Copying //UE4/Dev-Physics to Main (//UE4/Main)

diff --git a/Engine/Source/Runtime/Renderer/Private/HairStrands/HairStrandsDebug.cpp b/Engine/Source/Runtime/Renderer/Private/HairStrands/HairStrandsDebug.cpp
--- a/Engine/Source/Runtime/Renderer/Private/HairStrands/HairStrandsDebug.cpp
+++ b/Engine/Source/Runtime/Renderer/Private/HairStrands/HairStrandsDebug.cpp
@@ -23,6 +23,7 @@
 #include "CanvasTypes.h"
 #include "ShaderPrintParameters.h"
 #include "RenderGraphUtils.h"
+#include "GpuDebugRendering.h"
 
 static int32 GDeepShadowDebugIndex = 0;
 static float GDeepShadowDebugScale = 20;
@@ -51,6 +52,9 @@ static FAutoConsoleVariableRef CVarHairStrandsDebugBSDFExposure(TEXT("r.HairStra
 static int GHairStrandsDebugSampleIndex = -1;
 static FAutoConsoleVariableRef CVarHairStrandsDebugMaterialSampleIndex(TEXT("r.HairStrands.DebugMode.SampleIndex"), GHairStrandsDebugSampleIndex, TEXT("Debug value for a given sample index (default:-1, i.e., average sample information)."));
 
+static int32 GHairStrandsClusterDebug = 0;
+static FAutoConsoleVariableRef CVarHairStrandsDebugClusterAABB(TEXT("r.HairStrands.Cluster.Debug"), GHairStrandsClusterDebug, TEXT("Draw debug the world bounding box of hair clusters used for culling optimisation."));
+
 static int32 GHairDebugMeshProjection_SkinCacheMesh = 0;
 
 static int32 GHairDebugMeshProjection_Sim_HairRestTriangles = 0;
@@ -63,23 +67,25 @@ static int32 GHairDebugMeshProjection_Render_HairRestFrames = 0;
 static int32 GHairDebugMeshProjection_Render_HairDeformedTriangles = 0;
 static int32 GHairDebugMeshProjection_Render_HairDeformedFrames = 0;
 
-static FAutoConsoleVariableRef CVarHairDebugMeshProjection_SkinCacheMesh(					TEXT("r.HairStrands.MeshProjection.DebugSkinCache"),					GHairDebugMeshProjection_SkinCacheMesh, TEXT("Render debug mes projection"));
-static FAutoConsoleVariableRef CVarHairDebugMeshProjection_Render_HairRestTriangles(		TEXT("r.HairStrands.MeshProjection.Render.DebugHairRestTriangles"),		GHairDebugMeshProjection_Render_HairRestTriangles, TEXT("Render debug mes projection"));
-static FAutoConsoleVariableRef CVarHairDebugMeshProjection_Render_HairRestFrames(			TEXT("r.HairStrands.MeshProjection.Render.DebugHairRestFrames"),		GHairDebugMeshProjection_Render_HairRestFrames, TEXT("Render debug mes projection"));
-static FAutoConsoleVariableRef CVarHairDebugMeshProjection_Render_HairDeformedTriangles(	TEXT("r.HairStrands.MeshProjection.Render.DebugHairDeformedTriangles"),	GHairDebugMeshProjection_Render_HairDeformedTriangles, TEXT("Render debug mes projection"));
-static FAutoConsoleVariableRef CVarHairDebugMeshProjection_Render_HairDeformedFrames(		TEXT("r.HairStrands.MeshProjection.Render.DebugHairDeformedFrames"),	GHairDebugMeshProjection_Render_HairDeformedFrames, TEXT("Render debug mes projection"));
-
-static FAutoConsoleVariableRef CVarHairDebugMeshProjection_Sim_HairRestTriangles(			TEXT("r.HairStrands.MeshProjection.Sim.DebugHairRestTriangles"),		GHairDebugMeshProjection_Sim_HairRestTriangles, TEXT("Render debug mes projection"));
-static FAutoConsoleVariableRef CVarHairDebugMeshProjection_Sim_HairRestFrames(				TEXT("r.HairStrands.MeshProjection.Sim.DebugHairRestFrames"),			GHairDebugMeshProjection_Sim_HairRestFrames, TEXT("Render debug mes projection"));
-static FAutoConsoleVariableRef CVarHairDebugMeshProjection_Sim_HairDeformedTriangles(		TEXT("r.HairStrands.MeshProjection.Sim.DebugHairDeformedTriangles"),	GHairDebugMeshProjection_Sim_HairDeformedTriangles, TEXT("Render debug mes projection"));
-static FAutoConsoleVariableRef CVarHairDebugMeshProjection_Sim_HairDeformedFrames(			TEXT("r.HairStrands.MeshProjection.Sim.DebugHairDeformedFrames"),		GHairDebugMeshProjection_Sim_HairDeformedFrames, TEXT("Render debug mes projection"));
+static FAutoConsoleVariableRef CVarHairDebugMeshProjection_SkinCacheMesh(					TEXT("r.HairStrands.MeshProjection.DebugSkinCache"),			GHairDebugMeshProjection_SkinCacheMesh, TEXT("Render debug mes projection"));
+static FAutoConsoleVariableRef CVarHairDebugMeshProjection_Render_HairRestTriangles(		TEXT("r.HairStrands.MeshProjection.Render.Rest.Triangles"),		GHairDebugMeshProjection_Render_HairRestTriangles, TEXT("Render debug mes projection"));
+static FAutoConsoleVariableRef CVarHairDebugMeshProjection_Render_HairRestFrames(			TEXT("r.HairStrands.MeshProjection.Render.Rest.Frames"),		GHairDebugMeshProjection_Render_HairRestFrames, TEXT("Render debug mes projection"));
+static FAutoConsoleVariableRef CVarHairDebugMeshProjection_Render_HairDeformedTriangles(	TEXT("r.HairStrands.MeshProjection.Render.Deformed.Triangles"),	GHairDebugMeshProjection_Render_HairDeformedTriangles, TEXT("Render debug mes projection"));
+static FAutoConsoleVariableRef CVarHairDebugMeshProjection_Render_HairDeformedFrames(		TEXT("r.HairStrands.MeshProjection.Render.Deformed.Frames"),	GHairDebugMeshProjection_Render_HairDeformedFrames, TEXT("Render debug mes projection"));
 
+static FAutoConsoleVariableRef CVarHairDebugMeshProjection_Sim_HairRestTriangles(			TEXT("r.HairStrands.MeshProjection.Sim.Rest.Triangles"),		GHairDebugMeshProjection_Sim_HairRestTriangles, TEXT("Render debug mes projection"));
+static FAutoConsoleVariableRef CVarHairDebugMeshProjection_Sim_HairRestFrames(				TEXT("r.HairStrands.MeshProjection.Sim.Rest.Frames"),			GHairDebugMeshProjection_Sim_HairRestFrames, TEXT("Render debug mes projection"));
+static FAutoConsoleVariableRef CVarHairDebugMeshProjection_Sim_HairDeformedTriangles(		TEXT("r.HairStrands.MeshProjection.Sim.Deformed.Triangles"),	GHairDebugMeshProjection_Sim_HairDeformedTriangles, TEXT("Render debug mes projection"));
+static FAutoConsoleVariableRef CVarHairDebugMeshProjection_Sim_HairDeformedFrames(			TEXT("r.HairStrands.MeshProjection.Sim.Deformed.Frames"),		GHairDebugMeshProjection_Sim_HairDeformedFrames, TEXT("Render debug mes projection"));
 static int32 GHairStrandsDebugPPLL = 0;
-static FAutoConsoleVariableRef CVarHairStrandsDebugPPLL(									TEXT("r.HairStrands.DebugPPLL"),										GHairStrandsDebugPPLL, TEXT("Draw debug per pixel light list rendering."));
+static FAutoConsoleVariableRef CVarHairStrandsDebugPPLL(									TEXT("r.HairStrands.DebugPPLL"),									GHairStrandsDebugPPLL, TEXT("Draw debug per pixel light list rendering."));
+static int32 GHairVirtualVoxel_DrawDebugPage = 0;
+static FAutoConsoleVariableRef CVarHairVirtualVoxel_DrawDebugPage(TEXT("r.HairStrands.Voxelization.Virtual.DrawDebugPage"), GHairVirtualVoxel_DrawDebugPage, TEXT("When voxel debug rendering is enable, render the page bounds, instead of the voxel"));
+
 
 // Helper functions for accessing interpolation data for debug purpose.
 // Definitions is in HairStrandsInterface.cpp
-void GetGroomInterpolationData(const EWorldType::Type WorldType, FHairStrandsProjectionMeshData& OutGeometries);
+void GetGroomInterpolationData(const EWorldType::Type WorldType, FHairStrandsProjectionMeshData::LOD& OutGeometries);
 void GetGroomInterpolationData(const EWorldType::Type WorldType, const bool bRenderData, FHairStrandsProjectionHairData& OutHairData, TArray<int32>& OutLODIndices);
 
 static int32 GHairStrandsCull = 0;
@@ -102,10 +108,10 @@ FHairCullInfo GetHairStrandsCullInfo()
 enum class EHairDebugMode : uint8
 {
 	None,
-	ClusterData,
+	MacroGroups,
 	LightBounds,
 	DeepOpacityMaps,
-	ClusterScreenRect,
+	MacroGroupScreenRect,
 	SamplePerPixel,
 	CoverageType,
 	TAAResolveType,
@@ -119,7 +125,8 @@ enum class EHairDebugMode : uint8
 	MaterialBaseColor,
 	MaterialRoughness,
 	MaterialSpecular,
-	MaterialTangent
+	MaterialTangent,
+	Tile
 };
 
 static EHairDebugMode GetHairDebugMode()
@@ -127,9 +134,9 @@ static EHairDebugMode GetHairDebugMode()
 	switch (GHairStrandsDebugMode)
 	{
 	case 0:  return EHairDebugMode::None;
-	case 1:  return EHairDebugMode::ClusterData;
+	case 1:  return EHairDebugMode::MacroGroups;
 	case 2:  return EHairDebugMode::LightBounds;
-	case 3:  return EHairDebugMode::ClusterScreenRect;
+	case 3:  return EHairDebugMode::MacroGroupScreenRect;
 	case 4:  return EHairDebugMode::DeepOpacityMaps;
 	case 5:  return EHairDebugMode::SamplePerPixel;
 	case 6:  return EHairDebugMode::TAAResolveType;
@@ -145,6 +152,7 @@ static EHairDebugMode GetHairDebugMode()
 	case 16: return EHairDebugMode::MaterialRoughness;
 	case 17: return EHairDebugMode::MaterialSpecular;
 	case 18: return EHairDebugMode::MaterialTangent;
+	case 19: return EHairDebugMode::Tile;
 	default: return EHairDebugMode::None;
 	};
 }
@@ -154,9 +162,9 @@ static const TCHAR* ToString(EHairDebugMode DebugMode)
 	switch (DebugMode)
 	{
 	case EHairDebugMode::None: return TEXT("None");
-	case EHairDebugMode::ClusterData: return TEXT("Cluster info");
+	case EHairDebugMode::MacroGroups: return TEXT("Macro groups info");
 	case EHairDebugMode::LightBounds: return TEXT("All DOMs light bounds");
-	case EHairDebugMode::ClusterScreenRect: return TEXT("Screen projected clusters");
+	case EHairDebugMode::MacroGroupScreenRect: return TEXT("Screen projected macro groups");
 	case EHairDebugMode::DeepOpacityMaps: return TEXT("Deep opacity maps");
 	case EHairDebugMode::SamplePerPixel: return TEXT("Sub-pixel sample count");
 	case EHairDebugMode::TAAResolveType: return TEXT("TAA resolve type (regular/responsive)");
@@ -172,6 +180,7 @@ static const TCHAR* ToString(EHairDebugMode DebugMode)
 	case EHairDebugMode::MaterialRoughness: return TEXT("Hair material roughness");
 	case EHairDebugMode::MaterialSpecular: return TEXT("Hair material specular");
 	case EHairDebugMode::MaterialTangent: return TEXT("Hair material tangent");
+	case EHairDebugMode::Tile: return TEXT("Hair tile cotegorization");
 	default: return TEXT("None");
 	};
 }
@@ -192,6 +201,7 @@ EHairStrandsDebugMode GetHairStrandsDebugStrandsMode()
 	case  8:  return EHairStrandsDebugMode::RenderHairRadiusVariation;
 	case  9:  return EHairStrandsDebugMode::RenderHairBaseColor;
 	case 10:  return EHairStrandsDebugMode::RenderHairRoughness;
+	case 11:  return EHairStrandsDebugMode::RenderVisCluster;
 	default: return EHairStrandsDebugMode::None;
 	};
 }
@@ -211,6 +221,7 @@ static const TCHAR* ToString(EHairStrandsDebugMode DebugMode)
 	case EHairStrandsDebugMode::RenderHairRadiusVariation	: return TEXT("Hair radius variation");
 	case EHairStrandsDebugMode::RenderHairBaseColor			: return TEXT("Hair vertices color");
 	case EHairStrandsDebugMode::RenderHairRoughness			: return TEXT("Hair vertices roughness");
+	case EHairStrandsDebugMode::RenderVisCluster			: return TEXT("Hair visility clusters");
 	default													: return TEXT("None");
 	};
 }
@@ -227,13 +238,16 @@ class FHairDebugPrintCS : public FGlobalShader
 		SHADER_PARAMETER(FIntPoint, PixelCoord)
 		SHADER_PARAMETER(FIntPoint, MaxResolution)
 		SHADER_PARAMETER(uint32, FastResolveMask)
+		SHADER_PARAMETER(uint32, HairMacroGroupCount)
 		SHADER_PARAMETER_RDG_TEXTURE(Texture2D, CategorizationTexture)
 		SHADER_PARAMETER_RDG_TEXTURE(Texture2D, HairVisibilityNodeOffsetAndCount)
 		SHADER_PARAMETER_RDG_BUFFER_SRV(StructuredBuffer, HairVisibilityNodeData)
+		SHADER_PARAMETER_RDG_BUFFER_SRV(Buffer, HairMacroGroupAABBBuffer)
 		SHADER_PARAMETER_SRV(Texture2D, DepthStencilTexture)
 		SHADER_PARAMETER_SAMPLER(SamplerState, LinearSampler)
 		SHADER_PARAMETER_STRUCT_REF(FViewUniformShaderParameters, ViewUniformBuffer)
 		SHADER_PARAMETER_STRUCT_INCLUDE(ShaderPrint::FShaderParameters, ShaderPrintUniformBuffer)
+		SHADER_PARAMETER_STRUCT_INCLUDE(ShaderDrawDebug::FShaderDrawDebugParameters, ShaderDrawUniformBuffer)
 	END_SHADER_PARAMETER_STRUCT()
 
 public:
@@ -247,6 +261,7 @@ static void AddDebugHairPrintPass(
 	const FViewInfo* View,
 	const EHairDebugMode InDebugMode,
 	const FHairStrandsVisibilityData& VisibilityData,
+	const FHairStrandsMacroGroupDatas& MacroGroupDatas,
 	const FShaderResourceViewRHIRef& InDepthStencilTexture)
 {
 	const TRefCountPtr<IPooledRenderTarget>& InCategorizationTexture = VisibilityData.CategorizationTexture;
@@ -258,6 +273,7 @@ static void AddDebugHairPrintPass(
 	FRDGTextureRef CategorizationTexture = InCategorizationTexture ? GraphBuilder.RegisterExternalTexture(InCategorizationTexture, TEXT("CategorizationTexture")) : nullptr;
 	FRDGTextureRef NodeIndex = InNodeIndex ? GraphBuilder.RegisterExternalTexture(InNodeIndex, TEXT("NodeIndex")) : nullptr;
 	FRDGBufferRef  NodeData  = InNodeData ? GraphBuilder.RegisterExternalBuffer(InNodeData, TEXT("NodeData")) : nullptr;
+	FRDGBufferRef  HairMacroGroupAABBBuffer = GraphBuilder.RegisterExternalBuffer(MacroGroupDatas.MacroGroupResources.MacroGroupAABBsBuffer, TEXT("MacroGroupData"));
 
 	const FIntRect Viewport = View->ViewRect;
 	const FIntPoint Resolution(Viewport.Width(), Viewport.Height());
@@ -272,7 +288,10 @@ static void AddDebugHairPrintPass(
 	Parameters->HairVisibilityNodeOffsetAndCount = NodeIndex;
 	Parameters->DepthStencilTexture = InDepthStencilTexture;
 	Parameters->LinearSampler = TStaticSamplerState<SF_Bilinear, AM_Clamp, AM_Clamp, AM_Clamp>::GetRHI();
+	Parameters->HairMacroGroupCount = MacroGroupDatas.Datas.Num();
+	Parameters->HairMacroGroupAABBBuffer = GraphBuilder.CreateSRV(HairMacroGroupAABBBuffer, PF_R32_SINT);
 	ShaderPrint::SetParameters(*View, Parameters->ShaderPrintUniformBuffer);
+	ShaderDrawDebug::SetParameters(GraphBuilder, View->ShaderDrawData, Parameters->ShaderDrawUniformBuffer);
 	TShaderMapRef<FHairDebugPrintCS> ComputeShader(View->ShaderMap);
 
 	ClearUnusedGraphResources(*ComputeShader, Parameters);
@@ -296,8 +315,10 @@ class FHairDebugPS : public FGlobalShader
 		SHADER_PARAMETER(uint32, FastResolveMask)
 		SHADER_PARAMETER(uint32, DebugMode)
 		SHADER_PARAMETER(int32, SampleIndex)
+		SHADER_PARAMETER(uint32, TileSize)
 		SHADER_PARAMETER_RDG_TEXTURE(Texture2D, CategorizationTexture)
 		SHADER_PARAMETER_RDG_TEXTURE(Texture2D, NodeIndex)
+		SHADER_PARAMETER_RDG_TEXTURE(Texture2D, TileIndexTexture)
 		SHADER_PARAMETER_RDG_BUFFER_SRV(StructuredBuffer, NodeData)
 		SHADER_PARAMETER_SRV(Texture2D, DepthStencilTexture)
 		SHADER_PARAMETER_SAMPLER(SamplerState, LinearSampler)
@@ -315,9 +336,11 @@ static void AddDebugHairPass(
 	FRDGBuilder& GraphBuilder,
 	const FViewInfo* View,
 	const EHairDebugMode InDebugMode,
+	const uint32 InTileSize,
 	const TRefCountPtr<IPooledRenderTarget>& InCategorizationTexture,
 	const TRefCountPtr<IPooledRenderTarget>& InNodeIndex,
 	const TRefCountPtr<FPooledRDGBuffer>& InNodeData,
+	const TRefCountPtr<IPooledRenderTarget>& InTileIndexTexture,
 	const FShaderResourceViewRHIRef& InDepthStencilTexture,
 	FRDGTextureRef& OutTarget)
 {
@@ -330,14 +353,16 @@ static void AddDebugHairPass(
 		InDebugMode == EHairDebugMode::MaterialBaseColor ||
 		InDebugMode == EHairDebugMode::MaterialRoughness ||
 		InDebugMode == EHairDebugMode::MaterialSpecular ||
-		InDebugMode == EHairDebugMode::MaterialTangent);
+		InDebugMode == EHairDebugMode::MaterialTangent ||
+		InDebugMode == EHairDebugMode::Tile);
 
-	if (!InCategorizationTexture || !InNodeIndex || !InNodeData) return;
+	if (!InCategorizationTexture || !InNodeIndex || !InNodeData || !InTileIndexTexture) return;
 	if (InDebugMode == EHairDebugMode::TAAResolveType && !InDepthStencilTexture) return;
 
 	FRDGTextureRef CategorizationTexture = InCategorizationTexture ? GraphBuilder.RegisterExternalTexture(InCategorizationTexture, TEXT("CategorizationTexture")) : nullptr;
 	FRDGTextureRef NodeIndex = InNodeIndex ? GraphBuilder.RegisterExternalTexture(InNodeIndex, TEXT("NodeIndex")) : nullptr;
 	FRDGBufferRef  NodeData = InNodeData ? GraphBuilder.RegisterExternalBuffer(InNodeData, TEXT("NodeData")) : nullptr;
+	FRDGTextureRef TileIndexTexture = InTileIndexTexture ? GraphBuilder.RegisterExternalTexture(InTileIndexTexture, TEXT("TileIndexTexture")) : nullptr;
 
 	const FIntRect Viewport = View->ViewRect;
 	const FIntPoint Resolution(Viewport.Width(), Viewport.Height());
@@ -354,6 +379,7 @@ static void AddDebugHairPass(
 		case EHairDebugMode::MaterialRoughness:	InternalDebugMode = 6; break;
 		case EHairDebugMode::MaterialSpecular:	InternalDebugMode = 7; break;
 		case EHairDebugMode::MaterialTangent:	InternalDebugMode = 8; break;
+		case EHairDebugMode::Tile:				InternalDebugMode = 9; break;
 	};
 
 	FHairDebugPS::FParameters* Parameters = GraphBuilder.AllocParameters<FHairDebugPS::FParameters>();
@@ -363,6 +389,8 @@ static void AddDebugHairPass(
 	Parameters->CategorizationTexture = CategorizationTexture;
 	Parameters->NodeIndex = NodeIndex;
 	Parameters->NodeData = GraphBuilder.CreateSRV(NodeData);
+	Parameters->TileIndexTexture = TileIndexTexture;
+	Parameters->TileSize = InTileSize;
 	Parameters->DepthStencilTexture = InDepthStencilTexture;
 	Parameters->LinearSampler = TStaticSamplerState<SF_Bilinear, AM_Clamp, AM_Clamp, AM_Clamp>::GetRHI();
 	Parameters->DebugMode = InternalDebugMode;
@@ -382,7 +410,7 @@ static void AddDebugHairPass(
 	{
 		FGraphicsPipelineStateInitializer GraphicsPSOInit;
 		RHICmdList.ApplyCachedRenderTargets(GraphicsPSOInit);
-		GraphicsPSOInit.BlendState = TStaticBlendState<CW_RGBA, BO_Add, BF_One, BF_Zero, BO_Add, BF_One, BF_Zero>::GetRHI();
+		GraphicsPSOInit.BlendState = TStaticBlendState<CW_RGBA, BO_Add, BF_SourceAlpha, BF_InverseSourceAlpha, BO_Add, BF_One, BF_Zero>::GetRHI();
 		GraphicsPSOInit.RasterizerState = TStaticRasterizerState<>::GetRHI();
 		GraphicsPSOInit.DepthStencilState = TStaticDepthStencilState<false, CF_Always>::GetRHI();
 		GraphicsPSOInit.BoundShaderState.VertexDeclarationRHI = GFilterVertexDeclaration.VertexDeclarationRHI;
@@ -559,7 +587,7 @@ static void AddVoxelRaymarchingPass(
 	FRDGBuilder& GraphBuilder,
 	const FViewInfo& View,
 	const EHairDebugMode DebugMode,
-	const FHairStrandsClusterDatas& ClusterDatas,
+	const FHairStrandsMacroGroupDatas& MacroGroupDatas,
 	FRDGTextureRef& OutputTexture)
 {
 	check(DebugMode == EHairDebugMode::VoxelsDensity || DebugMode == EHairDebugMode::VoxelsTangent || DebugMode == EHairDebugMode::VoxelsBaseColor || DebugMode == EHairDebugMode::VoxelsRoughness);
@@ -568,22 +596,22 @@ static void AddVoxelRaymarchingPass(
 	SetupSceneTextureParameters(GraphBuilder, &SceneTextures);
 
 	const FIntPoint Resolution(OutputTexture->Desc.Extent);
-	for (const FHairStrandsClusterData& ClusterData : ClusterDatas.Datas)
+	for (const FHairStrandsMacroGroupData& MacroGroupData : MacroGroupDatas.Datas)
 	{
-		if (DebugMode == EHairDebugMode::VoxelsDensity && !ClusterData.VoxelResources.DensityTexture)
+		if (DebugMode == EHairDebugMode::VoxelsDensity && !MacroGroupData.VoxelResources.DensityTexture)
 			return;
 
-		if (DebugMode == EHairDebugMode::VoxelsTangent && (!ClusterData.VoxelResources.TangentXTexture || !ClusterData.VoxelResources.TangentYTexture || !ClusterData.VoxelResources.TangentZTexture))
+		if (DebugMode == EHairDebugMode::VoxelsTangent && (!MacroGroupData.VoxelResources.TangentXTexture || !MacroGroupData.VoxelResources.TangentYTexture || !MacroGroupData.VoxelResources.TangentZTexture))
 			return;
 
-		if ((DebugMode == EHairDebugMode::VoxelsBaseColor || DebugMode == EHairDebugMode::VoxelsRoughness) && !ClusterData.VoxelResources.MaterialTexture)
+		if ((DebugMode == EHairDebugMode::VoxelsBaseColor || DebugMode == EHairDebugMode::VoxelsRoughness) && !MacroGroupData.VoxelResources.MaterialTexture)
 			return;
 
-		const FRDGTextureRef VoxelDensityTexture = GraphBuilder.RegisterExternalTexture(ClusterData.VoxelResources.DensityTexture ? ClusterData.VoxelResources.DensityTexture : GSystemTextures.BlackDummy, TEXT("HairVoxelDensityTexture"));
-		const FRDGTextureRef VoxelTangentXTexture = GraphBuilder.RegisterExternalTexture(ClusterData.VoxelResources.TangentXTexture ? ClusterData.VoxelResources.TangentXTexture : GSystemTextures.BlackDummy, TEXT("HairVoxelTangentXTexture"));
-		const FRDGTextureRef VoxelTangentYTexture = GraphBuilder.RegisterExternalTexture(ClusterData.VoxelResources.TangentYTexture ? ClusterData.VoxelResources.TangentYTexture : GSystemTextures.BlackDummy, TEXT("HairVoxelTangentYTexture"));
-		const FRDGTextureRef VoxelTangentZTexture = GraphBuilder.RegisterExternalTexture(ClusterData.VoxelResources.TangentZTexture ? ClusterData.VoxelResources.TangentZTexture : GSystemTextures.BlackDummy, TEXT("HairVoxelTangentZTexture"));
-		const FRDGTextureRef VoxelMaterialTexture = GraphBuilder.RegisterExternalTexture(ClusterData.VoxelResources.MaterialTexture ? ClusterData.VoxelResources.MaterialTexture : GSystemTextures.BlackDummy, TEXT("HairVoxelMaterialTexture"));
+		const FRDGTextureRef VoxelDensityTexture  = GraphBuilder.RegisterExternalTexture(MacroGroupData.VoxelResources.DensityTexture  ? MacroGroupData.VoxelResources.DensityTexture : GSystemTextures.BlackDummy, TEXT("HairVoxelDensityTexture"));
+		const FRDGTextureRef VoxelTangentXTexture = GraphBuilder.RegisterExternalTexture(MacroGroupData.VoxelResources.TangentXTexture ? MacroGroupData.VoxelResources.TangentXTexture : GSystemTextures.BlackDummy, TEXT("HairVoxelTangentXTexture"));
+		const FRDGTextureRef VoxelTangentYTexture = GraphBuilder.RegisterExternalTexture(MacroGroupData.VoxelResources.TangentYTexture ? MacroGroupData.VoxelResources.TangentYTexture : GSystemTextures.BlackDummy, TEXT("HairVoxelTangentYTexture"));
+		const FRDGTextureRef VoxelTangentZTexture = GraphBuilder.RegisterExternalTexture(MacroGroupData.VoxelResources.TangentZTexture ? MacroGroupData.VoxelResources.TangentZTexture : GSystemTextures.BlackDummy, TEXT("HairVoxelTangentZTexture"));
+		const FRDGTextureRef VoxelMaterialTexture = GraphBuilder.RegisterExternalTexture(MacroGroupData.VoxelResources.MaterialTexture ? MacroGroupData.VoxelResources.MaterialTexture : GSystemTextures.BlackDummy, TEXT("HairVoxelMaterialTexture"));
 
 		FVoxelRaymarchingPS::FParameters* Parameters = GraphBuilder.AllocParameters<FVoxelRaymarchingPS::FParameters>();
 		Parameters->ViewUniformBuffer = View.ViewUniformBuffer;
@@ -594,9 +622,9 @@ static void AddVoxelRaymarchingPass(
 		Parameters->TangentYTexture = VoxelTangentYTexture;
 		Parameters->TangentZTexture = VoxelTangentZTexture;
 		Parameters->MaterialTexture = VoxelMaterialTexture;
-		Parameters->VoxelMinAABB = ClusterData.GetMinBound();
-		Parameters->VoxelMaxAABB = ClusterData.GetMaxBound();
-		Parameters->VoxelResolution = ClusterData.GetResolution();
+		Parameters->VoxelMinAABB = MacroGroupData.GetMinBound();
+		Parameters->VoxelMaxAABB = MacroGroupData.GetMaxBound();
+		Parameters->VoxelResolution = MacroGroupData.GetResolution();
 		Parameters->VoxelDensityScale = GetHairStrandsVoxelizationDensityScale();
 		Parameters->DensityIsoline = 1;
 		Parameters->LinearSampler = TStaticSamplerState<SF_Trilinear, AM_Clamp, AM_Clamp, AM_Clamp>::GetRHI();
@@ -658,7 +686,73 @@ static void AddVoxelRaymarchingPass(
 		});
 	}
 }
+
 	
+///////////////////////////////////////////////////////////////////////////////////////////////////
+class FVoxelVirtualRaymarchingCS : public FGlobalShader
+{
+	DECLARE_GLOBAL_SHADER(FVoxelVirtualRaymarchingCS);
+	SHADER_USE_PARAMETER_STRUCT(FVoxelVirtualRaymarchingCS, FGlobalShader);
+
+	BEGIN_SHADER_PARAMETER_STRUCT(FParameters, )
+		SHADER_PARAMETER_STRUCT_INCLUDE(FSceneTextureParameters, SceneTextures)
+		SHADER_PARAMETER_STRUCT_INCLUDE(ShaderDrawDebug::FShaderDrawDebugParameters, ShaderDrawParameters)
+		SHADER_PARAMETER_STRUCT_INCLUDE(ShaderPrint::FShaderParameters, ShaderPrintParameters)
+		SHADER_PARAMETER(FVector2D, OutputResolution)
+		SHADER_PARAMETER(uint32, bDrawPage)
+		SHADER_PARAMETER(uint32, MacroGroupId)
+		SHADER_PARAMETER(uint32, MacroGroupCount)
+		SHADER_PARAMETER(uint32, MaxTotalPageIndexCount)
+		SHADER_PARAMETER_STRUCT_REF(FVirtualVoxelParameters, VirtualVoxel)
+		SHADER_PARAMETER_STRUCT_REF(FViewUniformShaderParameters, ViewUniformBuffer)
+		SHADER_PARAMETER_RDG_BUFFER_SRV(Buffer, TotalValidPageCounter)
+		SHADER_PARAMETER_RDG_TEXTURE_UAV(RWTexture2D, OutputTexture)
+	END_SHADER_PARAMETER_STRUCT()
+
+public:
+	static bool ShouldCompilePermutation(const FGlobalShaderPermutationParameters& Parameters) { return IsHairStrandsSupported(Parameters.Platform); }
+};
+
+IMPLEMENT_GLOBAL_SHADER(FVoxelVirtualRaymarchingCS, "/Engine/Private/HairStrands/HairStrandsVoxelPageRayMarching.usf", "MainCS", SF_Compute);
+
+static void AddVoxelPageRaymarchingPass(
+	FRDGBuilder& GraphBuilder,
+	const FViewInfo& View,
+	const FHairStrandsMacroGroupDatas& MacroGroupDatas,
+	FRDGTextureRef& OutputTexture)
+{
+	FSceneTextureParameters SceneTextures;
+	SetupSceneTextureParameters(GraphBuilder, &SceneTextures);
+
+	const FVirtualVoxelResources& VoxelResources = MacroGroupDatas.VirtualVoxelResources;
+	const FIntPoint Resolution(OutputTexture->Desc.Extent);
+	for (const FHairStrandsMacroGroupData& MacroGroupData : MacroGroupDatas.Datas)
+	{
+		const FRDGBufferRef PageIndexBuffer	= GraphBuilder.RegisterExternalBuffer(VoxelResources.PageIndexBuffer, TEXT("HairVoxelPageIndexBuffer"));
+		const FRDGTextureRef PageTexture = GraphBuilder.RegisterExternalTexture(VoxelResources.PageTexture, TEXT("HairVoxelPageTexture"));
+		const FRDGBufferRef PageIndexGlobalCounter = GraphBuilder.RegisterExternalBuffer(VoxelResources.PageIndexGlobalCounter, TEXT("HairVoxelValidPageCountBuffer"));
+
+		FVoxelVirtualRaymarchingCS::FParameters* Parameters = GraphBuilder.AllocParameters<FVoxelVirtualRaymarchingCS::FParameters>();
+		Parameters->ViewUniformBuffer		= View.ViewUniformBuffer;
+		Parameters->OutputResolution		= Resolution;
+		Parameters->SceneTextures			= SceneTextures;
+		Parameters->bDrawPage				= GHairVirtualVoxel_DrawDebugPage ? 1 : 0;
+		Parameters->MacroGroupId			= MacroGroupData.MacroGroupId;
+		Parameters->MacroGroupCount			= MacroGroupDatas.Datas.Num();
+		Parameters->MaxTotalPageIndexCount  = VoxelResources.Parameters.Common.PageIndexCount;
+		Parameters->VirtualVoxel			= VoxelResources.UniformBuffer;
+		Parameters->TotalValidPageCounter	= GraphBuilder.CreateSRV(PageIndexGlobalCounter, PF_R32_UINT);
+		ShaderDrawDebug::SetParameters(GraphBuilder, View.ShaderDrawData, Parameters->ShaderDrawParameters);
+		ShaderPrint::SetParameters(View, Parameters->ShaderPrintParameters);
+		Parameters->OutputTexture			= GraphBuilder.CreateUAV(OutputTexture);
+
+		TShaderMapRef<FVoxelVirtualRaymarchingCS> ComputeShader(View.ShaderMap);
+
+		const FIntVector DispatchCount = DispatchCount.DivideAndRoundUp(FIntVector(OutputTexture->Desc.Extent.X, OutputTexture->Desc.Extent.Y, 1), FIntVector(8, 8, 1));
+		FComputeShaderUtils::AddPass(GraphBuilder, RDG_EVENT_NAME("HairStrandsVoxelVirtualRaymarching"), *ComputeShader, Parameters, DispatchCount);
+	}
+}
+
 ///////////////////////////////////////////////////////////////////////////////////////////////////
 class FHairStrandsBSDFPlotPS : public FGlobalShader
 {
@@ -996,12 +1090,12 @@ static void AddDebugProjectionHairPass(
 		Parameters->RootBarycentricBuffer = HairData.LODDatas[LODIndex].RootTriangleBarycentricBuffer->SRV;
 	}
 
-	Parameters->RestPositionOffset = HairData.LODDatas[LODIndex].RestPositionOffset;
+	Parameters->RestPositionOffset = *HairData.LODDatas[LODIndex].RestPositionOffset;
 	Parameters->RestPosition0Buffer = HairData.LODDatas[LODIndex].RestRootTrianglePosition0Buffer->SRV;
 	Parameters->RestPosition1Buffer = HairData.LODDatas[LODIndex].RestRootTrianglePosition1Buffer->SRV;
 	Parameters->RestPosition2Buffer = HairData.LODDatas[LODIndex].RestRootTrianglePosition2Buffer->SRV;
 	
-	Parameters->DeformedPositionOffset = HairData.LODDatas[LODIndex].DeformedPositionOffset;
+	Parameters->DeformedPositionOffset = *HairData.LODDatas[LODIndex].DeformedPositionOffset;
 	Parameters->DeformedPosition0Buffer = HairData.LODDatas[LODIndex].DeformedRootTrianglePosition0Buffer->SRV;
 	Parameters->DeformedPosition1Buffer = HairData.LODDatas[LODIndex].DeformedRootTrianglePosition1Buffer->SRV;
 	Parameters->DeformedPosition2Buffer = HairData.LODDatas[LODIndex].DeformedRootTrianglePosition2Buffer->SRV;
@@ -1114,7 +1208,111 @@ IMPLEMENT_GLOBAL_SHADER(FHairVisibilityDebugPPLLCS, "/Engine/Private/HairStrands
 
 ///////////////////////////////////////////////////////////////////////////////////////////////////
 
-void RenderHairStrandsDebugInfo(FRHICommandListImmediate& RHICmdList, TArray<FViewInfo>& Views, const FHairStrandsDatas* HairDatas)
+class FDrawDebugClusterAABBCS : public FGlobalShader
+{
+	DECLARE_GLOBAL_SHADER(FDrawDebugClusterAABBCS);
+	SHADER_USE_PARAMETER_STRUCT(FDrawDebugClusterAABBCS, FGlobalShader);
+
+	class FDebugAABBBuffer : SHADER_PERMUTATION_INT("PERMUTATION_DEBUGAABBBUFFER", 2);
+	using FPermutationDomain = TShaderPermutationDomain<FDebugAABBBuffer>;
+
+	BEGIN_SHADER_PARAMETER_STRUCT(FParameters, )
+		SHADER_PARAMETER_STRUCT_REF(FViewUniformShaderParameters, ViewUniformBuffer)
+		SHADER_PARAMETER_SRV(Buffer, ClusterAABBBuffer)
+		SHADER_PARAMETER_SRV(Buffer, GroupAABBBuffer)
+		SHADER_PARAMETER_RDG_BUFFER_SRV(Buffer, CulledDispatchIndirectParametersClusterCountBuffer)
+		SHADER_PARAMETER_RDG_BUFFER_SRV(Buffer, ClusterDebugAABBBuffer)
+		SHADER_PARAMETER_SRV(Buffer, CulledDrawIndirectParameters)
+		SHADER_PARAMETER(uint32, ClusterCount)
+		SHADER_PARAMETER(uint32, TriangleCount)
+		SHADER_PARAMETER(uint32, HairGroupId)
+		SHADER_PARAMETER(int32, ClusterDebugMode)
+		SHADER_PARAMETER_STRUCT_INCLUDE(ShaderDrawDebug::FShaderDrawDebugParameters, ShaderDrawParameters)
+		SHADER_PARAMETER_STRUCT_INCLUDE(ShaderPrint::FShaderParameters, ShaderPrintParameters)
+	END_SHADER_PARAMETER_STRUCT()
+
+public:
+	static bool ShouldCompilePermutation(const FGlobalShaderPermutationParameters& Parameters) { return IsHairStrandsSupported(Parameters.Platform); }
+	static void ModifyCompilationEnvironment(const FGlobalShaderPermutationParameters& Parameters, FShaderCompilerEnvironment& OutEnvironment)
+	{
+		FGlobalShader::ModifyCompilationEnvironment(Parameters, OutEnvironment);
+		OutEnvironment.SetDefine(TEXT("SHADER_DRAWDEBUGAABB"), 1);
+	}
+};
+
+IMPLEMENT_GLOBAL_SHADER(FDrawDebugClusterAABBCS, "/Engine/Private/HairStrands/HairStrandsClusterCulling.usf", "MainDrawDebugAABBCS", SF_Compute);
+
+static void AddDrawDebugClusterPass(
+	FRHICommandListImmediate& RHICmdList,
+	const FHairStrandClusterData& HairClusterData,
+	TArray<FViewInfo>& Views)
+{
+	check(ShaderDrawDebug::IsShaderDrawDebugEnabled());
+
+	static IConsoleVariable* CVarClusterCullingEnabled = IConsoleManager::Get().FindConsoleVariable(TEXT("r.HairStrands.Cluster.Culling"));
+	if (CVarClusterCullingEnabled && CVarClusterCullingEnabled->GetInt() <= 0)
+	{
+		return;
+	}
+
+	TShaderMap<FGlobalShaderType>* ShaderMap = GetGlobalShaderMap(ERHIFeatureLevel::SM5);
+
+	for (int DataIndex = 0; DataIndex < HairClusterData.HairGroups.Num(); ++DataIndex)
+	{
+		const FHairStrandClusterData::FHairGroup& HairGroupClusters = HairClusterData.HairGroups[DataIndex];
+
+		for (int32 ViewIndex = 0; ViewIndex < Views.Num(); ++ViewIndex)
+		{
+			FViewInfo& ViewInfo = Views[ViewIndex];
+			if (ShaderDrawDebug::IsShaderDrawDebugEnabled(ViewInfo))
+			{
+				FRDGBuilder GraphBuilder(RHICmdList);
+
+				FRDGBufferRef CulledDispatchIndirectParametersClusterCount = GraphBuilder.RegisterExternalBuffer(HairGroupClusters.CulledDispatchIndirectParametersClusterCount);
+				FRWBuffer& DrawIndirectBuffer = HairGroupClusters.HairGroupPublicPtr->GetDrawIndirectBuffer();
+
+				FDrawDebugClusterAABBCS::FPermutationDomain Permutation;
+				Permutation.Set<FDrawDebugClusterAABBCS::FDebugAABBBuffer>(GHairStrandsClusterDebug > 1 ? 1 : 0);
+				TShaderMapRef<FDrawDebugClusterAABBCS> ComputeShader(ShaderMap, Permutation);
+
+				FDrawDebugClusterAABBCS::FParameters* Parameters = GraphBuilder.AllocParameters<FDrawDebugClusterAABBCS::FParameters>();
+				Parameters->ViewUniformBuffer = ViewInfo.ViewUniformBuffer;
+				Parameters->ClusterCount = HairGroupClusters.ClusterCount;
+				Parameters->TriangleCount = HairGroupClusters.VertexCount / 3;
+				Parameters->HairGroupId = DataIndex;
+				Parameters->ClusterDebugMode = GHairStrandsClusterDebug;
+				Parameters->ClusterAABBBuffer = HairGroupClusters.ClusterAABBBuffer->SRV;
+				Parameters->CulledDispatchIndirectParametersClusterCountBuffer = GraphBuilder.CreateSRV(CulledDispatchIndirectParametersClusterCount, EPixelFormat::PF_R32_UINT);
+				Parameters->CulledDrawIndirectParameters = DrawIndirectBuffer.SRV;
+				Parameters->GroupAABBBuffer = HairGroupClusters.GroupAABBBuffer->SRV;
+
+				if (HairGroupClusters.ClusterDebugAABBBuffer && GHairStrandsClusterDebug>1)
+				{
+					FRDGBufferRef ClusterDebugAABBBuffer = GraphBuilder.RegisterExternalBuffer(HairGroupClusters.ClusterDebugAABBBuffer);
+					Parameters->ClusterDebugAABBBuffer = GraphBuilder.CreateSRV(ClusterDebugAABBBuffer, EPixelFormat::PF_R32_SINT);
+				}
+				ShaderDrawDebug::SetParameters(GraphBuilder, ViewInfo.ShaderDrawData, Parameters->ShaderDrawParameters);
+				ShaderPrint::SetParameters(ViewInfo, Parameters->ShaderPrintParameters);
+
+				check(Parameters->ClusterCount/64 <= 65535);
+				const FIntVector DispatchCount = DispatchCount.DivideAndRoundUp(FIntVector(Parameters->ClusterCount, 1, 1), FIntVector(64, 1, 1));// FIX ME, this could get over 65535
+				FComputeShaderUtils::AddPass(
+					GraphBuilder, RDG_EVENT_NAME("DrawDebugClusterAABB"),
+					*ComputeShader, Parameters, DispatchCount);
+
+				GraphBuilder.Execute();
+			}
+		}
+	}
+}
+
+///////////////////////////////////////////////////////////////////////////////////////////////////
+
+void RenderHairStrandsDebugInfo(
+	FRHICommandListImmediate& RHICmdList, 
+	TArray<FViewInfo>& Views, 
+	const FHairStrandsDatas* HairDatas, 
+	const struct FHairStrandClusterData& HairClusterData)
 {
 	const float YStep = 14;
 	const float ColumnWidth = 200;
@@ -1149,7 +1347,7 @@ void RenderHairStrandsDebugInfo(FRHICommandListImmediate& RHICmdList, TArray<FVi
 	}
 
 	float ClusterY = 38;
-	if (HairDebugMode == EHairDebugMode::ClusterData)
+	if (HairDebugMode == EHairDebugMode::MacroGroups)
 	{
 		// Component part of the clusters
 		FRenderTargetTemp TempRenderTarget(View, (const FTexture2DRHIRef&)SceneTargets.GetSceneColor()->GetRenderTargetItem().TargetableTexture);
@@ -1199,25 +1397,35 @@ void RenderHairStrandsDebugInfo(FRHICommandListImmediate& RHICmdList, TArray<FVi
 		return;
 
 	const FHairStrandsDeepShadowViews& InDomViews = HairDatas->DeepShadowViews;
-	const FHairStrandsClusterViews& InClusterViews = HairDatas->HairClusterPerViews;
+	const FHairStrandsMacroGroupViews& InMacroGroupViews = HairDatas->MacroGroupsPerViews;
 
-	if (HairDebugMode == EHairDebugMode::ClusterData)
+	if (HairDebugMode == EHairDebugMode::MacroGroups)
 	{
 		if (ViewIndex < uint32(HairDatas->HairVisibilityViews.HairDatas.Num()))
 		{
 			FRDGBuilder GraphBuilder(RHICmdList);
 			const FHairStrandsVisibilityData& VisibilityData = HairDatas->HairVisibilityViews.HairDatas[ViewIndex];
-			AddDebugHairPrintPass(GraphBuilder, &View, HairDebugMode, VisibilityData, SceneTargets.SceneStencilSRV);
+			const FHairStrandsMacroGroupDatas& MacroGroupDatas = InMacroGroupViews.Views[ViewIndex];
+			AddDebugHairPrintPass(GraphBuilder, &View, HairDebugMode, VisibilityData, MacroGroupDatas, SceneTargets.SceneStencilSRV);
 			GraphBuilder.Execute();	
 		}
 
-		// Clusters part of the clusters
+		// CPU bound of macro groups
 		FViewElementPDI ShadowFrustumPDI(&View, nullptr, nullptr);
-		const FHairStrandsClusterDatas& ClusterDatas = InClusterViews.Views[ViewIndex];
-		for (const FHairStrandsClusterData& ClusterData : ClusterDatas.Datas)
+		const FHairStrandsMacroGroupDatas& MacroGroupDatas = InMacroGroupViews.Views[ViewIndex];
+		const bool bUseVirtualVoxel = MacroGroupDatas.VirtualVoxelResources.IsValid();
+		for (const FHairStrandsMacroGroupData& MacroGroupData : MacroGroupDatas.Datas)
 		{
-			const FBox ClusterBox(ClusterData.GetMinBound(), ClusterData.GetMaxBound());
-			DrawWireBox(&ShadowFrustumPDI, ClusterBox, FColor::Red, 0);
+			if (bUseVirtualVoxel)
+			{
+				const FBox Bound(MacroGroupData.VirtualVoxelNodeDesc.WorldMinAABB, MacroGroupData.VirtualVoxelNodeDesc.WorldMaxAABB);
+				DrawWireBox(&ShadowFrustumPDI, Bound, FColor::Red, 0);
+			}
+			else
+			{
+				const FBox Bound(MacroGroupData.GetMinBound(), MacroGroupData.GetMaxBound());
+				DrawWireBox(&ShadowFrustumPDI, Bound, FColor::Red, 0);
+			}
 		}
 
 		FRenderTargetTemp TempRenderTarget(View, (const FTexture2DRHIRef&)SceneTargets.GetSceneColor()->GetRenderTargetItem().TargetableTexture);
@@ -1232,11 +1440,11 @@ void RenderHairStrandsDebugInfo(FRHICommandListImmediate& RHICmdList, TArray<FVi
 		Line = FString::Printf(TEXT("----------------------------------------------------------------"));
 		Canvas.DrawShadowedString(X, Y += YStep, *Line, GetStatsFont(), DebugColor);
 
-		Line = FString::Printf(TEXT("Cluster count : %d"), ClusterDatas.Datas.Num());
+		Line = FString::Printf(TEXT("Macro group count : %d"), MacroGroupDatas.Datas.Num());
 		Canvas.DrawShadowedString(X, Y += YStep, *Line, GetStatsFont(), DebugColor);
-		for (const FHairStrandsClusterData& ClusterData : ClusterDatas.Datas)
+		for (const FHairStrandsMacroGroupData& MacroGroupData : MacroGroupDatas.Datas)
 		{
-			Line = FString::Printf(TEXT(" %d - Bound Radus: %f.2m (%dx%d)"), ClusterData.ClusterId, ClusterData.Bounds.GetSphere().W);
+			Line = FString::Printf(TEXT(" %d - Bound Radus: %f.2m (%dx%d)"), MacroGroupData.MacroGroupId, MacroGroupData.Bounds.GetSphere().W);
 			Canvas.DrawShadowedString(X, Y += YStep, *Line, GetStatsFont(), DebugColor);
 		}
 
@@ -1267,19 +1475,19 @@ void RenderHairStrandsDebugInfo(FRHICommandListImmediate& RHICmdList, TArray<FVi
 	}
 
 	// View Rect
-	if (IsHairStrandsViewRectOptimEnable() && HairDebugMode == EHairDebugMode::ClusterScreenRect)
+	if (IsHairStrandsViewRectOptimEnable() && HairDebugMode == EHairDebugMode::MacroGroupScreenRect)
 	{
 		FRDGBuilder GraphBuilder(RHICmdList);
 		FRDGTextureRef SceneColorTexture = GraphBuilder.RegisterExternalTexture(SceneTargets.GetSceneColor(), TEXT("SceneColorTexture"));
-		if (ViewIndex < uint32(InClusterViews.Views.Num()))
+		if (ViewIndex < uint32(InMacroGroupViews.Views.Num()))
 		{
-			const FHairStrandsClusterDatas& ClusterDatas = InClusterViews.Views[ViewIndex];
-			for (const FHairStrandsClusterData& ClusterData : ClusterDatas.Datas)
+			const FHairStrandsMacroGroupDatas& MacroGroupDatas = InMacroGroupViews.Views[ViewIndex];
+			for (const FHairStrandsMacroGroupData& MacroGroupData : MacroGroupDatas.Datas)
 			{
-				AddDebugDeepShadowTexturePass(GraphBuilder, &View, ClusterData.ScreenRect, nullptr, SceneColorTexture);
+				AddDebugDeepShadowTexturePass(GraphBuilder, &View, MacroGroupData.ScreenRect, nullptr, SceneColorTexture);
 			}
 
-			const FIntRect TotalRect = ComputeVisibleHairStrandsClustersRect(View.ViewRect, ClusterDatas);
+			const FIntRect TotalRect = ComputeVisibleHairStrandsMacroGroupsRect(View.ViewRect, MacroGroupDatas);
 			AddDebugDeepShadowTexturePass(GraphBuilder, &View, TotalRect, nullptr, SceneColorTexture);
 		}
 		GraphBuilder.Execute();
@@ -1287,7 +1495,7 @@ void RenderHairStrandsDebugInfo(FRHICommandListImmediate& RHICmdList, TArray<FVi
 	
 	const bool bIsVoxelMode = HairDebugMode == EHairDebugMode::VoxelsDensity || HairDebugMode == EHairDebugMode::VoxelsTangent || HairDebugMode == EHairDebugMode::VoxelsBaseColor || HairDebugMode == EHairDebugMode::VoxelsRoughness;
 
-	// Render Frustum for all lights & clusters
+	// Render Frustum for all lights & macro groups
 	{
 		FViewElementPDI ShadowFrustumPDI(&View, nullptr, nullptr);
 
@@ -1315,15 +1523,24 @@ void RenderHairStrandsDebugInfo(FRHICommandListImmediate& RHICmdList, TArray<FVi
 		}
 
 		// Voxelization
-		if (bIsVoxelMode && ViewIndex < uint32(InClusterViews.Views.Num()))
+		if (bIsVoxelMode && ViewIndex < uint32(InMacroGroupViews.Views.Num()))
 		{
-			const FHairStrandsClusterDatas& ClusterDatas = InClusterViews.Views[ViewIndex];
-			for (const FHairStrandsClusterData& ClusterData : ClusterDatas.Datas)
+			const FHairStrandsMacroGroupDatas& MacroGroupDatas = InMacroGroupViews.Views[ViewIndex];
+			const bool bUseVirtualVoxel = MacroGroupDatas.VirtualVoxelResources.IsValid();
+			for (const FHairStrandsMacroGroupData& MacroGroupData : MacroGroupDatas.Datas)
 			{
-				DrawFrustumWireframe(&ShadowFrustumPDI, ClusterData.VoxelResources.WorldToClip.Inverse(), FColor::Purple, 0);
-
-				const FBox VoxelizationBox(ClusterData.GetMinBound(), ClusterData.GetMaxBound());
-				DrawWireBox(&ShadowFrustumPDI, VoxelizationBox, FColor::Red, 0);
+				if (bUseVirtualVoxel)
+				{
+					const FBox Bound(MacroGroupData.VirtualVoxelNodeDesc.WorldMinAABB, MacroGroupData.VirtualVoxelNodeDesc.WorldMaxAABB);
+					DrawWireBox(&ShadowFrustumPDI, Bound, FColor::Red, 0);
+					DrawFrustumWireframe(&ShadowFrustumPDI, MacroGroupData.VirtualVoxelNodeDesc.WorldToClip.Inverse(), FColor::Purple, 0);
+				}
+				else
+				{
+					const FBox VoxelizationBox(MacroGroupData.GetMinBound(), MacroGroupData.GetMaxBound());
+					DrawWireBox(&ShadowFrustumPDI, VoxelizationBox, FColor::Red, 0);
+					DrawFrustumWireframe(&ShadowFrustumPDI, MacroGroupData.VoxelResources.WorldToClip.Inverse(), FColor::Purple, 0);
+				}
 			}
 		}
 	}
@@ -1337,7 +1554,8 @@ void RenderHairStrandsDebugInfo(FRHICommandListImmediate& RHICmdList, TArray<FVi
 		HairDebugMode == EHairDebugMode::MaterialBaseColor ||
 		HairDebugMode == EHairDebugMode::MaterialRoughness ||
 		HairDebugMode == EHairDebugMode::MaterialSpecular ||
-		HairDebugMode == EHairDebugMode::MaterialTangent;
+		HairDebugMode == EHairDebugMode::MaterialTangent ||
+		HairDebugMode == EHairDebugMode::Tile;
 	if (bRunDebugPass)
 	{
 		FRDGBuilder GraphBuilder(RHICmdList);
@@ -1345,8 +1563,9 @@ void RenderHairStrandsDebugInfo(FRHICommandListImmediate& RHICmdList, TArray<FVi
 		if (ViewIndex < uint32(HairDatas->HairVisibilityViews.HairDatas.Num()))
 		{
 			const FHairStrandsVisibilityData& VisibilityData = HairDatas->HairVisibilityViews.HairDatas[ViewIndex];
-			AddDebugHairPass(GraphBuilder, &View, HairDebugMode, VisibilityData.CategorizationTexture, VisibilityData.NodeIndex, VisibilityData.NodeData, SceneTargets.SceneStencilSRV, SceneColorTexture);
-			AddDebugHairPrintPass(GraphBuilder, &View, HairDebugMode, VisibilityData,  SceneTargets.SceneStencilSRV);
+			const FHairStrandsMacroGroupDatas& MacroGroupDatas = InMacroGroupViews.Views[ViewIndex];
+			AddDebugHairPass(GraphBuilder, &View, HairDebugMode, VisibilityData.TileSize, VisibilityData.CategorizationTexture, VisibilityData.NodeIndex, VisibilityData.NodeData, VisibilityData.TileIndexTexture, SceneTargets.SceneStencilSRV, SceneColorTexture);
+			AddDebugHairPrintPass(GraphBuilder, &View, HairDebugMode, VisibilityData, MacroGroupDatas, SceneTargets.SceneStencilSRV);
 		}
 
 		GraphBuilder.Execute();
@@ -1354,13 +1573,24 @@ void RenderHairStrandsDebugInfo(FRHICommandListImmediate& RHICmdList, TArray<FVi
 
 	if (bIsVoxelMode)
 	{
-		if (ViewIndex < uint32(InClusterViews.Views.Num()))
+		if (ViewIndex < uint32(InMacroGroupViews.Views.Num()))
 		{
-			FRDGBuilder GraphBuilder(RHICmdList);
-			FRDGTextureRef SceneColorTexture = GraphBuilder.RegisterExternalTexture(SceneTargets.GetSceneColor(), TEXT("SceneColorTexture"));
-			const FHairStrandsClusterDatas& ClusterDatas = InClusterViews.Views[ViewIndex];
-			AddVoxelRaymarchingPass(GraphBuilder, View, HairDebugMode, ClusterDatas, SceneColorTexture);
-			GraphBuilder.Execute();
+			const FHairStrandsMacroGroupDatas& MacroGroupDatas = InMacroGroupViews.Views[ViewIndex];
+			const bool bIsVirtualVoxelizationEnabled = MacroGroupDatas.VirtualVoxelResources.IsValid();
+			if (bIsVirtualVoxelizationEnabled)
+			{
+				FRDGBuilder GraphBuilder(RHICmdList);
+				FRDGTextureRef SceneColorTexture = GraphBuilder.RegisterExternalTexture(SceneTargets.GetSceneColor(), TEXT("SceneColorTexture"));
+				AddVoxelPageRaymarchingPass(GraphBuilder, View, MacroGroupDatas, SceneColorTexture);
+				GraphBuilder.Execute();
+			}
+			else
+			{
+				FRDGBuilder GraphBuilder(RHICmdList);
+				FRDGTextureRef SceneColorTexture = GraphBuilder.RegisterExternalTexture(SceneTargets.GetSceneColor(), TEXT("SceneColorTexture"));
+				AddVoxelRaymarchingPass(GraphBuilder, View, HairDebugMode, MacroGroupDatas, SceneColorTexture);
+				GraphBuilder.Execute();
+			}
 		}
 	}
 
@@ -1391,9 +1621,9 @@ void RenderHairStrandsDebugInfo(FRHICommandListImmediate& RHICmdList, TArray<FVi
 
 			if (GHairDebugMeshProjection_SkinCacheMesh > 0)
 			{
-				FHairStrandsProjectionMeshData MeshProjectionData;
-				GetGroomInterpolationData(WorldType, MeshProjectionData);
-				for (FHairStrandsProjectionMeshData::Section& Section : MeshProjectionData.Sections)
+				FHairStrandsProjectionMeshData::LOD MeshProjectionLODData;
+				GetGroomInterpolationData(WorldType, MeshProjectionLODData);
+				for (FHairStrandsProjectionMeshData::Section& Section : MeshProjectionLODData.Sections)
 				{
 					AddDebugProjectionMeshPass(GraphBuilder, &View, bClearDepth, Section, SceneColorTexture, DepthTexture);
 					bClearDepth = false;
@@ -1494,6 +1724,11 @@ void RenderHairStrandsDebugInfo(FRHICommandListImmediate& RHICmdList, TArray<FVi
 		}
 	}
 
+	if (ShaderDrawDebug::IsShaderDrawDebugEnabled() && GHairStrandsClusterDebug > 0)
+	{
+		AddDrawDebugClusterPass(RHICmdList, HairClusterData, Views);
+	}
+
 	// Text
 	if (HairDebugMode == EHairDebugMode::LightBounds || HairDebugMode == EHairDebugMode::DeepOpacityMaps)
 	{
@@ -1518,9 +1753,9 @@ void RenderHairStrandsDebugInfo(FRHICommandListImmediate& RHICmdList, TArray<FVi
 		Line = FString::Printf(TEXT("Voxelization : %s"), IsHairStrandsVoxelizationEnable() ? TEXT("On") : TEXT("Off"));		Canvas.DrawShadowedString(X, Y += YStep, *Line, GetStatsFont(), DebugColor);
 		Line = FString::Printf(TEXT("View rect optim.: %s"), IsHairStrandsViewRectOptimEnable() ? TEXT("On") : TEXT("Off"));	Canvas.DrawShadowedString(X, Y += YStep, *Line, GetStatsFont(), DebugColor);
 		Line = FString::Printf(TEXT("----------------------------------------------------------------"));				Canvas.DrawShadowedString(X, Y += YStep, *Line, GetStatsFont(), DebugColor);
-		Line = FString::Printf(TEXT("DOM Atlas resolution : %d/%d"), AtlasResolution.X, AtlasResolution.Y);				Canvas.DrawShadowedString(X, Y += YStep, *Line, GetStatsFont(), DebugColor);
-		Line = FString::Printf(TEXT("DOM Cluster count : %d"), ViewData.Datas.Num());									Canvas.DrawShadowedString(X, Y += YStep, *Line, GetStatsFont(), DebugColor);
-		Line = FString::Printf(TEXT("DOM Texture Index : %d/%d"), DomTextureIndex, ViewData.Datas.Num());				Canvas.DrawShadowedString(X, Y += YStep, *Line, GetStatsFont(), DebugColor);
+		Line = FString::Printf(TEXT("DOM Atlas resolution  : %d/%d"), AtlasResolution.X, AtlasResolution.Y);				Canvas.DrawShadowedString(X, Y += YStep, *Line, GetStatsFont(), DebugColor);
+		Line = FString::Printf(TEXT("DOM macro group count : %d"), ViewData.Datas.Num());									Canvas.DrawShadowedString(X, Y += YStep, *Line, GetStatsFont(), DebugColor);
+		Line = FString::Printf(TEXT("DOM Texture Index     : %d/%d"), DomTextureIndex, ViewData.Datas.Num());				Canvas.DrawShadowedString(X, Y += YStep, *Line, GetStatsFont(), DebugColor);
 
 		uint32 BoundIndex = 0;
 		for (const FHairStrandsDeepShadowData& DomData : ViewData.Datas)
diff --git a/Engine/Source/Runtime/Renderer/Private/HairStrands/HairStrandsDebug.h b/Engine/Source/Runtime/Renderer/Private/HairStrands/HairStrandsDebug.h
--- a/Engine/Source/Runtime/Renderer/Private/HairStrands/HairStrandsDebug.h
+++ b/Engine/Source/Runtime/Renderer/Private/HairStrands/HairStrandsDebug.h
@@ -12,4 +12,5 @@
 void RenderHairStrandsDebugInfo(
 	FRHICommandListImmediate& RHICmdList,
 	TArray<FViewInfo>& Views,
-	const struct FHairStrandsDatas* HairDatas);
\ No newline at end of file
+	const struct FHairStrandsDatas* HairDatas,
+	const struct FHairStrandClusterData& HairClusterData);
\ No newline at end of file
diff --git a/Engine/Source/Runtime/Renderer/Private/HairStrands/HairStrandsDeepShadow.cpp b/Engine/Source/Runtime/Renderer/Private/HairStrands/HairStrandsDeepShadow.cpp
--- a/Engine/Source/Runtime/Renderer/Private/HairStrands/HairStrandsDeepShadow.cpp
+++ b/Engine/Source/Runtime/Renderer/Private/HairStrands/HairStrandsDeepShadow.cpp
@@ -45,76 +45,11 @@ static FLightSceneInfosArray GetVisibleDeepShadowLights(const FScene* Scene, con
 
 ///////////////////////////////////////////////////////////////////////////////////////////////////
 
-static void RenderDeepShadowFrontDepth(
-	FRHICommandList& RHICmdList,
-	const FScene* Scene,
-	const FViewInfo* ViewInfo,
-	const FHairStrandsClusterData::TPrimitiveInfos& PrimitiveSceneInfo,
-	const FIntRect& AtlasRect,
-	const TUniformBufferRef<FViewUniformShaderParameters>& ViewUniformShaderParameters,
-	const TUniformBufferRef<FDeepShadowPassUniformParameters>& DeepShadowPassUniformParameters,
-	const bool bClearOutput,
-	TRefCountPtr<IPooledRenderTarget>& outShadowDepthRT)
-{
-	DECLARE_GPU_STAT(HairStrandsDeepShadowFrontDepth);
-	SCOPED_DRAW_EVENT(RHICmdList, HairStrandsDeepShadowFrontDepth);
-	SCOPED_GPU_STAT(RHICmdList, HairStrandsDeepShadowFrontDepth);
-
-	EDepthStencilTargetActions DepthStencilAction = MakeDepthStencilTargetActions(
-		MakeRenderTargetActions(bClearOutput ? ERenderTargetLoadAction::EClear : ERenderTargetLoadAction::ELoad, ERenderTargetStoreAction::EStore),
-		MakeRenderTargetActions(ERenderTargetLoadAction::ENoAction, ERenderTargetStoreAction::ENoAction));
-	FRHIRenderPassInfo RPInfo(outShadowDepthRT->GetRenderTargetItem().TargetableTexture, DepthStencilAction, nullptr, FExclusiveDepthStencil::DepthWrite_StencilNop);
-	RHICmdList.BeginRenderPass(RPInfo, TEXT("DeepShadowDepth"));
-	RasterHairStrands(
-		RHICmdList,
-		Scene,
-		ViewInfo,
-		PrimitiveSceneInfo,
-		EHairStrandsRasterPassType::FrontDepth,
-		AtlasRect,
-		ViewUniformShaderParameters,
-		DeepShadowPassUniformParameters);
-	RHICmdList.EndRenderPass();
-}
-
-///////////////////////////////////////////////////////////////////////////////////////////////////
-
-static void RenderDeepShadowLayers(
-	FRHICommandListImmediate& RHICmdList,
-	const FScene* Scene,
-	const FViewInfo* ViewInfo,
-	const FHairStrandsClusterData::TPrimitiveInfos& PrimitiveSceneInfo,
-	const FIntRect& AtlasRect,
-	const TUniformBufferRef<FViewUniformShaderParameters>& ViewUniformShaderParameters,
-	const TUniformBufferRef<FDeepShadowPassUniformParameters>& DeepShadowPassUniformParameters,
-	const bool bClearOutput,
-	TRefCountPtr<IPooledRenderTarget>& outDeepShadowLayersRT)
-{
-	DECLARE_GPU_STAT(HairStrandsDeepShadowLayers);
-	SCOPED_DRAW_EVENT(RHICmdList, HairStrandsDeepShadowLayers);
-	SCOPED_GPU_STAT(RHICmdList, HairStrandsDeepShadowLayers);
-
-	FRHIRenderPassInfo RPInfo(outDeepShadowLayersRT->GetRenderTargetItem().TargetableTexture, bClearOutput ? ERenderTargetActions::Clear_Store : ERenderTargetActions::Load_Store);
-	RHICmdList.BeginRenderPass(RPInfo, TEXT("DeepShadowLayers"));
-	RasterHairStrands(
-		RHICmdList,
-		Scene,
-		ViewInfo,
-		PrimitiveSceneInfo,
-		EHairStrandsRasterPassType::DeepOpacityMap,
-		AtlasRect,
-		ViewUniformShaderParameters,
-		DeepShadowPassUniformParameters);
-	RHICmdList.EndRenderPass();
-}
-
-///////////////////////////////////////////////////////////////////////////////////////////////////
-
 FHairStrandsDeepShadowViews RenderHairStrandsDeepShadows(
 	FRHICommandListImmediate& RHICmdList,
 	const FScene* Scene,
 	const TArray<FViewInfo>& Views,
-	const FHairStrandsClusterViews& DeepShadowClusterViews)
+	const FHairStrandsMacroGroupViews& MacroGroupsViews)
 {
 	QUICK_SCOPE_CYCLE_COUNTER(STAT_CLM_RenderDeepShadow);
 	DECLARE_GPU_STAT(HairStrandsDeepShadow);
@@ -131,25 +66,25 @@ FHairStrandsDeepShadowViews RenderHairStrandsDeepShadows(
 		const FViewInfo& View = Views[ViewIndex];
 		if (View.Family)
 		{
-			const FHairStrandsClusterDatas& DeepShadowClusterDatas = DeepShadowClusterViews.Views[ViewIndex];
+			const FHairStrandsMacroGroupDatas& MacroGroupDatas = MacroGroupsViews.Views[ViewIndex];
 			DeepShadowsPerView.Views.AddDefaulted();
 
-			if (DeepShadowClusterDatas.Datas.Num() == 0 || 
+			if (MacroGroupDatas.Datas.Num() == 0 || 
 				VisibleLightsPerView[ViewIndex].Num() == 0 ||
 				IsHairStrandsForVoxelTransmittanceAndShadowEnable()) 
 			{
 				continue; 
 			}
 
-			for (int32 ClusterIt = 0; ClusterIt < DeepShadowClusterDatas.Datas.Num(); ++ClusterIt)
+			for (int32 MacroGroupIt = 0; MacroGroupIt < MacroGroupDatas.Datas.Num(); ++MacroGroupIt)
 			{
-				const FHairStrandsClusterData& Cluster = DeepShadowClusterDatas.Datas[ClusterIt];
-				const FBoxSphereBounds ClusterBounds = Cluster.Bounds;
+				const FHairStrandsMacroGroupData& MacroGroup = MacroGroupDatas.Datas[MacroGroupIt];
+				const FBoxSphereBounds MacroGroupBounds = MacroGroup.Bounds;
 
 				for (const FLightSceneInfo* LightInfo : VisibleLightsPerView[ViewIndex])
 				{
 					const FLightSceneProxy* LightProxy = LightInfo->Proxy;
-					if (!LightProxy->AffectsBounds(ClusterBounds))
+					if (!LightProxy->AffectsBounds(MacroGroupBounds))
 					{
 						continue;
 					}
@@ -167,28 +102,15 @@ FHairStrandsDeepShadowViews RenderHairStrandsDeepShadows(
 	const FIntPoint AtlasSlotDimension(AtlasSlotX, AtlasSlotX == DOMSlotCount ? 1 : AtlasSlotX);
 	const FIntPoint AtlasSlotResolution(GDeepShadowResolution, GDeepShadowResolution);
 	const FIntPoint AtlasResolution(AtlasSlotResolution.X * AtlasSlotDimension.X, AtlasSlotResolution.Y * AtlasSlotDimension.Y);
-	TRefCountPtr<IPooledRenderTarget> FrontDepthAtlasTexture;
-	TRefCountPtr<IPooledRenderTarget> DeepShadowLayersAtlasTexture;
+
+	FRDGBuilder GraphBuilder(RHICmdList);
 
 	// Create Atlas resources for DOM
 	bool bClearFrontDepthAtlasTexture = true;
 	bool bClearLayerAtlasTexture = true;
-	{
-		{
-			FPooledRenderTargetDesc ShadowDesc(FPooledRenderTargetDesc::Create2DDesc(AtlasResolution, PF_DepthStencil, FClearValueBinding::DepthFar, TexCreate_None, TexCreate_DepthStencilTargetable | TexCreate_ShaderResource, false));
-			GRenderTargetPool.FindFreeElement(RHICmdList, ShadowDesc, FrontDepthAtlasTexture, TEXT("ShadowDepth"));
-		}
+	FRDGTextureRef FrontDepthAtlasTexture = GraphBuilder.CreateTexture(FPooledRenderTargetDesc::Create2DDesc(AtlasResolution, PF_DepthStencil, FClearValueBinding::DepthFar, TexCreate_None, TexCreate_DepthStencilTargetable | TexCreate_ShaderResource, false), TEXT("ShadowDepth"));
+	FRDGTextureRef DeepShadowLayersAtlasTexture = GraphBuilder.CreateTexture(FPooledRenderTargetDesc::Create2DDesc(AtlasResolution, PF_FloatRGBA, FClearValueBinding::Transparent, TexCreate_None, TexCreate_RenderTargetable | TexCreate_ShaderResource, false), TEXT("DeepShadowLayers"));
 
-		{
-			FPooledRenderTargetDesc Desc(FPooledRenderTargetDesc::Create2DDesc(AtlasResolution, PF_FloatRGBA, FClearValueBinding::Transparent, TexCreate_None, TexCreate_RenderTargetable | TexCreate_ShaderResource, false));
-			GRenderTargetPool.FindFreeElement(RHICmdList, Desc, DeepShadowLayersAtlasTexture, TEXT("DeepShadowLayers"));
-		}
-	}
-
-	// @todo_hair: could we share the rendering of the deep shadow across all view? Is it needed?
-	// @todo_hair: need to clarify what how view are/should be managed. At the moment this is 
-	// confusing as it seems there are many views but in the renderLight code, only one shadow mask 
-	// texture is created for all lights.
 	uint32 AtlasSlotIndex = 0;
 	for (int32 ViewIndex = 0; ViewIndex < Views.Num(); ViewIndex++)
 	{
@@ -197,18 +119,18 @@ FHairStrandsDeepShadowViews RenderHairStrandsDeepShadows(
 		{
 			FHairStrandsDeepShadowDatas& DeepShadowDatas = DeepShadowsPerView.Views[ViewIndex];
 
-			const FHairStrandsClusterDatas& DeepShadowClusterDatas = DeepShadowClusterViews.Views[ViewIndex];
-			for (int32 ClusterIt = 0; ClusterIt < DeepShadowClusterDatas.Datas.Num(); ++ClusterIt)
+			const FHairStrandsMacroGroupDatas& MacroGroupDatas = MacroGroupsViews.Views[ViewIndex];
+			for (int32 MacroGroupIt = 0; MacroGroupIt < MacroGroupDatas.Datas.Num(); ++MacroGroupIt)
 			{
 				// List of all the light in the scene.
 				uint32 LightLocalIndex = 0;
 				for (const FLightSceneInfo* LightInfo : VisibleLightsPerView[ViewIndex])
 				{
-					const FHairStrandsClusterData& Cluster = DeepShadowClusterDatas.Datas[ClusterIt];
-					FBoxSphereBounds ClusterBounds = Cluster.Bounds;
+					const FHairStrandsMacroGroupData& MacroGroup = MacroGroupDatas.Datas[MacroGroupIt];
+					FBoxSphereBounds MacroGroupBounds = MacroGroup.Bounds;
 
 					const FLightSceneProxy* LightProxy = LightInfo->Proxy;
-					if (!LightProxy->AffectsBounds(ClusterBounds))
+					if (!LightProxy->AffectsBounds(MacroGroupBounds))
 					{
 						continue;
 					}
@@ -222,71 +144,74 @@ FHairStrandsDeepShadowViews RenderHairStrandsDeepShadows(
 
 					// Note: LightPosition.W is used in the transmittance mask shader to differentiate between directional and local lights.
 					FHairStrandsDeepShadowData& DomData = DeepShadowDatas.Datas.AddZeroed_GetRef();
-					ComputeWorldToLightClip(DomData.WorldToLightTransform, MinStrandRadiusAtDepth1, ClusterBounds, *LightProxy, LightType, AtlasSlotResolution);
+					ComputeWorldToLightClip(DomData.WorldToLightTransform, MinStrandRadiusAtDepth1, MacroGroupBounds, *LightProxy, LightType, AtlasSlotResolution);
 					DomData.LightDirection = LightProxy->GetDirection();
 					DomData.LightPosition = FVector4(LightProxy->GetPosition(), LightType == ELightComponentType::LightType_Directional ? 0 : 1);
 					DomData.LightLuminance = LightProxy->GetColor();
 					DomData.LightType = LightType;
 					DomData.LightId = LightInfo->Id;
 					DomData.ShadowResolution = AtlasSlotResolution;
-					DomData.Bounds = ClusterBounds;
+					DomData.Bounds = MacroGroupBounds;
 					DomData.AtlasRect = FIntRect(AtlasRectOffset, AtlasRectOffset + AtlasSlotResolution);
-					DomData.ClusterId = Cluster.ClusterId;
-					DomData.DepthTexture = FrontDepthAtlasTexture;
-					DomData.LayersTexture = DeepShadowLayersAtlasTexture;
+					DomData.MacroGroupId = MacroGroup.MacroGroupId;					
 					AtlasSlotIndex++;
 
-					TUniformBufferRef<FDeepShadowPassUniformParameters> DeepShadowPassUniformParameters;
+					const bool bIsOrtho = LightType == ELightComponentType::LightType_Directional;
+					const FVector4 HairRenderInfo = PackHairRenderInfo(MinStrandRadiusAtDepth1.Primary, MinStrandRadiusAtDepth1.Primary, 1, bIsOrtho, false);
+					
+					// Front depth
 					{
-						FDeepShadowPassUniformParameters PassParameters;
-						PassParameters.WorldToClipMatrix = DomData.WorldToLightTransform;;
-						PassParameters.SliceValue = FVector4(1, 1, 1, 1);
-						PassParameters.FrontDepthTexture = DomData.DepthTexture->GetRenderTargetItem().TargetableTexture.GetReference();
-						PassParameters.AtlasRect = DomData.AtlasRect;
-						PassParameters.VoxelMinAABB = FVector::ZeroVector;
-						PassParameters.VoxelMaxAABB = FVector::ZeroVector;
-						PassParameters.VoxelResolution = 0;
-						DeepShadowPassUniformParameters = CreateUniformBufferImmediate(PassParameters, UniformBuffer_SingleFrame, EUniformBufferValidation::None);
+						DECLARE_GPU_STAT(HairStrandsDeepShadowFrontDepth);
+						SCOPED_DRAW_EVENT(GraphBuilder.RHICmdList, HairStrandsDeepShadowFrontDepth);
+						SCOPED_GPU_STAT(GraphBuilder.RHICmdList, HairStrandsDeepShadowFrontDepth);
+
+						FHairDeepShadowRasterPassParameters* PassParameters = GraphBuilder.AllocParameters<FHairDeepShadowRasterPassParameters>();
+						PassParameters->WorldToClipMatrix = DomData.WorldToLightTransform;;
+						PassParameters->SliceValue = FVector4(1, 1, 1, 1);
+						PassParameters->AtlasRect = DomData.AtlasRect;
+						PassParameters->ViewportResolution = AtlasSlotResolution;
+						PassParameters->RenderTargets.DepthStencil = FDepthStencilBinding(FrontDepthAtlasTexture, bClearFrontDepthAtlasTexture ? ERenderTargetLoadAction::EClear : ERenderTargetLoadAction::ELoad, ERenderTargetLoadAction::ENoAction, FExclusiveDepthStencil::DepthWrite_StencilNop);
+
+						AddHairDeepShadowRasterPass(
+							GraphBuilder,
+							Scene,
+							&ViewInfo,
+							MacroGroup.PrimitivesInfos,
+							EHairStrandsRasterPassType::FrontDepth,
+							DomData.AtlasRect,
+							HairRenderInfo,
+							DomData.LightDirection,
+							PassParameters);
 					}
 
-					TUniformBufferRef<FViewUniformShaderParameters> ViewUniformShaderParameters;
+					// Deep layers
 					{
-						const float bIsOrtho = LightType == ELightComponentType::LightType_Directional;
-
-						// Save some status we need to restore
-						const FVector SavedViewForward = ViewInfo.CachedViewUniformShaderParameters->ViewForward;
-						// Update our view parameters
-						ViewInfo.CachedViewUniformShaderParameters->HairRenderInfo.X = MinStrandRadiusAtDepth1.Primary;
-						ViewInfo.CachedViewUniformShaderParameters->HairRenderInfo.Y = MinStrandRadiusAtDepth1.Primary;
-						ViewInfo.CachedViewUniformShaderParameters->HairRenderInfo.Z = bIsOrtho;
-						ViewInfo.CachedViewUniformShaderParameters->ViewForward = DomData.LightDirection;
-						// Create the uniform buffer
-						ViewUniformShaderParameters = TUniformBufferRef<FViewUniformShaderParameters>::CreateUniformBufferImmediate(*ViewInfo.CachedViewUniformShaderParameters, UniformBuffer_SingleFrame);
-						// Restore the view cached parameters
-						ViewInfo.CachedViewUniformShaderParameters->ViewForward = SavedViewForward;
+						DECLARE_GPU_STAT(HairStrandsDeepShadowLayers);
+						SCOPED_DRAW_EVENT(GraphBuilder.RHICmdList, HairStrandsDeepShadowLayers);
+						SCOPED_GPU_STAT(GraphBuilder.RHICmdList, HairStrandsDeepShadowLayers);
+
+						FHairDeepShadowRasterPassParameters* PassParameters = GraphBuilder.AllocParameters<FHairDeepShadowRasterPassParameters>();
+						PassParameters->WorldToClipMatrix = DomData.WorldToLightTransform;;
+						PassParameters->SliceValue = FVector4(1, 1, 1, 1);
+						PassParameters->AtlasRect = DomData.AtlasRect;
+						PassParameters->ViewportResolution = AtlasSlotResolution;
+						PassParameters->FrontDepthTexture = FrontDepthAtlasTexture;
+						PassParameters->RenderTargets[0] = FRenderTargetBinding(DeepShadowLayersAtlasTexture, bClearLayerAtlasTexture ? ERenderTargetLoadAction::EClear : ERenderTargetLoadAction::ELoad, 0);
+
+						AddHairDeepShadowRasterPass(
+							GraphBuilder,
+							Scene,
+							&ViewInfo,
+							MacroGroup.PrimitivesInfos,
+							EHairStrandsRasterPassType::DeepOpacityMap,
+							DomData.AtlasRect,
+							HairRenderInfo,
+							DomData.LightDirection,
+							PassParameters);						
 					}
 
-					RenderDeepShadowFrontDepth(
-						RHICmdList,
-						Scene,
-						&ViewInfo,
-						Cluster.PrimitivesInfos,
-						DomData.AtlasRect,
-						ViewUniformShaderParameters,
-						DeepShadowPassUniformParameters,
-						bClearFrontDepthAtlasTexture,
-						DomData.DepthTexture);
-
-					RenderDeepShadowLayers(
-						RHICmdList,
-						Scene,
-						&ViewInfo,
-						Cluster.PrimitivesInfos,
-						DomData.AtlasRect,
-						ViewUniformShaderParameters,
-						DeepShadowPassUniformParameters,
-						bClearLayerAtlasTexture,
-						DomData.LayersTexture);
+					GraphBuilder.QueueTextureExtraction(FrontDepthAtlasTexture, &DomData.DepthTexture);
+					GraphBuilder.QueueTextureExtraction(DeepShadowLayersAtlasTexture, &DomData.LayersTexture);
 
 					bClearFrontDepthAtlasTexture = false;
 					bClearLayerAtlasTexture = false;
@@ -294,6 +219,8 @@ FHairStrandsDeepShadowViews RenderHairStrandsDeepShadows(
 					++LightLocalIndex;
 				}
 			}
+
+			GraphBuilder.Execute();
 		}
 	}
 
diff --git a/Engine/Source/Runtime/Renderer/Private/HairStrands/HairStrandsDeepShadow.h b/Engine/Source/Runtime/Renderer/Private/HairStrands/HairStrandsDeepShadow.h
--- a/Engine/Source/Runtime/Renderer/Private/HairStrands/HairStrandsDeepShadow.h
+++ b/Engine/Source/Runtime/Renderer/Private/HairStrands/HairStrandsDeepShadow.h
@@ -13,13 +13,13 @@
 /// Hold deep shadow information for a given light.
 struct FHairStrandsDeepShadowData
 {
-	static const uint32 MaxClusterCount = 16u;
+	static const uint32 MaxMacroGroupCount = 16u;
 
 	TRefCountPtr<IPooledRenderTarget> DepthTexture;
 	TRefCountPtr<IPooledRenderTarget> LayersTexture;
 	FMatrix WorldToLightTransform;
 	FIntRect AtlasRect;
-	uint32 ClusterId = ~0;
+	uint32 MacroGroupId = ~0;
 
 	FIntPoint ShadowResolution = FIntPoint::ZeroValue;
 	uint32 LightId = ~0;
@@ -47,4 +47,4 @@ FHairStrandsDeepShadowViews RenderHairStrandsDeepShadows(
 	FRHICommandListImmediate& RHICmdList,
 	const class FScene* Scene,
 	const TArray<FViewInfo>& Views,
-	const struct FHairStrandsClusterViews& ClusterViews);
\ No newline at end of file
+	const struct FHairStrandsMacroGroupViews& MacroGroupsViews);
\ No newline at end of file
diff --git a/Engine/Source/Runtime/Renderer/Private/HairStrands/HairStrandsEnvironment.cpp b/Engine/Source/Runtime/Renderer/Private/HairStrands/HairStrandsEnvironment.cpp
--- a/Engine/Source/Runtime/Renderer/Private/HairStrands/HairStrandsEnvironment.cpp
+++ b/Engine/Source/Runtime/Renderer/Private/HairStrands/HairStrandsEnvironment.cpp
@@ -18,6 +18,7 @@
 #include "ScenePrivate.h"
 #include "RenderGraphEvent.h"
 #include "PostProcess/PostProcessing.h"
+#include "GpuDebugRendering.h"
 
 ///////////////////////////////////////////////////////////////////////////////////////////////////
 
@@ -27,9 +28,15 @@ static FAutoConsoleVariableRef CVarHairSkylightingEnable(TEXT("r.HairStrands.Sky
 static int32 GHairSkyAOEnable = 1;
 static FAutoConsoleVariableRef CVarHairSkyAOEnable(TEXT("r.HairStrands.SkyAO"), GHairSkyAOEnable, TEXT("Enable (sky) AO on hair."));
 
+static int32 GHairSkyAOVirtualVoxel = 4;
+static FAutoConsoleVariableRef CVarHairSkyAOVirtualVoxel(TEXT("r.HairStrands.SkyAO.VirtualVoxel"), GHairSkyAOVirtualVoxel, TEXT("Use virtual voxel for visibility/hair tracing/counting if available."));
+
 static float GHairSkylightingConeAngle = 3;
 static FAutoConsoleVariableRef CVarHairSkylightingConeAngle(TEXT("r.HairStrands.SkyLighting.ConeAngle"), GHairSkylightingConeAngle, TEXT("Cone angle for tracing sky lighting on hair."));
 
+static int GHairSkylightingVirtualVoxel = 4;
+static FAutoConsoleVariableRef CVarHairSkylightingVirtualVoxel(TEXT("r.HairStrands.SkyLighting.VirtualVoxel"), GHairSkylightingVirtualVoxel, TEXT("Use virtual voxel for visibility/hair tracing/counting if available."));
+
 static int32 GHairStrandsSkyLightingSampleCount = 16;
 static FAutoConsoleVariableRef CVarHairStrandsSkyLightingSampleCount(TEXT("r.HairStrands.SkyLighting.SampleCount"), GHairStrandsSkyLightingSampleCount, TEXT("Number of samples used for evaluating multiple scattering and visible area (default is set to 16)."), ECVF_Scalability | ECVF_RenderThreadSafe);
 
@@ -58,7 +65,8 @@ class FHairEnvironmentAO : public FGlobalShader
 	SHADER_USE_PARAMETER_STRUCT(FHairEnvironmentAO, FGlobalShader)
 
 	class FSampleSet : SHADER_PERMUTATION_INT("PERMUTATION_SAMPLESET", 2);
-	using FPermutationDomain = TShaderPermutationDomain<FSampleSet>;
+	class FVirtualVoxel : SHADER_PERMUTATION_INT("VOXEL_TRAVERSAL_TYPE", 5);
+	using FPermutationDomain = TShaderPermutationDomain<FSampleSet, FVirtualVoxel>;
 
 	static bool ShouldCompilePermutation(const FGlobalShaderPermutationParameters& Parameters)
 	{
@@ -68,7 +76,7 @@ class FHairEnvironmentAO : public FGlobalShader
 	BEGIN_SHADER_PARAMETER_STRUCT(FParameters, )
 
 		SHADER_PARAMETER(FVector, Voxel_MinAABB)
-		SHADER_PARAMETER(uint32, Voxel_ClusterId)
+		SHADER_PARAMETER(uint32, Voxel_MacroGroupId)
 		SHADER_PARAMETER(FVector, Voxel_MaxAABB)
 		SHADER_PARAMETER(uint32, Voxel_Resolution)
 		SHADER_PARAMETER(float, Voxel_DensityScale)
@@ -86,6 +94,10 @@ class FHairEnvironmentAO : public FGlobalShader
 
 		SHADER_PARAMETER_STRUCT_REF(FViewUniformShaderParameters, ViewUniformBuffer)
 
+		SHADER_PARAMETER_STRUCT_REF(FVirtualVoxelParameters, VirtualVoxel)
+
+		SHADER_PARAMETER_STRUCT_INCLUDE(ShaderDrawDebug::FShaderDrawDebugParameters, ShaderDrawParameters)
+
 		RENDER_TARGET_BINDING_SLOTS()
 	END_SHADER_PARAMETER_STRUCT()
 };
@@ -96,24 +108,31 @@ static void AddHairStrandsEnvironmentAOPass(
 	FRDGBuilder& GraphBuilder,
 	const FViewInfo& View,
 	const FHairStrandsVisibilityData& VisibilityData,
-	const FHairStrandsClusterData& ClusterData,
+	const FHairStrandsMacroGroupDatas& MacroGroupDatas,
+	const FHairStrandsMacroGroupData& MacroGroupData,
 	FRDGTextureRef Output)
 {
 	check(Output);
+	const bool bUseVirtualVoxel = MacroGroupDatas.VirtualVoxelResources.IsValid() && GHairSkyAOVirtualVoxel > 0;
+	const int32 VoxelTraversalType = FMath::Clamp(GHairSkyAOVirtualVoxel, 0, 4);
 
 	FSceneTextureParameters SceneTextures;
 	SetupSceneTextureParameters(GraphBuilder, &SceneTextures);
 
 	FHairEnvironmentAO::FParameters* PassParameters = GraphBuilder.AllocParameters<FHairEnvironmentAO::FParameters>();
-	PassParameters->Voxel_ClusterId = ClusterData.ClusterId;
-	PassParameters->Voxel_MinAABB = ClusterData.GetMinBound();
-	PassParameters->Voxel_MaxAABB = ClusterData.GetMaxBound();
-	PassParameters->Voxel_Resolution = ClusterData.GetResolution();
-	PassParameters->Voxel_DensityTexture = GraphBuilder.RegisterExternalTexture(ClusterData.VoxelResources.DensityTexture);
+	PassParameters->Voxel_MacroGroupId = MacroGroupData.MacroGroupId;
+	if (!bUseVirtualVoxel)
+	{
+		PassParameters->Voxel_MinAABB = MacroGroupData.GetMinBound();
+		PassParameters->Voxel_MaxAABB = MacroGroupData.GetMaxBound();
+		PassParameters->Voxel_Resolution = MacroGroupData.GetResolution();
+		PassParameters->Voxel_DensityTexture = GraphBuilder.RegisterExternalTexture(MacroGroupData.VoxelResources.DensityTexture);
+	}
 	PassParameters->Voxel_DensityScale = GetHairStrandsVoxelizationDensityScale();
 	PassParameters->Voxel_DepthBiasScale = GetHairStrandsVoxelizationDepthBiasScale();
 	PassParameters->Voxel_TanConeAngle = FMath::Tan(FMath::DegreesToRadians(GetHairStrandsSkyLightingConeAngle()));
 	PassParameters->SceneTextures = SceneTextures;
+	PassParameters->VirtualVoxel = MacroGroupDatas.VirtualVoxelResources.UniformBuffer;
 	SetupSceneTextureSamplers(&PassParameters->SceneTextureSamplers);
 
 	PassParameters->ViewUniformBuffer = View.ViewUniformBuffer;
@@ -125,8 +144,14 @@ static void AddHairStrandsEnvironmentAOPass(
 	PassParameters->AO_DistanceThreshold = FMath::Max(GHairStrandsSkyAODistanceThreshold, 1.f);
 	PassParameters->RenderTargets[0] = FRenderTargetBinding(Output, ERenderTargetLoadAction::ELoad);
 
+	if (ShaderDrawDebug::IsShaderDrawDebugEnabled(View))
+	{
+		ShaderDrawDebug::SetParameters(GraphBuilder, View.ShaderDrawData, PassParameters->ShaderDrawParameters);
+	}
+
 	FHairEnvironmentAO::FPermutationDomain PermutationVector;
 	PermutationVector.Set<FHairEnvironmentAO::FSampleSet>(PassParameters->AO_SampleCount <= 16 ? 0 : 1);
+	PermutationVector.Set<FHairEnvironmentAO::FVirtualVoxel>(bUseVirtualVoxel ? VoxelTraversalType : 0);
 
 	TShaderMapRef<FHairEnvironmentAO> PixelShader(View.ShaderMap, PermutationVector);
 	ClearUnusedGraphResources(*PixelShader, PassParameters);
@@ -172,7 +197,7 @@ static void AddHairEnvironmentLightingComposePass(
 	FRDGBuilder& GraphBuilder,
 	const FViewInfo& View,
 	const FHairStrandsVisibilityData& VisibilityData,
-	const FHairStrandsClusterData& ClusterData,
+	const FHairStrandsMacroGroupData& MacroGroupData,
 	const FRDGBufferRef SampleLightingBuffer,
 	FRDGTextureRef& OutColorTexture,
 	FRDGTextureRef& OutSubColorTexture)
@@ -238,7 +263,8 @@ class FHairEnvironmentLightingCS : public FGlobalShader
 
 	class FGroupSize : SHADER_PERMUTATION_INT("PERMUTATION_GROUP_SIZE", 2);
 	class FSampleSet : SHADER_PERMUTATION_INT("PERMUTATION_SAMPLESET", 2);
-	using FPermutationDomain = TShaderPermutationDomain<FGroupSize, FSampleSet>;
+	class FVirtualVoxel : SHADER_PERMUTATION_INT("VOXEL_TRAVERSAL_TYPE", 5);
+	using FPermutationDomain = TShaderPermutationDomain<FGroupSize, FSampleSet, FVirtualVoxel>;
 
 	static void ModifyCompilationEnvironment(const FGlobalShaderPermutationParameters& Parameters, FShaderCompilerEnvironment& OutEnvironment)
 	{
@@ -253,7 +279,7 @@ class FHairEnvironmentLightingCS : public FGlobalShader
 	BEGIN_SHADER_PARAMETER_STRUCT(FParameters, )
 
 		SHADER_PARAMETER(FVector, Voxel_MinAABB)
-		SHADER_PARAMETER(uint32, Voxel_ClusterId)
+		SHADER_PARAMETER(uint32, Voxel_MacroGroupId)
 		SHADER_PARAMETER(FVector, Voxel_MaxAABB)
 		SHADER_PARAMETER(uint32, Voxel_Resolution)
 		SHADER_PARAMETER(float, Voxel_DensityScale)
@@ -263,6 +289,7 @@ class FHairEnvironmentLightingCS : public FGlobalShader
 		SHADER_PARAMETER(uint32, MaxVisibilityNodeCount)
 		SHADER_PARAMETER(uint32, MultipleScatterSampleCount)
 		SHADER_PARAMETER(uint32, HairComponents)
+		SHADER_PARAMETER(float,  HairDualScatteringRoughnessOverride)
 		SHADER_PARAMETER(uint32, JitterSphericalIntegration)
 		SHADER_PARAMETER(float, HairDistanceThreshold)
 		SHADER_PARAMETER(uint32, bHairUseViewHairCount)
@@ -290,6 +317,10 @@ class FHairEnvironmentLightingCS : public FGlobalShader
 		SHADER_PARAMETER_STRUCT_REF(FReflectionCaptureShaderData, ReflectionCaptureData)
 		SHADER_PARAMETER_STRUCT_REF(FForwardLightData, ForwardLightData)
 
+		SHADER_PARAMETER_STRUCT_REF(FVirtualVoxelParameters, VirtualVoxel)
+
+		SHADER_PARAMETER_STRUCT_INCLUDE(ShaderDrawDebug::FShaderDrawDebugParameters, ShaderDrawParameters)
+
 		END_SHADER_PARAMETER_STRUCT()
 };
 
@@ -299,7 +330,8 @@ static FRDGBufferRef AddHairStrandsEnvironmentLightingPassCS(
 	FRDGBuilder& GraphBuilder,
 	const FViewInfo& View,
 	const FHairStrandsVisibilityData& VisibilityData,
-	const FHairStrandsClusterData& ClusterData,
+	const FHairStrandsMacroGroupDatas& MacroGroupDatas,
+	const FHairStrandsMacroGroupData& MacroGroupData,
 	const uint32 NodeGroupSize,
 	FRDGBufferRef IndirectArgsBuffer)
 {
@@ -311,6 +343,9 @@ static FRDGBufferRef AddHairStrandsEnvironmentLightingPassCS(
 	FSceneTextureParameters SceneTextures;
 	SetupSceneTextureParameters(GraphBuilder, &SceneTextures);
 
+	const bool bUseVirtualVoxel = MacroGroupDatas.VirtualVoxelResources.IsValid();
+	const int32 VoxelTraversalType = FMath::Clamp(GHairSkylightingVirtualVoxel, 0, 4);
+
 	// Render the reflection environment with tiled deferred culling
 	const bool bHasBoxCaptures = (View.NumBoxReflectionCaptures > 0);
 	const bool bHasSphereCaptures = (View.NumSphereReflectionCaptures > 0);
@@ -320,11 +355,14 @@ static FRDGBufferRef AddHairStrandsEnvironmentLightingPassCS(
 	PassParameters->HairEnergyLUTTexture = GraphBuilder.RegisterExternalTexture(InHairLUT.Textures[HairLUTType_MeanEnergy], TEXT("HairMeanEnergyLUTTexture"));;
 	PassParameters->HairScatteringLUTTexture = GraphBuilder.RegisterExternalTexture(InHairLUT.Textures[HairLUTType_DualScattering], TEXT("HairScatteringEnergyLUTTexture"));
 	PassParameters->HairLUTSampler = TStaticSamplerState<SF_Bilinear, AM_Clamp, AM_Clamp, AM_Clamp>::GetRHI();
-	PassParameters->Voxel_ClusterId = ClusterData.ClusterId;
-	PassParameters->Voxel_MinAABB = ClusterData.GetMinBound();
-	PassParameters->Voxel_MaxAABB = ClusterData.GetMaxBound();
-	PassParameters->Voxel_Resolution = ClusterData.GetResolution();
-	PassParameters->Voxel_DensityTexture = GraphBuilder.RegisterExternalTexture(ClusterData.VoxelResources.DensityTexture);
+	PassParameters->Voxel_MacroGroupId = MacroGroupData.MacroGroupId;
+	if (!bUseVirtualVoxel)
+	{
+		PassParameters->Voxel_MinAABB = MacroGroupData.GetMinBound();
+		PassParameters->Voxel_MaxAABB = MacroGroupData.GetMaxBound();
+		PassParameters->Voxel_Resolution = MacroGroupData.GetResolution();
+		PassParameters->Voxel_DensityTexture = GraphBuilder.RegisterExternalTexture(MacroGroupData.VoxelResources.DensityTexture);
+	}
 	PassParameters->Voxel_DensityScale = GetHairStrandsVoxelizationDensityScale();
 	PassParameters->Voxel_DepthBiasScale = GetHairStrandsVoxelizationDepthBiasScale();
 	PassParameters->Voxel_TanConeAngle = FMath::Tan(FMath::DegreesToRadians(GetHairStrandsSkyLightingConeAngle()));
@@ -334,10 +372,12 @@ static FRDGBufferRef AddHairStrandsEnvironmentLightingPassCS(
 	PassParameters->MultipleScatterSampleCount = FMath::Max(uint32(GHairStrandsSkyLightingSampleCount), 1u);
 	PassParameters->JitterSphericalIntegration = GHairStrandsSkyLightingJitterSphericalIntegration ? 1 : 0;
 	PassParameters->HairComponents = ToBitfield(GetHairComponents());
+	PassParameters->HairDualScatteringRoughnessOverride = GetHairDualScatteringRoughnessOverride();
 	PassParameters->PreIntegratedGF = GSystemTextures.PreintegratedGF->GetRenderTargetItem().ShaderResourceTexture;
 	PassParameters->PreIntegratedGFSampler = TStaticSamplerState<SF_Bilinear, AM_Clamp, AM_Clamp, AM_Clamp>::GetRHI();
 	PassParameters->HairCountTexture = GraphBuilder.RegisterExternalTexture(VisibilityData.ViewHairCountTexture ? VisibilityData.ViewHairCountTexture : GSystemTextures.BlackDummy);
 	PassParameters->SceneTextures = SceneTextures;
+	PassParameters->VirtualVoxel = MacroGroupDatas.VirtualVoxelResources.UniformBuffer;
 	SetupSceneTextureSamplers(&PassParameters->SceneTextureSamplers);
 	PassParameters->ViewUniformBuffer = View.ViewUniformBuffer;
 	PassParameters->ReflectionCaptureData = View.ReflectionCaptureUniformBuffer;
@@ -350,6 +390,11 @@ static FRDGBufferRef AddHairStrandsEnvironmentLightingPassCS(
 	PassParameters->IndirectArgsBuffer = IndirectArgsBuffer;
 	PassParameters->OutLightingBuffer = GraphBuilder.CreateUAV(FRDGBufferUAVDesc(OutBuffer, PF_FloatRGBA));
 
+	if (ShaderDrawDebug::IsShaderDrawDebugEnabled(View))
+	{
+		ShaderDrawDebug::SetParameters(GraphBuilder, View.ShaderDrawData, PassParameters->ShaderDrawParameters);
+	}
+
 	// Bind hair data
 	FRDGBufferRef InHairVisibilityNodeData = GraphBuilder.RegisterExternalBuffer(VisibilityData.NodeData, TEXT("HairVisibilityNodeData"));
 	FRDGBufferRef InHairVisibilityNodeCoord = GraphBuilder.RegisterExternalBuffer(VisibilityData.NodeCoord, TEXT("HairVisibilityNodeCoord"));
@@ -360,6 +405,7 @@ static FRDGBufferRef AddHairStrandsEnvironmentLightingPassCS(
 	FHairEnvironmentLightingCS::FPermutationDomain PermutationVector;
 	PermutationVector.Set<FHairEnvironmentLightingCS::FGroupSize>(NodeGroupSize == 64 ? 0 : (NodeGroupSize == 32 ? 1 : 2));
 	PermutationVector.Set<FHairEnvironmentLightingCS::FSampleSet>(PassParameters->MultipleScatterSampleCount <= 16 ? 0 : 1);
+	PermutationVector.Set<FHairEnvironmentLightingCS::FVirtualVoxel>(bUseVirtualVoxel ? VoxelTraversalType : 0);
 
 	TShaderMapRef<FHairEnvironmentLightingCS> ComputeShader(View.ShaderMap, PermutationVector);
 	FComputeShaderUtils::AddPass(
@@ -398,10 +444,11 @@ void RenderHairStrandsEnvironmentLighting(
 	FRDGBufferRef NodeIndirectArgBuffer = GraphBuilder.RegisterExternalBuffer(VisibilityData.NodeIndirectArg, TEXT("HairNodeIndirectArgBuffer"));
 
 	const FViewInfo& View = Views[ViewIndex];
-	for (const FHairStrandsClusterData& ClusterData : HairDatas->HairClusterPerViews.Views[ViewIndex].Datas)
+	const FHairStrandsMacroGroupDatas& MacroGroupDatas = HairDatas->MacroGroupsPerViews.Views[ViewIndex];
+	for (const FHairStrandsMacroGroupData& MacroGroupData : HairDatas->MacroGroupsPerViews.Views[ViewIndex].Datas)
 	{
-		FRDGBufferRef SampleLightingBuffer = AddHairStrandsEnvironmentLightingPassCS(GraphBuilder, View, VisibilityData, ClusterData, VisibilityData.NodeGroupSize, NodeIndirectArgBuffer);
-		AddHairEnvironmentLightingComposePass(GraphBuilder, View, VisibilityData, ClusterData, SampleLightingBuffer, SceneColorTexture, SceneColorSubPixelTexture);
+		FRDGBufferRef SampleLightingBuffer = AddHairStrandsEnvironmentLightingPassCS(GraphBuilder, View, VisibilityData, MacroGroupDatas, MacroGroupData, VisibilityData.NodeGroupSize, NodeIndirectArgBuffer);
+		AddHairEnvironmentLightingComposePass(GraphBuilder, View, VisibilityData, MacroGroupData, SampleLightingBuffer, SceneColorTexture, SceneColorSubPixelTexture);
 	}
 }
 
@@ -424,15 +471,16 @@ void RenderHairStrandsAmbientOcclusion(
 			continue;
 		}
 
-		if (ViewIndex > uint32(HairDatas->HairClusterPerViews.Views.Num()))
+		if (ViewIndex > uint32(HairDatas->MacroGroupsPerViews.Views.Num()))
 			continue;
 
 		const FViewInfo& View = Views[ViewIndex];
+		const FHairStrandsMacroGroupDatas& MacroGroupDatas = HairDatas->MacroGroupsPerViews.Views[ViewIndex];
 		FRDGBuilder GraphBuilder(RHICmdList);
 		FRDGTextureRef AOTexture = InAOTexture ? GraphBuilder.RegisterExternalTexture(InAOTexture, TEXT("AOTexture")) : nullptr;
-		for (const FHairStrandsClusterData& ClusterData : HairDatas->HairClusterPerViews.Views[ViewIndex].Datas)
+		for (const FHairStrandsMacroGroupData& MacroGroupData : HairDatas->MacroGroupsPerViews.Views[ViewIndex].Datas)
 		{
-			AddHairStrandsEnvironmentAOPass(GraphBuilder, View, VisibilityData, ClusterData, AOTexture);
+			AddHairStrandsEnvironmentAOPass(GraphBuilder, View, VisibilityData, MacroGroupDatas, MacroGroupData, AOTexture);
 
 		}
 		GraphBuilder.Execute();
diff --git a/Engine/Source/Runtime/Renderer/Private/HairStrands/HairStrandsInterface.cpp b/Engine/Source/Runtime/Renderer/Private/HairStrands/HairStrandsInterface.cpp
--- a/Engine/Source/Runtime/Renderer/Private/HairStrands/HairStrandsInterface.cpp
+++ b/Engine/Source/Runtime/Renderer/Private/HairStrands/HairStrandsInterface.cpp
@@ -10,6 +10,7 @@
 
 #include "SkeletalRenderPublic.h"
 #include "SceneRendering.h"
+#include "SystemTextures.h"
 
 DEFINE_LOG_CATEGORY_STATIC(LogHairRendering, Log, All);
 
@@ -18,6 +19,42 @@ static FAutoConsoleVariableRef CVarHairStrandsRenderingEnable(TEXT("r.HairStrand
 
 ////////////////////////////////////////////////////////////////////////////////////////////////////////////////////////
 
+FHairGroupPublicData::FHairGroupPublicData(uint32 InGroupIndex, uint32 InGroupInstanceVertexCount, uint32 InClusterCount, uint32 InVertexCount)
+{
+	GroupIndex = InGroupIndex;
+	GroupInstanceVertexCount = InGroupInstanceVertexCount;
+	ClusterCount = InClusterCount;
+	VertexCount = InVertexCount;
+}
+
+void FHairGroupPublicData::InitRHI()
+{
+	DrawIndirectBuffer.Initialize(sizeof(uint32), 4, PF_R32_UINT, BUF_DrawIndirect);
+	uint32* BufferData = (uint32*)RHILockVertexBuffer(DrawIndirectBuffer.Buffer, 0, sizeof(uint32) * 4, RLM_WriteOnly);
+	BufferData[0] = GroupInstanceVertexCount;
+	BufferData[1] = 1;
+	BufferData[2] = 0;
+	BufferData[3] = 0;
+	RHIUnlockVertexBuffer(DrawIndirectBuffer.Buffer);
+
+	ClusterAABBBuffer.Initialize(sizeof(int32), ClusterCount * 6, EPixelFormat::PF_R32_SINT, BUF_Static);
+	GroupAABBBuffer.Initialize(sizeof(int32), 6, EPixelFormat::PF_R32_SINT, BUF_Static);
+
+	CulledVertexIdBuffer.Initialize(sizeof(int32), VertexCount, EPixelFormat::PF_R32_UINT, BUF_Static);
+	CulledVertexRadiusScaleBuffer.Initialize(sizeof(float), VertexCount, EPixelFormat::PF_R32_FLOAT, BUF_Static);
+}
+
+void FHairGroupPublicData::ReleaseRHI()
+{
+	DrawIndirectBuffer.Release();
+	ClusterAABBBuffer.Release();
+	GroupAABBBuffer.Release();
+	CulledVertexIdBuffer.Release();
+	CulledVertexRadiusScaleBuffer.Release();
+}
+
+////////////////////////////////////////////////////////////////////////////////////////////////////////////////////////
+
 inline FHairStrandsProjectionMeshData::Section ConvertMeshSection(const FCachedGeometrySection& In, const FTransform& InTransform)
 {
 	FHairStrandsProjectionMeshData::Section Out;
@@ -39,8 +76,6 @@ inline FHairStrandsProjectionMeshData::Section ConvertMeshSection(const FCachedG
 //  * Register
 //  * For each frame
 //		* Update
-//		* AddProjectionQuery (Opt)
-//		* Project (Opt)
 //		* Update triangles information for dynamic meshes
 //		* RunHairStrandsInterpolation (Interpolation callback)
 //  * UnRegister
@@ -53,6 +88,7 @@ struct FHairStrandsManager
 		uint64 Id = 0;
 		EWorldType::Type WorldType = EWorldType::None;
 		FHairStrandsDebugInfo DebugInfo;
+		FHairStrandsPrimitiveResources PrimitiveResources;
 		FHairStrandsInterpolationData InterpolationData;
 		FHairStrandsProjectionHairData RenProjectionHairDatas;
 		FHairStrandsProjectionHairData SimProjectionHairDatas;
@@ -63,36 +99,17 @@ struct FHairStrandsManager
 		int32 FrameLODIndex = -1;
 	};
 	TArray<Element> Elements;
-
-	struct ProjectionQuery
-	{
-		uint64 Id = 0;
-		EWorldType::Type WorldType = EWorldType::None;
-		int32 LODIndex = -1;
-		FVector RestPositionOffset = FVector::ZeroVector;
-		bool bProcessed = false;
-	};
-	TArray<ProjectionQuery> ProjectionsQueries;
 };
 
 FHairStrandsManager GHairManager;
 
-void AddHairStrandsProjectionQuery(
-	FRHICommandListImmediate& RHICmdList,
-	uint64 Id,
-	EWorldType::Type WorldType,
-	int32 LODIndex,
-	const FVector& RestPositionOffset)
-{
-	GHairManager.ProjectionsQueries.Add({ Id, WorldType, LODIndex, RestPositionOffset, false });
-}
-
 void RegisterHairStrands(
 	uint64 Id, 
 	EWorldType::Type WorldType,
 	const FHairStrandsInterpolationData& InterpolationData, 
 	const FHairStrandsProjectionHairData& RenProjectionDatas,
  	const FHairStrandsProjectionHairData& SimProjectionDatas,
+	const FHairStrandsPrimitiveResources& PrimitiveResources,
 	const FHairStrandsDebugInfo& DebugInfo)
 {
 	for (int32 Index = 0; Index < GHairManager.Elements.Num(); ++Index)
@@ -110,17 +127,10 @@ void RegisterHairStrands(
 	E.WorldType = WorldType;
 	E.InterpolationData = InterpolationData;
 	E.RenProjectionHairDatas = RenProjectionDatas;
-	for (FHairStrandsProjectionHairData::HairGroup& ProjectionData : E.RenProjectionHairDatas.HairGroups)
-	{
-		ProjectionData.LocalToWorld = FTransform::Identity;
-	}
 	E.SimProjectionHairDatas = SimProjectionDatas;
-	for (FHairStrandsProjectionHairData::HairGroup& ProjectionData : E.SimProjectionHairDatas.HairGroups)
-	{
-		ProjectionData.LocalToWorld = FTransform::Identity;
-	}
 	E.SkeletalLocalToWorld = FTransform::Identity;
 	E.SkeletalDeformedPositionOffset = FVector::ZeroVector;
+	E.PrimitiveResources = PrimitiveResources;
 	E.DebugInfo = DebugInfo;
 }
 
@@ -225,11 +235,43 @@ void UnregisterHairStrands(uint64 Id)
 	}
 }
 
+void RunProjection(
+	FRHICommandListImmediate& RHICmdList,
+	const FTransform& LocalToWorld,
+	const FVector& RestPositionOffset,
+	const FHairStrandsProjectionMeshData& MeshData,
+	FHairStrandsProjectionHairData& RenProjectionHairData,
+	FHairStrandsProjectionHairData& SimProjectionHairData)
+{
+	const ERHIFeatureLevel::Type FeatureLevel = GMaxRHIFeatureLevel;
+	TShaderMap<FGlobalShaderType>* ShaderMap = GetGlobalShaderMap(FeatureLevel);
+
+	auto Project = [&RHICmdList, ShaderMap, &MeshData, LocalToWorld, RestPositionOffset](FHairStrandsProjectionHairData& ProjectionHairData)
+	{
+		for (FHairStrandsProjectionHairData::HairGroup& HairGroup : ProjectionHairData.HairGroups)
+		{
+			for (FHairStrandsProjectionHairData::LODData& LODData : HairGroup.LODDatas)
+			{
+				const uint32 LODIndex = LODData.LODIndex;
+				*HairGroup.LODDatas[LODIndex].RestPositionOffset = RestPositionOffset;
+				HairGroup.LocalToWorld = LocalToWorld;
+				ProjectHairStrandsOntoMesh(RHICmdList, ShaderMap, LODIndex, MeshData, HairGroup);
+				UpdateHairStrandsMeshTriangles(RHICmdList, ShaderMap, LODIndex, HairStrandsTriangleType::RestPose, MeshData.LODs[LODIndex], HairGroup);
+			}
+		}
+	};
+
+	Project(RenProjectionHairData);
+	Project(SimProjectionHairData);
+}
+
 void RunHairStrandsInterpolation(
 	FRHICommandListImmediate& RHICmdList, 
 	EWorldType::Type WorldType, 
+	const FShaderDrawDebugData* ShaderDrawData,
 	TShaderMap<FGlobalShaderType>* ShaderMap, 
-	EHairStrandsInterpolationType Type)
+	EHairStrandsInterpolationType Type,
+	FHairStrandClusterData* ClusterData)
 {
 	check(IsInRenderingThread());
 
@@ -241,39 +283,7 @@ void RunHairStrandsInterpolation(
 
 		E.CachedGeometry = E.MeshObject ? E.MeshObject->GetCachedGeometry() : FCachedGeometry();
 	}
-
-	// Process projection queries
-	for (FHairStrandsManager::ProjectionQuery& Q : GHairManager.ProjectionsQueries)
-	{	
-		for (FHairStrandsManager::Element& E : GHairManager.Elements)
-		{
-			if (E.Id != Q.Id || E.WorldType != Q.WorldType || E.CachedGeometry.Sections.Num() == 0 || Q.bProcessed)
-				continue;
 	
-			FHairStrandsProjectionMeshData MeshData;
-			for (FCachedGeometrySection Section : E. CachedGeometry.Sections)
-			{
-				check(Section.LODIndex == Q.LODIndex);
-				MeshData.Sections.Add(ConvertMeshSection(Section, E.SkeletalLocalToWorld));
-			}
-
-			for (FHairStrandsProjectionHairData::HairGroup& ProjectionHairData : E.SimProjectionHairDatas.HairGroups)
-			{
-				ProjectionHairData.LODDatas[Q.LODIndex].RestPositionOffset = Q.RestPositionOffset;
-				ProjectHairStrandsOntoMesh(RHICmdList, ShaderMap, Q.LODIndex, MeshData, ProjectionHairData);
-				UpdateHairStrandsMeshTriangles(RHICmdList, ShaderMap, Q.LODIndex, HairStrandsTriangleType::RestPose, MeshData, ProjectionHairData);
-			}
-
-			for (FHairStrandsProjectionHairData::HairGroup& ProjectionHairData : E.RenProjectionHairDatas.HairGroups)
-			{
-				ProjectionHairData.LODDatas[Q.LODIndex].RestPositionOffset = Q.RestPositionOffset;
-				ProjectHairStrandsOntoMesh(RHICmdList, ShaderMap, Q.LODIndex, MeshData, ProjectionHairData);
-				UpdateHairStrandsMeshTriangles(RHICmdList, ShaderMap, Q.LODIndex, HairStrandsTriangleType::RestPose, MeshData, ProjectionHairData);
-			}
-			Q.bProcessed = true;
-		}
-	}
-
 	// Update dynamic mesh triangles
 	for (FHairStrandsManager::Element& E : GHairManager.Elements)
 	{
@@ -281,22 +291,22 @@ void RunHairStrandsInterpolation(
 		if (E.WorldType != WorldType || E.CachedGeometry.Sections.Num() == 0)
 			continue;
 	
-		FHairStrandsProjectionMeshData MeshData;
+		FHairStrandsProjectionMeshData::LOD MeshDataLOD;
 		for (FCachedGeometrySection Section : E.CachedGeometry.Sections)
 		{
 			// Ensure all mesh's sections have the same LOD index
 			if (E.FrameLODIndex < 0) E.FrameLODIndex = Section.LODIndex;
 			check(E.FrameLODIndex == Section.LODIndex);
-	
-			MeshData.Sections.Add(ConvertMeshSection(Section, E.SkeletalLocalToWorld));
+
+			MeshDataLOD.Sections.Add(ConvertMeshSection(Section, E.SkeletalLocalToWorld));
 		}
 
 		for (FHairStrandsProjectionHairData::HairGroup& ProjectionHairData : E.RenProjectionHairDatas.HairGroups)
 		{
 			if (EHairStrandsInterpolationType::RenderStrands == Type && 0 <= E.FrameLODIndex && E.FrameLODIndex < ProjectionHairData.LODDatas.Num() && ProjectionHairData.LODDatas[E.FrameLODIndex].bIsValid)
 			{
-				ProjectionHairData.LODDatas[E.FrameLODIndex].DeformedPositionOffset = E.SkeletalDeformedPositionOffset;
-				UpdateHairStrandsMeshTriangles(RHICmdList, ShaderMap, E.FrameLODIndex, HairStrandsTriangleType::DeformedPose, MeshData, ProjectionHairData);
+				*ProjectionHairData.LODDatas[E.FrameLODIndex].DeformedPositionOffset = E.SkeletalDeformedPositionOffset;
+				UpdateHairStrandsMeshTriangles(RHICmdList, ShaderMap, E.FrameLODIndex, HairStrandsTriangleType::DeformedPose, MeshDataLOD, ProjectionHairData);
 			}
 		}
 
@@ -304,8 +314,23 @@ void RunHairStrandsInterpolation(
 		{
 			if (EHairStrandsInterpolationType::SimulationStrands == Type && 0 <= E.FrameLODIndex && E.FrameLODIndex < ProjectionHairData.LODDatas.Num() && ProjectionHairData.LODDatas[E.FrameLODIndex].bIsValid)
 			{
-				ProjectionHairData.LODDatas[E.FrameLODIndex].DeformedPositionOffset = E.SkeletalDeformedPositionOffset;
-				UpdateHairStrandsMeshTriangles(RHICmdList, ShaderMap, E.FrameLODIndex, HairStrandsTriangleType::DeformedPose, MeshData, ProjectionHairData);
+				*ProjectionHairData.LODDatas[E.FrameLODIndex].DeformedPositionOffset = E.SkeletalDeformedPositionOffset;
+				UpdateHairStrandsMeshTriangles(RHICmdList, ShaderMap, E.FrameLODIndex, HairStrandsTriangleType::DeformedPose, MeshDataLOD, ProjectionHairData);
+			}
+		}
+	}
+
+	// Reset deformation
+	if (EHairStrandsInterpolationType::SimulationStrands == Type)
+	{
+		for (FHairStrandsManager::Element& E : GHairManager.Elements)
+		{
+			if (E.WorldType != WorldType)
+				continue;
+
+			if (E.InterpolationData.Input && E.InterpolationData.Output && E.InterpolationData.ResetFunction)
+			{
+				E.InterpolationData.ResetFunction(RHICmdList, E.InterpolationData.Input, E.InterpolationData.Output);
 			}
 		}
 	}
@@ -320,13 +345,13 @@ void RunHairStrandsInterpolation(
 
 			if (E.InterpolationData.Input && E.InterpolationData.Output && E.InterpolationData.Function)
 			{
-				E.InterpolationData.Function(RHICmdList, E.InterpolationData.Input, E.InterpolationData.Output, E.RenProjectionHairDatas, E.SimProjectionHairDatas, E.FrameLODIndex);
+				E.InterpolationData.Function(RHICmdList, ShaderDrawData, E.InterpolationData.Input, E.InterpolationData.Output, E.RenProjectionHairDatas, E.SimProjectionHairDatas, E.FrameLODIndex, ClusterData);
 			}
 		}
 	}
 }
 
-void GetGroomInterpolationData(const EWorldType::Type WorldType, FHairStrandsProjectionMeshData& OutGeometries)
+void GetGroomInterpolationData(const EWorldType::Type WorldType, FHairStrandsProjectionMeshData::LOD& OutGeometries)
 {
 	for (FHairStrandsManager::Element& E : GHairManager.Elements)
 	{
@@ -401,10 +426,28 @@ FHairStrandsDebugInfos GetHairStandsDebugInfos()
 	return Infos;
 }
 
+FHairStrandsPrimitiveResources GetHairStandsPrimitiveResources(uint64 Id)
+{
+	for (FHairStrandsManager::Element& E : GHairManager.Elements)
+	{
+		if (E.Id != Id)
+			continue;
+
+		return E.PrimitiveResources;
+	}
+	return FHairStrandsPrimitiveResources();
+}
+
+
 bool IsHairStrandsEnable(EShaderPlatform Platform) 
 { 
 	return 
 		IsHairStrandsSupported(Platform) && 
 		GHairStrandsRenderingEnable == 1 && 
 		GHairManager.Elements.Num() > 0; 
+}
+
+bool IsHairRayTracingEnabled()
+{
+	return IsRayTracingEnabled();
 }
\ No newline at end of file
diff --git a/Engine/Source/Runtime/Renderer/Private/HairStrands/HairStrandsMeshProjection.cpp b/Engine/Source/Runtime/Renderer/Private/HairStrands/HairStrandsMeshProjection.cpp
--- a/Engine/Source/Runtime/Renderer/Private/HairStrands/HairStrandsMeshProjection.cpp
+++ b/Engine/Source/Runtime/Renderer/Private/HairStrands/HairStrandsMeshProjection.cpp
@@ -100,7 +100,7 @@ static void AddHairStrandMeshProjectionPass(
 		Parameters->OutRootTriangleBarycentrics = RootData.LODDatas[LODIndex].RootTriangleBarycentricBuffer->UAV;
 		Parameters->OutRootTriangleDistance		= DistanceUAV;
 
-		const FIntVector DispatchGroupCount = FComputeShaderUtils::GetGroupCount(RootData.RootCount, 32);
+		const FIntVector DispatchGroupCount = FComputeShaderUtils::GetGroupCount(RootData.RootCount, 128);
 		check(DispatchGroupCount.X < 65536);
 		TShaderMapRef<FHairMeshProjectionCS> ComputeShader(ShaderMap);
 		FComputeShaderUtils::AddPass(
@@ -116,18 +116,17 @@ void ProjectHairStrandsOntoMesh(
 	FRHICommandListImmediate& RHICmdList, 
 	TShaderMap<FGlobalShaderType>* ShaderMap,
 	const int32 LODIndex,
-	FHairStrandsProjectionMeshData& ProjectionMeshData, 
+	const FHairStrandsProjectionMeshData& ProjectionMeshData, 
 	FHairStrandsProjectionHairData::HairGroup& ProjectionHairData)
 {
 	if (LODIndex < 0 || LODIndex >= ProjectionHairData.LODDatas.Num())
 		return;
 
 	FRDGBuilder GraphBuilder(RHICmdList);
-
 	FRDGBufferRef RootDistanceBuffer = GraphBuilder.CreateBuffer(FRDGBufferDesc::CreateStructuredDesc(sizeof(float),  ProjectionHairData.RootCount),	TEXT("HairStrandsTriangleDistance"));
 
 	bool ClearDistance = true;
-	for (FHairStrandsProjectionMeshData::Section& MeshSection : ProjectionMeshData.Sections)
+	for (const FHairStrandsProjectionMeshData::Section& MeshSection : ProjectionMeshData.LODs[LODIndex].Sections)
 	{
 		check(ProjectionHairData.LODDatas[LODIndex].LODIndex == LODIndex);
 		AddHairStrandMeshProjectionPass(GraphBuilder, ShaderMap, ClearDistance, LODIndex, MeshSection, ProjectionHairData, RootDistanceBuffer);
@@ -137,8 +136,6 @@ void ProjectHairStrandsOntoMesh(
 	GraphBuilder.Execute();
 }
 
-
-
 ///////////////////////////////////////////////////////////////////////////////////////////////////
 class FHairUpdateMeshTriangleCS : public FGlobalShader
 {
@@ -195,20 +192,21 @@ static void AddHairStrandUpdateMeshTrianglesPass(
 
 	if (Type == HairStrandsTriangleType::RestPose)
 	{
-		Parameters->RootTrianglePositionOffset = LODData.RestPositionOffset;
+		Parameters->RootTrianglePositionOffset = *LODData.RestPositionOffset;
 		Parameters->OutRootTrianglePosition0 = LODData.RestRootTrianglePosition0Buffer->UAV;
 		Parameters->OutRootTrianglePosition1 = LODData.RestRootTrianglePosition1Buffer->UAV;
 		Parameters->OutRootTrianglePosition2 = LODData.RestRootTrianglePosition2Buffer->UAV;
 	}
 	else if (Type == HairStrandsTriangleType::DeformedPose)
 	{
-		Parameters->RootTrianglePositionOffset = LODData.DeformedPositionOffset;
+		Parameters->RootTrianglePositionOffset = *LODData.DeformedPositionOffset;
 		Parameters->OutRootTrianglePosition0 = LODData.DeformedRootTrianglePosition0Buffer->UAV;
 		Parameters->OutRootTrianglePosition1 = LODData.DeformedRootTrianglePosition1Buffer->UAV;
 		Parameters->OutRootTrianglePosition2 = LODData.DeformedRootTrianglePosition2Buffer->UAV;
+		if (LODData.Status) (*LODData.Status) = FHairStrandsProjectionHairData::LODData::EStatus::Completed;
 	}
 
-	const FIntVector DispatchGroupCount = FComputeShaderUtils::GetGroupCount(RootData.RootCount, 32);
+	const FIntVector DispatchGroupCount = FComputeShaderUtils::GetGroupCount(RootData.RootCount, 128);
 	check(DispatchGroupCount.X < 65536);
 	TShaderMapRef<FHairUpdateMeshTriangleCS> ComputeShader(ShaderMap);
 	FComputeShaderUtils::AddPass(
@@ -224,12 +222,12 @@ void UpdateHairStrandsMeshTriangles(
 	TShaderMap<FGlobalShaderType>* ShaderMap,
 	const int32 LODIndex,
 	const HairStrandsTriangleType Type,
-	FHairStrandsProjectionMeshData& ProjectionMeshData,
+	const FHairStrandsProjectionMeshData::LOD& ProjectionMeshData,
 	FHairStrandsProjectionHairData::HairGroup& ProjectionHairData)
 {
 	FRDGBuilder GraphBuilder(RHICmdList);
 
-	for (FHairStrandsProjectionMeshData::Section& MeshSection : ProjectionMeshData.Sections)
+	for (const FHairStrandsProjectionMeshData::Section& MeshSection : ProjectionMeshData.Sections)
 	{
 		AddHairStrandUpdateMeshTrianglesPass(GraphBuilder, ShaderMap, LODIndex, Type, MeshSection, ProjectionHairData);
 	}
diff --git a/Engine/Source/Runtime/Renderer/Private/HairStrands/HairStrandsMeshProjection.h b/Engine/Source/Runtime/Renderer/Private/HairStrands/HairStrandsMeshProjection.h
--- a/Engine/Source/Runtime/Renderer/Private/HairStrands/HairStrandsMeshProjection.h
+++ b/Engine/Source/Runtime/Renderer/Private/HairStrands/HairStrandsMeshProjection.h
@@ -12,7 +12,7 @@ void ProjectHairStrandsOntoMesh(
 	FRHICommandListImmediate& RHICmdList, 
 	TShaderMap<FGlobalShaderType>* ShaderMap,
 	const int32 LODIndex,
-	FHairStrandsProjectionMeshData& ProjectionMeshData, 
+	const FHairStrandsProjectionMeshData& ProjectionMeshData, 
 	FHairStrandsProjectionHairData::HairGroup& ProjectionHairData);
 
 enum class HairStrandsTriangleType
@@ -27,5 +27,5 @@ void UpdateHairStrandsMeshTriangles(
 	TShaderMap<FGlobalShaderType>* ShaderMap,
 	const int32 LODIndex,
 	const HairStrandsTriangleType Type,
-	FHairStrandsProjectionMeshData& ProjectionMeshData,
+	const FHairStrandsProjectionMeshData::LOD& ProjectionMeshData,
 	FHairStrandsProjectionHairData::HairGroup& ProjectionHairData);
\ No newline at end of file
