Copying //UE4/Dev-Physics to Main (//UE4/Main)

diff --git a/Engine/Source/Runtime/Renderer/Private/HairStrands/HairStrandsRasterCommon.cpp b/Engine/Source/Runtime/Renderer/Private/HairStrands/HairStrandsRasterCommon.cpp
--- a/Engine/Source/Runtime/Renderer/Private/HairStrands/HairStrandsRasterCommon.cpp
+++ b/Engine/Source/Runtime/Renderer/Private/HairStrands/HairStrandsRasterCommon.cpp
@@ -11,8 +11,55 @@
 #include "MeshPassProcessor.inl"
 #include "ScenePrivate.h"
 
-// this is temporary until we split the voxelize and DOM path
-IMPLEMENT_GLOBAL_SHADER_PARAMETER_STRUCT(FDeepShadowPassUniformParameters, "DeepShadowPass");
+/////////////////////////////////////////////////////////////////////////////////////////
+// Deep shadow global parameters
+BEGIN_GLOBAL_SHADER_PARAMETER_STRUCT(FHairDeepShadowRasterGlobalParameters, )
+	SHADER_PARAMETER(FMatrix, WorldToClipMatrix)
+	SHADER_PARAMETER(FVector4, SliceValue)
+	SHADER_PARAMETER(FIntRect, AtlasRect)
+	SHADER_PARAMETER(FIntPoint, ViewportResolution)
+	SHADER_PARAMETER_TEXTURE(Texture2D<float>, FrontDepthTexture)
+	END_GLOBAL_SHADER_PARAMETER_STRUCT()
+IMPLEMENT_GLOBAL_SHADER_PARAMETER_STRUCT(FHairDeepShadowRasterGlobalParameters, "DeepRasterPass");
+
+static FHairDeepShadowRasterGlobalParameters ConvertToGlobalPassParameter(const FHairDeepShadowRasterPassParameters* In)
+{
+	FHairDeepShadowRasterGlobalParameters Out;
+	Out.WorldToClipMatrix	= In->WorldToClipMatrix;
+	Out.SliceValue			= In->SliceValue;
+	Out.AtlasRect			= In->AtlasRect;
+	Out.ViewportResolution	= In->ViewportResolution;
+	Out.FrontDepthTexture	= In->FrontDepthTexture ? In->FrontDepthTexture->GetRHI() : (FRHITexture*)(GSystemTextures.DepthDummy->GetRenderTargetItem().ShaderResourceTexture);
+	return Out;
+}
+
+/////////////////////////////////////////////////////////////////////////////////////////
+// Voxelization global parameters
+BEGIN_GLOBAL_SHADER_PARAMETER_STRUCT(FHairVoxelizationRasterGlobalParameters, )
+	SHADER_PARAMETER_STRUCT(FVirtualVoxelCommonParameters, VirtualVoxel)
+	SHADER_PARAMETER(FMatrix, WorldToClipMatrix)
+	SHADER_PARAMETER(FVector, VoxelMinAABB)
+	SHADER_PARAMETER(FVector, VoxelMaxAABB)
+	SHADER_PARAMETER(FIntVector, VoxelResolution)
+	SHADER_PARAMETER(uint32, MacroGroupId)
+	SHADER_PARAMETER(FIntPoint, ViewportResolution)
+	SHADER_PARAMETER_SRV(StructuredBuffer<FVoxelizationViewInfo>, VoxelizationViewInfoBuffer)
+END_GLOBAL_SHADER_PARAMETER_STRUCT()
+IMPLEMENT_GLOBAL_SHADER_PARAMETER_STRUCT(FHairVoxelizationRasterGlobalParameters, "VoxelRasterPass");
+
+static FHairVoxelizationRasterGlobalParameters ConvertToGlobalPassParameter(const FHairVoxelizationRasterPassParameters* In)
+{
+	FHairVoxelizationRasterGlobalParameters Out;
+	Out.VirtualVoxel	   = In->VirtualVoxel;
+	Out.WorldToClipMatrix  = In->WorldToClipMatrix;
+	Out.VoxelMinAABB	   = In->VoxelMinAABB;
+	Out.VoxelMaxAABB	   = In->VoxelMaxAABB;
+	Out.VoxelResolution	   = In->VoxelResolution;
+	Out.MacroGroupId	   = In->MacroGroupId;
+	Out.ViewportResolution = In->ViewportResolution;
+	Out.VoxelizationViewInfoBuffer = In->VoxelizationViewInfoBuffer->GetRHI();
+	return Out;
+}
 
 /////////////////////////////////////////////////////////////////////////////////////////
 
@@ -28,7 +75,7 @@ protected:
 		ERHIFeatureLevel::Type FeatureLevel = GetMaxSupportedFeatureLevel((EShaderPlatform)Initializer.Target.Platform);
 		check(FSceneInterface::GetShadingPath(FeatureLevel) != EShadingPath::Mobile);
 		// deferred
-		PassUniformBuffer.Bind(Initializer.ParameterMap, FDeepShadowPassUniformParameters::StaticStructMetadata.GetShaderVariableName());
+		PassUniformBuffer.Bind(Initializer.ParameterMap, FHairDeepShadowRasterGlobalParameters::StaticStructMetadata.GetShaderVariableName());
 	}
 
 	FDeepShadowDepthMeshVS() {}
@@ -42,6 +89,7 @@ protected:
 	{
 		FMaterialShader::ModifyCompilationEnvironment(Parameters, OutEnvironment);
 		OutEnvironment.SetDefine(TEXT("MESH_RENDER_MODE"), 0);
+		OutEnvironment.SetDefine(TEXT("USE_CULLED_CLUSTER"), 1);
 	}
 };
 
@@ -61,7 +109,7 @@ protected:
 		ERHIFeatureLevel::Type FeatureLevel = GetMaxSupportedFeatureLevel((EShaderPlatform)Initializer.Target.Platform);
 		check(FSceneInterface::GetShadingPath(FeatureLevel) != EShadingPath::Mobile);
 		// deferred
-		PassUniformBuffer.Bind(Initializer.ParameterMap, FDeepShadowPassUniformParameters::StaticStructMetadata.GetShaderVariableName());
+		PassUniformBuffer.Bind(Initializer.ParameterMap, FHairDeepShadowRasterGlobalParameters::StaticStructMetadata.GetShaderVariableName());
 	}
 
 	FDeepShadowDomMeshVS() {}
@@ -75,6 +123,7 @@ protected:
 	{
 		FMaterialShader::ModifyCompilationEnvironment(Parameters, OutEnvironment);
 		OutEnvironment.SetDefine(TEXT("MESH_RENDER_MODE"), 1);
+		OutEnvironment.SetDefine(TEXT("USE_CULLED_CLUSTER"), 1);
 	}
 };
 
@@ -95,7 +144,7 @@ protected:
 		ERHIFeatureLevel::Type FeatureLevel = GetMaxSupportedFeatureLevel((EShaderPlatform)Initializer.Target.Platform);
 		check(FSceneInterface::GetShadingPath(FeatureLevel) != EShadingPath::Mobile);
 		// deferred
-		PassUniformBuffer.Bind(Initializer.ParameterMap, FDeepShadowPassUniformParameters::StaticStructMetadata.GetShaderVariableName());
+		PassUniformBuffer.Bind(Initializer.ParameterMap, FHairVoxelizationRasterGlobalParameters::StaticStructMetadata.GetShaderVariableName());
 	}
 
 	FVoxelMeshVS() {}
@@ -111,6 +160,7 @@ protected:
 		OutEnvironment.SetDefine(TEXT("MESH_RENDER_MODE"), 2);
 		OutEnvironment.SetDefine(TEXT("SUPPORT_TANGENT_PROPERTY"), bVoxelizeMaterial ? 1 : 0);
 		OutEnvironment.SetDefine(TEXT("SUPPORT_MATERIAL_PROPERTY"), bVoxelizeMaterial ? 1 : 0);
+		OutEnvironment.SetDefine(TEXT("USE_CULLED_CLUSTER"), 1);
 	}
 };
 IMPLEMENT_MATERIAL_SHADER_TYPE(template<>, FVoxelMeshVS<false>, TEXT("/Engine/Private/HairStrands/HairStrandsDeepShadowVS.usf"), TEXT("Main"), SF_Vertex);
@@ -129,7 +179,7 @@ public:
 	{
 		ERHIFeatureLevel::Type FeatureLevel = GetMaxSupportedFeatureLevel((EShaderPlatform)Initializer.Target.Platform);
 		check(FSceneInterface::GetShadingPath(FeatureLevel) != EShadingPath::Mobile);
-		PassUniformBuffer.Bind(Initializer.ParameterMap, FDeepShadowPassUniformParameters::StaticStructMetadata.GetShaderVariableName());
+		PassUniformBuffer.Bind(Initializer.ParameterMap, FHairDeepShadowRasterGlobalParameters::StaticStructMetadata.GetShaderVariableName());
 	}
 
 	FDeepShadowDepthMeshPS() {}
@@ -154,7 +204,7 @@ public:
 	{
 		ERHIFeatureLevel::Type FeatureLevel = GetMaxSupportedFeatureLevel((EShaderPlatform)Initializer.Target.Platform);
 		check(FSceneInterface::GetShadingPath(FeatureLevel) != EShadingPath::Mobile);
-		PassUniformBuffer.Bind(Initializer.ParameterMap, FDeepShadowPassUniformParameters::StaticStructMetadata.GetShaderVariableName());
+		PassUniformBuffer.Bind(Initializer.ParameterMap, FHairDeepShadowRasterGlobalParameters::StaticStructMetadata.GetShaderVariableName());
 	}
 
 	FDeepShadowDomMeshPS() {}
@@ -174,7 +224,13 @@ IMPLEMENT_MATERIAL_SHADER_TYPE(, FDeepShadowDomMeshPS, TEXT("/Engine/Private/Hai
 
 /////////////////////////////////////////////////////////////////////////////////////////
 
-template<bool bVoxelizeMaterial>
+enum class EVoxelMeshPSType
+{
+	Density,
+	Material,
+	DensityVirtual
+};
+template<EVoxelMeshPSType VoxelizationType>
 class FVoxelMeshPS : public FMeshMaterialShader
 {
 	DECLARE_SHADER_TYPE(FVoxelMeshPS, MeshMaterial);
@@ -186,9 +242,9 @@ public:
 	{
 		ERHIFeatureLevel::Type FeatureLevel = GetMaxSupportedFeatureLevel((EShaderPlatform)Initializer.Target.Platform);
 		check(FSceneInterface::GetShadingPath(FeatureLevel) != EShadingPath::Mobile);
-		PassUniformBuffer.Bind(Initializer.ParameterMap, FDeepShadowPassUniformParameters::StaticStructMetadata.GetShaderVariableName());
+		PassUniformBuffer.Bind(Initializer.ParameterMap, FHairVoxelizationRasterGlobalParameters::StaticStructMetadata.GetShaderVariableName());
 		DensityTexture.Bind(Initializer.ParameterMap, TEXT("DensityTexture"));
-		if (bVoxelizeMaterial)
+		if (VoxelizationType == EVoxelMeshPSType::Material)
 		{
 			TangentXTexture.Bind(Initializer.ParameterMap, TEXT("TangentXTexture"));
 			TangentYTexture.Bind(Initializer.ParameterMap, TEXT("TangentYTexture"));
@@ -207,15 +263,16 @@ public:
 	static void ModifyCompilationEnvironment(const FMaterialShaderPermutationParameters& Parameters, FShaderCompilerEnvironment& OutEnvironment)
 	{
 		FMaterialShader::ModifyCompilationEnvironment(Parameters, OutEnvironment);
-		OutEnvironment.SetDefine(TEXT("SUPPORT_TANGENT_PROPERTY"), bVoxelizeMaterial ? 1 : 0);
-		OutEnvironment.SetDefine(TEXT("SUPPORT_MATERIAL_PROPERTY"), bVoxelizeMaterial ? 1 : 0);
+		OutEnvironment.SetDefine(TEXT("SUPPORT_TANGENT_PROPERTY"), VoxelizationType == EVoxelMeshPSType::Material ? 1 : 0);
+		OutEnvironment.SetDefine(TEXT("SUPPORT_MATERIAL_PROPERTY"), VoxelizationType == EVoxelMeshPSType::Material ? 1 : 0);
+		OutEnvironment.SetDefine(TEXT("SUPPORT_SPARSE"), VoxelizationType == EVoxelMeshPSType::DensityVirtual ? 1 : 0);
 	}
 
 	virtual bool Serialize(FArchive& Ar) override
 	{
 		bool bShaderHasOutdatedParameters = FMeshMaterialShader::Serialize(Ar);
 		Ar << DensityTexture;
-		if (bVoxelizeMaterial)
+		if (VoxelizationType == EVoxelMeshPSType::Material)
 		{
 			Ar << TangentXTexture;
 			Ar << TangentYTexture;
@@ -231,16 +288,17 @@ private:
 	FRWShaderParameter TangentZTexture;
 	FRWShaderParameter MaterialTexture;
 };
-IMPLEMENT_MATERIAL_SHADER_TYPE(template<>, FVoxelMeshPS<true>, TEXT("/Engine/Private/HairStrands/HairStrandsDeepShadowPS.usf"), TEXT("MainVoxel"), SF_Pixel);
-IMPLEMENT_MATERIAL_SHADER_TYPE(template<>, FVoxelMeshPS<false>, TEXT("/Engine/Private/HairStrands/HairStrandsDeepShadowPS.usf"), TEXT("MainVoxel"), SF_Pixel);
+IMPLEMENT_MATERIAL_SHADER_TYPE(template<>, FVoxelMeshPS<EVoxelMeshPSType::Density>, TEXT("/Engine/Private/HairStrands/HairStrandsDeepShadowPS.usf"), TEXT("MainVoxel"), SF_Pixel);
+IMPLEMENT_MATERIAL_SHADER_TYPE(template<>, FVoxelMeshPS<EVoxelMeshPSType::Material>, TEXT("/Engine/Private/HairStrands/HairStrandsDeepShadowPS.usf"), TEXT("MainVoxel"), SF_Pixel);
+IMPLEMENT_MATERIAL_SHADER_TYPE(template<>, FVoxelMeshPS<EVoxelMeshPSType::DensityVirtual>, TEXT("/Engine/Private/HairStrands/HairStrandsDeepShadowPS.usf"), TEXT("MainVoxel"), SF_Pixel);
 
 /////////////////////////////////////////////////////////////////////////////////////////
 
-class FDeepShadowMeshProcessor : public FMeshPassProcessor
+class FHairRasterMeshProcessor : public FMeshPassProcessor
 {
 public:
 
-	FDeepShadowMeshProcessor(
+	FHairRasterMeshProcessor(
 		const FScene* Scene,
 		const FSceneView* InViewIfDynamicMeshCommand,
 		const FMeshPassProcessorRenderState& InPassDrawRenderState,
@@ -261,11 +319,11 @@ private:
 		ERasterizerFillMode MeshFillMode,
 		ERasterizerCullMode MeshCullMode);
 
-	EHairStrandsRasterPassType DeepShadowPassType;
+	EHairStrandsRasterPassType RasterPassType;
 	FMeshPassProcessorRenderState PassDrawRenderState;
 };
 
-void FDeepShadowMeshProcessor::AddMeshBatch(const FMeshBatch& RESTRICT MeshBatch, uint64 BatchElementMask, const FPrimitiveSceneProxy* RESTRICT PrimitiveSceneProxy, int32 StaticMeshId)
+void FHairRasterMeshProcessor::AddMeshBatch(const FMeshBatch& RESTRICT MeshBatch, uint64 BatchElementMask, const FPrimitiveSceneProxy* RESTRICT PrimitiveSceneProxy, int32 StaticMeshId)
 {
 	// Determine the mesh's material and blend mode.
 	const FMaterialRenderProxy* FallbackMaterialRenderProxyPtr = nullptr;
@@ -278,23 +336,25 @@ void FDeepShadowMeshProcessor::AddMeshBatch(const FMeshBatch& RESTRICT MeshBatch
 	{
 		const FMaterialRenderProxy& MaterialRenderProxy = FallbackMaterialRenderProxyPtr ? *FallbackMaterialRenderProxyPtr : *MeshBatch.MaterialRenderProxy;
 		const ERasterizerFillMode MeshFillMode = ComputeMeshFillMode(MeshBatch, Material);
-		const ERasterizerCullMode MeshCullMode = DeepShadowPassType == EHairStrandsRasterPassType::FrontDepth ? ComputeMeshCullMode(MeshBatch, Material) : CM_None;
+		const ERasterizerCullMode MeshCullMode = RasterPassType == EHairStrandsRasterPassType::FrontDepth ? ComputeMeshCullMode(MeshBatch, Material) : CM_None;
 
-		if (DeepShadowPassType == EHairStrandsRasterPassType::FrontDepth)
+		if (RasterPassType == EHairStrandsRasterPassType::FrontDepth)
 			Process<FDeepShadowDepthMeshVS, FDeepShadowDepthMeshPS>(MeshBatch, BatchElementMask, PrimitiveSceneProxy, StaticMeshId, MaterialRenderProxy, Material, MeshFillMode, MeshCullMode);
-		else if (DeepShadowPassType == EHairStrandsRasterPassType::DeepOpacityMap)
+		else if (RasterPassType == EHairStrandsRasterPassType::DeepOpacityMap)
 			Process<FDeepShadowDomMeshVS, FDeepShadowDomMeshPS>(MeshBatch, BatchElementMask, PrimitiveSceneProxy, StaticMeshId, MaterialRenderProxy, Material, MeshFillMode, MeshCullMode);
-		else if (DeepShadowPassType == EHairStrandsRasterPassType::Voxelization)
-			Process<FVoxelMeshVS<false>, FVoxelMeshPS<false>>(MeshBatch, BatchElementMask, PrimitiveSceneProxy, StaticMeshId, MaterialRenderProxy, Material, MeshFillMode, MeshCullMode);
-		else if (DeepShadowPassType == EHairStrandsRasterPassType::VoxelizationMaterial)
-			Process<FVoxelMeshVS<true>, FVoxelMeshPS<true>>(MeshBatch, BatchElementMask, PrimitiveSceneProxy, StaticMeshId, MaterialRenderProxy, Material, MeshFillMode, MeshCullMode);
+		else if (RasterPassType == EHairStrandsRasterPassType::Voxelization)
+			Process<FVoxelMeshVS<false>, FVoxelMeshPS<EVoxelMeshPSType::Density>>(MeshBatch, BatchElementMask, PrimitiveSceneProxy, StaticMeshId, MaterialRenderProxy, Material, MeshFillMode, MeshCullMode);
+		else if (RasterPassType == EHairStrandsRasterPassType::VoxelizationMaterial)
+			Process<FVoxelMeshVS<true>, FVoxelMeshPS<EVoxelMeshPSType::Material>>(MeshBatch, BatchElementMask, PrimitiveSceneProxy, StaticMeshId, MaterialRenderProxy, Material, MeshFillMode, MeshCullMode);
+		else if (RasterPassType == EHairStrandsRasterPassType::VoxelizationVirtual)
+			Process<FVoxelMeshVS<false>, FVoxelMeshPS<EVoxelMeshPSType::DensityVirtual>>(MeshBatch, BatchElementMask, PrimitiveSceneProxy, StaticMeshId, MaterialRenderProxy, Material, MeshFillMode, MeshCullMode);
 	}
 }
 
 // Vertex is either FDeepShadowDepthMeshVS, FDeepShadowDomMeshVS, or FVoxelMeshVS
 // Pixel  is either FDeepShadowDepthMeshPS, FDeepShadowDomMeshPS, or FVoxelMeshPS
 template<typename VertexShaderType, typename PixelShaderType>
-void FDeepShadowMeshProcessor::Process(
+void FHairRasterMeshProcessor::Process(
 	const FMeshBatch& MeshBatch,
 	uint64 BatchElementMask,
 	const FPrimitiveSceneProxy* RESTRICT PrimitiveSceneProxy,
@@ -345,80 +405,171 @@ void FDeepShadowMeshProcessor::Process(
 		ShaderElementData);
 }
 
-FDeepShadowMeshProcessor::FDeepShadowMeshProcessor(
+FHairRasterMeshProcessor::FHairRasterMeshProcessor(
 	const FScene* Scene,
 	const FSceneView* InViewIfDynamicMeshCommand,
 	const FMeshPassProcessorRenderState& InPassDrawRenderState,
 	FDynamicPassMeshDrawListContext* InDrawListContext,
 	const EHairStrandsRasterPassType PType)
 	: FMeshPassProcessor(Scene, Scene->GetFeatureLevel(), InViewIfDynamicMeshCommand, InDrawListContext)
-	, DeepShadowPassType(PType)
+	, RasterPassType(PType)
 	, PassDrawRenderState(InPassDrawRenderState)
 {
 }
 
 /////////////////////////////////////////////////////////////////////////////////////////
 
-void RasterHairStrands(
-	FRHICommandList& RHICmdList,
+template<typename TPassParameter, typename TGlobalParameter>
+void AddHairStrandsRasterPass(
+	FRDGBuilder& GraphBuilder,
 	const FScene* Scene,
 	const FViewInfo* ViewInfo,
-	const FHairStrandsClusterData::TPrimitiveInfos& PrimitiveSceneInfos,
-	const EHairStrandsRasterPassType ShadowPassType,
-	const FIntRect& AtlasRect,
-	const TUniformBufferRef<FViewUniformShaderParameters>& ViewUniformShaderParameters,
-	const TUniformBufferRef<FDeepShadowPassUniformParameters>& DeepShadowPassUniformParameters)
+	const FHairStrandsMacroGroupData::TPrimitiveInfos& PrimitiveSceneInfos,
+	const EHairStrandsRasterPassType RasterPassType,
+	const FIntRect& ViewportRect,
+	const FVector4& HairRenderInfo,
+	const FVector& RasterDirection,
+	TPassParameter* PassParameters)
 {
-	check(RHICmdList.IsInsideRenderPass());
-	check(IsInRenderingThread());
+	auto GetPassName = [](EHairStrandsRasterPassType Type)
+	{
+		switch (Type)
+		{
+		case EHairStrandsRasterPassType::DeepOpacityMap:		return RDG_EVENT_NAME("HairStrandsRasterDeepOpacityMap");
+		case EHairStrandsRasterPassType::FrontDepth:			return RDG_EVENT_NAME("HairStrandsRasterFrontDepth");
+		case EHairStrandsRasterPassType::Voxelization:			return RDG_EVENT_NAME("HairStrandsRasterVoxelization");
+		case EHairStrandsRasterPassType::VoxelizationMaterial:	return RDG_EVENT_NAME("HairStrandsRasterVoxelizationMaterial");
+		case EHairStrandsRasterPassType::VoxelizationVirtual:	return RDG_EVENT_NAME("HairStrandsRasterVoxelizationVirtual");
+		default:												return RDG_EVENT_NAME("Noname");
+		}
+	};
 
-	SCOPE_CYCLE_COUNTER(STAT_RenderPerObjectShadowDepthsTime);
+	GraphBuilder.AddPass(
+		GetPassName(RasterPassType),
+		PassParameters,
+		ERDGPassFlags::Raster,
+		[PassParameters, Scene = Scene, ViewInfo, RasterPassType, &PrimitiveSceneInfos, ViewportRect, HairRenderInfo, RasterDirection](FRHICommandListImmediate& RHICmdList)
+	{
+		check(RHICmdList.IsInsideRenderPass());
+		check(IsInRenderingThread());
 
-	FMeshPassProcessorRenderState DrawRenderState(*ViewInfo, DeepShadowPassUniformParameters);
-	DrawRenderState.SetViewUniformBuffer(ViewUniformShaderParameters);
+		check(RHICmdList.IsInsideRenderPass());
+		check(IsInRenderingThread());
 
-	RHICmdList.SetViewport(AtlasRect.Min.X, AtlasRect.Min.Y, 0.0f, AtlasRect.Max.X, AtlasRect.Max.Y, 1.0f);
+		SCOPE_CYCLE_COUNTER(STAT_RenderPerObjectShadowDepthsTime);
 
-	if (ShadowPassType == EHairStrandsRasterPassType::DeepOpacityMap)
-	{
-		DrawRenderState.SetBlendState(TStaticBlendState<
-			CW_RGBA, BO_Add, BF_One, BF_One, BO_Add, BF_One, BF_One,
-			CW_RGBA, BO_Add, BF_One, BF_One, BO_Add, BF_One, BF_One>::GetRHI());
-		DrawRenderState.SetDepthStencilState(TStaticDepthStencilState<false, CF_Always>::GetRHI());
-	}
-	else if (ShadowPassType == EHairStrandsRasterPassType::FrontDepth)
-	{
-		DrawRenderState.SetBlendState(TStaticBlendState<
-			CW_RGBA, BO_Add, BF_One, BF_Zero, BO_Add, BF_One, BF_Zero,
-			CW_RGBA, BO_Add, BF_One, BF_Zero, BO_Add, BF_One, BF_Zero>::GetRHI());
-		DrawRenderState.SetDepthStencilState(TStaticDepthStencilState<true, CF_DepthNearOrEqual>::GetRHI());
-	}
-	else if (ShadowPassType == EHairStrandsRasterPassType::Voxelization || ShadowPassType == EHairStrandsRasterPassType::VoxelizationMaterial)
-	{
-		DrawRenderState.SetBlendState(TStaticBlendState<
-			CW_RGBA, BO_Add, BF_One, BF_Zero, BO_Add, BF_One, BF_Zero>::GetRHI());
-		DrawRenderState.SetDepthStencilState(TStaticDepthStencilState<false, CF_Always>::GetRHI());
-	}
+		TUniformBufferRef<FViewUniformShaderParameters> ViewUniformShaderParameters;
+		ViewInfo->CachedViewUniformShaderParameters->HairRenderInfo = HairRenderInfo;
 
-	FDynamicMeshDrawCommandStorage DynamicMeshDrawCommandStorage; // << Were would thid be stored?
-	FMeshCommandOneFrameArray VisibleMeshDrawCommands;
-	FGraphicsMinimalPipelineStateSet GraphicsMinimalPipelineStateSet;
-	FDynamicPassMeshDrawListContext ShadowContext(DynamicMeshDrawCommandStorage, VisibleMeshDrawCommands, GraphicsMinimalPipelineStateSet);
+		const FVector SavedViewForward = ViewInfo->CachedViewUniformShaderParameters->ViewForward;	
+		ViewInfo->CachedViewUniformShaderParameters->ViewForward = RasterDirection;
+		ViewUniformShaderParameters = TUniformBufferRef<FViewUniformShaderParameters>::CreateUniformBufferImmediate(*ViewInfo->CachedViewUniformShaderParameters, UniformBuffer_SingleFrame);
+		ViewInfo->CachedViewUniformShaderParameters->ViewForward = SavedViewForward;
 
+		TGlobalParameter GlobalPassParameters = ConvertToGlobalPassParameter(PassParameters);
+		TUniformBufferRef<TGlobalParameter> GlobalPassParametersBuffer = TUniformBufferRef<TGlobalParameter>::CreateUniformBufferImmediate(GlobalPassParameters, UniformBuffer_SingleFrame);
 
-	FDeepShadowMeshProcessor DeepShadowMeshProcessor(Scene, ViewInfo /* is a SceneView */, DrawRenderState, &ShadowContext, ShadowPassType);
+		FMeshPassProcessorRenderState DrawRenderState(*ViewInfo, GlobalPassParametersBuffer);
+		DrawRenderState.SetViewUniformBuffer(ViewUniformShaderParameters);
 
-	for (const FHairStrandsClusterData::PrimitiveInfo& PrimitveInfo : PrimitiveSceneInfos)
-	{
-		const FMeshBatch& MeshBatch = *PrimitveInfo.MeshBatchAndRelevance.Mesh;
-		const uint64 BatchElementMask = ~0ull;
-		DeepShadowMeshProcessor.AddMeshBatch(MeshBatch, BatchElementMask, PrimitveInfo.MeshBatchAndRelevance.PrimitiveSceneProxy);
-	}
+		RHICmdList.SetViewport(ViewportRect.Min.X, ViewportRect.Min.Y, 0.0f, ViewportRect.Max.X, ViewportRect.Max.Y, 1.0f);
 
-	if (VisibleMeshDrawCommands.Num() > 0)
-	{
-		FRHIVertexBuffer* PrimitiveIdVertexBuffer = nullptr;
-		SortAndMergeDynamicPassMeshDrawCommands(ViewInfo->GetFeatureLevel(), VisibleMeshDrawCommands, DynamicMeshDrawCommandStorage, PrimitiveIdVertexBuffer, 1);
-		SubmitMeshDrawCommands(VisibleMeshDrawCommands, GraphicsMinimalPipelineStateSet, PrimitiveIdVertexBuffer, 0, false, 1, RHICmdList);
-	}
+		if (RasterPassType == EHairStrandsRasterPassType::DeepOpacityMap)
+		{
+			DrawRenderState.SetBlendState(TStaticBlendState<
+				CW_RGBA, BO_Add, BF_One, BF_One, BO_Add, BF_One, BF_One,
+				CW_RGBA, BO_Add, BF_One, BF_One, BO_Add, BF_One, BF_One>::GetRHI());
+			DrawRenderState.SetDepthStencilState(TStaticDepthStencilState<false, CF_Always>::GetRHI());
+		}
+		else if (RasterPassType == EHairStrandsRasterPassType::FrontDepth)
+		{
+			DrawRenderState.SetBlendState(TStaticBlendState<
+				CW_RGBA, BO_Add, BF_One, BF_Zero, BO_Add, BF_One, BF_Zero,
+				CW_RGBA, BO_Add, BF_One, BF_Zero, BO_Add, BF_One, BF_Zero>::GetRHI());
+			DrawRenderState.SetDepthStencilState(TStaticDepthStencilState<true, CF_DepthNearOrEqual>::GetRHI());
+		}
+		else if (RasterPassType == EHairStrandsRasterPassType::Voxelization || RasterPassType == EHairStrandsRasterPassType::VoxelizationMaterial)
+		{
+			DrawRenderState.SetBlendState(TStaticBlendState<
+				CW_RGBA, BO_Add, BF_One, BF_Zero, BO_Add, BF_One, BF_Zero>::GetRHI());
+			DrawRenderState.SetDepthStencilState(TStaticDepthStencilState<false, CF_Always>::GetRHI());
+		}
+		else if (RasterPassType == EHairStrandsRasterPassType::VoxelizationVirtual)
+		{
+			DrawRenderState.SetBlendState(TStaticBlendState<
+				CW_RGBA, BO_Add, BF_One, BF_Zero, BO_Add, BF_One, BF_Zero>::GetRHI());
+			DrawRenderState.SetDepthStencilState(TStaticDepthStencilState<false, CF_Always>::GetRHI());
+		}
+
+		FDynamicMeshDrawCommandStorage DynamicMeshDrawCommandStorage; // << Were would thid be stored?
+		FMeshCommandOneFrameArray VisibleMeshDrawCommands;
+		FGraphicsMinimalPipelineStateSet GraphicsMinimalPipelineStateSet;
+		FDynamicPassMeshDrawListContext ShadowContext(DynamicMeshDrawCommandStorage, VisibleMeshDrawCommands, GraphicsMinimalPipelineStateSet);
+
+		FHairRasterMeshProcessor HairRasterMeshProcessor(Scene, ViewInfo /* is a SceneView */, DrawRenderState, &ShadowContext, RasterPassType);
+
+		for (const FHairStrandsMacroGroupData::PrimitiveInfo& PrimitveInfo : PrimitiveSceneInfos)
+		{
+			const FMeshBatch& MeshBatch = *PrimitveInfo.MeshBatchAndRelevance.Mesh;
+			const uint64 BatchElementMask = ~0ull;
+			HairRasterMeshProcessor.AddMeshBatch(MeshBatch, BatchElementMask, PrimitveInfo.MeshBatchAndRelevance.PrimitiveSceneProxy);
+		}
+
+		if (VisibleMeshDrawCommands.Num() > 0)
+		{
+			FRHIVertexBuffer* PrimitiveIdVertexBuffer = nullptr;
+			SortAndMergeDynamicPassMeshDrawCommands(ViewInfo->GetFeatureLevel(), VisibleMeshDrawCommands, DynamicMeshDrawCommandStorage, PrimitiveIdVertexBuffer, 1);
+			SubmitMeshDrawCommands(VisibleMeshDrawCommands, GraphicsMinimalPipelineStateSet, PrimitiveIdVertexBuffer, 0, false, 1, RHICmdList);
+		}
+	});
+}
+
+void AddHairDeepShadowRasterPass(
+	FRDGBuilder& GraphBuilder,
+	const FScene* Scene,
+	const FViewInfo* ViewInfo,
+	const FHairStrandsMacroGroupData::TPrimitiveInfos& PrimitiveSceneInfos,
+	const EHairStrandsRasterPassType PassType,
+	const FIntRect& ViewportRect,
+	const FVector4& HairRenderInfo,
+	const FVector& LightDirection,
+	FHairDeepShadowRasterPassParameters* PassParameters)
+{
+	check(PassType == EHairStrandsRasterPassType::FrontDepth || PassType == EHairStrandsRasterPassType::DeepOpacityMap);
+
+	AddHairStrandsRasterPass<FHairDeepShadowRasterPassParameters, FHairDeepShadowRasterGlobalParameters>(
+		GraphBuilder, 
+		Scene, 
+		ViewInfo, 
+		PrimitiveSceneInfos, 
+		PassType, 
+		ViewportRect, 
+		HairRenderInfo, 
+		LightDirection,
+		PassParameters);
 }
+
+void AddHairVoxelizationRasterPass(
+	FRDGBuilder& GraphBuilder,
+	const FScene* Scene,
+	const FViewInfo* ViewInfo,
+	const FHairStrandsMacroGroupData::TPrimitiveInfos& PrimitiveSceneInfos,
+	const EHairStrandsRasterPassType PassType,
+	const FIntRect& ViewportRect,
+	const FVector4& HairRenderInfo,
+	const FVector& RasterDirection,
+	FHairVoxelizationRasterPassParameters* PassParameters)
+{
+	check(PassType == EHairStrandsRasterPassType::Voxelization || PassType == EHairStrandsRasterPassType::VoxelizationMaterial || PassType == EHairStrandsRasterPassType::VoxelizationVirtual);
+
+	AddHairStrandsRasterPass<FHairVoxelizationRasterPassParameters, FHairVoxelizationRasterGlobalParameters>(
+		GraphBuilder, 
+		Scene, 
+		ViewInfo, 
+		PrimitiveSceneInfos, 
+		PassType, 
+		ViewportRect, 
+		HairRenderInfo, 
+		RasterDirection,
+		PassParameters);
+}
\ No newline at end of file
diff --git a/Engine/Source/Runtime/Renderer/Private/HairStrands/HairStrandsRasterCommon.h b/Engine/Source/Runtime/Renderer/Private/HairStrands/HairStrandsRasterCommon.h
--- a/Engine/Source/Runtime/Renderer/Private/HairStrands/HairStrandsRasterCommon.h
+++ b/Engine/Source/Runtime/Renderer/Private/HairStrands/HairStrandsRasterCommon.h
@@ -9,6 +9,9 @@
 #include "CoreMinimal.h"
 #include "RendererInterface.h"
 #include "HairStrandsCluster.h"
+#include "HairStrandsVoxelization.h"
+
+class FViewInfo;
 
 enum class EHairStrandsRasterPassType : uint8
 {
@@ -16,25 +19,60 @@ enum class EHairStrandsRasterPassType : uint8
 	DeepOpacityMap,
 	Voxelization,
 	VoxelizationMaterial,
+	VoxelizationVirtual
 };
 
-BEGIN_GLOBAL_SHADER_PARAMETER_STRUCT(FDeepShadowPassUniformParameters, )
+// ////////////////////////////////////////////////////////////////
+// Deep shadow raster pass
+
+BEGIN_SHADER_PARAMETER_STRUCT(FHairDeepShadowRasterPassParameters, )
 	SHADER_PARAMETER(FMatrix, WorldToClipMatrix)
 	SHADER_PARAMETER(FVector4, SliceValue)
 	SHADER_PARAMETER(FIntRect, AtlasRect)
+	SHADER_PARAMETER(FIntPoint, ViewportResolution)
+	SHADER_PARAMETER_RDG_TEXTURE(Texture2D<float>, FrontDepthTexture)
+	RENDER_TARGET_BINDING_SLOTS()
+END_SHADER_PARAMETER_STRUCT()
+
+void AddHairDeepShadowRasterPass(
+	FRDGBuilder& GraphBuilder,
+	const FScene* Scene,
+	const FViewInfo* ViewInfo,
+	const FHairStrandsMacroGroupData::TPrimitiveInfos& PrimitiveSceneInfos,
+	const EHairStrandsRasterPassType ShadowPassType,
+	const FIntRect& ViewportRect,
+	const FVector4& HairRenderInfo,
+	const FVector& LightDirection,
+	FHairDeepShadowRasterPassParameters* PassParameters);
+
+// ////////////////////////////////////////////////////////////////
+// Voxelization raster pass
+BEGIN_SHADER_PARAMETER_STRUCT(FHairVoxelizationRasterPassParameters, )
+	SHADER_PARAMETER_STRUCT(FVirtualVoxelCommonParameters, VirtualVoxel)
+	SHADER_PARAMETER(FMatrix, WorldToClipMatrix)
 	SHADER_PARAMETER(FVector, VoxelMinAABB)
 	SHADER_PARAMETER(FVector, VoxelMaxAABB)
-	SHADER_PARAMETER(uint32, VoxelResolution)
-	SHADER_PARAMETER_TEXTURE(Texture2D<float>, FrontDepthTexture)
-END_GLOBAL_SHADER_PARAMETER_STRUCT()
+	SHADER_PARAMETER(FIntVector, VoxelResolution)
+	SHADER_PARAMETER(uint32, MacroGroupId)
+	SHADER_PARAMETER(FIntPoint, ViewportResolution)
+	SHADER_PARAMETER_RDG_BUFFER_SRV(StructuredBuffer<FVoxelizationViewInfo>, VoxelizationViewInfoBuffer)
+
+	SHADER_PARAMETER_RDG_TEXTURE_UAV(RWTexture3D<uint>, DensityTexture)
+	SHADER_PARAMETER_RDG_TEXTURE_UAV(RWTexture3D<uint>, TangentXTexture)
+	SHADER_PARAMETER_RDG_TEXTURE_UAV(RWTexture3D<uint>, TangentYTexture)
+	SHADER_PARAMETER_RDG_TEXTURE_UAV(RWTexture3D<uint>, TangentZTexture)
+	SHADER_PARAMETER_RDG_TEXTURE_UAV(RWTexture3D<uint>, MaterialTexture)
+
+	RENDER_TARGET_BINDING_SLOTS()
+END_SHADER_PARAMETER_STRUCT()
 
-// to be removed
-void RasterHairStrands(
-	FRHICommandList& RHICmdList,
+void AddHairVoxelizationRasterPass(
+	FRDGBuilder& GraphBuilder,
 	const FScene* Scene,
 	const FViewInfo* ViewInfo,
-	const FHairStrandsClusterData::TPrimitiveInfos& PrimitiveSceneInfos,
+	const FHairStrandsMacroGroupData::TPrimitiveInfos& PrimitiveSceneInfos,
 	const EHairStrandsRasterPassType ShadowPassType,
-	const FIntRect& AtlasRect,
-	const TUniformBufferRef<FViewUniformShaderParameters>& ViewUniformShaderParameters,
-	const TUniformBufferRef<FDeepShadowPassUniformParameters>& DeepShadowPassUniformParameters);
\ No newline at end of file
+	const FIntRect& ViewportRect,
+	const FVector4& HairRenderInfo,
+	const FVector& RasterDirection,
+	FHairVoxelizationRasterPassParameters* PassParameters);
diff --git a/Engine/Source/Runtime/Renderer/Private/HairStrands/HairStrandsRendering.h b/Engine/Source/Runtime/Renderer/Private/HairStrands/HairStrandsRendering.h
--- a/Engine/Source/Runtime/Renderer/Private/HairStrands/HairStrandsRendering.h
+++ b/Engine/Source/Runtime/Renderer/Private/HairStrands/HairStrandsRendering.h
@@ -11,6 +11,7 @@
 #include "Shader.h"
 #include "HairStrandsUtils.h"
 #include "HairStrandsCluster.h"
+#include "HairStrandsClusters.h"
 #include "HairStrandsLUT.h"
 #include "HairStrandsDeepShadow.h"
 #include "HairStrandsVoxelization.h"
@@ -26,7 +27,7 @@ struct FHairStrandsDatas
 {
 	FHairStrandsDeepShadowViews DeepShadowViews;
 	FHairStrandsVisibilityViews HairVisibilityViews;
-	FHairStrandsClusterViews HairClusterPerViews;
+	FHairStrandsMacroGroupViews MacroGroupsPerViews;
 };
 
 enum class EHairStrandsInterpolationType
@@ -38,5 +39,9 @@ enum class EHairStrandsInterpolationType
 void RunHairStrandsInterpolation(
 	FRHICommandListImmediate& RHICmdList, 
 	EWorldType::Type WorldType, 
+	const struct FShaderDrawDebugData* DebugShaderData,
 	TShaderMap<FGlobalShaderType>* ShaderMap, 
-	EHairStrandsInterpolationType Type);
\ No newline at end of file
+	EHairStrandsInterpolationType Type,
+	FHairStrandClusterData* ClusterData);
+
+
diff --git a/Engine/Source/Runtime/Renderer/Private/HairStrands/HairStrandsScatter.cpp b/Engine/Source/Runtime/Renderer/Private/HairStrands/HairStrandsScatter.cpp
new file mode 100644
--- /dev/null
+++ b/Engine/Source/Runtime/Renderer/Private/HairStrands/HairStrandsScatter.cpp
@@ -0,0 +1,311 @@
+// Copyright Epic Games, Inc. All Rights Reserved.
+
+#include "HairStrandsScatter.h"
+#include "Shader.h"
+#include "GlobalShader.h"
+#include "ShaderParameters.h"
+#include "ShaderParameterStruct.h"
+#include "ShaderParameterMacros.h"
+#include "RenderGraphUtils.h"
+#include "PostProcessing.h"
+#include "HairStrandsRendering.h"
+#include "GpuDebugRendering.h"
+
+/////////////////////////////////////////////////////////////////////////////////////////
+
+static int32 GHairStrandsScatter_Enable = 0;
+static int32 GHairStrandsScatter_Debug = 0;
+static float GHairStrandsScatter_WorldRadius = 1;
+static int32 GHairStrandsScatter_PassCount = 1;
+static int32 GHairStrandsScatter_SampleCount = 8;
+static FAutoConsoleVariableRef CVarHairStrandsScatter_WorldRadius(TEXT("r.HairStrands.Scatter.WorldRadius"), GHairStrandsScatter_WorldRadius, TEXT("Gather radius in world space (in cm)."));
+static FAutoConsoleVariableRef CVarHairStrandsScatter_Enable(TEXT("r.HairStrands.Scatter"), GHairStrandsScatter_Enable, TEXT("Enable screen space hair scattering."));
+static FAutoConsoleVariableRef CVarHairStrandsScatter_Debug(TEXT("r.HairStrands.Scatter.Debug"), GHairStrandsScatter_Debug, TEXT("Enable debug view of screen space hair scattering."));
+static FAutoConsoleVariableRef CVarHairStrandsScatter_PassCount(TEXT("r.HairStrands.Scatter.IterationCount"), GHairStrandsScatter_PassCount, TEXT("Number of diffusion iterations."));
+static FAutoConsoleVariableRef CVarHairStrandsScatter_SampleCount(TEXT("r.HairStrands.Scatter.SampleCount"), GHairStrandsScatter_SampleCount, TEXT("Number of sample using during the scattering integration."));
+
+///////////////////////////////////////////////////////////////////////////////////////////////////
+class FHairComposePS : public FGlobalShader
+{
+	DECLARE_GLOBAL_SHADER(FHairComposePS);
+	SHADER_USE_PARAMETER_STRUCT(FHairComposePS, FGlobalShader);
+
+	BEGIN_SHADER_PARAMETER_STRUCT(FParameters, )
+		SHADER_PARAMETER_RDG_TEXTURE(Texture2D, CategorizationTexture)
+		SHADER_PARAMETER_RDG_TEXTURE(Texture2D, SceneColorTexture)
+		SHADER_PARAMETER_RDG_TEXTURE(Texture2D, SceneColorSubPixelTexture)
+		RENDER_TARGET_BINDING_SLOTS()
+	END_SHADER_PARAMETER_STRUCT()
+
+public:
+	static bool ShouldCompilePermutation(const FGlobalShaderPermutationParameters& Parameters) { return Parameters.Platform == SP_PCD3D_SM5; }
+};
+
+IMPLEMENT_GLOBAL_SHADER(FHairComposePS, "/Engine/Private/HairStrands/HairScatterCompose.usf", "MainPS", SF_Pixel);
+
+static FRDGTextureRef AddPreScatterComposePass(
+	FRDGBuilder& GraphBuilder,
+	const FViewInfo& View,
+	const FRDGTextureRef& InCategorizationTexture,
+	const FRDGTextureRef& InSceneColorTexture,
+	const FRDGTextureRef& InSceneColorSubPixelTexture)
+{
+	const FIntPoint Resolution = InSceneColorTexture->Desc.Extent;
+	FRDGTextureRef OutputTexture;
+	{
+		FRDGTextureDesc Desc;
+		Desc.Extent = Resolution;
+		Desc.Depth = 0;
+		Desc.Format = InSceneColorTexture->Desc.Format;
+		Desc.NumMips = 1;
+		Desc.NumSamples = 1;
+		Desc.Flags = TexCreate_None;
+		Desc.TargetableFlags = TexCreate_RenderTargetable | TexCreate_ShaderResource;
+		Desc.ClearValue = FClearValueBinding(0);
+		OutputTexture = GraphBuilder.CreateTexture(Desc, TEXT("HairComposedTexture"));
+	}
+
+	FVector2D PixelVelocity(1.f / (Resolution.X * 2), 1.f / (Resolution.Y * 2));
+
+	FHairComposePS::FParameters* Parameters = GraphBuilder.AllocParameters<FHairComposePS::FParameters>();
+	Parameters->CategorizationTexture = InCategorizationTexture;
+	Parameters->SceneColorTexture = InSceneColorTexture;
+	Parameters->SceneColorSubPixelTexture = InSceneColorSubPixelTexture;
+	Parameters->RenderTargets[0] = FRenderTargetBinding(OutputTexture, ERenderTargetLoadAction::ENoAction);
+
+	TShaderMapRef<FPostProcessVS> VertexShader(View.ShaderMap);
+	TShaderMapRef<FHairComposePS> PixelShader(View.ShaderMap);
+	const TShaderMap<FGlobalShaderType>* GlobalShaderMap = View.ShaderMap;
+	const FIntRect Viewport = View.ViewRect;
+
+	const FViewInfo* CapturedView = &View;
+
+	{
+		ClearUnusedGraphResources(*PixelShader, Parameters);
+
+		GraphBuilder.AddPass(
+			RDG_EVENT_NAME("HairCompose"),
+			Parameters,
+			ERDGPassFlags::Raster,
+			[Parameters, VertexShader, PixelShader, Viewport, Resolution, CapturedView](FRHICommandList& RHICmdList)
+		{
+			FGraphicsPipelineStateInitializer GraphicsPSOInit;
+			RHICmdList.ApplyCachedRenderTargets(GraphicsPSOInit);
+			GraphicsPSOInit.BlendState = TStaticBlendState<CW_RGBA, BO_Add, BF_One, BF_Zero, BO_Add, BF_One, BF_Zero>::GetRHI();
+			GraphicsPSOInit.RasterizerState = TStaticRasterizerState<>::GetRHI();
+			GraphicsPSOInit.DepthStencilState = TStaticDepthStencilState<false, CF_Always>::GetRHI();
+			GraphicsPSOInit.BoundShaderState.VertexDeclarationRHI = GFilterVertexDeclaration.VertexDeclarationRHI;
+			GraphicsPSOInit.BoundShaderState.VertexShaderRHI = GETSAFERHISHADER_VERTEX(*VertexShader);
+			GraphicsPSOInit.BoundShaderState.PixelShaderRHI = GETSAFERHISHADER_PIXEL(*PixelShader);
+			GraphicsPSOInit.PrimitiveType = PT_TriangleList;
+			SetGraphicsPipelineState(RHICmdList, GraphicsPSOInit);
+
+			VertexShader->SetParameters(RHICmdList, CapturedView->ViewUniformBuffer);
+			RHICmdList.SetViewport(Viewport.Min.X, Viewport.Min.Y, 0.0f, Viewport.Max.X, Viewport.Max.Y, 1.0f);
+			SetShaderParameters(RHICmdList, *PixelShader, PixelShader->GetPixelShader(), *Parameters);
+			DrawRectangle(
+				RHICmdList,
+				0, 0,
+				Viewport.Width(), Viewport.Height(),
+				Viewport.Min.X, Viewport.Min.Y,
+				Viewport.Width(), Viewport.Height(),
+				Viewport.Size(),
+				Resolution,
+				*VertexShader,
+				EDRF_UseTriangleOptimization);
+		});
+	}
+
+	return OutputTexture;
+}
+
+///////////////////////////////////////////////////////////////////////////////////////////////////
+class FHairScatterPS : public FGlobalShader
+{
+	DECLARE_GLOBAL_SHADER(FHairScatterPS);
+	SHADER_USE_PARAMETER_STRUCT(FHairScatterPS, FGlobalShader);
+
+	BEGIN_SHADER_PARAMETER_STRUCT(FParameters, )
+		SHADER_PARAMETER_STRUCT_INCLUDE(ShaderDrawDebug::FShaderDrawDebugParameters, ShaderDrawParameters)
+		SHADER_PARAMETER(FVector2D, OutputResolution)
+		SHADER_PARAMETER(float, GatherWorldRadius)
+		SHADER_PARAMETER(float, PixelRadiusAtDepth1)
+		SHADER_PARAMETER(uint32, Enable)
+		SHADER_PARAMETER(uint32, Debug)
+		SHADER_PARAMETER(uint32, SampleCount)
+		SHADER_PARAMETER(uint32, HairComponents)
+		SHADER_PARAMETER_RDG_TEXTURE(Texture3D, HairLUTTexture)
+		SHADER_PARAMETER_RDG_TEXTURE(Texture3D, HairEnergyLUTTexture)
+		SHADER_PARAMETER_RDG_BUFFER_UAV(StructuredBuffer, OutputColor)
+		SHADER_PARAMETER_SAMPLER(SamplerState, LinearSampler)
+		SHADER_PARAMETER_RDG_TEXTURE(Texture2D, VisibilityNodeIndex)
+		SHADER_PARAMETER_RDG_BUFFER_SRV(StructuredBuffer, VisibilityNodeData)
+		SHADER_PARAMETER_RDG_TEXTURE(Texture2D, DiffusionInputTexture)
+		SHADER_PARAMETER_RDG_TEXTURE(Texture2D, CategorizationTexture)
+		SHADER_PARAMETER_STRUCT_REF(FViewUniformShaderParameters, ViewUniformBuffer)
+		SHADER_PARAMETER_STRUCT_REF(FVirtualVoxelParameters, VirtualVoxel)
+		RENDER_TARGET_BINDING_SLOTS()
+	END_SHADER_PARAMETER_STRUCT()
+
+public:
+	static bool ShouldCompilePermutation(const FGlobalShaderPermutationParameters& Parameters) { return Parameters.Platform == SP_PCD3D_SM5; }
+};
+
+IMPLEMENT_GLOBAL_SHADER(FHairScatterPS, "/Engine/Private/HairStrands/HairScatter.usf", "MainPS", SF_Pixel);
+
+static FRDGTextureRef AddScatterPass(
+	FRDGBuilder& GraphBuilder,
+	const FViewInfo& View,
+	const FVirtualVoxelResources& VoxelResources,
+	const FRDGTextureRef& InVisibilityNodeIndex,
+	const FRDGBufferRef&  InVisibilityNodeData,
+	const FRDGTextureRef& InCategorizationTexture,
+	const FRDGTextureRef& InDiffusionInput,
+	const FRDGTextureRef& OutSceneColorTexture,
+	const FRDGTextureRef& OutSceneColorSubPixelTexture)
+{
+	if (!VoxelResources.IsValid())
+		return nullptr;
+
+	const FIntPoint Resolution = OutSceneColorTexture->Desc.Extent;
+	const FHairLUT InHairLUT = GetHairLUT(GraphBuilder.RHICmdList, View);
+
+	float PixelRadiusAtDepth1 = 0;
+	{
+		const float DiameterToRadius = 0.5f;
+		const float hFOV = FMath::DegreesToRadians(View.FOV);
+		const float DiameterAtDepth1 = FMath::Tan(hFOV * 0.5f) / (0.5f * Resolution.X);
+		PixelRadiusAtDepth1 = DiameterAtDepth1 * DiameterToRadius;
+	}
+
+	FRDGTextureRef OutDiffusionOutput = GraphBuilder.CreateTexture(InDiffusionInput->Desc, TEXT("HairDiffusionOutput"));
+
+	FHairScatterPS::FParameters* Parameters = GraphBuilder.AllocParameters<FHairScatterPS::FParameters>();
+	Parameters->Enable = GHairStrandsScatter_Enable > 0 ? 1 : 0;
+	Parameters->Debug = GHairStrandsScatter_Debug > 0? 1 : 0;
+	Parameters->SampleCount = FMath::Clamp(GHairStrandsScatter_SampleCount, 1, 32);
+	Parameters->GatherWorldRadius = FMath::Clamp(GHairStrandsScatter_WorldRadius, 0.f, 100.f);
+	Parameters->PixelRadiusAtDepth1 = PixelRadiusAtDepth1;
+	Parameters->HairComponents = ToBitfield(GetHairComponents());
+	Parameters->OutputResolution = Resolution;
+	Parameters->VisibilityNodeIndex = InVisibilityNodeIndex;
+	Parameters->VisibilityNodeData = GraphBuilder.CreateSRV(InVisibilityNodeData);
+	Parameters->CategorizationTexture = InCategorizationTexture;
+	Parameters->DiffusionInputTexture = InDiffusionInput;
+	Parameters->HairLUTTexture = GraphBuilder.RegisterExternalTexture(InHairLUT.Textures[HairLUTType_DualScattering], TEXT("HairLUTTexture"));
+	Parameters->HairEnergyLUTTexture = GraphBuilder.RegisterExternalTexture(InHairLUT.Textures[HairLUTType_MeanEnergy], TEXT("HairEnergyLUTTexture"));
+	Parameters->LinearSampler = TStaticSamplerState<SF_Bilinear, AM_Clamp, AM_Clamp, AM_Clamp>::GetRHI();
+	Parameters->ViewUniformBuffer = View.ViewUniformBuffer;
+	Parameters->VirtualVoxel = VoxelResources.UniformBuffer;
+	if (ShaderDrawDebug::IsShaderDrawDebugEnabled(View))
+	{
+		ShaderDrawDebug::SetParameters(GraphBuilder, View.ShaderDrawData, Parameters->ShaderDrawParameters);
+	}
+	Parameters->RenderTargets[0] = FRenderTargetBinding(OutSceneColorTexture, ERenderTargetLoadAction::ELoad);
+	Parameters->RenderTargets[1] = FRenderTargetBinding(OutSceneColorSubPixelTexture, ERenderTargetLoadAction::ELoad);
+	Parameters->RenderTargets[2] = FRenderTargetBinding(OutDiffusionOutput, ERenderTargetLoadAction::ENoAction);
+
+	TShaderMapRef<FPostProcessVS> VertexShader(View.ShaderMap);
+	TShaderMapRef<FHairScatterPS> PixelShader(View.ShaderMap);
+	const TShaderMap<FGlobalShaderType>* GlobalShaderMap = View.ShaderMap;
+	const FIntRect Viewport = View.ViewRect;
+
+	const FViewInfo* CapturedView = &View;
+
+	{
+		ClearUnusedGraphResources(*PixelShader, Parameters);
+
+		GraphBuilder.AddPass(
+			RDG_EVENT_NAME("HairScatter"),
+			Parameters,
+			ERDGPassFlags::Raster,
+			[Parameters, VertexShader, PixelShader, Viewport, Resolution, CapturedView](FRHICommandList& RHICmdList)
+		{
+			FGraphicsPipelineStateInitializer GraphicsPSOInit;
+			RHICmdList.ApplyCachedRenderTargets(GraphicsPSOInit);
+			if (GHairStrandsScatter_Debug > 0)
+			{
+				GraphicsPSOInit.BlendState = TStaticBlendState<
+					CW_RGBA, BO_Add, BF_One, BF_Zero, BO_Add, BF_One, BF_Zero,
+					CW_RGBA, BO_Add, BF_One, BF_Zero, BO_Add, BF_One, BF_Zero > ::GetRHI();
+			}
+			else
+			{
+				GraphicsPSOInit.BlendState = TStaticBlendState<
+					CW_RGBA, BO_Add, BF_One, BF_One, BO_Add, BF_Zero, BF_One,
+					CW_RGBA, BO_Add, BF_One, BF_One, BO_Add, BF_Zero, BF_One > ::GetRHI();
+			}
+			GraphicsPSOInit.RasterizerState = TStaticRasterizerState<>::GetRHI();
+			GraphicsPSOInit.DepthStencilState = TStaticDepthStencilState<false, CF_Always>::GetRHI();
+			GraphicsPSOInit.BoundShaderState.VertexDeclarationRHI = GFilterVertexDeclaration.VertexDeclarationRHI;
+			GraphicsPSOInit.BoundShaderState.VertexShaderRHI = GETSAFERHISHADER_VERTEX(*VertexShader);
+			GraphicsPSOInit.BoundShaderState.PixelShaderRHI = GETSAFERHISHADER_PIXEL(*PixelShader);
+			GraphicsPSOInit.PrimitiveType = PT_TriangleList;
+			SetGraphicsPipelineState(RHICmdList, GraphicsPSOInit);
+
+			VertexShader->SetParameters(RHICmdList, CapturedView->ViewUniformBuffer);
+			RHICmdList.SetViewport(Viewport.Min.X, Viewport.Min.Y, 0.0f, Viewport.Max.X, Viewport.Max.Y, 1.0f);
+			SetShaderParameters(RHICmdList, *PixelShader, PixelShader->GetPixelShader(), *Parameters);
+			DrawRectangle(
+				RHICmdList,
+				0, 0,
+				Viewport.Width(), Viewport.Height(),
+				Viewport.Min.X, Viewport.Min.Y,
+				Viewport.Width(), Viewport.Height(),
+				Viewport.Size(),
+				Resolution,
+				*VertexShader,
+				EDRF_UseTriangleOptimization);
+		});
+	}
+
+	return OutDiffusionOutput;
+}
+
+///////////////////////////////////////////////////////////////////////////////////////////////////
+
+void AddHairDiffusionPass(
+	FRDGBuilder& GraphBuilder,
+	const FViewInfo& View,
+	const FHairStrandsVisibilityData& VisibilityData,
+	const FVirtualVoxelResources& VoxelResources,
+	const FRDGTextureRef SceneColorDepth,
+	FRDGTextureRef OutSceneColorSubPixelTexture,
+	FRDGTextureRef OutSceneColorTexture)
+{
+	const uint32 DiffusionPassCount = FMath::Clamp(GHairStrandsScatter_PassCount, 0, 8);
+	const bool bIsEnabled = DiffusionPassCount > 0 &&
+		GHairStrandsScatter_Enable > 0 &&
+		VisibilityData.NodeIndex &&
+		VisibilityData.NodeData &&
+		VisibilityData.CategorizationTexture;
+
+	if (!bIsEnabled)
+		return;
+
+	FRDGTextureRef NodedIndex = GraphBuilder.RegisterExternalTexture(VisibilityData.NodeIndex, TEXT("HairNodeIndex"));
+	FRDGBufferRef  NodedData = GraphBuilder.RegisterExternalBuffer(VisibilityData.NodeData, TEXT("HairNodeData"));
+	FRDGTextureRef CategorisationTexture = GraphBuilder.RegisterExternalTexture(VisibilityData.CategorizationTexture, TEXT("HairVisibilityCategorizationTexture"));
+
+	FRDGTextureRef DiffusionInput = AddPreScatterComposePass(
+		GraphBuilder,
+		View,
+		CategorisationTexture,
+		OutSceneColorTexture,
+		OutSceneColorSubPixelTexture);
+
+	for (uint32 DiffusionPassIt = 0; DiffusionPassIt < DiffusionPassCount; ++DiffusionPassIt)
+	{
+		FRDGTextureRef DiffusionOutput = AddScatterPass(
+			GraphBuilder,
+			View,
+			VoxelResources,
+			NodedIndex,
+			NodedData,
+			CategorisationTexture,
+			DiffusionInput,
+			OutSceneColorTexture,
+			OutSceneColorSubPixelTexture);
+		DiffusionInput = DiffusionOutput;
+	}
+}
\ No newline at end of file
diff --git a/Engine/Source/Runtime/Renderer/Private/HairStrands/HairStrandsScatter.h b/Engine/Source/Runtime/Renderer/Private/HairStrands/HairStrandsScatter.h
new file mode 100644
--- /dev/null
+++ b/Engine/Source/Runtime/Renderer/Private/HairStrands/HairStrandsScatter.h
@@ -0,0 +1,16 @@
+// Copyright Epic Games, Inc. All Rights Reserved.
+
+#pragma once
+
+#include "CoreMinimal.h"
+#include "RendererInterface.h"
+#include "RenderGraphResources.h"
+
+void AddHairDiffusionPass(
+	FRDGBuilder& GraphBuilder,
+	const FViewInfo& View,
+	const struct FHairStrandsVisibilityData& VisibilityData,
+	const struct FVirtualVoxelResources& VoxelResources,
+	const FRDGTextureRef SceneColorDepth,
+	FRDGTextureRef OutSceneColorSubPixelTexture,
+	FRDGTextureRef OutSceneColorTexture);
\ No newline at end of file
diff --git a/Engine/Source/Runtime/Renderer/Private/HairStrands/HairStrandsTransmittance.cpp b/Engine/Source/Runtime/Renderer/Private/HairStrands/HairStrandsTransmittance.cpp
--- a/Engine/Source/Runtime/Renderer/Private/HairStrands/HairStrandsTransmittance.cpp
+++ b/Engine/Source/Runtime/Renderer/Private/HairStrands/HairStrandsTransmittance.cpp
@@ -14,6 +14,8 @@
 #include "SceneTextureParameters.h"
 #include "RenderGraphUtils.h"
 #include "PostProcessing.h"
+#include "GpuDebugRendering.h"
+#include "ShaderPrintParameters.h"
 
 ///////////////////////////////////////////////////////////////////////////////////////////////////
 
@@ -37,6 +39,12 @@ static float GDeepShadowDepthBiasScale = 2;	// Default is arbitrary, based on co
 static FAutoConsoleVariableRef CVarDeepShadowDensityScale(TEXT("r.HairStrands.DeepShadow.DensityScale"), GDeepShadowDensityScale, TEXT("Set density scale for compensating the lack of hair fiber in an asset"));
 static FAutoConsoleVariableRef CVarDeepShadowDepthBiasScale(TEXT("r.HairStrands.DeepShadow.DepthBiasScale"), GDeepShadowDepthBiasScale, TEXT("Set depth bias scale for transmittance computation"));
 
+static int32 GHairStrandsTransmittanceSuperSampling = 0;
+static FAutoConsoleVariableRef CVarHairStrandsTransmittanceSuperSampling(TEXT("r.HairStrands.DeepShadow.SuperSampling"), GHairStrandsTransmittanceSuperSampling, TEXT("Evaluate transmittance with supersampling. This is expensive and intended to be used only in cine mode."), ECVF_Scalability | ECVF_RenderThreadSafe);
+
+static int32 GHairStrandsTransmittanceMaskUseMipTraversal = 1;
+static FAutoConsoleVariableRef CVarHairStrandsTransmittanceMaskUseMipTraversal(TEXT("r.HairStrands.DeepShadow.MipTraversal"), GHairStrandsTransmittanceMaskUseMipTraversal, TEXT("Evaluate transmittance using mip-map traversal (faster)."), ECVF_Scalability | ECVF_RenderThreadSafe);
+
 float GetDeepShadowDensityScale() { return FMath::Max(0.0f, GDeepShadowDensityScale); }
 float GetDeepShadowDepthBiasScale() { return FMath::Max(0.0f, GDeepShadowDepthBiasScale); }
 ///////////////////////////////////////////////////////////////////////////////////////////////////
@@ -60,6 +68,7 @@ enum FHairTransmittanceType
 {
 	FHairTransmittanceType_DeepShadow,
 	FHairTransmittanceType_Voxel,
+	FHairTransmittanceType_VirtualVoxel,
 	FHairTransmittanceTypeCount
 };
 
@@ -70,14 +79,19 @@ class FDeepTransmittanceMaskCS : public FGlobalShader
 
 	class FTransmittanceType : SHADER_PERMUTATION_INT("PERMUTATION_TRANSMITTANCE_TYPE", FHairTransmittanceTypeCount);
 	class FTransmittanceGroupSize : SHADER_PERMUTATION_INT("PERMUTATION_GROUP_SIZE", 2);
-	using FPermutationDomain = TShaderPermutationDomain<FTransmittanceType, FTransmittanceGroupSize>;
+	class FSuperSampling : SHADER_PERMUTATION_INT("PERMUTATION_SUPERSAMPLING", 2);
+	class FTraversal : SHADER_PERMUTATION_INT("PERMUTATION_TRAVERSAL", 2);
+	using FPermutationDomain = TShaderPermutationDomain<FTransmittanceType, FTransmittanceGroupSize, FSuperSampling, FTraversal>;
 
 	BEGIN_SHADER_PARAMETER_STRUCT(FParameters, )
 		SHADER_PARAMETER_STRUCT_INCLUDE(FSceneTextureParameters, SceneTextures)
+		SHADER_PARAMETER_STRUCT_INCLUDE(ShaderDrawDebug::FShaderDrawDebugParameters, ShaderDrawParameters)
+		SHADER_PARAMETER_STRUCT_INCLUDE(ShaderPrint::FShaderParameters, ShaderPrintParameters)
 
-		SHADER_PARAMETER_ARRAY(FIntVector4, DeepShadow_AtlasSlotOffsets, [FHairStrandsDeepShadowData::MaxClusterCount])
-		SHADER_PARAMETER_ARRAY(FMatrix, DeepShadow_WorldToLightTransforms, [FHairStrandsDeepShadowData::MaxClusterCount])
+		SHADER_PARAMETER_ARRAY(FIntVector4, DeepShadow_AtlasSlotOffsets, [FHairStrandsDeepShadowData::MaxMacroGroupCount])
+		SHADER_PARAMETER_ARRAY(FMatrix, DeepShadow_WorldToLightTransforms, [FHairStrandsDeepShadowData::MaxMacroGroupCount])
 		SHADER_PARAMETER(FIntPoint, DeepShadow_Resolution)
+		SHADER_PARAMETER(float, LightRadius)
 		SHADER_PARAMETER(FVector, LightDirection)
 		SHADER_PARAMETER(uint32, MaxVisibilityNodeCount)
 		SHADER_PARAMETER(FVector4, LightPosition)
@@ -88,8 +102,8 @@ class FDeepTransmittanceMaskCS : public FGlobalShader
 		SHADER_PARAMETER(uint32, DeepShadow_DebugMode)
 		SHADER_PARAMETER(FMatrix, DeepShadow_ShadowToWorld)
 
-		SHADER_PARAMETER_ARRAY(FVector4, Voxel_MinAABBs, [FHairStrandsDeepShadowData::MaxClusterCount])
-		SHADER_PARAMETER_ARRAY(FVector4, Voxel_MaxAABBs, [FHairStrandsDeepShadowData::MaxClusterCount])
+		SHADER_PARAMETER_ARRAY(FVector4, Voxel_MinAABBs, [FHairStrandsDeepShadowData::MaxMacroGroupCount])
+		SHADER_PARAMETER_ARRAY(FVector4, Voxel_MaxAABBs, [FHairStrandsDeepShadowData::MaxMacroGroupCount])
 		SHADER_PARAMETER(uint32, Voxel_Resolution)
 		SHADER_PARAMETER_RDG_TEXTURE(Texture3D, Voxel_DensityTexture0)
 		SHADER_PARAMETER_RDG_TEXTURE(Texture3D, Voxel_DensityTexture1)
@@ -100,6 +114,8 @@ class FDeepTransmittanceMaskCS : public FGlobalShader
 		SHADER_PARAMETER_RDG_TEXTURE(Texture3D, Voxel_DensityTexture6)
 		SHADER_PARAMETER_RDG_TEXTURE(Texture3D, Voxel_DensityTexture7)
 
+		SHADER_PARAMETER_RDG_TEXTURE(Texture2D, RayMarchMaskTexture)
+
 		SHADER_PARAMETER_RDG_TEXTURE(Texture2D, DeepShadow_FrontDepthTexture)
 		SHADER_PARAMETER_RDG_TEXTURE(Texture2D, DeepShadow_DomTexture)
 		SHADER_PARAMETER_RDG_TEXTURE(Texture3D, HairLUTTexture)
@@ -112,21 +128,51 @@ class FDeepTransmittanceMaskCS : public FGlobalShader
 		SHADER_PARAMETER_SAMPLER(SamplerState, ShadowSampler)
 
 		SHADER_PARAMETER_STRUCT_REF(FViewUniformShaderParameters, ViewUniformBuffer)
+		SHADER_PARAMETER_STRUCT_REF(FVirtualVoxelParameters, VirtualVoxel)
 		END_SHADER_PARAMETER_STRUCT()
 
 public:
-	static bool ShouldCompilePermutation(const FGlobalShaderPermutationParameters& Parameters) { return IsHairStrandsSupported(Parameters.Platform); }
+	static bool ShouldCompilePermutation(const FGlobalShaderPermutationParameters& Parameters) 
+	{ 
+		FPermutationDomain PermutationVector(Parameters.PermutationId);
+		if (PermutationVector.Get<FTransmittanceType>() == FHairTransmittanceType_DeepShadow && 
+			(PermutationVector.Get<FSuperSampling>() == 1 || PermutationVector.Get<FTraversal>() == 1))
+		{
+			return false;
+		}
+
+		if (PermutationVector.Get<FTransmittanceType>() == FHairTransmittanceType_Voxel && PermutationVector.Get<FTraversal>() == 1)
+		{
+			return false;
+		}
+		return IsHairStrandsSupported(Parameters.Platform);
+	}
+
+	static FPermutationDomain RemapPermutation(FPermutationDomain PermutationVector)
+	{
+		if (PermutationVector.Get<FTransmittanceType>() == FHairTransmittanceType_DeepShadow)
+		{
+			PermutationVector.Set<FSuperSampling>(0);
+			PermutationVector.Set<FTraversal>(0);
+		}
+		else if (PermutationVector.Get<FTransmittanceType>() == FHairTransmittanceType_Voxel)
+		{
+			PermutationVector.Set<FTraversal>(0);
+		}
+		return PermutationVector;
+	}
 };
 
 IMPLEMENT_GLOBAL_SHADER(FDeepTransmittanceMaskCS, "/Engine/Private/HairStrands/HairStrandsDeepTransmittanceMask.usf", "MainCS", SF_Compute);
 
 struct FDeepShadowTransmittanceParams
 {
-	FIntVector4 DeepShadow_AtlasSlotOffsets[FHairStrandsDeepShadowData::MaxClusterCount];
-	FMatrix DeepShadow_WorldToLightTransforms[FHairStrandsDeepShadowData::MaxClusterCount];
+	FIntVector4 DeepShadow_AtlasSlotOffsets[FHairStrandsDeepShadowData::MaxMacroGroupCount];
+	FMatrix DeepShadow_WorldToLightTransforms[FHairStrandsDeepShadowData::MaxMacroGroupCount];
 	FIntPoint DeepShadow_Resolution = FIntPoint(0, 0);
 	FVector LightDirection = FVector::ZeroVector;
 	FVector4 LightPosition = FVector4(0, 0, 0, 0);
+	float LightRadius = 0;
 	float DepthBiasScale = 0;
 	float DensityScale = 0;
 	FMatrix DeepShadow_ShadowToWorld = FMatrix::Identity;
@@ -137,10 +183,12 @@ struct FDeepShadowTransmittanceParams
 	FRDGBufferRef HairVisibilityNodeData = nullptr;
 	FRDGBufferRef HairVisibilityNodeCoord = nullptr;
 
-	FRDGTextureRef Cluster_DensityTextures[FHairStrandsDeepShadowData::MaxClusterCount];
-	FVector4 Cluster_MinAABBs[FHairStrandsDeepShadowData::MaxClusterCount];
-	FVector4 Cluster_MaxAABBs[FHairStrandsDeepShadowData::MaxClusterCount];
-	uint32  Cluster_Resolution;
+	FRDGTextureRef Voxel_DensityTextures[FHairStrandsDeepShadowData::MaxMacroGroupCount];
+	FVector4 Voxel_MinAABBs[FHairStrandsDeepShadowData::MaxMacroGroupCount];
+	FVector4 Voxel_MaxAABBs[FHairStrandsDeepShadowData::MaxMacroGroupCount];
+	uint32   Voxel_Resolution;
+
+	const FVirtualVoxelResources* VirtualVoxelResources = nullptr;
 };
 
 // Transmittance mask
@@ -152,7 +200,8 @@ static FRDGBufferRef AddDeepShadowTransmittanceMaskPass(
 	const FDeepShadowTransmittanceParams& Params,
 	const uint32 NodeGroupSize,
 	FRDGTextureRef HairLUTTexture,
-	FRDGBufferRef IndirectArgsBuffer)
+	FRDGBufferRef IndirectArgsBuffer,
+	TRefCountPtr<IPooledRenderTarget>& ScreenShadowMaskSubPixelTexture)
 {
 	FRDGBufferRef OutBuffer = GraphBuilder.CreateBuffer(FRDGBufferDesc::CreateStructuredDesc(
 		4 * sizeof(float),
@@ -170,6 +219,7 @@ static FRDGBufferRef AddDeepShadowTransmittanceMaskPass(
 	Parameters->DeepShadow_Resolution = Params.DeepShadow_Resolution;
 	Parameters->LightDirection = Params.LightDirection;
 	Parameters->LightPosition = Params.LightPosition;
+	Parameters->LightRadius = Params.LightRadius;
 	Parameters->DepthBiasScale = Params.DepthBiasScale;
 	Parameters->DensityScale = Params.DensityScale;
 	Parameters->DeepShadow_KernelAperture = GetDeepShadowKernelAperture();
@@ -179,28 +229,48 @@ static FRDGBufferRef AddDeepShadowTransmittanceMaskPass(
 	Parameters->IndirectArgsBuffer = IndirectArgsBuffer;
 	Parameters->MaxVisibilityNodeCount = Params.HairVisibilityNodeData->Desc.NumElements;
 
-	memcpy(&(Parameters->DeepShadow_AtlasSlotOffsets[0]), Params.DeepShadow_AtlasSlotOffsets, sizeof(FIntVector4) * FHairStrandsDeepShadowData::MaxClusterCount);
-	memcpy(&(Parameters->DeepShadow_WorldToLightTransforms[0]), Params.DeepShadow_WorldToLightTransforms, sizeof(FMatrix) * FHairStrandsDeepShadowData::MaxClusterCount);
-
-	memcpy(&(Parameters->Voxel_MinAABBs[0]), Params.Cluster_MinAABBs, sizeof(FVector4) * FHairStrandsDeepShadowData::MaxClusterCount);
-	memcpy(&(Parameters->Voxel_MaxAABBs[0]), Params.Cluster_MaxAABBs, sizeof(FVector4) * FHairStrandsDeepShadowData::MaxClusterCount);
-	Parameters->Voxel_Resolution = Params.Cluster_Resolution;
-	Parameters->Voxel_DensityTexture0 = Params.Cluster_DensityTextures[0];
-	Parameters->Voxel_DensityTexture1 = Params.Cluster_DensityTextures[1];
-	Parameters->Voxel_DensityTexture2 = Params.Cluster_DensityTextures[2];
-	Parameters->Voxel_DensityTexture3 = Params.Cluster_DensityTextures[3];
-	Parameters->Voxel_DensityTexture4 = Params.Cluster_DensityTextures[4];
-	Parameters->Voxel_DensityTexture5 = Params.Cluster_DensityTextures[5];
-	Parameters->Voxel_DensityTexture6 = Params.Cluster_DensityTextures[6];
-	Parameters->Voxel_DensityTexture7 = Params.Cluster_DensityTextures[7];
+	if (ShaderDrawDebug::IsShaderDrawDebugEnabled(View))
+	{
+		ShaderDrawDebug::SetParameters(GraphBuilder, View.ShaderDrawData, Parameters->ShaderDrawParameters);
+		ShaderPrint::SetParameters(View, Parameters->ShaderPrintParameters);
+	}
+
+	memcpy(&(Parameters->DeepShadow_AtlasSlotOffsets[0]), Params.DeepShadow_AtlasSlotOffsets, sizeof(FIntVector4) * FHairStrandsDeepShadowData::MaxMacroGroupCount);
+	memcpy(&(Parameters->DeepShadow_WorldToLightTransforms[0]), Params.DeepShadow_WorldToLightTransforms, sizeof(FMatrix) * FHairStrandsDeepShadowData::MaxMacroGroupCount);
+
+	memcpy(&(Parameters->Voxel_MinAABBs[0]), Params.Voxel_MinAABBs, sizeof(FVector4) * FHairStrandsDeepShadowData::MaxMacroGroupCount);
+	memcpy(&(Parameters->Voxel_MaxAABBs[0]), Params.Voxel_MaxAABBs, sizeof(FVector4) * FHairStrandsDeepShadowData::MaxMacroGroupCount);
+	Parameters->Voxel_Resolution = Params.Voxel_Resolution;
+	Parameters->Voxel_DensityTexture0 = Params.Voxel_DensityTextures[0];
+	Parameters->Voxel_DensityTexture1 = Params.Voxel_DensityTextures[1];
+	Parameters->Voxel_DensityTexture2 = Params.Voxel_DensityTextures[2];
+	Parameters->Voxel_DensityTexture3 = Params.Voxel_DensityTextures[3];
+	Parameters->Voxel_DensityTexture4 = Params.Voxel_DensityTextures[4];
+	Parameters->Voxel_DensityTexture5 = Params.Voxel_DensityTextures[5];
+	Parameters->Voxel_DensityTexture6 = Params.Voxel_DensityTextures[6];
+	Parameters->Voxel_DensityTexture7 = Params.Voxel_DensityTextures[7];
+
+	Parameters->RayMarchMaskTexture = GraphBuilder.RegisterExternalTexture(ScreenShadowMaskSubPixelTexture.IsValid() ? ScreenShadowMaskSubPixelTexture : GSystemTextures.WhiteDummy);
+
+	bool bIsSuperSampled = false;
+	if (TransmittanceType == FHairTransmittanceType_VirtualVoxel)
+	{
+		check(Params.VirtualVoxelResources);
+		Parameters->VirtualVoxel = Params.VirtualVoxelResources->UniformBuffer;
+		bIsSuperSampled = GHairStrandsTransmittanceSuperSampling > 0;
+	}
 
 	Parameters->HairVisibilityNodeData = GraphBuilder.CreateSRV(Params.HairVisibilityNodeData);
 	Parameters->HairVisibilityNodeCoord = GraphBuilder.CreateSRV(Params.HairVisibilityNodeCoord);
 
+	const bool bIsMipTraversal = GHairStrandsTransmittanceMaskUseMipTraversal > 0;
 	check(NodeGroupSize == 64 || NodeGroupSize == 32);
 	FDeepTransmittanceMaskCS::FPermutationDomain PermutationVector;
 	PermutationVector.Set<FDeepTransmittanceMaskCS::FTransmittanceType>(TransmittanceType);
 	PermutationVector.Set<FDeepTransmittanceMaskCS::FTransmittanceGroupSize>(NodeGroupSize == 64 ? 0 : (NodeGroupSize == 32 ? 1 : 2));
+	PermutationVector.Set<FDeepTransmittanceMaskCS::FSuperSampling>(bIsSuperSampled ? 1 : 0);
+	PermutationVector.Set<FDeepTransmittanceMaskCS::FTraversal>(bIsMipTraversal ? 1 : 0);
+	PermutationVector = FDeepTransmittanceMaskCS::RemapPermutation(PermutationVector);
 
 	TShaderMapRef<FDeepTransmittanceMaskCS> ComputeShader(View.ShaderMap, PermutationVector);
 	FComputeShaderUtils::AddPass(
@@ -221,6 +291,7 @@ enum FHairOpaqueMaskType
 {
 	FHairOpaqueMaskType_DeepShadow,
 	FHairOpaqueMaskType_Voxel,
+	FHairOpaqueMaskType_VirtualVoxel,
 	FHairOpaqueMaskTypeCount
 };
 
@@ -235,6 +306,8 @@ class FDeepShadowMaskPS : public FGlobalShader
 
 	BEGIN_SHADER_PARAMETER_STRUCT(FParameters, )
 		SHADER_PARAMETER_STRUCT_INCLUDE(FSceneTextureParameters, SceneTextures)
+		SHADER_PARAMETER_STRUCT_INCLUDE(ShaderDrawDebug::FShaderDrawDebugParameters, ShaderDrawParameters)
+		SHADER_PARAMETER_STRUCT_INCLUDE(ShaderPrint::FShaderParameters, ShaderPrintParameters)
 
 		SHADER_PARAMETER(FIntPoint, DeepShadow_SlotOffset)
 		SHADER_PARAMETER(FIntPoint, DeepShadow_SlotResolution)
@@ -247,6 +320,9 @@ class FDeepShadowMaskPS : public FGlobalShader
 		SHADER_PARAMETER(FVector, Voxel_MinAABB)
 		SHADER_PARAMETER(uint32, Voxel_Resolution)
 		SHADER_PARAMETER(FVector, Voxel_MaxAABB)
+		SHADER_PARAMETER(uint32, Voxel_MacroGroupId)
+
+		SHADER_PARAMETER_RDG_TEXTURE(Texture2D, RayMarchMaskTexture)
 
 		SHADER_PARAMETER_RDG_TEXTURE(Texture3D, Voxel_DensityTexture)
 		SHADER_PARAMETER_RDG_TEXTURE(Texture2D, DeepShadow_FrontDepthTexture)
@@ -255,6 +331,7 @@ class FDeepShadowMaskPS : public FGlobalShader
 		SHADER_PARAMETER_SAMPLER(SamplerState, ShadowSampler)
 
 		SHADER_PARAMETER_STRUCT_REF(FViewUniformShaderParameters, ViewUniformBuffer)
+		SHADER_PARAMETER_STRUCT_REF(FVirtualVoxelParameters, VirtualVoxel)
 		RENDER_TARGET_BINDING_SLOTS()
 		END_SHADER_PARAMETER_STRUCT()
 
@@ -281,6 +358,9 @@ struct FDeepShadowOpaqueParams
 	FVector			Voxel_MinAABB;
 	FVector			Voxel_MaxAABB;
 	uint32			Voxel_Resolution;
+	uint32			Voxel_MacroGroupId;
+
+	const FVirtualVoxelResources* Voxel_VirtualVoxel = nullptr;
 };
 
 // Opaque mask
@@ -307,15 +387,39 @@ static void AddDeepShadowOpaqueMaskPass(
 	Parameters->DeepShadow_SlotOffset = FIntPoint(Params.DeepShadow_AtlasRect.Min.X, Params.DeepShadow_AtlasRect.Min.Y);
 	Parameters->DeepShadow_SlotResolution = FIntPoint(Params.DeepShadow_AtlasRect.Max.X - Params.DeepShadow_AtlasRect.Min.X, Params.DeepShadow_AtlasRect.Max.Y - Params.DeepShadow_AtlasRect.Min.Y);
 
+	if (ShaderDrawDebug::IsShaderDrawDebugEnabled(View))
+	{
+		ShaderDrawDebug::SetParameters(GraphBuilder, View.ShaderDrawData, Parameters->ShaderDrawParameters);
+		ShaderPrint::SetParameters(View, Parameters->ShaderPrintParameters);
+	}
+
+	FRDGTextureRef RayMarchMask = nullptr;
+	if (HairOpaqueMaskType == FHairOpaqueMaskType_VirtualVoxel || HairOpaqueMaskType == FHairOpaqueMaskType_Voxel)
+	{
+		FRDGTextureDesc Desc = OutShadowMask->Desc;
+		Desc.TargetableFlags |= TexCreate_ShaderResource;
+		RayMarchMask = GraphBuilder.CreateTexture(Desc, TEXT("RayMarchMask"));
+		FRHICopyTextureInfo CopyInfo;
+		CopyInfo.Size = OutShadowMask->Desc.GetSize();
+		AddCopyTexturePass(GraphBuilder, OutShadowMask, RayMarchMask, CopyInfo);
+	}
+	Parameters->RayMarchMaskTexture = RayMarchMask;
+
 	Parameters->Voxel_LightPosition = Params.Voxel_LightPosition;
 	Parameters->Voxel_LightDirection = Params.Voxel_LightDirection;
 	Parameters->Voxel_DensityScale = Params.Voxel_DensityScale;
 	Parameters->Voxel_MinAABB = Params.Voxel_MinAABB;
 	Parameters->Voxel_Resolution = Params.Voxel_Resolution;
 	Parameters->Voxel_MaxAABB = Params.Voxel_MaxAABB;
-
+	Parameters->Voxel_MacroGroupId = Params.Voxel_MacroGroupId;
 	Parameters->Voxel_DensityTexture = Params.Voxel_DensityTexture;
 
+	if (HairOpaqueMaskType == FHairOpaqueMaskType_VirtualVoxel)
+	{		
+		check(Params.Voxel_VirtualVoxel);
+		Parameters->VirtualVoxel = Params.Voxel_VirtualVoxel->UniformBuffer;
+	}
+
 	FDeepShadowMaskPS::FPermutationDomain PermutationVector;
 	PermutationVector.Set<FDeepShadowMaskPS::FOpaqueMaskType>(HairOpaqueMaskType);
 	PermutationVector.Set<FDeepShadowMaskPS::FKernelType>(FMath::Clamp(GStrandHairShadowMaskKernelType, 0, 3));
@@ -374,10 +478,11 @@ static FHairStrandsTransmittanceMaskData RenderHairStrandsTransmittanceMask(
 	const FViewInfo& View,
 	const FLightSceneInfo* LightSceneInfo,
 	const FHairStrandsDeepShadowDatas& DeepShadowDatas,
-	const FHairStrandsClusterDatas& DeepClusterDatas,
-	const FHairStrandsVisibilityData& VisibilityData)
+	const FHairStrandsMacroGroupDatas& MacroGroupDatas,
+	const FHairStrandsVisibilityData& VisibilityData,
+	TRefCountPtr<IPooledRenderTarget>& ScreenShadowMaskSubPixelTexture)
 {
-	if (DeepClusterDatas.Datas.Num() == 0)
+	if (MacroGroupDatas.Datas.Num() == 0)
 		return FHairStrandsTransmittanceMaskData();
 
 	if (!HasDeepShadowData(LightSceneInfo, DeepShadowDatas) && !IsHairStrandsVoxelizationEnable())
@@ -423,9 +528,10 @@ static FHairStrandsTransmittanceMaskData RenderHairStrandsTransmittanceMask(
 				Params.DeepShadow_Resolution = DeepShadowData.ShadowResolution;
 				Params.LightDirection = DeepShadowData.LightDirection;
 				Params.LightPosition = DeepShadowData.LightPosition;
+				Params.LightRadius = 0;
 				Params.DepthBiasScale = GetDeepShadowDepthBiasScale();
-				Params.DeepShadow_AtlasSlotOffsets[DeepShadowData.ClusterId] = FIntVector4(DeepShadowData.AtlasRect.Min.X, DeepShadowData.AtlasRect.Min.Y, DeepShadowData.AtlasRect.Width(), DeepShadowData.AtlasRect.Height());
-				Params.DeepShadow_WorldToLightTransforms[DeepShadowData.ClusterId] = DeepShadowData.WorldToLightTransform;
+				Params.DeepShadow_AtlasSlotOffsets[DeepShadowData.MacroGroupId] = FIntVector4(DeepShadowData.AtlasRect.Min.X, DeepShadowData.AtlasRect.Min.Y, DeepShadowData.AtlasRect.Width(), DeepShadowData.AtlasRect.Height());
+				Params.DeepShadow_WorldToLightTransforms[DeepShadowData.MacroGroupId] = DeepShadowData.WorldToLightTransform;
 			}
 		}
 
@@ -441,47 +547,57 @@ static FHairStrandsTransmittanceMaskData RenderHairStrandsTransmittanceMask(
 				Params,
 				VisibilityData.NodeGroupSize,
 				HairLUTTexture,
-				NodeIndirectArgBuffer);
+				NodeIndirectArgBuffer,
+				ScreenShadowMaskSubPixelTexture);
 		}
 	}
 
 	if (!bHasFoundLight && IsHairStrandsVoxelizationEnable())
 	{
-		Params.Cluster_Resolution = 0;
-		memset(Params.Cluster_MinAABBs, 0, sizeof(FVector4) * FHairStrandsDeepShadowData::MaxClusterCount);
-		memset(Params.Cluster_MaxAABBs, 0, sizeof(FVector4) * FHairStrandsDeepShadowData::MaxClusterCount);
+		Params.Voxel_Resolution = 0;
+		memset(Params.Voxel_MinAABBs, 0, sizeof(FVector4) * FHairStrandsDeepShadowData::MaxMacroGroupCount);
+		memset(Params.Voxel_MaxAABBs, 0, sizeof(FVector4) * FHairStrandsDeepShadowData::MaxMacroGroupCount);
 
 		TRefCountPtr<IPooledRenderTarget> DummyVoxelResources;
 		FPooledRenderTargetDesc Desc(FPooledRenderTargetDesc::CreateVolumeDesc(1, 1, 1, PF_R32_UINT, FClearValueBinding::Black, TexCreate_None, TexCreate_UAV | TexCreate_ShaderResource, false, 1));
 		GRenderTargetPool.FindFreeElement(RHICmdList, Desc, DummyVoxelResources, TEXT("DummyDensityTexture"));
 		FRDGTextureRef DefaultDensityTexture = GraphBuilder.RegisterExternalTexture(DummyVoxelResources, TEXT("Voxel_DefaultDensityTexture"));
-		for (uint32 TexIt = 0; TexIt < FHairStrandsDeepShadowData::MaxClusterCount; ++TexIt)
+		for (uint32 TexIt = 0; TexIt < FHairStrandsDeepShadowData::MaxMacroGroupCount; ++TexIt)
 		{
-			Params.Cluster_DensityTextures[TexIt] = DefaultDensityTexture;
+			Params.Voxel_DensityTextures[TexIt] = DefaultDensityTexture;
 		}
 
+		FLightShaderParameters LightParameters;
+		LightSceneInfo->Proxy->GetLightShaderParameters(LightParameters);
+
 		Params.DensityScale = GetHairStrandsVoxelizationDensityScale();
 		Params.DepthBiasScale = GetHairStrandsVoxelizationDepthBiasScale();
 		Params.LightDirection = LightSceneInfo->Proxy->GetDirection();
 		Params.LightPosition = FVector4(LightSceneInfo->Proxy->GetPosition(), LightSceneInfo->Proxy->GetLightType() == ELightComponentType::LightType_Directional ? 0 : 1);
-		for (const FHairStrandsClusterData& ClusterData : DeepClusterDatas.Datas)
+		Params.LightRadius = FMath::Max(LightParameters.SourceLength, LightParameters.SourceRadius);
+		Params.VirtualVoxelResources = &MacroGroupDatas.VirtualVoxelResources;
+
+		const bool bUseVirtualVoxel = MacroGroupDatas.VirtualVoxelResources.IsValid();
+		const FHairTransmittanceType HairTransmittanceType = bUseVirtualVoxel ? FHairTransmittanceType_VirtualVoxel : FHairTransmittanceType_Voxel;
+		for (const FHairStrandsMacroGroupData& MacroGroupData : MacroGroupDatas.Datas)
 		{
-			Params.Cluster_Resolution = ClusterData.GetResolution();
-			Params.Cluster_MinAABBs[ClusterData.ClusterId] = ClusterData.GetMinBound();
-			Params.Cluster_MaxAABBs[ClusterData.ClusterId] = ClusterData.GetMaxBound();
-			if (ClusterData.VoxelResources.DensityTexture)
-				Params.Cluster_DensityTextures[ClusterData.ClusterId] = GraphBuilder.RegisterExternalTexture(ClusterData.VoxelResources.DensityTexture);
+			Params.Voxel_Resolution = MacroGroupData.GetResolution();
+			Params.Voxel_MinAABBs[MacroGroupData.MacroGroupId] = MacroGroupData.GetMinBound();
+			Params.Voxel_MaxAABBs[MacroGroupData.MacroGroupId] = MacroGroupData.GetMaxBound();
+			if (MacroGroupData.VoxelResources.DensityTexture)
+				Params.Voxel_DensityTextures[MacroGroupData.MacroGroupId] = GraphBuilder.RegisterExternalTexture(MacroGroupData.VoxelResources.DensityTexture);
 		}
 
 		OutShadowMask = AddDeepShadowTransmittanceMaskPass(
 			GraphBuilder,
 			SceneTextures,
 			View,
-			FHairTransmittanceType_Voxel,
+			HairTransmittanceType,
 			Params,
 			VisibilityData.NodeGroupSize,
 			HairLUTTexture,
-			NodeIndirectArgBuffer);
+			NodeIndirectArgBuffer,
+			ScreenShadowMaskSubPixelTexture);
 	}
 
 	FHairStrandsTransmittanceMaskData OutTransmittanceMaskData;
@@ -504,7 +620,8 @@ FHairStrandsTransmittanceMaskData RenderHairStrandsTransmittanceMask(
 	FRHICommandListImmediate& RHICmdList,
 	const TArray<FViewInfo>& Views,
 	const FLightSceneInfo* LightSceneInfo,
-	const FHairStrandsDatas* HairDatas)
+	const FHairStrandsDatas* HairDatas,
+	TRefCountPtr<IPooledRenderTarget>& ScreenShadowMaskSubPixelTexture)
 {
 	FHairStrandsTransmittanceMaskData TransmittanceMaskData;
 	if (HairDatas)
@@ -514,9 +631,9 @@ FHairStrandsTransmittanceMaskData RenderHairStrandsTransmittanceMask(
 			const FViewInfo& View = Views[ViewIndex];
 			const FHairStrandsDeepShadowDatas& InDeepShadowDatas = HairDatas->DeepShadowViews.Views[ViewIndex];
 			const FHairStrandsVisibilityData& InHairVisibilityData = HairDatas->HairVisibilityViews.HairDatas[ViewIndex];
-			const FHairStrandsClusterDatas& InDeepClusterDatas = HairDatas->HairClusterPerViews.Views[ViewIndex];
+			const FHairStrandsMacroGroupDatas& InMacroGroupDatas = HairDatas->MacroGroupsPerViews.Views[ViewIndex];
 
-			TransmittanceMaskData = RenderHairStrandsTransmittanceMask(RHICmdList, View, LightSceneInfo, InDeepShadowDatas, InDeepClusterDatas, InHairVisibilityData);
+			TransmittanceMaskData = RenderHairStrandsTransmittanceMask(RHICmdList, View, LightSceneInfo, InDeepShadowDatas, InMacroGroupDatas, InHairVisibilityData, ScreenShadowMaskSubPixelTexture);
 		}
 	}
 	return TransmittanceMaskData;
@@ -531,9 +648,9 @@ static void RenderHairStrandsShadowMask(
 	IPooledRenderTarget* ScreenShadowMaskTexture,
 	const FHairStrandsDeepShadowDatas& DeepShadowDatas,
 	const FHairStrandsVisibilityData& InVisibilityData,
-	const FHairStrandsClusterDatas& InClusterDatas)
+	const FHairStrandsMacroGroupDatas& InMacroGroupDatas)
 {
-	if (InClusterDatas.Datas.Num() == 0)
+	if (InMacroGroupDatas.Datas.Num() == 0)
 		return;
 
 	if (!HasDeepShadowData(LightSceneInfo, DeepShadowDatas) && !IsHairStrandsVoxelizationEnable())
@@ -586,25 +703,30 @@ static void RenderHairStrandsShadowMask(
 	// If there is no deep shadow for this light, fallback on the voxel representation
 	if (!bHasDeepShadow && IsHairStrandsVoxelizationEnable())
 	{
-		for (const FHairStrandsClusterData& ClusterData : InClusterDatas.Datas)
+		// TODO: Change this to be a single pass with virtual voxel?
+		for (const FHairStrandsMacroGroupData& MacroGroupData : InMacroGroupDatas.Datas)
 		{
 			const bool bIsWholeSceneLight = LightSceneInfo->Proxy->GetLightType() == LightType_Directional;
 
 			FDeepShadowOpaqueParams Params;
 			Params.CategorizationTexture = Categorization;
-			Params.Voxel_Resolution = ClusterData.GetResolution();
-			Params.Voxel_MinAABB = ClusterData.GetMinBound();
-			Params.Voxel_MaxAABB = ClusterData.GetMaxBound();
-			Params.Voxel_DensityTexture = GraphBuilder.RegisterExternalTexture(ClusterData.VoxelResources.DensityTexture ? ClusterData.VoxelResources.DensityTexture : GSystemTextures.WhiteDummy, TEXT("Voxel_DensityTexture"));
+			Params.Voxel_Resolution = MacroGroupData.GetResolution();
+			Params.Voxel_MinAABB = MacroGroupData.GetMinBound();
+			Params.Voxel_MaxAABB = MacroGroupData.GetMaxBound();
+			Params.Voxel_DensityTexture = GraphBuilder.RegisterExternalTexture(MacroGroupData.VoxelResources.DensityTexture ? MacroGroupData.VoxelResources.DensityTexture : GSystemTextures.WhiteDummy, TEXT("Voxel_DensityTexture"));
 			Params.Voxel_DensityScale = GetDeepShadowDensityScale();
 			Params.Voxel_LightDirection = LightSceneInfo->Proxy->GetDirection();
 			Params.Voxel_LightPosition = FVector4(LightSceneInfo->Proxy->GetPosition(), LightSceneInfo->Proxy->GetLightType() == ELightComponentType::LightType_Directional ? 0 : 1);
+			Params.Voxel_MacroGroupId = MacroGroupData.MacroGroupId;
 
+			const bool bUseVirtualVoxel = InMacroGroupDatas.VirtualVoxelResources.IsValid();
+			const FHairOpaqueMaskType HairOpaqueMaskType = bUseVirtualVoxel ? FHairOpaqueMaskType_VirtualVoxel : FHairOpaqueMaskType_Voxel;
+			Params.Voxel_VirtualVoxel = bUseVirtualVoxel ? &InMacroGroupDatas.VirtualVoxelResources : nullptr;
 			AddDeepShadowOpaqueMaskPass(
 				GraphBuilder,
 				SceneTextures,
 				View,
-				FHairOpaqueMaskType_Voxel,
+				HairOpaqueMaskType,
 				Params,
 				OutShadowMask);
 		}
@@ -625,16 +747,16 @@ void RenderHairStrandsShadowMask(
 {
 	const FHairStrandsDeepShadowViews& DeepShadowViews = HairDatas->DeepShadowViews;
 	const FHairStrandsVisibilityViews& HairVisibilityViews = HairDatas->HairVisibilityViews;
-	const FHairStrandsClusterViews& HairClusterViews = HairDatas->HairClusterPerViews;
+	const FHairStrandsMacroGroupViews& MacroGroupViews = HairDatas->MacroGroupsPerViews;
 
 	for (int32 ViewIndex = 0; ViewIndex < Views.Num(); ++ViewIndex)
 	{
-		if (ViewIndex < DeepShadowViews.Views.Num() && ViewIndex < HairVisibilityViews.HairDatas.Num() && ViewIndex < HairClusterViews.Views.Num())
+		if (ViewIndex < DeepShadowViews.Views.Num() && ViewIndex < HairVisibilityViews.HairDatas.Num() && ViewIndex < MacroGroupViews.Views.Num())
 		{
 			const FHairStrandsVisibilityData& HairVisibilityData = HairVisibilityViews.HairDatas[ViewIndex];
 			const FHairStrandsDeepShadowDatas& DeepShadowDatas = DeepShadowViews.Views[ViewIndex];
-			const FHairStrandsClusterDatas& ClusterDatas = HairClusterViews.Views[ViewIndex];
-			RenderHairStrandsShadowMask(RHICmdList, Views[ViewIndex], LightSceneInfo, ScreenShadowMaskTexture, DeepShadowDatas, HairVisibilityData, ClusterDatas);
+			const FHairStrandsMacroGroupDatas& MacroGroupDatas = MacroGroupViews.Views[ViewIndex];
+			RenderHairStrandsShadowMask(RHICmdList, Views[ViewIndex], LightSceneInfo, ScreenShadowMaskTexture, DeepShadowDatas, HairVisibilityData, MacroGroupDatas);
 		}
 	}
 }
diff --git a/Engine/Source/Runtime/Renderer/Private/HairStrands/HairStrandsTransmittance.h b/Engine/Source/Runtime/Renderer/Private/HairStrands/HairStrandsTransmittance.h
--- a/Engine/Source/Runtime/Renderer/Private/HairStrands/HairStrandsTransmittance.h
+++ b/Engine/Source/Runtime/Renderer/Private/HairStrands/HairStrandsTransmittance.h
@@ -29,4 +29,5 @@ FHairStrandsTransmittanceMaskData RenderHairStrandsTransmittanceMask(
 	FRHICommandListImmediate& RHICmdList,
 	const TArray<FViewInfo>& Views,
 	const class FLightSceneInfo* LightSceneInfo,
-	const struct FHairStrandsDatas* Hairdatas);
\ No newline at end of file
+	const struct FHairStrandsDatas* Hairdatas,
+	TRefCountPtr<IPooledRenderTarget>& ScreenShadowMaskSubPixelTexture);
\ No newline at end of file
diff --git a/Engine/Source/Runtime/Renderer/Private/HairStrands/HairStrandsUtils.cpp b/Engine/Source/Runtime/Renderer/Private/HairStrands/HairStrandsUtils.cpp
--- a/Engine/Source/Runtime/Renderer/Private/HairStrands/HairStrandsUtils.cpp
+++ b/Engine/Source/Runtime/Renderer/Private/HairStrands/HairStrandsUtils.cpp
@@ -16,10 +16,10 @@ static FAutoConsoleVariableRef CVarHairGlobalScattering(TEXT("r.HairStrands.Comp
 static FAutoConsoleVariableRef CVarHairLocalScattering(TEXT("r.HairStrands.Components.LocalScattering"), GHairLocalScattering, TEXT("Enable/disable hair BSDF component local scattering"));
 
 static float GStrandHairRasterizationScale = 0.5f; // For no AA without TAA, a good value is: 1.325f (Empirical)
-static FAutoConsoleVariableRef CVarStrandHairRasterizationScale(TEXT("r.HairStrands.RasterizationScale"), GStrandHairRasterizationScale, TEXT("Rasterization scale to snap strand to pixel"));
+static FAutoConsoleVariableRef CVarStrandHairRasterizationScale(TEXT("r.HairStrands.RasterizationScale"), GStrandHairRasterizationScale, TEXT("Rasterization scale to snap strand to pixel"), ECVF_Scalability | ECVF_RenderThreadSafe);
 
 static float GStrandHairVelocityRasterizationScale = 1.5f; // Tuned based on heavy motion example (e.g., head shaking)
-static FAutoConsoleVariableRef CVarStrandHairMaxRasterizationScale(TEXT("r.HairStrands.VelocityRasterizationScale"), GStrandHairVelocityRasterizationScale, TEXT("Rasterization scale to snap strand to pixel under high velocity"));
+static FAutoConsoleVariableRef CVarStrandHairMaxRasterizationScale(TEXT("r.HairStrands.VelocityRasterizationScale"), GStrandHairVelocityRasterizationScale, TEXT("Rasterization scale to snap strand to pixel under high velocity"), ECVF_Scalability | ECVF_RenderThreadSafe);
 
 static float GStrandHairShadowRasterizationScale = 1.0f;
 static FAutoConsoleVariableRef CVarStrandHairShadowRasterizationScale(TEXT("r.HairStrands.ShadowRasterizationScale"), GStrandHairShadowRasterizationScale, TEXT("Rasterization scale to snap strand to pixel in shadow view"));
@@ -30,6 +30,13 @@ static FAutoConsoleVariableRef CVarDeepShadowAABBScale(TEXT("r.HairStrands.DeepS
 static int32 GHairVisibilityRectOptimEnable = 0;
 static FAutoConsoleVariableRef CVarHairVisibilityRectOptimEnable(TEXT("r.HairStrands.RectLightingOptim"), GHairVisibilityRectOptimEnable, TEXT("Hair Visibility use projected view rect to light only relevant pixels"));
 
+static float GHairDualScatteringRoughnessOverride = 0;
+static FAutoConsoleVariableRef CVarHairDualScatteringRoughnessOverride(TEXT("r.HairStrands.DualScatteringRoughness"), GHairDualScatteringRoughnessOverride, TEXT("Override all roughness for the dual scattering evaluation. 0 means no override. Default:0"));
+float GetHairDualScatteringRoughnessOverride()
+{
+	return GHairDualScatteringRoughnessOverride;
+}
+
 float SampleCountToSubPixelSize(uint32 SamplePerPixelCount)
 {
 	float Scale = 1;
@@ -180,14 +187,14 @@ FIntRect ComputeProjectedScreenRect(const FBox& B, const FViewInfo& View)
 }
 
 
-FIntRect ComputeVisibleHairStrandsClustersRect(const FIntRect& ViewRect, const FHairStrandsClusterDatas& ClusterDatas)
+FIntRect ComputeVisibleHairStrandsMacroGroupsRect(const FIntRect& ViewRect, const FHairStrandsMacroGroupDatas& Datas)
 {
 	FIntRect TotalRect(INT_MAX, INT_MAX, -INT_MAX, -INT_MAX);
 	if (IsHairStrandsViewRectOptimEnable())
 	{
-		for (const FHairStrandsClusterData& ClusterData : ClusterDatas.Datas)
+		for (const FHairStrandsMacroGroupData& Data : Datas.Datas)
 		{
-			TotalRect.Union(ClusterData.ScreenRect);
+			TotalRect.Union(Data.ScreenRect);
 		}
 	}
 	else
@@ -233,4 +240,24 @@ FIntPoint GetVendorOptimalGroupSize2D()
 	case HairVisibilityVendor_INTEL:	return FIntPoint(8, 8);
 	default:							return FIntPoint(8, 8);
 	}
+}
+
+FVector4 PackHairRenderInfo(
+	float PrimaryRadiusAtDepth1,
+	float VelocityRadiusAtDepth1,
+	float VelocityMagnitudeScale,
+	bool bIsOrtho,
+	bool bIsGPUDriven)
+{
+	uint32 BitField = 0;
+	BitField |= bIsOrtho ? 0x1 : 0;
+	BitField |= bIsGPUDriven ? 0x2 : 0;
+
+	FVector4 Out;
+	Out.X = PrimaryRadiusAtDepth1;
+	Out.Y = VelocityRadiusAtDepth1;
+	Out.Z = VelocityMagnitudeScale;
+	Out.W = *((float*)(&BitField));
+
+	return Out;
 }
\ No newline at end of file
diff --git a/Engine/Source/Runtime/Renderer/Private/HairStrands/HairStrandsUtils.h b/Engine/Source/Runtime/Renderer/Private/HairStrands/HairStrandsUtils.h
--- a/Engine/Source/Runtime/Renderer/Private/HairStrands/HairStrandsUtils.h
+++ b/Engine/Source/Runtime/Renderer/Private/HairStrands/HairStrandsUtils.h
@@ -32,9 +32,11 @@ struct FHairComponent
 FHairComponent GetHairComponents();
 uint32 ToBitfield(const FHairComponent& Component);
 
+float GetHairDualScatteringRoughnessOverride();
+
 float SampleCountToSubPixelSize(uint32 SamplePerPixelCount);
 
-FIntRect ComputeVisibleHairStrandsClustersRect(const FIntRect& ViewRect, const struct FHairStrandsClusterDatas& ClusterDatas);
+FIntRect ComputeVisibleHairStrandsMacroGroupsRect(const FIntRect& ViewRect, const struct FHairStrandsMacroGroupDatas& Datas);
 
 bool IsHairStrandsViewRectOptimEnable();
 
@@ -52,3 +54,9 @@ FIntPoint GetVendorOptimalGroupSize2D();
 
 RENDERER_API bool IsHairStrandsSupported(const EShaderPlatform Platform);
 
+FVector4 PackHairRenderInfo(
+	float PrimaryRadiusAtDepth1,
+	float VelocityRadiusAtDepth1,
+	float VelocityScale,
+	bool  bIsOrtho,
+	bool  bIsGPUDriven);
\ No newline at end of file
diff --git a/Engine/Source/Runtime/Renderer/Private/HairStrands/HairStrandsVisibility.h b/Engine/Source/Runtime/Renderer/Private/HairStrands/HairStrandsVisibility.h
--- a/Engine/Source/Runtime/Renderer/Private/HairStrands/HairStrandsVisibility.h
+++ b/Engine/Source/Runtime/Renderer/Private/HairStrands/HairStrandsVisibility.h
@@ -28,6 +28,12 @@ struct FHairStrandsVisibilityData
 	TRefCountPtr<IPooledRenderTarget> PPLLNodeIndexTexture;
 	TRefCountPtr<FPooledRDGBuffer>	  PPLLNodeDataBuffer;
 
+	TRefCountPtr<IPooledRenderTarget> TileIndexTexture;
+	TRefCountPtr<FPooledRDGBuffer>	  TileBuffer;
+	TRefCountPtr<FPooledRDGBuffer>	  TileIndirectArgs;
+	const uint32					  TileSize = 8;
+	const uint32					  TileThreadGroupSize = 32;
+
 	TRefCountPtr<IPooledRenderTarget> NodeIndex;
 	TRefCountPtr<FPooledRDGBuffer>	  NodeData;
 	FShaderResourceViewRHIRef		  NodeDataSRV;
@@ -50,7 +56,7 @@ FHairStrandsVisibilityViews RenderHairStrandsVisibilityBuffer(
 	TRefCountPtr<IPooledRenderTarget> ColorTexture,
 	TRefCountPtr<IPooledRenderTarget> DepthTexture,
 	TRefCountPtr<IPooledRenderTarget> VelocityTexture,
-	const struct FHairStrandsClusterViews& ClusterViews);
+	const struct FHairStrandsMacroGroupViews& MacroGroupViews);
 
 void SetUpViewHairRenderInfo(const FViewInfo& ViewInfo, bool bEnableMSAA, FVector4& OutHairRenderInfo);
 
