Copying //UE4/Dev-Physics to Main (//UE4/Main)

diff --git a/Engine/Source/Runtime/Renderer/Private/HairStrands/HairStrandsVoxelization.cpp b/Engine/Source/Runtime/Renderer/Private/HairStrands/HairStrandsVoxelization.cpp
--- a/Engine/Source/Runtime/Renderer/Private/HairStrands/HairStrandsVoxelization.cpp
+++ b/Engine/Source/Runtime/Renderer/Private/HairStrands/HairStrandsVoxelization.cpp
@@ -25,7 +25,7 @@ static int32 GHairVoxelizationMaterialEnable = 0;
 static FAutoConsoleVariableRef CVarHairVoxelizationMaterialEnable(TEXT("r.HairStrands.VoxelizationMaterial"), GHairVoxelizationMaterialEnable, TEXT("Enable hair material voxelization for LOD evaluation"));
 
 static float GHairVoxelizationAABBScale = 1.0f;
-static FAutoConsoleVariableRef CVarHairVoxelizationAABBScale(TEXT("r.HairStrands.Voxelization.AABBScale"), GHairVoxelizationAABBScale, TEXT("Scale the hair cluster bounding box"));
+static FAutoConsoleVariableRef CVarHairVoxelizationAABBScale(TEXT("r.HairStrands.Voxelization.AABBScale"), GHairVoxelizationAABBScale, TEXT("Scale the hair macro group bounding box"));
 
 static float GHairVoxelizationDensityScale = 1.0f;
 static float GHairVoxelizationDepthBiasScale = 3.0f;
@@ -41,12 +41,36 @@ static FAutoConsoleVariableRef CVarHairVoxelFilterOpaqueDepthEnable(TEXT("r.Hair
 static int32 GHairStrandsVoxelMipMethod = 0;
 static FAutoConsoleVariableRef CVarHairVoxelMipMethod(TEXT("r.HairStrands.Voxelization.MipMethod"), GHairStrandsVoxelMipMethod, TEXT("Voxel mip methods (0 : one level per pass, 1 : two levels per pass."));
 
+static int32 GHairStransVoxelInjectOpaqueBiasCount = 3;
+static int32 GHairStransVoxelInjectOpaqueMarkCount = 6;
+static FAutoConsoleVariableRef CVarHairStransVoxelInjectOpaqueBiasCount(TEXT("r.HairStrands.Voxelization.InjectOpaque.BiasCount"), GHairStransVoxelInjectOpaqueBiasCount, TEXT("Bias, in number of voxel, at which opaque depth is injected."));
+static FAutoConsoleVariableRef CVarHairStransVoxelInjectOpaqueMarkCount(TEXT("r.HairStrands.Voxelization.InjectOpaque.MarkCount"), GHairStransVoxelInjectOpaqueMarkCount, TEXT("Number of voxel marked as opaque starting along the view direction beneath the opaque surface."));
+
+static float GHairStransVoxelRaymarchingSteppingScale = 1.15f;
+static FAutoConsoleVariableRef CVarHairStransVoxelRaymarchingSteppingScale(TEXT("r.HairStrands.Voxelization.Raymarching.SteppingScale"), GHairStransVoxelRaymarchingSteppingScale, TEXT("Stepping scale used for raymarching the voxel structure."));
+
 float GetHairStrandsVoxelizationDensityScale() { return FMath::Max(0.0f, GHairVoxelizationDensityScale); }
 float GetHairStrandsVoxelizationDepthBiasScale() { return FMath::Max(0.0f, GHairVoxelizationDepthBiasScale); }
 
 static int32 GHairForVoxelTransmittanceAndShadow = 0;
 static FAutoConsoleVariableRef CVarHairForVoxelTransmittanceAndShadow(TEXT("r.HairStrands.Voxelization.ForceTransmittanceAndShadow"), GHairForVoxelTransmittanceAndShadow, TEXT("For transmittance and shadow to be computed with density volume. This requires voxelization is enabled."));
 
+static int32 GHairVirtualVoxel = 1;
+static float GHairVirtualVoxel_VoxelWorldSize = 0.15f; // 1.5mm
+static int32 GHairVirtualVoxel_PageResolution = 32;
+static int32 GHairVirtualVoxel_PageCountPerDim = 14;
+static FAutoConsoleVariableRef CVarHairVirtualVoxel(TEXT("r.HairStrands.Voxelization.Virtual"), GHairVirtualVoxel, TEXT("Enable the two voxel hierachy."));
+static FAutoConsoleVariableRef CVarHairVirtualVoxel_VoxelWorldSize(TEXT("r.HairStrands.Voxelization.Virtual.VoxelWorldSize"), GHairVirtualVoxel_VoxelWorldSize, TEXT("World size of a voxel in cm."));
+static FAutoConsoleVariableRef CVarHairVirtualVoxel_VoxelPageResolution(TEXT("r.HairStrands.Voxelization.Virtual.VoxelPageResolution"), GHairVirtualVoxel_PageResolution, TEXT("Resolution of a voxel page."));
+static FAutoConsoleVariableRef CVarHairVirtualVoxel_VoxelPageCountPerDim(TEXT("r.HairStrands.Voxelization.Virtual.VoxelPageCountPerDim"), GHairVirtualVoxel_PageCountPerDim, TEXT("Number of voxel pages per texture dimension. The voxel page memory is allocated with a 3D texture. This value provide the resolution of this texture."));
+
+static int32 GHairVirtualVoxelGPUDriven = 0;
+static int32 GHairVirtualVoxelGPUDrivenMaxPageIndexRes = 32;
+static FAutoConsoleVariableRef CVarHairVirtualVoxelGPUDriven(TEXT("r.HairStrands.Voxelization.GPUDriven"), GHairVirtualVoxelGPUDriven, TEXT("Enable GPU driven voxelization."));
+static FAutoConsoleVariableRef CVarHairVirtualVoxelGPUDrivenMaxPageIndexRes(TEXT("r.HairStrands.Voxelization.GPUDriven.MaxPageIndexResolution"), GHairVirtualVoxelGPUDrivenMaxPageIndexRes, TEXT("Max resolution of the page index. This is used for allocating a conservative page index buffer when GPU driven allocation is enabled."));
+
+static const FIntPoint GPUDrivenViewportResolution = FIntPoint(4096, 4096);
+
 bool IsHairStrandsVoxelizationEnable()
 {
 	return GHairVoxelizationEnable > 0;
@@ -57,6 +81,10 @@ bool IsHairStrandsForVoxelTransmittanceAndShadowEnable()
 	return IsHairStrandsVoxelizationEnable() && GHairForVoxelTransmittanceAndShadow > 0;
 }
 
+///////////////////////////////////////////////////////////////////////////////////////////////////
+
+IMPLEMENT_GLOBAL_SHADER_PARAMETER_STRUCT(FVirtualVoxelParameters, "VirtualVoxel");
+
 ///////////////////////////////////////////////////////////////////////////////////////////////////
 class FVoxelInjectOpaquePS : public FGlobalShader
 {
@@ -69,8 +97,10 @@ class FVoxelInjectOpaquePS : public FGlobalShader
 		SHADER_PARAMETER(FVector, VoxelMinAABB)
 		SHADER_PARAMETER(uint32, VoxelResolution)
 		SHADER_PARAMETER(FVector, VoxelMaxAABB)
+		SHADER_PARAMETER(uint32, VoxelBiasCount)
 		SHADER_PARAMETER(FVector2D, OutputResolution)
 		SHADER_PARAMETER(FVector2D, SceneDepthResolution)
+		SHADER_PARAMETER(uint32, VoxelMarkCount)
 
 		SHADER_PARAMETER_RDG_TEXTURE(Texture2D, SceneDepthTexture)
 		SHADER_PARAMETER_RDG_TEXTURE_UAV(RWTexture3D, DensityTexture)
@@ -82,6 +112,11 @@ class FVoxelInjectOpaquePS : public FGlobalShader
 
 public:
 	static bool ShouldCompilePermutation(const FGlobalShaderPermutationParameters& Parameters) { return IsHairStrandsSupported(Parameters.Platform); }
+	static void ModifyCompilationEnvironment(const FGlobalShaderPermutationParameters& Parameters, FShaderCompilerEnvironment& OutEnvironment)
+	{
+		FGlobalShader::ModifyCompilationEnvironment(Parameters, OutEnvironment);
+		OutEnvironment.SetDefine(TEXT("SHADER_INJECTOPAQUE_VOXEL"), 1);
+	}
 };
 
 IMPLEMENT_GLOBAL_SHADER(FVoxelInjectOpaquePS, "/Engine/Private/HairStrands/HairStrandsVoxelOpaque.usf", "MainPS", SF_Pixel);
@@ -89,22 +124,22 @@ IMPLEMENT_GLOBAL_SHADER(FVoxelInjectOpaquePS, "/Engine/Private/HairStrands/HairS
 static void AddVoxelInjectOpaquePass(
 	FRDGBuilder& GraphBuilder,
 	const FViewInfo& View,
-	FHairStrandsClusterData& Cluster)
+	FHairStrandsMacroGroupData& MacroGroup)
 {
 	FSceneTextureParameters SceneTextures;
 	SetupSceneTextureParameters(GraphBuilder, &SceneTextures);
 
 	// #hair_todo: change this to a CS. PS was for easing debugging
-	const uint32 LinearizeRes = FMath::Sqrt(Cluster.VoxelResources.DensityTexture->GetDesc().Depth);
+	const uint32 LinearizeRes = FMath::Sqrt(MacroGroup.VoxelResources.DensityTexture->GetDesc().Depth);
 
 	FIntPoint Resolution;
-	Resolution.X = Cluster.VoxelResources.DensityTexture->GetDesc().Extent.X * LinearizeRes;
-	Resolution.Y = Cluster.VoxelResources.DensityTexture->GetDesc().Extent.Y * LinearizeRes;
+	Resolution.X = MacroGroup.VoxelResources.DensityTexture->GetDesc().Extent.X * LinearizeRes;
+	Resolution.Y = MacroGroup.VoxelResources.DensityTexture->GetDesc().Extent.Y * LinearizeRes;
 
-	if (!Cluster.VoxelResources.DensityTexture)
+	if (!MacroGroup.VoxelResources.DensityTexture)
 		return;
 
-	const FRDGTextureRef VoxelDensityTexture = GraphBuilder.RegisterExternalTexture(Cluster.VoxelResources.DensityTexture, TEXT("HairVoxelDensityTexture"));
+	const FRDGTextureRef VoxelDensityTexture = GraphBuilder.RegisterExternalTexture(MacroGroup.VoxelResources.DensityTexture, TEXT("HairVoxelDensityTexture"));
 
 	FRDGTextureDesc OutputDesc;
 	OutputDesc.Extent.X = Resolution.X;
@@ -119,13 +154,15 @@ static void AddVoxelInjectOpaquePass(
 	FVoxelInjectOpaquePS::FParameters* Parameters = GraphBuilder.AllocParameters<FVoxelInjectOpaquePS::FParameters>();
 	Parameters->ViewUniformBuffer = View.ViewUniformBuffer;
 	Parameters->OutputResolution = Resolution;
+	Parameters->VoxelBiasCount = FMath::Max(0, GHairStransVoxelInjectOpaqueBiasCount);
+	Parameters->VoxelMarkCount = FMath::Max(0, GHairStransVoxelInjectOpaqueMarkCount);
 	Parameters->SceneDepthResolution = SceneTextures.SceneDepthBuffer->Desc.Extent;
 	Parameters->SceneDepthTexture = SceneTextures.SceneDepthBuffer;
 	Parameters->SceneTextures = SceneTextures;
 	Parameters->DensityTexture = GraphBuilder.CreateUAV(VoxelDensityTexture);
-	Parameters->VoxelMinAABB = Cluster.GetMinBound();
-	Parameters->VoxelMaxAABB = Cluster.GetMaxBound();
-	Parameters->VoxelResolution = Cluster.GetResolution();
+	Parameters->VoxelMinAABB = MacroGroup.GetMinBound();
+	Parameters->VoxelMaxAABB = MacroGroup.GetMaxBound();
+	Parameters->VoxelResolution = MacroGroup.GetResolution();
 	Parameters->LinearSampler = TStaticSamplerState<SF_Trilinear, AM_Clamp, AM_Clamp, AM_Clamp>::GetRHI();
 	Parameters->RenderTargets[0] = FRenderTargetBinding(DummyTexture, ERenderTargetLoadAction::EClear);
 
@@ -173,6 +210,824 @@ static void AddVoxelInjectOpaquePass(
 	});
 }
 
+
+
+///////////////////////////////////////////////////////////////////////////////////////////////////
+class FVirtualVoxelInjectOpaqueCS : public FGlobalShader
+{
+	DECLARE_GLOBAL_SHADER(FVirtualVoxelInjectOpaqueCS);
+	SHADER_USE_PARAMETER_STRUCT(FVirtualVoxelInjectOpaqueCS, FGlobalShader);
+
+	BEGIN_SHADER_PARAMETER_STRUCT(FParameters, )
+		SHADER_PARAMETER_STRUCT_INCLUDE(FSceneTextureParameters, SceneTextures)
+		SHADER_PARAMETER_STRUCT(FVirtualVoxelCommonParameters, VirtualVoxel)
+		SHADER_PARAMETER(FIntVector, DispatchedPageIndexResolution)
+		SHADER_PARAMETER(uint32, MacroGroupId)
+		SHADER_PARAMETER(FVector2D, SceneDepthResolution)
+		SHADER_PARAMETER(uint32, VoxelBiasCount)
+		SHADER_PARAMETER(uint32, VoxelMarkCount)
+		SHADER_PARAMETER_RDG_TEXTURE(Texture2D, SceneDepthTexture)
+		SHADER_PARAMETER_RDG_BUFFER(StructuredBuffer, IndirectDispatchArgs)
+		SHADER_PARAMETER_RDG_TEXTURE_UAV(RWTexture3D, OutPageTexture)
+
+		SHADER_PARAMETER_STRUCT_REF(FViewUniformShaderParameters, ViewUniformBuffer)
+	END_SHADER_PARAMETER_STRUCT()
+
+public:
+	static bool ShouldCompilePermutation(const FGlobalShaderPermutationParameters& Parameters) { return IsHairStrandsSupported(Parameters.Platform); }
+	static void ModifyCompilationEnvironment(const FGlobalShaderPermutationParameters& Parameters, FShaderCompilerEnvironment& OutEnvironment)
+	{
+		FGlobalShader::ModifyCompilationEnvironment(Parameters, OutEnvironment);
+		OutEnvironment.SetDefine(TEXT("SHADER_INJECTOPAQUE_VIRTUALVOXEL"), 1);
+	}
+};
+
+IMPLEMENT_GLOBAL_SHADER(FVirtualVoxelInjectOpaqueCS, "/Engine/Private/HairStrands/HairStrandsVoxelOpaque.usf", "MainCS", SF_Compute);
+
+static void AddVirtualVoxelInjectOpaquePass(
+	FRDGBuilder& GraphBuilder,
+	const FViewInfo& View,
+	const FVirtualVoxelResources& VoxelResources,
+	const FHairStrandsMacroGroupData& MacroGroup)
+{
+	FSceneTextureParameters SceneTextures;
+	SetupSceneTextureParameters(GraphBuilder, &SceneTextures);
+
+	const FRDGTextureRef OutPageTexture = GraphBuilder.RegisterExternalTexture(VoxelResources.PageTexture, TEXT("HairVoxelPageTexture"));
+	FRDGBufferRef IndirectDispatchArgsBuffer = GraphBuilder.RegisterExternalBuffer(VoxelResources.IndirectArgsBuffer, TEXT("HairVoxelIndirectDispatchArgs"));
+
+	const uint32 TotalPageCount = VoxelResources.Parameters.Common.PageIndexCount;
+	const uint32 PageResolution = VoxelResources.Parameters.Common.PageResolution;
+
+	const uint32 SideSlotCount = FMath::CeilToInt(FMath::Pow(TotalPageCount, 1.f / 3.f));
+	const uint32 SideVoxelCount= SideSlotCount * PageResolution;
+
+	FVirtualVoxelInjectOpaqueCS::FParameters* Parameters = GraphBuilder.AllocParameters<FVirtualVoxelInjectOpaqueCS::FParameters>();
+	Parameters->ViewUniformBuffer = View.ViewUniformBuffer;
+	Parameters->VirtualVoxel = VoxelResources.Parameters.Common;
+	Parameters->VoxelBiasCount = FMath::Max(0, GHairStransVoxelInjectOpaqueBiasCount);
+	Parameters->VoxelMarkCount = FMath::Max(0, GHairStransVoxelInjectOpaqueMarkCount);
+	Parameters->SceneDepthResolution = SceneTextures.SceneDepthBuffer->Desc.Extent;
+	Parameters->SceneDepthTexture = SceneTextures.SceneDepthBuffer;
+	Parameters->SceneTextures = SceneTextures;
+	Parameters->MacroGroupId = MacroGroup.MacroGroupId;
+	Parameters->OutPageTexture = GraphBuilder.CreateUAV(OutPageTexture);
+	Parameters->DispatchedPageIndexResolution = MacroGroup.VirtualVoxelNodeDesc.PageIndexResolution;
+	Parameters->IndirectDispatchArgs = IndirectDispatchArgsBuffer;
+	TShaderMapRef<FVirtualVoxelInjectOpaqueCS> ComputeShader(View.ShaderMap);
+	const TShaderMap<FGlobalShaderType>* GlobalShaderMap = View.ShaderMap;
+
+	check(VoxelResources.Parameters.Common.IndirectDispatchGroupSize == 64);
+	const uint32 ArgsOffset = sizeof(uint32) * 3 * Parameters->MacroGroupId;
+
+	FComputeShaderUtils::AddPass(
+		GraphBuilder, 
+		RDG_EVENT_NAME("HairStrandsInjectOpaqueDepthInVoxel"), 
+		*ComputeShader, 
+		Parameters, 
+		IndirectDispatchArgsBuffer, 
+		ArgsOffset);
+}
+
+///////////////////////////////////////////////////////////////////////////////////////////////////
+class FVoxelAllocatePageIndexCS : public FGlobalShader
+{
+	DECLARE_GLOBAL_SHADER(FVoxelAllocatePageIndexCS);
+	SHADER_USE_PARAMETER_STRUCT(FVoxelAllocatePageIndexCS, FGlobalShader);
+
+	BEGIN_SHADER_PARAMETER_STRUCT(FParameters, )
+		SHADER_PARAMETER(float, PageWorldSize)
+		SHADER_PARAMETER(uint32, TotalPageIndexCount)
+		SHADER_PARAMETER(uint32, PageResolution)
+		SHADER_PARAMETER(uint32, MacroGroupCount)
+		SHADER_PARAMETER(uint32, IndirectDispatchGroupSize)
+
+		SHADER_PARAMETER(FVector, CPU_MinAABB)
+		SHADER_PARAMETER(FVector, CPU_MaxAABB)
+		SHADER_PARAMETER(FIntVector, CPU_PageIndexResolution)
+		SHADER_PARAMETER(uint32, CPU_bUseCPUData)
+
+		SHADER_PARAMETER_RDG_BUFFER_UAV(RWBuffer<int>, MacroGroupAABBBuffer)
+		SHADER_PARAMETER_RDG_BUFFER_UAV(RWBuffer<uint4>, OutPageIndexResolutionAndOffsetBuffer)
+		SHADER_PARAMETER_RDG_BUFFER_UAV(RWBuffer, OutVoxelizationViewInfoBuffer)
+		SHADER_PARAMETER_RDG_BUFFER_UAV(RWBuffer, OutPageIndexAllocationIndirectBufferArgs)
+	END_SHADER_PARAMETER_STRUCT()
+
+public:
+	static bool ShouldCompilePermutation(const FGlobalShaderPermutationParameters& Parameters) { return IsHairStrandsSupported(Parameters.Platform); }
+	static void ModifyCompilationEnvironment(const FGlobalShaderPermutationParameters& Parameters, FShaderCompilerEnvironment& OutEnvironment)
+	{
+		FGlobalShader::ModifyCompilationEnvironment(Parameters, OutEnvironment);
+		OutEnvironment.SetDefine(TEXT("SHADER_ALLOCATEPAGEINDEX"), 1);
+	}
+};
+
+class FVoxelMarkValidPageIndexCS : public FGlobalShader
+{
+	DECLARE_GLOBAL_SHADER(FVoxelMarkValidPageIndexCS);
+	SHADER_USE_PARAMETER_STRUCT(FVoxelMarkValidPageIndexCS, FGlobalShader);
+
+	class FGPUDriven : SHADER_PERMUTATION_INT("PERMUTATION_GPU_DRIVEN", 2);
+	using FPermutationDomain = TShaderPermutationDomain<FGPUDriven>;
+
+	BEGIN_SHADER_PARAMETER_STRUCT(FParameters, )
+		SHADER_PARAMETER(FIntVector, CPU_PageIndexResolution)
+		SHADER_PARAMETER(FVector, CPU_MinAABB)
+		SHADER_PARAMETER(uint32, MaxClusterCount)
+		SHADER_PARAMETER(FVector, CPU_MaxAABB)
+		SHADER_PARAMETER(uint32, CPU_PageIndexOffset)
+		SHADER_PARAMETER(uint32, MacroGroupId)
+		SHADER_PARAMETER_SRV(Buffer, ClusterAABBsBuffer)
+		SHADER_PARAMETER_RDG_BUFFER_SRV(Buffer, MacroGroupAABBBuffer)
+		SHADER_PARAMETER_RDG_BUFFER_SRV(Buffer, PageIndexResolutionAndOffsetBuffer)
+		SHADER_PARAMETER_RDG_BUFFER_UAV(Buffer<uint>, OutValidPageIndexBuffer)
+	END_SHADER_PARAMETER_STRUCT()
+
+public:
+	static bool ShouldCompilePermutation(const FGlobalShaderPermutationParameters& Parameters) { return IsHairStrandsSupported(Parameters.Platform); }
+	static void ModifyCompilationEnvironment(const FGlobalShaderPermutationParameters& Parameters, FShaderCompilerEnvironment& OutEnvironment)
+	{
+		FGlobalShader::ModifyCompilationEnvironment(Parameters, OutEnvironment);
+		OutEnvironment.SetDefine(TEXT("SHADER_MARKVALID"), 1);
+	}
+};
+
+class FVoxelAllocateVoxelPageCS : public FGlobalShader
+{
+	DECLARE_GLOBAL_SHADER(FVoxelAllocateVoxelPageCS);
+	SHADER_USE_PARAMETER_STRUCT(FVoxelAllocateVoxelPageCS, FGlobalShader);
+
+	class FGPUDriven : SHADER_PERMUTATION_INT("PERMUTATION_GPU_DRIVEN", 2);
+	using FPermutationDomain = TShaderPermutationDomain<FGPUDriven>;
+
+	BEGIN_SHADER_PARAMETER_STRUCT(FParameters, )
+		SHADER_PARAMETER(FIntVector, CPU_PageIndexResolution)
+		SHADER_PARAMETER(uint32, MacroGroupId)
+		SHADER_PARAMETER(uint32, PageCount)
+		SHADER_PARAMETER(uint32, CPU_PageIndexCount)
+		SHADER_PARAMETER(uint32, CPU_PageIndexOffset)		
+		SHADER_PARAMETER_RDG_BUFFER(Buffer, IndirectBufferArgs)
+		SHADER_PARAMETER_RDG_BUFFER_SRV(Buffer, PageIndexResolutionAndOffsetBuffer)
+		SHADER_PARAMETER_RDG_BUFFER_UAV(RWBuffer, PageIndexGlobalCounter)
+		SHADER_PARAMETER_RDG_BUFFER_UAV(RWBuffer, PageIndexBuffer)
+		SHADER_PARAMETER_RDG_BUFFER_UAV(RWBuffer, PageIndexCoordBuffer)
+	END_SHADER_PARAMETER_STRUCT()
+
+public:
+	static bool ShouldCompilePermutation(const FGlobalShaderPermutationParameters& Parameters) { return IsHairStrandsSupported(Parameters.Platform); }
+	static void ModifyCompilationEnvironment(const FGlobalShaderPermutationParameters& Parameters, FShaderCompilerEnvironment& OutEnvironment)
+	{
+		FGlobalShader::ModifyCompilationEnvironment(Parameters, OutEnvironment);
+		OutEnvironment.SetDefine(TEXT("SHADER_ALLOCATE"), 1);
+	}
+};
+
+class FVoxelAddNodeDescCS : public FGlobalShader
+{
+	DECLARE_GLOBAL_SHADER(FVoxelAddNodeDescCS);
+	SHADER_USE_PARAMETER_STRUCT(FVoxelAddNodeDescCS, FGlobalShader);
+
+	class FGPUDriven : SHADER_PERMUTATION_INT("PERMUTATION_GPU_DRIVEN", 2);
+	using FPermutationDomain = TShaderPermutationDomain<FGPUDriven>;
+
+	BEGIN_SHADER_PARAMETER_STRUCT(FParameters, )
+		SHADER_PARAMETER(FVector, CPU_MinAABB)
+		SHADER_PARAMETER(uint32, CPU_PageIndexOffset)
+		SHADER_PARAMETER(FVector, CPU_MaxAABB)
+		SHADER_PARAMETER(uint32, MacroGroupId)
+		SHADER_PARAMETER(FIntVector, CPU_PageIndexResolution)
+		SHADER_PARAMETER_RDG_BUFFER_SRV(Buffer, MacroGroupAABBBuffer)
+		SHADER_PARAMETER_RDG_BUFFER_SRV(Buffer, PageIndexResolutionAndOffsetBuffer)
+		SHADER_PARAMETER_RDG_BUFFER_UAV(RWStructuredBuffer, OutNodeDescBuffer)
+	END_SHADER_PARAMETER_STRUCT()
+
+public:
+	static bool ShouldCompilePermutation(const FGlobalShaderPermutationParameters& Parameters) { return IsHairStrandsSupported(Parameters.Platform); }
+	static void ModifyCompilationEnvironment(const FGlobalShaderPermutationParameters& Parameters, FShaderCompilerEnvironment& OutEnvironment)
+	{
+		FGlobalShader::ModifyCompilationEnvironment(Parameters, OutEnvironment);
+		OutEnvironment.SetDefine(TEXT("SHADER_ADDDESC"), 1);
+	}
+};
+
+class FVoxelAddIndirectBufferCS : public FGlobalShader
+{
+	DECLARE_GLOBAL_SHADER(FVoxelAddIndirectBufferCS);
+	SHADER_USE_PARAMETER_STRUCT(FVoxelAddIndirectBufferCS, FGlobalShader);
+
+	BEGIN_SHADER_PARAMETER_STRUCT(FParameters, )
+		SHADER_PARAMETER(uint32, MacroGroupId)
+		SHADER_PARAMETER(uint32, IndirectGroupSize)
+		SHADER_PARAMETER(uint32, PageResolution)
+		SHADER_PARAMETER_RDG_BUFFER_UAV(RWBuffer, OutPageIndexGlobalCounter)
+		SHADER_PARAMETER_RDG_BUFFER_UAV(RWBuffer, OutIndirectArgsBuffer)
+		END_SHADER_PARAMETER_STRUCT()
+
+public:
+	static bool ShouldCompilePermutation(const FGlobalShaderPermutationParameters& Parameters) { return IsHairStrandsSupported(Parameters.Platform); }
+	static void ModifyCompilationEnvironment(const FGlobalShaderPermutationParameters& Parameters, FShaderCompilerEnvironment& OutEnvironment)
+	{
+		FGlobalShader::ModifyCompilationEnvironment(Parameters, OutEnvironment);
+		OutEnvironment.SetDefine(TEXT("SHADER_ADDINDIRECTBUFFER"), 1);
+	}
+};
+
+
+class FVoxelIndPageClearBufferGenCS : public FGlobalShader
+{
+	DECLARE_GLOBAL_SHADER(FVoxelIndPageClearBufferGenCS);
+	SHADER_USE_PARAMETER_STRUCT(FVoxelIndPageClearBufferGenCS, FGlobalShader);
+
+	BEGIN_SHADER_PARAMETER_STRUCT(FParameters, )
+		SHADER_PARAMETER_RDG_BUFFER_SRV(Buffer, PageIndexGlobalCounter)
+		SHADER_PARAMETER_RDG_BUFFER_UAV(Buffer, OutIndirectArgsBuffer)
+		SHADER_PARAMETER(uint32, PageResolution)
+	END_SHADER_PARAMETER_STRUCT()
+
+public:
+	static bool ShouldCompilePermutation(const FGlobalShaderPermutationParameters& Parameters) { return IsHairStrandsSupported(Parameters.Platform); }
+	static void ModifyCompilationEnvironment(const FGlobalShaderPermutationParameters& Parameters, FShaderCompilerEnvironment& OutEnvironment)
+	{
+		FGlobalShader::ModifyCompilationEnvironment(Parameters, OutEnvironment);
+		OutEnvironment.SetDefine(TEXT("SHADER_INDPAGECLEARBUFFERGEN"), 1);
+	}
+};
+
+class FVoxelIndPageClearCS : public FGlobalShader
+{
+	DECLARE_GLOBAL_SHADER(FVoxelIndPageClearCS);
+	SHADER_USE_PARAMETER_STRUCT(FVoxelIndPageClearCS, FGlobalShader);
+
+	BEGIN_SHADER_PARAMETER_STRUCT(FParameters, )
+		SHADER_PARAMETER_STRUCT(FVirtualVoxelCommonParameters, VirtualVoxel)
+		SHADER_PARAMETER_RDG_TEXTURE_UAV(Texture3D, OutPageTexture)
+		SHADER_PARAMETER_RDG_BUFFER(Buffer, IndirectDispatchBuffer)
+	END_SHADER_PARAMETER_STRUCT()
+
+public:
+	static bool ShouldCompilePermutation(const FGlobalShaderPermutationParameters& Parameters) { return IsHairStrandsSupported(Parameters.Platform); }
+	static void ModifyCompilationEnvironment(const FGlobalShaderPermutationParameters& Parameters, FShaderCompilerEnvironment& OutEnvironment)
+	{
+		FGlobalShader::ModifyCompilationEnvironment(Parameters, OutEnvironment);
+		OutEnvironment.SetDefine(TEXT("SHADER_INDPAGECLEAR"), 1);
+	}
+};
+
+IMPLEMENT_GLOBAL_SHADER(FVoxelAllocatePageIndexCS, "/Engine/Private/HairStrands/HairStrandsVoxelPageAllocation.usf", "AllocatePageIndex", SF_Compute);
+IMPLEMENT_GLOBAL_SHADER(FVoxelMarkValidPageIndexCS, "/Engine/Private/HairStrands/HairStrandsVoxelPageAllocation.usf", "MarkValidCS", SF_Compute);
+IMPLEMENT_GLOBAL_SHADER(FVoxelAllocateVoxelPageCS, "/Engine/Private/HairStrands/HairStrandsVoxelPageAllocation.usf", "AllocateCS", SF_Compute);
+IMPLEMENT_GLOBAL_SHADER(FVoxelAddNodeDescCS, "/Engine/Private/HairStrands/HairStrandsVoxelPageAllocation.usf", "AddDescCS", SF_Compute);
+IMPLEMENT_GLOBAL_SHADER(FVoxelAddIndirectBufferCS, "/Engine/Private/HairStrands/HairStrandsVoxelPageAllocation.usf", "AddIndirectBufferCS", SF_Compute);
+IMPLEMENT_GLOBAL_SHADER(FVoxelIndPageClearBufferGenCS, "/Engine/Private/HairStrands/HairStrandsVoxelPageAllocation.usf", "VoxelIndPageClearBufferGenCS", SF_Compute);
+IMPLEMENT_GLOBAL_SHADER(FVoxelIndPageClearCS, "/Engine/Private/HairStrands/HairStrandsVoxelPageAllocation.usf", "VoxelIndPageClearCS", SF_Compute);
+
+inline FIntVector CeilToInt(const FVector& V)
+{
+	return FIntVector(FMath::CeilToInt(V.X), FMath::CeilToInt(V.Y), FMath::CeilToInt(V.Z));
+}
+
+static void AddAllocateVoxelPagesPass(
+	FRDGBuilder& GraphBuilder,
+	const FViewInfo& View,
+	FHairStrandsMacroGroupDatas& MacroGroups,
+	const FIntVector PageCountResolution,
+	const uint32 PageCount,
+	const float VoxelWorldSize,
+	const uint32 PageResolution,
+	const FIntVector PageTextureResolution,
+	const uint32 IndirectDispatchGroupSize,
+	uint32& OutTotalPageIndexCount,
+	FRDGBufferRef& OutPageIndexBuffer,
+	FRDGBufferRef& OutPageIndexCoordBuffer,
+	FRDGBufferRef& OutNodeDescBuffer,
+	FRDGBufferRef& OutIndirectArgsBuffer,
+	FRDGBufferRef& OutPageIndexGlobalCounter, 
+	FRDGBufferRef& OutVoxelizationViewInfoBuffer)
+{
+	const uint32 GroupSize = 32;
+	const bool bIsGPUDriven = GHairVirtualVoxelGPUDriven > 0;
+	const uint32 MacroGroupCount = MacroGroups.Datas.Num();
+	if (MacroGroupCount == 0)
+		return;
+
+	struct FCPUMacroGroupAllocation
+	{
+		FVector		MinAABB;
+		FVector		MaxAABB;
+		FIntVector	PageIndexResolution;
+		uint32		PageIndexCount;
+		uint32		PageIndexOffset;
+		uint32		MacroGroupId;
+	};
+	const float PageWorldSize = PageResolution * VoxelWorldSize;
+
+	OutTotalPageIndexCount = 0;
+	TArray<FCPUMacroGroupAllocation> CPUAllocationDescs;	
+	for (FHairStrandsMacroGroupData& MacroGroup : MacroGroups.Datas)
+	{
+		// Snap the max AABB to the voxel size
+		// Scale the bounding box in place of proper GPU driven AABB for now
+		const float Scale = FMath::Clamp(GHairVoxelizationAABBScale, 0.01f, 10.0f);
+		const FVector BoxCenter = MacroGroup.Bounds.GetBox().GetCenter();
+		FVector MinAABB = (MacroGroup.Bounds.GetBox().Min - BoxCenter) * Scale + BoxCenter;
+		FVector MaxAABB = (MacroGroup.Bounds.GetBox().Max - BoxCenter) * Scale + BoxCenter;
+
+		// Allocate enough pages to cover the AABB, where page (0,0,0) origin sit on MinAABB.
+		FVector MacroGroupSize = MaxAABB - MinAABB;
+		const FIntVector PageIndexResolution = CeilToInt(MacroGroupSize / PageWorldSize);
+		MacroGroupSize = FVector(PageIndexResolution) * PageWorldSize;
+		MaxAABB = MacroGroupSize + MinAABB;
+
+		FCPUMacroGroupAllocation& Out = CPUAllocationDescs.AddDefaulted_GetRef();
+		Out.MacroGroupId = MacroGroup.MacroGroupId;
+		Out.MinAABB = MinAABB; // >> these should actually be computed on the GPU ... 
+		Out.MaxAABB = MaxAABB; // >> these should actually be computed on the GPU ... 
+		Out.PageIndexResolution = PageIndexResolution;
+		Out.PageIndexCount = Out.PageIndexResolution.X * Out.PageIndexResolution.Y * Out.PageIndexResolution.Z;
+		Out.PageIndexOffset = OutTotalPageIndexCount;
+
+		OutTotalPageIndexCount += Out.PageIndexCount;
+
+		MacroGroup.VirtualVoxelNodeDesc.WorldMinAABB = Out.MinAABB;
+		MacroGroup.VirtualVoxelNodeDesc.WorldMaxAABB = Out.MaxAABB;
+		MacroGroup.VirtualVoxelNodeDesc.PageIndexResolution = Out.PageIndexResolution;
+	}
+
+	// Over-allocation (upper bound)
+	if (bIsGPUDriven)
+	{
+		OutTotalPageIndexCount = GHairVirtualVoxelGPUDrivenMaxPageIndexRes * GHairVirtualVoxelGPUDrivenMaxPageIndexRes * GHairVirtualVoxelGPUDrivenMaxPageIndexRes;
+	}
+	check(OutTotalPageIndexCount > 0);
+	
+	FRDGBufferRef PageIndexBuffer = GraphBuilder.CreateBuffer(FRDGBufferDesc::CreateBufferDesc(sizeof(uint32), OutTotalPageIndexCount), TEXT("PageIndexBuffer"));
+	FRDGBufferRef PageIndexCoordBuffer = GraphBuilder.CreateBuffer(FRDGBufferDesc::CreateBufferDesc(sizeof(uint32), OutTotalPageIndexCount), TEXT("PageIndexCoordBuffer"));
+	FRDGBufferRef PageIndexGlobalCounter = GraphBuilder.CreateBuffer(FRDGBufferDesc::CreateBufferDesc(sizeof(uint32), 2), TEXT("PageIndexGlobalCounter"));
+	FRDGBufferRef NodeDescBuffer = GraphBuilder.CreateBuffer(FRDGBufferDesc::CreateStructuredDesc(sizeof(FPackedVirtualVoxelNodeDesc), MacroGroupCount), TEXT("VirtualVoxelNodeDescBuffer"));
+	FRDGBufferRef IndirectArgsBuffer = GraphBuilder.CreateBuffer(FRDGBufferDesc::CreateIndirectDesc<FRHIDispatchIndirectParameters>(MacroGroupCount), TEXT("VirtualVoxelIndirectArgsBuffer"));
+
+	FRDGBufferUAVRef PageIndexBufferUAV = GraphBuilder.CreateUAV(PageIndexBuffer, PF_R32_UINT);
+	FRDGBufferUAVRef PageIndexGlobalCounterUAV = GraphBuilder.CreateUAV(PageIndexGlobalCounter, PF_R32_UINT);
+	
+	// Stored FVoxelizationViewInfo structs
+	// See HairStrandsVoxelPageCommonStruct.ush for more details
+	FRDGBufferRef VoxelizationViewInfoBuffer = GraphBuilder.CreateBuffer(FRDGBufferDesc::CreateStructuredDesc(24 * sizeof(float), MacroGroupCount), TEXT("VoxelizationViewInfo"));
+	FRDGBufferRef PageIndexResolutionBuffer = GraphBuilder.CreateBuffer(FRDGBufferDesc::CreateBufferDesc(MacroGroupCount * 4 * sizeof(uint32), OutTotalPageIndexCount), TEXT("PageIndexResolutionBuffer"));
+	FRDGBufferRef MacroGroupAABB = GraphBuilder.RegisterExternalBuffer(MacroGroups.MacroGroupResources.MacroGroupAABBsBuffer, TEXT("HairInstanceGroupAABBs"));
+	FRDGBufferRef PageIndexAllocationIndirectBufferArgs = GraphBuilder.CreateBuffer(FRDGBufferDesc::CreateIndirectDesc<FRHIDispatchIndirectParameters>(MacroGroupCount), TEXT("PageIndexAllocationIndirectBufferArgs"));
+
+	AddClearUAVPass(GraphBuilder, PageIndexBufferUAV, 0u);
+	AddClearUAVPass(GraphBuilder, PageIndexGlobalCounterUAV, 0u);
+
+	// Allocate page index for all instance group
+//	if (bIsGPUDriven)
+	{		
+
+		FVoxelAllocatePageIndexCS::FParameters* Parameters = GraphBuilder.AllocParameters<FVoxelAllocatePageIndexCS::FParameters>();
+		Parameters->PageWorldSize = PageWorldSize;
+		Parameters->TotalPageIndexCount = OutTotalPageIndexCount;
+		Parameters->PageResolution = PageResolution;
+		Parameters->MacroGroupCount = MacroGroupCount;
+		Parameters->MacroGroupAABBBuffer = GraphBuilder.CreateUAV(MacroGroupAABB, PF_R32_SINT);
+		Parameters->IndirectDispatchGroupSize = GroupSize; // This is the GroupSize used for FVoxelAllocateVoxelPageCS
+		Parameters->OutPageIndexResolutionAndOffsetBuffer = GraphBuilder.CreateUAV(PageIndexResolutionBuffer, PF_R32G32B32A32_UINT);
+		Parameters->OutVoxelizationViewInfoBuffer = GraphBuilder.CreateUAV(VoxelizationViewInfoBuffer);
+		Parameters->OutPageIndexAllocationIndirectBufferArgs = GraphBuilder.CreateUAV(PageIndexAllocationIndirectBufferArgs);
+		const bool bUseCPUData = GHairVirtualVoxelGPUDriven == 2;
+		if (bUseCPUData)
+		{
+			Parameters->CPU_bUseCPUData = bUseCPUData ? 1 : 0;
+			Parameters->CPU_MinAABB = CPUAllocationDescs.Num() > 0 ? CPUAllocationDescs[0].MinAABB : FVector::ZeroVector;
+			Parameters->CPU_MaxAABB = CPUAllocationDescs.Num() > 0 ? CPUAllocationDescs[0].MaxAABB : FVector::ZeroVector;
+			Parameters->CPU_PageIndexResolution = CPUAllocationDescs.Num() > 0 ? CPUAllocationDescs[0].PageIndexResolution : FIntVector(0, 0, 0);
+		}
+
+		// Currently support only 32 instance group at max
+		check(Parameters->MacroGroupCount < 32);
+		TShaderMapRef<FVoxelAllocatePageIndexCS> ComputeShader(View.ShaderMap);
+		FComputeShaderUtils::AddPass(
+			GraphBuilder,
+			RDG_EVENT_NAME("HairStrandsAllocatePageIndex"),
+			*ComputeShader,
+			Parameters,
+			FIntVector(1,1,1));
+	}
+
+	// Mark valid page index
+	for (uint32 MacroGroupIt=0; MacroGroupIt <MacroGroupCount;++MacroGroupIt)
+	{
+		DECLARE_GPU_STAT(HairStrandsAllocateMacroGroup);
+		SCOPED_DRAW_EVENT(GraphBuilder.RHICmdList, HairStrandsAllocateMacroGroup);
+		SCOPED_GPU_STAT(GraphBuilder.RHICmdList, HairStrandsAllocateMacroGroup);
+
+		const FHairStrandsMacroGroupData& MacroGroup = MacroGroups.Datas[MacroGroupIt];
+		FCPUMacroGroupAllocation& CPUAllocationDesc = CPUAllocationDescs[MacroGroupIt];
+
+		for (const FHairStrandsMacroGroupData::PrimitiveGroup& PrimitiveGroup : MacroGroup.PrimitivesGroups)
+		{
+			const FHairStrandsPrimitiveResources& Resources = GetHairStandsPrimitiveResources(PrimitiveGroup.ResourceId);
+			check(PrimitiveGroup.GroupIndex < uint32(Resources.Groups.Num()));
+			const FHairStrandsPrimitiveResources::FHairGroup& GroupResources = Resources.Groups[PrimitiveGroup.GroupIndex];
+
+			FVoxelMarkValidPageIndexCS::FParameters* Parameters = GraphBuilder.AllocParameters<FVoxelMarkValidPageIndexCS::FParameters>();
+			Parameters->MacroGroupId = MacroGroup.MacroGroupId;
+			Parameters->MaxClusterCount = GroupResources.ClusterCount;
+			Parameters->CPU_PageIndexResolution = CPUAllocationDesc.PageIndexResolution;
+			Parameters->CPU_PageIndexOffset = CPUAllocationDesc.PageIndexOffset;
+			Parameters->CPU_MinAABB = CPUAllocationDesc.MinAABB;
+			Parameters->CPU_MaxAABB = CPUAllocationDesc.MaxAABB;
+			Parameters->ClusterAABBsBuffer = GroupResources.ClusterAABBBuffer->SRV;
+			Parameters->OutValidPageIndexBuffer = PageIndexBufferUAV;
+
+			if (bIsGPUDriven)
+			{
+				Parameters->MacroGroupAABBBuffer = GraphBuilder.CreateSRV(MacroGroupAABB, PF_R32_SINT);
+				Parameters->PageIndexResolutionAndOffsetBuffer = GraphBuilder.CreateSRV(PageIndexResolutionBuffer, PF_R32G32B32A32_UINT);
+			}
+
+			FVoxelMarkValidPageIndexCS::FPermutationDomain PermutationVector;
+			PermutationVector.Set<FVoxelMarkValidPageIndexCS::FGPUDriven>(bIsGPUDriven ? 1 : 0);
+
+			FIntVector DispatchCount((GroupResources.ClusterCount + GroupSize - 1) / GroupSize, 1, 1);
+			check(DispatchCount.X < 65535);
+			TShaderMapRef<FVoxelMarkValidPageIndexCS> ComputeShader(View.ShaderMap, PermutationVector);
+			FComputeShaderUtils::AddPass(
+				GraphBuilder,
+				RDG_EVENT_NAME("HairStrandsMarkValidPageIndex"),
+				*ComputeShader,
+				Parameters,
+				DispatchCount);
+		}
+
+		// Fill in hair-macro-group information.
+		// Note: This need to happen before the allocation as we copy the index global count. This global index is 
+		// used as an offset, and thus refers to the previous pass
+		{
+			check(MacroGroup.MacroGroupId < MacroGroupCount);
+
+			FVoxelAddNodeDescCS::FParameters* Parameters = GraphBuilder.AllocParameters<FVoxelAddNodeDescCS::FParameters>();
+			Parameters->MacroGroupId = MacroGroup.MacroGroupId;
+			Parameters->CPU_MinAABB = CPUAllocationDesc.MinAABB;
+			Parameters->CPU_MaxAABB = CPUAllocationDesc.MaxAABB;
+			Parameters->CPU_PageIndexResolution = CPUAllocationDesc.PageIndexResolution;
+			Parameters->CPU_PageIndexOffset = CPUAllocationDesc.PageIndexOffset;
+			Parameters->OutNodeDescBuffer = GraphBuilder.CreateUAV(NodeDescBuffer);
+
+			if (bIsGPUDriven)
+			{
+				Parameters->MacroGroupAABBBuffer = GraphBuilder.CreateSRV(MacroGroupAABB, PF_R32_SINT);
+				Parameters->PageIndexResolutionAndOffsetBuffer = GraphBuilder.CreateSRV(PageIndexResolutionBuffer, PF_R32G32B32A32_UINT);
+			}
+
+			FVoxelAddNodeDescCS::FPermutationDomain PermutationVector;
+			PermutationVector.Set<FVoxelAddNodeDescCS::FGPUDriven>(bIsGPUDriven ? 1 : 0);
+
+			const FIntVector DispatchCount(1, 1, 1);
+			TShaderMapRef<FVoxelAddNodeDescCS> ComputeShader(View.ShaderMap, PermutationVector);
+			FComputeShaderUtils::AddPass(GraphBuilder, RDG_EVENT_NAME("HairStrandsAddNodeDesc"), *ComputeShader, Parameters, DispatchCount);
+		}
+
+		// Allocate pages
+		{
+			FVoxelAllocateVoxelPageCS::FParameters* Parameters = GraphBuilder.AllocParameters<FVoxelAllocateVoxelPageCS::FParameters>();
+			Parameters->MacroGroupId = MacroGroup.MacroGroupId;
+			Parameters->PageCount = PageCount;
+			Parameters->CPU_PageIndexCount = CPUAllocationDesc.PageIndexCount;
+			Parameters->CPU_PageIndexResolution = CPUAllocationDesc.PageIndexResolution;
+			Parameters->CPU_PageIndexOffset = CPUAllocationDesc.PageIndexOffset;
+			Parameters->PageIndexGlobalCounter = PageIndexGlobalCounterUAV;
+			Parameters->PageIndexBuffer = PageIndexBufferUAV;
+			Parameters->PageIndexCoordBuffer = GraphBuilder.CreateUAV(PageIndexCoordBuffer, PF_R8G8B8A8_UINT);
+
+			FVoxelAllocateVoxelPageCS::FPermutationDomain PermutationVector;
+			PermutationVector.Set<FVoxelAllocateVoxelPageCS::FGPUDriven>(bIsGPUDriven ? 1 : 0);
+			TShaderMapRef<FVoxelAllocateVoxelPageCS> ComputeShader(View.ShaderMap, PermutationVector);
+
+			if (bIsGPUDriven)
+			{
+				Parameters->PageIndexResolutionAndOffsetBuffer = GraphBuilder.CreateSRV(PageIndexResolutionBuffer, PF_R32G32B32A32_UINT);
+				Parameters->IndirectBufferArgs = PageIndexAllocationIndirectBufferArgs;
+
+				const uint32 ArgsOffset = sizeof(uint32) * 3 * MacroGroup.MacroGroupId;
+
+				FComputeShaderUtils::AddPass(GraphBuilder, RDG_EVENT_NAME("HairStrandsAllocateVoxelPage"), *ComputeShader, Parameters, 
+					PageIndexAllocationIndirectBufferArgs,
+					ArgsOffset);
+			}
+			else
+			{
+				const FIntVector DispatchCount((CPUAllocationDesc.PageIndexCount + GroupSize - 1) / GroupSize, 1, 1);
+				check(DispatchCount.X < 65535);
+				FComputeShaderUtils::AddPass(GraphBuilder, RDG_EVENT_NAME("HairStrandsAllocateVoxelPage"), *ComputeShader, Parameters, DispatchCount);
+			}
+		}
+
+		// Prepare indirect dispatch buffers
+		{
+			check(MacroGroup.MacroGroupId < MacroGroupCount);
+
+			FVoxelAddIndirectBufferCS::FParameters* Parameters = GraphBuilder.AllocParameters<FVoxelAddIndirectBufferCS::FParameters>();
+			Parameters->MacroGroupId = MacroGroup.MacroGroupId;
+			Parameters->PageResolution = PageResolution;
+			Parameters->IndirectGroupSize = IndirectDispatchGroupSize;
+			Parameters->OutPageIndexGlobalCounter = GraphBuilder.CreateUAV(PageIndexGlobalCounter, PF_R32_UINT);
+			Parameters->OutIndirectArgsBuffer = GraphBuilder.CreateUAV(IndirectArgsBuffer);
+
+			const FIntVector DispatchCount(1, 1, 1);
+			TShaderMapRef<FVoxelAddIndirectBufferCS> ComputeShader(View.ShaderMap);
+			FComputeShaderUtils::AddPass(GraphBuilder, RDG_EVENT_NAME("HairStrandsBuildVoxelIndirectArgs"), *ComputeShader, Parameters, DispatchCount);
+		}
+	}
+
+	OutPageIndexBuffer = PageIndexBuffer;
+	OutPageIndexCoordBuffer = PageIndexCoordBuffer;
+	OutNodeDescBuffer = NodeDescBuffer;
+	OutIndirectArgsBuffer = IndirectArgsBuffer;
+	OutPageIndexGlobalCounter = PageIndexGlobalCounter;
+	OutVoxelizationViewInfoBuffer = VoxelizationViewInfoBuffer;
+}
+
+FVirtualVoxelResources AllocateVirtualVoxelResources(
+	FRHICommandListImmediate& RHICmdList,
+	const FViewInfo& View,
+	FHairStrandsMacroGroupDatas& MacroGroups)
+{
+	FRDGBuilder GraphBuilder(RHICmdList);
+	FRDGBufferRef OutPageIndexBuffer = nullptr;
+	FRDGBufferRef OutPageIndexCoordBuffer = nullptr;
+	FRDGBufferRef OutNodeDescBuffer = nullptr;
+	FRDGBufferRef OutIndirectArgsBuffer = nullptr;
+	FRDGBufferRef OutPageIndexGlobalCounter = nullptr;
+	FRDGBufferRef OutVoxelizationViewInfoBuffer = nullptr;
+
+	FVirtualVoxelResources Out;
+
+	Out.Parameters.Common.PageCountResolution		= FIntVector(GHairVirtualVoxel_PageCountPerDim, GHairVirtualVoxel_PageCountPerDim, GHairVirtualVoxel_PageCountPerDim);
+	Out.Parameters.Common.PageCount					= Out.Parameters.Common.PageCountResolution.X * Out.Parameters.Common.PageCountResolution.Y * Out.Parameters.Common.PageCountResolution.Z;
+	Out.Parameters.Common.VoxelWorldSize			= FMath::Clamp(GHairVirtualVoxel_VoxelWorldSize, 0.01f, 10.f);
+	Out.Parameters.Common.PageResolution			= FMath::RoundUpToPowerOfTwo(FMath::Clamp(GHairVirtualVoxel_PageResolution, 2, 256));
+	Out.Parameters.Common.PageTextureResolution		= Out.Parameters.Common.PageCountResolution * Out.Parameters.Common.PageResolution;
+	Out.Parameters.Common.DensityScale				= GetHairStrandsVoxelizationDensityScale();
+	Out.Parameters.Common.DepthBiasScale			= GetHairStrandsVoxelizationDepthBiasScale();
+	Out.Parameters.Common.SteppingScale				= FMath::Clamp(GHairStransVoxelRaymarchingSteppingScale, 1.f, 10.f);
+	Out.Parameters.Common.IndirectDispatchGroupSize = 64;
+
+	AddAllocateVoxelPagesPass(
+		GraphBuilder, 
+		View, 
+		MacroGroups,
+		Out.Parameters.Common.PageCountResolution,
+		Out.Parameters.Common.PageCount,
+		Out.Parameters.Common.VoxelWorldSize,
+		Out.Parameters.Common.PageResolution,
+		Out.Parameters.Common.PageTextureResolution,
+		Out.Parameters.Common.IndirectDispatchGroupSize,
+		Out.Parameters.Common.PageIndexCount,
+		OutPageIndexBuffer, 
+		OutPageIndexCoordBuffer,
+		OutNodeDescBuffer,
+		OutIndirectArgsBuffer,
+		OutPageIndexGlobalCounter,
+		OutVoxelizationViewInfoBuffer);
+
+	if (OutPageIndexBuffer)
+	{
+		GraphBuilder.QueueBufferExtraction(OutPageIndexBuffer, &Out.PageIndexBuffer, FRDGResourceState::EAccess::Read, FRDGResourceState::EPipeline::Graphics);
+	}
+
+	if (OutPageIndexCoordBuffer)
+	{
+		GraphBuilder.QueueBufferExtraction(OutPageIndexCoordBuffer, &Out.PageIndexCoordBuffer, FRDGResourceState::EAccess::Read, FRDGResourceState::EPipeline::Compute);
+	}
+
+	if (OutNodeDescBuffer)
+	{
+		GraphBuilder.QueueBufferExtraction(OutNodeDescBuffer, &Out.NodeDescBuffer, FRDGResourceState::EAccess::Read, FRDGResourceState::EPipeline::Graphics);
+	}
+
+	if (OutIndirectArgsBuffer)
+	{
+		GraphBuilder.QueueBufferExtraction(OutIndirectArgsBuffer, &Out.IndirectArgsBuffer, FRDGResourceState::EAccess::Read, FRDGResourceState::EPipeline::Compute);
+	}
+
+	if (OutPageIndexGlobalCounter)
+	{
+		GraphBuilder.QueueBufferExtraction(OutPageIndexGlobalCounter, &Out.PageIndexGlobalCounter, FRDGResourceState::EAccess::Read, FRDGResourceState::EPipeline::Compute);
+	}
+
+	if (OutVoxelizationViewInfoBuffer)
+	{
+		GraphBuilder.QueueBufferExtraction(OutVoxelizationViewInfoBuffer, &Out.VoxelizationViewInfoBuffer, FRDGResourceState::EAccess::Read, FRDGResourceState::EPipeline::Compute);
+	}
+
+	GraphBuilder.Execute();
+
+	if (Out.PageIndexBuffer)
+	{
+		Out.PageIndexBufferSRV = RHICreateShaderResourceView(Out.PageIndexBuffer->VertexBuffer, sizeof(uint32), PF_R32_UINT);
+	}
+
+	if (Out.PageIndexCoordBuffer)
+	{
+		Out.PageIndexCoordBufferSRV = RHICreateShaderResourceView(Out.PageIndexCoordBuffer->VertexBuffer, sizeof(uint32), PF_R8G8B8A8_UINT);
+	}
+
+	if (Out.NodeDescBuffer)
+	{
+		Out.NodeDescBufferSRV = RHICreateShaderResourceView(Out.NodeDescBuffer->StructuredBuffer);
+	}
+
+	{
+		// Allocation should be conservative
+		// TODO: do a partial clear with indirect call: we know how many texture page will be touched, so we know how much thread we need to launch to clear what is relevant
+		check(FMath::IsPowerOfTwo(Out.Parameters.Common.PageResolution));
+		const uint32 MipCount = FMath::Log2(Out.Parameters.Common.PageResolution) + 1;
+
+		FPooledRenderTargetDesc Desc(FPooledRenderTargetDesc::CreateVolumeDesc(
+			Out.Parameters.Common.PageTextureResolution.X, 
+			Out.Parameters.Common.PageTextureResolution.Y, 
+			Out.Parameters.Common.PageTextureResolution.Z, 
+			PF_R32_UINT, 
+			FClearValueBinding::Black, 
+			TexCreate_None, TexCreate_UAV | TexCreate_ShaderResource, 
+			false, 
+			MipCount));
+
+		GRenderTargetPool.FindFreeElement(RHICmdList, Desc, Out.PageTexture, TEXT("VoxelPageTexture"));
+	}
+
+	Out.Parameters.Common.PageIndexBuffer		= Out.PageIndexBufferSRV;
+	Out.Parameters.Common.PageIndexCoordBuffer	= Out.PageIndexCoordBufferSRV;
+	Out.Parameters.Common.NodeDescBuffer		= Out.NodeDescBufferSRV; 
+	Out.Parameters.PageTexture					= Out.PageTexture->GetRenderTargetItem().ShaderResourceTexture;
+
+	if (Out.PageIndexBufferSRV && Out.NodeDescBufferSRV)
+	{
+		Out.UniformBuffer = CreateUniformBufferImmediate(Out.Parameters, UniformBuffer_SingleFrame);
+	}
+
+	return Out;
+}
+
+
+static FRDGBufferRef IndirectVoxelPageClear(
+	FRDGBuilder& GraphBuilder,
+	const FViewInfo& ViewInfo,
+	FVirtualVoxelResources& VoxelResources)
+{
+	DECLARE_GPU_STAT(HairStrandsIndVoxelPageClear);
+	SCOPED_DRAW_EVENT(GraphBuilder.RHICmdList, HairStrandsIndVoxelPageClear);
+	SCOPED_GPU_STAT(GraphBuilder.RHICmdList, HairStrandsIndVoxelPageClear);
+
+	FRDGBufferRef ClearIndArgsBuffer = GraphBuilder.CreateBuffer(FRDGBufferDesc::CreateIndirectDesc<FRHIDispatchIndirectParameters>(1), TEXT("VirtualVoxelClearIndirectArgsBuffer"));
+
+	// Generate the indirect buffer required to clear all voxel allocated linearly in the page volume texture, using the global counter.
+	{
+		FRDGBufferRef PageIndexGlobalCounter = GraphBuilder.RegisterExternalBuffer(VoxelResources.PageIndexGlobalCounter, TEXT("HairPageIndexGlobalCounter"));
+
+		FVoxelIndPageClearBufferGenCS::FParameters* Parameters = GraphBuilder.AllocParameters<FVoxelIndPageClearBufferGenCS::FParameters>();
+		Parameters->PageResolution = VoxelResources.Parameters.Common.PageResolution;
+		Parameters->OutIndirectArgsBuffer = GraphBuilder.CreateUAV(ClearIndArgsBuffer);
+		Parameters->PageIndexGlobalCounter = GraphBuilder.CreateSRV(PageIndexGlobalCounter, PF_R32_UINT);
+
+		TShaderMapRef<FVoxelIndPageClearBufferGenCS> ComputeShader(ViewInfo.ShaderMap);
+		FComputeShaderUtils::AddPass(
+			GraphBuilder,
+			RDG_EVENT_NAME("HairStrandsVoxelGenIndBufferClearCS"),
+			*ComputeShader,
+			Parameters,
+			FIntVector(1,1,1));
+	}
+
+	// Now single dispatch to clear all the pages
+	{
+		FRDGTextureRef OutPageTexture = GraphBuilder.RegisterExternalTexture(VoxelResources.PageTexture, TEXT("HairVoxelPageTexture"));
+		FRDGBufferRef IndirectDispatchArgsBuffer = GraphBuilder.RegisterExternalBuffer(VoxelResources.IndirectArgsBuffer, TEXT("HairVoxelIndirectDispatchArgs"));
+
+		FVoxelIndPageClearCS::FParameters* Parameters = GraphBuilder.AllocParameters<FVoxelIndPageClearCS::FParameters>();
+		Parameters->VirtualVoxel = VoxelResources.Parameters.Common;
+		Parameters->OutPageTexture = GraphBuilder.CreateUAV(FRDGTextureUAVDesc(OutPageTexture));
+		Parameters->IndirectDispatchBuffer = ClearIndArgsBuffer;
+
+		TShaderMapRef<FVoxelIndPageClearCS> ComputeShader(ViewInfo.ShaderMap);
+		FComputeShaderUtils::AddPass(
+			GraphBuilder,
+			RDG_EVENT_NAME("HairStrandsVoxelIndPageClearCS"),
+			*ComputeShader,
+			Parameters,
+			ClearIndArgsBuffer,
+			0);
+	}
+
+	return ClearIndArgsBuffer;
+}
+
+static void AddVirtualVoxelizationRasterPass(
+	FRDGBuilder& GraphBuilder,
+	const FScene* Scene,
+	const FViewInfo* ViewInfo,
+	FVirtualVoxelResources& VoxelResources,
+	FHairStrandsMacroGroupData& MacroGroup)
+{
+	const bool bIsGPUDriven = GHairVirtualVoxelGPUDriven > 0;
+	const FHairStrandsMacroGroupData::TPrimitiveInfos& PrimitiveSceneInfo = MacroGroup.PrimitivesInfos;
+	DECLARE_GPU_STAT(HairStrandsVoxelize);
+	SCOPED_DRAW_EVENT(GraphBuilder.RHICmdList, HairStrandsVoxelize);
+	SCOPED_GPU_STAT(GraphBuilder.RHICmdList, HairStrandsVoxelize);
+
+	// Find the largest resolution and its dominant axis
+	FIntPoint RasterResolution(0, 0);
+	FVector	RasterProjectionSize = FVector::ZeroVector;
+	FVector RasterDirection = FVector::ZeroVector;
+	FVector RasterUp = FVector::ZeroVector;
+	const FIntVector TotalVoxelResolution = MacroGroup.VirtualVoxelNodeDesc.PageIndexResolution * VoxelResources.Parameters.Common.PageResolution;
+	{
+		FIntVector ReorderIndex(0, 0, 0);
+
+		const uint32 ResolutionXY = TotalVoxelResolution.X * TotalVoxelResolution.Y;
+		const uint32 ResolutionXZ = TotalVoxelResolution.X * TotalVoxelResolution.Z;
+		const uint32 ResolutionYZ = TotalVoxelResolution.Y * TotalVoxelResolution.Y;
+		if (ResolutionXY >= ResolutionXZ && ResolutionXY >= ResolutionYZ)
+		{
+			RasterResolution	= FIntPoint(TotalVoxelResolution.X, TotalVoxelResolution.Y);
+			RasterDirection		= FVector(0, 0, 1);			
+			ReorderIndex		= FIntVector(0, 1, 2);
+			RasterUp			= FVector(0, 1, 0);
+		}
+		else if (ResolutionXZ >= ResolutionXY && ResolutionXZ >= ResolutionYZ)
+		{
+			RasterResolution	= FIntPoint(TotalVoxelResolution.X, TotalVoxelResolution.Z);
+			RasterDirection		= FVector(0, -1, 0);
+			ReorderIndex		= FIntVector(0, 2, 1);
+			RasterUp			= FVector(0, 0, 1);
+		}
+		else
+		{
+			RasterResolution	= FIntPoint(TotalVoxelResolution.Y, TotalVoxelResolution.Z);
+			RasterDirection		= FVector(1, 0, 0);
+			ReorderIndex		= FIntVector(1, 2, 0);
+			RasterUp			= FVector(0, 0, 1);
+		}
+
+		FBox ProjRasterAABB;
+		ProjRasterAABB.Min.X = MacroGroup.VirtualVoxelNodeDesc.WorldMinAABB[ReorderIndex[0]];
+		ProjRasterAABB.Min.Y = MacroGroup.VirtualVoxelNodeDesc.WorldMinAABB[ReorderIndex[1]];
+		ProjRasterAABB.Min.Z = MacroGroup.VirtualVoxelNodeDesc.WorldMinAABB[ReorderIndex[2]];
+
+		ProjRasterAABB.Max.X = MacroGroup.VirtualVoxelNodeDesc.WorldMaxAABB[ReorderIndex[0]];
+		ProjRasterAABB.Max.Y = MacroGroup.VirtualVoxelNodeDesc.WorldMaxAABB[ReorderIndex[1]];
+		ProjRasterAABB.Max.Z = MacroGroup.VirtualVoxelNodeDesc.WorldMaxAABB[ReorderIndex[2]];
+
+		RasterProjectionSize = ProjRasterAABB.GetSize();
+	}
+
+	if (bIsGPUDriven)
+	{
+		RasterResolution = GPUDrivenViewportResolution;
+	}
+
+
+	const FBox RasterAABB(MacroGroup.VirtualVoxelNodeDesc.WorldMinAABB, MacroGroup.VirtualVoxelNodeDesc.WorldMaxAABB);
+	const FVector RasterAABBSize = RasterAABB.GetSize();
+	const FVector RasterAABBCenter = RasterAABB.GetCenter();
+	const FIntRect ViewportRect = FIntRect(0, 0, RasterResolution.X, RasterResolution.Y);
+
+	const float RadiusAtDepth1 = GStrandHairVoxelizationRasterizationScale * VoxelResources.Parameters.Common.VoxelWorldSize;
+	const bool bIsOrtho = true;
+	const FVector4 HairRenderInfo = PackHairRenderInfo(RadiusAtDepth1, RadiusAtDepth1, 1, bIsOrtho, bIsGPUDriven);
+
+	FMatrix WorldToClip;
+	{
+		FReversedZOrthoMatrix OrthoMatrix(0.5f * RasterProjectionSize.X, 0.5f * RasterProjectionSize.Y, 1.0f / RasterProjectionSize.Z, 0);
+		FLookAtMatrix LookAt(RasterAABBCenter - RasterDirection * RasterProjectionSize.Z * 0.5f, RasterAABBCenter, RasterUp);
+		WorldToClip = LookAt * OrthoMatrix;
+		MacroGroup.VirtualVoxelNodeDesc.WorldToClip = WorldToClip;
+	}
+
+	FRDGBufferRef VoxelizationViewInfoBuffer = GraphBuilder.RegisterExternalBuffer(VoxelResources.VoxelizationViewInfoBuffer);
+	FRDGTextureRef PageTexture = GraphBuilder.RegisterExternalTexture(VoxelResources.PageTexture);
+
+	FHairVoxelizationRasterPassParameters* PassParameters = GraphBuilder.AllocParameters<FHairVoxelizationRasterPassParameters>();
+	PassParameters->VirtualVoxel = VoxelResources.Parameters.Common;
+	PassParameters->WorldToClipMatrix = WorldToClip;
+	PassParameters->VoxelMinAABB = MacroGroup.VirtualVoxelNodeDesc.WorldMinAABB;
+	PassParameters->VoxelMaxAABB = MacroGroup.VirtualVoxelNodeDesc.WorldMaxAABB;
+	PassParameters->VoxelResolution = TotalVoxelResolution; // i.e., the virtual resolution
+	PassParameters->MacroGroupId = MacroGroup.MacroGroupId;
+	PassParameters->ViewportResolution = RasterResolution;
+	PassParameters->VoxelizationViewInfoBuffer = GraphBuilder.CreateSRV(VoxelizationViewInfoBuffer);
+	PassParameters->DensityTexture = GraphBuilder.CreateUAV(PageTexture);
+
+	// For debug purpose
+	#if 0
+	FRDGTextureRef DummyTexture = GraphBuilder.CreateTexture(FPooledRenderTargetDesc::Create2DDesc(RasterResolution, PF_R32_UINT, FClearValueBinding::Black, TexCreate_None, TexCreate_RenderTargetable, false), TEXT("DummyTexture"));
+	PassParameters->RenderTargets[0] = FRenderTargetBinding(DummyTexture, ERenderTargetLoadAction::EClear);
+	#endif
+
+	AddHairVoxelizationRasterPass(GraphBuilder, Scene, ViewInfo, PrimitiveSceneInfo, EHairStrandsRasterPassType::VoxelizationVirtual, ViewportRect, HairRenderInfo, RasterDirection, PassParameters);
+}
+
+
 ///////////////////////////////////////////////////////////////////////////////////////////////////
 class FVoxelFilterDepthCS : public FGlobalShader
 {
@@ -189,6 +1044,11 @@ class FVoxelFilterDepthCS : public FGlobalShader
 
 public:
 	static bool ShouldCompilePermutation(const FGlobalShaderPermutationParameters& Parameters) { return IsHairStrandsSupported(Parameters.Platform); }
+	static void ModifyCompilationEnvironment(const FGlobalShaderPermutationParameters& Parameters, FShaderCompilerEnvironment& OutEnvironment)
+	{
+		FGlobalShader::ModifyCompilationEnvironment(Parameters, OutEnvironment);
+		OutEnvironment.SetDefine(TEXT("SHADER_FILTERDEPTH_VOXEL"), 1);
+	}
 };
 
 IMPLEMENT_GLOBAL_SHADER(FVoxelFilterDepthCS, "/Engine/Private/HairStrands/HairStrandsVoxelOpaque.usf", "MainCS", SF_Compute);
@@ -196,12 +1056,12 @@ IMPLEMENT_GLOBAL_SHADER(FVoxelFilterDepthCS, "/Engine/Private/HairStrands/HairSt
 static void AddFilterVoxelOpaqueDepthPass(
 	FRDGBuilder& GraphBuilder,
 	const FViewInfo& View,
-	FHairStrandsClusterData& Cluster)
+	FHairStrandsMacroGroupData& MacroGroup)
 {
-	if (!Cluster.VoxelResources.DensityTexture)
+	if (!MacroGroup.VoxelResources.DensityTexture)
 		return;
 
-	TRefCountPtr<IPooledRenderTarget> DensityTexture = Cluster.VoxelResources.DensityTexture;
+	TRefCountPtr<IPooledRenderTarget> DensityTexture = MacroGroup.VoxelResources.DensityTexture;
 	check(DensityTexture->GetDesc().Extent.X == DensityTexture->GetDesc().Extent.Y);
 
 	const uint32 VoxelResolution = DensityTexture->GetDesc().Extent.X;
@@ -249,6 +1109,11 @@ class FVoxelGenerateMipCS : public FGlobalShader
 
 public:
 	static bool ShouldCompilePermutation(const FGlobalShaderPermutationParameters& Parameters) { return IsHairStrandsSupported(Parameters.Platform); }
+	static void ModifyCompilationEnvironment(const FGlobalShaderPermutationParameters& Parameters, FShaderCompilerEnvironment& OutEnvironment)
+	{
+		FGlobalShader::ModifyCompilationEnvironment(Parameters, OutEnvironment);
+		OutEnvironment.SetDefine(TEXT("SHADER_VOXEL"), 1);
+	}
 };
 
 IMPLEMENT_GLOBAL_SHADER(FVoxelGenerateMipCS, "/Engine/Private/HairStrands/HairStrandsVoxelMip.usf", "MainCS", SF_Compute);
@@ -256,12 +1121,12 @@ IMPLEMENT_GLOBAL_SHADER(FVoxelGenerateMipCS, "/Engine/Private/HairStrands/HairSt
 static void AddVoxelGenerateMipPass(
 	FRDGBuilder& GraphBuilder,
 	const FViewInfo& View,
-	FHairStrandsClusterData& Cluster)
+	FHairStrandsMacroGroupData& MacroGroup)
 {
-	if (!Cluster.VoxelResources.DensityTexture)
+	if (!MacroGroup.VoxelResources.DensityTexture)
 		return;
 
-	TRefCountPtr<IPooledRenderTarget> DensityTexture = Cluster.VoxelResources.DensityTexture;
+	TRefCountPtr<IPooledRenderTarget> DensityTexture = MacroGroup.VoxelResources.DensityTexture;
 	check(DensityTexture->GetDesc().Extent.X == DensityTexture->GetDesc().Extent.Y);
 
 	const uint32 NumLevelPerPass = GHairStrandsVoxelMipMethod > 0 ? 2 : 1;
@@ -307,19 +1172,142 @@ static void AddVoxelGenerateMipPass(
 	}
 }
 
-static void RenderVoxelization(
-	FRHICommandListImmediate& RHICmdList,
+
+///////////////////////////////////////////////////////////////////////////////////////////////////
+class FVirtualVoxelGenerateMipCS : public FGlobalShader
+{
+	DECLARE_GLOBAL_SHADER(FVirtualVoxelGenerateMipCS);
+	SHADER_USE_PARAMETER_STRUCT(FVirtualVoxelGenerateMipCS, FGlobalShader);
+
+	BEGIN_SHADER_PARAMETER_STRUCT(FParameters, )
+		SHADER_PARAMETER_STRUCT_INCLUDE(FSceneTextureParameters, SceneTextures)
+		SHADER_PARAMETER(FIntVector, PageCountResolution)
+		SHADER_PARAMETER(uint32, PageResolution)
+		SHADER_PARAMETER(uint32, SourceMip)
+		SHADER_PARAMETER(uint32, TargetMip)
+
+		SHADER_PARAMETER_RDG_BUFFER(StructuredBuffer, IndirectDispatchArgs)
+		SHADER_PARAMETER_RDG_TEXTURE_SRV(Texture3D, InDensityTexture)
+		SHADER_PARAMETER_RDG_TEXTURE_UAV(RWTexture3D, OutDensityTexture)
+
+	END_SHADER_PARAMETER_STRUCT()
+
+public:
+	static bool ShouldCompilePermutation(const FGlobalShaderPermutationParameters& Parameters) { return IsHairStrandsSupported(Parameters.Platform); }
+	static void ModifyCompilationEnvironment(const FGlobalShaderPermutationParameters& Parameters, FShaderCompilerEnvironment& OutEnvironment)
+	{
+		FGlobalShader::ModifyCompilationEnvironment(Parameters, OutEnvironment);
+		OutEnvironment.SetDefine(TEXT("SHADER_VIRTUALVOXEL"), 1);
+	}
+};
+
+
+class FVirtualVoxelIndirectArgMipCS : public FGlobalShader
+{
+	DECLARE_GLOBAL_SHADER(FVirtualVoxelIndirectArgMipCS);
+	SHADER_USE_PARAMETER_STRUCT(FVirtualVoxelIndirectArgMipCS, FGlobalShader);
+
+	BEGIN_SHADER_PARAMETER_STRUCT(FParameters, )
+		SHADER_PARAMETER(uint32, PageResolution)
+		SHADER_PARAMETER(uint32, TargetMipIndex)
+		SHADER_PARAMETER(uint32, DispatchGroupSize)
+		SHADER_PARAMETER_RDG_BUFFER_SRV(Buffer, InIndirectArgs)
+		SHADER_PARAMETER_RDG_BUFFER_UAV(RWBuffer, OutIndirectArgs)
+	END_SHADER_PARAMETER_STRUCT()
+
+public:
+	static bool ShouldCompilePermutation(const FGlobalShaderPermutationParameters& Parameters) { return IsHairStrandsSupported(Parameters.Platform); }
+	static void ModifyCompilationEnvironment(const FGlobalShaderPermutationParameters& Parameters, FShaderCompilerEnvironment& OutEnvironment)
+	{
+		FGlobalShader::ModifyCompilationEnvironment(Parameters, OutEnvironment);
+		OutEnvironment.SetDefine(TEXT("SHADER_INDIRECTARGS"), 1);
+	}
+};
+
+IMPLEMENT_GLOBAL_SHADER(FVirtualVoxelGenerateMipCS, "/Engine/Private/HairStrands/HairStrandsVoxelMip.usf", "MainCS", SF_Compute);
+IMPLEMENT_GLOBAL_SHADER(FVirtualVoxelIndirectArgMipCS, "/Engine/Private/HairStrands/HairStrandsVoxelMip.usf", "MainCS", SF_Compute);
+
+
+static void AddVirtualVoxelGenerateMipPass(
+	FRDGBuilder& GraphBuilder,
+	const FViewInfo& View,
+	FHairStrandsMacroGroupDatas& MacroGroups,
+	FRDGBufferRef IndirectArgsBuffer)
+{
+	if (!MacroGroups.VirtualVoxelResources.IsValid())
+		return;
+
+	DECLARE_GPU_STAT(HairStrandsDensityMipGen);
+	SCOPED_DRAW_EVENT(GraphBuilder.RHICmdList, HairStrandsDensityMipGen);
+	SCOPED_GPU_STAT(GraphBuilder.RHICmdList, HairStrandsDensityMipGen);
+
+	FVirtualVoxelResources& VoxelResources = MacroGroups.VirtualVoxelResources;
+
+	const uint32 MipCount = VoxelResources.PageTexture->GetDesc().NumMips;
+	FRDGTextureRef VoxelDensityTexture = GraphBuilder.RegisterExternalTexture(VoxelResources.PageTexture, TEXT("HairVirtualVoxelDensityTexture"));
+
+	// Prepare indirect dispatch for all the pages this frame (allocated linearly in 3D DensityTexture)
+	TArray<FRDGBufferRef> MipIndirectArgsBuffers;
+	for (uint32 MipIt = 0; MipIt < MipCount - 1; ++MipIt)
+	{
+		const uint32 TargetMipIndex = MipIt + 1;
+		const uint32 DispatchGroupSize = 64;
+		FRDGBufferRef MipIndirectArgs = GraphBuilder.CreateBuffer(FRDGBufferDesc::CreateIndirectDesc<FRHIDispatchIndirectParameters>(), TEXT("VirtualVoxelMipIndirectArgsBuffer"));
+		MipIndirectArgsBuffers.Add(MipIndirectArgs);
+
+		FVirtualVoxelIndirectArgMipCS::FParameters* Parameters = GraphBuilder.AllocParameters<FVirtualVoxelIndirectArgMipCS::FParameters>();
+		Parameters->PageResolution		= VoxelResources.Parameters.Common.PageResolution;
+		Parameters->TargetMipIndex		= TargetMipIndex;
+		Parameters->DispatchGroupSize	= DispatchGroupSize;
+		Parameters->InIndirectArgs		= GraphBuilder.CreateSRV(IndirectArgsBuffer);
+		Parameters->OutIndirectArgs		= GraphBuilder.CreateUAV(MipIndirectArgs);
+
+		TShaderMapRef<FVirtualVoxelIndirectArgMipCS> ComputeShader(View.ShaderMap);
+		FComputeShaderUtils::AddPass(GraphBuilder, RDG_EVENT_NAME("HairStrandsBuildVoxelMipIndirectArgs"), *ComputeShader, Parameters, FIntVector(1, 1, 1));
+	}
+
+	// Generate MIP level (in one go for all allocated pages)
+	for (uint32 MipIt = 0; MipIt < MipCount - 1; ++MipIt)
+	{
+		const uint32 SourceMipIndex = MipIt;
+		const uint32 TargetMipIndex = MipIt + 1;
+
+		FVirtualVoxelGenerateMipCS::FParameters* Parameters = GraphBuilder.AllocParameters<FVirtualVoxelGenerateMipCS::FParameters>();
+		Parameters->InDensityTexture = GraphBuilder.CreateSRV(FRDGTextureSRVDesc::CreateForMipLevel(VoxelDensityTexture, MipIt));
+		Parameters->OutDensityTexture = GraphBuilder.CreateUAV(FRDGTextureUAVDesc(VoxelDensityTexture, MipIt + 1));
+		Parameters->PageResolution = VoxelResources.Parameters.Common.PageResolution;
+		Parameters->PageCountResolution = VoxelResources.Parameters.Common.PageCountResolution;
+		Parameters->SourceMip = SourceMipIndex;
+		Parameters->TargetMip = TargetMipIndex;
+		Parameters->IndirectDispatchArgs = MipIndirectArgsBuffers[MipIt];
+
+		TShaderMapRef<FVirtualVoxelGenerateMipCS> ComputeShader(View.ShaderMap);
+		ClearUnusedGraphResources(*ComputeShader, Parameters);
+		GraphBuilder.AddPass(
+			RDG_EVENT_NAME("HairStrandsComputeVoxelMip"),
+			Parameters,
+			ERDGPassFlags::Compute | ERDGPassFlags::GenerateMips,
+			[Parameters, ComputeShader](FRHICommandList& RHICmdList)
+		{
+			FComputeShaderUtils::DispatchIndirect(RHICmdList, *ComputeShader, *Parameters, Parameters->IndirectDispatchArgs->GetIndirectRHICallBuffer(), 0);
+		});
+	}
+}
+/////////////////////////////////////////////////////////////////////////////////////////
+
+static void AddVoxelizationRasterPass(
+	FRDGBuilder& GraphBuilder,
 	const FScene* Scene,
 	const FViewInfo* ViewInfo,
-	FHairStrandsClusterData& Cluster)
+	FHairStrandsMacroGroupData& MacroGroup)
 {
-	const FHairStrandsClusterData::TPrimitiveInfos& PrimitiveSceneInfo = Cluster.PrimitivesInfos;
-	const FBoxSphereBounds& Bounds = Cluster.Bounds;
-	FHairStrandsVoxelResources& VoxelResources = Cluster.VoxelResources;
+	DECLARE_GPU_STAT(HairStrandsVoxelize);
+	SCOPED_DRAW_EVENT(GraphBuilder.RHICmdList, HairStrandsVoxelize);
+	SCOPED_GPU_STAT(GraphBuilder.RHICmdList, HairStrandsVoxelize);
 
-	DECLARE_GPU_STAT(HairStrandsVoxelization);
-	SCOPED_DRAW_EVENT(RHICmdList, HairStrandsVoxelization);
-	SCOPED_GPU_STAT(RHICmdList, HairStrandsVoxelization);
+	const FHairStrandsMacroGroupData::TPrimitiveInfos& PrimitiveSceneInfo = MacroGroup.PrimitivesInfos;
+	const FBoxSphereBounds& Bounds = MacroGroup.Bounds;
+	FHairStrandsVoxelResources& VoxelResources = MacroGroup.VoxelResources;
 
 	const uint32 ResolutionDim = FMath::RoundUpToPowerOfTwo(GHairVoxelizationResolution);
 	const EPixelFormat Format = PF_R32_UINT;
@@ -334,145 +1322,161 @@ static void RenderVoxelization(
 		}
 	}
 
-	// todo remove this dummy texture
-	TRefCountPtr<IPooledRenderTarget> DummyTexture;
-	{
-		FPooledRenderTargetDesc Desc(FPooledRenderTargetDesc::Create2DDesc(Resolution, PF_R32_UINT, FClearValueBinding::Black, TexCreate_None, TexCreate_RenderTargetable, false));
-		GRenderTargetPool.FindFreeElement(RHICmdList, Desc, DummyTexture, TEXT("DummyTexture"));
-	}
-
 	const bool bVoxelizeMaterial = GHairVoxelizationMaterialEnable > 0;
+	FRDGTextureRef DensityTexture = nullptr;
+	FRDGTextureRef TangentXTexture = nullptr;
+	FRDGTextureRef TangentYTexture = nullptr;
+	FRDGTextureRef TangentZTexture = nullptr;
+	FRDGTextureRef MaterialTexture = nullptr;
 	{
-		FPooledRenderTargetDesc Desc(FPooledRenderTargetDesc::CreateVolumeDesc(ResolutionDim, ResolutionDim, ResolutionDim, PF_R32_UINT, FClearValueBinding::Black, TexCreate_None, TexCreate_UAV | TexCreate_ShaderResource, false, MipCount));
-		GRenderTargetPool.FindFreeElement(RHICmdList, Desc, VoxelResources.DensityTexture, TEXT("DensityTexture"));
+		FPooledRenderTargetDesc Desc = FPooledRenderTargetDesc::CreateVolumeDesc(ResolutionDim, ResolutionDim, ResolutionDim, PF_R32_UINT, FClearValueBinding::Black, TexCreate_None, TexCreate_UAV | TexCreate_ShaderResource, false, MipCount);
+		DensityTexture = GraphBuilder.CreateTexture(Desc, TEXT("DensityTexture"));
 		if (bVoxelizeMaterial)
 		{
-			GRenderTargetPool.FindFreeElement(RHICmdList, Desc, VoxelResources.TangentXTexture, TEXT("TangentXTexture"));
-			GRenderTargetPool.FindFreeElement(RHICmdList, Desc, VoxelResources.TangentYTexture, TEXT("TangentYTexture"));
-			GRenderTargetPool.FindFreeElement(RHICmdList, Desc, VoxelResources.TangentZTexture, TEXT("TangentZTexture"));
-			GRenderTargetPool.FindFreeElement(RHICmdList, Desc, VoxelResources.MaterialTexture, TEXT("MaterialTexture"));
-		}
-
-		uint32 ClearValues[4] = { 0, 0, 0, 0 };
-		RHICmdList.ClearTinyUAV(VoxelResources.DensityTexture->GetRenderTargetItem().UAV, ClearValues);
-		if (bVoxelizeMaterial)
-		{
-			RHICmdList.ClearTinyUAV(VoxelResources.TangentXTexture->GetRenderTargetItem().UAV, ClearValues);
-			RHICmdList.ClearTinyUAV(VoxelResources.TangentYTexture->GetRenderTargetItem().UAV, ClearValues);
-			RHICmdList.ClearTinyUAV(VoxelResources.TangentZTexture->GetRenderTargetItem().UAV, ClearValues);
-			RHICmdList.ClearTinyUAV(VoxelResources.MaterialTexture->GetRenderTargetItem().UAV, ClearValues);
+			TangentXTexture = GraphBuilder.CreateTexture(Desc, TEXT("TangentXTexture"));
+			TangentYTexture = GraphBuilder.CreateTexture(Desc, TEXT("TangentYTexture"));
+			TangentZTexture = GraphBuilder.CreateTexture(Desc, TEXT("TangentZTexture"));
+			MaterialTexture = GraphBuilder.CreateTexture(Desc, TEXT("MaterialTexture"));
 		}
 	}
 
-	FRHIRenderPassInfo RPInfo(DummyTexture->GetRenderTargetItem().TargetableTexture, ERenderTargetActions::Load_Store);
-	RPInfo.UAVIndex = 0;
-	RPInfo.NumUAVs = bVoxelizeMaterial ? 5 : 1;
-	RPInfo.UAVs[0] = VoxelResources.DensityTexture->GetRenderTargetItem().UAV;
-	if (bVoxelizeMaterial)
-	{
-		RPInfo.UAVs[1] = VoxelResources.TangentXTexture->GetRenderTargetItem().UAV;
-		RPInfo.UAVs[2] = VoxelResources.TangentYTexture->GetRenderTargetItem().UAV;
-		RPInfo.UAVs[3] = VoxelResources.TangentZTexture->GetRenderTargetItem().UAV;
-		RPInfo.UAVs[4] = VoxelResources.MaterialTexture->GetRenderTargetItem().UAV;
-	}
-
 	const FIntRect ViewportRect(0, 0, ResolutionDim, ResolutionDim);
-	const FVector LightDirection = FVector(1, 0, 0);
-	const FVector LightPosition = FVector(0, 0, 0);
+	const FVector RasterDirection = FVector(1, 0, 0);
 	const FSphere SphereBound = Bounds.GetSphere();
 	const float SphereRadius = SphereBound.W * GHairVoxelizationAABBScale;
+
 	const float RadiusAtDepth1 = GStrandHairVoxelizationRasterizationScale * SphereRadius / FMath::Min(Resolution.X, Resolution.Y);
-	const float bIsOrtho = 1.0f;
+	const bool bIsOrtho = true;
+	const bool bIsGPUDriven = false;
+	const FVector4 HairRenderInfo = PackHairRenderInfo(RadiusAtDepth1, RadiusAtDepth1, 1, bIsOrtho, bIsGPUDriven);
 
-	TUniformBufferRef<FDeepShadowPassUniformParameters> DeepShadowPassUniformParameters;
 	{
 		FReversedZOrthoMatrix OrthoMatrix(SphereRadius, SphereRadius, 1.f / (2 * SphereRadius), 0);
-		FLookAtMatrix LookAt(SphereBound.Center - LightDirection * SphereRadius, SphereBound.Center, FVector(0, 0, 1));
+		FLookAtMatrix LookAt(SphereBound.Center - RasterDirection * SphereRadius, SphereBound.Center, FVector(0, 0, 1));
 
 		VoxelResources.WorldToClip = LookAt * OrthoMatrix;
 		VoxelResources.MinAABB = Bounds.GetSphere().Center - SphereRadius;
 		VoxelResources.MaxAABB = Bounds.GetSphere().Center + SphereRadius;
+	}
 
-		FDeepShadowPassUniformParameters PassParameters;
-		PassParameters.WorldToClipMatrix = VoxelResources.WorldToClip;;
-		PassParameters.SliceValue = FVector4(1, 1, 1, 1);
-		PassParameters.FrontDepthTexture = nullptr;
-		PassParameters.AtlasRect = ViewportRect;
-		PassParameters.VoxelMinAABB = VoxelResources.MinAABB;
-		PassParameters.VoxelMaxAABB = VoxelResources.MaxAABB;
-		PassParameters.VoxelResolution = ViewportRect.Width();
-		DeepShadowPassUniformParameters = CreateUniformBufferImmediate(PassParameters, UniformBuffer_SingleDraw, EUniformBufferValidation::None);
-	}
-
-	TUniformBufferRef<FViewUniformShaderParameters> ViewUniformShaderParameters;
-	{
-		// Save some status we need to restore
-		const FVector SavedViewForward = ViewInfo->CachedViewUniformShaderParameters->ViewForward;
-		// Update our view parameters
-		ViewInfo->CachedViewUniformShaderParameters->HairRenderInfo.X = RadiusAtDepth1;
-		ViewInfo->CachedViewUniformShaderParameters->HairRenderInfo.Y = RadiusAtDepth1;
-		ViewInfo->CachedViewUniformShaderParameters->HairRenderInfo.Z = bIsOrtho;
-		ViewInfo->CachedViewUniformShaderParameters->ViewForward = LightDirection;
-		// Create the uniform buffer
-		ViewUniformShaderParameters = TUniformBufferRef<FViewUniformShaderParameters>::CreateUniformBufferImmediate(*ViewInfo->CachedViewUniformShaderParameters, UniformBuffer_SingleDraw);
-		// Restore the view cached parameters
-		ViewInfo->CachedViewUniformShaderParameters->ViewForward = SavedViewForward;
-	}
-
-	RHICmdList.BeginRenderPass(RPInfo, TEXT("HairVoxelization"));
-	RasterHairStrands(
-		RHICmdList,
+	FHairVoxelizationRasterPassParameters* PassParameters = GraphBuilder.AllocParameters<FHairVoxelizationRasterPassParameters>();
+	PassParameters->WorldToClipMatrix = VoxelResources.WorldToClip;
+	PassParameters->VoxelMinAABB = VoxelResources.MinAABB;
+	PassParameters->VoxelMaxAABB = VoxelResources.MaxAABB;
+	PassParameters->VoxelResolution = FIntVector(ViewportRect.Width(), ViewportRect.Width(), ViewportRect.Width());
+	PassParameters->MacroGroupId = MacroGroup.MacroGroupId;
+	PassParameters->ViewportResolution = FIntPoint(ResolutionDim, ResolutionDim);
+	PassParameters->DensityTexture = GraphBuilder.CreateUAV(DensityTexture);
+	if (bVoxelizeMaterial)
+	{
+		PassParameters->TangentXTexture = GraphBuilder.CreateUAV(TangentXTexture);
+		PassParameters->TangentYTexture = GraphBuilder.CreateUAV(TangentYTexture);
+		PassParameters->TangentZTexture = GraphBuilder.CreateUAV(TangentZTexture);
+		PassParameters->MaterialTexture = GraphBuilder.CreateUAV(MaterialTexture);
+	}
+
+	AddHairVoxelizationRasterPass(
+		GraphBuilder,
 		Scene,
 		ViewInfo,
 		PrimitiveSceneInfo,
 		bVoxelizeMaterial ? EHairStrandsRasterPassType::VoxelizationMaterial : EHairStrandsRasterPassType::Voxelization,
 		ViewportRect,
-		ViewUniformShaderParameters,
-		DeepShadowPassUniformParameters);
-	RHICmdList.EndRenderPass();
+		HairRenderInfo,
+		RasterDirection,
+		PassParameters);
 
-	if (GHairVoxelInjectOpaqueDepthEnable)
-	{
-		FRDGBuilder GraphBuilder(RHICmdList);
-		AddVoxelInjectOpaquePass(GraphBuilder, *ViewInfo, Cluster);
-		if (GHairVoxelFilterOpaqueDepthEnable)
-		{
-			AddFilterVoxelOpaqueDepthPass(GraphBuilder, *ViewInfo, Cluster);
-		}
-		GraphBuilder.Execute();
-	}
 
+	GraphBuilder.QueueTextureExtraction(DensityTexture, &VoxelResources.DensityTexture);
+	if (bVoxelizeMaterial)
 	{
-		FRDGBuilder GraphBuilder(RHICmdList);
-		AddVoxelGenerateMipPass(GraphBuilder, *ViewInfo, Cluster);
-		GraphBuilder.Execute();
+		GraphBuilder.QueueTextureExtraction(TangentXTexture, &VoxelResources.TangentXTexture);
+		GraphBuilder.QueueTextureExtraction(TangentYTexture, &VoxelResources.TangentYTexture);
+		GraphBuilder.QueueTextureExtraction(TangentZTexture, &VoxelResources.TangentZTexture);
+		GraphBuilder.QueueTextureExtraction(MaterialTexture, &VoxelResources.MaterialTexture);
 	}
+
 }
 
 void VoxelizeHairStrands(
 	FRHICommandListImmediate& RHICmdList, 
 	const FScene* Scene, 
 	const TArray<FViewInfo>& Views,
-	FHairStrandsClusterViews& DeepShadowClusterViews)
+	FHairStrandsMacroGroupViews& MacroGroupsViews)
 {
 	if (!IsHairStrandsVoxelizationEnable())
 		return;
 
-	FHairStrandsClusterViews PrimitivesClusterViews;
+	FHairStrandsMacroGroupViews PrimitivesClusterViews;
 	for (int32 ViewIndex = 0; ViewIndex < Views.Num(); ViewIndex++)
 	{			
-		if (ViewIndex >= DeepShadowClusterViews.Views.Num())
+		if (ViewIndex >= MacroGroupsViews.Views.Num())
 			continue;
 
 		const FViewInfo& View = Views[ViewIndex];
-		FHairStrandsClusterDatas& ClusterDatas = DeepShadowClusterViews.Views[ViewIndex];
-		for (FHairStrandsClusterData& Cluster : ClusterDatas.Datas)
+		FHairStrandsMacroGroupDatas& MacroGroupDatas = MacroGroupsViews.Views[ViewIndex];
+
+		if (MacroGroupDatas.Datas.Num() == 0)
+			continue;
+
+		DECLARE_GPU_STAT(HairStrandsVoxelization);
+		SCOPED_DRAW_EVENT(RHICmdList, HairStrandsVoxelization);
+		SCOPED_GPU_STAT(RHICmdList, HairStrandsVoxelization);
+
+		if (GHairVirtualVoxel)
+		{
+			if (MacroGroupDatas.Datas.Num() > 0)
+			{
+				// Toto moves this function into the render graph. At the moment this is not possible as this functions 
+				// generates internally a non-transient constant buffer which initialized VirtualVoxelResources. This 
+				// needs to be rewritten/worked out.
+				MacroGroupDatas.VirtualVoxelResources = AllocateVirtualVoxelResources(RHICmdList, View, MacroGroupDatas);
+
+				FRDGBuilder GraphBuilder(RHICmdList);
+				FRDGBufferRef ClearIndArgsBuffer = IndirectVoxelPageClear(GraphBuilder, View, MacroGroupDatas.VirtualVoxelResources);
+
+				for (FHairStrandsMacroGroupData& MacroGroup : MacroGroupDatas.Datas)
+				{
+					AddVirtualVoxelizationRasterPass(GraphBuilder, Scene, &View, MacroGroupDatas.VirtualVoxelResources, MacroGroup);
+				}
+
+				if (GHairVoxelInjectOpaqueDepthEnable)
+				{
+					for (FHairStrandsMacroGroupData& MacroGroup : MacroGroupDatas.Datas)
+					{
+						AddVirtualVoxelInjectOpaquePass(GraphBuilder, View, MacroGroupDatas.VirtualVoxelResources, MacroGroup);
+					}					
+				}
+
+				AddVirtualVoxelGenerateMipPass(GraphBuilder, View, MacroGroupDatas, ClearIndArgsBuffer);
+
+				GraphBuilder.Execute();
+			}
+		}
+		else
 		{
-			RenderVoxelization(
-				RHICmdList,
-				Scene,
-				&View,
-				Cluster);
+			for (FHairStrandsMacroGroupData& MacroGroup : MacroGroupDatas.Datas)
+			{
+				FRDGBuilder GraphBuilder(RHICmdList);
+				AddVoxelizationRasterPass(
+					GraphBuilder,
+					Scene,
+					&View,
+					MacroGroup);
+
+				if (GHairVoxelInjectOpaqueDepthEnable)
+				{
+					AddVoxelInjectOpaquePass(GraphBuilder, View, MacroGroup);
+					if (GHairVoxelFilterOpaqueDepthEnable)
+					{
+						AddFilterVoxelOpaqueDepthPass(GraphBuilder, View, MacroGroup);
+					}
+				}
+
+				AddVoxelGenerateMipPass(GraphBuilder, View, MacroGroup);
+
+				GraphBuilder.Execute();
+			}
 		}
 	}
 }
diff --git a/Engine/Source/Runtime/Renderer/Private/HairStrands/HairStrandsVoxelization.h b/Engine/Source/Runtime/Renderer/Private/HairStrands/HairStrandsVoxelization.h
--- a/Engine/Source/Runtime/Renderer/Private/HairStrands/HairStrandsVoxelization.h
+++ b/Engine/Source/Runtime/Renderer/Private/HairStrands/HairStrandsVoxelization.h
@@ -8,6 +8,7 @@
 
 #include "CoreMinimal.h"
 #include "RendererInterface.h"
+#include "RenderGraphResources.h"
 
 struct FHairStrandsVoxelResources
 {
@@ -21,6 +22,73 @@ struct FHairStrandsVoxelResources
 	FVector MaxAABB;
 };
 
+struct FVirtualVoxelNodeDesc
+{
+	FVector WorldMinAABB = FVector::ZeroVector;
+	FVector WorldMaxAABB = FVector::ZeroVector;
+	FIntVector PageIndexResolution = FIntVector::ZeroValue;
+	FMatrix WorldToClip;
+};
+
+struct FPackedVirtualVoxelNodeDesc
+{
+	// This is just a placeholder having the correct size. The actual definition is in HairStradsNVoxelPageCommon.ush
+	const static EPixelFormat Format = PF_R32G32B32A32_UINT;
+	const static uint32 ComponentCount = 2;
+
+	// Shader View is struct { uint4; uint4; }
+	FVector	MinAABB;
+	uint32	PackedPageIndexResolution;
+	FVector	MaxAABB;
+	uint32	PageIndexOffset;
+};
+
+BEGIN_SHADER_PARAMETER_STRUCT(FVirtualVoxelCommonParameters, )
+	SHADER_PARAMETER(FIntVector, PageCountResolution)
+	SHADER_PARAMETER(float, VoxelWorldSize)
+	SHADER_PARAMETER(FIntVector, PageTextureResolution)
+	SHADER_PARAMETER(uint32, PageCount)
+	SHADER_PARAMETER(uint32, PageResolution)
+	SHADER_PARAMETER(uint32, PageIndexCount)
+	SHADER_PARAMETER(uint32, IndirectDispatchGroupSize)
+	SHADER_PARAMETER(float, DensityScale)
+	SHADER_PARAMETER(float, DepthBiasScale)
+	SHADER_PARAMETER(float, SteppingScale)
+	SHADER_PARAMETER_SRV(Buffer<uint>, PageIndexBuffer)
+	SHADER_PARAMETER_SRV(Buffer<uint>, PageIndexCoordBuffer)
+	SHADER_PARAMETER_SRV(StructuredBuffer<FPackedVirtualVoxelNodeDesc>, NodeDescBuffer) // Packed into 2 x uint4
+END_SHADER_PARAMETER_STRUCT()
+
+BEGIN_GLOBAL_SHADER_PARAMETER_STRUCT(FVirtualVoxelParameters, RENDERER_API)
+	SHADER_PARAMETER_STRUCT_INCLUDE(FVirtualVoxelCommonParameters, Common)
+	SHADER_PARAMETER_TEXTURE(Texture3D<uint>, PageTexture)
+END_GLOBAL_SHADER_PARAMETER_STRUCT()
+
+struct FVirtualVoxelResources
+{
+	FVirtualVoxelParameters	Parameters;
+	TUniformBufferRef<FVirtualVoxelParameters> UniformBuffer;
+
+	TRefCountPtr<IPooledRenderTarget>	PageTexture;
+
+	TRefCountPtr<FPooledRDGBuffer>		PageIndexBuffer;
+	FShaderResourceViewRHIRef			PageIndexBufferSRV = nullptr;
+
+	TRefCountPtr<FPooledRDGBuffer>		NodeDescBuffer;
+	FShaderResourceViewRHIRef			NodeDescBufferSRV = nullptr;
+
+	TRefCountPtr<FPooledRDGBuffer>		PageIndexCoordBuffer;
+	FShaderResourceViewRHIRef			PageIndexCoordBufferSRV = nullptr;
+
+	TRefCountPtr<FPooledRDGBuffer>		IndirectArgsBuffer;
+
+	TRefCountPtr<FPooledRDGBuffer>		PageIndexGlobalCounter;
+
+	TRefCountPtr<FPooledRDGBuffer>		VoxelizationViewInfoBuffer;
+
+	const bool IsValid() const { return UniformBuffer.IsValid(); }
+};
+
 /// Global enable/disable for hair voxelization
 bool IsHairStrandsVoxelizationEnable();
 bool IsHairStrandsForVoxelTransmittanceAndShadowEnable();
@@ -31,4 +99,4 @@ void VoxelizeHairStrands(
 	FRHICommandListImmediate& RHICmdList,
 	const class FScene* Scene,
 	const TArray<FViewInfo>& Views,
-	struct FHairStrandsClusterViews& DeepShadowClusterViews);
\ No newline at end of file
+	struct FHairStrandsMacroGroupViews& MacroGroupViews);
\ No newline at end of file
diff --git a/Engine/Source/Runtime/Renderer/Private/LightRendering.cpp b/Engine/Source/Runtime/Renderer/Private/LightRendering.cpp
--- a/Engine/Source/Runtime/Renderer/Private/LightRendering.cpp
+++ b/Engine/Source/Runtime/Renderer/Private/LightRendering.cpp
@@ -316,6 +316,7 @@ class FDeferredLightPS : public FGlobalShader
 		HairLUTSampler.Bind(Initializer.ParameterMap, TEXT("HairLUTSampler"));
 		HairComponents.Bind(Initializer.ParameterMap, TEXT("HairComponents"));
 		HairShadowMaskValid.Bind(Initializer.ParameterMap, TEXT("HairShadowMaskValid"));
+		HairDualScatteringRoughnessOverride.Bind(Initializer.ParameterMap, TEXT("HairDualScatteringRoughnessOverride"));
 
 		HairCategorizationTexture.Bind(Initializer.ParameterMap, TEXT("HairCategorizationTexture"));
 		HairVisibilityNodeOffsetAndCount.Bind(Initializer.ParameterMap, TEXT("HairVisibilityNodeOffsetAndCount"));
@@ -374,6 +375,7 @@ public:
 		Ar << HairLUTSampler;
 		Ar << HairComponents;
 		Ar << HairShadowMaskValid;
+		Ar << HairDualScatteringRoughnessOverride;
 
 		return bShaderHasOutdatedParameters;
 	}
@@ -568,6 +570,17 @@ private:
 				HairComponents,
 				InHairComponents);
 		}
+
+		if (HairDualScatteringRoughnessOverride.IsBound())
+		{
+			const float DualScatteringRoughness = GetHairDualScatteringRoughnessOverride();
+			SetShaderValue(
+				RHICmdList,
+				ShaderRHI,
+				HairDualScatteringRoughnessOverride,
+				DualScatteringRoughness);
+		}
+		
 	}
 
 	FSceneTextureShaderParameters SceneTextureParameters;
@@ -595,6 +608,7 @@ private:
 	FShaderResourceParameter HairLUTSampler;
 	FShaderParameter HairComponents;
 	FShaderParameter HairShadowMaskValid;
+	FShaderParameter HairDualScatteringRoughnessOverride;
 	};
 
 IMPLEMENT_GLOBAL_SHADER(FDeferredLightPS, "/Engine/Private/DeferredLightPixelShaders.usf", "DeferredLightPixelMain", SF_Pixel);
@@ -920,7 +934,7 @@ void FDeferredShadingSceneRenderer::RenderLights(FRHICommandListImmediate& RHICm
 	const FHairStrandsVisibilityViews* InHairVisibilityViews = HairDatas ? &HairDatas->HairVisibilityViews : nullptr;
 	auto HasHairStrandsClusters = [bUseHairLighting, HairDatas](int32 ViewIndex)
 	{
-		return HairDatas && ViewIndex < HairDatas->HairClusterPerViews.Views.Num() && HairDatas->HairClusterPerViews.Views[ViewIndex].Datas.Num() > 0;
+		return HairDatas && ViewIndex < HairDatas->MacroGroupsPerViews.Views.Num() && HairDatas->MacroGroupsPerViews.Views[ViewIndex].Datas.Num() > 0;
 	};
 
 	auto GetHairStrandsCategorizationTexture = [HairDatas](int32 ViewIndex, FRDGBuilder& GraphBuilder)
@@ -1053,7 +1067,8 @@ void FDeferredShadingSceneRenderer::RenderLights(FRHICommandListImmediate& RHICm
 					const FLightSceneInfo* const LightSceneInfo = SortedLightInfo.LightSceneInfo;
 
 					FHairStrandsTransmittanceMaskData TransmittanceMaskData;
-					TransmittanceMaskData = RenderHairStrandsTransmittanceMask(RHICmdList, Views, LightSceneInfo, HairDatas);
+					TRefCountPtr<IPooledRenderTarget> NullScreenShadowMaskSubPixelTexture = nullptr;
+					TransmittanceMaskData = RenderHairStrandsTransmittanceMask(RHICmdList, Views, LightSceneInfo, HairDatas, NullScreenShadowMaskSubPixelTexture);
 
 					// Render the light to the scene color buffer, using a 1x1 white texture as input
 					SceneContext.BeginRenderingSceneColor(RHICmdList, ESimpleRenderTargetMode::EExistingColorAndDepth, FExclusiveDepthStencil::DepthRead_StencilWrite, true, true);
@@ -1228,7 +1243,9 @@ void FDeferredShadingSceneRenderer::RenderLights(FRHICommandListImmediate& RHICm
 					SceneContext.AllocateScreenShadowMask(RHICmdList, ScreenShadowMaskTexture);
 					bShadowMaskReadable = false;
 					if (bUseHairLighting)
-						SceneContext.AllocateScreenShadowMask(RHICmdList, ScreenShadowMaskSubPixelTexture);
+					{
+						SceneContext.AllocateScreenShadowMask(RHICmdList, ScreenShadowMaskSubPixelTexture, true);
+					}
 				}
 
 				FString LightNameWithLevel;
@@ -1809,7 +1826,7 @@ void FDeferredShadingSceneRenderer::RenderLights(FRHICommandListImmediate& RHICm
 				FHairStrandsTransmittanceMaskData TransmittanceMaskData;
 				if (bDrawDeepShadow)
 				{
-					TransmittanceMaskData = RenderHairStrandsTransmittanceMask(RHICmdList, Views, &LightSceneInfo, HairDatas);
+					TransmittanceMaskData = RenderHairStrandsTransmittanceMask(RHICmdList, Views, &LightSceneInfo, HairDatas, ScreenShadowMaskSubPixelTexture);
 				}
 
 				if (ShouldRenderRayTracingStochasticRectLight(LightSceneInfo))
diff --git a/Engine/Source/Runtime/Renderer/Private/MeshPassProcessor.cpp b/Engine/Source/Runtime/Renderer/Private/MeshPassProcessor.cpp
--- a/Engine/Source/Runtime/Renderer/Private/MeshPassProcessor.cpp
+++ b/Engine/Source/Runtime/Renderer/Private/MeshPassProcessor.cpp
@@ -1011,11 +1011,20 @@ void FMeshDrawCommand::SubmitDraw(
 	}
 	else
 	{
-		RHICmdList.DrawPrimitive(
-			MeshDrawCommand.VertexParams.BaseVertexIndex + MeshDrawCommand.FirstIndex,
-			MeshDrawCommand.NumPrimitives,
-			MeshDrawCommand.NumInstances * InstanceFactor
-		);
+		if (MeshDrawCommand.NumPrimitives > 0)
+		{
+			RHICmdList.DrawPrimitive(
+				MeshDrawCommand.VertexParams.BaseVertexIndex + MeshDrawCommand.FirstIndex,
+				MeshDrawCommand.NumPrimitives,
+				MeshDrawCommand.NumInstances * InstanceFactor);
+		}
+		else
+		{
+			RHICmdList.DrawPrimitiveIndirect(
+				MeshDrawCommand.IndirectArgs.Buffer,
+				MeshDrawCommand.IndirectArgs.Offset
+			);
+		}
 	}
 }
 #if MESH_DRAW_COMMAND_DEBUG_DATA
diff --git a/Engine/Source/Runtime/Renderer/Private/PostProcess/PostProcessing.cpp b/Engine/Source/Runtime/Renderer/Private/PostProcess/PostProcessing.cpp
--- a/Engine/Source/Runtime/Renderer/Private/PostProcess/PostProcessing.cpp
+++ b/Engine/Source/Runtime/Renderer/Private/PostProcess/PostProcessing.cpp
@@ -37,6 +37,7 @@
 #include "CompositionLighting/PostProcessPassThrough.h"
 #include "CompositionLighting/PostProcessLpvIndirect.h"
 #include "ShaderPrint.h"
+#include "GpuDebugRendering.h"
 #include "HighResScreenshot.h"
 #include "IHeadMountedDisplay.h"
 #include "IXRTrackingSystem.h"
@@ -818,6 +819,10 @@ void AddPostProcessingPasses(FRDGBuilder& GraphBuilder, const FViewInfo& View, c
 		AddTestImagePass(GraphBuilder, View, SceneColor);
 	}
 
+	if (ShaderDrawDebug::IsShaderDrawDebugEnabled(View))
+	{
+		ShaderDrawDebug::DrawView(GraphBuilder, View, SceneColor.Texture, SceneDepth.Texture);
+	}
 	if (ShaderPrint::IsEnabled() && ShaderPrint::IsSupported(View))
 	{
 		ShaderPrint::DrawView(GraphBuilder, View, SceneColor.Texture);
diff --git a/Engine/Source/Runtime/Renderer/Private/PostProcess/SceneRenderTargets.cpp b/Engine/Source/Runtime/Renderer/Private/PostProcess/SceneRenderTargets.cpp
--- a/Engine/Source/Runtime/Renderer/Private/PostProcess/SceneRenderTargets.cpp
+++ b/Engine/Source/Runtime/Renderer/Private/PostProcess/SceneRenderTargets.cpp
@@ -2316,9 +2316,10 @@ void FSceneRenderTargets::AllocateFoveationTexture(FRHICommandList& RHICmdList)
 	}
 }
 
-void FSceneRenderTargets::AllocateScreenShadowMask(FRHICommandList& RHICmdList, TRefCountPtr<IPooledRenderTarget>& ScreenShadowMaskTexture)
+void FSceneRenderTargets::AllocateScreenShadowMask(FRHICommandList& RHICmdList, TRefCountPtr<IPooledRenderTarget>& ScreenShadowMaskTexture, bool bCreateShaderResourceView)
 {
-	FPooledRenderTargetDesc Desc(FPooledRenderTargetDesc::Create2DDesc(GetBufferSizeXY(), PF_B8G8R8A8, FClearValueBinding::White, TexCreate_None, TexCreate_RenderTargetable, false));
+	ETextureCreateFlags TargetableFlags = ETextureCreateFlags(TexCreate_RenderTargetable | (bCreateShaderResourceView ? TexCreate_ShaderResource : TexCreate_None));
+	FPooledRenderTargetDesc Desc(FPooledRenderTargetDesc::Create2DDesc(GetBufferSizeXY(), PF_B8G8R8A8, FClearValueBinding::White, TexCreate_None, TargetableFlags, false));
 	Desc.Flags |= GFastVRamConfig.ScreenSpaceShadowMask;
 	Desc.NumSamples = GetNumSceneColorMSAASamples(GetCurrentFeatureLevel());
 	GRenderTargetPool.FindFreeElement(RHICmdList, Desc, ScreenShadowMaskTexture, TEXT("ScreenShadowMaskTexture"));
diff --git a/Engine/Source/Runtime/Renderer/Private/PostProcess/SceneRenderTargets.h b/Engine/Source/Runtime/Renderer/Private/PostProcess/SceneRenderTargets.h
--- a/Engine/Source/Runtime/Renderer/Private/PostProcess/SceneRenderTargets.h
+++ b/Engine/Source/Runtime/Renderer/Private/PostProcess/SceneRenderTargets.h
@@ -490,7 +490,7 @@ public:
 
 	void AllocateDebugViewModeTargets(FRHICommandList& RHICmdList);
 
-	void AllocateScreenShadowMask(FRHICommandList& RHICmdList, TRefCountPtr<IPooledRenderTarget>& ScreenShadowMaskTexture);
+	void AllocateScreenShadowMask(FRHICommandList& RHICmdList, TRefCountPtr<IPooledRenderTarget>& ScreenShadowMaskTexture, bool bCreateShaderResourceView = false);
 	void AllocateScreenTransmittanceMask(FRHICommandList& RHICmdList, TRefCountPtr<IPooledRenderTarget>& ScreenTransmittanceMaskTexture);
 
 	TRefCountPtr<IPooledRenderTarget>& GetReflectionBrightnessTarget();
diff --git a/Engine/Source/Runtime/Renderer/Private/SceneRendering.h b/Engine/Source/Runtime/Renderer/Private/SceneRendering.h
--- a/Engine/Source/Runtime/Renderer/Private/SceneRendering.h
+++ b/Engine/Source/Runtime/Renderer/Private/SceneRendering.h
@@ -30,6 +30,7 @@
 #include "Templates/UniquePtr.h"
 #include "RenderGraph.h"
 #include "MeshDrawCommands.h"
+#include "GpuDebugRendering.h"
 
 // Forward declarations.
 class FScene;
@@ -1151,6 +1152,8 @@ public:
 
 	FRWBufferStructured ShaderPrintValueBuffer;
 
+	FShaderDrawDebugData ShaderDrawData;
+
 #if RHI_RAYTRACING
 	TArray<FRayTracingGeometryInstance, SceneRenderingAllocator> RayTracingGeometryInstances;
 
diff --git a/Engine/Source/Runtime/Renderer/Private/SceneVisibility.cpp b/Engine/Source/Runtime/Renderer/Private/SceneVisibility.cpp
--- a/Engine/Source/Runtime/Renderer/Private/SceneVisibility.cpp
+++ b/Engine/Source/Runtime/Renderer/Private/SceneVisibility.cpp
@@ -3008,8 +3008,9 @@ void FSceneRenderer::PreVisibilityFrameSetup(FRHICommandListImmediate& RHICmdLis
 	if (IsHairStrandsEnable(Scene->GetShaderPlatform()) && Views.Num() > 0)
 	{
 		const EWorldType::Type WorldType = Views[0].Family->Scene->GetWorld()->WorldType;
+		const FShaderDrawDebugData* ShaderDrawData = &Views[0].ShaderDrawData;
 		auto ShaderMap = GetGlobalShaderMap(FeatureLevel);
-		RunHairStrandsInterpolation(RHICmdList, WorldType, ShaderMap, EHairStrandsInterpolationType::SimulationStrands);
+		RunHairStrandsInterpolation(RHICmdList, WorldType, ShaderDrawData, ShaderMap, EHairStrandsInterpolationType::SimulationStrands, nullptr);
 	}
 
 	// Notify the FX system that the scene is about to perform visibility checks.
diff --git a/Engine/Source/Runtime/Renderer/Private/ShadowRendering.cpp b/Engine/Source/Runtime/Renderer/Private/ShadowRendering.cpp
--- a/Engine/Source/Runtime/Renderer/Private/ShadowRendering.cpp
+++ b/Engine/Source/Runtime/Renderer/Private/ShadowRendering.cpp
@@ -281,21 +281,27 @@ IMPLEMENT_SHADER_TYPE(template<>,TShadowProjectionFromTranslucencyPS<4>,TEXT("/E
 IMPLEMENT_SHADER_TYPE(template<>,TShadowProjectionFromTranslucencyPS<5>,TEXT("/Engine/Private/ShadowProjectionPixelShader.usf"),TEXT("Main"),SF_Pixel);
 
 // Implement a pixel shader for rendering one pass point light shadows with different quality levels
-#define IMPLEMENT_ONEPASS_POINT_SHADOW_PROJECTION_PIXEL_SHADER(Quality,UseTransmission) \
-	typedef TOnePassPointShadowProjectionPS<Quality,  UseTransmission> FOnePassPointShadowProjectionPS##Quality##UseTransmission; \
-	IMPLEMENT_SHADER_TYPE(template<>,FOnePassPointShadowProjectionPS##Quality##UseTransmission,TEXT("/Engine/Private/ShadowProjectionPixelShader.usf"),TEXT("MainOnePassPointLightPS"),SF_Pixel);
-
-IMPLEMENT_ONEPASS_POINT_SHADOW_PROJECTION_PIXEL_SHADER(1, false);
-IMPLEMENT_ONEPASS_POINT_SHADOW_PROJECTION_PIXEL_SHADER(2, false);
-IMPLEMENT_ONEPASS_POINT_SHADOW_PROJECTION_PIXEL_SHADER(3, false);
-IMPLEMENT_ONEPASS_POINT_SHADOW_PROJECTION_PIXEL_SHADER(4, false);
-IMPLEMENT_ONEPASS_POINT_SHADOW_PROJECTION_PIXEL_SHADER(5, false);
-
-IMPLEMENT_ONEPASS_POINT_SHADOW_PROJECTION_PIXEL_SHADER(1, true);
-IMPLEMENT_ONEPASS_POINT_SHADOW_PROJECTION_PIXEL_SHADER(2, true);
-IMPLEMENT_ONEPASS_POINT_SHADOW_PROJECTION_PIXEL_SHADER(3, true);
-IMPLEMENT_ONEPASS_POINT_SHADOW_PROJECTION_PIXEL_SHADER(4, true);
-IMPLEMENT_ONEPASS_POINT_SHADOW_PROJECTION_PIXEL_SHADER(5, true);
+#define IMPLEMENT_ONEPASS_POINT_SHADOW_PROJECTION_PIXEL_SHADER(Quality,UseTransmission,UseSubPixel) \
+	typedef TOnePassPointShadowProjectionPS<Quality,  UseTransmission, UseSubPixel> FOnePassPointShadowProjectionPS##Quality##UseTransmission##UseSubPixel; \
+	IMPLEMENT_SHADER_TYPE(template<>,FOnePassPointShadowProjectionPS##Quality##UseTransmission##UseSubPixel,TEXT("/Engine/Private/ShadowProjectionPixelShader.usf"),TEXT("MainOnePassPointLightPS"),SF_Pixel);
+
+IMPLEMENT_ONEPASS_POINT_SHADOW_PROJECTION_PIXEL_SHADER(1, false, true);
+IMPLEMENT_ONEPASS_POINT_SHADOW_PROJECTION_PIXEL_SHADER(2, false, true);
+IMPLEMENT_ONEPASS_POINT_SHADOW_PROJECTION_PIXEL_SHADER(3, false, true);
+IMPLEMENT_ONEPASS_POINT_SHADOW_PROJECTION_PIXEL_SHADER(4, false, true);
+IMPLEMENT_ONEPASS_POINT_SHADOW_PROJECTION_PIXEL_SHADER(5, false, true);
+
+IMPLEMENT_ONEPASS_POINT_SHADOW_PROJECTION_PIXEL_SHADER(1, false, false);
+IMPLEMENT_ONEPASS_POINT_SHADOW_PROJECTION_PIXEL_SHADER(2, false, false);
+IMPLEMENT_ONEPASS_POINT_SHADOW_PROJECTION_PIXEL_SHADER(3, false, false);
+IMPLEMENT_ONEPASS_POINT_SHADOW_PROJECTION_PIXEL_SHADER(4, false, false);
+IMPLEMENT_ONEPASS_POINT_SHADOW_PROJECTION_PIXEL_SHADER(5, false, false);
+
+IMPLEMENT_ONEPASS_POINT_SHADOW_PROJECTION_PIXEL_SHADER(1, true, false);
+IMPLEMENT_ONEPASS_POINT_SHADOW_PROJECTION_PIXEL_SHADER(2, true, false);
+IMPLEMENT_ONEPASS_POINT_SHADOW_PROJECTION_PIXEL_SHADER(3, true, false);
+IMPLEMENT_ONEPASS_POINT_SHADOW_PROJECTION_PIXEL_SHADER(4, true, false);
+IMPLEMENT_ONEPASS_POINT_SHADOW_PROJECTION_PIXEL_SHADER(5, true, false);
 
 // Implements a pixel shader for directional light PCSS.
 #define IMPLEMENT_SHADOW_PROJECTION_PIXEL_SHADER(Quality,UseFadePlane) \
@@ -899,9 +905,11 @@ void FProjectedShadowInfo::RenderProjection(FRHICommandListImmediate& RHICmdList
 	TArray<FVector4, TInlineAllocator<8>> FrustumVertices;
 	SetupFrustumForProjection(View, FrustumVertices, bCameraInsideShadowFrustum);
 
-	const bool bDepthBoundsTestEnabled = IsWholeSceneDirectionalShadow() && GSupportsDepthBoundsTest && CVarCSMDepthBoundsTest.GetValueOnRenderThread() != 0;
-
-	if (!bDepthBoundsTestEnabled)
+	const bool bSubPixelSupport = HairVisibilityData != nullptr;
+	const bool bStencilTestEnabled = !bSubPixelSupport;
+	const bool bDepthBoundsTestEnabled = IsWholeSceneDirectionalShadow() && GSupportsDepthBoundsTest && CVarCSMDepthBoundsTest.GetValueOnRenderThread() != 0 && !bSubPixelSupport;
+	
+	if (!bDepthBoundsTestEnabled && bStencilTestEnabled)
 	{
 		SetupProjectionStencilMask(RHICmdList, View, ViewIndex, SceneRender, FrustumVertices, bMobileModulatedProjections, bCameraInsideShadowFrustum);
 	}
@@ -915,7 +923,7 @@ void FProjectedShadowInfo::RenderProjection(FRHICommandListImmediate& RHICmdList
 		// no depth test or writes
 		GraphicsPSOInit.DepthStencilState = TStaticDepthStencilState<false, CF_Always>::GetRHI();
 	}
-	else
+	else if (bStencilTestEnabled)
 	{
 		if (GStencilOptimization)
 		{
@@ -942,6 +950,11 @@ void FProjectedShadowInfo::RenderProjection(FRHICommandListImmediate& RHICmdList
 				>::GetRHI();
 		}
 	}
+	else
+	{
+		GraphicsPSOInit.DepthStencilState = TStaticDepthStencilState<false, CF_Always>::GetRHI();
+	}
+
 
 	SetBlendStateForProjection(GraphicsPSOInit, bProjectingForForwardShading, bMobileModulatedProjections);
 
@@ -982,7 +995,6 @@ void FProjectedShadowInfo::RenderProjection(FRHICommandListImmediate& RHICmdList
 		FShadowProjectionVertexShaderInterface* ShadowProjVS = nullptr;
 		FShadowProjectionPixelShaderInterface* ShadowProjPS = nullptr;
 
-		const bool bSubPixelSupport = HairVisibilityData != nullptr;
 		GetShadowProjectionShaders(LocalQuality, *View, this, bMobileModulatedProjections, bSubPixelSupport, &ShadowProjVS, &ShadowProjPS);
 
 		GraphicsPSOInit.BoundShaderState.VertexDeclarationRHI = GetVertexDeclarationFVector4();
@@ -1021,7 +1033,7 @@ void FProjectedShadowInfo::RenderProjection(FRHICommandListImmediate& RHICmdList
 		VertexBufferRHI.SafeRelease();
 	}
 
-	if (!bDepthBoundsTestEnabled)
+	if (!bDepthBoundsTestEnabled && bStencilTestEnabled)
 	{
 		// Clear the stencil buffer to 0.
 		if (!GStencilOptimization)
@@ -1032,11 +1044,11 @@ void FProjectedShadowInfo::RenderProjection(FRHICommandListImmediate& RHICmdList
 }
 
 
-template <uint32 Quality, bool bUseTransmission>
-static void SetPointLightShaderTempl(FRHICommandList& RHICmdList, FGraphicsPipelineStateInitializer& GraphicsPSOInit, int32 ViewIndex, const FViewInfo& View, const FProjectedShadowInfo* ShadowInfo)
+template <uint32 Quality, bool bUseTransmission, bool bUseSubPixel>
+static void SetPointLightShaderTempl(FRHICommandList& RHICmdList, FGraphicsPipelineStateInitializer& GraphicsPSOInit, int32 ViewIndex, const FViewInfo& View, const FProjectedShadowInfo* ShadowInfo, const FHairStrandsVisibilityData* HairVisibilityData)
 {
 	TShaderMapRef<FShadowVolumeBoundProjectionVS> VertexShader(View.ShaderMap);
-	TShaderMapRef<TOnePassPointShadowProjectionPS<Quality,bUseTransmission> > PixelShader(View.ShaderMap);
+	TShaderMapRef<TOnePassPointShadowProjectionPS<Quality,bUseTransmission,bUseSubPixel> > PixelShader(View.ShaderMap);
 
 	GraphicsPSOInit.BoundShaderState.VertexDeclarationRHI = GetVertexDeclarationFVector4();
 	GraphicsPSOInit.BoundShaderState.VertexShaderRHI = GETSAFERHISHADER_VERTEX(*VertexShader);
@@ -1045,11 +1057,11 @@ static void SetPointLightShaderTempl(FRHICommandList& RHICmdList, FGraphicsPipel
 	SetGraphicsPipelineState(RHICmdList, GraphicsPSOInit);
 	
 	VertexShader->SetParameters(RHICmdList, View, ShadowInfo);
-	PixelShader->SetParameters(RHICmdList, ViewIndex, View, nullptr, ShadowInfo);
+	PixelShader->SetParameters(RHICmdList, ViewIndex, View, HairVisibilityData, ShadowInfo);
 }
 
 /** Render one pass point light shadow projections. */
-void FProjectedShadowInfo::RenderOnePassPointLightProjection(FRHICommandListImmediate& RHICmdList, int32 ViewIndex, const FViewInfo& View, bool bProjectingForForwardShading) const
+void FProjectedShadowInfo::RenderOnePassPointLightProjection(FRHICommandListImmediate& RHICmdList, int32 ViewIndex, const FViewInfo& View, bool bProjectingForForwardShading, const FHairStrandsVisibilityData* HairVisibilityData) const
 {
 	SCOPE_CYCLE_COUNTER(STAT_RenderWholeSceneShadowProjectionsTime);
 
@@ -1099,15 +1111,29 @@ void FProjectedShadowInfo::RenderOnePassPointLightProjection(FRHICommandListImme
 			}
 		}
 
-		if (bUseTransmission)
+		const bool bSubPixelSupport = HairVisibilityData != nullptr;
+		if (bSubPixelSupport)
 		{
 			switch (LocalQuality)
 			{
-				case 1: SetPointLightShaderTempl<1, true>(RHICmdList, GraphicsPSOInit, ViewIndex, View, this); break;
-				case 2: SetPointLightShaderTempl<2, true>(RHICmdList, GraphicsPSOInit, ViewIndex, View, this); break;
-				case 3: SetPointLightShaderTempl<3, true>(RHICmdList, GraphicsPSOInit, ViewIndex, View, this); break;
-				case 4: SetPointLightShaderTempl<4, true>(RHICmdList, GraphicsPSOInit, ViewIndex, View, this); break;
-				case 5: SetPointLightShaderTempl<5, true>(RHICmdList, GraphicsPSOInit, ViewIndex, View, this); break;
+				case 1: SetPointLightShaderTempl<1, false, true>(RHICmdList, GraphicsPSOInit, ViewIndex, View, this, HairVisibilityData); break;
+				case 2: SetPointLightShaderTempl<2, false, true>(RHICmdList, GraphicsPSOInit, ViewIndex, View, this, HairVisibilityData); break;
+				case 3: SetPointLightShaderTempl<3, false, true>(RHICmdList, GraphicsPSOInit, ViewIndex, View, this, HairVisibilityData); break;
+				case 4: SetPointLightShaderTempl<4, false, true>(RHICmdList, GraphicsPSOInit, ViewIndex, View, this, HairVisibilityData); break;
+				case 5: SetPointLightShaderTempl<5, false, true>(RHICmdList, GraphicsPSOInit, ViewIndex, View, this, HairVisibilityData); break;
+				default:
+					check(0);
+			}
+		}
+		else if (bUseTransmission)
+		{
+			switch (LocalQuality)
+			{
+				case 1: SetPointLightShaderTempl<1, true, false>(RHICmdList, GraphicsPSOInit, ViewIndex, View, this, nullptr); break;
+				case 2: SetPointLightShaderTempl<2, true, false>(RHICmdList, GraphicsPSOInit, ViewIndex, View, this, nullptr); break;
+				case 3: SetPointLightShaderTempl<3, true, false>(RHICmdList, GraphicsPSOInit, ViewIndex, View, this, nullptr); break;
+				case 4: SetPointLightShaderTempl<4, true, false>(RHICmdList, GraphicsPSOInit, ViewIndex, View, this, nullptr); break;
+				case 5: SetPointLightShaderTempl<5, true, false>(RHICmdList, GraphicsPSOInit, ViewIndex, View, this, nullptr); break;
 				default:
 					check(0);
 			}
@@ -1116,11 +1142,11 @@ void FProjectedShadowInfo::RenderOnePassPointLightProjection(FRHICommandListImme
 		{
 			switch (LocalQuality)
 			{
-				case 1: SetPointLightShaderTempl<1, false>(RHICmdList, GraphicsPSOInit, ViewIndex, View, this); break;
-				case 2: SetPointLightShaderTempl<2, false>(RHICmdList, GraphicsPSOInit, ViewIndex, View, this); break;
-				case 3: SetPointLightShaderTempl<3, false>(RHICmdList, GraphicsPSOInit, ViewIndex, View, this); break;
-				case 4: SetPointLightShaderTempl<4, false>(RHICmdList, GraphicsPSOInit, ViewIndex, View, this); break;
-				case 5: SetPointLightShaderTempl<5, false>(RHICmdList, GraphicsPSOInit, ViewIndex, View, this); break;
+				case 1: SetPointLightShaderTempl<1, false, false>(RHICmdList, GraphicsPSOInit, ViewIndex, View, this, nullptr); break;
+				case 2: SetPointLightShaderTempl<2, false, false>(RHICmdList, GraphicsPSOInit, ViewIndex, View, this, nullptr); break;
+				case 3: SetPointLightShaderTempl<3, false, false>(RHICmdList, GraphicsPSOInit, ViewIndex, View, this, nullptr); break;
+				case 4: SetPointLightShaderTempl<4, false, false>(RHICmdList, GraphicsPSOInit, ViewIndex, View, this, nullptr); break;
+				case 5: SetPointLightShaderTempl<5, false, false>(RHICmdList, GraphicsPSOInit, ViewIndex, View, this, nullptr); break;
 				default:
 					check(0);
 			}
@@ -1560,7 +1586,7 @@ bool FSceneRenderer::RenderShadowProjections(FRHICommandListImmediate& RHICmdLis
 						{
 							if (ProjectedShadowInfo->bOnePassPointLightShadow)
 							{
-								ProjectedShadowInfo->RenderOnePassPointLightProjection(RHICmdList, ViewIndex, View, bProjectingForForwardShading);
+								ProjectedShadowInfo->RenderOnePassPointLightProjection(RHICmdList, ViewIndex, View, bProjectingForForwardShading, HairVisibilityData);
 							}
 							else
 							{
@@ -1605,7 +1631,7 @@ bool FSceneRenderer::RenderShadowProjections(FRHICommandListImmediate& RHICmdLis
 			FRHIRenderPassInfo RPInfo(ScreenShadowMaskSubPixelTexture->GetRenderTargetItem().TargetableTexture, ERenderTargetActions::Load_Store);
 			RPInfo.DepthStencilRenderTarget.Action = MakeDepthStencilTargetActions(ERenderTargetActions::Load_DontStore, ERenderTargetActions::Load_Store);
 			RPInfo.DepthStencilRenderTarget.DepthStencilTarget = SceneContext.GetSceneDepthSurface();
-			RPInfo.DepthStencilRenderTarget.ExclusiveDepthStencil = FExclusiveDepthStencil::DepthRead_StencilWrite;
+			RPInfo.DepthStencilRenderTarget.ExclusiveDepthStencil = FExclusiveDepthStencil::DepthRead_StencilNop;
 
 			TransitionRenderPassTargets(RHICmdList, RPInfo);
 			RHICmdList.BeginRenderPass(RPInfo, TEXT("RenderShadowProjectionSubPixel"));
