Copying //UE4/Dev-Physics to Main (//UE4/Main)

diff --git a/Engine/Source/Runtime/Renderer/Private/ShadowRendering.h b/Engine/Source/Runtime/Renderer/Private/ShadowRendering.h
--- a/Engine/Source/Runtime/Renderer/Private/ShadowRendering.h
+++ b/Engine/Source/Runtime/Renderer/Private/ShadowRendering.h
@@ -417,7 +417,7 @@ public:
 	void RenderRayTracedDistanceFieldProjection(FRHICommandListImmediate& RHICmdList, const class FViewInfo& View, const FIntRect& ScissorRect, IPooledRenderTarget* ScreenShadowMaskTexture, bool bProjectingForForwardShading);
 
 	/** Render one pass point light shadow projections. */
-	void RenderOnePassPointLightProjection(FRHICommandListImmediate& RHICmdList, int32 ViewIndex, const FViewInfo& View, bool bProjectingForForwardShading) const;
+	void RenderOnePassPointLightProjection(FRHICommandListImmediate& RHICmdList, int32 ViewIndex, const FViewInfo& View, bool bProjectingForForwardShading, const FHairStrandsVisibilityData* HairVisibilityData) const;
 
 	/**
 	 * Renders the projected shadow's frustum wireframe with the given FPrimitiveDrawInterface.
@@ -1295,7 +1295,7 @@ private:
  * Pixel shader used to project one pass point light shadows.
  */
 // Quality = 0 / 1
-template <uint32 Quality, bool bUseTransmission>
+template <uint32 Quality, bool bUseTransmission, bool bUseSubPixel>
 class TOnePassPointShadowProjectionPS : public FGlobalShader
 {
 	DECLARE_SHADER_TYPE(TOnePassPointShadowProjectionPS,Global);
@@ -1315,6 +1315,7 @@ public:
 		PointLightDepthBias.Bind(Initializer.ParameterMap,TEXT("PointLightDepthBias"));
 		PointLightProjParameters.Bind(Initializer.ParameterMap, TEXT("PointLightProjParameters"));
 		TransmissionProfilesTexture.Bind(Initializer.ParameterMap, TEXT("SSProfilesTexture"));
+		HairCategorizationTexture.Bind(Initializer.ParameterMap, TEXT("HairCategorizationTexture"));
 	}
 
 	static void ModifyCompilationEnvironment(const FGlobalShaderPermutationParameters& Parameters, FShaderCompilerEnvironment& OutEnvironment)
@@ -1322,6 +1323,7 @@ public:
 		FGlobalShader::ModifyCompilationEnvironment(Parameters,OutEnvironment);
 		OutEnvironment.SetDefine(TEXT("SHADOW_QUALITY"), Quality);
 		OutEnvironment.SetDefine(TEXT("USE_TRANSMISSION"), (uint32)(bUseTransmission ? 1 : 0));
+		OutEnvironment.SetDefine(TEXT("SUBPIXEL_SHADOW"), (uint32)(bUseSubPixel ? 1 : 0));
 	}
 
 	static bool ShouldCompilePermutation(const FGlobalShaderPermutationParameters& Parameters)
@@ -1357,6 +1359,11 @@ public:
 		SetShaderValue(RHICmdList, ShaderRHI, PointLightDepthBias, FVector(ShadowInfo->GetShaderDepthBias(), ShadowInfo->GetShaderSlopeDepthBias(), ShadowInfo->GetShaderMaxSlopeDepthBias()));
 		SetShaderValue(RHICmdList, ShaderRHI, PointLightProjParameters, FVector2D(projParam.X, projParam.Y));
 
+		if (bUseSubPixel && HairVisibilityData && HairVisibilityData->CategorizationTexture)
+		{
+			SetTextureParameter(RHICmdList, ShaderRHI, HairCategorizationTexture, HairVisibilityData->CategorizationTexture->GetRenderTargetItem().ShaderResourceTexture);
+		}
+
 		FSceneRenderTargets& SceneContext = FSceneRenderTargets::Get(RHICmdList);
 		{
 			const IPooledRenderTarget* PooledRT = GetSubsufaceProfileTexture_RT((FRHICommandListImmediate&)RHICmdList);
@@ -1401,6 +1408,7 @@ public:
 		Ar << PointLightDepthBias;
 		Ar << PointLightProjParameters;
 		Ar << TransmissionProfilesTexture;
+		Ar << HairCategorizationTexture;
 		return bShaderHasOutdatedParameters;
 	}
 
@@ -1414,6 +1422,7 @@ private:
 	FShaderParameter PointLightDepthBias;
 	FShaderParameter PointLightProjParameters;
 	FShaderResourceParameter TransmissionProfilesTexture;
+	FShaderResourceParameter HairCategorizationTexture;
 
 
 };
diff --git a/Engine/Source/Runtime/Renderer/Public/GpuDebugRendering.h b/Engine/Source/Runtime/Renderer/Public/GpuDebugRendering.h
new file mode 100644
--- /dev/null
+++ b/Engine/Source/Runtime/Renderer/Public/GpuDebugRendering.h
@@ -0,0 +1,44 @@
+// Copyright Epic Games, Inc. All Rights Reserved.
+
+#pragma once
+
+#include "RendererInterface.h"
+#include "RHICommandList.h"
+#include "RenderGraph.h"
+
+class FViewInfo;
+
+struct FShaderDrawDebugData
+{
+	FIntPoint CursorPosition;
+	TRefCountPtr<FPooledRDGBuffer> Buffer;
+	TRefCountPtr<FPooledRDGBuffer> IndirectBuffer;
+};
+
+namespace ShaderDrawDebug 
+{
+	// Call this to know if this is even just available (for exemple in shipping mode buffers won't exists)
+	RENDERER_API bool IsShaderDrawDebugEnabled();
+
+	// Call this to know if a view can render this debug information
+	bool IsShaderDrawDebugEnabled(const FViewInfo& View);
+
+	// Allocate the debug print buffer associated with the view
+	void BeginView(FRHICommandListImmediate& RHICmdList, FViewInfo& View);
+	// Draw info from the debug print buffer to the given output target
+	void DrawView(FRDGBuilder& GraphBuilder, const FViewInfo& View, FRDGTextureRef OutputTexture, FRDGTextureRef DepthTexture);
+	// Release the debug print buffer associated with the view
+	void EndView(FViewInfo& View);
+
+	// The structure to be set on the debug shader outputting debug primitive
+	BEGIN_SHADER_PARAMETER_STRUCT(FShaderDrawDebugParameters, )
+		SHADER_PARAMETER(FIntPoint, ShaderDrawCursorPos)
+		SHADER_PARAMETER(int32, ShaderDrawMaxElementCount)
+		SHADER_PARAMETER_RDG_BUFFER_UAV(RWStructuredBuffer, OutShaderDrawPrimitive)
+		SHADER_PARAMETER_RDG_BUFFER_UAV(RWBuffer<uint>, OutputShaderDrawIndirect)
+	END_SHADER_PARAMETER_STRUCT()
+
+	// Call this to fill the FShaderDrawParameters
+	RENDERER_API void SetParameters(FRDGBuilder& GraphBuilder, const FShaderDrawDebugData& Data, FShaderDrawDebugParameters& OutParameters);
+
+}
diff --git a/Engine/Source/Runtime/Renderer/Public/HairStrandsInterface.h b/Engine/Source/Runtime/Renderer/Public/HairStrandsInterface.h
--- a/Engine/Source/Runtime/Renderer/Public/HairStrandsInterface.h
+++ b/Engine/Source/Runtime/Renderer/Public/HairStrandsInterface.h
@@ -10,6 +10,8 @@
 #include "RendererInterface.h"
 #include "Containers/Array.h"
 #include "Engine/EngineTypes.h"
+#include "Shader.h"
+#include "RenderResource.h"
 
 ////////////////////////////////////////////////////////////////////////////////////////////////////
 // Misc/Helpers
@@ -27,6 +29,7 @@ enum class EHairStrandsDebugMode : uint8
 	RenderHairRadiusVariation,
 	RenderHairBaseColor,
 	RenderHairRoughness,
+	RenderVisCluster,
 	Count
 };
 
@@ -49,25 +52,122 @@ RENDERER_API FMinHairRadiusAtDepth1 ComputeMinStrandRadiusAtDepth1(
 	const uint32 SampleCount,
 	const float OverrideStrandHairRasterizationScale);
 
+struct FHairStrandClusterCullingData;
+struct FPooledRDGBuffer;
+struct IPooledRenderTarget;
+struct FRWBuffer;
+class  FHairGroupPublicData;
+class  FRDGShaderResourceView;
+class  FResourceArrayInterface;
+class  FSceneView;
+
+////////////////////////////////////////////////////////////////////////////////////////////////////
+// Public group data 
+class RENDERER_API FHairGroupPublicData : public FRenderResource
+{
+public:
+	FHairGroupPublicData(uint32 GroupIndex, uint32 GroupInstanceVertexCount, uint32 ClusterCount, uint32 VertexCount);
+	
+	virtual void InitRHI() override;
+	virtual void ReleaseRHI() override;
+	virtual FString GetFriendlyName() const override { return TEXT("FHairGroupPublicData"); }
+
+	// The primitive count when no culling and neither lod happens
+	uint32 GetGroupInstanceVertexCount() const { return GroupInstanceVertexCount; }
+
+	uint32 GetGroupIndex() const { return GroupIndex; }
+
+	FRWBuffer& GetDrawIndirectBuffer()	{ return DrawIndirectBuffer; }
+	FRWBuffer& GetClusterAABBBuffer()	{ return ClusterAABBBuffer; }
+	FRWBuffer& GetGroupAABBBuffer()		{ return GroupAABBBuffer; }
+
+	FRWBuffer& GetCulledVertexIdBuffer() { return CulledVertexIdBuffer; }
+	FRWBuffer& GetCulledVertexRadiusScaleBuffer() { return CulledVertexRadiusScaleBuffer; }
+	bool GetCullingResultAvailable() const { return bCullingResultAvailable; }
+	void SetCullingResultAvailable(bool b) { bCullingResultAvailable = b; }
+
+private:
+
+	uint32 GroupInstanceVertexCount;
+	uint32 GroupIndex;
+	uint32 ClusterCount;
+	uint32 VertexCount;
+
+	/* Indirect draw buffer to draw everything or the result of the culling per pass */
+	FRWBuffer DrawIndirectBuffer;
+
+	/* Hair Cluster & Hair Group bounding box buffer */
+	FRWBuffer ClusterAABBBuffer;
+	FRWBuffer GroupAABBBuffer;
+
+	/* Culling & LODing results for a hair group */ // Better to be transient?
+	FRWBuffer CulledVertexIdBuffer;
+	FRWBuffer CulledVertexRadiusScaleBuffer;
+	bool bCullingResultAvailable = false;
+};
+
+////////////////////////////////////////////////////////////////////////////////////////////////////
+// Cluster information exchanged between renderer and the hair strand plugin. 
+
+struct FHairStrandClusterData
+{
+	struct FHairGroup
+	{
+		uint32 ClusterCount = 0;
+		uint32 VertexCount = 0;
+
+		float LodBias = 0.0f;
+		float LodAverageVertexPerPixel = 0.0f;
+
+		// See FHairStrandsClusterCullingResource fro details about those buffers.
+		FRWBuffer* GroupAABBBuffer = nullptr;
+		FRWBuffer* ClusterAABBBuffer = nullptr;
+		FReadBuffer* ClusterInfoBuffer = nullptr;
+		FReadBuffer* VertexToClusterIdBuffer = nullptr;
+		FReadBuffer* ClusterVertexIdBuffer = nullptr;
+		FReadBuffer* ClusterIndexRadiusScaleInfoBuffer = nullptr;
+
+		// Culling & LOD output
+		FRWBuffer* GetCulledVertexIdBuffer() const { return HairGroupPublicPtr ? &HairGroupPublicPtr->GetCulledVertexIdBuffer() : nullptr; }
+		FRWBuffer* GetCulledVertexRadiusScaleBuffer() const { return HairGroupPublicPtr ? &HairGroupPublicPtr->GetCulledVertexRadiusScaleBuffer() : nullptr; }
+		bool GetCullingResultAvailable() const { return HairGroupPublicPtr ? HairGroupPublicPtr->GetCullingResultAvailable() : false; }
+		void SetCullingResultAvailable(bool b) { if (HairGroupPublicPtr) HairGroupPublicPtr->SetCullingResultAvailable(b); }
+
+		TRefCountPtr<FPooledRDGBuffer> ClusterDebugAABBBuffer;							// Null if this debug is not enabled.
+		TRefCountPtr<FPooledRDGBuffer> CulledDispatchIndirectParametersClusterCount;	// Null if this debug is not enabled.
+
+		FHairGroupPublicData* HairGroupPublicPtr = nullptr;
+	};
+
+	TArray<FHairGroup> HairGroups;
+};
+
 ////////////////////////////////////////////////////////////////////////////////////////////////////
 // Hair/Mesh projection & interpolation
+
+typedef void(*THairStrandsResetInterpolationFunction)(
+	FRHICommandListImmediate& RHICmdList,
+	struct FHairStrandsInterpolationInput* Input,
+	struct FHairStrandsInterpolationOutput* Output);
+
 typedef void (*THairStrandsInterpolationFunction)(
 	FRHICommandListImmediate& RHICmdList, 
+	const struct FShaderDrawDebugData* ShaderDrawData,
 	struct FHairStrandsInterpolationInput* Input, 
 	struct FHairStrandsInterpolationOutput* Output, 
 	struct FHairStrandsProjectionHairData& RenHairProjection,
 	struct FHairStrandsProjectionHairData& SimHairProjection,
-	int32 LODIndex);
+	int32 LODIndex,
+	FHairStrandClusterData* ClusterData);
 
 struct FHairStrandsInterpolationData
 {
 	struct FHairStrandsInterpolationInput* Input = nullptr;
 	struct FHairStrandsInterpolationOutput* Output = nullptr;
 	THairStrandsInterpolationFunction Function = nullptr;
+	THairStrandsResetInterpolationFunction ResetFunction = nullptr; 
 };
 
-struct FRWBuffer;
-
 struct FHairStrandsProjectionMeshData
 {
 	struct Section
@@ -84,15 +184,21 @@ struct FHairStrandsProjectionMeshData
 		int32 LODIndex = 0;
 	};
 
-	TArray<Section> Sections;
+	struct LOD
+	{
+		TArray<Section> Sections;
+	};
+	TArray<LOD> LODs;
 };
 
 struct FHairStrandsProjectionHairData
 {
 	struct LODData
 	{
+		enum class EStatus { Invalid, Initialized, Completed };
 		bool bIsValid = false;
 		int32 LODIndex = -1;
+		EStatus* Status = nullptr;
 
 		// The index buffers stores the mesh section & the triangle index into a single uint32 
 		// (3 highest bits store the section (up to 8 sections)
@@ -100,13 +206,13 @@ struct FHairStrandsProjectionHairData
 		FRWBuffer* RootTriangleBarycentricBuffer = nullptr;
 
 		// Rest root triangles' positions are relative to root center (for preserving precision)
-		FVector RestPositionOffset = FVector::ZeroVector;
+		FVector* RestPositionOffset = nullptr;
 		FRWBuffer* RestRootTrianglePosition0Buffer = nullptr;
 		FRWBuffer* RestRootTrianglePosition1Buffer = nullptr;
 		FRWBuffer* RestRootTrianglePosition2Buffer = nullptr;
 
 		// Deformed root triangles' positions are relative to root center (for preserving precision)
-		FVector DeformedPositionOffset = FVector::ZeroVector;
+		FVector* DeformedPositionOffset = nullptr;
 		FRWBuffer* DeformedRootTrianglePosition0Buffer = nullptr;
 		FRWBuffer* DeformedRootTrianglePosition1Buffer = nullptr;
 		FRWBuffer* DeformedRootTrianglePosition2Buffer = nullptr;
@@ -127,7 +233,20 @@ struct FHairStrandsProjectionHairData
 	TArray<HairGroup> HairGroups;
 };
 
-RENDERER_API void AddHairStrandsProjectionQuery(FRHICommandListImmediate& RHICmdList, uint64 Id, EWorldType::Type WorldType, int32 LODIndex, const FVector& RestRootCenter);
+////////////////////////////////////////////////////////////////////////////////////////////////////
+// Hair component/primitive resources (shared with the engine side)
+struct FHairStrandsPrimitiveResources
+{
+	struct FHairGroup
+	{
+		FRWBuffer* ClusterAABBBuffer = nullptr;
+		FRWBuffer* GroupAABBBuffer = nullptr;
+		uint32 ClusterCount = 0;
+	};
+	TArray<FHairGroup> Groups;
+};
+
+FHairStrandsPrimitiveResources GetHairStandsPrimitiveResources(uint64 Id);
 
 ////////////////////////////////////////////////////////////////////////////////////////////////////
 // Debug infos
@@ -162,6 +281,7 @@ RENDERER_API void RegisterHairStrands(
 	const FHairStrandsInterpolationData& E,
 	const FHairStrandsProjectionHairData& RenProjection,
 	const FHairStrandsProjectionHairData& SimProjection,
+	const FHairStrandsPrimitiveResources& PrimitiveResources,
 	const FHairStrandsDebugInfo& DebugInfo);
 
 RENDERER_API void UnregisterHairStrands(uint64 Id);
@@ -192,6 +312,9 @@ RENDERER_API bool UpdateHairStrands(
 RENDERER_API bool IsHairStrandsSupported(const EShaderPlatform Platform);
 bool IsHairStrandsEnable(EShaderPlatform Platform);
 
+RENDERER_API bool IsHairStrandsSupported(const EShaderPlatform Platform);
+bool IsHairStrandsEnable(EShaderPlatform Platform);
+
 // Return strands & guide indices to be preserved, while all others strands/guides should be culled
 enum class EHairCullMode : uint8
 {
@@ -206,3 +329,13 @@ struct FHairCullInfo
 	EHairCullMode CullMode = EHairCullMode::None;
 };
 RENDERER_API FHairCullInfo GetHairStrandsCullInfo();
+
+RENDERER_API bool IsHairRayTracingEnabled();
+
+RENDERER_API void RunProjection(
+	FRHICommandListImmediate& RHICmdList,
+	const FTransform& LocalToWorld,
+	const FVector& RestPositionOffset,
+	const FHairStrandsProjectionMeshData& MeshData,
+	FHairStrandsProjectionHairData& RenProjectionHairData,
+	FHairStrandsProjectionHairData& SimProjectionHairData);
\ No newline at end of file
diff --git a/Engine/Source/Runtime/Renderer/Public/MeshPassProcessor.h b/Engine/Source/Runtime/Renderer/Public/MeshPassProcessor.h
--- a/Engine/Source/Runtime/Renderer/Public/MeshPassProcessor.h
+++ b/Engine/Source/Runtime/Renderer/Public/MeshPassProcessor.h
@@ -915,7 +915,7 @@ enum class EMeshPassFeatures
 {
 	Default = 0,
 	PositionOnly = 1 << 0,
-	PositionAndNormalOnly = 1 << 1
+	PositionAndNormalOnly = 1 << 1,
 };
 ENUM_CLASS_FLAGS(EMeshPassFeatures);
 
