Fix implicit conversion compile error

diff --git a/Engine/Source/Runtime/Renderer/Private/DistanceFieldShadowing.cpp b/Engine/Source/Runtime/Renderer/Private/DistanceFieldShadowing.cpp
--- a/Engine/Source/Runtime/Renderer/Private/DistanceFieldShadowing.cpp
+++ b/Engine/Source/Runtime/Renderer/Private/DistanceFieldShadowing.cpp
@@ -1354,7 +1354,7 @@ void FProjectedShadowInfo::BeginRenderRayTracedDistanceFieldProjection(FRHIComma
 			LightSceneInfo->HeightFieldTileIntersectionResources);
 
 		FLightTileIntersectionResources* TileIntersectionResources = LightSceneInfo->HeightFieldTileIntersectionResources.Get();
-		const bool bHasPrevOutput = RayTracedShadowsRT;
+		const bool bHasPrevOutput = !!RayTracedShadowsRT;
 		TRefCountPtr<IPooledRenderTarget> PrevOutputRT;
 
 		if (!RHISupports4ComponentUAVReadWrite(View.GetShaderPlatform()))
