Fixed issue where post process material was assuming that the scene texture input matched the output viewport. This is not true when the post process material is 'After Tonemap' and the viewport rect is the final backbuffer.

diff --git a/Engine/Shaders/Private/MaterialTemplate.ush b/Engine/Shaders/Private/MaterialTemplate.ush
--- a/Engine/Shaders/Private/MaterialTemplate.ush
+++ b/Engine/Shaders/Private/MaterialTemplate.ush
@@ -657,13 +657,10 @@ MaterialFloat2 GetDefaultSceneTextureUV(FMaterialVertexParameters Parameters, co
 MaterialFloat2 GetDefaultSceneTextureUV(FMaterialPixelParameters Parameters, const uint SceneTextureId)
 {
 	#if POST_PROCESS_MATERIAL
-	if (IsPostProcessInputSceneTexture(SceneTextureId))
-		{
-			return ViewportUVToSceneTextureUV(GetViewportUV(Parameters), SceneTextureId);
-		}
-		#endif
-
-	return GetSceneTextureUV(Parameters);
+		return ViewportUVToSceneTextureUV(GetViewportUV(Parameters), SceneTextureId);
+	#else
+		return GetSceneTextureUV(Parameters);
+	#endif
 }
 
 
diff --git a/Engine/Source/Runtime/Renderer/Private/PostProcess/PostProcessMaterial.cpp b/Engine/Source/Runtime/Renderer/Private/PostProcess/PostProcessMaterial.cpp
--- a/Engine/Source/Runtime/Renderer/Private/PostProcess/PostProcessMaterial.cpp
+++ b/Engine/Source/Runtime/Renderer/Private/PostProcess/PostProcessMaterial.cpp
@@ -423,7 +423,7 @@ FScreenPassTexture AddPostProcessMaterialPass(
 
 	FPostProcessMaterialParameters* PostProcessMaterialParameters = GraphBuilder.AllocParameters<FPostProcessMaterialParameters>();
 
-	PostProcessMaterialParameters->PostProcessOutput = GetScreenPassTextureViewportParameters(SceneColorViewport);
+	PostProcessMaterialParameters->PostProcessOutput = GetScreenPassTextureViewportParameters(OutputViewport);
 	PostProcessMaterialParameters->CustomDepth = DepthStencilTexture;
 	PostProcessMaterialParameters->RenderTargets[0] = Output.GetRenderTargetBinding();
 
