Fix and re-enable denoiser for mac metal.
Added a Uint32 dummy texture that's 1x1 set to 0 (because metal needs an explicit conversion when reading 8888 as uint)
Switched Denoiser intermediates to 16F from R11G11B10 as in SSGI.
FSSDCompressMetadataCS now uses ShouldCompileSignalPipeline to determine which permutations to compile.

diff --git a/Engine/Source/Runtime/Renderer/Private/ScreenSpaceDenoise.cpp b/Engine/Source/Runtime/Renderer/Private/ScreenSpaceDenoise.cpp
--- a/Engine/Source/Runtime/Renderer/Private/ScreenSpaceDenoise.cpp
+++ b/Engine/Source/Runtime/Renderer/Private/ScreenSpaceDenoise.cpp
@@ -679,7 +679,7 @@ bool ShouldCompileSignalPipeline(ESignalProcessing SignalProcessing, EShaderPlat
 {
 	if (SignalProcessing == ESignalProcessing::ScreenSpaceDiffuseIndirect)
 	{
-		return Platform == SP_PCD3D_SM5 || Platform == SP_PS4 || Platform == SP_XBOXONE_D3D12;
+		return Platform == SP_PCD3D_SM5 || Platform == SP_PS4 || Platform == SP_XBOXONE_D3D12 || Platform == SP_METAL_SM5;
 	}
 	else if (
 		SignalProcessing == ESignalProcessing::Reflections ||
@@ -825,7 +825,7 @@ class FSSDCompressMetadataCS : public FGlobalShader
 			return false;
 		}
 
-		return Parameters.Platform == SP_PCD3D_SM5 || Parameters.Platform == SP_PS4;
+		return ShouldCompileSignalPipeline(ESignalProcessing::ScreenSpaceDiffuseIndirect, Parameters.Platform);
 	}
 
 	BEGIN_SHADER_PARAMETER_STRUCT(FParameters, )
@@ -1361,11 +1361,19 @@ static void DenoiseSignalAtConstantPixelDensity(
 		}
 		else if (Settings.SignalProcessing == ESignalProcessing::ScreenSpaceDiffuseIndirect)
 		{
+			// This is a workaround to fix UE-84870.
+			// In certain cases writes to R11G11B10 fail on AMD hw so we are falling back to 16F.
+#if PLATFORM_MAC
+			ReconstructionDescs[0].Format = PF_FloatRGBA;
+			HistoryDescs[0].Format = PF_FloatRGBA;
+#else
 			ReconstructionDescs[0].Format = PF_FloatR11G11B10;
+			HistoryDescs[0].Format = PF_FloatR11G11B10;
+#endif
 			ReconstructionDescs[1].Format = PF_R8G8;
 			ReconstructionTextureCount = 2;
 
-			HistoryDescs[0].Format = PF_FloatR11G11B10;
+			
 			HistoryDescs[1].Format = PF_R8G8;
 			HistoryTextureCountPerSignal = 2;
 			bHasReconstructionLayoutDifferentFromHistory = false;
@@ -1798,7 +1806,7 @@ static void DenoiseSignalAtConstantPixelDensity(
 		if (CompressedMetadataLayout == ECompressedMetadataLayout::DepthAndViewNormal)
 		{
 			PassParameters->PrevCompressedMetadata[0] = RegisterExternalTextureWithFallback(
-				GraphBuilder, View.PrevViewInfo.CompressedDepthViewNormal, GSystemTextures.BlackDummy);
+				GraphBuilder, View.PrevViewInfo.CompressedDepthViewNormal, GSystemTextures.ZeroUIntDummy);
 		}
 
 		FScreenSpaceDenoiserHistory DummyPrevFrameHistory;
diff --git a/Engine/Source/Runtime/Renderer/Private/SystemTextures.cpp b/Engine/Source/Runtime/Renderer/Private/SystemTextures.cpp
--- a/Engine/Source/Runtime/Renderer/Private/SystemTextures.cpp
+++ b/Engine/Source/Runtime/Renderer/Private/SystemTextures.cpp
@@ -48,6 +48,18 @@ void FSystemTextures::InitializeCommonTextures(FRHICommandListImmediate& RHICmdL
 			RHICmdList.EndRenderPass();
 			RHICmdList.CopyToResolveTarget(BlackDummy->GetRenderTargetItem().TargetableTexture, BlackDummy->GetRenderTargetItem().ShaderResourceTexture, FResolveParams());
 		}
+	
+		// Create a texture that is a single UInt32 value set to 0
+		{
+			FPooledRenderTargetDesc Desc(FPooledRenderTargetDesc::Create2DDesc(FIntPoint(1,1), PF_R32_UINT, FClearValueBinding::Transparent, TexCreate_HideInVisualizeTexture, TexCreate_RenderTargetable | TexCreate_NoFastClear | TexCreate_ShaderResource, false));
+			Desc.AutoWritable = false;
+			GRenderTargetPool.FindFreeElement(RHICmdList, Desc, ZeroUIntDummy, TEXT("ZeroUIntDummy"), true, ERenderTargetTransience::NonTransient);
+			
+			FRHIRenderPassInfo RPInfo(ZeroUIntDummy->GetRenderTargetItem().TargetableTexture, ERenderTargetActions::Clear_Store);
+			RHICmdList.BeginRenderPass(RPInfo, TEXT("ClearZeroUIntDummy"));
+			RHICmdList.EndRenderPass();
+			RHICmdList.CopyToResolveTarget(ZeroUIntDummy->GetRenderTargetItem().TargetableTexture, ZeroUIntDummy->GetRenderTargetItem().ShaderResourceTexture, FResolveParams());
+		}
 
 		// Create a BlackAlphaOneDummy texture
 		{
@@ -715,3 +727,9 @@ FRDGTextureRef FSystemTextures::GetVolumetricBlackDummy(FRDGBuilder& GraphBuilde
 {
 	return GraphBuilder.RegisterExternalTexture(VolumetricBlackDummy, TEXT("VolumetricBlackDummy"));
 }
+
+FRDGTextureRef FSystemTextures::GetZeroUIntDummy(FRDGBuilder& GraphBuilder) const
+{
+	return GraphBuilder.RegisterExternalTexture(ZeroUIntDummy, TEXT("ZeroUIntDummy"));
+}
+
diff --git a/Engine/Source/Runtime/Renderer/Private/SystemTextures.h b/Engine/Source/Runtime/Renderer/Private/SystemTextures.h
--- a/Engine/Source/Runtime/Renderer/Private/SystemTextures.h
+++ b/Engine/Source/Runtime/Renderer/Private/SystemTextures.h
@@ -91,6 +91,9 @@ public:
 
 	/** float4(0,0,0,0) volumetric texture. */
 	TRefCountPtr<IPooledRenderTarget> VolumetricBlackDummy;
+	
+	// Dummy 0 Uint texture for RHIs that need explicit overloads
+	TRefCountPtr<IPooledRenderTarget> ZeroUIntDummy;
 
 	// SRV for WhiteDummy Texture.
 	TRefCountPtr<FRHIShaderResourceView> WhiteDummySRV;
@@ -99,6 +102,7 @@ public:
 
 	FRDGTextureRef GetWhiteDummy(FRDGBuilder& GraphBuilder) const;
 	FRDGTextureRef GetBlackDummy(FRDGBuilder& GraphBuilder) const;
+	FRDGTextureRef GetZeroUIntDummy(FRDGBuilder& GraphBuilder) const;
 	FRDGTextureRef GetBlackAlphaOneDummy(FRDGBuilder& GraphBuilder) const;
 	FRDGTextureRef GetPerlinNoiseGradient(FRDGBuilder& GraphBuilder) const;
 	FRDGTextureRef GetPerlinNoise3D(FRDGBuilder& GraphBuilder) const;
