Fix dummy texture being overwritten in some cases (from KevinO)

diff --git a/Engine/Source/Runtime/Renderer/Private/TranslucentRendering.cpp b/Engine/Source/Runtime/Renderer/Private/TranslucentRendering.cpp
--- a/Engine/Source/Runtime/Renderer/Private/TranslucentRendering.cpp
+++ b/Engine/Source/Runtime/Renderer/Private/TranslucentRendering.cpp
@@ -860,7 +860,7 @@ void UpsampleTranslucency(FRHICommandList& RHICmdList, const FViewInfo& View, bo
 void FDeferredShadingSceneRenderer::RenderTranslucency(FRHICommandListImmediate& RHICmdList, bool bDrawUnderwaterViews)
 {
 	// For now there is only one resolve for all translucency passes. This can be changed by enabling the resolve in RenderTranslucencyInner()
-	TRefCountPtr<IPooledRenderTarget> SceneColorCopy = GSystemTextures.BlackDummy;
+	TRefCountPtr<IPooledRenderTarget> SceneColorCopy;
 	if (!bDrawUnderwaterViews)
 	{
 		ConditionalResolveSceneColorForTranslucentMaterials(RHICmdList, SceneColorCopy);
