Fix left-eye rendering black with post-process materials applied post-tonemapping.
PostProcess material should only copy the view portion of the rendertarget texture rather than the entire thing.

diff --git a/Engine/Source/Runtime/Renderer/Private/PostProcess/PostProcessMaterial.cpp b/Engine/Source/Runtime/Renderer/Private/PostProcess/PostProcessMaterial.cpp
--- a/Engine/Source/Runtime/Renderer/Private/PostProcess/PostProcessMaterial.cpp
+++ b/Engine/Source/Runtime/Renderer/Private/PostProcess/PostProcessMaterial.cpp
@@ -411,7 +411,7 @@ FScreenPassTexture AddPostProcessMaterialPass(
 		if (bPrimeOutputColor || bForceIntermediateRT)
 		{
 			// Copy existing contents to new output and use load-action to preserve untouched pixels.
-			AddDrawTexturePass(GraphBuilder, View, SceneColor.Texture, Output.Texture);
+			AddDrawTexturePass(GraphBuilder, View, SceneColor, Output);
 			Output.LoadAction = ERenderTargetLoadAction::ELoad;
 		}
 	}
