Fix shader compilation fails from adding mobile VT support
#rnx

diff --git a/Engine/Source/Runtime/Renderer/Private/VT/RuntimeVirtualTextureRender.cpp b/Engine/Source/Runtime/Renderer/Private/VT/RuntimeVirtualTextureRender.cpp
--- a/Engine/Source/Runtime/Renderer/Private/VT/RuntimeVirtualTextureRender.cpp
+++ b/Engine/Source/Runtime/Renderer/Private/VT/RuntimeVirtualTextureRender.cpp
@@ -28,7 +28,7 @@ namespace RuntimeVirtualTexture
 
 		static bool ShouldCompilePermutation(const FMeshMaterialShaderPermutationParameters& Parameters)
 		{
-			return IsFeatureLevelSupported(Parameters.Platform, ERHIFeatureLevel::ES3_1) && 
+			return UseVirtualTexturing(GetMaxSupportedFeatureLevel(Parameters.Platform)) &&
 				(Parameters.Material->GetMaterialDomain() == MD_RuntimeVirtualTexture || Parameters.Material->HasRuntimeVirtualTextureOutput());
 		}
 
@@ -516,7 +516,7 @@ namespace RuntimeVirtualTexture
 
 		static bool ShouldCompilePermutation(FGlobalShaderPermutationParameters const& Parameters)
 		{
-			return IsFeatureLevelSupported(Parameters.Platform, ERHIFeatureLevel::ES3_1);
+			return RHISupportsComputeShaders(Parameters.Platform);
 		}
 
 		FShader_VirtualTextureCopy()
