Fixed potential crash in ray tracing with multiple views if skylight was enabled. AnyRayTracingPassEnabled is now called for all views instead just the first one to skip raytracing AS creation and bindings.

diff --git a/Engine/Source/Runtime/Renderer/Private/DeferredShadingRenderer.cpp b/Engine/Source/Runtime/Renderer/Private/DeferredShadingRenderer.cpp
--- a/Engine/Source/Runtime/Renderer/Private/DeferredShadingRenderer.cpp
+++ b/Engine/Source/Runtime/Renderer/Private/DeferredShadingRenderer.cpp
@@ -623,7 +623,13 @@ bool FDeferredShadingSceneRenderer::GatherRayTracingWorldInstances(FRHICommandLi
 		return false;
 	}
 
-	if (!AnyRayTracingPassEnabled(Views[0]))
+	bool bAnyRayTracingPassEnabled = false;
+	for (int32 ViewIndex = 0; ViewIndex < Views.Num(); ViewIndex++)
+	{
+		bAnyRayTracingPassEnabled |= AnyRayTracingPassEnabled(Scene, Views[ViewIndex]);
+	}
+
+	if (!bAnyRayTracingPassEnabled)
 	{
 		return false;
 	}
@@ -883,7 +889,13 @@ bool FDeferredShadingSceneRenderer::DispatchRayTracingWorldUpdates(FRHICommandLi
 		return false;
 	}
 
-	if (!AnyRayTracingPassEnabled(Views[0]))
+	bool bAnyRayTracingPassEnabled = false;
+	for (int32 ViewIndex = 0; ViewIndex < Views.Num(); ViewIndex++)
+	{
+		bAnyRayTracingPassEnabled |= AnyRayTracingPassEnabled(Scene, Views[ViewIndex]);
+	}
+
+	if (!bAnyRayTracingPassEnabled)
 	{
 		return false;
 	}
@@ -2834,13 +2846,12 @@ void FDeferredShadingSceneRenderer::CopyStencilToLightingChannelTexture(FRHIComm
 
 #if RHI_RAYTRACING
 
-bool AnyRayTracingPassEnabled(const FViewInfo& View)
+bool AnyRayTracingPassEnabled(const FScene* Scene, const FViewInfo& View)
 {
 	static auto CVarRayTracingSkyLight = IConsoleManager::Get().FindConsoleVariable(TEXT("r.RayTracing.SkyLight"));
 	static auto CVarRayTracingShadows = IConsoleManager::Get().FindConsoleVariable(TEXT("r.RayTracing.Shadows"));
 	static auto CVarStochasticRectLight = IConsoleManager::Get().FindConsoleVariable(TEXT("r.RayTracing.StochasticRectLight"));
 
-	const bool bRayTracingSkyLight = CVarRayTracingSkyLight != nullptr && CVarRayTracingSkyLight->GetInt() > 0;
 	const bool bRayTracingShadows = CVarRayTracingShadows != nullptr && CVarRayTracingShadows->GetInt() > 0;
 	const bool bRayTracingStochasticRectLight = CVarStochasticRectLight != nullptr && CVarStochasticRectLight->GetInt() > 0;
 
@@ -2849,7 +2860,7 @@ bool AnyRayTracingPassEnabled(const FViewInfo& View)
 		|| ShouldRenderRayTracingReflections(View)
 		|| ShouldRenderRayTracingGlobalIllumination(View)
 		|| ShouldRenderRayTracingTranslucency(View)
-		|| bRayTracingSkyLight
+		|| ShouldRenderRayTracingSkyLight(Scene? Scene->SkyLight : nullptr)
 		|| bRayTracingShadows
 		|| bRayTracingStochasticRectLight
 		|| View.RayTracingRenderMode == ERayTracingRenderMode::PathTracing
diff --git a/Engine/Source/Runtime/Renderer/Private/RayTracing/RaytracingOptions.h b/Engine/Source/Runtime/Renderer/Private/RayTracing/RaytracingOptions.h
--- a/Engine/Source/Runtime/Renderer/Private/RayTracing/RaytracingOptions.h
+++ b/Engine/Source/Runtime/Renderer/Private/RayTracing/RaytracingOptions.h
@@ -42,7 +42,7 @@ struct FRayTracingPrimaryRaysOptions
 
 
 #if RHI_RAYTRACING
-extern bool AnyRayTracingPassEnabled(const FViewInfo& View);
+extern bool AnyRayTracingPassEnabled(const FScene* Scene, const FViewInfo& View);
 extern int32 GetForceRayTracingEffectsCVarValue();
 extern FRayTracingPrimaryRaysOptions GetRayTracingTranslucencyOptions();
 
@@ -61,7 +61,7 @@ extern float GetRaytracingMaxNormalBias();
 
 #else
 
-FORCEINLINE bool AnyRayTracingPassEnabled(const FViewInfo& View)
+FORCEINLINE bool AnyRayTracingPassEnabled(const FScene* Scene, const FViewInfo& View)
 {
 	return 0;
 }
