Release a missed global resource on RHI shutdown (fixes crash on exit). (UE-87341)

diff --git a/Engine/Source/Runtime/Renderer/Private/SystemTextures.cpp b/Engine/Source/Runtime/Renderer/Private/SystemTextures.cpp
--- a/Engine/Source/Runtime/Renderer/Private/SystemTextures.cpp
+++ b/Engine/Source/Runtime/Renderer/Private/SystemTextures.cpp
@@ -633,6 +633,7 @@ void FSystemTextures::ReleaseDynamicRHI()
 	GreenDummy.SafeRelease();
 	DefaultNormal8Bit.SafeRelease();
 	VolumetricBlackDummy.SafeRelease();
+	ZeroUIntDummy.SafeRelease();
 	MidGreyDummy.SafeRelease();
 	StencilDummy.SafeRelease();
 	StencilDummySRV.SafeRelease();
