Added shadow lights csv stats

#rnx

diff --git a/Engine/Source/Runtime/Renderer/Private/ShadowSetup.cpp b/Engine/Source/Runtime/Renderer/Private/ShadowSetup.cpp
--- a/Engine/Source/Runtime/Renderer/Private/ShadowSetup.cpp
+++ b/Engine/Source/Runtime/Renderer/Private/ShadowSetup.cpp
@@ -4121,12 +4121,14 @@ void FSceneRenderer::InitDynamicShadows(FRHICommandListImmediate& RHICmdList, FG
 	const bool bProjectEnablePointLightShadows = Scene->ReadOnlyCVARCache.bEnablePointLightShadows;
 	uint32 NumPointShadowCachesUpdatedThisFrame = 0;
 	uint32 NumSpotShadowCachesUpdatedThisFrame = 0;
+	uint32 NumInteractionShadows = 0;
 
 	TArray<FProjectedShadowInfo*,SceneRenderingAllocator> PreShadows;
 	TArray<FProjectedShadowInfo*,SceneRenderingAllocator> ViewDependentWholeSceneShadows;
 	TArray<FProjectedShadowInfo*,SceneRenderingAllocator> ViewDependentWholeSceneShadowsThatNeedCulling;
 	{
 		SCOPE_CYCLE_COUNTER(STAT_InitDynamicShadowsTime);
+		CSV_SCOPED_TIMING_STAT_EXCLUSIVE(ShadowInitDynamic);
 
 		for (TSparseArray<FLightSceneInfoCompact>::TConstIterator LightIt(Scene->Lights); LightIt; ++LightIt)
 		{
@@ -4229,6 +4231,7 @@ void FSceneRenderer::InitDynamicShadows(FRHICommandListImmediate& RHICmdList, FG
 								)
 							{
 								SetupInteractionShadows(RHICmdList, Interaction, VisibleLightInfo, bStaticSceneOnly, ViewDependentWholeSceneShadows, PreShadows);
+								NumInteractionShadows++;
 							}
 
 							for (FLightPrimitiveInteraction* Interaction = LightSceneInfo->DynamicInteractionStaticPrimitiveList;
@@ -4237,6 +4240,7 @@ void FSceneRenderer::InitDynamicShadows(FRHICommandListImmediate& RHICmdList, FG
 								)
 							{
 								SetupInteractionShadows(RHICmdList, Interaction, VisibleLightInfo, bStaticSceneOnly, ViewDependentWholeSceneShadows, PreShadows);
+								NumInteractionShadows++;
 							}
 						}
 					}
@@ -4244,6 +4248,9 @@ void FSceneRenderer::InitDynamicShadows(FRHICommandListImmediate& RHICmdList, FG
 			}
 		}
 
+		CSV_CUSTOM_STAT_GLOBAL(ShadowTotalLightUpdated, float(NumPointShadowCachesUpdatedThisFrame + NumSpotShadowCachesUpdatedThisFrame), ECsvCustomStatOp::Set);
+		CSV_CUSTOM_STAT_GLOBAL(ShadowInteractionShadowsCount, float(NumInteractionShadows), ECsvCustomStatOp::Set);
+
 		// Calculate visibility of the projected shadows.
 		InitProjectedShadowVisibility(RHICmdList);
 	}
