Added blueprint functions to change the lighting channels for primitive and light components at runtime.

diff --git a/Engine/Source/Runtime/Engine/Classes/Components/LightComponent.h b/Engine/Source/Runtime/Engine/Classes/Components/LightComponent.h
--- a/Engine/Source/Runtime/Engine/Classes/Components/LightComponent.h
+++ b/Engine/Source/Runtime/Engine/Classes/Components/LightComponent.h
@@ -302,6 +302,9 @@ public:
 	UFUNCTION(BlueprintCallable, Category = "Rendering|Components|Light")
 	void SetForceCachedShadowsForMovablePrimitives(bool bNewValue);
 
+	UFUNCTION(BlueprintCallable, Category = "Rendering|Components|Light")
+	void SetLightingChannels(bool bChannel0, bool bChannel1, bool bChannel2);
+
 public:
 	/** The light's scene info. */
 	class FLightSceneProxy* SceneProxy;
diff --git a/Engine/Source/Runtime/Engine/Classes/Components/PrimitiveComponent.h b/Engine/Source/Runtime/Engine/Classes/Components/PrimitiveComponent.h
--- a/Engine/Source/Runtime/Engine/Classes/Components/PrimitiveComponent.h
+++ b/Engine/Source/Runtime/Engine/Classes/Components/PrimitiveComponent.h
@@ -274,6 +274,9 @@ public:
 	UFUNCTION(BlueprintSetter)
 	void SetGenerateOverlapEvents(bool bInGenerateOverlapEvents);
 
+	UFUNCTION(BlueprintCallable, Category = "Rendering|Components")
+	void SetLightingChannels(bool bChannel0, bool bChannel1, bool bChannel2);
+
 private:
 	UPROPERTY(EditAnywhere, BlueprintGetter = GetGenerateOverlapEvents, BlueprintSetter = SetGenerateOverlapEvents, Category = Collision)
 	uint8 bGenerateOverlapEvents : 1;
diff --git a/Engine/Source/Runtime/Engine/Private/Components/LightComponent.cpp b/Engine/Source/Runtime/Engine/Private/Components/LightComponent.cpp
--- a/Engine/Source/Runtime/Engine/Private/Components/LightComponent.cpp
+++ b/Engine/Source/Runtime/Engine/Private/Components/LightComponent.cpp
@@ -1116,6 +1116,19 @@ void ULightComponent::SetForceCachedShadowsForMovablePrimitives(bool bNewValue)
 	}
 }
 
+void ULightComponent::SetLightingChannels(bool bChannel0, bool bChannel1, bool bChannel2)
+{
+	if (bChannel0 != LightingChannels.bChannel0 ||
+		bChannel1 != LightingChannels.bChannel1 ||
+		bChannel2 != LightingChannels.bChannel2)
+	{
+		LightingChannels.bChannel0 = bChannel0;
+		LightingChannels.bChannel1 = bChannel1;
+		LightingChannels.bChannel2 = bChannel2;
+		MarkRenderStateDirty();
+	}
+}
+
 // GetDirection
 FVector ULightComponent::GetDirection() const 
 { 
diff --git a/Engine/Source/Runtime/Engine/Private/Components/PrimitiveComponent.cpp b/Engine/Source/Runtime/Engine/Private/Components/PrimitiveComponent.cpp
--- a/Engine/Source/Runtime/Engine/Private/Components/PrimitiveComponent.cpp
+++ b/Engine/Source/Runtime/Engine/Private/Components/PrimitiveComponent.cpp
@@ -3273,6 +3273,22 @@ void UPrimitiveComponent::SetGenerateOverlapEvents(bool bInGenerateOverlapEvents
 	}
 }
 
+void UPrimitiveComponent::SetLightingChannels(bool bChannel0, bool bChannel1, bool bChannel2)
+{
+	if (bChannel0 != LightingChannels.bChannel0 ||
+		bChannel1 != LightingChannels.bChannel1 ||
+		bChannel2 != LightingChannels.bChannel2)
+	{
+		LightingChannels.bChannel0 = bChannel0;
+		LightingChannels.bChannel1 = bChannel1;
+		LightingChannels.bChannel2 = bChannel2;
+		if (SceneProxy)
+		{
+			SceneProxy->SetLightingChannels_GameThread(LightingChannels);
+		}
+	}
+}
+
 void UPrimitiveComponent::ClearComponentOverlaps(bool bDoNotifies, bool bSkipNotifySelf)
 {
 	if (OverlappingComponents.Num() > 0)
@@ -3618,14 +3634,10 @@ void UPrimitiveComponent::SetRenderCustomDepth(bool bValue)
 	if( bRenderCustomDepth != bValue )
 	{
 		bRenderCustomDepth = bValue;
-		if ( SceneProxy )
+		if (SceneProxy)
 		{
 			SceneProxy->SetCustomDepthEnabled_GameThread(bRenderCustomDepth);
 		}
-		else
-		{
-			MarkRenderStateDirty();
-		}
 	}
 }
 
@@ -3637,14 +3649,10 @@ void UPrimitiveComponent::SetCustomDepthStencilValue(int32 Value)
 	if (CustomDepthStencilValue != ClampedValue)
 	{
 		CustomDepthStencilValue = ClampedValue;
-		if ( SceneProxy )
+		if (SceneProxy)
 		{
 			SceneProxy->SetCustomDepthStencilValue_GameThread(CustomDepthStencilValue);
 		}
-		else
-		{
-			MarkRenderStateDirty();
-		}
 	}
 }
 
diff --git a/Engine/Source/Runtime/Engine/Private/PrimitiveSceneProxy.cpp b/Engine/Source/Runtime/Engine/Private/PrimitiveSceneProxy.cpp
--- a/Engine/Source/Runtime/Engine/Private/PrimitiveSceneProxy.cpp
+++ b/Engine/Source/Runtime/Engine/Private/PrimitiveSceneProxy.cpp
@@ -576,6 +576,20 @@ void FPrimitiveSceneProxy::SetHovered_GameThread(const bool bInHovered)
 		});
 }
 
+void FPrimitiveSceneProxy::SetLightingChannels_GameThread(FLightingChannels LightingChannels)
+{
+	check(IsInGameThread());
+
+	FPrimitiveSceneProxy* PrimitiveSceneProxy = this;
+	const uint8 LocalLightingChannelMask = GetLightingChannelMaskForStruct(LightingChannels);
+	ENQUEUE_RENDER_COMMAND(SetLightingChannelsCmd)(
+		[PrimitiveSceneProxy, LocalLightingChannelMask](FRHICommandListImmediate& RHICmdList)
+	{
+		PrimitiveSceneProxy->LightingChannelMask = LocalLightingChannelMask;
+		PrimitiveSceneProxy->GetPrimitiveSceneInfo()->SetNeedsUniformBufferUpdate(true);
+	});
+}
+
 #if !UE_BUILD_SHIPPING
 void FPrimitiveSceneProxy::FDebugMassData::DrawDebugMass(class FPrimitiveDrawInterface* PDI, const FTransform& ElemTM) const
 {
diff --git a/Engine/Source/Runtime/Engine/Public/PrimitiveSceneProxy.h b/Engine/Source/Runtime/Engine/Public/PrimitiveSceneProxy.h
--- a/Engine/Source/Runtime/Engine/Public/PrimitiveSceneProxy.h
+++ b/Engine/Source/Runtime/Engine/Public/PrimitiveSceneProxy.h
@@ -151,6 +151,11 @@ public:
 	 */
 	void SetHovered_GameThread(const bool bInHovered);
 
+	/**
+	 * Updates the lighting channels for the primitive proxy.
+	 */
+	void SetLightingChannels_GameThread(FLightingChannels LightingChannels);
+
 	/**
 	 * Updates the hidden editor view visibility map on the game thread which just enqueues a command on the render thread
 	 */
diff --git a/Engine/Source/Runtime/Renderer/Private/PrimitiveSceneInfo.cpp b/Engine/Source/Runtime/Renderer/Private/PrimitiveSceneInfo.cpp
--- a/Engine/Source/Runtime/Renderer/Private/PrimitiveSceneInfo.cpp
+++ b/Engine/Source/Runtime/Renderer/Private/PrimitiveSceneInfo.cpp
@@ -809,6 +809,7 @@ void FPrimitiveSceneInfo::UpdateUniformBuffer(FRHICommandListImmediate& RHICmdLi
 	checkSlow(bNeedsUniformBufferUpdate);
 	bNeedsUniformBufferUpdate = false;
 	Proxy->UpdateUniformBuffer();
+	AddPrimitiveToUpdateGPU(*Scene, PackedIndex);
 }
 
 void FPrimitiveSceneInfo::BeginDeferredUpdateStaticMeshes()
diff --git a/Engine/Source/Runtime/Renderer/Private/RendererScene.cpp b/Engine/Source/Runtime/Renderer/Private/RendererScene.cpp
--- a/Engine/Source/Runtime/Renderer/Private/RendererScene.cpp
+++ b/Engine/Source/Runtime/Renderer/Private/RendererScene.cpp
@@ -2431,9 +2431,6 @@ void FSceneVelocityData::StartFrame(FScene* Scene)
 		{
 			// Recreate PrimitiveUniformBuffer on the frame after the primitive moved, since it contains PreviousLocalToWorld
 			VelocityData.PrimitiveSceneInfo->SetNeedsUniformBufferUpdate(true);
-
-			check(VelocityData.PrimitiveSceneInfo->IsIndexValid());
-			AddPrimitiveToUpdateGPU(*Scene, VelocityData.PrimitiveSceneInfo->GetIndex());
 		}
 
 		if (bTrimOld && (InternalFrameIndex - VelocityData.LastFrameUsed) > 10)
@@ -3946,16 +3943,8 @@ void FScene::UpdateAllPrimitiveSceneInfos(FRHICommandListImmediate& RHICmdList)
 		FScopeCycleCounter Context(PrimitiveSceneProxy->GetStatId());
 		PrimitiveSceneProxy->CustomPrimitiveData = CustomParams.Value;
 
-		// No need to do any of this if GPUScene isn't used (the custom primitive data will make it to the primitive uniform buffer through FPrimitiveSceneProxy::UpdateUniformBuffer if that's the case)
-		if (UseGPUScene(GMaxRHIShaderPlatform, GetFeatureLevel()))
-		{
-			AddPrimitiveToUpdateGPU(*this, PrimitiveSceneProxy->GetPrimitiveSceneInfo()->PackedIndex);
-		}
-		else
-		{
-			// Make sure the uniform buffer is updated before rendering
-			PrimitiveSceneProxy->GetPrimitiveSceneInfo()->SetNeedsUniformBufferUpdate(true);
-		}
+		// Make sure the uniform buffer is updated before rendering
+		PrimitiveSceneProxy->GetPrimitiveSceneInfo()->SetNeedsUniformBufferUpdate(true);
 	}
 
 	for (FPrimitiveSceneInfo* PrimitiveSceneInfo : DistanceFieldSceneDataUpdates)
