World Origin Shifting Rendering Desync On Play

Made sure we apply the world offset to the AddedPrimitiveSceneInfos

diff --git a/Engine/Source/Runtime/Renderer/Private/RendererScene.cpp b/Engine/Source/Runtime/Renderer/Private/RendererScene.cpp
--- a/Engine/Source/Runtime/Renderer/Private/RendererScene.cpp
+++ b/Engine/Source/Runtime/Renderer/Private/RendererScene.cpp
@@ -3270,6 +3270,7 @@ void FScene::ApplyWorldOffset(FVector InOffset)
 	ENQUEUE_RENDER_COMMAND(FApplyWorldOffset)(
 		[Scene, InOffset](FRHICommandListImmediate& RHICmdList)
 		{
+			Scene->UpdateAllPrimitiveSceneInfos(RHICmdList);
 			Scene->ApplyWorldOffset_RenderThread(InOffset);
 		});
 }
@@ -3281,11 +3282,12 @@ void FScene::ApplyWorldOffset_RenderThread(const FVector& InOffset)
 	GPUScene.bUpdateAllPrimitives = true;
 
 	// Primitives
+	checkf(AddedPrimitiveSceneInfos.Num() == 0, TEXT("All primitives found in AddedPrimitiveSceneInfos must have been added to the scene before the world offset is applied"));
 	for (int32 Idx = 0; Idx < Primitives.Num(); ++Idx)
 	{
 		Primitives[Idx]->ApplyWorldOffset(InOffset);
 	}
-	
+
 	// Primitive transforms
 	for (int32 Idx = 0; Idx < PrimitiveTransforms.Num(); ++Idx)
 	{
