Merged fix for DBufferMask ensure.

diff --git a/Engine/Source/Runtime/Renderer/Private/RendererUtils.cpp b/Engine/Source/Runtime/Renderer/Private/RendererUtils.cpp
--- a/Engine/Source/Runtime/Renderer/Private/RendererUtils.cpp
+++ b/Engine/Source/Runtime/Renderer/Private/RendererUtils.cpp
@@ -2,6 +2,7 @@
 
 #include "RendererUtils.h"
 #include "RenderTargetPool.h"
+#include "VisualizeTexture.h"
 
 template <uint32 NumRenderTargets>
 class TRTWriteMaskDecodeCS : public FGlobalShader
@@ -104,6 +105,10 @@ void FRenderTargetWriteMask::Decode(FRHICommandListImmediate& RHICmdList, TShade
 		1,
 		false));
 
+#if SUPPORTS_VISUALIZE_TEXTURE
+	MaskDesc.TargetableFlags |= TexCreate_ShaderResource;
+#endif
+
 	GRenderTargetPool.FindFreeElement(RHICmdList, MaskDesc, OutRTWriteMask, RTWriteMaskDebugName);
 	RHICmdList.TransitionResource(EResourceTransitionAccess::EWritable, EResourceTransitionPipeline::EGfxToCompute, OutRTWriteMask->GetRenderTargetItem().UAV);
 
