Fix a bug where the wrong indirect argument buffer is used when culling heightfield objects in to screen tiles

diff --git a/Engine/Source/Runtime/Renderer/Private/DistanceFieldShadowing.cpp b/Engine/Source/Runtime/Renderer/Private/DistanceFieldShadowing.cpp
--- a/Engine/Source/Runtime/Renderer/Private/DistanceFieldShadowing.cpp
+++ b/Engine/Source/Runtime/Renderer/Private/DistanceFieldShadowing.cpp
@@ -928,10 +928,11 @@ void ScatterObjectsToShadowTiles(
 
 		RHICmdList.SetStreamSource(0, GetUnitCubeVertexBuffer(), 0);
 
-		RHICmdList.DrawIndexedPrimitiveIndirect(
-			GetUnitCubeIndexBuffer(),
-			GShadowCulledObjectBuffers.Buffers.ObjectIndirectArguments.Buffer,
-			0);
+		FRHIVertexBuffer* IndirectArgsBuffer = PrimitiveType == DFPT_HeightField ?
+			GShadowCulledHeightFieldObjectBuffers.Buffers.ObjectIndirectArguments.Buffer :
+			GShadowCulledObjectBuffers.Buffers.ObjectIndirectArguments.Buffer;
+
+		RHICmdList.DrawIndexedPrimitiveIndirect(GetUnitCubeIndexBuffer(), IndirectArgsBuffer, 0);
 	}
 	RHICmdList.EndRenderPass();
 
