Replace the SRVs TMap in FSceneRenderTargetItem with an array and clear cubemap SRVs after mip gen. This saved 0.09 ms RT time according to a 10-run replay A/B test.

diff --git a/Engine/Source/Runtime/RenderCore/Private/GenerateMips.cpp b/Engine/Source/Runtime/RenderCore/Private/GenerateMips.cpp
--- a/Engine/Source/Runtime/RenderCore/Private/GenerateMips.cpp
+++ b/Engine/Source/Runtime/RenderCore/Private/GenerateMips.cpp
@@ -256,13 +256,13 @@ TSharedPtr<FGenerateMipsStruct> FGenerateMips::SetupTexture(FRHITexture* InTextu
 		RenderTexture.TargetableTexture = InTexture;
 		RenderTexture.ShaderResourceTexture = InTexture;
 
-		RenderTexture.SRVs.Reserve(Desc.NumMips);
-		RenderTexture.MipUAVs.Reserve(Desc.NumMips);
+		RenderTexture.SRVs.Empty(Desc.NumMips);
+		RenderTexture.MipUAVs.Empty(Desc.NumMips);
 		for (uint8 MipLevel = 0; MipLevel < Desc.NumMips; MipLevel++)
 		{
 			FRHITextureSRVCreateInfo SRVDesc;
 			SRVDesc.MipLevel = MipLevel;
-			RenderTexture.SRVs.Add(SRVDesc, RHICreateShaderResourceView((FTexture2DRHIRef&)InTexture, SRVDesc));
+			RenderTexture.SRVs.Emplace(SRVDesc, RHICreateShaderResourceView((FTexture2DRHIRef&)InTexture, SRVDesc));
 
 			RenderTexture.MipUAVs.Add(RHICreateUnorderedAccessView(InTexture, MipLevel));
 		}
diff --git a/Engine/Source/Runtime/RenderCore/Private/RenderGraphBuilder.cpp b/Engine/Source/Runtime/RenderCore/Private/RenderGraphBuilder.cpp
--- a/Engine/Source/Runtime/RenderCore/Private/RenderGraphBuilder.cpp
+++ b/Engine/Source/Runtime/RenderCore/Private/RenderGraphBuilder.cpp
@@ -539,16 +539,19 @@ void FRDGBuilder::AllocateRHITextureSRVIfNeeded(FRDGTextureSRV* SRV)
 		return;
 	}
 
-	if (RenderTarget.SRVs.Contains(SRV->Desc))
+	for (int32 Idx = 0; Idx < RenderTarget.SRVs.Num(); ++Idx)
 	{
-		SRV->ResourceRHI = RenderTarget.SRVs[SRV->Desc];
-		return;
+		if (RenderTarget.SRVs[Idx].Key == SRV->Desc)
+		{
+			SRV->ResourceRHI = RenderTarget.SRVs[Idx].Value;
+			return;
+		}
 	}
 
 	FShaderResourceViewRHIRef RHIShaderResourceView = RHICreateShaderResourceView(RenderTarget.ShaderResourceTexture, SRV->Desc);
 
 	SRV->ResourceRHI = RHIShaderResourceView;
-	RenderTarget.SRVs.Add(SRV->Desc, RHIShaderResourceView);
+	RenderTarget.SRVs.Emplace(SRV->Desc, MoveTemp(RHIShaderResourceView));
 }
 
 void FRDGBuilder::AllocateRHITextureUAVIfNeeded(FRDGTextureUAV* UAV)
diff --git a/Engine/Source/Runtime/RenderCore/Public/RendererInterface.h b/Engine/Source/Runtime/RenderCore/Public/RendererInterface.h
--- a/Engine/Source/Runtime/RenderCore/Public/RendererInterface.h
+++ b/Engine/Source/Runtime/RenderCore/Public/RendererInterface.h
@@ -466,7 +466,7 @@ struct FSceneRenderTargetItem
 	TArray< FUnorderedAccessViewRHIRef, TInlineAllocator<1> > MipUAVs;
 
 	/** All SRVs that has been created on for that ShaderResourceTexture.  */
-	TMap<FRHITextureSRVCreateInfo, FShaderResourceViewRHIRef> SRVs;
+	TArray<TPair<FRHITextureSRVCreateInfo, FShaderResourceViewRHIRef>, TInlineAllocator<1>> SRVs;
 
 	FShaderResourceViewRHIRef RTWriteMaskSRV;
 	FShaderResourceViewRHIRef FmaskSRV;
diff --git a/Engine/Source/Runtime/Renderer/Private/ReflectionEnvironmentCapture.cpp b/Engine/Source/Runtime/Renderer/Private/ReflectionEnvironmentCapture.cpp
--- a/Engine/Source/Runtime/Renderer/Private/ReflectionEnvironmentCapture.cpp
+++ b/Engine/Source/Runtime/Renderer/Private/ReflectionEnvironmentCapture.cpp
@@ -89,6 +89,14 @@ static TAutoConsoleVariable<int32> CVarReflectionCaptureStaticSceneOnly(
 	TEXT(" 0 is off, 1 is on (default)"),
 	ECVF_ReadOnly);
 
+static int32 GReleaseSRVsAfterCubeMipsGen = 1;
+static FAutoConsoleVariableRef CVarReleaseSRVsAfterCubeMipsGen(
+	TEXT("r.ReleaseSRVsAfterCubeMipsGen"),
+	GReleaseSRVsAfterCubeMipsGen,
+	TEXT("Don't store SRVs used during mip gen to the cubemap render target item.\n")
+	TEXT("This reduces unnecessary memory allocations and copies when the item is copied around during snapshot creation."),
+	ECVF_RenderThreadSafe);
+
 bool DoGPUArrayCopy()
 {
 	return GRHISupportsResolveCubemapFaces && CVarReflectionCaptureGPUArrayCopy.GetValueOnAnyThread();
@@ -185,14 +193,15 @@ void CreateCubeMips( FRHICommandListImmediate& RHICmdList, ERHIFeatureLevel::Typ
 
 	auto* ShaderMap = GetGlobalShaderMap(FeatureLevel);
 
-	for (int32 MipIndex = 0; MipIndex < NumMips; MipIndex++)
+	if (!Cubemap.SRVs.Num())
 	{
-		FRHITextureSRVCreateInfo SRVDesc;
-		SRVDesc.MipLevel = MipIndex;
+		Cubemap.SRVs.Empty(NumMips);
 
-		if (!Cubemap.SRVs.Contains(SRVDesc))
+		for (int32 MipIndex = 0; MipIndex < NumMips; MipIndex++)
 		{
-			Cubemap.SRVs.Add(SRVDesc, RHICreateShaderResourceView(Cubemap.ShaderResourceTexture, SRVDesc));
+			FRHITextureSRVCreateInfo SRVDesc;
+			SRVDesc.MipLevel = MipIndex;
+			Cubemap.SRVs.Emplace(SRVDesc, RHICreateShaderResourceView(Cubemap.ShaderResourceTexture, SRVDesc));
 		}
 	}
 
@@ -235,10 +244,8 @@ void CreateCubeMips( FRHICommandListImmediate& RHICmdList, ERHIFeatureLevel::Typ
 
 				SetShaderValue(RHICmdList, ShaderRHI, PixelShader->NumMips, NumMips);
 
-				FRHITextureSRVCreateInfo SrcSRVDesc;
-				SrcSRVDesc.MipLevel = MipIndex - 1;
-
-				SetSRVParameter(RHICmdList, ShaderRHI, PixelShader->SourceTexture, Cubemap.SRVs[SrcSRVDesc]);
+				check(Cubemap.SRVs.IsValidIndex(MipIndex - 1) && Cubemap.SRVs[MipIndex - 1].Key.MipLevel == MipIndex - 1);
+				SetSRVParameter(RHICmdList, ShaderRHI, PixelShader->SourceTexture, Cubemap.SRVs[MipIndex - 1].Value);
 				SetSamplerParameter(RHICmdList, ShaderRHI, PixelShader->SourceTextureSampler, TStaticSamplerState<SF_Bilinear, AM_Clamp, AM_Clamp, AM_Clamp>::GetRHI());
 			}
 
@@ -257,6 +264,11 @@ void CreateCubeMips( FRHICommandListImmediate& RHICmdList, ERHIFeatureLevel::Typ
 	}
 
 	RHICmdList.TransitionResources(EResourceTransitionAccess::EReadable, &CubeRef, 1);
+
+	if (!!GReleaseSRVsAfterCubeMipsGen)
+	{
+		Cubemap.SRVs.Empty();
+	}
 }
 
 /** Computes the average brightness of the given reflection capture and stores it in the scene. */
