FPrimitiveViewRelevance inherits from FMaterialRelevance. On top of communizing code, copying a FMaterialRelevance instance to FPrimitiveViewRelevance one is now just a 64 bits assignment

#rnx

diff --git a/Engine/Plugins/2D/Paper2D/Source/Paper2D/Private/PaperRenderSceneProxy.cpp b/Engine/Plugins/2D/Paper2D/Source/Paper2D/Private/PaperRenderSceneProxy.cpp
--- a/Engine/Plugins/2D/Paper2D/Source/Paper2D/Private/PaperRenderSceneProxy.cpp
+++ b/Engine/Plugins/2D/Paper2D/Source/Paper2D/Private/PaperRenderSceneProxy.cpp
@@ -537,10 +537,10 @@ FPrimitiveViewRelevance FPaperRenderSceneProxy::GetViewRelevance(const FSceneVie
 #endif
 		)
 	{
-		Result.bOpaqueRelevance = true;
+		Result.bOpaque = true;
 	}
 
-	Result.bVelocityRelevance = IsMovable() && Result.bOpaqueRelevance && Result.bRenderInMainPass;
+	Result.bVelocityRelevance = IsMovable() && Result.bOpaque && Result.bRenderInMainPass;
 
 	return Result;
 }
diff --git a/Engine/Plugins/Editor/MeshEditor/Source/MeshEditor/OverlayComponent.cpp b/Engine/Plugins/Editor/MeshEditor/Source/MeshEditor/OverlayComponent.cpp
--- a/Engine/Plugins/Editor/MeshEditor/Source/MeshEditor/OverlayComponent.cpp
+++ b/Engine/Plugins/Editor/MeshEditor/Source/MeshEditor/OverlayComponent.cpp
@@ -281,7 +281,7 @@ public:
 		Result.bRenderCustomDepth = ShouldRenderCustomDepth();
 		Result.bTranslucentSelfShadow = bCastVolumetricTranslucentShadow;
 		MaterialRelevance.SetPrimitiveViewRelevance( Result );
-		Result.bVelocityRelevance = IsMovable() && Result.bOpaqueRelevance && Result.bRenderInMainPass;
+		Result.bVelocityRelevance = IsMovable() && Result.bOpaque && Result.bRenderInMainPass;
 		return Result;
 	}
 
diff --git a/Engine/Plugins/Editor/MeshEditor/Source/MeshEditor/WireframeMeshComponent.cpp b/Engine/Plugins/Editor/MeshEditor/Source/MeshEditor/WireframeMeshComponent.cpp
--- a/Engine/Plugins/Editor/MeshEditor/Source/MeshEditor/WireframeMeshComponent.cpp
+++ b/Engine/Plugins/Editor/MeshEditor/Source/MeshEditor/WireframeMeshComponent.cpp
@@ -371,7 +371,7 @@ public:
 		Result.bRenderCustomDepth = ShouldRenderCustomDepth();
 		Result.bTranslucentSelfShadow = bCastVolumetricTranslucentShadow;
 		MaterialRelevance.SetPrimitiveViewRelevance( Result );
-		Result.bVelocityRelevance = IsMovable() && Result.bOpaqueRelevance && Result.bRenderInMainPass;
+		Result.bVelocityRelevance = IsMovable() && Result.bOpaque && Result.bRenderInMainPass;
 		return Result;
 	}
 
diff --git a/Engine/Plugins/Experimental/GeometryCache/Source/GeometryCache/Private/GeometryCacheSceneProxy.cpp b/Engine/Plugins/Experimental/GeometryCache/Source/GeometryCache/Private/GeometryCacheSceneProxy.cpp
--- a/Engine/Plugins/Experimental/GeometryCache/Source/GeometryCache/Private/GeometryCacheSceneProxy.cpp
+++ b/Engine/Plugins/Experimental/GeometryCache/Source/GeometryCache/Private/GeometryCacheSceneProxy.cpp
@@ -524,7 +524,7 @@ FPrimitiveViewRelevance FGeometryCacheSceneProxy::GetViewRelevance(const FSceneV
 	Result.bDynamicRelevance = true;
 	Result.bRenderCustomDepth = ShouldRenderCustomDepth();
 	MaterialRelevance.SetPrimitiveViewRelevance(Result);
-	Result.bVelocityRelevance = IsMovable() && Result.bOpaqueRelevance && Result.bRenderInMainPass;
+	Result.bVelocityRelevance = IsMovable() && Result.bOpaque && Result.bRenderInMainPass;
 	return Result;
 }
 
diff --git a/Engine/Plugins/Experimental/HairStrands/Source/HairStrandsCore/Private/GroomComponent.cpp b/Engine/Plugins/Experimental/HairStrands/Source/HairStrandsCore/Private/GroomComponent.cpp
--- a/Engine/Plugins/Experimental/HairStrands/Source/HairStrandsCore/Private/GroomComponent.cpp
+++ b/Engine/Plugins/Experimental/HairStrands/Source/HairStrandsCore/Private/GroomComponent.cpp
@@ -337,7 +337,7 @@ public:
 		}
 
 		FPrimitiveViewRelevance Result;
-		Result.bHairStrandsRelevance = bIsViewModeValid;
+		Result.bHairStrands = bIsViewModeValid;
 
 		// Special pass for hair strands geometry (not part of the base pass, and shadowing is handlded in a custom fashion)
 		Result.bDrawRelevance = false;		
diff --git a/Engine/Plugins/Experimental/ImagePlate/Source/ImagePlate/Private/ImagePlateComponent.cpp b/Engine/Plugins/Experimental/ImagePlate/Source/ImagePlate/Private/ImagePlateComponent.cpp
--- a/Engine/Plugins/Experimental/ImagePlate/Source/ImagePlate/Private/ImagePlateComponent.cpp
+++ b/Engine/Plugins/Experimental/ImagePlate/Source/ImagePlate/Private/ImagePlateComponent.cpp
@@ -231,7 +231,7 @@ namespace
 			
 			MaterialRelevance.SetPrimitiveViewRelevance(Result);
 
-			Result.bVelocityRelevance = IsMovable() && Result.bOpaqueRelevance && Result.bRenderInMainPass;
+			Result.bVelocityRelevance = IsMovable() && Result.bOpaque && Result.bRenderInMainPass;
 
 			return Result;
 		}
diff --git a/Engine/Plugins/Experimental/MeshModelingToolset/Source/ModelingComponents/Private/OctreeDynamicMeshSceneProxy.h b/Engine/Plugins/Experimental/MeshModelingToolset/Source/ModelingComponents/Private/OctreeDynamicMeshSceneProxy.h
--- a/Engine/Plugins/Experimental/MeshModelingToolset/Source/ModelingComponents/Private/OctreeDynamicMeshSceneProxy.h
+++ b/Engine/Plugins/Experimental/MeshModelingToolset/Source/ModelingComponents/Private/OctreeDynamicMeshSceneProxy.h
@@ -217,7 +217,7 @@ public:
 		Result.bTranslucentSelfShadow = bCastVolumetricTranslucentShadow;
 		Result.bRenderCustomDepth = ShouldRenderCustomDepth();
 		MaterialRelevance.SetPrimitiveViewRelevance(Result);
-		Result.bVelocityRelevance = IsMovable() && Result.bOpaqueRelevance && Result.bRenderInMainPass;
+		Result.bVelocityRelevance = IsMovable() && Result.bOpaque && Result.bRenderInMainPass;
 
 		return Result;
 	}
diff --git a/Engine/Plugins/Experimental/MeshModelingToolset/Source/ModelingComponents/Private/SimpleDynamicMeshSceneProxy.h b/Engine/Plugins/Experimental/MeshModelingToolset/Source/ModelingComponents/Private/SimpleDynamicMeshSceneProxy.h
--- a/Engine/Plugins/Experimental/MeshModelingToolset/Source/ModelingComponents/Private/SimpleDynamicMeshSceneProxy.h
+++ b/Engine/Plugins/Experimental/MeshModelingToolset/Source/ModelingComponents/Private/SimpleDynamicMeshSceneProxy.h
@@ -294,7 +294,7 @@ public:
 			Result.bTranslucentSelfShadow = bCastVolumetricTranslucentShadow;
 			Result.bRenderCustomDepth = ShouldRenderCustomDepth();
 			MaterialRelevance.SetPrimitiveViewRelevance(Result);
-			Result.bVelocityRelevance = IsMovable() && Result.bOpaqueRelevance && Result.bRenderInMainPass;
+			Result.bVelocityRelevance = IsMovable() && Result.bOpaque && Result.bRenderInMainPass;
 		}
 
 		return Result;
diff --git a/Engine/Plugins/FX/Niagara/Source/Niagara/Private/NiagaraComponent.cpp b/Engine/Plugins/FX/Niagara/Source/Niagara/Private/NiagaraComponent.cpp
--- a/Engine/Plugins/FX/Niagara/Source/Niagara/Private/NiagaraComponent.cpp
+++ b/Engine/Plugins/FX/Niagara/Source/Niagara/Private/NiagaraComponent.cpp
@@ -294,7 +294,7 @@ FPrimitiveViewRelevance FNiagaraSceneProxy::GetViewRelevance(const FSceneView* V
 		}
 	}
 
-	Relevance.bVelocityRelevance = IsMovable() && Relevance.bOpaqueRelevance && Relevance.bRenderInMainPass;
+	Relevance.bVelocityRelevance = IsMovable() && Relevance.bOpaque && Relevance.bRenderInMainPass;
 
 	return Relevance;
 }
diff --git a/Engine/Plugins/FX/Niagara/Source/Niagara/Private/NiagaraRenderer.cpp b/Engine/Plugins/FX/Niagara/Source/Niagara/Private/NiagaraRenderer.cpp
--- a/Engine/Plugins/FX/Niagara/Source/Niagara/Private/NiagaraRenderer.cpp
+++ b/Engine/Plugins/FX/Niagara/Source/Niagara/Private/NiagaraRenderer.cpp
@@ -309,7 +309,7 @@ FPrimitiveViewRelevance FNiagaraRenderer::GetViewRelevance(const FSceneView* Vie
 	Result.bDynamicRelevance = bHasDynamicData;
 	if (bHasDynamicData)
 	{
-		Result.bOpaqueRelevance = View->Family->EngineShowFlags.Bounds;
+		Result.bOpaque = View->Family->EngineShowFlags.Bounds;
 		DynamicDataRender->GetMaterialRelevance().SetPrimitiveViewRelevance(Result);
 	}
 
diff --git a/Engine/Plugins/FX/Niagara/Source/Niagara/Private/NiagaraRendererLights.cpp b/Engine/Plugins/FX/Niagara/Source/Niagara/Private/NiagaraRendererLights.cpp
--- a/Engine/Plugins/FX/Niagara/Source/Niagara/Private/NiagaraRendererLights.cpp
+++ b/Engine/Plugins/FX/Niagara/Source/Niagara/Private/NiagaraRendererLights.cpp
@@ -37,7 +37,7 @@ FPrimitiveViewRelevance FNiagaraRendererLights::GetViewRelevance(const FSceneVie
 	Result.bDrawRelevance = false;
 	Result.bShadowRelevance = false;
 	Result.bDynamicRelevance = false;
-	Result.bOpaqueRelevance = false;
+	Result.bOpaque = false;
 	Result.bHasSimpleLights = true;
 	//MaterialRelevance.SetPrimitiveViewRelevance(Result);
 	return Result;
diff --git a/Engine/Plugins/Runtime/CustomMeshComponent/Source/CustomMeshComponent/Private/CustomMeshComponent.cpp b/Engine/Plugins/Runtime/CustomMeshComponent/Source/CustomMeshComponent/Private/CustomMeshComponent.cpp
--- a/Engine/Plugins/Runtime/CustomMeshComponent/Source/CustomMeshComponent/Private/CustomMeshComponent.cpp
+++ b/Engine/Plugins/Runtime/CustomMeshComponent/Source/CustomMeshComponent/Private/CustomMeshComponent.cpp
@@ -163,7 +163,7 @@ public:
 		Result.bRenderCustomDepth = ShouldRenderCustomDepth();
 		Result.bTranslucentSelfShadow = bCastVolumetricTranslucentShadow;
 		MaterialRelevance.SetPrimitiveViewRelevance(Result);
-		Result.bVelocityRelevance = IsMovable() && Result.bOpaqueRelevance && Result.bRenderInMainPass;
+		Result.bVelocityRelevance = IsMovable() && Result.bOpaque && Result.bRenderInMainPass;
 		return Result;
 	}
 
diff --git a/Engine/Plugins/Runtime/Oculus/OculusVR/Source/OculusMR/Private/OculusMR_BoundaryMeshComponent.cpp b/Engine/Plugins/Runtime/Oculus/OculusVR/Source/OculusMR/Private/OculusMR_BoundaryMeshComponent.cpp
--- a/Engine/Plugins/Runtime/Oculus/OculusVR/Source/OculusMR/Private/OculusMR_BoundaryMeshComponent.cpp
+++ b/Engine/Plugins/Runtime/Oculus/OculusVR/Source/OculusMR/Private/OculusMR_BoundaryMeshComponent.cpp
@@ -186,7 +186,7 @@ public:
 		Result.bUsesLightingChannels = false;
 		Result.bRenderCustomDepth = false;
 		MaterialRelevance.SetPrimitiveViewRelevance(Result);
-		Result.bOpaqueRelevance = true;
+		Result.bOpaque = true;
 		Result.bUsesSceneDepth = false;
 		Result.bUsesSkyMaterial = false;
 		Result.bUsesSingleLayerWaterMaterial = false;
diff --git a/Engine/Plugins/Runtime/ProceduralMeshComponent/Source/ProceduralMeshComponent/Private/ProceduralMeshComponent.cpp b/Engine/Plugins/Runtime/ProceduralMeshComponent/Source/ProceduralMeshComponent/Private/ProceduralMeshComponent.cpp
--- a/Engine/Plugins/Runtime/ProceduralMeshComponent/Source/ProceduralMeshComponent/Private/ProceduralMeshComponent.cpp
+++ b/Engine/Plugins/Runtime/ProceduralMeshComponent/Source/ProceduralMeshComponent/Private/ProceduralMeshComponent.cpp
@@ -402,7 +402,7 @@ public:
 		Result.bRenderCustomDepth = ShouldRenderCustomDepth();
 		Result.bTranslucentSelfShadow = bCastVolumetricTranslucentShadow;
 		MaterialRelevance.SetPrimitiveViewRelevance(Result);
-		Result.bVelocityRelevance = IsMovable() && Result.bOpaqueRelevance && Result.bRenderInMainPass;
+		Result.bVelocityRelevance = IsMovable() && Result.bOpaque && Result.bRenderInMainPass;
 		return Result;
 	}
 
diff --git a/Engine/Source/Developer/FunctionalTesting/Private/FuncTestRenderingComponent.cpp b/Engine/Source/Developer/FunctionalTesting/Private/FuncTestRenderingComponent.cpp
--- a/Engine/Source/Developer/FunctionalTesting/Private/FuncTestRenderingComponent.cpp
+++ b/Engine/Source/Developer/FunctionalTesting/Private/FuncTestRenderingComponent.cpp
@@ -47,7 +47,7 @@ public:
 		Result.bDrawRelevance = IsShown(View) && IsSelected();
 		Result.bDynamicRelevance = true;
 		// ideally the TranslucencyRelevance should be filled out by the material, here we do it conservative
-		Result.bSeparateTranslucencyRelevance = Result.bNormalTranslucencyRelevance = IsShown(View);
+		Result.bSeparateTranslucency = Result.bNormalTranslucency = IsShown(View);
 		return Result;
 	}
 
diff --git a/Engine/Source/Developer/LogVisualizer/Private/VisualLoggerRenderingActor.cpp b/Engine/Source/Developer/LogVisualizer/Private/VisualLoggerRenderingActor.cpp
--- a/Engine/Source/Developer/LogVisualizer/Private/VisualLoggerRenderingActor.cpp
+++ b/Engine/Source/Developer/LogVisualizer/Private/VisualLoggerRenderingActor.cpp
@@ -41,7 +41,7 @@ public:
 		Result.bDrawRelevance = IsShown(View);
 		Result.bDynamicRelevance = true;
 		// ideally the TranslucencyRelevance should be filled out by the material, here we do it conservative
-		Result.bSeparateTranslucencyRelevance = Result.bNormalTranslucencyRelevance = IsShown(View) && GIsEditor;
+		Result.bSeparateTranslucency = Result.bNormalTranslucency = IsShown(View) && GIsEditor;
 		return Result;
 	}
 
diff --git a/Engine/Source/Runtime/AIModule/Private/EnvironmentQuery/EQSRenderingComponent.cpp b/Engine/Source/Runtime/AIModule/Private/EnvironmentQuery/EQSRenderingComponent.cpp
--- a/Engine/Source/Runtime/AIModule/Private/EnvironmentQuery/EQSRenderingComponent.cpp
+++ b/Engine/Source/Runtime/AIModule/Private/EnvironmentQuery/EQSRenderingComponent.cpp
@@ -268,7 +268,7 @@ FPrimitiveViewRelevance FEQSSceneProxy::GetViewRelevance(const FSceneView* View)
 		&& (!bDrawOnlyWhenSelected || SafeIsActorSelected());
 	Result.bDynamicRelevance = true;
 	// ideally the TranslucencyRelevance should be filled out by the material, here we do it conservative
-	Result.bSeparateTranslucencyRelevance = Result.bNormalTranslucencyRelevance = IsShown(View);
+	Result.bSeparateTranslucency = Result.bNormalTranslucency = IsShown(View);
 	return Result;
 }
 
diff --git a/Engine/Source/Runtime/AIModule/Private/GameplayDebugger/GameplayDebuggerCategory_AI.cpp b/Engine/Source/Runtime/AIModule/Private/GameplayDebugger/GameplayDebuggerCategory_AI.cpp
--- a/Engine/Source/Runtime/AIModule/Private/GameplayDebugger/GameplayDebuggerCategory_AI.cpp
+++ b/Engine/Source/Runtime/AIModule/Private/GameplayDebugger/GameplayDebuggerCategory_AI.cpp
@@ -374,7 +374,7 @@ FDebugRenderSceneProxy* FGameplayDebuggerCategory_AI::CreateDebugSceneProxy(cons
 			const bool bCanShow = View->Family->EngineShowFlags.GetSingleFlag(ViewFlagIndex);
 
 			FPrimitiveViewRelevance Result;
-			Result.bDrawRelevance = Result.bSeparateTranslucencyRelevance = Result.bNormalTranslucencyRelevance = bCanShow;
+			Result.bDrawRelevance = Result.bSeparateTranslucency = Result.bNormalTranslucency = bCanShow;
 			Result.bDynamicRelevance = true;
 			return Result;
 		}
diff --git a/Engine/Source/Runtime/AIModule/Private/GameplayDebugger/GameplayDebuggerCategory_NavLocalGrid.cpp b/Engine/Source/Runtime/AIModule/Private/GameplayDebugger/GameplayDebuggerCategory_NavLocalGrid.cpp
--- a/Engine/Source/Runtime/AIModule/Private/GameplayDebugger/GameplayDebuggerCategory_NavLocalGrid.cpp
+++ b/Engine/Source/Runtime/AIModule/Private/GameplayDebugger/GameplayDebuggerCategory_NavLocalGrid.cpp
@@ -288,7 +288,7 @@ protected:
 		FPrimitiveViewRelevance Result;
 		Result.bDrawRelevance = IsShown(View);
 		Result.bDynamicRelevance = true;
-		Result.bNormalTranslucencyRelevance = Result.bSeparateTranslucencyRelevance = IsShown(View);
+		Result.bNormalTranslucency = Result.bSeparateTranslucency = IsShown(View);
 		return Result;
 	}
 
diff --git a/Engine/Source/Runtime/Engine/Classes/Materials/MaterialInterface.h b/Engine/Source/Runtime/Engine/Classes/Materials/MaterialInterface.h
--- a/Engine/Source/Runtime/Engine/Classes/Materials/MaterialInterface.h
+++ b/Engine/Source/Runtime/Engine/Classes/Materials/MaterialInterface.h
@@ -30,7 +30,6 @@ class UPhysicalMaterial;
 class UPhysicalMaterialMask;
 class USubsurfaceProfile;
 class UTexture;
-struct FPrimitiveViewRelevance;
 struct FMaterialParameterInfo;
 struct FMaterialResourceLocOnDisk;
 #if WITH_EDITORONLY_DATA
@@ -62,62 +61,6 @@ enum EMaterialUsage
 	MATUSAGE_MAX,
 };
 
-// the class is only storing bits, initialized to 0 and has an |= operator
-// to provide a combined set of multiple materials (component / mesh)
-struct ENGINE_API FMaterialRelevance
-{
-	// bits that express which EMaterialShadingModel are used
-	uint16 ShadingModelMask;
-	uint8 bOpaque : 1;
-	uint8 bMasked : 1;
-	uint8 bDistortion : 1;
-	uint8 bHairStrands : 1;
-	uint8 bSeparateTranslucency : 1; // Translucency After DOF
-	uint8 bNormalTranslucency : 1;
-	uint8 bUsesSceneColorCopy : 1;
-	uint8 bDisableOffscreenRendering : 1; // Blend Modulate
-	uint8 bDisableDepthTest : 1;
-	uint8 bOutputsTranslucentVelocity : 1;
-	uint8 bUsesGlobalDistanceField : 1;
-	uint8 bUsesWorldPositionOffset : 1;
-	uint8 bDecal : 1;
-	uint8 bTranslucentSurfaceLighting : 1;
-	uint8 bUsesSceneDepth : 1;
-	uint8 bUsesSkyMaterial : 1;
-	uint8 bUsesSingleLayerWaterMaterial : 1;
-	uint8 bHasVolumeMaterialDomain : 1;
-	uint8 bUsesDistanceCullFade : 1;
-	uint8 bUsesCustomDepthStencil : 1;
-
-	/** Default constructor */
-	FMaterialRelevance()
-	{
-		// the class is only storing bits initialized to 0, the following avoids code redundancy
-		uint8 * RESTRICT p = (uint8*)this;
-		for(uint32 i = 0; i < sizeof(*this); ++i)
-		{
-			*p++ = 0;
-		}
-	}
-
-	/** Bitwise OR operator.  Sets any relevance bits which are present in either. */
-	FMaterialRelevance& operator|=(const FMaterialRelevance& B)
-	{
-		// the class is only storing bits, the following avoids code redundancy
-		const uint8 * RESTRICT s = (const uint8*)&B;
-		uint8 * RESTRICT d = (uint8*)this;
-		for(uint32 i = 0; i < sizeof(*this); ++i)
-		{
-			*d = *d | *s; 
-			++s;++d;
-		}
-		return *this;
-	}
-
-	/** Copies the material's relevance flags to a primitive's view relevance flags. */
-	void SetPrimitiveViewRelevance(FPrimitiveViewRelevance& OutViewRelevance) const;
-};
-
 /** 
  *	UMaterial interface settings for Lightmass
  */
diff --git a/Engine/Source/Runtime/Engine/Classes/Materials/MaterialRelevance.h b/Engine/Source/Runtime/Engine/Classes/Materials/MaterialRelevance.h
new file mode 100644
--- /dev/null
+++ b/Engine/Source/Runtime/Engine/Classes/Materials/MaterialRelevance.h
@@ -0,0 +1,61 @@
+// Copyright Epic Games, Inc. All Rights Reserved.
+
+
+#pragma once
+
+#include "CoreMinimal.h"
+
+struct FPrimitiveViewRelevance;
+
+// the class is only storing bits, initialized to 0 and has an |= operator
+// to provide a combined set of multiple materials (component / mesh)
+struct ENGINE_API FMaterialRelevance
+{
+	// bits that express which EMaterialShadingModel are used
+	union
+	{
+		struct
+		{
+			uint16 ShadingModelMask;
+			uint8 bOpaque : 1;
+			uint8 bMasked : 1;
+			uint8 bDistortion : 1;
+			uint8 bHairStrands : 1;
+			uint8 bSeparateTranslucency : 1; // Translucency After DOF
+			uint8 bNormalTranslucency : 1;
+			uint8 bUsesSceneColorCopy : 1;
+			uint8 bDisableOffscreenRendering : 1; // Blend Modulate
+			uint8 bOutputsTranslucentVelocity : 1;
+			uint8 bUsesGlobalDistanceField : 1;
+			uint8 bUsesWorldPositionOffset : 1;
+			uint8 bDecal : 1;
+			uint8 bTranslucentSurfaceLighting : 1;
+			uint8 bUsesSceneDepth : 1;
+			uint8 bUsesSkyMaterial : 1;
+			uint8 bUsesSingleLayerWaterMaterial : 1;
+			uint8 bHasVolumeMaterialDomain : 1;
+			uint8 bUsesCustomDepthStencil : 1;
+			uint8 bUsesDistanceCullFade : 1;
+			uint8 bDisableDepthTest : 1;
+		};
+		uint64 Raw;
+	};
+
+	/** Default constructor */
+	FMaterialRelevance()
+		: Raw(0)
+	{
+	}
+
+	/** Bitwise OR operator.  Sets any relevance bits which are present in either. */
+	FMaterialRelevance& operator|=(const FMaterialRelevance& B)
+	{
+		Raw |= B.Raw;
+		return *this;
+	}
+
+	/** Copies the material's relevance flags to a primitive's view relevance flags. */
+	void SetPrimitiveViewRelevance(FPrimitiveViewRelevance& OutViewRelevance) const;
+};
+
+static_assert(sizeof(FMaterialRelevance) == sizeof(FMaterialRelevance::Raw), "Union Raw type is too small");
diff --git a/Engine/Source/Runtime/Engine/Private/Components/ArrowComponent.cpp b/Engine/Source/Runtime/Engine/Private/Components/ArrowComponent.cpp
--- a/Engine/Source/Runtime/Engine/Private/Components/ArrowComponent.cpp
+++ b/Engine/Source/Runtime/Engine/Private/Components/ArrowComponent.cpp
@@ -180,7 +180,7 @@ public:
 #endif
 		Result.bShadowRelevance = IsShadowCast(View);
 		Result.bEditorPrimitiveRelevance = UseEditorCompositing(View);
-		Result.bVelocityRelevance = IsMovable() && Result.bOpaqueRelevance && Result.bRenderInMainPass;
+		Result.bVelocityRelevance = IsMovable() && Result.bOpaque && Result.bRenderInMainPass;
 		return Result;
 	}
 
diff --git a/Engine/Source/Runtime/Engine/Private/Components/BillboardComponent.cpp b/Engine/Source/Runtime/Engine/Private/Components/BillboardComponent.cpp
--- a/Engine/Source/Runtime/Engine/Private/Components/BillboardComponent.cpp
+++ b/Engine/Source/Runtime/Engine/Private/Components/BillboardComponent.cpp
@@ -227,7 +227,7 @@ public:
 #endif
 		FPrimitiveViewRelevance Result;
 		Result.bDrawRelevance = IsShown(View) && bVisible;
-		Result.bOpaqueRelevance = true;
+		Result.bOpaque = true;
 		Result.bDynamicRelevance = true;
 		Result.bShadowRelevance = IsShadowCast(View);
 		Result.bEditorPrimitiveRelevance = UseEditorCompositing(View);
diff --git a/Engine/Source/Runtime/Engine/Private/Components/LineBatchComponent.cpp b/Engine/Source/Runtime/Engine/Private/Components/LineBatchComponent.cpp
--- a/Engine/Source/Runtime/Engine/Private/Components/LineBatchComponent.cpp
+++ b/Engine/Source/Runtime/Engine/Private/Components/LineBatchComponent.cpp
@@ -89,7 +89,7 @@ FPrimitiveViewRelevance FLineBatcherSceneProxy::GetViewRelevance(const FSceneVie
 	ViewRelevance.bDrawRelevance = IsShown(View);
 	ViewRelevance.bDynamicRelevance = true;
 	// ideally the TranslucencyRelevance should be filled out by the material, here we do it conservative
-	ViewRelevance.bSeparateTranslucencyRelevance = ViewRelevance.bNormalTranslucencyRelevance = true;
+	ViewRelevance.bSeparateTranslucency = ViewRelevance.bNormalTranslucency = true;
 	return ViewRelevance;
 }
 
diff --git a/Engine/Source/Runtime/Engine/Private/Components/TextRenderComponent.cpp b/Engine/Source/Runtime/Engine/Private/Components/TextRenderComponent.cpp
--- a/Engine/Source/Runtime/Engine/Private/Components/TextRenderComponent.cpp
+++ b/Engine/Source/Runtime/Engine/Private/Components/TextRenderComponent.cpp
@@ -820,7 +820,7 @@ FPrimitiveViewRelevance FTextRenderSceneProxy::GetViewRelevance(const FSceneView
 	}
 
 	MaterialRelevance.SetPrimitiveViewRelevance(Result);
-	Result.bVelocityRelevance = IsMovable() && Result.bOpaqueRelevance && Result.bRenderInMainPass;
+	Result.bVelocityRelevance = IsMovable() && Result.bOpaque && Result.bRenderInMainPass;
 	return Result;
 }
 
diff --git a/Engine/Source/Runtime/Engine/Private/Materials/MaterialInterface.cpp b/Engine/Source/Runtime/Engine/Private/Materials/MaterialInterface.cpp
--- a/Engine/Source/Runtime/Engine/Private/Materials/MaterialInterface.cpp
+++ b/Engine/Source/Runtime/Engine/Private/Materials/MaterialInterface.cpp
@@ -38,25 +38,7 @@ UEnum* UMaterialInterface::SamplerTypeEnum = nullptr;
 /** Copies the material's relevance flags to a primitive's view relevance flags. */
 void FMaterialRelevance::SetPrimitiveViewRelevance(FPrimitiveViewRelevance& OutViewRelevance) const
 {
-	OutViewRelevance.bOpaqueRelevance = bOpaque;
-	OutViewRelevance.bMaskedRelevance = bMasked;
-	OutViewRelevance.bTranslucentVelocityRelevance = bOutputsTranslucentVelocity;
-	OutViewRelevance.bDistortionRelevance = bDistortion;
-	OutViewRelevance.bHairStrandsRelevance = bHairStrands;
-	OutViewRelevance.bSeparateTranslucencyRelevance = bSeparateTranslucency;
-	OutViewRelevance.bNormalTranslucencyRelevance = bNormalTranslucency;
-	OutViewRelevance.bUsesSceneColorCopy = bUsesSceneColorCopy;
-	OutViewRelevance.bDisableOffscreenRendering = bDisableOffscreenRendering;
-	OutViewRelevance.ShadingModelMaskRelevance = ShadingModelMask;
-	OutViewRelevance.bUsesGlobalDistanceField = bUsesGlobalDistanceField;
-	OutViewRelevance.bUsesWorldPositionOffset = bUsesWorldPositionOffset;
-	OutViewRelevance.bDecal = bDecal;
-	OutViewRelevance.bTranslucentSurfaceLighting = bTranslucentSurfaceLighting;
-	OutViewRelevance.bUsesSceneDepth = bUsesSceneDepth;
-	OutViewRelevance.bUsesSkyMaterial = bUsesSkyMaterial;
-	OutViewRelevance.bUsesSingleLayerWaterMaterial = bUsesSingleLayerWaterMaterial;
-	OutViewRelevance.bHasVolumeMaterialDomain = bHasVolumeMaterialDomain;
-	OutViewRelevance.bUsesCustomDepthStencil = bUsesCustomDepthStencil;
+	OutViewRelevance.Raw = Raw;
 }
 
 //////////////////////////////////////////////////////////////////////////
diff --git a/Engine/Source/Runtime/Engine/Private/Particles/ParticleSystemRender.cpp b/Engine/Source/Runtime/Engine/Private/Particles/ParticleSystemRender.cpp
--- a/Engine/Source/Runtime/Engine/Private/Particles/ParticleSystemRender.cpp
+++ b/Engine/Source/Runtime/Engine/Private/Particles/ParticleSystemRender.cpp
@@ -7171,17 +7171,17 @@ FPrimitiveViewRelevance FParticleSystemSceneProxy::GetViewRelevance(const FScene
 	}
 	if (View->Family->EngineShowFlags.Bounds || View->Family->EngineShowFlags.VectorFields)
 	{
-		Result.bOpaqueRelevance = true;
+		Result.bOpaque = true;
 	}
 	// see if any of the emitters use dynamic vertex data
 	if (DynamicData == NULL)
 	{
 		// In order to get the LOD distances to update,
 		// we need to force a call to DrawDynamicElements...
-		Result.bOpaqueRelevance = true;
+		Result.bOpaque = true;
 	}
 
-	Result.bVelocityRelevance = IsMovable() && Result.bOpaqueRelevance && Result.bRenderInMainPass;
+	Result.bVelocityRelevance = IsMovable() && Result.bOpaque && Result.bRenderInMainPass;
 
 	return Result;
 }
diff --git a/Engine/Source/Runtime/Engine/Private/SkeletalMesh.cpp b/Engine/Source/Runtime/Engine/Private/SkeletalMesh.cpp
--- a/Engine/Source/Runtime/Engine/Private/SkeletalMesh.cpp
+++ b/Engine/Source/Runtime/Engine/Private/SkeletalMesh.cpp
@@ -5160,7 +5160,7 @@ FPrimitiveViewRelevance FSkeletalMeshSceneProxy::GetViewRelevance(const FSceneVi
 	MaterialRelevance.SetPrimitiveViewRelevance(Result);
 
 #if !UE_BUILD_SHIPPING
-	Result.bSeparateTranslucencyRelevance |= View->Family->EngineShowFlags.Constraints;
+	Result.bSeparateTranslucency |= View->Family->EngineShowFlags.Constraints;
 #endif
 
 #if WITH_EDITOR
@@ -5171,7 +5171,7 @@ FPrimitiveViewRelevance FSkeletalMeshSceneProxy::GetViewRelevance(const FSceneVi
 	}
 #endif
 
-	Result.bVelocityRelevance = IsMovable() && Result.bOpaqueRelevance && Result.bRenderInMainPass;
+	Result.bVelocityRelevance = IsMovable() && Result.bOpaque && Result.bRenderInMainPass;
 
 	return Result;
 }
diff --git a/Engine/Source/Runtime/Engine/Private/StaticMeshRender.cpp b/Engine/Source/Runtime/Engine/Private/StaticMeshRender.cpp
--- a/Engine/Source/Runtime/Engine/Private/StaticMeshRender.cpp
+++ b/Engine/Source/Runtime/Engine/Private/StaticMeshRender.cpp
@@ -1675,10 +1675,10 @@ FPrimitiveViewRelevance FStaticMeshSceneProxy::GetViewRelevance(const FSceneView
 #endif
 		)
 	{
-		Result.bOpaqueRelevance = true;
+		Result.bOpaque = true;
 	}
 
-	Result.bVelocityRelevance = IsMovable() && Result.bOpaqueRelevance && Result.bRenderInMainPass;
+	Result.bVelocityRelevance = IsMovable() && Result.bOpaque && Result.bRenderInMainPass;
 
 	return Result;
 }
diff --git a/Engine/Source/Runtime/Engine/Private/VectorField.cpp b/Engine/Source/Runtime/Engine/Private/VectorField.cpp
--- a/Engine/Source/Runtime/Engine/Private/VectorField.cpp
+++ b/Engine/Source/Runtime/Engine/Private/VectorField.cpp
@@ -551,7 +551,7 @@ public:
 		FPrimitiveViewRelevance Result;
 		Result.bDrawRelevance = IsShown(View); 
 		Result.bDynamicRelevance = true;
-		Result.bOpaqueRelevance = true;
+		Result.bOpaque = true;
 		return Result;
 	}
 
diff --git a/Engine/Source/Runtime/Engine/Public/PrimitiveViewRelevance.h b/Engine/Source/Runtime/Engine/Public/PrimitiveViewRelevance.h
--- a/Engine/Source/Runtime/Engine/Public/PrimitiveViewRelevance.h
+++ b/Engine/Source/Runtime/Engine/Public/PrimitiveViewRelevance.h
@@ -3,42 +3,18 @@
 #pragma once
 
 #include "CoreMinimal.h"
+#include "Materials/MaterialRelevance.h"
 
 /**
  * The different types of relevance a primitive scene proxy can declare towards a particular scene view.
  * the class is only storing bits, and has an |= operator
  */
-struct FPrimitiveViewRelevance
+struct FPrimitiveViewRelevance : public FMaterialRelevance
 {
 	// Warning: This class is memzeroed externally as 0 is assumed a
 	// valid value for all members meaning 'not relevant'. If this
 	// changes existing class usage should be re-evaluated
 
-	// from FMaterialRelevance (could be made the base class):
-
-	/** The LightingProfile supported by this primitive, as a bitmask. */
-	uint16 ShadingModelMaskRelevance;
-	/** The primitive has one or more opaque or masked elements. */
-	uint32 bOpaqueRelevance : 1;
-	/** The primitive has one or more masked elements. */
-	uint32 bMaskedRelevance : 1;
-	/** The primitive has one or more translucent elements which output velocity. */
-	uint32 bTranslucentVelocityRelevance : 1;
-	/** The primitive has one or more distortion elements. */
-	uint32 bDistortionRelevance : 1;
-	/** The primitive has one or more elements that have SeparateTranslucency. */
-	uint32 bSeparateTranslucencyRelevance : 1;
-	/** The primitive has one or more elements that have normal translucency. */
-	uint32 bNormalTranslucencyRelevance : 1;
-	/** For translucent primitives reading the scene color. */
-	uint32 bUsesSceneColorCopy : 1;
-	/** For primitive that can't render in offscreen buffers (blend modulate). */
-	uint32 bDisableOffscreenRendering : 1;
-	/** */
-	uint32 bUsesGlobalDistanceField : 1;
-
-	// others:
-
 	/** The primitive's static elements are rendered for the view. */
 	uint32 bStaticRelevance : 1; 
 	/** The primitive's dynamic elements are rendered for the view. */
@@ -47,8 +23,6 @@ struct FPrimitiveViewRelevance
 	uint32 bDrawRelevance : 1;
 	/** The primitive is casting a shadow. */
 	uint32 bShadowRelevance : 1;
-	/** The primitive is hair strands geometry. */
-	uint32 bHairStrandsRelevance : 1;
 	/** The primitive should render velocity. */
 	uint32 bVelocityRelevance : 1;
 	/** The primitive should render to the custom depth pass. */
@@ -57,8 +31,6 @@ struct FPrimitiveViewRelevance
 	uint32 bRenderInDepthPass : 1;
 	/** The primitive should render to the base pass / normal depth / velocity rendering. */
 	uint32 bRenderInMainPass : 1;
-	/** The primitive has materials using the volume domain. */
-	uint32 bHasVolumeMaterialDomain : 1;
 	/** The primitive is drawn only in the editor and composited onto the scene after post processing */
 	uint32 bEditorPrimitiveRelevance : 1;
 	/** The primitive's static elements are selected and rendered again in the selection outline pass*/
@@ -67,26 +39,12 @@ struct FPrimitiveViewRelevance
 	uint32 bEditorNoDepthTestPrimitiveRelevance : 1;
 	/** The primitive should have GatherSimpleLights called on the proxy when gathering simple lights. */
 	uint32 bHasSimpleLights : 1;
-	/** The primitive has one or more elements that have World Position Offset. */
-	uint32 bUsesWorldPositionOffset : 1;
 	/** Whether the primitive uses non-default lighting channels. */
 	uint32 bUsesLightingChannels : 1;
-	/** */
-	uint32 bDecal : 1;
-	/** Whether the primitive has materials that use translucent surface lighting. */
-	uint32 bTranslucentSurfaceLighting : 1;
 	/** Whether the primitive has materials that use volumetric translucent self shadow. */
 	uint32 bTranslucentSelfShadow : 1;
-	/** Whether the primitive has materials that read the scene depth. */
-	uint32 bUsesSceneDepth : 1;
-	/** Whether the primitive has materials that read the scene depth. */
-	uint32 bUsesSkyMaterial : 1;
-	/** Whether the primitive uses a single layer water material. */
-	uint32 bUsesSingleLayerWaterMaterial : 1;
 	/** Whether the view use custom data. */
 	uint32 bUseCustomViewData : 1;
-	/** Whether the primitive has materials that read the custom depth or custom stencil */
-	uint32 bUsesCustomDepthStencil : 1;
 
 	/** 
 	 * Whether this primitive view relevance has been initialized this frame.  
@@ -97,12 +55,12 @@ struct FPrimitiveViewRelevance
 
 	bool HasTranslucency() const 
 	{
-		return bSeparateTranslucencyRelevance || bNormalTranslucencyRelevance;
+		return bSeparateTranslucency || bNormalTranslucency;
 	}
 
 	bool HasVelocity() const
 	{
-		return bVelocityRelevance || bTranslucentVelocityRelevance;
+		return bVelocityRelevance || bOutputsTranslucentVelocity;
 	}
 
 	/** Default constructor */
@@ -117,7 +75,7 @@ struct FPrimitiveViewRelevance
 
 		// only exceptions (bugs we need to fix?):
 
-		bOpaqueRelevance = true;
+		bOpaque = true;
 		// without it BSP doesn't render
 		bRenderInMainPass = true;
 	}
diff --git a/Engine/Source/Runtime/Experimental/GeometryCollectionEngine/Private/GeometryCollection/GeometryCollectionSceneProxy.cpp b/Engine/Source/Runtime/Experimental/GeometryCollectionEngine/Private/GeometryCollection/GeometryCollectionSceneProxy.cpp
--- a/Engine/Source/Runtime/Experimental/GeometryCollectionEngine/Private/GeometryCollection/GeometryCollectionSceneProxy.cpp
+++ b/Engine/Source/Runtime/Experimental/GeometryCollectionEngine/Private/GeometryCollection/GeometryCollectionSceneProxy.cpp
@@ -767,7 +767,7 @@ FPrimitiveViewRelevance FGeometryCollectionSceneProxy::GetViewRelevance(const FS
 	MaterialRelevance.SetPrimitiveViewRelevance(Result);
 
 	// #todo(dmp): why does this crash?
-	// Result.bVelocityRelevance = IsMovable() && Result.bOpaqueRelevance && Result.bRenderInMainPass;
+	// Result.bVelocityRelevance = IsMovable() && Result.bOpaque && Result.bRenderInMainPass;
 
 	return Result;
 
diff --git a/Engine/Source/Runtime/Landscape/Private/LandscapeGizmoActor.cpp b/Engine/Source/Runtime/Landscape/Private/LandscapeGizmoActor.cpp
--- a/Engine/Source/Runtime/Landscape/Private/LandscapeGizmoActor.cpp
+++ b/Engine/Source/Runtime/Landscape/Private/LandscapeGizmoActor.cpp
@@ -343,7 +343,7 @@ public:
 		Result.bDrawRelevance = IsShown(View) && bVisible && !View->bIsGameView && GLandscapeEditRenderMode & ELandscapeEditRenderMode::Gizmo;
 		Result.bDynamicRelevance = true;
 		// ideally the TranslucencyRelevance should be filled out by the material, here we do it conservative
-		Result.bSeparateTranslucencyRelevance = Result.bNormalTranslucencyRelevance = true;
+		Result.bSeparateTranslucency = Result.bNormalTranslucency = true;
 #endif
 		return Result;
 	}
diff --git a/Engine/Source/Runtime/MRMesh/Private/MRMeshComponent.cpp b/Engine/Source/Runtime/MRMesh/Private/MRMeshComponent.cpp
--- a/Engine/Source/Runtime/MRMesh/Private/MRMeshComponent.cpp
+++ b/Engine/Source/Runtime/MRMesh/Private/MRMeshComponent.cpp
@@ -469,7 +469,7 @@ private:
 		Result.bUsesLightingChannels = GetLightingChannelMask() != GetDefaultLightingChannelMask();
 		Result.bRenderCustomDepth = ShouldRenderCustomDepth();
 		UMaterialInterface::TMicRecursionGuard RecursionGuard;
-		Result.bSeparateTranslucencyRelevance = MaterialToUse->GetMaterial_Concurrent(RecursionGuard)->bEnableSeparateTranslucency;
+		Result.bSeparateTranslucency = MaterialToUse->GetMaterial_Concurrent(RecursionGuard)->bEnableSeparateTranslucency;
 		//MaterialRelevance.SetPrimitiveViewRelevance(Result);
 		return Result;
 	}
diff --git a/Engine/Source/Runtime/NavigationSystem/Private/NavLinkRenderingComponent.cpp b/Engine/Source/Runtime/NavigationSystem/Private/NavLinkRenderingComponent.cpp
--- a/Engine/Source/Runtime/NavigationSystem/Private/NavLinkRenderingComponent.cpp
+++ b/Engine/Source/Runtime/NavigationSystem/Private/NavLinkRenderingComponent.cpp
@@ -442,7 +442,7 @@ FPrimitiveViewRelevance FNavLinkRenderingProxy::GetViewRelevance(const FSceneVie
 	Result.bDrawRelevance = IsShown(View) && IsSelected() && (View && View->Family && View->Family->EngineShowFlags.Navigation);
 	Result.bDynamicRelevance = true;
 	// ideally the TranslucencyRelevance should be filled out by the material, here we do it conservative
-	Result.bSeparateTranslucencyRelevance = Result.bNormalTranslucencyRelevance = IsShown(View);
+	Result.bSeparateTranslucency = Result.bNormalTranslucency = IsShown(View);
 	Result.bShadowRelevance = IsShadowCast(View);
 	Result.bEditorPrimitiveRelevance = UseEditorCompositing(View);
 	return Result;
diff --git a/Engine/Source/Runtime/NavigationSystem/Private/NavMesh/NavMeshRenderingComponent.cpp b/Engine/Source/Runtime/NavigationSystem/Private/NavMesh/NavMeshRenderingComponent.cpp
--- a/Engine/Source/Runtime/NavigationSystem/Private/NavMesh/NavMeshRenderingComponent.cpp
+++ b/Engine/Source/Runtime/NavigationSystem/Private/NavMesh/NavMeshRenderingComponent.cpp
@@ -1160,7 +1160,7 @@ FPrimitiveViewRelevance FNavMeshSceneProxy::GetViewRelevance(const FSceneView* V
 	Result.bDrawRelevance = bVisible && IsShown(View);
 	Result.bDynamicRelevance = true;
 	// ideally the TranslucencyRelevance should be filled out by the material, here we do it conservative
-	Result.bSeparateTranslucencyRelevance = Result.bNormalTranslucencyRelevance = bVisible && IsShown(View);
+	Result.bSeparateTranslucency = Result.bNormalTranslucency = bVisible && IsShown(View);
 	return Result;
 }
 
diff --git a/Engine/Source/Runtime/NavigationSystem/Private/NavMesh/NavTestRenderingComponent.cpp b/Engine/Source/Runtime/NavigationSystem/Private/NavMesh/NavTestRenderingComponent.cpp
--- a/Engine/Source/Runtime/NavigationSystem/Private/NavMesh/NavTestRenderingComponent.cpp
+++ b/Engine/Source/Runtime/NavigationSystem/Private/NavMesh/NavTestRenderingComponent.cpp
@@ -293,7 +293,7 @@ FPrimitiveViewRelevance FNavTestSceneProxy::GetViewRelevance(const FSceneView* V
 	Result.bDrawRelevance = IsShown(View);
 	Result.bDynamicRelevance = true;
 	// ideally the TranslucencyRelevance should be filled out by the material, here we do it conservative
-	Result.bSeparateTranslucencyRelevance = Result.bNormalTranslucencyRelevance = IsShown(View) && GIsEditor;
+	Result.bSeparateTranslucency = Result.bNormalTranslucency = IsShown(View) && GIsEditor;
 	return Result;
 }
 
diff --git a/Engine/Source/Runtime/Renderer/Private/IndirectLightingCache.cpp b/Engine/Source/Runtime/Renderer/Private/IndirectLightingCache.cpp
--- a/Engine/Source/Runtime/Renderer/Private/IndirectLightingCache.cpp
+++ b/Engine/Source/Runtime/Renderer/Private/IndirectLightingCache.cpp
@@ -549,7 +549,7 @@ void FIndirectLightingCache::UpdateCachePrimitivesInternal(FScene* Scene, FScene
 				{
 					uint32 PrimitiveIndex = BitIt.GetIndex();
 					// FDrawTranslucentMeshAction::AllowIndirectLightingCacheVolumeTexture doesn't allow volume samples on translucency, so we only need to support one if the primitive has at least one opaque material
-					const bool bAllowVolumeSample = View.PrimitiveViewRelevanceMap[PrimitiveIndex].bOpaqueRelevance;
+					const bool bAllowVolumeSample = View.PrimitiveViewRelevanceMap[PrimitiveIndex].bOpaque;
 					ProcessPrimitiveUpdate(Scene, View, PrimitiveIndex, bAllowUnbuiltPreview, bAllowVolumeSample, OutBlocksToUpdate, OutTransitionsOverTimeToUpdate, OutPrimitivesToUpdateStaticMeshes);
 				}
 
diff --git a/Engine/Source/Runtime/Renderer/Private/SceneSoftwareOcclusion.cpp b/Engine/Source/Runtime/Renderer/Private/SceneSoftwareOcclusion.cpp
--- a/Engine/Source/Runtime/Renderer/Private/SceneSoftwareOcclusion.cpp
+++ b/Engine/Source/Runtime/Renderer/Private/SceneSoftwareOcclusion.cpp
@@ -1046,8 +1046,8 @@ static FGraphEventRef SubmitScene(const FScene* Scene, FViewInfo& View, FOcclusi
 
 			// Relevance requirements
 			FPrimitiveViewRelevance ViewRelevance = Proxy->GetViewRelevance(&View);
-			const bool bNonOpaqueRelevance = (ViewRelevance.bMaskedRelevance || ViewRelevance.HasTranslucency()); // TODO: opaque sections
-			bool bCanBeOccluder = ViewRelevance.bDrawRelevance && (ViewRelevance.bOpaqueRelevance && !bNonOpaqueRelevance);
+			const bool bNonOpaqueRelevance = (ViewRelevance.bMasked || ViewRelevance.HasTranslucency()); // TODO: opaque sections
+			bool bCanBeOccluder = ViewRelevance.bDrawRelevance && (ViewRelevance.bOpaque && !bNonOpaqueRelevance);
 
 			if (bCanBeOccluder)
 			{
diff --git a/Engine/Source/Runtime/Renderer/Private/SceneVisibility.cpp b/Engine/Source/Runtime/Renderer/Private/SceneVisibility.cpp
--- a/Engine/Source/Runtime/Renderer/Private/SceneVisibility.cpp
+++ b/Engine/Source/Runtime/Renderer/Private/SceneVisibility.cpp
@@ -1972,7 +1972,7 @@ struct FRelevancePacket
 			const bool bEditorSelectionRelevance = ViewRelevance.bEditorStaticSelectionRelevance;
 			const bool bTranslucentRelevance = ViewRelevance.HasTranslucency();
 
-			const bool bHairStrandsEnabled = ViewRelevance.bHairStrandsRelevance && IsHairStrandsEnable(Scene->GetShaderPlatform());
+			const bool bHairStrandsEnabled = ViewRelevance.bHairStrands && IsHairStrandsEnable(Scene->GetShaderPlatform());
 			if (!bEditorRelevance && bHairStrandsEnabled)
 			{
 				++NumVisibleDynamicPrimitives;
@@ -2026,12 +2026,12 @@ struct FRelevancePacket
 			{
 				if (View.Family->AllowTranslucencyAfterDOF())
 				{
-					if (ViewRelevance.bNormalTranslucencyRelevance)
+					if (ViewRelevance.bNormalTranslucency)
 					{
 						TranslucentPrimCount.Add(ETranslucencyPass::TPT_StandardTranslucency, ViewRelevance.bUsesSceneColorCopy, ViewRelevance.bDisableOffscreenRendering);
 					}
 
-					if (ViewRelevance.bSeparateTranslucencyRelevance)
+					if (ViewRelevance.bSeparateTranslucency)
 					{
 						TranslucentPrimCount.Add(ETranslucencyPass::TPT_TranslucencyAfterDOF, ViewRelevance.bUsesSceneColorCopy, ViewRelevance.bDisableOffscreenRendering);
 					}
@@ -2042,13 +2042,13 @@ struct FRelevancePacket
 					TranslucentPrimCount.Add(ETranslucencyPass::TPT_AllTranslucency, ViewRelevance.bUsesSceneColorCopy, ViewRelevance.bDisableOffscreenRendering);
 				}
 
-				if (ViewRelevance.bDistortionRelevance)
+				if (ViewRelevance.bDistortion)
 				{
 					bHasDistortionPrimitives = true;
 				}
 			}
 			
-			CombinedShadingModelMask |= ViewRelevance.ShadingModelMaskRelevance;
+			CombinedShadingModelMask |= ViewRelevance.ShadingModelMask;
 			bUsesGlobalDistanceField |= ViewRelevance.bUsesGlobalDistanceField;
 			bUsesLightingChannels |= ViewRelevance.bUsesLightingChannels;
 			bTranslucentSurfaceLighting |= ViewRelevance.bTranslucentSurfaceLighting;
@@ -2303,7 +2303,7 @@ struct FRelevancePacket
 											DrawCommandPacket.AddCommandsForMesh(PrimitiveIndex, PrimitiveSceneInfo, StaticMeshRelevance, StaticMesh, Scene, bCanCache, EMeshPass::Velocity);
 										}
 
-										if (ViewRelevance.bTranslucentVelocityRelevance &&
+										if (ViewRelevance.bOutputsTranslucentVelocity &&
 											FTranslucentVelocityMeshProcessor::PrimitiveCanHaveVelocity(View.GetShaderPlatform(), PrimitiveSceneProxy) &&
 											FTranslucentVelocityMeshProcessor::PrimitiveHasVelocityForFrame(PrimitiveSceneProxy))
 										{
@@ -2327,12 +2327,12 @@ struct FRelevancePacket
 						{
 							if (View.Family->AllowTranslucencyAfterDOF())
 							{
-								if (ViewRelevance.bNormalTranslucencyRelevance)
+								if (ViewRelevance.bNormalTranslucency)
 								{
 									DrawCommandPacket.AddCommandsForMesh(PrimitiveIndex, PrimitiveSceneInfo, StaticMeshRelevance, StaticMesh, Scene, bCanCache, EMeshPass::TranslucencyStandard);
 								}
 
-								if (ViewRelevance.bSeparateTranslucencyRelevance)
+								if (ViewRelevance.bSeparateTranslucency)
 								{
 									DrawCommandPacket.AddCommandsForMesh(PrimitiveIndex, PrimitiveSceneInfo, StaticMeshRelevance, StaticMesh, Scene, bCanCache, EMeshPass::TranslucencyAfterDOF);
 								}
@@ -2344,7 +2344,7 @@ struct FRelevancePacket
 								DrawCommandPacket.AddCommandsForMesh(PrimitiveIndex, PrimitiveSceneInfo, StaticMeshRelevance, StaticMesh, Scene, bCanCache, EMeshPass::TranslucencyAll);
 							}
 
-							if (ViewRelevance.bDistortionRelevance)
+							if (ViewRelevance.bDistortion)
 							{
 								DrawCommandPacket.AddCommandsForMesh(PrimitiveIndex, PrimitiveSceneInfo, StaticMeshRelevance, StaticMesh, Scene, bCanCache, EMeshPass::Distortion);
 							}
@@ -2730,7 +2730,7 @@ void ComputeDynamicMeshRelevance(EShadingPath ShadingPath, bool bAddLightmapDens
 				View.NumVisibleDynamicMeshElements[EMeshPass::Velocity] += NumElements;
 			}
 
-			if (ViewRelevance.bTranslucentVelocityRelevance)
+			if (ViewRelevance.bOutputsTranslucentVelocity)
 			{
 				PassMask.Set(EMeshPass::TranslucentVelocity);
 				View.NumVisibleDynamicMeshElements[EMeshPass::TranslucentVelocity] += NumElements;
@@ -2750,13 +2750,13 @@ void ComputeDynamicMeshRelevance(EShadingPath ShadingPath, bool bAddLightmapDens
 	{
 		if (View.Family->AllowTranslucencyAfterDOF())
 		{
-			if (ViewRelevance.bNormalTranslucencyRelevance)
+			if (ViewRelevance.bNormalTranslucency)
 			{
 				PassMask.Set(EMeshPass::TranslucencyStandard);
 				View.NumVisibleDynamicMeshElements[EMeshPass::TranslucencyStandard] += NumElements;
 			}
 
-			if (ViewRelevance.bSeparateTranslucencyRelevance)
+			if (ViewRelevance.bSeparateTranslucency)
 			{
 				PassMask.Set(EMeshPass::TranslucencyAfterDOF);
 				View.NumVisibleDynamicMeshElements[EMeshPass::TranslucencyAfterDOF] += NumElements;
@@ -2768,7 +2768,7 @@ void ComputeDynamicMeshRelevance(EShadingPath ShadingPath, bool bAddLightmapDens
 			View.NumVisibleDynamicMeshElements[EMeshPass::TranslucencyAll] += NumElements;
 		}
 
-		if (ViewRelevance.bDistortionRelevance)
+		if (ViewRelevance.bDistortion)
 		{
 			PassMask.Set(EMeshPass::Distortion);
 			View.NumVisibleDynamicMeshElements[EMeshPass::Distortion] += NumElements;
@@ -2790,7 +2790,7 @@ void ComputeDynamicMeshRelevance(EShadingPath ShadingPath, bool bAddLightmapDens
 
 	// Hair strands are not rendered into the base pass (bRenderInMainPass=0) and so this 
 	// adds a special pass for allowing hair strands to be selectable.
-	if (View.bAllowTranslucentPrimitivesInHitProxy && ViewRelevance.bHairStrandsRelevance)
+	if (View.bAllowTranslucentPrimitivesInHitProxy && ViewRelevance.bHairStrands)
 	{
 		PassMask.Set(EMeshPass::HitProxy);
 		View.NumVisibleDynamicMeshElements[EMeshPass::HitProxy] += NumElements;
@@ -2814,7 +2814,7 @@ void ComputeDynamicMeshRelevance(EShadingPath ShadingPath, bool bAddLightmapDens
 		BatchAndProxy.SortKey = MeshBatch.PrimitiveSceneProxy->GetTranslucencySortPriority();
 	}
 	
-	const bool bIsHairStrandsCompatible = ViewRelevance.bHairStrandsRelevance && IsHairStrandsEnable(View.GetShaderPlatform());
+	const bool bIsHairStrandsCompatible = ViewRelevance.bHairStrands && IsHairStrandsEnable(View.GetShaderPlatform());
 	if (bIsHairStrandsCompatible)
 	{
 		View.HairStrandsMeshElements.AddUninitialized(1);
diff --git a/Engine/Source/Runtime/Renderer/Private/ShadowSetup.cpp b/Engine/Source/Runtime/Renderer/Private/ShadowSetup.cpp
--- a/Engine/Source/Runtime/Renderer/Private/ShadowSetup.cpp
+++ b/Engine/Source/Runtime/Renderer/Private/ShadowSetup.cpp
@@ -1049,7 +1049,7 @@ void FProjectedShadowInfo::AddSubjectPrimitive(FPrimitiveSceneInfo* PrimitiveSce
 			}
 		}
 
-		bool bOpaqueRelevance = false;
+		bool bOpaque = false;
 		bool bTranslucentRelevance = false;
 		bool bShadowRelevance = false;
 		bool bCustomDataRelevance = false;
@@ -1103,8 +1103,8 @@ void FProjectedShadowInfo::AddSubjectPrimitive(FPrimitiveSceneInfo* PrimitiveSce
 				ViewMask |= (1 << ViewIndex);
 			}
 
-			bOpaqueRelevance |= ViewRelevance.bOpaqueRelevance || ViewRelevance.bMaskedRelevance;
-			bTranslucentRelevance |= ViewRelevance.HasTranslucency() && !ViewRelevance.bMaskedRelevance;
+			bOpaque |= ViewRelevance.bOpaque || ViewRelevance.bMasked;
+			bTranslucentRelevance |= ViewRelevance.HasTranslucency() && !ViewRelevance.bMasked;
 			bShadowRelevance |= ViewRelevance.bShadowRelevance;
 			bCustomDataRelevance |= ViewRelevance.bUseCustomViewData;
 		}
@@ -1134,7 +1134,7 @@ void FProjectedShadowInfo::AddSubjectPrimitive(FPrimitiveSceneInfo* PrimitiveSce
 			}
 		}
 
-		if (bOpaqueRelevance && bShadowRelevance)
+		if (bOpaque && bShadowRelevance)
 		{
 			const FBoxSphereBounds& Bounds = Proxy->GetBounds();
 			bool bDrawingStaticMeshes = false;
@@ -1747,7 +1747,7 @@ void FSceneRenderer::CreatePerObjectProjectedShadow(
 	bool bShadowIsPotentiallyVisibleNextFrame = false;
 	bool bOpaqueShadowIsVisibleThisFrame = false;
 	bool bSubjectIsVisible = false;
-	bool bOpaqueRelevance = false;
+	bool bOpaque = false;
 	bool bTranslucentRelevance = false;
 	bool bTranslucentShadowIsVisibleThisFrame = false;
 	int32 NumBufferedFrames = FOcclusionQueryHelpers::GetNumBufferedFrames(FeatureLevel);
@@ -1799,7 +1799,7 @@ void FSceneRenderer::CreatePerObjectProjectedShadow(
 		bOpaqueShadowIsVisibleThisFrame |= (bPrimitiveIsShadowRelevant && !bOpaqueShadowIsOccluded);
 		bTranslucentShadowIsVisibleThisFrame |= (bPrimitiveIsShadowRelevant && !bTranslucentShadowIsOccluded);
 		bShadowIsPotentiallyVisibleNextFrame |= bPrimitiveIsShadowRelevant;
-		bOpaqueRelevance |= ViewRelevance.bOpaqueRelevance;
+		bOpaque |= ViewRelevance.bOpaque;
 		bTranslucentRelevance |= ViewRelevance.HasTranslucency();
 	}
 
@@ -1967,7 +1967,7 @@ void FSceneRenderer::CreatePerObjectProjectedShadow(
 			// Use the max resolution if the desired resolution is larger than that
 			const int32 SizeX = MaxDesiredResolution >= MaxShadowResolution ? MaxShadowResolution : (1 << (FMath::CeilLogTwo(MaxDesiredResolution) - 1));
 
-			if (bOpaqueRelevance && bCreateOpaqueObjectShadow && (bOpaqueShadowIsVisibleThisFrame || bShadowIsPotentiallyVisibleNextFrame))
+			if (bOpaque && bCreateOpaqueObjectShadow && (bOpaqueShadowIsVisibleThisFrame || bShadowIsPotentiallyVisibleNextFrame))
 			{
 				// Create a projected shadow for this interaction's shadow.
 				FProjectedShadowInfo* ProjectedShadowInfo = new(FMemStack::Get(),1,16) FProjectedShadowInfo;
@@ -2050,7 +2050,7 @@ void FSceneRenderer::CreatePerObjectProjectedShadow(
 		// If the subject is visible in at least one view, create a preshadow for static primitives shadowing the subject.
 		if (MaxPreFadeAlpha > 1.0f / 256.0f 
 			&& bRenderPreShadow
-			&& bOpaqueRelevance
+			&& bOpaque
 			&& Scene->GetFeatureLevel() >= ERHIFeatureLevel::SM5)
 		{
 			// Round down to the nearest power of two so that resolution changes are always doubling or halving the resolution, which increases filtering stability.
@@ -3509,7 +3509,7 @@ void FSceneRenderer::AllocateShadowDepthTargets(FRHICommandListImmediate& RHICmd
 				const FVisibleLightViewInfo& VisibleLightViewInfo = View.VisibleLightInfos[LightSceneInfo->Id];
 				const FPrimitiveViewRelevance ViewRelevance = VisibleLightViewInfo.ProjectedShadowViewRelevanceMap[ShadowIndex];
 				const bool bHasViewRelevance = (ProjectedShadowInfo->bTranslucentShadow && ViewRelevance.HasTranslucency()) 
-					|| (!ProjectedShadowInfo->bTranslucentShadow && ViewRelevance.bOpaqueRelevance);
+					|| (!ProjectedShadowInfo->bTranslucentShadow && ViewRelevance.bOpaque);
 
 				bShadowIsVisible |= bHasViewRelevance && VisibleLightViewInfo.ProjectedShadowVisibilityMap[ShadowIndex];
 			}
diff --git a/Engine/Source/Runtime/Renderer/Private/ShadowSetupMobile.cpp b/Engine/Source/Runtime/Renderer/Private/ShadowSetupMobile.cpp
--- a/Engine/Source/Runtime/Renderer/Private/ShadowSetupMobile.cpp
+++ b/Engine/Source/Runtime/Renderer/Private/ShadowSetupMobile.cpp
@@ -146,10 +146,10 @@ static bool MobileDetermineStaticMeshesCSMVisibilityStateInner(
 			FVisibleLightViewInfo& VisibleLightViewInfo = View.VisibleLightInfos[LightSceneInfo.Id];
 
 			const FPrimitiveViewRelevance& Relevance = View.PrimitiveViewRelevanceMap[PrimitiveSceneInfo->GetIndex()];
-			const bool bLit = (Relevance.ShadingModelMaskRelevance != (1 << MSM_Unlit));
+			const bool bLit = (Relevance.ShadingModelMask != (1 << MSM_Unlit));
 			bool bCanReceiveDynamicShadow =
 				bLit
-				&& (Relevance.bOpaqueRelevance || Relevance.bMaskedRelevance)
+				&& (Relevance.bOpaque || Relevance.bMasked)
 				&& IsReceiverFunc(PrimitiveBounds.Origin, PrimitiveBounds.BoxExtent, PrimitiveBounds.SphereRadius);
 
 			if (bCanReceiveDynamicShadow)
diff --git a/Engine/Source/Runtime/UMG/Private/Components/WidgetComponent.cpp b/Engine/Source/Runtime/UMG/Private/Components/WidgetComponent.cpp
--- a/Engine/Source/Runtime/UMG/Private/Components/WidgetComponent.cpp
+++ b/Engine/Source/Runtime/UMG/Private/Components/WidgetComponent.cpp
@@ -533,7 +533,7 @@ public:
 		Result.bShadowRelevance = IsShadowCast(View);
 		Result.bTranslucentSelfShadow = bCastVolumetricTranslucentShadow;
 		Result.bEditorPrimitiveRelevance = false;
-		Result.bVelocityRelevance = IsMovable() && Result.bOpaqueRelevance && Result.bRenderInMainPass;
+		Result.bVelocityRelevance = IsMovable() && Result.bOpaque && Result.bRenderInMainPass;
 
 		return Result;
 	}
