Support masking non-connected runtime virtual texture output attributes
Allows decals that only write to a subset of RVT channels

diff --git a/Engine/Source/Runtime/Engine/Classes/Materials/Material.h b/Engine/Source/Runtime/Engine/Classes/Materials/Material.h
--- a/Engine/Source/Runtime/Engine/Classes/Materials/Material.h
+++ b/Engine/Source/Runtime/Engine/Classes/Materials/Material.h
@@ -1770,7 +1770,16 @@ public:
 	void SaveShaderStableKeys(const class ITargetPlatform* TP);
 	ENGINE_API virtual void SaveShaderStableKeysInner(const class ITargetPlatform* TP, const struct FStableShaderKeyAndValue& SaveKeyVal) override;
 
+#if WITH_EDITORONLY_DATA
+	bool HasBaseColorConnected() const { return BaseColor.IsConnected(); }
+	bool HasRoughnessConnected() const { return Roughness.IsConnected(); }
+#else	
+	// Add to runtime data only if we need to call these at runtime
+	bool HasBaseColorConnected() const { check(0); return false; }
+	bool HasRoughnessConnected() const { check(0); return false; }
+#endif 	
 	bool HasNormalConnected() const { return Normal.IsConnected(); }
+	bool HasSpecularConnected() const { return Specular.IsConnected(); }
 	bool HasEmissiveColorConnected() const { return EmissiveColor.IsConnected(); }
 
 #if WITH_EDITOR
diff --git a/Engine/Source/Runtime/Engine/Private/Materials/HLSLMaterialTranslator.cpp b/Engine/Source/Runtime/Engine/Private/Materials/HLSLMaterialTranslator.cpp
--- a/Engine/Source/Runtime/Engine/Private/Materials/HLSLMaterialTranslator.cpp
+++ b/Engine/Source/Runtime/Engine/Private/Materials/HLSLMaterialTranslator.cpp
@@ -915,6 +915,15 @@ bool FHLSLMaterialTranslator::Translate()
 			Errorf(TEXT("Material using the DeferredDecal domain need to use the BlendModel Translucent (this saves performance)"));
 		}
 
+		if (Domain == MD_RuntimeVirtualTexture)
+		{
+			// Add connected material output nodes to the runtime virtual texture output mask.
+			MaterialCompilationOutput.RuntimeVirtualTextureOutputAttributeMask |= Material->HasBaseColorConnected() ? (1 << (uint8)ERuntimeVirtualTextureAttributeType::BaseColor) : 0;
+			MaterialCompilationOutput.RuntimeVirtualTextureOutputAttributeMask |= Material->HasNormalConnected() ? (1 << (uint8)ERuntimeVirtualTextureAttributeType::Normal) : 0;
+			MaterialCompilationOutput.RuntimeVirtualTextureOutputAttributeMask |= Material->HasRoughnessConnected() ? (1 << (uint8)ERuntimeVirtualTextureAttributeType::Roughness) : 0;
+			MaterialCompilationOutput.RuntimeVirtualTextureOutputAttributeMask |= Material->HasSpecularConnected() ? (1 << (uint8)ERuntimeVirtualTextureAttributeType::Specular) : 0;
+		}
+
 		if (MaterialCompilationOutput.bNeedsSceneTextures)
 		{
 			if (Domain != MD_DeferredDecal && Domain != MD_PostProcess)
@@ -1135,7 +1144,7 @@ void FHLSLMaterialTranslator::GetMaterialEnvironment(EShaderPlatform InPlatform,
 		OutEnvironment.SetDefine(TEXT("USES_EYE_ADAPTATION"), TEXT("1"));
 	}
 
-	if (MaterialCompilationOutput.bHasRuntimeVirtualTextureOutput)
+	if (MaterialCompilationOutput.bHasRuntimeVirtualTextureOutputNode)
 	{
 		OutEnvironment.SetDefine(TEXT("VIRTUAL_TEXTURE_OUTPUT"), 1);
 	}
@@ -6743,7 +6752,7 @@ int32 FHLSLMaterialTranslator::CustomOutput(class UMaterialExpressionCustomOutpu
 	return INDEX_NONE;
 }
 
-int32 FHLSLMaterialTranslator::VirtualTextureOutput()
+int32 FHLSLMaterialTranslator::VirtualTextureOutput(uint8 AttributeMask)
 {
 	if (Material->GetMaterialDomain() == MD_RuntimeVirtualTexture)
 	{
@@ -6753,7 +6762,8 @@ int32 FHLSLMaterialTranslator::VirtualTextureOutput()
 	}
 	else
 	{
-		MaterialCompilationOutput.bHasRuntimeVirtualTextureOutput = true;
+		MaterialCompilationOutput.bHasRuntimeVirtualTextureOutputNode |= AttributeMask != 0;
+		MaterialCompilationOutput.RuntimeVirtualTextureOutputAttributeMask |= AttributeMask;
 	}
 
 	// return value is not used
diff --git a/Engine/Source/Runtime/Engine/Private/Materials/HLSLMaterialTranslator.h b/Engine/Source/Runtime/Engine/Private/Materials/HLSLMaterialTranslator.h
--- a/Engine/Source/Runtime/Engine/Private/Materials/HLSLMaterialTranslator.h
+++ b/Engine/Source/Runtime/Engine/Private/Materials/HLSLMaterialTranslator.h
@@ -708,7 +708,7 @@ protected:
 	virtual int32 CustomExpression(class UMaterialExpressionCustom* Custom, TArray<int32>& CompiledInputs) override;
 	virtual int32 CustomOutput(class UMaterialExpressionCustomOutput* Custom, int32 OutputIndex, int32 OutputCode) override;
 
-	virtual int32 VirtualTextureOutput() override;
+	virtual int32 VirtualTextureOutput(uint8 MaterialAttributeMask) override;
 
 #if HANDLE_CUSTOM_OUTPUTS_AS_MATERIAL_ATTRIBUTES
 	/** Used to translate code for custom output attributes such as ClearCoatBottomNormal */
diff --git a/Engine/Source/Runtime/Engine/Private/Materials/MaterialExpressions.cpp b/Engine/Source/Runtime/Engine/Private/Materials/MaterialExpressions.cpp
--- a/Engine/Source/Runtime/Engine/Private/Materials/MaterialExpressions.cpp
+++ b/Engine/Source/Runtime/Engine/Private/Materials/MaterialExpressions.cpp
@@ -1993,33 +1993,34 @@ UMaterialExpressionRuntimeVirtualTextureOutput::UMaterialExpressionRuntimeVirtua
 int32 UMaterialExpressionRuntimeVirtualTextureOutput::Compile(class FMaterialCompiler* Compiler, int32 OutputIndex)
 {
 	int32 CodeInput = INDEX_NONE;
-
-	if (!BaseColor.IsConnected() && !Specular.IsConnected() && !Roughness.IsConnected() && !Normal.IsConnected() && !Opacity.IsConnected())
-	{
-		Compiler->Error(TEXT("No inputs to Runtime Virtual Texture Output."));
-	}
+	uint8 OutputAttributeMask = 0;
 
 	// Order of outputs generates function names GetVirtualTextureOutput{index}
 	// These must match the function names called in VirtualTextureMaterial.usf
 	if (OutputIndex == 0)
 	{
 		CodeInput = BaseColor.IsConnected() ? BaseColor.Compile(Compiler) : Compiler->Constant3(0.f, 0.f, 0.f);
+		OutputAttributeMask |= BaseColor.IsConnected() ? (1 << (uint8)ERuntimeVirtualTextureAttributeType::BaseColor) : 0;
 	}
 	else if (OutputIndex == 1)
 	{
 		CodeInput = Specular.IsConnected() ? Specular.Compile(Compiler) : Compiler->Constant(0.5f);
+		OutputAttributeMask |= Specular.IsConnected() ? (1 << (uint8)ERuntimeVirtualTextureAttributeType::Specular) : 0;
 	}
 	else if (OutputIndex == 2)
 	{
 		CodeInput = Roughness.IsConnected() ? Roughness.Compile(Compiler) : Compiler->Constant(0.5f);
+		OutputAttributeMask |= Roughness.IsConnected() ? (1 << (uint8)ERuntimeVirtualTextureAttributeType::Roughness) : 0;
 	}
 	else if (OutputIndex == 3)
 	{
 		CodeInput = Normal.IsConnected() ? Normal.Compile(Compiler) : Compiler->Constant3(0.f, 0.f, 1.f);
+		OutputAttributeMask |= Normal.IsConnected() ? (1 << (uint8)ERuntimeVirtualTextureAttributeType::Normal) : 0;
 	}
 	else if (OutputIndex == 4)
 	{
 		CodeInput = WorldHeight.IsConnected() ? WorldHeight.Compile(Compiler) : Compiler->Constant(0.f);
+		OutputAttributeMask |= WorldHeight.IsConnected() ? (1 << (uint8)ERuntimeVirtualTextureAttributeType::WorldHeight) : 0;
 	}
 	else if (OutputIndex == 5)
 	{
@@ -2028,9 +2029,10 @@ int32 UMaterialExpressionRuntimeVirtualTextureOutput::Compile(class FMaterialCom
 	else if (OutputIndex == 6)
 	{
 		CodeInput = Mask.IsConnected() ? Mask.Compile(Compiler) : Compiler->Constant(1.f);
+		OutputAttributeMask |= Mask.IsConnected() ? (1 << (uint8)ERuntimeVirtualTextureAttributeType::Mask) : 0;
 	}
 
-	Compiler->VirtualTextureOutput();
+	Compiler->VirtualTextureOutput(OutputAttributeMask);
 	return Compiler->CustomOutput(this, OutputIndex, CodeInput);
 }
 
diff --git a/Engine/Source/Runtime/Engine/Private/Materials/MaterialShared.cpp b/Engine/Source/Runtime/Engine/Private/Materials/MaterialShared.cpp
--- a/Engine/Source/Runtime/Engine/Private/Materials/MaterialShared.cpp
+++ b/Engine/Source/Runtime/Engine/Private/Materials/MaterialShared.cpp
@@ -534,24 +534,26 @@ void FMaterialCompilationOutput::Serialize(FArchive& Ar)
 	Ar << Tmp8;
 #endif
 
+	Ar << RuntimeVirtualTextureOutputAttributeMask;
+
 	uint8 PackedFlags = 0;
-	PackedFlags |= (bUsesEyeAdaptation				<< 0);
-	PackedFlags |= (bModifiesMeshPosition			<< 1);
-	PackedFlags |= (bUsesWorldPositionOffset		<< 2);
-	PackedFlags |= (bUsesGlobalDistanceField		<< 3);
-	PackedFlags |= (bUsesPixelDepthOffset			<< 4);
-	PackedFlags |= (bUsesDistanceCullFade			<< 5);
-	PackedFlags |= (bHasRuntimeVirtualTextureOutput << 6);
+	PackedFlags |= (bUsesEyeAdaptation << 0);
+	PackedFlags |= (bModifiesMeshPosition << 1);
+	PackedFlags |= (bUsesWorldPositionOffset << 2);
+	PackedFlags |= (bUsesGlobalDistanceField << 3);
+	PackedFlags |= (bUsesPixelDepthOffset << 4);
+	PackedFlags |= (bUsesDistanceCullFade << 5);
+	PackedFlags |= (bHasRuntimeVirtualTextureOutputNode << 6);
 
 	Ar << PackedFlags;
 
-	bUsesEyeAdaptation				= (PackedFlags >> 0) & 1;
-	bModifiesMeshPosition			= (PackedFlags >> 1) & 1;
-	bUsesWorldPositionOffset		= (PackedFlags >> 2) & 1;
-	bUsesGlobalDistanceField		= (PackedFlags >> 3) & 1;
-	bUsesPixelDepthOffset			= (PackedFlags >> 4) & 1;
-	bUsesDistanceCullFade			= (PackedFlags >> 5) & 1;
-	bHasRuntimeVirtualTextureOutput = (PackedFlags >> 6) & 1;
+	bUsesEyeAdaptation = (PackedFlags >> 0) & 1;
+	bModifiesMeshPosition = (PackedFlags >> 1) & 1;
+	bUsesWorldPositionOffset = (PackedFlags >> 2) & 1;
+	bUsesGlobalDistanceField = (PackedFlags >> 3) & 1;
+	bUsesPixelDepthOffset = (PackedFlags >> 4) & 1;
+	bUsesDistanceCullFade = (PackedFlags >> 5) & 1;
+	bHasRuntimeVirtualTextureOutputNode = (PackedFlags >> 6) & 1;
 }
 
 void FMaterial::GetShaderMapId(EShaderPlatform Platform, FMaterialShaderMapId& OutId) const
@@ -901,14 +903,14 @@ bool FMaterial::MaterialUsesSceneDepthLookup_GameThread() const
 	return GameThreadShaderMap.GetReference() ? GameThreadShaderMap->UsesSceneDepthLookup() : false;
 }
 
-bool FMaterial::HasRuntimeVirtualTextureOutput_RenderThread() const
+bool FMaterial::UsesCustomDepthStencil_GameThread() const
 {
-	return RenderingThreadShaderMap ? RenderingThreadShaderMap->HasRuntimeVirtualTextureOutput() : false;
+	return GameThreadShaderMap.GetReference() ? (GameThreadShaderMap->UsesSceneTexture(PPI_CustomDepth) || GameThreadShaderMap->UsesSceneTexture(PPI_CustomStencil)) : false;
 }
 
-bool FMaterial::UsesCustomDepthStencil_GameThread() const
+uint8 FMaterial::GetRuntimeVirtualTextureOutputAttibuteMask_RenderThread() const
 {
-	return GameThreadShaderMap.GetReference() ? (GameThreadShaderMap->UsesSceneTexture(PPI_CustomDepth) || GameThreadShaderMap->UsesSceneTexture(PPI_CustomStencil)) : false;
+	return RenderingThreadShaderMap ? RenderingThreadShaderMap->GetRuntimeVirtualTextureOutputAttributeMask() : 0;
 }
 
 FMaterialShaderMap* FMaterial::GetRenderingThreadShaderMap() const 
@@ -1411,11 +1413,26 @@ uint32 FMaterialResource::GetMaterialDecalResponse() const
 	return Material->GetMaterialDecalResponse();
 }
 
+bool FMaterialResource::HasBaseColorConnected() const
+{
+	return HasMaterialAttributesConnected() || Material->HasBaseColorConnected();
+}
+
 bool FMaterialResource::HasNormalConnected() const
 {
 	return HasMaterialAttributesConnected() || Material->HasNormalConnected();
 }
 
+bool FMaterialResource::HasRoughnessConnected() const
+{
+	return HasMaterialAttributesConnected() || Material->HasRoughnessConnected();
+}
+
+bool FMaterialResource::HasSpecularConnected() const
+{
+	return HasMaterialAttributesConnected() || Material->HasSpecularConnected();
+}
+
 bool FMaterialResource::HasEmissiveColorConnected() const
 {
 	return HasMaterialAttributesConnected() || Material->HasEmissiveColorConnected();
@@ -1463,8 +1480,7 @@ uint32 FMaterialResource::GetStencilCompare() const
 
 bool FMaterialResource::HasRuntimeVirtualTextureOutput() const
 {
-	// Slow check used only for ShouldCompilePermutation() calls. 
-	// Runtime calls from render thread can use faster FMaterial::HasRuntimeVirtualTextureOutput_RenderThread()
+	// Slow check. Used for ShouldCompilePermutation() calls. 
 	for (UMaterialExpression* Expression : Material->Expressions)
 	{
 		if (Expression->IsA(UMaterialExpressionRuntimeVirtualTextureOutput::StaticClass()))
diff --git a/Engine/Source/Runtime/Engine/Private/ShaderDerivedDataVersion.h b/Engine/Source/Runtime/Engine/Private/ShaderDerivedDataVersion.h
--- a/Engine/Source/Runtime/Engine/Private/ShaderDerivedDataVersion.h
+++ b/Engine/Source/Runtime/Engine/Private/ShaderDerivedDataVersion.h
@@ -11,4 +11,4 @@
 // In case of merge conflicts with DDC versions, you MUST generate a new GUID and set this new
 // guid as version
 #define GLOBALSHADERMAP_DERIVEDDATA_VER			TEXT("A85F9B4CA9264C3CBE27C6356D371896")
-#define MATERIALSHADERMAP_DERIVEDDATA_VER		TEXT("16B6C209A72740C49B69059E25CF3CCB")
+#define MATERIALSHADERMAP_DERIVEDDATA_VER		TEXT("DBB562AD91CE4620B207619E5C81A619")
diff --git a/Engine/Source/Runtime/Engine/Public/MaterialCompiler.h b/Engine/Source/Runtime/Engine/Public/MaterialCompiler.h
--- a/Engine/Source/Runtime/Engine/Public/MaterialCompiler.h
+++ b/Engine/Source/Runtime/Engine/Public/MaterialCompiler.h
@@ -323,7 +323,7 @@ public:
 
 	virtual int32 CustomExpression(class UMaterialExpressionCustom* Custom, TArray<int32>& CompiledInputs) = 0;
 	virtual int32 CustomOutput(class UMaterialExpressionCustomOutput* Custom, int32 OutputIndex, int32 OutputCode) = 0;
-	virtual int32 VirtualTextureOutput() = 0;
+	virtual int32 VirtualTextureOutput(uint8 AttributeMask) = 0;
 
 	virtual int32 DDX(int32 X) = 0;
 	virtual int32 DDY(int32 X) = 0;
@@ -583,7 +583,7 @@ public:
 
 	virtual int32 CustomExpression(class UMaterialExpressionCustom* Custom, TArray<int32>& CompiledInputs) override { return Compiler->CustomExpression(Custom,CompiledInputs); }
 	virtual int32 CustomOutput(class UMaterialExpressionCustomOutput* Custom, int32 OutputIndex, int32 OutputCode) override{ return Compiler->CustomOutput(Custom, OutputIndex, OutputCode); }
-	virtual int32 VirtualTextureOutput() override { return Compiler->VirtualTextureOutput(); }
+	virtual int32 VirtualTextureOutput(uint8 AttributeMask) override { return Compiler->VirtualTextureOutput(AttributeMask); }
 
 	virtual int32 DDX(int32 X) override { return Compiler->DDX(X); }
 	virtual int32 DDY(int32 X) override { return Compiler->DDY(X); }
diff --git a/Engine/Source/Runtime/Engine/Public/MaterialShared.h b/Engine/Source/Runtime/Engine/Public/MaterialShared.h
--- a/Engine/Source/Runtime/Engine/Public/MaterialShared.h
+++ b/Engine/Source/Runtime/Engine/Public/MaterialShared.h
@@ -547,6 +547,7 @@ public:
 		NumUsedUVScalars(0),
 		NumUsedCustomInterpolatorScalars(0),
 #endif
+		RuntimeVirtualTextureOutputAttributeMask(0),
 		bNeedsSceneTextures(false),
 		bUsesEyeAdaptation(false),
 		bModifiesMeshPosition(false),
@@ -554,7 +555,7 @@ public:
 		bUsesGlobalDistanceField(false),
 		bUsesPixelDepthOffset(false),
 		bUsesDistanceCullFade(false),
-		bHasRuntimeVirtualTextureOutput(false)
+		bHasRuntimeVirtualTextureOutputNode(false)
 	{}
 
 	ENGINE_API void Serialize(FArchive& Ar);
@@ -582,6 +583,9 @@ public:
 	uint8 NumUsedCustomInterpolatorScalars;
 #endif // WITH_EDITOR
 
+	/** Bitfield of runtime virtual texture output attributes. */
+	uint8 RuntimeVirtualTextureOutputAttributeMask;
+
 	/** true if the material needs the scenetexture lookups. */
 	uint8 bNeedsSceneTextures : 1;
 
@@ -603,8 +607,8 @@ public:
 	/** true if the material uses distance cull fade */
 	uint8 bUsesDistanceCullFade : 1;
 
-	/** true if the material supports virtual texture output */
-	uint8 bHasRuntimeVirtualTextureOutput : 1;
+	/** true if the material writes to a runtime virtual texture custom output node. */
+	uint8 bHasRuntimeVirtualTextureOutputNode : 1;
 
 	/** Indicates whether the material uses scene color. */
 	ENGINE_API bool RequiresSceneColorCopy() const { return IsSceneTextureUsed(PPI_SceneColor); }
@@ -1113,7 +1117,6 @@ public:
 	bool UsesSceneDepthLookup() const { return MaterialCompilationOutput.UsesSceneDepthLookup(); }
 	bool UsesVelocitySceneTexture() const { return MaterialCompilationOutput.UsesVelocitySceneTexture(); }
 	bool UsesDistanceCullFade() const { return MaterialCompilationOutput.bUsesDistanceCullFade; }
-	bool HasRuntimeVirtualTextureOutput() const { return MaterialCompilationOutput.bHasRuntimeVirtualTextureOutput; }
 #if WITH_EDITOR
 	uint32 GetNumUsedUVScalars() const { return MaterialCompilationOutput.NumUsedUVScalars; }
 	uint32 GetNumUsedCustomInterpolatorScalars() const { return MaterialCompilationOutput.NumUsedCustomInterpolatorScalars; }
@@ -1121,6 +1124,7 @@ public:
 	uint32 GetEstimatedNumVirtualTextureLookups() const { return MaterialCompilationOutput.EstimatedNumVirtualTextureLookups; }
 #endif
 	uint32 GetNumVirtualTextureStacks() const { return MaterialCompilationOutput.UniformExpressionSet.VTStacks.Num(); }
+	uint8 GetRuntimeVirtualTextureOutputAttributeMask() const { return MaterialCompilationOutput.RuntimeVirtualTextureOutputAttributeMask; }
 	bool UsesSceneTexture(uint32 TexId) const { return (MaterialCompilationOutput.UsedSceneTextures & (1ull << TexId)) != 0; }
 
 	bool IsValidForRendering(bool bFailOnInvalid = false) const
@@ -1513,7 +1517,10 @@ public:
 	virtual bool HasMaterialAttributesConnected() const { return false; }
 	virtual uint32 GetDecalBlendMode() const { return 0; }
 	virtual uint32 GetMaterialDecalResponse() const { return 0; }
+	virtual bool HasBaseColorConnected() const { return false; }
 	virtual bool HasNormalConnected() const { return false; }
+	virtual bool HasRoughnessConnected() const { return false; }
+	virtual bool HasSpecularConnected() const { return false; }
 	virtual bool HasEmissiveColorConnected() const { return false; }
 	virtual bool RequiresSynchronousCompilation() const { return false; };
 	virtual bool IsDefaultMaterial() const { return false; };
@@ -1622,15 +1629,15 @@ public:
 	ENGINE_API bool MaterialUsesSceneDepthLookup_RenderThread() const;
 	ENGINE_API bool MaterialUsesSceneDepthLookup_GameThread() const;
 
-	/** Does the material have a runtime virtual texture output node. */
-	ENGINE_API bool HasRuntimeVirtualTextureOutput_RenderThread() const;
-
 	/** Does the material use CustomDepth or CustomStencil lookup */
 	ENGINE_API bool UsesCustomDepthStencil_GameThread() const;
 
 	/** Note: This function is only intended for use in deciding whether or not shader permutations are required before material translation occurs. */
 	ENGINE_API bool MaterialMayModifyMeshPosition() const;
 
+	/** Get the runtime virtual texture output attribute mask for the material. */
+	ENGINE_API uint8 GetRuntimeVirtualTextureOutputAttibuteMask_RenderThread() const;
+
 	class FMaterialShaderMap* GetGameThreadShaderMap() const 
 	{ 
 		checkSlow(IsInGameThread() || IsInAsyncLoadingThread());
@@ -2259,8 +2266,10 @@ public:
 	ENGINE_API virtual enum ERefractionMode GetRefractionMode() const override;
 	ENGINE_API virtual uint32 GetDecalBlendMode() const override;
 	ENGINE_API virtual uint32 GetMaterialDecalResponse() const override;
+	ENGINE_API virtual bool HasBaseColorConnected() const override;
 	ENGINE_API virtual bool HasNormalConnected() const override;
-	ENGINE_API virtual bool HasEmissiveColorConnected() const override;
+	ENGINE_API virtual bool HasRoughnessConnected() const override;
+	ENGINE_API virtual bool HasSpecularConnected() const override;	ENGINE_API virtual bool HasEmissiveColorConnected() const override;
 	ENGINE_API virtual FMaterialShadingModelField GetShadingModels() const override;
 	ENGINE_API virtual bool IsShadingModelFromMaterialExpression() const override;
 	ENGINE_API virtual enum ETranslucencyLightingMode GetTranslucencyLightingMode() const override;
diff --git a/Engine/Source/Runtime/Engine/Public/VT/RuntimeVirtualTextureEnum.h b/Engine/Source/Runtime/Engine/Public/VT/RuntimeVirtualTextureEnum.h
--- a/Engine/Source/Runtime/Engine/Public/VT/RuntimeVirtualTextureEnum.h
+++ b/Engine/Source/Runtime/Engine/Public/VT/RuntimeVirtualTextureEnum.h
@@ -8,6 +8,24 @@
 /** Maximum number of texture layers we will have in a runtime virtual texture. Increase if we add a ERuntimeVirtualTextureMaterialType with more layers. */
 namespace RuntimeVirtualTexture { enum { MaxTextureLayers = 3 }; }
 
+/** 
+ * Enumeration of all runtime virtual texture material attributes.
+ * These can be combined to form full ERuntimeVirtualTextureMaterialType layouts.
+ */
+enum class ERuntimeVirtualTextureAttributeType : uint8
+{
+	BaseColor,
+	Normal,
+	Roughness,
+	Specular,
+	Mask,
+	WorldHeight,
+
+	Count
+};
+
+static_assert((uint32)ERuntimeVirtualTextureAttributeType::Count <= 8u, "ERuntimeVirtualTextureAttributeType can no longer be used to create 8bit masks.");
+
 /** 
  * Enumeration of virtual texture stack layouts to support. 
  * Extend this enumeration with other layouts as required. For example we will probably want to add a displacement texture option.
@@ -42,7 +60,6 @@ enum class ERuntimeVirtualTextureMainPassType : uint8
 };
 
 /** Enumeration of runtime virtual texture debug modes. */
-UENUM()
 enum class ERuntimeVirtualTextureDebugType
 {
 	None,
diff --git a/Engine/Source/Runtime/Renderer/Private/VT/RuntimeVirtualTextureRender.cpp b/Engine/Source/Runtime/Renderer/Private/VT/RuntimeVirtualTextureRender.cpp
--- a/Engine/Source/Runtime/Renderer/Private/VT/RuntimeVirtualTextureRender.cpp
+++ b/Engine/Source/Runtime/Renderer/Private/VT/RuntimeVirtualTextureRender.cpp
@@ -72,7 +72,7 @@ namespace RuntimeVirtualTexture
 			OutEnvironment.SetDefine(TEXT("OUT_BASECOLOR"), 1);
 		}
 
-		static FRHIBlendState* GetBlendState()
+		static FRHIBlendState* GetBlendState(uint8 OutputAttributeMask)
 		{
 			return TStaticBlendState< CW_RGBA, BO_Add, BF_One, BF_InverseSourceAlpha, BO_Add, BF_Zero, BF_One >::GetRHI();
 		}
@@ -81,18 +81,89 @@ namespace RuntimeVirtualTexture
 	/** Specialization for ERuntimeVirtualTextureMaterialType::BaseColor_Normal_Specular */
 	class FMaterialPolicy_BaseColorNormalSpecular
 	{
+	private:
+		/** Compile time helper to build blend state from the connected output attribute mask. */
+		static constexpr EColorWriteMask GetColorMaskFromAttributeMask(uint8 AttributeMask, uint8 RenderTargetIndex)
+		{
+			// Color mask in the output render targets for each of the relevant attributes in ERuntimeVirtualTextureAttributeType
+			const EColorWriteMask AttributeMasks[][3] = {
+				{ CW_RGBA, CW_NONE, CW_NONE }, // BaseColor
+				{ CW_NONE, EColorWriteMask(CW_RED | CW_GREEN | CW_ALPHA), EColorWriteMask(CW_BLUE | CW_ALPHA) }, // Normal
+				{ CW_NONE, CW_NONE, EColorWriteMask(CW_GREEN | CW_ALPHA) }, // Roughness
+				{ CW_NONE, CW_NONE, EColorWriteMask(CW_RED | CW_ALPHA) }, // Specular
+				{ CW_NONE, EColorWriteMask(CW_BLUE | CW_ALPHA), CW_NONE }, // Mask
+			};
+
+			// Combine the color masks for this AttributeMask
+			EColorWriteMask ColorWriteMask = CW_NONE;
+			for (int32 i = 0; i < 5; ++i)
+			{
+				if (AttributeMask & (1 << i))
+				{
+					ColorWriteMask = EColorWriteMask(ColorWriteMask | AttributeMasks[i][RenderTargetIndex]);
+				}
+			}
+			return ColorWriteMask;
+		}
+
+		/** Helper to convert the connected output attribute mask to a blend state with a color mask for these attributes. */
+		template< uint32 AttributeMask >
+		static FRHIBlendState* TGetBlendStateFromAttributeMask()
+		{
+			return TStaticBlendState< 
+				GetColorMaskFromAttributeMask(AttributeMask, 0), BO_Add, BF_One, BF_InverseSourceAlpha, BO_Add, BF_Zero, BF_One,
+				GetColorMaskFromAttributeMask(AttributeMask, 1), BO_Add, BF_One, BF_InverseSourceAlpha, BO_Add, BF_Zero, BF_One,
+				GetColorMaskFromAttributeMask(AttributeMask, 2), BO_Add, BF_One, BF_InverseSourceAlpha, BO_Add, BF_Zero, BF_One	>::GetRHI();
+		}
+
+		/** Runtime conversion of attribute mask to static blend state. */
+		static FRHIBlendState* GetBlendStateImpl(uint8 AttributeMask)
+		{
+			// We have 5 relevant bits in the attribute mask. Any more and this would get painful...
+			switch (AttributeMask & 0x1f)
+			{
+			case 1: return TGetBlendStateFromAttributeMask<1>();
+			case 2: return TGetBlendStateFromAttributeMask<2>();
+			case 3: return TGetBlendStateFromAttributeMask<3>();
+			case 4: return TGetBlendStateFromAttributeMask<4>();
+			case 5: return TGetBlendStateFromAttributeMask<5>();
+			case 6: return TGetBlendStateFromAttributeMask<6>();
+			case 7: return TGetBlendStateFromAttributeMask<7>();
+			case 8: return TGetBlendStateFromAttributeMask<8>();
+			case 9: return TGetBlendStateFromAttributeMask<9>();
+			case 10: return TGetBlendStateFromAttributeMask<10>();
+			case 11: return TGetBlendStateFromAttributeMask<11>();
+			case 12: return TGetBlendStateFromAttributeMask<12>();
+			case 13: return TGetBlendStateFromAttributeMask<13>();
+			case 14: return TGetBlendStateFromAttributeMask<14>();
+			case 15: return TGetBlendStateFromAttributeMask<15>();
+			case 16: return TGetBlendStateFromAttributeMask<16>();
+			case 17: return TGetBlendStateFromAttributeMask<17>();
+			case 18: return TGetBlendStateFromAttributeMask<18>();
+			case 19: return TGetBlendStateFromAttributeMask<19>();
+			case 21: return TGetBlendStateFromAttributeMask<21>();
+			case 22: return TGetBlendStateFromAttributeMask<22>();
+			case 23: return TGetBlendStateFromAttributeMask<23>();
+			case 24: return TGetBlendStateFromAttributeMask<24>();
+			case 25: return TGetBlendStateFromAttributeMask<25>();
+			case 26: return TGetBlendStateFromAttributeMask<26>();
+			case 27: return TGetBlendStateFromAttributeMask<27>();
+			case 28: return TGetBlendStateFromAttributeMask<28>();
+			case 29: return TGetBlendStateFromAttributeMask<29>();
+			case 30: return TGetBlendStateFromAttributeMask<30>();
+			default: return TGetBlendStateFromAttributeMask<31>();
+			}
+		}
+
 	public:
 		static void ModifyCompilationEnvironment(FShaderCompilerEnvironment& OutEnvironment)
 		{
 			OutEnvironment.SetDefine(TEXT("OUT_BASECOLOR_NORMAL_SPECULAR"), 1);
 		}
 
-		static FRHIBlendState* GetBlendState()
+		static FRHIBlendState* GetBlendState(uint8 OutputAttributeMask)
 		{
-			return TStaticBlendState<
-				CW_RGBA, BO_Add, BF_One, BF_InverseSourceAlpha, BO_Add, BF_Zero, BF_One,
-				CW_RGBA, BO_Add, BF_One, BF_InverseSourceAlpha, BO_Add, BF_Zero, BF_One,
-				CW_RGBA, BO_Add, BF_One, BF_InverseSourceAlpha, BO_Add, BF_Zero, BF_One >::GetRHI();
+			return GetBlendStateImpl(OutputAttributeMask);
 		}
 	};
 
@@ -106,7 +177,7 @@ namespace RuntimeVirtualTexture
 			OutEnvironment.SetRenderTargetOutputFormat(0, PF_R32_FLOAT);
 		}
 
-		static FRHIBlendState* GetBlendState()
+		static FRHIBlendState* GetBlendState(uint8 OutputAttributeMask)
 		{
 			return TStaticBlendState< CW_RED, BO_Max, BF_One, BF_One, BO_Add, BF_One, BF_One >::GetRHI();
 		}
@@ -185,6 +256,7 @@ namespace RuntimeVirtualTexture
 			const FMeshBatch& MeshBatch,
 			uint64 BatchElementMask,
 			int32 StaticMeshId,
+			uint8 OutputAttributeMask,
 			const FPrimitiveSceneProxy* RESTRICT PrimitiveSceneProxy,
 			const FMaterialRenderProxy& RESTRICT MaterialRenderProxy,
 			const FMaterial& RESTRICT MaterialResource)
@@ -200,7 +272,7 @@ namespace RuntimeVirtualTexture
 			Shaders.VertexShader = MaterialResource.GetShader< FShader_VirtualTextureMaterialDraw_VS< MaterialPolicy > >(VertexFactory->GetType());
 			Shaders.PixelShader = MaterialResource.GetShader< FShader_VirtualTextureMaterialDraw_PS< MaterialPolicy > >(VertexFactory->GetType());
 
-			DrawRenderState.SetBlendState(MaterialPolicy::GetBlendState());
+			DrawRenderState.SetBlendState(MaterialPolicy::GetBlendState(OutputAttributeMask));
 
 			ERasterizerFillMode MeshFillMode = ComputeMeshFillMode(MeshBatch, MaterialResource);
 			ERasterizerCullMode MeshCullMode = ComputeMeshCullMode(MeshBatch, MaterialResource);
@@ -236,21 +308,22 @@ namespace RuntimeVirtualTexture
 				const FMaterialRenderProxy* FallbackMaterialRenderProxyPtr = nullptr;
 				const FMaterial& Material = MeshBatch.MaterialRenderProxy->GetMaterialWithFallback(FeatureLevel, FallbackMaterialRenderProxyPtr);
 				const FMaterialRenderProxy& MaterialRenderProxy = FallbackMaterialRenderProxyPtr ? *FallbackMaterialRenderProxyPtr : *MeshBatch.MaterialRenderProxy;
+				const uint8 OutputAttributeMask = Material.GetRuntimeVirtualTextureOutputAttibuteMask_RenderThread();
 
-				if (Material.GetMaterialDomain() == MD_RuntimeVirtualTexture || Material.HasRuntimeVirtualTextureOutput_RenderThread())
+				if (OutputAttributeMask != 0)
 				{
 					switch ((ERuntimeVirtualTextureMaterialType)MeshBatch.RuntimeVirtualTextureMaterialType)
 					{
 					case ERuntimeVirtualTextureMaterialType::BaseColor:
-						Process<FMaterialPolicy_BaseColor>(MeshBatch, BatchElementMask, StaticMeshId, PrimitiveSceneProxy, MaterialRenderProxy, Material);
+						Process<FMaterialPolicy_BaseColor>(MeshBatch, BatchElementMask, StaticMeshId, OutputAttributeMask, PrimitiveSceneProxy, MaterialRenderProxy, Material);
 						break;
 					case ERuntimeVirtualTextureMaterialType::BaseColor_Normal_Specular:
 					case ERuntimeVirtualTextureMaterialType::BaseColor_Normal_Specular_YCoCg:
 					case ERuntimeVirtualTextureMaterialType::BaseColor_Normal_Specular_Mask_YCoCg:
-						Process<FMaterialPolicy_BaseColorNormalSpecular>(MeshBatch, BatchElementMask, StaticMeshId, PrimitiveSceneProxy, MaterialRenderProxy, Material);
+						Process<FMaterialPolicy_BaseColorNormalSpecular>(MeshBatch, BatchElementMask, StaticMeshId, OutputAttributeMask, PrimitiveSceneProxy, MaterialRenderProxy, Material);
 						break;
 					case ERuntimeVirtualTextureMaterialType::WorldHeight:
-						Process<FMaterialPolicy_WorldHeight>(MeshBatch, BatchElementMask, StaticMeshId, PrimitiveSceneProxy, MaterialRenderProxy, Material);
+						Process<FMaterialPolicy_WorldHeight>(MeshBatch, BatchElementMask, StaticMeshId, OutputAttributeMask, PrimitiveSceneProxy, MaterialRenderProxy, Material);
 						break;
 					default:
 						break;
