Added an option that moves light octree traversal and FLightPrimitiveInteraction creation to an async task. The task overlaps with a large portion of InitViews (e.g. ComputeViewVisibility) and sync doesn't happen until InitDynamicShadows. It can even overlap with RenderPrePass r.DoInitViewsLightingAfterPrepass is set. This option reduced RT time by 0.53 ms on a type of consoles in a 10-run replay A/B test.

diff --git a/Engine/Source/Runtime/Engine/Public/SceneInterface.h b/Engine/Source/Runtime/Engine/Public/SceneInterface.h
--- a/Engine/Source/Runtime/Engine/Public/SceneInterface.h
+++ b/Engine/Source/Runtime/Engine/Public/SceneInterface.h
@@ -66,7 +66,7 @@ public:
 	/**
 	* Updates all primitive scene info additions, remobals and translation changes
 	*/
-	virtual void UpdateAllPrimitiveSceneInfos(FRHICommandListImmediate& RHICmdList) = 0;
+	virtual void UpdateAllPrimitiveSceneInfos(FRHICommandListImmediate& RHICmdList, bool bAsyncCreateLPIs = false) = 0;
 	/** 
 	 * Updates the transform of a primitive which has already been added to the scene. 
 	 * 
diff --git a/Engine/Source/Runtime/RenderCore/Private/RenderCore.cpp b/Engine/Source/Runtime/RenderCore/Private/RenderCore.cpp
--- a/Engine/Source/Runtime/RenderCore/Private/RenderCore.cpp
+++ b/Engine/Source/Runtime/RenderCore/Private/RenderCore.cpp
@@ -176,6 +176,7 @@ DEFINE_STAT(STAT_RemoveScenePrimitiveTime);
 DEFINE_STAT(STAT_AddScenePrimitiveRenderThreadTime);
 DEFINE_STAT(STAT_UpdateScenePrimitiveRenderThreadTime);
 DEFINE_STAT(STAT_UpdatePrimitiveTransformRenderThreadTime);
+DEFINE_STAT(STAT_FlushAsyncLPICreation);
 
 DEFINE_STAT(STAT_RemoveScenePrimitiveGT);
 DEFINE_STAT(STAT_AddScenePrimitiveGT);
diff --git a/Engine/Source/Runtime/RenderCore/Public/RenderCore.h b/Engine/Source/Runtime/RenderCore/Public/RenderCore.h
--- a/Engine/Source/Runtime/RenderCore/Public/RenderCore.h
+++ b/Engine/Source/Runtime/RenderCore/Public/RenderCore.h
@@ -127,6 +127,7 @@ DECLARE_CYCLE_STAT_EXTERN(TEXT("Update CPU Skin"),STAT_CPUSkinUpdateRTTime,STATG
 DECLARE_CYCLE_STAT_EXTERN(TEXT("Update GPU Skin"),STAT_GPUSkinUpdateRTTime,STATGROUP_SceneUpdate, RENDERCORE_API);
 DECLARE_CYCLE_STAT_EXTERN(TEXT("Update particles"),STAT_ParticleUpdateRTTime,STATGROUP_SceneUpdate, RENDERCORE_API);
 DECLARE_CYCLE_STAT_EXTERN(TEXT("Update Influence Weights"),STAT_InfluenceWeightsUpdateRTTime,STATGROUP_SceneUpdate, RENDERCORE_API);
+DECLARE_CYCLE_STAT_EXTERN(TEXT("Flush async LPI creation"), STAT_FlushAsyncLPICreation, STATGROUP_SceneUpdate, RENDERCORE_API);
 
 DECLARE_CYCLE_STAT_EXTERN(TEXT("RemovePrimitive (GT)"),STAT_RemoveScenePrimitiveGT,STATGROUP_SceneUpdate, RENDERCORE_API);
 DECLARE_CYCLE_STAT_EXTERN(TEXT("AddPrimitive (GT)"),STAT_AddScenePrimitiveGT,STATGROUP_SceneUpdate, RENDERCORE_API);
diff --git a/Engine/Source/Runtime/Renderer/Private/DeferredShadingRenderer.cpp b/Engine/Source/Runtime/Renderer/Private/DeferredShadingRenderer.cpp
--- a/Engine/Source/Runtime/Renderer/Private/DeferredShadingRenderer.cpp
+++ b/Engine/Source/Runtime/Renderer/Private/DeferredShadingRenderer.cpp
@@ -1204,7 +1204,7 @@ static TAutoConsoleVariable<float> CVarStallInitViews(
 
 void FDeferredShadingSceneRenderer::Render(FRHICommandListImmediate& RHICmdList)
 {
-	Scene->UpdateAllPrimitiveSceneInfos(RHICmdList);
+	Scene->UpdateAllPrimitiveSceneInfos(RHICmdList, true);
 
 	check(RHICmdList.IsOutsideRenderPass());
 
diff --git a/Engine/Source/Runtime/Renderer/Private/LightSceneInfo.cpp b/Engine/Source/Runtime/Renderer/Private/LightSceneInfo.cpp
--- a/Engine/Source/Runtime/Renderer/Private/LightSceneInfo.cpp
+++ b/Engine/Source/Runtime/Renderer/Private/LightSceneInfo.cpp
@@ -139,6 +139,8 @@ void FLightSceneInfo::RemoveFromScene()
 void FLightSceneInfo::Detach()
 {
 	check(IsInRenderingThread());
+	
+	Scene->FlushAsyncLightPrimitiveInteractionCreation();
 
 	// implicit linked list. The destruction will update this "head" pointer to the next item in the list.
 	while(DynamicInteractionOftenMovingPrimitiveList)
diff --git a/Engine/Source/Runtime/Renderer/Private/PrimitiveSceneInfo.cpp b/Engine/Source/Runtime/Renderer/Private/PrimitiveSceneInfo.cpp
--- a/Engine/Source/Runtime/Renderer/Private/PrimitiveSceneInfo.cpp
+++ b/Engine/Source/Runtime/Renderer/Private/PrimitiveSceneInfo.cpp
@@ -505,7 +505,7 @@ int32 FPrimitiveSceneInfo::UpdateStaticLightingBuffer()
 	return LCIs.Num();
 }
 
-void FPrimitiveSceneInfo::AddToScene(FRHICommandListImmediate& RHICmdList, bool bUpdateStaticDrawLists, bool bAddToStaticDrawLists)
+void FPrimitiveSceneInfo::AddToScene(FRHICommandListImmediate& RHICmdList, bool bUpdateStaticDrawLists, bool bAddToStaticDrawLists, bool bAsyncCreateLPIs)
 {
 	check(IsInRenderingThread());
 
@@ -661,21 +661,8 @@ void FPrimitiveSceneInfo::AddToScene(FRHICommandListImmediate& RHICmdList, bool
 		LodInfo.CullValue = LodInfo.CullMethod == 0 ? Proxy->GetVirtualTextureCullMips() : Proxy->GetVirtualTextureMinCoverage();
 	}
 
-	{
-		FMemMark MemStackMark(FMemStack::Get());
-
-		// Find lights that affect the primitive in the light octree.
-		for (FSceneLightOctree::TConstElementBoxIterator<SceneRenderingAllocator> LightIt(Scene->LightOctree, Proxy->GetBounds().GetBox());
-			LightIt.HasPendingElements();
-			LightIt.Advance())
-		{
-			const FLightSceneInfoCompact& LightSceneInfoCompact = LightIt.GetCurrentElement();
-			if (LightSceneInfoCompact.AffectsPrimitive(CompactPrimitiveSceneInfo.Bounds, CompactPrimitiveSceneInfo.Proxy))
-			{
-				FLightPrimitiveInteraction::Create(LightSceneInfoCompact.LightSceneInfo,this);
-			}
-		}
-	}
+	// Find lights that affect the primitive in the light octree.
+	Scene->CreateLightPrimitiveInteractionsForPrimitive(this, bAsyncCreateLPIs);
 
 	INC_MEMORY_STAT_BY(STAT_PrimitiveInfoMemory, sizeof(*this) + StaticMeshes.GetAllocatedSize() + StaticMeshRelevances.GetAllocatedSize() + Proxy->GetMemoryFootprint());
 }
diff --git a/Engine/Source/Runtime/Renderer/Private/RendererScene.cpp b/Engine/Source/Runtime/Renderer/Private/RendererScene.cpp
--- a/Engine/Source/Runtime/Renderer/Private/RendererScene.cpp
+++ b/Engine/Source/Runtime/Renderer/Private/RendererScene.cpp
@@ -94,6 +94,13 @@ static TAutoConsoleVariable<int32> CVarBasePassWriteDepthEvenWithFullPrepass(
 	TEXT("0 to allow a readonly base pass, which skips an MSAA depth resolve, and allows masked materials to get EarlyZ (writing to depth while doing clip() disables EarlyZ) (default)\n")
 	TEXT("1 to force depth writes in the base pass.  Useful for debugging when the prepass and base pass don't match what they render."));
 
+static int32 GAsyncCreateLightPrimitiveInteractions = 1;
+static FAutoConsoleVariableRef CVarAsyncCreateLightPrimitiveInteractions(
+	TEXT("r.AsyncCreateLightPrimitiveInteractions"),
+	GAsyncCreateLightPrimitiveInteractions,
+	TEXT("Whether to create LPIs asynchronously."),
+	ECVF_RenderThreadSafe);
+
 DECLARE_CYCLE_STAT(TEXT("DeferredShadingSceneRenderer MotionBlurStartFrame"), STAT_FDeferredShadingSceneRenderer_MotionBlurStartFrame, STATGROUP_SceneRendering);
 
 IMPLEMENT_GLOBAL_SHADER_PARAMETER_STRUCT(FDistanceCullFadeUniformShaderParameters, "PrimitiveFade");
@@ -999,6 +1006,62 @@ void FPersistentUniformBuffers::UpdateViewUniformBufferImmediate(const FViewUnif
 	CachedView = nullptr;
 }
 
+class FAsyncCreateLightPrimitiveInteractionsTask : public FNonAbandonableTask
+{
+	TSet<FPrimitiveSceneInfo*> PendingPrimitives;
+	const FScene* Scene;
+
+public:
+	FAsyncCreateLightPrimitiveInteractionsTask()
+		: Scene(nullptr)
+	{}
+
+	void Init(const FScene* InScene)
+	{
+		Scene = InScene;
+	}
+
+	void AddPrimitive(FPrimitiveSceneInfo* PrimInfo)
+	{
+		PendingPrimitives.Add(PrimInfo);
+	}
+
+	bool HasPendingPrimitives() const
+	{
+		return PendingPrimitives.Num() > 0;
+	}
+
+	void DoWork()
+	{
+		for (TSet<FPrimitiveSceneInfo*>::TConstIterator It(PendingPrimitives); It; ++It)
+		{
+			FPrimitiveSceneInfo* PrimInfo = *It;
+			FPrimitiveSceneProxy* Proxy = PrimInfo->Proxy;
+			const FBoxSphereBounds& Bounds = Proxy->GetBounds();
+
+			for (FSceneLightOctree::TConstElementBoxIterator<> LightIt(Scene->LightOctree, FBoxCenterAndExtent(Bounds));
+				LightIt.HasPendingElements();
+				LightIt.Advance())
+			{
+				const FLightSceneInfoCompact& LightInfo = LightIt.GetCurrentElement();
+
+				if (LightInfo.AffectsPrimitive(Bounds, Proxy))
+				{
+					FLightPrimitiveInteraction::Create(LightInfo.LightSceneInfo, PrimInfo);
+				}
+			}
+		}
+
+		PendingPrimitives.Reset();
+		Scene = nullptr;
+	}
+
+	FORCEINLINE TStatId GetStatId() const
+	{
+		RETURN_QUICK_DECLARE_CYCLE_STAT(FAsyncCreateLightPrimitiveInteractionsTask, STATGROUP_ThreadPoolAsyncTasks);
+	}
+};
+
 FScene::FScene(UWorld* InWorld, bool bInRequiresHitProxies, bool bInIsEditorScene, bool bCreateFXSystem, ERHIFeatureLevel::Type InFeatureLevel)
 :	FSceneInterface(InFeatureLevel)
 ,	World(InWorld)
@@ -1031,6 +1094,7 @@ FScene::FScene(UWorld* InWorld, bool bInRequiresHitProxies, bool bInIsEditorScen
 #if RHI_RAYTRACING
 , RayTracingDynamicGeometryCollection(nullptr)
 #endif
+,	AsyncCreateLightPrimitiveInteractionsTask(nullptr)
 ,	NumVisibleLights_GameThread(0)
 ,	NumEnabledSkylights_GameThread(0)
 ,	SceneFrameNumber(0)
@@ -1128,6 +1192,12 @@ FScene::~FScene()
 	BeginReleaseResource(&HaltonPrimesResource);
 #endif
 
+	if (AsyncCreateLightPrimitiveInteractionsTask)
+	{
+		delete AsyncCreateLightPrimitiveInteractionsTask;
+		AsyncCreateLightPrimitiveInteractionsTask = nullptr;
+	}
+
 	checkf(RemovedPrimitiveSceneInfos.Num() == 0, TEXT("Leaking %d FPrimitiveSceneInfo instances."), RemovedPrimitiveSceneInfos.Num()); // Ensure UpdateAllPrimitiveSceneInfos() is called before destruction.
 }
 
@@ -3013,6 +3083,8 @@ void FScene::GetRelevantLights_RenderThread( UPrimitiveComponent* Primitive, TAr
 	check( RelevantLights );
 	if( Primitive->SceneProxy )
 	{
+		FlushAsyncLightPrimitiveInteractionCreation();
+
 		for( const FLightPrimitiveInteraction* Interaction=Primitive->SceneProxy->GetPrimitiveSceneInfo()->LightList; Interaction; Interaction=Interaction->GetNextLight() )
 		{
 			RelevantLights->Add( Interaction->GetLight()->Proxy->GetLightComponent() );
@@ -3199,6 +3271,11 @@ void FScene::DumpUnbuiltLightInteractions( FOutputDevice& Ar ) const
 {
 	FlushRenderingCommands();
 
+	if (AsyncCreateLightPrimitiveInteractionsTask)
+	{
+		AsyncCreateLightPrimitiveInteractionsTask->EnsureCompletion();
+	}
+
 	TSet<FString> LightsWithUnbuiltInteractions;
 	TSet<FString> PrimitivesWithUnbuiltInteractions;
 
@@ -3485,7 +3562,7 @@ struct FPrimitiveArraySortKey
 	}
 };
 
-void FScene::UpdateAllPrimitiveSceneInfos(FRHICommandListImmediate& RHICmdList)
+void FScene::UpdateAllPrimitiveSceneInfos(FRHICommandListImmediate& RHICmdList, bool bAsyncCreateLPIs)
 {
 	SCOPED_NAMED_EVENT(FScene_UpdateAllPrimitiveSceneInfos, FColor::Orange);
 	SCOPE_CYCLE_COUNTER(STAT_UpdateScenePrimitiveRenderThreadTime);
@@ -3501,6 +3578,20 @@ void FScene::UpdateAllPrimitiveSceneInfos(FRHICommandListImmediate& RHICmdList)
 	TSet<FPrimitiveSceneInfo*> DeletedSceneInfos;
 	DeletedSceneInfos.Reserve(RemovedLocalPrimitiveSceneInfos.Num());
 
+	if (!!GAsyncCreateLightPrimitiveInteractions && !AsyncCreateLightPrimitiveInteractionsTask)
+	{
+		AsyncCreateLightPrimitiveInteractionsTask = new FAsyncTask<FAsyncCreateLightPrimitiveInteractionsTask>();
+	}
+
+	if (AsyncCreateLightPrimitiveInteractionsTask)
+	{
+		if (!AsyncCreateLightPrimitiveInteractionsTask->IsDone())
+		{
+			AsyncCreateLightPrimitiveInteractionsTask->EnsureCompletion();
+		}
+		AsyncCreateLightPrimitiveInteractionsTask->GetTask().Init(this);
+	}
+
 	{
 		SCOPED_NAMED_EVENT(FScene_RemovePrimitiveSceneInfos, FColor::Red);
 		SCOPE_CYCLE_COUNTER(STAT_RemoveScenePrimitiveTime);
@@ -3814,18 +3905,18 @@ void FScene::UpdateAllPrimitiveSceneInfos(FRHICommandListImmediate& RHICmdList)
 					SCOPED_NAMED_EVENT(FScene_AddPrimitiveSceneInfoToScene, FColor::Turquoise);
 					if (GIsEditor)
 					{
-						PrimitiveSceneInfo->AddToScene(RHICmdList, true);
+						PrimitiveSceneInfo->AddToScene(RHICmdList, true, true, bAsyncCreateLPIs);
 					}
 					else
 					{
 						const bool bAddToDrawLists = !(CVarDoLazyStaticMeshUpdate.GetValueOnRenderThread());
 						if (bAddToDrawLists)
 						{
-							PrimitiveSceneInfo->AddToScene(RHICmdList, true);
+							PrimitiveSceneInfo->AddToScene(RHICmdList, true, true, bAsyncCreateLPIs);
 						}
 						else
 						{
-							PrimitiveSceneInfo->AddToScene(RHICmdList, true, false);
+							PrimitiveSceneInfo->AddToScene(RHICmdList, true, false, bAsyncCreateLPIs);
 							PrimitiveSceneInfo->BeginDeferredUpdateStaticMeshes();
 						}
 					}
@@ -3909,11 +4000,17 @@ void FScene::UpdateAllPrimitiveSceneInfos(FRHICommandListImmediate& RHICmdList)
 			check(!(bUpdateStaticDrawLists && PrimitiveSceneInfo->StaticMeshes.Num()));
 
 			// Re-add the primitive to the scene with the new transform.
-			PrimitiveSceneInfo->AddToScene(RHICmdList, bUpdateStaticDrawLists);
+			PrimitiveSceneInfo->AddToScene(RHICmdList, bUpdateStaticDrawLists, true, bAsyncCreateLPIs);
 
 			PrimitiveSceneInfo->FlushRuntimeVirtualTexture();
 		}
 
+		if (AsyncCreateLightPrimitiveInteractionsTask && AsyncCreateLightPrimitiveInteractionsTask->GetTask().HasPendingPrimitives())
+		{
+			check(GAsyncCreateLightPrimitiveInteractions);
+			AsyncCreateLightPrimitiveInteractionsTask->StartBackgroundTask();
+		}
+
 		for (const auto& Transform : OverridenPreviousTransforms)
 		{
 			FPrimitiveSceneInfo* PrimitiveSceneInfo = Transform.Key;
@@ -3988,6 +4085,50 @@ void FScene::UpdateAllPrimitiveSceneInfos(FRHICommandListImmediate& RHICmdList)
 	AddedPrimitiveSceneInfos.Reset();
 }
 
+void FScene::CreateLightPrimitiveInteractionsForPrimitive(FPrimitiveSceneInfo* PrimitiveInfo, bool bAsyncCreateLPIs)
+{
+	if (!bAsyncCreateLPIs
+		|| !GAsyncCreateLightPrimitiveInteractions
+		|| GetShadingPath() == EShadingPath::Mobile
+		|| PrimitiveInfo->LODParentComponentId.IsValid()
+		|| !PrimitiveInfo->Proxy->IsMeshShapeOftenMoving())
+	{
+		FMemMark MemStackMark(FMemStack::Get());
+		const FBoxSphereBounds& Bounds = PrimitiveInfo->Proxy->GetBounds();
+
+		for (FSceneLightOctree::TConstElementBoxIterator<SceneRenderingAllocator> LightIt(LightOctree, FBoxCenterAndExtent(Bounds));
+			LightIt.HasPendingElements();
+			LightIt.Advance())
+		{
+			const FLightSceneInfoCompact& LightInfo = LightIt.GetCurrentElement();
+
+			if (LightInfo.AffectsPrimitive(Bounds, PrimitiveInfo->Proxy))
+			{
+				FLightPrimitiveInteraction::Create(LightInfo.LightSceneInfo, PrimitiveInfo);
+			}
+		}
+	}
+	else
+	{
+		check(AsyncCreateLightPrimitiveInteractionsTask);
+		AsyncCreateLightPrimitiveInteractionsTask->GetTask().AddPrimitive(PrimitiveInfo);
+	}
+}
+
+void FScene::FlushAsyncLightPrimitiveInteractionCreation() const
+{
+	if (!GAsyncCreateLightPrimitiveInteractions)
+	{
+		check(!AsyncCreateLightPrimitiveInteractionsTask || AsyncCreateLightPrimitiveInteractionsTask->IsIdle());
+	}
+	else if (AsyncCreateLightPrimitiveInteractionsTask && !AsyncCreateLightPrimitiveInteractionsTask->IsWorkDone())
+	{
+		CSV_SCOPED_TIMING_STAT_EXCLUSIVE(FlushAsyncLPICreation);
+		SCOPE_CYCLE_COUNTER(STAT_FlushAsyncLPICreation);
+		AsyncCreateLightPrimitiveInteractionsTask->EnsureCompletion();
+	}
+}
+
 /**
  * Dummy NULL scene interface used by dedicated servers.
  */
@@ -4015,7 +4156,7 @@ public:
 	virtual void AddPrimitive(UPrimitiveComponent* Primitive) override {}
 	virtual void RemovePrimitive(UPrimitiveComponent* Primitive) override {}
 	virtual void ReleasePrimitive(UPrimitiveComponent* Primitive) override {}
-	virtual void UpdateAllPrimitiveSceneInfos(FRHICommandListImmediate& RHICmdList) override {}
+	virtual void UpdateAllPrimitiveSceneInfos(FRHICommandListImmediate& RHICmdList, bool bAsyncCreateLPIs = false) override {}
 	virtual FPrimitiveSceneInfo* GetPrimitiveSceneInfo(int32 PrimiteIndex) override { return NULL; }
 
 	/** Updates the transform of a primitive which has already been added to the scene. */
diff --git a/Engine/Source/Runtime/Renderer/Private/SceneCore.cpp b/Engine/Source/Runtime/Renderer/Private/SceneCore.cpp
--- a/Engine/Source/Runtime/Renderer/Private/SceneCore.cpp
+++ b/Engine/Source/Runtime/Renderer/Private/SceneCore.cpp
@@ -36,6 +36,8 @@ TAllocatorFixedSizeFreeList<sizeof(FLightPrimitiveInteraction), FREE_LIST_GROW_S
 
 uint32 FRendererModule::GetNumDynamicLightsAffectingPrimitive(const FPrimitiveSceneInfo* PrimitiveSceneInfo,const FLightCacheInterface* LCI)
 {
+	PrimitiveSceneInfo->Scene->FlushAsyncLightPrimitiveInteractionCreation();
+
 	uint32 NumDynamicLights = 0;
 
 	FLightPrimitiveInteraction *LightList = PrimitiveSceneInfo->LightList;
diff --git a/Engine/Source/Runtime/Renderer/Private/ScenePrivate.h b/Engine/Source/Runtime/Renderer/Private/ScenePrivate.h
--- a/Engine/Source/Runtime/Renderer/Private/ScenePrivate.h
+++ b/Engine/Source/Runtime/Renderer/Private/ScenePrivate.h
@@ -2727,7 +2727,7 @@ public:
 	virtual void AddPrimitive(UPrimitiveComponent* Primitive) override;
 	virtual void RemovePrimitive(UPrimitiveComponent* Primitive) override;
 	virtual void ReleasePrimitive(UPrimitiveComponent* Primitive) override;
-	virtual void UpdateAllPrimitiveSceneInfos(FRHICommandListImmediate& RHICmdList) override;
+	virtual void UpdateAllPrimitiveSceneInfos(FRHICommandListImmediate& RHICmdList, bool bAsyncCreateLPIs = false) override;
 	virtual void UpdatePrimitiveTransform(UPrimitiveComponent* Primitive) override;
 	virtual void UpdatePrimitiveAttachment(UPrimitiveComponent* Primitive) override;
 	virtual void UpdateCustomPrimitiveData(UPrimitiveComponent* Primitive) override;
@@ -2973,6 +2973,10 @@ public:
 
 	void DumpMeshDrawCommandMemoryStats();
 
+	void CreateLightPrimitiveInteractionsForPrimitive(FPrimitiveSceneInfo* PrimitiveInfo, bool bAsyncCreateLPIs);
+
+	void FlushAsyncLightPrimitiveInteractionCreation() const;
+
 private:
 
 	/**
@@ -3090,6 +3094,8 @@ private:
 	TSet<FPrimitiveSceneInfo*> RemovedPrimitiveSceneInfos;
 	TSet<FPrimitiveSceneInfo*> DistanceFieldSceneDataUpdates;
 
+	FAsyncTask<class FAsyncCreateLightPrimitiveInteractionsTask>* AsyncCreateLightPrimitiveInteractionsTask;
+
 	/** 
 	 * The number of visible lights in the scene
 	 * Note: This is tracked on the game thread!
diff --git a/Engine/Source/Runtime/Renderer/Private/SceneVisibility.cpp b/Engine/Source/Runtime/Renderer/Private/SceneVisibility.cpp
--- a/Engine/Source/Runtime/Renderer/Private/SceneVisibility.cpp
+++ b/Engine/Source/Runtime/Renderer/Private/SceneVisibility.cpp
@@ -3023,6 +3023,8 @@ void FSceneRenderer::PreVisibilityFrameSetup(FRHICommandListImmediate& RHICmdLis
 	// Draw lines to lights affecting this mesh if its selected.
 	if (ViewFamily.EngineShowFlags.LightInfluences)
 	{
+		Scene->FlushAsyncLightPrimitiveInteractionCreation();
+
 		for (TArray<FPrimitiveSceneInfo*>::TConstIterator It(Scene->Primitives); It; ++It)
 		{
 			const FPrimitiveSceneInfo* PrimitiveSceneInfo = *It;
diff --git a/Engine/Source/Runtime/Renderer/Private/ShadowSetup.cpp b/Engine/Source/Runtime/Renderer/Private/ShadowSetup.cpp
--- a/Engine/Source/Runtime/Renderer/Private/ShadowSetup.cpp
+++ b/Engine/Source/Runtime/Renderer/Private/ShadowSetup.cpp
@@ -2527,6 +2527,8 @@ void FSceneRenderer::CreateWholeSceneProjectedShadow(
 
 		if (MaxFadeAlpha > 1.0f / 256.0f)
 		{
+			Scene->FlushAsyncLightPrimitiveInteractionCreation();
+
 			for (int32 ShadowIndex = 0, ShadowCount = ProjectedShadowInitializers.Num(); ShadowIndex < ShadowCount; ShadowIndex++)
 			{
 				FWholeSceneProjectedShadowInitializer& ProjectedShadowInitializer = ProjectedShadowInitializers[ShadowIndex];
@@ -4209,6 +4211,8 @@ void FSceneRenderer::InitDynamicShadows(FRHICommandListImmediate& RHICmdList, FG
 
 						if( !bMobile || (LightSceneInfo->Proxy->CastsModulatedShadows() && !LightSceneInfo->Proxy->UseCSMForDynamicObjects()))
 						{
+							Scene->FlushAsyncLightPrimitiveInteractionCreation();
+
 							// Look for individual primitives with a dynamic shadow.
 							for (FLightPrimitiveInteraction* Interaction = LightSceneInfo->DynamicInteractionOftenMovingPrimitiveList;
 								Interaction;
diff --git a/Engine/Source/Runtime/Renderer/Public/PrimitiveSceneInfo.h b/Engine/Source/Runtime/Renderer/Public/PrimitiveSceneInfo.h
--- a/Engine/Source/Runtime/Renderer/Public/PrimitiveSceneInfo.h
+++ b/Engine/Source/Runtime/Renderer/Public/PrimitiveSceneInfo.h
@@ -338,7 +338,7 @@ public:
 	~FPrimitiveSceneInfo();
 
 	/** Adds the primitive to the scene. */
-	void AddToScene(FRHICommandListImmediate& RHICmdList, bool bUpdateStaticDrawLists, bool bAddToStaticDrawLists = true);
+	void AddToScene(FRHICommandListImmediate& RHICmdList, bool bUpdateStaticDrawLists, bool bAddToStaticDrawLists = true, bool bAsyncCreateLPIs = false);
 
 	/** Removes the primitive from the scene. */
 	void RemoveFromScene(bool bUpdateStaticDrawLists);
