Fix a crash when you enable RTHF shadow, disable, and then change scalability settings.

diff --git a/Engine/Source/Runtime/Renderer/Private/DeferredShadingRenderer.cpp b/Engine/Source/Runtime/Renderer/Private/DeferredShadingRenderer.cpp
--- a/Engine/Source/Runtime/Renderer/Private/DeferredShadingRenderer.cpp
+++ b/Engine/Source/Runtime/Renderer/Private/DeferredShadingRenderer.cpp
@@ -568,7 +568,10 @@ void FDeferredShadingSceneRenderer::PrepareDistanceFieldScene(FRHICommandListImm
 	CSV_SCOPED_TIMING_STAT_EXCLUSIVE(RenderDFAO);
 	SCOPE_CYCLE_COUNTER(STAT_FDeferredShadingSceneRenderer_DistanceFieldAO_Init);
 
-	if (ShouldPrepareHeightFieldScene())
+	const bool bShouldPrepareHeightFieldScene = ShouldPrepareHeightFieldScene();
+	const bool bShouldPrepareDistanceFieldScene = ShouldPrepareDistanceFieldScene();
+
+	if (bShouldPrepareHeightFieldScene)
 	{
 		extern int32 GHFShadowQuality;
 		if (GHFShadowQuality > 2)
@@ -578,8 +581,12 @@ void FDeferredShadingSceneRenderer::PrepareDistanceFieldScene(FRHICommandListImm
 		GHeightFieldTextureAtlas.UpdateAllocations(RHICmdList, FeatureLevel);
 		UpdateGlobalHeightFieldObjectBuffers(RHICmdList);
 	}
+	else if (bShouldPrepareDistanceFieldScene)
+	{
+		AddOrRemoveSceneHeightFieldPrimitives();
+	}
 
-	if (ShouldPrepareDistanceFieldScene())
+	if (bShouldPrepareDistanceFieldScene)
 	{
 		GDistanceFieldVolumeTextureAtlas.UpdateAllocations(RHICmdList, FeatureLevel);
 		UpdateGlobalDistanceFieldObjectBuffers(RHICmdList);
diff --git a/Engine/Source/Runtime/Renderer/Private/DeferredShadingRenderer.h b/Engine/Source/Runtime/Renderer/Private/DeferredShadingRenderer.h
--- a/Engine/Source/Runtime/Renderer/Private/DeferredShadingRenderer.h
+++ b/Engine/Source/Runtime/Renderer/Private/DeferredShadingRenderer.h
@@ -495,6 +495,7 @@ private:
 
 	void UpdateGlobalDistanceFieldObjectBuffers(FRHICommandListImmediate& RHICmdList);
 	void UpdateGlobalHeightFieldObjectBuffers(FRHICommandListImmediate& RHICmdList);
+	void AddOrRemoveSceneHeightFieldPrimitives(bool bSkipAdd = false);
 	void PrepareDistanceFieldScene(FRHICommandListImmediate& RHICmdList, bool bSplitDispatch);
 
 	void RenderViewTranslucency(FRHICommandListImmediate& RHICmdList, const FViewInfo& View, const FMeshPassProcessorRenderState& DrawRenderState, ETranslucencyPass::Type TranslucenyPass);
diff --git a/Engine/Source/Runtime/Renderer/Private/DistanceFieldObjectManagement.cpp b/Engine/Source/Runtime/Renderer/Private/DistanceFieldObjectManagement.cpp
--- a/Engine/Source/Runtime/Renderer/Private/DistanceFieldObjectManagement.cpp
+++ b/Engine/Source/Runtime/Renderer/Private/DistanceFieldObjectManagement.cpp
@@ -1633,6 +1633,15 @@ void FDeferredShadingSceneRenderer::UpdateGlobalHeightFieldObjectBuffers(FRHICom
 
 		if (!DistanceFieldSceneData.HeightFieldObjectBuffers)
 		{
+			AddOrRemoveSceneHeightFieldPrimitives(true);
+
+			for (int32 Idx = 0; Idx < DistanceFieldSceneData.HeightfieldPrimitives.Num(); ++Idx)
+			{
+				FPrimitiveSceneInfo* Primitive = DistanceFieldSceneData.HeightfieldPrimitives[Idx];
+				check(!DistanceFieldSceneData.PendingHeightFieldAddOps.Contains(Primitive));
+				DistanceFieldSceneData.PendingHeightFieldAddOps.Add(Primitive);
+			}
+			DistanceFieldSceneData.HeightfieldPrimitives.Reset();
 			DistanceFieldSceneData.HeightFieldObjectBuffers = new FHeightFieldObjectBuffers;
 		}
 
@@ -1749,6 +1758,59 @@ void FDeferredShadingSceneRenderer::UpdateGlobalHeightFieldObjectBuffers(FRHICom
 	}
 }
 
+void FDeferredShadingSceneRenderer::AddOrRemoveSceneHeightFieldPrimitives(bool bSkipAdd)
+{
+	FDistanceFieldSceneData& SceneData = Scene->DistanceFieldSceneData;
+
+	if (SceneData.HeightFieldObjectBuffers)
+	{
+		delete SceneData.HeightFieldObjectBuffers;
+		SceneData.HeightFieldObjectBuffers = nullptr;
+		SceneData.NumHeightFieldObjectsInBuffer = 0;
+		SceneData.HeightFieldAtlasGeneration = 0;
+		SceneData.HFVisibilityAtlasGenerattion = 0;
+	}
+
+	TArray<int32, SceneRenderingAllocator> PendingRemoveIndices;
+	for (int32 Idx = 0; Idx < SceneData.PendingHeightFieldRemoveOps.Num(); ++Idx)
+	{
+		const FHeightFieldPrimitiveRemoveInfo& RemoveInfo = SceneData.PendingHeightFieldRemoveOps[Idx];
+		check(RemoveInfo.DistanceFieldInstanceIndices.Num() == 1);
+		PendingRemoveIndices.Add(RemoveInfo.DistanceFieldInstanceIndices[0]);
+		const FGlobalDFCacheType CacheType = RemoveInfo.bOftenMoving ? GDF_Full : GDF_MostlyStatic;
+		SceneData.PrimitiveModifiedBounds[CacheType].Add(RemoveInfo.SphereBound);
+	}
+	SceneData.PendingHeightFieldRemoveOps.Reset();
+	Algo::Sort(PendingRemoveIndices);
+	for (int32 Idx = PendingRemoveIndices.Num() - 1; Idx >= 0; --Idx)
+	{
+		const int32 RemoveIdx = PendingRemoveIndices[Idx];
+		const int32 LastObjectIdx = SceneData.HeightfieldPrimitives.Num() - 1;
+		if (RemoveIdx != LastObjectIdx)
+		{
+			SceneData.HeightfieldPrimitives[LastObjectIdx]->DistanceFieldInstanceIndices[0] = RemoveIdx;
+		}
+		SceneData.HeightfieldPrimitives.RemoveAtSwap(RemoveIdx);
+	}
+
+	if (!bSkipAdd)
+	{
+		for (int32 Idx = 0; Idx < SceneData.PendingHeightFieldAddOps.Num(); ++Idx)
+		{
+			FPrimitiveSceneInfo* Primitive = SceneData.PendingHeightFieldAddOps[Idx];
+			const int32 HFIdx = SceneData.HeightfieldPrimitives.Add(Primitive);
+			Primitive->DistanceFieldInstanceIndices.Empty(1);
+			Primitive->DistanceFieldInstanceIndices.Add(HFIdx);
+			const FGlobalDFCacheType CacheType = Primitive->Proxy->IsOftenMoving() ? GDF_Full : GDF_MostlyStatic;
+			const FBoxSphereBounds& Bounds = Primitive->Proxy->GetBounds();
+			SceneData.PrimitiveModifiedBounds[CacheType].Add(FVector4(Bounds.Origin, Bounds.SphereRadius));
+		}
+		SceneData.PendingHeightFieldAddOps.Reset();
+	}
+
+	SceneData.PendingHeightFieldUpdateOps.Empty();
+}
+
 FString GetObjectBufferMemoryString()
 {
 	return FString::Printf(TEXT("Temp object buffers %.3fMb"), 
