Copying //UE4/Dev-Physics to Main (//UE4/Main)

diff --git a/Engine/Plugins/Experimental/HairStrands/Source/HairStrandsNiagara/Public/NiagaraDataInterfaceHairStrands.h b/Engine/Plugins/Experimental/HairStrands/Source/HairStrandsNiagara/Public/NiagaraDataInterfaceHairStrands.h
--- a/Engine/Plugins/Experimental/HairStrands/Source/HairStrandsNiagara/Public/NiagaraDataInterfaceHairStrands.h
+++ b/Engine/Plugins/Experimental/HairStrands/Source/HairStrandsNiagara/Public/NiagaraDataInterfaceHairStrands.h
@@ -10,17 +10,7 @@
 #include "GroomActor.h"
 #include "NiagaraDataInterfaceHairStrands.generated.h"
 
-/** Size of each strands*/
-UENUM(BlueprintType)
-enum class EHairStrandsSize : uint8
-{
-	None = 0 UMETA(Hidden),
-	Size2 = 0x02 UMETA(DisplatName = "2"),
-	Size4 = 0x04 UMETA(DisplatName = "4"),
-	Size8 = 0x08 UMETA(DisplatName = "8"),
-	Size16 = 0x10 UMETA(DisplatName = "16"),
-	Size32 = 0x20 UMETA(DisplatName = "32")
-};
+static const int32 MaxDelay = 5;
 
 /** Render buffers that will be used in hlsl functions */
 struct FNDIHairStrandsBuffer : public FRenderResource
@@ -66,12 +56,109 @@ struct FNDIHairStrandsBuffer : public FRenderResource
 /** Data stored per strand base instance*/
 struct FNDIHairStrandsData
 {
+	FNDIHairStrandsData()
+	{
+		ResetDatas();
+	}
 	/** Initialize the buffers */
 	bool Init(class UNiagaraDataInterfaceHairStrands* Interface, FNiagaraSystemInstance* SystemInstance);
 
 	/** Release the buffers */
 	void Release();
 
+	inline void ResetDatas()
+	{
+		WorldTransform.SetIdentity();
+		BoxCenter = FVector(0, 0, 0);
+		BoxExtent = FVector(0, 0, 0);
+
+		TickCount = 0;
+		ForceReset = true;
+		ResetTick = MaxDelay;
+
+		NumStrands = 0;
+		StrandsSize = 0;
+
+		SubSteps = 5;
+		IterationCount = 20;
+
+		GravityVector = FVector(0.0, 0.0, -981.0);
+		AirDrag = 0.1;
+		AirVelocity = FVector(0, 0, 0);
+
+		SolveBend = true;
+		ProjectBend = false;
+		BendDamping = 0.01;
+		BendStiffness = 0.01;
+
+		SolveStretch = true;
+		ProjectStretch = false;
+		StretchDamping = 0.01;
+		StretchStiffness = 1.0;
+
+		SolveCollision = true;
+		ProjectCollision = true;
+		KineticFriction = 0.1;
+		StaticFriction = 0.1;
+		StrandsViscosity = 1.0;
+		CollisionRadius = 1.0;
+
+		StrandsDensity = 1.0;
+		StrandsSmoothing = 0.1;
+		StrandsThickness = 0.01;
+	}
+
+	inline void CopyDatas(const FNDIHairStrandsData* OtherDatas)
+	{
+		if (OtherDatas != nullptr)
+		{
+			HairStrandsBuffer = OtherDatas->HairStrandsBuffer;
+
+			WorldTransform = OtherDatas->WorldTransform;
+			BoxCenter = OtherDatas->BoxCenter;
+			BoxExtent = OtherDatas->BoxExtent;
+
+			TickCount = OtherDatas->TickCount;
+			ForceReset = OtherDatas->ForceReset;
+			ResetTick = OtherDatas->ResetTick;
+
+			NumStrands = OtherDatas->NumStrands;
+			StrandsSize = OtherDatas->StrandsSize;
+
+			SubSteps = OtherDatas->SubSteps;
+			IterationCount = OtherDatas->IterationCount;
+
+			GravityVector = OtherDatas->GravityVector;
+			AirDrag = OtherDatas->AirDrag;
+			AirVelocity = OtherDatas->AirVelocity;
+
+			SolveBend = OtherDatas->SolveBend;
+			ProjectBend = OtherDatas->ProjectBend;
+			BendDamping = OtherDatas->BendDamping;
+			BendStiffness = OtherDatas->BendStiffness;
+			BendScale = OtherDatas->BendScale;
+
+			SolveStretch = OtherDatas->SolveStretch;
+			ProjectStretch = OtherDatas->ProjectStretch;
+			StretchDamping = OtherDatas->StretchDamping;
+			StretchStiffness = OtherDatas->StretchStiffness;
+			StretchScale = OtherDatas->StretchScale;
+
+			SolveCollision = OtherDatas->SolveCollision;
+			ProjectCollision = OtherDatas->ProjectCollision;
+			StaticFriction = OtherDatas->StaticFriction;
+			KineticFriction = OtherDatas->KineticFriction;
+			StrandsViscosity = OtherDatas->StrandsViscosity;
+			CollisionRadius = OtherDatas->CollisionRadius;
+			RadiusScale = OtherDatas->RadiusScale;
+
+			StrandsDensity = OtherDatas->StrandsDensity;
+			StrandsSmoothing = OtherDatas->StrandsSmoothing;
+			StrandsThickness = OtherDatas->StrandsThickness;
+			ThicknessScale = OtherDatas->ThicknessScale;
+		}
+	}
+
 	/** Cached World transform. */
 	FTransform WorldTransform;
 
@@ -79,7 +166,7 @@ struct FNDIHairStrandsData
 	int32 NumStrands;
 
 	/** Strand size */
-	int32 StrandSize;
+	int32 StrandsSize;
 
 	/** Bounding box center */
 	FVector BoxCenter;
@@ -98,6 +185,84 @@ struct FNDIHairStrandsData
 
 	/** Strands Gpu buffer */
 	FNDIHairStrandsBuffer* HairStrandsBuffer;
+	
+	/** Number of substeps to be used */
+	int32 SubSteps;
+
+	/** Number of iterations for the constraint solver  */
+	int32 IterationCount;
+
+	/** Acceleration vector in cm/s2 to be used for the gravity*/
+	FVector GravityVector;
+
+	/** Coefficient between 0 and 1 to be used for the air drag */
+	float AirDrag;
+
+	/** Velocity of the surrounding air in cm/s  */
+	FVector AirVelocity;
+
+	/** Velocity of the surrounding air in cm/s */
+	bool SolveBend;
+
+	/** Enable the solve of the bend constraint during the xpbd loop */
+	bool ProjectBend;
+
+	/** Damping for the bend constraint between 0 and 1 */
+	float BendDamping;
+
+	/** Stiffness for the bend constraint in GPa */
+	float BendStiffness;
+
+	/** Stiffness scale along the strand */
+	TStaticArray<float,4> BendScale;
+
+	/** Enable the solve of the stretch constraint during the xpbd loop */
+	bool SolveStretch;
+
+	/** Enable the projection of the stretch constraint after the xpbd loop */
+	bool ProjectStretch;
+
+	/** Damping for the stretch constraint between 0 and 1 */
+	float StretchDamping;
+
+	/** Stiffness for the stretch constraint in GPa */
+	float StretchStiffness;
+
+	/** Stiffness scale along the strand */
+	TStaticArray<float, 4> StretchScale;
+
+	/** Enable the solve of the collision constraint during the xpbd loop  */
+	bool SolveCollision;
+
+	/** Enable ther projection of the collision constraint after the xpbd loop */
+	bool ProjectCollision;
+
+	/** Static friction used for collision against the physics asset */
+	float StaticFriction;
+
+	/** Kinetic friction used for collision against the physics asset*/
+	float KineticFriction;
+
+	/** Radius that will be used for the collision detection against the physics asset */
+	float StrandsViscosity;
+
+	/** Radius scale along the strand */
+	float CollisionRadius;
+
+	/** Radius scale along the strand */
+	TStaticArray<float, 4> RadiusScale;
+
+	/** Density of the strands in g/cm3 */
+	float StrandsDensity;
+
+	/** Smoothing between 0 and 1 of the incoming guides curves for better stability */
+	float StrandsSmoothing;
+
+	/** Strands thickness in cm that will be used for mass and inertia computation */
+	float StrandsThickness;
+
+	/** Thickness scale along the curve */
+	TStaticArray<float, 4> ThicknessScale;
 };
 
 /** Data Interface for the strand base */
@@ -108,10 +273,6 @@ class HAIRSTRANDSNIAGARA_API UNiagaraDataInterfaceHairStrands : public UNiagaraD
 
 public:
 
-	/** Size of each strand. */
-	UPROPERTY(EditAnywhere, Category = "Spawn")
-	EHairStrandsSize StrandSize;
-
 	/** Hair Strands Asset used to sample from when not overridden by a source actor from the scene. Also useful for previewing in the editor. */
 	UPROPERTY(EditAnywhere, Category = "Source")
 	UGroomAsset* DefaultSource;
@@ -123,10 +284,6 @@ public:
 	/** The source component from which to sample */
 	TWeakObjectPtr<class UGroomComponent> SourceComponent;
 
-	/** Group Index to be used */
-	UPROPERTY(EditAnywhere, Category = "Source")
-	int32 GroupIndex;
-
 	/** UObject Interface */
 	virtual void PostInitProperties() override;
 
@@ -155,14 +312,66 @@ public:
 
 	/** Extract datas and resources */
 	void ExtractDatasAndResources(FNiagaraSystemInstance* SystemInstance, FHairStrandsDatas*& OutStrandsDatas,
-		FHairStrandsRestResource*& OutStrandsRestResource, FHairStrandsDeformedResource*& OutStrandsDeformedResource, FHairStrandsRootResource*& OutStrandsRootResource);
+		FHairStrandsRestResource*& OutStrandsRestResource, FHairStrandsDeformedResource*& OutStrandsDeformedResource, FHairStrandsRootResource*& OutStrandsRootResource, UGroomAsset*& OutGroomAsset);
 
 	/** Get the number of strands */
 	void GetNumStrands(FVectorVMContext& Context);
 
-	/** Get the strand size  */
+	/** Get the grrom asset datas  */
 	void GetStrandSize(FVectorVMContext& Context);
 
+	void GetSubSteps(FVectorVMContext& Context);
+
+	void GetIterationCount(FVectorVMContext& Context);
+
+	void GetGravityVector(FVectorVMContext& Context);
+
+	void GetAirDrag(FVectorVMContext& Context);
+
+	void GetAirVelocity(FVectorVMContext& Context);
+
+	void GetSolveBend(FVectorVMContext& Context);
+
+	void GetProjectBend(FVectorVMContext& Context);
+
+	void GetBendDamping(FVectorVMContext& Context);
+
+	void GetBendStiffness(FVectorVMContext& Context);
+
+	void GetBendScale(FVectorVMContext& Context);
+
+	void GetSolveStretch(FVectorVMContext& Context);
+
+	void GetProjectStretch(FVectorVMContext& Context);
+
+	void GetStretchDamping(FVectorVMContext& Context);
+
+	void GetStretchStiffness(FVectorVMContext& Context);
+
+	void GetStretchScale(FVectorVMContext& Context);
+
+	void GetSolveCollision(FVectorVMContext& Context);
+
+	void GetProjectCollision(FVectorVMContext& Context);
+
+	void GetStaticFriction(FVectorVMContext& Context);
+
+	void GetKineticFriction(FVectorVMContext& Context);
+
+	void GetStrandsViscosity(FVectorVMContext& Context);
+
+	void GetCollisionRadius(FVectorVMContext& Context);
+
+	void GetRadiusScale(FVectorVMContext& Context);
+
+	void GetStrandsSmoothing(FVectorVMContext& Context);
+
+	void GetStrandsDensity(FVectorVMContext& Context);
+
+	void GetStrandsThickness(FVectorVMContext& Context);
+
+	void GetThicknessScale(FVectorVMContext& Context);
+
 	/** Get the world transform */
 	void GetWorldTransform(FVectorVMContext& Context);
 
@@ -266,10 +475,10 @@ public:
 	void ProjectBendRodMaterial(FVectorVMContext& Context);
 
 	/** Solve the static collision constraint */
-	void SolveStaticCollisionConstraint(FVectorVMContext& Context);
+	void SolveHardCollisionConstraint(FVectorVMContext& Context);
 
 	/** Project the static collision constraint */
-	void ProjectStaticCollisionConstraint(FVectorVMContext& Context);
+	void ProjectHardCollisionConstraint(FVectorVMContext& Context);
 
 	/** Solve the soft collision constraint */
 	void SolveSoftCollisionConstraint(FVectorVMContext& Context);
@@ -283,8 +492,11 @@ public:
 	/** Compute the rest direction*/
 	void ComputeRestDirection(FVectorVMContext& Context);
 
-	/** Update the node orientation to match the bishop frame*/
-	void UpdateNodeOrientation(FVectorVMContext& Context);
+	/** Update the strands material frame */
+	void UpdateMaterialFrame(FVectorVMContext& Context);
+
+	/** Compute the strands material frame */
+	void ComputeMaterialFrame(FVectorVMContext& Context);
 
 	/** Compute the air drag force */
 	void ComputeAirDragForce(FVectorVMContext& Context);
diff --git a/Engine/Plugins/Experimental/HairStrands/Source/HairStrandsNiagara/Public/NiagaraDataInterfacePhysicsAsset.h b/Engine/Plugins/Experimental/HairStrands/Source/HairStrandsNiagara/Public/NiagaraDataInterfacePhysicsAsset.h
--- a/Engine/Plugins/Experimental/HairStrands/Source/HairStrandsNiagara/Public/NiagaraDataInterfacePhysicsAsset.h
+++ b/Engine/Plugins/Experimental/HairStrands/Source/HairStrandsNiagara/Public/NiagaraDataInterfacePhysicsAsset.h
@@ -224,6 +224,16 @@ struct FNDIPhysicsAssetProxy : public FNiagaraDataInterfaceProxy
 	/** Destroy the proxy data if necessary */
 	void DestroyPerInstanceData(NiagaraEmitterInstanceBatcher* Batcher, const FNiagaraSystemInstanceID& SystemInstance);
 
+	/** Launch all pre stage functions */
+	virtual void PreStage(FRHICommandList& RHICmdList, const FNiagaraDataInterfaceSetArgs& Context) override;
+
+	/** Launch all post stage functions */
+	virtual void PostStage(FRHICommandList& RHICmdList, const FNiagaraDataInterfaceSetArgs& Context) override;
+
+	/** Reset the buffers  */
+	virtual void ResetData(FRHICommandList& RHICmdList, const FNiagaraDataInterfaceSetArgs& Context) override;
+
+
 	/** List of proxy data for each system instances*/
 	TMap<FNiagaraSystemInstanceID, FNDIPhysicsAssetData> SystemInstancesToProxyData;
 };
diff --git a/Engine/Shaders/Private/DeferredLightPixelShaders.usf b/Engine/Shaders/Private/DeferredLightPixelShaders.usf
--- a/Engine/Shaders/Private/DeferredLightPixelShaders.usf
+++ b/Engine/Shaders/Private/DeferredLightPixelShaders.usf
@@ -157,7 +157,15 @@ void ComputeHairLighting(
 		float3			SampleBaseColor;
 		float			SampleRoughness;
 		float			SampleSpecular;
-		DecodeNodeData(NodeData, SampleDepth, SampleVisData, SampleBaseColor, SampleRoughness, SampleSpecular);
+		uint			SampleLightChannelMask;
+		DecodeNodeData(NodeData, SampleDepth, SampleVisData, SampleBaseColor, SampleRoughness, SampleSpecular, SampleLightChannelMask);
+
+	#if USE_LIGHTING_CHANNELS
+		if (!(SampleLightChannelMask & DeferredLightUniforms.LightingChannelMask))
+		{
+			continue;
+		}
+	#endif
 
 		// Inject material data from the visibility/mini-gbuffer for hair
 		const float SceneDepth = ConvertFromDeviceZ(SampleDepth);
@@ -276,11 +284,7 @@ void DeferredLightPixelMain(
 	OutColorSubPixel = 0;
 	float HairCoverage = 0;
 	BRANCH 
-	if (NodeDesc.Count > 0
-	#if USE_LIGHTING_CHANNELS
-		&& (GetLightingChannelMask(InputParams.ScreenUV) & DeferredLightUniforms.LightingChannelMask)
-	#endif
-		)
+	if (NodeDesc.Count > 0)
 	{
 		ComputeHairLighting(InputParams, NodeDesc, CatData, HairCoverage, OutColor, OutColorSubPixel);
 	#if USE_PREEXPOSURE
diff --git a/Engine/Shaders/Private/HairStrands/HairScatter.usf b/Engine/Shaders/Private/HairStrands/HairScatter.usf
--- a/Engine/Shaders/Private/HairStrands/HairScatter.usf
+++ b/Engine/Shaders/Private/HairStrands/HairScatter.usf
@@ -249,7 +249,8 @@ FGBufferData GetGBufferData(uint2 PixelCoord)
 			float3			SampleBaseColor;
 			float			SampleRoughness;
 			float			SampleSpecular;
-			DecodeNodeData(NodeData, SampleDepth, SampleVisData, SampleBaseColor, SampleRoughness, SampleSpecular);
+			uint			SampleLigthChannelMask;
+			DecodeNodeData(NodeData, SampleDepth, SampleVisData, SampleBaseColor, SampleRoughness, SampleSpecular, SampleLigthChannelMask);
 
 			if (SampleDepth >= ClosestDepth) // InverseZ
 			{
diff --git a/Engine/Shaders/Private/HairStrands/HairStrandsDebug.usf b/Engine/Shaders/Private/HairStrands/HairStrandsDebug.usf
--- a/Engine/Shaders/Private/HairStrands/HairStrandsDebug.usf
+++ b/Engine/Shaders/Private/HairStrands/HairStrandsDebug.usf
@@ -103,7 +103,8 @@ void MainPS(
 			float3			SampleBaseColor;
 			float			SampleRoughness;
 			float			SampleSpecular;
-			DecodeNodeData(Data, SampleDepth, SampleVisData, SampleBaseColor, SampleRoughness, SampleSpecular);
+			uint			SampleLigthChannelMask;
+			DecodeNodeData(Data, SampleDepth, SampleVisData, SampleBaseColor, SampleRoughness, SampleSpecular, SampleLigthChannelMask);
 
 			if (bAverageAllSamples)
 			{
diff --git a/Engine/Shaders/Private/HairStrands/HairStrandsDebugPrint.usf b/Engine/Shaders/Private/HairStrands/HairStrandsDebugPrint.usf
--- a/Engine/Shaders/Private/HairStrands/HairStrandsDebugPrint.usf
+++ b/Engine/Shaders/Private/HairStrands/HairStrandsDebugPrint.usf
@@ -22,7 +22,7 @@ StructuredBuffer<FNodeData> HairVisibilityNodeData;
 Buffer<int> HairMacroGroupAABBBuffer;
 
 #define MATERIAL_DATA 0
-#define DETAIL_DATA 0
+#define DETAIL_DATA 1
 
 [numthreads(1, 1, 1)]
 void MainCS(uint GroupIndex : SV_GroupIndex, uint3 DispatchThreadId : SV_DispatchThreadID)
@@ -54,16 +54,17 @@ void MainCS(uint GroupIndex : SV_GroupIndex, uint3 DispatchThreadId : SV_Dispatc
 
 	uint Total8BitCoverage = 0;
 	const FNodeDesc NodeDesc = DecodeNodeDesc(HairVisibilityNodeOffsetAndCount.Load(uint3(PixelCoord, 0)));
-	for (uint SampleIt = 0; SampleIt < NodeDesc.Count; SampleIt++)
+	for (uint SampleIt0 = 0; SampleIt0 < NodeDesc.Count; SampleIt0++)
 	{
-		const uint LocalOffset = NodeDesc.Offset + SampleIt;
+		const uint LocalOffset = NodeDesc.Offset + SampleIt0;
 		const FNodeData NodeData = HairVisibilityNodeData[LocalOffset];
 		float			SampleDepth;
 		FVisibilityData SampleVisData;
 		float3			SampleBaseColor;
 		float			SampleRoughness;
 		float			SampleSpecular;
-		DecodeNodeData(NodeData, SampleDepth, SampleVisData, SampleBaseColor, SampleRoughness, SampleSpecular);
+		uint			SampleLightChannelMask;
+		DecodeNodeData(NodeData, SampleDepth, SampleVisData, SampleBaseColor, SampleRoughness, SampleSpecular, SampleLightChannelMask);
 		Total8BitCoverage += SampleVisData.Coverage8bit;
 	}
 
@@ -234,7 +235,8 @@ void MainCS(uint GroupIndex : SV_GroupIndex, uint3 DispatchThreadId : SV_Dispatc
 		float3			SampleBaseColor;
 		float			SampleRoughness;
 		float			SampleSpecular;
-		DecodeNodeData(NodeData, SampleDepth, SampleVisData, SampleBaseColor, SampleRoughness, SampleSpecular);
+		uint			SampleLightChannelMask;
+		DecodeNodeData(NodeData, SampleDepth, SampleVisData, SampleBaseColor, SampleRoughness, SampleSpecular, SampleLightChannelMask);
 
 		// Base Color
 		Pos = ShaderPrintSymbol(Pos, _MINUS_);
@@ -353,6 +355,31 @@ void MainCS(uint GroupIndex : SV_GroupIndex, uint3 DispatchThreadId : SV_Dispatc
 
 		Pos = ShaderPrintNewline(Pos);
 		Pos.x = OriginalPos.x;		
+
+
+		// Ligtht channel mask
+		Pos = ShaderPrintSymbol(Pos, _DOT_);
+		Pos = ShaderPrintSymbol(Pos, _L_);
+		Pos = ShaderPrintSymbol(Pos, _I_);
+		Pos = ShaderPrintSymbol(Pos, _G_);
+		Pos = ShaderPrintSymbol(Pos, _H_);
+		Pos = ShaderPrintSymbol(Pos, _T_);
+		Pos = ShaderPrintSymbol(Pos, _SPC_);
+		Pos = ShaderPrintSymbol(Pos, _C_);
+		Pos = ShaderPrintSymbol(Pos, _H_);
+		Pos = ShaderPrintSymbol(Pos, _A_);
+		Pos = ShaderPrintSymbol(Pos, _N_);
+		Pos = ShaderPrintSymbol(Pos, _N_);
+		Pos = ShaderPrintSymbol(Pos, _E_);
+		Pos = ShaderPrintSymbol(Pos, _L_);
+		Pos = ShaderPrintSymbol(Pos, _S_);
+		Pos = ShaderPrintSymbol(Pos, _SPC_);
+		Pos = ShaderPrint(Pos, (SampleLightChannelMask & 0x1) ? 1u : 0u);
+		Pos = ShaderPrint(Pos, (SampleLightChannelMask & 0x2) ? 1u : 0u);
+		Pos = ShaderPrint(Pos, (SampleLightChannelMask & 0x4) ? 1u : 0u);
+
+		Pos = ShaderPrintNewline(Pos);
+		Pos.x = OriginalPos.x;
 #endif
 
 		// Coverage8bit
diff --git a/Engine/Shaders/Private/HairStrands/HairStrandsDeepTransmittanceCommon.ush b/Engine/Shaders/Private/HairStrands/HairStrandsDeepTransmittanceCommon.ush
--- a/Engine/Shaders/Private/HairStrands/HairStrandsDeepTransmittanceCommon.ush
+++ b/Engine/Shaders/Private/HairStrands/HairStrandsDeepTransmittanceCommon.ush
@@ -17,6 +17,14 @@ struct FPackedHairTransmittanceMask
 	uint PackedHairCountAndVisibility;
 };
 
+FPackedHairTransmittanceMask InitNullPackedHairTransmittanceMask()
+{
+	FPackedHairTransmittanceMask Out;
+	Out.Transmittance = 0;
+	Out.PackedHairCountAndVisibility = 0;
+	return Out;
+}
+
 FHairTransmittanceMask InitHairTransmittanceMask()
 {
 	FHairTransmittanceMask Out;
diff --git a/Engine/Shaders/Private/HairStrands/HairStrandsDeepTransmittanceMask.usf b/Engine/Shaders/Private/HairStrands/HairStrandsDeepTransmittanceMask.usf
--- a/Engine/Shaders/Private/HairStrands/HairStrandsDeepTransmittanceMask.usf
+++ b/Engine/Shaders/Private/HairStrands/HairStrandsDeepTransmittanceMask.usf
@@ -52,6 +52,7 @@ float3 ToWorldPosition(float2 UV, float Depth)
 uint2				DeepShadow_Resolution;
 float4x4			DeepShadow_WorldToLightTransforms[MAX_HAIR_MACROGROUP_COUNT];
 float3				LightDirection;
+uint				LightChannelMask;
 uint				MaxVisibilityNodeCount;
 float4				LightPosition;
 float				DepthBiasScale;
@@ -119,7 +120,14 @@ void MainCS(uint2 DispatchThreadId : SV_DispatchThreadID)
 		float3			SampleBaseColor;
 		float			SampleRoughness;
 		float			SampleSpecular;
-		DecodeNodeData(NodeData, SampleDepth, SampleVisData, SampleBaseColor, SampleRoughness, SampleSpecular);
+		uint			SampleLightChannelMask;
+		DecodeNodeData(NodeData, SampleDepth, SampleVisData, SampleBaseColor, SampleRoughness, SampleSpecular, SampleLightChannelMask);
+
+		if ((LightChannelMask & SampleLightChannelMask) == 0)
+		{
+			OutputColor[LocalOffset] = InitNullPackedHairTransmittanceMask();
+			return;
+		}
 
 		FGBufferData GBuffer;
 		GBuffer.ShadingModelID = SHADINGMODELID_HAIR;
@@ -340,7 +348,14 @@ void MainCS(uint2 DispatchThreadId : SV_DispatchThreadID)
 		float3			SampleBaseColor;
 		float			SampleRoughness;
 		float			SampleSpecular;
-		DecodeNodeData(NodeData, SampleDepth, SampleVisData, SampleBaseColor, SampleRoughness, SampleSpecular);
+		uint			SampleLightChannelMask;
+		DecodeNodeData(NodeData, SampleDepth, SampleVisData, SampleBaseColor, SampleRoughness, SampleSpecular, SampleLightChannelMask);
+
+		if ((LightChannelMask & SampleLightChannelMask) == 0)
+		{
+			OutputColor[LocalOffset] = InitNullPackedHairTransmittanceMask();
+			return;
+		}
 
 		FGBufferData GBuffer;
 		GBuffer.ShadingModelID = SHADINGMODELID_HAIR;
diff --git a/Engine/Shaders/Private/HairStrands/HairStrandsEnvironmentLighting.usf b/Engine/Shaders/Private/HairStrands/HairStrandsEnvironmentLighting.usf
--- a/Engine/Shaders/Private/HairStrands/HairStrandsEnvironmentLighting.usf
+++ b/Engine/Shaders/Private/HairStrands/HairStrandsEnvironmentLighting.usf
@@ -552,7 +552,8 @@ void MainCS(uint3 DispatchThreadId : SV_DispatchThreadID)
 	float3			SampleBaseColor;
 	float			SampleRoughness;
 	float			SampleSpecular;
-	DecodeNodeData(NodeData, SampleDepth, SampleVisData, SampleBaseColor, SampleRoughness, SampleSpecular);
+	uint			SampleLigthChannelMask;
+	DecodeNodeData(NodeData, SampleDepth, SampleVisData, SampleBaseColor, SampleRoughness, SampleSpecular, SampleLigthChannelMask);
 
 	FGBufferData GBuffer;
 	GBuffer.ShadingModelID = SHADINGMODELID_HAIR;
diff --git a/Engine/Shaders/Private/HairStrands/HairStrandsEnvironmentLightingCompose.usf b/Engine/Shaders/Private/HairStrands/HairStrandsEnvironmentLightingCompose.usf
--- a/Engine/Shaders/Private/HairStrands/HairStrandsEnvironmentLightingCompose.usf
+++ b/Engine/Shaders/Private/HairStrands/HairStrandsEnvironmentLightingCompose.usf
@@ -30,7 +30,7 @@ void MainPS(
 			const uint LocalOffset = NodeDesc.Offset + SampleIt;
 			const float4 LightingSample = HairLightingSampleBuffer[LocalOffset];
 			if (PixelCoverage >= 1)
-				OutColor += LightingSample;
+				OutColor += float4(LightingSample.xyz, 0);
 			else
 				OutColorSubPixel += LightingSample * PixelCoverage;
 		}		
diff --git a/Engine/Shaders/Private/HairStrands/HairStrandsGuideDeform.usf b/Engine/Shaders/Private/HairStrands/HairStrandsGuideDeform.usf
--- a/Engine/Shaders/Private/HairStrands/HairStrandsGuideDeform.usf
+++ b/Engine/Shaders/Private/HairStrands/HairStrandsGuideDeform.usf
@@ -1,6 +1,7 @@
 // Copyright Epic Games, Inc. All Rights Reserved.
 
 #include "../Common.ush"
+#include "HairStrandsMeshProjectionCommon.ush"
 
 float3 SimRestOffset;
 float3 SimDeformedOffset;
@@ -12,31 +13,68 @@ Buffer<uint4> SimRestPosePositionBuffer;
 Buffer<uint>  SimRootIndexBuffer;
 RWBuffer<uint4> OutSimDeformedPositionBuffer;
 
+
 #if PERMUTATION_GROUP_SIZE == 0
-	#define GROUP_SIZE 64
+#define GROUP_SIZE 64
 #elif PERMUTATION_GROUP_SIZE == 1
-	#define GROUP_SIZE 32
+#define GROUP_SIZE 32
 #else
-	#error Unknown group size	
+#error Unknown group size	
 #endif
 
-float3 Displacement(uint VertexIndex, float3 Pos, float Seed)
+
+#define DEFORMATION_BYPASS 0
+#define DEFORMATION_WAVE 1
+#define DEFORMATION_NORMAL_DIRECTION 2
+#define DEFORMATION_OFFSET_GUIDE 3
+#define DEFORMATION_OFFSET_GUIDE_DYNAMIC 4
+
+#if PERMUTATION_DEFORMATION == DEFORMATION_OFFSET_GUIDE_DYNAMIC
+
+Buffer<float4>  SimRestPosition0Buffer;
+Buffer<float4>  SimRestPosition1Buffer;
+Buffer<float4>  SimRestPosition2Buffer;
+
+Buffer<float4>  SimDeformedPosition0Buffer;
+Buffer<float4>  SimDeformedPosition1Buffer;
+Buffer<float4>  SimDeformedPosition2Buffer;
+
+Buffer<uint>	SimRootBarycentricBuffer;
+Buffer<uint>    SimVertexToRootIndexBuffer;
+
+// Return world position of a guide vertex deformed by a skinned triangle
+float3 DeformGuideByTriangle(
+	uint GuideIndex,
+	FHairMeshTriangle RestTriangle,
+	FHairMeshTriangle DeformedTriangle,
+	float3 RootBarycentric)
+{
+	const float3 WorldRestGuidePoint = f16tof32(SimRestPosePositionBuffer[GuideIndex].xyz) + SimRestOffset;
+	const float3 LocalRestGuidePoint = ToTriangle(WorldRestGuidePoint, RootBarycentric, RestTriangle);
+	const float3 WorldDeformedGuidePoint = ToWorld(LocalRestGuidePoint, RootBarycentric, DeformedTriangle);
+
+	return WorldDeformedGuidePoint;
+}
+
+#endif
+
+float3 DisplacePosition(uint VertexIndex, float3 Pos, float Seed)
 {
 // By pass
-#if PERMUTATION_DEFORMATION == 0
-	return 0;
+#if PERMUTATION_DEFORMATION == DEFORMATION_BYPASS
+	return Pos;
 #endif
 
 // Wave
-#if PERMUTATION_DEFORMATION == 1
+#if PERMUTATION_DEFORMATION == DEFORMATION_WAVE
 	const float Scale = 100.f;
 	Pos *= Scale;
 	//	return float3(0, cos(Pos.y + Seed / 16.f), cos(Pos.z + Seed / 32.f)) / Scale;
-	return float3(0, cos(Pos.y + Seed / 16.f), cos(Pos.z + Seed / 32.f));
+	return Pos + float3(0, cos(Pos.y + Seed / 16.f), cos(Pos.z + Seed / 32.f));
 #endif
 
 // Align with some normal
-#if PERMUTATION_DEFORMATION == 2
+#if PERMUTATION_DEFORMATION == DEFORMATION_NORMAL_DIRECTION
 	const uint RootIndex = SimRootIndexBuffer[VertexIndex];
 	const float3 V0 = f16tof32(SimRestPosePositionBuffer[RootIndex]);
 	const float3 V1 = f16tof32(SimRestPosePositionBuffer[RootIndex+1]);
@@ -46,11 +84,21 @@ float3 Displacement(uint VertexIndex, float3 Pos, float Seed)
 	const float3 Offset = V0 + D*T - Pos;
 
 	float Time = 0.5f * cos(Seed / 64.f) + 0.5f; // [0..1]
-	return lerp(0, Offset, Time);
+	return Pos + lerp(0, Offset, Time);
 #endif
 
-#if PERMUTATION_DEFORMATION == 3
-	return SimRestOffset - SimDeformedOffset;
+#if PERMUTATION_DEFORMATION == DEFORMATION_OFFSET_GUIDE
+	return Pos + (SimRestOffset - SimDeformedOffset);
+#endif
+
+#if PERMUTATION_DEFORMATION == DEFORMATION_OFFSET_GUIDE_DYNAMIC
+	const uint RootIndex = SimVertexToRootIndexBuffer[VertexIndex];
+	const float3 RootBarycentric = DecodeBarycentrics(SimRootBarycentricBuffer[RootIndex]);
+
+	const FHairMeshTriangle RestTriangle = GetTriangleTransformation(RootIndex, SimRestPosition0Buffer, SimRestPosition1Buffer, SimRestPosition2Buffer);
+	const FHairMeshTriangle DeformedTriangle = GetTriangleTransformation(RootIndex, SimDeformedPosition0Buffer, SimDeformedPosition1Buffer, SimDeformedPosition2Buffer);
+
+	return DeformGuideByTriangle(VertexIndex, RestTriangle, DeformedTriangle, RootBarycentric) - SimDeformedOffset;
 #endif
 }
 
@@ -60,14 +108,9 @@ void MainCS(uint2 DispatchThreadId : SV_DispatchThreadID)
 	const uint VertexIndex = DispatchThreadId.x + DispatchThreadId.y * DispatchCountX * GROUP_SIZE;
 	if (VertexIndex < VertexCount)
 	{
-		const uint4  Data 			= SimRestPosePositionBuffer[VertexIndex];
-		const float3 ControlPoint	= f16tof32(Data.xyz);
-
-		float3 PreviousOffset = 0;
-		{
-			const float3 Offset = Displacement(VertexIndex, ControlPoint, IterationCount);
-			OutSimDeformedPositionBuffer[VertexIndex] = uint4(f32tof16(ControlPoint + Offset), Data.w);
-			PreviousOffset = Offset;
-		}
+		const uint4  Data 				= SimRestPosePositionBuffer[VertexIndex];
+		const float3 ControlPoint		= f16tof32(Data.xyz);
+		const float3 DisplacedPosition	= DisplacePosition(VertexIndex, ControlPoint, IterationCount);
+		OutSimDeformedPositionBuffer[VertexIndex] = uint4(f32tof16(DisplacedPosition), Data.w);
 	}
 }
\ No newline at end of file
diff --git a/Engine/Shaders/Private/HairStrands/HairStrandsInterpolation.usf b/Engine/Shaders/Private/HairStrands/HairStrandsInterpolation.usf
--- a/Engine/Shaders/Private/HairStrands/HairStrandsInterpolation.usf
+++ b/Engine/Shaders/Private/HairStrands/HairStrandsInterpolation.usf
@@ -3,7 +3,7 @@
 #include "../Common.ush"
 #include "HairStrandsMeshProjectionCommon.ush"
 #include "HairStrandsVisibilityCommon.ush"
-#include "HairStrandsAABBCommon.ush"
+#include "HairStrandsAABBCommon.ush" 
 
 #if PERMUTATION_DEBUG == 1
 #include "../ShaderDrawDebug.ush"
@@ -46,31 +46,25 @@ RWBuffer<uint4> OutRenderDeformedPositionBuffer;
 #if PERMUTATION_DYNAMIC_GEOMETRY == 1
 
 // Compact all these buffers into 2 buffers: translation + quaternion
-float3			RestPositionOffset;
 Buffer<float4>  RestPosition0Buffer;
 Buffer<float4>  RestPosition1Buffer;
 Buffer<float4>  RestPosition2Buffer;
 
-float3			DeformedPositionOffset;
 Buffer<float4>  DeformedPosition0Buffer;
 Buffer<float4>  DeformedPosition1Buffer;
 Buffer<float4>  DeformedPosition2Buffer;
 
 Buffer<uint>	RootBarycentricBuffer;
-Buffer<uint>	RootToTriangleIndex;
 
-float3			SimRestPositionOffset;
 Buffer<float4>  SimRestPosition0Buffer;
 Buffer<float4>  SimRestPosition1Buffer;
 Buffer<float4>  SimRestPosition2Buffer;
 
-float3			SimDeformedPositionOffset;
 Buffer<float4>  SimDeformedPosition0Buffer;
 Buffer<float4>  SimDeformedPosition1Buffer;
 Buffer<float4>  SimDeformedPosition2Buffer;
 
 Buffer<uint>	SimRootBarycentricBuffer;
-Buffer<uint>	SimRootToTriangleIndex;
 #endif
 
 #if PERMUTATION_DEBUG == 1
@@ -217,8 +211,8 @@ void MainCS(uint2 DispatchThreadId : SV_DispatchThreadID)
 				const uint RootIndex = SimVertexToRootIndexBuffer[GuideIndex0];
 				const float3 RootBarycentric = DecodeBarycentrics(SimRootBarycentricBuffer[RootIndex]);
 
-				const FHairMeshTriangle RestTriangle		= GetTriangleTransformation(RootIndex, SimRestPosition0Buffer, SimRestPosition1Buffer, SimRestPosition2Buffer, SimRestPositionOffset);
-				const FHairMeshTriangle DeformedTriangle	= GetTriangleTransformation(RootIndex, SimDeformedPosition0Buffer, SimDeformedPosition1Buffer, SimDeformedPosition2Buffer, SimDeformedPositionOffset);
+				const FHairMeshTriangle RestTriangle		= GetTriangleTransformation(RootIndex, SimRestPosition0Buffer, SimRestPosition1Buffer, SimRestPosition2Buffer);
+				const FHairMeshTriangle DeformedTriangle	= GetTriangleTransformation(RootIndex, SimDeformedPosition0Buffer, SimDeformedPosition1Buffer, SimDeformedPosition2Buffer);
 
 				const float3 Offset0 = ComputeDynamicGeometryOffset(GuideIndex0, GuideVertexWeights[K], RestTriangle, DeformedTriangle, RootBarycentric);
 				const float3 Offset1 = ComputeDynamicGeometryOffset(GuideIndex1, GuideVertexWeights[K], RestTriangle, DeformedTriangle, RootBarycentric);
@@ -230,8 +224,8 @@ void MainCS(uint2 DispatchThreadId : SV_DispatchThreadID)
 			// Transform hair from rest post to deformed pose, based triangle deformation + simulation offset (optional depending of the simulation is actual running or not
 			const uint RootIndex						= RenVertexToRootIndexBuffer[VertexIndex];
 			const float3 RootBarycentric				= DecodeBarycentrics(RootBarycentricBuffer[RootIndex]);
-			const FHairMeshTriangle RestTriangle		= GetTriangleTransformation(RootIndex, RestPosition0Buffer, RestPosition1Buffer, RestPosition2Buffer, RestPositionOffset);
-			const FHairMeshTriangle DeformedTriangle	= GetTriangleTransformation(RootIndex, DeformedPosition0Buffer, DeformedPosition1Buffer, DeformedPosition2Buffer, DeformedPositionOffset);
+			const FHairMeshTriangle RestTriangle		= GetTriangleTransformation(RootIndex, RestPosition0Buffer, RestPosition1Buffer, RestPosition2Buffer);
+			const FHairMeshTriangle DeformedTriangle	= GetTriangleTransformation(RootIndex, DeformedPosition0Buffer, DeformedPosition1Buffer, DeformedPosition2Buffer);
 
 			ControlPoint = ToTriangle(ControlPoint,RootBarycentric,RestTriangle);
 			ControlPoint = ToWorld(ControlPoint,RootBarycentric,DeformedTriangle) + CurrOffset;
diff --git a/Engine/Shaders/Private/HairStrands/HairStrandsLightChannelMask.usf b/Engine/Shaders/Private/HairStrands/HairStrandsLightChannelMask.usf
new file mode 100644
--- /dev/null
+++ b/Engine/Shaders/Private/HairStrands/HairStrandsLightChannelMask.usf
@@ -0,0 +1,51 @@
+// Copyright Epic Games, Inc. All Rights Reserved.
+
+/*=============================================================================
+	ReflectionEnvironmentComputeShaders - functionality to apply local cubemaps.
+=============================================================================*/
+
+#include "../Common.ush"
+#include "HairStrandsCommon.ush"
+#include "HairStrandsVisibilityCommon.ush"
+
+uint2						OutputResolution;
+Texture2D<uint> 			NodeOffsetAndCount;
+StructuredBuffer<FNodeData> NodeData;
+RWTexture2D<uint>			OutLightChannelMaskTexture;
+
+#if PERMUTATION_GROUP_SIZE == 0
+#define GROUP_SIZE_X 8
+#define GROUP_SIZE_Y 8
+#elif PERMUTATION_GROUP_SIZE == 1
+#define GROUP_SIZE_X 8
+#define GROUP_SIZE_Y 4
+#else
+#error Unknown group size	
+#endif
+
+[numthreads(GROUP_SIZE_X, GROUP_SIZE_Y, 1)]
+void MainCS(uint3 DispatchThreadId : SV_DispatchThreadID)
+{
+	const uint2 PixelCoord = DispatchThreadId.xy;
+	if (PixelCoord.x >= OutputResolution.x || PixelCoord.y >= OutputResolution.y)
+		return;
+
+	const FNodeDesc NodeDesc = DecodeNodeDesc(NodeOffsetAndCount.Load(uint3(PixelCoord,0)));
+
+	uint OutLightChannelMask = 0;
+	for (uint SampleIt = 0; SampleIt < NodeDesc.Count; SampleIt++)
+	{
+		const uint LocalOffset = NodeDesc.Offset + SampleIt;
+		const FNodeData Data = NodeData[LocalOffset];
+		float			SampleDepth;
+		FVisibilityData SampleVisData;
+		float3			SampleBaseColor;
+		float			SampleRoughness;
+		float			SampleSpecular;
+		uint			SampleLightChannelMask;
+		DecodeNodeData(Data, SampleDepth, SampleVisData, SampleBaseColor, SampleRoughness, SampleSpecular, SampleLightChannelMask);
+		OutLightChannelMask |= SampleLightChannelMask;
+	}
+
+	OutLightChannelMaskTexture[PixelCoord] = OutLightChannelMask;
+}
diff --git a/Engine/Shaders/Private/HairStrands/HairStrandsMaterialPS.usf b/Engine/Shaders/Private/HairStrands/HairStrandsMaterialPS.usf
--- a/Engine/Shaders/Private/HairStrands/HairStrandsMaterialPS.usf
+++ b/Engine/Shaders/Private/HairStrands/HairStrandsMaterialPS.usf
@@ -14,6 +14,7 @@
 uint MaterialPass_MacroGroupId;
 uint MaterialPass_MaterialId;  
 uint MaterialPass_PrimitiveId;
+uint MaterialPass_LightChannelMask;
 RWStructuredBuffer<FNodeData>	OutNodeData;
 RWBuffer<float2>				OutNodeVelocity; 
 
@@ -220,7 +221,7 @@ FEvaluateOutput Evaluate(
 	}
 
 	FEvaluateOutput Out;
-	Out.NodeData		= ConvertToNodeData(SampleDepth, OutVisibilityData, OutBaseColor, OutRoughness, OutSpecular);
+	Out.NodeData		= ConvertToNodeData(SampleDepth, OutVisibilityData, OutBaseColor, OutRoughness, OutSpecular, MaterialPass_LightChannelMask);
 	Out.NodeVelocity	= EncodedVelocity;
 	return Out;
 }
diff --git a/Engine/Shaders/Private/HairStrands/HairStrandsMeshProjection.usf b/Engine/Shaders/Private/HairStrands/HairStrandsMeshProjection.usf
--- a/Engine/Shaders/Private/HairStrands/HairStrandsMeshProjection.usf
+++ b/Engine/Shaders/Private/HairStrands/HairStrandsMeshProjection.usf
@@ -3,6 +3,8 @@
 #include "../Common.ush"
 #include "HairStrandsMeshProjectionCommon.ush"
 
+#if SHADER_PROJECTION 
+
 uint bClear;
 uint MaxRootCount;
 float DistanceThreshold;
@@ -163,5 +165,225 @@ void MainCS(uint3 DispatchThreadId : SV_DispatchThreadID)
 		OutRootTriangleBarycentrics[RootIndex] = EncodedBarycentric;
 		OutRootTriangleDistance[RootIndex] = MinDistance;
 	}
+}
+
+#endif // SHADER_PROJECTION
+
+
+#if SHADER_MESHTRANSFER
+
+uint bNeedClear;
+
+uint Source_MeshPrimitiveCount_Iteration;
+uint Source_MeshMaxIndexCount;
+uint Source_MeshMaxVertexCount;
+uint Source_MeshIndexOffset;
+uint Source_MeshUVsChannelOffset;
+uint Source_MeshUVsChannelCount;
+
+Buffer<uint>   Source_MeshIndexBuffer;
+Buffer<float>  Source_MeshPositionBuffer;
+Buffer<float2> Source_MeshUVsBuffer;
+
+uint Target_MeshMaxVertexCount;
+uint Target_MeshUVsChannelOffset;
+uint Target_MeshUVsChannelCount;
+uint Target_SectionId;
+
+Buffer<uint>	Target_VertexSectionId;
+Buffer<float2>  Target_MeshUVsBuffer;
+RWBuffer<float> Target_MeshPositionBuffer;
+
+RWBuffer<uint>  OutDistanceBuffer;
+
+
+// Closest point on A triangle from another point
+// Code from the book "Real-Time Collision Detection" by Christer Ericson
+struct FTrianglePoint
+{
+	float3 P;
+	float3 Barycentric;
+};
+
+FTrianglePoint ComputeClosestPoint(const FHairMeshTriangleWithUVs Tri, const float2 UVs)
+{
+	const float3 A = float3(Tri.UV0, 0);
+	const float3 B = float3(Tri.UV1, 0);
+	const float3 C = float3(Tri.UV2, 0);
+	const float3 P = float3(UVs, 0);
+
+	// Check if P is in vertex region outside A.
+	float3 AB = B - A;
+	float3 AC = C - A;
+	float3 AP = P - A;
+	float D1 = dot(AB, AP);
+	float D2 = dot(AC, AP);
+	if (D1 <= 0.f && D2 <= 0.f)
+	{
+		FTrianglePoint Out;
+		Out.P = A;
+		Out.Barycentric = float3(1, 0, 0);
+		return Out;
+	}
+
+	// Check if P is in vertex region outside B.
+	float3 BP = P - B;
+	float D3 = dot(AB, BP);
+	float D4 = dot(AC, BP);
+	if (D3 >= 0.f && D4 <= D3)
+	{
+		FTrianglePoint Out;
+		Out.P = B;
+		Out.Barycentric = float3(0, 1, 0);
+		return Out;
+	}
+
+	// Check if P is in edge region of AB, and if so, return the projection of P onto AB.
+	float VC = D1 * D4 - D3 * D2;
+	if (VC <= 0.f && D1 >= 0.f && D3 <= 0.f)
+	{
+		float V = D1 / (D1 - D3);
+
+		FTrianglePoint Out;
+		Out.P = A + V * AB;
+		Out.Barycentric = float3(1 - V, V, 0);
+		return Out;
+	}
+
+	// Check if P is in vertex region outside C.
+	float3 CP = P - C;
+	float D5 = dot(AB, CP);
+	float D6 = dot(AC, CP);
+	if (D6 >= 0.f && D5 <= D6)
+	{
+		FTrianglePoint Out;
+		Out.P = C;
+		Out.Barycentric = float3(0, 0, 1);
+		return Out;
+	}
+
+	// Check if P is in edge region of AC, and if so, return the projection of P onto AC.
+	float VB = D5 * D2 - D1 * D6;
+	if (VB <= 0.f && D2 >= 0.f && D6 <= 0.f)
+	{
+		float W = D2 / (D2 - D6);
+		FTrianglePoint Out;
+		Out.P = A + W * AC;
+		Out.Barycentric = float3(1 - W, 0, W);
+		return Out;
+	}
+
+	// Check if P is in edge region of BC, and if so, return the projection of P onto BC.
+	float VA = D3 * D6 - D5 * D4;
+	if (VA <= 0.f && D4 - D3 >= 0.f && D5 - D6 >= 0.f)
+	{
+		float W = (D4 - D3) / (D4 - D3 + D5 - D6);
+		FTrianglePoint Out;
+		Out.P = B + W * (C - B);
+		Out.Barycentric = float3(0, 1 - W, W);
+		return Out;
+	}
+
+	// P must be inside the face region. Compute the closest point through its barycentric coordinates (u,V,W).
+	float Denom = 1.f / (VA + VB + VC);
+	float V = VB * Denom;
+	float W = VC * Denom;
+
+	FTrianglePoint Out;
+	Out.P = A + AB * V + AC * W;
+	Out.Barycentric = float3(1 - V - W, V, W);
+	return Out;
+}
+
+[numthreads(128, 1, 1)]
+void MainMeshTransferCS(uint3 DispatchThreadId : SV_DispatchThreadID)
+{
+	const float DistancePrecision = 100000;
+
+	const uint VertexIndex = DispatchThreadId.x;
+	if (VertexIndex >= Target_MeshMaxVertexCount) return;
+
+	if (bNeedClear)
+	{
+		Target_MeshPositionBuffer[VertexIndex * 3] = 0;
+		Target_MeshPositionBuffer[VertexIndex * 3 + 1] = 0;
+		Target_MeshPositionBuffer[VertexIndex * 3 + 2] = 0;
+
+		OutDistanceBuffer[VertexIndex] = ~0;
+	}
+
+	if (Target_VertexSectionId[VertexIndex] != Target_SectionId) return;
+
+	float ClosestDistance = OutDistanceBuffer[VertexIndex] / DistancePrecision;
+	float OriginalClosestDistance = ClosestDistance;
+
+	const uint UVIndex = Target_MeshUVsChannelOffset + Target_MeshUVsChannelCount * VertexIndex;
+	const float2 TargetUV = Target_MeshUVsBuffer[UVIndex];
+
+	float3 VertexPosition = 0;
+	for (uint PrimitiveIt = 0; PrimitiveIt < Source_MeshPrimitiveCount_Iteration; ++PrimitiveIt)
+	{
+		const FHairMeshTriangleWithUVs SourceTri = GetTriangleData(
+			PrimitiveIt,
+			Source_MeshPositionBuffer,
+			Source_MeshUVsBuffer,
+			Source_MeshIndexBuffer,
+			Source_MeshIndexOffset,
+			Source_MeshMaxIndexCount,
+			Source_MeshMaxVertexCount,
+			Source_MeshUVsChannelOffset,
+			Source_MeshUVsChannelCount);
+
+		// Project on triangle and find the closest distance
+		{
+			const FTrianglePoint ClosestPoint = ComputeClosestPoint(SourceTri, TargetUV);
+			const float DistanceToTriangle = distance(ClosestPoint.P.xy, TargetUV);
+			if (DistanceToTriangle < ClosestDistance)
+			{
+				VertexPosition = 
+					ClosestPoint.Barycentric.x * SourceTri.P0 + 
+					ClosestPoint.Barycentric.y * SourceTri.P1 + 
+					ClosestPoint.Barycentric.z * SourceTri.P2;
+				ClosestDistance = DistanceToTriangle;
+			}
+		}
+	}
+
+	if (ClosestDistance < OriginalClosestDistance)
+	{
+		Target_MeshPositionBuffer[VertexIndex * 3] = VertexPosition.x;
+		Target_MeshPositionBuffer[VertexIndex * 3 + 1] = VertexPosition.y;
+		Target_MeshPositionBuffer[VertexIndex * 3 + 2] = VertexPosition.z;
+		OutDistanceBuffer[VertexIndex] = uint(ClosestDistance * DistancePrecision);
+	}
+}
+
+#endif // SHADER_TRANSFER
+
+#if SHADER_SECTIONID
+
+uint MeshSectionId;
+uint MeshSectionPrimitiveCount;
+uint MeshMaxIndexCount;
+uint MeshMaxVertexCount;
+uint MeshIndexOffset;
+
+Buffer<uint>	MeshIndexBuffer;
+RWBuffer<uint>  OutVertexSectionId;
+
+[numthreads(128, 1, 1)]
+void MainMarkSectionIdCS(uint3 DispatchThreadId : SV_DispatchThreadID)
+{
+	const uint Index = DispatchThreadId.x;
+	if (Index >= MeshSectionPrimitiveCount*3) return;
+
+	if (Index + MeshIndexOffset >= MeshMaxIndexCount) return;
+
+	const uint VertexIndex = MeshIndexBuffer[Index + MeshIndexOffset];
+	if (VertexIndex < MeshMaxVertexCount)
+	{
+		OutVertexSectionId[VertexIndex] = MeshSectionId;
+	}
+}
 
-}
\ No newline at end of file
+#endif // SHADER_SECTIONID
diff --git a/Engine/Shaders/Private/HairStrands/HairStrandsMeshProjectionCommon.ush b/Engine/Shaders/Private/HairStrands/HairStrandsMeshProjectionCommon.ush
--- a/Engine/Shaders/Private/HairStrands/HairStrandsMeshProjectionCommon.ush
+++ b/Engine/Shaders/Private/HairStrands/HairStrandsMeshProjectionCommon.ush
@@ -12,6 +12,46 @@ struct FHairMeshTriangle
 	float3 N;
 };
 
+struct FHairMeshTriangleWithUVs : FHairMeshTriangle
+{
+	float2 UV0;
+	float2 UV1;
+	float2 UV2;
+};
+
+uint PackUVs(float2 UV)
+{
+	return (f32tof16(UV.x) & 0xFFFF) | ((f32tof16(UV.y) & 0xFFFF) << 16);
+}
+
+float PackUVsToFloat(float2 UV)
+{
+	return asfloat(PackUVs(UV));
+}
+
+float2 UnPackUVs(uint InUV)
+{
+	float2 Out;
+	Out.x = f16tof32(InUV & 0xFFFF);
+	Out.y = f16tof32((InUV >> 16) & 0xFFFF);
+	return Out;
+}
+
+float2 UnPackUVs(float InUV)
+{
+	return UnPackUVs(asuint(InUV));
+}
+
+float2 GetVertexUV(uint VId, Buffer<float2> InUVBuffer, Buffer<uint> InIndexBuffer, uint InIndexOffset, uint InMaxIndexCount, uint InMaxVertexCount, uint InMeshUVsChannelOffset, uint InMeshUVsChannelCount)
+{
+	const uint Index = min(InIndexOffset + VId, InMaxIndexCount - 1);
+	uint VertexIndex = InIndexBuffer.Load(Index);
+	VertexIndex = min(VertexIndex, InMaxVertexCount - 1);
+
+	const uint UVIndex = InMeshUVsChannelOffset + InMeshUVsChannelCount * VertexIndex;
+	return InUVBuffer.Load(UVIndex);
+}
+
 float3 GetVertexPosition(uint VId, Buffer<float> InPositionBuffer, Buffer<uint> InIndexBuffer, uint InIndexOffset, uint InMaxIndexCount, uint InMaxVertexCount)
 {
 	const uint Index = min(InIndexOffset + VId, InMaxIndexCount-1);
@@ -43,19 +83,40 @@ FHairMeshTriangle GetTriangleData(uint TriangleIndex, Buffer<float> InPositionBu
 	return Out;
 }
 
+FHairMeshTriangleWithUVs GetTriangleData(uint TriangleIndex, Buffer<float> InPositionBuffer, Buffer<float2> InUVBuffer, Buffer<uint> InIndexBuffer, uint InIndexOffset, uint InMaxIndexCount, uint InMaxVertexCount, uint InMeshUVsChannelOffset, uint InMeshUVsChannelCount)
+{
+	const uint BaseVertexId = TriangleIndex * 3;
+	const uint VId0 = BaseVertexId;
+	const uint VId1 = BaseVertexId + 1;
+	const uint VId2 = BaseVertexId + 2;
+
+	FHairMeshTriangleWithUVs Out;
+	Out.P0 = GetVertexPosition(VId0, InPositionBuffer, InIndexBuffer, InIndexOffset, InMaxIndexCount, InMaxVertexCount);
+	Out.P1 = GetVertexPosition(VId1, InPositionBuffer, InIndexBuffer, InIndexOffset, InMaxIndexCount, InMaxVertexCount);
+	Out.P2 = GetVertexPosition(VId2, InPositionBuffer, InIndexBuffer, InIndexOffset, InMaxIndexCount, InMaxVertexCount);
+	Out.E0 = Out.P1 - Out.P0;
+	Out.E1 = Out.P2 - Out.P0;
+	Out.N = normalize(cross(Out.E1, Out.E0));
+
+	Out.UV0 = GetVertexUV(VId0, InUVBuffer, InIndexBuffer, InIndexOffset, InMaxIndexCount, InMaxVertexCount, InMeshUVsChannelOffset, InMeshUVsChannelCount);
+	Out.UV1 = GetVertexUV(VId1, InUVBuffer, InIndexBuffer, InIndexOffset, InMaxIndexCount, InMaxVertexCount, InMeshUVsChannelOffset, InMeshUVsChannelCount);
+	Out.UV2 = GetVertexUV(VId2, InUVBuffer, InIndexBuffer, InIndexOffset, InMaxIndexCount, InMaxVertexCount, InMeshUVsChannelOffset, InMeshUVsChannelCount);
+	
+	return Out;
+}
+
 FHairMeshTriangle GetTriangleTransformation(
 	uint RootIndex,
 	Buffer<float4>  InPosition0Buffer,
 	Buffer<float4>  InPosition1Buffer,
-	Buffer<float4>  InPosition2Buffer,
-	float3 CenterOffset)
+	Buffer<float4>  InPosition2Buffer)
 {
 	FHairMeshTriangle Out;
 
 	// Build orthonormal frame
-	Out.P0 = InPosition0Buffer[RootIndex].xyz + CenterOffset;
-	Out.P1 = InPosition1Buffer[RootIndex].xyz + CenterOffset;
-	Out.P2 = InPosition2Buffer[RootIndex].xyz + CenterOffset;
+	Out.P0 = InPosition0Buffer[RootIndex].xyz;
+	Out.P1 = InPosition1Buffer[RootIndex].xyz;
+	Out.P2 = InPosition2Buffer[RootIndex].xyz;
 	Out.E0 = Out.P1 - Out.P0;
 	Out.E1 = Out.P2 - Out.P0;
 
@@ -66,6 +127,36 @@ FHairMeshTriangle GetTriangleTransformation(
 	return Out;
 }
 
+FHairMeshTriangleWithUVs GetTriangleWithUVsTransformation(
+	uint RootIndex,
+	Buffer<float4>  InPosition0Buffer,
+	Buffer<float4>  InPosition1Buffer,
+	Buffer<float4>  InPosition2Buffer)
+{
+	FHairMeshTriangleWithUVs Out;
+
+	const float4 V0 = InPosition0Buffer[RootIndex];
+	const float4 V1 = InPosition1Buffer[RootIndex];
+	const float4 V2 = InPosition2Buffer[RootIndex];
+
+	// Build orthonormal frame
+	Out.P0 = V0.xyz;
+	Out.P1 = V1.xyz;
+	Out.P2 = V2.xyz;
+	Out.E0 = Out.P1 - Out.P0;
+	Out.E1 = Out.P2 - Out.P0;
+
+	Out.N = normalize(cross(Out.E1, Out.E0));
+	Out.E0 = normalize(Out.E0);
+	Out.E1 = normalize(cross(Out.E0, Out.N));
+
+	Out.UV0 = UnPackUVs(V0.w);
+	Out.UV1 = UnPackUVs(V1.w);
+	Out.UV2 = UnPackUVs(V2.w);
+
+	return Out;
+}
+
 float3 ToTriangle(const float3 P, float3 B, const FHairMeshTriangle T)
 {
 	const float3 LocalP = P - (T.P0 * B.x + T.P1 * B.y + T.P2 * B.z);
diff --git a/Engine/Shaders/Private/HairStrands/HairStrandsMeshProjectionHairDebug.usf b/Engine/Shaders/Private/HairStrands/HairStrandsMeshProjectionHairDebug.usf
--- a/Engine/Shaders/Private/HairStrands/HairStrandsMeshProjectionHairDebug.usf
+++ b/Engine/Shaders/Private/HairStrands/HairStrandsMeshProjectionHairDebug.usf
@@ -3,13 +3,14 @@
 //#define SCENE_TEXTURES_DISABLED 1
 
 #include "../Common.ush"
-#include "HairStrandsMeshProjectionCommon.ush"
+#include "HairStrandsDebugCommon.ush"
+#include "HairStrandsMeshProjectionCommon.ush" 
 
 #define INPUT_ROOT_FRAME 0
 #define INPUT_ROOT_TRIANGLE 1
 uint MaxRootCount;
 float4x4 RootLocalToWorld;
-uint			DeformedFrameEnable;
+uint DeformedFrameEnable;
 
 #if PERMUTATION_INPUT_TYPE == INPUT_ROOT_FRAME
 Buffer<float4> RootPositionBuffer;
@@ -17,23 +18,24 @@ Buffer<float4> RootNormalBuffer;
 Buffer<uint>   RootBarycentricBuffer;
 #endif
 
-float3 RestPositionOffset;
 Buffer<float4>  RestPosition0Buffer;
 Buffer<float4>  RestPosition1Buffer;
 Buffer<float4>  RestPosition2Buffer;
 
-float3 DeformedPositionOffset;
 Buffer<float4>  DeformedPosition0Buffer;
 Buffer<float4>  DeformedPosition1Buffer;
 Buffer<float4>  DeformedPosition2Buffer;
 
 void MainVS(
 	uint VertexId : SV_VertexID,
+	out float2 OutVertexUVs : VERTEXUVs,
 	out float3 OutFaceNormal : FACE_NORMAL,
 	out float4 OutPosition : SV_POSITION)
 {
 	ResolvedView = ResolveView();
 
+	float2 UVs = 0;
+
 #if PERMUTATION_INPUT_TYPE == INPUT_ROOT_FRAME  
 	const uint LineVertex = (VertexId%2);
 	const uint LineIndex = VertexId/2;
@@ -45,14 +47,32 @@ void MainVS(
 	float3 N;
 	float3 P;
 	{
-		const FHairMeshTriangle RestTriangle     = GetTriangleTransformation(LineIndex, RestPosition0Buffer, RestPosition1Buffer, RestPosition2Buffer, RestPositionOffset);
-		const FHairMeshTriangle DeformedTriangle = GetTriangleTransformation(LineIndex, DeformedPosition0Buffer, DeformedPosition1Buffer, DeformedPosition2Buffer, DeformedPositionOffset);
+		const FHairMeshTriangleWithUVs RestTriangleWithUVs     = GetTriangleWithUVsTransformation(LineIndex, RestPosition0Buffer, RestPosition1Buffer, RestPosition2Buffer);
+		const FHairMeshTriangleWithUVs DeformedTriangleWithUVs = GetTriangleWithUVsTransformation(LineIndex, DeformedPosition0Buffer, DeformedPosition1Buffer, DeformedPosition2Buffer);
+
+		const FHairMeshTriangle RestTriangle			= GetTriangleTransformation(LineIndex, RestPosition0Buffer, RestPosition1Buffer, RestPosition2Buffer);
+		const FHairMeshTriangle DeformedTriangle		= GetTriangleTransformation(LineIndex, DeformedPosition0Buffer, DeformedPosition1Buffer, DeformedPosition2Buffer);
+
 
 		float3 P0 = RootPositionBuffer[LineIndex].xyz;
 		float3 LN = RootNormalBuffer[LineIndex].xyz;
 		float3 B  = DecodeBarycentrics(RootBarycentricBuffer[LineIndex]);
 		float3 P1 = P0 + LN;
 
+		if (DeformedFrameEnable > 0)
+		{
+			UVs = 
+				B.x * DeformedTriangleWithUVs.UV0 +
+				B.y * DeformedTriangleWithUVs.UV1 +
+				B.z * DeformedTriangleWithUVs.UV2;
+		}
+		else
+		{
+			UVs =
+				B.x * RestTriangleWithUVs.UV0 +
+				B.y * RestTriangleWithUVs.UV1 +
+				B.z * RestTriangleWithUVs.UV2;
+		}
 
 		if (DeformedFrameEnable > 0)
 		{
@@ -74,10 +94,10 @@ void MainVS(
 
 	if (TriangleIndex > MaxRootCount) return;
 
-	const FHairMeshTriangle RestTriangle     = GetTriangleTransformation(TriangleIndex, RestPosition0Buffer, RestPosition1Buffer, RestPosition2Buffer, RestPositionOffset);
-	const FHairMeshTriangle DeformedTriangle = GetTriangleTransformation(TriangleIndex, DeformedPosition0Buffer, DeformedPosition1Buffer, DeformedPosition2Buffer, DeformedPositionOffset);
+	const FHairMeshTriangleWithUVs RestTriangle		= GetTriangleWithUVsTransformation(TriangleIndex, RestPosition0Buffer, RestPosition1Buffer, RestPosition2Buffer);
+	const FHairMeshTriangleWithUVs DeformedTriangle = GetTriangleWithUVsTransformation(TriangleIndex, DeformedPosition0Buffer, DeformedPosition1Buffer, DeformedPosition2Buffer);
 
-	FHairMeshTriangle Tri;
+	FHairMeshTriangleWithUVs Tri;
 	if (DeformedFrameEnable > 0)
 		Tri = DeformedTriangle;
 	else
@@ -86,12 +106,21 @@ void MainVS(
 	float3 N = Tri.N;
 	float3 P = 0;
 	if (TriangleVertex == 0)
+	{
 		P = Tri.P0;
+		UVs = Tri.UV0;
+	}
 	if (TriangleVertex == 1)
+	{
 		P = Tri.P1;
+		UVs = Tri.UV1;
+	}
 	if (TriangleVertex == 2)
+	{
 		P = Tri.P2;
-	
+		UVs = Tri.UV2;
+	}
+
 	// Debug: To test if the rest-triangle-to-deformed-triangle transformation works correctly
 	#if 0
 	if (DeformedFrameEnable > 0)
@@ -119,15 +148,18 @@ void MainVS(
 	P = mul(float4(P, 1), RootLocalToWorld).xyz;
 #endif
 
+	OutVertexUVs = UVs;
 	OutFaceNormal = N;
 	OutPosition = mul(float4(P,1), View.WorldToClip);
 }
 
 void MainPS(
+	in float2 VertexUVs : VERTEXUVs,
 	in float3 FaceNormal : FACE_NORMAL,
 	out float4 OutColor : SV_Target0)
 {
 	FaceNormal += float3(1,1,1);
 	FaceNormal *= 0.5f;
 	OutColor = float4(FaceNormal,1);
+	OutColor = float4(VertexUVs,0, 1);
 }
\ No newline at end of file
diff --git a/Engine/Shaders/Private/HairStrands/HairStrandsMeshProjectionMeshDebug.usf b/Engine/Shaders/Private/HairStrands/HairStrandsMeshProjectionMeshDebug.usf
--- a/Engine/Shaders/Private/HairStrands/HairStrandsMeshProjectionMeshDebug.usf
+++ b/Engine/Shaders/Private/HairStrands/HairStrandsMeshProjectionMeshDebug.usf
@@ -4,21 +4,34 @@
 
 #include "../Common.ush"
 #include "HairStrandsMeshProjectionCommon.ush"
+#include "HairStrandsDebugCommon.ush"
 
 float4x4 LocalToWorld;
 uint VertexOffset;
 uint IndexOffset;
 uint MaxIndexCount;
 uint MaxVertexCount;
+uint MeshUVsChannelOffset;
+uint MeshUVsChannelCount;
+uint bOutputInUVsSpace;
+uint MeshType;
+uint SectionIndex;
+float2 OutputResolution;
 
-Buffer<float> InputVertexBuffer;
+Buffer<float> InputVertexPositionBuffer;
+Buffer<float2> InputVertexUVsBuffer;
 #if PERMUTATION_INPUT_TYPE == 1
 Buffer<uint> InputIndexBuffer;
 #endif
 
+#define REST_MESH 0
+#define DEFORMED_MESH 1
+#define SOURCE_MESH 2
+#define TARGET_MESH 3
 
 void MainVS(
 	uint VertexId : SV_VertexID,
+	out float2 OutVertexUV : TRI_UVs,
 	out float3 OutFaceNormal : FACE_NORMAL,
 	out float4 OutPosition : SV_POSITION)
 {
@@ -26,34 +39,80 @@ void MainVS(
 
 #if PERMUTATION_INPUT_TYPE == 0
 	const uint VertexIndex = min(VertexOffset + VertexId, MaxVertexCount-1);
-	float3 P = InputVertexBuffer[VertexIndex];
+	float3 P = InputVertexPositionBuffer[VertexIndex];
 	OutFaceNormal = float3(0,0,1);
+	OutVertexUV = InputVertexUVsBuffer[VertexIndex];
 #endif
 
 #if PERMUTATION_INPUT_TYPE == 1
 	const uint TriangleVertexIndex = (VertexId%3);
 	const uint TriangleIndex = VertexId / 3;
-	const FHairMeshTriangle Tri = GetTriangleData(TriangleIndex, InputVertexBuffer, InputIndexBuffer, IndexOffset, MaxIndexCount, MaxVertexCount);
+	const FHairMeshTriangleWithUVs Tri = GetTriangleData(TriangleIndex, InputVertexPositionBuffer, InputVertexUVsBuffer, InputIndexBuffer, IndexOffset, MaxIndexCount, MaxVertexCount, MeshUVsChannelOffset, MeshUVsChannelCount);
 	OutFaceNormal = Tri.N;
 
 	float3 P = 0;
+	float2 UV = 0;
 	if (TriangleVertexIndex == 0)
+	{
 		P = Tri.P0;
+		UV = Tri.UV0;
+	}
 	if (TriangleVertexIndex == 1)
+	{
 		P = Tri.P1;
+		UV = Tri.UV1;
+	}
 	if (TriangleVertexIndex == 2)
+	{
 		P = Tri.P2;
+		UV = Tri.UV2;
+	}
+
+	OutVertexUV = UV;
+#endif
+	if (bOutputInUVsSpace)
+	{
+		float2 PixelOffset = 0;
+		//float2 PixelOffset = 1.f / float2(OutputResolution);
+		if (MeshType == REST_MESH)		PixelOffset *= 0;
+		if (MeshType == DEFORMED_MESH)	PixelOffset *= 1;
+		if (MeshType == SOURCE_MESH)	PixelOffset *= 2;
+		if (MeshType == TARGET_MESH)	PixelOffset *= 3;
+
+		OutPosition = float4((PixelOffset + OutVertexUV) * 2 - 1, 0.5, 1);
+	}
+	else
+	{
+		P = mul(float4(P, 1), LocalToWorld).xyz;
+		OutPosition = mul(float4(P, 1), View.WorldToClip);
+	}
+
+	// Quick filter
+#if 0
+	//if (MeshType != REST_MESH)
+	//if (MeshType != DEFORMED_MESH)
+	//if (MeshType != SOURCE_MESH)
+	//if (MeshType != TARGET_MESH)
+	if (MeshType != SOURCE_MESH && MeshType != TARGET_MESH || SectionIndex != 0)
+	{
+		OutPosition = float4(1.#INF, 1.#INF, 1.#INF, 1.#INF);
+		OutVertexUV = 0;
+	}
 #endif
-	// TODO: the normal is in local space, not in world space so it won't appear correct
-	P = mul(float4(P, 1), LocalToWorld).xyz;
-	OutPosition = mul(float4(P,1), View.WorldToClip);
 }
 
 void MainPS(
+	in float2 VertexUV : TRI_UVs,
 	in float3 FaceNormal : FACE_NORMAL,
 	out float4 OutColor : SV_Target0)
 {
-	FaceNormal += float3(1,1,1);
-	FaceNormal *= 0.5f;
-	OutColor = float4(FaceNormal,1);
+	float3 Color = 0;
+	if (MeshType == REST_MESH)		Color = float3(1, 0, 0);
+	if (MeshType == DEFORMED_MESH)	Color = float3(0, 1, 0);
+	if (MeshType == SOURCE_MESH)	Color = float3(0, 0, 1);
+	if (MeshType == TARGET_MESH)	Color = float3(1, 1, 0);
+
+	OutColor = float4((FaceNormal + 1) * 0.5f,1);
+	OutColor = float4(VertexUV, 0, 1);
+	OutColor = float4(Color, 1);
 }
\ No newline at end of file
diff --git a/Engine/Shaders/Private/HairStrands/HairStrandsMeshUpdate.usf b/Engine/Shaders/Private/HairStrands/HairStrandsMeshUpdate.usf
--- a/Engine/Shaders/Private/HairStrands/HairStrandsMeshUpdate.usf
+++ b/Engine/Shaders/Private/HairStrands/HairStrandsMeshUpdate.usf
@@ -3,14 +3,20 @@
 #include "../Common.ush"
 #include "HairStrandsMeshProjectionCommon.ush"
 
-float3 RootTrianglePositionOffset;
 uint MeshSectionIndex;
 uint MeshMaxIndexCount;
 uint MeshMaxVertexCount;
 uint MeshIndexOffset;
+uint MeshUVsChannelOffset;
+uint MeshUVsChannelCount;
+
 Buffer<uint> MeshIndexBuffer;
 Buffer<float> MeshPositionBuffer;
 
+#if PERMUTATION_WITHUV == 1
+Buffer<float2> MeshUVsBuffer;
+#endif
+
 uint MaxRootCount;
 Buffer<uint> RootTriangleIndex;
 
@@ -29,9 +35,17 @@ void MainCS(uint3 DispatchThreadId : SV_DispatchThreadID)
 	if (TriangleIndex.SectionIndex != MeshSectionIndex)
 		return;
 
+#if PERMUTATION_WITHUV == 1
+	const FHairMeshTriangleWithUVs Tri = GetTriangleData(TriangleIndex.TriangleIndex, MeshPositionBuffer, MeshUVsBuffer, MeshIndexBuffer, MeshIndexOffset, MeshMaxIndexCount, MeshMaxVertexCount, MeshUVsChannelOffset, MeshUVsChannelCount);
+
+	OutRootTrianglePosition0[RootIndex] = float4(Tri.P0, PackUVsToFloat(Tri.UV0));
+	OutRootTrianglePosition1[RootIndex] = float4(Tri.P1, PackUVsToFloat(Tri.UV1));
+	OutRootTrianglePosition2[RootIndex] = float4(Tri.P2, PackUVsToFloat(Tri.UV2));
+#else
 	const FHairMeshTriangle Tri = GetTriangleData(TriangleIndex.TriangleIndex, MeshPositionBuffer, MeshIndexBuffer, MeshIndexOffset, MeshMaxIndexCount, MeshMaxVertexCount);
 
-	OutRootTrianglePosition0[RootIndex] = (float4(Tri.P0 - RootTrianglePositionOffset,1));
-	OutRootTrianglePosition1[RootIndex] = (float4(Tri.P1 - RootTrianglePositionOffset,1));
-	OutRootTrianglePosition2[RootIndex] = (float4(Tri.P2 - RootTrianglePositionOffset,1));
+	OutRootTrianglePosition0[RootIndex] = (float4(Tri.P0, 1));
+	OutRootTrianglePosition1[RootIndex] = (float4(Tri.P1, 1));
+	OutRootTrianglePosition2[RootIndex] = (float4(Tri.P2, 1));
+#endif
 }
\ No newline at end of file
diff --git a/Engine/Shaders/Private/HairStrands/HairStrandsVisibilityCommon.ush b/Engine/Shaders/Private/HairStrands/HairStrandsVisibilityCommon.ush
--- a/Engine/Shaders/Private/HairStrands/HairStrandsVisibilityCommon.ush
+++ b/Engine/Shaders/Private/HairStrands/HairStrandsVisibilityCommon.ush
@@ -37,7 +37,7 @@ struct FNodeData
 	uint  PrimitiveID_MacroGroupID;
 	uint  Tangent_Coverage8bit;
 	uint  BaseColor_Roughness;
-	float Specular;
+	uint  Specular_LightChannels;
 };
 
 float3 EncodeTangent(float3 N)
@@ -136,7 +136,7 @@ FNodeDesc DecodeNodeDesc(uint In)
 	return Out;
 }
 
-FNodeData ConvertToNodeData(float Depth, FVisibilityData VisibilityData, float3 BaseColor, float Roughness, float Specular)
+FNodeData ConvertToNodeData(float Depth, FVisibilityData VisibilityData, float3 BaseColor, float Roughness, float Specular, uint LightChannelMask)
 {
 	// 4x16bits
 	const uint4 EncodedVisibilityData = EncodeVisibilityData(VisibilityData);
@@ -146,7 +146,7 @@ FNodeData ConvertToNodeData(float Depth, FVisibilityData VisibilityData, float3
 	Out.Tangent_Coverage8bit	= (0xFFFF & EncodedVisibilityData.x) | ((0xFFFF & EncodedVisibilityData.y) << 16);	// 32bits uint
 	Out.PrimitiveID_MacroGroupID= (0xFFFF & EncodedVisibilityData.z) | ((0xFFFF & EncodedVisibilityData.w) << 16);	// 32bits uint
 	Out.BaseColor_Roughness		= float4ToUint(float4(sqrt(BaseColor), Roughness));									// 32bits uint
-	Out.Specular				= Specular;																			// 32bits float
+	Out.Specular_LightChannels	= clamp(uint(Specular * 0xFF), 0, 0xFF) | ((LightChannelMask & 0x7) << 8);
 
 	return Out;
 }
@@ -157,7 +157,8 @@ FNodeData ConvertToNodeData(float Depth, FVisibilityData VisibilityData, float3
 	 inout FVisibilityData	OutVisibilityData,
 	 inout float3			OutBaseColor, 
 	 inout float			OutRoughness,
-	 inout float			OutSpecular)
+	 inout float			OutSpecular,
+	 inout uint				OutLightChannelMask)
 {
 	OutDepth = InNode.Depth;
 
@@ -171,7 +172,8 @@ FNodeData ConvertToNodeData(float Depth, FVisibilityData VisibilityData, float3
 	const float4 BaseColorAndRoughness = UintToFloat4(InNode.BaseColor_Roughness);
 	OutBaseColor = BaseColorAndRoughness.xyz * BaseColorAndRoughness.xyz;
 	OutRoughness = BaseColorAndRoughness.w;
-	OutSpecular = InNode.Specular;
+	OutSpecular			= float(0xFF & InNode.Specular_LightChannels) / 0xFF;
+	OutLightChannelMask = (InNode.Specular_LightChannels>>8) & 0x7;
 }
 
 #define HAIR_VISIBILITY_MAXSAMPLECOUNT 8
@@ -337,7 +339,7 @@ FNodeData ConvertToNodeData(float Depth, FVisibilityData VisibilityData, float3
  struct FPPLLNodeData
  {
 	 float Depth;
-	 uint Specular_Coverage16bit;
+	 uint Specular_LightChannel_Coverage16bit;
 	 uint PrimitiveID_MacroGroupID;
 	 uint Tangent_Coverage8bit;
 	 uint BaseColor_Roughness;
@@ -355,7 +357,7 @@ FNodeData ConvertToNodeData(float Depth, FVisibilityData VisibilityData, float3
 	 return float2(f16tof32(PackedVelocity >> 16), f16tof32(PackedVelocity));
  }
 
- FPPLLNodeData ConvertToPPLLNodeData(float Depth, float CoverageFullPrecision, FVisibilityData VisibilityData, float3 BaseColor, float Roughness, float Specular, float2 Velocity, uint NextNodeIndex)
+ FPPLLNodeData ConvertToPPLLNodeData(float Depth, float CoverageFullPrecision, FVisibilityData VisibilityData, float3 BaseColor, float Roughness, float Specular, uint LightChannelMask, float2 Velocity, uint NextNodeIndex)
  {
 	 // 4x16bits
 	 const uint4 EncodedVisibilityData = EncodeVisibilityData(VisibilityData);
@@ -366,7 +368,7 @@ FNodeData ConvertToNodeData(float Depth, FVisibilityData VisibilityData, float3
 	 Out.Tangent_Coverage8bit = (0xFFFF & EncodedVisibilityData.x) | ((0xFFFF & EncodedVisibilityData.y) << 16);
 	 Out.PrimitiveID_MacroGroupID = (0xFFFF & EncodedVisibilityData.z) | ((0xFFFF & EncodedVisibilityData.w) << 16);
 	 Out.BaseColor_Roughness = float4ToUint(float4(sqrt(BaseColor), Roughness));
-	 Out.Specular_Coverage16bit = clamp(uint(Specular * 0xFFFF), 0, 0xFFFF) | ((0xFFFF & Coverage16bit) << 16);
+	 Out.Specular_LightChannel_Coverage16bit = clamp(uint(Specular * 0xFF), 0, 0xFF) | ((LightChannelMask & 0x7)<<8) | ((0xFFFF & Coverage16bit) << 16);
 	 Out.NextNodeIndex = NextNodeIndex;
 	 Out.PackedVelocity = PackVelocity(Velocity);
 
@@ -380,9 +382,9 @@ FNodeData ConvertToNodeData(float Depth, FVisibilityData VisibilityData, float3
 	 Out.PrimitiveID_MacroGroupID = In.PrimitiveID_MacroGroupID;
 	 Out.Tangent_Coverage8bit = In.Tangent_Coverage8bit;
 	 Out.BaseColor_Roughness = In.BaseColor_Roughness;
-	 Out.Specular = float(In.Specular_Coverage16bit & 0xFFFF) / float(0xFFFF);
+	 Out.Specular_LightChannels = In.Specular_LightChannel_Coverage16bit & 0xFFFF;
 
-	 OutCoverage16bit = (In.Specular_Coverage16bit >> 16) & 0xFFFF;
+	 OutCoverage16bit = (In.Specular_LightChannel_Coverage16bit >> 16) & 0xFFFF;
 
 	 return Out;
  }
diff --git a/Engine/Shaders/Private/HairStrands/HairStrandsVisibilityCompaction.usf b/Engine/Shaders/Private/HairStrands/HairStrandsVisibilityCompaction.usf
--- a/Engine/Shaders/Private/HairStrands/HairStrandsVisibilityCompaction.usf
+++ b/Engine/Shaders/Private/HairStrands/HairStrandsVisibilityCompaction.usf
@@ -110,8 +110,9 @@ void MainCS(uint GroupIndex : SV_GroupIndex, uint3 DispatchThreadId : SV_Dispatc
 			float3			BaseColor = 0.0f;
 			float			Roughness = 0.0f;
 			float			Specular = 0;
+			uint			LightChannelMask = 0;
 			uint			Coverage16bit = 0;
-			DecodeNodeData(ConvertToNodeData(PPLLNode, Coverage16bit), Depth, VisibilityData, BaseColor, Roughness, Specular);
+			DecodeNodeData(ConvertToNodeData(PPLLNode, Coverage16bit), Depth, VisibilityData, BaseColor, Roughness, Specular, LightChannelMask);
 
 			const float Coverage = From8bitCoverage(VisibilityData.Coverage8bit);
 			TotalTransmittance *= 1.0f - Coverage;
@@ -191,8 +192,9 @@ void MainCS(uint GroupIndex : SV_GroupIndex, uint3 DispatchThreadId : SV_Dispatc
 			float3			BaseColor = 0.0f;
 			float			Roughness = 0.0f;
 			float			Specular = 0;
+			uint			LightChannelMask = 0;
 			float			Coverage16bit = 0;
-			DecodeNodeData(ConvertToNodeData(PPLLNode, Coverage16bit), Depth, VisibilityData, BaseColor, Roughness, Specular);
+			DecodeNodeData(ConvertToNodeData(PPLLNode, Coverage16bit), Depth, VisibilityData, BaseColor, Roughness, Specular, LightChannelMask);
 
 			const float Coverage = From16bitCoverage(Coverage16bit);
 
@@ -219,13 +221,14 @@ void MainCS(uint GroupIndex : SV_GroupIndex, uint3 DispatchThreadId : SV_Dispatc
 			float3			BaseColor = 0.0f;
 			float			Roughness = 0.0f;
 			float			Specular = 0.0f;
-			DecodeNodeData(OutNodeData, Depth, VisibilityData, BaseColor, Roughness, Specular);
+			uint			LightChannelMask = 0;
+			DecodeNodeData(OutNodeData, Depth, VisibilityData, BaseColor, Roughness, Specular, LightChannelMask);
 
 			// Coverage8bit is a weight normalising to 1 the contribution of all the compacted samples. Because later it is weighted by Categorization.PixelCoverage.
 			uint Coverage8bit = To8bitCoverage(SortedCoverage[i] / float(ValidPixelSampleTotalCoverage));
 			// Patch the coverage on the out node
 			VisibilityData.Coverage8bit = Coverage8bit;
-			OutNodeData = ConvertToNodeData(Depth, VisibilityData, BaseColor, Roughness, Specular);
+			OutNodeData = ConvertToNodeData(Depth, VisibilityData, BaseColor, Roughness, Specular, LightChannelMask);
 
 			Categorization.TotalCoverage8bit += Coverage8bit;
 			Categorization.ClosestDepth = max(Categorization.ClosestDepth, OutNodeData.Depth); // Inverse-Z
@@ -617,14 +620,17 @@ void MainCS(uint GroupIndex : SV_GroupIndex, uint3 DispatchThreadId : SV_Dispatc
 
 					float4 BaseColorAndRoughness = MSAA_MaterialTexture.Load(PixelCoord, SampleIt);
 					BaseColorAndRoughness.xyz = BaseColorAndRoughness.xyz * BaseColorAndRoughness.xyz; // Cheap gamma encoding since this is encoded into 8bit target
-					const float Specular = MSAA_AttributeTexture.Load(PixelCoord, SampleIt).x;
+					const float2 SpecularAndLigthChannelMask = MSAA_AttributeTexture.Load(PixelCoord, SampleIt).xy;
+					const float Specular = SpecularAndLigthChannelMask.x;
+					const uint LightChannelMask = SpecularAndLigthChannelMask.y * 0xFF;
 
 					OutNodeData = ConvertToNodeData(
 						SampleDepth,
 						VisibilityData,
 						BaseColorAndRoughness.xyz,
 						BaseColorAndRoughness.w,
-						Specular);
+						Specular,
+						LightChannelMask);
 
 					#if PERMUTATION_VELOCITY == VELOCITY_TYPE_AVG
 					const float2 EncodedVelocity = MSAA_VelocityTexture.Load(PixelCoord, SampleIt);
diff --git a/Engine/Shaders/Private/HairStrands/HairStrandsVisibilityPS.usf b/Engine/Shaders/Private/HairStrands/HairStrandsVisibilityPS.usf
--- a/Engine/Shaders/Private/HairStrands/HairStrandsVisibilityPS.usf
+++ b/Engine/Shaders/Private/HairStrands/HairStrandsVisibilityPS.usf
@@ -42,6 +42,7 @@ RWStructuredBuffer<FPPLLNodeData> PPLLNodes;
 
 uint HairVisibilityPass_HairMacroGroupIndex;
 uint HairVisibilityPass_HairMaterialId;
+uint HairVisibilityPass_LightChannelMask;
 
 #if HAIR_RENDER_MODE == RENDER_MODE_PPLL
 EARLYDEPTHSTENCIL	// Required for UAV operations to not happen.
@@ -151,7 +152,7 @@ void MainVisibility(
 #if HAIR_RENDER_MODE == RENDER_MODE_MSAA
 	OutColor0 = EncodeVisibilityData(VisibilityData);
 	OutColor1 = float4(sqrt(BaseColor), Roughness);
-	OutColor2 = float4(Specular,0,0,0);
+	OutColor2 = float4(Specular, float(HairVisibilityPass_LightChannelMask) / 255.f,0,0);
 	OutColor3 = EncodedVelocity;
 #endif
 
@@ -170,7 +171,7 @@ void MainVisibility(
 		// If we can append new node, add it and make it point to the node previously in head of list
 		uint NextNodeIndex;
 		InterlockedExchange(PPLLNodeIndex[PixelCoord], NodeIndex, NextNodeIndex);
-		const FPPLLNodeData Node = ConvertToPPLLNodeData(SvPosition.z, Coverage, VisibilityData, BaseColor, Roughness, Specular, Velocity, NextNodeIndex);
+		const FPPLLNodeData Node = ConvertToPPLLNodeData(SvPosition.z, Coverage, VisibilityData, BaseColor, Roughness, Specular, HairVisibilityPass_LightChannelMask, Velocity, NextNodeIndex);
 		PPLLNodes[NodeIndex] = Node;
 	}
 #endif
diff --git a/Engine/Shaders/Private/HairStrands/HairStrandsVoxelPageAllocation.usf b/Engine/Shaders/Private/HairStrands/HairStrandsVoxelPageAllocation.usf
--- a/Engine/Shaders/Private/HairStrands/HairStrandsVoxelPageAllocation.usf
+++ b/Engine/Shaders/Private/HairStrands/HairStrandsVoxelPageAllocation.usf
@@ -32,7 +32,12 @@ groupshared uint PageIndexOffsets[MAX_INSTANCE_GROUP_COUNT];
 void AllocatePageIndex(uint2 DispatchThreadId : SV_DispatchThreadID)
 {
 	const uint MacroGroupId = DispatchThreadId.x;
-	FHairAABB Bound = ReadHairAABB(MacroGroupId, MacroGroupAABBBuffer);
+
+	FHairAABB Bound = InitHairAABB();
+	if (MacroGroupId < MacroGroupCount)
+	{
+		Bound = ReadHairAABB(MacroGroupId, MacroGroupAABBBuffer);
+	}
 
 	// Page index allocation
 	int3 PageIndexResolution = 0;
@@ -61,9 +66,17 @@ void AllocatePageIndex(uint2 DispatchThreadId : SV_DispatchThreadID)
 		}
 		GroupMemoryBarrierWithGroupSync();
 
-		const uint PageIndexOffset = PageIndexOffsets[MacroGroupId];
-		OutPageIndexResolutionAndOffsetBuffer[MacroGroupId] = uint4(PageIndexResolution, PageIndexOffset);
-		WriteHairAABB(MacroGroupId, Bound, MacroGroupAABBBuffer);
+		if (MacroGroupId < MacroGroupCount)
+		{
+			const uint PageIndexOffset = PageIndexOffsets[MacroGroupId];
+			OutPageIndexResolutionAndOffsetBuffer[MacroGroupId] = uint4(PageIndexResolution, PageIndexOffset);
+			WriteHairAABB(MacroGroupId, Bound, MacroGroupAABBBuffer);
+		}
+	}
+
+	if (MacroGroupId >= MacroGroupCount)
+	{
+		return;
 	}
 
 	// Use CPU data, which allow to compare the GPU driven part with CPU data
diff --git a/Engine/Shaders/Private/RayTracing/RayTracingOcclusionRGS.usf b/Engine/Shaders/Private/RayTracing/RayTracingOcclusionRGS.usf
--- a/Engine/Shaders/Private/RayTracing/RayTracingOcclusionRGS.usf
+++ b/Engine/Shaders/Private/RayTracing/RayTracingOcclusionRGS.usf
@@ -29,6 +29,7 @@ RaytracingAccelerationStructure TLAS;
 RWTexture2D<float4> RWOcclusionMaskUAV;
 RWTexture2D<float> RWRayDistanceUAV;
 Texture2D<uint4> HairCategorizationTexture;
+Texture2D<uint> HairLightChannelMaskTexture;
 
 uint ShadowMaskType;
 uint LightingChannelMask;
@@ -161,6 +162,7 @@ RAY_TRACING_ENTRY_RAYGEN(OcclusionRGS)
 #endif
 		const bool bIsDepthValid = SceneDepthBuffer.Load(int3(PixelCoord, 0)).r > 0.0;
 		bool bIsValidPixel = ShadingModelID != SHADINGMODELID_UNLIT && bIsDepthValid;
+		uint LightChannel = GetSceneLightingChannel(PixelCoord);
 
 		// Check if there is hair geometry
 		uint RaytracingMask = RAY_TRACING_MASK_SHADOW | RAY_TRACING_MASK_THIN_SHADOW;
@@ -174,13 +176,13 @@ RAY_TRACING_ENTRY_RAYGEN(OcclusionRGS)
 			{
 				DeviceZ = Categorization.ClosestDepth;
 				ShadingModelID = SHADINGMODELID_HAIR;
+				LightChannel = HairLightChannelMaskTexture.Load(uint3(PixelCoord, 0));
 			}
 		}
 	#endif
 
 		const float3 WorldPosition = ReconstructWorldPositionFromDeviceZ(PixelCoord, DeviceZ);
 
-		const uint LightChannel = GetSceneLightingChannel(PixelCoord);
 		const bool bTraceRay = bIsValidPixel && (LightChannel & LightingChannelMask) != 0;
 		if (!bTraceRay)
 		{
diff --git a/Engine/Source/Runtime/Engine/Private/GPUSkinCache.cpp b/Engine/Source/Runtime/Engine/Private/GPUSkinCache.cpp
--- a/Engine/Source/Runtime/Engine/Private/GPUSkinCache.cpp
+++ b/Engine/Source/Runtime/Engine/Private/GPUSkinCache.cpp
@@ -265,6 +265,7 @@ public:
 		FCachedGeometrySection MeshSection;
 		const FSkelMeshRenderSection& Section = *DispatchData[SectionIndex].Section;
 		MeshSection.PositionBuffer	= DispatchData[SectionIndex].PositionBuffer->SRV;
+		MeshSection.UVsBuffer		= DispatchData[SectionIndex].UVsBufferSRV;
 		MeshSection.TotalVertexCount= DispatchData[SectionIndex].PositionBuffer->NumBytes / (sizeof(float)*3);
 		MeshSection.NumPrimitives	= Section.NumTriangles;
 		MeshSection.IndexBaseIndex	= Section.BaseIndex;
diff --git a/Engine/Source/Runtime/Engine/Private/SkeletalRenderGPUSkin.cpp b/Engine/Source/Runtime/Engine/Private/SkeletalRenderGPUSkin.cpp
--- a/Engine/Source/Runtime/Engine/Private/SkeletalRenderGPUSkin.cpp
+++ b/Engine/Source/Runtime/Engine/Private/SkeletalRenderGPUSkin.cpp
@@ -646,9 +646,11 @@ void FSkeletalMeshObjectGPUSkin::ProcessUpdatedDynamicData(FGPUSkinCache* GPUSki
 					bClothFactory ? DynamicData->ClothingSimData.Find(Section.CorrespondClothAssetIndex) : 0, LocalToCloth, DynamicData->ClothBlendWeight, RevisionNumber, SectionIdx, SkinCacheEntry);
 
 				FCachedGeometrySection CachedSection = GPUSkinCache->GetCachedGeometry(SkinCacheEntry, SectionIdx);
-				CachedSection.IndexBuffer = LODData.MultiSizeIndexContainer.GetIndexBuffer()->GetSRV();
-				CachedSection.TotalIndexCount = LODData.MultiSizeIndexContainer.GetIndexBuffer()->Num();
-				CachedSection.LODIndex = DynamicData->LODIndex;
+				CachedSection.IndexBuffer		= LODData.MultiSizeIndexContainer.GetIndexBuffer()->GetSRV();
+				CachedSection.TotalIndexCount	= LODData.MultiSizeIndexContainer.GetIndexBuffer()->Num();
+				CachedSection.LODIndex			= DynamicData->LODIndex;
+				CachedSection.UVsChannelOffset	= 0; // Assume that we needs to pair meshes based on UVs 0
+				CachedSection.UVsChannelCount	= LODData.StaticVertexBuffers.StaticMeshVertexBuffer.GetNumTexCoords();
 				CachedGeometry.Sections.Add(CachedSection);
 			}
 			if (bNeedFence)
diff --git a/Engine/Source/Runtime/Engine/Public/SkeletalRenderPublic.h b/Engine/Source/Runtime/Engine/Public/SkeletalRenderPublic.h
--- a/Engine/Source/Runtime/Engine/Public/SkeletalRenderPublic.h
+++ b/Engine/Source/Runtime/Engine/Public/SkeletalRenderPublic.h
@@ -46,14 +46,17 @@ enum class EPreviousBoneTransformUpdateMode
 struct FCachedGeometrySection
 {
 	FRHIShaderResourceView* PositionBuffer = nullptr;
+	FRHIShaderResourceView* UVsBuffer = nullptr;
 	FRHIShaderResourceView* IndexBuffer = nullptr;
+	uint32 UVsChannelOffset = 0;
+	uint32 UVsChannelCount = 0;
 	uint32 NumPrimitives = 0;
 	uint32 VertexBaseIndex = 0;
 	uint32 IndexBaseIndex = 0;
 	uint32 TotalVertexCount = 0;
 	uint32 TotalIndexCount = 0;
 	uint32 SectionIndex = 0;
-	int32 LODIndex = 0;
+	int32  LODIndex = 0;
 };
 
 struct FCachedGeometry
diff --git a/Engine/Source/Runtime/Renderer/Private/DeferredShadingRenderer.cpp b/Engine/Source/Runtime/Renderer/Private/DeferredShadingRenderer.cpp
--- a/Engine/Source/Runtime/Renderer/Private/DeferredShadingRenderer.cpp
+++ b/Engine/Source/Runtime/Renderer/Private/DeferredShadingRenderer.cpp
@@ -1627,6 +1627,12 @@ void FDeferredShadingSceneRenderer::Render(FRHICommandListImmediate& RHICmdList)
 
 	RunGPUSkinCacheTransition(RHICmdList, Scene, EGPUSkinCacheTransition::Renderer);
 
+	if (HasHairStrandsProjectionQuery(Scene->GetShaderPlatform()))
+	{
+		auto ShaderMap = GetGlobalShaderMap(FeatureLevel);
+		RunHairStrandsBindingQueries(RHICmdList, ShaderMap);
+	}
+
 	// Interpolation needs to happen after the skin cache run as there is a dependency 
 	// on the skin cache output.
 	const bool bRunHairStrands = IsHairStrandsEnable(Scene->GetShaderPlatform()) && Views.Num() > 0;
@@ -1635,6 +1641,7 @@ void FDeferredShadingSceneRenderer::Render(FRHICommandListImmediate& RHICmdList)
 	{
 		const EWorldType::Type WorldType = Views[0].Family->Scene->GetWorld()->WorldType;
 		auto ShaderMap = GetGlobalShaderMap(FeatureLevel);
+
 		RunHairStrandsInterpolation(RHICmdList, WorldType, &Views[0].ShaderDrawData, ShaderMap, EHairStrandsInterpolationType::RenderStrands, &HairClusterData); // Send data to full up with culling
 	}
 
diff --git a/Engine/Source/Runtime/Renderer/Private/DeferredShadingRenderer.h b/Engine/Source/Runtime/Renderer/Private/DeferredShadingRenderer.h
--- a/Engine/Source/Runtime/Renderer/Private/DeferredShadingRenderer.h
+++ b/Engine/Source/Runtime/Renderer/Private/DeferredShadingRenderer.h
@@ -524,6 +524,7 @@ private:
 		const IScreenSpaceDenoiser::EShadowRequirements DenoiserRequirements,
 		const bool bSubPixelShadowMask,
 		FRDGTextureRef HairCategorizationTexture,
+		FRDGTextureRef HairLightChannelMaskTexture,
 		FRDGTextureUAV* OutShadowMaskUAV,
 		FRDGTextureUAV* OutRayHitDistanceUAV);
 
diff --git a/Engine/Source/Runtime/Renderer/Private/HairStrands/HairStrandsCluster.cpp b/Engine/Source/Runtime/Renderer/Private/HairStrands/HairStrandsCluster.cpp
--- a/Engine/Source/Runtime/Renderer/Private/HairStrands/HairStrandsCluster.cpp
+++ b/Engine/Source/Runtime/Renderer/Private/HairStrands/HairStrandsCluster.cpp
@@ -142,7 +142,7 @@ FHairStrandsMacroGroupViews CreateHairStrandsMacroGroups(
 			int32 MaterialId = 0;
 			FHairStrandsMacroGroupDatas& MacroGroups = MacroGroupsViews.Views.AddDefaulted_GetRef();
 
-			if (View.HairStrandsMeshElements.Num() == 0)
+			if (View.HairStrandsMeshElements.Num() == 0 || View.bIsPlanarReflection || View.bIsReflectionCapture)
 			{
 				continue;
 			}
diff --git a/Engine/Source/Runtime/Renderer/Private/HairStrands/HairStrandsComposition.cpp b/Engine/Source/Runtime/Renderer/Private/HairStrands/HairStrandsComposition.cpp
--- a/Engine/Source/Runtime/Renderer/Private/HairStrands/HairStrandsComposition.cpp
+++ b/Engine/Source/Runtime/Renderer/Private/HairStrands/HairStrandsComposition.cpp
@@ -68,7 +68,7 @@ static void AddHairVisibilityComposeSubPixelPass(
 		{
 			FGraphicsPipelineStateInitializer GraphicsPSOInit;
 			RHICmdList.ApplyCachedRenderTargets(GraphicsPSOInit);
-			GraphicsPSOInit.BlendState = TStaticBlendState<CW_RGBA, BO_Add, BF_One, BF_InverseSourceAlpha, BO_Add, BF_One, BF_Zero>::GetRHI();
+			GraphicsPSOInit.BlendState = TStaticBlendState<CW_RGBA, BO_Add, BF_One, BF_InverseSourceAlpha, BO_Add, BF_Zero, BF_InverseSourceAlpha>::GetRHI();
 			GraphicsPSOInit.RasterizerState = TStaticRasterizerState<>::GetRHI();
 
 			// Write stencil value for partially covered pixel, in order to run responsive AA on them
diff --git a/Engine/Source/Runtime/Renderer/Private/HairStrands/HairStrandsDebug.cpp b/Engine/Source/Runtime/Renderer/Private/HairStrands/HairStrandsDebug.cpp
--- a/Engine/Source/Runtime/Renderer/Private/HairStrands/HairStrandsDebug.cpp
+++ b/Engine/Source/Runtime/Renderer/Private/HairStrands/HairStrandsDebug.cpp
@@ -56,7 +56,7 @@ static int32 GHairStrandsClusterDebug = 0;
 static FAutoConsoleVariableRef CVarHairStrandsDebugClusterAABB(TEXT("r.HairStrands.Cluster.Debug"), GHairStrandsClusterDebug, TEXT("Draw debug the world bounding box of hair clusters used for culling optimisation."));
 
 static int32 GHairDebugMeshProjection_SkinCacheMesh = 0;
-
+static int32 GHairDebugMeshProjection_SkinCacheMeshInUVsSpace = 0;
 static int32 GHairDebugMeshProjection_Sim_HairRestTriangles = 0;
 static int32 GHairDebugMeshProjection_Sim_HairRestFrames = 0;
 static int32 GHairDebugMeshProjection_Sim_HairDeformedTriangles = 0;
@@ -67,6 +67,8 @@ static int32 GHairDebugMeshProjection_Render_HairRestFrames = 0;
 static int32 GHairDebugMeshProjection_Render_HairDeformedTriangles = 0;
 static int32 GHairDebugMeshProjection_Render_HairDeformedFrames = 0;
 
+
+static FAutoConsoleVariableRef CVarHairDebugMeshProjection_SkinCacheMeshInUVsSpace(			TEXT("r.HairStrands.MeshProjection.DebugInUVsSpace"),			GHairDebugMeshProjection_SkinCacheMeshInUVsSpace, TEXT("Render debug mes projection in UVs space"));
 static FAutoConsoleVariableRef CVarHairDebugMeshProjection_SkinCacheMesh(					TEXT("r.HairStrands.MeshProjection.DebugSkinCache"),			GHairDebugMeshProjection_SkinCacheMesh, TEXT("Render debug mes projection"));
 static FAutoConsoleVariableRef CVarHairDebugMeshProjection_Render_HairRestTriangles(		TEXT("r.HairStrands.MeshProjection.Render.Rest.Triangles"),		GHairDebugMeshProjection_Render_HairRestTriangles, TEXT("Render debug mes projection"));
 static FAutoConsoleVariableRef CVarHairDebugMeshProjection_Render_HairRestFrames(			TEXT("r.HairStrands.MeshProjection.Render.Rest.Frames"),		GHairDebugMeshProjection_Render_HairRestFrames, TEXT("Render debug mes projection"));
@@ -85,7 +87,7 @@ static FAutoConsoleVariableRef CVarHairVirtualVoxel_DrawDebugPage(TEXT("r.HairSt
 
 // Helper functions for accessing interpolation data for debug purpose.
 // Definitions is in HairStrandsInterface.cpp
-void GetGroomInterpolationData(const EWorldType::Type WorldType, FHairStrandsProjectionMeshData::LOD& OutGeometries);
+void GetGroomInterpolationData(const EWorldType::Type WorldType, const EHairStrandsProjectionMeshType MeshtType, FHairStrandsProjectionMeshData::LOD& OutGeometries);
 void GetGroomInterpolationData(const EWorldType::Type WorldType, const bool bRenderData, FHairStrandsProjectionHairData& OutHairData, TArray<int32>& OutLODIndices);
 
 static int32 GHairStrandsCull = 0;
@@ -849,9 +851,15 @@ BEGIN_SHADER_PARAMETER_STRUCT(FHairProjectionMeshDebugParameters, )
 	SHADER_PARAMETER(uint32, IndexOffset)
 	SHADER_PARAMETER(uint32, MaxIndexCount)
 	SHADER_PARAMETER(uint32, MaxVertexCount)
+	SHADER_PARAMETER(uint32, MeshUVsChannelOffset)
+	SHADER_PARAMETER(uint32, MeshUVsChannelCount)
+	SHADER_PARAMETER(uint32, bOutputInUVsSpace)
+	SHADER_PARAMETER(uint32, MeshType)
+	SHADER_PARAMETER(uint32, SectionIndex)
 	SHADER_PARAMETER(FVector2D, OutputResolution)
 	SHADER_PARAMETER_SRV(StructuredBuffer, InputIndexBuffer)
-	SHADER_PARAMETER_SRV(StructuredBuffer, InputVertexBuffer)
+	SHADER_PARAMETER_SRV(StructuredBuffer, InputVertexPositionBuffer)
+	SHADER_PARAMETER_SRV(StructuredBuffer, InputVertexUVsBuffer)
 	SHADER_PARAMETER_STRUCT_REF(FViewUniformShaderParameters, ViewUniformBuffer)
 	RENDER_TARGET_BINDING_SLOTS()
 END_SHADER_PARAMETER_STRUCT()
@@ -900,6 +908,7 @@ IMPLEMENT_GLOBAL_SHADER(FHairProjectionMeshDebugPS, "/Engine/Private/HairStrands
 static void AddDebugProjectionMeshPass(
 	FRDGBuilder& GraphBuilder,
 	const FViewInfo* View,
+	const EHairStrandsProjectionMeshType MeshType,
 	const bool bClearDepth,
 	FHairStrandsProjectionMeshData::Section& MeshSectionData,
 	FRDGTextureRef& ColorTexture, 
@@ -918,12 +927,18 @@ static void AddDebugProjectionMeshPass(
 	FHairProjectionMeshDebugParameters* Parameters = GraphBuilder.AllocParameters<FHairProjectionMeshDebugParameters>();
 	Parameters->LocalToWorld = MeshSectionData.LocalToWorld.ToMatrixWithScale();
 	Parameters->OutputResolution = Resolution;
+	Parameters->MeshType = uint32(MeshType);
+	Parameters->bOutputInUVsSpace = GHairDebugMeshProjection_SkinCacheMeshInUVsSpace ? 1 : 0;
 	Parameters->VertexOffset = MeshSectionData.VertexBaseIndex;
 	Parameters->IndexOffset = MeshSectionData.IndexBaseIndex;
 	Parameters->MaxIndexCount = MeshSectionData.TotalIndexCount;
 	Parameters->MaxVertexCount = MeshSectionData.TotalVertexCount;
+	Parameters->MeshUVsChannelOffset = MeshSectionData.UVsChannelOffset;
+	Parameters->MeshUVsChannelCount = MeshSectionData.UVsChannelCount;
 	Parameters->InputIndexBuffer  = MeshSectionData.IndexBuffer;
-	Parameters->InputVertexBuffer = MeshSectionData.PositionBuffer;
+	Parameters->InputVertexPositionBuffer = MeshSectionData.PositionBuffer;
+	Parameters->InputVertexUVsBuffer = MeshSectionData.UVsBuffer;
+	Parameters->SectionIndex = MeshSectionData.SectionIndex;
 	Parameters->ViewUniformBuffer = View->ViewUniformBuffer;
 	Parameters->RenderTargets[0] = FRenderTargetBinding(ColorTexture, ERenderTargetLoadAction::ELoad, 0);
 	Parameters->RenderTargets.DepthStencil = FDepthStencilBinding(DepthTexture, bClearDepth ? ERenderTargetLoadAction::EClear : ERenderTargetLoadAction::ELoad, ERenderTargetLoadAction::ENoAction, FExclusiveDepthStencil::DepthWrite_StencilNop);
@@ -958,7 +973,7 @@ static void AddDebugProjectionMeshPass(
 		FGraphicsPipelineStateInitializer GraphicsPSOInit;
 		RHICmdList.ApplyCachedRenderTargets(GraphicsPSOInit);
 		GraphicsPSOInit.BlendState = TStaticBlendState<CW_RGBA, BO_Add, BF_One, BF_Zero, BO_Add, BF_One, BF_Zero>::GetRHI();
-		GraphicsPSOInit.RasterizerState = TStaticRasterizerState<>::GetRHI();
+		GraphicsPSOInit.RasterizerState = TStaticRasterizerState<FM_Wireframe>::GetRHI();
 		GraphicsPSOInit.DepthStencilState = TStaticDepthStencilState<true, CF_DepthNearOrEqual>::GetRHI();
 		GraphicsPSOInit.BoundShaderState.VertexDeclarationRHI = GEmptyVertexDeclaration.VertexDeclarationRHI;
 		GraphicsPSOInit.BoundShaderState.VertexShaderRHI = GETSAFERHISHADER_VERTEX(*VertexShader);
@@ -979,8 +994,6 @@ BEGIN_SHADER_PARAMETER_STRUCT(FHairProjectionHairDebugParameters, )
 	SHADER_PARAMETER(FVector2D, OutputResolution)
 	SHADER_PARAMETER(uint32, MaxRootCount)
 	SHADER_PARAMETER(uint32, DeformedFrameEnable)
-	SHADER_PARAMETER(FVector, RestPositionOffset)
-	SHADER_PARAMETER(FVector, DeformedPositionOffset)
 	SHADER_PARAMETER(FMatrix, RootLocalToWorld)
 
 	SHADER_PARAMETER_SRV(StructuredBuffer, RestPosition0Buffer)
@@ -1090,12 +1103,10 @@ static void AddDebugProjectionHairPass(
 		Parameters->RootBarycentricBuffer = HairData.LODDatas[LODIndex].RootTriangleBarycentricBuffer->SRV;
 	}
 
-	Parameters->RestPositionOffset = *HairData.LODDatas[LODIndex].RestPositionOffset;
 	Parameters->RestPosition0Buffer = HairData.LODDatas[LODIndex].RestRootTrianglePosition0Buffer->SRV;
 	Parameters->RestPosition1Buffer = HairData.LODDatas[LODIndex].RestRootTrianglePosition1Buffer->SRV;
 	Parameters->RestPosition2Buffer = HairData.LODDatas[LODIndex].RestRootTrianglePosition2Buffer->SRV;
 	
-	Parameters->DeformedPositionOffset = *HairData.LODDatas[LODIndex].DeformedPositionOffset;
 	Parameters->DeformedPosition0Buffer = HairData.LODDatas[LODIndex].DeformedRootTrianglePosition0Buffer->SRV;
 	Parameters->DeformedPosition1Buffer = HairData.LODDatas[LODIndex].DeformedRootTrianglePosition1Buffer->SRV;
 	Parameters->DeformedPosition2Buffer = HairData.LODDatas[LODIndex].DeformedRootTrianglePosition2Buffer->SRV;
@@ -1621,13 +1632,22 @@ void RenderHairStrandsDebugInfo(
 
 			if (GHairDebugMeshProjection_SkinCacheMesh > 0)
 			{
-				FHairStrandsProjectionMeshData::LOD MeshProjectionLODData;
-				GetGroomInterpolationData(WorldType, MeshProjectionLODData);
-				for (FHairStrandsProjectionMeshData::Section& Section : MeshProjectionLODData.Sections)
+				FViewInfo* LocalView = &View;
+				auto RenderMeshProjection = [&bClearDepth, WorldType, LocalView, &SceneColorTexture, &DepthTexture](FRDGBuilder& InGraphBuilder, EHairStrandsProjectionMeshType MeshType)
 				{
-					AddDebugProjectionMeshPass(GraphBuilder, &View, bClearDepth, Section, SceneColorTexture, DepthTexture);
-					bClearDepth = false;
-				}
+					FHairStrandsProjectionMeshData::LOD MeshProjectionLODData;
+					GetGroomInterpolationData(WorldType, MeshType, MeshProjectionLODData);
+					for (FHairStrandsProjectionMeshData::Section& Section : MeshProjectionLODData.Sections)
+					{
+						AddDebugProjectionMeshPass(InGraphBuilder, LocalView, MeshType, bClearDepth, Section, SceneColorTexture, DepthTexture);
+						bClearDepth = false;
+					}
+				};
+
+				RenderMeshProjection(GraphBuilder, EHairStrandsProjectionMeshType::DeformedMesh);
+				RenderMeshProjection(GraphBuilder, EHairStrandsProjectionMeshType::RestMesh);
+				RenderMeshProjection(GraphBuilder, EHairStrandsProjectionMeshType::SourceMesh);
+				RenderMeshProjection(GraphBuilder, EHairStrandsProjectionMeshType::TargetMesh);
 			}
 
 			FViewInfo* LocalView = &View;
diff --git a/Engine/Source/Runtime/Renderer/Private/HairStrands/HairStrandsInterface.cpp b/Engine/Source/Runtime/Renderer/Private/HairStrands/HairStrandsInterface.cpp
--- a/Engine/Source/Runtime/Renderer/Private/HairStrands/HairStrandsInterface.cpp
+++ b/Engine/Source/Runtime/Renderer/Private/HairStrands/HairStrandsInterface.cpp
@@ -60,6 +60,9 @@ inline FHairStrandsProjectionMeshData::Section ConvertMeshSection(const FCachedG
 	FHairStrandsProjectionMeshData::Section Out;
 	Out.IndexBuffer = In.IndexBuffer;
 	Out.PositionBuffer = In.PositionBuffer;
+	Out.UVsBuffer = In.UVsBuffer;
+	Out.UVsChannelOffset = In.UVsChannelOffset;
+	Out.UVsChannelCount = In.UVsChannelCount;
 	Out.TotalVertexCount = In.TotalVertexCount;
 	Out.TotalIndexCount = In.TotalIndexCount;
 	Out.VertexBaseIndex = In.VertexBaseIndex;
@@ -94,8 +97,8 @@ struct FHairStrandsManager
 		FHairStrandsProjectionHairData SimProjectionHairDatas;
 		FCachedGeometry CachedGeometry;
 		FTransform SkeletalLocalToWorld;
-		FVector SkeletalDeformedPositionOffset;
 		const FSkeletalMeshObject* MeshObject = nullptr;
+		FHairStrandsProjectionDebugInfo DebugProjectionInfo;
 		int32 FrameLODIndex = -1;
 	};
 	TArray<Element> Elements;
@@ -110,13 +113,14 @@ void RegisterHairStrands(
 	const FHairStrandsProjectionHairData& RenProjectionDatas,
  	const FHairStrandsProjectionHairData& SimProjectionDatas,
 	const FHairStrandsPrimitiveResources& PrimitiveResources,
-	const FHairStrandsDebugInfo& DebugInfo)
+	const FHairStrandsDebugInfo& DebugInfo,
+	const FHairStrandsProjectionDebugInfo& DebugProjectionInfo)
 {
 	for (int32 Index = 0; Index < GHairManager.Elements.Num(); ++Index)
 	{
 		if (GHairManager.Elements[Index].Id == Id && GHairManager.Elements[Index].WorldType == WorldType)
 		{
-			// Component already registered. This should not happen.
+			// Component already registered. This should not happen. 
 			UE_LOG(LogHairRendering, Warning, TEXT("Component already register. This should't happen. Please report this to a rendering engineer."))
 			return;
 		}
@@ -129,17 +133,16 @@ void RegisterHairStrands(
 	E.RenProjectionHairDatas = RenProjectionDatas;
 	E.SimProjectionHairDatas = SimProjectionDatas;
 	E.SkeletalLocalToWorld = FTransform::Identity;
-	E.SkeletalDeformedPositionOffset = FVector::ZeroVector;
 	E.PrimitiveResources = PrimitiveResources;
 	E.DebugInfo = DebugInfo;
+	E.DebugProjectionInfo = DebugProjectionInfo;
 }
 
 bool UpdateHairStrands(
 	uint64 Id, 
 	EWorldType::Type WorldType, 
 	const FTransform& HairLocalToWorld, 
-	const FTransform& SkeletalLocalToWorld,
-	const FVector& SkeletalDeformedPositionOffset)
+	const FTransform& SkeletalLocalToWorld)
 {
 	for (FHairStrandsManager::Element& E : GHairManager.Elements)
 	{
@@ -155,7 +158,6 @@ bool UpdateHairStrands(
 			ProjectionData.LocalToWorld = HairLocalToWorld;
 		}
 		E.SkeletalLocalToWorld = SkeletalLocalToWorld;
-		E.SkeletalDeformedPositionOffset = SkeletalDeformedPositionOffset;
 		return true;
 	}
 
@@ -235,30 +237,54 @@ void UnregisterHairStrands(uint64 Id)
 	}
 }
 
+void RunMeshTransfer(
+	FRHICommandListImmediate& RHICmdList,
+	const FHairStrandsProjectionMeshData& SourceMeshData,
+	const FHairStrandsProjectionMeshData& TargetMeshData, 
+	TArray<FRWBuffer>& OutTransferedPositions)
+{
+	const ERHIFeatureLevel::Type FeatureLevel = GMaxRHIFeatureLevel;
+	TShaderMap<FGlobalShaderType>* ShaderMap = GetGlobalShaderMap(FeatureLevel);
+
+	FRDGBuilder GraphBuilder(RHICmdList);
+
+	const uint32 LODCount = TargetMeshData.LODs.Num();
+	OutTransferedPositions.SetNum(LODCount);
+	for (uint32 LODIndex = 0; LODIndex < LODCount; ++LODIndex)
+	{
+		check(TargetMeshData.LODs[LODIndex].Sections.Num() > 0);
+
+		OutTransferedPositions[LODIndex].Initialize(sizeof(float), TargetMeshData.LODs[LODIndex].Sections[0].TotalVertexCount * 3, PF_R32_FLOAT);
+		TransferMesh(GraphBuilder, ShaderMap, LODIndex, SourceMeshData, TargetMeshData, OutTransferedPositions[LODIndex]);
+	}
+	
+	GraphBuilder.Execute();	
+}
+
 void RunProjection(
 	FRHICommandListImmediate& RHICmdList,
 	const FTransform& LocalToWorld,
-	const FVector& RestPositionOffset,
-	const FHairStrandsProjectionMeshData& MeshData,
+	const FHairStrandsProjectionMeshData& TargetMeshData,
 	FHairStrandsProjectionHairData& RenProjectionHairData,
 	FHairStrandsProjectionHairData& SimProjectionHairData)
 {
 	const ERHIFeatureLevel::Type FeatureLevel = GMaxRHIFeatureLevel;
 	TShaderMap<FGlobalShaderType>* ShaderMap = GetGlobalShaderMap(FeatureLevel);
 
-	auto Project = [&RHICmdList, ShaderMap, &MeshData, LocalToWorld, RestPositionOffset](FHairStrandsProjectionHairData& ProjectionHairData)
+	auto Project = [&RHICmdList, ShaderMap, &TargetMeshData, LocalToWorld](FHairStrandsProjectionHairData& ProjectionHairData)
 	{
+		FRDGBuilder GraphBuilder(RHICmdList);
 		for (FHairStrandsProjectionHairData::HairGroup& HairGroup : ProjectionHairData.HairGroups)
 		{
 			for (FHairStrandsProjectionHairData::LODData& LODData : HairGroup.LODDatas)
 			{
 				const uint32 LODIndex = LODData.LODIndex;
-				*HairGroup.LODDatas[LODIndex].RestPositionOffset = RestPositionOffset;
 				HairGroup.LocalToWorld = LocalToWorld;
-				ProjectHairStrandsOntoMesh(RHICmdList, ShaderMap, LODIndex, MeshData, HairGroup);
-				UpdateHairStrandsMeshTriangles(RHICmdList, ShaderMap, LODIndex, HairStrandsTriangleType::RestPose, MeshData.LODs[LODIndex], HairGroup);
+				ProjectHairStrandsOntoMesh(GraphBuilder, ShaderMap, LODIndex, TargetMeshData, HairGroup);
+				UpdateHairStrandsMeshTriangles(GraphBuilder, ShaderMap, LODIndex, HairStrandsTriangleType::RestPose, TargetMeshData.LODs[LODIndex], HairGroup);
 			}
 		}
+		GraphBuilder.Execute();
 	};
 
 	Project(RenProjectionHairData);
@@ -301,12 +327,13 @@ void RunHairStrandsInterpolation(
 			MeshDataLOD.Sections.Add(ConvertMeshSection(Section, E.SkeletalLocalToWorld));
 		}
 
+		FRDGBuilder GraphBuilder(RHICmdList);
+
 		for (FHairStrandsProjectionHairData::HairGroup& ProjectionHairData : E.RenProjectionHairDatas.HairGroups)
 		{
 			if (EHairStrandsInterpolationType::RenderStrands == Type && 0 <= E.FrameLODIndex && E.FrameLODIndex < ProjectionHairData.LODDatas.Num() && ProjectionHairData.LODDatas[E.FrameLODIndex].bIsValid)
 			{
-				*ProjectionHairData.LODDatas[E.FrameLODIndex].DeformedPositionOffset = E.SkeletalDeformedPositionOffset;
-				UpdateHairStrandsMeshTriangles(RHICmdList, ShaderMap, E.FrameLODIndex, HairStrandsTriangleType::DeformedPose, MeshDataLOD, ProjectionHairData);
+				UpdateHairStrandsMeshTriangles(GraphBuilder, ShaderMap, E.FrameLODIndex, HairStrandsTriangleType::DeformedPose, MeshDataLOD, ProjectionHairData);
 			}
 		}
 
@@ -314,10 +341,10 @@ void RunHairStrandsInterpolation(
 		{
 			if (EHairStrandsInterpolationType::SimulationStrands == Type && 0 <= E.FrameLODIndex && E.FrameLODIndex < ProjectionHairData.LODDatas.Num() && ProjectionHairData.LODDatas[E.FrameLODIndex].bIsValid)
 			{
-				*ProjectionHairData.LODDatas[E.FrameLODIndex].DeformedPositionOffset = E.SkeletalDeformedPositionOffset;
-				UpdateHairStrandsMeshTriangles(RHICmdList, ShaderMap, E.FrameLODIndex, HairStrandsTriangleType::DeformedPose, MeshDataLOD, ProjectionHairData);
+				UpdateHairStrandsMeshTriangles(GraphBuilder, ShaderMap, E.FrameLODIndex, HairStrandsTriangleType::DeformedPose, MeshDataLOD, ProjectionHairData);
 			}
 		}
+		GraphBuilder.Execute();
 	}
 
 	// Reset deformation
@@ -330,7 +357,7 @@ void RunHairStrandsInterpolation(
 
 			if (E.InterpolationData.Input && E.InterpolationData.Output && E.InterpolationData.ResetFunction)
 			{
-				E.InterpolationData.ResetFunction(RHICmdList, E.InterpolationData.Input, E.InterpolationData.Output);
+				E.InterpolationData.ResetFunction(RHICmdList, E.InterpolationData.Input, E.InterpolationData.Output, E.SimProjectionHairDatas, E.FrameLODIndex);
 			}
 		}
 	}
@@ -351,16 +378,45 @@ void RunHairStrandsInterpolation(
 	}
 }
 
-void GetGroomInterpolationData(const EWorldType::Type WorldType, FHairStrandsProjectionMeshData::LOD& OutGeometries)
+void GetGroomInterpolationData(const EWorldType::Type WorldType, 
+	const EHairStrandsProjectionMeshType MeshtType,
+	FHairStrandsProjectionMeshData::LOD& OutGeometries)
 {
 	for (FHairStrandsManager::Element& E : GHairManager.Elements)
 	{
 		if (E.WorldType != WorldType)
 			continue;
 
-		for (FCachedGeometrySection Section : E.CachedGeometry.Sections)
-		{		
-			OutGeometries.Sections.Add(ConvertMeshSection(Section, E.SkeletalLocalToWorld));
+		if (MeshtType == EHairStrandsProjectionMeshType::DeformedMesh || MeshtType == EHairStrandsProjectionMeshType::RestMesh)
+		{
+			for (FCachedGeometrySection Section : E.CachedGeometry.Sections)
+			{			
+				FHairStrandsProjectionMeshData::Section OutSection = ConvertMeshSection(Section, E.SkeletalLocalToWorld);
+				if (MeshtType == EHairStrandsProjectionMeshType::RestMesh)
+				{
+					// If the mesh has some mesh-tranferred data, we display that otherwise we use the rest data
+					const bool bHasTransferData = E.FrameLODIndex < E.DebugProjectionInfo.TransferredPositions.Num();
+					if (bHasTransferData)
+					{
+						OutSection.PositionBuffer = E.DebugProjectionInfo.TransferredPositions[E.FrameLODIndex].SRV;
+					}
+					else if (E.DebugProjectionInfo.TargetMeshData.LODs.Num() > 0)
+					{
+						OutGeometries = E.DebugProjectionInfo.TargetMeshData.LODs[0];
+					}
+				}
+				OutGeometries.Sections.Add(OutSection);
+			}
+		}
+
+		if (MeshtType == EHairStrandsProjectionMeshType::TargetMesh && E.DebugProjectionInfo.TargetMeshData.LODs.Num() > 0)
+		{
+			OutGeometries = E.DebugProjectionInfo.TargetMeshData.LODs[0];
+		}
+
+		if (MeshtType == EHairStrandsProjectionMeshType::SourceMesh && E.DebugProjectionInfo.SourceMeshData.LODs.Num() > 0)
+		{
+			OutGeometries = E.DebugProjectionInfo.SourceMeshData.LODs[0];
 		}
 	}
 }
@@ -450,4 +506,38 @@ bool IsHairStrandsEnable(EShaderPlatform Platform)
 bool IsHairRayTracingEnabled()
 {
 	return IsRayTracingEnabled();
+}
+
+///////////////////////////////////////////////////////////////////////////////////////////////////
+
+struct FBindingQuery
+{
+	void* Asset = nullptr;
+	TBindingProcess BindingProcess;
+};
+TQueue<FBindingQuery> GBindingQueries;
+
+void EnqueueGroomBindingQuery(void* Asset, TBindingProcess BindingProcess)
+{
+	GBindingQueries.Enqueue({ Asset, BindingProcess });
+}
+
+bool HasHairStrandsProjectionQuery(EShaderPlatform Platform)
+{
+	return
+		IsHairStrandsSupported(Platform) &&
+		GHairStrandsRenderingEnable == 1 &&
+		!GBindingQueries.IsEmpty();
+}
+
+void RunHairStrandsBindingQueries(FRHICommandListImmediate& RHICmdList, TShaderMap<FGlobalShaderType>* ShaderMap)
+{
+	FBindingQuery Q;
+	while (GBindingQueries.Dequeue(Q))
+	{
+		if (Q.BindingProcess && Q.Asset)
+		{
+			Q.BindingProcess(RHICmdList, Q.Asset);
+		}
+	}
 }
\ No newline at end of file
