Copying //UE4/Dev-Physics to Main (//UE4/Main)

diff --git a/Engine/Source/Runtime/Renderer/Private/HairStrands/HairStrandsMeshProjection.cpp b/Engine/Source/Runtime/Renderer/Private/HairStrands/HairStrandsMeshProjection.cpp
--- a/Engine/Source/Runtime/Renderer/Private/HairStrands/HairStrandsMeshProjection.cpp
+++ b/Engine/Source/Runtime/Renderer/Private/HairStrands/HairStrandsMeshProjection.cpp
@@ -11,6 +11,185 @@ static int32 GHairProjectionMaxTrianglePerProjectionIteration = 8;
 static FAutoConsoleVariableRef CVarHairProjectionMaxTrianglePerProjectionIteration(TEXT("r.HairStrands.Projection.MaxTrianglePerIteration"), GHairProjectionMaxTrianglePerProjectionIteration, TEXT("Change the number of triangles which are iterated over during one projection iteration step. In kilo triangle (e.g., 8 == 8000 triangles). Default is 8."));
 
 
+///////////////////////////////////////////////////////////////////////////////////////////////////
+class FMarkMeshSectionIdCS : public FGlobalShader
+{
+	DECLARE_GLOBAL_SHADER(FMarkMeshSectionIdCS);
+	SHADER_USE_PARAMETER_STRUCT(FMarkMeshSectionIdCS, FGlobalShader);
+
+	BEGIN_SHADER_PARAMETER_STRUCT(FParameters, )
+		SHADER_PARAMETER(uint32, MeshSectionId)
+		SHADER_PARAMETER(uint32, MeshSectionPrimitiveCount)
+		SHADER_PARAMETER(uint32, MeshMaxIndexCount)
+		SHADER_PARAMETER(uint32, MeshMaxVertexCount)
+		SHADER_PARAMETER(uint32, MeshIndexOffset)
+		SHADER_PARAMETER_SRV(Buffer<uint>, MeshIndexBuffer)
+		SHADER_PARAMETER_RDG_BUFFER_UAV(RWBuffer<uint32>, OutVertexSectionId)
+
+	END_SHADER_PARAMETER_STRUCT()
+
+public:
+	static bool ShouldCompilePermutation(const FGlobalShaderPermutationParameters& Parameters) { return Parameters.Platform == SP_PCD3D_SM5; }
+	static void ModifyCompilationEnvironment(const FGlobalShaderPermutationParameters& Parameters, FShaderCompilerEnvironment& OutEnvironment)
+	{
+		FGlobalShader::ModifyCompilationEnvironment(Parameters, OutEnvironment);
+		OutEnvironment.SetDefine(TEXT("SHADER_SECTIONID"), 1);
+	}
+};
+
+IMPLEMENT_GLOBAL_SHADER(FMarkMeshSectionIdCS, "/Engine/Private/HairStrands/HairStrandsMeshProjection.usf", "MainMarkSectionIdCS", SF_Compute);
+
+static FRDGBufferRef AddMeshSectionId(
+	FRDGBuilder& GraphBuilder,
+	TShaderMap<FGlobalShaderType>* ShaderMap,
+	const FHairStrandsProjectionMeshData::LOD& MeshData)
+{
+	const int32 SectionCount = MeshData.Sections.Num();
+	if (SectionCount < 0)
+		return nullptr;
+
+	FRDGBufferRef VertexSectionIdBuffer = GraphBuilder.CreateBuffer(FRDGBufferDesc::CreateBufferDesc(sizeof(uint32), MeshData.Sections[0].TotalVertexCount), TEXT("SectionId"));
+	for (const FHairStrandsProjectionMeshData::Section& MeshSection : MeshData.Sections)
+	{	
+		FMarkMeshSectionIdCS::FParameters* Parameters = GraphBuilder.AllocParameters<FMarkMeshSectionIdCS::FParameters>();
+		Parameters->MeshSectionId				= MeshSection.SectionIndex;
+		Parameters->MeshSectionPrimitiveCount	= MeshSection.NumPrimitives;
+		Parameters->MeshMaxIndexCount			= MeshSection.TotalIndexCount;
+		Parameters->MeshMaxVertexCount			= MeshSection.TotalVertexCount;
+		Parameters->MeshIndexOffset				= MeshSection.IndexBaseIndex;
+		Parameters->MeshIndexBuffer				= MeshSection.IndexBuffer;
+		Parameters->OutVertexSectionId			= GraphBuilder.CreateUAV(VertexSectionIdBuffer, PF_R32_UINT);
+
+		const FIntVector DispatchGroupCount = FComputeShaderUtils::GetGroupCount(MeshSection.TotalVertexCount, 128);
+		check(DispatchGroupCount.X < 65536);
+		TShaderMapRef<FMarkMeshSectionIdCS> ComputeShader(ShaderMap);
+		FComputeShaderUtils::AddPass(
+			GraphBuilder,
+			RDG_EVENT_NAME("HairStrandsMarkVertexSectionId"),
+			*ComputeShader,
+			Parameters,
+			DispatchGroupCount);
+	}
+
+	return VertexSectionIdBuffer;
+}
+
+///////////////////////////////////////////////////////////////////////////////////////////////////
+class FMeshTransferCS : public FGlobalShader
+{
+	DECLARE_GLOBAL_SHADER(FMeshTransferCS);
+	SHADER_USE_PARAMETER_STRUCT(FMeshTransferCS, FGlobalShader);
+
+	BEGIN_SHADER_PARAMETER_STRUCT(FParameters, )
+		SHADER_PARAMETER(uint32, bNeedClear)
+		SHADER_PARAMETER(uint32, Source_MeshPrimitiveCount_Iteration)
+		SHADER_PARAMETER(uint32, Source_MeshMaxIndexCount)
+		SHADER_PARAMETER(uint32, Source_MeshMaxVertexCount)
+		SHADER_PARAMETER(uint32, Source_MeshIndexOffset)
+		SHADER_PARAMETER(uint32, Source_MeshUVsChannelOffset)
+		SHADER_PARAMETER(uint32, Source_MeshUVsChannelCount)
+		SHADER_PARAMETER(uint32, Target_MeshMaxVertexCount)
+		SHADER_PARAMETER(uint32, Target_MeshUVsChannelOffset)
+		SHADER_PARAMETER(uint32, Target_MeshUVsChannelCount)
+		SHADER_PARAMETER(uint32, Target_SectionId)
+		SHADER_PARAMETER_SRV(Buffer<uint>, Source_MeshIndexBuffer)
+		SHADER_PARAMETER_SRV(Buffer<float>, Source_MeshPositionBuffer)
+		SHADER_PARAMETER_SRV(Buffer<float2>, Source_MeshUVsBuffer)
+		SHADER_PARAMETER_SRV(Buffer<float2>, Target_MeshUVsBuffer)
+		SHADER_PARAMETER_RDG_BUFFER_SRV(Buffer<uint>, Target_VertexSectionId)
+		SHADER_PARAMETER_UAV(RWBuffer<float>, Target_MeshPositionBuffer)
+		SHADER_PARAMETER_RDG_BUFFER_UAV(RWBuffer<uint32>, OutDistanceBuffer)
+
+	END_SHADER_PARAMETER_STRUCT()
+
+public:
+	static bool ShouldCompilePermutation(const FGlobalShaderPermutationParameters& Parameters) { return Parameters.Platform == SP_PCD3D_SM5; }
+	static void ModifyCompilationEnvironment(const FGlobalShaderPermutationParameters& Parameters, FShaderCompilerEnvironment& OutEnvironment)
+	{
+		FGlobalShader::ModifyCompilationEnvironment(Parameters, OutEnvironment);
+		OutEnvironment.SetDefine(TEXT("SHADER_MESHTRANSFER"), 1);
+	}
+};
+
+IMPLEMENT_GLOBAL_SHADER(FMeshTransferCS, "/Engine/Private/HairStrands/HairStrandsMeshProjection.usf", "MainMeshTransferCS", SF_Compute);
+
+static void AddMeshTransferPass(
+	FRDGBuilder& GraphBuilder,
+	TShaderMap<FGlobalShaderType>* ShaderMap,
+	bool bClear,
+	const FHairStrandsProjectionMeshData::Section& SourceSectionData,
+	const FHairStrandsProjectionMeshData::Section& TargetSectionData,
+	FRDGBufferRef VertexSectionId, 
+	FRWBuffer& OutTargetRestPosition)
+{
+	if (!SourceSectionData.IndexBuffer ||
+		!SourceSectionData.PositionBuffer ||
+		 SourceSectionData.TotalIndexCount == 0 ||
+		 SourceSectionData.TotalVertexCount == 0||
+
+		!TargetSectionData.IndexBuffer ||
+		!TargetSectionData.PositionBuffer ||
+		 TargetSectionData.TotalIndexCount == 0 ||
+		 TargetSectionData.TotalVertexCount == 0)
+	{
+		return;
+	}
+	
+	FRDGBufferRef PositionDistanceBuffer = GraphBuilder.CreateBuffer(FRDGBufferDesc::CreateBufferDesc(sizeof(uint32), TargetSectionData.TotalVertexCount), TEXT("DistanceBuffer"));
+	FRDGBufferUAVRef PositionDistanceBufferUAV = GraphBuilder.CreateUAV(PositionDistanceBuffer, PF_R32_UINT);
+
+	// For projecting hair onto a skeletal mesh, 1 thread is spawn for each hair which iterates over all triangles.
+	// To avoid TDR, we split projection into multiple passes when the mesh is too large.
+	uint32 MeshPassNumPrimitive = 1024 * FMath::Clamp(GHairProjectionMaxTrianglePerProjectionIteration, 1, 256);
+	uint32 MeshPassCount = 1;
+	if (SourceSectionData.NumPrimitives < MeshPassNumPrimitive)
+	{
+		MeshPassNumPrimitive = SourceSectionData.NumPrimitives;
+	}
+	else
+	{
+		MeshPassCount = FMath::CeilToInt(SourceSectionData.NumPrimitives / float(MeshPassNumPrimitive));
+	}
+
+	FRDGBufferSRVRef VertexSectionIdSRV = GraphBuilder.CreateSRV(VertexSectionId, PF_R32_UINT);
+	for (uint32 MeshPassIt = 0; MeshPassIt < MeshPassCount; ++MeshPassIt)
+	{
+		FMeshTransferCS::FParameters* Parameters = GraphBuilder.AllocParameters<FMeshTransferCS::FParameters>();
+		Parameters->bNeedClear = bClear ? 1 : 0;
+
+		Parameters->Source_MeshPrimitiveCount_Iteration = (MeshPassIt < MeshPassCount - 1) ? MeshPassNumPrimitive : (SourceSectionData.NumPrimitives - MeshPassNumPrimitive * MeshPassIt);
+		Parameters->Source_MeshMaxIndexCount		= SourceSectionData.TotalIndexCount;
+		Parameters->Source_MeshMaxVertexCount		= SourceSectionData.TotalVertexCount;
+		Parameters->Source_MeshIndexOffset			= SourceSectionData.IndexBaseIndex + (MeshPassNumPrimitive * MeshPassIt * 3);
+		Parameters->Source_MeshUVsChannelOffset		= SourceSectionData.UVsChannelOffset;
+		Parameters->Source_MeshUVsChannelCount		= SourceSectionData.UVsChannelCount;
+		Parameters->Source_MeshIndexBuffer			= SourceSectionData.IndexBuffer;
+		Parameters->Source_MeshPositionBuffer		= SourceSectionData.PositionBuffer;
+		Parameters->Source_MeshUVsBuffer			= SourceSectionData.UVsBuffer;
+
+		Parameters->Target_MeshMaxVertexCount		= TargetSectionData.TotalVertexCount;
+		Parameters->Target_MeshUVsChannelOffset		= TargetSectionData.UVsChannelOffset;
+		Parameters->Target_MeshUVsChannelCount		= TargetSectionData.UVsChannelCount;
+		Parameters->Target_MeshUVsBuffer			= TargetSectionData.UVsBuffer;
+		Parameters->Target_MeshPositionBuffer		= OutTargetRestPosition.UAV;
+		Parameters->Target_VertexSectionId			= VertexSectionIdSRV;
+		Parameters->Target_SectionId				= TargetSectionData.SectionIndex;
+
+		Parameters->OutDistanceBuffer				= PositionDistanceBufferUAV;
+
+		const FIntVector DispatchGroupCount = FComputeShaderUtils::GetGroupCount(TargetSectionData.TotalVertexCount, 128);
+		check(DispatchGroupCount.X < 65536);
+		TShaderMapRef<FMeshTransferCS> ComputeShader(ShaderMap);
+		FComputeShaderUtils::AddPass(
+			GraphBuilder,
+			RDG_EVENT_NAME("HairStrandsTransferMesh"),
+			*ComputeShader,
+			Parameters,
+			DispatchGroupCount);
+		bClear = false;
+	}
+}
+
 ///////////////////////////////////////////////////////////////////////////////////////////////////
 class FHairMeshProjectionCS : public FGlobalShader
 {
@@ -27,6 +206,7 @@ class FHairMeshProjectionCS : public FGlobalShader
 		SHADER_PARAMETER(uint32, MeshMaxIndexCount)
 		SHADER_PARAMETER(uint32, MeshMaxVertexCount)
 		SHADER_PARAMETER(uint32, MeshIndexOffset)
+
 		SHADER_PARAMETER_SRV(Buffer, MeshIndexBuffer)
 		SHADER_PARAMETER_SRV(Buffer, MeshPositionBuffer)
 
@@ -35,11 +215,17 @@ class FHairMeshProjectionCS : public FGlobalShader
 
 		SHADER_PARAMETER_UAV(RWBuffer, OutRootTriangleIndex)
 		SHADER_PARAMETER_UAV(RWBuffer, OutRootTriangleBarycentrics)
-		SHADER_PARAMETER_RDG_BUFFER_UAV(StructuredBuffer, OutRootTriangleDistance)
+		SHADER_PARAMETER_RDG_BUFFER_UAV(RWStructuredBuffer, OutRootTriangleDistance)
+
 	END_SHADER_PARAMETER_STRUCT()
 
 public:
 	static bool ShouldCompilePermutation(const FGlobalShaderPermutationParameters& Parameters) { return Parameters.Platform == SP_PCD3D_SM5; }
+	static void ModifyCompilationEnvironment(const FGlobalShaderPermutationParameters& Parameters, FShaderCompilerEnvironment& OutEnvironment)
+	{
+		FGlobalShader::ModifyCompilationEnvironment(Parameters, OutEnvironment);
+		OutEnvironment.SetDefine(TEXT("SHADER_PROJECTION"), 1);
+	}
 };
 
 IMPLEMENT_GLOBAL_SHADER(FHairMeshProjectionCS, "/Engine/Private/HairStrands/HairStrandsMeshProjection.usf", "MainCS", SF_Compute);
@@ -66,6 +252,13 @@ static void AddHairStrandMeshProjectionPass(
 		return;
 	}
 
+	// The current shader code HairStrandsMeshProjection.usf encode the section ID onto the highest 3bits of a 32bits uint. 
+	// This limits the number of section to 8. See EncodeTriangleIndex & DecodeTriangleIndex functions in 
+	// HairStarndsMeshProjectionCommon.ush for mode details.
+	//
+	// This could be increase if necessary.
+	check(MeshSectionData.SectionIndex < 8);
+
 	// For projecting hair onto a skeletal mesh, 1 thread is spawn for each hair which iterates over all triangles.
 	// To avoid TDR, we split projection into multiple passes when the mesh is too large.
 	uint32 MeshPassNumPrimitive = 1024 * FMath::Clamp(GHairProjectionMaxTrianglePerProjectionIteration, 1, 256);
@@ -96,6 +289,7 @@ static void AddHairStrandMeshProjectionPass(
 		Parameters->MeshPrimitiveOffset_Iteration	= MeshPassNumPrimitive * MeshPassIt;
 		Parameters->MeshPrimitiveCount_Iteration	= (MeshPassIt < MeshPassCount-1) ? MeshPassNumPrimitive : (MeshSectionData.NumPrimitives - MeshPassNumPrimitive * MeshPassIt);
 
+		// The projection is always done onto the source/rest mesh
 		Parameters->OutRootTriangleIndex		= RootData.LODDatas[LODIndex].RootTriangleIndexBuffer->UAV;
 		Parameters->OutRootTriangleBarycentrics = RootData.LODDatas[LODIndex].RootTriangleBarycentricBuffer->UAV;
 		Parameters->OutRootTriangleDistance		= DistanceUAV;
@@ -113,17 +307,16 @@ static void AddHairStrandMeshProjectionPass(
 }
 
 void ProjectHairStrandsOntoMesh(
-	FRHICommandListImmediate& RHICmdList, 
+	FRDGBuilder& GraphBuilder,
 	TShaderMap<FGlobalShaderType>* ShaderMap,
 	const int32 LODIndex,
-	const FHairStrandsProjectionMeshData& ProjectionMeshData, 
+	const FHairStrandsProjectionMeshData& ProjectionMeshData,
 	FHairStrandsProjectionHairData::HairGroup& ProjectionHairData)
 {
 	if (LODIndex < 0 || LODIndex >= ProjectionHairData.LODDatas.Num())
 		return;
 
-	FRDGBuilder GraphBuilder(RHICmdList);
-	FRDGBufferRef RootDistanceBuffer = GraphBuilder.CreateBuffer(FRDGBufferDesc::CreateStructuredDesc(sizeof(float),  ProjectionHairData.RootCount),	TEXT("HairStrandsTriangleDistance"));
+	FRDGBufferRef RootDistanceBuffer = GraphBuilder.CreateBuffer(FRDGBufferDesc::CreateStructuredDesc(sizeof(float), ProjectionHairData.RootCount), TEXT("HairStrandsTriangleDistance"));
 
 	bool ClearDistance = true;
 	for (const FHairStrandsProjectionMeshData::Section& MeshSection : ProjectionMeshData.LODs[LODIndex].Sections)
@@ -133,7 +326,35 @@ void ProjectHairStrandsOntoMesh(
 		ProjectionHairData.LODDatas[LODIndex].bIsValid = true;
 		ClearDistance = false;
 	}
-	GraphBuilder.Execute();
+}
+
+void TransferMesh(
+	FRDGBuilder& GraphBuilder,
+	TShaderMap<FGlobalShaderType>* ShaderMap,
+	const int32 LODIndex,
+	const FHairStrandsProjectionMeshData& SourceMeshData,
+	const FHairStrandsProjectionMeshData& TargetMeshData,
+	FRWBuffer& OutPositionBuffer)
+{
+	if (LODIndex < 0)
+		return;
+
+	// LODs are transfered using the LOD0 of the source mesh, as the LOD count can mismatch between source and target meshes.
+	const int32 SourceLODIndex = 0;
+	const int32 TargetLODIndex = LODIndex;
+
+	// Assume that the section 0 contains the head section, which is where the hair/facial hair should be projected on
+	const uint32 SourceSectionIndex = 0;
+	const uint32 TargetSectionIndex = 0;
+
+	const int32 SectionCount = TargetMeshData.LODs[TargetLODIndex].Sections.Num();
+	if (SectionCount < 0)
+		return;
+
+	FRDGBufferRef VertexSectionId = AddMeshSectionId(GraphBuilder, ShaderMap, TargetMeshData.LODs[TargetLODIndex]);
+	const FHairStrandsProjectionMeshData::Section& SourceMeshSection = SourceMeshData.LODs[SourceLODIndex].Sections[SourceSectionIndex];
+	const FHairStrandsProjectionMeshData::Section& TargetMeshSection = TargetMeshData.LODs[TargetLODIndex].Sections[TargetSectionIndex];
+	AddMeshTransferPass(GraphBuilder, ShaderMap, true, SourceMeshSection, TargetMeshSection, VertexSectionId, OutPositionBuffer);
 }
 
 ///////////////////////////////////////////////////////////////////////////////////////////////////
@@ -142,16 +363,21 @@ class FHairUpdateMeshTriangleCS : public FGlobalShader
 	DECLARE_GLOBAL_SHADER(FHairUpdateMeshTriangleCS);
 	SHADER_USE_PARAMETER_STRUCT(FHairUpdateMeshTriangleCS, FGlobalShader);
 
+	class FUpdateUVs : SHADER_PERMUTATION_INT("PERMUTATION_WITHUV", 2);
+	using FPermutationDomain = TShaderPermutationDomain<FUpdateUVs>;
+
 	BEGIN_SHADER_PARAMETER_STRUCT(FParameters, )
-		SHADER_PARAMETER(FVector, RootTrianglePositionOffset)
 		SHADER_PARAMETER(uint32, MaxRootCount)
 		
 		SHADER_PARAMETER(uint32, MeshSectionIndex)
 		SHADER_PARAMETER(uint32, MeshMaxIndexCount)
 		SHADER_PARAMETER(uint32, MeshMaxVertexCount)
 		SHADER_PARAMETER(uint32, MeshIndexOffset)
+		SHADER_PARAMETER(uint32, MeshUVsChannelOffset)
+		SHADER_PARAMETER(uint32, MeshUVsChannelCount)
 		SHADER_PARAMETER_SRV(Buffer, MeshIndexBuffer)
 		SHADER_PARAMETER_SRV(Buffer, MeshPositionBuffer)
+		SHADER_PARAMETER_SRV(Buffer, MeshUVsBuffer)
 
 		SHADER_PARAMETER_SRV(Buffer, RootTriangleIndex)
 		SHADER_PARAMETER_UAV(StructuredBuffer, OutRootTrianglePosition0)
@@ -182,33 +408,42 @@ static void AddHairStrandUpdateMeshTrianglesPass(
 
 	FHairUpdateMeshTriangleCS::FParameters* Parameters = GraphBuilder.AllocParameters<FHairUpdateMeshTriangleCS::FParameters>();
 	Parameters->MaxRootCount		= RootData.RootCount;
-	Parameters->RootTriangleIndex	= LODData.RootTriangleIndexBuffer->SRV;
 	Parameters->MeshSectionIndex	= MeshSectionData.SectionIndex;
 	Parameters->MeshMaxIndexCount	= MeshSectionData.TotalIndexCount;
 	Parameters->MeshMaxVertexCount	= MeshSectionData.TotalVertexCount;
 	Parameters->MeshIndexOffset		= MeshSectionData.IndexBaseIndex;
 	Parameters->MeshIndexBuffer		= MeshSectionData.IndexBuffer;
 	Parameters->MeshPositionBuffer	= MeshSectionData.PositionBuffer;
+	Parameters->MeshUVsBuffer		= MeshSectionData.UVsBuffer;
+	Parameters->MeshUVsChannelOffset= MeshSectionData.UVsChannelOffset;
+	Parameters->MeshUVsChannelCount	= MeshSectionData.UVsChannelCount;
 
+	Parameters->RootTriangleIndex = LODData.RootTriangleIndexBuffer->SRV;
 	if (Type == HairStrandsTriangleType::RestPose)
 	{
-		Parameters->RootTrianglePositionOffset = *LODData.RestPositionOffset;
 		Parameters->OutRootTrianglePosition0 = LODData.RestRootTrianglePosition0Buffer->UAV;
 		Parameters->OutRootTrianglePosition1 = LODData.RestRootTrianglePosition1Buffer->UAV;
 		Parameters->OutRootTrianglePosition2 = LODData.RestRootTrianglePosition2Buffer->UAV;
 	}
 	else if (Type == HairStrandsTriangleType::DeformedPose)
 	{
-		Parameters->RootTrianglePositionOffset = *LODData.DeformedPositionOffset;
 		Parameters->OutRootTrianglePosition0 = LODData.DeformedRootTrianglePosition0Buffer->UAV;
 		Parameters->OutRootTrianglePosition1 = LODData.DeformedRootTrianglePosition1Buffer->UAV;
 		Parameters->OutRootTrianglePosition2 = LODData.DeformedRootTrianglePosition2Buffer->UAV;
 		if (LODData.Status) (*LODData.Status) = FHairStrandsProjectionHairData::LODData::EStatus::Completed;
 	}
+	else
+	{
+		// error
+		return;
+	}
+
+	FHairUpdateMeshTriangleCS::FPermutationDomain PermutationVector;
+	PermutationVector.Set<FHairUpdateMeshTriangleCS::FUpdateUVs>(1);
 
 	const FIntVector DispatchGroupCount = FComputeShaderUtils::GetGroupCount(RootData.RootCount, 128);
 	check(DispatchGroupCount.X < 65536);
-	TShaderMapRef<FHairUpdateMeshTriangleCS> ComputeShader(ShaderMap);
+	TShaderMapRef<FHairUpdateMeshTriangleCS> ComputeShader(ShaderMap, PermutationVector);
 	FComputeShaderUtils::AddPass(
 		GraphBuilder,
 		RDG_EVENT_NAME("HairStrandsTriangleMeshUpdate"),
@@ -218,19 +453,15 @@ static void AddHairStrandUpdateMeshTrianglesPass(
 }
 
 void UpdateHairStrandsMeshTriangles(
-	FRHICommandListImmediate& RHICmdList,
+	FRDGBuilder& GraphBuilder,
 	TShaderMap<FGlobalShaderType>* ShaderMap,
 	const int32 LODIndex,
 	const HairStrandsTriangleType Type,
 	const FHairStrandsProjectionMeshData::LOD& ProjectionMeshData,
 	FHairStrandsProjectionHairData::HairGroup& ProjectionHairData)
 {
-	FRDGBuilder GraphBuilder(RHICmdList);
-
 	for (const FHairStrandsProjectionMeshData::Section& MeshSection : ProjectionMeshData.Sections)
 	{
 		AddHairStrandUpdateMeshTrianglesPass(GraphBuilder, ShaderMap, LODIndex, Type, MeshSection, ProjectionHairData);
 	}
-	
-	GraphBuilder.Execute();
 }
\ No newline at end of file
diff --git a/Engine/Source/Runtime/Renderer/Private/HairStrands/HairStrandsMeshProjection.h b/Engine/Source/Runtime/Renderer/Private/HairStrands/HairStrandsMeshProjection.h
--- a/Engine/Source/Runtime/Renderer/Private/HairStrands/HairStrandsMeshProjection.h
+++ b/Engine/Source/Runtime/Renderer/Private/HairStrands/HairStrandsMeshProjection.h
@@ -9,12 +9,20 @@
 
 /* Project hair strands onto a LOD mesh */
 void ProjectHairStrandsOntoMesh(
-	FRHICommandListImmediate& RHICmdList, 
+	class FRDGBuilder& GraphBuilder,
 	TShaderMap<FGlobalShaderType>* ShaderMap,
 	const int32 LODIndex,
-	const FHairStrandsProjectionMeshData& ProjectionMeshData, 
+	const FHairStrandsProjectionMeshData& ProjectionMeshData,
 	FHairStrandsProjectionHairData::HairGroup& ProjectionHairData);
 
+void TransferMesh(
+	FRDGBuilder& GraphBuilder,
+	TShaderMap<FGlobalShaderType>* ShaderMap,
+	const int32 LODIndex,
+	const FHairStrandsProjectionMeshData& SourceMeshData,
+	const FHairStrandsProjectionMeshData& TargetMeshData,
+	struct FRWBuffer& OutPositionBuffer);
+
 enum class HairStrandsTriangleType
 {
 	RestPose,
@@ -23,7 +31,7 @@ enum class HairStrandsTriangleType
 
 /* Update the triangles information on which hair stands have been projected */
 void UpdateHairStrandsMeshTriangles(
-	FRHICommandListImmediate& RHICmdList,
+	class FRDGBuilder& GraphBuilder,
 	TShaderMap<FGlobalShaderType>* ShaderMap,
 	const int32 LODIndex,
 	const HairStrandsTriangleType Type,
diff --git a/Engine/Source/Runtime/Renderer/Private/HairStrands/HairStrandsRendering.h b/Engine/Source/Runtime/Renderer/Private/HairStrands/HairStrandsRendering.h
--- a/Engine/Source/Runtime/Renderer/Private/HairStrands/HairStrandsRendering.h
+++ b/Engine/Source/Runtime/Renderer/Private/HairStrands/HairStrandsRendering.h
@@ -44,4 +44,6 @@ void RunHairStrandsInterpolation(
 	EHairStrandsInterpolationType Type,
 	FHairStrandClusterData* ClusterData);
 
-
+void RunHairStrandsBindingQueries(
+	FRHICommandListImmediate& RHICmdList,
+	TShaderMap<FGlobalShaderType>* ShaderMap);
\ No newline at end of file
diff --git a/Engine/Source/Runtime/Renderer/Private/HairStrands/HairStrandsUtils.cpp b/Engine/Source/Runtime/Renderer/Private/HairStrands/HairStrandsUtils.cpp
--- a/Engine/Source/Runtime/Renderer/Private/HairStrands/HairStrandsUtils.cpp
+++ b/Engine/Source/Runtime/Renderer/Private/HairStrands/HairStrandsUtils.cpp
@@ -3,6 +3,7 @@
 #include "HairStrandsUtils.h"
 #include "ScenePrivate.h"
 #include "HairStrandsCluster.h"
+#include "Rendering/SkeletalMeshRenderData.h"
 
 static float GHairR = 1;
 static float GHairTT = 1;
@@ -260,4 +261,39 @@ FVector4 PackHairRenderInfo(
 	Out.W = *((float*)(&BitField));
 
 	return Out;
-}
\ No newline at end of file
+}
+
+FHairStrandsProjectionMeshData ExtractMeshData(FSkeletalMeshRenderData* RenderData)
+{
+	FHairStrandsProjectionMeshData MeshData;
+	uint32 LODIndex = 0;
+	for (FSkeletalMeshLODRenderData& LODRenderData : RenderData->LODRenderData)
+	{
+		FHairStrandsProjectionMeshData::LOD& LOD = MeshData.LODs.AddDefaulted_GetRef();
+		uint32 SectionIndex = 0;
+		for (FSkelMeshRenderSection& InSection : LODRenderData.RenderSections)
+		{
+			// Pick between float and halt
+			const uint32 UVSizeInByte = (LODRenderData.StaticVertexBuffers.StaticMeshVertexBuffer.GetUseFullPrecisionUVs() ? 4 : 2) * 2;
+
+			FHairStrandsProjectionMeshData::Section& OutSection = LOD.Sections.AddDefaulted_GetRef();
+			OutSection.UVsChannelOffset = 0; // Assume that we needs to pair meshes based on UVs 0
+			OutSection.UVsChannelCount = LODRenderData.StaticVertexBuffers.StaticMeshVertexBuffer.GetNumTexCoords();
+			OutSection.UVsBuffer = LODRenderData.StaticVertexBuffers.StaticMeshVertexBuffer.GetTexCoordsSRV();
+			OutSection.PositionBuffer = LODRenderData.StaticVertexBuffers.PositionVertexBuffer.GetSRV();
+			OutSection.IndexBuffer = LODRenderData.MultiSizeIndexContainer.GetIndexBuffer()->GetSRV();
+			OutSection.TotalVertexCount = LODRenderData.StaticVertexBuffers.PositionVertexBuffer.GetNumVertices();
+			OutSection.TotalIndexCount = LODRenderData.MultiSizeIndexContainer.GetIndexBuffer()->Num();
+			OutSection.NumPrimitives = InSection.NumTriangles;
+			OutSection.VertexBaseIndex = InSection.BaseVertexIndex;
+			OutSection.IndexBaseIndex = InSection.BaseIndex;
+			OutSection.SectionIndex = SectionIndex;
+			OutSection.LODIndex = LODIndex;
+
+			++SectionIndex;
+		}
+		++LODIndex;
+	}
+
+	return MeshData;
+}
diff --git a/Engine/Source/Runtime/Renderer/Private/HairStrands/HairStrandsVisibility.cpp b/Engine/Source/Runtime/Renderer/Private/HairStrands/HairStrandsVisibility.cpp
--- a/Engine/Source/Runtime/Renderer/Private/HairStrands/HairStrandsVisibility.cpp
+++ b/Engine/Source/Runtime/Renderer/Private/HairStrands/HairStrandsVisibility.cpp
@@ -133,7 +133,7 @@ uint32 GetPPLLMaxTotalRenderNode(FIntPoint Resolution)
 
 uint32 GetHairVisibilitySampleCount()
 {
-	return GetHairVisibilityRenderMode() == HairVisibilityRenderMode_MSAA ? FMath::Clamp(GHairVisibilitySampleCount, 1, 16) : 1;
+	return GetHairVisibilityRenderMode() == HairVisibilityRenderMode_MSAA ? FMath::Clamp(GHairVisibilitySampleCount,1,16) : 1;
 }
 
 void SetUpViewHairRenderInfo(const FViewInfo& ViewInfo, bool bEnableMSAA, FVector4& OutHairRenderInfo)
@@ -159,13 +159,13 @@ static bool IsCompatibleWithHairVisibility(const FMeshMaterialShaderPermutationP
 /////////////////////////////////////////////////////////////////////////////////////////
 
 BEGIN_GLOBAL_SHADER_PARAMETER_STRUCT(FMaterialPassParameters, )
-SHADER_PARAMETER(FIntPoint, MaxResolution)
-SHADER_PARAMETER(uint32, MaxSampleCount)
-SHADER_PARAMETER(uint32, NodeGroupSize)
-SHADER_PARAMETER_TEXTURE(Texture2D<uint>, NodeIndex)
-SHADER_PARAMETER_SRV(StructuredBuffer<uint>, NodeCoord)
-SHADER_PARAMETER_SRV(StructuredBuffer<FNodeVis>, NodeVis)
-SHADER_PARAMETER_SRV(Buffer<uint>, IndirectArgs)
+	SHADER_PARAMETER(FIntPoint, MaxResolution)
+	SHADER_PARAMETER(uint32, MaxSampleCount)
+	SHADER_PARAMETER(uint32, NodeGroupSize)
+	SHADER_PARAMETER_TEXTURE(Texture2D<uint>, NodeIndex)
+	SHADER_PARAMETER_SRV(StructuredBuffer<uint>, NodeCoord)
+	SHADER_PARAMETER_SRV(StructuredBuffer<FNodeVis>, NodeVis)
+	SHADER_PARAMETER_SRV(Buffer<uint>, IndirectArgs)
 END_GLOBAL_SHADER_PARAMETER_STRUCT()
 IMPLEMENT_GLOBAL_SHADER_PARAMETER_STRUCT(FMaterialPassParameters, "MaterialPassParameters");
 
@@ -202,10 +202,11 @@ IMPLEMENT_MATERIAL_SHADER_TYPE(, FHairMaterialVS, TEXT("/Engine/Private/HairStra
 class FHairMaterialShaderElementData : public FMeshMaterialShaderElementData
 {
 public:
-	FHairMaterialShaderElementData(int32 MacroGroupId, int32 MaterialId, int32 PrimitiveId) : MaterialPass_MacroGroupId(MacroGroupId), MaterialPass_MaterialId(MaterialId), MaterialPass_PrimitiveId(PrimitiveId) { }
+	FHairMaterialShaderElementData(int32 MacroGroupId, int32 MaterialId, int32 PrimitiveId, uint32 LightChannelMask) : MaterialPass_MacroGroupId(MacroGroupId), MaterialPass_MaterialId(MaterialId), MaterialPass_PrimitiveId(PrimitiveId), MaterialPass_LightChannelMask(LightChannelMask){ }
 	uint32 MaterialPass_MacroGroupId;
 	uint32 MaterialPass_MaterialId;
 	uint32 MaterialPass_PrimitiveId;
+	uint32 MaterialPass_LightChannelMask;
 };
 
 class FHairMaterialPS : public FMeshMaterialShader
@@ -222,6 +223,7 @@ public:
 		MaterialPass_MacroGroupId.Bind(Initializer.ParameterMap, TEXT("MaterialPass_MacroGroupId"));
 		MaterialPass_MaterialId.Bind(Initializer.ParameterMap, TEXT("MaterialPass_MaterialId"));
 		MaterialPass_PrimitiveId.Bind(Initializer.ParameterMap, TEXT("MaterialPass_PrimitiveId"));
+		MaterialPass_LightChannelMask.Bind(Initializer.ParameterMap, TEXT("MaterialPass_LightChannelMask"));
 
 		// Fake binding for avoiding error from the shader compiler. These output are actually properly bound by the rendering 
 		// pass, and not directly by this material shader
@@ -248,6 +250,7 @@ public:
 		Ar << MaterialPass_MacroGroupId;
 		Ar << MaterialPass_MaterialId;
 		Ar << MaterialPass_PrimitiveId;
+		Ar << MaterialPass_LightChannelMask;
 		Ar << OutNodeData;
 		Ar << OutNodeVelocity;
 		return bShaderHasOutdatedParameters;
@@ -267,12 +270,14 @@ public:
 		ShaderBindings.Add(MaterialPass_MacroGroupId, ShaderElementData.MaterialPass_MacroGroupId);
 		ShaderBindings.Add(MaterialPass_MaterialId, ShaderElementData.MaterialPass_MaterialId);
 		ShaderBindings.Add(MaterialPass_PrimitiveId, ShaderElementData.MaterialPass_PrimitiveId);
+		ShaderBindings.Add(MaterialPass_LightChannelMask, ShaderElementData.MaterialPass_LightChannelMask);
 	}
 
 private:
 	FShaderParameter MaterialPass_MacroGroupId;
 	FShaderParameter MaterialPass_MaterialId;
 	FShaderParameter MaterialPass_PrimitiveId;
+	FShaderParameter MaterialPass_LightChannelMask;
 	FShaderParameter OutNodeData;
 	FShaderParameter OutNodeVelocity;
 };
@@ -302,7 +307,8 @@ private:
 		const FMaterial& RESTRICT MaterialResource,
 		const int32 MacroGroupId,
 		const int32 HairMaterialId,
-		const int32 HairPrimitiveId);
+		const int32 HairPrimitiveId,
+		const uint32 HairPrimitiveLightChannelMask);
 
 	FMeshPassProcessorRenderState PassDrawRenderState;
 };
@@ -321,7 +327,7 @@ void FHairMaterialProcessor::AddMeshBatch(const FMeshBatch& RESTRICT MeshBatch,
 	const FMaterial& Material = MeshBatch.MaterialRenderProxy->GetMaterialWithFallback(FeatureLevel, FallbackMaterialRenderProxyPtr);
 	const bool bIsCompatible = IsCompatibleWithHairStrands(&Material, FeatureLevel);
 	const bool bIsHairStrandsFactory = MeshBatch.VertexFactory->GetType()->GetId() == CompatibleVF->GetId();
-	const bool bShouldRender = (!PrimitiveSceneProxy && MeshBatch.Elements.Num() > 0) || (PrimitiveSceneProxy && PrimitiveSceneProxy->ShouldRenderInMainPass());
+	const bool bShouldRender = (!PrimitiveSceneProxy && MeshBatch.Elements.Num()>0) || (PrimitiveSceneProxy && PrimitiveSceneProxy->ShouldRenderInMainPass());
 
 	if (bIsCompatible
 		&& bIsHairStrandsFactory
@@ -330,7 +336,7 @@ void FHairMaterialProcessor::AddMeshBatch(const FMeshBatch& RESTRICT MeshBatch,
 	{
 		// For the mesh patch to be rendered a single triangle triangle to spawn the necessary amount of thread
 		FMeshBatch MeshBatchCopy = MeshBatch;
-		for (uint32 ElementIt = 0, ElementCount = uint32(MeshBatch.Elements.Num()); ElementIt < ElementCount; ++ElementIt)
+		for (uint32 ElementIt = 0, ElementCount=uint32(MeshBatch.Elements.Num()); ElementIt < ElementCount; ++ElementIt)
 		{
 			MeshBatchCopy.Elements[ElementIt].FirstIndex = 0;
 			MeshBatchCopy.Elements[ElementIt].NumPrimitives = 1;
@@ -343,9 +349,10 @@ void FHairMaterialProcessor::AddMeshBatch(const FMeshBatch& RESTRICT MeshBatch,
 		int32 ScenePrimitiveId = 0;
 		FPrimitiveSceneInfo* SceneInfo = PrimitiveSceneProxy ? PrimitiveSceneProxy->GetPrimitiveSceneInfo() : nullptr;
 		GetDrawCommandPrimitiveId(SceneInfo, MeshBatch.Elements[0], PrimitiveId, ScenePrimitiveId);
+		uint32 LightChannelMask = PrimitiveSceneProxy ? PrimitiveSceneProxy->GetLightingChannelMask() : 0;
 
 		const FMaterialRenderProxy& MaterialRenderProxy = FallbackMaterialRenderProxyPtr ? *FallbackMaterialRenderProxyPtr : *MeshBatch.MaterialRenderProxy;
-		Process(MeshBatchCopy, BatchElementMask, PrimitiveSceneProxy, StaticMeshId, MaterialRenderProxy, Material, MacroGroupId, HairMaterialId, PrimitiveId);
+		Process(MeshBatchCopy, BatchElementMask, PrimitiveSceneProxy, StaticMeshId, MaterialRenderProxy, Material, MacroGroupId, HairMaterialId, PrimitiveId, LightChannelMask);
 	}
 }
 
@@ -358,7 +365,8 @@ void FHairMaterialProcessor::Process(
 	const FMaterial& RESTRICT MaterialResource,
 	const int32 MacroGroupId,
 	const int32 HairMaterialId,
-	const int32 HairPrimitiveId)
+	const int32 HairPrimitiveId,
+	const uint32 HairPrimitiveLightChannelMask)
 {
 	const FVertexFactory* VertexFactory = MeshBatch.VertexFactory;
 
@@ -374,7 +382,7 @@ void FHairMaterialProcessor::Process(
 	}
 
 	FMeshPassProcessorRenderState DrawRenderState(PassDrawRenderState);
-	FHairMaterialShaderElementData ShaderElementData(MacroGroupId, HairMaterialId, HairPrimitiveId);
+	FHairMaterialShaderElementData ShaderElementData(MacroGroupId, HairMaterialId, HairPrimitiveId, HairPrimitiveLightChannelMask);
 	ShaderElementData.InitializeMeshMaterialData(ViewIfDynamicMeshCommand, PrimitiveSceneProxy, MeshBatch, StaticMeshId, false);
 
 	BuildMeshDrawCommands(
@@ -405,13 +413,13 @@ FHairMaterialProcessor::FHairMaterialProcessor(
 /////////////////////////////////////////////////////////////////////////////////////////
 
 BEGIN_SHADER_PARAMETER_STRUCT(FVisibilityMaterialPassParameters, )
-SHADER_PARAMETER_RDG_TEXTURE(Texture2D<uint>, NodeIndex)
-SHADER_PARAMETER_RDG_BUFFER_SRV(StructuredBuffer<uint>, NodeCoord)
-SHADER_PARAMETER_RDG_BUFFER_SRV(StructuredBuffer<FNodeVis>, NodeVis)
-SHADER_PARAMETER_RDG_BUFFER_SRV(Buffer<uint>, IndirectArgs)
-SHADER_PARAMETER_RDG_BUFFER_UAV(RWStructuredBuffer<FNodeData>, OutNodeData)
-SHADER_PARAMETER_RDG_BUFFER_UAV(RWBuffer<float2>, OutNodeVelocity)
-RENDER_TARGET_BINDING_SLOTS()
+	SHADER_PARAMETER_RDG_TEXTURE(Texture2D<uint>, NodeIndex)
+	SHADER_PARAMETER_RDG_BUFFER_SRV(StructuredBuffer<uint>,	NodeCoord)
+	SHADER_PARAMETER_RDG_BUFFER_SRV(StructuredBuffer<FNodeVis>, NodeVis)
+	SHADER_PARAMETER_RDG_BUFFER_SRV(Buffer<uint>, IndirectArgs)
+	SHADER_PARAMETER_RDG_BUFFER_UAV(RWStructuredBuffer<FNodeData>, OutNodeData)
+	SHADER_PARAMETER_RDG_BUFFER_UAV(RWBuffer<float2>, OutNodeVelocity)
+	RENDER_TARGET_BINDING_SLOTS()
 END_SHADER_PARAMETER_STRUCT()
 
 struct FMaterialPassOutput
@@ -438,7 +446,7 @@ static FMaterialPassOutput AddHairMaterialPass(
 	const uint32 MaxNodeCount = CompactNodeVis->Desc.NumElements;
 
 	FMaterialPassOutput Output;
-	Output.NodeData = GraphBuilder.CreateBuffer(FRDGBufferDesc::CreateStructuredDesc(sizeof(HairStrandsVisibilityInternal::NodeData), MaxNodeCount), TEXT("HairCompactNodeData"));
+	Output.NodeData		= GraphBuilder.CreateBuffer(FRDGBufferDesc::CreateStructuredDesc(sizeof(HairStrandsVisibilityInternal::NodeData), MaxNodeCount), TEXT("HairCompactNodeData"));
 	Output.NodeVelocity = GraphBuilder.CreateBuffer(FRDGBufferDesc::CreateBufferDesc(4, CompactNodeVis->Desc.NumElements), TEXT("HairCompactNodeVelocity"));
 
 	const uint32 ResolutionDim = FMath::CeilToInt(FMath::Sqrt(MaxNodeCount));
@@ -457,12 +465,12 @@ static FMaterialPassOutput AddHairMaterialPass(
 
 	// Add resources reference to the pass parameters, in order to get the resource lifetime extended to this pass
 	FVisibilityMaterialPassParameters* PassParameters = GraphBuilder.AllocParameters<FVisibilityMaterialPassParameters>();
-	PassParameters->NodeIndex = CompactNodeIndex;
-	PassParameters->NodeVis = GraphBuilder.CreateSRV(CompactNodeVis);
-	PassParameters->NodeCoord = GraphBuilder.CreateSRV(CompactNodeCoord);
-	PassParameters->IndirectArgs = GraphBuilder.CreateSRV(IndirectArgBuffer);
-	PassParameters->OutNodeData = GraphBuilder.CreateUAV(FRDGBufferUAVDesc(Output.NodeData));
-	PassParameters->OutNodeVelocity = GraphBuilder.CreateUAV(FRDGBufferUAVDesc(Output.NodeVelocity, FMaterialPassOutput::VelocityFormat));
+	PassParameters->NodeIndex		= CompactNodeIndex;
+	PassParameters->NodeVis			= GraphBuilder.CreateSRV(CompactNodeVis);
+	PassParameters->NodeCoord		= GraphBuilder.CreateSRV(CompactNodeCoord);
+	PassParameters->IndirectArgs	= GraphBuilder.CreateSRV(IndirectArgBuffer);
+	PassParameters->OutNodeData		= GraphBuilder.CreateUAV(FRDGBufferUAVDesc(Output.NodeData));
+	PassParameters->OutNodeVelocity	= GraphBuilder.CreateUAV(FRDGBufferUAVDesc(Output.NodeVelocity, FMaterialPassOutput::VelocityFormat));
 	PassParameters->RenderTargets[0] = FRenderTargetBinding(OutDummyTexture0, ERenderTargetLoadAction::EClear, 0);
 
 
@@ -476,13 +484,13 @@ static FMaterialPassOutput AddHairMaterialPass(
 		check(IsInRenderingThread());
 
 		FMaterialPassParameters MaterialPassParameters;
-		MaterialPassParameters.MaxResolution = Resolution;
-		MaterialPassParameters.NodeGroupSize = NodeGroupSize;
-		MaterialPassParameters.MaxSampleCount = MaxNodeCount;
-		MaterialPassParameters.NodeIndex = PassParameters->NodeIndex->GetPooledRenderTarget()->GetRenderTargetItem().ShaderResourceTexture;
-		MaterialPassParameters.NodeCoord = PassParameters->NodeCoord->GetRHI();
-		MaterialPassParameters.NodeVis = PassParameters->NodeVis->GetRHI();
-		MaterialPassParameters.IndirectArgs = PassParameters->IndirectArgs->GetRHI();
+		MaterialPassParameters.MaxResolution	= Resolution;
+		MaterialPassParameters.NodeGroupSize	= NodeGroupSize;
+		MaterialPassParameters.MaxSampleCount	= MaxNodeCount;
+		MaterialPassParameters.NodeIndex		= PassParameters->NodeIndex->GetPooledRenderTarget()->GetRenderTargetItem().ShaderResourceTexture;
+		MaterialPassParameters.NodeCoord		= PassParameters->NodeCoord->GetRHI();
+		MaterialPassParameters.NodeVis			= PassParameters->NodeVis->GetRHI();
+		MaterialPassParameters.IndirectArgs		= PassParameters->IndirectArgs->GetRHI();
 		TUniformBufferRef<FMaterialPassParameters> MaterialPassParametersBuffer = TUniformBufferRef<FMaterialPassParameters>::CreateUniformBufferImmediate(MaterialPassParameters, UniformBuffer_SingleFrame);
 
 		FMeshPassProcessorRenderState DrawRenderState(*ViewInfo, MaterialPassParametersBuffer);
@@ -498,7 +506,7 @@ static FMaterialPassOutput AddHairMaterialPass(
 			RHICmdList.SetViewport(0, 0, 0.0f, Resolution.X, Resolution.Y, 1.0f);
 			DrawRenderState.SetBlendState(TStaticBlendState<>::GetRHI());
 			DrawRenderState.SetDepthStencilState(TStaticDepthStencilState <false, CF_Always> ::GetRHI());
-
+			
 			FDynamicMeshDrawCommandStorage DynamicMeshDrawCommandStorage;
 			FMeshCommandOneFrameArray VisibleMeshDrawCommands;
 			FGraphicsMinimalPipelineStateSet PipelineStateSet;
@@ -528,7 +536,7 @@ static FMaterialPassOutput AddHairMaterialPass(
 }
 
 ///////////////////////////////////////////////////////////////////////////////////////////////////
-class FHairVelocityCS : public FGlobalShader
+class FHairVelocityCS: public FGlobalShader
 {
 	DECLARE_GLOBAL_SHADER(FHairVelocityCS);
 	SHADER_USE_PARAMETER_STRUCT(FHairVelocityCS, FGlobalShader);
@@ -545,7 +553,7 @@ class FHairVelocityCS : public FGlobalShader
 		SHADER_PARAMETER_RDG_TEXTURE_UAV(Texture2D, OutVelocityTexture)
 		SHADER_PARAMETER_STRUCT_REF(FViewUniformShaderParameters, ViewUniformBuffer)
 		SHADER_PARAMETER_STRUCT_REF(FSceneTexturesUniformParameters, SceneTexturesStruct)
-		END_SHADER_PARAMETER_STRUCT()
+	END_SHADER_PARAMETER_STRUCT()
 
 public:
 
@@ -609,24 +617,89 @@ static void AddHairVelocityPass(
 		FComputeShaderUtils::GetGroupCount(RectResolution, GroupSize));
 }
 
+///////////////////////////////////////////////////////////////////////////////////////////////////
+class FHairLightChannelMaskCS : public FGlobalShader
+{
+	DECLARE_GLOBAL_SHADER(FHairLightChannelMaskCS);
+	SHADER_USE_PARAMETER_STRUCT(FHairLightChannelMaskCS, FGlobalShader);
+
+	class FVendor : SHADER_PERMUTATION_INT("PERMUTATION_VENDOR", HairVisibilityVendorCount);
+	using FPermutationDomain = TShaderPermutationDomain<FVendor>;
+
+	BEGIN_SHADER_PARAMETER_STRUCT(FParameters, )
+		SHADER_PARAMETER(FIntPoint, OutputResolution)
+		SHADER_PARAMETER_RDG_BUFFER_SRV(StructuredBuffer, NodeData)
+		SHADER_PARAMETER_RDG_TEXTURE(Texture2D, NodeOffsetAndCount)
+		SHADER_PARAMETER_RDG_TEXTURE_UAV(Texture2D, OutLightChannelMaskTexture)
+	END_SHADER_PARAMETER_STRUCT()
+		
+public:
+
+	static bool ShouldCompilePermutation(const FGlobalShaderPermutationParameters& Parameters)
+	{
+		return IsHairStrandsSupported(Parameters.Platform);
+	}
+};
+
+IMPLEMENT_GLOBAL_SHADER(FHairLightChannelMaskCS, "/Engine/Private/HairStrands/HairStrandsLightChannelMask.usf", "MainCS", SF_Compute);
+
+static FRDGTextureRef AddHairLightChannelMaskPass(
+	FRDGBuilder& GraphBuilder,
+	const FViewInfo& View,
+	const FIntPoint Resolution,
+	FRDGBufferRef NodeData,
+	FRDGTextureRef NodeOffsetAndCount)
+{
+	FRDGTextureDesc Desc;
+	Desc.Extent = Resolution;
+	Desc.Depth = 0;
+	Desc.Format = PF_R32_UINT;
+	Desc.NumMips = 1;
+	Desc.NumSamples = 1;
+	Desc.Flags = TexCreate_None;
+	Desc.TargetableFlags = TexCreate_UAV | TexCreate_ShaderResource;
+	Desc.ClearValue = FClearValueBinding(0);
+	FRDGTextureRef OutLightChannelMaskTexture = GraphBuilder.CreateTexture(Desc, TEXT("HairLightChannelMask"));
+
+	FHairLightChannelMaskCS::FPermutationDomain PermutationVector;
+	PermutationVector.Set<FHairLightChannelMaskCS::FVendor>(GetVendor());
+
+	FHairLightChannelMaskCS::FParameters* PassParameters = GraphBuilder.AllocParameters<FHairLightChannelMaskCS::FParameters>();
+	PassParameters->OutputResolution = Resolution;
+	PassParameters->NodeData = GraphBuilder.CreateSRV(NodeData);
+	PassParameters->NodeOffsetAndCount = NodeOffsetAndCount;
+	PassParameters->OutLightChannelMaskTexture = GraphBuilder.CreateUAV(OutLightChannelMaskTexture);
+	
+	const FIntPoint GroupSize = GetVendorOptimalGroupSize2D();
+	TShaderMapRef<FHairLightChannelMaskCS> ComputeShader(View.ShaderMap, PermutationVector);
+	FComputeShaderUtils::AddPass(
+		GraphBuilder,
+		RDG_EVENT_NAME("HairLightChannelMask"),
+		*ComputeShader,
+		PassParameters,
+		FComputeShaderUtils::GetGroupCount(Resolution, GroupSize));
+
+	return OutLightChannelMaskTexture;
+}
+
 /////////////////////////////////////////////////////////////////////////////////////////
 BEGIN_GLOBAL_SHADER_PARAMETER_STRUCT(FVisibilityPassGlobalParameters, )
-SHADER_PARAMETER(uint32, MaxPPLLNodeCount)
+	SHADER_PARAMETER(uint32, MaxPPLLNodeCount)
 END_GLOBAL_SHADER_PARAMETER_STRUCT()
 IMPLEMENT_GLOBAL_SHADER_PARAMETER_STRUCT(FVisibilityPassGlobalParameters, "VisibilityPassGlobalParameters");
 
 BEGIN_SHADER_PARAMETER_STRUCT(FVisibilityPassParameters, )
-SHADER_PARAMETER(uint32, HairVisibilityPass_MaxPPLLNodeCount)
-SHADER_PARAMETER_RDG_TEXTURE_UAV(Texture2D, PPLLCounter)
-SHADER_PARAMETER_RDG_TEXTURE_UAV(Texture2D, PPLLNodeIndex)
-SHADER_PARAMETER_RDG_BUFFER_UAV(StructuredBuffer, PPLLNodeData)
-RENDER_TARGET_BINDING_SLOTS()
+	SHADER_PARAMETER(uint32, HairVisibilityPass_MaxPPLLNodeCount)
+	SHADER_PARAMETER_RDG_TEXTURE_UAV(Texture2D, PPLLCounter)
+	SHADER_PARAMETER_RDG_TEXTURE_UAV(Texture2D, PPLLNodeIndex)
+	SHADER_PARAMETER_RDG_BUFFER_UAV(StructuredBuffer, PPLLNodeData)
+	RENDER_TARGET_BINDING_SLOTS()
 END_SHADER_PARAMETER_STRUCT()
 
 template<EHairVisibilityRenderMode RenderMode>
 class FHairVisibilityVS : public FMeshMaterialShader
 {
-	DECLARE_SHADER_TYPE(FHairVisibilityVS, MeshMaterial);
+	DECLARE_SHADER_TYPE(FHairVisibilityVS,MeshMaterial);
 
 protected:
 
@@ -664,9 +737,10 @@ IMPLEMENT_MATERIAL_SHADER_TYPE(template<>, FHairVisibilityVS<HairVisibilityRende
 class FHairVisibilityShaderElementData : public FMeshMaterialShaderElementData
 {
 public:
-	FHairVisibilityShaderElementData(uint32 InHairMacroGroupId, uint32 InHairMaterialId, uint32 InMaxPPLLNodeCount) : HairMacroGroupId(InHairMacroGroupId), HairMaterialId(InHairMaterialId), MaxPPLLNodeCount(InMaxPPLLNodeCount) { }
+	FHairVisibilityShaderElementData(uint32 InHairMacroGroupId, uint32 InHairMaterialId, uint32 InLightChannelMask, uint32 InMaxPPLLNodeCount) : HairMacroGroupId(InHairMacroGroupId), HairMaterialId(InHairMaterialId), LightChannelMask(InLightChannelMask), MaxPPLLNodeCount(InMaxPPLLNodeCount) { }
 	uint32 HairMacroGroupId;
 	uint32 HairMaterialId;
+	uint32 LightChannelMask;
 	uint32 MaxPPLLNodeCount;
 };
 
@@ -678,13 +752,14 @@ class FHairVisibilityPS : public FMeshMaterialShader
 public:
 
 	FHairVisibilityPS(const ShaderMetaType::CompiledShaderInitializerType& Initializer)
-		: FMeshMaterialShader(Initializer)
+		:	FMeshMaterialShader(Initializer)
 	{
 		ERHIFeatureLevel::Type FeatureLevel = GetMaxSupportedFeatureLevel((EShaderPlatform)Initializer.Target.Platform);
 		check(FSceneInterface::GetShadingPath(FeatureLevel) != EShadingPath::Mobile);
 		PassUniformBuffer.Bind(Initializer.ParameterMap, FVisibilityPassGlobalParameters::StaticStructMetadata.GetShaderVariableName());
 		HairVisibilityPass_HairMacroGroupIndex.Bind(Initializer.ParameterMap, TEXT("HairVisibilityPass_HairMacroGroupIndex"));
 		HairVisibilityPass_HairMaterialId.Bind(Initializer.ParameterMap, TEXT("HairVisibilityPass_HairMaterialId"));
+		HairVisibilityPass_LightChannelMask.Bind(Initializer.ParameterMap, TEXT("HairVisibilityPass_LightChannelMask"));
 		HairVisibilityPass_MaxPPLLNodeCount.Bind(Initializer.ParameterMap, TEXT("HairVisibilityPass_MaxPPLLNodeCount"));
 		if (RenderMode == EHairVisibilityRenderMode::HairVisibilityRenderMode_PPLL)
 		{
@@ -701,7 +776,7 @@ public:
 		return IsCompatibleWithHairStrands(Parameters.Material, Parameters.Platform) && Parameters.VertexFactoryType->GetFName() == FName(TEXT("FHairStrandsVertexFactory"));
 	}
 
-	static void ModifyCompilationEnvironment(const FMaterialShaderPermutationParameters& Parameters, FShaderCompilerEnvironment& OutEnvironment)
+	static void ModifyCompilationEnvironment(const FMaterialShaderPermutationParameters& Parameters, FShaderCompilerEnvironment& OutEnvironment)	
 	{
 		FMaterialShader::ModifyCompilationEnvironment(Parameters, OutEnvironment);
 		const uint32 RenderModeValue = uint32(RenderMode);
@@ -713,6 +788,7 @@ public:
 		bool bShaderHasOutdatedParameters = FMeshMaterialShader::Serialize(Ar);
 		Ar << HairVisibilityPass_HairMacroGroupIndex;
 		Ar << HairVisibilityPass_HairMaterialId;
+		Ar << HairVisibilityPass_LightChannelMask;
 		Ar << HairVisibilityPass_MaxPPLLNodeCount;
 		if (RenderMode == EHairVisibilityRenderMode::HairVisibilityRenderMode_PPLL)
 		{
@@ -736,11 +812,13 @@ public:
 		FMeshMaterialShader::GetShaderBindings(Scene, FeatureLevel, PrimitiveSceneProxy, MaterialRenderProxy, Material, DrawRenderState, ShaderElementData, ShaderBindings);
 		ShaderBindings.Add(HairVisibilityPass_HairMacroGroupIndex, ShaderElementData.HairMacroGroupId);
 		ShaderBindings.Add(HairVisibilityPass_HairMaterialId, ShaderElementData.HairMaterialId);
+		ShaderBindings.Add(HairVisibilityPass_LightChannelMask, ShaderElementData.LightChannelMask);
 		ShaderBindings.Add(HairVisibilityPass_MaxPPLLNodeCount, ShaderElementData.MaxPPLLNodeCount);
 	}
 
 	FShaderParameter HairVisibilityPass_HairMacroGroupIndex;
 	FShaderParameter HairVisibilityPass_HairMaterialId;
+	FShaderParameter HairVisibilityPass_LightChannelMask;
 	FShaderParameter HairVisibilityPass_MaxPPLLNodeCount;
 	// This is on FVisibilityPassParameters but needed to avoid the compiler to assert with unbound parameters 
 	FShaderParameter PPLLCounter;
@@ -781,6 +859,7 @@ private:
 		const FMaterial& RESTRICT MaterialResource,
 		const uint32 HairMacroGroupId,
 		const uint32 HairMaterialId,
+		const uint32 LightChannelMask,
 		ERasterizerFillMode MeshFillMode,
 		ERasterizerCullMode MeshCullMode);
 
@@ -803,9 +882,10 @@ void FHairVisibilityProcessor::AddMeshBatch(const FMeshBatch& RESTRICT MeshBatch
 	const FMaterial& Material = MeshBatch.MaterialRenderProxy->GetMaterialWithFallback(FeatureLevel, FallbackMaterialRenderProxyPtr);
 	const bool bIsCompatible = IsCompatibleWithHairStrands(&Material, FeatureLevel);
 	const bool bIsHairStrandsFactory = MeshBatch.VertexFactory->GetType()->GetId() == CompatibleVF->GetId();
-	const bool bShouldRender = (!PrimitiveSceneProxy && MeshBatch.Elements.Num() > 0) || (PrimitiveSceneProxy && PrimitiveSceneProxy->ShouldRenderInMainPass());
+	const bool bShouldRender = (!PrimitiveSceneProxy && MeshBatch.Elements.Num()>0) || (PrimitiveSceneProxy && PrimitiveSceneProxy->ShouldRenderInMainPass());
+	const uint32 LightChannelMask = PrimitiveSceneProxy && PrimitiveSceneProxy->GetLightingChannelMask();
 
-	if (bIsCompatible
+	if (bIsCompatible 
 		&& bIsHairStrandsFactory
 		&& bShouldRender
 		&& ShouldIncludeDomainInMeshPass(Material.GetMaterialDomain()))
@@ -814,28 +894,29 @@ void FHairVisibilityProcessor::AddMeshBatch(const FMeshBatch& RESTRICT MeshBatch
 		const ERasterizerFillMode MeshFillMode = ComputeMeshFillMode(MeshBatch, Material);
 		const ERasterizerCullMode MeshCullMode = ComputeMeshCullMode(MeshBatch, Material);
 		if (RenderMode == HairVisibilityRenderMode_MSAA_Visibility)
-			Process<HairVisibilityRenderMode_MSAA_Visibility>(MeshBatch, BatchElementMask, PrimitiveSceneProxy, StaticMeshId, MaterialRenderProxy, Material, HairMacroGroupId, HairMaterialId, MeshFillMode, MeshCullMode);
+			Process<HairVisibilityRenderMode_MSAA_Visibility>(MeshBatch, BatchElementMask, PrimitiveSceneProxy, StaticMeshId, MaterialRenderProxy, Material, HairMacroGroupId, HairMaterialId, LightChannelMask, MeshFillMode, MeshCullMode);
 		if (RenderMode == HairVisibilityRenderMode_MSAA)
-			Process<HairVisibilityRenderMode_MSAA>(MeshBatch, BatchElementMask, PrimitiveSceneProxy, StaticMeshId, MaterialRenderProxy, Material, HairMacroGroupId, HairMaterialId, MeshFillMode, MeshCullMode);
+			Process<HairVisibilityRenderMode_MSAA>(MeshBatch, BatchElementMask, PrimitiveSceneProxy, StaticMeshId, MaterialRenderProxy, Material, HairMacroGroupId, HairMaterialId, LightChannelMask, MeshFillMode, MeshCullMode);
 		if (RenderMode == HairVisibilityRenderMode_Transmittance)
-			Process<HairVisibilityRenderMode_Transmittance>(MeshBatch, BatchElementMask, PrimitiveSceneProxy, StaticMeshId, MaterialRenderProxy, Material, HairMacroGroupId, HairMaterialId, MeshFillMode, MeshCullMode);
+			Process<HairVisibilityRenderMode_Transmittance>(MeshBatch, BatchElementMask, PrimitiveSceneProxy, StaticMeshId, MaterialRenderProxy, Material, HairMacroGroupId, HairMaterialId, LightChannelMask, MeshFillMode, MeshCullMode);
 		if (RenderMode == HairVisibilityRenderMode_TransmittanceAndHairCount)
-			Process<HairVisibilityRenderMode_TransmittanceAndHairCount>(MeshBatch, BatchElementMask, PrimitiveSceneProxy, StaticMeshId, MaterialRenderProxy, Material, HairMacroGroupId, HairMaterialId, MeshFillMode, MeshCullMode);
+			Process<HairVisibilityRenderMode_TransmittanceAndHairCount>(MeshBatch, BatchElementMask, PrimitiveSceneProxy, StaticMeshId, MaterialRenderProxy, Material, HairMacroGroupId, HairMaterialId, LightChannelMask, MeshFillMode, MeshCullMode);
 		if (RenderMode == HairVisibilityRenderMode_PPLL)
-			Process<HairVisibilityRenderMode_PPLL>(MeshBatch, BatchElementMask, PrimitiveSceneProxy, StaticMeshId, MaterialRenderProxy, Material, HairMacroGroupId, HairMaterialId, MeshFillMode, MeshCullMode);
+			Process<HairVisibilityRenderMode_PPLL>(MeshBatch, BatchElementMask, PrimitiveSceneProxy, StaticMeshId, MaterialRenderProxy, Material, HairMacroGroupId, HairMaterialId, LightChannelMask, MeshFillMode, MeshCullMode);
 	}
 }
 
 template<EHairVisibilityRenderMode TRenderMode>
 void FHairVisibilityProcessor::Process(
-	const FMeshBatch& MeshBatch,
-	uint64 BatchElementMask,
+	const FMeshBatch& MeshBatch, 
+	uint64 BatchElementMask, 
 	const FPrimitiveSceneProxy* RESTRICT PrimitiveSceneProxy,
 	int32 StaticMeshId,
 	const FMaterialRenderProxy& RESTRICT MaterialRenderProxy,
 	const FMaterial& RESTRICT MaterialResource,
 	const uint32 HairMacroGroupId,
 	const uint32 HairMaterialId,
+	const uint32 LightChannelMask,
 	ERasterizerFillMode MeshFillMode,
 	ERasterizerCullMode MeshCullMode)
 {
@@ -849,11 +930,11 @@ void FHairVisibilityProcessor::Process(
 	{
 		FVertexFactoryType* VertexFactoryType = VertexFactory->GetType();
 		PassShaders.VertexShader = MaterialResource.GetShader<FHairVisibilityVS<TRenderMode>>(VertexFactoryType);
-		PassShaders.PixelShader = MaterialResource.GetShader<FHairVisibilityPS<TRenderMode>>(VertexFactoryType);
+		PassShaders.PixelShader  = MaterialResource.GetShader<FHairVisibilityPS<TRenderMode>>(VertexFactoryType);
 	}
 
 	FMeshPassProcessorRenderState DrawRenderState(PassDrawRenderState);
-	FHairVisibilityShaderElementData ShaderElementData(HairMacroGroupId, HairMaterialId, MaxPPLLNodeCount);
+	FHairVisibilityShaderElementData ShaderElementData(HairMacroGroupId, HairMaterialId, LightChannelMask, MaxPPLLNodeCount);
 	ShaderElementData.InitializeMeshMaterialData(ViewIfDynamicMeshCommand, PrimitiveSceneProxy, MeshBatch, StaticMeshId, false);
 
 	BuildMeshDrawCommands(
@@ -893,7 +974,7 @@ class FClearUIntGraphicPS : public FGlobalShader
 	BEGIN_SHADER_PARAMETER_STRUCT(FParameters, )
 		SHADER_PARAMETER(uint32, ClearValue)
 		RENDER_TARGET_BINDING_SLOTS()
-		END_SHADER_PARAMETER_STRUCT()
+	END_SHADER_PARAMETER_STRUCT()
 
 public:
 	static bool ShouldCompilePermutation(const FGlobalShaderPermutationParameters& Parameters) { return IsHairStrandsSupported(Parameters.Platform); }
@@ -946,9 +1027,9 @@ static void AddClearGraphicPass(
 		DrawRectangle(
 			RHICmdList,
 			0, 0,
-			Viewport.Width(), Viewport.Height(),
-			Viewport.Min.X, Viewport.Min.Y,
-			Viewport.Width(), Viewport.Height(),
+			Viewport.Width(),Viewport.Height(),
+			Viewport.Min.X,  Viewport.Min.Y,
+			Viewport.Width(),Viewport.Height(),
 			Viewport.Size(),
 			Resolution,
 			*VertexShader,
@@ -957,15 +1038,15 @@ static void AddClearGraphicPass(
 }
 
 BEGIN_SHADER_PARAMETER_STRUCT(FClearUAVTextureParameters, )
-SHADER_PARAMETER_RDG_TEXTURE_UAV(RWTexture2D, TextureUAV)
+	SHADER_PARAMETER_RDG_TEXTURE_UAV(RWTexture2D, TextureUAV)
 END_SHADER_PARAMETER_STRUCT()
 
-void AddClearUAVUintPass(
+void AddClearUAVPass(
 	FRDGBuilder& GraphBuilder,
 	FRDGEventName&& PassName,
 	FRDGTextureRef Texture,
 	uint32 Value)
-{
+{	
 	FClearUAVTextureParameters* Parameters = GraphBuilder.AllocParameters< FClearUAVTextureParameters >();
 	Parameters->TextureUAV = GraphBuilder.CreateUAV(Texture);
 
@@ -993,7 +1074,7 @@ class FCopyIndirectBufferCS : public FGlobalShader
 		SHADER_PARAMETER(uint32, ItemCountPerGroup)
 		SHADER_PARAMETER_RDG_TEXTURE(Texture2D, CounterTexture)
 		SHADER_PARAMETER_RDG_BUFFER_UAV(RWBuffer, OutArgBuffer)
-		END_SHADER_PARAMETER_STRUCT()
+	END_SHADER_PARAMETER_STRUCT()
 
 public:
 	static bool ShouldCompilePermutation(const FGlobalShaderPermutationParameters& Parameters) { return IsHairStrandsSupported(Parameters.Platform); }
@@ -1025,7 +1106,7 @@ static FRDGBufferRef AddCopyIndirectArgPass(
 		RDG_EVENT_NAME("HairStrandsVisbilityCopyIndirectArgs"),
 		*ComputeShader,
 		Parameters,
-		FIntVector(1, 1, 1));
+		FIntVector(1,1,1));
 
 	return OutBuffer;
 }
@@ -1036,12 +1117,12 @@ class FHairVisibilityPrimitiveIdCompactionCS : public FGlobalShader
 	DECLARE_GLOBAL_SHADER(FHairVisibilityPrimitiveIdCompactionCS);
 	SHADER_USE_PARAMETER_STRUCT(FHairVisibilityPrimitiveIdCompactionCS, FGlobalShader);
 
-	class FVendor : SHADER_PERMUTATION_INT("PERMUTATION_VENDOR", HairVisibilityVendorCount);
-	class FVelocity : SHADER_PERMUTATION_INT("PERMUTATION_VELOCITY", 4);
+	class FVendor		: SHADER_PERMUTATION_INT("PERMUTATION_VENDOR", HairVisibilityVendorCount);
+	class FVelocity		: SHADER_PERMUTATION_INT("PERMUTATION_VELOCITY", 4);
 	class FViewTransmittance : SHADER_PERMUTATION_INT("PERMUTATION_VIEWTRANSMITTANCE", 2);
-	class FMaterial : SHADER_PERMUTATION_INT("PERMUTATION_MATERIAL_COMPACTION", 2);
-	class FPPLL : SHADER_PERMUTATION_SPARSE_INT("PERMUTATION_PPLL", 0, 8, 16, 32); // See GetPPLLMaxRenderNodePerPixel
-	class FVisibility : SHADER_PERMUTATION_INT("PERMUTATION_VISIBILITY", 2);
+	class FMaterial		: SHADER_PERMUTATION_INT("PERMUTATION_MATERIAL_COMPACTION", 2);
+	class FPPLL			: SHADER_PERMUTATION_SPARSE_INT("PERMUTATION_PPLL", 0, 8, 16, 32); // See GetPPLLMaxRenderNodePerPixel
+	class FVisibility	: SHADER_PERMUTATION_INT("PERMUTATION_VISIBILITY", 2);
 	using FPermutationDomain = TShaderPermutationDomain<FVendor, FVelocity, FViewTransmittance, FMaterial, FPPLL, FVisibility>;
 
 	BEGIN_SHADER_PARAMETER_STRUCT(FParameters, )
@@ -1064,7 +1145,7 @@ class FHairVisibilityPrimitiveIdCompactionCS : public FGlobalShader
 		SHADER_PARAMETER_RDG_BUFFER_SRV(StructuredBuffer, PPLLNodeData)
 
 		SHADER_PARAMETER_RDG_TEXTURE(Texture2D, ViewTransmittanceTexture)
-
+		
 		SHADER_PARAMETER_RDG_TEXTURE_UAV(Texture2D, OutCompactNodeCounter)
 		SHADER_PARAMETER_RDG_TEXTURE_UAV(Texture2D, OutCompactNodeIndex)
 		SHADER_PARAMETER_RDG_TEXTURE_UAV(Texture2D, OutCategorizationTexture)
@@ -1074,7 +1155,7 @@ class FHairVisibilityPrimitiveIdCompactionCS : public FGlobalShader
 
 		SHADER_PARAMETER_STRUCT_REF(FViewUniformShaderParameters, ViewUniformBuffer)
 		SHADER_PARAMETER_STRUCT_REF(FSceneTexturesUniformParameters, SceneTexturesStruct)
-		END_SHADER_PARAMETER_STRUCT()
+	END_SHADER_PARAMETER_STRUCT()
 
 public:
 
@@ -1083,7 +1164,7 @@ public:
 		if (PermutationVector.Get<FPPLL>() > 0)
 		{
 			PermutationVector.Set<FViewTransmittance>(0);
-			PermutationVector.Set<FVisibility>(0);
+			PermutationVector.Set<FVisibility>(0); 
 		}
 		return PermutationVector;
 	}
@@ -1146,7 +1227,7 @@ static void AddHairVisibilityPrimitiveIdCompactionPass(
 		}
 	}
 
-
+	
 	FRDGTextureRef CompactCounter;
 	{
 		FRDGTextureDesc Desc;
@@ -1183,10 +1264,10 @@ static void AddHairVisibilityPrimitiveIdCompactionPass(
 		OutputDesc.TargetableFlags = TexCreate_UAV | TexCreate_ShaderResource;
 		OutCategorizationTexture = GraphBuilder.CreateTexture(OutputDesc, TEXT("CategorizationTexture"));
 	}
-
-	AddClearUAVUintPass(GraphBuilder, RDG_EVENT_NAME("HairStrandsClearCompactionCounter"), CompactCounter, 0);
-	AddClearUAVUintPass(GraphBuilder, RDG_EVENT_NAME("HairStrandsClearCompactionOffsetAndCount"), OutCompactNodeIndex, 0);
-	AddClearUAVUintPass(GraphBuilder, RDG_EVENT_NAME("HairStrandsClearCategorizationTexture"), OutCategorizationTexture, 0);
+	
+	AddClearUAVPass(GraphBuilder, RDG_EVENT_NAME("HairStrandsClearCompactionCounter"), CompactCounter, 0);
+	AddClearUAVPass(GraphBuilder, RDG_EVENT_NAME("HairStrandsClearCompactionOffsetAndCount"), OutCompactNodeIndex, 0);
+	AddClearUAVPass(GraphBuilder, RDG_EVENT_NAME("HairStrandsClearCategorizationTexture"), OutCategorizationTexture, 0);
 
 	// Select render node count according to current mode
 	const uint32 HairVisibilityMSAASampleCount = GetHairVisibilitySampleCount();
@@ -1232,7 +1313,7 @@ static void AddHairVisibilityPrimitiveIdCompactionPass(
 	{
 		PassParameters->OutVelocityTexture = GraphBuilder.CreateUAV(OutVelocityTexture);
 	}
-
+ 
 	FIntRect TotalRect = ComputeVisibleHairStrandsMacroGroupsRect(View.ViewRect, MacroGroupDatas);
 
 	// Snap the rect onto thread group boundary
@@ -1241,7 +1322,7 @@ static void AddHairVisibilityPrimitiveIdCompactionPass(
 	TotalRect.Min.Y = FMath::FloorToInt(float(TotalRect.Min.Y) / float(GroupSize.Y)) * GroupSize.Y;
 	TotalRect.Max.X = FMath::CeilToInt(float(TotalRect.Max.X) / float(GroupSize.X)) * GroupSize.X;
 	TotalRect.Max.Y = FMath::CeilToInt(float(TotalRect.Max.Y) / float(GroupSize.Y)) * GroupSize.Y;
-
+	
 	FIntPoint RectResolution(TotalRect.Width(), TotalRect.Height());
 	PassParameters->ResolutionOffset = FIntPoint(TotalRect.Min.X, TotalRect.Min.Y);
 
@@ -1273,10 +1354,10 @@ class FHairGenerateTileCS : public FGlobalShader
 		SHADER_PARAMETER_RDG_TEXTURE_UAV(Texture2D, OutTileCounter)
 		SHADER_PARAMETER_RDG_TEXTURE_UAV(RWTexture2D, OutTileIndexTexture)
 		SHADER_PARAMETER_RDG_BUFFER_UAV(RWBuffer, OutTileBuffer)
-		END_SHADER_PARAMETER_STRUCT()
+	END_SHADER_PARAMETER_STRUCT()
 
 public:
-
+	
 	static bool ShouldCompilePermutation(const FGlobalShaderPermutationParameters& Parameters)
 	{
 		return IsHairStrandsSupported(Parameters.Platform);
@@ -1327,11 +1408,10 @@ static void AddGenerateTilePass(
 		Desc.ClearValue = FClearValueBinding(0);
 		OutTileIndexTexture = GraphBuilder.CreateTexture(Desc, TEXT("HairTileIndexTexture"));
 	}
-
+	
 	OutTileBuffer = GraphBuilder.CreateBuffer(FRDGBufferDesc::CreateBufferDesc(sizeof(uint32), TileCount), TEXT("HairTileBuffer"));
 
-	uint32 ClearValues[4] = { 0,0,0,0 };
-	AddClearUAVPass(GraphBuilder, GraphBuilder.CreateUAV(TileCounter), ClearValues);
+	AddClearUAVPass(GraphBuilder, RDG_EVENT_NAME("HairTileCounter"), TileCounter, 0);
 
 	FHairGenerateTileCS::FPermutationDomain PermutationVector;
 	PermutationVector.Set<FHairGenerateTileCS::FTileSize>(0);
@@ -1339,7 +1419,7 @@ static void AddGenerateTilePass(
 	FHairGenerateTileCS::FParameters* PassParameters = GraphBuilder.AllocParameters<FHairGenerateTileCS::FParameters>();
 	PassParameters->Resolution = Resolution;
 	PassParameters->TileResolution = TileResolution;
-	PassParameters->CategorizationTexture = CategorizationTexture;
+	PassParameters->CategorizationTexture = CategorizationTexture;	
 	PassParameters->OutTileCounter = GraphBuilder.CreateUAV(TileCounter);
 	PassParameters->OutTileIndexTexture = GraphBuilder.CreateUAV(OutTileIndexTexture);
 	PassParameters->OutTileBuffer = GraphBuilder.CreateUAV(OutTileBuffer, PF_R16G16_UINT);
@@ -1368,7 +1448,7 @@ class FHairVisibilityFillOpaqueDepthPS : public FGlobalShader
 
 		SHADER_PARAMETER_STRUCT_REF(FViewUniformShaderParameters, ViewUniformBuffer)
 		RENDER_TARGET_BINDING_SLOTS()
-		END_SHADER_PARAMETER_STRUCT()
+	END_SHADER_PARAMETER_STRUCT()
 
 public:
 	static bool ShouldCompilePermutation(const FGlobalShaderPermutationParameters& Parameters) { return IsHairStrandsSupported(Parameters.Platform); }
@@ -1460,7 +1540,7 @@ static FRDGTextureRef AddHairVisibilityFillOpaqueDepth(
 					RHICmdList,
 					0, 0,
 					Viewport.Width(), Viewport.Height(),
-					Viewport.Min.X, Viewport.Min.Y,
+					Viewport.Min.X,   Viewport.Min.Y,
 					Viewport.Width(), Viewport.Height(),
 					Viewport.Size(),
 					Resolution,
@@ -1500,7 +1580,7 @@ static void AddHairVisibilityCommonPass(
 	};
 
 	GraphBuilder.AddPass(
-		GetPassName(),
+		GetPassName(), 
 		PassParameters,
 		ERDGPassFlags::Raster,
 		[PassParameters, Scene = Scene, ViewInfo, &MacroGroupDatas, RenderMode](FRHICommandListImmediate& RHICmdList)
@@ -1510,7 +1590,7 @@ static void AddHairVisibilityCommonPass(
 
 		FMeshPassProcessorRenderState DrawRenderState(*ViewInfo);
 
-		if (RenderMode == HairVisibilityRenderMode_Transmittance || RenderMode == HairVisibilityRenderMode_TransmittanceAndHairCount || RenderMode == HairVisibilityRenderMode_PPLL)
+		if(RenderMode == HairVisibilityRenderMode_Transmittance || RenderMode == HairVisibilityRenderMode_TransmittanceAndHairCount || RenderMode == HairVisibilityRenderMode_PPLL)
 		{
 			TUniformBufferRef<FViewUniformShaderParameters> ViewUniformShaderParameters;
 			const bool bEnableMSAA = false;
@@ -1519,7 +1599,7 @@ static void AddHairVisibilityCommonPass(
 			ViewUniformShaderParameters = TUniformBufferRef<FViewUniformShaderParameters>::CreateUniformBufferImmediate(*ViewInfo->CachedViewUniformShaderParameters, UniformBuffer_SingleFrame);
 			DrawRenderState.SetViewUniformBuffer(ViewUniformShaderParameters);
 		}
-
+		
 		{
 			RHICmdList.SetViewport(0, 0, 0.0f, ViewInfo->ViewRect.Width(), ViewInfo->ViewRect.Height(), 1.0f);
 			if (RenderMode == HairVisibilityRenderMode_MSAA)
@@ -1529,7 +1609,7 @@ static void AddHairVisibilityCommonPass(
 					CW_RGBA, BO_Add, BF_One, BF_Zero, BO_Add, BF_One, BF_Zero>::GetRHI());
 				DrawRenderState.SetDepthStencilState(TStaticDepthStencilState<true, CF_DepthNearOrEqual>::GetRHI());
 			}
-			else if (RenderMode == HairVisibilityRenderMode_MSAA_Visibility)
+			else if(RenderMode == HairVisibilityRenderMode_MSAA_Visibility)
 			{
 				DrawRenderState.SetBlendState(TStaticBlendState<CW_RGBA, BO_Add, BF_One, BF_Zero, BO_Add, BF_One, BF_Zero>::GetRHI());
 				DrawRenderState.SetDepthStencilState(TStaticDepthStencilState<true, CF_DepthNearOrEqual>::GetRHI());
@@ -1705,9 +1785,9 @@ static void AddHairVisibilityMSAAPass(
 
 		FVisibilityPassParameters* PassParameters = GraphBuilder.AllocParameters<FVisibilityPassParameters>();
 		PassParameters->RenderTargets[0] = FRenderTargetBinding(OutVisibilityIdTexture, ERenderTargetLoadAction::ELoad, 0);
-		PassParameters->RenderTargets[1] = FRenderTargetBinding(OutVisibilityMaterialTexture, LoadAction, 0);
+		PassParameters->RenderTargets[1] = FRenderTargetBinding(OutVisibilityMaterialTexture,  LoadAction, 0);
 		PassParameters->RenderTargets[2] = FRenderTargetBinding(OutVisibilityAttributeTexture, LoadAction, 0);
-		PassParameters->RenderTargets[3] = FRenderTargetBinding(OutVisibilityVelocityTexture, LoadAction, 0);
+		PassParameters->RenderTargets[3] = FRenderTargetBinding(OutVisibilityVelocityTexture,  LoadAction, 0);
 
 		PassParameters->RenderTargets.DepthStencil = FDepthStencilBinding(
 			OutVisibilityDepthTexture,
@@ -1773,8 +1853,8 @@ static void AddHairVisibilityPPLLPass(
 
 		OutVisibilityPPLLNodeData = GraphBuilder.CreateBuffer(FRDGBufferDesc::CreateStructuredDesc(sizeof(PPLLNodeData), PPLLMaxTotalListElementCount), TEXT("HairVisibilityPPLLNodeData"));
 	}
-	AddClearUAVUintPass(GraphBuilder, RDG_EVENT_NAME("ClearHairVisibilityPPLLCounter"), OutVisibilityPPLLNodeCounter, 0);
-	AddClearUAVUintPass(GraphBuilder, RDG_EVENT_NAME("ClearHairVisibilityPPLLNodeIndex"), OutVisibilityPPLLNodeIndex, 0xFFFFFFFF);
+	AddClearUAVPass(GraphBuilder, RDG_EVENT_NAME("ClearHairVisibilityPPLLCounter"), OutVisibilityPPLLNodeCounter, 0);
+	AddClearUAVPass(GraphBuilder, RDG_EVENT_NAME("ClearHairVisibilityPPLLNodeIndex"), OutVisibilityPPLLNodeIndex, 0xFFFFFFFF);
 
 	FVisibilityPassParameters* PassParameters = GraphBuilder.AllocParameters<FVisibilityPassParameters>();
 	PassParameters->PPLLCounter = GraphBuilder.CreateUAV(FRDGTextureUAVDesc(OutVisibilityPPLLNodeCounter, 0));
@@ -1850,10 +1930,10 @@ class FHairViewTransmittanceDepthPS : public FGlobalShader
 	BEGIN_SHADER_PARAMETER_STRUCT(FParameters, )
 		SHADER_PARAMETER(float, DistanceThreshold)
 		SHADER_PARAMETER_RDG_TEXTURE(Texture2D, SceneDepthTexture)
-		SHADER_PARAMETER_RDG_TEXTURE(Texture2D, CategorizationTexture)
+		SHADER_PARAMETER_RDG_TEXTURE(Texture2D, CategorizationTexture)		
 		SHADER_PARAMETER_STRUCT_REF(FViewUniformShaderParameters, ViewUniformBuffer)
 		RENDER_TARGET_BINDING_SLOTS()
-		END_SHADER_PARAMETER_STRUCT()
+	END_SHADER_PARAMETER_STRUCT()
 
 public:
 
@@ -1878,7 +1958,7 @@ static void AddHairViewTransmittanceDepthPass(
 	Parameters->SceneDepthTexture = SceneDepthTexture;
 	Parameters->ViewUniformBuffer = View.ViewUniformBuffer;
 	Parameters->RenderTargets[0] = FRenderTargetBinding(HairCountTexture, ERenderTargetLoadAction::ELoad);
-
+	
 	TShaderMapRef<FPostProcessVS> VertexShader(View.ShaderMap);
 	TShaderMapRef<FHairViewTransmittanceDepthPS> PixelShader(View.ShaderMap);
 	const TShaderMap<FGlobalShaderType>* GlobalShaderMap = View.ShaderMap;
@@ -1933,10 +2013,10 @@ class FHairVisibilityDepthPS : public FGlobalShader
 	BEGIN_SHADER_PARAMETER_STRUCT(FParameters, )
 		SHADER_PARAMETER_RDG_TEXTURE(Texture2D, CategorisationTexture)
 		RENDER_TARGET_BINDING_SLOTS()
-		END_SHADER_PARAMETER_STRUCT()
+	END_SHADER_PARAMETER_STRUCT()
 
 public:
-
+		
 	static FPermutationDomain RemapPermutation(FPermutationDomain PermutationVector)
 	{
 		return PermutationVector;
@@ -2018,9 +2098,9 @@ static void AddHairVisibilityColorAndDepthPatchPass(
 ///////////////////////////////////////////////////////////////////////////////////////////////////
 
 FHairStrandsVisibilityViews RenderHairStrandsVisibilityBuffer(
-	FRHICommandListImmediate& RHICmdList,
-	const FScene* Scene,
-	const TArray<FViewInfo>& Views,
+	FRHICommandListImmediate& RHICmdList, 
+	const FScene* Scene, 
+	const TArray<FViewInfo>& Views, 
 	TRefCountPtr<IPooledRenderTarget> InSceneGBufferBTexture,
 	TRefCountPtr<IPooledRenderTarget> InSceneColorTexture,
 	TRefCountPtr<IPooledRenderTarget> InSceneDepthTexture,
@@ -2057,7 +2137,7 @@ FHairStrandsVisibilityViews RenderHairStrandsVisibilityBuffer(
 			FRDGTextureRef SceneColorTexture = GraphBuilder.RegisterExternalTexture(InSceneColorTexture, TEXT("SceneColorTexture"));
 			FRDGTextureRef SceneDepthTexture = GraphBuilder.RegisterExternalTexture(InSceneDepthTexture, TEXT("SceneDepthTexture"));
 			FRDGTextureRef SceneVelocityTexture = InSceneVelocityTexture ? GraphBuilder.RegisterExternalTexture(InSceneVelocityTexture, TEXT("SceneVelocityTexture")) : nullptr;
-
+			
 			const EHairVisibilityRenderMode RenderMode = GetHairVisibilityRenderMode();
 
 			// Run the view transmittance pass if needed (not in PPLL mode that is already a high quality render path)
@@ -2070,12 +2150,14 @@ FHairStrandsVisibilityViews RenderHairStrandsVisibilityBuffer(
 					Scene,
 					&View,
 					MacroGroupDatas,
-					Resolution,
+					Resolution, 
 					bOutputHairCount,
 					SceneDepthTexture);
 			}
 
 			FRDGTextureRef CategorizationTexture = nullptr;
+			FRDGTextureRef CompactNodeIndex = nullptr;
+			FRDGBufferRef  CompactNodeData = nullptr;
 			if (RenderMode == HairVisibilityRenderMode_MSAA)
 			{
 				const bool bIsVisiblityEnable = GHairStrandsVisibilityMaterialPass > 0;
@@ -2098,16 +2180,16 @@ FHairStrandsVisibilityViews RenderHairStrandsVisibilityBuffer(
 
 				AddHairVisibilityMSAAPass(
 					bIsVisiblityEnable,
-					GraphBuilder,
-					Scene,
-					&View,
+					GraphBuilder, 
+					Scene, 
+					&View, 
 					MacroGroupDatas,
 					Resolution,
-					MsaaVisibilityResources.IdTexture,
-					MsaaVisibilityResources.MaterialTexture,
+					MsaaVisibilityResources.IdTexture, 
+					MsaaVisibilityResources.MaterialTexture, 
 					MsaaVisibilityResources.AttributeTexture,
 					MsaaVisibilityResources.VelocityTexture,
-					MsaaVisibilityResources.DepthTexture);
+					MsaaVisibilityResources.DepthTexture); 
 
 				// This is used when compaction is not enabled.
 				GraphBuilder.QueueTextureExtraction(MsaaVisibilityResources.IdTexture, &VisibilityData.IDTexture);
@@ -2128,8 +2210,6 @@ FHairStrandsVisibilityViews RenderHairStrandsVisibilityBuffer(
 					PassParameters->MSAA_VelocityTexture = MsaaVisibilityResources.VelocityTexture;
 					PassParameters->ViewTransmittanceTexture = ViewTransmittance.TransmittanceTexture;
 
-					FRDGTextureRef CompactNodeIndex;
-					FRDGBufferRef CompactNodeData;
 					FRDGBufferRef CompactNodeCoord;
 					FRDGBufferRef IndirectArgsBuffer;
 					AddHairVisibilityPrimitiveIdCompactionPass(
@@ -2175,11 +2255,11 @@ FHairStrandsVisibilityViews RenderHairStrandsVisibilityBuffer(
 						CompactNodeData = PassOutput.NodeData;
 					}
 
-					GraphBuilder.QueueTextureExtraction(CompactNodeIndex, &VisibilityData.NodeIndex);
-					GraphBuilder.QueueTextureExtraction(CategorizationTexture, &VisibilityData.CategorizationTexture);
-					GraphBuilder.QueueBufferExtraction(CompactNodeData, &VisibilityData.NodeData, FRDGResourceState::EAccess::Read, FRDGResourceState::EPipeline::Graphics);
-					GraphBuilder.QueueBufferExtraction(CompactNodeCoord, &VisibilityData.NodeCoord, FRDGResourceState::EAccess::Read, FRDGResourceState::EPipeline::Graphics);
-					GraphBuilder.QueueBufferExtraction(IndirectArgsBuffer, &VisibilityData.NodeIndirectArg, FRDGResourceState::EAccess::Read, FRDGResourceState::EPipeline::Compute);
+					GraphBuilder.QueueTextureExtraction(CompactNodeIndex,		&VisibilityData.NodeIndex);
+					GraphBuilder.QueueTextureExtraction(CategorizationTexture,	&VisibilityData.CategorizationTexture);
+					GraphBuilder.QueueBufferExtraction(CompactNodeData,			&VisibilityData.NodeData,					FRDGResourceState::EAccess::Read, FRDGResourceState::EPipeline::Graphics);
+					GraphBuilder.QueueBufferExtraction(CompactNodeCoord,		&VisibilityData.NodeCoord,					FRDGResourceState::EAccess::Read, FRDGResourceState::EPipeline::Graphics);
+					GraphBuilder.QueueBufferExtraction(IndirectArgsBuffer,		&VisibilityData.NodeIndirectArg,			FRDGResourceState::EAccess::Read, FRDGResourceState::EPipeline::Compute);
 				}
 
 				// View transmittance depth test needs to happen before the scene depth is patched with the hair depth (for fully-covered-by-hair pixels)
@@ -2224,13 +2304,11 @@ FHairStrandsVisibilityViews RenderHairStrandsVisibilityBuffer(
 				// Linked list sorting pass and compaction into common representation
 				{
 					FHairVisibilityPrimitiveIdCompactionCS::FParameters* PassParameters = GraphBuilder.AllocParameters<FHairVisibilityPrimitiveIdCompactionCS::FParameters>();
-					PassParameters->PPLLCounter = PPLLNodeCounterTexture;
-					PassParameters->PPLLNodeIndex = PPLLNodeIndexTexture;
+					PassParameters->PPLLCounter  = PPLLNodeCounterTexture;
+					PassParameters->PPLLNodeIndex= PPLLNodeIndexTexture;
 					PassParameters->PPLLNodeData = GraphBuilder.CreateSRV(PPLLNodeDataBuffer);
 					PassParameters->ViewTransmittanceTexture = ViewTransmittance.TransmittanceTexture;
 
-					FRDGTextureRef CompactNodeIndex;
-					FRDGBufferRef CompactNodeData;
 					FRDGBufferRef CompactNodeCoord;
 					FRDGBufferRef IndirectArgsBuffer;
 					AddHairVisibilityPrimitiveIdCompactionPass(
@@ -2270,6 +2348,19 @@ FHairStrandsVisibilityViews RenderHairStrandsVisibilityBuffer(
 #endif
 			}
 
+		#if RHI_RAYTRACING
+			if (IsRayTracingEnabled())
+			{
+				FRDGTextureRef LightingChannelMaskTexture = AddHairLightChannelMaskPass(
+					GraphBuilder,
+					View,
+					Resolution,
+					CompactNodeData,
+					CompactNodeIndex);
+				GraphBuilder.QueueTextureExtraction(LightingChannelMaskTexture, &VisibilityData.LightChannelMaskTexture);
+			}
+		#endif
+
 			// Generate Tile data
 			{
 				FRDGTextureRef TileIndexTexture = nullptr;
diff --git a/Engine/Source/Runtime/Renderer/Private/HairStrands/HairStrandsVisibility.h b/Engine/Source/Runtime/Renderer/Private/HairStrands/HairStrandsVisibility.h
--- a/Engine/Source/Runtime/Renderer/Private/HairStrands/HairStrandsVisibility.h
+++ b/Engine/Source/Runtime/Renderer/Private/HairStrands/HairStrandsVisibility.h
@@ -23,6 +23,7 @@ struct FHairStrandsVisibilityData
 	TRefCountPtr<IPooledRenderTarget> VelocityTexture;
 	TRefCountPtr<IPooledRenderTarget> CategorizationTexture;
 	TRefCountPtr<IPooledRenderTarget> ViewHairCountTexture;
+	TRefCountPtr<IPooledRenderTarget> LightChannelMaskTexture;
 
 	TRefCountPtr<IPooledRenderTarget> PPLLNodeCounterTexture;
 	TRefCountPtr<IPooledRenderTarget> PPLLNodeIndexTexture;
diff --git a/Engine/Source/Runtime/Renderer/Private/LightRendering.cpp b/Engine/Source/Runtime/Renderer/Private/LightRendering.cpp
--- a/Engine/Source/Runtime/Renderer/Private/LightRendering.cpp
+++ b/Engine/Source/Runtime/Renderer/Private/LightRendering.cpp
@@ -937,12 +937,26 @@ void FDeferredShadingSceneRenderer::RenderLights(FRHICommandListImmediate& RHICm
 		return HairDatas && ViewIndex < HairDatas->MacroGroupsPerViews.Views.Num() && HairDatas->MacroGroupsPerViews.Views[ViewIndex].Datas.Num() > 0;
 	};
 
-	auto GetHairStrandsCategorizationTexture = [HairDatas](int32 ViewIndex, FRDGBuilder& GraphBuilder)
+	struct FHairStrandsTextures
 	{
-		FRDGTextureRef HairCategorizationTexture = nullptr;
-		if (HairDatas && ViewIndex < HairDatas->HairVisibilityViews.HairDatas.Num() && HairDatas->HairVisibilityViews.HairDatas[ViewIndex].CategorizationTexture)
-			HairCategorizationTexture = GraphBuilder.RegisterExternalTexture(HairDatas->HairVisibilityViews.HairDatas[ViewIndex].CategorizationTexture);
-		return HairCategorizationTexture;
+		FRDGTextureRef CategorizationTexture = nullptr;
+		FRDGTextureRef LightChannelMaskTexture = nullptr;
+	};
+	auto GetHairStrandsTextures = [HairDatas](int32 ViewIndex, FRDGBuilder& GraphBuilder)
+	{
+		FHairStrandsTextures Out;
+		if (HairDatas && ViewIndex < HairDatas->HairVisibilityViews.HairDatas.Num())
+		{
+			if (HairDatas->HairVisibilityViews.HairDatas[ViewIndex].CategorizationTexture)
+			{
+				Out.CategorizationTexture = GraphBuilder.RegisterExternalTexture(HairDatas->HairVisibilityViews.HairDatas[ViewIndex].CategorizationTexture);
+			}
+			if (HairDatas->HairVisibilityViews.HairDatas[ViewIndex].LightChannelMaskTexture)
+			{
+				Out.LightChannelMaskTexture = GraphBuilder.RegisterExternalTexture(HairDatas->HairVisibilityViews.HairDatas[ViewIndex].LightChannelMaskTexture);
+			}
+		}
+		return Out;
 	};
 
 	check(RHICmdList.IsOutsideRenderPass());
@@ -1392,8 +1406,7 @@ void FDeferredShadingSceneRenderer::RenderLights(FRHICommandListImmediate& RHICm
 
 								FRDGTextureUAV* RayTracingShadowMaskUAV = GraphBuilder.CreateUAV(FRDGTextureUAVDesc(RayTracingShadowMaskTexture));
 								FRDGTextureUAV* RayHitDistanceUAV = GraphBuilder.CreateUAV(FRDGTextureUAVDesc(RayDistanceTexture));
-								FRDGTextureRef HairCategorizationTexture = GetHairStrandsCategorizationTexture(ViewIndex, GraphBuilder);
-
+								FHairStrandsTextures HairTextures = GetHairStrandsTextures(ViewIndex, GraphBuilder);
 								{
 									FString BatchLightNameWithLevel;
 									GetLightNameForDrawEvent(BatchLightSceneInfo.Proxy, BatchLightNameWithLevel);
@@ -1407,7 +1420,8 @@ void FDeferredShadingSceneRenderer::RenderLights(FRHICommandListImmediate& RHICm
 										BatchRayTracingConfig,
 										DenoiserRequirements,
 										false,
-										HairCategorizationTexture,
+										HairTextures.CategorizationTexture,
+										HairTextures.LightChannelMaskTexture,
 										RayTracingShadowMaskUAV,
 										RayHitDistanceUAV);
 								}
@@ -1452,7 +1466,8 @@ void FDeferredShadingSceneRenderer::RenderLights(FRHICommandListImmediate& RHICm
 										BatchRayTracingConfig,
 										IScreenSpaceDenoiser::EShadowRequirements::Bailout,
 										true,
-										HairCategorizationTexture,
+										HairTextures.CategorizationTexture,
+										HairTextures.LightChannelMaskTexture,
 										SubPixelRayTracingShadowMaskUAV,
 										SubPixelRayHitDistanceUAV);
 
@@ -1617,8 +1632,7 @@ void FDeferredShadingSceneRenderer::RenderLights(FRHICommandListImmediate& RHICm
 								DenoiserRequirements = DenoiserToUse->GetShadowRequirements(View, LightSceneInfo, RayTracingConfig);
 							}
 
-							FRDGTextureRef HairCategorizationTexture = GetHairStrandsCategorizationTexture(ViewIndex, GraphBuilder);
-
+							FHairStrandsTextures HairTextures = GetHairStrandsTextures(ViewIndex, GraphBuilder);
 							RenderRayTracingShadows(
 								GraphBuilder,
 								SceneTextures,
@@ -1627,7 +1641,8 @@ void FDeferredShadingSceneRenderer::RenderLights(FRHICommandListImmediate& RHICm
 								RayTracingConfig,
 								DenoiserRequirements,
 								false,
-								HairCategorizationTexture,
+								HairTextures.CategorizationTexture,
+								HairTextures.LightChannelMaskTexture,
 								RayTracingShadowMaskUAV,
 								RayHitDistanceUAV);
 
@@ -1685,7 +1700,8 @@ void FDeferredShadingSceneRenderer::RenderLights(FRHICommandListImmediate& RHICm
 									RayTracingConfig,
 									IScreenSpaceDenoiser::EShadowRequirements::Bailout,
 									true,
-									HairCategorizationTexture,
+									HairTextures.CategorizationTexture,
+									HairTextures.LightChannelMaskTexture,
 									SubPixelRayTracingShadowMaskUAV,
 									SubPixelRayHitDistanceUAV);
 
diff --git a/Engine/Source/Runtime/Renderer/Private/RayTracing/RayTracingShadows.cpp b/Engine/Source/Runtime/Renderer/Private/RayTracing/RayTracingShadows.cpp
--- a/Engine/Source/Runtime/Renderer/Private/RayTracing/RayTracingShadows.cpp
+++ b/Engine/Source/Runtime/Renderer/Private/RayTracing/RayTracingShadows.cpp
@@ -96,6 +96,7 @@ class FOcclusionRGS : public FGlobalShader
 		SHADER_PARAMETER_STRUCT_INCLUDE(FSceneTextureParameters, SceneTextures)
 
 		SHADER_PARAMETER_RDG_TEXTURE(Texture2D, HairCategorizationTexture)
+		SHADER_PARAMETER_RDG_TEXTURE(Texture2D, HairLightChannelMaskTexture)
 		SHADER_PARAMETER_SRV(RaytracingAccelerationStructure, TLAS)
 		SHADER_PARAMETER_RDG_TEXTURE_UAV(RWTexture2D<float4>, RWOcclusionMaskUAV)
 		SHADER_PARAMETER_RDG_TEXTURE_UAV(RWTexture2D<float>, RWRayDistanceUAV)
@@ -155,6 +156,7 @@ void FDeferredShadingSceneRenderer::RenderRayTracingShadows(
 	const IScreenSpaceDenoiser::EShadowRequirements DenoiserRequirements,
 	const bool bSubPixelShadowMask,
 	FRDGTextureRef HairCategorizationTexture,
+	FRDGTextureRef HairLightChannelMaskTexture,
 	FRDGTextureUAV* OutShadowMaskUAV,
 	FRDGTextureUAV* OutRayHitDistanceUAV)
 #if RHI_RAYTRACING
@@ -233,7 +235,11 @@ void FDeferredShadingSceneRenderer::RenderRayTracingShadows(
 		{
 			PassParameters->HairCategorizationTexture = GraphBuilder.RegisterExternalTexture(GSystemTextures.BlackDummy);
 		}
-		
+		PassParameters->HairLightChannelMaskTexture = HairLightChannelMaskTexture;
+		if (!PassParameters->HairLightChannelMaskTexture)
+		{
+			PassParameters->HairLightChannelMaskTexture = GraphBuilder.RegisterExternalTexture(GSystemTextures.BlackDummy);
+		}
 		FOcclusionRGS::FPermutationDomain PermutationVector;
 		PermutationVector.Set<FOcclusionRGS::FLightTypeDim>(LightSceneProxy->GetLightType());
 		if (DenoiserRequirements == IScreenSpaceDenoiser::EShadowRequirements::PenumbraAndAvgOccluder)
diff --git a/Engine/Source/Runtime/Renderer/Public/HairStrandsInterface.h b/Engine/Source/Runtime/Renderer/Public/HairStrandsInterface.h
--- a/Engine/Source/Runtime/Renderer/Public/HairStrandsInterface.h
+++ b/Engine/Source/Runtime/Renderer/Public/HairStrandsInterface.h
@@ -148,7 +148,9 @@ struct FHairStrandClusterData
 typedef void(*THairStrandsResetInterpolationFunction)(
 	FRHICommandListImmediate& RHICmdList,
 	struct FHairStrandsInterpolationInput* Input,
-	struct FHairStrandsInterpolationOutput* Output);
+	struct FHairStrandsInterpolationOutput* Output,
+	struct FHairStrandsProjectionHairData& SimHairProjection,
+	int32 LODIndex);
 
 typedef void (*THairStrandsInterpolationFunction)(
 	FRHICommandListImmediate& RHICmdList, 
@@ -174,7 +176,10 @@ struct FHairStrandsProjectionMeshData
 	{
 		FTransform LocalToWorld;
 		FRHIShaderResourceView* PositionBuffer = nullptr;
+		FRHIShaderResourceView* UVsBuffer = nullptr;
 		FRHIShaderResourceView* IndexBuffer = nullptr;
+		uint32 UVsChannelCount = 0;
+		uint32 UVsChannelOffset = 0;
 		uint32 NumPrimitives = 0;
 		uint32 VertexBaseIndex = 0;
 		uint32 IndexBaseIndex = 0;
@@ -202,17 +207,17 @@ struct FHairStrandsProjectionHairData
 
 		// The index buffers stores the mesh section & the triangle index into a single uint32 
 		// (3 highest bits store the section (up to 8 sections)
+		//
+		// See EncodeTriangleIndex & DecodeTriangleIndex functions in HairStrandsMeshProjectionCommon.ush
 		FRWBuffer* RootTriangleIndexBuffer = nullptr;
 		FRWBuffer* RootTriangleBarycentricBuffer = nullptr;
 
 		// Rest root triangles' positions are relative to root center (for preserving precision)
-		FVector* RestPositionOffset = nullptr;
 		FRWBuffer* RestRootTrianglePosition0Buffer = nullptr;
 		FRWBuffer* RestRootTrianglePosition1Buffer = nullptr;
 		FRWBuffer* RestRootTrianglePosition2Buffer = nullptr;
 
 		// Deformed root triangles' positions are relative to root center (for preserving precision)
-		FVector* DeformedPositionOffset = nullptr;
 		FRWBuffer* DeformedRootTrianglePosition0Buffer = nullptr;
 		FRWBuffer* DeformedRootTrianglePosition1Buffer = nullptr;
 		FRWBuffer* DeformedRootTrianglePosition2Buffer = nullptr;
@@ -272,6 +277,25 @@ struct FHairStrandsDebugInfo
 typedef TArray<FHairStrandsDebugInfo> FHairStrandsDebugInfos;
 FHairStrandsDebugInfos GetHairStandsDebugInfos();
 
+
+////////////////////////////////////////////////////////////////////////////////////////////////////
+// Mesh transfer and hair projection debug infos
+
+enum class EHairStrandsProjectionMeshType
+{
+	RestMesh,
+	DeformedMesh,
+	SourceMesh,
+	TargetMesh
+};
+
+struct FHairStrandsProjectionDebugInfo
+{
+	FHairStrandsProjectionMeshData SourceMeshData;
+	FHairStrandsProjectionMeshData TargetMeshData;
+	TArray<FRWBuffer> TransferredPositions;
+};
+
 ////////////////////////////////////////////////////////////////////////////////////////////////////
 // Registrations
 
@@ -282,7 +306,8 @@ RENDERER_API void RegisterHairStrands(
 	const FHairStrandsProjectionHairData& RenProjection,
 	const FHairStrandsProjectionHairData& SimProjection,
 	const FHairStrandsPrimitiveResources& PrimitiveResources,
-	const FHairStrandsDebugInfo& DebugInfo);
+	const FHairStrandsDebugInfo& DebugInfo,
+	const FHairStrandsProjectionDebugInfo& DebugProjectionInfo);
 
 RENDERER_API void UnregisterHairStrands(uint64 Id);
 
@@ -301,8 +326,7 @@ RENDERER_API bool UpdateHairStrands(
 	uint64 Id,
 	EWorldType::Type WorldType,
 	const FTransform& HairLocalToWorld,
-	const FTransform& MeshLocalToWorld,
-	const FVector& SkeletalDeformedPositionOffset);
+	const FTransform& MeshLocalToWorld);
 
 RENDERER_API bool UpdateHairStrands(
 	uint64 Id, 
@@ -332,10 +356,24 @@ RENDERER_API FHairCullInfo GetHairStrandsCullInfo();
 
 RENDERER_API bool IsHairRayTracingEnabled();
 
+RENDERER_API void RunMeshTransfer(
+	FRHICommandListImmediate& RHICmdList,
+	const FHairStrandsProjectionMeshData& SourceMeshData,
+	const FHairStrandsProjectionMeshData& TargetMeshData,
+	TArray<FRWBuffer>& TransferredLODsPositions);
+
 RENDERER_API void RunProjection(
 	FRHICommandListImmediate& RHICmdList,
 	const FTransform& LocalToWorld,
-	const FVector& RestPositionOffset,
-	const FHairStrandsProjectionMeshData& MeshData,
+	const FHairStrandsProjectionMeshData& TargetMeshData,
 	FHairStrandsProjectionHairData& RenProjectionHairData,
-	FHairStrandsProjectionHairData& SimProjectionHairData);
\ No newline at end of file
+	FHairStrandsProjectionHairData& SimProjectionHairData);
+
+typedef void (*TBindingProcess)(FRHICommandListImmediate& RHICmdList, void* Asset);
+RENDERER_API void EnqueueGroomBindingQuery(void* Asset, TBindingProcess BindingProcess);
+void RunHairStrandsBindingQueries(FRHICommandListImmediate& RHICmdList, TShaderMap<FGlobalShaderType>* ShaderMap);
+bool HasHairStrandsProjectionQuery(EShaderPlatform Platform);
+
+RENDERER_API FHairStrandsProjectionMeshData ExtractMeshData(class FSkeletalMeshRenderData* RenderData);
+
+FHairStrandsProjectionHairData::HairGroup ToProjectionHairData(struct FHairStrandsRootResource* In);
\ No newline at end of file
