Move the FLightPrimitiveInteraction lists in FLightSceneInfo to private and provide interfaces to access them. It syncs with the async LPI creation task by default to ensure thread-safe access.

diff --git a/Engine/Source/Runtime/Renderer/Private/LightSceneInfo.cpp b/Engine/Source/Runtime/Renderer/Private/LightSceneInfo.cpp
--- a/Engine/Source/Runtime/Renderer/Private/LightSceneInfo.cpp
+++ b/Engine/Source/Runtime/Renderer/Private/LightSceneInfo.cpp
@@ -36,9 +36,9 @@ void FLightSceneInfoCompact::Init(FLightSceneInfo* InLightSceneInfo)
 }
 
 FLightSceneInfo::FLightSceneInfo(FLightSceneProxy* InProxy, bool InbVisible)
-	: Proxy(InProxy)
-	, DynamicInteractionOftenMovingPrimitiveList(NULL)
+	: DynamicInteractionOftenMovingPrimitiveList(NULL)
 	, DynamicInteractionStaticPrimitiveList(NULL)
+	, Proxy(InProxy)
 	, Id(INDEX_NONE)
 	, TileIntersectionResources(nullptr)
 	, HeightFieldTileIntersectionResources(nullptr)
@@ -88,6 +88,7 @@ void FLightSceneInfo::AddToScene()
 		|| bIsValidLightTypeMobile)
 	{
 		// Add the light to the scene's light octree.
+		Scene->FlushAsyncLightPrimitiveInteractionCreation();
 		Scene->LightOctree.AddElement(LightSceneInfoCompact);
 
 		// TODO: Special case directional lights, no need to traverse the octree.
@@ -124,6 +125,8 @@ void FLightSceneInfo::CreateLightPrimitiveInteraction(const FLightSceneInfoCompa
 
 void FLightSceneInfo::RemoveFromScene()
 {
+	Scene->FlushAsyncLightPrimitiveInteractionCreation();
+
 	if (OctreeId.IsValidId())
 	{
 		// Remove the light from the octree.
@@ -139,8 +142,6 @@ void FLightSceneInfo::RemoveFromScene()
 void FLightSceneInfo::Detach()
 {
 	check(IsInRenderingThread());
-	
-	Scene->FlushAsyncLightPrimitiveInteractionCreation();
 
 	// implicit linked list. The destruction will update this "head" pointer to the next item in the list.
 	while(DynamicInteractionOftenMovingPrimitiveList)
@@ -203,6 +204,24 @@ bool FLightSceneInfo::IsPrecomputedLightingValid() const
 	return (bPrecomputedLightingIsValid && NumUnbuiltInteractions < GWholeSceneShadowUnbuiltInteractionThreshold) || !Proxy->HasStaticShadowing();
 }
 
+FLightPrimitiveInteraction* FLightSceneInfo::GetDynamicInteractionOftenMovingPrimitiveList(bool bSync) const
+{
+	if (bSync)
+	{
+		Scene->FlushAsyncLightPrimitiveInteractionCreation();
+	}
+	return DynamicInteractionOftenMovingPrimitiveList;
+}
+
+FLightPrimitiveInteraction* FLightSceneInfo::GetDynamicInteractionStaticPrimitiveList(bool bSync) const
+{
+	if (bSync)
+	{
+		Scene->FlushAsyncLightPrimitiveInteractionCreation();
+	}
+	return DynamicInteractionStaticPrimitiveList;
+}
+
 void FLightSceneInfo::ReleaseRHI()
 {
 	if (TileIntersectionResources)
diff --git a/Engine/Source/Runtime/Renderer/Private/LightSceneInfo.h b/Engine/Source/Runtime/Renderer/Private/LightSceneInfo.h
--- a/Engine/Source/Runtime/Renderer/Private/LightSceneInfo.h
+++ b/Engine/Source/Runtime/Renderer/Private/LightSceneInfo.h
@@ -151,15 +151,17 @@ typedef TOctree<FLightSceneInfoCompact,struct FLightOctreeSemantics> FSceneLight
  */
 class FLightSceneInfo : public FRenderResource
 {
-public:
-	/** The light's scene proxy. */
-	FLightSceneProxy* Proxy;
+	friend class FLightPrimitiveInteraction;
 
 	/** The list of dynamic primitives affected by the light. */
 	FLightPrimitiveInteraction* DynamicInteractionOftenMovingPrimitiveList;
 
 	FLightPrimitiveInteraction* DynamicInteractionStaticPrimitiveList;
 
+public:
+	/** The light's scene proxy. */
+	FLightSceneProxy* Proxy;
+
 	/** If bVisible == true, this is the index of the primitive in Scene->Lights. */
 	int32 Id;
 
@@ -300,6 +302,10 @@ public:
 		return DynamicShadowMapChannel;
 	}
 
+	FLightPrimitiveInteraction* GetDynamicInteractionOftenMovingPrimitiveList(bool bSync = true) const;
+
+	FLightPrimitiveInteraction* GetDynamicInteractionStaticPrimitiveList(bool bSync = true) const;
+
 	/** Hash function. */
 	friend uint32 GetTypeHash(const FLightSceneInfo* LightSceneInfo)
 	{
diff --git a/Engine/Source/Runtime/Renderer/Private/RendererScene.cpp b/Engine/Source/Runtime/Renderer/Private/RendererScene.cpp
--- a/Engine/Source/Runtime/Renderer/Private/RendererScene.cpp
+++ b/Engine/Source/Runtime/Renderer/Private/RendererScene.cpp
@@ -3083,8 +3083,6 @@ void FScene::GetRelevantLights_RenderThread( UPrimitiveComponent* Primitive, TAr
 	check( RelevantLights );
 	if( Primitive->SceneProxy )
 	{
-		FlushAsyncLightPrimitiveInteractionCreation();
-
 		for( const FLightPrimitiveInteraction* Interaction=Primitive->SceneProxy->GetPrimitiveSceneInfo()->LightList; Interaction; Interaction=Interaction->GetNextLight() )
 		{
 			RelevantLights->Add( Interaction->GetLight()->Proxy->GetLightComponent() );
@@ -3271,11 +3269,6 @@ void FScene::DumpUnbuiltLightInteractions( FOutputDevice& Ar ) const
 {
 	FlushRenderingCommands();
 
-	if (AsyncCreateLightPrimitiveInteractionsTask)
-	{
-		AsyncCreateLightPrimitiveInteractionsTask->EnsureCompletion();
-	}
-
 	TSet<FString> LightsWithUnbuiltInteractions;
 	TSet<FString> PrimitivesWithUnbuiltInteractions;
 
@@ -3287,7 +3280,7 @@ void FScene::DumpUnbuiltLightInteractions( FOutputDevice& Ar ) const
 
 		bool bLightHasUnbuiltInteractions = false;
 
-		for(FLightPrimitiveInteraction* Interaction = LightSceneInfo->DynamicInteractionOftenMovingPrimitiveList;
+		for(FLightPrimitiveInteraction* Interaction = LightSceneInfo->GetDynamicInteractionOftenMovingPrimitiveList();
 			Interaction;
 			Interaction = Interaction->GetNextPrimitive())
 		{
@@ -3298,7 +3291,7 @@ void FScene::DumpUnbuiltLightInteractions( FOutputDevice& Ar ) const
 			}
 		}
 
-		for(FLightPrimitiveInteraction* Interaction = LightSceneInfo->DynamicInteractionStaticPrimitiveList;
+		for(FLightPrimitiveInteraction* Interaction = LightSceneInfo->GetDynamicInteractionStaticPrimitiveList();
 			Interaction;
 			Interaction = Interaction->GetNextPrimitive())
 		{
@@ -3410,6 +3403,7 @@ void FScene::ApplyWorldOffset_RenderThread(const FVector& InOffset)
 	}
 
 	// Lights octree
+	FlushAsyncLightPrimitiveInteractionCreation();
 	LightOctree.ApplyOffset(InOffset, /*bGlobalOctee*/ true);
 
 	// Cached preshadows
diff --git a/Engine/Source/Runtime/Renderer/Private/SceneCore.cpp b/Engine/Source/Runtime/Renderer/Private/SceneCore.cpp
--- a/Engine/Source/Runtime/Renderer/Private/SceneCore.cpp
+++ b/Engine/Source/Runtime/Renderer/Private/SceneCore.cpp
@@ -36,8 +36,6 @@ TAllocatorFixedSizeFreeList<sizeof(FLightPrimitiveInteraction), FREE_LIST_GROW_S
 
 uint32 FRendererModule::GetNumDynamicLightsAffectingPrimitive(const FPrimitiveSceneInfo* PrimitiveSceneInfo,const FLightCacheInterface* LCI)
 {
-	PrimitiveSceneInfo->Scene->FlushAsyncLightPrimitiveInteractionCreation();
-
 	uint32 NumDynamicLights = 0;
 
 	FLightPrimitiveInteraction *LightList = PrimitiveSceneInfo->LightList;
diff --git a/Engine/Source/Runtime/Renderer/Private/SceneVisibility.cpp b/Engine/Source/Runtime/Renderer/Private/SceneVisibility.cpp
--- a/Engine/Source/Runtime/Renderer/Private/SceneVisibility.cpp
+++ b/Engine/Source/Runtime/Renderer/Private/SceneVisibility.cpp
@@ -3023,8 +3023,6 @@ void FSceneRenderer::PreVisibilityFrameSetup(FRHICommandListImmediate& RHICmdLis
 	// Draw lines to lights affecting this mesh if its selected.
 	if (ViewFamily.EngineShowFlags.LightInfluences)
 	{
-		Scene->FlushAsyncLightPrimitiveInteractionCreation();
-
 		for (TArray<FPrimitiveSceneInfo*>::TConstIterator It(Scene->Primitives); It; ++It)
 		{
 			const FPrimitiveSceneInfo* PrimitiveSceneInfo = *It;
diff --git a/Engine/Source/Runtime/Renderer/Private/ShadowSetup.cpp b/Engine/Source/Runtime/Renderer/Private/ShadowSetup.cpp
--- a/Engine/Source/Runtime/Renderer/Private/ShadowSetup.cpp
+++ b/Engine/Source/Runtime/Renderer/Private/ShadowSetup.cpp
@@ -2669,7 +2669,7 @@ void FSceneRenderer::CreateWholeSceneProjectedShadow(
 							&& (CacheMode[CacheModeIndex] != SDCM_MovablePrimitivesOnly || bCastCachedShadowFromMovablePrimitives))
 						{
 							// Add all the shadow casting primitives affected by the light to the shadow's subject primitive list.
-							for (FLightPrimitiveInteraction* Interaction = LightSceneInfo->DynamicInteractionOftenMovingPrimitiveList;
+							for (FLightPrimitiveInteraction* Interaction = LightSceneInfo->GetDynamicInteractionOftenMovingPrimitiveList(false);
 								Interaction;
 								Interaction = Interaction->GetNextPrimitive())
 							{
@@ -2690,7 +2690,7 @@ void FSceneRenderer::CreateWholeSceneProjectedShadow(
 						if (CacheMode[CacheModeIndex] != SDCM_MovablePrimitivesOnly)
 						{
 							// Add all the shadow casting primitives affected by the light to the shadow's subject primitive list.
-							for (FLightPrimitiveInteraction* Interaction = LightSceneInfo->DynamicInteractionStaticPrimitiveList;
+							for (FLightPrimitiveInteraction* Interaction = LightSceneInfo->GetDynamicInteractionStaticPrimitiveList(false);
 								Interaction;
 								Interaction = Interaction->GetNextPrimitive())
 							{
@@ -4214,7 +4214,7 @@ void FSceneRenderer::InitDynamicShadows(FRHICommandListImmediate& RHICmdList, FG
 							Scene->FlushAsyncLightPrimitiveInteractionCreation();
 
 							// Look for individual primitives with a dynamic shadow.
-							for (FLightPrimitiveInteraction* Interaction = LightSceneInfo->DynamicInteractionOftenMovingPrimitiveList;
+							for (FLightPrimitiveInteraction* Interaction = LightSceneInfo->GetDynamicInteractionOftenMovingPrimitiveList(false);
 								Interaction;
 								Interaction = Interaction->GetNextPrimitive()
 								)
@@ -4223,7 +4223,7 @@ void FSceneRenderer::InitDynamicShadows(FRHICommandListImmediate& RHICmdList, FG
 								NumInteractionShadows++;
 							}
 
-							for (FLightPrimitiveInteraction* Interaction = LightSceneInfo->DynamicInteractionStaticPrimitiveList;
+							for (FLightPrimitiveInteraction* Interaction = LightSceneInfo->GetDynamicInteractionStaticPrimitiveList(false);
 								Interaction;
 								Interaction = Interaction->GetNextPrimitive()
 								)
