Added LightCount csv stats

#rnx

diff --git a/Engine/Source/Runtime/Renderer/Private/DeferredShadingRenderer.cpp b/Engine/Source/Runtime/Renderer/Private/DeferredShadingRenderer.cpp
--- a/Engine/Source/Runtime/Renderer/Private/DeferredShadingRenderer.cpp
+++ b/Engine/Source/Runtime/Renderer/Private/DeferredShadingRenderer.cpp
@@ -1697,8 +1697,13 @@ void FDeferredShadingSceneRenderer::Render(FRHICommandListImmediate& RHICmdList)
 		SCOPED_GPU_STAT(RHICmdList, SortLights);
 		GatherAndSortLights(SortedLightSet);
 		ComputeLightGrid(RHICmdList, bComputeLightGrid, SortedLightSet);
+
 	}
 
+	CSV_CUSTOM_STAT_GLOBAL(LightCount, float(SortedLightSet.SortedLights.Num()), ECsvCustomStatOp::Set);
+	CSV_CUSTOM_STAT_GLOBAL(LightCountShadowOff, float(SortedLightSet.AttenuationLightStart), ECsvCustomStatOp::Set);
+	CSV_CUSTOM_STAT_GLOBAL(LightCountShadowOn, float(SortedLightSet.SortedLights.Num()) - float(SortedLightSet.AttenuationLightStart), ECsvCustomStatOp::Set);
+
 	{
 		SCOPE_CYCLE_COUNTER(STAT_FDeferredShadingSceneRenderer_AllocGBufferTargets);
 		SceneContext.PreallocGBufferTargets(); // Even if !bShouldRenderVelocities, the velocity buffer must be bound because it's a compile time option for the shader.
