Fixed raytracing testbed

diff --git a/Engine/Source/Runtime/Renderer/Private/Tests/RayTracingTestbed.cpp b/Engine/Source/Runtime/Renderer/Private/Tests/RayTracingTestbed.cpp
--- a/Engine/Source/Runtime/Renderer/Private/Tests/RayTracingTestbed.cpp
+++ b/Engine/Source/Runtime/Renderer/Private/Tests/RayTracingTestbed.cpp
@@ -101,6 +101,7 @@ bool RunRayTracingTestbed_RenderThread(const FString& Parameters)
 	Segment.VertexBuffer = VertexBuffer;
 	Segment.NumPrimitives = 1;
 	GeometryInitializer.Segments.Add(Segment);
+	GeometryInitializer.TotalPrimitiveCount = Segment.NumPrimitives;
 	FRayTracingGeometryRHIRef Geometry = RHICreateRayTracingGeometry(GeometryInitializer);
 
 	FRayTracingGeometryInstance Instances[] = {
@@ -143,7 +144,7 @@ bool RunRayTracingTestbed_RenderThread(const FString& Parameters)
 
 			RHIUnlockVertexBuffer(OcclusionResultBuffer);
 
-			bOcclusionTestOK = (MappedResults[0] != 0) && (MappedResults[1] == 0) && (MappedResults[2] == 0) && (MappedResults[3] == 0);
+			bOcclusionTestOK = (MappedResults[0] != 0) && (MappedResults[1] == 0) && (MappedResults[2] != 0) && (MappedResults[3] == 0);
 		}
 
 		// Read back and validate intersection trace results
