UE4 - Fix for static analysis warning
#rnx

diff --git a/Engine/Source/Runtime/Renderer/Private/PostProcess/SceneRenderTargets.cpp b/Engine/Source/Runtime/Renderer/Private/PostProcess/SceneRenderTargets.cpp
--- a/Engine/Source/Runtime/Renderer/Private/PostProcess/SceneRenderTargets.cpp
+++ b/Engine/Source/Runtime/Renderer/Private/PostProcess/SceneRenderTargets.cpp
@@ -2116,42 +2116,34 @@ void FSceneRenderTargets::AllocateMobileRenderTargets(FRHICommandListImmediate&
 // This is a helper class. It generates and provides N names with
 // sequentially incremented postfix starting from 0.
 // Example: SomeName0, SomeName1, ..., SomeName117
-class IncrementalNamesHolder
+class FIncrementalNamesHolder
 {
 public:
-	IncrementalNamesHolder(const TCHAR* const Name, uint32 Size)
-		: ArraySize(Size)
+	FIncrementalNamesHolder(const TCHAR* const Name, uint32 Size)
 	{
 		check(Size > 0);
-		Names = new FString[Size];
+		Names.Empty(Size);
 		for (uint32 i = 0; i < Size; ++i)
 		{
-			Names[i] = FString::Printf(TEXT("%s%d"), Name, i);
+			Names.Add(FString::Printf(TEXT("%s%d"), Name, i));
 		}
 	}
 
-	~IncrementalNamesHolder()
-	{
-		delete[] Names;
-	}
-
 	const TCHAR* const operator[] (uint32 Idx) const
 	{
-		check(Idx < ArraySize);
 		return *Names[Idx];
 	}
 
 private:
-	const uint32 ArraySize;
-	FString*     Names = nullptr;
+	TArray<FString> Names;
 };
 
 // for easier use of "VisualizeTexture"
 static const TCHAR* const GetVolumeName(uint32 Id, bool bDirectional)
 {
 	constexpr uint32 MaxNames = 128;
-	static IncrementalNamesHolder Names   (TEXT("TranslucentVolume"),    MaxNames);
-	static IncrementalNamesHolder NamesDir(TEXT("TranslucentVolumeDir"), MaxNames);
+	static FIncrementalNamesHolder Names(TEXT("TranslucentVolume"), MaxNames);
+	static FIncrementalNamesHolder NamesDir(TEXT("TranslucentVolumeDir"), MaxNames);
 
 	check(Id < MaxNames);
 
