Adding flag to force scene to ignore previous frame visibility history.

diff --git a/Engine/Plugins/MovieScene/MovieRenderPipeline/Source/MovieRenderPipelineRenderPasses/Private/MoviePipelineDeferredPasses.cpp b/Engine/Plugins/MovieScene/MovieRenderPipeline/Source/MovieRenderPipelineRenderPasses/Private/MoviePipelineDeferredPasses.cpp
--- a/Engine/Plugins/MovieScene/MovieRenderPipeline/Source/MovieRenderPipelineRenderPasses/Private/MoviePipelineDeferredPasses.cpp
+++ b/Engine/Plugins/MovieScene/MovieRenderPipeline/Source/MovieRenderPipelineRenderPasses/Private/MoviePipelineDeferredPasses.cpp
@@ -386,6 +386,8 @@ namespace MoviePipeline
 			View->LensPrincipalPointOffsetScale = FVector4(TilePrincipalPointOffset.X, -TilePrincipalPointOffset.Y, TilePrincipalPointScale.X, TilePrincipalPointScale.Y);
 		}
 
+		View->bForceCameraVisibilityReset = true;
+
 		View->EndFinalPostprocessSettings(ViewInitOptions);
 
 		return View;
diff --git a/Engine/Source/Runtime/Engine/Private/SceneView.cpp b/Engine/Source/Runtime/Engine/Private/SceneView.cpp
--- a/Engine/Source/Runtime/Engine/Private/SceneView.cpp
+++ b/Engine/Source/Runtime/Engine/Private/SceneView.cpp
@@ -656,6 +656,7 @@ FSceneView::FSceneView(const FSceneViewInitOptions& InitOptions)
 	, bIsMobileMultiViewEnabled(false)
 	, bIsMobileMultiViewDirectEnabled(false)
 	, bShouldBindInstancedViewUB(false)
+	, bForceCameraVisibilityReset(false)
 	, GlobalClippingPlane(FPlane(0, 0, 0, 0))
 	, LensPrincipalPointOffsetScale(0.0f, 0.0f, 1.0f, 1.0f)
 #if WITH_EDITOR
diff --git a/Engine/Source/Runtime/Engine/Public/SceneView.h b/Engine/Source/Runtime/Engine/Public/SceneView.h
--- a/Engine/Source/Runtime/Engine/Public/SceneView.h
+++ b/Engine/Source/Runtime/Engine/Public/SceneView.h
@@ -1017,6 +1017,11 @@ public:
 	/** True if we need to bind the instanced view uniform buffer parameters. */
 	bool bShouldBindInstancedViewUB;
 
+	/** True if we need to force the camera to discard previous frames occlusion. Necessary for overlapped tile rendering
+	  * where we discard previous frame occlusion because the projection matrix changes.
+	  */
+	bool bForceCameraVisibilityReset;
+
 	/** Global clipping plane being applied to the scene, or all 0's if disabled.  This is used when rendering the planar reflection pass. */
 	FPlane GlobalClippingPlane;
 
diff --git a/Engine/Source/Runtime/Renderer/Private/SceneVisibility.cpp b/Engine/Source/Runtime/Renderer/Private/SceneVisibility.cpp
--- a/Engine/Source/Runtime/Renderer/Private/SceneVisibility.cpp
+++ b/Engine/Source/Runtime/Renderer/Private/SceneVisibility.cpp
@@ -3298,7 +3298,7 @@ void FSceneRenderer::PreVisibilityFrameSetup(FRHICommandListImmediate& RHICmdLis
 				ViewState->PrevFrameViewInfo.ViewMatrices.GetViewMatrix(),
 				ViewState->PrevFrameViewInfo.ViewMatrices.GetViewOrigin(),
 				45.0f, 10000.0f);
-			const bool bResetCamera = (bFirstFrameOrTimeWasReset || View.bCameraCut || bIsLargeCameraMovement);
+			const bool bResetCamera = (bFirstFrameOrTimeWasReset || View.bCameraCut || bIsLargeCameraMovement || View.bForceCameraVisibilityReset);
 			
 #if RHI_RAYTRACING
 			// Note: 0.18 deg is the minimum angle for avoiding numerical precision issue (which would cause constant invalidation)
@@ -3360,6 +3360,7 @@ void FSceneRenderer::PreVisibilityFrameSetup(FRHICommandListImmediate& RHICmdLis
 			if (bFirstFrameOrTimeWasReset || 
 				ViewState->LastRenderTime + GEngine->PrimitiveProbablyVisibleTime < View.Family->CurrentRealTime ||
 				View.bCameraCut ||
+				View.bForceCameraVisibilityReset ||
 				IsLargeCameraMovement(
 					View, 
 				    ViewState->PrevViewMatrixForOcclusionQuery, 
