MANUAL MERGE

Merging //UE4/Main @ 11112898 to //UE4/Dev-Core (dev-core-l0178)

(ushell-p4-mergedown)

diff --git a/Engine/Source/Developer/SkeletalMeshUtilitiesCommon/Private/LODUtilities.cpp b/Engine/Source/Developer/SkeletalMeshUtilitiesCommon/Private/LODUtilities.cpp
--- a/Engine/Source/Developer/SkeletalMeshUtilitiesCommon/Private/LODUtilities.cpp
+++ b/Engine/Source/Developer/SkeletalMeshUtilitiesCommon/Private/LODUtilities.cpp
@@ -1955,7 +1955,7 @@ void FLODUtilities::GenerateImportedSkinWeightProfileData(const FSkeletalMeshLOD
 		int32 InfluenceBoneIndex = 0;
 		for (auto Kvp : WeightForBone)
 		{
-			SkinWeight.InfluenceBones[InfluenceBoneIndex] = (uint8)(Kvp.Key);
+			SkinWeight.InfluenceBones[InfluenceBoneIndex] = Kvp.Key;
 			SkinWeight.InfluenceWeights[InfluenceBoneIndex] = FMath::Clamp((uint8)(Kvp.Value*((float)0xFF)), (uint8)0x00, (uint8)0xFF);
 			TotalInfluenceWeight += SkinWeight.InfluenceWeights[InfluenceBoneIndex];
 			InfluenceBoneIndex++;
@@ -2031,7 +2031,6 @@ void FLODUtilities::RegenerateDependentLODs(USkeletalMesh* SkeletalMesh, int32 L
 				});
 			}
 
-			//Reduce all dependent LOD in same time
 			// Load the BulkData for all dependent LODs; this has to be done on the main thread since it requires access to the Linker and FLinkerLoad::Serialize is not threadsafe
 			{
 				FSkeletalMeshModel* SkeletalMeshResource = SkeletalMesh->GetImportedModel();
@@ -2560,12 +2559,13 @@ void FLODUtilities::UnbindClothingAndBackup(USkeletalMesh* SkeletalMesh, TArray<
 	{
 		if (Binding.LODIndex == LODIndex)
 		{
-			check(Binding.Asset);
-			Binding.Asset->UnbindFromSkeletalMesh(SkeletalMesh, Binding.LODIndex);
-			
 			//Use the UserSectionsData original section index, this will ensure we remap correctly the cloth if the reduction has change the number of sections
 			int32 OriginalDataSectionIndex = LODModel.Sections[Binding.SectionIndex].OriginalDataSectionIndex;
-			Binding.SectionIndex = OriginalDataSectionIndex;
+			if (Binding.Asset)
+			{
+				Binding.Asset->UnbindFromSkeletalMesh(SkeletalMesh, Binding.LODIndex);
+				Binding.SectionIndex = OriginalDataSectionIndex;
+			}
 			
 			FSkelMeshSourceSectionUserData& SectionUserData = LODModel.UserSectionsData.FindChecked(OriginalDataSectionIndex);
 			SectionUserData.ClothingData.AssetGuid = FGuid();
@@ -2598,9 +2598,8 @@ void FLODUtilities::RestoreClothingFromBackup(USkeletalMesh* SkeletalMesh, TArra
 			{
 				continue;
 			}
-			if (Binding.LODIndex == LODIndex)
+			if (Binding.LODIndex == LODIndex && Binding.Asset)
 			{
-				check(Binding.Asset);
 				if (Binding.Asset->BindToSkeletalMesh(SkeletalMesh, Binding.LODIndex, SectionIndex, Binding.AssetInternalLodIndex))
 				{
 					//If successfull set back the section user data
diff --git a/Engine/Source/Editor/UnrealEd/Private/CookOnTheFlyServer.cpp b/Engine/Source/Editor/UnrealEd/Private/CookOnTheFlyServer.cpp
--- a/Engine/Source/Editor/UnrealEd/Private/CookOnTheFlyServer.cpp
+++ b/Engine/Source/Editor/UnrealEd/Private/CookOnTheFlyServer.cpp
@@ -4533,7 +4533,7 @@ void UCookOnTheFlyServer::SaveCookedPackage(UPackage* Package, uint32 SaveFlags,
 						}
 						else
 						{
-							FSavePackageContext* const SavePackageContext = SavePackageContexts.Num() > 0 ? SavePackageContexts[PlatformIndex] : nullptr;
+							FSavePackageContext* const SavePackageContext = (IsCookByTheBookMode() && IsUsingPackageStore()) ? SavePackageContexts[PlatformIndex] : nullptr;
 
 							Result = GEditor->Save(	Package, World, FlagsToCook, *PlatFilename, 
 													GError, nullptr, bSwap, false, SaveFlags, Target, 
@@ -7487,8 +7487,6 @@ void UCookOnTheFlyServer::StartCookByTheBook( const FCookByTheBookStartupOptions
 
 	// Find all the localized packages and map them back to their source package
 	{
-		UE_LOG(LogCook, Display, TEXT("Discovering localized assets"));
-
 		TArray<FString> AllCulturesToCook = CookByTheBookStartupOptions.CookCultures;
 		for (const FString& CultureName : CookByTheBookStartupOptions.CookCultures)
 		{
@@ -7500,6 +7498,8 @@ void UCookOnTheFlyServer::StartCookByTheBook( const FCookByTheBookStartupOptions
 		}
 		AllCulturesToCook.Sort();
 
+		UE_LOG(LogCook, Display, TEXT("Discovering localized assets for cultures: %s"), *FString::Join(AllCulturesToCook, TEXT(", ")));
+
 		TArray<FString> RootPaths;
 		FPackageName::QueryRootContentPaths(RootPaths);
 
@@ -7547,7 +7547,7 @@ void UCookOnTheFlyServer::StartCookByTheBook( const FCookByTheBookStartupOptions
 			for (const ITargetPlatform* TargetPlatform : TargetPlatforms)
 			{
 				FAssetRegistryGenerator* RegistryGenerator = PlatformManager->GetPlatformData(TargetPlatform)->RegistryGenerator.Get();
-				RegistryGenerator->RegisterChunkDataGenerator(MakeShared<FLocalizationChunkDataGenerator>(MoveTemp(LocalizationTargetsToChunk), MoveTemp(AllCulturesToCook)));
+				RegistryGenerator->RegisterChunkDataGenerator(MakeShared<FLocalizationChunkDataGenerator>(PackagingSettings->LocalizationTargetCatchAllChunkId, LocalizationTargetsToChunk, AllCulturesToCook));
 			}
 		}
 	}
diff --git a/Engine/Source/Programs/HeadlessChaos/Private/HeadlessChaosTestJointSolver.cpp b/Engine/Source/Programs/HeadlessChaos/Private/HeadlessChaosTestJointSolver.cpp
deleted file mode 100644
diff --git a/Engine/Source/Runtime/AnimGraphRuntime/Private/AnimCustomInstance.cpp b/Engine/Source/Runtime/AnimGraphRuntime/Private/AnimCustomInstance.cpp
deleted file mode 100644
diff --git a/Engine/Source/Runtime/AnimGraphRuntime/Private/AnimCustomInstanceHelper.cpp b/Engine/Source/Runtime/AnimGraphRuntime/Private/AnimCustomInstanceHelper.cpp
new file mode 100644
--- /dev/null
+++ b/Engine/Source/Runtime/AnimGraphRuntime/Private/AnimCustomInstanceHelper.cpp
@@ -0,0 +1,19 @@
+// Copyright Epic Games, Inc. All Rights Reserved.
+
+
+#include "AnimCustomInstanceHelper.h"
+
+/////////////////////////////////////////////////////
+// FAnimCustomInstanceHelper
+/////////////////////////////////////////////////////
+
+
+bool FAnimCustomInstanceHelper::ShouldCreateCustomInstancePlayer(const USkeletalMeshComponent* SkeletalMeshComponent)
+{
+	const USkeletalMesh* SkeletalMesh = SkeletalMeshComponent->SkeletalMesh;
+	// create proper anim instance to animate
+	UAnimInstance* AnimInstance = SkeletalMeshComponent->GetAnimInstance();
+
+	return (AnimInstance == nullptr || SkeletalMeshComponent->GetAnimationMode() != EAnimationMode::AnimationBlueprint ||
+		AnimInstance->GetClass() != SkeletalMeshComponent->AnimClass || !SkeletalMesh->Skeleton->IsCompatible(AnimInstance->CurrentSkeleton));
+}
\ No newline at end of file
diff --git a/Engine/Source/Runtime/Core/Private/Windows/WindowsPlatformFile.cpp b/Engine/Source/Runtime/Core/Private/Windows/WindowsPlatformFile.cpp
--- a/Engine/Source/Runtime/Core/Private/Windows/WindowsPlatformFile.cpp
+++ b/Engine/Source/Runtime/Core/Private/Windows/WindowsPlatformFile.cpp
@@ -95,13 +95,13 @@ namespace
 		// This roundabout conversion clamps the precision of the returned time value to match that of time_t (1 second precision)
 		// This avoids issues when sending files over the network via cook-on-the-fly
 		SYSTEMTIME SysTime;
-		SysTime.wYear = InDateTime.GetYear();
-		SysTime.wMonth = InDateTime.GetMonth();
-		SysTime.wDay = InDateTime.GetDay();
-		SysTime.wDayOfWeek = UEDayOfWeekToWindowsSystemTimeDayOfWeek(InDateTime.GetDayOfWeek());
-		SysTime.wHour = InDateTime.GetHour();
-		SysTime.wMinute = InDateTime.GetMinute();
-		SysTime.wSecond = InDateTime.GetSecond();
+		SysTime.wYear = (WORD)InDateTime.GetYear();
+		SysTime.wMonth = (WORD)InDateTime.GetMonth();
+		SysTime.wDay = (WORD)InDateTime.GetDay();
+		SysTime.wDayOfWeek = (WORD)UEDayOfWeekToWindowsSystemTimeDayOfWeek(InDateTime.GetDayOfWeek());
+		SysTime.wHour = (WORD)InDateTime.GetHour();
+		SysTime.wMinute = (WORD)InDateTime.GetMinute();
+		SysTime.wSecond = (WORD)InDateTime.GetSecond();
 		SysTime.wMilliseconds = 0;
 
 		FILETIME FileTime;
@@ -348,7 +348,7 @@ public:
 		if (bWithinSerializeBuffer)
 		{
 			// Still within the serialize buffer so just update the position
-			SerializePos += PosDelta;
+			SerializePos += (int32)PosDelta;
 		}
 		else
 		{
@@ -421,7 +421,7 @@ public:
 				FMemory::Memcpy(Dest, &Buffers[SerializeBuffer][SerializePos], NumToCopy);
 
 				// Update the internal positions
-				SerializePos += NumToCopy;
+				SerializePos += (int32)NumToCopy;
 				check(SerializePos <= BufferSize);
 				FilePos += NumToCopy;
 				check(FilePos <= FileSize);
@@ -562,7 +562,7 @@ public:
 		int64 TotalNumRead = 0;
 		do
 		{
-			uint32 BytesToRead32 = FMath::Min(BytesToRead, int64(UINT32_MAX));
+			uint32 BytesToRead32 = (uint32)FMath::Min<int64>(BytesToRead, int64(UINT32_MAX));
 			uint32 NumRead = 0;
 
 			if (!ReadFile(FileHandle, Destination, BytesToRead32, (::DWORD*)&NumRead, &OverlappedIO))
diff --git a/Engine/Source/Runtime/Core/Public/Containers/UnrealString.h b/Engine/Source/Runtime/Core/Public/Containers/UnrealString.h
--- a/Engine/Source/Runtime/Core/Public/Containers/UnrealString.h
+++ b/Engine/Source/Runtime/Core/Public/Containers/UnrealString.h
@@ -41,7 +41,7 @@ struct TIsContiguousContainer<FString>
 
 TCHAR* GetData(FString& String);
 const TCHAR* GetData(const FString& String);
-SIZE_T GetNum(const FString& String);
+int32 GetNum(const FString& String);
 
 /**
  * A dynamically sizeable string.
@@ -249,7 +249,7 @@ public:
 	FORCEINLINE DataType::RangedForConstIteratorType end  () const { auto Result = Data.end();   if (Data.Num()) { --Result; }     return Result; }
 
 public:
-	FORCEINLINE uint32 GetAllocatedSize() const
+	FORCEINLINE SIZE_T GetAllocatedSize() const
 	{
 		return Data.GetAllocatedSize();
 	}
@@ -1997,7 +1997,7 @@ inline const TCHAR* GetData(const FString& String)
 	return String.GetCharArray().GetData();
 }
 
-inline SIZE_T GetNum(const FString& String)
+inline int32 GetNum(const FString& String)
 {
 	return String.Len();
 }
@@ -2154,13 +2154,13 @@ CORE_API int32 HexToBytes(const FString& HexString, uint8* OutBytes);
  */
 
  /** Covert a string buffer to intrinsic types */
-inline void LexFromString(int8& OutValue, 		const TCHAR* Buffer)	{	OutValue = FCString::Atoi(Buffer);		}
-inline void LexFromString(int16& OutValue,		const TCHAR* Buffer)	{	OutValue = FCString::Atoi(Buffer);		}
-inline void LexFromString(int32& OutValue,		const TCHAR* Buffer)	{	OutValue = FCString::Atoi(Buffer);		}
+inline void LexFromString(int8& OutValue, 		const TCHAR* Buffer)	{	OutValue = (int8)FCString::Atoi(Buffer);		}
+inline void LexFromString(int16& OutValue,		const TCHAR* Buffer)	{	OutValue = (int16)FCString::Atoi(Buffer);		}
+inline void LexFromString(int32& OutValue,		const TCHAR* Buffer)	{	OutValue = (int32)FCString::Atoi(Buffer);		}
 inline void LexFromString(int64& OutValue,		const TCHAR* Buffer)	{	OutValue = FCString::Atoi64(Buffer);	}
-inline void LexFromString(uint8& OutValue,		const TCHAR* Buffer)	{	OutValue = FCString::Atoi(Buffer);		}
-inline void LexFromString(uint16& OutValue, 	const TCHAR* Buffer)	{	OutValue = FCString::Atoi(Buffer);		}
-inline void LexFromString(uint32& OutValue, 	const TCHAR* Buffer)	{	OutValue = FCString::Atoi64(Buffer);	}	//64 because this unsigned and so Atoi might overflow
+inline void LexFromString(uint8& OutValue,		const TCHAR* Buffer)	{	OutValue = (uint8)FCString::Atoi(Buffer);		}
+inline void LexFromString(uint16& OutValue, 	const TCHAR* Buffer)	{	OutValue = (uint16)FCString::Atoi(Buffer);		}
+inline void LexFromString(uint32& OutValue, 	const TCHAR* Buffer)	{	OutValue = (uint32)FCString::Atoi64(Buffer);	}	//64 because this unsigned and so Atoi might overflow
 inline void LexFromString(uint64& OutValue, 	const TCHAR* Buffer)	{	OutValue = FCString::Strtoui64(Buffer, nullptr, 0); }
 inline void LexFromString(float& OutValue,		const TCHAR* Buffer)	{	OutValue = FCString::Atof(Buffer);		}
 inline void LexFromString(double& OutValue, 	const TCHAR* Buffer)	{	OutValue = FCString::Atod(Buffer);		}
@@ -2229,7 +2229,7 @@ LexTryParseString(T& OutValue, const TCHAR* Buffer)
 	}
 
 	LexFromString(OutValue, Buffer);
-	if (OutValue == 0 && FMath::IsFinite(OutValue))
+	if (OutValue == 0 && FMath::IsFinite((float)OutValue)) //@TODO:FLOATPRECISION: ? huh ?
 	{
 		bool bSawZero = false;
 		TCHAR C = *Buffer;
diff --git a/Engine/Source/Runtime/Core/Public/Stats/Stats2.h b/Engine/Source/Runtime/Core/Public/Stats/Stats2.h
--- a/Engine/Source/Runtime/Core/Public/Stats/Stats2.h
+++ b/Engine/Source/Runtime/Core/Public/Stats/Stats2.h
@@ -1287,7 +1287,7 @@ public:
 		if( Packet.ThreadType == EThreadType::Other )
 		{
 			FPlatformMisc::MemoryBarrier();
-			uint64 Frame = FStats::GameThreadStatsFrame;
+			int32 Frame = FStats::GameThreadStatsFrame;
 			const bool bFrameHasChanged = Frame > CurrentGameFrame;
 			if( bFrameHasChanged )
 			{
diff --git a/Engine/Source/Runtime/Core/Public/Unix/UnixPlatform.h b/Engine/Source/Runtime/Core/Public/Unix/UnixPlatform.h
--- a/Engine/Source/Runtime/Core/Public/Unix/UnixPlatform.h
+++ b/Engine/Source/Runtime/Core/Public/Unix/UnixPlatform.h
@@ -93,6 +93,11 @@ typedef FUnixPlatformTypes FPlatformTypes;
 	#define PLATFORM_ENABLE_VECTORINTRINSICS		0
 #endif // defined(_M_IX86) || defined(__i386__) || defined(_M_X64) || defined(__x86_64__) || defined (__amd64__) 
 
+#if PLATFORM_LINUXAARCH64
+	// Enable NEON intrinsics for ARM64 builds
+	#define PLATFORM_ENABLE_VECTORINTRINSICS_NEON	1
+#endif
+
 // We do not currently compile with -msse4 or higher on Unix or Linux
 #ifndef PLATFORM_MAYBE_HAS_SSE4_1 // May be set from UnrealBuildTool
 	#define PLATFORM_MAYBE_HAS_SSE4_1							0
@@ -101,6 +106,7 @@ typedef FUnixPlatformTypes FPlatformTypes;
 	#define PLATFORM_ALWAYS_HAS_SSE4_1							0
 #endif
 
+
 // Function type macros.
 #define VARARGS															/* Functions with variable arguments */
 #define CDECL															/* Standard C function */
diff --git a/Engine/Source/Runtime/RHI/Private/RHICommandList.cpp b/Engine/Source/Runtime/RHI/Private/RHICommandList.cpp
--- a/Engine/Source/Runtime/RHI/Private/RHICommandList.cpp
+++ b/Engine/Source/Runtime/RHI/Private/RHICommandList.cpp
@@ -26,6 +26,10 @@ DECLARE_DWORD_COUNTER_STAT(TEXT("Immed. Command count"), STAT_ImmedCmdListCount,
 
 UE_TRACE_CHANNEL_DEFINE(RHICommandsChannel);
 
+#if VALIDATE_UNIFORM_BUFFER_GLOBAL_BINDINGS
+bool FScopedUniformBufferGlobalBindings::bRecursionGuard = false;
+#endif
+
 #if !PLATFORM_USES_FIXED_RHI_CLASS
 #include "RHICommandListCommandExecutes.inl"
 #endif
@@ -138,7 +142,6 @@ uint32 GWorkingRHIThreadStartCycles = 0;
 uint64 GInputLatencyTime = 0;
 
 RHI_API bool GEnableAsyncCompute = true;
-RHI_API bool GSupportAsyncComputeRaytracingBuildBVH = false;
 RHI_API FRHICommandListExecutor GRHICommandList;
 
 static FGraphEventArray AllOutstandingTasks;
@@ -186,18 +189,6 @@ void FRHICommandListImmediate::SetCurrentStat(TStatId Stat)
 	}
 }
 
-void FRHICommandListBase::SetGPUMaskOnContext()
-{
-#if WITH_MGPU
-	if (Bypass())
-	{
-		GetContext().RHISetGPUMask(GPUMask);
-		return;
-	}
-	ALLOC_COMMAND(FRHICommandSetGPUMask)(GPUMask);
-#endif
-}
-
 DECLARE_CYCLE_STAT(TEXT("FNullGraphTask.RenderThreadTaskFence"), STAT_RenderThreadTaskFence, STATGROUP_TaskGraphTasks);
 DECLARE_CYCLE_STAT(TEXT("Render thread task fence wait"), STAT_RenderThreadTaskFenceWait, STATGROUP_TaskGraphTasks);
 FGraphEventRef FRHICommandListImmediate::RenderThreadTaskFence()
@@ -254,6 +245,7 @@ void FRHIAsyncComputeCommandListImmediate::ImmediateDispatch(FRHIAsyncComputeCom
 			static_assert(sizeof(FRHICommandList) == sizeof(FRHIAsyncComputeCommandListImmediate), "We are memswapping FRHICommandList and FRHICommandListImmediate; they need to be swappable.");
 			check(RHIComputeCmdList.IsImmediateAsyncCompute());
 			SwapCmdList->ExchangeCmdList(RHIComputeCmdList);
+			RHIComputeCmdList.CopyContext(*SwapCmdList);
 			RHIComputeCmdList.GPUMask = SwapCmdList->GPUMask;
 			// NB: InitialGPUMask set to GPUMask since exchanging the list
 			// is equivalent to a Reset.
@@ -271,18 +263,6 @@ void FRHIAsyncComputeCommandListImmediate::ImmediateDispatch(FRHIAsyncComputeCom
 	}
 }
 
-void FRHICommandListBase::SetGPUMaskOnComputeContext()
-{
-#if WITH_MGPU
-	if (Bypass())
-	{
-		GetComputeContext().RHISetGPUMask(GPUMask);
-		return;
-	}
-	ALLOC_COMMAND(FRHICommandSetGPUMask)(GPUMask);
-#endif
-}
-
 FRHICommandBase* GCurrentCommand = nullptr;
 
 DECLARE_CYCLE_STAT(TEXT("BigList"), STAT_BigList, STATGROUP_RHICMDLIST);
@@ -507,6 +487,7 @@ void FRHICommandListExecutor::ExecuteInner(FRHICommandListBase& CmdList)
 				// we should make command lists virtual and transfer ownership rather than this devious approach
 				static_assert(sizeof(FRHICommandList) == sizeof(FRHICommandListImmediate), "We are memswapping FRHICommandList and FRHICommandListImmediate; they need to be swappable.");
 				SwapCmdList->ExchangeCmdList(CmdList);
+				CmdList.CopyContext(*SwapCmdList);
 				CmdList.GPUMask = SwapCmdList->GPUMask;
 				// NB: InitialGPUMask set to GPUMask since exchanging the list
                 // is equivalent to a Reset.
@@ -806,24 +787,24 @@ FGraphEventRef FRHICommandListImmediate::RHIThreadFence(bool bSetLockFence)
 	}
 
 	return nullptr;
-}		
+}
 
 DECLARE_CYCLE_STAT(TEXT("Async Compute CmdList Execute"), STAT_AsyncComputeExecute, STATGROUP_RHICMDLIST);
 FRHICOMMAND_MACRO(FRHIAsyncComputeSubmitList)
 {
-	FRHIAsyncComputeCommandList* RHICmdList;
-	FORCEINLINE_DEBUGGABLE FRHIAsyncComputeSubmitList(FRHIAsyncComputeCommandList* InRHICmdList)
+	FRHIComputeCommandList* RHICmdList;
+	FORCEINLINE_DEBUGGABLE FRHIAsyncComputeSubmitList(FRHIComputeCommandList* InRHICmdList)
 		: RHICmdList(InRHICmdList)
 	{
 	}
 	void Execute(FRHICommandListBase& CmdList)
-	{		
+	{
 		SCOPE_CYCLE_COUNTER(STAT_AsyncComputeExecute);
 		delete RHICmdList;
 	}
 };
 
-void FRHICommandListImmediate::QueueAsyncCompute(FRHIAsyncComputeCommandList& RHIComputeCmdList)
+void FRHICommandListImmediate::QueueAsyncCompute(FRHIComputeCommandList& RHIComputeCmdList)
 {
 	if (Bypass())
 	{
@@ -865,7 +846,6 @@ FRHICommandListBase::FRHICommandListBase(FRHIGPUMask InGPUMask)
 	, ComputeContext(nullptr)
 	, MemManager(0)
 	, bAsyncPSOCompileAllowed(true)
-	, bRecursive(false)
 	, GPUMask(InGPUMask)
 	, InitialGPUMask(InGPUMask)
 {
@@ -892,20 +872,7 @@ void FRHICommandListBase::Reset()
 	NumCommands = 0;
 	Root = nullptr;
 	CommandLink = &Root;
-	if (!bRecursive)
-	{
-		Context = GDynamicRHI ? RHIGetDefaultContext() : nullptr;
 
-		if (GEnableAsyncCompute)
-		{
-			ComputeContext = GDynamicRHI ? RHIGetDefaultAsyncComputeContext() : nullptr;
-		}
-		else
-		{
-			ComputeContext = Context;
-		}
-	}
-	
 	UID = GRHICommandList.UIDCounter.Increment();
 	for (int32 Index = 0; ERenderThreadContext(Index) < ERenderThreadContext::Num; Index++)
 	{
@@ -1911,12 +1878,12 @@ void FRHICommandList::operator delete(void *RawMemory)
 	FMemory::Free(RawMemory);
 }
 
-void* FRHIAsyncComputeCommandList::operator new(size_t Size)
+void* FRHIComputeCommandList::operator new(size_t Size)
 {
 	return FMemory::Malloc(Size);
 }
 
-void FRHIAsyncComputeCommandList::operator delete(void *RawMemory)
+void FRHIComputeCommandList::operator delete(void *RawMemory)
 {
 	FMemory::Free(RawMemory);
 }
@@ -1978,6 +1945,13 @@ FShaderResourceViewRHIRef FDynamicRHI::CreateShaderResourceView_RenderThread(cla
 	return GDynamicRHI->RHICreateShaderResourceView(VertexBuffer, Stride, Format);
 }
 
+FShaderResourceViewRHIRef FDynamicRHI::CreateShaderResourceView_RenderThread(class FRHICommandListImmediate& RHICmdList, const FShaderResourceViewInitializer& Initializer)
+{
+	CSV_SCOPED_TIMING_STAT(RHITStalls, CreateShaderResourceView_RenderThread_VB);
+	FScopedRHIThreadStaller StallRHIThread(RHICmdList);
+	return GDynamicRHI->RHICreateShaderResourceView(Initializer);
+}
+
 FShaderResourceViewRHIRef FDynamicRHI::CreateShaderResourceView_RenderThread(class FRHICommandListImmediate& RHICmdList, FRHIIndexBuffer* Buffer)
 {
 	CSV_SCOPED_TIMING_STAT(RHITStalls, CreateShaderResourceView_RenderThread_IB);
@@ -2228,7 +2202,7 @@ void FDynamicRHI::RHIUnlockStructuredBuffer(class FRHICommandListImmediate& RHIC
 }
 
 // @todo-mattc-staging Default implementation
-void* FDynamicRHI::RHILockStagingBuffer(FRHIStagingBuffer* StagingBuffer, uint32 Offset, uint32 SizeRHI)
+void* FDynamicRHI::RHILockStagingBuffer(FRHIStagingBuffer* StagingBuffer, FRHIGPUFence* Fence, uint32 Offset, uint32 SizeRHI)
 {
 	check(false);
 	return nullptr;
@@ -2240,19 +2214,41 @@ void FDynamicRHI::RHIUnlockStagingBuffer(FRHIStagingBuffer* StagingBuffer)
 	//GDynamicRHI->RHIUnlockVertexBuffer(StagingBuffer->GetSourceBuffer());
 }
 
-void* FDynamicRHI::LockStagingBuffer_RenderThread(class FRHICommandListImmediate& RHICmdList, FRHIStagingBuffer* StagingBuffer, uint32 Offset, uint32 SizeRHI)
+void* FDynamicRHI::LockStagingBuffer_RenderThread(class FRHICommandListImmediate& RHICmdList, FRHIStagingBuffer* StagingBuffer, FRHIGPUFence* Fence, uint32 Offset, uint32 SizeRHI)
 {
-	QUICK_SCOPE_CYCLE_COUNTER(STAT_FDynamicRHI_LockStagingBuffer_RenderThread);
 	check(IsInRenderingThread());
-	RHICmdList.ImmediateFlush(EImmediateFlushType::FlushRHIThread);
-	return GDynamicRHI->RHILockStagingBuffer(StagingBuffer, Offset, SizeRHI);
+	if (Fence == nullptr || !Fence->Poll())
+	{
+		QUICK_SCOPE_CYCLE_COUNTER(STAT_FDynamicRHI_LockStagingBuffer_Flush);
+		RHICmdList.ImmediateFlush(EImmediateFlushType::FlushRHIThread);
+	}
+	{
+		QUICK_SCOPE_CYCLE_COUNTER(STAT_FDynamicRHI_LockStagingBuffer_RenderThread);
+		if (GRHISupportsMultithreading)
+		{
+			return GDynamicRHI->RHILockStagingBuffer(StagingBuffer, Fence, Offset, SizeRHI);
+		}
+		else
+		{
+			FScopedRHIThreadStaller StallRHIThread(RHICmdList);
+			return GDynamicRHI->RHILockStagingBuffer(StagingBuffer, Fence, Offset, SizeRHI);
+		}
+	}
 }
+
 void FDynamicRHI::UnlockStagingBuffer_RenderThread(class FRHICommandListImmediate& RHICmdList, FRHIStagingBuffer* StagingBuffer)
 {
 	QUICK_SCOPE_CYCLE_COUNTER(STAT_FDynamicRHI_UnlockStagingBuffer_RenderThread);
 	check(IsInRenderingThread());
-	RHICmdList.ImmediateFlush(EImmediateFlushType::FlushRHIThread);
-	GDynamicRHI->RHIUnlockStagingBuffer(StagingBuffer);
+	if (GRHISupportsMultithreading)
+	{
+		GDynamicRHI->RHIUnlockStagingBuffer(StagingBuffer);
+	}
+	else
+	{
+		FScopedRHIThreadStaller StallRHIThread(RHICmdList);
+		GDynamicRHI->RHIUnlockStagingBuffer(StagingBuffer);
+	}
 }
 
 FTexture2DRHIRef FDynamicRHI::AsyncReallocateTexture2D_RenderThread(class FRHICommandListImmediate& RHICmdList, FRHITexture2D* Texture2D, int32 NewMipCount, int32 NewSizeX, int32 NewSizeY, FThreadSafeCounter* RequestStatus)
@@ -2531,6 +2527,13 @@ FShaderResourceViewRHIRef FDynamicRHI::RHICreateShaderResourceView_RenderThread(
 	return GDynamicRHI->RHICreateShaderResourceView(VertexBuffer, Stride, Format);
 }
 
+FShaderResourceViewRHIRef FDynamicRHI::RHICreateShaderResourceView_RenderThread(class FRHICommandListImmediate& RHICmdList, const FShaderResourceViewInitializer& Initializer)
+{
+	CSV_SCOPED_TIMING_STAT(RHITStalls, RHICreateShaderResourceView_RenderThread_VB);
+	FScopedRHIThreadStaller StallRHIThread(RHICmdList);
+	return GDynamicRHI->RHICreateShaderResourceView(Initializer);
+}
+
 FShaderResourceViewRHIRef FDynamicRHI::RHICreateShaderResourceView_RenderThread(class FRHICommandListImmediate& RHICmdList, FRHIIndexBuffer* Buffer)
 {
 	CSV_SCOPED_TIMING_STAT(RHITStalls, RHICreateShaderResourceView_RenderThread_IB);
@@ -2605,15 +2608,30 @@ void FDynamicRHI::RHIMapStagingSurface_RenderThread(class FRHICommandListImmedia
 		RHICmdList.ImmediateFlush(EImmediateFlushType::FlushRHIThread);
 	}
 	{
-		FScopedRHIThreadStaller StallRHIThread(RHICmdList);
-		GDynamicRHI->RHIMapStagingSurface(Texture, Fence, OutData, OutWidth, OutHeight, RHICmdList.GetGPUMask().ToIndex());
+		QUICK_SCOPE_CYCLE_COUNTER(STAT_RHIMETHOD_MapStagingSurface_RenderThread);
+		if (GRHISupportsMultithreading)
+		{
+			GDynamicRHI->RHIMapStagingSurface(Texture, Fence, OutData, OutWidth, OutHeight, RHICmdList.GetGPUMask().ToIndex());
+		}
+		else
+		{
+			FScopedRHIThreadStaller StallRHIThread(RHICmdList);
+			GDynamicRHI->RHIMapStagingSurface(Texture, Fence, OutData, OutWidth, OutHeight, RHICmdList.GetGPUMask().ToIndex());
+		}
 	}
 }
 
 void FDynamicRHI::RHIUnmapStagingSurface_RenderThread(class FRHICommandListImmediate& RHICmdList, FRHITexture* Texture)
 {
-	FScopedRHIThreadStaller StallRHIThread(RHICmdList);
-	GDynamicRHI->RHIUnmapStagingSurface(Texture, RHICmdList.GetGPUMask().ToIndex());
+	if (GRHISupportsMultithreading)
+	{
+		GDynamicRHI->RHIUnmapStagingSurface(Texture, RHICmdList.GetGPUMask().ToIndex());
+	}
+	else
+	{
+		FScopedRHIThreadStaller StallRHIThread(RHICmdList);
+		GDynamicRHI->RHIUnmapStagingSurface(Texture, RHICmdList.GetGPUMask().ToIndex());
+	}
 }
 
 void FDynamicRHI::RHIReadSurfaceFloatData_RenderThread(class FRHICommandListImmediate& RHICmdList, FRHITexture* Texture, FIntRect Rect, TArray<FFloat16Color>& OutData, ECubeFace CubeFace, int32 ArrayIndex, int32 MipIndex)
diff --git a/Engine/Source/Runtime/Renderer/Private/HairStrands/HairStrandsVisibility.cpp b/Engine/Source/Runtime/Renderer/Private/HairStrands/HairStrandsVisibility.cpp
--- a/Engine/Source/Runtime/Renderer/Private/HairStrands/HairStrandsVisibility.cpp
+++ b/Engine/Source/Runtime/Renderer/Private/HairStrands/HairStrandsVisibility.cpp
@@ -133,7 +133,7 @@ uint32 GetPPLLMaxTotalRenderNode(FIntPoint Resolution)
 
 uint32 GetHairVisibilitySampleCount()
 {
-	return GetHairVisibilityRenderMode() == HairVisibilityRenderMode_MSAA ? FMath::Clamp(GHairVisibilitySampleCount,1,16) : 1;
+	return GetHairVisibilityRenderMode() == HairVisibilityRenderMode_MSAA ? FMath::Clamp(GHairVisibilitySampleCount, 1, 16) : 1;
 }
 
 void SetUpViewHairRenderInfo(const FViewInfo& ViewInfo, bool bEnableMSAA, FVector4& OutHairRenderInfo)
@@ -159,13 +159,13 @@ static bool IsCompatibleWithHairVisibility(const FMeshMaterialShaderPermutationP
 /////////////////////////////////////////////////////////////////////////////////////////
 
 BEGIN_GLOBAL_SHADER_PARAMETER_STRUCT(FMaterialPassParameters, )
-	SHADER_PARAMETER(FIntPoint, MaxResolution)
-	SHADER_PARAMETER(uint32, MaxSampleCount)
-	SHADER_PARAMETER(uint32, NodeGroupSize)
-	SHADER_PARAMETER_TEXTURE(Texture2D<uint>, NodeIndex)
-	SHADER_PARAMETER_SRV(StructuredBuffer<uint>, NodeCoord)
-	SHADER_PARAMETER_SRV(StructuredBuffer<FNodeVis>, NodeVis)
-	SHADER_PARAMETER_SRV(Buffer<uint>, IndirectArgs)
+SHADER_PARAMETER(FIntPoint, MaxResolution)
+SHADER_PARAMETER(uint32, MaxSampleCount)
+SHADER_PARAMETER(uint32, NodeGroupSize)
+SHADER_PARAMETER_TEXTURE(Texture2D<uint>, NodeIndex)
+SHADER_PARAMETER_SRV(StructuredBuffer<uint>, NodeCoord)
+SHADER_PARAMETER_SRV(StructuredBuffer<FNodeVis>, NodeVis)
+SHADER_PARAMETER_SRV(Buffer<uint>, IndirectArgs)
 END_GLOBAL_SHADER_PARAMETER_STRUCT()
 IMPLEMENT_GLOBAL_SHADER_PARAMETER_STRUCT(FMaterialPassParameters, "MaterialPassParameters");
 
@@ -321,7 +321,7 @@ void FHairMaterialProcessor::AddMeshBatch(const FMeshBatch& RESTRICT MeshBatch,
 	const FMaterial& Material = MeshBatch.MaterialRenderProxy->GetMaterialWithFallback(FeatureLevel, FallbackMaterialRenderProxyPtr);
 	const bool bIsCompatible = IsCompatibleWithHairStrands(&Material, FeatureLevel);
 	const bool bIsHairStrandsFactory = MeshBatch.VertexFactory->GetType()->GetId() == CompatibleVF->GetId();
-	const bool bShouldRender = (!PrimitiveSceneProxy && MeshBatch.Elements.Num()>0) || (PrimitiveSceneProxy && PrimitiveSceneProxy->ShouldRenderInMainPass());
+	const bool bShouldRender = (!PrimitiveSceneProxy && MeshBatch.Elements.Num() > 0) || (PrimitiveSceneProxy && PrimitiveSceneProxy->ShouldRenderInMainPass());
 
 	if (bIsCompatible
 		&& bIsHairStrandsFactory
@@ -330,7 +330,7 @@ void FHairMaterialProcessor::AddMeshBatch(const FMeshBatch& RESTRICT MeshBatch,
 	{
 		// For the mesh patch to be rendered a single triangle triangle to spawn the necessary amount of thread
 		FMeshBatch MeshBatchCopy = MeshBatch;
-		for (uint32 ElementIt = 0, ElementCount=uint32(MeshBatch.Elements.Num()); ElementIt < ElementCount; ++ElementIt)
+		for (uint32 ElementIt = 0, ElementCount = uint32(MeshBatch.Elements.Num()); ElementIt < ElementCount; ++ElementIt)
 		{
 			MeshBatchCopy.Elements[ElementIt].FirstIndex = 0;
 			MeshBatchCopy.Elements[ElementIt].NumPrimitives = 1;
@@ -405,13 +405,13 @@ FHairMaterialProcessor::FHairMaterialProcessor(
 /////////////////////////////////////////////////////////////////////////////////////////
 
 BEGIN_SHADER_PARAMETER_STRUCT(FVisibilityMaterialPassParameters, )
-	SHADER_PARAMETER_RDG_TEXTURE(Texture2D<uint>, NodeIndex)
-	SHADER_PARAMETER_RDG_BUFFER_SRV(StructuredBuffer<uint>,	NodeCoord)
-	SHADER_PARAMETER_RDG_BUFFER_SRV(StructuredBuffer<FNodeVis>, NodeVis)
-	SHADER_PARAMETER_RDG_BUFFER_SRV(Buffer<uint>, IndirectArgs)
-	SHADER_PARAMETER_RDG_BUFFER_UAV(RWStructuredBuffer<FNodeData>, OutNodeData)
-	SHADER_PARAMETER_RDG_BUFFER_UAV(RWBuffer<float2>, OutNodeVelocity)
-	RENDER_TARGET_BINDING_SLOTS()
+SHADER_PARAMETER_RDG_TEXTURE(Texture2D<uint>, NodeIndex)
+SHADER_PARAMETER_RDG_BUFFER_SRV(StructuredBuffer<uint>, NodeCoord)
+SHADER_PARAMETER_RDG_BUFFER_SRV(StructuredBuffer<FNodeVis>, NodeVis)
+SHADER_PARAMETER_RDG_BUFFER_SRV(Buffer<uint>, IndirectArgs)
+SHADER_PARAMETER_RDG_BUFFER_UAV(RWStructuredBuffer<FNodeData>, OutNodeData)
+SHADER_PARAMETER_RDG_BUFFER_UAV(RWBuffer<float2>, OutNodeVelocity)
+RENDER_TARGET_BINDING_SLOTS()
 END_SHADER_PARAMETER_STRUCT()
 
 struct FMaterialPassOutput
@@ -438,7 +438,7 @@ static FMaterialPassOutput AddHairMaterialPass(
 	const uint32 MaxNodeCount = CompactNodeVis->Desc.NumElements;
 
 	FMaterialPassOutput Output;
-	Output.NodeData		= GraphBuilder.CreateBuffer(FRDGBufferDesc::CreateStructuredDesc(sizeof(HairStrandsVisibilityInternal::NodeData), MaxNodeCount), TEXT("HairCompactNodeData"));
+	Output.NodeData = GraphBuilder.CreateBuffer(FRDGBufferDesc::CreateStructuredDesc(sizeof(HairStrandsVisibilityInternal::NodeData), MaxNodeCount), TEXT("HairCompactNodeData"));
 	Output.NodeVelocity = GraphBuilder.CreateBuffer(FRDGBufferDesc::CreateBufferDesc(4, CompactNodeVis->Desc.NumElements), TEXT("HairCompactNodeVelocity"));
 
 	const uint32 ResolutionDim = FMath::CeilToInt(FMath::Sqrt(MaxNodeCount));
@@ -457,12 +457,12 @@ static FMaterialPassOutput AddHairMaterialPass(
 
 	// Add resources reference to the pass parameters, in order to get the resource lifetime extended to this pass
 	FVisibilityMaterialPassParameters* PassParameters = GraphBuilder.AllocParameters<FVisibilityMaterialPassParameters>();
-	PassParameters->NodeIndex		= CompactNodeIndex;
-	PassParameters->NodeVis			= GraphBuilder.CreateSRV(CompactNodeVis);
-	PassParameters->NodeCoord		= GraphBuilder.CreateSRV(CompactNodeCoord);
-	PassParameters->IndirectArgs	= GraphBuilder.CreateSRV(IndirectArgBuffer);
-	PassParameters->OutNodeData		= GraphBuilder.CreateUAV(FRDGBufferUAVDesc(Output.NodeData));
-	PassParameters->OutNodeVelocity	= GraphBuilder.CreateUAV(FRDGBufferUAVDesc(Output.NodeVelocity, FMaterialPassOutput::VelocityFormat));
+	PassParameters->NodeIndex = CompactNodeIndex;
+	PassParameters->NodeVis = GraphBuilder.CreateSRV(CompactNodeVis);
+	PassParameters->NodeCoord = GraphBuilder.CreateSRV(CompactNodeCoord);
+	PassParameters->IndirectArgs = GraphBuilder.CreateSRV(IndirectArgBuffer);
+	PassParameters->OutNodeData = GraphBuilder.CreateUAV(FRDGBufferUAVDesc(Output.NodeData));
+	PassParameters->OutNodeVelocity = GraphBuilder.CreateUAV(FRDGBufferUAVDesc(Output.NodeVelocity, FMaterialPassOutput::VelocityFormat));
 	PassParameters->RenderTargets[0] = FRenderTargetBinding(OutDummyTexture0, ERenderTargetLoadAction::EClear, 0);
 
 
@@ -476,13 +476,13 @@ static FMaterialPassOutput AddHairMaterialPass(
 		check(IsInRenderingThread());
 
 		FMaterialPassParameters MaterialPassParameters;
-		MaterialPassParameters.MaxResolution	= Resolution;
-		MaterialPassParameters.NodeGroupSize	= NodeGroupSize;
-		MaterialPassParameters.MaxSampleCount	= MaxNodeCount;
-		MaterialPassParameters.NodeIndex		= PassParameters->NodeIndex->GetPooledRenderTarget()->GetRenderTargetItem().ShaderResourceTexture;
-		MaterialPassParameters.NodeCoord		= PassParameters->NodeCoord->GetRHI();
-		MaterialPassParameters.NodeVis			= PassParameters->NodeVis->GetRHI();
-		MaterialPassParameters.IndirectArgs		= PassParameters->IndirectArgs->GetRHI();
+		MaterialPassParameters.MaxResolution = Resolution;
+		MaterialPassParameters.NodeGroupSize = NodeGroupSize;
+		MaterialPassParameters.MaxSampleCount = MaxNodeCount;
+		MaterialPassParameters.NodeIndex = PassParameters->NodeIndex->GetPooledRenderTarget()->GetRenderTargetItem().ShaderResourceTexture;
+		MaterialPassParameters.NodeCoord = PassParameters->NodeCoord->GetRHI();
+		MaterialPassParameters.NodeVis = PassParameters->NodeVis->GetRHI();
+		MaterialPassParameters.IndirectArgs = PassParameters->IndirectArgs->GetRHI();
 		TUniformBufferRef<FMaterialPassParameters> MaterialPassParametersBuffer = TUniformBufferRef<FMaterialPassParameters>::CreateUniformBufferImmediate(MaterialPassParameters, UniformBuffer_SingleFrame);
 
 		FMeshPassProcessorRenderState DrawRenderState(*ViewInfo, MaterialPassParametersBuffer);
@@ -498,7 +498,7 @@ static FMaterialPassOutput AddHairMaterialPass(
 			RHICmdList.SetViewport(0, 0, 0.0f, Resolution.X, Resolution.Y, 1.0f);
 			DrawRenderState.SetBlendState(TStaticBlendState<>::GetRHI());
 			DrawRenderState.SetDepthStencilState(TStaticDepthStencilState <false, CF_Always> ::GetRHI());
-			
+
 			FDynamicMeshDrawCommandStorage DynamicMeshDrawCommandStorage;
 			FMeshCommandOneFrameArray VisibleMeshDrawCommands;
 			FGraphicsMinimalPipelineStateSet PipelineStateSet;
@@ -528,7 +528,7 @@ static FMaterialPassOutput AddHairMaterialPass(
 }
 
 ///////////////////////////////////////////////////////////////////////////////////////////////////
-class FHairVelocityCS: public FGlobalShader
+class FHairVelocityCS : public FGlobalShader
 {
 	DECLARE_GLOBAL_SHADER(FHairVelocityCS);
 	SHADER_USE_PARAMETER_STRUCT(FHairVelocityCS, FGlobalShader);
@@ -545,7 +545,7 @@ class FHairVelocityCS: public FGlobalShader
 		SHADER_PARAMETER_RDG_TEXTURE_UAV(Texture2D, OutVelocityTexture)
 		SHADER_PARAMETER_STRUCT_REF(FViewUniformShaderParameters, ViewUniformBuffer)
 		SHADER_PARAMETER_STRUCT_REF(FSceneTexturesUniformParameters, SceneTexturesStruct)
-	END_SHADER_PARAMETER_STRUCT()
+		END_SHADER_PARAMETER_STRUCT()
 
 public:
 
@@ -611,22 +611,22 @@ static void AddHairVelocityPass(
 
 /////////////////////////////////////////////////////////////////////////////////////////
 BEGIN_GLOBAL_SHADER_PARAMETER_STRUCT(FVisibilityPassGlobalParameters, )
-	SHADER_PARAMETER(uint32, MaxPPLLNodeCount)
+SHADER_PARAMETER(uint32, MaxPPLLNodeCount)
 END_GLOBAL_SHADER_PARAMETER_STRUCT()
 IMPLEMENT_GLOBAL_SHADER_PARAMETER_STRUCT(FVisibilityPassGlobalParameters, "VisibilityPassGlobalParameters");
 
 BEGIN_SHADER_PARAMETER_STRUCT(FVisibilityPassParameters, )
-	SHADER_PARAMETER(uint32, HairVisibilityPass_MaxPPLLNodeCount)
-	SHADER_PARAMETER_RDG_TEXTURE_UAV(Texture2D, PPLLCounter)
-	SHADER_PARAMETER_RDG_TEXTURE_UAV(Texture2D, PPLLNodeIndex)
-	SHADER_PARAMETER_RDG_BUFFER_UAV(StructuredBuffer, PPLLNodeData)
-	RENDER_TARGET_BINDING_SLOTS()
+SHADER_PARAMETER(uint32, HairVisibilityPass_MaxPPLLNodeCount)
+SHADER_PARAMETER_RDG_TEXTURE_UAV(Texture2D, PPLLCounter)
+SHADER_PARAMETER_RDG_TEXTURE_UAV(Texture2D, PPLLNodeIndex)
+SHADER_PARAMETER_RDG_BUFFER_UAV(StructuredBuffer, PPLLNodeData)
+RENDER_TARGET_BINDING_SLOTS()
 END_SHADER_PARAMETER_STRUCT()
 
 template<EHairVisibilityRenderMode RenderMode>
 class FHairVisibilityVS : public FMeshMaterialShader
 {
-	DECLARE_SHADER_TYPE(FHairVisibilityVS,MeshMaterial);
+	DECLARE_SHADER_TYPE(FHairVisibilityVS, MeshMaterial);
 
 protected:
 
@@ -678,7 +678,7 @@ class FHairVisibilityPS : public FMeshMaterialShader
 public:
 
 	FHairVisibilityPS(const ShaderMetaType::CompiledShaderInitializerType& Initializer)
-		:	FMeshMaterialShader(Initializer)
+		: FMeshMaterialShader(Initializer)
 	{
 		ERHIFeatureLevel::Type FeatureLevel = GetMaxSupportedFeatureLevel((EShaderPlatform)Initializer.Target.Platform);
 		check(FSceneInterface::GetShadingPath(FeatureLevel) != EShadingPath::Mobile);
@@ -701,7 +701,7 @@ public:
 		return IsCompatibleWithHairStrands(Parameters.Material, Parameters.Platform) && Parameters.VertexFactoryType->GetFName() == FName(TEXT("FHairStrandsVertexFactory"));
 	}
 
-	static void ModifyCompilationEnvironment(const FMaterialShaderPermutationParameters& Parameters, FShaderCompilerEnvironment& OutEnvironment)	
+	static void ModifyCompilationEnvironment(const FMaterialShaderPermutationParameters& Parameters, FShaderCompilerEnvironment& OutEnvironment)
 	{
 		FMaterialShader::ModifyCompilationEnvironment(Parameters, OutEnvironment);
 		const uint32 RenderModeValue = uint32(RenderMode);
@@ -803,9 +803,9 @@ void FHairVisibilityProcessor::AddMeshBatch(const FMeshBatch& RESTRICT MeshBatch
 	const FMaterial& Material = MeshBatch.MaterialRenderProxy->GetMaterialWithFallback(FeatureLevel, FallbackMaterialRenderProxyPtr);
 	const bool bIsCompatible = IsCompatibleWithHairStrands(&Material, FeatureLevel);
 	const bool bIsHairStrandsFactory = MeshBatch.VertexFactory->GetType()->GetId() == CompatibleVF->GetId();
-	const bool bShouldRender = (!PrimitiveSceneProxy && MeshBatch.Elements.Num()>0) || (PrimitiveSceneProxy && PrimitiveSceneProxy->ShouldRenderInMainPass());
+	const bool bShouldRender = (!PrimitiveSceneProxy && MeshBatch.Elements.Num() > 0) || (PrimitiveSceneProxy && PrimitiveSceneProxy->ShouldRenderInMainPass());
 
-	if (bIsCompatible 
+	if (bIsCompatible
 		&& bIsHairStrandsFactory
 		&& bShouldRender
 		&& ShouldIncludeDomainInMeshPass(Material.GetMaterialDomain()))
@@ -828,8 +828,8 @@ void FHairVisibilityProcessor::AddMeshBatch(const FMeshBatch& RESTRICT MeshBatch
 
 template<EHairVisibilityRenderMode TRenderMode>
 void FHairVisibilityProcessor::Process(
-	const FMeshBatch& MeshBatch, 
-	uint64 BatchElementMask, 
+	const FMeshBatch& MeshBatch,
+	uint64 BatchElementMask,
 	const FPrimitiveSceneProxy* RESTRICT PrimitiveSceneProxy,
 	int32 StaticMeshId,
 	const FMaterialRenderProxy& RESTRICT MaterialRenderProxy,
@@ -849,7 +849,7 @@ void FHairVisibilityProcessor::Process(
 	{
 		FVertexFactoryType* VertexFactoryType = VertexFactory->GetType();
 		PassShaders.VertexShader = MaterialResource.GetShader<FHairVisibilityVS<TRenderMode>>(VertexFactoryType);
-		PassShaders.PixelShader  = MaterialResource.GetShader<FHairVisibilityPS<TRenderMode>>(VertexFactoryType);
+		PassShaders.PixelShader = MaterialResource.GetShader<FHairVisibilityPS<TRenderMode>>(VertexFactoryType);
 	}
 
 	FMeshPassProcessorRenderState DrawRenderState(PassDrawRenderState);
@@ -893,7 +893,7 @@ class FClearUIntGraphicPS : public FGlobalShader
 	BEGIN_SHADER_PARAMETER_STRUCT(FParameters, )
 		SHADER_PARAMETER(uint32, ClearValue)
 		RENDER_TARGET_BINDING_SLOTS()
-	END_SHADER_PARAMETER_STRUCT()
+		END_SHADER_PARAMETER_STRUCT()
 
 public:
 	static bool ShouldCompilePermutation(const FGlobalShaderPermutationParameters& Parameters) { return IsHairStrandsSupported(Parameters.Platform); }
@@ -946,9 +946,9 @@ static void AddClearGraphicPass(
 		DrawRectangle(
 			RHICmdList,
 			0, 0,
-			Viewport.Width(),Viewport.Height(),
-			Viewport.Min.X,  Viewport.Min.Y,
-			Viewport.Width(),Viewport.Height(),
+			Viewport.Width(), Viewport.Height(),
+			Viewport.Min.X, Viewport.Min.Y,
+			Viewport.Width(), Viewport.Height(),
 			Viewport.Size(),
 			Resolution,
 			*VertexShader,
@@ -957,15 +957,15 @@ static void AddClearGraphicPass(
 }
 
 BEGIN_SHADER_PARAMETER_STRUCT(FClearUAVTextureParameters, )
-	SHADER_PARAMETER_RDG_TEXTURE_UAV(RWTexture2D, TextureUAV)
+SHADER_PARAMETER_RDG_TEXTURE_UAV(RWTexture2D, TextureUAV)
 END_SHADER_PARAMETER_STRUCT()
 
-void AddClearUAVPass(
+void AddClearUAVUintPass(
 	FRDGBuilder& GraphBuilder,
 	FRDGEventName&& PassName,
 	FRDGTextureRef Texture,
 	uint32 Value)
-{	
+{
 	FClearUAVTextureParameters* Parameters = GraphBuilder.AllocParameters< FClearUAVTextureParameters >();
 	Parameters->TextureUAV = GraphBuilder.CreateUAV(Texture);
 
@@ -975,9 +975,9 @@ void AddClearUAVPass(
 		ERDGPassFlags::Compute,
 		[Parameters, Texture, Value](FRHICommandList& RHICmdList)
 	{
-		uint32 ClearValue[4] = { Value, Value, Value, Value };
+		FUintVector4 ClearValue(Value, Value, Value, Value);
 		FRHIUnorderedAccessView* GlobalCounterUAV = Parameters->TextureUAV->GetRHI();
-		RHICmdList.ClearTinyUAV(GlobalCounterUAV, ClearValue);
+		RHICmdList.ClearUAVUint(GlobalCounterUAV, ClearValue);
 	});
 }
 
@@ -993,7 +993,7 @@ class FCopyIndirectBufferCS : public FGlobalShader
 		SHADER_PARAMETER(uint32, ItemCountPerGroup)
 		SHADER_PARAMETER_RDG_TEXTURE(Texture2D, CounterTexture)
 		SHADER_PARAMETER_RDG_BUFFER_UAV(RWBuffer, OutArgBuffer)
-	END_SHADER_PARAMETER_STRUCT()
+		END_SHADER_PARAMETER_STRUCT()
 
 public:
 	static bool ShouldCompilePermutation(const FGlobalShaderPermutationParameters& Parameters) { return IsHairStrandsSupported(Parameters.Platform); }
@@ -1025,7 +1025,7 @@ static FRDGBufferRef AddCopyIndirectArgPass(
 		RDG_EVENT_NAME("HairStrandsVisbilityCopyIndirectArgs"),
 		*ComputeShader,
 		Parameters,
-		FIntVector(1,1,1));
+		FIntVector(1, 1, 1));
 
 	return OutBuffer;
 }
@@ -1036,12 +1036,12 @@ class FHairVisibilityPrimitiveIdCompactionCS : public FGlobalShader
 	DECLARE_GLOBAL_SHADER(FHairVisibilityPrimitiveIdCompactionCS);
 	SHADER_USE_PARAMETER_STRUCT(FHairVisibilityPrimitiveIdCompactionCS, FGlobalShader);
 
-	class FVendor		: SHADER_PERMUTATION_INT("PERMUTATION_VENDOR", HairVisibilityVendorCount);
-	class FVelocity		: SHADER_PERMUTATION_INT("PERMUTATION_VELOCITY", 4);
+	class FVendor : SHADER_PERMUTATION_INT("PERMUTATION_VENDOR", HairVisibilityVendorCount);
+	class FVelocity : SHADER_PERMUTATION_INT("PERMUTATION_VELOCITY", 4);
 	class FViewTransmittance : SHADER_PERMUTATION_INT("PERMUTATION_VIEWTRANSMITTANCE", 2);
-	class FMaterial		: SHADER_PERMUTATION_INT("PERMUTATION_MATERIAL_COMPACTION", 2);
-	class FPPLL			: SHADER_PERMUTATION_SPARSE_INT("PERMUTATION_PPLL", 0, 8, 16, 32); // See GetPPLLMaxRenderNodePerPixel
-	class FVisibility	: SHADER_PERMUTATION_INT("PERMUTATION_VISIBILITY", 2);
+	class FMaterial : SHADER_PERMUTATION_INT("PERMUTATION_MATERIAL_COMPACTION", 2);
+	class FPPLL : SHADER_PERMUTATION_SPARSE_INT("PERMUTATION_PPLL", 0, 8, 16, 32); // See GetPPLLMaxRenderNodePerPixel
+	class FVisibility : SHADER_PERMUTATION_INT("PERMUTATION_VISIBILITY", 2);
 	using FPermutationDomain = TShaderPermutationDomain<FVendor, FVelocity, FViewTransmittance, FMaterial, FPPLL, FVisibility>;
 
 	BEGIN_SHADER_PARAMETER_STRUCT(FParameters, )
@@ -1064,7 +1064,7 @@ class FHairVisibilityPrimitiveIdCompactionCS : public FGlobalShader
 		SHADER_PARAMETER_RDG_BUFFER_SRV(StructuredBuffer, PPLLNodeData)
 
 		SHADER_PARAMETER_RDG_TEXTURE(Texture2D, ViewTransmittanceTexture)
-		
+
 		SHADER_PARAMETER_RDG_TEXTURE_UAV(Texture2D, OutCompactNodeCounter)
 		SHADER_PARAMETER_RDG_TEXTURE_UAV(Texture2D, OutCompactNodeIndex)
 		SHADER_PARAMETER_RDG_TEXTURE_UAV(Texture2D, OutCategorizationTexture)
@@ -1074,7 +1074,7 @@ class FHairVisibilityPrimitiveIdCompactionCS : public FGlobalShader
 
 		SHADER_PARAMETER_STRUCT_REF(FViewUniformShaderParameters, ViewUniformBuffer)
 		SHADER_PARAMETER_STRUCT_REF(FSceneTexturesUniformParameters, SceneTexturesStruct)
-	END_SHADER_PARAMETER_STRUCT()
+		END_SHADER_PARAMETER_STRUCT()
 
 public:
 
@@ -1083,7 +1083,7 @@ public:
 		if (PermutationVector.Get<FPPLL>() > 0)
 		{
 			PermutationVector.Set<FViewTransmittance>(0);
-			PermutationVector.Set<FVisibility>(0); 
+			PermutationVector.Set<FVisibility>(0);
 		}
 		return PermutationVector;
 	}
@@ -1146,7 +1146,7 @@ static void AddHairVisibilityPrimitiveIdCompactionPass(
 		}
 	}
 
-	
+
 	FRDGTextureRef CompactCounter;
 	{
 		FRDGTextureDesc Desc;
@@ -1183,10 +1183,10 @@ static void AddHairVisibilityPrimitiveIdCompactionPass(
 		OutputDesc.TargetableFlags = TexCreate_UAV | TexCreate_ShaderResource;
 		OutCategorizationTexture = GraphBuilder.CreateTexture(OutputDesc, TEXT("CategorizationTexture"));
 	}
-	
-	AddClearUAVPass(GraphBuilder, RDG_EVENT_NAME("HairStrandsClearCompactionCounter"), CompactCounter, 0);
-	AddClearUAVPass(GraphBuilder, RDG_EVENT_NAME("HairStrandsClearCompactionOffsetAndCount"), OutCompactNodeIndex, 0);
-	AddClearUAVPass(GraphBuilder, RDG_EVENT_NAME("HairStrandsClearCategorizationTexture"), OutCategorizationTexture, 0);
+
+	AddClearUAVUintPass(GraphBuilder, RDG_EVENT_NAME("HairStrandsClearCompactionCounter"), CompactCounter, 0);
+	AddClearUAVUintPass(GraphBuilder, RDG_EVENT_NAME("HairStrandsClearCompactionOffsetAndCount"), OutCompactNodeIndex, 0);
+	AddClearUAVUintPass(GraphBuilder, RDG_EVENT_NAME("HairStrandsClearCategorizationTexture"), OutCategorizationTexture, 0);
 
 	// Select render node count according to current mode
 	const uint32 HairVisibilityMSAASampleCount = GetHairVisibilitySampleCount();
@@ -1232,7 +1232,7 @@ static void AddHairVisibilityPrimitiveIdCompactionPass(
 	{
 		PassParameters->OutVelocityTexture = GraphBuilder.CreateUAV(OutVelocityTexture);
 	}
- 
+
 	FIntRect TotalRect = ComputeVisibleHairStrandsMacroGroupsRect(View.ViewRect, MacroGroupDatas);
 
 	// Snap the rect onto thread group boundary
@@ -1241,7 +1241,7 @@ static void AddHairVisibilityPrimitiveIdCompactionPass(
 	TotalRect.Min.Y = FMath::FloorToInt(float(TotalRect.Min.Y) / float(GroupSize.Y)) * GroupSize.Y;
 	TotalRect.Max.X = FMath::CeilToInt(float(TotalRect.Max.X) / float(GroupSize.X)) * GroupSize.X;
 	TotalRect.Max.Y = FMath::CeilToInt(float(TotalRect.Max.Y) / float(GroupSize.Y)) * GroupSize.Y;
-	
+
 	FIntPoint RectResolution(TotalRect.Width(), TotalRect.Height());
 	PassParameters->ResolutionOffset = FIntPoint(TotalRect.Min.X, TotalRect.Min.Y);
 
@@ -1273,10 +1273,10 @@ class FHairGenerateTileCS : public FGlobalShader
 		SHADER_PARAMETER_RDG_TEXTURE_UAV(Texture2D, OutTileCounter)
 		SHADER_PARAMETER_RDG_TEXTURE_UAV(RWTexture2D, OutTileIndexTexture)
 		SHADER_PARAMETER_RDG_BUFFER_UAV(RWBuffer, OutTileBuffer)
-	END_SHADER_PARAMETER_STRUCT()
+		END_SHADER_PARAMETER_STRUCT()
 
 public:
-	
+
 	static bool ShouldCompilePermutation(const FGlobalShaderPermutationParameters& Parameters)
 	{
 		return IsHairStrandsSupported(Parameters.Platform);
@@ -1327,10 +1327,11 @@ static void AddGenerateTilePass(
 		Desc.ClearValue = FClearValueBinding(0);
 		OutTileIndexTexture = GraphBuilder.CreateTexture(Desc, TEXT("HairTileIndexTexture"));
 	}
-	
+
 	OutTileBuffer = GraphBuilder.CreateBuffer(FRDGBufferDesc::CreateBufferDesc(sizeof(uint32), TileCount), TEXT("HairTileBuffer"));
 
-	AddClearUAVPass(GraphBuilder, RDG_EVENT_NAME("HairTileCounter"), TileCounter, 0);
+	uint32 ClearValues[4] = { 0,0,0,0 };
+	AddClearUAVPass(GraphBuilder, GraphBuilder.CreateUAV(TileCounter), ClearValues);
 
 	FHairGenerateTileCS::FPermutationDomain PermutationVector;
 	PermutationVector.Set<FHairGenerateTileCS::FTileSize>(0);
@@ -1338,7 +1339,7 @@ static void AddGenerateTilePass(
 	FHairGenerateTileCS::FParameters* PassParameters = GraphBuilder.AllocParameters<FHairGenerateTileCS::FParameters>();
 	PassParameters->Resolution = Resolution;
 	PassParameters->TileResolution = TileResolution;
-	PassParameters->CategorizationTexture = CategorizationTexture;	
+	PassParameters->CategorizationTexture = CategorizationTexture;
 	PassParameters->OutTileCounter = GraphBuilder.CreateUAV(TileCounter);
 	PassParameters->OutTileIndexTexture = GraphBuilder.CreateUAV(OutTileIndexTexture);
 	PassParameters->OutTileBuffer = GraphBuilder.CreateUAV(OutTileBuffer, PF_R16G16_UINT);
@@ -1367,7 +1368,7 @@ class FHairVisibilityFillOpaqueDepthPS : public FGlobalShader
 
 		SHADER_PARAMETER_STRUCT_REF(FViewUniformShaderParameters, ViewUniformBuffer)
 		RENDER_TARGET_BINDING_SLOTS()
-	END_SHADER_PARAMETER_STRUCT()
+		END_SHADER_PARAMETER_STRUCT()
 
 public:
 	static bool ShouldCompilePermutation(const FGlobalShaderPermutationParameters& Parameters) { return IsHairStrandsSupported(Parameters.Platform); }
@@ -1459,7 +1460,7 @@ static FRDGTextureRef AddHairVisibilityFillOpaqueDepth(
 					RHICmdList,
 					0, 0,
 					Viewport.Width(), Viewport.Height(),
-					Viewport.Min.X,   Viewport.Min.Y,
+					Viewport.Min.X, Viewport.Min.Y,
 					Viewport.Width(), Viewport.Height(),
 					Viewport.Size(),
 					Resolution,
@@ -1499,7 +1500,7 @@ static void AddHairVisibilityCommonPass(
 	};
 
 	GraphBuilder.AddPass(
-		GetPassName(), 
+		GetPassName(),
 		PassParameters,
 		ERDGPassFlags::Raster,
 		[PassParameters, Scene = Scene, ViewInfo, &MacroGroupDatas, RenderMode](FRHICommandListImmediate& RHICmdList)
@@ -1509,7 +1510,7 @@ static void AddHairVisibilityCommonPass(
 
 		FMeshPassProcessorRenderState DrawRenderState(*ViewInfo);
 
-		if(RenderMode == HairVisibilityRenderMode_Transmittance || RenderMode == HairVisibilityRenderMode_TransmittanceAndHairCount || RenderMode == HairVisibilityRenderMode_PPLL)
+		if (RenderMode == HairVisibilityRenderMode_Transmittance || RenderMode == HairVisibilityRenderMode_TransmittanceAndHairCount || RenderMode == HairVisibilityRenderMode_PPLL)
 		{
 			TUniformBufferRef<FViewUniformShaderParameters> ViewUniformShaderParameters;
 			const bool bEnableMSAA = false;
@@ -1518,7 +1519,7 @@ static void AddHairVisibilityCommonPass(
 			ViewUniformShaderParameters = TUniformBufferRef<FViewUniformShaderParameters>::CreateUniformBufferImmediate(*ViewInfo->CachedViewUniformShaderParameters, UniformBuffer_SingleFrame);
 			DrawRenderState.SetViewUniformBuffer(ViewUniformShaderParameters);
 		}
-		
+
 		{
 			RHICmdList.SetViewport(0, 0, 0.0f, ViewInfo->ViewRect.Width(), ViewInfo->ViewRect.Height(), 1.0f);
 			if (RenderMode == HairVisibilityRenderMode_MSAA)
@@ -1528,7 +1529,7 @@ static void AddHairVisibilityCommonPass(
 					CW_RGBA, BO_Add, BF_One, BF_Zero, BO_Add, BF_One, BF_Zero>::GetRHI());
 				DrawRenderState.SetDepthStencilState(TStaticDepthStencilState<true, CF_DepthNearOrEqual>::GetRHI());
 			}
-			else if(RenderMode == HairVisibilityRenderMode_MSAA_Visibility)
+			else if (RenderMode == HairVisibilityRenderMode_MSAA_Visibility)
 			{
 				DrawRenderState.SetBlendState(TStaticBlendState<CW_RGBA, BO_Add, BF_One, BF_Zero, BO_Add, BF_One, BF_Zero>::GetRHI());
 				DrawRenderState.SetDepthStencilState(TStaticDepthStencilState<true, CF_DepthNearOrEqual>::GetRHI());
@@ -1704,9 +1705,9 @@ static void AddHairVisibilityMSAAPass(
 
 		FVisibilityPassParameters* PassParameters = GraphBuilder.AllocParameters<FVisibilityPassParameters>();
 		PassParameters->RenderTargets[0] = FRenderTargetBinding(OutVisibilityIdTexture, ERenderTargetLoadAction::ELoad, 0);
-		PassParameters->RenderTargets[1] = FRenderTargetBinding(OutVisibilityMaterialTexture,  LoadAction, 0);
+		PassParameters->RenderTargets[1] = FRenderTargetBinding(OutVisibilityMaterialTexture, LoadAction, 0);
 		PassParameters->RenderTargets[2] = FRenderTargetBinding(OutVisibilityAttributeTexture, LoadAction, 0);
-		PassParameters->RenderTargets[3] = FRenderTargetBinding(OutVisibilityVelocityTexture,  LoadAction, 0);
+		PassParameters->RenderTargets[3] = FRenderTargetBinding(OutVisibilityVelocityTexture, LoadAction, 0);
 
 		PassParameters->RenderTargets.DepthStencil = FDepthStencilBinding(
 			OutVisibilityDepthTexture,
@@ -1772,8 +1773,8 @@ static void AddHairVisibilityPPLLPass(
 
 		OutVisibilityPPLLNodeData = GraphBuilder.CreateBuffer(FRDGBufferDesc::CreateStructuredDesc(sizeof(PPLLNodeData), PPLLMaxTotalListElementCount), TEXT("HairVisibilityPPLLNodeData"));
 	}
-	AddClearUAVPass(GraphBuilder, RDG_EVENT_NAME("ClearHairVisibilityPPLLCounter"), OutVisibilityPPLLNodeCounter, 0);
-	AddClearUAVPass(GraphBuilder, RDG_EVENT_NAME("ClearHairVisibilityPPLLNodeIndex"), OutVisibilityPPLLNodeIndex, 0xFFFFFFFF);
+	AddClearUAVUintPass(GraphBuilder, RDG_EVENT_NAME("ClearHairVisibilityPPLLCounter"), OutVisibilityPPLLNodeCounter, 0);
+	AddClearUAVUintPass(GraphBuilder, RDG_EVENT_NAME("ClearHairVisibilityPPLLNodeIndex"), OutVisibilityPPLLNodeIndex, 0xFFFFFFFF);
 
 	FVisibilityPassParameters* PassParameters = GraphBuilder.AllocParameters<FVisibilityPassParameters>();
 	PassParameters->PPLLCounter = GraphBuilder.CreateUAV(FRDGTextureUAVDesc(OutVisibilityPPLLNodeCounter, 0));
@@ -1849,10 +1850,10 @@ class FHairViewTransmittanceDepthPS : public FGlobalShader
 	BEGIN_SHADER_PARAMETER_STRUCT(FParameters, )
 		SHADER_PARAMETER(float, DistanceThreshold)
 		SHADER_PARAMETER_RDG_TEXTURE(Texture2D, SceneDepthTexture)
-		SHADER_PARAMETER_RDG_TEXTURE(Texture2D, CategorizationTexture)		
+		SHADER_PARAMETER_RDG_TEXTURE(Texture2D, CategorizationTexture)
 		SHADER_PARAMETER_STRUCT_REF(FViewUniformShaderParameters, ViewUniformBuffer)
 		RENDER_TARGET_BINDING_SLOTS()
-	END_SHADER_PARAMETER_STRUCT()
+		END_SHADER_PARAMETER_STRUCT()
 
 public:
 
@@ -1877,7 +1878,7 @@ static void AddHairViewTransmittanceDepthPass(
 	Parameters->SceneDepthTexture = SceneDepthTexture;
 	Parameters->ViewUniformBuffer = View.ViewUniformBuffer;
 	Parameters->RenderTargets[0] = FRenderTargetBinding(HairCountTexture, ERenderTargetLoadAction::ELoad);
-	
+
 	TShaderMapRef<FPostProcessVS> VertexShader(View.ShaderMap);
 	TShaderMapRef<FHairViewTransmittanceDepthPS> PixelShader(View.ShaderMap);
 	const TShaderMap<FGlobalShaderType>* GlobalShaderMap = View.ShaderMap;
@@ -1932,10 +1933,10 @@ class FHairVisibilityDepthPS : public FGlobalShader
 	BEGIN_SHADER_PARAMETER_STRUCT(FParameters, )
 		SHADER_PARAMETER_RDG_TEXTURE(Texture2D, CategorisationTexture)
 		RENDER_TARGET_BINDING_SLOTS()
-	END_SHADER_PARAMETER_STRUCT()
+		END_SHADER_PARAMETER_STRUCT()
 
 public:
-		
+
 	static FPermutationDomain RemapPermutation(FPermutationDomain PermutationVector)
 	{
 		return PermutationVector;
@@ -2017,9 +2018,9 @@ static void AddHairVisibilityColorAndDepthPatchPass(
 ///////////////////////////////////////////////////////////////////////////////////////////////////
 
 FHairStrandsVisibilityViews RenderHairStrandsVisibilityBuffer(
-	FRHICommandListImmediate& RHICmdList, 
-	const FScene* Scene, 
-	const TArray<FViewInfo>& Views, 
+	FRHICommandListImmediate& RHICmdList,
+	const FScene* Scene,
+	const TArray<FViewInfo>& Views,
 	TRefCountPtr<IPooledRenderTarget> InSceneGBufferBTexture,
 	TRefCountPtr<IPooledRenderTarget> InSceneColorTexture,
 	TRefCountPtr<IPooledRenderTarget> InSceneDepthTexture,
@@ -2056,7 +2057,7 @@ FHairStrandsVisibilityViews RenderHairStrandsVisibilityBuffer(
 			FRDGTextureRef SceneColorTexture = GraphBuilder.RegisterExternalTexture(InSceneColorTexture, TEXT("SceneColorTexture"));
 			FRDGTextureRef SceneDepthTexture = GraphBuilder.RegisterExternalTexture(InSceneDepthTexture, TEXT("SceneDepthTexture"));
 			FRDGTextureRef SceneVelocityTexture = InSceneVelocityTexture ? GraphBuilder.RegisterExternalTexture(InSceneVelocityTexture, TEXT("SceneVelocityTexture")) : nullptr;
-			
+
 			const EHairVisibilityRenderMode RenderMode = GetHairVisibilityRenderMode();
 
 			// Run the view transmittance pass if needed (not in PPLL mode that is already a high quality render path)
@@ -2069,7 +2070,7 @@ FHairStrandsVisibilityViews RenderHairStrandsVisibilityBuffer(
 					Scene,
 					&View,
 					MacroGroupDatas,
-					Resolution, 
+					Resolution,
 					bOutputHairCount,
 					SceneDepthTexture);
 			}
@@ -2097,16 +2098,16 @@ FHairStrandsVisibilityViews RenderHairStrandsVisibilityBuffer(
 
 				AddHairVisibilityMSAAPass(
 					bIsVisiblityEnable,
-					GraphBuilder, 
-					Scene, 
-					&View, 
+					GraphBuilder,
+					Scene,
+					&View,
 					MacroGroupDatas,
 					Resolution,
-					MsaaVisibilityResources.IdTexture, 
-					MsaaVisibilityResources.MaterialTexture, 
+					MsaaVisibilityResources.IdTexture,
+					MsaaVisibilityResources.MaterialTexture,
 					MsaaVisibilityResources.AttributeTexture,
 					MsaaVisibilityResources.VelocityTexture,
-					MsaaVisibilityResources.DepthTexture); 
+					MsaaVisibilityResources.DepthTexture);
 
 				// This is used when compaction is not enabled.
 				GraphBuilder.QueueTextureExtraction(MsaaVisibilityResources.IdTexture, &VisibilityData.IDTexture);
@@ -2174,11 +2175,11 @@ FHairStrandsVisibilityViews RenderHairStrandsVisibilityBuffer(
 						CompactNodeData = PassOutput.NodeData;
 					}
 
-					GraphBuilder.QueueTextureExtraction(CompactNodeIndex,		&VisibilityData.NodeIndex);
-					GraphBuilder.QueueTextureExtraction(CategorizationTexture,	&VisibilityData.CategorizationTexture);
-					GraphBuilder.QueueBufferExtraction(CompactNodeData,			&VisibilityData.NodeData,					FRDGResourceState::EAccess::Read, FRDGResourceState::EPipeline::Graphics);
-					GraphBuilder.QueueBufferExtraction(CompactNodeCoord,		&VisibilityData.NodeCoord,					FRDGResourceState::EAccess::Read, FRDGResourceState::EPipeline::Graphics);
-					GraphBuilder.QueueBufferExtraction(IndirectArgsBuffer,		&VisibilityData.NodeIndirectArg,			FRDGResourceState::EAccess::Read, FRDGResourceState::EPipeline::Compute);
+					GraphBuilder.QueueTextureExtraction(CompactNodeIndex, &VisibilityData.NodeIndex);
+					GraphBuilder.QueueTextureExtraction(CategorizationTexture, &VisibilityData.CategorizationTexture);
+					GraphBuilder.QueueBufferExtraction(CompactNodeData, &VisibilityData.NodeData, FRDGResourceState::EAccess::Read, FRDGResourceState::EPipeline::Graphics);
+					GraphBuilder.QueueBufferExtraction(CompactNodeCoord, &VisibilityData.NodeCoord, FRDGResourceState::EAccess::Read, FRDGResourceState::EPipeline::Graphics);
+					GraphBuilder.QueueBufferExtraction(IndirectArgsBuffer, &VisibilityData.NodeIndirectArg, FRDGResourceState::EAccess::Read, FRDGResourceState::EPipeline::Compute);
 				}
 
 				// View transmittance depth test needs to happen before the scene depth is patched with the hair depth (for fully-covered-by-hair pixels)
@@ -2223,8 +2224,8 @@ FHairStrandsVisibilityViews RenderHairStrandsVisibilityBuffer(
 				// Linked list sorting pass and compaction into common representation
 				{
 					FHairVisibilityPrimitiveIdCompactionCS::FParameters* PassParameters = GraphBuilder.AllocParameters<FHairVisibilityPrimitiveIdCompactionCS::FParameters>();
-					PassParameters->PPLLCounter  = PPLLNodeCounterTexture;
-					PassParameters->PPLLNodeIndex= PPLLNodeIndexTexture;
+					PassParameters->PPLLCounter = PPLLNodeCounterTexture;
+					PassParameters->PPLLNodeIndex = PPLLNodeIndexTexture;
 					PassParameters->PPLLNodeData = GraphBuilder.CreateSRV(PPLLNodeDataBuffer);
 					PassParameters->ViewTransmittanceTexture = ViewTransmittance.TransmittanceTexture;
 
