Fix CIS static analysis issue.

diff --git a/Engine/Source/Runtime/Renderer/Private/CapsuleShadowRendering.cpp b/Engine/Source/Runtime/Renderer/Private/CapsuleShadowRendering.cpp
--- a/Engine/Source/Runtime/Renderer/Private/CapsuleShadowRendering.cpp
+++ b/Engine/Source/Runtime/Renderer/Private/CapsuleShadowRendering.cpp
@@ -1643,7 +1643,6 @@ void FDeferredShadingSceneRenderer::RenderCapsuleShadowsForMovableSkylight(FRHIC
 									}
 									else
 									{
-										check(NumMeshDistanceFieldCasters > 0);
 										TShaderMapRef<TCapsuleShadowingCS<ShapeShadow_MovableSkylightTiledCulling, IPT_MeshDistanceFields> > ComputeShader(View.ShaderMap);
 										ComputeShaderBase = *ComputeShader;
 									}
