MaskedOnly Prepass Mobile Off by default. Enable with r.Mobile.EarlyZPassOnlyMaterialMasking=1. Requires shaders cooking!

diff --git a/Engine/Shaders/Private/MobileBasePassPixelShader.usf b/Engine/Shaders/Private/MobileBasePassPixelShader.usf
--- a/Engine/Shaders/Private/MobileBasePassPixelShader.usf
+++ b/Engine/Shaders/Private/MobileBasePassPixelShader.usf
@@ -458,8 +458,10 @@ void Main(
 	ApplyPixelDepthOffsetForMobileBasePass(MaterialParameters, PixelMaterialInputs, OutDepth);
 #endif
 	  
+#if !EARLY_Z_PASS_ONLY_MATERIAL_MASKING
 	//Clip if the blend mode requires it.
 	GetMaterialCoverageAndClipping(MaterialParameters, PixelMaterialInputs);
+#endif
 
 	half Opacity = GetMaterialOpacity(PixelMaterialInputs);
 	const float BaseMaterialCoverageOverWater = Opacity;
diff --git a/Engine/Source/Runtime/Engine/Private/ShaderCompiler/ShaderCompiler.cpp b/Engine/Source/Runtime/Engine/Private/ShaderCompiler/ShaderCompiler.cpp
--- a/Engine/Source/Runtime/Engine/Private/ShaderCompiler/ShaderCompiler.cpp
+++ b/Engine/Source/Runtime/Engine/Private/ShaderCompiler/ShaderCompiler.cpp
@@ -3657,8 +3657,14 @@ void GlobalBeginCompileShader(
 	}
 
 	{
-		static IConsoleVariable* CVar = IConsoleManager::Get().FindConsoleVariable(TEXT("r.EarlyZPassOnlyMaterialMasking"));
-		Input.Environment.SetDefine(TEXT("EARLY_Z_PASS_ONLY_MATERIAL_MASKING"), CVar ? (CVar->GetInt() != 0) : 0);
+		if (MaskedInEarlyPass((EShaderPlatform)Target.Platform))
+		{
+			Input.Environment.SetDefine(TEXT("EARLY_Z_PASS_ONLY_MATERIAL_MASKING"), 1);
+		}
+		else
+		{
+			Input.Environment.SetDefine(TEXT("EARLY_Z_PASS_ONLY_MATERIAL_MASKING"), 0);
+		}
 	}
 
 	{
diff --git a/Engine/Source/Runtime/RenderCore/Private/RenderUtils.cpp b/Engine/Source/Runtime/RenderCore/Private/RenderUtils.cpp
--- a/Engine/Source/Runtime/RenderCore/Private/RenderUtils.cpp
+++ b/Engine/Source/Runtime/RenderCore/Private/RenderUtils.cpp
@@ -1045,6 +1045,20 @@ RENDERCORE_API bool GPUSceneUseTexture2D(EShaderPlatform Platform)
 	return false;
 }
 
+RENDERCORE_API bool MaskedInEarlyPass(EShaderPlatform Platform)
+{
+	static IConsoleVariable* CVarMobileEarlyZPassOnlyMaterialMasking = IConsoleManager::Get().FindConsoleVariable(TEXT("r.Mobile.EarlyZPassOnlyMaterialMasking"));
+	static IConsoleVariable* CVarEarlyZPassOnlyMaterialMasking = IConsoleManager::Get().FindConsoleVariable(TEXT("r.EarlyZPassOnlyMaterialMasking"));
+	if (IsMobilePlatform(Platform))
+	{
+		return (CVarMobileEarlyZPassOnlyMaterialMasking && CVarMobileEarlyZPassOnlyMaterialMasking->GetInt() != 0);
+	}
+	else
+	{
+		return (CVarEarlyZPassOnlyMaterialMasking && CVarEarlyZPassOnlyMaterialMasking->GetInt() != 0);
+	}
+}
+
 RENDERCORE_API bool AllowPixelDepthOffset(EShaderPlatform Platform)
 {
 	if (IsMobilePlatform(Platform))
diff --git a/Engine/Source/Runtime/RenderCore/Private/Shader.cpp b/Engine/Source/Runtime/RenderCore/Private/Shader.cpp
--- a/Engine/Source/Runtime/RenderCore/Private/Shader.cpp
+++ b/Engine/Source/Runtime/RenderCore/Private/Shader.cpp
@@ -2511,8 +2511,7 @@ void ShaderMapAppendKeyString(EShaderPlatform Platform, FString& KeyString)
 	}
 
 	{
-		static const auto CVar = IConsoleManager::Get().FindTConsoleVariableDataInt(TEXT("r.EarlyZPassOnlyMaterialMasking"));
-		if (CVar && CVar->GetValueOnAnyThread() > 0)
+		if (MaskedInEarlyPass(Platform))
 		{
 			KeyString += TEXT("_EZPMM");
 		}
diff --git a/Engine/Source/Runtime/RenderCore/Public/RenderUtils.h b/Engine/Source/Runtime/RenderCore/Public/RenderUtils.h
--- a/Engine/Source/Runtime/RenderCore/Public/RenderUtils.h
+++ b/Engine/Source/Runtime/RenderCore/Public/RenderUtils.h
@@ -468,6 +468,8 @@ RENDERCORE_API bool MobileSupportsGPUScene(EShaderPlatform Platform);
 
 RENDERCORE_API bool GPUSceneUseTexture2D(EShaderPlatform Platform);
 
+RENDERCORE_API bool MaskedInEarlyPass(EShaderPlatform Platform);
+
 RENDERCORE_API bool AllowPixelDepthOffset(EShaderPlatform Platform);
 
 /** Returns if ForwardShading is enabled. Only valid for the current platform (otherwise call ITargetPlatform::UsesForwardShading()). */
diff --git a/Engine/Source/Runtime/Renderer/Private/BasePassRendering.cpp b/Engine/Source/Runtime/Renderer/Private/BasePassRendering.cpp
--- a/Engine/Source/Runtime/Renderer/Private/BasePassRendering.cpp
+++ b/Engine/Source/Runtime/Renderer/Private/BasePassRendering.cpp
@@ -239,8 +239,7 @@ void SetDepthStencilStateForBasePass(
 	const FPrimitiveSceneProxy* PrimitiveSceneProxy,
 	bool bEnableReceiveDecalOutput)
 {
-	static IConsoleVariable* EarlyZPassOnlyMaterialMaskingCVar = IConsoleManager::Get().FindConsoleVariable(TEXT("r.EarlyZPassOnlyMaterialMasking"));
-	const bool bMaskInEarlyPass = (Mesh.MaterialRenderProxy->GetMaterial(FeatureLevel)->IsMasked() || Mesh.bDitheredLODTransition) && EarlyZPassOnlyMaterialMaskingCVar->GetInt();
+	const bool bMaskedInEarlyPass = (Mesh.MaterialRenderProxy->GetMaterial(FeatureLevel)->IsMasked() || Mesh.bDitheredLODTransition) && MaskedInEarlyPass(GShaderPlatformForFeatureLevel[FeatureLevel]);
 
 	if (bEnableReceiveDecalOutput)
 	{
@@ -249,7 +248,7 @@ void SetDepthStencilStateForBasePass(
 		const uint8 StencilValue = GET_STENCIL_BIT_MASK(RECEIVE_DECAL, PrimitiveSceneProxy ? !!PrimitiveSceneProxy->ReceivesDecals() : 0x00)
 			| STENCIL_LIGHTING_CHANNELS_MASK(PrimitiveSceneProxy ? PrimitiveSceneProxy->GetLightingChannelStencilValue() : 0x00);
 
-		if (bMaskInEarlyPass)
+		if (bMaskedInEarlyPass)
 		{
 			DrawRenderState.SetDepthStencilState(TStaticDepthStencilState<
 				false, CF_Equal,
@@ -280,7 +279,7 @@ void SetDepthStencilStateForBasePass(
 			DrawRenderState.SetStencilRef(StencilValue);
 		}
 	}
-	else if (bMaskInEarlyPass)
+	else if (bMaskedInEarlyPass)
 	{
 		DrawRenderState.SetDepthStencilState(TStaticDepthStencilState<false, CF_Equal>::GetRHI());
 	}
diff --git a/Engine/Source/Runtime/Renderer/Private/DepthRendering.cpp b/Engine/Source/Runtime/Renderer/Private/DepthRendering.cpp
--- a/Engine/Source/Runtime/Renderer/Private/DepthRendering.cpp
+++ b/Engine/Source/Runtime/Renderer/Private/DepthRendering.cpp
@@ -774,6 +774,49 @@ bool FDeferredShadingSceneRenderer::RenderPrePass(FRHICommandListImmediate& RHIC
 	return bDepthWasCleared;
 }
 
+void FMobileSceneRenderer::RenderPrePass(FRHICommandListImmediate& RHICmdList)
+{
+	check(!RHICmdList.IsOutsideRenderPass());
+
+	SCOPED_NAMED_EVENT(FMobileSceneRenderer_RenderPrePass, FColor::Emerald);
+	SCOPED_DRAW_EVENT(RHICmdList, MobileRenderPrePass);
+
+	SCOPE_CYCLE_COUNTER(STAT_DepthDrawTime);
+	CSV_SCOPED_TIMING_STAT_EXCLUSIVE(RenderPrePass);
+	SCOPED_GPU_STAT(RHICmdList, Prepass);
+
+	// Draw a depth pass to avoid overdraw in the other passes.
+	// Mobile only does MaskedOnly DepthPass for the moment
+	if (Scene->EarlyZPassMode == DDM_MaskedOnly)
+	{
+		for (int32 ViewIndex = 0; ViewIndex < Views.Num(); ViewIndex++)
+		{
+			const FViewInfo& View = Views[ViewIndex];
+
+			SCOPED_GPU_MASK(RHICmdList, !View.IsInstancedStereoPass() ? View.GPUMask : (Views[0].GPUMask | Views[1].GPUMask));
+			SCOPED_CONDITIONAL_DRAW_EVENTF(RHICmdList, EventView, Views.Num() > 1, TEXT("View%d"), ViewIndex);
+
+			TUniformBufferRef<FSceneTexturesUniformParameters> PassUniformBuffer;
+			CreateDepthPassUniformBuffer(RHICmdList, View, PassUniformBuffer);
+
+			FMeshPassProcessorRenderState DrawRenderState(View, PassUniformBuffer);
+
+			SetupDepthPassState(DrawRenderState);
+
+			if (View.ShouldRenderView())
+			{
+				Scene->UniformBuffers.UpdateViewUniformBuffer(View);
+
+				{
+					SetupPrePassView(RHICmdList, View, this);
+
+					View.ParallelMeshDrawCommandPasses[EMeshPass::DepthPass].DispatchDraw(nullptr, RHICmdList);
+				}
+			}
+		}
+	}
+}
+
 extern bool IsHMDHiddenAreaMaskActive();
 
 bool FDeferredShadingSceneRenderer::RenderPrePassHMD(FRHICommandListImmediate& RHICmdList)
@@ -937,6 +980,7 @@ void FDepthPassMeshProcessor::AddMeshBatch(const FMeshBatch& RESTRICT MeshBatch,
 			&& ShouldIncludeMaterialInDefaultOpaquePass(Material))
 		{
 			if (BlendMode == BLEND_Opaque
+				&& EarlyZPassMode != DDM_MaskedOnly
 				&& MeshBatch.VertexFactory->SupportsPositionOnlyStream()
 				&& !Material.MaterialModifiesMeshPosition_RenderThread()
 				&& Material.WritesEveryPixel())
@@ -948,8 +992,7 @@ void FDepthPassMeshProcessor::AddMeshBatch(const FMeshBatch& RESTRICT MeshBatch,
 			else
 			{
 				const bool bMaterialMasked = !Material.WritesEveryPixel() || Material.IsTranslucencyWritingCustomDepth();
-
-				if (!bMaterialMasked || EarlyZPassMode != DDM_NonMaskedOnly)
+				if((!bMaterialMasked && EarlyZPassMode != DDM_MaskedOnly) || (bMaterialMasked && EarlyZPassMode != DDM_NonMaskedOnly))
 				{
 					const FMaterialRenderProxy* EffectiveMaterialRenderProxy = &MaterialRenderProxy;
 					const FMaterial* EffectiveMaterial = &Material;
@@ -993,4 +1036,5 @@ FMeshPassProcessor* CreateDepthPassProcessor(const FScene* Scene, const FSceneVi
 	return new(FMemStack::Get()) FDepthPassMeshProcessor(Scene, InViewIfDynamicMeshCommand, DepthPassState, true, Scene->EarlyZPassMode, Scene->bEarlyZPassMovable, InDrawListContext);
 }
 
-FRegisterPassProcessorCreateFunction RegisterDepthPass(&CreateDepthPassProcessor, EShadingPath::Deferred, EMeshPass::DepthPass, EMeshPassFlags::CachedMeshCommands | EMeshPassFlags::MainView);
\ No newline at end of file
+FRegisterPassProcessorCreateFunction RegisterDepthPass(&CreateDepthPassProcessor, EShadingPath::Deferred, EMeshPass::DepthPass, EMeshPassFlags::CachedMeshCommands | EMeshPassFlags::MainView);
+FRegisterPassProcessorCreateFunction RegisterMobileDepthPass(&CreateDepthPassProcessor, EShadingPath::Mobile, EMeshPass::DepthPass, EMeshPassFlags::CachedMeshCommands | EMeshPassFlags::MainView);
\ No newline at end of file
diff --git a/Engine/Source/Runtime/Renderer/Private/DepthRendering.h b/Engine/Source/Runtime/Renderer/Private/DepthRendering.h
--- a/Engine/Source/Runtime/Renderer/Private/DepthRendering.h
+++ b/Engine/Source/Runtime/Renderer/Private/DepthRendering.h
@@ -27,6 +27,8 @@ enum EDepthDrawingMode
 	DDM_AllOccluders	= 2,
 	// Full prepass, every object must be drawn and every pixel must match the base pass depth
 	DDM_AllOpaque		= 3,
+	// Masked materials only
+	DDM_MaskedOnly = 4,
 };
 
 extern const TCHAR* GetDepthDrawingModeString(EDepthDrawingMode Mode);
diff --git a/Engine/Source/Runtime/Renderer/Private/MobileBasePass.cpp b/Engine/Source/Runtime/Renderer/Private/MobileBasePass.cpp
--- a/Engine/Source/Runtime/Renderer/Private/MobileBasePass.cpp
+++ b/Engine/Source/Runtime/Renderer/Private/MobileBasePass.cpp
@@ -630,11 +630,17 @@ void FMobileBasePassMeshProcessor::Process(
 		BasePassShaders.VertexShader, 
 		BasePassShaders.PixelShader);
 
+	const bool bMaskedInEarlyPass = (MeshBatch.MaterialRenderProxy->GetMaterial(FeatureLevel)->IsMasked() || MeshBatch.bDitheredLODTransition) && MaskedInEarlyPass(Scene->GetShaderPlatform());
+	
 	FMeshPassProcessorRenderState DrawRenderState(PassDrawRenderState);
 	if (bTranslucentBasePass)
 	{
 		MobileBasePass::SetTranslucentRenderState(DrawRenderState, MaterialResource);
 	}
+	else if (bMaskedInEarlyPass)
+	{
+		DrawRenderState.SetDepthStencilState(TStaticDepthStencilState<false, CF_Equal>::GetRHI());
+	}
 	else
 	{
 		MobileBasePass::SetOpaqueRenderState(DrawRenderState, PrimitiveSceneProxy, MaterialResource, bEnableReceiveDecalOutput && IsMobileHDR());
diff --git a/Engine/Source/Runtime/Renderer/Private/MobileShadingRenderer.cpp b/Engine/Source/Runtime/Renderer/Private/MobileShadingRenderer.cpp
--- a/Engine/Source/Runtime/Renderer/Private/MobileShadingRenderer.cpp
+++ b/Engine/Source/Runtime/Renderer/Private/MobileShadingRenderer.cpp
@@ -85,6 +85,7 @@ DECLARE_CYCLE_STAT(TEXT("Post"), STAT_CLMM_Post, STATGROUP_CommandListMarkers);
 DECLARE_CYCLE_STAT(TEXT("Translucency"), STAT_CLMM_Translucency, STATGROUP_CommandListMarkers);
 DECLARE_CYCLE_STAT(TEXT("Shadows"), STAT_CLMM_Shadows, STATGROUP_CommandListMarkers);
 DECLARE_CYCLE_STAT(TEXT("SceneSimulation"), STAT_CLMM_SceneSim, STATGROUP_CommandListMarkers);
+DECLARE_CYCLE_STAT(TEXT("PrePass"), STAT_CLM_PrePass, STATGROUP_CommandListMarkers);
 
 FGlobalDynamicIndexBuffer FMobileSceneRenderer::DynamicIndexBuffer;
 FGlobalDynamicVertexBuffer FMobileSceneRenderer::DynamicVertexBuffer;
@@ -579,6 +580,9 @@ void FMobileSceneRenderer::Render(FRHICommandListImmediate& RHICmdList)
 	
 	RHICmdList.BeginRenderPass(SceneColorRenderPassInfo, TEXT("SceneColorRendering"));
 
+	RHICmdList.SetCurrentStat(GET_STATID(STAT_CLM_PrePass));
+	RenderPrePass(RHICmdList);
+
 	if (GIsEditor && !View.bIsSceneCapture)
 	{
 		DrawClearQuad(RHICmdList, Views[0].BackgroundColor);
diff --git a/Engine/Source/Runtime/Renderer/Private/RendererScene.cpp b/Engine/Source/Runtime/Renderer/Private/RendererScene.cpp
--- a/Engine/Source/Runtime/Renderer/Private/RendererScene.cpp
+++ b/Engine/Source/Runtime/Renderer/Private/RendererScene.cpp
@@ -77,6 +77,16 @@ static TAutoConsoleVariable<int32> CVarEarlyZPassOnlyMaterialMasking(
 	ECVF_RenderThreadSafe | ECVF_ReadOnly
 );
 
+/** Affects MobileBasePassPixelShader.usf so must relaunch editor to recompile shaders. */
+static TAutoConsoleVariable<int32> CVarMobileEarlyZPassOnlyMaterialMasking(
+	TEXT("r.Mobile.EarlyZPassOnlyMaterialMasking"),
+	0,
+	TEXT("Whether to compute materials' mask opacity only in early Z pass for Mobile platform. Changing this setting requires restarting the editor.\n")
+	TEXT("<=0: off\n")
+	TEXT("  1: on\n"),
+	ECVF_RenderThreadSafe | ECVF_ReadOnly
+);
+
 TAutoConsoleVariable<int32> CVarEarlyZPass(
 	TEXT("r.EarlyZPass"),
 	3,
@@ -3203,7 +3213,7 @@ bool ShouldForceFullDepthPass(EShaderPlatform ShaderPlatform)
 	static const auto AOComputeCVar = IConsoleManager::Get().FindTConsoleVariableDataInt(TEXT("r.AmbientOcclusion.Compute"));
 	const bool bAOCompute		= AOComputeCVar->GetValueOnAnyThread() > 0;
 
-	const bool bEarlyZMaterialMasking = CVarEarlyZPassOnlyMaterialMasking.GetValueOnAnyThread() != 0;
+	const bool bEarlyZMaterialMasking = MaskedInEarlyPass(ShaderPlatform);
 
 	// Note: ShouldForceFullDepthPass affects which static draw lists meshes go into, so nothing it depends on can change at runtime, unless you do a FGlobalComponentRecreateRenderStateContext to propagate the cvar change
 	return bAOCompute || bDBufferAllowed || bStencilLODDither || bEarlyZMaterialMasking || IsForwardShadingEnabled(ShaderPlatform) || IsUsingSelectiveBasePassOutputs(ShaderPlatform);
@@ -3246,6 +3256,17 @@ void FScene::UpdateEarlyZPassMode()
 			DefaultBasePassDepthStencilAccess = FExclusiveDepthStencil::DepthRead_StencilWrite;
 		}
 	}
+	else if (GetShadingPath(GetFeatureLevel()) == EShadingPath::Mobile)
+	{
+		if (MaskedInEarlyPass(GetFeatureLevelShaderPlatform(FeatureLevel)))
+		{
+			EarlyZPassMode = DDM_MaskedOnly;
+		}
+		else
+		{
+			EarlyZPassMode = DDM_None;
+		}
+	}
 }
 
 void FScene::ConditionalMarkStaticMeshElementsForUpdate()
diff --git a/Engine/Source/Runtime/Renderer/Private/SceneRendering.h b/Engine/Source/Runtime/Renderer/Private/SceneRendering.h
--- a/Engine/Source/Runtime/Renderer/Private/SceneRendering.h
+++ b/Engine/Source/Runtime/Renderer/Private/SceneRendering.h
@@ -1823,6 +1823,8 @@ protected:
 
 	void InitViews(FRHICommandListImmediate& RHICmdList);
 
+	void RenderPrePass(FRHICommandListImmediate& RHICmdList);
+
 	/** Renders the opaque base pass for mobile. */
 	void RenderMobileBasePass(FRHICommandListImmediate& RHICmdList, const TArrayView<const FViewInfo*> PassViews);
 
diff --git a/Engine/Source/Runtime/Renderer/Private/SceneVisibility.cpp b/Engine/Source/Runtime/Renderer/Private/SceneVisibility.cpp
--- a/Engine/Source/Runtime/Renderer/Private/SceneVisibility.cpp
+++ b/Engine/Source/Runtime/Renderer/Private/SceneVisibility.cpp
@@ -2169,6 +2169,7 @@ struct FRelevancePacket
 			const bool bDrawDepthOnly = ViewData.bFullEarlyZPass || FMath::Square(Bounds.BoxSphereBounds.SphereRadius) > GMinScreenRadiusForDepthPrepass * GMinScreenRadiusForDepthPrepass * LODFactorDistanceSquared;
 
 			const bool bAddLightmapDensityCommands = View.Family->EngineShowFlags.LightMapDensity && AllowDebugViewmodes();
+			const bool bMobileMaskedInEarlyPass = MaskedInEarlyPass(Scene->GetShaderPlatform()) && Scene->EarlyZPassMode == DDM_MaskedOnly;
 
 			const int32 NumStaticMeshes = PrimitiveSceneInfo->StaticMeshRelevances.Num();
 			for(int32 MeshIndex = 0;MeshIndex < NumStaticMeshes;MeshIndex++)
@@ -2234,7 +2235,8 @@ struct FRelevancePacket
 							&& (ViewRelevance.bRenderInMainPass || ViewRelevance.bRenderCustomDepth || ViewRelevance.bRenderInDepthPass)
 							&& !bHiddenByHLODFade)
 						{
-							if (StaticMeshRelevance.bUseForDepthPass && bDrawDepthOnly)
+							bool bMobileIsInDepthPassMaskedMesh = (bMobileMaskedInEarlyPass && ViewRelevance.bMasked) && ShadingPath == EShadingPath::Mobile;
+							if ((StaticMeshRelevance.bUseForDepthPass && bDrawDepthOnly && ShadingPath != EShadingPath::Mobile) || bMobileIsInDepthPassMaskedMesh)
 							{
 								DrawCommandPacket.AddCommandsForMesh(PrimitiveIndex, PrimitiveSceneInfo, StaticMeshRelevance, StaticMesh, Scene, bCanCache, EMeshPass::DepthPass);
 							}
