Build fix STAT_CLM_PrePass->STAT_CLM_MobilePrePass

diff --git a/Engine/Source/Runtime/Renderer/Private/MobileShadingRenderer.cpp b/Engine/Source/Runtime/Renderer/Private/MobileShadingRenderer.cpp
--- a/Engine/Source/Runtime/Renderer/Private/MobileShadingRenderer.cpp
+++ b/Engine/Source/Runtime/Renderer/Private/MobileShadingRenderer.cpp
@@ -85,7 +85,7 @@ DECLARE_CYCLE_STAT(TEXT("Post"), STAT_CLMM_Post, STATGROUP_CommandListMarkers);
 DECLARE_CYCLE_STAT(TEXT("Translucency"), STAT_CLMM_Translucency, STATGROUP_CommandListMarkers);
 DECLARE_CYCLE_STAT(TEXT("Shadows"), STAT_CLMM_Shadows, STATGROUP_CommandListMarkers);
 DECLARE_CYCLE_STAT(TEXT("SceneSimulation"), STAT_CLMM_SceneSim, STATGROUP_CommandListMarkers);
-DECLARE_CYCLE_STAT(TEXT("PrePass"), STAT_CLM_PrePass, STATGROUP_CommandListMarkers);
+DECLARE_CYCLE_STAT(TEXT("PrePass"), STAT_CLM_MobilePrePass, STATGROUP_CommandListMarkers);
 
 FGlobalDynamicIndexBuffer FMobileSceneRenderer::DynamicIndexBuffer;
 FGlobalDynamicVertexBuffer FMobileSceneRenderer::DynamicVertexBuffer;
@@ -580,7 +580,7 @@ void FMobileSceneRenderer::Render(FRHICommandListImmediate& RHICmdList)
 	
 	RHICmdList.BeginRenderPass(SceneColorRenderPassInfo, TEXT("SceneColorRendering"));
 
-	RHICmdList.SetCurrentStat(GET_STATID(STAT_CLM_PrePass));
+	RHICmdList.SetCurrentStat(GET_STATID(STAT_CLM_MobilePrePass));
 	RenderPrePass(RHICmdList);
 
 	if (GIsEditor && !View.bIsSceneCapture)
