Reenable use of query fences on PC for virtual texutre feedback buffer

diff --git a/Engine/Source/Runtime/Renderer/Private/VT/VirtualTextureFeedback.cpp b/Engine/Source/Runtime/Renderer/Private/VT/VirtualTextureFeedback.cpp
--- a/Engine/Source/Runtime/Renderer/Private/VT/VirtualTextureFeedback.cpp
+++ b/Engine/Source/Runtime/Renderer/Private/VT/VirtualTextureFeedback.cpp
@@ -139,11 +139,6 @@ FVirtualTextureFeedback::FVirtualTextureFeedback()
 	, PendingTargetCount(0)
 {
 	FeedBackFences = new FFeedbackFences;
-
-	for (int i = 0; i < TargetCapacity; ++i)
-	{
-		FeedbackCPU[i].ReadbackBuffer = nullptr;
-	}
 }
 
 FVirtualTextureFeedback::~FVirtualTextureFeedback()
@@ -156,15 +151,6 @@ void FVirtualTextureFeedback::ReleaseResources()
 	FeedbackBufferUAV.SafeRelease();
 	FeedbackBuffer.SafeRelease();
 	
-	for (int i = 0; i < TargetCapacity; ++i)
-	{
-		if (FeedbackCPU[i].ReadbackBuffer)
-		{
-			delete FeedbackCPU[i].ReadbackBuffer;
-			FeedbackCPU[i].ReadbackBuffer = nullptr;
-		}
-	}
-
 	FeedBackFences->ReleaseAll();
 
 	CPUReadIndex = 0u;
@@ -187,14 +173,6 @@ void FVirtualTextureFeedback::CreateResourceGPU( FRHICommandListImmediate& RHICm
 	FeedBackFences->Init(RHICmdList);
 }
 
-const static FName VTReadbackRequestNames[FVirtualTextureFeedback::TargetCapacity] = 
-{ 
-	FName("VTFeedbackCPU_0"), 
-	FName("VTFeedbackCPU_1"), 
-	FName("VTFeedbackCPU_2"), 
-	FName("VTFeedbackCPU_3") 
-};
-
 void FVirtualTextureFeedback::TransferGPUToCPU( FRHICommandListImmediate& RHICmdList, TArrayView<FIntRect> const& ViewRects)
 {
 	QUICK_SCOPE_CYCLE_COUNTER(STAT_VirtualTextureFeedback_TransferGPUToCPU);
@@ -212,13 +190,6 @@ void FVirtualTextureFeedback::TransferGPUToCPU( FRHICommandListImmediate& RHICmd
 
 	FFeedBackItem& FeedbackEntryCPU = FeedbackCPU[GPUWriteIndex];
 
-	///////////////////////////////
-	if (FeedbackEntryCPU.ReadbackBuffer == nullptr)
-	{
-		FName RequestName = VTReadbackRequestNames[GPUWriteIndex];
-		FeedbackEntryCPU.ReadbackBuffer = new FRHIGPUBufferReadback(RequestName);
-	}
-
 	FeedbackEntryCPU.NumRects = FMath::Min((int32)MaxRectPerTarget, ViewRects.Num());
 	for (int32 RectIndex = 0; RectIndex < FeedbackEntryCPU.NumRects; ++RectIndex)
 	{
@@ -228,13 +199,16 @@ void FVirtualTextureFeedback::TransferGPUToCPU( FRHICommandListImmediate& RHICmd
 		FeedbackEntryCPU.Rects[RectIndex].Max = FIntPoint((Rect.Max.X + 15) / 16, (Rect.Max.Y + 15) / 16);
 	}
 	
+	FeedBackFences->Allocate(RHICmdList, GPUWriteIndex);
+
 	// We only need to transfer 1 copy of the data, so restrict mask to the first active GPU.
 	FeedbackEntryCPU.GPUMask = FRHIGPUMask::FromIndex(RHICmdList.GetGPUMask().GetFirstIndex());
 	SCOPED_GPU_MASK(RHICmdList, FeedbackEntryCPU.GPUMask);
 
-	// copy whole feedback buffer
-	FeedbackEntryCPU.ReadbackBuffer->EnqueueCopy(RHICmdList, FeedbackBuffer, NumBytes);
-	///////////////////////////////
+	FeedbackEntryCPU.ReadbackBuffer = RHICreateStagingBuffer();
+	RHICmdList.CopyToStagingBuffer(FeedbackBuffer, FeedbackEntryCPU.ReadbackBuffer, 0, NumBytes);
+
+	FeedBackFences->Write(RHICmdList, GPUWriteIndex);
 
 	GPUWriteIndex = (GPUWriteIndex + 1) % TargetCapacity;
 	++PendingTargetCount;
@@ -243,7 +217,10 @@ void FVirtualTextureFeedback::TransferGPUToCPU( FRHICommandListImmediate& RHICmd
 bool FVirtualTextureFeedback::CanMap(FRHICommandListImmediate& RHICmdList)
 {
 	const FFeedBackItem& FeedbackEntryCPU = FeedbackCPU[CPUReadIndex];
-	return (PendingTargetCount > 0u && FeedbackEntryCPU.ReadbackBuffer && FeedbackEntryCPU.ReadbackBuffer->IsReady());
+	return 
+		PendingTargetCount > 0u && 
+		FeedbackEntryCPU.ReadbackBuffer->IsValid() && 
+		FeedBackFences->Poll(RHICmdList, CPUReadIndex);
 }
 
 bool FVirtualTextureFeedback::Map(FRHICommandListImmediate& RHICmdList, MapResult& OutResult)
@@ -253,8 +230,8 @@ bool FVirtualTextureFeedback::Map(FRHICommandListImmediate& RHICmdList, MapResul
 	FFeedBackItem& FeedbackEntryCPU = FeedbackCPU[CPUReadIndex];
 		
 	if (PendingTargetCount > 0u &&
-		FeedbackEntryCPU.ReadbackBuffer &&
-		FeedbackEntryCPU.ReadbackBuffer->IsReady())
+		FeedbackEntryCPU.ReadbackBuffer->IsValid() &&
+		FeedBackFences->Poll(RHICmdList, CPUReadIndex))
 	{
 		SCOPED_GPU_MASK(RHICmdList, FeedbackEntryCPU.GPUMask);
 	
@@ -264,11 +241,14 @@ bool FVirtualTextureFeedback::Map(FRHICommandListImmediate& RHICmdList, MapResul
 			OutResult.Rects[i] = FeedbackEntryCPU.Rects[i];
 		}
 		OutResult.Pitch = Size.X;
-		OutResult.Buffer.SetNumUninitialized(NumBytes/sizeof(uint32));
 		
-		void* MappedMem = FeedbackEntryCPU.ReadbackBuffer->Lock(NumBytes);
+		void* MappedMem = (uint32*)RHICmdList.LockStagingBuffer(FeedbackEntryCPU.ReadbackBuffer, FeedBackFences->GetMapFence(CPUReadIndex), 0, NumBytes);
+
+		OutResult.Buffer.SetNumUninitialized(NumBytes / sizeof(uint32));
 		FMemory::Memcpy(OutResult.Buffer.GetData(), MappedMem, NumBytes);
-		FeedbackEntryCPU.ReadbackBuffer->Unlock();
+
+		RHICmdList.UnlockStagingBuffer(FeedbackEntryCPU.ReadbackBuffer);
+		FeedbackEntryCPU.ReadbackBuffer.SafeRelease();
 		
 		check(PendingTargetCount > 0u);
 		--PendingTargetCount;
@@ -278,4 +258,3 @@ bool FVirtualTextureFeedback::Map(FRHICommandListImmediate& RHICmdList, MapResul
 
 	return false;
 }
-
diff --git a/Engine/Source/Runtime/Renderer/Private/VT/VirtualTextureFeedback.h b/Engine/Source/Runtime/Renderer/Private/VT/VirtualTextureFeedback.h
--- a/Engine/Source/Runtime/Renderer/Private/VT/VirtualTextureFeedback.h
+++ b/Engine/Source/Runtime/Renderer/Private/VT/VirtualTextureFeedback.h
@@ -5,8 +5,6 @@
 #include "CoreMinimal.h"
 #include "RenderCore/Public/RenderTargetPool.h"
 
-class FRHIGPUBufferReadback;
-
 /*
 ====================================
 	Manages GPU and CPU buffers for VT feedback.
@@ -50,10 +48,10 @@ public:
 private:
 	struct FFeedBackItem
 	{
-		FRHIGPUBufferReadback* ReadbackBuffer;
 		int32 NumRects;
 		FIntRect Rects[MaxRectPerTarget];
 		FRHIGPUMask GPUMask;
+		FStagingBufferRHIRef ReadbackBuffer;
 	};
 
 	FFeedBackItem FeedbackCPU[TargetCapacity];
