Added missing resource release on shutdown introduced by 11279565
#rnx

diff --git a/Engine/Source/Runtime/Renderer/Private/VT/VirtualTextureFeedback.cpp b/Engine/Source/Runtime/Renderer/Private/VT/VirtualTextureFeedback.cpp
--- a/Engine/Source/Runtime/Renderer/Private/VT/VirtualTextureFeedback.cpp
+++ b/Engine/Source/Runtime/Renderer/Private/VT/VirtualTextureFeedback.cpp
@@ -151,6 +151,11 @@ void FVirtualTextureFeedback::ReleaseResources()
 	FeedbackBufferUAV.SafeRelease();
 	FeedbackBuffer.SafeRelease();
 	
+	for (int i = 0; i < TargetCapacity; ++i)
+	{
+		FeedbackCPU[i].ReadbackBuffer.SafeRelease();
+	}
+
 	FeedBackFences->ReleaseAll();
 
 	CPUReadIndex = 0u;
