Integrate fix for editor crash with MSAA and forward shading.

diff --git a/Engine/Source/Runtime/Renderer/Private/PostProcess/SceneRenderTargets.cpp b/Engine/Source/Runtime/Renderer/Private/PostProcess/SceneRenderTargets.cpp
--- a/Engine/Source/Runtime/Renderer/Private/PostProcess/SceneRenderTargets.cpp
+++ b/Engine/Source/Runtime/Renderer/Private/PostProcess/SceneRenderTargets.cpp
@@ -1894,6 +1894,7 @@ void FSceneRenderTargets::ResolveDepthTexture(FRHICommandList& RHICmdList, const
 			break;
 		}
 
+		GraphicsPSOInit.BoundShaderState.VertexDeclarationRHI = GScreenVertexDeclaration.VertexDeclarationRHI;
 		GraphicsPSOInit.BoundShaderState.VertexShaderRHI = GETSAFERHISHADER_VERTEX(*ResolveVertexShader);
 		GraphicsPSOInit.BoundShaderState.PixelShaderRHI = ResolvePixelShader;
 		GraphicsPSOInit.PrimitiveType = PT_TriangleStrip;
