Copying Private-LoadTimes-4.24 stream to Main. Biggest changes are in Materials/Shader memory freezing.

diff --git a/Engine/Source/Runtime/Renderer/Private/AtmosphereRendering.cpp b/Engine/Source/Runtime/Renderer/Private/AtmosphereRendering.cpp
--- a/Engine/Source/Runtime/Renderer/Private/AtmosphereRendering.cpp
+++ b/Engine/Source/Runtime/Renderer/Private/AtmosphereRendering.cpp
@@ -34,6 +34,7 @@ DECLARE_GPU_STAT(AtmospherePreCompute);
 /** Shader parameters needed for atmosphere passes. */
 class FAtmosphereShaderTextureParameters
 {
+	DECLARE_INLINE_TYPE_LAYOUT(FAtmosphereShaderTextureParameters, NonVirtual);
 public:
 	void Bind(const FShaderParameterMap& ParameterMap);
 
@@ -54,12 +55,14 @@ public:
 	friend FArchive& operator<<(FArchive& Ar, FAtmosphereShaderTextureParameters& P);
 
 private:
-	FShaderResourceParameter TransmittanceTexture;
-	FShaderResourceParameter TransmittanceTextureSampler;
-	FShaderResourceParameter IrradianceTexture;
-	FShaderResourceParameter IrradianceTextureSampler;
-	FShaderResourceParameter InscatterTexture;
-	FShaderResourceParameter InscatterTextureSampler;
+	
+		LAYOUT_FIELD(FShaderResourceParameter, TransmittanceTexture)
+		LAYOUT_FIELD(FShaderResourceParameter, TransmittanceTextureSampler)
+		LAYOUT_FIELD(FShaderResourceParameter, IrradianceTexture)
+		LAYOUT_FIELD(FShaderResourceParameter, IrradianceTextureSampler)
+		LAYOUT_FIELD(FShaderResourceParameter, InscatterTexture)
+		LAYOUT_FIELD(FShaderResourceParameter, InscatterTextureSampler)
+	
 };
 
 void FAtmosphereShaderTextureParameters::Bind(const FShaderParameterMap& ParameterMap)
@@ -100,6 +103,7 @@ FArchive& operator<<(FArchive& Ar, FAtmosphereShaderTextureParameters& Parameter
 
 class FAtmosphereShaderPrecomputeTextureParameters
 {
+	DECLARE_INLINE_TYPE_LAYOUT(FAtmosphereShaderPrecomputeTextureParameters, NonVirtual);
 public:
 	enum TexType
 	{
@@ -212,8 +216,10 @@ public:
 	friend FArchive& operator<<(FArchive& Ar,FAtmosphereShaderPrecomputeTextureParameters& P);
 
 private:
-	FShaderResourceParameter AtmosphereTexture[4];
-	FShaderResourceParameter AtmosphereTextureSampler[4];
+	
+		LAYOUT_ARRAY(FShaderResourceParameter, AtmosphereTexture, 4)
+		LAYOUT_ARRAY(FShaderResourceParameter, AtmosphereTextureSampler, 4)
+	
 };
 
 FArchive& operator<<(FArchive& Ar,FAtmosphereShaderPrecomputeTextureParameters& Parameters)
@@ -234,7 +240,7 @@ FArchive& operator<<(FArchive& Ar,FAtmosphereShaderPrecomputeTextureParameters&
 /** A pixel shader for rendering atmospheric fog. */
 class FAtmosphericFogPS : public FGlobalShader
 {
-	DECLARE_SHADER_TYPE(FAtmosphericFogPS, Global);
+	DECLARE_INLINE_TYPE_LAYOUT(FAtmosphericFogPS, NonVirtual);
 public:
 	static bool ShouldCompilePermutation(const FGlobalShaderPermutationParameters& Parameters)
 	{
@@ -253,15 +259,15 @@ public:
 
 	void SetParameters(FRHICommandList& RHICmdList, const FSceneView& View, const TRefCountPtr<IPooledRenderTarget>& LightShaftOcclusion)
 	{
-		FGlobalShader::SetParameters<FViewUniformShaderParameters>(RHICmdList, GetPixelShader(), View.ViewUniformBuffer);
-		SceneTextureParameters.Set(RHICmdList, GetPixelShader(), View.FeatureLevel, ESceneTextureSetupMode::All);
-		AtmosphereTextureParameters.Set(RHICmdList, GetPixelShader(), View);
+		FGlobalShader::SetParameters<FViewUniformShaderParameters>(RHICmdList, RHICmdList.GetBoundPixelShader(), View.ViewUniformBuffer);
+		SceneTextureParameters.Set(RHICmdList, RHICmdList.GetBoundPixelShader(), View.FeatureLevel, ESceneTextureSetupMode::All);
+		AtmosphereTextureParameters.Set(RHICmdList, RHICmdList.GetBoundPixelShader(), View);
 
 		if (LightShaftOcclusion)
 		{
 			SetTextureParameter(
 				RHICmdList, 
-				GetPixelShader(),
+				RHICmdList.GetBoundPixelShader(),
 				OcclusionTextureParameter, OcclusionTextureSamplerParameter,
 				TStaticSamplerState<SF_Bilinear,AM_Clamp,AM_Clamp,AM_Clamp>::GetRHI(),
 				LightShaftOcclusion->GetRenderTargetItem().ShaderResourceTexture
@@ -271,7 +277,7 @@ public:
 		{
 			SetTextureParameter(
 				RHICmdList, 
-				GetPixelShader(),
+				RHICmdList.GetBoundPixelShader(),
 				OcclusionTextureParameter, OcclusionTextureSamplerParameter,
 				TStaticSamplerState<SF_Bilinear,AM_Clamp,AM_Clamp,AM_Clamp>::GetRHI(),
 				GWhiteTexture->TextureRHI
@@ -279,21 +285,12 @@ public:
 		}
 	}
 
-	virtual bool Serialize(FArchive& Ar) override
-	{
-		bool bShaderHasOutdatedParameters = FGlobalShader::Serialize(Ar);
-		Ar << SceneTextureParameters;
-		Ar << AtmosphereTextureParameters;
-		Ar << OcclusionTextureParameter;
-		Ar << OcclusionTextureSamplerParameter;
-		return bShaderHasOutdatedParameters;
-	}
-
 private:
-	FSceneTextureShaderParameters SceneTextureParameters;
-	FAtmosphereShaderTextureParameters AtmosphereTextureParameters;
-	FShaderResourceParameter OcclusionTextureParameter;
-	FShaderResourceParameter OcclusionTextureSamplerParameter;
+	
+	LAYOUT_FIELD(FSceneTextureShaderParameters, SceneTextureParameters);
+	LAYOUT_FIELD(FAtmosphereShaderTextureParameters, AtmosphereTextureParameters);
+	LAYOUT_FIELD(FShaderResourceParameter, OcclusionTextureParameter);
+	LAYOUT_FIELD(FShaderResourceParameter, OcclusionTextureSamplerParameter);
 };
 
 template<uint32 RenderFlag>
@@ -374,13 +371,7 @@ public:
 
 	void SetParameters(FRHICommandList& RHICmdList, const FViewInfo& View)
 	{
-		FGlobalShader::SetParameters<FViewUniformShaderParameters>(RHICmdList, GetVertexShader(),View.ViewUniformBuffer);
-	}
-
-	virtual bool Serialize(FArchive& Ar) override
-	{
-		bool bShaderHasOutdatedParameters = FGlobalShader::Serialize(Ar);
-		return bShaderHasOutdatedParameters;
+		FGlobalShader::SetParameters<FViewUniformShaderParameters>(RHICmdList, RHICmdList.GetBoundVertexShader(),View.ViewUniformBuffer);
 	}
 };
 
@@ -431,41 +422,41 @@ void SetAtmosphericFogShaders(FRHICommandList& RHICmdList, FGraphicsPipelineStat
 	}
 
 	TShaderMapRef<FAtmosphericVS> VertexShader(ShaderMap);
-	FAtmosphericFogPS* PixelShader = NULL;
+	TShaderRef<FAtmosphericFogPS> PixelShader;
 	switch (RenderFlag)
 	{
 	default:
 		checkSlow(false)
 
 	case EAtmosphereRenderFlag::E_EnableAll:
-		PixelShader = *TShaderMapRef<TAtmosphericFogPS<EAtmosphereRenderFlag::E_EnableAll> >(ShaderMap);
+		PixelShader = TShaderMapRef<TAtmosphericFogPS<EAtmosphereRenderFlag::E_EnableAll> >(ShaderMap);
 		break;
 	case EAtmosphereRenderFlag::E_DisableSunDisk:
-		PixelShader = *TShaderMapRef<TAtmosphericFogPS<EAtmosphereRenderFlag::E_DisableSunDisk> >(ShaderMap);
+		PixelShader = TShaderMapRef<TAtmosphericFogPS<EAtmosphereRenderFlag::E_DisableSunDisk> >(ShaderMap);
 		break;
 	case EAtmosphereRenderFlag::E_DisableGroundScattering:
-		PixelShader = *TShaderMapRef<TAtmosphericFogPS<EAtmosphereRenderFlag::E_DisableGroundScattering> >(ShaderMap);
+		PixelShader = TShaderMapRef<TAtmosphericFogPS<EAtmosphereRenderFlag::E_DisableGroundScattering> >(ShaderMap);
 		break;
 	case EAtmosphereRenderFlag::E_DisableSunAndGround :
-		PixelShader = *TShaderMapRef<TAtmosphericFogPS<EAtmosphereRenderFlag::E_DisableSunAndGround> >(ShaderMap);
+		PixelShader = TShaderMapRef<TAtmosphericFogPS<EAtmosphereRenderFlag::E_DisableSunAndGround> >(ShaderMap);
 		break;
 	case EAtmosphereRenderFlag::E_DisableLightShaft :
-		PixelShader = *TShaderMapRef<TAtmosphericFogPS<EAtmosphereRenderFlag::E_DisableLightShaft> >(ShaderMap);
+		PixelShader = TShaderMapRef<TAtmosphericFogPS<EAtmosphereRenderFlag::E_DisableLightShaft> >(ShaderMap);
 		break;
 	case EAtmosphereRenderFlag::E_DisableSunAndLightShaft :
-		PixelShader = *TShaderMapRef<TAtmosphericFogPS<EAtmosphereRenderFlag::E_DisableSunAndLightShaft> >(ShaderMap);
+		PixelShader = TShaderMapRef<TAtmosphericFogPS<EAtmosphereRenderFlag::E_DisableSunAndLightShaft> >(ShaderMap);
 		break;
 	case EAtmosphereRenderFlag::E_DisableGroundAndLightShaft :
-		PixelShader = *TShaderMapRef<TAtmosphericFogPS<EAtmosphereRenderFlag::E_DisableGroundAndLightShaft> >(ShaderMap);
+		PixelShader = TShaderMapRef<TAtmosphericFogPS<EAtmosphereRenderFlag::E_DisableGroundAndLightShaft> >(ShaderMap);
 		break;
 	case EAtmosphereRenderFlag::E_DisableAll:
-		PixelShader = *TShaderMapRef<TAtmosphericFogPS<EAtmosphereRenderFlag::E_DisableAll> >(ShaderMap);
+		PixelShader = TShaderMapRef<TAtmosphericFogPS<EAtmosphereRenderFlag::E_DisableAll> >(ShaderMap);
 		break;
 	}
 	
 	GraphicsPSOInit.BoundShaderState.VertexDeclarationRHI = GAtmophereVertexDeclaration.VertexDeclarationRHI;
-	GraphicsPSOInit.BoundShaderState.VertexShaderRHI = GETSAFERHISHADER_VERTEX(*VertexShader);
-	GraphicsPSOInit.BoundShaderState.PixelShaderRHI = GETSAFERHISHADER_PIXEL(PixelShader);
+	GraphicsPSOInit.BoundShaderState.VertexShaderRHI = VertexShader.GetVertexShader();
+	GraphicsPSOInit.BoundShaderState.PixelShaderRHI = PixelShader.GetPixelShader();
 	GraphicsPSOInit.PrimitiveType = PT_TriangleList;
 
 	SetGraphicsPipelineState(RHICmdList, GraphicsPSOInit);
@@ -538,16 +529,9 @@ public:
 	{
 	}
 
-	// FShader interface.
-	virtual bool Serialize(FArchive& Ar) override
-	{
-		bool bShaderHasOutdatedParameters = FGlobalShader::Serialize(Ar);
-		return bShaderHasOutdatedParameters;
-	}
-
 	void SetParameters(FRHICommandList& RHICmdList, const FViewInfo& View)
 	{
-		FGlobalShader::SetParameters<FViewUniformShaderParameters>(RHICmdList, GetPixelShader(), View.ViewUniformBuffer);
+		FGlobalShader::SetParameters<FViewUniformShaderParameters>(RHICmdList, RHICmdList.GetBoundPixelShader(), View.ViewUniformBuffer);
 	}
 };
 
@@ -564,7 +548,7 @@ class FAtmosphereIrradiance1PS : public FGlobalShader
 	FAtmosphereIrradiance1PS() {}
 
 public:
-	FAtmosphereShaderPrecomputeTextureParameters AtmosphereParameters;
+	LAYOUT_FIELD(FAtmosphereShaderPrecomputeTextureParameters, AtmosphereParameters);
 
 	/* Initialization constructor. */
 	FAtmosphereIrradiance1PS(const ShaderMetaType::CompiledShaderInitializerType& Initializer)
@@ -573,17 +557,9 @@ public:
 		AtmosphereParameters.Bind(Initializer.ParameterMap, 0, FAtmosphereShaderPrecomputeTextureParameters::Transmittance);
 	}
 
-	// FShader interface.
-	virtual bool Serialize(FArchive& Ar) override
-	{
-		bool bShaderHasOutdatedParameters = FGlobalShader::Serialize(Ar);
-		Ar << AtmosphereParameters;
-		return bShaderHasOutdatedParameters;
-	}
-
 	void SetParameters(FRHICommandList& RHICmdList, const FAtmosphereTextures* Textures)
 	{
-		AtmosphereParameters.Set(RHICmdList, GetPixelShader(), 0, FAtmosphereShaderPrecomputeTextureParameters::Transmittance, Textures);
+		AtmosphereParameters.Set(RHICmdList, RHICmdList.GetBoundPixelShader(), 0, FAtmosphereShaderPrecomputeTextureParameters::Transmittance, Textures);
 	}
 };
 
@@ -600,8 +576,8 @@ class FAtmosphereIrradianceNPS : public FGlobalShader
 	FAtmosphereIrradianceNPS() {}
 
 public:
-	FAtmosphereShaderPrecomputeTextureParameters AtmosphereParameters;
-	FShaderParameter FirstOrderParameter;
+	LAYOUT_FIELD(FAtmosphereShaderPrecomputeTextureParameters, AtmosphereParameters);
+	LAYOUT_FIELD(FShaderParameter, FirstOrderParameter);
 
 	/* Initialization constructor. */
 	FAtmosphereIrradianceNPS(const ShaderMetaType::CompiledShaderInitializerType& Initializer)
@@ -613,21 +589,13 @@ public:
 		FirstOrderParameter.Bind(Initializer.ParameterMap, TEXT("FirstOrder"));
 	}
 
-	// FShader interface.
-	virtual bool Serialize(FArchive& Ar) override
-	{
-		bool bShaderHasOutdatedParameters = FGlobalShader::Serialize(Ar);
-		Ar << AtmosphereParameters << FirstOrderParameter;
-		return bShaderHasOutdatedParameters;
-	}
-
 	void SetParameters(FRHICommandList& RHICmdList, const FViewInfo& View, float FirstOrder, const FAtmosphereTextures* Textures)
 	{
-		FGlobalShader::SetParameters<FViewUniformShaderParameters>(RHICmdList, GetPixelShader(), View.ViewUniformBuffer);
-		AtmosphereParameters.Set(RHICmdList, GetPixelShader(), 0, FAtmosphereShaderPrecomputeTextureParameters::Transmittance, Textures);
-		AtmosphereParameters.Set(RHICmdList, GetPixelShader(), 1, FAtmosphereShaderPrecomputeTextureParameters::DeltaSR, Textures);
-		AtmosphereParameters.Set(RHICmdList, GetPixelShader(), 2, FAtmosphereShaderPrecomputeTextureParameters::DeltaSM, Textures);
-		SetShaderValue(RHICmdList, GetPixelShader(), FirstOrderParameter, FirstOrder);
+		FGlobalShader::SetParameters<FViewUniformShaderParameters>(RHICmdList, RHICmdList.GetBoundPixelShader(), View.ViewUniformBuffer);
+		AtmosphereParameters.Set(RHICmdList, RHICmdList.GetBoundPixelShader(), 0, FAtmosphereShaderPrecomputeTextureParameters::Transmittance, Textures);
+		AtmosphereParameters.Set(RHICmdList, RHICmdList.GetBoundPixelShader(), 1, FAtmosphereShaderPrecomputeTextureParameters::DeltaSR, Textures);
+		AtmosphereParameters.Set(RHICmdList, RHICmdList.GetBoundPixelShader(), 2, FAtmosphereShaderPrecomputeTextureParameters::DeltaSM, Textures);
+		SetShaderValue(RHICmdList, RHICmdList.GetBoundPixelShader(), FirstOrderParameter, FirstOrder);
 	}
 };
 
@@ -644,7 +612,7 @@ class FAtmosphereCopyIrradiancePS : public FGlobalShader
 	FAtmosphereCopyIrradiancePS() {}
 
 public:
-	FAtmosphereShaderPrecomputeTextureParameters AtmosphereParameters;
+	LAYOUT_FIELD(FAtmosphereShaderPrecomputeTextureParameters, AtmosphereParameters);
 
 	/* Initialization constructor. */
 	FAtmosphereCopyIrradiancePS(const ShaderMetaType::CompiledShaderInitializerType& Initializer)
@@ -653,17 +621,9 @@ public:
 		AtmosphereParameters.Bind(Initializer.ParameterMap, 0, FAtmosphereShaderPrecomputeTextureParameters::DeltaE);
 	}
 
-	// FShader interface.
-	virtual bool Serialize(FArchive& Ar) override
-	{
-		bool bShaderHasOutdatedParameters = FGlobalShader::Serialize(Ar);
-		Ar << AtmosphereParameters;
-		return bShaderHasOutdatedParameters;
-	}
-
 	void SetParameters(FRHICommandList& RHICmdList, const FAtmosphereTextures* Textures)
 	{
-		AtmosphereParameters.Set(RHICmdList, GetPixelShader(), 0, FAtmosphereShaderPrecomputeTextureParameters::DeltaE, Textures);
+		AtmosphereParameters.Set(RHICmdList, RHICmdList.GetBoundPixelShader(), 0, FAtmosphereShaderPrecomputeTextureParameters::DeltaE, Textures);
 	}
 };
 
@@ -679,7 +639,7 @@ class FAtmosphereGS : public FGlobalShader
 	FAtmosphereGS() {}
 
 public:
-	FShaderParameter AtmosphereLayerParameter;
+	LAYOUT_FIELD(FShaderParameter, AtmosphereLayerParameter);
 
 	FAtmosphereGS(const ShaderMetaType::CompiledShaderInitializerType& Initializer)
 		: FGlobalShader(Initializer)
@@ -689,14 +649,7 @@ public:
 
 	void SetParameters(FRHICommandList& RHICmdList, int AtmosphereLayer)
 	{
-		SetShaderValue(RHICmdList, GetGeometryShader(), AtmosphereLayerParameter, AtmosphereLayer);
-	}
-
-	virtual bool Serialize(FArchive& Ar) override
-	{
-		bool bShaderHasOutdatedParameters = FGlobalShader::Serialize(Ar);
-		Ar << AtmosphereLayerParameter;
-		return bShaderHasOutdatedParameters;
+		SetShaderValue(RHICmdList, RHICmdList.GetBoundGeometryShader(), AtmosphereLayerParameter, AtmosphereLayer);
 	}
 };
 
@@ -713,9 +666,9 @@ class FAtmosphereInscatter1PS : public FGlobalShader
 	FAtmosphereInscatter1PS() {}
 
 public:
-	FAtmosphereShaderPrecomputeTextureParameters AtmosphereParameters;
-	FShaderParameter dhdHParameter;
-	FShaderParameter AtmosphereRParameter;
+	LAYOUT_FIELD(FAtmosphereShaderPrecomputeTextureParameters, AtmosphereParameters);
+	LAYOUT_FIELD(FShaderParameter, dhdHParameter);
+	LAYOUT_FIELD(FShaderParameter, AtmosphereRParameter);
 
 	/* Initialization constructor. */
 	FAtmosphereInscatter1PS(const ShaderMetaType::CompiledShaderInitializerType& Initializer)
@@ -726,20 +679,12 @@ public:
 		AtmosphereRParameter.Bind(Initializer.ParameterMap, TEXT("AtmosphereR"));
 	}
 
-	// FShader interface.
-	virtual bool Serialize(FArchive& Ar) override
-	{
-		bool bShaderHasOutdatedParameters = FGlobalShader::Serialize(Ar);
-		Ar << AtmosphereParameters << dhdHParameter << AtmosphereRParameter;
-		return bShaderHasOutdatedParameters;
-	}
-
 	void SetParameters(FRHICommandList& RHICmdList, const FViewInfo& View, float AtmosphereR, const FVector4& DhdH, const FAtmosphereTextures* Textures)
 	{
-		FGlobalShader::SetParameters<FViewUniformShaderParameters>(RHICmdList, GetPixelShader(), View.ViewUniformBuffer);
-		AtmosphereParameters.Set(RHICmdList, GetPixelShader(), 0, FAtmosphereShaderPrecomputeTextureParameters::Transmittance, Textures);
-		SetShaderValue(RHICmdList, GetPixelShader(), dhdHParameter, DhdH);
-		SetShaderValue(RHICmdList, GetPixelShader(), AtmosphereRParameter, AtmosphereR);
+		FGlobalShader::SetParameters<FViewUniformShaderParameters>(RHICmdList, RHICmdList.GetBoundPixelShader(), View.ViewUniformBuffer);
+		AtmosphereParameters.Set(RHICmdList, RHICmdList.GetBoundPixelShader(), 0, FAtmosphereShaderPrecomputeTextureParameters::Transmittance, Textures);
+		SetShaderValue(RHICmdList, RHICmdList.GetBoundPixelShader(), dhdHParameter, DhdH);
+		SetShaderValue(RHICmdList, RHICmdList.GetBoundPixelShader(), AtmosphereRParameter, AtmosphereR);
 	}
 };
 
@@ -756,10 +701,10 @@ class FAtmosphereCopyInscatter1PS : public FGlobalShader
 	FAtmosphereCopyInscatter1PS() {}
 
 public:
-	FAtmosphereShaderPrecomputeTextureParameters AtmosphereParameters;
-	FShaderParameter dhdHParameter;
-	FShaderParameter AtmosphereRParameter;
-	FShaderParameter AtmosphereLayerParameter;
+	LAYOUT_FIELD(FAtmosphereShaderPrecomputeTextureParameters, AtmosphereParameters);
+	LAYOUT_FIELD(FShaderParameter, dhdHParameter);
+	LAYOUT_FIELD(FShaderParameter, AtmosphereRParameter);
+	LAYOUT_FIELD(FShaderParameter, AtmosphereLayerParameter);
 
 	/* Initialization constructor. */
 	FAtmosphereCopyInscatter1PS(const ShaderMetaType::CompiledShaderInitializerType& Initializer)
@@ -772,22 +717,14 @@ public:
 		AtmosphereLayerParameter.Bind(Initializer.ParameterMap, TEXT("AtmosphereLayer"));
 	}
 
-	// FShader interface.
-	virtual bool Serialize(FArchive& Ar) override
-	{
-		bool bShaderHasOutdatedParameters = FGlobalShader::Serialize(Ar);
-		Ar << AtmosphereParameters << dhdHParameter << AtmosphereRParameter << AtmosphereLayerParameter;
-		return bShaderHasOutdatedParameters;
-	}
-
 	void SetParameters(FRHICommandList& RHICmdList, const FViewInfo& View, float AtmosphereR, const FVector4& DhdH, int AtmosphereLayer, const FAtmosphereTextures* Textures)
 	{
-		FGlobalShader::SetParameters<FViewUniformShaderParameters>(RHICmdList, GetPixelShader(), View.ViewUniformBuffer);
-		AtmosphereParameters.Set(RHICmdList, GetPixelShader(), 0, FAtmosphereShaderPrecomputeTextureParameters::DeltaSR, Textures);
-		AtmosphereParameters.Set(RHICmdList, GetPixelShader(), 1, FAtmosphereShaderPrecomputeTextureParameters::DeltaSM, Textures);
-		SetShaderValue(RHICmdList, GetPixelShader(), AtmosphereRParameter, AtmosphereR);
-		SetShaderValue(RHICmdList, GetPixelShader(), dhdHParameter, DhdH);
-		SetShaderValue(RHICmdList, GetPixelShader(), AtmosphereLayerParameter, AtmosphereLayer);
+		FGlobalShader::SetParameters<FViewUniformShaderParameters>(RHICmdList, RHICmdList.GetBoundPixelShader(), View.ViewUniformBuffer);
+		AtmosphereParameters.Set(RHICmdList, RHICmdList.GetBoundPixelShader(), 0, FAtmosphereShaderPrecomputeTextureParameters::DeltaSR, Textures);
+		AtmosphereParameters.Set(RHICmdList, RHICmdList.GetBoundPixelShader(), 1, FAtmosphereShaderPrecomputeTextureParameters::DeltaSM, Textures);
+		SetShaderValue(RHICmdList, RHICmdList.GetBoundPixelShader(), AtmosphereRParameter, AtmosphereR);
+		SetShaderValue(RHICmdList, RHICmdList.GetBoundPixelShader(), dhdHParameter, DhdH);
+		SetShaderValue(RHICmdList, RHICmdList.GetBoundPixelShader(), AtmosphereLayerParameter, AtmosphereLayer);
 	}
 };
 
@@ -804,10 +741,10 @@ class FAtmosphereCopyInscatterNPS : public FGlobalShader
 	FAtmosphereCopyInscatterNPS() {}
 
 public:
-	FAtmosphereShaderPrecomputeTextureParameters AtmosphereParameters;
-	FShaderParameter dhdHParameter;
-	FShaderParameter AtmosphereRParameter;
-	FShaderParameter AtmosphereLayerParameter;
+	LAYOUT_FIELD(FAtmosphereShaderPrecomputeTextureParameters, AtmosphereParameters);
+	LAYOUT_FIELD(FShaderParameter, dhdHParameter);
+	LAYOUT_FIELD(FShaderParameter, AtmosphereRParameter);
+	LAYOUT_FIELD(FShaderParameter, AtmosphereLayerParameter);
 
 	/* Initialization constructor. */
 	FAtmosphereCopyInscatterNPS(const ShaderMetaType::CompiledShaderInitializerType& Initializer)
@@ -819,21 +756,13 @@ public:
 		AtmosphereLayerParameter.Bind(Initializer.ParameterMap, TEXT("AtmosphereLayer"));
 	}
 
-	// FShader interface.
-	virtual bool Serialize(FArchive& Ar) override
-	{
-		bool bShaderHasOutdatedParameters = FGlobalShader::Serialize(Ar);
-		Ar << AtmosphereParameters << dhdHParameter << AtmosphereRParameter << AtmosphereLayerParameter;
-		return bShaderHasOutdatedParameters;
-	}
-
 	void SetParameters(FRHICommandList& RHICmdList, const FViewInfo& View, float AtmosphereR, const FVector4& DhdH, int AtmosphereLayer, const FAtmosphereTextures* Textures)
 	{
-		FGlobalShader::SetParameters<FViewUniformShaderParameters>(RHICmdList, GetPixelShader(), View.ViewUniformBuffer);
-		AtmosphereParameters.Set(RHICmdList, GetPixelShader(), 0, FAtmosphereShaderPrecomputeTextureParameters::DeltaSR, Textures);
-		SetShaderValue(RHICmdList, GetPixelShader(), AtmosphereRParameter, AtmosphereR);
-		SetShaderValue(RHICmdList, GetPixelShader(), dhdHParameter, DhdH);
-		SetShaderValue(RHICmdList, GetPixelShader(), AtmosphereLayerParameter, AtmosphereLayer);
+		FGlobalShader::SetParameters<FViewUniformShaderParameters>(RHICmdList, RHICmdList.GetBoundPixelShader(), View.ViewUniformBuffer);
+		AtmosphereParameters.Set(RHICmdList, RHICmdList.GetBoundPixelShader(), 0, FAtmosphereShaderPrecomputeTextureParameters::DeltaSR, Textures);
+		SetShaderValue(RHICmdList, RHICmdList.GetBoundPixelShader(), AtmosphereRParameter, AtmosphereR);
+		SetShaderValue(RHICmdList, RHICmdList.GetBoundPixelShader(), dhdHParameter, DhdH);
+		SetShaderValue(RHICmdList, RHICmdList.GetBoundPixelShader(), AtmosphereLayerParameter, AtmosphereLayer);
 	}
 };
 
@@ -850,10 +779,10 @@ class FAtmosphereInscatterSPS : public FGlobalShader
 	FAtmosphereInscatterSPS() {}
 
 public:
-	FAtmosphereShaderPrecomputeTextureParameters AtmosphereParameters;
-	FShaderParameter dhdHParameter;
-	FShaderParameter AtmosphereRParameter;
-	FShaderParameter FirstOrderParameter;
+	LAYOUT_FIELD(FAtmosphereShaderPrecomputeTextureParameters, AtmosphereParameters);
+	LAYOUT_FIELD(FShaderParameter, dhdHParameter);
+	LAYOUT_FIELD(FShaderParameter, AtmosphereRParameter);
+	LAYOUT_FIELD(FShaderParameter, FirstOrderParameter);
 
 	/* Initialization constructor. */
 	FAtmosphereInscatterSPS(const ShaderMetaType::CompiledShaderInitializerType& Initializer)
@@ -868,24 +797,16 @@ public:
 		FirstOrderParameter.Bind(Initializer.ParameterMap, TEXT("FirstOrder"));
 	}
 
-	// FShader interface.
-	virtual bool Serialize(FArchive& Ar) override
-	{
-		bool bShaderHasOutdatedParameters = FGlobalShader::Serialize(Ar);
-		Ar << AtmosphereParameters << dhdHParameter << AtmosphereRParameter << FirstOrderParameter;
-		return bShaderHasOutdatedParameters;
-	}
-
 	void SetParameters(FRHICommandList& RHICmdList, const FViewInfo& View, float AtmosphereR, const FVector4& DhdH, float FirstOrder, const FAtmosphereTextures* Textures)
 	{
-		FGlobalShader::SetParameters<FViewUniformShaderParameters>(RHICmdList, GetPixelShader(), View.ViewUniformBuffer);
-		AtmosphereParameters.Set(RHICmdList, GetPixelShader(), 0, FAtmosphereShaderPrecomputeTextureParameters::Transmittance, Textures);
-		AtmosphereParameters.Set(RHICmdList, GetPixelShader(), 1, FAtmosphereShaderPrecomputeTextureParameters::DeltaE, Textures);
-		AtmosphereParameters.Set(RHICmdList, GetPixelShader(), 2, FAtmosphereShaderPrecomputeTextureParameters::DeltaSR, Textures);
-		AtmosphereParameters.Set(RHICmdList, GetPixelShader(), 3, FAtmosphereShaderPrecomputeTextureParameters::DeltaSM, Textures);
-		SetShaderValue(RHICmdList, GetPixelShader(), AtmosphereRParameter, AtmosphereR);
-		SetShaderValue(RHICmdList, GetPixelShader(), dhdHParameter, DhdH);
-		SetShaderValue(RHICmdList, GetPixelShader(), FirstOrderParameter, FirstOrder);
+		FGlobalShader::SetParameters<FViewUniformShaderParameters>(RHICmdList, RHICmdList.GetBoundPixelShader(), View.ViewUniformBuffer);
+		AtmosphereParameters.Set(RHICmdList, RHICmdList.GetBoundPixelShader(), 0, FAtmosphereShaderPrecomputeTextureParameters::Transmittance, Textures);
+		AtmosphereParameters.Set(RHICmdList, RHICmdList.GetBoundPixelShader(), 1, FAtmosphereShaderPrecomputeTextureParameters::DeltaE, Textures);
+		AtmosphereParameters.Set(RHICmdList, RHICmdList.GetBoundPixelShader(), 2, FAtmosphereShaderPrecomputeTextureParameters::DeltaSR, Textures);
+		AtmosphereParameters.Set(RHICmdList, RHICmdList.GetBoundPixelShader(), 3, FAtmosphereShaderPrecomputeTextureParameters::DeltaSM, Textures);
+		SetShaderValue(RHICmdList, RHICmdList.GetBoundPixelShader(), AtmosphereRParameter, AtmosphereR);
+		SetShaderValue(RHICmdList, RHICmdList.GetBoundPixelShader(), dhdHParameter, DhdH);
+		SetShaderValue(RHICmdList, RHICmdList.GetBoundPixelShader(), FirstOrderParameter, FirstOrder);
 	}
 };
 
@@ -902,10 +823,10 @@ class FAtmosphereInscatterNPS : public FGlobalShader
 	FAtmosphereInscatterNPS() {}
 
 public:
-	FAtmosphereShaderPrecomputeTextureParameters AtmosphereParameters;
-	FShaderParameter dhdHParameter;
-	FShaderParameter AtmosphereRParameter;
-	FShaderParameter FirstOrderParameter;
+	LAYOUT_FIELD(FAtmosphereShaderPrecomputeTextureParameters, AtmosphereParameters);
+	LAYOUT_FIELD(FShaderParameter, dhdHParameter);
+	LAYOUT_FIELD(FShaderParameter, AtmosphereRParameter);
+	LAYOUT_FIELD(FShaderParameter, FirstOrderParameter);
 
 	/* Initialization constructor. */
 	FAtmosphereInscatterNPS(const ShaderMetaType::CompiledShaderInitializerType& Initializer)
@@ -918,22 +839,14 @@ public:
 		FirstOrderParameter.Bind(Initializer.ParameterMap, TEXT("FirstOrder"));
 	}
 
-	// FShader interface.
-	virtual bool Serialize(FArchive& Ar) override
-	{
-		bool bShaderHasOutdatedParameters = FGlobalShader::Serialize(Ar);
-		Ar << AtmosphereParameters << dhdHParameter << AtmosphereRParameter << FirstOrderParameter;
-		return bShaderHasOutdatedParameters;
-	}
-
 	void SetParameters(FRHICommandList& RHICmdList, const FViewInfo& View, float AtmosphereR, const FVector4& DhdH, float FirstOrder, const FAtmosphereTextures* Textures)
 	{
-		FGlobalShader::SetParameters<FViewUniformShaderParameters>(RHICmdList, GetPixelShader(), View.ViewUniformBuffer);
-		AtmosphereParameters.Set(RHICmdList, GetPixelShader(), 0, FAtmosphereShaderPrecomputeTextureParameters::Transmittance, Textures);
-		AtmosphereParameters.Set(RHICmdList, GetPixelShader(), 1, FAtmosphereShaderPrecomputeTextureParameters::DeltaJ, Textures);
-		SetShaderValue(RHICmdList, GetPixelShader(), AtmosphereRParameter, AtmosphereR);
-		SetShaderValue(RHICmdList, GetPixelShader(), dhdHParameter, DhdH);
-		SetShaderValue(RHICmdList, GetPixelShader(), FirstOrderParameter, FirstOrder);
+		FGlobalShader::SetParameters<FViewUniformShaderParameters>(RHICmdList, RHICmdList.GetBoundPixelShader(), View.ViewUniformBuffer);
+		AtmosphereParameters.Set(RHICmdList, RHICmdList.GetBoundPixelShader(), 0, FAtmosphereShaderPrecomputeTextureParameters::Transmittance, Textures);
+		AtmosphereParameters.Set(RHICmdList, RHICmdList.GetBoundPixelShader(), 1, FAtmosphereShaderPrecomputeTextureParameters::DeltaJ, Textures);
+		SetShaderValue(RHICmdList, RHICmdList.GetBoundPixelShader(), AtmosphereRParameter, AtmosphereR);
+		SetShaderValue(RHICmdList, RHICmdList.GetBoundPixelShader(), dhdHParameter, DhdH);
+		SetShaderValue(RHICmdList, RHICmdList.GetBoundPixelShader(), FirstOrderParameter, FirstOrder);
 	}
 };
 
@@ -956,12 +869,6 @@ public:
 		:	FGlobalShader(Initializer)
 	{
 	}
-
-	virtual bool Serialize(FArchive& Ar) override
-	{
-		bool bShaderHasOutdatedParameters = FGlobalShader::Serialize(Ar);
-		return bShaderHasOutdatedParameters;
-	}
 };
 		
 class FAtmospherePrecomputeInscatterVS : public FGlobalShader
@@ -976,7 +883,7 @@ class FAtmospherePrecomputeInscatterVS : public FGlobalShader
 	/* Default constructor. */
 	FAtmospherePrecomputeInscatterVS() {}
 	
-	FShaderParameter AtmosphereLayerParameter;
+	LAYOUT_FIELD(FShaderParameter, AtmosphereLayerParameter);
 	
 public:
 	
@@ -989,14 +896,7 @@ public:
 
 	void SetParameters(FRHICommandList& RHICmdList, int AtmosphereLayer)
 	{
-		SetShaderValue(RHICmdList, GetVertexShader(), AtmosphereLayerParameter, AtmosphereLayer);
-	}
-
-	virtual bool Serialize(FArchive& Ar) override
-	{
-		bool bShaderHasOutdatedParameters = FGlobalShader::Serialize(Ar);
-		Ar << AtmosphereLayerParameter;
-		return bShaderHasOutdatedParameters;
+		SetShaderValue(RHICmdList, RHICmdList.GetBoundVertexShader(), AtmosphereLayerParameter, AtmosphereLayer);
 	}
 };
 
@@ -1014,10 +914,10 @@ class FAtmosphereCopyInscatterFPS : public FGlobalShader
 	FAtmosphereCopyInscatterFPS() {}
 
 public:
-	FAtmosphereShaderPrecomputeTextureParameters AtmosphereParameters;
-	FShaderParameter dhdHParameter;
-	FShaderParameter AtmosphereRParameter;
-	FShaderParameter AtmosphereLayerParameter;
+	LAYOUT_FIELD(FAtmosphereShaderPrecomputeTextureParameters, AtmosphereParameters);
+	LAYOUT_FIELD(FShaderParameter, dhdHParameter);
+	LAYOUT_FIELD(FShaderParameter, AtmosphereRParameter);
+	LAYOUT_FIELD(FShaderParameter, AtmosphereLayerParameter);
 
 	/* Initialization constructor. */
 	FAtmosphereCopyInscatterFPS(const ShaderMetaType::CompiledShaderInitializerType& Initializer)
@@ -1029,20 +929,12 @@ public:
 		AtmosphereLayerParameter.Bind(Initializer.ParameterMap, TEXT("AtmosphereLayer"));
 	}
 
-	// FShader interface.
-	virtual bool Serialize(FArchive& Ar) override
-	{
-		bool bShaderHasOutdatedParameters = FGlobalShader::Serialize(Ar);
-		Ar << AtmosphereParameters << dhdHParameter << AtmosphereRParameter << AtmosphereLayerParameter;
-		return bShaderHasOutdatedParameters;
-	}
-
 	void SetParameters(FRHICommandList& RHICmdList, float AtmosphereR, const FVector4& DhdH, int AtmosphereLayer, const FAtmosphereTextures* Textures)
 	{
-		AtmosphereParameters.Set(RHICmdList, GetPixelShader(), 0, FAtmosphereShaderPrecomputeTextureParameters::Inscatter, Textures);
-		SetShaderValue(RHICmdList, GetPixelShader(), AtmosphereRParameter, AtmosphereR);
-		SetShaderValue(RHICmdList, GetPixelShader(), dhdHParameter, DhdH);
-		SetShaderValue(RHICmdList, GetPixelShader(), AtmosphereLayerParameter, AtmosphereLayer);
+		AtmosphereParameters.Set(RHICmdList, RHICmdList.GetBoundPixelShader(), 0, FAtmosphereShaderPrecomputeTextureParameters::Inscatter, Textures);
+		SetShaderValue(RHICmdList, RHICmdList.GetBoundPixelShader(), AtmosphereRParameter, AtmosphereR);
+		SetShaderValue(RHICmdList, RHICmdList.GetBoundPixelShader(), dhdHParameter, DhdH);
+		SetShaderValue(RHICmdList, RHICmdList.GetBoundPixelShader(), AtmosphereLayerParameter, AtmosphereLayer);
 	}
 };
 
@@ -1059,10 +951,10 @@ class FAtmosphereCopyInscatterFBackPS : public FGlobalShader
 	FAtmosphereCopyInscatterFBackPS() {}
 
 public:
-	FAtmosphereShaderPrecomputeTextureParameters AtmosphereParameters;
-	FShaderParameter dhdHParameter;
-	FShaderParameter AtmosphereRParameter;
-	FShaderParameter AtmosphereLayerParameter;
+	LAYOUT_FIELD(FAtmosphereShaderPrecomputeTextureParameters, AtmosphereParameters);
+	LAYOUT_FIELD(FShaderParameter, dhdHParameter);
+	LAYOUT_FIELD(FShaderParameter, AtmosphereRParameter);
+	LAYOUT_FIELD(FShaderParameter, AtmosphereLayerParameter);
 
 	/* Initialization constructor. */
 	FAtmosphereCopyInscatterFBackPS(const ShaderMetaType::CompiledShaderInitializerType& Initializer)
@@ -1074,20 +966,12 @@ public:
 		AtmosphereLayerParameter.Bind(Initializer.ParameterMap, TEXT("AtmosphereLayer"));
 	}
 
-	// FShader interface.
-	virtual bool Serialize(FArchive& Ar) override
-	{
-		bool bShaderHasOutdatedParameters = FGlobalShader::Serialize(Ar);
-		Ar << AtmosphereParameters << dhdHParameter << AtmosphereRParameter << AtmosphereLayerParameter;
-		return bShaderHasOutdatedParameters;
-	}
-
 	void SetParameters(FRHICommandList& RHICmdList, float AtmosphereR, const FVector4& DhdH, int AtmosphereLayer, const FAtmosphereTextures* Textures)
 	{
-		AtmosphereParameters.Set(RHICmdList, GetPixelShader(), 0, FAtmosphereShaderPrecomputeTextureParameters::DeltaSR, Textures);
-		SetShaderValue(RHICmdList, GetPixelShader(), AtmosphereRParameter, AtmosphereR);
-		SetShaderValue(RHICmdList, GetPixelShader(), dhdHParameter, DhdH);
-		SetShaderValue(RHICmdList, GetPixelShader(), AtmosphereLayerParameter, AtmosphereLayer);
+		AtmosphereParameters.Set(RHICmdList, RHICmdList.GetBoundPixelShader(), 0, FAtmosphereShaderPrecomputeTextureParameters::DeltaSR, Textures);
+		SetShaderValue(RHICmdList, RHICmdList.GetBoundPixelShader(), AtmosphereRParameter, AtmosphereR);
+		SetShaderValue(RHICmdList, RHICmdList.GetBoundPixelShader(), dhdHParameter, DhdH);
+		SetShaderValue(RHICmdList, RHICmdList.GetBoundPixelShader(), AtmosphereLayerParameter, AtmosphereLayer);
 	}
 };
 
@@ -1167,7 +1051,7 @@ FIntPoint FAtmosphericFogSceneInfo::GetTextureSize()
 	return AtmosphereTextures->AtmosphereTransmittance->GetDesc().Extent;
 }
 
-void FAtmosphericFogSceneInfo::DrawQuad(FRHICommandList& RHICmdList, const FIntRect& ViewRect, FShader* VertexShader)
+void FAtmosphericFogSceneInfo::DrawQuad(FRHICommandList& RHICmdList, const FIntRect& ViewRect, const TShaderRef<FShader>& VertexShader)
 {
 	// Draw a quad mapping scene color to the view's render target
 	DrawRectangle( 
@@ -1216,13 +1100,13 @@ void FAtmosphericFogSceneInfo::RenderAtmosphereShaders(FRHICommandList& RHICmdLi
 				TShaderMapRef<FAtmosphereTransmittancePS> PixelShader(ShaderMap);
 
 				GraphicsPSOInit.BoundShaderState.VertexDeclarationRHI = GFilterVertexDeclaration.VertexDeclarationRHI;
-				GraphicsPSOInit.BoundShaderState.VertexShaderRHI = GETSAFERHISHADER_VERTEX(*VertexShader);
-				GraphicsPSOInit.BoundShaderState.PixelShaderRHI = GETSAFERHISHADER_PIXEL(*PixelShader);
+				GraphicsPSOInit.BoundShaderState.VertexShaderRHI = VertexShader.GetVertexShader();
+				GraphicsPSOInit.BoundShaderState.PixelShaderRHI = PixelShader.GetPixelShader();
 				GraphicsPSOInit.PrimitiveType = PT_TriangleList;
 
 				SetGraphicsPipelineState(RHICmdList, GraphicsPSOInit);
 				PixelShader->SetParameters(RHICmdList, View);
-				DrawQuad(RHICmdList, ViewRect, *VertexShader);
+				DrawQuad(RHICmdList, ViewRect, VertexShader);
 			}
 			RHICmdList.EndRenderPass();
 			RHICmdList.CopyToResolveTarget(DestRenderTarget.TargetableTexture, DestRenderTarget.ShaderResourceTexture, FResolveParams());
@@ -1240,15 +1124,15 @@ void FAtmosphericFogSceneInfo::RenderAtmosphereShaders(FRHICommandList& RHICmdLi
 				TShaderMapRef<FAtmosphereIrradiance1PS> PixelShader(ShaderMap);
 
 				GraphicsPSOInit.BoundShaderState.VertexDeclarationRHI = GFilterVertexDeclaration.VertexDeclarationRHI;
-				GraphicsPSOInit.BoundShaderState.VertexShaderRHI = GETSAFERHISHADER_VERTEX(*VertexShader);
-				GraphicsPSOInit.BoundShaderState.PixelShaderRHI = GETSAFERHISHADER_PIXEL(*PixelShader);
+				GraphicsPSOInit.BoundShaderState.VertexShaderRHI = VertexShader.GetVertexShader();
+				GraphicsPSOInit.BoundShaderState.PixelShaderRHI = PixelShader.GetPixelShader();
 				GraphicsPSOInit.PrimitiveType = PT_TriangleList;
 
 				SetGraphicsPipelineState(RHICmdList, GraphicsPSOInit);
 
 				PixelShader->SetParameters(RHICmdList, AtmosphereTextures);
 
-				DrawQuad(RHICmdList, ViewRect, *VertexShader);
+				DrawQuad(RHICmdList, ViewRect, VertexShader);
 			}
 			RHICmdList.EndRenderPass();
 			RHICmdList.CopyToResolveTarget(DestRenderTarget.TargetableTexture, DestRenderTarget.ShaderResourceTexture, FResolveParams());
@@ -1274,11 +1158,11 @@ void FAtmosphericFogSceneInfo::RenderAtmosphereShaders(FRHICommandList& RHICmdLi
 					TShaderMapRef<FAtmosphereInscatter1PS> PixelShader(ShaderMap);
 
 					GraphicsPSOInit.BoundShaderState.VertexDeclarationRHI = GFilterVertexDeclaration.VertexDeclarationRHI;
-					GraphicsPSOInit.BoundShaderState.VertexShaderRHI = GETSAFERHISHADER_VERTEX(*VertexShader);
+					GraphicsPSOInit.BoundShaderState.VertexShaderRHI = VertexShader.GetVertexShader();
 #if PLATFORM_SUPPORTS_GEOMETRY_SHADERS
-					GraphicsPSOInit.BoundShaderState.GeometryShaderRHI = GETSAFERHISHADER_GEOMETRY(*GeometryShader);
+					GraphicsPSOInit.BoundShaderState.GeometryShaderRHI = GeometryShader.GetGeometryShader();
 #endif
-					GraphicsPSOInit.BoundShaderState.PixelShaderRHI = GETSAFERHISHADER_PIXEL(*PixelShader);
+					GraphicsPSOInit.BoundShaderState.PixelShaderRHI = PixelShader.GetPixelShader();
 					GraphicsPSOInit.PrimitiveType = PT_TriangleList;
 
 					SetGraphicsPipelineState(RHICmdList, GraphicsPSOInit);
@@ -1293,7 +1177,7 @@ void FAtmosphericFogSceneInfo::RenderAtmosphereShaders(FRHICommandList& RHICmdLi
 						GeometryShader->SetParameters(RHICmdList, Layer);
 					}
 					PixelShader->SetParameters(RHICmdList, View, r, DhdH, AtmosphereTextures);
-					DrawQuad(RHICmdList, ViewRect, *VertexShader);
+					DrawQuad(RHICmdList, ViewRect, VertexShader);
 				}
 				RHICmdList.EndRenderPass();
 				if (Atmosphere3DTextureIndex == Component->PrecomputeParams.InscatterAltitudeSampleNum - 1)
@@ -1333,11 +1217,11 @@ void FAtmosphericFogSceneInfo::RenderAtmosphereShaders(FRHICommandList& RHICmdLi
 					TShaderMapRef<FAtmosphereCopyInscatter1PS> PixelShader(ShaderMap);
 
 					GraphicsPSOInit.BoundShaderState.VertexDeclarationRHI = GFilterVertexDeclaration.VertexDeclarationRHI;
-					GraphicsPSOInit.BoundShaderState.VertexShaderRHI = GETSAFERHISHADER_VERTEX(*VertexShader);
+					GraphicsPSOInit.BoundShaderState.VertexShaderRHI = VertexShader.GetVertexShader();
 #if PLATFORM_SUPPORTS_GEOMETRY_SHADERS
-					GraphicsPSOInit.BoundShaderState.GeometryShaderRHI = GETSAFERHISHADER_GEOMETRY(*GeometryShader);
+					GraphicsPSOInit.BoundShaderState.GeometryShaderRHI = GeometryShader.GetGeometryShader();
 #endif
-					GraphicsPSOInit.BoundShaderState.PixelShaderRHI = GETSAFERHISHADER_PIXEL(*PixelShader);
+					GraphicsPSOInit.BoundShaderState.PixelShaderRHI = PixelShader.GetPixelShader();
 					GraphicsPSOInit.PrimitiveType = PT_TriangleList;
 
 					SetGraphicsPipelineState(RHICmdList, GraphicsPSOInit);
@@ -1351,7 +1235,7 @@ void FAtmosphericFogSceneInfo::RenderAtmosphereShaders(FRHICommandList& RHICmdLi
 						GeometryShader->SetParameters(RHICmdList, Layer);
 					}
 					PixelShader->SetParameters(RHICmdList, View, r, DhdH, Layer, AtmosphereTextures);
-					DrawQuad(RHICmdList, ViewRect, *VertexShader);
+					DrawQuad(RHICmdList, ViewRect, VertexShader);
 				}
 				RHICmdList.EndRenderPass();
 				if (Atmosphere3DTextureIndex == Component->PrecomputeParams.InscatterAltitudeSampleNum - 1)
@@ -1377,11 +1261,11 @@ void FAtmosphericFogSceneInfo::RenderAtmosphereShaders(FRHICommandList& RHICmdLi
 					TShaderMapRef<FAtmosphereInscatterSPS> PixelShader(ShaderMap);
 
 					GraphicsPSOInit.BoundShaderState.VertexDeclarationRHI = GFilterVertexDeclaration.VertexDeclarationRHI;
-					GraphicsPSOInit.BoundShaderState.VertexShaderRHI = GETSAFERHISHADER_VERTEX(*VertexShader);
+					GraphicsPSOInit.BoundShaderState.VertexShaderRHI = VertexShader.GetVertexShader();
 #if PLATFORM_SUPPORTS_GEOMETRY_SHADERS
-					GraphicsPSOInit.BoundShaderState.GeometryShaderRHI = GETSAFERHISHADER_GEOMETRY(*GeometryShader);
+					GraphicsPSOInit.BoundShaderState.GeometryShaderRHI = GeometryShader.GetGeometryShader();
 #endif
-					GraphicsPSOInit.BoundShaderState.PixelShaderRHI = GETSAFERHISHADER_PIXEL(*PixelShader);
+					GraphicsPSOInit.BoundShaderState.PixelShaderRHI = PixelShader.GetPixelShader();
 					GraphicsPSOInit.PrimitiveType = PT_TriangleList;
 
 					SetGraphicsPipelineState(RHICmdList, GraphicsPSOInit);
@@ -1395,7 +1279,7 @@ void FAtmosphericFogSceneInfo::RenderAtmosphereShaders(FRHICommandList& RHICmdLi
 						GeometryShader->SetParameters(RHICmdList, Layer);
 					}
 					PixelShader->SetParameters(RHICmdList, View, r, DhdH, AtmoshpereOrder == 2 ? 1.f : 0.f, AtmosphereTextures);
-					DrawQuad(RHICmdList, ViewRect, *VertexShader);
+					DrawQuad(RHICmdList, ViewRect, VertexShader);
 				}
 				RHICmdList.EndRenderPass();
 				if (Atmosphere3DTextureIndex == Component->PrecomputeParams.InscatterAltitudeSampleNum - 1)
@@ -1417,15 +1301,15 @@ void FAtmosphericFogSceneInfo::RenderAtmosphereShaders(FRHICommandList& RHICmdLi
 				TShaderMapRef<FAtmosphereIrradianceNPS> PixelShader(ShaderMap);
 
 				GraphicsPSOInit.BoundShaderState.VertexDeclarationRHI = GFilterVertexDeclaration.VertexDeclarationRHI;
-				GraphicsPSOInit.BoundShaderState.VertexShaderRHI = GETSAFERHISHADER_VERTEX(*VertexShader);
-				GraphicsPSOInit.BoundShaderState.PixelShaderRHI = GETSAFERHISHADER_PIXEL(*PixelShader);
+				GraphicsPSOInit.BoundShaderState.VertexShaderRHI = VertexShader.GetVertexShader();
+				GraphicsPSOInit.BoundShaderState.PixelShaderRHI = PixelShader.GetPixelShader();
 				GraphicsPSOInit.PrimitiveType = PT_TriangleList;
 
 				SetGraphicsPipelineState(RHICmdList, GraphicsPSOInit);
 
 				PixelShader->SetParameters(RHICmdList, View, AtmoshpereOrder == 2 ? 1.f : 0.f, AtmosphereTextures);
 
-				DrawQuad(RHICmdList, ViewRect, *VertexShader);
+				DrawQuad(RHICmdList, ViewRect, VertexShader);
 			}
 			RHICmdList.EndRenderPass();
 			RHICmdList.CopyToResolveTarget(DestRenderTarget.TargetableTexture, DestRenderTarget.ShaderResourceTexture, FResolveParams());
@@ -1448,11 +1332,11 @@ void FAtmosphericFogSceneInfo::RenderAtmosphereShaders(FRHICommandList& RHICmdLi
 					TShaderMapRef<FAtmosphereInscatterNPS> PixelShader(ShaderMap);
 
 					GraphicsPSOInit.BoundShaderState.VertexDeclarationRHI = GFilterVertexDeclaration.VertexDeclarationRHI;
-					GraphicsPSOInit.BoundShaderState.VertexShaderRHI = GETSAFERHISHADER_VERTEX(*VertexShader);
+					GraphicsPSOInit.BoundShaderState.VertexShaderRHI = VertexShader.GetVertexShader();
 #if PLATFORM_SUPPORTS_GEOMETRY_SHADERS
-					GraphicsPSOInit.BoundShaderState.GeometryShaderRHI = GETSAFERHISHADER_GEOMETRY(*GeometryShader);
+					GraphicsPSOInit.BoundShaderState.GeometryShaderRHI = GeometryShader.GetGeometryShader();
 #endif
-					GraphicsPSOInit.BoundShaderState.PixelShaderRHI = GETSAFERHISHADER_PIXEL(*PixelShader);
+					GraphicsPSOInit.BoundShaderState.PixelShaderRHI = PixelShader.GetPixelShader();
 					GraphicsPSOInit.PrimitiveType = PT_TriangleList;
 
 					SetGraphicsPipelineState(RHICmdList, GraphicsPSOInit);
@@ -1466,7 +1350,7 @@ void FAtmosphericFogSceneInfo::RenderAtmosphereShaders(FRHICommandList& RHICmdLi
 						GeometryShader->SetParameters(RHICmdList, Layer);
 					}
 					PixelShader->SetParameters(RHICmdList, View, r, DhdH, AtmoshpereOrder == 2 ? 1.f : 0.f, AtmosphereTextures);
-					DrawQuad(RHICmdList, ViewRect, *VertexShader);
+					DrawQuad(RHICmdList, ViewRect, VertexShader);
 				}
 				RHICmdList.EndRenderPass();
 				if (Atmosphere3DTextureIndex == Component->PrecomputeParams.InscatterAltitudeSampleNum - 1)
@@ -1491,15 +1375,15 @@ void FAtmosphericFogSceneInfo::RenderAtmosphereShaders(FRHICommandList& RHICmdLi
 				TShaderMapRef<FAtmosphereCopyIrradiancePS> PixelShader(ShaderMap);
 
 				GraphicsPSOInit.BoundShaderState.VertexDeclarationRHI = GFilterVertexDeclaration.VertexDeclarationRHI;
-				GraphicsPSOInit.BoundShaderState.VertexShaderRHI = GETSAFERHISHADER_VERTEX(*VertexShader);
-				GraphicsPSOInit.BoundShaderState.PixelShaderRHI = GETSAFERHISHADER_PIXEL(*PixelShader);
+				GraphicsPSOInit.BoundShaderState.VertexShaderRHI = VertexShader.GetVertexShader();
+				GraphicsPSOInit.BoundShaderState.PixelShaderRHI = PixelShader.GetPixelShader();
 				GraphicsPSOInit.PrimitiveType = PT_TriangleList;
 
 				SetGraphicsPipelineState(RHICmdList, GraphicsPSOInit);
 
 				PixelShader->SetParameters(RHICmdList, AtmosphereTextures);
 
-				DrawQuad(RHICmdList, ViewRect, *VertexShader);
+				DrawQuad(RHICmdList, ViewRect, VertexShader);
 			}
 			RHICmdList.EndRenderPass();
 			RHICmdList.CopyToResolveTarget(DestRenderTarget.TargetableTexture, DestRenderTarget.ShaderResourceTexture, FResolveParams());
@@ -1526,11 +1410,11 @@ void FAtmosphericFogSceneInfo::RenderAtmosphereShaders(FRHICommandList& RHICmdLi
 					TShaderMapRef<FAtmosphereCopyInscatterNPS> PixelShader(ShaderMap);
 
 					GraphicsPSOInit.BoundShaderState.VertexDeclarationRHI = GFilterVertexDeclaration.VertexDeclarationRHI;
-					GraphicsPSOInit.BoundShaderState.VertexShaderRHI = GETSAFERHISHADER_VERTEX(*VertexShader);
+					GraphicsPSOInit.BoundShaderState.VertexShaderRHI = VertexShader.GetVertexShader();
 #if PLATFORM_SUPPORTS_GEOMETRY_SHADERS
-					GraphicsPSOInit.BoundShaderState.GeometryShaderRHI = GETSAFERHISHADER_GEOMETRY(*GeometryShader);
+					GraphicsPSOInit.BoundShaderState.GeometryShaderRHI = GeometryShader.GetGeometryShader();
 #endif
-					GraphicsPSOInit.BoundShaderState.PixelShaderRHI = GETSAFERHISHADER_PIXEL(*PixelShader);
+					GraphicsPSOInit.BoundShaderState.PixelShaderRHI = PixelShader.GetPixelShader();
 					GraphicsPSOInit.PrimitiveType = PT_TriangleList;
 
 					SetGraphicsPipelineState(RHICmdList, GraphicsPSOInit);
@@ -1544,7 +1428,7 @@ void FAtmosphericFogSceneInfo::RenderAtmosphereShaders(FRHICommandList& RHICmdLi
 						GeometryShader->SetParameters(RHICmdList, Layer);
 					}
 					PixelShader->SetParameters(RHICmdList, View, r, DhdH, Layer, AtmosphereTextures);
-					DrawQuad(RHICmdList, ViewRect, *VertexShader);
+					DrawQuad(RHICmdList, ViewRect, VertexShader);
 				}
 				RHICmdList.EndRenderPass();
 				if (Atmosphere3DTextureIndex == Component->PrecomputeParams.InscatterAltitudeSampleNum - 1)
@@ -1573,11 +1457,11 @@ void FAtmosphericFogSceneInfo::RenderAtmosphereShaders(FRHICommandList& RHICmdLi
 					TShaderMapRef<FAtmosphereCopyInscatterFPS> PixelShader(ShaderMap);
 
 					GraphicsPSOInit.BoundShaderState.VertexDeclarationRHI = GFilterVertexDeclaration.VertexDeclarationRHI;
-					GraphicsPSOInit.BoundShaderState.VertexShaderRHI = GETSAFERHISHADER_VERTEX(*VertexShader);
+					GraphicsPSOInit.BoundShaderState.VertexShaderRHI = VertexShader.GetVertexShader();
 #if PLATFORM_SUPPORTS_GEOMETRY_SHADERS
-					GraphicsPSOInit.BoundShaderState.GeometryShaderRHI = GETSAFERHISHADER_GEOMETRY(*GeometryShader);
+					GraphicsPSOInit.BoundShaderState.GeometryShaderRHI = GeometryShader.GetGeometryShader();
 #endif
-					GraphicsPSOInit.BoundShaderState.PixelShaderRHI = GETSAFERHISHADER_PIXEL(*PixelShader);
+					GraphicsPSOInit.BoundShaderState.PixelShaderRHI = PixelShader.GetPixelShader();
 					GraphicsPSOInit.PrimitiveType = PT_TriangleList;
 
 					SetGraphicsPipelineState(RHICmdList, GraphicsPSOInit);
@@ -1591,7 +1475,7 @@ void FAtmosphericFogSceneInfo::RenderAtmosphereShaders(FRHICommandList& RHICmdLi
 						GeometryShader->SetParameters(RHICmdList, Layer);
 					}
 					PixelShader->SetParameters(RHICmdList, r, DhdH, Layer, AtmosphereTextures);
-					DrawQuad(RHICmdList, ViewRect, *VertexShader);
+					DrawQuad(RHICmdList, ViewRect, VertexShader);
 				}
 				RHICmdList.EndRenderPass();
 				if (Atmosphere3DTextureIndex == Component->PrecomputeParams.InscatterAltitudeSampleNum - 1)
@@ -1617,11 +1501,11 @@ void FAtmosphericFogSceneInfo::RenderAtmosphereShaders(FRHICommandList& RHICmdLi
 					TShaderMapRef<FAtmosphereCopyInscatterFBackPS> PixelShader(ShaderMap);
 
 					GraphicsPSOInit.BoundShaderState.VertexDeclarationRHI = GFilterVertexDeclaration.VertexDeclarationRHI;
-					GraphicsPSOInit.BoundShaderState.VertexShaderRHI = GETSAFERHISHADER_VERTEX(*VertexShader);
+					GraphicsPSOInit.BoundShaderState.VertexShaderRHI = VertexShader.GetVertexShader();
 #if PLATFORM_SUPPORTS_GEOMETRY_SHADERS
-					GraphicsPSOInit.BoundShaderState.GeometryShaderRHI = GETSAFERHISHADER_GEOMETRY(*GeometryShader);
+					GraphicsPSOInit.BoundShaderState.GeometryShaderRHI = GeometryShader.GetGeometryShader();
 #endif
-					GraphicsPSOInit.BoundShaderState.PixelShaderRHI = GETSAFERHISHADER_PIXEL(*PixelShader);
+					GraphicsPSOInit.BoundShaderState.PixelShaderRHI = PixelShader.GetPixelShader();
 					GraphicsPSOInit.PrimitiveType = PT_TriangleList;
 
 					SetGraphicsPipelineState(RHICmdList, GraphicsPSOInit);
@@ -1635,7 +1519,7 @@ void FAtmosphericFogSceneInfo::RenderAtmosphereShaders(FRHICommandList& RHICmdLi
 						GeometryShader->SetParameters(RHICmdList, Layer);
 					}
 					PixelShader->SetParameters(RHICmdList, r, DhdH, Layer, AtmosphereTextures);
-					DrawQuad(RHICmdList, ViewRect, *VertexShader);
+					DrawQuad(RHICmdList, ViewRect, VertexShader);
 				}
 				RHICmdList.EndRenderPass();
 				if (Atmosphere3DTextureIndex == Component->PrecomputeParams.InscatterAltitudeSampleNum - 1)
diff --git a/Engine/Source/Runtime/Renderer/Private/AtmosphereRendering.h b/Engine/Source/Runtime/Renderer/Private/AtmosphereRendering.h
--- a/Engine/Source/Runtime/Renderer/Private/AtmosphereRendering.h
+++ b/Engine/Source/Runtime/Renderer/Private/AtmosphereRendering.h
@@ -92,7 +92,7 @@ public:
 private:
 	/** Atmosphere pre-computation related functions */
 	FIntPoint GetTextureSize();
-	inline void DrawQuad(FRHICommandList& RHICmdList, const FIntRect& ViewRect, FShader* VertexShader);
+	inline void DrawQuad(FRHICommandList& RHICmdList, const FIntRect& ViewRect, const TShaderRef<FShader>& VertexShader);
 	void GetLayerValue(int Layer, float& AtmosphereR, FVector4& DhdH);
 	void RenderAtmosphereShaders(FRHICommandList& RHICmdList, FGraphicsPipelineStateInitializer& GraphicsPSOInit, const FViewInfo& View, const FIntRect& ViewRect);
 	void PrecomputeAtmosphereData(FRHICommandListImmediate& RHICmdList, const FViewInfo* View, FSceneViewFamily& ViewFamily);
diff --git a/Engine/Source/Runtime/Renderer/Private/BasePassRendering.cpp b/Engine/Source/Runtime/Renderer/Private/BasePassRendering.cpp
--- a/Engine/Source/Runtime/Renderer/Private/BasePassRendering.cpp
+++ b/Engine/Source/Runtime/Renderer/Private/BasePassRendering.cpp
@@ -353,10 +353,10 @@ void GetUniformBasePassShaders(
 	ERHIFeatureLevel::Type FeatureLevel,
 	bool bEnableAtmosphericFog,
 	bool bEnableSkyLight,
-	FBaseHS*& HullShader,
-	FBaseDS*& DomainShader,
-	TBasePassVertexShaderPolicyParamType<FUniformLightMapPolicy>*& VertexShader,
-	TBasePassPixelShaderPolicyParamType<FUniformLightMapPolicy>*& PixelShader
+	TShaderRef<FBaseHS>& HullShader,
+	TShaderRef<FBaseDS>& DomainShader,
+	TShaderRef<TBasePassVertexShaderPolicyParamType<FUniformLightMapPolicy>>& VertexShader,
+	TShaderRef<TBasePassPixelShaderPolicyParamType<FUniformLightMapPolicy>>& PixelShader
 	)
 {
 	const EMaterialTessellationMode MaterialTessellationMode = Material.GetTessellationMode();
@@ -370,7 +370,7 @@ void GetUniformBasePassShaders(
 		DomainShader = Material.GetShader<TBasePassDS<TUniformLightMapPolicy<Policy> > >(VertexFactoryType);
 		
 		// Metal requires matching permutations, but no other platform should worry about this complication.
-		if (bEnableAtmosphericFog && DomainShader && IsMetalPlatform(EShaderPlatform(DomainShader->GetTarget().Platform)))
+		if (bEnableAtmosphericFog && DomainShader.IsValid() && IsMetalPlatform(EShaderPlatform(DomainShader->GetTarget().Platform)))
 		{
 			HullShader = Material.GetShader<TBasePassHS<TUniformLightMapPolicy<Policy>, true > >(VertexFactoryType);
 		}
@@ -382,19 +382,19 @@ void GetUniformBasePassShaders(
 
 	if (bEnableAtmosphericFog)
 	{
-		VertexShader = (TBasePassVertexShaderPolicyParamType<FUniformLightMapPolicy>*)Material.GetShader<TBasePassVS<TUniformLightMapPolicy<Policy>, true> >(VertexFactoryType);
+		VertexShader = TShaderRef<TBasePassVertexShaderPolicyParamType<FUniformLightMapPolicy>>::ReinterpretCast(Material.GetShader<TBasePassVS<TUniformLightMapPolicy<Policy>, true> >(VertexFactoryType));
 	}
 	else
 	{
-		VertexShader = (TBasePassVertexShaderPolicyParamType<FUniformLightMapPolicy>*)Material.GetShader<TBasePassVS<TUniformLightMapPolicy<Policy>, false> >(VertexFactoryType);
+		VertexShader = TShaderRef<TBasePassVertexShaderPolicyParamType<FUniformLightMapPolicy>>::ReinterpretCast(Material.GetShader<TBasePassVS<TUniformLightMapPolicy<Policy>, false> >(VertexFactoryType));
 	}
 	if (bEnableSkyLight)
 	{
-		PixelShader = (TBasePassPixelShaderPolicyParamType<FUniformLightMapPolicy>*)Material.GetShader<TBasePassPS<TUniformLightMapPolicy<Policy>, true> >(VertexFactoryType);
+		PixelShader = TShaderRef<TBasePassPixelShaderPolicyParamType<FUniformLightMapPolicy>>::ReinterpretCast(Material.GetShader<TBasePassPS<TUniformLightMapPolicy<Policy>, true> >(VertexFactoryType));
 	}
 	else
 	{
-		PixelShader = (TBasePassPixelShaderPolicyParamType<FUniformLightMapPolicy>*)Material.GetShader<TBasePassPS<TUniformLightMapPolicy<Policy>, false> >(VertexFactoryType);
+		PixelShader = TShaderRef<TBasePassPixelShaderPolicyParamType<FUniformLightMapPolicy>>::ReinterpretCast(Material.GetShader<TBasePassPS<TUniformLightMapPolicy<Policy>, false> >(VertexFactoryType));
 	}
 }
 
@@ -406,10 +406,10 @@ void GetBasePassShaders<FUniformLightMapPolicy>(
 	ERHIFeatureLevel::Type FeatureLevel,
 	bool bEnableAtmosphericFog,
 	bool bEnableSkyLight,
-	FBaseHS*& HullShader,
-	FBaseDS*& DomainShader,
-	TBasePassVertexShaderPolicyParamType<FUniformLightMapPolicy>*& VertexShader,
-	TBasePassPixelShaderPolicyParamType<FUniformLightMapPolicy>*& PixelShader
+	TShaderRef<FBaseHS>& HullShader,
+	TShaderRef<FBaseDS>& DomainShader,
+	TShaderRef<TBasePassVertexShaderPolicyParamType<FUniformLightMapPolicy>>& VertexShader,
+	TShaderRef<TBasePassPixelShaderPolicyParamType<FUniformLightMapPolicy>>& PixelShader
 	)
 {
 	switch (LightMapPolicy.GetIndirectPolicy())
@@ -1020,7 +1020,7 @@ void FBasePassMeshProcessor::Process(
 	}
 	else
 	{
-		SortKey = CalculateBasePassMeshStaticSortKey(EarlyZPassMode, BlendMode, BasePassShaders.VertexShader, BasePassShaders.PixelShader);
+		SortKey = CalculateBasePassMeshStaticSortKey(EarlyZPassMode, BlendMode, BasePassShaders.VertexShader.GetShader(), BasePassShaders.PixelShader.GetShader());
 	}
 
 	BuildMeshDrawCommands(
@@ -1211,8 +1211,9 @@ void FBasePassMeshProcessor::AddMeshBatch(const FMeshBatch& RESTRICT MeshBatch,
 		const EBlendMode BlendMode = Material.GetBlendMode();
 		const FMaterialShadingModelField ShadingModels = Material.GetShadingModels();
 		const bool bIsTranslucent = IsTranslucentBlendMode(BlendMode);
-		const ERasterizerFillMode MeshFillMode = ComputeMeshFillMode(MeshBatch, Material);
-		const ERasterizerCullMode MeshCullMode = ComputeMeshCullMode(MeshBatch, Material);
+		const FMeshDrawingPolicyOverrideSettings OverrideSettings = ComputeMeshOverrideSettings(MeshBatch);
+		const ERasterizerFillMode MeshFillMode = ComputeMeshFillMode(MeshBatch, Material, OverrideSettings);
+		const ERasterizerCullMode MeshCullMode = ComputeMeshCullMode(MeshBatch, Material, OverrideSettings);
 
 
 		bool bShouldDraw = false;
diff --git a/Engine/Source/Runtime/Renderer/Private/BasePassRendering.h b/Engine/Source/Runtime/Renderer/Private/BasePassRendering.h
--- a/Engine/Source/Runtime/Renderer/Private/BasePassRendering.h
+++ b/Engine/Source/Runtime/Renderer/Private/BasePassRendering.h
@@ -190,6 +190,7 @@ public:
 template<typename LightMapPolicyType>
 class TBasePassVertexShaderPolicyParamType : public FMeshMaterialShader, public LightMapPolicyType::VertexParametersType
 {
+	DECLARE_INLINE_TYPE_LAYOUT_EXPLICIT_BASES(TBasePassVertexShaderPolicyParamType, NonVirtual, FMeshMaterialShader, typename LightMapPolicyType::VertexParametersType);
 protected:
 
 	TBasePassVertexShaderPolicyParamType() {}
@@ -209,14 +210,6 @@ public:
 		FForwardLightingParameters::ModifyCompilationEnvironment(Parameters.Platform, OutEnvironment);
 	}
 
-	virtual bool Serialize(FArchive& Ar)
-	{
-		bool bShaderHasOutdatedParameters = FMeshMaterialShader::Serialize(Ar);
-		LightMapPolicyType::VertexParametersType::Serialize(Ar);
-		Ar << ReflectionCaptureBuffer;
-		return bShaderHasOutdatedParameters;
-	}
-
 	void GetShaderBindings(
 		const FScene* Scene,
 		ERHIFeatureLevel::Type FeatureLevel,
@@ -228,6 +221,7 @@ public:
 		FMeshDrawSingleShaderBindings& ShaderBindings) const;
 
 	void GetElementShaderBindings(
+		const FShaderMapPointerTable& PointerTable,
 		const FScene* Scene,
 		const FSceneView* ViewIfDynamicMeshCommand,
 		const FVertexFactory* VertexFactory,
@@ -240,7 +234,7 @@ public:
 		FMeshDrawSingleShaderBindings& ShaderBindings,
 		FVertexInputStreamArray& VertexStreams) const;
 
-	FShaderUniformBufferParameter ReflectionCaptureBuffer;
+	LAYOUT_FIELD(FShaderUniformBufferParameter, ReflectionCaptureBuffer);
 };
 
 
@@ -255,9 +249,8 @@ template<typename LightMapPolicyType>
 class TBasePassVertexShaderBaseType : public TBasePassVertexShaderPolicyParamType<LightMapPolicyType>
 {
 	typedef TBasePassVertexShaderPolicyParamType<LightMapPolicyType> Super;
-
+	DECLARE_INLINE_TYPE_LAYOUT(TBasePassVertexShaderBaseType, NonVirtual);
 protected:
-
 	TBasePassVertexShaderBaseType(const FMeshMaterialShaderType::CompiledShaderInitializerType& Initializer) : Super(Initializer) {}
 
 	TBasePassVertexShaderBaseType() {}
@@ -274,6 +267,9 @@ public:
 		LightMapPolicyType::ModifyCompilationEnvironment(Parameters, OutEnvironment);
 		Super::ModifyCompilationEnvironment(Parameters, OutEnvironment);
 	}
+
+	
+	
 };
 
 template<typename LightMapPolicyType, bool bEnableAtmosphericFog>
@@ -300,7 +296,7 @@ public:
 		const bool bProjectAllowsAtmosphericFog = !SupportAtmosphericFog || SupportAtmosphericFog->GetValueOnAnyThread() != 0 || bForceAllPermutations;
 
 		bool bShouldCache = Super::ShouldCompilePermutation(Parameters);
-		bShouldCache &= (bEnableAtmosphericFog && bProjectAllowsAtmosphericFog && IsTranslucentBlendMode(Parameters.Material->GetBlendMode())) || !bEnableAtmosphericFog;
+		bShouldCache &= (bEnableAtmosphericFog && bProjectAllowsAtmosphericFog && IsTranslucentBlendMode(Parameters.MaterialParameters.BlendMode)) || !bEnableAtmosphericFog;
 
 		return bShouldCache
 			&& (IsFeatureLevelSupported(Parameters.Platform, ERHIFeatureLevel::SM5));
@@ -378,13 +374,6 @@ protected:
 		// Re-use vertex shader compilation environment
 		TBasePassVS<LightMapPolicyType,false>::ModifyCompilationEnvironment(Parameters, OutEnvironment);
 	}
-
-public:
-	virtual bool Serialize(FArchive& Ar) override
-	{
-		bool bShaderHasOutdatedParameters = FBaseDS::Serialize(Ar);
-		return bShaderHasOutdatedParameters;
-	}
 };
 
 /**
@@ -394,6 +383,7 @@ public:
 template<typename LightMapPolicyType>
 class TBasePassPixelShaderPolicyParamType : public FMeshMaterialShader, public LightMapPolicyType::PixelParametersType
 {
+	DECLARE_INLINE_TYPE_LAYOUT_EXPLICIT_BASES(TBasePassPixelShaderPolicyParamType, NonVirtual, FMeshMaterialShader, typename LightMapPolicyType::PixelParametersType);
 public:
 
 	// static bool ShouldCompilePermutation(const FMeshMaterialShaderPermutationParameters& Parameters)
@@ -412,7 +402,7 @@ public:
 		FForwardLightingParameters::ModifyCompilationEnvironment(Parameters.Platform, OutEnvironment);
 	}
 
-	static bool ValidateCompiledResult(EShaderPlatform Platform, const TArray<FMaterial*>& Materials, const FVertexFactoryType* VertexFactoryType, const FShaderParameterMap& ParameterMap, TArray<FString>& OutError)
+	static bool ValidateCompiledResult(EShaderPlatform Platform, const FShaderParameterMap& ParameterMap, TArray<FString>& OutError)
 	{
 		if (ParameterMap.ContainsParameterAllocation(FSceneTexturesUniformParameters::StaticStructMetadata.GetShaderVariableName()))
 		{
@@ -448,16 +438,8 @@ public:
 		const TBasePassShaderElementData<LightMapPolicyType>& ShaderElementData,
 		FMeshDrawSingleShaderBindings& ShaderBindings) const;
 
-	virtual bool Serialize(FArchive& Ar) override
-	{
-		bool bShaderHasOutdatedParameters = FMeshMaterialShader::Serialize(Ar);
-		LightMapPolicyType::PixelParametersType::Serialize(Ar);
-		Ar << ReflectionCaptureBuffer;
-		return bShaderHasOutdatedParameters;
-	}
-
 private:
-	FShaderUniformBufferParameter ReflectionCaptureBuffer;
+	LAYOUT_FIELD(FShaderUniformBufferParameter, ReflectionCaptureBuffer);
 };
 
 /**
@@ -468,9 +450,8 @@ template<typename LightMapPolicyType>
 class TBasePassPixelShaderBaseType : public TBasePassPixelShaderPolicyParamType<LightMapPolicyType>
 {
 	typedef TBasePassPixelShaderPolicyParamType<LightMapPolicyType> Super;
-
+	DECLARE_INLINE_TYPE_LAYOUT(TBasePassPixelShaderBaseType, NonVirtual);
 public:
-
 	static bool ShouldCompilePermutation(const FMeshMaterialShaderPermutationParameters& Parameters)
 	{
 		return LightMapPolicyType::ShouldCompilePermutation(Parameters);
@@ -501,9 +482,9 @@ public:
 		static const auto SupportStationarySkylight = IConsoleManager::Get().FindTConsoleVariableDataInt(TEXT("r.SupportStationarySkylight"));
 		static const auto SupportAllShaderPermutations = IConsoleManager::Get().FindTConsoleVariableDataInt(TEXT("r.SupportAllShaderPermutations"));
 
-		const bool IsSingleLayerWater = Parameters.Material->GetShadingModels().HasShadingModel(MSM_SingleLayerWater);
+		const bool IsSingleLayerWater = Parameters.MaterialParameters.ShadingModels.HasShadingModel(MSM_SingleLayerWater);
 
-		const bool bTranslucent = IsTranslucentBlendMode(Parameters.Material->GetBlendMode());
+		const bool bTranslucent = IsTranslucentBlendMode(Parameters.MaterialParameters.BlendMode);
 		const bool bForceAllPermutations = SupportAllShaderPermutations && SupportAllShaderPermutations->GetValueOnAnyThread() != 0;
 		const bool bProjectSupportsStationarySkylight = !SupportStationarySkylight || SupportStationarySkylight->GetValueOnAnyThread() != 0 || bForceAllPermutations;
 
@@ -513,7 +494,7 @@ public:
 			// Some lightmap policies (eg Simple Forward) always require skylight support
 			|| IsSingleLayerWater
 			|| LightMapPolicyType::RequiresSkylight()
-			|| ((bProjectSupportsStationarySkylight || IsForwardShadingEnabled(Parameters.Platform)) && Parameters.Material->GetShadingModels().IsLit());
+			|| ((bProjectSupportsStationarySkylight || IsForwardShadingEnabled(Parameters.Platform)) && Parameters.MaterialParameters.ShadingModels.IsLit());
 		return bCacheShaders
 			&& (IsFeatureLevelSupported(Parameters.Platform, ERHIFeatureLevel::SM5))
 			&& TBasePassPixelShaderBaseType<LightMapPolicyType>::ShouldCompilePermutation(Parameters);
@@ -522,7 +503,7 @@ public:
 	static void ModifyCompilationEnvironment(const FMaterialShaderPermutationParameters& Parameters, FShaderCompilerEnvironment& OutEnvironment)
 	{
 		// For deferred decals, the shader class used is FDeferredDecalPS. the TBasePassPS is only used in the material editor and will read wrong values.
-		OutEnvironment.SetDefine(TEXT("SCENE_TEXTURES_DISABLED"), Parameters.Material->GetMaterialDomain() != MD_Surface); 
+		OutEnvironment.SetDefine(TEXT("SCENE_TEXTURES_DISABLED"), Parameters.MaterialParameters.MaterialDomain != MD_Surface); 
 
 		OutEnvironment.SetDefine(TEXT("ENABLE_SKY_LIGHT"), bEnableSkyLight);
 		TBasePassPixelShaderBaseType<LightMapPolicyType>::ModifyCompilationEnvironment(Parameters, OutEnvironment);
@@ -549,10 +530,10 @@ void GetBasePassShaders(
 	ERHIFeatureLevel::Type FeatureLevel,
 	bool bEnableAtmosphericFog,
 	bool bEnableSkyLight,
-	FBaseHS*& HullShader,
-	FBaseDS*& DomainShader,
-	TBasePassVertexShaderPolicyParamType<LightMapPolicyType>*& VertexShader,
-	TBasePassPixelShaderPolicyParamType<LightMapPolicyType>*& PixelShader
+	TShaderRef<FBaseHS>& HullShader,
+	TShaderRef<FBaseDS>& DomainShader,
+	TShaderRef<TBasePassVertexShaderPolicyParamType<LightMapPolicyType>>& VertexShader,
+	TShaderRef<TBasePassPixelShaderPolicyParamType<LightMapPolicyType>>& PixelShader
 	)
 {
 	const EMaterialTessellationMode MaterialTessellationMode = Material.GetTessellationMode();
@@ -566,7 +547,7 @@ void GetBasePassShaders(
 		DomainShader = Material.GetShader<TBasePassDS<LightMapPolicyType > >(VertexFactoryType);
 		
 		// Metal requires matching permutations, but no other platform should worry about this complication.
-		if (bEnableAtmosphericFog && DomainShader && IsMetalPlatform(EShaderPlatform(DomainShader->GetTarget().Platform)))
+		if (bEnableAtmosphericFog && DomainShader.IsValid() && IsMetalPlatform(EShaderPlatform(DomainShader->GetTarget().Platform)))
 		{
 			HullShader = Material.GetShader<TBasePassHS<LightMapPolicyType, true > >(VertexFactoryType);
 		}
@@ -602,10 +583,10 @@ void GetBasePassShaders<FUniformLightMapPolicy>(
 	ERHIFeatureLevel::Type FeatureLevel,
 	bool bEnableAtmosphericFog,
 	bool bEnableSkyLight,
-	FBaseHS*& HullShader,
-	FBaseDS*& DomainShader,
-	TBasePassVertexShaderPolicyParamType<FUniformLightMapPolicy>*& VertexShader,
-	TBasePassPixelShaderPolicyParamType<FUniformLightMapPolicy>*& PixelShader
+	TShaderRef<FBaseHS>& HullShader,
+	TShaderRef<FBaseDS>& DomainShader,
+	TShaderRef<TBasePassVertexShaderPolicyParamType<FUniformLightMapPolicy>>& VertexShader,
+	TShaderRef<TBasePassPixelShaderPolicyParamType<FUniformLightMapPolicy>>& PixelShader
 	);
 
 void SetupBasePassView(FRHICommandList& RHICmdList, const FViewInfo& View, const FSceneRenderer* SceneRenderer, const bool bIsEditorPrimitivePass = false);
diff --git a/Engine/Source/Runtime/Renderer/Private/BasePassRendering.inl b/Engine/Source/Runtime/Renderer/Private/BasePassRendering.inl
--- a/Engine/Source/Runtime/Renderer/Private/BasePassRendering.inl
+++ b/Engine/Source/Runtime/Renderer/Private/BasePassRendering.inl
@@ -51,6 +51,7 @@ void TBasePassVertexShaderPolicyParamType<LightMapPolicyType>::GetShaderBindings
 
 template<typename LightMapPolicyType>
 void TBasePassVertexShaderPolicyParamType<LightMapPolicyType>::GetElementShaderBindings(
+	const FShaderMapPointerTable& PointerTable,
 	const FScene* Scene, 
 	const FSceneView* ViewIfDynamicMeshCommand, 
 	const FVertexFactory* VertexFactory,
@@ -63,7 +64,7 @@ void TBasePassVertexShaderPolicyParamType<LightMapPolicyType>::GetElementShaderB
 	FMeshDrawSingleShaderBindings& ShaderBindings,
 	FVertexInputStreamArray& VertexStreams) const
 {
-	FMeshMaterialShader::GetElementShaderBindings(Scene, ViewIfDynamicMeshCommand, VertexFactory, InputStreamType, FeatureLevel, PrimitiveSceneProxy, MeshBatch, BatchElement, ShaderElementData, ShaderBindings, VertexStreams);
+	FMeshMaterialShader::GetElementShaderBindings(PointerTable, Scene, ViewIfDynamicMeshCommand, VertexFactory, InputStreamType, FeatureLevel, PrimitiveSceneProxy, MeshBatch, BatchElement, ShaderElementData, ShaderBindings, VertexStreams);
 }
 
 template<typename LightMapPolicyType>
diff --git a/Engine/Source/Runtime/Renderer/Private/CapsuleShadowRendering.cpp b/Engine/Source/Runtime/Renderer/Private/CapsuleShadowRendering.cpp
--- a/Engine/Source/Runtime/Renderer/Private/CapsuleShadowRendering.cpp
+++ b/Engine/Source/Runtime/Renderer/Private/CapsuleShadowRendering.cpp
@@ -152,7 +152,7 @@ public:
 		FRHIShaderResourceView* LightDirectionDataSRV,
 		const FRWBuffer& ComputedLightDirectionDataBuffer)
 	{
-		FRHIComputeShader* ShaderRHI = GetComputeShader();
+		FRHIComputeShader* ShaderRHI = RHICmdList.GetBoundComputeShader();
 
 		FGlobalShader::SetParameters<FViewUniformShaderParameters>(RHICmdList, ShaderRHI, View.ViewUniformBuffer);
 
@@ -174,37 +174,22 @@ public:
 
 	void UnsetParameters(FRHICommandList& RHICmdList, const FRWBuffer& ComputedLightDirectionDataBuffer)
 	{
-		ComputedLightDirectionData.UnsetUAV(RHICmdList, GetComputeShader());
+		ComputedLightDirectionData.UnsetUAV(RHICmdList, RHICmdList.GetBoundComputeShader());
 
 		FRHIUnorderedAccessView* OutUAVs[1];
 		OutUAVs[0] = ComputedLightDirectionDataBuffer.UAV;
 		RHICmdList.TransitionResources(EResourceTransitionAccess::EReadable, EResourceTransitionPipeline::EComputeToCompute, OutUAVs, UE_ARRAY_COUNT(OutUAVs) );
 	}
 
-	// FShader interface.
-	virtual bool Serialize(FArchive& Ar) override
-	{
-		bool bShaderHasOutdatedParameters = FGlobalShader::Serialize(Ar);
-		Ar << NumLightDirectionData;
-		Ar << LightDirectionData;
-		Ar << UseStationarySkylightShadowing;
-		Ar << CapsuleIndirectConeAngle;
-		Ar << CapsuleSkyAngleScale;
-		Ar << CapsuleMinSkyAngle;
-		Ar << ComputedLightDirectionData;
-		return bShaderHasOutdatedParameters;
-	}
-
 private:
 
-	FShaderParameter NumLightDirectionData;
-	FShaderResourceParameter LightDirectionData;
-	FShaderParameter UseStationarySkylightShadowing;
-	FShaderParameter CapsuleIndirectConeAngle;
-	FShaderParameter CapsuleSkyAngleScale;
-	FShaderParameter CapsuleMinSkyAngle;
-	FRWShaderParameter ComputedLightDirectionData;
-
+	LAYOUT_FIELD(FShaderParameter, NumLightDirectionData);
+	LAYOUT_FIELD(FShaderResourceParameter, LightDirectionData);
+	LAYOUT_FIELD(FShaderParameter, UseStationarySkylightShadowing);
+	LAYOUT_FIELD(FShaderParameter, CapsuleIndirectConeAngle);
+	LAYOUT_FIELD(FShaderParameter, CapsuleSkyAngleScale);
+	LAYOUT_FIELD(FShaderParameter, CapsuleMinSkyAngle);
+	LAYOUT_FIELD(FRWShaderParameter, ComputedLightDirectionData);
 };
 
 IMPLEMENT_SHADER_TYPE(,FComputeLightDirectionFromVolumetricLightmapCS,TEXT("/Engine/Private/CapsuleShadowShaders.usf"),TEXT("ComputeLightDirectionFromVolumetricLightmapCS"),SF_Compute);
@@ -240,7 +225,7 @@ enum EIndirectShadowingPrimitiveTypes
 template<ECapsuleShadowingType ShadowingType>
 class TCapsuleShadowingBaseCS : public FGlobalShader
 {
-	DECLARE_SHADER_TYPE(TCapsuleShadowingBaseCS, Global);
+	DECLARE_INLINE_TYPE_LAYOUT(TCapsuleShadowingBaseCS, NonVirtual);
 public:
 
 	static bool ShouldCompilePermutation(const FGlobalShaderPermutationParameters& Parameters)
@@ -329,7 +314,7 @@ public:
 		FRHIShaderResourceView* LightDirectionDataSRV,
 		FRHITexture* ReceiverBentNormalTextureValue)
 	{
-		FRHIComputeShader* ShaderRHI = GetComputeShader();
+		FRHIComputeShader* ShaderRHI = RHICmdList.GetBoundComputeShader();
 
 		FGlobalShader::SetParameters<FViewUniformShaderParameters>(RHICmdList, ShaderRHI, View.ViewUniformBuffer);
 
@@ -454,9 +439,9 @@ public:
 
 	void UnsetParameters(FRHICommandList& RHICmdList, FSceneRenderTargetItem& OutputTexture, const FRWBuffer* TileIntersectionCountsBuffer)
 	{
-		ShadowFactors.UnsetUAV(RHICmdList, GetComputeShader());
-		BentNormalTexture.UnsetUAV(RHICmdList, GetComputeShader());
-		TileIntersectionCounts.UnsetUAV(RHICmdList, GetComputeShader());
+		ShadowFactors.UnsetUAV(RHICmdList, RHICmdList.GetBoundComputeShader());
+		BentNormalTexture.UnsetUAV(RHICmdList, RHICmdList.GetBoundComputeShader());
+		TileIntersectionCounts.UnsetUAV(RHICmdList, RHICmdList.GetBoundComputeShader());
 
 		FRHIUnorderedAccessView* OutUAVs[2];
 		OutUAVs[0] = OutputTexture.UAV;
@@ -469,61 +454,31 @@ public:
 		RHICmdList.TransitionResources(EResourceTransitionAccess::EReadable, EResourceTransitionPipeline::EComputeToCompute, OutUAVs, TileIntersectionCountsBuffer ? UE_ARRAY_COUNT(OutUAVs) : 1);
 	}
 
-	// FShader interface.
-	virtual bool Serialize(FArchive& Ar) override
-	{
-		bool bShaderHasOutdatedParameters = FGlobalShader::Serialize(Ar);
-		Ar << ShadowFactors;
-		Ar << TileIntersectionCounts;
-		Ar << TileDimensions;
-		Ar << BentNormalTexture;
-		Ar << ReceiverBentNormalTexture;
-		Ar << NumGroups;
-		Ar << LightDirection;
-		Ar << LightPositionAndInvRadius;
-		Ar << LightSourceRadius;
-		Ar << RayStartOffsetDepthScale;
-		Ar << LightAngleAndNormalThreshold;
-		Ar << ScissorRectMinAndSize;
-		Ar << SceneTextureParameters;
-		Ar << DownsampleFactor;
-		Ar << NumShadowCapsules;
-		Ar << ShadowCapsuleShapes;
-		Ar << NumMeshDistanceFieldCasters;
-		Ar << MeshDistanceFieldCasterIndices;
-		Ar << MaxOcclusionDistance;
-		Ar << CosFadeStartAngle;
-		Ar << LightDirectionData;
-		Ar << IndirectCapsuleSelfShadowingIntensity;
-		Ar << DistanceFieldObjectParameters;
-		return bShaderHasOutdatedParameters;
-	}
-
 private:
 
-	FRWShaderParameter ShadowFactors;
-	FRWShaderParameter TileIntersectionCounts;
-	FShaderParameter TileDimensions;
-	FRWShaderParameter BentNormalTexture;
-	FShaderResourceParameter ReceiverBentNormalTexture;
-	FShaderParameter NumGroups;
-	FShaderParameter LightDirection;
-	FShaderParameter LightPositionAndInvRadius;
-	FShaderParameter LightSourceRadius;
-	FShaderParameter RayStartOffsetDepthScale;
-	FShaderParameter LightAngleAndNormalThreshold;
-	FShaderParameter ScissorRectMinAndSize;
-	FSceneTextureShaderParameters SceneTextureParameters;
-	FShaderParameter DownsampleFactor;
-	FShaderParameter NumShadowCapsules;
-	FShaderResourceParameter ShadowCapsuleShapes;
-	FShaderParameter NumMeshDistanceFieldCasters;
-	FShaderResourceParameter MeshDistanceFieldCasterIndices;
-	FShaderParameter MaxOcclusionDistance;
-	FShaderParameter CosFadeStartAngle;
-	FShaderResourceParameter LightDirectionData;
-	FShaderParameter IndirectCapsuleSelfShadowingIntensity;
-	TDistanceFieldObjectBufferParameters<DFPT_SignedDistanceField> DistanceFieldObjectParameters;
+	LAYOUT_FIELD(FRWShaderParameter, ShadowFactors);
+	LAYOUT_FIELD(FRWShaderParameter, TileIntersectionCounts);
+	LAYOUT_FIELD(FShaderParameter, TileDimensions);
+	LAYOUT_FIELD(FRWShaderParameter, BentNormalTexture);
+	LAYOUT_FIELD(FShaderResourceParameter, ReceiverBentNormalTexture);
+	LAYOUT_FIELD(FShaderParameter, NumGroups);
+	LAYOUT_FIELD(FShaderParameter, LightDirection);
+	LAYOUT_FIELD(FShaderParameter, LightPositionAndInvRadius);
+	LAYOUT_FIELD(FShaderParameter, LightSourceRadius);
+	LAYOUT_FIELD(FShaderParameter, RayStartOffsetDepthScale);
+	LAYOUT_FIELD(FShaderParameter, LightAngleAndNormalThreshold);
+	LAYOUT_FIELD(FShaderParameter, ScissorRectMinAndSize);
+	LAYOUT_FIELD(FSceneTextureShaderParameters, SceneTextureParameters);
+	LAYOUT_FIELD(FShaderParameter, DownsampleFactor);
+	LAYOUT_FIELD(FShaderParameter, NumShadowCapsules);
+	LAYOUT_FIELD(FShaderResourceParameter, ShadowCapsuleShapes);
+	LAYOUT_FIELD(FShaderParameter, NumMeshDistanceFieldCasters);
+	LAYOUT_FIELD(FShaderResourceParameter, MeshDistanceFieldCasterIndices);
+	LAYOUT_FIELD(FShaderParameter, MaxOcclusionDistance);
+	LAYOUT_FIELD(FShaderParameter, CosFadeStartAngle);
+	LAYOUT_FIELD(FShaderResourceParameter, LightDirectionData);
+	LAYOUT_FIELD(FShaderParameter, IndirectCapsuleSelfShadowingIntensity);
+	LAYOUT_FIELD((TDistanceFieldObjectBufferParameters<DFPT_SignedDistanceField>), DistanceFieldObjectParameters);
 };
 
 template<ECapsuleShadowingType ShadowingType, EIndirectShadowingPrimitiveTypes PrimitiveTypes>
@@ -610,7 +565,7 @@ public:
 
 	void SetParameters(FRHICommandList& RHICmdList, const FSceneView& View, FIntPoint TileDimensionsValue, const FIntRect& ScissorRect, const FRWBuffer& TileIntersectionCountsBuffer)
 	{
-		FRHIVertexShader* ShaderRHI = GetVertexShader();
+		FRHIVertexShader* ShaderRHI = RHICmdList.GetBoundVertexShader();
 		FGlobalShader::SetParameters<FViewUniformShaderParameters>(RHICmdList, ShaderRHI, View.ViewUniformBuffer);
 
 		SetShaderValue(RHICmdList, ShaderRHI, TileDimensions, TileDimensionsValue);
@@ -621,23 +576,12 @@ public:
 		SetSRVParameter(RHICmdList, ShaderRHI, TileIntersectionCounts, TileIntersectionCountsBuffer.SRV);
 	}
 
-	// FShader interface.
-	virtual bool Serialize(FArchive& Ar) override
-	{
-		bool bShaderHasOutdatedParameters = FGlobalShader::Serialize(Ar);
-		Ar << TileDimensions;
-		Ar << TileSize;
-		Ar << ScissorRectMinAndSize;
-		Ar << TileIntersectionCounts;
-		return bShaderHasOutdatedParameters;
-	}
-
 private:
 
-	FShaderParameter TileDimensions;
-	FShaderParameter TileSize;
-	FShaderParameter ScissorRectMinAndSize;
-	FShaderResourceParameter TileIntersectionCounts;
+	LAYOUT_FIELD(FShaderParameter, TileDimensions);
+	LAYOUT_FIELD(FShaderParameter, TileSize);
+	LAYOUT_FIELD(FShaderParameter, ScissorRectMinAndSize);
+	LAYOUT_FIELD(FShaderResourceParameter, TileIntersectionCounts);
 };
 
 IMPLEMENT_SHADER_TYPE(,FCapsuleShadowingUpsampleVS,TEXT("/Engine/Private/CapsuleShadowShaders.usf"),TEXT("CapsuleShadowingUpsampleVS"),SF_Vertex);
@@ -677,7 +621,7 @@ public:
 
 	void SetParameters(FRHICommandList& RHICmdList, const FSceneView& View, const FIntRect& ScissorRect, TRefCountPtr<IPooledRenderTarget>& ShadowFactorsTextureValue, bool bOutputtingToLightAttenuation)
 	{
-		FRHIPixelShader* ShaderRHI = GetPixelShader();
+		FRHIPixelShader* ShaderRHI = RHICmdList.GetBoundPixelShader();
 
 		FGlobalShader::SetParameters<FViewUniformShaderParameters>(RHICmdList, ShaderRHI, View.ViewUniformBuffer);
 		SceneTextureParameters.Set(RHICmdList, ShaderRHI, View.FeatureLevel, ESceneTextureSetupMode::All);
@@ -688,25 +632,13 @@ public:
 		SetShaderValue(RHICmdList, ShaderRHI, OutputtingToLightAttenuation, bOutputtingToLightAttenuation ? 1.0f : 0.0f);
 	}
 
-	// FShader interface.
-	virtual bool Serialize(FArchive& Ar) override
-	{
-		bool bShaderHasOutdatedParameters = FGlobalShader::Serialize(Ar);
-		Ar << SceneTextureParameters;
-		Ar << ShadowFactorsTexture;
-		Ar << ShadowFactorsSampler;
-		Ar << ScissorRectMinAndSize;
-		Ar << OutputtingToLightAttenuation;
-		return bShaderHasOutdatedParameters;
-	}
-
 private:
 
-	FSceneTextureShaderParameters SceneTextureParameters;
-	FShaderResourceParameter ShadowFactorsTexture;
-	FShaderResourceParameter ShadowFactorsSampler;
-	FShaderParameter ScissorRectMinAndSize;
-	FShaderParameter OutputtingToLightAttenuation;
+	LAYOUT_FIELD(FSceneTextureShaderParameters, SceneTextureParameters);
+	LAYOUT_FIELD(FShaderResourceParameter, ShadowFactorsTexture);
+	LAYOUT_FIELD(FShaderResourceParameter, ShadowFactorsSampler);
+	LAYOUT_FIELD(FShaderParameter, ScissorRectMinAndSize);
+	LAYOUT_FIELD(FShaderParameter, OutputtingToLightAttenuation);
 };
 
 #define IMPLEMENT_CAPSULE_APPLY_SHADER_TYPE(bUpsampleRequired,bApplyToSSAO) \
@@ -897,7 +829,7 @@ bool FDeferredShadingSceneRenderer::RenderCapsuleDirectShadows(
 					if (bDirectionalLight)
 					{
 						TShaderMapRef<TCapsuleShadowingCS<ShapeShadow_DirectionalLightTiledCulling, IPT_CapsuleShapes> > ComputeShader(View.ShaderMap);
-						RHICmdList.SetComputeShader(ComputeShader->GetComputeShader());
+						RHICmdList.SetComputeShader(ComputeShader.GetComputeShader());
 
 						ComputeShader->SetParameters(
 							RHICmdList, 
@@ -918,13 +850,13 @@ bool FDeferredShadingSceneRenderer::RenderCapsuleDirectShadows(
 							NULL,
 							NULL);
 
-						DispatchComputeShader(RHICmdList, *ComputeShader, GroupSize.X, GroupSize.Y, 1);
+						DispatchComputeShader(RHICmdList, ComputeShader.GetShader(), GroupSize.X, GroupSize.Y, 1);
 						ComputeShader->UnsetParameters(RHICmdList, RayTracedShadowsRTI, &View.ViewState->CapsuleTileIntersectionCountsBuffer);
 					}
 					else
 					{
 						TShaderMapRef<TCapsuleShadowingCS<ShapeShadow_PointLightTiledCulling, IPT_CapsuleShapes> > ComputeShader(View.ShaderMap);
-						RHICmdList.SetComputeShader(ComputeShader->GetComputeShader());
+						RHICmdList.SetComputeShader(ComputeShader.GetComputeShader());
 
 						ComputeShader->SetParameters(
 							RHICmdList, 
@@ -945,7 +877,7 @@ bool FDeferredShadingSceneRenderer::RenderCapsuleDirectShadows(
 							NULL,
 							NULL);
 
-						DispatchComputeShader(RHICmdList, *ComputeShader, GroupSize.X, GroupSize.Y, 1);
+						DispatchComputeShader(RHICmdList, ComputeShader.GetShader(), GroupSize.X, GroupSize.Y, 1);
 						ComputeShader->UnsetParameters(RHICmdList, RayTracedShadowsRTI, &View.ViewState->CapsuleTileIntersectionCountsBuffer);
 					}
 				}
@@ -984,8 +916,8 @@ bool FDeferredShadingSceneRenderer::RenderCapsuleDirectShadows(
 						TShaderMapRef<TCapsuleShadowingUpsamplePS<false, false> > PixelShader(View.ShaderMap);
 
 						GraphicsPSOInit.BoundShaderState.VertexDeclarationRHI = GTileVertexDeclaration.VertexDeclarationRHI;
-						GraphicsPSOInit.BoundShaderState.VertexShaderRHI = GETSAFERHISHADER_VERTEX(*VertexShader);
-						GraphicsPSOInit.BoundShaderState.PixelShaderRHI = GETSAFERHISHADER_PIXEL(*PixelShader);
+						GraphicsPSOInit.BoundShaderState.VertexShaderRHI = VertexShader.GetVertexShader();
+						GraphicsPSOInit.BoundShaderState.PixelShaderRHI = PixelShader.GetPixelShader();
 
 						SetGraphicsPipelineState(RHICmdList, GraphicsPSOInit);
 
@@ -997,8 +929,8 @@ bool FDeferredShadingSceneRenderer::RenderCapsuleDirectShadows(
 						TShaderMapRef<TCapsuleShadowingUpsamplePS<true, false> > PixelShader(View.ShaderMap);
 
 						GraphicsPSOInit.BoundShaderState.VertexDeclarationRHI = GTileVertexDeclaration.VertexDeclarationRHI;
-						GraphicsPSOInit.BoundShaderState.VertexShaderRHI = GETSAFERHISHADER_VERTEX(*VertexShader);
-						GraphicsPSOInit.BoundShaderState.PixelShaderRHI = GETSAFERHISHADER_PIXEL(*PixelShader);
+						GraphicsPSOInit.BoundShaderState.VertexShaderRHI = VertexShader.GetVertexShader();
+						GraphicsPSOInit.BoundShaderState.PixelShaderRHI = PixelShader.GetPixelShader();
 
 						SetGraphicsPipelineState(RHICmdList, GraphicsPSOInit);
 
@@ -1272,11 +1204,11 @@ void FDeferredShadingSceneRenderer::SetupIndirectCapsuleShadows(
 
 			TShaderMapRef<FComputeLightDirectionFromVolumetricLightmapCS> ComputeShader(View.ShaderMap);
 	
-			RHICmdList.SetComputeShader(ComputeShader->GetComputeShader());
+			RHICmdList.SetComputeShader(ComputeShader.GetComputeShader());
 			ComputeShader->SetParameters(RHICmdList, Scene, View, NumLightDataElements, View.ViewState->IndirectShadowLightDirectionSRV, View.ViewState->IndirectShadowVolumetricLightmapDerivedLightDirection);
 
 			const int32 GroupSize = FMath::DivideAndRoundUp(NumLightDataElements, GComputeLightDirectionFromVolumetricLightmapGroupSize);
-			DispatchComputeShader(RHICmdList, *ComputeShader, GroupSize, 1, 1);
+			DispatchComputeShader(RHICmdList, ComputeShader.GetShader(), GroupSize, 1, 1);
 			ComputeShader->UnsetParameters(RHICmdList, View.ViewState->IndirectShadowVolumetricLightmapDerivedLightDirection);
 		}
 	}
@@ -1406,26 +1338,26 @@ void FDeferredShadingSceneRenderer::RenderIndirectCapsuleShadows(
 
 							FSceneRenderTargetItem& RayTracedShadowsRTI = RayTracedShadowsRT->GetRenderTargetItem();
 							{
-								TCapsuleShadowingBaseCS<ShapeShadow_IndirectTiledCulling>* ComputeShaderBase = NULL;
+								TShaderRef<TCapsuleShadowingBaseCS<ShapeShadow_IndirectTiledCulling>> ComputeShaderBase;
 
 								if (NumCapsuleShapes > 0 && NumMeshDistanceFieldCasters > 0)
 								{
 									TShaderMapRef<TCapsuleShadowingCS<ShapeShadow_IndirectTiledCulling, IPT_CapsuleShapesAndMeshDistanceFields> > ComputeShader(View.ShaderMap);
-									ComputeShaderBase = *ComputeShader;
+									ComputeShaderBase = ComputeShader;
 								}
 								else if (NumCapsuleShapes > 0)
 								{
 									TShaderMapRef<TCapsuleShadowingCS<ShapeShadow_IndirectTiledCulling, IPT_CapsuleShapes> > ComputeShader(View.ShaderMap);
-									ComputeShaderBase = *ComputeShader;
+									ComputeShaderBase = ComputeShader;
 								}
 								else
 								{
 									check(NumMeshDistanceFieldCasters > 0);
 									TShaderMapRef<TCapsuleShadowingCS<ShapeShadow_IndirectTiledCulling, IPT_MeshDistanceFields> > ComputeShader(View.ShaderMap);
-									ComputeShaderBase = *ComputeShader;
+									ComputeShaderBase = ComputeShader;
 								}
 
-								RHICmdList.SetComputeShader(ComputeShaderBase->GetComputeShader());
+								RHICmdList.SetComputeShader(ComputeShaderBase.GetComputeShader());
 
 								ComputeShaderBase->SetParameters(
 									RHICmdList, 
@@ -1446,7 +1378,7 @@ void FDeferredShadingSceneRenderer::RenderIndirectCapsuleShadows(
 									IndirectShadowLightDirectionSRV,
 									NULL);
 
-								DispatchComputeShader(RHICmdList, ComputeShaderBase, GroupSize.X, GroupSize.Y, 1);
+								DispatchComputeShader(RHICmdList, ComputeShaderBase.GetShader(), GroupSize.X, GroupSize.Y, 1);
 								ComputeShaderBase->UnsetParameters(RHICmdList, RayTracedShadowsRTI, &View.ViewState->CapsuleTileIntersectionCountsBuffer);
 							}
 						}
@@ -1481,7 +1413,7 @@ void FDeferredShadingSceneRenderer::RenderIndirectCapsuleShadows(
 
 							TShaderMapRef<FCapsuleShadowingUpsampleVS> VertexShader(View.ShaderMap);
 							GraphicsPSOInit.BoundShaderState.VertexDeclarationRHI = GTileVertexDeclaration.VertexDeclarationRHI;
-							GraphicsPSOInit.BoundShaderState.VertexShaderRHI = GETSAFERHISHADER_VERTEX(*VertexShader);
+							GraphicsPSOInit.BoundShaderState.VertexShaderRHI = VertexShader.GetVertexShader();
 							GraphicsPSOInit.PrimitiveType = PT_TriangleList;
 
 							if (RenderTargets.Num() > 1)
@@ -1489,7 +1421,7 @@ void FDeferredShadingSceneRenderer::RenderIndirectCapsuleShadows(
 								if (GCapsuleShadowsFullResolution)
 								{
 									TShaderMapRef<TCapsuleShadowingUpsamplePS<false, true> > PixelShader(View.ShaderMap);
-									GraphicsPSOInit.BoundShaderState.PixelShaderRHI = GETSAFERHISHADER_PIXEL(*PixelShader); 
+									GraphicsPSOInit.BoundShaderState.PixelShaderRHI = PixelShader.GetPixelShader(); 
 									SetGraphicsPipelineState(RHICmdList, GraphicsPSOInit);
 									VertexShader->SetParameters(RHICmdList, View, GroupSize, ScissorRect, View.ViewState->CapsuleTileIntersectionCountsBuffer);
 									PixelShader->SetParameters(RHICmdList, View, ScissorRect, RayTracedShadowsRT, false);
@@ -1497,7 +1429,7 @@ void FDeferredShadingSceneRenderer::RenderIndirectCapsuleShadows(
 								else
 								{
 									TShaderMapRef<TCapsuleShadowingUpsamplePS<true, true> > PixelShader(View.ShaderMap);
-									GraphicsPSOInit.BoundShaderState.PixelShaderRHI = GETSAFERHISHADER_PIXEL(*PixelShader);
+									GraphicsPSOInit.BoundShaderState.PixelShaderRHI = PixelShader.GetPixelShader();
 									SetGraphicsPipelineState(RHICmdList, GraphicsPSOInit);
 									VertexShader->SetParameters(RHICmdList, View, GroupSize, ScissorRect, View.ViewState->CapsuleTileIntersectionCountsBuffer);
 									PixelShader->SetParameters(RHICmdList, View, ScissorRect, RayTracedShadowsRT, false);
@@ -1508,7 +1440,7 @@ void FDeferredShadingSceneRenderer::RenderIndirectCapsuleShadows(
 								if (GCapsuleShadowsFullResolution)
 								{
 									TShaderMapRef<TCapsuleShadowingUpsamplePS<false, false> > PixelShader(View.ShaderMap);
-									GraphicsPSOInit.BoundShaderState.PixelShaderRHI = GETSAFERHISHADER_PIXEL(*PixelShader);
+									GraphicsPSOInit.BoundShaderState.PixelShaderRHI = PixelShader.GetPixelShader();
 									SetGraphicsPipelineState(RHICmdList, GraphicsPSOInit);
 									VertexShader->SetParameters(RHICmdList, View, GroupSize, ScissorRect, View.ViewState->CapsuleTileIntersectionCountsBuffer);
 									PixelShader->SetParameters(RHICmdList, View, ScissorRect, RayTracedShadowsRT, false);
@@ -1516,7 +1448,7 @@ void FDeferredShadingSceneRenderer::RenderIndirectCapsuleShadows(
 								else
 								{
 									TShaderMapRef<TCapsuleShadowingUpsamplePS<true, false> > PixelShader(View.ShaderMap);
-									GraphicsPSOInit.BoundShaderState.PixelShaderRHI = GETSAFERHISHADER_PIXEL(*PixelShader);
+									GraphicsPSOInit.BoundShaderState.PixelShaderRHI = PixelShader.GetPixelShader();
 									SetGraphicsPipelineState(RHICmdList, GraphicsPSOInit);
 									VertexShader->SetParameters(RHICmdList, View, GroupSize, ScissorRect, View.ViewState->CapsuleTileIntersectionCountsBuffer);
 									PixelShader->SetParameters(RHICmdList, View, ScissorRect, RayTracedShadowsRT, false);
@@ -1629,25 +1561,24 @@ void FDeferredShadingSceneRenderer::RenderCapsuleShadowsForMovableSkylight(FRHIC
 
 								FSceneRenderTargetItem& RayTracedShadowsRTI = NewBentNormal->GetRenderTargetItem();
 								{
-									TCapsuleShadowingBaseCS<ShapeShadow_MovableSkylightTiledCulling>* ComputeShaderBase = NULL;
-
+									TShaderRef<TCapsuleShadowingBaseCS<ShapeShadow_MovableSkylightTiledCulling>> ComputeShaderBase;
 									if (NumCapsuleShapes > 0 && NumMeshDistanceFieldCasters > 0)
 									{
 										TShaderMapRef<TCapsuleShadowingCS<ShapeShadow_MovableSkylightTiledCulling, IPT_CapsuleShapesAndMeshDistanceFields> > ComputeShader(View.ShaderMap);
-										ComputeShaderBase = *ComputeShader;
+										ComputeShaderBase = ComputeShader;
 									}
 									else if (NumCapsuleShapes > 0)
 									{
 										TShaderMapRef<TCapsuleShadowingCS<ShapeShadow_MovableSkylightTiledCulling, IPT_CapsuleShapes> > ComputeShader(View.ShaderMap);
-										ComputeShaderBase = *ComputeShader;
+										ComputeShaderBase = ComputeShader;
 									}
 									else
 									{
 										TShaderMapRef<TCapsuleShadowingCS<ShapeShadow_MovableSkylightTiledCulling, IPT_MeshDistanceFields> > ComputeShader(View.ShaderMap);
-										ComputeShaderBase = *ComputeShader;
+										ComputeShaderBase = ComputeShader;
 									}
 									
-									RHICmdList.SetComputeShader(ComputeShaderBase->GetComputeShader());
+									RHICmdList.SetComputeShader(ComputeShaderBase.GetComputeShader());
 
 									ComputeShaderBase->SetParameters(
 										RHICmdList, 
@@ -1668,7 +1599,7 @@ void FDeferredShadingSceneRenderer::RenderCapsuleShadowsForMovableSkylight(FRHIC
 										IndirectShadowLightDirectionSRV,
 										BentNormalOutput->GetRenderTargetItem().ShaderResourceTexture);
 
-									DispatchComputeShader(RHICmdList, ComputeShaderBase, GroupSizeX, GroupSizeY, 1);
+									DispatchComputeShader(RHICmdList, ComputeShaderBase.GetShader(), GroupSizeX, GroupSizeY, 1);
 									ComputeShaderBase->UnsetParameters(RHICmdList, RayTracedShadowsRTI, NULL);
 								}
 							}
diff --git a/Engine/Source/Runtime/Renderer/Private/CapsuleShadowRendering.h b/Engine/Source/Runtime/Renderer/Private/CapsuleShadowRendering.h
--- a/Engine/Source/Runtime/Renderer/Private/CapsuleShadowRendering.h
+++ b/Engine/Source/Runtime/Renderer/Private/CapsuleShadowRendering.h
@@ -13,29 +13,29 @@ extern int32 GCapsuleShadows;
 extern int32 GCapsuleDirectShadows;
 extern int32 GCapsuleIndirectShadows;
 
-inline bool DoesPlatformSupportCapsuleShadows(EShaderPlatform Platform)
+inline bool DoesPlatformSupportCapsuleShadows(const FStaticShaderPlatform Platform)
 {
 	// Hasn't been tested elsewhere yet
 	return Platform == SP_PCD3D_SM5 || Platform == SP_PS4 || Platform == SP_XBOXONE_D3D12
 		|| IsMetalSM5Platform(Platform)
 		|| IsVulkanSM5Platform(Platform)
 		|| Platform == SP_SWITCH
-		|| FDataDrivenShaderPlatformInfo::GetInfo(Platform).bSupportsCapsuleShadows;
+		|| FDataDrivenShaderPlatformInfo::GetSupportsCapsuleShadows(Platform);
 }
 
-inline bool SupportsCapsuleShadows(ERHIFeatureLevel::Type FeatureLevel, EShaderPlatform ShaderPlatform)
+inline bool SupportsCapsuleShadows(ERHIFeatureLevel::Type FeatureLevel, FStaticShaderPlatform ShaderPlatform)
 {
 	return GCapsuleShadows
 		&& FeatureLevel >= ERHIFeatureLevel::SM5
 		&& DoesPlatformSupportCapsuleShadows(ShaderPlatform);
 }
 
-inline bool SupportsCapsuleDirectShadows(ERHIFeatureLevel::Type FeatureLevel, EShaderPlatform ShaderPlatform)
+inline bool SupportsCapsuleDirectShadows(ERHIFeatureLevel::Type FeatureLevel, FStaticShaderPlatform ShaderPlatform)
 {
 	return GCapsuleDirectShadows && SupportsCapsuleShadows(FeatureLevel, ShaderPlatform);
 }
 
-inline bool SupportsCapsuleIndirectShadows(ERHIFeatureLevel::Type FeatureLevel, EShaderPlatform ShaderPlatform)
+inline bool SupportsCapsuleIndirectShadows(ERHIFeatureLevel::Type FeatureLevel, FStaticShaderPlatform ShaderPlatform)
 {
 	return GCapsuleIndirectShadows && SupportsCapsuleShadows(FeatureLevel, ShaderPlatform);
 }
diff --git a/Engine/Source/Runtime/Renderer/Private/ClusteredDeferredShadingPass.cpp b/Engine/Source/Runtime/Renderer/Private/ClusteredDeferredShadingPass.cpp
--- a/Engine/Source/Runtime/Renderer/Private/ClusteredDeferredShadingPass.cpp
+++ b/Engine/Source/Runtime/Renderer/Private/ClusteredDeferredShadingPass.cpp
@@ -158,19 +158,19 @@ void FDeferredShadingSceneRenderer::AddClusteredDeferredShadingPass(FRHICommandL
 					GraphicsPSOInit.DepthStencilState = TStaticDepthStencilState<false, CF_Always>::GetRHI();
 
 					GraphicsPSOInit.BoundShaderState.VertexDeclarationRHI = GFilterVertexDeclaration.VertexDeclarationRHI;
-					GraphicsPSOInit.BoundShaderState.VertexShaderRHI = GETSAFERHISHADER_VERTEX(*VertexShader);
-					GraphicsPSOInit.BoundShaderState.PixelShaderRHI = GETSAFERHISHADER_PIXEL(*PixelShader);
+					GraphicsPSOInit.BoundShaderState.VertexShaderRHI = VertexShader.GetVertexShader();
+					GraphicsPSOInit.BoundShaderState.PixelShaderRHI = PixelShader.GetPixelShader();
 					GraphicsPSOInit.PrimitiveType = PT_TriangleList;
 
 					SetGraphicsPipelineState(InRHICmdList, GraphicsPSOInit);
 				}
 				InRHICmdList.SetViewport(View.ViewRect.Min.X, View.ViewRect.Min.Y, 0.0f, View.ViewRect.Max.X, View.ViewRect.Max.Y, 1.0f);
 
-				SetShaderParameters(InRHICmdList, *PixelShader, PixelShader->GetPixelShader(), *PassParameters);
+				SetShaderParameters(InRHICmdList, PixelShader, PixelShader.GetPixelShader(), *PassParameters);
 
 				DrawRectangle(InRHICmdList, 0, 0, View.ViewRect.Width(), View.ViewRect.Height(),
 					View.ViewRect.Min.X, View.ViewRect.Min.Y, View.ViewRect.Width(), View.ViewRect.Height(),
-					FIntPoint(View.ViewRect.Width(), View.ViewRect.Height()), SceneContext.GetBufferSizeXY(), *VertexShader);
+					FIntPoint(View.ViewRect.Width(), View.ViewRect.Height()), SceneContext.GetBufferSizeXY(), VertexShader);
 			});
 
 		}
diff --git a/Engine/Source/Runtime/Renderer/Private/CompositionLighting/PostProcessAmbient.cpp b/Engine/Source/Runtime/Renderer/Private/CompositionLighting/PostProcessAmbient.cpp
--- a/Engine/Source/Runtime/Renderer/Private/CompositionLighting/PostProcessAmbient.cpp
+++ b/Engine/Source/Runtime/Renderer/Private/CompositionLighting/PostProcessAmbient.cpp
@@ -17,6 +17,7 @@
 /*-----------------------------------------------------------------------------
 FCubemapShaderParameters
 -----------------------------------------------------------------------------*/
+IMPLEMENT_TYPE_LAYOUT(FCubemapShaderParameters);
 
 void FCubemapShaderParameters::Bind(const FShaderParameterMap& ParameterMap)
 {
