Copying Private-LoadTimes-4.24 stream to Main. Biggest changes are in Materials/Shader memory freezing.

diff --git a/Engine/Source/Runtime/Renderer/Private/CompositionLighting/PostProcessAmbientOcclusion.cpp b/Engine/Source/Runtime/Renderer/Private/CompositionLighting/PostProcessAmbientOcclusion.cpp
--- a/Engine/Source/Runtime/Renderer/Private/CompositionLighting/PostProcessAmbientOcclusion.cpp
+++ b/Engine/Source/Runtime/Renderer/Private/CompositionLighting/PostProcessAmbientOcclusion.cpp
@@ -294,6 +294,7 @@ EGTAOType FSSAOHelper::GetGTAOPassType(const FViewInfo& View)
 /** Shader parameters needed for screen space AmbientOcclusion passes. */
 class FScreenSpaceAOParameters
 {
+	DECLARE_INLINE_TYPE_LAYOUT(FScreenSpaceAOParameters, NonVirtual);
 public:
 
 	enum class ERandTexType
@@ -381,7 +382,9 @@ public:
 	friend FArchive& operator<<(FArchive& Ar, FScreenSpaceAOParameters& This);
 
 private:
-	FShaderParameter ScreenSpaceAOParams;
+	
+		LAYOUT_FIELD(FShaderParameter, ScreenSpaceAOParams)
+	
 };
 
 /** Encapsulates the post processing ambient occlusion pixel shader. */
@@ -405,10 +408,10 @@ class FPostProcessAmbientOcclusionSetupPS : public FGlobalShader
 	FPostProcessAmbientOcclusionSetupPS() {}
 
 public:
-	FPostProcessPassParameters PostprocessParameter;
-	FSceneTextureShaderParameters SceneTextureParameters;
-	FScreenSpaceAOParameters ScreenSpaceAOParams;
-	FShaderParameter AmbientOcclusionSetupParams;
+	LAYOUT_FIELD(FPostProcessPassParameters, PostprocessParameter)
+	LAYOUT_FIELD(FSceneTextureShaderParameters, SceneTextureParameters)
+	LAYOUT_FIELD(FScreenSpaceAOParameters, ScreenSpaceAOParams);
+	LAYOUT_FIELD(FShaderParameter, AmbientOcclusionSetupParams)
 
 	/** Initialization constructor. */
 	FPostProcessAmbientOcclusionSetupPS(const ShaderMetaType::CompiledShaderInitializerType& Initializer)
@@ -423,7 +426,7 @@ public:
 	void SetParameters(const FRenderingCompositePassContext& Context)
 	{
 		const FFinalPostProcessSettings& Settings = Context.View.FinalPostProcessSettings;
-		FRHIPixelShader* ShaderRHI = GetPixelShader();
+		FRHIPixelShader* ShaderRHI = Context.RHICmdList.GetBoundPixelShader();
 
 		FGlobalShader::SetParameters<FViewUniformShaderParameters>(Context.RHICmdList, ShaderRHI, Context.View.ViewUniformBuffer);
 
@@ -442,12 +445,6 @@ public:
 	}
 
 	// FShader interface.
-	virtual bool Serialize(FArchive& Ar) override
-	{
-		bool bShaderHasOutdatedParameters = FGlobalShader::Serialize(Ar);
-		Ar << PostprocessParameter << SceneTextureParameters << ScreenSpaceAOParams << AmbientOcclusionSetupParams;
-		return bShaderHasOutdatedParameters;
-	}
 
 	static const TCHAR* GetSourceFilename()
 	{
@@ -507,8 +504,8 @@ void FRCPassPostProcessAmbientOcclusionSetup::Process(FRenderingCompositePassCon
 			TShaderMapRef<FPostProcessAmbientOcclusionSetupPS<1> > PixelShader(Context.GetShaderMap());
 
 			GraphicsPSOInit.BoundShaderState.VertexDeclarationRHI = GFilterVertexDeclaration.VertexDeclarationRHI;
-			GraphicsPSOInit.BoundShaderState.VertexShaderRHI = GETSAFERHISHADER_VERTEX(*VertexShader);
-			GraphicsPSOInit.BoundShaderState.PixelShaderRHI = GETSAFERHISHADER_PIXEL(*PixelShader);
+			GraphicsPSOInit.BoundShaderState.VertexShaderRHI = VertexShader.GetVertexShader();
+			GraphicsPSOInit.BoundShaderState.PixelShaderRHI = PixelShader.GetPixelShader();
 			SetGraphicsPipelineState(Context.RHICmdList, GraphicsPSOInit);
 
 			PixelShader->SetParameters(Context);
@@ -518,8 +515,8 @@ void FRCPassPostProcessAmbientOcclusionSetup::Process(FRenderingCompositePassCon
 			TShaderMapRef<FPostProcessAmbientOcclusionSetupPS<0> > PixelShader(Context.GetShaderMap());
 
 			GraphicsPSOInit.BoundShaderState.VertexDeclarationRHI = GFilterVertexDeclaration.VertexDeclarationRHI;
-			GraphicsPSOInit.BoundShaderState.VertexShaderRHI = GETSAFERHISHADER_VERTEX(*VertexShader);
-			GraphicsPSOInit.BoundShaderState.PixelShaderRHI = GETSAFERHISHADER_PIXEL(*PixelShader);
+			GraphicsPSOInit.BoundShaderState.VertexShaderRHI = VertexShader.GetVertexShader();
+			GraphicsPSOInit.BoundShaderState.PixelShaderRHI = PixelShader.GetPixelShader();
 			SetGraphicsPipelineState(Context.RHICmdList, GraphicsPSOInit);
 
 			PixelShader->SetParameters(Context);
@@ -534,7 +531,7 @@ void FRCPassPostProcessAmbientOcclusionSetup::Process(FRenderingCompositePassCon
 			SrcRect.Width(), SrcRect.Height(),
 			DestRect.Size(),
 			FSceneRenderTargets::Get(Context.RHICmdList).GetBufferSizeXY(),
-			*VertexShader,
+			VertexShader,
 			View.StereoPass,
 			Context.HasHmdMesh(),
 			EDRF_UseTriangleOptimization);
@@ -611,9 +608,9 @@ class FPostProcessAmbientOcclusionSmoothCS : public FGlobalShader
 	FPostProcessAmbientOcclusionSmoothCS() {}
 
 public:
-	FPostProcessPassParameters PostprocessParameter;
-	FShaderParameter SSAOSmoothParams;
-	FShaderParameter SSAOSmoothResult;
+	LAYOUT_FIELD(FPostProcessPassParameters, PostprocessParameter)
+	LAYOUT_FIELD(FShaderParameter, SSAOSmoothParams)
+	LAYOUT_FIELD(FShaderParameter, SSAOSmoothResult)
 
 	/** Initialization constructor. */
 	FPostProcessAmbientOcclusionSmoothCS(const ShaderMetaType::CompiledShaderInitializerType& Initializer)
@@ -631,7 +628,7 @@ public:
 		const FIntRect& OutputRect,
 		FRHIUnorderedAccessView* OutUAV)
 	{
-		FRHIComputeShader* ShaderRHI = GetComputeShader();
+		FRHIComputeShader* ShaderRHI = RHICmdList.GetBoundComputeShader();
 
 		FGlobalShader::SetParameters<FViewUniformShaderParameters>(RHICmdList, ShaderRHI, Context.View.ViewUniformBuffer);
 
@@ -646,18 +643,11 @@ public:
 	template <typename TRHICmdList>
 	void UnsetParameters(TRHICmdList& RHICmdList)
 	{
-		FRHIComputeShader* ShaderRHI = GetComputeShader();
+		FRHIComputeShader* ShaderRHI = RHICmdList.GetBoundComputeShader();
 		RHICmdList.SetUAVParameter(ShaderRHI, SSAOSmoothResult.GetBaseIndex(), nullptr);
 	}
 
 	// FShader interface.
-	virtual bool Serialize(FArchive& Ar) override
-	{
-		bool bShaderHasOutdatedParameters = FGlobalShader::Serialize(Ar);
-		Ar << PostprocessParameter << SSAOSmoothParams << SSAOSmoothResult;
-		return bShaderHasOutdatedParameters;
-	}
-
 	static const TCHAR* GetSourceFilename()
 	{
 		return TEXT("/Engine/Private/PostProcessAmbientOcclusion.usf");
@@ -684,11 +674,11 @@ void FRCPassPostProcessAmbientOcclusionSmooth::DispatchCS(
 	FRHIUnorderedAccessView* OutUAV) const
 {
 	TShaderMapRef<FPostProcessAmbientOcclusionSmoothCS> ComputeShader(Context.GetShaderMap());
-	RHICmdList.SetComputeShader(ComputeShader->GetComputeShader());
+	RHICmdList.SetComputeShader(ComputeShader.GetComputeShader());
 	ComputeShader->SetParameters(RHICmdList, Context, OutputRect, OutUAV);
 	const uint32 NumGroupsX = FMath::DivideAndRoundUp(OutputRect.Width(), ThreadGroupSize1D);
 	const uint32 NumGroupsY = FMath::DivideAndRoundUp(OutputRect.Height(), ThreadGroupSize1D);
-	DispatchComputeShader(RHICmdList, *ComputeShader, NumGroupsX, NumGroupsY, 1);
+	DispatchComputeShader(RHICmdList, ComputeShader.GetShader(), NumGroupsX, NumGroupsY, 1);
 	ComputeShader->UnsetParameters(RHICmdList);
 }
 
@@ -794,13 +784,13 @@ class FPostProcessAmbientOcclusionPSandCS : public FGlobalShader
 	FPostProcessAmbientOcclusionPSandCS() {}
 
 public:
-	FShaderParameter HZBRemapping;
-	FPostProcessPassParameters PostprocessParameter;
-	FSceneTextureShaderParameters SceneTextureParameters;
-	FScreenSpaceAOParameters ScreenSpaceAOParams;
-	FShaderResourceParameter RandomNormalTexture;
-	FShaderResourceParameter RandomNormalTextureSampler;
-	FShaderParameter OutTexture;
+	LAYOUT_FIELD(FShaderParameter, HZBRemapping)
+	LAYOUT_FIELD(FPostProcessPassParameters, PostprocessParameter)
+	LAYOUT_FIELD(FSceneTextureShaderParameters, SceneTextureParameters)
+	LAYOUT_FIELD(FScreenSpaceAOParameters, ScreenSpaceAOParams)
+	LAYOUT_FIELD(FShaderResourceParameter, RandomNormalTexture)
+	LAYOUT_FIELD(FShaderResourceParameter, RandomNormalTextureSampler)
+	LAYOUT_FIELD(FShaderParameter, OutTexture)
 
 	/** Initialization constructor. */
 	FPostProcessAmbientOcclusionPSandCS(const ShaderMetaType::CompiledShaderInitializerType& Initializer)
@@ -839,7 +829,7 @@ public:
 		const FVector4 HZBRemappingValue = GetHZBValue(View);				
 		const FSceneRenderTargetItem& SSAORandomization = GSystemTextures.SSAORandomization->GetRenderTargetItem();
 
-		FRHIComputeShader* ShaderRHI = GetComputeShader();
+		FRHIComputeShader* ShaderRHI = RHICmdList.GetBoundComputeShader();
 		FGlobalShader::SetParameters<FViewUniformShaderParameters>(RHICmdList, ShaderRHI, View.ViewUniformBuffer);			
 		RHICmdList.SetUAVParameter(ShaderRHI, OutTexture.GetBaseIndex(), OutUAV);
 
@@ -859,7 +849,7 @@ public:
 		const FVector4 HZBRemappingValue = GetHZBValue(View);
 		const FSceneRenderTargetItem& SSAORandomization = GSystemTextures.SSAORandomization->GetRenderTargetItem();
 
-		FRHIPixelShader* ShaderRHI = GetPixelShader();
+		FRHIPixelShader* ShaderRHI = RHICmdList.GetBoundPixelShader();
 		FGlobalShader::SetParameters<FViewUniformShaderParameters>(RHICmdList, ShaderRHI, View.ViewUniformBuffer);
 
 		// SF_Point is better than bilinear to avoid halos around objects
@@ -874,18 +864,11 @@ public:
 	template <typename TRHICmdList>
 	void UnsetParameters(TRHICmdList& RHICmdList)
 	{
-		FRHIComputeShader* ShaderRHI = GetComputeShader();
+		FRHIComputeShader* ShaderRHI = RHICmdList.GetBoundComputeShader();
 		RHICmdList.SetUAVParameter(ShaderRHI, OutTexture.GetBaseIndex(), NULL);
 	}
 	
 	// FShader interface.
-	virtual bool Serialize(FArchive& Ar) override
-	{
-		bool bShaderHasOutdatedParameters = FGlobalShader::Serialize(Ar);
-		Ar << HZBRemapping << PostprocessParameter << SceneTextureParameters << ScreenSpaceAOParams << RandomNormalTexture << RandomNormalTextureSampler << OutTexture;
-		return bShaderHasOutdatedParameters;
-	}
-
 	static const TCHAR* GetSourceFilename()
 	{
 		return TEXT("/Engine/Private/PostProcessAmbientOcclusion.usf");
@@ -920,14 +903,14 @@ public:
 // ---------------------------------
 
 template <uint32 bTAOSetupAsInput, uint32 bDoUpsample, uint32 ShaderQuality>
-FShader* FRCPassPostProcessAmbientOcclusion::SetShaderTemplPS(const FRenderingCompositePassContext& Context, FGraphicsPipelineStateInitializer& GraphicsPSOInit)
+TShaderRef<FShader> FRCPassPostProcessAmbientOcclusion::SetShaderTemplPS(const FRenderingCompositePassContext& Context, FGraphicsPipelineStateInitializer& GraphicsPSOInit)
 {
 	TShaderMapRef<FPostProcessVS> VertexShader(Context.GetShaderMap());
 	TShaderMapRef<FPostProcessAmbientOcclusionPSandCS<bTAOSetupAsInput, bDoUpsample, ShaderQuality, false> > PixelShader(Context.GetShaderMap());
 
 	GraphicsPSOInit.BoundShaderState.VertexDeclarationRHI = GFilterVertexDeclaration.VertexDeclarationRHI;
-	GraphicsPSOInit.BoundShaderState.VertexShaderRHI = GETSAFERHISHADER_VERTEX(*VertexShader);
-	GraphicsPSOInit.BoundShaderState.PixelShaderRHI = GETSAFERHISHADER_PIXEL(*PixelShader);
+	GraphicsPSOInit.BoundShaderState.VertexShaderRHI = VertexShader.GetVertexShader();
+	GraphicsPSOInit.BoundShaderState.PixelShaderRHI = PixelShader.GetPixelShader();
 	SetGraphicsPipelineState(Context.RHICmdList, GraphicsPSOInit);
 
 	const FPooledRenderTargetDesc* InputDesc0 = GetInputDesc(ePId_Input0);
@@ -937,7 +920,7 @@ FShader* FRCPassPostProcessAmbientOcclusion::SetShaderTemplPS(const FRenderingCo
 	VertexShader->SetParameters(Context);
 	PixelShader->SetParametersGfx(Context.RHICmdList, Context, TexSize, 0);
 
-	return *VertexShader;
+	return VertexShader;
 }
 
 
@@ -947,7 +930,7 @@ void FRCPassPostProcessAmbientOcclusion::DispatchCS(TRHICmdList& RHICmdList, con
 {
 	TShaderMapRef<FPostProcessAmbientOcclusionPSandCS<bTAOSetupAsInput, bDoUpsample, ShaderQuality, true> > ComputeShader(Context.GetShaderMap());
 
-	RHICmdList.SetComputeShader(ComputeShader->GetComputeShader());
+	RHICmdList.SetComputeShader(ComputeShader.GetComputeShader());
 
 	FSceneRenderTargets& SceneContext = FSceneRenderTargets::Get(Context.RHICmdList);
 
@@ -959,7 +942,7 @@ void FRCPassPostProcessAmbientOcclusion::DispatchCS(TRHICmdList& RHICmdList, con
 
 	uint32 GroupSizeX = FMath::DivideAndRoundUp(ViewRect.Size().X, GAmbientOcclusionTileSizeX);
 	uint32 GroupSizeY = FMath::DivideAndRoundUp(ViewRect.Size().Y, GAmbientOcclusionTileSizeY);
-	DispatchComputeShader(RHICmdList, *ComputeShader, GroupSizeX, GroupSizeY, 1);
+	DispatchComputeShader(RHICmdList, ComputeShader.GetShader(), GroupSizeX, GroupSizeY, 1);
 
 	ComputeShader->UnsetParameters(RHICmdList);
 }
@@ -1106,7 +1089,7 @@ void FRCPassPostProcessAmbientOcclusion::ProcessPS(FRenderingCompositePassContex
 		GraphicsPSOInit.PrimitiveType = PT_TriangleList;
 		GraphicsPSOInit.bDepthBounds = bDepthBoundsTestEnabled;
 
-		FShader* VertexShader = 0;
+		TShaderRef<FShader> VertexShader;
 
 #define SET_SHADER_CASE(ShaderQualityCase) \
 		case ShaderQualityCase: \
@@ -1255,6 +1238,7 @@ FPooledRenderTargetDesc FRCPassPostProcessAmbientOcclusion::ComputeOutputDesc(EP
 /** Shader parameters needed for screen space AmbientOcclusion passes. */
 class FGTAOParameters
 {
+	DECLARE_INLINE_TYPE_LAYOUT(FGTAOParameters, NonVirtual);
 public:
 
 	void Bind(const FShaderParameterMap& ParameterMap)
@@ -1318,7 +1302,7 @@ public:
 	friend FArchive& operator<<(FArchive& Ar, FGTAOParameters& This);
 
 private:
-	FShaderParameter GTAOParams;
+	LAYOUT_FIELD(FShaderParameter, GTAOParams);
 };
 
 FArchive& operator<<(FArchive& Ar, FGTAOParameters& This)
@@ -1363,16 +1347,16 @@ class FPostProcessGTAOHorizonSearchPSandCS : public FGlobalShader
 	FPostProcessGTAOHorizonSearchPSandCS() {}
 
 public:
-	FSceneTextureShaderParameters SceneTextureParameters;
-	FPostProcessPassParameters	  PostprocessParameter;
-	FScreenSpaceAOParameters	  ScreenSpaceAOParams;
-	FGTAOParameters				  GTAOParams;
-	FShaderResourceParameter 	  HorizonOutTexture;
-	FShaderResourceParameter 	  DepthOutTexture;
-	FShaderResourceParameter      RandomNormalTexture;
-	FShaderResourceParameter      RandomNormalTextureSampler;
-	FShaderParameter			  HZBRemapping;
-	FShaderParameter			  HorizonSearchParams;
+	LAYOUT_FIELD(FSceneTextureShaderParameters, SceneTextureParameters);
+	LAYOUT_FIELD(FPostProcessPassParameters, PostprocessParameter);
+	LAYOUT_FIELD(FScreenSpaceAOParameters, ScreenSpaceAOParams);
+	LAYOUT_FIELD(FGTAOParameters, GTAOParams);
+	LAYOUT_FIELD(FShaderResourceParameter, HorizonOutTexture);
+	LAYOUT_FIELD(FShaderResourceParameter, DepthOutTexture);
+	LAYOUT_FIELD(FShaderResourceParameter, RandomNormalTexture);
+	LAYOUT_FIELD(FShaderResourceParameter, RandomNormalTextureSampler);
+	LAYOUT_FIELD(FShaderParameter, HZBRemapping);
+	LAYOUT_FIELD(FShaderParameter, HorizonSearchParams);
 
 	/** Initialization constructor. */
 	FPostProcessGTAOHorizonSearchPSandCS(const ShaderMetaType::CompiledShaderInitializerType& Initializer)
@@ -1407,7 +1391,7 @@ public:
 		FRHIUnorderedAccessView *OutUAV, FRHIUnorderedAccessView *OutDepthsUAV)
 	{
 		const FFinalPostProcessSettings& Settings = Context.View.FinalPostProcessSettings;
-		FRHIComputeShader* ShaderRHI = GetComputeShader();
+		FRHIComputeShader* ShaderRHI = RHICmdList.GetBoundComputeShader();
 
 		FGlobalShader::SetParameters<FViewUniformShaderParameters>(RHICmdList, ShaderRHI, Context.View.ViewUniformBuffer);
 		SceneTextureParameters.Set(RHICmdList, ShaderRHI, Context.View.FeatureLevel, ESceneTextureSetupMode::All);
@@ -1433,7 +1417,7 @@ public:
 	void SetParametersPS(const FRenderingCompositePassContext& Context, FIntPoint DestSize, FIntPoint InputTextureSize)
 	{
 		const FFinalPostProcessSettings& Settings = Context.View.FinalPostProcessSettings;
-		FRHIPixelShader* ShaderRHI = GetPixelShader();
+		FRHIPixelShader* ShaderRHI = Context.RHICmdList.GetBoundPixelShader();
 
 		PostprocessParameter.SetPS(Context.RHICmdList, ShaderRHI, Context, TStaticSamplerState<SF_Point, AM_Clamp, AM_Clamp, AM_Clamp>::GetRHI());
 
@@ -1453,20 +1437,13 @@ public:
 	template <typename TRHICmdList>
 	void UnsetParameters(TRHICmdList& RHICmdList)
 	{
-		FRHIComputeShader* ShaderRHI = GetComputeShader();
+		FRHIComputeShader* ShaderRHI = RHICmdList.GetBoundComputeShader();
 		RHICmdList.SetUAVParameter(ShaderRHI, HorizonOutTexture.GetBaseIndex(), nullptr);
 		RHICmdList.SetUAVParameter(ShaderRHI, DepthOutTexture.GetBaseIndex(), nullptr);
 	}
 
 
 	// FShader interface.
-	virtual bool Serialize(FArchive& Ar) override
-	{
-		bool bShaderHasOutdatedParameters = FGlobalShader::Serialize(Ar);
-		Ar << PostprocessParameter << SceneTextureParameters << ScreenSpaceAOParams << GTAOParams << HorizonOutTexture << DepthOutTexture << HZBRemapping << RandomNormalTexture << RandomNormalTextureSampler << HorizonSearchParams;
-		return bShaderHasOutdatedParameters;
-	}
-
 	static const TCHAR* GetSourceFilename()
 	{
 		return TEXT("/Engine/Private/PostProcessAmbientOcclusion.usf");
@@ -1523,15 +1500,15 @@ class FPostProcessGTAOCombinedPSandCS : public FGlobalShader
 	FPostProcessGTAOCombinedPSandCS() {}
 
 public:
-	FSceneTextureShaderParameters SceneTextureParameters;
-	FPostProcessPassParameters	  PostprocessParameter;
-	FScreenSpaceAOParameters	  ScreenSpaceAOParams;
-	FShaderResourceParameter 	  OutTexture;
-	FShaderResourceParameter 	  DepthOutTexture;
-	FShaderParameter			  HZBRemapping;
-	FShaderResourceParameter      RandomNormalTexture;
-	FShaderResourceParameter      RandomNormalTextureSampler;
-	FGTAOParameters				  GTAOParams;
+	LAYOUT_FIELD(FSceneTextureShaderParameters, SceneTextureParameters);
+	LAYOUT_FIELD(FPostProcessPassParameters, PostprocessParameter);
+	LAYOUT_FIELD(FScreenSpaceAOParameters, ScreenSpaceAOParams);
+	LAYOUT_FIELD(FShaderResourceParameter, OutTexture);
+	LAYOUT_FIELD(FShaderResourceParameter, DepthOutTexture);
+	LAYOUT_FIELD(FShaderParameter, HZBRemapping);
+	LAYOUT_FIELD(FShaderResourceParameter, RandomNormalTexture);
+	LAYOUT_FIELD(FShaderResourceParameter, RandomNormalTextureSampler);
+	LAYOUT_FIELD(FGTAOParameters, GTAOParams);
 
 	/** Initialization constructor. */
 	FPostProcessGTAOCombinedPSandCS(const ShaderMetaType::CompiledShaderInitializerType& Initializer)
@@ -1567,7 +1544,7 @@ public:
 		FRHIUnorderedAccessView *OutUAV, FRHIUnorderedAccessView *DepthOutUAV)
 	{
 		const FFinalPostProcessSettings& Settings = Context.View.FinalPostProcessSettings;
-		FRHIComputeShader* ShaderRHI = GetComputeShader();
+		FRHIComputeShader* ShaderRHI = RHICmdList.GetBoundComputeShader();
 
 		FGlobalShader::SetParameters<FViewUniformShaderParameters>(RHICmdList, ShaderRHI, Context.View.ViewUniformBuffer);
 		SceneTextureParameters.Set(RHICmdList, ShaderRHI, Context.View.FeatureLevel, ESceneTextureSetupMode::All);
@@ -1589,7 +1566,7 @@ public:
 	void SetParametersPS(const FRenderingCompositePassContext& Context, FIntPoint DestSize, FIntPoint InputTextureSize)
 	{
 		const FFinalPostProcessSettings& Settings = Context.View.FinalPostProcessSettings;
-		FRHIPixelShader* ShaderRHI = GetPixelShader();
+		FRHIPixelShader* ShaderRHI = Context.RHICmdList.GetBoundPixelShader();
 
 		PostprocessParameter.SetPS(Context.RHICmdList, ShaderRHI, Context, TStaticSamplerState<SF_Point, AM_Clamp, AM_Clamp, AM_Clamp>::GetRHI());
 
@@ -1609,19 +1586,13 @@ public:
 	template <typename TRHICmdList>
 	void UnsetParameters(TRHICmdList& RHICmdList)
 	{
-		FRHIComputeShader* ShaderRHI = GetComputeShader();
+		FRHIComputeShader* ShaderRHI = RHICmdList.GetBoundComputeShader();
 		RHICmdList.SetUAVParameter(ShaderRHI, OutTexture.GetBaseIndex(), nullptr);
 		RHICmdList.SetUAVParameter(ShaderRHI, DepthOutTexture.GetBaseIndex(), nullptr);
 	}
 
 
 	// FShader interface.
-	virtual bool Serialize(FArchive& Ar) override
-	{
-		bool bShaderHasOutdatedParameters = FGlobalShader::Serialize(Ar);
-		Ar << PostprocessParameter << SceneTextureParameters << ScreenSpaceAOParams << OutTexture << DepthOutTexture << GTAOParams <<HZBRemapping << RandomNormalTexture << RandomNormalTextureSampler;
-		return bShaderHasOutdatedParameters;
-	}
 
 	static const TCHAR* GetSourceFilename()
 	{
@@ -1678,16 +1649,16 @@ class FPostProcessGTAOInnerIntegratePSandCS : public FGlobalShader
 	FPostProcessGTAOInnerIntegratePSandCS() {}
 
 public:
-	FSceneTextureShaderParameters SceneTextureParameters;
-	FPostProcessPassParameters	  PostprocessParameter;
-	FScreenSpaceAOParameters	  ScreenSpaceAOParams;
-	FShaderResourceParameter 	  OutTexture;
-	FShaderResourceParameter      RandomNormalTexture;
-	FShaderResourceParameter      RandomNormalTextureSampler;
-	FShaderResourceParameter      HorizonsTexture;
-	FShaderResourceParameter      HorizonsTextureSampler;
-	FGTAOParameters				  GTAOParams;
-	FShaderParameter			  InnerIntegrateParams;
+	LAYOUT_FIELD(FSceneTextureShaderParameters, SceneTextureParameters);
+	LAYOUT_FIELD(FPostProcessPassParameters, PostprocessParameter);
+	LAYOUT_FIELD(FScreenSpaceAOParameters, ScreenSpaceAOParams);
+	LAYOUT_FIELD(FShaderResourceParameter, OutTexture);
+	LAYOUT_FIELD(FShaderResourceParameter, RandomNormalTexture);
+	LAYOUT_FIELD(FShaderResourceParameter, RandomNormalTextureSampler);
+	LAYOUT_FIELD(FShaderResourceParameter, HorizonsTexture);
+	LAYOUT_FIELD(FShaderResourceParameter, HorizonsTextureSampler);
+	LAYOUT_FIELD(FGTAOParameters, GTAOParams);
+	LAYOUT_FIELD(FShaderParameter, InnerIntegrateParams);
 
 
 	/** Initialization constructor. */
@@ -1710,21 +1681,13 @@ public:
 		}
 	}
 	// FShader interface.
-	virtual bool Serialize(FArchive& Ar) override
-	{
-		bool bShaderHasOutdatedParameters = FGlobalShader::Serialize(Ar);
-		Ar << PostprocessParameter << SceneTextureParameters << ScreenSpaceAOParams << OutTexture << GTAOParams <<  RandomNormalTexture << RandomNormalTextureSampler;
-		Ar << HorizonsTexture << HorizonsTextureSampler << InnerIntegrateParams;
-
-		return bShaderHasOutdatedParameters;
-	}
 
 
 	template <typename TRHICmdList>
 	void SetParametersCS(TRHICmdList& RHICmdList, const FRenderingCompositePassContext& Context, uint32 DownScaleFactor, FIntPoint DestSize, FIntPoint InputTextureSize, FRHIUnorderedAccessView *OutUAV)
 	{
 		const FFinalPostProcessSettings& Settings = Context.View.FinalPostProcessSettings;
-		FRHIComputeShader* ShaderRHI = GetComputeShader();
+		FRHIComputeShader* ShaderRHI = RHICmdList.GetBoundComputeShader();
 
 		FGlobalShader::SetParameters<FViewUniformShaderParameters>(RHICmdList, ShaderRHI, Context.View.ViewUniformBuffer);
 		SceneTextureParameters.Set(RHICmdList, ShaderRHI, Context.View.FeatureLevel, ESceneTextureSetupMode::All);
@@ -1751,7 +1714,7 @@ public:
 	void SetParametersPS(const FRenderingCompositePassContext& Context, FIntPoint DestSize)
 	{
 		const FFinalPostProcessSettings& Settings = Context.View.FinalPostProcessSettings;
-		FRHIPixelShader* ShaderRHI = GetPixelShader();
+		FRHIPixelShader* ShaderRHI = Context.RHICmdList.GetBoundPixelShader();
 
 		PostprocessParameter.SetPS(Context.RHICmdList, ShaderRHI, Context, TStaticSamplerState<SF_Point, AM_Clamp, AM_Clamp, AM_Clamp>::GetRHI());
 
@@ -1772,7 +1735,7 @@ public:
 	template <typename TRHICmdList>
 	void UnsetParameters(TRHICmdList& RHICmdList)
 	{
-		FRHIComputeShader* ShaderRHI = GetComputeShader();
+		FRHIComputeShader* ShaderRHI = RHICmdList.GetBoundComputeShader();
 		RHICmdList.SetUAVParameter(ShaderRHI, OutTexture.GetBaseIndex(), nullptr);
 	}
 
@@ -1806,24 +1769,24 @@ void FRCPassPostProcessAmbientOcclusion_GTAOHorizonSearchIntegrate::DispatchCS(T
 	const FSceneRenderTargetItem& DestRenderTarget = PassOutputs[0].RequestSurface(Context);
 	const FSceneRenderTargetItem& DestRenderTarget1 = SceneContext.ScreenSpaceGTAODepths->GetRenderTargetItem();
 
-	RHICmdList.SetComputeShader(ComputeShader->GetComputeShader());
+	RHICmdList.SetComputeShader(ComputeShader.GetComputeShader());
 	ComputeShader->SetParametersCS(RHICmdList, Context, DestSize, TexSize, DownScaleFactor, DestRenderTarget.UAV, DestRenderTarget1.UAV);
 
 	uint32 GroupSizeX = FMath::DivideAndRoundUp(ViewRect.Width(), 8);
 	uint32 GroupSizeY = FMath::DivideAndRoundUp(ViewRect.Height(), 8);
-	DispatchComputeShader(RHICmdList, *ComputeShader, GroupSizeX, GroupSizeY, 1);
+	DispatchComputeShader(RHICmdList, ComputeShader, GroupSizeX, GroupSizeY, 1);
 	ComputeShader->UnsetParameters(RHICmdList);
 }
 
 template <uint32 ShaderQuality, uint32 UseNormals>
-FShader* FRCPassPostProcessAmbientOcclusion_GTAOHorizonSearchIntegrate::SetShaderPS(const FRenderingCompositePassContext& Context, FGraphicsPipelineStateInitializer& GraphicsPSOInit, FIntPoint DestSize)
+TShaderRef<FShader> FRCPassPostProcessAmbientOcclusion_GTAOHorizonSearchIntegrate::SetShaderPS(const FRenderingCompositePassContext& Context, FGraphicsPipelineStateInitializer& GraphicsPSOInit, FIntPoint DestSize)
 {
 	TShaderMapRef<FPostProcessVS> VertexShader(Context.GetShaderMap());
 	TShaderMapRef<FPostProcessGTAOCombinedPSandCS<false, ShaderQuality, UseNormals> > PixelShader(Context.GetShaderMap());
 
 	GraphicsPSOInit.BoundShaderState.VertexDeclarationRHI = GFilterVertexDeclaration.VertexDeclarationRHI;
-	GraphicsPSOInit.BoundShaderState.VertexShaderRHI = GETSAFERHISHADER_VERTEX(*VertexShader);
-	GraphicsPSOInit.BoundShaderState.PixelShaderRHI = GETSAFERHISHADER_PIXEL(*PixelShader);
+	GraphicsPSOInit.BoundShaderState.VertexShaderRHI = VertexShader.GetVertexShader();
+	GraphicsPSOInit.BoundShaderState.PixelShaderRHI = PixelShader.GetPixelShader();
 	SetGraphicsPipelineState(Context.RHICmdList, GraphicsPSOInit);
 
 	const FPooledRenderTargetDesc* InputDesc0 = GetInputDesc(ePId_Input0);
@@ -1833,7 +1796,7 @@ FShader* FRCPassPostProcessAmbientOcclusion_GTAOHorizonSearchIntegrate::SetShade
 	VertexShader->SetParameters(Context);
 	PixelShader->SetParametersPS(Context, DestSize, TexSize);
 
-	return *VertexShader;
+	return VertexShader;
 }
 
 void FRCPassPostProcessAmbientOcclusion_GTAOHorizonSearchIntegrate::Process(FRenderingCompositePassContext& Context)
@@ -1850,10 +1813,10 @@ void FRCPassPostProcessAmbientOcclusion_GTAOHorizonSearchIntegrate::Process(FRen
 	FIntPoint InputTexSize = InputDesc0 ? InputDesc0->Extent : SceneContext.GetBufferSizeXY();
 	FIntPoint OutputTexSize = PassOutputs[0].RenderTargetDesc.Extent;
 	FIntRect ViewRect = View.ViewRect / DownScaleFactor;
-   
+
 	static const auto CVarCompute = IConsoleManager::Get().FindTConsoleVariableDataInt(TEXT("r.AmbientOcclusion.Compute"));
 	static const auto CVarNormals = IConsoleManager::Get().FindTConsoleVariableDataInt(TEXT("r.GTAO.UseNormals"));
-
+	   
 	bool bUseNormals = CVarNormals->GetValueOnRenderThread() >= 1;
 	const int32 ShaderQuality = FSSAOHelper::GetAmbientOcclusionShaderLevel(Context.View);
 
@@ -1936,7 +1899,7 @@ void FRCPassPostProcessAmbientOcclusion_GTAOHorizonSearchIntegrate::Process(FRen
 			GraphicsPSOInit.PrimitiveType		= PT_TriangleList;
 			GraphicsPSOInit.bDepthBounds		= false;
 
-			FShader* VertexShader = 0;
+			TShaderRef<FShader> VertexShader;
 
 #define SET_SHADER_CASE_PS(ShaderQualityCase) \
 			case ShaderQualityCase: \
@@ -2026,12 +1989,12 @@ void FRCPassPostProcessAmbientOcclusion_GTAOInnerIntegrate::Process(FRenderingCo
 		UnbindRenderTargets(Context.RHICmdList);
 
 		TShaderMapRef<FPostProcessGTAOInnerIntegratePSandCS<true>> ComputeShader(Context.GetShaderMap());
-		Context.RHICmdList.SetComputeShader(ComputeShader->GetComputeShader());
+		Context.RHICmdList.SetComputeShader(ComputeShader.GetComputeShader());
 		ComputeShader->SetParametersCS(Context.RHICmdList, Context, DownScaleFactor, TexSize, TexSize, DestRenderTarget.UAV);
 
 		uint32 GroupSizeX = FMath::DivideAndRoundUp(ViewRect.Width(), 8);
 		uint32 GroupSizeY = FMath::DivideAndRoundUp(ViewRect.Height(), 8);
-		DispatchComputeShader(Context.RHICmdList, *ComputeShader, GroupSizeX, GroupSizeY, 1);
+		DispatchComputeShader(Context.RHICmdList, ComputeShader, GroupSizeX, GroupSizeY, 1);
 
 		ComputeShader->UnsetParameters(Context.RHICmdList);
 	}
@@ -2053,8 +2016,8 @@ void FRCPassPostProcessAmbientOcclusion_GTAOInnerIntegrate::Process(FRenderingCo
 			GraphicsPSOInit.DepthStencilState = TStaticDepthStencilState<false, CF_Always>::GetRHI();
 			GraphicsPSOInit.PrimitiveType = PT_TriangleList;
 			GraphicsPSOInit.BoundShaderState.VertexDeclarationRHI = GFilterVertexDeclaration.VertexDeclarationRHI;
-			GraphicsPSOInit.BoundShaderState.VertexShaderRHI = GETSAFERHISHADER_VERTEX(*VertexShader);
-			GraphicsPSOInit.BoundShaderState.PixelShaderRHI = GETSAFERHISHADER_PIXEL(*PixelShader);
+			GraphicsPSOInit.BoundShaderState.VertexShaderRHI = VertexShader.GetVertexShader();
+			GraphicsPSOInit.BoundShaderState.PixelShaderRHI = PixelShader.GetPixelShader();
 			SetGraphicsPipelineState(Context.RHICmdList, GraphicsPSOInit);
 
 			PixelShader->SetParametersPS(Context, TexSize);
@@ -2065,7 +2028,7 @@ void FRCPassPostProcessAmbientOcclusion_GTAOInnerIntegrate::Process(FRenderingCo
 				ViewRect.Min.X, ViewRect.Min.Y, ViewRect.Width(), ViewRect.Height(),
 				ViewRect.Size(),
 				TexSize,
-				*VertexShader,
+				VertexShader,
 				EDRF_UseTriangleOptimization);
 		}
 		Context.RHICmdList.EndRenderPass();
@@ -2112,25 +2075,25 @@ void FRCPassPostProcessAmbientOcclusion_HorizonSearch::DispatchCS(const FRenderi
 	FRHIAsyncComputeCommandListImmediate& RHICmdListComputeImmediate = FRHICommandListExecutor::GetImmediateAsyncComputeCommandList();
 
 	TShaderMapRef<FPostProcessGTAOHorizonSearchPSandCS<true, ShaderQuality>> ComputeShader(Context.GetShaderMap());
-	RHICmdListComputeImmediate.SetComputeShader(ComputeShader->GetComputeShader());
+	RHICmdListComputeImmediate.SetComputeShader(ComputeShader.GetComputeShader());
 	ComputeShader->SetParametersCS(RHICmdListComputeImmediate, Context, DestSize, TexSize, DownScaleFactor, DestRenderTarget.UAV, DestRenderTarget2.UAV);
 
 	uint32 GroupSizeX = FMath::DivideAndRoundUp(ViewRect.Width(), 8);
 	uint32 GroupSizeY = FMath::DivideAndRoundUp(ViewRect.Height(), 8);
-	DispatchComputeShader(RHICmdListComputeImmediate, *ComputeShader, GroupSizeX, GroupSizeY, 1);
+	DispatchComputeShader(RHICmdListComputeImmediate, ComputeShader, GroupSizeX, GroupSizeY, 1);
 
 	ComputeShader->UnsetParameters(RHICmdListComputeImmediate);
 }
 
 template <uint32 ShaderQuality>
-FShader* FRCPassPostProcessAmbientOcclusion_HorizonSearch::SetShaderPS(const FRenderingCompositePassContext& Context, FGraphicsPipelineStateInitializer& GraphicsPSOInit, FIntPoint DestSize)
+TShaderRef<FShader> FRCPassPostProcessAmbientOcclusion_HorizonSearch::SetShaderPS(const FRenderingCompositePassContext& Context, FGraphicsPipelineStateInitializer& GraphicsPSOInit, FIntPoint DestSize)
 {
 	TShaderMapRef<FPostProcessVS> VertexShader(Context.GetShaderMap());
 	TShaderMapRef<FPostProcessGTAOHorizonSearchPSandCS<false, ShaderQuality> > PixelShader(Context.GetShaderMap());
 
 	GraphicsPSOInit.BoundShaderState.VertexDeclarationRHI = GFilterVertexDeclaration.VertexDeclarationRHI;
-	GraphicsPSOInit.BoundShaderState.VertexShaderRHI = GETSAFERHISHADER_VERTEX(*VertexShader);
-	GraphicsPSOInit.BoundShaderState.PixelShaderRHI = GETSAFERHISHADER_PIXEL(*PixelShader);
+	GraphicsPSOInit.BoundShaderState.VertexShaderRHI = VertexShader.GetVertexShader();
+	GraphicsPSOInit.BoundShaderState.PixelShaderRHI = PixelShader.GetPixelShader();
 	SetGraphicsPipelineState(Context.RHICmdList, GraphicsPSOInit);
 
 	const FPooledRenderTargetDesc* InputDesc0 = GetInputDesc(ePId_Input0);
@@ -2140,7 +2103,7 @@ FShader* FRCPassPostProcessAmbientOcclusion_HorizonSearch::SetShaderPS(const FRe
 	VertexShader->SetParameters(Context);
 	PixelShader->SetParametersPS(Context, DestSize, TexSize);
 
-	return *VertexShader;
+	return VertexShader;
 }
 
 
@@ -2218,8 +2181,7 @@ void FRCPassPostProcessAmbientOcclusion_HorizonSearch::Process(FRenderingComposi
 			GraphicsPSOInit.PrimitiveType = PT_TriangleList;
 			GraphicsPSOInit.bDepthBounds = false;
 
-			FShader* VertexShader = 0;
-
+			TShaderRef<FShader> VertexShader;
 			switch (ShaderQuality)
 			{
 			case 0:
@@ -2301,30 +2263,30 @@ class FPostProcessGTAOTemporalFilterPSandCS : public FGlobalShader
 	FPostProcessGTAOTemporalFilterPSandCS() {}
 
 public:
-	FSceneTextureShaderParameters	SceneTextureParameters;
-	FPostProcessPassParameters		PostprocessParameter;
-	FShaderResourceParameter		HistoryTexture;
-	FShaderResourceParameter		HistoryTextureSampler;
+	LAYOUT_FIELD(FSceneTextureShaderParameters, SceneTextureParameters);
+	LAYOUT_FIELD(FPostProcessPassParameters, 	PostprocessParameter);
+	LAYOUT_FIELD(FShaderResourceParameter, 		HistoryTexture);
+	LAYOUT_FIELD(FShaderResourceParameter, 		HistoryTextureSampler);
 
-	FShaderResourceParameter		ZCurrTexture;
-	FShaderResourceParameter		ZCurrTextureSampler;
+	LAYOUT_FIELD(FShaderResourceParameter,		ZCurrTexture);
+	LAYOUT_FIELD(FShaderResourceParameter,		ZCurrTextureSampler);
 
-	FShaderResourceParameter		ZPrevTexture;
-	FShaderResourceParameter		ZPrevTextureSampler;
+	LAYOUT_FIELD(FShaderResourceParameter,		ZPrevTexture);
+	LAYOUT_FIELD(FShaderResourceParameter,		ZPrevTextureSampler);
 
-	FShaderResourceParameter		VelocityPrevTexture;
-	FShaderResourceParameter		VelocityPrevTextureSampler;
+	LAYOUT_FIELD(FShaderResourceParameter,		VelocityPrevTexture);
+	LAYOUT_FIELD(FShaderResourceParameter,		VelocityPrevTextureSampler);
 
-	FShaderResourceParameter		SceneVelocityTexture;
-	FShaderResourceParameter		SceneVelocityTextureSampler;
+	LAYOUT_FIELD(FShaderResourceParameter, 		SceneVelocityTexture);
+	LAYOUT_FIELD(FShaderResourceParameter, 		SceneVelocityTextureSampler);
 
-	FShaderParameter				PrevScreenPositionScaleBias;
-	FShaderParameter				BlendParams;
-	FScreenSpaceAOParameters		ScreenSpaceAOParams;
+	LAYOUT_FIELD(FShaderParameter, 				PrevScreenPositionScaleBias);
+	LAYOUT_FIELD(FShaderParameter, 				BlendParams);
+	LAYOUT_FIELD(FScreenSpaceAOParameters, 		ScreenSpaceAOParams);
 
-	FShaderResourceParameter		OutTexture;
-	FShaderResourceParameter 		DepthOutTexture;
-	FShaderResourceParameter 		VelocityOutTexture;
+	LAYOUT_FIELD(FShaderParameter, 				OutTexture);
+	LAYOUT_FIELD(FShaderResourceParameter, 		DepthOutTexture);
+	LAYOUT_FIELD(FShaderResourceParameter, 		VelocityOutTexture);
 
 
 	/** Initialization constructor. */
@@ -2367,7 +2329,7 @@ public:
 		FRHIUnorderedAccessView *OutUAV_AO, FRHIUnorderedAccessView *OutUAV_Depth, FRHIUnorderedAccessView *OutUAV_Velocity)
 	{
 		const FFinalPostProcessSettings& Settings = Context.View.FinalPostProcessSettings;
-		FRHIComputeShader* ShaderRHI = GetComputeShader();
+		FRHIComputeShader* ShaderRHI = RHICmdList.GetBoundComputeShader();
 		FSceneRenderTargets& SceneContext = FSceneRenderTargets::Get(Context.RHICmdList);
 		PostprocessParameter.SetCS(ShaderRHI, Context, RHICmdList, TStaticSamplerState<SF_Bilinear, AM_Clamp, AM_Clamp, AM_Clamp>::GetRHI());
 
@@ -2428,7 +2390,7 @@ public:
 	void SetParametersPS(const FRenderingCompositePassContext& Context, FIntPoint DestSize, FIntPoint InputTextureSize, bool bCameraCut, const FGTAOTAAHistory& InputHistory, TRefCountPtr<IPooledRenderTarget>& VelocityRT)
 	{
 		const FFinalPostProcessSettings& Settings = Context.View.FinalPostProcessSettings;
-		FRHIPixelShader* ShaderRHI = GetPixelShader();
+		FRHIPixelShader* ShaderRHI = Context.RHICmdList.GetBoundPixelShader();
 		FSceneRenderTargets& SceneContext = FSceneRenderTargets::Get(Context.RHICmdList);
 
 		PostprocessParameter.SetPS(Context.RHICmdList, ShaderRHI, Context, TStaticSamplerState<SF_Point, AM_Clamp, AM_Clamp, AM_Clamp>::GetRHI());
@@ -2485,7 +2447,7 @@ public:
 	template <typename TRHICmdList>
 	void UnsetParameters(TRHICmdList& RHICmdList)
 	{
-		FRHIComputeShader* ShaderRHI = GetComputeShader();
+		FRHIComputeShader* ShaderRHI = RHICmdList.GetBoundComputeShader();
 		RHICmdList.SetUAVParameter(ShaderRHI, OutTexture.GetBaseIndex(), nullptr);
 		RHICmdList.SetUAVParameter(ShaderRHI, DepthOutTexture.GetBaseIndex(), nullptr);
 		RHICmdList.SetUAVParameter(ShaderRHI, VelocityOutTexture.GetBaseIndex(), nullptr);
@@ -2493,14 +2455,6 @@ public:
 
 
 	// FShader interface.
-	virtual bool Serialize(FArchive& Ar) override
-	{
-		bool bShaderHasOutdatedParameters = FGlobalShader::Serialize(Ar);
-		Ar << ZCurrTexture << ZCurrTextureSampler << HistoryTexture << HistoryTextureSampler << PostprocessParameter << SceneTextureParameters << PrevScreenPositionScaleBias << OutTexture;
-		Ar << ZPrevTexture << ZPrevTextureSampler << VelocityPrevTexture << VelocityPrevTextureSampler << SceneVelocityTexture << SceneVelocityTextureSampler << BlendParams << DepthOutTexture << VelocityOutTexture;
-		return bShaderHasOutdatedParameters;
-	}
-
 	static const TCHAR* GetSourceFilename()
 	{
 		return TEXT("/Engine/Private/PostProcessAmbientOcclusion.usf");
@@ -2563,16 +2517,16 @@ void FRCPassPostProcessAmbientOcclusion_GTAO_TemporalFilter::Process(FRenderingC
 		TShaderMapRef<FPostProcessGTAOTemporalFilterPSandCS<true>> ComputeShader(Context.GetShaderMap());
 		uint32 GroupSizeX = FMath::DivideAndRoundUp(OutputViewRect.Width(), 8);
 		uint32 GroupSizeY = FMath::DivideAndRoundUp(OutputViewRect.Height(), 8);
-
+	
 		Context.RHICmdList.TransitionResource(EResourceTransitionAccess::ERWBarrier, EResourceTransitionPipeline::EComputeToCompute, DestRenderTarget0.UAV, nullptr);
 		Context.RHICmdList.TransitionResource(EResourceTransitionAccess::ERWBarrier, EResourceTransitionPipeline::EComputeToCompute, DestRenderTarget1.UAV, nullptr);
 		Context.RHICmdList.TransitionResource(EResourceTransitionAccess::ERWBarrier, EResourceTransitionPipeline::EComputeToCompute, DestRenderTarget2.UAV, nullptr);
 
-		Context.RHICmdList.SetComputeShader(ComputeShader->GetComputeShader());
+		Context.RHICmdList.SetComputeShader(ComputeShader.GetComputeShader());
 		ComputeShader->SetParametersCS(Context.RHICmdList, Context, OutputTexSize, InputTexSize, bCameraCut, InputHistory,
 			bCameraCut ? GSystemTextures.BlackDummy : SceneContext.SceneVelocity, DestRenderTarget0.UAV, DestRenderTarget1.UAV, DestRenderTarget2.UAV);
 
-		DispatchComputeShader(Context.RHICmdList, *ComputeShader, GroupSizeX, GroupSizeY, 1);
+		DispatchComputeShader(Context.RHICmdList, ComputeShader, GroupSizeX, GroupSizeY, 1);
 		ComputeShader->UnsetParameters(Context.RHICmdList);
 	}
 	else
@@ -2610,8 +2564,8 @@ void FRCPassPostProcessAmbientOcclusion_GTAO_TemporalFilter::Process(FRenderingC
 			GraphicsPSOInit.bDepthBounds = false;
 
 			GraphicsPSOInit.BoundShaderState.VertexDeclarationRHI = GFilterVertexDeclaration.VertexDeclarationRHI;
-			GraphicsPSOInit.BoundShaderState.VertexShaderRHI = GETSAFERHISHADER_VERTEX(*VertexShader);
-			GraphicsPSOInit.BoundShaderState.PixelShaderRHI = GETSAFERHISHADER_PIXEL(*PixelShader);
+			GraphicsPSOInit.BoundShaderState.VertexShaderRHI = VertexShader.GetVertexShader();
+			GraphicsPSOInit.BoundShaderState.PixelShaderRHI = PixelShader.GetPixelShader();
 			SetGraphicsPipelineState(Context.RHICmdList, GraphicsPSOInit);
 
 			VertexShader->SetParameters(Context);
@@ -2624,7 +2578,7 @@ void FRCPassPostProcessAmbientOcclusion_GTAO_TemporalFilter::Process(FRenderingC
 				InputViewRect.Min.X, InputViewRect.Min.Y, InputViewRect.Width(), InputViewRect.Height(),
 				OutputTexSize,
 				InputTexSize,
-				*VertexShader,
+				VertexShader,
 				EDRF_UseTriangleOptimization);
 
 		}
@@ -2688,12 +2642,12 @@ class FPostProcessGTAOSpatialFilterPSandCS : public FGlobalShader
 	}
 
 	FPostProcessGTAOSpatialFilterPSandCS() {}
-	FPostProcessPassParameters	PostprocessParameter;
-	FShaderParameter			OutTexture;
-	FShaderParameter			FilterParams;
-	FScreenSpaceAOParameters	ScreenSpaceAOParams;
-	FShaderResourceParameter	ZReadTexture;
-	FShaderResourceParameter    ZReadTextureSampler;
+	LAYOUT_FIELD(FPostProcessPassParameters, PostprocessParameter);
+	LAYOUT_FIELD(FShaderParameter, OutTexture);
+	LAYOUT_FIELD(FShaderParameter, FilterParams);
+	LAYOUT_FIELD(FScreenSpaceAOParameters, ScreenSpaceAOParams);
+	LAYOUT_FIELD(FShaderResourceParameter, ZReadTexture);
+	LAYOUT_FIELD(FShaderResourceParameter, ZReadTextureSampler);
 
 public:
 
@@ -2710,22 +2664,15 @@ public:
 	}
 
 	// FShader interface.
-	virtual bool Serialize(FArchive& Ar) override
-	{
-		bool bShaderHasOutdatedParameters = FGlobalShader::Serialize(Ar);
-		Ar << PostprocessParameter << OutTexture << FilterParams << ScreenSpaceAOParams << ZReadTexture << ZReadTextureSampler;
-		return bShaderHasOutdatedParameters;
-	}
-
 	void SetParametersPS(const FRenderingCompositePassContext& Context)
 	{
 		const FFinalPostProcessSettings& Settings = Context.View.FinalPostProcessSettings;
-		FRHIPixelShader* ShaderRHI = GetPixelShader();
+		FRHIPixelShader* ShaderRHI = Context.RHICmdList.GetBoundPixelShader();
 
 		FGlobalShader::SetParameters<FViewUniformShaderParameters>(Context.RHICmdList, ShaderRHI, Context.View.ViewUniformBuffer);
 		//ScreenSpaceAOParams.Set(Context.RHICmdList, Context.View, ShaderRHI, InputTextureSize);
 		PostprocessParameter.SetPS(Context.RHICmdList, ShaderRHI, Context, TStaticSamplerState<SF_Bilinear, AM_Clamp, AM_Clamp, AM_Clamp>::GetRHI());
-	
+
 		FSceneRenderTargets& SceneContext = FSceneRenderTargets::Get(Context.RHICmdList);
 
 		SetTextureParameter(Context.RHICmdList, ShaderRHI, ZReadTexture, ZReadTextureSampler, TStaticSamplerState<SF_Point, AM_Wrap, AM_Wrap, AM_Wrap>::GetRHI(),
@@ -2737,7 +2684,7 @@ public:
 	void SetParametersCS(TRHICmdList& RHICmdList, const FRenderingCompositePassContext& Context, FIntPoint InputTextureSize, const FIntRect& OutputRect, FRHIUnorderedAccessView *OutUAV)
 	{
 		const FFinalPostProcessSettings& Settings = Context.View.FinalPostProcessSettings;
-		FRHIComputeShader* ShaderRHI = GetComputeShader();
+		FRHIComputeShader* ShaderRHI = RHICmdList.GetBoundComputeShader();
 
 		FGlobalShader::SetParameters<FViewUniformShaderParameters>(RHICmdList, ShaderRHI, Context.View.ViewUniformBuffer);
 		PostprocessParameter.SetCS(ShaderRHI, Context, RHICmdList, TStaticSamplerState<SF_Bilinear, AM_Clamp, AM_Clamp, AM_Clamp>::GetRHI());
@@ -2756,7 +2703,7 @@ public:
 	template <typename TRHICmdList>
 	void UnsetParameters(TRHICmdList& RHICmdList)
 	{
-		FRHIComputeShader* ShaderRHI = GetComputeShader();
+		FRHIComputeShader* ShaderRHI = RHICmdList.GetBoundComputeShader();
 		RHICmdList.SetUAVParameter(ShaderRHI, OutTexture.GetBaseIndex(), nullptr);
 	}
 
@@ -2808,19 +2755,19 @@ void FRCPassPostProcessAmbientOcclusion_GTAO_SpatialFilter::Process(FRenderingCo
 
 		FRHIAsyncComputeCommandListImmediate& RHICmdListComputeImmediate = FRHICommandListExecutor::GetImmediateAsyncComputeCommandList();
 
-		RHICmdListComputeImmediate.SetComputeShader(ComputeShader->GetComputeShader());
+		RHICmdListComputeImmediate.SetComputeShader(ComputeShader.GetComputeShader());
 		RHICmdListComputeImmediate.TransitionResource(EResourceTransitionAccess::ERWBarrier, EResourceTransitionPipeline::EComputeToCompute, DestRenderTarget.UAV, nullptr);
 		ComputeShader->SetParametersCS(RHICmdListComputeImmediate, Context, InputTexSize, OutputViewRect, DestRenderTarget.UAV);
-		DispatchComputeShader(RHICmdListComputeImmediate, *ComputeShader, GroupSizeX, GroupSizeY, 1);
+		DispatchComputeShader(RHICmdListComputeImmediate, ComputeShader, GroupSizeX, GroupSizeY, 1);
 	}
 	else
 	{
 		const FSceneRenderTargetItem& DestRenderTarget = PassOutputs[0].RequestSurface(Context);
 
-		Context.RHICmdList.SetComputeShader(ComputeShader->GetComputeShader());
+		Context.RHICmdList.SetComputeShader(ComputeShader.GetComputeShader());
 		Context.RHICmdList.TransitionResource(EResourceTransitionAccess::ERWBarrier, EResourceTransitionPipeline::EGfxToCompute, DestRenderTarget.UAV, nullptr);
 		ComputeShader->SetParametersCS(Context.RHICmdList, Context, InputTexSize, OutputViewRect, DestRenderTarget.UAV);
-		DispatchComputeShader(Context.RHICmdList, *ComputeShader, GroupSizeX, GroupSizeY, 1);
+		DispatchComputeShader(Context.RHICmdList, ComputeShader, GroupSizeX, GroupSizeY, 1);
 	}
 
 	ComputeShader->UnsetParameters(Context.RHICmdList);
@@ -2869,11 +2816,11 @@ class FPostProcessGTAOUpsamplePSandCS : public FGlobalShader
 	}
 
 	FPostProcessGTAOUpsamplePSandCS() {}
-	FPostProcessPassParameters		PostprocessParameter;
-	FShaderParameter				OutTexture;
-	FScreenSpaceAOParameters		ScreenSpaceAOParams;
-	FShaderResourceParameter	ZReadTexture;
-	FShaderResourceParameter    ZReadTextureSampler;
+	LAYOUT_FIELD(FPostProcessPassParameters, PostprocessParameter);
+	LAYOUT_FIELD(FShaderParameter, OutTexture);
+	LAYOUT_FIELD(FScreenSpaceAOParameters, ScreenSpaceAOParams);
+	LAYOUT_FIELD(FShaderResourceParameter, ZReadTexture);
+	LAYOUT_FIELD(FShaderResourceParameter, ZReadTextureSampler);
 
 public:
 
@@ -2889,17 +2836,10 @@ public:
 	}
 
 	// FShader interface.
-	virtual bool Serialize(FArchive& Ar) override
-	{
-		bool bShaderHasOutdatedParameters = FGlobalShader::Serialize(Ar);
-		Ar << PostprocessParameter << OutTexture << ScreenSpaceAOParams << ZReadTexture << ZReadTextureSampler;
-		return bShaderHasOutdatedParameters;
-	}
-
 	void SetParametersPS(const FRenderingCompositePassContext& Context, FIntPoint InputTextureSize)
 	{
 		const FFinalPostProcessSettings& Settings = Context.View.FinalPostProcessSettings;
-		FRHIPixelShader* ShaderRHI = GetPixelShader();
+		FRHIPixelShader* ShaderRHI = Context.RHICmdList.GetBoundPixelShader();
 
 		FGlobalShader::SetParameters<FViewUniformShaderParameters>(Context.RHICmdList, ShaderRHI, Context.View.ViewUniformBuffer);
 		ScreenSpaceAOParams.Set(Context.RHICmdList, Context.View, ShaderRHI, InputTextureSize);
@@ -2917,7 +2857,7 @@ public:
 	void SetParametersCS(TRHICmdList& RHICmdList, const FRenderingCompositePassContext& Context, FIntPoint InputTextureSize, FRHIUnorderedAccessView *OutUAV)
 	{
 		const FFinalPostProcessSettings& Settings = Context.View.FinalPostProcessSettings;
-		FRHIComputeShader* ShaderRHI = GetComputeShader();
+		FRHIComputeShader* ShaderRHI = RHICmdList.GetBoundComputeShader();
 
 		FGlobalShader::SetParameters<FViewUniformShaderParameters>(RHICmdList, ShaderRHI, Context.View.ViewUniformBuffer);
 		ScreenSpaceAOParams.Set(RHICmdList, Context.View, ShaderRHI, InputTextureSize);
@@ -2935,7 +2875,7 @@ public:
 	template <typename TRHICmdList>
 	void UnsetParameters(TRHICmdList& RHICmdList)
 	{
-		FRHIComputeShader* ShaderRHI = GetComputeShader();
+		FRHIComputeShader* ShaderRHI = RHICmdList.GetBoundComputeShader();
 		RHICmdList.SetUAVParameter(ShaderRHI, OutTexture.GetBaseIndex(),nullptr);
 	}
 	
@@ -2983,9 +2923,9 @@ void FRCPassPostProcessAmbientOcclusion_GTAO_Upsample::Process(FRenderingComposi
 		uint32 GroupSizeY = FMath::DivideAndRoundUp(View.ViewRect.Height(), 8);
 
 		Context.RHICmdList.TransitionResource(EResourceTransitionAccess::ERWBarrier, EResourceTransitionPipeline::EComputeToCompute, DestRenderTarget.UAV, nullptr);
-		Context.RHICmdList.SetComputeShader(ComputeShader->GetComputeShader());
+		Context.RHICmdList.SetComputeShader(ComputeShader.GetComputeShader());
 		ComputeShader->SetParametersCS(Context.RHICmdList, Context, InputTexSize, DestRenderTarget.UAV);
-		DispatchComputeShader(Context.RHICmdList, *ComputeShader, GroupSizeX, GroupSizeY, 1);
+		DispatchComputeShader(Context.RHICmdList, ComputeShader, GroupSizeX, GroupSizeY, 1);
 		ComputeShader->UnsetParameters(Context.RHICmdList);
 		Context.RHICmdList.TransitionResource(EResourceTransitionAccess::EReadable, DestRenderTarget.TargetableTexture);
 	}
@@ -3012,8 +2952,8 @@ void FRCPassPostProcessAmbientOcclusion_GTAO_Upsample::Process(FRenderingComposi
 			GraphicsPSOInit.bDepthBounds		= false;
 
 			GraphicsPSOInit.BoundShaderState.VertexDeclarationRHI = GFilterVertexDeclaration.VertexDeclarationRHI;
-			GraphicsPSOInit.BoundShaderState.VertexShaderRHI = GETSAFERHISHADER_VERTEX(*VertexShader);
-			GraphicsPSOInit.BoundShaderState.PixelShaderRHI  = GETSAFERHISHADER_PIXEL(*PixelShader);
+			GraphicsPSOInit.BoundShaderState.VertexShaderRHI = VertexShader.GetVertexShader();
+			GraphicsPSOInit.BoundShaderState.PixelShaderRHI = PixelShader.GetPixelShader();
 			SetGraphicsPipelineState(Context.RHICmdList, GraphicsPSOInit);
 
 			PixelShader->SetParametersPS(Context, InputTexSize);
@@ -3024,7 +2964,7 @@ void FRCPassPostProcessAmbientOcclusion_GTAO_Upsample::Process(FRenderingComposi
 				ViewRect.Min.X, ViewRect.Min.Y, ViewRect.Width(), ViewRect.Height(),
 				ViewRect.Size(),
 				OutputTexSize,
-				*VertexShader,
+				VertexShader,
 				EDRF_UseTriangleOptimization);
 		}
 		Context.RHICmdList.EndRenderPass();
diff --git a/Engine/Source/Runtime/Renderer/Private/CompositionLighting/PostProcessAmbientOcclusion.h b/Engine/Source/Runtime/Renderer/Private/CompositionLighting/PostProcessAmbientOcclusion.h
--- a/Engine/Source/Runtime/Renderer/Private/CompositionLighting/PostProcessAmbientOcclusion.h
+++ b/Engine/Source/Runtime/Renderer/Private/CompositionLighting/PostProcessAmbientOcclusion.h
@@ -134,7 +134,7 @@ private:
 	void ProcessPS(FRenderingCompositePassContext& Context, const FSceneRenderTargetItem* DestRenderTarget, const FSceneRenderTargetItem* SceneDepthBuffer, const FIntRect& ViewRect, const FIntPoint& TexSize, int32 ShaderQuality, bool bDoUpsample);
 
 	template <uint32 bAOSetupAsInput, uint32 bDoUpsample, uint32 SampleSetQuality>
-	FShader* SetShaderTemplPS(const FRenderingCompositePassContext& Context, FGraphicsPipelineStateInitializer& GraphicsPSOInit);
+	TShaderRef<FShader> SetShaderTemplPS(const FRenderingCompositePassContext& Context, FGraphicsPipelineStateInitializer& GraphicsPSOInit);
 
 	template <uint32 bAOSetupAsInput, uint32 bDoUpsample, uint32 SampleSetQuality, typename TRHICmdList>
 	void DispatchCS(TRHICmdList& RHICmdList, const FRenderingCompositePassContext& Context, const FIntPoint& TexSize, FRHIUnorderedAccessView* OutTextureUAV);
@@ -161,7 +161,7 @@ public:
 	void DispatchCS(const FRenderingCompositePassContext& Context, FIntRect ViewRect, FIntPoint DestSize, FIntPoint TexSize);
 
 	template <uint32 ShaderQuality>
-	FShader* SetShaderPS(const FRenderingCompositePassContext& Context, FGraphicsPipelineStateInitializer& GraphicsPSOInit, FIntPoint DestSize);
+	TShaderRef<FShader> SetShaderPS(const FRenderingCompositePassContext& Context, FGraphicsPipelineStateInitializer& GraphicsPSOInit, FIntPoint DestSize);
 
 private:
 	const EGTAOType AOType;
@@ -200,7 +200,7 @@ public:
 	void DispatchCS(TRHICmdList& RHICmdList, const FRenderingCompositePassContext& Context, FIntRect ViewRect, FIntPoint DestSize, FIntPoint TexSize);
 
 	template <uint32 ShaderQuality, uint32 UseNormals>
-	FShader* SetShaderPS(const FRenderingCompositePassContext& Context, FGraphicsPipelineStateInitializer& GraphicsPSOInit, FIntPoint DestSize);
+	TShaderRef<FShader> SetShaderPS(const FRenderingCompositePassContext& Context, FGraphicsPipelineStateInitializer& GraphicsPSOInit, FIntPoint DestSize);
 
 private:
 	const EGTAOType AOType;
diff --git a/Engine/Source/Runtime/Renderer/Private/CompositionLighting/PostProcessLpvIndirect.cpp b/Engine/Source/Runtime/Renderer/Private/CompositionLighting/PostProcessLpvIndirect.cpp
--- a/Engine/Source/Runtime/Renderer/Private/CompositionLighting/PostProcessLpvIndirect.cpp
+++ b/Engine/Source/Runtime/Renderer/Private/CompositionLighting/PostProcessLpvIndirect.cpp
@@ -40,8 +40,8 @@ TAutoConsoleVariable<int32> CVarLPVMixing(
 /** Encapsulates the post processing ambient pixel shader. */
 class FPostProcessLpvIndirectPS : public FGlobalShader
 {
-	DECLARE_SHADER_TYPE(FPostProcessLpvIndirectPS, Global);
-
+	DECLARE_INLINE_TYPE_LAYOUT(FPostProcessLpvIndirectPS, NonVirtual);
+public:
 	//@todo-rco: Remove this when reenabling for OpenGL
 	static bool ShouldCompilePermutation(const FGlobalShaderPermutationParameters& Parameters)		{ return IsFeatureLevelSupported(Parameters.Platform, ERHIFeatureLevel::SM5) && IsLPVSupported(Parameters.Platform); }
 
@@ -49,14 +49,6 @@ class FPostProcessLpvIndirectPS : public FGlobalShader
 	FPostProcessLpvIndirectPS() {}
 
 public:
-	FPostProcessPassParameters PostprocessParameter;
-	FShaderResourceParameter LpvBufferSRVParameters[7];
-	FShaderResourceParameter LpvVolumeTextureSampler;
-	FShaderResourceParameter AOVolumeTextureSRVParameter;
-	FSceneTextureShaderParameters SceneTextureParameters;
-	FShaderResourceParameter PreIntegratedGF;
-	FShaderResourceParameter PreIntegratedGFSampler;
-
 	/** Initialization constructor. */
 	FPostProcessLpvIndirectPS(const ShaderMetaType::CompiledShaderInitializerType& Initializer)
 		: FGlobalShader(Initializer)
@@ -83,7 +75,7 @@ public:
 		FLpvReadUniformBufferRef LpvUniformBuffer, 
 		const FRenderingCompositePassContext& Context )
 	{
-		FRHIPixelShader* ShaderRHI = GetPixelShader();
+		FRHIPixelShader* ShaderRHI = RHICmdList.GetBoundPixelShader();
 
 		SetUniformBufferParameter(RHICmdList, ShaderRHI, GetUniformBufferParameter<FLpvReadUniformBufferParameters>(), LpvUniformBuffer);
 
@@ -105,24 +97,14 @@ public:
 		SceneTextureParameters.Set(RHICmdList, ShaderRHI, Context.View.FeatureLevel, ESceneTextureSetupMode::All);
 		SetTextureParameter(RHICmdList, ShaderRHI, PreIntegratedGF, PreIntegratedGFSampler, TStaticSamplerState<SF_Bilinear, AM_Clamp, AM_Clamp, AM_Clamp>::GetRHI(), GSystemTextures.PreintegratedGF->GetRenderTargetItem().ShaderResourceTexture);
 	}
-	
-	// FShader interface.
-	virtual bool Serialize(FArchive& Ar) override
-	{
-		bool bShaderHasOutdatedParameters = FGlobalShader::Serialize(Ar);
-		Ar << PostprocessParameter;
 
-		for ( int i=0; i<7; i++ )
-		{
-			Ar << LpvBufferSRVParameters[i];
-		}
-		Ar << LpvVolumeTextureSampler;
-		Ar << SceneTextureParameters;
-		Ar << PreIntegratedGF;
-		Ar << PreIntegratedGFSampler;
-		Ar << AOVolumeTextureSRVParameter;
-		return bShaderHasOutdatedParameters;
-	}
+	LAYOUT_FIELD(FPostProcessPassParameters, PostprocessParameter);
+	LAYOUT_ARRAY(FShaderResourceParameter, LpvBufferSRVParameters, 7);
+	LAYOUT_FIELD(FShaderResourceParameter, LpvVolumeTextureSampler);
+	LAYOUT_FIELD(FShaderResourceParameter, AOVolumeTextureSRVParameter);
+	LAYOUT_FIELD(FSceneTextureShaderParameters, SceneTextureParameters);
+	LAYOUT_FIELD(FShaderResourceParameter, PreIntegratedGF);
+	LAYOUT_FIELD(FShaderResourceParameter, PreIntegratedGFSampler);
 };
 
 template<bool bApplySeparateSpecularRT>
@@ -160,11 +142,11 @@ class FPostProcessLpvDirectionalOcclusionPS : public FGlobalShader
 	FPostProcessLpvDirectionalOcclusionPS() {}
 
 public:
-	FPostProcessPassParameters PostprocessParameter;
-	FShaderResourceParameter LpvVolumeTextureSampler;
-	FShaderResourceParameter AOVolumeTextureSRVParameter;
+	LAYOUT_FIELD(FPostProcessPassParameters, PostprocessParameter);
+	LAYOUT_FIELD(FShaderResourceParameter, LpvVolumeTextureSampler);
+	LAYOUT_FIELD(FShaderResourceParameter, AOVolumeTextureSRVParameter);
 
-	FSceneTextureShaderParameters SceneTextureParameters;
+	LAYOUT_FIELD(FSceneTextureShaderParameters, SceneTextureParameters);
 
 	/** Initialization constructor. */
 	FPostProcessLpvDirectionalOcclusionPS(const ShaderMetaType::CompiledShaderInitializerType& Initializer)
@@ -181,7 +163,7 @@ public:
 		FLpvReadUniformBufferRef LpvUniformBuffer, 
 		const FRenderingCompositePassContext& Context )
 	{
-		FRHIPixelShader* ShaderRHI = GetPixelShader();
+		FRHIPixelShader* ShaderRHI = Context.RHICmdList.GetBoundPixelShader();
 
 		SetUniformBufferParameter(Context.RHICmdList, ShaderRHI, GetUniformBufferParameter<FLpvReadUniformBufferParameters>(), LpvUniformBuffer);
 
@@ -195,17 +177,6 @@ public:
 		PostprocessParameter.SetPS(Context.RHICmdList, ShaderRHI, Context, TStaticSamplerState<SF_Bilinear, AM_Clamp, AM_Clamp, AM_Clamp>::GetRHI());
 		SceneTextureParameters.Set(Context.RHICmdList, ShaderRHI, Context.View.FeatureLevel, ESceneTextureSetupMode::All);
 	}
-
-	// FShader interface.
-	virtual bool Serialize(FArchive& Ar)
-	{
-		bool bShaderHasOutdatedParameters = FGlobalShader::Serialize(Ar);
-		Ar << LpvVolumeTextureSampler;
-		Ar << SceneTextureParameters;
-		Ar << AOVolumeTextureSRVParameter;
-		Ar << PostprocessParameter;
-		return bShaderHasOutdatedParameters;
-	}
 };
 
 IMPLEMENT_SHADER_TYPE(,FPostProcessLpvDirectionalOcclusionPS,TEXT("/Engine/Private/PostProcessLpvIndirect.usf"),TEXT("DirectionalOcclusionPS"),SF_Pixel);
@@ -296,21 +267,21 @@ void FRCPassPostProcessLpvIndirect::Process(FRenderingCompositePassContext& Cont
 
 		TShaderMapRef<FPostProcessVS> VertexShader(Context.GetShaderMap());
 		GraphicsPSOInit.BoundShaderState.VertexDeclarationRHI = GFilterVertexDeclaration.VertexDeclarationRHI;
-		GraphicsPSOInit.BoundShaderState.VertexShaderRHI = GETSAFERHISHADER_VERTEX(*VertexShader);
+		GraphicsPSOInit.BoundShaderState.VertexShaderRHI = VertexShader.GetVertexShader();
 
-		FPostProcessLpvIndirectPS* PixelShader = NULL;
+		TShaderRef<FPostProcessLpvIndirectPS> PixelShader;
 		if (bApplySeparateSpecularRT)
 		{
 			TShaderMapRef< TPostProcessLpvIndirectPS<true> > PixelShaderWithSpecular(Context.GetShaderMap());
-			PixelShader = (FPostProcessLpvIndirectPS*)*PixelShaderWithSpecular;
+			PixelShader = PixelShaderWithSpecular;
 		}
 		else
 		{
 			TShaderMapRef< TPostProcessLpvIndirectPS<false> > PixelShaderNoSpecular(Context.GetShaderMap());
-			PixelShader = (FPostProcessLpvIndirectPS*)*PixelShaderNoSpecular;
+			PixelShader = PixelShaderNoSpecular;
 		}
 
-		GraphicsPSOInit.BoundShaderState.PixelShaderRHI = GETSAFERHISHADER_PIXEL(PixelShader);
+		GraphicsPSOInit.BoundShaderState.PixelShaderRHI = PixelShader.GetPixelShader();
 		GraphicsPSOInit.PrimitiveType = PT_TriangleList;
 
 		SetGraphicsPipelineState(Context.RHICmdList, GraphicsPSOInit);
@@ -340,7 +311,7 @@ void FRCPassPostProcessLpvIndirect::Process(FRenderingCompositePassContext& Cont
 				View.ViewRect.Width(), View.ViewRect.Height(),
 				View.ViewRect.Size(),
 				SceneContext.GetBufferSizeXY(),
-				*VertexShader,
+				VertexShader,
 				View.StereoPass,
 				Context.HasHmdMesh());
 		}
@@ -403,8 +374,8 @@ void FRCPassPostProcessLpvIndirect::DoDirectionalOcclusionPass(FRenderingComposi
 		TShaderMapRef< FPostProcessLpvDirectionalOcclusionPS > PixelShader(View.ShaderMap);
 
 		GraphicsPSOInit.BoundShaderState.VertexDeclarationRHI = GFilterVertexDeclaration.VertexDeclarationRHI;
-		GraphicsPSOInit.BoundShaderState.VertexShaderRHI = GETSAFERHISHADER_VERTEX(*VertexShader);
-		GraphicsPSOInit.BoundShaderState.PixelShaderRHI = GETSAFERHISHADER_PIXEL(*PixelShader);
+		GraphicsPSOInit.BoundShaderState.VertexShaderRHI = VertexShader.GetVertexShader();
+		GraphicsPSOInit.BoundShaderState.PixelShaderRHI = PixelShader.GetPixelShader();
 		GraphicsPSOInit.PrimitiveType = PT_TriangleList;
 
 		SetGraphicsPipelineState(Context.RHICmdList, GraphicsPSOInit);
@@ -422,7 +393,7 @@ void FRCPassPostProcessLpvIndirect::DoDirectionalOcclusionPass(FRenderingComposi
 			View.ViewRect.Width(), View.ViewRect.Height(),
 			View.ViewRect.Size(),
 			SceneContext.GetBufferSizeXY(),
-			*VertexShader,
+			VertexShader,
 			View.StereoPass,
 			Context.HasHmdMesh());
 	}
diff --git a/Engine/Source/Runtime/Renderer/Private/CompositionLighting/PostProcessMeshDecals.cpp b/Engine/Source/Runtime/Renderer/Private/CompositionLighting/PostProcessMeshDecals.cpp
--- a/Engine/Source/Runtime/Renderer/Private/CompositionLighting/PostProcessMeshDecals.cpp
+++ b/Engine/Source/Runtime/Renderer/Private/CompositionLighting/PostProcessMeshDecals.cpp
@@ -33,7 +33,7 @@ class FMeshDecalAccumulatePolicy
 public:
 	static bool ShouldCompilePermutation(const FMeshMaterialShaderPermutationParameters& Parameters)
 	{
-		return Parameters.Material && Parameters.Material->IsDeferredDecal() && IsFeatureLevelSupported(Parameters.Platform, ERHIFeatureLevel::SM5);
+		return Parameters.MaterialParameters.MaterialDomain == MD_DeferredDecal && IsFeatureLevelSupported(Parameters.Platform, ERHIFeatureLevel::SM5);
 	}
 };
 
@@ -129,7 +129,7 @@ public:
 	static void ModifyCompilationEnvironment(const FMaterialShaderPermutationParameters& Parameters, FShaderCompilerEnvironment& OutEnvironment)
 	{
 		FMeshMaterialShader::ModifyCompilationEnvironment(Parameters, OutEnvironment);
-		FDecalRendering::SetDecalCompilationEnvironment(Parameters.Platform, Parameters.Material, OutEnvironment);
+		FDecalRendering::SetDecalCompilationEnvironment(Parameters, OutEnvironment);
 	}
 
 	FMeshDecalsPS(const ShaderMetaType::CompiledShaderInitializerType& Initializer)
@@ -141,12 +141,6 @@ public:
 	FMeshDecalsPS()
 	{
 	}
-
-	virtual bool Serialize(FArchive& Ar) override
-	{
-		bool bShaderHasOutdatedParameters = FMeshMaterialShader::Serialize(Ar);
-		return bShaderHasOutdatedParameters;
-	}
 };
 
 IMPLEMENT_MATERIAL_SHADER_TYPE(,FMeshDecalsPS,TEXT("/Engine/Private/MeshDecals.usf"),TEXT("MainPS"),SF_Pixel);
@@ -159,13 +153,13 @@ public:
 	static bool ShouldCompilePermutation(const FMeshMaterialShaderPermutationParameters& Parameters)
 	{
 		return FMeshDecalsPS::ShouldCompilePermutation(Parameters)
-			&& Parameters.Material->HasEmissiveColorConnected()
-			&& IsDBufferDecalBlendMode(FDecalRenderingCommon::ComputeFinalDecalBlendMode(Parameters.Platform, Parameters.Material));
+			&& Parameters.MaterialParameters.bHasEmissiveColorConnected
+			&& IsDBufferDecalBlendMode(FDecalRenderingCommon::ComputeFinalDecalBlendMode(Parameters.Platform, (EDecalBlendMode)Parameters.MaterialParameters.DecalBlendMode, Parameters.MaterialParameters.bHasNormalConnected));
 	}
 	static void ModifyCompilationEnvironment(const FMaterialShaderPermutationParameters& Parameters, FShaderCompilerEnvironment& OutEnvironment)
 	{
 		FMeshDecalsPS::ModifyCompilationEnvironment(Parameters, OutEnvironment);
-		FDecalRendering::SetEmissiveDBufferDecalCompilationEnvironment(Parameters.Platform, Parameters.Material, OutEnvironment);
+		FDecalRendering::SetEmissiveDBufferDecalCompilationEnvironment(Parameters, OutEnvironment);
 	}
 
 	FMeshDecalsEmissivePS() {}
@@ -235,8 +229,10 @@ void FMeshDecalMeshProcessor::AddMeshBatch(const FMeshBatch& RESTRICT MeshBatch,
 				const EShaderPlatform ShaderPlatform = ViewIfDynamicMeshCommand->GetShaderPlatform();
 				const EDecalBlendMode FinalDecalBlendMode = FDecalRenderingCommon::ComputeFinalDecalBlendMode(ShaderPlatform, Material);
 				const EDecalRenderStage LocalDecalRenderStage = FDecalRenderingCommon::ComputeRenderStage(ShaderPlatform, FinalDecalBlendMode);
-				const ERasterizerFillMode MeshFillMode = ComputeMeshFillMode(MeshBatch, *Material);
-				const ERasterizerCullMode MeshCullMode = ComputeMeshCullMode(MeshBatch, *Material);
+
+				const FMeshDrawingPolicyOverrideSettings OverrideSettings = ComputeMeshOverrideSettings(MeshBatch);
+				ERasterizerFillMode MeshFillMode = ComputeMeshFillMode(MeshBatch, *Material, OverrideSettings);
+				ERasterizerCullMode MeshCullMode = ComputeMeshCullMode(MeshBatch, *Material, OverrideSettings);
 
 				bool bShouldRender = FDecalRenderingCommon::IsCompatibleWithRenderStage(
 					PassDecalStage,
@@ -325,11 +321,14 @@ void FMeshDecalMeshProcessor::Process(
 	}
 
 	MeshDecalPassShaders.VertexShader = MaterialResource.GetShader<FMeshDecalsVS>(VertexFactoryType);
-
-	MeshDecalPassShaders.PixelShader = PassDecalStage == DRS_Emissive
-		? MaterialResource.GetShader<FMeshDecalsEmissivePS>(VertexFactoryType)
-		: MaterialResource.GetShader<FMeshDecalsPS>(VertexFactoryType);
-
+	if (PassDecalStage == DRS_Emissive)
+	{
+		MeshDecalPassShaders.PixelShader = MaterialResource.GetShader<FMeshDecalsEmissivePS>(VertexFactoryType);
+	}
+	else
+	{
+		MeshDecalPassShaders.PixelShader = MaterialResource.GetShader<FMeshDecalsPS>(VertexFactoryType);
+	}
 
 	FMeshMaterialShaderElementData ShaderElementData;
 	ShaderElementData.InitializeMeshMaterialData(ViewIfDynamicMeshCommand, PrimitiveSceneProxy, MeshBatch, StaticMeshId, true);
diff --git a/Engine/Source/Runtime/Renderer/Private/CompositionLighting/PostProcessPassThrough.cpp b/Engine/Source/Runtime/Renderer/Private/CompositionLighting/PostProcessPassThrough.cpp
--- a/Engine/Source/Runtime/Renderer/Private/CompositionLighting/PostProcessPassThrough.cpp
+++ b/Engine/Source/Runtime/Renderer/Private/CompositionLighting/PostProcessPassThrough.cpp
@@ -18,17 +18,17 @@ FPostProcessPassThroughPS::FPostProcessPassThroughPS(const ShaderMetaType::Compi
 	PostprocessParameter.Bind(Initializer.ParameterMap);
 }
 
-bool FPostProcessPassThroughPS::Serialize(FArchive& Ar)
+/*bool FPostProcessPassThroughPS::Serialize(FArchive& Ar)
 {
 	bool bShaderHasOutdatedParameters = FGlobalShader::Serialize(Ar);
 	Ar << PostprocessParameter;
 	return bShaderHasOutdatedParameters;
-}
+}*/
 
 template <typename TRHICmdList>
 void FPostProcessPassThroughPS::SetParameters(TRHICmdList& RHICmdList, const FRenderingCompositePassContext& Context)
 {
-	FRHIPixelShader* ShaderRHI = GetPixelShader();
+	FRHIPixelShader* ShaderRHI = RHICmdList.GetBoundPixelShader();
 
 	FGlobalShader::SetParameters<FViewUniformShaderParameters>(RHICmdList, ShaderRHI, Context.View.ViewUniformBuffer);
 
@@ -111,8 +111,8 @@ void FRCPassPostProcessPassThrough::Process(FRenderingCompositePassContext& Cont
 		TShaderMapRef<FPostProcessPassThroughPS> PixelShader(Context.GetShaderMap());
 
 		GraphicsPSOInit.BoundShaderState.VertexDeclarationRHI = GFilterVertexDeclaration.VertexDeclarationRHI;
-		GraphicsPSOInit.BoundShaderState.VertexShaderRHI = GETSAFERHISHADER_VERTEX(*VertexShader);
-		GraphicsPSOInit.BoundShaderState.PixelShaderRHI = GETSAFERHISHADER_PIXEL(*PixelShader);
+		GraphicsPSOInit.BoundShaderState.VertexShaderRHI = VertexShader.GetVertexShader();
+		GraphicsPSOInit.BoundShaderState.PixelShaderRHI = PixelShader.GetPixelShader();
 		GraphicsPSOInit.PrimitiveType = PT_TriangleList;
 
 		SetGraphicsPipelineState(Context.RHICmdList, GraphicsPSOInit);
@@ -128,7 +128,7 @@ void FRCPassPostProcessPassThrough::Process(FRenderingCompositePassContext& Cont
 			SrcRect.Width(), SrcRect.Height(),
 			DestSize,
 			SrcSize,
-			*VertexShader,
+			VertexShader,
 			View.StereoPass,
 			Context.HasHmdMesh(),
 			EDRF_UseTriangleOptimization);
diff --git a/Engine/Source/Runtime/Renderer/Private/CompositionLighting/PostProcessPassThrough.h b/Engine/Source/Runtime/Renderer/Private/CompositionLighting/PostProcessPassThrough.h
--- a/Engine/Source/Runtime/Renderer/Private/CompositionLighting/PostProcessPassThrough.h
+++ b/Engine/Source/Runtime/Renderer/Private/CompositionLighting/PostProcessPassThrough.h
@@ -27,13 +27,13 @@ class FPostProcessPassThroughPS : public FGlobalShader
 	FPostProcessPassThroughPS() {}
 
 public:
-	FPostProcessPassParameters PostprocessParameter;
+	LAYOUT_FIELD(FPostProcessPassParameters, PostprocessParameter)
 
 	/** Initialization constructor. */
 	FPostProcessPassThroughPS(const ShaderMetaType::CompiledShaderInitializerType& Initializer);
 
 	// FShader interface.
-	virtual bool Serialize(FArchive& Ar) override;
+	//virtual bool Serialize(FArchive& Ar) override;
 
 	template <typename TRHICmdList>
 	void SetParameters(TRHICmdList& RHICmdList, const FRenderingCompositePassContext& Context);
diff --git a/Engine/Source/Runtime/Renderer/Private/ConvertToUniformMesh.cpp b/Engine/Source/Runtime/Renderer/Private/ConvertToUniformMesh.cpp
--- a/Engine/Source/Runtime/Renderer/Private/ConvertToUniformMesh.cpp
+++ b/Engine/Source/Runtime/Renderer/Private/ConvertToUniformMesh.cpp
@@ -112,13 +112,6 @@ protected:
 			&& (FCString::Strstr(Parameters.VertexFactoryType->GetName(), TEXT("LocalVertexFactory")) != NULL
 				|| FCString::Strstr(Parameters.VertexFactoryType->GetName(), TEXT("InstancedStaticMeshVertexFactory")) != NULL);
 	}
-
-	static void GetStreamOutElements(FStreamOutElementList& ElementList, TArray<uint32>& StreamStrides, int32& RasterizedStream)
-	{
-		StreamStrides.Add(ComputeUniformVertexStride() * 4);
-		GetUniformMeshStreamOutLayout(ElementList);
-		RasterizedStream = -1;
-	}
 };
 
 IMPLEMENT_MATERIAL_SHADER_TYPE(,FConvertToUniformMeshGS,TEXT("/Engine/Private/ConvertToUniformMesh.usf"),TEXT("ConvertToUniformMeshGS"),SF_Geometry);
@@ -161,8 +154,9 @@ void FConvertToUniformMeshProcessor::AddMeshBatch(const FMeshBatch& RESTRICT Mes
 
 	const FMaterialRenderProxy& MaterialRenderProxy = FallbackMaterialRenderProxyPtr ? *FallbackMaterialRenderProxyPtr : *MeshBatch.MaterialRenderProxy;
 
-	const ERasterizerFillMode MeshFillMode = ComputeMeshFillMode(MeshBatch, Material);
-	const ERasterizerCullMode MeshCullMode = ComputeMeshCullMode(MeshBatch, Material);
+	const FMeshDrawingPolicyOverrideSettings OverrideSettings = ComputeMeshOverrideSettings(MeshBatch);
+	const ERasterizerFillMode MeshFillMode = ComputeMeshFillMode(MeshBatch, Material, OverrideSettings);
+	const ERasterizerCullMode MeshCullMode = ComputeMeshCullMode(MeshBatch, Material, OverrideSettings);
 
 	Process(MeshBatch, BatchElementMask, PrimitiveSceneProxy, MaterialRenderProxy, Material, MeshFillMode, MeshCullMode);
 }
@@ -311,12 +305,12 @@ public:
 
 	static bool ShouldCompilePermutation(const FMaterialShaderPermutationParameters& Parameters)
 	{
-		if (Parameters.Material->IsUIMaterial())
+		if (Parameters.MaterialParameters.MaterialDomain == MD_UI)
 		{
 			return false;
 		}
 
-		if (Parameters.Material->GetShadingModels().IsUnlit())
+		if (Parameters.MaterialParameters.ShadingModels.IsUnlit())
 		{
 			return false;
 		}
@@ -326,7 +320,7 @@ public:
 
 	static void ModifyCompilationEnvironment(const FMaterialShaderPermutationParameters& Parameters, FShaderCompilerEnvironment& OutEnvironment)
 	{
-		FMaterialShader::ModifyCompilationEnvironment(Parameters.Platform,OutEnvironment);
+		FMaterialShader::ModifyCompilationEnvironment(Parameters, OutEnvironment);
 		OutEnvironment.SetDefine(TEXT("EVALUATE_SURFEL_MATERIAL_GROUP_SIZE"), GEvaluateSurfelMaterialGroupSize);
 		OutEnvironment.SetDefine(TEXT("HAS_PRIMITIVE_UNIFORM_BUFFER"), 1);
 	}
@@ -354,7 +348,7 @@ public:
 		const FMatrix& Instance0Transform
 		)
 	{
-		FRHIComputeShader* ShaderRHI = GetComputeShader();
+		FRHIComputeShader* ShaderRHI = RHICmdList.GetBoundComputeShader();
 		FMaterialShader::SetParameters(RHICmdList, ShaderRHI, MaterialProxy, *MaterialProxy->GetMaterial(View.GetFeatureLevel()), View, View.ViewUniformBuffer, ESceneTextureSetupMode::None);
 
 		SetUniformBufferParameter(RHICmdList, ShaderRHI,GetUniformBufferParameter<FPrimitiveUniformShaderParameters>(),PrimitiveUniformBuffer);
@@ -375,7 +369,7 @@ public:
 
 	void UnsetParameters(FRHICommandList& RHICmdList, FViewInfo& View)
 	{
-		FRHIComputeShader* ShaderRHI = GetComputeShader();
+		FRHIComputeShader* ShaderRHI = RHICmdList.GetBoundComputeShader();
 		SurfelBufferParameters.UnsetParameters(RHICmdList, ShaderRHI);
 
 		const FScene* Scene = (const FScene*)View.Family->Scene;
@@ -384,22 +378,12 @@ public:
 		RHICmdList.TransitionResources(EResourceTransitionAccess::EReadable, EResourceTransitionPipeline::EComputeToCompute, UniformMeshUAVs, UE_ARRAY_COUNT(UniformMeshUAVs));
 	}
 
-	virtual bool Serialize(FArchive& Ar) override
-	{		
-		bool bShaderHasOutdatedParameters = FMaterialShader::Serialize(Ar);
-		Ar << SurfelBufferParameters;
-		Ar << SurfelStartIndex;
-		Ar << NumSurfelsToGenerate;
-		Ar << Instance0InverseTransform;
-		return bShaderHasOutdatedParameters;
-	}
-
 private:
 
-	FSurfelBufferParameters SurfelBufferParameters;
-	FShaderParameter SurfelStartIndex;
-	FShaderParameter NumSurfelsToGenerate;
-	FShaderParameter Instance0InverseTransform;
+	LAYOUT_FIELD(FSurfelBufferParameters, SurfelBufferParameters);
+	LAYOUT_FIELD(FShaderParameter, SurfelStartIndex);
+	LAYOUT_FIELD(FShaderParameter, NumSurfelsToGenerate);
+	LAYOUT_FIELD(FShaderParameter, Instance0InverseTransform);
 };
 
 IMPLEMENT_MATERIAL_SHADER_TYPE(,FEvaluateSurfelMaterialCS,TEXT("/Engine/Private/EvaluateSurfelMaterial.usf"),TEXT("EvaluateSurfelMaterialCS"),SF_Compute);
@@ -416,11 +400,11 @@ void FUniformMeshConverter::GenerateSurfels(
 {
 	const FMaterial* Material = MaterialProxy->GetMaterial(View.GetFeatureLevel());
 	const FMaterialShaderMap* MaterialShaderMap = Material->GetRenderingThreadShaderMap();
-	FEvaluateSurfelMaterialCS* ComputeShader = MaterialShaderMap->GetShader<FEvaluateSurfelMaterialCS>();
+	TShaderRef<FEvaluateSurfelMaterialCS> ComputeShader = MaterialShaderMap->GetShader<FEvaluateSurfelMaterialCS>();
 
-	RHICmdList.SetComputeShader(ComputeShader->GetComputeShader());
+	RHICmdList.SetComputeShader(ComputeShader.GetComputeShader());
 	ComputeShader->SetParameters(RHICmdList, View, SurfelOffset, NumSurfels, MaterialProxy, PrimitiveUniformBuffer, Instance0Transform);
-	DispatchComputeShader(RHICmdList, ComputeShader, FMath::DivideAndRoundUp(NumSurfels, GEvaluateSurfelMaterialGroupSize), 1, 1);
+	DispatchComputeShader(RHICmdList, ComputeShader.GetShader(), FMath::DivideAndRoundUp(NumSurfels, GEvaluateSurfelMaterialGroupSize), 1, 1);
 
 	ComputeShader->UnsetParameters(RHICmdList, View);
 }
diff --git a/Engine/Source/Runtime/Renderer/Private/CustomDepthRendering.cpp b/Engine/Source/Runtime/Renderer/Private/CustomDepthRendering.cpp
--- a/Engine/Source/Runtime/Renderer/Private/CustomDepthRendering.cpp
+++ b/Engine/Source/Runtime/Renderer/Private/CustomDepthRendering.cpp
@@ -64,8 +64,9 @@ void FCustomDepthPassMeshProcessor::AddMeshBatch(const FMeshBatch& RESTRICT Mesh
 		const FMaterialRenderProxy& MaterialRenderProxy = FallbackMaterialRenderProxyPtr ? *FallbackMaterialRenderProxyPtr : *MeshBatch.MaterialRenderProxy;
 
 		const EBlendMode BlendMode = Material.GetBlendMode();
-		const ERasterizerFillMode MeshFillMode = ComputeMeshFillMode(MeshBatch, Material);
-		const ERasterizerCullMode MeshCullMode = ComputeMeshCullMode(MeshBatch, Material);
+		const FMeshDrawingPolicyOverrideSettings OverrideSettings = ComputeMeshOverrideSettings(MeshBatch);
+		const ERasterizerFillMode MeshFillMode = ComputeMeshFillMode(MeshBatch, Material, OverrideSettings);
+		const ERasterizerCullMode MeshCullMode = ComputeMeshCullMode(MeshBatch, Material, OverrideSettings);
 		const bool bIsTranslucent = IsTranslucentBlendMode(BlendMode);
 
 
@@ -157,8 +158,7 @@ void FCustomDepthPassMeshProcessor::Process(
 		FDepthOnlyDS,
 		FDepthOnlyPS> DepthPassShaders;
 
-	FShaderPipeline* ShaderPipeline = nullptr;
-
+	FShaderPipelineRef ShaderPipeline;
 	GetDepthPassShaders<bPositionOnly>(
 		MaterialResource,
 		VertexFactory->GetType(),
diff --git a/Engine/Source/Runtime/Renderer/Private/DebugViewModeRendering.cpp b/Engine/Source/Runtime/Renderer/Private/DebugViewModeRendering.cpp
--- a/Engine/Source/Runtime/Renderer/Private/DebugViewModeRendering.cpp
+++ b/Engine/Source/Runtime/Renderer/Private/DebugViewModeRendering.cpp
@@ -59,14 +59,16 @@ bool FDebugViewModeVS::ShouldCompilePermutation(const FMeshMaterialShaderPermuta
 	if (AllowDebugViewVSDSHS(Parameters.Platform))
 	{
 		// If it comes from FDebugViewModeMaterialProxy, compile it.
-		if (Parameters.Material->GetFriendlyName().Contains(TEXT("DebugViewMode")))
+		if (Parameters.MaterialParameters.bIsMaterialDebugViewMode)
 		{
 			return true;
 		}
 		// Otherwise we only cache it if this for the shader complexity.
 		else if (GCacheShaderComplexityShaders)
 		{
-			return !FDebugViewModeInterface::AllowFallbackToDefaultMaterial(Parameters.Material) || Parameters.Material->IsDefaultMaterial();
+			return !FDebugViewModeInterface::AllowFallbackToDefaultMaterial(Parameters.MaterialParameters.TessellationMode,
+				Parameters.MaterialParameters.bHasVertexPositionOffsetConnected,
+				Parameters.MaterialParameters.bHasPixelDepthOffsetConnected) || Parameters.MaterialParameters.bIsDefaultMaterial;
 		}
 	}
 	return false;
@@ -104,11 +106,53 @@ bool FDeferredShadingSceneRenderer::RenderDebugViewMode(FRHICommandListImmediate
 	return bDirty;
 }
 
-FDebugViewModePS::FDebugViewModePS(const FMeshMaterialShaderType::CompiledShaderInitializerType& Initializer) : FMeshMaterialShader(Initializer) 
+FDebugViewModePS::FDebugViewModePS(const FMeshMaterialShaderType::CompiledShaderInitializerType& Initializer) : FMeshMaterialShader(Initializer)
 {
 	PassUniformBuffer.Bind(Initializer.ParameterMap, FDebugViewModePassPassUniformParameters::StaticStructMetadata.GetShaderVariableName());
 }
 
+void FDebugViewModePS::GetElementShaderBindings(
+	const FShaderMapPointerTable& PointerTable,
+	const FScene* Scene,
+	const FSceneView* ViewIfDynamicMeshCommand,
+	const FVertexFactory* VertexFactory,
+	const EVertexInputStreamType InputStreamType,
+	ERHIFeatureLevel::Type FeatureLevel,
+	const FPrimitiveSceneProxy* PrimitiveSceneProxy,
+	const FMeshBatch& MeshBatch,
+	const FMeshBatchElement& BatchElement,
+	const FDebugViewModeShaderElementData& ShaderElementData,
+	FMeshDrawSingleShaderBindings& ShaderBindings,
+	FVertexInputStreamArray& VertexStreams) const
+{
+	FMeshMaterialShader::GetElementShaderBindings(PointerTable, Scene, ViewIfDynamicMeshCommand, VertexFactory, InputStreamType, FeatureLevel, PrimitiveSceneProxy, MeshBatch, BatchElement, ShaderElementData, ShaderBindings, VertexStreams);
+
+	int8 VisualizeElementIndex = 0;
+#if !(UE_BUILD_SHIPPING || UE_BUILD_TEST)
+	VisualizeElementIndex = BatchElement.VisualizeElementIndex;
+#endif
+
+	const FDebugViewModeInterface* Interface = FDebugViewModeInterface::GetInterface(ShaderElementData.DebugViewMode);
+	if (ensure(Interface))
+	{
+		Interface->GetDebugViewModeShaderBindings(
+			*this,
+			PrimitiveSceneProxy,
+			ShaderElementData.MaterialRenderProxy,
+			ShaderElementData.Material,
+			ShaderElementData.DebugViewMode,
+			ShaderElementData.ViewOrigin,
+			ShaderElementData.VisualizeLODIndex,
+			VisualizeElementIndex,
+			ShaderElementData.NumVSInstructions,
+			ShaderElementData.NumPSInstructions,
+			ShaderElementData.ViewModeParam,
+			ShaderElementData.ViewModeParamName,
+			ShaderBindings
+		);
+	}
+}
+
 FDebugViewModeMeshProcessor::FDebugViewModeMeshProcessor(
 	const FScene* InScene, 
 	ERHIFeatureLevel::Type InFeatureLevel,
@@ -195,8 +239,9 @@ void FDebugViewModeMeshProcessor::AddMeshBatch(const FMeshBatch& RESTRICT MeshBa
 	}
 	DebugViewModePassShaders.PixelShader = DebugViewModeInterface->GetPixelShader(Material, VertexFactoryType);
 
-	const ERasterizerFillMode MeshFillMode = ComputeMeshFillMode(MeshBatch, *BatchMaterial);
-	const ERasterizerCullMode MeshCullMode = ComputeMeshCullMode(MeshBatch, *BatchMaterial);
+	const FMeshDrawingPolicyOverrideSettings OverrideSettings = ComputeMeshOverrideSettings(MeshBatch);
+	const ERasterizerFillMode MeshFillMode = ComputeMeshFillMode(MeshBatch, *BatchMaterial, OverrideSettings);
+	const ERasterizerCullMode MeshCullMode = ComputeMeshCullMode(MeshBatch, *BatchMaterial, OverrideSettings);
 
 	FMeshPassProcessorRenderState DrawRenderState;
 	DrawRenderState.SetViewUniformBuffer(ViewUniformBuffer);
@@ -276,13 +321,13 @@ void FDebugViewModeMeshProcessor::UpdateInstructionCount(FDebugViewModeShaderEle
 		}
 		else // EShadingPath::Mobile
 		{
-			TMobileBasePassVSPolicyParamType<FUniformLightMapPolicy>* MobileVS = nullptr;
-			TMobileBasePassPSPolicyParamType<FUniformLightMapPolicy>* MobilePS = nullptr;
+			TShaderRef<TMobileBasePassVSPolicyParamType<FUniformLightMapPolicy>> MobileVS;
+			TShaderRef<TMobileBasePassPSPolicyParamType<FUniformLightMapPolicy>> MobilePS;
 
 			MobileBasePass::GetShaders(LMP_NO_LIGHTMAP, 0, *InBatchMaterial, InVertexFactoryType, false, MobileVS, MobilePS);
 
-			OutShaderElementData.NumVSInstructions = MobileVS ? MobileVS->GetNumInstructions() : 0;
-			OutShaderElementData.NumPSInstructions = MobilePS ? MobilePS->GetNumInstructions() : 0;
+			OutShaderElementData.NumVSInstructions = MobileVS.IsValid() ? MobileVS->GetNumInstructions() : 0;
+			OutShaderElementData.NumPSInstructions = MobilePS.IsValid() ? MobilePS->GetNumInstructions() : 0;
 		}
 	}
 }
diff --git a/Engine/Source/Runtime/Renderer/Private/DebugViewModeRendering.h b/Engine/Source/Runtime/Renderer/Private/DebugViewModeRendering.h
--- a/Engine/Source/Runtime/Renderer/Private/DebugViewModeRendering.h
+++ b/Engine/Source/Runtime/Renderer/Private/DebugViewModeRendering.h
@@ -104,7 +104,7 @@ public:
 	static void SetCommonDefinitions(const FMaterialShaderPermutationParameters& Parameters, FShaderCompilerEnvironment& OutEnvironment)
 	{
 		// SM4 has less input interpolants. Also instanced meshes use more interpolants.
-		if (Parameters.Material->IsDefaultMaterial() || (IsFeatureLevelSupported(Parameters.Platform, ERHIFeatureLevel::SM5) && !Parameters.Material->IsUsedWithInstancedStaticMeshes()))
+		if (Parameters.MaterialParameters.bIsDefaultMaterial || (IsFeatureLevelSupported(Parameters.Platform, ERHIFeatureLevel::SM5) && !Parameters.MaterialParameters.bIsUsedWithInstancedStaticMeshes))
 		{	// Force the default material to pass enough texcoords to the pixel shaders (even though not using them).
 			// This is required to allow material shaders to have access to the sampled coords.
 			OutEnvironment.SetDefine(TEXT("MIN_MATERIAL_TEXCOORDS"), (uint32)4);
@@ -173,61 +173,22 @@ public:
 class FDebugViewModePS : public FMeshMaterialShader
 {
 public:
-
 	FDebugViewModePS(const FMeshMaterialShaderType::CompiledShaderInitializerType& Initializer);
-
 	FDebugViewModePS() {}
 
 	void GetElementShaderBindings(
-		const FScene* Scene, 
-		const FSceneView* ViewIfDynamicMeshCommand, 
+		const FShaderMapPointerTable& PointerTable,
+		const FScene* Scene,
+		const FSceneView* ViewIfDynamicMeshCommand,
 		const FVertexFactory* VertexFactory,
 		const EVertexInputStreamType InputStreamType,
 		ERHIFeatureLevel::Type FeatureLevel,
 		const FPrimitiveSceneProxy* PrimitiveSceneProxy,
 		const FMeshBatch& MeshBatch,
-		const FMeshBatchElement& BatchElement, 
+		const FMeshBatchElement& BatchElement,
 		const FDebugViewModeShaderElementData& ShaderElementData,
 		FMeshDrawSingleShaderBindings& ShaderBindings,
-		FVertexInputStreamArray& VertexStreams) const
-	{
-		FMeshMaterialShader::GetElementShaderBindings(Scene, ViewIfDynamicMeshCommand, VertexFactory, InputStreamType, FeatureLevel, PrimitiveSceneProxy, MeshBatch, BatchElement, ShaderElementData, ShaderBindings, VertexStreams);
-	
-		int8 VisualizeElementIndex = 0;
-#if !(UE_BUILD_SHIPPING || UE_BUILD_TEST)
-		VisualizeElementIndex = BatchElement.VisualizeElementIndex;
-#endif
-
-		GetDebugViewModeShaderBindings(
-			PrimitiveSceneProxy,
-			ShaderElementData.MaterialRenderProxy,
-			ShaderElementData.Material,
-			ShaderElementData.DebugViewMode,
-			ShaderElementData.ViewOrigin,
-			ShaderElementData.VisualizeLODIndex,
-			VisualizeElementIndex,
-			ShaderElementData.NumVSInstructions,
-			ShaderElementData.NumPSInstructions,
-			ShaderElementData.ViewModeParam,
-			ShaderElementData.ViewModeParamName,
-			ShaderBindings
-		);
-	}
-
-	virtual void GetDebugViewModeShaderBindings(
-		const FPrimitiveSceneProxy* RESTRICT PrimitiveSceneProxy,
-		const FMaterialRenderProxy& RESTRICT MaterialRenderProxy,
-		const FMaterial& RESTRICT Material,
-		EDebugViewShaderMode DebugViewMode,
-		const FVector& ViewOrigin,
-		int32 VisualizeLODIndex,
-		int32 VisualizeElementIndex,
-		int32 NumVSInstructions,
-		int32 NumPSInstructions,
-		int32 ViewModeParam,
-		FName ViewModeParamName,
-		FMeshDrawSingleShaderBindings& ShaderBindings
-	) const = 0;
+		FVertexInputStreamArray& VertexStreams) const;
 };
 
 class FDebugViewModeMeshProcessor : public FMeshPassProcessor
diff --git a/Engine/Source/Runtime/Renderer/Private/DecalRenderingShared.cpp b/Engine/Source/Runtime/Renderer/Private/DecalRenderingShared.cpp
--- a/Engine/Source/Runtime/Renderer/Private/DecalRenderingShared.cpp
+++ b/Engine/Source/Runtime/Renderer/Private/DecalRenderingShared.cpp
@@ -74,21 +74,14 @@ public:
 
 	void SetParameters(FRHICommandList& RHICmdList, FRHIUniformBuffer* ViewUniformBuffer, const FMatrix& InFrustumComponentToClip)
 	{
-		FRHIVertexShader* ShaderRHI = GetVertexShader();
+		FRHIVertexShader* ShaderRHI = RHICmdList.GetBoundVertexShader();
 
 		FGlobalShader::SetParameters<FViewUniformShaderParameters>(RHICmdList, ShaderRHI, ViewUniformBuffer);
 		SetShaderValue(RHICmdList, ShaderRHI, FrustumComponentToClip, InFrustumComponentToClip);
 	}
 
-	virtual bool Serialize(FArchive& Ar) override
-	{
-		bool bShaderHasOutdatedParameters = FGlobalShader::Serialize(Ar);
-		Ar << FrustumComponentToClip;
-		return bShaderHasOutdatedParameters;
-	}
-
 private:
-	FShaderParameter FrustumComponentToClip;
+	LAYOUT_FIELD(FShaderParameter, FrustumComponentToClip);
 };
 
 IMPLEMENT_SHADER_TYPE(,FDeferredDecalVS,TEXT("/Engine/Private/DeferredDecal.usf"),TEXT("MainVS"),SF_Vertex);
@@ -108,13 +101,13 @@ public:
 	  */
 	static bool ShouldCompilePermutation(const FMaterialShaderPermutationParameters& Parameters)
 	{
-		return Parameters.Material->IsDeferredDecal();
+		return Parameters.MaterialParameters.MaterialDomain == MD_DeferredDecal;
 	}
 
 	static void ModifyCompilationEnvironment(const FMaterialShaderPermutationParameters& Parameters, FShaderCompilerEnvironment& OutEnvironment )
 	{
-		FMaterialShader::ModifyCompilationEnvironment(Parameters.Platform, OutEnvironment);
-		FDecalRendering::SetDecalCompilationEnvironment(Parameters.Platform, Parameters.Material, OutEnvironment);
+		FMaterialShader::ModifyCompilationEnvironment(Parameters, OutEnvironment);
+		FDecalRendering::SetDecalCompilationEnvironment(Parameters, OutEnvironment);
 	}
 
 	FDeferredDecalPS() {}
@@ -130,7 +123,7 @@ public:
 
 	void SetParameters(FRHICommandList& RHICmdList, const FViewInfo& View, const FMaterialRenderProxy* MaterialProxy, const FDeferredDecalProxy& DecalProxy, const float FadeAlphaValue=1.0f)
 	{
-		FRHIPixelShader* ShaderRHI = GetPixelShader();
+		FRHIPixelShader* ShaderRHI = RHICmdList.GetBoundPixelShader();
 
 		FMaterialShader::SetParameters(RHICmdList, ShaderRHI, MaterialProxy, *MaterialProxy->GetMaterial(View.GetFeatureLevel()), View, View.ViewUniformBuffer, ESceneTextureSetupMode::All);
 
@@ -193,19 +186,12 @@ public:
 		SetShaderValue(RHICmdList, ShaderRHI, DecalParams, FVector2D(FadeAlphaValue, LifetimeAlpha));
 	}
 
-	virtual bool Serialize(FArchive& Ar) override
-	{
-		bool bShaderHasOutdatedParameters = FMaterialShader::Serialize(Ar);
-		Ar << SvPositionToDecal << DecalToWorld << WorldToDecal << DecalOrientation << DecalParams;
-		return bShaderHasOutdatedParameters;
-	}
-
 private:
-	FShaderParameter SvPositionToDecal;
-	FShaderParameter DecalToWorld;
-	FShaderParameter WorldToDecal;
-	FShaderParameter DecalOrientation;
-	FShaderParameter DecalParams;
+	LAYOUT_FIELD(FShaderParameter, SvPositionToDecal);
+	LAYOUT_FIELD(FShaderParameter, DecalToWorld);
+	LAYOUT_FIELD(FShaderParameter, WorldToDecal);
+	LAYOUT_FIELD(FShaderParameter, DecalOrientation);
+	LAYOUT_FIELD(FShaderParameter, DecalParams);
 };
 
 IMPLEMENT_MATERIAL_SHADER_TYPE(,FDeferredDecalPS,TEXT("/Engine/Private/DeferredDecal.usf"),TEXT("MainPS"),SF_Pixel);
@@ -217,13 +203,13 @@ public:
 	static bool ShouldCompilePermutation(const FMaterialShaderPermutationParameters& Parameters)
 	{
 		return FDeferredDecalPS::ShouldCompilePermutation(Parameters)
-			&& Parameters.Material->HasEmissiveColorConnected()
-			&& IsDBufferDecalBlendMode(FDecalRenderingCommon::ComputeFinalDecalBlendMode(Parameters.Platform, Parameters.Material));
+			&& Parameters.MaterialParameters.bHasEmissiveColorConnected
+			&& IsDBufferDecalBlendMode(FDecalRenderingCommon::ComputeFinalDecalBlendMode(Parameters.Platform, (EDecalBlendMode)Parameters.MaterialParameters.DecalBlendMode, Parameters.MaterialParameters.bHasNormalConnected));
 	}
 	static void ModifyCompilationEnvironment(const FMaterialShaderPermutationParameters& Parameters, FShaderCompilerEnvironment& OutEnvironment)
 	{
 		FDeferredDecalPS::ModifyCompilationEnvironment(Parameters, OutEnvironment);
-		FDecalRendering::SetEmissiveDBufferDecalCompilationEnvironment(Parameters.Platform, Parameters.Material, OutEnvironment);
+		FDecalRendering::SetEmissiveDBufferDecalCompilationEnvironment(Parameters, OutEnvironment);
 	}
 
 	FDeferredDecalEmissivePS() {}
@@ -382,16 +368,16 @@ void FDecalRendering::SetShader(FRHICommandList& RHICmdList, FGraphicsPipelineSt
 
 	// When in shader complexity, decals get rendered as emissive even though there might not be emissive decals.
 	// FDeferredDecalEmissivePS might not be available depending on the decal blend mode.
-	FDeferredDecalPS* PixelShader = (DecalRenderStage == DRS_Emissive && DebugViewMode == DVSM_None)
-		? MaterialShaderMap->GetShader<FDeferredDecalEmissivePS>()
+	TShaderRef<FDeferredDecalPS> PixelShader = (DecalRenderStage == DRS_Emissive && DebugViewMode == DVSM_None)
+		? TShaderRef<FDeferredDecalPS>(MaterialShaderMap->GetShader<FDeferredDecalEmissivePS>())
 		: MaterialShaderMap->GetShader<FDeferredDecalPS>();
 
 	TShaderMapRef<FDeferredDecalVS> VertexShader(View.ShaderMap);
 
 	{
 		GraphicsPSOInit.BoundShaderState.VertexDeclarationRHI = GetVertexDeclarationFVector4();
-		GraphicsPSOInit.BoundShaderState.VertexShaderRHI = GETSAFERHISHADER_VERTEX(*VertexShader);
-		GraphicsPSOInit.BoundShaderState.PixelShaderRHI = PixelShader->GetPixelShader();
+		GraphicsPSOInit.BoundShaderState.VertexShaderRHI = VertexShader.GetVertexShader();
+		GraphicsPSOInit.BoundShaderState.PixelShaderRHI = PixelShader.GetPixelShader();
 		GraphicsPSOInit.PrimitiveType = PT_TriangleList;
 
 		SetGraphicsPipelineState(RHICmdList, GraphicsPSOInit);
@@ -410,11 +396,11 @@ void FDecalRendering::SetShader(FRHICommandList& RHICmdList, FGraphicsPipelineSt
 		//	check(!PrimitivePS.IsBound());
 
 		// to prevent potential shader error (UE-18852 ElementalDemo crashes due to nil constant buffer)
-		SetUniformBufferParameter(RHICmdList, VertexShader->GetVertexShader(), PrimitiveVS, GIdentityPrimitiveUniformBuffer);
+		SetUniformBufferParameter(RHICmdList, VertexShader.GetVertexShader(), PrimitiveVS, GIdentityPrimitiveUniformBuffer);
 
 		if (DebugViewMode == DVSM_None)
 		{
-			SetUniformBufferParameter(RHICmdList, PixelShader->GetPixelShader(), PrimitivePS, GIdentityPrimitiveUniformBuffer);
+			SetUniformBufferParameter(RHICmdList, PixelShader.GetPixelShader(), PrimitivePS, GIdentityPrimitiveUniformBuffer);
 		}
 	}
 
@@ -429,7 +415,7 @@ void FDecalRendering::SetVertexShaderOnly(FRHICommandList& RHICmdList, FGraphics
 	TShaderMapRef<FDeferredDecalVS> VertexShader(View.ShaderMap);
 
 	GraphicsPSOInit.BoundShaderState.VertexDeclarationRHI = GetVertexDeclarationFVector4();
-	GraphicsPSOInit.BoundShaderState.VertexShaderRHI = VertexShader->GetVertexShader();
+	GraphicsPSOInit.BoundShaderState.VertexShaderRHI = VertexShader.GetVertexShader();
 	GraphicsPSOInit.PrimitiveType = PT_TriangleList;
 
 	SetGraphicsPipelineState(RHICmdList, GraphicsPSOInit);
@@ -464,13 +450,13 @@ static uint8 ComputeDBufferMRTMask(EDecalBlendMode DecalBlendMode)
 	return 0;
 }
 
-void FDecalRendering::SetDecalCompilationEnvironment(EShaderPlatform Platform, const FMaterial* Material, FShaderCompilerEnvironment& OutEnvironment)
+void FDecalRendering::SetDecalCompilationEnvironment(const FMaterialShaderPermutationParameters& Parameters, FShaderCompilerEnvironment& OutEnvironment)
 {
-	bool bHasNormalConnected = Material->HasNormalConnected();
-	EDecalBlendMode FinalDecalBlendMode = FDecalRenderingCommon::ComputeFinalDecalBlendMode(Platform, (EDecalBlendMode)Material->GetDecalBlendMode(), bHasNormalConnected);
-	EDecalRenderStage DecalRenderStage = FDecalRenderingCommon::ComputeRenderStage(Platform, FinalDecalBlendMode);
-	FDecalRenderingCommon::ERenderTargetMode RenderTargetMode = FDecalRenderingCommon::ComputeRenderTargetMode(Platform, FinalDecalBlendMode, bHasNormalConnected);
-	uint32 RenderTargetCount = FDecalRenderingCommon::ComputeRenderTargetCount(Platform, RenderTargetMode);
+	const bool bHasNormalConnected = Parameters.MaterialParameters.bHasNormalConnected;
+	EDecalBlendMode FinalDecalBlendMode = FDecalRenderingCommon::ComputeFinalDecalBlendMode(Parameters.Platform, (EDecalBlendMode)Parameters.MaterialParameters.DecalBlendMode, bHasNormalConnected);
+	EDecalRenderStage DecalRenderStage = FDecalRenderingCommon::ComputeRenderStage(Parameters.Platform, FinalDecalBlendMode);
+	FDecalRenderingCommon::ERenderTargetMode RenderTargetMode = FDecalRenderingCommon::ComputeRenderTargetMode(Parameters.Platform, FinalDecalBlendMode, bHasNormalConnected);
+	uint32 RenderTargetCount = FDecalRenderingCommon::ComputeRenderTargetCount(Parameters.Platform, RenderTargetMode);
 
 	uint32 DecalOutputNormal = (RenderTargetMode == FDecalRenderingCommon::RTM_SceneColorAndGBufferNoNormal || RenderTargetMode == FDecalRenderingCommon::RTM_SceneColorAndGBufferDepthWriteNoNormal) ? 0 : 1;
 	OutEnvironment.SetDefine(TEXT("DECAL_OUTPUT_NORMAL"), DecalOutputNormal);
@@ -489,7 +475,7 @@ void FDecalRendering::SetDecalCompilationEnvironment(EShaderPlatform Platform, c
 
 }
 
-void FDecalRendering::SetEmissiveDBufferDecalCompilationEnvironment(EShaderPlatform Platform, const FMaterial* Material, FShaderCompilerEnvironment& OutEnvironment)
+void FDecalRendering::SetEmissiveDBufferDecalCompilationEnvironment(const FMaterialShaderPermutationParameters& Parameters, FShaderCompilerEnvironment& OutEnvironment)
 {
 	OutEnvironment.SetDefine(TEXT("DECAL_OUTPUT_NORMAL"), 0);
 	OutEnvironment.SetDefine(TEXT("DECAL_BLEND_MODE"), (uint32)DBM_DBuffer_Emissive);
diff --git a/Engine/Source/Runtime/Renderer/Private/DecalRenderingShared.h b/Engine/Source/Runtime/Renderer/Private/DecalRenderingShared.h
--- a/Engine/Source/Runtime/Renderer/Private/DecalRenderingShared.h
+++ b/Engine/Source/Runtime/Renderer/Private/DecalRenderingShared.h
@@ -14,6 +14,7 @@ class FDeferredDecalProxy;
 class FScene;
 class FViewInfo;
 struct FShaderCompilerEnvironment;
+struct FMaterialShaderPermutationParameters;
 
 /**
  * Compact decal data for rendering
@@ -44,6 +45,6 @@ struct FDecalRendering
 	static void SetShader(FRHICommandList& RHICmdList, FGraphicsPipelineStateInitializer& GraphicsPSOInit, const FViewInfo& View, const FTransientDecalRenderData& DecalData, EDecalRenderStage DecalRenderStage, const FMatrix& FrustumComponentToClip);
 
 	// Set common compilation environment parameters for decal shaders (FDeferredDecalPS, FMeshDecalsPS, etc.)
-	static void SetDecalCompilationEnvironment(EShaderPlatform Platform, const FMaterial* Material, FShaderCompilerEnvironment& OutEnvironment);
-	static void SetEmissiveDBufferDecalCompilationEnvironment(EShaderPlatform Platform, const FMaterial* Material, FShaderCompilerEnvironment& OutEnvironment);
+	static void SetDecalCompilationEnvironment(const FMaterialShaderPermutationParameters& Parameters, FShaderCompilerEnvironment& OutEnvironment);
+	static void SetEmissiveDBufferDecalCompilationEnvironment(const FMaterialShaderPermutationParameters& Parameters, FShaderCompilerEnvironment& OutEnvironment);
 };
