Copying Private-LoadTimes-4.24 stream to Main. Biggest changes are in Materials/Shader memory freezing.

diff --git a/Engine/Source/Runtime/Renderer/Private/DeferredShadingRenderer.cpp b/Engine/Source/Runtime/Renderer/Private/DeferredShadingRenderer.cpp
--- a/Engine/Source/Runtime/Renderer/Private/DeferredShadingRenderer.cpp
+++ b/Engine/Source/Runtime/Renderer/Private/DeferredShadingRenderer.cpp
@@ -292,7 +292,7 @@ void FDeferredShadingSceneRenderer::ClearGBufferAtMaxZ(FRHICommandList& RHICmdLi
 	auto ShaderMap = GetGlobalShaderMap(FeatureLevel);
 
 	TShaderMapRef<TOneColorVS<true> > VertexShader(ShaderMap);
-	FOneColorPS* PixelShader = NULL; 
+	TShaderRef<FOneColorPS> PixelShader; 
 
 	// Assume for now all code path supports SM4, otherwise render target numbers are changed
 	switch(NumActiveRenderTargets)
@@ -300,20 +300,20 @@ void FDeferredShadingSceneRenderer::ClearGBufferAtMaxZ(FRHICommandList& RHICmdLi
 	case 5:
 		{
 			TShaderMapRef<TOneColorPixelShaderMRT<5> > MRTPixelShader(ShaderMap);
-			PixelShader = *MRTPixelShader;
+			PixelShader = MRTPixelShader;
 		}
 		break;
 	case 6:
 		{
 			TShaderMapRef<TOneColorPixelShaderMRT<6> > MRTPixelShader(ShaderMap);
-			PixelShader = *MRTPixelShader;
+			PixelShader = MRTPixelShader;
 		}
 		break;
 	default:
 	case 1:
 		{
 			TShaderMapRef<TOneColorPixelShaderMRT<1> > MRTPixelShader(ShaderMap);
-			PixelShader = *MRTPixelShader;
+			PixelShader = MRTPixelShader;
 		}
 		break;
 	}
@@ -327,8 +327,8 @@ void FDeferredShadingSceneRenderer::ClearGBufferAtMaxZ(FRHICommandList& RHICmdLi
 	GraphicsPSOInit.DepthStencilState = TStaticDepthStencilState<false, CF_DepthNearOrEqual>::GetRHI();
 
 	GraphicsPSOInit.BoundShaderState.VertexDeclarationRHI = GetVertexDeclarationFVector4();
-	GraphicsPSOInit.BoundShaderState.VertexShaderRHI = GETSAFERHISHADER_VERTEX(*VertexShader);
-	GraphicsPSOInit.BoundShaderState.PixelShaderRHI = GETSAFERHISHADER_PIXEL(PixelShader);
+	GraphicsPSOInit.BoundShaderState.VertexShaderRHI = VertexShader.GetVertexShader();
+	GraphicsPSOInit.BoundShaderState.PixelShaderRHI = PixelShader.GetPixelShader();
 	GraphicsPSOInit.PrimitiveType = PT_TriangleStrip;
 
 	SetGraphicsPipelineState(RHICmdList, GraphicsPSOInit);
@@ -1082,7 +1082,7 @@ bool FDeferredShadingSceneRenderer::DispatchRayTracingWorldUpdates(FRHICommandLi
 
 		if (RayGenShaders.Num())
 		{
-			auto DefaultHitShader = View.ShaderMap->GetShader<FOpaqueShadowHitGroup>()->GetRayTracingShader();
+			auto DefaultHitShader = View.ShaderMap->GetShader<FOpaqueShadowHitGroup>().GetRayTracingShader();
 
 			View.RayTracingMaterialPipeline = BindRayTracingMaterialPipeline(RHICmdList, View,
 				RayGenShaders,
@@ -1321,7 +1321,7 @@ void FDeferredShadingSceneRenderer::Render(FRHICommandListImmediate& RHICmdList)
 
 	// Find the visible primitives.
 	RHICmdList.ImmediateFlush(EImmediateFlushType::DispatchToRHIThread);
-	
+
 	bool bDoInitViewAftersPrepass = false;
 	{
 		SCOPED_GPU_STAT(RHICmdList, VisibilityCommands);
@@ -1634,7 +1634,7 @@ void FDeferredShadingSceneRenderer::Render(FRHICommandListImmediate& RHICmdList)
 
 	if (HasHairStrandsProjectionQuery(Scene->GetShaderPlatform()))
 	{
-		auto ShaderMap = GetGlobalShaderMap(FeatureLevel);
+		FGlobalShaderMap* ShaderMap = GetGlobalShaderMap(FeatureLevel);
 		RunHairStrandsBindingQueries(RHICmdList, ShaderMap);
 	}
 
@@ -2842,13 +2842,13 @@ public:
 
 	void SetParameters(FRHICommandList& RHICmdList, const FSceneView& View, bool bUseMaxDepth, FIntPoint ViewMax)
 	{
-		FGlobalShader::SetParameters<FViewUniformShaderParameters>(RHICmdList, GetPixelShader(), View.ViewUniformBuffer);
+		FGlobalShader::SetParameters<FViewUniformShaderParameters>(RHICmdList, RHICmdList.GetBoundPixelShader(), View.ViewUniformBuffer);
 		FSceneRenderTargets& SceneContext = FSceneRenderTargets::Get(RHICmdList);
 
 		// Used to remap view space Z (which is stored in scene color alpha) into post projection z and w so we can write z/w into the downsized depth buffer
 		const FVector2D ProjectionScaleBiasValue(View.ViewMatrices.GetProjectionMatrix().M[2][2], View.ViewMatrices.GetProjectionMatrix().M[3][2]);
-		SetShaderValue(RHICmdList, GetPixelShader(), ProjectionScaleBias, ProjectionScaleBiasValue);
-		SetShaderValue(RHICmdList, GetPixelShader(), UseMaxDepth, (bUseMaxDepth ? 1.0f : 0.0f));
+		SetShaderValue(RHICmdList, RHICmdList.GetBoundPixelShader(), ProjectionScaleBias, ProjectionScaleBiasValue);
+		SetShaderValue(RHICmdList, RHICmdList.GetBoundPixelShader(), UseMaxDepth, (bUseMaxDepth ? 1.0f : 0.0f));
 
 		FIntPoint BufferSize = SceneContext.GetBufferSizeXY();
 
@@ -2857,34 +2857,22 @@ public:
 
 		// Offsets of the four full resolution pixels corresponding with a low resolution pixel
 		const FVector4 Offsets01(0.0f, 0.0f, 1.0f / DownsampledBufferSizeX, 0.0f);
-		SetShaderValue(RHICmdList, GetPixelShader(), SourceTexelOffsets01, Offsets01);
+		SetShaderValue(RHICmdList, RHICmdList.GetBoundPixelShader(), SourceTexelOffsets01, Offsets01);
 		const FVector4 Offsets23(0.0f, 1.0f / DownsampledBufferSizeY, 1.0f / DownsampledBufferSizeX, 1.0f / DownsampledBufferSizeY);
-		SetShaderValue(RHICmdList, GetPixelShader(), SourceTexelOffsets23, Offsets23);
-		SceneTextureParameters.Set(RHICmdList, GetPixelShader(), View.FeatureLevel, ESceneTextureSetupMode::All);
+		SetShaderValue(RHICmdList, RHICmdList.GetBoundPixelShader(), SourceTexelOffsets23, Offsets23);
+		SceneTextureParameters.Set(RHICmdList, RHICmdList.GetBoundPixelShader(), View.FeatureLevel, ESceneTextureSetupMode::All);
 
 		// Set MaxUV, so we won't sample outside of a valid texture region.
 		FVector2D const SourceMaxUV((ViewMax.X - 0.5f) / BufferSize.X, (ViewMax.Y - 0.5f) / BufferSize.Y);
-		SetShaderValue(RHICmdList, GetPixelShader(), SourceMaxUVParameter, SourceMaxUV);
+		SetShaderValue(RHICmdList, RHICmdList.GetBoundPixelShader(), SourceMaxUVParameter, SourceMaxUV);
 	}
 
-	virtual bool Serialize(FArchive& Ar) override
-	{
-		bool bShaderHasOutdatedParameters = FGlobalShader::Serialize(Ar);
-		Ar << ProjectionScaleBias;
-		Ar << SourceTexelOffsets01;
-		Ar << SourceTexelOffsets23;
-		Ar << SceneTextureParameters;
-		Ar << UseMaxDepth;
-		Ar << SourceMaxUVParameter;
-		return bShaderHasOutdatedParameters;
-	}
-
-	FShaderParameter ProjectionScaleBias;
-	FShaderParameter SourceTexelOffsets01;
-	FShaderParameter SourceTexelOffsets23;
-	FShaderParameter SourceMaxUVParameter;
-	FSceneTextureShaderParameters SceneTextureParameters;
-	FShaderParameter UseMaxDepth;
+	LAYOUT_FIELD(FShaderParameter, ProjectionScaleBias);
+	LAYOUT_FIELD(FShaderParameter, SourceTexelOffsets01);
+	LAYOUT_FIELD(FShaderParameter, SourceTexelOffsets23);
+	LAYOUT_FIELD(FShaderParameter, SourceMaxUVParameter);
+	LAYOUT_FIELD(FSceneTextureShaderParameters, SceneTextureParameters);
+	LAYOUT_FIELD(FShaderParameter, UseMaxDepth);
 };
 
 IMPLEMENT_SHADER_TYPE(,FDownsampleSceneDepthPS,TEXT("/Engine/Private/DownsampleDepthPixelShader.usf"),TEXT("Main"),SF_Pixel);
@@ -2934,8 +2922,8 @@ void FDeferredShadingSceneRenderer::DownsampleDepthSurface(FRHICommandList& RHIC
 	GraphicsPSOInit.DepthStencilState = TStaticDepthStencilState<true, CF_Always>::GetRHI();
 
 	GraphicsPSOInit.BoundShaderState.VertexDeclarationRHI = GFilterVertexDeclaration.VertexDeclarationRHI;
-	GraphicsPSOInit.BoundShaderState.VertexShaderRHI = GETSAFERHISHADER_VERTEX(*ScreenVertexShader);
-	GraphicsPSOInit.BoundShaderState.PixelShaderRHI = GETSAFERHISHADER_PIXEL(*PixelShader);
+	GraphicsPSOInit.BoundShaderState.VertexShaderRHI = ScreenVertexShader.GetVertexShader();
+	GraphicsPSOInit.BoundShaderState.PixelShaderRHI = PixelShader.GetPixelShader();
 	GraphicsPSOInit.PrimitiveType = PT_TriangleList;
 
 	SetGraphicsPipelineState(RHICmdList, GraphicsPSOInit);
@@ -2956,7 +2944,7 @@ void FDeferredShadingSceneRenderer::DownsampleDepthSurface(FRHICommandList& RHIC
 		View.ViewRect.Width(), View.ViewRect.Height(),
 		FIntPoint(DownsampledSizeX, DownsampledSizeY),
 		SceneContext.GetBufferSizeXY(),
-		*ScreenVertexShader,
+		ScreenVertexShader,
 		EDRF_UseTriangleOptimization);
 	}
 	RHICmdList.EndRenderPass();
@@ -2989,20 +2977,13 @@ public:
 
 	void SetParameters(FRHICommandList& RHICmdList, const FSceneView& View)
 	{
-		FGlobalShader::SetParameters<FViewUniformShaderParameters>(RHICmdList, GetPixelShader(), View.ViewUniformBuffer);
+		FGlobalShader::SetParameters<FViewUniformShaderParameters>(RHICmdList, RHICmdList.GetBoundPixelShader(), View.ViewUniformBuffer);
 		FSceneRenderTargets& SceneContext = FSceneRenderTargets::Get(RHICmdList);
 
-		SetSRVParameter(RHICmdList, GetPixelShader(), SceneStencilTexture, SceneContext.SceneStencilSRV);
-	}
-
-	virtual bool Serialize(FArchive& Ar) override
-	{
-		bool bShaderHasOutdatedParameters = FGlobalShader::Serialize(Ar);
-		Ar << SceneStencilTexture;
-		return bShaderHasOutdatedParameters;
+		SetSRVParameter(RHICmdList, RHICmdList.GetBoundPixelShader(), SceneStencilTexture, SceneContext.SceneStencilSRV);
 	}
 
-	FShaderResourceParameter SceneStencilTexture;
+	LAYOUT_FIELD(FShaderResourceParameter, SceneStencilTexture);
 };
 
 IMPLEMENT_SHADER_TYPE(,FCopyStencilToLightingChannelsPS,TEXT("/Engine/Private/DownsampleDepthPixelShader.usf"),TEXT("CopyStencilToLightingChannelsPS"),SF_Pixel);
@@ -3048,8 +3029,8 @@ void FDeferredShadingSceneRenderer::CopyStencilToLightingChannelTexture(FRHIComm
 			TShaderMapRef<FScreenVS> ScreenVertexShader(View.ShaderMap);
 			TShaderMapRef<FCopyStencilToLightingChannelsPS> PixelShader(View.ShaderMap);
 
-			GraphicsPSOInit.BoundShaderState.VertexShaderRHI = GETSAFERHISHADER_VERTEX(*ScreenVertexShader);
-			GraphicsPSOInit.BoundShaderState.PixelShaderRHI = GETSAFERHISHADER_PIXEL(*PixelShader);
+			GraphicsPSOInit.BoundShaderState.VertexShaderRHI = ScreenVertexShader.GetVertexShader();
+			GraphicsPSOInit.BoundShaderState.PixelShaderRHI = PixelShader.GetPixelShader();
 			GraphicsPSOInit.PrimitiveType = PT_TriangleList;
 
 			SetGraphicsPipelineState(RHICmdList, GraphicsPSOInit);
@@ -3065,7 +3046,7 @@ void FDeferredShadingSceneRenderer::CopyStencilToLightingChannelTexture(FRHIComm
 				View.ViewRect.Width(), View.ViewRect.Height(),
 				FIntPoint(View.ViewRect.Width(), View.ViewRect.Height()),
 				SceneContext.GetBufferSizeXY(),
-				*ScreenVertexShader,
+				ScreenVertexShader,
 				EDRF_UseTriangleOptimization);
 		}
 		}
diff --git a/Engine/Source/Runtime/Renderer/Private/DepthRendering.cpp b/Engine/Source/Runtime/Renderer/Private/DepthRendering.cpp
--- a/Engine/Source/Runtime/Renderer/Private/DepthRendering.cpp
+++ b/Engine/Source/Runtime/Renderer/Private/DepthRendering.cpp
@@ -96,18 +96,18 @@ void GetDepthPassShaders(
 	const FMaterial& Material,
 	FVertexFactoryType* VertexFactoryType,
 	ERHIFeatureLevel::Type FeatureLevel,
-	FDepthOnlyHS*& HullShader,
-	FDepthOnlyDS*& DomainShader,
-	TDepthOnlyVS<bPositionOnly>*& VertexShader,
-	FDepthOnlyPS*& PixelShader,
-	FShaderPipeline*& ShaderPipeline,
+	TShaderRef<FDepthOnlyHS>& HullShader,
+	TShaderRef<FDepthOnlyDS>& DomainShader,
+	TShaderRef<TDepthOnlyVS<bPositionOnly>>& VertexShader,
+	TShaderRef<FDepthOnlyPS>& PixelShader,
+	FShaderPipelineRef& ShaderPipeline,
 	bool bUsesMobileColorValue)
 {
 	if (bPositionOnly && !bUsesMobileColorValue)
 	{
-		ShaderPipeline = UseShaderPipelines(FeatureLevel) ? Material.GetShaderPipeline(&DepthPosOnlyNoPixelPipeline, VertexFactoryType) : nullptr;
-		VertexShader = ShaderPipeline
-			? ShaderPipeline->GetShader<TDepthOnlyVS<bPositionOnly> >()
+		ShaderPipeline = UseShaderPipelines(FeatureLevel) ? Material.GetShaderPipeline(&DepthPosOnlyNoPixelPipeline, VertexFactoryType) : FShaderPipelineRef();
+		VertexShader = ShaderPipeline.IsValid()
+			? ShaderPipeline.GetShader<TDepthOnlyVS<bPositionOnly> >()
 			: Material.GetShader<TDepthOnlyVS<bPositionOnly> >(VertexFactoryType);
 	}
 	else
@@ -119,7 +119,7 @@ void GetDepthPassShaders(
 			&& VertexFactoryType->SupportsTessellationShaders() 
 			&& TessellationMode != MTM_NoTessellation)
 		{
-			ShaderPipeline = nullptr;
+			ShaderPipeline = FShaderPipelineRef();
 			VertexShader = Material.GetShader<TDepthOnlyVS<bPositionOnly> >(VertexFactoryType);
 			HullShader = Material.GetShader<FDepthOnlyHS>(VertexFactoryType);
 			DomainShader = Material.GetShader<FDepthOnlyDS>(VertexFactoryType);
@@ -130,24 +130,24 @@ void GetDepthPassShaders(
 		}
 		else
 		{
-			HullShader = nullptr;
-			DomainShader = nullptr;
+			HullShader.Reset();
+			DomainShader.Reset();
 			bool bUseShaderPipelines = UseShaderPipelines(FeatureLevel);
 			if (bNeedsPixelShader)
 			{
-				ShaderPipeline = bUseShaderPipelines ? Material.GetShaderPipeline(&DepthPipeline, VertexFactoryType, false) : nullptr;
+				ShaderPipeline = bUseShaderPipelines ? Material.GetShaderPipeline(&DepthPipeline, VertexFactoryType, false) : FShaderPipelineRef();
 			}
 			else
 			{
-				ShaderPipeline = bUseShaderPipelines ? Material.GetShaderPipeline(&DepthNoPixelPipeline, VertexFactoryType, false) : nullptr;
+				ShaderPipeline = bUseShaderPipelines ? Material.GetShaderPipeline(&DepthNoPixelPipeline, VertexFactoryType, false) : FShaderPipelineRef();
 			}
 
-			if (ShaderPipeline)
+			if (ShaderPipeline.IsValid())
 			{
-				VertexShader = ShaderPipeline->GetShader<TDepthOnlyVS<bPositionOnly> >();
+				VertexShader = ShaderPipeline.GetShader<TDepthOnlyVS<bPositionOnly> >();
 				if (bNeedsPixelShader)
 				{
-					PixelShader = ShaderPipeline->GetShader<FDepthOnlyPS>();
+					PixelShader = ShaderPipeline.GetShader<FDepthOnlyPS>();
 				}
 			}
 			else
@@ -167,11 +167,11 @@ void GetDepthPassShaders(
 		const FMaterial& Material, \
 		FVertexFactoryType* VertexFactoryType, \
 		ERHIFeatureLevel::Type FeatureLevel, \
-		FDepthOnlyHS*& HullShader, \
-		FDepthOnlyDS*& DomainShader, \
-		TDepthOnlyVS<bPositionOnly>*& VertexShader, \
-		FDepthOnlyPS*& PixelShader, \
-		FShaderPipeline*& ShaderPipeline, \
+		TShaderRef<FDepthOnlyHS>& HullShader, \
+		TShaderRef<FDepthOnlyDS>& DomainShader, \
+		TShaderRef<TDepthOnlyVS<bPositionOnly>>& VertexShader, \
+		TShaderRef<FDepthOnlyPS>& PixelShader, \
+		FShaderPipelineRef& ShaderPipeline, \
 		bool bUsesMobileColorValue \
 	);
 
@@ -247,7 +247,7 @@ static void RenderHiddenAreaMaskView(FRHICommandList& RHICmdList, FGraphicsPipel
 
 	extern TGlobalResource<FFilterVertexDeclaration> GFilterVertexDeclaration;
 	GraphicsPSOInit.BoundShaderState.VertexDeclarationRHI = GFilterVertexDeclaration.VertexDeclarationRHI;
-	GraphicsPSOInit.BoundShaderState.VertexShaderRHI = GETSAFERHISHADER_VERTEX(*VertexShader);
+	GraphicsPSOInit.BoundShaderState.VertexShaderRHI = VertexShader.GetVertexShader();
 	GraphicsPSOInit.PrimitiveType = PT_TriangleList;
 
 	SetGraphicsPipelineState(RHICmdList, GraphicsPSOInit);
@@ -343,20 +343,13 @@ public:
 
 	void SetParameters(FRHICommandList& RHICmdList, const FSceneView& View)
 	{
-		FGlobalShader::SetParameters<FViewUniformShaderParameters>(RHICmdList, GetPixelShader(), View.ViewUniformBuffer);
+		FGlobalShader::SetParameters<FViewUniformShaderParameters>(RHICmdList, RHICmdList.GetBoundPixelShader(), View.ViewUniformBuffer);
 
 		const float DitherFactor = View.GetTemporalLODTransition();
-		SetShaderValue(RHICmdList, GetPixelShader(), DitheredTransitionFactorParameter, DitherFactor);
+		SetShaderValue(RHICmdList, RHICmdList.GetBoundPixelShader(), DitheredTransitionFactorParameter, DitherFactor);
 	}
 
-	virtual bool Serialize(FArchive& Ar) override
-	{
-		bool bShaderHasOutdatedParameters = FGlobalShader::Serialize(Ar);
-		Ar << DitheredTransitionFactorParameter;
-		return bShaderHasOutdatedParameters;
-	}
-
-	FShaderParameter DitheredTransitionFactorParameter;
+	LAYOUT_FIELD(FShaderParameter, DitheredTransitionFactorParameter);
 };
 IMPLEMENT_SHADER_TYPE(, FDitheredTransitionStencilPS, TEXT("/Engine/Private/DitheredTransitionStencil.usf"), TEXT("Main"), SF_Pixel);
 
@@ -387,7 +380,7 @@ public:
 	template <typename TRHICmdList>
 	void SetParameters(TRHICmdList& RHICmdList, const FSceneView& View, FRHIUnorderedAccessView* StencilOutputUAV, FIntPoint BufferSizeXY, FIntPoint ViewOffsetXY, uint32 StencilValue)
 	{
-		FRHIComputeShader* ComputeShader = GetComputeShader();
+		FRHIComputeShader* ComputeShader = RHICmdList.GetBoundComputeShader();
 
 		FGlobalShader::SetParameters<FViewUniformShaderParameters>(RHICmdList, ComputeShader, View.ViewUniformBuffer);
 
@@ -410,23 +403,16 @@ public:
 	template <typename TRHICmdList>
 	void UnsetParameters(TRHICmdList& RHICmdList)
 	{
-		FRHIComputeShader* ComputeShader = GetComputeShader();
+		FRHIComputeShader* ComputeShader = RHICmdList.GetBoundComputeShader();
 		if (StencilOutputParameter.IsBound())
 		{
 			RHICmdList.SetUAVParameter(ComputeShader, StencilOutputParameter.GetBaseIndex(), nullptr);
 		}
 	}
 
-	virtual bool Serialize(FArchive& Ar) override
-	{
-		bool bShaderHasOutdatedParameters = FGlobalShader::Serialize(Ar);
-		Ar << DitheredTransitionFactorParameter << StencilOffsetAndValuesParameter << StencilOutputParameter;
-		return bShaderHasOutdatedParameters;
-	}
-
-	FShaderParameter DitheredTransitionFactorParameter;
-	FShaderParameter StencilOffsetAndValuesParameter;
-	FShaderResourceParameter StencilOutputParameter;
+	LAYOUT_FIELD(FShaderParameter, DitheredTransitionFactorParameter);
+	LAYOUT_FIELD(FShaderParameter, StencilOffsetAndValuesParameter);
+	LAYOUT_FIELD(FShaderResourceParameter, StencilOutputParameter);
 };
 IMPLEMENT_SHADER_TYPE(, FDitheredTransitionStencilCS, TEXT("/Engine/Private/DitheredTransitionStencil.usf"), TEXT("MainCS"), SF_Compute);
 
@@ -476,13 +462,13 @@ void FDeferredShadingSceneRenderer::PreRenderDitherFill(FRHIAsyncComputeCommandL
 		FViewInfo& View = Views[ViewIndex];
 
 		TShaderMapRef<FDitheredTransitionStencilCS> ComputeShader(View.ShaderMap);
-		RHICmdList.SetComputeShader(ComputeShader->GetComputeShader());
+		RHICmdList.SetComputeShader(ComputeShader.GetComputeShader());
 		ComputeShader->SetParameters(RHICmdList, View, StencilTextureUAV, BufferSizeXY, View.ViewRect.Min, STENCIL_SANDBOX_MASK);
 		const int32 SubWidth = FMath::Min(BufferSizeXY.X, View.ViewRect.Width());
 		const int32 SubHeight = FMath::Min(BufferSizeXY.Y, View.ViewRect.Height());
 		check(SubWidth > 0 && SubHeight > 0);
 
-		DispatchComputeShader(RHICmdList, *ComputeShader, FMath::DivideAndRoundUp(SubWidth, 8), FMath::DivideAndRoundUp(SubHeight, 8), 1);
+		DispatchComputeShader(RHICmdList, ComputeShader.GetShader(), FMath::DivideAndRoundUp(SubWidth, 8), FMath::DivideAndRoundUp(SubHeight, 8), 1);
 		ComputeShader->UnsetParameters(RHICmdList);
 	}
 }
@@ -502,13 +488,13 @@ void FDeferredShadingSceneRenderer::PreRenderDitherFill(FRHICommandListImmediate
 			FViewInfo& View = Views[ViewIndex];
 
 			TShaderMapRef<FDitheredTransitionStencilCS> ComputeShader(View.ShaderMap);
-			RHICmdList.SetComputeShader(ComputeShader->GetComputeShader());
+			RHICmdList.SetComputeShader(ComputeShader.GetComputeShader());
 			ComputeShader->SetParameters(RHICmdList, View, StencilTextureUAV, BufferSizeXY, View.ViewRect.Min, STENCIL_SANDBOX_MASK);
 			const int32 SubWidth = FMath::Min(BufferSizeXY.X, View.ViewRect.Width());
 			const int32 SubHeight = FMath::Min(BufferSizeXY.Y, View.ViewRect.Height());
 			check(SubWidth > 0 && SubHeight > 0);
 
-			DispatchComputeShader(RHICmdList, *ComputeShader, FMath::DivideAndRoundUp(SubWidth, 8), FMath::DivideAndRoundUp(SubHeight, 8), 1);
+			DispatchComputeShader(RHICmdList, ComputeShader.GetShader(), FMath::DivideAndRoundUp(SubWidth, 8), FMath::DivideAndRoundUp(SubHeight, 8), 1);
 			ComputeShader->UnsetParameters(RHICmdList);
 		}
 	}
@@ -538,8 +524,8 @@ void FDeferredShadingSceneRenderer::PreRenderDitherFill(FRHICommandListImmediate
 
 			extern TGlobalResource<FFilterVertexDeclaration> GFilterVertexDeclaration;
 			GraphicsPSOInit.BoundShaderState.VertexDeclarationRHI = GFilterVertexDeclaration.VertexDeclarationRHI;
-			GraphicsPSOInit.BoundShaderState.VertexShaderRHI = GETSAFERHISHADER_VERTEX(*ScreenVertexShader);
-			GraphicsPSOInit.BoundShaderState.PixelShaderRHI = GETSAFERHISHADER_PIXEL(*PixelShader);
+			GraphicsPSOInit.BoundShaderState.VertexShaderRHI = ScreenVertexShader.GetVertexShader();
+			GraphicsPSOInit.BoundShaderState.PixelShaderRHI = PixelShader.GetPixelShader();
 			GraphicsPSOInit.PrimitiveType = PT_TriangleList;
 
 			SetGraphicsPipelineState(RHICmdList, GraphicsPSOInit);
@@ -555,7 +541,7 @@ void FDeferredShadingSceneRenderer::PreRenderDitherFill(FRHICommandListImmediate
 				View.ViewRect.Width(), View.ViewRect.Height(),
 				BufferSizeXY,
 				BufferSizeXY,
-				*ScreenVertexShader,
+				ScreenVertexShader,
 				EDRF_UseTriangleOptimization);
 		}
 	}
@@ -895,7 +881,7 @@ void FDepthPassMeshProcessor::Process(
 		FDepthOnlyDS,
 		FDepthOnlyPS> DepthPassShaders;
 
-	FShaderPipeline* ShaderPipeline = nullptr;
+	FShaderPipelineRef ShaderPipeline;
 
 	GetDepthPassShaders<bPositionOnly>(
 		MaterialResource,
@@ -916,7 +902,7 @@ void FDepthPassMeshProcessor::Process(
 	FDepthOnlyShaderElementData ShaderElementData(0.0f);
 	ShaderElementData.InitializeMeshMaterialData(ViewIfDynamicMeshCommand, PrimitiveSceneProxy, MeshBatch, StaticMeshId, true);
 
-	const FMeshDrawCommandSortKey SortKey = CalculateDepthPassMeshStaticSortKey(BlendMode, DepthPassShaders.VertexShader, DepthPassShaders.PixelShader);
+	const FMeshDrawCommandSortKey SortKey = CalculateDepthPassMeshStaticSortKey(BlendMode, DepthPassShaders.VertexShader.GetShader(), DepthPassShaders.PixelShader.GetShader());
 
 	BuildMeshDrawCommands(
 		MeshBatch,
@@ -970,8 +956,9 @@ void FDepthPassMeshProcessor::AddMeshBatch(const FMeshBatch& RESTRICT MeshBatch,
 		const FMaterialRenderProxy& MaterialRenderProxy = FallbackMaterialRenderProxyPtr ? *FallbackMaterialRenderProxyPtr : *MeshBatch.MaterialRenderProxy;
 
 		const EBlendMode BlendMode = Material.GetBlendMode();
-		const ERasterizerFillMode MeshFillMode = ComputeMeshFillMode(MeshBatch, Material);
-		const ERasterizerCullMode MeshCullMode = ComputeMeshCullMode(MeshBatch, Material);
+		const FMeshDrawingPolicyOverrideSettings OverrideSettings = ComputeMeshOverrideSettings(MeshBatch);
+		const ERasterizerFillMode MeshFillMode = ComputeMeshFillMode(MeshBatch, Material, OverrideSettings);
+		const ERasterizerCullMode MeshCullMode = ComputeMeshCullMode(MeshBatch, Material, OverrideSettings);
 		const bool bIsTranslucent = IsTranslucentBlendMode(BlendMode);
 
 		if (!bIsTranslucent
diff --git a/Engine/Source/Runtime/Renderer/Private/DepthRendering.h b/Engine/Source/Runtime/Renderer/Private/DepthRendering.h
--- a/Engine/Source/Runtime/Renderer/Private/DepthRendering.h
+++ b/Engine/Source/Runtime/Renderer/Private/DepthRendering.h
@@ -68,15 +68,15 @@ public:
 		// Only the local vertex factory supports the position-only stream
 		if (bUsePositionOnlyStream)
 		{
-			return Parameters.VertexFactoryType->SupportsPositionOnly() && Parameters.Material->IsSpecialEngineMaterial();
+			return Parameters.VertexFactoryType->SupportsPositionOnly() && Parameters.MaterialParameters.bIsSpecialEngineMaterial;
 		}
 
 		// Only compile for the default material and masked materials
 		return (
-			Parameters.Material->IsSpecialEngineMaterial() ||
-			!Parameters.Material->WritesEveryPixel() ||
-			Parameters.Material->MaterialMayModifyMeshPosition() ||
-			Parameters.Material->IsTranslucencyWritingCustomDepth());
+			Parameters.MaterialParameters.bIsSpecialEngineMaterial ||
+			!Parameters.MaterialParameters.bWritesEveryPixel ||
+			Parameters.MaterialParameters.bMaterialMayModifyMeshPosition ||
+			Parameters.MaterialParameters.bIsTranslucencyWritingCustomDepth);
 	}
 
 	void GetShaderBindings(
@@ -147,9 +147,9 @@ public:
 	{
 		return 
 			// Compile for materials that are masked.
-			(!Parameters.Material->WritesEveryPixel() || Parameters.Material->HasPixelDepthOffsetConnected() || Parameters.Material->IsTranslucencyWritingCustomDepth()) 
+			(!Parameters.MaterialParameters.bWritesEveryPixel || Parameters.MaterialParameters.bHasPixelDepthOffsetConnected || Parameters.MaterialParameters.bIsTranslucencyWritingCustomDepth)
 			// Mobile uses material pixel shader to write custom stencil to color target
-			|| (IsMobilePlatform(Parameters.Platform) && (Parameters.Material->IsDefaultMaterial() || Parameters.Material->MaterialMayModifyMeshPosition()));
+			|| (IsMobilePlatform(Parameters.Platform) && (Parameters.MaterialParameters.bIsDefaultMaterial || Parameters.MaterialParameters.bMaterialMayModifyMeshPosition));
 	}
 
 	FDepthOnlyPS(const ShaderMetaType::CompiledShaderInitializerType& Initializer):
@@ -188,14 +188,7 @@ public:
 		ShaderBindings.Add(MobileColorValue, ShaderElementData.MobileColorValue);
 	}
 
-	virtual bool Serialize(FArchive& Ar) override
-	{
-		bool bShaderHasOutdatedParameters = FMeshMaterialShader::Serialize(Ar);
-		Ar << MobileColorValue;
-		return bShaderHasOutdatedParameters;
-	}
-
-	FShaderParameter MobileColorValue;
+	LAYOUT_FIELD(FShaderParameter, MobileColorValue);
 };
 
 template <bool bPositionOnly>
@@ -203,11 +196,11 @@ void GetDepthPassShaders(
 	const FMaterial& Material,
 	FVertexFactoryType* VertexFactoryType,
 	ERHIFeatureLevel::Type FeatureLevel,
-	FDepthOnlyHS*& HullShader,
-	FDepthOnlyDS*& DomainShader,
-	TDepthOnlyVS<bPositionOnly>*& VertexShader,
-	FDepthOnlyPS*& PixelShader,
-	FShaderPipeline*& ShaderPipeline,
+	TShaderRef<FDepthOnlyHS>& HullShader,
+	TShaderRef<FDepthOnlyDS>& DomainShader,
+	TShaderRef<TDepthOnlyVS<bPositionOnly>>& VertexShader,
+	TShaderRef<FDepthOnlyPS>& PixelShader,
+	FShaderPipelineRef& ShaderPipeline,
 	bool bUsesMobileColorValue);
 
 extern void CreateDepthPassUniformBuffer(
diff --git a/Engine/Source/Runtime/Renderer/Private/DistanceFieldAmbientOcclusion.cpp b/Engine/Source/Runtime/Renderer/Private/DistanceFieldAmbientOcclusion.cpp
--- a/Engine/Source/Runtime/Renderer/Private/DistanceFieldAmbientOcclusion.cpp
+++ b/Engine/Source/Runtime/Renderer/Private/DistanceFieldAmbientOcclusion.cpp
@@ -20,6 +20,11 @@
 #include "VisualizeTexture.h"
 #include "RayTracing/RaytracingOptions.h"
 
+IMPLEMENT_TYPE_LAYOUT(FAOParameters);
+IMPLEMENT_TYPE_LAYOUT(FDFAOUpsampleParameters);
+IMPLEMENT_TYPE_LAYOUT(FScreenGridParameters);
+IMPLEMENT_TYPE_LAYOUT(FTileIntersectionParameters);
+
 int32 GDistanceFieldAO = 1;
 FAutoConsoleVariableRef CVarDistanceFieldAO(
 	TEXT("r.DistanceFieldAO"),
@@ -302,25 +307,17 @@ public:
 
 	void SetParameters(FRHICommandList& RHICmdList, const FSceneView& View, const FDistanceFieldAOParameters& Parameters)
 	{
-		FRHIPixelShader* ShaderRHI = GetPixelShader();
+		FRHIPixelShader* ShaderRHI = RHICmdList.GetBoundPixelShader();
 
 		FGlobalShader::SetParameters<FViewUniformShaderParameters>(RHICmdList, ShaderRHI, View.ViewUniformBuffer);
 		AOParameters.Set(RHICmdList, ShaderRHI, Parameters);
 		SceneTextureParameters.Set(RHICmdList, ShaderRHI, View.FeatureLevel, ESceneTextureSetupMode::All);
 	}
-	// FShader interface.
-	virtual bool Serialize(FArchive& Ar) override
-	{
-		bool bShaderHasOutdatedParameters = FGlobalShader::Serialize(Ar);
-		Ar << SceneTextureParameters;
-		Ar << AOParameters;
-		return bShaderHasOutdatedParameters;
-	}
 
 private:
 
-	FSceneTextureShaderParameters SceneTextureParameters;
-	FAOParameters AOParameters;
+	LAYOUT_FIELD(FSceneTextureShaderParameters, SceneTextureParameters);
+	LAYOUT_FIELD(FAOParameters, AOParameters);
 };
 
 IMPLEMENT_SHADER_TYPE(,FComputeDistanceFieldNormalPS,TEXT("/Engine/Private/DistanceFieldScreenGridLighting.usf"),TEXT("ComputeDistanceFieldNormalPS"),SF_Pixel);
@@ -357,7 +354,7 @@ public:
 
 	void SetParameters(FRHICommandList& RHICmdList, const FSceneView& View, FSceneRenderTargetItem& DistanceFieldNormalValue, const FDistanceFieldAOParameters& Parameters)
 	{
-		FRHIComputeShader* ShaderRHI = GetComputeShader();
+		FRHIComputeShader* ShaderRHI = RHICmdList.GetBoundComputeShader();
 
 		FGlobalShader::SetParameters<FViewUniformShaderParameters>(RHICmdList, ShaderRHI, View.ViewUniformBuffer);
 
@@ -369,25 +366,15 @@ public:
 
 	void UnsetParameters(FRHICommandList& RHICmdList, FSceneRenderTargetItem& DistanceFieldNormalValue)
 	{
-		DistanceFieldNormal.UnsetUAV(RHICmdList, GetComputeShader());
+		DistanceFieldNormal.UnsetUAV(RHICmdList, RHICmdList.GetBoundComputeShader());
 		RHICmdList.TransitionResource(EResourceTransitionAccess::EReadable, EResourceTransitionPipeline::EComputeToCompute, DistanceFieldNormalValue.UAV);
 	}
 
-	// FShader interface.
-	virtual bool Serialize(FArchive& Ar) override
-	{
-		bool bShaderHasOutdatedParameters = FGlobalShader::Serialize(Ar);
-		Ar << DistanceFieldNormal;
-		Ar << SceneTextureParameters;
-		Ar << AOParameters;
-		return bShaderHasOutdatedParameters;
-	}
-
 private:
 
-	FRWShaderParameter DistanceFieldNormal;
-	FSceneTextureShaderParameters SceneTextureParameters;
-	FAOParameters AOParameters;
+	LAYOUT_FIELD(FRWShaderParameter, DistanceFieldNormal);
+	LAYOUT_FIELD(FSceneTextureShaderParameters, SceneTextureParameters);
+	LAYOUT_FIELD(FAOParameters, AOParameters);
 };
 
 IMPLEMENT_SHADER_TYPE(,FComputeDistanceFieldNormalCS,TEXT("/Engine/Private/DistanceFieldScreenGridLighting.usf"),TEXT("ComputeDistanceFieldNormalCS"),SF_Compute);
@@ -411,9 +398,9 @@ void ComputeDistanceFieldNormal(FRHICommandListImmediate& RHICmdList, const TArr
 				SCOPED_DRAW_EVENT(RHICmdList, ComputeNormalCS);
 				TShaderMapRef<FComputeDistanceFieldNormalCS> ComputeShader(View.ShaderMap);
 
-				RHICmdList.SetComputeShader(ComputeShader->GetComputeShader());
+				RHICmdList.SetComputeShader(ComputeShader.GetComputeShader());
 				ComputeShader->SetParameters(RHICmdList, View, DistanceFieldNormal, Parameters);
-				DispatchComputeShader(RHICmdList, *ComputeShader, GroupSizeX, GroupSizeY, 1);
+				DispatchComputeShader(RHICmdList, ComputeShader.GetShader(), GroupSizeX, GroupSizeY, 1);
 
 				ComputeShader->UnsetParameters(RHICmdList, DistanceFieldNormal);
 			}
@@ -446,8 +433,8 @@ void ComputeDistanceFieldNormal(FRHICommandListImmediate& RHICmdList, const TArr
 				TShaderMapRef<FComputeDistanceFieldNormalPS> PixelShader(View.ShaderMap);
 
 				GraphicsPSOInit.BoundShaderState.VertexDeclarationRHI = GFilterVertexDeclaration.VertexDeclarationRHI;
-				GraphicsPSOInit.BoundShaderState.VertexShaderRHI = GETSAFERHISHADER_VERTEX(*VertexShader);
-				GraphicsPSOInit.BoundShaderState.PixelShaderRHI = GETSAFERHISHADER_PIXEL(*PixelShader);
+				GraphicsPSOInit.BoundShaderState.VertexShaderRHI = VertexShader.GetVertexShader();
+				GraphicsPSOInit.BoundShaderState.PixelShaderRHI = PixelShader.GetPixelShader();
 
 				SetGraphicsPipelineState(RHICmdList, GraphicsPSOInit);
 
@@ -461,7 +448,7 @@ void ComputeDistanceFieldNormal(FRHICommandListImmediate& RHICmdList, const TArr
 					View.ViewRect.Width(), View.ViewRect.Height(),
 					FIntPoint(View.ViewRect.Width() / GAODownsampleFactor, View.ViewRect.Height() / GAODownsampleFactor),
 					FSceneRenderTargets::Get(RHICmdList).GetBufferSizeXY(),
-					*VertexShader);
+					VertexShader);
 			}
 		}
 		RHICmdList.EndRenderPass();
diff --git a/Engine/Source/Runtime/Renderer/Private/DistanceFieldAmbientOcclusion.h b/Engine/Source/Runtime/Renderer/Private/DistanceFieldAmbientOcclusion.h
--- a/Engine/Source/Runtime/Renderer/Private/DistanceFieldAmbientOcclusion.h
+++ b/Engine/Source/Runtime/Renderer/Private/DistanceFieldAmbientOcclusion.h
@@ -115,6 +115,7 @@ static int32 ConeTraceObjectsThreadGroupSize = 64;
 
 class FTileIntersectionParameters
 {
+	DECLARE_TYPE_LAYOUT(FTileIntersectionParameters, NonVirtual);
 public:
 
 	static void ModifyCompilationEnvironment(EShaderPlatform Platform, FShaderCompilerEnvironment& OutEnvironment)
@@ -201,12 +202,14 @@ public:
 	}
 
 private:
-	FShaderParameter TileListGroupSize;
 	
-	FRWShaderParameter NumCulledTilesArray;
-	FRWShaderParameter CulledTilesStartOffsetArray;
-	FRWShaderParameter CulledTileDataArray;
-	FRWShaderParameter ObjectTilesIndirectArguments;
+		LAYOUT_FIELD(FShaderParameter, TileListGroupSize)
+
+		LAYOUT_FIELD(FRWShaderParameter, NumCulledTilesArray)
+		LAYOUT_FIELD(FRWShaderParameter, CulledTilesStartOffsetArray)
+		LAYOUT_FIELD(FRWShaderParameter, CulledTileDataArray)
+		LAYOUT_FIELD(FRWShaderParameter, ObjectTilesIndirectArguments)
+	
 };
 
 class FAOScreenGridResources : public FRenderResource
@@ -287,6 +290,7 @@ inline float GetMaxAOViewDistance()
 
 class FAOParameters
 {
+	DECLARE_TYPE_LAYOUT(FAOParameters, NonVirtual);
 public:
 	void Bind(const FShaderParameterMap& ParameterMap)
 	{
@@ -329,15 +333,18 @@ public:
 	}
 
 private:
-	FShaderParameter AOObjectMaxDistance;
-	FShaderParameter AOStepScale;
-	FShaderParameter AOStepExponentScale;
-	FShaderParameter AOMaxViewDistance;
-	FShaderParameter AOGlobalMaxOcclusionDistance;
+	
+		LAYOUT_FIELD(FShaderParameter, AOObjectMaxDistance)
+		LAYOUT_FIELD(FShaderParameter, AOStepScale)
+		LAYOUT_FIELD(FShaderParameter, AOStepExponentScale)
+		LAYOUT_FIELD(FShaderParameter, AOMaxViewDistance)
+		LAYOUT_FIELD(FShaderParameter, AOGlobalMaxOcclusionDistance)
+	
 };
 
 class FDFAOUpsampleParameters
 {
+	DECLARE_TYPE_LAYOUT(FDFAOUpsampleParameters, NonVirtual);
 public:
 	void Bind(const FShaderParameterMap& ParameterMap)
 	{
@@ -379,11 +386,11 @@ public:
 	}
 
 private:
-	FShaderResourceParameter BentNormalAOTexture;
-	FShaderResourceParameter BentNormalAOSampler;
-	FShaderParameter AOBufferBilinearUVMax;
-	FShaderParameter DistanceFadeScale;
-	FShaderParameter AOMaxViewDistance;
+	LAYOUT_FIELD(FShaderResourceParameter, BentNormalAOTexture);
+	LAYOUT_FIELD(FShaderResourceParameter, BentNormalAOSampler);
+	LAYOUT_FIELD(FShaderParameter, AOBufferBilinearUVMax);
+	LAYOUT_FIELD(FShaderParameter, DistanceFadeScale);
+	LAYOUT_FIELD(FShaderParameter, AOMaxViewDistance);
 };
 
 class FMaxSizedRWBuffers : public FRenderResource
@@ -475,6 +482,7 @@ public:
 
 class FScreenGridParameters
 {
+	DECLARE_TYPE_LAYOUT(FScreenGridParameters, NonVirtual);
 public:
 	void Bind(const FShaderParameterMap& ParameterMap)
 	{
@@ -516,11 +524,13 @@ public:
 	}
 
 private:
-	FShaderParameter BaseLevelTexelSize;
-	FShaderParameter JitterOffset;
-	FShaderParameter ScreenGridConeVisibilitySize;
-	FShaderResourceParameter DistanceFieldNormalTexture;
-	FShaderResourceParameter DistanceFieldNormalSampler;
+	
+		LAYOUT_FIELD(FShaderParameter, BaseLevelTexelSize)
+		LAYOUT_FIELD(FShaderParameter, JitterOffset)
+		LAYOUT_FIELD(FShaderParameter, ScreenGridConeVisibilitySize)
+		LAYOUT_FIELD(FShaderResourceParameter, DistanceFieldNormalTexture)
+		LAYOUT_FIELD(FShaderResourceParameter, DistanceFieldNormalSampler)
+	
 };
 
 extern void TrackGPUProgress(FRHICommandListImmediate& RHICmdList, uint32 DebugId);
diff --git a/Engine/Source/Runtime/Renderer/Private/DistanceFieldGlobalIllumination.cpp b/Engine/Source/Runtime/Renderer/Private/DistanceFieldGlobalIllumination.cpp
--- a/Engine/Source/Runtime/Renderer/Private/DistanceFieldGlobalIllumination.cpp
+++ b/Engine/Source/Runtime/Renderer/Private/DistanceFieldGlobalIllumination.cpp
@@ -140,7 +140,7 @@ public:
 		const FMatrix& ShadowToWorldValue,
 		FLightTileIntersectionResources* TileIntersectionResources)
 	{
-		FRHIComputeShader* ShaderRHI = GetComputeShader();
+		FRHIComputeShader* ShaderRHI = RHICmdList.GetBoundComputeShader();
 		FGlobalShader::SetParameters<FViewUniformShaderParameters>(RHICmdList, ShaderRHI, View.ViewUniformBuffer);
 
 		FRHITexture* TextureAtlas;
@@ -181,8 +181,8 @@ public:
 
 	void UnsetParameters(FRHICommandList& RHICmdList)
 	{
-		VPLParameterBuffer.UnsetUAV(RHICmdList, GetComputeShader());
-		VPLData.UnsetUAV(RHICmdList, GetComputeShader());
+		VPLParameterBuffer.UnsetUAV(RHICmdList, RHICmdList.GetBoundComputeShader());
+		VPLData.UnsetUAV(RHICmdList, RHICmdList.GetBoundComputeShader());
 
 		FRHIUnorderedAccessView* OutUAVs[2];
 		OutUAVs[0] = GVPLResources.VPLParameterBuffer.UAV;
@@ -190,35 +190,17 @@ public:
 		RHICmdList.TransitionResources(EResourceTransitionAccess::EReadable, EResourceTransitionPipeline::EComputeToCompute, OutUAVs, UE_ARRAY_COUNT(OutUAVs));
 	}
 
-	// FShader interface.
-	virtual bool Serialize(FArchive& Ar) override
-	{
-		bool bShaderHasOutdatedParameters = FGlobalShader::Serialize(Ar);
-		Ar << VPLParameterBuffer;
-		Ar << VPLData;
-		Ar << InvPlacementGridSize;
-		Ar << WorldToShadow;
-		Ar << ShadowToWorld;
-		Ar << LightDirectionAndTraceDistance;
-		Ar << LightColor;
-		Ar << ObjectParameters;
-		Ar << LightTileIntersectionParameters;
-		Ar << VPLPlacementCameraRadius;
-		return bShaderHasOutdatedParameters;
-	}
-
 private:
-
-	FRWShaderParameter VPLParameterBuffer;
-	FRWShaderParameter VPLData;
-	FShaderParameter InvPlacementGridSize;
-	FShaderParameter WorldToShadow;
-	FShaderParameter ShadowToWorld;
-	FShaderParameter LightDirectionAndTraceDistance;
-	FShaderParameter LightColor;
-	TDistanceFieldCulledObjectBufferParameters<DFPT_SignedDistanceField> ObjectParameters;
-	FLightTileIntersectionParameters LightTileIntersectionParameters;
-	FShaderParameter VPLPlacementCameraRadius;
+	LAYOUT_FIELD(FRWShaderParameter, VPLParameterBuffer);
+	LAYOUT_FIELD(FRWShaderParameter, VPLData);
+	LAYOUT_FIELD(FShaderParameter, InvPlacementGridSize);
+	LAYOUT_FIELD(FShaderParameter, WorldToShadow);
+	LAYOUT_FIELD(FShaderParameter, ShadowToWorld);
+	LAYOUT_FIELD(FShaderParameter, LightDirectionAndTraceDistance);
+	LAYOUT_FIELD(FShaderParameter, LightColor);
+	LAYOUT_FIELD((TDistanceFieldCulledObjectBufferParameters<DFPT_SignedDistanceField>), ObjectParameters);
+	LAYOUT_FIELD(FLightTileIntersectionParameters, LightTileIntersectionParameters);
+	LAYOUT_FIELD(FShaderParameter, VPLPlacementCameraRadius);
 };
 
 IMPLEMENT_SHADER_TYPE(,FVPLPlacementCS,TEXT("/Engine/Private/DistanceFieldGlobalIllumination.usf"),TEXT("VPLPlacementCS"),SF_Compute);
@@ -255,7 +237,7 @@ public:
 
 	void SetParameters(FRHICommandList& RHICmdList, const FSceneView& View)
 	{
-		FRHIComputeShader* ShaderRHI = GetComputeShader();
+		FRHIComputeShader* ShaderRHI = RHICmdList.GetBoundComputeShader();
 		FGlobalShader::SetParameters<FViewUniformShaderParameters>(RHICmdList, ShaderRHI, View.ViewUniformBuffer);
 
 		RHICmdList.TransitionResource(EResourceTransitionAccess::ERWBarrier, EResourceTransitionPipeline::EComputeToCompute, GVPLResources.VPLDispatchIndirectBuffer.UAV);
@@ -265,22 +247,14 @@ public:
 
 	void UnsetParameters(FRHICommandList& RHICmdList)
 	{
-		DispatchParameters.UnsetUAV(RHICmdList, GetComputeShader());
+		DispatchParameters.UnsetUAV(RHICmdList, RHICmdList.GetBoundComputeShader());
 		RHICmdList.TransitionResource(EResourceTransitionAccess::EReadable, EResourceTransitionPipeline::EComputeToCompute, GVPLResources.VPLDispatchIndirectBuffer.UAV);
 	}
 
-	virtual bool Serialize(FArchive& Ar) override
-	{		
-		bool bShaderHasOutdatedParameters = FGlobalShader::Serialize(Ar);
-		Ar << DispatchParameters;
-		Ar << VPLParameterBuffer;
-		return bShaderHasOutdatedParameters;
-	}
-
 private:
 
-	FRWShaderParameter DispatchParameters;
-	FShaderResourceParameter VPLParameterBuffer;
+	LAYOUT_FIELD(FRWShaderParameter, DispatchParameters);
+	LAYOUT_FIELD(FShaderResourceParameter, VPLParameterBuffer);
 };
 
 IMPLEMENT_SHADER_TYPE(,FSetupVPLCullndirectArgumentsCS,TEXT("/Engine/Private/DistanceFieldGlobalIllumination.usf"),TEXT("SetupVPLCullndirectArgumentsCS"),SF_Compute);
@@ -320,7 +294,7 @@ public:
 
 	void SetParameters(FRHICommandList& RHICmdList, const FScene* Scene, const FSceneView& View, const FDistanceFieldAOParameters& Parameters)
 	{
-		FRHIComputeShader* ShaderRHI = GetComputeShader();
+		FRHIComputeShader* ShaderRHI = RHICmdList.GetBoundComputeShader();
 		FGlobalShader::SetParameters<FViewUniformShaderParameters>(RHICmdList, ShaderRHI, View.ViewUniformBuffer);
 
 		FRHIUnorderedAccessView* OutUAVs[2];
@@ -343,32 +317,19 @@ public:
 
 	void UnsetParameters(FRHICommandList& RHICmdList)
 	{
-		CulledVPLParameterBuffer.UnsetUAV(RHICmdList, GetComputeShader());
-		CulledVPLData.UnsetUAV(RHICmdList, GetComputeShader());
-	}
-
-	virtual bool Serialize(FArchive& Ar) override
-	{		
-		bool bShaderHasOutdatedParameters = FGlobalShader::Serialize(Ar);
-		Ar << VPLParameterBuffer;
-		Ar << VPLData;
-		Ar << CulledVPLParameterBuffer;
-		Ar << CulledVPLData;
-		Ar << AOParameters;
-		Ar << NumConvexHullPlanes;
-		Ar << ViewFrustumConvexHull;
-		return bShaderHasOutdatedParameters;
+		CulledVPLParameterBuffer.UnsetUAV(RHICmdList, RHICmdList.GetBoundComputeShader());
+		CulledVPLData.UnsetUAV(RHICmdList, RHICmdList.GetBoundComputeShader());
 	}
 
 private:
 
-	FShaderResourceParameter VPLParameterBuffer;
-	FShaderResourceParameter VPLData;
-	FRWShaderParameter CulledVPLParameterBuffer;
-	FRWShaderParameter CulledVPLData;
-	FAOParameters AOParameters;
-	FShaderParameter NumConvexHullPlanes;
-	FShaderParameter ViewFrustumConvexHull;
+	LAYOUT_FIELD(FShaderResourceParameter, VPLParameterBuffer);
+	LAYOUT_FIELD(FShaderResourceParameter, VPLData);
+	LAYOUT_FIELD(FRWShaderParameter, CulledVPLParameterBuffer);
+	LAYOUT_FIELD(FRWShaderParameter, CulledVPLData);
+	LAYOUT_FIELD(FAOParameters, AOParameters);
+	LAYOUT_FIELD(FShaderParameter, NumConvexHullPlanes);
+	LAYOUT_FIELD(FShaderParameter, ViewFrustumConvexHull);
 };
 
 IMPLEMENT_SHADER_TYPE(,FCullVPLsForViewCS,TEXT("/Engine/Private/DistanceFieldGlobalIllumination.usf"),TEXT("CullVPLsForViewCS"),SF_Compute);
@@ -521,9 +482,9 @@ void PlaceVPLs(
 
 			TShaderMapRef<FVPLPlacementCS> ComputeShader(View.ShaderMap);
 
-			RHICmdList.SetComputeShader(ComputeShader->GetComputeShader());
+			RHICmdList.SetComputeShader(ComputeShader.GetComputeShader());
 			ComputeShader->SetParameters(RHICmdList, View, DirectionalLightProxy, FVector2D(1.0f / GVPLGridDimension, 1.0f / GVPLGridDimension), DirectionalLightShadowToWorld, DirectionalLightShadowToWorld.InverseFast(), GVPLPlacementTileIntersectionResources.Get());
-			DispatchComputeShader(RHICmdList, *ComputeShader, FMath::DivideAndRoundUp<int32>(GVPLGridDimension, GDistanceFieldAOTileSizeX), FMath::DivideAndRoundUp<int32>(GVPLGridDimension, GDistanceFieldAOTileSizeY), 1);
+			DispatchComputeShader(RHICmdList, ComputeShader.GetShader(), FMath::DivideAndRoundUp<int32>(GVPLGridDimension, GDistanceFieldAOTileSizeX), FMath::DivideAndRoundUp<int32>(GVPLGridDimension, GDistanceFieldAOTileSizeY), 1);
 
 			ComputeShader->UnsetParameters(RHICmdList);
 		}
@@ -532,10 +493,10 @@ void PlaceVPLs(
 		{
 			{
 				TShaderMapRef<FSetupVPLCullndirectArgumentsCS> ComputeShader(View.ShaderMap);
-				RHICmdList.SetComputeShader(ComputeShader->GetComputeShader());
+				RHICmdList.SetComputeShader(ComputeShader.GetComputeShader());
 				ComputeShader->SetParameters(RHICmdList, View);
 
-				DispatchComputeShader(RHICmdList, *ComputeShader, 1, 1, 1);
+				DispatchComputeShader(RHICmdList, ComputeShader.GetShader(), 1, 1, 1);
 				ComputeShader->UnsetParameters(RHICmdList);
 			}
 
@@ -545,10 +506,10 @@ void PlaceVPLs(
 				RHICmdList.ClearUAVUint(GCulledVPLResources.VPLParameterBuffer.UAV, FUintVector4(0, 0, 0, 0));
 
 				TShaderMapRef<FCullVPLsForViewCS> ComputeShader(GetGlobalShaderMap(Scene->GetFeatureLevel()));
-				RHICmdList.SetComputeShader(ComputeShader->GetComputeShader());
+				RHICmdList.SetComputeShader(ComputeShader.GetComputeShader());
 				ComputeShader->SetParameters(RHICmdList, Scene, View, Parameters);
 
-				DispatchIndirectComputeShader(RHICmdList, *ComputeShader, GVPLResources.VPLDispatchIndirectBuffer.Buffer, 0);
+				DispatchIndirectComputeShader(RHICmdList, ComputeShader.GetShader(), GVPLResources.VPLDispatchIndirectBuffer.Buffer, 0);
 				ComputeShader->UnsetParameters(RHICmdList);
 			}
 		}
@@ -587,7 +548,7 @@ public:
 
 	void SetParameters(FRHICommandList& RHICmdList, const FSceneView& View)
 	{
-		FRHIComputeShader* ShaderRHI = GetComputeShader();
+		FRHIComputeShader* ShaderRHI = RHICmdList.GetBoundComputeShader();
 		FGlobalShader::SetParameters<FViewUniformShaderParameters>(RHICmdList, ShaderRHI, View.ViewUniformBuffer);
 
 		RHICmdList.TransitionResource(EResourceTransitionAccess::ERWBarrier, EResourceTransitionPipeline::EComputeToCompute, GAOCulledObjectBuffers.Buffers.ObjectIndirectDispatch.UAV);
@@ -610,27 +571,18 @@ public:
 
 	void UnsetParameters(FRHICommandList& RHICmdList)
 	{
-		FRHIComputeShader* ShaderRHI = GetComputeShader();
+		FRHIComputeShader* ShaderRHI = RHICmdList.GetBoundComputeShader();
 		DispatchParameters.UnsetUAV(RHICmdList, ShaderRHI);
 
 		ObjectParameters.UnsetParameters(RHICmdList, ShaderRHI);
 		RHICmdList.TransitionResource(EResourceTransitionAccess::EReadable, EResourceTransitionPipeline::EComputeToCompute, GAOCulledObjectBuffers.Buffers.ObjectIndirectDispatch.UAV);
 	}
 
-	virtual bool Serialize(FArchive& Ar) override
-	{		
-		bool bShaderHasOutdatedParameters = FGlobalShader::Serialize(Ar);
-		Ar << DispatchParameters;
-		Ar << ObjectParameters;
-		Ar << ObjectProcessStride;
-		return bShaderHasOutdatedParameters;
-	}
-
 private:
 
-	FRWShaderParameter DispatchParameters;
-	TDistanceFieldCulledObjectBufferParameters<DFPT_SignedDistanceField> ObjectParameters;
-	FShaderParameter ObjectProcessStride;
+	LAYOUT_FIELD(FRWShaderParameter, DispatchParameters);
+	LAYOUT_FIELD((TDistanceFieldCulledObjectBufferParameters<DFPT_SignedDistanceField>), ObjectParameters);
+	LAYOUT_FIELD(FShaderParameter, ObjectProcessStride);
 };
 
 IMPLEMENT_SHADER_TYPE(,FSetupLightVPLsIndirectArgumentsCS,TEXT("/Engine/Private/DistanceFieldGlobalIllumination.usf"),TEXT("SetupLightVPLsIndirectArgumentsCS"),SF_Compute);
@@ -683,7 +635,7 @@ public:
 		const FDistanceFieldAOParameters& Parameters,
 		FLightTileIntersectionResources* TileIntersectionResources)
 	{
-		FRHIComputeShader* ShaderRHI = GetComputeShader();
+		FRHIComputeShader* ShaderRHI = RHICmdList.GetBoundComputeShader();
 
 		FGlobalShader::SetParameters<FViewUniformShaderParameters>(RHICmdList, ShaderRHI, View.ViewUniformBuffer);
 		AOParameters.Set(RHICmdList, ShaderRHI, Parameters);
@@ -736,48 +688,25 @@ public:
 
 	void UnsetParameters(FRHICommandList& RHICmdList)
 	{
-		SurfelParameters.UnsetParameters(RHICmdList, GetComputeShader());
-	}
-
-	// FShader interface.
-	virtual bool Serialize(FArchive& Ar) override
-	{
-		bool bShaderHasOutdatedParameters = FGlobalShader::Serialize(Ar);
-		Ar << AOParameters;
-		Ar << LightDirection;
-		Ar << LightPositionAndInvRadius;
-		Ar << LightSourceRadius;
-		Ar << TanLightAngleAndNormalThreshold;
-		Ar << LightColor;
-		Ar << ObjectParameters;
-		Ar << SurfelParameters;
-		Ar << LightTileIntersectionParameters;
-		Ar << WorldToShadow;
-		Ar << ShadowObjectIndirectArguments;
-		Ar << ShadowCulledObjectBounds;
-		Ar << ShadowCulledObjectData;
-		Ar << ObjectProcessStride;
-		Ar << ObjectProcessStartIndex;
-		return bShaderHasOutdatedParameters;
+		SurfelParameters.UnsetParameters(RHICmdList, RHICmdList.GetBoundComputeShader());
 	}
 
 private:
-
-	FAOParameters AOParameters;
-	FShaderParameter LightDirection;
-	FShaderParameter LightPositionAndInvRadius;
-	FShaderParameter LightSourceRadius;
-	FShaderParameter TanLightAngleAndNormalThreshold;
-	FShaderParameter LightColor;
-	TDistanceFieldCulledObjectBufferParameters<DFPT_SignedDistanceField> ObjectParameters;
-	FSurfelBufferParameters SurfelParameters;
-	FLightTileIntersectionParameters LightTileIntersectionParameters;
-	FShaderParameter WorldToShadow;
-	FShaderResourceParameter ShadowObjectIndirectArguments;
-	FShaderResourceParameter ShadowCulledObjectBounds;
-	FShaderResourceParameter ShadowCulledObjectData;
-	FShaderParameter ObjectProcessStride;
-	FShaderParameter ObjectProcessStartIndex;
+	LAYOUT_FIELD(FAOParameters, AOParameters);
+	LAYOUT_FIELD(FShaderParameter, LightDirection);
+	LAYOUT_FIELD(FShaderParameter, LightPositionAndInvRadius);
+	LAYOUT_FIELD(FShaderParameter, LightSourceRadius);
+	LAYOUT_FIELD(FShaderParameter, TanLightAngleAndNormalThreshold);
+	LAYOUT_FIELD(FShaderParameter, LightColor);
+	LAYOUT_FIELD((TDistanceFieldCulledObjectBufferParameters<DFPT_SignedDistanceField>), ObjectParameters);
+	LAYOUT_FIELD(FSurfelBufferParameters, SurfelParameters);
+	LAYOUT_FIELD(FLightTileIntersectionParameters, LightTileIntersectionParameters);
+	LAYOUT_FIELD(FShaderParameter, WorldToShadow);
+	LAYOUT_FIELD(FShaderResourceParameter, ShadowObjectIndirectArguments);
+	LAYOUT_FIELD(FShaderResourceParameter, ShadowCulledObjectBounds);
+	LAYOUT_FIELD(FShaderResourceParameter, ShadowCulledObjectData);
+	LAYOUT_FIELD(FShaderParameter, ObjectProcessStride);
+	LAYOUT_FIELD(FShaderParameter, ObjectProcessStartIndex);
 };
 
 IMPLEMENT_SHADER_TYPE(,FLightVPLsCS,TEXT("/Engine/Private/DistanceFieldGlobalIllumination.usf"),TEXT("LightVPLsCS"),SF_Compute);
@@ -873,18 +802,18 @@ void UpdateVPLs(
 
 				{
 					TShaderMapRef<FSetupLightVPLsIndirectArgumentsCS> ComputeShader(View.ShaderMap);
-					RHICmdList.SetComputeShader(ComputeShader->GetComputeShader());
+					RHICmdList.SetComputeShader(ComputeShader.GetComputeShader());
 					ComputeShader->SetParameters(RHICmdList, View);
 
-					DispatchComputeShader(RHICmdList, *ComputeShader, 1, 1, 1);
+					DispatchComputeShader(RHICmdList, ComputeShader.GetShader(), 1, 1, 1);
 					ComputeShader->UnsetParameters(RHICmdList);
 				}
 
 				{
 					TShaderMapRef<FLightVPLsCS> ComputeShader(View.ShaderMap);
-					RHICmdList.SetComputeShader(ComputeShader->GetComputeShader());
+					RHICmdList.SetComputeShader(ComputeShader.GetComputeShader());
 					ComputeShader->SetParameters(RHICmdList, View, DirectionalLightProxy, DirectionalLightWorldToShadow, Parameters, GVPLPlacementTileIntersectionResources.Get());
-					DispatchIndirectComputeShader(RHICmdList, *ComputeShader, GAOCulledObjectBuffers.Buffers.ObjectIndirectDispatch.Buffer, 0);
+					DispatchIndirectComputeShader(RHICmdList, ComputeShader.GetShader(), GAOCulledObjectBuffers.Buffers.ObjectIndirectDispatch.Buffer, 0);
 					ComputeShader->UnsetParameters(RHICmdList);
 				}
 			}
@@ -942,7 +871,7 @@ public:
 		FSceneRenderTargetItem& DistanceFieldNormal, 
 		const FAOScreenGridResources& ScreenGridResources)
 	{
-		FRHIComputeShader* ShaderRHI = GetComputeShader();
+		FRHIComputeShader* ShaderRHI = RHICmdList.GetBoundComputeShader();
 		FGlobalShader::SetParameters<FViewUniformShaderParameters>(RHICmdList, ShaderRHI, View.ViewUniformBuffer);
 		ScreenGridParameters.Set(RHICmdList, ShaderRHI, View, DistanceFieldNormal);
 
@@ -976,26 +905,16 @@ public:
 
 	void UnsetParameters(FRHICommandList& RHICmdList, const FAOScreenGridResources& ScreenGridResources)
 	{
-		StepBentNormal.UnsetUAV(RHICmdList, GetComputeShader());
+		StepBentNormal.UnsetUAV(RHICmdList, RHICmdList.GetBoundComputeShader());
 		RHICmdList.TransitionResource(EResourceTransitionAccess::EReadable, EResourceTransitionPipeline::EComputeToCompute, ScreenGridResources.StepBentNormal.UAV);
 	}
 
-	virtual bool Serialize(FArchive& Ar) override
-	{		
-		bool bShaderHasOutdatedParameters = FGlobalShader::Serialize(Ar);
-		Ar << ScreenGridParameters;
-		Ar << BentNormalNormalizeFactor;
-		Ar << ConeDepthVisibilityFunction;
-		Ar << StepBentNormal;
-		return bShaderHasOutdatedParameters;
-	}
-
 private:
 
-	FScreenGridParameters ScreenGridParameters;
-	FShaderParameter BentNormalNormalizeFactor;
-	FShaderResourceParameter ConeDepthVisibilityFunction;
-	FRWShaderParameter StepBentNormal;
+	LAYOUT_FIELD(FScreenGridParameters, ScreenGridParameters);
+	LAYOUT_FIELD(FShaderParameter, BentNormalNormalizeFactor);
+	LAYOUT_FIELD(FShaderResourceParameter, ConeDepthVisibilityFunction);
+	LAYOUT_FIELD(FRWShaderParameter, StepBentNormal);
 };
 
 IMPLEMENT_SHADER_TYPE(,FComputeStepBentNormalScreenGridCS,TEXT("/Engine/Private/DistanceFieldGlobalIllumination.usf"),TEXT("ComputeStepBentNormalScreenGridCS"),SF_Compute);
@@ -1048,7 +967,7 @@ public:
 		FSceneRenderTargetItem& DistanceFieldNormal, 
 		const FDistanceFieldAOParameters& Parameters)
 	{
-		FRHIComputeShader* ShaderRHI = GetComputeShader();
+		FRHIComputeShader* ShaderRHI = RHICmdList.GetBoundComputeShader();
 		FGlobalShader::SetParameters<FViewUniformShaderParameters>(RHICmdList, ShaderRHI, View.ViewUniformBuffer);
 		SceneTextureParameters.Set(RHICmdList, ShaderRHI, View.FeatureLevel, ESceneTextureSetupMode::All);
 
@@ -1088,40 +1007,23 @@ public:
 
 	void UnsetParameters(FRHICommandList& RHICmdList, const FViewInfo& View)
 	{
-		SurfelIrradiance.UnsetUAV(RHICmdList, GetComputeShader());
+		SurfelIrradiance.UnsetUAV(RHICmdList, RHICmdList.GetBoundComputeShader());
 
 		FAOScreenGridResources* ScreenGridResources = View.ViewState->AOScreenGridResources;
 		RHICmdList.TransitionResource(EResourceTransitionAccess::EReadable, EResourceTransitionPipeline::EComputeToCompute, ScreenGridResources->SurfelIrradiance.UAV);
 	}
 
-	virtual bool Serialize(FArchive& Ar) override
-	{		
-		bool bShaderHasOutdatedParameters = FGlobalShader::Serialize(Ar);
-		Ar << SceneTextureParameters;
-		Ar << ObjectParameters;
-		Ar << AOParameters;
-		Ar << ScreenGridParameters;
-		Ar << SurfelParameters;
-		Ar << TileConeDepthRanges;
-		Ar << TileListGroupSize;
-		Ar << VPLGatherRadius;
-		Ar << StepBentNormalBuffer;
-		Ar << SurfelIrradiance;
-		return bShaderHasOutdatedParameters;
-	}
-
 private:
-
-	FSceneTextureShaderParameters SceneTextureParameters;
-	TDistanceFieldCulledObjectBufferParameters<DFPT_SignedDistanceField> ObjectParameters;
-	FAOParameters AOParameters;
-	FScreenGridParameters ScreenGridParameters;
-	FSurfelBufferParameters SurfelParameters;
-	FShaderResourceParameter TileConeDepthRanges;
-	FShaderParameter TileListGroupSize;
-	FShaderParameter VPLGatherRadius;
-	FShaderResourceParameter StepBentNormalBuffer;
-	FRWShaderParameter SurfelIrradiance;
+	LAYOUT_FIELD(FSceneTextureShaderParameters, SceneTextureParameters);
+	LAYOUT_FIELD((TDistanceFieldCulledObjectBufferParameters<DFPT_SignedDistanceField>), ObjectParameters);
+	LAYOUT_FIELD(FAOParameters, AOParameters);
+	LAYOUT_FIELD(FScreenGridParameters, ScreenGridParameters);
+	LAYOUT_FIELD(FSurfelBufferParameters, SurfelParameters);
+	LAYOUT_FIELD(FShaderResourceParameter, TileConeDepthRanges);
+	LAYOUT_FIELD(FShaderParameter, TileListGroupSize);
+	LAYOUT_FIELD(FShaderParameter, VPLGatherRadius);
+	LAYOUT_FIELD(FShaderResourceParameter, StepBentNormalBuffer);
+	LAYOUT_FIELD(FRWShaderParameter, SurfelIrradiance);
 };
 
 IMPLEMENT_SHADER_TYPE(,FComputeIrradianceScreenGridCS,TEXT("/Engine/Private/DistanceFieldGlobalIllumination.usf"),TEXT("ComputeIrradianceScreenGridCS"),SF_Compute);
@@ -1162,7 +1064,7 @@ public:
 		const FAOScreenGridResources& ScreenGridResources,
 		FSceneRenderTargetItem& IrradianceTextureValue)
 	{
-		FRHIComputeShader* ShaderRHI = GetComputeShader();
+		FRHIComputeShader* ShaderRHI = RHICmdList.GetBoundComputeShader();
 		FGlobalShader::SetParameters<FViewUniformShaderParameters>(RHICmdList, ShaderRHI, View.ViewUniformBuffer);
 
 		SetSRVParameter(RHICmdList, ShaderRHI, SurfelIrradiance, ScreenGridResources.SurfelIrradiance.SRV);
@@ -1176,26 +1078,16 @@ public:
 
 	void UnsetParameters(FRHICommandList& RHICmdList, FSceneRenderTargetItem& IrradianceTextureValue)
 	{
-		IrradianceTexture.UnsetUAV(RHICmdList, GetComputeShader());
+		IrradianceTexture.UnsetUAV(RHICmdList, RHICmdList.GetBoundComputeShader());
 		RHICmdList.TransitionResource(EResourceTransitionAccess::EReadable, EResourceTransitionPipeline::EComputeToCompute, IrradianceTextureValue.UAV);
 	}
 
-	virtual bool Serialize(FArchive& Ar) override
-	{		
-		bool bShaderHasOutdatedParameters = FGlobalShader::Serialize(Ar);
-		Ar << IrradianceTexture;
-		Ar << SurfelIrradiance;
-		Ar << HeightfieldIrradiance;
-		Ar << ScreenGridConeVisibilitySize;
-		return bShaderHasOutdatedParameters;
-	}
-
 private:
 
-	FRWShaderParameter IrradianceTexture;
-	FShaderResourceParameter SurfelIrradiance;
-	FShaderResourceParameter HeightfieldIrradiance;
-	FShaderParameter ScreenGridConeVisibilitySize;
+	LAYOUT_FIELD(FRWShaderParameter, IrradianceTexture);
+	LAYOUT_FIELD(FShaderResourceParameter, SurfelIrradiance);
+	LAYOUT_FIELD(FShaderResourceParameter, HeightfieldIrradiance);
+	LAYOUT_FIELD(FShaderParameter, ScreenGridConeVisibilitySize);
 };
 
 IMPLEMENT_SHADER_TYPE(,FCombineIrradianceScreenGridCS,TEXT("/Engine/Private/DistanceFieldGlobalIllumination.usf"),TEXT("CombineIrradianceScreenGridCS"),SF_Compute);
@@ -1225,9 +1117,9 @@ void ComputeIrradianceForScreenGrid(
 			SCOPED_DRAW_EVENT(RHICmdList, ComputeStepBentNormal);
 
 			TShaderMapRef<FComputeStepBentNormalScreenGridCS> ComputeShader(View.ShaderMap);
-			RHICmdList.SetComputeShader(ComputeShader->GetComputeShader());
+			RHICmdList.SetComputeShader(ComputeShader.GetComputeShader());
 			ComputeShader->SetParameters(RHICmdList, View, DistanceFieldNormal, ScreenGridResources);
-			DispatchComputeShader(RHICmdList, *ComputeShader, GroupSizeX, GroupSizeY, 1);
+			DispatchComputeShader(RHICmdList, ComputeShader.GetShader(), GroupSizeX, GroupSizeY, 1);
 			ComputeShader->UnsetParameters(RHICmdList, ScreenGridResources);
 		}
 
@@ -1236,21 +1128,21 @@ void ComputeIrradianceForScreenGrid(
 			SCOPED_DRAW_EVENT(RHICmdList, MeshIrradiance);
 
 			TShaderMapRef<FComputeIrradianceScreenGridCS> ComputeShader(View.ShaderMap);
-			RHICmdList.SetComputeShader(ComputeShader->GetComputeShader());
+			RHICmdList.SetComputeShader(ComputeShader.GetComputeShader());
 			ComputeShader->SetParameters(RHICmdList, View, DistanceFieldNormal, Parameters);
 
 			const uint32 ComputeIrradianceGroupSizeX = View.ViewRect.Size().X / GAODownsampleFactor;
 			const uint32 ComputeIrradianceGroupSizeY = View.ViewRect.Size().Y / GAODownsampleFactor;
-			DispatchComputeShader(RHICmdList, *ComputeShader, ComputeIrradianceGroupSizeX, ComputeIrradianceGroupSizeY, 1);
+			DispatchComputeShader(RHICmdList, ComputeShader.GetShader(), ComputeIrradianceGroupSizeX, ComputeIrradianceGroupSizeY, 1);
 			ComputeShader->UnsetParameters(RHICmdList, View);
 		}
 	}
 		
 	{
 		TShaderMapRef<FCombineIrradianceScreenGridCS> ComputeShader(View.ShaderMap);
-		RHICmdList.SetComputeShader(ComputeShader->GetComputeShader());
+		RHICmdList.SetComputeShader(ComputeShader.GetComputeShader());
 		ComputeShader->SetParameters(RHICmdList, View, ScreenGridResources, IrradianceTexture);
-		DispatchComputeShader(RHICmdList, *ComputeShader, GroupSizeX, GroupSizeY, 1);
+		DispatchComputeShader(RHICmdList, ComputeShader.GetShader(), GroupSizeX, GroupSizeY, 1);
 		ComputeShader->UnsetParameters(RHICmdList, IrradianceTexture);
 	}
 }
diff --git a/Engine/Source/Runtime/Renderer/Private/DistanceFieldLightingPost.cpp b/Engine/Source/Runtime/Renderer/Private/DistanceFieldLightingPost.cpp
--- a/Engine/Source/Runtime/Renderer/Private/DistanceFieldLightingPost.cpp
+++ b/Engine/Source/Runtime/Renderer/Private/DistanceFieldLightingPost.cpp
@@ -12,6 +12,9 @@
 #include "CompositionLighting/PostProcessAmbientOcclusion.h"
 #include "PipelineStateCache.h"
 
+IMPLEMENT_TYPE_LAYOUT(FSurfelBufferParameters);
+IMPLEMENT_TYPE_LAYOUT(FLightTileIntersectionParameters);
+
 int32 GAOUseHistory = 1;
 FAutoConsoleVariableRef CVarAOUseHistory(
 	TEXT("r.AOUseHistory"),
@@ -68,6 +71,7 @@ bool UseAOHistoryStabilityPass()
 
 class FGeometryAwareUpsampleParameters
 {
+	DECLARE_INLINE_TYPE_LAYOUT(FGeometryAwareUpsampleParameters, NonVirtual);
 public:
 	void Bind(const FShaderParameterMap& ParameterMap)
 	{
@@ -146,16 +150,18 @@ public:
 	}
 
 private:
-	FShaderResourceParameter DistanceFieldNormalTexture;
-	FShaderResourceParameter DistanceFieldNormalSampler;
-	FShaderResourceParameter BentNormalAOTexture;
-	FShaderResourceParameter BentNormalAOSampler;
-	FShaderParameter DistanceFieldGBufferTexelSize;
-	FShaderParameter DistanceFieldGBufferJitterOffset;
-	FShaderParameter BentNormalBufferAndTexelSize;
-	FShaderParameter MinDownsampleFactorToBaseLevel;
-	FShaderParameter DistanceFadeScale;
-	FShaderParameter JitterOffset;
+	
+		LAYOUT_FIELD(FShaderResourceParameter, DistanceFieldNormalTexture)
+		LAYOUT_FIELD(FShaderResourceParameter, DistanceFieldNormalSampler)
+		LAYOUT_FIELD(FShaderResourceParameter, BentNormalAOTexture)
+		LAYOUT_FIELD(FShaderResourceParameter, BentNormalAOSampler)
+		LAYOUT_FIELD(FShaderParameter, DistanceFieldGBufferTexelSize)
+		LAYOUT_FIELD(FShaderParameter, DistanceFieldGBufferJitterOffset)
+		LAYOUT_FIELD(FShaderParameter, BentNormalBufferAndTexelSize)
+		LAYOUT_FIELD(FShaderParameter, MinDownsampleFactorToBaseLevel)
+		LAYOUT_FIELD(FShaderParameter, DistanceFadeScale)
+		LAYOUT_FIELD(FShaderParameter, JitterOffset)
+	
 };
 
 class FUpdateHistoryDepthRejectionPS : public FGlobalShader
@@ -204,7 +210,7 @@ public:
 		IPooledRenderTarget* VelocityTextureValue,
 		const FDistanceFieldAOParameters& Parameters)
 	{
-		FRHIPixelShader* ShaderRHI = GetPixelShader();
+		FRHIPixelShader* ShaderRHI = RHICmdList.GetBoundPixelShader();
 
 		FGlobalShader::SetParameters<FViewUniformShaderParameters>(RHICmdList, ShaderRHI, View.ViewUniformBuffer);
 		SceneTextureParameters.Set(RHICmdList, ShaderRHI, View.FeatureLevel, ESceneTextureSetupMode::All);
@@ -256,39 +262,20 @@ public:
 			SetShaderValue(RHICmdList, ShaderRHI, HistoryUVMinMax, HistoryUVMinMaxValue);
 		}
 	}
-	// FShader interface.
-	virtual bool Serialize(FArchive& Ar)
-	{
-		bool bShaderHasOutdatedParameters = FGlobalShader::Serialize(Ar);
-		Ar << SceneTextureParameters;
-		Ar << AOParameters;
-		Ar << GeometryAwareUpsampleParameters;
-		Ar << BentNormalHistoryTexture;
-		Ar << BentNormalHistorySampler;
-		Ar << HistoryWeight;
-		Ar << HistoryDistanceThreshold;
-		Ar << UseHistoryFilter;
-		Ar << VelocityTexture;
-		Ar << VelocityTextureSampler;
-		Ar << HistoryScreenPositionScaleBias;
-		Ar << HistoryUVMinMax;
-		return bShaderHasOutdatedParameters;
-	}
 
 private:
-
-	FSceneTextureShaderParameters SceneTextureParameters;
-	FAOParameters AOParameters;
-	FGeometryAwareUpsampleParameters GeometryAwareUpsampleParameters;
-	FShaderResourceParameter BentNormalHistoryTexture;
-	FShaderResourceParameter BentNormalHistorySampler;
-	FShaderParameter HistoryWeight;
-	FShaderParameter HistoryDistanceThreshold;
-	FShaderParameter UseHistoryFilter;
-	FShaderResourceParameter VelocityTexture;
-	FShaderResourceParameter VelocityTextureSampler;
-	FShaderParameter HistoryScreenPositionScaleBias;
-	FShaderParameter HistoryUVMinMax;
+	LAYOUT_FIELD(FSceneTextureShaderParameters, SceneTextureParameters);
+	LAYOUT_FIELD(FAOParameters, AOParameters);
+	LAYOUT_FIELD(FGeometryAwareUpsampleParameters, GeometryAwareUpsampleParameters);
+	LAYOUT_FIELD(FShaderResourceParameter, BentNormalHistoryTexture);
+	LAYOUT_FIELD(FShaderResourceParameter, BentNormalHistorySampler);
+	LAYOUT_FIELD(FShaderParameter, HistoryWeight);
+	LAYOUT_FIELD(FShaderParameter, HistoryDistanceThreshold);
+	LAYOUT_FIELD(FShaderParameter, UseHistoryFilter);
+	LAYOUT_FIELD(FShaderResourceParameter, VelocityTexture);
+	LAYOUT_FIELD(FShaderResourceParameter, VelocityTextureSampler);
+	LAYOUT_FIELD(FShaderParameter, HistoryScreenPositionScaleBias);
+	LAYOUT_FIELD(FShaderParameter, HistoryUVMinMax);
 };
 
 IMPLEMENT_SHADER_TYPE(, FUpdateHistoryDepthRejectionPS,TEXT("/Engine/Private/DistanceFieldLightingPost.usf"),TEXT("UpdateHistoryDepthRejectionPS"),SF_Pixel);
@@ -333,7 +320,7 @@ public:
 		FSceneRenderTargetItem& DistanceFieldNormal, 
 		FSceneRenderTargetItem& BentNormalHistoryTextureValue)
 	{
-		FRHIPixelShader* ShaderRHI = GetPixelShader();
+		FRHIPixelShader* ShaderRHI = RHICmdList.GetBoundPixelShader();
 
 		FGlobalShader::SetParameters<FViewUniformShaderParameters>(RHICmdList, ShaderRHI, View.ViewUniformBuffer);
 
@@ -371,29 +358,15 @@ public:
 		}
 	}
 
-	// FShader interface.
-	virtual bool Serialize(FArchive& Ar)
-	{
-		bool bShaderHasOutdatedParameters = FGlobalShader::Serialize(Ar);
-		Ar << BentNormalAOTexture;
-		Ar << BentNormalAOSampler;
-		Ar << HistoryWeight;
-		Ar << BentNormalAOTexelSize;
-		Ar << MaxSampleBufferUV;
-		Ar << DistanceFieldNormalTexture;
-		Ar << DistanceFieldNormalSampler;
-		return bShaderHasOutdatedParameters;
-	}
-
 private:
 
-	FShaderResourceParameter BentNormalAOTexture;
-	FShaderResourceParameter BentNormalAOSampler;
-	FShaderParameter HistoryWeight;
-	FShaderParameter BentNormalAOTexelSize;
-	FShaderParameter MaxSampleBufferUV;
-	FShaderResourceParameter DistanceFieldNormalTexture;
-	FShaderResourceParameter DistanceFieldNormalSampler;
+	LAYOUT_FIELD(FShaderResourceParameter, BentNormalAOTexture);
+	LAYOUT_FIELD(FShaderResourceParameter, BentNormalAOSampler);
+	LAYOUT_FIELD(FShaderParameter, HistoryWeight);
+	LAYOUT_FIELD(FShaderParameter, BentNormalAOTexelSize);
+	LAYOUT_FIELD(FShaderParameter, MaxSampleBufferUV);
+	LAYOUT_FIELD(FShaderResourceParameter, DistanceFieldNormalTexture);
+	LAYOUT_FIELD(FShaderResourceParameter, DistanceFieldNormalSampler);
 };
 
 
@@ -439,7 +412,7 @@ public:
 		FSceneRenderTargetItem& DistanceFieldAOBentNormal,
 		const FDistanceFieldAOParameters& Parameters)
 	{
-		FRHIPixelShader* ShaderRHI = GetPixelShader();
+		FRHIPixelShader* ShaderRHI = RHICmdList.GetBoundPixelShader();
 
 		FGlobalShader::SetParameters<FViewUniformShaderParameters>(RHICmdList, ShaderRHI, View.ViewUniformBuffer);
 
@@ -447,18 +420,9 @@ public:
 		GeometryAwareUpsampleParameters.Set(RHICmdList, ShaderRHI, View, DistanceFieldNormal, DistanceFieldAOBentNormal);
 	}
 
-	// FShader interface.
-	virtual bool Serialize(FArchive& Ar) override
-	{
-		bool bShaderHasOutdatedParameters = FGlobalShader::Serialize(Ar);
-		Ar << AOParameters;
-		Ar << GeometryAwareUpsampleParameters;
-		return bShaderHasOutdatedParameters;
-	}
-
 private:
-	FAOParameters AOParameters;
-	FGeometryAwareUpsampleParameters GeometryAwareUpsampleParameters;
+	LAYOUT_FIELD(FAOParameters, AOParameters);
+	LAYOUT_FIELD(FGeometryAwareUpsampleParameters, GeometryAwareUpsampleParameters);
 };
 
 IMPLEMENT_SHADER_TYPE(,FGeometryAwareUpsamplePS, TEXT("/Engine/Private/DistanceFieldLightingPost.usf"), TEXT("GeometryAwareUpsamplePS"), SF_Pixel);
@@ -494,11 +458,11 @@ void GeometryAwareUpsample(FRHICommandList& RHICmdList, const FViewInfo& View, T
 
 		TShaderMapRef<FPostProcessVS> VertexShader(View.ShaderMap);
 		GraphicsPSOInit.BoundShaderState.VertexDeclarationRHI = GFilterVertexDeclaration.VertexDeclarationRHI;
-		GraphicsPSOInit.BoundShaderState.VertexShaderRHI = GETSAFERHISHADER_VERTEX(*VertexShader);
+		GraphicsPSOInit.BoundShaderState.VertexShaderRHI = VertexShader.GetVertexShader();
 		GraphicsPSOInit.PrimitiveType = PT_TriangleList;
 
 		TShaderMapRef<FGeometryAwareUpsamplePS> PixelShader(View.ShaderMap);
-		GraphicsPSOInit.BoundShaderState.PixelShaderRHI = GETSAFERHISHADER_PIXEL(*PixelShader);
+		GraphicsPSOInit.BoundShaderState.PixelShaderRHI = PixelShader.GetPixelShader();
 		SetGraphicsPipelineState(RHICmdList, GraphicsPSOInit);
 		PixelShader->SetParameters(RHICmdList, View, DistanceFieldNormal, BentNormalInterpolation, Parameters);
 
@@ -512,7 +476,7 @@ void GeometryAwareUpsample(FRHICommandList& RHICmdList, const FViewInfo& View, T
 			View.ViewRect.Width() / GAODownsampleFactor, View.ViewRect.Height() / GAODownsampleFactor,
 			FIntPoint(View.ViewRect.Width() / GAODownsampleFactor, View.ViewRect.Height() / GAODownsampleFactor),
 			SceneContext.GetBufferSizeXY() / FIntPoint(GAODownsampleFactor, GAODownsampleFactor),
-			*VertexShader);
+			VertexShader);
 	}
 	RHICmdList.EndRenderPass();
 	RHICmdList.CopyToResolveTarget(DistanceFieldAOBentNormal->GetRenderTargetItem().TargetableTexture, DistanceFieldAOBentNormal->GetRenderTargetItem().ShaderResourceTexture, FResolveParams());
@@ -573,12 +537,12 @@ void UpdateHistory(
 
 					TShaderMapRef<FPostProcessVS> VertexShader(View.ShaderMap);
 					GraphicsPSOInit.BoundShaderState.VertexDeclarationRHI = GFilterVertexDeclaration.VertexDeclarationRHI;
-					GraphicsPSOInit.BoundShaderState.VertexShaderRHI = GETSAFERHISHADER_VERTEX(*VertexShader);
+					GraphicsPSOInit.BoundShaderState.VertexShaderRHI = VertexShader.GetVertexShader();
 					GraphicsPSOInit.PrimitiveType = PT_TriangleList;
 
 					TShaderMapRef<FUpdateHistoryDepthRejectionPS> PixelShader(View.ShaderMap);
 
-					GraphicsPSOInit.BoundShaderState.PixelShaderRHI = GETSAFERHISHADER_PIXEL(*PixelShader);
+					GraphicsPSOInit.BoundShaderState.PixelShaderRHI = PixelShader.GetPixelShader();
 					SetGraphicsPipelineState(RHICmdList, GraphicsPSOInit);
 					PixelShader->SetParameters(RHICmdList, View, *DistanceFieldAOHistoryViewRect, DistanceFieldNormal, BentNormalInterpolation, (*BentNormalHistoryState)->GetRenderTargetItem(), VelocityTexture, Parameters);
 
@@ -592,7 +556,7 @@ void UpdateHistory(
 						View.ViewRect.Width() / GAODownsampleFactor, View.ViewRect.Height() / GAODownsampleFactor,
 						FIntPoint(View.ViewRect.Width() / GAODownsampleFactor, View.ViewRect.Height() / GAODownsampleFactor),
 						SceneContext.GetBufferSizeXY() / FIntPoint(GAODownsampleFactor, GAODownsampleFactor),
-						*VertexShader);
+						VertexShader);
 				}
 				RHICmdList.EndRenderPass();
 				RHICmdList.CopyToResolveTarget(NewBentNormalHistory->GetRenderTargetItem().TargetableTexture, NewBentNormalHistory->GetRenderTargetItem().ShaderResourceTexture, FResolveParams());
@@ -626,14 +590,14 @@ void UpdateHistory(
 
 					TShaderMapRef<FPostProcessVS> VertexShader(View.ShaderMap);
 					GraphicsPSOInit.BoundShaderState.VertexDeclarationRHI = GFilterVertexDeclaration.VertexDeclarationRHI;
-					GraphicsPSOInit.BoundShaderState.VertexShaderRHI = GETSAFERHISHADER_VERTEX(*VertexShader);
+					GraphicsPSOInit.BoundShaderState.VertexShaderRHI = VertexShader.GetVertexShader();
 					GraphicsPSOInit.PrimitiveType = PT_TriangleList;
 
 					if (View.ViewRect.Min == FIntPoint::ZeroValue && View.ViewRect.Max == SceneContext.GetBufferSizeXY())
 					{
 						TShaderMapRef<TFilterHistoryPS<false> > PixelShader(View.ShaderMap);
 
-						GraphicsPSOInit.BoundShaderState.PixelShaderRHI = GETSAFERHISHADER_PIXEL(*PixelShader);
+						GraphicsPSOInit.BoundShaderState.PixelShaderRHI = PixelShader.GetPixelShader();
 						SetGraphicsPipelineState(RHICmdList, GraphicsPSOInit);
 						PixelShader->SetParameters(RHICmdList, View, DistanceFieldNormal, NewBentNormalHistory->GetRenderTargetItem());
 					}
@@ -641,7 +605,7 @@ void UpdateHistory(
 					{
 						TShaderMapRef<TFilterHistoryPS<true> > PixelShader(View.ShaderMap);
 
-						GraphicsPSOInit.BoundShaderState.PixelShaderRHI = GETSAFERHISHADER_PIXEL(*PixelShader);
+						GraphicsPSOInit.BoundShaderState.PixelShaderRHI = PixelShader.GetPixelShader();
 						SetGraphicsPipelineState(RHICmdList, GraphicsPSOInit);
 						PixelShader->SetParameters(RHICmdList, View, DistanceFieldNormal, NewBentNormalHistory->GetRenderTargetItem());
 					}
@@ -656,7 +620,7 @@ void UpdateHistory(
 						View.ViewRect.Width() / GAODownsampleFactor, View.ViewRect.Height() / GAODownsampleFactor,
 						FIntPoint(View.ViewRect.Width() / GAODownsampleFactor, View.ViewRect.Height() / GAODownsampleFactor),
 						SceneContext.GetBufferSizeXY() / FIntPoint(GAODownsampleFactor, GAODownsampleFactor),
-						*VertexShader);
+						VertexShader);
 				}
 				RHICmdList.EndRenderPass();
 				RHICmdList.CopyToResolveTarget((*BentNormalHistoryState)->GetRenderTargetItem().TargetableTexture, (*BentNormalHistoryState)->GetRenderTargetItem().ShaderResourceTexture, FResolveParams());
@@ -737,7 +701,7 @@ public:
 
 	void SetParameters(FRHICommandList& RHICmdList, const FViewInfo& View, TRefCountPtr<IPooledRenderTarget>& DistanceFieldAOBentNormal)
 	{
-		FRHIPixelShader* ShaderRHI = GetPixelShader();
+		FRHIPixelShader* ShaderRHI = RHICmdList.GetBoundPixelShader();
 
 		FGlobalShader::SetParameters<FViewUniformShaderParameters>(RHICmdList, ShaderRHI, View.ViewUniformBuffer);
 		SceneTextureParameters.Set(RHICmdList, ShaderRHI, View.FeatureLevel, ESceneTextureSetupMode::All);
@@ -747,21 +711,12 @@ public:
 		const float MinOcclusion = Scene->SkyLight ? Scene->SkyLight->MinOcclusion : 0;
 		SetShaderValue(RHICmdList, ShaderRHI, MinIndirectDiffuseOcclusion, MinOcclusion);
 	}
-	// FShader interface.
-	virtual bool Serialize(FArchive& Ar) override
-	{
-		bool bShaderHasOutdatedParameters = FGlobalShader::Serialize(Ar);
-		Ar << SceneTextureParameters;
-		Ar << DFAOUpsampleParameters;
-		Ar << MinIndirectDiffuseOcclusion;
-		return bShaderHasOutdatedParameters;
-	}
-
+	
 private:
 
-	FSceneTextureShaderParameters SceneTextureParameters;
-	FDFAOUpsampleParameters DFAOUpsampleParameters;
-	FShaderParameter MinIndirectDiffuseOcclusion;
+	LAYOUT_FIELD(FSceneTextureShaderParameters, SceneTextureParameters);
+	LAYOUT_FIELD(FDFAOUpsampleParameters, DFAOUpsampleParameters);
+	LAYOUT_FIELD(FShaderParameter, MinIndirectDiffuseOcclusion);
 };
 
 #define IMPLEMENT_UPSAMPLE_PS_TYPE(bModulateToSceneColor) \
@@ -807,8 +762,8 @@ void UpsampleBentNormalAO(
 		{
 			TShaderMapRef<TDistanceFieldAOUpsamplePS<true> > PixelShader(View.ShaderMap);
 			GraphicsPSOInit.BoundShaderState.VertexDeclarationRHI = GFilterVertexDeclaration.VertexDeclarationRHI;
-			GraphicsPSOInit.BoundShaderState.VertexShaderRHI = GETSAFERHISHADER_VERTEX(*VertexShader);
-			GraphicsPSOInit.BoundShaderState.PixelShaderRHI = GETSAFERHISHADER_PIXEL(*PixelShader);
+			GraphicsPSOInit.BoundShaderState.VertexShaderRHI = VertexShader.GetVertexShader();
+			GraphicsPSOInit.BoundShaderState.PixelShaderRHI = PixelShader.GetPixelShader();
 			SetGraphicsPipelineState(RHICmdList, GraphicsPSOInit);
 			PixelShader->SetParameters(RHICmdList, View, DistanceFieldAOBentNormal);
 		}
@@ -816,8 +771,8 @@ void UpsampleBentNormalAO(
 		{
 			TShaderMapRef<TDistanceFieldAOUpsamplePS<false> > PixelShader(View.ShaderMap);
 			GraphicsPSOInit.BoundShaderState.VertexDeclarationRHI = GFilterVertexDeclaration.VertexDeclarationRHI;
-			GraphicsPSOInit.BoundShaderState.VertexShaderRHI = GETSAFERHISHADER_VERTEX(*VertexShader);
-			GraphicsPSOInit.BoundShaderState.PixelShaderRHI = GETSAFERHISHADER_PIXEL(*PixelShader);
+			GraphicsPSOInit.BoundShaderState.VertexShaderRHI = VertexShader.GetVertexShader();
+			GraphicsPSOInit.BoundShaderState.PixelShaderRHI = PixelShader.GetPixelShader();
 			SetGraphicsPipelineState(RHICmdList, GraphicsPSOInit);
 			PixelShader->SetParameters(RHICmdList, View, DistanceFieldAOBentNormal);
 		}
@@ -830,6 +785,6 @@ void UpsampleBentNormalAO(
 			View.ViewRect.Width() / GAODownsampleFactor, View.ViewRect.Height() / GAODownsampleFactor,
 			FIntPoint(View.ViewRect.Width(), View.ViewRect.Height()),
 			GetBufferSizeForAO(),
-			*VertexShader);
+			VertexShader);
 	}
 }
diff --git a/Engine/Source/Runtime/Renderer/Private/DistanceFieldLightingShared.h b/Engine/Source/Runtime/Renderer/Private/DistanceFieldLightingShared.h
--- a/Engine/Source/Runtime/Renderer/Private/DistanceFieldLightingShared.h
+++ b/Engine/Source/Runtime/Renderer/Private/DistanceFieldLightingShared.h
@@ -150,6 +150,7 @@ public:
 template <EDistanceFieldPrimitiveType PrimitiveType>
 class TDistanceFieldObjectBufferParameters
 {
+	DECLARE_INLINE_TYPE_LAYOUT(TDistanceFieldObjectBufferParameters, NonVirtual);
 public:
 	void Bind(const FShaderParameterMap& ParameterMap)
 	{
@@ -227,16 +228,18 @@ public:
 	}
 
 private:
-	FRWShaderParameter SceneObjectBounds;
-	FRWShaderParameter SceneObjectData;
-	FShaderParameter NumSceneObjects;
-	FShaderResourceParameter DistanceFieldTexture;
-	FShaderResourceParameter DistanceFieldSampler;
-	FShaderParameter DistanceFieldAtlasTexelSize;
+	
+	LAYOUT_FIELD(FRWShaderParameter, SceneObjectBounds)
+	LAYOUT_FIELD(FRWShaderParameter, SceneObjectData)
+	LAYOUT_FIELD(FShaderParameter, NumSceneObjects)
+	LAYOUT_FIELD(FShaderResourceParameter, DistanceFieldTexture)
+	LAYOUT_FIELD(FShaderResourceParameter, DistanceFieldSampler)
+	LAYOUT_FIELD(FShaderParameter, DistanceFieldAtlasTexelSize)
 };
 
 class FSurfelBufferParameters
 {
+	DECLARE_TYPE_LAYOUT(FSurfelBufferParameters, NonVirtual);
 public:
 	void Bind(const FShaderParameterMap& ParameterMap)
 	{
@@ -268,9 +271,11 @@ public:
 	}
 
 private:
-	FRWShaderParameter InterpolatedVertexData;
-	FRWShaderParameter SurfelData;
-	FRWShaderParameter VPLFlux;
+	
+		LAYOUT_FIELD(FRWShaderParameter, InterpolatedVertexData)
+		LAYOUT_FIELD(FRWShaderParameter, SurfelData)
+		LAYOUT_FIELD(FRWShaderParameter, VPLFlux)
+	
 };
 
 template <EDistanceFieldPrimitiveType PrimitiveType>
@@ -402,6 +407,7 @@ class FHeightFieldObjectBufferResource : public TDistanceFieldObjectBufferResour
 template <EDistanceFieldPrimitiveType PrimitiveType>
 class TDistanceFieldCulledObjectBufferParameters
 {
+	DECLARE_INLINE_TYPE_LAYOUT(TDistanceFieldCulledObjectBufferParameters, NonVirtual);
 public:
 	void Bind(const FShaderParameterMap& ParameterMap)
 	{
@@ -512,14 +518,15 @@ public:
 	}
 
 private:
-	FRWShaderParameter ObjectIndirectArguments;
-	FRWShaderParameter CulledObjectBounds;
-	FRWShaderParameter CulledObjectData;
-	FRWShaderParameter CulledObjectBoxBounds;
-	FShaderResourceParameter HFVisibilityTexture;
-	FShaderResourceParameter DistanceFieldTexture;
-	FShaderResourceParameter DistanceFieldSampler;
-	FShaderParameter DistanceFieldAtlasTexelSize;
+	
+	LAYOUT_FIELD(FRWShaderParameter, ObjectIndirectArguments)
+	LAYOUT_FIELD(FRWShaderParameter, CulledObjectBounds)
+	LAYOUT_FIELD(FRWShaderParameter, CulledObjectData)
+	LAYOUT_FIELD(FRWShaderParameter, CulledObjectBoxBounds)
+	LAYOUT_FIELD(FShaderResourceParameter, HFVisibilityTexture);
+	LAYOUT_FIELD(FShaderResourceParameter, DistanceFieldTexture)
+	LAYOUT_FIELD(FShaderResourceParameter, DistanceFieldSampler)
+	LAYOUT_FIELD(FShaderParameter, DistanceFieldAtlasTexelSize)
 };
 
 class FCPUUpdatedBuffer
@@ -626,6 +633,7 @@ public:
 
 class FLightTileIntersectionParameters
 {
+	DECLARE_TYPE_LAYOUT(FLightTileIntersectionParameters, NonVirtual);
 public:
 
 	static void ModifyCompilationEnvironment(EShaderPlatform Platform, FShaderCompilerEnvironment& OutEnvironment)
@@ -714,12 +722,14 @@ public:
 	}
 
 private:
-	FRWShaderParameter ShadowTileNumCulledObjects;
-	FRWShaderParameter ShadowTileStartOffsets;
-	FRWShaderParameter NextStartOffset;
-	FRWShaderParameter ShadowTileArrayData;
-	FShaderParameter ShadowTileListGroupSize;
-	FShaderParameter ShadowAverageObjectsPerTile;
+	
+		LAYOUT_FIELD(FRWShaderParameter, ShadowTileNumCulledObjects)
+		LAYOUT_FIELD(FRWShaderParameter, ShadowTileStartOffsets)
+		LAYOUT_FIELD(FRWShaderParameter, NextStartOffset)
+		LAYOUT_FIELD(FRWShaderParameter, ShadowTileArrayData)
+		LAYOUT_FIELD(FShaderParameter, ShadowTileListGroupSize)
+		LAYOUT_FIELD(FShaderParameter, ShadowAverageObjectsPerTile)
+	
 };
 
 extern void CullDistanceFieldObjectsForLight(
diff --git a/Engine/Source/Runtime/Renderer/Private/DistanceFieldObjectCulling.cpp b/Engine/Source/Runtime/Renderer/Private/DistanceFieldObjectCulling.cpp
--- a/Engine/Source/Runtime/Renderer/Private/DistanceFieldObjectCulling.cpp
+++ b/Engine/Source/Runtime/Renderer/Private/DistanceFieldObjectCulling.cpp
@@ -145,7 +145,7 @@ public:
 		OutUAVs[5] = Scene->DistanceFieldSceneData.GetCurrentObjectBuffers()->Bounds.UAV;
 		RHICmdList.TransitionResources(EResourceTransitionAccess::ERWBarrier, EResourceTransitionPipeline::EComputeToCompute, OutUAVs, UE_ARRAY_COUNT(OutUAVs));
 
-		FRHIComputeShader* ShaderRHI = GetComputeShader();
+		FRHIComputeShader* ShaderRHI = RHICmdList.GetBoundComputeShader();
 		FGlobalShader::SetParameters<FViewUniformShaderParameters>(RHICmdList, ShaderRHI, View.ViewUniformBuffer);
 
 		FRHITexture* TextureAtlas;
@@ -172,8 +172,8 @@ public:
 
 	void UnsetParameters(FRHICommandList& RHICmdList, const FScene* Scene)
 	{
-		ObjectBufferParameters.UnsetParameters(RHICmdList, GetComputeShader(), *(Scene->DistanceFieldSceneData.GetCurrentObjectBuffers()));
-		CulledObjectParameters.UnsetParameters(RHICmdList, GetComputeShader());
+		ObjectBufferParameters.UnsetParameters(RHICmdList, RHICmdList.GetBoundComputeShader(), *(Scene->DistanceFieldSceneData.GetCurrentObjectBuffers()));
+		CulledObjectParameters.UnsetParameters(RHICmdList, RHICmdList.GetBoundComputeShader());
 
 		FRHIUnorderedAccessView* OutUAVs[6];
 		OutUAVs[0] = GAOCulledObjectBuffers.Buffers.ObjectIndirectArguments.UAV;
@@ -185,26 +185,14 @@ public:
 		RHICmdList.TransitionResources(EResourceTransitionAccess::EReadable, EResourceTransitionPipeline::EComputeToGfx, OutUAVs, UE_ARRAY_COUNT(OutUAVs));
 	}
 
-	virtual bool Serialize(FArchive& Ar)
-	{		
-		bool bShaderHasOutdatedParameters = FGlobalShader::Serialize(Ar);
-		Ar << ObjectBufferParameters;
-		Ar << CulledObjectParameters;
-		Ar << AOParameters;
-		Ar << NumConvexHullPlanes;
-		Ar << ViewFrustumConvexHull;
-		Ar << ObjectBoundingGeometryIndexCount;
-		return bShaderHasOutdatedParameters;
-	}
-
 private:
 
-	TDistanceFieldObjectBufferParameters<DFPT_SignedDistanceField> ObjectBufferParameters;
-	TDistanceFieldCulledObjectBufferParameters<DFPT_SignedDistanceField> CulledObjectParameters;
-	FAOParameters AOParameters;
-	FShaderParameter NumConvexHullPlanes;
-	FShaderParameter ViewFrustumConvexHull;
-	FShaderParameter ObjectBoundingGeometryIndexCount;
+	LAYOUT_FIELD((TDistanceFieldObjectBufferParameters<DFPT_SignedDistanceField>), ObjectBufferParameters);
+	LAYOUT_FIELD((TDistanceFieldCulledObjectBufferParameters<DFPT_SignedDistanceField>), CulledObjectParameters);
+	LAYOUT_FIELD(FAOParameters, AOParameters);
+	LAYOUT_FIELD(FShaderParameter, NumConvexHullPlanes);
+	LAYOUT_FIELD(FShaderParameter, ViewFrustumConvexHull);
+	LAYOUT_FIELD(FShaderParameter, ObjectBoundingGeometryIndexCount);
 };
 
 IMPLEMENT_SHADER_TYPE(,FCullObjectsForViewCS,TEXT("/Engine/Private/DistanceFieldObjectCulling.usf"),TEXT("CullObjectsForViewCS"),SF_Compute);
@@ -228,10 +216,10 @@ void CullObjectsToView(FRHICommandListImmediate& RHICmdList, FScene* Scene, cons
 		RHICmdList.ClearUAVUint(CulledObjectBuffers.Buffers.ObjectIndirectArguments.UAV, FUintVector4(0, 0, 0, 0));
 
 		TShaderMapRef<FCullObjectsForViewCS> ComputeShader(GetGlobalShaderMap(Scene->GetFeatureLevel()));
-		RHICmdList.SetComputeShader(ComputeShader->GetComputeShader());
+		RHICmdList.SetComputeShader(ComputeShader.GetComputeShader());
 		ComputeShader->SetParameters(RHICmdList, Scene, View, Parameters);
 
-		DispatchComputeShader(RHICmdList, *ComputeShader, FMath::DivideAndRoundUp<uint32>(Scene->DistanceFieldSceneData.NumObjectsInBuffer, UpdateObjectsGroupSize), 1, 1);
+		DispatchComputeShader(RHICmdList, ComputeShader.GetShader(), FMath::DivideAndRoundUp<uint32>(Scene->DistanceFieldSceneData.NumObjectsInBuffer, UpdateObjectsGroupSize), 1, 1);
 		ComputeShader->UnsetParameters(RHICmdList, Scene);
 	}
 }
@@ -276,7 +264,7 @@ public:
 	}
 	void SetParameters(FRHICommandList& RHICmdList, const FViewInfo& View, FSceneRenderTargetItem& DistanceFieldNormal, FScene* Scene, FVector2D NumGroupsValue, const FDistanceFieldAOParameters& Parameters)
 	{
-		FRHIComputeShader* ShaderRHI = GetComputeShader();
+		FRHIComputeShader* ShaderRHI = RHICmdList.GetBoundComputeShader();
 
 		FGlobalShader::SetParameters<FViewUniformShaderParameters>(RHICmdList, ShaderRHI, View.ViewUniformBuffer);
 		SceneTextureParameters.Set(RHICmdList, ShaderRHI, View.FeatureLevel, ESceneTextureSetupMode::All);
@@ -308,8 +296,8 @@ public:
 
 	void UnsetParameters(FRHICommandList& RHICmdList, const FSceneView& View)
 	{
-		TileConeAxisAndCos.UnsetUAV(RHICmdList, GetComputeShader());
-		TileConeDepthRanges.UnsetUAV(RHICmdList, GetComputeShader());
+		TileConeAxisAndCos.UnsetUAV(RHICmdList, RHICmdList.GetBoundComputeShader());
+		TileConeDepthRanges.UnsetUAV(RHICmdList, RHICmdList.GetBoundComputeShader());
 
 		FTileIntersectionResources* TileIntersectionResources = ((FSceneViewState*)View.State)->AOTileIntersectionResources;
 
@@ -319,30 +307,15 @@ public:
 		RHICmdList.TransitionResources(EResourceTransitionAccess::EReadable, EResourceTransitionPipeline::EComputeToCompute, OutUAVs, UE_ARRAY_COUNT(OutUAVs));
 	}
 
-	virtual bool Serialize(FArchive& Ar)
-	{		
-		bool bShaderHasOutdatedParameters = FGlobalShader::Serialize(Ar);
-		Ar << SceneTextureParameters;
-		Ar << AOParameters;
-		Ar << TileConeAxisAndCos;
-		Ar << TileConeDepthRanges;
-		Ar << NumGroups;
-		Ar << ViewDimensionsParameter;
-		Ar << DistanceFieldNormalTexture;
-		Ar << DistanceFieldNormalSampler;
-		return bShaderHasOutdatedParameters;
-	}
-
 private:
-
-	FSceneTextureShaderParameters SceneTextureParameters;
-	FAOParameters AOParameters;
-	FRWShaderParameter TileConeAxisAndCos;
-	FRWShaderParameter TileConeDepthRanges;
-	FShaderParameter ViewDimensionsParameter;
-	FShaderParameter NumGroups;
-	FShaderResourceParameter DistanceFieldNormalTexture;
-	FShaderResourceParameter DistanceFieldNormalSampler;
+	LAYOUT_FIELD(FSceneTextureShaderParameters, SceneTextureParameters);
+	LAYOUT_FIELD(FAOParameters, AOParameters);
+	LAYOUT_FIELD(FRWShaderParameter, TileConeAxisAndCos);
+	LAYOUT_FIELD(FRWShaderParameter, TileConeDepthRanges);
+	LAYOUT_FIELD(FShaderParameter, ViewDimensionsParameter);
+	LAYOUT_FIELD(FShaderParameter, NumGroups);
+	LAYOUT_FIELD(FShaderResourceParameter, DistanceFieldNormalTexture);
+	LAYOUT_FIELD(FShaderResourceParameter, DistanceFieldNormalSampler);
 };
 
 IMPLEMENT_SHADER_TYPE(,FBuildTileConesCS,TEXT("/Engine/Private/DistanceFieldObjectCulling.usf"),TEXT("BuildTileConesMain"),SF_Compute);
@@ -370,7 +343,7 @@ public:
 
 	void SetParameters(FRHICommandList& RHICmdList, const FSceneView& View, const FDistanceFieldAOParameters& Parameters)
 	{
-		FRHIVertexShader* ShaderRHI = GetVertexShader();
+		FRHIVertexShader* ShaderRHI = RHICmdList.GetBoundVertexShader();
 		FGlobalShader::SetParameters<FViewUniformShaderParameters>(RHICmdList, ShaderRHI, View.ViewUniformBuffer);
 		
 		FRHITexture* TextureAtlas;
@@ -395,19 +368,10 @@ public:
 		SetShaderValue(RHICmdList, ShaderRHI, ConservativeRadiusScale, ConservativeRadiusScaleValue);
 	}
 
-	virtual bool Serialize(FArchive& Ar) override
-	{
-		bool bShaderHasOutdatedParameters = FGlobalShader::Serialize(Ar);
-		Ar << ObjectParameters;
-		Ar << AOParameters;
-		Ar << ConservativeRadiusScale;
-		return bShaderHasOutdatedParameters;
-	}
-
 private:
-	TDistanceFieldCulledObjectBufferParameters<DFPT_SignedDistanceField> ObjectParameters;
-	FAOParameters AOParameters;
-	FShaderParameter ConservativeRadiusScale;
+	LAYOUT_FIELD((TDistanceFieldCulledObjectBufferParameters<DFPT_SignedDistanceField>), ObjectParameters);
+	LAYOUT_FIELD(FAOParameters, AOParameters);
+	LAYOUT_FIELD(FShaderParameter, ConservativeRadiusScale);
 };
 
 IMPLEMENT_SHADER_TYPE(,FObjectCullVS,TEXT("/Engine/Private/DistanceFieldObjectCulling.usf"),TEXT("ObjectCullVS"),SF_Vertex);
@@ -449,7 +413,7 @@ public:
 
 	void SetParameters(FRHICommandList& RHICmdList, const FSceneView& View, FVector2D NumGroupsValue, const FDistanceFieldAOParameters& Parameters)
 	{
-		FRHIPixelShader* ShaderRHI = GetPixelShader();
+		FRHIPixelShader* ShaderRHI = RHICmdList.GetBoundPixelShader();
 		FGlobalShader::SetParameters<FViewUniformShaderParameters>(RHICmdList, ShaderRHI, View.ViewUniformBuffer);
 
 		FRHITexture* TextureAtlas;
@@ -481,26 +445,13 @@ public:
 		check(UAVs.Num() > 0);
 	}
 
-	// FShader interface.
-	virtual bool Serialize(FArchive& Ar) override
-	{
-		bool bShaderHasOutdatedParameters = FGlobalShader::Serialize(Ar);
-		Ar << ObjectParameters;
-		Ar << AOParameters;
-		Ar << TileIntersectionParameters;
-		Ar << TileConeAxisAndCos;
-		Ar << TileConeDepthRanges;
-		Ar << NumGroups;
-		return bShaderHasOutdatedParameters;
-	}
-
 private:
-	TDistanceFieldCulledObjectBufferParameters<DFPT_SignedDistanceField> ObjectParameters;
-	FAOParameters AOParameters;
-	FTileIntersectionParameters TileIntersectionParameters;
-	FShaderResourceParameter TileConeAxisAndCos;
-	FShaderResourceParameter TileConeDepthRanges;
-	FShaderParameter NumGroups;
+	LAYOUT_FIELD((TDistanceFieldCulledObjectBufferParameters<DFPT_SignedDistanceField>), ObjectParameters);
+	LAYOUT_FIELD(FAOParameters, AOParameters);
+	LAYOUT_FIELD(FTileIntersectionParameters, TileIntersectionParameters);
+	LAYOUT_FIELD(FShaderResourceParameter, TileConeAxisAndCos);
+	LAYOUT_FIELD(FShaderResourceParameter, TileConeDepthRanges);
+	LAYOUT_FIELD(FShaderParameter, NumGroups);
 };
 
 IMPLEMENT_SHADER_TYPE(template<>,TObjectCullPS<true>,TEXT("/Engine/Private/DistanceFieldObjectCulling.usf"),TEXT("ObjectCullPS"),SF_Pixel);
@@ -539,7 +490,7 @@ public:
 	}
 	void SetParameters(FRHICommandList& RHICmdList, const FSceneView& View)
 	{
-		FRHIComputeShader* ShaderRHI = GetComputeShader();
+		FRHIComputeShader* ShaderRHI = RHICmdList.GetBoundComputeShader();
 
 		FGlobalShader::SetParameters<FViewUniformShaderParameters>(RHICmdList, ShaderRHI, View.ViewUniformBuffer);
 
@@ -569,7 +520,7 @@ public:
 	{
 		FTileIntersectionResources* TileIntersectionResources = ((FSceneViewState*)View.State)->AOTileIntersectionResources;
 
-		TileIntersectionParameters.UnsetParameters(RHICmdList, GetComputeShader());
+		TileIntersectionParameters.UnsetParameters(RHICmdList, RHICmdList.GetBoundComputeShader());
 
 		TArray<FRHIUnorderedAccessView*> UAVs;
 		TileIntersectionParameters.GetUAVs(*TileIntersectionResources, UAVs);
@@ -577,18 +528,10 @@ public:
 		RHICmdList.TransitionResources(EResourceTransitionAccess::EReadable, EResourceTransitionPipeline::EComputeToCompute, UAVs.GetData(), UAVs.Num());
 	}
 
-	virtual bool Serialize(FArchive& Ar) override
-	{		
-		bool bShaderHasOutdatedParameters = FGlobalShader::Serialize(Ar);
-		Ar << ObjectParameters;
-		Ar << TileIntersectionParameters;
-		return bShaderHasOutdatedParameters;
-	}
-
 private:
 
-	TDistanceFieldCulledObjectBufferParameters<DFPT_SignedDistanceField> ObjectParameters;
-	FTileIntersectionParameters TileIntersectionParameters;
+	LAYOUT_FIELD((TDistanceFieldCulledObjectBufferParameters<DFPT_SignedDistanceField>), ObjectParameters);
+	LAYOUT_FIELD(FTileIntersectionParameters, TileIntersectionParameters);
 };
 
 IMPLEMENT_SHADER_TYPE(,FComputeCulledTilesStartOffsetCS,TEXT("/Engine/Private/DistanceFieldObjectCulling.usf"),TEXT("ComputeCulledTilesStartOffsetCS"),SF_Compute);
@@ -630,8 +573,8 @@ void ScatterTilesToObjects(FRHICommandListImmediate& RHICmdList, const FViewInfo
 		GraphicsPSOInit.PrimitiveType = PT_TriangleList;
 
 		GraphicsPSOInit.BoundShaderState.VertexDeclarationRHI = GetVertexDeclarationFVector4();
-		GraphicsPSOInit.BoundShaderState.VertexShaderRHI = GETSAFERHISHADER_VERTEX(*VertexShader);
-		GraphicsPSOInit.BoundShaderState.PixelShaderRHI = GETSAFERHISHADER_PIXEL(*PixelShader);
+		GraphicsPSOInit.BoundShaderState.VertexShaderRHI = VertexShader.GetVertexShader();
+		GraphicsPSOInit.BoundShaderState.PixelShaderRHI = PixelShader.GetPixelShader();
 
 		SetGraphicsPipelineState(RHICmdList, GraphicsPSOInit);
 
@@ -698,9 +641,9 @@ void BuildTileObjectLists(FRHICommandListImmediate& RHICmdList, FScene* Scene, T
 				SCOPED_DRAW_EVENT(RHICmdList, BuildTileCones);
 				TShaderMapRef<FBuildTileConesCS> ComputeShader(View.ShaderMap);
 
-				RHICmdList.SetComputeShader(ComputeShader->GetComputeShader());
+				RHICmdList.SetComputeShader(ComputeShader.GetComputeShader());
 				ComputeShader->SetParameters(RHICmdList, View, DistanceFieldNormal, Scene, FVector2D(TileListGroupSize.X, TileListGroupSize.Y), Parameters);
-				DispatchComputeShader(RHICmdList, *ComputeShader, TileListGroupSize.X, TileListGroupSize.Y, 1);
+				DispatchComputeShader(RHICmdList, ComputeShader.GetShader(), TileListGroupSize.X, TileListGroupSize.Y, 1);
 
 				ComputeShader->UnsetParameters(RHICmdList, View);
 			}
@@ -727,9 +670,9 @@ void BuildTileObjectLists(FRHICommandListImmediate& RHICmdList, FScene* Scene, T
 				const uint32 GroupSize = FMath::DivideAndRoundUp<uint32>(Scene->DistanceFieldSceneData.NumObjectsInBuffer, ComputeStartOffsetGroupSize);
 				// Must write to RWObjectTilesIndirectArguments
 				check(GroupSize != 0);
-				RHICmdList.SetComputeShader(ComputeShader->GetComputeShader());
+				RHICmdList.SetComputeShader(ComputeShader.GetComputeShader());
 				ComputeShader->SetParameters(RHICmdList, View);
-				DispatchComputeShader(RHICmdList, *ComputeShader, GroupSize, 1, 1);
+				DispatchComputeShader(RHICmdList, ComputeShader.GetShader(), GroupSize, 1, 1);
 
 				ComputeShader->UnsetParameters(RHICmdList, View);
 			}
