Copying Private-LoadTimes-4.24 stream to Main. Biggest changes are in Materials/Shader memory freezing.

diff --git a/Engine/Source/Runtime/Renderer/Private/DistanceFieldObjectManagement.cpp b/Engine/Source/Runtime/Renderer/Private/DistanceFieldObjectManagement.cpp
--- a/Engine/Source/Runtime/Renderer/Private/DistanceFieldObjectManagement.cpp
+++ b/Engine/Source/Runtime/Renderer/Private/DistanceFieldObjectManagement.cpp
@@ -207,7 +207,7 @@ public:
 
 	void SetParameters(FRHICommandList& RHICmdList, const FScene* Scene, uint32 NumUploadOperationsValue, FRHIShaderResourceView* InUploadOperationIndices, FRHIShaderResourceView* InUploadOperationData)
 	{
-		FRHIComputeShader* ShaderRHI = GetComputeShader();
+		FRHIComputeShader* ShaderRHI = RHICmdList.GetBoundComputeShader();
 
 		SetShaderValue(RHICmdList, ShaderRHI, NumUploadOperations, NumUploadOperationsValue);
 
@@ -248,25 +248,20 @@ public:
 		const FDistanceFieldSceneData& SceneData = Scene->DistanceFieldSceneData;
 		const auto& ObjectBuffers = TSelector<bIsHeightField>()(*SceneData.GetHeightFieldObjectBuffers(), *SceneData.GetCurrentObjectBuffers());
 
-		ObjectBufferParameters.UnsetParameters(RHICmdList, GetComputeShader(), ObjectBuffers, true);
-	}
-
-	virtual bool Serialize(FArchive& Ar) override
-	{		
-		bool bShaderHasOutdatedParameters = FGlobalShader::Serialize(Ar);
-		Ar << NumUploadOperations;
-		Ar << UploadOperationIndices;
-		Ar << UploadOperationData;
-		Ar << ObjectBufferParameters;
-		return bShaderHasOutdatedParameters;
+		ObjectBufferParameters.UnsetParameters(RHICmdList, RHICmdList.GetBoundComputeShader(), ObjectBuffers, true);
 	}
 
 private:
 
-	FShaderParameter NumUploadOperations;
-	FShaderResourceParameter UploadOperationIndices;
-	FShaderResourceParameter UploadOperationData;
-	TDistanceFieldObjectBufferParameters<PrimitiveType> ObjectBufferParameters;
+	LAYOUT_FIELD(FShaderParameter, NumUploadOperations)
+	LAYOUT_FIELD(FShaderResourceParameter, UploadOperationIndices)
+	LAYOUT_FIELD(FShaderResourceParameter, UploadOperationData)
+
+	// clang is segfaulting with the templated IMPLEMENT_SHADER_TYPE below with this templated type directly in the LAYOUT_FIELD, so the uses of using like this in this file works around it
+	// (and also in DistanceFieldShadowing.cpp)
+	// a bugreport has been sent to Apple
+	using FDistanceFieldObjectBufferParametersType = TDistanceFieldObjectBufferParameters<PrimitiveType>;
+	LAYOUT_FIELD(FDistanceFieldObjectBufferParametersType, ObjectBufferParameters)
 };
 
 IMPLEMENT_SHADER_TYPE(template<>, TUploadObjectsToBufferCS<DFPT_SignedDistanceField>, TEXT("/Engine/Private/DistanceFieldObjectCulling.usf"), TEXT("UploadObjectsToBufferCS"), SF_Compute);
@@ -306,7 +301,7 @@ public:
 
 	void SetParameters(FRHICommandList& RHICmdList, ObjectBuffersType& ObjectBuffersSource, ObjectBuffersType& ObjectBuffersDest, int32 NumObjectsValue)
 	{
-		FRHIComputeShader* ShaderRHI = GetComputeShader();
+		FRHIComputeShader* ShaderRHI = RHICmdList.GetBoundComputeShader();
 
 		FRHIUnorderedAccessView* OutUAVs[2];
 		OutUAVs[0] = ObjectBuffersDest.Bounds.UAV;
@@ -341,9 +336,9 @@ public:
 
 	void UnsetParameters(FRHICommandList& RHICmdList, ObjectBuffersType& ObjectBuffersDest)
 	{
-		ObjectBufferParameters.UnsetParameters(RHICmdList, GetComputeShader(), ObjectBuffersDest);
-		CopyObjectBounds.UnsetUAV(RHICmdList, GetComputeShader());
-		CopyObjectData.UnsetUAV(RHICmdList, GetComputeShader());
+		ObjectBufferParameters.UnsetParameters(RHICmdList, RHICmdList.GetBoundComputeShader(), ObjectBuffersDest);
+		CopyObjectBounds.UnsetUAV(RHICmdList, RHICmdList.GetBoundComputeShader());
+		CopyObjectData.UnsetUAV(RHICmdList, RHICmdList.GetBoundComputeShader());
 
 		FRHIUnorderedAccessView* OutUAVs[2];
 		OutUAVs[0] = ObjectBuffersDest.Bounds.UAV;
@@ -351,20 +346,12 @@ public:
 		RHICmdList.TransitionResources(EResourceTransitionAccess::EReadable, EResourceTransitionPipeline::EComputeToCompute, OutUAVs, UE_ARRAY_COUNT(OutUAVs));
 	}
 
-	virtual bool Serialize(FArchive& Ar) override
-	{		
-		bool bShaderHasOutdatedParameters = FGlobalShader::Serialize(Ar);
-		Ar << CopyObjectBounds;
-		Ar << CopyObjectData;
-		Ar << ObjectBufferParameters;
-		return bShaderHasOutdatedParameters;
-	}
-
 private:
 
-	FRWShaderParameter CopyObjectBounds;
-	FRWShaderParameter CopyObjectData;
-	TDistanceFieldObjectBufferParameters<PrimitiveType> ObjectBufferParameters;
+	LAYOUT_FIELD(FRWShaderParameter, CopyObjectBounds)
+	LAYOUT_FIELD(FRWShaderParameter, CopyObjectData)
+	using FDistanceFieldObjectBufferParametersType = TDistanceFieldObjectBufferParameters<PrimitiveType>;
+	LAYOUT_FIELD(FDistanceFieldObjectBufferParametersType, ObjectBufferParameters)
 };
 
 IMPLEMENT_SHADER_TYPE(template<>, TCopyObjectBufferCS<DFPT_SignedDistanceField>, TEXT("/Engine/Private/DistanceFieldObjectCulling.usf"), TEXT("CopyObjectBufferCS"), SF_Compute);
@@ -401,7 +388,7 @@ public:
 
 	void SetParameters(FRHICommandList& RHICmdList, const FSurfelBuffers& SurfelBuffersSource, const FInstancedSurfelBuffers& InstancedSurfelBuffersSource, FSurfelBuffers& SurfelBuffersDest, int32 NumSurfelsValue)
 	{
-		FRHIComputeShader* ShaderRHI = GetComputeShader();
+		FRHIComputeShader* ShaderRHI = RHICmdList.GetBoundComputeShader();
 
 		FRHIUnorderedAccessView* OutUAVs[2];
 		OutUAVs[0] = SurfelBuffersDest.InterpolatedVertexData.UAV;
@@ -416,9 +403,9 @@ public:
 
 	void UnsetParameters(FRHICommandList& RHICmdList, FSurfelBuffers& SurfelBuffersDest)
 	{
-		SurfelBufferParameters.UnsetParameters(RHICmdList, GetComputeShader());
-		CopyInterpolatedVertexData.UnsetUAV(RHICmdList, GetComputeShader());
-		CopySurfelData.UnsetUAV(RHICmdList, GetComputeShader());
+		SurfelBufferParameters.UnsetParameters(RHICmdList, RHICmdList.GetBoundComputeShader());
+		CopyInterpolatedVertexData.UnsetUAV(RHICmdList, RHICmdList.GetBoundComputeShader());
+		CopySurfelData.UnsetUAV(RHICmdList, RHICmdList.GetBoundComputeShader());
 
 		FRHIUnorderedAccessView* OutUAVs[2];
 		OutUAVs[0] = SurfelBuffersDest.InterpolatedVertexData.UAV;
@@ -426,22 +413,12 @@ public:
 		RHICmdList.TransitionResources(EResourceTransitionAccess::EReadable, EResourceTransitionPipeline::EComputeToCompute, OutUAVs, UE_ARRAY_COUNT(OutUAVs));
 	}
 
-	virtual bool Serialize(FArchive& Ar) override
-	{		
-		bool bShaderHasOutdatedParameters = FGlobalShader::Serialize(Ar);
-		Ar << CopyInterpolatedVertexData;
-		Ar << CopySurfelData;
-		Ar << SurfelBufferParameters;
-		Ar << NumSurfels;
-		return bShaderHasOutdatedParameters;
-	}
-
 private:
 
-	FRWShaderParameter CopyInterpolatedVertexData;
-	FRWShaderParameter CopySurfelData;
-	FSurfelBufferParameters SurfelBufferParameters;
-	FShaderParameter NumSurfels;
+	LAYOUT_FIELD(FRWShaderParameter, CopyInterpolatedVertexData)
+	LAYOUT_FIELD(FRWShaderParameter, CopySurfelData)
+	LAYOUT_FIELD(FSurfelBufferParameters, SurfelBufferParameters)
+	LAYOUT_FIELD(FShaderParameter, NumSurfels)
 };
 
 IMPLEMENT_SHADER_TYPE(,FCopySurfelBufferCS,TEXT("/Engine/Private/SurfelTree.usf"),TEXT("CopySurfelBufferCS"),SF_Compute);
@@ -477,7 +454,7 @@ public:
 
 	void SetParameters(FRHICommandList& RHICmdList, const FSurfelBuffers& SurfelBuffersSource, const FInstancedSurfelBuffers& InstancedSurfelBuffersSource, FInstancedSurfelBuffers& InstancedSurfelBuffersDest, int32 NumSurfelsValue)
 	{
-		FRHIComputeShader* ShaderRHI = GetComputeShader();
+		FRHIComputeShader* ShaderRHI = RHICmdList.GetBoundComputeShader();
 
 		RHICmdList.TransitionResource(EResourceTransitionAccess::ERWBarrier, EResourceTransitionPipeline::EComputeToCompute, InstancedSurfelBuffersDest.VPLFlux.UAV);
 		CopyVPLFlux.SetBuffer(RHICmdList, ShaderRHI, InstancedSurfelBuffersDest.VPLFlux);
@@ -487,25 +464,16 @@ public:
 
 	void UnsetParameters(FRHICommandList& RHICmdList, FInstancedSurfelBuffers& InstancedSurfelBuffersDest)
 	{
-		SurfelBufferParameters.UnsetParameters(RHICmdList, GetComputeShader());
-		CopyVPLFlux.UnsetUAV(RHICmdList, GetComputeShader());
+		SurfelBufferParameters.UnsetParameters(RHICmdList, RHICmdList.GetBoundComputeShader());
+		CopyVPLFlux.UnsetUAV(RHICmdList, RHICmdList.GetBoundComputeShader());
 		RHICmdList.TransitionResource(EResourceTransitionAccess::EReadable, EResourceTransitionPipeline::EComputeToCompute, InstancedSurfelBuffersDest.VPLFlux.UAV);
 	}
 
-	virtual bool Serialize(FArchive& Ar) override
-	{		
-		bool bShaderHasOutdatedParameters = FGlobalShader::Serialize(Ar);
-		Ar << CopyVPLFlux;
-		Ar << SurfelBufferParameters;
-		Ar << NumSurfels;
-		return bShaderHasOutdatedParameters;
-	}
-
 private:
 
-	FRWShaderParameter CopyVPLFlux;
-	FSurfelBufferParameters SurfelBufferParameters;
-	FShaderParameter NumSurfels;
+	LAYOUT_FIELD(FRWShaderParameter, CopyVPLFlux)
+	LAYOUT_FIELD(FSurfelBufferParameters, SurfelBufferParameters)
+	LAYOUT_FIELD(FShaderParameter, NumSurfels)
 };
 
 IMPLEMENT_SHADER_TYPE(,FCopyVPLFluxBufferCS,TEXT("/Engine/Private/SurfelTree.usf"),TEXT("CopyVPLFluxBufferCS"),SF_Compute);
@@ -551,7 +519,7 @@ public:
 		FRHIShaderResourceView* InObjectBounds2,
 		FRHIShaderResourceView* InObjectData2)
 	{
-		FRHIComputeShader* ShaderRHI = GetComputeShader();
+		FRHIComputeShader* ShaderRHI = RHICmdList.GetBoundComputeShader();
 
 		SetShaderValue(RHICmdList, ShaderRHI, NumRemoveOperations, NumRemoveOperationsValue);
 		SetSRVParameter(RHICmdList, ShaderRHI, RemoveOperationIndices, InRemoveOperationIndices);
@@ -594,27 +562,17 @@ public:
 		const FDistanceFieldSceneData& SceneData = Scene->DistanceFieldSceneData;
 		const auto& ObjectBuffers = TSelector<bIsHeightField>()(*SceneData.GetHeightFieldObjectBuffers(), *SceneData.GetCurrentObjectBuffers());
 
-		ObjectBufferParameters.UnsetParameters(RHICmdList, GetComputeShader(), ObjectBuffers, true);
-	}
-
-	virtual bool Serialize(FArchive& Ar) override
-	{		
-		bool bShaderHasOutdatedParameters = FGlobalShader::Serialize(Ar);
-		Ar << NumRemoveOperations;
-		Ar << RemoveOperationIndices;
-		Ar << ObjectBufferParameters;
-		Ar << ObjectBounds2;
-		Ar << ObjectData2;
-		return bShaderHasOutdatedParameters;
+		ObjectBufferParameters.UnsetParameters(RHICmdList, RHICmdList.GetBoundComputeShader(), ObjectBuffers, true);
 	}
 
 private:
 
-	FShaderParameter NumRemoveOperations;
-	FShaderResourceParameter RemoveOperationIndices;
-	TDistanceFieldObjectBufferParameters<PrimitiveType> ObjectBufferParameters;
-	FShaderResourceParameter ObjectBounds2;
-	FShaderResourceParameter ObjectData2;
+	LAYOUT_FIELD(FShaderParameter, NumRemoveOperations)
+	LAYOUT_FIELD(FShaderResourceParameter, RemoveOperationIndices)
+	using FDistanceFieldObjectBufferParametersType = TDistanceFieldObjectBufferParameters<PrimitiveType>;
+	LAYOUT_FIELD(FDistanceFieldObjectBufferParametersType, ObjectBufferParameters)
+	LAYOUT_FIELD(FShaderResourceParameter, ObjectBounds2)
+	LAYOUT_FIELD(FShaderResourceParameter, ObjectData2)
 };
 
 #define IMPLEMENT_REMOVE_OBJECTS_FROM_BUFFER_SHADER_TYPE(bRemoveFromSameBuffer, PrimitiveType) \
@@ -877,19 +835,19 @@ void UpdateGlobalDistanceFieldObjectRemoves(FRHICommandListImmediate& RHICmdList
 				if (bUseRemoveAtSwap)
 				{
 					TShaderMapRef<TRemoveObjectsFromBufferCS<true, DFPT_SignedDistanceField>> ComputeShader(GetGlobalShaderMap(Scene->GetFeatureLevel()));
-					RHICmdList.SetComputeShader(ComputeShader->GetComputeShader());
+					RHICmdList.SetComputeShader(ComputeShader.GetComputeShader());
 					ComputeShader->SetParameters(RHICmdList, Scene, RemoveObjectIndices.Num(), GDistanceFieldRemoveIndices.RemoveIndices.BufferSRV, NULL, NULL);
 
-					DispatchComputeShader(RHICmdList, *ComputeShader, FMath::DivideAndRoundUp<uint32>(RemoveObjectIndices.Num(), UpdateObjectsGroupSize), 1, 1);
+					DispatchComputeShader(RHICmdList, ComputeShader.GetShader(), FMath::DivideAndRoundUp<uint32>(RemoveObjectIndices.Num(), UpdateObjectsGroupSize), 1, 1);
 					ComputeShader->UnsetParameters(RHICmdList, Scene);
 				}
 				else
 				{
 					TShaderMapRef<TRemoveObjectsFromBufferCS<false, DFPT_SignedDistanceField>> ComputeShader(GetGlobalShaderMap(Scene->GetFeatureLevel()));
-					RHICmdList.SetComputeShader(ComputeShader->GetComputeShader());
+					RHICmdList.SetComputeShader(ComputeShader.GetComputeShader());
 					ComputeShader->SetParameters(RHICmdList, Scene, RemoveObjectIndices.Num(), GDistanceFieldRemoveIndices.RemoveIndices.BufferSRV, CurrentObjectBuffers->Bounds.SRV, CurrentObjectBuffers->Data.SRV);
 
-					DispatchComputeShader(RHICmdList, *ComputeShader, FMath::DivideAndRoundUp<uint32>(RemoveObjectIndices.Num(), UpdateObjectsGroupSize), 1, 1);
+					DispatchComputeShader(RHICmdList, ComputeShader.GetShader(), FMath::DivideAndRoundUp<uint32>(RemoveObjectIndices.Num(), UpdateObjectsGroupSize), 1, 1);
 					ComputeShader->UnsetParameters(RHICmdList, Scene);
 				}
 			}
@@ -972,9 +930,9 @@ void UpdateGlobalHeightFieldObjectRemoves(FRHICommandListImmediate& RHICmdList,
 			RHIUnlockVertexBuffer(RemoveIndices.Buffer);
 
 			TShaderMapRef<TRemoveObjectsFromBufferCS<true, DFPT_HeightField>> ComputeShader(GetGlobalShaderMap(Scene->GetFeatureLevel()));
-			RHICmdList.SetComputeShader(ComputeShader->GetComputeShader());
+			RHICmdList.SetComputeShader(ComputeShader.GetComputeShader());
 			ComputeShader->SetParameters(RHICmdList, Scene, DstSrcIndices.Num(), RemoveIndices.BufferSRV, nullptr, nullptr);
-			DispatchComputeShader(RHICmdList, *ComputeShader, FMath::DivideAndRoundUp<uint32>(DstSrcIndices.Num(), UpdateObjectsGroupSize), 1, 1);
+			DispatchComputeShader(RHICmdList, ComputeShader, FMath::DivideAndRoundUp<uint32>(DstSrcIndices.Num(), UpdateObjectsGroupSize), 1, 1);
 			ComputeShader->UnsetParameters(RHICmdList, Scene);
 		}
 	}
@@ -1431,10 +1389,10 @@ void FDeferredShadingSceneRenderer::UpdateGlobalDistanceFieldObjectBuffers(FRHIC
 
 						{
 							TShaderMapRef<FCopySurfelBufferCS> ComputeShader(GetGlobalShaderMap(Scene->GetFeatureLevel()));
-							RHICmdList.SetComputeShader(ComputeShader->GetComputeShader());
+							RHICmdList.SetComputeShader(ComputeShader.GetComputeShader());
 							ComputeShader->SetParameters(RHICmdList, *(DistanceFieldSceneData.SurfelBuffers), *(DistanceFieldSceneData.InstancedSurfelBuffers), *NewSurfelBuffers, OriginalNumSurfels);
 
-							DispatchComputeShader(RHICmdList, *ComputeShader, FMath::DivideAndRoundUp<uint32>(OriginalNumSurfels, UpdateObjectsGroupSize), 1, 1);
+							DispatchComputeShader(RHICmdList, ComputeShader.GetShader(), FMath::DivideAndRoundUp<uint32>(OriginalNumSurfels, UpdateObjectsGroupSize), 1, 1);
 							ComputeShader->UnsetParameters(RHICmdList, *NewSurfelBuffers);
 						}
 
@@ -1461,10 +1419,10 @@ void FDeferredShadingSceneRenderer::UpdateGlobalDistanceFieldObjectBuffers(FRHIC
 
 						{
 							TShaderMapRef<FCopyVPLFluxBufferCS> ComputeShader(GetGlobalShaderMap(Scene->GetFeatureLevel()));
-							RHICmdList.SetComputeShader(ComputeShader->GetComputeShader());
+							RHICmdList.SetComputeShader(ComputeShader.GetComputeShader());
 							ComputeShader->SetParameters(RHICmdList, *(DistanceFieldSceneData.SurfelBuffers), *(DistanceFieldSceneData.InstancedSurfelBuffers), *NewInstancedSurfelBuffers, OriginalNumInstancedSurfels);
 
-							DispatchComputeShader(RHICmdList, *ComputeShader, FMath::DivideAndRoundUp<uint32>(OriginalNumInstancedSurfels, UpdateObjectsGroupSize), 1, 1);
+							DispatchComputeShader(RHICmdList, ComputeShader.GetShader(), FMath::DivideAndRoundUp<uint32>(OriginalNumInstancedSurfels, UpdateObjectsGroupSize), 1, 1);
 							ComputeShader->UnsetParameters(RHICmdList, *NewInstancedSurfelBuffers);
 						}
 
@@ -1552,10 +1510,10 @@ void FDeferredShadingSceneRenderer::UpdateGlobalDistanceFieldObjectBuffers(FRHIC
 
 					{
 						TShaderMapRef<TCopyObjectBufferCS<DFPT_SignedDistanceField>> ComputeShader(GetGlobalShaderMap(Scene->GetFeatureLevel()));
-						RHICmdList.SetComputeShader(ComputeShader->GetComputeShader());
+						RHICmdList.SetComputeShader(ComputeShader.GetComputeShader());
 						ComputeShader->SetParameters(RHICmdList, *ObjectBuffers, *NewObjectBuffers, OriginalNumObjects);
 
-						DispatchComputeShader(RHICmdList, *ComputeShader, FMath::DivideAndRoundUp<uint32>(OriginalNumObjects, UpdateObjectsGroupSize), 1, 1);
+						DispatchComputeShader(RHICmdList, ComputeShader.GetShader(), FMath::DivideAndRoundUp<uint32>(OriginalNumObjects, UpdateObjectsGroupSize), 1, 1);
 						ComputeShader->UnsetParameters(RHICmdList, *NewObjectBuffers);
 					}
 
@@ -1601,10 +1559,10 @@ void FDeferredShadingSceneRenderer::UpdateGlobalDistanceFieldObjectBuffers(FRHIC
 
 			{
 				TShaderMapRef<TUploadObjectsToBufferCS<DFPT_SignedDistanceField>> ComputeShader(GetGlobalShaderMap(Scene->GetFeatureLevel()));
-				RHICmdList.SetComputeShader(ComputeShader->GetComputeShader());
+				RHICmdList.SetComputeShader(ComputeShader.GetComputeShader());
 				ComputeShader->SetParameters(RHICmdList, Scene, UploadObjectIndices.Num(), GDistanceFieldUploadIndices.UploadIndices.BufferSRV, GDistanceFieldUploadData.UploadData.BufferSRV);
 
-				DispatchComputeShader(RHICmdList, *ComputeShader, FMath::DivideAndRoundUp<uint32>(UploadObjectIndices.Num(), UpdateObjectsGroupSize), 1, 1);
+				DispatchComputeShader(RHICmdList, ComputeShader.GetShader(), FMath::DivideAndRoundUp<uint32>(UploadObjectIndices.Num(), UpdateObjectsGroupSize), 1, 1);
 				ComputeShader->UnsetParameters(RHICmdList, Scene);
 			}
 		}
@@ -1704,9 +1662,9 @@ void FDeferredShadingSceneRenderer::UpdateGlobalHeightFieldObjectBuffers(FRHICom
 					NewObjectBuffers->Initialize();
 
 					TShaderMapRef<TCopyObjectBufferCS<DFPT_HeightField>> ComputeShader(GetGlobalShaderMap(Scene->GetFeatureLevel()));
-					RHICmdList.SetComputeShader(ComputeShader->GetComputeShader());
+					RHICmdList.SetComputeShader(ComputeShader.GetComputeShader());
 					ComputeShader->SetParameters(RHICmdList, *ObjectBuffers, *NewObjectBuffers, OriginalNumObjects);
-					DispatchComputeShader(RHICmdList, *ComputeShader, FMath::DivideAndRoundUp<uint32>(OriginalNumObjects, UpdateObjectsGroupSize), 1, 1);
+					DispatchComputeShader(RHICmdList, ComputeShader, FMath::DivideAndRoundUp<uint32>(OriginalNumObjects, UpdateObjectsGroupSize), 1, 1);
 					ComputeShader->UnsetParameters(RHICmdList, *NewObjectBuffers);
 
 					ObjectBuffers->Release();
@@ -1749,9 +1707,9 @@ void FDeferredShadingSceneRenderer::UpdateGlobalHeightFieldObjectBuffers(FRHICom
 				RHIUnlockVertexBuffer(GHeightFieldUploadData.UploadData.Buffer);
 
 				TShaderMapRef<TUploadObjectsToBufferCS<DFPT_HeightField>> ComputeShader(GetGlobalShaderMap(Scene->GetFeatureLevel()));
-				RHICmdList.SetComputeShader(ComputeShader->GetComputeShader());
+				RHICmdList.SetComputeShader(ComputeShader.GetComputeShader());
 				ComputeShader->SetParameters(RHICmdList, Scene, NumObjectsToUpload, GHeightFieldUploadIndices.UploadIndices.BufferSRV, GHeightFieldUploadData.UploadData.BufferSRV);
-				DispatchComputeShader(RHICmdList, *ComputeShader, FMath::DivideAndRoundUp<uint32>(NumObjectsToUpload, UpdateObjectsGroupSize), 1, 1);
+				DispatchComputeShader(RHICmdList, ComputeShader, FMath::DivideAndRoundUp<uint32>(NumObjectsToUpload, UpdateObjectsGroupSize), 1, 1);
 				ComputeShader->UnsetParameters(RHICmdList, Scene);
 			}
 		}
diff --git a/Engine/Source/Runtime/Renderer/Private/DistanceFieldScreenGridLighting.cpp b/Engine/Source/Runtime/Renderer/Private/DistanceFieldScreenGridLighting.cpp
--- a/Engine/Source/Runtime/Renderer/Private/DistanceFieldScreenGridLighting.cpp
+++ b/Engine/Source/Runtime/Renderer/Private/DistanceFieldScreenGridLighting.cpp
@@ -128,7 +128,7 @@ public:
 		const FDistanceFieldAOParameters& Parameters,
 		const FGlobalDistanceFieldInfo& GlobalDistanceFieldInfo)
 	{
-		FRHIComputeShader* ShaderRHI = GetComputeShader();
+		FRHIComputeShader* ShaderRHI = RHICmdList.GetBoundComputeShader();
 		FGlobalShader::SetParameters<FViewUniformShaderParameters>(RHICmdList, ShaderRHI, View.ViewUniformBuffer);
 		SceneTextureParameters.Set(RHICmdList, ShaderRHI, View.FeatureLevel, ESceneTextureSetupMode::All);
 
@@ -193,40 +193,22 @@ public:
 
 	void UnsetParameters(FRHICommandList& RHICmdList, const FViewInfo& View)
 	{
-		ScreenGridConeVisibility.UnsetUAV(RHICmdList, GetComputeShader());
-		ConeDepthVisibilityFunction.UnsetUAV(RHICmdList, GetComputeShader());
-	}
-
-	virtual bool Serialize(FArchive& Ar) override
-	{		
-		bool bShaderHasOutdatedParameters = FGlobalShader::Serialize(Ar);
-		Ar << SceneTextureParameters;
-		Ar << ObjectParameters;
-		Ar << AOParameters;
-		Ar << ScreenGridParameters;
-		Ar << GlobalDistanceFieldParameters;
-		Ar << TileConeDepthRanges;
-		Ar << TileIntersectionParameters;
-		Ar << TanConeHalfAngle;
-		Ar << BentNormalNormalizeFactor;
-		Ar << ScreenGridConeVisibility;
-		Ar << ConeDepthVisibilityFunction;
-		return bShaderHasOutdatedParameters;
+		ScreenGridConeVisibility.UnsetUAV(RHICmdList, RHICmdList.GetBoundComputeShader());
+		ConeDepthVisibilityFunction.UnsetUAV(RHICmdList, RHICmdList.GetBoundComputeShader());
 	}
 
 private:
-
-	FSceneTextureShaderParameters SceneTextureParameters;
-	TDistanceFieldCulledObjectBufferParameters<DFPT_SignedDistanceField> ObjectParameters;
-	FAOParameters AOParameters;
-	FScreenGridParameters ScreenGridParameters;
-	FGlobalDistanceFieldParameters GlobalDistanceFieldParameters;
-	FShaderResourceParameter TileConeDepthRanges;
-	FTileIntersectionParameters TileIntersectionParameters;
-	FShaderParameter TanConeHalfAngle;
-	FShaderParameter BentNormalNormalizeFactor;
-	FRWShaderParameter ScreenGridConeVisibility;
-	FRWShaderParameter ConeDepthVisibilityFunction;
+	LAYOUT_FIELD(FSceneTextureShaderParameters, SceneTextureParameters);
+	LAYOUT_FIELD((TDistanceFieldCulledObjectBufferParameters<DFPT_SignedDistanceField>), ObjectParameters);
+	LAYOUT_FIELD(FAOParameters, AOParameters);
+	LAYOUT_FIELD(FScreenGridParameters, ScreenGridParameters);
+	LAYOUT_FIELD(FGlobalDistanceFieldParameters, GlobalDistanceFieldParameters);
+	LAYOUT_FIELD(FShaderResourceParameter, TileConeDepthRanges);
+	LAYOUT_FIELD(FTileIntersectionParameters, TileIntersectionParameters);
+	LAYOUT_FIELD(FShaderParameter, TanConeHalfAngle);
+	LAYOUT_FIELD(FShaderParameter, BentNormalNormalizeFactor);
+	LAYOUT_FIELD(FRWShaderParameter, ScreenGridConeVisibility);
+	LAYOUT_FIELD(FRWShaderParameter, ConeDepthVisibilityFunction);
 };
 
 #define IMPLEMENT_CONETRACE_CS_TYPE(bSupportIrradiance, bUseGlobalDistanceField) \
@@ -291,7 +273,7 @@ public:
 		const FDistanceFieldAOParameters& Parameters,
 		const FGlobalDistanceFieldInfo& GlobalDistanceFieldInfo)
 	{
-		FRHIComputeShader* ShaderRHI = GetComputeShader();
+		FRHIComputeShader* ShaderRHI = RHICmdList.GetBoundComputeShader();
 		FGlobalShader::SetParameters<FViewUniformShaderParameters>(RHICmdList, ShaderRHI, View.ViewUniformBuffer);
 		SceneTextureParameters.Set(RHICmdList, ShaderRHI, View.FeatureLevel, ESceneTextureSetupMode::All);
 
@@ -358,37 +340,20 @@ public:
 
 	void UnsetParameters(FRHICommandList& RHICmdList, const FViewInfo& View)
 	{
-		ScreenGridConeVisibility.UnsetUAV(RHICmdList, GetComputeShader());
-	}
-
-	virtual bool Serialize(FArchive& Ar) override
-	{		
-		bool bShaderHasOutdatedParameters = FGlobalShader::Serialize(Ar);
-		Ar << SceneTextureParameters;
-		Ar << ObjectParameters;
-		Ar << AOParameters;
-		Ar << ScreenGridParameters;
-		Ar << GlobalDistanceFieldParameters;
-		Ar << TileConeDepthRanges;
-		Ar << TileListGroupSize;
-		Ar << TanConeHalfAngle;
-		Ar << BentNormalNormalizeFactor;
-		Ar << ScreenGridConeVisibility;
-		return bShaderHasOutdatedParameters;
+		ScreenGridConeVisibility.UnsetUAV(RHICmdList, RHICmdList.GetBoundComputeShader());
 	}
 
 private:
-
-	FSceneTextureShaderParameters SceneTextureParameters;
-	TDistanceFieldCulledObjectBufferParameters<DFPT_SignedDistanceField> ObjectParameters;
-	FAOParameters AOParameters;
-	FScreenGridParameters ScreenGridParameters;
-	FGlobalDistanceFieldParameters GlobalDistanceFieldParameters;
-	FShaderResourceParameter TileConeDepthRanges;
-	FShaderParameter TileListGroupSize;
-	FShaderParameter TanConeHalfAngle;
-	FShaderParameter BentNormalNormalizeFactor;
-	FRWShaderParameter ScreenGridConeVisibility;
+	LAYOUT_FIELD(FSceneTextureShaderParameters, SceneTextureParameters);
+	LAYOUT_FIELD((TDistanceFieldCulledObjectBufferParameters<DFPT_SignedDistanceField>), ObjectParameters);
+	LAYOUT_FIELD(FAOParameters, AOParameters);
+	LAYOUT_FIELD(FScreenGridParameters, ScreenGridParameters);
+	LAYOUT_FIELD(FGlobalDistanceFieldParameters, GlobalDistanceFieldParameters);
+	LAYOUT_FIELD(FShaderResourceParameter, TileConeDepthRanges);
+	LAYOUT_FIELD(FShaderParameter, TileListGroupSize);
+	LAYOUT_FIELD(FShaderParameter, TanConeHalfAngle);
+	LAYOUT_FIELD(FShaderParameter, BentNormalNormalizeFactor);
+	LAYOUT_FIELD(FRWShaderParameter, ScreenGridConeVisibility);
 };
 
 #define IMPLEMENT_CONETRACE_GLOBAL_CS_TYPE(bSupportIrradiance) \
@@ -443,7 +408,7 @@ public:
 		FSceneRenderTargetItem& DistanceFieldNormal, 
 		FSceneRenderTargetItem& DownsampledBentNormal)
 	{
-		FRHIComputeShader* ShaderRHI = GetComputeShader();
+		FRHIComputeShader* ShaderRHI = RHICmdList.GetBoundComputeShader();
 		FGlobalShader::SetParameters<FViewUniformShaderParameters>(RHICmdList, ShaderRHI, View.ViewUniformBuffer);
 		ScreenGridParameters.Set(RHICmdList, ShaderRHI, View, DistanceFieldNormal);
 
@@ -488,30 +453,17 @@ public:
 
 	void UnsetParameters(FRHICommandList& RHICmdList, FSceneRenderTargetItem& DownsampledBentNormal)
 	{
-		DistanceFieldBentNormal.UnsetUAV(RHICmdList, GetComputeShader());
+		DistanceFieldBentNormal.UnsetUAV(RHICmdList, RHICmdList.GetBoundComputeShader());
 		RHICmdList.TransitionResource(EResourceTransitionAccess::EReadable, EResourceTransitionPipeline::EComputeToCompute, DownsampledBentNormal.UAV);
 	}
 
-	virtual bool Serialize(FArchive& Ar) override
-	{		
-		bool bShaderHasOutdatedParameters = FGlobalShader::Serialize(Ar);
-		Ar << ScreenGridParameters;
-		Ar << BentNormalNormalizeFactor;
-		Ar << ScreenGridConeVisibility;
-		Ar << DistanceFieldBentNormal;
-		Ar << DFNormalBufferUVMax;
-		Ar << ConeBufferMax;
-		return bShaderHasOutdatedParameters;
-	}
-
 private:
-
-	FScreenGridParameters ScreenGridParameters;
-	FShaderParameter BentNormalNormalizeFactor;
-	FShaderParameter DFNormalBufferUVMax;
-	FShaderParameter ConeBufferMax;
-	FShaderResourceParameter ScreenGridConeVisibility;
-	FRWShaderParameter DistanceFieldBentNormal;
+	LAYOUT_FIELD(FScreenGridParameters, ScreenGridParameters);
+	LAYOUT_FIELD(FShaderParameter, BentNormalNormalizeFactor);
+	LAYOUT_FIELD(FShaderParameter, DFNormalBufferUVMax);
+	LAYOUT_FIELD(FShaderParameter, ConeBufferMax);
+	LAYOUT_FIELD(FShaderResourceParameter, ScreenGridConeVisibility);
+	LAYOUT_FIELD(FRWShaderParameter, DistanceFieldBentNormal);
 };
 
 IMPLEMENT_SHADER_TYPE(,FCombineConeVisibilityCS,TEXT("/Engine/Private/DistanceFieldScreenGridLighting.usf"),TEXT("CombineConeVisibilityCS"),SF_Compute);
@@ -602,18 +554,18 @@ void FDeferredShadingSceneRenderer::RenderDistanceFieldAOScreenGrid(
 		{
 			TShaderMapRef<TConeTraceScreenGridGlobalOcclusionCS<true> > ComputeShader(View.ShaderMap);
 
-			RHICmdList.SetComputeShader(ComputeShader->GetComputeShader());
+			RHICmdList.SetComputeShader(ComputeShader.GetComputeShader());
 			ComputeShader->SetParameters(RHICmdList, View, TileListGroupSize, DistanceFieldNormal->GetRenderTargetItem(), Parameters, View.GlobalDistanceFieldInfo);
-			DispatchComputeShader(RHICmdList, *ComputeShader, GroupSizeX, GroupSizeY, 1);
+			DispatchComputeShader(RHICmdList, ComputeShader.GetShader(), GroupSizeX, GroupSizeY, 1);
 			ComputeShader->UnsetParameters(RHICmdList, View);
 		}
 		else
 		{
 			TShaderMapRef<TConeTraceScreenGridGlobalOcclusionCS<false> > ComputeShader(View.ShaderMap);
 
-			RHICmdList.SetComputeShader(ComputeShader->GetComputeShader());
+			RHICmdList.SetComputeShader(ComputeShader.GetComputeShader());
 			ComputeShader->SetParameters(RHICmdList, View, TileListGroupSize, DistanceFieldNormal->GetRenderTargetItem(), Parameters, View.GlobalDistanceFieldInfo);
-			DispatchComputeShader(RHICmdList, *ComputeShader, GroupSizeX, GroupSizeY, 1);
+			DispatchComputeShader(RHICmdList, ComputeShader.GetShader(), GroupSizeX, GroupSizeY, 1);
 			ComputeShader->UnsetParameters(RHICmdList, View);
 		}
 	}
@@ -633,18 +585,18 @@ void FDeferredShadingSceneRenderer::RenderDistanceFieldAOScreenGrid(
 			{
 				TShaderMapRef<TConeTraceScreenGridObjectOcclusionCS<true, true> > ComputeShader(View.ShaderMap);
 
-				RHICmdList.SetComputeShader(ComputeShader->GetComputeShader());
+				RHICmdList.SetComputeShader(ComputeShader.GetComputeShader());
 				ComputeShader->SetParameters(RHICmdList, View, DistanceFieldNormal->GetRenderTargetItem(), Parameters, View.GlobalDistanceFieldInfo);
-				DispatchIndirectComputeShader(RHICmdList, *ComputeShader, TileIntersectionResources->ObjectTilesIndirectArguments.Buffer, 0);
+				DispatchIndirectComputeShader(RHICmdList, ComputeShader.GetShader(), TileIntersectionResources->ObjectTilesIndirectArguments.Buffer, 0);
 				ComputeShader->UnsetParameters(RHICmdList, View);
 			}
 			else
 			{
 				TShaderMapRef<TConeTraceScreenGridObjectOcclusionCS<false, true> > ComputeShader(View.ShaderMap);
 
-				RHICmdList.SetComputeShader(ComputeShader->GetComputeShader());
+				RHICmdList.SetComputeShader(ComputeShader.GetComputeShader());
 				ComputeShader->SetParameters(RHICmdList, View, DistanceFieldNormal->GetRenderTargetItem(), Parameters, View.GlobalDistanceFieldInfo);
-				DispatchIndirectComputeShader(RHICmdList, *ComputeShader, TileIntersectionResources->ObjectTilesIndirectArguments.Buffer, 0);
+				DispatchIndirectComputeShader(RHICmdList, ComputeShader.GetShader(), TileIntersectionResources->ObjectTilesIndirectArguments.Buffer, 0);
 				ComputeShader->UnsetParameters(RHICmdList, View);
 			}
 		}
@@ -654,18 +606,18 @@ void FDeferredShadingSceneRenderer::RenderDistanceFieldAOScreenGrid(
 			{
 				TShaderMapRef<TConeTraceScreenGridObjectOcclusionCS<true, false> > ComputeShader(View.ShaderMap);
 
-				RHICmdList.SetComputeShader(ComputeShader->GetComputeShader());
+				RHICmdList.SetComputeShader(ComputeShader.GetComputeShader());
 				ComputeShader->SetParameters(RHICmdList, View, DistanceFieldNormal->GetRenderTargetItem(), Parameters, View.GlobalDistanceFieldInfo);
-				DispatchIndirectComputeShader(RHICmdList, *ComputeShader, TileIntersectionResources->ObjectTilesIndirectArguments.Buffer, 0);
+				DispatchIndirectComputeShader(RHICmdList, ComputeShader.GetShader(), TileIntersectionResources->ObjectTilesIndirectArguments.Buffer, 0);
 				ComputeShader->UnsetParameters(RHICmdList, View);
 			}
 			else
 			{
 				TShaderMapRef<TConeTraceScreenGridObjectOcclusionCS<false, false> > ComputeShader(View.ShaderMap);
 
-				RHICmdList.SetComputeShader(ComputeShader->GetComputeShader());
+				RHICmdList.SetComputeShader(ComputeShader.GetComputeShader());
 				ComputeShader->SetParameters(RHICmdList, View, DistanceFieldNormal->GetRenderTargetItem(), Parameters, View.GlobalDistanceFieldInfo);
-				DispatchIndirectComputeShader(RHICmdList, *ComputeShader, TileIntersectionResources->ObjectTilesIndirectArguments.Buffer, 0);
+				DispatchIndirectComputeShader(RHICmdList, ComputeShader.GetShader(), TileIntersectionResources->ObjectTilesIndirectArguments.Buffer, 0);
 				ComputeShader->UnsetParameters(RHICmdList, View);
 			}
 		}
@@ -712,9 +664,9 @@ void FDeferredShadingSceneRenderer::RenderDistanceFieldAOScreenGrid(
 
 		TShaderMapRef<FCombineConeVisibilityCS> ComputeShader(View.ShaderMap);
 
-		RHICmdList.SetComputeShader(ComputeShader->GetComputeShader());
+		RHICmdList.SetComputeShader(ComputeShader.GetComputeShader());
 		ComputeShader->SetParameters(RHICmdList, View, DistanceFieldNormal->GetRenderTargetItem(), DownsampledBentNormal->GetRenderTargetItem());
-		DispatchComputeShader(RHICmdList, *ComputeShader, GroupSizeX, GroupSizeY, 1);
+		DispatchComputeShader(RHICmdList, ComputeShader.GetShader(), GroupSizeX, GroupSizeY, 1);
 		ComputeShader->UnsetParameters(RHICmdList, DownsampledBentNormal->GetRenderTargetItem());
 	}
 
diff --git a/Engine/Source/Runtime/Renderer/Private/DistanceFieldShadowing.cpp b/Engine/Source/Runtime/Renderer/Private/DistanceFieldShadowing.cpp
--- a/Engine/Source/Runtime/Renderer/Private/DistanceFieldShadowing.cpp
+++ b/Engine/Source/Runtime/Renderer/Private/DistanceFieldShadowing.cpp
@@ -204,7 +204,7 @@ public:
 			AtlasSizeZ = GDistanceFieldVolumeTextureAtlas.GetSizeZ();
 		}
 
-		FRHIComputeShader* ShaderRHI = GetComputeShader();
+		FRHIComputeShader* ShaderRHI = RHICmdList.GetBoundComputeShader();
 		FGlobalShader::SetParameters<FViewUniformShaderParameters>(RHICmdList, ShaderRHI, View.ViewUniformBuffer);
 		ObjectBufferParameters.Set(RHICmdList, ShaderRHI, SceneObjectBuffers, NumObjectsInBuffer, TextureAtlas, AtlasSizeX, AtlasSizeY, AtlasSizeZ);
 
@@ -239,8 +239,8 @@ public:
 		const auto& SceneObjectBuffers = TSelector<bIsHeightField>()(*SceneData.GetHeightFieldObjectBuffers(), *SceneData.GetCurrentObjectBuffers());
 		auto& CulledObjectBuffers = TSelector<bIsHeightField>()(GShadowCulledHeightFieldObjectBuffers, GShadowCulledObjectBuffers);
 
-		ObjectBufferParameters.UnsetParameters(RHICmdList, GetComputeShader(), SceneObjectBuffers);
-		CulledObjectParameters.UnsetParameters(RHICmdList, GetComputeShader());
+		ObjectBufferParameters.UnsetParameters(RHICmdList, RHICmdList.GetBoundComputeShader(), SceneObjectBuffers);
+		CulledObjectParameters.UnsetParameters(RHICmdList, RHICmdList.GetBoundComputeShader());
 
 		FRHIUnorderedAccessView* OutUAVs[4];
 		OutUAVs[0] = CulledObjectBuffers.Buffers.ObjectIndirectArguments.UAV;
@@ -250,28 +250,17 @@ public:
 		RHICmdList.TransitionResources(EResourceTransitionAccess::EReadable, EResourceTransitionPipeline::EComputeToCompute, OutUAVs, UE_ARRAY_COUNT(OutUAVs));
 	}
 
-	virtual bool Serialize(FArchive& Ar) override
-	{		
-		bool bShaderHasOutdatedParameters = FGlobalShader::Serialize(Ar);
-		Ar << ObjectBufferParameters;
-		Ar << CulledObjectParameters;
-		Ar << ObjectBoundingGeometryIndexCount;
-		Ar << WorldToShadow;
-		Ar << NumShadowHullPlanes;
-		Ar << ShadowBoundingSphere;
-		Ar << ShadowConvexHull;
-		return bShaderHasOutdatedParameters;
-	}
-
 private:
 
-	TDistanceFieldObjectBufferParameters<PrimitiveType> ObjectBufferParameters;
-	TDistanceFieldCulledObjectBufferParameters<PrimitiveType> CulledObjectParameters;
-	FShaderParameter ObjectBoundingGeometryIndexCount;
-	FShaderParameter WorldToShadow;
-	FShaderParameter NumShadowHullPlanes;
-	FShaderParameter ShadowBoundingSphere;
-	FShaderParameter ShadowConvexHull;
+	using FDistanceFieldObjectBufferParametersType = TDistanceFieldObjectBufferParameters<PrimitiveType>;
+	LAYOUT_FIELD(FDistanceFieldObjectBufferParametersType, ObjectBufferParameters)
+	using FDistanceFieldCulledObjectParametersType = TDistanceFieldCulledObjectBufferParameters<PrimitiveType>;
+	LAYOUT_FIELD(FDistanceFieldCulledObjectParametersType, CulledObjectParameters)
+	LAYOUT_FIELD(FShaderParameter, ObjectBoundingGeometryIndexCount)
+	LAYOUT_FIELD(FShaderParameter, WorldToShadow)
+	LAYOUT_FIELD(FShaderParameter, NumShadowHullPlanes)
+	LAYOUT_FIELD(FShaderParameter, ShadowBoundingSphere)
+	LAYOUT_FIELD(FShaderParameter, ShadowConvexHull)
 };
 
 IMPLEMENT_SHADER_TYPE(template<>, TCullObjectsForShadowCS<DFPT_SignedDistanceField>, TEXT("/Engine/Private/DistanceFieldShadowing.usf"), TEXT("CullObjectsForShadowCS"), SF_Compute);
@@ -307,7 +296,7 @@ public:
 
 	void SetParameters(FRHICommandList& RHICmdList, const FSceneView& View, FVector2D NumGroupsValue, const FMatrix& WorldToShadowMatrixValue, float ShadowRadius)
 	{
-		FRHIVertexShader* ShaderRHI = GetVertexShader();
+		FRHIVertexShader* ShaderRHI = RHICmdList.GetBoundVertexShader();
 		FGlobalShader::SetParameters<FViewUniformShaderParameters>(RHICmdList, ShaderRHI, View.ViewUniformBuffer);
 		
 		constexpr bool bIsHeightField = PrimitiveType == DFPT_HeightField;
@@ -344,19 +333,11 @@ public:
 		SetShaderValue(RHICmdList, ShaderRHI, MinExpandRadius, MinExpandRadiusValue);
 	}
 
-	virtual bool Serialize(FArchive& Ar) override
-	{
-		bool bShaderHasOutdatedParameters = FGlobalShader::Serialize(Ar);
-		Ar << ObjectParameters;
-		Ar << WorldToShadow;
-		Ar << MinExpandRadius;
-		return bShaderHasOutdatedParameters;
-	}
-
 private:
-	TDistanceFieldCulledObjectBufferParameters<PrimitiveType> ObjectParameters;
-	FShaderParameter WorldToShadow;
-	FShaderParameter MinExpandRadius;
+	using FDistanceFieldObjectBufferParametersType = TDistanceFieldCulledObjectBufferParameters<PrimitiveType>;
+	LAYOUT_FIELD(FDistanceFieldObjectBufferParametersType, ObjectParameters)
+	LAYOUT_FIELD(FShaderParameter, WorldToShadow)
+	LAYOUT_FIELD(FShaderParameter, MinExpandRadius)
 };
 
 IMPLEMENT_SHADER_TYPE(template<>, TShadowObjectCullVS<DFPT_SignedDistanceField>, TEXT("/Engine/Private/DistanceFieldShadowing.usf"), TEXT("ShadowObjectCullVS"), SF_Vertex);
@@ -397,7 +378,7 @@ public:
 		const FSceneView& View, 
 		FLightTileIntersectionResources* TileIntersectionResources)
 	{
-		FRHIPixelShader* ShaderRHI = GetPixelShader();
+		FRHIPixelShader* ShaderRHI = RHICmdList.GetBoundPixelShader();
 		FGlobalShader::SetParameters<FViewUniformShaderParameters>(RHICmdList, ShaderRHI, View.ViewUniformBuffer);
 
 		constexpr bool bIsHeightField = PrimitiveType == DFPT_HeightField;
@@ -436,18 +417,10 @@ public:
 		LightTileIntersectionParameters.GetUAVs(*TileIntersectionResources, UAVs);
 	}
 
-	// FShader interface.
-	virtual bool Serialize(FArchive& Ar) override
-	{
-		bool bShaderHasOutdatedParameters = FGlobalShader::Serialize(Ar);
-		Ar << ObjectParameters;
-		Ar << LightTileIntersectionParameters;
-		return bShaderHasOutdatedParameters;
-	}
-
 private:
-	TDistanceFieldCulledObjectBufferParameters<PrimitiveType> ObjectParameters;
-	FLightTileIntersectionParameters LightTileIntersectionParameters;
+	using FDistanceFieldObjectBufferParametersType = TDistanceFieldCulledObjectBufferParameters<PrimitiveType>;
+	LAYOUT_FIELD(FDistanceFieldObjectBufferParametersType, ObjectParameters)
+	LAYOUT_FIELD(FLightTileIntersectionParameters, LightTileIntersectionParameters)
 };
 
 #define VARIATION(PrimitiveType) \
@@ -531,7 +504,7 @@ public:
 		FLightTileIntersectionResources* TileIntersectionResources,
 		FRHITexture* PrevOutput)
 	{
-		FRHIComputeShader* ShaderRHI = GetComputeShader();
+		FRHIComputeShader* ShaderRHI = RHICmdList.GetBoundComputeShader();
 
 		FGlobalShader::SetParameters<FViewUniformShaderParameters>(RHICmdList, ShaderRHI, View.ViewUniformBuffer);
 
@@ -631,36 +604,10 @@ public:
 	template<typename TRHICommandList>
 	void UnsetParameters(TRHICommandList& RHICmdList, FSceneRenderTargetItem& ShadowFactorsValue)
 	{
-		ShadowFactors.UnsetUAV(RHICmdList, GetComputeShader());
+		ShadowFactors.UnsetUAV(RHICmdList, RHICmdList.GetBoundComputeShader());
 		RHICmdList.TransitionResource(EResourceTransitionAccess::EReadable, EResourceTransitionPipeline::EComputeToGfx, ShadowFactorsValue.UAV);
 	}
 
-	// FShader interface.
-	virtual bool Serialize(FArchive& Ar) override
-	{
-		bool bShaderHasOutdatedParameters = FGlobalShader::Serialize(Ar);
-		Ar << ShadowFactors;
-		Ar << NumGroups;
-		Ar << LightDirection;
-		Ar << LightPositionAndInvRadius;
-		Ar << LightSourceRadius;
-		Ar << RayStartOffsetDepthScale;
-		Ar << TanLightAngleAndNormalThreshold;
-		Ar << ScissorRectMinAndSize;
-		Ar << ObjectParameters;
-		Ar << SceneTextureParameters;
-		Ar << LightTileIntersectionParameters;
-		Ar << WorldToShadow;
-		Ar << TwoSidedMeshDistanceBias;
-		Ar << MinDepth;
-		Ar << MaxDepth;
-		Ar << DownsampleFactor;
-		Ar << InvOutputBufferSize;
-		Ar << ShadowFactorsTexture;
-		Ar << ShadowFactorsSampler;
-		return bShaderHasOutdatedParameters;
-	}
-
 	static const TCHAR* GetSourceFilename()
 	{
 		return TEXT("/Engine/Private/DistanceFieldShadowing.usf");
@@ -672,26 +619,26 @@ public:
 	}
 
 private:
-
-	FRWShaderParameter ShadowFactors;
-	FShaderParameter NumGroups;
-	FShaderParameter LightDirection;
-	FShaderParameter LightPositionAndInvRadius;
-	FShaderParameter LightSourceRadius;
-	FShaderParameter RayStartOffsetDepthScale;
-	FShaderParameter TanLightAngleAndNormalThreshold;
-	FShaderParameter ScissorRectMinAndSize;
-	TDistanceFieldCulledObjectBufferParameters<PrimitiveType> ObjectParameters;
-	FSceneTextureShaderParameters SceneTextureParameters;
-	FLightTileIntersectionParameters LightTileIntersectionParameters;
-	FShaderParameter WorldToShadow;
-	FShaderParameter TwoSidedMeshDistanceBias;
-	FShaderParameter MinDepth;
-	FShaderParameter MaxDepth;
-	FShaderParameter DownsampleFactor;
-	FShaderParameter InvOutputBufferSize;
-	FShaderResourceParameter ShadowFactorsTexture;
-	FShaderResourceParameter ShadowFactorsSampler;
+	LAYOUT_FIELD(FRWShaderParameter, ShadowFactors);
+	LAYOUT_FIELD(FShaderParameter, NumGroups);
+	LAYOUT_FIELD(FShaderParameter, LightDirection);
+	LAYOUT_FIELD(FShaderParameter, LightPositionAndInvRadius);
+	LAYOUT_FIELD(FShaderParameter, LightSourceRadius);
+	LAYOUT_FIELD(FShaderParameter, RayStartOffsetDepthScale);
+	LAYOUT_FIELD(FShaderParameter, TanLightAngleAndNormalThreshold);
+	LAYOUT_FIELD(FShaderParameter, ScissorRectMinAndSize);
+	using FDistanceFieldObjectBufferParametersType = TDistanceFieldCulledObjectBufferParameters<PrimitiveType>;
+	LAYOUT_FIELD(FDistanceFieldObjectBufferParametersType, ObjectParameters)
+	LAYOUT_FIELD(FSceneTextureShaderParameters, SceneTextureParameters);
+	LAYOUT_FIELD(FLightTileIntersectionParameters, LightTileIntersectionParameters);
+	LAYOUT_FIELD(FShaderParameter, WorldToShadow);
+	LAYOUT_FIELD(FShaderParameter, TwoSidedMeshDistanceBias);
+	LAYOUT_FIELD(FShaderParameter, MinDepth);
+	LAYOUT_FIELD(FShaderParameter, MaxDepth);
+	LAYOUT_FIELD(FShaderParameter, DownsampleFactor);
+	LAYOUT_FIELD(FShaderParameter, InvOutputBufferSize);
+	LAYOUT_FIELD(FShaderResourceParameter, ShadowFactorsTexture);
+	LAYOUT_FIELD(FShaderResourceParameter, ShadowFactorsSampler);
 };
 
 // #define avoids a lot of code duplication
@@ -749,7 +696,7 @@ public:
 
 	void SetParameters(FRHICommandList& RHICmdList, const FSceneView& View, const FProjectedShadowInfo* ShadowInfo, const FIntRect& ScissorRect, TRefCountPtr<IPooledRenderTarget>& ShadowFactorsTextureValue)
 	{
-		FRHIPixelShader* ShaderRHI = GetPixelShader();
+		FRHIPixelShader* ShaderRHI = RHICmdList.GetBoundPixelShader();
 
 		FGlobalShader::SetParameters<FViewUniformShaderParameters>(RHICmdList, ShaderRHI, View.ViewUniformBuffer);
 		SceneTextureParameters.Set(RHICmdList, ShaderRHI, View.FeatureLevel, ESceneTextureSetupMode::All);
@@ -781,31 +728,15 @@ public:
 		}
 	}
 
-	// FShader interface.
-	virtual bool Serialize(FArchive& Ar) override
-	{
-		bool bShaderHasOutdatedParameters = FGlobalShader::Serialize(Ar);
-		Ar << SceneTextureParameters;
-		Ar << ShadowFactorsTexture;
-		Ar << ShadowFactorsSampler;
-		Ar << ScissorRectMinAndSize;
-		Ar << FadePlaneOffset;
-		Ar << InvFadePlaneLength;
-		Ar << NearFadePlaneOffset;
-		Ar << InvNearFadePlaneLength;
-		return bShaderHasOutdatedParameters;
-	}
-
 private:
-
-	FSceneTextureShaderParameters SceneTextureParameters;
-	FShaderResourceParameter ShadowFactorsTexture;
-	FShaderResourceParameter ShadowFactorsSampler;
-	FShaderParameter ScissorRectMinAndSize;
-	FShaderParameter FadePlaneOffset;
-	FShaderParameter InvFadePlaneLength;
-	FShaderParameter NearFadePlaneOffset;
-	FShaderParameter InvNearFadePlaneLength;
+	LAYOUT_FIELD(FSceneTextureShaderParameters, SceneTextureParameters);
+	LAYOUT_FIELD(FShaderResourceParameter, ShadowFactorsTexture);
+	LAYOUT_FIELD(FShaderResourceParameter, ShadowFactorsSampler);
+	LAYOUT_FIELD(FShaderParameter, ScissorRectMinAndSize);
+	LAYOUT_FIELD(FShaderParameter, FadePlaneOffset);
+	LAYOUT_FIELD(FShaderParameter, InvFadePlaneLength);
+	LAYOUT_FIELD(FShaderParameter, NearFadePlaneOffset);
+	LAYOUT_FIELD(FShaderParameter, InvNearFadePlaneLength);
 };
 
 IMPLEMENT_SHADER_TYPE(template<>,TDistanceFieldShadowingUpsamplePS<true>,TEXT("/Engine/Private/DistanceFieldShadowing.usf"),TEXT("DistanceFieldShadowingUpsamplePS"),SF_Pixel);
@@ -841,7 +772,7 @@ public:
 	}
 	void SetParameters(FRHICommandList& RHICmdList, const FSceneView& View, FLightTileIntersectionResources* TileIntersectionResources)
 	{
-		FRHIComputeShader* ShaderRHI = GetComputeShader();
+		FRHIComputeShader* ShaderRHI = RHICmdList.GetBoundComputeShader();
 
 		FGlobalShader::SetParameters<FViewUniformShaderParameters>(RHICmdList, ShaderRHI, View.ViewUniformBuffer);
 
@@ -855,7 +786,7 @@ public:
 
 	void UnsetParameters(FRHICommandList& RHICmdList, const FSceneView& View, FLightTileIntersectionResources* TileIntersectionResources)
 	{
-		TileIntersectionParameters.UnsetParameters(RHICmdList, GetComputeShader());
+		TileIntersectionParameters.UnsetParameters(RHICmdList, RHICmdList.GetBoundComputeShader());
 
 		TArray<FRHIUnorderedAccessView*> UAVs;
 		TileIntersectionParameters.GetUAVs(*TileIntersectionResources, UAVs);
@@ -863,16 +794,8 @@ public:
 		RHICmdList.TransitionResources(EResourceTransitionAccess::EReadable, EResourceTransitionPipeline::EComputeToCompute, UAVs.GetData(), UAVs.Num());
 	}
 
-	virtual bool Serialize(FArchive& Ar) override
-	{		
-		bool bShaderHasOutdatedParameters = FGlobalShader::Serialize(Ar);
-		Ar << TileIntersectionParameters;
-		return bShaderHasOutdatedParameters;
-	}
-
 private:
-
-	FLightTileIntersectionParameters TileIntersectionParameters;
+	LAYOUT_FIELD(FLightTileIntersectionParameters, TileIntersectionParameters);
 };
 
 IMPLEMENT_SHADER_TYPE(,FComputeCulledObjectStartOffsetCS,TEXT("/Engine/Private/DistanceFieldShadowing.usf"),TEXT("ComputeCulledTilesStartOffsetCS"),SF_Compute);
@@ -918,8 +841,8 @@ void ScatterObjectsToShadowTiles(
 		GraphicsPSOInit.DepthStencilState = TStaticDepthStencilState<false, CF_Always>::GetRHI();
 		GraphicsPSOInit.BlendState = TStaticBlendState<>::GetRHI();
 		GraphicsPSOInit.BoundShaderState.VertexDeclarationRHI = GetVertexDeclarationFVector4();
-		GraphicsPSOInit.BoundShaderState.VertexShaderRHI = GETSAFERHISHADER_VERTEX(*VertexShader);
-		GraphicsPSOInit.BoundShaderState.PixelShaderRHI = GETSAFERHISHADER_PIXEL(*PixelShader);
+		GraphicsPSOInit.BoundShaderState.VertexShaderRHI = VertexShader.GetVertexShader();
+		GraphicsPSOInit.BoundShaderState.PixelShaderRHI = PixelShader.GetPixelShader();
 		GraphicsPSOInit.PrimitiveType = PT_TriangleList;
 		SetGraphicsPipelineState(RHICmdList, GraphicsPSOInit);
 
@@ -982,10 +905,10 @@ void CullDistanceFieldObjectsForLight_Internal(
 			RHICmdList.ClearUAVUint(CulledObjectBuffers.Buffers.ObjectIndirectArguments.UAV, FUintVector4(0, 0, 0, 0));
 
 			TShaderMapRef<TCullObjectsForShadowCS<PrimitiveType>> ComputeShader(GetGlobalShaderMap(Scene->GetFeatureLevel()));
-			RHICmdList.SetComputeShader(ComputeShader->GetComputeShader());
+			RHICmdList.SetComputeShader(ComputeShader.GetComputeShader());
 			ComputeShader->SetParameters(RHICmdList, Scene, View, WorldToShadowValue, NumPlanes, PlaneData, ShadowBoundingSphereValue);
 
-			DispatchComputeShader(RHICmdList, *ComputeShader, FMath::DivideAndRoundUp<uint32>(NumObjectsInBuffer, UpdateObjectsGroupSize), 1, 1);
+			DispatchComputeShader(RHICmdList, ComputeShader.GetShader(), FMath::DivideAndRoundUp<uint32>(NumObjectsInBuffer, UpdateObjectsGroupSize), 1, 1);
 			ComputeShader->UnsetParameters(RHICmdList, Scene);
 		}
 	}
@@ -1042,9 +965,9 @@ void CullDistanceFieldObjectsForLight_Internal(
 
 			// Compute the start offset for each tile's culled object data
 			TShaderMapRef<FComputeCulledObjectStartOffsetCS> ComputeShader(View.ShaderMap);
-			RHICmdList.SetComputeShader(ComputeShader->GetComputeShader());
+			RHICmdList.SetComputeShader(ComputeShader.GetComputeShader());
 			ComputeShader->SetParameters(RHICmdList, View, TileIntersectionResources.Get());
-			DispatchComputeShader(RHICmdList, *ComputeShader, GroupSizeX, GroupSizeY, 1);
+			DispatchComputeShader(RHICmdList, ComputeShader.GetShader(), GroupSizeX, GroupSizeY, 1);
 			ComputeShader->UnsetParameters(RHICmdList, View, TileIntersectionResources.Get());
 		}
 
@@ -1192,10 +1115,10 @@ void RayTraceShadowsDispatch(
 
 	auto DispatchTemplatedCS = [&](auto CS)
 	{
-		RHICmdList.SetComputeShader(CS->GetComputeShader());
+		RHICmdList.SetComputeShader(CS.GetComputeShader());
 		FSceneRenderTargetItem& RayTracedShadowsRTI = ProjectedShadowInfo->RayTracedShadowsRT->GetRenderTargetItem();
 		CS->SetParameters(RHICmdList, View, ProjectedShadowInfo, RayTracedShadowsRTI, FVector2D(GroupSizeX, GroupSizeY), ScissorRect, TileIntersectionResources, PrevOutput);
-		DispatchComputeShader(RHICmdList, *CS, GroupSizeX, GroupSizeY, 1);
+		DispatchComputeShader(RHICmdList, CS.GetShader(), GroupSizeX, GroupSizeY, 1);
 		CS->UnsetParameters(RHICmdList, RayTracedShadowsRTI);
 	};
 
@@ -1434,14 +1357,14 @@ void FProjectedShadowInfo::RenderRayTracedDistanceFieldProjection(FRHICommandLis
 
 				TShaderMapRef<FPostProcessVS> VertexShader(View.ShaderMap);
 				GraphicsPSOInit.BoundShaderState.VertexDeclarationRHI = GFilterVertexDeclaration.VertexDeclarationRHI;
-				GraphicsPSOInit.BoundShaderState.VertexShaderRHI = GETSAFERHISHADER_VERTEX(*VertexShader);
+				GraphicsPSOInit.BoundShaderState.VertexShaderRHI = VertexShader.GetVertexShader();
 				GraphicsPSOInit.PrimitiveType = PT_TriangleList;
 				GraphicsPSOInit.bDepthBounds = bDirectionalLight;
 
 				if (GFullResolutionDFShadowing)
 				{
 					TShaderMapRef<TDistanceFieldShadowingUpsamplePS<false> > PixelShader(View.ShaderMap);
-					GraphicsPSOInit.BoundShaderState.PixelShaderRHI = GETSAFERHISHADER_PIXEL(*PixelShader);
+					GraphicsPSOInit.BoundShaderState.PixelShaderRHI = PixelShader.GetPixelShader();
 					SetGraphicsPipelineState(RHICmdList, GraphicsPSOInit);
 					VertexShader->SetParameters(RHICmdList, View.ViewUniformBuffer);
 					PixelShader->SetParameters(RHICmdList, View, this, ScissorRect, RayTracedShadowsRT);
@@ -1449,7 +1372,7 @@ void FProjectedShadowInfo::RenderRayTracedDistanceFieldProjection(FRHICommandLis
 				else
 				{
 					TShaderMapRef<TDistanceFieldShadowingUpsamplePS<true> > PixelShader(View.ShaderMap);
-					GraphicsPSOInit.BoundShaderState.PixelShaderRHI = GETSAFERHISHADER_PIXEL(*PixelShader);
+					GraphicsPSOInit.BoundShaderState.PixelShaderRHI = PixelShader.GetPixelShader();
 					SetGraphicsPipelineState(RHICmdList, GraphicsPSOInit);
 					VertexShader->SetParameters(RHICmdList, View.ViewUniformBuffer);
 					PixelShader->SetParameters(RHICmdList, View, this, ScissorRect, RayTracedShadowsRT);
@@ -1469,7 +1392,7 @@ void FProjectedShadowInfo::RenderRayTracedDistanceFieldProjection(FRHICommandLis
 					ScissorRect.Width() / GetDFShadowDownsampleFactor(), ScissorRect.Height() / GetDFShadowDownsampleFactor(),
 					FIntPoint(ScissorRect.Width(), ScissorRect.Height()),
 					GetBufferSizeForDFShadows(),
-					*VertexShader);
+					VertexShader);
 			}
 			RHICmdList.EndRenderPass();
 		}
diff --git a/Engine/Source/Runtime/Renderer/Private/DistanceFieldVisualization.cpp b/Engine/Source/Runtime/Renderer/Private/DistanceFieldVisualization.cpp
--- a/Engine/Source/Runtime/Renderer/Private/DistanceFieldVisualization.cpp
+++ b/Engine/Source/Runtime/Renderer/Private/DistanceFieldVisualization.cpp
@@ -59,7 +59,7 @@ public:
 		const FDistanceFieldAOParameters& Parameters,
 		const FGlobalDistanceFieldInfo& GlobalDistanceFieldInfo)
 	{
-		FRHIComputeShader* ShaderRHI = GetComputeShader();
+		FRHIComputeShader* ShaderRHI = RHICmdList.GetBoundComputeShader();
 
 		FGlobalShader::SetParameters<FViewUniformShaderParameters>(RHICmdList, ShaderRHI, View.ViewUniformBuffer);
 
@@ -92,30 +92,17 @@ public:
 	void UnsetParameters(FRHICommandList& RHICmdList, FSceneRenderTargetItem& VisualizeMeshDistanceFieldsValue)
 	{
 		RHICmdList.TransitionResource(EResourceTransitionAccess::EReadable, EResourceTransitionPipeline::EComputeToCompute, VisualizeMeshDistanceFieldsValue.UAV);
-		VisualizeMeshDistanceFields.UnsetUAV(RHICmdList, GetComputeShader());
-	}
-
-	// FShader interface.
-	virtual bool Serialize(FArchive& Ar) override
-	{
-		bool bShaderHasOutdatedParameters = FGlobalShader::Serialize(Ar);
-		Ar << VisualizeMeshDistanceFields;
-		Ar << NumGroups;
-		Ar << ObjectParameters;
-		Ar << SceneTextureParameters;
-		Ar << AOParameters;
-		Ar << GlobalDistanceFieldParameters;
-		return bShaderHasOutdatedParameters;
+		VisualizeMeshDistanceFields.UnsetUAV(RHICmdList, RHICmdList.GetBoundComputeShader());
 	}
 
 private:
 
-	FRWShaderParameter VisualizeMeshDistanceFields;
-	FShaderParameter NumGroups;
-	TDistanceFieldCulledObjectBufferParameters<DFPT_SignedDistanceField> ObjectParameters;
-	FSceneTextureShaderParameters SceneTextureParameters;
-	FAOParameters AOParameters;
-	FGlobalDistanceFieldParameters GlobalDistanceFieldParameters;
+	LAYOUT_FIELD(FRWShaderParameter, VisualizeMeshDistanceFields);
+	LAYOUT_FIELD(FShaderParameter, NumGroups);
+	LAYOUT_FIELD((TDistanceFieldCulledObjectBufferParameters<DFPT_SignedDistanceField>), ObjectParameters);
+	LAYOUT_FIELD(FSceneTextureShaderParameters, SceneTextureParameters);
+	LAYOUT_FIELD(FAOParameters, AOParameters);
+	LAYOUT_FIELD(FGlobalDistanceFieldParameters, GlobalDistanceFieldParameters);
 };
 
 IMPLEMENT_SHADER_TYPE(template<>,TVisualizeMeshDistanceFieldCS<true>,TEXT("/Engine/Private/DistanceFieldVisualization.usf"),TEXT("VisualizeMeshDistanceFieldCS"),SF_Compute);
@@ -150,28 +137,19 @@ public:
 
 	void SetParameters(FRHICommandList& RHICmdList, const FSceneView& View, TRefCountPtr<IPooledRenderTarget>& VisualizeDistanceField)
 	{
-		FRHIPixelShader* ShaderRHI = GetPixelShader();
+		FRHIPixelShader* ShaderRHI = RHICmdList.GetBoundPixelShader();
 
 		FGlobalShader::SetParameters<FViewUniformShaderParameters>(RHICmdList, ShaderRHI, View.ViewUniformBuffer);
 		SceneTextureParameters.Set(RHICmdList, ShaderRHI, View.FeatureLevel, ESceneTextureSetupMode::All);
 
 		SetTextureParameter(RHICmdList, ShaderRHI, VisualizeDistanceFieldTexture, VisualizeDistanceFieldSampler, TStaticSamplerState<SF_Bilinear>::GetRHI(), VisualizeDistanceField->GetRenderTargetItem().ShaderResourceTexture);
 	}
-	// FShader interface.
-	virtual bool Serialize(FArchive& Ar) override
-	{
-		bool bShaderHasOutdatedParameters = FGlobalShader::Serialize(Ar);
-		Ar << SceneTextureParameters;
-		Ar << VisualizeDistanceFieldTexture;
-		Ar << VisualizeDistanceFieldSampler;
-		return bShaderHasOutdatedParameters;
-	}
 
 private:
 
-	FSceneTextureShaderParameters SceneTextureParameters;
-	FShaderResourceParameter VisualizeDistanceFieldTexture;
-	FShaderResourceParameter VisualizeDistanceFieldSampler;
+	LAYOUT_FIELD(FSceneTextureShaderParameters, SceneTextureParameters);
+	LAYOUT_FIELD(FShaderResourceParameter, VisualizeDistanceFieldTexture);
+	LAYOUT_FIELD(FShaderResourceParameter, VisualizeDistanceFieldSampler);
 };
 
 IMPLEMENT_SHADER_TYPE(,FVisualizeDistanceFieldUpsamplePS,TEXT("/Engine/Private/DistanceFieldVisualization.usf"),TEXT("VisualizeDistanceFieldUpsamplePS"),SF_Pixel);
@@ -228,9 +206,9 @@ void FDeferredShadingSceneRenderer::RenderMeshDistanceFieldVisualization(FRHICom
 
 						TShaderMapRef<TVisualizeMeshDistanceFieldCS<true> > ComputeShader(ViewInfo.ShaderMap);
 
-						RHICmdList.SetComputeShader(ComputeShader->GetComputeShader());
+						RHICmdList.SetComputeShader(ComputeShader.GetComputeShader());
 						ComputeShader->SetParameters(RHICmdList, ViewInfo, VisualizeResultRTI, FVector2D(GroupSizeX, GroupSizeY), Parameters, View.GlobalDistanceFieldInfo);
-						DispatchComputeShader(RHICmdList, *ComputeShader, GroupSizeX, GroupSizeY, 1);
+						DispatchComputeShader(RHICmdList, ComputeShader.GetShader(), GroupSizeX, GroupSizeY, 1);
 
 						ComputeShader->UnsetParameters(RHICmdList, VisualizeResultRTI);
 					}
@@ -238,9 +216,9 @@ void FDeferredShadingSceneRenderer::RenderMeshDistanceFieldVisualization(FRHICom
 					{
 						TShaderMapRef<TVisualizeMeshDistanceFieldCS<false> > ComputeShader(ViewInfo.ShaderMap);
 
-						RHICmdList.SetComputeShader(ComputeShader->GetComputeShader());
+						RHICmdList.SetComputeShader(ComputeShader.GetComputeShader());
 						ComputeShader->SetParameters(RHICmdList, ViewInfo, VisualizeResultRTI, FVector2D(GroupSizeX, GroupSizeY), Parameters, View.GlobalDistanceFieldInfo);
-						DispatchComputeShader(RHICmdList, *ComputeShader, GroupSizeX, GroupSizeY, 1);
+						DispatchComputeShader(RHICmdList, ComputeShader.GetShader(), GroupSizeX, GroupSizeY, 1);
 
 						ComputeShader->UnsetParameters(RHICmdList, VisualizeResultRTI);
 					}
@@ -277,8 +255,8 @@ void FDeferredShadingSceneRenderer::RenderMeshDistanceFieldVisualization(FRHICom
 					TShaderMapRef<FVisualizeDistanceFieldUpsamplePS> PixelShader( ViewInfo.ShaderMap );
 
 					GraphicsPSOInit.BoundShaderState.VertexDeclarationRHI = GFilterVertexDeclaration.VertexDeclarationRHI;
-					GraphicsPSOInit.BoundShaderState.VertexShaderRHI = GETSAFERHISHADER_VERTEX(*VertexShader);
-					GraphicsPSOInit.BoundShaderState.PixelShaderRHI = GETSAFERHISHADER_PIXEL(*PixelShader);
+					GraphicsPSOInit.BoundShaderState.VertexShaderRHI = VertexShader.GetVertexShader();
+					GraphicsPSOInit.BoundShaderState.PixelShaderRHI = PixelShader.GetPixelShader();
 					GraphicsPSOInit.PrimitiveType = PT_TriangleList;
 					SetGraphicsPipelineState(RHICmdList, GraphicsPSOInit);
 
@@ -292,7 +270,7 @@ void FDeferredShadingSceneRenderer::RenderMeshDistanceFieldVisualization(FRHICom
 						ViewInfo.ViewRect.Width() / GAODownsampleFactor, ViewInfo.ViewRect.Height() / GAODownsampleFactor,
 						FIntPoint(ViewInfo.ViewRect.Width(), ViewInfo.ViewRect.Height()),
 						GetBufferSizeForAO(),
-						*VertexShader);
+						VertexShader);
 				}
 
 				FSceneRenderTargets::Get(RHICmdList).FinishRenderingSceneColor(RHICmdList);
diff --git a/Engine/Source/Runtime/Renderer/Private/DistortionRendering.cpp b/Engine/Source/Runtime/Renderer/Private/DistortionRendering.cpp
--- a/Engine/Source/Runtime/Renderer/Private/DistortionRendering.cpp
+++ b/Engine/Source/Runtime/Renderer/Private/DistortionRendering.cpp
@@ -109,7 +109,7 @@ public:
 	void SetParameters(const FRenderingCompositePassContext& Context, const FViewInfo& View, IPooledRenderTarget& DistortionRT)
 	{
 		FSceneRenderTargets& SceneContext = FSceneRenderTargets::Get(Context.RHICmdList);
-		FRHIPixelShader* ShaderRHI = GetPixelShader();
+		FRHIPixelShader* ShaderRHI = Context.RHICmdList.GetBoundPixelShader();
 
 		FGlobalShader::SetParameters<FViewUniformShaderParameters>(Context.RHICmdList, ShaderRHI, View.ViewUniformBuffer);
 
@@ -136,13 +136,6 @@ public:
 			);
 	}
 
-	virtual bool Serialize(FArchive& Ar) override
-	{
-		bool bShaderHasOutdatedParameters = FGlobalShader::Serialize(Ar);
-		Ar << DistortionTexture << DistortionTextureSampler << SceneColorTexture << SceneColorTextureSampler;
-		return bShaderHasOutdatedParameters;
-	}
-
 	static const TCHAR* GetSourceFilename()
 	{
 		return TEXT("/Engine/Private/DistortApplyScreenPS.usf");
@@ -154,10 +147,10 @@ public:
 	}
 
 private:
-	FShaderResourceParameter DistortionTexture;
-	FShaderResourceParameter DistortionTextureSampler;
-	FShaderResourceParameter SceneColorTexture;
-	FShaderResourceParameter SceneColorTextureSampler;
+	LAYOUT_FIELD(FShaderResourceParameter, DistortionTexture);
+	LAYOUT_FIELD(FShaderResourceParameter, DistortionTextureSampler);
+	LAYOUT_FIELD(FShaderResourceParameter, SceneColorTexture);
+	LAYOUT_FIELD(FShaderResourceParameter, SceneColorTextureSampler);
 
 	static void ModifyCompilationEnvironment(const FGlobalShaderPermutationParameters& Parameters, FShaderCompilerEnvironment& OutEnvironment)
 	{
@@ -206,7 +199,7 @@ public:
 	void SetParameters(const FRenderingCompositePassContext& Context, const FViewInfo& View, FRHITexture* PassTexture)
 	{
 		FSceneRenderTargets& SceneContext = FSceneRenderTargets::Get(Context.RHICmdList);
-		FRHIPixelShader* ShaderRHI = GetPixelShader();
+		FRHIPixelShader* ShaderRHI = Context.RHICmdList.GetBoundPixelShader();
 
 		FGlobalShader::SetParameters<FViewUniformShaderParameters>(Context.RHICmdList, ShaderRHI, View.ViewUniformBuffer);
 
@@ -220,13 +213,6 @@ public:
 			);
 	}
 
-	virtual bool Serialize(FArchive& Ar) override
-	{
-		bool bShaderHasOutdatedParameters = FGlobalShader::Serialize(Ar);
-		Ar << SceneColorTexture << SceneColorTextureSampler;
-		return bShaderHasOutdatedParameters;
-	}
-
 	static const TCHAR* GetSourceFilename()
 	{
 		return TEXT("/Engine/Private/DistortApplyScreenPS.usf");
@@ -238,8 +224,8 @@ public:
 	}
 
 private:
-	FShaderResourceParameter SceneColorTexture;
-	FShaderResourceParameter SceneColorTextureSampler;
+	LAYOUT_FIELD(FShaderResourceParameter, SceneColorTexture);
+	LAYOUT_FIELD(FShaderResourceParameter, SceneColorTextureSampler);
 
 	static void ModifyCompilationEnvironment(const FGlobalShaderPermutationParameters& Parameters, FShaderCompilerEnvironment& OutEnvironment)
 	{
@@ -285,7 +271,7 @@ protected:
 
 	static bool ShouldCompilePermutation(const FMeshMaterialShaderPermutationParameters& Parameters)
 	{
-		return Parameters.Material && IsTranslucentBlendMode(Parameters.Material->GetBlendMode()) && Parameters.Material->IsDistorted();
+		return IsTranslucentBlendMode(Parameters.MaterialParameters.BlendMode) && Parameters.MaterialParameters.bIsDistorted;
 	}
 };
 
@@ -310,7 +296,7 @@ protected:
 	static bool ShouldCompilePermutation(const FMeshMaterialShaderPermutationParameters& Parameters)
 	{
 		return FBaseHS::ShouldCompilePermutation(Parameters)
-			&& Parameters.Material && IsTranslucentBlendMode(Parameters.Material->GetBlendMode()) && Parameters.Material->IsDistorted();
+			&& IsTranslucentBlendMode(Parameters.MaterialParameters.BlendMode) && Parameters.MaterialParameters.bIsDistorted;
 	}
 };
 
@@ -334,7 +320,7 @@ protected:
 	static bool ShouldCompilePermutation(const FMeshMaterialShaderPermutationParameters& Parameters)
 	{
 		return FBaseDS::ShouldCompilePermutation(Parameters)
-			&& Parameters.Material && IsTranslucentBlendMode(Parameters.Material->GetBlendMode()) && Parameters.Material->IsDistorted();
+			&& IsTranslucentBlendMode(Parameters.MaterialParameters.BlendMode) && Parameters.MaterialParameters.bIsDistorted;
 	}
 };
 
@@ -354,7 +340,7 @@ class FDistortionMeshPS : public FMeshMaterialShader
 public:
 	static bool ShouldCompilePermutation(const FMeshMaterialShaderPermutationParameters& Parameters)
 	{
-		return Parameters.Material && IsTranslucentBlendMode(Parameters.Material->GetBlendMode()) && Parameters.Material->IsDistorted();
+		return IsTranslucentBlendMode(Parameters.MaterialParameters.BlendMode) && Parameters.MaterialParameters.bIsDistorted;
 	}
 
 	FDistortionMeshPS(const ShaderMetaType::CompiledShaderInitializerType& Initializer)
@@ -386,12 +372,6 @@ public:
 	FDistortionMeshPS()
 	{
 	}
-
-	virtual bool Serialize(FArchive& Ar) override
-	{
-		bool bShaderHasOutdatedParameters = FMeshMaterialShader::Serialize(Ar);
-		return bShaderHasOutdatedParameters;
-	}
 };
 
 //** distortion accumulate pixel shader type implementation */
@@ -433,8 +413,8 @@ static void DrawDistortionApplyScreenPass(FRHICommandListImmediate& RHICmdList,
 		kStencilMaskBit, kStencilMaskBit>::GetRHI();
 
 	GraphicsPSOInit.BoundShaderState.VertexDeclarationRHI = GFilterVertexDeclaration.VertexDeclarationRHI;
-	GraphicsPSOInit.BoundShaderState.VertexShaderRHI = GETSAFERHISHADER_VERTEX(*VertexShader);
-	GraphicsPSOInit.BoundShaderState.PixelShaderRHI = GETSAFERHISHADER_PIXEL(*PixelShader);
+	GraphicsPSOInit.BoundShaderState.VertexShaderRHI = VertexShader.GetVertexShader();
+	GraphicsPSOInit.BoundShaderState.PixelShaderRHI = PixelShader.GetPixelShader();
 	GraphicsPSOInit.PrimitiveType = PT_TriangleList;
 
 	SetGraphicsPipelineState(RHICmdList, GraphicsPSOInit);
@@ -452,7 +432,7 @@ static void DrawDistortionApplyScreenPass(FRHICommandListImmediate& RHICmdList,
 		View.ViewRect.Width(), View.ViewRect.Height(),
 		View.ViewRect.Size(),
 		SceneContext.GetBufferSizeXY(),
-		*VertexShader,
+		VertexShader,
 		EDRF_UseTriangleOptimization);
 }
 
@@ -478,8 +458,8 @@ static void DrawDistortionMergePass(FRHICommandListImmediate& RHICmdList, FScene
 		kStencilMaskBit, kStencilMaskBit>::GetRHI();
 
 	GraphicsPSOInit.BoundShaderState.VertexDeclarationRHI = GFilterVertexDeclaration.VertexDeclarationRHI;
-	GraphicsPSOInit.BoundShaderState.VertexShaderRHI = GETSAFERHISHADER_VERTEX(*VertexShader);
-	GraphicsPSOInit.BoundShaderState.PixelShaderRHI = GETSAFERHISHADER_PIXEL(*PixelShader);
+	GraphicsPSOInit.BoundShaderState.VertexShaderRHI = VertexShader.GetVertexShader();
+	GraphicsPSOInit.BoundShaderState.PixelShaderRHI = PixelShader.GetPixelShader();
 	GraphicsPSOInit.PrimitiveType = PT_TriangleList;
 
 	SetGraphicsPipelineState(RHICmdList, GraphicsPSOInit);
@@ -496,7 +476,7 @@ static void DrawDistortionMergePass(FRHICommandListImmediate& RHICmdList, FScene
 		View.ViewRect.Width(), View.ViewRect.Height(),
 		View.ViewRect.Size(),
 		SceneContext.GetBufferSizeXY(),
-		*VertexShader,
+		VertexShader,
 		EDRF_UseTriangleOptimization);
 }
 
@@ -744,8 +724,9 @@ void FDistortionMeshProcessor::AddMeshBatch(const FMeshBatch& RESTRICT MeshBatch
 		const FMaterialRenderProxy& MaterialRenderProxy = FallbackMaterialRenderProxyPtr ? *FallbackMaterialRenderProxyPtr : *MeshBatch.MaterialRenderProxy;
 
 		const EBlendMode BlendMode = Material.GetBlendMode();
-		const ERasterizerFillMode MeshFillMode = ComputeMeshFillMode(MeshBatch, Material);
-		const ERasterizerCullMode MeshCullMode = ComputeMeshCullMode(MeshBatch, Material);
+		const FMeshDrawingPolicyOverrideSettings OverrideSettings = ComputeMeshOverrideSettings(MeshBatch);
+		const ERasterizerFillMode MeshFillMode = ComputeMeshFillMode(MeshBatch, Material, OverrideSettings);
+		const ERasterizerCullMode MeshCullMode = ComputeMeshCullMode(MeshBatch, Material, OverrideSettings);
 		const bool bIsTranslucent = IsTranslucentBlendMode(BlendMode);
 
 		if (bIsTranslucent
@@ -762,10 +743,10 @@ void GetDistortionPassShaders(
 	const FMaterial& Material,
 	FVertexFactoryType* VertexFactoryType,
 	ERHIFeatureLevel::Type FeatureLevel,
-	FDistortionMeshHS*& HullShader,
-	FDistortionMeshDS*& DomainShader,
-	FDistortionMeshVS*& VertexShader,
-	FDistortionMeshPS*& PixelShader)
+	TShaderRef<FDistortionMeshHS>& HullShader,
+	TShaderRef<FDistortionMeshDS>& DomainShader,
+	TShaderRef<FDistortionMeshVS>& VertexShader,
+	TShaderRef<FDistortionMeshPS>& PixelShader)
 {
 	const EMaterialTessellationMode MaterialTessellationMode = Material.GetTessellationMode();
 
diff --git a/Engine/Source/Runtime/Renderer/Private/EditorPrimitivesRendering.cpp b/Engine/Source/Runtime/Renderer/Private/EditorPrimitivesRendering.cpp
--- a/Engine/Source/Runtime/Renderer/Private/EditorPrimitivesRendering.cpp
+++ b/Engine/Source/Runtime/Renderer/Private/EditorPrimitivesRendering.cpp
@@ -81,8 +81,9 @@ void FEditorPrimitivesBasePassMeshProcessor::ProcessDeferredShadingPath(const FM
 		SetTranslucentRenderState(DrawRenderState, Material);
 	}
 
-	ERasterizerFillMode MeshFillMode = ComputeMeshFillMode(MeshBatch, Material);
-	ERasterizerCullMode MeshCullMode = ComputeMeshCullMode(MeshBatch, Material);
+	const FMeshDrawingPolicyOverrideSettings OverrideSettings = ComputeMeshOverrideSettings(MeshBatch);
+	ERasterizerFillMode MeshFillMode = ComputeMeshFillMode(MeshBatch, Material, OverrideSettings);
+	ERasterizerCullMode MeshCullMode = ComputeMeshCullMode(MeshBatch, Material, OverrideSettings);
 	
 	TBasePassShaderElementData<LightMapPolicyType> ShaderElementData(nullptr);
 	ShaderElementData.InitializeMeshMaterialData(ViewIfDynamicMeshCommand, PrimitiveSceneProxy, MeshBatch, StaticMeshId, false);
@@ -135,8 +136,9 @@ void FEditorPrimitivesBasePassMeshProcessor::ProcessMobileShadingPath(const FMes
 		MobileBasePass::SetTranslucentRenderState(DrawRenderState, Material);
 	}
 
-	ERasterizerFillMode MeshFillMode = ComputeMeshFillMode(MeshBatch, Material);
-	ERasterizerCullMode MeshCullMode = ComputeMeshCullMode(MeshBatch, Material);
+	const FMeshDrawingPolicyOverrideSettings OverrideSettings = ComputeMeshOverrideSettings(MeshBatch);
+	ERasterizerFillMode MeshFillMode = ComputeMeshFillMode(MeshBatch, Material, OverrideSettings);
+	ERasterizerCullMode MeshCullMode = ComputeMeshCullMode(MeshBatch, Material, OverrideSettings);
 	
 	TMobileBasePassShaderElementData<LightMapPolicyType> ShaderElementData(nullptr);
 	ShaderElementData.InitializeMeshMaterialData(ViewIfDynamicMeshCommand, PrimitiveSceneProxy, MeshBatch, StaticMeshId, false);
diff --git a/Engine/Source/Runtime/Renderer/Private/FogRendering.cpp b/Engine/Source/Runtime/Renderer/Private/FogRendering.cpp
--- a/Engine/Source/Runtime/Renderer/Private/FogRendering.cpp
+++ b/Engine/Source/Runtime/Renderer/Private/FogRendering.cpp
@@ -119,7 +119,7 @@ public:
 
 	void SetParameters(FRHICommandList& RHICmdList, const FViewInfo& View)
 	{
-		FGlobalShader::SetParameters<FViewUniformShaderParameters>(RHICmdList, GetVertexShader(),View.ViewUniformBuffer);
+		FGlobalShader::SetParameters<FViewUniformShaderParameters>(RHICmdList, RHICmdList.GetBoundVertexShader(), View.ViewUniformBuffer);
 
 		{
 			// The fog can be set to start at a certain euclidean distance.
@@ -143,19 +143,12 @@ public:
 
 			float FogClipSpaceZ = ClipSpaceMaxDistance.Z / ClipSpaceMaxDistance.W;
 
-			SetShaderValue(RHICmdList, GetVertexShader(),FogStartZ, FogClipSpaceZ);
+			SetShaderValue(RHICmdList, RHICmdList.GetBoundVertexShader(),FogStartZ, FogClipSpaceZ);
 		}
 	}
 
-	virtual bool Serialize(FArchive& Ar) override
-	{
-		bool bShaderHasOutdatedParameters = FGlobalShader::Serialize(Ar);
-		Ar << FogStartZ;
-		return bShaderHasOutdatedParameters;
-	}
-
 private:
-	FShaderParameter FogStartZ;
+	LAYOUT_FIELD(FShaderParameter, FogStartZ);
 };
 
 IMPLEMENT_SHADER_TYPE(,FHeightFogVS,TEXT("/Engine/Private/HeightFogVertexShader.usf"),TEXT("Main"),SF_Vertex);
@@ -204,19 +197,19 @@ public:
 
 	void SetParameters(FRHICommandList& RHICmdList, const FViewInfo& View, const FHeightFogRenderingParameters& Params)
 	{
-		FGlobalShader::SetParameters<FViewUniformShaderParameters>(RHICmdList, GetPixelShader(), View.ViewUniformBuffer);
-		SceneTextureParameters.Set(RHICmdList, GetPixelShader(), View.FeatureLevel, ESceneTextureSetupMode::All);
+		FGlobalShader::SetParameters<FViewUniformShaderParameters>(RHICmdList, RHICmdList.GetBoundPixelShader(), View.ViewUniformBuffer);
+		SceneTextureParameters.Set(RHICmdList, RHICmdList.GetBoundPixelShader(), View.FeatureLevel, ESceneTextureSetupMode::All);
 		
 		FFogUniformParameters FogUniformParameters;
 		SetupFogUniformParameters(View, FogUniformParameters);
-		SetUniformBufferParameterImmediate(RHICmdList, GetPixelShader(), GetUniformBufferParameter<FFogUniformParameters>(), FogUniformParameters);
+		SetUniformBufferParameterImmediate(RHICmdList, RHICmdList.GetBoundPixelShader(), GetUniformBufferParameter<FFogUniformParameters>(), FogUniformParameters);
 
 		FTextureRHIRef TextureRHI = Params.LightShaftsOutput.LightShaftOcclusion ?
 			Params.LightShaftsOutput.LightShaftOcclusion->GetRenderTargetItem().ShaderResourceTexture :
 			GWhiteTexture->TextureRHI;
 		SetTextureParameter(
 			RHICmdList, 
-			GetPixelShader(),
+			RHICmdList.GetBoundPixelShader(),
 			OcclusionTexture, OcclusionSampler,
 			TStaticSamplerState<SF_Bilinear,AM_Clamp,AM_Clamp,AM_Clamp>::GetRHI(),
 			TextureRHI
@@ -226,37 +219,24 @@ public:
 		FTextureRHIRef LinearDepthTextureRHI = bUseLinearDepthTextureEnabled ? Params.LinearDepthTextureRHI : GSystemTextures.DepthDummy->GetRenderTargetItem().ShaderResourceTexture;
 		SetTextureParameter(
 			RHICmdList, 
-			GetPixelShader(),
+			RHICmdList.GetBoundPixelShader(),
 			LinearDepthTexture, LinearDepthSampler,
 			TStaticSamplerState<SF_Point,AM_Clamp,AM_Clamp,AM_Clamp>::GetRHI(),
 			LinearDepthTextureRHI
 			);
 
-		SetShaderValue(RHICmdList, GetPixelShader(), bOnlyOnRenderedOpaque, View.bFogOnlyOnRenderedOpaque ? 1.0f : 0.0f);
-		SetShaderValue(RHICmdList, GetPixelShader(), bUseLinearDepthTexture, bUseLinearDepthTextureEnabled ? Params.LinearDepthReadScale : 0.0f);
-	}
-
-	virtual bool Serialize(FArchive& Ar) override
-	{
-		bool bShaderHasOutdatedParameters = FGlobalShader::Serialize(Ar);
-		Ar << SceneTextureParameters;
-		Ar << OcclusionTexture;
-		Ar << OcclusionSampler;
-		Ar << LinearDepthTexture;
-		Ar << LinearDepthSampler;
-		Ar << bOnlyOnRenderedOpaque;
-		Ar << bUseLinearDepthTexture;
-		return bShaderHasOutdatedParameters;
+		SetShaderValue(RHICmdList, RHICmdList.GetBoundPixelShader(), bOnlyOnRenderedOpaque, View.bFogOnlyOnRenderedOpaque ? 1.0f : 0.0f);
+		SetShaderValue(RHICmdList, RHICmdList.GetBoundPixelShader(), bUseLinearDepthTexture, bUseLinearDepthTextureEnabled ? Params.LinearDepthReadScale : 0.0f);
 	}
 
 private:
-	FSceneTextureShaderParameters SceneTextureParameters;
-	FShaderResourceParameter OcclusionTexture;
-	FShaderResourceParameter OcclusionSampler;
-	FShaderResourceParameter LinearDepthTexture;
-	FShaderResourceParameter LinearDepthSampler;
-	FShaderParameter bOnlyOnRenderedOpaque;
-	FShaderParameter bUseLinearDepthTexture;
+	LAYOUT_FIELD(FSceneTextureShaderParameters, SceneTextureParameters)
+	LAYOUT_FIELD(FShaderResourceParameter, OcclusionTexture)
+	LAYOUT_FIELD(FShaderResourceParameter, OcclusionSampler)
+	LAYOUT_FIELD(FShaderResourceParameter, LinearDepthTexture);
+	LAYOUT_FIELD(FShaderResourceParameter, LinearDepthSampler);
+	LAYOUT_FIELD(FShaderParameter, bOnlyOnRenderedOpaque)
+	LAYOUT_FIELD(FShaderParameter, bUseLinearDepthTexture);
 };
 
 IMPLEMENT_SHADER_TYPE(template<>,TExponentialHeightFogPS<EHeightFogFeature::HeightFog>,TEXT("/Engine/Private/HeightFogPixelShader.usf"), TEXT("ExponentialPixelMain"),SF_Pixel)
@@ -380,7 +360,7 @@ void SetFogShaders(FRHICommandList& RHICmdList, FGraphicsPipelineStateInitialize
 	{
 		TShaderMapRef<FHeightFogVS> VertexShader(View.ShaderMap);
 		GraphicsPSOInit.BoundShaderState.VertexDeclarationRHI = GFogVertexDeclaration.VertexDeclarationRHI;
-		GraphicsPSOInit.BoundShaderState.VertexShaderRHI = GETSAFERHISHADER_VERTEX(*VertexShader);
+		GraphicsPSOInit.BoundShaderState.VertexShaderRHI = VertexShader.GetVertexShader();
 		
 		if (bShouldRenderVolumetricFog)
 		{
@@ -388,7 +368,7 @@ void SetFogShaders(FRHICommandList& RHICmdList, FGraphicsPipelineStateInitialize
 			{
 				TShaderMapRef<TExponentialHeightFogPS<EHeightFogFeature::InscatteringTextureAndVolumetricFog> > ExponentialHeightFogPixelShader(View.ShaderMap);
 
-				GraphicsPSOInit.BoundShaderState.PixelShaderRHI = GETSAFERHISHADER_PIXEL(*ExponentialHeightFogPixelShader);
+				GraphicsPSOInit.BoundShaderState.PixelShaderRHI = ExponentialHeightFogPixelShader.GetPixelShader();
 				SetGraphicsPipelineState(RHICmdList, GraphicsPSOInit);
 				VertexShader->SetParameters(RHICmdList, View);
 				ExponentialHeightFogPixelShader->SetParameters(RHICmdList, View, Params);
@@ -397,7 +377,7 @@ void SetFogShaders(FRHICommandList& RHICmdList, FGraphicsPipelineStateInitialize
 			{
 				TShaderMapRef<TExponentialHeightFogPS<EHeightFogFeature::DirectionalLightInscatteringAndVolumetricFog> > ExponentialHeightFogPixelShader(View.ShaderMap);
 
-				GraphicsPSOInit.BoundShaderState.PixelShaderRHI = GETSAFERHISHADER_PIXEL(*ExponentialHeightFogPixelShader);
+				GraphicsPSOInit.BoundShaderState.PixelShaderRHI = ExponentialHeightFogPixelShader.GetPixelShader();
 				SetGraphicsPipelineState(RHICmdList, GraphicsPSOInit);
 				VertexShader->SetParameters(RHICmdList, View);
 				ExponentialHeightFogPixelShader->SetParameters(RHICmdList, View, Params);
@@ -406,7 +386,7 @@ void SetFogShaders(FRHICommandList& RHICmdList, FGraphicsPipelineStateInitialize
 			{
 				TShaderMapRef<TExponentialHeightFogPS<EHeightFogFeature::HeightFogAndVolumetricFog> > ExponentialHeightFogPixelShader(View.ShaderMap);
 
-				GraphicsPSOInit.BoundShaderState.PixelShaderRHI = GETSAFERHISHADER_PIXEL(*ExponentialHeightFogPixelShader);
+				GraphicsPSOInit.BoundShaderState.PixelShaderRHI = ExponentialHeightFogPixelShader.GetPixelShader();
 				SetGraphicsPipelineState(RHICmdList, GraphicsPSOInit);
 				VertexShader->SetParameters(RHICmdList, View);
 				ExponentialHeightFogPixelShader->SetParameters(RHICmdList, View, Params);
@@ -418,7 +398,7 @@ void SetFogShaders(FRHICommandList& RHICmdList, FGraphicsPipelineStateInitialize
 			{
 				TShaderMapRef<TExponentialHeightFogPS<EHeightFogFeature::InscatteringTexture> > ExponentialHeightFogPixelShader(View.ShaderMap);
 
-				GraphicsPSOInit.BoundShaderState.PixelShaderRHI = GETSAFERHISHADER_PIXEL(*ExponentialHeightFogPixelShader);
+				GraphicsPSOInit.BoundShaderState.PixelShaderRHI = ExponentialHeightFogPixelShader.GetPixelShader();
 				SetGraphicsPipelineState(RHICmdList, GraphicsPSOInit);
 				VertexShader->SetParameters(RHICmdList, View);
 				ExponentialHeightFogPixelShader->SetParameters(RHICmdList, View, Params);
@@ -427,7 +407,7 @@ void SetFogShaders(FRHICommandList& RHICmdList, FGraphicsPipelineStateInitialize
 			{
 				TShaderMapRef<TExponentialHeightFogPS<EHeightFogFeature::DirectionalLightInscattering> > ExponentialHeightFogPixelShader(View.ShaderMap);
 
-				GraphicsPSOInit.BoundShaderState.PixelShaderRHI = GETSAFERHISHADER_PIXEL(*ExponentialHeightFogPixelShader);
+				GraphicsPSOInit.BoundShaderState.PixelShaderRHI = ExponentialHeightFogPixelShader.GetPixelShader();
 				SetGraphicsPipelineState(RHICmdList, GraphicsPSOInit);
 				VertexShader->SetParameters(RHICmdList, View);
 				ExponentialHeightFogPixelShader->SetParameters(RHICmdList, View, Params);
@@ -436,7 +416,7 @@ void SetFogShaders(FRHICommandList& RHICmdList, FGraphicsPipelineStateInitialize
 			{
 				TShaderMapRef<TExponentialHeightFogPS<EHeightFogFeature::HeightFog> > ExponentialHeightFogPixelShader(View.ShaderMap);
 
-				GraphicsPSOInit.BoundShaderState.PixelShaderRHI = GETSAFERHISHADER_PIXEL(*ExponentialHeightFogPixelShader);
+				GraphicsPSOInit.BoundShaderState.PixelShaderRHI = ExponentialHeightFogPixelShader.GetPixelShader();
 				SetGraphicsPipelineState(RHICmdList, GraphicsPSOInit);
 				VertexShader->SetParameters(RHICmdList, View);
 				ExponentialHeightFogPixelShader->SetParameters(RHICmdList, View, Params);
diff --git a/Engine/Source/Runtime/Renderer/Private/GPUBenchmark.cpp b/Engine/Source/Runtime/Renderer/Private/GPUBenchmark.cpp
--- a/Engine/Source/Runtime/Renderer/Private/GPUBenchmark.cpp
+++ b/Engine/Source/Runtime/Renderer/Private/GPUBenchmark.cpp
@@ -49,8 +49,8 @@ class FPostProcessBenchmarkPS : public FGlobalShader
 	FPostProcessBenchmarkPS() {}
 
 public:
-	FShaderResourceParameter InputTexture;
-	FShaderResourceParameter InputTextureSampler;
+	LAYOUT_FIELD(FShaderResourceParameter, InputTexture);
+	LAYOUT_FIELD(FShaderResourceParameter, InputTextureSampler);
 
 	/** Initialization constructor. */
 	FPostProcessBenchmarkPS(const ShaderMetaType::CompiledShaderInitializerType& Initializer)
@@ -60,17 +60,9 @@ public:
 		InputTextureSampler.Bind(Initializer.ParameterMap,TEXT("InputTextureSampler"));
 	}
 
-	// FShader interface.
-	virtual bool Serialize(FArchive& Ar) override
-	{
-		bool bShaderHasOutdatedParameters = FGlobalShader::Serialize(Ar);
-		Ar << InputTexture << InputTextureSampler;
-		return bShaderHasOutdatedParameters;
-	}
-
 	void SetParameters(FRHICommandList& RHICmdList, const FSceneView& View, TRefCountPtr<IPooledRenderTarget>& Src)
 	{
-		FRHIPixelShader* ShaderRHI = GetPixelShader();
+		FRHIPixelShader* ShaderRHI = RHICmdList.GetBoundPixelShader();
 
 		FGlobalShader::SetParameters<FViewUniformShaderParameters>(RHICmdList, ShaderRHI, View.ViewUniformBuffer);
 
@@ -123,17 +115,9 @@ public:
 	{
 	}
 
-	/** Serializer */
-	virtual bool Serialize(FArchive& Ar) override
-	{
-		bool bShaderHasOutdatedParameters = FGlobalShader::Serialize(Ar);
-
-		return bShaderHasOutdatedParameters;
-	}
-
 	void SetParameters(FRHICommandList& RHICmdList, const FSceneView& View)
 	{
-		FRHIVertexShader* ShaderRHI = GetVertexShader();
+		FRHIVertexShader* ShaderRHI = RHICmdList.GetBoundVertexShader();
 		
 		FGlobalShader::SetParameters<FViewUniformShaderParameters>(RHICmdList, ShaderRHI, View.ViewUniformBuffer);
 	}
@@ -210,8 +194,8 @@ void RunBenchmarkShader(FRHICommandList& RHICmdList, FRHIVertexBuffer* VertexThr
 		: GFilterVertexDeclaration.VertexDeclarationRHI;
 
 	GraphicsPSOInit.BoundShaderState.VertexDeclarationRHI = VertexDeclaration;
-	GraphicsPSOInit.BoundShaderState.VertexShaderRHI = GETSAFERHISHADER_VERTEX(*VertexShader);
-	GraphicsPSOInit.BoundShaderState.PixelShaderRHI = GETSAFERHISHADER_PIXEL(*PixelShader);
+	GraphicsPSOInit.BoundShaderState.VertexShaderRHI = VertexShader.GetVertexShader();
+	GraphicsPSOInit.BoundShaderState.PixelShaderRHI = PixelShader.GetPixelShader();
 	GraphicsPSOInit.PrimitiveType = PT_TriangleList;
 
 	SetGraphicsPipelineState(RHICmdList, GraphicsPSOInit);
@@ -262,7 +246,7 @@ void RunBenchmarkShader(FRHICommandList& RHICmdList, FRHIVertexBuffer* VertexThr
 				GBenchmarkResolution, LocalHeight,
 				FIntPoint(GBenchmarkResolution, GBenchmarkResolution),
 				FIntPoint(GBenchmarkResolution, GBenchmarkResolution),
-				*VertexShader,
+				VertexShader,
 				EDRF_Default);
 		}
 	}
