Copying Private-LoadTimes-4.24 stream to Main. Biggest changes are in Materials/Shader memory freezing.

diff --git a/Engine/Source/Runtime/Renderer/Private/GPUFastFourierTransform.cpp b/Engine/Source/Runtime/Renderer/Private/GPUFastFourierTransform.cpp
--- a/Engine/Source/Runtime/Renderer/Private/GPUFastFourierTransform.cpp
+++ b/Engine/Source/Runtime/Renderer/Private/GPUFastFourierTransform.cpp
@@ -215,7 +215,7 @@ namespace GPUFFT
 		{
 			using GPUFFTComputeShaderUtils::FComputeParamterValueSetter;
 
-			FRHIComputeShader* ShaderRHI = GetComputeShader();
+			FRHIComputeShader* ShaderRHI = RHICmdList.GetBoundComputeShader();
 
 			// Set up the input.  We have to do this explicitly because the FFT dispatches multiple compute shaders and manages their input/output.
 			FComputeParamterValueSetter ParamSetter(RHICmdList, ShaderRHI);
@@ -242,30 +242,15 @@ namespace GPUFFT
 		// Method for use with the FScopedUAVBind
 		FShaderResourceParameter& DestinationResourceParameter() { return DstRWTexture; }
 
-
-		// FShader interface.
-		virtual bool Serialize(FArchive& Ar) override
-		{
-			bool bShaderHasOutdatedParameters = FGlobalShader::Serialize(Ar);
-			Ar << SrcROTexture
-				<< DstRWTexture
-				<< TransformType
-				<< SrcRect
-				<< DstRect
-				<< LogTransformLength
-				<< BitCount;
-			return bShaderHasOutdatedParameters;
-		}
-
 	public:
 
-		FShaderResourceParameter SrcROTexture;
-		FShaderResourceParameter DstRWTexture;
-		FShaderParameter TransformType;
-		FShaderParameter SrcRect;
-		FShaderParameter DstRect;
-		FShaderParameter LogTransformLength;
-		FShaderParameter BitCount;
+		LAYOUT_FIELD(FShaderResourceParameter, SrcROTexture);
+		LAYOUT_FIELD(FShaderResourceParameter, DstRWTexture);
+		LAYOUT_FIELD(FShaderParameter, TransformType);
+		LAYOUT_FIELD(FShaderParameter, SrcRect);
+		LAYOUT_FIELD(FShaderParameter, DstRect);
+		LAYOUT_FIELD(FShaderParameter, LogTransformLength);
+		LAYOUT_FIELD(FShaderParameter, BitCount);
 	};
 
 
@@ -330,7 +315,7 @@ namespace GPUFFT
 
 			using GPUFFTComputeShaderUtils::FComputeParamterValueSetter;
 
-			FRHIComputeShader* ShaderRHI = GetComputeShader();
+			FRHIComputeShader* ShaderRHI = RHICmdList.GetBoundComputeShader();
 
 			// Set up the input.  We have to do this explicitly because the FFT dispatches multiple compute shaders and manages their input/output.
 			FComputeParamterValueSetter ParamSetter(RHICmdList, ShaderRHI);
@@ -349,28 +334,14 @@ namespace GPUFFT
 		// Method for use with the FScopedUAVBind
 		FShaderResourceParameter& DestinationResourceParameter() { return DstRWTexture; }
 
-
-		// FShader interface.
-		virtual bool Serialize(FArchive& Ar) override
-		{
-			bool bShaderHasOutdatedParameters = FGlobalShader::Serialize(Ar);
-			Ar << SrcROTexture
-				<< DstRWTexture
-				<< TransformType
-				<< SrcRect
-				<< TransformLength
-				<< NumSubRegions;
-			return bShaderHasOutdatedParameters;
-		}
-
 	public:
 
-		FShaderResourceParameter SrcROTexture;
-		FShaderResourceParameter DstRWTexture;
-		FShaderParameter TransformType;
-		FShaderParameter SrcRect;
-		FShaderParameter TransformLength;
-		FShaderParameter NumSubRegions;
+		LAYOUT_FIELD(FShaderResourceParameter, SrcROTexture);
+		LAYOUT_FIELD(FShaderResourceParameter, DstRWTexture);
+		LAYOUT_FIELD(FShaderParameter, TransformType);
+		LAYOUT_FIELD(FShaderParameter, SrcRect);
+		LAYOUT_FIELD(FShaderParameter, TransformLength);
+		LAYOUT_FIELD(FShaderParameter, NumSubRegions);
 	};
 
 	class FComplexFFTPassCS : public FGlobalShader
@@ -426,7 +397,7 @@ namespace GPUFFT
 
 			using GPUFFTComputeShaderUtils::FComputeParamterValueSetter;
 
-			FRHIComputeShader* ShaderRHI = GetComputeShader();
+			FRHIComputeShader* ShaderRHI = RHICmdList.GetBoundComputeShader();
 
 			// Set up the input.  We have to do this explicitly because the FFT dispatches multiple compute shaders and manages their input/output.
 			FComputeParamterValueSetter ParamSetter(RHICmdList, ShaderRHI);
@@ -450,30 +421,15 @@ namespace GPUFFT
 		// Method for use with the FScopedUAVBind
 		FShaderResourceParameter& DestinationResourceParameter() { return DstRWTexture; }
 
-
-		// FShader interface.
-		virtual bool Serialize(FArchive& Ar) override
-		{
-			bool bShaderHasOutdatedParameters = FGlobalShader::Serialize(Ar);
-			Ar << SrcROTexture
-				<< DstRWTexture
-				<< TransformType
-				<< SrcRect
-				<< DstRect
-				<< BitCount
-				<< PowTwoLength;
-			return bShaderHasOutdatedParameters;
-		}
-
 	public:
 
-		FShaderResourceParameter SrcROTexture;
-		FShaderResourceParameter DstRWTexture;
-		FShaderParameter TransformType;
-		FShaderParameter SrcRect;
-		FShaderParameter DstRect;
-		FShaderParameter BitCount;
-		FShaderParameter PowTwoLength;
+		LAYOUT_FIELD(FShaderResourceParameter, SrcROTexture);
+		LAYOUT_FIELD(FShaderResourceParameter, DstRWTexture);
+		LAYOUT_FIELD(FShaderParameter, TransformType);
+		LAYOUT_FIELD(FShaderParameter, SrcRect);
+		LAYOUT_FIELD(FShaderParameter, DstRect);
+		LAYOUT_FIELD(FShaderParameter, BitCount);
+		LAYOUT_FIELD(FShaderParameter, PowTwoLength);
 	};
 
 
@@ -524,7 +480,7 @@ namespace GPUFFT
 		{
 			using GPUFFTComputeShaderUtils::FComputeParamterValueSetter;
 
-			FRHIComputeShader* ShaderRHI = GetComputeShader();
+			FRHIComputeShader* ShaderRHI = RHICmdList.GetBoundComputeShader();
 
 			// Set up the input.  We have to do this explicitly because the FFT dispatches multiple compute shaders and manages their input/output.
 			FComputeParamterValueSetter ParamSetter(RHICmdList, ShaderRHI);
@@ -541,24 +497,12 @@ namespace GPUFFT
 		// Method for use with the FScopedUAVBind
 		FShaderResourceParameter& DestinationResourceParameter() { return DstRWTexture; }
 
-
-		// FShader interface.
-		virtual bool Serialize(FArchive& Ar) override
-		{
-			bool bShaderHasOutdatedParameters = FGlobalShader::Serialize(Ar);
-			Ar << SrcROTexture
-				<< DstRWTexture
-				<< TransformType
-				<< DstRect;
-			return bShaderHasOutdatedParameters;
-		}
-
 	public:
 
-		FShaderResourceParameter SrcROTexture;
-		FShaderResourceParameter DstRWTexture;
-		FShaderParameter TransformType;
-		FShaderParameter DstRect;
+		LAYOUT_FIELD(FShaderResourceParameter, SrcROTexture);
+		LAYOUT_FIELD(FShaderResourceParameter, DstRWTexture);
+		LAYOUT_FIELD(FShaderParameter, TransformType);
+		LAYOUT_FIELD(FShaderParameter, DstRect);
 	};
 
 
@@ -613,7 +557,7 @@ namespace GPUFFT
 		{
 			using GPUFFTComputeShaderUtils::FComputeParamterValueSetter;
 
-			FRHIComputeShader* ShaderRHI = GetComputeShader();
+			FRHIComputeShader* ShaderRHI = RHICmdList.GetBoundComputeShader();
 
 			// Set up the input.  We have to do this explicitly because the FFT dispatches multiple compute shaders and manages their input/output.
 			FComputeParamterValueSetter ParamSetter(RHICmdList, ShaderRHI);
@@ -627,26 +571,13 @@ namespace GPUFFT
 		// Method for use with the FScopedUAVBind
 		FShaderResourceParameter& DestinationResourceParameter() { return DstRWTexture; }
 
-
-		// FShader interface.
-		virtual bool Serialize(FArchive& Ar) override
-		{
-			bool bShaderHasOutdatedParameters = FGlobalShader::Serialize(Ar);
-			Ar << SrcROTexture
-				<< DstRWTexture
-				<< SrcRect
-				<< DstRect
-				<< PreFilter;
-			return bShaderHasOutdatedParameters;
-		}
-
 	public:
 
-		FShaderResourceParameter SrcROTexture;
-		FShaderResourceParameter DstRWTexture;
-		FShaderParameter SrcRect;
-		FShaderParameter DstRect;
-		FShaderParameter PreFilter;
+		LAYOUT_FIELD(FShaderResourceParameter, SrcROTexture);
+		LAYOUT_FIELD(FShaderResourceParameter, DstRWTexture);
+		LAYOUT_FIELD(FShaderParameter, SrcRect);
+		LAYOUT_FIELD(FShaderParameter, DstRect);
+		LAYOUT_FIELD(FShaderParameter, PreFilter);
 	};
 
 	class FComplexMultiplyImagesCS : public FGlobalShader
@@ -697,7 +628,7 @@ namespace GPUFFT
 		{
 			using GPUFFTComputeShaderUtils::FComputeParamterValueSetter;
 
-			FRHIComputeShader* ShaderRHI = GetComputeShader();
+			FRHIComputeShader* ShaderRHI = RHICmdList.GetBoundComputeShader();
 
 			const uint32 DataLayoutValue = (bHorizontalScanlines) ? 1 : 0;
 			// Set up the input.  We have to do this explicitly because the FFT dispatches multiple compute shaders and manages their input/output.
@@ -713,31 +644,18 @@ namespace GPUFFT
 		// Method for use with the FScopedUAVBind
 		FShaderResourceParameter& DestinationResourceParameter() { return DstRWTexture; }
 
-
-		// FShader interface.
-		virtual bool Serialize(FArchive& Ar) override
-		{
-			bool bShaderHasOutdatedParameters = FGlobalShader::Serialize(Ar);
-			Ar << SrcROTexture
-				<< KnlROTexture
-				<< DstRWTexture
-				<< SrcRect
-				<< DataLayout;
-			return bShaderHasOutdatedParameters;
-		}
-
 	public:
 
-		FShaderResourceParameter SrcROTexture;
-		FShaderResourceParameter KnlROTexture;
-		FShaderResourceParameter DstRWTexture;
-		FShaderParameter SrcRect;
-		FShaderParameter DataLayout;
+		LAYOUT_FIELD(FShaderResourceParameter, SrcROTexture);
+		LAYOUT_FIELD(FShaderResourceParameter, KnlROTexture);
+		LAYOUT_FIELD(FShaderResourceParameter, DstRWTexture);
+		LAYOUT_FIELD(FShaderParameter, SrcRect);
+		LAYOUT_FIELD(FShaderParameter, DataLayout);
 	};
 
 	class FGSComplexTransformBaseCS : public FGlobalShader
 	{
-
+		DECLARE_INLINE_TYPE_LAYOUT(FGSComplexTransformBaseCS, NonVirtual);
 	public:
 		
 
@@ -769,7 +687,7 @@ namespace GPUFFT
 		{
 			using GPUFFTComputeShaderUtils::FComputeParamterValueSetter;
 
-			FRHIComputeShader* ShaderRHI = GetComputeShader();
+			FRHIComputeShader* ShaderRHI = RHICmdList.GetBoundComputeShader();
 
 			// Set up the input.  We have to do this explicitly because the FFT dispatches multiple compute shaders and manages their input/output.
 			FComputeParamterValueSetter ParamSetter(RHICmdList, ShaderRHI);
@@ -797,32 +715,15 @@ namespace GPUFFT
 		// Method for use with the FScopedUAVBind
 		FShaderResourceParameter& DestinationResourceParameter() { return DstRWTexture; }
 
-
-		// FShader interface.
-		virtual bool Serialize(FArchive& Ar) override
-		{
-			bool bShaderHasOutdatedParameters = FGlobalShader::Serialize(Ar);
-			Ar << SrcROTexture
-				<< DstRWTexture
-				<< TransformType
-				<< SrcRectMin
-				<< SrcRectMax
-				<< DstExtent
-				<< DstRect
-				<< BrightPixelGain;
-			return bShaderHasOutdatedParameters;
-		}
-
 	public:
-
-		FShaderResourceParameter SrcROTexture;
-		FShaderResourceParameter DstRWTexture;
-		FShaderParameter TransformType;
-		FShaderParameter SrcRectMin;
-		FShaderParameter SrcRectMax;
-		FShaderParameter DstExtent;
-		FShaderParameter DstRect;
-		FShaderParameter BrightPixelGain;
+		LAYOUT_FIELD(FShaderResourceParameter, SrcROTexture);
+		LAYOUT_FIELD(FShaderResourceParameter, DstRWTexture);
+		LAYOUT_FIELD(FShaderParameter, TransformType);
+		LAYOUT_FIELD(FShaderParameter, SrcRectMin);
+		LAYOUT_FIELD(FShaderParameter, SrcRectMax);
+		LAYOUT_FIELD(FShaderParameter, DstExtent);
+		LAYOUT_FIELD(FShaderParameter, DstRect);
+		LAYOUT_FIELD(FShaderParameter, BrightPixelGain);
 	};
 
 	template <int PowRadixSignalLength>
@@ -902,7 +803,7 @@ namespace GPUFFT
 	*/
 	class FGSConvolutionBaseCS : public FGlobalShader
 	{
-
+		DECLARE_INLINE_TYPE_LAYOUT(FGSConvolutionBaseCS, NonVirtual);
 	public:
 
 		FGSConvolutionBaseCS() {};
@@ -940,7 +841,7 @@ namespace GPUFFT
 				TransformTypeValue |= 8;
 			}
 
-			FRHIComputeShader* ShaderRHI = GetComputeShader();
+			FRHIComputeShader* ShaderRHI = RHICmdList.GetBoundComputeShader();
 
 			// Set up the input.  We have to do this explicitly because the FFT dispatches multiple compute shaders and manages their input/output.
 
@@ -957,28 +858,13 @@ namespace GPUFFT
 		// Method for use with the FScopedUAVBind
 		FShaderResourceParameter& DestinationResourceParameter() { return DstRWTexture; }
 
-
-		// FShader interface.
-		virtual bool Serialize(FArchive& Ar) override
-		{
-			bool bShaderHasOutdatedParameters = FGlobalShader::Serialize(Ar);
-			Ar << SrcROTexture
-				<< DstRWTexture
-				<< SrcRectMin
-				<< SrcRectMax
-				<< DstExtent
-				<< TransformType;
-			return bShaderHasOutdatedParameters;
-		}
-
 	public:
-
-		FShaderResourceParameter SrcROTexture;
-		FShaderResourceParameter DstRWTexture;
-		FShaderParameter SrcRectMin;
-		FShaderParameter SrcRectMax;
-		FShaderParameter DstExtent;
-		FShaderParameter TransformType;
+		LAYOUT_FIELD(FShaderResourceParameter, SrcROTexture);
+		LAYOUT_FIELD(FShaderResourceParameter, DstRWTexture);
+		LAYOUT_FIELD(FShaderParameter, SrcRectMin);
+		LAYOUT_FIELD(FShaderParameter, SrcRectMax);
+		LAYOUT_FIELD(FShaderParameter, DstExtent);
+		LAYOUT_FIELD(FShaderParameter, TransformType);
 	};
 
 	/**
@@ -986,7 +872,7 @@ namespace GPUFFT
 	*/
 	class FGSConvolutionWithTextureKernelBaseCS : public FGSConvolutionBaseCS
 	{
-
+		DECLARE_INLINE_TYPE_LAYOUT(FGSConvolutionWithTextureKernelBaseCS, NonVirtual);
 	public:
 
 		FGSConvolutionWithTextureKernelBaseCS() {};
@@ -1016,23 +902,15 @@ namespace GPUFFT
 			FGSConvolutionBaseCS::SetCSParamters(RHICmdList, XFormType, SrcTexture, SrcRect, DstExtentValue);
 
 			// additional source input for sampling the spectral texture
-			FRHIComputeShader* ShaderRHI = GetComputeShader();
+			FRHIComputeShader* ShaderRHI = RHICmdList.GetBoundComputeShader();
 			GPUFFTComputeShaderUtils::FComputeParamterValueSetter ParamSetter(RHICmdList, ShaderRHI);
 			// set the texture
 			ParamSetter(FilterSrcROTexture, FilterSrcTexture);
 
 		}
 
-		// FShader interface.
-		virtual bool Serialize(FArchive& Ar) override
-		{
-			bool bShaderHasOutdatedParameters = FGSConvolutionBaseCS::Serialize(Ar);
-			Ar << FilterSrcROTexture;
-			return bShaderHasOutdatedParameters;
-		}
-
 	public:
-		FShaderResourceParameter FilterSrcROTexture;
+		LAYOUT_FIELD(FShaderResourceParameter, FilterSrcROTexture);
 	};
 
 
@@ -1137,40 +1015,39 @@ namespace
 {
 	using namespace GPUFFT;
 
-	FCopyWindowCS* GetCopyWindowCS(const FGPUFFTShaderContext::ShaderMapType& ShaderMap)
+	TShaderRef<FCopyWindowCS> GetCopyWindowCS(const FGPUFFTShaderContext::ShaderMapType& ShaderMap)
 	{
 		return  ShaderMap.GetShader<FCopyWindowCS>();
 	}
 
-	FComplexMultiplyImagesCS* GetComplexMultiplyImagesCS(const FGPUFFTShaderContext::ShaderMapType& ShaderMap)
+	TShaderRef<FComplexMultiplyImagesCS> GetComplexMultiplyImagesCS(const FGPUFFTShaderContext::ShaderMapType& ShaderMap)
 	{
 		return  ShaderMap.GetShader<FComplexMultiplyImagesCS>();
 	}
 
-	FGroupShardSubFFTPassCS* GetGroupSharedSubFFTPassCS(const FGPUFFTShaderContext::ShaderMapType& ShaderMap, const uint32 TransformLength)
+	TShaderRef<FGroupShardSubFFTPassCS> GetGroupSharedSubFFTPassCS(const FGPUFFTShaderContext::ShaderMapType& ShaderMap, const uint32 TransformLength)
 	{
 		return ShaderMap.GetShader<FGroupShardSubFFTPassCS>();
 	}
 
-	FReorderFFTPassCS* GetReorderFFTPassCS(const FGPUFFTShaderContext::ShaderMapType& ShaderMap)
+	TShaderRef<FReorderFFTPassCS> GetReorderFFTPassCS(const FGPUFFTShaderContext::ShaderMapType& ShaderMap)
 	{
 		return ShaderMap.GetShader<FReorderFFTPassCS>();
 	}
-	FPackTwoForOneFFTPassCS* GetPackTwoForOneFFTPassCS(const FGPUFFTShaderContext::ShaderMapType& ShaderMap, const uint32 TransformLength)
+	TShaderRef<FPackTwoForOneFFTPassCS> GetPackTwoForOneFFTPassCS(const FGPUFFTShaderContext::ShaderMapType& ShaderMap, const uint32 TransformLength)
 	{
 		return  ShaderMap.GetShader<FPackTwoForOneFFTPassCS>();
 	}
 
-	FComplexFFTPassCS* GetComplexFFTPassCS(const FGPUFFTShaderContext::ShaderMapType& ShaderMap, const uint32 TransformLength)
+	TShaderRef<FComplexFFTPassCS> GetComplexFFTPassCS(const FGPUFFTShaderContext::ShaderMapType& ShaderMap, const uint32 TransformLength)
 	{
-		FComplexFFTPassCS* Result = ShaderMap.GetShader<FComplexFFTPassCS>();
-		return Result;
+		return ShaderMap.GetShader<FComplexFFTPassCS>();
 	}
 
 	// Shader Permutation picker.
-	FGSComplexTransformBaseCS* GetComplexFFTCS(const FGPUFFTShaderContext::ShaderMapType& ShaderMap, const uint32 TransformLength)
+	TShaderRef<FGSComplexTransformBaseCS> GetComplexFFTCS(const FGPUFFTShaderContext::ShaderMapType& ShaderMap, const uint32 TransformLength)
 	{
-		FGSComplexTransformBaseCS* Result = NULL;
+		TShaderRef<FGSComplexTransformBaseCS> Result;
 
 	#define GET_COMPLEX_SHADER(_LENGTH)  ShaderMap.GetShader<TGSComplexTransformCS<_LENGTH>>();
 
@@ -1201,9 +1078,9 @@ namespace
 	}
 
 	// Shader Permutation picker.
-	FGSComplexTransformBaseCS* GetTwoForOneFFTCS(const FGPUFFTShaderContext::ShaderMapType& ShaderMap, const uint32 TransformLength)
+	TShaderRef<FGSComplexTransformBaseCS> GetTwoForOneFFTCS(const FGPUFFTShaderContext::ShaderMapType& ShaderMap, const uint32 TransformLength)
 	{
-		FGSComplexTransformBaseCS* Result = NULL;
+		TShaderRef<FGSComplexTransformBaseCS> Result;
 
 	#define GET_TWOFORONE_SHADER(_LENGTH)  ShaderMap.GetShader<TGSTwoForOneTransformCS<_LENGTH>>();
 
@@ -1234,9 +1111,9 @@ namespace
 	}
 
 
-	FGSConvolutionWithTextureKernelBaseCS* GetConvolutionWithTextureKernelCS(const FGPUFFTShaderContext::ShaderMapType& ShaderMap, const uint32 TransformLength)
+	TShaderRef<FGSConvolutionWithTextureKernelBaseCS> GetConvolutionWithTextureKernelCS(const FGPUFFTShaderContext::ShaderMapType& ShaderMap, const uint32 TransformLength)
 	{
-		FGSConvolutionWithTextureKernelBaseCS* Result = NULL;
+		TShaderRef<FGSConvolutionWithTextureKernelBaseCS> Result;
 
 		// Get the spectral filter.
 	#define GET_GROUP_SHARED_TEXTURE_FILTER(_LENGTH)  ShaderMap.GetShader< TGSConvolutionWithTexturerCS<_LENGTH> >();
@@ -1302,11 +1179,11 @@ namespace
 		SCOPED_DRAW_EVENTF(RHICmdList, CopyWindowCS, TEXT("FFT Multipass: Copy Subwindow"));
 
 		// Get pointer to the shader
-		FCopyWindowCS* ComputeShader = GetCopyWindowCS(ShaderMap);
+		TShaderRef<FCopyWindowCS> ComputeShader = GetCopyWindowCS(ShaderMap);
 
 		// #todo-renderpasses remove once everything is converted
 		UnbindRenderTargets(RHICmdList);
-		RHICmdList.SetComputeShader(ComputeShader->GetComputeShader());
+		RHICmdList.SetComputeShader(ComputeShader.GetComputeShader());
 
 		// Bind output
 		FScopedUAVBind ScopedBind = FScopedUAVBind::BindOutput(RHICmdList, ComputeShader, DstUAV);
@@ -1363,11 +1240,11 @@ namespace
 		SCOPED_DRAW_EVENTF(RHICmdList, ComplexMultiplyImagesCS, TEXT("FFT Multipass: Convolution in freq-space"));
 
 		// Get pointer to the shader
-		FComplexMultiplyImagesCS* ComputeShader = GetComplexMultiplyImagesCS(ShaderMap);
+		TShaderRef<FComplexMultiplyImagesCS> ComputeShader = GetComplexMultiplyImagesCS(ShaderMap);
 
 		// #todo-renderpasses remove once everything is converted
 		UnbindRenderTargets(RHICmdList);
-		RHICmdList.SetComputeShader(ComputeShader->GetComputeShader());
+		RHICmdList.SetComputeShader(ComputeShader.GetComputeShader());
 
 		// Bind output
 		FScopedUAVBind ScopedBind = FScopedUAVBind::BindOutput(RHICmdList, ComputeShader, DstUAV);
@@ -1432,11 +1309,11 @@ namespace
 		}
 
 		// Get pointer to the shader
-		FPackTwoForOneFFTPassCS* ComputeShader = GetPackTwoForOneFFTPassCS(ShaderMap, TransformLength);
+		TShaderRef<FPackTwoForOneFFTPassCS> ComputeShader = GetPackTwoForOneFFTPassCS(ShaderMap, TransformLength);
 
 		// #todo-renderpasses remove once everything is converted
 		UnbindRenderTargets(RHICmdList);
-		RHICmdList.SetComputeShader(ComputeShader->GetComputeShader());
+		RHICmdList.SetComputeShader(ComputeShader.GetComputeShader());
 
 		// Bind output
 		FScopedUAVBind ScopedBind = FScopedUAVBind::BindOutput(RHICmdList, ComputeShader, DstUAV);
@@ -1482,12 +1359,12 @@ namespace
 		FIntPoint DstExtent = FFTDesc.TransformExtent();
 		
 		// Get a base pointer to the shader
-		FComplexFFTPassCS* ComputeShader = GetComplexFFTPassCS(ShaderMap, TransformLength);
+		TShaderRef<FComplexFFTPassCS> ComputeShader = GetComplexFFTPassCS(ShaderMap, TransformLength);
 
 
 		// #todo-renderpasses remove once everything is converted
 		UnbindRenderTargets(RHICmdList);
-		RHICmdList.SetComputeShader(ComputeShader->GetComputeShader());
+		RHICmdList.SetComputeShader(ComputeShader.GetComputeShader());
 
 		// Bind output
 		FScopedUAVBind ScopedBind = FScopedUAVBind::BindOutput(RHICmdList, ComputeShader, DstUAV);
@@ -1535,12 +1412,12 @@ namespace
 		FIntPoint DstExtent = FFTDesc.TransformExtent();
 
 		// Get a base pointer to the shader
-		FReorderFFTPassCS* ComputeShader = GetReorderFFTPassCS(ShaderMap);
+		TShaderRef<FReorderFFTPassCS> ComputeShader = GetReorderFFTPassCS(ShaderMap);
 
 
 		// #todo-renderpasses remove once everything is converted
 		UnbindRenderTargets(RHICmdList);
-		RHICmdList.SetComputeShader(ComputeShader->GetComputeShader());
+		RHICmdList.SetComputeShader(ComputeShader.GetComputeShader());
 
 		// Bind output
 		FScopedUAVBind ScopedBind = FScopedUAVBind::BindOutput(RHICmdList, ComputeShader, DstUAV);
@@ -1598,12 +1475,12 @@ namespace
 		FIntPoint DstExtent = FFTDesc.TransformExtent();
 
 		// Get a base pointer to the shader
-		FGroupShardSubFFTPassCS* ComputeShader = GetGroupSharedSubFFTPassCS(ShaderMap, TransformLength);
+		TShaderRef<FGroupShardSubFFTPassCS> ComputeShader = GetGroupSharedSubFFTPassCS(ShaderMap, TransformLength);
 
 
 		// #todo-renderpasses remove once everything is converted
 		UnbindRenderTargets(RHICmdList);
-		RHICmdList.SetComputeShader(ComputeShader->GetComputeShader());
+		RHICmdList.SetComputeShader(ComputeShader.GetComputeShader());
 
 		// Bind output
 		FScopedUAVBind ScopedBind = FScopedUAVBind::BindOutput(RHICmdList, ComputeShader, DstUAV);
@@ -1646,12 +1523,12 @@ namespace
 
 
 		// Get a base pointer to the shader
-		FGSComplexTransformBaseCS* ComputeShader = GetComplexFFTCS(ShaderMap, TransformLength);
+		TShaderRef<FGSComplexTransformBaseCS> ComputeShader = GetComplexFFTCS(ShaderMap, TransformLength);
 
 
 		// #todo-renderpasses remove once everything is converted
 		UnbindRenderTargets(RHICmdList);
-		RHICmdList.SetComputeShader(ComputeShader->GetComputeShader());
+		RHICmdList.SetComputeShader(ComputeShader.GetComputeShader());
 
 		// Bind output
 		FScopedUAVBind ScopedBind = FScopedUAVBind::BindOutput(RHICmdList, ComputeShader, DstUAV);
@@ -1698,12 +1575,12 @@ namespace
 		SCOPED_DRAW_EVENTF(RHICmdList, FRCPassFFT, TEXT("FFT: Two-For-One %s of size %d of buffer %d x %d"), *TransformName, TransformLength, SrcRect.Size().X, SrcRect.Size().Y);
 
 		// Get a base pointer to the shader
-		FGSComplexTransformBaseCS* ComputeShader = GetTwoForOneFFTCS(ShaderMap, TransformLength);
+		TShaderRef<FGSComplexTransformBaseCS> ComputeShader = GetTwoForOneFFTCS(ShaderMap, TransformLength);
 
 		// DJH - do we need this SetRenderTarget?
 		// #todo-renderpasses remove once everything is converted
 		UnbindRenderTargets(RHICmdList);
-		RHICmdList.SetComputeShader(ComputeShader->GetComputeShader());
+		RHICmdList.SetComputeShader(ComputeShader.GetComputeShader());
 
 
 		FScopedUAVBind ScopedBind = FScopedUAVBind::BindOutput(RHICmdList, ComputeShader, DstUAV);
@@ -1756,12 +1633,12 @@ namespace
 		SCOPED_DRAW_EVENTF(RHICmdList, FRCPassFFTBloom, TEXT("FFT: Apply %s Transform, Multiply Texture, and InverseTransform size %d of buffer %d x %d"), *XFormDirName, SignalLength, SrcRectSize.X, SrcRectSize.Y);
 
 		// Get a base pointer to the shader
-		FGSConvolutionWithTextureKernelBaseCS* ComputeShader = GetConvolutionWithTextureKernelCS(ShaderMap, SignalLength);
+		TShaderRef<FGSConvolutionWithTextureKernelBaseCS> ComputeShader = GetConvolutionWithTextureKernelCS(ShaderMap, SignalLength);
 
 
 		// #todo-renderpasses remove once everything is converted
 		UnbindRenderTargets(RHICmdList);
-		RHICmdList.SetComputeShader(ComputeShader->GetComputeShader());
+		RHICmdList.SetComputeShader(ComputeShader.GetComputeShader());
 		
 		FScopedUAVBind ScopedBind = FScopedUAVBind::BindOutput(RHICmdList, ComputeShader, DstUAV);
 
diff --git a/Engine/Source/Runtime/Renderer/Private/GPUFastFourierTransform.h b/Engine/Source/Runtime/Renderer/Private/GPUFastFourierTransform.h
--- a/Engine/Source/Runtime/Renderer/Private/GPUFastFourierTransform.h
+++ b/Engine/Source/Runtime/Renderer/Private/GPUFastFourierTransform.h
@@ -82,7 +82,7 @@ namespace GPUFFT
 	{
 	public:
 
-		typedef TShaderMap<FGlobalShaderType>  ShaderMapType;
+		typedef FGlobalShaderMap  ShaderMapType;
 
 		FGPUFFTShaderContext(FRHICommandList& CmdList, const FGPUFFTShaderContext::ShaderMapType& Map)
 			: RHICmdList(&CmdList), ShaderMap(&Map)
@@ -571,9 +571,9 @@ namespace GPUFFTComputeShaderUtils
 
 		// Factory method.
 		template <typename ComputeShaderT>
-		inline static FScopedUAVBind BindOutput(FRHICommandList& CmdList, ComputeShaderT* Shader, FUnorderedAccessViewRHIRef BufferUAV)
+		inline static FScopedUAVBind BindOutput(FRHICommandList& CmdList, const TShaderRef<ComputeShaderT>& Shader, const FUnorderedAccessViewRHIRef& BufferUAV)
 		{
-			FRHIComputeShader* ShaderRHI = Shader->GetComputeShader();
+			FRHIComputeShader* ShaderRHI = Shader.GetComputeShader();
 			const uint32 BaseIndex = Shader->DestinationResourceParameter().GetBaseIndex();
 			return FScopedUAVBind(CmdList, ShaderRHI, BaseIndex, BufferUAV);
 		}
diff --git a/Engine/Source/Runtime/Renderer/Private/GammaCorrection.cpp b/Engine/Source/Runtime/Renderer/Private/GammaCorrection.cpp
--- a/Engine/Source/Runtime/Renderer/Private/GammaCorrection.cpp
+++ b/Engine/Source/Runtime/Renderer/Private/GammaCorrection.cpp
@@ -33,11 +33,11 @@ class FGammaCorrectionPS : public FGlobalShader
 
 public:
 
-	FShaderResourceParameter SceneTexture;
-	FShaderResourceParameter SceneTextureSampler;
-	FShaderParameter InverseGamma;
-	FShaderParameter ColorScale;
-	FShaderParameter OverlayColor;
+	LAYOUT_FIELD(FShaderResourceParameter, SceneTexture);
+	LAYOUT_FIELD(FShaderResourceParameter, SceneTextureSampler);
+	LAYOUT_FIELD(FShaderParameter, InverseGamma);
+	LAYOUT_FIELD(FShaderParameter, ColorScale);
+	LAYOUT_FIELD(FShaderParameter, OverlayColor);
 
 	/** Initialization constructor. */
 	FGammaCorrectionPS(const ShaderMetaType::CompiledShaderInitializerType& Initializer)
@@ -49,15 +49,6 @@ public:
 		ColorScale.Bind(Initializer.ParameterMap,TEXT("ColorScale"));
 		OverlayColor.Bind(Initializer.ParameterMap,TEXT("OverlayColor"));
 	}
-
-	// FShader interface.
-	virtual bool Serialize(FArchive& Ar) override
-	{
-		bool bShaderHasOutdatedParameters = FGlobalShader::Serialize(Ar);
-		
-		Ar << SceneTexture << SceneTextureSampler << InverseGamma << ColorScale << OverlayColor;
-		return bShaderHasOutdatedParameters;
-	}
 };
 
 /** Encapsulates the gamma correction vertex shader. */
@@ -127,8 +118,8 @@ void FSceneRenderer::GammaCorrectToViewportRenderTarget(FRHICommandList& RHICmdL
 	TShaderMapRef<FGammaCorrectionPS> PixelShader(View->ShaderMap);
 
 	GraphicsPSOInit.BoundShaderState.VertexDeclarationRHI = GFilterVertexDeclaration.VertexDeclarationRHI;
-	GraphicsPSOInit.BoundShaderState.VertexShaderRHI = GETSAFERHISHADER_VERTEX(*VertexShader);
-	GraphicsPSOInit.BoundShaderState.PixelShaderRHI = GETSAFERHISHADER_PIXEL(*PixelShader);
+	GraphicsPSOInit.BoundShaderState.VertexShaderRHI = VertexShader.GetVertexShader();
+	GraphicsPSOInit.BoundShaderState.PixelShaderRHI = PixelShader.GetPixelShader();
 	GraphicsPSOInit.PrimitiveType = PT_TriangleList;
 
 	SetGraphicsPipelineState(RHICmdList, GraphicsPSOInit);
@@ -140,7 +131,7 @@ void FSceneRenderer::GammaCorrectToViewportRenderTarget(FRHICommandList& RHICmdL
 		InvDisplayGamma = 1 / OverrideGamma;
 	}
 
-	FRHIPixelShader* ShaderRHI = PixelShader->GetPixelShader();
+	FRHIPixelShader* ShaderRHI = PixelShader.GetPixelShader();
 
 	SetShaderValue(
 		RHICmdList, 
@@ -172,7 +163,7 @@ void FSceneRenderer::GammaCorrectToViewportRenderTarget(FRHICommandList& RHICmdL
 		View->ViewRect.Width(),View->ViewRect.Height(),
 		ViewFamily.RenderTarget->GetSizeXY(),
 		SceneContext.GetBufferSizeXY(),
-		*VertexShader,
+		VertexShader,
 		EDRF_UseTriangleOptimization);
 
 	RHICmdList.EndRenderPass();
diff --git a/Engine/Source/Runtime/Renderer/Private/GlobalDistanceField.cpp b/Engine/Source/Runtime/Renderer/Private/GlobalDistanceField.cpp
--- a/Engine/Source/Runtime/Renderer/Private/GlobalDistanceField.cpp
+++ b/Engine/Source/Runtime/Renderer/Private/GlobalDistanceField.cpp
@@ -190,7 +190,7 @@ public:
 
 	void SetParameters(FRHICommandList& RHICmdList, const FScene* Scene, const FSceneView& View, float MaxOcclusionDistance, const FVector4& VolumeBoundsValue, FGlobalDFCacheType CacheType)
 	{
-		FRHIComputeShader* ShaderRHI = GetComputeShader();
+		FRHIComputeShader* ShaderRHI = RHICmdList.GetBoundComputeShader();
 		FGlobalShader::SetParameters<FViewUniformShaderParameters>(RHICmdList, ShaderRHI, View.ViewUniformBuffer);
 
 		FRHITexture* TextureAtlas;
@@ -245,32 +245,21 @@ public:
 
 	void UnsetParameters(FRHICommandList& RHICmdList, const FScene* Scene)
 	{
-		ObjectBufferParameters.UnsetParameters(RHICmdList, GetComputeShader(), *(Scene->DistanceFieldSceneData.GetCurrentObjectBuffers()));
-		CulledObjectParameters.UnsetParameters(RHICmdList, GetComputeShader());
+		ObjectBufferParameters.UnsetParameters(RHICmdList, RHICmdList.GetBoundComputeShader(), *(Scene->DistanceFieldSceneData.GetCurrentObjectBuffers()));
+		CulledObjectParameters.UnsetParameters(RHICmdList, RHICmdList.GetBoundComputeShader());
 
 		TArray<FRHIUnorderedAccessView*> UAVs;
 		CulledObjectParameters.GetUAVs(GGlobalDistanceFieldCulledObjectBuffers.Buffers, UAVs);
 		RHICmdList.TransitionResources(EResourceTransitionAccess::EReadable, EResourceTransitionPipeline::EComputeToCompute, UAVs.GetData(), UAVs.Num());
 	}
 
-	virtual bool Serialize(FArchive& Ar)
-	{		
-		bool bShaderHasOutdatedParameters = FGlobalShader::Serialize(Ar);
-		Ar << ObjectBufferParameters;
-		Ar << CulledObjectParameters;
-		Ar << AOGlobalMaxSphereQueryRadius;
-		Ar << VolumeBounds;
-		Ar << AcceptOftenMovingObjectsOnly;
-		return bShaderHasOutdatedParameters;
-	}
-
 private:
 
-	TDistanceFieldObjectBufferParameters<DFPT_SignedDistanceField> ObjectBufferParameters;
-	TDistanceFieldCulledObjectBufferParameters<DFPT_SignedDistanceField> CulledObjectParameters;
-	FShaderParameter AOGlobalMaxSphereQueryRadius;
-	FShaderParameter VolumeBounds;
-	FShaderParameter AcceptOftenMovingObjectsOnly;
+	LAYOUT_FIELD((TDistanceFieldObjectBufferParameters<DFPT_SignedDistanceField>), ObjectBufferParameters);
+	LAYOUT_FIELD((TDistanceFieldCulledObjectBufferParameters<DFPT_SignedDistanceField>), CulledObjectParameters);
+	LAYOUT_FIELD(FShaderParameter, AOGlobalMaxSphereQueryRadius);
+	LAYOUT_FIELD(FShaderParameter, VolumeBounds);
+	LAYOUT_FIELD(FShaderParameter, AcceptOftenMovingObjectsOnly);
 };
 
 IMPLEMENT_SHADER_TYPE(,FCullObjectsForVolumeCS,TEXT("/Engine/Private/GlobalDistanceField.usf"),TEXT("CullObjectsForVolumeCS"),SF_Compute);
@@ -385,7 +374,7 @@ public:
 		int32 ClipmapIndexValue,
 		const FVolumeUpdateRegion& UpdateRegion)
 	{
-		FRHIComputeShader* ShaderRHI = GetComputeShader();
+		FRHIComputeShader* ShaderRHI = RHICmdList.GetBoundComputeShader();
 		FGlobalShader::SetParameters<FViewUniformShaderParameters>(RHICmdList, ShaderRHI, View.ViewUniformBuffer);
 
 		FRHITexture* TextureAtlas;
@@ -423,38 +412,22 @@ public:
 
 	void UnsetParameters(FRHICommandList& RHICmdList)
 	{
-		CullGridObjectNum.UnsetUAV(RHICmdList, GetComputeShader());
-		CullGridObjectArray.UnsetUAV(RHICmdList, GetComputeShader());
+		CullGridObjectNum.UnsetUAV(RHICmdList, RHICmdList.GetBoundComputeShader());
+		CullGridObjectArray.UnsetUAV(RHICmdList, RHICmdList.GetBoundComputeShader());
 		RHICmdList.TransitionResource(EResourceTransitionAccess::EReadable, EResourceTransitionPipeline::EComputeToCompute, GObjectGridBuffers.CullGridObjectNum.UAV);
 		RHICmdList.TransitionResource(EResourceTransitionAccess::EReadable, EResourceTransitionPipeline::EComputeToCompute, GObjectGridBuffers.CullGridObjectArray.UAV);
 	}
 
-	virtual bool Serialize(FArchive& Ar)
-	{		
-		bool bShaderHasOutdatedParameters = FGlobalShader::Serialize(Ar);
-		Ar << CulledObjectBufferParameters;
-		Ar << GlobalDistanceFieldParameters;
-		Ar << CullGridObjectNum;
-		Ar << CullGridObjectArray;
-		Ar << CullGridDimension;
-		Ar << VolumeTexelSize;
-		Ar << UpdateRegionVolumeMin;
-		Ar << ClipmapIndex;
-		Ar << AOGlobalMaxSphereQueryRadius;
-		return bShaderHasOutdatedParameters;
-	}
-
 private:
-
-	TDistanceFieldCulledObjectBufferParameters<DFPT_SignedDistanceField> CulledObjectBufferParameters;
-	FGlobalDistanceFieldParameters GlobalDistanceFieldParameters;
-	FRWShaderParameter CullGridObjectNum;
-	FRWShaderParameter CullGridObjectArray;
-	FShaderParameter CullGridDimension;
-	FShaderParameter VolumeTexelSize;
-	FShaderParameter UpdateRegionVolumeMin;
-	FShaderParameter ClipmapIndex;
-	FShaderParameter AOGlobalMaxSphereQueryRadius;
+	LAYOUT_FIELD((TDistanceFieldCulledObjectBufferParameters<DFPT_SignedDistanceField>), CulledObjectBufferParameters);
+	LAYOUT_FIELD(FGlobalDistanceFieldParameters, GlobalDistanceFieldParameters);
+	LAYOUT_FIELD(FRWShaderParameter, CullGridObjectNum);
+	LAYOUT_FIELD(FRWShaderParameter, CullGridObjectArray);
+	LAYOUT_FIELD(FShaderParameter, CullGridDimension);
+	LAYOUT_FIELD(FShaderParameter, VolumeTexelSize);
+	LAYOUT_FIELD(FShaderParameter, UpdateRegionVolumeMin);
+	LAYOUT_FIELD(FShaderParameter, ClipmapIndex);
+	LAYOUT_FIELD(FShaderParameter, AOGlobalMaxSphereQueryRadius);
 };
 
 IMPLEMENT_SHADER_TYPE(,FCullObjectsToGridCS,TEXT("/Engine/Private/GlobalDistanceField.usf"),TEXT("CullObjectsToGridCS"),SF_Compute);
@@ -531,7 +504,7 @@ public:
 		int32 ClipmapIndexValue,
 		const FVolumeUpdateRegion& UpdateRegion)
 	{
-		FRHIComputeShader* ShaderRHI = GetComputeShader();
+		FRHIComputeShader* ShaderRHI = RHICmdList.GetBoundComputeShader();
 		FGlobalShader::SetParameters<FViewUniformShaderParameters>(RHICmdList, ShaderRHI, View.ViewUniformBuffer);
 
 		FRHITexture* TextureAtlas;
@@ -582,44 +555,25 @@ public:
 
 	void UnsetParameters(FRHICommandList& RHICmdList,const FGlobalDistanceFieldClipmap& Clipmap)
 	{
-		GlobalDistanceFieldTexture.UnsetUAV(RHICmdList, GetComputeShader());
+		GlobalDistanceFieldTexture.UnsetUAV(RHICmdList, RHICmdList.GetBoundComputeShader());
 
 		const FSceneRenderTargetItem& ClipMapRTI = Clipmap.RenderTarget->GetRenderTargetItem();
 		RHICmdList.TransitionResource(EResourceTransitionAccess::EReadable, EResourceTransitionPipeline::EComputeToCompute, ClipMapRTI.UAV);
 	}
 
-	virtual bool Serialize(FArchive& Ar)
-	{		
-		bool bShaderHasOutdatedParameters = FGlobalShader::Serialize(Ar);
-		Ar << CulledObjectBufferParameters;
-		Ar << GlobalDistanceFieldParameters;
-		Ar << GlobalDistanceFieldTexture;
-		Ar << ParentGlobalDistanceFieldTexture;
-		Ar << CullGridObjectNum;
-		Ar << CullGridObjectArray;
-		Ar << UpdateRegionSize;
-		Ar << CullGridDimension;
-		Ar << VolumeTexelSize;
-		Ar << UpdateRegionVolumeMin;
-		Ar << ClipmapIndex;
-		Ar << AOGlobalMaxSphereQueryRadius;
-		return bShaderHasOutdatedParameters;
-	}
-
 private:
-
-	TDistanceFieldCulledObjectBufferParameters<DFPT_SignedDistanceField> CulledObjectBufferParameters;
-	FGlobalDistanceFieldParameters GlobalDistanceFieldParameters;
-	FRWShaderParameter GlobalDistanceFieldTexture;
-	FShaderResourceParameter ParentGlobalDistanceFieldTexture;
-	FShaderResourceParameter CullGridObjectNum;
-	FShaderResourceParameter CullGridObjectArray;
-	FShaderParameter UpdateRegionSize;
-	FShaderParameter CullGridDimension;
-	FShaderParameter VolumeTexelSize;
-	FShaderParameter UpdateRegionVolumeMin;
-	FShaderParameter ClipmapIndex;
-	FShaderParameter AOGlobalMaxSphereQueryRadius;
+	LAYOUT_FIELD((TDistanceFieldCulledObjectBufferParameters<DFPT_SignedDistanceField>), CulledObjectBufferParameters);
+	LAYOUT_FIELD(FGlobalDistanceFieldParameters, GlobalDistanceFieldParameters);
+	LAYOUT_FIELD(FRWShaderParameter, GlobalDistanceFieldTexture);
+	LAYOUT_FIELD(FShaderResourceParameter, ParentGlobalDistanceFieldTexture);
+	LAYOUT_FIELD(FShaderResourceParameter, CullGridObjectNum);
+	LAYOUT_FIELD(FShaderResourceParameter, CullGridObjectArray);
+	LAYOUT_FIELD(FShaderParameter, UpdateRegionSize);
+	LAYOUT_FIELD(FShaderParameter, CullGridDimension);
+	LAYOUT_FIELD(FShaderParameter, VolumeTexelSize);
+	LAYOUT_FIELD(FShaderParameter, UpdateRegionVolumeMin);
+	LAYOUT_FIELD(FShaderParameter, ClipmapIndex);
+	LAYOUT_FIELD(FShaderParameter, AOGlobalMaxSphereQueryRadius);
 };
 
 #define IMPLEMENT_GLOBALDF_COMPOSITE_CS_TYPE(bUseParentDistanceField, FlattenedDimension) \
@@ -687,7 +641,7 @@ public:
 		UTexture2D* VisibilityTextureValue,
 		int32 NumHeightfieldsValue)
 	{
-		FRHIComputeShader* ShaderRHI = GetComputeShader();
+		FRHIComputeShader* ShaderRHI = RHICmdList.GetBoundComputeShader();
 		FGlobalShader::SetParameters<FViewUniformShaderParameters>(RHICmdList, ShaderRHI, View.ViewUniformBuffer);
 		GlobalDistanceFieldParameters.Set(RHICmdList, ShaderRHI, GlobalDistanceFieldInfo.ParameterData);
 
@@ -711,42 +665,24 @@ public:
 
 	void UnsetParameters(FRHICommandList& RHICmdList, const FGlobalDistanceFieldInfo& GlobalDistanceFieldInfo, const FGlobalDistanceFieldClipmap& Clipmap)
 	{
-		GlobalDistanceFieldTexture.UnsetUAV(RHICmdList, GetComputeShader());
+		GlobalDistanceFieldTexture.UnsetUAV(RHICmdList, RHICmdList.GetBoundComputeShader());
 
 		const FSceneRenderTargetItem& ClipMapRTI = Clipmap.RenderTarget->GetRenderTargetItem();
 		RHICmdList.TransitionResource(EResourceTransitionAccess::EReadable, EResourceTransitionPipeline::EComputeToGfx, ClipMapRTI.UAV);
 	}
 
-	virtual bool Serialize(FArchive& Ar)
-	{
-		bool bShaderHasOutdatedParameters = FGlobalShader::Serialize(Ar);
-		Ar << GlobalDistanceFieldParameters;
-		Ar << GlobalDistanceFieldTexture;
-		Ar << UpdateRegionSize;
-		Ar << VolumeTexelSize;
-		Ar << UpdateRegionVolumeMin;
-		Ar << UpdateRegionVolumeStep;
-		Ar << ClipmapIndex;
-		Ar << AOGlobalMaxSphereQueryRadius;
-		Ar << HeightfieldDescriptionParameters;
-		Ar << HeightfieldTextureParameters;
-		Ar << HeightfieldThickness;
-		return bShaderHasOutdatedParameters;
-	}
-
 private:
-
-	FGlobalDistanceFieldParameters GlobalDistanceFieldParameters;
-	FRWShaderParameter GlobalDistanceFieldTexture;
-	FShaderParameter UpdateRegionSize;
-	FShaderParameter VolumeTexelSize;
-	FShaderParameter UpdateRegionVolumeMin;
-	FShaderParameter UpdateRegionVolumeStep;
-	FShaderParameter ClipmapIndex;
-	FShaderParameter AOGlobalMaxSphereQueryRadius;
-	FHeightfieldDescriptionParameters HeightfieldDescriptionParameters;
-	FHeightfieldTextureParameters HeightfieldTextureParameters;
-	FShaderParameter HeightfieldThickness;
+	LAYOUT_FIELD(FGlobalDistanceFieldParameters, GlobalDistanceFieldParameters);
+	LAYOUT_FIELD(FRWShaderParameter, GlobalDistanceFieldTexture);
+	LAYOUT_FIELD(FShaderParameter, UpdateRegionSize);
+	LAYOUT_FIELD(FShaderParameter, VolumeTexelSize);
+	LAYOUT_FIELD(FShaderParameter, UpdateRegionVolumeMin);
+	LAYOUT_FIELD(FShaderParameter, UpdateRegionVolumeStep);
+	LAYOUT_FIELD(FShaderParameter, ClipmapIndex);
+	LAYOUT_FIELD(FShaderParameter, AOGlobalMaxSphereQueryRadius);
+	LAYOUT_FIELD(FHeightfieldDescriptionParameters, HeightfieldDescriptionParameters);
+	LAYOUT_FIELD(FHeightfieldTextureParameters, HeightfieldTextureParameters);
+	LAYOUT_FIELD(FShaderParameter, HeightfieldThickness);
 };
 
 IMPLEMENT_SHADER_TYPE(, FCompositeHeightfieldsIntoGlobalDistanceFieldCS, TEXT("/Engine/Private/GlobalDistanceField.usf"), TEXT("CompositeHeightfieldsIntoGlobalDistanceFieldCS"), SF_Compute);
@@ -827,14 +763,14 @@ void FHeightfieldLightingViewInfo::CompositeHeightfieldsIntoGlobalDistanceField(
 					UTexture2D* VisibilityTexture = It.Key().Visibility;
 
 					TShaderMapRef<FCompositeHeightfieldsIntoGlobalDistanceFieldCS> ComputeShader(View.ShaderMap);
-					RHICmdList.SetComputeShader(ComputeShader->GetComputeShader());
+					RHICmdList.SetComputeShader(ComputeShader.GetComputeShader());
 					ComputeShader->SetParameters(RHICmdList, Scene, View, GlobalMaxSphereQueryRadius, GlobalDistanceFieldInfo, Clipmap, ClipmapIndexValue, UpdateRegion, HeightfieldTexture, VisibilityTexture, HeightfieldDescriptions.Num());
 
 					//@todo - match typical update sizes.  Camera movement creates narrow slabs.
 					const uint32 NumGroupsX = FMath::DivideAndRoundUp<int32>(UpdateRegion.CellsSize.X, HeightfieldCompositeTileSize);
 					const uint32 NumGroupsY = FMath::DivideAndRoundUp<int32>(UpdateRegion.CellsSize.Y, HeightfieldCompositeTileSize);
 
-					DispatchComputeShader(RHICmdList, *ComputeShader, NumGroupsX, NumGroupsY, 1);
+					DispatchComputeShader(RHICmdList, ComputeShader.GetShader(), NumGroupsX, NumGroupsY, 1);
 					ComputeShader->UnsetParameters(RHICmdList, GlobalDistanceFieldInfo, Clipmap);
 				}
 			}
@@ -1469,25 +1405,25 @@ void UpdateGlobalDistanceFieldVolume(
 									RHICmdList.ClearUAVUint(GGlobalDistanceFieldCulledObjectBuffers.Buffers.ObjectIndirectArguments.UAV, FUintVector4(0, 0, 0, 0));
 
 									TShaderMapRef<FCullObjectsForVolumeCS> ComputeShader(View.ShaderMap);
-									RHICmdList.SetComputeShader(ComputeShader->GetComputeShader());
+									RHICmdList.SetComputeShader(ComputeShader.GetComputeShader());
 									const FVector4 VolumeBounds(UpdateRegion.Bounds.GetCenter(), UpdateRegion.Bounds.GetExtent().Size());
 									ComputeShader->SetParameters(RHICmdList, Scene, View, MaxOcclusionDistance, VolumeBounds, (FGlobalDFCacheType)CacheType);
 
-									DispatchComputeShader(RHICmdList, *ComputeShader, FMath::DivideAndRoundUp<uint32>(Scene->DistanceFieldSceneData.NumObjectsInBuffer, CullObjectsGroupSize), 1, 1);
+									DispatchComputeShader(RHICmdList, ComputeShader.GetShader(), FMath::DivideAndRoundUp<uint32>(Scene->DistanceFieldSceneData.NumObjectsInBuffer, CullObjectsGroupSize), 1, 1);
 									ComputeShader->UnsetParameters(RHICmdList, Scene);
 								}
 
 								// Further cull the objects into a low resolution grid
 								{
 									TShaderMapRef<FCullObjectsToGridCS> ComputeShader(View.ShaderMap);
-									RHICmdList.SetComputeShader(ComputeShader->GetComputeShader());
+									RHICmdList.SetComputeShader(ComputeShader.GetComputeShader());
 									ComputeShader->SetParameters(RHICmdList, Scene, View, MaxOcclusionDistance, GlobalDistanceFieldInfo, ClipmapIndex, UpdateRegion);
 
 									const uint32 NumGroupsX = FMath::DivideAndRoundUp<int32>(UpdateRegion.CellsSize.X, GCullGridTileSize);
 									const uint32 NumGroupsY = FMath::DivideAndRoundUp<int32>(UpdateRegion.CellsSize.Y, GCullGridTileSize);
 									const uint32 NumGroupsZ = FMath::DivideAndRoundUp<int32>(UpdateRegion.CellsSize.Z, GCullGridTileSize); 
 
-									DispatchComputeShader(RHICmdList, *ComputeShader, NumGroupsX, NumGroupsY, NumGroupsZ);
+									DispatchComputeShader(RHICmdList, ComputeShader.GetShader(), NumGroupsX, NumGroupsY, NumGroupsZ);
 									ComputeShader->UnsetParameters(RHICmdList);
 								}
 							}
@@ -1522,34 +1458,34 @@ void UpdateGlobalDistanceFieldVolume(
 									if (FlattenedDimension == Flatten_None)
 									{
 										TShaderMapRef<TCompositeObjectDistanceFieldsCS<true, Flatten_None>> ComputeShader(View.ShaderMap);
-										RHICmdList.SetComputeShader(ComputeShader->GetComputeShader());
+										RHICmdList.SetComputeShader(ComputeShader.GetComputeShader());
 										ComputeShader->SetParameters(RHICmdList, Scene, View, MaxOcclusionDistance, GlobalDistanceFieldInfo.ParameterData, Clipmap, ParentDistanceField, ClipmapIndex, UpdateRegion);
-										DispatchComputeShader(RHICmdList, *ComputeShader, NumGroupsX, NumGroupsY, NumGroupsZ);
+										DispatchComputeShader(RHICmdList, ComputeShader.GetShader(), NumGroupsX, NumGroupsY, NumGroupsZ);
 										ComputeShader->UnsetParameters(RHICmdList, Clipmap);
 									}
 									else if (FlattenedDimension == Flatten_XAxis)
 									{
 										TShaderMapRef<TCompositeObjectDistanceFieldsCS<true, Flatten_XAxis>> ComputeShader(View.ShaderMap);
-										RHICmdList.SetComputeShader(ComputeShader->GetComputeShader());
+										RHICmdList.SetComputeShader(ComputeShader.GetComputeShader());
 										ComputeShader->SetParameters(RHICmdList, Scene, View, MaxOcclusionDistance, GlobalDistanceFieldInfo.ParameterData, Clipmap, ParentDistanceField, ClipmapIndex, UpdateRegion);
-										DispatchComputeShader(RHICmdList, *ComputeShader, NumGroupsX, NumGroupsY, NumGroupsZ);
+										DispatchComputeShader(RHICmdList, ComputeShader.GetShader(), NumGroupsX, NumGroupsY, NumGroupsZ);
 										ComputeShader->UnsetParameters(RHICmdList, Clipmap);
 									}
 									else if (FlattenedDimension == Flatten_YAxis)
 									{
 										TShaderMapRef<TCompositeObjectDistanceFieldsCS<true, Flatten_YAxis>> ComputeShader(View.ShaderMap);
-										RHICmdList.SetComputeShader(ComputeShader->GetComputeShader());
+										RHICmdList.SetComputeShader(ComputeShader.GetComputeShader());
 										ComputeShader->SetParameters(RHICmdList, Scene, View, MaxOcclusionDistance, GlobalDistanceFieldInfo.ParameterData, Clipmap, ParentDistanceField, ClipmapIndex, UpdateRegion);
-										DispatchComputeShader(RHICmdList, *ComputeShader, NumGroupsX, NumGroupsY, NumGroupsZ);
+										DispatchComputeShader(RHICmdList, ComputeShader.GetShader(), NumGroupsX, NumGroupsY, NumGroupsZ);
 										ComputeShader->UnsetParameters(RHICmdList, Clipmap);
 									}
 									else
 									{
 										check(FlattenedDimension == Flatten_ZAxis);
 										TShaderMapRef<TCompositeObjectDistanceFieldsCS<true, Flatten_ZAxis>> ComputeShader(View.ShaderMap);
-										RHICmdList.SetComputeShader(ComputeShader->GetComputeShader());
+										RHICmdList.SetComputeShader(ComputeShader.GetComputeShader());
 										ComputeShader->SetParameters(RHICmdList, Scene, View, MaxOcclusionDistance, GlobalDistanceFieldInfo.ParameterData, Clipmap, ParentDistanceField, ClipmapIndex, UpdateRegion);
-										DispatchComputeShader(RHICmdList, *ComputeShader, NumGroupsX, NumGroupsY, NumGroupsZ);
+										DispatchComputeShader(RHICmdList, ComputeShader.GetShader(), NumGroupsX, NumGroupsY, NumGroupsZ);
 										ComputeShader->UnsetParameters(RHICmdList, Clipmap);
 									}
 								}
@@ -1558,34 +1494,34 @@ void UpdateGlobalDistanceFieldVolume(
 									if (FlattenedDimension == Flatten_None)
 									{
 										TShaderMapRef<TCompositeObjectDistanceFieldsCS<false, Flatten_None>> ComputeShader(View.ShaderMap);
-										RHICmdList.SetComputeShader(ComputeShader->GetComputeShader());
+										RHICmdList.SetComputeShader(ComputeShader.GetComputeShader());
 										ComputeShader->SetParameters(RHICmdList, Scene, View, MaxOcclusionDistance, GlobalDistanceFieldInfo.ParameterData, Clipmap, NULL, ClipmapIndex, UpdateRegion);
-										DispatchComputeShader(RHICmdList, *ComputeShader, NumGroupsX, NumGroupsY, NumGroupsZ);
+										DispatchComputeShader(RHICmdList, ComputeShader.GetShader(), NumGroupsX, NumGroupsY, NumGroupsZ);
 										ComputeShader->UnsetParameters(RHICmdList, Clipmap);
 									}
 									else if (FlattenedDimension == Flatten_XAxis)
 									{
 										TShaderMapRef<TCompositeObjectDistanceFieldsCS<false, Flatten_XAxis>> ComputeShader(View.ShaderMap);
-										RHICmdList.SetComputeShader(ComputeShader->GetComputeShader());
+										RHICmdList.SetComputeShader(ComputeShader.GetComputeShader());
 										ComputeShader->SetParameters(RHICmdList, Scene, View, MaxOcclusionDistance, GlobalDistanceFieldInfo.ParameterData, Clipmap, NULL, ClipmapIndex, UpdateRegion);
-										DispatchComputeShader(RHICmdList, *ComputeShader, NumGroupsX, NumGroupsY, NumGroupsZ);
+										DispatchComputeShader(RHICmdList, ComputeShader.GetShader(), NumGroupsX, NumGroupsY, NumGroupsZ);
 										ComputeShader->UnsetParameters(RHICmdList, Clipmap);
 									}
 									else if (FlattenedDimension == Flatten_YAxis)
 									{
 										TShaderMapRef<TCompositeObjectDistanceFieldsCS<false, Flatten_YAxis>> ComputeShader(View.ShaderMap);
-										RHICmdList.SetComputeShader(ComputeShader->GetComputeShader());
+										RHICmdList.SetComputeShader(ComputeShader.GetComputeShader());
 										ComputeShader->SetParameters(RHICmdList, Scene, View, MaxOcclusionDistance, GlobalDistanceFieldInfo.ParameterData, Clipmap, NULL, ClipmapIndex, UpdateRegion);
-										DispatchComputeShader(RHICmdList, *ComputeShader, NumGroupsX, NumGroupsY, NumGroupsZ);
+										DispatchComputeShader(RHICmdList, ComputeShader.GetShader(), NumGroupsX, NumGroupsY, NumGroupsZ);
 										ComputeShader->UnsetParameters(RHICmdList, Clipmap);
 									}
 									else
 									{
 										check(FlattenedDimension == Flatten_ZAxis);
 										TShaderMapRef<TCompositeObjectDistanceFieldsCS<false, Flatten_ZAxis>> ComputeShader(View.ShaderMap);
-										RHICmdList.SetComputeShader(ComputeShader->GetComputeShader());
+										RHICmdList.SetComputeShader(ComputeShader.GetComputeShader());
 										ComputeShader->SetParameters(RHICmdList, Scene, View, MaxOcclusionDistance, GlobalDistanceFieldInfo.ParameterData, Clipmap, NULL, ClipmapIndex, UpdateRegion);
-										DispatchComputeShader(RHICmdList, *ComputeShader, NumGroupsX, NumGroupsY, NumGroupsZ);
+										DispatchComputeShader(RHICmdList, ComputeShader.GetShader(), NumGroupsX, NumGroupsY, NumGroupsZ);
 										ComputeShader->UnsetParameters(RHICmdList, Clipmap);
 									}
 								}
diff --git a/Engine/Source/Runtime/Renderer/Private/GpuDebugRendering.cpp b/Engine/Source/Runtime/Renderer/Private/GpuDebugRendering.cpp
--- a/Engine/Source/Runtime/Renderer/Private/GpuDebugRendering.cpp
+++ b/Engine/Source/Runtime/Renderer/Private/GpuDebugRendering.cpp
@@ -235,7 +235,7 @@ namespace ShaderDrawDebug
 			ERDGPassFlags::Compute,
 			[Parameters, ComputeShader](FRHICommandList& RHICmdList)
 		{
-			FComputeShaderUtils::Dispatch(RHICmdList, *ComputeShader, *Parameters, FIntVector(1,1,1));
+			FComputeShaderUtils::Dispatch(RHICmdList, ComputeShader, *Parameters, FIntVector(1,1,1));
 		});
 
 		GraphBuilder.QueueBufferExtraction(DataBuffer, &View.ShaderDrawData.Buffer, FRDGResourceState::EAccess::Write, FRDGResourceState::EPipeline::Compute);
@@ -287,10 +287,10 @@ namespace ShaderDrawDebug
 				PassParameters->ShaderDrawVSParameters.LockedShaderDrawDebugPrimitive = LockedDataBuffer;
 			}
 
-			ValidateShaderParameters(*PixelShader, PassParameters->ShaderDrawPSParameters);
-			ClearUnusedGraphResources(*PixelShader, &PassParameters->ShaderDrawPSParameters, { IndirectBuffer });
-			ValidateShaderParameters(*VertexShader, PassParameters->ShaderDrawVSParameters);
-			ClearUnusedGraphResources(*VertexShader, &PassParameters->ShaderDrawVSParameters, { IndirectBuffer });
+			ValidateShaderParameters(PixelShader, PassParameters->ShaderDrawPSParameters);
+			ClearUnusedGraphResources(PixelShader, &PassParameters->ShaderDrawPSParameters, { IndirectBuffer });
+			ValidateShaderParameters(VertexShader, PassParameters->ShaderDrawVSParameters);
+			ClearUnusedGraphResources(VertexShader, &PassParameters->ShaderDrawVSParameters, { IndirectBuffer });
 
 			GraphBuilder.AddPass(
 				RDG_EVENT_NAME("ShaderDrawDebug"),
@@ -311,12 +311,12 @@ namespace ShaderDrawDebug
 				GraphicsPSOInit.RasterizerState = TStaticRasterizerState<FM_Solid, CM_None, true>::GetRHI();
 				GraphicsPSOInit.PrimitiveType = PT_LineList;
 				GraphicsPSOInit.BoundShaderState.VertexDeclarationRHI = GEmptyVertexDeclaration.VertexDeclarationRHI;
-				GraphicsPSOInit.BoundShaderState.VertexShaderRHI = GETSAFERHISHADER_VERTEX(*VertexShader);
-				GraphicsPSOInit.BoundShaderState.PixelShaderRHI = GETSAFERHISHADER_PIXEL(*PixelShader);
+				GraphicsPSOInit.BoundShaderState.VertexShaderRHI = VertexShader.GetVertexShader();
+				GraphicsPSOInit.BoundShaderState.PixelShaderRHI = PixelShader.GetPixelShader();
 				SetGraphicsPipelineState(RHICmdListImmediate, GraphicsPSOInit);
 
-				SetShaderParameters(RHICmdListImmediate, *VertexShader, VertexShader->GetVertexShader(), PassParameters->ShaderDrawVSParameters);
-				SetShaderParameters(RHICmdListImmediate, *PixelShader, PixelShader->GetPixelShader(), PassParameters->ShaderDrawPSParameters);
+				SetShaderParameters(RHICmdListImmediate, VertexShader, VertexShader.GetVertexShader(), PassParameters->ShaderDrawVSParameters);
+				SetShaderParameters(RHICmdListImmediate, PixelShader, PixelShader.GetPixelShader(), PassParameters->ShaderDrawPSParameters);
 
 
 				if (bUseRdgInput)
@@ -360,7 +360,7 @@ namespace ShaderDrawDebug
 				Parameters->NumElements = NumElements;
 				Parameters->InStructuredBuffer = GraphBuilder.CreateSRV(DataBuffer);
 				Parameters->OutStructuredBuffer = LockedData.Buffer.UAV;
-				FComputeShaderUtils::AddPass(GraphBuilder, RDG_EVENT_NAME("ShaderDrawDebugCopy"), *ComputeShader, Parameters, FIntVector(FMath::CeilToInt(NumElements / 1024.f), 1, 1));
+				FComputeShaderUtils::AddPass(GraphBuilder, RDG_EVENT_NAME("ShaderDrawDebugCopy"), ComputeShader, Parameters, FIntVector(FMath::CeilToInt(NumElements / 1024.f), 1, 1));
 			}
 			{
 				FShaderDrawDebugCopyCS::FPermutationDomain PermutationVector;
@@ -370,7 +370,7 @@ namespace ShaderDrawDebug
 				FShaderDrawDebugCopyCS::FParameters* Parameters = GraphBuilder.AllocParameters<FShaderDrawDebugCopyCS::FParameters>();
 				Parameters->InBuffer = GraphBuilder.CreateSRV(IndirectBuffer);
 				Parameters->OutBuffer = LockedData.IndirectBuffer.UAV;
-				FComputeShaderUtils::AddPass(GraphBuilder, RDG_EVENT_NAME("ShaderDrawDebugCopy"), *ComputeShader, Parameters, FIntVector(1, 1, 1));
+				FComputeShaderUtils::AddPass(GraphBuilder, RDG_EVENT_NAME("ShaderDrawDebugCopy"), ComputeShader, Parameters, FIntVector(1, 1, 1));
 			}
 		}
 	}
diff --git a/Engine/Source/Runtime/Renderer/Private/HairStrands/HairStrandsCluster.cpp b/Engine/Source/Runtime/Renderer/Private/HairStrands/HairStrandsCluster.cpp
--- a/Engine/Source/Runtime/Renderer/Private/HairStrands/HairStrandsCluster.cpp
+++ b/Engine/Source/Runtime/Renderer/Private/HairStrands/HairStrandsCluster.cpp
@@ -105,7 +105,7 @@ static void AddHairMacroGroupAABBPass(
 		FComputeShaderUtils::AddPass(
 			GraphBuilder,
 			RDG_EVENT_NAME("HairStrandsMacroGroupAABBUpdate"),
-			*ComputeShader,
+			ComputeShader,
 			Parameters,
 			FIntVector(1,1,1));
 
@@ -151,7 +151,7 @@ FHairStrandsMacroGroupViews CreateHairStrandsMacroGroups(
 			uint32 MacroGroupId = 0;
 			auto UpdateMacroGroup = [&MacroGroups, &View, &MacroGroupId, &MaterialId](const FMeshBatchAndRelevance* MeshBatchAndRelevance, const FPrimitiveSceneProxy* Proxy)
 			{
-				const bool bIsHairStrandsFactory = MeshBatchAndRelevance->Mesh->VertexFactory->GetType()->GetId() == CompatibleVF->GetId();
+				const bool bIsHairStrandsFactory = MeshBatchAndRelevance->Mesh->VertexFactory->GetType()->GetHashedName() == CompatibleVF->GetHashedName();
 				if (!bIsHairStrandsFactory)
 					return;
 
diff --git a/Engine/Source/Runtime/Renderer/Private/HairStrands/HairStrandsClusters.cpp b/Engine/Source/Runtime/Renderer/Private/HairStrands/HairStrandsClusters.cpp
--- a/Engine/Source/Runtime/Renderer/Private/HairStrands/HairStrandsClusters.cpp
+++ b/Engine/Source/Runtime/Renderer/Private/HairStrands/HairStrandsClusters.cpp
@@ -69,7 +69,7 @@ static void AddClearClusterAABBPass(
 	Parameters->OutGroupAABBBuffer = OutGroupAABBuffer;
 
 
-	TShaderMap<FGlobalShaderType>* ShaderMap = GetGlobalShaderMap(ERHIFeatureLevel::SM5);
+	FGlobalShaderMap* ShaderMap = GetGlobalShaderMap(ERHIFeatureLevel::SM5);
 	TShaderMapRef<FClearClusterAABBCS> ComputeShader(ShaderMap);
 
 	const FIntVector DispatchCount = DispatchCount.DivideAndRoundUp(FIntVector(ClusterCount * 6, 1, 1), FIntVector(64, 1, 1));
@@ -343,7 +343,7 @@ struct FHairHZBParameters
 
 static void AddClusterCullingPass(
 	FRDGBuilder& GraphBuilder,
-	TShaderMap<FGlobalShaderType>* ShaderMap,
+	FGlobalShaderMap* ShaderMap,
 	const FViewInfo& View,
 	const FHairCullingParams& CullingParameters,
 	const FHairHZBParameters& HZBParameters,
@@ -386,7 +386,7 @@ static void AddClusterCullingPass(
 		FComputeShaderUtils::AddPass(
 			GraphBuilder,
 			RDG_EVENT_NAME("BufferClearCS"),
-			*ComputeShader,
+			ComputeShader,
 			Parameters,
 			FIntVector(1, 1, 1));
 	}
@@ -446,7 +446,7 @@ static void AddClusterCullingPass(
 		FComputeShaderUtils::AddPass(
 			GraphBuilder,
 			RDG_EVENT_NAME("ClusterCullingCS"),
-			*ComputeShader,
+			ComputeShader,
 			Parameters,
 			DispatchCount);
 	}
@@ -463,7 +463,7 @@ static void AddClusterCullingPass(
 		FComputeShaderUtils::AddPass(
 			GraphBuilder,
 			RDG_EVENT_NAME("PrepareIndirectDrawsCS"),
-			*ComputeShader,
+			ComputeShader,
 			Parameters,
 			FIntVector(2, 1, 1));
 	}
@@ -481,7 +481,7 @@ static void AddClusterCullingPass(
 		FComputeShaderUtils::AddPass(
 			GraphBuilder,
 			RDG_EVENT_NAME("WithinBlockIndexCountPreFixSumCS"),
-			*ComputeShader,
+			ComputeShader,
 			Parameters,
 			DispatchIndirectParametersClusterCountDiv512, 0); // FIX ME, this could get over 65535
 		check(ClusterData.ClusterCount / 512 <= 65535);
@@ -499,7 +499,7 @@ static void AddClusterCullingPass(
 		FComputeShaderUtils::AddPass(
 			GraphBuilder,
 			RDG_EVENT_NAME("BlockIndexCountPreFixSumCS"),
-			*ComputeShader,
+			ComputeShader,
 			Parameters,
 			DispatchIndirectParametersClusterCountDiv512Div512, 0); // FIX ME, this could get over 65535
 		check(ClusterData.ClusterCount / (512*512) <= 65535);
@@ -531,7 +531,7 @@ static void AddClusterCullingPass(
 		FComputeShaderUtils::AddPass(
 			GraphBuilder,
 			RDG_EVENT_NAME("SplatCompactVertexIdsCS"),
-			*ComputeShader,
+			ComputeShader,
 			Parameters,
 			DispatchIndirectParametersClusterCount2D, 0); // DispatchIndirectParametersClusterCount2D is used to avoid having any dispatch dimension going above 65535.
 	}
@@ -553,7 +553,7 @@ static void AddClusterCullingPass(
 
 static void AddClusterResetLod0(
 	FRDGBuilder& GraphBuilder,
-	TShaderMap<FGlobalShaderType>* ShaderMap,
+	FGlobalShaderMap* ShaderMap,
 	FHairStrandClusterData::FHairGroup& ClusterData)
 {
 	// Set as culling result not available
@@ -570,7 +570,7 @@ static void AddClusterResetLod0(
 	FComputeShaderUtils::AddPass(
 		GraphBuilder,
 		RDG_EVENT_NAME("BufferClearCS"),
-		*ComputeShader,
+		ComputeShader,
 		Parameters,
 		FIntVector(1, 1, 1));
 }
@@ -579,7 +579,7 @@ static void AddClusterResetLod0(
 
 void ComputeHairStrandsClustersCulling(
 	FRHICommandListImmediate& RHICmdList,
-	TShaderMap<FGlobalShaderType>& ShaderMap,
+	FGlobalShaderMap& ShaderMap,
 	const TArray<FViewInfo>& Views,
 	const FHairCullingParams& CullingParameters,
 	FHairStrandClusterData& ClusterDatas)
@@ -654,7 +654,7 @@ void ComputeHairStrandsClustersCulling(
 
 void ResetHairStrandsClusterToLOD0(
 	FRHICommandListImmediate& RHICmdList,
-	TShaderMap<FGlobalShaderType>& ShaderMap,
+	FGlobalShaderMap& ShaderMap,
 	FHairStrandClusterData& ClusterDatas)
 {
 	DECLARE_GPU_STAT(HairStrandsResetLod0);
diff --git a/Engine/Source/Runtime/Renderer/Private/HairStrands/HairStrandsClusters.h b/Engine/Source/Runtime/Renderer/Private/HairStrands/HairStrandsClusters.h
--- a/Engine/Source/Runtime/Renderer/Private/HairStrands/HairStrandsClusters.h
+++ b/Engine/Source/Runtime/Renderer/Private/HairStrands/HairStrandsClusters.h
@@ -18,12 +18,12 @@ struct FHairCullingParams
 
 void ComputeHairStrandsClustersCulling(
 	FRHICommandListImmediate& RHICmdList,
-	TShaderMap<FGlobalShaderType>& ShaderMap,
+	FGlobalShaderMap& ShaderMap,
 	const TArray<FViewInfo>& Views,
 	const FHairCullingParams& CullingParameters,
 	FHairStrandClusterData& ClusterDatas);
 
 void ResetHairStrandsClusterToLOD0(
 	FRHICommandListImmediate& RHICmdList,
-	TShaderMap<FGlobalShaderType>& ShaderMap,
+	FGlobalShaderMap& ShaderMap,
 	FHairStrandClusterData& ClusterDatas);
\ No newline at end of file
diff --git a/Engine/Source/Runtime/Renderer/Private/HairStrands/HairStrandsComposition.cpp b/Engine/Source/Runtime/Renderer/Private/HairStrands/HairStrandsComposition.cpp
--- a/Engine/Source/Runtime/Renderer/Private/HairStrands/HairStrandsComposition.cpp
+++ b/Engine/Source/Runtime/Renderer/Private/HairStrands/HairStrandsComposition.cpp
@@ -49,7 +49,7 @@ static void AddHairVisibilityComposeSubPixelPass(
 	Parameters->CategorisationTexture = CategorisationTexture;
 	Parameters->RenderTargets[0] = FRenderTargetBinding(OutColorTexture, ERenderTargetLoadAction::ELoad);
 	Parameters->RenderTargets.DepthStencil = FDepthStencilBinding(OutDepthTexture, ERenderTargetLoadAction::ELoad, ERenderTargetLoadAction::ENoAction, FExclusiveDepthStencil::DepthWrite_StencilNop);
-
+	
 	FHairVisibilityComposeSubPixelPS::FPermutationDomain PermutationVector;
 	TShaderMapRef<FHairVisibilityComposeSubPixelPS> PixelShader(View.ShaderMap, PermutationVector);
 	TShaderMapRef<FPostProcessVS> VertexShader(View.ShaderMap);
@@ -58,7 +58,7 @@ static void AddHairVisibilityComposeSubPixelPass(
 	const FViewInfo* CapturedView = &View;
 
 	{
-		ClearUnusedGraphResources(*PixelShader, Parameters);
+		ClearUnusedGraphResources(PixelShader, Parameters);
 
 		GraphBuilder.AddPass(
 			RDG_EVENT_NAME("HairStrandsVisibilityComposeSubSPixel"),
@@ -74,14 +74,14 @@ static void AddHairVisibilityComposeSubPixelPass(
 			// Write stencil value for partially covered pixel, in order to run responsive AA on them
 			GraphicsPSOInit.DepthStencilState = TStaticDepthStencilState<true, CF_DepthNearOrEqual>::GetRHI();
 			GraphicsPSOInit.BoundShaderState.VertexDeclarationRHI = GFilterVertexDeclaration.VertexDeclarationRHI;
-			GraphicsPSOInit.BoundShaderState.VertexShaderRHI = GETSAFERHISHADER_VERTEX(*VertexShader);
-			GraphicsPSOInit.BoundShaderState.PixelShaderRHI = GETSAFERHISHADER_PIXEL(*PixelShader);
+			GraphicsPSOInit.BoundShaderState.VertexShaderRHI = VertexShader.GetVertexShader();
+			GraphicsPSOInit.BoundShaderState.PixelShaderRHI = PixelShader.GetPixelShader();
 			GraphicsPSOInit.PrimitiveType = PT_TriangleList;
 			SetGraphicsPipelineState(RHICmdList, GraphicsPSOInit);
 
 			VertexShader->SetParameters(RHICmdList, CapturedView->ViewUniformBuffer);
 			RHICmdList.SetViewport(Viewport.Min.X, Viewport.Min.Y, 0.0f, Viewport.Max.X, Viewport.Max.Y, 1.0f);
-			SetShaderParameters(RHICmdList, *PixelShader, PixelShader->GetPixelShader(), *Parameters);
+			SetShaderParameters(RHICmdList, PixelShader, PixelShader.GetPixelShader(), *Parameters);
 			DrawRectangle(
 				RHICmdList,
 				0, 0,
@@ -90,7 +90,7 @@ static void AddHairVisibilityComposeSubPixelPass(
 				Viewport.Width(), Viewport.Height(),
 				Viewport.Size(),
 				Resolution,
-				*VertexShader,
+				VertexShader,
 				EDRF_UseTriangleOptimization);
 		});
 	}
@@ -150,13 +150,13 @@ static void AddHairVisibilityFastResolvePass(
 
 	TShaderMapRef<FPostProcessVS> VertexShader(View.ShaderMap);
 	TShaderMapRef<FHairVisibilityFastResolvePS> PixelShader(View.ShaderMap);
-	const TShaderMap<FGlobalShaderType>* GlobalShaderMap = View.ShaderMap;
+	const FGlobalShaderMap* GlobalShaderMap = View.ShaderMap;
 	const FIntRect Viewport = View.ViewRect;
 	
 	const FViewInfo* CapturedView = &View;
 
 	{
-		ClearUnusedGraphResources(*PixelShader, Parameters);
+		ClearUnusedGraphResources(PixelShader, Parameters);
 
 		GraphBuilder.AddPass(
 			RDG_EVENT_NAME("HairStrandsVisibilityMarkTAAFastResolve"),
@@ -174,15 +174,15 @@ static void AddHairVisibilityFastResolvePass(
 				false, CF_Always, SO_Keep, SO_Keep, SO_Keep,
 				STENCIL_TEMPORAL_RESPONSIVE_AA_MASK, STENCIL_TEMPORAL_RESPONSIVE_AA_MASK>::GetRHI();
 			GraphicsPSOInit.BoundShaderState.VertexDeclarationRHI = GFilterVertexDeclaration.VertexDeclarationRHI;
-			GraphicsPSOInit.BoundShaderState.VertexShaderRHI = GETSAFERHISHADER_VERTEX(*VertexShader);
-			GraphicsPSOInit.BoundShaderState.PixelShaderRHI = GETSAFERHISHADER_PIXEL(*PixelShader);
+			GraphicsPSOInit.BoundShaderState.VertexShaderRHI = VertexShader.GetVertexShader();
+			GraphicsPSOInit.BoundShaderState.PixelShaderRHI = PixelShader.GetPixelShader();
 			GraphicsPSOInit.PrimitiveType = PT_TriangleList;
 			SetGraphicsPipelineState(RHICmdList, GraphicsPSOInit);
 
 			VertexShader->SetParameters(RHICmdList, CapturedView->ViewUniformBuffer);
 			RHICmdList.SetStencilRef(STENCIL_TEMPORAL_RESPONSIVE_AA_MASK);
 			RHICmdList.SetViewport(Viewport.Min.X, Viewport.Min.Y, 0.0f, Viewport.Max.X, Viewport.Max.Y, 1.0f);
-			SetShaderParameters(RHICmdList, *PixelShader, PixelShader->GetPixelShader(), *Parameters);
+			SetShaderParameters(RHICmdList, PixelShader, PixelShader.GetPixelShader(), *Parameters);
 			DrawRectangle(
 				RHICmdList,
 				0, 0,
@@ -191,7 +191,7 @@ static void AddHairVisibilityFastResolvePass(
 				Viewport.Width(), Viewport.Height(),
 				Viewport.Size(),
 				Resolution,
-				*VertexShader,
+				VertexShader,
 				EDRF_UseTriangleOptimization);
 		});
 	}
diff --git a/Engine/Source/Runtime/Renderer/Private/HairStrands/HairStrandsDebug.cpp b/Engine/Source/Runtime/Renderer/Private/HairStrands/HairStrandsDebug.cpp
--- a/Engine/Source/Runtime/Renderer/Private/HairStrands/HairStrandsDebug.cpp
+++ b/Engine/Source/Runtime/Renderer/Private/HairStrands/HairStrandsDebug.cpp
@@ -296,12 +296,12 @@ static void AddDebugHairPrintPass(
 	ShaderDrawDebug::SetParameters(GraphBuilder, View->ShaderDrawData, Parameters->ShaderDrawUniformBuffer);
 	TShaderMapRef<FHairDebugPrintCS> ComputeShader(View->ShaderMap);
 
-	ClearUnusedGraphResources(*ComputeShader, Parameters);
+	ClearUnusedGraphResources(ComputeShader, Parameters);
 
 	FComputeShaderUtils::AddPass(
 		GraphBuilder,
 		RDG_EVENT_NAME("HairStrandsDebugPrint"),
-		*ComputeShader,
+		ComputeShader,
 		Parameters,
 		FIntVector(1, 1, 1));
 }
@@ -402,7 +402,7 @@ static void AddDebugHairPass(
 
 	TShaderMapRef<FHairDebugPS> PixelShader(View->ShaderMap);
 
-	ClearUnusedGraphResources(*PixelShader, Parameters);
+	ClearUnusedGraphResources(PixelShader, Parameters);
 
 	GraphBuilder.AddPass(
 		RDG_EVENT_NAME("HairStrandsDebug"),
@@ -416,14 +416,14 @@ static void AddDebugHairPass(
 		GraphicsPSOInit.RasterizerState = TStaticRasterizerState<>::GetRHI();
 		GraphicsPSOInit.DepthStencilState = TStaticDepthStencilState<false, CF_Always>::GetRHI();
 		GraphicsPSOInit.BoundShaderState.VertexDeclarationRHI = GFilterVertexDeclaration.VertexDeclarationRHI;
-		GraphicsPSOInit.BoundShaderState.VertexShaderRHI = GETSAFERHISHADER_VERTEX(*VertexShader);
-		GraphicsPSOInit.BoundShaderState.PixelShaderRHI = GETSAFERHISHADER_PIXEL(*PixelShader);
+		GraphicsPSOInit.BoundShaderState.VertexShaderRHI = VertexShader.GetVertexShader();
+		GraphicsPSOInit.BoundShaderState.PixelShaderRHI = PixelShader.GetPixelShader();
 		GraphicsPSOInit.PrimitiveType = PT_TriangleList;
 		SetGraphicsPipelineState(RHICmdList, GraphicsPSOInit);
 
 		VertexShader->SetParameters(RHICmdList, View->ViewUniformBuffer);
 		RHICmdList.SetViewport(Viewport.Min.X, Viewport.Min.Y, 0.0f, Viewport.Max.X, Viewport.Max.Y, 1.0f);
-		SetShaderParameters(RHICmdList, *PixelShader, PixelShader->GetPixelShader(), *Parameters);
+		SetShaderParameters(RHICmdList, PixelShader, PixelShader.GetPixelShader(), *Parameters);
 
 		DrawRectangle(
 			RHICmdList,
@@ -433,7 +433,7 @@ static void AddDebugHairPass(
 			Viewport.Width(), Viewport.Height(),
 			Viewport.Size(),
 			Resolution,
-			*VertexShader,
+			VertexShader,
 			EDRF_UseTriangleOptimization);
 	});
 }
@@ -513,7 +513,7 @@ static void AddDebugDeepShadowTexturePass(
 	PermutationVector.Set<FDeepShadowVisualizePS::FOutputType>(ShadowData ? 0 : 1);
 	TShaderMapRef<FDeepShadowVisualizePS> PixelShader(View->ShaderMap, PermutationVector);
 
-	ClearUnusedGraphResources(*PixelShader, Parameters);
+	ClearUnusedGraphResources(PixelShader, Parameters);
 
 	GraphBuilder.AddPass(
 		ShadowData ? RDG_EVENT_NAME("DebugDeepShadowTexture") : RDG_EVENT_NAME("DebugHairViewRect"),
@@ -527,14 +527,14 @@ static void AddDebugDeepShadowTexturePass(
 		GraphicsPSOInit.RasterizerState = TStaticRasterizerState<>::GetRHI();
 		GraphicsPSOInit.DepthStencilState = TStaticDepthStencilState<false, CF_Always>::GetRHI();
 		GraphicsPSOInit.BoundShaderState.VertexDeclarationRHI = GFilterVertexDeclaration.VertexDeclarationRHI;
-		GraphicsPSOInit.BoundShaderState.VertexShaderRHI = GETSAFERHISHADER_VERTEX(*VertexShader);
-		GraphicsPSOInit.BoundShaderState.PixelShaderRHI = GETSAFERHISHADER_PIXEL(*PixelShader);
+		GraphicsPSOInit.BoundShaderState.VertexShaderRHI = VertexShader.GetVertexShader();
+		GraphicsPSOInit.BoundShaderState.PixelShaderRHI = PixelShader.GetPixelShader();
 		GraphicsPSOInit.PrimitiveType = PT_TriangleList;
 		SetGraphicsPipelineState(RHICmdList, GraphicsPSOInit);
 
 		VertexShader->SetParameters(RHICmdList, View->ViewUniformBuffer);
 		RHICmdList.SetViewport(Viewport.Min.X, Viewport.Min.Y, 0.0f, Viewport.Max.X, Viewport.Max.Y, 1.0f);
-		SetShaderParameters(RHICmdList, *PixelShader, PixelShader->GetPixelShader(), *Parameters);
+		SetShaderParameters(RHICmdList, PixelShader, PixelShader.GetPixelShader(), *Parameters);
 
 		DrawRectangle(
 			RHICmdList,
@@ -544,7 +544,7 @@ static void AddDebugDeepShadowTexturePass(
 			Viewport.Width(), Viewport.Height(),
 			Viewport.Size(),
 			Resolution,
-			*VertexShader,
+			VertexShader,
 			EDRF_UseTriangleOptimization);
 	});
 }
@@ -647,11 +647,11 @@ static void AddVoxelRaymarchingPass(
 		PermutationVector.Set<FVoxelRaymarchingPS::FDebugMode>(DebugPermutation);
 
 		TShaderMapRef<FVoxelRaymarchingPS> PixelShader(View.ShaderMap, PermutationVector);
-		const TShaderMap<FGlobalShaderType>* GlobalShaderMap = View.ShaderMap;
+		const FGlobalShaderMap* GlobalShaderMap = View.ShaderMap;
 		const FIntRect Viewport = View.ViewRect;
 		const FViewInfo* CapturedView = &View;
 
-		ClearUnusedGraphResources(*PixelShader, Parameters);
+		ClearUnusedGraphResources(PixelShader, Parameters);
 
 		GraphBuilder.AddPass(
 			RDG_EVENT_NAME("HairStrandsVoxelRaymarching"),
@@ -666,14 +666,14 @@ static void AddVoxelRaymarchingPass(
 			GraphicsPSOInit.DepthStencilState = TStaticDepthStencilState<false, CF_Always>::GetRHI();
 
 			GraphicsPSOInit.BoundShaderState.VertexDeclarationRHI = GFilterVertexDeclaration.VertexDeclarationRHI;
-			GraphicsPSOInit.BoundShaderState.VertexShaderRHI = GETSAFERHISHADER_VERTEX(*VertexShader);
-			GraphicsPSOInit.BoundShaderState.PixelShaderRHI = GETSAFERHISHADER_PIXEL(*PixelShader);
+			GraphicsPSOInit.BoundShaderState.VertexShaderRHI = VertexShader.GetVertexShader();
+			GraphicsPSOInit.BoundShaderState.PixelShaderRHI = PixelShader.GetPixelShader();
 			GraphicsPSOInit.PrimitiveType = PT_TriangleList;
 			SetGraphicsPipelineState(RHICmdList, GraphicsPSOInit);
 
 			VertexShader->SetParameters(RHICmdList, CapturedView->ViewUniformBuffer);
 			RHICmdList.SetViewport(Viewport.Min.X, Viewport.Min.Y, 0.0f, Viewport.Max.X, Viewport.Max.Y, 1.0f);
-			SetShaderParameters(RHICmdList, *PixelShader, PixelShader->GetPixelShader(), *Parameters);
+			SetShaderParameters(RHICmdList, PixelShader, PixelShader.GetPixelShader(), *Parameters);
 
 			DrawRectangle(
 				RHICmdList,
@@ -683,7 +683,7 @@ static void AddVoxelRaymarchingPass(
 				Viewport.Width(), Viewport.Height(),
 				Viewport.Size(),
 				Resolution,
-				*VertexShader,
+				VertexShader,
 				EDRF_UseTriangleOptimization);
 		});
 	}
@@ -751,7 +751,7 @@ static void AddVoxelPageRaymarchingPass(
 		TShaderMapRef<FVoxelVirtualRaymarchingCS> ComputeShader(View.ShaderMap);
 
 		const FIntVector DispatchCount = DispatchCount.DivideAndRoundUp(FIntVector(OutputTexture->Desc.Extent.X, OutputTexture->Desc.Extent.Y, 1), FIntVector(8, 8, 1));
-		FComputeShaderUtils::AddPass(GraphBuilder, RDG_EVENT_NAME("HairStrandsVoxelVirtualRaymarching"), *ComputeShader, Parameters, DispatchCount);
+		FComputeShaderUtils::AddPass(GraphBuilder, RDG_EVENT_NAME("HairStrandsVoxelVirtualRaymarching"), ComputeShader, Parameters, DispatchCount);
 	}
 }
 
@@ -803,11 +803,11 @@ static void AddPlotBSDFPass(
 	const FIntPoint OutputResolution = SceneTextures.SceneDepthBuffer->Desc.Extent;
 	TShaderMapRef<FPostProcessVS> VertexShader(View.ShaderMap);
 	TShaderMapRef<FHairStrandsBSDFPlotPS> PixelShader(View.ShaderMap);
-	const TShaderMap<FGlobalShaderType>* GlobalShaderMap = View.ShaderMap;
+	const FGlobalShaderMap* GlobalShaderMap = View.ShaderMap;
 	const FIntRect Viewport = View.ViewRect;
 	const FViewInfo* CapturedView = &View;
 
-	ClearUnusedGraphResources(*PixelShader, Parameters);
+	ClearUnusedGraphResources(PixelShader, Parameters);
 
 	GraphBuilder.AddPass(
 		RDG_EVENT_NAME("HairStrandsBsdfPlot"),
@@ -822,14 +822,14 @@ static void AddPlotBSDFPass(
 		GraphicsPSOInit.DepthStencilState = TStaticDepthStencilState<false, CF_Always>::GetRHI();
 
 		GraphicsPSOInit.BoundShaderState.VertexDeclarationRHI = GFilterVertexDeclaration.VertexDeclarationRHI;
-		GraphicsPSOInit.BoundShaderState.VertexShaderRHI = GETSAFERHISHADER_VERTEX(*VertexShader);
-		GraphicsPSOInit.BoundShaderState.PixelShaderRHI = GETSAFERHISHADER_PIXEL(*PixelShader);
+		GraphicsPSOInit.BoundShaderState.VertexShaderRHI = VertexShader.GetVertexShader();
+		GraphicsPSOInit.BoundShaderState.PixelShaderRHI = PixelShader.GetPixelShader();
 		GraphicsPSOInit.PrimitiveType = PT_TriangleList;
 		SetGraphicsPipelineState(RHICmdList, GraphicsPSOInit);
 
 		VertexShader->SetParameters(RHICmdList, CapturedView->ViewUniformBuffer);
 		RHICmdList.SetViewport(Viewport.Min.X, Viewport.Min.Y, 0.0f, Viewport.Max.X, Viewport.Max.Y, 1.0f);
-		SetShaderParameters(RHICmdList, *PixelShader, PixelShader->GetPixelShader(), *Parameters);
+		SetShaderParameters(RHICmdList, PixelShader, PixelShader.GetPixelShader(), *Parameters);
 
 		DrawRectangle(
 			RHICmdList,
@@ -839,7 +839,7 @@ static void AddPlotBSDFPass(
 			Viewport.Width(), Viewport.Height(),
 			Viewport.Size(),
 			Resolution,
-			*VertexShader,
+			VertexShader,
 			EDRF_UseTriangleOptimization);
 	});
 }
@@ -976,13 +976,13 @@ static void AddDebugProjectionMeshPass(
 		GraphicsPSOInit.RasterizerState = TStaticRasterizerState<FM_Wireframe>::GetRHI();
 		GraphicsPSOInit.DepthStencilState = TStaticDepthStencilState<true, CF_DepthNearOrEqual>::GetRHI();
 		GraphicsPSOInit.BoundShaderState.VertexDeclarationRHI = GEmptyVertexDeclaration.VertexDeclarationRHI;
-		GraphicsPSOInit.BoundShaderState.VertexShaderRHI = GETSAFERHISHADER_VERTEX(*VertexShader);
-		GraphicsPSOInit.BoundShaderState.PixelShaderRHI = GETSAFERHISHADER_PIXEL(*PixelShader);
+		GraphicsPSOInit.BoundShaderState.VertexShaderRHI = VertexShader.GetVertexShader();
+		GraphicsPSOInit.BoundShaderState.PixelShaderRHI = PixelShader.GetPixelShader();
 		GraphicsPSOInit.PrimitiveType = PrimitiveType;
 		SetGraphicsPipelineState(RHICmdList, GraphicsPSOInit);
 
-		SetShaderParameters(RHICmdList, *VertexShader, VertexShader->GetVertexShader(), VSParameters);
-		SetShaderParameters(RHICmdList, *PixelShader, PixelShader->GetPixelShader(), PSParameters);
+		SetShaderParameters(RHICmdList, VertexShader, VertexShader.GetVertexShader(), VSParameters);
+		SetShaderParameters(RHICmdList, PixelShader, PixelShader.GetPixelShader(), PSParameters);
 
 		// Emit an instanced quad draw call on the order of the number of pixels on the screen.	
 		RHICmdList.DrawPrimitive(0, PrimitiveCount, 1);
@@ -1148,13 +1148,13 @@ static void AddDebugProjectionHairPass(
 		GraphicsPSOInit.RasterizerState = TStaticRasterizerState<>::GetRHI();
 		GraphicsPSOInit.DepthStencilState = TStaticDepthStencilState<true, CF_DepthNearOrEqual>::GetRHI();
 		GraphicsPSOInit.BoundShaderState.VertexDeclarationRHI = GEmptyVertexDeclaration.VertexDeclarationRHI;
-		GraphicsPSOInit.BoundShaderState.VertexShaderRHI = GETSAFERHISHADER_VERTEX(*VertexShader);
-		GraphicsPSOInit.BoundShaderState.PixelShaderRHI = GETSAFERHISHADER_PIXEL(*PixelShader);
+		GraphicsPSOInit.BoundShaderState.VertexShaderRHI = VertexShader.GetVertexShader();
+		GraphicsPSOInit.BoundShaderState.PixelShaderRHI = PixelShader.GetPixelShader();
 		GraphicsPSOInit.PrimitiveType = PrimitiveType;
 		SetGraphicsPipelineState(RHICmdList, GraphicsPSOInit);
 
-		SetShaderParameters(RHICmdList, *VertexShader, VertexShader->GetVertexShader(), VSParameters);
-		SetShaderParameters(RHICmdList, *PixelShader, PixelShader->GetPixelShader(), PSParameters);
+		SetShaderParameters(RHICmdList, VertexShader, VertexShader.GetVertexShader(), VSParameters);
+		SetShaderParameters(RHICmdList, PixelShader, PixelShader.GetPixelShader(), PSParameters);
 
 		// Emit an instanced quad draw call on the order of the number of pixels on the screen.	
 		RHICmdList.DrawPrimitive(0, PrimitiveCount, 1);
@@ -1266,7 +1266,7 @@ static void AddDrawDebugClusterPass(
 		return;
 	}
 
-	TShaderMap<FGlobalShaderType>* ShaderMap = GetGlobalShaderMap(ERHIFeatureLevel::SM5);
+	FGlobalShaderMap* ShaderMap = GetGlobalShaderMap(ERHIFeatureLevel::SM5);
 
 	for (int DataIndex = 0; DataIndex < HairClusterData.HairGroups.Num(); ++DataIndex)
 	{
@@ -1309,7 +1309,7 @@ static void AddDrawDebugClusterPass(
 				const FIntVector DispatchCount = DispatchCount.DivideAndRoundUp(FIntVector(Parameters->ClusterCount, 1, 1), FIntVector(64, 1, 1));// FIX ME, this could get over 65535
 				FComputeShaderUtils::AddPass(
 					GraphBuilder, RDG_EVENT_NAME("DrawDebugClusterAABB"),
-					*ComputeShader, Parameters, DispatchCount);
+					ComputeShader, Parameters, DispatchCount);
 
 				GraphBuilder.Execute();
 			}
@@ -1738,7 +1738,7 @@ void RenderHairStrandsDebugInfo(
 			FHairVisibilityDebugPPLLCS::FPermutationDomain PermutationVector;
 			TShaderMapRef<FHairVisibilityDebugPPLLCS> ComputeShader(View.ShaderMap, PermutationVector);
 			FIntVector TextureSize = SceneColorTexture->Desc.GetSize(); TextureSize.Z = 1;
-			FComputeShaderUtils::AddPass(GraphBuilder, RDG_EVENT_NAME("HairPPLLDebug"), *ComputeShader, PassParameters,
+			FComputeShaderUtils::AddPass(GraphBuilder, RDG_EVENT_NAME("HairPPLLDebug"), ComputeShader, PassParameters,
 				FIntVector::DivideAndRoundUp(TextureSize, FIntVector(8, 8, 1)));
 			GraphBuilder.Execute();
 		}
diff --git a/Engine/Source/Runtime/Renderer/Private/HairStrands/HairStrandsEnvironment.cpp b/Engine/Source/Runtime/Renderer/Private/HairStrands/HairStrandsEnvironment.cpp
--- a/Engine/Source/Runtime/Renderer/Private/HairStrands/HairStrandsEnvironment.cpp
+++ b/Engine/Source/Runtime/Renderer/Private/HairStrands/HairStrandsEnvironment.cpp
@@ -154,7 +154,7 @@ static void AddHairStrandsEnvironmentAOPass(
 	PermutationVector.Set<FHairEnvironmentAO::FVirtualVoxel>(bUseVirtualVoxel ? VoxelTraversalType : 0);
 
 	TShaderMapRef<FHairEnvironmentAO> PixelShader(View.ShaderMap, PermutationVector);
-	ClearUnusedGraphResources(*PixelShader, PassParameters);
+	ClearUnusedGraphResources(PixelShader, PassParameters);
 
 	GraphBuilder.AddPass(
 		RDG_EVENT_NAME("HairStrandsAO %dx%d", View.ViewRect.Width(), View.ViewRect.Height()),
@@ -165,10 +165,10 @@ static void AddHairStrandsEnvironmentAOPass(
 		InRHICmdList.SetViewport(View.ViewRect.Min.X, View.ViewRect.Min.Y, 0.0f, View.ViewRect.Max.X, View.ViewRect.Max.Y, 1.0f);
 
 		FGraphicsPipelineStateInitializer GraphicsPSOInit;
-		FPixelShaderUtils::InitFullscreenPipelineState(InRHICmdList, View.ShaderMap, *PixelShader, GraphicsPSOInit);
+		FPixelShaderUtils::InitFullscreenPipelineState(InRHICmdList, View.ShaderMap, PixelShader, GraphicsPSOInit);
 		GraphicsPSOInit.BlendState = TStaticBlendState<CW_RGBA, BO_Min, BF_SourceColor, BF_DestColor, BO_Add, BF_Zero, BF_DestAlpha>::GetRHI();
 		SetGraphicsPipelineState(InRHICmdList, GraphicsPSOInit);
-		SetShaderParameters(InRHICmdList, *PixelShader, PixelShader->GetPixelShader(), *PassParameters);
+		SetShaderParameters(InRHICmdList, PixelShader, PixelShader.GetPixelShader(), *PassParameters);
 		FPixelShaderUtils::DrawFullscreenTriangle(InRHICmdList);
 	});
 }
@@ -211,12 +211,12 @@ static void AddHairEnvironmentLightingComposePass(
 
 	TShaderMapRef<FPostProcessVS> VertexShader(View.ShaderMap);
 	TShaderMapRef<FHairEnvironmentLightingComposePS> PixelShader(View.ShaderMap);
-	const TShaderMap<FGlobalShaderType>* GlobalShaderMap = View.ShaderMap;
+	const FGlobalShaderMap* GlobalShaderMap = View.ShaderMap;
 	const FIntRect Viewport = View.ViewRect;
 	const FIntPoint Resolution = OutColorTexture->Desc.Extent;
 	const FViewInfo* CapturedView = &View;
 
-	ClearUnusedGraphResources(*PixelShader, Parameters);
+	ClearUnusedGraphResources(PixelShader, Parameters);
 
 	GraphBuilder.AddPass(
 		RDG_EVENT_NAME("HairStrandsEnvironmentLightingCompose"),
@@ -232,14 +232,14 @@ static void AddHairEnvironmentLightingComposePass(
 		GraphicsPSOInit.RasterizerState = TStaticRasterizerState<>::GetRHI();
 		GraphicsPSOInit.DepthStencilState = TStaticDepthStencilState<false, CF_Always>::GetRHI();
 		GraphicsPSOInit.BoundShaderState.VertexDeclarationRHI = GFilterVertexDeclaration.VertexDeclarationRHI;
-		GraphicsPSOInit.BoundShaderState.VertexShaderRHI = GETSAFERHISHADER_VERTEX(*VertexShader);
-		GraphicsPSOInit.BoundShaderState.PixelShaderRHI = GETSAFERHISHADER_PIXEL(*PixelShader);
+		GraphicsPSOInit.BoundShaderState.VertexShaderRHI = VertexShader.GetVertexShader();
+		GraphicsPSOInit.BoundShaderState.PixelShaderRHI = PixelShader.GetPixelShader();
 		GraphicsPSOInit.PrimitiveType = PT_TriangleList;
 		SetGraphicsPipelineState(RHICmdList, GraphicsPSOInit);
 
 		VertexShader->SetParameters(RHICmdList, CapturedView->ViewUniformBuffer);
 		RHICmdList.SetViewport(Viewport.Min.X, Viewport.Min.Y, 0.0f, Viewport.Max.X, Viewport.Max.Y, 1.0f);
-		SetShaderParameters(RHICmdList, *PixelShader, PixelShader->GetPixelShader(), *Parameters);
+		SetShaderParameters(RHICmdList, PixelShader, PixelShader.GetPixelShader(), *Parameters);
 		DrawRectangle(
 			RHICmdList,
 			0, 0,
@@ -248,7 +248,7 @@ static void AddHairEnvironmentLightingComposePass(
 			Viewport.Width(), Viewport.Height(),
 			Viewport.Size(),
 			Resolution,
-			*VertexShader,
+			VertexShader,
 			EDRF_UseTriangleOptimization);
 	});
 }
@@ -411,7 +411,7 @@ static FRDGBufferRef AddHairStrandsEnvironmentLightingPassCS(
 	FComputeShaderUtils::AddPass(
 		GraphBuilder,
 		RDG_EVENT_NAME("HairStrandsEnvironmentCS %dx%d", View.ViewRect.Width(), View.ViewRect.Height()),
-		*ComputeShader,
+		ComputeShader,
 		PassParameters,
 		IndirectArgsBuffer,
 		0);
@@ -440,7 +440,7 @@ void RenderHairStrandsEnvironmentLighting(
 	{
 		return;
 	}
-	
+
 	FRDGBufferRef NodeIndirectArgBuffer = GraphBuilder.RegisterExternalBuffer(VisibilityData.NodeIndirectArg, TEXT("HairNodeIndirectArgBuffer"));
 
 	const FViewInfo& View = Views[ViewIndex];
diff --git a/Engine/Source/Runtime/Renderer/Private/HairStrands/HairStrandsInterface.cpp b/Engine/Source/Runtime/Renderer/Private/HairStrands/HairStrandsInterface.cpp
--- a/Engine/Source/Runtime/Renderer/Private/HairStrands/HairStrandsInterface.cpp
+++ b/Engine/Source/Runtime/Renderer/Private/HairStrands/HairStrandsInterface.cpp
@@ -244,7 +244,7 @@ void RunMeshTransfer(
 	TArray<FRWBuffer>& OutTransferedPositions)
 {
 	const ERHIFeatureLevel::Type FeatureLevel = GMaxRHIFeatureLevel;
-	TShaderMap<FGlobalShaderType>* ShaderMap = GetGlobalShaderMap(FeatureLevel);
+	FGlobalShaderMap* ShaderMap = GetGlobalShaderMap(FeatureLevel);
 
 	FRDGBuilder GraphBuilder(RHICmdList);
 
@@ -269,7 +269,7 @@ void RunProjection(
 	FHairStrandsProjectionHairData& SimProjectionHairData)
 {
 	const ERHIFeatureLevel::Type FeatureLevel = GMaxRHIFeatureLevel;
-	TShaderMap<FGlobalShaderType>* ShaderMap = GetGlobalShaderMap(FeatureLevel);
+	FGlobalShaderMap* ShaderMap = GetGlobalShaderMap(FeatureLevel);
 
 	auto Project = [&RHICmdList, ShaderMap, &TargetMeshData, LocalToWorld](FHairStrandsProjectionHairData& ProjectionHairData)
 	{
@@ -295,7 +295,7 @@ void RunHairStrandsInterpolation(
 	FRHICommandListImmediate& RHICmdList, 
 	EWorldType::Type WorldType, 
 	const FShaderDrawDebugData* ShaderDrawData,
-	TShaderMap<FGlobalShaderType>* ShaderMap, 
+	FGlobalShaderMap* ShaderMap, 
 	EHairStrandsInterpolationType Type,
 	FHairStrandClusterData* ClusterData)
 {
@@ -463,7 +463,7 @@ FHairStrandsDebugInfos GetHairStandsDebugInfos()
 
 		const uint32 GroupCount = Info.HairGroups.Num();
 		for (uint32 GroupIt=0; GroupIt < GroupCount; ++GroupIt)
-		{		
+		{
 			FHairStrandsDebugInfo::HairGroup& GroupInfo = Info.HairGroups[GroupIt];
 			if (GroupIt < uint32(E.RenProjectionHairDatas.HairGroups.Num()))
 			{
@@ -502,7 +502,6 @@ bool IsHairStrandsEnable(EShaderPlatform Platform)
 		GHairStrandsRenderingEnable == 1 && 
 		GHairManager.Elements.Num() > 0; 
 }
-
 bool IsHairRayTracingEnabled()
 {
 	return IsRayTracingEnabled();
@@ -530,7 +529,7 @@ bool HasHairStrandsProjectionQuery(EShaderPlatform Platform)
 		!GBindingQueries.IsEmpty();
 }
 
-void RunHairStrandsBindingQueries(FRHICommandListImmediate& RHICmdList, TShaderMap<FGlobalShaderType>* ShaderMap)
+void RunHairStrandsBindingQueries(FRHICommandListImmediate& RHICmdList, FGlobalShaderMap* ShaderMap)
 {
 	FBindingQuery Q;
 	while (GBindingQueries.Dequeue(Q))
diff --git a/Engine/Source/Runtime/Renderer/Private/HairStrands/HairStrandsLUT.cpp b/Engine/Source/Runtime/Renderer/Private/HairStrands/HairStrandsLUT.cpp
--- a/Engine/Source/Runtime/Renderer/Private/HairStrands/HairStrandsLUT.cpp
+++ b/Engine/Source/Runtime/Renderer/Private/HairStrands/HairStrandsLUT.cpp
@@ -76,7 +76,7 @@ static FRDGTextureRef AddHairLUTPass(
 	FComputeShaderUtils::AddPass(
 		GraphBuilder,
 		RDG_EVENT_NAME("HairStrandsLUT"),
-		*ComputeShader,
+		ComputeShader,
 		Parameters,
 		FComputeShaderUtils::GetGroupCount(OutputResolution, FIntVector(FComputeShaderUtils::kGolden2DGroupSize)));
 
diff --git a/Engine/Source/Runtime/Renderer/Private/HairStrands/HairStrandsMeshProjection.cpp b/Engine/Source/Runtime/Renderer/Private/HairStrands/HairStrandsMeshProjection.cpp
--- a/Engine/Source/Runtime/Renderer/Private/HairStrands/HairStrandsMeshProjection.cpp
+++ b/Engine/Source/Runtime/Renderer/Private/HairStrands/HairStrandsMeshProjection.cpp
@@ -41,7 +41,7 @@ IMPLEMENT_GLOBAL_SHADER(FMarkMeshSectionIdCS, "/Engine/Private/HairStrands/HairS
 
 static FRDGBufferRef AddMeshSectionId(
 	FRDGBuilder& GraphBuilder,
-	TShaderMap<FGlobalShaderType>* ShaderMap,
+	FGlobalShaderMap* ShaderMap,
 	const FHairStrandsProjectionMeshData::LOD& MeshData)
 {
 	const int32 SectionCount = MeshData.Sections.Num();
@@ -66,7 +66,7 @@ static FRDGBufferRef AddMeshSectionId(
 		FComputeShaderUtils::AddPass(
 			GraphBuilder,
 			RDG_EVENT_NAME("HairStrandsMarkVertexSectionId"),
-			*ComputeShader,
+			ComputeShader,
 			Parameters,
 			DispatchGroupCount);
 	}
@@ -115,7 +115,7 @@ IMPLEMENT_GLOBAL_SHADER(FMeshTransferCS, "/Engine/Private/HairStrands/HairStrand
 
 static void AddMeshTransferPass(
 	FRDGBuilder& GraphBuilder,
-	TShaderMap<FGlobalShaderType>* ShaderMap,
+	FGlobalShaderMap* ShaderMap,
 	bool bClear,
 	const FHairStrandsProjectionMeshData::Section& SourceSectionData,
 	const FHairStrandsProjectionMeshData::Section& TargetSectionData,
@@ -183,7 +183,7 @@ static void AddMeshTransferPass(
 		FComputeShaderUtils::AddPass(
 			GraphBuilder,
 			RDG_EVENT_NAME("HairStrandsTransferMesh"),
-			*ComputeShader,
+			ComputeShader,
 			Parameters,
 			DispatchGroupCount);
 		bClear = false;
@@ -232,7 +232,7 @@ IMPLEMENT_GLOBAL_SHADER(FHairMeshProjectionCS, "/Engine/Private/HairStrands/Hair
 
 static void AddHairStrandMeshProjectionPass(
 	FRDGBuilder& GraphBuilder,
-	TShaderMap<FGlobalShaderType>* ShaderMap,
+	FGlobalShaderMap* ShaderMap,
 	const bool bClear,
 	const int32 LODIndex,
 	const FHairStrandsProjectionMeshData::Section& MeshSectionData,
@@ -300,7 +300,7 @@ static void AddHairStrandMeshProjectionPass(
 		FComputeShaderUtils::AddPass(
 			GraphBuilder,
 			RDG_EVENT_NAME("HairStrandsMeshProjection"),
-			*ComputeShader,
+			ComputeShader,
 			Parameters,
 			DispatchGroupCount);
 	}
@@ -308,7 +308,7 @@ static void AddHairStrandMeshProjectionPass(
 
 void ProjectHairStrandsOntoMesh(
 	FRDGBuilder& GraphBuilder,
-	TShaderMap<FGlobalShaderType>* ShaderMap,
+	FGlobalShaderMap* ShaderMap,
 	const int32 LODIndex,
 	const FHairStrandsProjectionMeshData& ProjectionMeshData,
 	FHairStrandsProjectionHairData::HairGroup& ProjectionHairData)
@@ -330,7 +330,7 @@ void ProjectHairStrandsOntoMesh(
 
 void TransferMesh(
 	FRDGBuilder& GraphBuilder,
-	TShaderMap<FGlobalShaderType>* ShaderMap,
+	FGlobalShaderMap* ShaderMap,
 	const int32 LODIndex,
 	const FHairStrandsProjectionMeshData& SourceMeshData,
 	const FHairStrandsProjectionMeshData& TargetMeshData,
@@ -368,7 +368,7 @@ class FHairUpdateMeshTriangleCS : public FGlobalShader
 
 	BEGIN_SHADER_PARAMETER_STRUCT(FParameters, )
 		SHADER_PARAMETER(uint32, MaxRootCount)
-		
+
 		SHADER_PARAMETER(uint32, MeshSectionIndex)
 		SHADER_PARAMETER(uint32, MeshMaxIndexCount)
 		SHADER_PARAMETER(uint32, MeshMaxVertexCount)
@@ -393,7 +393,7 @@ IMPLEMENT_GLOBAL_SHADER(FHairUpdateMeshTriangleCS, "/Engine/Private/HairStrands/
 
 static void AddHairStrandUpdateMeshTrianglesPass(
 	FRDGBuilder& GraphBuilder,
-	TShaderMap<FGlobalShaderType>* ShaderMap,
+	FGlobalShaderMap* ShaderMap,
 	const int32 LODIndex,
 	const HairStrandsTriangleType Type,
 	const FHairStrandsProjectionMeshData::Section& MeshSectionData,
@@ -447,14 +447,14 @@ static void AddHairStrandUpdateMeshTrianglesPass(
 	FComputeShaderUtils::AddPass(
 		GraphBuilder,
 		RDG_EVENT_NAME("HairStrandsTriangleMeshUpdate"),
-		*ComputeShader,
+		ComputeShader,
 		Parameters,
 		DispatchGroupCount);
 }
 
 void UpdateHairStrandsMeshTriangles(
 	FRDGBuilder& GraphBuilder,
-	TShaderMap<FGlobalShaderType>* ShaderMap,
+	FGlobalShaderMap* ShaderMap,
 	const int32 LODIndex,
 	const HairStrandsTriangleType Type,
 	const FHairStrandsProjectionMeshData::LOD& ProjectionMeshData,
diff --git a/Engine/Source/Runtime/Renderer/Private/HairStrands/HairStrandsMeshProjection.h b/Engine/Source/Runtime/Renderer/Private/HairStrands/HairStrandsMeshProjection.h
--- a/Engine/Source/Runtime/Renderer/Private/HairStrands/HairStrandsMeshProjection.h
+++ b/Engine/Source/Runtime/Renderer/Private/HairStrands/HairStrandsMeshProjection.h
@@ -10,14 +10,14 @@
 /* Project hair strands onto a LOD mesh */
 void ProjectHairStrandsOntoMesh(
 	class FRDGBuilder& GraphBuilder,
-	TShaderMap<FGlobalShaderType>* ShaderMap,
+	FGlobalShaderMap* ShaderMap,
 	const int32 LODIndex,
 	const FHairStrandsProjectionMeshData& ProjectionMeshData,
 	FHairStrandsProjectionHairData::HairGroup& ProjectionHairData);
 
 void TransferMesh(
 	FRDGBuilder& GraphBuilder,
-	TShaderMap<FGlobalShaderType>* ShaderMap,
+	FGlobalShaderMap* ShaderMap,
 	const int32 LODIndex,
 	const FHairStrandsProjectionMeshData& SourceMeshData,
 	const FHairStrandsProjectionMeshData& TargetMeshData,
@@ -32,7 +32,7 @@ enum class HairStrandsTriangleType
 /* Update the triangles information on which hair stands have been projected */
 void UpdateHairStrandsMeshTriangles(
 	class FRDGBuilder& GraphBuilder,
-	TShaderMap<FGlobalShaderType>* ShaderMap,
+	FGlobalShaderMap* ShaderMap,
 	const int32 LODIndex,
 	const HairStrandsTriangleType Type,
 	const FHairStrandsProjectionMeshData::LOD& ProjectionMeshData,
diff --git a/Engine/Source/Runtime/Renderer/Private/HairStrands/HairStrandsRasterCommon.cpp b/Engine/Source/Runtime/Renderer/Private/HairStrands/HairStrandsRasterCommon.cpp
--- a/Engine/Source/Runtime/Renderer/Private/HairStrands/HairStrandsRasterCommon.cpp
+++ b/Engine/Source/Runtime/Renderer/Private/HairStrands/HairStrandsRasterCommon.cpp
@@ -82,7 +82,7 @@ protected:
 
 	static bool ShouldCompilePermutation(const FMeshMaterialShaderPermutationParameters& Parameters)
 	{
-		return IsCompatibleWithHairStrands(Parameters.Material, Parameters.Platform);
+		return IsCompatibleWithHairStrands(Parameters.Platform, Parameters.MaterialParameters);
 	}
 
 	static void ModifyCompilationEnvironment(const FMaterialShaderPermutationParameters& Parameters, FShaderCompilerEnvironment& OutEnvironment)
@@ -116,7 +116,7 @@ protected:
 
 	static bool ShouldCompilePermutation(const FMeshMaterialShaderPermutationParameters& Parameters)
 	{
-		return IsCompatibleWithHairStrands(Parameters.Material, Parameters.Platform);
+		return IsCompatibleWithHairStrands(Parameters.Platform, Parameters.MaterialParameters);
 	}
 
 	static void ModifyCompilationEnvironment(const FMaterialShaderPermutationParameters& Parameters, FShaderCompilerEnvironment& OutEnvironment)
@@ -151,7 +151,7 @@ protected:
 
 	static bool ShouldCompilePermutation(const FMeshMaterialShaderPermutationParameters& Parameters)
 	{
-		return IsCompatibleWithHairStrands(Parameters.Material, Parameters.Platform);
+		return IsCompatibleWithHairStrands(Parameters.Platform, Parameters.MaterialParameters);
 	}
 
 	static void ModifyCompilationEnvironment(const FMaterialShaderPermutationParameters& Parameters, FShaderCompilerEnvironment& OutEnvironment)
@@ -186,7 +186,7 @@ public:
 
 	static bool ShouldCompilePermutation(const FMeshMaterialShaderPermutationParameters& Parameters)
 	{
-		return IsCompatibleWithHairStrands(Parameters.Material, Parameters.Platform);
+		return IsCompatibleWithHairStrands(Parameters.Platform, Parameters.MaterialParameters);
 	}
 };
 IMPLEMENT_MATERIAL_SHADER_TYPE(, FDeepShadowDepthMeshPS, TEXT("/Engine/Private/HairStrands/HairStrandsDeepShadowPS.usf"), TEXT("MainDepth"), SF_Pixel);
@@ -211,13 +211,7 @@ public:
 
 	static bool ShouldCompilePermutation(const FMeshMaterialShaderPermutationParameters& Parameters)
 	{
-		return IsCompatibleWithHairStrands(Parameters.Material, Parameters.Platform);
-	}
-
-	virtual bool Serialize(FArchive& Ar) override
-	{
-		bool bShaderHasOutdatedParameters = FMeshMaterialShader::Serialize(Ar);
-		return bShaderHasOutdatedParameters;
+		return IsCompatibleWithHairStrands(Parameters.Platform, Parameters.MaterialParameters);
 	}
 };
 IMPLEMENT_MATERIAL_SHADER_TYPE(, FDeepShadowDomMeshPS, TEXT("/Engine/Private/HairStrands/HairStrandsDeepShadowPS.usf"), TEXT("MainDom"), SF_Pixel);
@@ -257,7 +251,7 @@ public:
 
 	static bool ShouldCompilePermutation(const FMeshMaterialShaderPermutationParameters& Parameters)
 	{
-		return IsCompatibleWithHairStrands(Parameters.Material, Parameters.Platform);
+		return IsCompatibleWithHairStrands(Parameters.Platform, Parameters.MaterialParameters);
 	}
 
 	static void ModifyCompilationEnvironment(const FMaterialShaderPermutationParameters& Parameters, FShaderCompilerEnvironment& OutEnvironment)
@@ -268,25 +262,12 @@ public:
 		OutEnvironment.SetDefine(TEXT("SUPPORT_SPARSE"), VoxelizationType == EVoxelMeshPSType::DensityVirtual ? 1 : 0);
 	}
 
-	virtual bool Serialize(FArchive& Ar) override
-	{
-		bool bShaderHasOutdatedParameters = FMeshMaterialShader::Serialize(Ar);
-		Ar << DensityTexture;
-		if (VoxelizationType == EVoxelMeshPSType::Material)
-		{
-			Ar << TangentXTexture;
-			Ar << TangentYTexture;
-			Ar << TangentZTexture;
-			Ar << MaterialTexture;
-		}
-		return bShaderHasOutdatedParameters;
-	}
 private:
-	FRWShaderParameter DensityTexture;
-	FRWShaderParameter TangentXTexture;
-	FRWShaderParameter TangentYTexture;
-	FRWShaderParameter TangentZTexture;
-	FRWShaderParameter MaterialTexture;
+	LAYOUT_FIELD(FRWShaderParameter, DensityTexture);
+	LAYOUT_FIELD(FRWShaderParameter, TangentXTexture);
+	LAYOUT_FIELD(FRWShaderParameter, TangentYTexture);
+	LAYOUT_FIELD(FRWShaderParameter, TangentZTexture);
+	LAYOUT_FIELD(FRWShaderParameter, MaterialTexture);
 };
 IMPLEMENT_MATERIAL_SHADER_TYPE(template<>, FVoxelMeshPS<EVoxelMeshPSType::Density>, TEXT("/Engine/Private/HairStrands/HairStrandsDeepShadowPS.usf"), TEXT("MainVoxel"), SF_Pixel);
 IMPLEMENT_MATERIAL_SHADER_TYPE(template<>, FVoxelMeshPS<EVoxelMeshPSType::Material>, TEXT("/Engine/Private/HairStrands/HairStrandsDeepShadowPS.usf"), TEXT("MainVoxel"), SF_Pixel);
@@ -335,8 +316,9 @@ void FHairRasterMeshProcessor::AddMeshBatch(const FMeshBatch& RESTRICT MeshBatch
 		&& ShouldIncludeDomainInMeshPass(Material.GetMaterialDomain()))
 	{
 		const FMaterialRenderProxy& MaterialRenderProxy = FallbackMaterialRenderProxyPtr ? *FallbackMaterialRenderProxyPtr : *MeshBatch.MaterialRenderProxy;
-		const ERasterizerFillMode MeshFillMode = ComputeMeshFillMode(MeshBatch, Material);
-		const ERasterizerCullMode MeshCullMode = RasterPassType == EHairStrandsRasterPassType::FrontDepth ? ComputeMeshCullMode(MeshBatch, Material) : CM_None;
+		const FMeshDrawingPolicyOverrideSettings OverrideSettings = ComputeMeshOverrideSettings(MeshBatch);
+		const ERasterizerFillMode MeshFillMode = ComputeMeshFillMode(MeshBatch, Material, OverrideSettings);
+		const ERasterizerCullMode MeshCullMode = RasterPassType == EHairStrandsRasterPassType::FrontDepth ? ComputeMeshCullMode(MeshBatch, Material, OverrideSettings) : CM_None;
 
 		if (RasterPassType == EHairStrandsRasterPassType::FrontDepth)
 			Process<FDeepShadowDepthMeshVS, FDeepShadowDepthMeshPS>(MeshBatch, BatchElementMask, PrimitiveSceneProxy, StaticMeshId, MaterialRenderProxy, Material, MeshFillMode, MeshCullMode);
@@ -375,12 +357,12 @@ void FHairRasterMeshProcessor::Process(
 	{
 		const EMaterialTessellationMode MaterialTessellationMode = MaterialResource.GetTessellationMode();
 		FVertexFactoryType* VertexFactoryType = VertexFactory->GetType();
-		const bool bIsHairStrandsFactory = MeshBatch.VertexFactory->GetType()->GetId() == CompatibleVF->GetId();
+		const bool bIsHairStrandsFactory = MeshBatch.VertexFactory->GetType()->GetHashedName() == CompatibleVF->GetHashedName();
 		if (!bIsHairStrandsFactory)
 			return;
 
-		PassShaders.DomainShader = nullptr;
-		PassShaders.HullShader = nullptr;
+		PassShaders.DomainShader.Reset();
+		PassShaders.HullShader.Reset();
 		PassShaders.VertexShader = MaterialResource.GetShader<VertexShaderType>(VertexFactoryType);
 		PassShaders.PixelShader = MaterialResource.GetShader<PixelShaderType>(VertexFactoryType);
 	}
@@ -504,7 +486,8 @@ void AddHairStrandsRasterPass(
 		FDynamicMeshDrawCommandStorage DynamicMeshDrawCommandStorage; // << Were would thid be stored?
 		FMeshCommandOneFrameArray VisibleMeshDrawCommands;
 		FGraphicsMinimalPipelineStateSet GraphicsMinimalPipelineStateSet;
-		FDynamicPassMeshDrawListContext ShadowContext(DynamicMeshDrawCommandStorage, VisibleMeshDrawCommands, GraphicsMinimalPipelineStateSet);
+		bool bNeedsInitialization;
+		FDynamicPassMeshDrawListContext ShadowContext(DynamicMeshDrawCommandStorage, VisibleMeshDrawCommands, GraphicsMinimalPipelineStateSet, bNeedsInitialization);
 
 		FHairRasterMeshProcessor HairRasterMeshProcessor(Scene, ViewInfo /* is a SceneView */, DrawRenderState, &ShadowContext, RasterPassType);
 
@@ -572,4 +555,4 @@ void AddHairVoxelizationRasterPass(
 		HairRenderInfo, 
 		RasterDirection,
 		PassParameters);
-}
\ No newline at end of file
+}
diff --git a/Engine/Source/Runtime/Renderer/Private/HairStrands/HairStrandsRendering.h b/Engine/Source/Runtime/Renderer/Private/HairStrands/HairStrandsRendering.h
--- a/Engine/Source/Runtime/Renderer/Private/HairStrands/HairStrandsRendering.h
+++ b/Engine/Source/Runtime/Renderer/Private/HairStrands/HairStrandsRendering.h
@@ -40,7 +40,7 @@ void RunHairStrandsInterpolation(
 	FRHICommandListImmediate& RHICmdList, 
 	EWorldType::Type WorldType, 
 	const struct FShaderDrawDebugData* DebugShaderData,
-	TShaderMap<FGlobalShaderType>* ShaderMap, 
+	FGlobalShaderMap* ShaderMap, 
 	EHairStrandsInterpolationType Type,
 	FHairStrandClusterData* ClusterData);
 
