Copying Private-LoadTimes-4.24 stream to Main. Biggest changes are in Materials/Shader memory freezing.

diff --git a/Engine/Source/Runtime/Renderer/Private/HairStrands/HairStrandsScatter.cpp b/Engine/Source/Runtime/Renderer/Private/HairStrands/HairStrandsScatter.cpp
--- a/Engine/Source/Runtime/Renderer/Private/HairStrands/HairStrandsScatter.cpp
+++ b/Engine/Source/Runtime/Renderer/Private/HairStrands/HairStrandsScatter.cpp
@@ -75,13 +75,13 @@ static FRDGTextureRef AddPreScatterComposePass(
 
 	TShaderMapRef<FPostProcessVS> VertexShader(View.ShaderMap);
 	TShaderMapRef<FHairComposePS> PixelShader(View.ShaderMap);
-	const TShaderMap<FGlobalShaderType>* GlobalShaderMap = View.ShaderMap;
+	const FGlobalShaderMap* GlobalShaderMap = View.ShaderMap;
 	const FIntRect Viewport = View.ViewRect;
 
 	const FViewInfo* CapturedView = &View;
 
 	{
-		ClearUnusedGraphResources(*PixelShader, Parameters);
+		ClearUnusedGraphResources(PixelShader, Parameters);
 
 		GraphBuilder.AddPass(
 			RDG_EVENT_NAME("HairCompose"),
@@ -95,14 +95,14 @@ static FRDGTextureRef AddPreScatterComposePass(
 			GraphicsPSOInit.RasterizerState = TStaticRasterizerState<>::GetRHI();
 			GraphicsPSOInit.DepthStencilState = TStaticDepthStencilState<false, CF_Always>::GetRHI();
 			GraphicsPSOInit.BoundShaderState.VertexDeclarationRHI = GFilterVertexDeclaration.VertexDeclarationRHI;
-			GraphicsPSOInit.BoundShaderState.VertexShaderRHI = GETSAFERHISHADER_VERTEX(*VertexShader);
-			GraphicsPSOInit.BoundShaderState.PixelShaderRHI = GETSAFERHISHADER_PIXEL(*PixelShader);
+			GraphicsPSOInit.BoundShaderState.VertexShaderRHI = VertexShader.GetVertexShader();
+			GraphicsPSOInit.BoundShaderState.PixelShaderRHI = PixelShader.GetPixelShader();
 			GraphicsPSOInit.PrimitiveType = PT_TriangleList;
 			SetGraphicsPipelineState(RHICmdList, GraphicsPSOInit);
 
 			VertexShader->SetParameters(RHICmdList, CapturedView->ViewUniformBuffer);
 			RHICmdList.SetViewport(Viewport.Min.X, Viewport.Min.Y, 0.0f, Viewport.Max.X, Viewport.Max.Y, 1.0f);
-			SetShaderParameters(RHICmdList, *PixelShader, PixelShader->GetPixelShader(), *Parameters);
+			SetShaderParameters(RHICmdList, PixelShader, PixelShader.GetPixelShader(), *Parameters);
 			DrawRectangle(
 				RHICmdList,
 				0, 0,
@@ -111,7 +111,7 @@ static FRDGTextureRef AddPreScatterComposePass(
 				Viewport.Width(), Viewport.Height(),
 				Viewport.Size(),
 				Resolution,
-				*VertexShader,
+				VertexShader,
 				EDRF_UseTriangleOptimization);
 		});
 	}
@@ -207,13 +207,13 @@ static FRDGTextureRef AddScatterPass(
 
 	TShaderMapRef<FPostProcessVS> VertexShader(View.ShaderMap);
 	TShaderMapRef<FHairScatterPS> PixelShader(View.ShaderMap);
-	const TShaderMap<FGlobalShaderType>* GlobalShaderMap = View.ShaderMap;
+	const FGlobalShaderMap* GlobalShaderMap = View.ShaderMap;
 	const FIntRect Viewport = View.ViewRect;
 
 	const FViewInfo* CapturedView = &View;
 
 	{
-		ClearUnusedGraphResources(*PixelShader, Parameters);
+		ClearUnusedGraphResources(PixelShader, Parameters);
 
 		GraphBuilder.AddPass(
 			RDG_EVENT_NAME("HairScatter"),
@@ -238,14 +238,14 @@ static FRDGTextureRef AddScatterPass(
 			GraphicsPSOInit.RasterizerState = TStaticRasterizerState<>::GetRHI();
 			GraphicsPSOInit.DepthStencilState = TStaticDepthStencilState<false, CF_Always>::GetRHI();
 			GraphicsPSOInit.BoundShaderState.VertexDeclarationRHI = GFilterVertexDeclaration.VertexDeclarationRHI;
-			GraphicsPSOInit.BoundShaderState.VertexShaderRHI = GETSAFERHISHADER_VERTEX(*VertexShader);
-			GraphicsPSOInit.BoundShaderState.PixelShaderRHI = GETSAFERHISHADER_PIXEL(*PixelShader);
+			GraphicsPSOInit.BoundShaderState.VertexShaderRHI = VertexShader.GetVertexShader();
+			GraphicsPSOInit.BoundShaderState.PixelShaderRHI = PixelShader.GetPixelShader();
 			GraphicsPSOInit.PrimitiveType = PT_TriangleList;
 			SetGraphicsPipelineState(RHICmdList, GraphicsPSOInit);
 
 			VertexShader->SetParameters(RHICmdList, CapturedView->ViewUniformBuffer);
 			RHICmdList.SetViewport(Viewport.Min.X, Viewport.Min.Y, 0.0f, Viewport.Max.X, Viewport.Max.Y, 1.0f);
-			SetShaderParameters(RHICmdList, *PixelShader, PixelShader->GetPixelShader(), *Parameters);
+			SetShaderParameters(RHICmdList, PixelShader, PixelShader.GetPixelShader(), *Parameters);
 			DrawRectangle(
 				RHICmdList,
 				0, 0,
@@ -254,7 +254,7 @@ static FRDGTextureRef AddScatterPass(
 				Viewport.Width(), Viewport.Height(),
 				Viewport.Size(),
 				Resolution,
-				*VertexShader,
+				VertexShader,
 				EDRF_UseTriangleOptimization);
 		});
 	}
diff --git a/Engine/Source/Runtime/Renderer/Private/HairStrands/HairStrandsTransmittance.cpp b/Engine/Source/Runtime/Renderer/Private/HairStrands/HairStrandsTransmittance.cpp
--- a/Engine/Source/Runtime/Renderer/Private/HairStrands/HairStrandsTransmittance.cpp
+++ b/Engine/Source/Runtime/Renderer/Private/HairStrands/HairStrandsTransmittance.cpp
@@ -279,7 +279,7 @@ static FRDGBufferRef AddDeepShadowTransmittanceMaskPass(
 	FComputeShaderUtils::AddPass(
 		GraphBuilder,
 		RDG_EVENT_NAME("HairStrandsTransmittanceMask"),
-		*ComputeShader,
+		ComputeShader,
 		Parameters,
 		IndirectArgsBuffer,
 		0);
@@ -430,12 +430,12 @@ static void AddDeepShadowOpaqueMaskPass(
 	const FIntPoint OutputResolution = SceneTextures.SceneDepthBuffer->Desc.Extent;
 	TShaderMapRef<FPostProcessVS> VertexShader(View.ShaderMap);
 	TShaderMapRef<FDeepShadowMaskPS> PixelShader(View.ShaderMap, PermutationVector);
-	const TShaderMap<FGlobalShaderType>* GlobalShaderMap = View.ShaderMap;
+	const FGlobalShaderMap* GlobalShaderMap = View.ShaderMap;
 	const FIntRect Viewport = View.ViewRect;
 	const FViewInfo* CapturedView = &View;
 
 	{
-		ClearUnusedGraphResources(*PixelShader, Parameters);
+		ClearUnusedGraphResources(PixelShader, Parameters);
 		FIntPoint Resolution = OutShadowMask->Desc.Extent;
 
 		GraphBuilder.AddPass(
@@ -451,14 +451,14 @@ static void AddDeepShadowOpaqueMaskPass(
 			GraphicsPSOInit.DepthStencilState = TStaticDepthStencilState<false, CF_Always>::GetRHI();
 
 			GraphicsPSOInit.BoundShaderState.VertexDeclarationRHI = GFilterVertexDeclaration.VertexDeclarationRHI;
-			GraphicsPSOInit.BoundShaderState.VertexShaderRHI = GETSAFERHISHADER_VERTEX(*VertexShader);
-			GraphicsPSOInit.BoundShaderState.PixelShaderRHI = GETSAFERHISHADER_PIXEL(*PixelShader);
+			GraphicsPSOInit.BoundShaderState.VertexShaderRHI = VertexShader.GetVertexShader();
+			GraphicsPSOInit.BoundShaderState.PixelShaderRHI = PixelShader.GetPixelShader();
 			GraphicsPSOInit.PrimitiveType = PT_TriangleList;
 			SetGraphicsPipelineState(RHICmdList, GraphicsPSOInit);
 
 			VertexShader->SetParameters(RHICmdList, CapturedView->ViewUniformBuffer);
 			RHICmdList.SetViewport(Viewport.Min.X, Viewport.Min.Y, 0.0f, Viewport.Max.X, Viewport.Max.Y, 1.0f);
-			SetShaderParameters(RHICmdList, *PixelShader, PixelShader->GetPixelShader(), *Parameters);
+			SetShaderParameters(RHICmdList, PixelShader, PixelShader.GetPixelShader(), *Parameters);
 
 			DrawRectangle(
 				RHICmdList,
@@ -468,7 +468,7 @@ static void AddDeepShadowOpaqueMaskPass(
 				Viewport.Width(), Viewport.Height(),
 				Viewport.Size(),
 				Resolution,
-				*VertexShader,
+				VertexShader,
 				EDRF_UseTriangleOptimization);
 		});
 	}
diff --git a/Engine/Source/Runtime/Renderer/Private/HairStrands/HairStrandsVisibility.cpp b/Engine/Source/Runtime/Renderer/Private/HairStrands/HairStrandsVisibility.cpp
--- a/Engine/Source/Runtime/Renderer/Private/HairStrands/HairStrandsVisibility.cpp
+++ b/Engine/Source/Runtime/Renderer/Private/HairStrands/HairStrandsVisibility.cpp
@@ -153,7 +153,7 @@ void SetUpViewHairRenderInfo(const FViewInfo& ViewInfo, bool bEnableMSAA, FVecto
 
 static bool IsCompatibleWithHairVisibility(const FMeshMaterialShaderPermutationParameters& Parameters)
 {
-	return IsCompatibleWithHairStrands(Parameters.Material, Parameters.Platform);
+	return IsCompatibleWithHairStrands(Parameters.Platform, Parameters.MaterialParameters);
 }
 
 /////////////////////////////////////////////////////////////////////////////////////////
@@ -235,7 +235,7 @@ public:
 
 	static bool ShouldCompilePermutation(const FMeshMaterialShaderPermutationParameters& Parameters)
 	{
-		const bool bIsCompatible = IsCompatibleWithHairStrands(Parameters.Material, Parameters.Platform) && Parameters.VertexFactoryType->GetFName() == FName(TEXT("FHairStrandsVertexFactory"));
+		const bool bIsCompatible = IsCompatibleWithHairStrands(Parameters.Platform, Parameters.MaterialParameters) && Parameters.VertexFactoryType->GetFName() == FName(TEXT("FHairStrandsVertexFactory"));
 		return bIsCompatible;
 	}
 
@@ -244,18 +244,6 @@ public:
 		FMaterialShader::ModifyCompilationEnvironment(Parameters, OutEnvironment);
 	}
 
-	virtual bool Serialize(FArchive& Ar) override
-	{
-		bool bShaderHasOutdatedParameters = FMeshMaterialShader::Serialize(Ar);
-		Ar << MaterialPass_MacroGroupId;
-		Ar << MaterialPass_MaterialId;
-		Ar << MaterialPass_PrimitiveId;
-		Ar << MaterialPass_LightChannelMask;
-		Ar << OutNodeData;
-		Ar << OutNodeVelocity;
-		return bShaderHasOutdatedParameters;
-	}
-
 	void GetShaderBindings(
 		const FScene* Scene,
 		ERHIFeatureLevel::Type FeatureLevel,
@@ -274,12 +262,12 @@ public:
 	}
 
 private:
-	FShaderParameter MaterialPass_MacroGroupId;
-	FShaderParameter MaterialPass_MaterialId;
-	FShaderParameter MaterialPass_PrimitiveId;
-	FShaderParameter MaterialPass_LightChannelMask;
-	FShaderParameter OutNodeData;
-	FShaderParameter OutNodeVelocity;
+	LAYOUT_FIELD(FShaderParameter, MaterialPass_MacroGroupId);
+	LAYOUT_FIELD(FShaderParameter, MaterialPass_MaterialId);
+	LAYOUT_FIELD(FShaderParameter, MaterialPass_PrimitiveId);
+	LAYOUT_FIELD(FShaderParameter, MaterialPass_LightChannelMask);
+	LAYOUT_FIELD(FShaderParameter, OutNodeData);
+	LAYOUT_FIELD(FShaderParameter, OutNodeVelocity);
 };
 IMPLEMENT_MATERIAL_SHADER_TYPE(, FHairMaterialPS, TEXT("/Engine/Private/HairStrands/HairStrandsMaterialPS.usf"), TEXT("Main"), SF_Pixel);
 
@@ -326,7 +314,7 @@ void FHairMaterialProcessor::AddMeshBatch(const FMeshBatch& RESTRICT MeshBatch,
 	const FMaterialRenderProxy* FallbackMaterialRenderProxyPtr = nullptr;
 	const FMaterial& Material = MeshBatch.MaterialRenderProxy->GetMaterialWithFallback(FeatureLevel, FallbackMaterialRenderProxyPtr);
 	const bool bIsCompatible = IsCompatibleWithHairStrands(&Material, FeatureLevel);
-	const bool bIsHairStrandsFactory = MeshBatch.VertexFactory->GetType()->GetId() == CompatibleVF->GetId();
+	const bool bIsHairStrandsFactory = MeshBatch.VertexFactory->GetType()->GetHashedName() == CompatibleVF->GetHashedName();
 	const bool bShouldRender = (!PrimitiveSceneProxy && MeshBatch.Elements.Num()>0) || (PrimitiveSceneProxy && PrimitiveSceneProxy->ShouldRenderInMainPass());
 
 	if (bIsCompatible
@@ -510,7 +498,9 @@ static FMaterialPassOutput AddHairMaterialPass(
 			FDynamicMeshDrawCommandStorage DynamicMeshDrawCommandStorage;
 			FMeshCommandOneFrameArray VisibleMeshDrawCommands;
 			FGraphicsMinimalPipelineStateSet PipelineStateSet;
-			FDynamicPassMeshDrawListContext ShadowContext(DynamicMeshDrawCommandStorage, VisibleMeshDrawCommands, PipelineStateSet);
+			// @todo loadtime arnes: do we need to pass this along to somewhere?
+			bool NeedsShaderInitialization;
+			FDynamicPassMeshDrawListContext ShadowContext(DynamicMeshDrawCommandStorage, VisibleMeshDrawCommands, PipelineStateSet, NeedsShaderInitialization);
 			FHairMaterialProcessor MeshProcessor(Scene, ViewInfo, DrawRenderState, &ShadowContext);
 
 			for (const FHairStrandsMacroGroupData& MacroGroupData : MacroGroupDatas.Datas)
@@ -612,7 +602,7 @@ static void AddHairVelocityPass(
 	FComputeShaderUtils::AddPass(
 		GraphBuilder,
 		RDG_EVENT_NAME("HairStrandsVelocity"),
-		*ComputeShader,
+		ComputeShader,
 		PassParameters,
 		FComputeShaderUtils::GetGroupCount(RectResolution, GroupSize));
 }
@@ -675,7 +665,7 @@ static FRDGTextureRef AddHairLightChannelMaskPass(
 	FComputeShaderUtils::AddPass(
 		GraphBuilder,
 		RDG_EVENT_NAME("HairLightChannelMask"),
-		*ComputeShader,
+		ComputeShader,
 		PassParameters,
 		FComputeShaderUtils::GetGroupCount(Resolution, GroupSize));
 
@@ -773,7 +763,7 @@ public:
 
 	static bool ShouldCompilePermutation(const FMeshMaterialShaderPermutationParameters& Parameters)
 	{
-		return IsCompatibleWithHairStrands(Parameters.Material, Parameters.Platform) && Parameters.VertexFactoryType->GetFName() == FName(TEXT("FHairStrandsVertexFactory"));
+		return IsCompatibleWithHairStrands(Parameters.Platform, Parameters.MaterialParameters) && Parameters.VertexFactoryType->GetFName() == FName(TEXT("FHairStrandsVertexFactory"));;
 	}
 
 	static void ModifyCompilationEnvironment(const FMaterialShaderPermutationParameters& Parameters, FShaderCompilerEnvironment& OutEnvironment)	
@@ -783,22 +773,6 @@ public:
 		OutEnvironment.SetDefine(TEXT("HAIR_RENDER_MODE"), RenderModeValue);
 	}
 
-	virtual bool Serialize(FArchive& Ar) override
-	{
-		bool bShaderHasOutdatedParameters = FMeshMaterialShader::Serialize(Ar);
-		Ar << HairVisibilityPass_HairMacroGroupIndex;
-		Ar << HairVisibilityPass_HairMaterialId;
-		Ar << HairVisibilityPass_LightChannelMask;
-		Ar << HairVisibilityPass_MaxPPLLNodeCount;
-		if (RenderMode == EHairVisibilityRenderMode::HairVisibilityRenderMode_PPLL)
-		{
-			Ar << PPLLCounter;
-			Ar << PPLLNodeIndex;
-			Ar << PPLLNodes;
-		}
-		return bShaderHasOutdatedParameters;
-	}
-
 	void GetShaderBindings(
 		const FScene* Scene,
 		ERHIFeatureLevel::Type FeatureLevel,
@@ -816,14 +790,14 @@ public:
 		ShaderBindings.Add(HairVisibilityPass_MaxPPLLNodeCount, ShaderElementData.MaxPPLLNodeCount);
 	}
 
-	FShaderParameter HairVisibilityPass_HairMacroGroupIndex;
-	FShaderParameter HairVisibilityPass_HairMaterialId;
-	FShaderParameter HairVisibilityPass_LightChannelMask;
-	FShaderParameter HairVisibilityPass_MaxPPLLNodeCount;
+	LAYOUT_FIELD(FShaderParameter, HairVisibilityPass_HairMacroGroupIndex);
+	LAYOUT_FIELD(FShaderParameter, HairVisibilityPass_HairMaterialId);
+	LAYOUT_FIELD(FShaderParameter, HairVisibilityPass_LightChannelMask);
+	LAYOUT_FIELD(FShaderParameter, HairVisibilityPass_MaxPPLLNodeCount);
 	// This is on FVisibilityPassParameters but needed to avoid the compiler to assert with unbound parameters 
-	FShaderParameter PPLLCounter;
-	FShaderParameter PPLLNodeIndex;
-	FShaderParameter PPLLNodes;
+	LAYOUT_FIELD(FShaderParameter, PPLLCounter);
+	LAYOUT_FIELD(FShaderParameter, PPLLNodeIndex);
+	LAYOUT_FIELD(FShaderParameter, PPLLNodes);
 };
 IMPLEMENT_MATERIAL_SHADER_TYPE(template<>, FHairVisibilityPS<HairVisibilityRenderMode_MSAA_Visibility>, TEXT("/Engine/Private/HairStrands/HairStrandsVisibilityPS.usf"), TEXT("MainVisibility"), SF_Pixel);
 IMPLEMENT_MATERIAL_SHADER_TYPE(template<>, FHairVisibilityPS<HairVisibilityRenderMode_MSAA>, TEXT("/Engine/Private/HairStrands/HairStrandsVisibilityPS.usf"), TEXT("MainVisibility"), SF_Pixel);
@@ -881,7 +855,7 @@ void FHairVisibilityProcessor::AddMeshBatch(const FMeshBatch& RESTRICT MeshBatch
 	const FMaterialRenderProxy* FallbackMaterialRenderProxyPtr = nullptr;
 	const FMaterial& Material = MeshBatch.MaterialRenderProxy->GetMaterialWithFallback(FeatureLevel, FallbackMaterialRenderProxyPtr);
 	const bool bIsCompatible = IsCompatibleWithHairStrands(&Material, FeatureLevel);
-	const bool bIsHairStrandsFactory = MeshBatch.VertexFactory->GetType()->GetId() == CompatibleVF->GetId();
+	const bool bIsHairStrandsFactory = MeshBatch.VertexFactory->GetType()->GetHashedName() == CompatibleVF->GetHashedName();
 	const bool bShouldRender = (!PrimitiveSceneProxy && MeshBatch.Elements.Num()>0) || (PrimitiveSceneProxy && PrimitiveSceneProxy->ShouldRenderInMainPass());
 	const uint32 LightChannelMask = PrimitiveSceneProxy && PrimitiveSceneProxy->GetLightingChannelMask();
 
@@ -891,8 +865,9 @@ void FHairVisibilityProcessor::AddMeshBatch(const FMeshBatch& RESTRICT MeshBatch
 		&& ShouldIncludeDomainInMeshPass(Material.GetMaterialDomain()))
 	{
 		const FMaterialRenderProxy& MaterialRenderProxy = FallbackMaterialRenderProxyPtr ? *FallbackMaterialRenderProxyPtr : *MeshBatch.MaterialRenderProxy;
-		const ERasterizerFillMode MeshFillMode = ComputeMeshFillMode(MeshBatch, Material);
-		const ERasterizerCullMode MeshCullMode = ComputeMeshCullMode(MeshBatch, Material);
+		const FMeshDrawingPolicyOverrideSettings OverrideSettings = ComputeMeshOverrideSettings(MeshBatch);
+		const ERasterizerFillMode MeshFillMode = ComputeMeshFillMode(MeshBatch, Material, OverrideSettings);
+		const ERasterizerCullMode MeshCullMode = ComputeMeshCullMode(MeshBatch, Material, OverrideSettings);
 		if (RenderMode == HairVisibilityRenderMode_MSAA_Visibility)
 			Process<HairVisibilityRenderMode_MSAA_Visibility>(MeshBatch, BatchElementMask, PrimitiveSceneProxy, StaticMeshId, MaterialRenderProxy, Material, HairMacroGroupId, HairMaterialId, LightChannelMask, MeshFillMode, MeshCullMode);
 		if (RenderMode == HairVisibilityRenderMode_MSAA)
@@ -1001,7 +976,7 @@ static void AddClearGraphicPass(
 	const FIntRect Viewport = View->ViewRect;
 	const FIntPoint Resolution = OutTarget->Desc.Extent;
 
-	ClearUnusedGraphResources(*PixelShader, Parameters);
+	ClearUnusedGraphResources(PixelShader, Parameters);
 
 	GraphBuilder.AddPass(
 		Forward<FRDGEventName>(PassName),
@@ -1015,14 +990,14 @@ static void AddClearGraphicPass(
 		GraphicsPSOInit.RasterizerState = TStaticRasterizerState<>::GetRHI();
 		GraphicsPSOInit.DepthStencilState = TStaticDepthStencilState<false, CF_Always>::GetRHI();
 		GraphicsPSOInit.BoundShaderState.VertexDeclarationRHI = GFilterVertexDeclaration.VertexDeclarationRHI;
-		GraphicsPSOInit.BoundShaderState.VertexShaderRHI = GETSAFERHISHADER_VERTEX(*VertexShader);
-		GraphicsPSOInit.BoundShaderState.PixelShaderRHI = GETSAFERHISHADER_PIXEL(*PixelShader);
+		GraphicsPSOInit.BoundShaderState.VertexShaderRHI = VertexShader.GetVertexShader();
+		GraphicsPSOInit.BoundShaderState.PixelShaderRHI = PixelShader.GetPixelShader();
 		GraphicsPSOInit.PrimitiveType = PT_TriangleList;
 		SetGraphicsPipelineState(RHICmdList, GraphicsPSOInit);
 
 		VertexShader->SetParameters(RHICmdList, View->ViewUniformBuffer);
 		RHICmdList.SetViewport(Viewport.Min.X, Viewport.Min.Y, 0.0f, Viewport.Max.X, Viewport.Max.Y, 1.0f);
-		SetShaderParameters(RHICmdList, *PixelShader, PixelShader->GetPixelShader(), *Parameters);
+		SetShaderParameters(RHICmdList, PixelShader, PixelShader.GetPixelShader(), *Parameters);
 
 		DrawRectangle(
 			RHICmdList,
@@ -1032,7 +1007,7 @@ static void AddClearGraphicPass(
 			Viewport.Width(),Viewport.Height(),
 			Viewport.Size(),
 			Resolution,
-			*VertexShader,
+			VertexShader,
 			EDRF_UseTriangleOptimization);
 	});
 }
@@ -1104,7 +1079,7 @@ static FRDGBufferRef AddCopyIndirectArgPass(
 	FComputeShaderUtils::AddPass(
 		GraphBuilder,
 		RDG_EVENT_NAME("HairStrandsVisbilityCopyIndirectArgs"),
-		*ComputeShader,
+		ComputeShader,
 		Parameters,
 		FIntVector(1,1,1));
 
@@ -1330,7 +1305,7 @@ static void AddHairVisibilityPrimitiveIdCompactionPass(
 	FComputeShaderUtils::AddPass(
 		GraphBuilder,
 		RDG_EVENT_NAME("HairStrandsVisibilityCompaction"),
-		*ComputeShader,
+		ComputeShader,
 		PassParameters,
 		FComputeShaderUtils::GetGroupCount(RectResolution, GroupSize));
 
@@ -1428,7 +1403,7 @@ static void AddGenerateTilePass(
 	FComputeShaderUtils::AddPass(
 		GraphBuilder,
 		RDG_EVENT_NAME("HairGenerateTile"),
-		*ComputeShader,
+		ComputeShader,
 		PassParameters,
 		FIntVector(TileResolution.X, TileResolution.Y, 1));
 
@@ -1492,7 +1467,7 @@ static FRDGTextureRef AddHairVisibilityFillOpaqueDepth(
 
 	TShaderMapRef<FPostProcessVS> VertexShader(View.ShaderMap);
 	TShaderMapRef<FHairVisibilityFillOpaqueDepthPS> PixelShader(View.ShaderMap);
-	const TShaderMap<FGlobalShaderType>* GlobalShaderMap = View.ShaderMap;
+	const FGlobalShaderMap* GlobalShaderMap = View.ShaderMap;
 	const FIntRect Viewport = View.ViewRect;
 	const FViewInfo* CapturedView = &View;
 
@@ -1510,7 +1485,7 @@ static FRDGTextureRef AddHairVisibilityFillOpaqueDepth(
 	}
 
 	{
-		ClearUnusedGraphResources(*PixelShader, Parameters);
+		ClearUnusedGraphResources(PixelShader, Parameters);
 
 		GraphBuilder.AddPass(
 			RDG_EVENT_NAME("HairStrandsVisibilityFillOpaqueDepth"),
@@ -1525,13 +1500,13 @@ static FRDGTextureRef AddHairVisibilityFillOpaqueDepth(
 			GraphicsPSOInit.DepthStencilState = TStaticDepthStencilState<true, CF_DepthNearOrEqual>::GetRHI();
 
 			GraphicsPSOInit.BoundShaderState.VertexDeclarationRHI = GFilterVertexDeclaration.VertexDeclarationRHI;
-			GraphicsPSOInit.BoundShaderState.VertexShaderRHI = GETSAFERHISHADER_VERTEX(*VertexShader);
-			GraphicsPSOInit.BoundShaderState.PixelShaderRHI = GETSAFERHISHADER_PIXEL(*PixelShader);
+			GraphicsPSOInit.BoundShaderState.VertexShaderRHI = VertexShader.GetVertexShader();
+			GraphicsPSOInit.BoundShaderState.PixelShaderRHI = PixelShader.GetPixelShader();
 			GraphicsPSOInit.PrimitiveType = PT_TriangleList;
 			SetGraphicsPipelineState(RHICmdList, GraphicsPSOInit);
 
 			VertexShader->SetParameters(RHICmdList, CapturedView->ViewUniformBuffer);
-			SetShaderParameters(RHICmdList, *PixelShader, PixelShader->GetPixelShader(), *Parameters);
+			SetShaderParameters(RHICmdList, PixelShader, PixelShader.GetPixelShader(), *Parameters);
 
 			for (const FIntRect& ViewRect : MacroGroupRects)
 			{
@@ -1544,7 +1519,7 @@ static FRDGTextureRef AddHairVisibilityFillOpaqueDepth(
 					Viewport.Width(), Viewport.Height(),
 					Viewport.Size(),
 					Resolution,
-					*VertexShader,
+					VertexShader,
 					EDRF_UseTriangleOptimization);
 			}
 		});
@@ -1635,7 +1610,8 @@ static void AddHairVisibilityCommonPass(
 			FDynamicMeshDrawCommandStorage DynamicMeshDrawCommandStorage;
 			FMeshCommandOneFrameArray VisibleMeshDrawCommands;
 			FGraphicsMinimalPipelineStateSet PipelineStateSet;
-			FDynamicPassMeshDrawListContext ShadowContext(DynamicMeshDrawCommandStorage, VisibleMeshDrawCommands, PipelineStateSet);
+			bool NeedsShaderInitialization;
+			FDynamicPassMeshDrawListContext ShadowContext(DynamicMeshDrawCommandStorage, VisibleMeshDrawCommands, PipelineStateSet, NeedsShaderInitialization);
 			FHairVisibilityProcessor MeshProcessor(Scene, ViewInfo, DrawRenderState, RenderMode, &ShadowContext);
 			if (RenderMode == HairVisibilityRenderMode_PPLL)
 			{
@@ -1961,11 +1937,11 @@ static void AddHairViewTransmittanceDepthPass(
 	
 	TShaderMapRef<FPostProcessVS> VertexShader(View.ShaderMap);
 	TShaderMapRef<FHairViewTransmittanceDepthPS> PixelShader(View.ShaderMap);
-	const TShaderMap<FGlobalShaderType>* GlobalShaderMap = View.ShaderMap;
+	const FGlobalShaderMap* GlobalShaderMap = View.ShaderMap;
 	const FIntRect Viewport = View.ViewRect;
 	const FIntPoint Resolution = HairCountTexture->Desc.Extent;
 	const FViewInfo* CapturedView = &View;
-	ClearUnusedGraphResources(*PixelShader, Parameters);
+	ClearUnusedGraphResources(PixelShader, Parameters);
 
 	GraphBuilder.AddPass(
 		RDG_EVENT_NAME("HairStrandsViewTransmittanceDepth"),
@@ -1980,14 +1956,14 @@ static void AddHairViewTransmittanceDepthPass(
 		GraphicsPSOInit.DepthStencilState = TStaticDepthStencilState<false, CF_Always>::GetRHI();
 
 		GraphicsPSOInit.BoundShaderState.VertexDeclarationRHI = GFilterVertexDeclaration.VertexDeclarationRHI;
-		GraphicsPSOInit.BoundShaderState.VertexShaderRHI = GETSAFERHISHADER_VERTEX(*VertexShader);
-		GraphicsPSOInit.BoundShaderState.PixelShaderRHI = GETSAFERHISHADER_PIXEL(*PixelShader);
+		GraphicsPSOInit.BoundShaderState.VertexShaderRHI = VertexShader.GetVertexShader();
+		GraphicsPSOInit.BoundShaderState.PixelShaderRHI = PixelShader.GetPixelShader();
 		GraphicsPSOInit.PrimitiveType = PT_TriangleList;
 		SetGraphicsPipelineState(RHICmdList, GraphicsPSOInit);
 
 		VertexShader->SetParameters(RHICmdList, CapturedView->ViewUniformBuffer);
 		RHICmdList.SetViewport(Viewport.Min.X, Viewport.Min.Y, 0.0f, Viewport.Max.X, Viewport.Max.Y, 1.0f);
-		SetShaderParameters(RHICmdList, *PixelShader, PixelShader->GetPixelShader(), *Parameters);
+		SetShaderParameters(RHICmdList, PixelShader, PixelShader.GetPixelShader(), *Parameters);
 		DrawRectangle(
 			RHICmdList,
 			0, 0,
@@ -1996,7 +1972,7 @@ static void AddHairViewTransmittanceDepthPass(
 			Viewport.Width(), Viewport.Height(),
 			Viewport.Size(),
 			Resolution,
-			*VertexShader,
+			VertexShader,
 			EDRF_UseTriangleOptimization);
 	});
 }
@@ -2052,13 +2028,13 @@ static void AddHairVisibilityColorAndDepthPatchPass(
 	FHairVisibilityDepthPS::FPermutationDomain PermutationVector;
 	PermutationVector = FHairVisibilityDepthPS::RemapPermutation(PermutationVector);
 	TShaderMapRef<FHairVisibilityDepthPS> PixelShader(View.ShaderMap, PermutationVector);
-	const TShaderMap<FGlobalShaderType>* GlobalShaderMap = View.ShaderMap;
+	const FGlobalShaderMap* GlobalShaderMap = View.ShaderMap;
 	const FIntRect Viewport = View.ViewRect;
 	const FIntPoint Resolution = OutDepthTexture->Desc.Extent;
 	const FViewInfo* CapturedView = &View;
 
 	{
-		ClearUnusedGraphResources(*PixelShader, Parameters);
+		ClearUnusedGraphResources(PixelShader, Parameters);
 
 		GraphBuilder.AddPass(
 			RDG_EVENT_NAME("HairStrandsVisibilityWriteColorAndDepth"),
@@ -2073,14 +2049,14 @@ static void AddHairVisibilityColorAndDepthPatchPass(
 			GraphicsPSOInit.DepthStencilState = TStaticDepthStencilState<true, CF_Greater>::GetRHI();
 
 			GraphicsPSOInit.BoundShaderState.VertexDeclarationRHI = GFilterVertexDeclaration.VertexDeclarationRHI;
-			GraphicsPSOInit.BoundShaderState.VertexShaderRHI = GETSAFERHISHADER_VERTEX(*VertexShader);
-			GraphicsPSOInit.BoundShaderState.PixelShaderRHI = GETSAFERHISHADER_PIXEL(*PixelShader);
+			GraphicsPSOInit.BoundShaderState.VertexShaderRHI = VertexShader.GetVertexShader();
+			GraphicsPSOInit.BoundShaderState.PixelShaderRHI = PixelShader.GetPixelShader();
 			GraphicsPSOInit.PrimitiveType = PT_TriangleList;
 			SetGraphicsPipelineState(RHICmdList, GraphicsPSOInit);
 
 			VertexShader->SetParameters(RHICmdList, CapturedView->ViewUniformBuffer);
 			RHICmdList.SetViewport(Viewport.Min.X, Viewport.Min.Y, 0.0f, Viewport.Max.X, Viewport.Max.Y, 1.0f);
-			SetShaderParameters(RHICmdList, *PixelShader, PixelShader->GetPixelShader(), *Parameters);
+			SetShaderParameters(RHICmdList, PixelShader, PixelShader.GetPixelShader(), *Parameters);
 			DrawRectangle(
 				RHICmdList,
 				0, 0,
@@ -2089,7 +2065,7 @@ static void AddHairVisibilityColorAndDepthPatchPass(
 				Viewport.Width(), Viewport.Height(),
 				Viewport.Size(),
 				Resolution,
-				*VertexShader,
+				VertexShader,
 				EDRF_UseTriangleOptimization);
 		});
 	}
@@ -2170,7 +2146,7 @@ FHairStrandsVisibilityViews RenderHairStrandsVisibilityBuffer(
 					FRDGTextureRef AttributeTexture;
 					FRDGTextureRef VelocityTexture;
 				} MsaaVisibilityResources;
-
+				
 				MsaaVisibilityResources.DepthTexture = AddHairVisibilityFillOpaqueDepth(
 					GraphBuilder,
 					View,
@@ -2331,7 +2307,7 @@ FHairStrandsVisibilityViews RenderHairStrandsVisibilityBuffer(
 					GraphBuilder.QueueBufferExtraction(CompactNodeCoord, &VisibilityData.NodeCoord, FRDGResourceState::EAccess::Read, FRDGResourceState::EPipeline::Graphics);
 					GraphBuilder.QueueBufferExtraction(IndirectArgsBuffer, &VisibilityData.NodeIndirectArg, FRDGResourceState::EAccess::Read, FRDGResourceState::EPipeline::Compute);
 				}
-
+		
 				AddHairVisibilityColorAndDepthPatchPass(
 					GraphBuilder,
 					View,
diff --git a/Engine/Source/Runtime/Renderer/Private/HairStrands/HairStrandsVoxelization.cpp b/Engine/Source/Runtime/Renderer/Private/HairStrands/HairStrandsVoxelization.cpp
--- a/Engine/Source/Runtime/Renderer/Private/HairStrands/HairStrandsVoxelization.cpp
+++ b/Engine/Source/Runtime/Renderer/Private/HairStrands/HairStrandsVoxelization.cpp
@@ -169,11 +169,11 @@ static void AddVoxelInjectOpaquePass(
 	const FIntPoint OutputResolution = SceneTextures.SceneDepthBuffer->Desc.Extent;
 	TShaderMapRef<FPostProcessVS> VertexShader(View.ShaderMap);
 	TShaderMapRef<FVoxelInjectOpaquePS> PixelShader(View.ShaderMap);
-	const TShaderMap<FGlobalShaderType>* GlobalShaderMap = View.ShaderMap;
+	const FGlobalShaderMap* GlobalShaderMap = View.ShaderMap;
 	const FIntRect Viewport(0, 0, Resolution.X, Resolution.Y);
 	const FViewInfo* CapturedView = &View;
 	
-	ClearUnusedGraphResources(*PixelShader, Parameters);
+	ClearUnusedGraphResources(PixelShader, Parameters);
 
 	GraphBuilder.AddPass(
 		RDG_EVENT_NAME("HairStrandsVoxelInjectOpaque"),
@@ -188,14 +188,14 @@ static void AddVoxelInjectOpaquePass(
 		GraphicsPSOInit.DepthStencilState = TStaticDepthStencilState<false, CF_Always>::GetRHI();
 
 		GraphicsPSOInit.BoundShaderState.VertexDeclarationRHI = GFilterVertexDeclaration.VertexDeclarationRHI;
-		GraphicsPSOInit.BoundShaderState.VertexShaderRHI = GETSAFERHISHADER_VERTEX(*VertexShader);
-		GraphicsPSOInit.BoundShaderState.PixelShaderRHI = GETSAFERHISHADER_PIXEL(*PixelShader);
+		GraphicsPSOInit.BoundShaderState.VertexShaderRHI = VertexShader.GetVertexShader();
+		GraphicsPSOInit.BoundShaderState.PixelShaderRHI = PixelShader.GetPixelShader();
 		GraphicsPSOInit.PrimitiveType = PT_TriangleList;
 		SetGraphicsPipelineState(RHICmdList, GraphicsPSOInit);
 
 		VertexShader->SetParameters(RHICmdList, CapturedView->ViewUniformBuffer);
 		RHICmdList.SetViewport(Viewport.Min.X, Viewport.Min.Y, 0.0f, Viewport.Max.X, Viewport.Max.Y, 1.0f);
-		SetShaderParameters(RHICmdList, *PixelShader, PixelShader->GetPixelShader(), *Parameters);
+		SetShaderParameters(RHICmdList, PixelShader, PixelShader.GetPixelShader(), *Parameters);
 
 		DrawRectangle(
 			RHICmdList,
@@ -205,7 +205,7 @@ static void AddVoxelInjectOpaquePass(
 			Viewport.Width(), Viewport.Height(),
 			Viewport.Size(),
 			Resolution,
-			*VertexShader,
+			VertexShader,
 			EDRF_UseTriangleOptimization);
 	});
 }
@@ -275,7 +275,7 @@ static void AddVirtualVoxelInjectOpaquePass(
 	Parameters->DispatchedPageIndexResolution = MacroGroup.VirtualVoxelNodeDesc.PageIndexResolution;
 	Parameters->IndirectDispatchArgs = IndirectDispatchArgsBuffer;
 	TShaderMapRef<FVirtualVoxelInjectOpaqueCS> ComputeShader(View.ShaderMap);
-	const TShaderMap<FGlobalShaderType>* GlobalShaderMap = View.ShaderMap;
+	const FGlobalShaderMap* GlobalShaderMap = View.ShaderMap;
 
 	check(VoxelResources.Parameters.Common.IndirectDispatchGroupSize == 64);
 	const uint32 ArgsOffset = sizeof(uint32) * 3 * Parameters->MacroGroupId;
@@ -283,7 +283,7 @@ static void AddVirtualVoxelInjectOpaquePass(
 	FComputeShaderUtils::AddPass(
 		GraphBuilder, 
 		RDG_EVENT_NAME("HairStrandsInjectOpaqueDepthInVoxel"), 
-		*ComputeShader, 
+		ComputeShader, 
 		Parameters, 
 		IndirectDispatchArgsBuffer, 
 		ArgsOffset);
@@ -608,7 +608,7 @@ static void AddAllocateVoxelPagesPass(
 		FComputeShaderUtils::AddPass(
 			GraphBuilder,
 			RDG_EVENT_NAME("HairStrandsAllocatePageIndex"),
-			*ComputeShader,
+			ComputeShader,
 			Parameters,
 			FIntVector(1,1,1));
 	}
@@ -654,7 +654,7 @@ static void AddAllocateVoxelPagesPass(
 			FComputeShaderUtils::AddPass(
 				GraphBuilder,
 				RDG_EVENT_NAME("HairStrandsMarkValidPageIndex"),
-				*ComputeShader,
+				ComputeShader,
 				Parameters,
 				DispatchCount);
 		}
@@ -684,7 +684,7 @@ static void AddAllocateVoxelPagesPass(
 
 			const FIntVector DispatchCount(1, 1, 1);
 			TShaderMapRef<FVoxelAddNodeDescCS> ComputeShader(View.ShaderMap, PermutationVector);
-			FComputeShaderUtils::AddPass(GraphBuilder, RDG_EVENT_NAME("HairStrandsAddNodeDesc"), *ComputeShader, Parameters, DispatchCount);
+			FComputeShaderUtils::AddPass(GraphBuilder, RDG_EVENT_NAME("HairStrandsAddNodeDesc"), ComputeShader, Parameters, DispatchCount);
 		}
 
 		// Allocate pages
@@ -710,7 +710,7 @@ static void AddAllocateVoxelPagesPass(
 
 				const uint32 ArgsOffset = sizeof(uint32) * 3 * MacroGroup.MacroGroupId;
 
-				FComputeShaderUtils::AddPass(GraphBuilder, RDG_EVENT_NAME("HairStrandsAllocateVoxelPage"), *ComputeShader, Parameters, 
+				FComputeShaderUtils::AddPass(GraphBuilder, RDG_EVENT_NAME("HairStrandsAllocateVoxelPage"), ComputeShader, Parameters, 
 					PageIndexAllocationIndirectBufferArgs,
 					ArgsOffset);
 			}
@@ -718,7 +718,7 @@ static void AddAllocateVoxelPagesPass(
 			{
 				const FIntVector DispatchCount((CPUAllocationDesc.PageIndexCount + GroupSize - 1) / GroupSize, 1, 1);
 				check(DispatchCount.X < 65535);
-				FComputeShaderUtils::AddPass(GraphBuilder, RDG_EVENT_NAME("HairStrandsAllocateVoxelPage"), *ComputeShader, Parameters, DispatchCount);
+				FComputeShaderUtils::AddPass(GraphBuilder, RDG_EVENT_NAME("HairStrandsAllocateVoxelPage"), ComputeShader, Parameters, DispatchCount);
 			}
 		}
 
@@ -735,7 +735,7 @@ static void AddAllocateVoxelPagesPass(
 
 			const FIntVector DispatchCount(1, 1, 1);
 			TShaderMapRef<FVoxelAddIndirectBufferCS> ComputeShader(View.ShaderMap);
-			FComputeShaderUtils::AddPass(GraphBuilder, RDG_EVENT_NAME("HairStrandsBuildVoxelIndirectArgs"), *ComputeShader, Parameters, DispatchCount);
+			FComputeShaderUtils::AddPass(GraphBuilder, RDG_EVENT_NAME("HairStrandsBuildVoxelIndirectArgs"), ComputeShader, Parameters, DispatchCount);
 		}
 	}
 
@@ -894,7 +894,7 @@ static FRDGBufferRef IndirectVoxelPageClear(
 		FComputeShaderUtils::AddPass(
 			GraphBuilder,
 			RDG_EVENT_NAME("HairStrandsVoxelGenIndBufferClearCS"),
-			*ComputeShader,
+			ComputeShader,
 			Parameters,
 			FIntVector(1,1,1));
 	}
@@ -913,7 +913,7 @@ static FRDGBufferRef IndirectVoxelPageClear(
 		FComputeShaderUtils::AddPass(
 			GraphBuilder,
 			RDG_EVENT_NAME("HairStrandsVoxelIndPageClearCS"),
-			*ComputeShader,
+			ComputeShader,
 			Parameters,
 			ClearIndArgsBuffer,
 			0);
@@ -1071,17 +1071,17 @@ static void AddFilterVoxelOpaqueDepthPass(
 	Parameters->VoxelResolution = VoxelResolution;
 
 	TShaderMapRef<FVoxelFilterDepthCS> ComputeShader(View.ShaderMap);
-	const TShaderMap<FGlobalShaderType>* GlobalShaderMap = View.ShaderMap;
+	const FGlobalShaderMap* GlobalShaderMap = View.ShaderMap;
 	const FIntVector DispatchCount = FComputeShaderUtils::GetGroupCount(FIntVector(VoxelResolution, VoxelResolution, VoxelResolution), FIntVector(4, 4, 4));
 
-	ClearUnusedGraphResources(*ComputeShader, Parameters);
+	ClearUnusedGraphResources(ComputeShader, Parameters);
 	GraphBuilder.AddPass(
 		RDG_EVENT_NAME("HairStrandsVoxelFilterDepth"),
 		Parameters,
 		ERDGPassFlags::Compute,
 		[Parameters, ComputeShader, DispatchCount](FRHICommandList& RHICmdList)
 	{
-		FComputeShaderUtils::Dispatch(RHICmdList, *ComputeShader, *Parameters, DispatchCount);
+		FComputeShaderUtils::Dispatch(RHICmdList, ComputeShader, *Parameters, DispatchCount);
 	});
 }
 
@@ -1152,22 +1152,22 @@ static void AddVoxelGenerateMipPass(
 
 		FVoxelGenerateMipCS::FPermutationDomain PermutationVector;
 		PermutationVector.Set<FVoxelGenerateMipCS::FMethod>(NumLevelPerPass == 2 ? 1 : 0);
-		
+
 		TShaderMapRef<FVoxelGenerateMipCS> ComputeShader(View.ShaderMap, PermutationVector);
-		const TShaderMap<FGlobalShaderType>* GlobalShaderMap = View.ShaderMap;
+		const FGlobalShaderMap* GlobalShaderMap = View.ShaderMap;
 		const FIntVector DispatchCount = 
 			NumLevelPerPass == 1 ? 
 			FComputeShaderUtils::GetGroupCount(FIntVector(TargetResolution, TargetResolution, TargetResolution), FIntVector(4, 4, 4)):
 			FComputeShaderUtils::GetGroupCount(FIntVector(SourceResolution, SourceResolution, SourceResolution), FIntVector(4, 4, 4));
 		
-		ClearUnusedGraphResources(*ComputeShader, Parameters);
+		ClearUnusedGraphResources(ComputeShader, Parameters);
 		GraphBuilder.AddPass(
 			RDG_EVENT_NAME("HairStrandsVoxelMip"),
 			Parameters,
 			ERDGPassFlags::Compute | ERDGPassFlags::GenerateMips,
 			[Parameters, ComputeShader, DispatchCount](FRHICommandList& RHICmdList)
 		{
-			FComputeShaderUtils::Dispatch(RHICmdList, *ComputeShader, *Parameters, DispatchCount);
+			FComputeShaderUtils::Dispatch(RHICmdList, ComputeShader, *Parameters, DispatchCount);
 		});
 	}
 }
@@ -1263,7 +1263,7 @@ static void AddVirtualVoxelGenerateMipPass(
 		Parameters->OutIndirectArgs		= GraphBuilder.CreateUAV(MipIndirectArgs);
 
 		TShaderMapRef<FVirtualVoxelIndirectArgMipCS> ComputeShader(View.ShaderMap);
-		FComputeShaderUtils::AddPass(GraphBuilder, RDG_EVENT_NAME("HairStrandsBuildVoxelMipIndirectArgs"), *ComputeShader, Parameters, FIntVector(1, 1, 1));
+		FComputeShaderUtils::AddPass(GraphBuilder, RDG_EVENT_NAME("HairStrandsBuildVoxelMipIndirectArgs"), ComputeShader, Parameters, FIntVector(1, 1, 1));
 	}
 
 	// Generate MIP level (in one go for all allocated pages)
@@ -1282,14 +1282,14 @@ static void AddVirtualVoxelGenerateMipPass(
 		Parameters->IndirectDispatchArgs = MipIndirectArgsBuffers[MipIt];
 
 		TShaderMapRef<FVirtualVoxelGenerateMipCS> ComputeShader(View.ShaderMap);
-		ClearUnusedGraphResources(*ComputeShader, Parameters);
+		ClearUnusedGraphResources(ComputeShader, Parameters);
 		GraphBuilder.AddPass(
 			RDG_EVENT_NAME("HairStrandsComputeVoxelMip"),
 			Parameters,
 			ERDGPassFlags::Compute | ERDGPassFlags::GenerateMips,
 			[Parameters, ComputeShader](FRHICommandList& RHICmdList)
 		{
-			FComputeShaderUtils::DispatchIndirect(RHICmdList, *ComputeShader, *Parameters, Parameters->IndirectDispatchArgs->GetIndirectRHICallBuffer(), 0);
+			FComputeShaderUtils::DispatchIndirect(RHICmdList, ComputeShader, *Parameters, Parameters->IndirectDispatchArgs->GetIndirectRHICallBuffer(), 0);
 		});
 	}
 }
diff --git a/Engine/Source/Runtime/Renderer/Private/HeightfieldLighting.cpp b/Engine/Source/Runtime/Renderer/Private/HeightfieldLighting.cpp
--- a/Engine/Source/Runtime/Renderer/Private/HeightfieldLighting.cpp
+++ b/Engine/Source/Runtime/Renderer/Private/HeightfieldLighting.cpp
@@ -17,6 +17,9 @@ HeightfieldLighting.cpp
 #include "LightRendering.h"
 #include "PipelineStateCache.h"
 
+IMPLEMENT_TYPE_LAYOUT(FHeightfieldDescriptionParameters);
+IMPLEMENT_TYPE_LAYOUT(FHeightfieldTextureParameters);
+
 // Currently disabled because the bHasHeightfieldRepresentation GBuffer bit has been reallocated, and self-shadowing artifacts are too severe without that bit
 int32 GAOHeightfieldOcclusion = 0;
 FAutoConsoleVariableRef CVarAOHeightfieldOcclusion(
@@ -176,6 +179,7 @@ TGlobalResource<FSubsectionHeightfieldDescriptionsResource> GSubsectionHeightfie
 
 class FSubsectionHeightfieldDescriptionParameters
 {
+	DECLARE_INLINE_TYPE_LAYOUT(FSubsectionHeightfieldDescriptionParameters, NonVirtual);
 public:
 	void Bind(const FShaderParameterMap& ParameterMap)
 	{
@@ -195,7 +199,9 @@ public:
 	}
 
 private:
-	FShaderResourceParameter SubsectionHeightfieldDescriptions;
+	
+		LAYOUT_FIELD(FShaderResourceParameter, SubsectionHeightfieldDescriptions)
+	
 };
 
 int32 UploadSubsectionHeightfieldDescriptions(const TArray<FHeightfieldComponentDescription>& HeightfieldDescriptions, FVector2D InvLightingAtlasSize, float InvDownsampleFactor)
@@ -267,23 +273,15 @@ public:
 
 	void SetParameters(FRHICommandList& RHICmdList, const FSceneView& View)
 	{
-		FRHIVertexShader* ShaderRHI = GetVertexShader();
+		FRHIVertexShader* ShaderRHI = RHICmdList.GetBoundVertexShader();
 
 		FGlobalShader::SetParameters<FViewUniformShaderParameters>(RHICmdList, ShaderRHI, View.ViewUniformBuffer);
 		SubsectionHeightfieldParameters.Set(RHICmdList, ShaderRHI);
 	}
 
-	// FShader interface.
-	virtual bool Serialize(FArchive& Ar) override
-	{
-		bool bShaderHasOutdatedParameters = FGlobalShader::Serialize(Ar);
-		Ar << SubsectionHeightfieldParameters;
-		return bShaderHasOutdatedParameters;
-	}
-
 private:
 
-	FSubsectionHeightfieldDescriptionParameters SubsectionHeightfieldParameters;
+	LAYOUT_FIELD(FSubsectionHeightfieldDescriptionParameters, SubsectionHeightfieldParameters);
 };
 
 IMPLEMENT_SHADER_TYPE(, FHeightfieldSubsectionQuadVS, TEXT("/Engine/Private/HeightfieldLighting.usf"), TEXT("HeightfieldSubsectionQuadVS"), SF_Vertex);
@@ -310,22 +308,15 @@ public:
 
 	void SetParameters(FRHICommandList& RHICmdList, const FSceneView& View, UTexture2D* HeightfieldTextureValue, UTexture2D* DiffuseColorTextureValue)
 	{
-		FRHIPixelShader* ShaderRHI = GetPixelShader();
+		FRHIPixelShader* ShaderRHI = RHICmdList.GetBoundPixelShader();
 
 		FGlobalShader::SetParameters<FViewUniformShaderParameters>(RHICmdList, ShaderRHI, View.ViewUniformBuffer);
 		HeightfieldTextureParameters.Set(RHICmdList, ShaderRHI, HeightfieldTextureValue, DiffuseColorTextureValue, nullptr);
 	}
-	// FShader interface.
-	virtual bool Serialize(FArchive& Ar) override
-	{
-		bool bShaderHasOutdatedParameters = FGlobalShader::Serialize(Ar);
-		Ar << HeightfieldTextureParameters;
-		return bShaderHasOutdatedParameters;
-	}
 
 private:
 
-	FHeightfieldTextureParameters HeightfieldTextureParameters;
+	LAYOUT_FIELD(FHeightfieldTextureParameters, HeightfieldTextureParameters);
 };
 
 IMPLEMENT_SHADER_TYPE(, FInitializeHeightfieldsPS, TEXT("/Engine/Private/HeightfieldLighting.usf"), TEXT("InitializeHeightfieldsPS"), SF_Pixel);
@@ -490,8 +481,8 @@ void FHeightfieldLightingViewInfo::SetupVisibleHeightfields(const FViewInfo& Vie
 						TShaderMapRef<FInitializeHeightfieldsPS> PixelShader(View.ShaderMap);
 
 						GraphicsPSOInit.BoundShaderState.VertexDeclarationRHI = GScreenVertexDeclaration.VertexDeclarationRHI;
-						GraphicsPSOInit.BoundShaderState.VertexShaderRHI = GETSAFERHISHADER_VERTEX(*VertexShader);
-						GraphicsPSOInit.BoundShaderState.PixelShaderRHI = GETSAFERHISHADER_PIXEL(*PixelShader);
+						GraphicsPSOInit.BoundShaderState.VertexShaderRHI = VertexShader.GetVertexShader();
+						GraphicsPSOInit.BoundShaderState.PixelShaderRHI = PixelShader.GetPixelShader();
 						GraphicsPSOInit.PrimitiveType = PT_TriangleList;
 
 						SetGraphicsPipelineState(RHICmdList, GraphicsPSOInit);
@@ -664,6 +655,7 @@ void UploadHeightfieldDescriptions(const TArray<FHeightfieldComponentDescription
 
 class FGlobalHeightfieldParameters
 {
+	DECLARE_INLINE_TYPE_LAYOUT(FGlobalHeightfieldParameters, NonVirtual);
 public:
 	void Bind(const FShaderParameterMap& ParameterMap)
 	{
@@ -691,10 +683,12 @@ public:
 	}
 
 private:
-	FShaderResourceParameter GlobalHeightfieldTexture;
-	FShaderResourceParameter GlobalNormalTexture;
-	FShaderResourceParameter GlobalDiffuseColorTexture;
-	FShaderResourceParameter GlobalHeightfieldSampler;
+	
+		LAYOUT_FIELD(FShaderResourceParameter, GlobalHeightfieldTexture)
+		LAYOUT_FIELD(FShaderResourceParameter, GlobalNormalTexture)
+		LAYOUT_FIELD(FShaderResourceParameter, GlobalDiffuseColorTexture)
+		LAYOUT_FIELD(FShaderResourceParameter, GlobalHeightfieldSampler)
+	
 };
 
 class FHeightfieldComponentQuadVS : public FGlobalShader
@@ -719,23 +713,14 @@ public:
 
 	void SetParameters(FRHICommandList& RHICmdList, const FSceneView& View, int32 NumHeightfieldsValue)
 	{
-		FRHIVertexShader* ShaderRHI = GetVertexShader();
+		FRHIVertexShader* ShaderRHI = RHICmdList.GetBoundVertexShader();
 
 		FGlobalShader::SetParameters<FViewUniformShaderParameters>(RHICmdList, ShaderRHI, View.ViewUniformBuffer);
 		HeightfieldDescriptionParameters.Set(RHICmdList, ShaderRHI, GetHeightfieldDescriptionsSRV(), NumHeightfieldsValue);
 	}
 
-	// FShader interface.
-	virtual bool Serialize(FArchive& Ar) override
-	{
-		bool bShaderHasOutdatedParameters = FGlobalShader::Serialize(Ar);
-		Ar << HeightfieldDescriptionParameters;
-		return bShaderHasOutdatedParameters;
-	}
-
 private:
-
-	FHeightfieldDescriptionParameters HeightfieldDescriptionParameters;
+	LAYOUT_FIELD(FHeightfieldDescriptionParameters, HeightfieldDescriptionParameters)
 };
 
 IMPLEMENT_SHADER_TYPE(, FHeightfieldComponentQuadVS, TEXT("/Engine/Private/HeightfieldLighting.usf"), TEXT("HeightfieldComponentQuadVS"), SF_Vertex);
@@ -769,7 +754,7 @@ public:
 
 	void SetParameters(FRHICommandList& RHICmdList, const FSceneView& View, const FProjectedShadowInfo* ProjectedShadowInfo, int32 NumHeightfieldsValue, const FHeightfieldLightingAtlas& Atlas)
 	{
-		FRHIPixelShader* ShaderRHI = GetPixelShader();
+		FRHIPixelShader* ShaderRHI = RHICmdList.GetBoundPixelShader();
 
 		FGlobalShader::SetParameters<FViewUniformShaderParameters>(RHICmdList, ShaderRHI, View.ViewUniformBuffer);
 		HeightfieldDescriptionParameters.Set(RHICmdList, ShaderRHI, GetHeightfieldDescriptionsSRV(), NumHeightfieldsValue);
@@ -788,31 +773,17 @@ public:
 
 		SetTextureParameter(RHICmdList, ShaderRHI, ShadowDepthTexture, ShadowDepthTextureSampler, DepthSamplerState, ShadowDepthTextureValue);
 	}
-	// FShader interface.
-	virtual bool Serialize(FArchive& Ar) override
-	{
-		bool bShaderHasOutdatedParameters = FGlobalShader::Serialize(Ar);
-		Ar << HeightfieldDescriptionParameters;
-		Ar << GlobalHeightfieldParameters;
-		Ar << WorldToShadow;
-		Ar << ShadowmapMinMax;
-		Ar << ShadowDepthBias;
-		Ar << CascadeDepthMinMax;
-		Ar << ShadowDepthTexture;
-		Ar << ShadowDepthTextureSampler;
-		return bShaderHasOutdatedParameters;
-	}
 
 private:
 
-	FHeightfieldDescriptionParameters HeightfieldDescriptionParameters;
-	FGlobalHeightfieldParameters GlobalHeightfieldParameters;
-	FShaderParameter WorldToShadow;
-	FShaderParameter ShadowmapMinMax;
-	FShaderParameter ShadowDepthBias;
-	FShaderParameter CascadeDepthMinMax;
-	FShaderResourceParameter ShadowDepthTexture;
-	FShaderResourceParameter ShadowDepthTextureSampler;
+	LAYOUT_FIELD(FHeightfieldDescriptionParameters, HeightfieldDescriptionParameters);
+	LAYOUT_FIELD(FGlobalHeightfieldParameters, GlobalHeightfieldParameters);
+	LAYOUT_FIELD(FShaderParameter, WorldToShadow);
+	LAYOUT_FIELD(FShaderParameter, ShadowmapMinMax);
+	LAYOUT_FIELD(FShaderParameter, ShadowDepthBias);
+	LAYOUT_FIELD(FShaderParameter, CascadeDepthMinMax);
+	LAYOUT_FIELD(FShaderResourceParameter, ShadowDepthTexture);
+	LAYOUT_FIELD(FShaderResourceParameter, ShadowDepthTextureSampler);
 };
 
 IMPLEMENT_SHADER_TYPE(, FShadowHeightfieldsPS, TEXT("/Engine/Private/HeightfieldLighting.usf"), TEXT("ShadowHeightfieldsPS"), SF_Pixel);
@@ -870,8 +841,8 @@ void FHeightfieldLightingViewInfo::ComputeShadowMapShadowing(const FViewInfo& Vi
 				TShaderMapRef<FShadowHeightfieldsPS> PixelShader(View.ShaderMap);
 
 				GraphicsPSOInit.BoundShaderState.VertexDeclarationRHI = GScreenVertexDeclaration.VertexDeclarationRHI;
-				GraphicsPSOInit.BoundShaderState.VertexShaderRHI = GETSAFERHISHADER_VERTEX(*VertexShader);
-				GraphicsPSOInit.BoundShaderState.PixelShaderRHI = GETSAFERHISHADER_PIXEL(*PixelShader);
+				GraphicsPSOInit.BoundShaderState.VertexShaderRHI = VertexShader.GetVertexShader();
+				GraphicsPSOInit.BoundShaderState.PixelShaderRHI = PixelShader.GetPixelShader();
 				GraphicsPSOInit.PrimitiveType = PT_TriangleList;
 
 				SetGraphicsPipelineState(RHICmdList, GraphicsPSOInit);
@@ -940,7 +911,7 @@ public:
 		FLightTileIntersectionResources* TileIntersectionResources,
 		FDistanceFieldObjectBufferResource& CulledObjectBuffers)
 	{
-		FRHIPixelShader* ShaderRHI = GetPixelShader();
+		FRHIPixelShader* ShaderRHI = RHICmdList.GetBoundPixelShader();
 
 		FGlobalShader::SetParameters<FViewUniformShaderParameters>(RHICmdList, ShaderRHI, View.ViewUniformBuffer);
 		HeightfieldDescriptionParameters.Set(RHICmdList, ShaderRHI, GetHeightfieldDescriptionsSRV(), NumHeightfieldsValue);
@@ -978,33 +949,18 @@ public:
 		extern float GTwoSidedMeshDistanceBias;
 		SetShaderValue(RHICmdList, ShaderRHI, TwoSidedMeshDistanceBias, GTwoSidedMeshDistanceBias);
 	}
-	// FShader interface.
-	virtual bool Serialize(FArchive& Ar) override
-	{
-		bool bShaderHasOutdatedParameters = FGlobalShader::Serialize(Ar);
-		Ar << HeightfieldDescriptionParameters;
-		Ar << GlobalHeightfieldParameters;
-		Ar << ObjectParameters;
-		Ar << LightDirection;
-		Ar << TanLightAngle;
-		Ar << CascadeDepthMinMax;
-		Ar << LightTileIntersectionParameters;
-		Ar << WorldToShadow;
-		Ar << TwoSidedMeshDistanceBias;
-		return bShaderHasOutdatedParameters;
-	}
 
 private:
 
-	FHeightfieldDescriptionParameters HeightfieldDescriptionParameters;
-	FGlobalHeightfieldParameters GlobalHeightfieldParameters;
-	TDistanceFieldCulledObjectBufferParameters<DFPT_SignedDistanceField> ObjectParameters;
-	FShaderParameter LightDirection;
-	FShaderParameter TanLightAngle;
-	FShaderParameter CascadeDepthMinMax;
-	FLightTileIntersectionParameters LightTileIntersectionParameters;
-	FShaderParameter WorldToShadow;
-	FShaderParameter TwoSidedMeshDistanceBias;
+	LAYOUT_FIELD(FHeightfieldDescriptionParameters, HeightfieldDescriptionParameters);
+	LAYOUT_FIELD(FGlobalHeightfieldParameters, GlobalHeightfieldParameters);
+	LAYOUT_FIELD((TDistanceFieldCulledObjectBufferParameters<DFPT_SignedDistanceField>), ObjectParameters);
+	LAYOUT_FIELD(FShaderParameter, LightDirection);
+	LAYOUT_FIELD(FShaderParameter, TanLightAngle);
+	LAYOUT_FIELD(FShaderParameter, CascadeDepthMinMax);
+	LAYOUT_FIELD(FLightTileIntersectionParameters, LightTileIntersectionParameters);
+	LAYOUT_FIELD(FShaderParameter, WorldToShadow);
+	LAYOUT_FIELD(FShaderParameter, TwoSidedMeshDistanceBias);
 };
 
 IMPLEMENT_SHADER_TYPE(, FRayTracedShadowHeightfieldsPS, TEXT("/Engine/Private/HeightfieldLighting.usf"), TEXT("RayTracedShadowHeightfieldsPS"), SF_Pixel);
@@ -1051,8 +1007,8 @@ void FHeightfieldLightingViewInfo::ComputeRayTracedShadowing(
 				TShaderMapRef<FRayTracedShadowHeightfieldsPS> PixelShader(View.ShaderMap);
 
 				GraphicsPSOInit.BoundShaderState.VertexDeclarationRHI = GScreenVertexDeclaration.VertexDeclarationRHI;
-				GraphicsPSOInit.BoundShaderState.VertexShaderRHI = GETSAFERHISHADER_VERTEX(*VertexShader);
-				GraphicsPSOInit.BoundShaderState.PixelShaderRHI = GETSAFERHISHADER_PIXEL(*PixelShader);
+				GraphicsPSOInit.BoundShaderState.VertexShaderRHI = VertexShader.GetVertexShader();
+				GraphicsPSOInit.BoundShaderState.PixelShaderRHI = PixelShader.GetPixelShader();
 				GraphicsPSOInit.PrimitiveType = PT_TriangleList;
 
 				SetGraphicsPipelineState(RHICmdList, GraphicsPSOInit);
@@ -1084,7 +1040,7 @@ public:
 
 	static bool ShouldCompilePermutation(const FMaterialShaderPermutationParameters& Parameters)
 	{
-		return Parameters.Material->IsLightFunction() && IsFeatureLevelSupported(Parameters.Platform, ERHIFeatureLevel::SM5) && DoesPlatformSupportDistanceFieldGI(Parameters.Platform);
+		return Parameters.MaterialParameters.MaterialDomain == MD_LightFunction && IsFeatureLevelSupported(Parameters.Platform, ERHIFeatureLevel::SM5) && DoesPlatformSupportDistanceFieldGI(Parameters.Platform);
 	}
 
 	static void ModifyCompilationEnvironment(const FMaterialShaderPermutationParameters& Parameters, FShaderCompilerEnvironment& OutEnvironment)
@@ -1119,7 +1075,7 @@ public:
 		const FHeightfieldLightingAtlas& Atlas,
 		float SkyLightIndirectScaleValue)
 	{
-		FRHIPixelShader* ShaderRHI = GetPixelShader();
+		FRHIPixelShader* ShaderRHI = RHICmdList.GetBoundPixelShader();
 
 		FMaterialShader::SetParameters(RHICmdList, ShaderRHI, MaterialProxy, *MaterialProxy->GetMaterial(View.GetFeatureLevel()), View, View.ViewUniformBuffer, ESceneTextureSetupMode::None);
 
@@ -1144,33 +1100,17 @@ public:
 		LightFunctionParameters.Set(RHICmdList, ShaderRHI, &LightSceneInfo, 1);
 	}
 
-	// FShader interface.
-	virtual bool Serialize(FArchive& Ar) override
-	{
-		bool bShaderHasOutdatedParameters = FMaterialShader::Serialize(Ar);
-		Ar << HeightfieldDescriptionParameters;
-		Ar << GlobalHeightfieldParameters;
-		Ar << LightDirection;
-		Ar << LightColor;
-		Ar << SkyLightIndirectScale;
-		Ar << HeightfieldShadowing;
-		Ar << HeightfieldShadowingSampler;
-		Ar << WorldToLight;
-		Ar << LightFunctionParameters;
-		return bShaderHasOutdatedParameters;
-	}
-
 private:
 
-	FHeightfieldDescriptionParameters HeightfieldDescriptionParameters;
-	FGlobalHeightfieldParameters GlobalHeightfieldParameters;
-	FShaderParameter LightDirection;
-	FShaderParameter LightColor;
-	FShaderParameter SkyLightIndirectScale;
-	FShaderResourceParameter HeightfieldShadowing;
-	FShaderResourceParameter HeightfieldShadowingSampler;
-	FShaderParameter WorldToLight;
-	FLightFunctionSharedParameters LightFunctionParameters;
+	LAYOUT_FIELD(FHeightfieldDescriptionParameters, HeightfieldDescriptionParameters);
+	LAYOUT_FIELD(FGlobalHeightfieldParameters, GlobalHeightfieldParameters);
+	LAYOUT_FIELD(FShaderParameter, LightDirection);
+	LAYOUT_FIELD(FShaderParameter, LightColor);
+	LAYOUT_FIELD(FShaderParameter, SkyLightIndirectScale);
+	LAYOUT_FIELD(FShaderResourceParameter, HeightfieldShadowing);
+	LAYOUT_FIELD(FShaderResourceParameter, HeightfieldShadowingSampler);
+	LAYOUT_FIELD(FShaderParameter, WorldToLight);
+	LAYOUT_FIELD(FLightFunctionSharedParameters, LightFunctionParameters);
 };
 
 IMPLEMENT_MATERIAL_SHADER_TYPE(, FLightHeightfieldsPS, TEXT("/Engine/Private/HeightfieldLighting.usf"), TEXT("LightHeightfieldsPS"), SF_Pixel);
@@ -1229,7 +1169,7 @@ void FHeightfieldLightingViewInfo::ComputeLighting(const FViewInfo& View, FRHICo
 
 				const FMaterial* Material = MaterialProxy->GetMaterial(FeatureLevel);
 				const FMaterialShaderMap* MaterialShaderMap = Material->GetRenderingThreadShaderMap();
-				FLightHeightfieldsPS* PixelShader = MaterialShaderMap->GetShader<FLightHeightfieldsPS>();
+				TShaderRef<FLightHeightfieldsPS> PixelShader = MaterialShaderMap->GetShader<FLightHeightfieldsPS>();
 
 				for (TMap<FHeightfieldComponentTextures, TArray<FHeightfieldComponentDescription>>::TConstIterator It(Heightfield.ComponentDescriptions); It; ++It)
 				{
@@ -1240,8 +1180,8 @@ void FHeightfieldLightingViewInfo::ComputeLighting(const FViewInfo& View, FRHICo
 						UploadHeightfieldDescriptions(HeightfieldDescriptions, InvLightingAtlasSize, 1.0f / Heightfield.DownsampleFactor);
 
 						GraphicsPSOInit.BoundShaderState.VertexDeclarationRHI = GScreenVertexDeclaration.VertexDeclarationRHI;
-						GraphicsPSOInit.BoundShaderState.VertexShaderRHI = GETSAFERHISHADER_VERTEX(*VertexShader);
-						GraphicsPSOInit.BoundShaderState.PixelShaderRHI = GETSAFERHISHADER_PIXEL(PixelShader);
+						GraphicsPSOInit.BoundShaderState.VertexShaderRHI = VertexShader.GetVertexShader();
+						GraphicsPSOInit.BoundShaderState.PixelShaderRHI = PixelShader.GetPixelShader();
 						GraphicsPSOInit.PrimitiveType = PT_TriangleList;
 
 						SetGraphicsPipelineState(RHICmdList, GraphicsPSOInit);
@@ -1306,7 +1246,7 @@ public:
 		const FAOScreenGridResources& ScreenGridResources,
 		const FDistanceFieldAOParameters& Parameters)
 	{
-		FRHIComputeShader* ShaderRHI = GetComputeShader();
+		FRHIComputeShader* ShaderRHI = RHICmdList.GetBoundComputeShader();
 		FGlobalShader::SetParameters<FViewUniformShaderParameters>(RHICmdList, ShaderRHI, View.ViewUniformBuffer);
 
 		AOParameters.Set(RHICmdList, ShaderRHI, Parameters);
@@ -1335,30 +1275,18 @@ public:
 
 	void UnsetParameters(FRHICommandList& RHICmdList, const FAOScreenGridResources& ScreenGridResources)
 	{
-		ScreenGridConeVisibility.UnsetUAV(RHICmdList, GetComputeShader());
+		ScreenGridConeVisibility.UnsetUAV(RHICmdList, RHICmdList.GetBoundComputeShader());
 		RHICmdList.TransitionResource(EResourceTransitionAccess::EReadable, EResourceTransitionPipeline::EComputeToCompute, ScreenGridResources.ScreenGridConeVisibility.UAV);
 	}
 
-	virtual bool Serialize(FArchive& Ar) override
-	{
-		bool bShaderHasOutdatedParameters = FGlobalShader::Serialize(Ar);
-		Ar << AOParameters;
-		Ar << ScreenGridParameters;
-		Ar << HeightfieldDescriptionParameters;
-		Ar << HeightfieldTextureParameters;
-		Ar << ScreenGridConeVisibility;
-		Ar << TanConeHalfAngle;
-		return bShaderHasOutdatedParameters;
-	}
-
 private:
 
-	FAOParameters AOParameters;
-	FScreenGridParameters ScreenGridParameters;
-	FHeightfieldDescriptionParameters HeightfieldDescriptionParameters;
-	FHeightfieldTextureParameters HeightfieldTextureParameters;
-	FRWShaderParameter ScreenGridConeVisibility;
-	FShaderParameter TanConeHalfAngle;
+	LAYOUT_FIELD(FAOParameters, AOParameters);
+	LAYOUT_FIELD(FScreenGridParameters, ScreenGridParameters);
+	LAYOUT_FIELD(FHeightfieldDescriptionParameters, HeightfieldDescriptionParameters);
+	LAYOUT_FIELD(FHeightfieldTextureParameters, HeightfieldTextureParameters);
+	LAYOUT_FIELD(FRWShaderParameter, ScreenGridConeVisibility);
+	LAYOUT_FIELD(FShaderParameter, TanConeHalfAngle);
 };
 
 IMPLEMENT_SHADER_TYPE(, FCalculateHeightfieldOcclusionScreenGridCS, TEXT("/Engine/Private/HeightfieldLighting.usf"), TEXT("CalculateHeightfieldOcclusionScreenGridCS"), SF_Compute);
@@ -1394,9 +1322,9 @@ void FHeightfieldLightingViewInfo::ComputeOcclusionForScreenGrid(
 					const uint32 GroupSizeY = FMath::DivideAndRoundUp(View.ViewRect.Size().Y / GAODownsampleFactor, GHeightfieldOcclusionDispatchSize);
 
 					TShaderMapRef<FCalculateHeightfieldOcclusionScreenGridCS> ComputeShader(View.ShaderMap);
-					RHICmdList.SetComputeShader(ComputeShader->GetComputeShader());
+					RHICmdList.SetComputeShader(ComputeShader.GetComputeShader());
 					ComputeShader->SetParameters(RHICmdList, View, HeightfieldTexture, HeightfieldDescriptions.Num(), DistanceFieldNormal, ScreenGridResources, Parameters);
-					DispatchComputeShader(RHICmdList, *ComputeShader, GroupSizeX, GroupSizeY, 1);
+					DispatchComputeShader(RHICmdList, ComputeShader.GetShader(), GroupSizeX, GroupSizeY, 1);
 					ComputeShader->UnsetParameters(RHICmdList, ScreenGridResources);
 				}
 			}
@@ -1457,7 +1385,7 @@ public:
 		const FDistanceFieldAOParameters& Parameters,
 		const FHeightfieldLightingAtlas& Atlas)
 	{
-		FRHIComputeShader* ShaderRHI = GetComputeShader();
+		FRHIComputeShader* ShaderRHI = RHICmdList.GetBoundComputeShader();
 		FGlobalShader::SetParameters<FViewUniformShaderParameters>(RHICmdList, ShaderRHI, View.ViewUniformBuffer);
 
 		AOParameters.Set(RHICmdList, ShaderRHI, Parameters);
@@ -1507,43 +1435,24 @@ public:
 
 	void UnsetParameters(FRHICommandList& RHICmdList, const FAOScreenGridResources& ScreenGridResources)
 	{
-		HeightfieldIrradiance.UnsetUAV(RHICmdList, GetComputeShader());
+		HeightfieldIrradiance.UnsetUAV(RHICmdList, RHICmdList.GetBoundComputeShader());
 		RHICmdList.TransitionResource(EResourceTransitionAccess::EReadable, EResourceTransitionPipeline::EComputeToCompute, ScreenGridResources.HeightfieldIrradiance.UAV);
 	}
 
-	virtual bool Serialize(FArchive& Ar) override
-	{
-		bool bShaderHasOutdatedParameters = FGlobalShader::Serialize(Ar);
-		Ar << AOParameters;
-		Ar << ScreenGridParameters;
-		Ar << HeightfieldDescriptionParameters;
-		Ar << HeightfieldIrradiance;
-		Ar << TanConeHalfAngle;
-		Ar << GlobalHeightfieldParameters;
-		Ar << BentNormalNormalizeFactor;
-		Ar << HeightfieldLighting;
-		Ar << HeightfieldLightingSampler;
-		Ar << InnerLightTransferDistance;
-		Ar << OuterLightTransferDistanceScale;
-		Ar << RecordConeVisibility;
-		return bShaderHasOutdatedParameters;
-	}
-
 private:
-
-	FAOParameters AOParameters;
-	FScreenGridParameters ScreenGridParameters;
-	FHeightfieldDescriptionParameters HeightfieldDescriptionParameters;
-	FRWShaderParameter HeightfieldIrradiance;
-	FShaderParameter TanConeHalfAngle;
-
-	FGlobalHeightfieldParameters GlobalHeightfieldParameters;
-	FShaderParameter BentNormalNormalizeFactor;
-	FShaderResourceParameter HeightfieldLighting;
-	FShaderResourceParameter HeightfieldLightingSampler;
-	FShaderParameter InnerLightTransferDistance;
-	FShaderParameter OuterLightTransferDistanceScale;
-	FShaderResourceParameter RecordConeVisibility;
+	LAYOUT_FIELD(FAOParameters, AOParameters);
+	LAYOUT_FIELD(FScreenGridParameters, ScreenGridParameters);
+	LAYOUT_FIELD(FHeightfieldDescriptionParameters, HeightfieldDescriptionParameters);
+	LAYOUT_FIELD(FRWShaderParameter, HeightfieldIrradiance);
+	LAYOUT_FIELD(FShaderParameter, TanConeHalfAngle);
+
+	LAYOUT_FIELD(FGlobalHeightfieldParameters, GlobalHeightfieldParameters);
+	LAYOUT_FIELD(FShaderParameter, BentNormalNormalizeFactor);
+	LAYOUT_FIELD(FShaderResourceParameter, HeightfieldLighting);
+	LAYOUT_FIELD(FShaderResourceParameter, HeightfieldLightingSampler);
+	LAYOUT_FIELD(FShaderParameter, InnerLightTransferDistance);
+	LAYOUT_FIELD(FShaderParameter, OuterLightTransferDistanceScale);
+	LAYOUT_FIELD(FShaderResourceParameter, RecordConeVisibility);
 };
 
 IMPLEMENT_SHADER_TYPE(, FCalculateHeightfieldIrradianceScreenGridCS, TEXT("/Engine/Private/HeightfieldLighting.usf"), TEXT("CalculateHeightfieldIrradianceScreenGridCS"), SF_Compute);
@@ -1585,13 +1494,13 @@ void FHeightfieldLightingViewInfo::ComputeIrradianceForScreenGrid(
 
 				{
 					TShaderMapRef<FCalculateHeightfieldIrradianceScreenGridCS> ComputeShader(View.ShaderMap);
-					RHICmdList.SetComputeShader(ComputeShader->GetComputeShader());
+					RHICmdList.SetComputeShader(ComputeShader.GetComputeShader());
 					ComputeShader->SetParameters(RHICmdList, View, CombinedHeightfieldDescriptions.Num(), DistanceFieldNormal, ScreenGridResources, Parameters, Atlas);
 
 					const uint32 GroupSizeX = View.ViewRect.Size().X / GAODownsampleFactor;
 					const uint32 GroupSizeY = View.ViewRect.Size().Y / GAODownsampleFactor;
 
-					DispatchComputeShader(RHICmdList, *ComputeShader, GroupSizeX, GroupSizeY, 1);
+					DispatchComputeShader(RHICmdList, ComputeShader.GetShader(), GroupSizeX, GroupSizeY, 1);
 
 					ComputeShader->UnsetParameters(RHICmdList, ScreenGridResources);
 				}
diff --git a/Engine/Source/Runtime/Renderer/Private/HeightfieldLighting.h b/Engine/Source/Runtime/Renderer/Private/HeightfieldLighting.h
--- a/Engine/Source/Runtime/Renderer/Private/HeightfieldLighting.h
+++ b/Engine/Source/Runtime/Renderer/Private/HeightfieldLighting.h
@@ -148,6 +148,7 @@ extern FRHIShaderResourceView* GetHeightfieldDescriptionsSRV();
 
 class FHeightfieldDescriptionParameters
 {
+	DECLARE_TYPE_LAYOUT(FHeightfieldDescriptionParameters, NonVirtual);
 public:
 	void Bind(const FShaderParameterMap& ParameterMap)
 	{
@@ -170,12 +171,15 @@ public:
 	}
 
 private:
-	FShaderResourceParameter HeightfieldDescriptions;
-	FShaderParameter NumHeightfields;
+	
+		LAYOUT_FIELD(FShaderResourceParameter, HeightfieldDescriptions)
+		LAYOUT_FIELD(FShaderParameter, NumHeightfields)
+	
 };
 
 class FHeightfieldTextureParameters
 {
+	DECLARE_TYPE_LAYOUT(FHeightfieldTextureParameters, NonVirtual);
 public:
 	void Bind(const FShaderParameterMap& ParameterMap)
 	{
@@ -224,10 +228,12 @@ public:
 	}
 
 private:
-	FShaderResourceParameter HeightfieldTexture;
-	FShaderResourceParameter HeightfieldSampler;
-	FShaderResourceParameter DiffuseColorTexture;
-	FShaderResourceParameter DiffuseColorSampler;
-	FShaderResourceParameter VisibilityTexture;
-	FShaderResourceParameter VisibilitySampler;
+	
+		LAYOUT_FIELD(FShaderResourceParameter, HeightfieldTexture)
+		LAYOUT_FIELD(FShaderResourceParameter, HeightfieldSampler)
+		LAYOUT_FIELD(FShaderResourceParameter, DiffuseColorTexture)
+		LAYOUT_FIELD(FShaderResourceParameter, DiffuseColorSampler)
+		LAYOUT_FIELD(FShaderResourceParameter, VisibilityTexture)
+		LAYOUT_FIELD(FShaderResourceParameter, VisibilitySampler)
+	
 };
diff --git a/Engine/Source/Runtime/Renderer/Private/IndirectLightRendering.cpp b/Engine/Source/Runtime/Renderer/Private/IndirectLightRendering.cpp
--- a/Engine/Source/Runtime/Renderer/Private/IndirectLightRendering.cpp
+++ b/Engine/Source/Runtime/Renderer/Private/IndirectLightRendering.cpp
@@ -486,7 +486,7 @@ void FDeferredShadingSceneRenderer::RenderDiffuseIndirectAndAmbientOcclusion(FRH
 			PermutationVector.Set<FDiffuseIndirectCompositePS::FApplyAmbientOcclusionDim>(PassParameters->AmbientOcclusionTexture != nullptr);
 
 			TShaderMapRef<FDiffuseIndirectCompositePS> PixelShader(View.ShaderMap, PermutationVector);
-			ClearUnusedGraphResources(*PixelShader, PassParameters);
+			ClearUnusedGraphResources(PixelShader, PassParameters);
 
 			GraphBuilder.AddPass(
 				RDG_EVENT_NAME(
@@ -501,7 +501,7 @@ void FDeferredShadingSceneRenderer::RenderDiffuseIndirectAndAmbientOcclusion(FRH
 				RHICmdList.SetViewport(View.ViewRect.Min.X, View.ViewRect.Min.Y, 0.0f, View.ViewRect.Max.X, View.ViewRect.Max.Y, 0.0);
 				
 				FGraphicsPipelineStateInitializer GraphicsPSOInit;
-				FPixelShaderUtils::InitFullscreenPipelineState(RHICmdList, View.ShaderMap, *PixelShader, /* out */ GraphicsPSOInit);
+				FPixelShaderUtils::InitFullscreenPipelineState(RHICmdList, View.ShaderMap, PixelShader, /* out */ GraphicsPSOInit);
 				
 				if (PermutationVector.Get<FDiffuseIndirectCompositePS::FApplyAmbientOcclusionDim>())
 				{
@@ -513,7 +513,7 @@ void FDeferredShadingSceneRenderer::RenderDiffuseIndirectAndAmbientOcclusion(FRH
 				}
 				SetGraphicsPipelineState(RHICmdList, GraphicsPSOInit);
 
-				SetShaderParameters(RHICmdList, *PixelShader, PixelShader->GetPixelShader(), *PassParameters);
+				SetShaderParameters(RHICmdList, PixelShader, PixelShader.GetPixelShader(), *PassParameters);
 
 				FPixelShaderUtils::DrawFullscreenTriangle(RHICmdList);
 			});
@@ -562,8 +562,8 @@ void FDeferredShadingSceneRenderer::RenderDiffuseIndirectAndAmbientOcclusion(FRH
 				GraphicsPSOInit.DepthStencilState = TStaticDepthStencilState<false, CF_Always>::GetRHI();
 
 				GraphicsPSOInit.BoundShaderState.VertexDeclarationRHI = GFilterVertexDeclaration.VertexDeclarationRHI;
-				GraphicsPSOInit.BoundShaderState.VertexShaderRHI = GETSAFERHISHADER_VERTEX(*VertexShader);
-				GraphicsPSOInit.BoundShaderState.PixelShaderRHI = GETSAFERHISHADER_PIXEL(*PixelShader);
+				GraphicsPSOInit.BoundShaderState.VertexShaderRHI = VertexShader.GetVertexShader();
+				GraphicsPSOInit.BoundShaderState.PixelShaderRHI = PixelShader.GetPixelShader();
 				GraphicsPSOInit.PrimitiveType = PT_TriangleList;
 
 				SetGraphicsPipelineState(RHICmdList, GraphicsPSOInit);
@@ -573,7 +573,7 @@ void FDeferredShadingSceneRenderer::RenderDiffuseIndirectAndAmbientOcclusion(FRH
 				{
 					FAmbientCubemapCompositePS::FParameters ShaderParameters = *PassParameters;
 					SetupAmbientCubemapParameters(CubemapEntry, &ShaderParameters.AmbientCubemap);
-					SetShaderParameters(RHICmdList, *PixelShader, PixelShader->GetPixelShader(), ShaderParameters);
+					SetShaderParameters(RHICmdList, PixelShader, PixelShader.GetPixelShader(), ShaderParameters);
 					
 					DrawPostProcessPass(
 						RHICmdList,
@@ -583,7 +583,7 @@ void FDeferredShadingSceneRenderer::RenderDiffuseIndirectAndAmbientOcclusion(FRH
 						View.ViewRect.Width(), View.ViewRect.Height(),
 						View.ViewRect.Size(),
 						FSceneRenderTargets::Get(RHICmdList).GetBufferSizeXY(),
-						*VertexShader,
+						VertexShader,
 						View.StereoPass, 
 						false, // TODO.
 						EDRF_UseTriangleOptimization);
@@ -881,7 +881,7 @@ void FDeferredShadingSceneRenderer::RenderDeferredReflectionsAndSkyLighting(FRHI
 				View, bHasBoxCaptures, bHasSphereCaptures, DynamicBentNormalAO != NULL, bSkyLight, bDynamicSkyLight, bApplySkyShadowing, bRayTracedReflections);
 
 			TShaderMapRef<FReflectionEnvironmentSkyLightingPS> PixelShader(View.ShaderMap, PermutationVector);
-			ClearUnusedGraphResources(*PixelShader, PassParameters);
+			ClearUnusedGraphResources(PixelShader, PassParameters);
 
 			GraphBuilder.AddPass(
 				RDG_EVENT_NAME("ReflectionEnvironmentAndSky %dx%d", View.ViewRect.Width(), View.ViewRect.Height()),
@@ -892,7 +892,7 @@ void FDeferredShadingSceneRenderer::RenderDeferredReflectionsAndSkyLighting(FRHI
 				InRHICmdList.SetViewport(View.ViewRect.Min.X, View.ViewRect.Min.Y, 0.0f, View.ViewRect.Max.X, View.ViewRect.Max.Y, 1.0f);
 
 				FGraphicsPipelineStateInitializer GraphicsPSOInit;
-				FPixelShaderUtils::InitFullscreenPipelineState(InRHICmdList, View.ShaderMap, *PixelShader, GraphicsPSOInit);
+				FPixelShaderUtils::InitFullscreenPipelineState(InRHICmdList, View.ShaderMap, PixelShader, GraphicsPSOInit);
 
 				extern int32 GAOOverwriteSceneColor;
 				if (GetReflectionEnvironmentCVar() == 2 || GAOOverwriteSceneColor)
@@ -913,7 +913,7 @@ void FDeferredShadingSceneRenderer::RenderDeferredReflectionsAndSkyLighting(FRHI
 				}
 
 				SetGraphicsPipelineState(InRHICmdList, GraphicsPSOInit);
-				SetShaderParameters(InRHICmdList, *PixelShader, PixelShader->GetPixelShader(), *PassParameters);
+				SetShaderParameters(InRHICmdList, PixelShader, PixelShader.GetPixelShader(), *PassParameters);
 				FPixelShaderUtils::DrawFullscreenTriangle(InRHICmdList);
 			});
 		} // if (bRequiresApply)
diff --git a/Engine/Source/Runtime/Renderer/Private/IndirectLightingCache.cpp b/Engine/Source/Runtime/Renderer/Private/IndirectLightingCache.cpp
--- a/Engine/Source/Runtime/Renderer/Private/IndirectLightingCache.cpp
+++ b/Engine/Source/Runtime/Renderer/Private/IndirectLightingCache.cpp
@@ -375,7 +375,7 @@ void FIndirectLightingCache::ReleasePrimitive(FPrimitiveComponentId PrimitiveId)
 	}
 }
 
-FIndirectLightingCacheAllocation* FIndirectLightingCache::FindPrimitiveAllocation(FPrimitiveComponentId PrimitiveId)
+FIndirectLightingCacheAllocation* FIndirectLightingCache::FindPrimitiveAllocation(FPrimitiveComponentId PrimitiveId) const
 {
 	return PrimitiveAllocations.FindRef(PrimitiveId);
 }
diff --git a/Engine/Source/Runtime/Renderer/Private/LightFunctionRendering.cpp b/Engine/Source/Runtime/Renderer/Private/LightFunctionRendering.cpp
--- a/Engine/Source/Runtime/Renderer/Private/LightFunctionRendering.cpp
+++ b/Engine/Source/Runtime/Renderer/Private/LightFunctionRendering.cpp
@@ -38,7 +38,7 @@ public:
 	  */
 	static bool ShouldCompilePermutation(const FMaterialShaderPermutationParameters& Parameters)
 	{
-		return Parameters.Material->IsLightFunction() && IsFeatureLevelSupported(Parameters.Platform, ERHIFeatureLevel::SM5);
+		return Parameters.MaterialParameters.MaterialDomain == MD_LightFunction && IsFeatureLevelSupported(Parameters.Platform, ERHIFeatureLevel::SM5);
 	}
 
 	FLightFunctionVS( )	{ }
@@ -50,7 +50,7 @@ public:
 
 	void SetParameters(FRHICommandList& RHICmdList, const FSceneView& View, const FLightSceneInfo* LightSceneInfo )
 	{
-		FMaterialShader::SetViewParameters(RHICmdList, GetVertexShader(), View, View.ViewUniformBuffer);
+		FMaterialShader::SetViewParameters(RHICmdList, RHICmdList.GetBoundVertexShader(), View, View.ViewUniformBuffer);
 		
 		// Light functions are projected using a bounding sphere.
 		// Calculate transform for bounding stencil sphere.
@@ -65,17 +65,8 @@ public:
 		StencilingGeometryParameters.Set(RHICmdList, this, StencilingSpherePosAndScale);
 	}
 
-	// Begin FShader interface
-	virtual bool Serialize(FArchive& Ar) override
-	{
-		bool bShaderHasOutdatedParameters = FMaterialShader::Serialize(Ar);
-		Ar << StencilingGeometryParameters;
-		return bShaderHasOutdatedParameters;
-	}
-	//  End FShader interface 
-
 private:
-	FStencilingGeometryShaderParameters StencilingGeometryParameters;
+	LAYOUT_FIELD(FStencilingGeometryShaderParameters, StencilingGeometryParameters);
 };
 
 IMPLEMENT_MATERIAL_SHADER_TYPE(,FLightFunctionVS,TEXT("/Engine/Private/LightFunctionVertexShader.usf"),TEXT("Main"),SF_Vertex);
@@ -95,7 +86,7 @@ public:
 	  */
 	static bool ShouldCompilePermutation(const FMaterialShaderPermutationParameters& Parameters)
 	{
-		return Parameters.Material->IsLightFunction() && IsFeatureLevelSupported(Parameters.Platform, ERHIFeatureLevel::SM5);
+		return Parameters.MaterialParameters.MaterialDomain == MD_LightFunction && IsFeatureLevelSupported(Parameters.Platform, ERHIFeatureLevel::SM5);
 	}
 
 	FLightFunctionPS() {}
@@ -110,7 +101,7 @@ public:
 
 	void SetParameters(FRHICommandList& RHICmdList, const FViewInfo& View, const FLightSceneInfo* LightSceneInfo, const FMaterialRenderProxy* MaterialProxy, bool bRenderingPreviewShadowIndicator, float ShadowFadeFraction )
 	{
-		FRHIPixelShader* ShaderRHI = GetPixelShader();
+		FRHIPixelShader* ShaderRHI = RHICmdList.GetBoundPixelShader();
 
 		FMaterialShader::SetParameters(RHICmdList, ShaderRHI, MaterialProxy, *MaterialProxy->GetMaterial(View.GetFeatureLevel()), View, View.ViewUniformBuffer, ESceneTextureSetupMode::All);
 
@@ -164,21 +155,11 @@ public:
 		}
 	}
 
-	virtual bool Serialize(FArchive& Ar) override
-	{
-		bool bShaderHasOutdatedParameters = FMaterialShader::Serialize(Ar);
-		Ar << SvPositionToLight;
-		Ar << LightFunctionParameters;
-		Ar << LightFunctionParameters2;
-		Ar << SceneTextureParameters;
-		return bShaderHasOutdatedParameters;
-	}
-
 private:
-	FShaderParameter SvPositionToLight;
-	FLightFunctionSharedParameters LightFunctionParameters;
-	FShaderParameter LightFunctionParameters2;
-	FSceneTextureShaderParameters SceneTextureParameters;
+	LAYOUT_FIELD(FShaderParameter, SvPositionToLight);
+	LAYOUT_FIELD(FLightFunctionSharedParameters, LightFunctionParameters);
+	LAYOUT_FIELD(FShaderParameter, LightFunctionParameters2);
+	LAYOUT_FIELD(FSceneTextureShaderParameters, SceneTextureParameters);
 };
 
 IMPLEMENT_MATERIAL_SHADER_TYPE(,FLightFunctionPS,TEXT("/Engine/Private/LightFunctionPixelShader.usf"),TEXT("Main"),SF_Pixel);
@@ -302,15 +283,15 @@ bool FDeferredShadingSceneRenderer::RenderLightFunctionForMaterial(
 			SCOPED_DRAW_EVENTF(RHICmdList, LightFunction, TEXT("LightFunction Material=%s"), *Material->GetFriendlyName());
 
 			const FMaterialShaderMap* MaterialShaderMap = Material->GetRenderingThreadShaderMap();
-			FLightFunctionVS* VertexShader = MaterialShaderMap->GetShader<FLightFunctionVS>();
-			FLightFunctionPS* PixelShader = MaterialShaderMap->GetShader<FLightFunctionPS>();
+			TShaderRef<FLightFunctionVS> VertexShader = MaterialShaderMap->GetShader<FLightFunctionVS>();
+			TShaderRef<FLightFunctionPS> PixelShader = MaterialShaderMap->GetShader<FLightFunctionPS>();
 
 			FGraphicsPipelineStateInitializer GraphicsPSOInit;
 			RHICmdList.ApplyCachedRenderTargets(GraphicsPSOInit);
 
 			GraphicsPSOInit.BoundShaderState.VertexDeclarationRHI = GetVertexDeclarationFVector4();
-			GraphicsPSOInit.BoundShaderState.VertexShaderRHI = VertexShader->GetVertexShader();
-			GraphicsPSOInit.BoundShaderState.PixelShaderRHI = PixelShader->GetPixelShader();
+			GraphicsPSOInit.BoundShaderState.VertexShaderRHI = VertexShader.GetVertexShader();
+			GraphicsPSOInit.BoundShaderState.PixelShaderRHI = PixelShader.GetPixelShader();
 
 			FSphere LightBounds = LightSceneInfo->Proxy->GetBoundingSphere();
 
diff --git a/Engine/Source/Runtime/Renderer/Private/LightGridInjection.cpp b/Engine/Source/Runtime/Renderer/Private/LightGridInjection.cpp
--- a/Engine/Source/Runtime/Renderer/Private/LightGridInjection.cpp
+++ b/Engine/Source/Runtime/Renderer/Private/LightGridInjection.cpp
@@ -711,7 +711,7 @@ void FDeferredShadingSceneRenderer::ComputeLightGrid(FRHICommandListImmediate& R
 						AddClearUAVPass(GraphBuilder, PassParameters->RWStartOffsetGrid, 0xFFFFFFFF);
 						AddClearUAVPass(GraphBuilder, PassParameters->RWNextCulledLightLink, 0);
 						AddClearUAVPass(GraphBuilder, GraphBuilder.CreateUAV(NextCulledLightDataBuffer, PF_R32_UINT), 0);
-						FComputeShaderUtils::AddPass(GraphBuilder, RDG_EVENT_NAME("LightGridInject:LinkedList"), *ComputeShader, PassParameters, NumGroups);
+						FComputeShaderUtils::AddPass(GraphBuilder, RDG_EVENT_NAME("LightGridInject:LinkedList"), ComputeShader, PassParameters, NumGroups);
 
 
 						{
@@ -726,13 +726,13 @@ void FDeferredShadingSceneRenderer::ComputeLightGrid(FRHICommandListImmediate& R
 							PassParametersCompact->RWNextCulledLightData = GraphBuilder.CreateUAV(NextCulledLightDataBuffer, PF_R32_UINT);
 							PassParametersCompact->StartOffsetGrid = GraphBuilder.CreateSRV(StartOffsetGridBuffer, PF_R32_UINT);
 
-							FComputeShaderUtils::AddPass(GraphBuilder, RDG_EVENT_NAME("CompactLinks"), *ComputeShaderCompact, PassParametersCompact, NumGroups);
+							FComputeShaderUtils::AddPass(GraphBuilder, RDG_EVENT_NAME("CompactLinks"), ComputeShaderCompact, PassParametersCompact, NumGroups);
 						}
 					}
 					else
 					{
 						RHICmdList.ClearUAVUint(View.ForwardLightingResources->NumCulledLightsGrid.UAV, FUintVector4(0, 0, 0, 0));
-						FComputeShaderUtils::AddPass(GraphBuilder, RDG_EVENT_NAME("LightGridInject:NotLinkedList"), *ComputeShader, PassParameters, NumGroups);
+						FComputeShaderUtils::AddPass(GraphBuilder, RDG_EVENT_NAME("LightGridInject:NotLinkedList"), ComputeShader, PassParameters, NumGroups);
 					}
 				}
 				GraphBuilder.Execute();
diff --git a/Engine/Source/Runtime/Renderer/Private/LightMapDensityRendering.cpp b/Engine/Source/Runtime/Renderer/Private/LightMapDensityRendering.cpp
--- a/Engine/Source/Runtime/Renderer/Private/LightMapDensityRendering.cpp
+++ b/Engine/Source/Runtime/Renderer/Private/LightMapDensityRendering.cpp
@@ -240,8 +240,9 @@ void FLightmapDensityMeshProcessor::AddMeshBatch(const FMeshBatch& RESTRICT Mesh
 		const bool bMaterialMasked = Material->IsMasked();
 		const bool bTranslucentBlendMode = IsTranslucentBlendMode(Material->GetBlendMode());
 		const bool bIsLitMaterial = Material->GetShadingModels().IsLit();
-		const ERasterizerFillMode MeshFillMode = ComputeMeshFillMode(MeshBatch, *Material);
-		const ERasterizerCullMode MeshCullMode = ComputeMeshCullMode(MeshBatch, *Material);
+		const FMeshDrawingPolicyOverrideSettings OverrideSettings = ComputeMeshOverrideSettings(MeshBatch);
+		const ERasterizerFillMode MeshFillMode = ComputeMeshFillMode(MeshBatch, *Material, OverrideSettings);
+		const ERasterizerCullMode MeshCullMode = ComputeMeshCullMode(MeshBatch, *Material, OverrideSettings);
 		const FLightMapInteraction LightMapInteraction = (MeshBatch.LCI && bIsLitMaterial) ? MeshBatch.LCI->GetLightMapInteraction(FeatureLevel) : FLightMapInteraction();
 
 		// Force simple lightmaps based on system settings.
diff --git a/Engine/Source/Runtime/Renderer/Private/LightMapDensityRendering.h b/Engine/Source/Runtime/Renderer/Private/LightMapDensityRendering.h
--- a/Engine/Source/Runtime/Renderer/Private/LightMapDensityRendering.h
+++ b/Engine/Source/Runtime/Renderer/Private/LightMapDensityRendering.h
@@ -52,14 +52,14 @@ public:
 template<typename LightMapPolicyType>
 class TLightMapDensityVS : public FMeshMaterialShader, public LightMapPolicyType::VertexParametersType
 {
-	DECLARE_SHADER_TYPE(TLightMapDensityVS,MeshMaterial);
+	DECLARE_SHADER_TYPE_EXPLICIT_BASES(TLightMapDensityVS,MeshMaterial, FMeshMaterialShader, typename LightMapPolicyType::VertexParametersType);
 
 public:
 
 	static bool ShouldCompilePermutation(const FMeshMaterialShaderPermutationParameters& Parameters)
 	{
 		return  AllowDebugViewmodes(Parameters.Platform) 
-				&& (Parameters.Material->IsSpecialEngineMaterial() || Parameters.Material->IsMasked() || Parameters.Material->MaterialMayModifyMeshPosition())
+				&& (Parameters.MaterialParameters.bIsSpecialEngineMaterial || Parameters.MaterialParameters.bIsMasked || Parameters.MaterialParameters.bMaterialMayModifyMeshPosition)
 				&& LightMapPolicyType::ShouldCompilePermutation(Parameters)
 				&& IsFeatureLevelSupported(Parameters.Platform, ERHIFeatureLevel::SM5);
 	}
@@ -79,13 +79,6 @@ public:
 	}
 	TLightMapDensityVS() {}
 
-	virtual bool Serialize(FArchive& Ar) override
-	{
-		bool bShaderHasOutdatedParameters = FMeshMaterialShader::Serialize(Ar);
-		LightMapPolicyType::VertexParametersType::Serialize(Ar);
-		return bShaderHasOutdatedParameters;
-	}
-
 	void GetShaderBindings(
 		const FScene* Scene,
 		ERHIFeatureLevel::Type FeatureLevel,
@@ -175,14 +168,14 @@ public:
 template<typename LightMapPolicyType>
 class TLightMapDensityPS : public FMeshMaterialShader, public LightMapPolicyType::PixelParametersType
 {
-	DECLARE_SHADER_TYPE(TLightMapDensityPS,MeshMaterial);
+	DECLARE_SHADER_TYPE_EXPLICIT_BASES(TLightMapDensityPS,MeshMaterial, FMeshMaterialShader, typename LightMapPolicyType::PixelParametersType);
 
 public:
 
 	static bool ShouldCompilePermutation(const FMeshMaterialShaderPermutationParameters& Parameters)
 	{
 		return	AllowDebugViewmodes(Parameters.Platform) 
-				&& (Parameters.Material->IsSpecialEngineMaterial() || Parameters.Material->IsMasked() || Parameters.Material->MaterialMayModifyMeshPosition())
+				&& (Parameters.MaterialParameters.bIsSpecialEngineMaterial || Parameters.MaterialParameters.bIsMasked || Parameters.MaterialParameters.bMaterialMayModifyMeshPosition)
 				&& LightMapPolicyType::ShouldCompilePermutation(Parameters)
 				&& IsFeatureLevelSupported(Parameters.Platform, ERHIFeatureLevel::SM5);
 	}
@@ -235,20 +228,10 @@ public:
 		ShaderBindings.Add(LightMapDensityDisplayOptions, OptionsParameter);
 	}
 
-	virtual bool Serialize(FArchive& Ar) override
-	{
-		bool bShaderHasOutdatedParameters = FMeshMaterialShader::Serialize(Ar);
-		LightMapPolicyType::PixelParametersType::Serialize(Ar);
-		Ar << BuiltLightingAndSelectedFlags;
-		Ar << LightMapResolutionScale;
-		Ar << LightMapDensityDisplayOptions;
-		return bShaderHasOutdatedParameters;
-	}
-
 private:
-	FShaderParameter BuiltLightingAndSelectedFlags;
-	FShaderParameter LightMapResolutionScale;
-	FShaderParameter LightMapDensityDisplayOptions;
+	LAYOUT_FIELD(FShaderParameter, BuiltLightingAndSelectedFlags);
+	LAYOUT_FIELD(FShaderParameter, LightMapResolutionScale);
+	LAYOUT_FIELD(FShaderParameter, LightMapDensityDisplayOptions);
 };
 
 
diff --git a/Engine/Source/Runtime/Renderer/Private/LightMapRendering.cpp b/Engine/Source/Runtime/Renderer/Private/LightMapRendering.cpp
--- a/Engine/Source/Runtime/Renderer/Private/LightMapRendering.cpp
+++ b/Engine/Source/Runtime/Renderer/Private/LightMapRendering.cpp
@@ -12,6 +12,8 @@ LightMapRendering.cpp: Light map rendering implementations.
 #include "Runtime/Renderer/Private/VT/VirtualTextureSpace.h"
 #include "VT/VirtualTextureSpace.h"
 
+IMPLEMENT_TYPE_LAYOUT(FUniformLightMapPolicyShaderParametersType);
+
 IMPLEMENT_GLOBAL_SHADER_PARAMETER_STRUCT(FIndirectLightingCacheUniformParameters, "IndirectLightingCache");
 
 const TCHAR* GLightmapDefineName[2] =
diff --git a/Engine/Source/Runtime/Renderer/Private/LightMapRendering.h b/Engine/Source/Runtime/Renderer/Private/LightMapRendering.h
--- a/Engine/Source/Runtime/Renderer/Private/LightMapRendering.h
+++ b/Engine/Source/Runtime/Renderer/Private/LightMapRendering.h
@@ -131,11 +131,11 @@ struct TLightMapPolicy
 		const bool bShouldCacheQuality = (LightmapQuality != ELightmapQuality::LQ_LIGHTMAP) || bProjectCanHaveLowQualityLightmaps;
 
 		// GetValueOnAnyThread() as it's possible that ShouldCache is called from rendering thread. That is to output some error message.
-		return (Parameters.Material->GetShadingModels().IsLit())
+		return (Parameters.MaterialParameters.ShadingModels.IsLit())
 			&& bShouldCacheQuality
 			&& Parameters.VertexFactoryType->SupportsStaticLighting()
 			&& (!AllowStaticLightingVar || AllowStaticLightingVar->GetValueOnAnyThread() != 0)
-			&& (Parameters.Material->IsUsedWithStaticLighting() || Parameters.Material->IsSpecialEngineMaterial());
+			&& (Parameters.MaterialParameters.bIsUsedWithStaticLighting || Parameters.MaterialParameters.bIsSpecialEngineMaterial);
 	}
 
 	static bool RequiresSkylight()
@@ -167,7 +167,7 @@ public:
 
 	static bool ShouldCompilePermutation(const FMeshMaterialShaderPermutationParameters& Parameters)
 	{
-		return Parameters.Material->GetShadingModels().IsLit() && Parameters.VertexFactoryType->SupportsStaticLighting();
+		return Parameters.MaterialParameters.ShadingModels.IsLit() && Parameters.VertexFactoryType->SupportsStaticLighting();
 	}
 };
 
@@ -182,6 +182,7 @@ public:
 
 	class VertexParametersType
 	{
+		DECLARE_INLINE_TYPE_LAYOUT(VertexParametersType, NonVirtual);
 	public:
 		void Bind(const FShaderParameterMap& ParameterMap) {}
 		void Serialize(FArchive& Ar) {}
@@ -189,6 +190,7 @@ public:
 
 	class PixelParametersType
 	{
+		DECLARE_INLINE_TYPE_LAYOUT(PixelParametersType, NonVirtual);
 	public:
 		void Bind(const FShaderParameterMap& ParameterMap)
 		{
@@ -200,13 +202,13 @@ public:
 			Ar << TranslucentSelfShadowBufferParameter;
 		}
 
-		FShaderUniformBufferParameter TranslucentSelfShadowBufferParameter;
+		LAYOUT_FIELD(FShaderUniformBufferParameter, TranslucentSelfShadowBufferParameter);
 	};
 
 	static bool ShouldCompilePermutation(const FMeshMaterialShaderPermutationParameters& Parameters)
 	{
-		return Parameters.Material->GetShadingModels().IsLit() &&
-			IsTranslucentBlendMode(Parameters.Material->GetBlendMode()) &&
+		return Parameters.MaterialParameters.ShadingModels.IsLit() &&
+			IsTranslucentBlendMode(Parameters.MaterialParameters.BlendMode) &&
 			IsFeatureLevelSupported(Parameters.Platform, ERHIFeatureLevel::SM5);
 	}
 
@@ -253,7 +255,7 @@ struct FPrecomputedVolumetricLightmapLightingPolicy
 	{
 		static const auto AllowStaticLightingVar = IConsoleManager::Get().FindTConsoleVariableDataInt(TEXT("r.AllowStaticLighting"));
 	
-		return Parameters.Material->GetShadingModels().IsLit() &&
+		return Parameters.MaterialParameters.ShadingModels.IsLit() &&
 			(!AllowStaticLightingVar || AllowStaticLightingVar->GetValueOnAnyThread() != 0);
 	}
 
@@ -272,8 +274,8 @@ struct FCachedVolumeIndirectLightingPolicy
 	{
 		static const auto AllowStaticLightingVar = IConsoleManager::Get().FindTConsoleVariableDataInt(TEXT("r.AllowStaticLighting"));
 
-		return Parameters.Material->GetShadingModels().IsLit()
-			&& !IsTranslucentBlendMode(Parameters.Material->GetBlendMode())
+		return Parameters.MaterialParameters.ShadingModels.IsLit()
+			&& !IsTranslucentBlendMode(Parameters.MaterialParameters.BlendMode)
 			&& (!AllowStaticLightingVar || AllowStaticLightingVar->GetValueOnAnyThread() != 0)
 			&& IsFeatureLevelSupported(Parameters.Platform, ERHIFeatureLevel::SM5);
 	}
@@ -299,7 +301,7 @@ struct FCachedPointIndirectLightingPolicy
 	{
 		static const auto AllowStaticLightingVar = IConsoleManager::Get().FindTConsoleVariableDataInt(TEXT("r.AllowStaticLighting"));
 	
-		return Parameters.Material->GetShadingModels().IsLit()
+		return Parameters.MaterialParameters.ShadingModels.IsLit()
 			&& (!AllowStaticLightingVar || AllowStaticLightingVar->GetValueOnAnyThread() != 0);
 	}
 
@@ -368,7 +370,7 @@ struct FSimpleDirectionalLightLightingPolicy
 	static bool ShouldCompilePermutation(const FMeshMaterialShaderPermutationParameters& Parameters)
 	{
 		return PlatformSupportsSimpleForwardShading(Parameters.Platform)
-			&& Parameters.Material->GetShadingModels().IsLit();
+			&& Parameters.MaterialParameters.ShadingModels.IsLit();
 	}
 
 	static void ModifyCompilationEnvironment(const FMaterialShaderPermutationParameters& Parameters, FShaderCompilerEnvironment& OutEnvironment)
@@ -490,7 +492,7 @@ public:
 	{
 		static auto* CVarMobileEnableStaticAndCSMShadowReceivers = IConsoleManager::Get().FindTConsoleVariableDataInt(TEXT("r.Mobile.EnableStaticAndCSMShadowReceivers"));
 		const bool bMobileEnableStaticAndCSMShadowReceivers = CVarMobileEnableStaticAndCSMShadowReceivers->GetValueOnAnyThread() == 1;
-		return bMobileEnableStaticAndCSMShadowReceivers && (Parameters.Material->GetShadingModels().IsLit()) && Super::ShouldCompilePermutation(Parameters);
+		return bMobileEnableStaticAndCSMShadowReceivers && (Parameters.MaterialParameters.ShadingModels.IsLit()) && Super::ShouldCompilePermutation(Parameters);
 	}
 
 	static void ModifyCompilationEnvironment(const FMaterialShaderPermutationParameters& Parameters, FShaderCompilerEnvironment& OutEnvironment)
@@ -515,7 +517,7 @@ struct FMobileDirectionalLightAndSHIndirectPolicy
 		static auto* CVarAllowStaticLighting = IConsoleManager::Get().FindTConsoleVariableDataInt(TEXT("r.AllowStaticLighting"));
 		const bool bAllowStaticLighting = CVarAllowStaticLighting->GetValueOnAnyThread() != 0;
 
-		return bAllowStaticLighting && Parameters.Material->GetShadingModels().IsLit() && FCachedPointIndirectLightingPolicy::ShouldCompilePermutation(Parameters);
+		return bAllowStaticLighting && Parameters.MaterialParameters.ShadingModels.IsLit() && FCachedPointIndirectLightingPolicy::ShouldCompilePermutation(Parameters);
 	}
 
 	static void ModifyCompilationEnvironment(const FMaterialShaderPermutationParameters& Parameters, FShaderCompilerEnvironment& OutEnvironment)
@@ -594,7 +596,7 @@ struct FMobileMovableDirectionalLightLightingPolicy
 		static auto* CVarMobileAllowMovableDirectionalLights = IConsoleManager::Get().FindTConsoleVariableDataInt(TEXT("r.Mobile.AllowMovableDirectionalLights"));
 		const bool bMobileAllowMovableDirectionalLights = CVarMobileAllowMovableDirectionalLights->GetValueOnAnyThread() != 0;
 
-		return bMobileAllowMovableDirectionalLights && Parameters.Material->GetShadingModels().IsLit();
+		return bMobileAllowMovableDirectionalLights && Parameters.MaterialParameters.ShadingModels.IsLit();
 	}
 
 	static void ModifyCompilationEnvironment(const FMaterialShaderPermutationParameters& Parameters, FShaderCompilerEnvironment& OutEnvironment)
@@ -616,7 +618,7 @@ struct FMobileMovableDirectionalLightCSMLightingPolicy
 		static auto* CVarMobileAllowMovableDirectionalLights = IConsoleManager::Get().FindTConsoleVariableDataInt(TEXT("r.Mobile.AllowMovableDirectionalLights"));
 		const bool bMobileAllowMovableDirectionalLights = CVarMobileAllowMovableDirectionalLights->GetValueOnAnyThread() != 0;
 
-		return bMobileAllowMovableDirectionalLights && Parameters.Material->GetShadingModels().IsLit();
+		return bMobileAllowMovableDirectionalLights && Parameters.MaterialParameters.ShadingModels.IsLit();
 	}	
 
 	static void ModifyCompilationEnvironment(const FMaterialShaderPermutationParameters& Parameters, FShaderCompilerEnvironment& OutEnvironment)
@@ -647,7 +649,7 @@ struct FMobileMovableDirectionalLightWithLightmapPolicy : public TLightMapPolicy
 		static auto* CVarAllowStaticLighting = IConsoleManager::Get().FindTConsoleVariableDataInt(TEXT("r.AllowStaticLighting"));
 		const bool bAllowStaticLighting = (!CVarAllowStaticLighting || CVarAllowStaticLighting->GetValueOnAnyThread() != 0);
 
-		return bAllowStaticLighting && bMobileAllowMovableDirectionalLights && (Parameters.Material->GetShadingModels().IsLit()) && Super::ShouldCompilePermutation(Parameters);
+		return bAllowStaticLighting && bMobileAllowMovableDirectionalLights && (Parameters.MaterialParameters.ShadingModels.IsLit()) && Super::ShouldCompilePermutation(Parameters);
 	}
 
 	static void ModifyCompilationEnvironment(const FMaterialShaderPermutationParameters& Parameters, FShaderCompilerEnvironment& OutEnvironment)
@@ -711,6 +713,7 @@ enum ELightMapPolicyType
 
 class FUniformLightMapPolicyShaderParametersType
 {
+	DECLARE_TYPE_LAYOUT(FUniformLightMapPolicyShaderParametersType, NonVirtual);
 public:
 	void Bind(const FShaderParameterMap& ParameterMap)
 	{
@@ -726,9 +729,9 @@ public:
 		Ar << LightmapResourceCluster;
 	}
 
-	FShaderUniformBufferParameter PrecomputedLightingBufferParameter;
-	FShaderUniformBufferParameter IndirectLightingCacheParameter;
-	FShaderUniformBufferParameter LightmapResourceCluster;
+	LAYOUT_FIELD(FShaderUniformBufferParameter, PrecomputedLightingBufferParameter);
+	LAYOUT_FIELD(FShaderUniformBufferParameter, IndirectLightingCacheParameter);
+	LAYOUT_FIELD(FShaderUniformBufferParameter, LightmapResourceCluster);
 };
 
 class FUniformLightMapPolicy
@@ -979,6 +982,7 @@ public:
 
 	class PixelParametersType : public FUniformLightMapPolicyShaderParametersType, public FSelfShadowedTranslucencyPolicy::PixelParametersType
 	{
+		DECLARE_INLINE_TYPE_LAYOUT_EXPLICIT_BASES(PixelParametersType, NonVirtual, FUniformLightMapPolicyShaderParametersType, FSelfShadowedTranslucencyPolicy::PixelParametersType);
 	public:
 		void Bind(const FShaderParameterMap& ParameterMap)
 		{
@@ -997,8 +1001,8 @@ public:
 	{
 		static IConsoleVariable* AllowStaticLightingVar = IConsoleManager::Get().FindConsoleVariable(TEXT("r.AllowStaticLighting"));
 
-		return Parameters.Material->GetShadingModels().IsLit()
-			&& IsTranslucentBlendMode(Parameters.Material->GetBlendMode())
+		return Parameters.MaterialParameters.ShadingModels.IsLit()
+			&& IsTranslucentBlendMode(Parameters.MaterialParameters.BlendMode)
 			&& (!AllowStaticLightingVar || AllowStaticLightingVar->GetInt() != 0)
 			&& FSelfShadowedTranslucencyPolicy::ShouldCompilePermutation(Parameters);
 	}
@@ -1028,6 +1032,7 @@ public:
 
 	class PixelParametersType : public FUniformLightMapPolicyShaderParametersType, public FSelfShadowedTranslucencyPolicy::PixelParametersType
 	{
+		DECLARE_INLINE_TYPE_LAYOUT_EXPLICIT_BASES(PixelParametersType, NonVirtual, FUniformLightMapPolicyShaderParametersType, FSelfShadowedTranslucencyPolicy::PixelParametersType);
 	public:
 		void Bind(const FShaderParameterMap& ParameterMap)
 		{
@@ -1046,8 +1051,8 @@ public:
 	{
 		static IConsoleVariable* AllowStaticLightingVar = IConsoleManager::Get().FindConsoleVariable(TEXT("r.AllowStaticLighting"));
 
-		return Parameters.Material->GetShadingModels().IsLit()
-			&& IsTranslucentBlendMode(Parameters.Material->GetBlendMode())
+		return Parameters.MaterialParameters.ShadingModels.IsLit()
+			&& IsTranslucentBlendMode(Parameters.MaterialParameters.BlendMode)
 			&& (!AllowStaticLightingVar || AllowStaticLightingVar->GetInt() != 0)
 			&& FSelfShadowedTranslucencyPolicy::ShouldCompilePermutation(Parameters);
 	}
@@ -1072,4 +1077,4 @@ public:
 		const ElementDataType& ShaderElementData,
 		const PixelParametersType* PixelShaderParameters,
 		FMeshDrawSingleShaderBindings& ShaderBindings);
-};
\ No newline at end of file
+};
