Copying Private-LoadTimes-4.24 stream to Main. Biggest changes are in Materials/Shader memory freezing.

diff --git a/Engine/Source/Runtime/Renderer/Private/LightPropagationVolume.cpp b/Engine/Source/Runtime/Renderer/Private/LightPropagationVolume.cpp
--- a/Engine/Source/Runtime/Renderer/Private/LightPropagationVolume.cpp
+++ b/Engine/Source/Runtime/Renderer/Private/LightPropagationVolume.cpp
@@ -128,6 +128,7 @@ IMPLEMENT_GLOBAL_SHADER_PARAMETER_STRUCT(FLpvWriteUniformBufferParameters,"LpvWr
 // ----------------------------------------------------------------------------
 class FLpvWriteShaderCSBase : public FGlobalShader
 {
+	DECLARE_INLINE_TYPE_LAYOUT(FLpvWriteShaderCSBase, NonVirtual);
 public:
 	// Default constructor
 	FLpvWriteShaderCSBase()	{	}
@@ -169,39 +170,10 @@ public:
 		OutEnvironment.SetDefine( TEXT("LPV_MULTIPLE_BOUNCES"), (uint32)LPV_MULTIPLE_BOUNCES );
 		OutEnvironment.SetDefine( TEXT("LPV_GV_SH_ORDER"),			(uint32)LPV_GV_SH_ORDER );
 	}
-	// Serialization
-	virtual bool Serialize( FArchive& Ar ) override
-	{
-		bool bShaderHasOutdatedParameters = FGlobalShader::Serialize( Ar );
-
-		for(int i = 0; i < 7; i++)
-		{
-			Ar << LpvBufferSRVParameters[i];
-			Ar << LpvBufferUAVs[i];
-		}
-
-		Ar << LpvVolumeTextureSampler;
-		Ar << VplListHeadBufferSRV;
-		Ar << VplListHeadBufferUAV;
-		Ar << VplListBufferSRV;
-		Ar << VplListBufferUAV;
-		for ( int i = 0; i < NUM_GV_TEXTURES; i++ )
-		{
-			Ar << GvBufferSRVParameters[i];
-			Ar << GvBufferUAVs[i];
-		}
-		Ar << GvListBufferUAV;
-		Ar << GvListHeadBufferUAV;
-		Ar << GvListBufferSRV;
-		Ar << GvListHeadBufferSRV;
-		Ar << AOVolumeTextureUAV;
-		Ar << AOVolumeTextureSRV;
-		return bShaderHasOutdatedParameters;
-	}
 
 	void SetParameters(FRHICommandList& RHICmdList, const FLpvBaseWriteShaderParams& Params )
 	{
-		FRHIComputeShader* ShaderRHI = GetComputeShader();
+		FRHIComputeShader* ShaderRHI = RHICmdList.GetBoundComputeShader();
 		SetUniformBufferParameter(RHICmdList, ShaderRHI, GetUniformBufferParameter<FLpvWriteUniformBufferParameters>(), Params.UniformBuffer );
 
 		TArray<int32> ResourceIndices;
@@ -297,7 +269,7 @@ public:
 		TArray<int32> ResourceIndices;
 		TArray<FRHIUnorderedAccessView*> UAVs;
 
-		FRHIComputeShader* ShaderRHI = GetComputeShader();
+		FRHIComputeShader* ShaderRHI = RHICmdList.GetBoundComputeShader();
 		for ( int i = 0; i < 7; i++ )
 		{
 			if ( LpvBufferSRVParameters[i].IsBound() )
@@ -379,22 +351,22 @@ public:
 	}
 
 protected:
-	FShaderResourceParameter LpvBufferSRVParameters[7];
-	FShaderResourceParameter LpvBufferUAVs[7];
-	FShaderResourceParameter LpvVolumeTextureSampler;
-	FShaderResourceParameter VplListHeadBufferSRV;
-	FShaderResourceParameter VplListHeadBufferUAV;
-	FShaderResourceParameter VplListBufferSRV;
-	FShaderResourceParameter VplListBufferUAV;
-
-	FShaderResourceParameter GvBufferSRVParameters[NUM_GV_TEXTURES];
-	FShaderResourceParameter GvBufferUAVs[NUM_GV_TEXTURES];
-	FShaderResourceParameter GvListBufferSRV;
-	FShaderResourceParameter GvListBufferUAV;
-	FShaderResourceParameter GvListHeadBufferSRV;
-	FShaderResourceParameter GvListHeadBufferUAV;
-	FShaderResourceParameter AOVolumeTextureUAV;
-	FShaderResourceParameter AOVolumeTextureSRV;
+	LAYOUT_ARRAY(FShaderResourceParameter, LpvBufferSRVParameters, 7);
+	LAYOUT_ARRAY(FShaderResourceParameter, LpvBufferUAVs, 7);
+	LAYOUT_FIELD(FShaderResourceParameter, LpvVolumeTextureSampler);
+	LAYOUT_FIELD(FShaderResourceParameter, VplListHeadBufferSRV);
+	LAYOUT_FIELD(FShaderResourceParameter, VplListHeadBufferUAV);
+	LAYOUT_FIELD(FShaderResourceParameter, VplListBufferSRV);
+	LAYOUT_FIELD(FShaderResourceParameter, VplListBufferUAV);
+
+	LAYOUT_ARRAY(FShaderResourceParameter, GvBufferSRVParameters, NUM_GV_TEXTURES);
+	LAYOUT_ARRAY(FShaderResourceParameter, GvBufferUAVs, NUM_GV_TEXTURES);
+	LAYOUT_FIELD(FShaderResourceParameter, GvListBufferSRV);
+	LAYOUT_FIELD(FShaderResourceParameter, GvListBufferUAV);
+	LAYOUT_FIELD(FShaderResourceParameter, GvListHeadBufferSRV);
+	LAYOUT_FIELD(FShaderResourceParameter, GvListHeadBufferUAV);
+	LAYOUT_FIELD(FShaderResourceParameter, AOVolumeTextureUAV);
+	LAYOUT_FIELD(FShaderResourceParameter, AOVolumeTextureSRV);
 };
 
 
@@ -411,8 +383,6 @@ public:
 	FLpvClearCS()	{	}
 
 	explicit FLpvClearCS( const ShaderMetaType::CompiledShaderInitializerType& Initializer ) : FLpvWriteShaderCSBase(Initializer)		{	}
-
-	virtual bool Serialize( FArchive& Ar ) override			{ return FLpvWriteShaderCSBase::Serialize( Ar ); }
 };
 IMPLEMENT_SHADER_TYPE(,FLpvClearCS,TEXT("/Engine/Private/LPVClear.usf"),TEXT("CSClear"),SF_Compute);
 
@@ -430,8 +400,6 @@ public:
 	FLpvClearGeometryVolumeCS()	{	}
 
 	explicit FLpvClearGeometryVolumeCS( const ShaderMetaType::CompiledShaderInitializerType& Initializer ) : FLpvWriteShaderCSBase(Initializer)		{	}
-
-	virtual bool Serialize( FArchive& Ar ) override			{ return FLpvWriteShaderCSBase::Serialize( Ar ); }
 };
 IMPLEMENT_SHADER_TYPE(,FLpvClearGeometryVolumeCS,TEXT("/Engine/Private/LPVClear.usf"),TEXT("CSClearGeometryVolume"),SF_Compute);
 
@@ -449,8 +417,6 @@ public:
 	FLpvClearListsCS()	{	}
 
 	explicit FLpvClearListsCS( const ShaderMetaType::CompiledShaderInitializerType& Initializer ) : FLpvWriteShaderCSBase(Initializer)		{	}
-
-	virtual bool Serialize( FArchive& Ar ) override			{ return FLpvWriteShaderCSBase::Serialize( Ar ); }
 };
 IMPLEMENT_SHADER_TYPE(,FLpvClearListsCS,TEXT("/Engine/Private/LPVClearLists.usf"),TEXT("CSClearLists"),SF_Compute);
 
@@ -483,7 +449,7 @@ public:
 		FRHITexture* RsmNormalTextureRHI,
 		FRHITexture* RsmDepthTextureRHI )
 	{
-		FRHIComputeShader* ShaderRHI = GetComputeShader();
+		FRHIComputeShader* ShaderRHI = RHICmdList.GetBoundComputeShader();
 		FLpvWriteShaderCSBase::SetParameters(RHICmdList, BaseParams );
 
 		FRHISamplerState* SamplerStateLinear  = TStaticSamplerState<SF_Bilinear,AM_Clamp,AM_Clamp,AM_Clamp>::GetRHI();
@@ -501,24 +467,13 @@ public:
 		FLpvWriteShaderCSBase::UnbindBuffers(RHICmdList, BaseParams);
 	}
 
-
-	virtual bool Serialize( FArchive& Ar ) override			
-	{ 
-		bool rv = FLpvWriteShaderCSBase::Serialize( Ar ); 
-		Ar << RsmDiffuseTexture;
-		Ar << RsmNormalTexture;
-		Ar << RsmDepthTexture;
-		Ar << LinearTextureSampler;
-		Ar << PointTextureSampler;
-		return rv;
-	}
 protected:
-	FShaderResourceParameter RsmDiffuseTexture;
-	FShaderResourceParameter RsmNormalTexture;
-	FShaderResourceParameter RsmDepthTexture;
+	LAYOUT_FIELD(FShaderResourceParameter, RsmDiffuseTexture);
+	LAYOUT_FIELD(FShaderResourceParameter, RsmNormalTexture);
+	LAYOUT_FIELD(FShaderResourceParameter, RsmDepthTexture);
 
-	FShaderResourceParameter LinearTextureSampler;
-	FShaderResourceParameter PointTextureSampler;
+	LAYOUT_FIELD(FShaderResourceParameter, LinearTextureSampler);
+	LAYOUT_FIELD(FShaderResourceParameter, PointTextureSampler);
 };
 IMPLEMENT_SHADER_TYPE(,FLpvInject_GenerateVplListsCS,TEXT("/Engine/Private/LPVInject_GenerateVplLists.usf"),TEXT("CSGenerateVplLists_LightDirectional"),SF_Compute);
 
@@ -535,8 +490,6 @@ public:
 	FLpvInject_AccumulateVplListsCS()	{	}
 
 	explicit FLpvInject_AccumulateVplListsCS( const ShaderMetaType::CompiledShaderInitializerType& Initializer ) : FLpvWriteShaderCSBase(Initializer)		{	}
-
-	virtual bool Serialize( FArchive& Ar ) override			{ return FLpvWriteShaderCSBase::Serialize( Ar ); }
 };
 IMPLEMENT_SHADER_TYPE(,FLpvInject_AccumulateVplListsCS,TEXT("/Engine/Private/LPVInject_AccumulateVplLists.usf"),TEXT("CSAccumulateVplLists"),SF_Compute);
 
@@ -556,15 +509,7 @@ public:
 	{	
 	}
 
-	// FShader interface.
-	virtual bool Serialize(FArchive& Ar) override
-	{
-
-		bool bShaderHasOutdatedParameters = FLpvWriteShaderCSBase::Serialize( Ar );
-		return bShaderHasOutdatedParameters;
-	}
-
-	virtual void SetParameters(
+	void SetParameters(
 		FRHICommandList& RHICmdList, 
 		FLpvBaseWriteShaderParams& BaseParams )
 	{
@@ -589,15 +534,7 @@ public:
 	{	
 	}
 
-	// FShader interface.
-	virtual bool Serialize(FArchive& Ar) override
-	{
-
-		bool bShaderHasOutdatedParameters = FLpvWriteShaderCSBase::Serialize( Ar );
-		return bShaderHasOutdatedParameters;
-	}
-
-	virtual void SetParameters(
+	void SetParameters(
 		FRHICommandList& RHICmdList, 
 		FLpvBaseWriteShaderParams& BaseParams )
 	{	
@@ -620,8 +557,6 @@ public:
 	FLpvBuildGeometryVolumeCS()	{	}
 
 	explicit FLpvBuildGeometryVolumeCS( const ShaderMetaType::CompiledShaderInitializerType& Initializer ) : FLpvWriteShaderCSBase(Initializer)		{	}
-
-	virtual bool Serialize( FArchive& Ar ) override			{ return FLpvWriteShaderCSBase::Serialize( Ar ); }
 };
 IMPLEMENT_SHADER_TYPE(,FLpvBuildGeometryVolumeCS,TEXT("/Engine/Private/LPVBuildGeometryVolume.usf"),TEXT("CSBuildGeometryVolume"),SF_Compute);
 
@@ -659,8 +594,6 @@ public:
 	TLpvPropagateCS()	{	}
 
 	explicit TLpvPropagateCS( const ShaderMetaType::CompiledShaderInitializerType& Initializer ) : FLpvWriteShaderCSBase(Initializer)		{	}
-
-	virtual bool Serialize( FArchive& Ar ) override			{ return FLpvWriteShaderCSBase::Serialize( Ar ); }
 };
 
 IMPLEMENT_SHADER_TYPE(template<>,TLpvPropagateCS<0>,																		TEXT("/Engine/Private/LPVPropagate.usf"),TEXT("CSPropagate"),SF_Compute);
@@ -672,28 +605,28 @@ IMPLEMENT_SHADER_TYPE(template<>,TLpvPropagateCS<PROPAGATE_AO|PROPAGATE_SECONDAR
 IMPLEMENT_SHADER_TYPE(template<>,TLpvPropagateCS<PROPAGATE_AO|PROPAGATE_MULTIPLE_BOUNCES>,									TEXT("/Engine/Private/LPVPropagate.usf"),TEXT("CSPropagate"),SF_Compute);
 IMPLEMENT_SHADER_TYPE(template<>,TLpvPropagateCS<PROPAGATE_AO|PROPAGATE_SECONDARY_OCCLUSION|PROPAGATE_MULTIPLE_BOUNCES>,	TEXT("/Engine/Private/LPVPropagate.usf"),TEXT("CSPropagate"),SF_Compute);
 
-FLpvWriteShaderCSBase* GetPropagateShader( FViewInfo& View, uint32 ShaderFlags )
+TShaderRef<FLpvWriteShaderCSBase> GetPropagateShader( FViewInfo& View, uint32 ShaderFlags )
 {
 	switch( ShaderFlags )
 	{
 	case 0:
-		return (FLpvWriteShaderCSBase*)*TShaderMapRef<TLpvPropagateCS<0> >( View.ShaderMap );
+		return TShaderMapRef<TLpvPropagateCS<0> >( View.ShaderMap );
 	case PROPAGATE_SECONDARY_OCCLUSION:										
-		return (FLpvWriteShaderCSBase*)*TShaderMapRef<TLpvPropagateCS<PROPAGATE_SECONDARY_OCCLUSION> >( View.ShaderMap );
+		return TShaderMapRef<TLpvPropagateCS<PROPAGATE_SECONDARY_OCCLUSION> >( View.ShaderMap );
 	case PROPAGATE_MULTIPLE_BOUNCES:	
-		return (FLpvWriteShaderCSBase*)*TShaderMapRef<TLpvPropagateCS<PROPAGATE_MULTIPLE_BOUNCES> >( View.ShaderMap );
+		return TShaderMapRef<TLpvPropagateCS<PROPAGATE_MULTIPLE_BOUNCES> >( View.ShaderMap );
 	case PROPAGATE_SECONDARY_OCCLUSION|PROPAGATE_MULTIPLE_BOUNCES:	
-		return (FLpvWriteShaderCSBase*)*TShaderMapRef<TLpvPropagateCS<PROPAGATE_SECONDARY_OCCLUSION|PROPAGATE_MULTIPLE_BOUNCES> >( View.ShaderMap );
+		return TShaderMapRef<TLpvPropagateCS<PROPAGATE_SECONDARY_OCCLUSION|PROPAGATE_MULTIPLE_BOUNCES> >( View.ShaderMap );
 	case PROPAGATE_AO:							
-		return (FLpvWriteShaderCSBase*)*TShaderMapRef<TLpvPropagateCS<PROPAGATE_AO> >( View.ShaderMap );
+		return TShaderMapRef<TLpvPropagateCS<PROPAGATE_AO> >( View.ShaderMap );
 	case PROPAGATE_AO|PROPAGATE_SECONDARY_OCCLUSION:							
-		return (FLpvWriteShaderCSBase*)*TShaderMapRef<TLpvPropagateCS<PROPAGATE_AO|PROPAGATE_SECONDARY_OCCLUSION> >( View.ShaderMap );
+		return TShaderMapRef<TLpvPropagateCS<PROPAGATE_AO|PROPAGATE_SECONDARY_OCCLUSION> >( View.ShaderMap );
 	case PROPAGATE_AO|PROPAGATE_MULTIPLE_BOUNCES:
-		return (FLpvWriteShaderCSBase*)*TShaderMapRef<TLpvPropagateCS<PROPAGATE_AO|PROPAGATE_MULTIPLE_BOUNCES> >( View.ShaderMap );
+		return TShaderMapRef<TLpvPropagateCS<PROPAGATE_AO|PROPAGATE_MULTIPLE_BOUNCES> >( View.ShaderMap );
 	case PROPAGATE_AO|PROPAGATE_SECONDARY_OCCLUSION|PROPAGATE_MULTIPLE_BOUNCES:
-		return (FLpvWriteShaderCSBase*)*TShaderMapRef<TLpvPropagateCS<PROPAGATE_AO|PROPAGATE_SECONDARY_OCCLUSION|PROPAGATE_MULTIPLE_BOUNCES> >( View.ShaderMap );
+		return TShaderMapRef<TLpvPropagateCS<PROPAGATE_AO|PROPAGATE_SECONDARY_OCCLUSION|PROPAGATE_MULTIPLE_BOUNCES> >( View.ShaderMap );
 	}
-	return NULL;
+	return TShaderRef<FLpvWriteShaderCSBase>();
 }
 
 // ----------------------------------------------------------------------------
@@ -714,15 +647,13 @@ public:
 		FDirectLightInjectBufferRef& InjectUniformBuffer )
 	{
 		FLpvWriteShaderCSBase::SetParameters(RHICmdList, BaseParams );
-		FRHIComputeShader* ComputeShaderRHI = GetComputeShader();
+		FRHIComputeShader* ComputeShaderRHI = RHICmdList.GetBoundComputeShader();
 		SetUniformBufferParameter(RHICmdList, ComputeShaderRHI, GetUniformBufferParameter<FLpvDirectLightInjectParameters>(), InjectUniformBuffer );
 	}
 
 	FLpvInjectShader_Base()	{	}
 
 	explicit FLpvInjectShader_Base( const ShaderMetaType::CompiledShaderInitializerType& Initializer ) : FLpvWriteShaderCSBase(Initializer)		{	}
-
-	virtual bool Serialize( FArchive& Ar ) override			{ return FLpvWriteShaderCSBase::Serialize( Ar ); }
 };
 
 // ----------------------------------------------------------------------------
@@ -749,8 +680,6 @@ public:
 	TLpvInject_LightCS()	{	}
 
 	explicit TLpvInject_LightCS( const ShaderMetaType::CompiledShaderInitializerType& Initializer ) : FLpvInjectShader_Base(Initializer)		{	}
-
-	virtual bool Serialize( FArchive& Ar ) override			{ return FLpvInjectShader_Base::Serialize( Ar ); }
 };
 
 IMPLEMENT_SHADER_TYPE(template<>,TLpvInject_LightCS<0>,TEXT("/Engine/Private/LPVDirectLightInject.usf"),TEXT("CSLightInject_ListGenCS"),SF_Compute);
@@ -1015,24 +944,24 @@ void FLightPropagationVolume::Clear(FRHICommandListImmediate& RHICmdList, FViewI
 	// Clear the list buffers
 	{
 		TShaderMapRef<FLpvClearListsCS> Shader(View.ShaderMap);
-		RHICmdList.SetComputeShader(Shader->GetComputeShader());
+		RHICmdList.SetComputeShader(Shader.GetComputeShader());
 
 		FLpvBaseWriteShaderParams ShaderParams;
 		GetShaderParams( ShaderParams );
 		Shader->SetParameters(RHICmdList, ShaderParams );
-		DispatchComputeShader(RHICmdList, *Shader, LPV_GRIDRES/4, LPV_GRIDRES/4, LPV_GRIDRES/4 );
+		DispatchComputeShader(RHICmdList, Shader.GetShader(), LPV_GRIDRES/4, LPV_GRIDRES/4, LPV_GRIDRES/4 );
 		Shader->UnbindBuffers(RHICmdList, ShaderParams );
 	}
 
 	// Clear the LPV (or fade, if REFINE_OVER_TIME is enabled)
 	{
 		TShaderMapRef<FLpvClearCS> Shader(View.ShaderMap);
-		RHICmdList.SetComputeShader(Shader->GetComputeShader());
+		RHICmdList.SetComputeShader(Shader.GetComputeShader());
 
 		FLpvBaseWriteShaderParams ShaderParams;
 		GetShaderParams( ShaderParams );
 		Shader->SetParameters(RHICmdList, ShaderParams );
-		DispatchComputeShader(RHICmdList, *Shader, LPV_GRIDRES/4, LPV_GRIDRES/4, LPV_GRIDRES/4 );
+		DispatchComputeShader(RHICmdList, Shader.GetShader(), LPV_GRIDRES/4, LPV_GRIDRES/4, LPV_GRIDRES/4 );
 		Shader->UnbindBuffers(RHICmdList, ShaderParams);
 	}
 
@@ -1040,12 +969,12 @@ void FLightPropagationVolume::Clear(FRHICommandListImmediate& RHICmdList, FViewI
 	if ( bGeometryVolumeNeeded )
 	{
 		TShaderMapRef<FLpvClearGeometryVolumeCS> Shader(View.ShaderMap);
-		RHICmdList.SetComputeShader(Shader->GetComputeShader());
+		RHICmdList.SetComputeShader(Shader.GetComputeShader());
 
 		FLpvBaseWriteShaderParams ShaderParams;
 		GetShaderParams( ShaderParams );
 		Shader->SetParameters(RHICmdList, ShaderParams );
-		DispatchComputeShader(RHICmdList, *Shader, LPV_GRIDRES/4, LPV_GRIDRES/4, LPV_GRIDRES/4 );
+		DispatchComputeShader(RHICmdList, Shader.GetShader(), LPV_GRIDRES/4, LPV_GRIDRES/4, LPV_GRIDRES/4 );
 		Shader->UnbindBuffers(RHICmdList, ShaderParams);
 	}
 	RHICmdList.AutomaticCacheFlushAfterComputeShader(true);
@@ -1125,7 +1054,7 @@ void FLightPropagationVolume::InjectDirectionalLightRSM(
 		SetVplInjectionConstants(ProjectedShadowInfo, LightProxy ); //-V595
 
 		TShaderMapRef<FLpvInject_GenerateVplListsCS> Shader(View.ShaderMap);
-		RHICmdList.SetComputeShader(Shader->GetComputeShader());
+		RHICmdList.SetComputeShader(Shader.GetComputeShader());
 
 		// Clear the list counter the first time this function is called in a frame
 		FLpvBaseWriteShaderParams ShaderParams;
@@ -1134,7 +1063,7 @@ void FLightPropagationVolume::InjectDirectionalLightRSM(
 
 		int32 RSMResolution = FSceneRenderTargets::Get_FrameConstantsOnly().GetReflectiveShadowMapResolution();
 		// todo: what if not divisible by 8?
-		DispatchComputeShader(RHICmdList, *Shader, RSMResolution / 8, RSMResolution / 8, 1 ); 
+		DispatchComputeShader(RHICmdList, Shader.GetShader(), RSMResolution / 8, RSMResolution / 8, 1 ); 
 
 		Shader->UnbindBuffers(RHICmdList, ShaderParams);
 	}
@@ -1152,7 +1081,7 @@ void FLightPropagationVolume::InjectDirectionalLightRSM(
 		LpvWriteUniformBufferParams->GeometryVolumeCaptureLightDirection = LightDirection;
 
 		TShaderMapRef<FLpvBuildGeometryVolumeCS> Shader(View.ShaderMap);
-		RHICmdList.SetComputeShader(Shader->GetComputeShader());
+		RHICmdList.SetComputeShader(Shader.GetComputeShader());
 
 		LpvWriteUniformBuffer.SetContents( *LpvWriteUniformBufferParams ); // FIXME: is this causing a stall? Double-buffer?
 
@@ -1160,7 +1089,7 @@ void FLightPropagationVolume::InjectDirectionalLightRSM(
 		GetShaderParams( ShaderParams );
 		Shader->SetParameters(RHICmdList, ShaderParams );
 
-		DispatchComputeShader(RHICmdList, *Shader, LPV_GRIDRES/4, LPV_GRIDRES/4, LPV_GRIDRES/4 );
+		DispatchComputeShader(RHICmdList, Shader.GetShader(), LPV_GRIDRES/4, LPV_GRIDRES/4, LPV_GRIDRES/4 );
 
 		Shader->UnbindBuffers(RHICmdList, ShaderParams);
 	}
@@ -1181,13 +1110,13 @@ void FLightPropagationVolume::ComputeDirectionalOcclusion( FRHICommandListImmedi
 
 			mWriteBufferIndex = 1-mWriteBufferIndex; // Swap buffers with each iteration
 			TShaderMapRef<FLpvDirectionalOcclusionCS> Shader(View.ShaderMap);
-			RHICmdList.SetComputeShader(Shader->GetComputeShader());
+			RHICmdList.SetComputeShader(Shader.GetComputeShader());
 			FLpvBaseWriteShaderParams ShaderParams;
 			GetShaderParams( ShaderParams );
 			Shader->SetParameters( RHICmdList, ShaderParams );
 			LpvWriteUniformBuffer.SetContents( *LpvWriteUniformBufferParams );
 
-			DispatchComputeShader( RHICmdList, *Shader, LPV_GRIDRES/4, LPV_GRIDRES/4, LPV_GRIDRES/4 );
+			DispatchComputeShader( RHICmdList, Shader.GetShader(), LPV_GRIDRES/4, LPV_GRIDRES/4, LPV_GRIDRES/4 );
 			Shader->UnbindBuffers(RHICmdList, ShaderParams);
 		}
 	}
@@ -1218,7 +1147,7 @@ void FLightPropagationVolume::Update( FRHICommandListImmediate& RHICmdList, FVie
 		mWriteBufferIndex = 1-mWriteBufferIndex; // Swap buffers with each iteration
 
 		TShaderMapRef<FLpvInject_AccumulateVplListsCS> Shader(View.ShaderMap);
-		RHICmdList.SetComputeShader(Shader->GetComputeShader());
+		RHICmdList.SetComputeShader(Shader.GetComputeShader());
 
 		//LpvWriteUniformBuffer.SetContents( *LpvWriteUniformBufferParams );
 
@@ -1226,7 +1155,7 @@ void FLightPropagationVolume::Update( FRHICommandListImmediate& RHICmdList, FVie
 		GetShaderParams( ShaderParams );
 		Shader->SetParameters(RHICmdList, ShaderParams );
 
-		DispatchComputeShader(RHICmdList, *Shader, LPV_GRIDRES/4, LPV_GRIDRES/4, LPV_GRIDRES/4 );
+		DispatchComputeShader(RHICmdList, Shader.GetShader(), LPV_GRIDRES/4, LPV_GRIDRES/4, LPV_GRIDRES/4 );
 		RHICmdList.FlushComputeShaderCache();
 
 		Shader->UnbindBuffers(RHICmdList, ShaderParams);
@@ -1257,8 +1186,8 @@ void FLightPropagationVolume::Update( FRHICommandListImmediate& RHICmdList, FVie
 			if ( i < CVarLPVNumAOPropagationSteps.GetValueOnRenderThread() )
 				ShaderFlags |= PROPAGATE_AO;
 
-			FLpvWriteShaderCSBase* Shader = GetPropagateShader( View, ShaderFlags );
-			RHICmdList.SetComputeShader(Shader->GetComputeShader());
+			TShaderRef<FLpvWriteShaderCSBase> Shader = GetPropagateShader( View, ShaderFlags );
+			RHICmdList.SetComputeShader(Shader.GetComputeShader());
 
 			LpvWriteUniformBufferParams->PropagationIndex = i;
 
@@ -1266,7 +1195,7 @@ void FLightPropagationVolume::Update( FRHICommandListImmediate& RHICmdList, FVie
 			GetShaderParams( ShaderParams );
 			Shader->SetParameters(RHICmdList, ShaderParams );
 
-			DispatchComputeShader(RHICmdList, Shader, LPV_GRIDRES/4, LPV_GRIDRES/4, LPV_GRIDRES/4 );
+			DispatchComputeShader(RHICmdList, Shader.GetShader(), LPV_GRIDRES/4, LPV_GRIDRES/4, LPV_GRIDRES/4 );
 
 			// Insert a flush for all iterations except the last - these dispatches can't overlap!
 			if ( i < LPVNumPropagationSteps - 1 )
@@ -1286,12 +1215,12 @@ void FLightPropagationVolume::Update( FRHICommandListImmediate& RHICmdList, FVie
 		SCOPED_DRAW_EVENT(RHICmdList, LpvCopyAOVolume);
 
 		TShaderMapRef<FLpvCopyAOVolumeCS> Shader(View.ShaderMap);
-		RHICmdList.SetComputeShader(Shader->GetComputeShader());
+		RHICmdList.SetComputeShader(Shader.GetComputeShader());
 		FLpvBaseWriteShaderParams ShaderParams;
 		GetShaderParams( ShaderParams );
 		Shader->SetParameters( RHICmdList, ShaderParams );
 		LpvWriteUniformBuffer.SetContents( *LpvWriteUniformBufferParams );
-		DispatchComputeShader( RHICmdList, *Shader, LPV_GRIDRES/4, LPV_GRIDRES/4, LPV_GRIDRES/4 );
+		DispatchComputeShader( RHICmdList, Shader.GetShader(), LPV_GRIDRES/4, LPV_GRIDRES/4, LPV_GRIDRES/4 );
 		Shader->UnbindBuffers(RHICmdList, ShaderParams);
 	}
 }
@@ -1374,19 +1303,18 @@ void FLightPropagationVolume::InjectLightDirect(FRHICommandListImmediate& RHICmd
 		InjectUniformBufferParams.LightSourceLength = LightParameters.SourceLength;
 		InjectUniformBufferParams.bLightInverseSquaredAttenuation = Light.IsInverseSquared() ? 1.0f : 0.0f;
 
-		FLpvInjectShader_Base* Shader = nullptr;
-
+		TShaderRef<FLpvInjectShader_Base> Shader;
 		switch ( Light.GetLightType() )
 		{
 			case LightType_Point:
 				{
 					if ( Light.CastsStaticShadow() || Light.CastsDynamicShadow() )
 					{
-						Shader = (FLpvInjectShader_Base*)*TShaderMapRef<TLpvInject_LightCS<INJECT_SHADOW_CASTING> >(View.ShaderMap);
+						Shader = TShaderMapRef<TLpvInject_LightCS<INJECT_SHADOW_CASTING> >(View.ShaderMap);
 					}
 					else
 					{
-						Shader = (FLpvInjectShader_Base*)*TShaderMapRef<TLpvInject_LightCS<0> >(View.ShaderMap);
+						Shader = TShaderMapRef<TLpvInject_LightCS<0> >(View.ShaderMap);
 					}
 				}
 				break;
@@ -1394,19 +1322,19 @@ void FLightPropagationVolume::InjectLightDirect(FRHICommandListImmediate& RHICmd
 				{
 					if ( Light.CastsStaticShadow() || Light.CastsDynamicShadow() )
 					{
-						Shader = (FLpvInjectShader_Base*)*TShaderMapRef<TLpvInject_LightCS<INJECT_SPOT_ATTENUATION | INJECT_SHADOW_CASTING> >(View.ShaderMap);
+						Shader = TShaderMapRef<TLpvInject_LightCS<INJECT_SPOT_ATTENUATION | INJECT_SHADOW_CASTING> >(View.ShaderMap);
 					}
 					else
 					{
-						Shader = (FLpvInjectShader_Base*)*TShaderMapRef<TLpvInject_LightCS<INJECT_SPOT_ATTENUATION | 0> >(View.ShaderMap);
+						Shader = TShaderMapRef<TLpvInject_LightCS<INJECT_SPOT_ATTENUATION | 0> >(View.ShaderMap);
 					}
 				}
 				break;
 		}
 
-		if (Shader)
+		if (Shader.IsValid())
 		{
-			RHICmdList.SetComputeShader(Shader->GetComputeShader());
+			RHICmdList.SetComputeShader(Shader.GetComputeShader());
 
 			FDirectLightInjectBufferRef InjectUniformBuffer =
 				FDirectLightInjectBufferRef::CreateUniformBufferImmediate(InjectUniformBufferParams, UniformBuffer_SingleFrame);
@@ -1419,7 +1347,7 @@ void FLightPropagationVolume::InjectLightDirect(FRHICommandListImmediate& RHICmd
 			LpvWriteUniformBuffer.SetContents(*LpvWriteUniformBufferParams);
 
 			Shader->SetParameters(RHICmdList, ShaderParams, InjectUniformBuffer);
-			DispatchComputeShader(RHICmdList, Shader, LPV_GRIDRES / 4, LPV_GRIDRES / 4, LPV_GRIDRES / 4);
+			DispatchComputeShader(RHICmdList, Shader.GetShader(), LPV_GRIDRES / 4, LPV_GRIDRES / 4, LPV_GRIDRES / 4);
 			Shader->UnbindBuffers(RHICmdList, ShaderParams);
 		}
 		else
diff --git a/Engine/Source/Runtime/Renderer/Private/LightPropagationVolumeVisualisation.cpp b/Engine/Source/Runtime/Renderer/Private/LightPropagationVolumeVisualisation.cpp
--- a/Engine/Source/Runtime/Renderer/Private/LightPropagationVolumeVisualisation.cpp
+++ b/Engine/Source/Runtime/Renderer/Private/LightPropagationVolumeVisualisation.cpp
@@ -53,7 +53,6 @@ public:
 
 	FLpvVisualiseGS()																												{}
 	explicit FLpvVisualiseGS( const ShaderMetaType::CompiledShaderInitializerType& Initializer ) : FLpvVisualiseBase(Initializer)	{}
-	virtual bool Serialize( FArchive& Ar ) override																					{ return FLpvVisualiseBase::Serialize( Ar ); }
 
 	static bool ShouldCompilePermutation(const FGlobalShaderPermutationParameters& Parameters)
 	{
@@ -64,7 +63,7 @@ public:
 		FRHICommandList& RHICmdList, 
 		const FSceneView& View )
 	{
-		FRHIGeometryShader* ShaderRHI = GetGeometryShader();
+		FRHIGeometryShader* ShaderRHI = RHICmdList.GetBoundGeometryShader();
 		FGlobalShader::SetParameters<FViewUniformShaderParameters>(RHICmdList, ShaderRHI, View.ViewUniformBuffer);
 	}
 };
@@ -76,7 +75,6 @@ public:
 
 	FLpvVisualiseVS()	{	}
 	explicit FLpvVisualiseVS( const ShaderMetaType::CompiledShaderInitializerType& Initializer ) : FLpvVisualiseBase(Initializer) {}
-	virtual bool Serialize( FArchive& Ar ) override																					{ return FLpvVisualiseBase::Serialize( Ar ); }
 
 	static bool ShouldCompilePermutation(const FGlobalShaderPermutationParameters& Parameters)		{ return IsFeatureLevelSupported(Parameters.Platform, ERHIFeatureLevel::SM5) && IsLPVSupported(Parameters.Platform); }
 
@@ -84,7 +82,7 @@ public:
 		FRHICommandList& RHICmdList, 
 		const FSceneView& View )
 	{
-		FRHIVertexShader* ShaderRHI = GetVertexShader();
+		FRHIVertexShader* ShaderRHI = RHICmdList.GetBoundVertexShader();
 		FGlobalShader::SetParameters<FViewUniformShaderParameters>(RHICmdList, ShaderRHI, View.ViewUniformBuffer);
 	}
 };
@@ -116,7 +114,7 @@ public:
 		const FLightPropagationVolume* LPV,
 		const FSceneView& View )
 	{
-		FRHIPixelShader* ShaderRHI = GetPixelShader();
+		FRHIPixelShader* ShaderRHI = RHICmdList.GetBoundPixelShader();
 		FGlobalShader::SetParameters<FViewUniformShaderParameters>(RHICmdList, ShaderRHI, View.ViewUniformBuffer);
 		
 		for ( int i = 0; i < 7; i++ )
@@ -141,35 +139,14 @@ public:
 
 	}
 
-
-	// Serialization
-	virtual bool Serialize( FArchive& Ar ) override
-	{
-		bool bShaderHasOutdatedParameters = FLpvVisualiseBase::Serialize( Ar );
-		
-		for ( int i = 0; i < 7; i++ )
-		{
-			Ar << LpvBufferSRVParameters[i];
-		}
-
-		Ar << LpvVolumeTextureSampler;
-
-		for ( int i = 0; i < NUM_GV_TEXTURES; i++ ) 
-		{
-			Ar << GvBufferSRVParameters[i];
-		}
-		return bShaderHasOutdatedParameters;
-	}
-	FShaderResourceParameter LpvBufferSRVParameters[7];
-
-	FShaderResourceParameter LpvVolumeTextureSampler;
-
-	FShaderResourceParameter GvBufferSRVParameters[3];
+	LAYOUT_ARRAY(FShaderResourceParameter, LpvBufferSRVParameters, 7);
+	LAYOUT_FIELD(FShaderResourceParameter, LpvVolumeTextureSampler);
+	LAYOUT_ARRAY(FShaderResourceParameter, GvBufferSRVParameters, 3);
 
 	void UnbindBuffers(FRHICommandList& RHICmdList)
 	{
 		// TODO: Is this necessary here?
-		FRHIPixelShader* ShaderRHI = GetPixelShader();
+		FRHIPixelShader* ShaderRHI = RHICmdList.GetBoundPixelShader();
 		for ( int i = 0; i < 7; i++ )
 		{
 			if ( LpvBufferSRVParameters[i].IsBound() )
@@ -220,11 +197,11 @@ void FLightPropagationVolume::Visualise(FRHICommandList& RHICmdList, const FView
 	}
 	
 	GraphicsPSOInit.BoundShaderState.VertexDeclarationRHI = GSimpleElementVertexDeclaration.VertexDeclarationRHI;
-	GraphicsPSOInit.BoundShaderState.VertexShaderRHI = GETSAFERHISHADER_VERTEX(*VertexShader);
+	GraphicsPSOInit.BoundShaderState.VertexShaderRHI = VertexShader.GetVertexShader();
 #if PLATFORM_SUPPORTS_GEOMETRY_SHADERS
-	GraphicsPSOInit.BoundShaderState.GeometryShaderRHI = GETSAFERHISHADER_GEOMETRY(*GeometryShader);
+	GraphicsPSOInit.BoundShaderState.GeometryShaderRHI = GeometryShader.GetGeometryShader();
 #endif
-	GraphicsPSOInit.BoundShaderState.PixelShaderRHI = GETSAFERHISHADER_PIXEL(*PixelShader);
+	GraphicsPSOInit.BoundShaderState.PixelShaderRHI = PixelShader.GetPixelShader();
 	GraphicsPSOInit.PrimitiveType = PrimType;
 
 	SetGraphicsPipelineState(RHICmdList, GraphicsPSOInit);
diff --git a/Engine/Source/Runtime/Renderer/Private/LightRendering.cpp b/Engine/Source/Runtime/Renderer/Private/LightRendering.cpp
--- a/Engine/Source/Runtime/Renderer/Private/LightRendering.cpp
+++ b/Engine/Source/Runtime/Renderer/Private/LightRendering.cpp
@@ -30,6 +30,11 @@
 
 DECLARE_GPU_STAT(Lights);
 
+IMPLEMENT_TYPE_LAYOUT(FLightFunctionSharedParameters);
+IMPLEMENT_TYPE_LAYOUT(FStencilingGeometryShaderParameters);
+IMPLEMENT_TYPE_LAYOUT(FOnePassPointShadowProjectionShaderParameters);
+IMPLEMENT_TYPE_LAYOUT(FShadowProjectionShaderParameters);
+
 IMPLEMENT_GLOBAL_SHADER_PARAMETER_STRUCT(FDeferredLightUniformStruct, "DeferredLightUniforms");
 
 extern int32 GUseTranslucentLightingVolumes;
@@ -228,7 +233,7 @@ enum class ELightSourceShape
 /** A pixel shader for rendering the light in a deferred pass. */
 class FDeferredLightPS : public FGlobalShader
 {
-	DECLARE_GLOBAL_SHADER(FDeferredLightPS)
+	DECLARE_SHADER_TYPE(FDeferredLightPS, Global)
 
 	class FSourceShapeDim		: SHADER_PERMUTATION_ENUM_CLASS("LIGHT_SOURCE_SHAPE", ELightSourceShape);
 	class FSourceTextureDim		: SHADER_PERMUTATION_BOOL("USE_SOURCE_TEXTURE");
@@ -334,52 +339,18 @@ public:
 		IPooledRenderTarget* ScreenShadowMaskTexture, 
 		FRenderLightParams* RenderLightParams)
 	{
-		FRHIPixelShader* ShaderRHI = GetPixelShader();
+		FRHIPixelShader* ShaderRHI = RHICmdList.GetBoundPixelShader();
 		SetParametersBase(RHICmdList, ShaderRHI, View, ScreenShadowMaskTexture, LightSceneInfo->Proxy->GetIESTextureResource(), RenderLightParams);
 		SetDeferredLightParameters(RHICmdList, ShaderRHI, GetUniformBufferParameter<FDeferredLightUniformStruct>(), LightSceneInfo, View);
 	}
 
 	void SetParametersSimpleLight(FRHICommandList& RHICmdList, const FSceneView& View, const FSimpleLightEntry& SimpleLight, const FSimpleLightPerViewEntry& SimpleLightPerViewData)
 	{
-		FRHIPixelShader* ShaderRHI = GetPixelShader();
+		FRHIPixelShader* ShaderRHI = RHICmdList.GetBoundPixelShader();
 		SetParametersBase(RHICmdList, ShaderRHI, View, nullptr, nullptr, nullptr);
 		SetSimpleDeferredLightParameters(RHICmdList, ShaderRHI, GetUniformBufferParameter<FDeferredLightUniformStruct>(), SimpleLight, SimpleLightPerViewData, View);
 	}
 
-	virtual bool Serialize(FArchive& Ar) override
-	{
-		bool bShaderHasOutdatedParameters = FGlobalShader::Serialize(Ar);
-		Ar << SceneTextureParameters;
-		Ar << LightAttenuationTexture;
-		Ar << LightAttenuationTextureSampler;
-		Ar << LTCMatTexture;
-		Ar << LTCMatSampler;
-		Ar << LTCAmpTexture;
-		Ar << LTCAmpSampler;
-		Ar << IESTexture;
-		Ar << IESTextureSampler;
-		Ar << LightingChannelsTexture;
-		Ar << LightingChannelsSampler;
-		Ar << TransmissionProfilesTexture;
-		Ar << TransmissionProfilesLinearSampler;
-
-		Ar << HairTransmittanceBuffer;
-		Ar << HairTransmittanceBufferMaxCount;
-
-		Ar << HairCategorizationTexture;
-		Ar << HairVisibilityNodeOffsetAndCount;
-		Ar << HairVisibilityNodeData;
-		Ar << ScreenShadowMaskSubPixelTexture;
-
-		Ar << HairLUTTexture;
-		Ar << HairLUTSampler;
-		Ar << HairComponents;
-		Ar << HairShadowMaskValid;
-		Ar << HairDualScatteringRoughnessOverride;
-
-		return bShaderHasOutdatedParameters;
-	}
-
 private:
 
 	void SetParametersBase(
@@ -583,33 +554,33 @@ private:
 		
 	}
 
-	FSceneTextureShaderParameters SceneTextureParameters;
-	FShaderResourceParameter LightAttenuationTexture;
-	FShaderResourceParameter LightAttenuationTextureSampler;
-	FShaderResourceParameter LTCMatTexture;
-	FShaderResourceParameter LTCMatSampler;
-	FShaderResourceParameter LTCAmpTexture;
-	FShaderResourceParameter LTCAmpSampler;
-	FShaderResourceParameter IESTexture;
-	FShaderResourceParameter IESTextureSampler;
-	FShaderResourceParameter LightingChannelsTexture;
-	FShaderResourceParameter LightingChannelsSampler;
-	FShaderResourceParameter TransmissionProfilesTexture;
-	FShaderResourceParameter TransmissionProfilesLinearSampler;
-
-	FShaderParameter HairTransmittanceBufferMaxCount;
-	FShaderResourceParameter HairTransmittanceBuffer;
-	FShaderResourceParameter HairCategorizationTexture;
-	FShaderResourceParameter HairVisibilityNodeOffsetAndCount;
-	FShaderResourceParameter HairVisibilityNodeData;
-	FShaderResourceParameter ScreenShadowMaskSubPixelTexture;
-
-	FShaderResourceParameter HairLUTTexture;
-	FShaderResourceParameter HairLUTSampler;
-	FShaderParameter HairComponents;
-	FShaderParameter HairShadowMaskValid;
-	FShaderParameter HairDualScatteringRoughnessOverride;
-	};
+	LAYOUT_FIELD(FSceneTextureShaderParameters, SceneTextureParameters);
+	LAYOUT_FIELD(FShaderResourceParameter, LightAttenuationTexture);
+	LAYOUT_FIELD(FShaderResourceParameter, LightAttenuationTextureSampler);
+	LAYOUT_FIELD(FShaderResourceParameter, LTCMatTexture);
+	LAYOUT_FIELD(FShaderResourceParameter, LTCMatSampler);
+	LAYOUT_FIELD(FShaderResourceParameter, LTCAmpTexture);
+	LAYOUT_FIELD(FShaderResourceParameter, LTCAmpSampler);
+	LAYOUT_FIELD(FShaderResourceParameter, IESTexture);
+	LAYOUT_FIELD(FShaderResourceParameter, IESTextureSampler);
+	LAYOUT_FIELD(FShaderResourceParameter, LightingChannelsTexture);
+	LAYOUT_FIELD(FShaderResourceParameter, LightingChannelsSampler);
+	LAYOUT_FIELD(FShaderResourceParameter, TransmissionProfilesTexture);
+	LAYOUT_FIELD(FShaderResourceParameter, TransmissionProfilesLinearSampler);
+
+	LAYOUT_FIELD(FShaderParameter, HairTransmittanceBufferMaxCount);
+	LAYOUT_FIELD(FShaderResourceParameter, HairTransmittanceBuffer);
+	LAYOUT_FIELD(FShaderResourceParameter, HairCategorizationTexture);
+	LAYOUT_FIELD(FShaderResourceParameter, HairVisibilityNodeOffsetAndCount);
+	LAYOUT_FIELD(FShaderResourceParameter, HairVisibilityNodeData);
+	LAYOUT_FIELD(FShaderResourceParameter, ScreenShadowMaskSubPixelTexture);
+
+	LAYOUT_FIELD(FShaderResourceParameter, HairLUTTexture);
+	LAYOUT_FIELD(FShaderResourceParameter, HairLUTSampler);
+	LAYOUT_FIELD(FShaderParameter, HairComponents);
+	LAYOUT_FIELD(FShaderParameter, HairShadowMaskValid);
+	LAYOUT_FIELD(FShaderParameter, HairDualScatteringRoughnessOverride);
+};
 
 IMPLEMENT_GLOBAL_SHADER(FDeferredLightPS, "/Engine/Private/DeferredLightPixelShaders.usf", "DeferredLightPixelMain", SF_Pixel);
 
@@ -645,7 +616,7 @@ public:
 
 	void SetParameters(FRHICommandList& RHICmdList, const FSceneView& View, const FLightSceneInfo* LightSceneInfo)
 	{
-		FRHIPixelShader* ShaderRHI = GetPixelShader();
+		FRHIPixelShader* ShaderRHI = RHICmdList.GetBoundPixelShader();
 		FGlobalShader::SetParameters<FViewUniformShaderParameters>(RHICmdList, ShaderRHI,View.ViewUniformBuffer);
 		const float HasValidChannelValue = LightSceneInfo->Proxy->GetPreviewShadowMapChannel() == INDEX_NONE ? 0.0f : 1.0f;
 		SetShaderValue(RHICmdList, ShaderRHI, HasValidChannel, HasValidChannelValue);
@@ -653,17 +624,9 @@ public:
 		SetDeferredLightParameters(RHICmdList, ShaderRHI, GetUniformBufferParameter<FDeferredLightUniformStruct>(), LightSceneInfo, View);
 	}
 
-	virtual bool Serialize(FArchive& Ar) override
-	{
-		bool bShaderHasOutdatedParameters = FGlobalShader::Serialize(Ar);
-		Ar << HasValidChannel;
-		Ar << SceneTextureParameters;
-		return bShaderHasOutdatedParameters;
-	}
-
 private:
-	FShaderParameter HasValidChannel;
-	FSceneTextureShaderParameters SceneTextureParameters;
+	LAYOUT_FIELD(FShaderParameter, HasValidChannel);
+	LAYOUT_FIELD(FSceneTextureShaderParameters, SceneTextureParameters);
 };
 
 IMPLEMENT_SHADER_TYPE(template<>, TDeferredLightOverlapPS<true>, TEXT("/Engine/Private/StationaryLightOverlapShaders.usf"), TEXT("OverlapRadialPixelMain"), SF_Pixel);
@@ -1961,7 +1924,7 @@ static void SetShaderTemplLightingSimple(
 	FRHICommandList& RHICmdList,
 	FGraphicsPipelineStateInitializer& GraphicsPSOInit,
 	const FViewInfo& View,
-	FShader* VertexShader,
+	const TShaderRef<FShader>& VertexShader,
 	const FSimpleLightEntry& SimpleLight,
 	const FSimpleLightPerViewEntry& SimpleLightPerViewData)
 {
@@ -1976,8 +1939,8 @@ static void SetShaderTemplLightingSimple(
 
 	TShaderMapRef< FDeferredLightPS > PixelShader( View.ShaderMap, PermutationVector );
 	GraphicsPSOInit.BoundShaderState.VertexDeclarationRHI = GetVertexDeclarationFVector4();
-	GraphicsPSOInit.BoundShaderState.VertexShaderRHI = GETSAFERHISHADER_VERTEX(VertexShader);
-	GraphicsPSOInit.BoundShaderState.PixelShaderRHI = GETSAFERHISHADER_PIXEL(*PixelShader);
+	GraphicsPSOInit.BoundShaderState.VertexShaderRHI = VertexShader.GetVertexShader();
+	GraphicsPSOInit.BoundShaderState.PixelShaderRHI = PixelShader.GetPixelShader();
 	SetGraphicsPipelineState(RHICmdList, GraphicsPSOInit);
 	PixelShader->SetParametersSimpleLight(RHICmdList, View, SimpleLight, SimpleLightPerViewData);
 }
@@ -2095,8 +2058,8 @@ void FDeferredShadingSceneRenderer::RenderLight(FRHICommandList& RHICmdList, con
 			{
 				TShaderMapRef<TDeferredLightOverlapPS<false> > PixelShader(View.ShaderMap);
 				GraphicsPSOInit.BoundShaderState.VertexDeclarationRHI = GFilterVertexDeclaration.VertexDeclarationRHI;
-				GraphicsPSOInit.BoundShaderState.VertexShaderRHI = GETSAFERHISHADER_VERTEX(*VertexShader);
-				GraphicsPSOInit.BoundShaderState.PixelShaderRHI = GETSAFERHISHADER_PIXEL(*PixelShader);
+				GraphicsPSOInit.BoundShaderState.VertexShaderRHI = VertexShader.GetVertexShader();
+				GraphicsPSOInit.BoundShaderState.PixelShaderRHI = PixelShader.GetPixelShader();
 				SetGraphicsPipelineState(RHICmdList, GraphicsPSOInit);
 				PixelShader->SetParameters(RHICmdList, View, LightSceneInfo);
 			}
@@ -2113,8 +2076,8 @@ void FDeferredShadingSceneRenderer::RenderLight(FRHICommandList& RHICmdList, con
 
 				TShaderMapRef< FDeferredLightPS > PixelShader( View.ShaderMap, PermutationVector );
 				GraphicsPSOInit.BoundShaderState.VertexDeclarationRHI = GFilterVertexDeclaration.VertexDeclarationRHI;
-				GraphicsPSOInit.BoundShaderState.VertexShaderRHI = GETSAFERHISHADER_VERTEX(*VertexShader);
-				GraphicsPSOInit.BoundShaderState.PixelShaderRHI = GETSAFERHISHADER_PIXEL(*PixelShader);
+				GraphicsPSOInit.BoundShaderState.VertexShaderRHI = VertexShader.GetVertexShader();
+				GraphicsPSOInit.BoundShaderState.PixelShaderRHI = PixelShader.GetPixelShader();
 
 				SetGraphicsPipelineState(RHICmdList, GraphicsPSOInit);
 				PixelShader->SetParameters(RHICmdList, View, LightSceneInfo, ScreenShadowMaskTexture, bHairRenderingEnabled ? &RenderLightParams : nullptr);
@@ -2131,7 +2094,7 @@ void FDeferredShadingSceneRenderer::RenderLight(FRHICommandList& RHICmdList, con
 				View.ViewRect.Width(), View.ViewRect.Height(),
 				View.ViewRect.Size(),
 				FSceneRenderTargets::Get(RHICmdList).GetBufferSizeXY(),
-				*VertexShader,
+				VertexShader,
 				EDRF_UseTriangleOptimization);
 		}
 		else
@@ -2148,8 +2111,8 @@ void FDeferredShadingSceneRenderer::RenderLight(FRHICommandList& RHICmdList, con
 			{
 				TShaderMapRef<TDeferredLightOverlapPS<true> > PixelShader(View.ShaderMap);
 				GraphicsPSOInit.BoundShaderState.VertexDeclarationRHI = GetVertexDeclarationFVector4();
-				GraphicsPSOInit.BoundShaderState.VertexShaderRHI = GETSAFERHISHADER_VERTEX(*VertexShader);
-				GraphicsPSOInit.BoundShaderState.PixelShaderRHI = GETSAFERHISHADER_PIXEL(*PixelShader);
+				GraphicsPSOInit.BoundShaderState.VertexShaderRHI = VertexShader.GetVertexShader();
+				GraphicsPSOInit.BoundShaderState.PixelShaderRHI = PixelShader.GetPixelShader();
 
 				SetGraphicsPipelineState(RHICmdList, GraphicsPSOInit);
 				PixelShader->SetParameters(RHICmdList, View, LightSceneInfo);
@@ -2168,8 +2131,8 @@ void FDeferredShadingSceneRenderer::RenderLight(FRHICommandList& RHICmdList, con
 
 				TShaderMapRef< FDeferredLightPS > PixelShader( View.ShaderMap, PermutationVector );
 				GraphicsPSOInit.BoundShaderState.VertexDeclarationRHI = GetVertexDeclarationFVector4();
-				GraphicsPSOInit.BoundShaderState.VertexShaderRHI = GETSAFERHISHADER_VERTEX(*VertexShader);
-				GraphicsPSOInit.BoundShaderState.PixelShaderRHI = GETSAFERHISHADER_PIXEL(*PixelShader);
+				GraphicsPSOInit.BoundShaderState.VertexShaderRHI = VertexShader.GetVertexShader();
+				GraphicsPSOInit.BoundShaderState.PixelShaderRHI = PixelShader.GetPixelShader();
 
 				SetGraphicsPipelineState(RHICmdList, GraphicsPSOInit);
 				PixelShader->SetParameters(RHICmdList, View, LightSceneInfo, ScreenShadowMaskTexture, &RenderLightParams);
@@ -2245,12 +2208,12 @@ void FDeferredShadingSceneRenderer::RenderSimpleLightsStandardDeferred(FRHIComma
 			if (SimpleLight.Exponent == 0)
 			{
 				// inverse squared
-				SetShaderTemplLightingSimple<false, true, true>(RHICmdList, GraphicsPSOInit, View, *VertexShader, SimpleLight, SimpleLightPerViewData);
+				SetShaderTemplLightingSimple<false, true, true>(RHICmdList, GraphicsPSOInit, View, VertexShader, SimpleLight, SimpleLightPerViewData);
 			}
 			else
 			{
 				// light's exponent, not inverse squared
-				SetShaderTemplLightingSimple<false, true, false>(RHICmdList, GraphicsPSOInit, View, *VertexShader, SimpleLight, SimpleLightPerViewData);
+				SetShaderTemplLightingSimple<false, true, false>(RHICmdList, GraphicsPSOInit, View, VertexShader, SimpleLight, SimpleLightPerViewData);
 			}
 
 			VertexShader->SetSimpleLightParameters(RHICmdList, View, LightBounds);
diff --git a/Engine/Source/Runtime/Renderer/Private/LightRendering.h b/Engine/Source/Runtime/Renderer/Private/LightRendering.h
--- a/Engine/Source/Runtime/Renderer/Private/LightRendering.h
+++ b/Engine/Source/Runtime/Renderer/Private/LightRendering.h
@@ -134,6 +134,7 @@ void SetSimpleDeferredLightParameters(
 /** Shader parameters needed to render a light function. */
 class FLightFunctionSharedParameters
 {
+	DECLARE_TYPE_LAYOUT(FLightFunctionSharedParameters, NonVirtual);
 public:
 
 	void Bind(const FShaderParameterMap& ParameterMap)
@@ -163,8 +164,9 @@ public:
 	}
 
 private:
-
-	FShaderParameter LightFunctionParameters;
+	
+		LAYOUT_FIELD(FShaderParameter, LightFunctionParameters)
+	
 };
 
 /** Utility functions for drawing a sphere */
@@ -434,6 +436,7 @@ namespace StencilingGeometry
 */
 class FStencilingGeometryShaderParameters
 {
+	DECLARE_TYPE_LAYOUT(FStencilingGeometryShaderParameters, NonVirtual);
 public:
 	void Bind(const FShaderParameterMap& ParameterMap)
 	{
@@ -445,8 +448,8 @@ public:
 
 	void Set(FRHICommandList& RHICmdList, FShader* Shader, const FVector4& InStencilingGeometryPosAndScale) const
 	{
-		SetShaderValue(RHICmdList, Shader->GetVertexShader(), StencilGeometryPosAndScale, InStencilingGeometryPosAndScale);
-		SetShaderValue(RHICmdList, Shader->GetVertexShader(), StencilConeParameters, FVector4(0.0f, 0.0f, 0.0f, 0.0f));
+		SetShaderValue(RHICmdList, RHICmdList.GetBoundVertexShader(), StencilGeometryPosAndScale, InStencilingGeometryPosAndScale);
+		SetShaderValue(RHICmdList, RHICmdList.GetBoundVertexShader(), StencilConeParameters, FVector4(0.0f, 0.0f, 0.0f, 0.0f));
 	}
 
 	void Set(FRHICommandList& RHICmdList, FShader* Shader, const FSceneView& View, const FLightSceneInfo* LightSceneInfo) const
@@ -456,22 +459,22 @@ public:
 			LightSceneInfo->Proxy->GetLightType() == LightType_Rect )
 		{
 			StencilingGeometry::GStencilSphereVertexBuffer.CalcTransform(GeometryPosAndScale, LightSceneInfo->Proxy->GetBoundingSphere(), View.ViewMatrices.GetPreViewTranslation());
-			SetShaderValue(RHICmdList, Shader->GetVertexShader(), StencilGeometryPosAndScale, GeometryPosAndScale);
-			SetShaderValue(RHICmdList, Shader->GetVertexShader(), StencilConeParameters, FVector4(0.0f, 0.0f, 0.0f, 0.0f));
+			SetShaderValue(RHICmdList, RHICmdList.GetBoundVertexShader(), StencilGeometryPosAndScale, GeometryPosAndScale);
+			SetShaderValue(RHICmdList, RHICmdList.GetBoundVertexShader(), StencilConeParameters, FVector4(0.0f, 0.0f, 0.0f, 0.0f));
 		}
 		else if(LightSceneInfo->Proxy->GetLightType() == LightType_Spot)
 		{
-			SetShaderValue(RHICmdList, Shader->GetVertexShader(), StencilConeTransform, LightSceneInfo->Proxy->GetLightToWorld());
+			SetShaderValue(RHICmdList, RHICmdList.GetBoundVertexShader(), StencilConeTransform, LightSceneInfo->Proxy->GetLightToWorld());
 			SetShaderValue(
 				RHICmdList, 
-				Shader->GetVertexShader(),
+				RHICmdList.GetBoundVertexShader(),
 				StencilConeParameters,
 				FVector4(
 					StencilingGeometry::FStencilConeIndexBuffer::NumSides,
 					StencilingGeometry::FStencilConeIndexBuffer::NumSlices,
 					LightSceneInfo->Proxy->GetOuterConeAngle(),
 					LightSceneInfo->Proxy->GetRadius()));
-			SetShaderValue(RHICmdList, Shader->GetVertexShader(), StencilPreViewTranslation, View.ViewMatrices.GetPreViewTranslation());
+			SetShaderValue(RHICmdList, RHICmdList.GetBoundVertexShader(), StencilPreViewTranslation, View.ViewMatrices.GetPreViewTranslation());
 		}
 	}
 
@@ -486,10 +489,12 @@ public:
 	}
 
 private:
-	FShaderParameter StencilGeometryPosAndScale;
-	FShaderParameter StencilConeParameters;
-	FShaderParameter StencilConeTransform;
-	FShaderParameter StencilPreViewTranslation;
+	
+		LAYOUT_FIELD(FShaderParameter, StencilGeometryPosAndScale)
+		LAYOUT_FIELD(FShaderParameter, StencilConeParameters)
+		LAYOUT_FIELD(FShaderParameter, StencilConeTransform)
+		LAYOUT_FIELD(FShaderParameter, StencilPreViewTranslation)
+	
 };
 
 
@@ -519,31 +524,22 @@ public:
 
 	void SetParameters(FRHICommandList& RHICmdList, const FViewInfo& View, const FLightSceneInfo* LightSceneInfo)
 	{
-		FGlobalShader::SetParameters<FViewUniformShaderParameters>(RHICmdList, GetVertexShader(),View.ViewUniformBuffer);
+		FGlobalShader::SetParameters<FViewUniformShaderParameters>(RHICmdList, RHICmdList.GetBoundVertexShader(), View.ViewUniformBuffer);
 		StencilingGeometryParameters.Set(RHICmdList, this, View, LightSceneInfo);
 	}
 
 	void SetSimpleLightParameters(FRHICommandList& RHICmdList, const FViewInfo& View, const FSphere& LightBounds)
 	{
-		FGlobalShader::SetParameters<FViewUniformShaderParameters>(RHICmdList, GetVertexShader(),View.ViewUniformBuffer);
+		FGlobalShader::SetParameters<FViewUniformShaderParameters>(RHICmdList, RHICmdList.GetBoundVertexShader(), View.ViewUniformBuffer);
 
 		FVector4 StencilingSpherePosAndScale;
 		StencilingGeometry::GStencilSphereVertexBuffer.CalcTransform(StencilingSpherePosAndScale, LightBounds, View.ViewMatrices.GetPreViewTranslation());
 		StencilingGeometryParameters.Set(RHICmdList, this, StencilingSpherePosAndScale);
 	}
 
-	//~ Begin FShader Interface
-	virtual bool Serialize(FArchive& Ar) override
-	{
-		bool bShaderHasOutdatedParameters = FGlobalShader::Serialize(Ar);
-		Ar << StencilingGeometryParameters;
-		return bShaderHasOutdatedParameters;
-	}
-	//~ End FShader Interface
-
 private:
 
-	FStencilingGeometryShaderParameters StencilingGeometryParameters;
+	LAYOUT_FIELD(FStencilingGeometryShaderParameters, StencilingGeometryParameters);
 };
 
 extern void SetBoundingGeometryRasterizerAndDepthState(FGraphicsPipelineStateInitializer& GraphicsPSOInit, const FViewInfo& View, const FSphere& LightBounds);
diff --git a/Engine/Source/Runtime/Renderer/Private/LightShaftRendering.cpp b/Engine/Source/Runtime/Renderer/Private/LightShaftRendering.cpp
--- a/Engine/Source/Runtime/Renderer/Private/LightShaftRendering.cpp
+++ b/Engine/Source/Runtime/Renderer/Private/LightShaftRendering.cpp
@@ -108,6 +108,7 @@ static FAutoConsoleVariableRef CVarCacheLightShaftNumSamples(
 /** Light shaft parameters that are shared between multiple pixel shaders. */
 class FLightShaftPixelShaderParameters
 {
+	DECLARE_INLINE_TYPE_LAYOUT(FLightShaftPixelShaderParameters, NonVirtual);
 public:
 	void Bind(const FShaderParameterMap& ParameterMap)
 	{
@@ -251,20 +252,22 @@ public:
 		}
 	}
 private:
-	FShaderParameter TextureSpaceBlurOriginParameter;
-	FShaderParameter WorldSpaceBlurOriginAndRadiusParameter;
-	FShaderParameter LightSourceRadius;
-	FShaderParameter SpotAnglesParameter;
-	FShaderParameter WorldSpaceSpotDirectionParameter;
-	FShaderParameter WorldSpaceCameraPositionParameter;
-	FShaderParameter UVMinMaxParameter;
-	FShaderParameter AspectRatioAndInvAspectRatioParameter;
-	FShaderParameter LightShaftParameters;
-	FShaderParameter BloomTintAndThresholdParameter;
-	FShaderParameter BloomMaxBrightnessParameter;
-	FShaderParameter DistanceFadeParameter;
-	FShaderResourceParameter SourceTextureParameter;
-	FShaderResourceParameter SourceTextureSamplerParameter;
+
+	LAYOUT_FIELD(FShaderParameter, TextureSpaceBlurOriginParameter)
+	LAYOUT_FIELD(FShaderParameter, WorldSpaceBlurOriginAndRadiusParameter)
+	LAYOUT_FIELD(FShaderParameter, LightSourceRadius)
+	LAYOUT_FIELD(FShaderParameter, SpotAnglesParameter)
+	LAYOUT_FIELD(FShaderParameter, WorldSpaceSpotDirectionParameter)
+	LAYOUT_FIELD(FShaderParameter, WorldSpaceCameraPositionParameter)
+	LAYOUT_FIELD(FShaderParameter, UVMinMaxParameter)
+	LAYOUT_FIELD(FShaderParameter, AspectRatioAndInvAspectRatioParameter)
+	LAYOUT_FIELD(FShaderParameter, LightShaftParameters)
+	LAYOUT_FIELD(FShaderParameter, BloomTintAndThresholdParameter)
+	LAYOUT_FIELD(FShaderParameter, BloomMaxBrightnessParameter);
+	LAYOUT_FIELD(FShaderParameter, DistanceFadeParameter)
+	LAYOUT_FIELD(FShaderResourceParameter, SourceTextureParameter)
+	LAYOUT_FIELD(FShaderResourceParameter, SourceTextureSamplerParameter)
+
 };
 
 /*-----------------------------------------------------------------------------
@@ -293,7 +296,7 @@ public:
 	/** Sets shader parameter values */
 	void SetParameters(FRHICommandList& RHICmdList, const FViewInfo& View)
 	{
-		FGlobalShader::SetParameters<FViewUniformShaderParameters>(RHICmdList, GetVertexShader(), View.ViewUniformBuffer);
+		FGlobalShader::SetParameters<FViewUniformShaderParameters>(RHICmdList, RHICmdList.GetBoundVertexShader(), View.ViewUniformBuffer);
 	}
 };
 
@@ -334,33 +337,23 @@ public:
 		SceneTextureParams.Bind(Initializer);
 	}
 
-	/** Serializer */
-	virtual bool Serialize(FArchive& Ar) override
-	{
-		bool bShaderHasOutdatedParameters = FShader::Serialize(Ar);
-		Ar << LightShaftParameters;
-		Ar << SampleOffsetsParameter;
-		Ar << SceneTextureParams;
-		return bShaderHasOutdatedParameters;
-	}
-
 	/** Sets shader parameter values */
 	void SetParameters(FRHICommandList& RHICmdList, const FLightSceneInfo* LightSceneInfo, const FViewInfo& View, TRefCountPtr<IPooledRenderTarget>& PassSource)
 	{
-		FGlobalShader::SetParameters<FViewUniformShaderParameters>(RHICmdList, GetPixelShader(), View.ViewUniformBuffer);
-		LightShaftParameters.SetParameters(RHICmdList, GetPixelShader(), LightSceneInfo, View, PassSource);
+		FGlobalShader::SetParameters<FViewUniformShaderParameters>(RHICmdList, RHICmdList.GetBoundPixelShader(), View.ViewUniformBuffer);
+		LightShaftParameters.SetParameters(RHICmdList, RHICmdList.GetBoundPixelShader(), LightSceneInfo, View, PassSource);
 
 		const FIntPoint BufferSize = FSceneRenderTargets::Get(RHICmdList).GetBufferSizeXY();
 		FVector2D SampleOffsets(1.0f / BufferSize.X, 1.0f / BufferSize.Y);
-		SetShaderValue(RHICmdList, GetPixelShader(),SampleOffsetsParameter,SampleOffsets);
-		SceneTextureParams.Set(RHICmdList, GetPixelShader(), View.FeatureLevel, ESceneTextureSetupMode::All);
+		SetShaderValue(RHICmdList, RHICmdList.GetBoundPixelShader(),SampleOffsetsParameter,SampleOffsets);
+		SceneTextureParams.Set(RHICmdList, RHICmdList.GetBoundPixelShader(), View.FeatureLevel, ESceneTextureSetupMode::All);
 	}
 
 private:
 
-	FLightShaftPixelShaderParameters LightShaftParameters;
-	FShaderParameter SampleOffsetsParameter;
-	FSceneTextureShaderParameters SceneTextureParams;
+	LAYOUT_FIELD(FLightShaftPixelShaderParameters, LightShaftParameters);
+	LAYOUT_FIELD(FShaderParameter, SampleOffsetsParameter);
+	LAYOUT_FIELD(FSceneTextureShaderParameters, SceneTextureParams);
 };
 
 #define IMPLEMENT_LSDOWNSAMPLE_PIXELSHADER_TYPE(LightType,DownsampleValue) \
@@ -404,29 +397,20 @@ public:
 		LightShaftParameters.Bind(Initializer.ParameterMap);
 	}
 
-	/** Serializer */
-	virtual bool Serialize(FArchive& Ar) override
-	{
-		bool bShaderHasOutdatedParameters = FShader::Serialize(Ar);
-		Ar << RadialBlurParameters;
-		Ar << LightShaftParameters;
-		return bShaderHasOutdatedParameters;
-	}
-
 	/** Sets shader parameter values */
 	void SetParameters(FRHICommandList& RHICmdList, const FLightSceneInfo* LightSceneInfo, const FViewInfo& View, int32 PassIndex, TRefCountPtr<IPooledRenderTarget>& PassSource)
 	{
-		FGlobalShader::SetParameters<FViewUniformShaderParameters>(RHICmdList, GetPixelShader(), View.ViewUniformBuffer);
-		LightShaftParameters.SetParameters(RHICmdList, GetPixelShader(), LightSceneInfo, View, PassSource);
+		FGlobalShader::SetParameters<FViewUniformShaderParameters>(RHICmdList, RHICmdList.GetBoundPixelShader(), View.ViewUniformBuffer);
+		LightShaftParameters.SetParameters(RHICmdList, RHICmdList.GetBoundPixelShader(), LightSceneInfo, View, PassSource);
 
 		const FVector4 Parameters(GLightShaftBlurNumSamples, GLightShaftFirstPassDistance, PassIndex);
-		SetShaderValue(RHICmdList, GetPixelShader(), RadialBlurParameters, Parameters);
+		SetShaderValue(RHICmdList, RHICmdList.GetBoundPixelShader(), RadialBlurParameters, Parameters);
 	}
 
 private:
 
-	FShaderParameter RadialBlurParameters;
-	FLightShaftPixelShaderParameters LightShaftParameters;
+	LAYOUT_FIELD(FShaderParameter, RadialBlurParameters);
+	LAYOUT_FIELD(FLightShaftPixelShaderParameters, LightShaftParameters);
 };
 
 IMPLEMENT_SHADER_TYPE(,FBlurLightShaftsPixelShader,TEXT("/Engine/Private/LightShaftShader.usf"),TEXT("BlurLightShaftsMain"),SF_Pixel);
@@ -455,24 +439,16 @@ public:
 		LightShaftParameters.Bind(Initializer.ParameterMap);
 	}
 
-	/** Serializer */
-	virtual bool Serialize(FArchive& Ar) override
-	{
-		bool bShaderHasOutdatedParameters = FShader::Serialize(Ar);
-		Ar << LightShaftParameters;
-		return bShaderHasOutdatedParameters;
-	}
-
 	/** Sets shader parameter values */
 	void SetParameters(FRHICommandList& RHICmdList, const FLightSceneInfo* LightSceneInfo, const FViewInfo& View, TRefCountPtr<IPooledRenderTarget>& PassSource)
 	{
-		FGlobalShader::SetParameters<FViewUniformShaderParameters>(RHICmdList, GetPixelShader(), View.ViewUniformBuffer);
-		LightShaftParameters.SetParameters(RHICmdList, GetPixelShader(), LightSceneInfo, View, PassSource);
+		FGlobalShader::SetParameters<FViewUniformShaderParameters>(RHICmdList, RHICmdList.GetBoundPixelShader(), View.ViewUniformBuffer);
+		LightShaftParameters.SetParameters(RHICmdList, RHICmdList.GetBoundPixelShader(), LightSceneInfo, View, PassSource);
 	}
 
 private:
 
-	FLightShaftPixelShaderParameters LightShaftParameters;
+	LAYOUT_FIELD(FLightShaftPixelShaderParameters, LightShaftParameters);
 };
 
 IMPLEMENT_SHADER_TYPE(,FFinishOcclusionPixelShader,TEXT("/Engine/Private/LightShaftShader.usf"),TEXT("FinishOcclusionMain"),SF_Pixel);
@@ -531,8 +507,8 @@ void DownsamplePass(FRHICommandListImmediate& RHICmdList, const FViewInfo& View,
 		{
 			TShaderMapRef<TDownsampleLightShaftsPixelShader<LightType_Directional, bDownsampleOcclusion> > DownsampleLightShaftsPixelShader(View.ShaderMap);
 			GraphicsPSOInit.BoundShaderState.VertexDeclarationRHI = GFilterVertexDeclaration.VertexDeclarationRHI;
-			GraphicsPSOInit.BoundShaderState.VertexShaderRHI = GETSAFERHISHADER_VERTEX(*DownsampleLightShaftsVertexShader);
-			GraphicsPSOInit.BoundShaderState.PixelShaderRHI = GETSAFERHISHADER_PIXEL(*DownsampleLightShaftsPixelShader);
+			GraphicsPSOInit.BoundShaderState.VertexShaderRHI = DownsampleLightShaftsVertexShader.GetVertexShader();
+			GraphicsPSOInit.BoundShaderState.PixelShaderRHI = DownsampleLightShaftsPixelShader.GetPixelShader();
 			SetGraphicsPipelineState(RHICmdList, GraphicsPSOInit);
 			DownsampleLightShaftsPixelShader->SetParameters(RHICmdList, LightSceneInfo, View, UnusedRT);
 		}
@@ -541,8 +517,8 @@ void DownsamplePass(FRHICommandListImmediate& RHICmdList, const FViewInfo& View,
 		{
 			TShaderMapRef<TDownsampleLightShaftsPixelShader<LightType_Spot, bDownsampleOcclusion> > DownsampleLightShaftsPixelShader(View.ShaderMap);
 			GraphicsPSOInit.BoundShaderState.VertexDeclarationRHI = GFilterVertexDeclaration.VertexDeclarationRHI;
-			GraphicsPSOInit.BoundShaderState.VertexShaderRHI = GETSAFERHISHADER_VERTEX(*DownsampleLightShaftsVertexShader);
-			GraphicsPSOInit.BoundShaderState.PixelShaderRHI = GETSAFERHISHADER_PIXEL(*DownsampleLightShaftsPixelShader);
+			GraphicsPSOInit.BoundShaderState.VertexShaderRHI = DownsampleLightShaftsVertexShader.GetVertexShader();
+			GraphicsPSOInit.BoundShaderState.PixelShaderRHI = DownsampleLightShaftsPixelShader.GetPixelShader();
 			SetGraphicsPipelineState(RHICmdList, GraphicsPSOInit);
 			DownsampleLightShaftsPixelShader->SetParameters(RHICmdList, LightSceneInfo, View, UnusedRT);
 		}
@@ -552,8 +528,8 @@ void DownsamplePass(FRHICommandListImmediate& RHICmdList, const FViewInfo& View,
 		{
 			TShaderMapRef<TDownsampleLightShaftsPixelShader<LightType_Point, bDownsampleOcclusion> > DownsampleLightShaftsPixelShader(View.ShaderMap);
 			GraphicsPSOInit.BoundShaderState.VertexDeclarationRHI = GFilterVertexDeclaration.VertexDeclarationRHI;
-			GraphicsPSOInit.BoundShaderState.VertexShaderRHI = GETSAFERHISHADER_VERTEX(*DownsampleLightShaftsVertexShader);
-			GraphicsPSOInit.BoundShaderState.PixelShaderRHI = GETSAFERHISHADER_PIXEL(*DownsampleLightShaftsPixelShader);
+			GraphicsPSOInit.BoundShaderState.VertexShaderRHI = DownsampleLightShaftsVertexShader.GetVertexShader();
+			GraphicsPSOInit.BoundShaderState.PixelShaderRHI = DownsampleLightShaftsPixelShader.GetPixelShader();
 			SetGraphicsPipelineState(RHICmdList, GraphicsPSOInit);
 			DownsampleLightShaftsPixelShader->SetParameters(RHICmdList, LightSceneInfo, View, UnusedRT);
 		}
@@ -571,7 +547,7 @@ void DownsamplePass(FRHICommandListImmediate& RHICmdList, const FViewInfo& View,
 			View.ViewRect.Width(), View.ViewRect.Height(),
 			FIntPoint(DownsampledSizeX, DownsampledSizeY),
 			BufferSize,
-			*DownsampleLightShaftsVertexShader,
+			DownsampleLightShaftsVertexShader,
 			EDRF_UseTriangleOptimization);
 	}
 	RHICmdList.EndRenderPass();
@@ -666,8 +642,8 @@ void ApplyRadialBlurPasses(
 			TShaderMapRef<FBlurLightShaftsPixelShader> BlurLightShaftsPixelShader(View.ShaderMap);
 
 			GraphicsPSOInit.BoundShaderState.VertexDeclarationRHI = GFilterVertexDeclaration.VertexDeclarationRHI;
-			GraphicsPSOInit.BoundShaderState.VertexShaderRHI = GETSAFERHISHADER_VERTEX(*ScreenVertexShader);
-			GraphicsPSOInit.BoundShaderState.PixelShaderRHI = GETSAFERHISHADER_PIXEL(*BlurLightShaftsPixelShader);
+			GraphicsPSOInit.BoundShaderState.VertexShaderRHI = ScreenVertexShader.GetVertexShader();
+			GraphicsPSOInit.BoundShaderState.PixelShaderRHI = BlurLightShaftsPixelShader.GetPixelShader();
 			GraphicsPSOInit.PrimitiveType = PT_TriangleList;
 
 			SetGraphicsPipelineState(RHICmdList, GraphicsPSOInit);
@@ -685,7 +661,7 @@ void ApplyRadialBlurPasses(
 					DownSampledXY.X, DownSampledXY.Y,
 					DownsampledSizeX, DownsampledSizeY,
 					FilterBufferSize, FilterBufferSize,
-					*ScreenVertexShader,
+					ScreenVertexShader,
 					EDRF_UseTriangleOptimization);
 			}
 		}
@@ -725,8 +701,8 @@ void FinishOcclusionTerm(FRHICommandList& RHICmdList, const FViewInfo& View, con
 		TShaderMapRef<FFinishOcclusionPixelShader> MaskOcclusionTermPixelShader(View.ShaderMap);
 
 		GraphicsPSOInit.BoundShaderState.VertexDeclarationRHI = GFilterVertexDeclaration.VertexDeclarationRHI;
-		GraphicsPSOInit.BoundShaderState.VertexShaderRHI = GETSAFERHISHADER_VERTEX(*ScreenVertexShader);
-		GraphicsPSOInit.BoundShaderState.PixelShaderRHI = GETSAFERHISHADER_PIXEL(*MaskOcclusionTermPixelShader);
+		GraphicsPSOInit.BoundShaderState.VertexShaderRHI = ScreenVertexShader.GetVertexShader();
+		GraphicsPSOInit.BoundShaderState.PixelShaderRHI = MaskOcclusionTermPixelShader.GetPixelShader();
 		GraphicsPSOInit.PrimitiveType = PT_TriangleList;
 
 		SetGraphicsPipelineState(RHICmdList, GraphicsPSOInit);
@@ -742,7 +718,7 @@ void FinishOcclusionTerm(FRHICommandList& RHICmdList, const FViewInfo& View, con
 				DownSampledXY.X, DownSampledXY.Y,
 				DownsampledSizeX, DownsampledSizeY,
 				FilterBufferSize, FilterBufferSize,
-				*ScreenVertexShader,
+				ScreenVertexShader,
 				EDRF_UseTriangleOptimization);
 		}
 	}
@@ -897,24 +873,14 @@ public:
 		UVMinMaxParameter.Bind(Initializer.ParameterMap, TEXT("UVMinMax"));
 	}
 
-	/** Serializer */
-	virtual bool Serialize(FArchive& Ar) override
-	{
-		bool bShaderHasOutdatedParameters = FShader::Serialize(Ar);
-		Ar << SourceTextureParameter;
-		Ar << SourceTextureSamplerParameter;
-		Ar << UVMinMaxParameter;
-		return bShaderHasOutdatedParameters;
-	}
-
 	/** Sets shader parameter values */
 	void SetParameters(FRHICommandList& RHICmdList, const FViewInfo& View, TRefCountPtr<IPooledRenderTarget>& LightShaftOcclusion, const FIntPoint& FilterBufferSize, const FIntRect& DownSampleRect)
 	{
-		FGlobalShader::SetParameters<FViewUniformShaderParameters>(RHICmdList, GetPixelShader(), View.ViewUniformBuffer);
+		FGlobalShader::SetParameters<FViewUniformShaderParameters>(RHICmdList, RHICmdList.GetBoundPixelShader(), View.ViewUniformBuffer);
 
 		SetTextureParameter(
 			RHICmdList, 
-			GetPixelShader(),
+			RHICmdList.GetBoundPixelShader(),
 			SourceTextureParameter, SourceTextureSamplerParameter,
 			TStaticSamplerState<SF_Bilinear,AM_Clamp,AM_Clamp,AM_Clamp>::GetRHI(),
 			LightShaftOcclusion->GetRenderTargetItem().ShaderResourceTexture
@@ -926,13 +892,13 @@ public:
 			(DownSampleRect.Max.X - 0.5f) / float(FilterBufferSize.X),
 			(DownSampleRect.Max.Y - 0.5f) / float(FilterBufferSize.Y));
 
-		SetShaderValue(RHICmdList, GetPixelShader(), UVMinMaxParameter, UVMinMaxValue);
+		SetShaderValue(RHICmdList, RHICmdList.GetBoundPixelShader(), UVMinMaxParameter, UVMinMaxValue);
 	}
 
 private:
-	FShaderResourceParameter SourceTextureParameter;
-	FShaderResourceParameter SourceTextureSamplerParameter;
-	FShaderParameter UVMinMaxParameter;
+	LAYOUT_FIELD(FShaderResourceParameter, SourceTextureParameter);
+	LAYOUT_FIELD(FShaderResourceParameter, SourceTextureSamplerParameter);
+	LAYOUT_FIELD(FShaderParameter, UVMinMaxParameter);
 };
 
 IMPLEMENT_SHADER_TYPE(,FApplyLightShaftsPixelShader,TEXT("/Engine/Private/LightShaftShader.usf"),TEXT("ApplyLightShaftsPixelMain"),SF_Pixel);
@@ -965,8 +931,8 @@ void ApplyLightShaftBloom(FRHICommandListImmediate& RHICmdList, const FViewInfo&
 	TShaderMapRef<FApplyLightShaftsPixelShader> ApplyLightShaftsPixelShader(View.ShaderMap);
 
 	GraphicsPSOInit.BoundShaderState.VertexDeclarationRHI = GFilterVertexDeclaration.VertexDeclarationRHI;
-	GraphicsPSOInit.BoundShaderState.VertexShaderRHI = GETSAFERHISHADER_VERTEX(*ScreenVertexShader);
-	GraphicsPSOInit.BoundShaderState.PixelShaderRHI = GETSAFERHISHADER_PIXEL(*ApplyLightShaftsPixelShader);
+	GraphicsPSOInit.BoundShaderState.VertexShaderRHI = ScreenVertexShader.GetVertexShader();
+	GraphicsPSOInit.BoundShaderState.PixelShaderRHI = ApplyLightShaftsPixelShader.GetPixelShader();
 	GraphicsPSOInit.PrimitiveType = PT_TriangleList;
 
 	SetGraphicsPipelineState(RHICmdList, GraphicsPSOInit);
@@ -993,7 +959,7 @@ void ApplyLightShaftBloom(FRHICommandListImmediate& RHICmdList, const FViewInfo&
 		DownSampledXY.X, DownSampledXY.Y, 
 		DownsampledSizeX, DownsampledSizeY,
 		FIntPoint(View.ViewRect.Width(), View.ViewRect.Height()), FilterBufferSize,
-		*ScreenVertexShader,
+		ScreenVertexShader,
 		EDRF_UseTriangleOptimization);
 
 	if (bUseSeparateTranslucency)
diff --git a/Engine/Source/Runtime/Renderer/Private/MaterialTexCoordScalesRendering.cpp b/Engine/Source/Runtime/Renderer/Private/MaterialTexCoordScalesRendering.cpp
--- a/Engine/Source/Runtime/Renderer/Private/MaterialTexCoordScalesRendering.cpp
+++ b/Engine/Source/Runtime/Renderer/Private/MaterialTexCoordScalesRendering.cpp
@@ -14,7 +14,8 @@ MaterialTexCoordScalesRendering.cpp: Contains definitions for rendering the view
 
 IMPLEMENT_MATERIAL_SHADER_TYPE(,FMaterialTexCoordScalePS,TEXT("/Engine/Private/MaterialTexCoordScalesPixelShader.usf"),TEXT("Main"),SF_Pixel);
 
-void FMaterialTexCoordScalePS::GetDebugViewModeShaderBindings(
+void FMaterialTexCoordScaleBaseInterface::GetDebugViewModeShaderBindings(
+	const FDebugViewModePS& ShaderBase,
 	const FPrimitiveSceneProxy* RESTRICT PrimitiveSceneProxy,
 	const FMaterialRenderProxy& RESTRICT MaterialRenderProxy,
 	const FMaterial& RESTRICT Material,
@@ -29,6 +30,7 @@ void FMaterialTexCoordScalePS::GetDebugViewModeShaderBindings(
 	FMeshDrawSingleShaderBindings& ShaderBindings
 ) const 
 {
+	const FMaterialTexCoordScalePS& Shader = static_cast<const FMaterialTexCoordScalePS&>(ShaderBase);
 	const int32 AnalysisIndex = ViewModeParam >= 0 ? FMath::Clamp<int32>(ViewModeParam, 0, TEXSTREAM_MAX_NUM_TEXTURES_PER_MATERIAL - 1) : -1;
 	FVector4 OneOverCPUTexCoordScales[TEXSTREAM_MAX_NUM_TEXTURES_PER_MATERIAL / 4];
 	FIntVector4 TexCoordIndices[TEXSTREAM_MAX_NUM_TEXTURES_PER_MATERIAL / 4];
@@ -42,10 +44,10 @@ void FMaterialTexCoordScalePS::GetDebugViewModeShaderBindings(
 #endif
 	const bool bOutputScales = DebugViewMode == DVSM_OutputMaterialTextureScales;
 
-	ShaderBindings.Add(OneOverCPUTexCoordScalesParameter, OneOverCPUTexCoordScales);
-	ShaderBindings.Add(TexCoordIndicesParameter, TexCoordIndices);
-	ShaderBindings.Add(AnalysisParamsParameter, FIntPoint(bOutputScales ? -1 : AnalysisIndex, bOutputScales ? 1 : 0));
-	ShaderBindings.Add(PrimitiveAlphaParameter,  (!PrimitiveSceneProxy || PrimitiveSceneProxy->IsSelected()) ? 1.f : .2f);
+	ShaderBindings.Add(Shader.OneOverCPUTexCoordScalesParameter, OneOverCPUTexCoordScales);
+	ShaderBindings.Add(Shader.TexCoordIndicesParameter, TexCoordIndices);
+	ShaderBindings.Add(Shader.AnalysisParamsParameter, FIntPoint(bOutputScales ? -1 : AnalysisIndex, bOutputScales ? 1 : 0));
+	ShaderBindings.Add(Shader.PrimitiveAlphaParameter,  (!PrimitiveSceneProxy || PrimitiveSceneProxy->IsSelected()) ? 1.f : .2f);
 }
 
 void FOutputMaterialTexCoordScaleInterface::SetDrawRenderState(EBlendMode BlendMode, FRenderState& DrawRenderState, bool bHasDepthPrepassForMaskedMaterial) const
diff --git a/Engine/Source/Runtime/Renderer/Private/MaterialTexCoordScalesRendering.h b/Engine/Source/Runtime/Renderer/Private/MaterialTexCoordScalesRendering.h
--- a/Engine/Source/Runtime/Renderer/Private/MaterialTexCoordScalesRendering.h
+++ b/Engine/Source/Runtime/Renderer/Private/MaterialTexCoordScalesRendering.h
@@ -33,8 +33,7 @@ public:
 
 	static bool ShouldCompilePermutation(const FMeshMaterialShaderPermutationParameters& Parameters)
 	{
-		// See FDebugViewModeMaterialProxy::GetFriendlyName()
-		return AllowDebugViewShaderMode(DVSM_OutputMaterialTextureScales, Parameters.Platform, Parameters.Material->GetFeatureLevel()) && Parameters.Material->GetFriendlyName().Contains(TEXT("MaterialTexCoordScale"));
+		return Parameters.MaterialParameters.bIsMaterialTexCoordScale && AllowDebugViewShaderMode(DVSM_OutputMaterialTextureScales, Parameters.Platform, Parameters.MaterialParameters.FeatureLevel);
 	}
 
 	FMaterialTexCoordScalePS(const ShaderMetaType::CompiledShaderInitializerType& Initializer):
@@ -48,16 +47,6 @@ public:
 
 	FMaterialTexCoordScalePS() {}
 
-	virtual bool Serialize(FArchive& Ar) override
-	{
-		bool bShaderHasOutdatedParameters = FMeshMaterialShader::Serialize(Ar);
-		Ar << AnalysisParamsParameter;
-		Ar << OneOverCPUTexCoordScalesParameter;
-		Ar << TexCoordIndicesParameter;
-		Ar << PrimitiveAlphaParameter;
-		return bShaderHasOutdatedParameters;
-	}
-
 	static void ModifyCompilationEnvironment(const FMaterialShaderPermutationParameters& Parameters, FShaderCompilerEnvironment& OutEnvironment)
 	{
 		OutEnvironment.SetDefine(TEXT("UNDEFINED_ACCURACY"), UndefinedStreamingAccuracyIntensity);
@@ -65,10 +54,31 @@ public:
 		OutEnvironment.SetDefine(TEXT("INITIAL_GPU_SCALE"), (uint32)TEXSTREAM_INITIAL_GPU_SCALE);
 		OutEnvironment.SetDefine(TEXT("TILE_RESOLUTION"), (uint32)TEXSTREAM_TILE_RESOLUTION);
 		OutEnvironment.SetDefine(TEXT("MAX_NUM_TEXTURE_REGISTER"), (uint32)TEXSTREAM_MAX_NUM_TEXTURES_PER_MATERIAL);
-		FMeshMaterialShader::ModifyCompilationEnvironment(Parameters.Platform, OutEnvironment);
+		FMeshMaterialShader::ModifyCompilationEnvironment(Parameters, OutEnvironment);
+	}
+
+	LAYOUT_FIELD(FShaderParameter, AnalysisParamsParameter);
+	LAYOUT_FIELD(FShaderParameter, OneOverCPUTexCoordScalesParameter);
+	LAYOUT_FIELD(FShaderParameter, TexCoordIndicesParameter);
+	LAYOUT_FIELD(FShaderParameter, PrimitiveAlphaParameter);
+};
+
+class FMaterialTexCoordScaleBaseInterface : public FDebugViewModeInterface
+{
+public:
+	FMaterialTexCoordScaleBaseInterface(
+		bool InNeedsOnlyLocalVertexFactor,  // Whether this viewmode will only be used with the local vertex factory (for draw tiled mesh).
+		bool InNeedsMaterialProperties,		// Whether the PS use any of the material properties (otherwise default material will be used, reducing shader compilation).
+		bool InNeedsInstructionCount		// Whether FDebugViewModePS::GetDebugViewModeShaderBindings() will use the num of instructions.
+	) : FDebugViewModeInterface(TEXT("MaterialTexCoordScale"), InNeedsOnlyLocalVertexFactor, InNeedsMaterialProperties, InNeedsInstructionCount) {}
+
+	virtual TShaderRef<FDebugViewModePS> GetPixelShader(const FMaterial* InMaterial, FVertexFactoryType* VertexFactoryType) const override
+	{
+		return InMaterial->GetShader<FMaterialTexCoordScalePS>(VertexFactoryType);
 	}
 
 	virtual void GetDebugViewModeShaderBindings(
+		const FDebugViewModePS& BaseShader,
 		const FPrimitiveSceneProxy* RESTRICT PrimitiveSceneProxy,
 		const FMaterialRenderProxy& RESTRICT MaterialRenderProxy,
 		const FMaterial& RESTRICT Material,
@@ -82,35 +92,18 @@ public:
 		FName ViewModeParamName,
 		FMeshDrawSingleShaderBindings& ShaderBindings
 	) const override;
-
-private:
-
-	FShaderParameter AnalysisParamsParameter;
-	FShaderParameter OneOverCPUTexCoordScalesParameter;
-	FShaderParameter TexCoordIndicesParameter;
-	FShaderParameter PrimitiveAlphaParameter;
 };
 
-class FMaterialTexCoordScaleAccuracyInterface : public FDebugViewModeInterface
+class FMaterialTexCoordScaleAccuracyInterface : public FMaterialTexCoordScaleBaseInterface
 {
 public:
-
-	FMaterialTexCoordScaleAccuracyInterface() : FDebugViewModeInterface(TEXT("MaterialTexCoordScale"), false, true, false) {}
-	virtual FDebugViewModePS* GetPixelShader(const FMaterial* InMaterial, FVertexFactoryType* VertexFactoryType) const override 
-	{ 
-		return InMaterial->GetShader<FMaterialTexCoordScalePS>(VertexFactoryType);
-	}
+	FMaterialTexCoordScaleAccuracyInterface() : FMaterialTexCoordScaleBaseInterface(false, true, false) {}
 };
 
-class FOutputMaterialTexCoordScaleInterface : public FDebugViewModeInterface
+class FOutputMaterialTexCoordScaleInterface : public FMaterialTexCoordScaleBaseInterface
 {
 public:
-
-	FOutputMaterialTexCoordScaleInterface() : FDebugViewModeInterface(TEXT("MaterialTexCoordScale"), true, true, false) {}
-	virtual FDebugViewModePS* GetPixelShader(const FMaterial* InMaterial, FVertexFactoryType* VertexFactoryType) const override 
-	{ 
-		return InMaterial->GetShader<FMaterialTexCoordScalePS>(VertexFactoryType); 
-	}
+	FOutputMaterialTexCoordScaleInterface() : FMaterialTexCoordScaleBaseInterface(true, true, false) {}
 	virtual void SetDrawRenderState(EBlendMode BlendMode, FRenderState& DrawRenderState, bool bHasDepthPrepassForMaskedMaterial) const override;
 };
 
diff --git a/Engine/Source/Runtime/Renderer/Private/MeshDrawCommands.cpp b/Engine/Source/Runtime/Renderer/Private/MeshDrawCommands.cpp
--- a/Engine/Source/Runtime/Renderer/Private/MeshDrawCommands.cpp
+++ b/Engine/Source/Runtime/Renderer/Private/MeshDrawCommands.cpp
@@ -504,7 +504,8 @@ void GenerateDynamicMeshDrawCommands(
 	int32 MaxNumBuildRequestElements,
 	FMeshCommandOneFrameArray& VisibleCommands,
 	FDynamicMeshDrawCommandStorage& MeshDrawCommandStorage,
-	FGraphicsMinimalPipelineStateSet& MinimalPipelineStatePassSet
+	FGraphicsMinimalPipelineStateSet& MinimalPipelineStatePassSet,
+	bool& NeedsShaderInitialisation
 )
 {
 	QUICK_SCOPE_CYCLE_COUNTER(STAT_GenerateDynamicMeshDrawCommands);
@@ -514,7 +515,8 @@ void GenerateDynamicMeshDrawCommands(
 	FDynamicPassMeshDrawListContext DynamicPassMeshDrawListContext(
 		MeshDrawCommandStorage,
 		VisibleCommands,
-		MinimalPipelineStatePassSet
+		MinimalPipelineStatePassSet,
+		NeedsShaderInitialisation
 	);
 	PassMeshProcessor->SetDrawListContext(&DynamicPassMeshDrawListContext);
 
@@ -574,7 +576,8 @@ void GenerateMobileBasePassDynamicMeshDrawCommands(
 	int32 MaxNumBuildRequestElements,
 	FMeshCommandOneFrameArray& VisibleCommands,
 	FDynamicMeshDrawCommandStorage& MeshDrawCommandStorage,
-	FGraphicsMinimalPipelineStateSet& GraphicsMinimalPipelineStateSet
+	FGraphicsMinimalPipelineStateSet& GraphicsMinimalPipelineStateSet,
+	bool& NeedsShaderInitialisation
 )
 {
 	QUICK_SCOPE_CYCLE_COUNTER(STAT_GenerateMobileBasePassDynamicMeshDrawCommands);
@@ -584,7 +587,8 @@ void GenerateMobileBasePassDynamicMeshDrawCommands(
 	FDynamicPassMeshDrawListContext DynamicPassMeshDrawListContext(
 		MeshDrawCommandStorage,
 		VisibleCommands,
-		GraphicsMinimalPipelineStateSet
+		GraphicsMinimalPipelineStateSet,
+		NeedsShaderInitialisation
 	);
 	PassMeshProcessor->SetDrawListContext(&DynamicPassMeshDrawListContext);
 	MobilePassCSMPassMeshProcessor->SetDrawListContext(&DynamicPassMeshDrawListContext);
@@ -662,6 +666,7 @@ void ApplyViewOverridesToMeshDrawCommands(
 	FMeshCommandOneFrameArray& VisibleMeshDrawCommands,
 	FDynamicMeshDrawCommandStorage& MeshDrawCommandStorage,
 	FGraphicsMinimalPipelineStateSet& MinimalPipelineStatePassSet,
+	bool& NeedsShaderInitialisation,
 	FMeshCommandOneFrameArray& TempVisibleMeshDrawCommands
 	)
 {
@@ -693,7 +698,7 @@ void ApplyViewOverridesToMeshDrawCommands(
 					PipelineState.DepthStencilState = PassDrawRenderState.GetDepthStencilState();
 				}
 
-				const FGraphicsMinimalPipelineStateId PipelineId = FGraphicsMinimalPipelineStateId::GetPipelineStateId(PipelineState, MinimalPipelineStatePassSet);
+				const FGraphicsMinimalPipelineStateId PipelineId = FGraphicsMinimalPipelineStateId::GetPipelineStateId(PipelineState, MinimalPipelineStatePassSet, NeedsShaderInitialisation);
 				NewMeshCommand.Finalize(PipelineId, nullptr);
 
 				FVisibleMeshDrawCommand NewVisibleMeshDrawCommand;
@@ -752,6 +757,7 @@ public:
 
 	void AnyThreadTask()
 	{
+		TRACE_CPUPROFILER_EVENT_SCOPE(MeshDrawCommandPassSetupTask);
 		// Mobile base pass is a special case, as final lists is created from two mesh passes based on CSM visibility.
 		const bool bMobileBasePass = Context.ShadingPath == EShadingPath::Mobile && Context.PassType == EMeshPass::BasePass;
 
@@ -777,7 +783,8 @@ public:
 				Context.NumDynamicMeshCommandBuildRequestElements,
 				Context.MeshDrawCommands,
 				Context.MeshDrawCommandStorage,
-				Context.MinimalPipelineStatePassSet
+				Context.MinimalPipelineStatePassSet,
+				Context.NeedsShaderInitialisation
 			);
 		}
 		else
@@ -794,7 +801,8 @@ public:
 				Context.NumDynamicMeshCommandBuildRequestElements,
 				Context.MeshDrawCommands,
 				Context.MeshDrawCommandStorage,
-				Context.MinimalPipelineStatePassSet
+				Context.MinimalPipelineStatePassSet,
+				Context.NeedsShaderInitialisation
 			);
 		}
 
@@ -812,6 +820,7 @@ public:
 					Context.MeshDrawCommands,
 					Context.MeshDrawCommandStorage,
 					Context.MinimalPipelineStatePassSet,
+					Context.NeedsShaderInitialisation,
 					Context.TempVisibleMeshDrawCommands
 				);
 			}
@@ -868,6 +877,54 @@ public:
 	}
 
 
+private:
+	FMeshDrawCommandPassSetupTaskContext& Context;
+};
+
+/**
+ * Task for shader initialization. This will run on the RenderThread after Commands have been generated.
+ */
+class FMeshDrawCommandInitResourcesTask
+{
+public:
+	FMeshDrawCommandInitResourcesTask(FMeshDrawCommandPassSetupTaskContext& InContext)
+		: Context(InContext)
+	{
+	}
+
+	FORCEINLINE TStatId GetStatId() const
+	{
+		RETURN_QUICK_DECLARE_CYCLE_STAT(FMeshDrawCommandInitResourcesTask, STATGROUP_TaskGraphTasks);
+	}
+
+	ENamedThreads::Type GetDesiredThread()
+	{
+		return ENamedThreads::GetRenderThread_Local();
+	}
+
+	static ESubsequentsMode::Type GetSubsequentsMode()
+	{
+		return ESubsequentsMode::TrackSubsequents;
+	}
+
+	void AnyThreadTask()
+	{
+		TRACE_CPUPROFILER_EVENT_SCOPE(MeshDrawCommandInitResourcesTask);
+		if (Context.NeedsShaderInitialisation)
+		{
+			for (const FGraphicsMinimalPipelineStateInitializer& Initializer : Context.MinimalPipelineStatePassSet)
+			{
+				Initializer.BoundShaderState.LazilyInitShaders();
+			}
+		}
+	}
+
+	void DoTask(ENamedThreads::Type CurrentThread, const FGraphEventRef& MyCompletionGraphEvent)
+	{
+		AnyThreadTask();
+	}
+
+
 private:
 	FMeshDrawCommandPassSetupTaskContext& Context;
 };
@@ -942,6 +999,7 @@ void FParallelMeshDrawCommandPass::DispatchPassSetup(
 	FMeshCommandOneFrameArray* InOutMobileBasePassCSMMeshDrawCommands
 )
 {
+	TRACE_CPUPROFILER_EVENT_SCOPE(ParallelMdcDispatchPassSetup);
 	check(!TaskEventRef.IsValid() && MeshPassProcessor != nullptr && TaskContext.PrimitiveIdBufferData == nullptr);
 	check((PassType == EMeshPass::Num) == (DynamicMeshElementsPassRelevance == nullptr));
 
@@ -1013,14 +1071,17 @@ void FParallelMeshDrawCommandPass::DispatchPassSetup(
 
 		if (bExecuteInParallel)
 		{
-			TaskEventRef = TGraphTask<FMeshDrawCommandPassSetupTask>::CreateTask(
-				nullptr, ENamedThreads::GetRenderThread()).ConstructAndDispatchWhenReady(TaskContext);
+			FGraphEventArray DependentGraphEvents;
+			DependentGraphEvents.Add(TGraphTask<FMeshDrawCommandPassSetupTask>::CreateTask(nullptr, ENamedThreads::GetRenderThread()).ConstructAndDispatchWhenReady(TaskContext));
+			TaskEventRef = TGraphTask<FMeshDrawCommandInitResourcesTask>::CreateTask(&DependentGraphEvents, ENamedThreads::GetRenderThread()).ConstructAndDispatchWhenReady(TaskContext);
 		}
 		else
 		{
 			QUICK_SCOPE_CYCLE_COUNTER(STAT_MeshPassSetupImmediate);
 			FMeshDrawCommandPassSetupTask Task(TaskContext);
 			Task.AnyThreadTask();
+			FMeshDrawCommandInitResourcesTask DependentTask(TaskContext);
+			DependentTask.AnyThreadTask();
 		}
 	}
 }
@@ -1126,6 +1187,7 @@ public:
 
 	void DoTask(ENamedThreads::Type CurrentThread, const FGraphEventRef& MyCompletionGraphEvent)
 	{
+		TRACE_CPUPROFILER_EVENT_SCOPE(DrawVisibleMeshCommandsAnyThreadTask);
 		checkSlow(RHICmdList.IsInsideRenderPass());
 
 		// Recompute draw range.
@@ -1143,6 +1205,7 @@ public:
 
 void FParallelMeshDrawCommandPass::DispatchDraw(FParallelCommandListSet* ParallelCommandListSet, FRHICommandList& RHICmdList) const
 {
+	TRACE_CPUPROFILER_EVENT_SCOPE(ParallelMdcDispatchDraw);
 	if (MaxNumDraws <= 0)
 	{
 		return;
diff --git a/Engine/Source/Runtime/Renderer/Private/MeshDrawCommands.h b/Engine/Source/Runtime/Renderer/Private/MeshDrawCommands.h
--- a/Engine/Source/Runtime/Renderer/Private/MeshDrawCommands.h
+++ b/Engine/Source/Runtime/Renderer/Private/MeshDrawCommands.h
@@ -62,6 +62,7 @@ public:
 		, InstanceFactor(1)
 		, NumDynamicMeshElements(0)
 		, NumDynamicMeshCommandBuildRequestElements(0)
+		, NeedsShaderInitialisation(false)
 		, PrimitiveIdBufferData(nullptr)
 		, PrimitiveIdBufferDataSize(0)
 		, PrimitiveBounds(nullptr)
@@ -98,6 +99,7 @@ public:
 	TArray<const FStaticMeshBatch*, SceneRenderingAllocator> MobileBasePassCSMDynamicMeshCommandBuildRequests;
 	FDynamicMeshDrawCommandStorage MeshDrawCommandStorage;
 	FGraphicsMinimalPipelineStateSet MinimalPipelineStatePassSet;
+	bool NeedsShaderInitialisation;
 
 	// Resources preallocated on rendering thread.
 	void* PrimitiveIdBufferData;
@@ -163,6 +165,10 @@ public:
 	void SetDumpInstancingStats(const FString& InPassName);
 	bool HasAnyDraw() const { return MaxNumDraws > 0; }
 
+	void InitCreateSnapshot()
+	{
+		new (&TaskContext.MinimalPipelineStatePassSet) FGraphicsMinimalPipelineStateSet();
+	}
 
 private:
 	FRHIVertexBuffer* PrimitiveIdVertexBufferRHI;
