Copying Private-LoadTimes-4.24 stream to Main. Biggest changes are in Materials/Shader memory freezing.

diff --git a/Engine/Source/Runtime/Renderer/Private/MeshPassProcessor.cpp b/Engine/Source/Runtime/Renderer/Private/MeshPassProcessor.cpp
--- a/Engine/Source/Runtime/Renderer/Private/MeshPassProcessor.cpp
+++ b/Engine/Source/Runtime/Renderer/Private/MeshPassProcessor.cpp
@@ -14,10 +14,14 @@
 #include "MeshPassProcessor.inl"
 #include "PipelineStateCache.h"
 #include "RayTracing/RayTracingMaterialHitShaders.h"
+#include "Hash/CityHash.h"
+
+FCriticalSection FGraphicsMinimalPipelineStateId::PersistentIdTableLock;
+FGraphicsMinimalPipelineStateId::PersistentTableType FGraphicsMinimalPipelineStateId::PersistentIdTable;
 
-TSet<FRefCountedGraphicsMinimalPipelineStateInitializer, RefCountedGraphicsMinimalPipelineStateInitializerKeyFuncs> FGraphicsMinimalPipelineStateId::PersistentIdTable;
 int32 FGraphicsMinimalPipelineStateId::LocalPipelineIdTableSize = 0;
 int32 FGraphicsMinimalPipelineStateId::CurrentLocalPipelineIdTableSize = 0;
+bool FGraphicsMinimalPipelineStateId::NeedsShaderInitialisation = true;
 
 const FMeshDrawCommandSortKey FMeshDrawCommandSortKey::Default = { {0} };
 
@@ -183,7 +187,7 @@ void FMeshDrawShaderBindings::SetShaderBindings(
 		{
 			RHICmdList.SetShaderParameter(
 				Shader,
-				LooseParameterBuffer.BufferIndex,
+				LooseParameterBuffer.BaseIndex,
 				Parameter.BaseIndex,
 				Parameter.Size,
 				LooseDataStart
@@ -256,7 +260,7 @@ void FMeshDrawShaderBindings::SetShaderBindings(
 		{
 			RHICmdList.SetShaderParameter(
 				Shader,
-				LooseParameterBuffer.BufferIndex,
+				LooseParameterBuffer.BaseIndex,
 				Parameter.BaseIndex,
 				Parameter.Size,
 				LooseDataStart
@@ -300,7 +304,7 @@ void FMeshDrawShaderBindings::SetRayTracingShaderBindingsForHitGroup(
 
 	const uint32 NumUniformBuffersToSet = MaxUniformBufferUsed + 1;
 
-	const TArray<FShaderLooseParameterBufferInfo>& LooseParameterBuffers = SingleShaderBindings.ParameterMapInfo.LooseParameterBuffers;
+	const TMemoryImageArray<FShaderLooseParameterBufferInfo>& LooseParameterBuffers = SingleShaderBindings.ParameterMapInfo.LooseParameterBuffers;
 	uint32 LooseParameterDataSize = 0;
 
 	if (LooseParameterBuffers.Num())
@@ -308,7 +312,7 @@ void FMeshDrawShaderBindings::SetRayTracingShaderBindingsForHitGroup(
 		check(LooseParameterBuffers.Num() <= 1);
 
 		const FShaderLooseParameterBufferInfo& LooseParameterBuffer = SingleShaderBindings.ParameterMapInfo.LooseParameterBuffers[0];
-		check(LooseParameterBuffer.BufferIndex == 0);
+		check(LooseParameterBuffer.BaseIndex == 0);
 
 		for (int32 LooseParameterIndex = 0; LooseParameterIndex < LooseParameterBuffer.Parameters.Num(); LooseParameterIndex++)
 		{
@@ -353,73 +357,110 @@ void FMeshDrawShaderBindings::SetRayTracingShaderBindingsForHitGroup(
 
 FGraphicsMinimalPipelineStateId FGraphicsMinimalPipelineStateId::GetPersistentId(const FGraphicsMinimalPipelineStateInitializer& InPipelineState)
 {
-	checkSlow(IsInRenderingThread());
-
-	FSetElementId TableId = PersistentIdTable.FindId(InPipelineState);
-	if (TableId.IsValidId())
-	{
-		++PersistentIdTable[TableId].RefNum;
-	}
-	else
+	Experimental::FHashElementId TableId;
+	auto hash = PersistentIdTable.ComputeHash(InPipelineState);
 	{
-		TableId = PersistentIdTable.Add(FRefCountedGraphicsMinimalPipelineStateInitializer(InPipelineState, 1));
+		FScopeLock Lock(&PersistentIdTableLock);
+
+#if UE_BUILD_DEBUG
+		FGraphicsMinimalPipelineStateInitializer PipelineStateDebug = FGraphicsMinimalPipelineStateInitializer(InPipelineState);
+		check(GetTypeHash(PipelineStateDebug) == GetTypeHash(InPipelineState));
+		check(PipelineStateDebug == InPipelineState);
+#endif
+
+		TableId = PersistentIdTable.FindOrAddIdByHash(hash, InPipelineState, FRefCountedGraphicsMinimalPipelineState());
+		FRefCountedGraphicsMinimalPipelineState& Value = PersistentIdTable.GetByElementId(TableId).Value;
+		if (Value.RefNum == 0 && !NeedsShaderInitialisation)
+		{
+			NeedsShaderInitialisation = true;
+		}
+		Value.RefNum++;
 	}
 
-	checkf(TableId.AsInteger() < (MAX_uint32 >> 2), TEXT("Persistent FGraphicsMinimalPipelineStateId table overflow!"));
+	checkf(TableId.GetIndex() < (MAX_uint32 >> 2), TEXT("Persistent FGraphicsMinimalPipelineStateId table overflow!"));
 
 	FGraphicsMinimalPipelineStateId Ret;
 	Ret.bValid = 1;
 	Ret.bComesFromLocalPipelineStateSet = 0;
-	Ret.SetElementIndex = TableId.AsInteger();
+	Ret.SetElementIndex = TableId.GetIndex();
 	return Ret;
 }
 
+
+void FGraphicsMinimalPipelineStateId::InitializePersistentIds()
+{
+	TRACE_CPUPROFILER_EVENT_SCOPE(InitializePersistentMdcIds);
+	FScopeLock Lock(&PersistentIdTableLock);
+	if (NeedsShaderInitialisation)
+	{
+		for (TPair<const FGraphicsMinimalPipelineStateInitializer, FRefCountedGraphicsMinimalPipelineState>& Element : PersistentIdTable)
+		{
+			Element.Key.BoundShaderState.LazilyInitShaders();
+		}
+		NeedsShaderInitialisation = false;
+	}
+}
+
 void FGraphicsMinimalPipelineStateId::RemovePersistentId(FGraphicsMinimalPipelineStateId Id)
 {
 	check(!Id.bComesFromLocalPipelineStateSet && Id.bValid);
 
-	const FSetElementId SetElementId = FSetElementId::FromInteger(Id.SetElementIndex);
-	FRefCountedGraphicsMinimalPipelineStateInitializer& RefCountedStateInitializer = PersistentIdTable[SetElementId];
-
-	check(RefCountedStateInitializer.RefNum > 0);
-	--RefCountedStateInitializer.RefNum;
-	if (RefCountedStateInitializer.RefNum <= 0)
 	{
-		PersistentIdTable.Remove(SetElementId);
+		FScopeLock Lock(&PersistentIdTableLock);
+		FRefCountedGraphicsMinimalPipelineState& RefCountedStateInitializer = PersistentIdTable.GetByElementId(Id.SetElementIndex).Value;
+
+		check(RefCountedStateInitializer.RefNum > 0);
+		--RefCountedStateInitializer.RefNum;
+		if (RefCountedStateInitializer.RefNum == 0)
+		{
+			PersistentIdTable.RemoveByElementId(Id.SetElementIndex);
+		}
 	}
 }
 
-FGraphicsMinimalPipelineStateId FGraphicsMinimalPipelineStateId::GetPipelineStateId(const FGraphicsMinimalPipelineStateInitializer& InPipelineState, FGraphicsMinimalPipelineStateSet& InOutPassSet)
+FGraphicsMinimalPipelineStateId FGraphicsMinimalPipelineStateId::GetPipelineStateId(const FGraphicsMinimalPipelineStateInitializer& InPipelineState, FGraphicsMinimalPipelineStateSet& InOutPassSet, bool& InNeedsShaderInitialisation)
 {
 	FGraphicsMinimalPipelineStateId Ret;
 	Ret.bValid = 1;
 	Ret.bComesFromLocalPipelineStateSet = 0;
 
-	FSetElementId TableId = PersistentIdTable.FindId(InPipelineState);
-	if (!TableId.IsValidId())
+	Experimental::FHashElementId TableIndex;
+	{
+		FScopeLock Lock(&PersistentIdTableLock);
+		TableIndex = PersistentIdTable.FindId(InPipelineState);
+	}
+
+	if (!TableIndex.IsValid())
 	{
 		Ret.bComesFromLocalPipelineStateSet = 1;
-		TableId = InOutPassSet.FindId(InPipelineState);
-		if (!TableId.IsValidId())
-		{
-			TableId = InOutPassSet.Add(InPipelineState);
-		}
+#if UE_BUILD_DEBUG
+		FGraphicsMinimalPipelineStateInitializer PipelineStateDebug = FGraphicsMinimalPipelineStateInitializer(InPipelineState);
+		check(GetTypeHash(PipelineStateDebug) == GetTypeHash(InPipelineState));
+		check(PipelineStateDebug == InPipelineState);
+#endif
+		TableIndex = InOutPassSet.FindOrAddId(InPipelineState);
+#if UE_BUILD_DEBUG
+		check(InOutPassSet.GetByElementId(TableIndex) == InPipelineState);
+#endif
+		InNeedsShaderInitialisation = InNeedsShaderInitialisation || InPipelineState.BoundShaderState.NeedsShaderInitialisation();
 	}
 
-	checkf(TableId.AsInteger() < (MAX_uint32 >> 2), TEXT("One frame FGraphicsMinimalPipelineStateId table overflow!"));
+	checkf(TableIndex.GetIndex() < (MAX_uint32 >> 2), TEXT("One frame FGraphicsMinimalPipelineStateId table overflow!"));
 
-	Ret.SetElementIndex = TableId.AsInteger();
+	Ret.SetElementIndex = TableIndex.GetIndex();
 	return Ret;
 }
 
 void FGraphicsMinimalPipelineStateId::ResetLocalPipelineIdTableSize()
 {
+	FScopeLock Lock(&PersistentIdTableLock);
 	LocalPipelineIdTableSize = CurrentLocalPipelineIdTableSize;
 	CurrentLocalPipelineIdTableSize = 0;
 }
 
 void FGraphicsMinimalPipelineStateId::AddSizeToLocalPipelineIdTableSize(SIZE_T Size)
 {
+	FScopeLock Lock(&PersistentIdTableLock);
 	CurrentLocalPipelineIdTableSize += int32(Size);
 }
 
@@ -493,56 +534,75 @@ FMeshDrawShaderBindings::~FMeshDrawShaderBindings()
 
 void FMeshDrawShaderBindings::Initialize(FMeshProcessorShaders Shaders)
 {
-	const int32 NumShaderFrequencies = (Shaders.VertexShader ? 1 : 0) + (Shaders.HullShader ? 1 : 0) + (Shaders.DomainShader ? 1 : 0) + (Shaders.PixelShader ? 1 : 0) + (Shaders.GeometryShader ? 1 : 0) + (Shaders.ComputeShader ? 1 : 0)
+	const int32 NumShaderFrequencies = (Shaders.VertexShader.IsValid() ? 1 : 0) +
+		(Shaders.HullShader.IsValid() ? 1 : 0) +
+		(Shaders.DomainShader.IsValid() ? 1 : 0) +
+		(Shaders.PixelShader.IsValid() ? 1 : 0) +
+		(Shaders.GeometryShader.IsValid() ? 1 : 0) +
+		(Shaders.ComputeShader.IsValid() ? 1 : 0)
 #if RHI_RAYTRACING
-		+ (Shaders.RayHitGroupShader ? 1 : 0)
+		+ (Shaders.RayHitGroupShader.IsValid() ? 1 : 0)
 #endif
 		;
 
 	ShaderLayouts.Empty(NumShaderFrequencies);
 	int32 ShaderBindingDataSize = 0;
 
-	if (Shaders.VertexShader)
+	if (Shaders.VertexShader.IsValid())
 	{
 		ShaderLayouts.Add(FMeshDrawShaderBindingsLayout(Shaders.VertexShader));
 		ShaderBindingDataSize += ShaderLayouts.Last().GetDataSizeBytes();
+		check(ShaderFrequencyBits < (1 << SF_Vertex));
+		ShaderFrequencyBits |= (1 << SF_Vertex);
 	}
 
-	if (Shaders.HullShader)
+	if (Shaders.HullShader.IsValid())
 	{
 		ShaderLayouts.Add(FMeshDrawShaderBindingsLayout(Shaders.HullShader));
 		ShaderBindingDataSize += ShaderLayouts.Last().GetDataSizeBytes();
+		check(ShaderFrequencyBits < (1 << SF_Hull));
+		ShaderFrequencyBits |= (1 << SF_Hull);
 	}
 
-	if (Shaders.DomainShader)
+	if (Shaders.DomainShader.IsValid())
 	{
 		ShaderLayouts.Add(FMeshDrawShaderBindingsLayout(Shaders.DomainShader));
 		ShaderBindingDataSize += ShaderLayouts.Last().GetDataSizeBytes();
+		check(ShaderFrequencyBits < (1 << SF_Domain));
+		ShaderFrequencyBits |= (1 << SF_Domain);
 	}
 
-	if (Shaders.PixelShader)
+	if (Shaders.PixelShader.IsValid())
 	{
 		ShaderLayouts.Add(FMeshDrawShaderBindingsLayout(Shaders.PixelShader));
 		ShaderBindingDataSize += ShaderLayouts.Last().GetDataSizeBytes();
+		check(ShaderFrequencyBits < (1 << SF_Pixel));
+		ShaderFrequencyBits |= (1 << SF_Pixel);
 	}
 
-	if (Shaders.GeometryShader)
+	if (Shaders.GeometryShader.IsValid())
 	{
 		ShaderLayouts.Add(FMeshDrawShaderBindingsLayout(Shaders.GeometryShader));
 		ShaderBindingDataSize += ShaderLayouts.Last().GetDataSizeBytes();
+		check(ShaderFrequencyBits < (1 << SF_Geometry));
+		ShaderFrequencyBits |= (1 << SF_Geometry);
 	}
 
-	if (Shaders.ComputeShader)
+	if (Shaders.ComputeShader.IsValid())
 	{
 		ShaderLayouts.Add(FMeshDrawShaderBindingsLayout(Shaders.ComputeShader));
 		ShaderBindingDataSize += ShaderLayouts.Last().GetDataSizeBytes();
+		check(ShaderFrequencyBits < (1 << SF_Compute));
+		ShaderFrequencyBits |= (1 << SF_Compute);
 	}
 
 #if RHI_RAYTRACING
-	if (Shaders.RayHitGroupShader)
+	if (Shaders.RayHitGroupShader.IsValid())
 	{
 		ShaderLayouts.Add(FMeshDrawShaderBindingsLayout(Shaders.RayHitGroupShader));
 		ShaderBindingDataSize += ShaderLayouts.Last().GetDataSizeBytes();
+		check(ShaderFrequencyBits < (1 << SF_RayHitGroup));
+		ShaderFrequencyBits |= (1 << SF_RayHitGroup);
 	}
 #endif
 
@@ -563,13 +623,26 @@ void FMeshDrawShaderBindings::Finalize(const FMeshProcessorShaders* ShadersForDe
 	}
 
 	const uint8* ShaderBindingDataPtr = GetData();
-
+	uint32 ShaderFrequencyBitIndex = ~0;
 	for (int32 ShaderBindingsIndex = 0; ShaderBindingsIndex < ShaderLayouts.Num(); ShaderBindingsIndex++)
 	{
+		EShaderFrequency Frequency = SF_NumFrequencies;
+		while (true)
+		{
+			ShaderFrequencyBitIndex++;
+			if ((ShaderFrequencyBits & (1 << ShaderFrequencyBitIndex)) != 0)
+			{
+				Frequency = EShaderFrequency(ShaderFrequencyBitIndex);
+				break;
+			}
+		}
+		check(Frequency < SF_NumFrequencies);
+
 		const FMeshDrawShaderBindingsLayout& ShaderLayout = ShaderLayouts[ShaderBindingsIndex];
 
-		FMeshMaterialShader* Shader = ShadersForDebugging->GetShader(ShaderLayout.Frequency);
-		check(Shader);
+		TShaderRef<FMeshMaterialShader> Shader = ShadersForDebugging->GetShader(Frequency);
+		check(Shader.IsValid());
+		const FVertexFactoryType* VFType = Shader.GetVertexFactoryType();
 
 		FReadOnlyMeshDrawSingleShaderBindings SingleShaderBindings(ShaderLayout, ShaderBindingDataPtr);
 
@@ -589,16 +662,16 @@ void FMeshDrawShaderBindings::Finalize(const FMeshProcessorShaders* ShadersForDe
 				if (AutomaticallyBoundUniformBufferStruct)
 				{
 					ensureMsgf(UniformBufferValue, TEXT("Shader %s with vertex factory %s never set automatically bound uniform buffer at BaseIndex %i.  Expected buffer of type %s.  This can cause GPU hangs, depending on how the shader uses it."),
-						Shader->GetType()->GetName(), 
-						Shader->GetVertexFactoryType()->GetName(),
+						Shader.GetType()->GetName(), 
+						VFType ? VFType->GetName() : TEXT("nullptr"),
 						ParameterInfo.BaseIndex,
 						AutomaticallyBoundUniformBufferStruct->GetStructTypeName());
 				}
 				else
 				{
 					ensureMsgf(UniformBufferValue, TEXT("Shader %s with vertex factory %s never set uniform buffer at BaseIndex %i.  This can cause GPU hangs, depending on how the shader uses it."), 
-						Shader->GetVertexFactoryType()->GetName(),
-						Shader->GetType()->GetName(), 
+						VFType ? VFType->GetName() : TEXT("nullptr"),
+						Shader.GetType()->GetName(), 
 						ParameterInfo.BaseIndex);
 				}
 			}
@@ -611,8 +684,8 @@ void FMeshDrawShaderBindings::Finalize(const FMeshProcessorShaders* ShadersForDe
 			FShaderParameterInfo ParameterInfo = ShaderLayout.ParameterMapInfo.TextureSamplers[BindingIndex];
 			const FRHISamplerState* SamplerValue = SamplerBindings[BindingIndex];
 			ensureMsgf(SamplerValue, TEXT("Shader %s with vertex factory %s never set sampler at BaseIndex %u.  This can cause GPU hangs, depending on how the shader uses it."), 
-				Shader->GetType()->GetName(), 
-				Shader->GetVertexFactoryType()->GetName(),
+				Shader.GetType()->GetName(), 
+				VFType ? VFType->GetName() : TEXT("nullptr"),
 				ParameterInfo.BaseIndex);
 		}
 
@@ -633,8 +706,8 @@ void FMeshDrawShaderBindings::Finalize(const FMeshProcessorShaders* ShadersForDe
 				FRHIShaderResourceView* SRV = (FRHIShaderResourceView*)SRVBindings[SRVIndex];
 
 				ensureMsgf(SRV, TEXT("Shader %s with vertex factory %s never set SRV at BaseIndex %u.  This can cause GPU hangs, depending on how the shader uses it."), 
-					Shader->GetType()->GetName(), 
-					Shader->GetVertexFactoryType()->GetName(),
+					Shader.GetType()->GetName(), 
+					VFType ? VFType->GetName() : TEXT("nullptr"),
 					Parameter.BaseIndex);
 			}
 			else
@@ -642,8 +715,8 @@ void FMeshDrawShaderBindings::Finalize(const FMeshProcessorShaders* ShadersForDe
 				FRHITexture* Texture = (FRHITexture*)SRVBindings[SRVIndex];
 
 				ensureMsgf(Texture, TEXT("Shader %s with vertex factory %s never set texture at BaseIndex %u.  This can cause GPU hangs, depending on how the shader uses it."), 
-					Shader->GetType()->GetName(), 
-					Shader->GetVertexFactoryType()->GetName(),
+					Shader.GetType()->GetName(), 
+					VFType ? VFType->GetName() : TEXT("nullptr"),
 					Parameter.BaseIndex);
 			}
 		}
@@ -656,11 +729,19 @@ void FMeshDrawShaderBindings::Finalize(const FMeshProcessorShaders* ShadersForDe
 void FMeshDrawShaderBindings::CopyFrom(const FMeshDrawShaderBindings& Other)
 {
 	Release();
-
 	ShaderLayouts = Other.ShaderLayouts;
+	ShaderFrequencyBits = Other.ShaderFrequencyBits;
 
 	Allocate(Other.Size);
-	FPlatformMemory::Memcpy(GetData(), Other.GetData(), Size);
+
+	if (Other.UsesInlineStorage())
+	{
+		Data = Other.Data;
+	}
+	else
+	{
+		FPlatformMemory::Memcpy(GetData(), Other.GetData(), Size);
+	}
 
 #if VALIDATE_UNIFORM_BUFFER_LIFETIME
 	uint8* ShaderBindingDataPtr = GetData();
@@ -712,37 +793,60 @@ void FMeshDrawShaderBindings::Release()
 	}
 #endif
 
-	if (Size > UE_ARRAY_COUNT(InlineStorage))
+	if (Size > sizeof(FData))
 	{
-		delete[] HeapData;
+		delete[] Data.GetHeapData();
 	}
 	Size = 0;
-	HeapData = nullptr;
+	Data.SetHeapData(nullptr);
 }
 
 void FMeshDrawCommand::SetShaders(FRHIVertexDeclaration* VertexDeclaration, const FMeshProcessorShaders& Shaders, FGraphicsMinimalPipelineStateInitializer& PipelineState)
 {
-	PipelineState.BoundShaderState = FBoundShaderStateInput(
-		VertexDeclaration
-		, GETSAFERHISHADER_VERTEX(Shaders.VertexShader)
-#if PLATFORM_SUPPORTS_TESSELLATION_SHADERS
-		, GETSAFERHISHADER_HULL(Shaders.HullShader)
-		, GETSAFERHISHADER_DOMAIN(Shaders.DomainShader)
-#endif
-		, GETSAFERHISHADER_PIXEL(Shaders.PixelShader)
+	PipelineState.BoundShaderState = FMinimalBoundShaderStateInput();
+	PipelineState.BoundShaderState.VertexDeclarationRHI = VertexDeclaration;
+	if(Shaders.VertexShader.IsValid())
+	{
+		checkSlow(Shaders.VertexShader->GetFrequency() == SF_Vertex);
+		PipelineState.BoundShaderState.VertexShaderResource = Shaders.VertexShader.GetResource();
+		PipelineState.BoundShaderState.VertexShaderIndex = Shaders.VertexShader->GetResourceIndex();
+	}
+	if (Shaders.PixelShader.IsValid())
+	{
+		checkSlow(Shaders.PixelShader->GetFrequency() == SF_Pixel);
+		PipelineState.BoundShaderState.PixelShaderResource = Shaders.PixelShader.GetResource();
+		PipelineState.BoundShaderState.PixelShaderIndex = Shaders.PixelShader->GetResourceIndex();
+	}
 #if PLATFORM_SUPPORTS_GEOMETRY_SHADERS
-		, GETSAFERHISHADER_GEOMETRY(Shaders.GeometryShader)
-#endif
-	);
-
+	if (Shaders.GeometryShader.IsValid())
+	{
+		checkSlow(Shaders.GeometryShader->GetFrequency() == SF_Geometry);
+		PipelineState.BoundShaderState.GeometryShaderResource = Shaders.GeometryShader.GetResource();
+		PipelineState.BoundShaderState.GeometryShaderIndex = Shaders.GeometryShader->GetResourceIndex();
+	}
+#endif // PLATFORM_SUPPORTS_GEOMETRY_SHADERS
+#if PLATFORM_SUPPORTS_TESSELLATION_SHADERS
+	if (Shaders.HullShader.IsValid())
+	{
+		checkSlow(Shaders.HullShader->GetFrequency() == SF_Hull);
+		PipelineState.BoundShaderState.HullShaderResource = Shaders.HullShader.GetResource();
+		PipelineState.BoundShaderState.HullShaderIndex = Shaders.HullShader->GetResourceIndex();
+	}
+	if (Shaders.DomainShader.IsValid())
+	{
+		checkSlow(Shaders.DomainShader->GetFrequency() == SF_Domain);
+		PipelineState.BoundShaderState.DomainShaderResource = Shaders.DomainShader.GetResource();
+		PipelineState.BoundShaderState.DomainShaderIndex = Shaders.DomainShader->GetResourceIndex();
+	}
+#endif // PLATFORM_SUPPORTS_TESSELLATION_SHADERS
 	ShaderBindings.Initialize(Shaders);
 }
 
 #if RHI_RAYTRACING
 void FRayTracingMeshCommand::SetShaders(const FMeshProcessorShaders& Shaders)
 {
-	check(Shaders.RayHitGroupShader)
-	MaterialShaderIndex = Shaders.RayHitGroupShader->GetRayTracingMaterialLibraryIndex();
+	check(Shaders.RayHitGroupShader.IsValid())
+	MaterialShaderIndex = Shaders.RayHitGroupShader.GetRayTracingMaterialLibraryIndex();
 	ShaderBindings.Initialize(Shaders);
 }
 #endif // RHI_RAYTRACING
@@ -757,7 +861,7 @@ void FMeshDrawCommand::SetDrawParametersAndFinalize(
 
 	check(!BatchElement.IndexBuffer || (BatchElement.IndexBuffer && BatchElement.IndexBuffer->IsInitialized() && BatchElement.IndexBuffer->IndexBufferRHI));
 	checkSlow(!BatchElement.bIsInstanceRuns);
-	IndexBuffer = BatchElement.IndexBuffer ? BatchElement.IndexBuffer->IndexBufferRHI : nullptr;
+	IndexBuffer = BatchElement.IndexBuffer ? BatchElement.IndexBuffer->IndexBufferRHI.GetReference() : nullptr;
 	FirstIndex = BatchElement.FirstIndex;
 	NumPrimitives = BatchElement.NumPrimitives;
 	NumInstances = BatchElement.NumInstances;
@@ -781,11 +885,22 @@ void FMeshDrawCommand::SetDrawParametersAndFinalize(
 void FMeshDrawShaderBindings::SetOnCommandList(FRHICommandList& RHICmdList, FBoundShaderStateInput Shaders, FShaderBindingState* StateCacheShaderBindings) const
 {
 	const uint8* ShaderBindingDataPtr = GetData();
-
+	uint32 ShaderFrequencyBitIndex = ~0;
 	for (int32 ShaderBindingsIndex = 0; ShaderBindingsIndex < ShaderLayouts.Num(); ShaderBindingsIndex++)
 	{
+		EShaderFrequency Frequency = SF_NumFrequencies;
+		while (true)
+		{
+			ShaderFrequencyBitIndex++;
+			if ((ShaderFrequencyBits & (1 << ShaderFrequencyBitIndex)) != 0)
+			{
+				Frequency = EShaderFrequency(ShaderFrequencyBitIndex);
+				break;
+			}
+		}
+		check(Frequency < SF_NumFrequencies);
+
 		FReadOnlyMeshDrawSingleShaderBindings SingleShaderBindings(ShaderLayouts[ShaderBindingsIndex], ShaderBindingDataPtr);
-		const EShaderFrequency Frequency = SingleShaderBindings.Frequency;
 		FShaderBindingState& ShaderBindingState = StateCacheShaderBindings[Frequency];
 
 		if (Frequency == SF_Vertex)
@@ -821,19 +936,31 @@ void FMeshDrawShaderBindings::SetOnCommandList(FRHIComputeCommandList& RHICmdLis
 {
 	check(ShaderLayouts.Num() == 1);
 	FReadOnlyMeshDrawSingleShaderBindings SingleShaderBindings(ShaderLayouts[0], GetData());
-	check(SingleShaderBindings.Frequency == SF_Compute);
+	check(ShaderFrequencyBits & (1 << SF_Compute));
 
 	SetShaderBindings(RHICmdList, Shader, SingleShaderBindings);
 }
 
 bool FMeshDrawShaderBindings::MatchesForDynamicInstancing(const FMeshDrawShaderBindings& Rhs) const
 {
-	if (!(ShaderLayouts == Rhs.ShaderLayouts
-		&& Size == Rhs.Size))
+	if (ShaderFrequencyBits != Rhs.ShaderFrequencyBits)
 	{
 		return false;
 	}
 
+	if (ShaderLayouts.Num() != Rhs.ShaderLayouts.Num())
+	{
+		return false;
+	}
+
+	for (int Index = 0; Index < ShaderLayouts.Num(); Index++)
+	{
+		if (!(ShaderLayouts[Index] == Rhs.ShaderLayouts[Index]))
+		{
+			return false;
+		}
+	}
+
 	const uint8* ShaderBindingDataPtr = GetData();
 	const uint8* OtherShaderBindingDataPtr = Rhs.GetData();
 
@@ -842,15 +969,49 @@ bool FMeshDrawShaderBindings::MatchesForDynamicInstancing(const FMeshDrawShaderB
 		FReadOnlyMeshDrawSingleShaderBindings SingleShaderBindings(ShaderLayouts[ShaderBindingsIndex], ShaderBindingDataPtr);
 		FReadOnlyMeshDrawSingleShaderBindings OtherSingleShaderBindings(Rhs.ShaderLayouts[ShaderBindingsIndex], OtherShaderBindingDataPtr);
 
-		if (SingleShaderBindings.ParameterMapInfo.SRVs.Num() > 0 || SingleShaderBindings.ParameterMapInfo.LooseParameterBuffers.Num() > 0 || SingleShaderBindings.ParameterMapInfo.TextureSamplers.Num() > 0)
+		if (SingleShaderBindings.ParameterMapInfo.LooseParameterBuffers.Num())
 		{
-			// Not implemented. Note: this must match with GetDynamicInstancingHash.
-			return false;
+			const uint8* LooseBindings = SingleShaderBindings.GetLooseDataStart();
+			const uint8* OtherLooseBindings = OtherSingleShaderBindings.GetLooseDataStart();
+
+			uint32 Length = 0;
+			for (const FShaderLooseParameterBufferInfo& ParamInfo : SingleShaderBindings.ParameterMapInfo.LooseParameterBuffers)
+			{
+				Length = FMath::Max<uint32>(Length, ParamInfo.BaseIndex + ParamInfo.Size);
+			}
+
+			if (memcmp(LooseBindings, OtherLooseBindings, Length) != 0)
+			{
+				return false;
+			}
+		}
+
+		FRHISamplerState* const* SamplerBindings = SingleShaderBindings.GetSamplerStart();
+		FRHISamplerState* const* OtherSamplerBindings = OtherSingleShaderBindings.GetSamplerStart();
+		for (int32 SamplerIndex = 0; SamplerIndex < SingleShaderBindings.ParameterMapInfo.TextureSamplers.Num(); SamplerIndex++)
+		{
+			const FRHIResource* Sampler = SamplerBindings[SamplerIndex];
+			const FRHIResource* OtherSampler = OtherSamplerBindings[SamplerIndex];
+			if (Sampler != OtherSampler)
+			{
+				return false;
+			}
+		}
+
+		FRHIResource* const* SrvBindings = SingleShaderBindings.GetSRVStart();
+		FRHIResource* const* OtherSrvBindings = SingleShaderBindings.GetSRVStart();
+		for (int32 SrvIndex = 0; SrvIndex < SingleShaderBindings.ParameterMapInfo.SRVs.Num(); SrvIndex++)
+		{
+			const FRHIResource* Srv = SrvBindings[SrvIndex];
+			const FRHIResource* OtherSrv = OtherSrvBindings[SrvIndex];
+			if (Srv != OtherSrv)
+			{
+				return false;
+			}
 		}
 
 		FRHIUniformBuffer* const* UniformBufferBindings = SingleShaderBindings.GetUniformBufferStart();
 		FRHIUniformBuffer* const* OtherUniformBufferBindings = OtherSingleShaderBindings.GetUniformBufferStart();
-
 		for (int32 UniformBufferIndex = 0; UniformBufferIndex < SingleShaderBindings.ParameterMapInfo.UniformBuffers.Num(); UniformBufferIndex++)
 		{
 			const FRHIUniformBuffer* UniformBuffer = UniformBufferBindings[UniformBufferIndex];
@@ -871,37 +1032,74 @@ bool FMeshDrawShaderBindings::MatchesForDynamicInstancing(const FMeshDrawShaderB
 
 uint32 FMeshDrawShaderBindings::GetDynamicInstancingHash() const
 {
-	uint32 Hash = FCrc::TypeCrc32(Size, 0);
-
-	for (const FMeshDrawShaderBindingsLayout& MeshDrawShaderBindingsLayout: ShaderLayouts)
+	//add and initialize any leftover padding within the struct to avoid unstable keys
+	struct FHashKey
 	{
-		Hash = HashCombine(MeshDrawShaderBindingsLayout.GetHash(), Hash);
-	}
+		uint32 LooseParametersHash = 0;
+		uint32 UniformBufferHash = 0;
+		uint16 Size;
+		uint16 Frequencies;
+		static inline uint32 PointerHash(const void* Key)
+		{
+#if PLATFORM_64BITS
+			// Ignoring the lower 4 bits since they are likely zero anyway.
+			// Higher bits are more significant in 64 bit builds.
+			return reinterpret_cast<UPTRINT>(Key) >> 4;
+#else
+			return reinterpret_cast<UPTRINT>(Key);
+#endif
+		};
+
+		static inline uint32 HashCombine(uint32 A, uint32 B)
+		{
+			return A ^ (B + 0x9e3779b9 + (A << 6) + (A >> 2));
+		}
+	} HashKey;
+
+	HashKey.Size = Size;
+	HashKey.Frequencies = ShaderFrequencyBits;
 
 	const uint8* ShaderBindingDataPtr = GetData();
 	for (int32 ShaderBindingsIndex = 0; ShaderBindingsIndex < ShaderLayouts.Num(); ShaderBindingsIndex++)
 	{
 		FReadOnlyMeshDrawSingleShaderBindings SingleShaderBindings(ShaderLayouts[ShaderBindingsIndex], ShaderBindingDataPtr);
 
-		if (SingleShaderBindings.ParameterMapInfo.SRVs.Num() > 0 || SingleShaderBindings.ParameterMapInfo.LooseParameterBuffers.Num() > 0 || SingleShaderBindings.ParameterMapInfo.TextureSamplers.Num() > 0)
+		if (SingleShaderBindings.ParameterMapInfo.LooseParameterBuffers.Num())
+		{
+			const uint8* LooseBindings = SingleShaderBindings.GetLooseDataStart();
+			uint32 Length = 0;
+			for (const FShaderLooseParameterBufferInfo& ParamInfo : SingleShaderBindings.ParameterMapInfo.LooseParameterBuffers)
+			{
+				Length = FMath::Max<uint32>(Length, ParamInfo.BaseIndex + ParamInfo.Size);
+			}
+			HashKey.LooseParametersHash = uint32(CityHash64((const char*)LooseBindings, Length));
+		}
+
+		FRHISamplerState* const* SamplerBindings = SingleShaderBindings.GetSamplerStart();
+		for (int32 SamplerIndex = 0; SamplerIndex < SingleShaderBindings.ParameterMapInfo.TextureSamplers.Num(); SamplerIndex++)
 		{
-			// Since this is not implemented, we must return a unique hash to minimize hash collisions.
-			// Note: this must match with MatchesForDynamicInstancing.
-			Hash = PointerHash(this, Hash);
-			return Hash;
+			const FRHIResource* Sampler = SamplerBindings[SamplerIndex];
+			HashKey.LooseParametersHash = FHashKey::HashCombine(FHashKey::PointerHash(Sampler), HashKey.LooseParametersHash);
+		}
+
+		FRHIResource* const* SrvBindings = SingleShaderBindings.GetSRVStart();
+		for (int32 SrvIndex = 0; SrvIndex < SingleShaderBindings.ParameterMapInfo.SRVs.Num(); SrvIndex++)
+		{
+			const FRHIResource* Srv = SrvBindings[SrvIndex];
+			HashKey.LooseParametersHash = FHashKey::HashCombine(FHashKey::PointerHash(Srv), HashKey.LooseParametersHash);
 		}
 
 		FRHIUniformBuffer* const* UniformBufferBindings = SingleShaderBindings.GetUniformBufferStart();
 		for (int32 UniformBufferIndex = 0; UniformBufferIndex < SingleShaderBindings.ParameterMapInfo.UniformBuffers.Num(); UniformBufferIndex++)
 		{
 			const FRHIUniformBuffer* UniformBuffer = UniformBufferBindings[UniformBufferIndex];
-			Hash = PointerHash(UniformBuffer, Hash);
+			HashKey.UniformBufferHash = FHashKey::HashCombine(FHashKey::PointerHash(UniformBuffer), HashKey.UniformBufferHash);
 		}
 
 		ShaderBindingDataPtr += ShaderLayouts[ShaderBindingsIndex].GetDataSizeBytes();
 	}
 
-	return Hash;
+	return uint32(CityHash64((char*)&HashKey, sizeof(FHashKey)));
 }
 
 void FMeshDrawCommand::SubmitDraw(
@@ -947,12 +1145,11 @@ void FMeshDrawCommand::SubmitDraw(
 	}
 #endif
 
-	{
-		const FGraphicsMinimalPipelineStateInitializer& MeshPipelineState = MeshDrawCommand.CachedPipelineId.GetPipelineState(GraphicsMinimalPipelineStateSet);
+	const FGraphicsMinimalPipelineStateInitializer& MeshPipelineState = MeshDrawCommand.CachedPipelineId.GetPipelineState(GraphicsMinimalPipelineStateSet);
 
 	if (MeshDrawCommand.CachedPipelineId.GetId() != StateCache.PipelineId)
 	{
-		FGraphicsPipelineStateInitializer GraphicsPSOInit = MeshPipelineState;
+		FGraphicsPipelineStateInitializer GraphicsPSOInit = MeshPipelineState.AsGraphicsPipelineStateInitializer();
 		RHICmdList.ApplyCachedRenderTargets(GraphicsPSOInit);
 		SetGraphicsPipelineState(RHICmdList, GraphicsPSOInit);
 		StateCache.SetPipelineState(MeshDrawCommand.CachedPipelineId.GetId());
@@ -980,8 +1177,7 @@ void FMeshDrawCommand::SubmitDraw(
 		}
 	}
 
-	MeshDrawCommand.ShaderBindings.SetOnCommandList(RHICmdList, MeshPipelineState.BoundShaderState, StateCache.ShaderBindings);
-	}
+	MeshDrawCommand.ShaderBindings.SetOnCommandList(RHICmdList, MeshPipelineState.BoundShaderState.AsBoundShaderState(), StateCache.ShaderBindings);
 
 	if (MeshDrawCommand.IndexBuffer)
 	{
@@ -1074,7 +1270,7 @@ void SubmitMeshDrawCommandsRange(
 	}
 }
 
-void ApplyViewOverridesToMeshDrawCommands(const FSceneView& View, FMeshCommandOneFrameArray& VisibleMeshDrawCommands, FDynamicMeshDrawCommandStorage& DynamicMeshDrawCommandStorage, FGraphicsMinimalPipelineStateSet& GraphicsMinimalPipelineStateSet)
+void ApplyViewOverridesToMeshDrawCommands(const FSceneView& View, FMeshCommandOneFrameArray& VisibleMeshDrawCommands, FDynamicMeshDrawCommandStorage& DynamicMeshDrawCommandStorage, FGraphicsMinimalPipelineStateSet& GraphicsMinimalPipelineStateSet, bool& InNeedsShaderInitialisation)
 {
 	if (View.bReverseCulling || View.bRenderSceneTwoSided)
 	{
@@ -1097,7 +1293,7 @@ void ApplyViewOverridesToMeshDrawCommands(const FSceneView& View, FMeshCommandOn
 			FGraphicsMinimalPipelineStateInitializer PipelineState = MeshCommand.CachedPipelineId.GetPipelineState(GraphicsMinimalPipelineStateSet);
 			PipelineState.RasterizerState = GetStaticRasterizerState<true>(VisibleMeshDrawCommand.MeshFillMode, LocalCullMode);
 
-			const FGraphicsMinimalPipelineStateId PipelineId = FGraphicsMinimalPipelineStateId::GetPipelineStateId(PipelineState, GraphicsMinimalPipelineStateSet);
+			const FGraphicsMinimalPipelineStateId PipelineId = FGraphicsMinimalPipelineStateId::GetPipelineStateId(PipelineState, GraphicsMinimalPipelineStateSet, InNeedsShaderInitialisation);
 			NewMeshCommand.Finalize(PipelineId, nullptr);
 
 			FVisibleMeshDrawCommand NewVisibleMeshDrawCommand;
@@ -1125,6 +1321,7 @@ void DrawDynamicMeshPassPrivate(
 	FMeshCommandOneFrameArray& VisibleMeshDrawCommands,
 	FDynamicMeshDrawCommandStorage& DynamicMeshDrawCommandStorage,
 	FGraphicsMinimalPipelineStateSet& GraphicsMinimalPipelineStateSet,
+	bool& InNeedsShaderInitialisation,
 	uint32 InstanceFactor)
 {
 	if (VisibleMeshDrawCommands.Num() > 0)
@@ -1133,7 +1330,7 @@ void DrawDynamicMeshPassPrivate(
 
 		FRHIVertexBuffer* PrimitiveIdVertexBuffer = nullptr;
 
-		ApplyViewOverridesToMeshDrawCommands(View, VisibleMeshDrawCommands, DynamicMeshDrawCommandStorage, GraphicsMinimalPipelineStateSet);
+		ApplyViewOverridesToMeshDrawCommands(View, VisibleMeshDrawCommands, DynamicMeshDrawCommandStorage, GraphicsMinimalPipelineStateSet, InNeedsShaderInitialisation);
 		SortAndMergeDynamicPassMeshDrawCommands(View.GetFeatureLevel(), VisibleMeshDrawCommands, DynamicMeshDrawCommandStorage, PrimitiveIdVertexBuffer, InstanceFactor);
 
 		SubmitMeshDrawCommandsRange(VisibleMeshDrawCommands, GraphicsMinimalPipelineStateSet, PrimitiveIdVertexBuffer, 0, bDynamicInstancing, 0, VisibleMeshDrawCommands.Num(), InstanceFactor, RHICmdList);
@@ -1157,23 +1354,7 @@ FMeshPassProcessor::FMeshPassProcessor(const FScene* InScene, ERHIFeatureLevel::
 {
 }
 
-enum class EDrawingPolicyOverrideFlags
-{
-	None = 0,
-	TwoSided = 1 << 0,
-	DitheredLODTransition = 1 << 1,
-	Wireframe = 1 << 2,
-	ReverseCullMode = 1 << 3,
-};
-ENUM_CLASS_FLAGS(EDrawingPolicyOverrideFlags);
-
-struct FMeshDrawingPolicyOverrideSettings
-{
-	EDrawingPolicyOverrideFlags	MeshOverrideFlags = EDrawingPolicyOverrideFlags::None;
-	EPrimitiveType				MeshPrimitiveType = PT_TriangleList;
-};
-
-FORCEINLINE_DEBUGGABLE FMeshDrawingPolicyOverrideSettings ComputeMeshOverrideSettings(const FMeshBatch& Mesh)
+FMeshPassProcessor::FMeshDrawingPolicyOverrideSettings FMeshPassProcessor::ComputeMeshOverrideSettings(const FMeshBatch& Mesh)
 {
 	FMeshDrawingPolicyOverrideSettings OverrideSettings;
 	OverrideSettings.MeshPrimitiveType = (EPrimitiveType)Mesh.Type;
@@ -1185,18 +1366,15 @@ FORCEINLINE_DEBUGGABLE FMeshDrawingPolicyOverrideSettings ComputeMeshOverrideSet
 	return OverrideSettings;
 }
 
-ERasterizerFillMode FMeshPassProcessor::ComputeMeshFillMode(const FMeshBatch& Mesh, const FMaterial& InMaterialResource) const
+ERasterizerFillMode FMeshPassProcessor::ComputeMeshFillMode(const FMeshBatch& Mesh, const FMaterial& InMaterialResource, const FMeshDrawingPolicyOverrideSettings& InOverrideSettings)
 {
-	const FMeshDrawingPolicyOverrideSettings InOverrideSettings = ComputeMeshOverrideSettings(Mesh);
-
 	const bool bMaterialResourceIsTwoSided = InMaterialResource.IsTwoSided();
 	const bool bIsWireframeMaterial = InMaterialResource.IsWireframe() || !!(InOverrideSettings.MeshOverrideFlags & EDrawingPolicyOverrideFlags::Wireframe);
 	return bIsWireframeMaterial ? FM_Wireframe : FM_Solid;
 }
 
-ERasterizerCullMode FMeshPassProcessor::ComputeMeshCullMode(const FMeshBatch& Mesh, const FMaterial& InMaterialResource) const
+ERasterizerCullMode FMeshPassProcessor::ComputeMeshCullMode(const FMeshBatch& Mesh, const FMaterial& InMaterialResource, const FMeshDrawingPolicyOverrideSettings& InOverrideSettings)
 {
-	const FMeshDrawingPolicyOverrideSettings InOverrideSettings = ComputeMeshOverrideSettings(Mesh);
 	const bool bMaterialResourceIsTwoSided = InMaterialResource.IsTwoSided();
 	const bool bInTwoSidedOverride = !!(InOverrideSettings.MeshOverrideFlags & EDrawingPolicyOverrideFlags::TwoSided);
 	const bool bInReverseCullModeOverride = !!(InOverrideSettings.MeshOverrideFlags & EDrawingPolicyOverrideFlags::ReverseCullMode);
@@ -1234,28 +1412,25 @@ void FMeshPassProcessor::GetDrawCommandPrimitiveId(
 	ScenePrimitiveId = PrimitiveSceneInfo ? PrimitiveSceneInfo->GetIndex() : -1;
 }
 
-FCachedPassMeshDrawListContext::FCachedPassMeshDrawListContext(FCachedMeshDrawCommandInfo& InCommandInfo, FCachedPassMeshDrawList& InDrawList, FScene& InScene) :
+FCachedPassMeshDrawListContext::FCachedPassMeshDrawListContext(FCachedMeshDrawCommandInfo& InCommandInfo, FCriticalSection& InCachedMeshDrawCommandLock, FCachedPassMeshDrawList& InCachedDrawLists, FStateBucketMap& InCachedMeshDrawCommandStateBuckets, const FScene& InScene) :
 	CommandInfo(InCommandInfo),
-	DrawList(InDrawList),
+	CachedMeshDrawCommandLock(InCachedMeshDrawCommandLock),
+	CachedDrawLists(InCachedDrawLists),
+	CachedMeshDrawCommandStateBuckets(InCachedMeshDrawCommandStateBuckets),
 	Scene(InScene)
 {
-	bUseStateBuckets = UseGPUScene(GMaxRHIShaderPlatform, GMaxRHIFeatureLevel);
 }
 
-FMeshDrawCommand& FCachedPassMeshDrawListContext::AddCommand(const FMeshDrawCommand& Initializer)
+FMeshDrawCommand& FCachedPassMeshDrawListContext::AddCommand(FMeshDrawCommand& Initializer, uint32 NumElements)
 {
-	if (bUseStateBuckets)
+	if (NumElements == 1)
 	{
-		MeshDrawCommandForStateBucketing = Initializer;
-		return MeshDrawCommandForStateBucketing;
+		return Initializer;
 	}
 	else
 	{
-		// Only one FMeshDrawCommand supported per FStaticMesh in a pass
-		check(CommandInfo.CommandIndex == -1);
-		// Allocate at lowest free index so that 'r.DoLazyStaticMeshUpdate' can shrink the TSparseArray more effectively
-		CommandInfo.CommandIndex = DrawList.MeshDrawCommands.AddAtLowestFreeIndex(Initializer, DrawList.LowestFreeIndexSearchStart);
-		return DrawList.MeshDrawCommands[CommandInfo.CommandIndex];
+		MeshDrawCommandForStateBucketing = Initializer;
+		return MeshDrawCommandForStateBucketing;
 	}
 }
 
@@ -1274,33 +1449,44 @@ void FCachedPassMeshDrawListContext::FinalizeCommand(
 	// disabling this by default as it incurs a high cost in perf captures due to sheer volume.  Recommendation is to re-enable locally if you need to profile this particular code.
 	// QUICK_SCOPE_CYCLE_COUNTER(STAT_FinalizeCachedMeshDrawCommand);
 
-	FGraphicsMinimalPipelineStateId PipelineId;
-	PipelineId = FGraphicsMinimalPipelineStateId::GetPersistentId(PipelineState);
+	FGraphicsMinimalPipelineStateId PipelineId = FGraphicsMinimalPipelineStateId::GetPersistentId(PipelineState);
 
 	MeshDrawCommand.SetDrawParametersAndFinalize(MeshBatch, BatchElementIndex, PipelineId, ShadersForDebugging);
-	if (bUseStateBuckets)
-	{
-		FSetElementId SetId = Scene.CachedMeshDrawCommandStateBuckets.FindId(MeshDrawCommand);
 
-		if (SetId.IsValidId())
-		{
-			Scene.CachedMeshDrawCommandStateBuckets[SetId].Num++;
-		}
-		else
+	if (UseGPUScene(GMaxRHIShaderPlatform, GMaxRHIFeatureLevel))
+	{
+		Experimental::FHashElementId SetId;
+		auto hash = CachedMeshDrawCommandStateBuckets.ComputeHash(MeshDrawCommand);
 		{
+			FScopeLock Lock(&CachedMeshDrawCommandLock);
+
+#if UE_BUILD_DEBUG
+			FMeshDrawCommand MeshDrawCommandDebug = FMeshDrawCommand(MeshDrawCommand);
+			check(MeshDrawCommandDebug.ShaderBindings.GetDynamicInstancingHash() == MeshDrawCommand.ShaderBindings.GetDynamicInstancingHash());
+			check(MeshDrawCommandDebug.GetDynamicInstancingHash() == MeshDrawCommand.GetDynamicInstancingHash());
+#endif
+			SetId = CachedMeshDrawCommandStateBuckets.FindOrAddIdByHash(hash, MeshDrawCommand, FMeshDrawCommandCount());
+			CachedMeshDrawCommandStateBuckets.GetByElementId(SetId).Value.Num++;
+
 #if MESH_DRAW_COMMAND_DEBUG_DATA
-			MeshDrawCommand.ClearDebugPrimitiveSceneProxy(); //When using State Buckets multiple PrimitiveSceneProxies use the same MeshDrawCommand, so The PrimitiveSceneProxy pointer can't be stored.
+			if (CachedMeshDrawCommandStateBuckets.GetByElementId(SetId).Value.Num == 1)
+			{
+				MeshDrawCommand.ClearDebugPrimitiveSceneProxy(); //When using State Buckets multiple PrimitiveSceneProxies use the same MeshDrawCommand, so The PrimitiveSceneProxy pointer can't be stored.
+			}
 #endif
-			SetId = Scene.CachedMeshDrawCommandStateBuckets.Add(FMeshDrawCommandStateBucket(1, MeshDrawCommand));
 		}
 
 		check(CommandInfo.StateBucketId == -1);
-		CommandInfo.StateBucketId = SetId.AsInteger();
+		CommandInfo.StateBucketId = SetId.GetIndex();
 		check(CommandInfo.CommandIndex == -1);
 	}
 	else
 	{
-		check(CommandInfo.CommandIndex != -1);
+		check(CommandInfo.CommandIndex == -1);
+		FScopeLock Lock(&CachedMeshDrawCommandLock);
+		// Only one FMeshDrawCommand supported per FStaticMesh in a pass
+		// Allocate at lowest free index so that 'r.DoLazyStaticMeshUpdate' can shrink the TSparseArray more effectively
+		CommandInfo.CommandIndex = CachedDrawLists.MeshDrawCommands.AddAtLowestFreeIndex(MeshDrawCommand, CachedDrawLists.LowestFreeIndexSearchStart);
 	}
 
 	CommandInfo.SortKey = SortKey;
diff --git a/Engine/Source/Runtime/Renderer/Private/MeshTexCoordSizeAccuracyRendering.cpp b/Engine/Source/Runtime/Renderer/Private/MeshTexCoordSizeAccuracyRendering.cpp
--- a/Engine/Source/Runtime/Renderer/Private/MeshTexCoordSizeAccuracyRendering.cpp
+++ b/Engine/Source/Runtime/Renderer/Private/MeshTexCoordSizeAccuracyRendering.cpp
@@ -15,7 +15,8 @@ MeshTexCoordSizeAccuracyRendering.cpp: Contains definitions for rendering the vi
 
 IMPLEMENT_SHADER_TYPE(,FMeshTexCoordSizeAccuracyPS,TEXT("/Engine/Private/MeshTexCoordSizeAccuracyPixelShader.usf"),TEXT("Main"),SF_Pixel);
 
-void FMeshTexCoordSizeAccuracyPS::GetDebugViewModeShaderBindings(
+void FMeshTexCoordSizeAccuracyInterface::GetDebugViewModeShaderBindings(
+	const FDebugViewModePS& ShaderBase,
 	const FPrimitiveSceneProxy* RESTRICT PrimitiveSceneProxy,
 	const FMaterialRenderProxy& RESTRICT MaterialRenderProxy,
 	const FMaterial& RESTRICT Material,
@@ -30,6 +31,7 @@ void FMeshTexCoordSizeAccuracyPS::GetDebugViewModeShaderBindings(
 	FMeshDrawSingleShaderBindings& ShaderBindings
 ) const
 {
+	const FMeshTexCoordSizeAccuracyPS& Shader = static_cast<const FMeshTexCoordSizeAccuracyPS&>(ShaderBase);
 	const int32 AnalysisIndex = ViewModeParam >= 0 ? FMath::Clamp<int32>(ViewModeParam, 0, MAX_TEXCOORDS - 1) : -1;
 
 	FVector4 WorldUVDensities;
@@ -40,9 +42,9 @@ void FMeshTexCoordSizeAccuracyPS::GetDebugViewModeShaderBindings(
 		FMemory::Memzero(WorldUVDensities);
 	}
 
-	ShaderBindings.Add(CPUTexelFactorParameter, WorldUVDensities);
-	ShaderBindings.Add(PrimitiveAlphaParameter, (!PrimitiveSceneProxy || PrimitiveSceneProxy->IsSelected()) ? 1.f : .2f);
-	ShaderBindings.Add(TexCoordAnalysisIndexParameter, AnalysisIndex);
+	ShaderBindings.Add(Shader.CPUTexelFactorParameter, WorldUVDensities);
+	ShaderBindings.Add(Shader.PrimitiveAlphaParameter, (!PrimitiveSceneProxy || PrimitiveSceneProxy->IsSelected()) ? 1.f : .2f);
+	ShaderBindings.Add(Shader.TexCoordAnalysisIndexParameter, AnalysisIndex);
 }
 
 
diff --git a/Engine/Source/Runtime/Renderer/Private/MeshTexCoordSizeAccuracyRendering.h b/Engine/Source/Runtime/Renderer/Private/MeshTexCoordSizeAccuracyRendering.h
--- a/Engine/Source/Runtime/Renderer/Private/MeshTexCoordSizeAccuracyRendering.h
+++ b/Engine/Source/Runtime/Renderer/Private/MeshTexCoordSizeAccuracyRendering.h
@@ -31,7 +31,7 @@ public:
 	static bool ShouldCompilePermutation(const FMeshMaterialShaderPermutationParameters& Parameters)
 	{
 		// See FDebugViewModeMaterialProxy::GetFriendlyName()
-		return AllowDebugViewShaderMode(DVSM_MeshUVDensityAccuracy, Parameters.Platform, Parameters.Material->GetFeatureLevel()) && Parameters.Material->GetFriendlyName().Contains(TEXT("MeshTexCoordSizeAccuracy"));
+		return AllowDebugViewShaderMode(DVSM_MeshUVDensityAccuracy, Parameters.Platform, Parameters.MaterialParameters.FeatureLevel) && Parameters.MaterialParameters.bIsMaterialMeshTexCoordSizeAccuracy;
 	}
 
 	FMeshTexCoordSizeAccuracyPS(const ShaderMetaType::CompiledShaderInitializerType& Initializer):
@@ -44,21 +44,24 @@ public:
 
 	FMeshTexCoordSizeAccuracyPS() {}
 
-	virtual bool Serialize(FArchive& Ar) override
-	{
-		bool bShaderHasOutdatedParameters = FMeshMaterialShader::Serialize(Ar);
-		Ar << CPUTexelFactorParameter;
-		Ar << PrimitiveAlphaParameter;
-		Ar << TexCoordAnalysisIndexParameter;
-		return bShaderHasOutdatedParameters;
-	}
-
 	static void ModifyCompilationEnvironment(const FMaterialShaderPermutationParameters& Parameters, FShaderCompilerEnvironment& OutEnvironment)
 	{
 		OutEnvironment.SetDefine(TEXT("UNDEFINED_ACCURACY"), UndefinedStreamingAccuracyIntensity);
 	}
 
+	LAYOUT_FIELD(FShaderParameter, CPUTexelFactorParameter);
+	LAYOUT_FIELD(FShaderParameter, PrimitiveAlphaParameter);
+	LAYOUT_FIELD(FShaderParameter, TexCoordAnalysisIndexParameter);
+};
+
+class FMeshTexCoordSizeAccuracyInterface : public FDebugViewModeInterface
+{
+public:
+	FMeshTexCoordSizeAccuracyInterface() : FDebugViewModeInterface(TEXT("MeshTexCoordSizeAccuracy"), false, false, false) {}
+	virtual TShaderRef<FDebugViewModePS> GetPixelShader(const FMaterial* InMaterial, FVertexFactoryType* VertexFactoryType) const override { return InMaterial->GetShader<FMeshTexCoordSizeAccuracyPS>(VertexFactoryType); }
+
 	virtual void GetDebugViewModeShaderBindings(
+		const FDebugViewModePS& BaseShader,
 		const FPrimitiveSceneProxy* RESTRICT PrimitiveSceneProxy,
 		const FMaterialRenderProxy& RESTRICT MaterialRenderProxy,
 		const FMaterial& RESTRICT Material,
@@ -72,20 +75,6 @@ public:
 		FName ViewModeParamName,
 		FMeshDrawSingleShaderBindings& ShaderBindings
 	) const override;
-
-private:
-
-	FShaderParameter CPUTexelFactorParameter;
-	FShaderParameter PrimitiveAlphaParameter;
-	FShaderParameter TexCoordAnalysisIndexParameter;
-};
-
-class FMeshTexCoordSizeAccuracyInterface : public FDebugViewModeInterface
-{
-public:
-
-	FMeshTexCoordSizeAccuracyInterface() : FDebugViewModeInterface(TEXT("MeshTexCoordSizeAccuracy"), false, false, false) {}
-	virtual FDebugViewModePS* GetPixelShader(const FMaterial* InMaterial, FVertexFactoryType* VertexFactoryType) const override { return InMaterial->GetShader<FMeshTexCoordSizeAccuracyPS>(VertexFactoryType); }
 };
 
 #endif // !(UE_BUILD_SHIPPING || UE_BUILD_TEST)
diff --git a/Engine/Source/Runtime/Renderer/Private/MobileBasePass.cpp b/Engine/Source/Runtime/Renderer/Private/MobileBasePass.cpp
--- a/Engine/Source/Runtime/Renderer/Private/MobileBasePass.cpp
+++ b/Engine/Source/Runtime/Renderer/Private/MobileBasePass.cpp
@@ -18,34 +18,36 @@ void GetUniformMobileBasePassShaders(
 	const FMaterial& Material, 
 	FVertexFactoryType* VertexFactoryType, 
 	bool bEnableSkyLight,
-	TMobileBasePassVSPolicyParamType<FUniformLightMapPolicy>*& VertexShader,
-	TMobileBasePassPSPolicyParamType<FUniformLightMapPolicy>*& PixelShader
+	TShaderRef<TMobileBasePassVSPolicyParamType<FUniformLightMapPolicy>>& VertexShader,
+	TShaderRef<TMobileBasePassPSPolicyParamType<FUniformLightMapPolicy>>& PixelShader
 	)
 {
+	using FVertexShaderType = TMobileBasePassVSPolicyParamType<FUniformLightMapPolicy>;
+	using FPixelShaderType = TMobileBasePassPSPolicyParamType<FUniformLightMapPolicy>;
 	if (IsMobileHDR())
 	{
-		VertexShader = (TMobileBasePassVSPolicyParamType<FUniformLightMapPolicy>*)Material.GetShader<TMobileBasePassVS<TUniformLightMapPolicy<Policy>, HDR_LINEAR_64> >(VertexFactoryType);
+		VertexShader = TShaderRef<FVertexShaderType>::ReinterpretCast(Material.GetShader<TMobileBasePassVS<TUniformLightMapPolicy<Policy>, HDR_LINEAR_64> >(VertexFactoryType));
 
 		if (bEnableSkyLight)
 		{
-			PixelShader = (TMobileBasePassPSPolicyParamType<FUniformLightMapPolicy>*)Material.GetShader< TMobileBasePassPS<TUniformLightMapPolicy<Policy>, HDR_LINEAR_64, true, NumMovablePointLights> >(VertexFactoryType);
+			PixelShader = TShaderRef<FPixelShaderType>::ReinterpretCast(Material.GetShader< TMobileBasePassPS<TUniformLightMapPolicy<Policy>, HDR_LINEAR_64, true, NumMovablePointLights> >(VertexFactoryType));
 		}
 		else
 		{
-			PixelShader = (TMobileBasePassPSPolicyParamType<FUniformLightMapPolicy>*)Material.GetShader< TMobileBasePassPS<TUniformLightMapPolicy<Policy>, HDR_LINEAR_64, false, NumMovablePointLights> >(VertexFactoryType);
+			PixelShader = TShaderRef<FPixelShaderType>::ReinterpretCast(Material.GetShader< TMobileBasePassPS<TUniformLightMapPolicy<Policy>, HDR_LINEAR_64, false, NumMovablePointLights> >(VertexFactoryType));
 		}	
 	}
 	else
 	{
-		VertexShader = (TMobileBasePassVSPolicyParamType<FUniformLightMapPolicy>*)Material.GetShader<TMobileBasePassVS<TUniformLightMapPolicy<Policy>, LDR_GAMMA_32> >(VertexFactoryType);
+		VertexShader = TShaderRef<FVertexShaderType>::ReinterpretCast(Material.GetShader<TMobileBasePassVS<TUniformLightMapPolicy<Policy>, LDR_GAMMA_32> >(VertexFactoryType));
 
 		if (bEnableSkyLight)
 		{
-			PixelShader = (TMobileBasePassPSPolicyParamType<FUniformLightMapPolicy>*)Material.GetShader< TMobileBasePassPS<TUniformLightMapPolicy<Policy>, LDR_GAMMA_32, true, NumMovablePointLights> >(VertexFactoryType);
+			PixelShader = TShaderRef<FPixelShaderType>::ReinterpretCast(Material.GetShader< TMobileBasePassPS<TUniformLightMapPolicy<Policy>, LDR_GAMMA_32, true, NumMovablePointLights> >(VertexFactoryType));
 		}
 		else
 		{
-			PixelShader = (TMobileBasePassPSPolicyParamType<FUniformLightMapPolicy>*)Material.GetShader< TMobileBasePassPS<TUniformLightMapPolicy<Policy>, LDR_GAMMA_32, false, NumMovablePointLights> >(VertexFactoryType);
+			PixelShader = TShaderRef<FPixelShaderType>::ReinterpretCast(Material.GetShader< TMobileBasePassPS<TUniformLightMapPolicy<Policy>, LDR_GAMMA_32, false, NumMovablePointLights> >(VertexFactoryType));
 		}			
 	}
 }
@@ -56,8 +58,8 @@ void GetMobileBasePassShaders(
 	const FMaterial& Material, 
 	FVertexFactoryType* VertexFactoryType, 
 	bool bEnableSkyLight,
-	TMobileBasePassVSPolicyParamType<FUniformLightMapPolicy>*& VertexShader,
-	TMobileBasePassPSPolicyParamType<FUniformLightMapPolicy>*& PixelShader
+	TShaderRef<TMobileBasePassVSPolicyParamType<FUniformLightMapPolicy>>& VertexShader,
+	TShaderRef<TMobileBasePassPSPolicyParamType<FUniformLightMapPolicy>>& PixelShader
 	)
 {
 	switch (LightMapPolicyType)
@@ -106,8 +108,8 @@ void MobileBasePass::GetShaders(
 	const FMaterial& MaterialResource,
 	FVertexFactoryType* VertexFactoryType,
 	bool bEnableSkyLight, 
-	TMobileBasePassVSPolicyParamType<FUniformLightMapPolicy>*& VertexShader,
-	TMobileBasePassPSPolicyParamType<FUniformLightMapPolicy>*& PixelShader)
+	TShaderRef<TMobileBasePassVSPolicyParamType<FUniformLightMapPolicy>>& VertexShader,
+	TShaderRef<TMobileBasePassPSPolicyParamType<FUniformLightMapPolicy>>& PixelShader)
 {
 	bool bIsLit = (MaterialResource.GetShadingModels().IsLit());
 	if (bIsLit && !UseSkylightPermutation(bEnableSkyLight, FReadOnlyCVARCache::Get().MobileSkyLightPermutation))	
@@ -662,8 +664,9 @@ void FMobileBasePassMeshProcessor::Process(
 		SortKey = GetBasePassStaticSortKey(BlendMode, bBackground);
 	}
 	
-	ERasterizerFillMode MeshFillMode = ComputeMeshFillMode(MeshBatch, MaterialResource);
-	ERasterizerCullMode MeshCullMode = ComputeMeshCullMode(MeshBatch, MaterialResource);
+	const FMeshDrawingPolicyOverrideSettings OverrideSettings = ComputeMeshOverrideSettings(MeshBatch);
+	ERasterizerFillMode MeshFillMode = ComputeMeshFillMode(MeshBatch, MaterialResource, OverrideSettings);
+	ERasterizerCullMode MeshCullMode = ComputeMeshCullMode(MeshBatch, MaterialResource, OverrideSettings);
 
 	TMobileBasePassShaderElementData<FUniformLightMapPolicy> ShaderElementData(LightMapElementData);
 	ShaderElementData.InitializeMeshMaterialData(ViewIfDynamicMeshCommand, PrimitiveSceneProxy, MeshBatch, StaticMeshId, false);
diff --git a/Engine/Source/Runtime/Renderer/Private/MobileBasePassRendering.h b/Engine/Source/Runtime/Renderer/Private/MobileBasePassRendering.h
--- a/Engine/Source/Runtime/Renderer/Private/MobileBasePassRendering.h
+++ b/Engine/Source/Runtime/Renderer/Private/MobileBasePassRendering.h
@@ -110,6 +110,7 @@ public:
 template<typename LightMapPolicyType>
 class TMobileBasePassVSPolicyParamType : public FMeshMaterialShader, public LightMapPolicyType::VertexParametersType
 {
+	DECLARE_INLINE_TYPE_LAYOUT_EXPLICIT_BASES(TMobileBasePassVSPolicyParamType, NonVirtual, FMeshMaterialShader, typename LightMapPolicyType::VertexParametersType);
 protected:
 
 	TMobileBasePassVSPolicyParamType() {}
@@ -129,13 +130,6 @@ public:
 		FMeshMaterialShader::ModifyCompilationEnvironment(Parameters, OutEnvironment);
 	}
 
-	virtual bool Serialize(FArchive& Ar) override
-	{
-		bool bShaderHasOutdatedParameters = FMeshMaterialShader::Serialize(Ar);
-		LightMapPolicyType::VertexParametersType::Serialize(Ar);
-		return bShaderHasOutdatedParameters;
-	}
-
 	void GetShaderBindings(
 		const FScene* Scene,
 		ERHIFeatureLevel::Type FeatureLevel,
@@ -160,7 +154,7 @@ template<typename LightMapPolicyType>
 class TMobileBasePassVSBaseType : public TMobileBasePassVSPolicyParamType<LightMapPolicyType>
 {
 	typedef TMobileBasePassVSPolicyParamType<LightMapPolicyType> Super;
-
+	DECLARE_INLINE_TYPE_LAYOUT(TMobileBasePassVSBaseType, NonVirtual);
 protected:
 
 	TMobileBasePassVSBaseType() {}
@@ -178,6 +172,9 @@ public:
 		LightMapPolicyType::ModifyCompilationEnvironment(Parameters, OutEnvironment);
 		Super::ModifyCompilationEnvironment(Parameters, OutEnvironment);
 	}
+
+	
+	
 };
 
 template< typename LightMapPolicyType, EOutputFormat OutputFormat >
@@ -217,6 +214,7 @@ public:
 template<typename LightMapPolicyType>
 class TMobileBasePassPSPolicyParamType : public FMeshMaterialShader, public LightMapPolicyType::PixelParametersType
 {
+	DECLARE_INLINE_TYPE_LAYOUT_EXPLICIT_BASES(TMobileBasePassPSPolicyParamType, NonVirtual, FMeshMaterialShader, typename LightMapPolicyType::PixelParametersType);
 public:
 
 	static bool ShouldCompilePermutation(const FMeshMaterialShaderPermutationParameters& Parameters)
@@ -228,7 +226,7 @@ public:
 	{
 		FMeshMaterialShader::ModifyCompilationEnvironment(Parameters, OutEnvironment);
 		// Modify compilation environment depending upon material shader quality level settings.
-		ModifyCompilationEnvironmentForQualityLevel(Parameters.Platform, Parameters.Material->GetQualityLevel(), OutEnvironment);
+		ModifyCompilationEnvironmentForQualityLevel(Parameters.Platform, Parameters.MaterialParameters.QualityLevel, OutEnvironment);
 	}
 
 	/** Initialization constructor. */
@@ -261,52 +259,24 @@ public:
 
 	TMobileBasePassPSPolicyParamType() {}
 
-	virtual bool Serialize(FArchive& Ar) override
-	{
-		bool bShaderHasOutdatedParameters = FMeshMaterialShader::Serialize(Ar);
-		LightMapPolicyType::PixelParametersType::Serialize(Ar);
-
-		Ar << MobileDirectionLightBufferParam;
-		Ar << ReflectionParameter;
-
-		Ar << HQReflectionCubemaps[0];
-		Ar << HQReflectionCubemaps[1];
-		Ar << HQReflectionCubemaps[2];
-		Ar << HQReflectionSamplers[0];
-		Ar << HQReflectionSamplers[1];
-		Ar << HQReflectionSamplers[2];
-		Ar << HQReflectionInvAverageBrigtnessParams;
-		Ar << HQReflectionPositionsAndRadii;
-		
-		Ar << LightPositionAndInvRadiusParameter;
-		Ar << LightColorAndFalloffExponentParameter;
-		Ar << NumDynamicPointLightsParameter;
-		Ar << SpotLightAnglesParameter;
-		Ar << SpotLightDirectionAndSpecularScaleParameter;
-			
-		Ar << CSMDebugHintParams;
-
-		return bShaderHasOutdatedParameters;
-	}
-
 private:
-	FShaderUniformBufferParameter MobileDirectionLightBufferParam;
-	FShaderUniformBufferParameter ReflectionParameter;
+	LAYOUT_FIELD(FShaderUniformBufferParameter, MobileDirectionLightBufferParam);
+	LAYOUT_FIELD(FShaderUniformBufferParameter, ReflectionParameter);
 
 	// HQ reflection bound as loose params
-	FShaderResourceParameter HQReflectionCubemaps[3];
-	FShaderResourceParameter HQReflectionSamplers[3];
-	FShaderParameter HQReflectionInvAverageBrigtnessParams;
-	FShaderParameter HQReflectionPositionsAndRadii;
+	LAYOUT_ARRAY(FShaderResourceParameter, HQReflectionCubemaps, 3);
+	LAYOUT_ARRAY(FShaderResourceParameter, HQReflectionSamplers, 3);
+	LAYOUT_FIELD(FShaderParameter, HQReflectionInvAverageBrigtnessParams);
+	LAYOUT_FIELD(FShaderParameter, HQReflectionPositionsAndRadii);
+
+	LAYOUT_FIELD(FShaderParameter, LightPositionAndInvRadiusParameter);
+	LAYOUT_FIELD(FShaderParameter, LightColorAndFalloffExponentParameter);
+	LAYOUT_FIELD(FShaderParameter, NumDynamicPointLightsParameter);
+	LAYOUT_FIELD(FShaderParameter, SpotLightDirectionAndSpecularScaleParameter);
+	LAYOUT_FIELD(FShaderParameter, SpotLightAnglesParameter);
+
+	LAYOUT_FIELD(FShaderParameter, CSMDebugHintParams);
 	
-	FShaderParameter LightPositionAndInvRadiusParameter;
-	FShaderParameter LightColorAndFalloffExponentParameter;
-	FShaderParameter NumDynamicPointLightsParameter;
-	FShaderParameter SpotLightDirectionAndSpecularScaleParameter;
-	FShaderParameter SpotLightAnglesParameter;
-
-	FShaderParameter CSMDebugHintParams;
-
 public:
 	void GetShaderBindings(
 		const FScene* Scene,
@@ -326,9 +296,8 @@ template<typename LightMapPolicyType>
 class TMobileBasePassPSBaseType : public TMobileBasePassPSPolicyParamType<LightMapPolicyType>
 {
 	typedef TMobileBasePassPSPolicyParamType<LightMapPolicyType> Super;
-
+	DECLARE_INLINE_TYPE_LAYOUT(TMobileBasePassPSBaseType, NonVirtual);
 public:
-
 	static bool ShouldCompilePermutation(const FMeshMaterialShaderPermutationParameters& Parameters)
 	{
 		return LightMapPolicyType::ShouldCompilePermutation(Parameters)
@@ -344,6 +313,9 @@ public:
 	/** Initialization constructor. */
 	TMobileBasePassPSBaseType(const FMeshMaterialShaderType::CompiledShaderInitializerType& Initializer) : Super(Initializer) {}
 	TMobileBasePassPSBaseType() {}
+
+	
+	
 };
 
 
@@ -364,8 +336,8 @@ namespace MobileBasePass
 		const FMaterial& MaterialResource,
 		FVertexFactoryType* VertexFactoryType,
 		bool bEnableSkyLight, 
-		TMobileBasePassVSPolicyParamType<FUniformLightMapPolicy>*& VertexShader,
-		TMobileBasePassPSPolicyParamType<FUniformLightMapPolicy>*& PixelShader);
+		TShaderRef<TMobileBasePassVSPolicyParamType<FUniformLightMapPolicy>>& VertexShader,
+		TShaderRef<TMobileBasePassPSPolicyParamType<FUniformLightMapPolicy>>& PixelShader);
 
 	const FLightSceneInfo* GetDirectionalLightInfo(const FScene* Scene, const FPrimitiveSceneProxy* PrimitiveSceneProxy);
 	int32 CalcNumMovablePointLights(const FMaterial& InMaterial, const FPrimitiveSceneProxy* InPrimitiveSceneProxy);
@@ -404,7 +376,7 @@ public:
 		const bool bMobileDynamicPointLightsUseStaticBranch = (MobileDynamicPointLightsUseStaticBranchCVar->GetValueOnAnyThread() == 1);
 		const int32 MobileNumDynamicPointLights = MobileNumDynamicPointLightsCVar->GetValueOnAnyThread();
 		const int32 MobileSkyLightPermutationOptions = MobileSkyLightPermutationCVar->GetValueOnAnyThread();
-		const bool bIsUnlit = Parameters.Material->GetShadingModels().IsUnlit();
+		const bool bIsUnlit = Parameters.MaterialParameters.ShadingModels.IsUnlit();
 
 		// Only compile skylight version for lit materials on ES2 (Metal) or higher
 		const bool bShouldCacheBySkylight = !bEnableSkyLight || !bIsUnlit;
diff --git a/Engine/Source/Runtime/Renderer/Private/MobileDistortionPass.cpp b/Engine/Source/Runtime/Renderer/Private/MobileDistortionPass.cpp
--- a/Engine/Source/Runtime/Renderer/Private/MobileDistortionPass.cpp
+++ b/Engine/Source/Runtime/Renderer/Private/MobileDistortionPass.cpp
@@ -85,7 +85,7 @@ class FDistortionMergePS_ES2 : public FGlobalShader
 	FDistortionMergePS_ES2() {}
 
 public:
-	FPostProcessPassParameters PostprocessParameter;
+	LAYOUT_FIELD(FPostProcessPassParameters, PostprocessParameter)
 
 	/** Initialization constructor. */
 	FDistortionMergePS_ES2(const ShaderMetaType::CompiledShaderInitializerType& Initializer)
@@ -96,17 +96,10 @@ public:
 
 	void SetParameters(const FRenderingCompositePassContext& Context)
 	{
-		FRHIPixelShader* ShaderRHI = GetPixelShader();
+		FRHIPixelShader* ShaderRHI = Context.RHICmdList.GetBoundPixelShader();
 		FGlobalShader::SetParameters<FViewUniformShaderParameters>(Context.RHICmdList, ShaderRHI, Context.View.ViewUniformBuffer);
 		PostprocessParameter.SetPS(Context.RHICmdList, ShaderRHI, Context, TStaticSamplerState<SF_Bilinear, AM_Clamp, AM_Clamp, AM_Clamp>::GetRHI());
 	}
-
-	virtual bool Serialize(FArchive& Ar) override
-	{
-		bool bShaderHasOutdatedParameters = FGlobalShader::Serialize(Ar);
-		Ar << PostprocessParameter;
-		return bShaderHasOutdatedParameters;
-	}
 };
 
 IMPLEMENT_SHADER_TYPE(, FDistortionMergePS_ES2, TEXT("/Engine/Private/DistortApplyScreenPS.usf"), TEXT("Merge_ES2"), SF_Pixel);
@@ -140,8 +133,8 @@ void FRCDistortionMergePassES2::Process(FRenderingCompositePassContext& Context)
 		TShaderMapRef<FDistortionMergePS_ES2> PixelShader(View.ShaderMap);
 
 		GraphicsPSOInit.BoundShaderState.VertexDeclarationRHI = GFilterVertexDeclaration.VertexDeclarationRHI;
-		GraphicsPSOInit.BoundShaderState.VertexShaderRHI = GETSAFERHISHADER_VERTEX(*VertexShader);
-		GraphicsPSOInit.BoundShaderState.PixelShaderRHI = GETSAFERHISHADER_PIXEL(*PixelShader);
+		GraphicsPSOInit.BoundShaderState.VertexShaderRHI = VertexShader.GetVertexShader();
+		GraphicsPSOInit.BoundShaderState.PixelShaderRHI = PixelShader.GetPixelShader();
 		GraphicsPSOInit.PrimitiveType = PT_TriangleList;
 
 		SetGraphicsPipelineState(Context.RHICmdList, GraphicsPSOInit);
@@ -157,7 +150,7 @@ void FRCDistortionMergePassES2::Process(FRenderingCompositePassContext& Context)
 			SrcRect.Width(), SrcRect.Height(),
 			DstSize,
 			SrcSize,
-			*VertexShader,
+			VertexShader,
 			EDRF_UseTriangleOptimization);
 	}
 	Context.RHICmdList.EndRenderPass();
diff --git a/Engine/Source/Runtime/Renderer/Private/MobileReflectionEnvironmentCapture.cpp b/Engine/Source/Runtime/Renderer/Private/MobileReflectionEnvironmentCapture.cpp
--- a/Engine/Source/Runtime/Renderer/Private/MobileReflectionEnvironmentCapture.cpp
+++ b/Engine/Source/Runtime/Renderer/Private/MobileReflectionEnvironmentCapture.cpp
@@ -47,33 +47,23 @@ public:
 
 	void SetParameters(FRHICommandList& RHICmdList, int32 CubeFaceValue, int32 SourceMipIndexValue, FSceneRenderTargetItem& SourceTextureValue)
 	{
-		SetShaderValue(RHICmdList, GetPixelShader(), CubeFace, CubeFaceValue);
-		SetShaderValue(RHICmdList, GetPixelShader(), SourceMipIndex, SourceMipIndexValue);
+		SetShaderValue(RHICmdList, RHICmdList.GetBoundPixelShader(), CubeFace, CubeFaceValue);
+		SetShaderValue(RHICmdList, RHICmdList.GetBoundPixelShader(), SourceMipIndex, SourceMipIndexValue);
 
 		SetTextureParameter(
 			RHICmdList,
-			GetPixelShader(),
+			RHICmdList.GetBoundPixelShader(),
 			SourceTexture,
 			SourceTextureSampler,
 			TStaticSamplerState<SF_Bilinear, AM_Clamp, AM_Clamp, AM_Clamp>::GetRHI(),
 			SourceTextureValue.ShaderResourceTexture);
 	}
 
-	virtual bool Serialize(FArchive& Ar) override
-	{
-		bool bShaderHasOutdatedParameters = FGlobalShader::Serialize(Ar);
-		Ar << CubeFace;
-		Ar << SourceMipIndex;
-		Ar << SourceTexture;
-		Ar << SourceTextureSampler;
-		return bShaderHasOutdatedParameters;
-	}
-
 private:
-	FShaderParameter CubeFace;
-	FShaderParameter SourceMipIndex;
-	FShaderResourceParameter SourceTexture;
-	FShaderResourceParameter SourceTextureSampler;
+	LAYOUT_FIELD(FShaderParameter, CubeFace)
+	LAYOUT_FIELD(FShaderParameter, SourceMipIndex)
+	LAYOUT_FIELD(FShaderResourceParameter, SourceTexture)
+	LAYOUT_FIELD(FShaderResourceParameter, SourceTextureSampler)
 };
 
 IMPLEMENT_SHADER_TYPE(, FMobileDownsamplePS, TEXT("/Engine/Private/ReflectionEnvironmentShaders.usf"), TEXT("DownsamplePS_Mobile"), SF_Pixel);
@@ -153,8 +143,8 @@ namespace MobileReflectionEnvironmentCapture
 						TShaderMapRef<FMobileDownsamplePS> PixelShader(ShaderMap);
 
 						GraphicsPSOInit.BoundShaderState.VertexDeclarationRHI = GFilterVertexDeclaration.VertexDeclarationRHI;
-						GraphicsPSOInit.BoundShaderState.VertexShaderRHI = GETSAFERHISHADER_VERTEX(*VertexShader);
-						GraphicsPSOInit.BoundShaderState.PixelShaderRHI = GETSAFERHISHADER_PIXEL(*PixelShader);
+						GraphicsPSOInit.BoundShaderState.VertexShaderRHI = VertexShader.GetVertexShader();
+						GraphicsPSOInit.BoundShaderState.PixelShaderRHI = PixelShader.GetPixelShader();
 						GraphicsPSOInit.PrimitiveType = PT_TriangleList;
 						SetGraphicsPipelineState(RHICmdList, GraphicsPSOInit);
 
@@ -168,7 +158,7 @@ namespace MobileReflectionEnvironmentCapture
 							ViewRect.Width(), ViewRect.Height(),
 							FIntPoint(ViewRect.Width(), ViewRect.Height()),
 							FIntPoint(MipSize, MipSize),
-							*VertexShader);
+							VertexShader);
 					}
 					RHICmdList.EndRenderPass();
 					RHICmdList.CopyToResolveTarget(EffectiveRT.TargetableTexture, EffectiveRT.ShaderResourceTexture, FResolveParams(FResolveRect(), (ECubeFace)CubeFace, MipIndex));
@@ -246,8 +236,8 @@ namespace MobileReflectionEnvironmentCapture
 					TShaderMapRef<FOneColorPS> PixelShader(GetGlobalShaderMap(FeatureLevel));
 
 					GraphicsPSOInit.BoundShaderState.VertexDeclarationRHI = GFilterVertexDeclaration.VertexDeclarationRHI;
-					GraphicsPSOInit.BoundShaderState.VertexShaderRHI = GETSAFERHISHADER_VERTEX(*VertexShader);
-					GraphicsPSOInit.BoundShaderState.PixelShaderRHI = GETSAFERHISHADER_PIXEL(*PixelShader);
+					GraphicsPSOInit.BoundShaderState.VertexShaderRHI = VertexShader.GetVertexShader();
+					GraphicsPSOInit.BoundShaderState.PixelShaderRHI = PixelShader.GetPixelShader();
 					GraphicsPSOInit.PrimitiveType = PT_TriangleList;
 
 					SetGraphicsPipelineState(RHICmdList, GraphicsPSOInit);
@@ -263,7 +253,7 @@ namespace MobileReflectionEnvironmentCapture
 						SourceDimensions.X, SourceDimensions.Y,
 						FIntPoint(ViewRect.Width(), ViewRect.Height()),
 						SourceDimensions,
-						*VertexShader);
+						VertexShader);
 				}
 				RHICmdList.EndRenderPass();
 				RHICmdList.CopyToResolveTarget(EffectiveColorRT.TargetableTexture, EffectiveColorRT.ShaderResourceTexture, FResolveParams(FResolveRect(), (ECubeFace)CubeFace));
@@ -312,8 +302,8 @@ namespace MobileReflectionEnvironmentCapture
 						TShaderMapRef<FMobileDownsamplePS> PixelShader(ShaderMap);
 
 						GraphicsPSOInit.BoundShaderState.VertexDeclarationRHI = GFilterVertexDeclaration.VertexDeclarationRHI;
-						GraphicsPSOInit.BoundShaderState.VertexShaderRHI = GETSAFERHISHADER_VERTEX(*VertexShader);
-						GraphicsPSOInit.BoundShaderState.PixelShaderRHI = GETSAFERHISHADER_PIXEL(*PixelShader);
+						GraphicsPSOInit.BoundShaderState.VertexShaderRHI = VertexShader.GetVertexShader();
+						GraphicsPSOInit.BoundShaderState.PixelShaderRHI = PixelShader.GetPixelShader();
 						GraphicsPSOInit.PrimitiveType = PT_TriangleList;
 
 						SetGraphicsPipelineState(RHICmdList, GraphicsPSOInit);
@@ -328,7 +318,7 @@ namespace MobileReflectionEnvironmentCapture
 							ViewRect.Width(), ViewRect.Height(),
 							FIntPoint(ViewRect.Width(), ViewRect.Height()),
 							FIntPoint(MipSize, MipSize),
-							*VertexShader);
+							VertexShader);
 					}
 					RHICmdList.EndRenderPass();
 					RHICmdList.CopyToResolveTarget(EffectiveRT.TargetableTexture, EffectiveRT.ShaderResourceTexture, FResolveParams(FResolveRect(), (ECubeFace)CubeFace, MipIndex));
@@ -399,13 +389,12 @@ namespace MobileReflectionEnvironmentCapture
 
 						TShaderMapRef<FScreenVS> VertexShader(GetGlobalShaderMap(FeatureLevel));
 
-						FCubeFilterPS* HQFilterPixelShader = *TShaderMapRef< TCubeFilterPS<0> >(ShaderMap);
-						check(HQFilterPixelShader);
+						TShaderMapRef<TCubeFilterPS<0>> HQFilterPixelShader(ShaderMap);
 						TShaderMapRef<FMobileDownsamplePS> BilinFilterPixelShader(ShaderMap);
-						FRHIPixelShader* PixelShaderRHI = bUseHQFiltering ? HQFilterPixelShader->GetPixelShader() : BilinFilterPixelShader->GetPixelShader();
+						FRHIPixelShader* PixelShaderRHI = bUseHQFiltering ? HQFilterPixelShader.GetPixelShader() : BilinFilterPixelShader.GetPixelShader();
 
 						GraphicsPSOInit.BoundShaderState.VertexDeclarationRHI = GFilterVertexDeclaration.VertexDeclarationRHI;
-						GraphicsPSOInit.BoundShaderState.VertexShaderRHI = GETSAFERHISHADER_VERTEX(*VertexShader);
+						GraphicsPSOInit.BoundShaderState.VertexShaderRHI = VertexShader.GetVertexShader();
 						GraphicsPSOInit.BoundShaderState.PixelShaderRHI = PixelShaderRHI;
 						GraphicsPSOInit.PrimitiveType = PT_TriangleList;
 
@@ -437,7 +426,7 @@ namespace MobileReflectionEnvironmentCapture
 							ViewRect.Width(), ViewRect.Height(),
 							FIntPoint(ViewRect.Width(), ViewRect.Height()),
 							FIntPoint(MipSize, MipSize),
-							*VertexShader);
+							VertexShader);
 					}
 					RHICmdList.EndRenderPass();
 					RHICmdList.CopyToResolveTarget(EffectiveRT.TargetableTexture, EffectiveRT.ShaderResourceTexture, FResolveParams(FResolveRect(), (ECubeFace)CubeFace, MipIndex));
diff --git a/Engine/Source/Runtime/Renderer/Private/MobileSceneCaptureRendering.cpp b/Engine/Source/Runtime/Renderer/Private/MobileSceneCaptureRendering.cpp
--- a/Engine/Source/Runtime/Renderer/Private/MobileSceneCaptureRendering.cpp
+++ b/Engine/Source/Runtime/Renderer/Private/MobileSceneCaptureRendering.cpp
@@ -67,24 +67,15 @@ public:
 
 	void SetParameters(FRHICommandList& RHICmdList, const FViewInfo& View, FRHISamplerState* SamplerStateRHI, FRHITexture* TextureRHI)
 	{
-		FGlobalShader::SetParameters<FViewUniformShaderParameters>(RHICmdList, GetPixelShader(), View.ViewUniformBuffer);
-		SetTextureParameter(RHICmdList, GetPixelShader(), InTexture, InTextureSampler, SamplerStateRHI, TextureRHI);
-		SceneTextureParameters.Set(RHICmdList, GetPixelShader(), View.FeatureLevel, ESceneTextureSetupMode::All);
-	}
-
-	virtual bool Serialize(FArchive& Ar) override
-	{
-		bool bShaderHasOutdatedParameters = FGlobalShader::Serialize(Ar);
-		Ar << InTexture;
-		Ar << InTextureSampler;
-		Ar << SceneTextureParameters;
-		return bShaderHasOutdatedParameters;
+		FGlobalShader::SetParameters<FViewUniformShaderParameters>(RHICmdList, RHICmdList.GetBoundPixelShader(), View.ViewUniformBuffer);
+		SetTextureParameter(RHICmdList, RHICmdList.GetBoundPixelShader(), InTexture, InTextureSampler, SamplerStateRHI, TextureRHI);
+		SceneTextureParameters.Set(RHICmdList, RHICmdList.GetBoundPixelShader(), View.FeatureLevel, ESceneTextureSetupMode::All);
 	}
 
 private:
-	FShaderResourceParameter InTexture;
-	FShaderResourceParameter InTextureSampler;
-	FSceneTextureShaderParameters SceneTextureParameters;
+	LAYOUT_FIELD(FShaderResourceParameter, InTexture)
+	LAYOUT_FIELD(FShaderResourceParameter, InTextureSampler)
+	LAYOUT_FIELD(FSceneTextureShaderParameters, SceneTextureParameters)
 };
 
 /**
@@ -113,22 +104,15 @@ public:
 
 	void SetParameters(FRHICommandList& RHICmdList, const FViewInfo& View, const FIntPoint& SourceTexSize)
 	{
-		FGlobalShader::SetParameters<FViewUniformShaderParameters>(RHICmdList, GetVertexShader(), View.ViewUniformBuffer);
+		FGlobalShader::SetParameters<FViewUniformShaderParameters>(RHICmdList, RHICmdList.GetBoundVertexShader(), View.ViewUniformBuffer);
 		if (InvTexSizeParameter.IsBound())
 		{
 			FVector2D InvTexSize(1.0f / SourceTexSize.X, 1.0f / SourceTexSize.Y);
-			SetShaderValue(RHICmdList, GetVertexShader(), InvTexSizeParameter, InvTexSize);
+			SetShaderValue(RHICmdList, RHICmdList.GetBoundVertexShader(), InvTexSizeParameter, InvTexSize);
 		}
 	}
 
-	virtual bool Serialize(FArchive& Ar) override
-	{
-		bool bShaderHasOutdatedParameters = FGlobalShader::Serialize(Ar);
-		Ar << InvTexSizeParameter;
-		return bShaderHasOutdatedParameters;
-	}
-
-	FShaderParameter InvTexSizeParameter;
+	LAYOUT_FIELD(FShaderParameter, InvTexSizeParameter)
 };
 
 #define IMPLEMENT_MOBILE_SCENE_CAPTURECOPY(SCENETYPE) \
@@ -148,24 +132,24 @@ IMPLEMENT_SHADER_TYPE(template<>, FMobileSceneCaptureCopyVS<true>, TEXT("/Engine
  
 
 template <bool bDemosaic, ESceneCaptureSource CaptureSource>
-static FShader* SetCaptureToTargetShaders(FRHICommandListImmediate& RHICmdList, FGraphicsPipelineStateInitializer& GraphicsPSOInit, FViewInfo& View, const FIntPoint& SourceTexSize, FRHITexture* SourceTextureRHI)
+static TShaderRef<FShader> SetCaptureToTargetShaders(FRHICommandListImmediate& RHICmdList, FGraphicsPipelineStateInitializer& GraphicsPSOInit, FViewInfo& View, const FIntPoint& SourceTexSize, FRHITexture* SourceTextureRHI)
 {
 	TShaderMapRef<FMobileSceneCaptureCopyVS<bDemosaic>> VertexShader(View.ShaderMap);
 	TShaderMapRef<FMobileSceneCaptureCopyPS<bDemosaic, CaptureSource>> PixelShader(View.ShaderMap);
 
 	GraphicsPSOInit.BoundShaderState.VertexDeclarationRHI = GFilterVertexDeclaration.VertexDeclarationRHI;
-	GraphicsPSOInit.BoundShaderState.VertexShaderRHI = GETSAFERHISHADER_VERTEX(*VertexShader);
-	GraphicsPSOInit.BoundShaderState.PixelShaderRHI = GETSAFERHISHADER_PIXEL(*PixelShader);
+	GraphicsPSOInit.BoundShaderState.VertexShaderRHI = VertexShader.GetVertexShader();
+	GraphicsPSOInit.BoundShaderState.PixelShaderRHI = PixelShader.GetPixelShader();
 	SetGraphicsPipelineState(RHICmdList, GraphicsPSOInit);
 
 	VertexShader->SetParameters(RHICmdList, View, SourceTexSize);
 	PixelShader->SetParameters(RHICmdList, View, TStaticSamplerState<SF_Point>::GetRHI(), SourceTextureRHI);
 
-	return *VertexShader;
+	return VertexShader;
 }
 
 template <bool bDemosaic>
-static FShader* SetCaptureToTargetShaders(FRHICommandListImmediate& RHICmdList, FGraphicsPipelineStateInitializer& GraphicsPSOInit, ESceneCaptureSource CaptureSource, FViewInfo& View, const FIntPoint& SourceTexSize, FRHITexture* SourceTextureRHI)
+static TShaderRef<FShader> SetCaptureToTargetShaders(FRHICommandListImmediate& RHICmdList, FGraphicsPipelineStateInitializer& GraphicsPSOInit, ESceneCaptureSource CaptureSource, FViewInfo& View, const FIntPoint& SourceTexSize, FRHITexture* SourceTextureRHI)
 {
 	switch (CaptureSource)
 	{
@@ -183,7 +167,7 @@ static FShader* SetCaptureToTargetShaders(FRHICommandListImmediate& RHICmdList,
 			return SetCaptureToTargetShaders<bDemosaic, SCS_DeviceDepth>(RHICmdList, GraphicsPSOInit, View, SourceTexSize, SourceTextureRHI);
 		default:
 			checkNoEntry();
-			return nullptr;
+			return TShaderRef<FShader>();
 	}
 }
 
@@ -242,7 +226,7 @@ static void CopyCaptureToTarget(
 		GraphicsPSOInit.RasterizerState = TStaticRasterizerState<FM_Solid, CM_None>::GetRHI();
 
 		const bool bUsingDemosaic = IsMobileHDRMosaic();
-		FShader* VertexShader;
+		TShaderRef<FShader> VertexShader;
 		if (bUsingDemosaic)
 		{
 			VertexShader = SetCaptureToTargetShaders<true>(RHICmdList, GraphicsPSOInit, CaptureSource, View, SourceTexSize, SourceTextureRHI);
@@ -310,8 +294,8 @@ static void CopyCaptureToTarget(
 			TShaderMapRef<FScreenPS> PixelShader(View.ShaderMap);
 
 			GraphicsPSOInit.BoundShaderState.VertexDeclarationRHI = GFilterVertexDeclaration.VertexDeclarationRHI;
-			GraphicsPSOInit.BoundShaderState.VertexShaderRHI = GETSAFERHISHADER_VERTEX(*ScreenVertexShader);
-			GraphicsPSOInit.BoundShaderState.PixelShaderRHI = GETSAFERHISHADER_PIXEL(*PixelShader);
+			GraphicsPSOInit.BoundShaderState.VertexShaderRHI = ScreenVertexShader.GetVertexShader();
+			GraphicsPSOInit.BoundShaderState.PixelShaderRHI = PixelShader.GetPixelShader();
 			GraphicsPSOInit.PrimitiveType = PT_TriangleList;
 
 			SetGraphicsPipelineState(RHICmdList, GraphicsPSOInit);
@@ -336,7 +320,7 @@ static void CopyCaptureToTarget(
 				ViewRect.Width(), TargetHeight,
 				TargetSize,
 				SourceTexSize,
-				*ScreenVertexShader,
+				ScreenVertexShader,
 				EDRF_UseTriangleOptimization);
 		}
 		RHICmdList.EndRenderPass();
diff --git a/Engine/Source/Runtime/Renderer/Private/MobileShadingRenderer.cpp b/Engine/Source/Runtime/Renderer/Private/MobileShadingRenderer.cpp
--- a/Engine/Source/Runtime/Renderer/Private/MobileShadingRenderer.cpp
+++ b/Engine/Source/Runtime/Renderer/Private/MobileShadingRenderer.cpp
@@ -1013,8 +1013,8 @@ void FMobileSceneRenderer::ConditionalResolveSceneDepth(FRHICommandListImmediate
 					TShaderMapRef<FScreenPS> PixelShader(View.ShaderMap);
 					
 					GraphicsPSOInit.BoundShaderState.VertexDeclarationRHI = GFilterVertexDeclaration.VertexDeclarationRHI;
-					GraphicsPSOInit.BoundShaderState.VertexShaderRHI = GETSAFERHISHADER_VERTEX(*ScreenVertexShader);
-					GraphicsPSOInit.BoundShaderState.PixelShaderRHI = GETSAFERHISHADER_PIXEL(*PixelShader);
+					GraphicsPSOInit.BoundShaderState.VertexShaderRHI = ScreenVertexShader.GetVertexShader();
+					GraphicsPSOInit.BoundShaderState.PixelShaderRHI = PixelShader.GetPixelShader();
 					GraphicsPSOInit.PrimitiveType = PT_TriangleList;
 
 					SetGraphicsPipelineState(RHICmdList, GraphicsPSOInit);
@@ -1029,7 +1029,7 @@ void FMobileSceneRenderer::ConditionalResolveSceneDepth(FRHICommandListImmediate
 						1, 1,
 						FIntPoint(1, 1),
 						FIntPoint(1, 1),
-						*ScreenVertexShader,
+						ScreenVertexShader,
 						EDRF_UseTriangleOptimization);
 				} // force depth resolve
 				RHICmdList.EndRenderPass();
@@ -1116,7 +1116,7 @@ public:
 
 	void SetParameters(FRHICommandList& RHICmdList, FRHIUniformBuffer* ViewUniformBuffer, FTextureRHIRef InMobileMultiViewSceneColorTexture)
 	{
-		FRHIPixelShader* ShaderRHI = GetPixelShader();
+		FRHIPixelShader* ShaderRHI = RHICmdList.GetBoundPixelShader();
 		FGlobalShader::SetParameters<FViewUniformShaderParameters>(RHICmdList, ShaderRHI, ViewUniformBuffer);
 		SetTextureParameter(
 			RHICmdList,
@@ -1127,16 +1127,8 @@ public:
 			InMobileMultiViewSceneColorTexture);
 	}
 
-	virtual bool Serialize(FArchive& Ar) override
-	{
-		const bool bShaderHasOutdatedParameters = FGlobalShader::Serialize(Ar);
-		Ar << MobileMultiViewSceneColorTexture;
-		Ar << MobileMultiViewSceneColorTextureSampler;
-		return bShaderHasOutdatedParameters;
-	}
-
-	FShaderResourceParameter MobileMultiViewSceneColorTexture;
-	FShaderResourceParameter MobileMultiViewSceneColorTextureSampler;
+	LAYOUT_FIELD(FShaderResourceParameter, MobileMultiViewSceneColorTexture)
+	LAYOUT_FIELD(FShaderResourceParameter, MobileMultiViewSceneColorTextureSampler)
 };
 
 IMPLEMENT_SHADER_TYPE(, FCopyMobileMultiViewSceneColorPS, TEXT("/Engine/Private/MobileMultiView.usf"), TEXT("MainPS"), SF_Pixel);
@@ -1173,8 +1165,8 @@ void FMobileSceneRenderer::CopyMobileMultiViewSceneColor(FRHICommandListImmediat
 	extern TGlobalResource<FFilterVertexDeclaration> GFilterVertexDeclaration;
 
 	GraphicsPSOInit.BoundShaderState.VertexDeclarationRHI = GFilterVertexDeclaration.VertexDeclarationRHI;
-	GraphicsPSOInit.BoundShaderState.VertexShaderRHI = GETSAFERHISHADER_VERTEX(*VertexShader);
-	GraphicsPSOInit.BoundShaderState.PixelShaderRHI = GETSAFERHISHADER_PIXEL(*PixelShader);
+	GraphicsPSOInit.BoundShaderState.VertexShaderRHI = VertexShader.GetVertexShader();
+	GraphicsPSOInit.BoundShaderState.PixelShaderRHI = PixelShader.GetPixelShader();
 	GraphicsPSOInit.PrimitiveType = PT_TriangleList;
 
 	SetGraphicsPipelineState(RHICmdList, GraphicsPSOInit);
@@ -1197,7 +1189,7 @@ void FMobileSceneRenderer::CopyMobileMultiViewSceneColor(FRHICommandListImmediat
 			View.ViewRect.Width(), View.ViewRect.Height(),
 			TargetSize,
 			TargetSize,
-			*VertexShader,
+			VertexShader,
 			EDRF_UseTriangleOptimization);
 	}
 
@@ -1253,8 +1245,8 @@ void FMobileSceneRenderer::PreTonemapMSAA(FRHICommandListImmediate& RHICmdList)
 
 	extern TGlobalResource<FFilterVertexDeclaration> GFilterVertexDeclaration;
 	GraphicsPSOInit.BoundShaderState.VertexDeclarationRHI = GFilterVertexDeclaration.VertexDeclarationRHI;
-	GraphicsPSOInit.BoundShaderState.VertexShaderRHI = GETSAFERHISHADER_VERTEX(*VertexShader);
-	GraphicsPSOInit.BoundShaderState.PixelShaderRHI = GETSAFERHISHADER_PIXEL(*PixelShader);
+	GraphicsPSOInit.BoundShaderState.VertexShaderRHI = VertexShader.GetVertexShader();
+	GraphicsPSOInit.BoundShaderState.PixelShaderRHI = PixelShader.GetPixelShader();
 	GraphicsPSOInit.PrimitiveType = PT_TriangleList;
 
 	SetGraphicsPipelineState(RHICmdList, GraphicsPSOInit);
@@ -1270,6 +1262,6 @@ void FMobileSceneRenderer::PreTonemapMSAA(FRHICommandListImmediate& RHICmdList)
 		TargetSize.X, TargetSize.Y,
 		TargetSize,
 		TargetSize,
-		*VertexShader,
+		VertexShader,
 		EDRF_UseTriangleOptimization);
 }
diff --git a/Engine/Source/Runtime/Renderer/Private/MobileTranslucentRendering.cpp b/Engine/Source/Runtime/Renderer/Private/MobileTranslucentRendering.cpp
--- a/Engine/Source/Runtime/Renderer/Private/MobileTranslucentRendering.cpp
+++ b/Engine/Source/Runtime/Renderer/Private/MobileTranslucentRendering.cpp
@@ -47,18 +47,11 @@ public:
 
 	void SetParameters(FRHICommandList& RHICmdList, const FSceneView& View)
 	{
-		SceneTextureParameters.Set(RHICmdList, GetPixelShader(), View.FeatureLevel, ESceneTextureSetupMode::All);
-	}
-
-	virtual bool Serialize(FArchive& Ar) override
-	{
-		bool bShaderHasOutdatedParameters = FGlobalShader::Serialize(Ar);
-		Ar << SceneTextureParameters;
-		return bShaderHasOutdatedParameters;
+		SceneTextureParameters.Set(RHICmdList, RHICmdList.GetBoundPixelShader(), View.FeatureLevel, ESceneTextureSetupMode::All);
 	}
 
 private:
-	FSceneTextureShaderParameters SceneTextureParameters;
+	LAYOUT_FIELD(FSceneTextureShaderParameters, SceneTextureParameters)
 };
 
 IMPLEMENT_SHADER_TYPE(,FMobileCopySceneAlphaPS,TEXT("/Engine/Private/TranslucentLightingShaders.usf"),TEXT("CopySceneAlphaMain"),SF_Pixel);
@@ -88,8 +81,8 @@ void FMobileSceneRenderer::CopySceneAlpha(FRHICommandListImmediate& RHICmdList,
 	TShaderMapRef<FMobileCopySceneAlphaPS> PixelShader(View.ShaderMap);
 
 	GraphicsPSOInit.BoundShaderState.VertexDeclarationRHI = GFilterVertexDeclaration.VertexDeclarationRHI;
-	GraphicsPSOInit.BoundShaderState.VertexShaderRHI = GETSAFERHISHADER_VERTEX(*ScreenVertexShader);
-	GraphicsPSOInit.BoundShaderState.PixelShaderRHI = GETSAFERHISHADER_PIXEL(*PixelShader);
+	GraphicsPSOInit.BoundShaderState.VertexShaderRHI = ScreenVertexShader.GetVertexShader();
+	GraphicsPSOInit.BoundShaderState.PixelShaderRHI = PixelShader.GetPixelShader();
 	GraphicsPSOInit.PrimitiveType = PT_TriangleList;
 
 	SetGraphicsPipelineState(RHICmdList, GraphicsPSOInit);
@@ -104,7 +97,7 @@ void FMobileSceneRenderer::CopySceneAlpha(FRHICommandListImmediate& RHICmdList,
 		X, Y,
 		FIntPoint(X, Y),
 		SceneContext.GetBufferSizeXY(),
-		*ScreenVertexShader,
+		ScreenVertexShader,
 		EDRF_UseTriangleOptimization);
 
 	SceneContext.FinishRenderingSceneAlphaCopy(RHICmdList);
@@ -172,7 +165,7 @@ public:
 
 	static bool ShouldCompilePermutation(const FMeshMaterialShaderPermutationParameters& Parameters)
 	{
-		return IsTranslucentBlendMode(Parameters.Material->GetBlendMode()) && IsMobilePlatform(Parameters.Platform);
+		return IsTranslucentBlendMode(Parameters.MaterialParameters.BlendMode) && IsMobilePlatform(Parameters.Platform);
 	}
 
 	static void ModifyCompilationEnvironment(const FMaterialShaderPermutationParameters& Parameters, FShaderCompilerEnvironment& OutEnvironment)
@@ -184,12 +177,6 @@ public:
 		OutEnvironment.SetDefine(TEXT("OUTPUT_GAMMA_SPACE"), IsMobileHDR() == false && !bMobileUseHWsRGBEncoding);
 		OutEnvironment.SetDefine(TEXT("OUTPUT_MOBILE_HDR"), IsMobileHDR() == true);
 	}
-
-	virtual bool Serialize(FArchive& Ar) override
-	{
-		const bool result = FMeshMaterialShader::Serialize(Ar);
-		return result;
-	}
 };
 
 IMPLEMENT_MATERIAL_SHADER_TYPE(, FOpacityOnlyVS, TEXT("/Engine/Private/MobileOpacityShaders.usf"), TEXT("MainVS"), SF_Vertex);
@@ -204,7 +191,7 @@ public:
 
 	static bool ShouldCompilePermutation(const FMeshMaterialShaderPermutationParameters& Parameters)
 	{
-		return IsTranslucentBlendMode(Parameters.Material->GetBlendMode()) && IsMobilePlatform(Parameters.Platform);
+		return IsTranslucentBlendMode(Parameters.MaterialParameters.BlendMode) && IsMobilePlatform(Parameters.Platform);
 	}
 
 	static void ModifyCompilationEnvironment(const FMaterialShaderPermutationParameters& Parameters, FShaderCompilerEnvironment& OutEnvironment)
@@ -220,12 +207,6 @@ public:
 	}
 
 	FOpacityOnlyPS() {}
-
-	virtual bool Serialize(FArchive& Ar) override
-	{
-		bool bShaderHasOutdatedParameters = FMeshMaterialShader::Serialize(Ar);
-		return bShaderHasOutdatedParameters;
-	}
 };
 
 IMPLEMENT_MATERIAL_SHADER_TYPE(, FOpacityOnlyPS, TEXT("/Engine/Private/MobileOpacityShaders.usf"), TEXT("MainPS"), SF_Pixel);
@@ -331,8 +312,9 @@ private:
 		FMeshPassProcessorRenderState DrawRenderState(PassDrawRenderState);
 		MobileBasePass::SetTranslucentRenderState(DrawRenderState, Material);
 
-		ERasterizerFillMode MeshFillMode = ComputeMeshFillMode(MeshBatch, Material);
-		ERasterizerCullMode MeshCullMode = ComputeMeshCullMode(MeshBatch, Material);
+		const FMeshDrawingPolicyOverrideSettings OverrideSettings = ComputeMeshOverrideSettings(MeshBatch);
+		ERasterizerFillMode MeshFillMode = ComputeMeshFillMode(MeshBatch, Material, OverrideSettings);
+		ERasterizerCullMode MeshCullMode = ComputeMeshCullMode(MeshBatch, Material, OverrideSettings);
 	
 		FMeshMaterialShaderElementData ShaderElementData;
 		ShaderElementData.InitializeMeshMaterialData(ViewIfDynamicMeshCommand, PrimitiveSceneProxy, MeshBatch, StaticMeshId, false);
diff --git a/Engine/Source/Runtime/Renderer/Private/PathTracing.cpp b/Engine/Source/Runtime/Renderer/Private/PathTracing.cpp
--- a/Engine/Source/Runtime/Renderer/Private/PathTracing.cpp
+++ b/Engine/Source/Runtime/Renderer/Private/PathTracing.cpp
@@ -113,7 +113,7 @@ public:
 	}
 
 	FPathTracingRG() {}
-	virtual ~FPathTracingRG() {}
+	~FPathTracingRG() {}
 
 	FPathTracingRG(const ShaderMetaType::CompiledShaderInitializerType& Initializer)
 		: FGlobalShader(Initializer)
@@ -311,7 +311,7 @@ public:
 		}
 	}
 
-	bool Serialize(FArchive& Ar)
+	/*bool Serialize(FArchive& Ar)
 	{
 		bool bShaderHasOutdatedParameters = FGlobalShader::Serialize(Ar);
 		Ar << TLASParameter;
@@ -327,20 +327,20 @@ public:
 		Ar << RayCountPerPixelRT;
 
 		return bShaderHasOutdatedParameters;
-	}
-
-	FShaderResourceParameter		TLASParameter;   // RaytracingAccelerationStructure
-	FShaderUniformBufferParameter	ViewParameter;
-	FShaderUniformBufferParameter	PathTracingParameters;
-	FShaderUniformBufferParameter	SceneLightsParameters;
-	FShaderUniformBufferParameter	SkyLightParameters;
-	FShaderUniformBufferParameter	AdaptiveSamplingParameters;
+	}*/
+	
+	LAYOUT_FIELD(FShaderResourceParameter, TLASParameter);   // RaytracingAccelerationStructure
+	LAYOUT_FIELD(FShaderUniformBufferParameter, ViewParameter);
+	LAYOUT_FIELD(FShaderUniformBufferParameter, PathTracingParameters);
+	LAYOUT_FIELD(FShaderUniformBufferParameter, SceneLightsParameters);
+	LAYOUT_FIELD(FShaderUniformBufferParameter, SkyLightParameters);
+	LAYOUT_FIELD(FShaderUniformBufferParameter, AdaptiveSamplingParameters);
 
 	// Output parameters
-	FShaderResourceParameter		RadianceRT;
-	FShaderResourceParameter		SampleCountRT;
-	FShaderResourceParameter        PixelPositionRT;
-	FShaderResourceParameter		RayCountPerPixelRT;
+	LAYOUT_FIELD(FShaderResourceParameter, RadianceRT);
+	LAYOUT_FIELD(FShaderResourceParameter, SampleCountRT);
+	LAYOUT_FIELD(FShaderResourceParameter, PixelPositionRT);
+	LAYOUT_FIELD(FShaderResourceParameter, RayCountPerPixelRT);
 };
 IMPLEMENT_SHADER_TYPE(, FPathTracingRG, TEXT("/Engine/Private/PathTracing/PathTracing.usf"), TEXT("PathTracingMainRG"), SF_RayGen);
 
@@ -396,7 +396,7 @@ class FPathTracingCompositorPS : public FGlobalShader
 		FRHITexture* CumulativeIrradianceRT,
 		FRHITexture* CumulativeSampleCountRT)
 	{
-		FRHIPixelShader* ShaderRHI = GetPixelShader();
+		FRHIPixelShader* ShaderRHI = RHICmdList.GetBoundPixelShader();
 		FGlobalShader::SetParameters<FViewUniformShaderParameters>(RHICmdList, ShaderRHI, View.ViewUniformBuffer);
 		SetTextureParameter(RHICmdList, ShaderRHI, RadianceRedTexture, RadianceRedRT);
 		SetTextureParameter(RHICmdList, ShaderRHI, RadianceGreenTexture, RadianceGreenRT);
@@ -407,28 +407,15 @@ class FPathTracingCompositorPS : public FGlobalShader
 		SetTextureParameter(RHICmdList, ShaderRHI, CumulativeSampleCountTexture, CumulativeSampleCountRT);
 	}
 
-	virtual bool Serialize(FArchive& Ar) override
-	{
-		bool bShaderHasOudatedParameters = FGlobalShader::Serialize(Ar);
-		Ar << RadianceRedTexture;
-		Ar << RadianceGreenTexture;
-		Ar << RadianceBlueTexture;
-		Ar << RadianceAlphaTexture;
-		Ar << SampleCountTexture;
-		Ar << CumulativeIrradianceTexture;
-		Ar << CumulativeSampleCountTexture;
-		return bShaderHasOudatedParameters;
-	}
-
 public:
-	FShaderResourceParameter RadianceRedTexture;
-	FShaderResourceParameter RadianceGreenTexture;
-	FShaderResourceParameter RadianceBlueTexture;
-	FShaderResourceParameter RadianceAlphaTexture;
-	FShaderResourceParameter SampleCountTexture;
-
-	FShaderResourceParameter CumulativeIrradianceTexture;
-	FShaderResourceParameter CumulativeSampleCountTexture;
+	LAYOUT_FIELD(FShaderResourceParameter, RadianceRedTexture);
+	LAYOUT_FIELD(FShaderResourceParameter, RadianceGreenTexture);
+	LAYOUT_FIELD(FShaderResourceParameter, RadianceBlueTexture);
+	LAYOUT_FIELD(FShaderResourceParameter, RadianceAlphaTexture);
+	LAYOUT_FIELD(FShaderResourceParameter, SampleCountTexture);
+
+	LAYOUT_FIELD(FShaderResourceParameter, CumulativeIrradianceTexture);
+	LAYOUT_FIELD(FShaderResourceParameter, CumulativeSampleCountTexture);
 };
 
 IMPLEMENT_SHADER_TYPE(, FPathTracingCompositorPS, TEXT("/Engine/Private/PathTracing/PathTracingCompositingPixelShader.usf"), TEXT("CompositeMain"), SF_Pixel);
@@ -437,7 +424,7 @@ void FDeferredShadingSceneRenderer::PreparePathTracing(const FViewInfo& View, TA
 {
 	// Declare all RayGen shaders that require material closest hit shaders to be bound
 	auto RayGenShader = View.ShaderMap->GetShader<FPathTracingRG>();
-	OutRayGenShaders.Add(RayGenShader->GetRayTracingShader());
+	OutRayGenShaders.Add(RayGenShader.GetRayTracingShader());
 }
 
 void FDeferredShadingSceneRenderer::RenderPathTracing(FRHICommandListImmediate& RHICmdList, const FViewInfo& View)
@@ -558,7 +545,7 @@ void FDeferredShadingSceneRenderer::RenderPathTracing(FRHICommandListImmediate&
 
 			RHICmdList.RayTraceDispatch(
 				View.RayTracingMaterialPipeline, 
-				RayGenShader->GetRayTracingShader(), 
+				RayGenShader.GetRayTracingShader(), 
 				RayTracingSceneRHI, GlobalResources, 
 				DispatchSizeX, DispatchSizeY
 			);
@@ -624,7 +611,7 @@ void FDeferredShadingSceneRenderer::RenderPathTracing(FRHICommandListImmediate&
 
 		RHICmdList.RayTraceDispatch(
 			View.RayTracingMaterialPipeline, 
-			RayGenShader->GetRayTracingShader(), 
+			RayGenShader.GetRayTracingShader(), 
 			RayTracingSceneRHI, GlobalResources, 
 			DispatchSizeX, DispatchSizeY 
 		);
@@ -721,8 +708,8 @@ void FDeferredShadingSceneRenderer::RenderPathTracing(FRHICommandListImmediate&
 	GraphicsPSOInit.RasterizerState = TStaticRasterizerState<FM_Solid, CM_None>::GetRHI();
 	GraphicsPSOInit.DepthStencilState = TStaticDepthStencilState<false, CF_Always>::GetRHI();
 	GraphicsPSOInit.BoundShaderState.VertexDeclarationRHI = GFilterVertexDeclaration.VertexDeclarationRHI;
-	GraphicsPSOInit.BoundShaderState.VertexShaderRHI = GETSAFERHISHADER_VERTEX(*VertexShader);
-	GraphicsPSOInit.BoundShaderState.PixelShaderRHI = GETSAFERHISHADER_PIXEL(*PixelShader);
+	GraphicsPSOInit.BoundShaderState.VertexShaderRHI = VertexShader.GetVertexShader();
+	GraphicsPSOInit.BoundShaderState.PixelShaderRHI = PixelShader.GetPixelShader();
 	GraphicsPSOInit.PrimitiveType = PT_TriangleList;
 	SetGraphicsPipelineState(RHICmdList, GraphicsPSOInit);
 
@@ -765,7 +752,7 @@ void FDeferredShadingSceneRenderer::RenderPathTracing(FRHICommandListImmediate&
 			DispatchSizeX, View.ViewRect.Height(),
 			FIntPoint(View.ViewRect.Width(), View.ViewRect.Height()),
 			SceneContext.GetBufferSizeXY(),
-			*VertexShader);
+			VertexShader);
 	}
 	RHICmdList.EndRenderPass();
 
diff --git a/Engine/Source/Runtime/Renderer/Private/PathTracing/PathCompactionCompute.cpp b/Engine/Source/Runtime/Renderer/Private/PathTracing/PathCompactionCompute.cpp
--- a/Engine/Source/Runtime/Renderer/Private/PathTracing/PathCompactionCompute.cpp
+++ b/Engine/Source/Runtime/Renderer/Private/PathTracing/PathCompactionCompute.cpp
@@ -64,7 +64,7 @@ public:
 		FRHIUnorderedAccessView* RadianceSortedAlphaUAV,
 		FRHIUnorderedAccessView* SampleCountSortedUAV)
 	{
-		FRHIComputeShader* ShaderRHI = GetComputeShader();
+		FRHIComputeShader* ShaderRHI = RHICmdList.GetBoundComputeShader();
 		FGlobalShader::SetParameters<FViewUniformShaderParameters>(RHICmdList, ShaderRHI, View.ViewUniformBuffer);
 
 		// Input textures
@@ -91,7 +91,7 @@ public:
 		FRHIUnorderedAccessView* SampleCountSortedUAV,
 		FRHIComputeFence* Fence)
 	{
-		FRHIComputeShader* ShaderRHI = GetComputeShader();
+		FRHIComputeShader* ShaderRHI = RHICmdList.GetBoundComputeShader();
 
 		SetUAVParameter(RHICmdList, ShaderRHI, RadianceSortedRedUAVParameter, FUnorderedAccessViewRHIRef());
 		SetUAVParameter(RHICmdList, ShaderRHI, RadianceSortedGreenUAVParameter, FUnorderedAccessViewRHIRef());
@@ -108,34 +108,19 @@ public:
 		RHICmdList.TransitionResources(TransitionAccess, TransitionPipeline, UAVs, 5, Fence);
 	}
 
-	virtual bool Serialize(FArchive& Ar)
-	{
-		bool bShaderHasOutdatedParameters = FGlobalShader::Serialize(Ar);
-		Ar << ViewParameter;
-		Ar << RadianceTextureParameter;
-		Ar << SampleCountTextureParameter;
-		Ar << PixelPositionTextureParameter;
-		Ar << RadianceSortedRedUAVParameter;
-		Ar << RadianceSortedGreenUAVParameter;
-		Ar << RadianceSortedBlueUAVParameter;
-		Ar << RadianceSortedAlphaUAVParameter;
-		Ar << SampleCountSortedUAVParameter;
-		return bShaderHasOutdatedParameters;
-	}
-
 private:
 	// Input parameters
-	FShaderResourceParameter ViewParameter;
-	FShaderResourceParameter RadianceTextureParameter;
-	FShaderResourceParameter SampleCountTextureParameter;
-	FShaderResourceParameter PixelPositionTextureParameter;
+	LAYOUT_FIELD(FShaderResourceParameter, ViewParameter);
+	LAYOUT_FIELD(FShaderResourceParameter, RadianceTextureParameter);
+	LAYOUT_FIELD(FShaderResourceParameter, SampleCountTextureParameter);
+	LAYOUT_FIELD(FShaderResourceParameter, PixelPositionTextureParameter);
 
 	// Output parameters
-	FShaderResourceParameter RadianceSortedRedUAVParameter;
-	FShaderResourceParameter RadianceSortedGreenUAVParameter;
-	FShaderResourceParameter RadianceSortedBlueUAVParameter;
-	FShaderResourceParameter RadianceSortedAlphaUAVParameter;
-	FShaderResourceParameter SampleCountSortedUAVParameter;
+	LAYOUT_FIELD(FShaderResourceParameter, RadianceSortedRedUAVParameter);
+	LAYOUT_FIELD(FShaderResourceParameter, RadianceSortedGreenUAVParameter);
+	LAYOUT_FIELD(FShaderResourceParameter, RadianceSortedBlueUAVParameter);
+	LAYOUT_FIELD(FShaderResourceParameter, RadianceSortedAlphaUAVParameter);
+	LAYOUT_FIELD(FShaderResourceParameter, SampleCountSortedUAVParameter);
 };
 
 IMPLEMENT_SHADER_TYPE(, FPathCompactionCS, TEXT("/Engine/Private/PathTracing/PathCompaction.usf"), TEXT("PathCompactionCS"), SF_Compute)
@@ -154,13 +139,13 @@ void FDeferredShadingSceneRenderer::ComputePathCompaction(
 {
 	const auto ShaderMap = GetGlobalShaderMap(FeatureLevel);
 	TShaderMapRef<FPathCompactionCS> PathCompactionComputeShader(ShaderMap);
-	RHICmdList.SetComputeShader(PathCompactionComputeShader->GetComputeShader());
+	RHICmdList.SetComputeShader(PathCompactionComputeShader.GetComputeShader());
 
 	FComputeFenceRHIRef Fence = RHICmdList.CreateComputeFence(TEXT("PathCompaction"));
 	PathCompactionComputeShader->SetParameters(RHICmdList, View, RadianceTexture, SampleCountTexture, PixelPositionTexture, RadianceSortedRedUAV, RadianceSortedGreenUAV, RadianceSortedBlueUAV, RadianceSortedAlphaUAV, SampleCountSortedUAV);
 	FIntPoint ViewSize = View.ViewRect.Size();
 	FIntVector NumGroups = FIntVector::DivideAndRoundUp(FIntVector(ViewSize.X, ViewSize.Y, 0), FPathCompactionCS::GetGroupSize());
-	DispatchComputeShader(RHICmdList, *PathCompactionComputeShader, NumGroups.X, NumGroups.Y, 1);
+	DispatchComputeShader(RHICmdList, PathCompactionComputeShader.GetShader(), NumGroups.X, NumGroups.Y, 1);
 	PathCompactionComputeShader->UnsetParameters(RHICmdList, EResourceTransitionAccess::EReadable, EResourceTransitionPipeline::EComputeToGfx, RadianceSortedRedUAV, RadianceSortedGreenUAV, RadianceSortedBlueUAV, RadianceSortedAlphaUAV, SampleCountSortedUAV, Fence);
 }
 #endif // RHI_RAYTRACING
