Copying Private-LoadTimes-4.24 stream to Main. Biggest changes are in Materials/Shader memory freezing.

diff --git a/Engine/Source/Runtime/Renderer/Private/PathTracing/PathTracingTestbed.cpp b/Engine/Source/Runtime/Renderer/Private/PathTracing/PathTracingTestbed.cpp
--- a/Engine/Source/Runtime/Renderer/Private/PathTracing/PathTracingTestbed.cpp
+++ b/Engine/Source/Runtime/Renderer/Private/PathTracing/PathTracingTestbed.cpp
@@ -54,7 +54,7 @@ public:
 		FRHIUnorderedAccessView* ResultsBuffer,
 		FRHIUnorderedAccessView* FloatResultsBuffer)
 	{
-		FRHIComputeShader* ShaderRHI = GetComputeShader();
+		FRHIComputeShader* ShaderRHI = RHICmdList.GetBoundComputeShader();
 		SetShaderValue(RHICmdList, ShaderRHI, SamplesCountParameter, SamplesCount);
 		SetShaderValue(RHICmdList, ShaderRHI, BrdfTypeParameter, BrdfType);
 		SetUAVParameter(RHICmdList, ShaderRHI, ResultsBufferParameter, ResultsBuffer);
@@ -69,27 +69,17 @@ public:
 		FRWBuffer& FloatResultsBuffer,
 		FRHIComputeFence* Fence)
 	{
-		FRHIComputeShader* ShaderRHI = GetComputeShader();
+		FRHIComputeShader* ShaderRHI = RHICmdList.GetBoundComputeShader();
 
 		RHICmdList.TransitionResource(TransitionAccess, TransitionPipeline, ResultsBuffer.UAV, Fence);
 		RHICmdList.TransitionResource(TransitionAccess, TransitionPipeline, FloatResultsBuffer.UAV, Fence);
 	}
 
-	virtual bool Serialize(FArchive& Ar)
-	{
-		bool bShaderHasOutdatedParameters = FGlobalShader::Serialize(Ar);
-		Ar << SamplesCountParameter;
-		Ar << BrdfTypeParameter;
-		Ar << ResultsBufferParameter;
-		Ar << FloatResultsBufferParameter;
-		return bShaderHasOutdatedParameters;
-	}
-
 private:
-	FShaderParameter SamplesCountParameter;
-	FShaderParameter BrdfTypeParameter;
-	FShaderResourceParameter ResultsBufferParameter;
-	FShaderResourceParameter FloatResultsBufferParameter;
+	LAYOUT_FIELD(FShaderParameter, SamplesCountParameter)
+	LAYOUT_FIELD(FShaderParameter, BrdfTypeParameter)
+	LAYOUT_FIELD(FShaderResourceParameter, ResultsBufferParameter)
+	LAYOUT_FIELD(FShaderResourceParameter, FloatResultsBufferParameter)
 };
 
 IMPLEMENT_SHADER_TYPE(, FTestBrdfIntegrityCS, TEXT("/Engine/Private/PathTracing/Material/PathTracingTestBrdfs.usf"), TEXT("TestBrdfIntegrityCS"), SF_Compute)
@@ -139,7 +129,7 @@ public:
 		uint32 NumPhiSteps,
 		FRHIUnorderedAccessView* PdfsResultsBuffer)
 	{
-		FRHIComputeShader* ShaderRHI = GetComputeShader();
+		FRHIComputeShader* ShaderRHI = RHICmdList.GetBoundComputeShader();
 		SetShaderValue(RHICmdList, ShaderRHI, BrdfTypeParameter, BrdfType);
 		SetShaderValue(RHICmdList, ShaderRHI, WoParameter, Wo);
 		SetShaderValue(RHICmdList, ShaderRHI, NumThetaStepsParameter, NumThetaSteps);
@@ -154,28 +144,17 @@ public:
 		FRWBuffer& PdfsResultsBuffer,
 		FRHIComputeFence* Fence)
 	{
-		FRHIComputeShader* ShaderRHI = GetComputeShader();
+		FRHIComputeShader* ShaderRHI = RHICmdList.GetBoundComputeShader();
 
 		RHICmdList.TransitionResource(TransitionAccess, TransitionPipeline, PdfsResultsBuffer.UAV, Fence);
 	}
 
-	virtual bool Serialize(FArchive& Ar)
-	{
-		bool bShaderHasOutdatedParameters = FGlobalShader::Serialize(Ar);
-		Ar << BrdfTypeParameter;
-		Ar << WoParameter;
-		Ar << NumThetaStepsParameter;
-		Ar << NumPhiStepsParameter;
-		Ar << PdfResultsBufferParameter;
-		return bShaderHasOutdatedParameters;
-	}
-
 private:
-	FShaderParameter BrdfTypeParameter;
-	FShaderParameter WoParameter;
-	FShaderParameter NumThetaStepsParameter;
-	FShaderParameter NumPhiStepsParameter;
-	FShaderResourceParameter PdfResultsBufferParameter;
+	LAYOUT_FIELD(FShaderParameter, BrdfTypeParameter)
+	LAYOUT_FIELD(FShaderParameter, WoParameter)
+	LAYOUT_FIELD(FShaderParameter, NumThetaStepsParameter)
+	LAYOUT_FIELD(FShaderParameter, NumPhiStepsParameter)
+	LAYOUT_FIELD(FShaderResourceParameter, PdfResultsBufferParameter)
 };
 
 IMPLEMENT_SHADER_TYPE(, FTestPdfIntegratesToOneCS, TEXT("/Engine/Private/PathTracing/Material/PathTracingTestPdfIntegration.usf"), TEXT("TestPDFIntegratesToOneCS"), SF_Compute)
@@ -227,7 +206,7 @@ public:
 		uint32 NumPhiSteps,
 		FRHIUnorderedAccessView* WisBinDistributionResultsBuffer)
 	{
-		FRHIComputeShader* ShaderRHI = GetComputeShader();
+		FRHIComputeShader* ShaderRHI = RHICmdList.GetBoundComputeShader();
 		SetShaderValue(RHICmdList, ShaderRHI, BrdfTypeParameter, BrdfType);
 		SetShaderValue(RHICmdList, ShaderRHI, NumSamplesParameter, NumSamples);
 		SetShaderValue(RHICmdList, ShaderRHI, WoParameter, Wo);
@@ -243,30 +222,18 @@ public:
 		FRWBuffer& PdfsResultsBuffer,
 		FRHIComputeFence* Fence)
 	{
-		FRHIComputeShader* ShaderRHI = GetComputeShader();
+		FRHIComputeShader* ShaderRHI = RHICmdList.GetBoundComputeShader();
 
 		RHICmdList.TransitionResource(TransitionAccess, TransitionPipeline, PdfsResultsBuffer.UAV, Fence);
 	}
 
-	virtual bool Serialize(FArchive& Ar)
-	{
-		bool bShaderHasOutdatedParameters = FGlobalShader::Serialize(Ar);
-		Ar << BrdfTypeParameter;
-		Ar << NumSamplesParameter;
-		Ar << WoParameter;
-		Ar << NumThetaStepsParameter;
-		Ar << NumPhiStepsParameter;
-		Ar << WisBinDistributionResultsParameter;
-		return bShaderHasOutdatedParameters;
-	}
-
 private:
-	FShaderParameter BrdfTypeParameter;
-	FShaderParameter NumSamplesParameter;
-	FShaderParameter WoParameter;
-	FShaderParameter NumThetaStepsParameter;
-	FShaderParameter NumPhiStepsParameter;
-	FShaderResourceParameter WisBinDistributionResultsParameter;
+	LAYOUT_FIELD(FShaderParameter, BrdfTypeParameter)
+	LAYOUT_FIELD(FShaderParameter, NumSamplesParameter)
+	LAYOUT_FIELD(FShaderParameter, WoParameter)
+	LAYOUT_FIELD(FShaderParameter, NumThetaStepsParameter)
+	LAYOUT_FIELD(FShaderParameter, NumPhiStepsParameter)
+	LAYOUT_FIELD(FShaderResourceParameter, WisBinDistributionResultsParameter)
 };
 
 IMPLEMENT_SHADER_TYPE(, FTestBrdfGenerateWiSamplesCS, TEXT("/Engine/Private/PathTracing/Material/PathTracingTestGenerateWiSamples.usf"), TEXT("TestGenerateWiSamplesCS"), SF_Compute)
@@ -316,7 +283,7 @@ public:
 		uint32 NumPhiSteps,
 		FRHIUnorderedAccessView* PdfsResultsBuffer)
 	{
-		FRHIComputeShader* ShaderRHI = GetComputeShader();
+		FRHIComputeShader* ShaderRHI = RHICmdList.GetBoundComputeShader();
 		SetShaderValue(RHICmdList, ShaderRHI, BrdfTypeParameter, BrdfType);
 		SetShaderValue(RHICmdList, ShaderRHI, WoParameter, Wo);
 		SetShaderValue(RHICmdList, ShaderRHI, NumThetaStepsParameter, NumThetaSteps);
@@ -331,28 +298,17 @@ public:
 		FRWBuffer& PdfsResultsBuffer,
 		FRHIComputeFence* Fence)
 	{
-		FRHIComputeShader* ShaderRHI = GetComputeShader();
+		FRHIComputeShader* ShaderRHI = RHICmdList.GetBoundComputeShader();
 
 		RHICmdList.TransitionResource(TransitionAccess, TransitionPipeline, PdfsResultsBuffer.UAV, Fence);
 	}
 
-	virtual bool Serialize(FArchive& Ar)
-	{
-		bool bShaderHasOutdatedParameters = FGlobalShader::Serialize(Ar);
-		Ar << BrdfTypeParameter;
-		Ar << WoParameter;
-		Ar << NumThetaStepsParameter;
-		Ar << NumPhiStepsParameter;
-		Ar << WisBinDistributionResultsParameter;
-		return bShaderHasOutdatedParameters;
-	}
-
 private:
-	FShaderParameter BrdfTypeParameter;
-	FShaderParameter WoParameter;
-	FShaderParameter NumThetaStepsParameter;
-	FShaderParameter NumPhiStepsParameter;
-	FShaderResourceParameter WisBinDistributionResultsParameter;
+	LAYOUT_FIELD(FShaderParameter, BrdfTypeParameter)
+	LAYOUT_FIELD(FShaderParameter, WoParameter)
+	LAYOUT_FIELD(FShaderParameter, NumThetaStepsParameter)
+	LAYOUT_FIELD(FShaderParameter, NumPhiStepsParameter)
+	LAYOUT_FIELD(FShaderResourceParameter, WisBinDistributionResultsParameter)
 };
 
 IMPLEMENT_SHADER_TYPE(, FTestBrdfIntegrateHemispherePatchCS, TEXT("/Engine/Private/PathTracing/Material/PathTracingTestIntegrateHemispherePatch.usf"), TEXT("TestIntegrateHemispherePatchCS"), SF_Compute)
@@ -403,7 +359,7 @@ void TestBRDFsIntegrity(void)
 	FRHICommandListImmediate& RHICmdList = FRHICommandListExecutor::GetImmediateCommandList();
 	const auto ShaderMap = GetGlobalShaderMap(GMaxRHIFeatureLevel);
 	TShaderMapRef<FTestBrdfIntegrityCS> BrdfIntegrityComputeShader(ShaderMap);
-	RHICmdList.SetComputeShader(BrdfIntegrityComputeShader->GetComputeShader());
+	RHICmdList.SetComputeShader(BrdfIntegrityComputeShader.GetComputeShader());
 
 	for (uint32 BrdfType = 0; BrdfType < (uint32) BrdfNames.Num(); ++BrdfType)
 	{
@@ -418,7 +374,7 @@ void TestBRDFsIntegrity(void)
 		BrdfIntegrityComputeShader->SetParameters(RHICmdList, NumSamples, BrdfType, BrdsResultsBuffer.UAV, FloatBrdsResultsBuffer.UAV);
 		
 		uint32 NumCSGroups = FMath::DivideAndRoundUp(NumSamples, FTestBrdfIntegrityCS::GetGroupSize());
-		DispatchComputeShader(RHICmdList, *BrdfIntegrityComputeShader, NumCSGroups, 1, 1);
+		DispatchComputeShader(RHICmdList, BrdfIntegrityComputeShader.GetShader(), NumCSGroups, 1, 1);
 
 		GDynamicRHI->RHISubmitCommandsAndFlushGPU();
 		GDynamicRHI->RHIBlockUntilGPUIdle();
@@ -496,7 +452,7 @@ void TestPDFsIntegrateToOne(void)
 	FRHICommandListImmediate& RHICmdList = FRHICommandListExecutor::GetImmediateCommandList();
 	const auto ShaderMap = GetGlobalShaderMap(GMaxRHIFeatureLevel);
 	TShaderMapRef<FTestPdfIntegratesToOneCS> PdfIntegrateToOneCS(ShaderMap);
-	RHICmdList.SetComputeShader(PdfIntegrateToOneCS->GetComputeShader());
+	RHICmdList.SetComputeShader(PdfIntegrateToOneCS.GetComputeShader());
 
 	TArray<FVector> Wos;
 	FRandomStream RandomStream(0);
@@ -526,7 +482,7 @@ void TestPDFsIntegrateToOne(void)
 		{
 			PdfIntegrateToOneCS->SetParameters(RHICmdList, BrdfType, Wo, NumThetaSteps, NumPhiSteps, PdfsResultsBuffer.UAV);
 
-			DispatchComputeShader(RHICmdList, *PdfIntegrateToOneCS, NumCSGroups.X, NumCSGroups.Y, 1);
+			DispatchComputeShader(RHICmdList, PdfIntegrateToOneCS.GetShader(), NumCSGroups.X, NumCSGroups.Y, 1);
 
 			GDynamicRHI->RHISubmitCommandsAndFlushGPU();
 			GDynamicRHI->RHIBlockUntilGPUIdle();
@@ -596,10 +552,10 @@ void TestBRDFandPDFConsistency(void)
 	const auto ShaderMap = GetGlobalShaderMap(GMaxRHIFeatureLevel);
 	
 	TShaderMapRef<FTestBrdfGenerateWiSamplesCS> BrdfGenerateWiSamplesCS(ShaderMap);
-	RHICmdList.SetComputeShader(BrdfGenerateWiSamplesCS->GetComputeShader());
+	RHICmdList.SetComputeShader(BrdfGenerateWiSamplesCS.GetComputeShader());
 
 	TShaderMapRef<FTestBrdfIntegrateHemispherePatchCS> IntegrateHemispherePatchCS(ShaderMap);
-	RHICmdList.SetComputeShader(IntegrateHemispherePatchCS->GetComputeShader());
+	RHICmdList.SetComputeShader(IntegrateHemispherePatchCS.GetComputeShader());
 
 	for (uint32 BrdfType = 0; BrdfType < (uint32)BrdfNames.Num(); ++BrdfType)
 	{
@@ -622,7 +578,7 @@ void TestBRDFandPDFConsistency(void)
 				BrdfGenerateWiSamplesCS->SetParameters(RHICmdList, BrdfType, WiSamplesCount, Wo, NumThetaSteps, NumPhiSteps, SampledWisResultsBuffer.UAV);
 
 				uint32 NumCSGroupsSamplesCS = FMath::DivideAndRoundUp(WiSamplesCount, FTestBrdfGenerateWiSamplesCS::GetGroupSize());
-				DispatchComputeShader(RHICmdList, *BrdfGenerateWiSamplesCS, NumCSGroupsSamplesCS, 1, 1);
+				DispatchComputeShader(RHICmdList, BrdfGenerateWiSamplesCS.GetShader(), NumCSGroupsSamplesCS, 1, 1);
 
 				GDynamicRHI->RHISubmitCommandsAndFlushGPU();
 				GDynamicRHI->RHIBlockUntilGPUIdle();
@@ -666,7 +622,7 @@ void TestBRDFandPDFConsistency(void)
 				FIntVector HemispherePatchCSDimensions = FIntVector(NumPhiSteps, NumThetaSteps, 1);
 				FIntVector HemispherePatchCSNumGroups = FIntVector::DivideAndRoundUp(HemispherePatchCSDimensions, FTestBrdfIntegrateHemispherePatchCS::GetGroupSize());
 
-				DispatchComputeShader(RHICmdList, *IntegrateHemispherePatchCS, HemispherePatchCSNumGroups.X, HemispherePatchCSNumGroups.Y, 1);
+				DispatchComputeShader(RHICmdList, IntegrateHemispherePatchCS.GetShader(), HemispherePatchCSNumGroups.X, HemispherePatchCSNumGroups.Y, 1);
 
 				GDynamicRHI->RHISubmitCommandsAndFlushGPU();
 				GDynamicRHI->RHIBlockUntilGPUIdle();
diff --git a/Engine/Source/Runtime/Renderer/Private/PathTracing/RayCounterCompute.cpp b/Engine/Source/Runtime/Renderer/Private/PathTracing/RayCounterCompute.cpp
--- a/Engine/Source/Runtime/Renderer/Private/PathTracing/RayCounterCompute.cpp
+++ b/Engine/Source/Runtime/Renderer/Private/PathTracing/RayCounterCompute.cpp
@@ -53,7 +53,7 @@ public:
 		const FIntPoint& ViewSize,
 		FRHIUnorderedAccessView* TotalRayCountBuffer)
 	{
-		FRHIComputeShader* ShaderRHI = GetComputeShader();
+		FRHIComputeShader* ShaderRHI = RHICmdList.GetBoundComputeShader();
 
 		SetTextureParameter(RHICmdList, ShaderRHI, RayCountPerPixelParameter, RayCountPerPixelBuffer);
 		SetShaderValue(RHICmdList, ShaderRHI, ViewSizeParameter, ViewSize);
@@ -67,27 +67,18 @@ public:
 		FRWBuffer& TotalRayCountBuffer,
 		FRHIComputeFence* Fence)
 	{
-		FRHIComputeShader* ShaderRHI = GetComputeShader();
+		FRHIComputeShader* ShaderRHI = RHICmdList.GetBoundComputeShader();
 
 		RHICmdList.TransitionResource(TransitionAccess, TransitionPipeline, TotalRayCountBuffer.UAV, Fence);
 	}
 
-	virtual bool Serialize(FArchive& Ar)
-	{
-		bool bShaderHasOutdatedParameters = FGlobalShader::Serialize(Ar);
-		Ar << RayCountPerPixelParameter;
-		Ar << ViewSizeParameter;
-		Ar << TotalRayCountParameter;
-		return bShaderHasOutdatedParameters;
-	}
-
 private:
 	// Input parameters
-	FShaderResourceParameter RayCountPerPixelParameter;
-	FShaderParameter ViewSizeParameter;
+	LAYOUT_FIELD(FShaderResourceParameter, RayCountPerPixelParameter)
+	LAYOUT_FIELD(FShaderParameter, ViewSizeParameter)
 
 	// Output parameters
-	FShaderResourceParameter TotalRayCountParameter;
+	LAYOUT_FIELD(FShaderResourceParameter, TotalRayCountParameter)
 };
 
 IMPLEMENT_SHADER_TYPE(, FRayCounterCS, TEXT("/Engine/Private/PathTracing/PathTracingRayCounterComputeShader.usf"), TEXT("RayCounterCS"), SF_Compute)
@@ -100,13 +91,13 @@ void FDeferredShadingSceneRenderer::ComputeRayCount(FRHICommandListImmediate& RH
 
 	const auto ShaderMap = GetGlobalShaderMap(FeatureLevel);
 	TShaderMapRef<FRayCounterCS> RayCounterComputeShader(ShaderMap);
-	RHICmdList.SetComputeShader(RayCounterComputeShader->GetComputeShader());
+	RHICmdList.SetComputeShader(RayCounterComputeShader.GetComputeShader());
 
 	FIntPoint ViewSize = View.ViewRect.Size();
 	RayCounterComputeShader->SetParameters(RHICmdList, RayCountPerPixelTexture, ViewSize, ViewState->TotalRayCountBuffer->UAV);
 
 	int32 NumGroups = FMath::DivideAndRoundUp<int32>(ViewSize.Y, FRayCounterCS::GetGroupSize());
-	DispatchComputeShader(RHICmdList, *RayCounterComputeShader, NumGroups, 1, 1);
+	DispatchComputeShader(RHICmdList, RayCounterComputeShader.GetShader(), NumGroups, 1, 1);
 
 	FRHIGPUBufferReadback* RayCountGPUReadback = ViewState->RayCountGPUReadback;
 
diff --git a/Engine/Source/Runtime/Renderer/Private/PathTracing/VarianceMipTreeBuild.cpp b/Engine/Source/Runtime/Renderer/Private/PathTracing/VarianceMipTreeBuild.cpp
--- a/Engine/Source/Runtime/Renderer/Private/PathTracing/VarianceMipTreeBuild.cpp
+++ b/Engine/Source/Runtime/Renderer/Private/PathTracing/VarianceMipTreeBuild.cpp
@@ -55,7 +55,7 @@ public:
 		uint32 MipLevel,
 		FRWBuffer& VarianceMipTree)
 	{
-		FRHIComputeShader* ShaderRHI = GetComputeShader();
+		FRHIComputeShader* ShaderRHI = RHICmdList.GetBoundComputeShader();
 
 		SetShaderValue(RHICmdList, ShaderRHI, ViewSizeParameter, ViewSize);
 		SetShaderValue(RHICmdList, ShaderRHI, VarianceMapDimensionsParameter, VarianceMapDimensions);
@@ -73,32 +73,20 @@ public:
 		FRWBuffer& VarianceMap,
 		FRHIComputeFence* Fence)
 	{
-		FRHIComputeShader* ShaderRHI = GetComputeShader();
+		FRHIComputeShader* ShaderRHI = RHICmdList.GetBoundComputeShader();
 
 		VarianceMipTreeParameter.UnsetUAV(RHICmdList, ShaderRHI);
 		RHICmdList.TransitionResource(TransitionAccess, TransitionPipeline, VarianceMap.UAV, Fence);
 	}
 
-	virtual bool Serialize(FArchive& Ar)
-	{
-		bool bShaderHasOutdatedParameters = FGlobalShader::Serialize(Ar);
-		Ar << RadianceTextureParameter;
-		Ar << RadianceTextureSamplerParameter;
-		Ar << ViewSizeParameter;
-		Ar << VarianceMapDimensionsParameter;
-		Ar << MipLevelParameter;
-		Ar << VarianceMipTreeParameter;
-		return bShaderHasOutdatedParameters;
-	}
-
 private:
-	FShaderResourceParameter RadianceTextureParameter;
-	FShaderResourceParameter RadianceTextureSamplerParameter;
-	FShaderParameter ViewSizeParameter;
+	LAYOUT_FIELD(FShaderResourceParameter, RadianceTextureParameter)
+	LAYOUT_FIELD(FShaderResourceParameter, RadianceTextureSamplerParameter)
+	LAYOUT_FIELD(FShaderParameter, ViewSizeParameter)
 
-	FShaderParameter VarianceMapDimensionsParameter;
-	FShaderParameter MipLevelParameter;
-	FRWShaderParameter VarianceMipTreeParameter;
+	LAYOUT_FIELD(FShaderParameter, VarianceMapDimensionsParameter)
+	LAYOUT_FIELD(FShaderParameter, MipLevelParameter)
+	LAYOUT_FIELD(FRWShaderParameter, VarianceMipTreeParameter)
 };
 
 IMPLEMENT_SHADER_TYPE(, FBuildVarianceMipTreeCS, TEXT("/Engine/Private/PathTracing/BuildVarianceMipTreeComputeShader.usf"), TEXT("BuildVarianceMipTreeCS"), SF_Compute)
@@ -108,7 +96,7 @@ void FDeferredShadingSceneRenderer::BuildVarianceMipTree(FRHICommandListImmediat
 {
 	const auto ShaderMap = GetGlobalShaderMap(FeatureLevel);
 	TShaderMapRef<FBuildVarianceMipTreeCS> BuildVarianceMipTreeComputeShader(ShaderMap);
-	RHICmdList.SetComputeShader(BuildVarianceMipTreeComputeShader->GetComputeShader());
+	RHICmdList.SetComputeShader(BuildVarianceMipTreeComputeShader.GetComputeShader());
 
 	// Allocate MIP tree
 	FIntPoint ViewSize = View.ViewRect.Size();
@@ -130,7 +118,7 @@ void FDeferredShadingSceneRenderer::BuildVarianceMipTree(FRHICommandListImmediat
 		BuildVarianceMipTreeComputeShader->SetParameters(RHICmdList, MeanAndDeviationTexture, ViewSize, VarianceMipTreeDimensions, MipLevel, VarianceMipTree);
 		FIntVector MipLevelDimensions = FIntVector(VarianceMipTreeDimensions.X >> MipLevel, VarianceMipTreeDimensions.Y >> MipLevel, 1);
 		FIntVector NumGroups = FIntVector::DivideAndRoundUp(MipLevelDimensions, FBuildVarianceMipTreeCS::GetGroupSize());
-		DispatchComputeShader(RHICmdList, *BuildVarianceMipTreeComputeShader, NumGroups.X, NumGroups.Y, 1);
+		DispatchComputeShader(RHICmdList, BuildVarianceMipTreeComputeShader.GetShader(), NumGroups.X, NumGroups.Y, 1);
 		BuildVarianceMipTreeComputeShader->UnsetParameters(RHICmdList, EResourceTransitionAccess::ERWBarrier, EResourceTransitionPipeline::EComputeToCompute, VarianceMipTree, MipLevelFence);
 	}
 	FComputeFenceRHIRef TransitionFence = RHICmdList.CreateComputeFence(TEXT("VarianceMipTree Transition"));
@@ -173,24 +161,16 @@ public:
 		const FIntVector Dimensions,
 		const FRWBuffer& MipTree)
 	{
-		FRHIPixelShader* ShaderRHI = GetPixelShader();
+		FRHIPixelShader* ShaderRHI = RHICmdList.GetBoundPixelShader();
 		FGlobalShader::SetParameters<FViewUniformShaderParameters>(RHICmdList, ShaderRHI, View.ViewUniformBuffer);
 
 		SetShaderValue(RHICmdList, ShaderRHI, DimensionsParameter, Dimensions);
 		SetSRVParameter(RHICmdList, ShaderRHI, MipTreeParameter, MipTree.SRV);
 	}
 
-	virtual bool Serialize(FArchive& Ar) override
-	{
-		bool bShaderHasOutdatedParameters = FGlobalShader::Serialize(Ar);
-		Ar << DimensionsParameter;
-		Ar << MipTreeParameter;
-		return bShaderHasOutdatedParameters;
-	}
-
 private:
-	FShaderParameter DimensionsParameter;
-	FShaderResourceParameter MipTreeParameter;
+	LAYOUT_FIELD(FShaderParameter, DimensionsParameter)
+	LAYOUT_FIELD(FShaderResourceParameter, MipTreeParameter)
 };
 
 IMPLEMENT_SHADER_TYPE(, FVisualizeMipTreePS, TEXT("/Engine/Private/PathTracing/VisualizeMipTreePixelShader.usf"), TEXT("VisualizeMipTreePS"), SF_Pixel)
@@ -223,8 +203,8 @@ void FDeferredShadingSceneRenderer::VisualizeVarianceMipTree(FRHICommandListImme
 	GraphicsPSOInit.RasterizerState = TStaticRasterizerState<FM_Solid, CM_None>::GetRHI();
 	GraphicsPSOInit.DepthStencilState = TStaticDepthStencilState<false, CF_Always>::GetRHI();
 	GraphicsPSOInit.BoundShaderState.VertexDeclarationRHI = GFilterVertexDeclaration.VertexDeclarationRHI;
-	GraphicsPSOInit.BoundShaderState.VertexShaderRHI = GETSAFERHISHADER_VERTEX(*VertexShader);
-	GraphicsPSOInit.BoundShaderState.PixelShaderRHI = GETSAFERHISHADER_PIXEL(*PixelShader);
+	GraphicsPSOInit.BoundShaderState.VertexShaderRHI = VertexShader.GetVertexShader();
+	GraphicsPSOInit.BoundShaderState.PixelShaderRHI = PixelShader.GetPixelShader();
 	GraphicsPSOInit.PrimitiveType = PT_TriangleList;
 	SetGraphicsPipelineState(RHICmdList, GraphicsPSOInit);
 
@@ -241,7 +221,7 @@ void FDeferredShadingSceneRenderer::VisualizeVarianceMipTree(FRHICommandListImme
 		View.ViewRect.Width(), View.ViewRect.Height(),
 		FIntPoint(View.ViewRect.Width(), View.ViewRect.Height()),
 		SceneContext.GetBufferSizeXY(),
-		*VertexShader);
+		VertexShader);
 	ResolveSceneColor(RHICmdList);
 	RHICmdList.EndRenderPass();
 
diff --git a/Engine/Source/Runtime/Renderer/Private/PlanarReflectionRendering.cpp b/Engine/Source/Runtime/Renderer/Private/PlanarReflectionRendering.cpp
--- a/Engine/Source/Runtime/Renderer/Private/PlanarReflectionRendering.cpp
+++ b/Engine/Source/Runtime/Renderer/Private/PlanarReflectionRendering.cpp
@@ -173,7 +173,7 @@ public:
 
 	void SetParameters(FRHICommandList& RHICmdList, const FSceneView& View, const FPlanarReflectionSceneProxy* ReflectionSceneProxy, FRHITexture* SceneColorInput, int32 FilterWidth)
 	{
-		FRHIPixelShader* ShaderRHI = GetPixelShader();
+		FRHIPixelShader* ShaderRHI = RHICmdList.GetBoundPixelShader();
 		FGlobalShader::SetParameters<FViewUniformShaderParameters>(RHICmdList, ShaderRHI, View.ViewUniformBuffer);
 		SceneTextureParameters.Set(RHICmdList, ShaderRHI, View.FeatureLevel, ESceneTextureSetupMode::All);
 
@@ -189,25 +189,13 @@ public:
 		SetUniformBufferParameterImmediate(RHICmdList, ShaderRHI, GetUniformBufferParameter<FPlanarReflectionUniformParameters>(), PlanarReflectionUniformParameters);
 	}
 
-	// FShader interface.
-	virtual bool Serialize(FArchive& Ar) override
-	{
-		bool bShaderHasOutdatedParameters = FGlobalShader::Serialize(Ar);
-		Ar << KernelRadiusY;
-		Ar << InvPrefilterRoughnessDistance;
-		Ar << SceneColorInputTexture;
-		Ar << SceneColorInputSampler;
-		Ar << SceneTextureParameters;
-		return bShaderHasOutdatedParameters;
-	}
-
 private:
 
-	FShaderParameter KernelRadiusY;
-	FShaderParameter InvPrefilterRoughnessDistance;
-	FShaderResourceParameter SceneColorInputTexture;
-	FShaderResourceParameter SceneColorInputSampler;
-	FSceneTextureShaderParameters SceneTextureParameters;
+	LAYOUT_FIELD(FShaderParameter, KernelRadiusY)
+	LAYOUT_FIELD(FShaderParameter, InvPrefilterRoughnessDistance)
+	LAYOUT_FIELD(FShaderResourceParameter, SceneColorInputTexture)
+	LAYOUT_FIELD(FShaderResourceParameter, SceneColorInputSampler)
+	LAYOUT_FIELD(FSceneTextureShaderParameters, SceneTextureParameters)
 };
 
 IMPLEMENT_SHADER_TYPE(template<>, FPrefilterPlanarReflectionPS<false>, TEXT("/Engine/Private/PlanarReflectionShaders.usf"), TEXT("PrefilterPlanarReflectionPS"), SF_Pixel);
@@ -252,8 +240,8 @@ void PrefilterPlanarReflection(FRHICommandListImmediate& RHICmdList, FViewInfo&
 			TShaderMapRef<FPrefilterPlanarReflectionPS<bEnablePlanarReflectionPrefilter> > PixelShader(View.ShaderMap);
 
 			GraphicsPSOInit.BoundShaderState.VertexDeclarationRHI = GFilterVertexDeclaration.VertexDeclarationRHI;
-			GraphicsPSOInit.BoundShaderState.VertexShaderRHI = GETSAFERHISHADER_VERTEX(*VertexShader);
-			GraphicsPSOInit.BoundShaderState.PixelShaderRHI = GETSAFERHISHADER_PIXEL(*PixelShader);
+			GraphicsPSOInit.BoundShaderState.VertexShaderRHI = VertexShader.GetVertexShader();
+			GraphicsPSOInit.BoundShaderState.PixelShaderRHI = PixelShader.GetPixelShader();
 			GraphicsPSOInit.PrimitiveType = PT_TriangleList;
 
 			SetGraphicsPipelineState(RHICmdList, GraphicsPSOInit);
@@ -278,7 +266,7 @@ void PrefilterPlanarReflection(FRHICommandListImmediate& RHICmdList, FViewInfo&
 				UVSize.X, UVSize.Y,
 				View.ViewRect.Size(),
 				FSceneRenderTargets::Get(RHICmdList).GetBufferSizeXY(),
-				*VertexShader,
+				VertexShader,
 				EDRF_UseTriangleOptimization);
 		}
 		RHICmdList.EndRenderPass();
@@ -768,8 +756,8 @@ void FDeferredShadingSceneRenderer::RenderDeferredPlanarReflections(FRDGBuilder&
 			TShaderMapRef<FPlanarReflectionPS> PixelShader(View.ShaderMap);
 
 			GraphicsPSOInit.BoundShaderState.VertexDeclarationRHI = GFilterVertexDeclaration.VertexDeclarationRHI;
-			GraphicsPSOInit.BoundShaderState.VertexShaderRHI = GETSAFERHISHADER_VERTEX(*VertexShader);
-			GraphicsPSOInit.BoundShaderState.PixelShaderRHI = GETSAFERHISHADER_PIXEL(*PixelShader);
+			GraphicsPSOInit.BoundShaderState.VertexShaderRHI = VertexShader.GetVertexShader();
+			GraphicsPSOInit.BoundShaderState.PixelShaderRHI = PixelShader.GetPixelShader();
 			GraphicsPSOInit.PrimitiveType = PT_TriangleList;
 
 			SetGraphicsPipelineState(RHICmdList, GraphicsPSOInit);
@@ -782,7 +770,7 @@ void FDeferredShadingSceneRenderer::RenderDeferredPlanarReflections(FRDGBuilder&
 		
 				FPlanarReflectionPS::FParameters ShaderParameters = *PassParameters;
 				ShaderParameters.PlanarReflectionParameters = CreateUniformBufferImmediate(PlanarReflectionUniformParameters, UniformBuffer_SingleDraw);
-				SetShaderParameters(RHICmdList, *PixelShader, PixelShader->GetPixelShader(), ShaderParameters);
+				SetShaderParameters(RHICmdList, PixelShader, PixelShader.GetPixelShader(), ShaderParameters);
 			}
 
 			DrawRectangle(
@@ -793,7 +781,7 @@ void FDeferredShadingSceneRenderer::RenderDeferredPlanarReflections(FRDGBuilder&
 				View.ViewRect.Width(), View.ViewRect.Height(),
 				View.ViewRect.Size(),
 				FSceneRenderTargets::Get(RHICmdList).GetBufferSizeXY(),
-				*VertexShader,
+				VertexShader,
 				EDRF_UseTriangleOptimization);
 		}
 	});
diff --git a/Engine/Source/Runtime/Renderer/Private/PostProcess/DiaphragmDOF.cpp b/Engine/Source/Runtime/Renderer/Private/PostProcess/DiaphragmDOF.cpp
--- a/Engine/Source/Runtime/Renderer/Private/PostProcess/DiaphragmDOF.cpp
+++ b/Engine/Source/Runtime/Renderer/Private/PostProcess/DiaphragmDOF.cpp
@@ -1598,7 +1598,7 @@ FRDGTextureRef DiaphragmDOF::AddPasses(
 				FMath::FloorToInt(CocModel.ComputeViewMinForegroundCocRadius(PassViewSize.X)),
 				FMath::CeilToInt(CocModel.ComputeViewMaxBackgroundCocRadius(PassViewSize.X)),
 				PassViewSize.X, PassViewSize.Y),
-			*ComputeShader,
+			ComputeShader,
 			PassParameters,
 			GroupCount);
 
@@ -1679,7 +1679,7 @@ FRDGTextureRef DiaphragmDOF::AddPasses(
 				RDG_EVENT_NAME("DOF FlattenCoc(Gather4=%s) %dx%d",
 					PermutationVector.Get<FDiaphragmDOFCocFlattenCS::FDoCocGather4>() ? TEXT("Yes") : TEXT("No"),
 					GatheringViewSize.X, GatheringViewSize.Y),
-				*ComputeShader,
+				ComputeShader,
 				PassParameters,
 				FComputeShaderUtils::GetGroupCount(GatheringViewSize, kDefaultGroupSize));
 		}
@@ -1735,7 +1735,7 @@ FRDGTextureRef DiaphragmDOF::AddPasses(
 				RDG_EVENT_NAME("DOF DilateCoc(1/16 %s radius=%d step=%d) %dx%d",
 					GetEventName(Mode), SampleRadiusCount, SampleOffsetMultipler,
 					DilatePassViewSize.X, DilatePassViewSize.Y),
-				*ComputeShader,
+				ComputeShader,
 				PassParameters,
 				DilateGroupCount);
 
@@ -1903,7 +1903,7 @@ FRDGTextureRef DiaphragmDOF::AddPasses(
 			FComputeShaderUtils::AddPass(
 				GraphBuilder,
 				RDG_EVENT_NAME("DOF Downsample %dx%d", PassViewSize.X, PassViewSize.Y),
-				*ComputeShader,
+				ComputeShader,
 				PassParameters,
 				FComputeShaderUtils::GetGroupCount(PassViewSize, kDefaultGroupSize));
 		}
@@ -1979,7 +1979,7 @@ FRDGTextureRef DiaphragmDOF::AddPasses(
 					bBackgroundHybridScattering ? TEXT("Yes") : TEXT("No"),
 					bRGBBufferSeparateCocBuffer ? TEXT(" R11G11B10") : TEXT(""),
 					PassViewSize.X, PassViewSize.Y),
-				*ComputeShader,
+				ComputeShader,
 				PassParameters,
 				FComputeShaderUtils::GetGroupCount(PassViewSize, kDefaultGroupSize));
 
@@ -2019,7 +2019,7 @@ FRDGTextureRef DiaphragmDOF::AddPasses(
 
 			// TODO(RDG): ERDGPassFlags::GenerateMips has no reason to exist.
 			TShaderMapRef<FDiaphragmDOFReduceCS> ComputeShader(View.ShaderMap, PermutationVector);
-			ClearUnusedGraphResources(*ComputeShader, PassParameters);
+			ClearUnusedGraphResources(ComputeShader, PassParameters);
 			GraphBuilder.AddPass(
 				RDG_EVENT_NAME("DOF Reduce(Mips=[%d;%d]%s) %dx%d",
 					ReducedMipLevelCount,
@@ -2030,7 +2030,7 @@ FRDGTextureRef DiaphragmDOF::AddPasses(
 				ERDGPassFlags::Compute | ERDGPassFlags::GenerateMips,
 				[PassParameters, ComputeShader, PassViewSize](FRHICommandList& RHICmdList)
 			{
-				FComputeShaderUtils::Dispatch(RHICmdList, *ComputeShader, *PassParameters, FComputeShaderUtils::GetGroupCount(PassViewSize, kDefaultGroupSize));
+				FComputeShaderUtils::Dispatch(RHICmdList, ComputeShader, *PassParameters, FComputeShaderUtils::GetGroupCount(PassViewSize, kDefaultGroupSize));
 			});
 
 			ReducedMipLevelCount += ProcessingMipLevelCount;
@@ -2057,7 +2057,7 @@ FRDGTextureRef DiaphragmDOF::AddPasses(
 			FComputeShaderUtils::AddPass(
 				GraphBuilder,
 				RDG_EVENT_NAME("DOF ScatterGroupPack"),
-				*ComputeShader,
+				ComputeShader,
 				PassParameters,
 				FIntVector(2, 1, 1));
 		}
@@ -2108,7 +2108,7 @@ FRDGTextureRef DiaphragmDOF::AddPasses(
 				BokehModel.BokehShape == DiaphragmDOF::EBokehShape::RoundedBlades ? TEXT("Rounded") : TEXT("Straight"),
 				GetEventName(LUTFormat),
 				BokehLUTDesc.Extent.X, BokehLUTDesc.Extent.Y),
-			*ComputeShader,
+			ComputeShader,
 			PassParameters,
 			FComputeShaderUtils::GetGroupCount(BokehLUTDesc.Extent, kDefaultGroupSize));
 
@@ -2304,7 +2304,7 @@ FRDGTextureRef DiaphragmDOF::AddPasses(
 					PermutationVector.Get<FDDOFClampInputUVDim>() ? TEXT(" ClampUV") : TEXT(""),
 					PermutationVector.Get<FDDOFRGBColorBufferDim>() ? TEXT(" R11G11B10") : TEXT(""),
 					GatheringViewSize.X, GatheringViewSize.Y),
-				*ComputeShader,
+				ComputeShader,
 				PassParameters,
 				FComputeShaderUtils::GetGroupCount(GatheringViewSize, kDefaultGroupSize));
 		}; // AddGatherPass()
@@ -2365,7 +2365,7 @@ FRDGTextureRef DiaphragmDOF::AddPasses(
 					GetEventName(ConvolutionSettings.LayerProcessing), GetEventName(ConvolutionSettings.PostfilterMethod),
 					PermutationVector.Get<FDiaphragmDOFPostfilterCS::FTileOptimization>() ? TEXT(" TileOptimisation") : TEXT(""),
 					GatheringViewSize.X, GatheringViewSize.Y),
-				*ComputeShader,
+				ComputeShader,
 				PassParameters,
 				FComputeShaderUtils::GetGroupCount(GatheringViewSize, kDefaultGroupSize));
 			
@@ -2412,7 +2412,7 @@ FRDGTextureRef DiaphragmDOF::AddPasses(
 			PassParameters->RenderTargets[0] = FRenderTargetBinding(
 				ConvolutionTextures->SceneColor, ERenderTargetLoadAction::ELoad);
 
-			ClearUnusedGraphResources(*VertexShader, *PixelShader, PassParameters);
+			ClearUnusedGraphResources(VertexShader, PixelShader, PassParameters);
 
 			GraphBuilder.AddPass(
 				RDG_EVENT_NAME("DOF IndirectScatter(%s Bokeh=%s Occlusion=%s) %dx%d",
@@ -2433,12 +2433,12 @@ FRDGTextureRef DiaphragmDOF::AddPasses(
 				GraphicsPSOInit.RasterizerState = TStaticRasterizerState<>::GetRHI();
 				GraphicsPSOInit.PrimitiveType = GRHISupportsRectTopology ? PT_RectList : PT_TriangleList;
 				GraphicsPSOInit.BoundShaderState.VertexDeclarationRHI = GEmptyVertexDeclaration.VertexDeclarationRHI;
-				GraphicsPSOInit.BoundShaderState.VertexShaderRHI = GETSAFERHISHADER_VERTEX(*VertexShader);
-				GraphicsPSOInit.BoundShaderState.PixelShaderRHI = GETSAFERHISHADER_PIXEL(*PixelShader);
+				GraphicsPSOInit.BoundShaderState.VertexShaderRHI = VertexShader.GetVertexShader();
+				GraphicsPSOInit.BoundShaderState.PixelShaderRHI = PixelShader.GetPixelShader();
 				SetGraphicsPipelineState(RHICmdList, GraphicsPSOInit);
 
-				SetShaderParameters(RHICmdList, *VertexShader, VertexShader->GetVertexShader(), *PassParameters);
-				SetShaderParameters(RHICmdList, *PixelShader, PixelShader->GetPixelShader(), *PassParameters);
+				SetShaderParameters(RHICmdList, VertexShader, VertexShader.GetVertexShader(), *PassParameters);
+				SetShaderParameters(RHICmdList, PixelShader, PixelShader.GetPixelShader(), *PassParameters);
 
 				RHICmdList.SetStreamSource(0, NULL, 0);
 
@@ -2598,7 +2598,7 @@ FRDGTextureRef DiaphragmDOF::AddPasses(
 				RecombineQuality,
 				GetEventName(PermutationVector.Get<FDDOFBokehSimulationDim>()),
 				PassViewRect.Width(), PassViewRect.Height()),
-			*ComputeShader,
+			ComputeShader,
 			PassParameters,
 			FComputeShaderUtils::GetGroupCount(PassViewRect.Size(), kDefaultGroupSize));
 	}
diff --git a/Engine/Source/Runtime/Renderer/Private/PostProcess/DiaphragmDOF.h b/Engine/Source/Runtime/Renderer/Private/PostProcess/DiaphragmDOF.h
--- a/Engine/Source/Runtime/Renderer/Private/PostProcess/DiaphragmDOF.h
+++ b/Engine/Source/Runtime/Renderer/Private/PostProcess/DiaphragmDOF.h
@@ -115,7 +115,7 @@ struct FBokehModel
 
 
 /** Returns whether DOF is supported. */
-inline bool IsSupported(EShaderPlatform ShaderPlatform)
+inline bool IsSupported(const FStaticShaderPlatform ShaderPlatform)
 {
 	// Only compile diaphragm DOF on platform it has been tested to ensure this is not blocking anyone else.
 	return 
@@ -125,7 +125,7 @@ inline bool IsSupported(EShaderPlatform ShaderPlatform)
 		IsVulkanSM5Platform(ShaderPlatform) ||
 		IsMetalSM5Platform(ShaderPlatform) ||
 		ShaderPlatform == SP_SWITCH ||
-		FDataDrivenShaderPlatformInfo::GetInfo(ShaderPlatform).bSupportsDiaphragmDOF;
+		FDataDrivenShaderPlatformInfo::GetSupportsDiaphragmDOF(ShaderPlatform);
 }
 
 
diff --git a/Engine/Source/Runtime/Renderer/Private/PostProcess/PostProcessAA.cpp b/Engine/Source/Runtime/Renderer/Private/PostProcess/PostProcessAA.cpp
--- a/Engine/Source/Runtime/Renderer/Private/PostProcess/PostProcessAA.cpp
+++ b/Engine/Source/Runtime/Renderer/Private/PostProcess/PostProcessAA.cpp
@@ -64,7 +64,7 @@ FScreenPassTexture AddFXAAPass(FRDGBuilder& GraphBuilder, const FViewInfo& View,
 	check(Inputs.Quality != EFXAAQuality::MAX);
 
 	FScreenPassRenderTarget Output = Inputs.OverrideOutput;
-
+	
 	if (!Output.IsValid())
 	{
 		Output = FScreenPassRenderTarget::CreateFromInput(GraphBuilder, Inputs.SceneColor, View.GetOverwriteLoadAction(), TEXT("FXAA"));
@@ -111,14 +111,14 @@ FScreenPassTexture AddFXAAPass(FRDGBuilder& GraphBuilder, const FViewInfo& View,
 		View,
 		OutputViewport,
 		FScreenPassTextureViewport(Inputs.SceneColor),
-		FScreenPassPipelineState(*VertexShader, *PixelShader),
+		FScreenPassPipelineState(VertexShader, PixelShader),
 		EScreenPassDrawFlags::AllowHMDHiddenAreaMask,
 		PassParameters,
 		[VertexShader, PixelShader, PassParameters](FRHICommandList& RHICmdList)
 	{
-		SetShaderParameters(RHICmdList, *VertexShader, VertexShader->GetVertexShader(), *PassParameters);
-		SetShaderParameters(RHICmdList, *PixelShader, PixelShader->GetPixelShader(), *PassParameters);
+		SetShaderParameters(RHICmdList, VertexShader, VertexShader.GetVertexShader(), *PassParameters);
+		SetShaderParameters(RHICmdList, PixelShader, PixelShader.GetPixelShader(), *PassParameters);
 	});
 
 	return MoveTemp(Output);
-}
\ No newline at end of file
+}
diff --git a/Engine/Source/Runtime/Renderer/Private/PostProcess/PostProcessBloomSetup.cpp b/Engine/Source/Runtime/Renderer/Private/PostProcess/PostProcessBloomSetup.cpp
--- a/Engine/Source/Runtime/Renderer/Private/PostProcess/PostProcessBloomSetup.cpp
+++ b/Engine/Source/Runtime/Renderer/Private/PostProcess/PostProcessBloomSetup.cpp
@@ -135,7 +135,7 @@ FScreenPassTexture AddBloomSetupPass(FRDGBuilder& GraphBuilder, const FViewInfo&
 		FComputeShaderUtils::AddPass(
 			GraphBuilder,
 			RDG_EVENT_NAME("BloomSetup %dx%d (CS)", Viewport.Rect.Width(), Viewport.Rect.Height()),
-			*ComputeShader,
+			ComputeShader,
 			PassParameters,
 			FComputeShaderUtils::GetGroupCount(Viewport.Rect.Size(), FIntPoint(GBloomSetupComputeTileSizeX, GBloomSetupComputeTileSizeY)));
 	}
@@ -154,13 +154,13 @@ FScreenPassTexture AddBloomSetupPass(FRDGBuilder& GraphBuilder, const FViewInfo&
 			View,
 			Viewport,
 			Viewport,
-			FScreenPassPipelineState(*VertexShader, *PixelShader),
+			FScreenPassPipelineState(VertexShader, PixelShader),
 			EScreenPassDrawFlags::None,
 			PassParameters,
 			[VertexShader, PixelShader, PassParameters] (FRHICommandListImmediate& RHICmdList)
 		{
-			SetShaderParameters(RHICmdList, *VertexShader, VertexShader->GetVertexShader(), PassParameters->BloomSetup);
-			SetShaderParameters(RHICmdList, *PixelShader, PixelShader->GetPixelShader(), *PassParameters);
+			SetShaderParameters(RHICmdList, VertexShader, VertexShader.GetVertexShader(), PassParameters->BloomSetup);
+			SetShaderParameters(RHICmdList, PixelShader, PixelShader.GetPixelShader(), *PassParameters);
 		});
 	}
 
diff --git a/Engine/Source/Runtime/Renderer/Private/PostProcess/PostProcessBokehDOF.cpp b/Engine/Source/Runtime/Renderer/Private/PostProcess/PostProcessBokehDOF.cpp
--- a/Engine/Source/Runtime/Renderer/Private/PostProcess/PostProcessBokehDOF.cpp
+++ b/Engine/Source/Runtime/Renderer/Private/PostProcess/PostProcessBokehDOF.cpp
@@ -66,7 +66,7 @@ FScreenPassTexture AddVisualizeDOFPass(FRDGBuilder& GraphBuilder, const FViewInf
 
 	RDG_EVENT_SCOPE(GraphBuilder, "VisualizeDOF");
 
-	AddDrawScreenPass(GraphBuilder, RDG_EVENT_NAME("Visualizer"), View, FScreenPassTextureViewport(Output), InputViewport, *PixelShader, PassParameters);
+	AddDrawScreenPass(GraphBuilder, RDG_EVENT_NAME("Visualizer"), View, FScreenPassTextureViewport(Output), InputViewport, PixelShader, PassParameters);
 
 	Output.LoadAction = ERenderTargetLoadAction::ELoad;
 
diff --git a/Engine/Source/Runtime/Renderer/Private/PostProcess/PostProcessBusyWait.cpp b/Engine/Source/Runtime/Renderer/Private/PostProcess/PostProcessBusyWait.cpp
--- a/Engine/Source/Runtime/Renderer/Private/PostProcess/PostProcessBusyWait.cpp
+++ b/Engine/Source/Runtime/Renderer/Private/PostProcess/PostProcessBusyWait.cpp
@@ -36,8 +36,8 @@ class FPostProcessBusyWaitPS : public FGlobalShader
 	FPostProcessBusyWaitPS() {}
 
 public:
-	FPostProcessPassParameters PostprocessParameter;
-	FShaderParameter GPUBusyWait;
+	LAYOUT_FIELD(FPostProcessPassParameters, PostprocessParameter);
+	LAYOUT_FIELD(FShaderParameter, GPUBusyWait);
 
 	/** Initialization constructor. */
 	FPostProcessBusyWaitPS(const ShaderMetaType::CompiledShaderInitializerType& Initializer)
@@ -50,7 +50,7 @@ public:
 	template <typename TRHICmdList>
 	void SetPS(TRHICmdList& RHICmdList, const FRenderingCompositePassContext& Context)
 	{
-		FRHIPixelShader* ShaderRHI = GetPixelShader();
+		FRHIPixelShader* ShaderRHI = RHICmdList.GetBoundPixelShader();
 
 		FGlobalShader::SetParameters<FViewUniformShaderParameters>(RHICmdList, ShaderRHI, Context.View.ViewUniformBuffer);
 
@@ -71,14 +71,6 @@ public:
 		}
 #endif
 	}
-	
-	// FShader interface.
-	virtual bool Serialize(FArchive& Ar) override
-	{
-		bool bShaderHasOutdatedParameters = FGlobalShader::Serialize(Ar);
-		Ar << GPUBusyWait;
-		return bShaderHasOutdatedParameters;
-	}
 };
 
 IMPLEMENT_SHADER_TYPE(,FPostProcessBusyWaitPS,TEXT("/Engine/Private/PostProcessBusyWait.usf"),TEXT("MainPS"),SF_Pixel);
@@ -112,8 +104,8 @@ void FRCPassPostProcessBusyWait::Process(FRenderingCompositePassContext& Context
 		TShaderMapRef<FPostProcessBusyWaitPS> PixelShader(Context.GetShaderMap());
 
 		GraphicsPSOInit.BoundShaderState.VertexDeclarationRHI = GFilterVertexDeclaration.VertexDeclarationRHI;
-		GraphicsPSOInit.BoundShaderState.VertexShaderRHI = GETSAFERHISHADER_VERTEX(*VertexShader);
-		GraphicsPSOInit.BoundShaderState.PixelShaderRHI = GETSAFERHISHADER_PIXEL(*PixelShader);
+		GraphicsPSOInit.BoundShaderState.VertexShaderRHI = VertexShader.GetVertexShader();
+		GraphicsPSOInit.BoundShaderState.PixelShaderRHI = PixelShader.GetPixelShader();
 		GraphicsPSOInit.PrimitiveType = PT_TriangleList;
 
 		SetGraphicsPipelineState(Context.RHICmdList, GraphicsPSOInit);
@@ -129,7 +121,7 @@ void FRCPassPostProcessBusyWait::Process(FRenderingCompositePassContext& Context
 			SrcRect.Width(), SrcRect.Height(),
 			DestRect.Size(),
 			SrcRect.Size(),
-			*VertexShader,
+			VertexShader,
 			EDRF_UseTriangleOptimization);
 	}
 	Context.RHICmdList.EndRenderPass();
diff --git a/Engine/Source/Runtime/Renderer/Private/PostProcess/PostProcessCombineLUTs.cpp b/Engine/Source/Runtime/Renderer/Private/PostProcess/PostProcessCombineLUTs.cpp
--- a/Engine/Source/Runtime/Renderer/Private/PostProcess/PostProcessCombineLUTs.cpp
+++ b/Engine/Source/Runtime/Renderer/Private/PostProcess/PostProcessCombineLUTs.cpp
@@ -449,7 +449,7 @@ FRDGTextureRef AddCombineLUTPass(FRDGBuilder& GraphBuilder, const FViewInfo& Vie
 		FComputeShaderUtils::AddPass(
 			GraphBuilder,
 			RDG_EVENT_NAME("CombineLUTs (CS)"),
-			*ComputeShader,
+			ComputeShader,
 			PassParameters,
 			FIntVector(GroupSizeXY, GroupSizeXY, GroupSizeZ));
 	}
@@ -482,14 +482,14 @@ FRDGTextureRef AddCombineLUTPass(FRDGBuilder& GraphBuilder, const FViewInfo& Vie
 
 				GraphicsPSOInit.PrimitiveType = PT_TriangleStrip;
 				GraphicsPSOInit.BoundShaderState.VertexDeclarationRHI = GScreenVertexDeclaration.VertexDeclarationRHI;
-				GraphicsPSOInit.BoundShaderState.VertexShaderRHI = GETSAFERHISHADER_VERTEX(*VertexShader);
-				GraphicsPSOInit.BoundShaderState.GeometryShaderRHI = GETSAFERHISHADER_GEOMETRY(*GeometryShader);
-				GraphicsPSOInit.BoundShaderState.PixelShaderRHI = GETSAFERHISHADER_PIXEL(*PixelShader);
+				GraphicsPSOInit.BoundShaderState.VertexShaderRHI = VertexShader.GetVertexShader();
+				GraphicsPSOInit.BoundShaderState.GeometryShaderRHI = GeometryShader.GetGeometryShader();
+				GraphicsPSOInit.BoundShaderState.PixelShaderRHI = PixelShader.GetPixelShader();
 				SetGraphicsPipelineState(RHICmdList, GraphicsPSOInit);
 
 				VertexShader->SetParameters(RHICmdList, VolumeBounds, FIntVector(VolumeBounds.MaxX - VolumeBounds.MinX));
 
-				SetShaderParameters(RHICmdList, *PixelShader, PixelShader->GetPixelShader(), *PassParameters);
+				SetShaderParameters(RHICmdList, PixelShader, PixelShader.GetPixelShader(), *PassParameters);
 
 				RasterizeToVolumeTexture(RHICmdList, VolumeBounds);
 			}
@@ -499,11 +499,11 @@ FRDGTextureRef AddCombineLUTPass(FRDGBuilder& GraphBuilder, const FViewInfo& Vie
 
 				GraphicsPSOInit.PrimitiveType = PT_TriangleList;
 				GraphicsPSOInit.BoundShaderState.VertexDeclarationRHI = GFilterVertexDeclaration.VertexDeclarationRHI;
-				GraphicsPSOInit.BoundShaderState.VertexShaderRHI = GETSAFERHISHADER_VERTEX(*VertexShader);
-				GraphicsPSOInit.BoundShaderState.PixelShaderRHI = GETSAFERHISHADER_PIXEL(*PixelShader);
+				GraphicsPSOInit.BoundShaderState.VertexShaderRHI = VertexShader.GetVertexShader();
+				GraphicsPSOInit.BoundShaderState.PixelShaderRHI = PixelShader.GetPixelShader();
 				SetGraphicsPipelineState(RHICmdList, GraphicsPSOInit);
 
-				SetShaderParameters(RHICmdList, *PixelShader, PixelShader->GetPixelShader(), *PassParameters);
+				SetShaderParameters(RHICmdList, PixelShader, PixelShader.GetPixelShader(), *PassParameters);
 
 				DrawRectangle(
 					RHICmdList,
@@ -513,7 +513,7 @@ FRDGTextureRef AddCombineLUTPass(FRDGBuilder& GraphBuilder, const FViewInfo& Vie
 					GLUTSize * GLUTSize, GLUTSize,				// SizeUV
 					FIntPoint(GLUTSize * GLUTSize, GLUTSize),	// TargetSize
 					FIntPoint(GLUTSize * GLUTSize, GLUTSize),	// TextureSize
-					*VertexShader,
+					VertexShader,
 					EDRF_UseTriangleOptimization);
 			}
 		});
diff --git a/Engine/Source/Runtime/Renderer/Private/PostProcess/PostProcessCompositeEditorPrimitives.cpp b/Engine/Source/Runtime/Renderer/Private/PostProcess/PostProcessCompositeEditorPrimitives.cpp
--- a/Engine/Source/Runtime/Renderer/Private/PostProcess/PostProcessCompositeEditorPrimitives.cpp
+++ b/Engine/Source/Runtime/Renderer/Private/PostProcess/PostProcessCompositeEditorPrimitives.cpp
@@ -320,11 +320,11 @@ FScreenPassTexture AddEditorPrimitivePass(
 					View,
 					EditorPrimitivesViewport,
 					SceneDepthViewport,
-					FScreenPassPipelineState(*PopulateDepthVertexShader, *PopulateDepthPixelShader, FScreenPassPipelineState::FDefaultBlendState::GetRHI(), TStaticDepthStencilState<true, CF_Always>::GetRHI()),
+					FScreenPassPipelineState(PopulateDepthVertexShader, PopulateDepthPixelShader, FScreenPassPipelineState::FDefaultBlendState::GetRHI(), TStaticDepthStencilState<true, CF_Always>::GetRHI()),
 					EScreenPassDrawFlags::None,
 					[&](FRHICommandList&)
 				{
-					SetShaderParameters(RHICmdList, *PopulateDepthPixelShader, PopulateDepthPixelShader->GetPixelShader(), *PassParameters);
+					SetShaderParameters(RHICmdList, PopulateDepthPixelShader, PopulateDepthPixelShader.GetPixelShader(), *PassParameters);
 				});
 			}
 			else
@@ -413,7 +413,7 @@ FScreenPassTexture AddEditorPrimitivePass(
 		View,
 		OutputViewport,
 		ColorViewport,
-		*PixelShader,
+		PixelShader,
 		PassParameters);
 
 	return MoveTemp(Output);
diff --git a/Engine/Source/Runtime/Renderer/Private/PostProcess/PostProcessDOF.cpp b/Engine/Source/Runtime/Renderer/Private/PostProcess/PostProcessDOF.cpp
--- a/Engine/Source/Runtime/Renderer/Private/PostProcess/PostProcessDOF.cpp
+++ b/Engine/Source/Runtime/Renderer/Private/PostProcess/PostProcessDOF.cpp
@@ -61,9 +61,9 @@ class FPostProcessDOFSetupPS : public FGlobalShader
 	FPostProcessDOFSetupPS() {}
 
 public:
-	FPostProcessPassParameters PostprocessParameter;
-	FSceneTextureShaderParameters SceneTextureParameters;
-	FShaderParameter DepthOfFieldParams;
+	LAYOUT_FIELD(FPostProcessPassParameters, PostprocessParameter);
+	LAYOUT_FIELD(FSceneTextureShaderParameters, SceneTextureParameters);
+	LAYOUT_FIELD(FShaderParameter, DepthOfFieldParams);
 	
 	/** Initialization constructor. */
 	FPostProcessDOFSetupPS(const ShaderMetaType::CompiledShaderInitializerType& Initializer)
@@ -74,18 +74,10 @@ public:
 		DepthOfFieldParams.Bind(Initializer.ParameterMap,TEXT("DepthOfFieldParams"));
 	}
 
-	// FShader interface.
-	virtual bool Serialize(FArchive& Ar) override
-	{
-		bool bShaderHasOutdatedParameters = FGlobalShader::Serialize(Ar);
-		Ar << PostprocessParameter << SceneTextureParameters << DepthOfFieldParams;
-		return bShaderHasOutdatedParameters;
-	}
-
 	template <typename TRHICmdList>
 	void SetParameters(TRHICmdList& RHICmdList, const FRenderingCompositePassContext& Context)
 	{
-		FRHIPixelShader* ShaderRHI = GetPixelShader();
+		FRHIPixelShader* ShaderRHI = RHICmdList.GetBoundPixelShader();
 
 		FGlobalShader::SetParameters<FViewUniformShaderParameters>(RHICmdList, ShaderRHI, Context.View.ViewUniformBuffer);
 
@@ -131,7 +123,7 @@ public:
 // @param FarBlur 0:off, 1:on
 // @param NearBlur 0:off, 1:on, 2:on with Vignette
 template <uint32 FarBlur, uint32 NearBlur>
-static FShader* SetDOFShaderTempl(const FRenderingCompositePassContext& Context)
+static TShaderRef<FShader> SetDOFShaderTempl(const FRenderingCompositePassContext& Context)
 {
 	FGraphicsPipelineStateInitializer GraphicsPSOInit;
 	Context.RHICmdList.ApplyCachedRenderTargets(GraphicsPSOInit);
@@ -146,14 +138,14 @@ static FShader* SetDOFShaderTempl(const FRenderingCompositePassContext& Context)
 	TShaderMapRef<FPostProcessDOFSetupPS<FarBlur, NearBlur> > PixelShader(Context.GetShaderMap());
 
 	GraphicsPSOInit.BoundShaderState.VertexDeclarationRHI = GFilterVertexDeclaration.VertexDeclarationRHI;
-	GraphicsPSOInit.BoundShaderState.VertexShaderRHI = GETSAFERHISHADER_VERTEX(*VertexShader);
-	GraphicsPSOInit.BoundShaderState.PixelShaderRHI = GETSAFERHISHADER_PIXEL(*PixelShader);
+	GraphicsPSOInit.BoundShaderState.VertexShaderRHI = VertexShader.GetVertexShader();
+	GraphicsPSOInit.BoundShaderState.PixelShaderRHI = PixelShader.GetPixelShader();
 	SetGraphicsPipelineState(Context.RHICmdList, GraphicsPSOInit);
 
 	VertexShader->SetParameters(Context);
 	PixelShader->SetParameters(Context.RHICmdList, Context);
 
-	return *VertexShader;
+	return VertexShader;
 }
 
 
@@ -235,8 +227,7 @@ void FRCPassPostProcessDOFSetup::Process(FRenderingCompositePassContext& Context
 			}
 		}
 
-		FShader* VertexShader = 0;
-
+		TShaderRef<FShader> VertexShader;
 		if (bFarBlur)
 		{
 			switch (NearBlur)
@@ -336,9 +327,9 @@ class FPostProcessDOFRecombinePS : public FGlobalShader
 	FPostProcessDOFRecombinePS() {}
 
 public:
-	FPostProcessPassParameters PostprocessParameter;
-	FSceneTextureShaderParameters SceneTextureParameters;
-	FShaderParameter DepthOfFieldUVLimit;
+	LAYOUT_FIELD(FPostProcessPassParameters, PostprocessParameter);
+	LAYOUT_FIELD(FSceneTextureShaderParameters, SceneTextureParameters);
+	LAYOUT_FIELD(FShaderParameter, DepthOfFieldUVLimit);
 
 	/** Initialization constructor. */
 	FPostProcessDOFRecombinePS(const ShaderMetaType::CompiledShaderInitializerType& Initializer)
@@ -349,18 +340,10 @@ public:
 		DepthOfFieldUVLimit.Bind(Initializer.ParameterMap,TEXT("DepthOfFieldUVLimit"));
 	}
 
-	// FShader interface.
-	virtual bool Serialize(FArchive& Ar) override
-	{
-		bool bShaderHasOutdatedParameters = FGlobalShader::Serialize(Ar);
-		Ar << PostprocessParameter << SceneTextureParameters << DepthOfFieldUVLimit;
-		return bShaderHasOutdatedParameters;
-	}
-
 	void SetParameters(const FRenderingCompositePassContext& Context)
 	{
 		FSceneRenderTargets& SceneContext = FSceneRenderTargets::Get(Context.RHICmdList);
-		FRHIPixelShader* ShaderRHI = GetPixelShader();
+		FRHIPixelShader* ShaderRHI = Context.RHICmdList.GetBoundPixelShader();
 
 		FGlobalShader::SetParameters<FViewUniformShaderParameters>(Context.RHICmdList, ShaderRHI, Context.View.ViewUniformBuffer);
 
@@ -400,7 +383,7 @@ public:
 // @param FarBlur 0:off, 1:on
 // @param NearBlur 0:off, 1:on, 2:on with Vignette
 template <uint32 FarBlur, uint32 NearBlur, uint32 SeparateTranslucency>
-static FShader* SetDOFRecombineShaderTempl(const FRenderingCompositePassContext& Context)
+static TShaderRef<FShader> SetDOFRecombineShaderTempl(const FRenderingCompositePassContext& Context)
 {
 	FGraphicsPipelineStateInitializer GraphicsPSOInit;
 	Context.RHICmdList.ApplyCachedRenderTargets(GraphicsPSOInit);
@@ -413,18 +396,18 @@ static FShader* SetDOFRecombineShaderTempl(const FRenderingCompositePassContext&
 	TShaderMapRef<FPostProcessDOFRecombinePS<FarBlur, NearBlur, SeparateTranslucency> > PixelShader(Context.GetShaderMap());
 
 	GraphicsPSOInit.BoundShaderState.VertexDeclarationRHI = GFilterVertexDeclaration.VertexDeclarationRHI;
-	GraphicsPSOInit.BoundShaderState.VertexShaderRHI = GETSAFERHISHADER_VERTEX(*VertexShader);
-	GraphicsPSOInit.BoundShaderState.PixelShaderRHI = GETSAFERHISHADER_PIXEL(*PixelShader);
+	GraphicsPSOInit.BoundShaderState.VertexShaderRHI = VertexShader.GetVertexShader();
+	GraphicsPSOInit.BoundShaderState.PixelShaderRHI = PixelShader.GetPixelShader();
 	SetGraphicsPipelineState(Context.RHICmdList, GraphicsPSOInit);
 
 	VertexShader->SetParameters(Context);
 	PixelShader->SetParameters(Context);
 
-	return *VertexShader;
+	return VertexShader;
 }
 
 template <uint32 FarBlur, uint32 NearBlur>
-static FShader* SetDOFRecombineShaderTempl(const FRenderingCompositePassContext& Context, bool bSeparateTranslucency)
+static TShaderRef<FShader> SetDOFRecombineShaderTempl(const FRenderingCompositePassContext& Context, bool bSeparateTranslucency)
 {
 	if (bSeparateTranslucency)
 	{
@@ -491,8 +474,7 @@ void FRCPassPostProcessDOFRecombine::Process(FRenderingCompositePassContext& Con
 		bool bNearBlur = GetInputDesc(ePId_Input2) != 0;
 		bool bSeparateTranslucency = GetInputDesc(ePId_Input3) != 0;
 
-		FShader* VertexShader = 0;
-
+		TShaderRef<FShader> VertexShader;
 		if (bFarBlur)
 		{
 			if (bNearBlur)
diff --git a/Engine/Source/Runtime/Renderer/Private/PostProcess/PostProcessDownsample.cpp b/Engine/Source/Runtime/Renderer/Private/PostProcess/PostProcessDownsample.cpp
--- a/Engine/Source/Runtime/Renderer/Private/PostProcess/PostProcessDownsample.cpp
+++ b/Engine/Source/Runtime/Renderer/Private/PostProcess/PostProcessDownsample.cpp
@@ -158,7 +158,7 @@ FScreenPassTexture AddDownsamplePass(
 		FComputeShaderUtils::AddPass(
 			GraphBuilder,
 			RDG_EVENT_NAME("Downsample.%s %dx%d (CS)", Inputs.Name, Inputs.SceneColor.ViewRect.Width(), Inputs.SceneColor.ViewRect.Height()),
-			*ComputeShader,
+			ComputeShader,
 			PassParameters,
 			FComputeShaderUtils::GetGroupCount(OutputViewport.Rect.Size(), FIntPoint(GDownsampleTileSizeX, GDownsampleTileSizeY)));
 	}
@@ -174,7 +174,7 @@ FScreenPassTexture AddDownsamplePass(
 			GraphBuilder,
 			View.ShaderMap,
 			RDG_EVENT_NAME("Downsample.%s %dx%d (PS)", Inputs.Name, Inputs.SceneColor.ViewRect.Width(), Inputs.SceneColor.ViewRect.Height()),
-			*PixelShader,
+			PixelShader,
 			PassParameters,
 			OutputViewport.Rect);
 	}
diff --git a/Engine/Source/Runtime/Renderer/Private/PostProcess/PostProcessEyeAdaptation.cpp b/Engine/Source/Runtime/Renderer/Private/PostProcess/PostProcessEyeAdaptation.cpp
--- a/Engine/Source/Runtime/Renderer/Private/PostProcess/PostProcessEyeAdaptation.cpp
+++ b/Engine/Source/Runtime/Renderer/Private/PostProcess/PostProcessEyeAdaptation.cpp
@@ -453,7 +453,7 @@ FRDGTextureRef AddHistogramEyeAdaptationPass(
 		FComputeShaderUtils::AddPass(
 			GraphBuilder,
 			RDG_EVENT_NAME("HistogramEyeAdaptation (CS)"),
-			*ComputeShader,
+			ComputeShader,
 			PassParameters,
 			FIntVector(1, 1, 1));
 	}
@@ -471,7 +471,7 @@ FRDGTextureRef AddHistogramEyeAdaptationPass(
 			View,
 			FScreenPassTextureViewport(OutputTexture),
 			FScreenPassTextureViewport(HistogramTexture),
-			*PixelShader,
+			PixelShader,
 			PassParameters);
 	}
 
@@ -498,7 +498,7 @@ public:
 
 	static bool ShouldCompilePermutation(const FGlobalShaderPermutationParameters& Parameters)
 	{
-		return IsFeatureLevelSupported(Parameters.Platform, BasicEyeAdaptationMinFeatureLevel);
+		return IsFeatureLevelSupported(Parameters.Platform, ERHIFeatureLevel::Type(BasicEyeAdaptationMinFeatureLevel));
 	}
 };
 
@@ -537,7 +537,7 @@ FScreenPassTexture AddBasicEyeAdaptationSetupPass(
 		View,
 		Viewport,
 		Viewport,
-		*PixelShader,
+		PixelShader,
 		PassParameters,
 		EScreenPassDrawFlags::AllowHMDHiddenAreaMask);
 
@@ -565,7 +565,7 @@ public:
 
 	static bool ShouldCompilePermutation(const FGlobalShaderPermutationParameters& Parameters)
 	{
-		return IsFeatureLevelSupported(Parameters.Platform, BasicEyeAdaptationMinFeatureLevel);
+		return IsFeatureLevelSupported(Parameters.Platform, ERHIFeatureLevel::Type(BasicEyeAdaptationMinFeatureLevel));
 	}
 
 	FBasicEyeAdaptationShader() = default;
@@ -646,7 +646,7 @@ FRDGTextureRef AddBasicEyeAdaptationPass(
 		FComputeShaderUtils::AddPass(
 			GraphBuilder,
 			RDG_EVENT_NAME("BasicEyeAdaptation (CS)"),
-			*ComputeShader,
+			ComputeShader,
 			PassParameters,
 			FIntVector(1, 1, 1));
 	}
@@ -666,7 +666,7 @@ FRDGTextureRef AddBasicEyeAdaptationPass(
 			View,
 			OutputViewport,
 			OutputViewport,
-			*PixelShader,
+			PixelShader,
 			PassParameters);
 	}
 
diff --git a/Engine/Source/Runtime/Renderer/Private/PostProcess/PostProcessFFTBloom.cpp b/Engine/Source/Runtime/Renderer/Private/PostProcess/PostProcessFFTBloom.cpp
--- a/Engine/Source/Runtime/Renderer/Private/PostProcess/PostProcessFFTBloom.cpp
+++ b/Engine/Source/Runtime/Renderer/Private/PostProcess/PostProcessFFTBloom.cpp
@@ -263,7 +263,7 @@ void ResizeAndCenterTexture(
 		// Use multiple threads per scan line to insure memory coalescing during the write
 		const int32 ThreadsPerGroup = FResizeAndCenterTextureCS::ThreadsPerGroup;
 		const int32 ThreadsGroupsPerScanLine = (DstBufferSize.X % ThreadsPerGroup == 0) ? DstBufferSize.X / ThreadsPerGroup : DstBufferSize.X / ThreadsPerGroup + 1;
-		FComputeShaderUtils::Dispatch(RHICmdList, *ComputeShader, PassParameters, FIntVector(ThreadsGroupsPerScanLine, DstBufferSize.Y, 1));
+		FComputeShaderUtils::Dispatch(RHICmdList, ComputeShader, PassParameters, FIntVector(ThreadsGroupsPerScanLine, DstBufferSize.Y, 1));
 	}
 }
 
@@ -308,7 +308,7 @@ void CaptureKernelWeight(
 		PassParameters.HalfResSumLocation = HalfResSumLocation;
 		PassParameters.UVCenter = CenterUV;
 
-		FComputeShaderUtils::Dispatch(RHICmdList, *ComputeShader, PassParameters, FIntVector(1, 1, 1));
+		FComputeShaderUtils::Dispatch(RHICmdList, ComputeShader, PassParameters, FIntVector(1, 1, 1));
 	}
 
 	RHICmdList.TransitionResource(EResourceTransitionAccess::EReadable, EResourceTransitionPipeline::EComputeToCompute, DstTargetItem.UAV);
@@ -370,7 +370,7 @@ void BlendLowRes(
 		const int32 ThreadsPerGroup = ComputeShader->ThreadsPerGroup;
 		const int32 ThreadsGroupsPerScanLine = (TargetExtent.X % ThreadsPerGroup == 0) ? TargetExtent.X / ThreadsPerGroup : TargetExtent.X / ThreadsPerGroup + 1;
 
-		FComputeShaderUtils::Dispatch(RHICmdList, *ComputeShader, PassParameters, FIntVector(ThreadsGroupsPerScanLine, TargetExtent.Y, 1));
+		FComputeShaderUtils::Dispatch(RHICmdList, ComputeShader, PassParameters, FIntVector(ThreadsGroupsPerScanLine, TargetExtent.Y, 1));
 	}
 
 	RHICmdList.TransitionResource(EResourceTransitionAccess::EReadable, EResourceTransitionPipeline::EComputeToGfx, DstUAV);
@@ -414,7 +414,7 @@ void CopyImageRect(
 		const int32 ThreadsPerGroup = ComputeShader->ThreadsPerGroup;
 		const int32 ThreadsGroupsPerScanLine = (DstRectSize.X % ThreadsPerGroup == 0) ? DstRectSize.X / ThreadsPerGroup : DstRectSize.X / ThreadsPerGroup + 1;
 
-		FComputeShaderUtils::Dispatch(RHICmdList, *ComputeShader, PassParameters, FIntVector(ThreadsGroupsPerScanLine, DstRectSize.Y, 1));
+		FComputeShaderUtils::Dispatch(RHICmdList, ComputeShader, PassParameters, FIntVector(ThreadsGroupsPerScanLine, DstRectSize.Y, 1));
 	}
 }
 
diff --git a/Engine/Source/Runtime/Renderer/Private/PostProcess/PostProcessGBufferHints.cpp b/Engine/Source/Runtime/Renderer/Private/PostProcess/PostProcessGBufferHints.cpp
--- a/Engine/Source/Runtime/Renderer/Private/PostProcess/PostProcessGBufferHints.cpp
+++ b/Engine/Source/Runtime/Renderer/Private/PostProcess/PostProcessGBufferHints.cpp
@@ -64,7 +64,7 @@ FScreenPassTexture AddVisualizeGBufferHintsPass(FRDGBuilder& GraphBuilder, const
 
 	RDG_EVENT_SCOPE(GraphBuilder, "VisualizeGBufferHints");
 
-	AddDrawScreenPass(GraphBuilder, RDG_EVENT_NAME("Visualizer"), View, OutputViewport, InputViewport, *PixelShader, PassParameters);
+	AddDrawScreenPass(GraphBuilder, RDG_EVENT_NAME("Visualizer"), View, OutputViewport, InputViewport, PixelShader, PassParameters);
 
 	Output.LoadAction = ERenderTargetLoadAction::ELoad;
 
diff --git a/Engine/Source/Runtime/Renderer/Private/PostProcess/PostProcessHMD.cpp b/Engine/Source/Runtime/Renderer/Private/PostProcess/PostProcessHMD.cpp
--- a/Engine/Source/Runtime/Renderer/Private/PostProcess/PostProcessHMD.cpp
+++ b/Engine/Source/Runtime/Renderer/Private/PostProcess/PostProcessHMD.cpp
@@ -117,12 +117,12 @@ FScreenPassTexture AddDefaultHMDDistortionPass(FRDGBuilder& GraphBuilder, const
 	{
 		RHICmdList.SetViewport(OutputViewRect.Min.X, OutputViewRect.Min.Y, 0.0f, OutputViewRect.Max.X, OutputViewRect.Max.Y, 1.0f);
 
-		FScreenPassPipelineState PipelineState(*VertexShader, *PixelShader);
+		FScreenPassPipelineState PipelineState(VertexShader,PixelShader);
 		PipelineState.VertexDeclaration = GDistortionVertexDeclaration.VertexDeclarationRHI;
 		SetScreenPassPipelineState(RHICmdList, PipelineState);
 
-		SetShaderParameters(RHICmdList, *VertexShader, VertexShader->GetVertexShader(), *PassParameters);
-		SetShaderParameters(RHICmdList, *PixelShader, PixelShader->GetPixelShader(), *PassParameters);
+		SetShaderParameters(RHICmdList, VertexShader, VertexShader.GetVertexShader(), *PassParameters);
+		SetShaderParameters(RHICmdList, PixelShader, PixelShader.GetPixelShader(), *PassParameters);
 
 		FRenderingCompositePassContext PassContext(RHICmdList, View);
 		HMDDevice->DrawDistortionMesh_RenderThread(PassContext, PassParameters->InputTexture->Desc.Extent);
diff --git a/Engine/Source/Runtime/Renderer/Private/PostProcess/PostProcessHistogram.cpp b/Engine/Source/Runtime/Renderer/Private/PostProcess/PostProcessHistogram.cpp
--- a/Engine/Source/Runtime/Renderer/Private/PostProcess/PostProcessHistogram.cpp
+++ b/Engine/Source/Runtime/Renderer/Private/PostProcess/PostProcessHistogram.cpp
@@ -144,7 +144,7 @@ FRDGTextureRef AddHistogramPass(
 		FComputeShaderUtils::AddPass(
 			GraphBuilder,
 			RDG_EVENT_NAME("Histogram %dx%d (CS)", SceneColor.ViewRect.Width(), SceneColor.ViewRect.Height()),
-			*ComputeShader,
+			ComputeShader,
 			PassParameters,
 			FIntVector(HistogramThreadGroupCount.X, HistogramThreadGroupCount.Y, 1));
 	}
@@ -184,7 +184,7 @@ FRDGTextureRef AddHistogramPass(
 			View,
 			OutputViewport,
 			InputViewport,
-			*PixelShader,
+			PixelShader,
 			PassParameters);
 	}
 
diff --git a/Engine/Source/Runtime/Renderer/Private/PostProcess/PostProcessLensFlares.cpp b/Engine/Source/Runtime/Renderer/Private/PostProcess/PostProcessLensFlares.cpp
--- a/Engine/Source/Runtime/Renderer/Private/PostProcess/PostProcessLensFlares.cpp
+++ b/Engine/Source/Runtime/Renderer/Private/PostProcess/PostProcessLensFlares.cpp
@@ -206,13 +206,13 @@ FScreenPassTexture AddLensFlaresPass(
 			GraphicsPSOInit.RasterizerState = TStaticRasterizerState<>::GetRHI();
 			GraphicsPSOInit.DepthStencilState = TStaticDepthStencilState<false, CF_Always>::GetRHI();
 			GraphicsPSOInit.BoundShaderState.VertexDeclarationRHI = GEmptyVertexDeclaration.VertexDeclarationRHI;
-			GraphicsPSOInit.BoundShaderState.VertexShaderRHI = GETSAFERHISHADER_VERTEX(*VertexShader);
-			GraphicsPSOInit.BoundShaderState.PixelShaderRHI = GETSAFERHISHADER_PIXEL(*PixelShader);
+			GraphicsPSOInit.BoundShaderState.VertexShaderRHI = VertexShader.GetVertexShader();
+			GraphicsPSOInit.BoundShaderState.PixelShaderRHI = PixelShader.GetPixelShader();
 			GraphicsPSOInit.PrimitiveType = PT_TriangleList;
 			SetGraphicsPipelineState(RHICmdList, GraphicsPSOInit);
 
-			SetShaderParameters(RHICmdList, *VertexShader, VertexShader->GetVertexShader(), VertexParameters);
-			SetShaderParameters(RHICmdList, *PixelShader, PixelShader->GetPixelShader(), PixelParameters);
+			SetShaderParameters(RHICmdList, VertexShader, VertexShader.GetVertexShader(), VertexParameters);
+			SetShaderParameters(RHICmdList, PixelShader, PixelShader.GetPixelShader(), PixelParameters);
 
 			// Emit an instanced quad draw call on the order of the number of pixels on the screen.
 			RHICmdList.SetStreamSource(0, nullptr, 0);
@@ -283,7 +283,7 @@ FScreenPassTexture AddLensFlaresPass(
 		TShaderMapRef<FScreenPassVS> VertexShader(View.ShaderMap);
 		TShaderMapRef<FLensFlareCompositePS> PixelShader(View.ShaderMap);
 
-		const FScreenPassPipelineState PipelineState(*VertexShader, *PixelShader, AdditiveBlendState);
+		const FScreenPassPipelineState PipelineState(VertexShader, PixelShader, AdditiveBlendState);
 
 		// This pass rasterizes the lens flare quad scaled and centered within the viewport.
 		GraphBuilder.AddPass(
@@ -296,7 +296,7 @@ FScreenPassTexture AddLensFlaresPass(
 
 			SetScreenPassPipelineState(RHICmdList, PipelineState);
 
-			SetShaderParameters(RHICmdList, *PixelShader, PixelShader->GetPixelShader(), *PassParameters);
+			SetShaderParameters(RHICmdList, PixelShader, PixelShader.GetPixelShader(), *PassParameters);
 
 			DrawRectangle(
 				RHICmdList,
diff --git a/Engine/Source/Runtime/Renderer/Private/PostProcess/PostProcessMaterial.cpp b/Engine/Source/Runtime/Renderer/Private/PostProcess/PostProcessMaterial.cpp
--- a/Engine/Source/Runtime/Renderer/Private/PostProcess/PostProcessMaterial.cpp
+++ b/Engine/Source/Runtime/Renderer/Private/PostProcess/PostProcessMaterial.cpp
@@ -190,7 +190,7 @@ public:
 
 	static bool ShouldCompilePermutation(const FMaterialShaderPermutationParameters& Parameters)
 	{
-		if (Parameters.Material->GetMaterialDomain() == MD_PostProcess)
+		if (Parameters.MaterialParameters.MaterialDomain == MD_PostProcess)
 		{
 			const FPermutationDomain PermutationVector(Parameters.PermutationId);
 
@@ -211,33 +211,34 @@ public:
 		FMaterialShader::ModifyCompilationEnvironment(Parameters, OutEnvironment);
 		OutEnvironment.SetDefine(TEXT("POST_PROCESS_MATERIAL"), 1);
 
-		EBlendableLocation Location = EBlendableLocation(Parameters.Material->GetBlendableLocation());
+		EBlendableLocation Location = EBlendableLocation(Parameters.MaterialParameters.BlendableLocation);
 		OutEnvironment.SetDefine(TEXT("POST_PROCESS_MATERIAL_BEFORE_TONEMAP"), (Location == BL_AfterTonemapping || Location == BL_ReplacingTonemapper) ? 0 : 1);
 
 		FPermutationDomain PermutationVector(Parameters.PermutationId);
 		if (PermutationVector.Get<FMobileDimension>())
 		{
-			OutEnvironment.SetDefine(TEXT("POST_PROCESS_MATERIAL_BEFORE_TONEMAP"), (Parameters.Material->GetBlendableLocation() != BL_AfterTonemapping) ? 1 : 0);
+			OutEnvironment.SetDefine(TEXT("POST_PROCESS_MATERIAL_BEFORE_TONEMAP"), (Parameters.MaterialParameters.BlendableLocation != BL_AfterTonemapping) ? 1 : 0);
 		}
 	}
 
 protected:
 	template <typename TRHIShader>
-	void SetParameters(FRHICommandList& RHICmdList, TRHIShader* ShaderRHI, const FViewInfo& View, const FMaterialRenderProxy* Proxy, const FParameters& Parameters)
+	static void SetParameters(FRHICommandList& RHICmdList, const TShaderRef<FPostProcessMaterialShader> & Shader, TRHIShader* ShaderRHI, const FViewInfo& View, const FMaterialRenderProxy* Proxy, const FParameters& Parameters)
 	{
-		FMaterialShader::SetParameters(RHICmdList, ShaderRHI, Proxy, *Proxy->GetMaterial(View.GetFeatureLevel()), View, View.ViewUniformBuffer, ESceneTextureSetupMode::All);
-		SetShaderParameters(RHICmdList, this, ShaderRHI, Parameters);
+		FMaterialShader* MaterialShader = Shader.GetShader();
+		MaterialShader->SetParameters(RHICmdList, ShaderRHI, Proxy, *Proxy->GetMaterial(View.GetFeatureLevel()), View, View.ViewUniformBuffer, ESceneTextureSetupMode::All);
+		SetShaderParameters(RHICmdList, Shader, ShaderRHI, Parameters);
 	}
 };
 
 class FPostProcessMaterialVS : public FPostProcessMaterialShader
 {
 public:
-	DECLARE_MATERIAL_SHADER(FPostProcessMaterialVS);
+	DECLARE_SHADER_TYPE(FPostProcessMaterialVS, Material);
 
-	void SetParameters(FRHICommandList& RHICmdList, const FViewInfo& View, const FMaterialRenderProxy* Proxy, const FParameters& Parameters)
+	static void SetParameters(FRHICommandList& RHICmdList, const TShaderRef<FPostProcessMaterialVS>& Shader, const FViewInfo& View, const FMaterialRenderProxy* Proxy, const FParameters& Parameters)
 	{
-		FPostProcessMaterialShader::SetParameters(RHICmdList, GetVertexShader(), View, Proxy, Parameters);
+		FPostProcessMaterialShader::SetParameters(RHICmdList, Shader, Shader.GetVertexShader(), View, Proxy, Parameters);
 	}
 
 	FPostProcessMaterialVS() = default;
@@ -249,11 +250,11 @@ public:
 class FPostProcessMaterialPS : public FPostProcessMaterialShader
 {
 public:
-	DECLARE_MATERIAL_SHADER(FPostProcessMaterialPS);
+	DECLARE_SHADER_TYPE(FPostProcessMaterialPS, Material);
 
-	void SetParameters(FRHICommandList& RHICmdList, const FViewInfo& View, const FMaterialRenderProxy* Proxy, const FParameters& Parameters)
+	static void SetParameters(FRHICommandList& RHICmdList, const TShaderRef<FPostProcessMaterialPS>& Shader, const FViewInfo& View, const FMaterialRenderProxy* Proxy, const FParameters& Parameters)
 	{
-		FPostProcessMaterialShader::SetParameters(RHICmdList, GetPixelShader(), View, Proxy, Parameters);
+		FPostProcessMaterialShader::SetParameters(RHICmdList, Shader, Shader.GetPixelShader(), View, Proxy, Parameters);
 	}
 
 	FPostProcessMaterialPS() = default;
@@ -262,8 +263,8 @@ public:
 	{}
 };
 
-IMPLEMENT_MATERIAL_SHADER(FPostProcessMaterialVS, "/Engine/Private/PostProcessMaterialShaders.usf", "MainVS", SF_Vertex);
-IMPLEMENT_MATERIAL_SHADER(FPostProcessMaterialPS, "/Engine/Private/PostProcessMaterialShaders.usf", "MainPS", SF_Pixel);
+IMPLEMENT_SHADER_TYPE(,FPostProcessMaterialVS, TEXT("/Engine/Private/PostProcessMaterialShaders.usf"), TEXT("MainVS"), SF_Vertex);
+IMPLEMENT_SHADER_TYPE(,FPostProcessMaterialPS, TEXT("/Engine/Private/PostProcessMaterialShaders.usf"), TEXT("MainPS"), SF_Pixel);
 
 class FPostProcessMaterialVertexDeclaration : public FRenderResource
 {
@@ -320,7 +321,7 @@ static void AddCopyAndFlipTexturePass(FRDGBuilder& GraphBuilder, const FViewInfo
 	Parameters->InputSampler = TStaticSamplerState<>::GetRHI();
 	Parameters->RenderTargets[0] = FRenderTargetBinding(DestTexture, ERenderTargetLoadAction::ENoAction);
 
-	AddDrawScreenPass(GraphBuilder, RDG_EVENT_NAME("DrawTexture"), View, OutputViewport, InputViewport, *PixelShader, Parameters, EScreenPassDrawFlags::FlipYAxis);
+	AddDrawScreenPass(GraphBuilder, RDG_EVENT_NAME("DrawTexture"), View, OutputViewport, InputViewport, PixelShader, Parameters, EScreenPassDrawFlags::FlipYAxis);
 }
 
 FScreenPassTexture AddPostProcessMaterialPass(
@@ -465,8 +466,8 @@ FScreenPassTexture AddPostProcessMaterialPass(
 	FPostProcessMaterialShader::FPermutationDomain PermutationVector;
 	PermutationVector.Set<FPostProcessMaterialShader::FMobileDimension>(bIsMobile);
 
-	FPostProcessMaterialVS* VertexShader = MaterialShaderMap->GetShader<FPostProcessMaterialVS>(PermutationVector);
-	FPostProcessMaterialPS* PixelShader = MaterialShaderMap->GetShader<FPostProcessMaterialPS>(PermutationVector);
+	TShaderRef<FPostProcessMaterialVS> VertexShader = MaterialShaderMap->GetShader<FPostProcessMaterialVS>(PermutationVector);
+	TShaderRef<FPostProcessMaterialPS> PixelShader = MaterialShaderMap->GetShader<FPostProcessMaterialPS>(PermutationVector);
 	ClearUnusedGraphResources(VertexShader, PixelShader, PostProcessMaterialParameters);
 
 	const uint32 MaterialStencilRef = Material->GetStencilRefValue();
@@ -500,8 +501,8 @@ FScreenPassTexture AddPostProcessMaterialPass(
 			ScreenPassFlags,
 			[&](FRHICommandListImmediate&)
 		{
-			VertexShader->SetParameters(RHICmdList, View, MaterialRenderProxy, *PostProcessMaterialParameters);
-			PixelShader->SetParameters(RHICmdList, View, MaterialRenderProxy, *PostProcessMaterialParameters);
+			FPostProcessMaterialVS::SetParameters(RHICmdList, VertexShader, View, MaterialRenderProxy, *PostProcessMaterialParameters);
+			FPostProcessMaterialPS::SetParameters(RHICmdList, PixelShader, View, MaterialRenderProxy, *PostProcessMaterialParameters);
 			RHICmdList.SetStencilRef(MaterialStencilRef);
 		});
 
diff --git a/Engine/Source/Runtime/Renderer/Private/PostProcess/PostProcessMitchellNetravali.cpp b/Engine/Source/Runtime/Renderer/Private/PostProcess/PostProcessMitchellNetravali.cpp
--- a/Engine/Source/Runtime/Renderer/Private/PostProcess/PostProcessMitchellNetravali.cpp
+++ b/Engine/Source/Runtime/Renderer/Private/PostProcess/PostProcessMitchellNetravali.cpp
@@ -70,7 +70,7 @@ FRDGTextureRef ComputeMitchellNetravaliDownsample(
 	FComputeShaderUtils::AddPass(
 		GraphBuilder,
 		RDG_EVENT_NAME("MitchellNetravaliDownsample %dx%d -> %dx%d", Input.ViewRect.Width(), Input.ViewRect.Height(), OutputViewport.Rect.Width(), OutputViewport.Rect.Height()),
-		*ComputeShader,
+		ComputeShader,
 		PassParameters,
 		FComputeShaderUtils::GetGroupCount(OutputViewport.Rect.Size(), FComputeShaderUtils::kGolden2DGroupSize));
 
