Copying Private-LoadTimes-4.24 stream to Main. Biggest changes are in Materials/Shader memory freezing.

diff --git a/Engine/Source/Runtime/Renderer/Private/PostProcess/PostProcessMobile.cpp b/Engine/Source/Runtime/Renderer/Private/PostProcess/PostProcessMobile.cpp
--- a/Engine/Source/Runtime/Renderer/Private/PostProcess/PostProcessMobile.cpp
+++ b/Engine/Source/Runtime/Renderer/Private/PostProcess/PostProcessMobile.cpp
@@ -55,7 +55,7 @@ class FPostProcessBloomSetupVS_ES2 : public FGlobalShader
 	FPostProcessBloomSetupVS_ES2() {}
 
 public:
-	FPostProcessPassParameters PostprocessParameter;
+	LAYOUT_FIELD(FPostProcessPassParameters, PostprocessParameter);
 
 	/** Initialization constructor. */
 	FPostProcessBloomSetupVS_ES2(const ShaderMetaType::CompiledShaderInitializerType& Initializer)
@@ -66,20 +66,12 @@ public:
 
 	void SetVS(const FRenderingCompositePassContext& Context)
 	{
-		FRHIVertexShader* ShaderRHI = GetVertexShader();
+		FRHIVertexShader* ShaderRHI = Context.RHICmdList.GetBoundVertexShader();
 
 		FGlobalShader::SetParameters<FViewUniformShaderParameters>(Context.RHICmdList, ShaderRHI, Context.View.ViewUniformBuffer);
 
 		PostprocessParameter.SetVS(ShaderRHI, Context, TStaticSamplerState<SF_Bilinear,AM_Clamp,AM_Clamp,AM_Clamp>::GetRHI());
 	}
-
-	// FShader interface.
-	virtual bool Serialize(FArchive& Ar) override
-	{
-		bool bShaderHasOutdatedParameters = FGlobalShader::Serialize(Ar);
-		Ar << PostprocessParameter;
-		return bShaderHasOutdatedParameters;
-	}
 };
 
 template <uint32 UseSunDof> // 0=none, 1=dof, 2=sun, 3=sun&dof, 4=msaa
@@ -110,8 +102,8 @@ class FPostProcessBloomSetupPS_ES2 : public FGlobalShader
 	FPostProcessBloomSetupPS_ES2() {}
 
 public:
-	FPostProcessPassParameters PostprocessParameter;
-	FShaderParameter BloomThreshold;
+	LAYOUT_FIELD(FPostProcessPassParameters, PostprocessParameter);
+	LAYOUT_FIELD(FShaderParameter, BloomThreshold);
 
 	/** Initialization constructor. */
 	FPostProcessBloomSetupPS_ES2(const ShaderMetaType::CompiledShaderInitializerType& Initializer)
@@ -123,7 +115,7 @@ public:
 
 	void SetPS(const FRenderingCompositePassContext& Context)
 	{
-		FRHIPixelShader* ShaderRHI = GetPixelShader();
+		FRHIPixelShader* ShaderRHI = Context.RHICmdList.GetBoundPixelShader();
 
 		const FPostProcessSettings& Settings = Context.View.FinalPostProcessSettings;
 
@@ -133,14 +125,6 @@ public:
 
 		SetShaderValue(Context.RHICmdList, ShaderRHI, BloomThreshold, Settings.BloomThreshold);
 	}
-
-	// FShader interface.
-	virtual bool Serialize(FArchive& Ar) override
-	{
-		bool bShaderHasOutdatedParameters = FGlobalShader::Serialize(Ar);
-		Ar << PostprocessParameter << BloomThreshold;
-		return bShaderHasOutdatedParameters;
-	}
 };
 
 IMPLEMENT_SHADER_TYPE(,FPostProcessBloomSetupVS_ES2,TEXT("/Engine/Private/PostProcessMobile.usf"),TEXT("BloomVS_ES2"),SF_Vertex);
@@ -169,8 +153,8 @@ static void BloomSetup_SetShader(const FRenderingCompositePassContext& Context)
 	TShaderMapRef<FPostProcessBloomSetupPS_ES2<UseSunDof> > PixelShader(Context.GetShaderMap());
 
 	GraphicsPSOInit.BoundShaderState.VertexDeclarationRHI = GFilterVertexDeclaration.VertexDeclarationRHI;
-	GraphicsPSOInit.BoundShaderState.VertexShaderRHI = GETSAFERHISHADER_VERTEX(*VertexShader);
-	GraphicsPSOInit.BoundShaderState.PixelShaderRHI = GETSAFERHISHADER_PIXEL(*PixelShader);
+	GraphicsPSOInit.BoundShaderState.VertexShaderRHI = VertexShader.GetVertexShader();
+	GraphicsPSOInit.BoundShaderState.PixelShaderRHI = PixelShader.GetPixelShader();
 	GraphicsPSOInit.PrimitiveType = PT_TriangleList;
 	SetGraphicsPipelineState(Context.RHICmdList, GraphicsPSOInit);
 
@@ -270,7 +254,7 @@ void FRCPassPostProcessBloomSetupES2::Process(FRenderingCompositePassContext& Co
 			SrcRect.Width(), SrcRect.Height(),
 			DstSize,
 			SrcSize,
-			*VertexShader,
+			VertexShader,
 			EDRF_UseTriangleOptimization);
 	}
 	Context.RHICmdList.EndRenderPass();
@@ -315,7 +299,7 @@ class FPostProcessBloomSetupSmallVS_ES2 : public FGlobalShader
 	FPostProcessBloomSetupSmallVS_ES2() {}
 
 public:
-	FPostProcessPassParameters PostprocessParameter;
+	LAYOUT_FIELD(FPostProcessPassParameters, PostprocessParameter);
 
 	/** Initialization constructor. */
 	FPostProcessBloomSetupSmallVS_ES2(const ShaderMetaType::CompiledShaderInitializerType& Initializer)
@@ -326,20 +310,12 @@ public:
 
 	void SetVS(const FRenderingCompositePassContext& Context)
 	{
-		FRHIVertexShader* ShaderRHI = GetVertexShader();
+		FRHIVertexShader* ShaderRHI = Context.RHICmdList.GetBoundVertexShader();
 
 		FGlobalShader::SetParameters<FViewUniformShaderParameters>(Context.RHICmdList, ShaderRHI, Context.View.ViewUniformBuffer);
 
 		PostprocessParameter.SetVS(ShaderRHI, Context, TStaticSamplerState<SF_Bilinear,AM_Clamp,AM_Clamp,AM_Clamp>::GetRHI());
 	}
-
-	// FShader interface.
-	virtual bool Serialize(FArchive& Ar) override
-	{
-		bool bShaderHasOutdatedParameters = FGlobalShader::Serialize(Ar);
-		Ar << PostprocessParameter;
-		return bShaderHasOutdatedParameters;
-	}
 };
 
 class FPostProcessBloomSetupSmallPS_ES2 : public FGlobalShader
@@ -354,8 +330,8 @@ class FPostProcessBloomSetupSmallPS_ES2 : public FGlobalShader
 	FPostProcessBloomSetupSmallPS_ES2() {}
 
 public:
-	FPostProcessPassParameters PostprocessParameter;
-	FShaderParameter BloomThreshold;
+	LAYOUT_FIELD(FPostProcessPassParameters, PostprocessParameter);
+	LAYOUT_FIELD(FShaderParameter, BloomThreshold);
 
 	/** Initialization constructor. */
 	FPostProcessBloomSetupSmallPS_ES2(const ShaderMetaType::CompiledShaderInitializerType& Initializer)
@@ -368,7 +344,7 @@ public:
 	template <typename TRHICmdList>
 	void SetPS(TRHICmdList& RHICmdList, const FRenderingCompositePassContext& Context)
 	{
-		FRHIPixelShader* ShaderRHI = GetPixelShader();
+		FRHIPixelShader* ShaderRHI = Context.RHICmdList.GetBoundPixelShader();
 
 		const FPostProcessSettings& Settings = Context.View.FinalPostProcessSettings;
 
@@ -378,14 +354,6 @@ public:
 
 		SetShaderValue(RHICmdList, ShaderRHI, BloomThreshold, Settings.BloomThreshold);
 	}
-
-	// FShader interface.
-	virtual bool Serialize(FArchive& Ar) override
-	{
-		bool bShaderHasOutdatedParameters = FGlobalShader::Serialize(Ar);
-		Ar << PostprocessParameter << BloomThreshold;
-		return bShaderHasOutdatedParameters;
-	}
 };
 
 IMPLEMENT_SHADER_TYPE(,FPostProcessBloomSetupSmallVS_ES2,TEXT("/Engine/Private/PostProcessMobile.usf"),TEXT("BloomSmallVS_ES2"),SF_Vertex);
@@ -405,8 +373,8 @@ void FRCPassPostProcessBloomSetupSmallES2::SetShader(const FRenderingCompositePa
 	TShaderMapRef<FPostProcessBloomSetupSmallPS_ES2> PixelShader(Context.GetShaderMap());
 
 	GraphicsPSOInit.BoundShaderState.VertexDeclarationRHI = GFilterVertexDeclaration.VertexDeclarationRHI;
-	GraphicsPSOInit.BoundShaderState.VertexShaderRHI = GETSAFERHISHADER_VERTEX(*VertexShader);
-	GraphicsPSOInit.BoundShaderState.PixelShaderRHI = GETSAFERHISHADER_PIXEL(*PixelShader);
+	GraphicsPSOInit.BoundShaderState.VertexShaderRHI = VertexShader.GetVertexShader();
+	GraphicsPSOInit.BoundShaderState.PixelShaderRHI = PixelShader.GetPixelShader();
 	GraphicsPSOInit.PrimitiveType = PT_TriangleList;
 	SetGraphicsPipelineState(Context.RHICmdList, GraphicsPSOInit);
 
@@ -484,7 +452,7 @@ void FRCPassPostProcessBloomSetupSmallES2::Process(FRenderingCompositePassContex
 			SrcRect.Width(), SrcRect.Height(),
 			DstSize,
 			SrcSize,
-			*VertexShader,
+			VertexShader,
 			EDRF_UseTriangleOptimization);
 	}
 	Context.RHICmdList.EndRenderPass();
@@ -530,7 +498,7 @@ class FPostProcessBloomDownPS_ES2 : public FGlobalShader
 	FPostProcessBloomDownPS_ES2() {}
 
 public:
-	FPostProcessPassParameters PostprocessParameter;
+	LAYOUT_FIELD(FPostProcessPassParameters, PostprocessParameter);
 
 	FPostProcessBloomDownPS_ES2(const ShaderMetaType::CompiledShaderInitializerType& Initializer)
 		: FGlobalShader(Initializer)
@@ -541,18 +509,11 @@ public:
 	template <typename TRHICmdList>
 	void SetPS(TRHICmdList& RHICmdList, const FRenderingCompositePassContext& Context)
 	{
-		FRHIPixelShader* ShaderRHI = GetPixelShader();
+		FRHIPixelShader* ShaderRHI = Context.RHICmdList.GetBoundPixelShader();
 		const FPostProcessSettings& Settings = Context.View.FinalPostProcessSettings;
 		FGlobalShader::SetParameters<FViewUniformShaderParameters>(RHICmdList, ShaderRHI, Context.View.ViewUniformBuffer);
 		PostprocessParameter.SetPS(RHICmdList, ShaderRHI, Context, TStaticSamplerState<SF_Bilinear,AM_Clamp,AM_Clamp,AM_Clamp>::GetRHI());
 	}
-
-	virtual bool Serialize(FArchive& Ar) override
-	{
-		bool bShaderHasOutdatedParameters = FGlobalShader::Serialize(Ar);
-		Ar << PostprocessParameter;
-		return bShaderHasOutdatedParameters;
-	}
 };
 
 IMPLEMENT_SHADER_TYPE(,FPostProcessBloomDownPS_ES2,TEXT("/Engine/Private/PostProcessMobile.usf"),TEXT("BloomDownPS_ES2"),SF_Pixel);
@@ -570,8 +531,8 @@ class FPostProcessBloomDownVS_ES2 : public FGlobalShader
 	FPostProcessBloomDownVS_ES2(){}
 
 public:
-	FPostProcessPassParameters PostprocessParameter;
-	FShaderParameter BloomDownScale;
+	LAYOUT_FIELD(FPostProcessPassParameters, PostprocessParameter);
+	LAYOUT_FIELD(FShaderParameter, BloomDownScale);
 
 	FPostProcessBloomDownVS_ES2(const ShaderMetaType::CompiledShaderInitializerType& Initializer)
 		: FGlobalShader(Initializer)
@@ -582,18 +543,11 @@ public:
 
 	void SetVS(const FRenderingCompositePassContext& Context, float InScale)
 	{
-		FRHIVertexShader* ShaderRHI = GetVertexShader();
+		FRHIVertexShader* ShaderRHI = Context.RHICmdList.GetBoundVertexShader();
 		FGlobalShader::SetParameters<FViewUniformShaderParameters>(Context.RHICmdList, ShaderRHI, Context.View.ViewUniformBuffer);
 		PostprocessParameter.SetVS(ShaderRHI, Context, TStaticSamplerState<SF_Bilinear,AM_Clamp,AM_Clamp,AM_Clamp>::GetRHI());
 		SetShaderValue(Context.RHICmdList, ShaderRHI, BloomDownScale, InScale);
 	}
-
-	virtual bool Serialize(FArchive& Ar) override
-	{
-		bool bShaderHasOutdatedParameters = FGlobalShader::Serialize(Ar);
-		Ar << PostprocessParameter << BloomDownScale;
-		return bShaderHasOutdatedParameters;
-	}
 };
 
 IMPLEMENT_SHADER_TYPE(,FPostProcessBloomDownVS_ES2,TEXT("/Engine/Private/PostProcessMobile.usf"),TEXT("BloomDownVS_ES2"),SF_Vertex);
@@ -644,8 +598,8 @@ void FRCPassPostProcessBloomDownES2::Process(FRenderingCompositePassContext& Con
 		TShaderMapRef<FPostProcessBloomDownPS_ES2> PixelShader(Context.GetShaderMap());
 
 		GraphicsPSOInit.BoundShaderState.VertexDeclarationRHI = GFilterVertexDeclaration.VertexDeclarationRHI;
-		GraphicsPSOInit.BoundShaderState.VertexShaderRHI = GETSAFERHISHADER_VERTEX(*VertexShader);
-		GraphicsPSOInit.BoundShaderState.PixelShaderRHI = GETSAFERHISHADER_PIXEL(*PixelShader);
+		GraphicsPSOInit.BoundShaderState.VertexShaderRHI = VertexShader.GetVertexShader();
+		GraphicsPSOInit.BoundShaderState.PixelShaderRHI = PixelShader.GetPixelShader();
 		GraphicsPSOInit.PrimitiveType = PT_TriangleList;
 
 		SetGraphicsPipelineState(Context.RHICmdList, GraphicsPSOInit);
@@ -663,7 +617,7 @@ void FRCPassPostProcessBloomDownES2::Process(FRenderingCompositePassContext& Con
 			DstX, DstY,
 			SrcDstSize,
 			SrcDstSize,
-			*VertexShader,
+			VertexShader,
 			EDRF_UseTriangleOptimization);
 	}
 	Context.RHICmdList.EndRenderPass();
@@ -708,9 +662,9 @@ class FPostProcessBloomUpPS_ES2 : public FGlobalShader
 	FPostProcessBloomUpPS_ES2() {}
 
 public:
-	FPostProcessPassParameters PostprocessParameter;
-	FShaderParameter TintA;
-	FShaderParameter TintB;
+	LAYOUT_FIELD(FPostProcessPassParameters, PostprocessParameter);
+	LAYOUT_FIELD(FShaderParameter, TintA);
+	LAYOUT_FIELD(FShaderParameter, TintB);
 
 	FPostProcessBloomUpPS_ES2(const ShaderMetaType::CompiledShaderInitializerType& Initializer)
 		: FGlobalShader(Initializer)
@@ -723,20 +677,13 @@ public:
 	template <typename TRHICmdList>
 	void SetPS(TRHICmdList& RHICmdList, const FRenderingCompositePassContext& Context, FVector4& InTintA, FVector4& InTintB)
 	{
-		FRHIPixelShader* ShaderRHI = GetPixelShader();
+		FRHIPixelShader* ShaderRHI = Context.RHICmdList.GetBoundPixelShader();
 		const FPostProcessSettings& Settings = Context.View.FinalPostProcessSettings;
 		FGlobalShader::SetParameters<FViewUniformShaderParameters>(RHICmdList, ShaderRHI, Context.View.ViewUniformBuffer);
 		PostprocessParameter.SetPS(RHICmdList, ShaderRHI, Context, TStaticSamplerState<SF_Bilinear,AM_Clamp,AM_Clamp,AM_Clamp>::GetRHI());
 		SetShaderValue(RHICmdList, ShaderRHI, TintA, InTintA);
 		SetShaderValue(RHICmdList, ShaderRHI, TintB, InTintB);
 	}
-
-	virtual bool Serialize(FArchive& Ar) override
-	{
-		bool bShaderHasOutdatedParameters = FGlobalShader::Serialize(Ar);
-		Ar << PostprocessParameter << TintA << TintB;
-		return bShaderHasOutdatedParameters;
-	}
 };
 
 IMPLEMENT_SHADER_TYPE(,FPostProcessBloomUpPS_ES2,TEXT("/Engine/Private/PostProcessMobile.usf"),TEXT("BloomUpPS_ES2"),SF_Pixel);
@@ -754,8 +701,8 @@ class FPostProcessBloomUpVS_ES2 : public FGlobalShader
 	FPostProcessBloomUpVS_ES2(){}
 
 public:
-	FPostProcessPassParameters PostprocessParameter;
-	FShaderParameter BloomUpScales;
+	LAYOUT_FIELD(FPostProcessPassParameters, PostprocessParameter);
+	LAYOUT_FIELD(FShaderParameter, BloomUpScales);
 
 	FPostProcessBloomUpVS_ES2(const ShaderMetaType::CompiledShaderInitializerType& Initializer)
 		: FGlobalShader(Initializer)
@@ -766,18 +713,11 @@ public:
 
 	void SetVS(const FRenderingCompositePassContext& Context, FVector2D InScale)
 	{
-		FRHIVertexShader* ShaderRHI = GetVertexShader();
+		FRHIVertexShader* ShaderRHI = Context.RHICmdList.GetBoundVertexShader();
 		FGlobalShader::SetParameters<FViewUniformShaderParameters>(Context.RHICmdList, ShaderRHI, Context.View.ViewUniformBuffer);
 		PostprocessParameter.SetVS(ShaderRHI, Context, TStaticSamplerState<SF_Bilinear,AM_Clamp,AM_Clamp,AM_Clamp>::GetRHI());
 		SetShaderValue(Context.RHICmdList, ShaderRHI, BloomUpScales, InScale);
 	}
-
-	virtual bool Serialize(FArchive& Ar) override
-	{
-		bool bShaderHasOutdatedParameters = FGlobalShader::Serialize(Ar);
-		Ar << PostprocessParameter << BloomUpScales;
-		return bShaderHasOutdatedParameters;
-	}
 };
 
 IMPLEMENT_SHADER_TYPE(,FPostProcessBloomUpVS_ES2,TEXT("/Engine/Private/PostProcessMobile.usf"),TEXT("BloomUpVS_ES2"),SF_Vertex);
@@ -828,8 +768,8 @@ void FRCPassPostProcessBloomUpES2::Process(FRenderingCompositePassContext& Conte
 		TShaderMapRef<FPostProcessBloomUpPS_ES2> PixelShader(Context.GetShaderMap());
 
 		GraphicsPSOInit.BoundShaderState.VertexDeclarationRHI = GFilterVertexDeclaration.VertexDeclarationRHI;
-		GraphicsPSOInit.BoundShaderState.VertexShaderRHI = GETSAFERHISHADER_VERTEX(*VertexShader);
-		GraphicsPSOInit.BoundShaderState.PixelShaderRHI = GETSAFERHISHADER_PIXEL(*PixelShader);
+		GraphicsPSOInit.BoundShaderState.VertexShaderRHI = VertexShader.GetVertexShader();
+		GraphicsPSOInit.BoundShaderState.PixelShaderRHI = PixelShader.GetPixelShader();
 		GraphicsPSOInit.PrimitiveType = PT_TriangleList;
 
 		SetGraphicsPipelineState(Context.RHICmdList, GraphicsPSOInit);
@@ -850,7 +790,7 @@ void FRCPassPostProcessBloomUpES2::Process(FRenderingCompositePassContext& Conte
 			DstX, DstY,
 			SrcDstSize,
 			SrcDstSize,
-			*VertexShader,
+			VertexShader,
 			EDRF_UseTriangleOptimization);
 	}
 	Context.RHICmdList.EndRenderPass();
@@ -906,10 +846,10 @@ class FPostProcessSunMaskPS_ES2 : public FGlobalShader
 	FPostProcessSunMaskPS_ES2() {}
 
 public:
-	FPostProcessPassParameters PostprocessParameter;
-	FShaderParameter SunColorApertureDiv2Parameter;
+	LAYOUT_FIELD(FPostProcessPassParameters, PostprocessParameter);
+	LAYOUT_FIELD(FShaderParameter, SunColorApertureDiv2Parameter);
 
-	FSceneTextureShaderParameters SceneTextureParameters;
+	LAYOUT_FIELD(FSceneTextureShaderParameters, SceneTextureParameters)
 
 	FPostProcessSunMaskPS_ES2(const ShaderMetaType::CompiledShaderInitializerType& Initializer)
 		: FGlobalShader(Initializer)
@@ -921,7 +861,7 @@ public:
 
 	void SetPS(const FRenderingCompositePassContext& Context)
 	{
-		FRHIPixelShader* ShaderRHI = GetPixelShader();
+		FRHIPixelShader* ShaderRHI = Context.RHICmdList.GetBoundPixelShader();
 		const FPostProcessSettings& Settings = Context.View.FinalPostProcessSettings;
 		FGlobalShader::SetParameters<FViewUniformShaderParameters>(Context.RHICmdList, ShaderRHI, Context.View.ViewUniformBuffer);
 		PostprocessParameter.SetPS(Context.RHICmdList, ShaderRHI, Context, TStaticSamplerState<SF_Bilinear,AM_Clamp,AM_Clamp,AM_Clamp>::GetRHI());
@@ -935,13 +875,6 @@ public:
 
 		SceneTextureParameters.Set(Context.RHICmdList, ShaderRHI, Context.View.FeatureLevel, ESceneTextureSetupMode::All);
 	}
-
-	virtual bool Serialize(FArchive& Ar) override
-	{
-		bool bShaderHasOutdatedParameters = FGlobalShader::Serialize(Ar);
-		Ar << PostprocessParameter << SunColorApertureDiv2Parameter << SceneTextureParameters;
-		return bShaderHasOutdatedParameters;
-	}
 };
 
 // #define avoids a lot of code duplication
@@ -966,7 +899,7 @@ class FPostProcessSunMaskVS_ES2 : public FGlobalShader
 	FPostProcessSunMaskVS_ES2(){}
 
 public:
-	FPostProcessPassParameters PostprocessParameter;
+	LAYOUT_FIELD(FPostProcessPassParameters, PostprocessParameter);
 
 	FPostProcessSunMaskVS_ES2(const ShaderMetaType::CompiledShaderInitializerType& Initializer)
 		: FGlobalShader(Initializer)
@@ -976,17 +909,10 @@ public:
 
 	void SetVS(const FRenderingCompositePassContext& Context)
 	{
-		FRHIVertexShader* ShaderRHI = GetVertexShader();
+		FRHIVertexShader* ShaderRHI = Context.RHICmdList.GetBoundVertexShader();
 		FGlobalShader::SetParameters<FViewUniformShaderParameters>(Context.RHICmdList, ShaderRHI, Context.View.ViewUniformBuffer);
 		PostprocessParameter.SetVS(ShaderRHI, Context, TStaticSamplerState<SF_Bilinear,AM_Clamp,AM_Clamp,AM_Clamp>::GetRHI());
 	}
-
-	virtual bool Serialize(FArchive& Ar) override
-	{
-		bool bShaderHasOutdatedParameters = FGlobalShader::Serialize(Ar);
-		Ar << PostprocessParameter;
-		return bShaderHasOutdatedParameters;
-	}
 };
 
 IMPLEMENT_SHADER_TYPE(,FPostProcessSunMaskVS_ES2,TEXT("/Engine/Private/PostProcessMobile.usf"),TEXT("SunMaskVS_ES2"),SF_Vertex);
@@ -1004,8 +930,8 @@ static void SunMask_SetShader(const FRenderingCompositePassContext& Context)
 	TShaderMapRef<FPostProcessSunMaskPS_ES2<UseSunDof, bUseDepthTexture> > PixelShader(Context.GetShaderMap());
 
 	GraphicsPSOInit.BoundShaderState.VertexDeclarationRHI = GFilterVertexDeclaration.VertexDeclarationRHI;
-	GraphicsPSOInit.BoundShaderState.VertexShaderRHI = GETSAFERHISHADER_VERTEX(*VertexShader);
-	GraphicsPSOInit.BoundShaderState.PixelShaderRHI = GETSAFERHISHADER_PIXEL(*PixelShader);
+	GraphicsPSOInit.BoundShaderState.VertexShaderRHI = VertexShader.GetVertexShader();
+	GraphicsPSOInit.BoundShaderState.PixelShaderRHI = PixelShader.GetPixelShader();
 	GraphicsPSOInit.PrimitiveType = PT_TriangleList;
 
 	SetGraphicsPipelineState(Context.RHICmdList, GraphicsPSOInit);
@@ -1098,7 +1024,7 @@ void FRCPassPostProcessSunMaskES2::Process(FRenderingCompositePassContext& Conte
 				SrcRect.Width(), SrcRect.Height(),
 				DstSize,
 				SrcSize,
-				*VertexShader,
+				VertexShader,
 				EDRF_UseTriangleOptimization);
 		}
 		Context.RHICmdList.EndRenderPass();
@@ -1150,7 +1076,7 @@ class FPostProcessSunAlphaPS_ES2 : public FGlobalShader
 	FPostProcessSunAlphaPS_ES2() {}
 
 public:
-	FPostProcessPassParameters PostprocessParameter;
+	LAYOUT_FIELD(FPostProcessPassParameters, PostprocessParameter);
 
 	FPostProcessSunAlphaPS_ES2(const ShaderMetaType::CompiledShaderInitializerType& Initializer)
 		: FGlobalShader(Initializer)
@@ -1160,18 +1086,11 @@ public:
 
 	void SetPS(const FRenderingCompositePassContext& Context)
 	{
-		FRHIPixelShader* ShaderRHI = GetPixelShader();
+		FRHIPixelShader* ShaderRHI = Context.RHICmdList.GetBoundPixelShader();
 		const FPostProcessSettings& Settings = Context.View.FinalPostProcessSettings;
 		FGlobalShader::SetParameters<FViewUniformShaderParameters>(Context.RHICmdList, ShaderRHI, Context.View.ViewUniformBuffer);
 		PostprocessParameter.SetPS(Context.RHICmdList, ShaderRHI, Context, TStaticSamplerState<SF_Bilinear, AM_Clamp, AM_Clamp, AM_Clamp>::GetRHI());
 	}
-
-	virtual bool Serialize(FArchive& Ar) override
-	{
-		bool bShaderHasOutdatedParameters = FGlobalShader::Serialize(Ar);
-		Ar << PostprocessParameter;
-		return bShaderHasOutdatedParameters;
-	}
 };
 
 typedef FPostProcessSunAlphaPS_ES2<0> FPostProcessSunAlphaPS_ES2_0;
@@ -1191,8 +1110,8 @@ class FPostProcessSunAlphaVS_ES2 : public FGlobalShader
 	FPostProcessSunAlphaVS_ES2(){}
 
 public:
-	FPostProcessPassParameters PostprocessParameter;
-	FShaderParameter LightShaftCenter;
+	LAYOUT_FIELD(FPostProcessPassParameters, PostprocessParameter);
+	LAYOUT_FIELD(FShaderParameter, LightShaftCenter);
 
 	FPostProcessSunAlphaVS_ES2(const ShaderMetaType::CompiledShaderInitializerType& Initializer)
 		: FGlobalShader(Initializer)
@@ -1203,19 +1122,12 @@ public:
 
 	void SetVS(const FRenderingCompositePassContext& Context)
 	{
-		FRHIVertexShader* ShaderRHI = GetVertexShader();
+		FRHIVertexShader* ShaderRHI = Context.RHICmdList.GetBoundVertexShader();
 		FGlobalShader::SetParameters<FViewUniformShaderParameters>(Context.RHICmdList, ShaderRHI, Context.View.ViewUniformBuffer);
 		PostprocessParameter.SetVS(ShaderRHI, Context, TStaticSamplerState<SF_Bilinear,AM_Clamp,AM_Clamp,AM_Clamp>::GetRHI());
 
 		SetShaderValue(Context.RHICmdList, ShaderRHI, LightShaftCenter, Context.View.LightShaftCenter);
 	}
-
-	virtual bool Serialize(FArchive& Ar) override
-	{
-		bool bShaderHasOutdatedParameters = FGlobalShader::Serialize(Ar);
-		Ar << PostprocessParameter << LightShaftCenter;
-		return bShaderHasOutdatedParameters;
-	}
 };
 
 IMPLEMENT_SHADER_TYPE(,FPostProcessSunAlphaVS_ES2,TEXT("/Engine/Private/PostProcessMobile.usf"),TEXT("SunAlphaVS_ES2"),SF_Vertex);
@@ -1233,8 +1145,8 @@ static void SunAlpha_SetShader(const FRenderingCompositePassContext& Context)
 	TShaderMapRef<FPostProcessSunAlphaPS_ES2<UseDof> > PixelShader(Context.GetShaderMap());
 
 	GraphicsPSOInit.BoundShaderState.VertexDeclarationRHI = GFilterVertexDeclaration.VertexDeclarationRHI;
-	GraphicsPSOInit.BoundShaderState.VertexShaderRHI = GETSAFERHISHADER_VERTEX(*VertexShader);
-	GraphicsPSOInit.BoundShaderState.PixelShaderRHI = GETSAFERHISHADER_PIXEL(*PixelShader);
+	GraphicsPSOInit.BoundShaderState.VertexShaderRHI = VertexShader.GetVertexShader();
+	GraphicsPSOInit.BoundShaderState.PixelShaderRHI = PixelShader.GetPixelShader();
 	GraphicsPSOInit.PrimitiveType = PT_TriangleList;
 
 	SetGraphicsPipelineState(Context.RHICmdList, GraphicsPSOInit);
@@ -1302,7 +1214,7 @@ void FRCPassPostProcessSunAlphaES2::Process(FRenderingCompositePassContext& Cont
 			DstX, DstY,
 			SrcDstSize,
 			SrcDstSize,
-			*VertexShader,
+			VertexShader,
 			EDRF_UseTriangleOptimization);
 	}
 	Context.RHICmdList.EndRenderPass();
@@ -1349,7 +1261,7 @@ class FPostProcessSunBlurPS_ES2 : public FGlobalShader
 	FPostProcessSunBlurPS_ES2() {}
 
 public:
-	FPostProcessPassParameters PostprocessParameter;
+	LAYOUT_FIELD(FPostProcessPassParameters, PostprocessParameter);
 
 	FPostProcessSunBlurPS_ES2(const ShaderMetaType::CompiledShaderInitializerType& Initializer)
 		: FGlobalShader(Initializer)
@@ -1359,18 +1271,11 @@ public:
 
 	void SetPS(const FRenderingCompositePassContext& Context)
 	{
-		FRHIPixelShader* ShaderRHI = GetPixelShader();
+		FRHIPixelShader* ShaderRHI = Context.RHICmdList.GetBoundPixelShader();
 		const FPostProcessSettings& Settings = Context.View.FinalPostProcessSettings;
 		FGlobalShader::SetParameters<FViewUniformShaderParameters>(Context.RHICmdList, ShaderRHI, Context.View.ViewUniformBuffer);
 		PostprocessParameter.SetPS(Context.RHICmdList, ShaderRHI, Context, TStaticSamplerState<SF_Bilinear, AM_Clamp, AM_Clamp, AM_Clamp>::GetRHI());
 	}
-
-	virtual bool Serialize(FArchive& Ar) override
-	{
-		bool bShaderHasOutdatedParameters = FGlobalShader::Serialize(Ar);
-		Ar << PostprocessParameter;
-		return bShaderHasOutdatedParameters;
-	}
 };
 
 IMPLEMENT_SHADER_TYPE(,FPostProcessSunBlurPS_ES2,TEXT("/Engine/Private/PostProcessMobile.usf"),TEXT("SunBlurPS_ES2"),SF_Pixel);
@@ -1388,8 +1293,8 @@ class FPostProcessSunBlurVS_ES2 : public FGlobalShader
 	FPostProcessSunBlurVS_ES2(){}
 
 public:
-	FPostProcessPassParameters PostprocessParameter;
-	FShaderParameter LightShaftCenter;
+	LAYOUT_FIELD(FPostProcessPassParameters, PostprocessParameter);
+	LAYOUT_FIELD(FShaderParameter, LightShaftCenter);
 
 	FPostProcessSunBlurVS_ES2(const ShaderMetaType::CompiledShaderInitializerType& Initializer)
 		: FGlobalShader(Initializer)
@@ -1400,19 +1305,12 @@ public:
 
 	void SetVS(const FRenderingCompositePassContext& Context)
 	{
-		FRHIVertexShader* ShaderRHI = GetVertexShader();
+		FRHIVertexShader* ShaderRHI = Context.RHICmdList.GetBoundVertexShader();
 		FGlobalShader::SetParameters<FViewUniformShaderParameters>(Context.RHICmdList, ShaderRHI, Context.View.ViewUniformBuffer);
 		PostprocessParameter.SetVS(ShaderRHI, Context, TStaticSamplerState<SF_Bilinear,AM_Clamp,AM_Clamp,AM_Clamp>::GetRHI());
 
 		SetShaderValue(Context.RHICmdList, ShaderRHI, LightShaftCenter, Context.View.LightShaftCenter);
 	}
-
-	virtual bool Serialize(FArchive& Ar) override
-	{
-		bool bShaderHasOutdatedParameters = FGlobalShader::Serialize(Ar);
-		Ar << PostprocessParameter << LightShaftCenter;
-		return bShaderHasOutdatedParameters;
-	}
 };
 
 IMPLEMENT_SHADER_TYPE(,FPostProcessSunBlurVS_ES2,TEXT("/Engine/Private/PostProcessMobile.usf"),TEXT("SunBlurVS_ES2"),SF_Vertex);
@@ -1461,8 +1359,8 @@ void FRCPassPostProcessSunBlurES2::Process(FRenderingCompositePassContext& Conte
 		TShaderMapRef<FPostProcessSunBlurPS_ES2> PixelShader(Context.GetShaderMap());
 
 		GraphicsPSOInit.BoundShaderState.VertexDeclarationRHI = GFilterVertexDeclaration.VertexDeclarationRHI;
-		GraphicsPSOInit.BoundShaderState.VertexShaderRHI = GETSAFERHISHADER_VERTEX(*VertexShader);
-		GraphicsPSOInit.BoundShaderState.PixelShaderRHI = GETSAFERHISHADER_PIXEL(*PixelShader);
+		GraphicsPSOInit.BoundShaderState.VertexShaderRHI = VertexShader.GetVertexShader();
+		GraphicsPSOInit.BoundShaderState.PixelShaderRHI = PixelShader.GetPixelShader();
 		GraphicsPSOInit.PrimitiveType = PT_TriangleList;
 
 		SetGraphicsPipelineState(Context.RHICmdList, GraphicsPSOInit);
@@ -1480,7 +1378,7 @@ void FRCPassPostProcessSunBlurES2::Process(FRenderingCompositePassContext& Conte
 			DstX, DstY,
 			SrcDstSize,
 			SrcDstSize,
-			*VertexShader,
+			VertexShader,
 			EDRF_UseTriangleOptimization);
 	}
 	Context.RHICmdList.EndRenderPass();
@@ -1535,10 +1433,10 @@ class FPostProcessSunMergePS_ES2 : public FGlobalShader
 	FPostProcessSunMergePS_ES2() {}
 
 public:
-	FPostProcessPassParameters PostprocessParameter;
-	FShaderParameter SunColorVignetteIntensity;
-	FShaderParameter VignetteColor;
-	FShaderParameter BloomColor;
+	LAYOUT_FIELD(FPostProcessPassParameters, PostprocessParameter);
+	LAYOUT_FIELD(FShaderParameter, SunColorVignetteIntensity);
+	LAYOUT_FIELD(FShaderParameter, VignetteColor);
+	LAYOUT_FIELD(FShaderParameter, BloomColor);
 
 	FPostProcessSunMergePS_ES2(const ShaderMetaType::CompiledShaderInitializerType& Initializer)
 		: FGlobalShader(Initializer)
@@ -1551,7 +1449,7 @@ public:
 
 	void SetPS(const FRenderingCompositePassContext& Context)
 	{
-		FRHIPixelShader* ShaderRHI = GetPixelShader();
+		FRHIPixelShader* ShaderRHI = Context.RHICmdList.GetBoundPixelShader();
 		const FPostProcessSettings& Settings = Context.View.FinalPostProcessSettings;
 		FGlobalShader::SetParameters<FViewUniformShaderParameters>(Context.RHICmdList, ShaderRHI, Context.View.ViewUniformBuffer);
 		PostprocessParameter.SetPS(Context.RHICmdList, ShaderRHI, Context, TStaticSamplerState<SF_Bilinear, AM_Clamp, AM_Clamp, AM_Clamp>::GetRHI());
@@ -1566,13 +1464,6 @@ public:
 		// Scaling Bloom1 by extra factor to match filter area difference between PC default and mobile.
 		SetShaderValue(Context.RHICmdList, ShaderRHI, BloomColor, Context.View.FinalPostProcessSettings.Bloom1Tint * Context.View.FinalPostProcessSettings.BloomIntensity * 0.5);
 	}
-
-	virtual bool Serialize(FArchive& Ar) override
-	{
-		bool bShaderHasOutdatedParameters = FGlobalShader::Serialize(Ar);
-		Ar << PostprocessParameter << SunColorVignetteIntensity << BloomColor;
-		return bShaderHasOutdatedParameters;
-	}
 };
 
 typedef FPostProcessSunMergePS_ES2<0> FPostProcessSunMergePS_ES2_0;
@@ -1597,8 +1488,8 @@ class FPostProcessSunMergeVS_ES2 : public FGlobalShader
 	FPostProcessSunMergeVS_ES2(){}
 
 public:
-	FPostProcessPassParameters PostprocessParameter;
-	FShaderParameter LightShaftCenter;
+	LAYOUT_FIELD(FPostProcessPassParameters, PostprocessParameter);
+	LAYOUT_FIELD(FShaderParameter, LightShaftCenter);
 
 	FPostProcessSunMergeVS_ES2(const ShaderMetaType::CompiledShaderInitializerType& Initializer)
 		: FGlobalShader(Initializer)
@@ -1609,19 +1500,12 @@ public:
 
 	void SetVS(const FRenderingCompositePassContext& Context)
 	{
-		FRHIVertexShader* ShaderRHI = GetVertexShader();
+		FRHIVertexShader* ShaderRHI = Context.RHICmdList.GetBoundVertexShader();
 		FGlobalShader::SetParameters<FViewUniformShaderParameters>(Context.RHICmdList, ShaderRHI, Context.View.ViewUniformBuffer);
 		PostprocessParameter.SetVS(ShaderRHI, Context, TStaticSamplerState<SF_Bilinear,AM_Clamp,AM_Clamp,AM_Clamp>::GetRHI());
 
 		SetShaderValue(Context.RHICmdList, ShaderRHI, LightShaftCenter, Context.View.LightShaftCenter);
 	}
-
-	virtual bool Serialize(FArchive& Ar) override
-	{
-		bool bShaderHasOutdatedParameters = FGlobalShader::Serialize(Ar);
-		Ar << PostprocessParameter << LightShaftCenter;
-		return bShaderHasOutdatedParameters;
-	}
 };
 
 IMPLEMENT_SHADER_TYPE(,FPostProcessSunMergeVS_ES2,TEXT("/Engine/Private/PostProcessMobile.usf"),TEXT("SunMergeVS_ES2"),SF_Vertex);
@@ -1629,7 +1513,7 @@ IMPLEMENT_SHADER_TYPE(,FPostProcessSunMergeVS_ES2,TEXT("/Engine/Private/PostProc
 
 
 template <uint32 UseSunBloom>
-FShader* SunMerge_SetShader(const FRenderingCompositePassContext& Context)
+TShaderRef<FShader> SunMerge_SetShader(const FRenderingCompositePassContext& Context)
 {
 	FGraphicsPipelineStateInitializer GraphicsPSOInit;
 	Context.RHICmdList.ApplyCachedRenderTargets(GraphicsPSOInit);
@@ -1641,8 +1525,8 @@ FShader* SunMerge_SetShader(const FRenderingCompositePassContext& Context)
 	TShaderMapRef<FPostProcessSunMergePS_ES2<UseSunBloom> > PixelShader(Context.GetShaderMap());
 
 	GraphicsPSOInit.BoundShaderState.VertexDeclarationRHI = GFilterVertexDeclaration.VertexDeclarationRHI;
-	GraphicsPSOInit.BoundShaderState.VertexShaderRHI = GETSAFERHISHADER_VERTEX(*VertexShader);
-	GraphicsPSOInit.BoundShaderState.PixelShaderRHI = GETSAFERHISHADER_PIXEL(*PixelShader);
+	GraphicsPSOInit.BoundShaderState.VertexShaderRHI = VertexShader.GetVertexShader();
+	GraphicsPSOInit.BoundShaderState.PixelShaderRHI = PixelShader.GetPixelShader();
 	GraphicsPSOInit.PrimitiveType = PT_TriangleList;
 
 	SetGraphicsPipelineState(Context.RHICmdList, GraphicsPSOInit);
@@ -1650,10 +1534,10 @@ FShader* SunMerge_SetShader(const FRenderingCompositePassContext& Context)
 	VertexShader->SetVS(Context);
 	PixelShader->SetPS(Context);
 
-	return *VertexShader;
+	return VertexShader;
 }
 
-FShader* FRCPassPostProcessSunMergeES2::SetShader(const FRenderingCompositePassContext& Context)
+TShaderRef<FShader> FRCPassPostProcessSunMergeES2::SetShader(const FRenderingCompositePassContext& Context)
 {
 	const FViewInfo& View = Context.View;
 	uint32 UseBloom = (View.FinalPostProcessSettings.BloomIntensity > 0.0f) ? 1 : 0;
@@ -1669,7 +1553,7 @@ FShader* FRCPassPostProcessSunMergeES2::SetShader(const FRenderingCompositePassC
 	}
 
 	check(false);
-	return NULL;
+	return TShaderRef<FShader>();
 }
 
 void FRCPassPostProcessSunMergeES2::Process(FRenderingCompositePassContext& Context)
@@ -1704,7 +1588,7 @@ void FRCPassPostProcessSunMergeES2::Process(FRenderingCompositePassContext& Cont
 
 		Context.SetViewportAndCallRHI(0, 0, 0.0f, DstX, DstY, 1.0f);
 
-		FShader* VertexShader = SetShader(Context);
+		TShaderRef<FShader> VertexShader = SetShader(Context);
 
 		FIntPoint SrcDstSize = PrePostSourceViewportSize / 4;
 
@@ -1773,11 +1657,11 @@ class FPostProcessSunMergeSmallPS_ES2 : public FGlobalShader
 	FPostProcessSunMergeSmallPS_ES2() {}
 
 public:
-	FPostProcessPassParameters PostprocessParameter;
-	FShaderParameter SunColorVignetteIntensity;
-	FShaderParameter VignetteColor;
-	FShaderParameter BloomColor;
-	FShaderParameter BloomColor2;
+	LAYOUT_FIELD(FPostProcessPassParameters, PostprocessParameter);
+	LAYOUT_FIELD(FShaderParameter, SunColorVignetteIntensity);
+	LAYOUT_FIELD(FShaderParameter, VignetteColor);
+	LAYOUT_FIELD(FShaderParameter, BloomColor);
+	LAYOUT_FIELD(FShaderParameter, BloomColor2);
 
 	FPostProcessSunMergeSmallPS_ES2(const ShaderMetaType::CompiledShaderInitializerType& Initializer)
 		: FGlobalShader(Initializer)
@@ -1791,7 +1675,7 @@ public:
 
 	void SetPS(const FRenderingCompositePassContext& Context)
 	{
-		FRHIPixelShader* ShaderRHI = GetPixelShader();
+		FRHIPixelShader* ShaderRHI = Context.RHICmdList.GetBoundPixelShader();
 		const FPostProcessSettings& Settings = Context.View.FinalPostProcessSettings;
 		FGlobalShader::SetParameters<FViewUniformShaderParameters>(Context.RHICmdList, ShaderRHI, Context.View.ViewUniformBuffer);
 		PostprocessParameter.SetPS(Context.RHICmdList, ShaderRHI, Context, TStaticSamplerState<SF_Bilinear,AM_Clamp,AM_Clamp,AM_Clamp>::GetRHI());
@@ -1807,13 +1691,6 @@ public:
 		SetShaderValue(Context.RHICmdList, ShaderRHI, BloomColor, Context.View.FinalPostProcessSettings.Bloom1Tint * Context.View.FinalPostProcessSettings.BloomIntensity * 0.5);
 		SetShaderValue(Context.RHICmdList, ShaderRHI, BloomColor2, Context.View.FinalPostProcessSettings.Bloom2Tint * Context.View.FinalPostProcessSettings.BloomIntensity);
 	}
-
-	virtual bool Serialize(FArchive& Ar) override
-	{
-		bool bShaderHasOutdatedParameters = FGlobalShader::Serialize(Ar);
-		Ar << PostprocessParameter << SunColorVignetteIntensity << BloomColor << BloomColor2;
-		return bShaderHasOutdatedParameters;
-	}
 };
 
 IMPLEMENT_SHADER_TYPE(,FPostProcessSunMergeSmallPS_ES2,TEXT("/Engine/Private/PostProcessMobile.usf"),TEXT("SunMergeSmallPS_ES2"),SF_Pixel);
@@ -1830,7 +1707,7 @@ class FPostProcessSunMergeSmallVS_ES2 : public FGlobalShader
 	FPostProcessSunMergeSmallVS_ES2(){}
 
 public:
-	FPostProcessPassParameters PostprocessParameter;
+	LAYOUT_FIELD(FPostProcessPassParameters, PostprocessParameter);
 
 	FPostProcessSunMergeSmallVS_ES2(const ShaderMetaType::CompiledShaderInitializerType& Initializer)
 		: FGlobalShader(Initializer)
@@ -1840,17 +1717,10 @@ public:
 
 	void SetVS(const FRenderingCompositePassContext& Context)
 	{
-		FRHIVertexShader* ShaderRHI = GetVertexShader();
+		FRHIVertexShader* ShaderRHI = Context.RHICmdList.GetBoundVertexShader();
 		FGlobalShader::SetParameters<FViewUniformShaderParameters>(Context.RHICmdList, ShaderRHI, Context.View.ViewUniformBuffer);
 		PostprocessParameter.SetVS(ShaderRHI, Context, TStaticSamplerState<SF_Bilinear,AM_Clamp,AM_Clamp,AM_Clamp>::GetRHI());
 	}
-
-	virtual bool Serialize(FArchive& Ar) override
-	{
-		bool bShaderHasOutdatedParameters = FGlobalShader::Serialize(Ar);
-		Ar << PostprocessParameter;
-		return bShaderHasOutdatedParameters;
-	}
 };
 
 IMPLEMENT_SHADER_TYPE(,FPostProcessSunMergeSmallVS_ES2,TEXT("/Engine/Private/PostProcessMobile.usf"),TEXT("SunMergeSmallVS_ES2"),SF_Vertex);
@@ -1867,8 +1737,8 @@ void FRCPassPostProcessSunMergeSmallES2::SetShader(const FRenderingCompositePass
 	TShaderMapRef<FPostProcessSunMergeSmallPS_ES2> PixelShader(Context.GetShaderMap());
 
 	GraphicsPSOInit.BoundShaderState.VertexDeclarationRHI = GFilterVertexDeclaration.VertexDeclarationRHI;
-	GraphicsPSOInit.BoundShaderState.VertexShaderRHI = GETSAFERHISHADER_VERTEX(*VertexShader);
-	GraphicsPSOInit.BoundShaderState.PixelShaderRHI = GETSAFERHISHADER_PIXEL(*PixelShader);
+	GraphicsPSOInit.BoundShaderState.VertexShaderRHI = VertexShader.GetVertexShader();
+	GraphicsPSOInit.BoundShaderState.PixelShaderRHI = PixelShader.GetPixelShader();
 	GraphicsPSOInit.PrimitiveType = PT_TriangleList;
 
 	SetGraphicsPipelineState(Context.RHICmdList, GraphicsPSOInit);
@@ -1922,7 +1792,7 @@ void FRCPassPostProcessSunMergeSmallES2::Process(FRenderingCompositePassContext&
 			DstX, DstY,
 			SrcDstSize,
 			SrcDstSize,
-			*VertexShader,
+			VertexShader,
 			EDRF_UseTriangleOptimization);
 	}
 	Context.RHICmdList.EndRenderPass();
@@ -1985,7 +1855,7 @@ class FPostProcessDofDownVS_ES2 : public FGlobalShader
 	FPostProcessDofDownVS_ES2() {}
 
 public:
-	FPostProcessPassParameters PostprocessParameter;
+	LAYOUT_FIELD(FPostProcessPassParameters, PostprocessParameter);
 
 	/** Initialization constructor. */
 	FPostProcessDofDownVS_ES2(const ShaderMetaType::CompiledShaderInitializerType& Initializer)
@@ -1996,20 +1866,12 @@ public:
 
 	void SetVS(const FRenderingCompositePassContext& Context)
 	{
-		FRHIVertexShader* ShaderRHI = GetVertexShader();
+		FRHIVertexShader* ShaderRHI = Context.RHICmdList.GetBoundVertexShader();
 
 		FGlobalShader::SetParameters<FViewUniformShaderParameters>(Context.RHICmdList, ShaderRHI, Context.View.ViewUniformBuffer);
 
 		PostprocessParameter.SetVS(ShaderRHI, Context, TStaticSamplerState<SF_Point,AM_Clamp,AM_Clamp,AM_Clamp>::GetRHI());
 	}
-
-	// FShader interface.
-	virtual bool Serialize(FArchive& Ar) override
-	{
-		bool bShaderHasOutdatedParameters = FGlobalShader::Serialize(Ar);
-		Ar << PostprocessParameter;
-		return bShaderHasOutdatedParameters;
-	}
 };
 
 template<uint32 UseSun>
@@ -2031,7 +1893,7 @@ class FPostProcessDofDownPS_ES2 : public FGlobalShader
 	FPostProcessDofDownPS_ES2() {}
 
 public:
-	FPostProcessPassParameters PostprocessParameter;
+	LAYOUT_FIELD(FPostProcessPassParameters, PostprocessParameter);
 
 	FPostProcessDofDownPS_ES2(const ShaderMetaType::CompiledShaderInitializerType& Initializer)
 		: FGlobalShader(Initializer)
@@ -2041,18 +1903,11 @@ public:
 
 	void SetPS(const FRenderingCompositePassContext& Context)
 	{
-		FRHIPixelShader* ShaderRHI = GetPixelShader();
+		FRHIPixelShader* ShaderRHI = Context.RHICmdList.GetBoundPixelShader();
 		const FPostProcessSettings& Settings = Context.View.FinalPostProcessSettings;
 		FGlobalShader::SetParameters<FViewUniformShaderParameters>(Context.RHICmdList, ShaderRHI, Context.View.ViewUniformBuffer);
 		PostprocessParameter.SetPS(Context.RHICmdList, ShaderRHI, Context, TStaticSamplerState<SF_Point, AM_Clamp, AM_Clamp, AM_Clamp>::GetRHI());
 	}
-
-	virtual bool Serialize(FArchive& Ar) override
-	{
-		bool bShaderHasOutdatedParameters = FGlobalShader::Serialize(Ar);
-		Ar << PostprocessParameter;
-		return bShaderHasOutdatedParameters;
-	}
 };
 
 IMPLEMENT_SHADER_TYPE(,FPostProcessDofDownVS_ES2,TEXT("/Engine/Private/PostProcessMobile.usf"),TEXT("DofDownVS_ES2"),SF_Vertex);
@@ -2075,8 +1930,8 @@ static void DofDown_SetShader(const FRenderingCompositePassContext& Context)
 	TShaderMapRef<FPostProcessDofDownPS_ES2<UseSun> > PixelShader(Context.GetShaderMap());
 
 	GraphicsPSOInit.BoundShaderState.VertexDeclarationRHI = GFilterVertexDeclaration.VertexDeclarationRHI;
-	GraphicsPSOInit.BoundShaderState.VertexShaderRHI = GETSAFERHISHADER_VERTEX(*VertexShader);
-	GraphicsPSOInit.BoundShaderState.PixelShaderRHI = GETSAFERHISHADER_PIXEL(*PixelShader);
+	GraphicsPSOInit.BoundShaderState.VertexShaderRHI = VertexShader.GetVertexShader();
+	GraphicsPSOInit.BoundShaderState.PixelShaderRHI = PixelShader.GetPixelShader();
 	GraphicsPSOInit.PrimitiveType = PT_TriangleList;
 
 	SetGraphicsPipelineState(Context.RHICmdList, GraphicsPSOInit);
@@ -2165,7 +2020,7 @@ void FRCPassPostProcessDofDownES2::Process(FRenderingCompositePassContext& Conte
 			SrcRect.Width(), SrcRect.Height(),
 			DstSize,
 			SrcSize,
-			*VertexShader,
+			VertexShader,
 			EDRF_UseTriangleOptimization);
 	}
 	Context.RHICmdList.EndRenderPass();
@@ -2211,7 +2066,7 @@ class FPostProcessDofNearVS_ES2 : public FGlobalShader
 	FPostProcessDofNearVS_ES2() {}
 
 public:
-	FPostProcessPassParameters PostprocessParameter;
+	LAYOUT_FIELD(FPostProcessPassParameters, PostprocessParameter);
 
 	/** Initialization constructor. */
 	FPostProcessDofNearVS_ES2(const ShaderMetaType::CompiledShaderInitializerType& Initializer)
@@ -2222,20 +2077,12 @@ public:
 
 	void SetVS(const FRenderingCompositePassContext& Context)
 	{
-		FRHIVertexShader* ShaderRHI = GetVertexShader();
+		FRHIVertexShader* ShaderRHI = Context.RHICmdList.GetBoundVertexShader();
 
 		FGlobalShader::SetParameters<FViewUniformShaderParameters>(Context.RHICmdList, ShaderRHI, Context.View.ViewUniformBuffer);
 
 		PostprocessParameter.SetVS(ShaderRHI, Context, TStaticSamplerState<SF_Point,AM_Clamp,AM_Clamp,AM_Clamp>::GetRHI());
 	}
-
-	// FShader interface.
-	virtual bool Serialize(FArchive& Ar) override
-	{
-		bool bShaderHasOutdatedParameters = FGlobalShader::Serialize(Ar);
-		Ar << PostprocessParameter;
-		return bShaderHasOutdatedParameters;
-	}
 };
 
 template<uint32 UseSun>
@@ -2257,7 +2104,7 @@ class FPostProcessDofNearPS_ES2 : public FGlobalShader
 	FPostProcessDofNearPS_ES2() {}
 
 public:
-	FPostProcessPassParameters PostprocessParameter;
+	LAYOUT_FIELD(FPostProcessPassParameters, PostprocessParameter);
 
 	FPostProcessDofNearPS_ES2(const ShaderMetaType::CompiledShaderInitializerType& Initializer)
 		: FGlobalShader(Initializer)
@@ -2267,18 +2114,11 @@ public:
 
 	void SetPS(const FRenderingCompositePassContext& Context)
 	{
-		FRHIPixelShader* ShaderRHI = GetPixelShader();
+		FRHIPixelShader* ShaderRHI = Context.RHICmdList.GetBoundPixelShader();
 		const FPostProcessSettings& Settings = Context.View.FinalPostProcessSettings;
 		FGlobalShader::SetParameters<FViewUniformShaderParameters>(Context.RHICmdList, ShaderRHI, Context.View.ViewUniformBuffer);
 		PostprocessParameter.SetPS(Context.RHICmdList, ShaderRHI, Context, TStaticSamplerState<SF_Point, AM_Clamp, AM_Clamp, AM_Clamp>::GetRHI());
 	}
-
-	virtual bool Serialize(FArchive& Ar) override
-	{
-		bool bShaderHasOutdatedParameters = FGlobalShader::Serialize(Ar);
-		Ar << PostprocessParameter;
-		return bShaderHasOutdatedParameters;
-	}
 };
 
 IMPLEMENT_SHADER_TYPE(,FPostProcessDofNearVS_ES2,TEXT("/Engine/Private/PostProcessMobile.usf"),TEXT("DofNearVS_ES2"),SF_Vertex);
@@ -2301,8 +2141,8 @@ static void DofNear_SetShader(const FRenderingCompositePassContext& Context)
 	TShaderMapRef<FPostProcessDofNearPS_ES2<UseSun> > PixelShader(Context.GetShaderMap());
 
 	GraphicsPSOInit.BoundShaderState.VertexDeclarationRHI = GFilterVertexDeclaration.VertexDeclarationRHI;
-	GraphicsPSOInit.BoundShaderState.VertexShaderRHI = GETSAFERHISHADER_VERTEX(*VertexShader);
-	GraphicsPSOInit.BoundShaderState.PixelShaderRHI = GETSAFERHISHADER_PIXEL(*PixelShader);
+	GraphicsPSOInit.BoundShaderState.VertexShaderRHI = VertexShader.GetVertexShader();
+	GraphicsPSOInit.BoundShaderState.PixelShaderRHI = PixelShader.GetPixelShader();
 	GraphicsPSOInit.PrimitiveType = PT_TriangleList;
 
 	SetGraphicsPipelineState(Context.RHICmdList, GraphicsPSOInit);
@@ -2372,7 +2212,7 @@ void FRCPassPostProcessDofNearES2::Process(FRenderingCompositePassContext& Conte
 			DstX, DstY,
 			SrcDstSize,
 			SrcSize,
-			*VertexShader,
+			VertexShader,
 			EDRF_UseTriangleOptimization);
 	}
 	Context.RHICmdList.EndRenderPass();
@@ -2416,7 +2256,7 @@ class FPostProcessDofBlurPS_ES2 : public FGlobalShader
 	FPostProcessDofBlurPS_ES2() {}
 
 public:
-	FPostProcessPassParameters PostprocessParameter;
+	LAYOUT_FIELD(FPostProcessPassParameters, PostprocessParameter);
 
 	FPostProcessDofBlurPS_ES2(const ShaderMetaType::CompiledShaderInitializerType& Initializer)
 		: FGlobalShader(Initializer)
@@ -2426,18 +2266,11 @@ public:
 
 	void SetPS(const FRenderingCompositePassContext& Context)
 	{
-		FRHIPixelShader* ShaderRHI = GetPixelShader();
+		FRHIPixelShader* ShaderRHI = Context.RHICmdList.GetBoundPixelShader();
 		const FPostProcessSettings& Settings = Context.View.FinalPostProcessSettings;
 		FGlobalShader::SetParameters<FViewUniformShaderParameters>(Context.RHICmdList, ShaderRHI, Context.View.ViewUniformBuffer);
 		PostprocessParameter.SetPS(Context.RHICmdList, ShaderRHI, Context, TStaticSamplerState<SF_Bilinear,AM_Clamp,AM_Clamp,AM_Clamp>::GetRHI());
 	}
-
-	virtual bool Serialize(FArchive& Ar) override
-	{
-		bool bShaderHasOutdatedParameters = FGlobalShader::Serialize(Ar);
-		Ar << PostprocessParameter;
-		return bShaderHasOutdatedParameters;
-	}
 };
 
 IMPLEMENT_SHADER_TYPE(,FPostProcessDofBlurPS_ES2,TEXT("/Engine/Private/PostProcessMobile.usf"),TEXT("DofBlurPS_ES2"),SF_Pixel);
@@ -2455,7 +2288,7 @@ class FPostProcessDofBlurVS_ES2 : public FGlobalShader
 	FPostProcessDofBlurVS_ES2(){}
 
 public:
-	FPostProcessPassParameters PostprocessParameter;
+	LAYOUT_FIELD(FPostProcessPassParameters, PostprocessParameter);
 
 	FPostProcessDofBlurVS_ES2(const ShaderMetaType::CompiledShaderInitializerType& Initializer)
 		: FGlobalShader(Initializer)
@@ -2465,17 +2298,10 @@ public:
 
 	void SetVS(const FRenderingCompositePassContext& Context)
 	{
-		FRHIVertexShader* ShaderRHI = GetVertexShader();
+		FRHIVertexShader* ShaderRHI = Context.RHICmdList.GetBoundVertexShader();
 		FGlobalShader::SetParameters<FViewUniformShaderParameters>(Context.RHICmdList, ShaderRHI, Context.View.ViewUniformBuffer);
 		PostprocessParameter.SetVS(ShaderRHI, Context, TStaticSamplerState<SF_Point,AM_Clamp,AM_Clamp,AM_Clamp>::GetRHI());
 	}
-
-	virtual bool Serialize(FArchive& Ar) override
-	{
-		bool bShaderHasOutdatedParameters = FGlobalShader::Serialize(Ar);
-		Ar << PostprocessParameter;
-		return bShaderHasOutdatedParameters;
-	}
 };
 
 IMPLEMENT_SHADER_TYPE(,FPostProcessDofBlurVS_ES2,TEXT("/Engine/Private/PostProcessMobile.usf"),TEXT("DofBlurVS_ES2"),SF_Vertex);
@@ -2523,8 +2349,8 @@ void FRCPassPostProcessDofBlurES2::Process(FRenderingCompositePassContext& Conte
 		TShaderMapRef<FPostProcessDofBlurPS_ES2> PixelShader(Context.GetShaderMap());
 
 		GraphicsPSOInit.BoundShaderState.VertexDeclarationRHI = GFilterVertexDeclaration.VertexDeclarationRHI;
-		GraphicsPSOInit.BoundShaderState.VertexShaderRHI = GETSAFERHISHADER_VERTEX(*VertexShader);
-		GraphicsPSOInit.BoundShaderState.PixelShaderRHI = GETSAFERHISHADER_PIXEL(*PixelShader);
+		GraphicsPSOInit.BoundShaderState.VertexShaderRHI = VertexShader.GetVertexShader();
+		GraphicsPSOInit.BoundShaderState.PixelShaderRHI = PixelShader.GetPixelShader();
 		GraphicsPSOInit.PrimitiveType = PT_TriangleList;
 
 		SetGraphicsPipelineState(Context.RHICmdList, GraphicsPSOInit);
@@ -2542,7 +2368,7 @@ void FRCPassPostProcessDofBlurES2::Process(FRenderingCompositePassContext& Conte
 			DstX, DstY,
 			SrcDstSize,
 			SrcDstSize,
-			*VertexShader,
+			VertexShader,
 			EDRF_UseTriangleOptimization);
 	}
 	Context.RHICmdList.EndRenderPass();
@@ -2588,7 +2414,7 @@ class FPostProcessSunAvgPS_ES2 : public FGlobalShader
 	FPostProcessSunAvgPS_ES2() {}
 
 public:
-	FPostProcessPassParameters PostprocessParameter;
+	LAYOUT_FIELD(FPostProcessPassParameters, PostprocessParameter);
 
 	FPostProcessSunAvgPS_ES2(const ShaderMetaType::CompiledShaderInitializerType& Initializer)
 		: FGlobalShader(Initializer)
@@ -2598,18 +2424,11 @@ public:
 
 	void SetPS(const FRenderingCompositePassContext& Context)
 	{
-		FRHIPixelShader* ShaderRHI = GetPixelShader();
+		FRHIPixelShader* ShaderRHI = Context.RHICmdList.GetBoundPixelShader();
 		const FPostProcessSettings& Settings = Context.View.FinalPostProcessSettings;
 		FGlobalShader::SetParameters<FViewUniformShaderParameters>(Context.RHICmdList, ShaderRHI, Context.View.ViewUniformBuffer);
 		PostprocessParameter.SetPS(Context.RHICmdList, ShaderRHI, Context, TStaticSamplerState<SF_Bilinear,AM_Clamp,AM_Clamp,AM_Clamp>::GetRHI());
 	}
-
-	virtual bool Serialize(FArchive& Ar) override
-	{
-		bool bShaderHasOutdatedParameters = FGlobalShader::Serialize(Ar);
-		Ar << PostprocessParameter;
-		return bShaderHasOutdatedParameters;
-	}
 };
 
 IMPLEMENT_SHADER_TYPE(,FPostProcessSunAvgPS_ES2,TEXT("/Engine/Private/PostProcessMobile.usf"),TEXT("SunAvgPS_ES2"),SF_Pixel);
@@ -2628,7 +2447,7 @@ class FPostProcessSunAvgVS_ES2 : public FGlobalShader
 	FPostProcessSunAvgVS_ES2(){}
 
 public:
-	FPostProcessPassParameters PostprocessParameter;
+	LAYOUT_FIELD(FPostProcessPassParameters, PostprocessParameter);
 
 	FPostProcessSunAvgVS_ES2(const ShaderMetaType::CompiledShaderInitializerType& Initializer)
 		: FGlobalShader(Initializer)
@@ -2638,17 +2457,10 @@ public:
 
 	void SetVS(const FRenderingCompositePassContext& Context)
 	{
-		FRHIVertexShader* ShaderRHI = GetVertexShader();
+		FRHIVertexShader* ShaderRHI = Context.RHICmdList.GetBoundVertexShader();
 		FGlobalShader::SetParameters<FViewUniformShaderParameters>(Context.RHICmdList, ShaderRHI, Context.View.ViewUniformBuffer);
 		PostprocessParameter.SetVS(ShaderRHI, Context, TStaticSamplerState<SF_Bilinear,AM_Clamp,AM_Clamp,AM_Clamp>::GetRHI());
 	}
-
-	virtual bool Serialize(FArchive& Ar) override
-	{
-		bool bShaderHasOutdatedParameters = FGlobalShader::Serialize(Ar);
-		Ar << PostprocessParameter;
-		return bShaderHasOutdatedParameters;
-	}
 };
 
 IMPLEMENT_SHADER_TYPE(,FPostProcessSunAvgVS_ES2,TEXT("/Engine/Private/PostProcessMobile.usf"),TEXT("SunAvgVS_ES2"),SF_Vertex);
@@ -2667,8 +2479,8 @@ static void SunAvg_SetShader(const FRenderingCompositePassContext& Context)
 	TShaderMapRef<FPostProcessSunAvgPS_ES2> PixelShader(Context.GetShaderMap());
 
 	GraphicsPSOInit.BoundShaderState.VertexDeclarationRHI = GFilterVertexDeclaration.VertexDeclarationRHI;
-	GraphicsPSOInit.BoundShaderState.VertexShaderRHI = GETSAFERHISHADER_VERTEX(*VertexShader);
-	GraphicsPSOInit.BoundShaderState.PixelShaderRHI = GETSAFERHISHADER_PIXEL(*PixelShader);
+	GraphicsPSOInit.BoundShaderState.VertexShaderRHI = VertexShader.GetVertexShader();
+	GraphicsPSOInit.BoundShaderState.PixelShaderRHI = PixelShader.GetPixelShader();
 	GraphicsPSOInit.PrimitiveType = PT_TriangleList;
 
 	SetGraphicsPipelineState(Context.RHICmdList, GraphicsPSOInit);
@@ -2727,7 +2539,7 @@ void FRCPassPostProcessSunAvgES2::Process(FRenderingCompositePassContext& Contex
 			DstX, DstY,
 			SrcDstSize,
 			SrcDstSize,
-			*VertexShader,
+			VertexShader,
 			EDRF_UseTriangleOptimization);
 	}
 	Context.RHICmdList.EndRenderPass();
@@ -2772,8 +2584,8 @@ class FPostProcessAaPS_ES2 : public FGlobalShader
 	FPostProcessAaPS_ES2() {}
 
 public:
-	FPostProcessPassParameters PostprocessParameter;
-	FShaderParameter AaBlendAmount;
+	LAYOUT_FIELD(FPostProcessPassParameters, PostprocessParameter);
+	LAYOUT_FIELD(FShaderParameter, AaBlendAmount);
 
 	FPostProcessAaPS_ES2(const ShaderMetaType::CompiledShaderInitializerType& Initializer)
 		: FGlobalShader(Initializer)
@@ -2784,7 +2596,7 @@ public:
 
 	void SetPS(const FRenderingCompositePassContext& Context)
 	{
-		FRHIPixelShader* ShaderRHI = GetPixelShader();
+		FRHIPixelShader* ShaderRHI = Context.RHICmdList.GetBoundPixelShader();
 		const FPostProcessSettings& Settings = Context.View.FinalPostProcessSettings;
 		FGlobalShader::SetParameters<FViewUniformShaderParameters>(Context.RHICmdList, ShaderRHI, Context.View.ViewUniformBuffer);
 		PostprocessParameter.SetPS(Context.RHICmdList, ShaderRHI, Context, TStaticSamplerState<SF_Bilinear,AM_Clamp,AM_Clamp,AM_Clamp>::GetRHI());
@@ -2854,13 +2666,6 @@ public:
 			SetShaderValue(Context.RHICmdList, ShaderRHI, AaBlendAmount, BlendAmount);
 		}
 	}
-
-	virtual bool Serialize(FArchive& Ar) override
-	{
-		bool bShaderHasOutdatedParameters = FGlobalShader::Serialize(Ar);
-		Ar << PostprocessParameter << AaBlendAmount;
-		return bShaderHasOutdatedParameters;
-	}
 };
 
 IMPLEMENT_SHADER_TYPE(,FPostProcessAaPS_ES2,TEXT("/Engine/Private/PostProcessMobile.usf"),TEXT("AaPS_ES2"),SF_Pixel);
@@ -2879,7 +2684,7 @@ class FPostProcessAaVS_ES2 : public FGlobalShader
 	FPostProcessAaVS_ES2(){}
 
 public:
-	FPostProcessPassParameters PostprocessParameter;
+	LAYOUT_FIELD(FPostProcessPassParameters, PostprocessParameter);
 
 	FPostProcessAaVS_ES2(const ShaderMetaType::CompiledShaderInitializerType& Initializer)
 		: FGlobalShader(Initializer)
@@ -2889,17 +2694,10 @@ public:
 
 	void SetVS(const FRenderingCompositePassContext& Context)
 	{
-		FRHIVertexShader* ShaderRHI = GetVertexShader();
+		FRHIVertexShader* ShaderRHI = Context.RHICmdList.GetBoundVertexShader();
 		FGlobalShader::SetParameters<FViewUniformShaderParameters>(Context.RHICmdList, ShaderRHI, Context.View.ViewUniformBuffer);
 		PostprocessParameter.SetVS(ShaderRHI, Context, TStaticSamplerState<SF_Bilinear,AM_Clamp,AM_Clamp,AM_Clamp>::GetRHI());
 	}
-
-	virtual bool Serialize(FArchive& Ar) override
-	{
-		bool bShaderHasOutdatedParameters = FGlobalShader::Serialize(Ar);
-		Ar << PostprocessParameter;
-		return bShaderHasOutdatedParameters;
-	}
 };
 
 IMPLEMENT_SHADER_TYPE(,FPostProcessAaVS_ES2,TEXT("/Engine/Private/PostProcessMobile.usf"),TEXT("AaVS_ES2"),SF_Vertex);
@@ -2918,8 +2716,8 @@ static void Aa_SetShader(const FRenderingCompositePassContext& Context)
 	TShaderMapRef<FPostProcessAaPS_ES2> PixelShader(Context.GetShaderMap());
 
 	GraphicsPSOInit.BoundShaderState.VertexDeclarationRHI = GFilterVertexDeclaration.VertexDeclarationRHI;
-	GraphicsPSOInit.BoundShaderState.VertexShaderRHI = GETSAFERHISHADER_VERTEX(*VertexShader);
-	GraphicsPSOInit.BoundShaderState.PixelShaderRHI = GETSAFERHISHADER_PIXEL(*PixelShader);
+	GraphicsPSOInit.BoundShaderState.VertexShaderRHI = VertexShader.GetVertexShader();
+	GraphicsPSOInit.BoundShaderState.PixelShaderRHI = PixelShader.GetPixelShader();
 	GraphicsPSOInit.PrimitiveType = PT_TriangleList;
 
 	SetGraphicsPipelineState(Context.RHICmdList, GraphicsPSOInit);
@@ -2994,7 +2792,7 @@ void FRCPassPostProcessAaES2::Process(FRenderingCompositePassContext& Context)
 			Width, Height,
 			DestSize,
 			SrcSize,
-			*VertexShader,
+			VertexShader,
 			EDRF_UseTriangleOptimization);
 	}
 	Context.RHICmdList.EndRenderPass();
diff --git a/Engine/Source/Runtime/Renderer/Private/PostProcess/PostProcessMobile.h b/Engine/Source/Runtime/Renderer/Private/PostProcess/PostProcessMobile.h
--- a/Engine/Source/Runtime/Renderer/Private/PostProcess/PostProcessMobile.h
+++ b/Engine/Source/Runtime/Renderer/Private/PostProcess/PostProcessMobile.h
@@ -152,7 +152,7 @@ public:
 	virtual void Release() override { delete this; }
 private:
 	FIntPoint PrePostSourceViewportSize;
-	FShader* SetShader(const FRenderingCompositePassContext& Context);
+	TShaderRef<FShader> SetShader(const FRenderingCompositePassContext& Context);
 };
 
 class FRCPassPostProcessSunMergeSmallES2 : public TRenderingCompositePassBase<2, 1>
diff --git a/Engine/Source/Runtime/Renderer/Private/PostProcess/PostProcessMotionBlur.cpp b/Engine/Source/Runtime/Renderer/Private/PostProcess/PostProcessMotionBlur.cpp
--- a/Engine/Source/Runtime/Renderer/Private/PostProcess/PostProcessMotionBlur.cpp
+++ b/Engine/Source/Runtime/Renderer/Private/PostProcess/PostProcessMotionBlur.cpp
@@ -488,7 +488,7 @@ void AddMotionBlurVelocityPass(
 		FComputeShaderUtils::AddPass(
 			GraphBuilder,
 			RDG_EVENT_NAME("Velocity Flatten"),
-			*ComputeShader,
+			ComputeShader,
 			PassParameters,
 			FComputeShaderUtils::GetGroupCount(Viewports.Velocity.Rect.Size(), FMotionBlurVelocityFlattenCS::ThreadGroupSize));
 	}
@@ -542,8 +542,8 @@ void AddMotionBlurVelocityPass(
 		TShaderMapRef<FMotionBlurVelocityDilateScatterVS> VertexShader(View.ShaderMap);
 		TShaderMapRef<FMotionBlurVelocityDilateScatterPS> PixelShader(View.ShaderMap);
 		
-		ValidateShaderParameters(*VertexShader, *PassParameters);
-		ValidateShaderParameters(*PixelShader, *PassParameters);
+		ValidateShaderParameters(VertexShader, *PassParameters);
+		ValidateShaderParameters(PixelShader, *PassParameters);
 
 		GraphBuilder.AddPass(
 			RDG_EVENT_NAME("VelocityTileScatter %dx%d", VelocityTileCount.X, VelocityTileCount.Y),
@@ -551,12 +551,12 @@ void AddMotionBlurVelocityPass(
 			ERDGPassFlags::Raster,
 			[VertexShader, PixelShader, VelocityTileCount, PassParameters](FRHICommandListImmediate& RHICmdList)
 		{
-			FRHIVertexShader* RHIVertexShader = GETSAFERHISHADER_VERTEX(*VertexShader);
+			FRHIVertexShader* RHIVertexShader = VertexShader.GetVertexShader();
 
 			FGraphicsPipelineStateInitializer GraphicsPSOInit;
 			GraphicsPSOInit.BoundShaderState.VertexDeclarationRHI = GEmptyVertexDeclaration.VertexDeclarationRHI;
 			GraphicsPSOInit.BoundShaderState.VertexShaderRHI = RHIVertexShader;
-			GraphicsPSOInit.BoundShaderState.PixelShaderRHI = GETSAFERHISHADER_PIXEL(*PixelShader);
+			GraphicsPSOInit.BoundShaderState.PixelShaderRHI = PixelShader.GetPixelShader();
 			GraphicsPSOInit.PrimitiveType = PT_TriangleList;
 			GraphicsPSOInit.RasterizerState = TStaticRasterizerState<>::GetRHI();
 			RHICmdList.ApplyCachedRenderTargets(GraphicsPSOInit);
@@ -584,7 +584,7 @@ void AddMotionBlurVelocityPass(
 
 				PassParameters->ScatterPass = ScatterPassIndex;
 
-				SetShaderParameters(RHICmdList, *VertexShader, RHIVertexShader, *PassParameters);
+				SetShaderParameters(RHICmdList, VertexShader, RHIVertexShader, *PassParameters);
 
 				// Needs to be the same on shader side (faster on NVIDIA and AMD)
 				const int32 QuadsPerInstance = 8;
@@ -604,7 +604,7 @@ void AddMotionBlurVelocityPass(
 		FComputeShaderUtils::AddPass(
 			GraphBuilder,
 			RDG_EVENT_NAME("VelocityTileGatherCS %dx%d", VelocityTileCount.X, VelocityTileCount.Y),
-			*ComputeShader,
+			ComputeShader,
 			PassParameters,
 			FComputeShaderUtils::GetGroupCount(VelocityTileCount, FMotionBlurVelocityDilateGatherCS::ThreadGroupSize));
 	}
@@ -701,7 +701,7 @@ FRDGTextureRef AddMotionBlurFilterPass(
 		FComputeShaderUtils::AddPass(
 			GraphBuilder,
 			RDG_EVENT_NAME("Motion Blur %dx%d (CS)", Viewports.Color.Rect.Width(), Viewports.Color.Rect.Height()),
-			*ComputeShader,
+			ComputeShader,
 			PassParameters,
 			FComputeShaderUtils::GetGroupCount(Viewports.Color.Rect.Size(), FComputeShaderUtils::kGolden2DGroupSize));
 	}
@@ -719,7 +719,7 @@ FRDGTextureRef AddMotionBlurFilterPass(
 			View,
 			Viewports.Color,
 			Viewports.Color,
-			*PixelShader,
+			PixelShader,
 			PassParameters,
 			EScreenPassDrawFlags::AllowHMDHiddenAreaMask);
 	}
@@ -765,7 +765,7 @@ FRDGTextureRef AddVisualizeMotionBlurPass(
 
 	TShaderMapRef<FMotionBlurVisualizePS> PixelShader(View.ShaderMap);
 
-	AddDrawScreenPass(GraphBuilder, RDG_EVENT_NAME("Visualizer"), View, ColorViewport, ColorViewport, *PixelShader, PassParameters);
+	AddDrawScreenPass(GraphBuilder, RDG_EVENT_NAME("Visualizer"), View, ColorViewport, ColorViewport, PixelShader, PassParameters);
 
 	Output.LoadAction = ERenderTargetLoadAction::ELoad;
 
diff --git a/Engine/Source/Runtime/Renderer/Private/PostProcess/PostProcessNoiseBlur.cpp b/Engine/Source/Runtime/Renderer/Private/PostProcess/PostProcessNoiseBlur.cpp
--- a/Engine/Source/Runtime/Renderer/Private/PostProcess/PostProcessNoiseBlur.cpp
+++ b/Engine/Source/Runtime/Renderer/Private/PostProcess/PostProcessNoiseBlur.cpp
@@ -35,9 +35,9 @@ class FPostProcessNoiseBlurPS : public FGlobalShader
 	FPostProcessNoiseBlurPS() {}
 
 public:
-	FPostProcessPassParameters PostprocessParameter;
-	FSceneTextureShaderParameters SceneTextureParameters;
-	FShaderParameter NoiseParams;
+	LAYOUT_FIELD(FPostProcessPassParameters, PostprocessParameter);
+	LAYOUT_FIELD(FSceneTextureShaderParameters, SceneTextureParameters);
+	LAYOUT_FIELD(FShaderParameter, NoiseParams);
 
 	/** Initialization constructor. */
 	FPostProcessNoiseBlurPS(const ShaderMetaType::CompiledShaderInitializerType& Initializer)
@@ -48,18 +48,10 @@ public:
 		NoiseParams.Bind(Initializer.ParameterMap, TEXT("NoiseParams"));
 	}
 
-	// FShader interface.
-	virtual bool Serialize(FArchive& Ar) override
-	{
-		bool bShaderHasOutdatedParameters = FGlobalShader::Serialize(Ar);
-		Ar << PostprocessParameter << SceneTextureParameters << NoiseParams;
-		return bShaderHasOutdatedParameters;
-	}
-
 	template <typename TRHICmdList>
 	void SetParameters(TRHICmdList& RHICmdList, const FRenderingCompositePassContext& Context, float InRadius)
 	{
-		FRHIPixelShader* ShaderRHI = GetPixelShader();
+		FRHIPixelShader* ShaderRHI = RHICmdList.GetBoundPixelShader();
 
 		FGlobalShader::SetParameters<FViewUniformShaderParameters>(RHICmdList, ShaderRHI, Context.View.ViewUniformBuffer);
 		SceneTextureParameters.Set(RHICmdList, ShaderRHI, Context.View.FeatureLevel, ESceneTextureSetupMode::All);
@@ -111,8 +103,8 @@ static void SetNoiseBlurShader(const FRenderingCompositePassContext& Context, fl
 	TShaderMapRef<FPostProcessNoiseBlurPS<Method> > PixelShader(Context.GetShaderMap());
 
 	GraphicsPSOInit.BoundShaderState.VertexDeclarationRHI = GFilterVertexDeclaration.VertexDeclarationRHI;
-	GraphicsPSOInit.BoundShaderState.VertexShaderRHI = GETSAFERHISHADER_VERTEX(*VertexShader);
-	GraphicsPSOInit.BoundShaderState.PixelShaderRHI = GETSAFERHISHADER_PIXEL(*PixelShader);
+	GraphicsPSOInit.BoundShaderState.VertexShaderRHI = VertexShader.GetVertexShader();
+	GraphicsPSOInit.BoundShaderState.PixelShaderRHI = PixelShader.GetPixelShader();
 	GraphicsPSOInit.PrimitiveType = PT_TriangleList;
 
 	SetGraphicsPipelineState(Context.RHICmdList, GraphicsPSOInit);
@@ -184,7 +176,7 @@ void FRCPassPostProcessNoiseBlur::Process(FRenderingCompositePassContext& Contex
 			SrcRect.Width(), SrcRect.Height(),
 			DestSize,
 			SrcSize,
-			*VertexShader,
+			VertexShader,
 			View.StereoPass,
 			Context.HasHmdMesh(),
 			EDRF_UseTriangleOptimization);
diff --git a/Engine/Source/Runtime/Renderer/Private/PostProcess/PostProcessSelectionOutline.cpp b/Engine/Source/Runtime/Renderer/Private/PostProcess/PostProcessSelectionOutline.cpp
--- a/Engine/Source/Runtime/Renderer/Private/PostProcess/PostProcessSelectionOutline.cpp
+++ b/Engine/Source/Runtime/Renderer/Private/PostProcess/PostProcessSelectionOutline.cpp
@@ -172,7 +172,7 @@ FScreenPassTexture AddSelectionOutlinePass(FRDGBuilder& GraphBuilder, const FVie
 			View,
 			OutputViewport,
 			ColorViewport,
-			*PixelShader,
+			PixelShader,
 			PassParameters);
 	}
 
diff --git a/Engine/Source/Runtime/Renderer/Private/PostProcess/PostProcessSubsurface.cpp b/Engine/Source/Runtime/Renderer/Private/PostProcess/PostProcessSubsurface.cpp
--- a/Engine/Source/Runtime/Renderer/Private/PostProcess/PostProcessSubsurface.cpp
+++ b/Engine/Source/Runtime/Renderer/Private/PostProcess/PostProcessSubsurface.cpp
@@ -638,7 +638,7 @@ FRDGTextureRef ResolveTextureToSRV(FRDGBuilder& GraphBuilder, FRDGTextureRef Inp
 
 	TShaderMapRef<FSubsurfaceSRVResolvePS> PixelShader(View.ShaderMap);
 
-	AddDrawScreenPass(GraphBuilder, RDG_EVENT_NAME("SubsurfaceTextureResolve"), View, SceneViewport, SceneViewport, *PixelShader, PassParameters);
+	AddDrawScreenPass(GraphBuilder, RDG_EVENT_NAME("SubsurfaceTextureResolve"), View, SceneViewport, SceneViewport, PixelShader, PassParameters);
 
 	return SRVTextureOutput;
 }
@@ -657,7 +657,7 @@ FRDGTextureRef CreateBlackUAVTexture(FRDGBuilder& GraphBuilder, FRDGTextureDesc
 
 	TShaderMapRef<FSubsurfaceSRVResolvePS> PixelShader(View.ShaderMap);
 
-	AddDrawScreenPass(GraphBuilder, RDG_EVENT_NAME("ClearUAV"), View, SceneViewport, SceneViewport, *PixelShader, PassParameters);
+	AddDrawScreenPass(GraphBuilder, RDG_EVENT_NAME("ClearUAV"), View, SceneViewport, SceneViewport, PixelShader, PassParameters);
 #else
 	FRDGTextureRef SRVTextureOutput = GraphBuilder.CreateTexture(SRVDesc, Name);
 	FRDGTextureUAVDesc UAVClearDesc(SRVTextureOutput, 0);
@@ -866,7 +866,7 @@ void ComputeSubsurfaceForView(
 			SHADER::FParameters* PassParameters = GraphBuilder.AllocParameters<SHADER::FParameters>();
 			PassParameters->RWBurleyGroupBuffer = GraphBuilder.CreateUAV(BurleyGroupBuffer, EPixelFormat::PF_R32_UINT);
 			PassParameters->RWSeparableGroupBuffer = GraphBuilder.CreateUAV(SeparableGroupBuffer, EPixelFormat::PF_R32_UINT);
-			FComputeShaderUtils::AddPass(GraphBuilder, RDG_EVENT_NAME("InitGroupCounter"), *ComputeShader, PassParameters, FIntVector(1, 1, 1));
+			FComputeShaderUtils::AddPass(GraphBuilder, RDG_EVENT_NAME("InitGroupCounter"), ComputeShader, PassParameters, FIntVector(1, 1, 1));
 		}
 
 		// Call the indirect setup
@@ -893,7 +893,7 @@ void ComputeSubsurfaceForView(
 
 			FIntPoint ComputeGroupCount = FIntPoint::DivideAndRoundUp(SubsurfaceViewport.Extent, kSubsurfaceGroupSize);
 
-			FComputeShaderUtils::AddPass(GraphBuilder, RDG_EVENT_NAME("SubsurfaceSetup"), *ComputeShader, PassParameters, FIntVector(ComputeGroupCount.X, ComputeGroupCount.Y, 1));
+			FComputeShaderUtils::AddPass(GraphBuilder, RDG_EVENT_NAME("SubsurfaceSetup"), ComputeShader, PassParameters, FIntVector(ComputeGroupCount.X, ComputeGroupCount.Y, 1));
 		}
 
 		// In half resolution, only Separable is used. We do not need this mipmap.
@@ -926,7 +926,7 @@ void ComputeSubsurfaceForView(
 				PassParameters->GroupBuffer = GraphBuilder.CreateSRV(SubsurfaceBufferUsage[SubsurfaceTypeIndex], EPixelFormat::PF_R32_UINT);
 
 				TShaderMapRef<ARGSETUPSHADER> ComputeShader(View.ShaderMap);
-				FComputeShaderUtils::AddPass(GraphBuilder, FRDGEventName(SubsurfacePhaseName[SubsurfaceTypeIndex]), *ComputeShader, PassParameters, FIntVector(1, 1, 1));
+				FComputeShaderUtils::AddPass(GraphBuilder, FRDGEventName(SubsurfacePhaseName[SubsurfaceTypeIndex]), ComputeShader, PassParameters, FIntVector(1, 1, 1));
 			}
 		}
 
@@ -1008,7 +1008,7 @@ void ComputeSubsurfaceForView(
 				ComputeShaderPermutationVector.Set<SHADER::FRunningInSeparable>(bForceRunningInSeparable);
 				TShaderMapRef<SHADER> ComputeShader(View.ShaderMap, ComputeShaderPermutationVector);
 
-				FComputeShaderUtils::AddPass(GraphBuilder, FRDGEventName(PassInfo.Name), *ComputeShader, PassParameters, SubsurfaceBufferArgs[SubsurfaceTypeIndex], 0);
+				FComputeShaderUtils::AddPass(GraphBuilder, FRDGEventName(PassInfo.Name), ComputeShader, PassParameters, SubsurfaceBufferArgs[SubsurfaceTypeIndex], 0);
 			}
 		}
 	}
@@ -1049,7 +1049,7 @@ void ComputeSubsurfaceForView(
 			View,
 			SceneViewport,
 			SceneViewport,
-			*PixelShader,
+			PixelShader,
 			PassParameters,
 			EScreenPassDrawFlags::AllowHMDHiddenAreaMask);
 	}
@@ -1113,7 +1113,7 @@ FRDGTextureRef ComputeSubsurface(
 					const FViewInfo& View = Views[ViewIndex];
 					const FScreenPassTextureViewport TextureViewport(SceneTexture, View.ViewRect);
 
-					DrawScreenPass(RHICmdList, View, TextureViewport, TextureViewport, *PixelShader, *PassParameters);
+					DrawScreenPass(RHICmdList, View, TextureViewport, TextureViewport, PixelShader, *PassParameters);
 				}
 			}
 		});
@@ -1170,7 +1170,7 @@ FScreenPassTexture AddVisualizeSubsurfacePass(FRDGBuilder& GraphBuilder, const F
 
 	RDG_EVENT_SCOPE(GraphBuilder, "VisualizeSubsurface");
 
-	AddDrawScreenPass(GraphBuilder, RDG_EVENT_NAME("Visualizer"), View, FScreenPassTextureViewport(Output), InputViewport, *PixelShader, PassParameters);
+	AddDrawScreenPass(GraphBuilder, RDG_EVENT_NAME("Visualizer"), View, FScreenPassTextureViewport(Output), InputViewport, PixelShader, PassParameters);
 
 	Output.LoadAction = ERenderTargetLoadAction::ELoad;
 
diff --git a/Engine/Source/Runtime/Renderer/Private/PostProcess/PostProcessTestImage.cpp b/Engine/Source/Runtime/Renderer/Private/PostProcess/PostProcessTestImage.cpp
--- a/Engine/Source/Runtime/Renderer/Private/PostProcess/PostProcessTestImage.cpp
+++ b/Engine/Source/Runtime/Renderer/Private/PostProcess/PostProcessTestImage.cpp
@@ -51,7 +51,7 @@ void AddTestImagePass(FRDGBuilder& GraphBuilder, const FViewInfo& View, FScreenP
 		ERDGPassFlags::Raster,
 		[&View, OutputTexture, OutputViewport, PixelShader, PassParameters](FRHICommandListImmediate& RHICmdList)
 	{
-		DrawScreenPass(RHICmdList, View, OutputViewport, OutputViewport, *PixelShader, *PassParameters);
+		DrawScreenPass(RHICmdList, View, OutputViewport, OutputViewport, PixelShader, *PassParameters);
 
 		// Draw debug text
 		{
diff --git a/Engine/Source/Runtime/Renderer/Private/PostProcess/PostProcessTonemap.cpp b/Engine/Source/Runtime/Renderer/Private/PostProcess/PostProcessTonemap.cpp
--- a/Engine/Source/Runtime/Renderer/Private/PostProcess/PostProcessTonemap.cpp
+++ b/Engine/Source/Runtime/Renderer/Private/PostProcess/PostProcessTonemap.cpp
@@ -839,7 +839,7 @@ FScreenPassTexture AddTonemapPass(FRDGBuilder& GraphBuilder, const FViewInfo& Vi
 		FComputeShaderUtils::AddPass(
 			GraphBuilder,
 			RDG_EVENT_NAME("Tonemap %dx%d (CS GammaOnly=%d)", OutputViewport.Rect.Width(), OutputViewport.Rect.Height(), Inputs.bGammaOnly),
-			*ComputeShader,
+			ComputeShader,
 			PassParameters,
 			FComputeShaderUtils::GetGroupCount(OutputViewport.Rect.Size(), FIntPoint(GTonemapComputeTileSizeX, GTonemapComputeTileSizeY)));
 	}
@@ -870,13 +870,13 @@ FScreenPassTexture AddTonemapPass(FRDGBuilder& GraphBuilder, const FViewInfo& Vi
 			View,
 			OutputViewport,
 			SceneColorViewport,
-			FScreenPassPipelineState(*VertexShader, *PixelShader, BlendState, DepthStencilState),
+			FScreenPassPipelineState(VertexShader, PixelShader, BlendState, DepthStencilState),
 			DrawFlags,
 			PassParameters,
 			[VertexShader, PixelShader, PassParameters](FRHICommandList& RHICmdList)
 		{
-			SetShaderParameters(RHICmdList, *VertexShader, VertexShader->GetVertexShader(), PassParameters->Tonemap);
-			SetShaderParameters(RHICmdList, *PixelShader, PixelShader->GetPixelShader(), *PassParameters);
+			SetShaderParameters(RHICmdList, VertexShader, VertexShader.GetVertexShader(), PassParameters->Tonemap);
+			SetShaderParameters(RHICmdList, PixelShader, PixelShader.GetPixelShader(), *PassParameters);
 		});
 	}
 
@@ -1211,27 +1211,27 @@ class FPostProcessTonemapPS_ES2 : public FGlobalShader
 	FPostProcessTonemapPS_ES2() {}
 
 public:
-	FPostProcessPassParameters PostprocessParameter;
-	FShaderParameter ColorScale0;
-	FShaderParameter ColorScale1;
-	FShaderParameter TexScale;
-	FShaderParameter GrainScaleBiasJitter;
-	FShaderParameter InverseGamma;
-	FShaderParameter TonemapperParams;
-
-	FShaderParameter ColorMatrixR_ColorCurveCd1;
-	FShaderParameter ColorMatrixG_ColorCurveCd3Cm3;
-	FShaderParameter ColorMatrixB_ColorCurveCm2;
-	FShaderParameter ColorCurve_Cm0Cd0_Cd2_Ch0Cm1_Ch3;
-	FShaderParameter ColorCurve_Ch1_Ch2;
-	FShaderParameter ColorShadow_Luma;
-	FShaderParameter ColorShadow_Tint1;
-	FShaderParameter ColorShadow_Tint2;
-
-	FShaderParameter OverlayColor;
-	FShaderParameter FringeIntensity;
-	FShaderParameter SRGBAwareTargetParam;
-	FShaderParameter DefaultEyeExposure;
+	LAYOUT_FIELD(FPostProcessPassParameters, PostprocessParameter);
+	LAYOUT_FIELD(FShaderParameter, ColorScale0);
+	LAYOUT_FIELD(FShaderParameter, ColorScale1);
+	LAYOUT_FIELD(FShaderParameter, TexScale);
+	LAYOUT_FIELD(FShaderParameter, GrainScaleBiasJitter);
+	LAYOUT_FIELD(FShaderParameter, InverseGamma);
+	LAYOUT_FIELD(FShaderParameter, TonemapperParams);
+
+	LAYOUT_FIELD(FShaderParameter, ColorMatrixR_ColorCurveCd1);
+	LAYOUT_FIELD(FShaderParameter, ColorMatrixG_ColorCurveCd3Cm3);
+	LAYOUT_FIELD(FShaderParameter, ColorMatrixB_ColorCurveCm2);
+	LAYOUT_FIELD(FShaderParameter, ColorCurve_Cm0Cd0_Cd2_Ch0Cm1_Ch3);
+	LAYOUT_FIELD(FShaderParameter, ColorCurve_Ch1_Ch2);
+	LAYOUT_FIELD(FShaderParameter, ColorShadow_Luma);
+	LAYOUT_FIELD(FShaderParameter, ColorShadow_Tint1);
+	LAYOUT_FIELD(FShaderParameter, ColorShadow_Tint2);
+
+	LAYOUT_FIELD(FShaderParameter, OverlayColor);
+	LAYOUT_FIELD(FShaderParameter, FringeIntensity);
+	LAYOUT_FIELD(FShaderParameter, SRGBAwareTargetParam);
+	LAYOUT_FIELD(FShaderParameter, DefaultEyeExposure);
 
 	/** Initialization constructor. */
 	FPostProcessTonemapPS_ES2(const ShaderMetaType::CompiledShaderInitializerType& Initializer)
@@ -1262,28 +1262,13 @@ public:
 		DefaultEyeExposure.Bind(Initializer.ParameterMap, TEXT("DefaultEyeExposure"));
 	}
 	
-	// FShader interface.
-	virtual bool Serialize(FArchive& Ar) override
-	{
-		bool bShaderHasOutdatedParameters = FGlobalShader::Serialize(Ar);
-		Ar  << PostprocessParameter << ColorScale0 << ColorScale1 << InverseGamma
-			<< TexScale << GrainScaleBiasJitter << TonemapperParams
-			<< ColorMatrixR_ColorCurveCd1 << ColorMatrixG_ColorCurveCd3Cm3 << ColorMatrixB_ColorCurveCm2 << ColorCurve_Cm0Cd0_Cd2_Ch0Cm1_Ch3 << ColorCurve_Ch1_Ch2 << ColorShadow_Luma << ColorShadow_Tint1 << ColorShadow_Tint2
-			<< OverlayColor
-			<< FringeIntensity
-			<< SRGBAwareTargetParam
-			<< DefaultEyeExposure;
-
-		return bShaderHasOutdatedParameters;
-	}
-
 	template <typename TRHICmdList>
 	void SetPS(TRHICmdList& RHICmdList, const FRenderingCompositePassContext& Context, const FPermutationDomain& PermutationVector, bool bSRGBAwareTarget)
 	{
 		const FPostProcessSettings& Settings = Context.View.FinalPostProcessSettings;
 		const FSceneViewFamily& ViewFamily = *(Context.View.Family);
 
-		FRHIPixelShader* ShaderRHI = GetPixelShader();
+		FRHIPixelShader* ShaderRHI = Context.RHICmdList.GetBoundPixelShader();
 		
 		FGlobalShader::SetParameters<FViewUniformShaderParameters>(RHICmdList, ShaderRHI, Context.View.ViewUniformBuffer);
 
@@ -1402,13 +1387,13 @@ class FPostProcessTonemapVS_ES2 : public FGlobalShader
 	FPostProcessTonemapVS_ES2() { }
 
 public:
-	FPostProcessPassParameters PostprocessParameter;
-	FShaderResourceParameter EyeAdaptation;
-	FShaderParameter GrainRandomFull;
-	FShaderParameter FringeIntensity;
-	FShaderParameter ScreenPosToViewportBias;
-	FShaderParameter ScreenPosToViewportScale;
-	bool bUsedFramebufferFetch;
+	LAYOUT_FIELD(FPostProcessPassParameters, PostprocessParameter);
+	LAYOUT_FIELD(FShaderResourceParameter, EyeAdaptation);
+	LAYOUT_FIELD(FShaderParameter, GrainRandomFull);
+	LAYOUT_FIELD(FShaderParameter, FringeIntensity);
+	LAYOUT_FIELD(FShaderParameter, ScreenPosToViewportBias);
+	LAYOUT_FIELD(FShaderParameter, ScreenPosToViewportScale);
+	LAYOUT_FIELD(bool, bUsedFramebufferFetch);
 
 	FPostProcessTonemapVS_ES2(const ShaderMetaType::CompiledShaderInitializerType& Initializer)
 		: FGlobalShader(Initializer)
@@ -1422,7 +1407,7 @@ public:
 
 	void SetVS(const FRenderingCompositePassContext& Context)
 	{
-		FRHIVertexShader* ShaderRHI = GetVertexShader();
+		FRHIVertexShader* ShaderRHI = Context.RHICmdList.GetBoundVertexShader();
 		FGlobalShader::SetParameters<FViewUniformShaderParameters>(Context.RHICmdList, ShaderRHI, Context.View.ViewUniformBuffer);
 
 		PostprocessParameter.SetVS(ShaderRHI, Context, TStaticSamplerState<SF_Bilinear,AM_Clamp,AM_Clamp,AM_Clamp>::GetRHI());
@@ -1457,13 +1442,6 @@ public:
 			SetShaderValue(Context.RHICmdList, ShaderRHI, ScreenPosToViewportBias, FVector2D(ScreenPosToScenePixelValue.Z, ScreenPosToScenePixelValue.W));
 		}
 	}
-	
-	virtual bool Serialize(FArchive& Ar) override
-	{
-		bool bShaderHasOutdatedParameters = FGlobalShader::Serialize(Ar);
-		Ar << PostprocessParameter << GrainRandomFull << FringeIntensity << ScreenPosToViewportBias << ScreenPosToViewportScale;
-		return bShaderHasOutdatedParameters;
-	}
 };
 
 IMPLEMENT_GLOBAL_SHADER(FPostProcessTonemapVS_ES2, "/Engine/Private/PostProcessTonemap.usf", "MainVS_ES2", SF_Vertex);
@@ -1530,8 +1508,8 @@ void FRCPassPostProcessTonemapES2::Process(FRenderingCompositePassContext& Conte
 			VertexShader->bUsedFramebufferFetch = bUsedFramebufferFetch;
 
 			GraphicsPSOInit.BoundShaderState.VertexDeclarationRHI = GFilterVertexDeclaration.VertexDeclarationRHI;
-			GraphicsPSOInit.BoundShaderState.VertexShaderRHI = GETSAFERHISHADER_VERTEX(*VertexShader);
-			GraphicsPSOInit.BoundShaderState.PixelShaderRHI = GETSAFERHISHADER_PIXEL(*PixelShader);
+			GraphicsPSOInit.BoundShaderState.VertexShaderRHI = VertexShader.GetVertexShader();
+			GraphicsPSOInit.BoundShaderState.PixelShaderRHI = PixelShader.GetPixelShader();
 			GraphicsPSOInit.PrimitiveType = PT_TriangleList;
 
 			SetGraphicsPipelineState(Context.RHICmdList, GraphicsPSOInit);
@@ -1548,7 +1526,7 @@ void FRCPassPostProcessTonemapES2::Process(FRenderingCompositePassContext& Conte
 			SrcRect.Width(), SrcRect.Height(),
 			DstSize,
 			SrcSize,
-			*VertexShader,
+			VertexShader,
 			EDRF_UseTriangleOptimization);
 	}
 	Context.RHICmdList.EndRenderPass();
diff --git a/Engine/Source/Runtime/Renderer/Private/PostProcess/PostProcessUpscale.cpp b/Engine/Source/Runtime/Renderer/Private/PostProcess/PostProcessUpscale.cpp
--- a/Engine/Source/Runtime/Renderer/Private/PostProcess/PostProcessUpscale.cpp
+++ b/Engine/Source/Runtime/Renderer/Private/PostProcess/PostProcessUpscale.cpp
@@ -233,24 +233,23 @@ FScreenPassTexture AddUpscalePass(FRDGBuilder& GraphBuilder, const FViewInfo& Vi
 	{
 		RHICmdList.SetViewport(OutputViewport.Rect.Min.X, OutputViewport.Rect.Min.Y, 0.0f, OutputViewport.Rect.Max.X, OutputViewport.Rect.Max.Y, 1.0f);
 
-		FShader* VertexShader = nullptr;
-
+		TShaderRef<FShader> VertexShader;
 		if (bUsePaniniProjection)
 		{
 			TShaderMapRef<FUpscaleVS> TypedVertexShader(View.ShaderMap);
-			SetScreenPassPipelineState(RHICmdList, FScreenPassPipelineState(*TypedVertexShader, *PixelShader));
-			SetShaderParameters(RHICmdList, *TypedVertexShader, TypedVertexShader->GetVertexShader(), *PassParameters);
-			VertexShader = *TypedVertexShader;
+			SetScreenPassPipelineState(RHICmdList, FScreenPassPipelineState(TypedVertexShader, PixelShader));
+			SetShaderParameters(RHICmdList, TypedVertexShader, TypedVertexShader.GetVertexShader(), *PassParameters);
+			VertexShader = TypedVertexShader;
 		}
 		else
 		{
 			TShaderMapRef<FScreenPassVS> TypedVertexShader(View.ShaderMap);
-			SetScreenPassPipelineState(RHICmdList, FScreenPassPipelineState(*TypedVertexShader, *PixelShader));
-			VertexShader = *TypedVertexShader;
+			SetScreenPassPipelineState(RHICmdList, FScreenPassPipelineState(TypedVertexShader, PixelShader));
+			VertexShader = TypedVertexShader;
 		}
-		check(VertexShader != nullptr);
+		check(VertexShader.IsValid());
 
-		SetShaderParameters(RHICmdList, *PixelShader, PixelShader->GetPixelShader(), *PassParameters);
+		SetShaderParameters(RHICmdList, PixelShader, PixelShader.GetPixelShader(), *PassParameters);
 
 		DrawRectangle(
 			RHICmdList,
diff --git a/Engine/Source/Runtime/Renderer/Private/PostProcess/PostProcessVisualizeBuffer.cpp b/Engine/Source/Runtime/Renderer/Private/PostProcess/PostProcessVisualizeBuffer.cpp
--- a/Engine/Source/Runtime/Renderer/Private/PostProcess/PostProcessVisualizeBuffer.cpp
+++ b/Engine/Source/Runtime/Renderer/Private/PostProcess/PostProcessVisualizeBuffer.cpp
@@ -131,12 +131,12 @@ FScreenPassTexture AddVisualizeBufferPass(FRDGBuilder& GraphBuilder, const FView
 			View,
 			OutputViewport,
 			FScreenPassTextureViewport(Tile.Input),
-			FScreenPassPipelineState(*VertexShader, *PixelShader, BlendState),
+			FScreenPassPipelineState(VertexShader, PixelShader, BlendState),
 			EScreenPassDrawFlags::None,
 			PassParameters,
 			[PixelShader, PassParameters](FRHICommandList& RHICmdList)
 		{
-			SetShaderParameters(RHICmdList, *PixelShader, PixelShader->GetPixelShader(), *PassParameters);
+			SetShaderParameters(RHICmdList, PixelShader, PixelShader.GetPixelShader(), *PassParameters);
 		});
 
 		FTileLabel TileLabel;
diff --git a/Engine/Source/Runtime/Renderer/Private/PostProcess/PostProcessVisualizeComplexity.cpp b/Engine/Source/Runtime/Renderer/Private/PostProcess/PostProcessVisualizeComplexity.cpp
--- a/Engine/Source/Runtime/Renderer/Private/PostProcess/PostProcessVisualizeComplexity.cpp
+++ b/Engine/Source/Runtime/Renderer/Private/PostProcess/PostProcessVisualizeComplexity.cpp
@@ -128,7 +128,7 @@ FScreenPassTexture AddVisualizeComplexityPass(FRDGBuilder& GraphBuilder, const F
 
 	RDG_EVENT_SCOPE(GraphBuilder, "VisualizeComplexity");
 
-	AddDrawScreenPass(GraphBuilder, RDG_EVENT_NAME("Visualizer"), View, FScreenPassTextureViewport(Output), InputViewport, *PixelShader, PassParameters);
+	AddDrawScreenPass(GraphBuilder, RDG_EVENT_NAME("Visualizer"), View, FScreenPassTextureViewport(Output), InputViewport, PixelShader, PassParameters);
 
 	Output.LoadAction = ERenderTargetLoadAction::ELoad;
 
diff --git a/Engine/Source/Runtime/Renderer/Private/PostProcess/PostProcessVisualizeHDR.cpp b/Engine/Source/Runtime/Renderer/Private/PostProcess/PostProcessVisualizeHDR.cpp
--- a/Engine/Source/Runtime/Renderer/Private/PostProcess/PostProcessVisualizeHDR.cpp
+++ b/Engine/Source/Runtime/Renderer/Private/PostProcess/PostProcessVisualizeHDR.cpp
@@ -92,7 +92,7 @@ FScreenPassTexture AddVisualizeHDRPass(FRDGBuilder& GraphBuilder, const FViewInf
 
 	RDG_EVENT_SCOPE(GraphBuilder, "VisualizeComplexity");
 
-	AddDrawScreenPass(GraphBuilder, RDG_EVENT_NAME("Visualizer"), View, OutputViewport, InputViewport, *PixelShader, PassParameters);
+	AddDrawScreenPass(GraphBuilder, RDG_EVENT_NAME("Visualizer"), View, OutputViewport, InputViewport, PixelShader, PassParameters);
 
 	Output.LoadAction = ERenderTargetLoadAction::ELoad;
 
diff --git a/Engine/Source/Runtime/Renderer/Private/PostProcess/PostProcessWeightedSampleSum.cpp b/Engine/Source/Runtime/Renderer/Private/PostProcess/PostProcessWeightedSampleSum.cpp
--- a/Engine/Source/Runtime/Renderer/Private/PostProcess/PostProcessWeightedSampleSum.cpp
+++ b/Engine/Source/Runtime/Renderer/Private/PostProcess/PostProcessWeightedSampleSum.cpp
@@ -396,7 +396,7 @@ FScreenPassTexture AddGaussianBlurPass(
 		FComputeShaderUtils::AddPass(
 			GraphBuilder,
 			RDG_EVENT_NAME("GaussianBlur.%s %dx%d (CS)", Name, OutputViewport.Rect.Width(), OutputViewport.Rect.Height()),
-			*ComputeShader,
+			ComputeShader,
 			PassParameters,
 			FComputeShaderUtils::GetGroupCount(OutputViewport.Rect.Size(), FIntPoint(GFilterComputeTileSizeX, GFilterComputeTileSizeY)));
 	}
@@ -424,13 +424,13 @@ FScreenPassTexture AddGaussianBlurPass(
 				View,
 				OutputViewport,
 				FScreenPassTextureViewport(Filter),
-				FScreenPassPipelineState(*VertexShader, *PixelShader),
+				FScreenPassPipelineState(VertexShader, PixelShader),
 				EScreenPassDrawFlags::None,
 				PassParameters,
 				[VertexShader, PixelShader, PassParameters] (FRHICommandListImmediate& RHICmdList)
 			{
-				SetShaderParameters(RHICmdList, *VertexShader, VertexShader->GetVertexShader(), PassParameters->Filter);
-				SetShaderParameters(RHICmdList, *PixelShader, PixelShader->GetPixelShader(), *PassParameters);
+				SetShaderParameters(RHICmdList, VertexShader, VertexShader.GetVertexShader(), PassParameters->Filter);
+				SetShaderParameters(RHICmdList, PixelShader, PixelShader.GetPixelShader(), *PassParameters);
 			});
 		}
 		else
@@ -441,7 +441,7 @@ FScreenPassTexture AddGaussianBlurPass(
 				View,
 				OutputViewport,
 				FScreenPassTextureViewport(Filter),
-				*PixelShader,
+				PixelShader,
 				PassParameters);
 		}
 	}
diff --git a/Engine/Source/Runtime/Renderer/Private/PostProcess/RenderingCompositionGraph.cpp b/Engine/Source/Runtime/Renderer/Private/PostProcess/RenderingCompositionGraph.cpp
--- a/Engine/Source/Runtime/Renderer/Private/PostProcess/RenderingCompositionGraph.cpp
+++ b/Engine/Source/Runtime/Renderer/Private/PostProcess/RenderingCompositionGraph.cpp
@@ -25,6 +25,8 @@
 #include "ImageWriteStream.h"
 #include "Modules/ModuleManager.h"
 
+IMPLEMENT_TYPE_LAYOUT(FPostProcessPassParameters);
+
 void ExecuteCompositionGraphDebug();
 
 static TAutoConsoleVariable<int32> CVarCompositionGraphOrder(
diff --git a/Engine/Source/Runtime/Renderer/Private/PostProcess/RenderingCompositionGraph.h b/Engine/Source/Runtime/Renderer/Private/PostProcess/RenderingCompositionGraph.h
--- a/Engine/Source/Runtime/Renderer/Private/PostProcess/RenderingCompositionGraph.h
+++ b/Engine/Source/Runtime/Renderer/Private/PostProcess/RenderingCompositionGraph.h
@@ -24,8 +24,6 @@ struct FRenderingCompositeOutputRef;
 struct FRenderingCompositePass;
 struct FRenderingCompositePassContext;
 
-template<typename ShaderMetaType> class TShaderMap;
-
 class FRenderingCompositionGraph
 {
 public:
@@ -206,7 +204,7 @@ struct FRenderingCompositePassContext
 
 	ERHIFeatureLevel::Type GetFeatureLevel() const { return FeatureLevel; }
 	EShaderPlatform GetShaderPlatform() const { return GShaderPlatformForFeatureLevel[FeatureLevel]; }
-	TShaderMap<FGlobalShaderType>* GetShaderMap() const { check(ShaderMap); return ShaderMap; }
+	FGlobalShaderMap* GetShaderMap() const { check(ShaderMap); return ShaderMap; }
 
 	//
 	const FViewInfo& View;
@@ -232,7 +230,7 @@ private:
 	//
 	ERHIFeatureLevel::Type FeatureLevel;
 	//
-	TShaderMap<FGlobalShaderType>* ShaderMap;
+	FGlobalShaderMap* ShaderMap;
 	// to ensure we only process the graph once
 	bool bWasProcessed;
 	// updated once a frame in Process()
diff --git a/Engine/Source/Runtime/Renderer/Private/PostProcess/SceneFilterRendering.cpp b/Engine/Source/Runtime/Renderer/Private/PostProcess/SceneFilterRendering.cpp
--- a/Engine/Source/Runtime/Renderer/Private/PostProcess/SceneFilterRendering.cpp
+++ b/Engine/Source/Runtime/Renderer/Private/PostProcess/SceneFilterRendering.cpp
@@ -13,6 +13,8 @@
 #include "PixelShaderUtils.h"
 #include "CommonRenderResources.h"
 
+IMPLEMENT_TYPE_LAYOUT(FGammaShaderParameters);
+
 void FTesselatedScreenRectangleIndexBuffer::InitRHI()
 {
 	TResourceArray<uint16, INDEXBUFFER_ALIGNMENT> IndexBuffer;
@@ -104,7 +106,7 @@ static inline void InternalDrawRectangle(
 	float SizeV,
 	FIntPoint TargetSize,
 	FIntPoint TextureSize,
-	FShader* VertexShader,
+	const TShaderRef<FShader>& VertexShader,
 	EDrawRectangleFlags Flags,
 	uint32 InstanceCount
 	)
@@ -131,7 +133,7 @@ static inline void InternalDrawRectangle(
 		1.0f / TargetSize.X, 1.0f / TargetSize.Y,
 		1.0f / TextureSize.X, 1.0f / TextureSize.Y);
 
-	SetUniformBufferParameterImmediate(RHICmdList, VertexShader->GetVertexShader(), VertexShader->GetUniformBufferParameter<FDrawRectangleParameters>(), Parameters);
+	SetUniformBufferParameterImmediate(RHICmdList, VertexShader.GetVertexShader(), VertexShader->GetUniformBufferParameter<FDrawRectangleParameters>(), Parameters);
 
 	if(Flags == EDRF_UseTesselatedIndexBuffer)
 	{
@@ -173,7 +175,7 @@ void DrawRectangle(
 	float SizeV,
 	FIntPoint TargetSize,
 	FIntPoint TextureSize,
-	FShader* VertexShader,
+	const TShaderRef<FShader>& VertexShader,
 	EDrawRectangleFlags Flags,
 	uint32 InstanceCount
 )
@@ -245,7 +247,7 @@ void DrawHmdMesh(
 	FIntPoint TargetSize,
 	FIntPoint TextureSize,
 	EStereoscopicPass StereoView,
-	FShader* VertexShader
+	const TShaderRef<FShader>& VertexShader
 	)
 {
 	FDrawRectangleParameters Parameters;
@@ -256,7 +258,7 @@ void DrawHmdMesh(
 		1.0f / TargetSize.X, 1.0f / TargetSize.Y,
 		1.0f / TextureSize.X, 1.0f / TextureSize.Y);
 
-	SetUniformBufferParameterImmediate(RHICmdList, VertexShader->GetVertexShader(), VertexShader->GetUniformBufferParameter<FDrawRectangleParameters>(), Parameters);
+	SetUniformBufferParameterImmediate(RHICmdList, VertexShader.GetVertexShader(), VertexShader->GetUniformBufferParameter<FDrawRectangleParameters>(), Parameters);
 
 	if (GEngine->XRSystem->GetHMDDevice())
 	{
@@ -276,7 +278,7 @@ void DrawPostProcessPass(
 	float SizeV,
 	FIntPoint TargetSize,
 	FIntPoint TextureSize,
-	FShader* VertexShader,
+	const TShaderRef<FShader>& VertexShader,
 	EStereoscopicPass StereoView,
 	bool bHasCustomMesh,
 	EDrawRectangleFlags Flags)
diff --git a/Engine/Source/Runtime/Renderer/Private/PostProcess/SceneFilterRendering.h b/Engine/Source/Runtime/Renderer/Private/PostProcess/SceneFilterRendering.h
--- a/Engine/Source/Runtime/Renderer/Private/PostProcess/SceneFilterRendering.h
+++ b/Engine/Source/Runtime/Renderer/Private/PostProcess/SceneFilterRendering.h
@@ -50,7 +50,7 @@ extern RENDERER_API void DrawRectangle(
 	float SizeV,
 	FIntPoint TargetSize,
 	FIntPoint TextureSize,
-	class FShader* VertexShader,
+	const TShaderRef<FShader>& VertexShader,
 	EDrawRectangleFlags Flags = EDRF_Default,
 	uint32 InstanceCount = 1
 	);
@@ -86,7 +86,7 @@ extern RENDERER_API void DrawHmdMesh(
 	FIntPoint TargetSize,
 	FIntPoint TextureSize,
 	EStereoscopicPass StereoView,
-	FShader* VertexShader
+	const TShaderRef<FShader>& VertexShader
 	);
 
 // NOTE: Assumes previously set PSO has PrimitiveType = PT_TriangleList
@@ -102,7 +102,7 @@ extern RENDERER_API void DrawPostProcessPass(
 	float SizeV,
 	FIntPoint TargetSize,
 	FIntPoint TextureSize,
-	class FShader* VertexShader,
+	const TShaderRef<FShader>& VertexShader,
 	EStereoscopicPass StereoView,
 	bool bHasCustomMesh,
 	EDrawRectangleFlags Flags = EDRF_Default
@@ -117,6 +117,7 @@ FMGammaShaderParameters
 /** Encapsulates the gamma correction parameters. */
 class FGammaShaderParameters
 {
+	DECLARE_TYPE_LAYOUT(FGammaShaderParameters, NonVirtual);
 public:
 
 	/** Default constructor. */
@@ -131,7 +132,7 @@ public:
 	}
 
 	/** Set the material shader parameter values. */
-	void Set(FRHICommandList& RHICmdList, FShader* PixelShader, float DisplayGamma, FLinearColor const& ColorScale, FLinearColor const& ColorOverlay)
+	void Set(FRHICommandList& RHICmdList, FRHIPixelShader* RHIShader, float DisplayGamma, FLinearColor const& ColorScale, FLinearColor const& ColorOverlay)
 	{
 		// GammaColorScaleAndInverse
 
@@ -147,7 +148,7 @@ public:
 
 		SetShaderValue(
 			RHICmdList,
-			PixelShader->GetPixelShader(),
+			RHIShader,
 			GammaColorScaleAndInverse,
 			ColorScaleAndInverse
 			);
@@ -163,7 +164,7 @@ public:
 
 		SetShaderValue(
 			RHICmdList,
-			PixelShader->GetPixelShader(),
+			RHIShader,
 			GammaOverlayColor,
 			OverlayColor
 			);
@@ -178,7 +179,7 @@ public:
 
 		SetShaderValue(
 			RHICmdList,
-			PixelShader->GetPixelShader(),
+			RHIShader,
 			RenderTargetExtent,
 			vRenderTargetExtent);
 	}
@@ -195,9 +196,11 @@ public:
 
 
 private:
-	FShaderParameter				GammaColorScaleAndInverse;
-	FShaderParameter				GammaOverlayColor;
-	FShaderParameter				RenderTargetExtent;
+	
+		LAYOUT_FIELD(FShaderParameter, GammaColorScaleAndInverse)
+		LAYOUT_FIELD(FShaderParameter, GammaOverlayColor)
+		LAYOUT_FIELD(FShaderParameter, RenderTargetExtent)
+	
 };
 
 
diff --git a/Engine/Source/Runtime/Renderer/Private/PostProcess/SceneRenderTargets.cpp b/Engine/Source/Runtime/Renderer/Private/PostProcess/SceneRenderTargets.cpp
--- a/Engine/Source/Runtime/Renderer/Private/PostProcess/SceneRenderTargets.cpp
+++ b/Engine/Source/Runtime/Renderer/Private/PostProcess/SceneRenderTargets.cpp
@@ -29,6 +29,8 @@
 #include "VisualizeTexture.h"
 #include "GpuDebugRendering.h"
 
+IMPLEMENT_TYPE_LAYOUT(FSceneTextureShaderParameters);
+
 static TAutoConsoleVariable<int32> CVarRSMResolution(
 	TEXT("r.LPV.RSMResolution"),
 	360,
@@ -243,7 +245,7 @@ FSceneRenderTargets::FSceneRenderTargets(const FViewInfo& View, const FSceneRend
 	, LightAccumulation(GRenderTargetPool.MakeSnapshot(SnapshotSource.LightAccumulation))
 	, SceneColorSubPixel(GRenderTargetPool.MakeSnapshot(SnapshotSource.SceneColorSubPixel))
 	, DirectionalOcclusion(GRenderTargetPool.MakeSnapshot(SnapshotSource.DirectionalOcclusion))
-	, SceneDepthZ(GRenderTargetPool.MakeSnapshot(SnapshotSource.SceneDepthZ))
+	, SceneDepthZ(GRenderTargetPool.MakeSnapshot(SnapshotSource.SceneDepthZ))	
 	, SceneVelocity(GRenderTargetPool.MakeSnapshot(SnapshotSource.SceneVelocity))
 	, LightingChannels(GRenderTargetPool.MakeSnapshot(SnapshotSource.LightingChannels))
 	, SceneAlphaCopy(GRenderTargetPool.MakeSnapshot(SnapshotSource.SceneAlphaCopy))
@@ -1503,7 +1505,7 @@ void FSceneRenderTargets::BeginRenderingPrePass(FRHICommandList& RHICmdList, boo
 		// Note, this is a reversed Z depth surface, so 0.0f is the far plane.
 		ERenderTargetActions StencilAction = bStencilClear ? ERenderTargetActions::Clear_Store : ERenderTargetActions::Load_Store;
 		RPInfo.DepthStencilRenderTarget.Action = MakeDepthStencilTargetActions(ERenderTargetActions::Clear_Store, StencilAction);
-		bSceneDepthCleared = true;
+		bSceneDepthCleared = true;	
 	}
 	else
 	{
@@ -1877,25 +1879,25 @@ void FSceneRenderTargets::ResolveDepthTexture(FRHICommandList& RHICmdList, const
 		{
 		case 2:
 			TextureIndex = ResolvePixelShader2X->UnresolvedSurface.GetBaseIndex();
-			ResolvePixelShader = GETSAFERHISHADER_PIXEL(*ResolvePixelShader2X);
+			ResolvePixelShader = ResolvePixelShader2X.GetPixelShader();
 			break;
 		case 4:
 			TextureIndex = ResolvePixelShader4X->UnresolvedSurface.GetBaseIndex();
-			ResolvePixelShader = GETSAFERHISHADER_PIXEL(*ResolvePixelShader4X);
+			ResolvePixelShader = ResolvePixelShader4X.GetPixelShader();
 			break;
 		case 8:
 			TextureIndex = ResolvePixelShader8X->UnresolvedSurface.GetBaseIndex();
-			ResolvePixelShader = GETSAFERHISHADER_PIXEL(*ResolvePixelShader8X);
+			ResolvePixelShader = ResolvePixelShader8X.GetPixelShader();
 			break;
 		default:
 			ensureMsgf(false, TEXT("Unsupported depth resolve for samples: %i.  Dynamic loop method isn't supported on all platforms.  Please add specific case."), SourceTexture->GetNumSamples());
 			TextureIndex = ResolvePixelShaderAny->UnresolvedSurface.GetBaseIndex();
-			ResolvePixelShader = GETSAFERHISHADER_PIXEL(*ResolvePixelShaderAny);
+			ResolvePixelShader = ResolvePixelShaderAny.GetPixelShader();
 			break;
 		}
 
 		GraphicsPSOInit.BoundShaderState.VertexDeclarationRHI = GScreenVertexDeclaration.VertexDeclarationRHI;
-		GraphicsPSOInit.BoundShaderState.VertexShaderRHI = GETSAFERHISHADER_VERTEX(*ResolveVertexShader);
+		GraphicsPSOInit.BoundShaderState.VertexShaderRHI = ResolveVertexShader.GetVertexShader();
 		GraphicsPSOInit.BoundShaderState.PixelShaderRHI = ResolvePixelShader;
 		GraphicsPSOInit.PrimitiveType = PT_TriangleStrip;
 
@@ -2258,7 +2260,7 @@ void FSceneRenderTargets::AllocateCommonDepthTargets(FRHICommandList& RHICmdList
 		if (SceneDepthZ && bHMDAllocatedDepthTarget)
 		{
 			const uint32 OldElementSize = SceneDepthZ->ComputeMemorySize();
-
+		
 			{
 				FSceneRenderTargetItem& Item = SceneDepthZ->GetRenderTargetItem();
 
@@ -2296,7 +2298,7 @@ void FSceneRenderTargets::AllocateFoveationTexture(FRHICommandList& RHICmdList)
 {
 	const bool bStereo = GEngine->StereoRenderingDevice.IsValid() && GEngine->StereoRenderingDevice->IsStereoEnabled();
 	IStereoRenderTargetManager* const StereoRenderTargetManager = bStereo ? GEngine->StereoRenderingDevice->GetRenderTargetManager() : nullptr;
-	
+
 	FTexture2DRHIRef Texture;
 	FIntPoint TextureSize;
 
@@ -2405,7 +2407,7 @@ void FSceneRenderTargets::AllocateDeferredShadingPathRenderTargets(FRHICommandLi
 			}
 			GRenderTargetPool.FindFreeElement(RHICmdList, Desc, ScreenSpaceAO, TEXT("ScreenSpaceAO"), true, ERenderTargetTransience::NonTransient);
 		}
-
+		
 		{
 			const int32 TranslucencyLightingVolumeDim = GetTranslucencyLightingVolumeDim();
 
@@ -2609,11 +2611,11 @@ void FSceneRenderTargets::ClearVolumeTextures(FRHICommandList& RHICmdList, ERHIF
 	TShaderMapRef<TOneColorPixelShaderMRT<NumRenderTargets> > PixelShader(ShaderMap);
 
 	GraphicsPSOInit.BoundShaderState.VertexDeclarationRHI = GScreenVertexDeclaration.VertexDeclarationRHI;
-	GraphicsPSOInit.BoundShaderState.VertexShaderRHI = GETSAFERHISHADER_VERTEX(*VertexShader);
+	GraphicsPSOInit.BoundShaderState.VertexShaderRHI = VertexShader.GetVertexShader();
 #if PLATFORM_SUPPORTS_GEOMETRY_SHADERS
-	GraphicsPSOInit.BoundShaderState.GeometryShaderRHI = GETSAFERHISHADER_GEOMETRY(*GeometryShader);
+	GraphicsPSOInit.BoundShaderState.GeometryShaderRHI = GeometryShader.GetGeometryShader();
 #endif
-	GraphicsPSOInit.BoundShaderState.PixelShaderRHI = GETSAFERHISHADER_PIXEL(*PixelShader);
+	GraphicsPSOInit.BoundShaderState.PixelShaderRHI = PixelShader.GetPixelShader();
 	GraphicsPSOInit.PrimitiveType = PT_TriangleStrip;
 
 	SetGraphicsPipelineState(RHICmdList, GraphicsPSOInit);
diff --git a/Engine/Source/Runtime/Renderer/Private/PostProcess/VisualizeShadingModels.cpp b/Engine/Source/Runtime/Renderer/Private/PostProcess/VisualizeShadingModels.cpp
--- a/Engine/Source/Runtime/Renderer/Private/PostProcess/VisualizeShadingModels.cpp
+++ b/Engine/Source/Runtime/Renderer/Private/PostProcess/VisualizeShadingModels.cpp
@@ -17,7 +17,7 @@ public:
 	{
 		return IsFeatureLevelSupported(Parameters.Platform, ERHIFeatureLevel::SM5);
 	}
-
+	
 	BEGIN_SHADER_PARAMETER_STRUCT(FParameters, )
 		SHADER_PARAMETER_STRUCT_REF(FViewUniformShaderParameters, View)
 		SHADER_PARAMETER_STRUCT_INCLUDE(FSceneTextureParameters, SceneTextures)
@@ -39,7 +39,7 @@ FScreenPassTexture AddVisualizeShadingModelPass(FRDGBuilder& GraphBuilder, const
 {
 	check(Inputs.SceneTextures);
 	check(Inputs.SceneColor.IsValid());
-
+	
 	FScreenPassRenderTarget Output = Inputs.OverrideOutput;
 
 	if (!Output.IsValid())
@@ -69,7 +69,7 @@ FScreenPassTexture AddVisualizeShadingModelPass(FRDGBuilder& GraphBuilder, const
 
 	RDG_EVENT_SCOPE(GraphBuilder, "VisualizeShadingModels");
 
-	AddDrawScreenPass(GraphBuilder, RDG_EVENT_NAME("Visualizer"), View, OutputViewport, InputViewport, *PixelShader, PassParameters);
+	AddDrawScreenPass(GraphBuilder, RDG_EVENT_NAME("Visualizer"), View, OutputViewport, InputViewport, PixelShader, PassParameters);
 
 	Output.LoadAction = ERenderTargetLoadAction::ELoad;
 
@@ -79,7 +79,7 @@ FScreenPassTexture AddVisualizeShadingModelPass(FRDGBuilder& GraphBuilder, const
 		float Y = 28;
 		const float YStep = 14;
 		const float ColumnWidth = 250;
-
+	
 		FString Line;
 
 		Canvas.DrawShadowedString(X, Y += YStep, TEXT("Visualize ShadingModels"), GetStatsFont(), FLinearColor(1, 1, 1));
@@ -113,4 +113,4 @@ FScreenPassTexture AddVisualizeShadingModelPass(FRDGBuilder& GraphBuilder, const
 	});
 
 	return MoveTemp(Output);
-}
\ No newline at end of file
+}
diff --git a/Engine/Source/Runtime/Renderer/Private/PrimitiveDistanceAccuracyRendering.cpp b/Engine/Source/Runtime/Renderer/Private/PrimitiveDistanceAccuracyRendering.cpp
--- a/Engine/Source/Runtime/Renderer/Private/PrimitiveDistanceAccuracyRendering.cpp
+++ b/Engine/Source/Runtime/Renderer/Private/PrimitiveDistanceAccuracyRendering.cpp
@@ -14,7 +14,8 @@ PrimitiveDistanceAccuracyRendering.cpp: Contains definitions for rendering the v
 
 IMPLEMENT_MATERIAL_SHADER_TYPE(,FPrimitiveDistanceAccuracyPS,TEXT("/Engine/Private/PrimitiveDistanceAccuracyPixelShader.usf"),TEXT("Main"),SF_Pixel);
 
-void FPrimitiveDistanceAccuracyPS::GetDebugViewModeShaderBindings(
+void FPrimitiveDistanceAccuracyInterface::GetDebugViewModeShaderBindings(
+	const FDebugViewModePS& ShaderBase,
 	const FPrimitiveSceneProxy* RESTRICT PrimitiveSceneProxy,
 	const FMaterialRenderProxy& RESTRICT MaterialRenderProxy,
 	const FMaterial& RESTRICT Material,
@@ -29,6 +30,7 @@ void FPrimitiveDistanceAccuracyPS::GetDebugViewModeShaderBindings(
 	FMeshDrawSingleShaderBindings& ShaderBindings
 ) const
 {
+	const FPrimitiveDistanceAccuracyPS& Shader = static_cast<const FPrimitiveDistanceAccuracyPS&>(ShaderBase);
 	float CPULogDistance = -1.f;
 #if WITH_EDITORONLY_DATA
 	float Distance = 0;
@@ -38,8 +40,8 @@ void FPrimitiveDistanceAccuracyPS::GetDebugViewModeShaderBindings(
 	}
 #endif
 	// Because the streamer use FMath::FloorToFloat, here we need to use -1 to have a useful result.
-	ShaderBindings.Add(CPULogDistanceParameter, CPULogDistance);
-	ShaderBindings.Add(PrimitiveAlphaParameter,  (!PrimitiveSceneProxy || PrimitiveSceneProxy->IsSelected()) ? 1.f : .2f);
+	ShaderBindings.Add(Shader.CPULogDistanceParameter, CPULogDistance);
+	ShaderBindings.Add(Shader.PrimitiveAlphaParameter,  (!PrimitiveSceneProxy || PrimitiveSceneProxy->IsSelected()) ? 1.f : .2f);
 }
 
 #endif // !(UE_BUILD_SHIPPING || UE_BUILD_TEST)
