Copying Private-LoadTimes-4.24 stream to Main. Biggest changes are in Materials/Shader memory freezing.

diff --git a/Engine/Source/Runtime/Renderer/Private/PrimitiveDistanceAccuracyRendering.h b/Engine/Source/Runtime/Renderer/Private/PrimitiveDistanceAccuracyRendering.h
--- a/Engine/Source/Runtime/Renderer/Private/PrimitiveDistanceAccuracyRendering.h
+++ b/Engine/Source/Runtime/Renderer/Private/PrimitiveDistanceAccuracyRendering.h
@@ -31,7 +31,7 @@ public:
 	static bool ShouldCompilePermutation(const FMeshMaterialShaderPermutationParameters& Parameters)
 	{
 		// See FDebugViewModeMaterialProxy::GetFriendlyName()
-		return AllowDebugViewShaderMode(DVSM_PrimitiveDistanceAccuracy, Parameters.Platform, Parameters.Material->GetFeatureLevel()) && Parameters.Material->GetFriendlyName().Contains(TEXT("PrimitiveDistanceAccuracy"));
+		return AllowDebugViewShaderMode(DVSM_PrimitiveDistanceAccuracy, Parameters.Platform, Parameters.MaterialParameters.FeatureLevel) && Parameters.MaterialParameters.bMaterialIsPrimitiveDistanceAccuracy;
 	}
 
 	FPrimitiveDistanceAccuracyPS(const ShaderMetaType::CompiledShaderInitializerType& Initializer):
@@ -43,20 +43,24 @@ public:
 
 	FPrimitiveDistanceAccuracyPS() {}
 
-	virtual bool Serialize(FArchive& Ar) override
-	{
-		bool bShaderHasOutdatedParameters = FMeshMaterialShader::Serialize(Ar);
-		Ar << CPULogDistanceParameter;
-		Ar << PrimitiveAlphaParameter;
-		return bShaderHasOutdatedParameters;
-	}
-
 	static void ModifyCompilationEnvironment(const FMaterialShaderPermutationParameters& Parameters, FShaderCompilerEnvironment& OutEnvironment)
 	{
 		OutEnvironment.SetDefine(TEXT("UNDEFINED_ACCURACY"), UndefinedStreamingAccuracyIntensity);
 	}
 
+	LAYOUT_FIELD(FShaderParameter, CPULogDistanceParameter);
+	LAYOUT_FIELD(FShaderParameter, PrimitiveAlphaParameter);
+};
+
+class FPrimitiveDistanceAccuracyInterface : public FDebugViewModeInterface
+{
+public:
+
+	FPrimitiveDistanceAccuracyInterface() : FDebugViewModeInterface(TEXT("PrimitiveDistanceAccuracy"), false, false, false) {}
+	virtual TShaderRef<FDebugViewModePS> GetPixelShader(const FMaterial* InMaterial, FVertexFactoryType* VertexFactoryType) const override { return InMaterial->GetShader<FPrimitiveDistanceAccuracyPS>(VertexFactoryType); }
+
 	virtual void GetDebugViewModeShaderBindings(
+		const FDebugViewModePS& ShaderBase,
 		const FPrimitiveSceneProxy* RESTRICT PrimitiveSceneProxy,
 		const FMaterialRenderProxy& RESTRICT MaterialRenderProxy,
 		const FMaterial& RESTRICT Material,
@@ -70,19 +74,6 @@ public:
 		FName ViewModeParamName,
 		FMeshDrawSingleShaderBindings& ShaderBindings
 	) const override;
-
-private:
-
-	FShaderParameter CPULogDistanceParameter;
-	FShaderParameter PrimitiveAlphaParameter;
-};
-
-class FPrimitiveDistanceAccuracyInterface : public FDebugViewModeInterface
-{
-public:
-
-	FPrimitiveDistanceAccuracyInterface() : FDebugViewModeInterface(TEXT("PrimitiveDistanceAccuracy"), false, false, false) {}
-	virtual FDebugViewModePS* GetPixelShader(const FMaterial* InMaterial, FVertexFactoryType* VertexFactoryType) const override { return InMaterial->GetShader<FPrimitiveDistanceAccuracyPS>(VertexFactoryType); }
 };
 
 #endif // !(UE_BUILD_SHIPPING || UE_BUILD_TEST)
diff --git a/Engine/Source/Runtime/Renderer/Private/PrimitiveSceneInfo.cpp b/Engine/Source/Runtime/Renderer/Private/PrimitiveSceneInfo.cpp
--- a/Engine/Source/Runtime/Renderer/Private/PrimitiveSceneInfo.cpp
+++ b/Engine/Source/Runtime/Renderer/Private/PrimitiveSceneInfo.cpp
@@ -17,6 +17,8 @@
 #include "VT/RuntimeVirtualTextureSceneProxy.h"
 #include "VT/VirtualTextureSystem.h"
 #include "GPUScene.h"
+#include "Async/ParallelFor.h"
+#include "ProfilingDebugging/ExternalProfiler.h"
 
 /** An implementation of FStaticPrimitiveDrawInterface that stores the drawn elements for the rendering thread to use. */
 class FBatchingSPDI : public FStaticPrimitiveDrawInterface
@@ -63,7 +65,7 @@ public:
 		{
 			check(Mesh.VertexFactory);
 			check(Mesh.VertexFactory->IsInitialized());
-			checkSlow(IsInRenderingThread());
+			checkSlow(IsInParallelRenderingThread());
 
 			FPrimitiveSceneProxy* PrimitiveSceneProxy = PrimitiveSceneInfo->Proxy;
 			PrimitiveSceneProxy->VerifyUsedMaterial(Mesh.MaterialRenderProxy);
@@ -77,7 +79,8 @@ public:
 			const ERHIFeatureLevel::Type FeatureLevel = PrimitiveSceneInfo->Scene->GetFeatureLevel();
 			StaticMesh->PreparePrimitiveUniformBuffer(PrimitiveSceneProxy, FeatureLevel);
 
-			const bool bSupportsCachingMeshDrawCommands = SupportsCachingMeshDrawCommands(PrimitiveSceneProxy, *StaticMesh, FeatureLevel);
+			// Volumetric self shadow mesh commands need to be generated every frame, as they depend on single frame uniform buffers with self shadow data.
+			const bool bSupportsCachingMeshDrawCommands = SupportsCachingMeshDrawCommands(*StaticMesh, FeatureLevel) && !PrimitiveSceneProxy->CastsVolumetricTranslucentShadow();
 
 			bool bUseSkyMaterial = Mesh.MaterialRenderProxy->GetMaterial(FeatureLevel)->IsSky();
 			bool bUseSingleLayerWaterMaterial = Mesh.MaterialRenderProxy->GetMaterial(FeatureLevel)->GetShadingModels().HasShadingModel(MSM_SingleLayerWater);
@@ -186,7 +189,10 @@ FPrimitiveSceneInfo::FPrimitiveSceneInfo(UPrimitiveComponent* InComponent,FScene
 FPrimitiveSceneInfo::~FPrimitiveSceneInfo()
 {
 	check(!OctreeId.IsValidId());
-	check(StaticMeshCommandInfos.Num() == 0);
+	for (int32 PassIndex = 0; PassIndex < EMeshPass::Num; PassIndex++)
+	{
+		check(StaticMeshCommandInfos.Num() == 0);
+	}
 }
 
 #if RHI_RAYTRACING
@@ -203,138 +209,220 @@ FRHIRayTracingGeometry* FPrimitiveSceneInfo::GetStaticRayTracingGeometryInstance
 }
 #endif
 
-void FPrimitiveSceneInfo::CacheMeshDrawCommands(FRHICommandListImmediate& RHICmdList)
+void FPrimitiveSceneInfo::CacheMeshDrawCommands(FRHICommandListImmediate& RHICmdList, FScene* Scene, const TArrayView<FPrimitiveSceneInfo*>& SceneInfos)
 {
-	check(StaticMeshCommandInfos.Num() == 0);
+	//@todo - only need material uniform buffers to be created since we are going to cache pointers to them
+	// Any updates (after initial creation) don't need to be forced here
+	FMaterialRenderProxy::UpdateDeferredCachedUniformExpressions();
 
-	int32 MeshWithCachedCommandsNum = 0;
-	for (int32 MeshIndex = 0; MeshIndex < StaticMeshes.Num(); MeshIndex++)
-	{
-		const FStaticMeshBatch& Mesh = StaticMeshes[MeshIndex];
-		if (SupportsCachingMeshDrawCommands(Proxy, Mesh))
-		{
-			++MeshWithCachedCommandsNum;
-		}
-	}
+	SCOPED_NAMED_EVENT(FPrimitiveSceneInfo_CacheMeshDrawCommands, FColor::Emerald);
 
-#if RHI_RAYTRACING
-	if (IsRayTracingEnabled() && Proxy->IsRayTracingStaticRelevant())
-	{
-		int MaxLOD = -1;
+	QUICK_SCOPE_CYCLE_COUNTER(STAT_CacheMeshDrawCommands);
+	FMemMark Mark(FMemStack::Get());
 
-		for (int32 MeshIndex = 0; MeshIndex < StaticMeshes.Num(); MeshIndex++)
-		{
-			FStaticMeshBatch& Mesh = StaticMeshes[MeshIndex];
-			MaxLOD = MaxLOD < Mesh.LODIndex ? Mesh.LODIndex : MaxLOD;
-		}
+	static constexpr int BATCH_SIZE = 64;
+	const int NumBatches = (SceneInfos.Num() + BATCH_SIZE - 1) / BATCH_SIZE;
 
-		if (StaticMeshes.Num() > 0)
-		{
-			CachedRayTracingMeshCommandIndicesPerLOD.Empty(MaxLOD + 1);
-			CachedRayTracingMeshCommandIndicesPerLOD.AddDefaulted(MaxLOD + 1);
-		}
-	}
-#endif
-
-	if (MeshWithCachedCommandsNum > 0)
+	auto DoWorkLambda = [Scene, SceneInfos](int32 Index)
 	{
-		//@todo - only need material uniform buffers to be created since we are going to cache pointers to them
-		// Any updates (after initial creation) don't need to be forced here
-		FMaterialRenderProxy::UpdateDeferredCachedUniformExpressions();
+		SCOPED_NAMED_EVENT(FPrimitiveSceneInfo_CacheMeshDrawCommand, FColor::Green);
 
-		// Reserve based on assumption that we have on average 2 cached mesh draw commands per mesh.
-		StaticMeshCommandInfos.Reserve(MeshWithCachedCommandsNum * 2);
-
-		QUICK_SCOPE_CYCLE_COUNTER(STAT_CacheMeshDrawCommands);
-		FMemMark Mark(FMemStack::Get());
+		struct FMeshInfoAndIndex
+		{
+			int32 InfoIndex;
+			int32 MeshIndex;
+		};
 
-		const EShadingPath ShadingPath = Scene->GetShadingPath();
+		TArray<FMeshInfoAndIndex, TMemStackAllocator<>> MeshBatches;
+		MeshBatches.Reserve(3 * BATCH_SIZE);
 
-		for (int32 MeshIndex = 0; MeshIndex < StaticMeshes.Num(); MeshIndex++)
+		int LocalNum = FMath::Min((Index * BATCH_SIZE) + BATCH_SIZE, SceneInfos.Num());
+		for (int LocalIndex = (Index * BATCH_SIZE); LocalIndex < LocalNum; LocalIndex++)
 		{
-			FStaticMeshBatchRelevance& MeshRelevance = StaticMeshRelevances[MeshIndex];
-			FStaticMeshBatch& Mesh = StaticMeshes[MeshIndex];
-
-			check(MeshRelevance.CommandInfosMask.IsEmpty());
-			MeshRelevance.CommandInfosBase = StaticMeshCommandInfos.Num();
+			FPrimitiveSceneInfo* SceneInfo = SceneInfos[LocalIndex];
+			check(SceneInfo->StaticMeshCommandInfos.Num() == 0);
+			SceneInfo->StaticMeshCommandInfos.AddDefaulted(EMeshPass::Num * SceneInfo->StaticMeshes.Num());
+			FPrimitiveSceneProxy* SceneProxy = SceneInfo->Proxy;
 
-			if (SupportsCachingMeshDrawCommands(Proxy, Mesh))
+			// Volumetric self shadow mesh commands need to be generated every frame, as they depend on single frame uniform buffers with self shadow data.
+			if (!SceneProxy->CastsVolumetricTranslucentShadow())
 			{
-				for (int32 PassIndex = 0; PassIndex < EMeshPass::Num; PassIndex++)
+				for (int32 MeshIndex = 0; MeshIndex < SceneInfo->StaticMeshes.Num(); MeshIndex++)
 				{
-					EMeshPass::Type PassType = (EMeshPass::Type)PassIndex;
-
-					if ((FPassProcessorManager::GetPassFlags(ShadingPath, PassType) & EMeshPassFlags::CachedMeshCommands) != EMeshPassFlags::None)
+					FStaticMeshBatch& Mesh = SceneInfo->StaticMeshes[MeshIndex];
+					if (SupportsCachingMeshDrawCommands(Mesh))
 					{
-						FCachedMeshDrawCommandInfo CommandInfo;
-						CommandInfo.MeshPass = PassType;
+						MeshBatches.Add(FMeshInfoAndIndex{ LocalIndex, MeshIndex });
+					}
+				}
+			}
+		}
 
-						FCachedPassMeshDrawList& SceneDrawList = Scene->CachedDrawLists[PassType];
-						FCachedPassMeshDrawListContext CachedPassMeshDrawListContext(CommandInfo, SceneDrawList, *Scene);
+		for (int32 PassIndex = 0; PassIndex < EMeshPass::Num; PassIndex++)
+		{
+			const EShadingPath ShadingPath = Scene->GetShadingPath();
+			EMeshPass::Type PassType = (EMeshPass::Type)PassIndex;
 
-						PassProcessorCreateFunction CreateFunction = FPassProcessorManager::GetCreateFunction(ShadingPath, PassType);
-						FMeshPassProcessor* PassMeshProcessor = CreateFunction(Scene, nullptr, &CachedPassMeshDrawListContext);
+			if ((FPassProcessorManager::GetPassFlags(ShadingPath, PassType) & EMeshPassFlags::CachedMeshCommands) != EMeshPassFlags::None)
+			{
+				FCachedMeshDrawCommandInfo CommandInfo(PassType);
 
-						if (PassMeshProcessor != nullptr)
-						{
-							check(!Mesh.bRequiresPerElementVisibility);
-							uint64 BatchElementMask = ~0ull;
-							PassMeshProcessor->AddMeshBatch(Mesh, BatchElementMask, Proxy);
+				FCriticalSection& CachedMeshDrawCommandLock = Scene->CachedMeshDrawCommandLock[PassType];
+				FCachedPassMeshDrawList& SceneDrawList = Scene->CachedDrawLists[PassType];
+				FStateBucketMap& CachedMeshDrawCommandStateBuckets = Scene->CachedMeshDrawCommandStateBuckets[PassType];
+				FCachedPassMeshDrawListContext CachedPassMeshDrawListContext(CommandInfo, CachedMeshDrawCommandLock, SceneDrawList, CachedMeshDrawCommandStateBuckets, *Scene);
 
-							PassMeshProcessor->~FMeshPassProcessor();
-						}
+				PassProcessorCreateFunction CreateFunction = FPassProcessorManager::GetCreateFunction(ShadingPath, PassType);
+				FMeshPassProcessor* PassMeshProcessor = CreateFunction(Scene, nullptr, &CachedPassMeshDrawListContext);
+
+				if (PassMeshProcessor != nullptr)
+				{
+					for (const FMeshInfoAndIndex& MeshAndInfo : MeshBatches)
+					{
+						FPrimitiveSceneInfo* SceneInfo = SceneInfos[MeshAndInfo.InfoIndex];
+						FStaticMeshBatch& Mesh = SceneInfo->StaticMeshes[MeshAndInfo.MeshIndex];
+						
+						CommandInfo = FCachedMeshDrawCommandInfo(PassType);
+						FStaticMeshBatchRelevance& MeshRelevance = SceneInfo->StaticMeshRelevances[MeshAndInfo.MeshIndex];
+
+						check(!MeshRelevance.CommandInfosMask.Get(PassType));
+
+						check(!Mesh.bRequiresPerElementVisibility);
+						uint64 BatchElementMask = ~0ull;
+						PassMeshProcessor->AddMeshBatch(Mesh, BatchElementMask, SceneInfo->Proxy);
 
 						if (CommandInfo.CommandIndex != -1 || CommandInfo.StateBucketId != -1)
 						{
 							static_assert(sizeof(MeshRelevance.CommandInfosMask) * 8 >= EMeshPass::Num, "CommandInfosMask is too small to contain all mesh passes.");
-
 							MeshRelevance.CommandInfosMask.Set(PassType);
+							MeshRelevance.CommandInfosBase++;
 
-							StaticMeshCommandInfos.Add(CommandInfo);
-
+							int CommandInfoIndex = MeshAndInfo.MeshIndex * EMeshPass::Num + PassType;
+							check(SceneInfo->StaticMeshCommandInfos[CommandInfoIndex].MeshPass == EMeshPass::Num);
+							SceneInfo->StaticMeshCommandInfos[CommandInfoIndex] = CommandInfo;
 #if DO_GUARD_SLOW
 							if (ShadingPath == EShadingPath::Deferred)
 							{
-								FMeshDrawCommand* MeshDrawCommand = CommandInfo.StateBucketId >= 0
-									? &Scene->CachedMeshDrawCommandStateBuckets[FSetElementId::FromInteger(CommandInfo.StateBucketId)].MeshDrawCommand
+								FScopeLock Lock(&CachedMeshDrawCommandLock);
+								const FMeshDrawCommand* MeshDrawCommand = CommandInfo.StateBucketId >= 0
+									? &Scene->CachedMeshDrawCommandStateBuckets[PassType].GetByElementId(CommandInfo.StateBucketId).Key
 									: &SceneDrawList.MeshDrawCommands[CommandInfo.CommandIndex];
 
 								ensureMsgf(MeshDrawCommand->VertexStreams.GetAllocatedSize() == 0, TEXT("Cached Mesh Draw command overflows VertexStreams.  VertexStream inline size should be tweaked."));
-								
+
 								if (PassType == EMeshPass::BasePass || PassType == EMeshPass::DepthPass || PassType == EMeshPass::CSMShadowDepth)
 								{
 									TArray<EShaderFrequency, TInlineAllocator<SF_NumFrequencies>> ShaderFrequencies;
 									MeshDrawCommand->ShaderBindings.GetShaderFrequencies(ShaderFrequencies);
-
+								
+									int32 DataOffset = 0;
 									for (int32 i = 0; i < ShaderFrequencies.Num(); i++)
 									{
-										FMeshDrawSingleShaderBindings SingleShaderBindings = MeshDrawCommand->ShaderBindings.GetSingleShaderBindings(ShaderFrequencies[i]);
-										ensureMsgf(SingleShaderBindings.ParameterMapInfo.LooseParameterBuffers.Num() == 0, TEXT("Cached Mesh Draw command uses loose parameters.  This will break dynamic instancing in performance critical pass.  Use Uniform Buffers instead."));
-										ensureMsgf(SingleShaderBindings.ParameterMapInfo.SRVs.Num() == 0, TEXT("Cached Mesh Draw command uses individual SRVs.  This will break dynamic instancing in performance critical pass.  Use Uniform Buffers instead."));
-										ensureMsgf(SingleShaderBindings.ParameterMapInfo.TextureSamplers.Num() == 0, TEXT("Cached Mesh Draw command uses individual Texture Samplers.  This will break dynamic instancing in performance critical pass.  Use Uniform Buffers instead."));
+										FMeshDrawSingleShaderBindings SingleShaderBindings = const_cast<FMeshDrawCommand*>(MeshDrawCommand)->ShaderBindings.GetSingleShaderBindings(ShaderFrequencies[i], DataOffset);
+										ensureMsgf(SingleShaderBindings.GetParameterMapInfo().LooseParameterBuffers.Num() == 0, TEXT("Cached Mesh Draw command uses loose parameters.  This will break dynamic instancing in performance critical pass.  Use Uniform Buffers instead."));
+										ensureMsgf(SingleShaderBindings.GetParameterMapInfo().SRVs.Num() == 0, TEXT("Cached Mesh Draw command uses individual SRVs.  This will break dynamic instancing in performance critical pass.  Use Uniform Buffers instead."));
+										ensureMsgf(SingleShaderBindings.GetParameterMapInfo().TextureSamplers.Num() == 0, TEXT("Cached Mesh Draw command uses individual Texture Samplers.  This will break dynamic instancing in performance critical pass.  Use Uniform Buffers instead."));
 									}
 								}
 							}
 #endif
 						}
 					}
+					PassMeshProcessor->~FMeshPassProcessor();
 				}
+			}
+		}
 
-			#if RHI_RAYTRACING
-				if (IsRayTracingEnabled() && Proxy->IsRayTracingStaticRelevant())
+		for (int LocalIndex = (Index * BATCH_SIZE); LocalIndex < LocalNum; LocalIndex++)
+		{
+			FPrimitiveSceneInfo* SceneInfo = SceneInfos[LocalIndex];	
+			int PrefixSum = 0;
+			for (int32 MeshIndex = 0; MeshIndex < SceneInfo->StaticMeshes.Num(); MeshIndex++)
+			{
+				FStaticMeshBatchRelevance& MeshRelevance = SceneInfo->StaticMeshRelevances[MeshIndex];
+				if (MeshRelevance.CommandInfosBase > 0)
 				{
-					FCachedRayTracingMeshCommandContext CommandContext(Scene->CachedRayTracingMeshCommands);
-					FRayTracingMeshProcessor RayTracingMeshProcessor(&CommandContext, Scene, nullptr);
+					EMeshPass::Type PassType = EMeshPass::DepthPass;
+					int NewPrefixSum = PrefixSum;
+					for (;;)
+					{
+						PassType = MeshRelevance.CommandInfosMask.SkipEmpty(PassType);
+						if (PassType == EMeshPass::Num)
+						{
+							break;
+						}
 
-					check(!Mesh.bRequiresPerElementVisibility);
-					RayTracingMeshProcessor.AddMeshBatch(Mesh, ~0ull, Proxy);
+						int CommandInfoIndex = MeshIndex * EMeshPass::Num + PassType;
+						checkSlow(CommandInfoIndex >= NewPrefixSum);
+						SceneInfo->StaticMeshCommandInfos[NewPrefixSum] = SceneInfo->StaticMeshCommandInfos[CommandInfoIndex];
+						NewPrefixSum++;
+						PassType = EMeshPass::Type(PassType + 1);
+					}
+
+#if DO_GUARD_SLOW
+					int NumBits = MeshRelevance.CommandInfosMask.GetNum();
+					check(PrefixSum + NumBits == NewPrefixSum);
+					int LastPass = -1;
+					for (int32 TestIndex = PrefixSum; TestIndex < NewPrefixSum; TestIndex++)
+					{
+						int MeshPass = SceneInfo->StaticMeshCommandInfos[TestIndex].MeshPass;
+						check(MeshPass > LastPass);
+						LastPass = MeshPass;
+					}
+#endif
+					MeshRelevance.CommandInfosBase = PrefixSum;
+					PrefixSum = NewPrefixSum;
+				}
+			}
+			SceneInfo->StaticMeshCommandInfos.SetNum(PrefixSum, true);
+		}
+	};
+
+	if (FApp::ShouldUseThreadingForPerformance())
+	{
+		ParallelForTemplate(NumBatches, DoWorkLambda, EParallelForFlags::PumpRenderingThread);
+	}
+	else
+	{
+		for (int Idx = 0; Idx < NumBatches; Idx++)
+		{
+			DoWorkLambda(Idx);
+		}
+	}
+
+	FGraphicsMinimalPipelineStateId::InitializePersistentIds();
 					
-					CachedRayTracingMeshCommandIndicesPerLOD[Mesh.LODIndex].Add(CommandContext.CommandIndex);
+#if RHI_RAYTRACING
+	if (IsRayTracingEnabled())
+	{
+		FCachedRayTracingMeshCommandContext CommandContext(Scene->CachedRayTracingMeshCommands);
+		FRayTracingMeshProcessor RayTracingMeshProcessor(&CommandContext, Scene, nullptr);
+
+		for (FPrimitiveSceneInfo* SceneInfo : SceneInfos)
+		{
+			if (SceneInfo->Proxy->IsRayTracingStaticRelevant())
+			{
+				int MaxLOD = -1;
+				for (int32 MeshIndex = 0; MeshIndex < SceneInfo->StaticMeshes.Num(); MeshIndex++)
+				{
+					FStaticMeshBatch& Mesh = SceneInfo->StaticMeshes[MeshIndex];
+					MaxLOD = MaxLOD < Mesh.LODIndex ? Mesh.LODIndex : MaxLOD;
+
+					check(!Mesh.bRequiresPerElementVisibility);
+					RayTracingMeshProcessor.AddMeshBatch(Mesh, ~0ull, SceneInfo->Proxy);
+
+					SceneInfo->CachedRayTracingMeshCommandIndicesPerLOD[Mesh.LODIndex].Add(CommandContext.CommandIndex);
+				}
+
+				if (SceneInfo->StaticMeshes.Num() > 0)
+				{
+					SceneInfo->CachedRayTracingMeshCommandIndicesPerLOD.Empty(MaxLOD + 1);
+					SceneInfo->CachedRayTracingMeshCommandIndicesPerLOD.AddDefaulted(MaxLOD + 1);
 				}
-			#endif
 			}
 		}
 	}
+#endif
 }
 
 void FPrimitiveSceneInfo::RemoveCachedMeshDrawCommands()
@@ -344,29 +432,31 @@ void FPrimitiveSceneInfo::RemoveCachedMeshDrawCommands()
 	for (int32 CommandIndex = 0; CommandIndex < StaticMeshCommandInfos.Num(); ++CommandIndex)
 	{
 		const FCachedMeshDrawCommandInfo& CachedCommand = StaticMeshCommandInfos[CommandIndex];
-		const FSetElementId StateBucketId = FSetElementId::FromInteger(CachedCommand.StateBucketId);
 
-		if (StateBucketId.IsValidId())
+		if (CachedCommand.StateBucketId != INDEX_NONE)
 		{
-			FMeshDrawCommandStateBucket& StateBucket = Scene->CachedMeshDrawCommandStateBuckets[StateBucketId];
+			EMeshPass::Type PassIndex = CachedCommand.MeshPass;
+			FScopeLock Lock(&Scene->CachedMeshDrawCommandLock[PassIndex]);
 
-			FGraphicsMinimalPipelineStateId::RemovePersistentId(StateBucket.MeshDrawCommand.CachedPipelineId);
+			FGraphicsMinimalPipelineStateId CachedPipelineId = Scene->CachedMeshDrawCommandStateBuckets[PassIndex].GetByElementId(CachedCommand.StateBucketId).Key.CachedPipelineId;
 
-			if (StateBucket.Num == 1)
-			{
-				Scene->CachedMeshDrawCommandStateBuckets.Remove(StateBucketId);
-			}
-			else
+			FMeshDrawCommandCount& StateBucketCount = Scene->CachedMeshDrawCommandStateBuckets[PassIndex].GetByElementId(CachedCommand.StateBucketId).Value;
+			check(StateBucketCount.Num > 0);
+			StateBucketCount.Num--;
+			if (StateBucketCount.Num == 0)
 			{
-				StateBucket.Num--;
+				check(Scene->CachedMeshDrawCommandStateBuckets[PassIndex].RemoveByElementId(CachedCommand.StateBucketId));
 			}
+
+			FGraphicsMinimalPipelineStateId::RemovePersistentId(CachedPipelineId);
 		}
 		else if (CachedCommand.CommandIndex >= 0)
 		{
 			FCachedPassMeshDrawList& PassDrawList = Scene->CachedDrawLists[CachedCommand.MeshPass];
-
-			FGraphicsMinimalPipelineStateId::RemovePersistentId(PassDrawList.MeshDrawCommands[CachedCommand.CommandIndex].CachedPipelineId);
+			FGraphicsMinimalPipelineStateId CachedPipelineId = PassDrawList.MeshDrawCommands[CachedCommand.CommandIndex].CachedPipelineId;
+			
 			PassDrawList.MeshDrawCommands.RemoveAt(CachedCommand.CommandIndex);
+			FGraphicsMinimalPipelineStateId::RemovePersistentId(CachedPipelineId);
 
 			// Track the lowest index that might be free for faster AddAtLowestFreeIndex
 			PassDrawList.LowestFreeIndexSearchStart = FMath::Min(PassDrawList.LowestFreeIndexSearchStart, CachedCommand.CommandIndex);
@@ -378,7 +468,6 @@ void FPrimitiveSceneInfo::RemoveCachedMeshDrawCommands()
 	{
 		FStaticMeshBatchRelevance& MeshRelevance = StaticMeshRelevances[MeshIndex];
 
-		MeshRelevance.CommandInfosBase = 0;
 		MeshRelevance.CommandInfosMask.Reset();
 	}
 
@@ -403,41 +492,54 @@ void FPrimitiveSceneInfo::RemoveCachedMeshDrawCommands()
 #endif
 }
 
-void FPrimitiveSceneInfo::AddStaticMeshes(FRHICommandListImmediate& RHICmdList, bool bAddToStaticDrawLists)
+void FPrimitiveSceneInfo::AddStaticMeshes(FRHICommandListImmediate& RHICmdList, FScene* Scene, const TArrayView<FPrimitiveSceneInfo*>& SceneInfos, bool bAddToStaticDrawLists)
 {
 	LLM_SCOPE(ELLMTag::StaticMesh);
 
-	// Cache the primitive's static mesh elements.
-	FBatchingSPDI BatchingSPDI(this);
-	BatchingSPDI.SetHitProxy(DefaultDynamicHitProxy);
-	Proxy->DrawStaticElements(&BatchingSPDI);
-	StaticMeshes.Shrink();
-	StaticMeshRelevances.Shrink();
-
-	check(StaticMeshRelevances.Num() == StaticMeshes.Num());
+	{
+		ParallelForTemplate(SceneInfos.Num(), [Scene, &SceneInfos](int32 Index)
+		{
+			SCOPED_NAMED_EVENT(FPrimitiveSceneInfo_AddStaticMeshes_DrawStaticElements, FColor::Magenta);
+			FPrimitiveSceneInfo* SceneInfo = SceneInfos[Index];
+			// Cache the primitive's static mesh elements.
+			FBatchingSPDI BatchingSPDI(SceneInfo);
+			BatchingSPDI.SetHitProxy(SceneInfo->DefaultDynamicHitProxy);
+			SceneInfo->Proxy->DrawStaticElements(&BatchingSPDI);
+			SceneInfo->StaticMeshes.Shrink();
+			SceneInfo->StaticMeshRelevances.Shrink();
+
+			check(SceneInfo->StaticMeshRelevances.Num() == SceneInfo->StaticMeshes.Num());
+		});
+	}
 
-	for(int32 MeshIndex = 0;MeshIndex < StaticMeshes.Num();MeshIndex++)
 	{
-		FStaticMeshBatchRelevance& MeshRelevance = StaticMeshRelevances[MeshIndex];
-		FStaticMeshBatch& Mesh = StaticMeshes[MeshIndex];
+		SCOPED_NAMED_EVENT(FPrimitiveSceneInfo_AddStaticMeshes_UpdateSceneArrays, FColor::Blue);
+		for (FPrimitiveSceneInfo* SceneInfo : SceneInfos)
+		{
+			for (int32 MeshIndex = 0; MeshIndex < SceneInfo->StaticMeshes.Num(); MeshIndex++)
+			{
+				FStaticMeshBatchRelevance& MeshRelevance = SceneInfo->StaticMeshRelevances[MeshIndex];
+				FStaticMeshBatch& Mesh = SceneInfo->StaticMeshes[MeshIndex];
 
-		// Add the static mesh to the scene's static mesh list.
-		FSparseArrayAllocationInfo SceneArrayAllocation = Scene->StaticMeshes.AddUninitialized();
-		Scene->StaticMeshes[SceneArrayAllocation.Index] = &Mesh;
-		Mesh.Id = SceneArrayAllocation.Index;
-		MeshRelevance.Id = SceneArrayAllocation.Index;
+				// Add the static mesh to the scene's static mesh list.
+				FSparseArrayAllocationInfo SceneArrayAllocation = Scene->StaticMeshes.AddUninitialized();
+				Scene->StaticMeshes[SceneArrayAllocation.Index] = &Mesh;
+				Mesh.Id = SceneArrayAllocation.Index;
+				MeshRelevance.Id = SceneArrayAllocation.Index;
 
-		if (Mesh.bRequiresPerElementVisibility)
-		{
-			// Use a separate index into StaticMeshBatchVisibility, since most meshes don't use it
-			Mesh.BatchVisibilityId = Scene->StaticMeshBatchVisibility.AddUninitialized().Index;
-			Scene->StaticMeshBatchVisibility[Mesh.BatchVisibilityId] = true;
+				if (Mesh.bRequiresPerElementVisibility)
+				{
+					// Use a separate index into StaticMeshBatchVisibility, since most meshes don't use it
+					Mesh.BatchVisibilityId = Scene->StaticMeshBatchVisibility.AddUninitialized().Index;
+					Scene->StaticMeshBatchVisibility[Mesh.BatchVisibilityId] = true;
+				}
+			}
 		}
 	}
 
 	if (bAddToStaticDrawLists)
 	{
-		CacheMeshDrawCommands(RHICmdList);
+		CacheMeshDrawCommands(RHICmdList, Scene, SceneInfos);
 	}
 }
 
@@ -505,166 +607,221 @@ int32 FPrimitiveSceneInfo::UpdateStaticLightingBuffer()
 	return LCIs.Num();
 }
 
-void FPrimitiveSceneInfo::AddToScene(FRHICommandListImmediate& RHICmdList, bool bUpdateStaticDrawLists, bool bAddToStaticDrawLists, bool bAsyncCreateLPIs)
+void FPrimitiveSceneInfo::AddToScene(FRHICommandListImmediate& RHICmdList, FScene* Scene, const TArrayView<FPrimitiveSceneInfo*>& SceneInfos, bool bUpdateStaticDrawLists, bool bAddToStaticDrawLists, bool bAsyncCreateLPIs)
 {
 	check(IsInRenderingThread());
 
-	// Create an indirect lighting cache uniform buffer if we attaching a primitive that may require it, as it may be stored inside a cached mesh command.
-	if (IsIndirectLightingCacheAllowed(Scene->GetFeatureLevel())
-		&& Proxy->WillEverBeLit()
-		&& ((Proxy->HasStaticLighting() && Proxy->NeedsUnbuiltPreviewLighting()) || (Proxy->IsMovable() && Proxy->GetIndirectLightingCacheQuality() != ILCQ_Off) || Proxy->GetLightmapType() == ELightmapType::ForceVolumetric))
 	{
-		if (!IndirectLightingCacheUniformBuffer)
+		SCOPED_NAMED_EVENT(FPrimitiveSceneInfo_AddToScene_IndirectLightingCacheUniformBuffer, FColor::Turquoise);
+		for (FPrimitiveSceneInfo* SceneInfo : SceneInfos)
 		{
-			FIndirectLightingCacheUniformParameters Parameters;
-
-			GetIndirectLightingCacheParameters(
-				Scene->GetFeatureLevel(), 
-				Parameters, 
-				nullptr,
-				nullptr,
-				FVector(0.0f, 0.0f, 0.0f),
-				0, 
-				nullptr);
-
-			IndirectLightingCacheUniformBuffer = TUniformBufferRef<FIndirectLightingCacheUniformParameters>::CreateUniformBufferImmediate(Parameters, UniformBuffer_MultiFrame, EUniformBufferValidation::None);
+			FPrimitiveSceneProxy* Proxy = SceneInfo->Proxy;
+			// Create an indirect lighting cache uniform buffer if we attaching a primitive that may require it, as it may be stored inside a cached mesh command.
+			if (IsIndirectLightingCacheAllowed(Scene->GetFeatureLevel())
+				&& Proxy->WillEverBeLit()
+				&& ((Proxy->HasStaticLighting() && Proxy->NeedsUnbuiltPreviewLighting()) || (Proxy->IsMovable() && Proxy->GetIndirectLightingCacheQuality() != ILCQ_Off) || Proxy->GetLightmapType() == ELightmapType::ForceVolumetric))
+			{
+				if (!SceneInfo->IndirectLightingCacheUniformBuffer)
+				{
+					FIndirectLightingCacheUniformParameters Parameters;
+
+					GetIndirectLightingCacheParameters(
+						Scene->GetFeatureLevel(),
+						Parameters,
+						nullptr,
+						nullptr,
+						FVector(0.0f, 0.0f, 0.0f),
+						0,
+						nullptr);
+
+					SceneInfo->IndirectLightingCacheUniformBuffer = TUniformBufferRef<FIndirectLightingCacheUniformParameters>::CreateUniformBufferImmediate(Parameters, UniformBuffer_MultiFrame, EUniformBufferValidation::None);
+				}
+			}
 		}
 	}
 	
-	// If we are attaching a primitive that should be statically lit but has unbuilt lighting,
-	// Allocate space in the indirect lighting cache so that it can be used for previewing indirect lighting
-	if (Proxy->HasStaticLighting() 
-		&& Proxy->NeedsUnbuiltPreviewLighting() 
-		&& IsIndirectLightingCacheAllowed(Scene->GetFeatureLevel()))
 	{
-		FIndirectLightingCacheAllocation* PrimitiveAllocation = Scene->IndirectLightingCache.FindPrimitiveAllocation(PrimitiveComponentId);
-
-		if (PrimitiveAllocation)
+		SCOPED_NAMED_EVENT(FPrimitiveSceneInfo_AddToScene_IndirectLightingCacheAllocation, FColor::Orange);
+		for (FPrimitiveSceneInfo* SceneInfo : SceneInfos)
 		{
-			IndirectLightingCacheAllocation = PrimitiveAllocation;
-			PrimitiveAllocation->SetDirty();
+			FPrimitiveSceneProxy* Proxy = SceneInfo->Proxy;
+			// If we are attaching a primitive that should be statically lit but has unbuilt lighting,
+			// Allocate space in the indirect lighting cache so that it can be used for previewing indirect lighting
+			if (Proxy->HasStaticLighting()
+				&& Proxy->NeedsUnbuiltPreviewLighting()
+				&& IsIndirectLightingCacheAllowed(Scene->GetFeatureLevel()))
+			{
+				FIndirectLightingCacheAllocation* PrimitiveAllocation = Scene->IndirectLightingCache.FindPrimitiveAllocation(SceneInfo->PrimitiveComponentId);
+
+				if (PrimitiveAllocation)
+				{
+					SceneInfo->IndirectLightingCacheAllocation = PrimitiveAllocation;
+					PrimitiveAllocation->SetDirty();
+				}
+				else
+				{
+					PrimitiveAllocation = Scene->IndirectLightingCache.AllocatePrimitive(SceneInfo, true);
+					PrimitiveAllocation->SetDirty();
+					SceneInfo->IndirectLightingCacheAllocation = PrimitiveAllocation;
+				}
+			}
+			SceneInfo->MarkIndirectLightingCacheBufferDirty();
 		}
-		else
+	}
+
+	{
+		SCOPED_NAMED_EVENT(FPrimitiveSceneInfo_AddToScene_LightmapDataOffset, FColor::Green);
+		for (FPrimitiveSceneInfo* SceneInfo : SceneInfos)
 		{
-			PrimitiveAllocation = Scene->IndirectLightingCache.AllocatePrimitive(this, true);
-			PrimitiveAllocation->SetDirty();
-			IndirectLightingCacheAllocation = PrimitiveAllocation;
+			const bool bAllowStaticLighting = FReadOnlyCVARCache::Get().bAllowStaticLighting;
+			if (bAllowStaticLighting)
+			{
+				SceneInfo->NumLightmapDataEntries = SceneInfo->UpdateStaticLightingBuffer();
+				if (SceneInfo->NumLightmapDataEntries > 0 && UseGPUScene(GMaxRHIShaderPlatform, Scene->GetFeatureLevel()))
+				{
+					SceneInfo->LightmapDataOffset = Scene->GPUScene.LightmapDataAllocator.Allocate(SceneInfo->NumLightmapDataEntries);
+				}
+			}
 		}
 	}
-	MarkIndirectLightingCacheBufferDirty();
 
-	const bool bAllowStaticLighting = FReadOnlyCVARCache::Get().bAllowStaticLighting;
-	if (bAllowStaticLighting)
 	{
-		NumLightmapDataEntries = UpdateStaticLightingBuffer();
-		if (NumLightmapDataEntries > 0 && UseGPUScene(GMaxRHIShaderPlatform, Scene->GetFeatureLevel()))
+		SCOPED_NAMED_EVENT(FPrimitiveSceneInfo_AddToScene_ReflectionCaptures, FColor::Yellow);
+		for (FPrimitiveSceneInfo* SceneInfo : SceneInfos)
 		{
-			LightmapDataOffset = Scene->GPUScene.LightmapDataAllocator.Allocate(NumLightmapDataEntries);
+			// Cache the nearest reflection proxy if needed
+			if (SceneInfo->NeedsReflectionCaptureUpdate())
+			{
+				SceneInfo->CacheReflectionCaptures();
+			}
 		}
 	}
 
-	// Cache the nearest reflection proxy if needed
-	if (NeedsReflectionCaptureUpdate())
 	{
-		CacheReflectionCaptures();
+		SCOPED_NAMED_EVENT(FPrimitiveSceneInfo_AddToScene_AddStaticMeshes, FColor::Magenta);
+		if (bUpdateStaticDrawLists)
+		{
+			AddStaticMeshes(RHICmdList, Scene, SceneInfos, bAddToStaticDrawLists);
+		}
 	}
 
-	if (bUpdateStaticDrawLists)
 	{
-		AddStaticMeshes(RHICmdList, bAddToStaticDrawLists);
+		SCOPED_NAMED_EVENT(FPrimitiveSceneInfo_AddToScene_AddToPrimitiveOctree, FColor::Red);
+		for (FPrimitiveSceneInfo* SceneInfo : SceneInfos)
+		{
+			// create potential storage for our compact info
+			FPrimitiveSceneInfoCompact CompactPrimitiveSceneInfo(SceneInfo);
+
+			// Add the primitive to the octree.
+			check(!SceneInfo->OctreeId.IsValidId());
+			Scene->PrimitiveOctree.AddElement(CompactPrimitiveSceneInfo);
+			check(SceneInfo->OctreeId.IsValidId());
+		}
 	}
 
-	// create potential storage for our compact info
-	FPrimitiveSceneInfoCompact CompactPrimitiveSceneInfo(this);
+	{
+		SCOPED_NAMED_EVENT(FPrimitiveSceneInfo_AddToScene_UpdateBounds, FColor::Cyan);
+		for (FPrimitiveSceneInfo* SceneInfo : SceneInfos)
+		{
+			FPrimitiveSceneProxy* Proxy = SceneInfo->Proxy;
+			int32 PackedIndex = SceneInfo->PackedIndex;
 
-	// Add the primitive to the octree.
-	check(!OctreeId.IsValidId());
-	Scene->PrimitiveOctree.AddElement(CompactPrimitiveSceneInfo);
-	check(OctreeId.IsValidId());
+			if (Proxy->CastsDynamicIndirectShadow())
+			{
+				Scene->DynamicIndirectCasterPrimitives.Add(SceneInfo);
+			}
 
-	if (Proxy->CastsDynamicIndirectShadow())
-	{
-		Scene->DynamicIndirectCasterPrimitives.Add(this);
-	}
+			Scene->PrimitiveSceneProxies[PackedIndex] = Proxy;
+			Scene->PrimitiveTransforms[PackedIndex] = Proxy->GetLocalToWorld();
+
+			// Set bounds.
+			FPrimitiveBounds& PrimitiveBounds = Scene->PrimitiveBounds[PackedIndex];
+			FBoxSphereBounds BoxSphereBounds = Proxy->GetBounds();
+			PrimitiveBounds.BoxSphereBounds = BoxSphereBounds;
+			PrimitiveBounds.MinDrawDistanceSq = FMath::Square(Proxy->GetMinDrawDistance());
+			PrimitiveBounds.MaxDrawDistance = Proxy->GetMaxDrawDistance();
+			PrimitiveBounds.MaxCullDistance = PrimitiveBounds.MaxDrawDistance;
 
-	Scene->PrimitiveSceneProxies[PackedIndex] = Proxy;
-	Scene->PrimitiveTransforms[PackedIndex] = Proxy->GetLocalToWorld();
+			Scene->PrimitiveFlagsCompact[PackedIndex] = FPrimitiveFlagsCompact(Proxy);
 
-	// Set bounds.
-	FPrimitiveBounds& PrimitiveBounds = Scene->PrimitiveBounds[PackedIndex];
-	FBoxSphereBounds BoxSphereBounds = Proxy->GetBounds();
-	PrimitiveBounds.BoxSphereBounds = BoxSphereBounds;
-	PrimitiveBounds.MinDrawDistanceSq = FMath::Square(Proxy->GetMinDrawDistance());
-	PrimitiveBounds.MaxDrawDistance = Proxy->GetMaxDrawDistance();
-	PrimitiveBounds.MaxCullDistance = PrimitiveBounds.MaxDrawDistance;
+			// Store precomputed visibility ID.
+			int32 VisibilityBitIndex = Proxy->GetVisibilityId();
+			FPrimitiveVisibilityId& VisibilityId = Scene->PrimitiveVisibilityIds[PackedIndex];
+			VisibilityId.ByteIndex = VisibilityBitIndex / 8;
+			VisibilityId.BitMask = (1 << (VisibilityBitIndex & 0x7));
 
-	Scene->PrimitiveFlagsCompact[PackedIndex] = FPrimitiveFlagsCompact(Proxy);
+			// Store occlusion flags.
+			uint8 OcclusionFlags = EOcclusionFlags::None;
+			if (Proxy->CanBeOccluded())
+			{
+				OcclusionFlags |= EOcclusionFlags::CanBeOccluded;
+			}
+			if (Proxy->HasSubprimitiveOcclusionQueries())
+			{
+				OcclusionFlags |= EOcclusionFlags::HasSubprimitiveQueries;
+			}
+			if (Proxy->AllowApproximateOcclusion()
+				// Allow approximate occlusion if attached, even if the parent does not have bLightAttachmentsAsGroup enabled
+				|| SceneInfo->LightingAttachmentRoot.IsValid())
+			{
+				OcclusionFlags |= EOcclusionFlags::AllowApproximateOcclusion;
+			}
+			if (VisibilityBitIndex >= 0)
+			{
+				OcclusionFlags |= EOcclusionFlags::HasPrecomputedVisibility;
+			}
+			Scene->PrimitiveOcclusionFlags[PackedIndex] = OcclusionFlags;
 	
-	// Store precomputed visibility ID.
-	int32 VisibilityBitIndex = Proxy->GetVisibilityId();
-	FPrimitiveVisibilityId& VisibilityId = Scene->PrimitiveVisibilityIds[PackedIndex];
-	VisibilityId.ByteIndex = VisibilityBitIndex / 8;
-	VisibilityId.BitMask = (1 << (VisibilityBitIndex & 0x7));
-
-	// Store occlusion flags.
-	uint8 OcclusionFlags = EOcclusionFlags::None;
-	if (Proxy->CanBeOccluded())
-	{
-		OcclusionFlags |= EOcclusionFlags::CanBeOccluded;
-	}
-	if (Proxy->HasSubprimitiveOcclusionQueries())
-	{
-		OcclusionFlags |= EOcclusionFlags::HasSubprimitiveQueries;
-	}
-	if (Proxy->AllowApproximateOcclusion()
-		// Allow approximate occlusion if attached, even if the parent does not have bLightAttachmentsAsGroup enabled
-		|| LightingAttachmentRoot.IsValid())
-	{
-		OcclusionFlags |= EOcclusionFlags::AllowApproximateOcclusion;
-	}
-	if (VisibilityBitIndex >= 0)
-	{
-		OcclusionFlags |= EOcclusionFlags::HasPrecomputedVisibility;
+			// Store occlusion bounds.
+			FBoxSphereBounds OcclusionBounds = BoxSphereBounds;
+			if (Proxy->HasCustomOcclusionBounds())
+			{
+				OcclusionBounds = Proxy->GetCustomOcclusionBounds();
+			}
+			OcclusionBounds.BoxExtent.X = OcclusionBounds.BoxExtent.X + OCCLUSION_SLOP;
+			OcclusionBounds.BoxExtent.Y = OcclusionBounds.BoxExtent.Y + OCCLUSION_SLOP;
+			OcclusionBounds.BoxExtent.Z = OcclusionBounds.BoxExtent.Z + OCCLUSION_SLOP;
+			OcclusionBounds.SphereRadius = OcclusionBounds.SphereRadius + OCCLUSION_SLOP;
+			Scene->PrimitiveOcclusionBounds[PackedIndex] = OcclusionBounds;
+
+			// Store the component.
+			Scene->PrimitiveComponentIds[PackedIndex] = SceneInfo->PrimitiveComponentId;
+		}
 	}
-	Scene->PrimitiveOcclusionFlags[PackedIndex] = OcclusionFlags;
 
-	// Store occlusion bounds.
-	FBoxSphereBounds OcclusionBounds = BoxSphereBounds;
-	if (Proxy->HasCustomOcclusionBounds())
 	{
-		OcclusionBounds = Proxy->GetCustomOcclusionBounds();
-	}
-	OcclusionBounds.BoxExtent.X = OcclusionBounds.BoxExtent.X + OCCLUSION_SLOP;
-	OcclusionBounds.BoxExtent.Y = OcclusionBounds.BoxExtent.Y + OCCLUSION_SLOP;
-	OcclusionBounds.BoxExtent.Z = OcclusionBounds.BoxExtent.Z + OCCLUSION_SLOP;
-	OcclusionBounds.SphereRadius = OcclusionBounds.SphereRadius + OCCLUSION_SLOP;
-	Scene->PrimitiveOcclusionBounds[PackedIndex] = OcclusionBounds;
-
-	// Store the component.
-	Scene->PrimitiveComponentIds[PackedIndex] = PrimitiveComponentId;
-
-	// Store the runtime virtual texture flags.
-	UpdateRuntimeVirtualTextureFlags();
-	Scene->PrimitiveVirtualTextureFlags[PackedIndex] = RuntimeVirtualTextureFlags;
+		SCOPED_NAMED_EVENT(FPrimitiveSceneInfo_AddToScene_UpdateVirtualTexture, FColor::Emerald);
+		for (FPrimitiveSceneInfo* SceneInfo : SceneInfos)
+		{
+			FPrimitiveSceneProxy* Proxy = SceneInfo->Proxy;
+			// Store the runtime virtual texture flags.
+			SceneInfo->UpdateRuntimeVirtualTextureFlags();
+			Scene->PrimitiveVirtualTextureFlags[SceneInfo->PackedIndex] = SceneInfo->RuntimeVirtualTextureFlags;
 
-	// Store the runtime virtual texture Lod info.
-	if (RuntimeVirtualTextureFlags.bRenderToVirtualTexture)
-	{
-		int8 MinLod, MaxLod;
-		GetRuntimeVirtualTextureLODRange(StaticMeshRelevances, MinLod, MaxLod);
-
-		FPrimitiveVirtualTextureLodInfo& LodInfo = Scene->PrimitiveVirtualTextureLod[PackedIndex];
-		LodInfo.MinLod = FMath::Clamp((int32)MinLod, 0, 15);
-		LodInfo.MaxLod = FMath::Clamp((int32)MaxLod, 0, 15);
-		LodInfo.LodBias = FMath::Clamp(Proxy->GetVirtualTextureLodBias() + FPrimitiveVirtualTextureLodInfo::LodBiasOffset, 0, 15);
-		LodInfo.CullMethod = Proxy->GetVirtualTextureMinCoverage() == 0 ? 0 : 1;
-		LodInfo.CullValue = LodInfo.CullMethod == 0 ? Proxy->GetVirtualTextureCullMips() : Proxy->GetVirtualTextureMinCoverage();
+			// Store the runtime virtual texture Lod info.
+			if (SceneInfo->RuntimeVirtualTextureFlags.bRenderToVirtualTexture)
+			{
+				int8 MinLod, MaxLod;
+				GetRuntimeVirtualTextureLODRange(SceneInfo->StaticMeshRelevances, MinLod, MaxLod);
+
+				FPrimitiveVirtualTextureLodInfo& LodInfo = Scene->PrimitiveVirtualTextureLod[SceneInfo->PackedIndex];
+				LodInfo.MinLod = FMath::Clamp((int32)MinLod, 0, 15);
+				LodInfo.MaxLod = FMath::Clamp((int32)MaxLod, 0, 15);
+				LodInfo.LodBias = FMath::Clamp(Proxy->GetVirtualTextureLodBias() + FPrimitiveVirtualTextureLodInfo::LodBiasOffset, 0, 15);
+				LodInfo.CullMethod = Proxy->GetVirtualTextureMinCoverage() == 0 ? 0 : 1;
+				LodInfo.CullValue = LodInfo.CullMethod == 0 ? Proxy->GetVirtualTextureCullMips() : Proxy->GetVirtualTextureMinCoverage();
+			}
+		}
 	}
 
 	// Find lights that affect the primitive in the light octree.
-	Scene->CreateLightPrimitiveInteractionsForPrimitive(this, bAsyncCreateLPIs);
+	for (FPrimitiveSceneInfo* SceneInfo : SceneInfos)
+	{
+		Scene->CreateLightPrimitiveInteractionsForPrimitive(SceneInfo, bAsyncCreateLPIs);
 
-	INC_MEMORY_STAT_BY(STAT_PrimitiveInfoMemory, sizeof(*this) + StaticMeshes.GetAllocatedSize() + StaticMeshRelevances.GetAllocatedSize() + Proxy->GetMemoryFootprint());
+		FPrimitiveSceneProxy* Proxy = SceneInfo->Proxy;
+		INC_MEMORY_STAT_BY(STAT_PrimitiveInfoMemory, sizeof(*SceneInfo) + SceneInfo->StaticMeshes.GetAllocatedSize() + SceneInfo->StaticMeshRelevances.GetAllocatedSize() + Proxy->GetMemoryFootprint());
+	}
 }
 
 void FPrimitiveSceneInfo::RemoveStaticMeshes()
@@ -702,7 +859,11 @@ void FPrimitiveSceneInfo::RemoveFromScene(bool bUpdateStaticDrawLists)
 	}
 
 	IndirectLightingCacheAllocation = NULL;
-	ClearIndirectLightingCacheBuffer(false);
+
+	if (Proxy->IsOftenMoving())
+	{
+		MarkIndirectLightingCacheBufferDirty();
+	}
 
 	DEC_MEMORY_STAT_BY(STAT_PrimitiveInfoMemory, sizeof(*this) + StaticMeshes.GetAllocatedSize() + StaticMeshRelevances.GetAllocatedSize() + Proxy->GetMemoryFootprint());
 
@@ -771,23 +932,30 @@ bool FPrimitiveSceneInfo::NeedsUpdateStaticMeshes()
 	return Scene->PrimitivesNeedingStaticMeshUpdate[PackedIndex];
 }
 
-void FPrimitiveSceneInfo::UpdateStaticMeshes(FRHICommandListImmediate& RHICmdList, bool bReAddToDrawLists)
+void FPrimitiveSceneInfo::UpdateStaticMeshes(FRHICommandListImmediate& RHICmdList, FScene* Scene, const TArrayView<FPrimitiveSceneInfo*>& SceneInfos, bool bReAddToDrawLists)
 {
 	QUICK_SCOPE_CYCLE_COUNTER(STAT_FPrimitiveSceneInfo_UpdateStaticMeshes);
+
 	const bool bNeedsStaticMeshUpdate = !bReAddToDrawLists;
-	Scene->PrimitivesNeedingStaticMeshUpdate[PackedIndex] = bNeedsStaticMeshUpdate;
 
-	if (!bNeedsStaticMeshUpdate && bNeedsStaticMeshUpdateWithoutVisibilityCheck)
+	for (int Index = 0; Index < SceneInfos.Num(); Index++)
 	{
-		Scene->PrimitivesNeedingStaticMeshUpdateWithoutVisibilityCheck.Remove(this);
+		FPrimitiveSceneInfo* SceneInfo = SceneInfos[Index];
+		Scene->PrimitivesNeedingStaticMeshUpdate[SceneInfo->PackedIndex] = bNeedsStaticMeshUpdate;
 
-		bNeedsStaticMeshUpdateWithoutVisibilityCheck = false;
+		if (!bNeedsStaticMeshUpdate && SceneInfo->bNeedsStaticMeshUpdateWithoutVisibilityCheck)
+		{
+			Scene->PrimitivesNeedingStaticMeshUpdateWithoutVisibilityCheck.Remove(SceneInfo);
+
+			SceneInfo->bNeedsStaticMeshUpdateWithoutVisibilityCheck = false;
+		}
+
+		SceneInfo->RemoveCachedMeshDrawCommands();
 	}
 
-	RemoveCachedMeshDrawCommands();
 	if (bReAddToDrawLists)
 	{
-		CacheMeshDrawCommands(RHICmdList);
+		CacheMeshDrawCommands(RHICmdList, Scene, SceneInfos);
 	}
 }
 
@@ -1085,14 +1253,6 @@ void FPrimitiveSceneInfo::UpdateIndirectLightingCacheBuffer()
 	}
 }
 
-void FPrimitiveSceneInfo::ClearIndirectLightingCacheBuffer(bool bSingleFrameOnly)
-{
-	if (!bSingleFrameOnly || Proxy->IsOftenMoving())
-	{
-		MarkIndirectLightingCacheBufferDirty();
-	}
-}
-
 void FPrimitiveSceneInfo::GetStaticMeshesLODRange(int8& OutMinLOD, int8& OutMaxLOD) const
 {
 	OutMinLOD = MAX_int8;
diff --git a/Engine/Source/Runtime/Renderer/Private/RayTracing/RayTracingAmbientOcclusion.cpp b/Engine/Source/Runtime/Renderer/Private/RayTracing/RayTracingAmbientOcclusion.cpp
--- a/Engine/Source/Runtime/Renderer/Private/RayTracing/RayTracingAmbientOcclusion.cpp
+++ b/Engine/Source/Runtime/Renderer/Private/RayTracing/RayTracingAmbientOcclusion.cpp
@@ -130,7 +130,7 @@ void FDeferredShadingSceneRenderer::PrepareRayTracingAmbientOcclusion(const FVie
 			PermutationVector.Set<FRayTracingAmbientOcclusionRGS::FEnableTwoSidedGeometryDim>(TwoSidedGeometryIndex != 0);
 			PermutationVector.Set<FRayTracingAmbientOcclusionRGS::FEnableMaterialsDim>(EnableMaterialsIndex != 0);
 			TShaderMapRef<FRayTracingAmbientOcclusionRGS> RayGenerationShader(View.ShaderMap, PermutationVector);
-			OutRayGenShaders.Add(RayGenerationShader->GetRayTracingShader());
+			OutRayGenShaders.Add(RayGenerationShader.GetRayTracingShader());
 		}
 	}
 }
@@ -180,7 +180,7 @@ void FDeferredShadingSceneRenderer::RenderRayTracingAmbientOcclusion(
 	PermutationVector.Set<FRayTracingAmbientOcclusionRGS::FEnableTwoSidedGeometryDim>(CVarRayTracingAmbientOcclusionEnableTwoSidedGeometry.GetValueOnRenderThread() != 0);
 	PermutationVector.Set<FRayTracingAmbientOcclusionRGS::FEnableMaterialsDim>(CVarRayTracingAmbientOcclusionEnableMaterials.GetValueOnRenderThread() != 0);
 	TShaderMapRef<FRayTracingAmbientOcclusionRGS> RayGenerationShader(GetGlobalShaderMap(FeatureLevel), PermutationVector);
-	ClearUnusedGraphResources(*RayGenerationShader, PassParameters);
+	ClearUnusedGraphResources(RayGenerationShader, PassParameters);
 
 	FIntPoint RayTracingResolution = View.ViewRect.Size();
 	GraphBuilder.AddPass(
@@ -190,7 +190,7 @@ void FDeferredShadingSceneRenderer::RenderRayTracingAmbientOcclusion(
 		[PassParameters, this, &View, RayGenerationShader, RayTracingResolution](FRHICommandList& RHICmdList)
 	{
 		FRayTracingShaderBindingsWriter GlobalResources;
-		SetShaderParameters(GlobalResources, *RayGenerationShader, *PassParameters);
+		SetShaderParameters(GlobalResources, RayGenerationShader, *PassParameters);
 
 		// TODO: Provide material support for opacity mask
 		FRayTracingPipelineState* Pipeline = View.RayTracingMaterialPipeline;
@@ -199,10 +199,10 @@ void FDeferredShadingSceneRenderer::RenderRayTracingAmbientOcclusion(
 			// Declare default pipeline
 			FRayTracingPipelineStateInitializer Initializer;
 			Initializer.MaxPayloadSizeInBytes = 52; // sizeof(FPackedMaterialClosestHitPayload)
-			FRHIRayTracingShader* RayGenShaderTable[] = { RayGenerationShader->GetRayTracingShader() };
+			FRHIRayTracingShader* RayGenShaderTable[] = { RayGenerationShader.GetRayTracingShader() };
 			Initializer.SetRayGenShaderTable(RayGenShaderTable);
 
-			FRHIRayTracingShader* HitGroupTable[] = { View.ShaderMap->GetShader<FOpaqueShadowHitGroup>()->GetRayTracingShader() };
+			FRHIRayTracingShader* HitGroupTable[] = { View.ShaderMap->GetShader<FOpaqueShadowHitGroup>().GetRayTracingShader() };
 			Initializer.SetHitGroupTable(HitGroupTable);
 			Initializer.bAllowHitGroupIndexing = false; // Use the same hit shader for all geometry in the scene by disabling SBT indexing.
 
@@ -210,7 +210,7 @@ void FDeferredShadingSceneRenderer::RenderRayTracingAmbientOcclusion(
 		}
 
 		FRHIRayTracingScene* RayTracingSceneRHI = View.RayTracingScene.RayTracingSceneRHI;
-		RHICmdList.RayTraceDispatch(Pipeline, RayGenerationShader->GetRayTracingShader(), RayTracingSceneRHI, GlobalResources, RayTracingResolution.X, RayTracingResolution.Y);
+		RHICmdList.RayTraceDispatch(Pipeline, RayGenerationShader.GetRayTracingShader(), RayTracingSceneRHI, GlobalResources, RayTracingResolution.X, RayTracingResolution.Y);
 	});
 
 	int32 DenoiserMode = CVarUseAODenoiser.GetValueOnRenderThread();
diff --git a/Engine/Source/Runtime/Renderer/Private/RayTracing/RayTracingBarycentrics.cpp b/Engine/Source/Runtime/Renderer/Private/RayTracing/RayTracingBarycentrics.cpp
--- a/Engine/Source/Runtime/Renderer/Private/RayTracing/RayTracingBarycentrics.cpp
+++ b/Engine/Source/Runtime/Renderer/Private/RayTracing/RayTracingBarycentrics.cpp
@@ -58,17 +58,17 @@ void FDeferredShadingSceneRenderer::RenderRayTracingBarycentrics(FRHICommandList
 
 	FRDGBuilder GraphBuilder(RHICmdList);
 
-	TShaderMap<FGlobalShaderType>* ShaderMap = GetGlobalShaderMap(FeatureLevel);
+	FGlobalShaderMap* ShaderMap = GetGlobalShaderMap(FeatureLevel);
 
 	auto RayGenShader = ShaderMap->GetShader<FRayTracingBarycentricsRGS>();
 	auto ClosestHitShader = ShaderMap->GetShader<FRayTracingBarycentricsCHS>();
 
 	FRayTracingPipelineStateInitializer Initializer;
 
-	FRHIRayTracingShader* RayGenShaderTable[] = { RayGenShader->GetRayTracingShader() };
+	FRHIRayTracingShader* RayGenShaderTable[] = { RayGenShader.GetRayTracingShader() };
 	Initializer.SetRayGenShaderTable(RayGenShaderTable);
 
-	FRHIRayTracingShader* HitGroupTable[] = { ClosestHitShader->GetRayTracingShader() };
+	FRHIRayTracingShader* HitGroupTable[] = { ClosestHitShader.GetRayTracingShader() };
 	Initializer.SetHitGroupTable(HitGroupTable);
 	Initializer.bAllowHitGroupIndexing = false; // Use the same hit shader for all geometry in the scene by disabling SBT indexing.
 
@@ -94,7 +94,7 @@ void FDeferredShadingSceneRenderer::RenderRayTracingBarycentrics(FRHICommandList
 		SetShaderParameters(GlobalResources, RayGenShader, *RayGenParameters);
 
 		// Dispatch rays using default shader binding table
-		RHICmdList.RayTraceDispatch(Pipeline, RayGenShader->GetRayTracingShader(), RayTracingSceneRHI, GlobalResources, ViewRect.Size().X, ViewRect.Size().Y);
+		RHICmdList.RayTraceDispatch(Pipeline, RayGenShader.GetRayTracingShader(), RayTracingSceneRHI, GlobalResources, ViewRect.Size().X, ViewRect.Size().Y);
 	});
 
 	GraphBuilder.Execute();
diff --git a/Engine/Source/Runtime/Renderer/Private/RayTracing/RayTracingDebug.cpp b/Engine/Source/Runtime/Renderer/Private/RayTracing/RayTracingDebug.cpp
--- a/Engine/Source/Runtime/Renderer/Private/RayTracing/RayTracingDebug.cpp
+++ b/Engine/Source/Runtime/Renderer/Private/RayTracing/RayTracingDebug.cpp
@@ -58,7 +58,7 @@ void FDeferredShadingSceneRenderer::PrepareRayTracingDebug(const FViewInfo& View
 	// Declare all RayGen shaders that require material closest hit shaders to be bound
 
 	auto RayGenShader = View.ShaderMap->GetShader<FRayTracingDebugRGS>();
-	OutRayGenShaders.Add(RayGenShader->GetRayTracingShader());
+	OutRayGenShaders.Add(RayGenShader.GetRayTracingShader());
 }
 
 void FDeferredShadingSceneRenderer::RenderRayTracingDebug(FRHICommandListImmediate& RHICmdList, const FViewInfo& View)
@@ -131,7 +131,7 @@ void FDeferredShadingSceneRenderer::RenderRayTracingDebug(FRHICommandListImmedia
 	}
 
 
-	TShaderMap<FGlobalShaderType>* ShaderMap = GetGlobalShaderMap(FeatureLevel);
+	FGlobalShaderMap* ShaderMap = GetGlobalShaderMap(FeatureLevel);
 
 	auto RayGenShader = ShaderMap->GetShader<FRayTracingDebugRGS>();
 
@@ -161,7 +161,7 @@ void FDeferredShadingSceneRenderer::RenderRayTracingDebug(FRHICommandListImmedia
 		FRayTracingShaderBindingsWriter GlobalResources;
 		SetShaderParameters(GlobalResources, RayGenShader, *RayGenParameters);
 
-		RHICmdList.RayTraceDispatch(Pipeline, RayGenShader->GetRayTracingShader(), RayTracingSceneRHI, GlobalResources, ViewRect.Size().X, ViewRect.Size().Y);
+		RHICmdList.RayTraceDispatch(Pipeline, RayGenShader.GetRayTracingShader(), RayTracingSceneRHI, GlobalResources, ViewRect.Size().X, ViewRect.Size().Y);
 	});
 
 	GraphBuilder.Execute();
diff --git a/Engine/Source/Runtime/Renderer/Private/RayTracing/RayTracingDeferredMaterials.cpp b/Engine/Source/Runtime/Renderer/Private/RayTracing/RayTracingDeferredMaterials.cpp
--- a/Engine/Source/Runtime/Renderer/Private/RayTracing/RayTracingDeferredMaterials.cpp
+++ b/Engine/Source/Runtime/Renderer/Private/RayTracing/RayTracingDeferredMaterials.cpp
@@ -38,8 +38,8 @@ FRayTracingPipelineState* FDeferredShadingSceneRenderer::BindRayTracingDeferredM
 	Initializer.MaxPayloadSizeInBytes = 12; // sizeof FDeferredMaterialPayload
 
 	// Get the ray tracing materials
-	auto* ClosestHitShader = View.ShaderMap->GetShader<FRayTracingDeferredMaterialCHS>();
-	FRHIRayTracingShader* HitShaderTable[] = { ClosestHitShader->GetRayTracingShader() };
+	auto ClosestHitShader = View.ShaderMap->GetShader<FRayTracingDeferredMaterialCHS>();
+	FRHIRayTracingShader* HitShaderTable[] = { ClosestHitShader.GetRayTracingShader() };
 	Initializer.SetHitGroupTable(HitShaderTable);
 
 	FRayTracingPipelineState* PipelineState = PipelineStateCache::GetAndOrCreateRayTracingPipelineState(RHICmdList, Initializer);
@@ -121,7 +121,7 @@ void SortDeferredMaterials(
 	FComputeShaderUtils::AddPass(
 		GraphBuilder,
 		RDG_EVENT_NAME("MaterialSort SortSize=%d NumElements=%d", ElementBlockSize, NumElements),
-		*SortShader,
+		SortShader,
 		PassParameters,
 		FIntVector(DispatchWidth, 1, 1));		
 }
diff --git a/Engine/Source/Runtime/Renderer/Private/RayTracing/RayTracingDynamicGeometry.cpp b/Engine/Source/Runtime/Renderer/Private/RayTracing/RayTracingDynamicGeometry.cpp
--- a/Engine/Source/Runtime/Renderer/Private/RayTracing/RayTracingDynamicGeometry.cpp
+++ b/Engine/Source/Runtime/Renderer/Private/RayTracing/RayTracingDynamicGeometry.cpp
@@ -41,17 +41,6 @@ public:
 		return IsSupportedDynamicVertexFactoryType(Parameters.VertexFactoryType) && ShouldCompileRayTracingShadersForProject(Parameters.Platform);
 	}
 
-	virtual bool Serialize(FArchive& Ar) override
-	{
-		bool bShaderHasOutdatedParameters = FMeshMaterialShader::Serialize(Ar);
-		Ar << RWVertexPositions;
-		Ar << VertexBufferSize;
-		Ar << NumVertices;
-		Ar << MinVertexIndex;
-		Ar << PrimitiveId;
-		return bShaderHasOutdatedParameters;
-	}
-
 	void GetShaderBindings(
 		const FScene* Scene,
 		ERHIFeatureLevel::Type FeatureLevel,
@@ -66,6 +55,7 @@ public:
 	}
 
 	void GetElementShaderBindings(
+		const FShaderMapPointerTable& PointerTable,
 		const FScene* Scene,
 		const FSceneView* ViewIfDynamicMeshCommand,
 		const FVertexFactory* VertexFactory,
@@ -78,14 +68,14 @@ public:
 		FMeshDrawSingleShaderBindings& ShaderBindings,
 		FVertexInputStreamArray& VertexStreams) const
 	{
-		FMeshMaterialShader::GetElementShaderBindings(Scene, ViewIfDynamicMeshCommand, VertexFactory, InputStreamType, FeatureLevel, PrimitiveSceneProxy, MeshBatch, BatchElement, ShaderElementData, ShaderBindings, VertexStreams);
+		FMeshMaterialShader::GetElementShaderBindings(PointerTable, Scene, ViewIfDynamicMeshCommand, VertexFactory, InputStreamType, FeatureLevel, PrimitiveSceneProxy, MeshBatch, BatchElement, ShaderElementData, ShaderBindings, VertexStreams);
 	}
 
-	FRWShaderParameter RWVertexPositions;
-	FShaderParameter VertexBufferSize;
-	FShaderParameter NumVertices;
-	FShaderParameter MinVertexIndex;
-	FShaderParameter PrimitiveId;
+	LAYOUT_FIELD(FRWShaderParameter, RWVertexPositions);
+	LAYOUT_FIELD(FShaderParameter, VertexBufferSize);
+	LAYOUT_FIELD(FShaderParameter, NumVertices);
+	LAYOUT_FIELD(FShaderParameter, MinVertexIndex);
+	LAYOUT_FIELD(FShaderParameter, PrimitiveId);
 };
 
 IMPLEMENT_MATERIAL_SHADER_TYPE(, FRayTracingDynamicGeometryConverterCS, TEXT("/Engine/Private/RayTracing/RayTracingDynamicMesh.usf"), TEXT("RayTracingDynamicGeometryConverterCS"), SF_Compute);
@@ -126,7 +116,7 @@ void FRayTracingDynamicGeometryCollection::AddDynamicMeshBatchForGeometryUpdate(
 
 		FMeshComputeDispatchCommand DispatchCmd;
 
-		FRayTracingDynamicGeometryConverterCS* Shader = Material.GetShader<FRayTracingDynamicGeometryConverterCS>(MeshBatch.VertexFactory->GetType());
+		TShaderRef<FRayTracingDynamicGeometryConverterCS> Shader = Material.GetShader<FRayTracingDynamicGeometryConverterCS>(MeshBatch.VertexFactory->GetType());
 		DispatchCmd.MaterialShader = Shader;
 		FMeshDrawShaderBindings& ShaderBindings = DispatchCmd.ShaderBindings;
 
@@ -136,12 +126,13 @@ void FRayTracingDynamicGeometryCollection::AddDynamicMeshBatchForGeometryUpdate(
 		FMeshMaterialShaderElementData ShaderElementData;
 		ShaderElementData.InitializeMeshMaterialData(View, PrimitiveSceneProxy, MeshBatch, -1, false);
 
-		FMeshDrawSingleShaderBindings SingleShaderBindings = ShaderBindings.GetSingleShaderBindings(SF_Compute);
+		int32 DataOffset = 0;
+		FMeshDrawSingleShaderBindings SingleShaderBindings = ShaderBindings.GetSingleShaderBindings(SF_Compute, DataOffset);
 		FMeshPassProcessorRenderState DrawRenderState(Scene->UniformBuffers.ViewUniformBuffer, Scene->UniformBuffers.OpaqueBasePassUniformBuffer);
 		Shader->GetShaderBindings(Scene, Scene->GetFeatureLevel(), PrimitiveSceneProxy, MaterialRenderProxy, Material, DrawRenderState, ShaderElementData, SingleShaderBindings);
 
 		FVertexInputStreamArray DummyArray;
-		Shader->GetElementShaderBindings(Scene, View, MeshBatch.VertexFactory, EVertexInputStreamType::Default, Scene->GetFeatureLevel(), PrimitiveSceneProxy, MeshBatch, MeshBatch.Elements[0], ShaderElementData, SingleShaderBindings, DummyArray);
+		FMeshMaterialShader::GetElementShaderBindings(Shader, Scene, View, MeshBatch.VertexFactory, EVertexInputStreamType::Default, Scene->GetFeatureLevel(), PrimitiveSceneProxy, MeshBatch, MeshBatch.Elements[0], ShaderElementData, SingleShaderBindings, DummyArray);
 
 		DispatchCmd.TargetBuffer = &Buffer;
 		DispatchCmd.NumMaxVertices = UpdateParams.NumVertices;
@@ -240,19 +231,19 @@ void FRayTracingDynamicGeometryCollection::DispatchUpdates(FRHIComputeCommandLis
 			for (FMeshComputeDispatchCommand& Cmd : *DispatchCommands)
 			{
 				{
-					FRayTracingDynamicGeometryConverterCS* Shader = Cmd.MaterialShader;
+					const TShaderRef<FRayTracingDynamicGeometryConverterCS>& Shader = Cmd.MaterialShader;
 
-					RHICmdList.SetComputeShader(Shader->GetComputeShader());
+					RHICmdList.SetComputeShader(Shader.GetComputeShader());
 
-					Cmd.ShaderBindings.SetOnCommandList(RHICmdList, Shader->GetComputeShader());
-					Shader->RWVertexPositions.SetBuffer(RHICmdList, Shader->GetComputeShader(), *Cmd.TargetBuffer);
-					SetShaderValue(RHICmdList, Shader->GetComputeShader(), Shader->VertexBufferSize, Cmd.TargetBuffer->NumBytes / sizeof(FVector));
-					SetShaderValue(RHICmdList, Shader->GetComputeShader(), Shader->NumVertices, Cmd.NumCPUVertices);
-					SetShaderValue(RHICmdList, Shader->GetComputeShader(), Shader->MinVertexIndex, Cmd.MinVertexIndex);
-					SetShaderValue(RHICmdList, Shader->GetComputeShader(), Shader->PrimitiveId, Cmd.PrimitiveId);
+					Cmd.ShaderBindings.SetOnCommandList(RHICmdList, Shader.GetComputeShader());
+					Shader->RWVertexPositions.SetBuffer(RHICmdList, Shader.GetComputeShader(), *Cmd.TargetBuffer);
+					SetShaderValue(RHICmdList, Shader.GetComputeShader(), Shader->VertexBufferSize, Cmd.TargetBuffer->NumBytes / sizeof(FVector));
+					SetShaderValue(RHICmdList, Shader.GetComputeShader(), Shader->NumVertices, Cmd.NumCPUVertices);
+					SetShaderValue(RHICmdList, Shader.GetComputeShader(), Shader->MinVertexIndex, Cmd.MinVertexIndex);
+					SetShaderValue(RHICmdList, Shader.GetComputeShader(), Shader->PrimitiveId, Cmd.PrimitiveId);
 					RHICmdList.DispatchComputeShader(FMath::DivideAndRoundUp<uint32>(Cmd.NumMaxVertices, 64), 1, 1);
 
-					Shader->RWVertexPositions.UnsetUAV(RHICmdList, Shader->GetComputeShader());
+					Shader->RWVertexPositions.UnsetUAV(RHICmdList, Shader.GetComputeShader());
 				}
 			}
 
diff --git a/Engine/Source/Runtime/Renderer/Private/RayTracing/RayTracingGlobalIllumination.cpp b/Engine/Source/Runtime/Renderer/Private/RayTracing/RayTracingGlobalIllumination.cpp
--- a/Engine/Source/Runtime/Renderer/Private/RayTracing/RayTracingGlobalIllumination.cpp
+++ b/Engine/Source/Runtime/Renderer/Private/RayTracing/RayTracingGlobalIllumination.cpp
@@ -554,7 +554,7 @@ void FDeferredShadingSceneRenderer::PrepareRayTracingGlobalIllumination(const FV
 			PermutationVector.Set<FGlobalIlluminationRGS::FUseAttenuationTermDim>(CVarRayTracingGlobalIlluminationEnableLightAttenuation.GetValueOnRenderThread() != 0);
 			PermutationVector.Set<FGlobalIlluminationRGS::FEnableTwoSidedGeometryDim>(EnableTwoSidedGeometry == 1);
 			TShaderMapRef<FGlobalIlluminationRGS> RayGenerationShader(View.ShaderMap, PermutationVector);
-			OutRayGenShaders.Add(RayGenerationShader->GetRayTracingShader());
+			OutRayGenShaders.Add(RayGenerationShader.GetRayTracingShader());
 
 			if (bSortMaterials)
 			{
@@ -565,7 +565,7 @@ void FDeferredShadingSceneRenderer::PrepareRayTracingGlobalIllumination(const FV
 					CreateGatherPointsPermutationVector.Set<FRayTracingGlobalIlluminationCreateGatherPointsTraceRGS::FEnableTwoSidedGeometryDim>(EnableTwoSidedGeometry == 1);
 					CreateGatherPointsPermutationVector.Set<FRayTracingGlobalIlluminationCreateGatherPointsTraceRGS::FDeferredMaterialMode>(EDeferredMaterialMode::Gather);
 					TShaderMapRef<FRayTracingGlobalIlluminationCreateGatherPointsTraceRGS> CreateGatherPointsRayGenerationShader(View.ShaderMap, CreateGatherPointsPermutationVector);
-					OutRayGenShaders.Add(CreateGatherPointsRayGenerationShader->GetRayTracingShader());
+					OutRayGenShaders.Add(CreateGatherPointsRayGenerationShader.GetRayTracingShader());
 				}
 
 				// Shade
@@ -575,7 +575,7 @@ void FDeferredShadingSceneRenderer::PrepareRayTracingGlobalIllumination(const FV
 					CreateGatherPointsPermutationVector.Set<FRayTracingGlobalIlluminationCreateGatherPointsRGS::FEnableTwoSidedGeometryDim>(EnableTwoSidedGeometry == 1);
 					CreateGatherPointsPermutationVector.Set<FRayTracingGlobalIlluminationCreateGatherPointsRGS::FDeferredMaterialMode>(EDeferredMaterialMode::Shade);
 					TShaderMapRef<FRayTracingGlobalIlluminationCreateGatherPointsRGS> CreateGatherPointsRayGenerationShader(View.ShaderMap, CreateGatherPointsPermutationVector);
-					OutRayGenShaders.Add(CreateGatherPointsRayGenerationShader->GetRayTracingShader());
+					OutRayGenShaders.Add(CreateGatherPointsRayGenerationShader.GetRayTracingShader());
 				}
 			}
 			else
@@ -585,14 +585,14 @@ void FDeferredShadingSceneRenderer::PrepareRayTracingGlobalIllumination(const FV
 				CreateGatherPointsPermutationVector.Set<FRayTracingGlobalIlluminationCreateGatherPointsRGS::FEnableTwoSidedGeometryDim>(EnableTwoSidedGeometry == 1);
 				CreateGatherPointsPermutationVector.Set<FRayTracingGlobalIlluminationCreateGatherPointsRGS::FDeferredMaterialMode>(EDeferredMaterialMode::None);
 				TShaderMapRef<FRayTracingGlobalIlluminationCreateGatherPointsRGS> CreateGatherPointsRayGenerationShader(View.ShaderMap, CreateGatherPointsPermutationVector);
-				OutRayGenShaders.Add(CreateGatherPointsRayGenerationShader->GetRayTracingShader());
+				OutRayGenShaders.Add(CreateGatherPointsRayGenerationShader.GetRayTracingShader());
 			}
 
 			FRayTracingGlobalIlluminationFinalGatherRGS::FPermutationDomain GatherPassPermutationVector;
 			GatherPassPermutationVector.Set<FRayTracingGlobalIlluminationFinalGatherRGS::FUseAttenuationTermDim>(UseAttenuationTerm == 1);
 			GatherPassPermutationVector.Set<FRayTracingGlobalIlluminationFinalGatherRGS::FEnableTwoSidedGeometryDim>(EnableTwoSidedGeometry == 1);
 			TShaderMapRef<FRayTracingGlobalIlluminationFinalGatherRGS> GatherPassRayGenerationShader(View.ShaderMap, GatherPassPermutationVector);
-			OutRayGenShaders.Add(GatherPassRayGenerationShader->GetRayTracingShader());
+			OutRayGenShaders.Add(GatherPassRayGenerationShader.GetRayTracingShader());
 		}
 	}
 
@@ -607,7 +607,7 @@ bool FDeferredShadingSceneRenderer::RenderRayTracingGlobalIllumination(
 	IScreenSpaceDenoiser::FAmbientOcclusionRayTracingConfig* OutRayTracingConfig,
 	IScreenSpaceDenoiser::FDiffuseIndirectInputs* OutDenoiserInputs)
 #if RHI_RAYTRACING
-{	
+{
 	if (!View.ViewState) return false;
 
 	int32 RayTracingGISamplesPerPixel = GetRayTracingGlobalIlluminationSamplesPerPixel(View);
@@ -620,7 +620,7 @@ bool FDeferredShadingSceneRenderer::RenderRayTracingGlobalIllumination(
 	}
 
 	OutRayTracingConfig->RayCountPerPixel = RayTracingGISamplesPerPixel;
-	
+
 	int32 UpscaleFactor = int32(1.0 / OutRayTracingConfig->ResolutionFraction);
 
 	// Allocate input for the denoiser.
@@ -709,7 +709,7 @@ void FDeferredShadingSceneRenderer::RayTracingGlobalIlluminationCreateGatherPoin
 	FIntVector& GatherPointsResolution
 )
 #if RHI_RAYTRACING
-{
+{	
 	RDG_GPU_STAT_SCOPE(GraphBuilder, RayTracingGICreateGatherPoints);
 	RDG_EVENT_SCOPE(GraphBuilder, "Ray Tracing GI: Create Gather Points");
 
@@ -801,7 +801,7 @@ void FDeferredShadingSceneRenderer::RayTracingGlobalIlluminationCreateGatherPoin
 	PermutationVector.Set<FRayTracingGlobalIlluminationCreateGatherPointsRGS::FUseAttenuationTermDim>(true);
 	PermutationVector.Set<FRayTracingGlobalIlluminationCreateGatherPointsRGS::FEnableTwoSidedGeometryDim>(CVarRayTracingGlobalIlluminationEnableTwoSidedGeometry.GetValueOnRenderThread() != 0);
 	TShaderMapRef<FRayTracingGlobalIlluminationCreateGatherPointsRGS> RayGenerationShader(GetGlobalShaderMap(FeatureLevel), PermutationVector);
-	ClearUnusedGraphResources(*RayGenerationShader, PassParameters);
+	ClearUnusedGraphResources(RayGenerationShader, PassParameters);
 
 	GraphBuilder.AddPass(
 		RDG_EVENT_NAME("GatherPoints %d%d", GatherPointsResolution.X, GatherPointsResolution.Y),
@@ -811,8 +811,8 @@ void FDeferredShadingSceneRenderer::RayTracingGlobalIlluminationCreateGatherPoin
 	{
 		FRHIRayTracingScene* RayTracingSceneRHI = View.RayTracingScene.RayTracingSceneRHI;
 		FRayTracingShaderBindingsWriter GlobalResources;
-		SetShaderParameters(GlobalResources, *RayGenerationShader, *PassParameters);
-		RHICmdList.RayTraceDispatch(View.RayTracingMaterialPipeline, RayGenerationShader->GetRayTracingShader(), RayTracingSceneRHI, GlobalResources, GatherPointsResolution.X, GatherPointsResolution.Y);
+		SetShaderParameters(GlobalResources, RayGenerationShader, *PassParameters);
+		RHICmdList.RayTraceDispatch(View.RayTracingMaterialPipeline, RayGenerationShader.GetRayTracingShader(), RayTracingSceneRHI, GlobalResources, GatherPointsResolution.X, GatherPointsResolution.Y);
 	});
 	}
 	else
@@ -845,7 +845,7 @@ void FDeferredShadingSceneRenderer::RayTracingGlobalIlluminationCreateGatherPoin
 			PermutationVector.Set<FRayTracingGlobalIlluminationCreateGatherPointsTraceRGS::FDeferredMaterialMode>(EDeferredMaterialMode::Gather);
 			TShaderMapRef<FRayTracingGlobalIlluminationCreateGatherPointsTraceRGS> RayGenerationShader(GetGlobalShaderMap(FeatureLevel), PermutationVector);
 
-			ClearUnusedGraphResources(*RayGenerationShader, GatherPassParameters);
+			ClearUnusedGraphResources(RayGenerationShader, GatherPassParameters);
 
 			GraphBuilder.AddPass(
 				RDG_EVENT_NAME("GlobalIlluminationRayTracingGatherMaterials %dx%d", TileAlignedResolution.X, TileAlignedResolution.Y),
@@ -853,13 +853,13 @@ void FDeferredShadingSceneRenderer::RayTracingGlobalIlluminationCreateGatherPoin
 				ERDGPassFlags::Compute,
 				[GatherPassParameters, this, &View, RayGenerationShader, TileAlignedResolution](FRHICommandList& RHICmdList)
 			{
-				FRayTracingPipelineState* Pipeline = BindRayTracingDeferredMaterialGatherPipeline(RHICmdList, View, RayGenerationShader->GetRayTracingShader());
+				FRayTracingPipelineState* Pipeline = BindRayTracingDeferredMaterialGatherPipeline(RHICmdList, View, RayGenerationShader.GetRayTracingShader());
 
 				FRayTracingShaderBindingsWriter GlobalResources;
-				SetShaderParameters(GlobalResources, *RayGenerationShader, *GatherPassParameters);
+				SetShaderParameters(GlobalResources, RayGenerationShader, *GatherPassParameters);
 
 				FRHIRayTracingScene* RayTracingSceneRHI = View.RayTracingScene.RayTracingSceneRHI;
-				RHICmdList.RayTraceDispatch(Pipeline, RayGenerationShader->GetRayTracingShader(), RayTracingSceneRHI, GlobalResources, TileAlignedResolution.X, TileAlignedResolution.Y);
+				RHICmdList.RayTraceDispatch(Pipeline, RayGenerationShader.GetRayTracingShader(), RayTracingSceneRHI, GlobalResources, TileAlignedResolution.X, TileAlignedResolution.Y);
 			});
 		}
 
@@ -876,7 +876,7 @@ void FDeferredShadingSceneRenderer::RayTracingGlobalIlluminationCreateGatherPoin
 			PermutationVector.Set<FRayTracingGlobalIlluminationCreateGatherPointsRGS::FEnableTwoSidedGeometryDim>(CVarRayTracingGlobalIlluminationEnableTwoSidedGeometry.GetValueOnRenderThread() != 0);
 			PermutationVector.Set<FRayTracingGlobalIlluminationCreateGatherPointsRGS::FDeferredMaterialMode>(EDeferredMaterialMode::Shade);
 			TShaderMapRef<FRayTracingGlobalIlluminationCreateGatherPointsRGS> RayGenerationShader(GetGlobalShaderMap(FeatureLevel), PermutationVector);
-			ClearUnusedGraphResources(*RayGenerationShader, PassParameters);
+			ClearUnusedGraphResources(RayGenerationShader, PassParameters);
 
 			GraphBuilder.AddPass(
 				RDG_EVENT_NAME("GlobalIlluminationRayTracingShadeMaterials %d", DeferredMaterialBufferNumElements),
@@ -886,11 +886,11 @@ void FDeferredShadingSceneRenderer::RayTracingGlobalIlluminationCreateGatherPoin
 			{
 				FRHIRayTracingScene* RayTracingSceneRHI = View.RayTracingScene.RayTracingSceneRHI;
 				FRayTracingShaderBindingsWriter GlobalResources;
-				SetShaderParameters(GlobalResources, *RayGenerationShader, *PassParameters);
+				SetShaderParameters(GlobalResources, RayGenerationShader, *PassParameters);
 
 				// Shading pass for sorted materials uses 1D dispatch over all elements in the material buffer.
 				// This can be reduced to the number of output pixels if sorting pass guarantees that all invalid entries are moved to the end.
-				RHICmdList.RayTraceDispatch(View.RayTracingMaterialPipeline, RayGenerationShader->GetRayTracingShader(), RayTracingSceneRHI, GlobalResources, DeferredMaterialBufferNumElements, 1);
+				RHICmdList.RayTraceDispatch(View.RayTracingMaterialPipeline, RayGenerationShader.GetRayTracingShader(), RayTracingSceneRHI, GlobalResources, DeferredMaterialBufferNumElements, 1);
 			});
 		}
 	}
@@ -959,7 +959,7 @@ void FDeferredShadingSceneRenderer::RenderRayTracingGlobalIlluminationFinalGathe
 	PermutationVector.Set<FRayTracingGlobalIlluminationFinalGatherRGS::FUseAttenuationTermDim>(true);
 	PermutationVector.Set<FRayTracingGlobalIlluminationFinalGatherRGS::FEnableTwoSidedGeometryDim>(CVarRayTracingGlobalIlluminationEnableTwoSidedGeometry.GetValueOnRenderThread() != 0);
 	TShaderMapRef<FRayTracingGlobalIlluminationFinalGatherRGS> RayGenerationShader(GetGlobalShaderMap(FeatureLevel), PermutationVector);
-	ClearUnusedGraphResources(*RayGenerationShader, PassParameters);
+	ClearUnusedGraphResources(RayGenerationShader, PassParameters);
 
 	FIntPoint RayTracingResolution = FIntPoint::DivideAndRoundUp(View.ViewRect.Size(), UpscaleFactor);
 	GraphBuilder.AddPass(
@@ -971,8 +971,8 @@ void FDeferredShadingSceneRenderer::RenderRayTracingGlobalIlluminationFinalGathe
 		FRHIRayTracingScene* RayTracingSceneRHI = View.RayTracingScene.RayTracingSceneRHI;
 
 		FRayTracingShaderBindingsWriter GlobalResources;
-		SetShaderParameters(GlobalResources, *RayGenerationShader, *PassParameters);
-		RHICmdList.RayTraceDispatch(View.RayTracingMaterialPipeline, RayGenerationShader->GetRayTracingShader(), RayTracingSceneRHI, GlobalResources, RayTracingResolution.X, RayTracingResolution.Y);
+		SetShaderParameters(GlobalResources, RayGenerationShader, *PassParameters);
+		RHICmdList.RayTraceDispatch(View.RayTracingMaterialPipeline, RayGenerationShader.GetRayTracingShader(), RayTracingSceneRHI, GlobalResources, RayTracingResolution.X, RayTracingResolution.Y);
 	});
 
 
@@ -1060,7 +1060,7 @@ void FDeferredShadingSceneRenderer::RenderRayTracingGlobalIlluminationBruteForce
 	PermutationVector.Set<FGlobalIlluminationRGS::FUseAttenuationTermDim>(true);
 	PermutationVector.Set<FGlobalIlluminationRGS::FEnableTwoSidedGeometryDim>(CVarRayTracingGlobalIlluminationEnableTwoSidedGeometry.GetValueOnRenderThread() != 0);
 	TShaderMapRef<FGlobalIlluminationRGS> RayGenerationShader(GetGlobalShaderMap(FeatureLevel), PermutationVector);
-	ClearUnusedGraphResources(*RayGenerationShader, PassParameters);
+	ClearUnusedGraphResources(RayGenerationShader, PassParameters);
 
 	FIntPoint RayTracingResolution = FIntPoint::DivideAndRoundUp(View.ViewRect.Size(), UpscaleFactor);
 
@@ -1075,8 +1075,8 @@ void FDeferredShadingSceneRenderer::RenderRayTracingGlobalIlluminationBruteForce
 			FRHIRayTracingScene* RayTracingSceneRHI = View.RayTracingScene.RayTracingSceneRHI;
 
 			FRayTracingShaderBindingsWriter GlobalResources;
-			SetShaderParameters(GlobalResources, *RayGenerationShader, *PassParameters);
-			RHICmdList.RayTraceDispatch(View.RayTracingMaterialPipeline, RayGenerationShader->GetRayTracingShader(), RayTracingSceneRHI, GlobalResources, RayTracingResolution.X, RayTracingResolution.Y);
+			SetShaderParameters(GlobalResources, RayGenerationShader, *PassParameters);
+			RHICmdList.RayTraceDispatch(View.RayTracingMaterialPipeline, RayGenerationShader.GetRayTracingShader(), RayTracingSceneRHI, GlobalResources, RayTracingResolution.X, RayTracingResolution.Y);
 		});
 	}
 	else
@@ -1111,8 +1111,8 @@ void FDeferredShadingSceneRenderer::RenderRayTracingGlobalIlluminationBruteForce
 					FRHIRayTracingScene* RayTracingSceneRHI = View.RayTracingScene.RayTracingSceneRHI;
 
 					FRayTracingShaderBindingsWriter GlobalResources;
-					SetShaderParameters(GlobalResources, *RayGenerationShader, *TilePassParameters);
-					RHICmdList.RayTraceDispatch(View.RayTracingMaterialPipeline, RayGenerationShader->GetRayTracingShader(), RayTracingSceneRHI, 
+					SetShaderParameters(GlobalResources, RayGenerationShader, *TilePassParameters);
+					RHICmdList.RayTraceDispatch(View.RayTracingMaterialPipeline, RayGenerationShader.GetRayTracingShader(), RayTracingSceneRHI, 
 						GlobalResources, DispatchSizeX, DispatchSizeY);
 					RHICmdList.SubmitCommandsHint();
 				});
diff --git a/Engine/Source/Runtime/Renderer/Private/RayTracing/RayTracingLighting.cpp b/Engine/Source/Runtime/Renderer/Private/RayTracing/RayTracingLighting.cpp
--- a/Engine/Source/Runtime/Renderer/Private/RayTracing/RayTracingLighting.cpp
+++ b/Engine/Source/Runtime/Renderer/Private/RayTracing/RayTracingLighting.cpp
@@ -18,7 +18,7 @@ bool CanUseRayTracingLightingMissShader(EShaderPlatform ShaderPlatform)
 {
 	return GRHISupportsRayTracingMissShaderBindings
 		&& CVarRayTracingLightingMissShader.GetValueOnRenderThread() != 0
-		&& FDataDrivenShaderPlatformInfo::GetInfo(ShaderPlatform).bSupportsRayTracingMissShaderBindings;
+		&& FDataDrivenShaderPlatformInfo::GetSupportsRayTracingMissShaderBindings(ShaderPlatform);
 }
 
 IMPLEMENT_GLOBAL_SHADER_PARAMETER_STRUCT(FRaytracingLightDataPacked, "RaytracingLightsDataPacked");
@@ -238,7 +238,7 @@ class FRayTracingLightingMS : public FGlobalShader
 	static bool ShouldCompilePermutation(const FGlobalShaderPermutationParameters& Parameters)
 	{
 		return ShouldCompileRayTracingShadersForProject(Parameters.Platform)
-			&& FDataDrivenShaderPlatformInfo::GetInfo(Parameters.Platform).bSupportsRayTracingMissShaderBindings;
+			&& FDataDrivenShaderPlatformInfo::GetSupportsRayTracingMissShaderBindings(Parameters.Platform);
 	}
 };
 
@@ -246,15 +246,15 @@ IMPLEMENT_GLOBAL_SHADER(FRayTracingLightingMS, "/Engine/Private/RayTracing/RayTr
 
 FRHIRayTracingShader* FDeferredShadingSceneRenderer::GetRayTracingLightingMissShader(FViewInfo& View)
 {
-	return View.ShaderMap->GetShader<FRayTracingLightingMS>()->GetRayTracingShader();
+	return View.ShaderMap->GetShader<FRayTracingLightingMS>().GetRayTracingShader();
 }
 
 template< typename ShaderClass>
-static int32 BindParameters(ShaderClass* Shader, typename ShaderClass::FParameters & Parameters, int32 MaxParams, const FRHIUniformBuffer **OutUniformBuffers)
+static int32 BindParameters(const TShaderRef<ShaderClass>& Shader, typename ShaderClass::FParameters & Parameters, int32 MaxParams, const FRHIUniformBuffer **OutUniformBuffers)
 {
 	FRayTracingShaderBindingsWriter ResourceBinder;
 
-	auto &ParameterMap = Shader->GetParameterMapInfo();
+	auto &ParameterMap = Shader->ParameterMapInfo;
 
 	// all parameters should be in uniform buffers
 	check(ParameterMap.LooseParameterBuffers.Num() == 0);
diff --git a/Engine/Source/Runtime/Renderer/Private/RayTracing/RayTracingMaterialHitShaders.cpp b/Engine/Source/Runtime/Renderer/Private/RayTracing/RayTracingMaterialHitShaders.cpp
--- a/Engine/Source/Runtime/Renderer/Private/RayTracing/RayTracingMaterialHitShaders.cpp
+++ b/Engine/Source/Runtime/Renderer/Private/RayTracing/RayTracingMaterialHitShaders.cpp
@@ -46,6 +46,7 @@ static bool IsSupportedVertexFactoryType(const FVertexFactoryType* VertexFactory
 
 class FMaterialCHS : public FMeshMaterialShader, public FUniformLightMapPolicyShaderParametersType
 {
+	DECLARE_INLINE_TYPE_LAYOUT_EXPLICIT_BASES(FMaterialCHS, NonVirtual, FMeshMaterialShader, FUniformLightMapPolicyShaderParametersType);
 public:
 	FMaterialCHS(const FMeshMaterialShaderType::CompiledShaderInitializerType& Initializer)
 		: FMeshMaterialShader(Initializer)
@@ -56,13 +57,6 @@ public:
 
 	FMaterialCHS() {}
 
-	virtual bool Serialize(FArchive& Ar) override
-	{
-		bool bShaderHasOutdatedParameters = FMeshMaterialShader::Serialize(Ar);
-		FUniformLightMapPolicyShaderParametersType::Serialize(Ar);
-		return bShaderHasOutdatedParameters;
-	}
-
 	void GetShaderBindings(
 		const FScene* Scene,
 		ERHIFeatureLevel::Type FeatureLevel,
@@ -83,6 +77,7 @@ public:
 	}
 
 	void GetElementShaderBindings(
+		const FShaderMapPointerTable& PointerTable,
 		const FScene* Scene,
 		const FSceneView* ViewIfDynamicMeshCommand,
 		const FVertexFactory* VertexFactory,
@@ -95,7 +90,7 @@ public:
 		FMeshDrawSingleShaderBindings& ShaderBindings,
 		FVertexInputStreamArray& VertexStreams) const
 	{
-		FMeshMaterialShader::GetElementShaderBindings(Scene, ViewIfDynamicMeshCommand, VertexFactory, InputStreamType, FeatureLevel, PrimitiveSceneProxy, MeshBatch, BatchElement, ShaderElementData, ShaderBindings, VertexStreams);
+		FMeshMaterialShader::GetElementShaderBindings(PointerTable, Scene, ViewIfDynamicMeshCommand, VertexFactory, InputStreamType, FeatureLevel, PrimitiveSceneProxy, MeshBatch, BatchElement, ShaderElementData, ShaderBindings, VertexStreams);
 	}
 };
 
@@ -114,7 +109,7 @@ public:
 	static bool ShouldCompilePermutation(const FMeshMaterialShaderPermutationParameters& Parameters)
 	{
 		return IsSupportedVertexFactoryType(Parameters.VertexFactoryType)
-			&& (Parameters.Material->IsMasked() == UseAnyHitShader)
+			&& (Parameters.MaterialParameters.bIsMasked == UseAnyHitShader)
 			&& LightMapPolicyType::ShouldCompilePermutation(Parameters)
 			&& ShouldCompileRayTracingShadersForProject(Parameters.Platform)
 			&& (bool)GRayTracingUseTextureLod == UseRayConeTextureLod;
@@ -128,7 +123,7 @@ public:
 		FMeshMaterialShader::ModifyCompilationEnvironment(Parameters, OutEnvironment);
 	}
 
-	static bool ValidateCompiledResult(EShaderPlatform Platform, const TArray<FMaterial*>& Materials, const FVertexFactoryType* VertexFactoryType, const FShaderParameterMap& ParameterMap, TArray<FString>& OutError)
+	static bool ValidateCompiledResult(EShaderPlatform Platform, const FShaderParameterMap& ParameterMap, TArray<FString>& OutError)
 	{
 		if (ParameterMap.ContainsParameterAllocation(FSceneTexturesUniformParameters::StaticStructMetadata.GetShaderVariableName()))
 		{
@@ -172,7 +167,7 @@ IMPLEMENT_GLOBAL_SHADER(FHiddenMaterialHitGroup, "/Engine/Private/RayTracing/Ray
 IMPLEMENT_GLOBAL_SHADER(FOpaqueShadowHitGroup, "/Engine/Private/RayTracing/RayTracingMaterialDefaultHitShaders.usf", "closesthit=OpaqueShadowCHS", SF_RayHitGroup);
 
 template<typename LightMapPolicyType>
-static FMaterialCHS* GetMaterialHitShader(const FMaterial& RESTRICT MaterialResource, const FVertexFactory* VertexFactory, bool UseTextureLod)
+static TShaderRef<FMaterialCHS> GetMaterialHitShader(const FMaterial& RESTRICT MaterialResource, const FVertexFactory* VertexFactory, bool UseTextureLod)
 {
 	if (MaterialResource.IsMasked())
 	{
@@ -224,9 +219,10 @@ void FRayTracingMeshProcessor::BuildRayTracingMeshCommands(
 	FVertexInputStreamArray VertexStreams;
 	VertexFactory->GetStreams(ERHIFeatureLevel::SM5, EVertexInputStreamType::Default, VertexStreams);
 
-	if (PassShaders.RayHitGroupShader)
+	int32 DataOffset = 0;
+	if (PassShaders.RayHitGroupShader.IsValid())
 	{
-		FMeshDrawSingleShaderBindings ShaderBindings = SharedCommand.ShaderBindings.GetSingleShaderBindings(SF_RayHitGroup);
+		FMeshDrawSingleShaderBindings ShaderBindings = SharedCommand.ShaderBindings.GetSingleShaderBindings(SF_RayHitGroup, DataOffset);
 		PassShaders.RayHitGroupShader->GetShaderBindings(Scene, FeatureLevel, PrimitiveSceneProxy, MaterialRenderProxy, MaterialResource, DrawRenderState, ShaderElementData, ShaderBindings);
 	}
 
@@ -236,13 +232,14 @@ void FRayTracingMeshProcessor::BuildRayTracingMeshCommands(
 	{
 		if ((1ull << BatchElementIndex) & BatchElementMask)
 		{
+			DataOffset = 0;
 			const FMeshBatchElement& BatchElement = MeshBatch.Elements[BatchElementIndex];
 			FRayTracingMeshCommand& RayTracingMeshCommand = CommandContext->AddCommand(SharedCommand);
 
-			if (PassShaders.RayHitGroupShader)
+			if (PassShaders.RayHitGroupShader.IsValid())
 			{
-				FMeshDrawSingleShaderBindings RayHitGroupShaderBindings = RayTracingMeshCommand.ShaderBindings.GetSingleShaderBindings(SF_RayHitGroup);
-				PassShaders.RayHitGroupShader->GetElementShaderBindings(Scene, ViewIfDynamicMeshCommand, VertexFactory, EVertexInputStreamType::Default, FeatureLevel, PrimitiveSceneProxy, MeshBatch, BatchElement, ShaderElementData, RayHitGroupShaderBindings, VertexStreams);
+				FMeshDrawSingleShaderBindings RayHitGroupShaderBindings = RayTracingMeshCommand.ShaderBindings.GetSingleShaderBindings(SF_RayHitGroup, DataOffset);
+				FMeshMaterialShader::GetElementShaderBindings(PassShaders.RayHitGroupShader, Scene, ViewIfDynamicMeshCommand, VertexFactory, EVertexInputStreamType::Default, FeatureLevel, PrimitiveSceneProxy, MeshBatch, BatchElement, ShaderElementData, RayHitGroupShaderBindings, VertexStreams);
 			}
 
 			int32 GeometrySegmentIndex = MeshBatch.SegmentIndex + BatchElementIndex;
@@ -471,7 +468,7 @@ FRayTracingPipelineState* FDeferredShadingSceneRenderer::BindRayTracingMaterialP
 
 	const bool bLightingMissShader = CanUseRayTracingLightingMissShader(View.GetShaderPlatform());
 
-	FRHIRayTracingShader* DefaultMissShader = View.ShaderMap->GetShader<FDefaultMainMS>()->GetRayTracingShader();
+	FRHIRayTracingShader* DefaultMissShader = View.ShaderMap->GetShader<FDefaultMainMS>().GetRayTracingShader();
 
 	FRHIRayTracingShader* RayTracingMissShaderLibrary[RAY_TRACING_NUM_MISS_SHADER_SLOTS];
 	RayTracingMissShaderLibrary[RAY_TRACING_MISS_SHADER_SLOT_DEFAULT] = DefaultMissShader;
@@ -486,15 +483,15 @@ FRayTracingPipelineState* FDeferredShadingSceneRenderer::BindRayTracingMaterialP
 
 	if (bEnableMaterials)
 	{
-		FShaderResource::GetRayTracingMaterialLibrary(RayTracingMaterialLibrary, DefaultClosestHitShader);
+		FShaderMapResource::GetRayTracingMaterialLibrary(RayTracingMaterialLibrary, DefaultClosestHitShader);
 	}
 	else
 	{
 		RayTracingMaterialLibrary.Add(DefaultClosestHitShader);
 	}
 
-	int32 OpaqueShadowMaterialIndex = RayTracingMaterialLibrary.Add(View.ShaderMap->GetShader<FOpaqueShadowHitGroup>()->GetRayTracingShader());
-	int32 HiddenMaterialIndex = RayTracingMaterialLibrary.Add(View.ShaderMap->GetShader<FHiddenMaterialHitGroup>()->GetRayTracingShader());
+	int32 OpaqueShadowMaterialIndex = RayTracingMaterialLibrary.Add(View.ShaderMap->GetShader<FOpaqueShadowHitGroup>().GetRayTracingShader());
+	int32 HiddenMaterialIndex = RayTracingMaterialLibrary.Add(View.ShaderMap->GetShader<FHiddenMaterialHitGroup>().GetRayTracingShader());
 
 	Initializer.SetHitGroupTable(RayTracingMaterialLibrary);
 
diff --git a/Engine/Source/Runtime/Renderer/Private/RayTracing/RayTracingPrimaryRays.cpp b/Engine/Source/Runtime/Renderer/Private/RayTracing/RayTracingPrimaryRays.cpp
--- a/Engine/Source/Runtime/Renderer/Private/RayTracing/RayTracingPrimaryRays.cpp
+++ b/Engine/Source/Runtime/Renderer/Private/RayTracing/RayTracingPrimaryRays.cpp
@@ -82,7 +82,7 @@ void FDeferredShadingSceneRenderer::PrepareRayTracingTranslucency(const FViewInf
 	PermutationVector.Set<FRayTracingPrimaryRaysRGS::FEnableTwoSidedGeometryForShadowDim>(EnableRayTracingShadowTwoSidedGeometry());
 
 	auto RayGenShader = View.ShaderMap->GetShader<FRayTracingPrimaryRaysRGS>(PermutationVector);
-	OutRayGenShaders.Add(RayGenShader->GetRayTracingShader());
+	OutRayGenShaders.Add(RayGenShader.GetRayTracingShader());
 }
 
 void FDeferredShadingSceneRenderer::RenderRayTracingPrimaryRaysView(
@@ -187,7 +187,7 @@ void FDeferredShadingSceneRenderer::RenderRayTracingPrimaryRaysView(
 		SetShaderParameters(GlobalResources, RayGenShader, *PassParameters);
 
 		FRHIRayTracingScene* RayTracingSceneRHI = View.RayTracingScene.RayTracingSceneRHI;
-		RHICmdList.RayTraceDispatch(Pipeline, RayGenShader->GetRayTracingShader(), RayTracingSceneRHI, GlobalResources, RayTracingResolution.X, RayTracingResolution.Y);
+		RHICmdList.RayTraceDispatch(Pipeline, RayGenShader.GetRayTracingShader(), RayTracingSceneRHI, GlobalResources, RayTracingResolution.X, RayTracingResolution.Y);
 	});
 }
 
diff --git a/Engine/Source/Runtime/Renderer/Private/RayTracing/RayTracingRectLight.cpp b/Engine/Source/Runtime/Renderer/Private/RayTracing/RayTracingRectLight.cpp
--- a/Engine/Source/Runtime/Renderer/Private/RayTracing/RayTracingRectLight.cpp
+++ b/Engine/Source/Runtime/Renderer/Private/RayTracing/RayTracingRectLight.cpp
@@ -119,7 +119,7 @@ public:
 		FRWBuffer& MipTree
 	)
 	{
-		FRHIComputeShader* ShaderRHI = GetComputeShader();
+		FRHIComputeShader* ShaderRHI = RHICmdList.GetBoundComputeShader();
 
 		SetShaderValue(RHICmdList, ShaderRHI, DimensionsParameter, Dimensions);
 		SetShaderValue(RHICmdList, ShaderRHI, MipLevelParameter, MipLevel);
@@ -136,30 +136,19 @@ public:
 		FRWBuffer& MipTree,
 		FRHIComputeFence* Fence)
 	{
-		FRHIComputeShader* ShaderRHI = GetComputeShader();
+		FRHIComputeShader* ShaderRHI = RHICmdList.GetBoundComputeShader();
 
 		MipTreeParameter.UnsetUAV(RHICmdList, ShaderRHI);
 		RHICmdList.TransitionResource(TransitionAccess, TransitionPipeline, MipTree.UAV, Fence);
 	}
 
-	virtual bool Serialize(FArchive& Ar)
-	{
-		bool bShaderHasOutdatedParameters = FGlobalShader::Serialize(Ar);
-		Ar << TextureParameter;
-		Ar << TextureSamplerParameter;
-		Ar << DimensionsParameter;
-		Ar << MipLevelParameter;
-		Ar << MipTreeParameter;
-		return bShaderHasOutdatedParameters;
-	}
-
 private:
-	FShaderResourceParameter TextureParameter;
-	FShaderResourceParameter TextureSamplerParameter;
+	LAYOUT_FIELD(FShaderResourceParameter, TextureParameter);
+	LAYOUT_FIELD(FShaderResourceParameter, TextureSamplerParameter);
 
-	FShaderParameter DimensionsParameter;
-	FShaderParameter MipLevelParameter;
-	FRWShaderParameter MipTreeParameter;
+	LAYOUT_FIELD(FShaderParameter, DimensionsParameter);
+	LAYOUT_FIELD(FShaderParameter, MipLevelParameter);
+	LAYOUT_FIELD(FRWShaderParameter, MipTreeParameter);
 };
 
 IMPLEMENT_SHADER_TYPE(, FBuildRectLightMipTreeCS, TEXT("/Engine/Private/Raytracing/BuildMipTreeCS.usf"), TEXT("BuildRectLightMipTreeCS"), SF_Compute)
@@ -176,7 +165,7 @@ FRectLightRayTracingData BuildRectLightMipTree(FRHICommandListImmediate& RHICmdL
 
 	const auto ShaderMap = GetGlobalShaderMap(ERHIFeatureLevel::SM5);
 	TShaderMapRef<FBuildRectLightMipTreeCS> BuildRectLightMipTreeComputeShader(ShaderMap);
-	RHICmdList.SetComputeShader(BuildRectLightMipTreeComputeShader->GetComputeShader());
+	RHICmdList.SetComputeShader(BuildRectLightMipTreeComputeShader.GetComputeShader());
 
 	// Allocate MIP tree
 	FIntVector TextureSize = RhiTexture->GetSizeXYZ();
@@ -198,7 +187,7 @@ FRectLightRayTracingData BuildRectLightMipTree(FRHICommandListImmediate& RHICmdL
 		BuildRectLightMipTreeComputeShader->SetParameters(RHICmdList, RhiTexture, Data.RectLightMipTreeDimensions, MipLevel, Data.RectLightMipTree);
 		FIntVector MipLevelDimensions = FIntVector(Data.RectLightMipTreeDimensions.X >> MipLevel, Data.RectLightMipTreeDimensions.Y >> MipLevel, 1);
 		FIntVector NumGroups = FIntVector::DivideAndRoundUp(MipLevelDimensions, FBuildRectLightMipTreeCS::GetGroupSize());
-		DispatchComputeShader(RHICmdList, *BuildRectLightMipTreeComputeShader, NumGroups.X, NumGroups.Y, 1);
+		DispatchComputeShader(RHICmdList, BuildRectLightMipTreeComputeShader.GetShader(), NumGroups.X, NumGroups.Y, 1);
 		BuildRectLightMipTreeComputeShader->UnsetParameters(RHICmdList, EResourceTransitionAccess::ERWBarrier, EResourceTransitionPipeline::EComputeToCompute, Data.RectLightMipTree, MipLevelFence);
 	}
 	FComputeFenceRHIRef TransitionFence = RHICmdList.CreateComputeFence(TEXT("RectLightMipTree Transition"));
@@ -226,7 +215,7 @@ public:
 	}
 
 	FRectLightRGS() {}
-	virtual ~FRectLightRGS() {}
+	~FRectLightRGS() {}
 
 	FRectLightRGS(const ShaderMetaType::CompiledShaderInitializerType& Initializer)
 		: FGlobalShader(Initializer)
@@ -242,7 +231,7 @@ public:
 		RayDistanceUAVParameter.Bind(Initializer.ParameterMap, TEXT("RWRayDistanceUAV"));
 	}
 
-	bool Serialize(FArchive& Ar)
+	/*bool Serialize(FArchive& Ar)
 	{
 		bool bShaderHasOutdatedParameters = FGlobalShader::Serialize(Ar);
 		Ar << ViewParameter;
@@ -254,10 +243,11 @@ public:
 		Ar << LuminanceUAVParameter;
 		Ar << RayDistanceUAVParameter;
 		return bShaderHasOutdatedParameters;
-	}
+	}*/
 
 	void Dispatch(
 		FRHICommandListImmediate& RHICmdList,
+		FRHIRayTracingShader* ShaderRHI,
 		const FRayTracingScene& RayTracingScene,
 		FRHIUniformBuffer* ViewUniformBuffer,
 		FRHIUniformBuffer* SceneTexturesUniformBuffer,
@@ -269,7 +259,7 @@ public:
 	{
 		FRayTracingPipelineStateInitializer Initializer;
 
-		FRHIRayTracingShader* RayGenShaderTable[] = { GetRayTracingShader() };
+		FRHIRayTracingShader* RayGenShaderTable[] = { ShaderRHI };
 		Initializer.SetRayGenShaderTable(RayGenShaderTable);
 
 		FRayTracingPipelineState* Pipeline = PipelineStateCache::GetAndOrCreateRayTracingPipelineState(RHICmdList, Initializer);
@@ -298,23 +288,23 @@ public:
 			GlobalResources.SetSampler(TransmissionProfilesLinearSamplerParameter.GetBaseIndex(), TStaticSamplerState<SF_Bilinear, AM_Clamp, AM_Clamp, AM_Clamp>::GetRHI());
 		}
 
-		RHICmdList.RayTraceDispatch(Pipeline, GetRayTracingShader(), RayTracingScene.RayTracingSceneRHI, GlobalResources, Width, Height);
+		RHICmdList.RayTraceDispatch(Pipeline, ShaderRHI, RayTracingScene.RayTracingSceneRHI, GlobalResources, Width, Height);
 	}
 
 private:
 	// Input
-	FShaderResourceParameter TLASParameter;
-	FShaderUniformBufferParameter ViewParameter;
-	FShaderUniformBufferParameter SceneTexturesParameter;
-	FShaderUniformBufferParameter RectLightParameter;
+	LAYOUT_FIELD(FShaderResourceParameter, TLASParameter);
+	LAYOUT_FIELD(FShaderUniformBufferParameter, ViewParameter);
+	LAYOUT_FIELD(FShaderUniformBufferParameter, SceneTexturesParameter);
+	LAYOUT_FIELD(FShaderUniformBufferParameter, RectLightParameter);
 
 	// SSS Profile
-	FShaderResourceParameter TransmissionProfilesTextureParameter;
-	FShaderResourceParameter TransmissionProfilesLinearSamplerParameter;
+	LAYOUT_FIELD(FShaderResourceParameter, TransmissionProfilesTextureParameter);
+	LAYOUT_FIELD(FShaderResourceParameter, TransmissionProfilesLinearSamplerParameter);
 
 	// Output
-	FShaderResourceParameter LuminanceUAVParameter;
-	FShaderResourceParameter RayDistanceUAVParameter;
+	LAYOUT_FIELD(FShaderResourceParameter, LuminanceUAVParameter);
+	LAYOUT_FIELD(FShaderResourceParameter, RayDistanceUAVParameter);
 };
 
 #define IMPLEMENT_RECT_LIGHT_TYPE(TextureImportanceSampling) \
@@ -360,24 +350,16 @@ public:
 		const FRWBuffer& MipTree,
 		const FIntVector Dimensions)
 	{
-		FRHIPixelShader* ShaderRHI = GetPixelShader();
+		FRHIPixelShader* ShaderRHI = RHICmdList.GetBoundPixelShader();
 		FGlobalShader::SetParameters<FViewUniformShaderParameters>(RHICmdList, ShaderRHI, View.ViewUniformBuffer);
 
 		SetShaderValue(RHICmdList, ShaderRHI, DimensionsParameter, Dimensions);
 		SetSRVParameter(RHICmdList, ShaderRHI, MipTreeParameter, MipTree.SRV);
 	}
 
-	virtual bool Serialize(FArchive& Ar) override
-	{
-		bool bShaderHasOutdatedParameters = FGlobalShader::Serialize(Ar);
-		Ar << DimensionsParameter;
-		Ar << MipTreeParameter;
-		return bShaderHasOutdatedParameters;
-	}
-
 private:
-	FShaderParameter DimensionsParameter;
-	FShaderResourceParameter MipTreeParameter;
+	LAYOUT_FIELD(FShaderParameter, DimensionsParameter);
+	LAYOUT_FIELD(FShaderResourceParameter, MipTreeParameter);
 };
 
 IMPLEMENT_SHADER_TYPE(, FVisualizeRectLightMipTreePS, TEXT("/Engine/Private/PathTracing/VisualizeMipTreePixelShader.usf"), TEXT("VisualizeMipTreePS"), SF_Pixel)
@@ -415,8 +397,8 @@ void FDeferredShadingSceneRenderer::VisualizeRectLightMipTree(
 	GraphicsPSOInit.RasterizerState = TStaticRasterizerState<FM_Solid, CM_None>::GetRHI();
 	GraphicsPSOInit.DepthStencilState = TStaticDepthStencilState<false, CF_Always>::GetRHI();
 	GraphicsPSOInit.BoundShaderState.VertexDeclarationRHI = GFilterVertexDeclaration.VertexDeclarationRHI;
-	GraphicsPSOInit.BoundShaderState.VertexShaderRHI = GETSAFERHISHADER_VERTEX(*VertexShader);
-	GraphicsPSOInit.BoundShaderState.PixelShaderRHI = GETSAFERHISHADER_PIXEL(*PixelShader);
+	GraphicsPSOInit.BoundShaderState.VertexShaderRHI = VertexShader.GetVertexShader();
+	GraphicsPSOInit.BoundShaderState.PixelShaderRHI = PixelShader.GetPixelShader();
 	GraphicsPSOInit.PrimitiveType = PT_TriangleList;
 	SetGraphicsPipelineState(RHICmdList, GraphicsPSOInit);
 
@@ -434,7 +416,7 @@ void FDeferredShadingSceneRenderer::VisualizeRectLightMipTree(
 		View.ViewRect.Width(), View.ViewRect.Height(),
 		FIntPoint(View.ViewRect.Width(), View.ViewRect.Height()),
 		SceneContext.GetBufferSizeXY(),
-		*VertexShader);
+		VertexShader);
 	ResolveSceneColor(RHICmdList);
 	RHICmdList.EndRenderPass();
 	GVisualizeTexture.SetCheckPoint(RHICmdList, RectLightMipTreeRT);
@@ -450,8 +432,8 @@ void FDeferredShadingSceneRenderer::PrepareRayTracingRectLight(const FViewInfo&
 	TShaderMapRef<FRectLightRGS<0>> Shader0(GetGlobalShaderMap(View.FeatureLevel));
 	TShaderMapRef<FRectLightRGS<1>> Shader1(GetGlobalShaderMap(View.FeatureLevel));
 
-	OutRayGenShaders.Add(Shader0->GetRayTracingShader());
-	OutRayGenShaders.Add(Shader1->GetRayTracingShader());
+	OutRayGenShaders.Add(Shader0.GetRayTracingShader());
+	OutRayGenShaders.Add(Shader1.GetRayTracingShader());
 }
 
 template <int TextureImportanceSampling>
@@ -537,6 +519,7 @@ void FDeferredShadingSceneRenderer::RenderRayTracingRectLightInternal(
 		// Dispatch
 		RectLightRayGenerationShader->Dispatch(
 			RHICmdList,
+			RectLightRayGenerationShader.GetRayTracingShader(),
 			View.RayTracingScene,
 			View.ViewUniformBuffer,
 			SceneTexturesUniformBuffer,
diff --git a/Engine/Source/Runtime/Renderer/Private/RayTracing/RayTracingReflections.cpp b/Engine/Source/Runtime/Renderer/Private/RayTracing/RayTracingReflections.cpp
--- a/Engine/Source/Runtime/Renderer/Private/RayTracing/RayTracingReflections.cpp
+++ b/Engine/Source/Runtime/Renderer/Private/RayTracing/RayTracingReflections.cpp
@@ -379,7 +379,7 @@ void FDeferredShadingSceneRenderer::PrepareRayTracingReflections(const FViewInfo
 			PermutationVector.Set<FRayTracingReflectionsRGS::FMissShaderLighting>(bMissShaderLighting);
 			PermutationVector.Set<FRayTracingReflectionsRGS::FRayTraceSkyLightContribution>(bRayTraceSkyLightContribution);
 			auto RayGenShader = View.ShaderMap->GetShader<FRayTracingReflectionsRGS>(PermutationVector);
-			OutRayGenShaders.Add(RayGenShader->GetRayTracingShader());
+			OutRayGenShaders.Add(RayGenShader.GetRayTracingShader());
 		}
 		
 		{
@@ -390,7 +390,7 @@ void FDeferredShadingSceneRenderer::PrepareRayTracingReflections(const FViewInfo
 			PermutationVector.Set<FRayTracingReflectionsRGS::FMissShaderLighting>(bMissShaderLighting);
 			PermutationVector.Set<FRayTracingReflectionsRGS::FRayTraceSkyLightContribution>(bRayTraceSkyLightContribution);
 			auto RayGenShader = View.ShaderMap->GetShader<FRayTracingReflectionsRGS>(PermutationVector);
-			OutRayGenShaders.Add(RayGenShader->GetRayTracingShader());
+			OutRayGenShaders.Add(RayGenShader.GetRayTracingShader());
 		}
 	}
 	else
@@ -401,7 +401,7 @@ void FDeferredShadingSceneRenderer::PrepareRayTracingReflections(const FViewInfo
 		PermutationVector.Set<FRayTracingReflectionsRGS::FMissShaderLighting>(bMissShaderLighting);
 		PermutationVector.Set<FRayTracingReflectionsRGS::FRayTraceSkyLightContribution>(bRayTraceSkyLightContribution);
 		auto RayGenShader = View.ShaderMap->GetShader<FRayTracingReflectionsRGS>(PermutationVector);
-		OutRayGenShaders.Add(RayGenShader->GetRayTracingShader());
+		OutRayGenShaders.Add(RayGenShader.GetRayTracingShader());
 	}
 }
 
@@ -466,7 +466,7 @@ void FDeferredShadingSceneRenderer::RenderRayTracingReflections(
 			/* bInForceSeparateTargetAndShaderResource = */ false);
 
 		OutDenoiserInputs->Color = GraphBuilder.CreateTexture(Desc, TEXT("RayTracingReflections"));
-
+		
 		Desc.Format = PF_R16F;
 		OutDenoiserInputs->RayHitDistance = GraphBuilder.CreateTexture(Desc, TEXT("RayTracingReflectionsHitDistance"));
 		OutDenoiserInputs->RayImaginaryDepth = GraphBuilder.CreateTexture(Desc, TEXT("RayTracingReflectionsImaginaryDepth"));
@@ -612,13 +612,13 @@ void FDeferredShadingSceneRenderer::RenderRayTracingReflections(
 					ERDGPassFlags::Compute,
 					[PassParameters, this, &View, RayGenShader, TileAlignedResolution](FRHICommandList& RHICmdList)
 				{
-					FRayTracingPipelineState* Pipeline = BindRayTracingDeferredMaterialGatherPipeline(RHICmdList, View, RayGenShader->GetRayTracingShader());
+					FRayTracingPipelineState* Pipeline = BindRayTracingDeferredMaterialGatherPipeline(RHICmdList, View, RayGenShader.GetRayTracingShader());
 
 					FRayTracingShaderBindingsWriter GlobalResources;
 					SetShaderParameters(GlobalResources, RayGenShader, *PassParameters);
 
 					FRHIRayTracingScene* RayTracingSceneRHI = View.RayTracingScene.RayTracingSceneRHI;
-					RHICmdList.RayTraceDispatch(Pipeline, RayGenShader->GetRayTracingShader(), RayTracingSceneRHI, GlobalResources, TileAlignedResolution.X, TileAlignedResolution.Y);
+					RHICmdList.RayTraceDispatch(Pipeline, RayGenShader.GetRayTracingShader(), RayTracingSceneRHI, GlobalResources, TileAlignedResolution.X, TileAlignedResolution.Y);
 				});
 
 				// A material sorting pass
@@ -649,11 +649,11 @@ void FDeferredShadingSceneRenderer::RenderRayTracingReflections(
 						{
 							// Shading pass for sorted materials uses 1D dispatch over all elements in the material buffer.
 							// This can be reduced to the number of output pixels if sorting pass guarantees that all invalid entries are moved to the end.
-							RHICmdList.RayTraceDispatch(View.RayTracingMaterialPipeline, RayGenShader->GetRayTracingShader(), RayTracingSceneRHI, GlobalResources, DeferredMaterialBufferNumElements, 1);
+							RHICmdList.RayTraceDispatch(View.RayTracingMaterialPipeline, RayGenShader.GetRayTracingShader(), RayTracingSceneRHI, GlobalResources, DeferredMaterialBufferNumElements, 1);
 						}
 						else // EDeferredMaterialMode::None
 						{
-							RHICmdList.RayTraceDispatch(View.RayTracingMaterialPipeline, RayGenShader->GetRayTracingShader(), RayTracingSceneRHI, GlobalResources, RayTracingResolution.X, RayTracingResolution.Y);
+							RHICmdList.RayTraceDispatch(View.RayTracingMaterialPipeline, RayGenShader.GetRayTracingShader(), RayTracingSceneRHI, GlobalResources, RayTracingResolution.X, RayTracingResolution.Y);
 						}
 					});
 				}
@@ -687,7 +687,7 @@ void FDeferredShadingSceneRenderer::RenderRayTracingReflections(
 								SetShaderParameters(GlobalResources, RayGenShader, *TilePassParameters);
 
 								FRHIRayTracingScene* RayTracingSceneRHI = View.RayTracingScene.RayTracingSceneRHI;
-								RHICmdList.RayTraceDispatch(View.RayTracingMaterialPipeline, RayGenShader->GetRayTracingShader(), RayTracingSceneRHI, GlobalResources, DispatchSizeX, DispatchSizeY);
+								RHICmdList.RayTraceDispatch(View.RayTracingMaterialPipeline, RayGenShader.GetRayTracingShader(), RayTracingSceneRHI, GlobalResources, DispatchSizeX, DispatchSizeY);
 							});
 						}
 					}
@@ -740,7 +740,7 @@ void FDeferredShadingSceneRenderer::RenderRayTracingReflections(
 		FComputeShaderUtils::AddPass(
 			GraphBuilder,
 			RDG_EVENT_NAME("SplitImaginaryReflectionGBuffer"),
-			*ComputeShader,
+			ComputeShader,
 			PassParameters,
 			FComputeShaderUtils::GetGroupCount(FIntPoint(RayTracingResolution.X, RayTracingResolution.Y), FSplitImaginaryReflectionGBufferCS::GetGroupSize())
 		);
