Copying Private-LoadTimes-4.24 stream to Main. Biggest changes are in Materials/Shader memory freezing.

diff --git a/Engine/Source/Runtime/Renderer/Private/SceneRendering.cpp b/Engine/Source/Runtime/Renderer/Private/SceneRendering.cpp
--- a/Engine/Source/Runtime/Renderer/Private/SceneRendering.cpp
+++ b/Engine/Source/Runtime/Renderer/Private/SceneRendering.cpp
@@ -1780,6 +1780,11 @@ FViewInfo* FViewInfo::CreateSnapshot() const
 	TStaticArray<FParallelMeshDrawCommandPass, EMeshPass::Num> NullParallelMeshDrawCommandPasses;
 	FMemory::Memcpy(Result->ParallelMeshDrawCommandPasses, NullParallelMeshDrawCommandPasses);
 
+	for (int i = 0; i < EMeshPass::Num; i++)
+	{
+		Result->ParallelMeshDrawCommandPasses[i].InitCreateSnapshot();
+	}
+	
 	Result->bIsSnapshot = true;
 	ViewInfoSnapshots.Add(Result);
 	return Result;
@@ -2335,7 +2340,7 @@ void FSceneRenderer::PrepareViewRectsForRendering()
 			TEXT("ISceneViewFamilyScreenPercentage::GetPrimaryResolutionFractionUpperBound() should not lie to the renderer."));
 		
 		check(FSceneViewScreenPercentageConfig::IsValidResolutionFraction(PrimaryResolutionFraction));
-		
+
 		// Compute final resolution fraction.
 		float ResolutionFraction = PrimaryResolutionFraction * ViewFamily.SecondaryViewFraction;
 
@@ -2474,7 +2479,7 @@ void FSceneRenderer::ComputeViewGPUMasks(FRHIGPUMask RenderTargetGPUMask)
 				{
 					ViewInfo.GPUMask = FRHIGPUMask::FromIndex(*GPUIterator);
 					ViewFamily.bMultiGPUForkAndJoin |= (ViewInfo.GPUMask != RenderTargetGPUMask);
-
+					
 					// Increment and wrap around if we reach the last index.
 					++GPUIterator;
 					if (!GPUIterator)
@@ -2594,10 +2599,7 @@ bool FSceneRenderer::DoOcclusionQueries(ERHIFeatureLevel::Type InFeatureLevel) c
 }
 
 FSceneRenderer::~FSceneRenderer()
-{
-	// To prevent keeping persistent references to single frame buffers, clear any such reference at this point.
-	ClearPrimitiveSingleFrameIndirectLightingCacheBuffers();
-	
+{	
 	if(Scene)
 	{
 		// Destruct the projected shadow infos.
@@ -2710,7 +2712,7 @@ void FSceneRenderer::RenderFinish(FRHICommandListImmediate& RHICmdList)
 				bMobileShowVertexFogWarning = true;
 			}
 		}
-
+		
 		const bool bSingleLayerWaterWarning = ShouldRenderSingleLayerWaterSkippedRenderEditorNotification(Views);
 		
 		const bool bAnyWarning = bShowPrecomputedVisibilityWarning || bShowGlobalClipPlaneWarning || bShowAtmosphericFogWarning || bShowSkylightWarning || bShowPointLightWarning 
@@ -3169,7 +3171,7 @@ void FSceneRenderer::RenderCustomDepthPass(FRHICommandListImmediate& RHICmdList)
 						Scene->UniformBuffers.InstancedCustomDepthViewUniformBuffer.UpdateUniformBufferImmediate(reinterpret_cast<FInstancedViewUniformShaderParameters&>(*View.CachedViewUniformShaderParameters));
 					}
 				}
-
+	
 				extern TSet<IPersistentViewUniformBufferExtension*> PersistentViewUniformBufferExtensions;
 				
 				for (IPersistentViewUniformBufferExtension* Extension : PersistentViewUniformBufferExtensions)
@@ -3360,24 +3362,6 @@ void FSceneRenderer::UpdatePrimitiveIndirectLightingCacheBuffers()
 	}
 }
 
-void FSceneRenderer::ClearPrimitiveSingleFrameIndirectLightingCacheBuffers()
-{
-	for (int32 ViewIndex = 0; ViewIndex < Views.Num(); ViewIndex++)
-	{		
-		FViewInfo& View = Views[ViewIndex];
-
-		for (int32 Index = 0; Index < View.DirtyIndirectLightingCacheBufferPrimitives.Num(); ++Index)
-		{
-			FPrimitiveSceneInfo* PrimitiveSceneInfo = View.DirtyIndirectLightingCacheBufferPrimitives[Index];
-	
-			if (PrimitiveSceneInfo) // Could be null if it was a duplicate.
-			{
-				PrimitiveSceneInfo->ClearIndirectLightingCacheBuffer(true);
-			}
-		}
-	}
-}
-
 /*-----------------------------------------------------------------------------
 	FRendererModule
 -----------------------------------------------------------------------------*/
@@ -3669,7 +3653,7 @@ void FRendererModule::BeginRenderingViewFamily(FCanvas* Canvas, FSceneViewFamily
 	}
 
 	if (Scene)
-	{
+	{		
 		// Set the world's "needs full lighting rebuild" flag if the scene has any uncached static lighting interactions.
 		if(World)
 		{
@@ -3748,7 +3732,7 @@ void FRendererModule::DrawRectangle(
 		float SizeV,
 		FIntPoint TargetSize,
 		FIntPoint TextureSize,
-		class FShader* VertexShader,
+		const TShaderRef<FShader>& VertexShader,
 		EDrawRectangleFlags Flags
 		)
 {
@@ -4095,7 +4079,7 @@ void FSceneRenderer::ResolveSceneColor(FRHICommandList& RHICmdList)
 					auto ShaderMap = GetGlobalShaderMap(SceneContext.GetCurrentFeatureLevel());
 					TShaderMapRef<FHdrCustomResolveVS> VertexShader(ShaderMap);
 					GraphicsPSOInit.BoundShaderState.VertexDeclarationRHI = GetVertexDeclarationFVector4();
-					GraphicsPSOInit.BoundShaderState.VertexShaderRHI = GETSAFERHISHADER_VERTEX(*VertexShader);
+					GraphicsPSOInit.BoundShaderState.VertexShaderRHI = VertexShader.GetVertexShader();
 					GraphicsPSOInit.PrimitiveType = PT_TriangleList;
 
 					if (SceneColor.FmaskSRV)
@@ -4103,7 +4087,7 @@ void FSceneRenderer::ResolveSceneColor(FRHICommandList& RHICmdList)
 						if (CurrentNumSamples == 2)
 						{
 							TShaderMapRef<FHdrCustomResolveFMask2xPS> PixelShader(ShaderMap);
-							GraphicsPSOInit.BoundShaderState.PixelShaderRHI = GETSAFERHISHADER_PIXEL(*PixelShader);
+							GraphicsPSOInit.BoundShaderState.PixelShaderRHI = PixelShader.GetPixelShader();
 
 							SetGraphicsPipelineState(RHICmdList, GraphicsPSOInit);
 							PixelShader->SetParameters(RHICmdList, SceneColor.TargetableTexture, SceneColor.FmaskSRV);
@@ -4111,7 +4095,7 @@ void FSceneRenderer::ResolveSceneColor(FRHICommandList& RHICmdList)
 						else if (CurrentNumSamples == 4)
 						{
 							TShaderMapRef<FHdrCustomResolveFMask4xPS> PixelShader(ShaderMap);
-							GraphicsPSOInit.BoundShaderState.PixelShaderRHI = GETSAFERHISHADER_PIXEL(*PixelShader);
+							GraphicsPSOInit.BoundShaderState.PixelShaderRHI = PixelShader.GetPixelShader();
 
 							SetGraphicsPipelineState(RHICmdList, GraphicsPSOInit);
 							PixelShader->SetParameters(RHICmdList, SceneColor.TargetableTexture, SceneColor.FmaskSRV);
@@ -4119,7 +4103,7 @@ void FSceneRenderer::ResolveSceneColor(FRHICommandList& RHICmdList)
 						else if (CurrentNumSamples == 8)
 						{
 							TShaderMapRef<FHdrCustomResolveFMask8xPS> PixelShader(ShaderMap);
-							GraphicsPSOInit.BoundShaderState.PixelShaderRHI = GETSAFERHISHADER_PIXEL(*PixelShader);
+							GraphicsPSOInit.BoundShaderState.PixelShaderRHI = PixelShader.GetPixelShader();
 
 							SetGraphicsPipelineState(RHICmdList, GraphicsPSOInit);
 							PixelShader->SetParameters(RHICmdList, SceneColor.TargetableTexture, SceneColor.FmaskSRV);
@@ -4136,7 +4120,7 @@ void FSceneRenderer::ResolveSceneColor(FRHICommandList& RHICmdList)
 						if (CurrentNumSamples == 2)
 						{
 							TShaderMapRef<FHdrCustomResolve2xPS> PixelShader(ShaderMap);
-							GraphicsPSOInit.BoundShaderState.PixelShaderRHI = GETSAFERHISHADER_PIXEL(*PixelShader);
+							GraphicsPSOInit.BoundShaderState.PixelShaderRHI = PixelShader.GetPixelShader();
 
 							SetGraphicsPipelineState(RHICmdList, GraphicsPSOInit);
 							PixelShader->SetParameters(RHICmdList, SceneColor.TargetableTexture);
@@ -4144,7 +4128,7 @@ void FSceneRenderer::ResolveSceneColor(FRHICommandList& RHICmdList)
 						else if (CurrentNumSamples == 4)
 						{
 							TShaderMapRef<FHdrCustomResolve4xPS> PixelShader(ShaderMap);
-							GraphicsPSOInit.BoundShaderState.PixelShaderRHI = GETSAFERHISHADER_PIXEL(*PixelShader);
+							GraphicsPSOInit.BoundShaderState.PixelShaderRHI = PixelShader.GetPixelShader();
 
 							SetGraphicsPipelineState(RHICmdList, GraphicsPSOInit);
 							PixelShader->SetParameters(RHICmdList, SceneColor.TargetableTexture);
@@ -4152,7 +4136,7 @@ void FSceneRenderer::ResolveSceneColor(FRHICommandList& RHICmdList)
 						else if (CurrentNumSamples == 8)
 						{
 							TShaderMapRef<FHdrCustomResolve8xPS> PixelShader(ShaderMap);
-							GraphicsPSOInit.BoundShaderState.PixelShaderRHI = GETSAFERHISHADER_PIXEL(*PixelShader);
+							GraphicsPSOInit.BoundShaderState.PixelShaderRHI = PixelShader.GetPixelShader();
 
 							SetGraphicsPipelineState(RHICmdList, GraphicsPSOInit);
 							PixelShader->SetParameters(RHICmdList, SceneColor.TargetableTexture);
diff --git a/Engine/Source/Runtime/Renderer/Private/SceneRendering.h b/Engine/Source/Runtime/Renderer/Private/SceneRendering.h
--- a/Engine/Source/Runtime/Renderer/Private/SceneRendering.h
+++ b/Engine/Source/Runtime/Renderer/Private/SceneRendering.h
@@ -39,7 +39,6 @@ class FViewInfo;
 struct FILCUpdatePrimTaskData;
 class FPostprocessContext;
 struct FILCUpdatePrimTaskData;
-template<typename ShaderMetaType> class TShaderMap;
 class FRaytracingLightDataPacked;
 class FRayTracingLocalShaderBindingWriter;
 
@@ -1133,7 +1132,7 @@ public:
 
 	FHeightfieldLightingViewInfo HeightfieldLightingViewInfo;
 
-	TShaderMap<FGlobalShaderType>* ShaderMap;
+	FGlobalShaderMap* ShaderMap;
 
 	bool bIsSnapshot;
 
@@ -1773,7 +1772,6 @@ protected:
 	FRHITexture* GetMultiViewSceneColor(const FSceneRenderTargets& SceneContext) const;
 
 	void UpdatePrimitiveIndirectLightingCacheBuffers();
-	void ClearPrimitiveSingleFrameIndirectLightingCacheBuffers();
 
 	void RenderPlanarReflection(class FPlanarReflectionSceneProxy* ReflectionSceneProxy);
 
diff --git a/Engine/Source/Runtime/Renderer/Private/SceneTextureReductions.cpp b/Engine/Source/Runtime/Renderer/Private/SceneTextureReductions.cpp
--- a/Engine/Source/Runtime/Renderer/Private/SceneTextureReductions.cpp
+++ b/Engine/Source/Runtime/Renderer/Private/SceneTextureReductions.cpp
@@ -177,7 +177,7 @@ void BuildHZB(FRDGBuilder& GraphBuilder, const FSceneTextureParameters& SceneTex
 			TShaderMapRef<FHZBBuildCS> ComputeShader(View.ShaderMap, PermutationVector);
 
 			// TODO(RDG): remove ERDGPassFlags::GenerateMips to use FComputeShaderUtils::AddPass().
-			ClearUnusedGraphResources(*ComputeShader, PassParameters);
+			ClearUnusedGraphResources(ComputeShader, PassParameters);
 			GraphBuilder.AddPass(
 				RDG_EVENT_NAME("ReduceHZB(mips=[%d;%d]%s%s) %dx%d",
 					StartDestMip, EndDestMip - 1,
@@ -188,7 +188,7 @@ void BuildHZB(FRDGBuilder& GraphBuilder, const FSceneTextureParameters& SceneTex
 				StartDestMip ? (ERDGPassFlags::Compute | ERDGPassFlags::GenerateMips) : ERDGPassFlags::Compute,
 				[PassParameters, ComputeShader, DstSize](FRHICommandList& RHICmdList)
 			{
-				FComputeShaderUtils::Dispatch(RHICmdList, *ComputeShader, *PassParameters, FComputeShaderUtils::GetGroupCount(DstSize, 8));
+				FComputeShaderUtils::Dispatch(RHICmdList, ComputeShader, *PassParameters, FComputeShaderUtils::GetGroupCount(DstSize, 8));
 			});
 		}
 		else
@@ -203,14 +203,14 @@ void BuildHZB(FRDGBuilder& GraphBuilder, const FSceneTextureParameters& SceneTex
 			TShaderMapRef<FHZBBuildPS> PixelShader(View.ShaderMap);
 
 			// TODO(RDG): remove ERDGPassFlags::GenerateMips to use FPixelShaderUtils::AddFullscreenPass().
-			ClearUnusedGraphResources(*PixelShader, PassParameters);
+			ClearUnusedGraphResources(PixelShader, PassParameters);
 			GraphBuilder.AddPass(
 				RDG_EVENT_NAME("DownsampleHZB(mip=%d) %dx%d", StartDestMip, DstSize.X, DstSize.Y),
 				PassParameters,
 				StartDestMip ? (ERDGPassFlags::Raster | ERDGPassFlags::GenerateMips) : ERDGPassFlags::Raster,
 				[PassParameters, &View, PixelShader, DstSize](FRHICommandList& RHICmdList)
 			{
-				FPixelShaderUtils::DrawFullscreenPixelShader(RHICmdList, View.ShaderMap, *PixelShader, *PassParameters, FIntRect(0, 0, DstSize.X, DstSize.Y));
+				FPixelShaderUtils::DrawFullscreenPixelShader(RHICmdList, View.ShaderMap, PixelShader, *PassParameters, FIntRect(0, 0, DstSize.X, DstSize.Y));
 			});
 		}
 	};
diff --git a/Engine/Source/Runtime/Renderer/Private/SceneVisibility.cpp b/Engine/Source/Runtime/Renderer/Private/SceneVisibility.cpp
--- a/Engine/Source/Runtime/Renderer/Private/SceneVisibility.cpp
+++ b/Engine/Source/Runtime/Renderer/Private/SceneVisibility.cpp
@@ -1795,7 +1795,7 @@ struct FDrawCommandRelevancePacket
 				FVisibleMeshDrawCommand& NewVisibleMeshDrawCommand = VisibleCachedDrawCommands[(uint32)PassType].Last();
 
 				const FMeshDrawCommand* MeshDrawCommand = CachedMeshDrawCommand.StateBucketId >= 0
-					? &Scene->CachedMeshDrawCommandStateBuckets[FSetElementId::FromInteger(CachedMeshDrawCommand.StateBucketId)].MeshDrawCommand
+					? &Scene->CachedMeshDrawCommandStateBuckets[PassType].GetByElementId(CachedMeshDrawCommand.StateBucketId).Key
 					: &SceneDrawList.MeshDrawCommands[CachedMeshDrawCommand.CommandIndex];
 
 				NewVisibleMeshDrawCommand.Setup(
@@ -2505,6 +2505,8 @@ static void ComputeAndMarkRelevanceForViewParallel(
 	FPrimitiveViewMasks& HasViewCustomDataMasks
 	)
 {
+	SCOPED_NAMED_EVENT(FSceneRenderer_ComputeAndMarkRelevanceForViewParallel, FColor::Blue);
+
 	check(OutHasDynamicMeshElementsMasks.Num() == Scene->Primitives.Num());
 
 	FFrozenSceneViewMatricesGuard FrozenMatricesGuard(View);
@@ -2998,10 +3000,12 @@ void FSceneRenderer::PreVisibilityFrameSetup(FRHICommandListImmediate& RHICmdLis
 				Scene->CachedDrawLists[RollingPassShrinkIndex].MeshDrawCommands.Shrink();
 			}
 			const int32 NumRemovedPerFrame = 10;
+			TArray<FPrimitiveSceneInfo*, TInlineAllocator<10>> SceneInfos;
 			for (int32 NumRemoved = 0; NumRemoved < NumRemovedPerFrame && RollingRemoveIndex < Scene->Primitives.Num(); NumRemoved++, RollingRemoveIndex++)
 			{
-				Scene->Primitives[RollingRemoveIndex]->UpdateStaticMeshes(RHICmdList, false);
+				SceneInfos.Add(Scene->Primitives[RollingRemoveIndex]);
 			}
+			FPrimitiveSceneInfo::UpdateStaticMeshes(RHICmdList, Scene, SceneInfos, false);
 		}
 	}
 
@@ -3493,6 +3497,8 @@ static TAutoConsoleVariable<int32> CVarUseFastIntersect(
 void UpdateReflectionSceneData(FScene* Scene)
 {
 	QUICK_SCOPE_CYCLE_COUNTER(STAT_UpdateReflectionSceneData)
+	SCOPED_NAMED_EVENT(UpdateReflectionScene, FColor::Red);
+
 
 	FReflectionEnvironmentSceneData& ReflectionSceneData = Scene->ReflectionSceneData;
 
@@ -3647,14 +3653,20 @@ void FSceneRenderer::ComputeViewVisibility(FRHICommandListImmediate& RHICmdList,
 
 	{
 		QUICK_SCOPE_CYCLE_COUNTER(STAT_ViewVisibilityTime_ConditionalUpdateStaticMeshesWithoutVisibilityCheck);
+		SCOPED_NAMED_EVENT(FSceneRenderer_ConditionalUpdateStaticMeshes, FColor::Red);
 
 		Scene->ConditionalMarkStaticMeshElementsForUpdate();
 
+		TArray<FPrimitiveSceneInfo*> UpdatedSceneInfos;
 		for (TSet<FPrimitiveSceneInfo*>::TIterator It(Scene->PrimitivesNeedingStaticMeshUpdateWithoutVisibilityCheck); It; ++It)
 		{
 			FPrimitiveSceneInfo* Primitive = *It;
-			Primitive->ConditionalUpdateStaticMeshes(RHICmdList);
+			if (Primitive->NeedsUpdateStaticMeshes())
+			{
+				UpdatedSceneInfos.Add(Primitive);
+			}
 		}
+		FPrimitiveSceneInfo::UpdateStaticMeshes(RHICmdList, Scene, UpdatedSceneInfos);
 		Scene->PrimitivesNeedingStaticMeshUpdateWithoutVisibilityCheck.Reset();
 	}
 
@@ -3876,11 +3888,15 @@ void FSceneRenderer::ComputeViewVisibility(FRHICommandListImmediate& RHICmdList,
 
 		{
 			QUICK_SCOPE_CYCLE_COUNTER(STAT_ViewVisibilityTime_ConditionalUpdateStaticMeshes);
+			SCOPED_NAMED_EVENT(FSceneRenderer_UpdateStaticMeshes, FColor::Red);
+
+			TArray<FPrimitiveSceneInfo*> AddedSceneInfos;
 			for (TConstDualSetBitIterator<SceneRenderingBitArrayAllocator, FDefaultBitArrayAllocator> BitIt(View.PrimitiveVisibilityMap, Scene->PrimitivesNeedingStaticMeshUpdate); BitIt; ++BitIt)
 			{
 				int32 PrimitiveIndex = BitIt.GetIndex();
-				Scene->Primitives[PrimitiveIndex]->UpdateStaticMeshes(RHICmdList);
+				AddedSceneInfos.Add(Scene->Primitives[PrimitiveIndex]);
 			}
+			FPrimitiveSceneInfo::UpdateStaticMeshes(RHICmdList, Scene, AddedSceneInfos);
 		}
 
 		// ISR views can't compute relevance until all views are frustum culled
@@ -3949,9 +3965,12 @@ void FSceneRenderer::ComputeViewVisibility(FRHICommandListImmediate& RHICmdList,
 		ViewBit <<= 1;
 	}
 
-	// Gather FMeshBatches from scene proxies
-	GatherDynamicMeshElements(Views, Scene, ViewFamily, DynamicIndexBuffer, DynamicVertexBuffer, DynamicReadBuffer,
-		HasDynamicMeshElementsMasks, HasDynamicEditorMeshElementsMasks, HasViewCustomDataMasks, MeshCollector);
+	{
+		SCOPED_NAMED_EVENT(FSceneRenderer_GatherDynamicMeshElements, FColor::Yellow);
+		// Gather FMeshBatches from scene proxies
+		GatherDynamicMeshElements(Views, Scene, ViewFamily, DynamicIndexBuffer, DynamicVertexBuffer, DynamicReadBuffer,
+			HasDynamicMeshElementsMasks, HasDynamicEditorMeshElementsMasks, HasViewCustomDataMasks, MeshCollector);
+	}
 
 	for (int32 ViewIndex = 0; ViewIndex < Views.Num(); ViewIndex++)
 	{
diff --git a/Engine/Source/Runtime/Renderer/Private/ScreenPass.cpp b/Engine/Source/Runtime/Renderer/Private/ScreenPass.cpp
--- a/Engine/Source/Runtime/Renderer/Private/ScreenPass.cpp
+++ b/Engine/Source/Runtime/Renderer/Private/ScreenPass.cpp
@@ -96,17 +96,14 @@ FScreenPassTextureViewportParameters GetScreenPassTextureViewportParameters(cons
 
 void SetScreenPassPipelineState(FRHICommandList& RHICmdList, const FScreenPassPipelineState& ScreenPassDraw)
 {
-	FRHIPixelShader* PixelShaderRHI = GETSAFERHISHADER_PIXEL(ScreenPassDraw.PixelShader);
-	FRHIVertexShader* VertexShaderRHI = GETSAFERHISHADER_VERTEX(ScreenPassDraw.VertexShader);
-
 	FGraphicsPipelineStateInitializer GraphicsPSOInit;
 	RHICmdList.ApplyCachedRenderTargets(GraphicsPSOInit);
 	GraphicsPSOInit.BlendState = ScreenPassDraw.BlendState;
 	GraphicsPSOInit.RasterizerState = TStaticRasterizerState<FM_Solid, CM_None>::GetRHI();
 	GraphicsPSOInit.DepthStencilState = ScreenPassDraw.DepthStencilState;
 	GraphicsPSOInit.BoundShaderState.VertexDeclarationRHI = ScreenPassDraw.VertexDeclaration;
-	GraphicsPSOInit.BoundShaderState.VertexShaderRHI = VertexShaderRHI;
-	GraphicsPSOInit.BoundShaderState.PixelShaderRHI = PixelShaderRHI;
+	GraphicsPSOInit.BoundShaderState.VertexShaderRHI = ScreenPassDraw.VertexShader.GetVertexShader();
+	GraphicsPSOInit.BoundShaderState.PixelShaderRHI = ScreenPassDraw.PixelShader.GetPixelShader();
 	GraphicsPSOInit.PrimitiveType = PT_TriangleList;
 
 	SetGraphicsPipelineState(RHICmdList, GraphicsPSOInit);
@@ -151,7 +148,7 @@ void AddDrawTexturePass(
 	Parameters->InputSampler = TStaticSamplerState<>::GetRHI();
 	Parameters->RenderTargets[0] = FRenderTargetBinding(OutputTexture, LoadAction);
 
-	AddDrawScreenPass(GraphBuilder, RDG_EVENT_NAME("DrawTexture"), View, OutputViewport, InputViewport, *PixelShader, Parameters);
+	AddDrawScreenPass(GraphBuilder, RDG_EVENT_NAME("DrawTexture"), View, OutputViewport, InputViewport, PixelShader, Parameters);
 }
 
 void AddDrawTexturePass(
@@ -170,5 +167,5 @@ void AddDrawTexturePass(
 	Parameters->InputSampler = TStaticSamplerState<>::GetRHI();
 	Parameters->RenderTargets[0] = Output.GetRenderTargetBinding();
 
-	AddDrawScreenPass(GraphBuilder, RDG_EVENT_NAME("DrawTexture"), View, OutputViewport, InputViewport, *PixelShader, Parameters);
+	AddDrawScreenPass(GraphBuilder, RDG_EVENT_NAME("DrawTexture"), View, OutputViewport, InputViewport, PixelShader, Parameters);
 }
\ No newline at end of file
diff --git a/Engine/Source/Runtime/Renderer/Private/ScreenPass.h b/Engine/Source/Runtime/Renderer/Private/ScreenPass.h
--- a/Engine/Source/Runtime/Renderer/Private/ScreenPass.h
+++ b/Engine/Source/Runtime/Renderer/Private/ScreenPass.h
@@ -223,8 +223,8 @@ struct FScreenPassPipelineState
 	FScreenPassPipelineState() = default;
 
 	FScreenPassPipelineState(
-		FShader* InVertexShader,
-		FShader* InPixelShader,
+		const TShaderRef<FShader>& InVertexShader,
+		const TShaderRef<FShader>& InPixelShader,
 		FRHIBlendState* InBlendState = FDefaultBlendState::GetRHI(),
 		FRHIDepthStencilState* InDepthStencilState = FDefaultDepthStencilState::GetRHI(),
 		FRHIVertexDeclaration* InVertexDeclaration = GFilterVertexDeclaration.VertexDeclarationRHI)
@@ -237,15 +237,15 @@ struct FScreenPassPipelineState
 
 	void Validate() const
 	{
-		check(VertexShader);
-		check(PixelShader);
+		check(VertexShader.IsValid());
+		check(PixelShader.IsValid());
 		check(BlendState);
 		check(DepthStencilState);
 		check(VertexDeclaration);
 	}
 
-	FShader* VertexShader = nullptr;
-	FShader* PixelShader = nullptr;
+	TShaderRef<FShader> VertexShader;
+	TShaderRef<FShader> PixelShader;
 	FRHIBlendState* BlendState = nullptr;
 	FRHIDepthStencilState* DepthStencilState = nullptr;
 	FRHIVertexDeclaration* VertexDeclaration = nullptr;
@@ -279,7 +279,7 @@ void DrawScreenPass(
 	const FViewInfo& View,
 	const FScreenPassTextureViewport& OutputViewport,
 	const FScreenPassTextureViewport& InputViewport,
-	TPixelShaderType* PixelShader,
+	const TShaderRef<TPixelShaderType>& PixelShader,
 	const typename TPixelShaderType::FParameters& PixelShaderParameters,
 	EScreenPassDrawFlags Flags = EScreenPassDrawFlags::None)
 {
@@ -290,11 +290,11 @@ void DrawScreenPass(
 		View,
 		OutputViewport,
 		InputViewport,
-		FScreenPassPipelineState(*ScreenPassVS, PixelShader),
+		FScreenPassPipelineState(ScreenPassVS, PixelShader),
 		Flags,
 		[&](FRHICommandListImmediate&)
 	{
-		SetShaderParameters(RHICmdList, PixelShader, PixelShader->GetPixelShader(), PixelShaderParameters);
+		SetShaderParameters(RHICmdList, PixelShader, PixelShader.GetPixelShader(), PixelShaderParameters);
 	});
 }
 
@@ -368,11 +368,11 @@ void AddDrawScreenPass(
 	const FViewInfo& View,
 	const FScreenPassTextureViewport& OutputViewport,
 	const FScreenPassTextureViewport& InputViewport,
-	TPixelShaderType* PixelShader,
+	const TShaderRef<TPixelShaderType>& PixelShader,
 	typename TPixelShaderType::FParameters* PixelShaderParameters,
 	EScreenPassDrawFlags Flags = EScreenPassDrawFlags::None)
 {
-	check(PixelShader);
+	check(PixelShader.IsValid());
 	check(PixelShaderParameters);
 
 	ClearUnusedGraphResources(PixelShader, PixelShaderParameters);
@@ -414,7 +414,6 @@ void AddDrawScreenPass(
 		DrawScreenPass(RHICmdList, View, OutputViewport, InputViewport, PipelineState, Flags, SetupFunction);
 	});
 }
-
 /** Helper function which copies a region of an input texture to a region of the output texture,
  *  with support for format conversion. If formats match, the method falls back to a simple DMA
  *  (CopyTexture); otherwise, it rasterizes using a pixel shader. Use this method if the two
diff --git a/Engine/Source/Runtime/Renderer/Private/ScreenSpaceDenoise.cpp b/Engine/Source/Runtime/Renderer/Private/ScreenSpaceDenoise.cpp
--- a/Engine/Source/Runtime/Renderer/Private/ScreenSpaceDenoise.cpp
+++ b/Engine/Source/Runtime/Renderer/Private/ScreenSpaceDenoise.cpp
@@ -1586,7 +1586,7 @@ static void DenoiseSignalAtConstantPixelDensity(
 		FComputeShaderUtils::AddPass(
 			GraphBuilder,
 			RDG_EVENT_NAME("SSD CompressMetadata %dx%d", Viewport.Width(), Viewport.Height()),
-			*ComputeShader,
+			ComputeShader,
 			PassParameters,
 			FComputeShaderUtils::GetGroupCount(Viewport.Size(), FComputeShaderUtils::kGolden2DGroupSize));
 	}
@@ -1619,7 +1619,7 @@ static void DenoiseSignalAtConstantPixelDensity(
 			GraphBuilder,
 			RDG_EVENT_NAME("SSD Injest(MultiSPP=%i)",
 				int32(PermutationVector.Get<FMultiSPPDim>())),
-			*ComputeShader,
+			ComputeShader,
 			PassParameters,
 			FComputeShaderUtils::GetGroupCount(Viewport.Size(), FComputeShaderUtils::kGolden2DGroupSize));
 
@@ -1653,7 +1653,7 @@ static void DenoiseSignalAtConstantPixelDensity(
 			RDG_EVENT_NAME("SSD Reduce(Mips=%i MultiSPP=%i)",
 				kMaxMipLevel,
 				int32(PermutationVector.Get<FMultiSPPDim>())),
-			*ComputeShader,
+			ComputeShader,
 			PassParameters,
 			FComputeShaderUtils::GetGroupCount(Viewport.Size(), FComputeShaderUtils::kGolden2DGroupSize));
 
@@ -1706,7 +1706,7 @@ static void DenoiseSignalAtConstantPixelDensity(
 				Viewport.Width(), Viewport.Height(),
 				PermutationVector.Get<FSSDSpatialAccumulationCS::FUpscaleDim>() ? TEXT(" Upscale") : TEXT(""),
 				PermutationVector.Get<FMultiSPPDim>() ? TEXT("") : TEXT(" 1SPP")),
-			*ComputeShader,
+			ComputeShader,
 			PassParameters,
 			FComputeShaderUtils::GetGroupCount(Viewport.Size(), FSSDSpatialAccumulationCS::kGroupSize));
 
@@ -1742,7 +1742,7 @@ static void DenoiseSignalAtConstantPixelDensity(
 		FComputeShaderUtils::AddPass(
 			GraphBuilder,
 			RDG_EVENT_NAME("SSD PreConvolution(MaxSamples=7 Spread=%d)", PassParameters->KernelSpreadFactor),
-			*ComputeShader,
+			ComputeShader,
 			PassParameters,
 			FComputeShaderUtils::GetGroupCount(Viewport.Size(), FSSDSpatialAccumulationCS::kGroupSize));
 
@@ -1824,7 +1824,7 @@ static void DenoiseSignalAtConstantPixelDensity(
 			FComputeShaderUtils::AddPass(
 				GraphBuilder,
 				RDG_EVENT_NAME("SSD SpatialAccumulation(RejectionPreConvolution MaxSamples=5)"),
-				*ComputeShader,
+				ComputeShader,
 				PassParameters,
 				FComputeShaderUtils::GetGroupCount(Viewport.Size(), FSSDSpatialAccumulationCS::kGroupSize));
 		} // if (SignalUsesRejectionPreConvolution(Settings.SignalProcessing))
@@ -1895,7 +1895,7 @@ static void DenoiseSignalAtConstantPixelDensity(
 
 		// Manually cleans the unused resource, to find out what the shader is actually going to need for next frame.
 		{
-			ClearUnusedGraphResources(*ComputeShader, PassParameters);
+			ClearUnusedGraphResources(ComputeShader, PassParameters);
 
 			bExtractSceneDepth = PassParameters->PrevDepthBuffer != nullptr;
 			bExtractSceneGBufferA = PassParameters->PrevGBufferA != nullptr;
@@ -1908,7 +1908,7 @@ static void DenoiseSignalAtConstantPixelDensity(
 		FComputeShaderUtils::AddPass(
 			GraphBuilder,
 			RDG_EVENT_NAME("SSD TemporalAccumulation"),
-			*ComputeShader,
+			ComputeShader,
 			PassParameters,
 			FComputeShaderUtils::GetGroupCount(Viewport.Size(), FComputeShaderUtils::kGolden2DGroupSize));
 
@@ -1944,7 +1944,7 @@ static void DenoiseSignalAtConstantPixelDensity(
 		FComputeShaderUtils::AddPass(
 			GraphBuilder,
 			RDG_EVENT_NAME("SSD SpatialAccumulation(PostFiltering MaxSamples=%i)", MaxPostFilterSampleCount),
-			*ComputeShader,
+			ComputeShader,
 			PassParameters,
 			FComputeShaderUtils::GetGroupCount(Viewport.Size(), FSSDSpatialAccumulationCS::kGroupSize));
 
@@ -2066,7 +2066,7 @@ static void DenoiseSignalAtConstantPixelDensity(
 		FComputeShaderUtils::AddPass(
 			GraphBuilder,
 			RDG_EVENT_NAME("SSD SpatialAccumulation(Final)"),
-			*ComputeShader,
+			ComputeShader,
 			PassParameters,
 			FComputeShaderUtils::GetGroupCount(Viewport.Size(), FSSDSpatialAccumulationCS::kGroupSize));
 	}
@@ -2370,7 +2370,7 @@ public:
 			FComputeShaderUtils::AddPass(
 				GraphBuilder,
 				RDG_EVENT_NAME("SSD ComposeHarmonics"),
-				*ComputeShader, ComposePassParameters,
+				ComputeShader, ComposePassParameters,
 				FComputeShaderUtils::GetGroupCount(View.ViewRect.Size(), FSSDSpatialAccumulationCS::kGroupSize));
 		}
 
diff --git a/Engine/Source/Runtime/Renderer/Private/ScreenSpaceRayTracing.cpp b/Engine/Source/Runtime/Renderer/Private/ScreenSpaceRayTracing.cpp
--- a/Engine/Source/Runtime/Renderer/Private/ScreenSpaceRayTracing.cpp
+++ b/Engine/Source/Runtime/Renderer/Private/ScreenSpaceRayTracing.cpp
@@ -398,14 +398,14 @@ class FScreenSpaceReflectionsTileVS : public FGlobalShader
 	SHADER_USE_PARAMETER_STRUCT(FScreenSpaceReflectionsTileVS, FGlobalShader);
 
 	using FPermutationDomain = TShaderPermutationDomain<>;
-
+		
 	using FParameters = FScreenSpaceReflectionsPS::FParameters; // Sharing parameters for proper registration with RDG
-
+		
 	static FPermutationDomain RemapPermutation(FPermutationDomain PermutationVector)
 	{
 		return PermutationVector;
 	}
-
+		
 	static bool ShouldCompilePermutation(const FGlobalShaderPermutationParameters& Parameters)
 	{
 		return ::UseSingleLayerWaterIndirectDraw(Parameters.Platform);
@@ -675,7 +675,7 @@ void RenderScreenSpaceReflections(
 		PassParameters->RenderTargets = RenderTargets;
 		
 		TShaderMapRef<FScreenSpaceReflectionsStencilPS> PixelShader(View.ShaderMap, PermutationVector);
-		ClearUnusedGraphResources(*PixelShader, PassParameters);
+		ClearUnusedGraphResources(PixelShader, PassParameters);
 		
 		GraphBuilder.AddPass(
 			RDG_EVENT_NAME("SSR StencilSetup %dx%d", View.ViewRect.Width(), View.ViewRect.Height()),
@@ -688,12 +688,12 @@ void RenderScreenSpaceReflections(
 			RHICmdList.SetStencilRef(0x80);
 		
 			FGraphicsPipelineStateInitializer GraphicsPSOInit;
-			FPixelShaderUtils::InitFullscreenPipelineState(RHICmdList, View.ShaderMap, *PixelShader, /* out */ GraphicsPSOInit);
+			FPixelShaderUtils::InitFullscreenPipelineState(RHICmdList, View.ShaderMap, PixelShader, /* out */ GraphicsPSOInit);
 			// Clobers the stencil to pixel that should not compute SSR
 			GraphicsPSOInit.DepthStencilState = TStaticDepthStencilState<false, CF_Always, true, CF_Always, SO_Replace, SO_Replace, SO_Replace>::GetRHI();
 
 			SetGraphicsPipelineState(RHICmdList, GraphicsPSOInit);
-			SetShaderParameters(RHICmdList, *PixelShader, PixelShader->GetPixelShader(), *PassParameters);
+			SetShaderParameters(RHICmdList, PixelShader, PixelShader.GetPixelShader(), *PassParameters);
 
 			FPixelShaderUtils::DrawFullscreenTriangle(RHICmdList);
 		});
@@ -735,10 +735,10 @@ void RenderScreenSpaceReflections(
 			PassParameters->ScreenSpaceRayTracingDebugOutput = CreateScreenSpaceRayTracingDebugUAV(GraphBuilder, DenoiserInputs->Color->Desc, TEXT("DebugSSR"), true);
 		}
 		PassParameters->PrevSceneColorPreExposureCorrection = InputColor != CurrentSceneColor ? View.PreExposure / View.PrevViewInfo.SceneColorPreExposure : 1.0f;
-
+		
 		PassParameters->SceneColor = InputColor;
 		PassParameters->SceneColorSampler = GSSRHalfResSceneColor ? TStaticSamplerState<SF_Bilinear>::GetRHI() : TStaticSamplerState<SF_Point>::GetRHI();
-
+		
 		PassParameters->HZB = GraphBuilder.RegisterExternalTexture(View.HZB);
 		PassParameters->HZBSampler = TStaticSamplerState<SF_Point>::GetRHI();
 	};
@@ -746,7 +746,7 @@ void RenderScreenSpaceReflections(
 	FScreenSpaceReflectionsPS::FPermutationDomain PermutationVector;
 	PermutationVector.Set<FSSRQualityDim>(SSRQuality);
 	PermutationVector.Set<FSSROutputForDenoiser>(bDenoiser);
-
+		
 	FScreenSpaceReflectionsPS::FParameters* PassParameters = GraphBuilder.AllocParameters<FScreenSpaceReflectionsPS::FParameters>();
 	PassParameters->CommonParameters = CommonParameters;
 	SetSSRParameters(&PassParameters->SSRPassCommonParameter);
@@ -756,8 +756,8 @@ void RenderScreenSpaceReflections(
 
 	if (TiledScreenSpaceReflection == nullptr)
 	{
-		ClearUnusedGraphResources(*PixelShader, PassParameters);
-
+		ClearUnusedGraphResources(PixelShader, PassParameters);
+		
 		GraphBuilder.AddPass(
 			RDG_EVENT_NAME("SSR RayMarch(Quality=%d RayPerPixel=%d%s) %dx%d",
 				SSRQuality, RayTracingConfigs.RayCountPerPixel, bDenoiser ? TEXT(" DenoiserOutput") : TEXT(""),
@@ -769,9 +769,9 @@ void RenderScreenSpaceReflections(
 			SCOPED_GPU_STAT(RHICmdList, ScreenSpaceReflections);
 			RHICmdList.SetViewport(View.ViewRect.Min.X, View.ViewRect.Min.Y, 0.0f, View.ViewRect.Max.X, View.ViewRect.Max.Y, 1.0f);
 			RHICmdList.SetStencilRef(0x80);
-
+		
 			FGraphicsPipelineStateInitializer GraphicsPSOInit;
-			FPixelShaderUtils::InitFullscreenPipelineState(RHICmdList, View.ShaderMap, *PixelShader, /* out */ GraphicsPSOInit);
+			FPixelShaderUtils::InitFullscreenPipelineState(RHICmdList, View.ShaderMap, PixelShader, /* out */ GraphicsPSOInit);
 			if (SSRStencilPrePass)
 			{
 				// Clobers the stencil to pixel that should not compute SSR
@@ -779,7 +779,7 @@ void RenderScreenSpaceReflections(
 			}
 
 			SetGraphicsPipelineState(RHICmdList, GraphicsPSOInit);
-			SetShaderParameters(RHICmdList, *PixelShader, PixelShader->GetPixelShader(), *PassParameters);
+			SetShaderParameters(RHICmdList, PixelShader, PixelShader.GetPixelShader(), *PassParameters);
 
 			FPixelShaderUtils::DrawFullscreenTriangle(RHICmdList);
 		});
@@ -794,8 +794,8 @@ void RenderScreenSpaceReflections(
 		PassParameters->TileListData = TiledScreenSpaceReflection->TileListStructureBufferSRV;
 		PassParameters->IndirectDrawParameter = TiledScreenSpaceReflection->DispatchIndirectParametersBuffer;
 
-		ValidateShaderParameters(*VertexShader, *PassParameters);
-		ValidateShaderParameters(*PixelShader, *PassParameters);
+		ValidateShaderParameters(VertexShader, *PassParameters);
+		ValidateShaderParameters(PixelShader, *PassParameters);
 
 		GraphBuilder.AddPass(
 			RDG_EVENT_NAME("SSR RayMarch(Quality=%d RayPerPixel=%d%s) %dx%d",
@@ -810,7 +810,7 @@ void RenderScreenSpaceReflections(
 			RHICmdList.SetStencilRef(0x80);
 
 			FGraphicsPipelineStateInitializer GraphicsPSOInit;
-			FPixelShaderUtils::InitFullscreenPipelineState(RHICmdList, View.ShaderMap, *PixelShader, /* out */ GraphicsPSOInit);
+			FPixelShaderUtils::InitFullscreenPipelineState(RHICmdList, View.ShaderMap, PixelShader, /* out */ GraphicsPSOInit);
 			if (SSRStencilPrePass)
 			{
 				// Clobers the stencil to pixel that should not compute SSR
@@ -818,12 +818,12 @@ void RenderScreenSpaceReflections(
 			}
 			GraphicsPSOInit.PrimitiveType = GRHISupportsRectTopology ? PT_RectList : PT_TriangleList;
 			GraphicsPSOInit.BoundShaderState.VertexDeclarationRHI = GEmptyVertexDeclaration.VertexDeclarationRHI;
-			GraphicsPSOInit.BoundShaderState.VertexShaderRHI = VertexShader->GetVertexShader();
-			GraphicsPSOInit.BoundShaderState.PixelShaderRHI = PixelShader->GetPixelShader();
+			GraphicsPSOInit.BoundShaderState.VertexShaderRHI = VertexShader.GetVertexShader();
+			GraphicsPSOInit.BoundShaderState.PixelShaderRHI = PixelShader.GetPixelShader();
 
 			SetGraphicsPipelineState(RHICmdList, GraphicsPSOInit);
-			SetShaderParameters(RHICmdList, *VertexShader, VertexShader->GetVertexShader(), *PassParameters);
-			SetShaderParameters(RHICmdList, *PixelShader, PixelShader->GetPixelShader(), *PassParameters);
+			SetShaderParameters(RHICmdList, VertexShader, VertexShader.GetVertexShader(), *PassParameters);
+			SetShaderParameters(RHICmdList, PixelShader, PixelShader.GetPixelShader(), *PassParameters);
 
 			PassParameters->IndirectDrawParameter->MarkResourceAsUsed();
 			RHICmdList.DrawPrimitiveIndirect(PassParameters->IndirectDrawParameter->GetIndirectRHICallBuffer(), 0);
@@ -975,7 +975,7 @@ void RenderScreenSpaceDiffuseIndirect(
 				RDG_EVENT_NAME("PrevFrameReduction(LeakFree=%i) %dx%d",
 					bUseLeakFree ? 1 : 0,
 					View.ViewRect.Width(), View.ViewRect.Height()),
-				*ComputeShader,
+				ComputeShader,
 				PassParameters,
 				FComputeShaderUtils::GetGroupCount(View.ViewRect.Size(), 8));
 		}
@@ -1020,7 +1020,7 @@ void RenderScreenSpaceDiffuseIndirect(
 			PermutationVector.Set<FSSRTPrevFrameReductionCS::FLeakFree>(bUseLeakFree);
 
 			TShaderMapRef<FSSRTPrevFrameReductionCS> ComputeShader(View.ShaderMap, PermutationVector);
-			ClearUnusedGraphResources(*ComputeShader, PassParameters);
+			ClearUnusedGraphResources(ComputeShader, PassParameters);
 			GraphBuilder.AddPass(
 				RDG_EVENT_NAME("PrevFrameReduction(LeakFree=%i) %dx%d",
 					bUseLeakFree ? 1 : 0,
@@ -1029,7 +1029,7 @@ void RenderScreenSpaceDiffuseIndirect(
 				ERDGPassFlags::Compute | ERDGPassFlags::GenerateMips,
 				[PassParameters, ComputeShader, &View, Divisor](FRHICommandList& RHICmdList)
 			{
-				FComputeShaderUtils::Dispatch(RHICmdList, *ComputeShader, *PassParameters, FComputeShaderUtils::GetGroupCount(View.ViewRect.Size(), 8 * Divisor));
+				FComputeShaderUtils::Dispatch(RHICmdList, ComputeShader, *PassParameters, FComputeShaderUtils::GetGroupCount(View.ViewRect.Size(), 8 * Divisor));
 			});
 		}
 	}
@@ -1057,7 +1057,7 @@ void RenderScreenSpaceDiffuseIndirect(
 		PassParameters->ClosestHZBTextureSampler = TStaticSamplerState<SF_Point>::GetRHI();
 
 		PassParameters->View = View.ViewUniformBuffer;
-
+			
 		PassParameters->TileClassificationParameters = ClassificationParameters;
 		PassParameters->TileClassificationUAVs = CreateUAVs(GraphBuilder, ClassificationResources);
 
@@ -1077,7 +1077,7 @@ void RenderScreenSpaceDiffuseIndirect(
 		FComputeShaderUtils::AddPass(
 			GraphBuilder,
 			RDG_EVENT_NAME("ScreenSpaceDiffuseClassification %dx%d", ThreadCount.X, ThreadCount.Y),
-			*ComputeShader,
+			ComputeShader,
 			PassParameters,
 			FComputeShaderUtils::GetGroupCount(ThreadCount, 8));
 	}
@@ -1099,9 +1099,9 @@ void RenderScreenSpaceDiffuseIndirect(
 			Desc.Format = PF_R16F;
 			OutDenoiserInputs->AmbientOcclusionMask = GraphBuilder.CreateTexture(Desc, TEXT("SSRTAmbientOcclusion"));
 		}
-
+	
 		FScreenSpaceDiffuseIndirectCS::FParameters* PassParameters = GraphBuilder.AllocParameters<FScreenSpaceDiffuseIndirectCS::FParameters>();
-		
+
 		if (bHalfResolution)
 		{
 			PassParameters->PixelPositionToFullResPixel = 2.0f;
@@ -1157,7 +1157,7 @@ void RenderScreenSpaceDiffuseIndirect(
 			GraphBuilder,
 			RDG_EVENT_NAME("ScreenSpaceDiffuseIndirect(Quality=%d RayPerPixel=%d) %dx%d",
 				Quality, RayCountPerPixel, Viewport.Width(), Viewport.Height()),
-			*ComputeShader,
+			ComputeShader,
 			PassParameters,
 			FComputeShaderUtils::GetGroupCount(Viewport.Size(), GroupSize));
 	}
diff --git a/Engine/Source/Runtime/Renderer/Private/ShaderBaseClasses.cpp b/Engine/Source/Runtime/Renderer/Private/ShaderBaseClasses.cpp
--- a/Engine/Source/Runtime/Renderer/Private/ShaderBaseClasses.cpp
+++ b/Engine/Source/Runtime/Renderer/Private/ShaderBaseClasses.cpp
@@ -13,6 +13,12 @@
 #include "VT/VirtualTextureSpace.h"
 #include "VT/VirtualTextureSystem.h"
 
+IMPLEMENT_TYPE_LAYOUT(FMaterialShader);
+IMPLEMENT_TYPE_LAYOUT(FMeshMaterialShader);
+IMPLEMENT_TYPE_LAYOUT(FBaseHS);
+IMPLEMENT_TYPE_LAYOUT(FBaseDS);
+IMPLEMENT_TYPE_LAYOUT(FDebugUniformExpressionSet);
+
 /** If true, cached uniform expressions are allowed. */
 int32 FMaterialShader::bAllowCachedUniformExpressions = true;
 
@@ -61,16 +67,16 @@ FName FMaterialShader::UniformBufferLayoutName(TEXT("Material"));
 
 FMaterialShader::FMaterialShader(const FMaterialShaderType::CompiledShaderInitializerType& Initializer)
 :	FShader(Initializer)
-#if ALLOW_SHADERMAP_DEBUG_DATA
-,	DebugUniformExpressionSet(Initializer.UniformExpressionSet)
-,	DebugUniformExpressionUBLayout(FRHIUniformBufferLayout::Zero)
-,	DebugDescription(Initializer.DebugDescription)
-#endif
+#if WITH_EDITORONLY_DATA
+, DebugUniformExpressionSet(Initializer.UniformExpressionSet)
+, DebugUniformExpressionUBLayout(FRHIUniformBufferLayout::Zero)
+, DebugDescription(Initializer.DebugDescription)
+#endif // WITH_EDITORONLY_DATA
 {
-#if ALLOW_SHADERMAP_DEBUG_DATA
+#if WITH_EDITORONLY_DATA
 	check(!DebugDescription.IsEmpty());
-	DebugUniformExpressionUBLayout.CopyFrom(Initializer.UniformExpressionSet.GetUniformBufferStruct().GetLayout());
-#endif
+	DebugUniformExpressionUBLayout.CopyFrom(Initializer.UniformExpressionSet.GetUniformBufferLayout());
+#endif // WITH_EDITORONLY_DATA
 
 	// Bind the material uniform buffer parameter.
 	MaterialUniformBuffer.Bind(Initializer.ParameterMap,TEXT("Material"));
@@ -114,8 +120,7 @@ void FMaterialShader::VerifyExpressionAndShaderMaps(const FMaterialRenderProxy*
 	{
 		auto DumpUB = [](const FRHIUniformBufferLayout& Layout)
 		{
-			FString DebugName = Layout.GetDebugName().GetPlainNameString();
-			UE_LOG(LogShaders, Warning, TEXT("Layout %s, Hash %08x"), *DebugName, Layout.GetHash());
+			UE_LOG(LogShaders, Warning, TEXT("Layout %s, Hash %08x"), *Layout.GetDebugName(), Layout.GetHash());
 			FString ResourcesString;
 			for (int32 Index = 0; Index < Layout.Resources.Num(); ++Index)
 			{
@@ -159,10 +164,10 @@ void FMaterialShader::VerifyExpressionAndShaderMaps(const FMaterialRenderProxy*
 	}
 	if (bUniformExpressionSetMismatch)
 	{
-		FString DebugDesc;
-#if ALLOW_SHADERMAP_DEBUG_DATA
-		DebugDesc = DebugDescription;
-#endif
+		const FShaderType* ShaderType = FindShaderTypeByName(GetTypeName());
+		const FString ProxyName = *MaterialRenderProxy->GetFriendlyName();
+		const FString MaterialName = *Material.GetFriendlyName();
+		const TCHAR* ShaderMapDesc = Material.GetRenderingThreadShaderMap()->GetDebugDescription();
 		UE_LOG(
 			LogShaders,
 			Fatal,
@@ -171,26 +176,26 @@ void FMaterialShader::VerifyExpressionAndShaderMaps(const FMaterialRenderProxy*
 			TEXT("Material render proxy compilation info: %s\n")
 			TEXT("Shader uniform expression set:   %u vectors, %u scalars, %u 2D textures, %u cube textures, %u array textures, %u 3D textures, %u virtual textures, shader map %p\n")
 			TEXT("Material uniform expression set: %u vectors, %u scalars, %u 2D textures, %u cube textures, %u array textures, %u 3D textures, %u virtual textures, shader map %p\n"),
-			GetType()->GetName(),
-			*MaterialRenderProxy->GetFriendlyName(),
-			*Material.GetFriendlyName(),
-			*DebugDesc,
-			*Material.GetRenderingThreadShaderMap()->GetDebugDescription(),
+			ShaderType->GetName(),
+			*ProxyName,
+			*MaterialName,
+			*DebugDescription,
+			ShaderMapDesc,
 			DebugUniformExpressionSet.NumVectorExpressions,
 			DebugUniformExpressionSet.NumScalarExpressions,
-			DebugUniformExpressionSet.Num2DTextureExpressions,
-			DebugUniformExpressionSet.NumCubeTextureExpressions,
-			DebugUniformExpressionSet.Num2DArrayTextureExpressions,
-			DebugUniformExpressionSet.NumVolumeTextureExpressions,
-			DebugUniformExpressionSet.NumVirtualTextureExpressions,
+			DebugUniformExpressionSet.NumTextureExpressions[(uint32)EMaterialTextureParameterType::Standard2D],
+			DebugUniformExpressionSet.NumTextureExpressions[(uint32)EMaterialTextureParameterType::Cube],
+			DebugUniformExpressionSet.NumTextureExpressions[(uint32)EMaterialTextureParameterType::Array2D],
+			DebugUniformExpressionSet.NumTextureExpressions[(uint32)EMaterialTextureParameterType::Volume],
+			DebugUniformExpressionSet.NumTextureExpressions[(uint32)EMaterialTextureParameterType::Virtual],
 			UniformExpressionCache->CachedUniformExpressionShaderMap,
-			MaterialUniformExpressionSet.UniformVectorExpressions.Num(),
-			MaterialUniformExpressionSet.UniformScalarExpressions.Num(),
-			MaterialUniformExpressionSet.Uniform2DTextureExpressions.Num(),
-			MaterialUniformExpressionSet.UniformCubeTextureExpressions.Num(),
-			MaterialUniformExpressionSet.Uniform2DArrayTextureExpressions.Num(),
-			MaterialUniformExpressionSet.UniformVolumeTextureExpressions.Num(),
-			MaterialUniformExpressionSet.UniformVirtualTextureExpressions.Num(),
+			MaterialUniformExpressionSet.UniformVectorPreshaders.Num(),
+			MaterialUniformExpressionSet.UniformScalarPreshaders.Num(),
+			MaterialUniformExpressionSet.UniformTextureParameters[(uint32)EMaterialTextureParameterType::Standard2D].Num(),
+			MaterialUniformExpressionSet.UniformTextureParameters[(uint32)EMaterialTextureParameterType::Cube].Num(),
+			MaterialUniformExpressionSet.UniformTextureParameters[(uint32)EMaterialTextureParameterType::Array2D].Num(),
+			MaterialUniformExpressionSet.UniformTextureParameters[(uint32)EMaterialTextureParameterType::Volume].Num(),
+			MaterialUniformExpressionSet.UniformTextureParameters[(uint32)EMaterialTextureParameterType::Virtual].Num(),
 			Material.GetRenderingThreadShaderMap()
 		);
 	}
@@ -222,11 +227,12 @@ void FMaterialShader::SetParametersInner(
 			UMaterialInterface* MtlInterface = Material.GetMaterialInterface();
 			UMaterialInterface* ProxyInterface = MaterialRenderProxy->GetMaterialInterface();
 
+			const FShaderType* ShaderType = FindShaderTypeByName(GetTypeName());
 			ensureMsgf(false,
 				TEXT("%s shader uniform expression set mismatched shader map for material %s/%s, forcing expression cache evaluation.\n")
 				TEXT("Material:  %s\n")
 				TEXT("Proxy:  %s\n"),
-				GetType()->GetName(),
+				ShaderType->GetName(),
 				*MaterialRenderProxy->GetFriendlyName(), *Material.GetFriendlyName(),
 				MtlInterface ? *MtlInterface->GetFullName() : TEXT("nullptr"),
 				ProxyInterface ? *ProxyInterface->GetFullName() : TEXT("nullptr"));
@@ -414,7 +420,7 @@ IMPLEMENT_MATERIAL_SHADER_SetParameters(FRHIComputeShader);
 
 void FMaterialShader::GetShaderBindings(
 	const FScene* Scene,
-	ERHIFeatureLevel::Type FeatureLevel,
+	const FStaticFeatureLevel FeatureLevel,
 	const FMaterialRenderProxy& MaterialRenderProxy,
 	const FMaterial& Material,
 	FMeshDrawSingleShaderBindings& ShaderBindings) const
@@ -484,6 +490,7 @@ void FMaterialShader::GetShaderBindings(
 	}
 }
 
+#if 0
 bool FMaterialShader::Serialize(FArchive& Ar)
 {
 	const bool bShaderHasOutdatedParameters = FShader::Serialize(Ar);
@@ -544,15 +551,12 @@ bool FMaterialShader::Serialize(FArchive& Ar)
 
 	return bShaderHasOutdatedParameters;
 }
+#endif
 
-uint32 FMaterialShader::GetAllocatedSize() const
+FMeshMaterialShader::FMeshMaterialShader(const FMeshMaterialShaderType::CompiledShaderInitializerType& Initializer)
+	: FMaterialShader(Initializer)
 {
-	return FShader::GetAllocatedSize()
-		+ ParameterCollectionUniformBuffers.GetAllocatedSize()
-#if ALLOW_SHADERMAP_DEBUG_DATA
-		+ DebugDescription.GetAllocatedSize()
-#endif
-	;
+	VertexFactoryParameters = Initializer.VertexFactoryType->CreateShaderParameters(Initializer.Target.GetFrequency(), Initializer.ParameterMap);
 }
 
 void FMeshMaterialShader::GetShaderBindings(
@@ -574,11 +578,12 @@ void FMeshMaterialShader::GetShaderBindings(
 }
 
 void FMeshMaterialShader::GetElementShaderBindings(
+	const FShaderMapPointerTable& PointerTable,
 	const FScene* Scene, 
 	const FSceneView* ViewIfDynamicMeshCommand, 
 	const FVertexFactory* VertexFactory,
 	const EVertexInputStreamType InputStreamType,
-	ERHIFeatureLevel::Type FeatureLevel,
+	const FStaticFeatureLevel FeatureLevel,
 	const FPrimitiveSceneProxy* PrimitiveSceneProxy,
 	const FMeshBatch& MeshBatch,
 	const FMeshBatchElement& BatchElement, 
@@ -586,14 +591,21 @@ void FMeshMaterialShader::GetElementShaderBindings(
 	FMeshDrawSingleShaderBindings& ShaderBindings,
 	FVertexInputStreamArray& VertexStreams) const
 {
-	checkSlow(ShaderBindings.Frequency == GetType()->GetFrequency());
-
-	VertexFactoryParameters.GetElementShaderBindings(Scene, ViewIfDynamicMeshCommand, this, InputStreamType, FeatureLevel, VertexFactory, BatchElement, ShaderBindings, VertexStreams);
+	const FVertexFactoryType* VertexFactoryType = GetVertexFactoryType(PointerTable);
+	if (VertexFactoryType)
+	{
+		const FVertexFactoryShaderParameters* VFParameters = VertexFactoryParameters.Get();
+		if (VFParameters)
+		{
+			VertexFactoryType->GetShaderParameterElementShaderBindings(GetFrequency(), VFParameters, Scene, ViewIfDynamicMeshCommand, this, InputStreamType, FeatureLevel, VertexFactory, BatchElement, ShaderBindings, VertexStreams);
+		}
+	}
 		
 	if (UseGPUScene(GMaxRHIShaderPlatform, FeatureLevel) && VertexFactory->GetPrimitiveIdStreamIndex(InputStreamType) >= 0)
 	{
-		ensureMsgf(!GetUniformBufferParameter<FPrimitiveUniformShaderParameters>().IsBound(), TEXT("Shader %s attempted to bind the Primitive uniform buffer even though Vertex Factory computes a PrimitiveId per-instance.  This will break auto-instancing.  Shaders should use GetPrimitiveData(PrimitiveId).Member instead of Primitive.Member."), GetType()->GetName());
-		ensureMsgf(!BatchElement.PrimitiveUniformBuffer, TEXT("FMeshBatchElement was assigned a PrimitiveUniformBuffer even though Vertex Factory %s fetches primitive shader data through a Scene buffer.  The assigned PrimitiveUniformBuffer cannot be respected.  Use PrimitiveUniformBufferResource instead for dynamic primitive data."), GetType()->GetName());
+		const FShaderType* ShaderType = GetType(PointerTable);
+		ensureMsgf(!GetUniformBufferParameter<FPrimitiveUniformShaderParameters>().IsBound(), TEXT("Shader %s attempted to bind the Primitive uniform buffer even though Vertex Factory computes a PrimitiveId per-instance.  This will break auto-instancing.  Shaders should use GetPrimitiveData(PrimitiveId).Member instead of Primitive.Member."), ShaderType->GetName());
+		ensureMsgf(!BatchElement.PrimitiveUniformBuffer, TEXT("FMeshBatchElement was assigned a PrimitiveUniformBuffer even though Vertex Factory %s fetches primitive shader data through a Scene buffer.  The assigned PrimitiveUniformBuffer cannot be respected.  Use PrimitiveUniformBufferResource instead for dynamic primitive data."), ShaderType->GetName());
 	}
 	else
 	{
@@ -603,22 +615,23 @@ void FMeshMaterialShader::GetElementShaderBindings(
 		}
 		else
 		{
-			checkf(BatchElement.PrimitiveUniformBufferResource, TEXT("%s expected a primitive uniform buffer but none was set on BatchElement.PrimitiveUniformBuffer or BatchElement.PrimitiveUniformBufferResource"), GetType()->GetName());
+			const FShaderType* ShaderType = GetType(PointerTable);
+			checkf(BatchElement.PrimitiveUniformBufferResource, TEXT("%s expected a primitive uniform buffer but none was set on BatchElement.PrimitiveUniformBuffer or BatchElement.PrimitiveUniformBufferResource"), ShaderType->GetName());
 			ShaderBindings.Add(GetUniformBufferParameter<FPrimitiveUniformShaderParameters>(), BatchElement.PrimitiveUniformBufferResource->GetUniformBufferRHI());
 		}
 	}
 }
 
-bool FMeshMaterialShader::Serialize(FArchive& Ar)
+/*bool FMeshMaterialShader::Serialize(FArchive& Ar)
 {
 	bool bShaderHasOutdatedParameters = FMaterialShader::Serialize(Ar);
 	Ar << PassUniformBuffer;
 	bShaderHasOutdatedParameters |= Ar << VertexFactoryParameters;
 	return bShaderHasOutdatedParameters;
-}
+}*/
 
-uint32 FMeshMaterialShader::GetAllocatedSize() const
+void FMeshMaterialShader::WriteFrozenVertexFactoryParameters(FMemoryImageWriter& Writer, const TMemoryImagePtr<FVertexFactoryShaderParameters>& InVertexFactoryParameters) const
 {
-	return FMaterialShader::GetAllocatedSize()
-		+ VertexFactoryParameters.GetAllocatedSize();
+	const FVertexFactoryType* VertexFactoryType = GetVertexFactoryType(Writer.TryGetPrevPointerTable());
+	InVertexFactoryParameters.WriteMemoryImageWithDerivedType(Writer, VertexFactoryType->GetShaderParameterLayout(GetFrequency()));
 }
diff --git a/Engine/Source/Runtime/Renderer/Private/ShaderBaseClasses.h b/Engine/Source/Runtime/Renderer/Private/ShaderBaseClasses.h
--- a/Engine/Source/Runtime/Renderer/Private/ShaderBaseClasses.h
+++ b/Engine/Source/Runtime/Renderer/Private/ShaderBaseClasses.h
@@ -36,7 +36,7 @@ typedef TUniformBufferRef< FDitherUniformShaderParameters > FDitherUniformBuffer
 /** Base Hull shader for drawing policy rendering */
 class FBaseHS : public FMeshMaterialShader
 {
-	DECLARE_SHADER_TYPE(FBaseHS,MeshMaterial);
+	DECLARE_TYPE_LAYOUT(FBaseHS, NonVirtual);
 public:
 
 	static bool ShouldCompilePermutation(const FMeshMaterialShaderPermutationParameters& Parameters)
@@ -52,7 +52,7 @@ public:
 			return false;	
 		}
 
-		if (!Parameters.Material || Parameters.Material->GetTessellationMode() == MTM_NoTessellation)
+		if (Parameters.MaterialParameters.TessellationMode == MTM_NoTessellation)
 		{
 			// Material controls use of tessellation
 			return false;	
@@ -71,12 +71,15 @@ public:
 	}
 
 	FBaseHS() {}
+
+	
+	
 };
 
 /** Base Domain shader for drawing policy rendering */
 class FBaseDS : public FMeshMaterialShader
 {
-	DECLARE_SHADER_TYPE(FBaseDS,MeshMaterial);
+	DECLARE_TYPE_LAYOUT(FBaseDS, NonVirtual);
 public:
 
 	static bool ShouldCompilePermutation(const FMeshMaterialShaderPermutationParameters& Parameters)
@@ -92,7 +95,7 @@ public:
 			return false;	
 		}
 
-		if (!Parameters.Material || Parameters.Material->GetTessellationMode() == MTM_NoTessellation)
+		if (Parameters.MaterialParameters.TessellationMode == MTM_NoTessellation)
 		{
 			// Material controls use of tessellation
 			return false;	
@@ -111,5 +114,8 @@ public:
 	}
 
 	FBaseDS() {}
+
+	
+	
 };
 
diff --git a/Engine/Source/Runtime/Renderer/Private/ShaderComplexityRendering.cpp b/Engine/Source/Runtime/Renderer/Private/ShaderComplexityRendering.cpp
--- a/Engine/Source/Runtime/Renderer/Private/ShaderComplexityRendering.cpp
+++ b/Engine/Source/Runtime/Renderer/Private/ShaderComplexityRendering.cpp
@@ -69,8 +69,9 @@ static FAutoConsoleVariableRef CVarShaderComplexityBaselineDeferredUnlitPS(
 IMPLEMENT_SHADER_TYPE(template<>,TComplexityAccumulatePS<false>,TEXT("/Engine/Private/ShaderComplexityAccumulatePixelShader.usf"),TEXT("Main"),SF_Pixel);
 IMPLEMENT_SHADER_TYPE(template<>,TComplexityAccumulatePS<true>,TEXT("/Engine/Private/QuadComplexityAccumulatePixelShader.usf"),TEXT("Main"),SF_Pixel);
 
-template <bool bQuadComplexity>
-void TComplexityAccumulatePS<bQuadComplexity>::GetDebugViewModeShaderBindings(
+
+void FComplexityAccumulateInterface::GetDebugViewModeShaderBindings(
+	const FDebugViewModePS& BaseShader,
 	const FPrimitiveSceneProxy* RESTRICT PrimitiveSceneProxy,
 	const FMaterialRenderProxy& RESTRICT MaterialRenderProxy,
 	const FMaterial& RESTRICT Material,
@@ -89,17 +90,20 @@ void TComplexityAccumulatePS<bQuadComplexity>::GetDebugViewModeShaderBindings(
 	// late value is for overdraw which can be problematic with a low precision float format, at some point the precision isn't there any more and it doesn't accumulate
 	if (DebugViewMode == DVSM_QuadComplexity)
 	{
-		ShaderBindings.Add(NormalizedComplexity, FVector4(NormalizedQuadComplexityValue));
+		const TComplexityAccumulatePS<true>& Shader = static_cast<const TComplexityAccumulatePS<true>&>(BaseShader);
+		ShaderBindings.Add(Shader.NormalizedComplexity, FVector4(NormalizedQuadComplexityValue));
+		ShaderBindings.Add(Shader.ShowQuadOverdraw, DebugViewMode != DVSM_ShaderComplexity ? 1 : 0);
 	}
 	else
 	{
+		const TComplexityAccumulatePS<false>& Shader = static_cast<const TComplexityAccumulatePS<false>&>(BaseShader);
 		const float NormalizeMul = 1.0f / GetMaxShaderComplexityCount(Material.GetFeatureLevel());
-		ShaderBindings.Add(NormalizedComplexity, FVector4(NumPSInstructions * NormalizeMul, NumVSInstructions * NormalizeMul, 1 / 32.0f));
+		ShaderBindings.Add(Shader.NormalizedComplexity, FVector4(NumPSInstructions * NormalizeMul, NumVSInstructions * NormalizeMul, 1 / 32.0f));
+		ShaderBindings.Add(Shader.ShowQuadOverdraw, DebugViewMode != DVSM_ShaderComplexity ? 1 : 0);
 	}
-	ShaderBindings.Add(ShowQuadOverdraw, DebugViewMode != DVSM_ShaderComplexity ? 1 : 0);
 }
 
-FDebugViewModePS* FComplexityAccumulateInterface::GetPixelShader(const FMaterial* InMaterial, FVertexFactoryType* VertexFactoryType) const
+TShaderRef<FDebugViewModePS> FComplexityAccumulateInterface::GetPixelShader(const FMaterial* InMaterial, FVertexFactoryType* VertexFactoryType) const
 {
 	if (bShowQuadComplexity)
 	{
diff --git a/Engine/Source/Runtime/Renderer/Private/ShaderComplexityRendering.h b/Engine/Source/Runtime/Renderer/Private/ShaderComplexityRendering.h
--- a/Engine/Source/Runtime/Renderer/Private/ShaderComplexityRendering.h
+++ b/Engine/Source/Runtime/Renderer/Private/ShaderComplexityRendering.h
@@ -35,21 +35,22 @@ class TComplexityAccumulatePS : public FDebugViewModePS
 {
 	DECLARE_SHADER_TYPE(TComplexityAccumulatePS,MeshMaterial);
 public:
-
 	static bool ShouldCompilePermutation(const FMeshMaterialShaderPermutationParameters& Parameters)
 	{
 		// See FDebugViewModeMaterialProxy::GetFriendlyName()
-		if (AllowDebugViewShaderMode(bQuadComplexity ? DVSM_QuadComplexity : DVSM_ShaderComplexity, Parameters.Platform, Parameters.Material->GetFeatureLevel()))
+		if (AllowDebugViewShaderMode(bQuadComplexity ? DVSM_QuadComplexity : DVSM_ShaderComplexity, Parameters.Platform, Parameters.MaterialParameters.FeatureLevel))
 		{
 			// If it comes from FDebugViewModeMaterialProxy, compile it.
-			if (Parameters.Material->GetFriendlyName().Contains(TEXT("ComplexityAccumulate")))
+			if (Parameters.MaterialParameters.bMaterialIsComplexityAccumulate)
 			{
 				return true;
 			}
 			// Otherwise we only cache it if this for the shader complexity.
 			else if (GCacheShaderComplexityShaders)
 			{
-				return !FDebugViewModeInterface::AllowFallbackToDefaultMaterial(Parameters.Material) || Parameters.Material->IsDefaultMaterial();
+				return !FDebugViewModeInterface::AllowFallbackToDefaultMaterial(Parameters.MaterialParameters.TessellationMode,
+					Parameters.MaterialParameters.bHasVertexPositionOffsetConnected,
+					Parameters.MaterialParameters.bHasPixelDepthOffsetConnected) || Parameters.MaterialParameters.bIsDefaultMaterial;
 			}
 		}
 		return false;
@@ -65,43 +66,17 @@ public:
 
 	TComplexityAccumulatePS() {}
 
-	virtual bool Serialize(FArchive& Ar) override
-	{
-		bool bShaderHasOutdatedParameters = FMeshMaterialShader::Serialize(Ar);
-		Ar << NormalizedComplexity;
-		Ar << ShowQuadOverdraw;
-		Ar << QuadBufferUAV;
-		return bShaderHasOutdatedParameters;
-	}
-
 	static void ModifyCompilationEnvironment(const FMaterialShaderPermutationParameters& Parameters, FShaderCompilerEnvironment& OutEnvironment)
 	{
-		OutEnvironment.SetDefine(TEXT("OUTPUT_QUAD_OVERDRAW"), AllowDebugViewShaderMode(DVSM_QuadComplexity, Parameters.Platform, Parameters.Material->GetFeatureLevel()));
+		OutEnvironment.SetDefine(TEXT("OUTPUT_QUAD_OVERDRAW"), AllowDebugViewShaderMode(DVSM_QuadComplexity, Parameters.Platform, Parameters.MaterialParameters.FeatureLevel));
 		TCHAR BufferRegister[] = { 'u', '0', 0 };
-		BufferRegister[1] += FSceneRenderTargets::GetQuadOverdrawUAVIndex(Parameters.Platform, Parameters.Material->GetFeatureLevel());
+		BufferRegister[1] += FSceneRenderTargets::GetQuadOverdrawUAVIndex(Parameters.Platform, Parameters.MaterialParameters.FeatureLevel);
 		OutEnvironment.SetDefine(TEXT("QUAD_BUFFER_REGISTER"), BufferRegister);
 	}
 
-	virtual void GetDebugViewModeShaderBindings(
-		const FPrimitiveSceneProxy* RESTRICT PrimitiveSceneProxy,
-		const FMaterialRenderProxy& RESTRICT MaterialRenderProxy,
-		const FMaterial& RESTRICT Material,
-		EDebugViewShaderMode DebugViewMode,
-		const FVector& ViewOrigin,
-		int32 VisualizeLODIndex,
-		int32 VisualizeElementIndex,
-		int32 NumVSInstructions,
-		int32 NumPSInstructions,
-		int32 ViewModeParam,
-		FName ViewModeParamName,
-		FMeshDrawSingleShaderBindings& ShaderBindings
-	) const override;
-
-private:
-
-	FShaderParameter NormalizedComplexity;
-	FShaderParameter ShowQuadOverdraw;
-	FShaderResourceParameter QuadBufferUAV;
+	LAYOUT_FIELD(FShaderParameter, NormalizedComplexity);
+	LAYOUT_FIELD(FShaderParameter, ShowQuadOverdraw);
+	LAYOUT_FIELD(FShaderResourceParameter, QuadBufferUAV);
 };
 
 class FComplexityAccumulateInterface : public FDebugViewModeInterface
@@ -117,8 +92,24 @@ public:
 		, bShowQuadComplexity(InShowQuadComplexity)
 	{}
 
-	virtual FDebugViewModePS* GetPixelShader(const FMaterial* InMaterial, FVertexFactoryType* VertexFactoryType) const override;
+	virtual TShaderRef<FDebugViewModePS> GetPixelShader(const FMaterial* InMaterial, FVertexFactoryType* VertexFactoryType) const override;
 	virtual void SetDrawRenderState(EBlendMode BlendMode, FRenderState& DrawRenderState, bool bHasDepthPrepassForMaskedMaterial) const;
+
+	virtual void GetDebugViewModeShaderBindings(
+		const FDebugViewModePS& Shader,
+		const FPrimitiveSceneProxy* PrimitiveSceneProxy,
+		const FMaterialRenderProxy& MaterialRenderProxy,
+		const FMaterial& Material,
+		EDebugViewShaderMode DebugViewMode,
+		const FVector& ViewOrigin,
+		int32 VisualizeLODIndex,
+		int32 VisualizeElementIndex,
+		int32 NumVSInstructions,
+		int32 NumPSInstructions,
+		int32 ViewModeParam,
+		FName ViewModeParamName,
+		FMeshDrawSingleShaderBindings& ShaderBindings
+	) const override;
 };
 
 #endif //!(UE_BUILD_SHIPPING || UE_BUILD_TEST)
diff --git a/Engine/Source/Runtime/Renderer/Private/ShaderPrint.cpp b/Engine/Source/Runtime/Renderer/Private/ShaderPrint.cpp
--- a/Engine/Source/Runtime/Renderer/Private/ShaderPrint.cpp
+++ b/Engine/Source/Runtime/Renderer/Private/ShaderPrint.cpp
@@ -12,6 +12,7 @@
 #include "SceneRendering.h"
 #include "SystemTextures.h"
 
+IMPLEMENT_TYPE_LAYOUT(ShaderPrint::FShaderParametersLegacy);
 
 namespace ShaderPrint
 {
@@ -316,14 +317,14 @@ namespace ShaderPrint
 
 		// Clear the output buffer internal counter ready for use
 		const ERHIFeatureLevel::Type FeatureLevel = View.GetFeatureLevel();
-		TShaderMap<FGlobalShaderType>* GlobalShaderMap = GetGlobalShaderMap(FeatureLevel);
+		FGlobalShaderMap* GlobalShaderMap = GetGlobalShaderMap(FeatureLevel);
 
 		TShaderMapRef< FShaderInitValueBufferCS > ComputeShader(GlobalShaderMap);
 
 		FShaderInitValueBufferCS::FParameters Parameters;
 		Parameters.RWValuesBuffer = View.ShaderPrintValueBuffer.UAV;
 
-		FComputeShaderUtils::Dispatch(RHICmdList, *ComputeShader, Parameters, FIntVector(1, 1, 1));
+		FComputeShaderUtils::Dispatch(RHICmdList, ComputeShader, Parameters, FIntVector(1, 1, 1));
 	}
 
 	void DrawView(FRDGBuilder& GraphBuilder, const FViewInfo& View, FRDGTextureRef OutputTexture)
@@ -344,7 +345,7 @@ namespace ShaderPrint
 		FTextureRHIRef FontTexture = GEngine->MiniFontTexture != nullptr ? GEngine->MiniFontTexture->Resource->TextureRHI : GSystemTextures.BlackDummy->GetRenderTargetItem().ShaderResourceTexture;;
 
 		const ERHIFeatureLevel::Type FeatureLevel = View.GetFeatureLevel();
-		TShaderMap<FGlobalShaderType>* GlobalShaderMap = GetGlobalShaderMap(FeatureLevel);
+		FGlobalShaderMap* GlobalShaderMap = GetGlobalShaderMap(FeatureLevel);
 
 		// BuildIndirectDispatchArgs
 		{
@@ -360,7 +361,7 @@ namespace ShaderPrint
 			FComputeShaderUtils::AddPass(
 				GraphBuilder, 
 				RDG_EVENT_NAME("BuildIndirectDispatchArgs"), 
-				*ComputeShader, PassParameters,
+				ComputeShader, PassParameters,
 				FIntVector(1, 1, 1));
 		}
 
@@ -378,7 +379,7 @@ namespace ShaderPrint
 			FComputeShaderUtils::AddPass(
 				GraphBuilder,
 				RDG_EVENT_NAME("BuildSymbolBuffer"),
-				*ComputeShader, PassParameters,
+				ComputeShader, PassParameters,
 				IndirectDispatchArgsBuffer, 0);
 		}
 
@@ -395,7 +396,7 @@ namespace ShaderPrint
 			FComputeShaderUtils::AddPass(
 				GraphBuilder,
 				RDG_EVENT_NAME("BuildIndirectDrawArgs"),
-				*ComputeShader, PassParameters,
+				ComputeShader, PassParameters,
 				FIntVector(1, 1, 1));
 		}
 
@@ -425,12 +426,12 @@ namespace ShaderPrint
 				GraphicsPSOInit.RasterizerState = TStaticRasterizerState<>::GetRHI();
 				GraphicsPSOInit.PrimitiveType = PT_TriangleList;
 				GraphicsPSOInit.BoundShaderState.VertexDeclarationRHI = GetVertexDeclarationFVector4();
-				GraphicsPSOInit.BoundShaderState.VertexShaderRHI = GETSAFERHISHADER_VERTEX(*VertexShader);
-				GraphicsPSOInit.BoundShaderState.PixelShaderRHI = GETSAFERHISHADER_PIXEL(*PixelShader);
+				GraphicsPSOInit.BoundShaderState.VertexShaderRHI = VertexShader.GetVertexShader();
+				GraphicsPSOInit.BoundShaderState.PixelShaderRHI = PixelShader.GetPixelShader();
 				SetGraphicsPipelineState(RHICmdListImmediate, GraphicsPSOInit);
 
-				SetShaderParameters(RHICmdListImmediate, *VertexShader, VertexShader->GetVertexShader(), *PassParameters);
-				SetShaderParameters(RHICmdListImmediate, *PixelShader, PixelShader->GetPixelShader(), *PassParameters);
+				SetShaderParameters(RHICmdListImmediate, VertexShader, VertexShader.GetVertexShader(), *PassParameters);
+				SetShaderParameters(RHICmdListImmediate, PixelShader, PixelShader.GetPixelShader(), *PassParameters);
 
 				RHICmdListImmediate.DrawIndexedPrimitiveIndirect(GTwoTrianglesIndexBuffer.IndexBufferRHI, PassParameters->IndirectDrawArgsBuffer->GetIndirectRHICallBuffer(), 0);
 			});
diff --git a/Engine/Source/Runtime/Renderer/Private/ShaderPrintParameters.h b/Engine/Source/Runtime/Renderer/Private/ShaderPrintParameters.h
--- a/Engine/Source/Runtime/Renderer/Private/ShaderPrintParameters.h
+++ b/Engine/Source/Runtime/Renderer/Private/ShaderPrintParameters.h
@@ -43,6 +43,8 @@ namespace ShaderPrint
 	// Legacy parameter binding helper for an FShader that doesn't use FParameters style parameter declaration
 	struct FShaderParametersLegacy
 	{
+		DECLARE_TYPE_LAYOUT(FShaderParametersLegacy, NonVirtual);
+	public:
 		// Call on shader construction to bind
 		void Bind(FShaderParameterMap const& ParameterMap);
 
@@ -58,7 +60,7 @@ namespace ShaderPrint
 		friend FArchive& operator<<(FArchive& Ar, FShaderParametersLegacy& P);
 
 		// Parameters
-		FShaderUniformBufferParameter UniformBufferParameter;
-		FRWShaderParameter ValuesBufferParameter;
+		LAYOUT_FIELD(FShaderUniformBufferParameter, UniformBufferParameter);
+		LAYOUT_FIELD(FRWShaderParameter, ValuesBufferParameter);
 	};
 }
diff --git a/Engine/Source/Runtime/Renderer/Private/ShadowDepthRendering.cpp b/Engine/Source/Runtime/Renderer/Private/ShadowDepthRendering.cpp
--- a/Engine/Source/Runtime/Renderer/Private/ShadowDepthRendering.cpp
+++ b/Engine/Source/Runtime/Renderer/Private/ShadowDepthRendering.cpp
@@ -141,7 +141,7 @@ public:
 */
 class FShadowDepthVS : public FMeshMaterialShader
 {
-	DECLARE_SHADER_TYPE(FShadowDepthVS, MeshMaterial);
+	DECLARE_INLINE_TYPE_LAYOUT(FShadowDepthVS, NonVirtual);
 public:
 
 	static bool ShouldCompilePermutation(const FMeshMaterialShaderPermutationParameters& Parameters)
@@ -169,13 +169,6 @@ public:
 
 	FShadowDepthVS() {}
 
-	virtual bool Serialize(FArchive& Ar) override
-	{
-		bool bShaderHasOutdatedParameters = FMeshMaterialShader::Serialize(Ar);
-		Ar << LayerId;
-		return bShaderHasOutdatedParameters;
-	}
-
 	void GetShaderBindings(
 		const FScene* Scene,
 		ERHIFeatureLevel::Type FeatureLevel,
@@ -192,7 +185,8 @@ public:
 	}
 
 private:
-	FShaderParameter LayerId;
+	
+	LAYOUT_FIELD(FShaderParameter, LayerId);
 };
 
 enum EShadowDepthVertexShaderMode
@@ -227,7 +221,6 @@ public:
 	static bool ShouldCompilePermutation(const FMeshMaterialShaderPermutationParameters& Parameters)
 	{
 		const EShaderPlatform Platform = Parameters.Platform;
-		const FMaterial* Material = Parameters.Material;
 
 		static const auto SupportAllShaderPermutationsVar = IConsoleManager::Get().FindTConsoleVariableDataInt(TEXT("r.SupportAllShaderPermutations"));
 		const bool bForceAllPermutations = SupportAllShaderPermutationsVar && SupportAllShaderPermutationsVar->GetValueOnAnyThread() != 0;
@@ -253,22 +246,21 @@ public:
 				// Reflective shadow map shaders must be compiled for every material because they access the material normal
 				return !bUsePositionOnlyStream
 					// Don't render ShadowDepth for translucent unlit materials, unless we're injecting emissive
-					&& (Material->ShouldCastDynamicShadows() || Material->ShouldInjectEmissiveIntoLPV()
-						|| Material->ShouldBlockGI())
+					&& (Parameters.MaterialParameters.bShouldCastDynamicShadows || Parameters.MaterialParameters.bShouldInjectEmissiveIntoLPV || Parameters.MaterialParameters.bShouldBlockGI)
 					&& IsFeatureLevelSupported(Platform, ERHIFeatureLevel::SM5);
 			}
 		}
 		else
 		{
-			return (Material->IsSpecialEngineMaterial()
+			return (Parameters.MaterialParameters.bIsSpecialEngineMaterial
 				// Masked and WPO materials need their shaders but cannot be used with a position only stream.
-				|| ((!Material->WritesEveryPixel(true) || Material->MaterialMayModifyMeshPosition()) && !bUsePositionOnlyStream))
+				|| ((!Parameters.MaterialParameters.bWritesEveryPixelShadowPass || Parameters.MaterialParameters.bMaterialMayModifyMeshPosition) && !bUsePositionOnlyStream))
 				// Only compile one pass point light shaders for feature levels >= SM5
 				&& (ShaderMode != VertexShadowDepth_OnePassPointLight || IsFeatureLevelSupported(Platform, ERHIFeatureLevel::SM5))
 				// Only compile position-only shaders for vertex factories that support it. (Note: this assumes that a vertex factor which supports PositionOnly, supports also PositionAndNormalOnly)
 				&& (!bUsePositionOnlyStream || Parameters.VertexFactoryType->SupportsPositionOnly())
 				// Don't render ShadowDepth for translucent unlit materials
-				&& Material->ShouldCastDynamicShadows()
+				&& Parameters.MaterialParameters.bShouldCastDynamicShadows
 				// Only compile perspective correct light shaders for feature levels >= SM5
 				&& (ShaderMode != VertexShadowDepth_PerspectiveCorrect || IsFeatureLevelSupported(Platform, ERHIFeatureLevel::SM5));
 		}
@@ -433,13 +425,12 @@ IMPLEMENT_MATERIAL_SHADER_TYPE(template<>, TShadowDepthVSForGSVertexShadowDepth_
 /**
 * A pixel shader for rendering the depth of a mesh.
 */
-template <bool bRenderReflectiveShadowMap>
-class TShadowDepthBasePS : public FMeshMaterialShader
+class FShadowDepthBasePS : public FMeshMaterialShader
 {
-	DECLARE_SHADER_TYPE(TShadowDepthBasePS, MeshMaterial);
+	DECLARE_INLINE_TYPE_LAYOUT(FShadowDepthBasePS, NonVirtual);
 public:
 
-	TShadowDepthBasePS(const ShaderMetaType::CompiledShaderInitializerType& Initializer)
+	FShadowDepthBasePS(const ShaderMetaType::CompiledShaderInitializerType& Initializer)
 		: FMeshMaterialShader(Initializer)
 	{
 		GvListBuffer.Bind(Initializer.ParameterMap, TEXT("RWGvListBuffer"));
@@ -460,26 +451,14 @@ public:
 		}
 	}
 
-	TShadowDepthBasePS() {}
-
-	virtual bool Serialize(FArchive& Ar) override
-	{
-		bool bShaderHasOutdatedParameters = FMeshMaterialShader::Serialize(Ar);
-
-		Ar << GvListBuffer;
-		Ar << GvListHeadBuffer;
-		Ar << VplListBuffer;
-		Ar << VplListHeadBuffer;
-
-		return bShaderHasOutdatedParameters;
-	}
+	FShadowDepthBasePS() {}
 
 private:
-
-	FRWShaderParameter GvListBuffer;
-	FRWShaderParameter GvListHeadBuffer;
-	FRWShaderParameter VplListBuffer;
-	FRWShaderParameter VplListHeadBuffer;
+	
+	LAYOUT_FIELD(FRWShaderParameter, GvListBuffer);
+	LAYOUT_FIELD(FRWShaderParameter, GvListHeadBuffer);
+	LAYOUT_FIELD(FRWShaderParameter, VplListBuffer);
+	LAYOUT_FIELD(FRWShaderParameter, VplListHeadBuffer);
 };
 
 enum EShadowDepthPixelShaderMode
@@ -490,7 +469,7 @@ enum EShadowDepthPixelShaderMode
 };
 
 template <EShadowDepthPixelShaderMode ShaderMode, bool bRenderReflectiveShadowMap>
-class TShadowDepthPS : public TShadowDepthBasePS<bRenderReflectiveShadowMap>
+class TShadowDepthPS : public FShadowDepthBasePS
 {
 	DECLARE_SHADER_TYPE(TShadowDepthPS, MeshMaterial);
 public:
@@ -498,19 +477,18 @@ public:
 	static bool ShouldCompilePermutation(const FMeshMaterialShaderPermutationParameters& Parameters)
 	{
 		const EShaderPlatform Platform = Parameters.Platform;
-		const FMaterial* Material = Parameters.Material;
 
 		if (!IsFeatureLevelSupported(Platform, ERHIFeatureLevel::SM5))
 		{
-			return (Material->IsSpecialEngineMaterial()
+			return (Parameters.MaterialParameters.bIsSpecialEngineMaterial
 				// Only compile for masked or lit translucent materials
-				|| !Material->WritesEveryPixel(true)
-				|| (Material->MaterialMayModifyMeshPosition() && Material->IsUsedWithInstancedStaticMeshes())
+				|| !Parameters.MaterialParameters.bWritesEveryPixelShadowPass
+				|| (Parameters.MaterialParameters.bMaterialMayModifyMeshPosition && Parameters.MaterialParameters.bIsUsedWithInstancedStaticMeshes)
 				// Perspective correct rendering needs a pixel shader and WPO materials can't be overridden with default material.
-				|| (ShaderMode == PixelShadowDepth_PerspectiveCorrect && Material->MaterialMayModifyMeshPosition()))
+				|| (ShaderMode == PixelShadowDepth_PerspectiveCorrect && Parameters.MaterialParameters.bMaterialMayModifyMeshPosition))
 				&& ShaderMode == PixelShadowDepth_NonPerspectiveCorrect
 				// Don't render ShadowDepth for translucent unlit materials
-				&& Material->ShouldCastDynamicShadows()
+				&& Parameters.MaterialParameters.bShouldCastDynamicShadows
 				&& !bRenderReflectiveShadowMap;
 		}
 
@@ -527,30 +505,30 @@ public:
 				// Reflective shadow map shaders must be compiled for every material because they access the material normal
 				return
 					// Only compile one pass point light shaders for feature levels >= SM4
-					(Material->ShouldCastDynamicShadows() || Material->ShouldInjectEmissiveIntoLPV() || Material->ShouldBlockGI())
+					(Parameters.MaterialParameters.bShouldCastDynamicShadows || Parameters.MaterialParameters.bShouldInjectEmissiveIntoLPV || Parameters.MaterialParameters.bShouldBlockGI)
 					&& IsFeatureLevelSupported(Platform, ERHIFeatureLevel::SM5);
 			}
 		}
 		else
 		{
 			//Note: This logic needs to stay in sync with OverrideWithDefaultMaterialForShadowDepth!
-			return (Material->IsSpecialEngineMaterial()
+			return (Parameters.MaterialParameters.bIsSpecialEngineMaterial
 				// Only compile for masked or lit translucent materials
-				|| !Material->WritesEveryPixel(true)
-				|| (Material->MaterialMayModifyMeshPosition() && Material->IsUsedWithInstancedStaticMeshes())
+				|| !Parameters.MaterialParameters.bWritesEveryPixelShadowPass
+				|| (Parameters.MaterialParameters.bMaterialMayModifyMeshPosition && Parameters.MaterialParameters.bIsUsedWithInstancedStaticMeshes)
 				// Perspective correct rendering needs a pixel shader and WPO materials can't be overridden with default material.
-				|| (ShaderMode == PixelShadowDepth_PerspectiveCorrect && Material->MaterialMayModifyMeshPosition()))
+				|| (ShaderMode == PixelShadowDepth_PerspectiveCorrect && Parameters.MaterialParameters.bMaterialMayModifyMeshPosition))
 				// Only compile one pass point light shaders for feature levels >= SM5
 				&& (ShaderMode != PixelShadowDepth_OnePassPointLight || IsFeatureLevelSupported(Platform, ERHIFeatureLevel::SM5))
 				// Don't render ShadowDepth for translucent unlit materials
-				&& Material->ShouldCastDynamicShadows()
+				&& Parameters.MaterialParameters.bShouldCastDynamicShadows
 				&& IsFeatureLevelSupported(Platform, ERHIFeatureLevel::SM5);
 		}
 	}
 
 	static void ModifyCompilationEnvironment(const FMaterialShaderPermutationParameters& Parameters, FShaderCompilerEnvironment& OutEnvironment)
 	{
-		TShadowDepthBasePS<bRenderReflectiveShadowMap>::ModifyCompilationEnvironment(Parameters, OutEnvironment);
+		FShadowDepthBasePS::ModifyCompilationEnvironment(Parameters, OutEnvironment);
 		OutEnvironment.SetDefine(TEXT("PERSPECTIVE_CORRECT_DEPTH"), (uint32)(ShaderMode == PixelShadowDepth_PerspectiveCorrect));
 		OutEnvironment.SetDefine(TEXT("ONEPASS_POINTLIGHT_SHADOW"), (uint32)(ShaderMode == PixelShadowDepth_OnePassPointLight));
 		OutEnvironment.SetDefine(TEXT("REFLECTIVE_SHADOW_MAP"), (uint32)bRenderReflectiveShadowMap);
@@ -561,7 +539,7 @@ public:
 	}
 
 	TShadowDepthPS(const ShaderMetaType::CompiledShaderInitializerType& Initializer)
-		: TShadowDepthBasePS<bRenderReflectiveShadowMap>(Initializer)
+		: FShadowDepthBasePS(Initializer)
 	{
 	}
 };
@@ -611,11 +589,11 @@ void GetShadowDepthPassShaders(
 	bool bDirectionalLight,
 	bool bOnePassPointLightShadow,
 	bool bPositionOnlyVS,
-	FShadowDepthVS*& VertexShader,
-	FBaseHS*& HullShader,
-	FBaseDS*& DomainShader,
-	TShadowDepthBasePS<bRenderingReflectiveShadowMaps>*& PixelShader,
-	FOnePassPointShadowDepthGS*& GeometryShader)
+	TShaderRef<FShadowDepthVS>& VertexShader,
+	TShaderRef<FBaseHS>& HullShader,
+	TShaderRef<FBaseDS>& DomainShader,
+	TShaderRef<FShadowDepthBasePS>& PixelShader,
+	TShaderRef<FOnePassPointShadowDepthGS>& GeometryShader)
 {
 	check(!bOnePassPointLightShadow || !bRenderingReflectiveShadowMaps);
 
@@ -625,9 +603,9 @@ void GetShadowDepthPassShaders(
 	// One pass point lights don't output a linear depth, so they are already perspective correct.
 	const bool bUsePerspectiveCorrectShadowDepths = !bDirectionalLight && !bOnePassPointLightShadow;
 
-	HullShader = nullptr;
-	DomainShader = nullptr;
-	GeometryShader = nullptr;
+	HullShader.Reset();
+	DomainShader.Reset();
+	GeometryShader.Reset();
 
 	FVertexFactoryType* VFType = VertexFactory->GetType();
 
@@ -719,21 +697,21 @@ void GetShadowDepthPassShaders(
 	if (Material.WritesEveryPixel(true) && !bUsePerspectiveCorrectShadowDepths && !bRenderingReflectiveShadowMaps && VertexFactory->SupportsNullPixelShader())
 	{
 		// No pixel shader necessary.
-		PixelShader = NULL;
+		PixelShader.Reset();
 	}
 	else
 	{
 		if (bUsePerspectiveCorrectShadowDepths)
 		{
-			PixelShader = (TShadowDepthBasePS<bRenderingReflectiveShadowMaps> *)Material.GetShader<TShadowDepthPS<PixelShadowDepth_PerspectiveCorrect, bRenderingReflectiveShadowMaps> >(VFType, false);
+			PixelShader = Material.GetShader<TShadowDepthPS<PixelShadowDepth_PerspectiveCorrect, bRenderingReflectiveShadowMaps> >(VFType, false);
 		}
 		else if (bOnePassPointLightShadow)
 		{
-			PixelShader = (TShadowDepthBasePS<bRenderingReflectiveShadowMaps> *)Material.GetShader<TShadowDepthPS<PixelShadowDepth_OnePassPointLight, false> >(VFType, false);
+			PixelShader = Material.GetShader<TShadowDepthPS<PixelShadowDepth_OnePassPointLight, false> >(VFType, false);
 		}
 		else
 		{
-			PixelShader = (TShadowDepthBasePS<bRenderingReflectiveShadowMaps> *)Material.GetShader<TShadowDepthPS<PixelShadowDepth_NonPerspectiveCorrect, bRenderingReflectiveShadowMaps> >(VFType, false);
+			PixelShader = Material.GetShader<TShadowDepthPS<PixelShadowDepth_NonPerspectiveCorrect, bRenderingReflectiveShadowMaps> >(VFType, false);
 		}
 	}
 }
@@ -936,21 +914,13 @@ public:
 
 	void SetParameters(FRHICommandList& RHICmdList, const FSceneView& View, IPooledRenderTarget* SourceShadowMap)
 	{
-		FGlobalShader::SetParameters<FViewUniformShaderParameters>(RHICmdList, GetPixelShader(), View.ViewUniformBuffer);
-
-		SetTextureParameter(RHICmdList, GetPixelShader(), ShadowDepthTexture, ShadowDepthSampler, TStaticSamplerState<SF_Point, AM_Clamp, AM_Clamp, AM_Clamp>::GetRHI(), SourceShadowMap->GetRenderTargetItem().ShaderResourceTexture);
-	}
+		FGlobalShader::SetParameters<FViewUniformShaderParameters>(RHICmdList, RHICmdList.GetBoundPixelShader(), View.ViewUniformBuffer);
 
-	virtual bool Serialize(FArchive& Ar) override
-	{
-		bool bShaderHasOutdatedParameters = FGlobalShader::Serialize(Ar);
-		Ar << ShadowDepthTexture;
-		Ar << ShadowDepthSampler;
-		return bShaderHasOutdatedParameters;
+		SetTextureParameter(RHICmdList, RHICmdList.GetBoundPixelShader(), ShadowDepthTexture, ShadowDepthSampler, TStaticSamplerState<SF_Point, AM_Clamp, AM_Clamp, AM_Clamp>::GetRHI(), SourceShadowMap->GetRenderTargetItem().ShaderResourceTexture);
 	}
 
-	FShaderResourceParameter ShadowDepthTexture;
-	FShaderResourceParameter ShadowDepthSampler;
+	LAYOUT_FIELD(FShaderResourceParameter, ShadowDepthTexture);
+	LAYOUT_FIELD(FShaderResourceParameter, ShadowDepthSampler);
 };
 
 IMPLEMENT_SHADER_TYPE(, FCopyShadowMapsCubePS, TEXT("/Engine/Private/CopyShadowMaps.usf"), TEXT("CopyCubeDepthPS"), SF_Pixel);
@@ -976,21 +946,13 @@ public:
 
 	void SetParameters(FRHICommandList& RHICmdList, const FSceneView& View, IPooledRenderTarget* SourceShadowMap)
 	{
-		FGlobalShader::SetParameters<FViewUniformShaderParameters>(RHICmdList, GetPixelShader(), View.ViewUniformBuffer);
-
-		SetTextureParameter(RHICmdList, GetPixelShader(), ShadowDepthTexture, ShadowDepthSampler, TStaticSamplerState<SF_Point, AM_Clamp, AM_Clamp, AM_Clamp>::GetRHI(), SourceShadowMap->GetRenderTargetItem().ShaderResourceTexture);
-	}
+		FGlobalShader::SetParameters<FViewUniformShaderParameters>(RHICmdList, RHICmdList.GetBoundPixelShader(), View.ViewUniformBuffer);
 
-	virtual bool Serialize(FArchive& Ar) override
-	{
-		bool bShaderHasOutdatedParameters = FGlobalShader::Serialize(Ar);
-		Ar << ShadowDepthTexture;
-		Ar << ShadowDepthSampler;
-		return bShaderHasOutdatedParameters;
+		SetTextureParameter(RHICmdList, RHICmdList.GetBoundPixelShader(), ShadowDepthTexture, ShadowDepthSampler, TStaticSamplerState<SF_Point, AM_Clamp, AM_Clamp, AM_Clamp>::GetRHI(), SourceShadowMap->GetRenderTargetItem().ShaderResourceTexture);
 	}
 
-	FShaderResourceParameter ShadowDepthTexture;
-	FShaderResourceParameter ShadowDepthSampler;
+	LAYOUT_FIELD(FShaderResourceParameter, ShadowDepthTexture);
+	LAYOUT_FIELD(FShaderResourceParameter, ShadowDepthSampler);
 };
 
 IMPLEMENT_SHADER_TYPE(, FCopyShadowMaps2DPS, TEXT("/Engine/Private/CopyShadowMaps.usf"), TEXT("Copy2DDepthPS"), SF_Pixel);
@@ -1026,11 +988,11 @@ void FProjectedShadowInfo::CopyCachedShadowMap(FRHICommandList& RHICmdList, cons
 				TShaderMapRef<FCopyShadowMapsCubePS> PixelShader(View.ShaderMap);
 
 				GraphicsPSOInit.BoundShaderState.VertexDeclarationRHI = GFilterVertexDeclaration.VertexDeclarationRHI;
-				GraphicsPSOInit.BoundShaderState.VertexShaderRHI = GETSAFERHISHADER_VERTEX(*ScreenVertexShader);
+				GraphicsPSOInit.BoundShaderState.VertexShaderRHI = ScreenVertexShader.GetVertexShader();
 #if PLATFORM_SUPPORTS_GEOMETRY_SHADERS
-				GraphicsPSOInit.BoundShaderState.GeometryShaderRHI = GETSAFERHISHADER_GEOMETRY(*GeometryShader);
+				GraphicsPSOInit.BoundShaderState.GeometryShaderRHI = GeometryShader.GetGeometryShader();
 #endif
-				GraphicsPSOInit.BoundShaderState.PixelShaderRHI = GETSAFERHISHADER_PIXEL(*PixelShader);
+				GraphicsPSOInit.BoundShaderState.PixelShaderRHI = PixelShader.GetPixelShader();
 				GraphicsPSOInit.PrimitiveType = PT_TriangleList;
 
 				SetGraphicsPipelineState(RHICmdList, GraphicsPSOInit);
@@ -1046,7 +1008,7 @@ void FProjectedShadowInfo::CopyCachedShadowMap(FRHICommandList& RHICmdList, cons
 					ResolutionX, ResolutionY,
 					FIntPoint(ResolutionX, ResolutionY),
 					CachedShadowMapData.ShadowMap.GetSize(),
-					*ScreenVertexShader,
+					ScreenVertexShader,
 					EDRF_Default);
 			}
 			else
@@ -1058,8 +1020,8 @@ void FProjectedShadowInfo::CopyCachedShadowMap(FRHICommandList& RHICmdList, cons
 				TShaderMapRef<FCopyShadowMapsCubePS> PixelShader(View.ShaderMap);
 
 				GraphicsPSOInit.BoundShaderState.VertexDeclarationRHI = GFilterVertexDeclaration.VertexDeclarationRHI;
-				GraphicsPSOInit.BoundShaderState.VertexShaderRHI = GETSAFERHISHADER_VERTEX(*ScreenVertexShader);
-				GraphicsPSOInit.BoundShaderState.PixelShaderRHI = GETSAFERHISHADER_PIXEL(*PixelShader);
+				GraphicsPSOInit.BoundShaderState.VertexShaderRHI = ScreenVertexShader.GetVertexShader();
+				GraphicsPSOInit.BoundShaderState.PixelShaderRHI = PixelShader.GetPixelShader();
 				GraphicsPSOInit.PrimitiveType = PT_TriangleList;
 
 				SetGraphicsPipelineState(RHICmdList, GraphicsPSOInit);
@@ -1075,7 +1037,7 @@ void FProjectedShadowInfo::CopyCachedShadowMap(FRHICommandList& RHICmdList, cons
 					ResolutionX, ResolutionY,
 					FIntPoint(ResolutionX, ResolutionY),
 					CachedShadowMapData.ShadowMap.GetSize(),
-					*ScreenVertexShader,
+					ScreenVertexShader,
 					EDRF_Default,
 					6);
 			}
@@ -1087,8 +1049,8 @@ void FProjectedShadowInfo::CopyCachedShadowMap(FRHICommandList& RHICmdList, cons
 			TShaderMapRef<FCopyShadowMaps2DPS> PixelShader(View.ShaderMap);
 
 			GraphicsPSOInit.BoundShaderState.VertexDeclarationRHI = GFilterVertexDeclaration.VertexDeclarationRHI;
-			GraphicsPSOInit.BoundShaderState.VertexShaderRHI = GETSAFERHISHADER_VERTEX(*ScreenVertexShader);
-			GraphicsPSOInit.BoundShaderState.PixelShaderRHI = GETSAFERHISHADER_PIXEL(*PixelShader);
+			GraphicsPSOInit.BoundShaderState.VertexShaderRHI = ScreenVertexShader.GetVertexShader();
+			GraphicsPSOInit.BoundShaderState.PixelShaderRHI = PixelShader.GetPixelShader();
 			GraphicsPSOInit.PrimitiveType = PT_TriangleList;
 
 			SetGraphicsPipelineState(RHICmdList, GraphicsPSOInit);
@@ -1104,7 +1066,7 @@ void FProjectedShadowInfo::CopyCachedShadowMap(FRHICommandList& RHICmdList, cons
 				ResolutionX, ResolutionY,
 				FIntPoint(ResolutionX, ResolutionY),
 				CachedShadowMapData.ShadowMap.GetSize(),
-				*ScreenVertexShader,
+				ScreenVertexShader,
 				EDRF_Default);
 		}
 	}
@@ -1899,7 +1861,7 @@ void FShadowDepthPassMeshProcessor::Process(
 		FShadowDepthVS,
 		FBaseHS,
 		FBaseDS,
-		TShadowDepthBasePS<bRenderReflectiveShadowMap>,
+		FShadowDepthBasePS,
 		FOnePassPointShadowDepthGS> ShadowDepthPassShaders;
 
 	const bool bUsePositionOnlyVS = 
@@ -1960,12 +1922,13 @@ void FShadowDepthPassMeshProcessor::AddMeshBatch(const FMeshBatch& RESTRICT Mesh
 		const bool bReflectiveShadowmap = ShadowDepthType.bReflectiveShadowmap && !ShadowDepthType.bOnePassPointLightShadow;
 		const bool bShouldCastShadow = Material.ShouldCastDynamicShadows();
 
-		const ERasterizerFillMode MeshFillMode = ComputeMeshFillMode(MeshBatch, Material);
+		const FMeshDrawingPolicyOverrideSettings OverrideSettings = ComputeMeshOverrideSettings(MeshBatch);
+		const ERasterizerFillMode MeshFillMode = ComputeMeshFillMode(MeshBatch, Material, OverrideSettings);
 
 		ERasterizerCullMode FinalCullMode;
 
 		{
-			const ERasterizerCullMode MeshCullMode = ComputeMeshCullMode(MeshBatch, Material);
+			const ERasterizerCullMode MeshCullMode = ComputeMeshCullMode(MeshBatch, Material, OverrideSettings);
 
 			const bool bTwoSided = Material.IsTwoSided() || PrimitiveSceneProxy->CastsShadowAsTwoSided();
 			// @TODO: only render directional light shadows as two sided, and only when blocking is enabled (required by geometry volume injection)
