Copying Private-LoadTimes-4.24 stream to Main. Biggest changes are in Materials/Shader memory freezing.

diff --git a/Engine/Source/Runtime/Renderer/Private/ShadowRendering.cpp b/Engine/Source/Runtime/Renderer/Private/ShadowRendering.cpp
--- a/Engine/Source/Runtime/Renderer/Private/ShadowRendering.cpp
+++ b/Engine/Source/Runtime/Renderer/Private/ShadowRendering.cpp
@@ -198,7 +198,7 @@ uint32 GetShadowQuality()
 
 void FShadowVolumeBoundProjectionVS::SetParameters(FRHICommandList& RHICmdList, const FSceneView& View, const FProjectedShadowInfo* ShadowInfo)
 {
-	FGlobalShader::SetParameters<FViewUniformShaderParameters>(RHICmdList, GetVertexShader(),View.ViewUniformBuffer);
+	FGlobalShader::SetParameters<FViewUniformShaderParameters>(RHICmdList, RHICmdList.GetBoundVertexShader(), View.ViewUniformBuffer);
 	
 	if(ShadowInfo->IsWholeSceneDirectionalShadow())
 	{
@@ -218,6 +218,9 @@ void FShadowVolumeBoundProjectionVS::SetParameters(FRHICommandList& RHICmdList,
 	}
 }
 
+IMPLEMENT_TYPE_LAYOUT(FShadowProjectionVertexShaderInterface);
+IMPLEMENT_TYPE_LAYOUT(FShadowProjectionPixelShaderInterface);
+
 IMPLEMENT_SHADER_TYPE(,FShadowProjectionNoTransformVS,TEXT("/Engine/Private/ShadowProjectionVertexShader.usf"),TEXT("Main"),SF_Vertex);
 
 IMPLEMENT_SHADER_TYPE(,FShadowVolumeBoundProjectionVS,TEXT("/Engine/Private/ShadowProjectionVertexShader.usf"),TEXT("Main"),SF_Vertex);
@@ -262,16 +265,24 @@ IMPLEMENT_SHADOW_PROJECTION_PIXEL_SHADER(2, false, false, true);
 IMPLEMENT_SHADOW_PROJECTION_PIXEL_SHADER(3, false, false, true);
 IMPLEMENT_SHADOW_PROJECTION_PIXEL_SHADER(4, false, false, true);
 IMPLEMENT_SHADOW_PROJECTION_PIXEL_SHADER(5, false, false, true);
-													   
+
 #undef IMPLEMENT_SHADOW_PROJECTION_PIXEL_SHADER
-#endif
+
+#define IMPLEMENT_MODULATED_SHADOW_PROJECTION_PIXEL_SHADER(Quality) \
+	using FShadowModulatedProjectionPS##Quality = TShadowProjectionPS<Quality, false, true>; \
+	IMPLEMENT_TEMPLATE_TYPE_LAYOUT(template<>, FShadowModulatedProjectionPS##Quality); \
+	IMPLEMENT_SHADER_TYPE(template<>, TModulatedShadowProjection<Quality>, TEXT("/Engine/Private/ShadowProjectionPixelShader.usf"), TEXT("Main"), SF_Pixel);
 
 // Implement a pixel shader for rendering modulated shadow projections.
-IMPLEMENT_SHADER_TYPE(template<>, TModulatedShadowProjection<1>, TEXT("/Engine/Private/ShadowProjectionPixelShader.usf"), TEXT("Main"), SF_Pixel);
-IMPLEMENT_SHADER_TYPE(template<>, TModulatedShadowProjection<2>, TEXT("/Engine/Private/ShadowProjectionPixelShader.usf"), TEXT("Main"), SF_Pixel);
-IMPLEMENT_SHADER_TYPE(template<>, TModulatedShadowProjection<3>, TEXT("/Engine/Private/ShadowProjectionPixelShader.usf"), TEXT("Main"), SF_Pixel);
-IMPLEMENT_SHADER_TYPE(template<>, TModulatedShadowProjection<4>, TEXT("/Engine/Private/ShadowProjectionPixelShader.usf"), TEXT("Main"), SF_Pixel);
-IMPLEMENT_SHADER_TYPE(template<>, TModulatedShadowProjection<5>, TEXT("/Engine/Private/ShadowProjectionPixelShader.usf"), TEXT("Main"), SF_Pixel);
+IMPLEMENT_MODULATED_SHADOW_PROJECTION_PIXEL_SHADER(1);
+IMPLEMENT_MODULATED_SHADOW_PROJECTION_PIXEL_SHADER(2);
+IMPLEMENT_MODULATED_SHADOW_PROJECTION_PIXEL_SHADER(3);
+IMPLEMENT_MODULATED_SHADOW_PROJECTION_PIXEL_SHADER(4);
+IMPLEMENT_MODULATED_SHADOW_PROJECTION_PIXEL_SHADER(5);
+
+#undef IMPLEMENT_MODULATED_SHADOW_PROJECTION_PIXEL_SHADER
+
+#endif
 
 // with different quality levels
 IMPLEMENT_SHADER_TYPE(template<>,TShadowProjectionFromTranslucencyPS<1>,TEXT("/Engine/Private/ShadowProjectionPixelShader.usf"),TEXT("Main"),SF_Pixel);
@@ -319,60 +330,79 @@ IMPLEMENT_SHADOW_PROJECTION_PIXEL_SHADER(5, false);
 IMPLEMENT_SHADOW_PROJECTION_PIXEL_SHADER(5, true);
 #undef IMPLEMENT_SHADOW_PROJECTION_PIXEL_SHADER
 
-static void GetShadowProjectionShaders(
-	int32 Quality, const FViewInfo& View, const FProjectedShadowInfo* ShadowInfo, bool bMobileModulatedProjections, bool bSubPixelSupport,
-	FShadowProjectionVertexShaderInterface** OutShadowProjVS, FShadowProjectionPixelShaderInterface** OutShadowProjPS)
+template<typename VertexShaderType, typename PixelShaderType>
+static void BindShaderShaders(FRHICommandList& RHICmdList, FGraphicsPipelineStateInitializer& GraphicsPSOInit,
+	int32 ViewIndex, const FViewInfo& View, const FHairStrandsVisibilityData* HairVisibilityData, const FProjectedShadowInfo* ShadowInfo)
 {
-	check(!*OutShadowProjVS);
-	check(!*OutShadowProjPS);
+	TShaderRef<VertexShaderType> VertexShader = View.ShaderMap->GetShader<VertexShaderType>();
+	TShaderRef<PixelShaderType> PixelShader = View.ShaderMap->GetShader<PixelShaderType>();
+
+	GraphicsPSOInit.BoundShaderState.VertexShaderRHI = VertexShader.GetVertexShader();
+	GraphicsPSOInit.BoundShaderState.PixelShaderRHI = PixelShader.GetPixelShader();
+	SetGraphicsPipelineState(RHICmdList, GraphicsPSOInit);
+
+	VertexShader->SetParameters(RHICmdList, View, ShadowInfo);
+	PixelShader->SetParameters(RHICmdList, ViewIndex, View, HairVisibilityData, ShadowInfo);
+}
 
-	if (bSubPixelSupport)
+
+static void BindShadowProjectionShaders(int32 Quality, FRHICommandList& RHICmdList, FGraphicsPipelineStateInitializer GraphicsPSOInit, int32 ViewIndex, const FViewInfo& View,
+	const FHairStrandsVisibilityData* HairVisibilityData, const FProjectedShadowInfo* ShadowInfo, bool bMobileModulatedProjections)
+{
+	if (HairVisibilityData)
 	{
 		check(!bMobileModulatedProjections);
 
 		if (ShadowInfo->IsWholeSceneDirectionalShadow())
-			*OutShadowProjVS = View.ShaderMap->GetShader<FShadowProjectionNoTransformVS>();
+		{
+			switch (Quality)
+			{
+			case 1: BindShaderShaders<FShadowProjectionNoTransformVS, TShadowProjectionPS<1, false, false, false, true> >(RHICmdList, GraphicsPSOInit, ViewIndex, View, HairVisibilityData, ShadowInfo); break;
+			case 2: BindShaderShaders<FShadowProjectionNoTransformVS, TShadowProjectionPS<2, false, false, false, true> >(RHICmdList, GraphicsPSOInit, ViewIndex, View, HairVisibilityData, ShadowInfo); break;
+			case 3: BindShaderShaders<FShadowProjectionNoTransformVS, TShadowProjectionPS<3, false, false, false, true> >(RHICmdList, GraphicsPSOInit, ViewIndex, View, HairVisibilityData, ShadowInfo); break;
+			case 4: BindShaderShaders<FShadowProjectionNoTransformVS, TShadowProjectionPS<4, false, false, false, true> >(RHICmdList, GraphicsPSOInit, ViewIndex, View, HairVisibilityData, ShadowInfo); break;
+			case 5: BindShaderShaders<FShadowProjectionNoTransformVS, TShadowProjectionPS<5, false, false, false, true> >(RHICmdList, GraphicsPSOInit, ViewIndex, View, HairVisibilityData, ShadowInfo); break;
+			default:
+				check(0);
+			}
+		}
 		else
-			*OutShadowProjVS = View.ShaderMap->GetShader<FShadowVolumeBoundProjectionVS>();
-
-		switch (Quality)
 		{
-		case 1: *OutShadowProjPS = View.ShaderMap->GetShader<TShadowProjectionPS<1, false, false, false, true> >(); break;
-		case 2: *OutShadowProjPS = View.ShaderMap->GetShader<TShadowProjectionPS<2, false, false, false, true> >(); break;
-		case 3: *OutShadowProjPS = View.ShaderMap->GetShader<TShadowProjectionPS<3, false, false, false, true> >(); break;
-		case 4: *OutShadowProjPS = View.ShaderMap->GetShader<TShadowProjectionPS<4, false, false, false, true> >(); break;
-		case 5: *OutShadowProjPS = View.ShaderMap->GetShader<TShadowProjectionPS<5, false, false, false, true> >(); break;
-		default:
-			check(0);
+			switch (Quality)
+			{
+			case 1: BindShaderShaders<FShadowVolumeBoundProjectionVS, TShadowProjectionPS<1, false, false, false, true> >(RHICmdList, GraphicsPSOInit, ViewIndex, View, HairVisibilityData, ShadowInfo); break;
+			case 2: BindShaderShaders<FShadowVolumeBoundProjectionVS, TShadowProjectionPS<2, false, false, false, true> >(RHICmdList, GraphicsPSOInit, ViewIndex, View, HairVisibilityData, ShadowInfo); break;
+			case 3: BindShaderShaders<FShadowVolumeBoundProjectionVS, TShadowProjectionPS<3, false, false, false, true> >(RHICmdList, GraphicsPSOInit, ViewIndex, View, HairVisibilityData, ShadowInfo); break;
+			case 4: BindShaderShaders<FShadowVolumeBoundProjectionVS, TShadowProjectionPS<4, false, false, false, true> >(RHICmdList, GraphicsPSOInit, ViewIndex, View, HairVisibilityData, ShadowInfo); break;
+			case 5: BindShaderShaders<FShadowVolumeBoundProjectionVS, TShadowProjectionPS<5, false, false, false, true> >(RHICmdList, GraphicsPSOInit, ViewIndex, View, HairVisibilityData, ShadowInfo); break;
+			default:
+				check(0);
+			}
 		}
 		return;
 	}
 
 	if (ShadowInfo->bTranslucentShadow)
 	{
-		*OutShadowProjVS = View.ShaderMap->GetShader<FShadowVolumeBoundProjectionVS>();
-
 		switch (Quality)
 		{
-		case 1: *OutShadowProjPS = View.ShaderMap->GetShader<TShadowProjectionFromTranslucencyPS<1> >(); break;
-		case 2: *OutShadowProjPS = View.ShaderMap->GetShader<TShadowProjectionFromTranslucencyPS<2> >(); break;
-		case 3: *OutShadowProjPS = View.ShaderMap->GetShader<TShadowProjectionFromTranslucencyPS<3> >(); break;
-		case 4: *OutShadowProjPS = View.ShaderMap->GetShader<TShadowProjectionFromTranslucencyPS<4> >(); break;
-		case 5: *OutShadowProjPS = View.ShaderMap->GetShader<TShadowProjectionFromTranslucencyPS<5> >(); break;
+		case 1: BindShaderShaders<FShadowVolumeBoundProjectionVS, TShadowProjectionFromTranslucencyPS<1> >(RHICmdList, GraphicsPSOInit, ViewIndex, View, HairVisibilityData, ShadowInfo); break;
+		case 2: BindShaderShaders<FShadowVolumeBoundProjectionVS, TShadowProjectionFromTranslucencyPS<2> >(RHICmdList, GraphicsPSOInit, ViewIndex, View, HairVisibilityData, ShadowInfo); break;
+		case 3: BindShaderShaders<FShadowVolumeBoundProjectionVS, TShadowProjectionFromTranslucencyPS<3> >(RHICmdList, GraphicsPSOInit, ViewIndex, View, HairVisibilityData, ShadowInfo); break;
+		case 4: BindShaderShaders<FShadowVolumeBoundProjectionVS, TShadowProjectionFromTranslucencyPS<4> >(RHICmdList, GraphicsPSOInit, ViewIndex, View, HairVisibilityData, ShadowInfo); break;
+		case 5: BindShaderShaders<FShadowVolumeBoundProjectionVS, TShadowProjectionFromTranslucencyPS<5> >(RHICmdList, GraphicsPSOInit, ViewIndex, View, HairVisibilityData, ShadowInfo); break;
 		default:
 			check(0);
 		}
 	}
 	else if (ShadowInfo->IsWholeSceneDirectionalShadow())
 	{
-		*OutShadowProjVS = View.ShaderMap->GetShader<FShadowProjectionNoTransformVS>();
-
 		if (CVarFilterMethod.GetValueOnRenderThread() == 1)
 		{
 			if (ShadowInfo->CascadeSettings.FadePlaneLength > 0)
-				*OutShadowProjPS = View.ShaderMap->GetShader<TDirectionalPercentageCloserShadowProjectionPS<5, true> >();
+				BindShaderShaders<FShadowProjectionNoTransformVS, TDirectionalPercentageCloserShadowProjectionPS<5, true> >(RHICmdList, GraphicsPSOInit, ViewIndex, View, HairVisibilityData, ShadowInfo);
 			else
-				*OutShadowProjPS = View.ShaderMap->GetShader<TDirectionalPercentageCloserShadowProjectionPS<5, false> >();
+				BindShaderShaders<FShadowProjectionNoTransformVS, TDirectionalPercentageCloserShadowProjectionPS<5, false> >(RHICmdList, GraphicsPSOInit, ViewIndex, View, HairVisibilityData, ShadowInfo);
 		}
 		else if (ShadowInfo->CascadeSettings.FadePlaneLength > 0)
 		{
@@ -380,11 +410,11 @@ static void GetShadowProjectionShaders(
 			{
 				switch (Quality)
 				{
-				case 1: *OutShadowProjPS = View.ShaderMap->GetShader<TShadowProjectionPS<1, true, false, true> >(); break;
-				case 2: *OutShadowProjPS = View.ShaderMap->GetShader<TShadowProjectionPS<2, true, false, true> >(); break;
-				case 3: *OutShadowProjPS = View.ShaderMap->GetShader<TShadowProjectionPS<3, true, false, true> >(); break;
-				case 4: *OutShadowProjPS = View.ShaderMap->GetShader<TShadowProjectionPS<4, true, false, true> >(); break;
-				case 5: *OutShadowProjPS = View.ShaderMap->GetShader<TShadowProjectionPS<5, true, false, true> >(); break;
+				case 1: BindShaderShaders<FShadowProjectionNoTransformVS, TShadowProjectionPS<1, true, false, true> >(RHICmdList, GraphicsPSOInit, ViewIndex, View, HairVisibilityData, ShadowInfo); break;
+				case 2: BindShaderShaders<FShadowProjectionNoTransformVS, TShadowProjectionPS<2, true, false, true> >(RHICmdList, GraphicsPSOInit, ViewIndex, View, HairVisibilityData, ShadowInfo); break;
+				case 3: BindShaderShaders<FShadowProjectionNoTransformVS, TShadowProjectionPS<3, true, false, true> >(RHICmdList, GraphicsPSOInit, ViewIndex, View, HairVisibilityData, ShadowInfo); break;
+				case 4: BindShaderShaders<FShadowProjectionNoTransformVS, TShadowProjectionPS<4, true, false, true> >(RHICmdList, GraphicsPSOInit, ViewIndex, View, HairVisibilityData, ShadowInfo); break;
+				case 5: BindShaderShaders<FShadowProjectionNoTransformVS, TShadowProjectionPS<5, true, false, true> >(RHICmdList, GraphicsPSOInit, ViewIndex, View, HairVisibilityData, ShadowInfo); break;
 				default:
 					check(0);
 				}
@@ -393,11 +423,11 @@ static void GetShadowProjectionShaders(
 			{
 				switch (Quality)
 				{
-				case 1: *OutShadowProjPS = View.ShaderMap->GetShader<TShadowProjectionPS<1, true> >(); break;
-				case 2: *OutShadowProjPS = View.ShaderMap->GetShader<TShadowProjectionPS<2, true> >(); break;
-				case 3: *OutShadowProjPS = View.ShaderMap->GetShader<TShadowProjectionPS<3, true> >(); break;
-				case 4: *OutShadowProjPS = View.ShaderMap->GetShader<TShadowProjectionPS<4, true> >(); break;
-				case 5: *OutShadowProjPS = View.ShaderMap->GetShader<TShadowProjectionPS<5, true> >(); break;
+				case 1: BindShaderShaders<FShadowProjectionNoTransformVS, TShadowProjectionPS<1, true> >(RHICmdList, GraphicsPSOInit, ViewIndex, View, HairVisibilityData, ShadowInfo); break;
+				case 2: BindShaderShaders<FShadowProjectionNoTransformVS, TShadowProjectionPS<2, true> >(RHICmdList, GraphicsPSOInit, ViewIndex, View, HairVisibilityData, ShadowInfo); break;
+				case 3: BindShaderShaders<FShadowProjectionNoTransformVS, TShadowProjectionPS<3, true> >(RHICmdList, GraphicsPSOInit, ViewIndex, View, HairVisibilityData, ShadowInfo); break;
+				case 4: BindShaderShaders<FShadowProjectionNoTransformVS, TShadowProjectionPS<4, true> >(RHICmdList, GraphicsPSOInit, ViewIndex, View, HairVisibilityData, ShadowInfo); break;
+				case 5: BindShaderShaders<FShadowProjectionNoTransformVS, TShadowProjectionPS<5, true> >(RHICmdList, GraphicsPSOInit, ViewIndex, View, HairVisibilityData, ShadowInfo); break;
 				default:
 					check(0);
 				}
@@ -409,11 +439,11 @@ static void GetShadowProjectionShaders(
 			{
 				switch (Quality)
 				{
-				case 1: *OutShadowProjPS = View.ShaderMap->GetShader<TShadowProjectionPS<1, false, false, true> >(); break;
-				case 2: *OutShadowProjPS = View.ShaderMap->GetShader<TShadowProjectionPS<2, false, false, true> >(); break;
-				case 3: *OutShadowProjPS = View.ShaderMap->GetShader<TShadowProjectionPS<3, false, false, true> >(); break;
-				case 4: *OutShadowProjPS = View.ShaderMap->GetShader<TShadowProjectionPS<4, false, false, true> >(); break;
-				case 5: *OutShadowProjPS = View.ShaderMap->GetShader<TShadowProjectionPS<5, false, false, true> >(); break;
+				case 1: BindShaderShaders<FShadowProjectionNoTransformVS, TShadowProjectionPS<1, false, false, true> >(RHICmdList, GraphicsPSOInit, ViewIndex, View, HairVisibilityData, ShadowInfo); break;
+				case 2: BindShaderShaders<FShadowProjectionNoTransformVS, TShadowProjectionPS<2, false, false, true> >(RHICmdList, GraphicsPSOInit, ViewIndex, View, HairVisibilityData, ShadowInfo); break;
+				case 3: BindShaderShaders<FShadowProjectionNoTransformVS, TShadowProjectionPS<3, false, false, true> >(RHICmdList, GraphicsPSOInit, ViewIndex, View, HairVisibilityData, ShadowInfo); break;
+				case 4: BindShaderShaders<FShadowProjectionNoTransformVS, TShadowProjectionPS<4, false, false, true> >(RHICmdList, GraphicsPSOInit, ViewIndex, View, HairVisibilityData, ShadowInfo); break;
+				case 5: BindShaderShaders<FShadowProjectionNoTransformVS, TShadowProjectionPS<5, false, false, true> >(RHICmdList, GraphicsPSOInit, ViewIndex, View, HairVisibilityData, ShadowInfo); break;
 				default:
 					check(0);
 				}
@@ -422,11 +452,11 @@ static void GetShadowProjectionShaders(
 			{ 
 				switch (Quality)
 				{
-				case 1: *OutShadowProjPS = View.ShaderMap->GetShader<TShadowProjectionPS<1, false> >(); break;
-				case 2: *OutShadowProjPS = View.ShaderMap->GetShader<TShadowProjectionPS<2, false> >(); break;
-				case 3: *OutShadowProjPS = View.ShaderMap->GetShader<TShadowProjectionPS<3, false> >(); break;
-				case 4: *OutShadowProjPS = View.ShaderMap->GetShader<TShadowProjectionPS<4, false> >(); break;
-				case 5: *OutShadowProjPS = View.ShaderMap->GetShader<TShadowProjectionPS<5, false> >(); break;
+				case 1: BindShaderShaders<FShadowProjectionNoTransformVS, TShadowProjectionPS<1, false> >(RHICmdList, GraphicsPSOInit, ViewIndex, View, HairVisibilityData, ShadowInfo); break;
+				case 2: BindShaderShaders<FShadowProjectionNoTransformVS, TShadowProjectionPS<2, false> >(RHICmdList, GraphicsPSOInit, ViewIndex, View, HairVisibilityData, ShadowInfo); break;
+				case 3: BindShaderShaders<FShadowProjectionNoTransformVS, TShadowProjectionPS<3, false> >(RHICmdList, GraphicsPSOInit, ViewIndex, View, HairVisibilityData, ShadowInfo); break;
+				case 4: BindShaderShaders<FShadowProjectionNoTransformVS, TShadowProjectionPS<4, false> >(RHICmdList, GraphicsPSOInit, ViewIndex, View, HairVisibilityData, ShadowInfo); break;
+				case 5: BindShaderShaders<FShadowProjectionNoTransformVS, TShadowProjectionPS<5, false> >(RHICmdList, GraphicsPSOInit, ViewIndex, View, HairVisibilityData, ShadowInfo); break;
 				default:
 					check(0);
 				}
@@ -435,17 +465,15 @@ static void GetShadowProjectionShaders(
 	}
 	else
 	{
-		*OutShadowProjVS = View.ShaderMap->GetShader<FShadowVolumeBoundProjectionVS>();
-
 		if(bMobileModulatedProjections)
 		{
 			switch (Quality)
 			{
-			case 1: *OutShadowProjPS = View.ShaderMap->GetShader<TModulatedShadowProjection<1> >(); break;
-			case 2: *OutShadowProjPS = View.ShaderMap->GetShader<TModulatedShadowProjection<2> >(); break;
-			case 3: *OutShadowProjPS = View.ShaderMap->GetShader<TModulatedShadowProjection<3> >(); break;
-			case 4: *OutShadowProjPS = View.ShaderMap->GetShader<TModulatedShadowProjection<4> >(); break;
-			case 5: *OutShadowProjPS = View.ShaderMap->GetShader<TModulatedShadowProjection<5> >(); break;
+			case 1: BindShaderShaders<FShadowVolumeBoundProjectionVS, TModulatedShadowProjection<1> >(RHICmdList, GraphicsPSOInit, ViewIndex, View, HairVisibilityData, ShadowInfo); break;
+			case 2: BindShaderShaders<FShadowVolumeBoundProjectionVS, TModulatedShadowProjection<2> >(RHICmdList, GraphicsPSOInit, ViewIndex, View, HairVisibilityData, ShadowInfo); break;
+			case 3: BindShaderShaders<FShadowVolumeBoundProjectionVS, TModulatedShadowProjection<3> >(RHICmdList, GraphicsPSOInit, ViewIndex, View, HairVisibilityData, ShadowInfo); break;
+			case 4: BindShaderShaders<FShadowVolumeBoundProjectionVS, TModulatedShadowProjection<4> >(RHICmdList, GraphicsPSOInit, ViewIndex, View, HairVisibilityData, ShadowInfo); break;
+			case 5: BindShaderShaders<FShadowVolumeBoundProjectionVS, TModulatedShadowProjection<5> >(RHICmdList, GraphicsPSOInit, ViewIndex, View, HairVisibilityData, ShadowInfo); break;
 			default:
 				check(0);
 			}
@@ -454,11 +482,11 @@ static void GetShadowProjectionShaders(
 		{
 			switch (Quality)
 			{
-			case 1: *OutShadowProjPS = View.ShaderMap->GetShader<TShadowProjectionPS<1, false, false, true> >(); break;
-			case 2: *OutShadowProjPS = View.ShaderMap->GetShader<TShadowProjectionPS<2, false, false, true> >(); break;
-			case 3: *OutShadowProjPS = View.ShaderMap->GetShader<TShadowProjectionPS<3, false, false, true> >(); break;
-			case 4: *OutShadowProjPS = View.ShaderMap->GetShader<TShadowProjectionPS<4, false, false, true> >(); break;
-			case 5: *OutShadowProjPS = View.ShaderMap->GetShader<TShadowProjectionPS<5, false, false, true> >(); break;
+			case 1: BindShaderShaders<FShadowVolumeBoundProjectionVS, TShadowProjectionPS<1, false, false, true> >(RHICmdList, GraphicsPSOInit, ViewIndex, View, HairVisibilityData, ShadowInfo); break;
+			case 2: BindShaderShaders<FShadowVolumeBoundProjectionVS, TShadowProjectionPS<2, false, false, true> >(RHICmdList, GraphicsPSOInit, ViewIndex, View, HairVisibilityData, ShadowInfo); break;
+			case 3: BindShaderShaders<FShadowVolumeBoundProjectionVS, TShadowProjectionPS<3, false, false, true> >(RHICmdList, GraphicsPSOInit, ViewIndex, View, HairVisibilityData, ShadowInfo); break;
+			case 4: BindShaderShaders<FShadowVolumeBoundProjectionVS, TShadowProjectionPS<4, false, false, true> >(RHICmdList, GraphicsPSOInit, ViewIndex, View, HairVisibilityData, ShadowInfo); break;
+			case 5: BindShaderShaders<FShadowVolumeBoundProjectionVS, TShadowProjectionPS<5, false, false, true> >(RHICmdList, GraphicsPSOInit, ViewIndex, View, HairVisibilityData, ShadowInfo); break;
 			default:
 				check(0);
 			}
@@ -467,17 +495,17 @@ static void GetShadowProjectionShaders(
 		{
 			if (CVarFilterMethod.GetValueOnRenderThread() == 1 && ShadowInfo->GetLightSceneInfo().Proxy->GetLightType() == LightType_Spot)
 			{
-				*OutShadowProjPS = View.ShaderMap->GetShader<TSpotPercentageCloserShadowProjectionPS<5, false> >();
+				BindShaderShaders<FShadowVolumeBoundProjectionVS, TSpotPercentageCloserShadowProjectionPS<5, false> >(RHICmdList, GraphicsPSOInit, ViewIndex, View, HairVisibilityData, ShadowInfo);
 			}
 			else
 			{
 				switch (Quality)
 				{
-				case 1: *OutShadowProjPS = View.ShaderMap->GetShader<TShadowProjectionPS<1, false> >(); break;
-				case 2: *OutShadowProjPS = View.ShaderMap->GetShader<TShadowProjectionPS<2, false> >(); break;
-				case 3: *OutShadowProjPS = View.ShaderMap->GetShader<TShadowProjectionPS<3, false> >(); break;
-				case 4: *OutShadowProjPS = View.ShaderMap->GetShader<TShadowProjectionPS<4, false> >(); break;
-				case 5: *OutShadowProjPS = View.ShaderMap->GetShader<TShadowProjectionPS<5, false> >(); break;
+				case 1: BindShaderShaders<FShadowVolumeBoundProjectionVS, TShadowProjectionPS<1, false> >(RHICmdList, GraphicsPSOInit, ViewIndex, View, HairVisibilityData, ShadowInfo); break;
+				case 2: BindShaderShaders<FShadowVolumeBoundProjectionVS, TShadowProjectionPS<2, false> >(RHICmdList, GraphicsPSOInit, ViewIndex, View, HairVisibilityData, ShadowInfo); break;
+				case 3: BindShaderShaders<FShadowVolumeBoundProjectionVS, TShadowProjectionPS<3, false> >(RHICmdList, GraphicsPSOInit, ViewIndex, View, HairVisibilityData, ShadowInfo); break;
+				case 4: BindShaderShaders<FShadowVolumeBoundProjectionVS, TShadowProjectionPS<4, false> >(RHICmdList, GraphicsPSOInit, ViewIndex, View, HairVisibilityData, ShadowInfo); break;
+				case 5: BindShaderShaders<FShadowVolumeBoundProjectionVS, TShadowProjectionPS<5, false> >(RHICmdList, GraphicsPSOInit, ViewIndex, View, HairVisibilityData, ShadowInfo); break;
 				default:
 					check(0);
 				}
@@ -485,8 +513,8 @@ static void GetShadowProjectionShaders(
 		}
 	}
 
-	check(*OutShadowProjVS);
-	check(*OutShadowProjPS);
+	check(GraphicsPSOInit.BoundShaderState.VertexShaderRHI);
+	check(GraphicsPSOInit.BoundShaderState.PixelShaderRHI);
 }
 
 FRHIBlendState* FProjectedShadowInfo::GetBlendStateForProjection(
@@ -763,7 +791,7 @@ void FProjectedShadowInfo::SetupProjectionStencilMask(
 		VertexShaderNoTransform->SetParameters(RHICmdList, View->ViewUniformBuffer);
 
 		GraphicsPSOInit.BoundShaderState.VertexDeclarationRHI = GetVertexDeclarationFVector4();
-		GraphicsPSOInit.BoundShaderState.VertexShaderRHI = GETSAFERHISHADER_VERTEX(*VertexShaderNoTransform);
+		GraphicsPSOInit.BoundShaderState.VertexShaderRHI = VertexShaderNoTransform.GetVertexShader();
 		GraphicsPSOInit.PrimitiveType = PT_TriangleList;
 
 		SetGraphicsPipelineState(RHICmdList, GraphicsPSOInit);
@@ -840,7 +868,7 @@ void FProjectedShadowInfo::SetupProjectionStencilMask(
 		TShaderMapRef<FShadowVolumeBoundProjectionVS> VertexShader(View->ShaderMap);
 
 		GraphicsPSOInit.BoundShaderState.VertexDeclarationRHI = GetVertexDeclarationFVector4();
-		GraphicsPSOInit.BoundShaderState.VertexShaderRHI = GETSAFERHISHADER_VERTEX(*VertexShader);
+		GraphicsPSOInit.BoundShaderState.VertexShaderRHI = VertexShader.GetVertexShader();
 		GraphicsPSOInit.PrimitiveType = PT_TriangleList;
 
 		SetGraphicsPipelineState(RHICmdList, GraphicsPSOInit);
@@ -991,16 +1019,8 @@ void FProjectedShadowInfo::RenderProjection(FRHICommandListImmediate& RHICmdList
 			}
 		}
 
-		FShadowProjectionVertexShaderInterface* ShadowProjVS = nullptr;
-		FShadowProjectionPixelShaderInterface* ShadowProjPS = nullptr;
-
-		GetShadowProjectionShaders(LocalQuality, *View, this, bMobileModulatedProjections, bSubPixelSupport, &ShadowProjVS, &ShadowProjPS);
-
 		GraphicsPSOInit.BoundShaderState.VertexDeclarationRHI = GetVertexDeclarationFVector4();
-		GraphicsPSOInit.BoundShaderState.VertexShaderRHI = GETSAFERHISHADER_VERTEX(ShadowProjVS);
-		GraphicsPSOInit.BoundShaderState.PixelShaderRHI = GETSAFERHISHADER_PIXEL(ShadowProjPS);
-
-		SetGraphicsPipelineState(RHICmdList, GraphicsPSOInit);
+		BindShadowProjectionShaders(LocalQuality, RHICmdList, GraphicsPSOInit, ViewIndex, *View, HairVisibilityData, this, bMobileModulatedProjections);
 
 		if (bDepthBoundsTestEnabled)
 		{
@@ -1008,9 +1028,6 @@ void FProjectedShadowInfo::RenderProjection(FRHICommandListImmediate& RHICmdList
 		}
 
 		RHICmdList.SetStencilRef(0);
-
-		ShadowProjVS->SetParameters(RHICmdList, *View, this);
-		ShadowProjPS->SetParameters(RHICmdList, ViewIndex, *View, HairVisibilityData, this);
 	}
 
 	if (IsWholeSceneDirectionalShadow())
@@ -1050,8 +1067,8 @@ static void SetPointLightShaderTempl(FRHICommandList& RHICmdList, FGraphicsPipel
 	TShaderMapRef<TOnePassPointShadowProjectionPS<Quality,bUseTransmission,bUseSubPixel> > PixelShader(View.ShaderMap);
 
 	GraphicsPSOInit.BoundShaderState.VertexDeclarationRHI = GetVertexDeclarationFVector4();
-	GraphicsPSOInit.BoundShaderState.VertexShaderRHI = GETSAFERHISHADER_VERTEX(*VertexShader);
-	GraphicsPSOInit.BoundShaderState.PixelShaderRHI = GETSAFERHISHADER_PIXEL(*PixelShader);
+	GraphicsPSOInit.BoundShaderState.VertexShaderRHI = VertexShader.GetVertexShader();
+	GraphicsPSOInit.BoundShaderState.PixelShaderRHI = PixelShader.GetPixelShader();
 
 	SetGraphicsPipelineState(RHICmdList, GraphicsPSOInit);
 	
diff --git a/Engine/Source/Runtime/Renderer/Private/ShadowRendering.h b/Engine/Source/Runtime/Renderer/Private/ShadowRendering.h
--- a/Engine/Source/Runtime/Renderer/Private/ShadowRendering.h
+++ b/Engine/Source/Runtime/Renderer/Private/ShadowRendering.h
@@ -42,7 +42,7 @@ class FViewInfo;
 /** Renders a cone with a spherical cap, used for rendering spot lights in deferred passes. */
 extern void DrawStencilingCone(const FMatrix& ConeToWorld, float ConeAngle, float SphereRadius, const FVector& PreViewTranslation);
 
-template <bool bRenderingReflectiveShadowMaps> class TShadowDepthBasePS;
+class FShadowDepthBasePS;
 
 /** 
  * Overrides a material used for shadow depth rendering with the default material when appropriate.
@@ -582,6 +582,7 @@ private:
 
 	FDynamicMeshDrawCommandStorage DynamicMeshDrawCommandStorage;
 	FGraphicsMinimalPipelineStateSet GraphicsMinimalPipelineStateSet;
+	bool NeedsShaderInitialisation;
 
 	/**
 	 * Bias during in shadowmap rendering, stored redundantly for better performance 
@@ -651,7 +652,7 @@ private:
 		bool bCameraInsideShadowFrustum) const;
 
 	friend class FShadowDepthVS;
-	template <bool bRenderingReflectiveShadowMaps> friend class TShadowDepthBasePS;
+	friend class FShadowDepthBasePS;
 	friend class FShadowVolumeBoundProjectionVS;
 	friend class FShadowProjectionPS;
 };
@@ -661,13 +662,17 @@ private:
 */
 class FShadowProjectionVertexShaderInterface : public FGlobalShader
 {
+	DECLARE_TYPE_LAYOUT(FShadowProjectionVertexShaderInterface, NonVirtual);
 public:
 	FShadowProjectionVertexShaderInterface() {}
 	FShadowProjectionVertexShaderInterface(const ShaderMetaType::CompiledShaderInitializerType& Initializer)
 		: FGlobalShader(Initializer)
 	{ }
 
-	virtual void SetParameters(FRHICommandList& RHICmdList, const FSceneView& View, const FProjectedShadowInfo* ShadowInfo) = 0;
+	void SetParameters(FRHICommandList& RHICmdList, const FSceneView& View, const FProjectedShadowInfo* ShadowInfo) {}
+
+	
+	
 };
 
 /**
@@ -696,19 +701,10 @@ public:
 		OutEnvironment.SetDefine(TEXT("USE_TRANSFORM"), (uint32)1);
 	}
 
-	void SetParameters(FRHICommandList& RHICmdList, const FSceneView& View, const FProjectedShadowInfo* ShadowInfo) override;
-
-	//~ Begin FShader Interface
-	virtual bool Serialize(FArchive& Ar) override
-	{
-		bool bShaderHasOutdatedParameters = FGlobalShader::Serialize(Ar);
-		Ar << StencilingGeometryParameters;
-		return bShaderHasOutdatedParameters;
-	}
-	//~ Begin  End FShader Interface 
+	void SetParameters(FRHICommandList& RHICmdList, const FSceneView& View, const FProjectedShadowInfo* ShadowInfo);
 
 private:
-	FStencilingGeometryShaderParameters StencilingGeometryParameters;
+	LAYOUT_FIELD(FStencilingGeometryShaderParameters, StencilingGeometryParameters);
 };
 
 class FShadowProjectionNoTransformVS : public FShadowProjectionVertexShaderInterface
@@ -738,12 +734,12 @@ public:
 
 	void SetParameters(FRHICommandList& RHICmdList, FRHIUniformBuffer* ViewUniformBuffer)
 	{
-		FGlobalShader::SetParameters<FViewUniformShaderParameters>(RHICmdList, GetVertexShader(), ViewUniformBuffer);
+		FGlobalShader::SetParameters<FViewUniformShaderParameters>(RHICmdList, RHICmdList.GetBoundVertexShader(), ViewUniformBuffer);
 	}
 
-	void SetParameters(FRHICommandList& RHICmdList, const FSceneView& View, const FProjectedShadowInfo*) override
+	void SetParameters(FRHICommandList& RHICmdList, const FSceneView& View, const FProjectedShadowInfo*)
 	{
-		FGlobalShader::SetParameters<FViewUniformShaderParameters>(RHICmdList, GetVertexShader(), View.ViewUniformBuffer);
+		FGlobalShader::SetParameters<FViewUniformShaderParameters>(RHICmdList, RHICmdList.GetBoundVertexShader(), View.ViewUniformBuffer);
 	}
 };
 
@@ -753,7 +749,7 @@ public:
 
 class FShadowProjectionPixelShaderInterface : public FGlobalShader
 {
-	DECLARE_SHADER_TYPE(FShadowProjectionPixelShaderInterface,Global);
+	DECLARE_TYPE_LAYOUT(FShadowProjectionPixelShaderInterface, NonVirtual);
 public:
 
 	FShadowProjectionPixelShaderInterface() 
@@ -773,22 +769,21 @@ public:
 	 * @param View - current view
 	 * @param ShadowInfo - projected shadow info for a single light
 	 */
-	virtual void SetParameters(
+	void SetParameters(
 		FRHICommandList& RHICmdList, 
 		int32 ViewIndex,
 		const FSceneView& View,
 		const FHairStrandsVisibilityData* HairVisibilityData,
 		const FProjectedShadowInfo* ShadowInfo)
 	{ 
-		FGlobalShader::SetParameters<FViewUniformShaderParameters>(RHICmdList, GetPixelShader(), View.ViewUniformBuffer);
+		FGlobalShader::SetParameters<FViewUniformShaderParameters>(RHICmdList, RHICmdList.GetBoundPixelShader(), View.ViewUniformBuffer);
 	}
-
 };
 
 /** Shadow projection parameters used by multiple shaders. */
-template<bool bModulatedShadows>
-class TShadowProjectionShaderParameters
+class FShadowProjectionShaderParameters
 {
+	DECLARE_TYPE_LAYOUT(FShadowProjectionShaderParameters, NonVirtual);
 public:
 	void Bind(const FShader::CompiledShaderInitializerType& Initializer)
 	{
@@ -809,9 +804,9 @@ public:
 		InvPerObjectShadowFadeLength.Bind(ParameterMap, TEXT("InvPerObjectShadowFadeLength"));
 	}
 
-	void Set(FRHICommandList& RHICmdList, FShader* Shader, const FSceneView& View, const FProjectedShadowInfo* ShadowInfo, const FHairStrandsVisibilityData* HairVisibilityData)
+	void Set(FRHICommandList& RHICmdList, FShader* Shader, const FSceneView& View, const FProjectedShadowInfo* ShadowInfo, const FHairStrandsVisibilityData* HairVisibilityData, bool bModulatedShadows)
 	{
-		FRHIPixelShader* ShaderRHI = Shader->GetPixelShader();
+		FRHIPixelShader* ShaderRHI = RHICmdList.GetBoundPixelShader();
 
 		SceneTextureParameters.Set(RHICmdList, ShaderRHI, View.FeatureLevel, ESceneTextureSetupMode::All);
 
@@ -908,7 +903,7 @@ public:
 	}
 
 	/** Serializer. */
-	friend FArchive& operator<<(FArchive& Ar, TShadowProjectionShaderParameters& P)
+	/*friend FArchive& operator<<(FArchive& Ar, TShadowProjectionShaderParameters& P)
 	{
 		Ar << P.SceneTextureParameters;
 		Ar << P.ScreenToShadowMatrix;
@@ -925,24 +920,23 @@ public:
 		Ar << P.PerObjectShadowFadeStart;
 		Ar << P.InvPerObjectShadowFadeLength;
 		return Ar;
-	}
+	}*/
 
 private:
-
-	FSceneTextureShaderParameters SceneTextureParameters;
-	FShaderParameter ScreenToShadowMatrix;
-	FShaderParameter SoftTransitionScale;
-	FShaderParameter ShadowBufferSize;
-	FShaderResourceParameter ShadowDepthTexture;
-	FShaderResourceParameter ShadowDepthTextureSampler;
-	FShaderParameter ProjectionDepthBias;
-	FShaderParameter FadePlaneOffset;
-	FShaderParameter InvFadePlaneLength;
-	FShaderParameter ShadowTileOffsetAndSizeParam;
-	FShaderParameter LightPositionOrDirection;
-	FShaderResourceParameter HairCategorizationTexture;
-	FShaderParameter PerObjectShadowFadeStart;
-	FShaderParameter InvPerObjectShadowFadeLength;
+	LAYOUT_FIELD(FSceneTextureShaderParameters, SceneTextureParameters);
+	LAYOUT_FIELD(FShaderParameter, ScreenToShadowMatrix);
+	LAYOUT_FIELD(FShaderParameter, SoftTransitionScale);
+	LAYOUT_FIELD(FShaderParameter, ShadowBufferSize);
+	LAYOUT_FIELD(FShaderResourceParameter, ShadowDepthTexture);
+	LAYOUT_FIELD(FShaderResourceParameter, ShadowDepthTextureSampler);
+	LAYOUT_FIELD(FShaderParameter, ProjectionDepthBias);
+	LAYOUT_FIELD(FShaderParameter, FadePlaneOffset);
+	LAYOUT_FIELD(FShaderParameter, InvFadePlaneLength);
+	LAYOUT_FIELD(FShaderParameter, ShadowTileOffsetAndSizeParam);
+	LAYOUT_FIELD(FShaderParameter, LightPositionOrDirection);
+	LAYOUT_FIELD(FShaderResourceParameter, HairCategorizationTexture);
+	LAYOUT_FIELD(FShaderParameter, PerObjectShadowFadeStart);
+	LAYOUT_FIELD(FShaderParameter, InvPerObjectShadowFadeLength);
 };
 
 /**
@@ -998,18 +992,18 @@ public:
 	 * @param View - current view
 	 * @param ShadowInfo - projected shadow info for a single light
 	 */
-	virtual void SetParameters(
-		FRHICommandList& RHICmdList,
+	void SetParameters(
+		FRHICommandList& RHICmdList, 
 		int32 ViewIndex,
 		const FSceneView& View,
 		const FHairStrandsVisibilityData* HairVisibilityData,
-		const FProjectedShadowInfo* ShadowInfo) override
+		const FProjectedShadowInfo* ShadowInfo)
 	{
-		FRHIPixelShader* ShaderRHI = GetPixelShader();
+		FRHIPixelShader* ShaderRHI = RHICmdList.GetBoundPixelShader();
 
 		FShadowProjectionPixelShaderInterface::SetParameters(RHICmdList, ViewIndex, View, HairVisibilityData, ShadowInfo);
 
-		ProjectionParameters.Set(RHICmdList, this, View, ShadowInfo, HairVisibilityData);
+		ProjectionParameters.Set(RHICmdList, this, View, ShadowInfo, HairVisibilityData, bModulatedShadows);
 		const FLightSceneProxy& LightProxy = *(ShadowInfo->GetLightSceneInfo().Proxy);
 
 		SetShaderValue(RHICmdList, ShaderRHI, ShadowFadeFraction, ShadowInfo->FadeAlphas[ViewIndex] );
@@ -1047,28 +1041,12 @@ public:
 		}
 	}
 
-	/**
-	 * Serialize the parameters for this shader
-	 * @param Ar - archive to serialize to
-	 */
-	virtual bool Serialize(FArchive& Ar) override
-	{
-		bool bShaderHasOutdatedParameters = FShadowProjectionPixelShaderInterface::Serialize(Ar);
-		Ar << ProjectionParameters;
-		Ar << ShadowFadeFraction;
-		Ar << ShadowSharpen;
-		Ar << TransmissionProfilesTexture;
-		Ar << LightPosition;
-		
-		return bShaderHasOutdatedParameters;
-	}
-
 protected:
-	TShadowProjectionShaderParameters<bModulatedShadows> ProjectionParameters;
-	FShaderParameter ShadowFadeFraction;
-	FShaderParameter ShadowSharpen;
-	FShaderParameter LightPosition;
-	FShaderResourceParameter TransmissionProfilesTexture;
+	LAYOUT_FIELD(FShadowProjectionShaderParameters, ProjectionParameters);
+	LAYOUT_FIELD(FShaderParameter, ShadowFadeFraction);
+	LAYOUT_FIELD(FShaderParameter, ShadowSharpen);
+	LAYOUT_FIELD(FShaderParameter, LightPosition);
+	LAYOUT_FIELD(FShaderResourceParameter, TransmissionProfilesTexture);
 };
 
 /** Pixel shader to project modulated shadows onto the scene. */
@@ -1097,31 +1075,20 @@ public:
 		ModulatedShadowColorParameter.Bind(Initializer.ParameterMap, TEXT("ModulatedShadowColor"));
 	}
 
-	virtual void SetParameters(
+	void SetParameters(
 		FRHICommandList& RHICmdList,
 		int32 ViewIndex,
 		const FSceneView& View,
 		const FHairStrandsVisibilityData* HairVisibilityData,
-		const FProjectedShadowInfo* ShadowInfo) override
+		const FProjectedShadowInfo* ShadowInfo)
 	{
 		TShadowProjectionPS<Quality, false, true>::SetParameters(RHICmdList, ViewIndex, View, HairVisibilityData, ShadowInfo);
-		FRHIPixelShader* ShaderRHI = this->GetPixelShader();
+		FRHIPixelShader* ShaderRHI = RHICmdList.GetBoundPixelShader();
 		SetShaderValue(RHICmdList, ShaderRHI, ModulatedShadowColorParameter, ShadowInfo->GetLightSceneInfo().Proxy->GetModulatedShadowColor());
 	}
 
-	/**
-	* Serialize the parameters for this shader
-	* @param Ar - archive to serialize to
-	*/
-	virtual bool Serialize(FArchive& Ar) override
-	{
-		bool bShaderHasOutdatedParameters = TShadowProjectionPS<Quality, false, true>::Serialize(Ar);
-		Ar << ModulatedShadowColorParameter;
-		return bShaderHasOutdatedParameters;
-	}
-
 protected:
-	FShaderParameter ModulatedShadowColorParameter;
+	LAYOUT_FIELD(FShaderParameter, ModulatedShadowColorParameter);
 };
 
 /** Translucency shadow projection uniform buffer containing data needed for Fourier opacity maps. */
@@ -1176,28 +1143,18 @@ public:
 	{
 	}
 
-	virtual void SetParameters(
+	void SetParameters(
 		FRHICommandList& RHICmdList, 
 		int32 ViewIndex,
 		const FSceneView& View,
 		const FHairStrandsVisibilityData* HairVisibilityData,
-		const FProjectedShadowInfo* ShadowInfo) override
+		const FProjectedShadowInfo* ShadowInfo)
 	{
 		TShadowProjectionPS<Quality>::SetParameters(RHICmdList, ViewIndex, View, HairVisibilityData, ShadowInfo);
 
 		FTranslucentSelfShadowUniformParameters TranslucentSelfShadowUniformParameters;
 		SetupTranslucentSelfShadowUniformParameters(ShadowInfo, TranslucentSelfShadowUniformParameters);
-		SetUniformBufferParameterImmediate(RHICmdList, this->GetPixelShader(), this->template GetUniformBufferParameter<FTranslucentSelfShadowUniformParameters>(), TranslucentSelfShadowUniformParameters);
-	}
-
-	/**
-	 * Serialize the parameters for this shader
-	 * @param Ar - archive to serialize to
-	 */
-	virtual bool Serialize(FArchive& Ar) override
-	{
-		bool bShaderHasOutdatedParameters = TShadowProjectionPS<Quality>::Serialize(Ar);
-		return bShaderHasOutdatedParameters;
+		SetUniformBufferParameterImmediate(RHICmdList, RHICmdList.GetBoundPixelShader(), this->template GetUniformBufferParameter<FTranslucentSelfShadowUniformParameters>(), TranslucentSelfShadowUniformParameters);
 	}
 };
 
@@ -1205,6 +1162,7 @@ public:
 /** One pass point light shadow projection parameters used by multiple shaders. */
 class FOnePassPointShadowProjectionShaderParameters
 {
+	DECLARE_TYPE_LAYOUT(FOnePassPointShadowProjectionShaderParameters, NonVirtual);
 public:
 
 	void Bind(const FShaderParameterMap& ParameterMap)
@@ -1288,7 +1246,7 @@ public:
 	}
 
 	/** Serializer. */ 
-	friend FArchive& operator<<(FArchive& Ar,FOnePassPointShadowProjectionShaderParameters& P)
+	/*friend FArchive& operator<<(FArchive& Ar,FOnePassPointShadowProjectionShaderParameters& P)
 	{
 		Ar << P.ShadowDepthTexture;
 		Ar << P.ShadowDepthTexture2;
@@ -1297,15 +1255,15 @@ public:
 		Ar << P.InvShadowmapResolution;
 		Ar << P.LightPositionOrDirection;
 		return Ar;
-	}
+	}*/
 
 private:
-	FShaderResourceParameter ShadowDepthTexture;
-	FShaderResourceParameter ShadowDepthTexture2;
-	FShaderResourceParameter ShadowDepthCubeComparisonSampler;
-	FShaderParameter ShadowViewProjectionMatrices;
-	FShaderParameter InvShadowmapResolution;
-	FShaderParameter LightPositionOrDirection;
+	LAYOUT_FIELD(FShaderResourceParameter, ShadowDepthTexture);
+	LAYOUT_FIELD(FShaderResourceParameter, ShadowDepthTexture2);
+	LAYOUT_FIELD(FShaderResourceParameter, ShadowDepthCubeComparisonSampler);
+	LAYOUT_FIELD(FShaderParameter, ShadowViewProjectionMatrices);
+	LAYOUT_FIELD(FShaderParameter, InvShadowmapResolution);
+	LAYOUT_FIELD(FShaderParameter, LightPositionOrDirection);
 };
 
 /**
@@ -1355,7 +1313,7 @@ public:
 		const FHairStrandsVisibilityData* HairVisibilityData,
 		const FProjectedShadowInfo* ShadowInfo)
 	{
-		FRHIPixelShader* ShaderRHI = GetPixelShader();
+		FRHIPixelShader* ShaderRHI = RHICmdList.GetBoundPixelShader();
 
 		FGlobalShader::SetParameters<FViewUniformShaderParameters>(RHICmdList, ShaderRHI, View.ViewUniformBuffer);
 
@@ -1413,35 +1371,17 @@ public:
 		}
 	}
 
-	virtual bool Serialize(FArchive& Ar) override
-	{
-		bool bShaderHasOutdatedParameters = FGlobalShader::Serialize(Ar);
-		Ar << SceneTextureParameters;
-		Ar << OnePassShadowParameters;
-		Ar << ShadowDepthTextureSampler;
-		Ar << LightPosition;
-		Ar << ShadowFadeFraction;
-		Ar << ShadowSharpen;
-		Ar << PointLightDepthBias;
-		Ar << PointLightProjParameters;
-		Ar << TransmissionProfilesTexture;
-		Ar << HairCategorizationTexture;
-		return bShaderHasOutdatedParameters;
-	}
-
 private:
-	FSceneTextureShaderParameters SceneTextureParameters;
-	FOnePassPointShadowProjectionShaderParameters OnePassShadowParameters;
-	FShaderResourceParameter ShadowDepthTextureSampler;
-	FShaderParameter LightPosition;
-	FShaderParameter ShadowFadeFraction;
-	FShaderParameter ShadowSharpen;
-	FShaderParameter PointLightDepthBias;
-	FShaderParameter PointLightProjParameters;
-	FShaderResourceParameter TransmissionProfilesTexture;
-	FShaderResourceParameter HairCategorizationTexture;
-
-
+	LAYOUT_FIELD(FSceneTextureShaderParameters, SceneTextureParameters);
+	LAYOUT_FIELD(FOnePassPointShadowProjectionShaderParameters, OnePassShadowParameters);
+	LAYOUT_FIELD(FShaderResourceParameter, ShadowDepthTextureSampler);
+	LAYOUT_FIELD(FShaderParameter, LightPosition);
+	LAYOUT_FIELD(FShaderParameter, ShadowFadeFraction);
+	LAYOUT_FIELD(FShaderParameter, ShadowSharpen);
+	LAYOUT_FIELD(FShaderParameter, PointLightDepthBias);
+	LAYOUT_FIELD(FShaderParameter, PointLightProjParameters);
+	LAYOUT_FIELD(FShaderResourceParameter, TransmissionProfilesTexture);
+	LAYOUT_FIELD(FShaderResourceParameter, HairCategorizationTexture);
 };
 
 /** A transform the remaps depth and potentially projects onto some plane. */
@@ -1484,16 +1424,16 @@ public:
 			&& (Parameters.Platform == SP_PCD3D_SM5 || IsVulkanSM5Platform(Parameters.Platform) || Parameters.Platform == SP_METAL_SM5 || Parameters.Platform == SP_METAL_SM5_NOTESS);
 	}
 
-	virtual void SetParameters(
+	void SetParameters(
 		FRHICommandList& RHICmdList,
 		int32 ViewIndex,
 		const FSceneView& View,
 		const FHairStrandsVisibilityData* HairVisibilityData,
-		const FProjectedShadowInfo* ShadowInfo) override
+		const FProjectedShadowInfo* ShadowInfo)
 	{
 		TShadowProjectionPS<Quality, bUseFadePlane>::SetParameters(RHICmdList, ViewIndex, View, HairVisibilityData, ShadowInfo);
 
-		FRHIPixelShader* ShaderRHI = this->GetPixelShader();
+		FRHIPixelShader* ShaderRHI = RHICmdList.GetBoundPixelShader();
 
 		// GetLightSourceAngle returns the full angle.
 		float TanLightSourceAngle = FMath::Tan(0.5 * FMath::DegreesToRadians(ShadowInfo->GetLightSceneInfo().Proxy->GetLightSourceAngle()));
@@ -1509,19 +1449,8 @@ public:
 		SetShaderValue(RHICmdList, ShaderRHI, PCSSParameters, PCSSParameterValues);
 	}
 
-	/**
-	* Serialize the parameters for this shader
-	* @param Ar - archive to serialize to
-	*/
-	virtual bool Serialize(FArchive& Ar) override
-	{
-		bool bShaderHasOutdatedParameters = TShadowProjectionPS<Quality, bUseFadePlane>::Serialize(Ar);
-		Ar << PCSSParameters;
-		return bShaderHasOutdatedParameters;
-	}
-
 protected:
-	FShaderParameter PCSSParameters;
+	LAYOUT_FIELD(FShaderParameter, PCSSParameters);
 };
 
 
@@ -1552,18 +1481,18 @@ public:
 		OutEnvironment.SetDefine(TEXT("SPOT_LIGHT_PCSS"), 1);
 	}
 
-	virtual void SetParameters(
+	void SetParameters(
 		FRHICommandList& RHICmdList,
 		int32 ViewIndex,
 		const FSceneView& View,
 		const FHairStrandsVisibilityData* HairVisibilityData,
-		const FProjectedShadowInfo* ShadowInfo) override
+		const FProjectedShadowInfo* ShadowInfo)
 	{
 		check(ShadowInfo->GetLightSceneInfo().Proxy->GetLightType() == LightType_Spot);
 
 		TShadowProjectionPS<Quality, bUseFadePlane>::SetParameters(RHICmdList, ViewIndex, View, HairVisibilityData, ShadowInfo);
 
-		FRHIPixelShader* ShaderRHI = this->GetPixelShader();
+		FRHIPixelShader* ShaderRHI = RHICmdList.GetBoundPixelShader();
 
 		static IConsoleVariable* CVarMaxSoftShadowKernelSize = IConsoleManager::Get().FindConsoleVariable(TEXT("r.Shadow.MaxSoftKernelSize"));
 		check(CVarMaxSoftShadowKernelSize);
@@ -1573,19 +1502,8 @@ public:
 		SetShaderValue(RHICmdList, ShaderRHI, PCSSParameters, PCSSParameterValues);
 	}
 
-	/**
-	* Serialize the parameters for this shader
-	* @param Ar - archive to serialize to
-	*/
-	virtual bool Serialize(FArchive& Ar) override
-	{
-		bool bShaderHasOutdatedParameters = TShadowProjectionPS<Quality, bUseFadePlane>::Serialize(Ar);
-		Ar << PCSSParameters;
-		return bShaderHasOutdatedParameters;
-	}
-
 protected:
-	FShaderParameter PCSSParameters;
+	LAYOUT_FIELD(FShaderParameter, PCSSParameters);
 };
 
 
diff --git a/Engine/Source/Runtime/Renderer/Private/ShadowSetup.cpp b/Engine/Source/Runtime/Renderer/Private/ShadowSetup.cpp
--- a/Engine/Source/Runtime/Renderer/Private/ShadowSetup.cpp
+++ b/Engine/Source/Runtime/Renderer/Private/ShadowSetup.cpp
@@ -842,7 +842,7 @@ void FProjectedShadowInfo::AddCachedMeshDrawCommandsForPass(
 			const FCachedMeshDrawCommandInfo& CachedMeshDrawCommand = InPrimitiveSceneInfo->StaticMeshCommandInfos[StaticMeshCommandInfoIndex];
 			const FCachedPassMeshDrawList& SceneDrawList = Scene->CachedDrawLists[PassType];
 			const FMeshDrawCommand* MeshDrawCommand = CachedMeshDrawCommand.StateBucketId >= 0
-					? &Scene->CachedMeshDrawCommandStateBuckets[FSetElementId::FromInteger(CachedMeshDrawCommand.StateBucketId)].MeshDrawCommand
+					? &Scene->CachedMeshDrawCommandStateBuckets[PassType].GetByElementId(CachedMeshDrawCommand.StateBucketId).Key
 					: &SceneDrawList.MeshDrawCommands[CachedMeshDrawCommand.CommandIndex];
 
 			FVisibleMeshDrawCommand NewVisibleMeshDrawCommand;
@@ -1259,7 +1259,7 @@ void FProjectedShadowInfo::SetupMeshDrawCommandsForProjectionStenciling(FSceneRe
 			ProjectionStencilingPasses.Add(FShadowMeshDrawCommandPass());
 			FShadowMeshDrawCommandPass& ProjectionStencilingPass = ProjectionStencilingPasses[ViewIndex];
 
-			FDynamicPassMeshDrawListContext ProjectionStencilingContext(DynamicMeshDrawCommandStorage, ProjectionStencilingPass.VisibleMeshDrawCommands, GraphicsMinimalPipelineStateSet);
+			FDynamicPassMeshDrawListContext ProjectionStencilingContext(DynamicMeshDrawCommandStorage, ProjectionStencilingPass.VisibleMeshDrawCommands, GraphicsMinimalPipelineStateSet, NeedsShaderInitialisation);
 
 			FMeshPassProcessorRenderState DrawRenderState;
 			DrawRenderState.SetBlendState(TStaticBlendState<CW_NONE>::GetRHI());
@@ -1352,7 +1352,7 @@ void FProjectedShadowInfo::SetupMeshDrawCommandsForProjectionStenciling(FSceneRe
 				}
 			}
 
-			ApplyViewOverridesToMeshDrawCommands(View, ProjectionStencilingPass.VisibleMeshDrawCommands, DynamicMeshDrawCommandStorage, GraphicsMinimalPipelineStateSet);
+			ApplyViewOverridesToMeshDrawCommands(View, ProjectionStencilingPass.VisibleMeshDrawCommands, DynamicMeshDrawCommandStorage, GraphicsMinimalPipelineStateSet, NeedsShaderInitialisation);
 
 			// If instanced stereo is enabled, we need to render each view of the stereo pair using the instanced stereo transform to avoid bias issues.
 			// TODO: Support instanced stereo properly in the projection stenciling pass.
diff --git a/Engine/Source/Runtime/Renderer/Private/SingleLayerWaterRendering.cpp b/Engine/Source/Runtime/Renderer/Private/SingleLayerWaterRendering.cpp
--- a/Engine/Source/Runtime/Renderer/Private/SingleLayerWaterRendering.cpp
+++ b/Engine/Source/Runtime/Renderer/Private/SingleLayerWaterRendering.cpp
@@ -161,8 +161,9 @@ void FSingleLayerWaterPassMeshProcessor::AddMeshBatch(const FMeshBatch& RESTRICT
 
 	if (Material.GetShadingModels().HasShadingModel(MSM_SingleLayerWater))
 	{
-		const ERasterizerFillMode MeshFillMode = ComputeMeshFillMode(MeshBatch, Material);
-		const ERasterizerCullMode MeshCullMode = ComputeMeshCullMode(MeshBatch, Material);
+		const FMeshDrawingPolicyOverrideSettings OverrideSettings = ComputeMeshOverrideSettings(MeshBatch);
+		const ERasterizerFillMode MeshFillMode = ComputeMeshFillMode(MeshBatch, Material, OverrideSettings);
+		const ERasterizerCullMode MeshCullMode = ComputeMeshCullMode(MeshBatch, Material, OverrideSettings);
 		const FMaterialRenderProxy& MaterialRenderProxy = FallbackMaterialRenderProxyPtr ? *FallbackMaterialRenderProxyPtr : *MeshBatch.MaterialRenderProxy;
 		Process(MeshBatch, BatchElementMask, PrimitiveSceneProxy, StaticMeshId, MaterialRenderProxy, Material, MeshFillMode, MeshCullMode);
 	}
@@ -474,17 +475,17 @@ void FDeferredShadingSceneRenderer::CopySingleLayerWaterTextures(FRHICommandList
 	    PassParameters->SVPositionToSourceTextureUV = FVector2D(RefractionDownsampleFactor / float(SceneContext.GetBufferSizeXY().X), RefractionDownsampleFactor / float(SceneContext.GetBufferSizeXY().Y));
     
 	    PassParameters->RenderTargets[0] = FRenderTargetBinding(SceneDepthWithoutSingleLayerWaterTexture, ERenderTargetLoadAction::ENoAction);
-    
+
 	    if (bCopyColor)
 	    {
 		    PassParameters->RenderTargets[1] = FRenderTargetBinding(SceneColorWithoutSingleLayerWaterTexture, ERenderTargetLoadAction::ENoAction);
 	    }
-    
+
 	    FWaterRefractionCopyPS::FPermutationDomain PermutationVector;
 	    PermutationVector.Set<FWaterRefractionCopyPS::FDownsampleRefraction>(RefractionDownsampleFactor > 1);
 	    PermutationVector.Set<FWaterRefractionCopyPS::FDownsampleColor>(bCopyColor);
 	    auto PixelShader = View.ShaderMap->GetShader<FWaterRefractionCopyPS>(PermutationVector);
-    
+
 		const FIntRect RefractionViewRect = FIntRect(FIntPoint::DivideAndRoundDown(View.ViewRect.Min, RefractionDownsampleFactor), FIntPoint::DivideAndRoundDown(View.ViewRect.Max, RefractionDownsampleFactor));
 
 		PassData.ViewData[ViewIndex].SceneWithoutSingleLayerWaterViewRect   = RefractionViewRect;
@@ -659,7 +660,7 @@ void FDeferredShadingSceneRenderer::RenderSingleLayerWaterReflections(FRHIComman
 				PassParameters->DispatchIndirectDataUAV = TiledScreenSpaceReflection.DispatchIndirectParametersBufferUAV;
 				PassParameters->WaterTileListDataUAV = TiledScreenSpaceReflection.TileListStructureBufferUAV;
 
-				FComputeShaderUtils::AddPass(GraphBuilder, RDG_EVENT_NAME("WaterTileCategorisation"), *ComputeShader, PassParameters, TiledViewRes);
+				FComputeShaderUtils::AddPass(GraphBuilder, RDG_EVENT_NAME("WaterTileCategorisation"), ComputeShader, PassParameters, TiledViewRes);
 			}
 		}
 
@@ -726,15 +727,15 @@ void FDeferredShadingSceneRenderer::RenderSingleLayerWaterReflections(FRHIComman
 			PassParameters->IndirectDrawParameter = TiledScreenSpaceReflection.DispatchIndirectParametersBuffer;
 			PassParameters->RenderTargets[0] = FRenderTargetBinding(SceneColorTexture, ERenderTargetLoadAction::ELoad);
 
-			ValidateShaderParameters(*PixelShader, PassParameters->PS);
-			ClearUnusedGraphResources(*PixelShader, &PassParameters->PS);
+			ValidateShaderParameters(PixelShader, PassParameters->PS);
+			ClearUnusedGraphResources(PixelShader, &PassParameters->PS);
 
 			if (bRunTiled)
 			{
 				FWaterTileVS::FPermutationDomain VsPermutationVector;
 				TShaderMapRef<FWaterTileVS> VertexShader(View.ShaderMap, VsPermutationVector);
-				ValidateShaderParameters(*VertexShader, PassParameters->VS);
-				ClearUnusedGraphResources(*VertexShader, &PassParameters->VS);
+				ValidateShaderParameters(VertexShader, PassParameters->VS);
+				ClearUnusedGraphResources(VertexShader, &PassParameters->VS);
 
 				GraphBuilder.AddPass(
 					RDG_EVENT_NAME("Water Composite %dx%d", View.ViewRect.Width(), View.ViewRect.Height()),
@@ -751,12 +752,12 @@ void FDeferredShadingSceneRenderer::RenderSingleLayerWaterReflections(FRHIComman
 					GraphicsPSOInit.DepthStencilState = TStaticDepthStencilState<false, CF_Always>::GetRHI();
 					GraphicsPSOInit.RasterizerState = TStaticRasterizerState<FM_Solid, CM_None>::GetRHI();
 					GraphicsPSOInit.BoundShaderState.VertexDeclarationRHI = GEmptyVertexDeclaration.VertexDeclarationRHI;
-					GraphicsPSOInit.BoundShaderState.VertexShaderRHI = VertexShader->GetVertexShader();
-					GraphicsPSOInit.BoundShaderState.PixelShaderRHI = PixelShader->GetPixelShader();
+					GraphicsPSOInit.BoundShaderState.VertexShaderRHI = VertexShader.GetVertexShader();
+					GraphicsPSOInit.BoundShaderState.PixelShaderRHI = PixelShader.GetPixelShader();
 					SetGraphicsPipelineState(InRHICmdList, GraphicsPSOInit);
 
-					SetShaderParameters(InRHICmdList, *VertexShader, VertexShader->GetVertexShader(), PassParameters->VS);
-					SetShaderParameters(InRHICmdList, *PixelShader, PixelShader->GetPixelShader(), PassParameters->PS);
+					SetShaderParameters(InRHICmdList, VertexShader, VertexShader.GetVertexShader(), PassParameters->VS);
+					SetShaderParameters(InRHICmdList, PixelShader, PixelShader.GetPixelShader(), PassParameters->PS);
 
 					InRHICmdList.DrawPrimitiveIndirect(PassParameters->IndirectDrawParameter->GetIndirectRHICallBuffer(), 0);
 				});
@@ -772,13 +773,13 @@ void FDeferredShadingSceneRenderer::RenderSingleLayerWaterReflections(FRHIComman
 					InRHICmdList.SetViewport(View.ViewRect.Min.X, View.ViewRect.Min.Y, 0.0f, View.ViewRect.Max.X, View.ViewRect.Max.Y, 1.0f);
 
 					FGraphicsPipelineStateInitializer GraphicsPSOInit;
-					FPixelShaderUtils::InitFullscreenPipelineState(InRHICmdList, View.ShaderMap, *PixelShader, GraphicsPSOInit);
+					FPixelShaderUtils::InitFullscreenPipelineState(InRHICmdList, View.ShaderMap, PixelShader, GraphicsPSOInit);
 
 					// Premultiplied alpha where alpha is transmittance.
 					GraphicsPSOInit.BlendState = TStaticBlendState<CW_RGB, BO_Add, BF_One, BF_SourceAlpha>::GetRHI();
 
 					SetGraphicsPipelineState(InRHICmdList, GraphicsPSOInit);
-					SetShaderParameters(InRHICmdList, *PixelShader, PixelShader->GetPixelShader(), PassParameters->PS);
+					SetShaderParameters(InRHICmdList, PixelShader, PixelShader.GetPixelShader(), PassParameters->PS);
 					FPixelShaderUtils::DrawFullscreenTriangle(InRHICmdList);
 				});
 			}
diff --git a/Engine/Source/Runtime/Renderer/Private/SkyAtmosphereRendering.cpp b/Engine/Source/Runtime/Renderer/Private/SkyAtmosphereRendering.cpp
--- a/Engine/Source/Runtime/Renderer/Private/SkyAtmosphereRendering.cpp
+++ b/Engine/Source/Runtime/Renderer/Private/SkyAtmosphereRendering.cpp
@@ -1091,7 +1091,7 @@ void FSceneRenderer::RenderSkyAtmosphereLookUpTables(FRHICommandListImmediate& R
 	FRDGTextureUAVRef MultiScatteredLuminanceLutUAV = GraphBuilder.CreateUAV(FRDGTextureUAVDesc(MultiScatteredLuminanceLut, 0));
 
 	// Transmittance LUT
-	TShaderMap<FGlobalShaderType>* GlobalShaderMap = GetGlobalShaderMap(FeatureLevel);
+	FGlobalShaderMap* GlobalShaderMap = GetGlobalShaderMap(FeatureLevel);
 	{
 		TShaderMapRef<FRenderTransmittanceLutCS> ComputeShader(GlobalShaderMap);
 
@@ -1103,7 +1103,7 @@ void FSceneRenderer::RenderSkyAtmosphereLookUpTables(FRHICommandListImmediate& R
 		FIntVector TextureSize = TransmittanceLut->Desc.GetSize();
 		TextureSize.Z = 1;
 		const FIntVector NumGroups = FIntVector::DivideAndRoundUp(TextureSize, FRenderTransmittanceLutCS::GroupSize);
-		FComputeShaderUtils::AddPass(GraphBuilder, RDG_EVENT_NAME("TransmittanceLut"), *ComputeShader, PassParameters, NumGroups);
+		FComputeShaderUtils::AddPass(GraphBuilder, RDG_EVENT_NAME("TransmittanceLut"), ComputeShader, PassParameters, NumGroups);
 	}
 	
 	// Mean Illuminance LUT
@@ -1124,7 +1124,7 @@ void FSceneRenderer::RenderSkyAtmosphereLookUpTables(FRHICommandListImmediate& R
 		FIntVector TextureSize = MultiScatteredLuminanceLut->Desc.GetSize();
 		TextureSize.Z = 1;
 		const FIntVector NumGroups = FIntVector::DivideAndRoundUp(TextureSize, FRenderMultiScatteredLuminanceLutCS::GroupSize);
-		FComputeShaderUtils::AddPass(GraphBuilder, RDG_EVENT_NAME("MeanIllumLut"), *ComputeShader, PassParameters, NumGroups);
+		FComputeShaderUtils::AddPass(GraphBuilder, RDG_EVENT_NAME("MeanIllumLut"), ComputeShader, PassParameters, NumGroups);
 	}
 
 	// Distant Sky Light LUT
@@ -1164,7 +1164,7 @@ void FSceneRenderer::RenderSkyAtmosphereLookUpTables(FRHICommandListImmediate& R
 
 		FIntVector TextureSize = FIntVector(1, 1, 1);
 		const FIntVector NumGroups = FIntVector::DivideAndRoundUp(TextureSize, FRenderDistantSkyLightLutCS::GroupSize);
-		FComputeShaderUtils::AddPass(GraphBuilder, RDG_EVENT_NAME("DistantSkyLightLut"), *ComputeShader, PassParameters, NumGroups);
+		FComputeShaderUtils::AddPass(GraphBuilder, RDG_EVENT_NAME("DistantSkyLightLut"), ComputeShader, PassParameters, NumGroups);
 	}
 
 	for (int32 ViewIndex = 0; ViewIndex < Views.Num(); ViewIndex++)
@@ -1199,7 +1199,7 @@ void FSceneRenderer::RenderSkyAtmosphereLookUpTables(FRHICommandListImmediate& R
 			FIntVector TextureSize = SkyAtmosphereViewLutTexture->Desc.GetSize();
 			TextureSize.Z = 1;
 			const FIntVector NumGroups = FIntVector::DivideAndRoundUp(TextureSize, FRenderSkyViewLutCS::GroupSize);
-			FComputeShaderUtils::AddPass(GraphBuilder, RDG_EVENT_NAME("SkyViewLut"), *ComputeShader, PassParameters, NumGroups);
+			FComputeShaderUtils::AddPass(GraphBuilder, RDG_EVENT_NAME("SkyViewLut"), ComputeShader, PassParameters, NumGroups);
 		}
 
 		// Camera Atmosphere Volume
@@ -1222,7 +1222,7 @@ void FSceneRenderer::RenderSkyAtmosphereLookUpTables(FRHICommandListImmediate& R
 
 			FIntVector TextureSize = SkyAtmosphereCameraAerialPerspectiveVolume->Desc.GetSize();
 			const FIntVector NumGroups = FIntVector::DivideAndRoundUp(TextureSize, FRenderCameraAerialPerspectiveVolumeCS::GroupSize);
-			FComputeShaderUtils::AddPass(GraphBuilder, RDG_EVENT_NAME("CameraVolumeLut"), *ComputeShader, PassParameters, NumGroups);
+			FComputeShaderUtils::AddPass(GraphBuilder, RDG_EVENT_NAME("CameraVolumeLut"), ComputeShader, PassParameters, NumGroups);
 		}
 	}
 	
@@ -1321,7 +1321,7 @@ void FSceneRenderer::RenderSkyAtmosphere(FRHICommandListImmediate& RHICmdList)
 			PsPassParameters->SkyViewLutTexture = SkyAtmosphereViewLutTexture;
 			PsPassParameters->CameraAerialPerspectiveVolumeTexture = SkyAtmosphereCameraAerialPerspectiveVolume;
 			PsPassParameters->AerialPerspectiveStartDepth = AerialPerspectiveStartDepthInCm * CM_TO_KM;
-			ClearUnusedGraphResources(*PixelShader, PsPassParameters);
+			ClearUnusedGraphResources(PixelShader, PsPassParameters);
 
 			float StartDepthZ = 0.1f;
 			if (bFastAerialPerspectiveDepthTest)
@@ -1359,8 +1359,8 @@ void FSceneRenderer::RenderSkyAtmosphere(FRHICommandListImmediate& RHICmdList)
 				}
 
 				GraphicsPSOInit.BoundShaderState.VertexDeclarationRHI = GEmptyVertexDeclaration.VertexDeclarationRHI;
-				GraphicsPSOInit.BoundShaderState.VertexShaderRHI = VertexShader->GetVertexShader();
-				GraphicsPSOInit.BoundShaderState.PixelShaderRHI = PixelShader->GetPixelShader();
+				GraphicsPSOInit.BoundShaderState.VertexShaderRHI = VertexShader.GetVertexShader();
+				GraphicsPSOInit.BoundShaderState.PixelShaderRHI = PixelShader.GetPixelShader();
 				GraphicsPSOInit.PrimitiveType = PT_TriangleList;
 
 				if (!bRenderSkyPixel && GSupportsDepthBoundsTest)
@@ -1383,11 +1383,11 @@ void FSceneRenderer::RenderSkyAtmosphere(FRHICommandListImmediate& RHICmdList)
 
 				SetGraphicsPipelineState(RHICmdListLambda, GraphicsPSOInit);
 
-				SetShaderParameters(RHICmdListLambda, *PixelShader, PixelShader->GetPixelShader(), *PsPassParameters);
+				SetShaderParameters(RHICmdListLambda, PixelShader, PixelShader.GetPixelShader(), *PsPassParameters);
 
 				FRenderSkyAtmosphereVS::FParameters VsPassParameters;
 				VsPassParameters.StartDepthZ = StartDepthZ;
-				SetShaderParameters(RHICmdListLambda, *VertexShader, VertexShader->GetVertexShader(), VsPassParameters);
+				SetShaderParameters(RHICmdListLambda, VertexShader, VertexShader.GetVertexShader(), VsPassParameters);
 
 				RHICmdListLambda.DrawPrimitive(0, 1, 1);
 			});
@@ -1542,7 +1542,7 @@ void FDeferredShadingSceneRenderer::RenderDebugSkyAtmosphere(FRHICommandListImme
 				PassParameters->ViewPortWidth = float(View.ViewRect.Width());
 				PassParameters->ViewPortHeight = float(View.ViewRect.Height());
 
-				FPixelShaderUtils::AddFullscreenPass<FRenderDebugSkyAtmospherePS>(GraphBuilder, View.ShaderMap, RDG_EVENT_NAME("SkyAtmosphere"), *PixelShader, PassParameters,
+				FPixelShaderUtils::AddFullscreenPass<FRenderDebugSkyAtmospherePS>(GraphBuilder, View.ShaderMap, RDG_EVENT_NAME("SkyAtmosphere"), PixelShader, PassParameters,
 					View.ViewRect, PreMultipliedColorTransmittanceBlend, nullptr, DepthStencilStateWrite);
 			}
 		}
diff --git a/Engine/Source/Runtime/Renderer/Private/SkyPassRendering.cpp b/Engine/Source/Runtime/Renderer/Private/SkyPassRendering.cpp
--- a/Engine/Source/Runtime/Renderer/Private/SkyPassRendering.cpp
+++ b/Engine/Source/Runtime/Renderer/Private/SkyPassRendering.cpp
@@ -27,8 +27,9 @@ void FSkyPassMeshProcessor::AddMeshBatch(const FMeshBatch& RESTRICT MeshBatch, u
 
 	if (Material.IsSky())
 	{
-		const ERasterizerFillMode MeshFillMode = ComputeMeshFillMode(MeshBatch, Material);
-		const ERasterizerCullMode MeshCullMode = ComputeMeshCullMode(MeshBatch, Material);
+		const FMeshDrawingPolicyOverrideSettings OverrideSettings = ComputeMeshOverrideSettings(MeshBatch);
+		const ERasterizerFillMode MeshFillMode = ComputeMeshFillMode(MeshBatch, Material, OverrideSettings);
+		const ERasterizerCullMode MeshCullMode = ComputeMeshCullMode(MeshBatch, Material, OverrideSettings);
 		const FMaterialRenderProxy& MaterialRenderProxy = FallbackMaterialRenderProxyPtr ? *FallbackMaterialRenderProxyPtr : *MeshBatch.MaterialRenderProxy;
 		Process(MeshBatch, BatchElementMask, PrimitiveSceneProxy, StaticMeshId, MaterialRenderProxy, Material, MeshFillMode, MeshCullMode);
 	}
diff --git a/Engine/Source/Runtime/Renderer/Private/SurfelTree.cpp b/Engine/Source/Runtime/Renderer/Private/SurfelTree.cpp
--- a/Engine/Source/Runtime/Renderer/Private/SurfelTree.cpp
+++ b/Engine/Source/Runtime/Renderer/Private/SurfelTree.cpp
@@ -73,7 +73,7 @@ public:
 		const FUniformMeshBuffers& UniformMeshBuffers
 		)
 	{
-		FRHIComputeShader* ShaderRHI = GetComputeShader();
+		FRHIComputeShader* ShaderRHI = RHICmdList.GetBoundComputeShader();
 		FGlobalShader::SetParameters<FViewUniformShaderParameters>(RHICmdList, ShaderRHI, View.ViewUniformBuffer);
 
 		SetShaderValue(RHICmdList, ShaderRHI, NumTriangles, NumTrianglesValue);
@@ -84,26 +84,17 @@ public:
 
 	void UnsetParameters(FRHICommandList& RHICmdList)
 	{
-		FRHIComputeShader* ShaderRHI = GetComputeShader();
+		FRHIComputeShader* ShaderRHI = RHICmdList.GetBoundComputeShader();
 		TriangleAreas.UnsetUAV(RHICmdList, ShaderRHI);
 		// RHISetStreamOutTargets doesn't unbind existing uses like render targets do 
 		SetSRVParameter(RHICmdList, ShaderRHI, TriangleVertexData, NULL);
 	}
 
-	virtual bool Serialize(FArchive& Ar) override
-	{		
-		bool bShaderHasOutdatedParameters = FGlobalShader::Serialize(Ar);
-		Ar << NumTriangles;
-		Ar << TriangleVertexData;
-		Ar << TriangleAreas;
-		return bShaderHasOutdatedParameters;
-	}
-
 private:
 
-	FShaderParameter NumTriangles;
-	FShaderResourceParameter TriangleVertexData;
-	FRWShaderParameter TriangleAreas;
+	LAYOUT_FIELD(FShaderParameter, NumTriangles);
+	LAYOUT_FIELD(FShaderResourceParameter, TriangleVertexData);
+	LAYOUT_FIELD(FRWShaderParameter, TriangleAreas);
 };
 
 IMPLEMENT_SHADER_TYPE(,FComputeTriangleAreasCS,TEXT("/Engine/Private/SurfelTree.usf"),TEXT("ComputeTriangleAreasCS"),SF_Compute);
@@ -145,7 +136,7 @@ public:
 		const FUniformMeshBuffers& UniformMeshBuffers
 		)
 	{
-		FRHIComputeShader* ShaderRHI = GetComputeShader();
+		FRHIComputeShader* ShaderRHI = RHICmdList.GetBoundComputeShader();
 		FGlobalShader::SetParameters<FViewUniformShaderParameters>(RHICmdList, ShaderRHI, View.ViewUniformBuffer);
 
 		SetShaderValue(RHICmdList, ShaderRHI, NumTriangles, NumTrianglesValue);
@@ -156,24 +147,15 @@ public:
 
 	void UnsetParameters(FRHICommandList& RHICmdList)
 	{
-		FRHIComputeShader* ShaderRHI = GetComputeShader();
+		FRHIComputeShader* ShaderRHI = RHICmdList.GetBoundComputeShader();
 		TriangleCDFs.UnsetUAV(RHICmdList, ShaderRHI);
 	}
 
-	virtual bool Serialize(FArchive& Ar) override
-	{		
-		bool bShaderHasOutdatedParameters = FGlobalShader::Serialize(Ar);
-		Ar << NumTriangles;
-		Ar << TriangleAreas;
-		Ar << TriangleCDFs;
-		return bShaderHasOutdatedParameters;
-	}
-
 private:
 
-	FShaderParameter NumTriangles;
-	FShaderResourceParameter TriangleAreas;
-	FRWShaderParameter TriangleCDFs;
+	LAYOUT_FIELD(FShaderParameter, NumTriangles);
+	LAYOUT_FIELD(FShaderResourceParameter, TriangleAreas);
+	LAYOUT_FIELD(FRWShaderParameter, TriangleCDFs);
 };
 
 IMPLEMENT_SHADER_TYPE(,FComputeTriangleCDFsCS,TEXT("/Engine/Private/SurfelTree.usf"),TEXT("ComputeTriangleCDFsCS"),SF_Compute);
@@ -220,7 +202,7 @@ public:
 		const FUniformMeshBuffers& UniformMeshBuffers
 		)
 	{
-		FRHIComputeShader* ShaderRHI = GetComputeShader();
+		FRHIComputeShader* ShaderRHI = RHICmdList.GetBoundComputeShader();
 		FGlobalShader::SetParameters<FViewUniformShaderParameters>(RHICmdList, ShaderRHI, View.ViewUniformBuffer);
 
 		const FScene* Scene = (const FScene*)View.Family->Scene;
@@ -241,7 +223,7 @@ public:
 
 	void UnsetParameters(FRHICommandList& RHICmdList, const FSceneView& View)
 	{
-		FRHIComputeShader* ShaderRHI = GetComputeShader();
+		FRHIComputeShader* ShaderRHI = RHICmdList.GetBoundComputeShader();
 
 		const FScene* Scene = (const FScene*)View.Family->Scene;
 		SurfelBufferParameters.UnsetParameters(RHICmdList, ShaderRHI);
@@ -253,26 +235,14 @@ public:
 		RHICmdList.TransitionResources(EResourceTransitionAccess::EReadable, EResourceTransitionPipeline::EComputeToCompute, UniformMeshUAVs, UE_ARRAY_COUNT(UniformMeshUAVs));
 	}
 
-	virtual bool Serialize(FArchive& Ar) override
-	{		
-		bool bShaderHasOutdatedParameters = FGlobalShader::Serialize(Ar);
-		Ar << SurfelBufferParameters;
-		Ar << SurfelStartIndex;
-		Ar << NumSurfelsToGenerate;
-		Ar << NumTriangles;
-		Ar << TriangleVertexData;
-		Ar << TriangleCDFs;
-		return bShaderHasOutdatedParameters;
-	}
-
 private:
 
-	FSurfelBufferParameters SurfelBufferParameters;
-	FShaderParameter SurfelStartIndex;
-	FShaderParameter NumSurfelsToGenerate;
-	FShaderParameter NumTriangles;
-	FShaderResourceParameter TriangleVertexData;
-	FShaderResourceParameter TriangleCDFs;
+	LAYOUT_FIELD(FSurfelBufferParameters, SurfelBufferParameters);
+	LAYOUT_FIELD(FShaderParameter, SurfelStartIndex);
+	LAYOUT_FIELD(FShaderParameter, NumSurfelsToGenerate);
+	LAYOUT_FIELD(FShaderParameter, NumTriangles);
+	LAYOUT_FIELD(FShaderResourceParameter, TriangleVertexData);
+	LAYOUT_FIELD(FShaderResourceParameter, TriangleCDFs);
 };
 
 IMPLEMENT_SHADER_TYPE(,FSampleTrianglesCS,TEXT("/Engine/Private/SurfelTree.usf"),TEXT("SampleTrianglesCS"),SF_Compute);
@@ -310,9 +280,9 @@ void GenerateSurfelRepresentation(FRHICommandListImmediate& RHICmdList, FSceneRe
 			{
 				TShaderMapRef<FComputeTriangleAreasCS> ComputeShader(GetGlobalShaderMap(View.GetFeatureLevel()));
 
-				RHICmdList.SetComputeShader(ComputeShader->GetComputeShader());
+				RHICmdList.SetComputeShader(ComputeShader.GetComputeShader());
 				ComputeShader->SetParameters(RHICmdList, View, NumUniformTriangles, *UniformMeshBuffers);
-				DispatchComputeShader(RHICmdList, *ComputeShader, FMath::DivideAndRoundUp(NumUniformTriangles, GDistanceFieldAOTileSizeX * GDistanceFieldAOTileSizeY), 1, 1);
+				DispatchComputeShader(RHICmdList, ComputeShader.GetShader(), FMath::DivideAndRoundUp(NumUniformTriangles, GDistanceFieldAOTileSizeX * GDistanceFieldAOTileSizeY), 1, 1);
 
 				ComputeShader->UnsetParameters(RHICmdList);
 				RHICmdList.TransitionResource(EResourceTransitionAccess::EReadable, EResourceTransitionPipeline::EComputeToCompute, UniformMeshBuffers->TriangleAreas.UAV);
@@ -321,9 +291,9 @@ void GenerateSurfelRepresentation(FRHICommandListImmediate& RHICmdList, FSceneRe
 			{
 				TShaderMapRef<FComputeTriangleCDFsCS> ComputeShader(GetGlobalShaderMap(View.GetFeatureLevel()));
 
-				RHICmdList.SetComputeShader(ComputeShader->GetComputeShader());
+				RHICmdList.SetComputeShader(ComputeShader.GetComputeShader());
 				ComputeShader->SetParameters(RHICmdList, View, NumUniformTriangles, *UniformMeshBuffers);
-				DispatchComputeShader(RHICmdList, *ComputeShader, FMath::DivideAndRoundUp(NumUniformTriangles, GDistanceFieldAOTileSizeX * GDistanceFieldAOTileSizeY), 1, 1);
+				DispatchComputeShader(RHICmdList, ComputeShader.GetShader(), FMath::DivideAndRoundUp(NumUniformTriangles, GDistanceFieldAOTileSizeX * GDistanceFieldAOTileSizeY), 1, 1);
 
 				ComputeShader->UnsetParameters(RHICmdList);
 				RHICmdList.TransitionResource(EResourceTransitionAccess::EReadable, EResourceTransitionPipeline::EComputeToCompute, UniformMeshBuffers->TriangleCDFs.UAV);
@@ -332,9 +302,9 @@ void GenerateSurfelRepresentation(FRHICommandListImmediate& RHICmdList, FSceneRe
 			{
 				TShaderMapRef<FSampleTrianglesCS> ComputeShader(GetGlobalShaderMap(View.GetFeatureLevel()));
 
-				RHICmdList.SetComputeShader(ComputeShader->GetComputeShader());
+				RHICmdList.SetComputeShader(ComputeShader.GetComputeShader());
 				ComputeShader->SetParameters(RHICmdList, View, Allocation.Offset, Allocation.NumLOD0, NumUniformTriangles, *UniformMeshBuffers);
-				DispatchComputeShader(RHICmdList, *ComputeShader, FMath::DivideAndRoundUp(Allocation.NumLOD0, GDistanceFieldAOTileSizeX * GDistanceFieldAOTileSizeY), 1, 1);
+				DispatchComputeShader(RHICmdList, ComputeShader.GetShader(), FMath::DivideAndRoundUp(Allocation.NumLOD0, GDistanceFieldAOTileSizeX * GDistanceFieldAOTileSizeY), 1, 1);
 
 				ComputeShader->UnsetParameters(RHICmdList, View);
 			}
@@ -348,9 +318,9 @@ void GenerateSurfelRepresentation(FRHICommandListImmediate& RHICmdList, FSceneRe
 				{
 					TShaderMapRef<FSampleTrianglesCS> ComputeShader(GetGlobalShaderMap(View.GetFeatureLevel()));
 
-					RHICmdList.SetComputeShader(ComputeShader->GetComputeShader());
+					RHICmdList.SetComputeShader(ComputeShader.GetComputeShader());
 					ComputeShader->SetParameters(RHICmdList, View, Allocation.Offset + Allocation.NumLOD0, NumLOD1, NumUniformTriangles, *UniformMeshBuffers);
-					DispatchComputeShader(RHICmdList, *ComputeShader, FMath::DivideAndRoundUp(NumLOD1, GDistanceFieldAOTileSizeX * GDistanceFieldAOTileSizeY), 1, 1);
+					DispatchComputeShader(RHICmdList, ComputeShader.GetShader(), FMath::DivideAndRoundUp(NumLOD1, GDistanceFieldAOTileSizeX * GDistanceFieldAOTileSizeY), 1, 1);
 
 					ComputeShader->UnsetParameters(RHICmdList, View);
 				}
diff --git a/Engine/Source/Runtime/Renderer/Private/Tests/RayTracingTestbed.cpp b/Engine/Source/Runtime/Renderer/Private/Tests/RayTracingTestbed.cpp
--- a/Engine/Source/Runtime/Renderer/Private/Tests/RayTracingTestbed.cpp
+++ b/Engine/Source/Runtime/Renderer/Private/Tests/RayTracingTestbed.cpp
@@ -191,7 +191,7 @@ public:
 	}
 
 	FTestRaygenShader() {}
-	virtual ~FTestRaygenShader() {}
+	//virtual ~FTestRaygenShader() {}
 
 	/** Initialization constructor. */
 	FTestRaygenShader(const ShaderMetaType::CompiledShaderInitializerType& Initializer)
@@ -202,18 +202,18 @@ public:
 		Output.Bind(Initializer.ParameterMap, TEXT("Output"));
 	}
 
-	bool Serialize(FArchive& Ar)
+	/*bool Serialize(FArchive& Ar)
 	{
 		bool bShaderHasOutdatedParameters = FGlobalShader::Serialize(Ar);
 		Ar << TLAS;
 		Ar << Rays;
 		Ar << Output;
 		return bShaderHasOutdatedParameters;
-	}
+	}*/
 
-	FShaderResourceParameter	TLAS;   // SRV RaytracingAccelerationStructure
-	FShaderResourceParameter	Rays;   // SRV StructuredBuffer<FBasicRayData>
-	FShaderResourceParameter	Output; // UAV RWStructuredBuffer<uint>
+	LAYOUT_FIELD(FShaderResourceParameter, TLAS)   // SRV RaytracingAccelerationStructure
+	LAYOUT_FIELD(FShaderResourceParameter, Rays)   // SRV StructuredBuffer<FBasicRayData>
+	LAYOUT_FIELD(FShaderResourceParameter, Output) // UAV RWStructuredBuffer<uint>
 };
 
 IMPLEMENT_SHADER_TYPE(, FTestRaygenShader, TEXT("/Engine/Private/RayTracing/RayTracingTest.usf"), TEXT("TestMainRGS"), SF_RayGen);
diff --git a/Engine/Source/Runtime/Renderer/Private/TiledDeferredLightRendering.cpp b/Engine/Source/Runtime/Renderer/Private/TiledDeferredLightRendering.cpp
--- a/Engine/Source/Runtime/Renderer/Private/TiledDeferredLightRendering.cpp
+++ b/Engine/Source/Runtime/Renderer/Private/TiledDeferredLightRendering.cpp
@@ -133,7 +133,7 @@ public:
 		IPooledRenderTarget& InTextureValue,
 		IPooledRenderTarget& OutTextureValue)
 	{
-		FRHIComputeShader* ShaderRHI = GetComputeShader();
+		FRHIComputeShader* ShaderRHI = RHICmdList.GetBoundComputeShader();
 
 		FGlobalShader::SetParameters<FViewUniformShaderParameters>(RHICmdList, ShaderRHI, View.ViewUniformBuffer);
 		SceneTextureParameters.Set(RHICmdList, ShaderRHI, View.FeatureLevel, ESceneTextureSetupMode::All);
@@ -267,29 +267,12 @@ public:
 
 	void UnsetParameters(FRHICommandList& RHICmdList, IPooledRenderTarget& OutTextureValue)
 	{
-		OutTexture.UnsetUAV(RHICmdList, GetComputeShader());
+		OutTexture.UnsetUAV(RHICmdList, RHICmdList.GetBoundComputeShader());
 
 		FRHIUnorderedAccessView* OutUAV = OutTextureValue.GetRenderTargetItem().UAV;
 		RHICmdList.TransitionResources(EResourceTransitionAccess::EReadable, EResourceTransitionPipeline::EComputeToGfx, &OutUAV, 1);
 	}
 
-	virtual bool Serialize(FArchive& Ar) override
-	{		
-		bool bShaderHasOutdatedParameters = FGlobalShader::Serialize(Ar);
-		Ar << SceneTextureParameters;
-		Ar << OutTexture;
-		Ar << InTexture;
-		Ar << NumLights;
-		Ar << ViewDimensions;
-		Ar << LTCMatTexture;
-		Ar << LTCMatSampler;
-		Ar << LTCAmpTexture;
-		Ar << LTCAmpSampler;
-		Ar << TransmissionProfilesTexture;
-		Ar << TransmissionProfilesLinearSampler;
-		return bShaderHasOutdatedParameters;
-	}
-
 	static const TCHAR* GetSourceFilename()
 	{
 		return TEXT("/Engine/Private/TiledDeferredLightShaders.usf");
@@ -301,18 +284,17 @@ public:
 	}
 
 private:
-
-	FSceneTextureShaderParameters SceneTextureParameters;
-	FShaderResourceParameter InTexture;
-	FRWShaderParameter OutTexture;
-	FShaderParameter NumLights;
-	FShaderParameter ViewDimensions;
-	FShaderResourceParameter LTCMatTexture;
-	FShaderResourceParameter LTCMatSampler;
-	FShaderResourceParameter LTCAmpTexture;
-	FShaderResourceParameter LTCAmpSampler;
-	FShaderResourceParameter TransmissionProfilesTexture;
-	FShaderResourceParameter TransmissionProfilesLinearSampler;
+	LAYOUT_FIELD(FSceneTextureShaderParameters, SceneTextureParameters);
+	LAYOUT_FIELD(FShaderResourceParameter, InTexture);
+	LAYOUT_FIELD(FRWShaderParameter, OutTexture);
+	LAYOUT_FIELD(FShaderParameter, NumLights);
+	LAYOUT_FIELD(FShaderParameter, ViewDimensions);
+	LAYOUT_FIELD(FShaderResourceParameter, LTCMatTexture);
+	LAYOUT_FIELD(FShaderResourceParameter, LTCMatSampler);
+	LAYOUT_FIELD(FShaderResourceParameter, LTCAmpTexture);
+	LAYOUT_FIELD(FShaderResourceParameter, LTCAmpSampler);
+	LAYOUT_FIELD(FShaderResourceParameter, TransmissionProfilesTexture);
+	LAYOUT_FIELD(FShaderResourceParameter, TransmissionProfilesLinearSampler);
 };
 
 // #define avoids a lot of code duplication
@@ -350,13 +332,13 @@ static void SetShaderTemplTiledLighting(
 	IPooledRenderTarget& OutTexture)
 {
 	TShaderMapRef<FTiledDeferredLightingCS<bVisualizeLightCulling> > ComputeShader(View.ShaderMap);
-	RHICmdList.SetComputeShader(ComputeShader->GetComputeShader());
+	RHICmdList.SetComputeShader(ComputeShader.GetComputeShader());
 
 	ComputeShader->SetParameters(RHICmdList, View, ViewIndex, NumViews, SortedLights, TiledDeferredLightsStart, TiledDeferredLightsEnd, SimpleLights, StartIndex, NumThisPass, InTexture, OutTexture);
 
 	uint32 GroupSizeX = (View.ViewRect.Size().X + GDeferredLightTileSizeX - 1) / GDeferredLightTileSizeX;
 	uint32 GroupSizeY = (View.ViewRect.Size().Y + GDeferredLightTileSizeY - 1) / GDeferredLightTileSizeY;
-	DispatchComputeShader(RHICmdList, *ComputeShader, GroupSizeX, GroupSizeY, 1);
+	DispatchComputeShader(RHICmdList, ComputeShader.GetShader(), GroupSizeX, GroupSizeY, 1);
 
 	ComputeShader->UnsetParameters(RHICmdList, OutTexture);
 }
