Copying Private-LoadTimes-4.24 stream to Main. Biggest changes are in Materials/Shader memory freezing.

diff --git a/Engine/Source/Runtime/Renderer/Private/TranslucentLighting.cpp b/Engine/Source/Runtime/Renderer/Private/TranslucentLighting.cpp
--- a/Engine/Source/Runtime/Renderer/Private/TranslucentLighting.cpp
+++ b/Engine/Source/Runtime/Renderer/Private/TranslucentLighting.cpp
@@ -187,6 +187,7 @@ void FViewInfo::CalcTranslucencyLightingVolumeBounds(FBox* InOutCascadeBoundsArr
 /** Shader parameters for rendering the depth of a mesh for shadowing. */
 class FShadowDepthShaderParameters
 {
+	DECLARE_INLINE_TYPE_LAYOUT(FShadowDepthShaderParameters, NonVirtual);
 public:
 
 	void Bind(const FShaderParameterMap& ParameterMap)
@@ -213,15 +214,15 @@ public:
 	}
 
 	/** Set the vertex shader parameter values. */
-	void SetVertexShader(FRHICommandList& RHICmdList, FShader* VertexShader, const FSceneView& View, const FProjectedShadowInfo* ShadowInfo, const FMaterialRenderProxy* MaterialRenderProxy)
+	void SetVertexShader(FRHICommandList& RHICmdList, FRHIVertexShader* ShaderRHI, const FSceneView& View, const FProjectedShadowInfo* ShadowInfo, const FMaterialRenderProxy* MaterialRenderProxy)
 	{
-		Set(RHICmdList, VertexShader->GetVertexShader(), View, ShadowInfo, MaterialRenderProxy);
+		Set(RHICmdList, ShaderRHI, View, ShadowInfo, MaterialRenderProxy);
 	}
 
 	/** Set the domain shader parameter values. */
-	void SetDomainShader(FRHICommandList& RHICmdList, FShader* DomainShader, const FSceneView& View, const FProjectedShadowInfo* ShadowInfo, const FMaterialRenderProxy* MaterialRenderProxy)
+	void SetDomainShader(FRHICommandList& RHICmdList, FRHIDomainShader* ShaderRHI, const FSceneView& View, const FProjectedShadowInfo* ShadowInfo, const FMaterialRenderProxy* MaterialRenderProxy)
 	{
-		Set(RHICmdList, DomainShader->GetDomainShader(), View, ShadowInfo, MaterialRenderProxy);
+		Set(RHICmdList, ShaderRHI, View, ShadowInfo, MaterialRenderProxy);
 	}
 
 	/** Serializer. */
@@ -234,9 +235,11 @@ public:
 	}
 
 private:
-	FShaderParameter ProjectionMatrix;
-	FShaderParameter ShadowParams;
-	FShaderParameter ClampToNearPlane;
+	
+		LAYOUT_FIELD(FShaderParameter, ProjectionMatrix)
+		LAYOUT_FIELD(FShaderParameter, ShadowParams)
+		LAYOUT_FIELD(FShaderParameter, ClampToNearPlane)
+	
 };
 
 class FTranslucencyDepthShaderElementData : public FMeshMaterialShaderElementData
@@ -284,12 +287,12 @@ void SetupTranslucencyDepthPassUniformBuffer(
 */
 class FTranslucencyShadowDepthVS : public FMeshMaterialShader
 {
-	DECLARE_SHADER_TYPE(FTranslucencyShadowDepthVS, MeshMaterial);
+	DECLARE_INLINE_TYPE_LAYOUT(FTranslucencyShadowDepthVS, NonVirtual);
 public:
 
 	static bool ShouldCompilePermutation(const FMeshMaterialShaderPermutationParameters& Parameters)
 	{
-		return IsTranslucentBlendMode(Parameters.Material->GetBlendMode()) && IsFeatureLevelSupported(Parameters.Platform, ERHIFeatureLevel::SM5);
+		return IsTranslucentBlendMode(Parameters.MaterialParameters.BlendMode) && IsFeatureLevelSupported(Parameters.Platform, ERHIFeatureLevel::SM5);
 	}
 
 	FTranslucencyShadowDepthVS() {}
@@ -298,6 +301,9 @@ public:
 	{
 		PassUniformBuffer.Bind(Initializer.ParameterMap, FTranslucencyDepthPassUniformParameters::StaticStructMetadata.GetShaderVariableName());
 	}
+
+	
+	
 };
 
 enum ETranslucencyShadowDepthShaderMode
@@ -334,12 +340,12 @@ IMPLEMENT_MATERIAL_SHADER_TYPE(template<>,TTranslucencyShadowDepthVS<Translucenc
  */
 class FTranslucencyShadowDepthPS : public FMeshMaterialShader
 {
-	DECLARE_SHADER_TYPE(FTranslucencyShadowDepthPS,MeshMaterial);
+	DECLARE_INLINE_TYPE_LAYOUT(FTranslucencyShadowDepthPS, NonVirtual);
 public:
 
 	static bool ShouldCompilePermutation(const FMeshMaterialShaderPermutationParameters& Parameters)
 	{
-		return IsTranslucentBlendMode(Parameters.Material->GetBlendMode()) && IsFeatureLevelSupported(Parameters.Platform, ERHIFeatureLevel::SM5);
+		return IsTranslucentBlendMode(Parameters.MaterialParameters.BlendMode) && IsFeatureLevelSupported(Parameters.Platform, ERHIFeatureLevel::SM5);
 	}
 
 	FTranslucencyShadowDepthPS(const ShaderMetaType::CompiledShaderInitializerType& Initializer):
@@ -367,15 +373,9 @@ public:
 		ShaderBindings.Add(TranslucentShadowStartOffset, ShaderElementData.TranslucentShadowStartOffset);
 	}
 
-	virtual bool Serialize(FArchive& Ar) override
-	{
-		bool bShaderHasOutdatedParameters = FMeshMaterialShader::Serialize(Ar);
-		Ar << TranslucentShadowStartOffset;
-		return bShaderHasOutdatedParameters;
-	}
-
 private:
-	FShaderParameter TranslucentShadowStartOffset;
+	
+	LAYOUT_FIELD(FShaderParameter, TranslucentShadowStartOffset);
 };
 
 template <ETranslucencyShadowDepthShaderMode ShaderMode>
@@ -505,8 +505,9 @@ void FTranslucencyDepthPassMeshProcessor::AddMeshBatch(const FMeshBatch& RESTRIC
 
 		const EBlendMode BlendMode = Material.GetBlendMode();
 		const float MaterialTranslucentShadowStartOffset = Material.GetTranslucentShadowStartOffset();
-		const ERasterizerFillMode MeshFillMode = ComputeMeshFillMode(MeshBatch, Material);
-		const ERasterizerCullMode MeshCullMode = ComputeMeshCullMode(MeshBatch, Material);
+		const FMeshDrawingPolicyOverrideSettings OverrideSettings = ComputeMeshOverrideSettings(MeshBatch);
+		const ERasterizerFillMode MeshFillMode = ComputeMeshFillMode(MeshBatch, Material, OverrideSettings);
+		const ERasterizerCullMode MeshCullMode = ComputeMeshCullMode(MeshBatch, Material, OverrideSettings);
 		const bool bIsTranslucent = IsTranslucentBlendMode(BlendMode);
 
 		// Only render translucent meshes into the Fourier opacity maps
@@ -575,7 +576,7 @@ void FProjectedShadowInfo::RenderTranslucencyDepths(FRHICommandList& RHICmdList,
 
 
 		FMeshCommandOneFrameArray VisibleMeshDrawCommands;
-		FDynamicPassMeshDrawListContext TranslucencyDepthContext(DynamicMeshDrawCommandStorage, VisibleMeshDrawCommands, GraphicsMinimalPipelineStateSet);
+		FDynamicPassMeshDrawListContext TranslucencyDepthContext(DynamicMeshDrawCommandStorage, VisibleMeshDrawCommands, GraphicsMinimalPipelineStateSet, NeedsShaderInitialisation);
 
 		FTranslucencyDepthPassMeshProcessor TranslucencyDepthPassMeshProcessor(
 			SceneRenderer->Scene,
@@ -620,7 +621,7 @@ void FProjectedShadowInfo::RenderTranslucencyDepths(FRHICommandList& RHICmdList,
 			const bool bDynamicInstancing = IsDynamicInstancingEnabled(ShadowDepthView->FeatureLevel);
 
 			FRHIVertexBuffer* PrimitiveIdVertexBuffer = nullptr;
-			ApplyViewOverridesToMeshDrawCommands(*ShadowDepthView, VisibleMeshDrawCommands, DynamicMeshDrawCommandStorage, GraphicsMinimalPipelineStateSet);
+			ApplyViewOverridesToMeshDrawCommands(*ShadowDepthView, VisibleMeshDrawCommands, DynamicMeshDrawCommandStorage, GraphicsMinimalPipelineStateSet, NeedsShaderInitialisation);
 			SortAndMergeDynamicPassMeshDrawCommands(SceneRenderer->FeatureLevel, VisibleMeshDrawCommands, DynamicMeshDrawCommandStorage, PrimitiveIdVertexBuffer, 1);
 			SubmitMeshDrawCommands(VisibleMeshDrawCommands, GraphicsMinimalPipelineStateSet, PrimitiveIdVertexBuffer, 0, bDynamicInstancing, 1, RHICmdList);
 		}
@@ -652,7 +653,7 @@ public:
 	void SetParameters(FRHICommandList& RHICmdList, const FViewInfo& View, int32 VolumeCascadeIndex, const int32 ViewIndex)
 	{
 		FSceneRenderTargets& SceneContext = FSceneRenderTargets::Get(RHICmdList);
-		FRHIPixelShader* ShaderRHI = GetPixelShader();
+		FRHIPixelShader* ShaderRHI = RHICmdList.GetBoundPixelShader();
 
 		FGlobalShader::SetParameters<FViewUniformShaderParameters>(RHICmdList, ShaderRHI, View.ViewUniformBuffer);
 
@@ -676,23 +677,12 @@ public:
 			SceneContext.TranslucencyLightingVolumeDirectional[VolumeCascadeIndex + NumTranslucentVolumeRenderTargetSets * ViewIndex]->GetRenderTargetItem().ShaderResourceTexture);
 	}
 
-	virtual bool Serialize(FArchive& Ar) override
-	{
-		bool bShaderHasOutdatedParameters = FGlobalShader::Serialize(Ar);
-		Ar << TexelSize;
-		Ar << TranslucencyLightingVolumeAmbient;
-		Ar << TranslucencyLightingVolumeAmbientSampler;
-		Ar << TranslucencyLightingVolumeDirectional;
-		Ar << TranslucencyLightingVolumeDirectionalSampler;
-		return bShaderHasOutdatedParameters;
-	}
-
 private:
-	FShaderParameter TexelSize;
-	FShaderResourceParameter TranslucencyLightingVolumeAmbient;
-	FShaderResourceParameter TranslucencyLightingVolumeAmbientSampler;
-	FShaderResourceParameter TranslucencyLightingVolumeDirectional;
-	FShaderResourceParameter TranslucencyLightingVolumeDirectionalSampler;
+	LAYOUT_FIELD(FShaderParameter, TexelSize);
+	LAYOUT_FIELD(FShaderResourceParameter, TranslucencyLightingVolumeAmbient);
+	LAYOUT_FIELD(FShaderResourceParameter, TranslucencyLightingVolumeAmbientSampler);
+	LAYOUT_FIELD(FShaderResourceParameter, TranslucencyLightingVolumeDirectional);
+	LAYOUT_FIELD(FShaderResourceParameter, TranslucencyLightingVolumeDirectionalSampler);
 };
 
 IMPLEMENT_SHADER_TYPE(,FFilterTranslucentVolumePS,TEXT("/Engine/Private/TranslucentLightingShaders.usf"),TEXT("FilterMainPS"),SF_Pixel);
@@ -700,6 +690,7 @@ IMPLEMENT_SHADER_TYPE(,FFilterTranslucentVolumePS,TEXT("/Engine/Private/Transluc
 /** Shader parameters needed to inject direct lighting into a volume. */
 class FTranslucentInjectParameters
 {
+	DECLARE_INLINE_TYPE_LAYOUT(FTranslucentInjectParameters, NonVirtual);
 public:
 
 	void Bind(const FShaderParameterMap& ParameterMap)
@@ -735,19 +726,20 @@ public:
 	}
 
 	/** Serializer. */ 
-	friend FArchive& operator<<(FArchive& Ar,FTranslucentInjectParameters& P)
+	/*friend FArchive& operator<<(FArchive& Ar,FTranslucentInjectParameters& P)
 	{
 		Ar << P.WorldToShadowMatrix;
 		Ar << P.ShadowmapMinMax;
 		Ar << P.VolumeCascadeIndex;
 		return Ar;
-	}
+	}*/
 
 private:
-
-	FShaderParameter WorldToShadowMatrix;
-	FShaderParameter ShadowmapMinMax;
-	FShaderParameter VolumeCascadeIndex;
+	
+		LAYOUT_FIELD(FShaderParameter, WorldToShadowMatrix)
+		LAYOUT_FIELD(FShaderParameter, ShadowmapMinMax)
+		LAYOUT_FIELD(FShaderParameter, VolumeCascadeIndex)
+	
 };
 
 /** Pixel shader used to accumulate per-object translucent shadows into a volume texture. */
@@ -776,23 +768,16 @@ public:
 
 	void SetParameters(FRHICommandList& RHICmdList, const FViewInfo& View, const FLightSceneInfo* LightSceneInfo, const FProjectedShadowInfo* ShadowMap, uint32 VolumeCascadeIndex)
 	{
-		FGlobalShader::SetParameters<FViewUniformShaderParameters>(RHICmdList, GetPixelShader(), View.ViewUniformBuffer);
-		TranslucentInjectParameters.Set(RHICmdList, GetPixelShader(), this, View, LightSceneInfo, ShadowMap, VolumeCascadeIndex, true);
+		FGlobalShader::SetParameters<FViewUniformShaderParameters>(RHICmdList, RHICmdList.GetBoundPixelShader(), View.ViewUniformBuffer);
+		TranslucentInjectParameters.Set(RHICmdList, RHICmdList.GetBoundPixelShader(), this, View, LightSceneInfo, ShadowMap, VolumeCascadeIndex, true);
 
 		FTranslucentSelfShadowUniformParameters TranslucentSelfShadowUniformParameters;
 		SetupTranslucentSelfShadowUniformParameters(ShadowMap, TranslucentSelfShadowUniformParameters);
-		SetUniformBufferParameterImmediate(RHICmdList, GetPixelShader(), GetUniformBufferParameter<FTranslucentSelfShadowUniformParameters>(), TranslucentSelfShadowUniformParameters);
-	}
-
-	virtual bool Serialize(FArchive& Ar) override
-	{
-		bool bShaderHasOutdatedParameters = FGlobalShader::Serialize(Ar);
-		Ar << TranslucentInjectParameters;
-		return bShaderHasOutdatedParameters;
+		SetUniformBufferParameterImmediate(RHICmdList, RHICmdList.GetBoundPixelShader(), GetUniformBufferParameter<FTranslucentSelfShadowUniformParameters>(), TranslucentSelfShadowUniformParameters);
 	}
 
 private:
-	FTranslucentInjectParameters TranslucentInjectParameters;
+	LAYOUT_FIELD(FTranslucentInjectParameters, TranslucentInjectParameters);
 };
 
 IMPLEMENT_SHADER_TYPE(,FTranslucentObjectShadowingPS,TEXT("/Engine/Private/TranslucentLightingShaders.usf"),TEXT("PerObjectShadowingMainPS"),SF_Pixel);
@@ -821,7 +806,7 @@ public:
 	  */
 	static bool ShouldCompilePermutation(const FMaterialShaderPermutationParameters& Parameters)
 	{
-		return (Parameters.Material->IsLightFunction() || Parameters.Material->IsSpecialEngineMaterial()) &&
+		return (Parameters.MaterialParameters.MaterialDomain == MD_LightFunction || Parameters.MaterialParameters.bIsSpecialEngineMaterial) &&
 			(IsFeatureLevelSupported(Parameters.Platform, ERHIFeatureLevel::SM5) &&
 			(RHISupportsGeometryShaders(Parameters.Platform) || RHISupportsVertexShaderLayer(Parameters.Platform)));
 	}
@@ -850,7 +835,7 @@ public:
 	{
 		check(ShadowMap || !bDynamicallyShadowed);
 		
-		FRHIPixelShader* ShaderRHI = GetPixelShader();
+		FRHIPixelShader* ShaderRHI = RHICmdList.GetBoundPixelShader();
 
 		FMaterialShader::SetParameters(RHICmdList, ShaderRHI, MaterialProxy, *MaterialProxy->GetMaterial(View.GetFeatureLevel()), View, View.ViewUniformBuffer, ESceneTextureSetupMode::All);
 		
@@ -874,23 +859,12 @@ public:
 		}
 	}
 
-	virtual bool Serialize(FArchive& Ar) override
-	{
-		bool bShaderHasOutdatedParameters = FMaterialShader::Serialize(Ar);
-		Ar << VolumeShadowingParameters;
-		Ar << SpotlightMask;
-		Ar << LightFunctionParameters;
-		Ar << TranslucentInjectParameters;
-		Ar << LightFunctionWorldToLight;
-		return bShaderHasOutdatedParameters;
-	}
-
 private:
-	FVolumeShadowingParameters VolumeShadowingParameters;
-	FShaderParameter SpotlightMask;
-	FLightFunctionSharedParameters LightFunctionParameters;
-	FTranslucentInjectParameters TranslucentInjectParameters;
-	FShaderParameter LightFunctionWorldToLight;
+	LAYOUT_FIELD(FVolumeShadowingParameters, VolumeShadowingParameters);
+	LAYOUT_FIELD(FShaderParameter, SpotlightMask);
+	LAYOUT_FIELD(FLightFunctionSharedParameters, LightFunctionParameters);
+	LAYOUT_FIELD(FTranslucentInjectParameters, TranslucentInjectParameters);
+	LAYOUT_FIELD(FShaderParameter, LightFunctionWorldToLight);
 };
 
 #define IMPLEMENT_INJECTION_PIXELSHADER_TYPE(LightType,bDynamicallyShadowed,bApplyLightFunction,bInverseSquared) \
@@ -964,7 +938,7 @@ public:
 	)
 	{
 		check(NumUAVs == 4);
-		FRHIComputeShader* ShaderRHI = GetComputeShader();
+		FRHIComputeShader* ShaderRHI = RHICmdList.GetBoundComputeShader();
 		Ambient0.SetTexture(RHICmdList, ShaderRHI, NULL, VolumeUAVs[0]);
 		Directional0.SetTexture(RHICmdList, ShaderRHI, NULL, VolumeUAVs[1]);
 		Ambient1.SetTexture(RHICmdList, ShaderRHI, NULL, VolumeUAVs[2]);
@@ -973,28 +947,18 @@ public:
 
 	void UnsetParameters(FRHIAsyncComputeCommandListImmediate& RHICmdList)
 	{
-		FRHIComputeShader* ShaderRHI = GetComputeShader();
+		FRHIComputeShader* ShaderRHI = RHICmdList.GetBoundComputeShader();
 		Ambient0.UnsetUAV(RHICmdList, ShaderRHI);
 		Directional0.UnsetUAV(RHICmdList, ShaderRHI);
 		Ambient1.UnsetUAV(RHICmdList, ShaderRHI);
 		Directional1.UnsetUAV(RHICmdList, ShaderRHI);
 	}
 
-	virtual bool Serialize(FArchive& Ar) override
-	{
-		bool bShaderHasOutdatedParameters = FGlobalShader::Serialize(Ar);
-		Ar << Ambient0;
-		Ar << Directional0;
-		Ar << Ambient1;
-		Ar << Directional1;
-		return bShaderHasOutdatedParameters;
-	}
-
 private:
-	FRWShaderParameter Ambient0;
-	FRWShaderParameter Directional0;
-	FRWShaderParameter Ambient1;
-	FRWShaderParameter Directional1;
+	LAYOUT_FIELD(FRWShaderParameter, Ambient0);
+	LAYOUT_FIELD(FRWShaderParameter, Directional0);
+	LAYOUT_FIELD(FRWShaderParameter, Ambient1);
+	LAYOUT_FIELD(FRWShaderParameter, Directional1);
 };
 
 IMPLEMENT_SHADER_TYPE(, FClearTranslucentLightingVolumeCS, TEXT("/Engine/Private/TranslucentLightInjectionShaders.usf"), TEXT("ClearTranslucentLightingVolumeCS"), SF_Compute)
@@ -1005,7 +969,7 @@ void FDeferredShadingSceneRenderer::ClearTranslucentVolumeLightingAsyncCompute(F
 
 	const int32 NumUAVs = 4;
 
-	FClearTranslucentLightingVolumeCS* ComputeShader = *TShaderMapRef<FClearTranslucentLightingVolumeCS>(GetGlobalShaderMap(FeatureLevel));
+	TShaderMapRef<FClearTranslucentLightingVolumeCS> ComputeShader(GetGlobalShaderMap(FeatureLevel));
 	static const FName EndComputeFenceName(TEXT("TranslucencyLightingVolumeClearEndComputeFence"));
 	TranslucencyLightingVolumeClearEndFence = RHICmdList.CreateComputeFence(EndComputeFenceName);
 
@@ -1037,12 +1001,12 @@ void FDeferredShadingSceneRenderer::ClearTranslucentVolumeLightingAsyncCompute(F
 			RHICmdListComputeImmediate.WaitComputeFence(ClearBeginFence);
 
 			//standard compute setup, but on the async commandlist.
-			RHICmdListComputeImmediate.SetComputeShader(ComputeShader->GetComputeShader());
+			RHICmdListComputeImmediate.SetComputeShader(ComputeShader.GetComputeShader());
 
 			ComputeShader->SetParameters(RHICmdListComputeImmediate, VolumeUAVs.GetData() + i * NumUAVs, NumUAVs);
-
+		
 			int32 GroupsPerDim = TranslucencyLightingVolumeDim / FClearTranslucentLightingVolumeCS::CLEAR_BLOCK_SIZE;
-			DispatchComputeShader(RHICmdListComputeImmediate, ComputeShader, GroupsPerDim, GroupsPerDim, GroupsPerDim);
+			DispatchComputeShader(RHICmdListComputeImmediate, ComputeShader.GetShader(), GroupsPerDim, GroupsPerDim, GroupsPerDim);
 
 			ComputeShader->UnsetParameters(RHICmdListComputeImmediate);
 		}
@@ -1070,7 +1034,7 @@ class FInjectAmbientCubemapPS : public FGlobalShader
 	FInjectAmbientCubemapPS() {}
 
 public:
-	FCubemapShaderParameters CubemapShaderParameters;
+	LAYOUT_FIELD(FCubemapShaderParameters, CubemapShaderParameters)
 
 	/** Initialization constructor. */
 	FInjectAmbientCubemapPS(const ShaderMetaType::CompiledShaderInitializerType& Initializer)
@@ -1079,17 +1043,9 @@ public:
 		CubemapShaderParameters.Bind(Initializer.ParameterMap);
 	}
 
-	// FShader interface.
-	virtual bool Serialize(FArchive& Ar) override
-	{
-		bool bShaderHasOutdatedParameters = FGlobalShader::Serialize(Ar);
-		Ar << CubemapShaderParameters;
-		return bShaderHasOutdatedParameters;
-	}
-
 	void SetParameters(FRHICommandList& RHICmdList, const FViewInfo& View, const FFinalPostProcessSettings::FCubemapEntry& CubemapEntry)
 	{
-		FRHIPixelShader* ShaderRHI = GetPixelShader();
+		FRHIPixelShader* ShaderRHI = RHICmdList.GetBoundPixelShader();
 
 		FGlobalShader::SetParameters<FViewUniformShaderParameters>(RHICmdList, ShaderRHI, View.ViewUniformBuffer);
 
@@ -1136,11 +1092,11 @@ void FDeferredShadingSceneRenderer::InjectAmbientCubemapTranslucentVolumeLightin
 				TShaderMapRef<FInjectAmbientCubemapPS> PixelShader(ShaderMap);
 
 				GraphicsPSOInit.BoundShaderState.VertexDeclarationRHI = GScreenVertexDeclaration.VertexDeclarationRHI;
-				GraphicsPSOInit.BoundShaderState.VertexShaderRHI = GETSAFERHISHADER_VERTEX(*VertexShader);
+				GraphicsPSOInit.BoundShaderState.VertexShaderRHI = VertexShader.GetVertexShader();
 #if PLATFORM_SUPPORTS_GEOMETRY_SHADERS
-				GraphicsPSOInit.BoundShaderState.GeometryShaderRHI = GETSAFERHISHADER_GEOMETRY(*GeometryShader);
+				GraphicsPSOInit.BoundShaderState.GeometryShaderRHI = GeometryShader.GetGeometryShader();
 #endif
-				GraphicsPSOInit.BoundShaderState.PixelShaderRHI = GETSAFERHISHADER_PIXEL(*PixelShader);
+				GraphicsPSOInit.BoundShaderState.PixelShaderRHI = PixelShader.GetPixelShader();
 				GraphicsPSOInit.PrimitiveType = PT_TriangleStrip;
 									
 				SetGraphicsPipelineState(RHICmdList, GraphicsPSOInit);
@@ -1270,11 +1226,11 @@ void FDeferredShadingSceneRenderer::AccumulateTranslucentVolumeObjectShadowing(F
 					TShaderMapRef<FTranslucentObjectShadowingPS> PixelShader(ShaderMap);
 
 					GraphicsPSOInit.BoundShaderState.VertexDeclarationRHI = GScreenVertexDeclaration.VertexDeclarationRHI;
-					GraphicsPSOInit.BoundShaderState.VertexShaderRHI = GETSAFERHISHADER_VERTEX(*VertexShader);
+					GraphicsPSOInit.BoundShaderState.VertexShaderRHI = VertexShader.GetVertexShader();
 #if PLATFORM_SUPPORTS_GEOMETRY_SHADERS
-					GraphicsPSOInit.BoundShaderState.GeometryShaderRHI = GETSAFERHISHADER_GEOMETRY(*GeometryShader);
+					GraphicsPSOInit.BoundShaderState.GeometryShaderRHI = GeometryShader.GetGeometryShader();
 #endif
-					GraphicsPSOInit.BoundShaderState.PixelShaderRHI = GETSAFERHISHADER_PIXEL(*PixelShader);
+					GraphicsPSOInit.BoundShaderState.PixelShaderRHI = PixelShader.GetPixelShader();
 					GraphicsPSOInit.PrimitiveType = PT_TriangleStrip;
 
 					SetGraphicsPipelineState(RHICmdList, GraphicsPSOInit);
@@ -1314,15 +1270,15 @@ void SetInjectionShader(
 	const FProjectedShadowInfo* ShadowMap, 
 	int32 InnerSplitIndex, 
 	int32 VolumeCascadeIndexValue,
-	FWriteToSliceVS* VertexShader,
-	FWriteToSliceGS* GeometryShader,
+	const TShaderRef<FWriteToSliceVS>& VertexShader,
+	const TShaderRef<FWriteToSliceGS>& GeometryShader,
 	bool bApplyLightFunction,
 	bool bInverseSquared)
 {
 	check(ShadowMap || !bDynamicallyShadowed);
 
 	const FMaterialShaderMap* MaterialShaderMap = MaterialProxy->GetMaterial(View.GetFeatureLevel())->GetRenderingThreadShaderMap();
-	FMaterialShader* PixelShader = NULL;
+	TShaderRef<FMaterialShader> PixelShader;
 
 	const bool Directional = InjectionType == LightType_Directional;
 
@@ -1331,15 +1287,11 @@ void SetInjectionShader(
 		if( bInverseSquared )
 		{
 			auto InjectionPixelShader = MaterialShaderMap->GetShader< TTranslucentLightingInjectPS<InjectionType, bDynamicallyShadowed, true, true && !Directional> >();
-
-			check(InjectionPixelShader);
 			PixelShader = InjectionPixelShader;
 		}
 		else
 		{
 			auto InjectionPixelShader = MaterialShaderMap->GetShader< TTranslucentLightingInjectPS<InjectionType, bDynamicallyShadowed, true, false> >();
-
-			check(InjectionPixelShader);
 			PixelShader = InjectionPixelShader;
 		}
 	}
@@ -1348,25 +1300,21 @@ void SetInjectionShader(
 		if( bInverseSquared )
 		{
 			auto InjectionPixelShader = MaterialShaderMap->GetShader< TTranslucentLightingInjectPS<InjectionType, bDynamicallyShadowed, false, true && !Directional> >();
-
-			check(InjectionPixelShader);
 			PixelShader = InjectionPixelShader;
 		}
 		else
 		{
 			auto InjectionPixelShader = MaterialShaderMap->GetShader< TTranslucentLightingInjectPS<InjectionType, bDynamicallyShadowed, false, false> >();
-
-			check(InjectionPixelShader);
 			PixelShader = InjectionPixelShader;
 		}
 	}
 	
 	GraphicsPSOInit.BoundShaderState.VertexDeclarationRHI = GScreenVertexDeclaration.VertexDeclarationRHI;
-	GraphicsPSOInit.BoundShaderState.VertexShaderRHI = GETSAFERHISHADER_VERTEX(VertexShader);
+	GraphicsPSOInit.BoundShaderState.VertexShaderRHI = VertexShader.GetVertexShader();
 #if PLATFORM_SUPPORTS_GEOMETRY_SHADERS
-	GraphicsPSOInit.BoundShaderState.GeometryShaderRHI = GETSAFERHISHADER_GEOMETRY(GeometryShader);
+	GraphicsPSOInit.BoundShaderState.GeometryShaderRHI = GeometryShader.GetGeometryShader();
 #endif
-	GraphicsPSOInit.BoundShaderState.PixelShaderRHI = GETSAFERHISHADER_PIXEL(PixelShader);
+	GraphicsPSOInit.BoundShaderState.PixelShaderRHI = PixelShader.GetPixelShader();
 	SetGraphicsPipelineState(RHICmdList, GraphicsPSOInit);
 
 	// Now shader is set, bind parameters
@@ -1375,13 +1323,11 @@ void SetInjectionShader(
 		if( bInverseSquared )
 		{
 			auto InjectionPixelShader = MaterialShaderMap->GetShader< TTranslucentLightingInjectPS<InjectionType, bDynamicallyShadowed, true, true && !Directional> >();
-			check(InjectionPixelShader);
 			InjectionPixelShader->SetParameters(RHICmdList, View, LightSceneInfo, MaterialProxy, ShadowMap, InnerSplitIndex, VolumeCascadeIndexValue);
 		}
 		else
 		{
 			auto InjectionPixelShader = MaterialShaderMap->GetShader< TTranslucentLightingInjectPS<InjectionType, bDynamicallyShadowed, true, false> >();
-			check(InjectionPixelShader);
 			InjectionPixelShader->SetParameters(RHICmdList, View, LightSceneInfo, MaterialProxy, ShadowMap, InnerSplitIndex, VolumeCascadeIndexValue);
 		}
 	}
@@ -1390,13 +1336,11 @@ void SetInjectionShader(
 		if( bInverseSquared )
 		{
 			auto InjectionPixelShader = MaterialShaderMap->GetShader< TTranslucentLightingInjectPS<InjectionType, bDynamicallyShadowed, false, true && !Directional> >();
-			check(InjectionPixelShader);
 			InjectionPixelShader->SetParameters(RHICmdList, View, LightSceneInfo, MaterialProxy, ShadowMap, InnerSplitIndex, VolumeCascadeIndexValue);
 		}
 		else
 		{
 			auto InjectionPixelShader = MaterialShaderMap->GetShader< TTranslucentLightingInjectPS<InjectionType, bDynamicallyShadowed, false, false> >();
-			check(InjectionPixelShader);
 			InjectionPixelShader->SetParameters(RHICmdList, View, LightSceneInfo, MaterialProxy, ShadowMap, InnerSplitIndex, VolumeCascadeIndexValue);
 		}
 	}
@@ -1514,14 +1458,14 @@ static void InjectTranslucentLightArray(FRHICommandListImmediate& RHICmdList, co
 							// shadows, restricting light contribution to the cascade bounds (except last cascade far to get light functions and no shadows there)
 							SetInjectionShader<LightType_Directional, true>(RHICmdList, GraphicsPSOInit, View, InjectionData.LightFunctionMaterialProxy, LightSceneInfo,
 								InjectionData.ProjectedShadowInfo, InjectionData.ProjectedShadowInfo->CascadeSettings.ShadowSplitIndex, VolumeCascadeIndex,
-								*VertexShader, *GeometryShader, InjectionData.bApplyLightFunction, false);
+								VertexShader, GeometryShader, InjectionData.bApplyLightFunction, false);
 						}
 						else
 						{
 							// no shadows
 							SetInjectionShader<LightType_Directional, false>(RHICmdList, GraphicsPSOInit, View, InjectionData.LightFunctionMaterialProxy, LightSceneInfo,
 								InjectionData.ProjectedShadowInfo, INDEX_NONE, VolumeCascadeIndex,
-								*VertexShader, *GeometryShader, InjectionData.bApplyLightFunction, false);
+								VertexShader, GeometryShader, InjectionData.bApplyLightFunction, false);
 						}
 					}
 					else
@@ -1535,13 +1479,13 @@ static void InjectTranslucentLightArray(FRHICommandListImmediate& RHICmdList, co
 						{
 							SetInjectionShader<LightType_Point, true>(RHICmdList, GraphicsPSOInit, View, InjectionData.LightFunctionMaterialProxy, LightSceneInfo,
 								InjectionData.ProjectedShadowInfo, INDEX_NONE, VolumeCascadeIndex,
-								*VertexShader, *GeometryShader, InjectionData.bApplyLightFunction, bInverseSquared);
+								VertexShader, GeometryShader, InjectionData.bApplyLightFunction, bInverseSquared);
 						}
 						else
 						{
 							SetInjectionShader<LightType_Point, false>(RHICmdList, GraphicsPSOInit, View, InjectionData.LightFunctionMaterialProxy, LightSceneInfo,
 								InjectionData.ProjectedShadowInfo, INDEX_NONE, VolumeCascadeIndex,
-								*VertexShader, *GeometryShader, InjectionData.bApplyLightFunction, bInverseSquared);
+								VertexShader, GeometryShader, InjectionData.bApplyLightFunction, bInverseSquared);
 						}
 					}
 
@@ -1633,30 +1577,21 @@ public:
 
 	void SetParameters(FRHICommandList& RHICmdList, const FViewInfo& View, const FSimpleLightEntry& SimpleLight, const FSimpleLightPerViewEntry& SimpleLightPerViewData, int32 VolumeCascadeIndexValue)
 	{
-		FGlobalShader::SetParameters<FViewUniformShaderParameters>(RHICmdList, GetPixelShader(), View.ViewUniformBuffer);
+		FGlobalShader::SetParameters<FViewUniformShaderParameters>(RHICmdList, RHICmdList.GetBoundPixelShader(), View.ViewUniformBuffer);
 
 		FVector4 PositionAndRadius(SimpleLightPerViewData.Position, SimpleLight.Radius);
-		SetShaderValue(RHICmdList, GetPixelShader(), VolumeCascadeIndex, VolumeCascadeIndexValue);
-		SetShaderValue(RHICmdList, GetPixelShader(), SimpleLightPositionAndRadius, PositionAndRadius);
+		SetShaderValue(RHICmdList, RHICmdList.GetBoundPixelShader(), VolumeCascadeIndex, VolumeCascadeIndexValue);
+		SetShaderValue(RHICmdList, RHICmdList.GetBoundPixelShader(), SimpleLightPositionAndRadius, PositionAndRadius);
 
 		FVector4 LightColorAndExponent(SimpleLight.Color, SimpleLight.Exponent);
 
-		SetShaderValue(RHICmdList, GetPixelShader(), SimpleLightColorAndExponent, LightColorAndExponent);
-	}
-
-	virtual bool Serialize(FArchive& Ar) override
-	{
-		bool bShaderHasOutdatedParameters = FGlobalShader::Serialize(Ar);
-		Ar << VolumeCascadeIndex;
-		Ar << SimpleLightPositionAndRadius;
-		Ar << SimpleLightColorAndExponent;
-		return bShaderHasOutdatedParameters;
+		SetShaderValue(RHICmdList, RHICmdList.GetBoundPixelShader(), SimpleLightColorAndExponent, LightColorAndExponent);
 	}
 
 private:
-	FShaderParameter VolumeCascadeIndex;
-	FShaderParameter SimpleLightPositionAndRadius;
-	FShaderParameter SimpleLightColorAndExponent;
+	LAYOUT_FIELD(FShaderParameter, VolumeCascadeIndex);
+	LAYOUT_FIELD(FShaderParameter, SimpleLightPositionAndRadius);
+	LAYOUT_FIELD(FShaderParameter, SimpleLightColorAndExponent);
 };
 
 IMPLEMENT_SHADER_TYPE(,FSimpleLightTranslucentLightingInjectPS,TEXT("/Engine/Private/TranslucentLightInjectionShaders.usf"),TEXT("SimpleLightInjectMainPS"),SF_Pixel);
@@ -1727,11 +1662,11 @@ void FDeferredShadingSceneRenderer::InjectSimpleTranslucentVolumeLightingArray(F
 							TShaderMapRef<FSimpleLightTranslucentLightingInjectPS> PixelShader(View.ShaderMap);
 
 							GraphicsPSOInit.BoundShaderState.VertexDeclarationRHI = GScreenVertexDeclaration.VertexDeclarationRHI;
-							GraphicsPSOInit.BoundShaderState.VertexShaderRHI = GETSAFERHISHADER_VERTEX(*VertexShader);
+							GraphicsPSOInit.BoundShaderState.VertexShaderRHI = VertexShader.GetVertexShader();
 #if PLATFORM_SUPPORTS_GEOMETRY_SHADERS
-							GraphicsPSOInit.BoundShaderState.GeometryShaderRHI = GETSAFERHISHADER_GEOMETRY(*GeometryShader);
+							GraphicsPSOInit.BoundShaderState.GeometryShaderRHI = GeometryShader.GetGeometryShader();
 #endif
-							GraphicsPSOInit.BoundShaderState.PixelShaderRHI = GETSAFERHISHADER_PIXEL(*PixelShader);
+							GraphicsPSOInit.BoundShaderState.PixelShaderRHI = PixelShader.GetPixelShader();
 
 							SetGraphicsPipelineState(RHICmdList, GraphicsPSOInit);
 
@@ -1832,11 +1767,11 @@ void FDeferredShadingSceneRenderer::FilterTranslucentVolumeLighting(FRHICommandL
 					TShaderMapRef<FFilterTranslucentVolumePS> PixelShader(View.ShaderMap);
 				
 					GraphicsPSOInit.BoundShaderState.VertexDeclarationRHI = GScreenVertexDeclaration.VertexDeclarationRHI;
-					GraphicsPSOInit.BoundShaderState.VertexShaderRHI = GETSAFERHISHADER_VERTEX(*VertexShader);
+					GraphicsPSOInit.BoundShaderState.VertexShaderRHI = VertexShader.GetVertexShader();
 #if PLATFORM_SUPPORTS_GEOMETRY_SHADERS
-					GraphicsPSOInit.BoundShaderState.GeometryShaderRHI = GETSAFERHISHADER_GEOMETRY(*GeometryShader);
+					GraphicsPSOInit.BoundShaderState.GeometryShaderRHI = GeometryShader.GetGeometryShader();
 #endif
-					GraphicsPSOInit.BoundShaderState.PixelShaderRHI = GETSAFERHISHADER_PIXEL(*PixelShader);
+					GraphicsPSOInit.BoundShaderState.PixelShaderRHI = PixelShader.GetPixelShader();
 					GraphicsPSOInit.PrimitiveType = PT_TriangleStrip;
 
 					SetGraphicsPipelineState(RHICmdList, GraphicsPSOInit);
diff --git a/Engine/Source/Runtime/Renderer/Private/TranslucentRendering.cpp b/Engine/Source/Runtime/Renderer/Private/TranslucentRendering.cpp
--- a/Engine/Source/Runtime/Renderer/Private/TranslucentRendering.cpp
+++ b/Engine/Source/Runtime/Renderer/Private/TranslucentRendering.cpp
@@ -248,18 +248,11 @@ public:
 
 	void SetParameters(FRHICommandList& RHICmdList, const FViewInfo& View)
 	{
-		SceneTextureParameters.Set(RHICmdList, GetPixelShader(), View.FeatureLevel, ESceneTextureSetupMode::All);
-	}
-
-	virtual bool Serialize(FArchive& Ar) override
-	{
-		bool bShaderHasOutdatedParameters = FGlobalShader::Serialize(Ar);
-		Ar << SceneTextureParameters;
-		return bShaderHasOutdatedParameters;
+		SceneTextureParameters.Set(RHICmdList, RHICmdList.GetBoundPixelShader(), View.FeatureLevel, ESceneTextureSetupMode::All);
 	}
 
 private:
-	FSceneTextureShaderParameters SceneTextureParameters;
+	LAYOUT_FIELD(FSceneTextureShaderParameters, SceneTextureParameters);
 };
 
 IMPLEMENT_SHADER_TYPE(, FCopySceneColorPS, TEXT("/Engine/Private/TranslucentLightingShaders.usf"), TEXT("CopySceneColorMain"), SF_Pixel);
@@ -537,8 +530,8 @@ void FDeferredShadingSceneRenderer::ConditionalResolveSceneColorForTranslucentMa
 				TShaderMapRef<FCopySceneColorPS> PixelShader(View.ShaderMap);
 
 				GraphicsPSOInit.BoundShaderState.VertexDeclarationRHI = GFilterVertexDeclaration.VertexDeclarationRHI;
-				GraphicsPSOInit.BoundShaderState.VertexShaderRHI = GETSAFERHISHADER_VERTEX(*ScreenVertexShader);
-				GraphicsPSOInit.BoundShaderState.PixelShaderRHI = GETSAFERHISHADER_PIXEL(*PixelShader);
+				GraphicsPSOInit.BoundShaderState.VertexShaderRHI = ScreenVertexShader.GetVertexShader();
+				GraphicsPSOInit.BoundShaderState.PixelShaderRHI = PixelShader.GetPixelShader();
 				GraphicsPSOInit.PrimitiveType = PT_TriangleList;
 
 				SetGraphicsPipelineState(RHICmdList, GraphicsPSOInit);
@@ -553,7 +546,7 @@ void FDeferredShadingSceneRenderer::ConditionalResolveSceneColorForTranslucentMa
 					View.ViewRect.Width(), View.ViewRect.Height(),
 					FIntPoint(View.ViewRect.Width(), View.ViewRect.Height()),
 					SceneContext.GetBufferSizeXY(),
-					*ScreenVertexShader,
+					ScreenVertexShader,
 					EDRF_UseTriangleOptimization);
 			}
 			RHICmdList.EndRenderPass();
@@ -694,6 +687,7 @@ void CreateTranslucentBasePassUniformBuffer(
 
 class FTranslucencyUpsamplingPS : public FGlobalShader
 {
+	DECLARE_INLINE_TYPE_LAYOUT(FTranslucencyUpsamplingPS, NonVirtual);
 protected:
 
 	static bool ShouldCompilePermutation(const FGlobalShaderPermutationParameters& Parameters)
@@ -707,12 +701,12 @@ protected:
 	{
 	}
 
-	FSceneTextureShaderParameters SceneTextureParameters;
-	FShaderParameter LowResColorTexelSize;
-	FShaderResourceParameter LowResDepthTexture;
-	FShaderResourceParameter LowResColorTexture;
-	FShaderResourceParameter BilinearClampedSampler;
-	FShaderResourceParameter PointClampedSampler;
+	LAYOUT_FIELD(FSceneTextureShaderParameters, SceneTextureParameters);
+	LAYOUT_FIELD(FShaderParameter, LowResColorTexelSize);
+	LAYOUT_FIELD(FShaderResourceParameter, LowResDepthTexture);
+	LAYOUT_FIELD(FShaderResourceParameter, LowResColorTexture);
+	LAYOUT_FIELD(FShaderResourceParameter, BilinearClampedSampler);
+	LAYOUT_FIELD(FShaderResourceParameter, PointClampedSampler);
 
 public:
 
@@ -729,17 +723,9 @@ public:
 		PointClampedSampler.Bind(Initializer.ParameterMap, TEXT("PointClampedSampler"));
 	}
 
-	// FShader interface.
-	virtual bool Serialize(FArchive& Ar) override
-	{
-		bool bShaderHasOutdatedParameters = FGlobalShader::Serialize(Ar);
-		Ar << SceneTextureParameters << LowResColorTexelSize << LowResDepthTexture << LowResColorTexture << BilinearClampedSampler << PointClampedSampler;
-		return bShaderHasOutdatedParameters;
-	}
-
 	void SetParameters(FRHICommandList& RHICmdList, const FViewInfo& View)
 	{
-		FRHIPixelShader* ShaderRHI = GetPixelShader();
+		FRHIPixelShader* ShaderRHI = RHICmdList.GetBoundPixelShader();
 		FGlobalShader::SetParameters<FViewUniformShaderParameters>(RHICmdList, ShaderRHI, View.ViewUniformBuffer);
 
 		FSceneRenderTargets& SceneContext = FSceneRenderTargets::Get(RHICmdList);
@@ -755,10 +741,10 @@ public:
 		SetSamplerParameter(RHICmdList, ShaderRHI, BilinearClampedSampler, TStaticSamplerState<SF_Bilinear, AM_Clamp, AM_Clamp, AM_Clamp>::GetRHI());
 		SetSamplerParameter(RHICmdList, ShaderRHI, PointClampedSampler, TStaticSamplerState<SF_Point, AM_Clamp, AM_Clamp, AM_Clamp>::GetRHI());
 
-		SceneTextureParameters.Set(RHICmdList, GetPixelShader(), View.FeatureLevel, ESceneTextureSetupMode::All);
+		SceneTextureParameters.Set(RHICmdList, RHICmdList.GetBoundPixelShader(), View.FeatureLevel, ESceneTextureSetupMode::All);
 	}
 
-	const bool bUseNearestDepthNeighborUpsample;
+	LAYOUT_FIELD(const bool, bUseNearestDepthNeighborUpsample);
 };
 
 class FTranslucencySimpleUpsamplingPS : public FTranslucencyUpsamplingPS
@@ -815,21 +801,21 @@ void UpsampleTranslucency(FRHICommandList& RHICmdList, const FViewInfo& View, bo
 	}
 
 	TShaderMapRef<FScreenVS> ScreenVertexShader(View.ShaderMap);
-	FTranslucencyUpsamplingPS* UpsamplingPixelShader = nullptr;
+	TShaderRef<FTranslucencyUpsamplingPS> UpsamplingPixelShader;
 	if (UseNearestDepthNeighborUpsampleForSeparateTranslucency(SceneContext))
 	{
 		TShaderMapRef<FTranslucencyNearestDepthNeighborUpsamplingPS> PixelShader(View.ShaderMap);
-		UpsamplingPixelShader = *PixelShader;
+		UpsamplingPixelShader = PixelShader;
 	}
 	else
 	{
 		TShaderMapRef<FTranslucencySimpleUpsamplingPS> PixelShader(View.ShaderMap);
-		UpsamplingPixelShader = *PixelShader;
+		UpsamplingPixelShader = PixelShader;
 	}
 
 	GraphicsPSOInit.BoundShaderState.VertexDeclarationRHI = GFilterVertexDeclaration.VertexDeclarationRHI;
-	GraphicsPSOInit.BoundShaderState.VertexShaderRHI = GETSAFERHISHADER_VERTEX(*ScreenVertexShader);
-	GraphicsPSOInit.BoundShaderState.PixelShaderRHI = GETSAFERHISHADER_PIXEL(UpsamplingPixelShader);
+	GraphicsPSOInit.BoundShaderState.VertexShaderRHI = ScreenVertexShader.GetVertexShader();
+	GraphicsPSOInit.BoundShaderState.PixelShaderRHI = UpsamplingPixelShader.GetPixelShader();
 	GraphicsPSOInit.PrimitiveType = PT_TriangleList;
 
 	SetGraphicsPipelineState(RHICmdList, GraphicsPSOInit);
@@ -851,7 +837,7 @@ void UpsampleTranslucency(FRHICommandList& RHICmdList, const FViewInfo& View, bo
 		View.ViewRect.Width() * OutScale, View.ViewRect.Height() * OutScale,
 		View.ViewRect.Size(),
 		FIntPoint(TextureWidth, TextureHeight),
-		*ScreenVertexShader,
+		ScreenVertexShader,
 		EDRF_UseTriangleOptimization);
 
 	SceneContext.FinishRenderingSceneColor(RHICmdList);
diff --git a/Engine/Source/Runtime/Renderer/Private/VT/RuntimeVirtualTextureRender.cpp b/Engine/Source/Runtime/Renderer/Private/VT/RuntimeVirtualTextureRender.cpp
--- a/Engine/Source/Runtime/Renderer/Private/VT/RuntimeVirtualTextureRender.cpp
+++ b/Engine/Source/Runtime/Renderer/Private/VT/RuntimeVirtualTextureRender.cpp
@@ -29,7 +29,7 @@ namespace RuntimeVirtualTexture
 		static bool ShouldCompilePermutation(const FMeshMaterialShaderPermutationParameters& Parameters)
 		{
 			return UseVirtualTexturing(GetMaxSupportedFeatureLevel(Parameters.Platform)) &&
-				(Parameters.Material->GetMaterialDomain() == MD_RuntimeVirtualTexture || Parameters.Material->HasRuntimeVirtualTextureOutput());
+				(Parameters.MaterialParameters.MaterialDomain == MD_RuntimeVirtualTexture || Parameters.MaterialParameters.bHasRuntimeVirtualTextureOutput);
 		}
 
 		static void ModifyCompilationEnvironment(const FMaterialShaderPermutationParameters& Parameters, FShaderCompilerEnvironment& OutEnvironment)
@@ -53,7 +53,7 @@ namespace RuntimeVirtualTexture
 		{
 			FMeshMaterialShader::SetParameters(
 				RHICmdList,
-				GetPixelShader(),
+				RHICmdList.GetBoundPixelShader(),
 				&MaterialProxy,
 				*MaterialProxy.GetMaterial(View.FeatureLevel),
 				View,
@@ -274,8 +274,9 @@ namespace RuntimeVirtualTexture
 
 			DrawRenderState.SetBlendState(MaterialPolicy::GetBlendState(OutputAttributeMask));
 
-			ERasterizerFillMode MeshFillMode = ComputeMeshFillMode(MeshBatch, MaterialResource);
-			ERasterizerCullMode MeshCullMode = ComputeMeshCullMode(MeshBatch, MaterialResource);
+			const FMeshDrawingPolicyOverrideSettings OverrideSettings = ComputeMeshOverrideSettings(MeshBatch);
+			ERasterizerFillMode MeshFillMode = ComputeMeshFillMode(MeshBatch, MaterialResource, OverrideSettings);
+			ERasterizerCullMode MeshCullMode = ComputeMeshCullMode(MeshBatch, MaterialResource, OverrideSettings);
 
 			FMeshMaterialShaderElementData ShaderElementData;
 			ShaderElementData.InitializeMeshMaterialData(ViewIfDynamicMeshCommand, PrimitiveSceneProxy, MeshBatch, StaticMeshId, false);
@@ -358,7 +359,9 @@ namespace RuntimeVirtualTexture
 		FDynamicMeshDrawCommandStorage MeshDrawCommandStorage;
 		FMeshCommandOneFrameArray AllocatedCommands;
 		FGraphicsMinimalPipelineStateSet GraphicsMinimalPipelineStateSet;
-		FDynamicPassMeshDrawListContext DynamicMeshPassContext(MeshDrawCommandStorage, AllocatedCommands, GraphicsMinimalPipelineStateSet);
+		bool NeedsShaderInitialisation;
+
+		FDynamicPassMeshDrawListContext DynamicMeshPassContext(MeshDrawCommandStorage, AllocatedCommands, GraphicsMinimalPipelineStateSet, NeedsShaderInitialisation);
 		FRuntimeVirtualTextureMeshProcessor MeshProcessor(Scene, View, &DynamicMeshPassContext);
 
 		// Pre-calculate view factors used for culling
@@ -430,7 +433,7 @@ namespace RuntimeVirtualTexture
 							FCachedMeshDrawCommandInfo& CachedMeshDrawCommand = PrimitiveSceneInfo->StaticMeshCommandInfos[StaticMeshCommandInfoIndex];
 
 							const FMeshDrawCommand* MeshDrawCommand = CachedMeshDrawCommand.StateBucketId >= 0
-								? &Scene->CachedMeshDrawCommandStateBuckets[FSetElementId::FromInteger(CachedMeshDrawCommand.StateBucketId)].MeshDrawCommand
+								? &Scene->CachedMeshDrawCommandStateBuckets[EMeshPass::VirtualTexture].GetByElementId(CachedMeshDrawCommand.StateBucketId).Key
 								: &SceneDrawList.MeshDrawCommands[CachedMeshDrawCommand.CommandIndex];
 
 							FVisibleMeshDrawCommand NewVisibleMeshDrawCommand;
@@ -539,13 +542,13 @@ namespace RuntimeVirtualTexture
 	template< ERuntimeVirtualTextureMaterialType MaterialType >
 	void AddCompressPass(FRDGBuilder& GraphBuilder, ERHIFeatureLevel::Type FeatureLevel, FShader_VirtualTextureCompress::FParameters* Parameters, FIntVector GroupCount)
 	{
-		TShaderMap<FGlobalShaderType>* GlobalShaderMap = GetGlobalShaderMap(FeatureLevel);
+		FGlobalShaderMap* GlobalShaderMap = GetGlobalShaderMap(FeatureLevel);
 		TShaderMapRef< FShader_VirtualTextureCompress_CS< MaterialType > > ComputeShader(GlobalShaderMap);
 
 		FComputeShaderUtils::AddPass(
 			GraphBuilder,
 			RDG_EVENT_NAME("VirtualTextureCompress"),
-			*ComputeShader, Parameters, GroupCount);
+			ComputeShader, Parameters, GroupCount);
 	}
 
 	/** Set up the BC compression pass for the given MaterialType. */
@@ -643,7 +646,7 @@ namespace RuntimeVirtualTexture
 	template< ERuntimeVirtualTextureMaterialType MaterialType >
 	void AddCopyPass(FRDGBuilder& GraphBuilder, ERHIFeatureLevel::Type FeatureLevel, FShader_VirtualTextureCopy::FParameters* Parameters, FIntPoint TextureSize)
 	{
-		TShaderMap<FGlobalShaderType>* GlobalShaderMap = GetGlobalShaderMap(FeatureLevel);
+		FGlobalShaderMap* GlobalShaderMap = GetGlobalShaderMap(FeatureLevel);
 		TShaderMapRef< FShader_VirtualTextureCopy_VS > VertexShader(GlobalShaderMap);
 		TShaderMapRef< FShader_VirtualTextureCopy_PS< MaterialType > > PixelShader(GlobalShaderMap);
 
@@ -660,12 +663,12 @@ namespace RuntimeVirtualTexture
 			GraphicsPSOInit.RasterizerState = TStaticRasterizerState<>::GetRHI();
 			GraphicsPSOInit.PrimitiveType = PT_TriangleList;
 			GraphicsPSOInit.BoundShaderState.VertexDeclarationRHI = GetVertexDeclarationFVector4();
-			GraphicsPSOInit.BoundShaderState.VertexShaderRHI = GETSAFERHISHADER_VERTEX(*VertexShader);
-			GraphicsPSOInit.BoundShaderState.PixelShaderRHI = GETSAFERHISHADER_PIXEL(*PixelShader);
+			GraphicsPSOInit.BoundShaderState.VertexShaderRHI = VertexShader.GetVertexShader();
+			GraphicsPSOInit.BoundShaderState.PixelShaderRHI = PixelShader.GetPixelShader();
 			SetGraphicsPipelineState(RHICmdList, GraphicsPSOInit);
 
-			SetShaderParameters(RHICmdList, *VertexShader, VertexShader->GetVertexShader(), *Parameters);
-			SetShaderParameters(RHICmdList, *PixelShader, PixelShader->GetPixelShader(), *Parameters);
+			SetShaderParameters(RHICmdList, VertexShader, VertexShader.GetVertexShader(), *Parameters);
+			SetShaderParameters(RHICmdList, PixelShader, PixelShader.GetPixelShader(), *Parameters);
 
 			RHICmdList.SetViewport(0.0f, 0.0f, 0.0f, TextureSize[0], TextureSize[1], 1.0f);
 			RHICmdList.DrawIndexedPrimitive(GTwoTrianglesIndexBuffer.IndexBufferRHI, 0, 0, 4, 0, 2, 1);
diff --git a/Engine/Source/Runtime/Renderer/Private/VT/VirtualTextureSpace.cpp b/Engine/Source/Runtime/Renderer/Private/VT/VirtualTextureSpace.cpp
--- a/Engine/Source/Runtime/Renderer/Private/VT/VirtualTextureSpace.cpp
+++ b/Engine/Source/Runtime/Renderer/Private/VT/VirtualTextureSpace.cpp
@@ -163,16 +163,11 @@ TGlobalResource< FSpriteIndexBuffer<8> > GQuadIndexBuffer;
 
 class FPageTableUpdateVS : public FGlobalShader
 {
-	DECLARE_SHADER_TYPE(FPageTableUpdateVS,Global);
-
+	DECLARE_INLINE_TYPE_LAYOUT(FPageTableUpdateVS, NonVirtual);
+protected:
 	FPageTableUpdateVS() {}
 	
 public:
-	FShaderParameter			PageTableSize;
-	FShaderParameter			FirstUpdate;
-	FShaderParameter			NumUpdates;
-	FShaderResourceParameter	UpdateBuffer;
-
 	FPageTableUpdateVS(const ShaderMetaType::CompiledShaderInitializerType& Initializer)
 		: FGlobalShader(Initializer)
 	{
@@ -188,21 +183,16 @@ public:
 			|| GetMaxSupportedFeatureLevel(Parameters.Platform) == ERHIFeatureLevel::ES3_1;
 	}
 
-	virtual bool Serialize( FArchive& Ar ) override
-	{
-		bool bShaderHasOutdatedParameters = FGlobalShader::Serialize(Ar);
-		Ar << PageTableSize;
-		Ar << FirstUpdate;
-		Ar << NumUpdates;
-		Ar << UpdateBuffer;
-		return bShaderHasOutdatedParameters;
-	}
+	LAYOUT_FIELD(FShaderParameter, PageTableSize);
+	LAYOUT_FIELD(FShaderParameter, FirstUpdate);
+	LAYOUT_FIELD(FShaderParameter, NumUpdates);
+	LAYOUT_FIELD(FShaderResourceParameter, UpdateBuffer);
 };
 
 class FPageTableUpdatePS : public FGlobalShader
 {
-	DECLARE_SHADER_TYPE(FPageTableUpdatePS, Global);
-
+	DECLARE_INLINE_TYPE_LAYOUT(FPageTableUpdatePS, NonVirtual);
+protected:
 	FPageTableUpdatePS() {}
 
 public:
@@ -215,6 +205,9 @@ public:
 		return IsFeatureLevelSupported(Parameters.Platform, ERHIFeatureLevel::SM5)
 			|| GetMaxSupportedFeatureLevel(Parameters.Platform) == ERHIFeatureLevel::ES3_1;
 	}
+
+	
+	
 };
 
 template<bool Use16Bits>
@@ -402,7 +395,7 @@ void FVirtualTextureSpace::ApplyUpdates(FVirtualTextureSystem* System, FRHIComma
 	SCOPED_DRAW_EVENT(RHICmdList, PageTableUpdate);
 
 	auto ShaderMap = GetGlobalShaderMap(GetFeatureLevel());
-	FPageTableUpdateVS* VertexShader = nullptr;
+	TShaderRef<FPageTableUpdateVS> VertexShader;
 	if (Description.PageTableFormat == EVTPageTableFormat::UInt16)
 	{
 		VertexShader = ShaderMap->GetShader< TPageTableUpdateVS<true> >();
@@ -411,7 +404,7 @@ void FVirtualTextureSpace::ApplyUpdates(FVirtualTextureSystem* System, FRHIComma
 	{
 		VertexShader = ShaderMap->GetShader< TPageTableUpdateVS<false> >();
 	}
-	check(VertexShader);
+	check(VertexShader.IsValid());
 
 	uint32 FirstUpdate = 0;
 	for (uint32 LayerIndex = 0u; LayerIndex < Description.NumPageTableLayers; ++LayerIndex)
@@ -431,7 +424,7 @@ void FVirtualTextureSpace::ApplyUpdates(FVirtualTextureSystem* System, FRHIComma
 		default: check(false); break;
 		}
 
-		FPageTableUpdatePS* PixelShader = nullptr;
+		TShaderRef<FPageTableUpdatePS> PixelShader;
 		switch (TexturePixelFormat[TextureIndex])
 		{
 		case PF_R16_UINT: PixelShader = ShaderMap->GetShader< TPageTableUpdatePS<PF_R16_UINT> >(); break;
@@ -442,7 +435,7 @@ void FVirtualTextureSpace::ApplyUpdates(FVirtualTextureSystem* System, FRHIComma
 		case PF_R32G32B32A32_UINT: PixelShader = ShaderMap->GetShader< TPageTableUpdatePS<PF_R32G32B32A32_UINT> >(); break;
 		default: checkNoEntry(); break;
 		}
-		check(PixelShader);
+		check(PixelShader.IsValid());
 
 		uint32 MipSize = PageTableSize;
 		for (uint32 Mip = 0; Mip < NumPageTableLevels; Mip++)
@@ -464,13 +457,13 @@ void FVirtualTextureSpace::ApplyUpdates(FVirtualTextureSystem* System, FRHIComma
 				GraphicsPSOInit.PrimitiveType = PT_TriangleList;
 
 				GraphicsPSOInit.BoundShaderState.VertexDeclarationRHI = GEmptyVertexDeclaration.VertexDeclarationRHI;
-				GraphicsPSOInit.BoundShaderState.VertexShaderRHI = VertexShader->GetVertexShader();
-				GraphicsPSOInit.BoundShaderState.PixelShaderRHI = PixelShader->GetPixelShader();
+				GraphicsPSOInit.BoundShaderState.VertexShaderRHI = VertexShader.GetVertexShader();
+				GraphicsPSOInit.BoundShaderState.PixelShaderRHI = PixelShader.GetPixelShader();
 
 				SetGraphicsPipelineState(RHICmdList, GraphicsPSOInit);
 
 				{
-					FRHIVertexShader* ShaderRHI = VertexShader->GetVertexShader();
+					FRHIVertexShader* ShaderRHI = VertexShader.GetVertexShader();
 					SetShaderValue(RHICmdList, ShaderRHI, VertexShader->PageTableSize, PageTableSize);
 					SetShaderValue(RHICmdList, ShaderRHI, VertexShader->FirstUpdate, FirstUpdate);
 					SetShaderValue(RHICmdList, ShaderRHI, VertexShader->NumUpdates, NumUpdates);
diff --git a/Engine/Source/Runtime/Renderer/Private/VelocityRendering.cpp b/Engine/Source/Runtime/Renderer/Private/VelocityRendering.cpp
--- a/Engine/Source/Runtime/Renderer/Private/VelocityRendering.cpp
+++ b/Engine/Source/Runtime/Renderer/Private/VelocityRendering.cpp
@@ -61,23 +61,21 @@ bool IsParallelVelocity()
 class FVelocityVS : public FMeshMaterialShader
 {
 public:
-	DECLARE_MESH_MATERIAL_SHADER(FVelocityVS);
+	DECLARE_SHADER_TYPE(FVelocityVS, MeshMaterial);
 
 	static bool ShouldCompilePermutation(const FMeshMaterialShaderPermutationParameters& Parameters)
 	{
-		const FMaterial* Material = Parameters.Material;
-
 		// Compile for default material.
-		const bool bIsDefault = Material->IsSpecialEngineMaterial();
+		const bool bIsDefault = Parameters.MaterialParameters.bIsSpecialEngineMaterial;
 
 		// Compile for masked materials.
-		const bool bIsMasked = !Material->WritesEveryPixel();
+		const bool bIsMasked = !Parameters.MaterialParameters.bWritesEveryPixel;
 
 		// Compile for opaque and two-sided materials.
-		const bool bIsOpaqueAndTwoSided = (Material->IsTwoSided() && !IsTranslucentBlendMode(Material->GetBlendMode()));
+		const bool bIsOpaqueAndTwoSided = (Parameters.MaterialParameters.bIsTwoSided && !IsTranslucentBlendMode(Parameters.MaterialParameters.BlendMode));
 
 		// Compile for materials which modify meshes.
-		const bool bMayModifyMeshes = Material->MaterialMayModifyMeshPosition();
+		const bool bMayModifyMeshes = Parameters.MaterialParameters.bMaterialMayModifyMeshPosition;
 
 		// Compile if supported by the hardware.
 		const bool bIsFeatureSupported = IsFeatureLevelSupported(Parameters.Platform, ERHIFeatureLevel::SM5);
@@ -91,7 +89,7 @@ public:
 			FVelocityRendering::IsSeparateVelocityPassRequiredByVertexFactory(Parameters.Platform, Parameters.VertexFactoryType->SupportsStaticLighting());
 
 		// The material may explicitly override and request that it be rendered into the velocity pass.
-		const bool bIsSeparateVelocityPassRequiredByMaterial = Material->IsTranslucencyWritingVelocity();
+		const bool bIsSeparateVelocityPassRequiredByMaterial = Parameters.MaterialParameters.bIsTranslucencyWritingVelocity;
 
 		return bIsFeatureSupported && (bIsSeparateVelocityPassRequired || bIsSeparateVelocityPassRequiredByMaterial);
 	}
@@ -107,7 +105,7 @@ public:
 class FVelocityHS : public FBaseHS
 {
 public:
-	DECLARE_MESH_MATERIAL_SHADER(FVelocityHS);
+	DECLARE_SHADER_TYPE(FVelocityHS, MeshMaterial);
 
 	static bool ShouldCompilePermutation(const FMeshMaterialShaderPermutationParameters& Parameters)
 	{
@@ -123,7 +121,7 @@ public:
 class FVelocityDS : public FBaseDS
 {
 public:
-	DECLARE_MESH_MATERIAL_SHADER(FVelocityDS);
+	DECLARE_SHADER_TYPE(FVelocityDS, MeshMaterial);
 
 	static bool ShouldCompilePermutation(const FMeshMaterialShaderPermutationParameters& Parameters)
 	{
@@ -139,7 +137,7 @@ public:
 class FVelocityPS : public FMeshMaterialShader
 {
 public:
-	DECLARE_MESH_MATERIAL_SHADER(FVelocityPS);
+	DECLARE_SHADER_TYPE(FVelocityPS, MeshMaterial);
 
 	static bool ShouldCompilePermutation(const FMeshMaterialShaderPermutationParameters& Parameters)
 	{
@@ -160,10 +158,10 @@ public:
 	}
 };
 
-IMPLEMENT_MESH_MATERIAL_SHADER(FVelocityVS, "/Engine/Private/VelocityShader.usf", "MainVertexShader", SF_Vertex); 
-IMPLEMENT_MESH_MATERIAL_SHADER(FVelocityHS, "/Engine/Private/VelocityShader.usf", "MainHull", SF_Hull); 
-IMPLEMENT_MESH_MATERIAL_SHADER(FVelocityDS, "/Engine/Private/VelocityShader.usf", "MainDomain", SF_Domain);
-IMPLEMENT_MESH_MATERIAL_SHADER(FVelocityPS, "/Engine/Private/VelocityShader.usf", "MainPixelShader", SF_Pixel);
+IMPLEMENT_SHADER_TYPE(,FVelocityVS, TEXT("/Engine/Private/VelocityShader.usf"), TEXT("MainVertexShader"), SF_Vertex);
+IMPLEMENT_SHADER_TYPE(,FVelocityHS, TEXT("/Engine/Private/VelocityShader.usf"), TEXT("MainHull"), SF_Hull);
+IMPLEMENT_SHADER_TYPE(,FVelocityDS, TEXT("/Engine/Private/VelocityShader.usf"), TEXT("MainDomain"), SF_Domain);
+IMPLEMENT_SHADER_TYPE(,FVelocityPS, TEXT("/Engine/Private/VelocityShader.usf"), TEXT("MainPixelShader"), SF_Pixel);
 IMPLEMENT_SHADERPIPELINE_TYPE_VSPS(VelocityPipeline, FVelocityVS, FVelocityPS, true);
 
 EMeshPass::Type GetMeshPassFromVelocityPass(EVelocityPass VelocityPass)
@@ -615,8 +613,9 @@ void FOpaqueVelocityMeshProcessor::AddMeshBatch(
 	if (MeshBatch.bUseForMaterial && bIsNotTranslucent && ShouldIncludeMaterialInDefaultOpaquePass(*Material))
 	{
 		// This is specifically done *before* the material swap, as swapped materials may have different fill / cull modes.
-		const ERasterizerFillMode MeshFillMode = ComputeMeshFillMode(MeshBatch, *Material);
-		const ERasterizerCullMode MeshCullMode = ComputeMeshCullMode(MeshBatch, *Material);
+		const FMeshDrawingPolicyOverrideSettings OverrideSettings = ComputeMeshOverrideSettings(MeshBatch);
+		const ERasterizerFillMode MeshFillMode = ComputeMeshFillMode(MeshBatch, *Material, OverrideSettings);
+		const ERasterizerCullMode MeshCullMode = ComputeMeshCullMode(MeshBatch, *Material, OverrideSettings);
 
 		/**
 		 * Materials without masking or custom vertex modifications can be swapped out
@@ -689,8 +688,9 @@ void FTranslucentVelocityMeshProcessor::AddMeshBatch(
 
 	if (MeshBatch.bUseForMaterial && bMaterialWritesVelocity)
 	{
-		const ERasterizerFillMode MeshFillMode = ComputeMeshFillMode(MeshBatch, *Material);
-		const ERasterizerCullMode MeshCullMode = ComputeMeshCullMode(MeshBatch, *Material);
+		const FMeshDrawingPolicyOverrideSettings OverrideSettings = ComputeMeshOverrideSettings(MeshBatch);
+		const ERasterizerFillMode MeshFillMode = ComputeMeshFillMode(MeshBatch, *Material, OverrideSettings);
+		const ERasterizerCullMode MeshCullMode = ComputeMeshCullMode(MeshBatch, *Material, OverrideSettings);
 
 		Process(MeshBatch, BatchElementMask, StaticMeshId, PrimitiveSceneProxy, *MaterialRenderProxy, *Material, MeshFillMode, MeshCullMode);
 	}
@@ -700,10 +700,10 @@ void GetVelocityPassShaders(
 	const FMaterial& Material,
 	FVertexFactoryType* VertexFactoryType,
 	ERHIFeatureLevel::Type FeatureLevel,
-	FVelocityHS*& HullShader,
-	FVelocityDS*& DomainShader,
-	FVelocityVS*& VertexShader,
-	FVelocityPS*& PixelShader)
+	TShaderRef<FVelocityHS>& HullShader,
+	TShaderRef<FVelocityDS>& DomainShader,
+	TShaderRef<FVelocityVS>& VertexShader,
+	TShaderRef<FVelocityPS>& PixelShader)
 {
 	const EMaterialTessellationMode MaterialTessellationMode = Material.GetTessellationMode();
 
@@ -720,18 +720,18 @@ void GetVelocityPassShaders(
 	static const auto* CVar = IConsoleManager::Get().FindTConsoleVariableDataInt(TEXT("r.ShaderPipelines"));
 	const bool bUseShaderPipelines = RHISupportsShaderPipelines(GShaderPlatformForFeatureLevel[FeatureLevel]) && !bNeedsHSDS && CVar && CVar->GetValueOnAnyThread() != 0;
 
-	FShaderPipeline* ShaderPipeline = bUseShaderPipelines ? Material.GetShaderPipeline(&VelocityPipeline, VertexFactoryType, false) : nullptr;
-	if (ShaderPipeline)
+	FShaderPipelineRef ShaderPipeline = bUseShaderPipelines ? Material.GetShaderPipeline(&VelocityPipeline, VertexFactoryType, false) : FShaderPipelineRef();
+	if (ShaderPipeline.IsValid())
 	{
-		VertexShader = ShaderPipeline->GetShader<FVelocityVS>();
-		PixelShader = ShaderPipeline->GetShader<FVelocityPS>();
-		check(VertexShader && PixelShader);
+		VertexShader = ShaderPipeline.GetShader<FVelocityVS>();
+		PixelShader = ShaderPipeline.GetShader<FVelocityPS>();
+		check(VertexShader.IsValid() && PixelShader.IsValid());
 	}
 	else
 	{
 		VertexShader = Material.GetShader<FVelocityVS>(VertexFactoryType);
 		PixelShader = Material.GetShader<FVelocityPS>(VertexFactoryType);
-		check(VertexShader && PixelShader);
+		check(VertexShader.IsValid() && PixelShader.IsValid());
 	}
 }
 
@@ -806,7 +806,6 @@ FMeshPassProcessor* CreateVelocityPassProcessor(const FScene* Scene, const FScen
 }
 
 FRegisterPassProcessorCreateFunction RegisterVelocityPass(&CreateVelocityPassProcessor, EShadingPath::Deferred,  EMeshPass::Velocity, EMeshPassFlags::CachedMeshCommands | EMeshPassFlags::MainView);
-
 FTranslucentVelocityMeshProcessor::FTranslucentVelocityMeshProcessor(const FScene* Scene, const FSceneView* InViewIfDynamicMeshCommand, const FMeshPassProcessorRenderState& InPassDrawRenderState, FMeshPassDrawListContext* InDrawListContext)
 	: FVelocityMeshProcessor(Scene, InViewIfDynamicMeshCommand, InPassDrawRenderState, InDrawListContext)
 {}
diff --git a/Engine/Source/Runtime/Renderer/Private/VisualizeTexturePresent.cpp b/Engine/Source/Runtime/Renderer/Private/VisualizeTexturePresent.cpp
--- a/Engine/Source/Runtime/Renderer/Private/VisualizeTexturePresent.cpp
+++ b/Engine/Source/Runtime/Renderer/Private/VisualizeTexturePresent.cpp
@@ -473,8 +473,8 @@ void FVisualizeTexturePresent::PresentContent(FRHICommandListImmediate& RHICmdLi
 	TShaderMapRef<FVisualizeTexturePresentPS> PixelShader(ShaderMap);
 
 	GraphicsPSOInit.BoundShaderState.VertexDeclarationRHI = GFilterVertexDeclaration.VertexDeclarationRHI;
-	GraphicsPSOInit.BoundShaderState.VertexShaderRHI = GETSAFERHISHADER_VERTEX(*VertexShader);
-	GraphicsPSOInit.BoundShaderState.PixelShaderRHI = GETSAFERHISHADER_PIXEL(*PixelShader);
+	GraphicsPSOInit.BoundShaderState.VertexShaderRHI = VertexShader.GetVertexShader();
+	GraphicsPSOInit.BoundShaderState.PixelShaderRHI = PixelShader.GetPixelShader();
 	GraphicsPSOInit.PrimitiveType = PT_TriangleList;
 	SetGraphicsPipelineState(RHICmdList, GraphicsPSOInit);
 
@@ -483,7 +483,7 @@ void FVisualizeTexturePresent::PresentContent(FRHICommandListImmediate& RHICmdLi
 		FVisualizeTexturePresentPS::FParameters Parameters;
 		Parameters.VisualizeTexture2D = GVisualizeTexture.VisualizeTextureContent->GetRenderTargetItem().ShaderResourceTexture;
 		Parameters.VisualizeTexture2DSampler = TStaticSamplerState<SF_Point, AM_Clamp, AM_Clamp, AM_Clamp>::GetRHI();
-		SetShaderParameters(RHICmdList, *PixelShader, PixelShader->GetPixelShader(), Parameters);
+		SetShaderParameters(RHICmdList, PixelShader, PixelShader.GetPixelShader(), Parameters);
 	}
 
 	{
@@ -496,7 +496,7 @@ void FVisualizeTexturePresent::PresentContent(FRHICommandListImmediate& RHICmdLi
 			SrcRect.Width(), SrcRect.Height(),
 			DestRect.Size(),
 			SrcSize,
-			*VertexShader,
+			VertexShader,
 			EDRF_Default);
 	}
 	RHICmdList.EndRenderPass();
diff --git a/Engine/Source/Runtime/Renderer/Private/VisualizeVolumetricLightmap.cpp b/Engine/Source/Runtime/Renderer/Private/VisualizeVolumetricLightmap.cpp
--- a/Engine/Source/Runtime/Renderer/Private/VisualizeVolumetricLightmap.cpp
+++ b/Engine/Source/Runtime/Renderer/Private/VisualizeVolumetricLightmap.cpp
@@ -74,25 +74,16 @@ public:
 
 	void SetParameters(FRHICommandList& RHICmdList, const FSceneView& View)
 	{
-		FRHIVertexShader* ShaderRHI = GetVertexShader();
+		FRHIVertexShader* ShaderRHI = RHICmdList.GetBoundVertexShader();
 		FGlobalShader::SetParameters<FViewUniformShaderParameters>(RHICmdList, ShaderRHI, View.ViewUniformBuffer);
 
 		SetShaderValue(RHICmdList, ShaderRHI, VisualizationRadiusScale, GVolumetricLightmapVisualizationRadiusScale);
 		SetShaderValue(RHICmdList, ShaderRHI, VisualizationMinScreenFraction, GVolumetricLightmapVisualizationMinScreenFraction);
 	}
 
-	// FShader interface.
-	virtual bool Serialize(FArchive& Ar) override
-	{
-		bool bShaderHasOutdatedParameters = FGlobalShader::Serialize(Ar);
-		Ar << VisualizationRadiusScale;
-		Ar << VisualizationMinScreenFraction;
-		return bShaderHasOutdatedParameters;
-	}
-
 private:
-	FShaderParameter VisualizationRadiusScale;
-	FShaderParameter VisualizationMinScreenFraction;
+	LAYOUT_FIELD(FShaderParameter, VisualizationRadiusScale);
+	LAYOUT_FIELD(FShaderParameter, VisualizationMinScreenFraction);
 };
 
 IMPLEMENT_SHADER_TYPE(,FVisualizeVolumetricLightmapVS,TEXT("/Engine/Private/VisualizeVolumetricLightmap.usf"),TEXT("VisualizeVolumetricLightmapVS"),SF_Vertex);
@@ -120,7 +111,7 @@ public:
 
 	void SetParameters(FRHICommandList& RHICmdList, const FSceneView& View)
 	{
-		FRHIPixelShader* ShaderRHI = GetPixelShader();
+		FRHIPixelShader* ShaderRHI = RHICmdList.GetBoundPixelShader();
 
 		FGlobalShader::SetParameters<FViewUniformShaderParameters>(RHICmdList, ShaderRHI, View.ViewUniformBuffer);
 
@@ -134,16 +125,8 @@ public:
 		SetShaderValue(RHICmdList, ShaderRHI, DiffuseColor, DiffuseColorValue);
 	}
 
-	// FShader interface.
-	virtual bool Serialize(FArchive& Ar) override
-	{
-		bool bShaderHasOutdatedParameters = FGlobalShader::Serialize(Ar);
-		Ar << DiffuseColor;
-		return bShaderHasOutdatedParameters;
-	}
-
 private:
-	FShaderParameter DiffuseColor;
+	LAYOUT_FIELD(FShaderParameter, DiffuseColor);
 };
 
 IMPLEMENT_SHADER_TYPE(,FVisualizeVolumetricLightmapPS,TEXT("/Engine/Private/VisualizeVolumetricLightmap.usf"),TEXT("VisualizeVolumetricLightmapPS"),SF_Pixel);
@@ -197,8 +180,8 @@ void FDeferredShadingSceneRenderer::VisualizeVolumetricLightmap(FRHICommandListI
 				TShaderMapRef<FVisualizeVolumetricLightmapPS> PixelShader(View.ShaderMap);
 
 				GraphicsPSOInit.BoundShaderState.VertexDeclarationRHI = GEmptyVertexDeclaration.VertexDeclarationRHI;
-				GraphicsPSOInit.BoundShaderState.VertexShaderRHI = GETSAFERHISHADER_VERTEX(*VertexShader);
-				GraphicsPSOInit.BoundShaderState.PixelShaderRHI = GETSAFERHISHADER_PIXEL(*PixelShader);
+				GraphicsPSOInit.BoundShaderState.VertexShaderRHI = VertexShader.GetVertexShader();
+				GraphicsPSOInit.BoundShaderState.PixelShaderRHI = PixelShader.GetPixelShader();
 
 				SetGraphicsPipelineState(RHICmdList, GraphicsPSOInit);
 
diff --git a/Engine/Source/Runtime/Renderer/Private/VolumeLighting.h b/Engine/Source/Runtime/Renderer/Private/VolumeLighting.h
--- a/Engine/Source/Runtime/Renderer/Private/VolumeLighting.h
+++ b/Engine/Source/Runtime/Renderer/Private/VolumeLighting.h
@@ -14,6 +14,7 @@
 
 class FVolumeShadowingParameters
 {
+	DECLARE_TYPE_LAYOUT(FVolumeShadowingParameters, NonVirtual);
 public:
 
 	void Bind(const FShaderParameterMap& ParameterMap)
@@ -143,7 +144,7 @@ public:
 	}
 
 	/** Serializer. */ 
-	friend FArchive& operator<<(FArchive& Ar,FVolumeShadowingParameters& P)
+	/*friend FArchive& operator<<(FArchive& Ar,FVolumeShadowingParameters& P)
 	{
 		Ar << P.WorldToShadowMatrix;
 		Ar << P.ShadowmapMinMax;
@@ -159,21 +160,22 @@ public:
 		Ar << P.WorldToStaticShadowMatrix;
 		Ar << P.StaticShadowBufferSize;
 		return Ar;
-	}
+	}*/
 
 private:
-
-	FShaderParameter WorldToShadowMatrix;
-	FShaderParameter ShadowmapMinMax;
-	FShaderParameter DepthBiasParameters;
-	FShaderParameter ShadowInjectParams;
-	FShaderParameter ClippingPlanes;
-	FShaderResourceParameter ShadowDepthTexture;
-	FShaderResourceParameter ShadowDepthTextureSampler;
-	FOnePassPointShadowProjectionShaderParameters OnePassShadowParameters;
-	FShaderParameter bStaticallyShadowed;
-	FShaderResourceParameter StaticShadowDepthTexture;
-	FShaderResourceParameter StaticShadowDepthTextureSampler;
-	FShaderParameter WorldToStaticShadowMatrix;
-	FShaderParameter StaticShadowBufferSize;
+	
+		LAYOUT_FIELD(FShaderParameter, WorldToShadowMatrix)
+		LAYOUT_FIELD(FShaderParameter, ShadowmapMinMax)
+		LAYOUT_FIELD(FShaderParameter, DepthBiasParameters)
+		LAYOUT_FIELD(FShaderParameter, ShadowInjectParams)
+		LAYOUT_FIELD(FShaderParameter, ClippingPlanes)
+		LAYOUT_FIELD(FShaderResourceParameter, ShadowDepthTexture)
+		LAYOUT_FIELD(FShaderResourceParameter, ShadowDepthTextureSampler)
+		LAYOUT_FIELD(FOnePassPointShadowProjectionShaderParameters, OnePassShadowParameters)
+		LAYOUT_FIELD(FShaderParameter, bStaticallyShadowed)
+		LAYOUT_FIELD(FShaderResourceParameter, StaticShadowDepthTexture)
+		LAYOUT_FIELD(FShaderResourceParameter, StaticShadowDepthTextureSampler)
+		LAYOUT_FIELD(FShaderParameter, WorldToStaticShadowMatrix)
+		LAYOUT_FIELD(FShaderParameter, StaticShadowBufferSize)
+	
 };
\ No newline at end of file
diff --git a/Engine/Source/Runtime/Renderer/Private/VolumetricFog.cpp b/Engine/Source/Runtime/Renderer/Private/VolumetricFog.cpp
--- a/Engine/Source/Runtime/Renderer/Private/VolumetricFog.cpp
+++ b/Engine/Source/Runtime/Renderer/Private/VolumetricFog.cpp
@@ -20,6 +20,9 @@ VolumetricFog.cpp
 #include "ShaderParameterStruct.h"
 #include "Math/Halton.h"
 
+IMPLEMENT_TYPE_LAYOUT(FVolumetricFogIntegrationParameters);
+IMPLEMENT_TYPE_LAYOUT(FVolumeShadowingParameters);
+
 int32 GVolumetricFog = 1;
 FAutoConsoleVariableRef CVarVolumetricFog(
 	TEXT("r.VolumetricFog"),
@@ -174,20 +177,13 @@ public:
 		const FViewInfo& View,
 		const FVolumetricFogIntegrationParameterData& IntegrationData)
 	{
-		FRHIComputeShader* ShaderRHI = GetComputeShader();
+		FRHIComputeShader* ShaderRHI = RHICmdList.GetBoundComputeShader();
 		VolumetricFogParameters.Set(RHICmdList, ShaderRHI, View, IntegrationData);
 	}
 
-	virtual bool Serialize(FArchive& Ar)
-	{
-		bool bShaderHasOutdatedParameters = FGlobalShader::Serialize(Ar);
-		Ar << VolumetricFogParameters;
-		return bShaderHasOutdatedParameters;
-	}
-
 private:
 
-	FVolumetricFogIntegrationParameters VolumetricFogParameters;
+	LAYOUT_FIELD(FVolumetricFogIntegrationParameters, VolumetricFogParameters);
 };
 
 IMPLEMENT_SHADER_TYPE(, FVolumetricFogMaterialSetupCS, TEXT("/Engine/Private/VolumetricFog.usf"), TEXT("MaterialSetupCS"), SF_Compute);
@@ -223,32 +219,22 @@ public:
 
 	void SetParameters(FRHICommandList& RHICmdList, const FViewInfo& View, const FVolumetricFogIntegrationParameterData& IntegrationData, const FSphere& BoundingSphere, int32 MinZValue)
 	{
-		SetShaderValue(RHICmdList, GetVertexShader(), MinZ, MinZValue);
+		SetShaderValue(RHICmdList, RHICmdList.GetBoundVertexShader(), MinZ, MinZValue);
 
 		const FVector ViewSpaceBoundingSphereCenter = View.ViewMatrices.GetViewMatrix().TransformPosition(BoundingSphere.Center);
-		SetShaderValue(RHICmdList, GetVertexShader(), ViewSpaceBoundingSphere, FVector4(ViewSpaceBoundingSphereCenter, BoundingSphere.W));
+		SetShaderValue(RHICmdList, RHICmdList.GetBoundVertexShader(), ViewSpaceBoundingSphere, FVector4(ViewSpaceBoundingSphereCenter, BoundingSphere.W));
 
 		const FMatrix ProjectionMatrix = View.ViewMatrices.ComputeProjectionNoAAMatrix();
-		SetShaderValue(RHICmdList, GetVertexShader(), ViewToVolumeClip, ProjectionMatrix);
+		SetShaderValue(RHICmdList, RHICmdList.GetBoundVertexShader(), ViewToVolumeClip, ProjectionMatrix);
 
-		VolumetricFogParameters.Set(RHICmdList, GetVertexShader(), View, IntegrationData);
-	}
-
-	virtual bool Serialize(FArchive& Ar) override
-	{
-		bool bShaderHasOutdatedParameters = FGlobalShader::Serialize(Ar);
-		Ar << MinZ;
-		Ar << ViewSpaceBoundingSphere;
-		Ar << ViewToVolumeClip;
-		Ar << VolumetricFogParameters;
-		return bShaderHasOutdatedParameters;
+		VolumetricFogParameters.Set(RHICmdList, RHICmdList.GetBoundVertexShader(), View, IntegrationData);
 	}
 
 private:
-	FShaderParameter MinZ;
-	FShaderParameter ViewSpaceBoundingSphere;
-	FShaderParameter ViewToVolumeClip;
-	FVolumetricFogIntegrationParameters VolumetricFogParameters;
+	LAYOUT_FIELD(FShaderParameter, MinZ);
+	LAYOUT_FIELD(FShaderParameter, ViewSpaceBoundingSphere);
+	LAYOUT_FIELD(FShaderParameter, ViewToVolumeClip);
+	LAYOUT_FIELD(FVolumetricFogIntegrationParameters, VolumetricFogParameters);
 };
 
 IMPLEMENT_SHADER_TYPE(, FWriteToBoundingSphereVS, TEXT("/Engine/Private/VolumetricFog.usf"), TEXT("WriteToBoundingSphereVS"), SF_Vertex);
@@ -301,7 +287,7 @@ public:
 		const FProjectedShadowInfo* ShadowMap,
 		bool bDynamicallyShadowed)
 	{
-		FRHIPixelShader* ShaderRHI = GetPixelShader();
+		FRHIPixelShader* ShaderRHI = RHICmdList.GetBoundPixelShader();
 
 		FGlobalShader::SetParameters<FViewUniformShaderParameters>(RHICmdList, ShaderRHI, View.ViewUniformBuffer);
 
@@ -315,21 +301,11 @@ public:
 		VolumeShadowingParameters.Set(RHICmdList, ShaderRHI, View, LightSceneInfo, ShadowMap, INDEX_NONE, bDynamicallyShadowed);
 	}
 
-	virtual bool Serialize(FArchive& Ar) override
-	{
-		bool bShaderHasOutdatedParameters = FGlobalShader::Serialize(Ar);
-		Ar << PhaseG;
-		Ar << InverseSquaredLightDistanceBiasScale;
-		Ar << VolumetricFogParameters;
-		Ar << VolumeShadowingParameters;
-		return bShaderHasOutdatedParameters;
-	}
-
 private:
-	FShaderParameter PhaseG;
-	FShaderParameter InverseSquaredLightDistanceBiasScale;
-	FVolumetricFogIntegrationParameters VolumetricFogParameters;
-	FVolumeShadowingParameters VolumeShadowingParameters;
+	LAYOUT_FIELD(FShaderParameter, PhaseG);
+	LAYOUT_FIELD(FShaderParameter, InverseSquaredLightDistanceBiasScale);
+	LAYOUT_FIELD(FVolumetricFogIntegrationParameters, VolumetricFogParameters);
+	LAYOUT_FIELD(FVolumeShadowingParameters, VolumeShadowingParameters);
 };
 
 IMPLEMENT_GLOBAL_SHADER(TInjectShadowedLocalLightPS, "/Engine/Private/VolumetricFog.usf", "InjectShadowedLocalLightPS", SF_Pixel);
@@ -538,9 +514,9 @@ void FDeferredShadingSceneRenderer::RenderLocalLightsForVolumetricFog(
 					GraphicsPSOInit.BlendState = TStaticBlendState<CW_RGBA, BO_Add, BF_One, BF_One, BO_Add, BF_Zero, BF_One>::GetRHI();
 
 					GraphicsPSOInit.BoundShaderState.VertexDeclarationRHI = GScreenVertexDeclaration.VertexDeclarationRHI;
-					GraphicsPSOInit.BoundShaderState.VertexShaderRHI = GETSAFERHISHADER_VERTEX(VertexShader);
-					GraphicsPSOInit.BoundShaderState.GeometryShaderRHI = GETSAFERHISHADER_GEOMETRY(*GeometryShader);
-					GraphicsPSOInit.BoundShaderState.PixelShaderRHI = GETSAFERHISHADER_PIXEL(PixelShader);
+					GraphicsPSOInit.BoundShaderState.VertexShaderRHI = VertexShader.GetVertexShader();
+					GraphicsPSOInit.BoundShaderState.GeometryShaderRHI = GeometryShader.GetGeometryShader();
+					GraphicsPSOInit.BoundShaderState.PixelShaderRHI = PixelShader.GetPixelShader();
 					GraphicsPSOInit.PrimitiveType = PT_TriangleList;
 
 					SetGraphicsPipelineState(RHICmdList, GraphicsPSOInit);
@@ -565,7 +541,7 @@ void FDeferredShadingSceneRenderer::RenderLocalLightsForVolumetricFog(
 
 class TVolumetricFogLightScatteringCS : public FGlobalShader
 {
-	DECLARE_GLOBAL_SHADER(TVolumetricFogLightScatteringCS)
+	DECLARE_SHADER_TYPE(TVolumetricFogLightScatteringCS, Global)
 
 	class FTemporalReprojection			: SHADER_PERMUTATION_BOOL("USE_TEMPORAL_REPROJECTION");
 	class FDistanceFieldSkyOcclusion	: SHADER_PERMUTATION_BOOL("DISTANCE_FIELD_SKY_OCCLUSION");
@@ -639,7 +615,7 @@ public:
 		bool bUseDirectionalLightShadowing,
 		const FMatrix& DirectionalLightFunctionWorldToShadowValue)
 	{
-		FRHIComputeShader* ShaderRHI = GetComputeShader();
+		FRHIComputeShader* ShaderRHI = RHICmdList.GetBoundComputeShader();
 		FGlobalShader::SetParameters<FViewUniformShaderParameters>(RHICmdList, ShaderRHI, View.ViewUniformBuffer);
 
 		if (!LightScatteringHistoryTexture)
@@ -716,48 +692,24 @@ public:
 		SetUniformBufferParameterImmediate(RHICmdList, ShaderRHI, GetUniformBufferParameter<FFogUniformParameters>(), FogUniformParameters);
 	}
 
-	virtual bool Serialize(FArchive& Ar)
-	{
-		bool bShaderHasOutdatedParameters = FGlobalShader::Serialize(Ar);
-		Ar << LocalShadowedLightScattering;
-		Ar << LightScatteringHistory;
-		Ar << LightScatteringHistorySampler;
-		Ar << VolumetricFogParameters;
-		Ar << DirectionalLightFunctionWorldToShadow;
-		Ar << LightFunctionTexture;
-		Ar << LightFunctionSampler;
-		Ar << StaticLightingScatteringIntensity;
-		Ar << SkyLightUseStaticShadowing;
-		Ar << SkyLightVolumetricScatteringIntensity;
-		Ar << SkySH;
-		Ar << PhaseG;
-		Ar << InverseSquaredLightDistanceBiasScale;
-		Ar << UseHeightFogColors;
-		Ar << UseDirectionalLightShadowing;
-		Ar << AOParameters;
-		Ar << GlobalDistanceFieldParameters;
-		return bShaderHasOutdatedParameters;
-	}
-
 private:
-
-	FShaderResourceParameter LocalShadowedLightScattering;
-	FShaderResourceParameter LightScatteringHistory;
-	FShaderResourceParameter LightScatteringHistorySampler;
-	FVolumetricFogIntegrationParameters VolumetricFogParameters;
-	FShaderParameter DirectionalLightFunctionWorldToShadow;
-	FShaderResourceParameter LightFunctionTexture;
-	FShaderResourceParameter LightFunctionSampler;
-	FShaderParameter StaticLightingScatteringIntensity;
-	FShaderParameter SkyLightUseStaticShadowing;
-	FShaderParameter SkyLightVolumetricScatteringIntensity;
-	FShaderParameter SkySH;
-	FShaderParameter PhaseG;
-	FShaderParameter InverseSquaredLightDistanceBiasScale;
-	FShaderParameter UseHeightFogColors;
-	FShaderParameter UseDirectionalLightShadowing;
-	FAOParameters AOParameters;
-	FGlobalDistanceFieldParameters GlobalDistanceFieldParameters;
+	LAYOUT_FIELD(FShaderResourceParameter, LocalShadowedLightScattering);
+	LAYOUT_FIELD(FShaderResourceParameter, LightScatteringHistory);
+	LAYOUT_FIELD(FShaderResourceParameter, LightScatteringHistorySampler);
+	LAYOUT_FIELD(FVolumetricFogIntegrationParameters, VolumetricFogParameters);
+	LAYOUT_FIELD(FShaderParameter, DirectionalLightFunctionWorldToShadow);
+	LAYOUT_FIELD(FShaderResourceParameter, LightFunctionTexture);
+	LAYOUT_FIELD(FShaderResourceParameter, LightFunctionSampler);
+	LAYOUT_FIELD(FShaderParameter, StaticLightingScatteringIntensity);
+	LAYOUT_FIELD(FShaderParameter, SkyLightUseStaticShadowing);
+	LAYOUT_FIELD(FShaderParameter, SkyLightVolumetricScatteringIntensity);
+	LAYOUT_FIELD(FShaderParameter, SkySH);
+	LAYOUT_FIELD(FShaderParameter, PhaseG);
+	LAYOUT_FIELD(FShaderParameter, InverseSquaredLightDistanceBiasScale);
+	LAYOUT_FIELD(FShaderParameter, UseHeightFogColors);
+	LAYOUT_FIELD(FShaderParameter, UseDirectionalLightShadowing);
+	LAYOUT_FIELD(FAOParameters, AOParameters);
+	LAYOUT_FIELD(FGlobalDistanceFieldParameters, GlobalDistanceFieldParameters);
 };
 
 IMPLEMENT_GLOBAL_SHADER(TVolumetricFogLightScatteringCS, "/Engine/Private/VolumetricFog.usf", "LightScatteringCS", SF_Compute);
@@ -800,21 +752,14 @@ class FVolumetricFogFinalIntegrationCS : public FGlobalShader
 public:
 	void SetParameters(FRHICommandList& RHICmdList, const FViewInfo& View, const FVolumetricFogIntegrationParameterData& IntegrationData)
 	{
-		FRHIComputeShader* ShaderRHI = GetComputeShader();
+		FRHIComputeShader* ShaderRHI = RHICmdList.GetBoundComputeShader();
 		FGlobalShader::SetParameters<FViewUniformShaderParameters>(RHICmdList, ShaderRHI, View.ViewUniformBuffer);
 		VolumetricFogParameters.Set(RHICmdList, ShaderRHI, View, IntegrationData);
 	}
 
-	virtual bool Serialize(FArchive& Ar)
-	{
-		bool bShaderHasOutdatedParameters = FGlobalShader::Serialize(Ar);
-		Ar << VolumetricFogParameters;
-		return bShaderHasOutdatedParameters;
-	}
-
 private:
 
-	FVolumetricFogIntegrationParameters VolumetricFogParameters;
+	LAYOUT_FIELD(FVolumetricFogIntegrationParameters, VolumetricFogParameters);
 };
 
 IMPLEMENT_SHADER_TYPE(, FVolumetricFogFinalIntegrationCS, TEXT("/Engine/Private/VolumetricFog.usf"), TEXT("FinalIntegrationCS"), SF_Compute);
@@ -1086,13 +1031,13 @@ void FDeferredShadingSceneRenderer::ComputeVolumetricFog(FRHICommandListImmediat
 				{
 					const FIntVector NumGroups = FIntVector::DivideAndRoundUp(VolumetricFogGridSize, VolumetricFogGridInjectionGroupSize);
 
-					RHICmdList.SetComputeShader(ComputeShader->GetComputeShader());
+					RHICmdList.SetComputeShader(ComputeShader.GetComputeShader());
 
 					ComputeShader->SetParameters(RHICmdList, View, IntegrationData);
 
-					SetShaderParameters(RHICmdList, ComputeShader, ComputeShader->GetComputeShader(), *PassParameters);
-					DispatchComputeShader(RHICmdList, ComputeShader, NumGroups.X, NumGroups.Y, NumGroups.Z);
-					UnsetShaderUAVs(RHICmdList, ComputeShader, ComputeShader->GetComputeShader());
+					SetShaderParameters(RHICmdList, ComputeShader, ComputeShader.GetComputeShader(), *PassParameters);
+					DispatchComputeShader(RHICmdList, ComputeShader.GetShader(), NumGroups.X, NumGroups.Y, NumGroups.Z);
+					UnsetShaderUAVs(RHICmdList, ComputeShader, ComputeShader.GetComputeShader());
 				});
 
 				VoxelizeFogVolumePrimitives(
@@ -1155,7 +1100,7 @@ void FDeferredShadingSceneRenderer::ComputeVolumetricFog(FRHICommandListImmediat
 					UnbindRenderTargets(RHICmdList);
 					const FIntVector NumGroups = FIntVector::DivideAndRoundUp(VolumetricFogGridSize, FIntVector(VolumetricFogLightScatteringGroupSizeX, VolumetricFogLightScatteringGroupSizeY, VolumetricFogLightScatteringGroupSizeZ));
 
-					RHICmdList.SetComputeShader(ComputeShader->GetComputeShader());
+					RHICmdList.SetComputeShader(ComputeShader.GetComputeShader());
 
 					FRHITexture* LightScatteringHistoryTexture = bUseTemporalReprojection && View.ViewState->LightScatteringHistory.IsValid()
 						? View.ViewState->LightScatteringHistory->GetRenderTargetItem().ShaderResourceTexture
@@ -1164,9 +1109,9 @@ void FDeferredShadingSceneRenderer::ComputeVolumetricFog(FRHICommandListImmediat
 					RHICmdList.TransitionResource(EResourceTransitionAccess::EReadable, LightScatteringHistoryTexture);
 					ComputeShader->SetParameters(RHICmdList, View, IntegrationData, FogInfo, LightScatteringHistoryTexture, bUseDirectionalLightShadowing, LightFunctionWorldToShadow);
 
-					SetShaderParameters(RHICmdList, ComputeShader, ComputeShader->GetComputeShader(), *PassParameters);
-					DispatchComputeShader(RHICmdList, ComputeShader, NumGroups.X, NumGroups.Y, NumGroups.Z);
-					UnsetShaderUAVs(RHICmdList, ComputeShader, ComputeShader->GetComputeShader());
+					SetShaderParameters(RHICmdList, ComputeShader, ComputeShader.GetComputeShader(), *PassParameters);
+					DispatchComputeShader(RHICmdList, ComputeShader.GetShader(), NumGroups.X, NumGroups.Y, NumGroups.Z);
+					UnsetShaderUAVs(RHICmdList, ComputeShader, ComputeShader.GetComputeShader());
 				});
 			}
 
@@ -1187,12 +1132,12 @@ void FDeferredShadingSceneRenderer::ComputeVolumetricFog(FRHICommandListImmediat
 					const FIntVector NumGroups = FIntVector::DivideAndRoundUp(VolumetricFogGridSize, VolumetricFogIntegrationGroupSize);
 
 					auto ComputeShader = View.ShaderMap->GetShader< FVolumetricFogFinalIntegrationCS >();
-					RHICmdList.SetComputeShader(ComputeShader->GetComputeShader());
+					RHICmdList.SetComputeShader(ComputeShader.GetComputeShader());
 					ComputeShader->SetParameters(RHICmdList, View, IntegrationData);
 
-					SetShaderParameters(RHICmdList, ComputeShader, ComputeShader->GetComputeShader(), *PassParameters);
-					DispatchComputeShader(RHICmdList, ComputeShader, NumGroups.X, NumGroups.Y, 1);
-					UnsetShaderUAVs(RHICmdList, ComputeShader, ComputeShader->GetComputeShader());
+					SetShaderParameters(RHICmdList, ComputeShader, ComputeShader.GetComputeShader(), *PassParameters);
+					DispatchComputeShader(RHICmdList, ComputeShader.GetShader(), NumGroups.X, NumGroups.Y, 1);
+					UnsetShaderUAVs(RHICmdList, ComputeShader, ComputeShader.GetComputeShader());
 				});
 			}
 
diff --git a/Engine/Source/Runtime/Renderer/Private/VolumetricFog.h b/Engine/Source/Runtime/Renderer/Private/VolumetricFog.h
--- a/Engine/Source/Runtime/Renderer/Private/VolumetricFog.h
+++ b/Engine/Source/Runtime/Renderer/Private/VolumetricFog.h
@@ -10,20 +10,20 @@
 #include "SceneView.h"
 #include "SceneRendering.h"
 
-inline bool DoesPlatformSupportVolumetricFog(EShaderPlatform Platform)
+inline bool DoesPlatformSupportVolumetricFog(const FStaticShaderPlatform Platform)
 {
 	return Platform == SP_PCD3D_SM5 || Platform == SP_PS4 || Platform == SP_XBOXONE_D3D12
 		|| Platform == SP_METAL_SM5 || Platform == SP_METAL_SM5_NOTESS
 		|| IsVulkanSM5Platform(Platform)
-		|| FDataDrivenShaderPlatformInfo::GetInfo(Platform).bSupportsVolumetricFog;
+		|| FDataDrivenShaderPlatformInfo::GetSupportsVolumetricFog(Platform);
 }
 
-inline bool DoesPlatformSupportVolumetricFogVoxelization(EShaderPlatform Platform)
+inline bool DoesPlatformSupportVolumetricFogVoxelization(const FStaticShaderPlatform Platform)
 {
 	return Platform == SP_PCD3D_SM5 || Platform == SP_PS4 || Platform == SP_XBOXONE_D3D12
 		|| Platform == SP_METAL_SM5 || Platform == SP_METAL_SM5_NOTESS
 		|| IsVulkanSM5Platform(Platform)
-		|| FDataDrivenShaderPlatformInfo::GetInfo(Platform).bSupportsVolumetricFog;
+		|| FDataDrivenShaderPlatformInfo::GetSupportsVolumetricFog(Platform);
 }
 
 extern bool ShouldRenderVolumetricFog(const FScene* Scene, const FSceneViewFamily& ViewFamily);
diff --git a/Engine/Source/Runtime/Renderer/Private/VolumetricFogLightFunction.cpp b/Engine/Source/Runtime/Renderer/Private/VolumetricFogLightFunction.cpp
--- a/Engine/Source/Runtime/Renderer/Private/VolumetricFogLightFunction.cpp
+++ b/Engine/Source/Runtime/Renderer/Private/VolumetricFogLightFunction.cpp
@@ -28,7 +28,7 @@ public:
 
 	static bool ShouldCompilePermutation(const FMaterialShaderPermutationParameters& Parameters)
 	{
-		return Parameters.Material->IsLightFunction() && DoesPlatformSupportVolumetricFog(Parameters.Platform);
+		return Parameters.MaterialParameters.MaterialDomain == MD_LightFunction && DoesPlatformSupportVolumetricFog(Parameters.Platform);
 	}
 
 	FVolumetricFogLightFunctionPS() {}
@@ -50,7 +50,7 @@ public:
 		FVector2D LightFunctionTexelSizeValue,
 		const FMatrix& ShadowToWorldValue)
 	{
-		FRHIPixelShader* ShaderRHI = GetPixelShader();
+		FRHIPixelShader* ShaderRHI = RHICmdList.GetBoundPixelShader();
 
 		FMaterialShader::SetParameters(RHICmdList, ShaderRHI, MaterialProxy, *MaterialProxy->GetMaterial(View.GetFeatureLevel()), View, View.ViewUniformBuffer, ESceneTextureSetupMode::None);
 
@@ -76,23 +76,12 @@ public:
 		SetShaderValue(RHICmdList, ShaderRHI, ShadowToWorld, ShadowToWorldValue);
 	}
 
-	virtual bool Serialize(FArchive& Ar) override
-	{
-		bool bShaderHasOutdatedParameters = FMaterialShader::Serialize(Ar);
-		Ar << LightFunctionParameters;
-		Ar << LightFunctionParameters2;
-		Ar << LightFunctionWorldToLight;
-		Ar << LightFunctionTexelSize;
-		Ar << ShadowToWorld;
-		return bShaderHasOutdatedParameters;
-	}
-
 private:
-	FLightFunctionSharedParameters LightFunctionParameters;
-	FShaderParameter LightFunctionParameters2;
-	FShaderParameter LightFunctionWorldToLight;
-	FShaderParameter LightFunctionTexelSize;
-	FShaderParameter ShadowToWorld;
+	LAYOUT_FIELD(FLightFunctionSharedParameters, LightFunctionParameters);
+	LAYOUT_FIELD(FShaderParameter, LightFunctionParameters2);
+	LAYOUT_FIELD(FShaderParameter, LightFunctionWorldToLight);
+	LAYOUT_FIELD(FShaderParameter, LightFunctionTexelSize);
+	LAYOUT_FIELD(FShaderParameter, ShadowToWorld);
 };
 
 IMPLEMENT_MATERIAL_SHADER_TYPE(,FVolumetricFogLightFunctionPS,TEXT("/Engine/Private/VolumetricFogLightFunction.usf"),TEXT("Main"),SF_Pixel);
@@ -243,11 +232,11 @@ void FDeferredShadingSceneRenderer::RenderLightFunctionForVolumetricFog(
 
 				const FMaterialShaderMap* MaterialShaderMap = Material->GetRenderingThreadShaderMap();
 				TShaderMapRef<FPostProcessVS> VertexShader(View.ShaderMap);
-				FVolumetricFogLightFunctionPS* PixelShader = MaterialShaderMap->GetShader<FVolumetricFogLightFunctionPS>();
+				TShaderRef<FVolumetricFogLightFunctionPS> PixelShader = MaterialShaderMap->GetShader<FVolumetricFogLightFunctionPS>();
 
 				GraphicsPSOInit.BoundShaderState.VertexDeclarationRHI = GFilterVertexDeclaration.VertexDeclarationRHI;
-				GraphicsPSOInit.BoundShaderState.VertexShaderRHI = GETSAFERHISHADER_VERTEX(*VertexShader);
-				GraphicsPSOInit.BoundShaderState.PixelShaderRHI = GETSAFERHISHADER_PIXEL(PixelShader);
+				GraphicsPSOInit.BoundShaderState.VertexShaderRHI = VertexShader.GetVertexShader();
+				GraphicsPSOInit.BoundShaderState.PixelShaderRHI = PixelShader.GetPixelShader();
 
 				SetGraphicsPipelineState(RHICmdList, GraphicsPSOInit);
 
@@ -262,7 +251,7 @@ void FDeferredShadingSceneRenderer::RenderLightFunctionForVolumetricFog(
 					LightFunctionResolution.X, LightFunctionResolution.Y,
 					LightFunctionResolution,
 					LightFunctionResolution,
-					*VertexShader);
+					VertexShader);
 			});
 		}
 	}
diff --git a/Engine/Source/Runtime/Renderer/Private/VolumetricFogShared.h b/Engine/Source/Runtime/Renderer/Private/VolumetricFogShared.h
--- a/Engine/Source/Runtime/Renderer/Private/VolumetricFogShared.h
+++ b/Engine/Source/Runtime/Renderer/Private/VolumetricFogShared.h
@@ -39,6 +39,7 @@ struct FVolumetricFogIntegrationParameterData
 /**  */
 class FVolumetricFogIntegrationParameters
 {
+	DECLARE_TYPE_LAYOUT(FVolumetricFogIntegrationParameters, NonVirtual);
 public:
 
 	static void ModifyCompilationEnvironment(EShaderPlatform Platform, FShaderCompilerEnvironment& OutEnvironment)
@@ -116,18 +117,19 @@ public:
 	}
 
 private:
-
-	FRWShaderParameter VBufferA;
-	FRWShaderParameter VBufferB;
-	FRWShaderParameter LightScattering;
-	FRWShaderParameter IntegratedLightScattering;
-	FShaderResourceParameter IntegratedLightScatteringSampler;
-	FShaderUniformBufferParameter VolumetricFogData;
-	FShaderParameter UnjitteredClipToTranslatedWorld;
-	FShaderParameter UnjitteredPrevWorldToClip;
-	FShaderParameter FrameJitterOffsets;
-	FShaderParameter HistoryWeight;
-	FShaderParameter HistoryMissSuperSampleCount;
+	
+		LAYOUT_FIELD(FRWShaderParameter, VBufferA)
+		LAYOUT_FIELD(FRWShaderParameter, VBufferB)
+		LAYOUT_FIELD(FRWShaderParameter, LightScattering)
+		LAYOUT_FIELD(FRWShaderParameter, IntegratedLightScattering)
+		LAYOUT_FIELD(FShaderResourceParameter, IntegratedLightScatteringSampler)
+		LAYOUT_FIELD(FShaderUniformBufferParameter, VolumetricFogData)
+		LAYOUT_FIELD(FShaderParameter, UnjitteredClipToTranslatedWorld)
+		LAYOUT_FIELD(FShaderParameter, UnjitteredPrevWorldToClip)
+		LAYOUT_FIELD(FShaderParameter, FrameJitterOffsets)
+		LAYOUT_FIELD(FShaderParameter, HistoryWeight)
+		LAYOUT_FIELD(FShaderParameter, HistoryMissSuperSampleCount)
+	
 };
 
 inline int32 ComputeZSliceFromDepth(float SceneDepth, FVector GridZParams)
diff --git a/Engine/Source/Runtime/Renderer/Private/VolumetricFogVoxelization.cpp b/Engine/Source/Runtime/Renderer/Private/VolumetricFogVoxelization.cpp
--- a/Engine/Source/Runtime/Renderer/Private/VolumetricFogVoxelization.cpp
+++ b/Engine/Source/Runtime/Renderer/Private/VolumetricFogVoxelization.cpp
@@ -151,7 +151,8 @@ public:
 
 class FVoxelizeVolumeVS : public FMeshMaterialShader
 {
-    protected:
+	DECLARE_INLINE_TYPE_LAYOUT(FVoxelizeVolumeVS, NonVirtual);
+protected:
 
 	FVoxelizeVolumeVS(const FMeshMaterialShaderType::CompiledShaderInitializerType& Initializer)
 		:	FMeshMaterialShader(Initializer)
@@ -168,7 +169,7 @@ class FVoxelizeVolumeVS : public FMeshMaterialShader
 	{
 		return IsFeatureLevelSupported(Parameters.Platform, ERHIFeatureLevel::SM5) 
 			&& DoesPlatformSupportVolumetricFogVoxelization(Parameters.Platform)
-			&& Parameters.Material->GetMaterialDomain() == MD_Volume;
+			&& Parameters.MaterialParameters.MaterialDomain == MD_Volume;
 	}
 
 	static void ModifyCompilationEnvironment(const FMaterialShaderPermutationParameters& Parameters, FShaderCompilerEnvironment& OutEnvironment)
@@ -182,13 +183,6 @@ class FVoxelizeVolumeVS : public FMeshMaterialShader
 
 public:
 
-	virtual bool Serialize(FArchive& Ar)
-	{		
-		bool bShaderHasOutdatedParameters = FMeshMaterialShader::Serialize(Ar);
-		Ar << VoxelizationPassIndex;
-		return bShaderHasOutdatedParameters;
-	}
-
 	void GetShaderBindings(
 		const FScene* Scene,
 		ERHIFeatureLevel::Type FeatureLevel,
@@ -207,8 +201,7 @@ public:
 	}
 
 protected:
-
-	FShaderParameter VoxelizationPassIndex;
+	LAYOUT_FIELD(FShaderParameter, VoxelizationPassIndex);
 };
 
 enum EVoxelizeShapeMode
@@ -253,7 +246,7 @@ IMPLEMENT_MATERIAL_SHADER_TYPE(template<>,TVoxelizeVolumeVS<VMode_Object_Box>,TE
 
 class FVoxelizeVolumeGS : public FMeshMaterialShader
 {
-	DECLARE_SHADER_TYPE(FVoxelizeVolumeGS,MeshMaterial);
+	DECLARE_INLINE_TYPE_LAYOUT(FVoxelizeVolumeGS, NonVirtual);
 
 protected:
 
@@ -273,7 +266,7 @@ protected:
 		return IsFeatureLevelSupported(Parameters.Platform, ERHIFeatureLevel::SM5) 
 			&& RHISupportsGeometryShaders(Parameters.Platform)
 			&& DoesPlatformSupportVolumetricFogVoxelization(Parameters.Platform)
-			&& Parameters.Material->GetMaterialDomain() == MD_Volume;
+			&& Parameters.MaterialParameters.MaterialDomain == MD_Volume;
 	}
 
 	static void ModifyCompilationEnvironment(const FMaterialShaderPermutationParameters& Parameters, FShaderCompilerEnvironment& OutEnvironment)
@@ -283,14 +276,6 @@ protected:
 	}
 
 public:
-	
-	virtual bool Serialize(FArchive& Ar)
-	{		
-		bool bShaderHasOutdatedParameters = FMeshMaterialShader::Serialize(Ar);
-		Ar << VoxelizationPassIndex;
-		return bShaderHasOutdatedParameters;
-	}
-
 	void GetShaderBindings(
 		const FScene* Scene,
 		ERHIFeatureLevel::Type FeatureLevel,
@@ -307,8 +292,8 @@ public:
 	}
 	
 protected:
-
-	FShaderParameter VoxelizationPassIndex;
+	
+	LAYOUT_FIELD(FShaderParameter, VoxelizationPassIndex);
 };
 
 template<EVoxelizeShapeMode Mode>
@@ -347,7 +332,7 @@ IMPLEMENT_MATERIAL_SHADER_TYPE(template<>,TVoxelizeVolumeGS<VMode_Object_Box>,TE
 
 class FVoxelizeVolumePS : public FMeshMaterialShader
 {
-	DECLARE_SHADER_TYPE(FVoxelizeVolumePS,MeshMaterial);
+	DECLARE_INLINE_TYPE_LAYOUT(FVoxelizeVolumePS, NonVirtual);
 
 protected:
 
@@ -365,8 +350,11 @@ protected:
 	{
 		return IsFeatureLevelSupported(Parameters.Platform, ERHIFeatureLevel::SM5) 
 			&& DoesPlatformSupportVolumetricFogVoxelization(Parameters.Platform)
-			&& Parameters.Material->GetMaterialDomain() == MD_Volume;
+			&& Parameters.MaterialParameters.MaterialDomain == MD_Volume;
 	}
+
+	
+	
 };
 
 template<EVoxelizeShapeMode Mode>
@@ -447,7 +435,8 @@ void FVoxelizeVolumeMeshProcessor::AddMeshBatch(const FMeshBatch& RESTRICT MeshB
 	const FMaterialRenderProxy* FallbackMaterialRenderProxyPtr = nullptr;
 	const FMaterial& Material = MeshBatch.MaterialRenderProxy->GetMaterialWithFallback(FeatureLevel, FallbackMaterialRenderProxyPtr);
 
-	const ERasterizerFillMode MeshFillMode = ComputeMeshFillMode(MeshBatch, Material);
+	const FMeshDrawingPolicyOverrideSettings OverrideSettings = ComputeMeshOverrideSettings(MeshBatch);
+	const ERasterizerFillMode MeshFillMode = ComputeMeshFillMode(MeshBatch, Material, OverrideSettings);
 	const ERasterizerCullMode MeshCullMode = CM_None;
 
 	const FMaterialRenderProxy& MaterialRenderProxy = FallbackMaterialRenderProxyPtr ? *FallbackMaterialRenderProxyPtr : *MeshBatch.MaterialRenderProxy;
diff --git a/Engine/Source/Runtime/Renderer/Private/WideCustomResolveShaders.h b/Engine/Source/Runtime/Renderer/Private/WideCustomResolveShaders.h
--- a/Engine/Source/Runtime/Renderer/Private/WideCustomResolveShaders.h
+++ b/Engine/Source/Runtime/Renderer/Private/WideCustomResolveShaders.h
@@ -18,12 +18,6 @@ public:
 	{
 	}
 
-	virtual bool Serialize(FArchive& Ar) override
-	{
-		bool bShaderHasOutdatedParameters = FGlobalShader::Serialize(Ar);
-		return bShaderHasOutdatedParameters;
-	}
-
 	static bool ShouldCompilePermutation(const FGlobalShaderPermutationParameters& Parameters)
 	{
 		return IsFeatureLevelSupported(Parameters.Platform, ERHIFeatureLevel::SM5);
@@ -50,15 +44,6 @@ public:
 		ResolveOrigin.Bind(Initializer.ParameterMap, TEXT("ResolveOrigin"));
 	}
 
-	virtual bool Serialize(FArchive& Ar) override
-	{
-		bool bShaderHasOutdatedParameters = FGlobalShader::Serialize(Ar);
-		Ar << Tex;
-		Ar << FMaskTex;
-		Ar << ResolveOrigin;
-		return bShaderHasOutdatedParameters;
-	}
-
 	static bool ShouldCompilePermutation(const FGlobalShaderPermutationParameters& Parameters)
 	{
 		return IsFeatureLevelSupported(Parameters.Platform, ERHIFeatureLevel::SM5);
@@ -66,7 +51,7 @@ public:
 
 	void SetParameters(FRHICommandList& RHICmdList, FRHITexture* Texture2DMS, FRHIShaderResourceView* FmaskSRV, FIntPoint Origin)
 	{
-		FRHIPixelShader* PixelShaderRHI = GetPixelShader();
+		FRHIPixelShader* PixelShaderRHI = RHICmdList.GetBoundPixelShader();
 		SetTextureParameter(RHICmdList, PixelShaderRHI, Tex, Texture2DMS);
 		if (MSAASampleCount > 0)
 		{
@@ -84,9 +69,9 @@ public:
 	}
 
 protected:
-	FShaderResourceParameter Tex;
-	FShaderResourceParameter FMaskTex;
-	FShaderParameter ResolveOrigin;
+	LAYOUT_FIELD(FShaderResourceParameter, Tex);
+	LAYOUT_FIELD(FShaderResourceParameter, FMaskTex);
+	LAYOUT_FIELD(FShaderParameter, ResolveOrigin);
 };
 
 extern void ResolveFilterWide(
diff --git a/Engine/Source/Runtime/Renderer/Public/GlobalDistanceFieldParameters.h b/Engine/Source/Runtime/Renderer/Public/GlobalDistanceFieldParameters.h
--- a/Engine/Source/Runtime/Renderer/Public/GlobalDistanceFieldParameters.h
+++ b/Engine/Source/Runtime/Renderer/Public/GlobalDistanceFieldParameters.h
@@ -34,6 +34,7 @@ public:
 
 class FGlobalDistanceFieldParameters
 {
+	DECLARE_INLINE_TYPE_LAYOUT(FGlobalDistanceFieldParameters, NonVirtual);
 public:
 	void Bind(const FShaderParameterMap& ParameterMap)
 	{
@@ -94,17 +95,19 @@ public:
 	}
 
 private:
-	FShaderResourceParameter GlobalDistanceFieldTexture0;
-	FShaderResourceParameter GlobalDistanceFieldTexture1;
-	FShaderResourceParameter GlobalDistanceFieldTexture2;
-	FShaderResourceParameter GlobalDistanceFieldTexture3;
-	FShaderResourceParameter GlobalDistanceFieldSampler0;
-	FShaderResourceParameter GlobalDistanceFieldSampler1;
-	FShaderResourceParameter GlobalDistanceFieldSampler2;
-	FShaderResourceParameter GlobalDistanceFieldSampler3;
-	FShaderParameter GlobalVolumeCenterAndExtent;
-	FShaderParameter GlobalVolumeWorldToUVAddAndMul;
-	FShaderParameter GlobalVolumeDimension;
-	FShaderParameter GlobalVolumeTexelSize;
-	FShaderParameter MaxGlobalDistance;
+	
+		LAYOUT_FIELD(FShaderResourceParameter, GlobalDistanceFieldTexture0)
+		LAYOUT_FIELD(FShaderResourceParameter, GlobalDistanceFieldTexture1)
+		LAYOUT_FIELD(FShaderResourceParameter, GlobalDistanceFieldTexture2)
+		LAYOUT_FIELD(FShaderResourceParameter, GlobalDistanceFieldTexture3)
+		LAYOUT_FIELD(FShaderResourceParameter, GlobalDistanceFieldSampler0)
+		LAYOUT_FIELD(FShaderResourceParameter, GlobalDistanceFieldSampler1)
+		LAYOUT_FIELD(FShaderResourceParameter, GlobalDistanceFieldSampler2)
+		LAYOUT_FIELD(FShaderResourceParameter, GlobalDistanceFieldSampler3)
+		LAYOUT_FIELD(FShaderParameter, GlobalVolumeCenterAndExtent)
+		LAYOUT_FIELD(FShaderParameter, GlobalVolumeWorldToUVAddAndMul)
+		LAYOUT_FIELD(FShaderParameter, GlobalVolumeDimension)
+		LAYOUT_FIELD(FShaderParameter, GlobalVolumeTexelSize)
+		LAYOUT_FIELD(FShaderParameter, MaxGlobalDistance)
+	
 };
diff --git a/Engine/Source/Runtime/Renderer/Public/HairStrandsInterface.h b/Engine/Source/Runtime/Renderer/Public/HairStrandsInterface.h
--- a/Engine/Source/Runtime/Renderer/Public/HairStrandsInterface.h
+++ b/Engine/Source/Runtime/Renderer/Public/HairStrandsInterface.h
@@ -371,7 +371,7 @@ RENDERER_API void RunProjection(
 
 typedef void (*TBindingProcess)(FRHICommandListImmediate& RHICmdList, void* Asset);
 RENDERER_API void EnqueueGroomBindingQuery(void* Asset, TBindingProcess BindingProcess);
-void RunHairStrandsBindingQueries(FRHICommandListImmediate& RHICmdList, TShaderMap<FGlobalShaderType>* ShaderMap);
+void RunHairStrandsBindingQueries(FRHICommandListImmediate& RHICmdList, class FGlobalShaderMap* ShaderMap);
 bool HasHairStrandsProjectionQuery(EShaderPlatform Platform);
 
 RENDERER_API FHairStrandsProjectionMeshData ExtractMeshData(class FSkeletalMeshRenderData* RenderData);
