Copying Private-LoadTimes-4.24 stream to Main. Biggest changes are in Materials/Shader memory freezing.

diff --git a/Engine/Source/Runtime/Renderer/Public/HdrCustomResolveShaders.h b/Engine/Source/Runtime/Renderer/Public/HdrCustomResolveShaders.h
--- a/Engine/Source/Runtime/Renderer/Public/HdrCustomResolveShaders.h
+++ b/Engine/Source/Runtime/Renderer/Public/HdrCustomResolveShaders.h
@@ -19,12 +19,6 @@ public:
 	{
 	}
 
-	virtual bool Serialize(FArchive& Ar) override
-	{
-		bool bShaderHasOutdatedParameters = FGlobalShader::Serialize(Ar);
-		return bShaderHasOutdatedParameters;
-	}
-
 	static bool ShouldCompilePermutation(const FGlobalShaderPermutationParameters& Parameters)
 	{
 		return IsFeatureLevelSupported(Parameters.Platform, ERHIFeatureLevel::ES2);
@@ -42,13 +36,6 @@ public:
 		Tex.Bind(Initializer.ParameterMap, TEXT("Tex"), SPF_Mandatory);
 	}
 
-	virtual bool Serialize(FArchive& Ar) override
-	{
-		bool bShaderHasOutdatedParameters = FGlobalShader::Serialize(Ar);
-		Ar << Tex;
-		return bShaderHasOutdatedParameters;
-	}
-
 	static bool ShouldCompilePermutation(const FGlobalShaderPermutationParameters& Parameters)
 	{
 		return IsFeatureLevelSupported(Parameters.Platform, ERHIFeatureLevel::ES2);
@@ -56,7 +43,7 @@ public:
 
 	void SetParameters(FRHICommandList& RHICmdList, FRHITexture* Texture2DMS)
 	{
-		FRHIPixelShader* PixelShaderRHI = GetPixelShader();
+		FRHIPixelShader* PixelShaderRHI = RHICmdList.GetBoundPixelShader();
 		SetTextureParameter(RHICmdList, PixelShaderRHI, Tex, Texture2DMS);
 	}
 
@@ -67,7 +54,7 @@ public:
 	}
 
 protected:
-	FShaderResourceParameter Tex;
+	LAYOUT_FIELD(FShaderResourceParameter, Tex);
 };
 
 class FHdrCustomResolve4xPS : public FGlobalShader
@@ -81,13 +68,6 @@ public:
 		Tex.Bind(Initializer.ParameterMap, TEXT("Tex"), SPF_Mandatory);
 	}
 
-	virtual bool Serialize(FArchive& Ar) override
-	{
-		bool bShaderHasOutdatedParameters = FGlobalShader::Serialize(Ar);
-		Ar << Tex;
-		return bShaderHasOutdatedParameters;
-	}
-
 	static bool ShouldCompilePermutation(const FGlobalShaderPermutationParameters& Parameters)
 	{
 		return IsFeatureLevelSupported(Parameters.Platform, ERHIFeatureLevel::ES2);
@@ -95,7 +75,7 @@ public:
 
 	void SetParameters(FRHICommandList& RHICmdList, FRHITexture* Texture2DMS)
 	{
-		FRHIPixelShader* PixelShaderRHI = GetPixelShader();
+		FRHIPixelShader* PixelShaderRHI = RHICmdList.GetBoundPixelShader();
 		SetTextureParameter(RHICmdList, PixelShaderRHI, Tex, Texture2DMS);
 	}
 
@@ -106,7 +86,7 @@ public:
 	}
 
 protected:
-	FShaderResourceParameter Tex;
+	LAYOUT_FIELD(FShaderResourceParameter, Tex);
 };
 
 
@@ -121,13 +101,6 @@ public:
 		Tex.Bind(Initializer.ParameterMap, TEXT("Tex"), SPF_Mandatory);
 	}
 
-	virtual bool Serialize(FArchive& Ar) override
-	{
-		bool bShaderHasOutdatedParameters = FGlobalShader::Serialize(Ar);
-		Ar << Tex;
-		return bShaderHasOutdatedParameters;
-	}
-
 	static bool ShouldCompilePermutation(const FGlobalShaderPermutationParameters& Parameters)
 	{
 		return IsFeatureLevelSupported(Parameters.Platform, ERHIFeatureLevel::ES2);
@@ -135,7 +108,7 @@ public:
 
 	void SetParameters(FRHICommandList& RHICmdList, FRHITexture* Texture2DMS)
 	{
-		FRHIPixelShader* PixelShaderRHI = GetPixelShader();
+		FRHIPixelShader* PixelShaderRHI = RHICmdList.GetBoundPixelShader();
 		SetTextureParameter(RHICmdList, PixelShaderRHI, Tex, Texture2DMS);
 	}
 
@@ -146,7 +119,7 @@ public:
 	}
 
 protected:
-	FShaderResourceParameter Tex;
+	LAYOUT_FIELD(FShaderResourceParameter, Tex);
 };
 
 class FHdrCustomResolveFMask2xPS : public FGlobalShader
@@ -161,14 +134,6 @@ public:
 		FMaskTex.Bind(Initializer.ParameterMap, TEXT("FMaskTex"), SPF_Optional);
 	}
 
-	virtual bool Serialize(FArchive& Ar) override
-	{
-		bool bShaderHasOutdatedParameters = FGlobalShader::Serialize(Ar);
-		Ar << Tex;
-		Ar << FMaskTex;
-		return bShaderHasOutdatedParameters;
-	}
-
 	static bool ShouldCompilePermutation(const FGlobalShaderPermutationParameters& Parameters)
 	{
 		return IsFeatureLevelSupported(Parameters.Platform, ERHIFeatureLevel::ES2);
@@ -176,7 +141,7 @@ public:
 
 	void SetParameters(FRHICommandList& RHICmdList, FRHITexture* Texture2DMS, FRHIShaderResourceView* FMaskSRV)
 	{
-		FRHIPixelShader* PixelShaderRHI = GetPixelShader();
+		FRHIPixelShader* PixelShaderRHI = RHICmdList.GetBoundPixelShader();
 		SetTextureParameter(RHICmdList, PixelShaderRHI, Tex, Texture2DMS);
 		SetSRVParameter(RHICmdList, PixelShaderRHI, FMaskTex, FMaskSRV);
 	}
@@ -189,8 +154,8 @@ public:
 	}
 
 protected:
-	FShaderResourceParameter Tex;
-	FShaderResourceParameter FMaskTex;
+	LAYOUT_FIELD(FShaderResourceParameter, Tex);
+	LAYOUT_FIELD(FShaderResourceParameter, FMaskTex);
 };
 
 class FHdrCustomResolveFMask4xPS : public FGlobalShader
@@ -205,14 +170,6 @@ public:
 		FMaskTex.Bind(Initializer.ParameterMap, TEXT("FMaskTex"), SPF_Optional);
 	}
 
-	virtual bool Serialize(FArchive& Ar) override
-	{
-		bool bShaderHasOutdatedParameters = FGlobalShader::Serialize(Ar);
-		Ar << Tex;
-		Ar << FMaskTex;
-		return bShaderHasOutdatedParameters;
-	}
-
 	static bool ShouldCompilePermutation(const FGlobalShaderPermutationParameters& Parameters)
 	{
 		return IsFeatureLevelSupported(Parameters.Platform, ERHIFeatureLevel::ES2);
@@ -220,7 +177,7 @@ public:
 
 	void SetParameters(FRHICommandList& RHICmdList, FRHITexture* Texture2DMS, FRHIShaderResourceView* FMaskSRV)
 	{
-		FRHIPixelShader* PixelShaderRHI = GetPixelShader();
+		FRHIPixelShader* PixelShaderRHI = RHICmdList.GetBoundPixelShader();
 		SetTextureParameter(RHICmdList, PixelShaderRHI, Tex, Texture2DMS);
 		SetSRVParameter(RHICmdList, PixelShaderRHI, FMaskTex, FMaskSRV);
 	}
@@ -233,8 +190,8 @@ public:
 	}
 
 protected:
-	FShaderResourceParameter Tex;
-	FShaderResourceParameter FMaskTex;
+	LAYOUT_FIELD(FShaderResourceParameter, Tex);
+	LAYOUT_FIELD(FShaderResourceParameter, FMaskTex);
 };
 
 
@@ -250,14 +207,6 @@ public:
 		FMaskTex.Bind(Initializer.ParameterMap, TEXT("FMaskTex"), SPF_Optional);
 	}
 
-	virtual bool Serialize(FArchive& Ar) override
-	{
-		bool bShaderHasOutdatedParameters = FGlobalShader::Serialize(Ar);
-		Ar << Tex;
-		Ar << FMaskTex;
-		return bShaderHasOutdatedParameters;
-	}
-
 	static bool ShouldCompilePermutation(const FGlobalShaderPermutationParameters& Parameters)
 	{
 		return IsFeatureLevelSupported(Parameters.Platform, ERHIFeatureLevel::ES2);
@@ -265,7 +214,7 @@ public:
 
 	void SetParameters(FRHICommandList& RHICmdList, FRHITexture* Texture2DMS, FRHIShaderResourceView* FMaskSRV)
 	{
-		FRHIPixelShader* PixelShaderRHI = GetPixelShader();
+		FRHIPixelShader* PixelShaderRHI = RHICmdList.GetBoundPixelShader();
 		SetTextureParameter(RHICmdList, PixelShaderRHI, Tex, Texture2DMS);
 		SetSRVParameter(RHICmdList, PixelShaderRHI, FMaskTex, FMaskSRV);
 	}
@@ -278,6 +227,6 @@ public:
 	}
 
 protected:
-	FShaderResourceParameter Tex;
-	FShaderResourceParameter FMaskTex;
+	LAYOUT_FIELD(FShaderResourceParameter, Tex);
+	LAYOUT_FIELD(FShaderResourceParameter, FMaskTex);
 };
diff --git a/Engine/Source/Runtime/Renderer/Public/MaterialShader.h b/Engine/Source/Runtime/Renderer/Public/MaterialShader.h
--- a/Engine/Source/Runtime/Renderer/Public/MaterialShader.h
+++ b/Engine/Source/Runtime/Renderer/Public/MaterialShader.h
@@ -25,25 +25,13 @@ template<typename TBufferStruct> class TUniformBufferRef;
  */
 class FDebugUniformExpressionSet
 {
+	DECLARE_TYPE_LAYOUT(FDebugUniformExpressionSet, NonVirtual);
 public:
-	/** The number of each type of expression contained in the set. */
-	int32 NumVectorExpressions;
-	int32 NumScalarExpressions;
-	int32 Num2DTextureExpressions;
-	int32 NumCubeTextureExpressions;
-	int32 Num2DArrayTextureExpressions;
-	int32 NumVolumeTextureExpressions;
-	int32 NumVirtualTextureExpressions;
-
 	FDebugUniformExpressionSet()
 		: NumVectorExpressions(0)
 		, NumScalarExpressions(0)
-		, Num2DTextureExpressions(0)
-		, NumCubeTextureExpressions(0)
-		, Num2DArrayTextureExpressions(0)
-		, NumVolumeTextureExpressions(0)
-		, NumVirtualTextureExpressions(0)
 	{
+		FMemory::Memzero(NumTextureExpressions);
 	}
 
 	explicit FDebugUniformExpressionSet(const FUniformExpressionSet& InUniformExpressionSet)
@@ -54,69 +42,61 @@ public:
 	/** Initialize from a uniform expression set. */
 	void InitFromExpressionSet(const FUniformExpressionSet& InUniformExpressionSet)
 	{
-		NumVectorExpressions = InUniformExpressionSet.UniformVectorExpressions.Num();
-		NumScalarExpressions = InUniformExpressionSet.UniformScalarExpressions.Num();
-		Num2DTextureExpressions = InUniformExpressionSet.Uniform2DTextureExpressions.Num();
-		Num2DArrayTextureExpressions = InUniformExpressionSet.Uniform2DArrayTextureExpressions.Num();
-		NumCubeTextureExpressions = InUniformExpressionSet.UniformCubeTextureExpressions.Num();
-		NumVolumeTextureExpressions = InUniformExpressionSet.UniformVolumeTextureExpressions.Num();
-		NumVirtualTextureExpressions = InUniformExpressionSet.UniformVirtualTextureExpressions.Num();
+		NumVectorExpressions = InUniformExpressionSet.UniformVectorPreshaders.Num();
+		NumScalarExpressions = InUniformExpressionSet.UniformScalarPreshaders.Num();
+		for (uint32 TypeIndex = 0u; TypeIndex < NumMaterialTextureParameterTypes; ++TypeIndex)
+		{
+			NumTextureExpressions[TypeIndex] = InUniformExpressionSet.UniformTextureParameters[TypeIndex].Num();
+		}
 	}
 
 	/** Returns true if the number of uniform expressions matches those with which the debug set was initialized. */
 	bool Matches(const FUniformExpressionSet& InUniformExpressionSet) const
 	{
-		return NumVectorExpressions == InUniformExpressionSet.UniformVectorExpressions.Num()
-			&& NumScalarExpressions == InUniformExpressionSet.UniformScalarExpressions.Num()
-			&& Num2DTextureExpressions == InUniformExpressionSet.Uniform2DTextureExpressions.Num()
-			&& NumCubeTextureExpressions == InUniformExpressionSet.UniformCubeTextureExpressions.Num()
-			&& Num2DArrayTextureExpressions == InUniformExpressionSet.Uniform2DArrayTextureExpressions.Num()
-			&& NumVolumeTextureExpressions == InUniformExpressionSet.UniformVolumeTextureExpressions.Num()
-			&& NumVirtualTextureExpressions == InUniformExpressionSet.UniformVirtualTextureExpressions.Num();
+		for (uint32 TypeIndex = 0u; TypeIndex < NumMaterialTextureParameterTypes; ++TypeIndex)
+		{
+			if (NumTextureExpressions[TypeIndex] != InUniformExpressionSet.UniformTextureParameters[TypeIndex].Num())
+			{
+				return false;
+			}
+		}
+		return NumVectorExpressions == InUniformExpressionSet.UniformVectorPreshaders.Num() && NumScalarExpressions == InUniformExpressionSet.UniformScalarPreshaders.Num();
 	}
+	
+	/** The number of each type of expression contained in the set. */
+	LAYOUT_FIELD(int32, NumVectorExpressions);
+	LAYOUT_FIELD(int32, NumScalarExpressions);
+	LAYOUT_ARRAY(int32, NumTextureExpressions, NumMaterialTextureParameterTypes);
 };
 
-/** Serialization for debug uniform expression sets. */
-inline FArchive& operator<<(FArchive& Ar, FDebugUniformExpressionSet& DebugExpressionSet)
+struct FMaterialShaderPermutationParameters : public FShaderPermutationParameters
 {
-	Ar << DebugExpressionSet.NumVectorExpressions;
-	Ar << DebugExpressionSet.NumScalarExpressions;
-	Ar << DebugExpressionSet.Num2DTextureExpressions;
-	Ar << DebugExpressionSet.NumCubeTextureExpressions;
-	Ar << DebugExpressionSet.Num2DArrayTextureExpressions;
-	Ar << DebugExpressionSet.NumVolumeTextureExpressions;
-	Ar << DebugExpressionSet.NumVirtualTextureExpressions;
-	return Ar;
-}
+	FMaterialShaderParameters MaterialParameters;
 
+	FMaterialShaderPermutationParameters(EShaderPlatform InPlatform, const FMaterialShaderParameters& InMaterialParameters, int32 InPermutationId)
+		: FShaderPermutationParameters(InPlatform, InPermutationId)
+		, MaterialParameters(InMaterialParameters)
+	{}
+};
 
 /** Base class of all shaders that need material parameters. */
 class RENDERER_API FMaterialShader : public FShader
 {
-	DECLARE_SHADER_TYPE(FMaterialShader, Material);
+	DECLARE_TYPE_LAYOUT(FMaterialShader, NonVirtual);
 public:
+	using FPermutationParameters = FMaterialShaderPermutationParameters;
+	using ShaderMetaType = FMaterialShaderType;
+
 	static FName UniformBufferLayoutName;
 
 	FMaterialShader()
-#if ALLOW_SHADERMAP_DEBUG_DATA
+#if WITH_EDITORONLY_DATA
 		: DebugUniformExpressionUBLayout(FRHIUniformBufferLayout::Zero)
 #endif
-	{
-	}
+	{}
 
 	FMaterialShader(const FMaterialShaderType::CompiledShaderInitializerType& Initializer);
 
-	typedef void (*ModifyCompilationEnvironmentType)(EShaderPlatform, const FMaterial*, FShaderCompilerEnvironment&);
-
-	static void ModifyCompilationEnvironment(const FMaterialShaderPermutationParameters& Parameters, FShaderCompilerEnvironment& OutEnvironment)
-	{
-	}
-
-	static bool ValidateCompiledResult(EShaderPlatform Platform, const TArray<FMaterial*>& Materials, const FShaderParameterMap& ParameterMap, TArray<FString>& OutError)
-	{
-		return true;
-	}
-
 	FRHIUniformBuffer* GetParameterCollectionBuffer(const FGuid& Id, const FSceneInterface* SceneInterface) const;
 
 	template<typename ShaderRHIParamRef>
@@ -136,9 +116,6 @@ public:
 		}
 	}
 
-	static void ModifyCompilationEnvironment(EShaderPlatform Platform, FShaderCompilerEnvironment& OutEnvironment)
-	{ }
-
 	/** Sets pixel parameters that are material specific but not FMeshBatch specific. */
 	template< typename TRHIShader >
 	void SetParametersInner(
@@ -172,29 +149,12 @@ public:
 
 	void GetShaderBindings(
 		const FScene* Scene,
-		ERHIFeatureLevel::Type FeatureLevel,
+		const FStaticFeatureLevel FeatureLevel,
 		const FMaterialRenderProxy& MaterialRenderProxy,
 		const FMaterial& Material,
 		FMeshDrawSingleShaderBindings& ShaderBindings) const;
 
-	// FShader interface.
-	virtual bool Serialize(FArchive& Ar) override;
-	virtual uint32 GetAllocatedSize() const override;
-
-protected:
-
-	FSceneTextureShaderParameters SceneTextureParameters;
-
 private:
-
-	FShaderUniformBufferParameter MaterialUniformBuffer;
-	TArray<FShaderUniformBufferParameter> ParameterCollectionUniformBuffers;
-
-#if ALLOW_SHADERMAP_DEBUG_DATA
-	FDebugUniformExpressionSet	DebugUniformExpressionSet;
-	FRHIUniformBufferLayout		DebugUniformExpressionUBLayout;
-	FString						DebugDescription;
-#endif
 	/** If true, cached uniform expressions are allowed. */
 	static int32 bAllowCachedUniformExpressions;
 	/** Console variable ref to toggle cached uniform expressions. */
@@ -203,4 +163,19 @@ private:
 #if !(UE_BUILD_TEST || UE_BUILD_SHIPPING || !WITH_EDITOR)
 	void VerifyExpressionAndShaderMaps(const FMaterialRenderProxy* MaterialRenderProxy, const FMaterial& Material, const FUniformExpressionCache* UniformExpressionCache) const;
 #endif
+
+	LAYOUT_FIELD(TMemoryImageArray<FShaderUniformBufferParameter>, ParameterCollectionUniformBuffers);
+
+	LAYOUT_FIELD(FShaderUniformBufferParameter, MaterialUniformBuffer);
+
+	// Only needed to avoid unbound parameter error
+	// This texture is bound as an UAV (RWTexture) and so it must be bound together with any RT. So it actually bound but not as part of the material
+	LAYOUT_FIELD(FShaderResourceParameter, VTFeedbackBuffer);
+
+protected:
+	LAYOUT_FIELD(FSceneTextureShaderParameters, SceneTextureParameters);
+
+	LAYOUT_FIELD_EDITORONLY(FDebugUniformExpressionSet, DebugUniformExpressionSet);
+	LAYOUT_FIELD_EDITORONLY(FRHIUniformBufferLayout, DebugUniformExpressionUBLayout);
+	LAYOUT_FIELD_EDITORONLY(FMemoryImageString, DebugDescription);
 };
diff --git a/Engine/Source/Runtime/Renderer/Public/MeshDrawShaderBindings.h b/Engine/Source/Runtime/Renderer/Public/MeshDrawShaderBindings.h
--- a/Engine/Source/Runtime/Renderer/Public/MeshDrawShaderBindings.h
+++ b/Engine/Source/Runtime/Renderer/Public/MeshDrawShaderBindings.h
@@ -13,28 +13,27 @@
 /** Stores the number of each resource type that will need to be bound to a single shader, computed during shader reflection. */
 class FMeshDrawShaderBindingsLayout
 {
-public:
-	EShaderFrequency Frequency : SF_NumBits + 1;
+protected:
 	const FShaderParameterMapInfo& ParameterMapInfo;
 
-	FMeshDrawShaderBindingsLayout(const FShader* Shader) :
-		ParameterMapInfo(Shader->GetParameterMapInfo())
+public:
+
+	FMeshDrawShaderBindingsLayout(const TShaderRef<FShader>& Shader) :
+		ParameterMapInfo(Shader->ParameterMapInfo)
 	{
-		check(Shader);
-		Frequency = (EShaderFrequency)Shader->GetTarget().Frequency;
-		checkSlow((EShaderFrequency)Frequency == (EShaderFrequency)Shader->GetTarget().Frequency);
+		check(Shader.IsValid());
 	}
-	
-	inline uint32 GetHash() const
+
+#if DO_GUARD_SLOW
+	const FShaderParameterMapInfo& GetParameterMapInfo()
 	{
-		uint32 LocalFrequency = Frequency;
-		return FCrc::TypeCrc32(LocalFrequency, 0);
+		return ParameterMapInfo;
 	}
+#endif
 
 	bool operator==(const FMeshDrawShaderBindingsLayout& Rhs) const
 	{
-		return Frequency == Rhs.Frequency
-			&& ParameterMapInfo == Rhs.ParameterMapInfo;
+		return ParameterMapInfo == Rhs.ParameterMapInfo;
 	}
 
 	inline uint32 GetDataSizeBytes() const
@@ -49,7 +48,7 @@ public:
 
 		for (int32 LooseBufferIndex = 0; LooseBufferIndex < ParameterMapInfo.LooseParameterBuffers.Num(); LooseBufferIndex++)
 		{
-			DataSize += ParameterMapInfo.LooseParameterBuffers[LooseBufferIndex].BufferSize;
+			DataSize += ParameterMapInfo.LooseParameterBuffers[LooseBufferIndex].Size;
 		}
 
 		// Align to pointer size so subsequent packed shader bindings will have their pointers aligned
@@ -106,8 +105,8 @@ public:
 
 		if (Parameter.IsBound())
 		{
-			checkf(Value.GetReference(), TEXT("Attempted to set null uniform buffer for type %s on %s"), UniformBufferStructType::StaticStructMetadata.GetStructTypeName(), GetShaderFrequencyString(Frequency));
-			checkfSlow(Value.GetReference()->IsValid(), TEXT("Attempted to set already deleted uniform buffer for type %s on %s"), UniformBufferStructType::StaticStructMetadata.GetStructTypeName(), GetShaderFrequencyString(Frequency));
+			checkf(Value.GetReference(), TEXT("Attempted to set null uniform buffer for type %s"), UniformBufferStructType::StaticStructMetadata.GetStructTypeName());
+			checkfSlow(Value.GetReference()->IsValid(), TEXT("Attempted to set already deleted uniform buffer for type %s"), UniformBufferStructType::StaticStructMetadata.GetStructTypeName());
 			WriteBindingUniformBuffer(Value.GetReference(), Parameter.GetBaseIndex());
 		}
 	}
@@ -119,20 +118,20 @@ public:
 
 		if (Parameter.IsBound())
 		{
-			checkf(Value.GetUniformBufferRHI(), TEXT("Attempted to set null uniform buffer for type %s on %s"), UniformBufferStructType::StaticStructMetadata.GetStructTypeName(), GetShaderFrequencyString(Frequency));
-			checkfSlow(Value.GetUniformBufferRHI()->IsValid(), TEXT("Attempted to set already deleted uniform buffer for type %s on %s"), UniformBufferStructType::StaticStructMetadata.GetStructTypeName(), GetShaderFrequencyString(Frequency));
+			checkf(Value.GetUniformBufferRHI(), TEXT("Attempted to set null uniform buffer for type %s"), UniformBufferStructType::StaticStructMetadata.GetStructTypeName());
+			checkfSlow(Value.GetUniformBufferRHI()->IsValid(), TEXT("Attempted to set already deleted uniform buffer for type %s"), UniformBufferStructType::StaticStructMetadata.GetStructTypeName());
 			WriteBindingUniformBuffer(Value.GetUniformBufferRHI(), Parameter.GetBaseIndex());
 		}
 	}
 
-	void Add(FShaderUniformBufferParameter Parameter, FRHIUniformBuffer* Value)
+	void Add(const FShaderUniformBufferParameter& Parameter, const FRHIUniformBuffer* Value)
 	{
 		checkfSlow(Parameter.IsInitialized(), TEXT("Parameter was not serialized"));
 
 		if (Parameter.IsBound())
 		{
-			checkf(Value, TEXT("Attempted to set null uniform buffer with unknown type on %s"), GetShaderFrequencyString(Frequency));
-			checkfSlow(Value->IsValid(), TEXT("Attempted to set already deleted uniform buffer of type %s on %s"), *Value->GetLayout().GetDebugName().ToString(), GetShaderFrequencyString(Frequency));
+			checkf(Value, TEXT("Attempted to set null uniform buffer"));
+			checkfSlow(Value->IsValid(), TEXT("Attempted to set already deleted uniform buffer of type %s"), *Value->GetLayout().GetDebugName());
 			WriteBindingUniformBuffer(Value, Parameter.GetBaseIndex());
 		}
 	}
@@ -143,8 +142,8 @@ public:
 
 		if (Parameter.IsBound())
 		{
-			checkf(Value, TEXT("Attempted to set null SRV on slot %u of %s"), Parameter.GetBaseIndex(), GetShaderFrequencyString(Frequency));
-			checkfSlow(Value->IsValid(), TEXT("Attempted to set already deleted SRV on slot %u of %s"), Parameter.GetBaseIndex(), GetShaderFrequencyString(Frequency));
+			checkf(Value, TEXT("Attempted to set null SRV on slot %u"), Parameter.GetBaseIndex());
+			checkfSlow(Value->IsValid(), TEXT("Attempted to set already deleted SRV on slot %u"), Parameter.GetBaseIndex());
 			WriteBindingSRV(Value, Parameter.GetBaseIndex());
 		}
 	}
@@ -160,13 +159,13 @@ public:
 
 		if (TextureParameter.IsBound())
 		{
-			checkf(TextureRHI, TEXT("Attempted to set null Texture on slot %u of %s"), TextureParameter.GetBaseIndex(), GetShaderFrequencyString(Frequency));
+			checkf(TextureRHI, TEXT("Attempted to set null Texture on slot %u"), TextureParameter.GetBaseIndex());
 			WriteBindingTexture(TextureRHI, TextureParameter.GetBaseIndex());
 		}
 
 		if (SamplerParameter.IsBound())
 		{
-			checkf(SamplerStateRHI, TEXT("Attempted to set null Sampler on slot %u of %s"), SamplerParameter.GetBaseIndex(), GetShaderFrequencyString(Frequency));
+			checkf(SamplerStateRHI, TEXT("Attempted to set null Sampler on slot %u"), SamplerParameter.GetBaseIndex());
 			WriteBindingSampler(SamplerStateRHI, SamplerParameter.GetBaseIndex());
 		}
 	}
@@ -186,7 +185,7 @@ public:
 			{
 				const FShaderLooseParameterBufferInfo& LooseParameterBuffer = ParameterMapInfo.LooseParameterBuffers[LooseBufferIndex];
 
-				if (LooseParameterBuffer.BufferIndex == Parameter.GetBufferIndex())
+				if (LooseParameterBuffer.BaseIndex == Parameter.GetBufferIndex())
 				{
 					for (int32 LooseParameterIndex = 0; LooseParameterIndex < LooseParameterBuffer.Parameters.Num(); LooseParameterIndex++)
 					{
@@ -207,7 +206,7 @@ public:
 					break;
 				}
 
-				LooseDataOffset += LooseParameterBuffer.BufferSize;
+				LooseDataOffset += LooseParameterBuffer.Size;
 			}
 
 			checkfSlow(bFoundParameter, TEXT("Attempted to set loose parameter at BaseIndex %u, Size %u which was never in the shader's parameter map."), Parameter.GetBaseIndex(), Parameter.GetNumBytes());
@@ -217,9 +216,9 @@ public:
 private:
 	uint8* Data;
 
-	inline FRHIUniformBuffer** GetUniformBufferStart() const
+	inline const FRHIUniformBuffer** GetUniformBufferStart() const
 	{
-		return (FRHIUniformBuffer**)(Data + GetUniformBufferOffset());
+		return (const FRHIUniformBuffer**)(Data + GetUniformBufferOffset());
 	}
 
 	inline FRHISamplerState** GetSamplerStart() const
@@ -246,22 +245,39 @@ private:
 		uint8* LooseDataStart = Data + GetLooseDataOffset();
 		return LooseDataStart;
 	}
-
-	inline void WriteBindingUniformBuffer(FRHIUniformBuffer* Value, uint32 BaseIndex)
+	
+	template<typename ElementType>
+	inline int FindSortedArrayBaseIndex(const TConstArrayView<ElementType>& Array, uint32 BaseIndex)
 	{
-		int32 FoundIndex = -1;
-
-		for (int32 SearchIndex = 0; SearchIndex < ParameterMapInfo.UniformBuffers.Num(); SearchIndex++)
+		int S0 = 0;
+		int S1 = Array.Num() - 1;
+		while (S0 < S1)
 		{
-			FShaderParameterInfo Parameter = ParameterMapInfo.UniformBuffers[SearchIndex];
+			int MED = (S0 + S1) / 2;
+			if (Array[MED].BaseIndex >= BaseIndex)
+			{
+				S1 = MED;
+			}
+			else
+			{
+				S0 = (S0 == MED) ? S0 + 1 : MED;
+			}
+		}
 
-			if (Parameter.BaseIndex == BaseIndex)
+		if (S0 == S1)
+		{
+			if (Array[S0].BaseIndex == BaseIndex)
 			{
-				FoundIndex = SearchIndex;
-				break;
+				return S0;
 			}
 		}
 
+		return -1;
+	}
+
+	inline void WriteBindingUniformBuffer(const FRHIUniformBuffer* Value, uint32 BaseIndex)
+	{
+		int32 FoundIndex = FindSortedArrayBaseIndex(MakeArrayView(ParameterMapInfo.UniformBuffers), BaseIndex);
 		if (FoundIndex >= 0)
 		{
 #if VALIDATE_UNIFORM_BUFFER_LIFETIME
@@ -281,19 +297,7 @@ private:
 
 	inline void WriteBindingSampler(FRHISamplerState* Value, uint32 BaseIndex)
 	{
-		int32 FoundIndex = -1;
-
-		for (int32 SearchIndex = 0; SearchIndex < ParameterMapInfo.TextureSamplers.Num(); SearchIndex++)
-		{
-			FShaderParameterInfo Parameter = ParameterMapInfo.TextureSamplers[SearchIndex];
-
-			if (Parameter.BaseIndex == BaseIndex)
-			{
-				FoundIndex = SearchIndex;
-				break;
-			}
-		}
-
+		int32 FoundIndex = FindSortedArrayBaseIndex(MakeArrayView(ParameterMapInfo.TextureSamplers), BaseIndex);
 		if (FoundIndex >= 0)
 		{
 			GetSamplerStart()[FoundIndex] = Value;
@@ -304,19 +308,7 @@ private:
 
 	inline void WriteBindingSRV(FRHIShaderResourceView* Value, uint32 BaseIndex)
 	{
-		int32 FoundIndex = -1;
-
-		for (int32 SearchIndex = 0; SearchIndex < ParameterMapInfo.SRVs.Num(); SearchIndex++)
-		{
-			FShaderParameterInfo Parameter = ParameterMapInfo.SRVs[SearchIndex];
-
-			if (Parameter.BaseIndex == BaseIndex)
-			{
-				FoundIndex = SearchIndex;
-				break;
-			}
-		}
-
+		int32 FoundIndex = FindSortedArrayBaseIndex(MakeArrayView(ParameterMapInfo.SRVs), BaseIndex);
 		if (FoundIndex >= 0)
 		{
 			uint32 TypeByteIndex = FoundIndex / 8;
@@ -330,19 +322,7 @@ private:
 
 	inline void WriteBindingTexture(FRHITexture* Value, uint32 BaseIndex)
 	{
-		int32 FoundIndex = -1;
-
-		for (int32 SearchIndex = 0; SearchIndex < ParameterMapInfo.SRVs.Num(); SearchIndex++)
-		{
-			FShaderParameterInfo Parameter = ParameterMapInfo.SRVs[SearchIndex];
-
-			if (Parameter.BaseIndex == BaseIndex)
-			{
-				FoundIndex = SearchIndex;
-				break;
-			}
-		}
-
+		int32 FoundIndex = FindSortedArrayBaseIndex(MakeArrayView(ParameterMapInfo.SRVs), BaseIndex);
 		if (FoundIndex >= 0)
 		{
 			GetSRVStart()[FoundIndex] = Value;
diff --git a/Engine/Source/Runtime/Renderer/Public/MeshMaterialShader.h b/Engine/Source/Runtime/Renderer/Public/MeshMaterialShader.h
--- a/Engine/Source/Runtime/Renderer/Public/MeshMaterialShader.h
+++ b/Engine/Source/Runtime/Renderer/Public/MeshMaterialShader.h
@@ -29,28 +29,46 @@ public:
 	RENDERER_API void InitializeMeshMaterialData(const FSceneView* SceneView, const FPrimitiveSceneProxy* RESTRICT PrimitiveSceneProxy, const FMeshBatch& RESTRICT MeshBatch, int32 StaticMeshId, bool bAllowStencilDither);
 };
 
+struct FMeshMaterialShaderPermutationParameters : public FMaterialShaderPermutationParameters
+{
+	// Type of vertex factory to compile.
+	const FVertexFactoryType* VertexFactoryType;
+
+	FMeshMaterialShaderPermutationParameters(EShaderPlatform InPlatform, const FMaterialShaderParameters& InMaterialParameters, const FVertexFactoryType* InVertexFactoryType, const int32 InPermutationId)
+		: FMaterialShaderPermutationParameters(InPlatform, InMaterialParameters, InPermutationId)
+		, VertexFactoryType(InVertexFactoryType)
+	{}
+};
+
+struct FVertexFactoryShaderPermutationParameters
+{
+	EShaderPlatform Platform;
+	FMaterialShaderParameters MaterialParameters;
+	const FVertexFactoryType* VertexFactoryType;
+
+	FVertexFactoryShaderPermutationParameters(EShaderPlatform InPlatform, const FMaterialShaderParameters& InMaterialParameters, const FVertexFactoryType* InVertexFactoryType)
+		: Platform(InPlatform)
+		, MaterialParameters(InMaterialParameters)
+		, VertexFactoryType(InVertexFactoryType)
+	{}
+};
+
 /** Base class of all shaders that need material and vertex factory parameters. */
 class RENDERER_API FMeshMaterialShader : public FMaterialShader
 {
-	DECLARE_SHADER_TYPE(FMeshMaterialShader, MeshMaterial);
+	DECLARE_TYPE_LAYOUT(FMeshMaterialShader, NonVirtual);
 public:
+	using FPermutationParameters = FMeshMaterialShaderPermutationParameters;
+	using ShaderMetaType = FMeshMaterialShaderType;
+
 	FMeshMaterialShader() {}
 
-	FMeshMaterialShader(const FMeshMaterialShaderType::CompiledShaderInitializerType& Initializer)
-		:	FMaterialShader(Initializer)
-		,	VertexFactoryParameters(Initializer.VertexFactoryType, Initializer.ParameterMap, Initializer.Target.GetFrequency(), Initializer.Target.GetPlatform())
-	{
-	}
-
-	static bool ValidateCompiledResult(EShaderPlatform Platform, const TArray<FMaterial*>& Materials, const FVertexFactoryType* VertexFactoryType, const FShaderParameterMap& ParameterMap, TArray<FString>& OutError)
-	{
-		return true;
-	}
+	FMeshMaterialShader(const FMeshMaterialShaderType::CompiledShaderInitializerType& Initializer);
 
 	// Declared as a friend, so that it can be called from other modules via static linkage, even if the compiler doesn't inline it.
-	FORCEINLINE friend void ValidateAfterBind(FMeshMaterialShader* Shader)
+	FORCEINLINE friend void ValidateAfterBind(const FShaderType* Type, FMeshMaterialShader* Shader)
 	{
-		checkfSlow(Shader->PassUniformBuffer.IsInitialized(), TEXT("FMeshMaterialShader must bind a pass uniform buffer, even if it is just FSceneTexturesUniformParameters: %s"), Shader->GetType()->GetName());
+		checkfSlow(Shader->PassUniformBuffer.IsInitialized(), TEXT("FMeshMaterialShader must bind a pass uniform buffer, even if it is just FSceneTexturesUniformParameters: %s"), Type->GetName());
 	}
 
 	void GetShaderBindings(
@@ -64,11 +82,12 @@ public:
 		FMeshDrawSingleShaderBindings& ShaderBindings) const;
 
 	void GetElementShaderBindings(
+		const FShaderMapPointerTable& PointerTable,
 		const FScene* Scene, 
 		const FSceneView* ViewIfDynamicMeshCommand, 
 		const FVertexFactory* VertexFactory,
 		const EVertexInputStreamType InputStreamType,
-		ERHIFeatureLevel::Type FeatureLevel,
+		const FStaticFeatureLevel FeatureLevel,
 		const FPrimitiveSceneProxy* PrimitiveSceneProxy,
 		const FMeshBatch& MeshBatch,
 		const FMeshBatchElement& BatchElement, 
@@ -76,14 +95,27 @@ public:
 		FMeshDrawSingleShaderBindings& ShaderBindings,
 		FVertexInputStreamArray& VertexStreams) const;
 
-	// FShader interface.
-	virtual const FVertexFactoryParameterRef* GetVertexFactoryParameterRef() const override { return &VertexFactoryParameters; }
-	virtual bool Serialize(FArchive& Ar) override;
-	virtual uint32 GetAllocatedSize() const override;
-
-protected:
-	FShaderUniformBufferParameter PassUniformBuffer;
+	template<typename ShaderType, typename PointerTableType, typename ShaderElementDataType>
+	static inline void GetElementShaderBindings(const TShaderRefBase<ShaderType, PointerTableType>& Shader,
+		const FScene* Scene,
+		const FSceneView* ViewIfDynamicMeshCommand,
+		const FVertexFactory* VertexFactory,
+		const EVertexInputStreamType InputStreamType,
+		ERHIFeatureLevel::Type FeatureLevel,
+		const FPrimitiveSceneProxy* PrimitiveSceneProxy,
+		const FMeshBatch& MeshBatch,
+		const FMeshBatchElement& BatchElement,
+		const ShaderElementDataType& ShaderElementData,
+		FMeshDrawSingleShaderBindings& ShaderBindings,
+		FVertexInputStreamArray& VertexStreams)
+	{
+		Shader->GetElementShaderBindings(Shader.GetPointerTable(), Scene, ViewIfDynamicMeshCommand, VertexFactory, InputStreamType, FeatureLevel, PrimitiveSceneProxy, MeshBatch, BatchElement, ShaderElementData, ShaderBindings, VertexStreams);
+	}
 
 private:
-	FVertexFactoryParameterRef VertexFactoryParameters;
+	void WriteFrozenVertexFactoryParameters(FMemoryImageWriter& Writer, const TMemoryImagePtr<FVertexFactoryShaderParameters>& InVertexFactoryParameters) const;
+	LAYOUT_FIELD_WITH_WRITER(TMemoryImagePtr<FVertexFactoryShaderParameters>, VertexFactoryParameters, WriteFrozenVertexFactoryParameters);
+
+protected:
+	LAYOUT_FIELD(FShaderUniformBufferParameter, PassUniformBuffer);
 };
diff --git a/Engine/Source/Runtime/Renderer/Public/MeshPassProcessor.h b/Engine/Source/Runtime/Renderer/Public/MeshPassProcessor.h
--- a/Engine/Source/Runtime/Renderer/Public/MeshPassProcessor.h
+++ b/Engine/Source/Runtime/Renderer/Public/MeshPassProcessor.h
@@ -9,6 +9,8 @@
 #include "MeshMaterialShader.h"
 #include "SceneUtils.h"
 #include "MeshBatch.h"
+#include "Hash/CityHash.h"
+#include "Experimental/Containers/RobinHoodHashTable.h"
 
 #define MESH_DRAW_COMMAND_DEBUG_DATA ((!UE_BUILD_SHIPPING && !UE_BUILD_TEST) || VALIDATE_MESH_COMMAND_BINDINGS || WANTS_DRAW_MESH_EVENTS)
 
@@ -89,59 +91,339 @@ public:
 	{
 	}
 
-	void Set(EMeshPass::Type Pass) { Data |= (1 << Pass); }
-	bool Get(EMeshPass::Type Pass) const { return !!(Data & (1 << Pass)); }
+	void Set(EMeshPass::Type Pass) 
+	{ 
+		Data |= (1 << Pass); 
+	}
+
+	bool Get(EMeshPass::Type Pass) const 
+	{ 
+		return !!(Data & (1 << Pass)); 
+	}
+
+	EMeshPass::Type SkipEmpty(EMeshPass::Type Pass) const 
+	{
+		uint32 Mask = 0xFFffFFff << Pass;
+		return EMeshPass::Type(FMath::Min<uint32>(EMeshPass::Num, FMath::CountTrailingZeros(Data & Mask)));
+	}
 
-	void AppendTo(FMeshPassMask& Mask) const { Mask.Data |= Data; }
-	void Reset() { Data = 0; }
-	bool IsEmpty() const { return Data == 0; }
+	int GetNum() 
+	{ 
+		return FMath::CountBits(Data); 
+	}
+
+	void AppendTo(FMeshPassMask& Mask) const 
+	{ 
+		Mask.Data |= Data; 
+	}
+
+	void Reset() 
+	{ 
+		Data = 0; 
+	}
+
+	bool IsEmpty() const 
+	{ 
+		return Data == 0; 
+	}
 
 	uint32 Data;
 };
 
-static_assert(sizeof(FMeshPassMask::Data) * 8 >= EMeshPass::Num, "FMeshPassMask::Data is too small to fit all mesh passes.");
-
-struct FRefCountedGraphicsMinimalPipelineStateInitializer
+struct FMinimalBoundShaderStateInput
 {
-	FRefCountedGraphicsMinimalPipelineStateInitializer(const FGraphicsMinimalPipelineStateInitializer& InStateInitializer, int32 InRefNum = 0)
-		: StateInitializer(InStateInitializer)
-		, RefNum(InRefNum)
+	inline FMinimalBoundShaderStateInput() {}
+
+	FBoundShaderStateInput AsBoundShaderState() const
+	{
+		return FBoundShaderStateInput(VertexDeclarationRHI
+			, VertexShaderResource ? static_cast<FRHIVertexShader*>(VertexShaderResource->GetShader(VertexShaderIndex)) : nullptr
+#if PLATFORM_SUPPORTS_TESSELLATION_SHADERS
+			, HullShaderResource ? static_cast<FRHIHullShader*>(HullShaderResource->GetShader(HullShaderIndex)) : nullptr
+			, DomainShaderResource ? static_cast<FRHIDomainShader*>(DomainShaderResource->GetShader(DomainShaderIndex)) : nullptr
+#endif
+			, PixelShaderResource ? static_cast<FRHIPixelShader*>(PixelShaderResource->GetShader(PixelShaderIndex)) : nullptr
+#if PLATFORM_SUPPORTS_GEOMETRY_SHADERS
+			, GeometryShaderResource ? static_cast<FRHIGeometryShader*>(GeometryShaderResource->GetShader(GeometryShaderIndex)) : nullptr
+#endif
+		);
+	}
+
+	void LazilyInitShaders() const
 	{
+		AsBoundShaderState(); // querying shaders will initialize on demand
 	}
 
-	FGraphicsMinimalPipelineStateInitializer StateInitializer;
-	int32 RefNum = 0;
+	bool NeedsShaderInitialisation() const
+	{
+		if (VertexShaderResource && !VertexShaderResource->HasShader(VertexShaderIndex))
+		{
+			return true;
+		}
+		if (HullShaderResource && !HullShaderResource->HasShader(HullShaderIndex))
+		{
+			return true;
+		}
+		if (DomainShaderResource && !DomainShaderResource->HasShader(DomainShaderIndex))
+		{
+			return true;
+		}
+		if (PixelShaderResource && !PixelShaderResource->HasShader(PixelShaderIndex))
+		{
+			return true;
+		}
+		if (GeometryShaderResource && !GeometryShaderResource->HasShader(GeometryShaderIndex))
+		{
+			return true;
+		}
+		return false;
+	}
+
+	FRHIVertexDeclaration* VertexDeclarationRHI = nullptr;
+	// TODO - all of these resources should always (?) come from the same shader map
+	// can we enforce this, and only store a single resource pointer?
+	class FShaderMapResource* VertexShaderResource = nullptr;
+	class FShaderMapResource* HullShaderResource = nullptr;
+	class FShaderMapResource* DomainShaderResource = nullptr;
+	class FShaderMapResource* PixelShaderResource = nullptr;
+	class FShaderMapResource* GeometryShaderResource = nullptr;
+	int32 VertexShaderIndex = INDEX_NONE;
+	int32 HullShaderIndex = INDEX_NONE;
+	int32 DomainShaderIndex = INDEX_NONE;
+	int32 PixelShaderIndex = INDEX_NONE;
+	int32 GeometryShaderIndex = INDEX_NONE;
 };
 
-struct RefCountedGraphicsMinimalPipelineStateInitializerKeyFuncs : DefaultKeyFuncs<FRefCountedGraphicsMinimalPipelineStateInitializer, false>
+
+/**
+ * Pipeline state without render target state
+ * Useful for mesh passes where the render target state is not changing between draws.
+ * Note: the size of this class affects rendering mesh pass traversal performance.
+ */
+class FGraphicsMinimalPipelineStateInitializer
 {
-	typedef typename TCallTraits<FGraphicsMinimalPipelineStateInitializer>::ConstReference KeyInitType;
+public:
+	// Can't use TEnumByte<EPixelFormat> as it changes the struct to be non trivially constructible, breaking memset
+	using TRenderTargetFormats = TStaticArray<uint8/*EPixelFormat*/, MaxSimultaneousRenderTargets>;
+	using TRenderTargetFlags = TStaticArray<uint32, MaxSimultaneousRenderTargets>;
 
-	/**
-	 * @return True if the keys match.
-	 */
-	static FORCEINLINE bool Matches(KeyInitType A, KeyInitType B)
+	FGraphicsMinimalPipelineStateInitializer()
+		: BlendState(nullptr)
+		, RasterizerState(nullptr)
+		, DepthStencilState(nullptr)
+		, PrimitiveType(PT_Num)
+	{
+		static_assert(sizeof(EPixelFormat) != sizeof(uint8), "Change TRenderTargetFormats's uint8 to EPixelFormat");
+		static_assert(PF_MAX < MAX_uint8, "TRenderTargetFormats assumes EPixelFormat can fit in a uint8!");
+	}
+
+	FGraphicsMinimalPipelineStateInitializer(
+		FMinimalBoundShaderStateInput	InBoundShaderState,
+		FRHIBlendState*					InBlendState,
+		FRHIRasterizerState*			InRasterizerState,
+		FRHIDepthStencilState*			InDepthStencilState,
+		FImmutableSamplerState			InImmutableSamplerState,
+		EPrimitiveType					InPrimitiveType
+	)
+		: BoundShaderState(InBoundShaderState)
+		, BlendState(InBlendState)
+		, RasterizerState(InRasterizerState)
+		, DepthStencilState(InDepthStencilState)
+		, ImmutableSamplerState(InImmutableSamplerState)
+		, PrimitiveType(InPrimitiveType)
+	{
+	}
+
+	FGraphicsMinimalPipelineStateInitializer(const FGraphicsMinimalPipelineStateInitializer& InMinimalState)
+		: BoundShaderState(InMinimalState.BoundShaderState)
+		, BlendState(InMinimalState.BlendState)
+		, RasterizerState(InMinimalState.RasterizerState)
+		, DepthStencilState(InMinimalState.DepthStencilState)
+		, ImmutableSamplerState(InMinimalState.ImmutableSamplerState)
+		, bDepthBounds(InMinimalState.bDepthBounds)
+		, PrimitiveType(InMinimalState.PrimitiveType)
+	{
+	}
+
+	FGraphicsPipelineStateInitializer AsGraphicsPipelineStateInitializer() const
+	{	
+		return FGraphicsPipelineStateInitializer
+		(	BoundShaderState.AsBoundShaderState()
+			, BlendState
+			, RasterizerState
+			, DepthStencilState
+			, ImmutableSamplerState
+			, PrimitiveType
+			, 0
+			, FGraphicsPipelineStateInitializer::TRenderTargetFormats(PF_Unknown)
+			, FGraphicsPipelineStateInitializer::TRenderTargetFlags(0)
+			, PF_Unknown
+			, 0
+			, ERenderTargetLoadAction::ENoAction
+			, ERenderTargetStoreAction::ENoAction
+			, ERenderTargetLoadAction::ENoAction
+			, ERenderTargetStoreAction::ENoAction
+			, FExclusiveDepthStencil::DepthNop
+			, 0
+			, ESubpassHint::None
+			, 0
+			, 0
+			, bDepthBounds
+			, bMultiView
+		);
+	}
+
+	inline bool operator==(const FGraphicsMinimalPipelineStateInitializer& rhs) const
+	{
+		if (BoundShaderState.VertexDeclarationRHI != rhs.BoundShaderState.VertexDeclarationRHI ||
+			BoundShaderState.VertexShaderResource != rhs.BoundShaderState.VertexShaderResource ||
+			BoundShaderState.PixelShaderResource != rhs.BoundShaderState.PixelShaderResource ||
+			BoundShaderState.VertexShaderIndex != rhs.BoundShaderState.VertexShaderIndex ||
+			BoundShaderState.PixelShaderIndex != rhs.BoundShaderState.PixelShaderIndex ||
+#if PLATFORM_SUPPORTS_GEOMETRY_SHADERS
+			BoundShaderState.GeometryShaderResource != rhs.BoundShaderState.GeometryShaderResource ||
+			BoundShaderState.GeometryShaderIndex != rhs.BoundShaderState.GeometryShaderIndex ||
+#endif
+#if PLATFORM_SUPPORTS_TESSELLATION_SHADERS
+			BoundShaderState.DomainShaderResource != rhs.BoundShaderState.DomainShaderResource ||
+			BoundShaderState.HullShaderResource != rhs.BoundShaderState.HullShaderResource ||
+			BoundShaderState.DomainShaderIndex != rhs.BoundShaderState.DomainShaderIndex ||
+			BoundShaderState.HullShaderIndex != rhs.BoundShaderState.HullShaderIndex ||
+#endif		
+			BlendState != rhs.BlendState ||
+			RasterizerState != rhs.RasterizerState ||
+			DepthStencilState != rhs.DepthStencilState ||
+			ImmutableSamplerState != rhs.ImmutableSamplerState ||
+			bDepthBounds != rhs.bDepthBounds ||
+			bMultiView != rhs.bMultiView ||
+			PrimitiveType != rhs.PrimitiveType)
+		{
+			return false;
+		}
+
+		return true;
+	}
+
+	inline bool operator!=(const FGraphicsMinimalPipelineStateInitializer& rhs) const
 	{
-		return A == B;
+		return !(*this == rhs);
 	}
 
-	/**
-	 * @return The key used to index the given element.
-	 */
-	static FORCEINLINE KeyInitType GetSetKey(ElementInitType Element)
+	inline friend uint32 GetTypeHash(const FGraphicsMinimalPipelineStateInitializer& Initializer)
 	{
-		return Element.StateInitializer;
+		//add and initialize any leftover padding within the struct to avoid unstable key
+		struct FHashKey
+		{
+			uint32 VertexDeclaration;
+			uint32 VertexShader;
+			uint32 PixelShader;
+			uint32 RasterizerState;
+		} HashKey;
+		HashKey.VertexDeclaration = PointerHash(Initializer.BoundShaderState.VertexDeclarationRHI);
+		HashKey.VertexShader = GetTypeHash(Initializer.BoundShaderState.VertexShaderIndex);
+		HashKey.PixelShader = GetTypeHash(Initializer.BoundShaderState.PixelShaderIndex);
+		HashKey.RasterizerState = PointerHash(Initializer.RasterizerState);
+
+		return uint32(CityHash64((const char*)&HashKey, sizeof(FHashKey)));
+	}
+
+#define COMPARE_FIELD_BEGIN(Field) \
+		if (Field != rhs.Field) \
+		{ return Field COMPARE_OP rhs.Field; }
+
+#define COMPARE_FIELD(Field) \
+		else if (Field != rhs.Field) \
+		{ return Field COMPARE_OP rhs.Field; }
+
+#define COMPARE_FIELD_END \
+		else { return false; }
+
+	bool operator<(const FGraphicsMinimalPipelineStateInitializer& rhs) const
+	{
+#define COMPARE_OP <
+
+		COMPARE_FIELD_BEGIN(BoundShaderState.VertexDeclarationRHI)
+			COMPARE_FIELD(BoundShaderState.VertexShaderIndex)
+			COMPARE_FIELD(BoundShaderState.PixelShaderIndex)
+			COMPARE_FIELD(BoundShaderState.VertexShaderResource)
+			COMPARE_FIELD(BoundShaderState.PixelShaderResource)
+#if PLATFORM_SUPPORTS_GEOMETRY_SHADERS
+			COMPARE_FIELD(BoundShaderState.GeometryShaderIndex)
+			COMPARE_FIELD(BoundShaderState.GeometryShaderResource)
+#endif
+#if PLATFORM_SUPPORTS_TESSELLATION_SHADERS
+			COMPARE_FIELD(BoundShaderState.DomainShaderIndex)
+			COMPARE_FIELD(BoundShaderState.HullShaderIndex)
+			COMPARE_FIELD(BoundShaderState.DomainShaderResource)
+			COMPARE_FIELD(BoundShaderState.HullShaderResource)
+#endif
+			COMPARE_FIELD(BlendState)
+			COMPARE_FIELD(RasterizerState)
+			COMPARE_FIELD(DepthStencilState)
+			COMPARE_FIELD(bDepthBounds)
+			COMPARE_FIELD(bMultiView)
+			COMPARE_FIELD(PrimitiveType)
+		COMPARE_FIELD_END;
+
+#undef COMPARE_OP
 	}
 
-	/** Calculates a hash index for a key. */
-	static FORCEINLINE uint32 GetKeyHash(KeyInitType Key)
+	bool operator>(const FGraphicsMinimalPipelineStateInitializer& rhs) const
 	{
-		return GetTypeHash(Key);
+#define COMPARE_OP >
+
+		COMPARE_FIELD_BEGIN(BoundShaderState.VertexDeclarationRHI)
+			COMPARE_FIELD(BoundShaderState.VertexShaderIndex)
+			COMPARE_FIELD(BoundShaderState.PixelShaderIndex)
+			COMPARE_FIELD(BoundShaderState.VertexShaderResource)
+			COMPARE_FIELD(BoundShaderState.PixelShaderResource)
+#if PLATFORM_SUPPORTS_GEOMETRY_SHADERS
+			COMPARE_FIELD(BoundShaderState.GeometryShaderIndex)
+			COMPARE_FIELD(BoundShaderState.GeometryShaderResource)
+#endif
+#if PLATFORM_SUPPORTS_TESSELLATION_SHADERS
+			COMPARE_FIELD(BoundShaderState.DomainShaderIndex)
+			COMPARE_FIELD(BoundShaderState.HullShaderIndex)
+			COMPARE_FIELD(BoundShaderState.DomainShaderResource)
+			COMPARE_FIELD(BoundShaderState.HullShaderResource)
+#endif
+			COMPARE_FIELD(BlendState)
+			COMPARE_FIELD(RasterizerState)
+			COMPARE_FIELD(DepthStencilState)
+			COMPARE_FIELD(bDepthBounds)
+			COMPARE_FIELD(bMultiView)
+			COMPARE_FIELD(PrimitiveType)
+			COMPARE_FIELD_END;
+
+#undef COMPARE_OP
 	}
+
+#undef COMPARE_FIELD_BEGIN
+#undef COMPARE_FIELD
+#undef COMPARE_FIELD_END
+
+	// TODO: [PSO API] - As we migrate reuse existing API objects, but eventually we can move to the direct initializers. 
+	// When we do that work, move this to RHI.h as its more appropriate there, but here for now since dependent typdefs are here.
+	FMinimalBoundShaderStateInput	BoundShaderState;
+	FRHIBlendState*					BlendState;
+	FRHIRasterizerState*			RasterizerState;
+	FRHIDepthStencilState*			DepthStencilState;
+	FImmutableSamplerState			ImmutableSamplerState;
+
+	// Note: FGraphicsMinimalPipelineStateInitializer is 8-byte aligned and can't have any implicit padding,
+	// as it is sometimes hashed and compared as raw bytes. Explicit padding is therefore required between
+	// all data members and at the end of the structure.
+	bool							bDepthBounds = false;
+	bool							bMultiView = false;
+	uint8							Padding[2] = {};
+
+	EPrimitiveType			PrimitiveType;
 };
 
+static_assert(sizeof(FMeshPassMask::Data) * 8 >= EMeshPass::Num, "FMeshPassMask::Data is too small to fit all mesh passes.");
+
 /** Set of FGraphicsMinimalPipelineStateInitializer unique per MeshDrawCommandsPassContext */
-typedef TSet< FGraphicsMinimalPipelineStateInitializer > FGraphicsMinimalPipelineStateSet;
+typedef Experimental::TRobinHoodHashSet< FGraphicsMinimalPipelineStateInitializer > FGraphicsMinimalPipelineStateSet;
 
 /** Uniquely represents a FGraphicsMinimalPipelineStateInitializer for fast compares. */
 class FGraphicsMinimalPipelineStateId
@@ -171,16 +453,18 @@ public:
 	
 	inline const FGraphicsMinimalPipelineStateInitializer& GetPipelineState(const FGraphicsMinimalPipelineStateSet& InPipelineSet) const
 	{
-		const FSetElementId SetElementId = FSetElementId::FromInteger(SetElementIndex);
-
 		if (bComesFromLocalPipelineStateSet)
 		{
-			return InPipelineSet[SetElementId];
+			return InPipelineSet.GetByElementId(SetElementIndex);
 		}
 
-		return PersistentIdTable[SetElementId].StateInitializer;
+		{
+			FScopeLock Lock(&PersistentIdTableLock);
+			return PersistentIdTable.GetByElementId(SetElementIndex).Key;
+		}
 	}
 
+	static void InitializePersistentIds();
 	/**
 	 * Get a ref counted persistent pipeline id, which needs to manually released.
 	 */
@@ -194,14 +478,26 @@ public:
 	/**
 	 * Get a pipeline state id in this order: global persistent Id table. If not found, will lookup in PassSet argument. If not found in PassSet argument, create a blank pipeline set id and add it PassSet argument
 	 */
-	RENDERER_API static FGraphicsMinimalPipelineStateId GetPipelineStateId(const FGraphicsMinimalPipelineStateInitializer& InPipelineState, FGraphicsMinimalPipelineStateSet& InOutPassSet);
+	RENDERER_API static FGraphicsMinimalPipelineStateId GetPipelineStateId(const FGraphicsMinimalPipelineStateInitializer& InPipelineState, FGraphicsMinimalPipelineStateSet& InOutPassSet, bool& NeedsShaderInitialisation);
 
-	static int32 GetLocalPipelineIdTableSize() { return LocalPipelineIdTableSize; }
+	static int32 GetLocalPipelineIdTableSize() 
+	{ 
+		FScopeLock Lock(&PersistentIdTableLock);
+		return LocalPipelineIdTableSize; 
+	}
 	static void ResetLocalPipelineIdTableSize();
 	static void AddSizeToLocalPipelineIdTableSize(SIZE_T Size);
 
-	static SIZE_T GetPersistentIdTableSize() { return PersistentIdTable.GetAllocatedSize(); }
-	static int32 GetPersistentIdNum() { return PersistentIdTable.Num(); }
+	static SIZE_T GetPersistentIdTableSize() 
+	{ 
+		FScopeLock Lock(&PersistentIdTableLock);
+		return PersistentIdTable.GetAllocatedSize(); 
+	}
+	static int32 GetPersistentIdNum() 
+	{ 
+		FScopeLock Lock(&PersistentIdTableLock);
+		return PersistentIdTable.Num(); 
+	}
 
 private:
 	union
@@ -216,25 +512,36 @@ private:
 		};
 	};
 
-	static TSet<FRefCountedGraphicsMinimalPipelineStateInitializer, RefCountedGraphicsMinimalPipelineStateInitializerKeyFuncs> PersistentIdTable;
-	
+	struct FRefCountedGraphicsMinimalPipelineState
+	{
+		FRefCountedGraphicsMinimalPipelineState() : RefNum(0)
+		{
+		}
+		uint32 RefNum;
+	};
+
+	static FCriticalSection PersistentIdTableLock;
+	using PersistentTableType = Experimental::TRobinHoodHashMap<FGraphicsMinimalPipelineStateInitializer, FRefCountedGraphicsMinimalPipelineState>;
+	static PersistentTableType PersistentIdTable;
+
 	static int32 LocalPipelineIdTableSize;
 	static int32 CurrentLocalPipelineIdTableSize;
+	static bool NeedsShaderInitialisation;
 };
 
 struct FMeshProcessorShaders
 {
-	mutable FMeshMaterialShader* VertexShader;
-	mutable FMeshMaterialShader* HullShader;
-	mutable FMeshMaterialShader* DomainShader;
-	mutable FMeshMaterialShader* PixelShader;
-	mutable FMeshMaterialShader* GeometryShader;
-	mutable FMeshMaterialShader* ComputeShader;
+	mutable TShaderRef<FMeshMaterialShader> VertexShader;
+	mutable TShaderRef<FMeshMaterialShader> HullShader;
+	mutable TShaderRef<FMeshMaterialShader> DomainShader;
+	mutable TShaderRef<FMeshMaterialShader> PixelShader;
+	mutable TShaderRef<FMeshMaterialShader> GeometryShader;
+	mutable TShaderRef<FMeshMaterialShader> ComputeShader;
 #if RHI_RAYTRACING
-	mutable FMeshMaterialShader* RayHitGroupShader;
+	mutable TShaderRef<FMeshMaterialShader> RayHitGroupShader;
 #endif
 
-	FMeshMaterialShader* GetShader(EShaderFrequency Frequency) const
+	TShaderRef<FMeshMaterialShader> GetShader(EShaderFrequency Frequency) const
 	{
 		if (Frequency == SF_Vertex)
 		{
@@ -268,7 +575,7 @@ struct FMeshProcessorShaders
 #endif // RHI_RAYTRACING
 
 		checkf(0, TEXT("Unhandled shader frequency"));
-		return nullptr;
+		return TShaderRef<FMeshMaterialShader>();
 	}
 };
 
@@ -288,8 +595,8 @@ struct FMeshDrawCommandDebugData
 	const FPrimitiveSceneProxy* PrimitiveSceneProxyIfNotUsingStateBuckets;
 	const FMaterial* Material;
 	const FMaterialRenderProxy* MaterialRenderProxy;
-	FMeshMaterialShader* VertexShader;
-	FMeshMaterialShader* PixelShader;
+	TShaderRef<FMeshMaterialShader> VertexShader;
+	TShaderRef<FMeshMaterialShader> PixelShader;
 	const FVertexFactory* VertexFactory;
 	FName ResourceName;
 #endif
@@ -302,7 +609,31 @@ class FMeshDrawShaderBindings
 {
 public:
 
-	FMeshDrawShaderBindings() {}
+	FMeshDrawShaderBindings() 
+	{
+		static_assert(sizeof(ShaderFrequencyBits) * 8 > SF_NumFrequencies, "Please increase ShaderFrequencyBits size");
+	}
+	FMeshDrawShaderBindings(FMeshDrawShaderBindings&& Other)
+	{
+		if (!UsesInlineStorage())
+		{
+			delete[] Data.GetHeapData();
+		}
+		Size = Other.Size;
+		ShaderFrequencyBits = Other.ShaderFrequencyBits;
+		ShaderLayouts = MoveTemp(Other.ShaderLayouts);
+		if (Other.UsesInlineStorage())
+		{
+			Data = MoveTemp(Other.Data);
+		}
+		else
+		{		
+			Data.SetHeapData(Other.Data.GetHeapData());
+			Other.Data.SetHeapData(nullptr);
+		}
+		Other.Size = 0;	
+	}
+
 	FMeshDrawShaderBindings(const FMeshDrawShaderBindings& Other)
 	{
 		CopyFrom(Other);
@@ -315,29 +646,55 @@ public:
 		return *this;
 	}
 
+	FMeshDrawShaderBindings& operator=(FMeshDrawShaderBindings&& Other)
+	{
+		if (!UsesInlineStorage())
+		{
+			delete[] Data.GetHeapData();
+		}
+		Size = Other.Size;
+		ShaderFrequencyBits = Other.ShaderFrequencyBits;
+		ShaderLayouts = MoveTemp(Other.ShaderLayouts);
+		if (Other.UsesInlineStorage())
+		{
+			Data = MoveTemp(Other.Data);
+		}
+		else
+		{	
+			Data.SetHeapData(Other.Data.GetHeapData());
+			Other.Data.SetHeapData(nullptr);
+		}
+		Other.Size = 0;
+		return *this;
+	}
+
 	/** Allocates space for the bindings of all shaders. */
 	void Initialize(FMeshProcessorShaders Shaders);
 
 	/** Called once binding setup is complete. */
 	void Finalize(const FMeshProcessorShaders* ShadersForDebugging);
 
-	inline FMeshDrawSingleShaderBindings GetSingleShaderBindings(EShaderFrequency Frequency)
+	inline FMeshDrawSingleShaderBindings GetSingleShaderBindings(EShaderFrequency Frequency, int32& DataOffset)
 	{
-		int32 DataOffset = 0;
+		int FrequencyIndex = FPlatformMath::CountBits(ShaderFrequencyBits & ((1 << (Frequency + 1)) - 1)) - 1;
 
-		for (int32 BindingIndex = 0; BindingIndex < ShaderLayouts.Num(); BindingIndex++)
+#if DO_CHECK && !(UE_BUILD_TEST || UE_BUILD_SHIPPING)
+		int32 CheckedDataOffset = 0;
+		for (int32 BindingIndex = 0; BindingIndex < FrequencyIndex; BindingIndex++)
 		{
-			if (ShaderLayouts[BindingIndex].Frequency == Frequency)
-			{
-				return FMeshDrawSingleShaderBindings(ShaderLayouts[BindingIndex], GetData() + DataOffset);
-			}
-
-			DataOffset += ShaderLayouts[BindingIndex].GetDataSizeBytes();
+			CheckedDataOffset += ShaderLayouts[BindingIndex].GetDataSizeBytes();
+		}
+		checkf(CheckedDataOffset == DataOffset, TEXT("GetSingleShaderBindings was not called in the order of ShaderFrequencies"));
+#endif
+		if (FrequencyIndex >= 0)
+		{
+			int32 StartDataOffset = DataOffset;
+			DataOffset += ShaderLayouts[FrequencyIndex].GetDataSizeBytes();
+			return FMeshDrawSingleShaderBindings(ShaderLayouts[FrequencyIndex], GetData() + StartDataOffset);
 		}
 
 		checkf(0, TEXT("Invalid shader binding frequency requested"));
-		FShader Shader;
-		return FMeshDrawSingleShaderBindings(FMeshDrawShaderBindingsLayout(&Shader), nullptr);
+		return FMeshDrawSingleShaderBindings(FMeshDrawShaderBindingsLayout(TShaderRef<FShader>()), nullptr);
 	}
 
 	/** Set shader bindings on the commandlist, filtered by state cache. */
@@ -356,7 +713,7 @@ public:
 	SIZE_T GetAllocatedSize() const
 	{
 		SIZE_T Bytes = ShaderLayouts.GetAllocatedSize();
-		if (Size > sizeof(InlineStorage))
+		if (!UsesInlineStorage())
 		{
 			Bytes += Size;
 		}
@@ -368,35 +725,47 @@ public:
 	{
 		OutShaderFrequencies.Empty(ShaderLayouts.Num());
 
-		for (int32 BindingIndex = 0; BindingIndex < ShaderLayouts.Num(); BindingIndex++)
+		for (int32 BindingIndex = 0; BindingIndex < SF_NumFrequencies; BindingIndex++)
 		{
-			OutShaderFrequencies.Add(ShaderLayouts[BindingIndex].Frequency);
+			if ((ShaderFrequencyBits & (1 << BindingIndex)) != 0)
+			{
+				OutShaderFrequencies.Add(EShaderFrequency(BindingIndex));
+			}
 		}
 	}
 
 	inline int32 GetDataSize() const { return Size; }
 
 private:
-
 	TArray<FMeshDrawShaderBindingsLayout, TInlineAllocator<2>> ShaderLayouts;
-
-	union
+	struct FData
 	{
-		uint8 InlineStorage[NumInlineShaderBindings * sizeof(void*)];
-		uint8* HeapData = nullptr;
-	};
-	
+		uint8* InlineStorage[NumInlineShaderBindings] = {};
+		uint8* GetHeapData()
+		{
+			return InlineStorage[0];
+		}
+		const uint8* GetHeapData() const
+		{
+			return InlineStorage[0];
+		}
+		void SetHeapData(uint8* HeapData)
+		{
+			InlineStorage[0] = HeapData;
+		}
+	} Data = {};
+	uint16 ShaderFrequencyBits = 0;
 	uint16 Size = 0;
 
 	void Allocate(uint16 InSize)
 	{
-		check(Size == 0 && HeapData == nullptr);
+		check(Size == 0 && Data.GetHeapData() == nullptr);
 
 		Size = InSize;
 
-		if (InSize > UE_ARRAY_COUNT(InlineStorage))
+		if (InSize > sizeof(FData))
 		{
-			HeapData = new uint8[InSize];
+			Data.SetHeapData(new uint8[InSize]);
 		}
 	}
 
@@ -407,17 +776,25 @@ private:
 		// Verify no type overflow
 		check(Size == InSize);
 
-		FPlatformMemory::Memzero(GetData(), InSize);
+		if (!UsesInlineStorage())
+		{
+			FPlatformMemory::Memzero(GetData(), InSize);
+		}
+	}
+
+	inline bool UsesInlineStorage() const
+	{
+		return Size <= sizeof(FData);
 	}
 
 	uint8* GetData()
 	{
-		return Size <= UE_ARRAY_COUNT(InlineStorage) ? &InlineStorage[0] : HeapData;
+		return UsesInlineStorage() ? reinterpret_cast<uint8*>(&Data.InlineStorage[0]) : Data.GetHeapData();
 	}
 
 	const uint8* GetData() const
 	{
-		return Size <= UE_ARRAY_COUNT(InlineStorage) ? &InlineStorage[0] : HeapData;
+		return UsesInlineStorage() ? reinterpret_cast<const uint8*>(&Data.InlineStorage[0]) : Data.GetHeapData();
 	}
 
 	RENDERER_API void CopyFrom(const FMeshDrawShaderBindings& Other);
@@ -490,8 +867,11 @@ public:
 	/** Non-pipeline state */
 	uint8 StencilRef;
 
-	FMeshDrawCommand()
-	{}
+	FMeshDrawCommand() {};
+	FMeshDrawCommand(FMeshDrawCommand&& Other) = default;
+	FMeshDrawCommand(const FMeshDrawCommand& Other) = default;
+	FMeshDrawCommand& operator=(const FMeshDrawCommand& Other) = default;
+	FMeshDrawCommand& operator=(FMeshDrawCommand&& Other) = default; 
 
 	bool MatchesForDynamicInstancing(const FMeshDrawCommand& Rhs) const
 	{
@@ -510,38 +890,70 @@ public:
 
 	uint32 GetDynamicInstancingHash() const
 	{
-		uint32 Hash = FCrc::TypeCrc32(CachedPipelineId.GetId(), 0);
-		Hash = FCrc::TypeCrc32(StencilRef, Hash);
-		Hash = HashCombine(ShaderBindings.GetDynamicInstancingHash(), Hash);
+		//add and initialize any leftover padding within the struct to avoid unstable keys
+		struct FHashKey
+		{
+			uint32 IndexBuffer;
+			uint32 VertexBuffers = 0;
+		    uint32 VertexStreams = 0;
+			uint32 PipelineId;
+			uint32 DynamicInstancingHash;
+			uint32 FirstIndex;
+			uint32 NumPrimitives;
+			uint32 NumInstances;
+			uint32 IndirectArgsBufferOrBaseVertexIndex;
+			uint32 NumVertices;
+			uint32 StencilRefAndPrimitiveIdStreamIndex;
+
+			static inline uint32 PointerHash(const void* Key)
+			{
+#if PLATFORM_64BITS
+				// Ignoring the lower 4 bits since they are likely zero anyway.
+				// Higher bits are more significant in 64 bit builds.
+				return reinterpret_cast<UPTRINT>(Key) >> 4;
+#else
+				return reinterpret_cast<UPTRINT>(Key);
+#endif
+			};
+
+			static inline uint32 HashCombine(uint32 A, uint32 B)
+			{
+				return A ^ (B + 0x9e3779b9 + (A << 6) + (A >> 2));
+			}
+		} HashKey;
+
+		HashKey.PipelineId = CachedPipelineId.GetId();
+		HashKey.StencilRefAndPrimitiveIdStreamIndex = StencilRef | (PrimitiveIdStreamIndex << 8);
+		HashKey.DynamicInstancingHash = ShaderBindings.GetDynamicInstancingHash();
 
-		for (const FVertexInputStream& VertexInputStream: VertexStreams)
+		for (int index = 0; index < VertexStreams.Num(); index++)
 		{
+			const FVertexInputStream& VertexInputStream = VertexStreams[index];
 			const uint32 StreamIndex = VertexInputStream.StreamIndex;
 			const uint32 Offset = VertexInputStream.Offset;
 
-			Hash = FCrc::TypeCrc32(StreamIndex, Hash);
-			Hash = FCrc::TypeCrc32(Offset, Hash);
-			Hash = PointerHash(VertexInputStream.VertexBuffer, Hash);
+			uint32 Packed = (StreamIndex << 28) | Offset;
+			HashKey.VertexStreams = FHashKey::HashCombine(HashKey.VertexStreams, Packed);
+			HashKey.VertexBuffers = FHashKey::HashCombine(HashKey.VertexBuffers, FHashKey::PointerHash(VertexInputStream.VertexBuffer));
 		}
 
-		Hash = FCrc::TypeCrc32(PrimitiveIdStreamIndex, Hash);
-		Hash = PointerHash(IndexBuffer, Hash);
-		Hash = FCrc::TypeCrc32(FirstIndex, Hash);
-		Hash = FCrc::TypeCrc32(NumPrimitives, Hash);
-		Hash = FCrc::TypeCrc32(NumInstances, Hash);
+		HashKey.IndexBuffer = FHashKey::PointerHash(IndexBuffer);
+		HashKey.FirstIndex = FirstIndex;
+		HashKey.NumPrimitives = NumPrimitives;
+		HashKey.NumInstances = NumInstances;
 
 		if (NumPrimitives > 0)
 		{
-			Hash = FCrc::TypeCrc32(VertexParams.BaseVertexIndex, Hash);
-			Hash = FCrc::TypeCrc32(VertexParams.NumVertices, Hash);
+			HashKey.IndirectArgsBufferOrBaseVertexIndex = VertexParams.BaseVertexIndex;
+			HashKey.NumVertices = VertexParams.NumVertices;
 		}
 		else
 		{
-			Hash = PointerHash(IndirectArgs.Buffer, Hash);
-			Hash = FCrc::TypeCrc32(IndirectArgs.Offset, Hash);
+			HashKey.IndirectArgsBufferOrBaseVertexIndex = FHashKey::PointerHash(IndirectArgs.Buffer);
+			HashKey.NumVertices = IndirectArgs.Offset;
 		}		
 
-		return Hash;
+		return uint32(CityHash64((char*)&HashKey, sizeof(FHashKey)));
 	}
 
 	/** Sets shaders on the mesh draw command and allocates room for the shader bindings. */
@@ -601,12 +1013,12 @@ public:
 	}
 
 #if MESH_DRAW_COMMAND_DEBUG_DATA
-	void ClearDebugPrimitiveSceneProxy()
+	void ClearDebugPrimitiveSceneProxy() const
 	{
 		DebugData.PrimitiveSceneProxyIfNotUsingStateBuckets = nullptr;
 	}
 private:
-	FMeshDrawCommandDebugData DebugData;
+	mutable FMeshDrawCommandDebugData DebugData;
 #endif
 };
 
@@ -659,7 +1071,7 @@ public:
 
 	virtual ~FMeshPassDrawListContext() {}
 
-	virtual FMeshDrawCommand& AddCommand(const FMeshDrawCommand& Initializer) = 0;
+	virtual FMeshDrawCommand& AddCommand(FMeshDrawCommand& Initializer, uint32 NumElements) = 0;
 
 	virtual void FinalizeCommand(
 		const FMeshBatch& MeshBatch, 
@@ -755,14 +1167,16 @@ public:
 	(
 		FDynamicMeshDrawCommandStorage& InDrawListStorage, 
 		FMeshCommandOneFrameArray& InDrawList,
-		FGraphicsMinimalPipelineStateSet& InPipelineStateSet
+		FGraphicsMinimalPipelineStateSet& InPipelineStateSet,
+		bool& InNeedsShaderInitialisation
 	) :
 		DrawListStorage(InDrawListStorage),
 		DrawList(InDrawList),
-		GraphicsMinimalPipelineStateSet(InPipelineStateSet)
+		GraphicsMinimalPipelineStateSet(InPipelineStateSet),
+		NeedsShaderInitialisation(InNeedsShaderInitialisation)
 	{}
 
-	virtual FMeshDrawCommand& AddCommand(const FMeshDrawCommand& Initializer) override final
+	virtual FMeshDrawCommand& AddCommand(FMeshDrawCommand& Initializer, uint32 NumElements) override final
 	{
 		const int32 Index = DrawListStorage.MeshDrawCommands.AddElement(Initializer);
 		FMeshDrawCommand& NewCommand = DrawListStorage.MeshDrawCommands[Index];
@@ -781,7 +1195,7 @@ public:
 		const FMeshProcessorShaders* ShadersForDebugging,
 		FMeshDrawCommand& MeshDrawCommand) override final
 	{
-		FGraphicsMinimalPipelineStateId PipelineId = FGraphicsMinimalPipelineStateId::GetPipelineStateId(PipelineState, GraphicsMinimalPipelineStateSet);
+		FGraphicsMinimalPipelineStateId PipelineId = FGraphicsMinimalPipelineStateId::GetPipelineStateId(PipelineState, GraphicsMinimalPipelineStateSet, NeedsShaderInitialisation);
 
 		MeshDrawCommand.SetDrawParametersAndFinalize(MeshBatch, BatchElementIndex, PipelineId, ShadersForDebugging);
 
@@ -796,6 +1210,7 @@ private:
 	FDynamicMeshDrawCommandStorage& DrawListStorage;
 	FMeshCommandOneFrameArray& DrawList;
 	FGraphicsMinimalPipelineStateSet& GraphicsMinimalPipelineStateSet;
+	bool& NeedsShaderInitialisation;
 };
 
 #if PLATFORM_SUPPORTS_PRAGMA_PACK
@@ -809,11 +1224,14 @@ private:
 class FCachedMeshDrawCommandInfo
 {
 public:
-	explicit FCachedMeshDrawCommandInfo() :
+	FCachedMeshDrawCommandInfo() : FCachedMeshDrawCommandInfo(EMeshPass::Num)
+	{}
+
+	explicit FCachedMeshDrawCommandInfo(EMeshPass::Type InMeshPass) :
 		SortKey(FMeshDrawCommandSortKey::Default),
-		CommandIndex(-1),
-		StateBucketId(-1),
-		MeshPass(EMeshPass::Num),
+		CommandIndex(INDEX_NONE),
+		StateBucketId(INDEX_NONE),
+		MeshPass(InMeshPass),
 		MeshFillMode(ERasterizerFillMode_Num),
 		MeshCullMode(ERasterizerCullMode_Num)
 	{}
@@ -851,14 +1269,37 @@ public:
 	int32 LowestFreeIndexSearchStart;
 };
 
-typedef TArray<int32, TInlineAllocator<5>> FDrawCommandIndices;
+struct FMeshDrawCommandCount 
+{
+	uint32 Num = 0;
+};
+
+struct MeshDrawCommandKeyFuncs : TDefaultMapHashableKeyFuncs<FMeshDrawCommand, FMeshDrawCommandCount, false>
+{
+	/**
+	 * @return True if the keys match.
+	 */
+	static inline bool Matches(KeyInitType A, KeyInitType B)
+	{
+		return A.MatchesForDynamicInstancing(B);
+	}
+
+	/** Calculates a hash index for a key. */
+	static inline uint32 GetKeyHash(KeyInitType Key)
+	{
+		return Key.GetDynamicInstancingHash();
+	}
+};
+
+using FDrawCommandIndices = TArray<int32, TInlineAllocator<5>>;
+using FStateBucketMap = Experimental::TRobinHoodHashMap<FMeshDrawCommand, FMeshDrawCommandCount, MeshDrawCommandKeyFuncs>;
 
 class FCachedPassMeshDrawListContext : public FMeshPassDrawListContext
 {
 public:
-	FCachedPassMeshDrawListContext(FCachedMeshDrawCommandInfo& InCommandInfo, FCachedPassMeshDrawList& InDrawList, FScene& InScene);
+	FCachedPassMeshDrawListContext(FCachedMeshDrawCommandInfo& InCommandInfo, FCriticalSection& InCachedMeshDrawCommandLock, FCachedPassMeshDrawList& InCachedDrawLists, FStateBucketMap& InCachedMeshDrawCommandStateBuckets, const FScene& InScene);
 
-	virtual FMeshDrawCommand& AddCommand(const FMeshDrawCommand& Initializer) override final;
+	virtual FMeshDrawCommand& AddCommand(FMeshDrawCommand& Initializer, uint32 NumElements) override final;
 
 	virtual void FinalizeCommand(
 		const FMeshBatch& MeshBatch, 
@@ -875,22 +1316,23 @@ public:
 private:
 	FMeshDrawCommand MeshDrawCommandForStateBucketing;
 	FCachedMeshDrawCommandInfo& CommandInfo;
-	FCachedPassMeshDrawList& DrawList;
-	FScene& Scene;
-	bool bUseStateBuckets;
+	FCriticalSection& CachedMeshDrawCommandLock;
+	FCachedPassMeshDrawList& CachedDrawLists;
+	FStateBucketMap& CachedMeshDrawCommandStateBuckets;
+	const FScene& Scene;
 };
 
 template<typename VertexType, typename HullType, typename DomainType, typename PixelType, typename GeometryType = FMeshMaterialShader, typename RayHitGroupType = FMeshMaterialShader, typename ComputeType = FMeshMaterialShader>
 struct TMeshProcessorShaders
 {
-	VertexType* VertexShader = nullptr;
-	HullType* HullShader = nullptr;
-	DomainType* DomainShader = nullptr;
-	PixelType* PixelShader = nullptr;
-	GeometryType* GeometryShader = nullptr;
-	ComputeType* ComputeShader = nullptr;
+	TShaderRef<VertexType> VertexShader;
+	TShaderRef<HullType> HullShader;
+	TShaderRef<DomainType> DomainShader;
+	TShaderRef<PixelType> PixelShader;
+	TShaderRef<GeometryType> GeometryShader;
+	TShaderRef<ComputeType> ComputeShader;
 #if RHI_RAYTRACING
-	RayHitGroupType* RayHitGroupShader = nullptr;
+	TShaderRef<RayHitGroupType> RayHitGroupShader;
 #endif
 
 	TMeshProcessorShaders() = default;
@@ -1017,7 +1459,7 @@ public:
 		ViewUniformBuffer = InViewUniformBuffer;
 	}
 
-	FORCEINLINE_DEBUGGABLE const TUniformBufferRef<FViewUniformShaderParameters>& GetViewUniformBuffer() const
+	FORCEINLINE_DEBUGGABLE const FRHIUniformBuffer* GetViewUniformBuffer() const
 	{
 		return ViewUniformBuffer;
 	}
@@ -1027,9 +1469,9 @@ public:
 		InstancedViewUniformBuffer = InViewUniformBuffer;
 	}
 
-	FORCEINLINE_DEBUGGABLE const TUniformBufferRef<FInstancedViewUniformShaderParameters>& GetInstancedViewUniformBuffer() const
+	FORCEINLINE_DEBUGGABLE const FRHIUniformBuffer* GetInstancedViewUniformBuffer() const
 	{
-		return InstancedViewUniformBuffer.IsValid() ? InstancedViewUniformBuffer : reinterpret_cast<const TUniformBufferRef<FInstancedViewUniformShaderParameters>&>(ViewUniformBuffer);
+		return InstancedViewUniformBuffer != nullptr ? InstancedViewUniformBuffer : ViewUniformBuffer;
 	}
 
 	FORCEINLINE_DEBUGGABLE void SetReflectionCaptureUniformBuffer(FRHIUniformBuffer* InUniformBuffer)
@@ -1037,7 +1479,7 @@ public:
 		ReflectionCaptureUniformBuffer = InUniformBuffer;
 	}
 
-	FORCEINLINE_DEBUGGABLE const FUniformBufferRHIRef& GetReflectionCaptureUniformBuffer() const
+	FORCEINLINE_DEBUGGABLE const FRHIUniformBuffer* GetReflectionCaptureUniformBuffer() const
 	{
 		return ReflectionCaptureUniformBuffer;
 	}
@@ -1068,16 +1510,26 @@ private:
 	FRHIDepthStencilState*			DepthStencilState;
 	FExclusiveDepthStencil::Type	DepthStencilAccess;
 
-	TUniformBufferRef<FViewUniformShaderParameters>	ViewUniformBuffer;
-	TUniformBufferRef<FInstancedViewUniformShaderParameters> InstancedViewUniformBuffer;
+	FRHIUniformBuffer*				ViewUniformBuffer;
+	FRHIUniformBuffer*				InstancedViewUniformBuffer;
 
 	/** Will be bound as reflection capture uniform buffer in case where scene is not available, typically set to dummy/empty buffer to avoid null binding */
-	FUniformBufferRHIRef			ReflectionCaptureUniformBuffer;
+	FRHIUniformBuffer*				ReflectionCaptureUniformBuffer;
 
 	FRHIUniformBuffer*				PassUniformBuffer;
 	uint32							StencilRef;
 };
 
+enum class EDrawingPolicyOverrideFlags
+{
+	None = 0,
+	TwoSided = 1 << 0,
+	DitheredLODTransition = 1 << 1,
+	Wireframe = 1 << 2,
+	ReverseCullMode = 1 << 3,
+};
+ENUM_CLASS_FLAGS(EDrawingPolicyOverrideFlags);
+
 /** 
  * Base class of mesh processors, whose job is to transform FMeshBatch draw descriptions received from scene proxy implementations into FMeshDrawCommands ready for the RHI command list
  */
@@ -1108,8 +1560,15 @@ public:
 		return CullMode == CM_None ? CM_None : (CullMode == CM_CCW ? CM_CW : CM_CCW);
 	}
 
-	RENDERER_API ERasterizerFillMode ComputeMeshFillMode(const FMeshBatch& Mesh, const FMaterial& InMaterialResource) const;
-	RENDERER_API ERasterizerCullMode ComputeMeshCullMode(const FMeshBatch& Mesh, const FMaterial& InMaterialResource) const;
+	struct FMeshDrawingPolicyOverrideSettings
+	{
+		EDrawingPolicyOverrideFlags	MeshOverrideFlags = EDrawingPolicyOverrideFlags::None;
+		EPrimitiveType				MeshPrimitiveType = PT_TriangleList;
+	};
+
+	RENDERER_API static FMeshDrawingPolicyOverrideSettings ComputeMeshOverrideSettings(const FMeshBatch& Mesh);
+	RENDERER_API static ERasterizerFillMode ComputeMeshFillMode(const FMeshBatch& Mesh, const FMaterial& InMaterialResource, const FMeshDrawingPolicyOverrideSettings& InOverrideSettings);
+	RENDERER_API static ERasterizerCullMode ComputeMeshCullMode(const FMeshBatch& Mesh, const FMaterial& InMaterialResource, const FMeshDrawingPolicyOverrideSettings& InOverrideSettings);
 
 	template<typename PassShadersType, typename ShaderElementDataType>
 	void BuildMeshDrawCommands(
@@ -1216,7 +1675,8 @@ extern void ApplyViewOverridesToMeshDrawCommands(
 	const FSceneView& View,
 	FMeshCommandOneFrameArray& VisibleMeshDrawCommands,
 	FDynamicMeshDrawCommandStorage& DynamicMeshDrawCommandStorage,
-	FGraphicsMinimalPipelineStateSet& GraphicsMinimalPipelineStateSet);
+	FGraphicsMinimalPipelineStateSet& GraphicsMinimalPipelineStateSet,
+	bool& NeedsShaderInitialisation);
 
 RENDERER_API extern void DrawDynamicMeshPassPrivate(
 	const FSceneView& View,
@@ -1224,10 +1684,16 @@ RENDERER_API extern void DrawDynamicMeshPassPrivate(
 	FMeshCommandOneFrameArray& VisibleMeshDrawCommands,
 	FDynamicMeshDrawCommandStorage& DynamicMeshDrawCommandStorage,
 	FGraphicsMinimalPipelineStateSet& GraphicsMinimalPipelineStateSet,
+	bool& InNeedsShaderInitialisation,
 	uint32 InstanceFactor);
 
 RENDERER_API extern FMeshDrawCommandSortKey CalculateMeshStaticSortKey(const FMeshMaterialShader* VertexShader, const FMeshMaterialShader* PixelShader);
 
+inline FMeshDrawCommandSortKey CalculateMeshStaticSortKey(const TShaderRef<FMeshMaterialShader>& VertexShader, const TShaderRef<FMeshMaterialShader>& PixelShader)
+{
+	return CalculateMeshStaticSortKey(VertexShader.GetShader(), PixelShader.GetShader());
+}
+
 #if RHI_RAYTRACING
 class FRayTracingMeshCommand
 {
diff --git a/Engine/Source/Runtime/Renderer/Public/MeshPassProcessor.inl b/Engine/Source/Runtime/Renderer/Public/MeshPassProcessor.inl
--- a/Engine/Source/Runtime/Renderer/Public/MeshPassProcessor.inl
+++ b/Engine/Source/Runtime/Renderer/Public/MeshPassProcessor.inl
@@ -56,33 +56,30 @@ void FMeshPassProcessor::BuildMeshDrawCommands(
 
 	SharedMeshDrawCommand.PrimitiveIdStreamIndex = VertexFactory->GetPrimitiveIdStreamIndex(InputStreamType);
 
-	if (PassShaders.VertexShader)
+	int32 DataOffset = 0;
+	if (PassShaders.VertexShader.IsValid())
 	{
-		FMeshDrawSingleShaderBindings ShaderBindings = SharedMeshDrawCommand.ShaderBindings.GetSingleShaderBindings(SF_Vertex);
+		FMeshDrawSingleShaderBindings ShaderBindings = SharedMeshDrawCommand.ShaderBindings.GetSingleShaderBindings(SF_Vertex, DataOffset);
 		PassShaders.VertexShader->GetShaderBindings(Scene, FeatureLevel, PrimitiveSceneProxy, MaterialRenderProxy, MaterialResource, DrawRenderState, ShaderElementData, ShaderBindings);
 	}
 
-	if (PassShaders.HullShader)
+	if (PassShaders.HullShader.IsValid() & PassShaders.DomainShader.IsValid())
 	{
-		FMeshDrawSingleShaderBindings ShaderBindings = SharedMeshDrawCommand.ShaderBindings.GetSingleShaderBindings(SF_Hull);
-		PassShaders.HullShader->GetShaderBindings(Scene, FeatureLevel, PrimitiveSceneProxy, MaterialRenderProxy, MaterialResource, DrawRenderState, ShaderElementData, ShaderBindings);
+		FMeshDrawSingleShaderBindings HullShaderBindings = SharedMeshDrawCommand.ShaderBindings.GetSingleShaderBindings(SF_Hull, DataOffset);
+		FMeshDrawSingleShaderBindings DomainShaderBindings = SharedMeshDrawCommand.ShaderBindings.GetSingleShaderBindings(SF_Domain, DataOffset);
+		PassShaders.HullShader->GetShaderBindings(Scene, FeatureLevel, PrimitiveSceneProxy, MaterialRenderProxy, MaterialResource, DrawRenderState, ShaderElementData, HullShaderBindings);
+		PassShaders.DomainShader->GetShaderBindings(Scene, FeatureLevel, PrimitiveSceneProxy, MaterialRenderProxy, MaterialResource, DrawRenderState, ShaderElementData, DomainShaderBindings);
 	}
 
-	if (PassShaders.DomainShader)
+	if (PassShaders.PixelShader.IsValid())
 	{
-		FMeshDrawSingleShaderBindings ShaderBindings = SharedMeshDrawCommand.ShaderBindings.GetSingleShaderBindings(SF_Domain);
-		PassShaders.DomainShader->GetShaderBindings(Scene, FeatureLevel, PrimitiveSceneProxy, MaterialRenderProxy, MaterialResource, DrawRenderState, ShaderElementData, ShaderBindings);
-	}
-
-	if (PassShaders.PixelShader)
-	{
-		FMeshDrawSingleShaderBindings ShaderBindings = SharedMeshDrawCommand.ShaderBindings.GetSingleShaderBindings(SF_Pixel);
+		FMeshDrawSingleShaderBindings ShaderBindings = SharedMeshDrawCommand.ShaderBindings.GetSingleShaderBindings(SF_Pixel, DataOffset);
 		PassShaders.PixelShader->GetShaderBindings(Scene, FeatureLevel, PrimitiveSceneProxy, MaterialRenderProxy, MaterialResource, DrawRenderState, ShaderElementData, ShaderBindings);
 	}
 
-	if (PassShaders.GeometryShader)
+	if (PassShaders.GeometryShader.IsValid())
 	{
-		FMeshDrawSingleShaderBindings ShaderBindings = SharedMeshDrawCommand.ShaderBindings.GetSingleShaderBindings(SF_Geometry);
+		FMeshDrawSingleShaderBindings ShaderBindings = SharedMeshDrawCommand.ShaderBindings.GetSingleShaderBindings(SF_Geometry, DataOffset);
 		PassShaders.GeometryShader->GetShaderBindings(Scene, FeatureLevel, PrimitiveSceneProxy, MaterialRenderProxy, MaterialResource, DrawRenderState, ShaderElementData, ShaderBindings);
 	}
 
@@ -95,33 +92,34 @@ void FMeshPassProcessor::BuildMeshDrawCommands(
 		if ((1ull << BatchElementIndex) & BatchElementMask)
 		{
 			const FMeshBatchElement& BatchElement = MeshBatch.Elements[BatchElementIndex];
-			FMeshDrawCommand& MeshDrawCommand = DrawListContext->AddCommand(SharedMeshDrawCommand);
+			FMeshDrawCommand& MeshDrawCommand = DrawListContext->AddCommand(SharedMeshDrawCommand, NumElements);
 
-			if (PassShaders.VertexShader)
+			DataOffset = 0;
+			if (PassShaders.VertexShader.IsValid())
 			{
-				FMeshDrawSingleShaderBindings VertexShaderBindings = MeshDrawCommand.ShaderBindings.GetSingleShaderBindings(SF_Vertex);
-				PassShaders.VertexShader->GetElementShaderBindings(Scene, ViewIfDynamicMeshCommand, VertexFactory, InputStreamType, FeatureLevel, PrimitiveSceneProxy, MeshBatch, BatchElement, ShaderElementData, VertexShaderBindings, MeshDrawCommand.VertexStreams);
+				FMeshDrawSingleShaderBindings VertexShaderBindings = MeshDrawCommand.ShaderBindings.GetSingleShaderBindings(SF_Vertex, DataOffset);
+				FMeshMaterialShader::GetElementShaderBindings(PassShaders.VertexShader, Scene, ViewIfDynamicMeshCommand, VertexFactory, InputStreamType, FeatureLevel, PrimitiveSceneProxy, MeshBatch, BatchElement, ShaderElementData, VertexShaderBindings, MeshDrawCommand.VertexStreams);
 			}
 
-			if (PassShaders.HullShader && PassShaders.DomainShader)
+			if (PassShaders.HullShader.IsValid() & PassShaders.DomainShader.IsValid())
 			{
-				FMeshDrawSingleShaderBindings HullShaderBindings = MeshDrawCommand.ShaderBindings.GetSingleShaderBindings(SF_Hull);
-				FMeshDrawSingleShaderBindings DomainShaderBindings = MeshDrawCommand.ShaderBindings.GetSingleShaderBindings(SF_Domain);
-				PassShaders.HullShader->GetElementShaderBindings(Scene, ViewIfDynamicMeshCommand, VertexFactory, EVertexInputStreamType::Default, FeatureLevel, PrimitiveSceneProxy, MeshBatch, BatchElement, ShaderElementData, HullShaderBindings, MeshDrawCommand.VertexStreams);
-				PassShaders.DomainShader->GetElementShaderBindings(Scene, ViewIfDynamicMeshCommand, VertexFactory, EVertexInputStreamType::Default, FeatureLevel, PrimitiveSceneProxy, MeshBatch, BatchElement, ShaderElementData, DomainShaderBindings, MeshDrawCommand.VertexStreams);
+				FMeshDrawSingleShaderBindings HullShaderBindings = MeshDrawCommand.ShaderBindings.GetSingleShaderBindings(SF_Hull, DataOffset);
+				FMeshDrawSingleShaderBindings DomainShaderBindings = MeshDrawCommand.ShaderBindings.GetSingleShaderBindings(SF_Domain, DataOffset);
+				FMeshMaterialShader::GetElementShaderBindings(PassShaders.HullShader, Scene, ViewIfDynamicMeshCommand, VertexFactory, EVertexInputStreamType::Default, FeatureLevel, PrimitiveSceneProxy, MeshBatch, BatchElement, ShaderElementData, HullShaderBindings, MeshDrawCommand.VertexStreams);
+				FMeshMaterialShader::GetElementShaderBindings(PassShaders.DomainShader, Scene, ViewIfDynamicMeshCommand, VertexFactory, EVertexInputStreamType::Default, FeatureLevel, PrimitiveSceneProxy, MeshBatch, BatchElement, ShaderElementData, DomainShaderBindings, MeshDrawCommand.VertexStreams);
 			}
 
-			if (PassShaders.PixelShader)
+			if (PassShaders.PixelShader.IsValid())
 			{
-				FMeshDrawSingleShaderBindings PixelShaderBindings = MeshDrawCommand.ShaderBindings.GetSingleShaderBindings(SF_Pixel);
-				PassShaders.PixelShader->GetElementShaderBindings(Scene, ViewIfDynamicMeshCommand, VertexFactory, EVertexInputStreamType::Default, FeatureLevel, PrimitiveSceneProxy, MeshBatch, BatchElement, ShaderElementData, PixelShaderBindings, MeshDrawCommand.VertexStreams);
+				FMeshDrawSingleShaderBindings PixelShaderBindings = MeshDrawCommand.ShaderBindings.GetSingleShaderBindings(SF_Pixel, DataOffset);
+				FMeshMaterialShader::GetElementShaderBindings(PassShaders.PixelShader, Scene, ViewIfDynamicMeshCommand, VertexFactory, EVertexInputStreamType::Default, FeatureLevel, PrimitiveSceneProxy, MeshBatch, BatchElement, ShaderElementData, PixelShaderBindings, MeshDrawCommand.VertexStreams);
 			}
 
 
-			if (PassShaders.GeometryShader)
+			if (PassShaders.GeometryShader.IsValid())
 			{
-				FMeshDrawSingleShaderBindings GeometryShaderBindings = MeshDrawCommand.ShaderBindings.GetSingleShaderBindings(SF_Geometry);
-				PassShaders.GeometryShader->GetElementShaderBindings(Scene, ViewIfDynamicMeshCommand, VertexFactory, EVertexInputStreamType::Default, FeatureLevel, PrimitiveSceneProxy, MeshBatch, BatchElement, ShaderElementData, GeometryShaderBindings, MeshDrawCommand.VertexStreams);
+				FMeshDrawSingleShaderBindings GeometryShaderBindings = MeshDrawCommand.ShaderBindings.GetSingleShaderBindings(SF_Geometry, DataOffset);
+				FMeshMaterialShader::GetElementShaderBindings(PassShaders.GeometryShader, Scene, ViewIfDynamicMeshCommand, VertexFactory, EVertexInputStreamType::Default, FeatureLevel, PrimitiveSceneProxy, MeshBatch, BatchElement, ShaderElementData, GeometryShaderBindings, MeshDrawCommand.VertexStreams);
 			}
 
 			int32 DrawPrimitiveId;
@@ -144,12 +142,13 @@ void DrawDynamicMeshPass(const FSceneView& View, FRHICommandList& RHICmdList, co
 	FDynamicMeshDrawCommandStorage DynamicMeshDrawCommandStorage;
 	FMeshCommandOneFrameArray VisibleMeshDrawCommands;
 	FGraphicsMinimalPipelineStateSet GraphicsMinimalPipelineStateSet;
+	bool NeedsShaderInitialisation;
 
-	FDynamicPassMeshDrawListContext DynamicMeshPassContext(DynamicMeshDrawCommandStorage, VisibleMeshDrawCommands, GraphicsMinimalPipelineStateSet);
+	FDynamicPassMeshDrawListContext DynamicMeshPassContext(DynamicMeshDrawCommandStorage, VisibleMeshDrawCommands, GraphicsMinimalPipelineStateSet, NeedsShaderInitialisation);
 
 	BuildPassProcessorLambda(&DynamicMeshPassContext);
 
 	// We assume all dynamic passes are in stereo if it is enabled in the view, so we apply ISR to them
 	const uint32 InstanceFactor = (!bForceStereoInstancingOff && View.IsInstancedStereoPass()) ? 2 : 1;
-	DrawDynamicMeshPassPrivate(View, RHICmdList, VisibleMeshDrawCommands, DynamicMeshDrawCommandStorage, GraphicsMinimalPipelineStateSet, InstanceFactor);
+	DrawDynamicMeshPassPrivate(View, RHICmdList, VisibleMeshDrawCommands, DynamicMeshDrawCommandStorage, GraphicsMinimalPipelineStateSet, NeedsShaderInitialisation, InstanceFactor);
 }
diff --git a/Engine/Source/Runtime/Renderer/Public/PostProcessParameters.h b/Engine/Source/Runtime/Renderer/Public/PostProcessParameters.h
--- a/Engine/Source/Runtime/Renderer/Public/PostProcessParameters.h
+++ b/Engine/Source/Runtime/Renderer/Public/PostProcessParameters.h
@@ -70,6 +70,8 @@ enum EFallbackColor
 // convenient but not the most optimized solution
 struct RENDERER_API FPostProcessPassParameters
 {
+	DECLARE_TYPE_LAYOUT(FPostProcessPassParameters, NonVirtual);
+public:
 	/** Initialization constructor. */
 	 void Bind(const FShaderParameterMap& ParameterMap);
 
@@ -88,17 +90,16 @@ struct RENDERER_API FPostProcessPassParameters
 	RENDERER_API friend FArchive& operator<<(FArchive& Ar,FPostProcessPassParameters& P);
 
 private:
-
-	FShaderParameter ViewportSize;
-	FShaderParameter ViewportRect;
-	FShaderResourceParameter PostprocessInputParameter[ePId_Input_MAX];
-	FShaderResourceParameter PostprocessInputParameterSampler[ePId_Input_MAX];
-	FShaderParameter PostprocessInputSizeParameter[ePId_Input_MAX];
-	FShaderParameter PostProcessInputMinMaxParameter[ePId_Input_MAX];
-	FShaderParameter ScreenPosToPixel;
-	FShaderParameter SceneColorBufferUVViewport;
-	FShaderResourceParameter BilinearTextureSampler;
-
+	LAYOUT_FIELD(FShaderParameter, ViewportSize);
+	LAYOUT_FIELD(FShaderParameter, ViewportRect);
+	LAYOUT_ARRAY(FShaderResourceParameter, PostprocessInputParameter, ePId_Input_MAX);
+	LAYOUT_ARRAY(FShaderResourceParameter, PostprocessInputParameterSampler, ePId_Input_MAX);
+	LAYOUT_ARRAY(FShaderParameter, PostprocessInputSizeParameter, ePId_Input_MAX);
+	LAYOUT_ARRAY(FShaderParameter, PostProcessInputMinMaxParameter, ePId_Input_MAX);
+	LAYOUT_FIELD(FShaderParameter, ScreenPosToPixel);
+	LAYOUT_FIELD(FShaderParameter, SceneColorBufferUVViewport);
+	LAYOUT_FIELD(FShaderResourceParameter, BilinearTextureSampler);
+	
 public:
 	// @param Filter can be 0 if FilterOverrideArray is used
 	// @param FilterOverrideArray can be 0 if Filter is used
diff --git a/Engine/Source/Runtime/Renderer/Public/PrimitiveSceneInfo.h b/Engine/Source/Runtime/Renderer/Public/PrimitiveSceneInfo.h
--- a/Engine/Source/Runtime/Renderer/Public/PrimitiveSceneInfo.h
+++ b/Engine/Source/Runtime/Renderer/Public/PrimitiveSceneInfo.h
@@ -14,6 +14,7 @@
 #include "Math/GenericOctreePublic.h"
 #include "Engine/Scene.h"
 #include "RendererInterface.h"
+#include "MeshPassProcessor.h"
 
 class FPrimitiveSceneInfo;
 class FPrimitiveSceneProxy;
@@ -213,6 +214,7 @@ typedef TOctree<FPrimitiveSceneInfoCompact,struct FPrimitiveOctreeSemantics> FSc
  */
 class FPrimitiveSceneInfo : public FDeferredCleanupInterface
 {
+	friend class FSceneRenderer;
 public:
 
 	/** The render proxy for the primitive. */
@@ -338,7 +340,7 @@ public:
 	~FPrimitiveSceneInfo();
 
 	/** Adds the primitive to the scene. */
-	void AddToScene(FRHICommandListImmediate& RHICmdList, bool bUpdateStaticDrawLists, bool bAddToStaticDrawLists = true, bool bAsyncCreateLPIs = false);
+	static void AddToScene(FRHICommandListImmediate& RHICmdList, FScene* Scene, const TArrayView<FPrimitiveSceneInfo*>& SceneInfos, bool bUpdateStaticDrawLists, bool bAddToStaticDrawLists = true, bool bAsyncCreateLPIs = false);
 
 	/** Removes the primitive from the scene. */
 	void RemoveFromScene(bool bUpdateStaticDrawLists);
@@ -359,16 +361,7 @@ public:
 	}
 
 	/** Updates the primitive's static meshes in the scene. */
-	void UpdateStaticMeshes(FRHICommandListImmediate& RHICmdList, bool bReAddToDrawLists = true);
-
-	/** Updates the primitive's static meshes in the scene. */
-	FORCEINLINE void ConditionalUpdateStaticMeshes(FRHICommandListImmediate& RHICmdList)
-	{
-		if (NeedsUpdateStaticMeshes())
-		{
-			UpdateStaticMeshes(RHICmdList);
-		}
-	}
+	static void UpdateStaticMeshes(FRHICommandListImmediate& RHICmdList, FScene* Scene, const TArrayView<FPrimitiveSceneInfo*>& SceneInfos, bool bReAddToDrawLists = true);
 
 	/** Updates the primitive's uniform buffer. */
 	void UpdateUniformBuffer(FRHICommandListImmediate& RHICmdList);
@@ -389,7 +382,7 @@ public:
 	void BeginDeferredUpdateStaticMeshesWithoutVisibilityCheck();
 
 	/** Adds the primitive's static meshes to the scene. */
-	void AddStaticMeshes(FRHICommandListImmediate& RHICmdList, bool bUpdateStaticDrawLists = true);
+	static void AddStaticMeshes(FRHICommandListImmediate& RHICmdList, FScene* Scene, const TArrayView<FPrimitiveSceneInfo*>& SceneInfos, bool bUpdateStaticDrawLists = true);
 
 	/** Removes the primitive's static meshes from the scene. */
 	void RemoveStaticMeshes();
@@ -451,11 +444,13 @@ public:
 
 	FORCEINLINE void MarkIndirectLightingCacheBufferDirty()
 	{
-		bIndirectLightingCacheBufferDirty = true;
+		if (!bIndirectLightingCacheBufferDirty)
+		{
+			bIndirectLightingCacheBufferDirty = true;
+		}
 	}
 
 	void UpdateIndirectLightingCacheBuffer();
-	void ClearIndirectLightingCacheBuffer(bool bSingleFrameOnly);
 
 	/** Will output the LOD ranges of the static meshes used with this primitive. */
 	RENDERER_API void GetStaticMeshesLODRange(int8& OutMinLOD, int8& OutMaxLOD) const;
@@ -537,7 +532,7 @@ private:
 		class FVolumetricLightmapSceneData* VolumetricLightmapSceneData);
 
 	/** Creates cached mesh draw commands for all meshes. */
-	void CacheMeshDrawCommands(FRHICommandListImmediate& RHICmdList);
+	static void CacheMeshDrawCommands(FRHICommandListImmediate& RHICmdList, FScene* Scene, const TArrayView<FPrimitiveSceneInfo*>& SceneInfos);
 
 	/** Removes cached mesh draw commands for all meshes. */
 	void RemoveCachedMeshDrawCommands();
diff --git a/Engine/Source/Runtime/Renderer/Public/RendererUtils.h b/Engine/Source/Runtime/Renderer/Public/RendererUtils.h
--- a/Engine/Source/Runtime/Renderer/Public/RendererUtils.h
+++ b/Engine/Source/Runtime/Renderer/Public/RendererUtils.h
@@ -9,9 +9,9 @@
 #include "ShaderParameterUtils.h"
 #include "RendererInterface.h"
 
-class FRenderTargetWriteMask
+class RENDERER_API FRenderTargetWriteMask
 {
 public:
 	template <uint32 NumRenderTargets>
-	static void Decode(FRHICommandListImmediate& RHICmdList, TShaderMap<FGlobalShaderType>* ShaderMap, IPooledRenderTarget* InRenderTargets[NumRenderTargets], TRefCountPtr<IPooledRenderTarget>& OutRTWriteMask, uint32 RTWriteMaskFastVRamConfig, const TCHAR* RTWriteMaskDebugName);
-};
\ No newline at end of file
+	static void Decode(FRHICommandListImmediate& RHICmdList, FGlobalShaderMap* ShaderMap, IPooledRenderTarget* InRenderTargets[NumRenderTargets], TRefCountPtr<IPooledRenderTarget>& OutRTWriteMask, uint32 RTWriteMaskFastVRamConfig, const TCHAR* RTWriteMaskDebugName);
+};	
\ No newline at end of file
diff --git a/Engine/Source/Runtime/Renderer/Public/SceneRenderTargetParameters.h b/Engine/Source/Runtime/Renderer/Public/SceneRenderTargetParameters.h
--- a/Engine/Source/Runtime/Renderer/Public/SceneRenderTargetParameters.h
+++ b/Engine/Source/Runtime/Renderer/Public/SceneRenderTargetParameters.h
@@ -133,6 +133,7 @@ extern RENDERER_API void BindSceneTextureUniformBufferDependentOnShadingPath(
 /** Encapsulates scene texture shader parameter bindings. */
 class RENDERER_API FSceneTextureShaderParameters
 {
+	DECLARE_TYPE_LAYOUT(FSceneTextureShaderParameters, NonVirtual);
 public:
 	/** Binds the parameters using a compiled shader's parameter map. */
 	void Bind(const FShader::CompiledShaderInitializerType& Initializer)
@@ -183,5 +184,5 @@ public:
 	}
 
 private:
-	FShaderUniformBufferParameter SceneTexturesUniformBuffer;
+	LAYOUT_FIELD(FShaderUniformBufferParameter, SceneTexturesUniformBuffer);
 };
diff --git a/Engine/Source/Runtime/Renderer/Public/VolumeRendering.h b/Engine/Source/Runtime/Renderer/Public/VolumeRendering.h
--- a/Engine/Source/Runtime/Renderer/Public/VolumeRendering.h
+++ b/Engine/Source/Runtime/Renderer/Public/VolumeRendering.h
@@ -73,25 +73,17 @@ public:
 	{
 		const float InvVolumeResolutionX = 1.0f / VolumeResolution.X;
 		const float InvVolumeResolutionY = 1.0f / VolumeResolution.Y;
-		SetShaderValue(RHICmdList, GetVertexShader(), UVScaleBias, FVector4(
+		SetShaderValue(RHICmdList, RHICmdList.GetBoundVertexShader(), UVScaleBias, FVector4(
 			(VolumeBounds.MaxX - VolumeBounds.MinX) * InvVolumeResolutionX,
 			(VolumeBounds.MaxY - VolumeBounds.MinY) * InvVolumeResolutionY,
 			VolumeBounds.MinX * InvVolumeResolutionX,
 			VolumeBounds.MinY * InvVolumeResolutionY));
-		SetShaderValue(RHICmdList, GetVertexShader(), MinZ, VolumeBounds.MinZ);
-	}
-
-	virtual bool Serialize(FArchive& Ar) override
-	{
-		bool bShaderHasOutdatedParameters = FGlobalShader::Serialize(Ar);
-		Ar << UVScaleBias;
-		Ar << MinZ;
-		return bShaderHasOutdatedParameters;
+		SetShaderValue(RHICmdList, RHICmdList.GetBoundVertexShader(), MinZ, VolumeBounds.MinZ);
 	}
 
 private:
-	FShaderParameter UVScaleBias;
-	FShaderParameter MinZ;
+	LAYOUT_FIELD(FShaderParameter, UVScaleBias);
+	LAYOUT_FIELD(FShaderParameter, MinZ);
 };
 
 /** Geometry shader used to write to a range of slices of a 3d volume texture. */
@@ -115,18 +107,11 @@ public:
 	template <typename TRHICommandList>
 	void SetParameters(TRHICommandList& RHICmdList, int32 MinZValue)
 	{
-		SetShaderValue(RHICmdList, GetGeometryShader(), MinZ, MinZValue);
-	}
-
-	virtual bool Serialize(FArchive& Ar) override
-	{
-		bool bShaderHasOutdatedParameters = FGlobalShader::Serialize(Ar);
-		Ar << MinZ;
-		return bShaderHasOutdatedParameters;
+		SetShaderValue(RHICmdList, RHICmdList.GetBoundGeometryShader(), MinZ, MinZValue);
 	}
 
 private:
-	FShaderParameter MinZ;
+	LAYOUT_FIELD(FShaderParameter, MinZ);
 };
 
 // This function assumes the PSO had a PrimitiveType of PT_TriangleStrip
diff --git a/Engine/Source/Runtime/SlateRHIRenderer/Private/SlateMaterialShader.cpp b/Engine/Source/Runtime/SlateRHIRenderer/Private/SlateMaterialShader.cpp
--- a/Engine/Source/Runtime/SlateRHIRenderer/Private/SlateMaterialShader.cpp
+++ b/Engine/Source/Runtime/SlateRHIRenderer/Private/SlateMaterialShader.cpp
@@ -4,6 +4,8 @@
 #include "Materials/Material.h"
 #include "ShaderParameterUtils.h"
 
+IMPLEMENT_TYPE_LAYOUT(FSlateMaterialShaderVS);
+IMPLEMENT_TYPE_LAYOUT(FSlateMaterialShaderPS);
 
 FSlateMaterialShaderVS::FSlateMaterialShaderVS(const FMaterialShaderType::CompiledShaderInitializerType& Initializer)
 	: FMaterialShader(Initializer)
@@ -17,36 +19,36 @@ void FSlateMaterialShaderVS::ModifyCompilationEnvironment(const FMaterialShaderP
 {
 	// Set defines based on what this shader will be used for
 	OutEnvironment.SetDefine( TEXT("USE_MATERIALS"), 1 );
-	OutEnvironment.SetDefine( TEXT("NUM_CUSTOMIZED_UVS"), Parameters.Material->GetNumCustomizedUVs() );
-	OutEnvironment.SetDefine( TEXT("HAS_SCREEN_POSITION"), Parameters.Material->HasVertexPositionOffsetConnected() );
+	OutEnvironment.SetDefine( TEXT("NUM_CUSTOMIZED_UVS"), Parameters.MaterialParameters.NumCustomizedUVs );
+	OutEnvironment.SetDefine( TEXT("HAS_SCREEN_POSITION"), (bool)Parameters.MaterialParameters.bHasVertexPositionOffsetConnected );
 
 	FMaterialShader::ModifyCompilationEnvironment( Parameters, OutEnvironment );
 }
 
 bool FSlateMaterialShaderVS::ShouldCompilePermutation(const FMaterialShaderPermutationParameters& Parameters)
 {
-	return Parameters.Material->GetMaterialDomain() == MD_UI;
+	return Parameters.MaterialParameters.MaterialDomain == MD_UI;
 }
 
 void FSlateMaterialShaderVS::SetViewProjection(FRHICommandList& RHICmdList, const FMatrix& InViewProjection )
 {
-	SetShaderValue(RHICmdList, GetVertexShader(), ViewProjection, InViewProjection );
+	SetShaderValue(RHICmdList, RHICmdList.GetBoundVertexShader(), ViewProjection, InViewProjection );
 }
 
 void FSlateMaterialShaderVS::SetMaterialShaderParameters(FRHICommandList& RHICmdList, const FSceneView& View, const FMaterialRenderProxy* MaterialRenderProxy, const FMaterial* Material)
 {
-	FRHIVertexShader* ShaderRHI = GetVertexShader();
+	FRHIVertexShader* ShaderRHI = RHICmdList.GetBoundVertexShader();
 
 	FMaterialShader::SetParameters<FRHIVertexShader>(RHICmdList, ShaderRHI, MaterialRenderProxy, *Material, View, View.ViewUniformBuffer, ESceneTextureSetupMode::None);
 }
 
 void FSlateMaterialShaderVS::SetVerticalAxisMultiplier(FRHICommandList& RHICmdList, float InMultiplier )
 {
-	SetShaderValue(RHICmdList, GetVertexShader(), SwitchVerticalAxisMultiplier, InMultiplier );
+	SetShaderValue(RHICmdList, RHICmdList.GetBoundVertexShader(), SwitchVerticalAxisMultiplier, InMultiplier );
 }
 
 
-bool FSlateMaterialShaderVS::Serialize(FArchive& Ar)
+/*bool FSlateMaterialShaderVS::Serialize(FArchive& Ar)
 {
 	bool bShaderHasOutdatedParameters = FMaterialShader::Serialize(Ar);
 
@@ -54,12 +56,12 @@ bool FSlateMaterialShaderVS::Serialize(FArchive& Ar)
 	Ar << SwitchVerticalAxisMultiplier;
 
 	return bShaderHasOutdatedParameters;
-}
+}*/
 
 
 bool FSlateMaterialShaderPS::ShouldCompilePermutation(const FMaterialShaderPermutationParameters& Parameters)
 {
-	return Parameters.Material->GetMaterialDomain() == MD_UI;
+	return Parameters.MaterialParameters.MaterialDomain == MD_UI;
 }
 
 
@@ -67,7 +69,7 @@ void FSlateMaterialShaderPS::ModifyCompilationEnvironment(const FMaterialShaderP
 {
 	// Set defines based on what this shader will be used for
 	OutEnvironment.SetDefine( TEXT("USE_MATERIALS"), 1 );
-	OutEnvironment.SetDefine( TEXT("NUM_CUSTOMIZED_UVS"), Parameters.Material->GetNumCustomizedUVs() );
+	OutEnvironment.SetDefine( TEXT("NUM_CUSTOMIZED_UVS"), Parameters.MaterialParameters.NumCustomizedUVs);
 
 	FMaterialShader::ModifyCompilationEnvironment( Parameters, OutEnvironment );
 }
@@ -120,7 +122,7 @@ void FSlateMaterialShaderPS::SetBlendState(FGraphicsPipelineStateInitializer& Gr
 
 void FSlateMaterialShaderPS::SetParameters(FRHICommandList& RHICmdList, const FSceneView& View, const FMaterialRenderProxy* MaterialRenderProxy, const FMaterial* Material, const FVector4& InShaderParams)
 {
-	FRHIPixelShader* ShaderRHI = GetPixelShader();
+	FRHIPixelShader* ShaderRHI = RHICmdList.GetBoundPixelShader();
 
 	SetShaderValue( RHICmdList, ShaderRHI, ShaderParams, InShaderParams );
 
@@ -130,32 +132,21 @@ void FSlateMaterialShaderPS::SetParameters(FRHICommandList& RHICmdList, const FS
 
 void FSlateMaterialShaderPS::SetAdditionalTexture( FRHICommandList& RHICmdList, FRHITexture* InTexture, const FSamplerStateRHIRef SamplerState )
 {
-	SetTextureParameter(RHICmdList, GetPixelShader(), AdditionalTextureParameter, TextureParameterSampler, SamplerState, InTexture );
+	SetTextureParameter(RHICmdList, RHICmdList.GetBoundPixelShader(), AdditionalTextureParameter, TextureParameterSampler, SamplerState, InTexture );
 }
 
 void FSlateMaterialShaderPS::SetDisplayGammaAndContrast(FRHICommandList& RHICmdList, float InDisplayGamma, float InContrast)
 {
 	FVector4 InGammaValues(2.2f / InDisplayGamma, 1.0f / InDisplayGamma, 0.0f, InContrast);
 
-	SetShaderValue(RHICmdList, GetPixelShader(), GammaAndAlphaValues, InGammaValues);
+	SetShaderValue(RHICmdList, RHICmdList.GetBoundPixelShader(), GammaAndAlphaValues, InGammaValues);
 }
 
 void FSlateMaterialShaderPS::SetDrawFlags(FRHICommandList& RHICmdList, bool bDrawDisabledEffect)
 {
 	FVector4 InDrawFlags(bDrawDisabledEffect ? 1 : 0, 0, 0, 0);
 
-	SetShaderValue(RHICmdList, GetPixelShader(), DrawFlags, InDrawFlags);
-}
-
-bool FSlateMaterialShaderPS::Serialize(FArchive& Ar)
-{
-	bool bShaderHasOutdatedParameters = FMaterialShader::Serialize(Ar);
-	Ar << GammaAndAlphaValues;
-	Ar << DrawFlags;
-	Ar << ShaderParams;
-	Ar << TextureParameterSampler;
-	Ar << AdditionalTextureParameter;
-	return bShaderHasOutdatedParameters;
+	SetShaderValue(RHICmdList, RHICmdList.GetBoundPixelShader(), DrawFlags, InDrawFlags);
 }
 
 
diff --git a/Engine/Source/Runtime/SlateRHIRenderer/Private/SlateMaterialShader.h b/Engine/Source/Runtime/SlateRHIRenderer/Private/SlateMaterialShader.h
--- a/Engine/Source/Runtime/SlateRHIRenderer/Private/SlateMaterialShader.h
+++ b/Engine/Source/Runtime/SlateRHIRenderer/Private/SlateMaterialShader.h
@@ -11,6 +11,7 @@
 
 class FSlateMaterialShaderVS : public FMaterialShader
 {
+	DECLARE_TYPE_LAYOUT(FSlateMaterialShaderVS, NonVirtual);
 public:
 	FSlateMaterialShaderVS() {}
 	FSlateMaterialShaderVS(const FMaterialShaderType::CompiledShaderInitializerType& Initializer);
@@ -34,16 +35,19 @@ public:
 	void SetVerticalAxisMultiplier(FRHICommandList& RHICmdList, float InMultiplier);
 
 	/** Serializes the shader data */
-	virtual bool Serialize( FArchive& Ar ) override;
+	//virtual bool Serialize( FArchive& Ar ) override;
 private:
-	/** ViewProjection parameter used by the shader */
-	FShaderParameter ViewProjection;
-	/** Parameter used to determine if we need to swtich the vertical axis for opengl */
-	FShaderParameter SwitchVerticalAxisMultiplier;
+	
+		/** ViewProjection parameter used by the shader */
+		LAYOUT_FIELD(FShaderParameter, ViewProjection)
+		/** Parameter used to determine if we need to swtich the vertical axis for opengl */
+		LAYOUT_FIELD(FShaderParameter, SwitchVerticalAxisMultiplier)
+	
 };
 
 class FSlateMaterialShaderPS : public FMaterialShader
 {
+	DECLARE_TYPE_LAYOUT(FSlateMaterialShaderPS, NonVirtual);
 public:
 
 	/** Only compile shaders used with UI. */
@@ -65,15 +69,13 @@ public:
 
 	void SetDrawFlags(FRHICommandList& RHICmdList, bool bDrawDisabledEffect);
 
-	virtual bool Serialize(FArchive& Ar) override;
-
 private:
-	FShaderParameter GammaAndAlphaValues;
-	FShaderParameter DrawFlags;
-	FShaderParameter ShaderParams;
+	LAYOUT_FIELD(FShaderParameter, GammaAndAlphaValues);
+	LAYOUT_FIELD(FShaderParameter, DrawFlags);
+	LAYOUT_FIELD(FShaderParameter, ShaderParams);
 	/** Extra texture (like a font atlas) to be used in addition to any material textures */
-	FShaderResourceParameter TextureParameterSampler;
-	FShaderResourceParameter AdditionalTextureParameter;
+	LAYOUT_FIELD(FShaderResourceParameter, TextureParameterSampler);
+	LAYOUT_FIELD(FShaderResourceParameter, AdditionalTextureParameter);
 };
 
 template<bool bUseInstancing>
@@ -101,11 +103,6 @@ public:
 
 		OutEnvironment.SetDefine(TEXT("USE_SLATE_INSTANCING"), (uint32)( bUseInstancing ? 1 : 0 ));
 	}
-
-	virtual bool Serialize(FArchive& Ar) override
-	{
-		return FSlateMaterialShaderVS::Serialize( Ar );
-	}
 };
 
 template<ESlateShader ShaderType>
@@ -133,9 +130,4 @@ public:
 
 		OutEnvironment.SetDefine(TEXT("SHADER_TYPE"), (uint32)ShaderType);
 	}
-
-	virtual bool Serialize(FArchive& Ar) override
-	{
-		return FSlateMaterialShaderPS::Serialize( Ar );
-	}
 };
