More missing includes
#rnx

diff --git a/Engine/Source/Runtime/Renderer/Private/AtmosphereRendering.h b/Engine/Source/Runtime/Renderer/Private/AtmosphereRendering.h
--- a/Engine/Source/Runtime/Renderer/Private/AtmosphereRendering.h
+++ b/Engine/Source/Runtime/Renderer/Private/AtmosphereRendering.h
@@ -9,6 +9,7 @@
 #include "CoreMinimal.h"
 #include "RenderResource.h"
 #include "Serialization/BulkData.h"
+#include "Shader.h"
 #include "RendererInterface.h"
 
 class FLightSceneInfo;
diff --git a/Engine/Source/Runtime/Renderer/Private/GPUFastFourierTransform.h b/Engine/Source/Runtime/Renderer/Private/GPUFastFourierTransform.h
--- a/Engine/Source/Runtime/Renderer/Private/GPUFastFourierTransform.h
+++ b/Engine/Source/Runtime/Renderer/Private/GPUFastFourierTransform.h
@@ -8,6 +8,7 @@ GPUFastFourierTransform.h: Interface for Fast Fourier Transform (FFT) on GPU.
 
 #include "CoreMinimal.h"
 #include "Math/NumericLimits.h"
+#include "GlobalShader.h"
 #include "RendererInterface.h"
 #include "RHI.h"
 #include "RHIStaticStates.h"
