Add option to disable custom depth rendering on mobile even translucency uses it

diff --git a/Engine/Source/Runtime/Renderer/Private/MobileShadingRenderer.cpp b/Engine/Source/Runtime/Renderer/Private/MobileShadingRenderer.cpp
--- a/Engine/Source/Runtime/Renderer/Private/MobileShadingRenderer.cpp
+++ b/Engine/Source/Runtime/Renderer/Private/MobileShadingRenderer.cpp
@@ -74,6 +74,14 @@ static TAutoConsoleVariable<int32> CVarMobileAdrenoOcclusionMode(
 	TEXT("1: Render occlusion queries after translucency and a flush, which can help Adreno devices in GL mode."),
 	ECVF_RenderThreadSafe);
 
+static TAutoConsoleVariable<int32> CVarMobileCustomDepthForTranslucency(
+	TEXT("r.Mobile.CustomDepthForTranslucency"),
+	1,
+	TEXT(" Whether to render custom depth/stencil if any tranclucency in the scene uses it. \n")
+	TEXT(" 0 = Off \n")
+	TEXT(" 1 = On [default]"),
+	ECVF_Scalability | ECVF_RenderThreadSafe);
+
 DECLARE_GPU_STAT_NAMED(MobileSceneRender, TEXT("Mobile Scene Render"));
 
 DECLARE_CYCLE_STAT(TEXT("SceneStart"), STAT_CLMM_SceneStart, STATGROUP_CommandListMarkers);
@@ -93,7 +101,7 @@ TGlobalResource<FGlobalDynamicReadBuffer> FMobileSceneRenderer::DynamicReadBuffe
 
 static bool UsesCustomDepthStencilLookup(const FViewInfo& View)
 {
-	if (View.bUsesCustomDepthStencil)
+	if (View.bUsesCustomDepthStencil && CVarMobileCustomDepthForTranslucency.GetValueOnAnyThread() != 0)
 	{
 		return true;
 	}
