Moved individual light count csv stats to LightCount category

#rnx

diff --git a/Engine/Source/Runtime/Renderer/Private/DeferredShadingRenderer.cpp b/Engine/Source/Runtime/Renderer/Private/DeferredShadingRenderer.cpp
--- a/Engine/Source/Runtime/Renderer/Private/DeferredShadingRenderer.cpp
+++ b/Engine/Source/Runtime/Renderer/Private/DeferredShadingRenderer.cpp
@@ -222,6 +222,8 @@ DECLARE_GPU_STAT(VirtualTextureUpdate);
 DECLARE_GPU_STAT(UploadDynamicBuffers);
 DECLARE_GPU_STAT(PostOpaqueExtensions);
 
+CSV_DEFINE_CATEGORY(LightCount, true);
+
 const TCHAR* GetDepthPassReason(bool bDitheredLODTransitionsUseStencil, EShaderPlatform ShaderPlatform)
 {
 	if (IsForwardShadingEnabled(ShaderPlatform))
@@ -1705,9 +1707,9 @@ void FDeferredShadingSceneRenderer::Render(FRHICommandListImmediate& RHICmdList)
 
 	}
 
-	CSV_CUSTOM_STAT_GLOBAL(LightCount, float(SortedLightSet.SortedLights.Num()), ECsvCustomStatOp::Set);
-	CSV_CUSTOM_STAT_GLOBAL(LightCountShadowOff, float(SortedLightSet.AttenuationLightStart), ECsvCustomStatOp::Set);
-	CSV_CUSTOM_STAT_GLOBAL(LightCountShadowOn, float(SortedLightSet.SortedLights.Num()) - float(SortedLightSet.AttenuationLightStart), ECsvCustomStatOp::Set);
+	CSV_CUSTOM_STAT(LightCount, All,  float(SortedLightSet.SortedLights.Num()), ECsvCustomStatOp::Set);
+	CSV_CUSTOM_STAT(LightCount, ShadowOff, float(SortedLightSet.AttenuationLightStart), ECsvCustomStatOp::Set);
+	CSV_CUSTOM_STAT(LightCount, ShadowOn, float(SortedLightSet.SortedLights.Num()) - float(SortedLightSet.AttenuationLightStart), ECsvCustomStatOp::Set);
 
 	{
 		SCOPE_CYCLE_COUNTER(STAT_FDeferredShadingSceneRenderer_AllocGBufferTargets);
diff --git a/Engine/Source/Runtime/Renderer/Private/ShadowSetup.cpp b/Engine/Source/Runtime/Renderer/Private/ShadowSetup.cpp
--- a/Engine/Source/Runtime/Renderer/Private/ShadowSetup.cpp
+++ b/Engine/Source/Runtime/Renderer/Private/ShadowSetup.cpp
@@ -246,6 +246,8 @@ FAutoConsoleVariableRef CVarUseOctreeForShadowCulling(
 	ECVF_Scalability | ECVF_RenderThreadSafe
 	);
 
+CSV_DECLARE_CATEGORY_EXTERN(LightCount);
+
 #if !UE_BUILD_SHIPPING
 // read and written on the render thread
 bool GDumpShadowSetup = false;
@@ -4108,7 +4110,6 @@ void FSceneRenderer::InitDynamicShadows(FRHICommandListImmediate& RHICmdList, FG
 	const bool bProjectEnablePointLightShadows = Scene->ReadOnlyCVARCache.bEnablePointLightShadows;
 	uint32 NumPointShadowCachesUpdatedThisFrame = 0;
 	uint32 NumSpotShadowCachesUpdatedThisFrame = 0;
-	uint32 NumInteractionShadows = 0;
 
 	TArray<FProjectedShadowInfo*,SceneRenderingAllocator> PreShadows;
 	TArray<FProjectedShadowInfo*,SceneRenderingAllocator> ViewDependentWholeSceneShadows;
@@ -4220,7 +4221,6 @@ void FSceneRenderer::InitDynamicShadows(FRHICommandListImmediate& RHICmdList, FG
 								)
 							{
 								SetupInteractionShadows(RHICmdList, Interaction, VisibleLightInfo, bStaticSceneOnly, ViewDependentWholeSceneShadows, PreShadows);
-								NumInteractionShadows++;
 							}
 
 							for (FLightPrimitiveInteraction* Interaction = LightSceneInfo->GetDynamicInteractionStaticPrimitiveList(false);
@@ -4229,7 +4229,6 @@ void FSceneRenderer::InitDynamicShadows(FRHICommandListImmediate& RHICmdList, FG
 								)
 							{
 								SetupInteractionShadows(RHICmdList, Interaction, VisibleLightInfo, bStaticSceneOnly, ViewDependentWholeSceneShadows, PreShadows);
-								NumInteractionShadows++;
 							}
 						}
 					}
@@ -4237,8 +4236,7 @@ void FSceneRenderer::InitDynamicShadows(FRHICommandListImmediate& RHICmdList, FG
 			}
 		}
 
-		CSV_CUSTOM_STAT_GLOBAL(ShadowTotalLightUpdated, float(NumPointShadowCachesUpdatedThisFrame + NumSpotShadowCachesUpdatedThisFrame), ECsvCustomStatOp::Set);
-		CSV_CUSTOM_STAT_GLOBAL(ShadowInteractionShadowsCount, float(NumInteractionShadows), ECsvCustomStatOp::Set);
+		CSV_CUSTOM_STAT(LightCount, UpdatedShadowMaps, float(NumPointShadowCachesUpdatedThisFrame + NumSpotShadowCachesUpdatedThisFrame), ECsvCustomStatOp::Set);
 
 		// Calculate visibility of the projected shadows.
 		InitProjectedShadowVisibility(RHICmdList);
