Make sure light shaft downsample pass output rect have the same size as the input rect of the light shaft TAA pass.

diff --git a/Engine/Source/Runtime/Renderer/Private/LightShaftRendering.cpp b/Engine/Source/Runtime/Renderer/Private/LightShaftRendering.cpp
--- a/Engine/Source/Runtime/Renderer/Private/LightShaftRendering.cpp
+++ b/Engine/Source/Runtime/Renderer/Private/LightShaftRendering.cpp
@@ -477,9 +477,10 @@ void DownsamplePass(FRHICommandListImmediate& RHICmdList, const FViewInfo& View,
 	const FIntPoint BufferSize = FSceneRenderTargets::Get(RHICmdList).GetBufferSizeXY();
 	const uint32 DownsampleFactor	= GetLightShaftDownsampleFactor();
 	const FIntPoint FilterBufferSize = BufferSize / DownsampleFactor;
-	const FIntPoint DownSampledXY = View.ViewRect.Min / DownsampleFactor;
-	const uint32 DownsampledSizeX = View.ViewRect.Width() / DownsampleFactor;
-	const uint32 DownsampledSizeY = View.ViewRect.Height() / DownsampleFactor;
+	const FIntRect DownsampledRect = FIntRect::DivideAndRoundUp(View.ViewRect, DownsampleFactor);
+	const FIntPoint DownSampledXY = DownsampledRect.Min;
+	const uint32 DownsampledSizeX = DownsampledRect.Width();
+	const uint32 DownsampledSizeY = DownsampledRect.Height();
 
 	FRHIRenderPassInfo RPInfo(LightShaftsDest->GetRenderTargetItem().TargetableTexture, ERenderTargetActions::Load_Store);
 	TransitionRenderPassTargets(RHICmdList, RPInfo);
