FMinimalBoundShaderStateInput holds ref-counted pointers to shader map resource
- Without this, it's possible to get state pointers held by persistent pipeline state ids

diff --git a/Engine/Source/Runtime/Renderer/Public/MeshPassProcessor.h b/Engine/Source/Runtime/Renderer/Public/MeshPassProcessor.h
--- a/Engine/Source/Runtime/Renderer/Public/MeshPassProcessor.h
+++ b/Engine/Source/Runtime/Renderer/Public/MeshPassProcessor.h
@@ -180,13 +180,11 @@ struct FMinimalBoundShaderStateInput
 	}
 
 	FRHIVertexDeclaration* VertexDeclarationRHI = nullptr;
-	// TODO - all of these resources should always (?) come from the same shader map
-	// can we enforce this, and only store a single resource pointer?
-	class FShaderMapResource* VertexShaderResource = nullptr;
-	class FShaderMapResource* HullShaderResource = nullptr;
-	class FShaderMapResource* DomainShaderResource = nullptr;
-	class FShaderMapResource* PixelShaderResource = nullptr;
-	class FShaderMapResource* GeometryShaderResource = nullptr;
+	TRefCountPtr<FShaderMapResource> VertexShaderResource;
+	TRefCountPtr<FShaderMapResource> HullShaderResource;
+	TRefCountPtr<FShaderMapResource> DomainShaderResource;
+	TRefCountPtr<FShaderMapResource> PixelShaderResource;
+	TRefCountPtr<FShaderMapResource> GeometryShaderResource;
 	int32 VertexShaderIndex = INDEX_NONE;
 	int32 HullShaderIndex = INDEX_NONE;
 	int32 DomainShaderIndex = INDEX_NONE;
