Change mobile custom stencil format from B8G8R8A8 to G8, save depth to R16F color target and use memoryless for depth target to reduce memory cost on mobile platform.
Add a Cvar for using half-res custom depth.
Add a seperate permuation for mobile custom depth stencil pixel shader, so mobile depth prepass could not be affected.
Support PostprocessMaterial stencil test on mobile platform base on the new mobile custom depth implementation

diff --git a/Engine/Shaders/Private/DepthOnlyPixelShader.usf b/Engine/Shaders/Private/DepthOnlyPixelShader.usf
--- a/Engine/Shaders/Private/DepthOnlyPixelShader.usf
+++ b/Engine/Shaders/Private/DepthOnlyPixelShader.usf
@@ -8,14 +8,17 @@
 #include "/Engine/Generated/Material.ush"
 #include "/Engine/Generated/VertexFactory.ush"
 
-#if (FEATURE_LEVEL <= FEATURE_LEVEL_ES3_1)
+#if OUTPUT_MOBILE_COLOR_VALUE
 // used to write custom stencil values to color target
 float MobileColorValue;
 #endif
 
 void Main(
-#if !MATERIALBLENDING_SOLID || OUTPUT_PIXEL_DEPTH_OFFSET
+#if !MATERIALBLENDING_SOLID || OUTPUT_PIXEL_DEPTH_OFFSET || OUTPUT_MOBILE_COLOR_VALUE
 	in INPUT_POSITION_QUALIFIERS float4 SvPosition : SV_Position,
+#endif
+
+#if !MATERIALBLENDING_SOLID || OUTPUT_PIXEL_DEPTH_OFFSET
 #if USE_WORLD_POSITION_EXCLUDING_SHADER_OFFSETS
 	float3 PixelPositionExcludingWPO : TEXCOORD7,
 #endif
@@ -24,7 +27,12 @@ void Main(
 	OPTIONAL_IsFrontFace
 	OPTIONAL_OutDepthConservative,
 #endif
+#if OUTPUT_MOBILE_COLOR_VALUE
+	out half OutDepth : SV_Target0,
+	out uint OutStencil : SV_Target1
+#else
 	out float4 OutColor : SV_Target0
+#endif
 #if MATERIALBLENDING_MASKED_USING_COVERAGE
 	, out uint OutCoverage : SV_Coverage
 #endif
@@ -67,8 +75,9 @@ void Main(
 	#endif
 #endif
 
-#if (FEATURE_LEVEL <= FEATURE_LEVEL_ES3_1)
-	OutColor = MobileColorValue;
+#if OUTPUT_MOBILE_COLOR_VALUE
+	OutDepth = min(SvPosition.w, 65500);
+	OutStencil = MobileColorValue;
 #else
 	OutColor = 0;
 #endif
diff --git a/Engine/Shaders/Private/MaterialTemplate.ush b/Engine/Shaders/Private/MaterialTemplate.ush
--- a/Engine/Shaders/Private/MaterialTemplate.ush
+++ b/Engine/Shaders/Private/MaterialTemplate.ush
@@ -1661,7 +1661,7 @@ MaterialFloat4 MobileSceneTextureLookup(inout FMaterialPixelParameters Parameter
 	}
 	else if (SceneTextureId == PPI_CustomDepth)
 	{
-		MaterialFloat Depth = ConvertFromDeviceZ(Texture2DSample(MobileSceneTextures.CustomDepthTexture, MobileSceneTextures.CustomDepthTextureSampler, UV).r);
+		MaterialFloat Depth = Texture2DSample(MobileSceneTextures.CustomDepthTexture, MobileSceneTextures.CustomDepthTextureSampler, UV).r;
 		return MaterialFloat4(Depth.rrr, 0.f);
 	}
 	else if (SceneTextureId == PPI_PostProcessInput0)
@@ -1682,8 +1682,7 @@ MaterialFloat4 MobileSceneTextureLookup(inout FMaterialPixelParameters Parameter
 	}
 	else if (SceneTextureId == PPI_CustomStencil)
 	{
-		MaterialFloat Stencil = Texture2DSample(MobileSceneTextures.MobileCustomStencilTexture, MobileSceneTextures.MobileCustomStencilTextureSampler, UV).r*255.0;
-		Stencil = floor(Stencil + 0.5);
+		MaterialFloat Stencil = MobileSceneTextures.CustomStencilTexture.Load(int3(Parameters.SvPosition.xy, 0));
 		return MaterialFloat4(Stencil.rrr, 0.f);
 	}
 #endif// FEATURE_LEVEL
diff --git a/Engine/Shaders/Private/PostProcessMaterialShaders.usf b/Engine/Shaders/Private/PostProcessMaterialShaders.usf
--- a/Engine/Shaders/Private/PostProcessMaterialShaders.usf
+++ b/Engine/Shaders/Private/PostProcessMaterialShaders.usf
@@ -248,6 +248,10 @@ void MainPS_VideoOverlay(
 
 #else // FEATURE_LEVEL_ES3_1
 
+Texture2D<uint> MobileCustomStencilTexture;
+
+int MobileStencilValueRef; // Use integer cause it has to be negative to make the less function to pass the test as always
+
 uint bFlipYAxis;
 
 //
@@ -290,12 +294,48 @@ void MainVS(
 #endif
 }
 
+#define MOBILE_STENCIL_COMPARE_LESS				0
+#define MOBILE_STENCIL_COMPARE_LESSEQUAL		1
+#define MOBILE_STENCIL_COMPARE_GREATE			2
+#define MOBILE_STENCIL_COMPARE_GREATEEQUAL		3
+#define MOBILE_STENCIL_COMPARE_EQUAL			4
+#define MOBILE_STENCIL_COMPARE_NOTEQUAL			5
+#define MOBILE_STENCIL_COMPARE_NEVER			6
+#define MOBILE_STENCIL_COMPARE_ALWAYS			7
+
 void MainPS(
 	in FPostProcessMaterialVSToPS Input,
 	in float2 InViewportUV : TEXCOORD8,
 	out half4 OutColor : SV_Target0
 	)
 {
+	float ClipValue = 1.0f;
+
+#if MOBILE_STENCIL_COMPARE_FUNCTION != MOBILE_STENCIL_COMPARE_NEVER && MOBILE_STENCIL_COMPARE_FUNCTION != MOBILE_STENCIL_COMPARE_ALWAYS
+	int MobileStencilValue = MobileCustomStencilTexture.Load(int3(Input.Position.xy, 0));
+#endif
+
+#if MOBILE_STENCIL_COMPARE_FUNCTION == MOBILE_STENCIL_COMPARE_LESS
+	ClipValue = MobileStencilValueRef < MobileStencilValue ? 1.0f : -1.0f;
+#elif MOBILE_STENCIL_COMPARE_FUNCTION == MOBILE_STENCIL_COMPARE_LESSEQUAL
+	ClipValue = MobileStencilValueRef <= MobileStencilValue ? 1.0f : -1.0f;
+#elif MOBILE_STENCIL_COMPARE_FUNCTION == MOBILE_STENCIL_COMPARE_GREATE
+	ClipValue = MobileStencilValueRef > MobileStencilValue ? 1.0f : -1.0f;
+#elif MOBILE_STENCIL_COMPARE_FUNCTION == MOBILE_STENCIL_COMPARE_GREATEEQUAL
+	ClipValue = MobileStencilValueRef >= MobileStencilValue ? 1.0f : -1.0f;
+#elif MOBILE_STENCIL_COMPARE_FUNCTION == MOBILE_STENCIL_COMPARE_EQUAL
+	ClipValue = MobileStencilValueRef == MobileStencilValue ? 1.0f : -1.0f;
+#elif MOBILE_STENCIL_COMPARE_FUNCTION == MOBILE_STENCIL_COMPARE_NOTEQUAL
+	ClipValue = MobileStencilValueRef != MobileStencilValue ? 1.0f : -1.0f;
+#elif MOBILE_STENCIL_COMPARE_FUNCTION == MOBILE_STENCIL_COMPARE_NEVER
+#elif MOBILE_STENCIL_COMPARE_FUNCTION == MOBILE_STENCIL_COMPARE_ALWAYS
+	ClipValue = MobileStencilValueRef == 256 ? 1.0f : -1.0f;
+#endif
+
+#if MOBILE_STENCIL_COMPARE_FUNCTION != MOBILE_STENCIL_COMPARE_NEVER
+	clip(ClipValue);
+#endif
+
 	ResolvedView = ResolveView();
 	FMaterialPixelParameters Parameters = MakeInitializedMaterialPixelParameters();
 	FPixelMaterialInputs PixelMaterialInputs;
diff --git a/Engine/Source/Runtime/Renderer/Private/CustomDepthRendering.cpp b/Engine/Source/Runtime/Renderer/Private/CustomDepthRendering.cpp
--- a/Engine/Source/Runtime/Renderer/Private/CustomDepthRendering.cpp
+++ b/Engine/Source/Runtime/Renderer/Private/CustomDepthRendering.cpp
@@ -19,7 +19,7 @@ public:
 	virtual void AddMeshBatch(const FMeshBatch& RESTRICT MeshBatch, uint64 BatchElementMask, const FPrimitiveSceneProxy* RESTRICT PrimitiveSceneProxy, int32 StaticMeshId = -1) override final;
 
 private:
-	template<bool bPositionOnly>
+	template<bool bPositionOnly, bool bUsesMobileColorValue>
 	void Process(
 		const FMeshBatch& MeshBatch,
 		uint64 BatchElementMask,
@@ -29,8 +29,7 @@ private:
 		const FMaterial& RESTRICT MaterialResource,
 		ERasterizerFillMode MeshFillMode,
 		ERasterizerCullMode MeshCullMode,
-		float MobileColorValue,
-		bool bUsesMobileColorValue);
+		float MobileColorValue);
 
 	FMeshPassProcessorRenderState PassDrawRenderState;
 };
@@ -71,7 +70,7 @@ void FCustomDepthPassMeshProcessor::AddMeshBatch(const FMeshBatch& RESTRICT Mesh
 
 
 		const bool bWriteCustomStencilValues = FSceneRenderTargets::IsCustomDepthPassWritingStencil();
-		float bMobileColorValue = 0.0f;
+		float MobileColorValue = 0.0f;
 
 		if (bWriteCustomStencilValues)
 		{
@@ -97,7 +96,7 @@ void FCustomDepthPassMeshProcessor::AddMeshBatch(const FMeshBatch& RESTRICT Mesh
 			if (FeatureLevel <= ERHIFeatureLevel::ES3_1)
 			{
 				// On mobile platforms write custom stencil value to color target
-				bMobileColorValue = CustomDepthStencilValue / 255.0f;
+				MobileColorValue = CustomDepthStencilValue;
 			}
 		}
 		else
@@ -105,7 +104,7 @@ void FCustomDepthPassMeshProcessor::AddMeshBatch(const FMeshBatch& RESTRICT Mesh
 			PassDrawRenderState.SetDepthStencilState(TStaticDepthStencilState<true, CF_DepthNearOrEqual>::GetRHI());
 		}
 
-		const bool bUsesMobileColorValue = bMobileColorValue != 0.0f;
+		const bool bUsesMobileColorValue = FeatureLevel <= ERHIFeatureLevel::ES3_1;
 
 
 		if (BlendMode == BLEND_Opaque
@@ -116,7 +115,7 @@ void FCustomDepthPassMeshProcessor::AddMeshBatch(const FMeshBatch& RESTRICT Mesh
 		{
 			const FMaterialRenderProxy& DefaultProxy = *UMaterial::GetDefaultMaterial(MD_Surface)->GetRenderProxy();
 			const FMaterial& DefaultMaterial = *DefaultProxy.GetMaterial(FeatureLevel);
-			Process<true>(MeshBatch, BatchElementMask, StaticMeshId, PrimitiveSceneProxy, DefaultProxy, DefaultMaterial, MeshFillMode, MeshCullMode, bMobileColorValue, bUsesMobileColorValue);
+			Process<true, false>(MeshBatch, BatchElementMask, StaticMeshId, PrimitiveSceneProxy, DefaultProxy, DefaultMaterial, MeshFillMode, MeshCullMode, MobileColorValue);
 		}
 		else if (!IsTranslucentBlendMode(BlendMode) || Material.IsTranslucencyWritingCustomDepth())
 		{
@@ -132,12 +131,19 @@ void FCustomDepthPassMeshProcessor::AddMeshBatch(const FMeshBatch& RESTRICT Mesh
 				EffectiveMaterial = EffectiveMaterialRenderProxy->GetMaterial(FeatureLevel);
 			}
 
-			Process<false>(MeshBatch, BatchElementMask, StaticMeshId, PrimitiveSceneProxy, *EffectiveMaterialRenderProxy, *EffectiveMaterial, MeshFillMode, MeshCullMode, bMobileColorValue, bUsesMobileColorValue);
+			if (bUsesMobileColorValue)
+			{
+				Process<false, true>(MeshBatch, BatchElementMask, StaticMeshId, PrimitiveSceneProxy, *EffectiveMaterialRenderProxy, *EffectiveMaterial, MeshFillMode, MeshCullMode, MobileColorValue);
+			}
+			else
+			{
+				Process<false, false>(MeshBatch, BatchElementMask, StaticMeshId, PrimitiveSceneProxy, *EffectiveMaterialRenderProxy, *EffectiveMaterial, MeshFillMode, MeshCullMode, MobileColorValue);
+			}
 		}
 	}
 }
 
-template<bool bPositionOnly>
+template<bool bPositionOnly, bool bUsesMobileColorValue>
 void FCustomDepthPassMeshProcessor::Process(
 	const FMeshBatch& RESTRICT MeshBatch,
 	uint64 BatchElementMask,
@@ -147,8 +153,7 @@ void FCustomDepthPassMeshProcessor::Process(
 	const FMaterial& RESTRICT MaterialResource,
 	ERasterizerFillMode MeshFillMode,
 	ERasterizerCullMode MeshCullMode,
-	float MobileColorValue,
-	bool bUsesMobileColorValue)
+	float MobileColorValue)
 {
 	const FVertexFactory* VertexFactory = MeshBatch.VertexFactory;
 
@@ -156,10 +161,10 @@ void FCustomDepthPassMeshProcessor::Process(
 		TDepthOnlyVS<bPositionOnly>,
 		FDepthOnlyHS,
 		FDepthOnlyDS,
-		FDepthOnlyPS> DepthPassShaders;
+		FDepthOnlyPS<bUsesMobileColorValue>> DepthPassShaders;
 
 	FShaderPipelineRef ShaderPipeline;
-	GetDepthPassShaders<bPositionOnly>(
+	GetDepthPassShaders<bPositionOnly, bUsesMobileColorValue>(
 		MaterialResource,
 		VertexFactory->GetType(),
 		FeatureLevel,
@@ -167,8 +172,7 @@ void FCustomDepthPassMeshProcessor::Process(
 		DepthPassShaders.DomainShader,
 		DepthPassShaders.VertexShader,
 		DepthPassShaders.PixelShader,
-		ShaderPipeline,
-		bUsesMobileColorValue
+		ShaderPipeline
 		);
 
 	FDepthOnlyShaderElementData ShaderElementData(MobileColorValue);
diff --git a/Engine/Source/Runtime/Renderer/Private/DepthRendering.cpp b/Engine/Source/Runtime/Renderer/Private/DepthRendering.cpp
--- a/Engine/Source/Runtime/Renderer/Private/DepthRendering.cpp
+++ b/Engine/Source/Runtime/Renderer/Private/DepthRendering.cpp
@@ -78,12 +78,13 @@ IMPLEMENT_MATERIAL_SHADER_TYPE(template<>,TDepthOnlyVS<false>,TEXT("/Engine/Priv
 IMPLEMENT_MATERIAL_SHADER_TYPE(,FDepthOnlyHS,TEXT("/Engine/Private/DepthOnlyVertexShader.usf"),TEXT("MainHull"),SF_Hull);	
 IMPLEMENT_MATERIAL_SHADER_TYPE(,FDepthOnlyDS,TEXT("/Engine/Private/DepthOnlyVertexShader.usf"),TEXT("MainDomain"),SF_Domain);
 
-IMPLEMENT_MATERIAL_SHADER_TYPE(,FDepthOnlyPS,TEXT("/Engine/Private/DepthOnlyPixelShader.usf"),TEXT("Main"),SF_Pixel);
+IMPLEMENT_MATERIAL_SHADER_TYPE(template<>,FDepthOnlyPS<true>,TEXT("/Engine/Private/DepthOnlyPixelShader.usf"),TEXT("Main"),SF_Pixel);
+IMPLEMENT_MATERIAL_SHADER_TYPE(template<>,FDepthOnlyPS<false>,TEXT("/Engine/Private/DepthOnlyPixelShader.usf"),TEXT("Main"),SF_Pixel);
 
 IMPLEMENT_SHADERPIPELINE_TYPE_VS(DepthNoPixelPipeline, TDepthOnlyVS<false>, true);
 IMPLEMENT_SHADERPIPELINE_TYPE_VS(DepthPosOnlyNoPixelPipeline, TDepthOnlyVS<true>, true);
-IMPLEMENT_SHADERPIPELINE_TYPE_VSPS(DepthPipeline, TDepthOnlyVS<false>, FDepthOnlyPS, true);
-
+IMPLEMENT_SHADERPIPELINE_TYPE_VSPS(DepthNoColorOutputPipeline, TDepthOnlyVS<false>, FDepthOnlyPS<false>, true);
+IMPLEMENT_SHADERPIPELINE_TYPE_VSPS(DepthWithColorOutputPipeline, TDepthOnlyVS<false>, FDepthOnlyPS<true>, true);
 
 static FORCEINLINE bool UseShaderPipelines(ERHIFeatureLevel::Type InFeatureLevel)
 {
@@ -91,7 +92,7 @@ static FORCEINLINE bool UseShaderPipelines(ERHIFeatureLevel::Type InFeatureLevel
 	return RHISupportsShaderPipelines(GShaderPlatformForFeatureLevel[InFeatureLevel]) && CVar && CVar->GetValueOnAnyThread() != 0;
 }
 
-template <bool bPositionOnly>
+template <bool bPositionOnly, bool bUsesMobileColorValue>
 void GetDepthPassShaders(
 	const FMaterial& Material,
 	FVertexFactoryType* VertexFactoryType,
@@ -99,9 +100,8 @@ void GetDepthPassShaders(
 	TShaderRef<FDepthOnlyHS>& HullShader,
 	TShaderRef<FDepthOnlyDS>& DomainShader,
 	TShaderRef<TDepthOnlyVS<bPositionOnly>>& VertexShader,
-	TShaderRef<FDepthOnlyPS>& PixelShader,
-	FShaderPipelineRef& ShaderPipeline,
-	bool bUsesMobileColorValue)
+	TShaderRef<FDepthOnlyPS<bUsesMobileColorValue>>& PixelShader,
+	FShaderPipelineRef& ShaderPipeline)
 {
 	if (bPositionOnly && !bUsesMobileColorValue)
 	{
@@ -125,7 +125,7 @@ void GetDepthPassShaders(
 			DomainShader = Material.GetShader<FDepthOnlyDS>(VertexFactoryType);
 			if (bNeedsPixelShader)
 			{
-				PixelShader = Material.GetShader<FDepthOnlyPS>(VertexFactoryType);
+				PixelShader = Material.GetShader<FDepthOnlyPS<bUsesMobileColorValue>>(VertexFactoryType);
 			}
 		}
 		else
@@ -135,7 +135,14 @@ void GetDepthPassShaders(
 			bool bUseShaderPipelines = UseShaderPipelines(FeatureLevel);
 			if (bNeedsPixelShader)
 			{
-				ShaderPipeline = bUseShaderPipelines ? Material.GetShaderPipeline(&DepthPipeline, VertexFactoryType, false) : FShaderPipelineRef();
+				if (bUsesMobileColorValue)
+				{
+					ShaderPipeline = bUseShaderPipelines ? Material.GetShaderPipeline(&DepthWithColorOutputPipeline, VertexFactoryType, false) : FShaderPipelineRef();
+				}
+				else
+				{
+					ShaderPipeline = bUseShaderPipelines ? Material.GetShaderPipeline(&DepthNoColorOutputPipeline, VertexFactoryType, false) : FShaderPipelineRef();
+				}
 			}
 			else
 			{
@@ -147,7 +154,7 @@ void GetDepthPassShaders(
 				VertexShader = ShaderPipeline.GetShader<TDepthOnlyVS<bPositionOnly> >();
 				if (bNeedsPixelShader)
 				{
-					PixelShader = ShaderPipeline.GetShader<FDepthOnlyPS>();
+					PixelShader = ShaderPipeline.GetShader<FDepthOnlyPS<bUsesMobileColorValue>>();
 				}
 			}
 			else
@@ -155,28 +162,28 @@ void GetDepthPassShaders(
 				VertexShader = Material.GetShader<TDepthOnlyVS<bPositionOnly> >(VertexFactoryType);
 				if (bNeedsPixelShader)
 				{
-					PixelShader = Material.GetShader<FDepthOnlyPS>(VertexFactoryType);
+					PixelShader = Material.GetShader<FDepthOnlyPS<bUsesMobileColorValue>>(VertexFactoryType);
 				}
 			}
 		}
 	}
 }
 
-#define IMPLEMENT_GetDepthPassShaders( bPositionOnly ) \
-	template void GetDepthPassShaders< bPositionOnly >( \
+#define IMPLEMENT_GetDepthPassShaders( bPositionOnly, bUsesMobileColorValue ) \
+	template void GetDepthPassShaders< bPositionOnly, bUsesMobileColorValue >( \
 		const FMaterial& Material, \
 		FVertexFactoryType* VertexFactoryType, \
 		ERHIFeatureLevel::Type FeatureLevel, \
 		TShaderRef<FDepthOnlyHS>& HullShader, \
 		TShaderRef<FDepthOnlyDS>& DomainShader, \
 		TShaderRef<TDepthOnlyVS<bPositionOnly>>& VertexShader, \
-		TShaderRef<FDepthOnlyPS>& PixelShader, \
-		FShaderPipelineRef& ShaderPipeline, \
-		bool bUsesMobileColorValue \
+		TShaderRef<FDepthOnlyPS<bUsesMobileColorValue>>& PixelShader, \
+		FShaderPipelineRef& ShaderPipeline \
 	);
 
-IMPLEMENT_GetDepthPassShaders( true );
-IMPLEMENT_GetDepthPassShaders( false );
+IMPLEMENT_GetDepthPassShaders( true, false );
+IMPLEMENT_GetDepthPassShaders( false, false );
+IMPLEMENT_GetDepthPassShaders( false, true );
 
 void SetDepthPassDitheredLODTransitionState(const FSceneView* SceneView, const FMeshBatch& RESTRICT Mesh, int32 StaticMeshId, FMeshPassProcessorRenderState& DrawRenderState)
 {
@@ -879,11 +886,11 @@ void FDepthPassMeshProcessor::Process(
 		TDepthOnlyVS<bPositionOnly>,
 		FDepthOnlyHS,
 		FDepthOnlyDS,
-		FDepthOnlyPS> DepthPassShaders;
+		FDepthOnlyPS<false>> DepthPassShaders;
 
 	FShaderPipelineRef ShaderPipeline;
 
-	GetDepthPassShaders<bPositionOnly>(
+	GetDepthPassShaders<bPositionOnly, false>(
 		MaterialResource,
 		VertexFactory->GetType(),
 		FeatureLevel,
@@ -891,8 +898,7 @@ void FDepthPassMeshProcessor::Process(
 		DepthPassShaders.DomainShader,
 		DepthPassShaders.VertexShader,
 		DepthPassShaders.PixelShader,
-		ShaderPipeline,
-		false
+		ShaderPipeline
 		);
 
 	FMeshPassProcessorRenderState DrawRenderState(PassDrawRenderState);
diff --git a/Engine/Source/Runtime/Renderer/Private/DepthRendering.h b/Engine/Source/Runtime/Renderer/Private/DepthRendering.h
--- a/Engine/Source/Runtime/Renderer/Private/DepthRendering.h
+++ b/Engine/Source/Runtime/Renderer/Private/DepthRendering.h
@@ -138,6 +138,7 @@ public:
 /**
 * A pixel shader for rendering the depth of a mesh.
 */
+template <bool bUsesMobileColorValue>
 class FDepthOnlyPS : public FMeshMaterialShader
 {
 	DECLARE_SHADER_TYPE(FDepthOnlyPS,MeshMaterial);
@@ -145,9 +146,9 @@ public:
 
 	static bool ShouldCompilePermutation(const FMeshMaterialShaderPermutationParameters& Parameters)
 	{
-		return 
-			// Compile for materials that are masked.
-			(!Parameters.MaterialParameters.bWritesEveryPixel || Parameters.MaterialParameters.bHasPixelDepthOffsetConnected || Parameters.MaterialParameters.bIsTranslucencyWritingCustomDepth)
+		return
+			// Compile for materials that are masked, avoid generating permutation for other platforms if bUsesMobileColorValue is true
+			((!Parameters.MaterialParameters.bWritesEveryPixel || Parameters.MaterialParameters.bHasPixelDepthOffsetConnected || Parameters.MaterialParameters.bIsTranslucencyWritingCustomDepth) && (!bUsesMobileColorValue || IsMobilePlatform(Parameters.Platform)))
 			// Mobile uses material pixel shader to write custom stencil to color target
 			|| (IsMobilePlatform(Parameters.Platform) && (Parameters.MaterialParameters.bIsDefaultMaterial || Parameters.MaterialParameters.bMaterialMayModifyMeshPosition));
 	}
@@ -168,6 +169,12 @@ public:
 		{
 			// No access to scene textures during depth rendering on mobile
 			OutEnvironment.SetDefine(TEXT("SCENE_TEXTURES_DISABLED"), 1u);
+
+			OutEnvironment.SetDefine(TEXT("OUTPUT_MOBILE_COLOR_VALUE"), bUsesMobileColorValue ? 1u : 0u);
+		}
+		else
+		{
+			OutEnvironment.SetDefine(TEXT("OUTPUT_MOBILE_COLOR_VALUE"), 0u);
 		}
 	}
 
@@ -191,7 +198,7 @@ public:
 	LAYOUT_FIELD(FShaderParameter, MobileColorValue);
 };
 
-template <bool bPositionOnly>
+template <bool bPositionOnly, bool bUsesMobileColorValue>
 void GetDepthPassShaders(
 	const FMaterial& Material,
 	FVertexFactoryType* VertexFactoryType,
@@ -199,9 +206,8 @@ void GetDepthPassShaders(
 	TShaderRef<FDepthOnlyHS>& HullShader,
 	TShaderRef<FDepthOnlyDS>& DomainShader,
 	TShaderRef<TDepthOnlyVS<bPositionOnly>>& VertexShader,
-	TShaderRef<FDepthOnlyPS>& PixelShader,
-	FShaderPipelineRef& ShaderPipeline,
-	bool bUsesMobileColorValue);
+	TShaderRef<FDepthOnlyPS<bUsesMobileColorValue>>& PixelShader,
+	FShaderPipelineRef& ShaderPipeline);
 
 extern void CreateDepthPassUniformBuffer(
 	FRHICommandListImmediate& RHICmdList, 
diff --git a/Engine/Source/Runtime/Renderer/Private/PostProcess/PostProcessMaterial.cpp b/Engine/Source/Runtime/Renderer/Private/PostProcess/PostProcessMaterial.cpp
--- a/Engine/Source/Runtime/Renderer/Private/PostProcess/PostProcessMaterial.cpp
+++ b/Engine/Source/Runtime/Renderer/Private/PostProcess/PostProcessMaterial.cpp
@@ -261,6 +261,15 @@ public:
 	FPostProcessMaterialPS(const ShaderMetaType::CompiledShaderInitializerType& Initializer)
 		: FPostProcessMaterialShader(Initializer)
 	{}
+
+	static void ModifyCompilationEnvironment(const FMaterialShaderPermutationParameters& Parameters, FShaderCompilerEnvironment& OutEnvironment)
+	{
+		FPostProcessMaterialShader::ModifyCompilationEnvironment(Parameters, OutEnvironment);
+
+		uint32 StencilCompareFunction = Parameters.MaterialParameters.bIsStencilTestEnabled ? Parameters.MaterialParameters.StencilCompare : EMaterialStencilCompare::MSC_Never;
+
+		OutEnvironment.SetDefine(TEXT("MOBILE_STENCIL_COMPARE_FUNCTION"), StencilCompareFunction);
+	}
 };
 
 IMPLEMENT_SHADER_TYPE(,FPostProcessMaterialVS, TEXT("/Engine/Private/PostProcessMaterialShaders.usf"), TEXT("MainVS"), SF_Vertex);
@@ -422,13 +431,18 @@ FScreenPassTexture AddPostProcessMaterialPass(
 
 	RDG_EVENT_SCOPE(GraphBuilder, "PostProcessMaterial %dx%d Material=%s", SceneColorViewport.Rect.Width(), SceneColorViewport.Rect.Height(), *Material->GetFriendlyName());
 
+	const uint32 MaterialStencilRef = Material->GetStencilRefValue();
+
 	FPostProcessMaterialParameters* PostProcessMaterialParameters = GraphBuilder.AllocParameters<FPostProcessMaterialParameters>();
 
 	PostProcessMaterialParameters->PostProcessOutput = GetScreenPassTextureViewportParameters(OutputViewport);
-	PostProcessMaterialParameters->CustomDepth = DepthStencilTexture;
+	PostProcessMaterialParameters->MobileCustomStencilTexture = DepthStencilTexture;
+	PostProcessMaterialParameters->MobileStencilValueRef = MaterialStencilRef;
 	PostProcessMaterialParameters->RenderTargets[0] = Output.GetRenderTargetBinding();
 
-	if (DepthStencilTexture)
+	bool bMobilePlatform = IsMobilePlatform(View.GetShaderPlatform());
+
+	if (DepthStencilTexture && !bMobilePlatform)
 	{
 		PostProcessMaterialParameters->RenderTargets.DepthStencil = FDepthStencilBinding(
 			DepthStencilTexture,
@@ -436,6 +450,31 @@ FScreenPassTexture AddPostProcessMaterialPass(
 			ERenderTargetLoadAction::ELoad,
 			FExclusiveDepthStencil::DepthRead_StencilRead);
 	}
+	else if (!DepthStencilTexture && bMobilePlatform && Material->IsStencilTestEnabled()) // we have to set a default texture for MobileStencilTexture and override the MobileStencilValueRef to make all function to pass the stencil test
+	{
+		PostProcessMaterialParameters->MobileCustomStencilTexture = GSystemTextures.GetZeroUIntDummy(GraphBuilder);
+		
+		switch (Material->GetStencilCompare())
+		{
+		case EMaterialStencilCompare::MSC_Less:
+			PostProcessMaterialParameters->MobileStencilValueRef = -1;
+			break;
+		case EMaterialStencilCompare::MSC_LessEqual:
+		case EMaterialStencilCompare::MSC_GreaterEqual:
+		case EMaterialStencilCompare::MSC_Equal:
+			PostProcessMaterialParameters->MobileStencilValueRef = 0;
+			break;
+		case EMaterialStencilCompare::MSC_Greater:
+		case EMaterialStencilCompare::MSC_NotEqual:
+			PostProcessMaterialParameters->MobileStencilValueRef = 1;
+			break;
+		case EMaterialStencilCompare::MSC_Always:
+			PostProcessMaterialParameters->MobileStencilValueRef = 256;
+			break;
+		default:
+			break;
+		}
+	}
 
 	PostProcessMaterialParameters->PostProcessInput_BilinearSampler = TStaticSamplerState<SF_Bilinear, AM_Clamp, AM_Clamp, AM_Clamp>::GetRHI();;
 
@@ -470,8 +509,6 @@ FScreenPassTexture AddPostProcessMaterialPass(
 	TShaderRef<FPostProcessMaterialPS> PixelShader = MaterialShaderMap->GetShader<FPostProcessMaterialPS>(PermutationVector);
 	ClearUnusedGraphResources(VertexShader, PixelShader, PostProcessMaterialParameters);
 
-	const uint32 MaterialStencilRef = Material->GetStencilRefValue();
-
 	EScreenPassDrawFlags ScreenPassFlags = EScreenPassDrawFlags::AllowHMDHiddenAreaMask;
 
 	if (PostProcessMaterialParameters->bFlipYAxis)
@@ -691,7 +728,9 @@ FRenderingCompositePass* AddPostProcessMaterialPass(
 
 		Inputs.bFlipYAxis = ShouldMobilePassFlipVerticalAxis(InContext, Pass);
 
-		if (TRefCountPtr<IPooledRenderTarget> CustomDepthTarget = FSceneRenderTargets::Get(InContext.RHICmdList).CustomDepth)
+		bool bMobilePlatform = IsMobilePlatform(InContext.View.GetShaderPlatform());
+
+		if (TRefCountPtr<IPooledRenderTarget> CustomDepthTarget = bMobilePlatform ? FSceneRenderTargets::Get(InContext.RHICmdList).MobileCustomStencil : FSceneRenderTargets::Get(InContext.RHICmdList).CustomDepth)
 		{
 			Inputs.CustomDepthTexture = GraphBuilder.RegisterExternalTexture(CustomDepthTarget, TEXT("CustomDepth"));
 		}
diff --git a/Engine/Source/Runtime/Renderer/Private/PostProcess/PostProcessMaterial.h b/Engine/Source/Runtime/Renderer/Private/PostProcess/PostProcessMaterial.h
--- a/Engine/Source/Runtime/Renderer/Private/PostProcess/PostProcessMaterial.h
+++ b/Engine/Source/Runtime/Renderer/Private/PostProcess/PostProcessMaterial.h
@@ -40,7 +40,8 @@ BEGIN_SHADER_PARAMETER_STRUCT(FPostProcessMaterialParameters, )
 	SHADER_PARAMETER_STRUCT(FScreenPassTextureViewportParameters, PostProcessOutput)
 	SHADER_PARAMETER_STRUCT_ARRAY(FScreenPassTextureInput, PostProcessInput, [kPostProcessMaterialInputCountMax])
 	SHADER_PARAMETER_SAMPLER(SamplerState, PostProcessInput_BilinearSampler)
-	SHADER_PARAMETER_RDG_TEXTURE(Texture2D, CustomDepth)
+	SHADER_PARAMETER_RDG_TEXTURE(Texture2D<uint>, MobileCustomStencilTexture)
+	SHADER_PARAMETER(int32, MobileStencilValueRef)
 	SHADER_PARAMETER(uint32, bFlipYAxis)
 	RENDER_TARGET_BINDING_SLOTS()
 END_SHADER_PARAMETER_STRUCT()
diff --git a/Engine/Source/Runtime/Renderer/Private/PostProcess/SceneRenderTargets.cpp b/Engine/Source/Runtime/Renderer/Private/PostProcess/SceneRenderTargets.cpp
--- a/Engine/Source/Runtime/Renderer/Private/PostProcess/SceneRenderTargets.cpp
+++ b/Engine/Source/Runtime/Renderer/Private/PostProcess/SceneRenderTargets.cpp
@@ -79,6 +79,15 @@ static TAutoConsoleVariable<int32> CVarCustomDepth(
 	ECVF_RenderThreadSafe
 	);
 
+static TAutoConsoleVariable<int32> CVarMobileCustomDepthDownSample(
+	TEXT("r.Mobile.CustomDepthDownSample"),
+	0,
+	TEXT("Perform Mobile CustomDepth at HalfRes \n ")
+	TEXT("0: Off (default)\n ")
+	TEXT("1: On \n "),
+	ECVF_RenderThreadSafe
+);
+
 static TAutoConsoleVariable<int32> CVarMSAACount(
 	TEXT("r.MSAACount"),
 	4,
@@ -263,6 +272,7 @@ FSceneRenderTargets::FSceneRenderTargets(const FViewInfo& View, const FSceneRend
 	, ScreenSpaceGTAOHorizons(GRenderTargetPool.MakeSnapshot(SnapshotSource.ScreenSpaceGTAOHorizons))
 	, QuadOverdrawBuffer(GRenderTargetPool.MakeSnapshot(SnapshotSource.QuadOverdrawBuffer))
 	, CustomDepth(GRenderTargetPool.MakeSnapshot(SnapshotSource.CustomDepth))
+	, MobileCustomDepth(GRenderTargetPool.MakeSnapshot(SnapshotSource.MobileCustomDepth))
 	, MobileCustomStencil(GRenderTargetPool.MakeSnapshot(SnapshotSource.MobileCustomStencil))
 	, CustomStencilSRV(SnapshotSource.CustomStencilSRV)
 	, SkySHIrradianceMap(GRenderTargetPool.MakeSnapshot(SnapshotSource.SkySHIrradianceMap))
@@ -1430,21 +1440,24 @@ bool FSceneRenderTargets::BeginRenderingCustomDepth(FRHICommandListImmediate& RH
 		FRHIRenderPassInfo RPInfo;
 
 		const bool bWritesCustomStencilValues = IsCustomDepthPassWritingStencil();
-		const bool bRequiresStencilColorTarget = (bWritesCustomStencilValues && CurrentFeatureLevel <= ERHIFeatureLevel::ES3_1);
+		const bool bRequiresStencilColorTarget = (CurrentFeatureLevel <= ERHIFeatureLevel::ES3_1);
 
 		int32 NumColorTargets = 0;
-		FRHIRenderTargetView ColorView = {};
 		if (bRequiresStencilColorTarget)
 		{
-			checkSlow(MobileCustomStencil.IsValid());
+			checkSlow(MobileCustomDepth.IsValid() && MobileCustomStencil.IsValid());
 
 			RPInfo.ColorRenderTargets[0].Action = ERenderTargetActions::Clear_Store;
 			RPInfo.ColorRenderTargets[0].ArraySlice = -1;
 			RPInfo.ColorRenderTargets[0].MipIndex = 0;
-			RPInfo.ColorRenderTargets[0].RenderTarget = MobileCustomStencil->GetRenderTargetItem().ShaderResourceTexture;
+			RPInfo.ColorRenderTargets[0].RenderTarget = MobileCustomDepth->GetRenderTargetItem().ShaderResourceTexture;
+
+			RPInfo.ColorRenderTargets[1].Action = bWritesCustomStencilValues ? ERenderTargetActions::Clear_Store : ERenderTargetActions::DontLoad_DontStore;
+			RPInfo.ColorRenderTargets[1].ArraySlice = -1;
+			RPInfo.ColorRenderTargets[1].MipIndex = 0;
+			RPInfo.ColorRenderTargets[1].RenderTarget = MobileCustomStencil->GetRenderTargetItem().ShaderResourceTexture;
 
-			ColorView =	FRHIRenderTargetView(MobileCustomStencil->GetRenderTargetItem().ShaderResourceTexture, 0, -1, ERenderTargetLoadAction::EClear, ERenderTargetStoreAction::EStore);
-			NumColorTargets = 1;
+			NumColorTargets = 2;
 		}
 
 		RPInfo.DepthStencilRenderTarget.Action = MakeDepthStencilTargetActions(ERenderTargetActions::Clear_Store, ERenderTargetActions::Load_Store);
@@ -1476,12 +1489,15 @@ void FSceneRenderTargets::FinishRenderingCustomDepth(FRHICommandListImmediate& R
 
 	RHICmdList.EndRenderPass();
 	
-	if (CurrentFeatureLevel <= ERHIFeatureLevel::ES3_1 && IsCustomDepthPassWritingStencil() && MobileCustomStencil.IsValid())
+	if (CurrentFeatureLevel <= ERHIFeatureLevel::ES3_1)
 	{
+		RHICmdList.CopyToResolveTarget(MobileCustomDepth->GetRenderTargetItem().TargetableTexture, MobileCustomDepth->GetRenderTargetItem().ShaderResourceTexture, FResolveParams(ResolveRect));
 		RHICmdList.CopyToResolveTarget(MobileCustomStencil->GetRenderTargetItem().TargetableTexture, MobileCustomStencil->GetRenderTargetItem().ShaderResourceTexture, FResolveParams(ResolveRect));
 	}
-
-	RHICmdList.CopyToResolveTarget(CustomDepth->GetRenderTargetItem().TargetableTexture, CustomDepth->GetRenderTargetItem().ShaderResourceTexture, FResolveParams(ResolveRect));
+	else
+	{
+		RHICmdList.CopyToResolveTarget(CustomDepth->GetRenderTargetItem().TargetableTexture, CustomDepth->GetRenderTargetItem().ShaderResourceTexture, FResolveParams(ResolveRect));
+	}
 
 	bCustomDepthIsValid = true;
 }
@@ -2753,6 +2769,7 @@ void FSceneRenderTargets::ReleaseAllTargets()
 	SceneColorSubPixel.SafeRelease();
 	DirectionalOcclusion.SafeRelease();
 	CustomDepth.SafeRelease();
+	MobileCustomDepth.SafeRelease();
 	MobileCustomStencil.SafeRelease();
 	CustomStencilSRV.SafeRelease();
 
@@ -2870,14 +2887,19 @@ IPooledRenderTarget* FSceneRenderTargets::RequestCustomDepth(FRHICommandListImme
 	int Value = CVarCustomDepth.GetValueOnRenderThread();
 	const bool bCustomDepthPassWritingStencil = IsCustomDepthPassWritingStencil();
 	const bool bMobilePath = (CurrentFeatureLevel <= ERHIFeatureLevel::ES3_1);
+	int32 DownsampleFactor = bMobilePath && CVarMobileCustomDepthDownSample.GetValueOnRenderThread() > 0 ? 2 : 1;
+
+	FIntPoint CustomDepthBufferSize = FIntPoint::DivideAndRoundUp(BufferSize, DownsampleFactor);
 
 	if ((Value == 1 && bPrimitives) || Value == 2 || bCustomDepthPassWritingStencil)
 	{
-		bool bHasValidCustomDepth = (CustomDepth.IsValid() && BufferSize == CustomDepth->GetDesc().Extent && !GFastVRamConfig.bDirty);
+		bool bHasValidCustomDepth = (CustomDepth.IsValid() && CustomDepthBufferSize == CustomDepth->GetDesc().Extent && !GFastVRamConfig.bDirty);
 		bool bHasValidCustomStencil;
 		if (bMobilePath)
 		{
-			bHasValidCustomStencil = (MobileCustomStencil.IsValid() && BufferSize == MobileCustomStencil->GetDesc().Extent);
+			bHasValidCustomStencil = (MobileCustomStencil.IsValid() && CustomDepthBufferSize == MobileCustomStencil->GetDesc().Extent) && 
+									//Use memory less when stencil writing is disabled and vice versa
+									bCustomDepthPassWritingStencil == ((MobileCustomStencil->GetDesc().TargetableFlags & TexCreate_Memoryless) == 0);
 		}
 		else
 		{
@@ -2891,14 +2913,32 @@ IPooledRenderTarget* FSceneRenderTargets::RequestCustomDepth(FRHICommandListImme
 			uint32 CustomDepthFlags = TexCreate_NoFastClear;
 
 			// Todo: Could check if writes stencil here and create min viable target
-			FPooledRenderTargetDesc Desc(FPooledRenderTargetDesc::Create2DDesc(BufferSize, PF_DepthStencil, FClearValueBinding::DepthFar, CustomDepthFlags, TexCreate_DepthStencilTargetable | TexCreate_ShaderResource, false));
+			FPooledRenderTargetDesc Desc(FPooledRenderTargetDesc::Create2DDesc(CustomDepthBufferSize, PF_DepthStencil, FClearValueBinding::DepthFar, CustomDepthFlags, TexCreate_DepthStencilTargetable | TexCreate_ShaderResource, false));
 			Desc.Flags |= GFastVRamConfig.CustomDepth;
+			
+			if (bMobilePath)
+			{
+				Desc.TargetableFlags |= TexCreate_Memoryless;
+			}
+
 			GRenderTargetPool.FindFreeElement(RHICmdList, Desc, CustomDepth, TEXT("CustomDepth"), true, ERenderTargetTransience::NonTransient);
 			
 			if (bMobilePath)
 			{
-				FPooledRenderTargetDesc MobileCustomStencilDesc(FPooledRenderTargetDesc::Create2DDesc(BufferSize, PF_B8G8R8A8, FClearValueBinding::Transparent, TexCreate_None, TexCreate_RenderTargetable | TexCreate_ShaderResource, false));
+				uint32 CustomDepthStencilTargetableFlags = TexCreate_RenderTargetable | TexCreate_ShaderResource;
+
+				FPooledRenderTargetDesc MobileCustomDepthDesc(FPooledRenderTargetDesc::Create2DDesc(CustomDepthBufferSize, PF_R16F, FClearValueBinding(FLinearColor(65500.0f, 65500.0f, 65500.0f)), TexCreate_None, CustomDepthStencilTargetableFlags, false));
+				GRenderTargetPool.FindFreeElement(RHICmdList, MobileCustomDepthDesc, MobileCustomDepth, TEXT("MobileCustomDepth"));
+
+				if (!bCustomDepthPassWritingStencil)
+				{
+					CustomDepthStencilTargetableFlags |= TexCreate_Memoryless;
+				}
+
+				FPooledRenderTargetDesc MobileCustomStencilDesc(FPooledRenderTargetDesc::Create2DDesc(CustomDepthBufferSize, PF_R8_UINT, FClearValueBinding::Transparent, TexCreate_None, CustomDepthStencilTargetableFlags, false));
 				GRenderTargetPool.FindFreeElement(RHICmdList, MobileCustomStencilDesc, MobileCustomStencil, TEXT("MobileCustomStencil"));
+
+				CustomStencilSRV.SafeRelease();
 			}
 			else
 			{
@@ -3203,6 +3243,7 @@ void SetupMobileSceneTextureUniformParameters(
 {
 	FRHITexture* BlackDefault2D = GSystemTextures.BlackDummy->GetRenderTargetItem().ShaderResourceTexture;
 	FRHITexture* DepthDefault = GSystemTextures.DepthDummy->GetRenderTargetItem().ShaderResourceTexture;
+	FRHITexture* ZeroUIntDummy = GSystemTextures.ZeroUIntDummy->GetRenderTargetItem().ShaderResourceTexture;
 
 	SceneTextureParameters.SceneColorTexture = bSceneTexturesValid ? SceneContext.GetSceneColorTexture().GetReference() : BlackDefault2D;
 	SceneTextureParameters.SceneColorTextureSampler = TStaticSamplerState<SF_Point, AM_Clamp, AM_Clamp, AM_Clamp>::GetRHI();
@@ -3214,11 +3255,11 @@ void SetupMobileSceneTextureUniformParameters(
 	SceneTextureParameters.SceneAlphaCopyTexture = bSceneTexturesValid && SceneContext.HasSceneAlphaCopyTexture() ? SceneContext.GetSceneAlphaCopyTexture() : BlackDefault2D;
 	SceneTextureParameters.SceneAlphaCopyTextureSampler = TStaticSamplerState<SF_Point,AM_Clamp,AM_Clamp,AM_Clamp>::GetRHI();
 
-	FRHITexture* CustomDepth = DepthDefault;
+	FRHITexture* CustomDepth = BlackDefault2D;
 
 	// if there is no custom depth it's better to have the far distance there
 	// we should update all pass uniform buffers at the start of frame on mobile, SceneContext.bCustomDepthIsValid is invalid at InitView, so pass the parameter from View.bUsesCustomDepthStencil
-	IPooledRenderTarget* CustomDepthTarget = bCustomDepthIsValid ? SceneContext.CustomDepth.GetReference() : 0;
+	IPooledRenderTarget* CustomDepthTarget = bCustomDepthIsValid ? SceneContext.MobileCustomDepth.GetReference() : 0;
 	if (CustomDepthTarget)
 	{
 		CustomDepth = CustomDepthTarget->GetRenderTargetItem().ShaderResourceTexture;
@@ -3227,15 +3268,14 @@ void SetupMobileSceneTextureUniformParameters(
 	SceneTextureParameters.CustomDepthTexture = CustomDepth;
 	SceneTextureParameters.CustomDepthTextureSampler = TStaticSamplerState<>::GetRHI();
 
-	FRHITexture* MobileCustomStencil = BlackDefault2D;
+	FRHITexture* MobileCustomStencil = ZeroUIntDummy;
 
 	if (SceneContext.MobileCustomStencil.IsValid())
 	{
 		MobileCustomStencil = SceneContext.MobileCustomStencil->GetRenderTargetItem().ShaderResourceTexture;
 	}
 
-	SceneTextureParameters.MobileCustomStencilTexture = MobileCustomStencil;
-	SceneTextureParameters.MobileCustomStencilTextureSampler = TStaticSamplerState<SF_Point, AM_Clamp, AM_Clamp, AM_Clamp>::GetRHI();
+	SceneTextureParameters.CustomStencilTexture = MobileCustomStencil;
 	
 	if (SceneContext.VirtualTextureFeedback.FeedbackBufferUAV.IsValid())
 	{
diff --git a/Engine/Source/Runtime/Renderer/Private/PostProcess/SceneRenderTargets.h b/Engine/Source/Runtime/Renderer/Private/PostProcess/SceneRenderTargets.h
--- a/Engine/Source/Runtime/Renderer/Private/PostProcess/SceneRenderTargets.h
+++ b/Engine/Source/Runtime/Renderer/Private/PostProcess/SceneRenderTargets.h
@@ -561,6 +561,8 @@ public:
 	TRefCountPtr<IPooledRenderTarget> QuadOverdrawBuffer;
 	// used by the CustomDepth material feature, is allocated on demand or if r.CustomDepth is 2
 	TRefCountPtr<IPooledRenderTarget> CustomDepth;
+	// CustomDepth is memoryless on mobile, depth is saved in MobileCustomDepth Color RT 
+	TRefCountPtr<IPooledRenderTarget> MobileCustomDepth;
 	TRefCountPtr<IPooledRenderTarget> MobileCustomStencil;
 	// used by the CustomDepth material feature for stencil
 	TRefCountPtr<FRHIShaderResourceView> CustomStencilSRV;
diff --git a/Engine/Source/Runtime/Renderer/Private/SceneRendering.cpp b/Engine/Source/Runtime/Renderer/Private/SceneRendering.cpp
--- a/Engine/Source/Runtime/Renderer/Private/SceneRendering.cpp
+++ b/Engine/Source/Runtime/Renderer/Private/SceneRendering.cpp
@@ -3077,7 +3077,9 @@ void FSceneRenderer::RenderCustomDepthPass(FRHICommandListImmediate& RHICmdList)
 			{
 				FRHIUniformBuffer* PassUniformBuffer = nullptr;
 
-				if (FSceneInterface::GetShadingPath(View.FeatureLevel) == EShadingPath::Mobile)
+				bool bMobilePath = FSceneInterface::GetShadingPath(View.FeatureLevel) == EShadingPath::Mobile;
+
+				if (bMobilePath)
 				{
 					FMobileSceneTextureUniformParameters SceneTextureParameters;
 					SetupMobileSceneTextureUniformParameters(SceneContext, View.FeatureLevel, false, false, SceneTextureParameters);
@@ -3092,25 +3094,34 @@ void FSceneRenderer::RenderCustomDepthPass(FRHICommandListImmediate& RHICmdList)
 					PassUniformBuffer = Scene->UniformBuffers.CustomDepthPassUniformBuffer;
 				}
 				
+				static const auto CVarMobileCustomDepthDownSample = IConsoleManager::Get().FindTConsoleVariableDataInt(TEXT("r.Mobile.CustomDepthDownSample"));
+
+				bool bMobileCustomDepthDownSample = bMobilePath && CVarMobileCustomDepthDownSample && CVarMobileCustomDepthDownSample->GetValueOnRenderThread() > 0;
+
+				FIntRect ViewRect = bMobileCustomDepthDownSample ? FIntRect::DivideAndRoundUp(View.ViewRect, 2) : View.ViewRect;
+
 				if (!View.IsInstancedStereoPass())
 				{
-					RHICmdList.SetViewport(View.ViewRect.Min.X, View.ViewRect.Min.Y, 0.0f, View.ViewRect.Max.X, View.ViewRect.Max.Y, 1.0f);
+					RHICmdList.SetViewport(ViewRect.Min.X, ViewRect.Min.Y, 0.0f, ViewRect.Max.X, ViewRect.Max.Y, 1.0f);
 				}
 				else
 				{
 					if (View.bIsMultiViewEnabled)
 					{
-						const uint32 LeftMinX = Views[0].ViewRect.Min.X;
-						const uint32 LeftMaxX = Views[0].ViewRect.Max.X;
-						const uint32 RightMinX = Views[1].ViewRect.Min.X;
-						const uint32 RightMaxX = Views[1].ViewRect.Max.X;
-						const uint32 LeftMaxY = Views[0].ViewRect.Max.Y;
-						const uint32 RightMaxY = Views[1].ViewRect.Max.Y;
+						FIntRect ViewRect0 = bMobileCustomDepthDownSample ? FIntRect::DivideAndRoundUp(Views[0].ViewRect, 2) : Views[0].ViewRect;
+						FIntRect ViewRect1 = bMobileCustomDepthDownSample ? FIntRect::DivideAndRoundUp(Views[1].ViewRect, 2) : Views[1].ViewRect;
+
+						const uint32 LeftMinX = ViewRect0.Min.X;
+						const uint32 LeftMaxX = ViewRect0.Max.X;
+						const uint32 RightMinX = ViewRect1.Min.X;
+						const uint32 RightMaxX = ViewRect1.Max.X;
+						const uint32 LeftMaxY = ViewRect0.Max.Y;
+						const uint32 RightMaxY = ViewRect1.Max.Y;
 						RHICmdList.SetStereoViewport(LeftMinX, RightMinX, 0, 0, 0.0f, LeftMaxX, RightMaxX, LeftMaxY, RightMaxY, 1.0f);
 					}
 					else
 					{
-						RHICmdList.SetViewport(0, 0, 0, InstancedStereoWidth, View.ViewRect.Max.Y, 1);
+						RHICmdList.SetViewport(0, 0, 0, InstancedStereoWidth, ViewRect.Max.Y, 1);
 					}
 				}
 
diff --git a/Engine/Source/Runtime/Renderer/Public/SceneRenderTargetParameters.h b/Engine/Source/Runtime/Renderer/Public/SceneRenderTargetParameters.h
--- a/Engine/Source/Runtime/Renderer/Public/SceneRenderTargetParameters.h
+++ b/Engine/Source/Runtime/Renderer/Public/SceneRenderTargetParameters.h
@@ -111,8 +111,7 @@ BEGIN_GLOBAL_SHADER_PARAMETER_STRUCT(FMobileSceneTextureUniformParameters, RENDE
 	SHADER_PARAMETER_SAMPLER(SamplerState, SceneAlphaCopyTextureSampler)
 	SHADER_PARAMETER_TEXTURE(Texture2D, CustomDepthTexture)
 	SHADER_PARAMETER_SAMPLER(SamplerState, CustomDepthTextureSampler)
-	SHADER_PARAMETER_TEXTURE(Texture2D, MobileCustomStencilTexture)
-	SHADER_PARAMETER_SAMPLER(SamplerState, MobileCustomStencilTextureSampler)
+	SHADER_PARAMETER_TEXTURE(Texture2D<uint>, CustomStencilTexture)
 	SHADER_PARAMETER_UAV(RWBuffer<uint>, VirtualTextureFeedbackUAV)
 END_GLOBAL_SHADER_PARAMETER_STRUCT()
 
