Fix Static Analysis error

diff --git a/Engine/Source/Runtime/Renderer/Private/MobileBasePass.cpp b/Engine/Source/Runtime/Renderer/Private/MobileBasePass.cpp
--- a/Engine/Source/Runtime/Renderer/Private/MobileBasePass.cpp
+++ b/Engine/Source/Runtime/Renderer/Private/MobileBasePass.cpp
@@ -632,7 +632,7 @@ void FMobileBasePassMeshProcessor::Process(
 		BasePassShaders.VertexShader, 
 		BasePassShaders.PixelShader);
 
-	const bool bMaskedInEarlyPass = (MeshBatch.MaterialRenderProxy->GetMaterial(FeatureLevel)->IsMasked() || MeshBatch.bDitheredLODTransition) && MaskedInEarlyPass(Scene->GetShaderPlatform());
+	const bool bMaskedInEarlyPass = (MeshBatch.MaterialRenderProxy->GetMaterial(FeatureLevel)->IsMasked() || MeshBatch.bDitheredLODTransition) && Scene && MaskedInEarlyPass(Scene->GetShaderPlatform());
 	
 	FMeshPassProcessorRenderState DrawRenderState(PassDrawRenderState);
 	if (bTranslucentBasePass)
