Try to fix a bug of mobile custom depth using uncleared custom depth render target when custom depth primitives are hidden in game.

diff --git a/Engine/Source/Runtime/Renderer/Private/DistortionRendering.cpp b/Engine/Source/Runtime/Renderer/Private/DistortionRendering.cpp
--- a/Engine/Source/Runtime/Renderer/Private/DistortionRendering.cpp
+++ b/Engine/Source/Runtime/Renderer/Private/DistortionRendering.cpp
@@ -70,7 +70,7 @@ void SetupDistortionPassUniformBuffer(FRHICommandListImmediate& RHICmdList, cons
 void SetupMobileDistortionPassUniformBuffer(FRHICommandListImmediate& RHICmdList, const FViewInfo& View, FMobileDistortionPassUniformParameters& DistortionPassParameters)
 {
 	FSceneRenderTargets& SceneRenderTargets = FSceneRenderTargets::Get(RHICmdList);
-	SetupMobileSceneTextureUniformParameters(SceneRenderTargets, View.FeatureLevel, true, View.bUsesCustomDepthStencil, DistortionPassParameters.SceneTextures);
+	SetupMobileSceneTextureUniformParameters(SceneRenderTargets, View.FeatureLevel, true, View.bCustomDepthStencilValid, DistortionPassParameters.SceneTextures);
 	SetupDistortionParams(DistortionPassParameters.DistortionParams, View);
 }
 
diff --git a/Engine/Source/Runtime/Renderer/Private/MobileBasePassRendering.cpp b/Engine/Source/Runtime/Renderer/Private/MobileBasePassRendering.cpp
--- a/Engine/Source/Runtime/Renderer/Private/MobileBasePassRendering.cpp
+++ b/Engine/Source/Runtime/Renderer/Private/MobileBasePassRendering.cpp
@@ -165,7 +165,7 @@ void SetupMobileBasePassUniformParameters(
 	SetupPlanarReflectionUniformParameters(View, ReflectionSceneProxy, BasePassParameters.PlanarReflection);
 
 	FSceneRenderTargets& SceneContext = FSceneRenderTargets::Get(RHICmdList);
-	SetupMobileSceneTextureUniformParameters(SceneContext, View.FeatureLevel, bTranslucentPass, View.bUsesCustomDepthStencil, BasePassParameters.SceneTextures);
+	SetupMobileSceneTextureUniformParameters(SceneContext, View.FeatureLevel, bTranslucentPass, View.bCustomDepthStencilValid, BasePassParameters.SceneTextures);
 
 	BasePassParameters.PreIntegratedGFTexture = GSystemTextures.PreintegratedGF->GetRenderTargetItem().ShaderResourceTexture;
 	BasePassParameters.PreIntegratedGFSampler = TStaticSamplerState<SF_Bilinear, AM_Clamp, AM_Clamp, AM_Clamp>::GetRHI();
diff --git a/Engine/Source/Runtime/Renderer/Private/MobileShadingRenderer.cpp b/Engine/Source/Runtime/Renderer/Private/MobileShadingRenderer.cpp
--- a/Engine/Source/Runtime/Renderer/Private/MobileShadingRenderer.cpp
+++ b/Engine/Source/Runtime/Renderer/Private/MobileShadingRenderer.cpp
@@ -101,39 +101,50 @@ TGlobalResource<FGlobalDynamicReadBuffer> FMobileSceneRenderer::DynamicReadBuffe
 
 static bool UsesCustomDepthStencilLookup(const FViewInfo& View)
 {
-	if (View.bUsesCustomDepthStencil && CVarMobileCustomDepthForTranslucency.GetValueOnAnyThread() != 0)
+	bool bUsesCustomDepthStencil = false;
+
+	// Find out whether CustomDepth/Stencil used in translucent materials
+	if (View.bUsesCustomDepthStencilInTranslucentMaterials && CVarMobileCustomDepthForTranslucency.GetValueOnAnyThread() != 0)
 	{
-		return true;
+		bUsesCustomDepthStencil = true;
 	}
-
-	// Find out whether post-process materials use CustomDepth/Stencil lookups
-	const FBlendableManager& BlendableManager = View.FinalPostProcessSettings.BlendableManager;
-	FBlendableEntry* BlendableIt = nullptr;
-
-	while (FPostProcessMaterialNode* DataPtr = BlendableManager.IterateBlendables<FPostProcessMaterialNode>(BlendableIt))
+	else
 	{
-		if (DataPtr->IsValid())
-		{
-			FMaterialRenderProxy* Proxy = DataPtr->GetMaterialInterface()->GetRenderProxy();
-			check(Proxy);
+		// Find out whether post-process materials use CustomDepth/Stencil lookups
+		const FBlendableManager& BlendableManager = View.FinalPostProcessSettings.BlendableManager;
+		FBlendableEntry* BlendableIt = nullptr;
 
-			const FMaterial* Material = Proxy->GetMaterial(View.GetFeatureLevel());
-			check(Material);
-			if (Material->IsStencilTestEnabled())
+		while (FPostProcessMaterialNode* DataPtr = BlendableManager.IterateBlendables<FPostProcessMaterialNode>(BlendableIt))
+		{
+			if (DataPtr->IsValid())
 			{
-				return true;
-			}
+				FMaterialRenderProxy* Proxy = DataPtr->GetMaterialInterface()->GetRenderProxy();
+				check(Proxy);
 
-			const FMaterialShaderMap* MaterialShaderMap = Material->GetRenderingThreadShaderMap();
-			if (MaterialShaderMap->UsesSceneTexture(PPI_CustomDepth) || MaterialShaderMap->UsesSceneTexture(PPI_CustomStencil))
-			{
-				return true;
+				const FMaterial* Material = Proxy->GetMaterial(View.GetFeatureLevel());
+				check(Material);
+				if (Material->IsStencilTestEnabled())
+				{
+					bUsesCustomDepthStencil = true;
+					break;
+				}
+
+				const FMaterialShaderMap* MaterialShaderMap = Material->GetRenderingThreadShaderMap();
+				if (MaterialShaderMap->UsesSceneTexture(PPI_CustomDepth) || MaterialShaderMap->UsesSceneTexture(PPI_CustomStencil))
+				{
+					bUsesCustomDepthStencil = true;
+					break;
+				}
 			}
 		}
 	}
 
-	//TODO: check if translucency uses CustomDepth
-	return false;
+	// Find out whether there are primitives will render in custom depth pass or just always render custom depth 
+	static const auto CVarCustomDepth = IConsoleManager::Get().FindTConsoleVariableDataInt(TEXT("r.CustomDepth"));
+
+	bUsesCustomDepthStencil &= (View.bHasCustomDepthPrimitives || (CVarCustomDepth && CVarCustomDepth->GetValueOnRenderThread() > 1));
+	
+	return bUsesCustomDepthStencil;
 }
 
 
@@ -141,6 +152,7 @@ FMobileSceneRenderer::FMobileSceneRenderer(const FSceneViewFamily* InViewFamily,
 	:	FSceneRenderer(InViewFamily, HitProxyConsumer)
 {
 	bModulatedShadowsInUse = false;
+	bShouldRenderCustomDepth = false;
 }
 
 class FMobileDirLightShaderParamsRenderResource : public FRenderResource
@@ -263,6 +275,20 @@ void FMobileSceneRenderer::InitViews(FRHICommandListImmediate& RHICmdList)
 
 	PostVisibilityFrameSetup(ILCTaskData);
 
+	// Find out whether custom depth pass should be rendered.
+	{
+		const bool bGammaSpace = !IsMobileHDR();
+
+		bool bCouldUseCustomDepthStencil = !bGammaSpace && (!Scene->World || (Scene->World->WorldType != EWorldType::EditorPreview && Scene->World->WorldType != EWorldType::Inactive));
+
+		for (int32 ViewIndex = 0; ViewIndex < Views.Num(); ViewIndex++)
+		{
+			Views[ViewIndex].bCustomDepthStencilValid = bCouldUseCustomDepthStencil && UsesCustomDepthStencilLookup(Views[ViewIndex]);
+
+			bShouldRenderCustomDepth |= Views[ViewIndex].bCustomDepthStencilValid;
+		}
+	}
+
 	const bool bDynamicShadows = ViewFamily.EngineShowFlags.DynamicShadows;
 
 	if (bDynamicShadows && !IsSimpleForwardShadingEnabled(ShaderPlatform))
@@ -492,23 +518,10 @@ void FMobileSceneRenderer::Render(FRHICommandListImmediate& RHICmdList)
 	}
 
 	// Custom depth
-	if (!bGammaSpace)
+	// bShouldRenderCustomDepth has been initialized in InitViews on mobile platform
+	if (bShouldRenderCustomDepth)
 	{
-		// see if anything uses CustomDepth, and skip custom depth rendering if possible
-		bool bUsesCustomDepthStencil = false;
-		for (int32 ViewIndex = 0; ViewIndex < Views.Num(); ViewIndex++) 
-		{
-			if (UsesCustomDepthStencilLookup(Views[ViewIndex]))
-			{
-				bUsesCustomDepthStencil = true;
-				break;
-			}
-		}
-
-		if (bUsesCustomDepthStencil)
-		{
-			RenderCustomDepthPass(RHICmdList);
-		}
+		RenderCustomDepthPass(RHICmdList);
 	}
 		
 	//
diff --git a/Engine/Source/Runtime/Renderer/Private/PostProcess/SceneRenderTargets.cpp b/Engine/Source/Runtime/Renderer/Private/PostProcess/SceneRenderTargets.cpp
--- a/Engine/Source/Runtime/Renderer/Private/PostProcess/SceneRenderTargets.cpp
+++ b/Engine/Source/Runtime/Renderer/Private/PostProcess/SceneRenderTargets.cpp
@@ -3292,7 +3292,7 @@ TUniformBufferRef<FMobileSceneTextureUniformParameters> CreateMobileSceneTexture
 {
 	FSceneRenderTargets& SceneContext = FSceneRenderTargets::Get(RHICmdList);
 	FMobileSceneTextureUniformParameters SceneTextureParameters;
-	SetupMobileSceneTextureUniformParameters(SceneContext, FeatureLevel, true, true, SceneTextureParameters);
+	SetupMobileSceneTextureUniformParameters(SceneContext, FeatureLevel, true, SceneContext.bCustomDepthIsValid, SceneTextureParameters);
 	return TUniformBufferRef<FMobileSceneTextureUniformParameters>::CreateUniformBufferImmediate(SceneTextureParameters, UniformBuffer_SingleDraw);
 }
 
diff --git a/Engine/Source/Runtime/Renderer/Private/SceneRendering.cpp b/Engine/Source/Runtime/Renderer/Private/SceneRendering.cpp
--- a/Engine/Source/Runtime/Renderer/Private/SceneRendering.cpp
+++ b/Engine/Source/Runtime/Renderer/Private/SceneRendering.cpp
@@ -873,7 +873,8 @@ void FViewInfo::Init()
 	bHasCustomDepthPrimitives = false;
 	bHasDistortionPrimitives = false;
 	bAllowStencilDither = false;
-	bUsesCustomDepthStencil = false;
+	bCustomDepthStencilValid = false;
+	bUsesCustomDepthStencilInTranslucentMaterials = false;
 
 	ForwardLightingResources = nullptr;
 
diff --git a/Engine/Source/Runtime/Renderer/Private/SceneRendering.h b/Engine/Source/Runtime/Renderer/Private/SceneRendering.h
--- a/Engine/Source/Runtime/Renderer/Private/SceneRendering.h
+++ b/Engine/Source/Runtime/Renderer/Private/SceneRendering.h
@@ -1061,7 +1061,8 @@ public:
 	uint32 bTranslucentSurfaceLighting : 1;
 	/** Whether the view has any materials that read from scene depth. */
 	uint32 bUsesSceneDepth : 1;
-	uint32 bUsesCustomDepthStencil : 1;
+	uint32 bCustomDepthStencilValid : 1;
+	uint32 bUsesCustomDepthStencilInTranslucentMaterials : 1;
 
 	/** Whether fog should only be computed on rendered opaque pixels or not. */
 	uint32 bFogOnlyOnRenderedOpaque : 1;
@@ -1877,6 +1878,7 @@ protected:
 	
 private:
 	bool bModulatedShadowsInUse;
+	bool bShouldRenderCustomDepth;
 	static FGlobalDynamicIndexBuffer DynamicIndexBuffer;
 	static FGlobalDynamicVertexBuffer DynamicVertexBuffer;
 	static TGlobalResource<FGlobalDynamicReadBuffer> DynamicReadBuffer;
diff --git a/Engine/Source/Runtime/Renderer/Private/SceneVisibility.cpp b/Engine/Source/Runtime/Renderer/Private/SceneVisibility.cpp
--- a/Engine/Source/Runtime/Renderer/Private/SceneVisibility.cpp
+++ b/Engine/Source/Runtime/Renderer/Private/SceneVisibility.cpp
@@ -2422,7 +2422,7 @@ struct FRelevancePacket
 		WriteView.TranslucentPrimCount.Append(TranslucentPrimCount);
 		WriteView.bHasDistortionPrimitives |= bHasDistortionPrimitives;
 		WriteView.bHasCustomDepthPrimitives |= bHasCustomDepthPrimitives;
-		WriteView.bUsesCustomDepthStencil |= bUsesCustomDepthStencil;
+		WriteView.bUsesCustomDepthStencilInTranslucentMaterials |= bUsesCustomDepthStencil;
 		DirtyIndirectLightingCacheBufferPrimitives.AppendTo(WriteView.DirtyIndirectLightingCacheBufferPrimitives);
 
 		WriteView.MeshDecalBatches.Append(MeshDecalBatches);
