Fix unity stress test warning
#rnx

diff --git a/Engine/Source/Runtime/Renderer/Private/SceneRendering.cpp b/Engine/Source/Runtime/Renderer/Private/SceneRendering.cpp
--- a/Engine/Source/Runtime/Renderer/Private/SceneRendering.cpp
+++ b/Engine/Source/Runtime/Renderer/Private/SceneRendering.cpp
@@ -3095,9 +3095,9 @@ void FSceneRenderer::RenderCustomDepthPass(FRHICommandListImmediate& RHICmdList)
 					PassUniformBuffer = Scene->UniformBuffers.CustomDepthPassUniformBuffer;
 				}
 				
-				static const auto CVarMobileCustomDepthDownSample = IConsoleManager::Get().FindTConsoleVariableDataInt(TEXT("r.Mobile.CustomDepthDownSample"));
+				static const auto MobileCustomDepthDownSampleCVar = IConsoleManager::Get().FindTConsoleVariableDataInt(TEXT("r.Mobile.CustomDepthDownSample"));
 
-				bool bMobileCustomDepthDownSample = bMobilePath && CVarMobileCustomDepthDownSample && CVarMobileCustomDepthDownSample->GetValueOnRenderThread() > 0;
+				bool bMobileCustomDepthDownSample = bMobilePath && MobileCustomDepthDownSampleCVar && MobileCustomDepthDownSampleCVar->GetValueOnRenderThread() > 0;
 
 				FIntRect ViewRect = bMobileCustomDepthDownSample ? FIntRect::DivideAndRoundUp(View.ViewRect, 2) : View.ViewRect;
 
