UE4 - Static analysis fix
#rnx

diff --git a/Engine/Source/Runtime/Renderer/Private/ShaderComplexityRendering.cpp b/Engine/Source/Runtime/Renderer/Private/ShaderComplexityRendering.cpp
--- a/Engine/Source/Runtime/Renderer/Private/ShaderComplexityRendering.cpp
+++ b/Engine/Source/Runtime/Renderer/Private/ShaderComplexityRendering.cpp
@@ -92,7 +92,7 @@ void FComplexityAccumulateInterface::GetDebugViewModeShaderBindings(
 	{
 		const TComplexityAccumulatePS<true>& Shader = static_cast<const TComplexityAccumulatePS<true>&>(BaseShader);
 		ShaderBindings.Add(Shader.NormalizedComplexity, FVector4(NormalizedQuadComplexityValue));
-		ShaderBindings.Add(Shader.ShowQuadOverdraw, DebugViewMode != DVSM_ShaderComplexity ? 1 : 0);
+		ShaderBindings.Add(Shader.ShowQuadOverdraw, 0);
 	}
 	else
 	{
