- Fixed Windows Clang link errors (didn't make a lot of sense, but should not be an issue)
#jria UE-88388

diff --git a/Engine/Source/Runtime/AudioExtensions/Public/IAudioExtensionPlugin.h b/Engine/Source/Runtime/AudioExtensions/Public/IAudioExtensionPlugin.h
--- a/Engine/Source/Runtime/AudioExtensions/Public/IAudioExtensionPlugin.h
+++ b/Engine/Source/Runtime/AudioExtensions/Public/IAudioExtensionPlugin.h
@@ -552,7 +552,7 @@ struct AUDIOEXTENSIONS_API FSoundModulationControls
 
 /** Parameter allowing modulation control override for systems opting in to the Modulation System. */
 USTRUCT(BlueprintType)
-struct AUDIOEXTENSIONS_API FSoundModulationParameter
+struct  FSoundModulationParameter
 {
 	GENERATED_BODY()
 
diff --git a/Engine/Source/Runtime/Engine/Private/ScreenRendering.cpp b/Engine/Source/Runtime/Engine/Private/ScreenRendering.cpp
--- a/Engine/Source/Runtime/Engine/Private/ScreenRendering.cpp
+++ b/Engine/Source/Runtime/Engine/Private/ScreenRendering.cpp
@@ -13,6 +13,4 @@ TGlobalResource<FScreenVertexDeclaration> GScreenVertexDeclaration;
 IMPLEMENT_SHADER_TYPE(,FScreenPS,TEXT("/Engine/Private/ScreenPixelShader.usf"),TEXT("Main"),SF_Pixel);
 IMPLEMENT_SHADER_TYPE(,FScreenPSsRGBSource, TEXT("/Engine/Private/ScreenPixelShader.usf"), TEXT("MainsRGBSource"), SF_Pixel);
 IMPLEMENT_SHADER_TYPE(,FScreenVS,TEXT("/Engine/Private/ScreenVertexShader.usf"),TEXT("Main"),SF_Vertex);
-IMPLEMENT_SHADER_TYPE(template<>,TScreenVSForGS<false>,TEXT("/Engine/Private/ScreenVertexShader.usf"),TEXT("MainForGS"),SF_Vertex);
-IMPLEMENT_SHADER_TYPE(template<>,TScreenVSForGS<true>,TEXT("/Engine/Private/ScreenVertexShader.usf"),TEXT("MainForGS"),SF_Vertex);
 IMPLEMENT_SHADER_TYPE(,FScreenPS_OSE,TEXT("/Engine/Private/ScreenPixelShaderOES.usf"),TEXT("Main"),SF_Pixel);
diff --git a/Engine/Source/Runtime/Engine/Public/ScreenRendering.h b/Engine/Source/Runtime/Engine/Public/ScreenRendering.h
--- a/Engine/Source/Runtime/Engine/Public/ScreenRendering.h
+++ b/Engine/Source/Runtime/Engine/Public/ScreenRendering.h
@@ -176,31 +176,4 @@ public:
 	}
 };
 
-template<bool bUsingVertexLayers=false>
-class TScreenVSForGS : public FScreenVS
-{
-	DECLARE_EXPORTED_SHADER_TYPE(TScreenVSForGS,Global,ENGINE_API);
-public:
-
-	static bool ShouldCompilePermutation(const FGlobalShaderPermutationParameters& Parameters)
-	{
-		return IsFeatureLevelSupported(Parameters.Platform, ERHIFeatureLevel::SM5) && (!bUsingVertexLayers || RHISupportsVertexShaderLayer(Parameters.Platform));
-	}
-
-	TScreenVSForGS(const ShaderMetaType::CompiledShaderInitializerType& Initializer):
-	  FScreenVS(Initializer)
-	  {
-	  }
-	TScreenVSForGS() {}
-
-	static void ModifyCompilationEnvironment(const FGlobalShaderPermutationParameters& Parameters, FShaderCompilerEnvironment& OutEnvironment)
-	{
-		FScreenVS::ModifyCompilationEnvironment(Parameters, OutEnvironment);
-		OutEnvironment.SetDefine(TEXT("USING_LAYERS"), (uint32)(bUsingVertexLayers ? 1 : 0));
-		if (!bUsingVertexLayers)
-		{
-			OutEnvironment.CompilerFlags.Add( CFLAG_VertexToGeometryShader );
-		}
-	}
-};
 
diff --git a/Engine/Source/Runtime/Renderer/Private/ShadowDepthRendering.cpp b/Engine/Source/Runtime/Renderer/Private/ShadowDepthRendering.cpp
--- a/Engine/Source/Runtime/Renderer/Private/ShadowDepthRendering.cpp
+++ b/Engine/Source/Runtime/Renderer/Private/ShadowDepthRendering.cpp
@@ -45,6 +45,37 @@ DECLARE_GPU_STAT_NAMED(ShadowDepths, TEXT("Shadow Depths"));
 IMPLEMENT_GLOBAL_SHADER_PARAMETER_STRUCT(FShadowDepthPassUniformParameters, "ShadowDepthPass");
 IMPLEMENT_GLOBAL_SHADER_PARAMETER_STRUCT(FMobileShadowDepthPassUniformParameters, "MobileShadowDepthPass");
 
+template<bool bUsingVertexLayers = false>
+class TScreenVSForGS : public FScreenVS
+{
+	DECLARE_SHADER_TYPE(TScreenVSForGS, Global);
+public:
+
+	static bool ShouldCompilePermutation(const FGlobalShaderPermutationParameters& Parameters)
+	{
+		return IsFeatureLevelSupported(Parameters.Platform, ERHIFeatureLevel::SM5) && (!bUsingVertexLayers || RHISupportsVertexShaderLayer(Parameters.Platform));
+	}
+
+	TScreenVSForGS(const ShaderMetaType::CompiledShaderInitializerType& Initializer) :
+		FScreenVS(Initializer)
+	{
+	}
+	TScreenVSForGS() {}
+
+	static void ModifyCompilationEnvironment(const FGlobalShaderPermutationParameters& Parameters, FShaderCompilerEnvironment& OutEnvironment)
+	{
+		FScreenVS::ModifyCompilationEnvironment(Parameters, OutEnvironment);
+		OutEnvironment.SetDefine(TEXT("USING_LAYERS"), (uint32)(bUsingVertexLayers ? 1 : 0));
+		if (!bUsingVertexLayers)
+		{
+			OutEnvironment.CompilerFlags.Add(CFLAG_VertexToGeometryShader);
+		}
+	}
+};
+
+IMPLEMENT_SHADER_TYPE(template<>, TScreenVSForGS<false>, TEXT("/Engine/Private/ScreenVertexShader.usf"), TEXT("MainForGS"), SF_Vertex);
+IMPLEMENT_SHADER_TYPE(template<>, TScreenVSForGS<true>, TEXT("/Engine/Private/ScreenVertexShader.usf"), TEXT("MainForGS"), SF_Vertex);
+
 
 static TAutoConsoleVariable<int32> CVarShadowForceSerialSingleRenderPass(
 	TEXT("r.Shadow.ForceSerialSingleRenderPass"),
