Fix a bug in mobile custom depth stencil when using half resolution.

diff --git a/Engine/Shaders/Private/DepthOnlyPixelShader.usf b/Engine/Shaders/Private/DepthOnlyPixelShader.usf
--- a/Engine/Shaders/Private/DepthOnlyPixelShader.usf
+++ b/Engine/Shaders/Private/DepthOnlyPixelShader.usf
@@ -28,8 +28,8 @@ void Main(
 	OPTIONAL_OutDepthConservative,
 #endif
 #if OUTPUT_MOBILE_COLOR_VALUE
-	out half OutDepth : SV_Target0,
-	out uint OutStencil : SV_Target1
+	out half OutCustomDepth : SV_Target0,
+	out half OutCustomStencil : SV_Target1
 #else
 	out float4 OutColor : SV_Target0
 #endif
@@ -76,8 +76,8 @@ void Main(
 #endif
 
 #if OUTPUT_MOBILE_COLOR_VALUE
-	OutDepth = min(SvPosition.w, 65500);
-	OutStencil = MobileColorValue;
+	OutCustomDepth = min(SvPosition.w, 65500);
+	OutCustomStencil = MobileColorValue;
 #else
 	OutColor = 0;
 #endif
diff --git a/Engine/Shaders/Private/MaterialTemplate.ush b/Engine/Shaders/Private/MaterialTemplate.ush
--- a/Engine/Shaders/Private/MaterialTemplate.ush
+++ b/Engine/Shaders/Private/MaterialTemplate.ush
@@ -1682,7 +1682,8 @@ MaterialFloat4 MobileSceneTextureLookup(inout FMaterialPixelParameters Parameter
 	}
 	else if (SceneTextureId == PPI_CustomStencil)
 	{
-		MaterialFloat Stencil = MobileSceneTextures.CustomStencilTexture.Load(int3(Parameters.SvPosition.xy, 0));
+		MaterialFloat Stencil = Texture2DSample(MobileSceneTextures.MobileCustomStencilTexture, MobileSceneTextures.MobileCustomStencilTextureSampler, UV).r*255.0;
+		Stencil = floor(Stencil + 0.5);
 		return MaterialFloat4(Stencil.rrr, 0.f);
 	}
 #endif// FEATURE_LEVEL
diff --git a/Engine/Shaders/Private/PostProcessMaterialShaders.usf b/Engine/Shaders/Private/PostProcessMaterialShaders.usf
--- a/Engine/Shaders/Private/PostProcessMaterialShaders.usf
+++ b/Engine/Shaders/Private/PostProcessMaterialShaders.usf
@@ -248,7 +248,8 @@ void MainPS_VideoOverlay(
 
 #else // FEATURE_LEVEL_ES3_1
 
-Texture2D<uint> MobileCustomStencilTexture;
+Texture2D MobileCustomStencilTexture;
+SamplerState MobileCustomStencilTextureSampler;
 
 int MobileStencilValueRef; // Use integer cause it has to be negative to make the less function to pass the test as always
 
@@ -312,7 +313,7 @@ void MainPS(
 	float ClipValue = 1.0f;
 
 #if MOBILE_STENCIL_COMPARE_FUNCTION != MOBILE_STENCIL_COMPARE_NEVER && MOBILE_STENCIL_COMPARE_FUNCTION != MOBILE_STENCIL_COMPARE_ALWAYS
-	int MobileStencilValue = MobileCustomStencilTexture.Load(int3(Input.Position.xy, 0));
+	int MobileStencilValue = floor(Texture2DSample(MobileCustomStencilTexture, MobileCustomStencilTextureSampler, InViewportUV).r*255.0f + 0.5f);
 #endif
 
 #if MOBILE_STENCIL_COMPARE_FUNCTION == MOBILE_STENCIL_COMPARE_LESS
diff --git a/Engine/Source/Runtime/Renderer/Private/CustomDepthRendering.cpp b/Engine/Source/Runtime/Renderer/Private/CustomDepthRendering.cpp
--- a/Engine/Source/Runtime/Renderer/Private/CustomDepthRendering.cpp
+++ b/Engine/Source/Runtime/Renderer/Private/CustomDepthRendering.cpp
@@ -96,7 +96,7 @@ void FCustomDepthPassMeshProcessor::AddMeshBatch(const FMeshBatch& RESTRICT Mesh
 			if (FeatureLevel <= ERHIFeatureLevel::ES3_1)
 			{
 				// On mobile platforms write custom stencil value to color target
-				MobileColorValue = CustomDepthStencilValue;
+				MobileColorValue = CustomDepthStencilValue / 255.0f;
 			}
 		}
 		else
diff --git a/Engine/Source/Runtime/Renderer/Private/PostProcess/PostProcessMaterial.cpp b/Engine/Source/Runtime/Renderer/Private/PostProcess/PostProcessMaterial.cpp
--- a/Engine/Source/Runtime/Renderer/Private/PostProcess/PostProcessMaterial.cpp
+++ b/Engine/Source/Runtime/Renderer/Private/PostProcess/PostProcessMaterial.cpp
@@ -437,6 +437,7 @@ FScreenPassTexture AddPostProcessMaterialPass(
 
 	PostProcessMaterialParameters->PostProcessOutput = GetScreenPassTextureViewportParameters(OutputViewport);
 	PostProcessMaterialParameters->MobileCustomStencilTexture = DepthStencilTexture;
+	PostProcessMaterialParameters->MobileCustomStencilTextureSampler = TStaticSamplerState<SF_Point, AM_Clamp, AM_Clamp, AM_Clamp>::GetRHI();
 	PostProcessMaterialParameters->MobileStencilValueRef = MaterialStencilRef;
 	PostProcessMaterialParameters->RenderTargets[0] = Output.GetRenderTargetBinding();
 
@@ -452,7 +453,7 @@ FScreenPassTexture AddPostProcessMaterialPass(
 	}
 	else if (!DepthStencilTexture && bMobilePlatform && Material->IsStencilTestEnabled()) // we have to set a default texture for MobileStencilTexture and override the MobileStencilValueRef to make all function to pass the stencil test
 	{
-		PostProcessMaterialParameters->MobileCustomStencilTexture = GSystemTextures.GetZeroUIntDummy(GraphBuilder);
+		PostProcessMaterialParameters->MobileCustomStencilTexture = GSystemTextures.GetBlackDummy(GraphBuilder);
 		
 		switch (Material->GetStencilCompare())
 		{
diff --git a/Engine/Source/Runtime/Renderer/Private/PostProcess/PostProcessMaterial.h b/Engine/Source/Runtime/Renderer/Private/PostProcess/PostProcessMaterial.h
--- a/Engine/Source/Runtime/Renderer/Private/PostProcess/PostProcessMaterial.h
+++ b/Engine/Source/Runtime/Renderer/Private/PostProcess/PostProcessMaterial.h
@@ -40,7 +40,8 @@ BEGIN_SHADER_PARAMETER_STRUCT(FPostProcessMaterialParameters, )
 	SHADER_PARAMETER_STRUCT(FScreenPassTextureViewportParameters, PostProcessOutput)
 	SHADER_PARAMETER_STRUCT_ARRAY(FScreenPassTextureInput, PostProcessInput, [kPostProcessMaterialInputCountMax])
 	SHADER_PARAMETER_SAMPLER(SamplerState, PostProcessInput_BilinearSampler)
-	SHADER_PARAMETER_RDG_TEXTURE(Texture2D<uint>, MobileCustomStencilTexture)
+	SHADER_PARAMETER_RDG_TEXTURE(Texture2D, MobileCustomStencilTexture)
+	SHADER_PARAMETER_SAMPLER(SamplerState, MobileCustomStencilTextureSampler)
 	SHADER_PARAMETER(int32, MobileStencilValueRef)
 	SHADER_PARAMETER(uint32, bFlipYAxis)
 	RENDER_TARGET_BINDING_SLOTS()
diff --git a/Engine/Source/Runtime/Renderer/Private/PostProcess/SceneRenderTargets.cpp b/Engine/Source/Runtime/Renderer/Private/PostProcess/SceneRenderTargets.cpp
--- a/Engine/Source/Runtime/Renderer/Private/PostProcess/SceneRenderTargets.cpp
+++ b/Engine/Source/Runtime/Renderer/Private/PostProcess/SceneRenderTargets.cpp
@@ -2935,7 +2935,7 @@ IPooledRenderTarget* FSceneRenderTargets::RequestCustomDepth(FRHICommandListImme
 					CustomDepthStencilTargetableFlags |= TexCreate_Memoryless;
 				}
 
-				FPooledRenderTargetDesc MobileCustomStencilDesc(FPooledRenderTargetDesc::Create2DDesc(CustomDepthBufferSize, PF_R8_UINT, FClearValueBinding::Transparent, TexCreate_None, CustomDepthStencilTargetableFlags, false));
+				FPooledRenderTargetDesc MobileCustomStencilDesc(FPooledRenderTargetDesc::Create2DDesc(CustomDepthBufferSize, PF_G8, FClearValueBinding::Transparent, TexCreate_None, CustomDepthStencilTargetableFlags, false));
 				GRenderTargetPool.FindFreeElement(RHICmdList, MobileCustomStencilDesc, MobileCustomStencil, TEXT("MobileCustomStencil"));
 
 				CustomStencilSRV.SafeRelease();
@@ -3243,7 +3243,6 @@ void SetupMobileSceneTextureUniformParameters(
 {
 	FRHITexture* BlackDefault2D = GSystemTextures.BlackDummy->GetRenderTargetItem().ShaderResourceTexture;
 	FRHITexture* DepthDefault = GSystemTextures.DepthDummy->GetRenderTargetItem().ShaderResourceTexture;
-	FRHITexture* ZeroUIntDummy = GSystemTextures.ZeroUIntDummy->GetRenderTargetItem().ShaderResourceTexture;
 
 	SceneTextureParameters.SceneColorTexture = bSceneTexturesValid ? SceneContext.GetSceneColorTexture().GetReference() : BlackDefault2D;
 	SceneTextureParameters.SceneColorTextureSampler = TStaticSamplerState<SF_Point, AM_Clamp, AM_Clamp, AM_Clamp>::GetRHI();
@@ -3268,14 +3267,15 @@ void SetupMobileSceneTextureUniformParameters(
 	SceneTextureParameters.CustomDepthTexture = CustomDepth;
 	SceneTextureParameters.CustomDepthTextureSampler = TStaticSamplerState<>::GetRHI();
 
-	FRHITexture* MobileCustomStencil = ZeroUIntDummy;
+	FRHITexture* MobileCustomStencil = BlackDefault2D;
 
 	if (SceneContext.MobileCustomStencil.IsValid())
 	{
 		MobileCustomStencil = SceneContext.MobileCustomStencil->GetRenderTargetItem().ShaderResourceTexture;
 	}
 
-	SceneTextureParameters.CustomStencilTexture = MobileCustomStencil;
+	SceneTextureParameters.MobileCustomStencilTexture = MobileCustomStencil;
+	SceneTextureParameters.MobileCustomStencilTextureSampler = TStaticSamplerState<SF_Point, AM_Clamp, AM_Clamp, AM_Clamp>::GetRHI();
 	
 	if (SceneContext.VirtualTextureFeedback.FeedbackBufferUAV.IsValid())
 	{
diff --git a/Engine/Source/Runtime/Renderer/Public/SceneRenderTargetParameters.h b/Engine/Source/Runtime/Renderer/Public/SceneRenderTargetParameters.h
--- a/Engine/Source/Runtime/Renderer/Public/SceneRenderTargetParameters.h
+++ b/Engine/Source/Runtime/Renderer/Public/SceneRenderTargetParameters.h
@@ -111,7 +111,8 @@ BEGIN_GLOBAL_SHADER_PARAMETER_STRUCT(FMobileSceneTextureUniformParameters, RENDE
 	SHADER_PARAMETER_SAMPLER(SamplerState, SceneAlphaCopyTextureSampler)
 	SHADER_PARAMETER_TEXTURE(Texture2D, CustomDepthTexture)
 	SHADER_PARAMETER_SAMPLER(SamplerState, CustomDepthTextureSampler)
-	SHADER_PARAMETER_TEXTURE(Texture2D<uint>, CustomStencilTexture)
+	SHADER_PARAMETER_TEXTURE(Texture2D, MobileCustomStencilTexture)
+	SHADER_PARAMETER_SAMPLER(SamplerState, MobileCustomStencilTextureSampler)
 	SHADER_PARAMETER_UAV(RWBuffer<uint>, VirtualTextureFeedbackUAV)
 END_GLOBAL_SHADER_PARAMETER_STRUCT()
 
