Fixed bogus static warning C6385 (Reading invalid data from xxx the readable size is..)

diff --git a/Engine/Source/Runtime/Renderer/Private/PostProcess/SceneRenderTargets.cpp b/Engine/Source/Runtime/Renderer/Private/PostProcess/SceneRenderTargets.cpp
--- a/Engine/Source/Runtime/Renderer/Private/PostProcess/SceneRenderTargets.cpp
+++ b/Engine/Source/Runtime/Renderer/Private/PostProcess/SceneRenderTargets.cpp
@@ -2150,6 +2150,7 @@ public:
 
 	const TCHAR* const operator[] (uint32 Idx) const
 	{
+		CA_SUPPRESS(6385);	// Doesn't like COM
 		return *Names[Idx];
 	}
 
@@ -2166,6 +2167,7 @@ static const TCHAR* const GetVolumeName(uint32 Id, bool bDirectional)
 
 	check(Id < MaxNames);
 
+	CA_SUPPRESS(6385);	// Doesn't like COM
 	return bDirectional ? NamesDir[Id] : Names[Id];
 }
 
diff --git a/Engine/Source/Runtime/Renderer/Private/TranslucentLighting.cpp b/Engine/Source/Runtime/Renderer/Private/TranslucentLighting.cpp
--- a/Engine/Source/Runtime/Renderer/Private/TranslucentLighting.cpp
+++ b/Engine/Source/Runtime/Renderer/Private/TranslucentLighting.cpp
@@ -1528,6 +1528,7 @@ void FDeferredShadingSceneRenderer::InjectTranslucentVolumeLightingArray(FRHICom
 	TArray<FTranslucentLightInjectionData, SceneRenderingAllocator> *LightInjectionData = new TArray<FTranslucentLightInjectionData, SceneRenderingAllocator>[Views.Num()];
 	for(int32 ViewIndex = 0; ViewIndex < Views.Num(); ++ViewIndex)
 	{
+		CA_SUPPRESS(6385);	// Doesn't like COM
 		LightInjectionData[ViewIndex].Empty(LightsEndIndex - FirstLightIndex);
 	}
 	
@@ -1540,6 +1541,7 @@ void FDeferredShadingSceneRenderer::InjectTranslucentVolumeLightingArray(FRHICom
 		{
 			if(LightSceneInfo->ShouldRenderLight(Views[ViewIndex]))
 			{
+				CA_SUPPRESS(6385);	// Doesn't like COM
 				AddLightForInjection(*this, *LightSceneInfo, NULL, LightInjectionData[ViewIndex]);
 			}
 		}
diff --git a/Engine/Source/Runtime/Windows/D3D11RHI/Private/Windows/WindowsD3D11ConstantBuffer.cpp b/Engine/Source/Runtime/Windows/D3D11RHI/Private/Windows/WindowsD3D11ConstantBuffer.cpp
--- a/Engine/Source/Runtime/Windows/D3D11RHI/Private/Windows/WindowsD3D11ConstantBuffer.cpp
+++ b/Engine/Source/Runtime/Windows/D3D11RHI/Private/Windows/WindowsD3D11ConstantBuffer.cpp
@@ -25,6 +25,7 @@ void FWinD3D11ConstantBuffer::InitDynamicRHI()
 	CurrentSubBuffer = 0;
 	for(uint32 s = 0;s < NumSubBuffers;s++)
 	{
+		CA_SUPPRESS(6385);	// Doesn't like COM
 		VERIFYD3D11RESULT_EX(D3DRHI->GetDevice()->CreateBuffer(&BufferDesc, NULL, Buffers[s].GetInitReference()), D3DRHI->GetDevice());
 		UpdateBufferStats(Buffers[s],true);
 		BufferDesc.ByteWidth = Align(BufferDesc.ByteWidth / 2, 16);
