Fixing mobile auto exposure behavior.

diff --git a/Engine/Source/Runtime/Renderer/Private/PostProcess/PostProcessEyeAdaptation.cpp b/Engine/Source/Runtime/Renderer/Private/PostProcess/PostProcessEyeAdaptation.cpp
--- a/Engine/Source/Runtime/Renderer/Private/PostProcess/PostProcessEyeAdaptation.cpp
+++ b/Engine/Source/Runtime/Renderer/Private/PostProcess/PostProcessEyeAdaptation.cpp
@@ -70,7 +70,7 @@ static TAutoConsoleVariable<int32> CVarEnablePreExposureOnlyInTheEditor(
 	TEXT("When pre-exposure is enabled, 0 to enable it everywhere, 1 to enable it only in the editor (default).\n")
 	TEXT("This is to because it currently has an impact on the renderthread performance\n"),
 	ECVF_ReadOnly);
-const ERHIFeatureLevel::Type BasicEyeAdaptationMinFeatureLevel = ERHIFeatureLevel::ES3_1;
+const ERHIFeatureLevel::Type BasicEyeAdaptationMinFeatureLevel = ERHIFeatureLevel::SM5;
 }
 
 bool IsAutoExposureMethodSupported(ERHIFeatureLevel::Type FeatureLevel, EAutoExposureMethod AutoExposureMethodId)
