Fix a software occlusion bug caused by candidates can't be set to be occluded.(eg. when enable custom depth)

diff --git a/Engine/Source/Runtime/Renderer/Private/SceneSoftwareOcclusion.cpp b/Engine/Source/Runtime/Renderer/Private/SceneSoftwareOcclusion.cpp
--- a/Engine/Source/Runtime/Renderer/Private/SceneSoftwareOcclusion.cpp
+++ b/Engine/Source/Runtime/Renderer/Private/SceneSoftwareOcclusion.cpp
@@ -908,6 +908,11 @@ static int32 ApplyResults(const FScene* Scene, FViewInfo& View, const FOcclusion
 				View.PrimitiveDefinitelyUnoccludedMap[PrimitiveIndex] = true;
 			}
 		}
+		else
+		{
+			// This applies for those are not occludees, so just follow the frustum culling result.
+			View.PrimitiveDefinitelyUnoccludedMap[PrimitiveIndex] = View.PrimitiveVisibilityMap[PrimitiveIndex];
+		}
 	}
 
 	INC_DWORD_STAT_BY(STAT_SoftwareCulledPrimitives, NumOccluded);
