Fix a bug that ILC/VLM doesn't affect Static Staticmesh at all on Mobile.

diff --git a/Engine/Source/Runtime/Renderer/Private/MobileBasePass.cpp b/Engine/Source/Runtime/Renderer/Private/MobileBasePass.cpp
--- a/Engine/Source/Runtime/Renderer/Private/MobileBasePass.cpp
+++ b/Engine/Source/Runtime/Renderer/Private/MobileBasePass.cpp
@@ -302,7 +302,9 @@ ELightMapPolicyType MobileBasePass::SelectMeshLightmapPolicy(
 		else if (IsIndirectLightingCacheAllowed(FeatureLevel) /* implies bAllowStaticLighting*/
 			&& PrimitiveSceneProxy
 			// Movable objects need to get their GI from the indirect lighting cache
-			&& PrimitiveSceneProxy->IsMovable())
+			&& (PrimitiveSceneProxy->IsMovable() 
+			// Static objects with lightmap type set to ForceVolumetric also need to get GI from ILC
+			|| LightMapInteraction.GetType() == LMIT_GlobalVolume) )
 		{
 			if (bUseMovableLight)
 			{
