Fix a bug that Translucent material can't get valid custom depth RT on mobile platform and pc preview

diff --git a/Engine/Source/Runtime/Engine/Classes/Materials/MaterialInterface.h b/Engine/Source/Runtime/Engine/Classes/Materials/MaterialInterface.h
--- a/Engine/Source/Runtime/Engine/Classes/Materials/MaterialInterface.h
+++ b/Engine/Source/Runtime/Engine/Classes/Materials/MaterialInterface.h
@@ -82,6 +82,7 @@ struct ENGINE_API FMaterialRelevance
 	uint8 bUsesSingleLayerWaterMaterial : 1;
 	uint8 bHasVolumeMaterialDomain : 1;
 	uint8 bUsesDistanceCullFade : 1;
+	uint8 bUsesCustomDepthStencil : 1;
 
 	/** Default constructor */
 	FMaterialRelevance()
diff --git a/Engine/Source/Runtime/Engine/Private/Materials/MaterialInterface.cpp b/Engine/Source/Runtime/Engine/Private/Materials/MaterialInterface.cpp
--- a/Engine/Source/Runtime/Engine/Private/Materials/MaterialInterface.cpp
+++ b/Engine/Source/Runtime/Engine/Private/Materials/MaterialInterface.cpp
@@ -55,6 +55,7 @@ void FMaterialRelevance::SetPrimitiveViewRelevance(FPrimitiveViewRelevance& OutV
 	OutViewRelevance.bUsesSkyMaterial = bUsesSkyMaterial;
 	OutViewRelevance.bUsesSingleLayerWaterMaterial = bUsesSingleLayerWaterMaterial;
 	OutViewRelevance.bHasVolumeMaterialDomain = bHasVolumeMaterialDomain;
+	OutViewRelevance.bUsesCustomDepthStencil = bUsesCustomDepthStencil;
 }
 
 //////////////////////////////////////////////////////////////////////////
@@ -168,6 +169,7 @@ FMaterialRelevance UMaterialInterface::GetRelevance_Internal(const UMaterial* Ma
 			MaterialRelevance.bUsesSceneDepth = MaterialResource->MaterialUsesSceneDepthLookup_GameThread();
 			MaterialRelevance.bHasVolumeMaterialDomain = MaterialResource->IsVolumetricPrimitive();
 			MaterialRelevance.bUsesDistanceCullFade = MaterialResource->MaterialUsesDistanceCullFade_GameThread();
+			MaterialRelevance.bUsesCustomDepthStencil = MaterialResource->UsesCustomDepthStencil_GameThread();
 			MaterialRelevance.bUsesSkyMaterial = Material->bIsSky;
 			MaterialRelevance.bUsesSingleLayerWaterMaterial = bUsesSingleLayerWaterMaterial;
 		}
diff --git a/Engine/Source/Runtime/Engine/Private/Materials/MaterialShared.cpp b/Engine/Source/Runtime/Engine/Private/Materials/MaterialShared.cpp
--- a/Engine/Source/Runtime/Engine/Private/Materials/MaterialShared.cpp
+++ b/Engine/Source/Runtime/Engine/Private/Materials/MaterialShared.cpp
@@ -881,6 +881,11 @@ bool FMaterial::HasRuntimeVirtualTextureOutput_RenderThread() const
 	return RenderingThreadShaderMap ? RenderingThreadShaderMap->HasRuntimeVirtualTextureOutput() : false;
 }
 
+bool FMaterial::UsesCustomDepthStencil_GameThread() const
+{
+	return GameThreadShaderMap.GetReference() ? (GameThreadShaderMap->UsesSceneTexture(PPI_CustomDepth) || GameThreadShaderMap->UsesSceneTexture(PPI_CustomStencil)) : false;
+}
+
 FMaterialShaderMap* FMaterial::GetRenderingThreadShaderMap() const 
 { 
 	check(IsInParallelRenderingThread());
diff --git a/Engine/Source/Runtime/Engine/Public/MaterialShared.h b/Engine/Source/Runtime/Engine/Public/MaterialShared.h
--- a/Engine/Source/Runtime/Engine/Public/MaterialShared.h
+++ b/Engine/Source/Runtime/Engine/Public/MaterialShared.h
@@ -1619,6 +1619,9 @@ public:
 	/** Does the material have a runtime virtual texture output node. */
 	ENGINE_API bool HasRuntimeVirtualTextureOutput_RenderThread() const;
 
+	/** Does the material use CustomDepth or CustomStencil lookup */
+	ENGINE_API bool UsesCustomDepthStencil_GameThread() const;
+
 	/** Note: This function is only intended for use in deciding whether or not shader permutations are required before material translation occurs. */
 	ENGINE_API bool MaterialMayModifyMeshPosition() const;
 
diff --git a/Engine/Source/Runtime/Engine/Public/PrimitiveViewRelevance.h b/Engine/Source/Runtime/Engine/Public/PrimitiveViewRelevance.h
--- a/Engine/Source/Runtime/Engine/Public/PrimitiveViewRelevance.h
+++ b/Engine/Source/Runtime/Engine/Public/PrimitiveViewRelevance.h
@@ -83,6 +83,8 @@ struct FPrimitiveViewRelevance
 	uint32 bUsesSingleLayerWaterMaterial : 1;
 	/** Whether the view use custom data. */
 	uint32 bUseCustomViewData : 1;
+	/** Whether the primitive has materials that read the custom depth or custom stencil */
+	uint32 bUsesCustomDepthStencil : 1;
 
 	/** 
 	 * Whether this primitive view relevance has been initialized this frame.  
diff --git a/Engine/Source/Runtime/Renderer/Private/MobileShadingRenderer.cpp b/Engine/Source/Runtime/Renderer/Private/MobileShadingRenderer.cpp
--- a/Engine/Source/Runtime/Renderer/Private/MobileShadingRenderer.cpp
+++ b/Engine/Source/Runtime/Renderer/Private/MobileShadingRenderer.cpp
@@ -110,7 +110,7 @@ static bool UsesCustomDepthStencilLookup(const FViewInfo& View)
 	}
 
 	//TODO: check if translucency uses CustomDepth
-	return bPPUsesCustomDepth || bPPUsesCustomStencil;
+	return bPPUsesCustomDepth || bPPUsesCustomStencil || View.bUsesCustomDepthStencil;
 }
 
 
@@ -281,17 +281,12 @@ void FMobileSceneRenderer::InitViews(FRHICommandListImmediate& RHICmdList)
 
 	// update buffers used in cached mesh path
 	// in case there are multiple views, these buffers will be updated before rendering each view
+	// OpaqueBasePassUniformBuffer, TranslucentBasePassUniformBuffer and DistortionPassUniformBuffer may depend on custom depth render target, so they should be updated after custom depth rendering 
 	if (Views.Num() > 0)
 	{
 		const FViewInfo& View = Views[0];
 		Scene->UniformBuffers.UpdateViewUniformBuffer(View);
-		UpdateOpaqueBasePassUniformBuffer(RHICmdList, View);
-		UpdateTranslucentBasePassUniformBuffer(RHICmdList, View);
 		UpdateDirectionalLightUniformBuffers(RHICmdList, View);
-
-		FMobileDistortionPassUniformParameters DistortionPassParameters;
-		SetupMobileDistortionPassUniformBuffer(RHICmdList, View, DistortionPassParameters);
-		Scene->UniformBuffers.MobileDistortionPassUniformBuffer.UpdateUniformBufferImmediate(DistortionPassParameters);
 	}
 	UpdateSkyReflectionUniformBuffer();
 
@@ -445,6 +440,16 @@ void FMobileSceneRenderer::Render(FRHICommandListImmediate& RHICmdList)
 			RenderCustomDepthPass(RHICmdList);
 		}
 	}
+
+	// Update OpaqueBasePassUniformBuffer, TranslucentBasePassUniformBuffer and DistortionUniformBuffer right after CustomDepthPass and before BeginRenderPass of MobileBasePass
+	{
+		UpdateOpaqueBasePassUniformBuffer(RHICmdList, View);
+		UpdateTranslucentBasePassUniformBuffer(RHICmdList, View);
+
+		FMobileDistortionPassUniformParameters DistortionPassParameters;
+		SetupMobileDistortionPassUniformBuffer(RHICmdList, View, DistortionPassParameters);
+		Scene->UniformBuffers.MobileDistortionPassUniformBuffer.UpdateUniformBufferImmediate(DistortionPassParameters);
+	}
 		
 	//
 	FRHITexture* SceneColor = nullptr;
diff --git a/Engine/Source/Runtime/Renderer/Private/SceneRendering.cpp b/Engine/Source/Runtime/Renderer/Private/SceneRendering.cpp
--- a/Engine/Source/Runtime/Renderer/Private/SceneRendering.cpp
+++ b/Engine/Source/Runtime/Renderer/Private/SceneRendering.cpp
@@ -866,6 +866,7 @@ void FViewInfo::Init()
 	bHasCustomDepthPrimitives = false;
 	bHasDistortionPrimitives = false;
 	bAllowStencilDither = false;
+	bUsesCustomDepthStencil = false;
 
 	ForwardLightingResources = nullptr;
 
diff --git a/Engine/Source/Runtime/Renderer/Private/SceneRendering.h b/Engine/Source/Runtime/Renderer/Private/SceneRendering.h
--- a/Engine/Source/Runtime/Renderer/Private/SceneRendering.h
+++ b/Engine/Source/Runtime/Renderer/Private/SceneRendering.h
@@ -1010,7 +1010,7 @@ public:
 	uint32 bTranslucentSurfaceLighting : 1;
 	/** Whether the view has any materials that read from scene depth. */
 	uint32 bUsesSceneDepth : 1;
-
+	uint32 bUsesCustomDepthStencil : 1;
 
 	/** Whether fog should only be computed on rendered opaque pixels or not. */
 	uint32 bFogOnlyOnRenderedOpaque : 1;
diff --git a/Engine/Source/Runtime/Renderer/Private/SceneVisibility.cpp b/Engine/Source/Runtime/Renderer/Private/SceneVisibility.cpp
--- a/Engine/Source/Runtime/Renderer/Private/SceneVisibility.cpp
+++ b/Engine/Source/Runtime/Renderer/Private/SceneVisibility.cpp
@@ -1891,6 +1891,7 @@ struct FRelevancePacket
 	bool bUsesLightingChannels;
 	bool bTranslucentSurfaceLighting;
 	bool bUsesSceneDepth;
+	bool bUsesCustomDepthStencil;
 	bool bSceneHasSkyMaterial;
 	bool bHasSingleLayerWaterMaterial;
 
@@ -1929,6 +1930,7 @@ struct FRelevancePacket
 		, bUsesLightingChannels(false)
 		, bTranslucentSurfaceLighting(false)
 		, bUsesSceneDepth(false)
+		, bUsesCustomDepthStencil(false)
 		, bSceneHasSkyMaterial(false)
 		, bHasSingleLayerWaterMaterial(false)
 	{
@@ -2049,6 +2051,7 @@ struct FRelevancePacket
 			bUsesLightingChannels |= ViewRelevance.bUsesLightingChannels;
 			bTranslucentSurfaceLighting |= ViewRelevance.bTranslucentSurfaceLighting;
 			bUsesSceneDepth |= ViewRelevance.bUsesSceneDepth;
+			bUsesCustomDepthStencil |= ViewRelevance.bUsesCustomDepthStencil;
 			bSceneHasSkyMaterial |= ViewRelevance.bUsesSkyMaterial;
 			bHasSingleLayerWaterMaterial |= ViewRelevance.bUsesSingleLayerWaterMaterial;
 
@@ -2403,6 +2406,7 @@ struct FRelevancePacket
 		WriteView.TranslucentPrimCount.Append(TranslucentPrimCount);
 		WriteView.bHasDistortionPrimitives |= bHasDistortionPrimitives;
 		WriteView.bHasCustomDepthPrimitives |= bHasCustomDepthPrimitives;
+		WriteView.bUsesCustomDepthStencil |= bUsesCustomDepthStencil;
 		DirtyIndirectLightingCacheBufferPrimitives.AppendTo(WriteView.DirtyIndirectLightingCacheBufferPrimitives);
 
 		WriteView.MeshDecalBatches.Append(MeshDecalBatches);
