Try to fix a bug of Stationary Sky Light has double lighting contributions to StaticMeshActor on MobilePlatform.

diff --git a/Engine/Source/Runtime/Renderer/Private/MobileBasePass.cpp b/Engine/Source/Runtime/Renderer/Private/MobileBasePass.cpp
--- a/Engine/Source/Runtime/Renderer/Private/MobileBasePass.cpp
+++ b/Engine/Source/Runtime/Renderer/Private/MobileBasePass.cpp
@@ -620,7 +620,8 @@ void FMobileBasePassMeshProcessor::Process(
 		FBaseDS,
 		TMobileBasePassPSPolicyParamType<FUniformLightMapPolicy>> BasePassShaders;
 	
-	bool bEnableSkyLight = ShadingModels.IsLit() && Scene && Scene->ShouldRenderSkylightInBasePass(BlendMode);
+	//The stationary skylight contribution has been added to the LowQuality Lightmap for StaticMeshActor on mobile, so we should skip the sky light spherical harmonic contribution for it.
+	bool bEnableSkyLight = ShadingModels.IsLit() && Scene && Scene->ShouldRenderSkylightInBasePass(BlendMode) && !(FReadOnlyCVARCache::Get().bAllowStaticLighting && PrimitiveSceneProxy->IsStatic());
 	int32 NumMovablePointLights = MobileBasePass::CalcNumMovablePointLights(MaterialResource, PrimitiveSceneProxy);
 
 	MobileBasePass::GetShaders(
