Eliminate a warning caused by CL#9293420

diff --git a/Engine/Source/Runtime/Renderer/Private/MobileBasePass.cpp b/Engine/Source/Runtime/Renderer/Private/MobileBasePass.cpp
--- a/Engine/Source/Runtime/Renderer/Private/MobileBasePass.cpp
+++ b/Engine/Source/Runtime/Renderer/Private/MobileBasePass.cpp
@@ -621,7 +621,7 @@ void FMobileBasePassMeshProcessor::Process(
 		TMobileBasePassPSPolicyParamType<FUniformLightMapPolicy>> BasePassShaders;
 	
 	//The stationary skylight contribution has been added to the LowQuality Lightmap for StaticMeshActor on mobile, so we should skip the sky light spherical harmonic contribution for it.
-	bool bEnableSkyLight = ShadingModels.IsLit() && Scene && Scene->ShouldRenderSkylightInBasePass(BlendMode) && !(FReadOnlyCVARCache::Get().bAllowStaticLighting && PrimitiveSceneProxy->IsStatic());
+	bool bEnableSkyLight = ShadingModels.IsLit() && Scene && Scene->ShouldRenderSkylightInBasePass(BlendMode) && !(FReadOnlyCVARCache::Get().bAllowStaticLighting && PrimitiveSceneProxy && PrimitiveSceneProxy->IsStatic());
 	int32 NumMovablePointLights = MobileBasePass::CalcNumMovablePointLights(MaterialResource, PrimitiveSceneProxy);
 
 	MobileBasePass::GetShaders(
