Conditional creation of SkyLight sampling structure before running CreateGatherPoints.

diff --git a/Engine/Source/Runtime/Renderer/Private/RayTracing/RayTracingGlobalIllumination.cpp b/Engine/Source/Runtime/Renderer/Private/RayTracing/RayTracingGlobalIllumination.cpp
--- a/Engine/Source/Runtime/Renderer/Private/RayTracing/RayTracingGlobalIllumination.cpp
+++ b/Engine/Source/Runtime/Renderer/Private/RayTracing/RayTracingGlobalIllumination.cpp
@@ -535,6 +535,11 @@ void FDeferredShadingSceneRenderer::RayTracingGlobalIlluminationCreateGatherPoin
 	FPathTracingLightData LightParameters;
 	SetupLightParameters(Scene->Lights, View, &LightParameters);
 
+	if (Scene->SkyLight && Scene->SkyLight->ShouldRebuildCdf())
+	{
+		// TODO: should be converted to RDG.
+		BuildSkyLightCdfs(GraphBuilder.RHICmdList, Scene->SkyLight);
+	}
 	FSkyLightData SkyLightParameters;
 	SetupSkyLightParameters(*Scene, &SkyLightParameters);
 
