DRPS - Merge fixes
#rnx

diff --git a/Engine/Platforms/HTML5/Source/Developer/HTML5PlatformEditor/Private/HTML5TargetSettingsCustomization.cpp b/Engine/Platforms/HTML5/Source/Developer/HTML5PlatformEditor/Private/HTML5TargetSettingsCustomization.cpp
--- a/Engine/Platforms/HTML5/Source/Developer/HTML5PlatformEditor/Private/HTML5TargetSettingsCustomization.cpp
+++ b/Engine/Platforms/HTML5/Source/Developer/HTML5PlatformEditor/Private/HTML5TargetSettingsCustomization.cpp
@@ -105,7 +105,7 @@ void FHTML5TargetSettingsCustomization::CustomizeDetails(IDetailLayoutBuilder& D
 	BuildDeprecationMessage(DetailLayout); // !!! HTML5 Deprecation Message
 	// --------------------------------------------------
 	BuildEmscriptenCategory(DetailLayout); // !!! HTML5 Deprecation Message
-	AudioPluginWidgetManager.BuildAudioCategory(DetailLayout, EAudioPlatform::HTML5);
+	AudioPluginWidgetManager.BuildAudioCategory(DetailLayout, FString("HTML5"));
 }
 
 void FHTML5TargetSettingsCustomization::BuildEmscriptenCategory(IDetailLayoutBuilder& DetailLayout)
diff --git a/Engine/Source/Runtime/HeadMountedDisplay/Private/DefaultStereoLayers.cpp b/Engine/Source/Runtime/HeadMountedDisplay/Private/DefaultStereoLayers.cpp
--- a/Engine/Source/Runtime/HeadMountedDisplay/Private/DefaultStereoLayers.cpp
+++ b/Engine/Source/Runtime/HeadMountedDisplay/Private/DefaultStereoLayers.cpp
@@ -209,7 +209,7 @@ void FDefaultStereoLayers::PreRenderViewFamily_RenderThread(FRHICommandListImmed
 
 void FDefaultStereoLayers::PostRenderView_RenderThread(FRHICommandListImmediate& RHICmdList, FSceneView& InView)
 {
-	if (!HMDDevice->DeviceIsStereoEyeView(InView.StereoPass))
+	if (!HMDDevice->IsStereoEyePass(InView.StereoPass))
 	{
 		return;
 	}
diff --git a/Engine/Source/Runtime/Renderer/Private/BasePassRendering.cpp b/Engine/Source/Runtime/Renderer/Private/BasePassRendering.cpp
--- a/Engine/Source/Runtime/Renderer/Private/BasePassRendering.cpp
+++ b/Engine/Source/Runtime/Renderer/Private/BasePassRendering.cpp
@@ -485,7 +485,7 @@ void SetupSharedBasePassParameters(
 
 	SetupFogUniformParameters(View, SharedParameters.Fog);
 
-	if (View.bIsInstancedStereoEnabled && IStereoRendering::IsAPrimaryView(View.StereoPass))
+	if (View.bIsInstancedStereoEnabled && View.StereoPass == EStereoscopicPass::eSSP_LEFT_EYE)
 	{
 		const FSceneView& RightEye = *View.Family->Views[1];
 		SharedParameters.ForwardISR = RightEye.ForwardLightingResources->ForwardLightData;
diff --git a/Engine/Source/Runtime/Renderer/Private/PostProcess/PostProcessing.cpp b/Engine/Source/Runtime/Renderer/Private/PostProcess/PostProcessing.cpp
--- a/Engine/Source/Runtime/Renderer/Private/PostProcess/PostProcessing.cpp
+++ b/Engine/Source/Runtime/Renderer/Private/PostProcess/PostProcessing.cpp
@@ -326,7 +326,7 @@ void AddPostProcessingPasses(FRDGBuilder& GraphBuilder, const FViewInfo& View, c
 		};
 
 		const EStereoscopicPass StereoPass = View.StereoPass;
-		const bool bPrimaryView = IStereoRendering::IsAPrimaryView(StereoPass);
+		const bool bPrimaryView = IStereoRendering::IsAPrimaryView(View, GEngine->StereoRenderingDevice);
 		const bool bHasViewState = View.ViewState != nullptr;
 		const bool bDepthOfFieldEnabled = DiaphragmDOF::IsEnabled(View);
 		const bool bVisualizeDepthOfField = bDepthOfFieldEnabled && EngineShowFlags.VisualizeDOF;
@@ -1198,7 +1198,7 @@ static FRCPassPostProcessTonemap* AddTonemapper(
 	const EStereoscopicPass StereoPass = View.StereoPass;
 
 	FRenderingCompositeOutputRef TonemapperCombinedLUTOutputRef;
-	if (IStereoRendering::IsAPrimaryView(StereoPass))
+	if (IStereoRendering::IsAPrimaryView(View, GEngine->StereoRenderingDevice))
 	{
 		TonemapperCombinedLUTOutputRef = AddCombineLUTPass(Context.Graph);
 	}
