Fix invalid permutations of shadow depth vertex shader causing spline mesh component to assert on unbound parameters.
Also reinstated the SPF_Mandatory flag in SplineMeshComponent.cpp

diff --git a/Engine/Shaders/Private/ShadowDepthVertexShader.usf b/Engine/Shaders/Private/ShadowDepthVertexShader.usf
--- a/Engine/Shaders/Private/ShadowDepthVertexShader.usf
+++ b/Engine/Shaders/Private/ShadowDepthVertexShader.usf
@@ -359,6 +359,8 @@ void Main(
 	OutputVertexID( OutParameters );
 }
 
+#if IS_FOR_GEOMETRY_SHADER
+
 // Vertex shader exports must match GS inputs.
 // This is a requirement on some platforms.
 void MainForGS(
@@ -368,11 +370,11 @@ void MainForGS(
 	out FShadowDepthVSToDS OutParameters
 #elif ONEPASS_POINTLIGHT_SHADOW
 	out FShadowDepthVSToGS OutParameters
-#if FEATURE_LEVEL >= FEATURE_LEVEL_SM4 && USING_VERTEX_SHADER_LAYER
-	, out uint LayerIndex : SV_RenderTargetArrayIndex
-#endif
+	#if FEATURE_LEVEL >= FEATURE_LEVEL_SM4 && USING_VERTEX_SHADER_LAYER
+		, out uint LayerIndex : SV_RenderTargetArrayIndex
+	#endif
 #else
-	out FShadowDepthVSToPS OutParameters
+	#error Invalid permutation for MainForGS
 #endif
 #if COMPILER_METAL
 #define SKIP_UNUSED_POSITION
@@ -394,6 +396,8 @@ void MainForGS(
 		UnusedPosition);
 }
 
+#endif // IS_FOR_GEOMETRY_SHADER
+
 #define USING_ONEPASS_POINT_VERTEX_LAYER	(FEATURE_LEVEL >= FEATURE_LEVEL_SM4 && ONEPASS_POINTLIGHT_SHADOW && USING_VERTEX_SHADER_LAYER)
 
 #if POSITION_ONLY
@@ -458,6 +462,8 @@ void PositionOnlyMain(
 }
 
 
+#if IS_FOR_GEOMETRY_SHADER
+
 // Vertex shader exports must match GS inputs.
 // This is a requirement on some platforms.
 void PositionOnlyMainForGS(
@@ -468,7 +474,7 @@ void PositionOnlyMainForGS(
 #if ONEPASS_POINTLIGHT_SHADOW
 	out FShadowDepthVSToGS OutParameters
 #else
-	out FShadowDepthVSToPS OutParameters
+	#error Invalid permutation for PositionOnlyMainForGS
 #endif
 #if COMPILER_METAL
 	#ifndef SKIP_UNUSED_POSITION
@@ -490,4 +496,6 @@ void PositionOnlyMainForGS(
 		OutParameters, UnusedPosition);
 }
 
+#endif // IS_FOR_GEOMETRY_SHADER
+
 #endif
diff --git a/Engine/Source/Runtime/Engine/Private/Components/SplineMeshComponent.cpp b/Engine/Source/Runtime/Engine/Private/Components/SplineMeshComponent.cpp
--- a/Engine/Source/Runtime/Engine/Private/Components/SplineMeshComponent.cpp
+++ b/Engine/Source/Runtime/Engine/Private/Components/SplineMeshComponent.cpp
@@ -34,8 +34,7 @@ static FAutoConsoleVariableRef CVarNoRecreateSplineMeshProxy(
 
 void FSplineMeshVertexFactoryShaderParameters::Bind(const FShaderParameterMap& ParameterMap)
 {
-	// @todo: fix the shader compile errors and restore the mandatory flag
-	SplineMeshParams.Bind(ParameterMap, TEXT("SplineParams")/*, SPF_Mandatory*/);
+	SplineMeshParams.Bind(ParameterMap, TEXT("SplineParams"), SPF_Mandatory);
 }
 
 void FSplineMeshVertexFactoryShaderParameters::GetElementShaderBindings(
diff --git a/Engine/Source/Runtime/Renderer/Private/ShadowDepthRendering.cpp b/Engine/Source/Runtime/Renderer/Private/ShadowDepthRendering.cpp
--- a/Engine/Source/Runtime/Renderer/Private/ShadowDepthRendering.cpp
+++ b/Engine/Source/Runtime/Renderer/Private/ShadowDepthRendering.cpp
@@ -253,6 +253,7 @@ public:
 		OutEnvironment.SetDefine(TEXT("ONEPASS_POINTLIGHT_SHADOW"), (uint32)(ShaderMode == VertexShadowDepth_OnePassPointLight));
 		OutEnvironment.SetDefine(TEXT("REFLECTIVE_SHADOW_MAP"), (uint32)bRenderReflectiveShadowMap);
 		OutEnvironment.SetDefine(TEXT("POSITION_ONLY"), (uint32)bUsePositionOnlyStream);
+		OutEnvironment.SetDefine(TEXT("IS_FOR_GEOMETRY_SHADER"), (uint32)bIsForGeometryShader);
 
 		if (bIsForGeometryShader)
 		{
@@ -371,18 +372,16 @@ public:
 	FOnePassPointShadowDepthGS() {}
 };
 
-#define IMPLEMENT_SHADOW_DEPTH_SHADERMODE_SHADERS(ShaderMode,bRenderReflectiveShadowMap) \
-	typedef TShadowDepthVS<ShaderMode, bRenderReflectiveShadowMap, false> TShadowDepthVS##ShaderMode##bRenderReflectiveShadowMap;	\
-	IMPLEMENT_MATERIAL_SHADER_TYPE(template<>,TShadowDepthVS##ShaderMode##bRenderReflectiveShadowMap,TEXT("/Engine/Private/ShadowDepthVertexShader.usf"),TEXT("Main"),SF_Vertex);	\
-	typedef TShadowDepthVS<ShaderMode, bRenderReflectiveShadowMap, false, true> TShadowDepthVSForGS##ShaderMode##bRenderReflectiveShadowMap;	\
-	IMPLEMENT_MATERIAL_SHADER_TYPE(template<>,TShadowDepthVSForGS##ShaderMode##bRenderReflectiveShadowMap,TEXT("/Engine/Private/ShadowDepthVertexShader.usf"),TEXT("MainForGS"),SF_Vertex);	\
-	typedef TShadowDepthHS<ShaderMode, bRenderReflectiveShadowMap> TShadowDepthHS##ShaderMode##bRenderReflectiveShadowMap;	\
-	IMPLEMENT_MATERIAL_SHADER_TYPE(template<>,TShadowDepthHS##ShaderMode##bRenderReflectiveShadowMap,TEXT("/Engine/Private/ShadowDepthVertexShader.usf"),TEXT("MainHull"),SF_Hull);	\
-	typedef TShadowDepthDS<ShaderMode, bRenderReflectiveShadowMap> TShadowDepthDS##ShaderMode##bRenderReflectiveShadowMap;	\
-	IMPLEMENT_MATERIAL_SHADER_TYPE(template<>,TShadowDepthDS##ShaderMode##bRenderReflectiveShadowMap,TEXT("/Engine/Private/ShadowDepthVertexShader.usf"),TEXT("MainDomain"),SF_Domain);
-
 IMPLEMENT_SHADER_TYPE(, FOnePassPointShadowDepthGS, TEXT("/Engine/Private/ShadowDepthVertexShader.usf"), TEXT("MainOnePassPointLightGS"), SF_Geometry);
 
+#define IMPLEMENT_SHADOW_DEPTH_SHADERMODE_SHADERS(ShaderMode,bRenderReflectiveShadowMap) \
+	typedef TShadowDepthVS<ShaderMode, bRenderReflectiveShadowMap, false> TShadowDepthVS##ShaderMode##bRenderReflectiveShadowMap; \
+	typedef TShadowDepthHS<ShaderMode, bRenderReflectiveShadowMap       > TShadowDepthHS##ShaderMode##bRenderReflectiveShadowMap; \
+	typedef TShadowDepthDS<ShaderMode, bRenderReflectiveShadowMap       > TShadowDepthDS##ShaderMode##bRenderReflectiveShadowMap; \
+	IMPLEMENT_MATERIAL_SHADER_TYPE(template<>, TShadowDepthVS##ShaderMode##bRenderReflectiveShadowMap, TEXT("/Engine/Private/ShadowDepthVertexShader.usf"), TEXT("Main"),       SF_Vertex); \
+	IMPLEMENT_MATERIAL_SHADER_TYPE(template<>, TShadowDepthHS##ShaderMode##bRenderReflectiveShadowMap, TEXT("/Engine/Private/ShadowDepthVertexShader.usf"), TEXT("MainHull"),   SF_Hull  ); \
+	IMPLEMENT_MATERIAL_SHADER_TYPE(template<>, TShadowDepthDS##ShaderMode##bRenderReflectiveShadowMap, TEXT("/Engine/Private/ShadowDepthVertexShader.usf"), TEXT("MainDomain"), SF_Domain);
+
 IMPLEMENT_SHADOW_DEPTH_SHADERMODE_SHADERS(VertexShadowDepth_PerspectiveCorrect, true);
 IMPLEMENT_SHADOW_DEPTH_SHADERMODE_SHADERS(VertexShadowDepth_PerspectiveCorrect, false);
 IMPLEMENT_SHADOW_DEPTH_SHADERMODE_SHADERS(VertexShadowDepth_OutputDepth, true);
@@ -391,12 +390,16 @@ IMPLEMENT_SHADOW_DEPTH_SHADERMODE_SHADERS(VertexShadowDepth_OnePassPointLight, f
 
 // Position only vertex shaders.
 typedef TShadowDepthVS<VertexShadowDepth_PerspectiveCorrect, false, true> TShadowDepthVSVertexShadowDepth_PerspectiveCorrectPositionOnly;
+typedef TShadowDepthVS<VertexShadowDepth_OutputDepth,        false, true> TShadowDepthVSVertexShadowDepth_OutputDepthPositionOnly;
+typedef TShadowDepthVS<VertexShadowDepth_OnePassPointLight,  false, true> TShadowDepthVSVertexShadowDepth_OnePassPointLightPositionOnly;
 IMPLEMENT_MATERIAL_SHADER_TYPE(template<>, TShadowDepthVSVertexShadowDepth_PerspectiveCorrectPositionOnly, TEXT("/Engine/Private/ShadowDepthVertexShader.usf"), TEXT("PositionOnlyMain"), SF_Vertex);
-typedef TShadowDepthVS<VertexShadowDepth_OutputDepth, false, true> TShadowDepthVSVertexShadowDepth_OutputDepthPositionOnly;
-IMPLEMENT_MATERIAL_SHADER_TYPE(template<>, TShadowDepthVSVertexShadowDepth_OutputDepthPositionOnly, TEXT("/Engine/Private/ShadowDepthVertexShader.usf"), TEXT("PositionOnlyMain"), SF_Vertex);
-typedef TShadowDepthVS<VertexShadowDepth_OnePassPointLight, false, true> TShadowDepthVSVertexShadowDepth_OnePassPointLightPositionOnly;
-IMPLEMENT_MATERIAL_SHADER_TYPE(template<>, TShadowDepthVSVertexShadowDepth_OnePassPointLightPositionOnly, TEXT("/Engine/Private/ShadowDepthVertexShader.usf"), TEXT("PositionOnlyMain"), SF_Vertex);
-typedef TShadowDepthVS<VertexShadowDepth_OnePassPointLight, false, true, true> TShadowDepthVSForGSVertexShadowDepth_OnePassPointLightPositionOnly;
+IMPLEMENT_MATERIAL_SHADER_TYPE(template<>, TShadowDepthVSVertexShadowDepth_OutputDepthPositionOnly,        TEXT("/Engine/Private/ShadowDepthVertexShader.usf"), TEXT("PositionOnlyMain"), SF_Vertex);
+IMPLEMENT_MATERIAL_SHADER_TYPE(template<>, TShadowDepthVSVertexShadowDepth_OnePassPointLightPositionOnly,  TEXT("/Engine/Private/ShadowDepthVertexShader.usf"), TEXT("PositionOnlyMain"), SF_Vertex);
+
+// One pass point light VS for GS shaders.
+typedef TShadowDepthVS<VertexShadowDepth_OnePassPointLight, false, false, true> TShadowDepthVSForGSVertexShadowDepth_OnePassPointLight;
+typedef TShadowDepthVS<VertexShadowDepth_OnePassPointLight, false, true,  true> TShadowDepthVSForGSVertexShadowDepth_OnePassPointLightPositionOnly;
+IMPLEMENT_MATERIAL_SHADER_TYPE(template<>, TShadowDepthVSForGSVertexShadowDepth_OnePassPointLight,             TEXT("/Engine/Private/ShadowDepthVertexShader.usf"), TEXT("MainForGS"),             SF_Vertex); \
 IMPLEMENT_MATERIAL_SHADER_TYPE(template<>, TShadowDepthVSForGSVertexShadowDepth_OnePassPointLightPositionOnly, TEXT("/Engine/Private/ShadowDepthVertexShader.usf"), TEXT("PositionOnlyMainForGS"), SF_Vertex);
 
 /**
