Add an option on material for forward shaing (including mobile) to use EnvBRDF instead of EnvBRDFApprox to have a better IBL effect.

diff --git a/Engine/Shaders/Private/BasePassPixelShader.usf b/Engine/Shaders/Private/BasePassPixelShader.usf
--- a/Engine/Shaders/Private/BasePassPixelShader.usf
+++ b/Engine/Shaders/Private/BasePassPixelShader.usf
@@ -14,7 +14,8 @@
 	#define EyeAdaptationStruct OpaqueBasePass
 #endif
 
-#if MATERIAL_SHADINGMODEL_SINGLELAYERWATER
+//ForwardShading (including mobile) use preintegrated GF lut for simple IBL, it could be enable in material editor.
+#if MATERIAL_SHADINGMODEL_SINGLELAYERWATER || (FORWARD_SHADING && MATERIAL_USE_PREINTEGRATED_GF)
 #define PreIntegratedGF			OpaqueBasePass.PreIntegratedGFTexture
 #if SUPPORTS_INDEPENDENT_SAMPLERS
 #define PreIntegratedGFSampler	View.SharedBilinearClampedSampler
diff --git a/Engine/Shaders/Private/ForwardLightingCommon.ush b/Engine/Shaders/Private/ForwardLightingCommon.ush
--- a/Engine/Shaders/Private/ForwardLightingCommon.ush
+++ b/Engine/Shaders/Private/ForwardLightingCommon.ush
@@ -390,7 +390,11 @@ half3 GetImageBasedReflectionLighting(FMaterialPixelParameters MaterialParameter
 	}
 #endif
 
+#if MATERIAL_USE_PREINTEGRATED_GF
+	SpecularColor = EnvBRDF(SpecularColor, Roughness, NoV);
+#else
 	SpecularColor = EnvBRDFApprox(SpecularColor, Roughness, NoV);
+#endif
 
 	float3 SpecularLighting = SpecularIBL.rgb;
 
diff --git a/Engine/Shaders/Private/MobileBasePassPixelShader.usf b/Engine/Shaders/Private/MobileBasePassPixelShader.usf
--- a/Engine/Shaders/Private/MobileBasePassPixelShader.usf
+++ b/Engine/Shaders/Private/MobileBasePassPixelShader.usf
@@ -9,6 +9,18 @@
 // Reroute MobileSceneTextures uniform buffer references to the base pass uniform buffer
 #define MobileSceneTextures MobileBasePass.SceneTextures
 
+#ifndef MOBILE_QL_FORCE_DISABLE_PREINTEGRATEDGF
+#define MOBILE_QL_FORCE_DISABLE_PREINTEGRATEDGF 0
+#endif
+
+#define FORWARDSHADING_USE_HQ_ENV_BRDF (MATERIAL_USE_PREINTEGRATED_GF && !MOBILE_QL_FORCE_DISABLE_PREINTEGRATEDGF)
+
+//use preintegrated GF lut for simple IBL
+#if FORWARDSHADING_USE_HQ_ENV_BRDF
+#define PreIntegratedGF			MobileBasePass.PreIntegratedGFTexture
+#define PreIntegratedGFSampler	MobileBasePass.PreIntegratedGFSampler
+#endif
+
 #include "SHCommon.ush"
 #include "MobileBasePassCommon.ush"
 #include "/Engine/Generated/Material.ush"
@@ -462,10 +474,14 @@ void Main(
 	// then defaults are the same as a non-metal,
 	// so this define is safe.
 	SpecularColor = EnvBRDFApproxNonmetal( Roughness, NoV );
+#else
+#if FORWARDSHADING_USE_HQ_ENV_BRDF
+	SpecularColor = EnvBRDF(SpecularColor, Roughness, NoV);
 #else
 	SpecularColor = EnvBRDFApprox( SpecularColor, Roughness, NoV );
 #endif
 #endif
+#endif
 
 #if MOBILE_EMULATION
 	{
diff --git a/Engine/Source/Runtime/Engine/Classes/Materials/Material.h b/Engine/Source/Runtime/Engine/Classes/Materials/Material.h
--- a/Engine/Source/Runtime/Engine/Classes/Materials/Material.h
+++ b/Engine/Source/Runtime/Engine/Classes/Materials/Material.h
@@ -767,6 +767,10 @@ public:
 	UPROPERTY(EditAnywhere, BlueprintReadOnly, Category=Mobile)
 	uint8 bUseLightmapDirectionality : 1;
 
+	/* Forward (including mobile) renderer: use preintegrated GF lut for simple IBL, but will use one more sampler. */
+	UPROPERTY(EditAnywhere, BlueprintReadOnly, Category = ForwardShading, meta = (DisplayName = "PreintegratedGF For Simple IBL"))
+	uint32 bForwardRenderUsePreintegratedGFForSimpleIBL : 1;
+
 	/* 
 	 * Forward renderer: enables multiple parallax-corrected reflection captures that blend together.
 	 * Mobile renderer: blend between nearest 3 reflection captures, but reduces the number of samplers available to the material as two more samplers will be used for reflection cubemaps.
diff --git a/Engine/Source/Runtime/Engine/Private/Materials/MaterialShared.cpp b/Engine/Source/Runtime/Engine/Private/Materials/MaterialShared.cpp
--- a/Engine/Source/Runtime/Engine/Private/Materials/MaterialShared.cpp
+++ b/Engine/Source/Runtime/Engine/Private/Materials/MaterialShared.cpp
@@ -1251,6 +1251,11 @@ bool FMaterialResource::IsUsingFullPrecision() const
 	return Material->bUseFullPrecision;
 }
 
+bool FMaterialResource::IsUsingPreintegratedGFForSimpleIBL() const
+{
+	return Material->bForwardRenderUsePreintegratedGFForSimpleIBL;
+}
+
 bool FMaterialResource::IsUsingHQForwardReflections() const
 {
 	return Material->bUseHQForwardReflections;
@@ -1670,6 +1675,7 @@ void FMaterial::SetupMaterialEnvironment(
 	OutEnvironment.SetDefine(TEXT("MATERIAL_TANGENTSPACENORMAL"), IsTangentSpaceNormal());
 	OutEnvironment.SetDefine(TEXT("GENERATE_SPHERICAL_PARTICLE_NORMALS"),ShouldGenerateSphericalParticleNormals());
 	OutEnvironment.SetDefine(TEXT("MATERIAL_USES_SCENE_COLOR_COPY"), RequiresSceneColorCopy_GameThread());
+	OutEnvironment.SetDefine(TEXT("MATERIAL_USE_PREINTEGRATED_GF"), IsUsingPreintegratedGFForSimpleIBL());
 	OutEnvironment.SetDefine(TEXT("MATERIAL_HQ_FORWARD_REFLECTIONS"), IsUsingHQForwardReflections());
 	OutEnvironment.SetDefine(TEXT("MATERIAL_PLANAR_FORWARD_REFLECTIONS"), IsUsingPlanarForwardReflections());
 	OutEnvironment.SetDefine(TEXT("MATERIAL_NONMETAL"), IsNonmetal());
diff --git a/Engine/Source/Runtime/Engine/Public/MaterialShared.h b/Engine/Source/Runtime/Engine/Public/MaterialShared.h
--- a/Engine/Source/Runtime/Engine/Public/MaterialShared.h
+++ b/Engine/Source/Runtime/Engine/Public/MaterialShared.h
@@ -1470,6 +1470,7 @@ public:
 	virtual bool IsFullyRough() const { return false; }
 	virtual bool UseNormalCurvatureToRoughness() const { return false; }
 	virtual bool IsUsingFullPrecision() const { return false; }
+	virtual bool IsUsingPreintegratedGFForSimpleIBL() const { return false; }
 	virtual bool IsUsingHQForwardReflections() const { return false; }
 	virtual bool IsUsingPlanarForwardReflections() const { return false; }
 	virtual bool OutputsVelocityOnBasePass() const { return true; }
@@ -2231,6 +2232,7 @@ public:
 	ENGINE_API virtual bool IsFullyRough() const override;
 	ENGINE_API virtual bool UseNormalCurvatureToRoughness() const override;
 	ENGINE_API virtual bool IsUsingFullPrecision() const override;
+	ENGINE_API virtual bool IsUsingPreintegratedGFForSimpleIBL() const override;
 	ENGINE_API virtual bool IsUsingHQForwardReflections() const override;
 	ENGINE_API virtual bool IsUsingPlanarForwardReflections() const override;
 	ENGINE_API virtual bool OutputsVelocityOnBasePass() const override;
diff --git a/Engine/Source/Runtime/MaterialShaderQualitySettings/Classes/ShaderPlatformQualitySettings.h b/Engine/Source/Runtime/MaterialShaderQualitySettings/Classes/ShaderPlatformQualitySettings.h
--- a/Engine/Source/Runtime/MaterialShaderQualitySettings/Classes/ShaderPlatformQualitySettings.h
+++ b/Engine/Source/Runtime/MaterialShaderQualitySettings/Classes/ShaderPlatformQualitySettings.h
@@ -36,6 +36,7 @@ public:
 		, bForceNonMetal(false)
 		, bForceDisableLMDirectionality(false)
 		, bForceLQReflections(false)
+		, bForceNotUsePreintegratedGF(false)
 		, bDisableMaterialNormalCalculation(false)
 		, MobileCSMQuality(EMobileCSMQuality::PCF_2x2)
 	{
@@ -59,6 +60,9 @@ public:
 	UPROPERTY(EditAnywhere, Config, Meta = (DisplayName = "Force low quality reflections"), Category = "Quality")
 	bool bForceLQReflections;
 
+	UPROPERTY(EditAnywhere, Config, Meta = (DisplayName = "Force not use preintegrated GF for simple IBL"), Category = "Quality")
+	bool bForceNotUsePreintegratedGF;
+
 	UPROPERTY(EditAnywhere, Config, Meta = (DisplayName = "Disable material normal calculation"), Category = "Quality")
 	bool bDisableMaterialNormalCalculation;
 
diff --git a/Engine/Source/Runtime/MaterialShaderQualitySettings/Private/MaterialShaderQualitySettings.cpp b/Engine/Source/Runtime/MaterialShaderQualitySettings/Private/MaterialShaderQualitySettings.cpp
--- a/Engine/Source/Runtime/MaterialShaderQualitySettings/Private/MaterialShaderQualitySettings.cpp
+++ b/Engine/Source/Runtime/MaterialShaderQualitySettings/Private/MaterialShaderQualitySettings.cpp
@@ -147,6 +147,7 @@ bool FMaterialQualityOverrides::HasAnyOverridesSet() const
 		|| bForceFullyRough != DefaultOverrides.bForceFullyRough
 		|| bForceNonMetal != DefaultOverrides.bForceNonMetal
 		|| bForceLQReflections != DefaultOverrides.bForceLQReflections
+		|| bForceNotUsePreintegratedGF != DefaultOverrides.bForceNotUsePreintegratedGF
 		|| bDisableMaterialNormalCalculation != DefaultOverrides.bDisableMaterialNormalCalculation
 		|| bDiscardQualityDuringCook != DefaultOverrides.bDiscardQualityDuringCook;
 }
diff --git a/Engine/Source/Runtime/Renderer/Private/MobileBasePassRendering.cpp b/Engine/Source/Runtime/Renderer/Private/MobileBasePassRendering.cpp
--- a/Engine/Source/Runtime/Renderer/Private/MobileBasePassRendering.cpp
+++ b/Engine/Source/Runtime/Renderer/Private/MobileBasePassRendering.cpp
@@ -99,6 +99,7 @@ bool TMobileBasePassPSPolicyParamType<LightMapPolicyType>::ModifyCompilationEnvi
 	OutEnvironment.SetDefine(TEXT("MOBILE_QL_FORCE_NONMETAL"), QualityOverrides.bEnableOverride && QualityOverrides.bForceNonMetal != 0 ? 1u : 0u);
 	OutEnvironment.SetDefine(TEXT("QL_FORCEDISABLE_LM_DIRECTIONALITY"), QualityOverrides.bEnableOverride && QualityOverrides.bForceDisableLMDirectionality != 0 ? 1u : 0u);
 	OutEnvironment.SetDefine(TEXT("MOBILE_QL_FORCE_LQ_REFLECTIONS"), QualityOverrides.bEnableOverride && QualityOverrides.bForceLQReflections != 0 ? 1u : 0u);
+	OutEnvironment.SetDefine(TEXT("MOBILE_QL_FORCE_DISABLE_PREINTEGRATEDGF"), QualityOverrides.bEnableOverride && QualityOverrides.bForceNotUsePreintegratedGF != 0 ? 1u : 0u);
 	OutEnvironment.SetDefine(TEXT("MOBILE_CSM_QUALITY"), (uint32)QualityOverrides.MobileCSMQuality);
 	OutEnvironment.SetDefine(TEXT("MOBILE_QL_DISABLE_MATERIAL_NORMAL"), QualityOverrides.bEnableOverride && QualityOverrides.bDisableMaterialNormalCalculation);
 	return true;
@@ -163,6 +164,9 @@ void SetupMobileBasePassUniformParameters(
 
 	FSceneRenderTargets& SceneContext = FSceneRenderTargets::Get(RHICmdList);
 	SetupMobileSceneTextureUniformParameters(SceneContext, View.FeatureLevel, bTranslucentPass, BasePassParameters.SceneTextures);
+
+	BasePassParameters.PreIntegratedGFTexture = GSystemTextures.PreintegratedGF->GetRenderTargetItem().ShaderResourceTexture;
+	BasePassParameters.PreIntegratedGFSampler = TStaticSamplerState<SF_Bilinear, AM_Clamp, AM_Clamp, AM_Clamp>::GetRHI();
 }
 
 void CreateMobileBasePassUniformBuffer(
diff --git a/Engine/Source/Runtime/Renderer/Private/MobileBasePassRendering.h b/Engine/Source/Runtime/Renderer/Private/MobileBasePassRendering.h
--- a/Engine/Source/Runtime/Renderer/Private/MobileBasePassRendering.h
+++ b/Engine/Source/Runtime/Renderer/Private/MobileBasePassRendering.h
@@ -28,6 +28,8 @@ BEGIN_GLOBAL_SHADER_PARAMETER_STRUCT(FMobileBasePassUniformParameters, )
 	SHADER_PARAMETER_STRUCT(FFogUniformParameters, Fog)
 	SHADER_PARAMETER_STRUCT(FPlanarReflectionUniformParameters, PlanarReflection) // Single global planar reflection for the forward pass.
 	SHADER_PARAMETER_STRUCT(FMobileSceneTextureUniformParameters, SceneTextures)
+	SHADER_PARAMETER_TEXTURE(Texture2D, PreIntegratedGFTexture)
+	SHADER_PARAMETER_SAMPLER(SamplerState, PreIntegratedGFSampler)
 END_GLOBAL_SHADER_PARAMETER_STRUCT()
 
 extern void SetupMobileBasePassUniformParameters(
diff --git a/Engine/Source/Runtime/Renderer/Private/SystemTextures.cpp b/Engine/Source/Runtime/Renderer/Private/SystemTextures.cpp
--- a/Engine/Source/Runtime/Renderer/Private/SystemTextures.cpp
+++ b/Engine/Source/Runtime/Renderer/Private/SystemTextures.cpp
@@ -24,98 +24,98 @@ void FSystemTextures::InitializeCommonTextures(FRHICommandListImmediate& RHICmdL
 	// requested feature level
 
 		// Create a WhiteDummy texture
-		{
-			FPooledRenderTargetDesc Desc(FPooledRenderTargetDesc::Create2DDesc(FIntPoint(1, 1), PF_B8G8R8A8, FClearValueBinding::White, TexCreate_HideInVisualizeTexture, TexCreate_RenderTargetable | TexCreate_NoFastClear | TexCreate_ShaderResource, false));
-			Desc.AutoWritable = false;
-			GRenderTargetPool.FindFreeElement(RHICmdList, Desc, WhiteDummy, TEXT("WhiteDummy"), true, ERenderTargetTransience::NonTransient);
+	{
+		FPooledRenderTargetDesc Desc(FPooledRenderTargetDesc::Create2DDesc(FIntPoint(1, 1), PF_B8G8R8A8, FClearValueBinding::White, TexCreate_HideInVisualizeTexture, TexCreate_RenderTargetable | TexCreate_NoFastClear | TexCreate_ShaderResource, false));
+		Desc.AutoWritable = false;
+		GRenderTargetPool.FindFreeElement(RHICmdList, Desc, WhiteDummy, TEXT("WhiteDummy"), true, ERenderTargetTransience::NonTransient);
 
-			FRHIRenderPassInfo RPInfo(WhiteDummy->GetRenderTargetItem().TargetableTexture, ERenderTargetActions::Clear_Store);
-			RHICmdList.BeginRenderPass(RPInfo, TEXT("WhiteDummy"));
-			RHICmdList.EndRenderPass();
-			RHICmdList.CopyToResolveTarget(WhiteDummy->GetRenderTargetItem().TargetableTexture, WhiteDummy->GetRenderTargetItem().ShaderResourceTexture, FResolveParams());
+		FRHIRenderPassInfo RPInfo(WhiteDummy->GetRenderTargetItem().TargetableTexture, ERenderTargetActions::Clear_Store);
+		RHICmdList.BeginRenderPass(RPInfo, TEXT("WhiteDummy"));
+		RHICmdList.EndRenderPass();
+		RHICmdList.CopyToResolveTarget(WhiteDummy->GetRenderTargetItem().TargetableTexture, WhiteDummy->GetRenderTargetItem().ShaderResourceTexture, FResolveParams());
 
-			WhiteDummySRV = RHICreateShaderResourceView((FRHITexture2D*)WhiteDummy->GetRenderTargetItem().ShaderResourceTexture.GetReference(), 0);
-		}
+		WhiteDummySRV = RHICreateShaderResourceView((FRHITexture2D*)WhiteDummy->GetRenderTargetItem().ShaderResourceTexture.GetReference(), 0);
+	}
 
-		// Create a BlackDummy texture
-		{
-			FPooledRenderTargetDesc Desc(FPooledRenderTargetDesc::Create2DDesc(FIntPoint(1, 1), PF_B8G8R8A8, FClearValueBinding::Transparent, TexCreate_HideInVisualizeTexture, TexCreate_RenderTargetable | TexCreate_NoFastClear | TexCreate_ShaderResource, false));
-			Desc.AutoWritable = false;
-			GRenderTargetPool.FindFreeElement(RHICmdList, Desc, BlackDummy, TEXT("BlackDummy"), true, ERenderTargetTransience::NonTransient);
+	// Create a BlackDummy texture
+	{
+		FPooledRenderTargetDesc Desc(FPooledRenderTargetDesc::Create2DDesc(FIntPoint(1, 1), PF_B8G8R8A8, FClearValueBinding::Transparent, TexCreate_HideInVisualizeTexture, TexCreate_RenderTargetable | TexCreate_NoFastClear | TexCreate_ShaderResource, false));
+		Desc.AutoWritable = false;
+		GRenderTargetPool.FindFreeElement(RHICmdList, Desc, BlackDummy, TEXT("BlackDummy"), true, ERenderTargetTransience::NonTransient);
 
-			FRHIRenderPassInfo RPInfo(BlackDummy->GetRenderTargetItem().TargetableTexture, ERenderTargetActions::Clear_Store);
-			RHICmdList.BeginRenderPass(RPInfo, TEXT("BlackDummy"));
-			RHICmdList.EndRenderPass();
-			RHICmdList.CopyToResolveTarget(BlackDummy->GetRenderTargetItem().TargetableTexture, BlackDummy->GetRenderTargetItem().ShaderResourceTexture, FResolveParams());
-		}
+		FRHIRenderPassInfo RPInfo(BlackDummy->GetRenderTargetItem().TargetableTexture, ERenderTargetActions::Clear_Store);
+		RHICmdList.BeginRenderPass(RPInfo, TEXT("BlackDummy"));
+		RHICmdList.EndRenderPass();
+		RHICmdList.CopyToResolveTarget(BlackDummy->GetRenderTargetItem().TargetableTexture, BlackDummy->GetRenderTargetItem().ShaderResourceTexture, FResolveParams());
+	}
 
-		// Create a BlackAlphaOneDummy texture
-		{
-			FPooledRenderTargetDesc Desc(FPooledRenderTargetDesc::Create2DDesc(FIntPoint(1, 1), PF_B8G8R8A8, FClearValueBinding::Black, TexCreate_HideInVisualizeTexture, TexCreate_RenderTargetable | TexCreate_NoFastClear | TexCreate_ShaderResource, false));
-			Desc.AutoWritable = false;
-			GRenderTargetPool.FindFreeElement(RHICmdList, Desc, BlackAlphaOneDummy, TEXT("BlackAlphaOneDummy"), true, ERenderTargetTransience::NonTransient);
+	// Create a BlackAlphaOneDummy texture
+	{
+		FPooledRenderTargetDesc Desc(FPooledRenderTargetDesc::Create2DDesc(FIntPoint(1, 1), PF_B8G8R8A8, FClearValueBinding::Black, TexCreate_HideInVisualizeTexture, TexCreate_RenderTargetable | TexCreate_NoFastClear | TexCreate_ShaderResource, false));
+		Desc.AutoWritable = false;
+		GRenderTargetPool.FindFreeElement(RHICmdList, Desc, BlackAlphaOneDummy, TEXT("BlackAlphaOneDummy"), true, ERenderTargetTransience::NonTransient);
 
-			FRHIRenderPassInfo RPInfo(BlackAlphaOneDummy->GetRenderTargetItem().TargetableTexture, ERenderTargetActions::Clear_Store);
-			RHICmdList.BeginRenderPass(RPInfo, TEXT("BlackAlphaOneDummy"));
-			RHICmdList.EndRenderPass();
-			RHICmdList.CopyToResolveTarget(BlackAlphaOneDummy->GetRenderTargetItem().TargetableTexture, BlackAlphaOneDummy->GetRenderTargetItem().ShaderResourceTexture, FResolveParams());
-		}
+		FRHIRenderPassInfo RPInfo(BlackAlphaOneDummy->GetRenderTargetItem().TargetableTexture, ERenderTargetActions::Clear_Store);
+		RHICmdList.BeginRenderPass(RPInfo, TEXT("BlackAlphaOneDummy"));
+		RHICmdList.EndRenderPass();
+		RHICmdList.CopyToResolveTarget(BlackAlphaOneDummy->GetRenderTargetItem().TargetableTexture, BlackAlphaOneDummy->GetRenderTargetItem().ShaderResourceTexture, FResolveParams());
+	}
 
-		// Create a GreenDummy texture
-		{
-			FPooledRenderTargetDesc Desc(FPooledRenderTargetDesc::Create2DDesc(FIntPoint(1, 1), PF_B8G8R8A8, FClearValueBinding::Green, TexCreate_HideInVisualizeTexture, TexCreate_RenderTargetable | TexCreate_NoFastClear | TexCreate_ShaderResource, false));
-			Desc.AutoWritable = false;
-			GRenderTargetPool.FindFreeElement(RHICmdList, Desc, GreenDummy, TEXT("GreenDummy"), true, ERenderTargetTransience::NonTransient);
+	// Create a GreenDummy texture
+	{
+		FPooledRenderTargetDesc Desc(FPooledRenderTargetDesc::Create2DDesc(FIntPoint(1, 1), PF_B8G8R8A8, FClearValueBinding::Green, TexCreate_HideInVisualizeTexture, TexCreate_RenderTargetable | TexCreate_NoFastClear | TexCreate_ShaderResource, false));
+		Desc.AutoWritable = false;
+		GRenderTargetPool.FindFreeElement(RHICmdList, Desc, GreenDummy, TEXT("GreenDummy"), true, ERenderTargetTransience::NonTransient);
 
-			FRHIRenderPassInfo RPInfo(GreenDummy->GetRenderTargetItem().TargetableTexture, ERenderTargetActions::Clear_Store);
-			RHICmdList.BeginRenderPass(RPInfo, TEXT("GreenDummy"));
-			RHICmdList.EndRenderPass();
-			RHICmdList.CopyToResolveTarget(GreenDummy->GetRenderTargetItem().TargetableTexture, GreenDummy->GetRenderTargetItem().ShaderResourceTexture, FResolveParams());
-		}
+		FRHIRenderPassInfo RPInfo(GreenDummy->GetRenderTargetItem().TargetableTexture, ERenderTargetActions::Clear_Store);
+		RHICmdList.BeginRenderPass(RPInfo, TEXT("GreenDummy"));
+		RHICmdList.EndRenderPass();
+		RHICmdList.CopyToResolveTarget(GreenDummy->GetRenderTargetItem().TargetableTexture, GreenDummy->GetRenderTargetItem().ShaderResourceTexture, FResolveParams());
+	}
 
-		// Create a DefaultNormal8Bit texture
-		{
-			FPooledRenderTargetDesc Desc(FPooledRenderTargetDesc::Create2DDesc(FIntPoint(1, 1), PF_B8G8R8A8, FClearValueBinding::DefaultNormal8Bit, TexCreate_HideInVisualizeTexture, TexCreate_RenderTargetable | TexCreate_NoFastClear | TexCreate_ShaderResource, false));
-			Desc.AutoWritable = false;
-			GRenderTargetPool.FindFreeElement(RHICmdList, Desc, DefaultNormal8Bit, TEXT("DefaultNormal8Bit"), true, ERenderTargetTransience::NonTransient);
+	// Create a DefaultNormal8Bit texture
+	{
+		FPooledRenderTargetDesc Desc(FPooledRenderTargetDesc::Create2DDesc(FIntPoint(1, 1), PF_B8G8R8A8, FClearValueBinding::DefaultNormal8Bit, TexCreate_HideInVisualizeTexture, TexCreate_RenderTargetable | TexCreate_NoFastClear | TexCreate_ShaderResource, false));
+		Desc.AutoWritable = false;
+		GRenderTargetPool.FindFreeElement(RHICmdList, Desc, DefaultNormal8Bit, TEXT("DefaultNormal8Bit"), true, ERenderTargetTransience::NonTransient);
 
-			FRHIRenderPassInfo RPInfo(DefaultNormal8Bit->GetRenderTargetItem().TargetableTexture, ERenderTargetActions::Clear_Store);
-			RHICmdList.BeginRenderPass(RPInfo, TEXT("DefaultNormal8Bit"));
-			RHICmdList.EndRenderPass();
-			RHICmdList.CopyToResolveTarget(DefaultNormal8Bit->GetRenderTargetItem().TargetableTexture, DefaultNormal8Bit->GetRenderTargetItem().ShaderResourceTexture, FResolveParams());
-		}
+		FRHIRenderPassInfo RPInfo(DefaultNormal8Bit->GetRenderTargetItem().TargetableTexture, ERenderTargetActions::Clear_Store);
+		RHICmdList.BeginRenderPass(RPInfo, TEXT("DefaultNormal8Bit"));
+		RHICmdList.EndRenderPass();
+		RHICmdList.CopyToResolveTarget(DefaultNormal8Bit->GetRenderTargetItem().TargetableTexture, DefaultNormal8Bit->GetRenderTargetItem().ShaderResourceTexture, FResolveParams());
+	}
 
-		// Create the PerlinNoiseGradient texture
+	// Create the PerlinNoiseGradient texture
+	{
+		FPooledRenderTargetDesc Desc(FPooledRenderTargetDesc::Create2DDesc(FIntPoint(128, 128), PF_B8G8R8A8, FClearValueBinding::None, TexCreate_HideInVisualizeTexture, TexCreate_None | TexCreate_NoFastClear | TexCreate_ShaderResource, false));
+		Desc.AutoWritable = false;
+		GRenderTargetPool.FindFreeElement(RHICmdList, Desc, PerlinNoiseGradient, TEXT("PerlinNoiseGradient"), true, ERenderTargetTransience::NonTransient);
+		// Write the contents of the texture.
+		uint32 DestStride;
+		uint8* DestBuffer = (uint8*)RHICmdList.LockTexture2D((FTexture2DRHIRef&)PerlinNoiseGradient->GetRenderTargetItem().ShaderResourceTexture, 0, RLM_WriteOnly, DestStride, false);
+		// seed the pseudo random stream with a good value
+		FRandomStream RandomStream(12345);
+		// Values represent float3 values in the -1..1 range.
+		// The vectors are the edge mid point of a cube from -1 .. 1
+		static uint32 gradtable[] =
 		{
-			FPooledRenderTargetDesc Desc(FPooledRenderTargetDesc::Create2DDesc(FIntPoint(128, 128), PF_B8G8R8A8, FClearValueBinding::None, TexCreate_HideInVisualizeTexture, TexCreate_None | TexCreate_NoFastClear | TexCreate_ShaderResource, false));
-			Desc.AutoWritable = false;
-			GRenderTargetPool.FindFreeElement(RHICmdList, Desc, PerlinNoiseGradient, TEXT("PerlinNoiseGradient"), true, ERenderTargetTransience::NonTransient);
-			// Write the contents of the texture.
-			uint32 DestStride;
-			uint8* DestBuffer = (uint8*)RHICmdList.LockTexture2D((FTexture2DRHIRef&)PerlinNoiseGradient->GetRenderTargetItem().ShaderResourceTexture, 0, RLM_WriteOnly, DestStride, false);
-			// seed the pseudo random stream with a good value
-			FRandomStream RandomStream(12345);
-			// Values represent float3 values in the -1..1 range.
-			// The vectors are the edge mid point of a cube from -1 .. 1
-			static uint32 gradtable[] =
-			{
-				0x88ffff, 0xff88ff, 0xffff88,
-				0x88ff00, 0xff8800, 0xff0088,
-				0x8800ff, 0x0088ff, 0x00ff88,
-				0x880000, 0x008800, 0x000088,
-			};
-			for (int32 y = 0; y < Desc.Extent.Y; ++y)
+			0x88ffff, 0xff88ff, 0xffff88,
+			0x88ff00, 0xff8800, 0xff0088,
+			0x8800ff, 0x0088ff, 0x00ff88,
+			0x880000, 0x008800, 0x000088,
+		};
+		for (int32 y = 0; y < Desc.Extent.Y; ++y)
+		{
+			for (int32 x = 0; x < Desc.Extent.X; ++x)
 			{
-				for (int32 x = 0; x < Desc.Extent.X; ++x)
-				{
 				uint32* Dest = (uint32*)(DestBuffer + x * sizeof(uint32) + y * DestStride);
 
-					// pick a random direction (hacky way to overcome the quality issues FRandomStream has)
-					*Dest = gradtable[(uint32)(RandomStream.GetFraction() * 11.9999999f)];
-				}
+				// pick a random direction (hacky way to overcome the quality issues FRandomStream has)
+				*Dest = gradtable[(uint32)(RandomStream.GetFraction() * 11.9999999f)];
 			}
-			RHICmdList.UnlockTexture2D((FTexture2DRHIRef&)PerlinNoiseGradient->GetRenderTargetItem().ShaderResourceTexture, 0, false);
 		}
+		RHICmdList.UnlockTexture2D((FTexture2DRHIRef&)PerlinNoiseGradient->GetRenderTargetItem().ShaderResourceTexture, 0, false);
+	}
 
 	if (!GSupportsShaderFramebufferFetch && GPixelFormats[PF_FloatRGBA].Supported)
 	{
@@ -165,35 +165,35 @@ void FSystemTextures::InitializeFeatureLevelDependentTextures(FRHICommandListImm
 	// this should change if we add any ERHIFeatureLevel::ES2 dependent textures here!
 	ERHIFeatureLevel::Type CurrentFeatureLevel = FeatureLevelInitializedTo == ERHIFeatureLevel::Num ? ERHIFeatureLevel::ES2 : FeatureLevelInitializedTo;
 
-		// Create the SobolSampling texture
+	// Create the SobolSampling texture
 	if (CurrentFeatureLevel < ERHIFeatureLevel::ES3_1 && InFeatureLevel >= ERHIFeatureLevel::ES3_1 && GPixelFormats[PF_R16_UINT].Supported)
+	{
+		FPooledRenderTargetDesc Desc(FPooledRenderTargetDesc::Create2DDesc(FIntPoint(32, 16), PF_R16_UINT, FClearValueBinding::None, TexCreate_HideInVisualizeTexture, TexCreate_NoFastClear | TexCreate_ShaderResource, false));
+		Desc.AutoWritable = false;
+		GRenderTargetPool.FindFreeElement(RHICmdList, Desc, SobolSampling, TEXT("SobolSampling"));
+		// Write the contents of the texture.
+		uint32 DestStride;
+		uint8* DestBuffer = (uint8*)RHICmdList.LockTexture2D((FTexture2DRHIRef&)SobolSampling->GetRenderTargetItem().ShaderResourceTexture, 0, RLM_WriteOnly, DestStride, false);
+
+		uint16 *Dest;
+		for (int y = 0; y < 16; ++y)
 		{
-			FPooledRenderTargetDesc Desc(FPooledRenderTargetDesc::Create2DDesc(FIntPoint(32, 16), PF_R16_UINT, FClearValueBinding::None, TexCreate_HideInVisualizeTexture, TexCreate_NoFastClear | TexCreate_ShaderResource, false));
-			Desc.AutoWritable = false;
-			GRenderTargetPool.FindFreeElement(RHICmdList, Desc, SobolSampling, TEXT("SobolSampling"));
-			// Write the contents of the texture.
-			uint32 DestStride;
-			uint8* DestBuffer = (uint8*)RHICmdList.LockTexture2D((FTexture2DRHIRef&)SobolSampling->GetRenderTargetItem().ShaderResourceTexture, 0, RLM_WriteOnly, DestStride, false);
+			Dest = (uint16*)(DestBuffer + y * DestStride);
 
-			uint16 *Dest;
-			for (int y = 0; y < 16; ++y)
+			// 16x16 block starting at 0,0 = Sobol X,Y from bottom 4 bits of cell X,Y
+			for (int x = 0; x < 16; ++x, ++Dest)
 			{
-				Dest = (uint16*)(DestBuffer + y * DestStride);
-
-				// 16x16 block starting at 0,0 = Sobol X,Y from bottom 4 bits of cell X,Y
-				for (int x = 0; x < 16; ++x, ++Dest)
-				{
-					*Dest = FSobol::ComputeGPUSpatialSeed(x, y, /* Index = */ 0);
-				}
+				*Dest = FSobol::ComputeGPUSpatialSeed(x, y, /* Index = */ 0);
+			}
 
-				// 16x16 block starting at 16,0 = Sobol X,Y from 2nd 4 bits of cell X,Y
-				for (int x = 0; x < 16; ++x, ++Dest)
-				{
-					*Dest = FSobol::ComputeGPUSpatialSeed(x, y, /* Index = */ 1);
-				}
+			// 16x16 block starting at 16,0 = Sobol X,Y from 2nd 4 bits of cell X,Y
+			for (int x = 0; x < 16; ++x, ++Dest)
+			{
+				*Dest = FSobol::ComputeGPUSpatialSeed(x, y, /* Index = */ 1);
 			}
-			RHICmdList.UnlockTexture2D((FTexture2DRHIRef&)SobolSampling->GetRenderTargetItem().ShaderResourceTexture, 0, false);
 		}
+		RHICmdList.UnlockTexture2D((FTexture2DRHIRef&)SobolSampling->GetRenderTargetItem().ShaderResourceTexture, 0, false);
+	}
 
 	// Create a VolumetricBlackDummy texture
 	if (CurrentFeatureLevel < ERHIFeatureLevel::SM5 && InFeatureLevel >= ERHIFeatureLevel::SM5)
@@ -213,191 +213,320 @@ void FSystemTextures::InitializeFeatureLevelDependentTextures(FRHICommandListImm
 			BlackBytes);
 	}
 
-	if (CurrentFeatureLevel < ERHIFeatureLevel::SM5 && InFeatureLevel >= ERHIFeatureLevel::SM5)
-		{
-	// Create the PerlinNoise3D texture (similar to http://prettyprocs.wordpress.com/2012/10/20/fast-perlin-noise/)
+	// The PreintegratedGF maybe used on forward shading inluding mobile platorm, intialize it anyway.
 	{
-		uint32 Extent = 16;
+		// for testing, with 128x128 R8G8 we are very close to the reference (if lower res is needed we might have to add an offset to counter the 0.5f texel shift)
+		const bool bReference = false;
 
-		const uint32 Square = Extent * Extent;
+		EPixelFormat Format = PF_R8G8;
+		// for low roughness we would get banding with PF_R8G8 but for low spec it could be used, for now we don't do this optimization
+		if (GPixelFormats[PF_G16R16].Supported)
+		{
+			Format = PF_G16R16;
+		}
 
-		FPooledRenderTargetDesc Desc(FPooledRenderTargetDesc::CreateVolumeDesc(Extent, Extent, Extent, PF_B8G8R8A8, FClearValueBinding::None, TexCreate_ShaderResource | TexCreate_HideInVisualizeTexture | TexCreate_NoTiling, TexCreate_ShaderResource, false));
+		FPooledRenderTargetDesc Desc(FPooledRenderTargetDesc::Create2DDesc(FIntPoint(128, 32), Format, FClearValueBinding::None, 0, TexCreate_ShaderResource, false));
 		Desc.AutoWritable = false;
-		GRenderTargetPool.FindFreeElement(RHICmdList, Desc, PerlinNoise3D, TEXT("PerlinNoise3D"), true, ERenderTargetTransience::NonTransient);
+		if (bReference)
+		{
+			Desc.Extent.X = 128;
+			Desc.Extent.Y = 128;
+		}
+
+		GRenderTargetPool.FindFreeElement(RHICmdList, Desc, PreintegratedGF, TEXT("PreintegratedGF"), true, ERenderTargetTransience::NonTransient);
 		// Write the contents of the texture.
-		TArray<uint32> DestBuffer;
+		uint32 DestStride;
+		uint8* DestBuffer = (uint8*)RHICmdList.LockTexture2D((FTexture2DRHIRef&)PreintegratedGF->GetRenderTargetItem().ShaderResourceTexture, 0, RLM_WriteOnly, DestStride, false);
 
-		DestBuffer.AddZeroed(Extent * Extent * Extent);
-		// seed the pseudo random stream with a good value
-		FRandomStream RandomStream(0x1234);
-		// Values represent float3 values in the -1..1 range.
-		// The vectors are the edge mid point of a cube from -1 .. 1
-		// -1:0 0:7f 1:fe, can be reconstructed with * 512/254 - 1
-		// * 2 - 1 cannot be used because 0 would not be mapped
-			static uint32 gradtable[] =
+		// x is NoV, y is roughness
+		for (int32 y = 0; y < Desc.Extent.Y; y++)
 		{
-			0x7ffefe, 0xfe7ffe, 0xfefe7f,
-			0x7ffe00, 0xfe7f00, 0xfe007f,
-			0x7f00fe, 0x007ffe, 0x00fe7f,
-			0x7f0000, 0x007f00, 0x00007f,
-		};
-		// set random directions
-		{
-				for (uint32 z = 0; z < Extent - 1; ++z)
+			float Roughness = (float)(y + 0.5f) / Desc.Extent.Y;
+			float m = Roughness * Roughness;
+			float m2 = m * m;
+
+			for (int32 x = 0; x < Desc.Extent.X; x++)
 			{
-					for (uint32 y = 0; y < Extent - 1; ++y)
+				float NoV = (float)(x + 0.5f) / Desc.Extent.X;
+
+				FVector V;
+				V.X = FMath::Sqrt(1.0f - NoV * NoV);	// sin
+				V.Y = 0.0f;
+				V.Z = NoV;								// cos
+
+				float A = 0.0f;
+				float B = 0.0f;
+				float C = 0.0f;
+
+				const uint32 NumSamples = 128;
+				for (uint32 i = 0; i < NumSamples; i++)
 				{
-						for (uint32 x = 0; x < Extent - 1; ++x)
+					float E1 = (float)i / NumSamples;
+					float E2 = (double)ReverseBits(i) / (double)0x100000000LL;
+
 					{
-						uint32& Value = DestBuffer[x + y * Extent + z * Square];
+						float Phi = 2.0f * PI * E1;
+						float CosPhi = FMath::Cos(Phi);
+						float SinPhi = FMath::Sin(Phi);
+						float CosTheta = FMath::Sqrt((1.0f - E2) / (1.0f + (m2 - 1.0f) * E2));
+						float SinTheta = FMath::Sqrt(1.0f - CosTheta * CosTheta);
+
+						FVector H(SinTheta * FMath::Cos(Phi), SinTheta * FMath::Sin(Phi), CosTheta);
+						FVector L = 2.0f * (V | H) * H - V;
+
+						float NoL = FMath::Max(L.Z, 0.0f);
+						float NoH = FMath::Max(H.Z, 0.0f);
+						float VoH = FMath::Max(V | H, 0.0f);
+
+						if (NoL > 0.0f)
+						{
+							float Vis_SmithV = NoL * (NoV * (1 - m) + m);
+							float Vis_SmithL = NoV * (NoL * (1 - m) + m);
+							float Vis = 0.5f / (Vis_SmithV + Vis_SmithL);
+
+							float NoL_Vis_PDF = NoL * Vis * (4.0f * VoH / NoH);
+							float Fc = 1.0f - VoH;
+							Fc *= FMath::Square(Fc*Fc);
+							A += NoL_Vis_PDF * (1.0f - Fc);
+							B += NoL_Vis_PDF * Fc;
+						}
+					}
 
-						// pick a random direction (hacky way to overcome the quality issues FRandomStream has)
-						Value = gradtable[(uint32)(RandomStream.GetFraction() * 11.9999999f)];
+					{
+						float Phi = 2.0f * PI * E1;
+						float CosPhi = FMath::Cos(Phi);
+						float SinPhi = FMath::Sin(Phi);
+						float CosTheta = FMath::Sqrt(E2);
+						float SinTheta = FMath::Sqrt(1.0f - CosTheta * CosTheta);
+
+						FVector L(SinTheta * FMath::Cos(Phi), SinTheta * FMath::Sin(Phi), CosTheta);
+						FVector H = (V + L).GetUnsafeNormal();
+
+						float NoL = FMath::Max(L.Z, 0.0f);
+						float NoH = FMath::Max(H.Z, 0.0f);
+						float VoH = FMath::Max(V | H, 0.0f);
+
+						float FD90 = 0.5f + 2.0f * VoH * VoH * Roughness;
+						float FdV = 1.0f + (FD90 - 1.0f) * pow(1.0f - NoV, 5);
+						float FdL = 1.0f + (FD90 - 1.0f) * pow(1.0f - NoL, 5);
+						C += FdV * FdL;// * ( 1.0f - 0.3333f * Roughness );
 					}
 				}
+				A /= NumSamples;
+				B /= NumSamples;
+				C /= NumSamples;
+
+				if (Desc.Format == PF_A16B16G16R16)
+				{
+					uint16* Dest = (uint16*)(DestBuffer + x * 8 + y * DestStride);
+					Dest[0] = (int32)(FMath::Clamp(A, 0.0f, 1.0f) * 65535.0f + 0.5f);
+					Dest[1] = (int32)(FMath::Clamp(B, 0.0f, 1.0f) * 65535.0f + 0.5f);
+					Dest[2] = (int32)(FMath::Clamp(C, 0.0f, 1.0f) * 65535.0f + 0.5f);
+				}
+				else if (Desc.Format == PF_G16R16)
+				{
+					uint16* Dest = (uint16*)(DestBuffer + x * 4 + y * DestStride);
+					Dest[0] = (int32)(FMath::Clamp(A, 0.0f, 1.0f) * 65535.0f + 0.5f);
+					Dest[1] = (int32)(FMath::Clamp(B, 0.0f, 1.0f) * 65535.0f + 0.5f);
+				}
+				else
+				{
+					check(Desc.Format == PF_R8G8);
+
+					uint8* Dest = (uint8*)(DestBuffer + x * 2 + y * DestStride);
+					Dest[0] = (int32)(FMath::Clamp(A, 0.0f, 1.0f) * 255.9999f);
+					Dest[1] = (int32)(FMath::Clamp(B, 0.0f, 1.0f) * 255.9999f);
+				}
 			}
 		}
-		// replicate a border for filtering
+		RHICmdList.UnlockTexture2D((FTexture2DRHIRef&)PreintegratedGF->GetRenderTargetItem().ShaderResourceTexture, 0, false);
+	}
+
+	if (CurrentFeatureLevel < ERHIFeatureLevel::SM5 && InFeatureLevel >= ERHIFeatureLevel::SM5)
+	{
+		// Create the PerlinNoise3D texture (similar to http://prettyprocs.wordpress.com/2012/10/20/fast-perlin-noise/)
 		{
-			uint32 Last = Extent - 1;
+			uint32 Extent = 16;
 
-				for (uint32 z = 0; z < Extent; ++z)
+			const uint32 Square = Extent * Extent;
+
+			FPooledRenderTargetDesc Desc(FPooledRenderTargetDesc::CreateVolumeDesc(Extent, Extent, Extent, PF_B8G8R8A8, FClearValueBinding::None, TexCreate_ShaderResource | TexCreate_HideInVisualizeTexture | TexCreate_NoTiling, TexCreate_ShaderResource, false));
+			Desc.AutoWritable = false;
+			GRenderTargetPool.FindFreeElement(RHICmdList, Desc, PerlinNoise3D, TEXT("PerlinNoise3D"), true, ERenderTargetTransience::NonTransient);
+			// Write the contents of the texture.
+			TArray<uint32> DestBuffer;
+
+			DestBuffer.AddZeroed(Extent * Extent * Extent);
+			// seed the pseudo random stream with a good value
+			FRandomStream RandomStream(0x1234);
+			// Values represent float3 values in the -1..1 range.
+			// The vectors are the edge mid point of a cube from -1 .. 1
+			// -1:0 0:7f 1:fe, can be reconstructed with * 512/254 - 1
+			// * 2 - 1 cannot be used because 0 would not be mapped
+			static uint32 gradtable[] =
 			{
-					for (uint32 y = 0; y < Extent; ++y)
+				0x7ffefe, 0xfe7ffe, 0xfefe7f,
+				0x7ffe00, 0xfe7f00, 0xfe007f,
+				0x7f00fe, 0x007ffe, 0x00fe7f,
+				0x7f0000, 0x007f00, 0x00007f,
+			};
+			// set random directions
+			{
+				for (uint32 z = 0; z < Extent - 1; ++z)
 				{
-					DestBuffer[Last + y * Extent + z * Square] = DestBuffer[0 + y * Extent + z * Square];
+					for (uint32 y = 0; y < Extent - 1; ++y)
+					{
+						for (uint32 x = 0; x < Extent - 1; ++x)
+						{
+							uint32& Value = DestBuffer[x + y * Extent + z * Square];
+
+							// pick a random direction (hacky way to overcome the quality issues FRandomStream has)
+							Value = gradtable[(uint32)(RandomStream.GetFraction() * 11.9999999f)];
+						}
+					}
 				}
 			}
-				for (uint32 z = 0; z < Extent; ++z)
+			// replicate a border for filtering
 			{
-					for (uint32 x = 0; x < Extent; ++x)
+				uint32 Last = Extent - 1;
+
+				for (uint32 z = 0; z < Extent; ++z)
 				{
-					DestBuffer[x + Last * Extent + z * Square] = DestBuffer[x + 0 * Extent + z * Square];
+					for (uint32 y = 0; y < Extent; ++y)
+					{
+						DestBuffer[Last + y * Extent + z * Square] = DestBuffer[0 + y * Extent + z * Square];
+					}
 				}
-			}
-				for (uint32 y = 0; y < Extent; ++y)
-			{
+				for (uint32 z = 0; z < Extent; ++z)
+				{
 					for (uint32 x = 0; x < Extent; ++x)
+					{
+						DestBuffer[x + Last * Extent + z * Square] = DestBuffer[x + 0 * Extent + z * Square];
+					}
+				}
+				for (uint32 y = 0; y < Extent; ++y)
 				{
-					DestBuffer[x + y * Extent + Last * Square] = DestBuffer[x + y * Extent + 0 * Square];
+					for (uint32 x = 0; x < Extent; ++x)
+					{
+						DestBuffer[x + y * Extent + Last * Square] = DestBuffer[x + y * Extent + 0 * Square];
+					}
 				}
 			}
-		}
-		// precompute gradients
-		{
-			uint32* Dest = DestBuffer.GetData();
+			// precompute gradients
+			{
+				uint32* Dest = DestBuffer.GetData();
 
 				for (uint32 z = 0; z < Desc.Depth; ++z)
-			{
-					for (uint32 y = 0; y < (uint32)Desc.Extent.Y; ++y)
 				{
-						for (uint32 x = 0; x < (uint32)Desc.Extent.X; ++x)
+					for (uint32 y = 0; y < (uint32)Desc.Extent.Y; ++y)
 					{
-						uint32 Value = *Dest;
+						for (uint32 x = 0; x < (uint32)Desc.Extent.X; ++x)
+						{
+							uint32 Value = *Dest;
 
-						// todo: check if rgb order is correct
-						int32 r = Value >> 16;
-						int32 g = (Value >> 8) & 0xff;
-						int32 b = Value & 0xff;
+							// todo: check if rgb order is correct
+							int32 r = Value >> 16;
+							int32 g = (Value >> 8) & 0xff;
+							int32 b = Value & 0xff;
 
-						int nx = (r / 0x7f) - 1;
-						int ny = (g / 0x7f) - 1;
-						int nz = (b / 0x7f) - 1;
+							int nx = (r / 0x7f) - 1;
+							int ny = (g / 0x7f) - 1;
+							int nz = (b / 0x7f) - 1;
 
-						int32 d = nx * x + ny * y + nz * z;
+							int32 d = nx * x + ny * y + nz * z;
 
-						// compress in 8bit
-						uint32 a = d + 127;
+							// compress in 8bit
+							uint32 a = d + 127;
 
-						*Dest++ = Value | (a << 24);
+							*Dest++ = Value | (a << 24);
+						}
 					}
 				}
 			}
-		}
 
-		FUpdateTextureRegion3D Region(0, 0, 0, 0, 0, 0, Desc.Extent.X, Desc.Extent.Y, Desc.Depth);
+			FUpdateTextureRegion3D Region(0, 0, 0, 0, 0, 0, Desc.Extent.X, Desc.Extent.Y, Desc.Depth);
 
-		RHICmdList.UpdateTexture3D(
-			(FTexture3DRHIRef&)PerlinNoise3D->GetRenderTargetItem().ShaderResourceTexture,
-			0,
-			Region,
-			Desc.Extent.X * sizeof(uint32),
-			Desc.Extent.X * Desc.Extent.Y * sizeof(uint32),
-			(const uint8*)DestBuffer.GetData());
+			RHICmdList.UpdateTexture3D(
+				(FTexture3DRHIRef&)PerlinNoise3D->GetRenderTargetItem().ShaderResourceTexture,
+				0,
+				Region,
+				Desc.Extent.X * sizeof(uint32),
+				Desc.Extent.X * Desc.Extent.Y * sizeof(uint32),
+				(const uint8*)DestBuffer.GetData());
 		} // end Create the PerlinNoise3D texture
 
 	// Create the SSAO randomization texture
 		{
-	{
-		float g_AngleOff1 = 127;
-		float g_AngleOff2 = 198;
-		float g_AngleOff3 = 23;
+			{
+				float g_AngleOff1 = 127;
+				float g_AngleOff2 = 198;
+				float g_AngleOff3 = 23;
 
-		FColor Bases[16];
+				FColor Bases[16];
 
 				for (int32 Pos = 0; Pos < 16; ++Pos)
-		{
-			// distribute rotations over 4x4 pattern
-					//			int32 Reorder[16] = { 0, 8, 2, 10, 12, 6, 14, 4, 3, 11, 1, 9, 15, 5, 13, 7 };
-			int32 Reorder[16] = { 0, 11, 7, 3, 10, 4, 15, 12, 6, 8, 1, 14, 13, 2, 9, 5 };
-			int32 w = Reorder[Pos];
-
-			// ordered sampling of the rotation basis (*2 is missing as we use mirrored samples)
-			float ww = w / 16.0f * PI;
+				{
+					// distribute rotations over 4x4 pattern
+							//			int32 Reorder[16] = { 0, 8, 2, 10, 12, 6, 14, 4, 3, 11, 1, 9, 15, 5, 13, 7 };
+					int32 Reorder[16] = { 0, 11, 7, 3, 10, 4, 15, 12, 6, 8, 1, 14, 13, 2, 9, 5 };
+					int32 w = Reorder[Pos];
 
-			// randomize base scale
-			float lenm = 1.0f - (FMath::Sin(g_AngleOff2 * w * 0.01f) * 0.5f + 0.5f) * g_AngleOff3 * 0.01f;
-			float s = FMath::Sin(ww) * lenm;
-			float c = FMath::Cos(ww) * lenm;
+					// ordered sampling of the rotation basis (*2 is missing as we use mirrored samples)
+					float ww = w / 16.0f * PI;
 
-			Bases[Pos] = FColor(FMath::Quantize8SignedByte(c), FMath::Quantize8SignedByte(s), 0, 0);
-		}
+					// randomize base scale
+					float lenm = 1.0f - (FMath::Sin(g_AngleOff2 * w * 0.01f) * 0.5f + 0.5f) * g_AngleOff3 * 0.01f;
+					float s = FMath::Sin(ww) * lenm;
+					float c = FMath::Cos(ww) * lenm;
 
-		{
-			// could be PF_V8U8 to save shader instructions but that doesn't work on all hardware
-			FPooledRenderTargetDesc Desc(FPooledRenderTargetDesc::Create2DDesc(FIntPoint(64, 64), PF_R8G8, FClearValueBinding::None, TexCreate_HideInVisualizeTexture, TexCreate_NoFastClear | TexCreate_ShaderResource, false));
-			Desc.AutoWritable = false;
-			GRenderTargetPool.FindFreeElement(RHICmdList, Desc, SSAORandomization, TEXT("SSAORandomization"), true, ERenderTargetTransience::NonTransient);
-			// Write the contents of the texture.
-			uint32 DestStride;
-			uint8* DestBuffer = (uint8*)RHICmdList.LockTexture2D((FTexture2DRHIRef&)SSAORandomization->GetRenderTargetItem().ShaderResourceTexture, 0, RLM_WriteOnly, DestStride, false);
+					Bases[Pos] = FColor(FMath::Quantize8SignedByte(c), FMath::Quantize8SignedByte(s), 0, 0);
+				}
 
-			for(int32 y = 0; y < Desc.Extent.Y; ++y)
-			{
-				for(int32 x = 0; x < Desc.Extent.X; ++x)
 				{
-					uint8* Dest = (uint8*)(DestBuffer + x * sizeof(uint16) + y * DestStride);
+					// could be PF_V8U8 to save shader instructions but that doesn't work on all hardware
+					FPooledRenderTargetDesc Desc(FPooledRenderTargetDesc::Create2DDesc(FIntPoint(64, 64), PF_R8G8, FClearValueBinding::None, TexCreate_HideInVisualizeTexture, TexCreate_NoFastClear | TexCreate_ShaderResource, false));
+					Desc.AutoWritable = false;
+					GRenderTargetPool.FindFreeElement(RHICmdList, Desc, SSAORandomization, TEXT("SSAORandomization"), true, ERenderTargetTransience::NonTransient);
+					// Write the contents of the texture.
+					uint32 DestStride;
+					uint8* DestBuffer = (uint8*)RHICmdList.LockTexture2D((FTexture2DRHIRef&)SSAORandomization->GetRenderTargetItem().ShaderResourceTexture, 0, RLM_WriteOnly, DestStride, false);
+
+					for (int32 y = 0; y < Desc.Extent.Y; ++y)
+					{
+						for (int32 x = 0; x < Desc.Extent.X; ++x)
+						{
+							uint8* Dest = (uint8*)(DestBuffer + x * sizeof(uint16) + y * DestStride);
 
-					uint32 Index = (x % 4) + (y % 4) * 4; 
+							uint32 Index = (x % 4) + (y % 4) * 4;
 
-					Dest[0] = Bases[Index].R;
-					Dest[1] = Bases[Index].G;
+							Dest[0] = Bases[Index].R;
+							Dest[1] = Bases[Index].G;
+						}
+					}
 				}
+				RHICmdList.UnlockTexture2D((FTexture2DRHIRef&)SSAORandomization->GetRenderTargetItem().ShaderResourceTexture, 0, false);
 			}
-		}
-		RHICmdList.UnlockTexture2D((FTexture2DRHIRef&)SSAORandomization->GetRenderTargetItem().ShaderResourceTexture, 0, false);
-			}
-	
+
 			{
 				FPooledRenderTargetDesc Desc(FPooledRenderTargetDesc::Create2DDesc(FIntPoint(4, 4), PF_R8G8, FClearValueBinding::None, TexCreate_HideInVisualizeTexture, TexCreate_None | TexCreate_NoFastClear, false));
 				Desc.AutoWritable = false;
-				GRenderTargetPool.FindFreeElement(RHICmdList, Desc,GTAORandomization, TEXT("GTAORandomization"), true, ERenderTargetTransience::NonTransient);
+				GRenderTargetPool.FindFreeElement(RHICmdList, Desc, GTAORandomization, TEXT("GTAORandomization"), true, ERenderTargetTransience::NonTransient);
 				// Write the contents of the texture.
 				uint32 DestStride;
 				uint8* DestBuffer = (uint8*)RHICmdList.LockTexture2D((FTexture2DRHIRef&)GTAORandomization->GetRenderTargetItem().ShaderResourceTexture, 0, RLM_WriteOnly, DestStride, false);
 
-				for(int32 y = 0; y < Desc.Extent.Y; ++y)
+				for (int32 y = 0; y < Desc.Extent.Y; ++y)
 				{
-					for(int32 x = 0; x < Desc.Extent.X; ++x)
+					for (int32 x = 0; x < Desc.Extent.X; ++x)
 					{
 						uint8* Dest = (uint8*)(DestBuffer + x * sizeof(uint16) + y * DestStride);
-			
-						float Angle  = (PI/16.0f)  * ((((x + y) & 0x3) << 2) + (x & 0x3));
-						float Step	 = 0.5f + ((1.0f/8.0f) * ((y - x) & 0x3));
 
-						float ScaleCos = FMath::Cos(Angle) ;
-						float ScaleSin = FMath::Sin(Angle) ;
-			
+						float Angle = (PI / 16.0f)  * ((((x + y) & 0x3) << 2) + (x & 0x3));
+						float Step = 0.5f + ((1.0f / 8.0f) * ((y - x) & 0x3));
+
+						float ScaleCos = FMath::Cos(Angle);
+						float ScaleSin = FMath::Sin(Angle);
+
 						Dest[0] = (uint8_t)(ScaleCos*127.5f + 127.5f);
 						Dest[1] = (uint8_t)(ScaleSin*127.5f + 127.5f);
 					}
@@ -405,178 +534,49 @@ void FSystemTextures::InitializeFeatureLevelDependentTextures(FRHICommandListImm
 			}
 			RHICmdList.UnlockTexture2D((FTexture2DRHIRef&)GTAORandomization->GetRenderTargetItem().ShaderResourceTexture, 0, false);
 			
-			
-		{
-			// for testing, with 128x128 R8G8 we are very close to the reference (if lower res is needed we might have to add an offset to counter the 0.5f texel shift)
-			const bool bReference = false;
-
-			EPixelFormat Format = PF_R8G8;
-			// for low roughness we would get banding with PF_R8G8 but for low spec it could be used, for now we don't do this optimization
-			if (GPixelFormats[PF_G16R16].Supported)
-			{
-				Format = PF_G16R16;
-			}
-
-			FPooledRenderTargetDesc Desc(FPooledRenderTargetDesc::Create2DDesc(FIntPoint(128, 32), Format, FClearValueBinding::None, 0, TexCreate_ShaderResource, false));
-			Desc.AutoWritable = false;
-			if (bReference)
-			{
-				Desc.Extent.X = 128;
-				Desc.Extent.Y = 128;
-			}
-
-			GRenderTargetPool.FindFreeElement(RHICmdList, Desc, PreintegratedGF, TEXT("PreintegratedGF"), true, ERenderTargetTransience::NonTransient);
-			// Write the contents of the texture.
-			uint32 DestStride;
-			uint8* DestBuffer = (uint8*)RHICmdList.LockTexture2D((FTexture2DRHIRef&)PreintegratedGF->GetRenderTargetItem().ShaderResourceTexture, 0, RLM_WriteOnly, DestStride, false);
-
-			// x is NoV, y is roughness
-			for (int32 y = 0; y < Desc.Extent.Y; y++)
 			{
-				float Roughness = (float)(y + 0.5f) / Desc.Extent.Y;
-				float m = Roughness * Roughness;
-					float m2 = m * m;
-
-				for (int32 x = 0; x < Desc.Extent.X; x++)
-				{
-					float NoV = (float)(x + 0.5f) / Desc.Extent.X;
-
-					FVector V;
-					V.X = FMath::Sqrt(1.0f - NoV * NoV);	// sin
-					V.Y = 0.0f;
-					V.Z = NoV;								// cos
-
-					float A = 0.0f;
-					float B = 0.0f;
-					float C = 0.0f;
-
-					const uint32 NumSamples = 128;
-					for (uint32 i = 0; i < NumSamples; i++)
-					{
-						float E1 = (float)i / NumSamples;
-						float E2 = (double)ReverseBits(i) / (double)0x100000000LL;
-
-						{
-							float Phi = 2.0f * PI * E1;
-							float CosPhi = FMath::Cos(Phi);
-							float SinPhi = FMath::Sin(Phi);
-							float CosTheta = FMath::Sqrt((1.0f - E2) / (1.0f + (m2 - 1.0f) * E2));
-							float SinTheta = FMath::Sqrt(1.0f - CosTheta * CosTheta);
-
-							FVector H(SinTheta * FMath::Cos(Phi), SinTheta * FMath::Sin(Phi), CosTheta);
-							FVector L = 2.0f * (V | H) * H - V;
-
-							float NoL = FMath::Max(L.Z, 0.0f);
-							float NoH = FMath::Max(H.Z, 0.0f);
-							float VoH = FMath::Max(V | H, 0.0f);
-
-							if (NoL > 0.0f)
-							{
-									float Vis_SmithV = NoL * (NoV * (1 - m) + m);
-									float Vis_SmithL = NoV * (NoL * (1 - m) + m);
-									float Vis = 0.5f / (Vis_SmithV + Vis_SmithL);
-
-								float NoL_Vis_PDF = NoL * Vis * (4.0f * VoH / NoH);
-								float Fc = 1.0f - VoH;
-								Fc *= FMath::Square(Fc*Fc);
-								A += NoL_Vis_PDF * (1.0f - Fc);
-								B += NoL_Vis_PDF * Fc;
-							}
-						}
-
-						{
-							float Phi = 2.0f * PI * E1;
-							float CosPhi = FMath::Cos(Phi);
-							float SinPhi = FMath::Sin(Phi);
-							float CosTheta = FMath::Sqrt(E2);
-							float SinTheta = FMath::Sqrt(1.0f - CosTheta * CosTheta);
-
-							FVector L(SinTheta * FMath::Cos(Phi), SinTheta * FMath::Sin(Phi), CosTheta);
-							FVector H = (V + L).GetUnsafeNormal();
-
-							float NoL = FMath::Max(L.Z, 0.0f);
-							float NoH = FMath::Max(H.Z, 0.0f);
-							float VoH = FMath::Max(V | H, 0.0f);
-
-							float FD90 = 0.5f + 2.0f * VoH * VoH * Roughness;
-								float FdV = 1.0f + (FD90 - 1.0f) * pow(1.0f - NoV, 5);
-								float FdL = 1.0f + (FD90 - 1.0f) * pow(1.0f - NoL, 5);
-							C += FdV * FdL;// * ( 1.0f - 0.3333f * Roughness );
-						}
-					}
-					A /= NumSamples;
-					B /= NumSamples;
-					C /= NumSamples;
-
-						if (Desc.Format == PF_A16B16G16R16)
-					{
-						uint16* Dest = (uint16*)(DestBuffer + x * 8 + y * DestStride);
-						Dest[0] = (int32)(FMath::Clamp(A, 0.0f, 1.0f) * 65535.0f + 0.5f);
-						Dest[1] = (int32)(FMath::Clamp(B, 0.0f, 1.0f) * 65535.0f + 0.5f);
-						Dest[2] = (int32)(FMath::Clamp(C, 0.0f, 1.0f) * 65535.0f + 0.5f);
-					}
-					else if (Desc.Format == PF_G16R16)
-					{
-						uint16* Dest = (uint16*)(DestBuffer + x * 4 + y * DestStride);
- 						Dest[0] = (int32)(FMath::Clamp(A, 0.0f, 1.0f) * 65535.0f + 0.5f);
- 						Dest[1] = (int32)(FMath::Clamp(B, 0.0f, 1.0f) * 65535.0f + 0.5f);
-					}
-					else
-					{
-						check(Desc.Format == PF_R8G8);
-
-						uint8* Dest = (uint8*)(DestBuffer + x * 2 + y * DestStride);
-						Dest[0] = (int32)(FMath::Clamp(A, 0.0f, 1.0f) * 255.9999f);
-						Dest[1] = (int32)(FMath::Clamp(B, 0.0f, 1.0f) * 255.9999f);
-					}
-				}
-			}
-			RHICmdList.UnlockTexture2D((FTexture2DRHIRef&)PreintegratedGF->GetRenderTargetItem().ShaderResourceTexture, 0, false);
-		}
-
-		{
-			FPooledRenderTargetDesc Desc(FPooledRenderTargetDesc::Create2DDesc(FIntPoint(LTC_Size, LTC_Size), PF_FloatRGBA, FClearValueBinding::None, TexCreate_FastVRAM, TexCreate_ShaderResource, false));
-			Desc.AutoWritable = false;
+				FPooledRenderTargetDesc Desc(FPooledRenderTargetDesc::Create2DDesc(FIntPoint(LTC_Size, LTC_Size), PF_FloatRGBA, FClearValueBinding::None, TexCreate_FastVRAM, TexCreate_ShaderResource, false));
+				Desc.AutoWritable = false;
 
-			GRenderTargetPool.FindFreeElement(RHICmdList, Desc, LTCMat, TEXT("LTCMat"));
-			// Write the contents of the texture.
-			uint32 DestStride;
-			uint8* DestBuffer = (uint8*)RHICmdList.LockTexture2D((FTexture2DRHIRef&)LTCMat->GetRenderTargetItem().ShaderResourceTexture, 0, RLM_WriteOnly, DestStride, false);
+				GRenderTargetPool.FindFreeElement(RHICmdList, Desc, LTCMat, TEXT("LTCMat"));
+				// Write the contents of the texture.
+				uint32 DestStride;
+				uint8* DestBuffer = (uint8*)RHICmdList.LockTexture2D((FTexture2DRHIRef&)LTCMat->GetRenderTargetItem().ShaderResourceTexture, 0, RLM_WriteOnly, DestStride, false);
 
 				for (int32 y = 0; y < Desc.Extent.Y; ++y)
-			{
-					for (int32 x = 0; x < Desc.Extent.X; ++x)
 				{
-					uint16* Dest = (uint16*)(DestBuffer + x * 4 * sizeof(uint16) + y * DestStride);
+					for (int32 x = 0; x < Desc.Extent.X; ++x)
+					{
+						uint16* Dest = (uint16*)(DestBuffer + x * 4 * sizeof(uint16) + y * DestStride);
 
 						for (int k = 0; k < 4; k++)
 							Dest[k] = FFloat16(LTC_Mat[4 * (x + y * LTC_Size) + k]).Encoded;
+					}
 				}
+				RHICmdList.UnlockTexture2D((FTexture2DRHIRef&)LTCMat->GetRenderTargetItem().ShaderResourceTexture, 0, false);
 			}
-			RHICmdList.UnlockTexture2D((FTexture2DRHIRef&)LTCMat->GetRenderTargetItem().ShaderResourceTexture, 0, false);
-		}
 
-		{
-			FPooledRenderTargetDesc Desc(FPooledRenderTargetDesc::Create2DDesc(FIntPoint(LTC_Size, LTC_Size), PF_G16R16F, FClearValueBinding::None, TexCreate_FastVRAM, TexCreate_ShaderResource, false));
-			Desc.AutoWritable = false;
+			{
+				FPooledRenderTargetDesc Desc(FPooledRenderTargetDesc::Create2DDesc(FIntPoint(LTC_Size, LTC_Size), PF_G16R16F, FClearValueBinding::None, TexCreate_FastVRAM, TexCreate_ShaderResource, false));
+				Desc.AutoWritable = false;
 
-			GRenderTargetPool.FindFreeElement(RHICmdList, Desc, LTCAmp, TEXT("LTCAmp"));
-			// Write the contents of the texture.
-			uint32 DestStride;
-			uint8* DestBuffer = (uint8*)RHICmdList.LockTexture2D((FTexture2DRHIRef&)LTCAmp->GetRenderTargetItem().ShaderResourceTexture, 0, RLM_WriteOnly, DestStride, false);
+				GRenderTargetPool.FindFreeElement(RHICmdList, Desc, LTCAmp, TEXT("LTCAmp"));
+				// Write the contents of the texture.
+				uint32 DestStride;
+				uint8* DestBuffer = (uint8*)RHICmdList.LockTexture2D((FTexture2DRHIRef&)LTCAmp->GetRenderTargetItem().ShaderResourceTexture, 0, RLM_WriteOnly, DestStride, false);
 
 				for (int32 y = 0; y < Desc.Extent.Y; ++y)
-			{
-					for (int32 x = 0; x < Desc.Extent.X; ++x)
 				{
-					uint16* Dest = (uint16*)(DestBuffer + x * 2 * sizeof(uint16) + y * DestStride);
+					for (int32 x = 0; x < Desc.Extent.X; ++x)
+					{
+						uint16* Dest = (uint16*)(DestBuffer + x * 2 * sizeof(uint16) + y * DestStride);
 
 						for (int k = 0; k < 2; k++)
 							Dest[k] = FFloat16(LTC_Amp[4 * (x + y * LTC_Size) + k]).Encoded;
+					}
 				}
+				RHICmdList.UnlockTexture2D((FTexture2DRHIRef&)LTCAmp->GetRenderTargetItem().ShaderResourceTexture, 0, false);
 			}
-			RHICmdList.UnlockTexture2D((FTexture2DRHIRef&)LTCAmp->GetRenderTargetItem().ShaderResourceTexture, 0, false);
-		}
 		} // end Create the SSAO randomization texture
 	} // end if (FeatureLevelInitializedTo < ERHIFeatureLevel::SM5 && InFeatureLevel >= ERHIFeatureLevel::SM5)
 
