Correct a variable name and a comment.

diff --git a/Engine/Shaders/Private/BasePassPixelShader.usf b/Engine/Shaders/Private/BasePassPixelShader.usf
--- a/Engine/Shaders/Private/BasePassPixelShader.usf
+++ b/Engine/Shaders/Private/BasePassPixelShader.usf
@@ -14,7 +14,7 @@
 	#define EyeAdaptationStruct OpaqueBasePass
 #endif
 
-//ForwardShading (including mobile) use preintegrated GF lut for simple IBL, it could be enable in material editor.
+// Material setting to allow forward shading (including mobile) to use preintegrated GF lut for simple IBL.
 #if MATERIAL_SHADINGMODEL_SINGLELAYERWATER || (FORWARD_SHADING && MATERIAL_USE_PREINTEGRATED_GF)
 #define PreIntegratedGF			OpaqueBasePass.PreIntegratedGFTexture
 #if SUPPORTS_INDEPENDENT_SAMPLERS
diff --git a/Engine/Source/Runtime/MaterialShaderQualitySettings/Classes/ShaderPlatformQualitySettings.h b/Engine/Source/Runtime/MaterialShaderQualitySettings/Classes/ShaderPlatformQualitySettings.h
--- a/Engine/Source/Runtime/MaterialShaderQualitySettings/Classes/ShaderPlatformQualitySettings.h
+++ b/Engine/Source/Runtime/MaterialShaderQualitySettings/Classes/ShaderPlatformQualitySettings.h
@@ -36,7 +36,7 @@ public:
 		, bForceNonMetal(false)
 		, bForceDisableLMDirectionality(false)
 		, bForceLQReflections(false)
-		, bForceNotUsePreintegratedGF(false)
+		, bForceDisablePreintegratedGF(false)
 		, bDisableMaterialNormalCalculation(false)
 		, MobileCSMQuality(EMobileCSMQuality::PCF_2x2)
 	{
@@ -61,7 +61,7 @@ public:
 	bool bForceLQReflections;
 
 	UPROPERTY(EditAnywhere, Config, Meta = (DisplayName = "Force not use preintegrated GF for simple IBL"), Category = "Quality")
-	bool bForceNotUsePreintegratedGF;
+	bool bForceDisablePreintegratedGF;
 
 	UPROPERTY(EditAnywhere, Config, Meta = (DisplayName = "Disable material normal calculation"), Category = "Quality")
 	bool bDisableMaterialNormalCalculation;
diff --git a/Engine/Source/Runtime/MaterialShaderQualitySettings/Private/MaterialShaderQualitySettings.cpp b/Engine/Source/Runtime/MaterialShaderQualitySettings/Private/MaterialShaderQualitySettings.cpp
--- a/Engine/Source/Runtime/MaterialShaderQualitySettings/Private/MaterialShaderQualitySettings.cpp
+++ b/Engine/Source/Runtime/MaterialShaderQualitySettings/Private/MaterialShaderQualitySettings.cpp
@@ -147,7 +147,7 @@ bool FMaterialQualityOverrides::HasAnyOverridesSet() const
 		|| bForceFullyRough != DefaultOverrides.bForceFullyRough
 		|| bForceNonMetal != DefaultOverrides.bForceNonMetal
 		|| bForceLQReflections != DefaultOverrides.bForceLQReflections
-		|| bForceNotUsePreintegratedGF != DefaultOverrides.bForceNotUsePreintegratedGF
+		|| bForceDisablePreintegratedGF != DefaultOverrides.bForceDisablePreintegratedGF
 		|| bDisableMaterialNormalCalculation != DefaultOverrides.bDisableMaterialNormalCalculation
 		|| bDiscardQualityDuringCook != DefaultOverrides.bDiscardQualityDuringCook;
 }
diff --git a/Engine/Source/Runtime/Renderer/Private/MobileBasePassRendering.cpp b/Engine/Source/Runtime/Renderer/Private/MobileBasePassRendering.cpp
--- a/Engine/Source/Runtime/Renderer/Private/MobileBasePassRendering.cpp
+++ b/Engine/Source/Runtime/Renderer/Private/MobileBasePassRendering.cpp
@@ -99,7 +99,7 @@ bool TMobileBasePassPSPolicyParamType<LightMapPolicyType>::ModifyCompilationEnvi
 	OutEnvironment.SetDefine(TEXT("MOBILE_QL_FORCE_NONMETAL"), QualityOverrides.bEnableOverride && QualityOverrides.bForceNonMetal != 0 ? 1u : 0u);
 	OutEnvironment.SetDefine(TEXT("QL_FORCEDISABLE_LM_DIRECTIONALITY"), QualityOverrides.bEnableOverride && QualityOverrides.bForceDisableLMDirectionality != 0 ? 1u : 0u);
 	OutEnvironment.SetDefine(TEXT("MOBILE_QL_FORCE_LQ_REFLECTIONS"), QualityOverrides.bEnableOverride && QualityOverrides.bForceLQReflections != 0 ? 1u : 0u);
-	OutEnvironment.SetDefine(TEXT("MOBILE_QL_FORCE_DISABLE_PREINTEGRATEDGF"), QualityOverrides.bEnableOverride && QualityOverrides.bForceNotUsePreintegratedGF != 0 ? 1u : 0u);
+	OutEnvironment.SetDefine(TEXT("MOBILE_QL_FORCE_DISABLE_PREINTEGRATEDGF"), QualityOverrides.bEnableOverride && QualityOverrides.bForceDisablePreintegratedGF != 0 ? 1u : 0u);
 	OutEnvironment.SetDefine(TEXT("MOBILE_CSM_QUALITY"), (uint32)QualityOverrides.MobileCSMQuality);
 	OutEnvironment.SetDefine(TEXT("MOBILE_QL_DISABLE_MATERIAL_NORMAL"), QualityOverrides.bEnableOverride && QualityOverrides.bDisableMaterialNormalCalculation);
 	return true;
