DRPS - Merge fixes
#rnx

diff --git a/Engine/Source/Runtime/Renderer/Private/DeferredShadingRenderer.cpp b/Engine/Source/Runtime/Renderer/Private/DeferredShadingRenderer.cpp
--- a/Engine/Source/Runtime/Renderer/Private/DeferredShadingRenderer.cpp
+++ b/Engine/Source/Runtime/Renderer/Private/DeferredShadingRenderer.cpp
@@ -2369,7 +2369,11 @@ void FDeferredShadingSceneRenderer::Render(FRHICommandListImmediate& RHICmdList)
 		{
 			for (int32 ViewIndex = 0; ViewIndex < Views.Num(); ViewIndex++)
 			{
-				GPostProcessing.Process(RHICmdList, Views[ ViewIndex ], SceneContext.SceneVelocity);
+				FViewInfo& View = Views[ViewIndex];
+
+				SCOPED_GPU_MASK(RHICmdList, View.GPUMask);
+				RDG_EVENT_SCOPE_CONDITIONAL(GraphBuilder, Views.Num() > 1, "View%d", ViewIndex);
+				AddPostProcessingPasses(GraphBuilder, View, PostProcessingInputs);
 			}
 		}
 
diff --git a/Engine/Source/Runtime/Renderer/Private/PostProcess/PostProcessing.h b/Engine/Source/Runtime/Renderer/Private/PostProcess/PostProcessing.h
--- a/Engine/Source/Runtime/Renderer/Private/PostProcess/PostProcessing.h
+++ b/Engine/Source/Runtime/Renderer/Private/PostProcess/PostProcessing.h
@@ -92,9 +92,6 @@ public:
 class FPostProcessing
 {
 public:
-	// @param VelocityRT only valid if motion blur is supported
-	void Process(FRHICommandListImmediate& RHICmdList, const FViewInfo& View, TRefCountPtr<IPooledRenderTarget>& VelocityRT);
-
 	void ProcessES2(FRHICommandListImmediate& RHICmdList, FScene* Scene, const FViewInfo& View);
 
 	void ProcessPlanarReflection(FRHICommandListImmediate& RHICmdList, const FViewInfo& View, TRefCountPtr<IPooledRenderTarget>& OutFilteredSceneColor);
