Replace RHIGetStencilTexture with RHICreateUnorderedAccessViewStencil to match the other subresource view RHI functions, and resolve some compile errors.
Also fix compiler errors with RHIMapStagingSurface

diff --git a/Engine/Source/Runtime/RHI/Public/DynamicRHI.h b/Engine/Source/Runtime/RHI/Public/DynamicRHI.h
--- a/Engine/Source/Runtime/RHI/Public/DynamicRHI.h
+++ b/Engine/Source/Runtime/RHI/Public/DynamicRHI.h
@@ -953,12 +953,12 @@ public:
 		return nullptr;
 	}
 
-	virtual FTexture2DRHIRef RHIGetStencilTexture(FRHITexture* SourceTextureRHI)
+	virtual FUnorderedAccessViewRHIRef RHICreateUnorderedAccessViewHTile(FRHITexture2D* RenderTarget)
 	{
 		return nullptr;
 	}
 
-	virtual FUnorderedAccessViewRHIRef RHICreateUnorderedAccessViewHTile(FRHITexture2D* RenderTarget)
+	virtual FUnorderedAccessViewRHIRef RHICreateUnorderedAccessViewStencil(FRHITexture2D* DepthTarget, int32 MipLevel)
 	{
 		return nullptr;
 	}
@@ -1419,12 +1419,6 @@ FORCEINLINE FTexture2DRHIRef RHIGetViewportBackBuffer(FRHIViewport* Viewport)
 	return GDynamicRHI->RHIGetViewportBackBuffer(Viewport);
 }
 
-FORCEINLINE FTexture2DRHIRef RHIGetStencilTexture(FRHITexture* SourceTextureRHI)
-{
-	return GDynamicRHI->RHIGetStencilTexture(SourceTextureRHI);
-
-}
-
 FORCEINLINE FShaderResourceViewRHIRef RHICreateShaderResourceViewHTile(FRHITexture2D* RenderTarget)
 {
 	return GDynamicRHI->RHICreateShaderResourceViewHTile(RenderTarget);
@@ -1435,6 +1429,11 @@ FORCEINLINE FUnorderedAccessViewRHIRef RHICreateUnorderedAccessViewHTile(FRHITex
 	return GDynamicRHI->RHICreateUnorderedAccessViewHTile(RenderTarget);
 }
 
+FORCEINLINE FUnorderedAccessViewRHIRef RHICreateUnorderedAccessViewStencil(FRHITexture2D* DepthTarget, int32 MipLevel)
+{
+	return GDynamicRHI->RHICreateUnorderedAccessViewStencil(DepthTarget, MipLevel);
+}
+
 FORCEINLINE void RHIAdvanceFrameForGetViewportBackBuffer(FRHIViewport* Viewport)
 {
 	return GDynamicRHI->RHIAdvanceFrameForGetViewportBackBuffer(Viewport);
diff --git a/Engine/Source/Runtime/RHI/Public/RHIValidation.h b/Engine/Source/Runtime/RHI/Public/RHIValidation.h
--- a/Engine/Source/Runtime/RHI/Public/RHIValidation.h
+++ b/Engine/Source/Runtime/RHI/Public/RHIValidation.h
@@ -979,6 +979,11 @@ public:
 		return RHI->RHICreateUnorderedAccessViewHTile(RenderTarget);
 	}
 
+	virtual FUnorderedAccessViewRHIRef RHICreateUnorderedAccessViewStencil(FRHITexture2D* DepthTarget, int32 MipLevel) override final
+	{
+		return RHI->RHICreateUnorderedAccessViewStencil(DepthTarget, MipLevel);
+	}
+
 	virtual void RHIAliasTextureResources(FRHITexture* DestTexture, FRHITexture* SourceTexture) override final
 	{
 		// Source and target need to be valid objects.
@@ -988,11 +993,6 @@ public:
 		RHI->RHIAliasTextureResources(DestTexture, SourceTexture);
 	}
 
-	virtual FTexture2DRHIRef RHIGetStencilTexture(FRHITexture* SourceTextureRHI) override final
-	{
-		return RHI->RHIGetStencilTexture(SourceTextureRHI);
-	}
-
 	virtual FTextureRHIRef RHICreateAliasedTexture(FRHITexture* SourceTexture) override final
 	{
 		check(SourceTexture);
diff --git a/Engine/Source/Runtime/Renderer/Private/DeferredShadingRenderer.cpp b/Engine/Source/Runtime/Renderer/Private/DeferredShadingRenderer.cpp
--- a/Engine/Source/Runtime/Renderer/Private/DeferredShadingRenderer.cpp
+++ b/Engine/Source/Runtime/Renderer/Private/DeferredShadingRenderer.cpp
@@ -1268,26 +1268,19 @@ void FDeferredShadingSceneRenderer::Render(FRHICommandListImmediate& RHICmdList)
 	const bool bStencilLODCompute = (StencilLODMode == 1 || StencilLODMode == 2);
 	const bool bStencilLODComputeAsync = StencilLODMode == 2;
 
-	FComputeFenceRHIRef AsyncDitherLODStartFence;
 	FComputeFenceRHIRef AsyncDitherLODEndFence;
-
-	FTexture2DRHIRef StencilTexture;
-	FUnorderedAccessViewRHIRef StencilTextureUAV;
 	if (bStencilLODCompute && bDitheredLODTransitionsUseStencil)
 	{
 		// Either compute pass will happen prior to the prepass, and the
 		// stencil clear will be skipped there.
-
-		StencilTexture = GDynamicRHI->RHIGetStencilTexture(SceneContext.GetSceneDepthSurface());
-		StencilTextureUAV = GDynamicRHI->RHICreateUnorderedAccessView(StencilTexture, 0 /* Mip Level */);
-
+		FUnorderedAccessViewRHIRef StencilTextureUAV = RHICreateUnorderedAccessViewStencil(SceneContext.GetSceneDepthSurface(), 0 /* Mip Level */);
 		RHICmdList.TransitionResource(EResourceTransitionAccess::ERWBarrier, EResourceTransitionPipeline::EGfxToCompute, StencilTextureUAV);
 
 		if (bStencilLODComputeAsync)
 		{
 			static FName AsyncDitherLODStartFenceName(TEXT("AsyncDitherLODStartFence"));
 			static FName AsyncDitherLODEndFenceName(TEXT("AsyncDitherLODEndFence"));
-			AsyncDitherLODStartFence = RHICmdList.CreateComputeFence(AsyncDitherLODStartFenceName);
+			FComputeFenceRHIRef AsyncDitherLODStartFence = RHICmdList.CreateComputeFence(AsyncDitherLODStartFenceName);
 			AsyncDitherLODEndFence = RHICmdList.CreateComputeFence(AsyncDitherLODEndFenceName);
 
 			FRHIAsyncComputeCommandListImmediate& RHICmdListComputeImmediate = FRHICommandListExecutor::GetImmediateAsyncComputeCommandList();
