DRPS - Merge fixes
#rnx

diff --git a/Engine/Source/Runtime/Renderer/Private/PostProcess/PostProcessing.h b/Engine/Source/Runtime/Renderer/Private/PostProcess/PostProcessing.h
--- a/Engine/Source/Runtime/Renderer/Private/PostProcess/PostProcessing.h
+++ b/Engine/Source/Runtime/Renderer/Private/PostProcess/PostProcessing.h
@@ -92,9 +92,6 @@ public:
 class FPostProcessing
 {
 public:
-	// @param VelocityRT only valid if motion blur is supported
-	void Process(FRHICommandListImmediate& RHICmdList, const FViewInfo& View, TRefCountPtr<IPooledRenderTarget>& VelocityRT);
-
 	void ProcessES2(FRHICommandListImmediate& RHICmdList, FScene* Scene, const FViewInfo& View);
 
 	void ProcessPlanarReflection(FRHICommandListImmediate& RHICmdList, const FViewInfo& View, TRefCountPtr<IPooledRenderTarget>& OutFilteredSceneColor);
