Replace RHIGetStencilTexture with RHICreateUnorderedAccessViewStencil to match the other subresource view RHI functions, and resolve some compile errors.
Also fix compiler errors with RHIMapStagingSurface

diff --git a/Engine/Source/Runtime/RHI/Public/DynamicRHI.h b/Engine/Source/Runtime/RHI/Public/DynamicRHI.h
--- a/Engine/Source/Runtime/RHI/Public/DynamicRHI.h
+++ b/Engine/Source/Runtime/RHI/Public/DynamicRHI.h
@@ -958,6 +958,11 @@ public:
 		return nullptr;
 	}
 
+	virtual FUnorderedAccessViewRHIRef RHICreateUnorderedAccessViewStencil(FRHITexture2D* DepthTarget, int32 MipLevel)
+	{
+		return nullptr;
+	}
+
 	virtual void RHIAliasTextureResources(FRHITexture* DestTexture, FRHITexture* SrcTexture)
 	{
 		checkNoEntry();
@@ -1424,6 +1429,11 @@ FORCEINLINE FUnorderedAccessViewRHIRef RHICreateUnorderedAccessViewHTile(FRHITex
 	return GDynamicRHI->RHICreateUnorderedAccessViewHTile(RenderTarget);
 }
 
+FORCEINLINE FUnorderedAccessViewRHIRef RHICreateUnorderedAccessViewStencil(FRHITexture2D* DepthTarget, int32 MipLevel)
+{
+	return GDynamicRHI->RHICreateUnorderedAccessViewStencil(DepthTarget, MipLevel);
+}
+
 FORCEINLINE void RHIAdvanceFrameForGetViewportBackBuffer(FRHIViewport* Viewport)
 {
 	return GDynamicRHI->RHIAdvanceFrameForGetViewportBackBuffer(Viewport);
diff --git a/Engine/Source/Runtime/RHI/Public/RHIValidation.h b/Engine/Source/Runtime/RHI/Public/RHIValidation.h
--- a/Engine/Source/Runtime/RHI/Public/RHIValidation.h
+++ b/Engine/Source/Runtime/RHI/Public/RHIValidation.h
@@ -979,6 +979,11 @@ public:
 		return RHI->RHICreateUnorderedAccessViewHTile(RenderTarget);
 	}
 
+	virtual FUnorderedAccessViewRHIRef RHICreateUnorderedAccessViewStencil(FRHITexture2D* DepthTarget, int32 MipLevel) override final
+	{
+		return RHI->RHICreateUnorderedAccessViewStencil(DepthTarget, MipLevel);
+	}
+
 	virtual void RHIAliasTextureResources(FRHITexture* DestTexture, FRHITexture* SourceTexture) override final
 	{
 		// Source and target need to be valid objects.
diff --git a/Engine/Source/Runtime/Renderer/Private/DeferredShadingRenderer.cpp b/Engine/Source/Runtime/Renderer/Private/DeferredShadingRenderer.cpp
--- a/Engine/Source/Runtime/Renderer/Private/DeferredShadingRenderer.cpp
+++ b/Engine/Source/Runtime/Renderer/Private/DeferredShadingRenderer.cpp
@@ -1235,6 +1235,65 @@ void FDeferredShadingSceneRenderer::Render(FRHICommandListImmediate& RHICmdList)
 		}
 	}
 
+	StencilLODMode = CVarStencilLODDitherMode.GetValueOnRenderThread();
+	if (!GRHISupportsDepthUAV)
+	{
+		// RHI doesn't support depth/stencil UAVs - enforce graphics path
+		StencilLODMode = 0;
+	}
+	else if (StencilLODMode == 2 && !GSupportsEfficientAsyncCompute)
+	{
+		// Async compute is not supported, fall back to compute path (on graphics queue)
+		StencilLODMode = 1;
+	}
+	else if (IsHMDHiddenAreaMaskActive())
+	{
+		// Unsupported mode for compute path - enforce graphics path on VR
+		StencilLODMode = 0;
+	}
+
+	const bool bStencilLODCompute = (StencilLODMode == 1 || StencilLODMode == 2);
+	const bool bStencilLODComputeAsync = StencilLODMode == 2;
+
+	FComputeFenceRHIRef AsyncDitherLODEndFence;
+	if (bStencilLODCompute && bDitheredLODTransitionsUseStencil)
+	{
+		// Either compute pass will happen prior to the prepass, and the
+		// stencil clear will be skipped there.
+		FUnorderedAccessViewRHIRef StencilTextureUAV = RHICreateUnorderedAccessViewStencil(SceneContext.GetSceneDepthSurface(), 0 /* Mip Level */);
+		RHICmdList.TransitionResource(EResourceTransitionAccess::ERWBarrier, EResourceTransitionPipeline::EGfxToCompute, StencilTextureUAV);
+
+		if (bStencilLODComputeAsync)
+		{
+			static FName AsyncDitherLODStartFenceName(TEXT("AsyncDitherLODStartFence"));
+			static FName AsyncDitherLODEndFenceName(TEXT("AsyncDitherLODEndFence"));
+			FComputeFenceRHIRef AsyncDitherLODStartFence = RHICmdList.CreateComputeFence(AsyncDitherLODStartFenceName);
+			AsyncDitherLODEndFence = RHICmdList.CreateComputeFence(AsyncDitherLODEndFenceName);
+
+			FRHIAsyncComputeCommandListImmediate& RHICmdListComputeImmediate = FRHICommandListExecutor::GetImmediateAsyncComputeCommandList();
+
+			RHICmdList.TransitionResource(EResourceTransitionAccess::ERWBarrier, EResourceTransitionPipeline::EComputeToCompute, nullptr, AsyncDitherLODStartFence);
+			RHICmdListComputeImmediate.WaitComputeFence(AsyncDitherLODStartFence);
+
+			PreRenderDitherFill(RHICmdListComputeImmediate, SceneContext, StencilTextureUAV);
+
+			RHICmdListComputeImmediate.TransitionResources(EResourceTransitionAccess::ERWBarrier, EResourceTransitionPipeline::EComputeToCompute, nullptr, 0, AsyncDitherLODEndFence);
+			FRHIAsyncComputeCommandListImmediate::ImmediateDispatch(RHICmdListComputeImmediate);
+		}
+		else
+		{
+			PreRenderDitherFill(RHICmdList, SceneContext, StencilTextureUAV);
+		}
+	}
+
+	// TODO: Move to async compute with proper RDG support.
+	const bool bShouldRenderSkyAtmosphere = ShouldRenderSkyAtmosphere(Scene, ViewFamily.EngineShowFlags);
+	if (bShouldRenderSkyAtmosphere)
+	{
+		// Generate the Sky/Atmosphere look up tables
+		RenderSkyAtmosphereLookUpTables(RHICmdList);
+	}
+
 	// Notify the FX system that the scene is about to be rendered.
 	if (Scene->FXSystem && Views.IsValidIndex(0))
 	{
