Fix a OverriddenRenderTarget missing srv flag bug.

diff --git a/Engine/Source/Runtime/Renderer/Private/PostProcess/PostProcessing.cpp b/Engine/Source/Runtime/Renderer/Private/PostProcess/PostProcessing.cpp
--- a/Engine/Source/Runtime/Renderer/Private/PostProcess/PostProcessing.cpp
+++ b/Engine/Source/Runtime/Renderer/Private/PostProcess/PostProcessing.cpp
@@ -1768,7 +1768,7 @@ void FPostProcessing::ProcessES2(FRHICommandListImmediate& RHICmdList, FScene* S
 			Desc.Format = PF_B8G8R8A8;
 			Desc.NumMips = 1;
 			Desc.DebugName = TEXT("OverriddenRenderTarget");
-			Desc.TargetableFlags |= TexCreate_RenderTargetable;
+			Desc.TargetableFlags |= TexCreate_RenderTargetable | TexCreate_ShaderResource;
 
 			GRenderTargetPool.CreateUntrackedElement(Desc, Temp, Item);
 
