Copying Dev-Mobile @ CL 10787175

diff --git a/Engine/Config/Android/AndroidScalability.ini b/Engine/Config/Android/AndroidScalability.ini
--- a/Engine/Config/Android/AndroidScalability.ini
+++ b/Engine/Config/Android/AndroidScalability.ini
@@ -1,3 +1,110 @@
+; ------------------------------------------------------------------------------
+; Foliage settings
+; ------------------------------------------------------------------------------
+
+[FoliageQuality@0]
+foliage.DensityScale=0.5
+grass.DensityScale=0.5
+
+[FoliageQuality@1]
+grass.densityScale=0.8
+grass.CullDistanceScale=0.8
+foliage.DensityScale=0.8
+
+[FoliageQuality@2]
+grass.densityScale=1.0
+grass.CullDistanceScale=1.0
+foliage.DensityScale=1.0
+
+[FoliageQuality@3]
+grass.densityScale=1.0
+grass.CullDistanceScale=1.0
+foliage.DensityScale=1.0
+
+; ------------------------------------------------------------------------------
+; View distance settings
+; ------------------------------------------------------------------------------
+
+[ViewDistanceQuality@0]
+r.ViewDistanceScale=0.70
+
+[ViewDistanceQuality@1]
+r.ViewDistanceScale=1.0
+
+[ViewDistanceQuality@2]
+r.ViewDistanceScale=1.0
+
+[ViewDistanceQuality@3]
+r.ViewDistanceScale=1.0
+
+; ------------------------------------------------------------------------------
+; Shadow settings
+; ------------------------------------------------------------------------------
+
+[ShadowQuality@0]
+r.ShadowQuality=0
+r.Shadow.MaxCSMResolution=1024
+r.Shadow.DistanceScale=0.5
+r.Shadow.CSM.MaxCascades=1
+r.Shadow.RadiusThreshold=0.06
+
+[ShadowQuality@1]
+r.ShadowQuality=1
+r.Shadow.MaxCSMResolution=1024
+r.Shadow.DistanceScale=0.7
+r.Shadow.CSM.MaxCascades=1
+r.Shadow.RadiusThreshold=0.06
+
+[ShadowQuality@2]
+r.ShadowQuality=1
+r.Shadow.MaxCSMResolution=2048
+r.Shadow.DistanceScale=1.0
+r.Shadow.CSM.MaxCascades=2
+r.Shadow.RadiusThreshold=0.06
+
+[ShadowQuality@3]
+r.ShadowQuality=2
+r.Shadow.MaxCSMResolution=2048
+r.Shadow.DistanceScale=1.0
+r.Shadow.CSM.MaxCascades=2
+r.Shadow.RadiusThreshold=0.06
+
+; ------------------------------------------------------------------------------
+; Postprocess settings
+; ------------------------------------------------------------------------------
+
+[PostProcessQuality@0]
+r.Tonemapper.Quality=1
+r.BloomQuality=0
+r.DepthOfFieldQuality=0
+r.LightShaftQuality=0
+r.RefractionQuality=0
+Slate.ForceBackgroundBlurLowQualityOverride=1
+
+[PostProcessQuality@1]
+r.Tonemapper.Quality=2
+r.BloomQuality=1
+r.DepthOfFieldQuality=1
+r.LightShaftQuality=1
+r.RefractionQuality=1
+Slate.ForceBackgroundBlurLowQualityOverride=1
+
+[PostProcessQuality@2]
+r.Tonemapper.Quality=3
+r.BloomQuality=1
+r.DepthOfFieldQuality=1
+r.LightShaftQuality=1
+r.RefractionQuality=1
+Slate.ForceBackgroundBlurLowQualityOverride=1
+
+[PostProcessQuality@3]
+r.Tonemapper.Quality=3
+r.BloomQuality=1
+r.DepthOfFieldQuality=1
+r.LightShaftQuality=1
+r.RefractionQuality=1
+Slate.ForceBackgroundBlurLowQualityOverride=1
+
 ; ------------------------------------------------------------------------------
 ; Effects settings
 ; ------------------------------------------------------------------------------
@@ -45,3 +152,27 @@ r.SkyAtmosphere.FastSkyLUT.SampleCountMax=8.0
 r.SkyAtmosphere.FastSkyLUT.Width=96
 r.SkyAtmosphere.FastSkyLUT.Height=50
 r.SkyAtmosphere.TransmittanceLUT.UseSmallFormat=1
+
+; ------------------------------------------------------------------------------
+; Texture settings
+; ------------------------------------------------------------------------------
+
+[TextureQuality@0]
+r.MaxAnisotropy=1
+
+[TextureQuality@1]
+r.MaxAnisotropy=1
+
+[TextureQuality@2]
+r.MaxAnisotropy=2
+
+[TextureQuality@3]
+r.MaxAnisotropy=4
+
+[ScalabilityPreview]
+; don't want to preview any texture streaming changes in the editor
+BlacklistCVars=r.MaxAnisotropy
++BlacklistCVars=r.Streaming.Boost
++BlacklistCVars=r.Streaming.AmortizeCPUToGPUCopy
++BlacklistCVars=r.Streaming.MaxNumTexturesToStreamPerFrame
++BlacklistCVars=r.Streaming.PoolSize
\ No newline at end of file
diff --git a/Engine/Config/BaseDeviceProfiles.ini b/Engine/Config/BaseDeviceProfiles.ini
--- a/Engine/Config/BaseDeviceProfiles.ini
+++ b/Engine/Config/BaseDeviceProfiles.ini
@@ -9,19 +9,14 @@
 +DeviceProfileNameAndTypes=WindowsServer,Windows
 +DeviceProfileNameAndTypes=WindowsClient,Windows
 +DeviceProfileNameAndTypes=IOS,IOS
-+DeviceProfileNameAndTypes=iPad2,IOS
-+DeviceProfileNameAndTypes=iPad3,IOS
-+DeviceProfileNameAndTypes=iPad4,IOS
 +DeviceProfileNameAndTypes=iPadAir,IOS
-+DeviceProfileNameAndTypes=iPadMini,IOS
++DeviceProfileNameAndTypes=iPadAir2,IOS
++DeviceProfileNameAndTypes=iPadAir3,IOS
 +DeviceProfileNameAndTypes=iPadMini2,IOS
 +DeviceProfileNameAndTypes=iPadMini4,IOS
-+DeviceProfileNameAndTypes=iPodTouch5,IOS
++DeviceProfileNameAndTypes=iPadMini5,IOS
 +DeviceProfileNameAndTypes=iPodTouch6,IOS
 +DeviceProfileNameAndTypes=iPodTouch7,IOS
-+DeviceProfileNameAndTypes=iPhone4,IOS
-+DeviceProfileNameAndTypes=iPhone4S,IOS
-+DeviceProfileNameAndTypes=iPhone5,IOS
 +DeviceProfileNameAndTypes=iPhone5S,IOS
 +DeviceProfileNameAndTypes=iPhone6,IOS
 +DeviceProfileNameAndTypes=iPhone6Plus,IOS
@@ -49,8 +44,6 @@
 +DeviceProfileNameAndTypes=iPad7,IOS
 +DeviceProfileNameAndTypes=iPadPro11,IOS
 +DeviceProfileNameAndTypes=iPadPro3_129,IOS
-+DeviceProfileNameAndTypes=iPadAir3,IOS
-+DeviceProfileNameAndTypes=iPadMini5,IOS
 +DeviceProfileNameAndTypes=AppleTV,IOS
 +DeviceProfileNameAndTypes=AppleTV4K,IOS
 +DeviceProfileNameAndTypes=Mac,Mac
@@ -69,32 +62,25 @@
 +DeviceProfileNameAndTypes=Android_Mid,Android
 +DeviceProfileNameAndTypes=Android_High,Android
 +DeviceProfileNameAndTypes=Android_Default,Android
-+DeviceProfileNameAndTypes=Android_Adreno320,Android
-+DeviceProfileNameAndTypes=Android_Adreno2xx,Android
-+DeviceProfileNameAndTypes=Android_Adreno330,Android
-+DeviceProfileNameAndTypes=Android_Adreno330_Ver53,Android
-+DeviceProfileNameAndTypes=Android_Adreno4xx_Nexus6,Android
 +DeviceProfileNameAndTypes=Android_Adreno4xx,Android
++DeviceProfileNameAndTypes=Android_Adreno5xx_Low,Android
 +DeviceProfileNameAndTypes=Android_Adreno5xx,Android
-+DeviceProfileNameAndTypes=Android_Adreno5xx_No_Vulkan,Android
 +DeviceProfileNameAndTypes=Android_Adreno6xx,Android
-+DeviceProfileNameAndTypes=Android_PowerVR54x,Android
-+DeviceProfileNameAndTypes=Android_PowerVR_G6xxx,Android
-+DeviceProfileNameAndTypes=Android_PowerVRHigh,Android
-+DeviceProfileNameAndTypes=Android_PowerVRLow,Android
-+DeviceProfileNameAndTypes=Android_PowerVRMid,Android
-+DeviceProfileNameAndTypes=Android_TegraK1,Android
-+DeviceProfileNameAndTypes=Android_Tegra4,Android
-+DeviceProfileNameAndTypes=Android_IntelHD,Android
-+DeviceProfileNameAndTypes=Android_IntelHD_ES3,Android
-+DeviceProfileNameAndTypes=Android_Mali_4xx,Android
++DeviceProfileNameAndTypes=Android_Adreno6xx_Vulkan,Android
 +DeviceProfileNameAndTypes=Android_Mali_T6xx,Android
 +DeviceProfileNameAndTypes=Android_Mali_T7xx,Android
 +DeviceProfileNameAndTypes=Android_Mali_T8xx,Android
-+DeviceProfileNameAndTypes=Android_Mali_T8xx_No_Vulkan,Android
 +DeviceProfileNameAndTypes=Android_Mali_G71,Android
 +DeviceProfileNameAndTypes=Android_Mali_G72,Android
++DeviceProfileNameAndTypes=Android_Mali_G72_Vulkan,Android
 +DeviceProfileNameAndTypes=Android_Mali_G76,Android
++DeviceProfileNameAndTypes=Android_Mali_G76_Vulkan,Android
++DeviceProfileNameAndTypes=Android_PowerVR_G6xxx,Android
++DeviceProfileNameAndTypes=Android_PowerVR_GT7xxx,Android
++DeviceProfileNameAndTypes=Android_PowerVR_GE8xxx,Android
++DeviceProfileNameAndTypes=Android_PowerVR_GM9xxx,Android
++DeviceProfileNameAndTypes=Android_PowerVR_GM9xxx_Vulkan,Android
++DeviceProfileNameAndTypes=Android_TegraK1,Android
 +DeviceProfileNameAndTypes=Lumin,Lumin
 +DeviceProfileNameAndTypes=Lumin_Desktop,Lumin
 
@@ -149,14 +135,12 @@ BaseProfileName=Windows
 DeviceType=Windows
 BaseProfileName=Windows
 
+;----------------
+; iOS Devices
+
 [IOS DeviceProfile]
 DeviceType=IOS
 BaseProfileName=
-+CVars=r.BloomQuality=0
-+CVars=r.DepthOfFieldQuality=0
-+CVars=r.LightShaftQuality=0
-+CVars=r.RefractionQuality=0
-+CVars=r.ShadowQuality=2
 ; HZB will be slower with tiled, and not needed
 +CVars=r.HZBOcclusion=0
 ; Non need for depth prepass
@@ -165,12 +149,18 @@ BaseProfileName=
 +CVars=r.TranslucentLightingVolume=0
 ; Needs geometry shader support
 +CVars=r.AllowPointLightCubemapShadows=0
-; Temporal AA (?) looks blurry right now
-+CVars=r.PostProcessAAQuality=0
 ; Not sure on this one - may need the whole D Buffer stuff going
 +CVars=r.Decal.StencilSizeThreshold=-1
 +CVars=slate.AbsoluteIndices=1
 +CVars=r.MorphTarget.Mode=0
+; Scalability
++CVars=sg.AntiAliasingQuality=0
++CVars=sg.ShadowQuality=1
++CVars=sg.ViewDistanceQuality=3
++CVars=sg.PostProcessQuality=3
++CVars=sg.TextureQuality=3
++CVars=sg.EffectsQuality=3
++CVars=sg.FoliageQuality=3
 
 [iPad2 DeviceProfile]
 DeviceType=IOS
@@ -190,15 +180,13 @@ BaseProfileName=IOS
 [iPadAir DeviceProfile]
 DeviceType=IOS
 BaseProfileName=IOS
-+CVars=r.BloomQuality=1
++CVars=sg.PostProcessQuality=1
 
 [iPadAir2 DeviceProfile]
 DeviceType=IOS
 BaseProfileName=IOS
-+CVars=r.BloomQuality=1
-+CVars=r.DepthOfFieldQuality=1
-+CVars=r.LightShaftQuality=1
-+CVars=r.RefractionQuality=1
++CVars=sg.PostProcessQuality=3
++CVars=sg.ShadowQuality=2
 
 [iPadMini DeviceProfile]
 DeviceType=IOS
@@ -212,36 +200,13 @@ BaseProfileName=iPadAir
 [iPadMini4 DeviceProfile]
 DeviceType=IOS
 BaseProfileName=iPadMini2
-+CVars=r.RefractionQuality=1
-
-[iPhone4 DeviceProfile]
-DeviceType=IOS
-BaseProfileName=IOS
-+CVars=r.RenderTargetSwitchWorkaround=1
-
-[iPhone4S DeviceProfile]
-DeviceType=IOS
-BaseProfileName=IOS
-+CVars=r.RenderTargetSwitchWorkaround=1
-
-[iPhone5 DeviceProfile]
-DeviceType=IOS
-BaseProfileName=IOS
-+CVars=r.RenderTargetSwitchWorkaround=1
++CVars=sg.PostProcessQuality=2
 
 [iPhone5S DeviceProfile]
 DeviceType=IOS
 BaseProfileName=IOS
 +CVars=r.MobileContentScaleFactor=2
-+CVars=r.BloomQuality=1
-+CVars=r.DepthOfFieldQuality=1
-+CVars=r.LightShaftQuality=1
-+CVars=r.RefractionQuality=1
-
-[iPodTouch5 DeviceProfile]
-DeviceType=IOS
-BaseProfileName=IOS
-+CVars=r.RenderTargetSwitchWorkaround=1
++CVars=sg.PostProcessQuality=3
 
 [iPodTouch6 DeviceProfile]
 DeviceType=IOS
@@ -257,38 +222,31 @@ BaseProfileName=iPhone7
 DeviceType=IOS
 BaseProfileName=IOS
 +CVars=r.MobileContentScaleFactor=2
-+CVars=r.BloomQuality=1
-+CVars=r.DepthOfFieldQuality=1
-+CVars=r.LightShaftQuality=1
-+CVars=r.RefractionQuality=1
++CVars=sg.PostProcessQuality=3
 
 [iPhone6Plus DeviceProfile]
 DeviceType=IOS
 BaseProfileName=IOS
 +CVars=r.MobileContentScaleFactor=2
-+CVars=r.BloomQuality=1
-+CVars=r.DepthOfFieldQuality=1
-+CVars=r.LightShaftQuality=1
-+CVars=r.RefractionQuality=1
++CVars=sg.PostProcessQuality=3
 
 [iPhone6S DeviceProfile]
 DeviceType=IOS
 BaseProfileName=iPhone6
 +CVars=r.MobileContentScaleFactor=2
++CVars=sg.ShadowQuality=2
 
 [iPhone6SPlus DeviceProfile]
 DeviceType=IOS
 BaseProfileName=iPhone6Plus
 +CVars=r.MobileContentScaleFactor=2
++CVars=sg.ShadowQuality=2
 
 [iPhoneSE DeviceProfile]
 DeviceType=IOS
 BaseProfileName=IOS
 +CVars=r.MobileContentScaleFactor=2
-+CVars=r.BloomQuality=1
-+CVars=r.DepthOfFieldQuality=1
-+CVars=r.LightShaftQuality=1
-+CVars=r.RefractionQuality=1
++CVars=sg.PostProcessQuality=3
 
 [iPhone7 DeviceProfile]
 DeviceType=IOS
@@ -303,10 +261,12 @@ BaseProfileName=iPhone6Plus
 [iPhone8 DeviceProfile]
 DeviceType=IOS
 BaseProfileName=iPhone7
++CVars=sg.ShadowQuality=3
 
 [iPhone8Plus DeviceProfile]
 DeviceType=IOS
 BaseProfileName=iPhone7Plus
++CVars=sg.ShadowQuality=3
 
 [iPhoneX DeviceProfile]
 DeviceType=IOS
@@ -346,24 +306,20 @@ BaseProfileName=iPhone8
 [iPadPro DeviceProfile]
 DeviceType=IOS
 BaseProfileName=IOS
-+CVars=r.BloomQuality=1
-+CVars=r.DepthOfFieldQuality=1
-+CVars=r.LightShaftQuality=1
-+CVars=r.RefractionQuality=1
 +CVars=r.MobileContentScaleFactor=1.5
 +CVars=g.TimeoutForBlockOnRenderFence=3000000
 +CVars=r.Decal.StencilSizeThreshold=0.1
 +CVars=r.MetalComputeParameterSize=1024
 +CVars=r.EarlyZPass=3
-+CVars=r.PostProcessAAQuality=4
++CVars=sg.ShadowQuality=3
++CVars=sg.PostProcessQuality=3
++CVars=sg.AntiAliasingQuality=3
 
 [AppleTV DeviceProfile]
 DeviceType=IOS
 BaseProfileName=IOS
-+CVars=r.BloomQuality=1
-+CVars=r.DepthOfFieldQuality=1
-+CVars=r.LightShaftQuality=1
-+CVars=r.RefractionQuality=1
++CVars=sg.PostProcessQuality=3
++CVars=sg.ShadowQuality=3
 
 [AppleTV4K DeviceProfile]
 DeviceType=IOS
@@ -413,243 +369,176 @@ BaseProfileName=IPadPro
 DeviceType=IOS
 BaseProfileName=IPadPro
 
+;----------------
+; Android devices
+
+; These rules are checked in turn and stop at the profile that matches all criteria
 [/Script/AndroidDeviceProfileSelector.AndroidDeviceProfileMatchingRules]
-MatchProfile=(Profile="Android_PowerVR54x",Match=((SourceType=SRC_GpuFamily,CompareType=CMP_Regex,MatchString="PowerVR SGX 54[0-9]")))
-+MatchProfile=(Profile="Android_PowerVR_G6xxx",Match=((SourceType=SRC_GpuFamily,CompareType=CMP_Regex,MatchString="PowerVR Rogue G6[0-9]+")))
-+MatchProfile=(Profile="Android_PowerVRHigh",Match=((SourceType=SRC_GpuFamily,CompareType=CMP_Regex,MatchString="PowerVR Rogue G"),(SourceType=SRC_GlVersion,CompareType=CMP_Regex,MatchString="^OpenGL ES 3\\.")))
-+MatchProfile=(Profile="Android_PowerVRMid",Match=((SourceType=SRC_GpuFamily,CompareType=CMP_Regex,MatchString="PowerVR"),(SourceType=SRC_GlVersion,CompareType=CMP_Regex,MatchString="^OpenGL ES 3\\.")))
-+MatchProfile=(Profile="Android_PowerVRLow",Match=((SourceType=SRC_GpuFamily,CompareType=CMP_Regex,MatchString="PowerVR")))
-+MatchProfile=(Profile="Android_Adreno2xx",Match=((SourceType=SRC_GpuFamily,CompareType=CMP_Regex,MatchString="Adreno \\(TM\\) 2[0-9][0-9]")))
-+MatchProfile=(Profile="Android_Adreno320",Match=((SourceType=SRC_GpuFamily,CompareType=CMP_Equal,MatchString="Adreno (TM) 320")))
-+MatchProfile=(Profile="Android_Adreno330_Ver53",Match=((SourceType=SRC_GpuFamily,CompareType=CMP_Equal,MatchString="Adreno (TM) 330"),(SourceType=SRC_GlVersion,CompareType=CMP_Regex,MatchString="ES 3\\.0 V@([0-9]+)"),(SourceType=SRC_PreviousRegexMatch,CompareType=CMP_GreaterEqual,MatchString="53")))
-+MatchProfile=(Profile="Android_Adreno330",Match=((SourceType=SRC_GpuFamily,CompareType=CMP_Equal,MatchString="Adreno (TM) 330")))
-+MatchProfile=(Profile="Android_Adreno4xx_Nexus6",Match=((SourceType=SRC_DeviceModel,CompareType=CMP_Equal,MatchString="Nexus 6")))
-+MatchProfile=(Profile="Android_Adreno4xx",Match=((SourceType=SRC_GpuFamily,CompareType=CMP_Regex,MatchString="Adreno \\(TM\\) 4[0-9][0-9]")))
-+MatchProfile=(Profile="Android_Adreno51x",Match=((SourceType=SRC_GpuFamily,CompareType=CMP_Regex,MatchString="Adreno \\(TM\\) 51[0-9]")))
-+MatchProfile=(Profile="Android_Adreno5xx_No_Vulkan",Match=((SourceType=SRC_GpuFamily,CompareType=CMP_Regex,MatchString="Adreno \\(TM\\) 5[0-9][0-9]"),(SourceType=SRC_AndroidVersion, CompareType=CMP_Regex,MatchString="([0-9]+).*"),(SourceType=SRC_PreviousRegexMatch,CompareType=CMP_Less,MatchString="8")))
+MatchProfile=(Profile="Android_Adreno4xx",Match=((SourceType=SRC_GpuFamily,CompareType=CMP_Regex,MatchString="Adreno \\(TM\\) 4[0-9][0-9]")))
++MatchProfile=(Profile="Android_Adreno5xx_Low",Match=((SourceType=SRC_GpuFamily,CompareType=CMP_Regex,MatchString="Adreno \\(TM\\) 5[0-1][0-9]")))
 +MatchProfile=(Profile="Android_Adreno5xx",Match=((SourceType=SRC_GpuFamily,CompareType=CMP_Regex,MatchString="Adreno \\(TM\\) 5[0-9][0-9]")))
++MatchProfile=(Profile="Android_Adreno6xx_Vulkan",Match=((SourceType=SRC_GpuFamily,CompareType=CMP_Regex,MatchString="Adreno \\(TM\\) 6[0-9][0-9]")),(SourceType=SRC_AndroidVersion, CompareType=CMP_Regex,MatchString="([0-9]+).*"),(SourceType=SRC_PreviousRegexMatch,CompareType=CMP_GreaterEqual,MatchString="9"),(SourceType=SRC_VulkanAvailable,CompareType=CMP_Equal,MatchString="true")))
 +MatchProfile=(Profile="Android_Adreno6xx",Match=((SourceType=SRC_GpuFamily,CompareType=CMP_Regex,MatchString="Adreno \\(TM\\) 6[0-9][0-9]")))
-+MatchProfile=(Profile="Android_TegraK1",Match=((SourceType=SRC_GpuFamily,CompareType=CMP_Equal,MatchString="NVIDIA Tegra"),(SourceType=SRC_GlVersion,CompareType=CMP_Regex,MatchString="^OpenGL ES 3\\.")))
-+MatchProfile=(Profile="Android_Tegra4",Match=((SourceType=SRC_GpuFamily,CompareType=CMP_Equal,MatchString="NVIDIA Tegra"),(SourceType=SRC_GlVersion,CompareType=CMP_Regex,MatchString="^OpenGL ES 2\\.")))
-+MatchProfile=(Profile="Android_IntelHD_ES3",Match=((SourceType=SRC_GpuFamily,CompareType=CMP_Equal,MatchString="Intel(R) HD Graphics"),(SourceType=SRC_GlVersion,CompareType=CMP_Regex,MatchString="^OpenGL ES 3\\.")))
-+MatchProfile=(Profile="Android_IntelHD",Match=((SourceType=SRC_GpuFamily,CompareType=CMP_Equal,MatchString="Intel(R) HD Graphics"),(SourceType=SRC_GlVersion,CompareType=CMP_Regex,MatchString="^OpenGL ES 2\\.")))
-+MatchProfile=(Profile="Android_Mali_4xx_LenovoS939",Match=((SourceType=SRC_DeviceModel,CompareType=CMP_Equal,MatchString="Lenovo S939")))
-+MatchProfile=(Profile="Android_Mali_4xx_LenovoA7600",Match=((SourceType=SRC_DeviceModel,CompareType=CMP_Regex,MatchString="^Lenovo A7600")))
-+MatchProfile=(Profile="Android_Mali_4xx_WikoHighwaySigns",Match=((SourceType=SRC_DeviceMake,CompareType=CMP_Equal,MatchString="Wiko"),(SourceType=SRC_DeviceModel,CompareType=CMP_Equal,MatchString="Highway Signs")))
-+MatchProfile=(Profile="Android_Mali_4xx_AlcatelPixi3",Match=((SourceType=SRC_DeviceMake,CompareType=CMP_Equal,MatchString="Alcatel"),(SourceType=SRC_DeviceModel,CompareType=CMP_Equal,MatchString="Pixi 3 (4) 3G")))
-+MatchProfile=(Profile="Android_Mali_4xx_KindleFire7",Match=((SourceType=SRC_DeviceModel,CompareType=CMP_Equal,MatchString="KFFOWI")))
-+MatchProfile=(Profile="Android_Mali_4xx",Match=((SourceType=SRC_GpuFamily,CompareType=CMP_Regex,MatchString="^Mali\\-4")))
 +MatchProfile=(Profile="Android_Mali_T6xx",Match=((SourceType=SRC_GpuFamily,CompareType=CMP_Regex,MatchString="^Mali\\-T6")))
 +MatchProfile=(Profile="Android_Mali_T7xx",Match=((SourceType=SRC_GpuFamily,CompareType=CMP_Regex,MatchString="^Mali\\-T7")))
-+MatchProfile=(Profile="Android_Mali_T8xx_No_Vulkan",Match=((SourceType=SRC_GpuFamily,CompareType=CMP_Regex,MatchString="^Mali\\-T8"),(SourceType=SRC_VulkanVersion,CompareType=CMP_Regex,MatchString="^1\\.0\\.([0-9]+)"),(SourceType=SRC_PreviousRegexMatch,CompareType=CMP_Less,MatchString="11")))
-+MatchProfile=(Profile="Android_Mali_T8xx_No_Vulkan",Match=((SourceType=SRC_GpuFamily,CompareType=CMP_Regex,MatchString="^Mali\\-T8"),(SourceType=SRC_AndroidVersion, CompareType=CMP_Regex,MatchString="([0-9]+).*"),(SourceType=SRC_PreviousRegexMatch,CompareType=CMP_Less,MatchString="8")))
 +MatchProfile=(Profile="Android_Mali_T8xx",Match=((SourceType=SRC_GpuFamily,CompareType=CMP_Regex,MatchString="^Mali\\-T8")))
 +MatchProfile=(Profile="Android_Mali_G71",Match=((SourceType=SRC_GpuFamily,CompareType=CMP_Regex,MatchString="^Mali\\-G71")))
++MatchProfile=(Profile="Android_Mali_G72_Vulkan",Match=((SourceType=SRC_GpuFamily,CompareType=CMP_Regex,MatchString="^Mali\\-G72"),(SourceType=SRC_VulkanAvailable,CompareType=CMP_Equal,MatchString="true")))
 +MatchProfile=(Profile="Android_Mali_G72",Match=((SourceType=SRC_GpuFamily,CompareType=CMP_Regex,MatchString="^Mali\\-G72")))
++MatchProfile=(Profile="Android_Mali_G76_Vulkan",Match=((SourceType=SRC_GpuFamily,CompareType=CMP_Regex,MatchString="^Mali\\-G76"),(SourceType=SRC_VulkanAvailable,CompareType=CMP_Equal,MatchString="true")))
 +MatchProfile=(Profile="Android_Mali_G76",Match=((SourceType=SRC_GpuFamily,CompareType=CMP_Regex,MatchString="^Mali\\-G76")))
-+MatchProfile=(Profile="Android_Mid",Match=((SourceType=SRC_GlVersion,CompareType=CMP_Regex,MatchString="^OpenGL ES 3\\.")))
-+MatchProfile=(Profile="Android_Low",Match=((SourceType=SRC_GlVersion,CompareType=CMP_Regex,MatchString="^OpenGL ES 2\\.")))
-
-[/Script/AndroidDeviceProfileSelector.AndroidJavaSurfaceViewDevices]
-SurfaceViewDevices=(Manufacturer="Amazon",Model="Kindle Fire")	; Kindle Fire (2011)
-+SurfaceViewDevices=(Manufacturer="Amazon",Model="KFOT")		; Kindle Fire (2012)
-+SurfaceViewDevices=(Manufacturer="Amazon",Model="KFTT")		; Kindle Fire HD 7" 2012 (WiFi)
-+SurfaceViewDevices=(Manufacturer="Amazon",Model="KFJWA")		; Kindle Fire HD 8.9" 2012 (WAN)
-+SurfaceViewDevices=(Manufacturer="Amazon",Model="KFJWI")		; Kindle Fire HD 8.9" 2012 (WiFi)
-+SurfaceViewDevices=(Manufacturer="HTC",Model="HTC6435LVW")		; HTC One M7
++MatchProfile=(Profile="Android_PowerVR_G6xxx",Match=((SourceType=SRC_GpuFamily,CompareType=CMP_Regex,MatchString="PowerVR Rogue G6[0-9]+")))
++MatchProfile=(Profile="Android_PowerVR_GT7xxx",Match=((SourceType=SRC_GpuFamily,CompareType=CMP_Regex,MatchString="PowerVR Rogue GT7[0-9]+")))
++MatchProfile=(Profile="Android_PowerVR_GE8xxx",Match=((SourceType=SRC_GpuFamily,CompareType=CMP_Regex,MatchString="PowerVR Rogue GE8[0-9]+")))
++MatchProfile=(Profile="Android_PowerVR_GM9xxx_Vulkan",Match=((SourceType=SRC_GpuFamily,CompareType=CMP_Regex,MatchString="PowerVR Rogue GM9[0-9]+"),(SourceType=SRC_VulkanAvailable,CompareType=CMP_Equal,MatchString="true")))
++MatchProfile=(Profile="Android_PowerVR_GM9xxx",Match=((SourceType=SRC_GpuFamily,CompareType=CMP_Regex,MatchString="PowerVR Rogue GM9[0-9]+")))
++MatchProfile=(Profile="Android_TegraK1",Match=((SourceType=SRC_GpuFamily,CompareType=CMP_Equal,MatchString="NVIDIA Tegra"),(SourceType=SRC_GlVersion,CompareType=CMP_Regex,MatchString="^OpenGL ES 3\\.")))
+; Android_Default is the profile of last resort
 
 [Android DeviceProfile]
 DeviceType=Android
 BaseProfileName=
-+CVars=r.MobileContentScaleFactor=1
-+CVars=r.BloomQuality=0
-+CVars=r.DepthOfFieldQuality=0
-+CVars=r.LightShaftQuality=0
-+CVars=r.RefractionQuality=0
-+CVars=r.ShadowQuality=2
++CVars=r.MobileContentScaleFactor=1.0
 +CVars=slate.AbsoluteIndices=1
 +CVars=r.Vulkan.DelayAcquireBackBuffer=2
 +CVars=r.Vulkan.RobustBufferAccess=1
 +CVars=r.Vulkan.DescriptorSetLayoutMode=2
+; Don't enable Vulkan by default. Specific device profiles can set this cvar to 0 to enable Vulkan.
++CVars=r.Android.DisableVulkanSupport=1
 
 [Android_Low DeviceProfile]
 DeviceType=Android
 BaseProfileName=Android
-+CVars=r.MobileContentScaleFactor=0.5
-+CVars=r.ShadowQuality=1
++CVars=r.MobileContentScaleFactor=0.8
+; Scalability groups, see AndroidScalability.ini
++CVars=sg.ViewDistanceQuality=0
++CVars=sg.AntiAliasingQuality=0
++CVars=sg.ShadowQuality=0
++CVars=sg.PostProcessQuality=0
++CVars=sg.TextureQuality=0
 +CVars=sg.EffectsQuality=0
++CVars=sg.FoliageQuality=0
 
 [Android_Mid DeviceProfile]
 DeviceType=Android
 BaseProfileName=Android
-+CVars=r.BloomQuality=1
-+CVars=r.MobileContentScaleFactor=0.8
++CVars=r.MobileContentScaleFactor=1.0
++CVars=sg.ViewDistanceQuality=1
++CVars=sg.AntiAliasingQuality=1
++CVars=sg.ShadowQuality=1
++CVars=sg.PostProcessQuality=1
++CVars=sg.TextureQuality=1
++CVars=sg.EffectsQuality=1
++CVars=sg.FoliageQuality=1
 
 [Android_High DeviceProfile]
 DeviceType=Android
 BaseProfileName=Android
-+CVars=r.BloomQuality=1
-+CVars=r.DepthOfFieldQuality=1
-+CVars=r.LightShaftQuality=1
-+CVars=r.RefractionQuality=1
-+CVars=r.MobileContentScaleFactor=1.0
++CVars=sg.ViewDistanceQuality=2
++CVars=sg.AntiAliasingQuality=2
++CVars=sg.ShadowQuality=2
++CVars=sg.PostProcessQuality=2
++CVars=sg.TextureQuality=2
++CVars=sg.EffectsQuality=2
++CVars=sg.FoliageQuality=2+CVars=r.MobileContentScaleFactor=1.0
 
 [Android_Default DeviceProfile]
 DeviceType=Android
 BaseProfileName=Android_Mid
 
-[Android_Adreno320 DeviceProfile]
-DeviceType=Android
-BaseProfileName=Android_Mid
-
-;This offset needs to be set for the mosaic fallback to work on Galaxy S4 (SAMSUNG-IGH-I337)
-;+CVars=r.DemosaicVposOffset=0.5
+;----------------
+; Adreno devices, Adreno 4xx and above are supported
 
-[Android_Adreno2xx DeviceProfile]
+[Android_Adreno4xx DeviceProfile]
 DeviceType=Android
 BaseProfileName=Android_Low
 
-[Android_Adreno330 DeviceProfile]
+[Android_Adreno5xx_Low DeviceProfile]
 DeviceType=Android
-BaseProfileName=Android_High
-+CVars=OpenGL.UseBufferDiscard=0
-
-[Android_Adreno330_Ver53 DeviceProfile]
-DeviceType=Android
-BaseProfileName=Android_Adreno330
-+CVars=r.DisjointTimerQueries=0
-
-[Android_Adreno51x DeviceProfile]
-DeviceType=Android
-BaseProfileName=Android_High
+BaseProfileName=Android_Low
 
 [Android_Adreno5xx DeviceProfile]
 DeviceType=Android
-BaseProfileName=Android_High
-+CVars=r.DisjointTimerQueries=1
-
-[Android_Adreno5xx_No_Vulkan DeviceProfile]
-DeviceType=Android
-BaseProfileName=Android_Adreno5xx
-; There are several issues (vulkan sub-passes, occlusion queries) on devices running Android 7 and earier
-+CVars=r.Android.DisableVulkanSupport=1
+BaseProfileName=Android_Mid
 
 [Android_Adreno6xx DeviceProfile]
 DeviceType=Android
 BaseProfileName=Android_High
-+CVars=r.DisjointTimerQueries=1
-
-[Android_Adreno4xx DeviceProfile]
-DeviceType=Android
-BaseProfileName=Android_High
 
-[Android_Adreno4xx_Nexus6 DeviceProfile]
+[Android_Adreno6xx_Vulkan DeviceProfile]
 DeviceType=Android
-BaseProfileName=Android_Adreno4xx
+BaseProfileName=Android_Adreno6xx
+; enable Vulkan
++CVars=r.Android.DisableVulkanSupport=0
 
-[Android_PowerVR54x DeviceProfile]
-DeviceType=Android
-BaseProfileName=Android_Mid
-+CVars=r.Mobile.ForceDepthResolve=1
-+CVars=r.Android.DisableASTCSupport=1
-
-[Android_PowerVRMid DeviceProfile]
-DeviceType=Android
-BaseProfileName=Android_Mid
-+CVars=r.Mobile.ForceDepthResolve=1
-+CVars=r.Android.DisableASTCSupport=1
+;----------------
+; PowerVR devices
+; some PowerVR G6xxx devices are known to have problems with ES 3.1 on UE4
 
 [Android_PowerVR_G6xxx DeviceProfile]
 DeviceType=Android
-BaseProfileName=Android_PowerVRMid
-; Disable ES3.1 support on PoverVR G6xxx devices
-+CVars=r.Android.DisableOpenGLES31Support=1
-
-[Android_PowerVRLow DeviceProfile]
-DeviceType=Android
-BaseProfileName=Android_Low
-+CVars=r.Mobile.ForceDepthResolve=1
-+CVars=r.Android.DisableASTCSupport=1
-
-[Android_PowerVRHigh DeviceProfile]
-DeviceType=Android
-BaseProfileName=Android_High
-
-[Android_IntelHD DeviceProfile]
-DeviceType=Android
 BaseProfileName=Android_Low
 
-[Android_IntelHD_ES3 DeviceProfile]
+[Android_PowerVR_GT7xxx DeviceProfile]
 DeviceType=Android
-BaseProfileName=Android
-+CVars=r.BloomQuality=1
-
-[Android_Mali_4xx DeviceProfile]
-DeviceType=Android
-BaseProfileName=Android_Low
-
-[Android_Mali_4xx_LenovoS939 DeviceProfile]
-DeviceType=Android
-BaseProfileName=Android_Mali_4xx
-+CVars=r.AndroidDisableThreadedRendering=1
+BaseProfileName=Android_Mid
 
-[Android_Mali_4xx_LenovoA7600 DeviceProfile]
+[Android_PowerVR_GE8xxx DeviceProfile]
 DeviceType=Android
-BaseProfileName=Android_Mali_4xx
-+CVars=r.AndroidDisableThreadedRendering=1
+BaseProfileName=Android_Mid
 
-[Android_Mali_4xx_WikoHighwaySigns DeviceProfile]
+[Android_PowerVR_GM9xxx DeviceProfile]
 DeviceType=Android
-BaseProfileName=Android_Mali_4xx
-+CVars=r.AndroidDisableThreadedRendering=1
+BaseProfileName=Android_High
 
-[Android_Mali_4xx_AlcatelPixi3 DeviceProfile]
+[Android_PowerVR_GM9xxx_Vulkan DeviceProfile]
 DeviceType=Android
-BaseProfileName=Android_Mali_4xx
-+CVars=r.AndroidDisableThreadedRendering=1
+BaseProfileName=Android_PowerVR_GM9xxx
+; enable Vulkan
++CVars=r.Android.DisableVulkanSupport=0
 
-[Android_Mali_4xx_KindleFire7 DeviceProfile]
-DeviceType=Android
-BaseProfileName=Android_Mali_4xx
-+CVars=r.AndroidDisableThreadedRendering=1
-+CVars=r.AndroidDisableThreadedRenderingFirstLoad=1
+;----------------
+; Mali devices, we support from Mali-T6xx and above
 
 [Android_Mali_T6xx DeviceProfile]
 DeviceType=Android
-BaseProfileName=Android_Mid
+BaseProfileName=Android_Low
 
 [Android_Mali_T7xx DeviceProfile]
 DeviceType=Android
-BaseProfileName=Android_High
+BaseProfileName=Android_Low
 
 [Android_Mali_T8xx DeviceProfile]
 DeviceType=Android
-BaseProfileName=Android_High
+BaseProfileName=Android_Mid
 
 [Android_Mali_G71 DeviceProfile]
 DeviceType=Android
-BaseProfileName=Android_High
+BaseProfileName=Android_Mid
 
 [Android_Mali_G72 DeviceProfile]
 DeviceType=Android
 BaseProfileName=Android_High
 
-[Android_Mali_G76 DeviceProfile]
+[Android_Mali_G72_Vulkan DeviceProfile]
 DeviceType=Android
-BaseProfileName=Android_High
+BaseProfileName=Android_Mali_G72
+; enable Vulkan
++CVars=r.Android.DisableVulkanSupport=0
 
-[Android_TegraK1 DeviceProfile]
+[Android_Mali_G76 DeviceProfile]
 DeviceType=Android
 BaseProfileName=Android_High
 
-[Android_Tegra4 DeviceProfile]
+[Android_Mali_G76_Vulkan DeviceProfile]
 DeviceType=Android
-BaseProfileName=Android_Mid
+BaseProfileName=Android_Mali_G76
+; enable Vulkan
++CVars=r.Android.DisableVulkanSupport=0
 
-[Android_Mali_T8xx_No_Vulkan DeviceProfile]
+[Android_TegraK1 DeviceProfile]
 DeviceType=Android
-BaseProfileName=Android_Mali_T8xx
-+CVars=r.Android.DisableVulkanSupport=1
+BaseProfileName=Android_High
 
 [Lumin DeviceProfile]
 DeviceType=Lumin
diff --git a/Engine/Config/IOS/IOSScalability.ini b/Engine/Config/IOS/IOSScalability.ini
--- a/Engine/Config/IOS/IOSScalability.ini
+++ b/Engine/Config/IOS/IOSScalability.ini
@@ -1,3 +1,110 @@
+; ------------------------------------------------------------------------------
+; Foliage settings
+; ------------------------------------------------------------------------------
+
+[FoliageQuality@0]
+foliage.DensityScale=0.5
+grass.DensityScale=0.5
+
+[FoliageQuality@1]
+grass.densityScale=0.8
+grass.CullDistanceScale=0.8
+foliage.DensityScale=0.8
+
+[FoliageQuality@2]
+grass.densityScale=1.0
+grass.CullDistanceScale=1.0
+foliage.DensityScale=1.0
+
+[FoliageQuality@3]
+grass.densityScale=1.0
+grass.CullDistanceScale=1.0
+foliage.DensityScale=1.0
+
+; ------------------------------------------------------------------------------
+; View distance settings
+; ------------------------------------------------------------------------------
+
+[ViewDistanceQuality@0]
+r.ViewDistanceScale=0.70
+
+[ViewDistanceQuality@1]
+r.ViewDistanceScale=1.0
+
+[ViewDistanceQuality@2]
+r.ViewDistanceScale=1.0
+
+[ViewDistanceQuality@3]
+r.ViewDistanceScale=1.0
+
+; ------------------------------------------------------------------------------
+; Shadow settings
+; ------------------------------------------------------------------------------
+
+[ShadowQuality@0]
+r.ShadowQuality=0
+r.Shadow.MaxCSMResolution=1024
+r.Shadow.DistanceScale=0.5
+r.Shadow.CSM.MaxCascades=1
+r.Shadow.RadiusThreshold=0.06
+
+[ShadowQuality@1]
+r.ShadowQuality=1
+r.Shadow.MaxCSMResolution=1024
+r.Shadow.DistanceScale=0.7
+r.Shadow.CSM.MaxCascades=1
+r.Shadow.RadiusThreshold=0.06
+
+[ShadowQuality@2]
+r.ShadowQuality=1
+r.Shadow.MaxCSMResolution=2048
+r.Shadow.DistanceScale=1.0
+r.Shadow.CSM.MaxCascades=2
+r.Shadow.RadiusThreshold=0.06
+
+[ShadowQuality@3]
+r.ShadowQuality=2
+r.Shadow.MaxCSMResolution=2048
+r.Shadow.DistanceScale=1.0
+r.Shadow.CSM.MaxCascades=2
+r.Shadow.RadiusThreshold=0.06
+
+; ------------------------------------------------------------------------------
+; Postprocess settings
+; ------------------------------------------------------------------------------
+
+[PostProcessQuality@0]
+r.Tonemapper.Quality=1
+r.BloomQuality=0
+r.DepthOfFieldQuality=0
+r.LightShaftQuality=0
+r.RefractionQuality=0
+Slate.ForceBackgroundBlurLowQualityOverride=1
+
+[PostProcessQuality@1]
+r.Tonemapper.Quality=2
+r.BloomQuality=1
+r.DepthOfFieldQuality=0
+r.LightShaftQuality=0
+r.RefractionQuality=0
+Slate.ForceBackgroundBlurLowQualityOverride=1
+
+[PostProcessQuality@2]
+r.Tonemapper.Quality=3
+r.BloomQuality=1
+r.DepthOfFieldQuality=0
+r.LightShaftQuality=0
+r.RefractionQuality=1
+Slate.ForceBackgroundBlurLowQualityOverride=1
+
+[PostProcessQuality@3]
+r.Tonemapper.Quality=3
+r.BloomQuality=1
+r.DepthOfFieldQuality=1
+r.LightShaftQuality=1
+r.RefractionQuality=1
+Slate.ForceBackgroundBlurLowQualityOverride=1
+
 ; ------------------------------------------------------------------------------
 ; Effects settings
 ; ------------------------------------------------------------------------------
@@ -45,3 +152,27 @@ r.SkyAtmosphere.FastSkyLUT.SampleCountMax=8.0
 r.SkyAtmosphere.FastSkyLUT.Width=96
 r.SkyAtmosphere.FastSkyLUT.Height=50
 r.SkyAtmosphere.TransmittanceLUT.UseSmallFormat=1
+
+; ------------------------------------------------------------------------------
+; Texture settings
+; ------------------------------------------------------------------------------
+
+[TextureQuality@0]
+r.MaxAnisotropy=1
+
+[TextureQuality@1]
+r.MaxAnisotropy=1
+
+[TextureQuality@2]
+r.MaxAnisotropy=2
+
+[TextureQuality@3]
+r.MaxAnisotropy=4
+
+[ScalabilityPreview]
+; don't want to preview any texture streaming changes in the editor
+BlacklistCVars=r.MaxAnisotropy
++BlacklistCVars=r.Streaming.Boost
++BlacklistCVars=r.Streaming.AmortizeCPUToGPUCopy
++BlacklistCVars=r.Streaming.MaxNumTexturesToStreamPerFrame
++BlacklistCVars=r.Streaming.PoolSize
\ No newline at end of file
diff --git a/Engine/Plugins/Compositing/Composure/Source/Composure/Private/MovieScene/MovieSceneComposureExportSectionTemplate.cpp b/Engine/Plugins/Compositing/Composure/Source/Composure/Private/MovieScene/MovieSceneComposureExportSectionTemplate.cpp
--- a/Engine/Plugins/Compositing/Composure/Source/Composure/Private/MovieScene/MovieSceneComposureExportSectionTemplate.cpp
+++ b/Engine/Plugins/Compositing/Composure/Source/Composure/Private/MovieScene/MovieSceneComposureExportSectionTemplate.cpp
@@ -31,7 +31,7 @@ struct FBufferVisualizationIterator
 	/** Reference to the post processing settings to add vialization materials to */
 	FFinalPostProcessSettings& FinalPostProcessSettings;
 
-	void ProcessValue(const FString& InName, UMaterial* Material, const FText& InText)
+	void ProcessValue(const FString& InName, UMaterialInterface* Material, const FText& InText)
 	{
 		if (BuffersToCapture.Contains(InName) || BuffersToCapture.Contains(InText.ToString()))
 		{
diff --git a/Engine/Shaders/Private/BasePassPixelShader.usf b/Engine/Shaders/Private/BasePassPixelShader.usf
--- a/Engine/Shaders/Private/BasePassPixelShader.usf
+++ b/Engine/Shaders/Private/BasePassPixelShader.usf
@@ -14,7 +14,8 @@
 	#define EyeAdaptationStruct OpaqueBasePass
 #endif
 
-#if MATERIAL_SHADINGMODEL_SINGLELAYERWATER
+// Material setting to allow forward shading (including mobile) to use preintegrated GF lut for simple IBL.
+#if MATERIAL_SHADINGMODEL_SINGLELAYERWATER || (FORWARD_SHADING && MATERIAL_USE_PREINTEGRATED_GF)
 #define PreIntegratedGF			OpaqueBasePass.PreIntegratedGFTexture
 #if SUPPORTS_INDEPENDENT_SAMPLERS
 #define PreIntegratedGFSampler	View.SharedBilinearClampedSampler
diff --git a/Engine/Shaders/Private/ForwardLightingCommon.ush b/Engine/Shaders/Private/ForwardLightingCommon.ush
--- a/Engine/Shaders/Private/ForwardLightingCommon.ush
+++ b/Engine/Shaders/Private/ForwardLightingCommon.ush
@@ -398,7 +398,11 @@ half3 GetImageBasedReflectionLighting(FMaterialPixelParameters MaterialParameter
 	}
 #endif
 
+#if MATERIAL_USE_PREINTEGRATED_GF
+	SpecularColor = EnvBRDF(SpecularColor, Roughness, NoV);
+#else
 	SpecularColor = EnvBRDFApprox(SpecularColor, Roughness, NoV);
+#endif
 
 	float3 SpecularLighting = SpecularIBL.rgb;
 
diff --git a/Engine/Shaders/Private/MobileBasePassPixelShader.usf b/Engine/Shaders/Private/MobileBasePassPixelShader.usf
--- a/Engine/Shaders/Private/MobileBasePassPixelShader.usf
+++ b/Engine/Shaders/Private/MobileBasePassPixelShader.usf
@@ -9,6 +9,18 @@
 // Reroute MobileSceneTextures uniform buffer references to the base pass uniform buffer
 #define MobileSceneTextures MobileBasePass.SceneTextures
 
+#ifndef MOBILE_QL_FORCE_DISABLE_PREINTEGRATEDGF
+#define MOBILE_QL_FORCE_DISABLE_PREINTEGRATEDGF 0
+#endif
+
+#define FORWARDSHADING_USE_HQ_ENV_BRDF (MATERIAL_USE_PREINTEGRATED_GF && !MOBILE_QL_FORCE_DISABLE_PREINTEGRATEDGF)
+
+//use preintegrated GF lut for simple IBL
+#if FORWARDSHADING_USE_HQ_ENV_BRDF
+#define PreIntegratedGF			MobileBasePass.PreIntegratedGFTexture
+#define PreIntegratedGFSampler	MobileBasePass.PreIntegratedGFSampler
+#endif
+
 #include "SHCommon.ush"
 #include "MobileBasePassCommon.ush"
 #include "/Engine/Generated/Material.ush"
@@ -462,10 +474,14 @@ void Main(
 	// then defaults are the same as a non-metal,
 	// so this define is safe.
 	SpecularColor = EnvBRDFApproxNonmetal( Roughness, NoV );
+#else
+#if FORWARDSHADING_USE_HQ_ENV_BRDF
+	SpecularColor = EnvBRDF(SpecularColor, Roughness, NoV);
 #else
 	SpecularColor = EnvBRDFApprox( SpecularColor, Roughness, NoV );
 #endif
 #endif
+#endif
 
 #if MOBILE_EMULATION
 	{
@@ -634,7 +650,7 @@ void Main(
 		#if ENABLE_SKY_LIGHT
 			//@mw todo
 			// TODO: Also need to do specular.
-#if MATERIAL_TWO_SIDED && LQ_TEXTURE_LIGHTMAP
+#if MATERIAL_TWOSIDED && LQ_TEXTURE_LIGHTMAP
 			if (NoL == 0)
 			{
 #endif
@@ -644,7 +660,7 @@ void Main(
 			WaterDiffuseIndirectLuminance += SkyDiffuseLighting;
 		#endif
 			Color += SkyDiffuseLighting * half3(ResolvedView.SkyLightColor.rgb) * DiffuseColor * MaterialAO;
-#if MATERIAL_TWO_SIDED && LQ_TEXTURE_LIGHTMAP
+#if MATERIAL_TWOSIDED && LQ_TEXTURE_LIGHTMAP
 			
 			}
 #endif
@@ -717,10 +733,10 @@ void Main(
 
 #if MATERIAL_SHADINGMODEL_SINGLELAYERWATER
 	{
-		const bool CameraIsUnderWater = false;												// Fade out the material contribution over to water contribution according to material opacity.
+		const bool CameraIsUnderWater = false;	// Fade out the material contribution over to water contribution according to material opacity.
 		const float3 SunIlluminance = ResolvedView.DirectionalLightColor.rgb * PI;			// times PI because it is divided by PI on CPU (=luminance) and we want illuminance here. 
 		const float3 WaterDiffuseIndirectIlluminance = WaterDiffuseIndirectLuminance * PI;	// DiffuseIndirectLighting is luminance. So we need to multiply by PI to get illuminance.
-		const float3 EnvBrdf = SpecularColor;												// SpecularColor is not F0 as in BasePassPixelShader, it is EnvBRDF.
+		const float3 EnvBrdf = SpecularColor; // SpecularColor is not F0 as in BasePassPixelShader, it is EnvBRDF.
 		const uint EyeIndex = 0;
 
 		const float4 NullDistortionParams = 1.0f;
diff --git a/Engine/Source/Editor/CommonMenuExtensions/Private/BufferVisualizationMenuCommands.cpp b/Engine/Source/Editor/CommonMenuExtensions/Private/BufferVisualizationMenuCommands.cpp
--- a/Engine/Source/Editor/CommonMenuExtensions/Private/BufferVisualizationMenuCommands.cpp
+++ b/Engine/Source/Editor/CommonMenuExtensions/Private/BufferVisualizationMenuCommands.cpp
@@ -16,7 +16,7 @@ namespace
 {
 	struct FMaterialIteratorWithLambda
 	{
-		typedef TFunction<void(const FString&, const UMaterial*, const FText&)> LambdaFunctionType;
+		typedef TFunction<void(const FString&, const UMaterialInterface*, const FText&)> LambdaFunctionType;
 		LambdaFunctionType Lambda;
 
 		inline FMaterialIteratorWithLambda(const LambdaFunctionType& InLambda)
@@ -24,7 +24,7 @@ namespace
 		{
 		}
 
-		inline void ProcessValue(const FString& InMaterialName, const UMaterial* InMaterial, const FText& InDisplayName)
+		inline void ProcessValue(const FString& InMaterialName, const UMaterialInterface* InMaterial, const FText& InDisplayName)
 		{
 			Lambda(InMaterialName, InMaterial, InDisplayName);
 		}
@@ -74,7 +74,7 @@ void FBufferVisualizationMenuCommands::CreateOverviewCommand()
 
 void FBufferVisualizationMenuCommands::CreateVisualizationCommands()
 {
-	FMaterialIteratorWithLambda Iterator([this](const FString& InMaterialName, const UMaterial* InMaterial, const FText& InDisplayName)
+	FMaterialIteratorWithLambda Iterator([this](const FString& InMaterialName, const UMaterialInterface* InMaterial, const FText& InDisplayName)
 	{
 		const FName CommandName = NameFromMaterial(InMaterialName);
 
@@ -115,7 +115,7 @@ void FBufferVisualizationMenuCommands::AddVisualizationCommandsToMenu(FMenuBuild
 
 	const TBufferVisualizationModeCommandMap& Commands = CommandMap;
 
-	FMaterialIteratorWithLambda Iterator([&Menu, &Commands](const FString& InMaterialName, const UMaterial* InMaterial, const FText& InDisplayName)
+	FMaterialIteratorWithLambda Iterator([&Menu, &Commands](const FString& InMaterialName, const UMaterialInterface* InMaterial, const FText& InDisplayName)
 	{
 		const FName CommandName = NameFromMaterial(InMaterialName);
 
diff --git a/Engine/Source/Editor/DetailCustomizations/Private/RenderPassesCustomization.cpp b/Engine/Source/Editor/DetailCustomizations/Private/RenderPassesCustomization.cpp
--- a/Engine/Source/Editor/DetailCustomizations/Private/RenderPassesCustomization.cpp
+++ b/Engine/Source/Editor/DetailCustomizations/Private/RenderPassesCustomization.cpp
@@ -83,7 +83,7 @@ public:
 
 			FIterator(TMap<FString, FText>& InRenderPasses) : RenderPasses(InRenderPasses) {}
 
-			void ProcessValue(const FString& InName, UMaterial* Material, const FText& InText)
+			void ProcessValue(const FString& InName, UMaterialInterface* Material, const FText& InText)
 			{
 				RenderPasses.Add(InName, InText);
 			}
diff --git a/Engine/Source/Runtime/Engine/Classes/Materials/Material.h b/Engine/Source/Runtime/Engine/Classes/Materials/Material.h
--- a/Engine/Source/Runtime/Engine/Classes/Materials/Material.h
+++ b/Engine/Source/Runtime/Engine/Classes/Materials/Material.h
@@ -767,6 +767,10 @@ public:
 	UPROPERTY(EditAnywhere, BlueprintReadOnly, Category=Mobile)
 	uint8 bUseLightmapDirectionality : 1;
 
+	/* Forward (including mobile) renderer: use preintegrated GF lut for simple IBL, but will use one more sampler. */
+	UPROPERTY(EditAnywhere, BlueprintReadOnly, Category = ForwardShading, meta = (DisplayName = "PreintegratedGF For Simple IBL"))
+	uint32 bForwardRenderUsePreintegratedGFForSimpleIBL : 1;
+
 	/* 
 	 * Forward renderer: enables multiple parallax-corrected reflection captures that blend together.
 	 * Mobile renderer: blend between nearest 3 reflection captures, but reduces the number of samplers available to the material as two more samplers will be used for reflection cubemaps.
diff --git a/Engine/Source/Runtime/Engine/Private/BufferVisualizationData.cpp b/Engine/Source/Runtime/Engine/Private/BufferVisualizationData.cpp
--- a/Engine/Source/Runtime/Engine/Private/BufferVisualizationData.cpp
+++ b/Engine/Source/Runtime/Engine/Private/BufferVisualizationData.cpp
@@ -39,7 +39,7 @@ void FBufferVisualizationData::Initialize()
 					FString MaterialName;
 					if( FParse::Value( *It.Value().GetValue(), TEXT("Material="), MaterialName, true ) )
 					{
-						UMaterial* Material = LoadObject<UMaterial>(NULL, *MaterialName);
+						UMaterialInterface* Material = LoadObject<UMaterialInterface>(NULL, *MaterialName);
 			
 						if (Material)
 						{
@@ -73,7 +73,7 @@ void FBufferVisualizationData::ConfigureConsoleCommand()
 		{
 		}
 
-		void ProcessValue(const FString& InMaterialName, const UMaterial* InMaterial, const FText& InDisplayName)
+		void ProcessValue(const FString& InMaterialName, const UMaterialInterface* InMaterial, const FText& InDisplayName)
 		{
 			Message += FString(TEXT("\n  "));
 			Message += InMaterialName;
@@ -102,7 +102,7 @@ void FBufferVisualizationData::ConfigureConsoleCommand()
 		ECVF_Default);
 }
 
-UMaterial* FBufferVisualizationData::GetMaterial(FName InMaterialName)
+UMaterialInterface* FBufferVisualizationData::GetMaterial(FName InMaterialName)
 {
 	Record* Result = MaterialMap.Find(InMaterialName);
 	if (Result)
@@ -125,7 +125,7 @@ bool FBufferVisualizationData::IsDifferentToCurrentOverviewMaterialNames(const F
 	return InNameList != CurrentOverviewMaterialNames;
 }
 
-TArray<UMaterial*>& FBufferVisualizationData::GetOverviewMaterials()
+TArray<UMaterialInterface*>& FBufferVisualizationData::GetOverviewMaterials()
 {
 	return OverviewMaterials;
 }
diff --git a/Engine/Source/Runtime/Engine/Private/DebugCameraController.cpp b/Engine/Source/Runtime/Engine/Private/DebugCameraController.cpp
--- a/Engine/Source/Runtime/Engine/Private/DebugCameraController.cpp
+++ b/Engine/Source/Runtime/Engine/Private/DebugCameraController.cpp
@@ -960,7 +960,7 @@ FString ADebugCameraController::GetBufferMaterialName(const FString& InBufferNam
 {
 	if (!InBufferName.IsEmpty())
 	{
-		if (UMaterial* Material = GetBufferVisualizationData().GetMaterial(*InBufferName))
+		if (UMaterialInterface* Material = GetBufferVisualizationData().GetMaterial(*InBufferName))
 		{
 			return Material->GetName();
 		}
diff --git a/Engine/Source/Runtime/Engine/Private/Materials/HLSLMaterialTranslator.cpp b/Engine/Source/Runtime/Engine/Private/Materials/HLSLMaterialTranslator.cpp
--- a/Engine/Source/Runtime/Engine/Private/Materials/HLSLMaterialTranslator.cpp
+++ b/Engine/Source/Runtime/Engine/Private/Materials/HLSLMaterialTranslator.cpp
@@ -1160,7 +1160,6 @@ void FHLSLMaterialTranslator::GetMaterialEnvironment(EShaderPlatform InPlatform,
 	OutEnvironment.SetDefine(TEXT("MATERIAL_IS_SKY"), Material->IsSky());
 	OutEnvironment.SetDefine(TEXT("MATERIAL_COMPUTE_FOG_PER_PIXEL"), Material->ComputeFogPerPixel());
 	OutEnvironment.SetDefine(TEXT("MATERIAL_FULLY_ROUGH"), bIsFullyRough || Material->IsFullyRough());
-	OutEnvironment.SetDefine(TEXT("MATERIAL_TWO_SIDED"), Material->IsTwoSided());
 
 	// Count the number of VTStacks (each stack will allocate a feedback slot)
 	OutEnvironment.SetDefine(TEXT("NUM_VIRTUALTEXTURE_SAMPLES"), VTStacks.Num());
@@ -5237,21 +5236,21 @@ int32 FHLSLMaterialTranslator::VertexInterpolator(uint32 InterpolatorIndex)
 	const TCHAR* TypeName = HLSLTypeString(Interpolator->InterpolatedType);
 	const TCHAR* Swizzle[2] = { TEXT("x"), TEXT("y") };
 	const int32 Offset = Interpolator->InterpolatorOffset;
-
+	
 	// Note: We reference the UV define directly to avoid having to pre-accumulate UV counts before property translation
-	FString GetValueCode = FString::Printf(TEXT("%s(Parameters.TexCoords[VERTEX_INTERPOLATOR_%i_TEXCOORDS_X].%s"), TypeName, InterpolatorIndex, Swizzle[Offset % 2]);
+	FString GetValueCode = FString::Printf(TEXT("%s(Parameters.TexCoords[VERTEX_INTERPOLATOR_%i_TEXCOORDS_X].%s"), TypeName, InterpolatorIndex, Swizzle[Offset%2]);
 	if (InterpolatorSize >= 2)
 	{
-		GetValueCode += FString::Printf(TEXT(", Parameters.TexCoords[VERTEX_INTERPOLATOR_%i_TEXCOORDS_Y].%s"), InterpolatorIndex, Swizzle[(Offset + 1) % 2]);
+		GetValueCode += FString::Printf(TEXT(", Parameters.TexCoords[VERTEX_INTERPOLATOR_%i_TEXCOORDS_Y].%s"), InterpolatorIndex, Swizzle[(Offset+1)%2]);
 
 		if (InterpolatorSize >= 3)
 		{
-			GetValueCode += FString::Printf(TEXT(", Parameters.TexCoords[VERTEX_INTERPOLATOR_%i_TEXCOORDS_Z].%s"), InterpolatorIndex, Swizzle[(Offset + 2) % 2]);
+			GetValueCode += FString::Printf(TEXT(", Parameters.TexCoords[VERTEX_INTERPOLATOR_%i_TEXCOORDS_Z].%s"), InterpolatorIndex, Swizzle[(Offset+2)%2]);
 
 			if (InterpolatorSize >= 4)
 			{
 				check(InterpolatorSize == 4);
-				GetValueCode += FString::Printf(TEXT(", Parameters.TexCoords[VERTEX_INTERPOLATOR_%i_TEXCOORDS_W].%s"), InterpolatorIndex, Swizzle[(Offset + 3) % 2]);
+				GetValueCode += FString::Printf(TEXT(", Parameters.TexCoords[VERTEX_INTERPOLATOR_%i_TEXCOORDS_W].%s"), InterpolatorIndex, Swizzle[(Offset+3)%2]);
 			}
 		}
 	}
diff --git a/Engine/Source/Runtime/Engine/Private/Materials/MaterialShared.cpp b/Engine/Source/Runtime/Engine/Private/Materials/MaterialShared.cpp
--- a/Engine/Source/Runtime/Engine/Private/Materials/MaterialShared.cpp
+++ b/Engine/Source/Runtime/Engine/Private/Materials/MaterialShared.cpp
@@ -1251,6 +1251,11 @@ bool FMaterialResource::IsUsingFullPrecision() const
 	return Material->bUseFullPrecision;
 }
 
+bool FMaterialResource::IsUsingPreintegratedGFForSimpleIBL() const
+{
+	return Material->bForwardRenderUsePreintegratedGFForSimpleIBL;
+}
+
 bool FMaterialResource::IsUsingHQForwardReflections() const
 {
 	return Material->bUseHQForwardReflections;
@@ -1665,6 +1670,7 @@ void FMaterial::SetupMaterialEnvironment(
 	OutEnvironment.SetDefine(TEXT("MATERIAL_TANGENTSPACENORMAL"), IsTangentSpaceNormal());
 	OutEnvironment.SetDefine(TEXT("GENERATE_SPHERICAL_PARTICLE_NORMALS"),ShouldGenerateSphericalParticleNormals());
 	OutEnvironment.SetDefine(TEXT("MATERIAL_USES_SCENE_COLOR_COPY"), RequiresSceneColorCopy_GameThread());
+	OutEnvironment.SetDefine(TEXT("MATERIAL_USE_PREINTEGRATED_GF"), IsUsingPreintegratedGFForSimpleIBL());
 	OutEnvironment.SetDefine(TEXT("MATERIAL_HQ_FORWARD_REFLECTIONS"), IsUsingHQForwardReflections());
 	OutEnvironment.SetDefine(TEXT("MATERIAL_PLANAR_FORWARD_REFLECTIONS"), IsUsingPlanarForwardReflections());
 	OutEnvironment.SetDefine(TEXT("MATERIAL_NONMETAL"), IsNonmetal());
diff --git a/Engine/Source/Runtime/Engine/Private/SceneView.cpp b/Engine/Source/Runtime/Engine/Private/SceneView.cpp
--- a/Engine/Source/Runtime/Engine/Private/SceneView.cpp
+++ b/Engine/Source/Runtime/Engine/Private/SceneView.cpp
@@ -2083,7 +2083,7 @@ void FSceneView::ConfigureBufferVisualizationSettings()
 
 				// Lookup this material from the list that was parsed out of the global ini file
 				Left.TrimStartInline();
-				UMaterial* Material = BufferVisualizationData.GetMaterial(*Left);
+				UMaterialInterface* Material = BufferVisualizationData.GetMaterial(*Left);
 
 				if (Material == NULL && Left.Len() > 0)
 				{
@@ -2099,7 +2099,7 @@ void FSceneView::ConfigureBufferVisualizationSettings()
 		}
 
 		// Copy current material list into settings material list
-		for (TArray<UMaterial*>::TConstIterator It = BufferVisualizationData.GetOverviewMaterials().CreateConstIterator(); It; ++It)
+		for (TArray<UMaterialInterface*>::TConstIterator It = BufferVisualizationData.GetOverviewMaterials().CreateConstIterator(); It; ++It)
 		{
 			FinalPostProcessSettings.BufferVisualizationOverviewMaterials.Add(*It);
 		}
diff --git a/Engine/Source/Runtime/Engine/Public/BufferVisualizationData.h b/Engine/Source/Runtime/Engine/Public/BufferVisualizationData.h
--- a/Engine/Source/Runtime/Engine/Public/BufferVisualizationData.h
+++ b/Engine/Source/Runtime/Engine/Public/BufferVisualizationData.h
@@ -3,7 +3,7 @@
 
 #include "CoreMinimal.h"
 
-class UMaterial;
+class UMaterialInterface;
 
 class FBufferVisualizationData
 {
@@ -22,14 +22,14 @@ public:
 	bool IsInitialized() const { return bIsInitialized; }
 
 	/** Get a named material from the available material map **/
-	ENGINE_API UMaterial* GetMaterial(FName InMaterialName);
+	ENGINE_API UMaterialInterface* GetMaterial(FName InMaterialName);
 
 	/** We cache the overview material name list from the console command here, so all dynamically created views can re-use the existing cached list of materials */
 	void SetCurrentOverviewMaterialNames(const FString& InNameList);
 	bool IsDifferentToCurrentOverviewMaterialNames(const FString& InNameList);
 
 	/** Access the list of materials currently in use by the buffer visualization overview */
-	TArray<UMaterial*>& GetOverviewMaterials();
+	TArray<UMaterialInterface*>& GetOverviewMaterials();
 
 	/** Iterator function for iterating over available materials */
 	template<class T> void IterateOverAvailableMaterials(T& Iterator) const
@@ -54,7 +54,7 @@ private:
 	{
 		FString Name;
 		FText DisplayName;
-		UMaterial* Material;
+		UMaterialInterface* Material;
 	};
 
 	/** Mapping of FName (first parameter in ini file material list) to a material record */
@@ -67,7 +67,7 @@ private:
 	FString CurrentOverviewMaterialNames;
 
 	/** List of material currently in use by the buffer visualization overview */
-	TArray<UMaterial*> OverviewMaterials;
+	TArray<UMaterialInterface*> OverviewMaterials;
 	
 	/** Storage for console variable documentation strings **/
 	FString ConsoleDocumentationVisualizationMode;
diff --git a/Engine/Source/Runtime/Engine/Public/MaterialShared.h b/Engine/Source/Runtime/Engine/Public/MaterialShared.h
--- a/Engine/Source/Runtime/Engine/Public/MaterialShared.h
+++ b/Engine/Source/Runtime/Engine/Public/MaterialShared.h
@@ -1475,6 +1475,7 @@ public:
 	virtual bool IsFullyRough() const { return false; }
 	virtual bool UseNormalCurvatureToRoughness() const { return false; }
 	virtual bool IsUsingFullPrecision() const { return false; }
+	virtual bool IsUsingPreintegratedGFForSimpleIBL() const { return false; }
 	virtual bool IsUsingHQForwardReflections() const { return false; }
 	virtual bool IsUsingPlanarForwardReflections() const { return false; }
 	virtual bool IsNonmetal() const { return false; }
@@ -2239,6 +2240,7 @@ public:
 	ENGINE_API virtual bool IsFullyRough() const override;
 	ENGINE_API virtual bool UseNormalCurvatureToRoughness() const override;
 	ENGINE_API virtual bool IsUsingFullPrecision() const override;
+	ENGINE_API virtual bool IsUsingPreintegratedGFForSimpleIBL() const override;
 	ENGINE_API virtual bool IsUsingHQForwardReflections() const override;
 	ENGINE_API virtual bool IsUsingPlanarForwardReflections() const override;
 	ENGINE_API virtual bool IsNonmetal() const override;
diff --git a/Engine/Source/Runtime/Landscape/Private/LandscapeEdit.cpp b/Engine/Source/Runtime/Landscape/Private/LandscapeEdit.cpp
--- a/Engine/Source/Runtime/Landscape/Private/LandscapeEdit.cpp
+++ b/Engine/Source/Runtime/Landscape/Private/LandscapeEdit.cpp
@@ -6450,7 +6450,7 @@ void ULandscapeComponent::GeneratePlatformVertexData(const ITargetPlatform* Targ
 						FColor* CurrentMipSrcRow = HeightmapMipData[OcclusionMeshMip] + (CurrentMipOfsY + y) * MipSizeX + CurrentMipOfsX;
 						uint16 Height = CurrentMipSrcRow[x].R << 8 | CurrentMipSrcRow[x].G;
 
-						FVector VtxPos = FVector(x*MipRatio, y*MipRatio, ((float)Height - 32768.f) * LANDSCAPE_ZSCALE);
+						FVector VtxPos = FVector(x*MipRatio + SubX * SubsectionSizeQuads, y*MipRatio + SubY * SubsectionSizeQuads, ((float)Height - 32768.f) * LANDSCAPE_ZSCALE);
 						OccluderVertices.Add(VtxPos);
 					}
 				}
diff --git a/Engine/Source/Runtime/Landscape/Private/LandscapeRender.cpp b/Engine/Source/Runtime/Landscape/Private/LandscapeRender.cpp
--- a/Engine/Source/Runtime/Landscape/Private/LandscapeRender.cpp
+++ b/Engine/Source/Runtime/Landscape/Private/LandscapeRender.cpp
@@ -3018,33 +3018,42 @@ void FLandscapeSharedBuffers::CreateOccluderIndexBuffer(int32 NumOccluderVertice
 		return;
 	}
 
-	uint16 NumLineQuads = ((uint16)FMath::Sqrt(NumOccluderVertices) - 1);
-	uint16 NumLineVtx = NumLineQuads + 1;
-	check(NumLineVtx*NumLineVtx == NumOccluderVertices);
+	const uint16 NumLineQuads = ((uint16)(FMath::Sqrt(NumOccluderVertices)/NumSubsections) - 1);
+	const uint16 NumLineVtx = NumLineQuads + 1;
+	check(NumLineVtx*NumLineVtx*NumSubsections*NumSubsections == NumOccluderVertices);
 
-	int32 NumTris = NumLineQuads*NumLineQuads * 2;
-	int32 NumIndices = NumTris * 3;
+	const int32 NumTrisPerSubsection = NumLineQuads*NumLineQuads * 2;
+	const int32 NumIndicesPerSubsection = NumTrisPerSubsection * 3;
 	OccluderIndicesSP = MakeShared<FOccluderIndexArray, ESPMode::ThreadSafe>();
-	OccluderIndicesSP->SetNumUninitialized(NumIndices, false);
+	OccluderIndicesSP->SetNumUninitialized(NumIndicesPerSubsection*NumSubsections*NumSubsections, false);
 
 	uint16* OcclusionIndices = OccluderIndicesSP->GetData();
 	const uint16 NumLineVtxPlusOne = NumLineVtx + 1;
 	const uint16 QuadIndices[2][3] = { {0, NumLineVtx, NumLineVtxPlusOne}, {0, NumLineVtxPlusOne, 1} };
 	uint16 QuadOffset = 0;
 	int32 Index = 0;
-	for (int32 y = 0; y < NumLineQuads; y++)
+
+	for (int32 SubY = 0; SubY < NumSubsections; SubY++)
 	{
-		for (int32 x = 0; x < NumLineQuads; x++)
+		for (int32 SubX = 0; SubX < NumSubsections; SubX++)
 		{
-			for (int32 i = 0; i < 2; i++)
-			{
-				OcclusionIndices[Index++] = QuadIndices[i][0] + QuadOffset;
-				OcclusionIndices[Index++] = QuadIndices[i][1] + QuadOffset;
-				OcclusionIndices[Index++] = QuadIndices[i][2] + QuadOffset;
-			}
-			QuadOffset++;
+			QuadOffset = (SubY*NumSubsections + SubX) * NumLineVtx * NumLineVtx;
+
+	        for (int32 y = 0; y < NumLineQuads; y++)
+	        {
+		        for (int32 x = 0; x < NumLineQuads; x++)
+		        {
+			        for (int32 i = 0; i < 2; i++)
+			        {
+				        OcclusionIndices[Index++] = QuadIndices[i][0] + QuadOffset;
+				        OcclusionIndices[Index++] = QuadIndices[i][1] + QuadOffset;
+				        OcclusionIndices[Index++] = QuadIndices[i][2] + QuadOffset;
+			        }
+			        QuadOffset++;
+		        }
+		        QuadOffset++;
+	        }
 		}
-		QuadOffset++;
 	}
 
 	INC_DWORD_STAT_BY(STAT_LandscapeOccluderMem, OccluderIndicesSP->GetAllocatedSize());
diff --git a/Engine/Source/Runtime/MaterialShaderQualitySettings/Classes/ShaderPlatformQualitySettings.h b/Engine/Source/Runtime/MaterialShaderQualitySettings/Classes/ShaderPlatformQualitySettings.h
--- a/Engine/Source/Runtime/MaterialShaderQualitySettings/Classes/ShaderPlatformQualitySettings.h
+++ b/Engine/Source/Runtime/MaterialShaderQualitySettings/Classes/ShaderPlatformQualitySettings.h
@@ -36,6 +36,7 @@ public:
 		, bForceNonMetal(false)
 		, bForceDisableLMDirectionality(false)
 		, bForceLQReflections(false)
+		, bForceDisablePreintegratedGF(false)
 		, bDisableMaterialNormalCalculation(false)
 		, MobileCSMQuality(EMobileCSMQuality::PCF_2x2)
 	{
@@ -59,6 +60,9 @@ public:
 	UPROPERTY(EditAnywhere, Config, Meta = (DisplayName = "Force low quality reflections"), Category = "Quality")
 	bool bForceLQReflections;
 
+	UPROPERTY(EditAnywhere, Config, Meta = (DisplayName = "Force not use preintegrated GF for simple IBL"), Category = "Quality")
+	bool bForceDisablePreintegratedGF;
+
 	UPROPERTY(EditAnywhere, Config, Meta = (DisplayName = "Disable material normal calculation"), Category = "Quality")
 	bool bDisableMaterialNormalCalculation;
 
diff --git a/Engine/Source/Runtime/MaterialShaderQualitySettings/Private/MaterialShaderQualitySettings.cpp b/Engine/Source/Runtime/MaterialShaderQualitySettings/Private/MaterialShaderQualitySettings.cpp
--- a/Engine/Source/Runtime/MaterialShaderQualitySettings/Private/MaterialShaderQualitySettings.cpp
+++ b/Engine/Source/Runtime/MaterialShaderQualitySettings/Private/MaterialShaderQualitySettings.cpp
@@ -147,6 +147,7 @@ bool FMaterialQualityOverrides::HasAnyOverridesSet() const
 		|| bForceFullyRough != DefaultOverrides.bForceFullyRough
 		|| bForceNonMetal != DefaultOverrides.bForceNonMetal
 		|| bForceLQReflections != DefaultOverrides.bForceLQReflections
+		|| bForceDisablePreintegratedGF != DefaultOverrides.bForceDisablePreintegratedGF
 		|| bDisableMaterialNormalCalculation != DefaultOverrides.bDisableMaterialNormalCalculation
 		|| bDiscardQualityDuringCook != DefaultOverrides.bDiscardQualityDuringCook;
 }
diff --git a/Engine/Source/Runtime/MovieSceneCapture/Private/CompositionGraphCaptureProtocol.cpp b/Engine/Source/Runtime/MovieSceneCapture/Private/CompositionGraphCaptureProtocol.cpp
--- a/Engine/Source/Runtime/MovieSceneCapture/Private/CompositionGraphCaptureProtocol.cpp
+++ b/Engine/Source/Runtime/MovieSceneCapture/Private/CompositionGraphCaptureProtocol.cpp
@@ -110,7 +110,7 @@ struct FFrameCaptureViewExtension : public FSceneViewExtensionBase
 				: FinalPostProcessSettings(InFinalPostProcessSettings), RenderPasses(InRenderPasses)
 			{}
 
-			void ProcessValue(const FString& InName, UMaterial* Material, const FText& InText)
+			void ProcessValue(const FString& InName, UMaterialInterface* Material, const FText& InText)
 			{
 				if (!RenderPasses.Num() || RenderPasses.Contains(InName) || RenderPasses.Contains(InText.ToString()))
 				{
diff --git a/Engine/Source/Runtime/Renderer/Private/MobileBasePassRendering.cpp b/Engine/Source/Runtime/Renderer/Private/MobileBasePassRendering.cpp
--- a/Engine/Source/Runtime/Renderer/Private/MobileBasePassRendering.cpp
+++ b/Engine/Source/Runtime/Renderer/Private/MobileBasePassRendering.cpp
@@ -99,6 +99,7 @@ bool TMobileBasePassPSPolicyParamType<LightMapPolicyType>::ModifyCompilationEnvi
 	OutEnvironment.SetDefine(TEXT("MOBILE_QL_FORCE_NONMETAL"), QualityOverrides.bEnableOverride && QualityOverrides.bForceNonMetal != 0 ? 1u : 0u);
 	OutEnvironment.SetDefine(TEXT("QL_FORCEDISABLE_LM_DIRECTIONALITY"), QualityOverrides.bEnableOverride && QualityOverrides.bForceDisableLMDirectionality != 0 ? 1u : 0u);
 	OutEnvironment.SetDefine(TEXT("MOBILE_QL_FORCE_LQ_REFLECTIONS"), QualityOverrides.bEnableOverride && QualityOverrides.bForceLQReflections != 0 ? 1u : 0u);
+	OutEnvironment.SetDefine(TEXT("MOBILE_QL_FORCE_DISABLE_PREINTEGRATEDGF"), QualityOverrides.bEnableOverride && QualityOverrides.bForceDisablePreintegratedGF != 0 ? 1u : 0u);
 	OutEnvironment.SetDefine(TEXT("MOBILE_CSM_QUALITY"), (uint32)QualityOverrides.MobileCSMQuality);
 	OutEnvironment.SetDefine(TEXT("MOBILE_QL_DISABLE_MATERIAL_NORMAL"), QualityOverrides.bEnableOverride && QualityOverrides.bDisableMaterialNormalCalculation);
 	return true;
@@ -163,6 +164,9 @@ void SetupMobileBasePassUniformParameters(
 
 	FSceneRenderTargets& SceneContext = FSceneRenderTargets::Get(RHICmdList);
 	SetupMobileSceneTextureUniformParameters(SceneContext, View.FeatureLevel, bTranslucentPass, View.bUsesCustomDepthStencil, BasePassParameters.SceneTextures);
+
+	BasePassParameters.PreIntegratedGFTexture = GSystemTextures.PreintegratedGF->GetRenderTargetItem().ShaderResourceTexture;
+	BasePassParameters.PreIntegratedGFSampler = TStaticSamplerState<SF_Bilinear, AM_Clamp, AM_Clamp, AM_Clamp>::GetRHI();
 }
 
 void CreateMobileBasePassUniformBuffer(
diff --git a/Engine/Source/Runtime/Renderer/Private/MobileBasePassRendering.h b/Engine/Source/Runtime/Renderer/Private/MobileBasePassRendering.h
--- a/Engine/Source/Runtime/Renderer/Private/MobileBasePassRendering.h
+++ b/Engine/Source/Runtime/Renderer/Private/MobileBasePassRendering.h
@@ -28,6 +28,8 @@ BEGIN_GLOBAL_SHADER_PARAMETER_STRUCT(FMobileBasePassUniformParameters, )
 	SHADER_PARAMETER_STRUCT(FFogUniformParameters, Fog)
 	SHADER_PARAMETER_STRUCT(FPlanarReflectionUniformParameters, PlanarReflection) // Single global planar reflection for the forward pass.
 	SHADER_PARAMETER_STRUCT(FMobileSceneTextureUniformParameters, SceneTextures)
+	SHADER_PARAMETER_TEXTURE(Texture2D, PreIntegratedGFTexture)
+	SHADER_PARAMETER_SAMPLER(SamplerState, PreIntegratedGFSampler)
 END_GLOBAL_SHADER_PARAMETER_STRUCT()
 
 extern void SetupMobileBasePassUniformParameters(
diff --git a/Engine/Source/Runtime/Renderer/Private/PostProcess/PostProcessMaterial.cpp b/Engine/Source/Runtime/Renderer/Private/PostProcess/PostProcessMaterial.cpp
--- a/Engine/Source/Runtime/Renderer/Private/PostProcess/PostProcessMaterial.cpp
+++ b/Engine/Source/Runtime/Renderer/Private/PostProcess/PostProcessMaterial.cpp
@@ -574,7 +574,8 @@ FPostProcessMaterialChain GetPostProcessMaterialChain(const FViewInfo& View, EBl
 
 	if (ViewFamily.EngineShowFlags.VisualizeBuffer)
 	{
-		UMaterial* Material = GetBufferVisualizationData().GetMaterial(View.CurrentBufferVisualizationMode);
+		UMaterialInterface* VisMaterial = GetBufferVisualizationData().GetMaterial(View.CurrentBufferVisualizationMode);
+		UMaterial* Material = VisMaterial ? VisMaterial->GetMaterial() : nullptr;
 
 		if (Material && Material->BlendableLocation == Location)
 		{
diff --git a/Engine/Source/Runtime/Renderer/Private/SystemTextures.cpp b/Engine/Source/Runtime/Renderer/Private/SystemTextures.cpp
--- a/Engine/Source/Runtime/Renderer/Private/SystemTextures.cpp
+++ b/Engine/Source/Runtime/Renderer/Private/SystemTextures.cpp
@@ -227,169 +227,298 @@ void FSystemTextures::InitializeFeatureLevelDependentTextures(FRHICommandListImm
 			BlackBytes);
 	}
 
-	if (CurrentFeatureLevel < ERHIFeatureLevel::SM5 && InFeatureLevel >= ERHIFeatureLevel::SM5)
-		{
-	// Create the PerlinNoise3D texture (similar to http://prettyprocs.wordpress.com/2012/10/20/fast-perlin-noise/)
+	// The PreintegratedGF maybe used on forward shading inluding mobile platorm, intialize it anyway.
 	{
-		uint32 Extent = 16;
+		// for testing, with 128x128 R8G8 we are very close to the reference (if lower res is needed we might have to add an offset to counter the 0.5f texel shift)
+		const bool bReference = false;
 
-		const uint32 Square = Extent * Extent;
+		EPixelFormat Format = PF_R8G8;
+		// for low roughness we would get banding with PF_R8G8 but for low spec it could be used, for now we don't do this optimization
+		if (GPixelFormats[PF_G16R16].Supported)
+		{
+			Format = PF_G16R16;
+		}
 
-		FPooledRenderTargetDesc Desc(FPooledRenderTargetDesc::CreateVolumeDesc(Extent, Extent, Extent, PF_B8G8R8A8, FClearValueBinding::None, TexCreate_ShaderResource | TexCreate_HideInVisualizeTexture | TexCreate_NoTiling, TexCreate_ShaderResource, false));
+		FPooledRenderTargetDesc Desc(FPooledRenderTargetDesc::Create2DDesc(FIntPoint(128, 32), Format, FClearValueBinding::None, 0, TexCreate_ShaderResource, false));
 		Desc.AutoWritable = false;
-		GRenderTargetPool.FindFreeElement(RHICmdList, Desc, PerlinNoise3D, TEXT("PerlinNoise3D"), true, ERenderTargetTransience::NonTransient);
-		// Write the contents of the texture.
-		TArray<uint32> DestBuffer;
-
-		DestBuffer.AddZeroed(Extent * Extent * Extent);
-		// seed the pseudo random stream with a good value
-		FRandomStream RandomStream(0x1234);
-		// Values represent float3 values in the -1..1 range.
-		// The vectors are the edge mid point of a cube from -1 .. 1
-		// -1:0 0:7f 1:fe, can be reconstructed with * 512/254 - 1
-		// * 2 - 1 cannot be used because 0 would not be mapped
-			static uint32 gradtable[] =
+		if (bReference)
 		{
-			0x7ffefe, 0xfe7ffe, 0xfefe7f,
-			0x7ffe00, 0xfe7f00, 0xfe007f,
-			0x7f00fe, 0x007ffe, 0x00fe7f,
-			0x7f0000, 0x007f00, 0x00007f,
-		};
-		// set random directions
+			Desc.Extent.X = 128;
+			Desc.Extent.Y = 128;
+		}
+
+		GRenderTargetPool.FindFreeElement(RHICmdList, Desc, PreintegratedGF, TEXT("PreintegratedGF"), true, ERenderTargetTransience::NonTransient);
+		// Write the contents of the texture.
+		uint32 DestStride;
+		uint8* DestBuffer = (uint8*)RHICmdList.LockTexture2D((FTexture2DRHIRef&)PreintegratedGF->GetRenderTargetItem().ShaderResourceTexture, 0, RLM_WriteOnly, DestStride, false);
+
+		// x is NoV, y is roughness
+		for (int32 y = 0; y < Desc.Extent.Y; y++)
 		{
-				for (uint32 z = 0; z < Extent - 1; ++z)
+			float Roughness = (float)(y + 0.5f) / Desc.Extent.Y;
+			float m = Roughness * Roughness;
+			float m2 = m * m;
+
+			for (int32 x = 0; x < Desc.Extent.X; x++)
 			{
-					for (uint32 y = 0; y < Extent - 1; ++y)
+				float NoV = (float)(x + 0.5f) / Desc.Extent.X;
+
+				FVector V;
+				V.X = FMath::Sqrt(1.0f - NoV * NoV);	// sin
+				V.Y = 0.0f;
+				V.Z = NoV;								// cos
+
+				float A = 0.0f;
+				float B = 0.0f;
+				float C = 0.0f;
+
+				const uint32 NumSamples = 128;
+				for (uint32 i = 0; i < NumSamples; i++)
 				{
-						for (uint32 x = 0; x < Extent - 1; ++x)
+					float E1 = (float)i / NumSamples;
+					float E2 = (double)ReverseBits(i) / (double)0x100000000LL;
+
 					{
-						uint32& Value = DestBuffer[x + y * Extent + z * Square];
+						float Phi = 2.0f * PI * E1;
+						float CosPhi = FMath::Cos(Phi);
+						float SinPhi = FMath::Sin(Phi);
+						float CosTheta = FMath::Sqrt((1.0f - E2) / (1.0f + (m2 - 1.0f) * E2));
+						float SinTheta = FMath::Sqrt(1.0f - CosTheta * CosTheta);
 
-						// pick a random direction (hacky way to overcome the quality issues FRandomStream has)
-						Value = gradtable[(uint32)(RandomStream.GetFraction() * 11.9999999f)];
+						FVector H(SinTheta * FMath::Cos(Phi), SinTheta * FMath::Sin(Phi), CosTheta);
+						FVector L = 2.0f * (V | H) * H - V;
+
+						float NoL = FMath::Max(L.Z, 0.0f);
+						float NoH = FMath::Max(H.Z, 0.0f);
+						float VoH = FMath::Max(V | H, 0.0f);
+
+						if (NoL > 0.0f)
+						{
+							float Vis_SmithV = NoL * (NoV * (1 - m) + m);
+							float Vis_SmithL = NoV * (NoL * (1 - m) + m);
+							float Vis = 0.5f / (Vis_SmithV + Vis_SmithL);
+
+							float NoL_Vis_PDF = NoL * Vis * (4.0f * VoH / NoH);
+							float Fc = 1.0f - VoH;
+							Fc *= FMath::Square(Fc*Fc);
+							A += NoL_Vis_PDF * (1.0f - Fc);
+							B += NoL_Vis_PDF * Fc;
+						}
+					}
+
+					{
+						float Phi = 2.0f * PI * E1;
+						float CosPhi = FMath::Cos(Phi);
+						float SinPhi = FMath::Sin(Phi);
+						float CosTheta = FMath::Sqrt(E2);
+						float SinTheta = FMath::Sqrt(1.0f - CosTheta * CosTheta);
+
+						FVector L(SinTheta * FMath::Cos(Phi), SinTheta * FMath::Sin(Phi), CosTheta);
+						FVector H = (V + L).GetUnsafeNormal();
+
+						float NoL = FMath::Max(L.Z, 0.0f);
+						float NoH = FMath::Max(H.Z, 0.0f);
+						float VoH = FMath::Max(V | H, 0.0f);
+
+						float FD90 = 0.5f + 2.0f * VoH * VoH * Roughness;
+						float FdV = 1.0f + (FD90 - 1.0f) * pow(1.0f - NoV, 5);
+						float FdL = 1.0f + (FD90 - 1.0f) * pow(1.0f - NoL, 5);
+						C += FdV * FdL;// * ( 1.0f - 0.3333f * Roughness );
 					}
 				}
-			}
-		}
-		// replicate a border for filtering
-		{
-			uint32 Last = Extent - 1;
+				A /= NumSamples;
+				B /= NumSamples;
+				C /= NumSamples;
 
-				for (uint32 z = 0; z < Extent; ++z)
-			{
-					for (uint32 y = 0; y < Extent; ++y)
+				if (Desc.Format == PF_A16B16G16R16)
 				{
-					DestBuffer[Last + y * Extent + z * Square] = DestBuffer[0 + y * Extent + z * Square];
+					uint16* Dest = (uint16*)(DestBuffer + x * 8 + y * DestStride);
+					Dest[0] = (int32)(FMath::Clamp(A, 0.0f, 1.0f) * 65535.0f + 0.5f);
+					Dest[1] = (int32)(FMath::Clamp(B, 0.0f, 1.0f) * 65535.0f + 0.5f);
+					Dest[2] = (int32)(FMath::Clamp(C, 0.0f, 1.0f) * 65535.0f + 0.5f);
 				}
-			}
-				for (uint32 z = 0; z < Extent; ++z)
-			{
-					for (uint32 x = 0; x < Extent; ++x)
+				else if (Desc.Format == PF_G16R16)
 				{
-					DestBuffer[x + Last * Extent + z * Square] = DestBuffer[x + 0 * Extent + z * Square];
+					uint16* Dest = (uint16*)(DestBuffer + x * 4 + y * DestStride);
+					Dest[0] = (int32)(FMath::Clamp(A, 0.0f, 1.0f) * 65535.0f + 0.5f);
+					Dest[1] = (int32)(FMath::Clamp(B, 0.0f, 1.0f) * 65535.0f + 0.5f);
 				}
-			}
-				for (uint32 y = 0; y < Extent; ++y)
-			{
-					for (uint32 x = 0; x < Extent; ++x)
+				else
 				{
-					DestBuffer[x + y * Extent + Last * Square] = DestBuffer[x + y * Extent + 0 * Square];
+					check(Desc.Format == PF_R8G8);
+
+					uint8* Dest = (uint8*)(DestBuffer + x * 2 + y * DestStride);
+					Dest[0] = (int32)(FMath::Clamp(A, 0.0f, 1.0f) * 255.9999f);
+					Dest[1] = (int32)(FMath::Clamp(B, 0.0f, 1.0f) * 255.9999f);
 				}
 			}
 		}
-		// precompute gradients
-		{
-			uint32* Dest = DestBuffer.GetData();
+		RHICmdList.UnlockTexture2D((FTexture2DRHIRef&)PreintegratedGF->GetRenderTargetItem().ShaderResourceTexture, 0, false);
+	}
 
-				for (uint32 z = 0; z < Desc.Depth; ++z)
+	if (CurrentFeatureLevel < ERHIFeatureLevel::SM5 && InFeatureLevel >= ERHIFeatureLevel::SM5)
+	{
+	    // Create the PerlinNoise3D texture (similar to http://prettyprocs.wordpress.com/2012/10/20/fast-perlin-noise/)
+	    {
+		    uint32 Extent = 16;
+    
+		    const uint32 Square = Extent * Extent;
+    
+		    FPooledRenderTargetDesc Desc(FPooledRenderTargetDesc::CreateVolumeDesc(Extent, Extent, Extent, PF_B8G8R8A8, FClearValueBinding::None, TexCreate_ShaderResource | TexCreate_HideInVisualizeTexture | TexCreate_NoTiling, TexCreate_ShaderResource, false));
+		    Desc.AutoWritable = false;
+		    GRenderTargetPool.FindFreeElement(RHICmdList, Desc, PerlinNoise3D, TEXT("PerlinNoise3D"), true, ERenderTargetTransience::NonTransient);
+		    // Write the contents of the texture.
+		    TArray<uint32> DestBuffer;
+    
+		    DestBuffer.AddZeroed(Extent * Extent * Extent);
+		    // seed the pseudo random stream with a good value
+		    FRandomStream RandomStream(0x1234);
+		    // Values represent float3 values in the -1..1 range.
+		    // The vectors are the edge mid point of a cube from -1 .. 1
+		    // -1:0 0:7f 1:fe, can be reconstructed with * 512/254 - 1
+		    // * 2 - 1 cannot be used because 0 would not be mapped
+			    static uint32 gradtable[] =
+		    {
+			    0x7ffefe, 0xfe7ffe, 0xfefe7f,
+			    0x7ffe00, 0xfe7f00, 0xfe007f,
+			    0x7f00fe, 0x007ffe, 0x00fe7f,
+			    0x7f0000, 0x007f00, 0x00007f,
+		    };
+		    // set random directions
+		    {
+				    for (uint32 z = 0; z < Extent - 1; ++z)
+			    {
+					    for (uint32 y = 0; y < Extent - 1; ++y)
+				    {
+						    for (uint32 x = 0; x < Extent - 1; ++x)
+					    {
+						    uint32& Value = DestBuffer[x + y * Extent + z * Square];
+    
+						    // pick a random direction (hacky way to overcome the quality issues FRandomStream has)
+						    Value = gradtable[(uint32)(RandomStream.GetFraction() * 11.9999999f)];
+					    }
+				    }
+			    }
+		    }
+		    // replicate a border for filtering
+		    {
+			    uint32 Last = Extent - 1;
+    
+				    for (uint32 z = 0; z < Extent; ++z)
+			    {
+					    for (uint32 y = 0; y < Extent; ++y)
+				    {
+					    DestBuffer[Last + y * Extent + z * Square] = DestBuffer[0 + y * Extent + z * Square];
+				    }
+			    }
+				for (uint32 z = 0; z < Extent; ++z)
+				{
+					for (uint32 x = 0; x < Extent; ++x)
+					{
+					    DestBuffer[x + Last * Extent + z * Square] = DestBuffer[x + 0 * Extent + z * Square];
+				    }
+			    }
+				for (uint32 y = 0; y < Extent; ++y)
+				{
+					for (uint32 x = 0; x < Extent; ++x)
+				    {
+					    DestBuffer[x + y * Extent + Last * Square] = DestBuffer[x + y * Extent + 0 * Square];
+				    }
+			    }
+		    }
+		    // precompute gradients
 			{
+			    uint32* Dest = DestBuffer.GetData();
+    
+				for (uint32 z = 0; z < Desc.Depth; ++z)
+			    {
 					for (uint32 y = 0; y < (uint32)Desc.Extent.Y; ++y)
-				{
+				    {
 						for (uint32 x = 0; x < (uint32)Desc.Extent.X; ++x)
-					{
-						uint32 Value = *Dest;
-
-						// todo: check if rgb order is correct
-						int32 r = Value >> 16;
-						int32 g = (Value >> 8) & 0xff;
-						int32 b = Value & 0xff;
-
-						int nx = (r / 0x7f) - 1;
-						int ny = (g / 0x7f) - 1;
-						int nz = (b / 0x7f) - 1;
-
-						int32 d = nx * x + ny * y + nz * z;
-
-						// compress in 8bit
-						uint32 a = d + 127;
-
-						*Dest++ = Value | (a << 24);
-					}
-				}
-			}
-		}
-
-		FUpdateTextureRegion3D Region(0, 0, 0, 0, 0, 0, Desc.Extent.X, Desc.Extent.Y, Desc.Depth);
-
-		RHICmdList.UpdateTexture3D(
-			(FTexture3DRHIRef&)PerlinNoise3D->GetRenderTargetItem().ShaderResourceTexture,
-			0,
-			Region,
-			Desc.Extent.X * sizeof(uint32),
-			Desc.Extent.X * Desc.Extent.Y * sizeof(uint32),
-			(const uint8*)DestBuffer.GetData());
+					    {
+						    uint32 Value = *Dest;
+    
+						    // todo: check if rgb order is correct
+						    int32 r = Value >> 16;
+						    int32 g = (Value >> 8) & 0xff;
+						    int32 b = Value & 0xff;
+    
+						    int nx = (r / 0x7f) - 1;
+						    int ny = (g / 0x7f) - 1;
+						    int nz = (b / 0x7f) - 1;
+    
+						    int32 d = nx * x + ny * y + nz * z;
+    
+						    // compress in 8bit
+						    uint32 a = d + 127;
+    
+						    *Dest++ = Value | (a << 24);
+					    }
+				    }
+			    }
+		    }
+    
+		    FUpdateTextureRegion3D Region(0, 0, 0, 0, 0, 0, Desc.Extent.X, Desc.Extent.Y, Desc.Depth);
+    
+		    RHICmdList.UpdateTexture3D(
+			    (FTexture3DRHIRef&)PerlinNoise3D->GetRenderTargetItem().ShaderResourceTexture,
+			    0,
+			    Region,
+			    Desc.Extent.X * sizeof(uint32),
+			    Desc.Extent.X * Desc.Extent.Y * sizeof(uint32),
+			    (const uint8*)DestBuffer.GetData());
 		} // end Create the PerlinNoise3D texture
 
-	// Create the SSAO randomization texture
+		// Create the SSAO randomization texture
 		{
-	{
-		float g_AngleOff1 = 127;
-		float g_AngleOff2 = 198;
-		float g_AngleOff3 = 23;
-
-		FColor Bases[16];
-
-				for (int32 Pos = 0; Pos < 16; ++Pos)
-		{
-			// distribute rotations over 4x4 pattern
-					//			int32 Reorder[16] = { 0, 8, 2, 10, 12, 6, 14, 4, 3, 11, 1, 9, 15, 5, 13, 7 };
-			int32 Reorder[16] = { 0, 11, 7, 3, 10, 4, 15, 12, 6, 8, 1, 14, 13, 2, 9, 5 };
-			int32 w = Reorder[Pos];
-
-			// ordered sampling of the rotation basis (*2 is missing as we use mirrored samples)
-			float ww = w / 16.0f * PI;
-
-			// randomize base scale
-			float lenm = 1.0f - (FMath::Sin(g_AngleOff2 * w * 0.01f) * 0.5f + 0.5f) * g_AngleOff3 * 0.01f;
-			float s = FMath::Sin(ww) * lenm;
-			float c = FMath::Cos(ww) * lenm;
-
-			Bases[Pos] = FColor(FMath::Quantize8SignedByte(c), FMath::Quantize8SignedByte(s), 0, 0);
-		}
-
-		{
-			// could be PF_V8U8 to save shader instructions but that doesn't work on all hardware
-			FPooledRenderTargetDesc Desc(FPooledRenderTargetDesc::Create2DDesc(FIntPoint(64, 64), PF_R8G8, FClearValueBinding::None, TexCreate_HideInVisualizeTexture, TexCreate_NoFastClear | TexCreate_ShaderResource, false));
-			Desc.AutoWritable = false;
-			GRenderTargetPool.FindFreeElement(RHICmdList, Desc, SSAORandomization, TEXT("SSAORandomization"), true, ERenderTargetTransience::NonTransient);
-			// Write the contents of the texture.
-			uint32 DestStride;
-			uint8* DestBuffer = (uint8*)RHICmdList.LockTexture2D((FTexture2DRHIRef&)SSAORandomization->GetRenderTargetItem().ShaderResourceTexture, 0, RLM_WriteOnly, DestStride, false);
-
-			for(int32 y = 0; y < Desc.Extent.Y; ++y)
-			{
-				for(int32 x = 0; x < Desc.Extent.X; ++x)
-				{
-					uint8* Dest = (uint8*)(DestBuffer + x * sizeof(uint16) + y * DestStride);
-
-					uint32 Index = (x % 4) + (y % 4) * 4; 
-
-					Dest[0] = Bases[Index].R;
-					Dest[1] = Bases[Index].G;
-				}
-			}
-		}
-		RHICmdList.UnlockTexture2D((FTexture2DRHIRef&)SSAORandomization->GetRenderTargetItem().ShaderResourceTexture, 0, false);
+	        {
+		        float g_AngleOff1 = 127;
+		        float g_AngleOff2 = 198;
+		        float g_AngleOff3 = 23;
+        
+		        FColor Bases[16];
+    
+			    for (int32 Pos = 0; Pos < 16; ++Pos)
+		        {
+			        // distribute rotations over 4x4 pattern
+					        //			int32 Reorder[16] = { 0, 8, 2, 10, 12, 6, 14, 4, 3, 11, 1, 9, 15, 5, 13, 7 };
+			        int32 Reorder[16] = { 0, 11, 7, 3, 10, 4, 15, 12, 6, 8, 1, 14, 13, 2, 9, 5 };
+			        int32 w = Reorder[Pos];
+        
+			        // ordered sampling of the rotation basis (*2 is missing as we use mirrored samples)
+			        float ww = w / 16.0f * PI;
+        
+			        // randomize base scale
+			        float lenm = 1.0f - (FMath::Sin(g_AngleOff2 * w * 0.01f) * 0.5f + 0.5f) * g_AngleOff3 * 0.01f;
+			        float s = FMath::Sin(ww) * lenm;
+			        float c = FMath::Cos(ww) * lenm;
+        
+			        Bases[Pos] = FColor(FMath::Quantize8SignedByte(c), FMath::Quantize8SignedByte(s), 0, 0);
+		        }
+    
+		        {
+			        // could be PF_V8U8 to save shader instructions but that doesn't work on all hardware
+			        FPooledRenderTargetDesc Desc(FPooledRenderTargetDesc::Create2DDesc(FIntPoint(64, 64), PF_R8G8, FClearValueBinding::None, TexCreate_HideInVisualizeTexture, TexCreate_NoFastClear | TexCreate_ShaderResource, false));
+			        Desc.AutoWritable = false;
+			        GRenderTargetPool.FindFreeElement(RHICmdList, Desc, SSAORandomization, TEXT("SSAORandomization"), true, ERenderTargetTransience::NonTransient);
+			        // Write the contents of the texture.
+			        uint32 DestStride;
+			        uint8* DestBuffer = (uint8*)RHICmdList.LockTexture2D((FTexture2DRHIRef&)SSAORandomization->GetRenderTargetItem().ShaderResourceTexture, 0, RLM_WriteOnly, DestStride, false);
+        
+			        for(int32 y = 0; y < Desc.Extent.Y; ++y)
+			        {
+				        for(int32 x = 0; x < Desc.Extent.X; ++x)
+				        {
+					        uint8* Dest = (uint8*)(DestBuffer + x * sizeof(uint16) + y * DestStride);
+        
+					        uint32 Index = (x % 4) + (y % 4) * 4; 
+        
+					        Dest[0] = Bases[Index].R;
+					        Dest[1] = Bases[Index].G;
+				        }
+			        }
+		        }
+		        RHICmdList.UnlockTexture2D((FTexture2DRHIRef&)SSAORandomization->GetRenderTargetItem().ShaderResourceTexture, 0, false);
 			}
 	
 			{
@@ -409,8 +538,8 @@ void FSystemTextures::InitializeFeatureLevelDependentTextures(FRHICommandListImm
 						float Angle  = (PI/16.0f)  * ((((x + y) & 0x3) << 2) + (x & 0x3));
 						float Step   = (1.0f / 4.0f) * ((y - x) & 0x3);
 
-						float ScaleCos = FMath::Cos(Angle);
-						float ScaleSin = FMath::Sin(Angle);
+						float ScaleCos = FMath::Cos(Angle) ;
+						float ScaleSin = FMath::Sin(Angle) ;
 			
 						Dest[0] = (uint8_t)(ScaleCos*127.5f + 127.5f);
 						Dest[1] = (uint8_t)(ScaleSin*127.5f + 127.5f);
@@ -421,178 +550,49 @@ void FSystemTextures::InitializeFeatureLevelDependentTextures(FRHICommandListImm
 			}
 			RHICmdList.UnlockTexture2D((FTexture2DRHIRef&)GTAORandomization->GetRenderTargetItem().ShaderResourceTexture, 0, false);
 			
-			
-		{
-			// for testing, with 128x128 R8G8 we are very close to the reference (if lower res is needed we might have to add an offset to counter the 0.5f texel shift)
-			const bool bReference = false;
-
-			EPixelFormat Format = PF_R8G8;
-			// for low roughness we would get banding with PF_R8G8 but for low spec it could be used, for now we don't do this optimization
-			if (GPixelFormats[PF_G16R16].Supported)
-			{
-				Format = PF_G16R16;
-			}
-
-			FPooledRenderTargetDesc Desc(FPooledRenderTargetDesc::Create2DDesc(FIntPoint(128, 32), Format, FClearValueBinding::None, 0, TexCreate_ShaderResource, false));
-			Desc.AutoWritable = false;
-			if (bReference)
-			{
-				Desc.Extent.X = 128;
-				Desc.Extent.Y = 128;
-			}
-
-			GRenderTargetPool.FindFreeElement(RHICmdList, Desc, PreintegratedGF, TEXT("PreintegratedGF"), true, ERenderTargetTransience::NonTransient);
-			// Write the contents of the texture.
-			uint32 DestStride;
-			uint8* DestBuffer = (uint8*)RHICmdList.LockTexture2D((FTexture2DRHIRef&)PreintegratedGF->GetRenderTargetItem().ShaderResourceTexture, 0, RLM_WriteOnly, DestStride, false);
-
-			// x is NoV, y is roughness
-			for (int32 y = 0; y < Desc.Extent.Y; y++)
-			{
-				float Roughness = (float)(y + 0.5f) / Desc.Extent.Y;
-				float m = Roughness * Roughness;
-					float m2 = m * m;
-
-				for (int32 x = 0; x < Desc.Extent.X; x++)
-				{
-					float NoV = (float)(x + 0.5f) / Desc.Extent.X;
-
-					FVector V;
-					V.X = FMath::Sqrt(1.0f - NoV * NoV);	// sin
-					V.Y = 0.0f;
-					V.Z = NoV;								// cos
-
-					float A = 0.0f;
-					float B = 0.0f;
-					float C = 0.0f;
-
-					const uint32 NumSamples = 128;
-					for (uint32 i = 0; i < NumSamples; i++)
-					{
-						float E1 = (float)i / NumSamples;
-						float E2 = (double)ReverseBits(i) / (double)0x100000000LL;
-
-						{
-							float Phi = 2.0f * PI * E1;
-							float CosPhi = FMath::Cos(Phi);
-							float SinPhi = FMath::Sin(Phi);
-							float CosTheta = FMath::Sqrt((1.0f - E2) / (1.0f + (m2 - 1.0f) * E2));
-							float SinTheta = FMath::Sqrt(1.0f - CosTheta * CosTheta);
-
-							FVector H(SinTheta * FMath::Cos(Phi), SinTheta * FMath::Sin(Phi), CosTheta);
-							FVector L = 2.0f * (V | H) * H - V;
-
-							float NoL = FMath::Max(L.Z, 0.0f);
-							float NoH = FMath::Max(H.Z, 0.0f);
-							float VoH = FMath::Max(V | H, 0.0f);
-
-							if (NoL > 0.0f)
-							{
-									float Vis_SmithV = NoL * (NoV * (1 - m) + m);
-									float Vis_SmithL = NoV * (NoL * (1 - m) + m);
-									float Vis = 0.5f / (Vis_SmithV + Vis_SmithL);
-
-								float NoL_Vis_PDF = NoL * Vis * (4.0f * VoH / NoH);
-								float Fc = 1.0f - VoH;
-								Fc *= FMath::Square(Fc*Fc);
-								A += NoL_Vis_PDF * (1.0f - Fc);
-								B += NoL_Vis_PDF * Fc;
-							}
-						}
-
-						{
-							float Phi = 2.0f * PI * E1;
-							float CosPhi = FMath::Cos(Phi);
-							float SinPhi = FMath::Sin(Phi);
-							float CosTheta = FMath::Sqrt(E2);
-							float SinTheta = FMath::Sqrt(1.0f - CosTheta * CosTheta);
-
-							FVector L(SinTheta * FMath::Cos(Phi), SinTheta * FMath::Sin(Phi), CosTheta);
-							FVector H = (V + L).GetUnsafeNormal();
-
-							float NoL = FMath::Max(L.Z, 0.0f);
-							float NoH = FMath::Max(H.Z, 0.0f);
-							float VoH = FMath::Max(V | H, 0.0f);
-
-							float FD90 = 0.5f + 2.0f * VoH * VoH * Roughness;
-								float FdV = 1.0f + (FD90 - 1.0f) * pow(1.0f - NoV, 5);
-								float FdL = 1.0f + (FD90 - 1.0f) * pow(1.0f - NoL, 5);
-							C += FdV * FdL;// * ( 1.0f - 0.3333f * Roughness );
-						}
-					}
-					A /= NumSamples;
-					B /= NumSamples;
-					C /= NumSamples;
-
-						if (Desc.Format == PF_A16B16G16R16)
-					{
-						uint16* Dest = (uint16*)(DestBuffer + x * 8 + y * DestStride);
-						Dest[0] = (int32)(FMath::Clamp(A, 0.0f, 1.0f) * 65535.0f + 0.5f);
-						Dest[1] = (int32)(FMath::Clamp(B, 0.0f, 1.0f) * 65535.0f + 0.5f);
-						Dest[2] = (int32)(FMath::Clamp(C, 0.0f, 1.0f) * 65535.0f + 0.5f);
-					}
-					else if (Desc.Format == PF_G16R16)
-					{
-						uint16* Dest = (uint16*)(DestBuffer + x * 4 + y * DestStride);
- 						Dest[0] = (int32)(FMath::Clamp(A, 0.0f, 1.0f) * 65535.0f + 0.5f);
- 						Dest[1] = (int32)(FMath::Clamp(B, 0.0f, 1.0f) * 65535.0f + 0.5f);
-					}
-					else
-					{
-						check(Desc.Format == PF_R8G8);
-
-						uint8* Dest = (uint8*)(DestBuffer + x * 2 + y * DestStride);
-						Dest[0] = (int32)(FMath::Clamp(A, 0.0f, 1.0f) * 255.9999f);
-						Dest[1] = (int32)(FMath::Clamp(B, 0.0f, 1.0f) * 255.9999f);
-					}
-				}
-			}
-			RHICmdList.UnlockTexture2D((FTexture2DRHIRef&)PreintegratedGF->GetRenderTargetItem().ShaderResourceTexture, 0, false);
-		}
-
-		{
-			FPooledRenderTargetDesc Desc(FPooledRenderTargetDesc::Create2DDesc(FIntPoint(LTC_Size, LTC_Size), PF_FloatRGBA, FClearValueBinding::None, TexCreate_FastVRAM, TexCreate_ShaderResource, false));
-			Desc.AutoWritable = false;
-
-			GRenderTargetPool.FindFreeElement(RHICmdList, Desc, LTCMat, TEXT("LTCMat"));
-			// Write the contents of the texture.
-			uint32 DestStride;
-			uint8* DestBuffer = (uint8*)RHICmdList.LockTexture2D((FTexture2DRHIRef&)LTCMat->GetRenderTargetItem().ShaderResourceTexture, 0, RLM_WriteOnly, DestStride, false);
-
-				for (int32 y = 0; y < Desc.Extent.Y; ++y)
-			{
-					for (int32 x = 0; x < Desc.Extent.X; ++x)
-				{
-					uint16* Dest = (uint16*)(DestBuffer + x * 4 * sizeof(uint16) + y * DestStride);
-
-						for (int k = 0; k < 4; k++)
-							Dest[k] = FFloat16(LTC_Mat[4 * (x + y * LTC_Size) + k]).Encoded;
-				}
-			}
-			RHICmdList.UnlockTexture2D((FTexture2DRHIRef&)LTCMat->GetRenderTargetItem().ShaderResourceTexture, 0, false);
-		}
-
-		{
-			FPooledRenderTargetDesc Desc(FPooledRenderTargetDesc::Create2DDesc(FIntPoint(LTC_Size, LTC_Size), PF_G16R16F, FClearValueBinding::None, TexCreate_FastVRAM, TexCreate_ShaderResource, false));
-			Desc.AutoWritable = false;
-
-			GRenderTargetPool.FindFreeElement(RHICmdList, Desc, LTCAmp, TEXT("LTCAmp"));
-			// Write the contents of the texture.
-			uint32 DestStride;
-			uint8* DestBuffer = (uint8*)RHICmdList.LockTexture2D((FTexture2DRHIRef&)LTCAmp->GetRenderTargetItem().ShaderResourceTexture, 0, RLM_WriteOnly, DestStride, false);
-
+		    {
+			    FPooledRenderTargetDesc Desc(FPooledRenderTargetDesc::Create2DDesc(FIntPoint(LTC_Size, LTC_Size), PF_FloatRGBA, FClearValueBinding::None, TexCreate_FastVRAM, TexCreate_ShaderResource, false));
+			    Desc.AutoWritable = false;
+    
+			    GRenderTargetPool.FindFreeElement(RHICmdList, Desc, LTCMat, TEXT("LTCMat"));
+			    // Write the contents of the texture.
+			    uint32 DestStride;
+			    uint8* DestBuffer = (uint8*)RHICmdList.LockTexture2D((FTexture2DRHIRef&)LTCMat->GetRenderTargetItem().ShaderResourceTexture, 0, RLM_WriteOnly, DestStride, false);
+    
+				    for (int32 y = 0; y < Desc.Extent.Y; ++y)
+			    {
+					    for (int32 x = 0; x < Desc.Extent.X; ++x)
+				    {
+					    uint16* Dest = (uint16*)(DestBuffer + x * 4 * sizeof(uint16) + y * DestStride);
+    
+						    for (int k = 0; k < 4; k++)
+							    Dest[k] = FFloat16(LTC_Mat[4 * (x + y * LTC_Size) + k]).Encoded;
+				    }
+			    }
+			    RHICmdList.UnlockTexture2D((FTexture2DRHIRef&)LTCMat->GetRenderTargetItem().ShaderResourceTexture, 0, false);
+		    }
+    
+		    {
+			    FPooledRenderTargetDesc Desc(FPooledRenderTargetDesc::Create2DDesc(FIntPoint(LTC_Size, LTC_Size), PF_G16R16F, FClearValueBinding::None, TexCreate_FastVRAM, TexCreate_ShaderResource, false));
+			    Desc.AutoWritable = false;
+    
+			    GRenderTargetPool.FindFreeElement(RHICmdList, Desc, LTCAmp, TEXT("LTCAmp"));
+			    // Write the contents of the texture.
+			    uint32 DestStride;
+			    uint8* DestBuffer = (uint8*)RHICmdList.LockTexture2D((FTexture2DRHIRef&)LTCAmp->GetRenderTargetItem().ShaderResourceTexture, 0, RLM_WriteOnly, DestStride, false);
+    
 				for (int32 y = 0; y < Desc.Extent.Y; ++y)
-			{
+			    {
 					for (int32 x = 0; x < Desc.Extent.X; ++x)
-				{
-					uint16* Dest = (uint16*)(DestBuffer + x * 2 * sizeof(uint16) + y * DestStride);
-
+				    {
+					    uint16* Dest = (uint16*)(DestBuffer + x * 2 * sizeof(uint16) + y * DestStride);
+    
 						for (int k = 0; k < 2; k++)
-							Dest[k] = FFloat16(LTC_Amp[4 * (x + y * LTC_Size) + k]).Encoded;
-				}
-			}
-			RHICmdList.UnlockTexture2D((FTexture2DRHIRef&)LTCAmp->GetRenderTargetItem().ShaderResourceTexture, 0, false);
-		}
+						    Dest[k] = FFloat16(LTC_Amp[4 * (x + y * LTC_Size) + k]).Encoded;
+				    }
+			    }
+			    RHICmdList.UnlockTexture2D((FTexture2DRHIRef&)LTCAmp->GetRenderTargetItem().ShaderResourceTexture, 0, false);
+		    }
 		} // end Create the SSAO randomization texture
 	} // end if (FeatureLevelInitializedTo < ERHIFeatureLevel::SM5 && InFeatureLevel >= ERHIFeatureLevel::SM5)
 
@@ -714,4 +714,4 @@ FRDGTextureRef FSystemTextures::GetMidGreyDummy(FRDGBuilder& GraphBuilder) const
 FRDGTextureRef FSystemTextures::GetVolumetricBlackDummy(FRDGBuilder& GraphBuilder) const
 {
 	return GraphBuilder.RegisterExternalTexture(VolumetricBlackDummy, TEXT("VolumetricBlackDummy"));
-}
+}
\ No newline at end of file
