Fixed error in visualize texture due to missing shader resource flag.

diff --git a/Engine/Source/Runtime/Renderer/Private/CompositionLighting/PostProcessDeferredDecals.cpp b/Engine/Source/Runtime/Renderer/Private/CompositionLighting/PostProcessDeferredDecals.cpp
--- a/Engine/Source/Runtime/Renderer/Private/CompositionLighting/PostProcessDeferredDecals.cpp
+++ b/Engine/Source/Runtime/Renderer/Private/CompositionLighting/PostProcessDeferredDecals.cpp
@@ -550,6 +550,9 @@ void FRCPassPostProcessDeferredDecals::Process(FRenderingCompositePassContext& C
 				// While a smaller format, such as R8_UINT, will use less video memory, it will result in slower clears and higher bandwidth requirements.
 				check(Desc.Format == PF_B8G8R8A8);
 				Desc.Flags = TexCreate_None;
+#if SUPPORTS_VISUALIZE_TEXTURE
+				Desc.TargetableFlags |= TexCreate_ShaderResource;
+#endif
 				Desc.ClearValue = FClearValueBinding::Transparent;
 				GRenderTargetPool.FindFreeElement(RHICmdList, Desc, SceneContext.DBufferMask, TEXT("DBufferMask"));
 			}
