Fix ensure in mobile post-processing

diff --git a/Engine/Source/Runtime/Renderer/Private/PostProcess/PostProcessing.cpp b/Engine/Source/Runtime/Renderer/Private/PostProcess/PostProcessing.cpp
--- a/Engine/Source/Runtime/Renderer/Private/PostProcess/PostProcessing.cpp
+++ b/Engine/Source/Runtime/Renderer/Private/PostProcess/PostProcessing.cpp
@@ -1769,7 +1769,10 @@ void FPostProcessing::ProcessES2(FRHICommandListImmediate& RHICmdList, FScene* S
 			Desc.NumMips = 1;
 			Desc.DebugName = TEXT("OverriddenRenderTarget");
 			Desc.TargetableFlags |= TexCreate_RenderTargetable;
-
+#if !(UE_BUILD_SHIPPING || UE_BUILD_TEST)
+			// for FVisualizeTexture
+			Desc.TargetableFlags |= TexCreate_ShaderResource;
+#endif
 			GRenderTargetPool.CreateUntrackedElement(Desc, Temp, Item);
 
 			OverrideRenderTarget(Context.FinalOutput, Temp, Desc);
