Fixed merge compilation issues

diff --git a/Engine/Source/Runtime/Engine/Public/BufferVisualizationData.h b/Engine/Source/Runtime/Engine/Public/BufferVisualizationData.h
--- a/Engine/Source/Runtime/Engine/Public/BufferVisualizationData.h
+++ b/Engine/Source/Runtime/Engine/Public/BufferVisualizationData.h
@@ -8,9 +8,9 @@ class UMaterialInterface;
 class FBufferVisualizationData
 {
 public:
-	
+
 	FBufferVisualizationData()
-	: bIsInitialized(false)
+		: bIsInitialized(false)
 	{
 
 	}
@@ -24,6 +24,12 @@ public:
 	/** Get a named material from the available material map **/
 	ENGINE_API UMaterialInterface* GetMaterial(FName InMaterialName);
 
+	/** Get the display name of a named material from the available material map **/
+	ENGINE_API FText GetMaterialDisplayName(FName InMaterialName) const;
+
+	/** Get the default display name if no material is used and the Overview window is being used **/
+	ENGINE_API static FText GetMaterialDefaultDisplayName();
+
 	/** We cache the overview material name list from the console command here, so all dynamically created views can re-use the existing cached list of materials */
 	void SetCurrentOverviewMaterialNames(const FString& InNameList);
 	bool IsDifferentToCurrentOverviewMaterialNames(const FString& InNameList);
@@ -37,7 +43,7 @@ public:
 		for (TMaterialMap::TConstIterator It = MaterialMap.CreateConstIterator(); It; ++It)
 		{
 			const Record& Rec = It.Value();
-			Iterator.ProcessValue(Rec.Name, Rec.Material, Rec.DisplayName);
+			Iterator.ProcessValue(Rec.Name, Rec.Material->GetMaterial(), Rec.DisplayName);
 		}
 	}
 
@@ -62,13 +68,16 @@ private:
 
 	/** The name->material mapping table */
 	TMaterialMap MaterialMap;
-	
+
+	/** The UMaterial.name->material mapping table */
+	TMaterialMap MaterialMapFromMaterialName;
+
 	/** List of material names to use in the buffer visualization overview */
 	FString CurrentOverviewMaterialNames;
 
 	/** List of material currently in use by the buffer visualization overview */
 	TArray<UMaterialInterface*> OverviewMaterials;
-	
+
 	/** Storage for console variable documentation strings **/
 	FString ConsoleDocumentationVisualizationMode;
 	FString ConsoleDocumentationOverviewTargets;
diff --git a/Engine/Source/Runtime/Renderer/Private/HairStrands/HairStrandsVisibility.cpp b/Engine/Source/Runtime/Renderer/Private/HairStrands/HairStrandsVisibility.cpp
--- a/Engine/Source/Runtime/Renderer/Private/HairStrands/HairStrandsVisibility.cpp
+++ b/Engine/Source/Runtime/Renderer/Private/HairStrands/HairStrandsVisibility.cpp
@@ -53,7 +53,7 @@ static FAutoConsoleVariableRef CVarGHairVisibilityPPLLMeanNodeCountPerPixel(TEXT
 static int32 GHairStrandsVisibilityMaterialPass = 0;
 static FAutoConsoleVariableRef CVarHairStrandsVisibilityMaterialPass(TEXT("r.HairStrands.Visibility.MaterialPass"), GHairStrandsVisibilityMaterialPass, TEXT("Enable the deferred material pass evaluation after the hair visibility is resolved."));
 
-static int32 GHairStrandsViewHairCount = 0;
+static int32 GHairStrandsViewHairCount = 1;
 static float GHairStrandsViewHairCountDepthDistanceThreshold = 30.f;
 static FAutoConsoleVariableRef CVarHairStrandsViewHairCount(TEXT("r.HairStrands.Visibility.HairCount"), GHairStrandsViewHairCount, TEXT("Enable the computation of 'view-hair-count' during the transmission pass."));
 static FAutoConsoleVariableRef CVarHairStrandsViewHairCountDepthDistanceThreshold(TEXT("r.HairStrands.Visibility.HairCount.DistanceThreshold"), GHairStrandsViewHairCountDepthDistanceThreshold, TEXT("Distance threshold defining if opaque depth get injected into the 'view-hair-count' buffer."));
@@ -65,7 +65,7 @@ namespace HairStrandsVisibilityInternal
 	struct NodeData
 	{
 		uint32 Depth;
-		uint32 PrimitiveId_ClusterId;
+		uint32 PrimitiveId_MacroGroupId;
 		uint32 Tangent_Coverage;
 		uint32 BaseColor_Roughness;
 		uint32 Specular;
@@ -75,8 +75,8 @@ namespace HairStrandsVisibilityInternal
 	struct NodeVis
 	{
 		uint32 Depth;
-		uint32 PrimitiveId_ClusterId;
-		uint32 Coverage_ClusterID_Pad;
+		uint32 PrimitiveId_MacroGroupId;
+		uint32 Coverage_MacroGroupId_Pad;
 		uint32 Pad;
 	};
 }
@@ -133,7 +133,7 @@ uint32 GetPPLLMaxTotalRenderNode(FIntPoint Resolution)
 
 uint32 GetHairVisibilitySampleCount()
 {
-	return GetHairVisibilityRenderMode() == HairVisibilityRenderMode_MSAA ? FMath::Clamp(GHairVisibilitySampleCount,1,16) : 1;
+	return GetHairVisibilityRenderMode() == HairVisibilityRenderMode_MSAA ? FMath::Clamp(GHairVisibilitySampleCount, 1, 16) : 1;
 }
 
 void SetUpViewHairRenderInfo(const FViewInfo& ViewInfo, bool bEnableMSAA, FVector4& OutHairRenderInfo)
@@ -148,11 +148,7 @@ void SetUpViewHairRenderInfo(const FViewInfo& ViewInfo, bool bEnableMSAA, FVecto
 		FIntPoint(ViewInfo.UnconstrainedViewRect.Width(), ViewInfo.UnconstrainedViewRect.Height()), ViewInfo.FOV, HairVisibilitySampleCount, RasterizationScaleOverride);
 
 	TUniformBufferRef<FViewUniformShaderParameters> ViewUniformShaderParameters;
-	// Update our view parameters
-	OutHairRenderInfo.X = MinHairRadius.Primary;
-	OutHairRenderInfo.Y = MinHairRadius.Velocity;
-	OutHairRenderInfo.Z = ViewInfo.IsPerspectiveProjection() ? 0.0f : 1.0f;
-	OutHairRenderInfo.W = VelocityMagnitudeScale;
+	OutHairRenderInfo = PackHairRenderInfo(MinHairRadius.Primary, MinHairRadius.Velocity, VelocityMagnitudeScale, !ViewInfo.IsPerspectiveProjection(), false);
 }
 
 static bool IsCompatibleWithHairVisibility(const FMeshMaterialShaderPermutationParameters& Parameters)
@@ -163,13 +159,13 @@ static bool IsCompatibleWithHairVisibility(const FMeshMaterialShaderPermutationP
 /////////////////////////////////////////////////////////////////////////////////////////
 
 BEGIN_GLOBAL_SHADER_PARAMETER_STRUCT(FMaterialPassParameters, )
-	SHADER_PARAMETER(FIntPoint, MaxResolution)
-	SHADER_PARAMETER(uint32, MaxSampleCount)
-	SHADER_PARAMETER(uint32, NodeGroupSize)
-	SHADER_PARAMETER_TEXTURE(Texture2D<uint>, NodeIndex)
-	SHADER_PARAMETER_SRV(StructuredBuffer<uint>, NodeCoord)
-	SHADER_PARAMETER_SRV(StructuredBuffer<FNodeVis>, NodeVis)
-	SHADER_PARAMETER_SRV(Buffer<uint>, IndirectArgs)
+SHADER_PARAMETER(FIntPoint, MaxResolution)
+SHADER_PARAMETER(uint32, MaxSampleCount)
+SHADER_PARAMETER(uint32, NodeGroupSize)
+SHADER_PARAMETER_TEXTURE(Texture2D<uint>, NodeIndex)
+SHADER_PARAMETER_SRV(StructuredBuffer<uint>, NodeCoord)
+SHADER_PARAMETER_SRV(StructuredBuffer<FNodeVis>, NodeVis)
+SHADER_PARAMETER_SRV(Buffer<uint>, IndirectArgs)
 END_GLOBAL_SHADER_PARAMETER_STRUCT()
 IMPLEMENT_GLOBAL_SHADER_PARAMETER_STRUCT(FMaterialPassParameters, "MaterialPassParameters");
 
@@ -206,8 +202,8 @@ IMPLEMENT_MATERIAL_SHADER_TYPE(, FHairMaterialVS, TEXT("/Engine/Private/HairStra
 class FHairMaterialShaderElementData : public FMeshMaterialShaderElementData
 {
 public:
-	FHairMaterialShaderElementData(int32 ClusterId, int32 MaterialId, int32 PrimitiveId) : MaterialPass_ClusterId(ClusterId), MaterialPass_MaterialId(MaterialId), MaterialPass_PrimitiveId(PrimitiveId) { }
-	uint32 MaterialPass_ClusterId;
+	FHairMaterialShaderElementData(int32 MacroGroupId, int32 MaterialId, int32 PrimitiveId) : MaterialPass_MacroGroupId(MacroGroupId), MaterialPass_MaterialId(MaterialId), MaterialPass_PrimitiveId(PrimitiveId) { }
+	uint32 MaterialPass_MacroGroupId;
 	uint32 MaterialPass_MaterialId;
 	uint32 MaterialPass_PrimitiveId;
 };
@@ -223,7 +219,7 @@ public:
 		ERHIFeatureLevel::Type FeatureLevel = GetMaxSupportedFeatureLevel((EShaderPlatform)Initializer.Target.Platform);
 		check(FSceneInterface::GetShadingPath(FeatureLevel) != EShadingPath::Mobile);
 		PassUniformBuffer.Bind(Initializer.ParameterMap, FMaterialPassParameters::StaticStructMetadata.GetShaderVariableName());
-		MaterialPass_ClusterId.Bind(Initializer.ParameterMap, TEXT("MaterialPass_ClusterId"));
+		MaterialPass_MacroGroupId.Bind(Initializer.ParameterMap, TEXT("MaterialPass_MacroGroupId"));
 		MaterialPass_MaterialId.Bind(Initializer.ParameterMap, TEXT("MaterialPass_MaterialId"));
 		MaterialPass_PrimitiveId.Bind(Initializer.ParameterMap, TEXT("MaterialPass_PrimitiveId"));
 
@@ -249,7 +245,7 @@ public:
 	virtual bool Serialize(FArchive& Ar) override
 	{
 		bool bShaderHasOutdatedParameters = FMeshMaterialShader::Serialize(Ar);
-		Ar << MaterialPass_ClusterId;
+		Ar << MaterialPass_MacroGroupId;
 		Ar << MaterialPass_MaterialId;
 		Ar << MaterialPass_PrimitiveId;
 		Ar << OutNodeData;
@@ -268,13 +264,13 @@ public:
 		FMeshDrawSingleShaderBindings& ShaderBindings) const
 	{
 		FMeshMaterialShader::GetShaderBindings(Scene, FeatureLevel, PrimitiveSceneProxy, MaterialRenderProxy, Material, DrawRenderState, ShaderElementData, ShaderBindings);
-		ShaderBindings.Add(MaterialPass_ClusterId, ShaderElementData.MaterialPass_ClusterId);
+		ShaderBindings.Add(MaterialPass_MacroGroupId, ShaderElementData.MaterialPass_MacroGroupId);
 		ShaderBindings.Add(MaterialPass_MaterialId, ShaderElementData.MaterialPass_MaterialId);
 		ShaderBindings.Add(MaterialPass_PrimitiveId, ShaderElementData.MaterialPass_PrimitiveId);
 	}
 
 private:
-	FShaderParameter MaterialPass_ClusterId;
+	FShaderParameter MaterialPass_MacroGroupId;
 	FShaderParameter MaterialPass_MaterialId;
 	FShaderParameter MaterialPass_PrimitiveId;
 	FShaderParameter OutNodeData;
@@ -294,7 +290,7 @@ public:
 		FDynamicPassMeshDrawListContext* InDrawListContext);
 
 	virtual void AddMeshBatch(const FMeshBatch& RESTRICT MeshBatch, uint64 BatchElementMask, const FPrimitiveSceneProxy* RESTRICT PrimitiveSceneProxy, int32 StaticMeshId = -1) override final;
-	void AddMeshBatch(const FMeshBatch& RESTRICT MeshBatch, uint64 BatchElementMask, const FPrimitiveSceneProxy* RESTRICT PrimitiveSceneProxy, int32 StaticMeshId, int32 HairClusterId, int32 HairMaterialId);
+	void AddMeshBatch(const FMeshBatch& RESTRICT MeshBatch, uint64 BatchElementMask, const FPrimitiveSceneProxy* RESTRICT PrimitiveSceneProxy, int32 StaticMeshId, int32 MacroGroupId, int32 HairMaterialId);
 
 private:
 	void Process(
@@ -304,7 +300,7 @@ private:
 		int32 StaticMeshId,
 		const FMaterialRenderProxy& RESTRICT MaterialRenderProxy,
 		const FMaterial& RESTRICT MaterialResource,
-		const int32 HairClusterId,
+		const int32 MacroGroupId,
 		const int32 HairMaterialId,
 		const int32 HairPrimitiveId);
 
@@ -316,7 +312,7 @@ void FHairMaterialProcessor::AddMeshBatch(const FMeshBatch& RESTRICT MeshBatch,
 	AddMeshBatch(MeshBatch, BatchElementMask, PrimitiveSceneProxy, StaticMeshId, 0, 0);
 }
 
-void FHairMaterialProcessor::AddMeshBatch(const FMeshBatch& RESTRICT MeshBatch, uint64 BatchElementMask, const FPrimitiveSceneProxy* RESTRICT PrimitiveSceneProxy, int32 StaticMeshId, int32 HairClusterId, int32 HairMaterialId)
+void FHairMaterialProcessor::AddMeshBatch(const FMeshBatch& RESTRICT MeshBatch, uint64 BatchElementMask, const FPrimitiveSceneProxy* RESTRICT PrimitiveSceneProxy, int32 StaticMeshId, int32 MacroGroupId, int32 HairMaterialId)
 {
 	static const FVertexFactoryType* CompatibleVF = FVertexFactoryType::GetVFByName(TEXT("FHairStrandsVertexFactory"));
 
@@ -325,7 +321,7 @@ void FHairMaterialProcessor::AddMeshBatch(const FMeshBatch& RESTRICT MeshBatch,
 	const FMaterial& Material = MeshBatch.MaterialRenderProxy->GetMaterialWithFallback(FeatureLevel, FallbackMaterialRenderProxyPtr);
 	const bool bIsCompatible = IsCompatibleWithHairStrands(&Material, FeatureLevel);
 	const bool bIsHairStrandsFactory = MeshBatch.VertexFactory->GetType()->GetId() == CompatibleVF->GetId();
-	const bool bShouldRender = (!PrimitiveSceneProxy && MeshBatch.Elements.Num()>0) || (PrimitiveSceneProxy && PrimitiveSceneProxy->ShouldRenderInMainPass());
+	const bool bShouldRender = (!PrimitiveSceneProxy && MeshBatch.Elements.Num() > 0) || (PrimitiveSceneProxy && PrimitiveSceneProxy->ShouldRenderInMainPass());
 
 	if (bIsCompatible
 		&& bIsHairStrandsFactory
@@ -334,11 +330,13 @@ void FHairMaterialProcessor::AddMeshBatch(const FMeshBatch& RESTRICT MeshBatch,
 	{
 		// For the mesh patch to be rendered a single triangle triangle to spawn the necessary amount of thread
 		FMeshBatch MeshBatchCopy = MeshBatch;
-		for (uint32 ElementIt = 0, ElementCount=uint32(MeshBatch.Elements.Num()); ElementIt < ElementCount; ++ElementIt)
+		for (uint32 ElementIt = 0, ElementCount = uint32(MeshBatch.Elements.Num()); ElementIt < ElementCount; ++ElementIt)
 		{
 			MeshBatchCopy.Elements[ElementIt].FirstIndex = 0;
 			MeshBatchCopy.Elements[ElementIt].NumPrimitives = 1;
 			MeshBatchCopy.Elements[ElementIt].NumInstances = 1;
+			MeshBatchCopy.Elements[ElementIt].IndirectArgsBuffer = nullptr;
+			MeshBatchCopy.Elements[ElementIt].IndirectArgsOffset = 0;
 		}
 
 		int32 PrimitiveId = 0;
@@ -347,7 +345,7 @@ void FHairMaterialProcessor::AddMeshBatch(const FMeshBatch& RESTRICT MeshBatch,
 		GetDrawCommandPrimitiveId(SceneInfo, MeshBatch.Elements[0], PrimitiveId, ScenePrimitiveId);
 
 		const FMaterialRenderProxy& MaterialRenderProxy = FallbackMaterialRenderProxyPtr ? *FallbackMaterialRenderProxyPtr : *MeshBatch.MaterialRenderProxy;
-		Process(MeshBatchCopy, BatchElementMask, PrimitiveSceneProxy, StaticMeshId, MaterialRenderProxy, Material, HairClusterId, HairMaterialId, PrimitiveId);
+		Process(MeshBatchCopy, BatchElementMask, PrimitiveSceneProxy, StaticMeshId, MaterialRenderProxy, Material, MacroGroupId, HairMaterialId, PrimitiveId);
 	}
 }
 
@@ -358,7 +356,7 @@ void FHairMaterialProcessor::Process(
 	int32 StaticMeshId,
 	const FMaterialRenderProxy& RESTRICT MaterialRenderProxy,
 	const FMaterial& RESTRICT MaterialResource,
-	const int32 HairClusterId,
+	const int32 MacroGroupId,
 	const int32 HairMaterialId,
 	const int32 HairPrimitiveId)
 {
@@ -376,7 +374,7 @@ void FHairMaterialProcessor::Process(
 	}
 
 	FMeshPassProcessorRenderState DrawRenderState(PassDrawRenderState);
-	FHairMaterialShaderElementData ShaderElementData(HairClusterId, HairMaterialId, HairPrimitiveId);
+	FHairMaterialShaderElementData ShaderElementData(MacroGroupId, HairMaterialId, HairPrimitiveId);
 	ShaderElementData.InitializeMeshMaterialData(ViewIfDynamicMeshCommand, PrimitiveSceneProxy, MeshBatch, StaticMeshId, false);
 
 	BuildMeshDrawCommands(
@@ -407,13 +405,13 @@ FHairMaterialProcessor::FHairMaterialProcessor(
 /////////////////////////////////////////////////////////////////////////////////////////
 
 BEGIN_SHADER_PARAMETER_STRUCT(FVisibilityMaterialPassParameters, )
-	SHADER_PARAMETER_RDG_TEXTURE(Texture2D<uint>, NodeIndex)
-	SHADER_PARAMETER_RDG_BUFFER_SRV(StructuredBuffer<uint>,	NodeCoord)
-	SHADER_PARAMETER_RDG_BUFFER_SRV(StructuredBuffer<FNodeVis>, NodeVis)
-	SHADER_PARAMETER_RDG_BUFFER_SRV(Buffer<uint>, IndirectArgs)
-	SHADER_PARAMETER_RDG_BUFFER_UAV(RWStructuredBuffer<FNodeData>, OutNodeData)
-	SHADER_PARAMETER_RDG_BUFFER_UAV(RWBuffer<float2>, OutNodeVelocity)
-	RENDER_TARGET_BINDING_SLOTS()
+SHADER_PARAMETER_RDG_TEXTURE(Texture2D<uint>, NodeIndex)
+SHADER_PARAMETER_RDG_BUFFER_SRV(StructuredBuffer<uint>, NodeCoord)
+SHADER_PARAMETER_RDG_BUFFER_SRV(StructuredBuffer<FNodeVis>, NodeVis)
+SHADER_PARAMETER_RDG_BUFFER_SRV(Buffer<uint>, IndirectArgs)
+SHADER_PARAMETER_RDG_BUFFER_UAV(RWStructuredBuffer<FNodeData>, OutNodeData)
+SHADER_PARAMETER_RDG_BUFFER_UAV(RWBuffer<float2>, OutNodeVelocity)
+RENDER_TARGET_BINDING_SLOTS()
 END_SHADER_PARAMETER_STRUCT()
 
 struct FMaterialPassOutput
@@ -427,7 +425,7 @@ static FMaterialPassOutput AddHairMaterialPass(
 	FRDGBuilder& GraphBuilder,
 	const FScene* Scene,
 	const FViewInfo* ViewInfo,
-	const FHairStrandsClusterDatas& ClusterDatas,
+	const FHairStrandsMacroGroupDatas& MacroGroupDatas,
 	const uint32 NodeGroupSize,
 	FRDGTextureRef CompactNodeIndex,
 	FRDGBufferRef CompactNodeVis,
@@ -440,7 +438,7 @@ static FMaterialPassOutput AddHairMaterialPass(
 	const uint32 MaxNodeCount = CompactNodeVis->Desc.NumElements;
 
 	FMaterialPassOutput Output;
-	Output.NodeData		= GraphBuilder.CreateBuffer(FRDGBufferDesc::CreateStructuredDesc(sizeof(HairStrandsVisibilityInternal::NodeData), MaxNodeCount), TEXT("HairCompactNodeData"));
+	Output.NodeData = GraphBuilder.CreateBuffer(FRDGBufferDesc::CreateStructuredDesc(sizeof(HairStrandsVisibilityInternal::NodeData), MaxNodeCount), TEXT("HairCompactNodeData"));
 	Output.NodeVelocity = GraphBuilder.CreateBuffer(FRDGBufferDesc::CreateBufferDesc(4, CompactNodeVis->Desc.NumElements), TEXT("HairCompactNodeVelocity"));
 
 	const uint32 ResolutionDim = FMath::CeilToInt(FMath::Sqrt(MaxNodeCount));
@@ -459,12 +457,12 @@ static FMaterialPassOutput AddHairMaterialPass(
 
 	// Add resources reference to the pass parameters, in order to get the resource lifetime extended to this pass
 	FVisibilityMaterialPassParameters* PassParameters = GraphBuilder.AllocParameters<FVisibilityMaterialPassParameters>();
-	PassParameters->NodeIndex		= CompactNodeIndex;
-	PassParameters->NodeVis			= GraphBuilder.CreateSRV(CompactNodeVis);
-	PassParameters->NodeCoord		= GraphBuilder.CreateSRV(CompactNodeCoord);
-	PassParameters->IndirectArgs	= GraphBuilder.CreateSRV(IndirectArgBuffer);
-	PassParameters->OutNodeData		= GraphBuilder.CreateUAV(FRDGBufferUAVDesc(Output.NodeData));
-	PassParameters->OutNodeVelocity	= GraphBuilder.CreateUAV(FRDGBufferUAVDesc(Output.NodeVelocity, FMaterialPassOutput::VelocityFormat));
+	PassParameters->NodeIndex = CompactNodeIndex;
+	PassParameters->NodeVis = GraphBuilder.CreateSRV(CompactNodeVis);
+	PassParameters->NodeCoord = GraphBuilder.CreateSRV(CompactNodeCoord);
+	PassParameters->IndirectArgs = GraphBuilder.CreateSRV(IndirectArgBuffer);
+	PassParameters->OutNodeData = GraphBuilder.CreateUAV(FRDGBufferUAVDesc(Output.NodeData));
+	PassParameters->OutNodeVelocity = GraphBuilder.CreateUAV(FRDGBufferUAVDesc(Output.NodeVelocity, FMaterialPassOutput::VelocityFormat));
 	PassParameters->RenderTargets[0] = FRenderTargetBinding(OutDummyTexture0, ERenderTargetLoadAction::EClear, 0);
 
 
@@ -472,19 +470,19 @@ static FMaterialPassOutput AddHairMaterialPass(
 		RDG_EVENT_NAME("HairStrandsMaterialPass"),
 		PassParameters,
 		ERDGPassFlags::Raster,
-		[PassParameters, Scene = Scene, ViewInfo, &ClusterDatas, MaxNodeCount, Resolution, NodeGroupSize](FRHICommandListImmediate& RHICmdList)
+		[PassParameters, Scene = Scene, ViewInfo, &MacroGroupDatas, MaxNodeCount, Resolution, NodeGroupSize](FRHICommandListImmediate& RHICmdList)
 	{
 		check(RHICmdList.IsInsideRenderPass());
 		check(IsInRenderingThread());
 
 		FMaterialPassParameters MaterialPassParameters;
-		MaterialPassParameters.MaxResolution	= Resolution;
-		MaterialPassParameters.NodeGroupSize	= NodeGroupSize;
-		MaterialPassParameters.MaxSampleCount	= MaxNodeCount;
-		MaterialPassParameters.NodeIndex		= PassParameters->NodeIndex->GetPooledRenderTarget()->GetRenderTargetItem().ShaderResourceTexture;
-		MaterialPassParameters.NodeCoord		= PassParameters->NodeCoord->GetRHI();
-		MaterialPassParameters.NodeVis			= PassParameters->NodeVis->GetRHI();
-		MaterialPassParameters.IndirectArgs		= PassParameters->IndirectArgs->GetRHI();
+		MaterialPassParameters.MaxResolution = Resolution;
+		MaterialPassParameters.NodeGroupSize = NodeGroupSize;
+		MaterialPassParameters.MaxSampleCount = MaxNodeCount;
+		MaterialPassParameters.NodeIndex = PassParameters->NodeIndex->GetPooledRenderTarget()->GetRenderTargetItem().ShaderResourceTexture;
+		MaterialPassParameters.NodeCoord = PassParameters->NodeCoord->GetRHI();
+		MaterialPassParameters.NodeVis = PassParameters->NodeVis->GetRHI();
+		MaterialPassParameters.IndirectArgs = PassParameters->IndirectArgs->GetRHI();
 		TUniformBufferRef<FMaterialPassParameters> MaterialPassParametersBuffer = TUniformBufferRef<FMaterialPassParameters>::CreateUniformBufferImmediate(MaterialPassParameters, UniformBuffer_SingleFrame);
 
 		FMeshPassProcessorRenderState DrawRenderState(*ViewInfo, MaterialPassParametersBuffer);
@@ -500,20 +498,20 @@ static FMaterialPassOutput AddHairMaterialPass(
 			RHICmdList.SetViewport(0, 0, 0.0f, Resolution.X, Resolution.Y, 1.0f);
 			DrawRenderState.SetBlendState(TStaticBlendState<>::GetRHI());
 			DrawRenderState.SetDepthStencilState(TStaticDepthStencilState <false, CF_Always> ::GetRHI());
-			
+
 			FDynamicMeshDrawCommandStorage DynamicMeshDrawCommandStorage;
 			FMeshCommandOneFrameArray VisibleMeshDrawCommands;
 			FGraphicsMinimalPipelineStateSet PipelineStateSet;
 			FDynamicPassMeshDrawListContext ShadowContext(DynamicMeshDrawCommandStorage, VisibleMeshDrawCommands, PipelineStateSet);
 			FHairMaterialProcessor MeshProcessor(Scene, ViewInfo, DrawRenderState, &ShadowContext);
 
-			for (const FHairStrandsClusterData& ClusterData : ClusterDatas.Datas)
+			for (const FHairStrandsMacroGroupData& MacroGroupData : MacroGroupDatas.Datas)
 			{
-				for (const FHairStrandsClusterData::PrimitiveInfo& PrimitiveInfo : ClusterData.PrimitivesInfos)
+				for (const FHairStrandsMacroGroupData::PrimitiveInfo& PrimitiveInfo : MacroGroupData.PrimitivesInfos)
 				{
 					const FMeshBatch& MeshBatch = *PrimitiveInfo.MeshBatchAndRelevance.Mesh;
 					const uint64 BatchElementMask = ~0ull;
-					MeshProcessor.AddMeshBatch(MeshBatch, BatchElementMask, PrimitiveInfo.MeshBatchAndRelevance.PrimitiveSceneProxy, -1, ClusterData.ClusterId, PrimitiveInfo.MaterialId);
+					MeshProcessor.AddMeshBatch(MeshBatch, BatchElementMask, PrimitiveInfo.MeshBatchAndRelevance.PrimitiveSceneProxy, -1, MacroGroupData.MacroGroupId, PrimitiveInfo.MaterialId);
 				}
 			}
 
@@ -530,7 +528,7 @@ static FMaterialPassOutput AddHairMaterialPass(
 }
 
 ///////////////////////////////////////////////////////////////////////////////////////////////////
-class FHairVelocityCS: public FGlobalShader
+class FHairVelocityCS : public FGlobalShader
 {
 	DECLARE_GLOBAL_SHADER(FHairVelocityCS);
 	SHADER_USE_PARAMETER_STRUCT(FHairVelocityCS, FGlobalShader);
@@ -547,7 +545,7 @@ class FHairVelocityCS: public FGlobalShader
 		SHADER_PARAMETER_RDG_TEXTURE_UAV(Texture2D, OutVelocityTexture)
 		SHADER_PARAMETER_STRUCT_REF(FViewUniformShaderParameters, ViewUniformBuffer)
 		SHADER_PARAMETER_STRUCT_REF(FSceneTexturesUniformParameters, SceneTexturesStruct)
-	END_SHADER_PARAMETER_STRUCT()
+		END_SHADER_PARAMETER_STRUCT()
 
 public:
 
@@ -562,7 +560,7 @@ IMPLEMENT_GLOBAL_SHADER(FHairVelocityCS, "/Engine/Private/HairStrands/HairStrand
 static void AddHairVelocityPass(
 	FRDGBuilder& GraphBuilder,
 	const FViewInfo& View,
-	const FHairStrandsClusterDatas& ClusterDatas,
+	const FHairStrandsMacroGroupDatas& MacroGroupDatas,
 	FRDGTextureRef& NodeIndex,
 	FRDGBufferRef& NodeVis,
 	FRDGBufferRef& NodeVelocity,
@@ -590,7 +588,7 @@ static void AddHairVelocityPass(
 	PassParameters->NodeVelocity = GraphBuilder.CreateSRV(NodeVelocity, FMaterialPassOutput::VelocityFormat);
 	PassParameters->OutVelocityTexture = GraphBuilder.CreateUAV(OutVelocityTexture);
 
-	FIntRect TotalRect = ComputeVisibleHairStrandsClustersRect(View.ViewRect, ClusterDatas);
+	FIntRect TotalRect = ComputeVisibleHairStrandsMacroGroupsRect(View.ViewRect, MacroGroupDatas);
 
 	// Snap the rect onto thread group boundary
 	const FIntPoint GroupSize = GetVendorOptimalGroupSize2D();
@@ -613,22 +611,22 @@ static void AddHairVelocityPass(
 
 /////////////////////////////////////////////////////////////////////////////////////////
 BEGIN_GLOBAL_SHADER_PARAMETER_STRUCT(FVisibilityPassGlobalParameters, )
-	SHADER_PARAMETER(uint32, MaxPPLLNodeCount)
+SHADER_PARAMETER(uint32, MaxPPLLNodeCount)
 END_GLOBAL_SHADER_PARAMETER_STRUCT()
 IMPLEMENT_GLOBAL_SHADER_PARAMETER_STRUCT(FVisibilityPassGlobalParameters, "VisibilityPassGlobalParameters");
 
 BEGIN_SHADER_PARAMETER_STRUCT(FVisibilityPassParameters, )
-	SHADER_PARAMETER(uint32, HairVisibilityPass_MaxPPLLNodeCount)
-	SHADER_PARAMETER_RDG_TEXTURE_UAV(Texture2D, PPLLCounter)
-	SHADER_PARAMETER_RDG_TEXTURE_UAV(Texture2D, PPLLNodeIndex)
-	SHADER_PARAMETER_RDG_BUFFER_UAV(StructuredBuffer, PPLLNodeData)
-	RENDER_TARGET_BINDING_SLOTS()
+SHADER_PARAMETER(uint32, HairVisibilityPass_MaxPPLLNodeCount)
+SHADER_PARAMETER_RDG_TEXTURE_UAV(Texture2D, PPLLCounter)
+SHADER_PARAMETER_RDG_TEXTURE_UAV(Texture2D, PPLLNodeIndex)
+SHADER_PARAMETER_RDG_BUFFER_UAV(StructuredBuffer, PPLLNodeData)
+RENDER_TARGET_BINDING_SLOTS()
 END_SHADER_PARAMETER_STRUCT()
 
 template<EHairVisibilityRenderMode RenderMode>
 class FHairVisibilityVS : public FMeshMaterialShader
 {
-	DECLARE_SHADER_TYPE(FHairVisibilityVS,MeshMaterial);
+	DECLARE_SHADER_TYPE(FHairVisibilityVS, MeshMaterial);
 
 protected:
 
@@ -637,6 +635,7 @@ protected:
 	{
 		ERHIFeatureLevel::Type FeatureLevel = GetMaxSupportedFeatureLevel((EShaderPlatform)Initializer.Target.Platform);
 		check(FSceneInterface::GetShadingPath(FeatureLevel) != EShadingPath::Mobile);
+		PassUniformBuffer.Bind(Initializer.ParameterMap, FVisibilityPassGlobalParameters::StaticStructMetadata.GetShaderVariableName());
 	}
 
 	FHairVisibilityVS() {}
@@ -651,6 +650,7 @@ protected:
 		FMaterialShader::ModifyCompilationEnvironment(Parameters, OutEnvironment);
 		const uint32 RenderModeValue = uint32(RenderMode);
 		OutEnvironment.SetDefine(TEXT("HAIR_RENDER_MODE"), RenderModeValue);
+		OutEnvironment.SetDefine(TEXT("USE_CULLED_CLUSTER"), 1);
 	}
 };
 IMPLEMENT_MATERIAL_SHADER_TYPE(template<>, FHairVisibilityVS<HairVisibilityRenderMode_MSAA_Visibility>, TEXT("/Engine/Private/HairStrands/HairStrandsVisibilityVS.usf"), TEXT("Main"), SF_Vertex);
@@ -664,8 +664,8 @@ IMPLEMENT_MATERIAL_SHADER_TYPE(template<>, FHairVisibilityVS<HairVisibilityRende
 class FHairVisibilityShaderElementData : public FMeshMaterialShaderElementData
 {
 public:
-	FHairVisibilityShaderElementData(uint32 InHairClusterId, uint32 InHairMaterialId, uint32 InMaxPPLLNodeCount) : HairClusterId(InHairClusterId), HairMaterialId(InHairMaterialId), MaxPPLLNodeCount(InMaxPPLLNodeCount) { }
-	uint32 HairClusterId;
+	FHairVisibilityShaderElementData(uint32 InHairMacroGroupId, uint32 InHairMaterialId, uint32 InMaxPPLLNodeCount) : HairMacroGroupId(InHairMacroGroupId), HairMaterialId(InHairMaterialId), MaxPPLLNodeCount(InMaxPPLLNodeCount) { }
+	uint32 HairMacroGroupId;
 	uint32 HairMaterialId;
 	uint32 MaxPPLLNodeCount;
 };
@@ -678,17 +678,20 @@ class FHairVisibilityPS : public FMeshMaterialShader
 public:
 
 	FHairVisibilityPS(const ShaderMetaType::CompiledShaderInitializerType& Initializer)
-		:	FMeshMaterialShader(Initializer)
+		: FMeshMaterialShader(Initializer)
 	{
 		ERHIFeatureLevel::Type FeatureLevel = GetMaxSupportedFeatureLevel((EShaderPlatform)Initializer.Target.Platform);
 		check(FSceneInterface::GetShadingPath(FeatureLevel) != EShadingPath::Mobile);
 		PassUniformBuffer.Bind(Initializer.ParameterMap, FVisibilityPassGlobalParameters::StaticStructMetadata.GetShaderVariableName());
-		HairVisibilityPass_HairClusterIndex.Bind(Initializer.ParameterMap, TEXT("HairVisibilityPass_HairClusterIndex"));
+		HairVisibilityPass_HairMacroGroupIndex.Bind(Initializer.ParameterMap, TEXT("HairVisibilityPass_HairMacroGroupIndex"));
 		HairVisibilityPass_HairMaterialId.Bind(Initializer.ParameterMap, TEXT("HairVisibilityPass_HairMaterialId"));
 		HairVisibilityPass_MaxPPLLNodeCount.Bind(Initializer.ParameterMap, TEXT("HairVisibilityPass_MaxPPLLNodeCount"));
-		PPLLCounter.Bind(Initializer.ParameterMap, TEXT("PPLLCounter"));
-		PPLLNodeIndex.Bind(Initializer.ParameterMap, TEXT("PPLLNodeIndex"));
-		PPLLNodes.Bind(Initializer.ParameterMap, TEXT("PPLLNodes"));
+		if (RenderMode == EHairVisibilityRenderMode::HairVisibilityRenderMode_PPLL)
+		{
+			PPLLCounter.Bind(Initializer.ParameterMap, TEXT("PPLLCounter"));
+			PPLLNodeIndex.Bind(Initializer.ParameterMap, TEXT("PPLLNodeIndex"));
+			PPLLNodes.Bind(Initializer.ParameterMap, TEXT("PPLLNodes"));
+		}
 	}
 
 	FHairVisibilityPS() {}
@@ -698,7 +701,7 @@ public:
 		return IsCompatibleWithHairStrands(Parameters.Material, Parameters.Platform) && Parameters.VertexFactoryType->GetFName() == FName(TEXT("FHairStrandsVertexFactory"));
 	}
 
-	static void ModifyCompilationEnvironment(const FMaterialShaderPermutationParameters& Parameters, FShaderCompilerEnvironment& OutEnvironment)	
+	static void ModifyCompilationEnvironment(const FMaterialShaderPermutationParameters& Parameters, FShaderCompilerEnvironment& OutEnvironment)
 	{
 		FMaterialShader::ModifyCompilationEnvironment(Parameters, OutEnvironment);
 		const uint32 RenderModeValue = uint32(RenderMode);
@@ -708,12 +711,15 @@ public:
 	virtual bool Serialize(FArchive& Ar) override
 	{
 		bool bShaderHasOutdatedParameters = FMeshMaterialShader::Serialize(Ar);
-		Ar << HairVisibilityPass_HairClusterIndex;
+		Ar << HairVisibilityPass_HairMacroGroupIndex;
 		Ar << HairVisibilityPass_HairMaterialId;
 		Ar << HairVisibilityPass_MaxPPLLNodeCount;
-		Ar << PPLLCounter;
-		Ar << PPLLNodeIndex;
-		Ar << PPLLNodes;
+		if (RenderMode == EHairVisibilityRenderMode::HairVisibilityRenderMode_PPLL)
+		{
+			Ar << PPLLCounter;
+			Ar << PPLLNodeIndex;
+			Ar << PPLLNodes;
+		}
 		return bShaderHasOutdatedParameters;
 	}
 
@@ -728,12 +734,12 @@ public:
 		FMeshDrawSingleShaderBindings& ShaderBindings) const
 	{
 		FMeshMaterialShader::GetShaderBindings(Scene, FeatureLevel, PrimitiveSceneProxy, MaterialRenderProxy, Material, DrawRenderState, ShaderElementData, ShaderBindings);
-		ShaderBindings.Add(HairVisibilityPass_HairClusterIndex, ShaderElementData.HairClusterId);
+		ShaderBindings.Add(HairVisibilityPass_HairMacroGroupIndex, ShaderElementData.HairMacroGroupId);
 		ShaderBindings.Add(HairVisibilityPass_HairMaterialId, ShaderElementData.HairMaterialId);
 		ShaderBindings.Add(HairVisibilityPass_MaxPPLLNodeCount, ShaderElementData.MaxPPLLNodeCount);
 	}
 
-	FShaderParameter HairVisibilityPass_HairClusterIndex;
+	FShaderParameter HairVisibilityPass_HairMacroGroupIndex;
 	FShaderParameter HairVisibilityPass_HairMaterialId;
 	FShaderParameter HairVisibilityPass_MaxPPLLNodeCount;
 	// This is on FVisibilityPassParameters but needed to avoid the compiler to assert with unbound parameters 
@@ -760,7 +766,7 @@ public:
 		FDynamicPassMeshDrawListContext* InDrawListContext);
 
 	virtual void AddMeshBatch(const FMeshBatch& RESTRICT MeshBatch, uint64 BatchElementMask, const FPrimitiveSceneProxy* RESTRICT PrimitiveSceneProxy, int32 StaticMeshId = -1) override final;
-	void AddMeshBatch(const FMeshBatch& RESTRICT MeshBatch, uint64 BatchElementMask, const FPrimitiveSceneProxy* RESTRICT PrimitiveSceneProxy, int32 StaticMeshId, uint32 HairClusterId, uint32 HairMaterialId);
+	void AddMeshBatch(const FMeshBatch& RESTRICT MeshBatch, uint64 BatchElementMask, const FPrimitiveSceneProxy* RESTRICT PrimitiveSceneProxy, int32 StaticMeshId, uint32 HairMacroGroupId, uint32 HairMaterialId);
 
 	void SetMaxPPLLNodeCount(uint32 NewValue) { MaxPPLLNodeCount = NewValue; }
 
@@ -773,7 +779,7 @@ private:
 		int32 StaticMeshId,
 		const FMaterialRenderProxy& RESTRICT MaterialRenderProxy,
 		const FMaterial& RESTRICT MaterialResource,
-		const uint32 HairClusterId,
+		const uint32 HairMacroGroupId,
 		const uint32 HairMaterialId,
 		ERasterizerFillMode MeshFillMode,
 		ERasterizerCullMode MeshCullMode);
@@ -788,7 +794,7 @@ void FHairVisibilityProcessor::AddMeshBatch(const FMeshBatch& RESTRICT MeshBatch
 	AddMeshBatch(MeshBatch, BatchElementMask, PrimitiveSceneProxy, StaticMeshId, 0, 0);
 }
 
-void FHairVisibilityProcessor::AddMeshBatch(const FMeshBatch& RESTRICT MeshBatch, uint64 BatchElementMask, const FPrimitiveSceneProxy* RESTRICT PrimitiveSceneProxy, int32 StaticMeshId, uint32 HairClusterId, uint32 HairMaterialId)
+void FHairVisibilityProcessor::AddMeshBatch(const FMeshBatch& RESTRICT MeshBatch, uint64 BatchElementMask, const FPrimitiveSceneProxy* RESTRICT PrimitiveSceneProxy, int32 StaticMeshId, uint32 HairMacroGroupId, uint32 HairMaterialId)
 {
 	static const FVertexFactoryType* CompatibleVF = FVertexFactoryType::GetVFByName(TEXT("FHairStrandsVertexFactory"));
 
@@ -797,9 +803,9 @@ void FHairVisibilityProcessor::AddMeshBatch(const FMeshBatch& RESTRICT MeshBatch
 	const FMaterial& Material = MeshBatch.MaterialRenderProxy->GetMaterialWithFallback(FeatureLevel, FallbackMaterialRenderProxyPtr);
 	const bool bIsCompatible = IsCompatibleWithHairStrands(&Material, FeatureLevel);
 	const bool bIsHairStrandsFactory = MeshBatch.VertexFactory->GetType()->GetId() == CompatibleVF->GetId();
-	const bool bShouldRender = (!PrimitiveSceneProxy && MeshBatch.Elements.Num()>0) || (PrimitiveSceneProxy && PrimitiveSceneProxy->ShouldRenderInMainPass());
+	const bool bShouldRender = (!PrimitiveSceneProxy && MeshBatch.Elements.Num() > 0) || (PrimitiveSceneProxy && PrimitiveSceneProxy->ShouldRenderInMainPass());
 
-	if (bIsCompatible 
+	if (bIsCompatible
 		&& bIsHairStrandsFactory
 		&& bShouldRender
 		&& ShouldIncludeDomainInMeshPass(Material.GetMaterialDomain()))
@@ -808,27 +814,27 @@ void FHairVisibilityProcessor::AddMeshBatch(const FMeshBatch& RESTRICT MeshBatch
 		const ERasterizerFillMode MeshFillMode = ComputeMeshFillMode(MeshBatch, Material);
 		const ERasterizerCullMode MeshCullMode = ComputeMeshCullMode(MeshBatch, Material);
 		if (RenderMode == HairVisibilityRenderMode_MSAA_Visibility)
-			Process<HairVisibilityRenderMode_MSAA_Visibility>(MeshBatch, BatchElementMask, PrimitiveSceneProxy, StaticMeshId, MaterialRenderProxy, Material, HairClusterId, HairMaterialId, MeshFillMode, MeshCullMode);
+			Process<HairVisibilityRenderMode_MSAA_Visibility>(MeshBatch, BatchElementMask, PrimitiveSceneProxy, StaticMeshId, MaterialRenderProxy, Material, HairMacroGroupId, HairMaterialId, MeshFillMode, MeshCullMode);
 		if (RenderMode == HairVisibilityRenderMode_MSAA)
-			Process<HairVisibilityRenderMode_MSAA>(MeshBatch, BatchElementMask, PrimitiveSceneProxy, StaticMeshId, MaterialRenderProxy, Material, HairClusterId, HairMaterialId, MeshFillMode, MeshCullMode);
+			Process<HairVisibilityRenderMode_MSAA>(MeshBatch, BatchElementMask, PrimitiveSceneProxy, StaticMeshId, MaterialRenderProxy, Material, HairMacroGroupId, HairMaterialId, MeshFillMode, MeshCullMode);
 		if (RenderMode == HairVisibilityRenderMode_Transmittance)
-			Process<HairVisibilityRenderMode_Transmittance>(MeshBatch, BatchElementMask, PrimitiveSceneProxy, StaticMeshId, MaterialRenderProxy, Material, HairClusterId, HairMaterialId, MeshFillMode, MeshCullMode);
+			Process<HairVisibilityRenderMode_Transmittance>(MeshBatch, BatchElementMask, PrimitiveSceneProxy, StaticMeshId, MaterialRenderProxy, Material, HairMacroGroupId, HairMaterialId, MeshFillMode, MeshCullMode);
 		if (RenderMode == HairVisibilityRenderMode_TransmittanceAndHairCount)
-			Process<HairVisibilityRenderMode_TransmittanceAndHairCount>(MeshBatch, BatchElementMask, PrimitiveSceneProxy, StaticMeshId, MaterialRenderProxy, Material, HairClusterId, HairMaterialId, MeshFillMode, MeshCullMode);
+			Process<HairVisibilityRenderMode_TransmittanceAndHairCount>(MeshBatch, BatchElementMask, PrimitiveSceneProxy, StaticMeshId, MaterialRenderProxy, Material, HairMacroGroupId, HairMaterialId, MeshFillMode, MeshCullMode);
 		if (RenderMode == HairVisibilityRenderMode_PPLL)
-			Process<HairVisibilityRenderMode_PPLL>(MeshBatch, BatchElementMask, PrimitiveSceneProxy, StaticMeshId, MaterialRenderProxy, Material, HairClusterId, HairMaterialId, MeshFillMode, MeshCullMode);
+			Process<HairVisibilityRenderMode_PPLL>(MeshBatch, BatchElementMask, PrimitiveSceneProxy, StaticMeshId, MaterialRenderProxy, Material, HairMacroGroupId, HairMaterialId, MeshFillMode, MeshCullMode);
 	}
 }
 
 template<EHairVisibilityRenderMode TRenderMode>
 void FHairVisibilityProcessor::Process(
-	const FMeshBatch& MeshBatch, 
-	uint64 BatchElementMask, 
+	const FMeshBatch& MeshBatch,
+	uint64 BatchElementMask,
 	const FPrimitiveSceneProxy* RESTRICT PrimitiveSceneProxy,
 	int32 StaticMeshId,
 	const FMaterialRenderProxy& RESTRICT MaterialRenderProxy,
 	const FMaterial& RESTRICT MaterialResource,
-	const uint32 HairClusterId,
+	const uint32 HairMacroGroupId,
 	const uint32 HairMaterialId,
 	ERasterizerFillMode MeshFillMode,
 	ERasterizerCullMode MeshCullMode)
@@ -843,11 +849,11 @@ void FHairVisibilityProcessor::Process(
 	{
 		FVertexFactoryType* VertexFactoryType = VertexFactory->GetType();
 		PassShaders.VertexShader = MaterialResource.GetShader<FHairVisibilityVS<TRenderMode>>(VertexFactoryType);
-		PassShaders.PixelShader  = MaterialResource.GetShader<FHairVisibilityPS<TRenderMode>>(VertexFactoryType);
+		PassShaders.PixelShader = MaterialResource.GetShader<FHairVisibilityPS<TRenderMode>>(VertexFactoryType);
 	}
 
 	FMeshPassProcessorRenderState DrawRenderState(PassDrawRenderState);
-	FHairVisibilityShaderElementData ShaderElementData(HairClusterId, HairMaterialId, MaxPPLLNodeCount);
+	FHairVisibilityShaderElementData ShaderElementData(HairMacroGroupId, HairMaterialId, MaxPPLLNodeCount);
 	ShaderElementData.InitializeMeshMaterialData(ViewIfDynamicMeshCommand, PrimitiveSceneProxy, MeshBatch, StaticMeshId, false);
 
 	BuildMeshDrawCommands(
@@ -887,7 +893,7 @@ class FClearUIntGraphicPS : public FGlobalShader
 	BEGIN_SHADER_PARAMETER_STRUCT(FParameters, )
 		SHADER_PARAMETER(uint32, ClearValue)
 		RENDER_TARGET_BINDING_SLOTS()
-	END_SHADER_PARAMETER_STRUCT()
+		END_SHADER_PARAMETER_STRUCT()
 
 public:
 	static bool ShouldCompilePermutation(const FGlobalShaderPermutationParameters& Parameters) { return IsHairStrandsSupported(Parameters.Platform); }
@@ -940,9 +946,9 @@ static void AddClearGraphicPass(
 		DrawRectangle(
 			RHICmdList,
 			0, 0,
-			Viewport.Width(),Viewport.Height(),
-			Viewport.Min.X,  Viewport.Min.Y,
-			Viewport.Width(),Viewport.Height(),
+			Viewport.Width(), Viewport.Height(),
+			Viewport.Min.X, Viewport.Min.Y,
+			Viewport.Width(), Viewport.Height(),
 			Viewport.Size(),
 			Resolution,
 			*VertexShader,
@@ -951,7 +957,7 @@ static void AddClearGraphicPass(
 }
 
 BEGIN_SHADER_PARAMETER_STRUCT(FClearUAVTextureParameters, )
-	SHADER_PARAMETER_RDG_TEXTURE_UAV(RWTexture2D, TextureUAV)
+SHADER_PARAMETER_RDG_TEXTURE_UAV(RWTexture2D, TextureUAV)
 END_SHADER_PARAMETER_STRUCT()
 
 void AddClearUAVUintPass(
@@ -959,7 +965,7 @@ void AddClearUAVUintPass(
 	FRDGEventName&& PassName,
 	FRDGTextureRef Texture,
 	uint32 Value)
-{	
+{
 	FClearUAVTextureParameters* Parameters = GraphBuilder.AllocParameters< FClearUAVTextureParameters >();
 	Parameters->TextureUAV = GraphBuilder.CreateUAV(Texture);
 
@@ -984,9 +990,10 @@ class FCopyIndirectBufferCS : public FGlobalShader
 
 	BEGIN_SHADER_PARAMETER_STRUCT(FParameters, )
 		SHADER_PARAMETER(uint32, ThreadGroupSize)
+		SHADER_PARAMETER(uint32, ItemCountPerGroup)
 		SHADER_PARAMETER_RDG_TEXTURE(Texture2D, CounterTexture)
 		SHADER_PARAMETER_RDG_BUFFER_UAV(RWBuffer, OutArgBuffer)
-	END_SHADER_PARAMETER_STRUCT()
+		END_SHADER_PARAMETER_STRUCT()
 
 public:
 	static bool ShouldCompilePermutation(const FGlobalShaderPermutationParameters& Parameters) { return IsHairStrandsSupported(Parameters.Platform); }
@@ -998,6 +1005,7 @@ static FRDGBufferRef AddCopyIndirectArgPass(
 	FRDGBuilder& GraphBuilder,
 	const FViewInfo* View,
 	const uint32 ThreadGroupSize,
+	const uint32 ItemCountPerGroup,
 	FRDGTextureRef CounterTexture)
 {
 	check(CounterTexture);
@@ -1005,7 +1013,8 @@ static FRDGBufferRef AddCopyIndirectArgPass(
 	FRDGBufferRef OutBuffer = GraphBuilder.CreateBuffer(FRDGBufferDesc::CreateIndirectDesc<FRHIDispatchIndirectParameters>(), TEXT("HairVisibilityIndirectArgBuffer"));
 
 	FCopyIndirectBufferCS::FParameters* Parameters = GraphBuilder.AllocParameters<FCopyIndirectBufferCS::FParameters>();
-	Parameters->ThreadGroupSize = 32;
+	Parameters->ThreadGroupSize = ThreadGroupSize;
+	Parameters->ItemCountPerGroup = ItemCountPerGroup;
 	Parameters->CounterTexture = CounterTexture;
 	Parameters->OutArgBuffer = GraphBuilder.CreateUAV(OutBuffer);
 
@@ -1016,7 +1025,7 @@ static FRDGBufferRef AddCopyIndirectArgPass(
 		RDG_EVENT_NAME("HairStrandsVisbilityCopyIndirectArgs"),
 		*ComputeShader,
 		Parameters,
-		FIntVector(1,1,1));
+		FIntVector(1, 1, 1));
 
 	return OutBuffer;
 }
@@ -1027,12 +1036,12 @@ class FHairVisibilityPrimitiveIdCompactionCS : public FGlobalShader
 	DECLARE_GLOBAL_SHADER(FHairVisibilityPrimitiveIdCompactionCS);
 	SHADER_USE_PARAMETER_STRUCT(FHairVisibilityPrimitiveIdCompactionCS, FGlobalShader);
 
-	class FVendor		: SHADER_PERMUTATION_INT("PERMUTATION_VENDOR", HairVisibilityVendorCount);
-	class FVelocity		: SHADER_PERMUTATION_INT("PERMUTATION_VELOCITY", 4);
+	class FVendor : SHADER_PERMUTATION_INT("PERMUTATION_VENDOR", HairVisibilityVendorCount);
+	class FVelocity : SHADER_PERMUTATION_INT("PERMUTATION_VELOCITY", 4);
 	class FViewTransmittance : SHADER_PERMUTATION_INT("PERMUTATION_VIEWTRANSMITTANCE", 2);
-	class FMaterial		: SHADER_PERMUTATION_INT("PERMUTATION_MATERIAL_COMPACTION", 2);
-	class FPPLL			: SHADER_PERMUTATION_SPARSE_INT("PERMUTATION_PPLL", 0, 8, 16, 32); // See GetPPLLMaxRenderNodePerPixel
-	class FVisibility	: SHADER_PERMUTATION_INT("PERMUTATION_VISIBILITY", 2);
+	class FMaterial : SHADER_PERMUTATION_INT("PERMUTATION_MATERIAL_COMPACTION", 2);
+	class FPPLL : SHADER_PERMUTATION_SPARSE_INT("PERMUTATION_PPLL", 0, 8, 16, 32); // See GetPPLLMaxRenderNodePerPixel
+	class FVisibility : SHADER_PERMUTATION_INT("PERMUTATION_VISIBILITY", 2);
 	using FPermutationDomain = TShaderPermutationDomain<FVendor, FVelocity, FViewTransmittance, FMaterial, FPPLL, FVisibility>;
 
 	BEGIN_SHADER_PARAMETER_STRUCT(FParameters, )
@@ -1055,7 +1064,7 @@ class FHairVisibilityPrimitiveIdCompactionCS : public FGlobalShader
 		SHADER_PARAMETER_RDG_BUFFER_SRV(StructuredBuffer, PPLLNodeData)
 
 		SHADER_PARAMETER_RDG_TEXTURE(Texture2D, ViewTransmittanceTexture)
-		
+
 		SHADER_PARAMETER_RDG_TEXTURE_UAV(Texture2D, OutCompactNodeCounter)
 		SHADER_PARAMETER_RDG_TEXTURE_UAV(Texture2D, OutCompactNodeIndex)
 		SHADER_PARAMETER_RDG_TEXTURE_UAV(Texture2D, OutCategorizationTexture)
@@ -1065,7 +1074,7 @@ class FHairVisibilityPrimitiveIdCompactionCS : public FGlobalShader
 
 		SHADER_PARAMETER_STRUCT_REF(FViewUniformShaderParameters, ViewUniformBuffer)
 		SHADER_PARAMETER_STRUCT_REF(FSceneTexturesUniformParameters, SceneTexturesStruct)
-	END_SHADER_PARAMETER_STRUCT()
+		END_SHADER_PARAMETER_STRUCT()
 
 public:
 
@@ -1074,7 +1083,7 @@ public:
 		if (PermutationVector.Get<FPPLL>() > 0)
 		{
 			PermutationVector.Set<FViewTransmittance>(0);
-			PermutationVector.Set<FVisibility>(0); 
+			PermutationVector.Set<FVisibility>(0);
 		}
 		return PermutationVector;
 	}
@@ -1101,7 +1110,7 @@ static void AddHairVisibilityPrimitiveIdCompactionPass(
 	const bool bUseVisibility,
 	FRDGBuilder& GraphBuilder,
 	const FViewInfo& View,
-	const FHairStrandsClusterDatas& ClusterDatas,
+	const FHairStrandsMacroGroupDatas& MacroGroupDatas,
 	const uint32 NodeGroupSize,
 	FHairVisibilityPrimitiveIdCompactionCS::FParameters* PassParameters,
 	FRDGTextureRef& OutCompactNodeIndex,
@@ -1137,7 +1146,7 @@ static void AddHairVisibilityPrimitiveIdCompactionPass(
 		}
 	}
 
-	
+
 	FRDGTextureRef CompactCounter;
 	{
 		FRDGTextureDesc Desc;
@@ -1174,7 +1183,7 @@ static void AddHairVisibilityPrimitiveIdCompactionPass(
 		OutputDesc.TargetableFlags = TexCreate_UAV | TexCreate_ShaderResource;
 		OutCategorizationTexture = GraphBuilder.CreateTexture(OutputDesc, TEXT("CategorizationTexture"));
 	}
-	
+
 	AddClearUAVUintPass(GraphBuilder, RDG_EVENT_NAME("HairStrandsClearCompactionCounter"), CompactCounter, 0);
 	AddClearUAVUintPass(GraphBuilder, RDG_EVENT_NAME("HairStrandsClearCompactionOffsetAndCount"), OutCompactNodeIndex, 0);
 	AddClearUAVUintPass(GraphBuilder, RDG_EVENT_NAME("HairStrandsClearCategorizationTexture"), OutCategorizationTexture, 0);
@@ -1223,8 +1232,8 @@ static void AddHairVisibilityPrimitiveIdCompactionPass(
 	{
 		PassParameters->OutVelocityTexture = GraphBuilder.CreateUAV(OutVelocityTexture);
 	}
- 
-	FIntRect TotalRect = ComputeVisibleHairStrandsClustersRect(View.ViewRect, ClusterDatas);
+
+	FIntRect TotalRect = ComputeVisibleHairStrandsMacroGroupsRect(View.ViewRect, MacroGroupDatas);
 
 	// Snap the rect onto thread group boundary
 	const FIntPoint GroupSize = GetVendorOptimalGroupSize2D();
@@ -1232,7 +1241,7 @@ static void AddHairVisibilityPrimitiveIdCompactionPass(
 	TotalRect.Min.Y = FMath::FloorToInt(float(TotalRect.Min.Y) / float(GroupSize.Y)) * GroupSize.Y;
 	TotalRect.Max.X = FMath::CeilToInt(float(TotalRect.Max.X) / float(GroupSize.X)) * GroupSize.X;
 	TotalRect.Max.Y = FMath::CeilToInt(float(TotalRect.Max.Y) / float(GroupSize.Y)) * GroupSize.Y;
-	
+
 	FIntPoint RectResolution(TotalRect.Width(), TotalRect.Height());
 	PassParameters->ResolutionOffset = FIntPoint(TotalRect.Min.X, TotalRect.Min.Y);
 
@@ -1244,7 +1253,105 @@ static void AddHairVisibilityPrimitiveIdCompactionPass(
 		PassParameters,
 		FComputeShaderUtils::GetGroupCount(RectResolution, GroupSize));
 
-	OutIndirectArgsBuffer = AddCopyIndirectArgPass(GraphBuilder, &View, NodeGroupSize, CompactCounter);
+	OutIndirectArgsBuffer = AddCopyIndirectArgPass(GraphBuilder, &View, NodeGroupSize, 1, CompactCounter);
+}
+
+
+///////////////////////////////////////////////////////////////////////////////////////////////////
+class FHairGenerateTileCS : public FGlobalShader
+{
+	DECLARE_GLOBAL_SHADER(FHairGenerateTileCS);
+	SHADER_USE_PARAMETER_STRUCT(FHairGenerateTileCS, FGlobalShader);
+
+	class FTileSize : SHADER_PERMUTATION_INT("PERMUTATION_TILESIZE", 2);
+	using FPermutationDomain = TShaderPermutationDomain<FTileSize>;
+
+	BEGIN_SHADER_PARAMETER_STRUCT(FParameters, )
+		SHADER_PARAMETER(FIntPoint, Resolution)
+		SHADER_PARAMETER(FIntPoint, TileResolution)
+		SHADER_PARAMETER_RDG_TEXTURE(Texture2D, CategorizationTexture)
+		SHADER_PARAMETER_RDG_TEXTURE_UAV(Texture2D, OutTileCounter)
+		SHADER_PARAMETER_RDG_TEXTURE_UAV(RWTexture2D, OutTileIndexTexture)
+		SHADER_PARAMETER_RDG_BUFFER_UAV(RWBuffer, OutTileBuffer)
+		END_SHADER_PARAMETER_STRUCT()
+
+public:
+
+	static bool ShouldCompilePermutation(const FGlobalShaderPermutationParameters& Parameters)
+	{
+		return IsHairStrandsSupported(Parameters.Platform);
+	}
+};
+
+IMPLEMENT_GLOBAL_SHADER(FHairGenerateTileCS, "/Engine/Private/HairStrands/HairStrandsVisibilityTile.usf", "MainCS", SF_Compute);
+
+static void AddGenerateTilePass(
+	FRDGBuilder& GraphBuilder,
+	const FViewInfo& View,
+	const uint32 ThreadGroupSize,
+	const uint32 TileSize,
+	const FRDGTextureRef& CategorizationTexture,
+	FRDGTextureRef& OutTileIndexTexture,
+	FRDGBufferRef& OutTileBuffer,
+	FRDGBufferRef& OutTileIndirectArgs)
+{
+	check(TileSize == 8); // only size supported for now
+	const FIntPoint Resolution = CategorizationTexture->Desc.Extent;
+	const FIntPoint TileResolution = FIntPoint(FMath::CeilToInt(Resolution.X / float(TileSize)), FMath::CeilToInt(Resolution.Y / float(TileSize)));
+	const uint32 TileCount = TileResolution.X * TileResolution.Y;
+
+	FRDGTextureRef TileCounter;
+	{
+		FRDGTextureDesc Desc;
+		Desc.Extent.X = 1;
+		Desc.Extent.Y = 1;
+		Desc.Depth = 0;
+		Desc.Format = PF_R32_UINT;
+		Desc.NumMips = 1;
+		Desc.NumSamples = 1;
+		Desc.Flags = TexCreate_None;
+		Desc.TargetableFlags = TexCreate_UAV | TexCreate_ShaderResource;
+		Desc.ClearValue = FClearValueBinding(0);
+		TileCounter = GraphBuilder.CreateTexture(Desc, TEXT("HairTileCounter"));
+	}
+
+	{
+		FRDGTextureDesc Desc;
+		Desc.Extent = TileResolution;
+		Desc.Depth = 0;
+		Desc.Format = PF_R32_UINT;
+		Desc.NumMips = 1;
+		Desc.NumSamples = 1;
+		Desc.Flags = TexCreate_None;
+		Desc.TargetableFlags = TexCreate_UAV | TexCreate_ShaderResource;
+		Desc.ClearValue = FClearValueBinding(0);
+		OutTileIndexTexture = GraphBuilder.CreateTexture(Desc, TEXT("HairTileIndexTexture"));
+	}
+
+	OutTileBuffer = GraphBuilder.CreateBuffer(FRDGBufferDesc::CreateBufferDesc(sizeof(uint32), TileCount), TEXT("HairTileBuffer"));
+
+	AddClearRenderTargetPass(GraphBuilder, TileCounter, FLinearColor(0.0f, 0.0f, 0.0f, 1.0f));
+
+	FHairGenerateTileCS::FPermutationDomain PermutationVector;
+	PermutationVector.Set<FHairGenerateTileCS::FTileSize>(0);
+
+	FHairGenerateTileCS::FParameters* PassParameters = GraphBuilder.AllocParameters<FHairGenerateTileCS::FParameters>();
+	PassParameters->Resolution = Resolution;
+	PassParameters->TileResolution = TileResolution;
+	PassParameters->CategorizationTexture = CategorizationTexture;
+	PassParameters->OutTileCounter = GraphBuilder.CreateUAV(TileCounter);
+	PassParameters->OutTileIndexTexture = GraphBuilder.CreateUAV(OutTileIndexTexture);
+	PassParameters->OutTileBuffer = GraphBuilder.CreateUAV(OutTileBuffer, PF_R16G16_UINT);
+
+	TShaderMapRef<FHairGenerateTileCS> ComputeShader(View.ShaderMap, PermutationVector);
+	FComputeShaderUtils::AddPass(
+		GraphBuilder,
+		RDG_EVENT_NAME("HairGenerateTile"),
+		*ComputeShader,
+		PassParameters,
+		FIntVector(TileResolution.X, TileResolution.Y, 1));
+
+	OutTileIndirectArgs = AddCopyIndirectArgPass(GraphBuilder, &View, ThreadGroupSize, TileSize*TileSize, TileCounter);
 }
 
 ///////////////////////////////////////////////////////////////////////////////////////////////////
@@ -1260,7 +1367,7 @@ class FHairVisibilityFillOpaqueDepthPS : public FGlobalShader
 
 		SHADER_PARAMETER_STRUCT_REF(FViewUniformShaderParameters, ViewUniformBuffer)
 		RENDER_TARGET_BINDING_SLOTS()
-	END_SHADER_PARAMETER_STRUCT()
+		END_SHADER_PARAMETER_STRUCT()
 
 public:
 	static bool ShouldCompilePermutation(const FGlobalShaderPermutationParameters& Parameters) { return IsHairStrandsSupported(Parameters.Platform); }
@@ -1272,7 +1379,7 @@ static FRDGTextureRef AddHairVisibilityFillOpaqueDepth(
 	FRDGBuilder& GraphBuilder,
 	const FViewInfo& View,
 	const FIntPoint& Resolution,
-	const FHairStrandsClusterDatas& ClusterDatas,
+	const FHairStrandsMacroGroupDatas& MacroGroupDatas,
 	const FRDGTextureRef& SceneDepthTexture)
 {
 	FRDGTextureRef OutVisibilityDepthTexture;
@@ -1308,17 +1415,17 @@ static FRDGTextureRef AddHairVisibilityFillOpaqueDepth(
 	const FIntRect Viewport = View.ViewRect;
 	const FViewInfo* CapturedView = &View;
 
-	TArray<FIntRect> ClusterRects;
+	TArray<FIntRect> MacroGroupRects;
 	if (IsHairStrandsViewRectOptimEnable())
 	{
-		for (const FHairStrandsClusterData& ClusterData : ClusterDatas.Datas)
+		for (const FHairStrandsMacroGroupData& MacroGroupData : MacroGroupDatas.Datas)
 		{
-			ClusterRects.Add(ClusterData.ScreenRect);
+			MacroGroupRects.Add(MacroGroupData.ScreenRect);
 		}
 	}
 	else
 	{
-		ClusterRects.Add(Viewport);
+		MacroGroupRects.Add(Viewport);
 	}
 
 	{
@@ -1328,7 +1435,7 @@ static FRDGTextureRef AddHairVisibilityFillOpaqueDepth(
 			RDG_EVENT_NAME("HairStrandsVisibilityFillOpaqueDepth"),
 			Parameters,
 			ERDGPassFlags::Raster,
-			[Parameters, VertexShader, PixelShader, ClusterRects, Viewport, Resolution, CapturedView](FRHICommandList& RHICmdList)
+			[Parameters, VertexShader, PixelShader, MacroGroupRects, Viewport, Resolution, CapturedView](FRHICommandList& RHICmdList)
 		{
 			FGraphicsPipelineStateInitializer GraphicsPSOInit;
 			RHICmdList.ApplyCachedRenderTargets(GraphicsPSOInit);
@@ -1345,14 +1452,14 @@ static FRDGTextureRef AddHairVisibilityFillOpaqueDepth(
 			VertexShader->SetParameters(RHICmdList, CapturedView->ViewUniformBuffer);
 			SetShaderParameters(RHICmdList, *PixelShader, PixelShader->GetPixelShader(), *Parameters);
 
-			for (const FIntRect& ViewRect : ClusterRects)
+			for (const FIntRect& ViewRect : MacroGroupRects)
 			{
 				RHICmdList.SetViewport(ViewRect.Min.X, ViewRect.Min.Y, 0.0f, ViewRect.Max.X, ViewRect.Max.Y, 1.0f);
 				DrawRectangle(
 					RHICmdList,
 					0, 0,
 					Viewport.Width(), Viewport.Height(),
-					Viewport.Min.X,   Viewport.Min.Y,
+					Viewport.Min.X, Viewport.Min.Y,
 					Viewport.Width(), Viewport.Height(),
 					Viewport.Size(),
 					Resolution,
@@ -1373,7 +1480,7 @@ static void AddHairVisibilityCommonPass(
 	FRDGBuilder& GraphBuilder,
 	const FScene* Scene,
 	const FViewInfo* ViewInfo,
-	const FHairStrandsClusterDatas& ClusterDatas,
+	const FHairStrandsMacroGroupDatas& MacroGroupDatas,
 	const EHairVisibilityRenderMode RenderMode,
 	TPassParameters* PassParameters)
 {
@@ -1392,17 +1499,17 @@ static void AddHairVisibilityCommonPass(
 	};
 
 	GraphBuilder.AddPass(
-		GetPassName(), 
+		GetPassName(),
 		PassParameters,
 		ERDGPassFlags::Raster,
-		[PassParameters, Scene = Scene, ViewInfo, &ClusterDatas, RenderMode](FRHICommandListImmediate& RHICmdList)
+		[PassParameters, Scene = Scene, ViewInfo, &MacroGroupDatas, RenderMode](FRHICommandListImmediate& RHICmdList)
 	{
 		check(RHICmdList.IsInsideRenderPass());
 		check(IsInRenderingThread());
 
 		FMeshPassProcessorRenderState DrawRenderState(*ViewInfo);
 
-		if(RenderMode == HairVisibilityRenderMode_Transmittance || RenderMode == HairVisibilityRenderMode_TransmittanceAndHairCount || RenderMode == HairVisibilityRenderMode_PPLL)
+		if (RenderMode == HairVisibilityRenderMode_Transmittance || RenderMode == HairVisibilityRenderMode_TransmittanceAndHairCount || RenderMode == HairVisibilityRenderMode_PPLL)
 		{
 			TUniformBufferRef<FViewUniformShaderParameters> ViewUniformShaderParameters;
 			const bool bEnableMSAA = false;
@@ -1411,7 +1518,7 @@ static void AddHairVisibilityCommonPass(
 			ViewUniformShaderParameters = TUniformBufferRef<FViewUniformShaderParameters>::CreateUniformBufferImmediate(*ViewInfo->CachedViewUniformShaderParameters, UniformBuffer_SingleFrame);
 			DrawRenderState.SetViewUniformBuffer(ViewUniformShaderParameters);
 		}
-		
+
 		{
 			RHICmdList.SetViewport(0, 0, 0.0f, ViewInfo->ViewRect.Width(), ViewInfo->ViewRect.Height(), 1.0f);
 			if (RenderMode == HairVisibilityRenderMode_MSAA)
@@ -1421,7 +1528,7 @@ static void AddHairVisibilityCommonPass(
 					CW_RGBA, BO_Add, BF_One, BF_Zero, BO_Add, BF_One, BF_Zero>::GetRHI());
 				DrawRenderState.SetDepthStencilState(TStaticDepthStencilState<true, CF_DepthNearOrEqual>::GetRHI());
 			}
-			else if(RenderMode == HairVisibilityRenderMode_MSAA_Visibility)
+			else if (RenderMode == HairVisibilityRenderMode_MSAA_Visibility)
 			{
 				DrawRenderState.SetBlendState(TStaticBlendState<CW_RGBA, BO_Add, BF_One, BF_Zero, BO_Add, BF_One, BF_Zero>::GetRHI());
 				DrawRenderState.SetDepthStencilState(TStaticDepthStencilState<true, CF_DepthNearOrEqual>::GetRHI());
@@ -1456,13 +1563,13 @@ static void AddHairVisibilityCommonPass(
 				MeshProcessor.SetMaxPPLLNodeCount(PassParameters->HairVisibilityPass_MaxPPLLNodeCount);
 			}
 
-			for (const FHairStrandsClusterData& ClusterData : ClusterDatas.Datas)
+			for (const FHairStrandsMacroGroupData& MacroGroupData : MacroGroupDatas.Datas)
 			{
-				for (const FHairStrandsClusterData::PrimitiveInfo& PrimitiveInfo : ClusterData.PrimitivesInfos)
+				for (const FHairStrandsMacroGroupData::PrimitiveInfo& PrimitiveInfo : MacroGroupData.PrimitivesInfos)
 				{
 					const FMeshBatch& MeshBatch = *PrimitiveInfo.MeshBatchAndRelevance.Mesh;
 					const uint64 BatchElementMask = ~0ull;
-					MeshProcessor.AddMeshBatch(MeshBatch, BatchElementMask, PrimitiveInfo.MeshBatchAndRelevance.PrimitiveSceneProxy, -1, ClusterData.ClusterId, PrimitiveInfo.MaterialId);
+					MeshProcessor.AddMeshBatch(MeshBatch, BatchElementMask, PrimitiveInfo.MeshBatchAndRelevance.PrimitiveSceneProxy, -1, MacroGroupData.MacroGroupId, PrimitiveInfo.MaterialId);
 				}
 			}
 
@@ -1481,7 +1588,7 @@ static void AddHairVisibilityMSAAPass(
 	FRDGBuilder& GraphBuilder,
 	const FScene* Scene,
 	const FViewInfo* ViewInfo,
-	const FHairStrandsClusterDatas& ClusterDatas,
+	const FHairStrandsMacroGroupDatas& MacroGroupDatas,
 	const FIntPoint& Resolution,
 	FRDGTextureRef& OutVisibilityIdTexture,
 	FRDGTextureRef& OutVisibilityMaterialTexture,
@@ -1519,7 +1626,7 @@ static void AddHairVisibilityMSAAPass(
 			ERenderTargetLoadAction::ELoad,
 			ERenderTargetLoadAction::ENoAction,
 			FExclusiveDepthStencil::DepthWrite_StencilNop);
-		AddHairVisibilityCommonPass(GraphBuilder, Scene, ViewInfo, ClusterDatas, HairVisibilityRenderMode_MSAA_Visibility, PassParameters);
+		AddHairVisibilityCommonPass(GraphBuilder, Scene, ViewInfo, MacroGroupDatas, HairVisibilityRenderMode_MSAA_Visibility, PassParameters);
 	}
 	else
 	{
@@ -1597,16 +1704,16 @@ static void AddHairVisibilityMSAAPass(
 
 		FVisibilityPassParameters* PassParameters = GraphBuilder.AllocParameters<FVisibilityPassParameters>();
 		PassParameters->RenderTargets[0] = FRenderTargetBinding(OutVisibilityIdTexture, ERenderTargetLoadAction::ELoad, 0);
-		PassParameters->RenderTargets[1] = FRenderTargetBinding(OutVisibilityMaterialTexture,  LoadAction, 0);
+		PassParameters->RenderTargets[1] = FRenderTargetBinding(OutVisibilityMaterialTexture, LoadAction, 0);
 		PassParameters->RenderTargets[2] = FRenderTargetBinding(OutVisibilityAttributeTexture, LoadAction, 0);
-		PassParameters->RenderTargets[3] = FRenderTargetBinding(OutVisibilityVelocityTexture,  LoadAction, 0);
+		PassParameters->RenderTargets[3] = FRenderTargetBinding(OutVisibilityVelocityTexture, LoadAction, 0);
 
 		PassParameters->RenderTargets.DepthStencil = FDepthStencilBinding(
 			OutVisibilityDepthTexture,
 			ERenderTargetLoadAction::ELoad,
 			ERenderTargetLoadAction::ENoAction,
 			FExclusiveDepthStencil::DepthWrite_StencilNop);
-		AddHairVisibilityCommonPass(GraphBuilder, Scene, ViewInfo, ClusterDatas, HairVisibilityRenderMode_MSAA, PassParameters);
+		AddHairVisibilityCommonPass(GraphBuilder, Scene, ViewInfo, MacroGroupDatas, HairVisibilityRenderMode_MSAA, PassParameters);
 	}
 }
 
@@ -1614,7 +1721,7 @@ static void AddHairVisibilityPPLLPass(
 	FRDGBuilder& GraphBuilder,
 	const FScene* Scene,
 	const FViewInfo* ViewInfo,
-	const FHairStrandsClusterDatas& ClusterDatas,
+	const FHairStrandsMacroGroupDatas& MacroGroupDatas,
 	const FIntPoint& Resolution,
 	FRDGTextureRef& InViewZDepthTexture,
 	FRDGTextureRef& OutVisibilityPPLLNodeCounter,
@@ -1655,7 +1762,7 @@ static void AddHairVisibilityPPLLPass(
 		struct PPLLNodeData
 		{
 			uint32 Depth;
-			uint32 PrimitiveId_ClusterId;
+			uint32 PrimitiveId_MacroGroupId;
 			uint32 Tangent_Coverage;
 			uint32 BaseColor_Roughness;
 			uint32 Specular;
@@ -1674,7 +1781,7 @@ static void AddHairVisibilityPPLLPass(
 	PassParameters->PPLLNodeData = GraphBuilder.CreateUAV(FRDGBufferUAVDesc(OutVisibilityPPLLNodeData));
 	PassParameters->HairVisibilityPass_MaxPPLLNodeCount = PPLLMaxTotalListElementCount;
 	PassParameters->RenderTargets.DepthStencil = FDepthStencilBinding(InViewZDepthTexture, ERenderTargetLoadAction::ELoad, ERenderTargetLoadAction::ENoAction, FExclusiveDepthStencil::DepthRead_StencilNop);
-	AddHairVisibilityCommonPass(GraphBuilder, Scene, ViewInfo, ClusterDatas, HairVisibilityRenderMode_PPLL, PassParameters);
+	AddHairVisibilityCommonPass(GraphBuilder, Scene, ViewInfo, MacroGroupDatas, HairVisibilityRenderMode_PPLL, PassParameters);
 }
 
 ///////////////////////////////////////////////////////////////////////////////////////////////////
@@ -1689,7 +1796,7 @@ static FHairPrimaryTransmittance AddHairViewTransmittancePass(
 	FRDGBuilder& GraphBuilder,
 	const FScene* Scene,
 	const FViewInfo* ViewInfo,
-	const FHairStrandsClusterDatas& ClusterDatas,
+	const FHairStrandsMacroGroupDatas& MacroGroupDatas,
 	const FIntPoint& Resolution,
 	const bool bOutputHairCount,
 	FRDGTextureRef SceneDepthTexture)
@@ -1727,7 +1834,7 @@ static FHairPrimaryTransmittance AddHairViewTransmittancePass(
 		ERenderTargetLoadAction::ELoad,
 		ERenderTargetLoadAction::ENoAction,
 		FExclusiveDepthStencil::DepthRead_StencilNop);
-	AddHairVisibilityCommonPass(GraphBuilder, Scene, ViewInfo, ClusterDatas, RenderMode, PassParameters);
+	AddHairVisibilityCommonPass(GraphBuilder, Scene, ViewInfo, MacroGroupDatas, RenderMode, PassParameters);
 
 	return Out;
 }
@@ -1742,10 +1849,10 @@ class FHairViewTransmittanceDepthPS : public FGlobalShader
 	BEGIN_SHADER_PARAMETER_STRUCT(FParameters, )
 		SHADER_PARAMETER(float, DistanceThreshold)
 		SHADER_PARAMETER_RDG_TEXTURE(Texture2D, SceneDepthTexture)
-		SHADER_PARAMETER_RDG_TEXTURE(Texture2D, CategorizationTexture)		
+		SHADER_PARAMETER_RDG_TEXTURE(Texture2D, CategorizationTexture)
 		SHADER_PARAMETER_STRUCT_REF(FViewUniformShaderParameters, ViewUniformBuffer)
 		RENDER_TARGET_BINDING_SLOTS()
-	END_SHADER_PARAMETER_STRUCT()
+		END_SHADER_PARAMETER_STRUCT()
 
 public:
 
@@ -1770,7 +1877,7 @@ static void AddHairViewTransmittanceDepthPass(
 	Parameters->SceneDepthTexture = SceneDepthTexture;
 	Parameters->ViewUniformBuffer = View.ViewUniformBuffer;
 	Parameters->RenderTargets[0] = FRenderTargetBinding(HairCountTexture, ERenderTargetLoadAction::ELoad);
-	
+
 	TShaderMapRef<FPostProcessVS> VertexShader(View.ShaderMap);
 	TShaderMapRef<FHairViewTransmittanceDepthPS> PixelShader(View.ShaderMap);
 	const TShaderMap<FGlobalShaderType>* GlobalShaderMap = View.ShaderMap;
@@ -1825,10 +1932,10 @@ class FHairVisibilityDepthPS : public FGlobalShader
 	BEGIN_SHADER_PARAMETER_STRUCT(FParameters, )
 		SHADER_PARAMETER_RDG_TEXTURE(Texture2D, CategorisationTexture)
 		RENDER_TARGET_BINDING_SLOTS()
-	END_SHADER_PARAMETER_STRUCT()
+		END_SHADER_PARAMETER_STRUCT()
 
 public:
-		
+
 	static FPermutationDomain RemapPermutation(FPermutationDomain PermutationVector)
 	{
 		return PermutationVector;
@@ -1910,14 +2017,14 @@ static void AddHairVisibilityColorAndDepthPatchPass(
 ///////////////////////////////////////////////////////////////////////////////////////////////////
 
 FHairStrandsVisibilityViews RenderHairStrandsVisibilityBuffer(
-	FRHICommandListImmediate& RHICmdList, 
-	const FScene* Scene, 
-	const TArray<FViewInfo>& Views, 
+	FRHICommandListImmediate& RHICmdList,
+	const FScene* Scene,
+	const TArray<FViewInfo>& Views,
 	TRefCountPtr<IPooledRenderTarget> InSceneGBufferBTexture,
 	TRefCountPtr<IPooledRenderTarget> InSceneColorTexture,
 	TRefCountPtr<IPooledRenderTarget> InSceneDepthTexture,
 	TRefCountPtr<IPooledRenderTarget> InSceneVelocityTexture,
-	const FHairStrandsClusterViews& ClusterViews)
+	const FHairStrandsMacroGroupViews& MacroGroupViews)
 {
 	QUICK_SCOPE_CYCLE_COUNTER(STAT_CLM_RenderHairStrandsVisibility);
 	SCOPED_DRAW_EVENT(RHICmdList, HairStrandsVisibility);
@@ -1932,9 +2039,9 @@ FHairStrandsVisibilityViews RenderHairStrandsVisibilityBuffer(
 		{
 			FHairStrandsVisibilityData& VisibilityData = Output.HairDatas.AddDefaulted_GetRef();
 			VisibilityData.NodeGroupSize = GetVendorOptimalGroupSize1D();
-			const FHairStrandsClusterDatas& ClusterDatas = ClusterViews.Views[ViewIndex];
+			const FHairStrandsMacroGroupDatas& MacroGroupDatas = MacroGroupViews.Views[ViewIndex];
 
-			if (ClusterDatas.Datas.Num() == 0)
+			if (MacroGroupDatas.Datas.Num() == 0)
 				continue;
 
 			// Use the scene color for computing target resolution as the View.ViewRect, 
@@ -1949,7 +2056,7 @@ FHairStrandsVisibilityViews RenderHairStrandsVisibilityBuffer(
 			FRDGTextureRef SceneColorTexture = GraphBuilder.RegisterExternalTexture(InSceneColorTexture, TEXT("SceneColorTexture"));
 			FRDGTextureRef SceneDepthTexture = GraphBuilder.RegisterExternalTexture(InSceneDepthTexture, TEXT("SceneDepthTexture"));
 			FRDGTextureRef SceneVelocityTexture = InSceneVelocityTexture ? GraphBuilder.RegisterExternalTexture(InSceneVelocityTexture, TEXT("SceneVelocityTexture")) : nullptr;
-			
+
 			const EHairVisibilityRenderMode RenderMode = GetHairVisibilityRenderMode();
 
 			// Run the view transmittance pass if needed (not in PPLL mode that is already a high quality render path)
@@ -1961,12 +2068,13 @@ FHairStrandsVisibilityViews RenderHairStrandsVisibilityBuffer(
 					GraphBuilder,
 					Scene,
 					&View,
-					ClusterDatas,
-					Resolution, 
+					MacroGroupDatas,
+					Resolution,
 					bOutputHairCount,
 					SceneDepthTexture);
 			}
 
+			FRDGTextureRef CategorizationTexture = nullptr;
 			if (RenderMode == HairVisibilityRenderMode_MSAA)
 			{
 				const bool bIsVisiblityEnable = GHairStrandsVisibilityMaterialPass > 0;
@@ -1984,21 +2092,21 @@ FHairStrandsVisibilityViews RenderHairStrandsVisibilityBuffer(
 					GraphBuilder,
 					View,
 					Resolution,
-					ClusterDatas,
+					MacroGroupDatas,
 					SceneDepthTexture);
 
 				AddHairVisibilityMSAAPass(
 					bIsVisiblityEnable,
-					GraphBuilder, 
-					Scene, 
-					&View, 
-					ClusterDatas, 
+					GraphBuilder,
+					Scene,
+					&View,
+					MacroGroupDatas,
 					Resolution,
-					MsaaVisibilityResources.IdTexture, 
-					MsaaVisibilityResources.MaterialTexture, 
+					MsaaVisibilityResources.IdTexture,
+					MsaaVisibilityResources.MaterialTexture,
 					MsaaVisibilityResources.AttributeTexture,
 					MsaaVisibilityResources.VelocityTexture,
-					MsaaVisibilityResources.DepthTexture); 
+					MsaaVisibilityResources.DepthTexture);
 
 				// This is used when compaction is not enabled.
 				GraphBuilder.QueueTextureExtraction(MsaaVisibilityResources.IdTexture, &VisibilityData.IDTexture);
@@ -2010,7 +2118,6 @@ FHairStrandsVisibilityViews RenderHairStrandsVisibilityBuffer(
 					GraphBuilder.QueueTextureExtraction(MsaaVisibilityResources.VelocityTexture, &VisibilityData.VelocityTexture);
 				}
 
-				FRDGTextureRef CategorizationTexture;
 				{
 					FHairVisibilityPrimitiveIdCompactionCS::FParameters* PassParameters = GraphBuilder.AllocParameters<FHairVisibilityPrimitiveIdCompactionCS::FParameters>();
 					PassParameters->MSAA_DepthTexture = MsaaVisibilityResources.DepthTexture;
@@ -2029,7 +2136,7 @@ FHairStrandsVisibilityViews RenderHairStrandsVisibilityBuffer(
 						bIsVisiblityEnable,
 						GraphBuilder,
 						View,
-						ClusterDatas,
+						MacroGroupDatas,
 						VisibilityData.NodeGroupSize,
 						PassParameters,
 						CompactNodeIndex,
@@ -2047,7 +2154,7 @@ FHairStrandsVisibilityViews RenderHairStrandsVisibilityBuffer(
 							GraphBuilder,
 							Scene,
 							&View,
-							ClusterDatas,
+							MacroGroupDatas,
 							VisibilityData.NodeGroupSize,
 							CompactNodeIndex,
 							CompactNodeData,
@@ -2058,7 +2165,7 @@ FHairStrandsVisibilityViews RenderHairStrandsVisibilityBuffer(
 						AddHairVelocityPass(
 							GraphBuilder,
 							View,
-							ClusterDatas,
+							MacroGroupDatas,
 							CompactNodeIndex,
 							CompactNodeData,
 							PassOutput.NodeVelocity,
@@ -2067,11 +2174,23 @@ FHairStrandsVisibilityViews RenderHairStrandsVisibilityBuffer(
 						CompactNodeData = PassOutput.NodeData;
 					}
 
-					GraphBuilder.QueueTextureExtraction(CompactNodeIndex,		&VisibilityData.NodeIndex);
-					GraphBuilder.QueueTextureExtraction(CategorizationTexture,	&VisibilityData.CategorizationTexture);
-					GraphBuilder.QueueBufferExtraction(CompactNodeData,			&VisibilityData.NodeData,					FRDGResourceState::EAccess::Read, FRDGResourceState::EPipeline::Graphics);
-					GraphBuilder.QueueBufferExtraction(CompactNodeCoord,		&VisibilityData.NodeCoord,					FRDGResourceState::EAccess::Read, FRDGResourceState::EPipeline::Graphics);
-					GraphBuilder.QueueBufferExtraction(IndirectArgsBuffer,		&VisibilityData.NodeIndirectArg,			FRDGResourceState::EAccess::Read, FRDGResourceState::EPipeline::Compute);
+					GraphBuilder.QueueTextureExtraction(CompactNodeIndex, &VisibilityData.NodeIndex);
+					GraphBuilder.QueueTextureExtraction(CategorizationTexture, &VisibilityData.CategorizationTexture);
+					GraphBuilder.QueueBufferExtraction(CompactNodeData, &VisibilityData.NodeData, FRDGResourceState::EAccess::Read, FRDGResourceState::EPipeline::Graphics);
+					GraphBuilder.QueueBufferExtraction(CompactNodeCoord, &VisibilityData.NodeCoord, FRDGResourceState::EAccess::Read, FRDGResourceState::EPipeline::Graphics);
+					GraphBuilder.QueueBufferExtraction(IndirectArgsBuffer, &VisibilityData.NodeIndirectArg, FRDGResourceState::EAccess::Read, FRDGResourceState::EPipeline::Compute);
+				}
+
+				// View transmittance depth test needs to happen before the scene depth is patched with the hair depth (for fully-covered-by-hair pixels)
+				if (ViewTransmittance.HairCountTexture)
+				{
+					AddHairViewTransmittanceDepthPass(
+						GraphBuilder,
+						View,
+						CategorizationTexture,
+						SceneDepthTexture,
+						ViewTransmittance.HairCountTexture);
+					GraphBuilder.QueueTextureExtraction(ViewTransmittance.HairCountTexture, &VisibilityData.ViewHairCountTexture);
 				}
 
 				// For fully covered pixels, write: 
@@ -2085,17 +2204,6 @@ FHairStrandsVisibilityViews RenderHairStrandsVisibilityBuffer(
 					SceneGBufferBTexture,
 					SceneColorTexture,
 					SceneDepthTexture);
-
-				if (ViewTransmittance.HairCountTexture)
-				{
-					AddHairViewTransmittanceDepthPass(
-						GraphBuilder,
-						View,
-						CategorizationTexture,
-						SceneDepthTexture,
-						ViewTransmittance.HairCountTexture);
-					GraphBuilder.QueueTextureExtraction(ViewTransmittance.HairCountTexture, &VisibilityData.ViewHairCountTexture);
-				}
 			}
 			else if (RenderMode == HairVisibilityRenderMode_PPLL)
 			{
@@ -2110,14 +2218,13 @@ FHairStrandsVisibilityViews RenderHairStrandsVisibilityBuffer(
 				FRDGTextureRef ViewZDepthTexture = GraphBuilder.RegisterExternalTexture(SceneContext.SceneDepthZ);
 
 				// Linked list generation pass
-				AddHairVisibilityPPLLPass(GraphBuilder, Scene, &View, ClusterDatas, Resolution, ViewZDepthTexture, PPLLNodeCounterTexture, PPLLNodeIndexTexture, PPLLNodeDataBuffer);
+				AddHairVisibilityPPLLPass(GraphBuilder, Scene, &View, MacroGroupDatas, Resolution, ViewZDepthTexture, PPLLNodeCounterTexture, PPLLNodeIndexTexture, PPLLNodeDataBuffer);
 
 				// Linked list sorting pass and compaction into common representation
-				FRDGTextureRef CategorizationTexture;
 				{
 					FHairVisibilityPrimitiveIdCompactionCS::FParameters* PassParameters = GraphBuilder.AllocParameters<FHairVisibilityPrimitiveIdCompactionCS::FParameters>();
-					PassParameters->PPLLCounter  = PPLLNodeCounterTexture;
-					PassParameters->PPLLNodeIndex= PPLLNodeIndexTexture;
+					PassParameters->PPLLCounter = PPLLNodeCounterTexture;
+					PassParameters->PPLLNodeIndex = PPLLNodeIndexTexture;
 					PassParameters->PPLLNodeData = GraphBuilder.CreateSRV(PPLLNodeDataBuffer);
 					PassParameters->ViewTransmittanceTexture = ViewTransmittance.TransmittanceTexture;
 
@@ -2130,7 +2237,7 @@ FHairStrandsVisibilityViews RenderHairStrandsVisibilityBuffer(
 						false,
 						GraphBuilder,
 						View,
-						ClusterDatas,
+						MacroGroupDatas,
 						VisibilityData.NodeGroupSize,
 						PassParameters,
 						CompactNodeIndex,
@@ -2162,6 +2269,18 @@ FHairStrandsVisibilityViews RenderHairStrandsVisibilityBuffer(
 #endif
 			}
 
+			// Generate Tile data
+			{
+				FRDGTextureRef TileIndexTexture = nullptr;
+				FRDGBufferRef TileBuffer = nullptr;
+				FRDGBufferRef TileIndirectArgs = nullptr;
+				AddGenerateTilePass(GraphBuilder, View, VisibilityData.TileThreadGroupSize, VisibilityData.TileSize, CategorizationTexture, TileIndexTexture, TileBuffer, TileIndirectArgs);
+
+				GraphBuilder.QueueTextureExtraction(TileIndexTexture, &VisibilityData.TileIndexTexture);
+				GraphBuilder.QueueBufferExtraction(TileBuffer, &VisibilityData.TileBuffer, FRDGResourceState::EAccess::Read, FRDGResourceState::EPipeline::Compute);
+				GraphBuilder.QueueBufferExtraction(TileIndirectArgs, &VisibilityData.TileIndirectArgs, FRDGResourceState::EAccess::Read, FRDGResourceState::EPipeline::Compute);
+			}
+
 			GraphBuilder.Execute();
 
 			// #hair_todo: is there a better way to get SRV view of a RDG buffer? should work as long as there is not reuse between the pass
