Skipped gpu debug rendering pass to unblock RPS because it is crashing on startup

diff --git a/Engine/Source/Runtime/Renderer/Private/GpuDebugRendering.cpp b/Engine/Source/Runtime/Renderer/Private/GpuDebugRendering.cpp
--- a/Engine/Source/Runtime/Renderer/Private/GpuDebugRendering.cpp
+++ b/Engine/Source/Runtime/Renderer/Private/GpuDebugRendering.cpp
@@ -253,7 +253,7 @@ namespace ShaderDrawDebug
 
 	void DrawView(FRDGBuilder& GraphBuilder, const FViewInfo& View, FRDGTextureRef OutputTexture, FRDGTextureRef DepthTexture)
 	{
-		if (!IsShaderDrawDebugEnabled(View))
+		//if (!IsShaderDrawDebugEnabled(View))
 		{
 			return;
 		}
