Fix error when enabling hair rendering.

diff --git a/Engine/Source/Runtime/Renderer/Private/HairStrands/HairStrandsVisibility.cpp b/Engine/Source/Runtime/Renderer/Private/HairStrands/HairStrandsVisibility.cpp
--- a/Engine/Source/Runtime/Renderer/Private/HairStrands/HairStrandsVisibility.cpp
+++ b/Engine/Source/Runtime/Renderer/Private/HairStrands/HairStrandsVisibility.cpp
@@ -1330,7 +1330,8 @@ static void AddGenerateTilePass(
 
 	OutTileBuffer = GraphBuilder.CreateBuffer(FRDGBufferDesc::CreateBufferDesc(sizeof(uint32), TileCount), TEXT("HairTileBuffer"));
 
-	AddClearRenderTargetPass(GraphBuilder, TileCounter, FLinearColor(0.0f, 0.0f, 0.0f, 1.0f));
+	uint32 ClearValues[4] = { 0,0,0,0 };
+	AddClearUAVPass(GraphBuilder, GraphBuilder.CreateUAV(TileCounter), ClearValues);
 
 	FHairGenerateTileCS::FPermutationDomain PermutationVector;
 	PermutationVector.Set<FHairGenerateTileCS::FTileSize>(0);
