Renable GPU debug rendering.

This is possible as CL 11083355 fixed the ResourceRHI nullptr issue.

diff --git a/Engine/Source/Runtime/Renderer/Private/GpuDebugRendering.cpp b/Engine/Source/Runtime/Renderer/Private/GpuDebugRendering.cpp
--- a/Engine/Source/Runtime/Renderer/Private/GpuDebugRendering.cpp
+++ b/Engine/Source/Runtime/Renderer/Private/GpuDebugRendering.cpp
@@ -253,7 +253,7 @@ namespace ShaderDrawDebug
 
 	void DrawView(FRDGBuilder& GraphBuilder, const FViewInfo& View, FRDGTextureRef OutputTexture, FRDGTextureRef DepthTexture)
 	{
-		//if (!IsShaderDrawDebugEnabled(View))
+		if (!IsShaderDrawDebugEnabled(View))
 		{
 			return;
 		}
@@ -287,9 +287,9 @@ namespace ShaderDrawDebug
 			}
 
 			ValidateShaderParameters(*PixelShader, PassParameters->ShaderDrawPSParameters);
-			ClearUnusedGraphResources(*PixelShader, &PassParameters->ShaderDrawPSParameters);
+			ClearUnusedGraphResources(*PixelShader, &PassParameters->ShaderDrawPSParameters, { IndirectBuffer });
 			ValidateShaderParameters(*VertexShader, PassParameters->ShaderDrawVSParameters);
-			ClearUnusedGraphResources(*VertexShader, &PassParameters->ShaderDrawVSParameters);
+			ClearUnusedGraphResources(*VertexShader, &PassParameters->ShaderDrawVSParameters, { IndirectBuffer });
 
 			GraphBuilder.AddPass(
 				RDG_EVENT_NAME("ShaderDrawDebug"),
@@ -297,6 +297,12 @@ namespace ShaderDrawDebug
 				ERDGPassFlags::Raster,
 				[VertexShader, PixelShader, PassParameters, IndirectBuffer, LockedIndirectBuffer, bIsBehindDepth, bUseRdgInput](FRHICommandListImmediate& RHICmdListImmediate)
 			{
+				// Marks the indirect draw parameter as used by the pass, given it's not used directly by any of the shaders.
+				if (bUseRdgInput)
+				{
+					PassParameters->ShaderDrawVSParameters.IndirectBuffer->MarkResourceAsUsed();
+				}
+
 				FGraphicsPipelineStateInitializer GraphicsPSOInit;
 				RHICmdListImmediate.ApplyCachedRenderTargets(GraphicsPSOInit);
 				GraphicsPSOInit.DepthStencilState = bIsBehindDepth ? TStaticDepthStencilState<false, CF_DepthFarther>::GetRHI() : TStaticDepthStencilState<false, CF_DepthNearOrEqual>::GetRHI();
@@ -314,7 +320,10 @@ namespace ShaderDrawDebug
 
 				if (bUseRdgInput)
 				{
-					RHICmdListImmediate.DrawPrimitiveIndirect(IndirectBuffer->GetIndirectRHICallBuffer(), 0);
+					// Marks the indirect draw parameter as used by the pass, given it's not used directly by any of the shaders.
+					FRHIVertexBuffer* IndirectBufferRHI = PassParameters->ShaderDrawVSParameters.IndirectBuffer->GetIndirectRHICallBuffer();
+					check(IndirectBufferRHI != nullptr);
+					RHICmdListImmediate.DrawPrimitiveIndirect(IndirectBufferRHI, 0);
 				}
 				else
 				{
