Integrated fix for negative draw call counts.

#6584

diff --git a/Engine/Source/Runtime/Renderer/Private/MeshDrawCommands.cpp b/Engine/Source/Runtime/Renderer/Private/MeshDrawCommands.cpp
--- a/Engine/Source/Runtime/Renderer/Private/MeshDrawCommands.cpp
+++ b/Engine/Source/Runtime/Renderer/Private/MeshDrawCommands.cpp
@@ -1130,7 +1130,7 @@ public:
 
 		// Recompute draw range.
 		const int32 DrawNum = VisibleMeshDrawCommands.Num();
-		const int32 NumDrawsPerTask = FMath::DivideAndRoundUp(DrawNum, TaskNum);
+		const int32 NumDrawsPerTask = TaskIndex < DrawNum ? FMath::DivideAndRoundUp(DrawNum, TaskNum) : 0;
 		const int32 StartIndex = TaskIndex * NumDrawsPerTask;
 		const int32 NumDraws = FMath::Min(NumDrawsPerTask, DrawNum - StartIndex);
 
