Fixed crash in raytracing mesh draw command caching.

diff --git a/Engine/Source/Runtime/Renderer/Private/PrimitiveSceneInfo.cpp b/Engine/Source/Runtime/Renderer/Private/PrimitiveSceneInfo.cpp
--- a/Engine/Source/Runtime/Renderer/Private/PrimitiveSceneInfo.cpp
+++ b/Engine/Source/Runtime/Renderer/Private/PrimitiveSceneInfo.cpp
@@ -400,24 +400,30 @@ void FPrimitiveSceneInfo::CacheMeshDrawCommands(FRHICommandListImmediate& RHICmd
 
 		for (FPrimitiveSceneInfo* SceneInfo : SceneInfos)
 		{
-			if (SceneInfo->Proxy->IsRayTracingStaticRelevant())
+			if (SceneInfo->Proxy->IsRayTracingStaticRelevant() && SceneInfo->StaticMeshes.Num() > 0)
 			{
-				int MaxLOD = -1;
+				int32 MaxLOD = -1;
 				for (int32 MeshIndex = 0; MeshIndex < SceneInfo->StaticMeshes.Num(); MeshIndex++)
 				{
 					FStaticMeshBatch& Mesh = SceneInfo->StaticMeshes[MeshIndex];
 					MaxLOD = MaxLOD < Mesh.LODIndex ? Mesh.LODIndex : MaxLOD;
+				}
+
+				SceneInfo->CachedRayTracingMeshCommandIndicesPerLOD.Empty(MaxLOD + 1);
+				SceneInfo->CachedRayTracingMeshCommandIndicesPerLOD.AddDefaulted(MaxLOD + 1);
+
+				for (int32 MeshIndex = 0; MeshIndex < SceneInfo->StaticMeshes.Num(); MeshIndex++)
+				{
+					FStaticMeshBatch& Mesh = SceneInfo->StaticMeshes[MeshIndex];
 
 					check(!Mesh.bRequiresPerElementVisibility);
 					RayTracingMeshProcessor.AddMeshBatch(Mesh, ~0ull, SceneInfo->Proxy);
 
-					SceneInfo->CachedRayTracingMeshCommandIndicesPerLOD[Mesh.LODIndex].Add(CommandContext.CommandIndex);
-				}
-
-				if (SceneInfo->StaticMeshes.Num() > 0)
-				{
-					SceneInfo->CachedRayTracingMeshCommandIndicesPerLOD.Empty(MaxLOD + 1);
-					SceneInfo->CachedRayTracingMeshCommandIndicesPerLOD.AddDefaulted(MaxLOD + 1);
+					if (CommandContext.CommandIndex >= 0)
+					{
+						SceneInfo->CachedRayTracingMeshCommandIndicesPerLOD[Mesh.LODIndex].Add(CommandContext.CommandIndex);
+						CommandContext.CommandIndex = -1;
+					}
 				}
 			}
 		}
