Add missing MemStack mark

diff --git a/Engine/Source/Runtime/Renderer/Private/RendererScene.cpp b/Engine/Source/Runtime/Renderer/Private/RendererScene.cpp
--- a/Engine/Source/Runtime/Renderer/Private/RendererScene.cpp
+++ b/Engine/Source/Runtime/Renderer/Private/RendererScene.cpp
@@ -1047,6 +1047,8 @@ public:
 	{
 		for (TSet<FPrimitiveSceneInfo*>::TConstIterator It(PendingPrimitives); It; ++It)
 		{
+			FMemMark MemStackMark(FMemStack::Get());
+
 			FPrimitiveSceneInfo* PrimInfo = *It;
 			FPrimitiveSceneProxy* Proxy = PrimInfo->Proxy;
 			const FBoxSphereBounds& Bounds = Proxy->GetBounds();
