Fixed GPrimitiveIdVertexBufferPool race condition (integrated fix from dev-rendering)

#rnx

diff --git a/Engine/Source/Runtime/Renderer/Private/MeshDrawCommands.cpp b/Engine/Source/Runtime/Renderer/Private/MeshDrawCommands.cpp
--- a/Engine/Source/Runtime/Renderer/Private/MeshDrawCommands.cpp
+++ b/Engine/Source/Runtime/Renderer/Private/MeshDrawCommands.cpp
@@ -35,8 +35,12 @@ FPrimitiveIdVertexBufferPool::~FPrimitiveIdVertexBufferPool()
 	check(!Entries.Num());
 }
 
-FRHIVertexBuffer* FPrimitiveIdVertexBufferPool::Allocate(int32 BufferSize)
+FPrimitiveIdVertexBufferPoolEntry FPrimitiveIdVertexBufferPool::Allocate(int32 BufferSize)
 {
+	check(IsInRenderingThread());
+
+	FScopeLock Lock(&AllocationCS);
+
 	BufferSize = Align(BufferSize, 1024);
 
 	// First look for a smallest unused one.
@@ -59,11 +63,13 @@ FRHIVertexBuffer* FPrimitiveIdVertexBufferPool::Allocate(int32 BufferSize)
 		}
 	}
 	
-	if (BestFitBufferIndex >=0 )
+	if (BestFitBufferIndex >= 0)
 	{
 		// Reuse existing buffer.
-		Entries[BestFitBufferIndex].LastDiscardId = DiscardId;
-		return Entries[BestFitBufferIndex].BufferRHI;
+		FPrimitiveIdVertexBufferPoolEntry ReusedEntry = MoveTemp(Entries[BestFitBufferIndex]);
+		ReusedEntry.LastDiscardId = DiscardId;
+		Entries.RemoveAt(BestFitBufferIndex);
+		return ReusedEntry;
 	}
 	else
 	{
@@ -73,14 +79,23 @@ FRHIVertexBuffer* FPrimitiveIdVertexBufferPool::Allocate(int32 BufferSize)
 		NewEntry.BufferSize = BufferSize;
 		FRHIResourceCreateInfo CreateInfo;
 		NewEntry.BufferRHI = RHICreateVertexBuffer(NewEntry.BufferSize, BUF_Volatile, CreateInfo);
-		Entries.Add(NewEntry);
 
-		return NewEntry.BufferRHI;
+		return NewEntry;
 	}
 }
 
+void FPrimitiveIdVertexBufferPool::ReturnToFreeList(FPrimitiveIdVertexBufferPoolEntry Entry)
+{
+	// Entries can be returned from RHIT or RT, depending on if FParallelMeshDrawCommandPass::DispatchDraw() takes the parallel path
+	FScopeLock Lock(&AllocationCS);
+
+	Entries.Add(MoveTemp(Entry));
+}
+
 void FPrimitiveIdVertexBufferPool::DiscardAll()
 {
+	FScopeLock Lock(&AllocationCS);
+
 	++DiscardId;
 
 	// Remove old unused pool entries.
@@ -960,7 +975,8 @@ void SortAndMergeDynamicPassMeshDrawCommands(
 			}
 
 			const int32 PrimitiveIdBufferDataSize = InstanceFactor * NumDrawCommands * sizeof(int32);
-			OutPrimitiveIdVertexBuffer = GPrimitiveIdVertexBufferPool.Allocate(PrimitiveIdBufferDataSize);
+			FPrimitiveIdVertexBufferPoolEntry Entry = GPrimitiveIdVertexBufferPool.Allocate(PrimitiveIdBufferDataSize);
+			OutPrimitiveIdVertexBuffer = Entry.BufferRHI;
 			void* PrimitiveIdBufferData = RHILockVertexBuffer(OutPrimitiveIdVertexBuffer, 0, PrimitiveIdBufferDataSize, RLM_WriteOnly);
 
 			BuildMeshDrawCommandPrimitiveIdBuffer(
@@ -978,6 +994,7 @@ void SortAndMergeDynamicPassMeshDrawCommands(
 			);
 
 			RHIUnlockVertexBuffer(OutPrimitiveIdVertexBuffer);
+			GPrimitiveIdVertexBufferPool.ReturnToFreeList(Entry);
 		}
 	}
 }
@@ -1062,7 +1079,7 @@ void FParallelMeshDrawCommandPass::DispatchPassSetup(
 		bPrimitiveIdBufferDataOwnedByRHIThread = false;
 		TaskContext.PrimitiveIdBufferDataSize = TaskContext.InstanceFactor * MaxNumDraws * sizeof(int32);
 		TaskContext.PrimitiveIdBufferData = FMemory::Malloc(TaskContext.PrimitiveIdBufferDataSize);
-		PrimitiveIdVertexBufferRHI = GPrimitiveIdVertexBufferPool.Allocate(TaskContext.PrimitiveIdBufferDataSize);
+		PrimitiveIdVertexBufferPoolEntry = GPrimitiveIdVertexBufferPool.Allocate(TaskContext.PrimitiveIdBufferDataSize);
 		TaskContext.MeshDrawCommands.Reserve(MaxNumDraws);
 		TaskContext.TempVisibleMeshDrawCommands.Reserve(MaxNumDraws);
 
@@ -1116,6 +1133,20 @@ void FParallelMeshDrawCommandPass::WaitForTasksAndEmpty()
 		TaskContext.MobileBasePassCSMMeshPassProcessor = nullptr;
 	}
 
+	if (MaxNumDraws > 0)
+	{
+		if (bPrimitiveIdBufferDataOwnedByRHIThread)
+		{
+			FRHICommandListExecutor::GetImmediateCommandList().EnqueueLambda([PrimitiveIdVertexBufferPoolEntry = PrimitiveIdVertexBufferPoolEntry](FRHICommandListImmediate&) {
+				GPrimitiveIdVertexBufferPool.ReturnToFreeList(PrimitiveIdVertexBufferPoolEntry);
+			});
+		}
+		else
+		{
+			GPrimitiveIdVertexBufferPool.ReturnToFreeList(PrimitiveIdVertexBufferPoolEntry);
+		}
+	}
+
 	if (!bPrimitiveIdBufferDataOwnedByRHIThread)
 	{
 		FMemory::Free(TaskContext.PrimitiveIdBufferData);
@@ -1212,7 +1243,7 @@ void FParallelMeshDrawCommandPass::DispatchDraw(FParallelCommandListSet* Paralle
 		return;
 	}
 
-	FRHIVertexBuffer* PrimitiveIdsBuffer = PrimitiveIdVertexBufferRHI;
+	FRHIVertexBuffer* PrimitiveIdsBuffer = PrimitiveIdVertexBufferPoolEntry.BufferRHI;
 	const int32 BasePrimitiveIdsOffset = 0;
 
 	if (ParallelCommandListSet)
@@ -1230,7 +1261,8 @@ void FParallelMeshDrawCommandPass::DispatchDraw(FParallelCommandListSet* Paralle
 			RHICommandList.EnqueueLambda([
 				VertexBuffer = PrimitiveIdsBuffer,
 				VertexBufferData = TaskContext.PrimitiveIdBufferData, 
-				VertexBufferDataSize = TaskContext.PrimitiveIdBufferDataSize](FRHICommandListImmediate& CmdList)
+				VertexBufferDataSize = TaskContext.PrimitiveIdBufferDataSize,
+				PrimitiveIdVertexBufferPoolEntry = PrimitiveIdVertexBufferPoolEntry](FRHICommandListImmediate& CmdList)
 			{
 				// Upload vertex buffer data.
 				void* RESTRICT Data = (void* RESTRICT)CmdList.LockVertexBuffer(VertexBuffer, 0, VertexBufferDataSize, RLM_WriteOnly);
@@ -1285,9 +1317,9 @@ void FParallelMeshDrawCommandPass::DispatchDraw(FParallelCommandListSet* Paralle
 		if (TaskContext.bUseGPUScene)
 		{
 			// Can immediately upload vertex buffer data, as there is no parallel draw task.
-			void* RESTRICT Data = RHILockVertexBuffer(PrimitiveIdVertexBufferRHI, 0, TaskContext.PrimitiveIdBufferDataSize, RLM_WriteOnly);
+			void* RESTRICT Data = RHILockVertexBuffer(PrimitiveIdVertexBufferPoolEntry.BufferRHI, 0, TaskContext.PrimitiveIdBufferDataSize, RLM_WriteOnly);
 			FMemory::Memcpy(Data, TaskContext.PrimitiveIdBufferData, TaskContext.PrimitiveIdBufferDataSize);
-			RHIUnlockVertexBuffer(PrimitiveIdVertexBufferRHI);
+			RHIUnlockVertexBuffer(PrimitiveIdVertexBufferPoolEntry.BufferRHI);
 		}
 
 		SubmitMeshDrawCommandsRange(TaskContext.MeshDrawCommands, TaskContext.MinimalPipelineStatePassSet, PrimitiveIdsBuffer, BasePrimitiveIdsOffset, TaskContext.bDynamicInstancing, 0, TaskContext.MeshDrawCommands.Num(), TaskContext.InstanceFactor, RHICmdList);
diff --git a/Engine/Source/Runtime/Renderer/Private/MeshDrawCommands.h b/Engine/Source/Runtime/Renderer/Private/MeshDrawCommands.h
--- a/Engine/Source/Runtime/Renderer/Private/MeshDrawCommands.h
+++ b/Engine/Source/Runtime/Renderer/Private/MeshDrawCommands.h
@@ -16,27 +16,30 @@ class FParallelCommandListSet;
 /**
  * Global vertex buffer pool used for GPUScene primitive id arrays.
  */
+
+struct FPrimitiveIdVertexBufferPoolEntry
+{
+	int32 BufferSize = 0;
+	uint32 LastDiscardId = 0;
+	FVertexBufferRHIRef BufferRHI;
+};
+
 class FPrimitiveIdVertexBufferPool : public FRenderResource
 {
 public:
 	FPrimitiveIdVertexBufferPool();
 	~FPrimitiveIdVertexBufferPool();
 
-	FRHIVertexBuffer* Allocate(int32 BufferSize);
+	FPrimitiveIdVertexBufferPoolEntry Allocate(int32 BufferSize);
+	void ReturnToFreeList(FPrimitiveIdVertexBufferPoolEntry Entry);
 	void DiscardAll();
 
 	virtual void ReleaseDynamicRHI() override;
 
 private:
-	struct FPrimitiveIdVertexBufferPoolEntry
-	{
-		int32 BufferSize;
-		uint32 LastDiscardId;
-		FVertexBufferRHIRef BufferRHI;
-	};
-
 	uint32 DiscardId;
 	TArray<FPrimitiveIdVertexBufferPoolEntry> Entries;
+	FCriticalSection AllocationCS;
 };
 
 extern TGlobalResource<FPrimitiveIdVertexBufferPool> GPrimitiveIdVertexBufferPool;
@@ -128,8 +131,7 @@ class FParallelMeshDrawCommandPass
 {
 public:
 	FParallelMeshDrawCommandPass()
-		: PrimitiveIdVertexBufferRHI(nullptr)
-		, bPrimitiveIdBufferDataOwnedByRHIThread(false)
+		: bPrimitiveIdBufferDataOwnedByRHIThread(false)
 		, MaxNumDraws(0)
 	{
 	}
@@ -171,7 +173,7 @@ public:
 	}
 
 private:
-	FRHIVertexBuffer* PrimitiveIdVertexBufferRHI;
+	FPrimitiveIdVertexBufferPoolEntry PrimitiveIdVertexBufferPoolEntry;
 	FMeshDrawCommandPassSetupTaskContext TaskContext;
 	FGraphEventRef TaskEventRef;
 	FString PassNameForStats;
