Fix a bug that dynamic skylight doesn't apply to static meshes.

diff --git a/Engine/Source/Runtime/Renderer/Private/MobileBasePass.cpp b/Engine/Source/Runtime/Renderer/Private/MobileBasePass.cpp
--- a/Engine/Source/Runtime/Renderer/Private/MobileBasePass.cpp
+++ b/Engine/Source/Runtime/Renderer/Private/MobileBasePass.cpp
@@ -614,9 +614,20 @@ void FMobileBasePassMeshProcessor::Process(
 		FBaseDS,
 		TMobileBasePassPSPolicyParamType<FUniformLightMapPolicy>> BasePassShaders;
 	
-	//The stationary skylight contribution has been added to the LowQuality Lightmap for StaticMeshActor on mobile, so we should skip the sky light spherical harmonic contribution for it.
-	//Enable skylight if LowQualityLightmaps is disabled or the Lightmap has not been built.
-	bool bEnableSkyLight = ShadingModels.IsLit() && Scene && Scene->ShouldRenderSkylightInBasePass(BlendMode) && !(FReadOnlyCVARCache::Get().bAllowStaticLighting && FReadOnlyCVARCache::Get().bEnableLowQualityLightmaps && PrimitiveSceneProxy && PrimitiveSceneProxy->IsStatic() && MeshBatch.LCI && MeshBatch.LCI->GetLightMapInteraction(FeatureLevel).GetType() == LMIT_Texture && !MaterialResource.IsTwoSided());
+	bool bEnableSkyLight = false;
+	
+	if (Scene && Scene->SkyLight)
+	{
+		//The stationary skylight contribution has been added to the LowQuality Lightmap for StaticMeshActor on mobile, so we should skip the sky light spherical harmonic contribution for it.
+		//Enable skylight if LowQualityLightmaps is disabled or the Lightmap has not been built or if it is a dynamic skylight
+		bool bStaticMeshHasValidLightmapFromStationarySkyLight = FReadOnlyCVARCache::Get().bAllowStaticLighting && FReadOnlyCVARCache::Get().bEnableLowQualityLightmaps && PrimitiveSceneProxy && PrimitiveSceneProxy->IsStatic() && MeshBatch.LCI && MeshBatch.LCI->GetLightMapInteraction(FeatureLevel).GetType() == LMIT_Texture && Scene->SkyLight->bWantsStaticShadowing;
+
+		//Two side material should enable sky light for the back face since only the front face has light map and it will be corrected in base pass shader.
+		bool bDisableStationarySkyLightForStaticMesh = bStaticMeshHasValidLightmapFromStationarySkyLight && !MaterialResource.IsTwoSided();
+
+		bEnableSkyLight = ShadingModels.IsLit() && Scene->ShouldRenderSkylightInBasePass(BlendMode) && (!bDisableStationarySkyLightForStaticMesh);
+	}
+
 	int32 NumMovablePointLights = MobileBasePass::CalcNumMovablePointLights(MaterialResource, PrimitiveSceneProxy);
 
 	MobileBasePass::GetShaders(
