Fix out-of-date view uniform buffer for selection rendering post-process (broken selection highlights in VRMode).

diff --git a/Engine/Source/Runtime/Renderer/Private/PostProcess/PostProcessCompositeEditorPrimitives.cpp b/Engine/Source/Runtime/Renderer/Private/PostProcess/PostProcessCompositeEditorPrimitives.cpp
--- a/Engine/Source/Runtime/Renderer/Private/PostProcess/PostProcessCompositeEditorPrimitives.cpp
+++ b/Engine/Source/Runtime/Renderer/Private/PostProcess/PostProcessCompositeEditorPrimitives.cpp
@@ -232,8 +232,6 @@ const FViewInfo* UpdateEditorPrimitiveView(
 
 	EditorView->CachedViewUniformShaderParameters = MoveTemp(ViewParameters);
 
-	SceneUniformBuffers.UpdateViewUniformBufferImmediate(*EditorView->CachedViewUniformShaderParameters);
-
 	EditorView->ViewUniformBuffer = SceneUniformBuffers.ViewUniformBuffer;
 
 	return EditorView;
diff --git a/Engine/Source/Runtime/Renderer/Private/PostProcess/PostProcessSelectionOutline.cpp b/Engine/Source/Runtime/Renderer/Private/PostProcess/PostProcessSelectionOutline.cpp
--- a/Engine/Source/Runtime/Renderer/Private/PostProcess/PostProcessSelectionOutline.cpp
+++ b/Engine/Source/Runtime/Renderer/Private/PostProcess/PostProcessSelectionOutline.cpp
@@ -99,6 +99,9 @@ FScreenPassTexture AddSelectionOutlinePass(FRDGBuilder& GraphBuilder, const FVie
 			ERDGPassFlags::Raster,
 			[&View, SceneColorViewport](FRHICommandList& RHICmdList)
 		{
+			FScene* Scene = View.Family->Scene->GetRenderScene();
+			Scene->UniformBuffers.UpdateViewUniformBufferImmediate(*View.CachedViewUniformShaderParameters);
+
 			RHICmdList.SetViewport(SceneColorViewport.Rect.Min.X, SceneColorViewport.Rect.Min.Y, 0.0f, SceneColorViewport.Rect.Max.X, SceneColorViewport.Rect.Max.Y, 1.0f);
 
 			// Run selection pass on static elements
