Fix for robomerge merge error

diff --git a/Engine/Source/Runtime/Renderer/Private/CompositionLighting/PostProcessAmbientOcclusion.cpp b/Engine/Source/Runtime/Renderer/Private/CompositionLighting/PostProcessAmbientOcclusion.cpp
--- a/Engine/Source/Runtime/Renderer/Private/CompositionLighting/PostProcessAmbientOcclusion.cpp
+++ b/Engine/Source/Runtime/Renderer/Private/CompositionLighting/PostProcessAmbientOcclusion.cpp
@@ -2773,20 +2773,20 @@ void FRCPassPostProcessAmbientOcclusion_GTAO_SpatialFilter::Process(FRenderingCo
 			const FSceneRenderTargetItem& DestRenderTarget = SceneContext.ScreenSpaceGTAOHorizons->GetRenderTargetItem();
 
 			FRHIAsyncComputeCommandListImmediate& RHICmdListComputeImmediate = FRHICommandListExecutor::GetImmediateAsyncComputeCommandList();
-			RHICmdListComputeImmediate.SetComputeShader(ComputeShader->GetComputeShader());
+			RHICmdListComputeImmediate.SetComputeShader(ComputeShader.GetComputeShader());
 			RHICmdListComputeImmediate.TransitionResource(EResourceTransitionAccess::ERWBarrier, EResourceTransitionPipeline::EComputeToCompute, DestRenderTarget.UAV, nullptr);
 			ComputeShader->SetParametersCS(RHICmdListComputeImmediate, Context, InputTexSize, OutputViewRect, DestRenderTarget.UAV);
-			DispatchComputeShader(RHICmdListComputeImmediate, *ComputeShader, GroupSizeX, GroupSizeY, 1);
+			DispatchComputeShader(RHICmdListComputeImmediate, ComputeShader, GroupSizeX, GroupSizeY, 1);
 			ComputeShader->UnsetParameters(Context.RHICmdList);
 		}
 		else
 		{
 			const FSceneRenderTargetItem& DestRenderTarget = PassOutputs[0].RequestSurface(Context);
 
-			Context.RHICmdList.SetComputeShader(ComputeShader->GetComputeShader());
+			Context.RHICmdList.SetComputeShader(ComputeShader.GetComputeShader());
 			Context.RHICmdList.TransitionResource(EResourceTransitionAccess::ERWBarrier, EResourceTransitionPipeline::EGfxToCompute, DestRenderTarget.UAV, nullptr);
 			ComputeShader->SetParametersCS(Context.RHICmdList, Context, InputTexSize, OutputViewRect, DestRenderTarget.UAV);
-			DispatchComputeShader(Context.RHICmdList, *ComputeShader, GroupSizeX, GroupSizeY, 1);
+			DispatchComputeShader(Context.RHICmdList, ComputeShader, GroupSizeX, GroupSizeY, 1);
 			ComputeShader->UnsetParameters(Context.RHICmdList);
 		}
 	}
@@ -2800,20 +2800,20 @@ void FRCPassPostProcessAmbientOcclusion_GTAO_SpatialFilter::Process(FRenderingCo
 			const FSceneRenderTargetItem& DestRenderTarget = SceneContext.ScreenSpaceGTAOHorizons->GetRenderTargetItem();
 
 			FRHIAsyncComputeCommandListImmediate& RHICmdListComputeImmediate = FRHICommandListExecutor::GetImmediateAsyncComputeCommandList();
-			RHICmdListComputeImmediate.SetComputeShader(ComputeShader->GetComputeShader());
+			RHICmdListComputeImmediate.SetComputeShader(ComputeShader.GetComputeShader());
 			RHICmdListComputeImmediate.TransitionResource(EResourceTransitionAccess::ERWBarrier, EResourceTransitionPipeline::EComputeToCompute, DestRenderTarget.UAV, nullptr);
 			ComputeShader->SetParametersCS(RHICmdListComputeImmediate, Context, InputTexSize, OutputViewRect, DestRenderTarget.UAV);
-			DispatchComputeShader(RHICmdListComputeImmediate, *ComputeShader, GroupSizeX, GroupSizeY, 1);
+			DispatchComputeShader(RHICmdListComputeImmediate, ComputeShader, GroupSizeX, GroupSizeY, 1);
 			ComputeShader->UnsetParameters(Context.RHICmdList);
 		}
 		else
 		{
 			const FSceneRenderTargetItem& DestRenderTarget = PassOutputs[0].RequestSurface(Context);
 
-			Context.RHICmdList.SetComputeShader(ComputeShader->GetComputeShader());
+			Context.RHICmdList.SetComputeShader(ComputeShader.GetComputeShader());
 			Context.RHICmdList.TransitionResource(EResourceTransitionAccess::ERWBarrier, EResourceTransitionPipeline::EGfxToCompute, DestRenderTarget.UAV, nullptr);
 			ComputeShader->SetParametersCS(Context.RHICmdList, Context, InputTexSize, OutputViewRect, DestRenderTarget.UAV);
-			DispatchComputeShader(Context.RHICmdList, *ComputeShader, GroupSizeX, GroupSizeY, 1);
+			DispatchComputeShader(Context.RHICmdList, ComputeShader, GroupSizeX, GroupSizeY, 1);
 			ComputeShader->UnsetParameters(Context.RHICmdList);
 		}
 	}
