Set bSupportsWaveOperations=false for PCD3D_SM5 to prevent DXIL shaders from being loaded in D3D11 mode. Removed GRHISupportsRayTracing check from FMaterialSortCS::ShouldCompilePermutation, as this prevents this shader from being compiled when cooking.

This fixes ray tracing support in cooked builds.

diff --git a/Engine/Config/Windows/DataDrivenPlatformInfo.ini b/Engine/Config/Windows/DataDrivenPlatformInfo.ini
--- a/Engine/Config/Windows/DataDrivenPlatformInfo.ini
+++ b/Engine/Config/Windows/DataDrivenPlatformInfo.ini
@@ -137,5 +137,5 @@ PlatformGroupName=Desktop
 bSupportsRayTracing=true
 bSupportsRayTracingMissShaderBindings=true
 bSupportsRayTracingIndirectInstanceData=true
-bSupportsWaveOperations=true
+bSupportsWaveOperations=false
 bSupportsGPUScene=true
diff --git a/Engine/Shaders/Private/RayTracing/MaterialSort.usf b/Engine/Shaders/Private/RayTracing/MaterialSort.usf
--- a/Engine/Shaders/Private/RayTracing/MaterialSort.usf
+++ b/Engine/Shaders/Private/RayTracing/MaterialSort.usf
@@ -1,6 +1,8 @@
 #include "../Common.ush"
 #include "RayTracingCommon.ush"
 
+#define MATERIAL_SORT_SHADER_VERSION 0xd74cbc87 // Change to force shader compilation of this shader
+
 #ifndef DIM_WAVE_OPS
 	#define DIM_WAVE_OPS 0
 #endif // DIM_WAVE_OPS
diff --git a/Engine/Source/Runtime/Renderer/Private/RayTracing/RayTracingDeferredMaterials.cpp b/Engine/Source/Runtime/Renderer/Private/RayTracing/RayTracingDeferredMaterials.cpp
--- a/Engine/Source/Runtime/Renderer/Private/RayTracing/RayTracingDeferredMaterials.cpp
+++ b/Engine/Source/Runtime/Renderer/Private/RayTracing/RayTracingDeferredMaterials.cpp
@@ -86,7 +86,7 @@ class FMaterialSortCS : public FGlobalShader
 
 	static bool ShouldCompilePermutation(const FGlobalShaderPermutationParameters& Parameters)
 	{
-		if (!GRHISupportsRayTracing || !ShouldCompileRayTracingShadersForProject(Parameters.Platform))
+		if (!ShouldCompileRayTracingShadersForProject(Parameters.Platform))
 		{
 			return false;
 		}
@@ -137,7 +137,7 @@ void SortDeferredMaterials(
 
 	FMaterialSortCS::FPermutationDomain PermutationVector;
 	PermutationVector.Set<FMaterialSortCS::FSortSize>(SortSize - 1);
-	PermutationVector.Set<FMaterialSortCS::FWaveOps>(GRHISupportsWaveOperations && GRHIMinimumWaveSize >= 32);
+	PermutationVector.Set<FMaterialSortCS::FWaveOps>(GRHISupportsWaveOperations && GRHIMinimumWaveSize >= 32 && RHISupportsWaveOperations(View.GetShaderPlatform()));
 
 	// Sort size represents an index into pow2 sizes, not an actual size, so convert to the actual number of elements being sorted
 	const uint32 ElementBlockSize = 256 * (1 << (SortSize - 1));
