Fix for bloom being sampled incorrectly in tonemap pass.

diff --git a/Engine/Shaders/Private/PostProcessTonemap.usf b/Engine/Shaders/Private/PostProcessTonemap.usf
--- a/Engine/Shaders/Private/PostProcessTonemap.usf
+++ b/Engine/Shaders/Private/PostProcessTonemap.usf
@@ -37,6 +37,7 @@
 SCREEN_PASS_TEXTURE_VIEWPORT(Color)
 SCREEN_PASS_TEXTURE_VIEWPORT(Bloom)
 SCREEN_PASS_TEXTURE_VIEWPORT(Output)
+SCREEN_PASS_TEXTURE_VIEWPORT_TRANSFORM(ColorToBloom)
 
 // 64x64 tiling noise texture, optimized to have mostly high frequency content
 Texture2D NoiseTexture;
@@ -525,7 +526,7 @@ float4 TonemapCommonPS(
 		half3 LinearColor = SceneColor.rgb * ColorScale0.rgb;
 
 		#if USE_BLOOM
-			float2 BloomUV = Bloom_UVViewportMin + UV * Bloom_UVViewportSize;			
+			float2 BloomUV = ColorToBloom_Scale * UV + ColorToBloom_Bias;
 			BloomUV = clamp(BloomUV, Bloom_UVViewportBilinearMin, Bloom_UVViewportBilinearMax);
 			#if ES3_1_PROFILE
 				BloomUV = FullViewUV.xy;
diff --git a/Engine/Source/Runtime/Renderer/Private/PostProcess/PostProcessTonemap.cpp b/Engine/Source/Runtime/Renderer/Private/PostProcess/PostProcessTonemap.cpp
--- a/Engine/Source/Runtime/Renderer/Private/PostProcess/PostProcessTonemap.cpp
+++ b/Engine/Source/Runtime/Renderer/Private/PostProcess/PostProcessTonemap.cpp
@@ -522,6 +522,7 @@ BEGIN_SHADER_PARAMETER_STRUCT(FTonemapParameters, )
 	SHADER_PARAMETER_STRUCT(FScreenPassTextureViewportParameters, Color)
 	SHADER_PARAMETER_STRUCT(FScreenPassTextureViewportParameters, Bloom)
 	SHADER_PARAMETER_STRUCT(FScreenPassTextureViewportParameters, Output)
+	SHADER_PARAMETER_STRUCT(FScreenPassTextureViewportTransform, ColorToBloom)
 	SHADER_PARAMETER_RDG_TEXTURE(Texture2D, ColorTexture)
 	SHADER_PARAMETER_RDG_TEXTURE(Texture2D, BloomTexture)
 	// SM5 and above use Texture2D for EyeAdaptation
@@ -765,6 +766,7 @@ FScreenPassTexture AddTonemapPass(FRDGBuilder& GraphBuilder, const FViewInfo& Vi
 	{
 		const FScreenPassTextureViewport BloomViewport(Inputs.Bloom);
 		CommonParameters.Bloom = GetScreenPassTextureViewportParameters(BloomViewport);
+		CommonParameters.ColorToBloom = GetScreenPassTextureViewportTransform(CommonParameters.Color, CommonParameters.Bloom);
 	}
 	CommonParameters.Output = GetScreenPassTextureViewportParameters(OutputViewport);
 	CommonParameters.ColorTexture = Inputs.SceneColor.Texture;
