Record a separate list of primitives that requires interaction shadows to speed up traversal when setting up interaction shadows

diff --git a/Engine/Source/Runtime/Renderer/Private/LightSceneInfo.cpp b/Engine/Source/Runtime/Renderer/Private/LightSceneInfo.cpp
--- a/Engine/Source/Runtime/Renderer/Private/LightSceneInfo.cpp
+++ b/Engine/Source/Runtime/Renderer/Private/LightSceneInfo.cpp
@@ -18,6 +18,13 @@ static FAutoConsoleVariableRef CVarWholeSceneShadowUnbuiltInteractionThreshold(
 	ECVF_RenderThreadSafe
 	);
 
+static int32 GRecordInteractionShadowPrimitives = 1;
+FAutoConsoleVariableRef CVarRecordInteractionShadowPrimitives(
+	TEXT("r.Shadow.RecordInteractionShadowPrimitives"),
+	GRecordInteractionShadowPrimitives,
+	TEXT(""),
+	ECVF_RenderThreadSafe);
+
 void FLightSceneInfoCompact::Init(FLightSceneInfo* InLightSceneInfo)
 {
 	LightSceneInfo = InLightSceneInfo;
@@ -36,7 +43,8 @@ void FLightSceneInfoCompact::Init(FLightSceneInfo* InLightSceneInfo)
 }
 
 FLightSceneInfo::FLightSceneInfo(FLightSceneProxy* InProxy, bool InbVisible)
-	: DynamicInteractionOftenMovingPrimitiveList(NULL)
+	: bRecordInteractionShadowPrimitives(!!GRecordInteractionShadowPrimitives && InProxy->GetLightType() != ELightComponentType::LightType_Directional)
+	, DynamicInteractionOftenMovingPrimitiveList(NULL)
 	, DynamicInteractionStaticPrimitiveList(NULL)
 	, Proxy(InProxy)
 	, Id(INDEX_NONE)
@@ -156,6 +164,8 @@ void FLightSceneInfo::Detach()
 {
 	check(IsInRenderingThread());
 
+	InteractionShadowPrimitives.Empty();
+
 	// implicit linked list. The destruction will update this "head" pointer to the next item in the list.
 	while(DynamicInteractionOftenMovingPrimitiveList)
 	{
@@ -217,6 +227,15 @@ bool FLightSceneInfo::IsPrecomputedLightingValid() const
 	return (bPrecomputedLightingIsValid && NumUnbuiltInteractions < GWholeSceneShadowUnbuiltInteractionThreshold) || !Proxy->HasStaticShadowing();
 }
 
+const TArray<FLightPrimitiveInteraction*>* FLightSceneInfo::GetInteractionShadowPrimitives(bool bSync) const
+{
+	if (bSync)
+	{
+		Scene->FlushAsyncLightPrimitiveInteractionCreation();
+	}
+	return bRecordInteractionShadowPrimitives ? &InteractionShadowPrimitives : nullptr;
+}
+
 FLightPrimitiveInteraction* FLightSceneInfo::GetDynamicInteractionOftenMovingPrimitiveList(bool bSync) const
 {
 	if (bSync)
diff --git a/Engine/Source/Runtime/Renderer/Private/LightSceneInfo.h b/Engine/Source/Runtime/Renderer/Private/LightSceneInfo.h
--- a/Engine/Source/Runtime/Renderer/Private/LightSceneInfo.h
+++ b/Engine/Source/Runtime/Renderer/Private/LightSceneInfo.h
@@ -153,6 +153,9 @@ class FLightSceneInfo : public FRenderResource
 {
 	friend class FLightPrimitiveInteraction;
 
+	bool bRecordInteractionShadowPrimitives;
+	TArray<FLightPrimitiveInteraction*> InteractionShadowPrimitives;
+
 	/** The list of dynamic primitives affected by the light. */
 	FLightPrimitiveInteraction* DynamicInteractionOftenMovingPrimitiveList;
 
@@ -302,6 +305,8 @@ public:
 		return DynamicShadowMapChannel;
 	}
 
+	const TArray<FLightPrimitiveInteraction*>* GetInteractionShadowPrimitives(bool bSync = true) const;
+
 	FLightPrimitiveInteraction* GetDynamicInteractionOftenMovingPrimitiveList(bool bSync = true) const;
 
 	FLightPrimitiveInteraction* GetDynamicInteractionStaticPrimitiveList(bool bSync = true) const;
diff --git a/Engine/Source/Runtime/Renderer/Private/SceneCore.cpp b/Engine/Source/Runtime/Renderer/Private/SceneCore.cpp
--- a/Engine/Source/Runtime/Renderer/Private/SceneCore.cpp
+++ b/Engine/Source/Runtime/Renderer/Private/SceneCore.cpp
@@ -150,6 +150,8 @@ void FLightPrimitiveInteraction::Destroy(FLightPrimitiveInteraction* LightPrimit
 	delete LightPrimitiveInteraction;
 }
 
+extern bool ShouldCreateObjectShadowForStationaryLight(const FLightSceneInfo* LightSceneInfo, const FPrimitiveSceneProxy* PrimitiveSceneProxy, bool bInteractionShadowMapped);
+
 FLightPrimitiveInteraction::FLightPrimitiveInteraction(
 	FLightSceneInfo* InLightSceneInfo,
 	FPrimitiveSceneInfo* InPrimitiveSceneInfo,
@@ -268,6 +270,21 @@ FLightPrimitiveInteraction::FLightPrimitiveInteraction(
 		(*PrevLightLink)->PrevLightLink = &NextLight;
 	}
 	*PrevLightLink = this;
+
+	if (HasShadow()
+		&& LightSceneInfo->bRecordInteractionShadowPrimitives
+		&& (HasTranslucentObjectShadow() || HasInsetObjectShadow() || ShouldCreateObjectShadowForStationaryLight(LightSceneInfo, PrimitiveSceneInfo->Proxy, IsShadowMapped())))
+	{
+		if (LightSceneInfo->InteractionShadowPrimitives.Num() < 16)
+		{
+			LightSceneInfo->InteractionShadowPrimitives.Add(this);
+		}
+		else
+		{
+			LightSceneInfo->bRecordInteractionShadowPrimitives = false;
+			LightSceneInfo->InteractionShadowPrimitives.Empty();
+		}
+	}
 }
 
 FLightPrimitiveInteraction::~FLightPrimitiveInteraction()
@@ -311,6 +328,8 @@ FLightPrimitiveInteraction::~FLightPrimitiveInteraction()
 		NextLight->PrevLightLink = PrevLightLink;
 	}
 	*PrevLightLink = NextLight;
+
+	LightSceneInfo->InteractionShadowPrimitives.RemoveSingleSwap(this);
 }
 
 void FLightPrimitiveInteraction::FlushCachedShadowMapData()
diff --git a/Engine/Source/Runtime/Renderer/Private/ShadowSetup.cpp b/Engine/Source/Runtime/Renderer/Private/ShadowSetup.cpp
--- a/Engine/Source/Runtime/Renderer/Private/ShadowSetup.cpp
+++ b/Engine/Source/Runtime/Renderer/Private/ShadowSetup.cpp
@@ -4214,21 +4214,34 @@ void FSceneRenderer::InitDynamicShadows(FRHICommandListImmediate& RHICmdList, FG
 						{
 							Scene->FlushAsyncLightPrimitiveInteractionCreation();
 
-							// Look for individual primitives with a dynamic shadow.
-							for (FLightPrimitiveInteraction* Interaction = LightSceneInfo->GetDynamicInteractionOftenMovingPrimitiveList(false);
-								Interaction;
-								Interaction = Interaction->GetNextPrimitive()
-								)
+							const TArray<FLightPrimitiveInteraction*>* InteractionShadowPrimitives = LightSceneInfo->GetInteractionShadowPrimitives(false);
+
+							if (InteractionShadowPrimitives)
 							{
-								SetupInteractionShadows(RHICmdList, Interaction, VisibleLightInfo, bStaticSceneOnly, ViewDependentWholeSceneShadows, PreShadows);
+								const int32 NumPrims = InteractionShadowPrimitives->Num();
+								for (int32 Idx = 0; Idx < NumPrims; ++Idx)
+								{
+									SetupInteractionShadows(RHICmdList, (*InteractionShadowPrimitives)[Idx], VisibleLightInfo, bStaticSceneOnly, ViewDependentWholeSceneShadows, PreShadows);
+								}
 							}
-
-							for (FLightPrimitiveInteraction* Interaction = LightSceneInfo->GetDynamicInteractionStaticPrimitiveList(false);
-								Interaction;
-								Interaction = Interaction->GetNextPrimitive()
-								)
+							else
 							{
-								SetupInteractionShadows(RHICmdList, Interaction, VisibleLightInfo, bStaticSceneOnly, ViewDependentWholeSceneShadows, PreShadows);
+								// Look for individual primitives with a dynamic shadow.
+								for (FLightPrimitiveInteraction* Interaction = LightSceneInfo->GetDynamicInteractionOftenMovingPrimitiveList(false);
+									Interaction;
+									Interaction = Interaction->GetNextPrimitive()
+									)
+								{
+									SetupInteractionShadows(RHICmdList, Interaction, VisibleLightInfo, bStaticSceneOnly, ViewDependentWholeSceneShadows, PreShadows);
+								}
+
+								for (FLightPrimitiveInteraction* Interaction = LightSceneInfo->GetDynamicInteractionStaticPrimitiveList(false);
+									Interaction;
+									Interaction = Interaction->GetNextPrimitive()
+									)
+								{
+									SetupInteractionShadows(RHICmdList, Interaction, VisibleLightInfo, bStaticSceneOnly, ViewDependentWholeSceneShadows, PreShadows);
+								}
 							}
 						}
 					}
