Fix an exponential height fog sample out of bound issue

diff --git a/Engine/Shaders/Private/HeightFogPixelShader.usf b/Engine/Shaders/Private/HeightFogPixelShader.usf
--- a/Engine/Shaders/Private/HeightFogPixelShader.usf
+++ b/Engine/Shaders/Private/HeightFogPixelShader.usf
@@ -12,6 +12,8 @@
 float bOnlyOnRenderedOpaque;
 float bUseLinearDepthTexture;
 
+float4 LinearDepthTextureMinMaxUV;
+
 Texture2D OcclusionTexture;
 SamplerState OcclusionSampler;
 
@@ -29,7 +31,8 @@ void ExponentialPixelMain(
 	float DeviceZ;
 	if (bUseLinearDepthTexture > 0.0f)
 	{
-		SceneDepth = Texture2DSampleLevel(LinearDepthTexture, LinearDepthSampler, TexCoord, 0).r * bUseLinearDepthTexture; // bUseLinearDepthTexture is also a scale to get depth in correct range
+		float2 SampleUV = clamp(TexCoord, LinearDepthTextureMinMaxUV.xy, LinearDepthTextureMinMaxUV.zw);
+		SceneDepth = Texture2DSampleLevel(LinearDepthTexture, LinearDepthSampler, SampleUV, 0).r * bUseLinearDepthTexture; // bUseLinearDepthTexture is also a scale to get depth in correct range
 		DeviceZ = ConvertToDeviceZ(SceneDepth);
 	}
 	else
diff --git a/Engine/Source/Runtime/Renderer/Private/FogRendering.cpp b/Engine/Source/Runtime/Renderer/Private/FogRendering.cpp
--- a/Engine/Source/Runtime/Renderer/Private/FogRendering.cpp
+++ b/Engine/Source/Runtime/Renderer/Private/FogRendering.cpp
@@ -50,6 +50,7 @@ struct FHeightFogRenderingParameters
 	FTextureRHIRef LinearDepthTextureRHI;
 	FIntRect ViewRect;
 	float LinearDepthReadScale;
+	FVector4 LinearDepthMinMaxUV;
 };
 
 void SetupFogUniformParameters(const FViewInfo& View, FFogUniformParameters& OutParameters)
@@ -192,6 +193,7 @@ public:
 		LinearDepthSampler.Bind(Initializer.ParameterMap, TEXT("LinearDepthSampler"));
 		bOnlyOnRenderedOpaque.Bind(Initializer.ParameterMap, TEXT("bOnlyOnRenderedOpaque"));
 		bUseLinearDepthTexture.Bind(Initializer.ParameterMap, TEXT("bUseLinearDepthTexture"));
+		LinearDepthTextureMinMaxUV.Bind(Initializer.ParameterMap, TEXT("LinearDepthTextureMinMaxUV"));
 		SceneTextureParameters.Bind(Initializer);
 	}
 
@@ -227,6 +229,7 @@ public:
 
 		SetShaderValue(RHICmdList, RHICmdList.GetBoundPixelShader(), bOnlyOnRenderedOpaque, View.bFogOnlyOnRenderedOpaque ? 1.0f : 0.0f);
 		SetShaderValue(RHICmdList, RHICmdList.GetBoundPixelShader(), bUseLinearDepthTexture, bUseLinearDepthTextureEnabled ? Params.LinearDepthReadScale : 0.0f);
+		SetShaderValue(RHICmdList, RHICmdList.GetBoundPixelShader(), LinearDepthTextureMinMaxUV, Params.LinearDepthMinMaxUV);
 	}
 
 private:
@@ -237,6 +240,7 @@ private:
 	LAYOUT_FIELD(FShaderResourceParameter, LinearDepthSampler);
 	LAYOUT_FIELD(FShaderParameter, bOnlyOnRenderedOpaque)
 	LAYOUT_FIELD(FShaderParameter, bUseLinearDepthTexture);
+	LAYOUT_FIELD(FShaderParameter, LinearDepthTextureMinMaxUV);
 };
 
 IMPLEMENT_SHADER_TYPE(template<>,TExponentialHeightFogPS<EHeightFogFeature::HeightFog>,TEXT("/Engine/Private/HeightFogPixelShader.usf"), TEXT("ExponentialPixelMain"),SF_Pixel)
@@ -463,7 +467,7 @@ bool FDeferredShadingSceneRenderer::RenderFog(FRHICommandListImmediate& RHICmdLi
 
 		SceneContext.BeginRenderingSceneColor(RHICmdList, ESimpleRenderTargetMode::EExistingColorAndDepth, FExclusiveDepthStencil::DepthRead_StencilWrite, true);
 
-		FHeightFogRenderingParameters Parameters = { LightShaftsOutput, nullptr, FIntRect(), 1.0f };
+		FHeightFogRenderingParameters Parameters = { LightShaftsOutput, nullptr, FIntRect(), 1.0f, FVector4() };
 
 		for(int32 ViewIndex = 0;ViewIndex < Views.Num();ViewIndex++)
 		{
@@ -503,7 +507,7 @@ void FDeferredShadingSceneRenderer::RenderUnderWaterFog(FRHICommandListImmediate
 		LightShaftsOutput.LightShaftOcclusion = nullptr;
 		FTextureRHIRef LinearDepthTextureRHI = PassData.SceneDepthWithoutSingleLayerWater->GetRenderTargetItem().ShaderResourceTexture;
 		const float SINGLE_LAYER_WATER_DEPTH_SCALE = 100.0f; // This must match SINGLE_LAYER_WATER_DEPTH_SCALE from SingleLayerWaterCommon.ush and SingleLayerWaterComposite.usf. TODO deduplicate
-		FHeightFogRenderingParameters Parameters = { LightShaftsOutput, LinearDepthTextureRHI, FIntRect(), SINGLE_LAYER_WATER_DEPTH_SCALE };
+		FHeightFogRenderingParameters Parameters = { LightShaftsOutput, LinearDepthTextureRHI, FIntRect(), SINGLE_LAYER_WATER_DEPTH_SCALE, FVector4() };
 
 		for (int32 ViewIndex = 0; ViewIndex < Views.Num(); ViewIndex++)
 		{
@@ -514,6 +518,7 @@ void FDeferredShadingSceneRenderer::RenderUnderWaterFog(FRHICommandListImmediate
 
 				// Specify the low resolution view rect
 				Parameters.ViewRect = PassData.ViewData[ViewIndex].SceneWithoutSingleLayerWaterViewRect;
+				Parameters.LinearDepthMinMaxUV = PassData.ViewData[ViewIndex].SceneWithoutSingleLayerWaterMinMaxUV;
 
 				RenderViewFog(RHICmdList, View, Parameters);
 			}
