Integrate Fix CL 11421022 from Dev-Rendering

#rnx

diff --git a/Engine/Source/Runtime/Renderer/Private/RendererScene.cpp b/Engine/Source/Runtime/Renderer/Private/RendererScene.cpp
--- a/Engine/Source/Runtime/Renderer/Private/RendererScene.cpp
+++ b/Engine/Source/Runtime/Renderer/Private/RendererScene.cpp
@@ -3516,6 +3516,8 @@ void FScene::OnLevelAddedToWorld(FName LevelAddedName, UWorld* InWorld, bool bIs
 
 void FScene::OnLevelAddedToWorld_RenderThread(FName InLevelName)
 {
+	TArray<FPrimitiveSceneInfo*> PrimitivesToReAddStaticMeshes;
+
 	// Mark level primitives
 	for (auto It = Primitives.CreateIterator(); It; ++It)
 	{
@@ -3528,10 +3530,12 @@ void FScene::OnLevelAddedToWorld_RenderThread(FName InLevelName)
 				// The only type of SceneProxy using this is landscape
 				(*It)->RemoveStaticMeshes();
 				Proxy->OnLevelAddedToWorld();
-				(*It)->AddStaticMeshes(FRHICommandListExecutor::GetImmediateCommandList(), this, Primitives);
+				PrimitivesToReAddStaticMeshes.Add(*It);
 			}
 		}
 	}
+
+	FPrimitiveSceneInfo::AddStaticMeshes(FRHICommandListExecutor::GetImmediateCommandList(), this, PrimitivesToReAddStaticMeshes);
 }
 
 void FScene::ProcessAtmosphereLightAddition_RenderThread(FLightSceneInfo* LightSceneInfo)
