Fixed crash in shadow generation code.

diff --git a/Engine/Source/Runtime/Renderer/Private/RendererScene.cpp b/Engine/Source/Runtime/Renderer/Private/RendererScene.cpp
--- a/Engine/Source/Runtime/Renderer/Private/RendererScene.cpp
+++ b/Engine/Source/Runtime/Renderer/Private/RendererScene.cpp
@@ -4072,7 +4072,7 @@ void FScene::UpdateAllPrimitiveSceneInfos(FRHICommandListImmediate& RHICmdList,
 
 		// Re-add the primitive to the scene with the new transform.
 		FPrimitiveSceneInfo::AddToScene(RHICmdList, this, UpdatedSceneInfosWithStaticDrawListUpdate, true, true, bAsyncCreateLPIs);
-		FPrimitiveSceneInfo::AddToScene(RHICmdList, this, UpdatedSceneInfosWithoutStaticDrawListUpdate, false, bAsyncCreateLPIs);
+		FPrimitiveSceneInfo::AddToScene(RHICmdList, this, UpdatedSceneInfosWithoutStaticDrawListUpdate, false, true, bAsyncCreateLPIs);
 
 		for (FPrimitiveSceneInfo* PrimitiveSceneInfo : UpdatedSceneInfosWithoutStaticDrawListUpdate)
 		{
