Fix for sublevel loading.

diff --git a/Engine/Source/Runtime/Renderer/Private/RendererScene.cpp b/Engine/Source/Runtime/Renderer/Private/RendererScene.cpp
--- a/Engine/Source/Runtime/Renderer/Private/RendererScene.cpp
+++ b/Engine/Source/Runtime/Renderer/Private/RendererScene.cpp
@@ -3516,9 +3516,9 @@ void FScene::OnLevelAddedToWorld(FName LevelAddedName, UWorld* InWorld, bool bIs
 
 void FScene::OnLevelAddedToWorld_RenderThread(FName InLevelName)
 {
-	TArray<FPrimitiveSceneInfo*> PrimitivesToReAddStaticMeshes;
-
 	// Mark level primitives
+	TArray<FPrimitiveSceneInfo*> PrimitivesToAdd;
+
 	for (auto It = Primitives.CreateIterator(); It; ++It)
 	{
 		FPrimitiveSceneProxy* Proxy = (*It)->Proxy;
@@ -3530,12 +3530,12 @@ void FScene::OnLevelAddedToWorld_RenderThread(FName InLevelName)
 				// The only type of SceneProxy using this is landscape
 				(*It)->RemoveStaticMeshes();
 				Proxy->OnLevelAddedToWorld();
-				PrimitivesToReAddStaticMeshes.Add(*It);
+				PrimitivesToAdd.Add(*It);
 			}
 		}
 	}
 
-	FPrimitiveSceneInfo::AddStaticMeshes(FRHICommandListExecutor::GetImmediateCommandList(), this, PrimitivesToReAddStaticMeshes);
+	FPrimitiveSceneInfo::AddStaticMeshes(FRHICommandListExecutor::GetImmediateCommandList(), this, PrimitivesToAdd);
 }
 
 void FScene::ProcessAtmosphereLightAddition_RenderThread(FLightSceneInfo* LightSceneInfo)
