Make CVarMobileEyeAdaptation could be toggled at runtime.

diff --git a/Engine/Source/Runtime/Renderer/Private/PostProcess/PostProcessMobile.cpp b/Engine/Source/Runtime/Renderer/Private/PostProcess/PostProcessMobile.cpp
--- a/Engine/Source/Runtime/Renderer/Private/PostProcess/PostProcessMobile.cpp
+++ b/Engine/Source/Runtime/Renderer/Private/PostProcess/PostProcessMobile.cpp
@@ -29,7 +29,7 @@ static TAutoConsoleVariable<int32> CVarMobileEyeAdaptation(
 	TEXT("EyeAdaptation for mobile platform.\n")
 	TEXT(" 0: Disable\n")
 	TEXT(" 1: Enabled (Default)"),
-	ECVF_ReadOnly);
+	ECVF_RenderThreadSafe);
 
 static EPixelFormat GetHDRPixelFormat()
 {
