FixedL "Render in main pass" primitive option is ignored for Mobile Renderer

diff --git a/Engine/Source/Runtime/Renderer/Private/MobileBasePass.cpp b/Engine/Source/Runtime/Renderer/Private/MobileBasePass.cpp
--- a/Engine/Source/Runtime/Renderer/Private/MobileBasePass.cpp
+++ b/Engine/Source/Runtime/Renderer/Private/MobileBasePass.cpp
@@ -554,7 +554,7 @@ FMobileBasePassMeshProcessor::FMobileBasePassMeshProcessor(
 
 void FMobileBasePassMeshProcessor::AddMeshBatch(const FMeshBatch& RESTRICT MeshBatch, uint64 BatchElementMask, const FPrimitiveSceneProxy* RESTRICT PrimitiveSceneProxy, int32 StaticMeshId)
 {
-	if (!MeshBatch.bUseForMaterial)
+	if (!MeshBatch.bUseForMaterial || (PrimitiveSceneProxy && !PrimitiveSceneProxy->ShouldRenderInMainPass()))
 	{
 		return;
 	}
