Shadow rays in RTGI are shortened by default to avoid hitting the sky sphere. That makes it possible to get desired sky reflections without affecting GI.

diff --git a/Engine/Shaders/Private/RayTracing/RayTracingCreateGatherPointsRGS.usf b/Engine/Shaders/Private/RayTracing/RayTracingCreateGatherPointsRGS.usf
--- a/Engine/Shaders/Private/RayTracing/RayTracingCreateGatherPointsRGS.usf
+++ b/Engine/Shaders/Private/RayTracing/RayTracingCreateGatherPointsRGS.usf
@@ -35,6 +35,7 @@ uint UpscaleFactor;
 uint RenderTileOffsetX;
 uint RenderTileOffsetY;
 float MaxRayDistanceForGI;
+float MaxShadowDistance;
 float NextEventEstimationSamples;
 float DiffuseThreshold;
 float MaxNormalBias;
@@ -349,6 +350,7 @@ RAY_TRACING_ENTRY_RAYGEN(RayTracingCreateGatherPointsRGS)
 				{
 					RayDesc LightRay;
 					GenerateLightRay(Ray, LightId, LightUV, LightRay);
+					LightRay.TMax = min(LightRay.TMax, MaxShadowDistance);
 					ApplyPositionBias(LightRay, Payload.WorldNormal, MaxNormalBias);
 
 					// Trace visibility ray
diff --git a/Engine/Shaders/Private/RayTracing/RayTracingGlobalIlluminationRGS.usf b/Engine/Shaders/Private/RayTracing/RayTracingGlobalIlluminationRGS.usf
--- a/Engine/Shaders/Private/RayTracing/RayTracingGlobalIlluminationRGS.usf
+++ b/Engine/Shaders/Private/RayTracing/RayTracingGlobalIlluminationRGS.usf
@@ -26,6 +26,7 @@ uint MaxBounces;
 uint UpscaleFactor;
 float MaxRayDistanceForGI;
 float MaxRayDistanceForAO;
+float MaxShadowDistance;
 float NextEventEstimationSamples;
 float DiffuseThreshold;
 uint EvalSkyLight;
@@ -235,6 +236,7 @@ RAY_TRACING_ENTRY_RAYGEN(GlobalIlluminationRGS)
 				{
 					RayDesc LightRay;
 					GenerateLightRay(Ray, LightId, LightUV, LightRay);
+					LightRay.TMax = min(LightRay.TMax, MaxShadowDistance);
 					ApplyPositionBias(LightRay, Payload.WorldNormal, MaxNormalBias);
 
 					// Trace visibility ray
diff --git a/Engine/Source/Runtime/Renderer/Private/RayTracing/RayTracingGlobalIllumination.cpp b/Engine/Source/Runtime/Renderer/Private/RayTracing/RayTracingGlobalIllumination.cpp
--- a/Engine/Source/Runtime/Renderer/Private/RayTracing/RayTracingGlobalIllumination.cpp
+++ b/Engine/Source/Runtime/Renderer/Private/RayTracing/RayTracingGlobalIllumination.cpp
@@ -47,6 +47,13 @@ static FAutoConsoleVariableRef CVarRayTracingGlobalIlluminationMaxRayDistance(
 	TEXT("Max ray distance (default = 1.0e27)")
 );
 
+static float GRayTracingGlobalIlluminationMaxShadowDistance = -1.0;
+static FAutoConsoleVariableRef CVarRayTracingGlobalIlluminationMaxShadowDistance(
+	TEXT("r.RayTracing.GlobalIllumination.MaxShadowDistance"),
+	GRayTracingGlobalIlluminationMaxShadowDistance,
+	TEXT("Max shadow distance (default = -1.0, distance adjusted automatically so shadow rays do not hit the sky sphere) ")
+);
+
 static TAutoConsoleVariable<int32> CVarRayTracingGlobalIlluminationMaxBounces(
 	TEXT("r.RayTracing.GlobalIllumination.MaxBounces"),
 	-1,
@@ -336,6 +343,7 @@ class FGlobalIlluminationRGS : public FGlobalShader
 		SHADER_PARAMETER(uint32, UpscaleFactor)
 		SHADER_PARAMETER(float, MaxRayDistanceForGI)
 		SHADER_PARAMETER(float, MaxRayDistanceForAO)
+		SHADER_PARAMETER(float, MaxShadowDistance)
 		SHADER_PARAMETER(float, NextEventEstimationSamples)
 		SHADER_PARAMETER(float, DiffuseThreshold)
 		SHADER_PARAMETER(uint32, EvalSkyLight)
@@ -396,6 +404,7 @@ class FRayTracingGlobalIlluminationCreateGatherPointsRGS : public FGlobalShader
 		SHADER_PARAMETER(uint32, RenderTileOffsetX)
 		SHADER_PARAMETER(uint32, RenderTileOffsetY)
 		SHADER_PARAMETER(float, MaxRayDistanceForGI)
+		SHADER_PARAMETER(float, MaxShadowDistance)
 		SHADER_PARAMETER(float, NextEventEstimationSamples)
 		SHADER_PARAMETER(float, DiffuseThreshold)
 		SHADER_PARAMETER(float, MaxNormalBias)
@@ -458,6 +467,7 @@ class FRayTracingGlobalIlluminationCreateGatherPointsTraceRGS : public FGlobalSh
 		SHADER_PARAMETER(uint32, RenderTileOffsetX)
 		SHADER_PARAMETER(uint32, RenderTileOffsetY)
 		SHADER_PARAMETER(float, MaxRayDistanceForGI)
+		SHADER_PARAMETER(float, MaxShadowDistance)
 		SHADER_PARAMETER(float, NextEventEstimationSamples)
 		SHADER_PARAMETER(float, DiffuseThreshold)
 		SHADER_PARAMETER(float, MaxNormalBias)
@@ -671,6 +681,7 @@ void CopyGatherPassParameters(
 	NewParameters->RenderTileOffsetX = PassParameters.RenderTileOffsetX;
 	NewParameters->RenderTileOffsetY = PassParameters.RenderTileOffsetY;
 	NewParameters->MaxRayDistanceForGI = PassParameters.MaxRayDistanceForGI;
+	NewParameters->MaxShadowDistance = PassParameters.MaxShadowDistance;
 	NewParameters->NextEventEstimationSamples = PassParameters.NextEventEstimationSamples;
 	NewParameters->DiffuseThreshold = PassParameters.DiffuseThreshold;
 	NewParameters->MaxNormalBias = PassParameters.MaxNormalBias;
@@ -734,6 +745,17 @@ void FDeferredShadingSceneRenderer::RayTracingGlobalIlluminationCreateGatherPoin
 	FPathTracingLightData LightParameters;
 	SetupLightParameters(*Scene, View, &LightParameters);
 
+	float MaxShadowDistance = 1.0e27;
+	if (GRayTracingGlobalIlluminationMaxShadowDistance > 0.0)
+	{
+		MaxShadowDistance = GRayTracingGlobalIlluminationMaxShadowDistance;
+	}
+	else if (Scene->SkyLight)
+	{
+		// Adjust ray TMax so shadow rays do not hit the sky sphere 
+		MaxShadowDistance = FMath::Max(0.0, 0.99 * Scene->SkyLight->SkyDistanceThreshold);
+	}
+
 	FSkyLightData SkyLightParameters;
 	SetupGlobalIlluminationSkyLightParameters(*Scene, &SkyLightParameters);
 
@@ -744,6 +766,7 @@ void FDeferredShadingSceneRenderer::RayTracingGlobalIlluminationCreateGatherPoin
 	PassParameters->MaxBounces = 1;
 	PassParameters->MaxNormalBias = GetRaytracingMaxNormalBias();
 	PassParameters->MaxRayDistanceForGI = GRayTracingGlobalIlluminationMaxRayDistance;
+	PassParameters->MaxShadowDistance = MaxShadowDistance;
 	PassParameters->EvalSkyLight = GRayTracingGlobalIlluminationEvalSkyLight != 0;
 	PassParameters->UseRussianRoulette = GRayTracingGlobalIlluminationUseRussianRoulette != 0;
 	PassParameters->DiffuseThreshold = GRayTracingGlobalIlluminationDiffuseThreshold;
@@ -1014,6 +1037,17 @@ void FDeferredShadingSceneRenderer::RenderRayTracingGlobalIlluminationBruteForce
 	FPathTracingLightData LightParameters;
 	SetupLightParameters(*Scene, View, &LightParameters);
 
+	float MaxShadowDistance = 1.0e27;
+	if (GRayTracingGlobalIlluminationMaxShadowDistance > 0.0)
+	{
+		MaxShadowDistance = GRayTracingGlobalIlluminationMaxShadowDistance;
+	}
+	else if (Scene->SkyLight)
+	{
+		// Adjust ray TMax so shadow rays do not hit the sky sphere 
+		MaxShadowDistance = FMath::Max(0.0, 0.99 * Scene->SkyLight->SkyDistanceThreshold);
+	}
+
 	FSkyLightData SkyLightParameters;
 	SetupGlobalIlluminationSkyLightParameters(*Scene, &SkyLightParameters);
 
@@ -1029,6 +1063,7 @@ void FDeferredShadingSceneRenderer::RenderRayTracingGlobalIlluminationBruteForce
 	}
 	PassParameters->MaxRayDistanceForGI = MaxRayDistanceForGI;
 	PassParameters->MaxRayDistanceForAO = View.FinalPostProcessSettings.AmbientOcclusionRadius;
+	PassParameters->MaxShadowDistance = MaxShadowDistance;
 	PassParameters->UpscaleFactor = UpscaleFactor;
 	PassParameters->EvalSkyLight = GRayTracingGlobalIlluminationEvalSkyLight != 0;
 	PassParameters->UseRussianRoulette = GRayTracingGlobalIlluminationUseRussianRoulette != 0;
