Fixes responsive TAA

diff --git a/Engine/Shaders/Private/PostProcessTemporalAA.usf b/Engine/Shaders/Private/PostProcessTemporalAA.usf
--- a/Engine/Shaders/Private/PostProcessTemporalAA.usf
+++ b/Engine/Shaders/Private/PostProcessTemporalAA.usf
@@ -465,6 +465,8 @@ SamplerState SceneDepthBufferSampler;
 Texture2D SceneVelocityBuffer;
 SamplerState SceneVelocityBufferSampler;
 
+Texture2D<uint2> StencilTexture;
+
 Texture2D HistoryBuffer0;
 Texture2D HistoryBuffer1;
 SamplerState HistoryBuffer0Sampler;
@@ -1902,7 +1904,7 @@ FTAAHistoryPayload TemporalAASample(uint2 GroupId, uint2 GroupThreadId, uint Gro
 			const uint kResponsiveStencilMask = 1 << 3;
 			
 			int2 SceneStencilUV = (int2)trunc(InputParams.NearestBufferUV * InputSceneColorSize.xy);
-			uint SceneStencilRef = SceneDepthBuffer.Load(int3(SceneStencilUV, 0)) STENCIL_COMPONENT_SWIZZLE;
+			uint SceneStencilRef = StencilTexture.Load(int3(SceneStencilUV, 0)) STENCIL_COMPONENT_SWIZZLE;
 
 			InputParams.bIsResponsiveAAPixel = (SceneStencilRef & kResponsiveStencilMask) ? 1.f : 0.f;
 		}
diff --git a/Engine/Source/Runtime/Renderer/Private/PostProcess/PostProcessTemporalAA.cpp b/Engine/Source/Runtime/Renderer/Private/PostProcess/PostProcessTemporalAA.cpp
--- a/Engine/Source/Runtime/Renderer/Private/PostProcess/PostProcessTemporalAA.cpp
+++ b/Engine/Source/Runtime/Renderer/Private/PostProcess/PostProcessTemporalAA.cpp
@@ -121,6 +121,8 @@ class FTemporalAACS : public FGlobalShader
 		SHADER_PARAMETER_SAMPLER(SamplerState, SceneDepthBufferSampler)
 		SHADER_PARAMETER_SAMPLER(SamplerState, SceneVelocityBufferSampler)
 
+		SHADER_PARAMETER_RDG_TEXTURE_SRV(Texture2D, StencilTexture)
+
 		SHADER_PARAMETER_STRUCT_REF(FViewUniformShaderParameters, ViewUniformBuffer)
 		
 		// Temporal upsample specific parameters.
@@ -513,6 +515,8 @@ FTAAOutputs AddTemporalAAPass(
 		PassParameters->SceneDepthBufferSampler = TStaticSamplerState<SF_Point>::GetRHI();
 		PassParameters->SceneVelocityBufferSampler = TStaticSamplerState<SF_Point>::GetRHI();
 
+		PassParameters->StencilTexture = GraphBuilder.CreateSRV(FRDGTextureSRVDesc::CreateWithPixelFormat(SceneTextures.SceneDepthBuffer, PF_X24_G8));
+
 		// We need a valid velocity buffer texture. Use black (no velocity) if none exists.
 		if (!PassParameters->SceneTextures.SceneVelocityBuffer)
 		{
