Move more work from FGatherShadowPrimitivesPacket::RenderThreadFinalize to AnyThreadTask

diff --git a/Engine/Source/Runtime/Renderer/Private/SceneRendering.h b/Engine/Source/Runtime/Renderer/Private/SceneRendering.h
--- a/Engine/Source/Runtime/Renderer/Private/SceneRendering.h
+++ b/Engine/Source/Runtime/Renderer/Private/SceneRendering.h
@@ -1251,7 +1251,12 @@ public:
 	void InitRHIResources();
 
 	/** Determines distance culling and fades if the state changes */
-	bool IsDistanceCulled(float DistanceSquared, float MaxDrawDistance, float MinDrawDistance, const FPrimitiveSceneInfo* PrimitiveSceneInfo);
+	bool IsDistanceCulled(float DistanceSquared, float MinDrawDistance, float InMaxDrawDistance, const FPrimitiveSceneInfo* PrimitiveSceneInfo);
+
+	bool IsDistanceCulled_AnyThread(float DistanceSquared, float MinDrawDistance, float InMaxDrawDistance, const FPrimitiveSceneInfo* PrimitiveSceneInfo, bool& bOutMayBeFading, bool& bOutFadingIn) const;
+
+	/** @return - whether this primitive has completely faded out */
+	bool UpdatePrimitiveFadingState(const FPrimitiveSceneInfo* PrimitiveSceneInfo, bool bFadingIn);
 
 	/** Gets the eye adaptation render target for this view. Same as GetEyeAdaptationRT */
 	IPooledRenderTarget* GetEyeAdaptation(FRHICommandList& RHICmdList) const;
diff --git a/Engine/Source/Runtime/Renderer/Private/SceneVisibility.cpp b/Engine/Source/Runtime/Renderer/Private/SceneVisibility.cpp
--- a/Engine/Source/Runtime/Renderer/Private/SceneVisibility.cpp
+++ b/Engine/Source/Runtime/Renderer/Private/SceneVisibility.cpp
@@ -256,7 +256,7 @@ DECLARE_CYCLE_STAT(TEXT("After Occlusion Readback"), STAT_CLMM_AfterOcclusionRea
  * @param View - The view for which to update.
  * @param bVisible - Whether the primitive should be visible in the view.
  */
-static void UpdatePrimitiveFadingState(FPrimitiveFadingState& FadingState, FViewInfo& View, bool bVisible)
+static void UpdatePrimitiveFadingStateHelper(FPrimitiveFadingState& FadingState, FViewInfo& View, bool bVisible)
 {
 	if (FadingState.bValid)
 	{
@@ -323,47 +323,70 @@ static void UpdatePrimitiveFadingState(FPrimitiveFadingState& FadingState, FView
 
 bool FViewInfo::IsDistanceCulled( float DistanceSquared, float MinDrawDistance, float InMaxDrawDistance, const FPrimitiveSceneInfo* PrimitiveSceneInfo)
 {
-	float MaxDrawDistanceScale = GetCachedScalabilityCVars().ViewDistanceScale;
-	float FadeRadius = GDisableLODFade ? 0.0f : GDistanceFadeMaxTravel;
-	float MaxDrawDistance = InMaxDrawDistance * MaxDrawDistanceScale;
+	bool bMayBeFading;
+	bool bFadingIn;
+	bool bDistanceCulled = IsDistanceCulled_AnyThread(DistanceSquared, MinDrawDistance, InMaxDrawDistance, PrimitiveSceneInfo, bMayBeFading, bFadingIn);
+
+	if (bMayBeFading)
+	{
+		bDistanceCulled = UpdatePrimitiveFadingState(PrimitiveSceneInfo, bFadingIn);
+	}
+
+	return bDistanceCulled;
+}
+
+bool FViewInfo::IsDistanceCulled_AnyThread(float DistanceSquared, float MinDrawDistance, float InMaxDrawDistance, const FPrimitiveSceneInfo* PrimitiveSceneInfo, bool& bOutMayBeFading, bool& bOutFadingIn) const
+{
+	const float MaxDrawDistanceScale = GetCachedScalabilityCVars().ViewDistanceScale;
+	const float FadeRadius = GDisableLODFade ? 0.0f : GDistanceFadeMaxTravel;
+	const float MaxDrawDistance = InMaxDrawDistance * MaxDrawDistanceScale;
+	bOutMayBeFading = false;
 
 	// If cull distance is disabled, always show (except foliage)
-	if (Family->EngineShowFlags.DistanceCulledPrimitives
-		&& !PrimitiveSceneInfo->Proxy->IsDetailMesh())
+	if (Family->EngineShowFlags.DistanceCulledPrimitives && !PrimitiveSceneInfo->Proxy->IsDetailMesh())
 	{
 		return false;
 	}
 
 	// The primitive is always culled if it exceeds the max fade distance.
-	if (DistanceSquared > FMath::Square(MaxDrawDistance + FadeRadius) ||
-		DistanceSquared < FMath::Square(MinDrawDistance))
+	if (DistanceSquared > FMath::Square(MaxDrawDistance + FadeRadius) || DistanceSquared < FMath::Square(MinDrawDistance))
 	{
 		return true;
 	}
 
 	const bool bDistanceCulled = (DistanceSquared > FMath::Square(MaxDrawDistance));
 	const bool bMayBeFading = (DistanceSquared > FMath::Square(MaxDrawDistance - FadeRadius));
+	
+	if (!GDisableLODFade && bMayBeFading && State != NULL && !bDisableDistanceBasedFadeTransitions && PrimitiveSceneInfo->Proxy->IsUsingDistanceCullFade())
+	{
+		// Don't update primitive fading state yet because current thread may be not render thread
+		bOutMayBeFading = true;
+		bOutFadingIn = !bDistanceCulled;
+	}
+
+	return bDistanceCulled && !bOutMayBeFading;
+}
 
+bool FViewInfo::UpdatePrimitiveFadingState(const FPrimitiveSceneInfo* PrimitiveSceneInfo, bool bFadingIn)
+{
+	// Update distance-based visibility and fading state if it has not already been updated.
+	const int32 PrimitiveIndex = PrimitiveSceneInfo->GetIndex();
+	const FRelativeBitReference PrimitiveBit(PrimitiveIndex);
 	bool bStillFading = false;
-	if  (!GDisableLODFade && bMayBeFading && State != NULL && !bDisableDistanceBasedFadeTransitions && PrimitiveSceneInfo->Proxy->IsUsingDistanceCullFade())
+
+	if (!PotentiallyFadingPrimitiveMap.AccessCorrespondingBit(PrimitiveBit))
 	{
-		// Update distance-based visibility and fading state if it has not already been updated.
-		int32 PrimitiveIndex = PrimitiveSceneInfo->GetIndex();
-		FRelativeBitReference PrimitiveBit(PrimitiveIndex);
-		if (PotentiallyFadingPrimitiveMap.AccessCorrespondingBit(PrimitiveBit) == false)
-		{
-			FPrimitiveFadingState& FadingState = ((FSceneViewState*)State)->PrimitiveFadingStates.FindOrAdd(PrimitiveSceneInfo->PrimitiveComponentId);
-			UpdatePrimitiveFadingState(FadingState, *this, !bDistanceCulled);
-			FRHIUniformBuffer* UniformBuffer = FadingState.UniformBuffer;
-			bStillFading = (UniformBuffer != NULL);
-			PrimitiveFadeUniformBuffers[PrimitiveIndex] = UniformBuffer;
-			PrimitiveFadeUniformBufferMap[PrimitiveIndex] = UniformBuffer != nullptr;
-			PotentiallyFadingPrimitiveMap.AccessCorrespondingBit(PrimitiveBit) = true;
-		}
+		FPrimitiveFadingState& FadingState = ((FSceneViewState*)State)->PrimitiveFadingStates.FindOrAdd(PrimitiveSceneInfo->PrimitiveComponentId);
+		UpdatePrimitiveFadingStateHelper(FadingState, *this, bFadingIn);
+		FRHIUniformBuffer* UniformBuffer = FadingState.UniformBuffer;
+		bStillFading = UniformBuffer != nullptr;
+		PrimitiveFadeUniformBuffers[PrimitiveIndex] = UniformBuffer;
+		PrimitiveFadeUniformBufferMap[PrimitiveIndex] = UniformBuffer != nullptr;
+		PotentiallyFadingPrimitiveMap.AccessCorrespondingBit(PrimitiveBit) = true;
 	}
 
 	// If we're still fading then make sure the object is still drawn, even if it's beyond the max draw distance
-	return ( bDistanceCulled && !bStillFading );
+	return !bFadingIn && !bStillFading;
 }
 
 FORCEINLINE bool IntersectBox8Plane(const FVector& InOrigin, const FVector& InExtent, const FPlane*PermutedPlanePtr)
@@ -598,7 +621,7 @@ static void UpdatePrimitiveFading(const FScene* Scene, FViewInfo& View)
 			{
 				bool bVisible = View.PrimitiveVisibilityMap.AccessCorrespondingBit(BitIt);
 				FPrimitiveFadingState& FadingState = ViewState->PrimitiveFadingStates.FindOrAdd(Scene->PrimitiveComponentIds[BitIt.GetIndex()]);
-				UpdatePrimitiveFadingState(FadingState, View, bVisible);
+				UpdatePrimitiveFadingStateHelper(FadingState, View, bVisible);
 				FRHIUniformBuffer* UniformBuffer = FadingState.UniformBuffer;
 				if (UniformBuffer && !bVisible)
 				{
diff --git a/Engine/Source/Runtime/Renderer/Private/ShadowRendering.h b/Engine/Source/Runtime/Renderer/Private/ShadowRendering.h
--- a/Engine/Source/Runtime/Renderer/Private/ShadowRendering.h
+++ b/Engine/Source/Runtime/Renderer/Private/ShadowRendering.h
@@ -446,6 +446,21 @@ public:
 	 */
 	void AddSubjectPrimitive(FPrimitiveSceneInfo* PrimitiveSceneInfo, TArray<FViewInfo>* ViewArray, ERHIFeatureLevel::Type FeatureLevel, bool bRecordShadowSubjectForMobileShading);
 
+	uint64 AddSubjectPrimitive_AnyThread(
+		const FPrimitiveSceneInfoCompact& PrimitiveSceneInfoCompact,
+		TArray<FViewInfo>* ViewArray,
+		ERHIFeatureLevel::Type FeatureLevel,
+		struct FAddSubjectPrimitiveStats& OutStats,
+		struct FAddSubjectPrimitiveOverflowedIndices& OverflowBuffer) const;
+
+	void PresizeSubjectPrimitiveArrays(struct FAddSubjectPrimitiveStats const& Stats);
+
+	void FinalizeAddSubjectPrimitive(
+		struct FAddSubjectPrimitiveOp const& Op,
+		TArray<FViewInfo>* ViewArray,
+		ERHIFeatureLevel::Type FeatureLevel,
+		struct FFinalizeAddSubjectPrimitiveContext& Context);
+
 	/**
 	* @return TRUE if this shadow info has any casting subject prims to render
 	*/
@@ -620,9 +635,27 @@ private:
 		TArray<const FStaticMeshBatch*, SceneRenderingAllocator>& MeshCommandBuildRequests,
 		int32& NumMeshCommandBuildRequestElements);
 
+	void AddCachedMeshDrawCommands_AnyThread(
+		const FScene* Scene,
+		const FStaticMeshBatchRelevance& RESTRICT StaticMeshRelevance,
+		int32 StaticMeshIdx,
+		int32& NumAcceptedStaticMeshes,
+		struct FAddSubjectPrimitiveResult& OutResult,
+		struct FAddSubjectPrimitiveStats& OutStats,
+		struct FAddSubjectPrimitiveOverflowedIndices& OverflowBuffer) const;
+
 	/** Will return if we should draw the static mesh for the shadow, and will perform lazy init of primitive if it wasn't visible */
 	bool ShouldDrawStaticMeshes(FViewInfo& InCurrentView, bool bInCustomDataRelevance, FPrimitiveSceneInfo* InPrimitiveSceneInfo);
 
+	bool ShouldDrawStaticMeshes_AnyThread(
+		FViewInfo& InCurrentView,
+		const FPrimitiveSceneInfoCompact& PrimitiveSceneInfoCompact,
+		bool bMayBeFading,
+		bool bNeedUpdateStaticMeshes,
+		struct FAddSubjectPrimitiveResult& OutResult,
+		struct FAddSubjectPrimitiveStats& OutStats,
+		struct FAddSubjectPrimitiveOverflowedIndices& OverflowBuffer) const;
+
 	void GetShadowTypeNameForDrawEvent(FString& TypeName) const;
 
 	/** Updates object buffers needed by ray traced distance field shadows. */
diff --git a/Engine/Source/Runtime/Renderer/Private/ShadowSetup.cpp b/Engine/Source/Runtime/Renderer/Private/ShadowSetup.cpp
--- a/Engine/Source/Runtime/Renderer/Private/ShadowSetup.cpp
+++ b/Engine/Source/Runtime/Renderer/Private/ShadowSetup.cpp
@@ -868,6 +868,249 @@ void FProjectedShadowInfo::AddCachedMeshDrawCommandsForPass(
 	}
 }
 
+struct FAddSubjectPrimitiveOverflowedIndices
+{
+	TArray<uint16> MDCIndices;
+	TArray<uint16> MeshIndices;
+};
+
+struct FFinalizeAddSubjectPrimitiveContext
+{
+	const uint16* OverflowedMDCIndices;
+	const uint16* OverflowedMeshIndices;
+};
+
+struct FAddSubjectPrimitiveResult
+{
+	union
+	{
+		uint64 Qword;
+		struct
+		{
+			uint32 bCopyCachedMeshDrawCommand : 1;
+			uint32 bRequestMeshCommandBuild : 1;
+			uint32 bOverflowed : 1;
+			uint32 bDynamicSubjectPrimitive : 1;
+			uint32 bTranslucentSubjectPrimitive : 1;
+			uint32 bNeedUniformBufferUpdate : 1;
+			uint32 bNeedUpdateStaticMeshes : 1;
+			uint32 bNeedPrimitiveFadingStateUpdate : 1;
+			uint32 bFadingIn : 1;
+			uint32 bAddOnRenderThread : 1;
+			uint32 bRecordShadowSubjectsForMobile : 1;
+
+			union
+			{
+				uint16 MDCOrMeshIndices[2];
+				struct 
+				{
+					uint16 NumMDCIndices;
+					uint16 NumMeshIndices;
+				};
+			};
+		};
+	};
+
+	void AcceptMDC(int32 NumAcceptedStaticMeshes, int32 MDCIdx, FAddSubjectPrimitiveOverflowedIndices& OverflowBuffer)
+	{
+		check(NumAcceptedStaticMeshes >= 0 && MDCIdx < MAX_uint16);
+		if (NumAcceptedStaticMeshes < 2)
+		{
+			MDCOrMeshIndices[NumAcceptedStaticMeshes] = uint16(MDCIdx + 1);
+			if (bRequestMeshCommandBuild)
+			{
+				const uint16 Tmp = MDCOrMeshIndices[1];
+				MDCOrMeshIndices[1] = MDCOrMeshIndices[0];
+				MDCOrMeshIndices[0] = Tmp;
+			}
+		}
+		else
+		{
+			if (NumAcceptedStaticMeshes == 2)
+			{
+				HandleOverflow(OverflowBuffer);
+			}
+			check(bOverflowed);
+			OverflowBuffer.MDCIndices.Add(MDCIdx);
+			++NumMDCIndices;
+		}
+		bCopyCachedMeshDrawCommand = true;
+	}
+
+	void AcceptMesh(int32 NumAcceptedStaticMeshes, int32 MeshIdx, FAddSubjectPrimitiveOverflowedIndices& OverflowBuffer)
+	{
+		check(NumAcceptedStaticMeshes >= 0 && MeshIdx < MAX_uint16);
+		if (NumAcceptedStaticMeshes < 2)
+		{
+			MDCOrMeshIndices[NumAcceptedStaticMeshes] = uint16(MeshIdx + 1);
+		}
+		else
+		{
+			if (NumAcceptedStaticMeshes == 2)
+			{
+				HandleOverflow(OverflowBuffer);
+			}
+			check(bOverflowed);
+			OverflowBuffer.MeshIndices.Add(MeshIdx);
+			++NumMeshIndices;
+		}
+		bRequestMeshCommandBuild = true;
+	}
+
+	int32 GetMDCIndices(FFinalizeAddSubjectPrimitiveContext& Context, const uint16*& OutMDCIndices, int32& OutIdxBias) const
+	{
+		int32 NumMDCs;
+		OutIdxBias = -1;
+		if (bOverflowed)
+		{
+			OutMDCIndices = Context.OverflowedMDCIndices;
+			NumMDCs = NumMDCIndices;
+			check(NumMDCs > 0);
+			Context.OverflowedMDCIndices += NumMDCs;
+			OutIdxBias = 0;
+		}
+		else
+		{
+			OutMDCIndices = MDCOrMeshIndices;
+			NumMDCs = !MDCOrMeshIndices[1] ? 1 : (!bRequestMeshCommandBuild ? 2 : 1);
+		}
+		return NumMDCs;
+	}
+
+	int32 GetMeshIndices(FFinalizeAddSubjectPrimitiveContext& Context, const uint16*& OutMeshIndices, int32& OutIdxBias) const
+	{
+		int32 NumMeshes;
+		OutIdxBias = -1;
+		if (bOverflowed)
+		{
+			OutMeshIndices = Context.OverflowedMeshIndices;
+			NumMeshes = NumMeshIndices;
+			check(NumMeshes > 0);
+			Context.OverflowedMeshIndices += NumMeshes;
+			OutIdxBias = 0;
+		}
+		else if (!bCopyCachedMeshDrawCommand)
+		{
+			OutMeshIndices = MDCOrMeshIndices;
+			NumMeshes = !MDCOrMeshIndices[1] ? 1 : 2;
+		}
+		else
+		{
+			OutMeshIndices = &MDCOrMeshIndices[1];
+			NumMeshes = 1;
+		}
+		return NumMeshes;
+	}
+
+private:
+	void HandleOverflow(FAddSubjectPrimitiveOverflowedIndices& OverflowBuffer)
+	{
+		if (bCopyCachedMeshDrawCommand && !bRequestMeshCommandBuild)
+		{
+			OverflowBuffer.MDCIndices.Add(MDCOrMeshIndices[0] - 1);
+			OverflowBuffer.MDCIndices.Add(MDCOrMeshIndices[1] - 1);
+			NumMDCIndices = 2;
+			NumMeshIndices = 0;
+		}
+		else if (bCopyCachedMeshDrawCommand)
+		{
+			OverflowBuffer.MDCIndices.Add(MDCOrMeshIndices[0] - 1);
+			OverflowBuffer.MeshIndices.Add(MDCOrMeshIndices[1] - 1);
+			NumMDCIndices = 1;
+			NumMeshIndices = 1;
+		}
+		else
+		{
+			check(bRequestMeshCommandBuild);
+			OverflowBuffer.MeshIndices.Add(MDCOrMeshIndices[0] - 1);
+			OverflowBuffer.MeshIndices.Add(MDCOrMeshIndices[1] - 1);
+			NumMDCIndices = 0;
+			NumMeshIndices = 2;
+		}
+		bOverflowed = true;
+	}
+};
+
+static_assert(sizeof(FAddSubjectPrimitiveResult) == 8, "Unexpected size for FAddSubjectPrimitiveResult");
+
+struct FAddSubjectPrimitiveOp
+{
+	FPrimitiveSceneInfo* PrimitiveSceneInfo;
+	FAddSubjectPrimitiveResult Result;
+};
+
+struct FAddSubjectPrimitiveStats
+{
+	int32 NumCachedMDCCopies;
+	int32 NumMDCBuildRequests;
+	int32 NumDynamicSubs;
+	int32 NumTranslucentSubs;
+	int32 NumDeferredPrimitives;
+
+	FAddSubjectPrimitiveStats()
+		: NumCachedMDCCopies(0)
+		, NumMDCBuildRequests(0)
+		, NumDynamicSubs(0)
+		, NumTranslucentSubs(0)
+		, NumDeferredPrimitives(0)
+	{}
+
+	void InterlockedAdd(const FAddSubjectPrimitiveStats& Other)
+	{
+		if (Other.NumCachedMDCCopies > 0)
+		{
+			FPlatformAtomics::InterlockedAdd(&NumCachedMDCCopies, Other.NumCachedMDCCopies);
+		}
+		if (Other.NumMDCBuildRequests > 0)
+		{
+			FPlatformAtomics::InterlockedAdd(&NumMDCBuildRequests, Other.NumMDCBuildRequests);
+		}
+		if (Other.NumDynamicSubs > 0)
+		{
+			FPlatformAtomics::InterlockedAdd(&NumDynamicSubs, Other.NumDynamicSubs);
+		}
+		if (Other.NumTranslucentSubs > 0)
+		{
+			FPlatformAtomics::InterlockedAdd(&NumTranslucentSubs, Other.NumTranslucentSubs);
+		}
+		if (Other.NumDeferredPrimitives > 0)
+		{
+			FPlatformAtomics::InterlockedAdd(&NumDeferredPrimitives, Other.NumDeferredPrimitives);
+		}
+	}
+};
+
+void FProjectedShadowInfo::AddCachedMeshDrawCommands_AnyThread(
+	const FScene* Scene,
+	const FStaticMeshBatchRelevance& RESTRICT StaticMeshRelevance,
+	int32 StaticMeshIdx,
+	int32& NumAcceptedStaticMeshes,
+	FAddSubjectPrimitiveResult& OutResult,
+	FAddSubjectPrimitiveStats& OutStats,
+	FAddSubjectPrimitiveOverflowedIndices& OverflowBuffer) const
+{
+	const EMeshPass::Type PassType = EMeshPass::CSMShadowDepth;
+	const EShadingPath ShadingPath = Scene->GetShadingPath();
+	const bool bUseCachedMeshCommand = UseCachedMeshDrawCommands()
+		&& !!(FPassProcessorManager::GetPassFlags(ShadingPath, PassType) & EMeshPassFlags::CachedMeshCommands)
+		&& StaticMeshRelevance.bSupportsCachingMeshDrawCommands;
+
+	if (bUseCachedMeshCommand)
+	{
+		const int32 StaticMeshCommandInfoIndex = StaticMeshRelevance.GetStaticMeshCommandInfoIndex(PassType);
+		if (StaticMeshCommandInfoIndex >= 0)
+		{
+			++OutStats.NumCachedMDCCopies;
+			OutResult.AcceptMDC(NumAcceptedStaticMeshes++, StaticMeshCommandInfoIndex, OverflowBuffer);
+		}
+	}
+	else
+	{
+		++OutStats.NumMDCBuildRequests;
+		OutResult.AcceptMesh(NumAcceptedStaticMeshes++, StaticMeshIdx, OverflowBuffer);
+	}
+}
+
 bool FProjectedShadowInfo::ShouldDrawStaticMeshes(FViewInfo& InCurrentView, bool InbCustomDataRelevance, FPrimitiveSceneInfo* InPrimitiveSceneInfo)
 {
 	bool WholeSceneDirectionalShadow = IsWholeSceneDirectionalShadow();
@@ -1017,6 +1260,149 @@ bool FProjectedShadowInfo::ShouldDrawStaticMeshes(FViewInfo& InCurrentView, bool
 	return bDrawingStaticMeshes;
 }
 
+bool FProjectedShadowInfo::ShouldDrawStaticMeshes_AnyThread(
+	FViewInfo& CurrentView,
+	const FPrimitiveSceneInfoCompact& PrimitiveSceneInfoCompact,
+	bool bMayBeFading,
+	bool bNeedUpdateStaticMeshes,
+	FAddSubjectPrimitiveResult& OutResult,
+	FAddSubjectPrimitiveStats& OutStats,
+	FAddSubjectPrimitiveOverflowedIndices& OverflowBuffer) const
+{
+	bool bDrawingStaticMeshes = false;
+	const bool WholeSceneDirectionalShadow = IsWholeSceneDirectionalShadow();
+	const FPrimitiveSceneInfo* PrimitiveSceneInfo = PrimitiveSceneInfoCompact.PrimitiveSceneInfo;
+	const FPrimitiveSceneProxy* Proxy = PrimitiveSceneInfoCompact.Proxy;
+	const int32 PrimitiveId = PrimitiveSceneInfo->GetIndex();
+
+	{
+		const int32 ForcedLOD = CurrentView.Family->EngineShowFlags.LOD ? (GetCVarForceLODShadow() != -1 ? GetCVarForceLODShadow() : GetCVarForceLOD()) : -1;
+		const FLODMask* VisibilePrimitiveLODMask = nullptr;
+
+		if (CurrentView.PrimitivesLODMask[PrimitiveId].ContainsLOD(MAX_int8)) // only calculate it if it's not set
+		{
+			FLODMask ViewLODToRender;
+			float MeshScreenSizeSquared = 0;
+			const int8 CurFirstLODIdx = Proxy->GetCurrentFirstLODIdx_RenderThread();
+
+			if (PrimitiveSceneInfo->bIsUsingCustomLODRules)
+			{
+				ViewLODToRender = Proxy->GetCustomLOD(CurrentView, CurrentView.LODDistanceFactor, ForcedLOD, MeshScreenSizeSquared);
+				ViewLODToRender.ClampToFirstLOD(CurFirstLODIdx);
+			}
+			else
+			{
+				const FBoxSphereBounds& Bounds = PrimitiveSceneInfoCompact.Bounds;
+				ViewLODToRender = ComputeLODForMeshes(PrimitiveSceneInfo->StaticMeshRelevances, CurrentView, Bounds.Origin, Bounds.SphereRadius, ForcedLOD, MeshScreenSizeSquared, CurFirstLODIdx, CurrentView.LODDistanceFactor);
+			}
+
+			CurrentView.PrimitivesLODMask[PrimitiveId] = ViewLODToRender;
+		}
+
+		VisibilePrimitiveLODMask = &CurrentView.PrimitivesLODMask[PrimitiveId];
+		check(VisibilePrimitiveLODMask != nullptr);
+
+		FLODMask ShadowLODToRender = *VisibilePrimitiveLODMask;
+
+		// Use lowest LOD for PreShadow
+		if (bReflectiveShadowmap || (bPreShadow && GPreshadowsForceLowestLOD))
+		{
+			int8 LODToRenderScan = -MAX_int8;
+			FLODMask LODToRender;
+
+			for (int32 Index = 0; Index < PrimitiveSceneInfo->StaticMeshRelevances.Num(); Index++)
+			{
+				LODToRenderScan = FMath::Max<int8>(PrimitiveSceneInfo->StaticMeshRelevances[Index].LODIndex, LODToRenderScan);
+			}
+			if (LODToRenderScan != -MAX_int8)
+			{
+				ShadowLODToRender.SetLOD(LODToRenderScan);
+			}
+		}
+
+		if (CascadeSettings.bFarShadowCascade)
+		{
+			extern ENGINE_API int32 GFarShadowStaticMeshLODBias;
+			int8 LODToRenderScan = ShadowLODToRender.DitheredLODIndices[0] + GFarShadowStaticMeshLODBias;
+
+			for (int32 Index = PrimitiveSceneInfo->StaticMeshRelevances.Num() - 1; Index >= 0; Index--)
+			{
+				if (LODToRenderScan == PrimitiveSceneInfo->StaticMeshRelevances[Index].LODIndex)
+				{
+					ShadowLODToRender.SetLOD(LODToRenderScan);
+					break;
+				}
+			}
+		}
+
+		if (WholeSceneDirectionalShadow)
+		{
+			if (PrimitiveSceneInfo->bIsUsingCustomWholeSceneShadowLODRules)
+			{
+				const bool bHasShelfShadow = !bReflectiveShadowmap && !bPreShadow;
+				const float ShadowMapTextureResolution = BorderSize * 2 + ResolutionX; // We assume Shadow Map texture to be squared (current design)
+				ShadowLODToRender = Proxy->GetCustomWholeSceneShadowLOD(CurrentView, CurrentView.LODDistanceFactor, ForcedLOD, *VisibilePrimitiveLODMask, ShadowMapTextureResolution, ShadowBounds.W * 2.0f, ShadowId, bHasShelfShadow);
+			}
+
+			// Don't cache if it requires per view per mesh state for distance cull fade.
+			const bool bCanCache = !bMayBeFading && !bNeedUpdateStaticMeshes;
+			int32 NumAcceptedStaticMeshes = 0;
+
+			for (int32 MeshIndex = 0; MeshIndex < PrimitiveSceneInfo->StaticMeshRelevances.Num(); MeshIndex++)
+			{
+				const FStaticMeshBatchRelevance& StaticMeshRelevance = PrimitiveSceneInfo->StaticMeshRelevances[MeshIndex];
+				const FStaticMeshBatch& StaticMesh = PrimitiveSceneInfo->StaticMeshes[MeshIndex];
+
+				if ((StaticMeshRelevance.CastShadow || (bSelfShadowOnly && StaticMeshRelevance.bUseForDepthPass)) && ShadowLODToRender.ContainsLOD(StaticMeshRelevance.LODIndex))
+				{
+					if (bCanCache && GetShadowDepthType() == CSMShadowDepthType)
+					{
+						AddCachedMeshDrawCommands_AnyThread(PrimitiveSceneInfo->Scene, StaticMeshRelevance, MeshIndex, NumAcceptedStaticMeshes, OutResult, OutStats, OverflowBuffer);
+					}
+					else
+					{
+						++OutStats.NumMDCBuildRequests;
+						OutResult.AcceptMesh(NumAcceptedStaticMeshes++, MeshIndex, OverflowBuffer);
+					}
+
+					if (StaticMeshRelevance.bRequiresPerElementVisibility)
+					{
+						CurrentView.StaticMeshBatchVisibility[StaticMesh.BatchVisibilityId] = StaticMesh.VertexFactory->GetStaticBatchElementVisibility(CurrentView, &StaticMesh);
+					}
+
+					bDrawingStaticMeshes = true;
+				}
+			}
+		}
+		else
+		{
+			int32 NumAcceptedStaticMeshes = 0;
+
+			for (int32 MeshIndex = 0; MeshIndex < PrimitiveSceneInfo->StaticMeshRelevances.Num(); MeshIndex++)
+			{
+				const FStaticMeshBatchRelevance& StaticMeshRelevance = PrimitiveSceneInfo->StaticMeshRelevances[MeshIndex];
+				const FStaticMeshBatch& StaticMesh = PrimitiveSceneInfo->StaticMeshes[MeshIndex];
+
+				if ((StaticMeshRelevance.CastShadow || (bSelfShadowOnly && StaticMeshRelevance.bUseForDepthPass)) && ShadowLODToRender.ContainsLOD(StaticMeshRelevance.LODIndex))
+				{
+					check(MeshIndex < MAX_uint16);
+					++OutStats.NumMDCBuildRequests;
+					OutResult.AcceptMesh(NumAcceptedStaticMeshes++, MeshIndex, OverflowBuffer);
+
+					if (StaticMeshRelevance.bRequiresPerElementVisibility)
+					{
+						CurrentView.StaticMeshBatchVisibility[StaticMesh.BatchVisibilityId] = StaticMesh.VertexFactory->GetStaticBatchElementVisibility(CurrentView, &StaticMesh);
+					}
+
+					bDrawingStaticMeshes = true;
+				}
+			}
+		}
+	}
+
+	return bDrawingStaticMeshes;
+}
+
 void FProjectedShadowInfo::AddSubjectPrimitive(FPrimitiveSceneInfo* PrimitiveSceneInfo, TArray<FViewInfo>* ViewArray, ERHIFeatureLevel::Type FeatureLevel, bool bRecordShadowSubjectsForMobileShading)
 {
 	// Ray traced shadows use the GPU managed distance field object buffers, no CPU culling should be used
@@ -1192,6 +1578,315 @@ void FProjectedShadowInfo::AddSubjectPrimitive(FPrimitiveSceneInfo* PrimitiveSce
 	}
 }
 
+uint64 FProjectedShadowInfo::AddSubjectPrimitive_AnyThread(
+	const FPrimitiveSceneInfoCompact& PrimitiveSceneInfoCompact,
+	TArray<FViewInfo>* ViewArray,
+	ERHIFeatureLevel::Type FeatureLevel,
+	FAddSubjectPrimitiveStats& OutStats,
+	FAddSubjectPrimitiveOverflowedIndices& OverflowBuffer) const
+{
+	SCOPE_CYCLE_COUNTER(STAT_AddSubjectPrimitive);
+
+	// Ray traced shadows use the GPU managed distance field object buffers, no CPU culling should be used
+	check(!bRayTracedDistanceField);
+
+	FAddSubjectPrimitiveResult Result;
+	Result.Qword = 0;
+
+	if (FSceneInterface::GetShadingPath(FeatureLevel) == EShadingPath::Mobile)
+	{
+		// record shadow casters if CSM culling is enabled for the light's mobility type and the culling mode requires the list of casters.
+		static auto* CVarMobileEnableStaticAndCSMShadowReceivers = IConsoleManager::Get().FindTConsoleVariableDataInt(TEXT("r.Mobile.EnableStaticAndCSMShadowReceivers"));
+		static auto* CVarMobileEnableMovableLightCSMShaderCulling = IConsoleManager::Get().FindTConsoleVariableDataInt(TEXT("r.Mobile.EnableMovableLightCSMShaderCulling"));
+		static auto* CVarMobileCSMShaderCullingMethod = IConsoleManager::Get().FindTConsoleVariableDataInt(TEXT("r.Mobile.Shadow.CSMShaderCullingMethod"));
+		const uint32 MobileCSMCullingMode = CVarMobileCSMShaderCullingMethod->GetValueOnRenderThread() & 0xF;
+		const bool bRecordShadowSubjectsForMobile =
+			(MobileCSMCullingMode == 2 || MobileCSMCullingMode == 3)
+			&& ((CVarMobileEnableMovableLightCSMShaderCulling->GetValueOnRenderThread() && GetLightSceneInfo().Proxy->IsMovable() && GetLightSceneInfo().ShouldRenderViewIndependentWholeSceneShadows())
+				|| (CVarMobileEnableStaticAndCSMShadowReceivers->GetValueOnRenderThread() && GetLightSceneInfo().Proxy->UseCSMForDynamicObjects()));
+
+		if (bRecordShadowSubjectsForMobile)
+		{
+			Result.bAddOnRenderThread = true;
+			Result.bRecordShadowSubjectsForMobile = true;
+			++OutStats.NumDeferredPrimitives;
+			return Result.Qword;
+		}
+	}
+
+	FPrimitiveSceneInfo* PrimitiveSceneInfo = PrimitiveSceneInfoCompact.PrimitiveSceneInfo;
+
+	if (!ReceiverPrimitives.Contains(PrimitiveSceneInfo)
+		// Far cascade only casts from primitives marked for it
+		&& (!CascadeSettings.bFarShadowCascade || PrimitiveSceneInfoCompact.Proxy->CastsFarShadow()))
+	{
+		FViewInfo* CurrentView;
+		const bool bWholeSceneDirectionalShadow = IsWholeSceneDirectionalShadow();
+
+		if (bWholeSceneDirectionalShadow)
+		{
+			CurrentView = DependentView;
+		}
+		else
+		{
+			checkf(ViewArray,
+				TEXT("bWholeSceneShadow=%d, CascadeSettings.ShadowSplitIndex=%d, bDirectionalLight=%s"),
+				bWholeSceneShadow ? TEXT("true") : TEXT("false"),
+				CascadeSettings.ShadowSplitIndex,
+				bDirectionalLight ? TEXT("true") : TEXT("false"));
+
+			if (ViewArray->Num() > 1)
+			{
+				Result.bAddOnRenderThread = true;
+				++OutStats.NumDeferredPrimitives;
+				return Result.Qword;
+			}
+
+			CurrentView = &(*ViewArray)[0];
+		}
+
+		bool bOpaque = false;
+		bool bTranslucentRelevance = false;
+		bool bShadowRelevance = false;
+		bool bStaticRelevance = false;
+		bool bCustomDataRelevance = false;
+		bool bMayBeFading = false;
+		bool bNeedUpdateStaticMeshes = false;
+
+		const int32 PrimitiveId = PrimitiveSceneInfo->GetIndex();
+		FPrimitiveViewRelevance& ViewRelevance = CurrentView->PrimitiveViewRelevanceMap[PrimitiveId];
+
+		if (!ViewRelevance.bInitializedThisFrame)
+		{
+			if (CurrentView->IsPerspectiveProjection())
+			{
+				bool bFadingIn;
+				// Compute the distance between the view and the primitive.
+				const float DistanceSquared = (PrimitiveSceneInfoCompact.Bounds.Origin - CurrentView->ShadowViewMatrices.GetViewOrigin()).SizeSquared();
+
+				if (CurrentView->IsDistanceCulled_AnyThread(
+					DistanceSquared,
+					PrimitiveSceneInfoCompact.MinDrawDistance,
+					PrimitiveSceneInfoCompact.MaxDrawDistance,
+					PrimitiveSceneInfo,
+					bMayBeFading,
+					bFadingIn))
+				{
+					return 0;
+				}
+
+				if (bMayBeFading)
+				{
+					Result.bNeedPrimitiveFadingStateUpdate = true;
+					Result.bFadingIn = bFadingIn;
+				}
+			}
+
+			// Respect HLOD visibility which can hide child LOD primitives
+			if (CurrentView->ViewState &&
+				CurrentView->ViewState->HLODVisibilityState.IsValidPrimitiveIndex(PrimitiveId) &&
+				CurrentView->ViewState->HLODVisibilityState.IsNodeForcedHidden(PrimitiveId))
+			{
+				return 0;
+			}
+
+			if ((CurrentView->ShowOnlyPrimitives.IsSet() &&
+				!CurrentView->ShowOnlyPrimitives->Contains(PrimitiveSceneInfoCompact.Proxy->GetPrimitiveComponentId())) ||
+				CurrentView->HiddenPrimitives.Contains(PrimitiveSceneInfoCompact.Proxy->GetPrimitiveComponentId()))
+			{
+				return 0;
+			}
+
+			// Compute the subject primitive's view relevance since it wasn't cached
+			// Update the main view's PrimitiveViewRelevanceMap
+			ViewRelevance = PrimitiveSceneInfoCompact.Proxy->GetViewRelevance(CurrentView);
+		}
+
+		bOpaque = ViewRelevance.bOpaque || ViewRelevance.bMasked;
+		bTranslucentRelevance = ViewRelevance.HasTranslucency() && !ViewRelevance.bMasked;
+		bShadowRelevance = ViewRelevance.bShadowRelevance;
+		bStaticRelevance = ViewRelevance.bStaticRelevance;
+		bCustomDataRelevance = ViewRelevance.bUseCustomViewData;
+
+		if (!bShadowRelevance)
+		{
+			return 0;
+		}
+
+		if (bCustomDataRelevance)
+		{
+			// The global custom data allocator is not thread-safe
+			Result.bAddOnRenderThread = true;
+			++OutStats.NumDeferredPrimitives;
+			return Result.Qword;
+		}
+
+		// Update the primitive component's last render time. Allows the component to update when using bCastWhenHidden.
+		const float CurrentWorldTime = CurrentView->Family->CurrentWorldTime;
+		PrimitiveSceneInfo->UpdateComponentLastRenderTime(CurrentWorldTime, /*bUpdateLastRenderTimeOnScreen=*/false);
+
+		if (PrimitiveSceneInfo->NeedsUniformBufferUpdate())
+		{
+			// Main view visible primitives are processed on parallel tasks, updating uniform buffer them here will cause a race condition.
+			check(!CurrentView->PrimitiveVisibilityMap[PrimitiveId]);
+			Result.bNeedUniformBufferUpdate = true;
+		}
+
+		if (PrimitiveSceneInfo->NeedsUpdateStaticMeshes())
+		{
+			// Need to defer to next InitViews, as main view visible primitives are processed on parallel tasks and calling 
+			// CacheMeshDrawCommands may resize CachedDrawLists/CachedMeshDrawCommandStateBuckets causing a crash.
+			Result.bNeedUpdateStaticMeshes = true;
+			bNeedUpdateStaticMeshes = true;
+		}
+
+		if (bOpaque)
+		{
+			bool bDrawingStaticMeshes = false;
+
+			if (PrimitiveSceneInfo->StaticMeshes.Num() > 0)
+			{
+				if (bWholeSceneShadow)
+				{
+					const FBoxSphereBounds& Bounds = PrimitiveSceneInfoCompact.Bounds;
+					const float DistanceSquared = (Bounds.Origin - CurrentView->ShadowViewMatrices.GetViewOrigin()).SizeSquared();
+					const bool bDrawShadowDepth = FMath::Square(Bounds.SphereRadius) > FMath::Square(GMinScreenRadiusForShadowCaster) * DistanceSquared * CurrentView->LODDistanceFactorSquared;
+					if (!bDrawShadowDepth)
+					{
+						// cull object if it's too small to be considered as shadow caster
+						return 0;
+					}
+				}
+
+				// Update visibility for meshes which weren't visible in the main views or were visible with static relevance
+				if (bStaticRelevance || !CurrentView->PrimitiveVisibilityMap[PrimitiveId])
+				{
+					bDrawingStaticMeshes |= ShouldDrawStaticMeshes_AnyThread(
+						*CurrentView,
+						PrimitiveSceneInfoCompact,
+						bMayBeFading,
+						bNeedUpdateStaticMeshes,
+						Result,
+						OutStats,
+						OverflowBuffer);
+				}
+			}
+
+			if (!bDrawingStaticMeshes)
+			{
+				Result.bDynamicSubjectPrimitive = true;
+				++OutStats.NumDynamicSubs;
+			}
+		}
+
+		if (bTranslucentRelevance)
+		{
+			Result.bTranslucentSubjectPrimitive = true;
+			++OutStats.NumTranslucentSubs;
+		}
+	}
+
+	return Result.Qword;
+}
+
+void FProjectedShadowInfo::PresizeSubjectPrimitiveArrays(const FAddSubjectPrimitiveStats& Stats)
+{
+	ShadowDepthPassVisibleCommands.Reserve(ShadowDepthPassVisibleCommands.Num() + Stats.NumDeferredPrimitives * 2 + Stats.NumCachedMDCCopies);
+	SubjectMeshCommandBuildRequests.Reserve(SubjectMeshCommandBuildRequests.Num() + Stats.NumMDCBuildRequests);
+	DynamicSubjectPrimitives.Reserve(DynamicSubjectPrimitives.Num() + Stats.NumDeferredPrimitives + Stats.NumDynamicSubs);
+	SubjectTranslucentPrimitives.Reserve(SubjectTranslucentPrimitives.Num() + Stats.NumTranslucentSubs);
+}
+
+void FProjectedShadowInfo::FinalizeAddSubjectPrimitive(
+	const FAddSubjectPrimitiveOp& Op,
+	TArray<FViewInfo>* ViewArray,
+	ERHIFeatureLevel::Type FeatureLevel,
+	FFinalizeAddSubjectPrimitiveContext& Context)
+{
+	FPrimitiveSceneInfo* PrimitiveSceneInfo = Op.PrimitiveSceneInfo;
+	const FAddSubjectPrimitiveResult& Result = Op.Result;
+
+	if (Result.bAddOnRenderThread)
+	{
+		AddSubjectPrimitive(PrimitiveSceneInfo, ViewArray, FeatureLevel, Result.bRecordShadowSubjectsForMobile);
+		return;
+	}
+
+	if (Result.bNeedPrimitiveFadingStateUpdate)
+	{
+		FViewInfo& View = IsWholeSceneDirectionalShadow() ? *DependentView : (*ViewArray)[0];
+		if (View.UpdatePrimitiveFadingState(PrimitiveSceneInfo, Result.bFadingIn))
+		{
+			if (Result.bOverflowed)
+			{
+				Context.OverflowedMDCIndices += Result.NumMDCIndices;
+				Context.OverflowedMeshIndices += Result.NumMeshIndices;
+			}
+			return;
+		}
+	}
+
+	if (Result.bCopyCachedMeshDrawCommand)
+	{
+		check(!Result.bDynamicSubjectPrimitive);
+		const uint16* MDCIndices;
+		int32 IdxBias;
+		int32 NumMDCs = Result.GetMDCIndices(Context, MDCIndices, IdxBias);
+
+		for (int32 Idx = 0; Idx < NumMDCs; ++Idx)
+		{
+			const int32 CmdIdx = (int32)MDCIndices[Idx] + IdxBias;
+			const FCachedMeshDrawCommandInfo& CmdInfo = PrimitiveSceneInfo->StaticMeshCommandInfos[CmdIdx];
+			const FScene* Scene = PrimitiveSceneInfo->Scene;
+			const FMeshDrawCommand* CachedCmd = CmdInfo.StateBucketId >= 0 ?
+				&Scene->CachedMeshDrawCommandStateBuckets[FSetElementId::FromInteger(CmdInfo.StateBucketId)].MeshDrawCommand :
+				&Scene->CachedDrawLists[EMeshPass::CSMShadowDepth].MeshDrawCommands[CmdInfo.CommandIndex];
+			const int32 PrimIdx = PrimitiveSceneInfo->GetIndex();
+
+			FVisibleMeshDrawCommand& VisibleCmd = ShadowDepthPassVisibleCommands[ShadowDepthPassVisibleCommands.AddUninitialized()];
+			VisibleCmd.Setup(CachedCmd, PrimIdx, PrimIdx, CmdInfo.StateBucketId, CmdInfo.MeshFillMode, CmdInfo.MeshCullMode, CmdInfo.SortKey);
+		}
+	}
+
+	if (Result.bRequestMeshCommandBuild)
+	{
+		check(!Result.bDynamicSubjectPrimitive);
+		const uint16* MeshIndices;
+		int32 IdxBias;
+		int32 NumMeshes = Result.GetMeshIndices(Context, MeshIndices, IdxBias);
+
+		for (int32 Idx = 0; Idx < NumMeshes; ++Idx)
+		{
+			const int32 MeshIdx = (int32)MeshIndices[Idx] + IdxBias;
+			const FStaticMeshBatchRelevance& MeshRelevance = PrimitiveSceneInfo->StaticMeshRelevances[MeshIdx];
+			const FStaticMeshBatch& MeshBatch = PrimitiveSceneInfo->StaticMeshes[MeshIdx];
+
+			NumSubjectMeshCommandBuildRequestElements += MeshRelevance.NumElements;
+			SubjectMeshCommandBuildRequests.Add(&MeshBatch);
+		}
+	}
+
+	if (Result.bDynamicSubjectPrimitive)
+	{
+		DynamicSubjectPrimitives.Add(PrimitiveSceneInfo);
+	}
+
+	if (Result.bTranslucentSubjectPrimitive)
+	{
+		SubjectTranslucentPrimitives.Add(PrimitiveSceneInfo);
+	}
+
+	if (Result.bNeedUniformBufferUpdate)
+	{
+		PrimitiveSceneInfo->ConditionalUpdateUniformBuffer(FRHICommandListExecutor::GetImmediateCommandList());
+	}
+
+	if (Result.bNeedUpdateStaticMeshes)
+	{
+		PrimitiveSceneInfo->BeginDeferredUpdateStaticMeshesWithoutVisibilityCheck();
+	}
+}
+
 bool FProjectedShadowInfo::HasSubjectPrims() const
 {
 	return DynamicSubjectPrimitives.Num() > 0
@@ -2983,7 +3678,9 @@ void FSceneRenderer::GatherShadowDynamicMeshElements(FGlobalDynamicIndexBuffer&
 	}
 }
 
-typedef TArray<FPrimitiveSceneInfo*> FShadowSubjectPrimitives;
+typedef TArray<FAddSubjectPrimitiveOp> FShadowSubjectPrimitives;
+typedef TArray<FAddSubjectPrimitiveStats, TInlineAllocator<4, SceneRenderingAllocator>> FPerShadowGatherStats;
+typedef TArray<FAddSubjectPrimitiveOverflowedIndices, SceneRenderingAllocator> FPerShadowOverflowedIndices;
 
 struct FGatherShadowPrimitivesPacket
 {
@@ -2998,10 +3695,16 @@ struct FGatherShadowPrimitivesPacket
 	ERHIFeatureLevel::Type FeatureLevel;
 	bool bStaticSceneOnly;
 
-	// Outputs
+	// Scratch
+	FPerShadowGatherStats ViewDependentWholeSceneShadowStats;
+	FPerShadowOverflowedIndices PreShadowOverflowedIndices;
+	FPerShadowOverflowedIndices ViewDependentWholeSceneShadowOverflowedIndices;
 	TArray<FShadowSubjectPrimitives, SceneRenderingAllocator> PreShadowSubjectPrimitives;
 	TArray<FShadowSubjectPrimitives, SceneRenderingAllocator> ViewDependentWholeSceneShadowSubjectPrimitives;
 
+	// Outputs
+	FPerShadowGatherStats& GlobalStats;
+
 	FGatherShadowPrimitivesPacket(
 		const FScene* InScene,
 		TArray<FViewInfo>& InViews,
@@ -3011,7 +3714,8 @@ struct FGatherShadowPrimitivesPacket
 		const TArray<FProjectedShadowInfo*, SceneRenderingAllocator>& InPreShadows,
 		const TArray<FProjectedShadowInfo*, SceneRenderingAllocator>& InViewDependentWholeSceneShadows,
 		ERHIFeatureLevel::Type InFeatureLevel,
-		bool bInStaticSceneOnly) 
+		bool bInStaticSceneOnly,
+		FPerShadowGatherStats& OutGlobalStats) 
 		: Scene(InScene)
 		, Views(InViews)
 		, Node(InNode)
@@ -3021,12 +3725,26 @@ struct FGatherShadowPrimitivesPacket
 		, ViewDependentWholeSceneShadows(InViewDependentWholeSceneShadows)
 		, FeatureLevel(InFeatureLevel)
 		, bStaticSceneOnly(bInStaticSceneOnly)
+		, GlobalStats(OutGlobalStats)
 	{
-		PreShadowSubjectPrimitives.Empty(PreShadows.Num());
-		PreShadowSubjectPrimitives.AddDefaulted(PreShadows.Num());
+		const int32 NumPreShadows = PreShadows.Num();
+		const int32 NumVDWSShadows = ViewDependentWholeSceneShadows.Num();
+
+		check(GlobalStats.Num() == NumVDWSShadows);
+		ViewDependentWholeSceneShadowStats.Empty(NumVDWSShadows);
+		ViewDependentWholeSceneShadowStats.AddDefaulted(NumVDWSShadows);
+
+		PreShadowOverflowedIndices.Empty(NumPreShadows);
+		PreShadowOverflowedIndices.AddDefaulted(NumPreShadows);
 
-		ViewDependentWholeSceneShadowSubjectPrimitives.Empty(ViewDependentWholeSceneShadows.Num());
-		ViewDependentWholeSceneShadowSubjectPrimitives.AddDefaulted(ViewDependentWholeSceneShadows.Num());
+		ViewDependentWholeSceneShadowOverflowedIndices.Empty(NumVDWSShadows);
+		ViewDependentWholeSceneShadowOverflowedIndices.AddDefaulted(NumVDWSShadows);
+
+		PreShadowSubjectPrimitives.Empty(NumPreShadows);
+		PreShadowSubjectPrimitives.AddDefaulted(NumPreShadows);
+
+		ViewDependentWholeSceneShadowSubjectPrimitives.Empty(NumVDWSShadows);
+		ViewDependentWholeSceneShadowSubjectPrimitives.AddDefaulted(NumVDWSShadows);
 	}
 
 	void AnyThreadTask()
@@ -3040,7 +3758,7 @@ struct FGatherShadowPrimitivesPacket
 			{
 				if (NodePrimitiveIt->PrimitiveFlagsCompact.bCastDynamicShadow)
 				{
-					FilterPrimitiveForShadows(NodePrimitiveIt->Bounds, NodePrimitiveIt->PrimitiveFlagsCompact, NodePrimitiveIt->PrimitiveSceneInfo, NodePrimitiveIt->Proxy);
+					FilterPrimitiveForShadows(*NodePrimitiveIt);
 				}
 			}
 		}
@@ -3051,18 +3769,32 @@ struct FGatherShadowPrimitivesPacket
 			// Check primitives in this packet's range
 			for (int32 PrimitiveIndex = StartPrimitiveIndex; PrimitiveIndex < StartPrimitiveIndex + NumPrimitives; PrimitiveIndex++)
 			{
-				FPrimitiveFlagsCompact PrimitiveFlagsCompact = Scene->PrimitiveFlagsCompact[PrimitiveIndex];
+				const FPrimitiveFlagsCompact PrimitiveFlagsCompact = Scene->PrimitiveFlagsCompact[PrimitiveIndex];
 
 				if (PrimitiveFlagsCompact.bCastDynamicShadow)
 				{
-					FilterPrimitiveForShadows(Scene->PrimitiveBounds[PrimitiveIndex].BoxSphereBounds, PrimitiveFlagsCompact, Scene->Primitives[PrimitiveIndex], Scene->PrimitiveSceneProxies[PrimitiveIndex]);
+					FPrimitiveSceneInfo* PrimitiveSceneInfo = Scene->Primitives[PrimitiveIndex];
+					const FPrimitiveSceneInfoCompact PrimitiveSceneInfoCompact(PrimitiveSceneInfo);
+
+					FilterPrimitiveForShadows(PrimitiveSceneInfoCompact);
 				}
 			}
 		}
+
+		const int32 NumStatsToMerge = ViewDependentWholeSceneShadowStats.Num();
+		for (int32 StatIdx = 0; StatIdx < NumStatsToMerge; ++StatIdx)
+		{
+			GlobalStats[StatIdx].InterlockedAdd(ViewDependentWholeSceneShadowStats[StatIdx]);
+		}
 	}
 
-	void FilterPrimitiveForShadows(const FBoxSphereBounds& PrimitiveBounds, FPrimitiveFlagsCompact PrimitiveFlagsCompact, FPrimitiveSceneInfo* PrimitiveSceneInfo, FPrimitiveSceneProxy* PrimitiveProxy)
+	void FilterPrimitiveForShadows(const FPrimitiveSceneInfoCompact& PrimitiveSceneInfoCompact)
 	{
+		const FPrimitiveFlagsCompact& PrimitiveFlagsCompact = PrimitiveSceneInfoCompact.PrimitiveFlagsCompact;
+		const FBoxSphereBounds& PrimitiveBounds = PrimitiveSceneInfoCompact.Bounds;
+		FPrimitiveSceneInfo* PrimitiveSceneInfo = PrimitiveSceneInfoCompact.PrimitiveSceneInfo;
+		const FPrimitiveSceneProxy* PrimitiveProxy = PrimitiveSceneInfoCompact.Proxy;
+
 		// Check if the primitive is a subject for any of the preshadows.
 		// Only allow preshadows from lightmapped primitives that cast both dynamic and static shadows.
 		if (PreShadows.Num() && PrimitiveFlagsCompact.bCastStaticShadow && PrimitiveFlagsCompact.bStaticLighting)
@@ -3077,7 +3809,17 @@ struct FGatherShadowPrimitivesPacket
 
 				if (bInFrustum && ProjectedShadowInfo->GetLightSceneInfoCompact().AffectsPrimitive(PrimitiveBounds, PrimitiveProxy))
 				{
-					PreShadowSubjectPrimitives[ShadowIndex].Add(PrimitiveSceneInfo);
+					FAddSubjectPrimitiveResult Result;
+					FAddSubjectPrimitiveStats UnusedStats;
+					Result.Qword = ProjectedShadowInfo->AddSubjectPrimitive_AnyThread(PrimitiveSceneInfoCompact, &Views, FeatureLevel, UnusedStats, PreShadowOverflowedIndices[ShadowIndex]);
+
+					if (!!Result.Qword)
+					{
+						FShadowSubjectPrimitives& SubjectPrimitives = PreShadowSubjectPrimitives[ShadowIndex];
+						FAddSubjectPrimitiveOp& Op = SubjectPrimitives[SubjectPrimitives.AddUninitialized()];
+						Op.PrimitiveSceneInfo = PrimitiveSceneInfo;
+						Op.Result.Qword = Result.Qword;
+					}
 				}
 			}
 		}
@@ -3141,7 +3883,25 @@ struct FGatherShadowPrimitivesPacket
 					// Render dynamic lit objects if CSMForDynamicObjects is enabled.
 					&& (!LightProxy.UseCSMForDynamicObjects() || !PrimitiveProxy->HasStaticLighting()))
 				{
-					ViewDependentWholeSceneShadowSubjectPrimitives[ShadowIndex].Add(PrimitiveSceneInfo);
+					FAddSubjectPrimitiveResult Result;
+					Result.Qword = ProjectedShadowInfo->AddSubjectPrimitive_AnyThread(
+						PrimitiveSceneInfoCompact,
+						nullptr,
+						FeatureLevel,
+						ViewDependentWholeSceneShadowStats[ShadowIndex],
+						ViewDependentWholeSceneShadowOverflowedIndices[ShadowIndex]);
+
+					if (!!Result.Qword)
+					{
+						FShadowSubjectPrimitives& SubjectPrimitives = ViewDependentWholeSceneShadowSubjectPrimitives[ShadowIndex];
+						if (!SubjectPrimitives.Num())
+						{
+							SubjectPrimitives.Reserve(16);
+						}
+						FAddSubjectPrimitiveOp& Op = SubjectPrimitives[SubjectPrimitives.AddUninitialized()];
+						Op.PrimitiveSceneInfo = PrimitiveSceneInfo;
+						Op.Result.Qword = Result.Qword;
+					}
 				}
 			}
 		}
@@ -3152,35 +3912,30 @@ struct FGatherShadowPrimitivesPacket
 		for (int32 ShadowIndex = 0; ShadowIndex < PreShadowSubjectPrimitives.Num(); ShadowIndex++)
 		{
 			FProjectedShadowInfo* ProjectedShadowInfo = PreShadows[ShadowIndex];
+			const FShadowSubjectPrimitives& SubjectPrimitives = PreShadowSubjectPrimitives[ShadowIndex];
+			const FAddSubjectPrimitiveOverflowedIndices& OverflowBuffer = PreShadowOverflowedIndices[ShadowIndex];
+			FFinalizeAddSubjectPrimitiveContext Context;
+			Context.OverflowedMDCIndices = OverflowBuffer.MDCIndices.GetData();
+			Context.OverflowedMeshIndices = OverflowBuffer.MeshIndices.GetData();
 
-			for (int32 PrimitiveIndex = 0; PrimitiveIndex < PreShadowSubjectPrimitives[ShadowIndex].Num(); PrimitiveIndex++)
+			for (int32 PrimitiveIndex = 0; PrimitiveIndex < SubjectPrimitives.Num(); PrimitiveIndex++)
 			{
-				ProjectedShadowInfo->AddSubjectPrimitive(PreShadowSubjectPrimitives[ShadowIndex][PrimitiveIndex], &Views, FeatureLevel, false);
+				ProjectedShadowInfo->FinalizeAddSubjectPrimitive(SubjectPrimitives[PrimitiveIndex], &Views, FeatureLevel, Context);
 			}
 		}
 
 		for (int32 ShadowIndex = 0; ShadowIndex < ViewDependentWholeSceneShadowSubjectPrimitives.Num(); ShadowIndex++)
 		{
 			FProjectedShadowInfo* ProjectedShadowInfo = ViewDependentWholeSceneShadows[ShadowIndex];
+			const FShadowSubjectPrimitives& SubjectPrimitives = ViewDependentWholeSceneShadowSubjectPrimitives[ShadowIndex];
+			const FAddSubjectPrimitiveOverflowedIndices& OverflowBuffer = ViewDependentWholeSceneShadowOverflowedIndices[ShadowIndex];
+			FFinalizeAddSubjectPrimitiveContext Context;
+			Context.OverflowedMDCIndices = OverflowBuffer.MDCIndices.GetData();
+			Context.OverflowedMeshIndices = OverflowBuffer.MeshIndices.GetData();
 
-			bool bRecordShadowSubjectsForMobile = false;
-							
-			if (FSceneInterface::GetShadingPath(FeatureLevel) == EShadingPath::Mobile)
-			{
-				// record shadow casters if CSM culling is enabled for the light's mobility type and the culling mode requires the list of casters.
-				static auto* CVarMobileEnableStaticAndCSMShadowReceivers = IConsoleManager::Get().FindTConsoleVariableDataInt(TEXT("r.Mobile.EnableStaticAndCSMShadowReceivers"));
-				static auto* CVarMobileEnableMovableLightCSMShaderCulling = IConsoleManager::Get().FindTConsoleVariableDataInt(TEXT("r.Mobile.EnableMovableLightCSMShaderCulling"));
-				static auto* CVarMobileCSMShaderCullingMethod = IConsoleManager::Get().FindTConsoleVariableDataInt(TEXT("r.Mobile.Shadow.CSMShaderCullingMethod"));
-				uint32 MobileCSMCullingMode = CVarMobileCSMShaderCullingMethod->GetValueOnRenderThread()  & 0xF;
-				bRecordShadowSubjectsForMobile = 
-					(MobileCSMCullingMode == 2 || MobileCSMCullingMode == 3)
-					&& ((CVarMobileEnableMovableLightCSMShaderCulling->GetValueOnRenderThread() && ProjectedShadowInfo->GetLightSceneInfo().Proxy->IsMovable() && ProjectedShadowInfo->GetLightSceneInfo().ShouldRenderViewIndependentWholeSceneShadows())
-						|| (CVarMobileEnableStaticAndCSMShadowReceivers->GetValueOnRenderThread() && ProjectedShadowInfo->GetLightSceneInfo().Proxy->UseCSMForDynamicObjects()));
-			}
-
-			for (int32 PrimitiveIndex = 0; PrimitiveIndex < ViewDependentWholeSceneShadowSubjectPrimitives[ShadowIndex].Num(); PrimitiveIndex++)
+			for (int32 PrimitiveIndex = 0; PrimitiveIndex < SubjectPrimitives.Num(); PrimitiveIndex++)
 			{
-				ProjectedShadowInfo->AddSubjectPrimitive(ViewDependentWholeSceneShadowSubjectPrimitives[ShadowIndex][PrimitiveIndex], NULL, FeatureLevel, bRecordShadowSubjectsForMobile);
+				ProjectedShadowInfo->FinalizeAddSubjectPrimitive(SubjectPrimitives[PrimitiveIndex], nullptr, FeatureLevel, Context);
 			}
 		}
 	}
@@ -3197,6 +3952,9 @@ void FSceneRenderer::GatherShadowPrimitives(
 	if (PreShadows.Num() || ViewDependentWholeSceneShadows.Num())
 	{
 		TArray<FGatherShadowPrimitivesPacket*,SceneRenderingAllocator> Packets;
+		FPerShadowGatherStats GatherStats;
+
+		GatherStats.AddDefaulted(ViewDependentWholeSceneShadows.Num());
 
 		if (GUseOctreeForShadowCulling)
 		{
@@ -3273,7 +4031,17 @@ void FSceneRenderer::GatherShadowPrimitives(
 
 				if (PrimitiveOctreeNode.GetElementCount() > 0)
 				{
-					FGatherShadowPrimitivesPacket* Packet = new(FMemStack::Get()) FGatherShadowPrimitivesPacket(Scene, Views, &PrimitiveOctreeNode, 0, 0, PreShadows, ViewDependentWholeSceneShadows, FeatureLevel, bStaticSceneOnly);
+					FGatherShadowPrimitivesPacket* Packet = new(FMemStack::Get()) FGatherShadowPrimitivesPacket(
+						Scene,
+						Views,
+						&PrimitiveOctreeNode,
+						0,
+						0,
+						PreShadows,
+						ViewDependentWholeSceneShadows,
+						FeatureLevel,
+						bStaticSceneOnly,
+						GatherStats);
 					Packets.Add(Packet);
 				}
 			}
@@ -3289,7 +4057,17 @@ void FSceneRenderer::GatherShadowPrimitives(
 			{
 				const int32 StartPrimitiveIndex = PacketIndex * PacketSize;
 				const int32 NumPrimitives = FMath::Min(PacketSize, Scene->Primitives.Num() - StartPrimitiveIndex);
-				FGatherShadowPrimitivesPacket* Packet = new(FMemStack::Get()) FGatherShadowPrimitivesPacket(Scene, Views, NULL, StartPrimitiveIndex, NumPrimitives, PreShadows, ViewDependentWholeSceneShadows, FeatureLevel, bStaticSceneOnly);
+				FGatherShadowPrimitivesPacket* Packet = new(FMemStack::Get()) FGatherShadowPrimitivesPacket(
+					Scene,
+					Views,
+					NULL,
+					StartPrimitiveIndex,
+					NumPrimitives,
+					PreShadows,
+					ViewDependentWholeSceneShadows,
+					FeatureLevel,
+					bStaticSceneOnly,
+					GatherStats);
 				Packets.Add(Packet);
 			}
 		}
@@ -3309,6 +4087,11 @@ void FSceneRenderer::GatherShadowPrimitives(
 		{
 			QUICK_SCOPE_CYCLE_COUNTER(STAT_RenderThreadFinalize);
 
+			for (int32 ShadowIdx = 0, Num = ViewDependentWholeSceneShadows.Num(); ShadowIdx < Num; ++ShadowIdx)
+			{
+				ViewDependentWholeSceneShadows[ShadowIdx]->PresizeSubjectPrimitiveArrays(GatherStats[ShadowIdx]);
+			}
+
 			for (int32 PacketIndex = 0; PacketIndex < Packets.Num(); PacketIndex++)
 			{
 				FGatherShadowPrimitivesPacket* Packet = Packets[PacketIndex];
