fix compile error

diff --git a/Engine/Source/Runtime/Renderer/Private/ShadowSetup.cpp b/Engine/Source/Runtime/Renderer/Private/ShadowSetup.cpp
--- a/Engine/Source/Runtime/Renderer/Private/ShadowSetup.cpp
+++ b/Engine/Source/Runtime/Renderer/Private/ShadowSetup.cpp
@@ -1585,8 +1585,6 @@ uint64 FProjectedShadowInfo::AddSubjectPrimitive_AnyThread(
 	FAddSubjectPrimitiveStats& OutStats,
 	FAddSubjectPrimitiveOverflowedIndices& OverflowBuffer) const
 {
-	SCOPE_CYCLE_COUNTER(STAT_AddSubjectPrimitive);
-
 	// Ray traced shadows use the GPU managed distance field object buffers, no CPU culling should be used
 	check(!bRayTracedDistanceField);
 
