Do not generate DepthOnly shaders for translucent materials

diff --git a/Engine/Source/Runtime/Renderer/Private/DepthRendering.h b/Engine/Source/Runtime/Renderer/Private/DepthRendering.h
--- a/Engine/Source/Runtime/Renderer/Private/DepthRendering.h
+++ b/Engine/Source/Runtime/Renderer/Private/DepthRendering.h
@@ -70,13 +70,17 @@ public:
 		{
 			return Parameters.VertexFactoryType->SupportsPositionOnly() && Parameters.MaterialParameters.bIsSpecialEngineMaterial;
 		}
+		
+		if (IsTranslucentBlendMode(Parameters.MaterialParameters.BlendMode) && !Parameters.MaterialParameters.bIsTranslucencyWritingCustomDepth)
+		{
+			return false;
+		}
 
 		// Only compile for the default material and masked materials
 		return (
 			Parameters.MaterialParameters.bIsSpecialEngineMaterial ||
 			!Parameters.MaterialParameters.bWritesEveryPixel ||
-			Parameters.MaterialParameters.bMaterialMayModifyMeshPosition ||
-			Parameters.MaterialParameters.bIsTranslucencyWritingCustomDepth);
+			Parameters.MaterialParameters.bMaterialMayModifyMeshPosition);
 	}
 
 	void GetShaderBindings(
@@ -146,9 +150,14 @@ public:
 
 	static bool ShouldCompilePermutation(const FMeshMaterialShaderPermutationParameters& Parameters)
 	{
+		if (IsTranslucentBlendMode(Parameters.MaterialParameters.BlendMode) && !Parameters.MaterialParameters.bIsTranslucencyWritingCustomDepth)
+		{
+			return false;
+		}
+		
 		return
 			// Compile for materials that are masked, avoid generating permutation for other platforms if bUsesMobileColorValue is true
-			((!Parameters.MaterialParameters.bWritesEveryPixel || Parameters.MaterialParameters.bHasPixelDepthOffsetConnected || Parameters.MaterialParameters.bIsTranslucencyWritingCustomDepth) && (!bUsesMobileColorValue || IsMobilePlatform(Parameters.Platform)))
+			((!Parameters.MaterialParameters.bWritesEveryPixel || Parameters.MaterialParameters.bHasPixelDepthOffsetConnected) && (!bUsesMobileColorValue || IsMobilePlatform(Parameters.Platform)))
 			// Mobile uses material pixel shader to write custom stencil to color target
 			|| (IsMobilePlatform(Parameters.Platform) && (Parameters.MaterialParameters.bIsDefaultMaterial || Parameters.MaterialParameters.bMaterialMayModifyMeshPosition));
 	}
