Repair ShadowSetup as a result of LoadTimes changes interacting with new shadow code
#rnx

diff --git a/Engine/Source/Runtime/Renderer/Private/ShadowSetup.cpp b/Engine/Source/Runtime/Renderer/Private/ShadowSetup.cpp
--- a/Engine/Source/Runtime/Renderer/Private/ShadowSetup.cpp
+++ b/Engine/Source/Runtime/Renderer/Private/ShadowSetup.cpp
@@ -1837,7 +1837,7 @@ void FProjectedShadowInfo::FinalizeAddSubjectPrimitive(
 			const FCachedMeshDrawCommandInfo& CmdInfo = PrimitiveSceneInfo->StaticMeshCommandInfos[CmdIdx];
 			const FScene* Scene = PrimitiveSceneInfo->Scene;
 			const FMeshDrawCommand* CachedCmd = CmdInfo.StateBucketId >= 0 ?
-				&Scene->CachedMeshDrawCommandStateBuckets[FSetElementId::FromInteger(CmdInfo.StateBucketId)].MeshDrawCommand :
+				&Scene->CachedMeshDrawCommandStateBuckets[EMeshPass::CSMShadowDepth].GetByElementId(CmdInfo.StateBucketId).Key :
 				&Scene->CachedDrawLists[EMeshPass::CSMShadowDepth].MeshDrawCommands[CmdInfo.CommandIndex];
 			const int32 PrimIdx = PrimitiveSceneInfo->GetIndex();
 
