Vulkan: fix crash on exit (UE-89075)

diff --git a/Engine/Source/Runtime/Renderer/Private/VolumetricFogVoxelization.cpp b/Engine/Source/Runtime/Renderer/Private/VolumetricFogVoxelization.cpp
--- a/Engine/Source/Runtime/Renderer/Private/VolumetricFogVoxelization.cpp
+++ b/Engine/Source/Runtime/Renderer/Private/VolumetricFogVoxelization.cpp
@@ -93,7 +93,9 @@ public:
 
 	virtual void ReleaseRHI() override
 	{
+		Buffers.PositionVertexBuffer.ReleaseRHI();
 		Buffers.PositionVertexBuffer.ReleaseResource();
+		Buffers.StaticMeshVertexBuffer.ReleaseRHI();
 		Buffers.StaticMeshVertexBuffer.ReleaseResource();
 	}
 };
