Integrate hair changes (rendering+simulation)

diff --git a/Engine/Shaders/Private/HairStrands/HairStrandsVisibilityVS.usf b/Engine/Shaders/Private/HairStrands/HairStrandsVisibilityVS.usf
--- a/Engine/Shaders/Private/HairStrands/HairStrandsVisibilityVS.usf
+++ b/Engine/Shaders/Private/HairStrands/HairStrandsVisibilityVS.usf
@@ -1,9 +1,11 @@
 // Copyright Epic Games, Inc. All Rights Reserved.
 
+#include "HairStrandsVisibilityCommonStruct.ush"
+
 #include "../Common.ush"
 #include "../VelocityCommon.ush"
 #include "/Engine/Generated/Material.ush"
-#include "/Engine/Generated/VertexFactory.ush"
+#include "/Engine/Generated/VertexFactory.ush" 
 #include "HairStrandsVisibilityCommon.ush"
 
 struct FHairVisibilityAccumulateVSToPS
@@ -27,12 +29,31 @@ struct FHairVisibilityAccumulateVSToPS
 
 #if VERTEXSHADER 
 
-/** transform mesh as normal */
 void Main(
 	FVertexFactoryInput Input,
 	OPTIONAL_VertexID
 	out VS_OUTPUT_TYPE Output
-	)
+)
+#if HAIR_RENDER_MODE == RENDER_MODE_MSAA_VISIBILITY && VF_STRAND_HAIR
+{
+	ResolvedView = ResolveView();
+
+	const bool bUseStableRasterization = UseStableRasterization();
+	const FHairRenderInfo HairRenderInfo = GetHairRenderInfo(ResolvedView.HairRenderInfo, ResolvedView.HairRenderInfoBits, bUseStableRasterization);
+	const float VelocityScale = 0;
+
+	FHairViewInfo HairViewInfo;
+	HairViewInfo.TranslatedWorldCameraOrigin = ResolvedView.TranslatedWorldCameraOrigin;
+	HairViewInfo.ViewForward = ResolvedView.ViewForward;
+	HairViewInfo.RadiusAtDepth1 = lerp(HairRenderInfo.RadiusAtDepth1Primary, HairRenderInfo.RadiusAtDepth1Velocity, VelocityScale);
+	HairViewInfo.bIsOrthoView = HairRenderInfo.bIsOrthoView;
+
+	uint HairPrimitiveId = 0;
+	const float4 WorldPosition = VertexFactoryGetWorldPosition_Visibility(Input, HairViewInfo, HairPrimitiveId);
+	Output.Position = mul(WorldPosition, View.TranslatedWorldToClip);
+	Output.HairPrimitiveId = HairPrimitiveId;
+}
+#else // HAIR_RENDER_MODE == RENDER_MODE_MSAA_VISIBILITY
 {
 	ResolvedView = ResolveView();
 
@@ -45,15 +66,20 @@ void Main(
 #else
 	float4 WorldPositionRaw = VertexFactoryGetWorldPosition(Input, VFIntermediates);
 #endif
-	float4 PrevWorldPosition = VertexFactoryGetPreviousWorldPosition(Input, VFIntermediates);
 	
+	bool bUseStableRasterization = false;
+#if VF_STRAND_HAIR
+	bUseStableRasterization = UseStableRasterization();
+#endif
+
 	// Velocity computation
 	// The velocity is used to adapt/increase the minimum rasterization size in order to avoid aliasing under heavy-motion.
 	// The higher the velocity, the larger a strand becomes.
 	float VelocityScale = 0;
-	const FHairRenderInfo HairRenderInfo = GetHairRenderInfo(ResolvedView.HairRenderInfo);
+	const FHairRenderInfo HairRenderInfo = GetHairRenderInfo(ResolvedView.HairRenderInfo, ResolvedView.HairRenderInfoBits, bUseStableRasterization);
 	#if HAIR_RENDER_MODE == RENDER_MODE_MSAA || HAIR_RENDER_MODE == RENDER_MODE_PPLL
 	{
+		const float4 PrevWorldPosition	= VertexFactoryGetPreviousWorldPosition(Input, VFIntermediates);
 		const float4 ScreenPos			= mul(float4(WorldPositionRaw.xyz, 1), ResolvedView.TranslatedWorldToClip);
 		const float4 PrevScreenPosObj	= mul(float4(PrevWorldPosition.xyz, 1), ResolvedView.PrevTranslatedWorldToClip);
 		Output.PackedVelocityA			= INVARIANT(ScreenPos);
@@ -105,5 +131,6 @@ void Main(
 
 	OutputVertexID( Output );
 }
+#endif // HAIR_RENDER_MODE == RENDER_MODE_MSAA_VISIBILITY
 
 #endif // VERTEXSHADER
\ No newline at end of file
diff --git a/Engine/Shaders/Private/HairStrands/HairStrandsVoxelMip.usf b/Engine/Shaders/Private/HairStrands/HairStrandsVoxelMip.usf
--- a/Engine/Shaders/Private/HairStrands/HairStrandsVoxelMip.usf
+++ b/Engine/Shaders/Private/HairStrands/HairStrandsVoxelMip.usf
@@ -24,7 +24,7 @@ uint ComputeMipDensity(
 		((RawDensity7 & VOXEL_OPAQUE_MASK) >> VOXEL_OPAQUE_SHIFT);
 	const uint OutTotalOpaque = uint(clamp(TotalOpaque / 8.f, TotalOpaque > 0 ? 1 : 0, 0xFF)) << VOXEL_OPAQUE_SHIFT;
 
-	uint TotalHair = min(VOXEL_HAIR_MASK,
+	uint TotalHair = 
 		(RawDensity0 & VOXEL_HAIR_MASK) +
 		(RawDensity1 & VOXEL_HAIR_MASK) +
 		(RawDensity2 & VOXEL_HAIR_MASK) +
@@ -32,13 +32,13 @@ uint ComputeMipDensity(
 		(RawDensity4 & VOXEL_HAIR_MASK) +
 		(RawDensity5 & VOXEL_HAIR_MASK) +
 		(RawDensity6 & VOXEL_HAIR_MASK) +
-		(RawDensity7 & VOXEL_HAIR_MASK));
+		(RawDensity7 & VOXEL_HAIR_MASK);
 	TotalHair /= 8;
 
-	return TotalHair | OutTotalOpaque;
+	return min(VOXEL_HAIR_MASK, TotalHair) | OutTotalOpaque;
 }
 
-#if SHADER_VOXEL
+#if SHADER_MIP_VOXEL
 
 uint VoxelResolution;
 uint SourceMip;
@@ -158,7 +158,7 @@ void MainCS(uint3 DispatchThreadId : SV_DispatchThreadID)
 
 #endif
 
-#if SHADER_VIRTUALVOXEL
+#if SHADER_MIP_VIRTUALVOXEL
 int3 PageCountResolution;
 uint PageResolution;
 uint SourceMip;
@@ -219,7 +219,7 @@ void MainCS(uint3 DispatchThreadId : SV_DispatchThreadID)
 }
 #endif
 
-#if SHADER_INDIRECTARGS
+#if SHADER_MIP_INDIRECTARGS
 #include "HairStrandsVoxelPageCommon.ush"
 
 uint PageResolution;
@@ -243,4 +243,42 @@ void MainCS(uint3 DispatchThreadId : SV_DispatchThreadID)
 	OutIndirectArgs[2] = DispatchZ;
 }
 
-#endif
\ No newline at end of file
+#endif
+
+
+#if SHADER_UPDATE_PAGEINDEX
+int3 PageCountResolution;
+uint PageResolution;
+uint MipIt;
+
+Texture3D<uint>	  DensityTexture;
+Buffer<uint>	  PageToPageIndexBuffer;
+RWBuffer<uint>	  OutPageIndexBuffer;
+
+[numthreads(1, 1, 64)]
+void MainCS(uint3 DispatchThreadId : SV_DispatchThreadID)
+{
+	const uint  MaxDispatchCountPerDim = 65535;
+	const uint  PageIndex = DispatchThreadId.x + DispatchThreadId.y * MaxDispatchCountPerDim;
+	const uint3 PageCoord = IndexToCoord(PageIndex, PageCountResolution);
+
+	const uint OutPageResolution = PageResolution >> MipIt;
+	const uint TotalVoxelPerOutPageResolution = OutPageResolution * OutPageResolution * OutPageResolution;
+	const uint VoxelIndex = DispatchThreadId.z;
+	if (VoxelIndex >= TotalVoxelPerOutPageResolution)
+		return;
+
+	const uint3 OutVoxelCoordOffset = IndexToCoord(VoxelIndex, OutPageResolution);
+	const uint3 OutVoxelCoord = PageCoord * OutPageResolution + OutVoxelCoordOffset;
+
+	const uint Density = DensityTexture.Load(uint4(OutVoxelCoord, MipIt));
+	// Update the page index with invalid page index if the voxel does not contain any data.
+	// This allow to save tracing cost when evaluating the transmittance.
+	if (Density == 0)
+	{
+		const uint PageIndexOffset = PageToPageIndexBuffer[PageIndex];
+		OutPageIndexBuffer[PageIndexOffset] = INVALID_VOXEL_PAGE_INDEX;
+	}
+}
+#endif
+
diff --git a/Engine/Shaders/Private/HairStrands/HairStrandsVoxelPageAllocation.usf b/Engine/Shaders/Private/HairStrands/HairStrandsVoxelPageAllocation.usf
--- a/Engine/Shaders/Private/HairStrands/HairStrandsVoxelPageAllocation.usf
+++ b/Engine/Shaders/Private/HairStrands/HairStrandsVoxelPageAllocation.usf
@@ -163,6 +163,200 @@ void AllocatePageIndex(uint2 DispatchThreadId : SV_DispatchThreadID)
 }
 #endif
 
+#if SHADER_MARKVALID_PREPARE || SHADER_MARKVALID_INDIRECTARG || SHADER_MARKVALID_SCATTER
+struct FScatterCoord
+{
+	uint3 MinCoord;
+	uint  PartIndex;
+	uint3 Resolution;
+	uint  MacroGroupId;
+};
+
+uint2 PackScatterCoord(FScatterCoord In)
+{
+	uint2 Out;
+	Out.x =
+		((In.MinCoord.x & 0xFF))    |
+		((In.MinCoord.y & 0xFF)<<8) |
+		((In.MinCoord.z & 0xFF)<<16)|
+		((In.PartIndex  & 0xFF)<<24);
+
+	Out.y =
+		((In.Resolution.x & 0xFF))	  |
+		((In.Resolution.y & 0xFF)<<8) |
+		((In.Resolution.z & 0xFF)<<16)|
+		((In.MacroGroupId & 0xFF)<<24);
+
+	return Out;
+}
+
+FScatterCoord UnpackScatterCoord(uint2 In)
+{
+	FScatterCoord Out;
+	Out.MinCoord = uint3(
+		 In.x        & 0xFF,
+		(In.x >> 8)  & 0xFF,
+		(In.x >> 16) & 0xFF);
+	Out.PartIndex = (In.x>>24) & 0xFF;
+
+	Out.Resolution = uint3(
+		 In.y        & 0xFF,
+		(In.y >> 8)  & 0xFF,
+		(In.y >> 16) & 0xFF);
+	Out.MacroGroupId = (In.y >> 24) & 0xFF;
+
+	return Out;
+}
+
+void UpdatePackScatterCoord(inout uint2 Out, uint PartIndex)
+{
+	Out.x = (Out.x & 0x00FFFFFF) | ((PartIndex & 0xFF) << 24);
+}
+
+#define SCATTER_WORK_ITEM_COUNT 8
+#define SCATTER_GROUP_SIZE 32
+#define SCATTER_GROUP_LINE 4
+
+#endif
+
+#if SHADER_MARKVALID_PREPARE
+uint			MaxClusterCount;
+uint			MacroGroupId;
+uint			MaxScatterAllocationCount;
+
+Buffer<int>		ClusterAABBsBuffer;
+Buffer<int>		MacroGroupAABBBuffer;
+Buffer<uint4>	PageIndexResolutionAndOffsetBuffer;
+
+RWBuffer<uint>	OutValidPageIndexBuffer;
+RWBuffer<uint>  OutDeferredScatterCounter;
+RWBuffer<uint2> OutDeferredScatterBuffer;
+
+
+// PageIndexBuffer      is sampled with linear coordinate computed from the 3d page coordinate. VALID NODE ARE NOT COMPACTED. It contains the LINEAR PAGE INDEX (to map to the 3d volume).
+// PageIndexCoordBuffer is sampled with linear coordinate for allocated nodes. VALID NODE ARE COMPACTED. It contains the 3d page coordinate and ClustedId. Only used for opaque voxel injection.
+
+[numthreads(SCATTER_GROUP_SIZE, 1, 1)]
+void MarkValid_PrepareCS(uint2 DispatchThreadId : SV_DispatchThreadID)
+{
+	const uint ClusterIndex = DispatchThreadId.x;
+	if (ClusterIndex >= MaxClusterCount)
+		return;
+
+	const uint BaseClusterIndex = 6 * ClusterIndex;
+
+	FHairAABB ClusterBound;
+	ClusterBound.Min.x = ClusterAABBsBuffer[BaseClusterIndex + 0];
+	ClusterBound.Min.y = ClusterAABBsBuffer[BaseClusterIndex + 1];
+	ClusterBound.Min.z = ClusterAABBsBuffer[BaseClusterIndex + 2];
+
+	ClusterBound.Max.x = ClusterAABBsBuffer[BaseClusterIndex + 3];
+	ClusterBound.Max.y = ClusterAABBsBuffer[BaseClusterIndex + 4];
+	ClusterBound.Max.z = ClusterAABBsBuffer[BaseClusterIndex + 5];
+
+	const uint4 PageIndexResolutionAndOffset = PageIndexResolutionAndOffsetBuffer.Load(MacroGroupId);
+	const FHairAABB MacroGroupBound			 = ReadHairAABB(MacroGroupId, MacroGroupAABBBuffer);
+	const int3 PageIndexResolution			 = PageIndexResolutionAndOffset.xyz;
+	const uint PageIndexOffset				 = PageIndexResolutionAndOffset.w;
+
+
+	uint3 MinCoord = PositionToCoord(ClusterBound.Min, MacroGroupBound.Min, MacroGroupBound.Max, PageIndexResolution);
+	uint3 MaxCoord = PositionToCoord(ClusterBound.Max, MacroGroupBound.Min, MacroGroupBound.Max, PageIndexResolution);
+
+	MinCoord = clamp(MinCoord, 0, PageIndexResolution - 1);
+	MaxCoord = clamp(MaxCoord, 0, PageIndexResolution - 1);
+
+	FScatterCoord Coord;
+	Coord.MacroGroupId = MacroGroupId;
+	Coord.PartIndex = 0;
+	Coord.MinCoord = MinCoord;
+	Coord.Resolution = (MaxCoord - MinCoord) + 1;
+
+	const uint ScatterCount = Coord.Resolution.x * Coord.Resolution.y * Coord.Resolution.z;
+
+	// Find a good sweet spot
+	if (ScatterCount <= 4)
+	{
+		for (uint z = MinCoord.z; z <= MaxCoord.z; ++z)
+		for (uint y = MinCoord.y; y <= MaxCoord.y; ++y)
+		for (uint x = MinCoord.x; x <= MaxCoord.x; ++x)
+		{
+			const uint3 PageIndexCoord = uint3(x, y, z);
+			const uint LinearPageIndexCoord = CoordToIndex(PageIndexCoord, PageIndexResolution, PageIndexOffset);
+			InterlockedOr(OutValidPageIndexBuffer[LinearPageIndexCoord], 1u);
+		}
+	}
+	else
+	{
+		const uint WorkItemCount = (ScatterCount + SCATTER_WORK_ITEM_COUNT - 1) / SCATTER_WORK_ITEM_COUNT;
+	
+		uint AllocationOffset = 0;
+		InterlockedAdd(OutDeferredScatterCounter[0], WorkItemCount, AllocationOffset);
+		if (WorkItemCount + AllocationOffset < MaxScatterAllocationCount)
+		{
+			uint2 PackedCoord = PackScatterCoord(Coord);
+			for (uint Item = 0; Item < WorkItemCount; ++Item)
+			{
+				UpdatePackScatterCoord(PackedCoord, Item);
+				OutDeferredScatterBuffer[AllocationOffset + Item] = PackedCoord;
+			}
+		}
+	}
+}
+#endif
+
+#if SHADER_MARKVALID_INDIRECTARG
+Buffer<uint> DeferredScatterCounter;
+RWBuffer<uint> OutIndirectArgsBuffer;
+
+[numthreads(1, 1, 1)]
+void MarkValid_BuildIndirectArgCS(uint2 DispatchThreadId : SV_DispatchThreadID)
+{
+	const uint ScatterItemCount = DeferredScatterCounter[0] * SCATTER_WORK_ITEM_COUNT;
+	const uint DispatchCount = (ScatterItemCount + SCATTER_GROUP_SIZE - 1) / SCATTER_GROUP_SIZE;
+	const uint DispatchX = min(DispatchCount, SCATTER_GROUP_LINE);
+	const uint DispatchY = (DispatchCount + SCATTER_GROUP_LINE - 1) / SCATTER_GROUP_LINE;
+
+	OutIndirectArgsBuffer[0] = DispatchX;
+	OutIndirectArgsBuffer[1] = DispatchY;
+	OutIndirectArgsBuffer[2] = 1;
+}
+#endif
+
+#if SHADER_MARKVALID_SCATTER
+Buffer<int>		IndirectBufferArgs;
+Buffer<uint4>	PageIndexResolutionAndOffsetBuffer;
+Buffer<uint>	DeferredScatterCounter;
+Buffer<uint2>	DeferredScatterBuffer;
+RWBuffer<uint>	OutValidPageIndexBuffer;
+
+[numthreads(SCATTER_GROUP_SIZE, 1, 1)]
+void MarkValid_ScatterCS(uint2 DispatchThreadId : SV_DispatchThreadID)
+{
+	const uint LinearThreadId = DispatchThreadId.x + DispatchThreadId.y * SCATTER_GROUP_SIZE * SCATTER_GROUP_LINE;
+	const uint ScatterItem	 = LinearThreadId;
+	const uint WorkItemIndex = LinearThreadId / SCATTER_WORK_ITEM_COUNT;
+	const uint ItemIndex	 = LinearThreadId % SCATTER_WORK_ITEM_COUNT;
+	const uint MaxScatterItemCount = DeferredScatterCounter[0] * SCATTER_WORK_ITEM_COUNT;
+	if (ScatterItem < MaxScatterItemCount)
+	{
+		const FScatterCoord Coord = UnpackScatterCoord(DeferredScatterBuffer[WorkItemIndex]);
+
+		const uint4 PageIndexResolutionAndOffset = PageIndexResolutionAndOffsetBuffer.Load(Coord.MacroGroupId);
+		const int3 PageIndexResolution			 = PageIndexResolutionAndOffset.xyz;
+		const uint PageIndexOffset				 = PageIndexResolutionAndOffset.w;
+
+		const uint  ClusterPageCount = Coord.Resolution.x * Coord.Resolution.y * Coord.Resolution.z;
+		const uint LinearVoxelIt = (Coord.PartIndex * SCATTER_WORK_ITEM_COUNT) + ItemIndex;
+		if (LinearVoxelIt < ClusterPageCount)
+		{
+			const uint3 PageIndexCoord = IndexToCoord(LinearVoxelIt, Coord.Resolution) + Coord.MinCoord;
+			const uint LinearPageIndexCoord = CoordToIndex(PageIndexCoord, PageIndexResolution, PageIndexOffset);
+			InterlockedOr(OutValidPageIndexBuffer[LinearPageIndexCoord], 1u);
+		}
+	}
+}
+#endif
 
 #if SHADER_MARKVALID
 uint			MaxClusterCount;
@@ -248,6 +442,7 @@ Buffer<int> IndirectBufferArgs;
 
 RWBuffer<uint> PageIndexGlobalCounter;
 RWBuffer<uint> PageIndexBuffer;
+RWBuffer<uint> PageToPageIndexBuffer;
 RWBuffer<uint4> PageIndexCoordBuffer;
 
 groupshared uint LocalCounter;
@@ -300,7 +495,7 @@ void AllocateCS(uint GroupIndex : SV_GroupIndex, uint3 DispatchThreadId : SV_Dis
 	{
 		const uint PageIndex = GroupBase + Offset;
 		PageIndexBuffer[GridIndex] = PageIndex < PageCount ? PageIndex : INVALID_VOXEL_PAGE_INDEX;
-
+		PageToPageIndexBuffer[PageIndex] = GridIndex;
 		// Output the coordinates of the allocated page for indirect dispatch usage
 		const uint LinearIndex = DispatchThreadId.x;
 		const uint3 PageIndexCoord = IndexToCoord(LinearIndex, PageIndexResolution);
diff --git a/Engine/Shaders/Private/HairStrands/HairStrandsVoxelPageCommon.ush b/Engine/Shaders/Private/HairStrands/HairStrandsVoxelPageCommon.ush
--- a/Engine/Shaders/Private/HairStrands/HairStrandsVoxelPageCommon.ush
+++ b/Engine/Shaders/Private/HairStrands/HairStrandsVoxelPageCommon.ush
@@ -99,7 +99,7 @@ FPackedVirtualVoxelNodeDesc PackVoxelNode(FVirtualVoxelNodeDesc In)
 	Out.MinAABB = In.MinAABB;
 	Out.MaxAABB = In.MaxAABB;
 	Out.PackedPageIndexResolution =
-		  In.PageIndexResolution.x & 0xFF | 
+		 (In.PageIndexResolution.x & 0xFF) | 
 		((In.PageIndexResolution.y & 0xFF)<<8) | 
 		((In.PageIndexResolution.z & 0xFF)<<16);
 	Out.PageIndexOffset = In.PageIndexOffset;
@@ -112,6 +112,7 @@ struct FHairTraversalResult
 {
 	float HairCount;
 	float Visibility;
+	float HitT;
 };
 
 FHairTraversalResult InitHairTraversalResult()
@@ -119,6 +120,7 @@ FHairTraversalResult InitHairTraversalResult()
 	FHairTraversalResult Out;
 	Out.HairCount = 0;
 	Out.Visibility = 1;
+	Out.HitT = -1;
 	return Out;
 }
 
@@ -126,6 +128,7 @@ void Acc(inout FHairTraversalResult Out, const FHairTraversalResult In)
 {
 	Out.HairCount += In.HairCount;
 	Out.Visibility = min(Out.Visibility, In.Visibility);
+	Out.HitT = In.HitT >= 0 && Out.HitT >= 0 ? min(In.HitT, Out.HitT) : (In.HitT >= 0 ? In.HitT : Out.HitT);
 }
 
 ///////////////////////////////////////////////////////////////////////////////////////////////////
@@ -175,3 +178,22 @@ FHairTraversalResult GetHairVoxelDensity(uint3 InVolumeCoord, Texture3D<uint> In
 	Out.Visibility = GetInternalVoxelOpaqueVisibility(RawDensity);
 	return Out;
 }
+
+///////////////////////////////////////////////////////////////////////////////////////////////////
+
+uint EncodeBaseColorAndRoughness(float3 BaseColor, float Roughness)
+{
+	return
+		(uint(BaseColor.x * 0xFF) & 0xFF) |
+		(uint(BaseColor.y * 0xFF) & 0xFF) << 8 |
+		(uint(BaseColor.z * 0xFF) & 0xFF) << 16 |
+		(uint(Roughness * 0xFF) & 0xFF) << 24;
+}
+
+void DecodeBaseColorAndRoughness(uint Encoded, inout float3 BaseColor, out float Roughness)
+{
+	BaseColor.x = (float((Encoded) & 0xFF) / 255.f);
+	BaseColor.y = (float((Encoded >> 8) & 0xFF) / 255.f);
+	BaseColor.z = (float((Encoded >> 16) & 0xFF) / 255.f);
+	Roughness = (float((Encoded >> 24) & 0xFF) / 255.f);
+}
diff --git a/Engine/Shaders/Private/HairStrands/HairStrandsVoxelPageTraversal.ush b/Engine/Shaders/Private/HairStrands/HairStrandsVoxelPageTraversal.ush
--- a/Engine/Shaders/Private/HairStrands/HairStrandsVoxelPageTraversal.ush
+++ b/Engine/Shaders/Private/HairStrands/HairStrandsVoxelPageTraversal.ush
@@ -61,10 +61,10 @@ void AddAABBVolume(bool bEnabled, float3 P0, float3 P1)
 
 ////////////////////////////////////////////////////////////////////////////////////////////////////////////////////////////////
 
-bool ShouldStopTraversal(const FHairTraversalResult InResult, float InHairCountThreshold)
+bool ShouldStopTraversal(const FHairTraversalResult InResult, float InHairCountThreshold, bool InUseOpaqueVisibility=true)
 {
 	// Change this function once we have better result/control over the opaque filtering
-	return InResult.Visibility == 0 || (InHairCountThreshold != 0 && InResult.HairCount > InHairCountThreshold);
+	return (InUseOpaqueVisibility && InResult.Visibility == 0) || (InHairCountThreshold != 0 && InResult.HairCount > InHairCountThreshold);
 }
 
 float3 GetHairVoxelJitter(uint2 PixelCoord, uint Seed)
@@ -118,7 +118,10 @@ FHairTraversalResult ComputeHairCountVoxel(
 			Acc(Out, StepResult);
 
 			if (ShouldStopTraversal(Out, CountThreshold))
+			{
+				Out.HitT = StepIt * VoxelWorldSize;
 				return Out;
+			}
 		}
 	}
 
@@ -187,7 +190,10 @@ FHairTraversalResult ComputeHairCountVoxelCone(
 			Acc(Out, StepResult);
 
 			if (ShouldStopTraversal(Out, CountThreshold))
+			{
+				Out.HitT = T;
 				return Out;
+			}
 
 			const float VoxelRadius = T * TanConeAngle;
 			const float NumVoxel = VolumeSize / max(2 * VoxelRadius, VoxelSize);
@@ -220,18 +226,25 @@ struct FHairTraversalSettings
 	float	SteppingScale;		// Rate at which the raymarching step increase
 	float3	Random;
 	float	TanConeAngle;
+	bool	bUseOpaqueVisibility;
 	bool	bDebugEnabled;
 };
 
+float GetHairTraversalMaxT()
+{
+	return 999999;
+}
+
 FHairTraversalSettings InitHairTraversalSettings()
 {
 	FHairTraversalSettings Out;
 	Out.DensityScale = 1;
 	Out.CountThreshold = 0;
-	Out.DistanceThreshold = 100000.f;
+	Out.DistanceThreshold = GetHairTraversalMaxT();
 	Out.SteppingScale = 1;
 	Out.Random = 0.5f;
 	Out.TanConeAngle = 0;
+	Out.bUseOpaqueVisibility = true;
 	Out.bDebugEnabled = false;
 	return Out;
 }
@@ -347,8 +360,11 @@ FHairTraversalResult ComputeHairCountVirtualVoxel(
 				}
 				#endif
 
-				if (ShouldStopTraversal(Out, InSettings.CountThreshold))
+				if (ShouldStopTraversal(Out, InSettings.CountThreshold, InSettings.bUseOpaqueVisibility))
+				{
+					Out.HitT = length(HitP - O);
 					return Out;
+				}
 			}
 			#if VOXEL_TRAVERSAL_TYPE == VOXEL_TRAVERSAL_LINEAR_MIPMAP
 			StepScale = min(InCommonDesc.PageResolution, StepScale * InSettings.SteppingScale);
@@ -513,6 +529,7 @@ FHairTraversalResult ComputeHairCountVirtualVoxel(
 						#if VOXEL_TRAVERSAL_DEBUG
 						AddStepAABB(InSettings.bDebugEnabled, InnerHitP - D * 2, InnerHitP + D * 2, float3(1, 0.25f, 0));
 						#endif
+						Out.HitT = length(InnerHitP - O);
 						return Out;
 					}
 				}
diff --git a/Engine/Shaders/Private/HairStrands/HairStrandsVoxelRasterCompute.usf b/Engine/Shaders/Private/HairStrands/HairStrandsVoxelRasterCompute.usf
new file mode 100644
--- /dev/null
+++ b/Engine/Shaders/Private/HairStrands/HairStrandsVoxelRasterCompute.usf
@@ -0,0 +1,146 @@
+// Copyright Epic Games, Inc. All Rights Reserved.
+
+#include "HairStrandsVoxelPageCommonStruct.ush"
+#include "../Common.ush"
+
+#include "HairStrandsVertexFactoryCommon.ush"
+#include "HairStrandsVoxelPageCommon.ush"
+
+///////////////////////////////////////////////////////////////////////////
+uint			DispatchCountX;
+uint			MaxRasterCount;
+uint			FrameIdMod8;
+uint			MacroGroupId;
+float			HairStrandsVF_Radius;
+float			HairStrandsVF_Length;
+float			HairStrandsVF_Density;
+float3			HairStrandsVF_PositionOffset;
+uint			HairStrandsVF_VertexCount;
+float4x4		HairStrandsVF_LocalToWorldPrimitiveTransform;
+Buffer<uint4>	HairStrandsVF_PositionBuffer;
+
+uint			VirtualVoxel_PageIndexCount;
+uint			VirtualVoxel_PageResolution;
+float			VirtualVoxel_VoxelWorldSize;
+uint3			VirtualVoxel_PageCountResolution;
+uint3			VirtualVoxel_PageTextureResolution;
+Buffer<uint>	VirtualVoxel_PageIndexBuffer;
+Buffer<uint4>	VirtualVoxel_PageIndexCoordBuffer;
+StructuredBuffer<FPackedVirtualVoxelNodeDesc> VirtualVoxel_NodeDescBuffer;
+
+RWTexture3D<uint> OutPageTexture;
+
+///////////////////////////////////////////////////////////////////////////
+
+float3 GetHairVoxelJitter(uint2 PixelCoord, uint Seed)
+{
+	return float3(
+		InterleavedGradientNoise(PixelCoord.xy, Seed),
+		InterleavedGradientNoise(PixelCoord.xy, Seed * 117),
+		InterleavedGradientNoise(PixelCoord.xy, Seed * 7901));
+}
+
+#define JITTER_ENABLE 0
+
+#define GROUP_SIZE 32
+[numthreads(GROUP_SIZE, 1, 1)]
+void MainCS(uint3 DispatchThreadID : SV_DispatchThreadID)
+{
+	const uint VertexIndex = DispatchThreadID.x + DispatchThreadID.y * GROUP_SIZE * DispatchCountX;
+	if (VertexIndex >= HairStrandsVF_VertexCount)
+		return;
+
+	const FHairControlPoint CP0 = UnpackHairControlPoint(
+		HairStrandsVF_PositionBuffer[VertexIndex],
+		HairStrandsVF_PositionOffset,
+		HairStrandsVF_Radius,
+		HairStrandsVF_Length);
+
+	if (CP0.Type == HAIR_CONTROLPOINT_END)
+		return;
+
+	const uint VertexIndex1 = min(VertexIndex + 1, HairStrandsVF_VertexCount - 1);
+	const FHairControlPoint CP1 = UnpackHairControlPoint(
+		HairStrandsVF_PositionBuffer[VertexIndex1],
+		HairStrandsVF_PositionOffset,
+		HairStrandsVF_Radius,
+		HairStrandsVF_Length);
+
+	const float3 WP0 = mul(float4(CP0.Position, 1), HairStrandsVF_LocalToWorldPrimitiveTransform).xyz;
+	const float HairCoverage0 = CP0.WorldRadius / max(CP0.WorldRadius, VirtualVoxel_VoxelWorldSize) * HairStrandsVF_Density;
+
+	const float3 WP1 = mul(float4(CP1.Position, 1), HairStrandsVF_LocalToWorldPrimitiveTransform).xyz;
+	const float HairCoverage1 = CP1.WorldRadius / max(CP1.WorldRadius, VirtualVoxel_VoxelWorldSize) * HairStrandsVF_Density;
+
+	const FVirtualVoxelNodeDesc NodeDesc = UnpackVoxelNode(VirtualVoxel_NodeDescBuffer[MacroGroupId], VirtualVoxel_PageResolution);
+
+	// In order to reduce aliasing, we increase the number of steps. This makes the result more comparable to the raster pass.
+	const float LineStepMultiplier = 1.5f;
+
+	FVirtualVoxelCommonDesc CommonDesc;
+	CommonDesc.PageCountResolution		= VirtualVoxel_PageCountResolution;
+	CommonDesc.VoxelWorldSize			= VirtualVoxel_VoxelWorldSize / LineStepMultiplier;
+	CommonDesc.PageTextureResolution	= VirtualVoxel_PageTextureResolution;
+	CommonDesc.PageResolution			= VirtualVoxel_PageResolution;
+
+	// Count the number of fibers which are within a cylinder defined by the voxel size, 
+	// and the distance between the origin and the extent of the volume
+	// This assumes that the voxel volume is cubic (i.e. equal dimensions on all sides)
+	const float3 Line = WP1 - WP0;
+	const float LineLength = length(Line);
+	const float3 StepD = normalize(Line) * CommonDesc.VoxelWorldSize;
+
+	// Step according to voxel size
+	int3 CurrentPageIndexCoord = -1;
+	bool bIsPageValid = false;
+	uint3 PageCoord = 0;
+
+	// If we have long segment we could break them into batch (e.g. 8 voxels long), queue them, and indirect dispatch them. 
+	// This would make the workload more uniform/coherent. Currently, breaking into smaller batch, does not seems to save a 
+	// lot of cost
+	#if JITTER_ENABLE
+	const float3 Jitter = (GetHairVoxelJitter(frac(CP0.Position.xy), FrameIdMod8) * 2 - 1);
+	#else
+	const float3 Jitter = 0;
+	#endif
+
+	const float fMaxStep = LineLength / CommonDesc.VoxelWorldSize;
+	const float MaxStep = float(min(ceil(fMaxStep), MaxRasterCount));
+
+	for (float StepIt = 0.0f; StepIt < MaxStep; ++StepIt)
+	{
+		const float U = (StepIt + 0.5f) / float(MaxStep);
+		const float Radius = lerp(CP0.WorldRadius, CP1.WorldRadius, U);
+
+		const float3 HitP = WP0 + StepIt * StepD + Jitter * Radius;
+		const int3 VolumeCoord = PositionToCoord(HitP, NodeDesc.MinAABB, NodeDesc.MaxAABB, NodeDesc.VirtualResolution);
+		const int3 PageIndexCoord = VolumeCoord / CommonDesc.PageResolution;
+
+		// Update page index only when needed
+		if (PageIndexCoord.x != CurrentPageIndexCoord.x ||
+			PageIndexCoord.y != CurrentPageIndexCoord.y ||
+			PageIndexCoord.z != CurrentPageIndexCoord.z)
+		{
+			CurrentPageIndexCoord = PageIndexCoord;
+			const uint LinearPageIndexCoord = CoordToIndex(PageIndexCoord, NodeDesc.PageIndexResolution, NodeDesc.PageIndexOffset);
+			const uint PageIndex = VirtualVoxel_PageIndexBuffer.Load(LinearPageIndexCoord);
+
+			bIsPageValid = PageIndex != INVALID_VOXEL_PAGE_INDEX;
+			if (bIsPageValid)
+			{
+				PageCoord = IndexToCoord(PageIndex, CommonDesc.PageCountResolution);
+			}
+		}
+
+		if (bIsPageValid)
+		{
+			const int3 VoxelPageBase = PageCoord * CommonDesc.PageResolution;
+			const int3 VoxelPageOffset = VolumeCoord - PageIndexCoord * CommonDesc.PageResolution;
+			const int3 VoxelPageCoord = VoxelPageBase + VoxelPageOffset;
+
+			const float VoxelFixPointScale = GetVoxelDensityFixPointScale();
+			const float HairCoverage = lerp(HairCoverage0, HairCoverage1, U);
+			InterlockedAdd(OutPageTexture[VoxelPageCoord], HairCoverage * VoxelFixPointScale);
+		}
+	}
+}
\ No newline at end of file
diff --git a/Engine/Shaders/Private/MaterialTemplate.ush b/Engine/Shaders/Private/MaterialTemplate.ush
--- a/Engine/Shaders/Private/MaterialTemplate.ush
+++ b/Engine/Shaders/Private/MaterialTemplate.ush
@@ -13,6 +13,7 @@
 #include "/Engine/Private/CircleDOFCommon.ush"
 #include "/Engine/Private/GlobalDistanceFieldShared.ush"
 #include "/Engine/Private/SceneData.ush"
+#include "/Engine/Private/HairShadingCommon.ush"
 
 #if USES_SPEEDTREE
 	#include "/Engine/Private/SpeedTreeCommon.ush"
@@ -1088,6 +1089,11 @@ float3 MaterialExpressionGetHairTangent(FMaterialPixelParameters Parameters)
 	return Parameters.TangentToWorld[2];
 }
 
+float3 MaterialExpressionGetHairColorFromMelanin(float Melanin, float Redness, float3 DyeColor)
+{
+	return GetHairColorFromMelanin(Melanin, Redness, DyeColor);
+}
+
 float4 MaterialExpressionAtmosphericFog(FMaterialPixelParameters Parameters, float3 AbsoluteWorldPosition)
 {
 #if MATERIAL_ATMOSPHERIC_FOG
diff --git a/Engine/Shaders/Private/RayTracing/RayTracingOcclusionRGS.usf b/Engine/Shaders/Private/RayTracing/RayTracingOcclusionRGS.usf
--- a/Engine/Shaders/Private/RayTracing/RayTracingOcclusionRGS.usf
+++ b/Engine/Shaders/Private/RayTracing/RayTracingOcclusionRGS.usf
@@ -1,5 +1,7 @@
 // Copyright Epic Games, Inc. All Rights Reserved.
 
+#include "../HairStrands/HairStrandsVoxelPageCommonStruct.ush"
+
 #include "../Common.ush"
 #include "../MonteCarlo.ush"
 #include "../DeferredShadingCommon.ush"
@@ -28,8 +30,6 @@
 RaytracingAccelerationStructure TLAS;
 RWTexture2D<float4> RWOcclusionMaskUAV;
 RWTexture2D<float> RWRayDistanceUAV;
-Texture2D<uint4> HairCategorizationTexture;
-Texture2D<uint> HairLightChannelMaskTexture;
 
 uint ShadowMaskType;
 uint LightingChannelMask;
@@ -38,6 +38,85 @@ float NormalBias;
 int4 LightScissor;
 int2 PixelOffset;
 
+#if USE_HAIR_LIGHTING
+#define VOXEL_TRAVERSAL_DEBUG 0
+#define VOXEL_TRAVERSAL_TYPE VOXEL_TRAVERSAL_LINEAR_MIPMAP
+//#define VOXEL_TRAVERSAL_TYPE VOXEL_TRAVERSAL_LINEAR
+
+#include "../HairStrands/HairStrandsVoxelPageCommon.ush"
+#include "../HairStrands/HairStrandsVoxelPageTraversal.ush"
+
+uint				bUseHairVoxel;
+Texture2D<uint4>	HairCategorizationTexture;
+Texture2D<uint>		HairLightChannelMaskTexture;
+
+float TraverseHair(uint2 PixelCoord, float3 RayOrigin, float3 RayDirection, float InMinT)
+{
+	float HitT = InMinT;
+
+	// Depth bias
+	// Origin is shifted 2 voxels away towards the light + a constant bias of the size of the voxel
+	const float3 Random = 0.5;// GetHairVoxelJitter(PixelCoord.xy, View.StateFrameIndexMod8);
+	{
+		const float MaxVoxelSize = VirtualVoxel.VoxelWorldSize;
+		const float3 DepthBias = MaxVoxelSize * (2 * RayDirection + Random);
+		RayOrigin += DepthBias;
+	}
+
+	#if VOXEL_TRAVERSAL_DEBUG
+	const bool bDebugEnabled = PixelCoord.x == GetCursorPos().x && PixelCoord.y == GetCursorPos().y;
+	#else
+	const bool bDebugEnabled = false;
+	#endif
+
+	FHairTraversalResult Result = InitHairTraversalResult();
+
+	const float DistanceThreshold = 1000.0f; 
+	const float CoverageThreshold = 0.995f; // When Coverage is high, we do not trace shadow on opaque since hair/fur is covering the background.
+
+	FVirtualVoxelCommonDesc CommonDesc;
+	CommonDesc.PageCountResolution = VirtualVoxel.PageCountResolution;
+	CommonDesc.VoxelWorldSize = VirtualVoxel.VoxelWorldSize;
+	CommonDesc.PageTextureResolution = VirtualVoxel.PageTextureResolution;
+	CommonDesc.PageResolution = VirtualVoxel.PageResolution;
+
+	const float3 P0 = RayOrigin;
+	const float3 P1 = RayOrigin + RayDirection * (HitT>=0 ? min(HitT, DistanceThreshold) : DistanceThreshold);
+
+	for (uint MacroGroupId = 0; MacroGroupId < VirtualVoxel.NodeDescCount; ++MacroGroupId)
+	{
+		const FPackedVirtualVoxelNodeDesc PackedNode = VirtualVoxel.NodeDescBuffer[MacroGroupId];
+		const FVirtualVoxelNodeDesc NodeDesc = UnpackVoxelNode(PackedNode, VirtualVoxel.PageResolution);
+
+		FHairTraversalSettings TraversalSettings = InitHairTraversalSettings();
+		TraversalSettings.DensityScale = VirtualVoxel.DensityScale;
+		TraversalSettings.CountThreshold = 1;
+		TraversalSettings.DistanceThreshold = DistanceThreshold;
+		TraversalSettings.bDebugEnabled = bDebugEnabled;
+		TraversalSettings.SteppingScale = VirtualVoxel.SteppingScale;
+		TraversalSettings.Random = Random;
+		TraversalSettings.bUseOpaqueVisibility = false;
+
+		Result = ComputeHairCountVirtualVoxel(
+			P0,
+			P1,
+			CommonDesc,
+			NodeDesc,
+			VirtualVoxel.PageIndexBuffer,
+			VirtualVoxel.PageTexture,
+			TraversalSettings);
+
+		if (Result.HitT >= 0)
+		{
+			const float WorldHitT = DistanceThreshold * Result.HitT;
+			HitT = HitT >= 0 ? min(HitT, WorldHitT) : WorldHitT;
+		}
+	}
+
+	return HitT;
+}
+#endif
+
 bool GenerateOcclusionRay(
 	FLightShaderParameters LightParameters,
 	float3 WorldPosition,
@@ -302,6 +381,13 @@ RAY_TRACING_ENTRY_RAYGEN(OcclusionRGS)
 				RaytracingMask,
 				Ray);
 
+			#if USE_HAIR_LIGHTING
+			if (ShadowMaskType != ShadowMaskType_Hair && bUseHairVoxel)
+			{
+				MinimalPayload.HitT = TraverseHair(PixelCoord, Ray.Origin, Ray.Direction, MinimalPayload.HitT);
+			}
+			#endif
+
 			RayCount += 1.0;
 			Visibility += MinimalPayload.IsMiss() ? 1.0 : 0.0;
 			if (MinimalPayload.IsHit())
@@ -364,4 +450,4 @@ RAY_TRACING_ENTRY_RAYGEN(OcclusionRGS)
 
 		RWOcclusionMaskUAV[PixelCoord] = OutColor;
 	}
-}
\ No newline at end of file
+}
diff --git a/Engine/Source/Runtime/Engine/Classes/Components/SkinnedMeshComponent.h b/Engine/Source/Runtime/Engine/Classes/Components/SkinnedMeshComponent.h
--- a/Engine/Source/Runtime/Engine/Classes/Components/SkinnedMeshComponent.h
+++ b/Engine/Source/Runtime/Engine/Classes/Components/SkinnedMeshComponent.h
@@ -132,18 +132,6 @@ struct FActiveMorphTarget
 	}
 };
 
-struct FSkeletalMeshObjectCallbackData
-{
-	enum class EEventType { Register, Unregister, Update };
-	typedef void (*TCallbackMeshObjectCallback)(
-		EEventType Event, 
-		class FSkeletalMeshObject* MeshObject,
-		uint64 UserData);
-
-	uint64 UserData = 0;
-	TCallbackMeshObjectCallback Run = nullptr;
-};
-
 /** Vertex skin weight info supplied for a component override. */
 USTRUCT(BlueprintType, meta = (HasNativeMake = "Engine.KismetRenderingLibrary.MakeSkinWeightInfo", HasNativeBreak = "Engine.KismetRenderingLibrary.BreakSkinWeightInfo"))
 struct FSkelMeshSkinWeightInfo
@@ -667,8 +655,6 @@ public:
 
 	/** Object responsible for sending bone transforms, morph target state etc. to render thread. */
 	class FSkeletalMeshObject*	MeshObject;
-	FSkeletalMeshObjectCallbackData MeshObjectCallbackData;
-	void SetMeshObjectCallbackData(FSkeletalMeshObjectCallbackData& MeshObjectCallbackData);
 
 	/** Gets the skeletal mesh resource used for rendering the component. */
 	FSkeletalMeshRenderData* GetSkeletalMeshRenderData() const;
diff --git a/Engine/Source/Runtime/Engine/Private/Components/SkinnedMeshComponent.cpp b/Engine/Source/Runtime/Engine/Private/Components/SkinnedMeshComponent.cpp
--- a/Engine/Source/Runtime/Engine/Private/Components/SkinnedMeshComponent.cpp
+++ b/Engine/Source/Runtime/Engine/Private/Components/SkinnedMeshComponent.cpp
@@ -778,15 +778,6 @@ void USkinnedMeshComponent::TickUpdateRate(float DeltaTime, bool bNeedsValidRoot
 	}
 }
 
-void USkinnedMeshComponent::SetMeshObjectCallbackData(FSkeletalMeshObjectCallbackData& InData)
-{
-	MeshObjectCallbackData = InData;
-	if (MeshObject)
-	{
-		MeshObject->SetCallbackData(InData);
-	}
-}
-
 void USkinnedMeshComponent::TickComponent(float DeltaTime, enum ELevelTick TickType, FActorComponentTickFunction *ThisTickFunction)
 {
 	SCOPED_NAMED_EVENT(USkinnedMeshComponent_TickComponent, FColor::Yellow);
@@ -795,12 +786,6 @@ void USkinnedMeshComponent::TickComponent(float DeltaTime, enum ELevelTick TickT
 	// Tick ActorComponent first.
 	Super::TickComponent(DeltaTime, TickType, ThisTickFunction);
 
-	// Refresh callback in case MeshObjectCallbackData was updated
-	if (MeshObject)
-	{
-		MeshObject->SetCallbackData(MeshObjectCallbackData);
-	}
-
 	// See if this mesh was rendered recently. This has to happen first because other data will rely on this
 	bRecentlyRendered = (GetLastRenderTime() > GetWorld()->TimeSeconds - 1.0f);
 
diff --git a/Engine/Source/Runtime/Engine/Private/GPUSkinCache.cpp b/Engine/Source/Runtime/Engine/Private/GPUSkinCache.cpp
--- a/Engine/Source/Runtime/Engine/Private/GPUSkinCache.cpp
+++ b/Engine/Source/Runtime/Engine/Private/GPUSkinCache.cpp
@@ -284,9 +284,9 @@ public:
 		DispatchData[Section].UpdateVertexFactoryDeclaration();
 	}
 
-	inline FCachedGeometrySection GetCachedGeometry(int32 SectionIndex) const
+	inline FCachedGeometry::Section GetCachedGeometry(int32 SectionIndex) const
 	{
-		FCachedGeometrySection MeshSection;
+		FCachedGeometry::Section MeshSection;
 		const FSkelMeshRenderSection& Section = *DispatchData[SectionIndex].Section;
 		MeshSection.PositionBuffer	= DispatchData[SectionIndex].PositionBuffer->SRV;
 		MeshSection.UVsBuffer		= DispatchData[SectionIndex].UVsBufferSRV;
@@ -1577,9 +1577,37 @@ void FGPUSkinCache::InvalidateAllEntries()
 	SET_MEMORY_STAT(STAT_GPUSkinCache_TangentsIntermediateMemUsed, 0);
 }
 
-FCachedGeometrySection FGPUSkinCache::GetCachedGeometry(FGPUSkinCacheEntry* InOutEntry, uint32 sectionIndex)
+FCachedGeometry FGPUSkinCache::GetCachedGeometry(uint32 ComponentId) const
 {
-	return InOutEntry ? InOutEntry->GetCachedGeometry(sectionIndex) : FCachedGeometrySection();
+	FCachedGeometry Out;
+	for (FGPUSkinCacheEntry* Entry : Entries)
+	{
+		if (Entry && Entry->GPUSkin && Entry->GPUSkin->GetComponentId() == ComponentId)
+		{
+			const uint32 LODIndex = Entry->GPUSkin->GetLOD();
+			const FSkeletalMeshRenderData& RenderData = Entry->GPUSkin->GetSkeletalMeshRenderData();
+			const FSkeletalMeshLODRenderData& LODData = RenderData.LODRenderData[LODIndex];
+			const uint32 SectionCount = LODData.RenderSections.Num();
+			for (uint32 SectionIdx=0; SectionIdx< SectionCount;++SectionIdx)
+			{
+				FCachedGeometry::Section CachedSection = Entry->GetCachedGeometry(SectionIdx);
+				CachedSection.IndexBuffer		= LODData.MultiSizeIndexContainer.GetIndexBuffer()->GetSRV();
+				CachedSection.TotalIndexCount	= LODData.MultiSizeIndexContainer.GetIndexBuffer()->Num();
+				CachedSection.LODIndex			= LODIndex;
+				CachedSection.UVsChannelOffset	= 0; // Assume that we needs to pair meshes based on UVs 0
+				CachedSection.UVsChannelCount	= LODData.StaticVertexBuffers.StaticMeshVertexBuffer.GetNumTexCoords();
+				Out.Sections.Add(CachedSection);
+			}
+			break;
+		}
+	}
+
+	return Out;
+}
+
+FCachedGeometry::Section FGPUSkinCache::GetCachedGeometry(FGPUSkinCacheEntry* InOutEntry, uint32 sectionIndex)
+{
+	return InOutEntry ? InOutEntry->GetCachedGeometry(sectionIndex) : FCachedGeometry::Section();
 }
 
 void FGPUSkinCache::UpdateSkinWeightBuffer(FGPUSkinCacheEntry* Entry)
diff --git a/Engine/Source/Runtime/Engine/Private/SkeletalRender.cpp b/Engine/Source/Runtime/Engine/Private/SkeletalRender.cpp
--- a/Engine/Source/Runtime/Engine/Private/SkeletalRender.cpp
+++ b/Engine/Source/Runtime/Engine/Private/SkeletalRender.cpp
@@ -45,6 +45,7 @@ FSkeletalMeshObject::FSkeletalMeshObject(USkinnedMeshComponent* InMeshComponent,
 ,	bUsePerBoneMotionBlur(InMeshComponent->bPerBoneMotionBlur)
 ,	StatId(InMeshComponent->SkeletalMesh->GetStatID(true))
 ,	FeatureLevel(InFeatureLevel)
+,	ComponentId(InMeshComponent->ComponentId.PrimIDValue)
 {
 	check(SkeletalMeshRenderData);
 
diff --git a/Engine/Source/Runtime/Engine/Private/SkeletalRenderGPUSkin.cpp b/Engine/Source/Runtime/Engine/Private/SkeletalRenderGPUSkin.cpp
--- a/Engine/Source/Runtime/Engine/Private/SkeletalRenderGPUSkin.cpp
+++ b/Engine/Source/Runtime/Engine/Private/SkeletalRenderGPUSkin.cpp
@@ -134,15 +134,6 @@ void FMorphVertexBuffer::ReleaseDynamicRHI()
 /*-----------------------------------------------------------------------------
 FSkeletalMeshObjectGPUSkin
 -----------------------------------------------------------------------------*/
-void FSkeletalMeshObjectGPUSkin::SetCallbackData(FSkeletalMeshObjectCallbackData& InMeshObjectCallbackData)
-{
-	CallbackData = InMeshObjectCallbackData;
-	if (CallbackData.Run)
-	{
-		CallbackData.Run(FSkeletalMeshObjectCallbackData::EEventType::Register, this, CallbackData.UserData);
-	}
-}
-
 FSkeletalMeshObjectGPUSkin::FSkeletalMeshObjectGPUSkin(USkinnedMeshComponent* InMeshComponent, FSkeletalMeshRenderData* InSkelMeshRenderData, ERHIFeatureLevel::Type InFeatureLevel)
 	: FSkeletalMeshObject(InMeshComponent, InSkelMeshRenderData, InFeatureLevel)
 	,	DynamicData(NULL)
@@ -159,7 +150,6 @@ FSkeletalMeshObjectGPUSkin::FSkeletalMeshObjectGPUSkin(USkinnedMeshComponent* In
 	}
 
 	InitResources(InMeshComponent);
-	SetCallbackData(InMeshComponent->MeshObjectCallbackData);
 }
 
 
@@ -202,11 +192,6 @@ void FSkeletalMeshObjectGPUSkin::InitResources(USkinnedMeshComponent* InMeshComp
 
 void FSkeletalMeshObjectGPUSkin::ReleaseResources()
 {
-	if (CallbackData.Run)
-	{
-		CallbackData.Run(FSkeletalMeshObjectCallbackData::EEventType::Unregister, this, CallbackData.UserData);
-	}
-
 	for( int32 LODIndex=0;LODIndex < LODs.Num();LODIndex++ )
 	{
 		FSkeletalMeshObjectLOD& SkelLOD = LODs[LODIndex];
@@ -322,11 +307,6 @@ void FSkeletalMeshObjectGPUSkin::Update(int32 LODIndex,USkinnedMeshComponent* In
 			MeshObject->UpdateDynamicData_RenderThread(GPUSkinCache, RHICmdList, NewDynamicData, nullptr, FrameNumberToPrepare, RevisionNumber);
 		}
 	);
-
-	if (CallbackData.Run)
-	{
-		CallbackData.Run(FSkeletalMeshObjectCallbackData::EEventType::Update, this, CallbackData.UserData);
-	}
 }
 
 void FSkeletalMeshObjectGPUSkin::UpdateSkinWeightBuffer(USkinnedMeshComponent* InMeshComponent)
@@ -581,7 +561,6 @@ void FSkeletalMeshObjectGPUSkin::ProcessUpdatedDynamicData(FGPUSkinCache* GPUSki
 		bDataPresent = VertexFactoryData.VertexFactories.Num() > 0;
 	}
 
-	CachedGeometry.Reset();
 	if (bDataPresent)
 	{
 		for (int32 SectionIdx = 0; SectionIdx < Sections.Num(); SectionIdx++)
@@ -671,14 +650,6 @@ void FSkeletalMeshObjectGPUSkin::ProcessUpdatedDynamicData(FGPUSkinCache* GPUSki
 				GPUSkinCache->ProcessEntry(RHICmdList, VertexFactory,
 					VertexFactoryData.PassthroughVertexFactories[SectionIdx].Get(), Section, this, bMorph ? &LOD.MorphVertexBuffer : 0, bClothFactory ? &LODData.ClothVertexBuffer : 0,
 					bClothFactory ? DynamicData->ClothingSimData.Find(Section.CorrespondClothAssetIndex) : 0, LocalToCloth, DynamicData->ClothBlendWeight, RevisionNumber, SectionIdx, SkinCacheEntry);
-
-				FCachedGeometrySection CachedSection = GPUSkinCache->GetCachedGeometry(SkinCacheEntry, SectionIdx);
-				CachedSection.IndexBuffer		= LODData.MultiSizeIndexContainer.GetIndexBuffer()->GetSRV();
-				CachedSection.TotalIndexCount	= LODData.MultiSizeIndexContainer.GetIndexBuffer()->Num();
-				CachedSection.LODIndex			= DynamicData->LODIndex;
-				CachedSection.UVsChannelOffset	= 0; // Assume that we needs to pair meshes based on UVs 0
-				CachedSection.UVsChannelCount	= LODData.StaticVertexBuffers.StaticMeshVertexBuffer.GetNumTexCoords();
-				CachedGeometry.Sections.Add(CachedSection);
 			}
 
 			if (bNeedFence)
diff --git a/Engine/Source/Runtime/Engine/Private/SkeletalRenderGPUSkin.h b/Engine/Source/Runtime/Engine/Private/SkeletalRenderGPUSkin.h
--- a/Engine/Source/Runtime/Engine/Private/SkeletalRenderGPUSkin.h
+++ b/Engine/Source/Runtime/Engine/Private/SkeletalRenderGPUSkin.h
@@ -262,8 +262,6 @@ public:
 	virtual bool IsCPUSkinned() const override { return false; }
 	virtual TArray<FTransform>* GetComponentSpaceTransforms() const override;
 	virtual const TArray<FMatrix>& GetReferenceToLocalMatrices() const override;
-	virtual void SetCallbackData(FSkeletalMeshObjectCallbackData& CallbackData) override;
-	virtual FCachedGeometry GetCachedGeometry() const override { return CachedGeometry; };
 
 #if RHI_RAYTRACING
 	/** Geometry for ray tracing. */
@@ -534,12 +532,6 @@ private:
 
 	/** last updated bone transform revision number */
 	uint32 LastBoneTransformRevisionNumber;
-
-	/** Last cached geometry description */
-	FCachedGeometry CachedGeometry;
-
-	/** Callback for register/unregister/update events */
-	FSkeletalMeshObjectCallbackData CallbackData;
 };
 
 
diff --git a/Engine/Source/Runtime/Engine/Public/GPUSkinCache.h b/Engine/Source/Runtime/Engine/Public/GPUSkinCache.h
--- a/Engine/Source/Runtime/Engine/Public/GPUSkinCache.h
+++ b/Engine/Source/Runtime/Engine/Public/GPUSkinCache.h
@@ -85,6 +85,28 @@ struct FGPUSkinBatchElementUserData
 	int32 Section;
 };
 
+struct FCachedGeometry
+{
+	struct Section
+	{
+		FRHIShaderResourceView* PositionBuffer = nullptr;
+		FRHIShaderResourceView* UVsBuffer = nullptr;
+		FRHIShaderResourceView* IndexBuffer = nullptr;
+		uint32 UVsChannelOffset = 0;
+		uint32 UVsChannelCount = 0;
+		uint32 NumPrimitives = 0;
+		uint32 VertexBaseIndex = 0;
+		uint32 IndexBaseIndex = 0;
+		uint32 TotalVertexCount = 0;
+		uint32 TotalIndexCount = 0;
+		uint32 SectionIndex = 0;
+		int32 LODIndex = 0;
+	};
+
+	int32 LODIndex = 0;
+	TArray<Section> Sections;
+};
+
 class FGPUSkinCache
 {
 public:
@@ -104,7 +126,8 @@ public:
 	ENGINE_API FGPUSkinCache(bool bInRequiresMemoryLimit);
 	ENGINE_API ~FGPUSkinCache();
 
-	struct FCachedGeometrySection GetCachedGeometry(FGPUSkinCacheEntry* InOutEntry, uint32 SectionId);
+	ENGINE_API FCachedGeometry GetCachedGeometry(uint32 ComponentId) const;
+	FCachedGeometry::Section GetCachedGeometry(FGPUSkinCacheEntry* InOutEntry, uint32 SectionId);
 	void UpdateSkinWeightBuffer(FGPUSkinCacheEntry* Entry);
 
 	void ProcessEntry(FRHICommandListImmediate& RHICmdList, FGPUBaseSkinVertexFactory* VertexFactory,
diff --git a/Engine/Source/Runtime/Engine/Public/SceneView.h b/Engine/Source/Runtime/Engine/Public/SceneView.h
--- a/Engine/Source/Runtime/Engine/Public/SceneView.h
+++ b/Engine/Source/Runtime/Engine/Public/SceneView.h
@@ -731,6 +731,7 @@ enum ETranslucencyVolumeCascade
 	VIEW_UNIFORM_BUFFER_MEMBER(int32, FarShadowStaticMeshLODBias) \
 	VIEW_UNIFORM_BUFFER_MEMBER(float, MinRoughness) \
 	VIEW_UNIFORM_BUFFER_MEMBER(FVector4, HairRenderInfo) \
+	VIEW_UNIFORM_BUFFER_MEMBER(uint32, HairRenderInfoBits) \
 
 #define VIEW_UNIFORM_BUFFER_MEMBER(type, identifier) \
 	SHADER_PARAMETER(type, identifier)
diff --git a/Engine/Source/Runtime/Engine/Public/SkeletalRenderPublic.h b/Engine/Source/Runtime/Engine/Public/SkeletalRenderPublic.h
--- a/Engine/Source/Runtime/Engine/Public/SkeletalRenderPublic.h
+++ b/Engine/Source/Runtime/Engine/Public/SkeletalRenderPublic.h
@@ -43,34 +43,6 @@ enum class EPreviousBoneTransformUpdateMode
 	DuplicateCurrentToPrevious,
 };
 
-struct FCachedGeometrySection
-{
-	FRHIShaderResourceView* PositionBuffer = nullptr;
-	FRHIShaderResourceView* UVsBuffer = nullptr;
-	FRHIShaderResourceView* IndexBuffer = nullptr;
-	uint32 UVsChannelOffset = 0;
-	uint32 UVsChannelCount = 0;
-	uint32 NumPrimitives = 0;
-	uint32 VertexBaseIndex = 0;
-	uint32 IndexBaseIndex = 0;
-	uint32 TotalVertexCount = 0;
-	uint32 TotalIndexCount = 0;
-	uint32 SectionIndex = 0;
-	int32  LODIndex = 0;
-};
-
-struct FCachedGeometry
-{
-	int32 LODIndex = 0;
-	TArray<FCachedGeometrySection> Sections;
-
-	void Reset()
-	{
-		LODIndex = 0;
-		Sections.Empty();
-	}
-};
-
 /**
 * Interface for mesh rendering data
 */
@@ -217,17 +189,16 @@ public:
 	 */
 	void InitLODInfos(const USkinnedMeshComponent* InMeshComponent);
 
-	FORCEINLINE TStatId GetStatId() const 
-	{ 
-		return StatId; 
-	}
+	/**
+	 * Return the ID of the component to which the skeletal mesh object belongs to.
+	 */
+	FORCEINLINE uint32 GetComponentId() const { return ComponentId; }
+
+	FORCEINLINE TStatId GetStatId() const { return StatId; }
 
 	/** Get the skeletal mesh resource for which this mesh object was created. */
 	FORCEINLINE FSkeletalMeshRenderData& GetSkeletalMeshRenderData() const { return *SkeletalMeshRenderData; }
 
-	/** Set callback data for register/unregister/update events  */
-	virtual void SetCallbackData(struct FSkeletalMeshObjectCallbackData& CallbackData) {};
-
 #if RHI_RAYTRACING
 	/** Retrieve ray tracing geometry from the underlying mesh object */
 	virtual FRayTracingGeometry* GetRayTracingGeometry() { return nullptr; }
@@ -235,8 +206,6 @@ public:
 	virtual FRWBuffer* GetRayTracingDynamicVertexBuffer() { return nullptr; }
 #endif // RHI_RAYTRACING
 
-	virtual FCachedGeometry GetCachedGeometry() const { return FCachedGeometry(); };
-
 	/** Called to notify clothing data that component transform has changed */
 	virtual void RefreshClothingTransforms(const FMatrix& InNewLocalToWorld, uint32 FrameNumber) {};
 
@@ -322,4 +291,7 @@ protected:
 
 	/** Feature level to render for. */
 	ERHIFeatureLevel::Type FeatureLevel;
+
+	/** Component ID to which belong this  mesh object  */
+	uint32 ComponentId;
 };
diff --git a/Engine/Source/Runtime/Renderer/Private/DeferredShadingRenderer.cpp b/Engine/Source/Runtime/Renderer/Private/DeferredShadingRenderer.cpp
--- a/Engine/Source/Runtime/Renderer/Private/DeferredShadingRenderer.cpp
+++ b/Engine/Source/Runtime/Renderer/Private/DeferredShadingRenderer.cpp
@@ -1653,10 +1653,10 @@ void FDeferredShadingSceneRenderer::Render(FRHICommandListImmediate& RHICmdList)
 
 	RunGPUSkinCacheTransition(RHICmdList, Scene, EGPUSkinCacheTransition::Renderer);
 
-	if (HasHairStrandsProjectionQuery(Scene->GetShaderPlatform()))
+	if (HasHairStrandsProcess(Scene->GetShaderPlatform()))
 	{
-		FGlobalShaderMap* ShaderMap = GetGlobalShaderMap(FeatureLevel);
-		RunHairStrandsBindingQueries(RHICmdList, ShaderMap);
+		auto ShaderMap = GetGlobalShaderMap(FeatureLevel);
+		RunHairStrandsProcess(RHICmdList, ShaderMap);
 	}
 
 	// Interpolation needs to happen after the skin cache run as there is a dependency 
@@ -1668,7 +1668,8 @@ void FDeferredShadingSceneRenderer::Render(FRHICommandListImmediate& RHICmdList)
 		const EWorldType::Type WorldType = Views[0].Family->Scene->GetWorld()->WorldType;
 		auto ShaderMap = GetGlobalShaderMap(FeatureLevel);
 
-		RunHairStrandsInterpolation(RHICmdList, WorldType, &Views[0].ShaderDrawData, ShaderMap, EHairStrandsInterpolationType::RenderStrands, &HairClusterData); // Send data to full up with culling
+		FGPUSkinCache* GPUSkinCache = Scene->GetGPUSkinCache();
+		RunHairStrandsInterpolation(RHICmdList, WorldType, GPUSkinCache, &Views[0].ShaderDrawData, ShaderMap, EHairStrandsInterpolationType::RenderStrands, &HairClusterData); // Send data to full up with culling
 	}
 
 	// Before starting the render, all async task for the Custom data must be completed
@@ -2345,6 +2346,10 @@ void FDeferredShadingSceneRenderer::Render(FRHICommandListImmediate& RHICmdList)
 
 		ComputeSubsurfaceShim(RHICmdList, Views);
 
+		if (HairDatas)
+		{
+			RenderHairStrandsSceneColorScattering(RHICmdList, Views, HairDatas);
+		}
 #if RHI_RAYTRACING
 		if (SkyLightRT)
 		{
diff --git a/Engine/Source/Runtime/Renderer/Private/DeferredShadingRenderer.h b/Engine/Source/Runtime/Renderer/Private/DeferredShadingRenderer.h
--- a/Engine/Source/Runtime/Renderer/Private/DeferredShadingRenderer.h
+++ b/Engine/Source/Runtime/Renderer/Private/DeferredShadingRenderer.h
@@ -529,8 +529,7 @@ private:
 		const IScreenSpaceDenoiser::FShadowRayTracingConfig& RayTracingConfig,
 		const IScreenSpaceDenoiser::EShadowRequirements DenoiserRequirements,
 		const bool bSubPixelShadowMask,
-		FRDGTextureRef HairCategorizationTexture,
-		FRDGTextureRef HairLightChannelMaskTexture,
+		const struct FHairStrandsOcclusionResources* HairResources,
 		FRDGTextureUAV* OutShadowMaskUAV,
 		FRDGTextureUAV* OutRayHitDistanceUAV);
 
diff --git a/Engine/Source/Runtime/Renderer/Private/HairStrands/HairStrandsCluster.cpp b/Engine/Source/Runtime/Renderer/Private/HairStrands/HairStrandsCluster.cpp
--- a/Engine/Source/Runtime/Renderer/Private/HairStrands/HairStrandsCluster.cpp
+++ b/Engine/Source/Runtime/Renderer/Private/HairStrands/HairStrandsCluster.cpp
@@ -113,12 +113,12 @@ static void AddHairMacroGroupAABBPass(
 	}
 }
 
-static bool DoesGroupExists(uint32 ResourceId, FHairStrandsMacroGroupData::TPrimitiveGroups& PrimitivesGroups)
+static bool DoesGroupExists(uint32 ResourceId, uint32 GroupIndex, const FHairStrandsMacroGroupData::TPrimitiveGroups& PrimitivesGroups)
 {
 	// Simple linear search as the expected number of groups is supposed to be low (<10)
-	for (FHairStrandsMacroGroupData::PrimitiveGroup& Group : PrimitivesGroups)
+	for (const FHairStrandsMacroGroupData::PrimitiveGroup& Group : PrimitivesGroups)
 	{
-		if (Group.ResourceId == ResourceId)
+		if (Group.GroupIndex == GroupIndex && Group.ResourceId == ResourceId)
 		{
 			return true;
 		}
@@ -179,7 +179,7 @@ FHairStrandsMacroGroupViews CreateHairStrandsMacroGroups(
 							PrimitiveInfo.GroupIndex = HairGroupPublicData->GetGroupIndex();
 							check(PrimitiveInfo.GroupIndex < 32); // Sanity check
 
-							const bool bAlreadyExists = DoesGroupExists(PrimitiveInfo.ResourceId, MacroGroup.PrimitivesGroups);
+							const bool bAlreadyExists = DoesGroupExists(PrimitiveInfo.ResourceId, PrimitiveInfo.GroupIndex, MacroGroup.PrimitivesGroups);
 							if (!bAlreadyExists)
 							{
 								FHairStrandsMacroGroupData::PrimitiveGroup& PrimitiveGroup = MacroGroup.PrimitivesGroups.AddZeroed_GetRef();
@@ -210,7 +210,7 @@ FHairStrandsMacroGroupViews CreateHairStrandsMacroGroups(
 						PrimitiveInfo.GroupIndex = HairGroupPublicData->GetGroupIndex();
 						check(PrimitiveInfo.GroupIndex < 32); // Sanity check
 
-						const bool bAlreadyExists = DoesGroupExists(PrimitiveInfo.ResourceId, MacroGroup.PrimitivesGroups);
+						const bool bAlreadyExists = DoesGroupExists(PrimitiveInfo.ResourceId, PrimitiveInfo.GroupIndex, MacroGroup.PrimitivesGroups);
 						if (!bAlreadyExists)
 						{
 							FHairStrandsMacroGroupData::PrimitiveGroup& PrimitiveGroup = MacroGroup.PrimitivesGroups.AddZeroed_GetRef();
@@ -256,3 +256,9 @@ FHairStrandsMacroGroupViews CreateHairStrandsMacroGroups(
 
 	return MacroGroupsViews;
 }
+
+bool FHairStrandsMacroGroupData::PrimitiveInfo::IsCullingEnable() const
+{
+	const FHairGroupPublicData* HairGroupPublicData = reinterpret_cast<const FHairGroupPublicData*>(MeshBatchAndRelevance.Mesh->Elements[0].VertexFactoryUserData);
+	return HairGroupPublicData->GetCullingResultAvailable();
+}
\ No newline at end of file
diff --git a/Engine/Source/Runtime/Renderer/Private/HairStrands/HairStrandsCluster.h b/Engine/Source/Runtime/Renderer/Private/HairStrands/HairStrandsCluster.h
--- a/Engine/Source/Runtime/Renderer/Private/HairStrands/HairStrandsCluster.h
+++ b/Engine/Source/Runtime/Renderer/Private/HairStrands/HairStrandsCluster.h
@@ -37,7 +37,8 @@ struct FHairStrandsMacroGroupData
 		FMeshBatchAndRelevance MeshBatchAndRelevance;
 		uint32 MaterialId;
 		uint32 ResourceId;
-		uint32 GroupIndex;
+		uint32 GroupIndex;		
+		bool IsCullingEnable() const;
 	};
 	typedef TArray<PrimitiveInfo, SceneRenderingAllocator> TPrimitiveInfos;
 
diff --git a/Engine/Source/Runtime/Renderer/Private/HairStrands/HairStrandsComposition.cpp b/Engine/Source/Runtime/Renderer/Private/HairStrands/HairStrandsComposition.cpp
--- a/Engine/Source/Runtime/Renderer/Private/HairStrands/HairStrandsComposition.cpp
+++ b/Engine/Source/Runtime/Renderer/Private/HairStrands/HairStrandsComposition.cpp
@@ -32,6 +32,11 @@ class FHairVisibilityComposeSubPixelPS : public FGlobalShader
 
 public:
 	static bool ShouldCompilePermutation(const FGlobalShaderPermutationParameters& Parameters) { return Parameters.Platform == SP_PCD3D_SM5; }
+	static void ModifyCompilationEnvironment(const FGlobalShaderPermutationParameters& Parameters, FShaderCompilerEnvironment& OutEnvironment)
+	{
+		FGlobalShader::ModifyCompilationEnvironment(Parameters, OutEnvironment);
+		OutEnvironment.SetDefine(TEXT("SHADER_SUBCOLOR"), 1);
+	}
 };
 
 IMPLEMENT_GLOBAL_SHADER(FHairVisibilityComposeSubPixelPS, "/Engine/Private/HairStrands/HairStrandsVisibilityComposeSubPixelPS.usf", "SubColorPS", SF_Pixel);
@@ -49,7 +54,7 @@ static void AddHairVisibilityComposeSubPixelPass(
 	Parameters->CategorisationTexture = CategorisationTexture;
 	Parameters->RenderTargets[0] = FRenderTargetBinding(OutColorTexture, ERenderTargetLoadAction::ELoad);
 	Parameters->RenderTargets.DepthStencil = FDepthStencilBinding(OutDepthTexture, ERenderTargetLoadAction::ELoad, ERenderTargetLoadAction::ENoAction, FExclusiveDepthStencil::DepthWrite_StencilNop);
-	
+
 	FHairVisibilityComposeSubPixelPS::FPermutationDomain PermutationVector;
 	TShaderMapRef<FHairVisibilityComposeSubPixelPS> PixelShader(View.ShaderMap, PermutationVector);
 	TShaderMapRef<FPostProcessVS> VertexShader(View.ShaderMap);
@@ -102,6 +107,9 @@ class FHairVisibilityFastResolvePS : public FGlobalShader
 	DECLARE_GLOBAL_SHADER(FHairVisibilityFastResolvePS);
 	SHADER_USE_PARAMETER_STRUCT(FHairVisibilityFastResolvePS, FGlobalShader);
 
+	class FMSAACount : SHADER_PERMUTATION_SPARSE_INT("PERMUTATION_MSAACOUNT", 4, 8);
+	using FPermutationDomain = TShaderPermutationDomain<FMSAACount>;
+
 	BEGIN_SHADER_PARAMETER_STRUCT(FParameters, )
 		SHADER_PARAMETER(float, VelocityThreshold)
 		SHADER_PARAMETER_RDG_TEXTURE(Texture2D, HairVisibilityVelocityTexture)
@@ -110,6 +118,11 @@ class FHairVisibilityFastResolvePS : public FGlobalShader
 
 public:
 	static bool ShouldCompilePermutation(const FGlobalShaderPermutationParameters& Parameters) { return Parameters.Platform == SP_PCD3D_SM5; }
+	static void ModifyCompilationEnvironment(const FGlobalShaderPermutationParameters& Parameters, FShaderCompilerEnvironment& OutEnvironment)
+	{
+		FGlobalShader::ModifyCompilationEnvironment(Parameters, OutEnvironment);
+		OutEnvironment.SetDefine(TEXT("SHADER_FASTRESOLVE"), 1);
+	}
 };
 
 IMPLEMENT_GLOBAL_SHADER(FHairVisibilityFastResolvePS, "/Engine/Private/HairStrands/HairStrandsVisibilityComposeSubPixelPS.usf", "FastResolvePS", SF_Pixel);
@@ -148,8 +161,13 @@ static void AddHairVisibilityFastResolvePass(
 		ERenderTargetLoadAction::ELoad,
 		FExclusiveDepthStencil::DepthNop_StencilWrite);
 
+	const uint32 MSAASampleCount = HairVisibilityVelocityTexture->Desc.NumSamples;
+	check(MSAASampleCount == 4 || MSAASampleCount == 8);
+	FHairVisibilityFastResolvePS::FPermutationDomain PermutationVector;
+	PermutationVector.Set<FHairVisibilityFastResolvePS::FMSAACount>(MSAASampleCount == 4 ? 4 : 8);
+
 	TShaderMapRef<FPostProcessVS> VertexShader(View.ShaderMap);
-	TShaderMapRef<FHairVisibilityFastResolvePS> PixelShader(View.ShaderMap);
+	TShaderMapRef<FHairVisibilityFastResolvePS> PixelShader(View.ShaderMap, PermutationVector);
 	const FGlobalShaderMap* GlobalShaderMap = View.ShaderMap;
 	const FIntRect Viewport = View.ViewRect;
 	
diff --git a/Engine/Source/Runtime/Renderer/Private/HairStrands/HairStrandsDebug.cpp b/Engine/Source/Runtime/Renderer/Private/HairStrands/HairStrandsDebug.cpp
--- a/Engine/Source/Runtime/Renderer/Private/HairStrands/HairStrandsDebug.cpp
+++ b/Engine/Source/Runtime/Renderer/Private/HairStrands/HairStrandsDebug.cpp
@@ -87,8 +87,8 @@ static FAutoConsoleVariableRef CVarHairVirtualVoxel_DrawDebugPage(TEXT("r.HairSt
 
 // Helper functions for accessing interpolation data for debug purpose.
 // Definitions is in HairStrandsInterface.cpp
-void GetGroomInterpolationData(const EWorldType::Type WorldType, const EHairStrandsProjectionMeshType MeshtType, FHairStrandsProjectionMeshData::LOD& OutGeometries);
-void GetGroomInterpolationData(const EWorldType::Type WorldType, const bool bRenderData, FHairStrandsProjectionHairData& OutHairData, TArray<int32>& OutLODIndices);
+void GetGroomInterpolationData(const EWorldType::Type WorldType, const EHairStrandsProjectionMeshType MeshtType, const FGPUSkinCache* SkinCache, FHairStrandsProjectionMeshData::LOD& OutGeometries);
+void GetGroomInterpolationData(const EWorldType::Type WorldType, const EHairStrandsInterpolationType StrandType, const FGPUSkinCache* SkinCache, FHairStrandsProjectionHairData& OutHairData, TArray<int32>& OutLODIndices);
 
 static int32 GHairStrandsCull = 0;
 static int32 GHairStrandsCullIndex = -1;
@@ -241,6 +241,8 @@ class FHairDebugPrintCS : public FGlobalShader
 		SHADER_PARAMETER(FIntPoint, MaxResolution)
 		SHADER_PARAMETER(uint32, FastResolveMask)
 		SHADER_PARAMETER(uint32, HairMacroGroupCount)
+		SHADER_PARAMETER(uint32, MaxSampleCount)
+		SHADER_PARAMETER_RDG_TEXTURE(Texture2D, HairCountTexture)
 		SHADER_PARAMETER_RDG_TEXTURE(Texture2D, CategorizationTexture)
 		SHADER_PARAMETER_RDG_TEXTURE(Texture2D, HairVisibilityNodeOffsetAndCount)
 		SHADER_PARAMETER_RDG_BUFFER_SRV(StructuredBuffer, HairVisibilityNodeData)
@@ -266,12 +268,14 @@ static void AddDebugHairPrintPass(
 	const FHairStrandsMacroGroupDatas& MacroGroupDatas,
 	const FShaderResourceViewRHIRef& InDepthStencilTexture)
 {
+	const TRefCountPtr<IPooledRenderTarget>& InViewHairCountTexture = VisibilityData.ViewHairCountTexture;
 	const TRefCountPtr<IPooledRenderTarget>& InCategorizationTexture = VisibilityData.CategorizationTexture;
 	const TRefCountPtr<IPooledRenderTarget>& InNodeIndex = VisibilityData.NodeIndex;
 	const TRefCountPtr<FPooledRDGBuffer>& InNodeData = VisibilityData.NodeData;
 	
 	if (!InCategorizationTexture || !InNodeIndex || !InNodeData || !InDepthStencilTexture) return;
 
+	FRDGTextureRef ViewHairCountTexture = InViewHairCountTexture ? GraphBuilder.RegisterExternalTexture(InViewHairCountTexture, TEXT("ViewHairCountTexture")) : GSystemTextures.GetBlackDummy(GraphBuilder);
 	FRDGTextureRef CategorizationTexture = InCategorizationTexture ? GraphBuilder.RegisterExternalTexture(InCategorizationTexture, TEXT("CategorizationTexture")) : nullptr;
 	FRDGTextureRef NodeIndex = InNodeIndex ? GraphBuilder.RegisterExternalTexture(InNodeIndex, TEXT("NodeIndex")) : nullptr;
 	FRDGBufferRef  NodeData  = InNodeData ? GraphBuilder.RegisterExternalBuffer(InNodeData, TEXT("NodeData")) : nullptr;
@@ -286,11 +290,13 @@ static void AddDebugHairPrintPass(
 	Parameters->PixelCoord = View->CursorPos;
 	Parameters->FastResolveMask = STENCIL_TEMPORAL_RESPONSIVE_AA_MASK;
 	Parameters->CategorizationTexture = CategorizationTexture;
+	Parameters->HairCountTexture = ViewHairCountTexture;
 	Parameters->HairVisibilityNodeData = GraphBuilder.CreateSRV(NodeData);
 	Parameters->HairVisibilityNodeOffsetAndCount = NodeIndex;
 	Parameters->DepthStencilTexture = InDepthStencilTexture;
 	Parameters->LinearSampler = TStaticSamplerState<SF_Bilinear, AM_Clamp, AM_Clamp, AM_Clamp>::GetRHI();
 	Parameters->HairMacroGroupCount = MacroGroupDatas.Datas.Num();
+	Parameters->MaxSampleCount = VisibilityData.MaxSampleCount;
 	Parameters->HairMacroGroupAABBBuffer = GraphBuilder.CreateSRV(HairMacroGroupAABBBuffer, PF_R32_SINT);
 	ShaderPrint::SetParameters(*View, Parameters->ShaderPrintUniformBuffer);
 	ShaderDrawDebug::SetParameters(GraphBuilder, View->ShaderDrawData, Parameters->ShaderDrawUniformBuffer);
@@ -318,6 +324,7 @@ class FHairDebugPS : public FGlobalShader
 		SHADER_PARAMETER(uint32, DebugMode)
 		SHADER_PARAMETER(int32, SampleIndex)
 		SHADER_PARAMETER(uint32, TileSize)
+		SHADER_PARAMETER(uint32, MaxSampleCount)
 		SHADER_PARAMETER_RDG_TEXTURE(Texture2D, CategorizationTexture)
 		SHADER_PARAMETER_RDG_TEXTURE(Texture2D, NodeIndex)
 		SHADER_PARAMETER_RDG_TEXTURE(Texture2D, TileIndexTexture)
@@ -338,11 +345,7 @@ static void AddDebugHairPass(
 	FRDGBuilder& GraphBuilder,
 	const FViewInfo* View,
 	const EHairDebugMode InDebugMode,
-	const uint32 InTileSize,
-	const TRefCountPtr<IPooledRenderTarget>& InCategorizationTexture,
-	const TRefCountPtr<IPooledRenderTarget>& InNodeIndex,
-	const TRefCountPtr<FPooledRDGBuffer>& InNodeData,
-	const TRefCountPtr<IPooledRenderTarget>& InTileIndexTexture,
+	const FHairStrandsVisibilityData& VisibilityData,
 	const FShaderResourceViewRHIRef& InDepthStencilTexture,
 	FRDGTextureRef& OutTarget)
 {
@@ -358,6 +361,11 @@ static void AddDebugHairPass(
 		InDebugMode == EHairDebugMode::MaterialTangent ||
 		InDebugMode == EHairDebugMode::Tile);
 
+	const TRefCountPtr<IPooledRenderTarget>& InCategorizationTexture = VisibilityData.CategorizationTexture;
+	const TRefCountPtr<IPooledRenderTarget>& InNodeIndex = VisibilityData.NodeIndex;
+	const TRefCountPtr<FPooledRDGBuffer>& InNodeData = VisibilityData.NodeData;
+	const TRefCountPtr<IPooledRenderTarget>& InTileIndexTexture = VisibilityData.TileIndexTexture;
+
 	if (!InCategorizationTexture || !InNodeIndex || !InNodeData || !InTileIndexTexture) return;
 	if (InDebugMode == EHairDebugMode::TAAResolveType && !InDepthStencilTexture) return;
 
@@ -392,11 +400,12 @@ static void AddDebugHairPass(
 	Parameters->NodeIndex = NodeIndex;
 	Parameters->NodeData = GraphBuilder.CreateSRV(NodeData);
 	Parameters->TileIndexTexture = TileIndexTexture;
-	Parameters->TileSize = InTileSize;
+	Parameters->TileSize = VisibilityData.TileSize;
 	Parameters->DepthStencilTexture = InDepthStencilTexture;
 	Parameters->LinearSampler = TStaticSamplerState<SF_Bilinear, AM_Clamp, AM_Clamp, AM_Clamp>::GetRHI();
 	Parameters->DebugMode = InternalDebugMode;
 	Parameters->SampleIndex = GHairStrandsDebugSampleIndex;
+	Parameters->MaxSampleCount = VisibilityData.MaxSampleCount;
 	Parameters->RenderTargets[0] = FRenderTargetBinding(OutTarget, ERenderTargetLoadAction::ELoad, 0);
 	TShaderMapRef<FPostProcessVS> VertexShader(View->ShaderMap);
 
@@ -1383,17 +1392,19 @@ void RenderHairStrandsDebugInfo(
 			check(ViewFamily.Scene && ViewFamily.Scene->GetWorld());
 			const bool bIsActive = DebugInfo.WorldType == ViewFamily.Scene->GetWorld()->WorldType;
 
-			Line = FString::Printf(TEXT(" * Id:%d | WorldType:%s | Group count : %d"), DebugInfo.Id, ToString(DebugInfo.WorldType), DebugInfo.HairGroups.Num());
+			Line = FString::Printf(TEXT(" * Id:%d | WorldType:%s | Group count : %d"), DebugInfo.ComponentId, ToString(DebugInfo.WorldType), DebugInfo.HairGroups.Num());
 			Canvas.DrawShadowedString(X, Y += YStep, *Line, GetStatsFont(), bIsActive ? DebugColor : InactiveColor);
 
 			for (const FHairStrandsDebugInfo::HairGroup& DebugHairGroup : DebugInfo.HairGroups)
 			{
-				Line = FString::Printf(TEXT("        |> CurveCount : %d | VertexCount : %d | MaxRadius : %f | MaxLength : %f | Skinned: %s | LOD count : %d"),
+				Line = FString::Printf(TEXT("        |> CurveCount : %d | VertexCount : %d | MaxRadius : %f | MaxLength : %f | Skinned: %s | Binding: %s | Simulation: %s| LOD count : %d"),
 					DebugHairGroup.CurveCount,
 					DebugHairGroup.VertexCount,
 					DebugHairGroup.MaxRadius,
 					DebugHairGroup.MaxLength,
 					DebugHairGroup.bHasSkinInterpolation ? TEXT("True") : TEXT("False"),
+					DebugHairGroup.bHasBinding ? TEXT("True") : TEXT("False"),
+					DebugHairGroup.bHasSimulation ? TEXT("True") : TEXT("False"),
 					DebugHairGroup.LODCount);
 				Canvas.DrawShadowedString(X, Y += YStep, *Line, GetStatsFont(), bIsActive ? DebugColor : InactiveColor);
 			}
@@ -1575,7 +1586,7 @@ void RenderHairStrandsDebugInfo(
 		{
 			const FHairStrandsVisibilityData& VisibilityData = HairDatas->HairVisibilityViews.HairDatas[ViewIndex];
 			const FHairStrandsMacroGroupDatas& MacroGroupDatas = InMacroGroupViews.Views[ViewIndex];
-			AddDebugHairPass(GraphBuilder, &View, HairDebugMode, VisibilityData.TileSize, VisibilityData.CategorizationTexture, VisibilityData.NodeIndex, VisibilityData.NodeData, VisibilityData.TileIndexTexture, SceneTargets.SceneStencilSRV, SceneColorTexture);
+			AddDebugHairPass(GraphBuilder, &View, HairDebugMode, VisibilityData, SceneTargets.SceneStencilSRV, SceneColorTexture);
 			AddDebugHairPrintPass(GraphBuilder, &View, HairDebugMode, VisibilityData, MacroGroupDatas, SceneTargets.SceneStencilSRV);
 		}
 
@@ -1630,16 +1641,17 @@ void RenderHairStrandsDebugInfo(
 				DepthTexture = GraphBuilder.CreateTexture(Desc, TEXT("HairInterpolationDepthTexture"));
 			}
 
+			const FGPUSkinCache* SkinCache = View.Family->Scene->GetGPUSkinCache();
 			if (GHairDebugMeshProjection_SkinCacheMesh > 0)
 			{
 				FViewInfo* LocalView = &View;
-				auto RenderMeshProjection = [&bClearDepth, WorldType, LocalView, &SceneColorTexture, &DepthTexture](FRDGBuilder& InGraphBuilder, EHairStrandsProjectionMeshType MeshType)
+				auto RenderMeshProjection = [&bClearDepth, WorldType, LocalView, SkinCache, &SceneColorTexture, &DepthTexture](FRDGBuilder& GraphBuilder, EHairStrandsProjectionMeshType MeshType)
 				{
 					FHairStrandsProjectionMeshData::LOD MeshProjectionLODData;
-					GetGroomInterpolationData(WorldType, MeshType, MeshProjectionLODData);
+					GetGroomInterpolationData(WorldType, MeshType, SkinCache, MeshProjectionLODData);
 					for (FHairStrandsProjectionMeshData::Section& Section : MeshProjectionLODData.Sections)
 					{
-						AddDebugProjectionMeshPass(InGraphBuilder, LocalView, MeshType, bClearDepth, Section, SceneColorTexture, DepthTexture);
+						AddDebugProjectionMeshPass(GraphBuilder, LocalView, MeshType, bClearDepth, Section, SceneColorTexture, DepthTexture);
 						bClearDepth = false;
 					}
 				};
@@ -1651,11 +1663,11 @@ void RenderHairStrandsDebugInfo(
 			}
 
 			FViewInfo* LocalView = &View;
-			auto RenderProjectionData = [&GraphBuilder, LocalView, WorldType, &bClearDepth, SceneColorTexture, DepthTexture](bool bRenderData, bool bRestTriangle, bool bRestFrame, bool bDeformedTriangle, bool bDeformedFrame)
+			auto RenderProjectionData = [&GraphBuilder, LocalView, WorldType, SkinCache, &bClearDepth, SceneColorTexture, DepthTexture](EHairStrandsInterpolationType StrandType, bool bRestTriangle, bool bRestFrame, bool bDeformedTriangle, bool bDeformedFrame)
 			{
 				FHairStrandsProjectionHairData HairProjectionDatas;
 				TArray<int32> HairLODIndices;
-				GetGroomInterpolationData(WorldType, bRenderData, HairProjectionDatas, HairLODIndices);
+				GetGroomInterpolationData(WorldType, StrandType, SkinCache, HairProjectionDatas, HairLODIndices);
 				check(HairProjectionDatas.HairGroups.Num() == HairLODIndices.Num());
 				for (int32 HairIndex=0; HairIndex < HairProjectionDatas.HairGroups.Num(); ++HairIndex)
 				{
@@ -1691,7 +1703,7 @@ void RenderHairStrandsDebugInfo(
 				GHairDebugMeshProjection_Render_HairDeformedFrames > 0)
 			{
 				RenderProjectionData(
-					true, 
+					EHairStrandsInterpolationType::RenderStrands, 
 					GHairDebugMeshProjection_Render_HairRestTriangles > 0, 
 					GHairDebugMeshProjection_Render_HairRestFrames > 0, 
 					GHairDebugMeshProjection_Render_HairDeformedTriangles > 0, 
@@ -1704,7 +1716,7 @@ void RenderHairStrandsDebugInfo(
 				GHairDebugMeshProjection_Sim_HairDeformedFrames > 0)
 			{
 				RenderProjectionData(
-					false, 
+					EHairStrandsInterpolationType::SimulationStrands,
 					GHairDebugMeshProjection_Sim_HairRestTriangles > 0,
 					GHairDebugMeshProjection_Sim_HairRestFrames > 0, 
 					GHairDebugMeshProjection_Sim_HairDeformedTriangles > 0, 
@@ -1725,9 +1737,10 @@ void RenderHairStrandsDebugInfo(
 			FRDGTextureRef PPLLNodeIndexTexture = GraphBuilder.RegisterExternalTexture(VisibilityData.PPLLNodeIndexTexture, TEXT("PPLLNodeIndexTexture"));
 			FRDGBufferRef  PPLLNodeDataBuffer = GraphBuilder.RegisterExternalBuffer(VisibilityData.PPLLNodeDataBuffer, TEXT("PPLLNodeDataBuffer"));
 
+			const FIntPoint PPLLResolution = VisibilityData.PPLLNodeIndexTexture->GetDesc().Extent;
 			FHairVisibilityDebugPPLLCS::FParameters* PassParameters = GraphBuilder.AllocParameters<FHairVisibilityDebugPPLLCS::FParameters>();
-			PassParameters->PPLLMeanListElementCountPerPixel = float(GetPPLLMeanListElementCountPerPixel());
-			PassParameters->PPLLMaxTotalListElementCount = float(GetPPLLMaxTotalListElementCount(VisibilityData.PPLLNodeIndexTexture->GetDesc().Extent));
+			PassParameters->PPLLMeanListElementCountPerPixel = float(VisibilityData.MaxSampleCount);
+			PassParameters->PPLLMaxTotalListElementCount = float(VisibilityData.MaxSampleCount * PPLLResolution.X * PPLLResolution.Y);
 			PassParameters->PPLLCounter = PPLLNodeCounterTexture;
 			PassParameters->PPLLNodeIndex = PPLLNodeIndexTexture;
 			PassParameters->PPLLNodeData = GraphBuilder.CreateSRV(FRDGBufferSRVDesc(PPLLNodeDataBuffer));
diff --git a/Engine/Source/Runtime/Renderer/Private/HairStrands/HairStrandsDeepShadow.cpp b/Engine/Source/Runtime/Renderer/Private/HairStrands/HairStrandsDeepShadow.cpp
--- a/Engine/Source/Runtime/Renderer/Private/HairStrands/HairStrandsDeepShadow.cpp
+++ b/Engine/Source/Runtime/Renderer/Private/HairStrands/HairStrandsDeepShadow.cpp
@@ -157,7 +157,8 @@ FHairStrandsDeepShadowViews RenderHairStrandsDeepShadows(
 					AtlasSlotIndex++;
 
 					const bool bIsOrtho = LightType == ELightComponentType::LightType_Directional;
-					const FVector4 HairRenderInfo = PackHairRenderInfo(MinStrandRadiusAtDepth1.Primary, MinStrandRadiusAtDepth1.Primary, 1, bIsOrtho, false);
+					const FVector4 HairRenderInfo = PackHairRenderInfo(MinStrandRadiusAtDepth1.Primary, MinStrandRadiusAtDepth1.Stable, MinStrandRadiusAtDepth1.Primary, 1);
+					const uint32 HairRenderInfoBits = PackHairRenderInfoBits(bIsOrtho, false);
 					
 					// Front depth
 					{
@@ -180,6 +181,7 @@ FHairStrandsDeepShadowViews RenderHairStrandsDeepShadows(
 							EHairStrandsRasterPassType::FrontDepth,
 							DomData.AtlasRect,
 							HairRenderInfo,
+							HairRenderInfoBits,
 							DomData.LightDirection,
 							PassParameters);
 					}
@@ -206,6 +208,7 @@ FHairStrandsDeepShadowViews RenderHairStrandsDeepShadows(
 							EHairStrandsRasterPassType::DeepOpacityMap,
 							DomData.AtlasRect,
 							HairRenderInfo,
+							HairRenderInfoBits,
 							DomData.LightDirection,
 							PassParameters);						
 					}
diff --git a/Engine/Source/Runtime/Renderer/Private/HairStrands/HairStrandsEnvironment.cpp b/Engine/Source/Runtime/Renderer/Private/HairStrands/HairStrandsEnvironment.cpp
--- a/Engine/Source/Runtime/Renderer/Private/HairStrands/HairStrandsEnvironment.cpp
+++ b/Engine/Source/Runtime/Renderer/Private/HairStrands/HairStrandsEnvironment.cpp
@@ -22,6 +22,9 @@
 
 ///////////////////////////////////////////////////////////////////////////////////////////////////
 
+static int32 GHairScatterSceneLighting = 1;
+static FAutoConsoleVariableRef CVarHairScatterSceneLighting(TEXT("r.HairStrands.ScatterSceneLighting"), GHairScatterSceneLighting, TEXT("Enable scene color lighting scattering into hair (valid for short hair only)."));
+
 static int32 GHairSkylightingEnable = 1;
 static FAutoConsoleVariableRef CVarHairSkylightingEnable(TEXT("r.HairStrands.SkyLighting"), GHairSkylightingEnable, TEXT("Enable sky lighting on hair."));
 
@@ -51,7 +54,7 @@ static float GHairStrandsSkyLightingDistanceThreshold = 10;
 static FAutoConsoleVariableRef CVarHairStrandsSkyAOThreshold(TEXT("r.HairStrands.SkyAO.DistanceThreshold"), GHairStrandsSkyAODistanceThreshold, TEXT("Max distance for occlusion search."), ECVF_Scalability | ECVF_RenderThreadSafe);
 static FAutoConsoleVariableRef CVarHairStrandsSkyLightingDistanceThreshold(TEXT("r.HairStrands.SkyLighting.DistanceThreshold"), GHairStrandsSkyLightingDistanceThreshold, TEXT("Max distance for occlusion search."), ECVF_Scalability | ECVF_RenderThreadSafe);
 
-static bool GetHairStrandsSkyLightingEnable() { return GHairSkylightingEnable > 0; }
+bool GetHairStrandsSkyLightingEnable() { return GHairSkylightingEnable > 0; }
 static bool GetHairStrandsSkyAOEnable() { return GHairSkyAOEnable > 0; }
 static float GetHairStrandsSkyLightingConeAngle() { return FMath::Max(0.f, GHairSkylightingConeAngle); }
 
@@ -264,7 +267,8 @@ class FHairEnvironmentLightingCS : public FGlobalShader
 	class FGroupSize : SHADER_PERMUTATION_INT("PERMUTATION_GROUP_SIZE", 2);
 	class FSampleSet : SHADER_PERMUTATION_INT("PERMUTATION_SAMPLESET", 2);
 	class FVirtualVoxel : SHADER_PERMUTATION_INT("VOXEL_TRAVERSAL_TYPE", 5);
-	using FPermutationDomain = TShaderPermutationDomain<FGroupSize, FSampleSet, FVirtualVoxel>;
+	class FUseSceneColor : SHADER_PERMUTATION_INT("USE_SCENE_COLOR", 2);
+	using FPermutationDomain = TShaderPermutationDomain<FGroupSize, FSampleSet, FVirtualVoxel, FUseSceneColor>;
 
 	static void ModifyCompilationEnvironment(const FGlobalShaderPermutationParameters& Parameters, FShaderCompilerEnvironment& OutEnvironment)
 	{
@@ -273,7 +277,31 @@ class FHairEnvironmentLightingCS : public FGlobalShader
 
 	static bool ShouldCompilePermutation(const FGlobalShaderPermutationParameters& Parameters)
 	{
-		return IsHairStrandsSupported(Parameters.Platform);
+		if (!IsHairStrandsSupported(Parameters.Platform))
+		{
+			return false;
+		}
+
+		const FPermutationDomain PermutationVector(Parameters.PermutationId);
+		if (PermutationVector.Get<FUseSceneColor>() == 1 &&
+		   (PermutationVector.Get<FSampleSet>() == 1 || PermutationVector.Get<FVirtualVoxel>() == 0))
+		{
+			return false;
+		}
+
+		return true;
+	}
+
+	static FPermutationDomain RemapPermutation(FPermutationDomain PermutationVector)
+	{
+		// If not rendering the sky, ignore the fastsky and sundisk permutations
+		if (PermutationVector.Get<FUseSceneColor>() == 1)
+		{
+			PermutationVector.Set<FSampleSet>(0);
+			PermutationVector.Set<FVirtualVoxel>(1);
+		}
+
+		return PermutationVector;
 	}
 
 	BEGIN_SHADER_PARAMETER_STRUCT(FParameters, )
@@ -303,7 +331,9 @@ class FHairEnvironmentLightingCS : public FGlobalShader
 		SHADER_PARAMETER_STRUCT_INCLUDE(FSceneTextureParameters, SceneTextures)
 		SHADER_PARAMETER_STRUCT_INCLUDE(FSceneTextureSamplerParameters, SceneTextureSamplers)
 
+		SHADER_PARAMETER_RDG_TEXTURE(Texture2D, SceneColorTexture)
 		SHADER_PARAMETER_RDG_TEXTURE(Texture2D, HairCountTexture)
+		SHADER_PARAMETER_RDG_TEXTURE(Texture2D, HairCategorizationTexture)
 		SHADER_PARAMETER_RDG_TEXTURE(TEXTURE3D, HairEnergyLUTTexture)
 		SHADER_PARAMETER_RDG_TEXTURE(TEXTURE3D, HairScatteringLUTTexture)
 		SHADER_PARAMETER_SAMPLER(SamplerState, HairLUTSampler)
@@ -332,6 +362,7 @@ static FRDGBufferRef AddHairStrandsEnvironmentLightingPassCS(
 	const FHairStrandsVisibilityData& VisibilityData,
 	const FHairStrandsMacroGroupDatas& MacroGroupDatas,
 	const FHairStrandsMacroGroupData& MacroGroupData,
+	const FRDGTextureRef SceneColorTexture,
 	const uint32 NodeGroupSize,
 	FRDGBufferRef IndirectArgsBuffer)
 {
@@ -363,6 +394,12 @@ static FRDGBufferRef AddHairStrandsEnvironmentLightingPassCS(
 		PassParameters->Voxel_Resolution = MacroGroupData.GetResolution();
 		PassParameters->Voxel_DensityTexture = GraphBuilder.RegisterExternalTexture(MacroGroupData.VoxelResources.DensityTexture);
 	}
+	const bool bUseSceneColor = SceneColorTexture != nullptr;
+	if (bUseSceneColor)
+	{
+		PassParameters->SceneColorTexture = SceneColorTexture;
+		PassParameters->HairCategorizationTexture = GraphBuilder.RegisterExternalTexture(VisibilityData.CategorizationTexture ? VisibilityData.CategorizationTexture : GSystemTextures.BlackDummy);
+	}
 	PassParameters->Voxel_DensityScale = GetHairStrandsVoxelizationDensityScale();
 	PassParameters->Voxel_DepthBiasScale = GetHairStrandsVoxelizationDepthBiasScale();
 	PassParameters->Voxel_TanConeAngle = FMath::Tan(FMath::DegreesToRadians(GetHairStrandsSkyLightingConeAngle()));
@@ -406,6 +443,8 @@ static FRDGBufferRef AddHairStrandsEnvironmentLightingPassCS(
 	PermutationVector.Set<FHairEnvironmentLightingCS::FGroupSize>(NodeGroupSize == 64 ? 0 : (NodeGroupSize == 32 ? 1 : 2));
 	PermutationVector.Set<FHairEnvironmentLightingCS::FSampleSet>(PassParameters->MultipleScatterSampleCount <= 16 ? 0 : 1);
 	PermutationVector.Set<FHairEnvironmentLightingCS::FVirtualVoxel>(bUseVirtualVoxel ? VoxelTraversalType : 0);
+	PermutationVector.Set<FHairEnvironmentLightingCS::FUseSceneColor>(bUseSceneColor ? 1 : 0);
+	PermutationVector = FHairEnvironmentLightingCS::RemapPermutation(PermutationVector);
 
 	TShaderMapRef<FHairEnvironmentLightingCS> ComputeShader(View.ShaderMap, PermutationVector);
 	FComputeShaderUtils::AddPass(
@@ -421,6 +460,45 @@ static FRDGBufferRef AddHairStrandsEnvironmentLightingPassCS(
 
 ///////////////////////////////////////////////////////////////////////////////////////////////////
 
+void RenderHairStrandsSceneColorScattering(
+	FRHICommandListImmediate& RHICmdList,
+	const TArray<FViewInfo>& Views,
+	const struct FHairStrandsDatas* HairDatas)
+{
+	if (Views.Num() == 0 || !HairDatas || GHairScatterSceneLighting <= 0)
+		return;
+
+	for (uint32 ViewIndex = 0; ViewIndex < uint32(Views.Num()); ++ViewIndex)
+	{
+		check(ViewIndex < uint32(HairDatas->HairVisibilityViews.HairDatas.Num()));
+		const FHairStrandsVisibilityData& VisibilityData = HairDatas->HairVisibilityViews.HairDatas[ViewIndex];
+		const bool bRenderHairLighting = VisibilityData.NodeIndex && VisibilityData.NodeDataSRV;
+		if (!bRenderHairLighting)
+		{
+			continue;
+		}
+
+		if (ViewIndex > uint32(HairDatas->MacroGroupsPerViews.Views.Num()))
+			continue;
+
+		const FViewInfo& View = Views[ViewIndex];
+		const FHairStrandsMacroGroupDatas& MacroGroupDatas = HairDatas->MacroGroupsPerViews.Views[ViewIndex];
+		FRDGBuilder GraphBuilder(RHICmdList);
+
+		FSceneRenderTargets& SceneContext = FSceneRenderTargets::Get(RHICmdList);
+		FRDGTextureRef SceneColorTexture = GraphBuilder.RegisterExternalTexture(SceneContext.GetSceneColor());
+		FRDGTextureRef SceneColorSubPixelTexture = GraphBuilder.RegisterExternalTexture(SceneContext.GetSceneColorSubPixel());
+		FRDGBufferRef NodeIndirectArgBuffer = GraphBuilder.RegisterExternalBuffer(VisibilityData.NodeIndirectArg, TEXT("HairNodeIndirectArgBuffer"));
+
+		for (const FHairStrandsMacroGroupData& MacroGroupData : HairDatas->MacroGroupsPerViews.Views[ViewIndex].Datas)
+		{
+			FRDGBufferRef SampleLightingBuffer = AddHairStrandsEnvironmentLightingPassCS(GraphBuilder, View, VisibilityData, MacroGroupDatas, MacroGroupData, SceneColorTexture, VisibilityData.NodeGroupSize, NodeIndirectArgBuffer);
+			AddHairEnvironmentLightingComposePass(GraphBuilder, View, VisibilityData, MacroGroupData, SampleLightingBuffer, SceneColorTexture, SceneColorSubPixelTexture);
+		}
+		GraphBuilder.Execute();
+	}
+}
+
 void RenderHairStrandsEnvironmentLighting(
 	FRDGBuilder& GraphBuilder,
 	const uint32 ViewIndex,
@@ -440,14 +518,14 @@ void RenderHairStrandsEnvironmentLighting(
 	{
 		return;
 	}
-
+	
 	FRDGBufferRef NodeIndirectArgBuffer = GraphBuilder.RegisterExternalBuffer(VisibilityData.NodeIndirectArg, TEXT("HairNodeIndirectArgBuffer"));
 
 	const FViewInfo& View = Views[ViewIndex];
 	const FHairStrandsMacroGroupDatas& MacroGroupDatas = HairDatas->MacroGroupsPerViews.Views[ViewIndex];
 	for (const FHairStrandsMacroGroupData& MacroGroupData : HairDatas->MacroGroupsPerViews.Views[ViewIndex].Datas)
 	{
-		FRDGBufferRef SampleLightingBuffer = AddHairStrandsEnvironmentLightingPassCS(GraphBuilder, View, VisibilityData, MacroGroupDatas, MacroGroupData, VisibilityData.NodeGroupSize, NodeIndirectArgBuffer);
+		FRDGBufferRef SampleLightingBuffer = AddHairStrandsEnvironmentLightingPassCS(GraphBuilder, View, VisibilityData, MacroGroupDatas, MacroGroupData, SceneColorTexture, VisibilityData.NodeGroupSize, NodeIndirectArgBuffer);
 		AddHairEnvironmentLightingComposePass(GraphBuilder, View, VisibilityData, MacroGroupData, SampleLightingBuffer, SceneColorTexture, SceneColorSubPixelTexture);
 	}
 }
diff --git a/Engine/Source/Runtime/Renderer/Private/HairStrands/HairStrandsEnvironment.h b/Engine/Source/Runtime/Renderer/Private/HairStrands/HairStrandsEnvironment.h
--- a/Engine/Source/Runtime/Renderer/Private/HairStrands/HairStrandsEnvironment.h
+++ b/Engine/Source/Runtime/Renderer/Private/HairStrands/HairStrandsEnvironment.h
@@ -22,4 +22,9 @@ void RenderHairStrandsEnvironmentLighting(
 	const TArray<FViewInfo>& Views,
 	const struct FHairStrandsDatas* HairDatas,
 	FRDGTextureRef SceneColorTexture,
-	FRDGTextureRef SceneColorSubPixelTexture);
\ No newline at end of file
+	FRDGTextureRef SceneColorSubPixelTexture);
+
+void RenderHairStrandsSceneColorScattering(
+	FRHICommandListImmediate& RHICmdList,
+	const TArray<FViewInfo>& Views,
+	const struct FHairStrandsDatas* HairDatas);
