Integrate hair changes (rendering+simulation)

diff --git a/Engine/Source/Runtime/Renderer/Private/HairStrands/HairStrandsInterface.cpp b/Engine/Source/Runtime/Renderer/Private/HairStrands/HairStrandsInterface.cpp
--- a/Engine/Source/Runtime/Renderer/Private/HairStrands/HairStrandsInterface.cpp
+++ b/Engine/Source/Runtime/Renderer/Private/HairStrands/HairStrandsInterface.cpp
@@ -8,6 +8,7 @@
 #include "HairStrandsRendering.h"
 #include "HairStrandsMeshProjection.h"
 
+#include "GPUSkinCache.h"
 #include "SkeletalRenderPublic.h"
 #include "SceneRendering.h"
 #include "SystemTextures.h"
@@ -17,6 +18,9 @@ DEFINE_LOG_CATEGORY_STATIC(LogHairRendering, Log, All);
 static int32 GHairStrandsRenderingEnable = 1;
 static FAutoConsoleVariableRef CVarHairStrandsRenderingEnable(TEXT("r.HairStrands.Enable"), GHairStrandsRenderingEnable, TEXT("Enable/Disable hair strands rendering"));
 
+static int32 GHairStrandsRaytracingEnable = 1;
+static FAutoConsoleVariableRef CVarHairStrandsRaytracingEnable(TEXT("r.HairStrands.Raytracing"), GHairStrandsRaytracingEnable, TEXT("Enable/Disable hair strands raytracing (fallback onto raster techniques"));
+
 ////////////////////////////////////////////////////////////////////////////////////////////////////////////////////////
 
 FHairGroupPublicData::FHairGroupPublicData(uint32 InGroupIndex, uint32 InGroupInstanceVertexCount, uint32 InClusterCount, uint32 InVertexCount)
@@ -55,7 +59,7 @@ void FHairGroupPublicData::ReleaseRHI()
 
 ////////////////////////////////////////////////////////////////////////////////////////////////////////////////////////
 
-inline FHairStrandsProjectionMeshData::Section ConvertMeshSection(const FCachedGeometrySection& In, const FTransform& InTransform)
+inline FHairStrandsProjectionMeshData::Section ConvertMeshSection(const FCachedGeometry::Section& In, const FTransform& InTransform)
 {
 	FHairStrandsProjectionMeshData::Section Out;
 	Out.IndexBuffer = In.IndexBuffer;
@@ -88,26 +92,36 @@ struct FHairStrandsManager
 {
 	struct Element
 	{
-		uint64 Id = 0;
+		uint32 ComponentId = 0;
+		uint32 SkeletalComponentId = 0;
 		EWorldType::Type WorldType = EWorldType::None;
 		FHairStrandsDebugInfo DebugInfo;
 		FHairStrandsPrimitiveResources PrimitiveResources;
 		FHairStrandsInterpolationData InterpolationData;
 		FHairStrandsProjectionHairData RenProjectionHairDatas;
 		FHairStrandsProjectionHairData SimProjectionHairDatas;
-		FCachedGeometry CachedGeometry;
 		FTransform SkeletalLocalToWorld;
-		const FSkeletalMeshObject* MeshObject = nullptr;
+		FTransform LocalToWorld;
 		FHairStrandsProjectionDebugInfo DebugProjectionInfo;
 		int32 FrameLODIndex = -1;
 	};
+
+	// TODO change this array to a queue update, in order make processing/update thread safe.
 	TArray<Element> Elements;
+
+	FHairStrandsManager()
+	{
+		// Reserve a large a mount of object to avoid any potential memory reallocation, which 
+		// could cause some thread safety issue. This is a workaround against the non-thread-safe array
+		Elements.Reserve(64);
+	}
 };
 
 FHairStrandsManager GHairManager;
 
 void RegisterHairStrands(
-	uint64 Id, 
+	uint32 ComponentId,
+	uint32 SkeletalComponentId, 
 	EWorldType::Type WorldType,
 	const FHairStrandsInterpolationData& InterpolationData, 
 	const FHairStrandsProjectionHairData& RenProjectionDatas,
@@ -118,7 +132,7 @@ void RegisterHairStrands(
 {
 	for (int32 Index = 0; Index < GHairManager.Elements.Num(); ++Index)
 	{
-		if (GHairManager.Elements[Index].Id == Id && GHairManager.Elements[Index].WorldType == WorldType)
+		if (GHairManager.Elements[Index].ComponentId == ComponentId && GHairManager.Elements[Index].WorldType == WorldType)
 		{
 			// Component already registered. This should not happen. 
 			UE_LOG(LogHairRendering, Warning, TEXT("Component already register. This should't happen. Please report this to a rendering engineer."))
@@ -127,28 +141,50 @@ void RegisterHairStrands(
 	}
 
 	FHairStrandsManager::Element& E =  GHairManager.Elements.AddDefaulted_GetRef();
-	E.Id = Id;
+	E.ComponentId = ComponentId;
+	E.SkeletalComponentId = SkeletalComponentId;
 	E.WorldType = WorldType;
 	E.InterpolationData = InterpolationData;
 	E.RenProjectionHairDatas = RenProjectionDatas;
 	E.SimProjectionHairDatas = SimProjectionDatas;
 	E.SkeletalLocalToWorld = FTransform::Identity;
+	E.LocalToWorld = FTransform::Identity;
 	E.PrimitiveResources = PrimitiveResources;
 	E.DebugInfo = DebugInfo;
 	E.DebugProjectionInfo = DebugProjectionInfo;
 }
 
+bool UpdateHairStrandsDebugInfo(
+	uint32 ComponentId,
+	EWorldType::Type WorldType,
+	const uint32 GroupIt,
+	const bool bSimulationEnable)
+{
+	for (FHairStrandsManager::Element& E : GHairManager.Elements)
+	{
+		if (E.ComponentId != ComponentId || E.WorldType != WorldType || GroupIt >= uint32(E.DebugInfo.HairGroups.Num()))
+			continue;
+
+		E.DebugInfo.HairGroups[GroupIt].bHasSimulation = bSimulationEnable;
+		return true;
+	}
+
+	return false;
+}
+
 bool UpdateHairStrands(
-	uint64 Id, 
+	uint32 ComponentId,
 	EWorldType::Type WorldType, 
 	const FTransform& HairLocalToWorld, 
 	const FTransform& SkeletalLocalToWorld)
 {
 	for (FHairStrandsManager::Element& E : GHairManager.Elements)
 	{
-		if (E.Id != Id || E.WorldType != WorldType)
+		if (E.ComponentId != ComponentId || E.WorldType != WorldType)
 			continue;
 
+		E.LocalToWorld = HairLocalToWorld;
+
 		for (FHairStrandsProjectionHairData::HairGroup& ProjectionData : E.RenProjectionHairDatas.HairGroups)
 		{
 			ProjectionData.LocalToWorld = HairLocalToWorld;
@@ -165,12 +201,12 @@ bool UpdateHairStrands(
 }
 
 bool UpdateHairStrands(
-	uint64 Id, 
+	uint32 ComponentId,
 	EWorldType::Type NewWorldType)
 {
 	for (FHairStrandsManager::Element& E : GHairManager.Elements)
 	{
-		if (E.Id != Id)
+		if (E.ComponentId != ComponentId)
 			continue;
 
 		E.WorldType = NewWorldType;
@@ -181,7 +217,7 @@ bool UpdateHairStrands(
 }
 
 bool UpdateHairStrands(
-	uint64 Id,
+	uint32 ComponentId,
 	EWorldType::Type WorldType,
 	const FTransform& HairLocalToWorld,
 	const FHairStrandsProjectionHairData& RenProjectionDatas,
@@ -189,9 +225,11 @@ bool UpdateHairStrands(
 {
 	for (FHairStrandsManager::Element& E : GHairManager.Elements)
 	{
-		if (E.Id != Id || E.WorldType != WorldType)
+		if (E.ComponentId != ComponentId || E.WorldType != WorldType)
 			continue;
 
+		E.LocalToWorld = HairLocalToWorld;
+
 		E.RenProjectionHairDatas = RenProjectionDatas;
 		for (FHairStrandsProjectionHairData::HairGroup& ProjectionData : E.RenProjectionHairDatas.HairGroups)
 		{
@@ -208,31 +246,14 @@ bool UpdateHairStrands(
 	return false;
 }
 
-bool UpdateHairStrands(
-	uint64 Id,
-	EWorldType::Type WorldType,
-	const FSkeletalMeshObject* MeshObject)
-{
-	for (FHairStrandsManager::Element& E : GHairManager.Elements)
-	{
-		if (E.Id != Id || E.WorldType != WorldType)
-			continue;
-
-		E.MeshObject = MeshObject;
-		return true;
-	}
-
-	return false;
-}
-
-void UnregisterHairStrands(uint64 Id)
+void UnregisterHairStrands(uint32 ComponentId)
 {
 	for (int32 Index=0;Index< GHairManager.Elements.Num();++Index)
 	{
-		if (GHairManager.Elements[Index].Id == Id)
+		if (GHairManager.Elements[Index].ComponentId == ComponentId)
 		{
 			GHairManager.Elements[Index] = GHairManager.Elements[GHairManager.Elements.Num()-1];
-			GHairManager.Elements.SetNum(GHairManager.Elements.Num() - 1);
+			GHairManager.Elements.SetNum(GHairManager.Elements.Num() - 1, false);
 		}
 	}
 }
@@ -294,6 +315,7 @@ void RunProjection(
 void RunHairStrandsInterpolation(
 	FRHICommandListImmediate& RHICmdList, 
 	EWorldType::Type WorldType, 
+	const FGPUSkinCache* SkinCache,
 	const FShaderDrawDebugData* ShaderDrawData,
 	FGlobalShaderMap* ShaderMap, 
 	EHairStrandsInterpolationType Type,
@@ -301,24 +323,27 @@ void RunHairStrandsInterpolation(
 {
 	check(IsInRenderingThread());
 
-	// Update geometry cached based on GPU Skin output
-	for (FHairStrandsManager::Element& E : GHairManager.Elements)
-	{
-		if (E.WorldType != WorldType)
-			continue;
+	DECLARE_GPU_STAT(HairStrandsInterpolation);
+	SCOPED_DRAW_EVENT(RHICmdList, HairStrandsInterpolation);
+	SCOPED_GPU_STAT(RHICmdList, HairStrandsInterpolation);
 
-		E.CachedGeometry = E.MeshObject ? E.MeshObject->GetCachedGeometry() : FCachedGeometry();
-	}
-	
 	// Update dynamic mesh triangles
 	for (FHairStrandsManager::Element& E : GHairManager.Elements)
 	{
 		E.FrameLODIndex = -1;
-		if (E.WorldType != WorldType || E.CachedGeometry.Sections.Num() == 0)
+		if (E.WorldType != WorldType)
 			continue;
 	
+		FCachedGeometry CachedGeometry;
+		if (SkinCache)
+		{
+			CachedGeometry = SkinCache->GetCachedGeometry(E.SkeletalComponentId);
+		}
+		if (CachedGeometry.Sections.Num() == 0)
+			continue;
+
 		FHairStrandsProjectionMeshData::LOD MeshDataLOD;
-		for (FCachedGeometrySection Section : E.CachedGeometry.Sections)
+		for (FCachedGeometry::Section& Section : CachedGeometry.Sections)
 		{
 			// Ensure all mesh's sections have the same LOD index
 			if (E.FrameLODIndex < 0) E.FrameLODIndex = Section.LODIndex;
@@ -344,6 +369,26 @@ void RunHairStrandsInterpolation(
 				UpdateHairStrandsMeshTriangles(GraphBuilder, ShaderMap, E.FrameLODIndex, HairStrandsTriangleType::DeformedPose, MeshDataLOD, ProjectionHairData);
 			}
 		}
+
+		for (FHairStrandsProjectionHairData::HairGroup& ProjectionHairData : E.SimProjectionHairDatas.HairGroups)
+		{
+			if (EHairStrandsInterpolationType::SimulationStrands == Type && 0 <= E.FrameLODIndex && E.FrameLODIndex < ProjectionHairData.LODDatas.Num() && ProjectionHairData.LODDatas[E.FrameLODIndex].bIsValid)
+			{
+				InitHairStrandsMeshSamples(GraphBuilder, ShaderMap, E.FrameLODIndex, HairStrandsTriangleType::DeformedPose, MeshDataLOD, ProjectionHairData);
+				UpdateHairStrandsMeshSamples(GraphBuilder, ShaderMap, E.FrameLODIndex, MeshDataLOD, ProjectionHairData);
+				//InterpolateHairStrandsMeshTriangles(GraphBuilder, ShaderMap, E.FrameLODIndex, MeshDataLOD, ProjectionHairData);
+			}
+		}
+
+		/*for (FHairStrandsProjectionHairData::HairGroup& ProjectionHairData : E.RenProjectionHairDatas.HairGroups)
+		{
+			if (EHairStrandsInterpolationType::RenderStrands == Type && 0 <= E.FrameLODIndex && E.FrameLODIndex < ProjectionHairData.LODDatas.Num() && ProjectionHairData.LODDatas[E.FrameLODIndex].bIsValid)
+			{
+				InitHairStrandsMeshSamples(GraphBuilder, ShaderMap, E.FrameLODIndex, HairStrandsTriangleType::DeformedPose, MeshDataLOD, ProjectionHairData);
+				UpdateHairStrandsMeshSamples(GraphBuilder, ShaderMap, E.FrameLODIndex, MeshDataLOD, ProjectionHairData);
+				InterpolateHairStrandsMeshTriangles(GraphBuilder, ShaderMap, E.FrameLODIndex, MeshDataLOD, ProjectionHairData);
+			}
+		}*/
 		GraphBuilder.Execute();
 	}
 
@@ -372,14 +417,16 @@ void RunHairStrandsInterpolation(
 
 			if (E.InterpolationData.Input && E.InterpolationData.Output && E.InterpolationData.Function)
 			{
-				E.InterpolationData.Function(RHICmdList, ShaderDrawData, E.InterpolationData.Input, E.InterpolationData.Output, E.RenProjectionHairDatas, E.SimProjectionHairDatas, E.FrameLODIndex, ClusterData);
+				E.InterpolationData.Function(RHICmdList, ShaderDrawData, E.LocalToWorld, E.InterpolationData.Input, E.InterpolationData.Output, E.RenProjectionHairDatas, E.SimProjectionHairDatas, E.FrameLODIndex, ClusterData);
 			}
 		}
 	}
 }
 
-void GetGroomInterpolationData(const EWorldType::Type WorldType, 
-	const EHairStrandsProjectionMeshType MeshtType,
+void GetGroomInterpolationData(
+	const EWorldType::Type WorldType, 
+	const EHairStrandsProjectionMeshType MeshType,
+	const FGPUSkinCache* SkinCache,
 	FHairStrandsProjectionMeshData::LOD& OutGeometries)
 {
 	for (FHairStrandsManager::Element& E : GHairManager.Elements)
@@ -387,12 +434,16 @@ void GetGroomInterpolationData(const EWorldType::Type WorldType,
 		if (E.WorldType != WorldType)
 			continue;
 
-		if (MeshtType == EHairStrandsProjectionMeshType::DeformedMesh || MeshtType == EHairStrandsProjectionMeshType::RestMesh)
+		FCachedGeometry CachedGeometry = SkinCache->GetCachedGeometry(E.SkeletalComponentId);
+		if (CachedGeometry.Sections.Num() == 0)
+			continue;
+
+		if (MeshType == EHairStrandsProjectionMeshType::DeformedMesh || MeshType == EHairStrandsProjectionMeshType::RestMesh)
 		{
-			for (FCachedGeometrySection Section : E.CachedGeometry.Sections)
+			for (FCachedGeometry::Section Section : CachedGeometry.Sections)
 			{			
 				FHairStrandsProjectionMeshData::Section OutSection = ConvertMeshSection(Section, E.SkeletalLocalToWorld);
-				if (MeshtType == EHairStrandsProjectionMeshType::RestMesh)
+				if (MeshType == EHairStrandsProjectionMeshType::RestMesh)
 				{
 					// If the mesh has some mesh-tranferred data, we display that otherwise we use the rest data
 					const bool bHasTransferData = E.FrameLODIndex < E.DebugProjectionInfo.TransferredPositions.Num();
@@ -409,29 +460,36 @@ void GetGroomInterpolationData(const EWorldType::Type WorldType,
 			}
 		}
 
-		if (MeshtType == EHairStrandsProjectionMeshType::TargetMesh && E.DebugProjectionInfo.TargetMeshData.LODs.Num() > 0)
+		if (MeshType == EHairStrandsProjectionMeshType::TargetMesh && E.DebugProjectionInfo.TargetMeshData.LODs.Num() > 0)
 		{
 			OutGeometries = E.DebugProjectionInfo.TargetMeshData.LODs[0];
 		}
 
-		if (MeshtType == EHairStrandsProjectionMeshType::SourceMesh && E.DebugProjectionInfo.SourceMeshData.LODs.Num() > 0)
+		if (MeshType == EHairStrandsProjectionMeshType::SourceMesh && E.DebugProjectionInfo.SourceMeshData.LODs.Num() > 0)
 		{
 			OutGeometries = E.DebugProjectionInfo.SourceMeshData.LODs[0];
 		}
 	}
 }
 
-void GetGroomInterpolationData(const EWorldType::Type WorldType, const bool bRenderData, FHairStrandsProjectionHairData& Out, TArray<int32>& OutLODIndices)
+void GetGroomInterpolationData(
+	const EWorldType::Type WorldType, 
+	EHairStrandsInterpolationType StrandsType,
+	const FGPUSkinCache* SkinCache,
+	FHairStrandsProjectionHairData& Out, TArray<int32>& OutLODIndices)
 {
 	for (FHairStrandsManager::Element& E : GHairManager.Elements)
 	{
 		if (E.WorldType != WorldType)
 			continue;
 
-		const bool bHasDynamicMesh = E.CachedGeometry.Sections.Num() > 0;
+		FCachedGeometry CachedGeometry = SkinCache->GetCachedGeometry(E.SkeletalComponentId);
+		if (CachedGeometry.Sections.Num() == 0)
+			continue;
+		const bool bHasDynamicMesh = CachedGeometry.Sections.Num() > 0;
 		if (bHasDynamicMesh)
 		{
-			if (bRenderData)
+			if (StrandsType == EHairStrandsInterpolationType::RenderStrands)
 			{
 				for (FHairStrandsProjectionHairData::HairGroup& ProjectionHairData : E.RenProjectionHairDatas.HairGroups)
 				{
@@ -439,7 +497,7 @@ void GetGroomInterpolationData(const EWorldType::Type WorldType, const bool bRen
 					OutLODIndices.Add(E.FrameLODIndex);
 				}
 			}
-			else
+			else if (StrandsType == EHairStrandsInterpolationType::SimulationStrands)
 			{
 				for (FHairStrandsProjectionHairData::HairGroup& ProjectionHairData : E.SimProjectionHairDatas.HairGroups)
 				{
@@ -458,12 +516,12 @@ FHairStrandsDebugInfos GetHairStandsDebugInfos()
 	{
 		FHairStrandsDebugInfo& Info = Infos.AddDefaulted_GetRef();
 		Info = E.DebugInfo;
-		Info.Id = E.Id;
+		Info.ComponentId = E.ComponentId;
 		Info.WorldType = E.WorldType;
 
 		const uint32 GroupCount = Info.HairGroups.Num();
 		for (uint32 GroupIt=0; GroupIt < GroupCount; ++GroupIt)
-		{
+		{		
 			FHairStrandsDebugInfo::HairGroup& GroupInfo = Info.HairGroups[GroupIt];
 			if (GroupIt < uint32(E.RenProjectionHairDatas.HairGroups.Num()))
 			{
@@ -482,11 +540,11 @@ FHairStrandsDebugInfos GetHairStandsDebugInfos()
 	return Infos;
 }
 
-FHairStrandsPrimitiveResources GetHairStandsPrimitiveResources(uint64 Id)
+FHairStrandsPrimitiveResources GetHairStandsPrimitiveResources(uint32 ComponentId)
 {
 	for (FHairStrandsManager::Element& E : GHairManager.Elements)
 	{
-		if (E.Id != Id)
+		if (E.ComponentId != ComponentId)
 			continue;
 
 		return E.PrimitiveResources;
@@ -502,9 +560,10 @@ bool IsHairStrandsEnable(EShaderPlatform Platform)
 		GHairStrandsRenderingEnable == 1 && 
 		GHairManager.Elements.Num() > 0; 
 }
+
 bool IsHairRayTracingEnabled()
 {
-	return IsRayTracingEnabled();
+	return IsRayTracingEnabled() && GHairStrandsRaytracingEnable;
 }
 
 ///////////////////////////////////////////////////////////////////////////////////////////////////
@@ -521,15 +580,7 @@ void EnqueueGroomBindingQuery(void* Asset, TBindingProcess BindingProcess)
 	GBindingQueries.Enqueue({ Asset, BindingProcess });
 }
 
-bool HasHairStrandsProjectionQuery(EShaderPlatform Platform)
-{
-	return
-		IsHairStrandsSupported(Platform) &&
-		GHairStrandsRenderingEnable == 1 &&
-		!GBindingQueries.IsEmpty();
-}
-
-void RunHairStrandsBindingQueries(FRHICommandListImmediate& RHICmdList, FGlobalShaderMap* ShaderMap)
+static void RunHairStrandsBindingQueries(FRHICommandListImmediate& RHICmdList, FGlobalShaderMap* ShaderMap)
 {
 	FBindingQuery Q;
 	while (GBindingQueries.Dequeue(Q))
@@ -539,4 +590,124 @@ void RunHairStrandsBindingQueries(FRHICommandListImmediate& RHICmdList, FGlobalS
 			Q.BindingProcess(RHICmdList, Q.Asset);
 		}
 	}
-}
\ No newline at end of file
+}
+
+///////////////////////////////////////////////////////////////////////////////////////////////////
+
+struct FFollicleQuery
+{
+	static const uint32 MaxInfoCount = 16;
+	FFollicleInfo Infos[MaxInfoCount];
+	uint32 InfoCount = 0;
+	UTexture2D* OutTexture = nullptr;
+};
+TQueue<FFollicleQuery> GFollicleQueries;
+
+void EnqueueFollicleMaskUpdateQuery(const TArray<FFollicleInfo>& Infos, UTexture2D* OutTexture)
+{
+	if (!OutTexture)
+	{
+		return;
+	}
+
+	FFollicleQuery Q;
+	Q.OutTexture = OutTexture;
+	for (const FFollicleInfo& Info : Infos)
+	{
+		if (Q.InfoCount >= FFollicleQuery::MaxInfoCount)
+		{
+			return;
+		}
+		Q.Infos[Q.InfoCount++] = Info;
+	}
+	GFollicleQueries.Enqueue(Q);
+}
+
+static void RunFolliculeMaskGeneration(FRHICommandListImmediate& RHICmdList, FGlobalShaderMap* ShaderMap)
+{
+	struct FFollicleElement
+	{
+		FHairStrandsManager::Element* E = nullptr;
+		FFollicleInfo::EChannel Channel = FFollicleInfo::R;
+		uint32 KernelSizeInPixels = 0;
+	};
+
+	FFollicleQuery Q;
+	while (GFollicleQueries.Dequeue(Q))
+	{
+		// TODO deduplicate queries (compute hash?)
+
+		// Filter valid/ready groom
+		TArray<FFollicleElement> Elements;
+		Elements.Reserve(Q.InfoCount);
+		for (uint32 InfoIt=0; InfoIt <Q.InfoCount;++InfoIt)
+		{
+			for (FHairStrandsManager::Element& E : GHairManager.Elements)
+			{
+				if (E.ComponentId == Q.Infos[InfoIt].GroomId && E.FrameLODIndex >= 0)
+				{
+					Elements.Add({ &E, Q.Infos[InfoIt].Channel, Q.Infos[InfoIt].KernelSizeInPixels });
+				}
+			}
+		}
+
+		// Generate texture
+		if (Elements.Num() > 0)
+		{
+			TRefCountPtr<IPooledRenderTarget> OutMaskTexture;
+
+			const uint32 MipCount = Q.OutTexture->GetNumMips();
+			const FIntPoint Resolution(Q.OutTexture->Resource->GetSizeX(), Q.OutTexture->Resource->GetSizeY());
+			FRDGBuilder GraphBuilder(RHICmdList);
+			FRDGTextureRef FollicleMaskTexture = nullptr;
+			for (FFollicleElement& E : Elements)
+			{
+				GenerateFolliculeMask(
+					GraphBuilder,
+					ShaderMap,
+					Resolution,
+					MipCount,
+					E.KernelSizeInPixels,
+					uint32(E.Channel),
+					E.E->FrameLODIndex,
+					E.E->RenProjectionHairDatas,
+					FollicleMaskTexture);
+			}
+
+			AddComputeMipsPass(GraphBuilder, ShaderMap, FollicleMaskTexture);
+
+			GraphBuilder.QueueTextureExtraction(FollicleMaskTexture, &OutMaskTexture);
+			GraphBuilder.Execute();
+
+			check(FollicleMaskTexture->Desc.Format == Q.OutTexture->GetPixelFormat());
+
+			FRHICopyTextureInfo CopyInfo;
+			CopyInfo.NumMips = MipCount;
+			RHICmdList.CopyTexture(
+				OutMaskTexture->GetRenderTargetItem().ShaderResourceTexture,
+				Q.OutTexture->Resource->TextureRHI,
+				CopyInfo);
+		}
+	}
+}
+
+void RunHairStrandsProcess(FRHICommandListImmediate& RHICmdList, FGlobalShaderMap* ShaderMap)
+{
+	if (!GFollicleQueries.IsEmpty())
+	{
+		RunFolliculeMaskGeneration(RHICmdList, ShaderMap);
+	}
+
+	if (!GBindingQueries.IsEmpty())
+	{
+		RunHairStrandsBindingQueries(RHICmdList, ShaderMap);
+	}
+}
+
+bool HasHairStrandsProcess(EShaderPlatform Platform)
+{
+	return
+		IsHairStrandsSupported(Platform) &&
+		GHairStrandsRenderingEnable == 1 &&
+		(!GBindingQueries.IsEmpty() || !GFollicleQueries.IsEmpty());
+}
diff --git a/Engine/Source/Runtime/Renderer/Private/HairStrands/HairStrandsMeshProjection.cpp b/Engine/Source/Runtime/Renderer/Private/HairStrands/HairStrandsMeshProjection.cpp
--- a/Engine/Source/Runtime/Renderer/Private/HairStrands/HairStrandsMeshProjection.cpp
+++ b/Engine/Source/Runtime/Renderer/Private/HairStrands/HairStrandsMeshProjection.cpp
@@ -48,7 +48,10 @@ static FRDGBufferRef AddMeshSectionId(
 	if (SectionCount < 0)
 		return nullptr;
 
+	// Initialized the section ID to a large number, as the shader will do an atomic min on the section ID.
 	FRDGBufferRef VertexSectionIdBuffer = GraphBuilder.CreateBuffer(FRDGBufferDesc::CreateBufferDesc(sizeof(uint32), MeshData.Sections[0].TotalVertexCount), TEXT("SectionId"));
+	FRDGBufferUAVRef VertexSectionIdBufferUAV = GraphBuilder.CreateUAV(VertexSectionIdBuffer, PF_R32_UINT);
+	AddClearUAVPass(GraphBuilder, VertexSectionIdBufferUAV, ~0u);
 	for (const FHairStrandsProjectionMeshData::Section& MeshSection : MeshData.Sections)
 	{	
 		FMarkMeshSectionIdCS::FParameters* Parameters = GraphBuilder.AllocParameters<FMarkMeshSectionIdCS::FParameters>();
@@ -58,9 +61,9 @@ static FRDGBufferRef AddMeshSectionId(
 		Parameters->MeshMaxVertexCount			= MeshSection.TotalVertexCount;
 		Parameters->MeshIndexOffset				= MeshSection.IndexBaseIndex;
 		Parameters->MeshIndexBuffer				= MeshSection.IndexBuffer;
-		Parameters->OutVertexSectionId			= GraphBuilder.CreateUAV(VertexSectionIdBuffer, PF_R32_UINT);
+		Parameters->OutVertexSectionId			= VertexSectionIdBufferUAV;
 
-		const FIntVector DispatchGroupCount = FComputeShaderUtils::GetGroupCount(MeshSection.TotalVertexCount, 128);
+		const FIntVector DispatchGroupCount = FComputeShaderUtils::GetGroupCount(MeshSection.NumPrimitives*3, 128);
 		check(DispatchGroupCount.X < 65536);
 		TShaderMapRef<FMarkMeshSectionIdCS> ComputeShader(ShaderMap);
 		FComputeShaderUtils::AddPass(
@@ -252,12 +255,14 @@ static void AddHairStrandMeshProjectionPass(
 		return;
 	}
 
-	// The current shader code HairStrandsMeshProjection.usf encode the section ID onto the highest 3bits of a 32bits uint. 
-	// This limits the number of section to 8. See EncodeTriangleIndex & DecodeTriangleIndex functions in 
+	// The current shader code HairStrandsMeshProjection.usf encode the section ID onto the highest 4bits of a 32bits uint. 
+	// This limits the number of section to 16. See EncodeTriangleIndex & DecodeTriangleIndex functions in 
 	// HairStarndsMeshProjectionCommon.ush for mode details.
+	// This means that the mesh needs to have less than 285M triangles (since triangle ID is stored onto 28bits).
 	//
 	// This could be increase if necessary.
-	check(MeshSectionData.SectionIndex < 8);
+	check(MeshSectionData.SectionIndex < 16);
+	check(MeshSectionData.NumPrimitives < ((1<<28)-1))
 
 	// For projecting hair onto a skeletal mesh, 1 thread is spawn for each hair which iterates over all triangles.
 	// To avoid TDR, we split projection into multiple passes when the mesh is too large.
@@ -358,8 +363,12 @@ void TransferMesh(
 }
 
 ///////////////////////////////////////////////////////////////////////////////////////////////////
+#define SECTION_ARRAY_COUNT 16
 class FHairUpdateMeshTriangleCS : public FGlobalShader
 {
+public:
+	const static uint32 SectionArrayCount = 16;
+private:
 	DECLARE_GLOBAL_SHADER(FHairUpdateMeshTriangleCS);
 	SHADER_USE_PARAMETER_STRUCT(FHairUpdateMeshTriangleCS, FGlobalShader);
 
@@ -368,13 +377,14 @@ class FHairUpdateMeshTriangleCS : public FGlobalShader
 
 	BEGIN_SHADER_PARAMETER_STRUCT(FParameters, )
 		SHADER_PARAMETER(uint32, MaxRootCount)
-
-		SHADER_PARAMETER(uint32, MeshSectionIndex)
-		SHADER_PARAMETER(uint32, MeshMaxIndexCount)
-		SHADER_PARAMETER(uint32, MeshMaxVertexCount)
-		SHADER_PARAMETER(uint32, MeshIndexOffset)
-		SHADER_PARAMETER(uint32, MeshUVsChannelOffset)
-		SHADER_PARAMETER(uint32, MeshUVsChannelCount)
+		SHADER_PARAMETER(uint32, MaxSectionCount)
+		
+		SHADER_PARAMETER_ARRAY(uint32, MeshSectionIndex, [SectionArrayCount])
+		SHADER_PARAMETER_ARRAY(uint32, MeshMaxIndexCount, [SectionArrayCount])
+		SHADER_PARAMETER_ARRAY(uint32, MeshMaxVertexCount, [SectionArrayCount])
+		SHADER_PARAMETER_ARRAY(uint32, MeshIndexOffset, [SectionArrayCount])
+		SHADER_PARAMETER_ARRAY(uint32, MeshUVsChannelOffset, [SectionArrayCount])
+		SHADER_PARAMETER_ARRAY(uint32, MeshUVsChannelCount, [SectionArrayCount])
 		SHADER_PARAMETER_SRV(Buffer, MeshIndexBuffer)
 		SHADER_PARAMETER_SRV(Buffer, MeshPositionBuffer)
 		SHADER_PARAMETER_SRV(Buffer, MeshUVsBuffer)
@@ -387,6 +397,11 @@ class FHairUpdateMeshTriangleCS : public FGlobalShader
 
 public:
 	static bool ShouldCompilePermutation(const FGlobalShaderPermutationParameters& Parameters) { return Parameters.Platform == SP_PCD3D_SM5; }
+	static void ModifyCompilationEnvironment(const FGlobalShaderPermutationParameters& Parameters, FShaderCompilerEnvironment& OutEnvironment)
+	{
+		FGlobalShader::ModifyCompilationEnvironment(Parameters, OutEnvironment);
+		OutEnvironment.SetDefine(TEXT("MAX_SECTION_COUNT"), SectionArrayCount);
+	}
 };
 
 IMPLEMENT_GLOBAL_SHADER(FHairUpdateMeshTriangleCS, "/Engine/Private/HairStrands/HairStrandsMeshUpdate.usf", "MainCS", SF_Compute);
@@ -396,7 +411,7 @@ static void AddHairStrandUpdateMeshTrianglesPass(
 	FGlobalShaderMap* ShaderMap,
 	const int32 LODIndex,
 	const HairStrandsTriangleType Type,
-	const FHairStrandsProjectionMeshData::Section& MeshSectionData,
+	const FHairStrandsProjectionMeshData::LOD& MeshData,
 	FHairStrandsProjectionHairData::HairGroup& RootData)
 {
 	if (RootData.RootCount == 0 || LODIndex < 0 || LODIndex >= RootData.LODDatas.Num())
@@ -406,17 +421,31 @@ static void AddHairStrandUpdateMeshTrianglesPass(
 	FHairStrandsProjectionHairData::LODData& LODData = RootData.LODDatas[LODIndex];
 	check(LODData.LODIndex == LODIndex);
 
+	const int32 SectionCount = MeshData.Sections.Num();
 	FHairUpdateMeshTriangleCS::FParameters* Parameters = GraphBuilder.AllocParameters<FHairUpdateMeshTriangleCS::FParameters>();
 	Parameters->MaxRootCount		= RootData.RootCount;
-	Parameters->MeshSectionIndex	= MeshSectionData.SectionIndex;
-	Parameters->MeshMaxIndexCount	= MeshSectionData.TotalIndexCount;
-	Parameters->MeshMaxVertexCount	= MeshSectionData.TotalVertexCount;
-	Parameters->MeshIndexOffset		= MeshSectionData.IndexBaseIndex;
-	Parameters->MeshIndexBuffer		= MeshSectionData.IndexBuffer;
-	Parameters->MeshPositionBuffer	= MeshSectionData.PositionBuffer;
-	Parameters->MeshUVsBuffer		= MeshSectionData.UVsBuffer;
-	Parameters->MeshUVsChannelOffset= MeshSectionData.UVsChannelOffset;
-	Parameters->MeshUVsChannelCount	= MeshSectionData.UVsChannelCount;
+	Parameters->MaxSectionCount		= SectionCount;
+	Parameters->MeshPositionBuffer	= MeshData.Sections[0].PositionBuffer;
+	Parameters->MeshIndexBuffer		= MeshData.Sections[0].IndexBuffer;
+	Parameters->MeshUVsBuffer		= MeshData.Sections[0].UVsBuffer;
+
+	check(SectionCount < FHairUpdateMeshTriangleCS::SectionArrayCount);
+	for (int32 SectionIt = 0; SectionIt < SectionCount; ++SectionIt)
+	{
+		const FHairStrandsProjectionMeshData::Section& MeshSectionData = MeshData.Sections[SectionIt];
+		
+		// Sanity check to insure all sections share the same underlying buffer 
+		check(Parameters->MeshPositionBuffer == MeshSectionData.PositionBuffer);
+		check(Parameters->MeshIndexBuffer == MeshSectionData.IndexBuffer);
+		check(Parameters->MeshUVsBuffer == MeshSectionData.UVsBuffer);
+
+		Parameters->MeshSectionIndex[SectionIt]		= MeshSectionData.SectionIndex;
+		Parameters->MeshMaxIndexCount[SectionIt]	= MeshSectionData.TotalIndexCount;
+		Parameters->MeshMaxVertexCount[SectionIt]	= MeshSectionData.TotalVertexCount;
+		Parameters->MeshIndexOffset[SectionIt]		= MeshSectionData.IndexBaseIndex;
+		Parameters->MeshUVsChannelOffset[SectionIt] = MeshSectionData.UVsChannelOffset;
+		Parameters->MeshUVsChannelCount[SectionIt]	= MeshSectionData.UVsChannelCount;
+	}
 
 	Parameters->RootTriangleIndex = LODData.RootTriangleIndexBuffer->SRV;
 	if (Type == HairStrandsTriangleType::RestPose)
@@ -460,8 +489,483 @@ void UpdateHairStrandsMeshTriangles(
 	const FHairStrandsProjectionMeshData::LOD& ProjectionMeshData,
 	FHairStrandsProjectionHairData::HairGroup& ProjectionHairData)
 {
-	for (const FHairStrandsProjectionMeshData::Section& MeshSection : ProjectionMeshData.Sections)
+	AddHairStrandUpdateMeshTrianglesPass(GraphBuilder, ShaderMap, LODIndex, Type, ProjectionMeshData, ProjectionHairData);
+}
+
+///////////////////////////////////////////////////////////////////////////////////////////////////
+
+class FHairInterpolateMeshTriangleCS : public FGlobalShader
+{
+private:
+	DECLARE_GLOBAL_SHADER(FHairInterpolateMeshTriangleCS);
+	SHADER_USE_PARAMETER_STRUCT(FHairInterpolateMeshTriangleCS, FGlobalShader);
+
+	using FPermutationDomain = TShaderPermutationDomain<>;
+
+	BEGIN_SHADER_PARAMETER_STRUCT(FParameters, )
+		SHADER_PARAMETER(uint32, MaxRootCount)
+		SHADER_PARAMETER(uint32, MaxSampleCount)
+
+		SHADER_PARAMETER_SRV(Buffer, RestSamplePositionsBuffer)
+		SHADER_PARAMETER_SRV(Buffer, MeshSampleWeightsBuffer)
+
+		SHADER_PARAMETER_SRV(StructuredBuffer, RestRootTrianglePosition0)
+		SHADER_PARAMETER_SRV(StructuredBuffer, RestRootTrianglePosition1)
+		SHADER_PARAMETER_SRV(StructuredBuffer, RestRootTrianglePosition2)
+
+		SHADER_PARAMETER_UAV(StructuredBuffer, OutDeformedRootTrianglePosition0)
+		SHADER_PARAMETER_UAV(StructuredBuffer, OutDeformedRootTrianglePosition1)
+		SHADER_PARAMETER_UAV(StructuredBuffer, OutDeformedRootTrianglePosition2)
+		END_SHADER_PARAMETER_STRUCT()
+
+public:
+	static bool ShouldCompilePermutation(const FGlobalShaderPermutationParameters& Parameters) { return Parameters.Platform == SP_PCD3D_SM5; }
+	static void ModifyCompilationEnvironment(const FGlobalShaderPermutationParameters& Parameters, FShaderCompilerEnvironment& OutEnvironment)
+	{
+		FGlobalShader::ModifyCompilationEnvironment(Parameters, OutEnvironment);
+	}
+};
+
+IMPLEMENT_GLOBAL_SHADER(FHairInterpolateMeshTriangleCS, "/Engine/Private/HairStrands/HairStrandsMeshInterpolate.usf", "MainCS", SF_Compute);
+
+static void AddHairStrandInterpolateMeshTrianglesPass(
+	FRDGBuilder& GraphBuilder,
+	FGlobalShaderMap* ShaderMap,
+	const int32 LODIndex,
+	const FHairStrandsProjectionMeshData::LOD& MeshData,
+	FHairStrandsProjectionHairData::HairGroup& RootData)
+{
+	if (RootData.RootCount == 0 || LODIndex < 0 || LODIndex >= RootData.LODDatas.Num())
 	{
-		AddHairStrandUpdateMeshTrianglesPass(GraphBuilder, ShaderMap, LODIndex, Type, MeshSection, ProjectionHairData);
+		return;
+	}
+	FHairStrandsProjectionHairData::LODData& LODData = RootData.LODDatas[LODIndex];
+	check(LODData.LODIndex == LODIndex);
+
+	FHairInterpolateMeshTriangleCS::FParameters* Parameters = GraphBuilder.AllocParameters<FHairInterpolateMeshTriangleCS::FParameters>();
+	Parameters->MaxRootCount = RootData.RootCount;
+	Parameters->MaxSampleCount = LODData.SampleCount;
+
+	Parameters->RestRootTrianglePosition0 = LODData.RestRootTrianglePosition0Buffer->SRV;
+	Parameters->RestRootTrianglePosition1 = LODData.RestRootTrianglePosition1Buffer->SRV;
+	Parameters->RestRootTrianglePosition2 = LODData.RestRootTrianglePosition2Buffer->SRV;
+
+	Parameters->OutDeformedRootTrianglePosition0 = LODData.DeformedRootTrianglePosition0Buffer->UAV;
+	Parameters->OutDeformedRootTrianglePosition1 = LODData.DeformedRootTrianglePosition1Buffer->UAV;
+	Parameters->OutDeformedRootTrianglePosition2 = LODData.DeformedRootTrianglePosition2Buffer->UAV;
+
+	Parameters->MeshSampleWeightsBuffer = LODData.MeshSampleWeightsBuffer->SRV;
+	Parameters->RestSamplePositionsBuffer = LODData.RestSamplePositionsBuffer->SRV;
+
+	const FIntVector DispatchGroupCount = FComputeShaderUtils::GetGroupCount(RootData.RootCount, 128);
+	check(DispatchGroupCount.X < 65536);
+	TShaderMapRef<FHairInterpolateMeshTriangleCS> ComputeShader(ShaderMap);
+	FComputeShaderUtils::AddPass(
+		GraphBuilder,
+		RDG_EVENT_NAME("HairStrandsTriangleMeshInterpolate"),
+		ComputeShader,
+		Parameters,
+		DispatchGroupCount);
+}
+
+void InterpolateHairStrandsMeshTriangles(
+	FRDGBuilder& GraphBuilder,
+	FGlobalShaderMap* ShaderMap,
+	const int32 LODIndex,
+	const FHairStrandsProjectionMeshData::LOD& ProjectionMeshData,
+	FHairStrandsProjectionHairData::HairGroup& ProjectionHairData)
+{
+	AddHairStrandInterpolateMeshTrianglesPass(GraphBuilder, ShaderMap, LODIndex, ProjectionMeshData, ProjectionHairData);
+}
+
+///////////////////////////////////////////////////////////////////////////////////////////////////
+//
+class FHairInitMeshSamplesCS : public FGlobalShader
+{
+private:
+	DECLARE_GLOBAL_SHADER(FHairInitMeshSamplesCS);
+	SHADER_USE_PARAMETER_STRUCT(FHairInitMeshSamplesCS, FGlobalShader);
+
+	using FPermutationDomain = TShaderPermutationDomain<>;
+
+	BEGIN_SHADER_PARAMETER_STRUCT(FParameters, )
+		SHADER_PARAMETER(uint32, MaxSampleCount)
+		SHADER_PARAMETER(uint32, MaxVertexCount)
+
+		SHADER_PARAMETER_SRV(Buffer, VertexPositionsBuffer)
+
+		SHADER_PARAMETER_SRV(Buffer, SampleIndicesBuffer)
+		SHADER_PARAMETER_UAV(StructuredBuffer, OutSamplePositionsBuffer)
+		END_SHADER_PARAMETER_STRUCT()
+
+public:
+	static bool ShouldCompilePermutation(const FGlobalShaderPermutationParameters& Parameters) { return Parameters.Platform == SP_PCD3D_SM5; }
+	static void ModifyCompilationEnvironment(const FGlobalShaderPermutationParameters& Parameters, FShaderCompilerEnvironment& OutEnvironment)
+	{
+		FGlobalShader::ModifyCompilationEnvironment(Parameters, OutEnvironment);
+	}
+};
+
+IMPLEMENT_GLOBAL_SHADER(FHairInitMeshSamplesCS, "/Engine/Private/HairStrands/HairStrandsSamplesInit.usf", "MainCS", SF_Compute);
+
+static void AddHairStrandInitMeshSamplesPass(
+	FRDGBuilder& GraphBuilder,
+	FGlobalShaderMap* ShaderMap,
+	const int32 LODIndex,
+	const HairStrandsTriangleType Type,
+	const FHairStrandsProjectionMeshData::LOD& MeshData,
+	FHairStrandsProjectionHairData::HairGroup& RootData)
+{
+	if (LODIndex < 0 || LODIndex >= RootData.LODDatas.Num())
+	{
+		return;
+	}
+	FHairStrandsProjectionHairData::LODData& LODData = RootData.LODDatas[LODIndex];
+	check(LODData.LODIndex == LODIndex);
+
+	const uint32 SectionCount = MeshData.Sections.Num();
+	if (SectionCount > 0 && LODData.SampleCount > 0)
+	{
+		FHairInitMeshSamplesCS::FParameters* Parameters = GraphBuilder.AllocParameters<FHairInitMeshSamplesCS::FParameters>();
+
+		Parameters->MaxVertexCount = MeshData.Sections[0].TotalVertexCount;
+		Parameters->VertexPositionsBuffer = MeshData.Sections[0].PositionBuffer;
+
+		Parameters->MaxSampleCount = LODData.SampleCount;
+		Parameters->SampleIndicesBuffer = LODData.MeshSampleIndicesBuffer->SRV;
+		if (Type == HairStrandsTriangleType::RestPose)
+		{
+			Parameters->OutSamplePositionsBuffer = LODData.RestSamplePositionsBuffer->UAV;
+		}
+		else if (Type == HairStrandsTriangleType::DeformedPose)
+		{
+			Parameters->OutSamplePositionsBuffer = LODData.DeformedSamplePositionsBuffer->UAV;
+		}
+		else
+		{
+			return;
+		}
+
+		const FIntVector DispatchGroupCount = FComputeShaderUtils::GetGroupCount(LODData.SampleCount, 128);
+		check(DispatchGroupCount.X < 65536);
+		TShaderMapRef<FHairInitMeshSamplesCS> ComputeShader(ShaderMap);
+		FComputeShaderUtils::AddPass(
+			GraphBuilder,
+			RDG_EVENT_NAME("HairStrandsInitMeshSamples"),
+			ComputeShader,
+			Parameters,
+			DispatchGroupCount);
+	}
+}
+
+void InitHairStrandsMeshSamples(
+	FRDGBuilder& GraphBuilder,
+	FGlobalShaderMap* ShaderMap,
+	const int32 LODIndex,
+	const HairStrandsTriangleType Type,
+	const FHairStrandsProjectionMeshData::LOD& ProjectionMeshData,
+	FHairStrandsProjectionHairData::HairGroup& ProjectionHairData)
+{
+	AddHairStrandInitMeshSamplesPass(GraphBuilder, ShaderMap, LODIndex, Type, ProjectionMeshData, ProjectionHairData);
+}
+
+///////////////////////////////////////////////////////////////////////////////////////////////////
+//
+class FHairUpdateMeshSamplesCS : public FGlobalShader
+{
+private:
+	DECLARE_GLOBAL_SHADER(FHairUpdateMeshSamplesCS);
+	SHADER_USE_PARAMETER_STRUCT(FHairUpdateMeshSamplesCS, FGlobalShader);
+
+	using FPermutationDomain = TShaderPermutationDomain<>;
+
+	BEGIN_SHADER_PARAMETER_STRUCT(FParameters, )
+		SHADER_PARAMETER(uint32, MaxSampleCount)
+
+		SHADER_PARAMETER_SRV(Buffer, SampleIndicesBuffer)
+		SHADER_PARAMETER_SRV(Buffer, InterpolationWeightsBuffer)
+		SHADER_PARAMETER_SRV(StructuredBuffer, SampleRestPositionsBuffer)
+		SHADER_PARAMETER_SRV(StructuredBuffer, SampleDeformedPositionsBuffer)
+		SHADER_PARAMETER_UAV(StructuredBuffer, OutSampleDeformationsBuffer)
+		END_SHADER_PARAMETER_STRUCT()
+
+public:
+	static bool ShouldCompilePermutation(const FGlobalShaderPermutationParameters& Parameters) { return Parameters.Platform == SP_PCD3D_SM5; }
+	static void ModifyCompilationEnvironment(const FGlobalShaderPermutationParameters& Parameters, FShaderCompilerEnvironment& OutEnvironment)
+	{
+		FGlobalShader::ModifyCompilationEnvironment(Parameters, OutEnvironment);
+	}
+};
+
+IMPLEMENT_GLOBAL_SHADER(FHairUpdateMeshSamplesCS, "/Engine/Private/HairStrands/HairStrandsSamplesUpdate.usf", "MainCS", SF_Compute);
+
+static void AddHairStrandUpdateMeshSamplesPass(
+	FRDGBuilder& GraphBuilder,
+	FGlobalShaderMap* ShaderMap,
+	const int32 LODIndex,
+	const FHairStrandsProjectionMeshData::LOD& MeshData,
+	FHairStrandsProjectionHairData::HairGroup& RootData)
+{
+	if (LODIndex < 0 || LODIndex >= RootData.LODDatas.Num())
+	{
+		return;
+	}
+	FHairStrandsProjectionHairData::LODData& LODData = RootData.LODDatas[LODIndex];
+	check(LODData.LODIndex == LODIndex);
+
+	const uint32 SectionCount = MeshData.Sections.Num();
+	if (SectionCount > 0 && LODData.SampleCount > 0)
+	{
+		FHairUpdateMeshSamplesCS::FParameters* Parameters = GraphBuilder.AllocParameters<FHairUpdateMeshSamplesCS::FParameters>();
+
+		Parameters->MaxSampleCount = LODData.SampleCount;
+		Parameters->SampleIndicesBuffer = LODData.MeshSampleIndicesBuffer->SRV;
+		Parameters->InterpolationWeightsBuffer = LODData.MeshInterpolationWeightsBuffer->SRV;
+		Parameters->SampleRestPositionsBuffer = LODData.RestSamplePositionsBuffer->SRV;
+		Parameters->SampleDeformedPositionsBuffer = LODData.DeformedSamplePositionsBuffer->SRV;
+		Parameters->OutSampleDeformationsBuffer = LODData.MeshSampleWeightsBuffer->UAV;
+
+		const FIntVector DispatchGroupCount = FComputeShaderUtils::GetGroupCount(LODData.SampleCount, 128);
+		check(DispatchGroupCount.X < 65536);
+		TShaderMapRef<FHairUpdateMeshSamplesCS> ComputeShader(ShaderMap);
+		FComputeShaderUtils::AddPass(
+			GraphBuilder,
+			RDG_EVENT_NAME("HairStrandsUpdateMeshSamples"),
+			ComputeShader,
+			Parameters,
+			DispatchGroupCount);
+	}
+}
+
+void UpdateHairStrandsMeshSamples(
+	FRDGBuilder& GraphBuilder,
+	FGlobalShaderMap* ShaderMap,
+	const int32 LODIndex,
+	const FHairStrandsProjectionMeshData::LOD& ProjectionMeshData,
+	FHairStrandsProjectionHairData::HairGroup& ProjectionHairData)
+{
+	AddHairStrandUpdateMeshSamplesPass(GraphBuilder, ShaderMap, LODIndex, ProjectionMeshData, ProjectionHairData);
+}
+
+///////////////////////////////////////////////////////////////////////////////////////////////////
+
+// Generate follicle mask texture
+BEGIN_SHADER_PARAMETER_STRUCT(FHairFollicleMaskParameters, )
+	SHADER_PARAMETER(FVector2D, OutputResolution)
+	SHADER_PARAMETER(uint32, MaxRootCount)
+	SHADER_PARAMETER(uint32, Channel)
+	SHADER_PARAMETER(uint32, KernelSizeInPixels)
+
+	SHADER_PARAMETER_SRV(Buffer, TrianglePosition0Buffer)
+	SHADER_PARAMETER_SRV(Buffer, TrianglePosition1Buffer)
+	SHADER_PARAMETER_SRV(Buffer, TrianglePosition2Buffer)
+	SHADER_PARAMETER_SRV(Buffer, RootBarycentricBuffer)
+
+	SHADER_PARAMETER_STRUCT_REF(FViewUniformShaderParameters, ViewUniformBuffer)
+	RENDER_TARGET_BINDING_SLOTS()
+END_SHADER_PARAMETER_STRUCT()
+
+class FHairFollicleMask : public FGlobalShader
+{
+public:
+	static bool ShouldCompilePermutation(const FGlobalShaderPermutationParameters& Parameters)
+	{
+		return IsHairStrandsSupported(Parameters.Platform);
+	}
+	static void ModifyCompilationEnvironment(const FGlobalShaderPermutationParameters& Parameters, FShaderCompilerEnvironment& OutEnvironment)
+	{
+		FGlobalShader::ModifyCompilationEnvironment(Parameters, OutEnvironment);
+		OutEnvironment.SetDefine(TEXT("SHADER_FOLLICLE_MASK"), 1);
+	}
+
+	FHairFollicleMask() = default;
+	FHairFollicleMask(const CompiledShaderInitializerType& Initializer) : FGlobalShader(Initializer) {}
+};
+
+class FHairFollicleMaskVS : public FHairFollicleMask
+{
+	DECLARE_GLOBAL_SHADER(FHairFollicleMaskVS);
+	SHADER_USE_PARAMETER_STRUCT(FHairFollicleMaskVS, FHairFollicleMask);
+
+	BEGIN_SHADER_PARAMETER_STRUCT(FParameters, )
+		SHADER_PARAMETER_STRUCT_INCLUDE(FHairFollicleMaskParameters, Pass)
+	END_SHADER_PARAMETER_STRUCT()
+};
+
+class FHairFollicleMaskPS : public FHairFollicleMask
+{
+	DECLARE_GLOBAL_SHADER(FHairFollicleMaskPS);
+	SHADER_USE_PARAMETER_STRUCT(FHairFollicleMaskPS, FHairFollicleMask);
+
+	BEGIN_SHADER_PARAMETER_STRUCT(FParameters, )
+		SHADER_PARAMETER_STRUCT_INCLUDE(FHairFollicleMaskParameters, Pass)
+	END_SHADER_PARAMETER_STRUCT()
+};
+
+IMPLEMENT_GLOBAL_SHADER(FHairFollicleMaskPS, "/Engine/Private/HairStrands/HairStrandsFollicleMask.usf", "MainPS", SF_Pixel);
+IMPLEMENT_GLOBAL_SHADER(FHairFollicleMaskVS, "/Engine/Private/HairStrands/HairStrandsFollicleMask.usf", "MainVS", SF_Vertex);
+
+static void AddFollicleMaskPass(
+	FRDGBuilder& GraphBuilder,
+	FGlobalShaderMap* ShaderMap,
+	const bool bNeedClear,
+	const uint32 KernelSizeInPixels,
+	const uint32 Channel,
+	const uint32 LODIndex,
+	const FHairStrandsProjectionHairData::HairGroup& HairData,
+	FRDGTextureRef OutTexture)
+{
+	if (LODIndex >= uint32(HairData.LODDatas.Num()) || HairData.RootCount == 0)
+		return;
+
+	const FHairStrandsProjectionHairData::LODData& LODData = HairData.LODDatas[LODIndex];
+	if (!LODData.RootTriangleBarycentricBuffer ||
+		!LODData.RestRootTrianglePosition0Buffer ||
+		!LODData.RestRootTrianglePosition1Buffer ||
+		!LODData.RestRootTrianglePosition2Buffer)
+		return;
+
+	const FIntPoint OutputResolution = OutTexture->Desc.Extent;
+	FHairFollicleMaskParameters* Parameters = GraphBuilder.AllocParameters<FHairFollicleMaskParameters>();
+	Parameters->TrianglePosition0Buffer = LODData.RestRootTrianglePosition0Buffer->SRV;
+	Parameters->TrianglePosition1Buffer = LODData.RestRootTrianglePosition1Buffer->SRV;
+	Parameters->TrianglePosition2Buffer = LODData.RestRootTrianglePosition2Buffer->SRV;
+	Parameters->RootBarycentricBuffer = LODData.RootTriangleBarycentricBuffer->SRV;
+	Parameters->OutputResolution = OutputResolution;
+	Parameters->MaxRootCount = HairData.RootCount;
+	Parameters->Channel = FMath::Min(Channel, 3u);
+	Parameters->KernelSizeInPixels = FMath::Clamp(KernelSizeInPixels, 2u, 200u);
+	Parameters->RenderTargets[0] = FRenderTargetBinding(OutTexture, bNeedClear ? ERenderTargetLoadAction::EClear : ERenderTargetLoadAction::ELoad, 0);
+
+	TShaderMapRef<FHairFollicleMaskVS> VertexShader(ShaderMap);
+	TShaderMapRef<FHairFollicleMaskPS> PixelShader(ShaderMap);
+	FHairFollicleMaskVS::FParameters ParametersVS;
+	FHairFollicleMaskPS::FParameters ParametersPS;
+	ParametersVS.Pass = *Parameters;
+	ParametersPS.Pass = *Parameters;
+
+	GraphBuilder.AddPass(
+		RDG_EVENT_NAME("HairStrandsFollicleMask"),
+		Parameters,
+		ERDGPassFlags::Raster,
+		[Parameters, ParametersVS, ParametersPS, VertexShader, PixelShader, OutputResolution](FRHICommandList& RHICmdList)
+	{
+
+		RHICmdList.SetViewport(0, 0, 0.0f, OutputResolution.X, OutputResolution.Y, 1.0f);
+
+		// Apply additive blending pipeline state.
+		FGraphicsPipelineStateInitializer GraphicsPSOInit;
+		RHICmdList.ApplyCachedRenderTargets(GraphicsPSOInit);
+		GraphicsPSOInit.BlendState = TStaticBlendState<CW_RGBA, BO_Max, BF_SourceColor, BF_DestColor, BO_Max, BF_SourceAlpha, BF_DestAlpha>::GetRHI();
+		GraphicsPSOInit.RasterizerState = TStaticRasterizerState<>::GetRHI();
+		GraphicsPSOInit.DepthStencilState = TStaticDepthStencilState<false, CF_Always>::GetRHI();
+		GraphicsPSOInit.BoundShaderState.VertexDeclarationRHI = GEmptyVertexDeclaration.VertexDeclarationRHI;
+		GraphicsPSOInit.BoundShaderState.VertexShaderRHI = VertexShader.GetVertexShader();
+		GraphicsPSOInit.BoundShaderState.PixelShaderRHI = PixelShader.GetPixelShader();
+		GraphicsPSOInit.PrimitiveType = PT_TriangleList;
+		SetGraphicsPipelineState(RHICmdList, GraphicsPSOInit);
+
+		SetShaderParameters(RHICmdList, VertexShader, VertexShader.GetVertexShader(), ParametersVS);
+		SetShaderParameters(RHICmdList, PixelShader, PixelShader.GetPixelShader(), ParametersPS);
+
+		// Emit an instanced quad draw call on the order of the number of pixels on the screen.	
+		RHICmdList.DrawPrimitive(0, Parameters->MaxRootCount, 1);
+	});
+}
+
+void GenerateFolliculeMask(
+	FRDGBuilder& GraphBuilder,
+	FGlobalShaderMap* ShaderMap,
+	const FIntPoint Resolution,
+	const uint32 MipCount,
+	const uint32 KernelSizeInPixels,
+	const uint32 Channel,
+	const int32 LODIndex,
+	const FHairStrandsProjectionHairData& HairData, 
+	FRDGTextureRef& OutTexture)
+{
+	const FLinearColor ClearColor(0.0f, 0.f, 0.f, 0.f);
+
+	bool bClear = OutTexture == nullptr;
+	if (OutTexture == nullptr)
+	{
+		FRDGTextureDesc OutputDesc;
+		OutputDesc.ClearValue = FClearValueBinding(ClearColor);
+		OutputDesc.Extent.X = Resolution.X;
+		OutputDesc.Extent.Y = Resolution.Y;
+		OutputDesc.Depth = 0;
+		OutputDesc.Format = PF_R8G8B8A8;
+		OutputDesc.NumMips = MipCount;
+		OutputDesc.Flags = 0;
+		OutputDesc.TargetableFlags = TexCreate_RenderTargetable | TexCreate_ShaderResource | TexCreate_UAV;
+		OutTexture = GraphBuilder.CreateTexture(OutputDesc, TEXT("FollicleMask"));
+	}
+
+	for (const FHairStrandsProjectionHairData::HairGroup& HairGroup : HairData.HairGroups)
+	{
+		AddFollicleMaskPass(GraphBuilder, ShaderMap, bClear, KernelSizeInPixels, Channel, LODIndex, HairGroup, OutTexture);
+		bClear = false;
+	}
+}
+
+class FGenerateMipCS : public FGlobalShader
+{
+	DECLARE_GLOBAL_SHADER(FGenerateMipCS);
+	SHADER_USE_PARAMETER_STRUCT(FGenerateMipCS, FGlobalShader);
+
+	BEGIN_SHADER_PARAMETER_STRUCT(FParameters, )
+		SHADER_PARAMETER(uint32, Resolution)
+		SHADER_PARAMETER(uint32, SourceMip)
+		SHADER_PARAMETER(uint32, TargetMip)
+
+		SHADER_PARAMETER_RDG_TEXTURE_SRV(Texture2D, InTexture)
+		SHADER_PARAMETER_RDG_TEXTURE_UAV(RWTexture2D, OutTexture)
+		SHADER_PARAMETER_SAMPLER(SamplerState, LinearSampler)
+	END_SHADER_PARAMETER_STRUCT()
+
+public:
+	static bool ShouldCompilePermutation(const FGlobalShaderPermutationParameters& Parameters) { return IsHairStrandsSupported(Parameters.Platform); }
+	static void ModifyCompilationEnvironment(const FGlobalShaderPermutationParameters& Parameters, FShaderCompilerEnvironment& OutEnvironment)
+	{
+		FGlobalShader::ModifyCompilationEnvironment(Parameters, OutEnvironment);
+		OutEnvironment.SetDefine(TEXT("SHADER_GENERATE_MIPS"), 1);
+	}
+};
+
+IMPLEMENT_GLOBAL_SHADER(FGenerateMipCS, "/Engine/Private/HairStrands/HairStrandsFollicleMask.usf", "MainCS", SF_Compute);
+
+void AddComputeMipsPass(
+	FRDGBuilder& GraphBuilder,
+	FGlobalShaderMap* ShaderMap,
+	FRDGTextureRef& OutTexture)
+{
+	check(OutTexture->Desc.Extent.X == OutTexture->Desc.Extent.Y);
+	const uint32 Resolution = OutTexture->Desc.Extent.X;
+	const uint32 MipCount = OutTexture->Desc.NumMips;
+	for (uint32 MipIt = 0; MipIt < MipCount - 1; ++MipIt)
+	{
+		const uint32 SourceMipIndex = MipIt;
+		const uint32 TargetMipIndex = MipIt + 1;
+		const uint32 TargetResolution = Resolution << TargetMipIndex;
+
+		FGenerateMipCS::FParameters* Parameters = GraphBuilder.AllocParameters<FGenerateMipCS::FParameters>();
+		Parameters->InTexture = GraphBuilder.CreateSRV(FRDGTextureSRVDesc::CreateForMipLevel(OutTexture, SourceMipIndex));
+		Parameters->OutTexture = GraphBuilder.CreateUAV(FRDGTextureUAVDesc(OutTexture, TargetMipIndex));
+		Parameters->Resolution = Resolution;
+		Parameters->SourceMip = SourceMipIndex;
+		Parameters->TargetMip = TargetMipIndex;
+		Parameters->LinearSampler = TStaticSamplerState<SF_Bilinear, AM_Clamp, AM_Clamp, AM_Clamp>::GetRHI();
+
+		TShaderMapRef<FGenerateMipCS> ComputeShader(ShaderMap);
+		ClearUnusedGraphResources(ComputeShader, Parameters);
+		GraphBuilder.AddPass(
+			RDG_EVENT_NAME("HairStrandsComputeVoxelMip"),
+			Parameters,
+			ERDGPassFlags::Compute | ERDGPassFlags::GenerateMips,
+			[Parameters, ComputeShader, TargetResolution](FRHICommandList& RHICmdList)
+		{
+			const FIntVector GroupCount = FComputeShaderUtils::GetGroupCount(FIntPoint(TargetResolution, TargetResolution), FIntPoint(8, 8));
+			FComputeShaderUtils::Dispatch(RHICmdList, ComputeShader, *Parameters, GroupCount);
+		});
 	}
 }
\ No newline at end of file
diff --git a/Engine/Source/Runtime/Renderer/Private/HairStrands/HairStrandsMeshProjection.h b/Engine/Source/Runtime/Renderer/Private/HairStrands/HairStrandsMeshProjection.h
--- a/Engine/Source/Runtime/Renderer/Private/HairStrands/HairStrandsMeshProjection.h
+++ b/Engine/Source/Runtime/Renderer/Private/HairStrands/HairStrandsMeshProjection.h
@@ -5,6 +5,7 @@
 #include "CoreMinimal.h"
 #include "RendererInterface.h"
 #include "HairStrandsInterface.h"
+#include "RenderGraphResources.h"
 #include "Shader.h"
 
 /* Project hair strands onto a LOD mesh */
@@ -36,4 +37,45 @@ void UpdateHairStrandsMeshTriangles(
 	const int32 LODIndex,
 	const HairStrandsTriangleType Type,
 	const FHairStrandsProjectionMeshData::LOD& ProjectionMeshData,
-	FHairStrandsProjectionHairData::HairGroup& ProjectionHairData);
\ No newline at end of file
+	FHairStrandsProjectionHairData::HairGroup& ProjectionHairData);
+
+/* Init the samples information to be used for interpolation*/
+void InitHairStrandsMeshSamples(
+	class FRDGBuilder& GraphBuilder,
+	FGlobalShaderMap* ShaderMap,
+	const int32 LODIndex,
+	const HairStrandsTriangleType Type,
+	const FHairStrandsProjectionMeshData::LOD& ProjectionMeshData,
+	FHairStrandsProjectionHairData::HairGroup& ProjectionHairData);
+
+/* Update the samples information to be used for interpolation*/
+void UpdateHairStrandsMeshSamples(
+	FRDGBuilder& GraphBuilder,
+	FGlobalShaderMap* ShaderMap,
+	const int32 LODIndex,
+	const FHairStrandsProjectionMeshData::LOD& ProjectionMeshData,
+	FHairStrandsProjectionHairData::HairGroup& ProjectionHairData);
+
+void GenerateFolliculeMask(
+	FRDGBuilder& GraphBuilder,
+	FGlobalShaderMap* ShaderMap,
+	const FIntPoint Resolution,
+	const uint32 MipCount,
+	const uint32 KernelSizeInPixels,
+	const uint32 Channel,
+	const int32 LODIndex,
+	const FHairStrandsProjectionHairData& HairData,
+	FRDGTextureRef& OutTexture);
+
+void AddComputeMipsPass(
+	FRDGBuilder& GraphBuilder,
+	FGlobalShaderMap* ShaderMap,
+	FRDGTextureRef& OutTexture);
+
+void InterpolateHairStrandsMeshTriangles(
+	FRDGBuilder& GraphBuilder,
+	FGlobalShaderMap* ShaderMap,
+	const int32 LODIndex,
+	const FHairStrandsProjectionMeshData::LOD& ProjectionMeshData,
+	FHairStrandsProjectionHairData::HairGroup& ProjectionHairData);
+
diff --git a/Engine/Source/Runtime/Renderer/Private/HairStrands/HairStrandsRasterCommon.cpp b/Engine/Source/Runtime/Renderer/Private/HairStrands/HairStrandsRasterCommon.cpp
--- a/Engine/Source/Runtime/Renderer/Private/HairStrands/HairStrandsRasterCommon.cpp
+++ b/Engine/Source/Runtime/Renderer/Private/HairStrands/HairStrandsRasterCommon.cpp
@@ -44,6 +44,12 @@ BEGIN_GLOBAL_SHADER_PARAMETER_STRUCT(FHairVoxelizationRasterGlobalParameters, )
 	SHADER_PARAMETER(uint32, MacroGroupId)
 	SHADER_PARAMETER(FIntPoint, ViewportResolution)
 	SHADER_PARAMETER_SRV(StructuredBuffer<FVoxelizationViewInfo>, VoxelizationViewInfoBuffer)
+
+	SHADER_PARAMETER_UAV(RWTexture3D<uint>, DensityTexture)
+	SHADER_PARAMETER_UAV(RWTexture3D<uint>, TangentXTexture)
+	SHADER_PARAMETER_UAV(RWTexture3D<uint>, TangentYTexture)
+	SHADER_PARAMETER_UAV(RWTexture3D<uint>, TangentZTexture)
+	SHADER_PARAMETER_UAV(RWTexture3D<uint>, MaterialTexture)
 END_GLOBAL_SHADER_PARAMETER_STRUCT()
 IMPLEMENT_GLOBAL_SHADER_PARAMETER_STRUCT(FHairVoxelizationRasterGlobalParameters, "VoxelRasterPass");
 
@@ -58,6 +64,13 @@ static FHairVoxelizationRasterGlobalParameters ConvertToGlobalPassParameter(cons
 	Out.MacroGroupId	   = In->MacroGroupId;
 	Out.ViewportResolution = In->ViewportResolution;
 	Out.VoxelizationViewInfoBuffer = In->VoxelizationViewInfoBuffer->GetRHI();
+
+	Out.DensityTexture  = In->DensityTexture->GetRHI();
+	Out.TangentXTexture = In->TangentXTexture->GetRHI();
+	Out.TangentYTexture = In->TangentYTexture->GetRHI();
+	Out.TangentZTexture = In->TangentZTexture->GetRHI();
+	Out.MaterialTexture = In->MaterialTexture->GetRHI();
+
 	return Out;
 }
 
@@ -131,7 +144,7 @@ IMPLEMENT_MATERIAL_SHADER_TYPE(, FDeepShadowDomMeshVS, TEXT("/Engine/Private/Hai
 
 /////////////////////////////////////////////////////////////////////////////////////////
 
-template<bool bVoxelizeMaterial>
+template<bool bVoxelizeMaterial, bool bClusterCulling>
 class FVoxelMeshVS : public FMeshMaterialShader
 {
 	DECLARE_SHADER_TYPE(FVoxelMeshVS, MeshMaterial);
@@ -156,15 +169,22 @@ protected:
 
 	static void ModifyCompilationEnvironment(const FMaterialShaderPermutationParameters& Parameters, FShaderCompilerEnvironment& OutEnvironment)
 	{
+		// Note: at the moment only the plain voxelization support material voxelization
 		FMaterialShader::ModifyCompilationEnvironment(Parameters, OutEnvironment);
 		OutEnvironment.SetDefine(TEXT("MESH_RENDER_MODE"), 2);
 		OutEnvironment.SetDefine(TEXT("SUPPORT_TANGENT_PROPERTY"), bVoxelizeMaterial ? 1 : 0);
 		OutEnvironment.SetDefine(TEXT("SUPPORT_MATERIAL_PROPERTY"), bVoxelizeMaterial ? 1 : 0);
-		OutEnvironment.SetDefine(TEXT("USE_CULLED_CLUSTER"), 1);
+		OutEnvironment.SetDefine(TEXT("USE_CULLED_CLUSTER"), bClusterCulling ? 1 : 0);
 	}
 };
-IMPLEMENT_MATERIAL_SHADER_TYPE(template<>, FVoxelMeshVS<false>, TEXT("/Engine/Private/HairStrands/HairStrandsDeepShadowVS.usf"), TEXT("Main"), SF_Vertex);
-IMPLEMENT_MATERIAL_SHADER_TYPE(template<>, FVoxelMeshVS<true>, TEXT("/Engine/Private/HairStrands/HairStrandsDeepShadowVS.usf"), TEXT("Main"), SF_Vertex);
+
+typedef FVoxelMeshVS<false, false> TVoxelMeshVS_NoMaterial_NoCluster;
+typedef FVoxelMeshVS<false, true>  TVoxelMeshVS_NoMaterial_Cluster;
+typedef FVoxelMeshVS<true,  true>  TVoxelMeshVS_Material_Cluster;
+
+IMPLEMENT_MATERIAL_SHADER_TYPE(template<>, TVoxelMeshVS_NoMaterial_NoCluster, TEXT("/Engine/Private/HairStrands/HairStrandsDeepShadowVS.usf"), TEXT("Main"), SF_Vertex);
+IMPLEMENT_MATERIAL_SHADER_TYPE(template<>, TVoxelMeshVS_NoMaterial_Cluster, TEXT("/Engine/Private/HairStrands/HairStrandsDeepShadowVS.usf"), TEXT("Main"), SF_Vertex);
+IMPLEMENT_MATERIAL_SHADER_TYPE(template<>, TVoxelMeshVS_Material_Cluster, TEXT("/Engine/Private/HairStrands/HairStrandsDeepShadowVS.usf"), TEXT("Main"), SF_Vertex);
 
 /////////////////////////////////////////////////////////////////////////////////////////
 
@@ -237,14 +257,6 @@ public:
 		ERHIFeatureLevel::Type FeatureLevel = GetMaxSupportedFeatureLevel((EShaderPlatform)Initializer.Target.Platform);
 		check(FSceneInterface::GetShadingPath(FeatureLevel) != EShadingPath::Mobile);
 		PassUniformBuffer.Bind(Initializer.ParameterMap, FHairVoxelizationRasterGlobalParameters::StaticStructMetadata.GetShaderVariableName());
-		DensityTexture.Bind(Initializer.ParameterMap, TEXT("DensityTexture"));
-		if (VoxelizationType == EVoxelMeshPSType::Material)
-		{
-			TangentXTexture.Bind(Initializer.ParameterMap, TEXT("TangentXTexture"));
-			TangentYTexture.Bind(Initializer.ParameterMap, TEXT("TangentYTexture"));
-			TangentZTexture.Bind(Initializer.ParameterMap, TEXT("TangentZTexture"));
-			MaterialTexture.Bind(Initializer.ParameterMap, TEXT("MaterialTexture"));
-		}
 	}
 
 	FVoxelMeshPS() {}
@@ -261,13 +273,6 @@ public:
 		OutEnvironment.SetDefine(TEXT("SUPPORT_MATERIAL_PROPERTY"), VoxelizationType == EVoxelMeshPSType::Material ? 1 : 0);
 		OutEnvironment.SetDefine(TEXT("SUPPORT_SPARSE"), VoxelizationType == EVoxelMeshPSType::DensityVirtual ? 1 : 0);
 	}
-
-private:
-	LAYOUT_FIELD(FRWShaderParameter, DensityTexture);
-	LAYOUT_FIELD(FRWShaderParameter, TangentXTexture);
-	LAYOUT_FIELD(FRWShaderParameter, TangentYTexture);
-	LAYOUT_FIELD(FRWShaderParameter, TangentZTexture);
-	LAYOUT_FIELD(FRWShaderParameter, MaterialTexture);
 };
 IMPLEMENT_MATERIAL_SHADER_TYPE(template<>, FVoxelMeshPS<EVoxelMeshPSType::Density>, TEXT("/Engine/Private/HairStrands/HairStrandsDeepShadowPS.usf"), TEXT("MainVoxel"), SF_Pixel);
 IMPLEMENT_MATERIAL_SHADER_TYPE(template<>, FVoxelMeshPS<EVoxelMeshPSType::Material>, TEXT("/Engine/Private/HairStrands/HairStrandsDeepShadowPS.usf"), TEXT("MainVoxel"), SF_Pixel);
@@ -286,7 +291,12 @@ public:
 		FDynamicPassMeshDrawListContext* InDrawListContext,
 		const EHairStrandsRasterPassType PType);
 
-	virtual void AddMeshBatch(const FMeshBatch& RESTRICT MeshBatch, uint64 BatchElementMask, const FPrimitiveSceneProxy* RESTRICT PrimitiveSceneProxy, int32 StaticMeshId = -1) override final;
+	virtual void AddMeshBatch(const FMeshBatch& RESTRICT MeshBatch, uint64 BatchElementMask, const FPrimitiveSceneProxy* RESTRICT PrimitiveSceneProxy, int32 StaticMeshId = -1) override final
+	{
+		AddMeshBatch(MeshBatch, BatchElementMask, PrimitiveSceneProxy, StaticMeshId, false);
+	}
+
+	void AddMeshBatch(const FMeshBatch& RESTRICT MeshBatch, uint64 BatchElementMask, const FPrimitiveSceneProxy* RESTRICT PrimitiveSceneProxy, int32 StaticMeshId, bool bCullingEnable);
 
 private:
 	template<typename VertexShaderType, typename PixelShaderType>
@@ -304,7 +314,7 @@ private:
 	FMeshPassProcessorRenderState PassDrawRenderState;
 };
 
-void FHairRasterMeshProcessor::AddMeshBatch(const FMeshBatch& RESTRICT MeshBatch, uint64 BatchElementMask, const FPrimitiveSceneProxy* RESTRICT PrimitiveSceneProxy, int32 StaticMeshId)
+void FHairRasterMeshProcessor::AddMeshBatch(const FMeshBatch& RESTRICT MeshBatch, uint64 BatchElementMask, const FPrimitiveSceneProxy* RESTRICT PrimitiveSceneProxy, int32 StaticMeshId, bool bCullingEnable)
 {
 	// Determine the mesh's material and blend mode.
 	const FMaterialRenderProxy* FallbackMaterialRenderProxyPtr = nullptr;
@@ -325,11 +335,13 @@ void FHairRasterMeshProcessor::AddMeshBatch(const FMeshBatch& RESTRICT MeshBatch
 		else if (RasterPassType == EHairStrandsRasterPassType::DeepOpacityMap)
 			Process<FDeepShadowDomMeshVS, FDeepShadowDomMeshPS>(MeshBatch, BatchElementMask, PrimitiveSceneProxy, StaticMeshId, MaterialRenderProxy, Material, MeshFillMode, MeshCullMode);
 		else if (RasterPassType == EHairStrandsRasterPassType::Voxelization)
-			Process<FVoxelMeshVS<false>, FVoxelMeshPS<EVoxelMeshPSType::Density>>(MeshBatch, BatchElementMask, PrimitiveSceneProxy, StaticMeshId, MaterialRenderProxy, Material, MeshFillMode, MeshCullMode);
+			Process<FVoxelMeshVS<false, true>, FVoxelMeshPS<EVoxelMeshPSType::Density>>(MeshBatch, BatchElementMask, PrimitiveSceneProxy, StaticMeshId, MaterialRenderProxy, Material, MeshFillMode, MeshCullMode);
 		else if (RasterPassType == EHairStrandsRasterPassType::VoxelizationMaterial)
-			Process<FVoxelMeshVS<true>, FVoxelMeshPS<EVoxelMeshPSType::Material>>(MeshBatch, BatchElementMask, PrimitiveSceneProxy, StaticMeshId, MaterialRenderProxy, Material, MeshFillMode, MeshCullMode);
-		else if (RasterPassType == EHairStrandsRasterPassType::VoxelizationVirtual)
-			Process<FVoxelMeshVS<false>, FVoxelMeshPS<EVoxelMeshPSType::DensityVirtual>>(MeshBatch, BatchElementMask, PrimitiveSceneProxy, StaticMeshId, MaterialRenderProxy, Material, MeshFillMode, MeshCullMode);
+			Process<FVoxelMeshVS<true, true>, FVoxelMeshPS<EVoxelMeshPSType::Material>>(MeshBatch, BatchElementMask, PrimitiveSceneProxy, StaticMeshId, MaterialRenderProxy, Material, MeshFillMode, MeshCullMode);
+		else if (RasterPassType == EHairStrandsRasterPassType::VoxelizationVirtual && bCullingEnable)
+			Process<FVoxelMeshVS<false, true>, FVoxelMeshPS<EVoxelMeshPSType::DensityVirtual>>(MeshBatch, BatchElementMask, PrimitiveSceneProxy, StaticMeshId, MaterialRenderProxy, Material, MeshFillMode, MeshCullMode);
+		else if (RasterPassType == EHairStrandsRasterPassType::VoxelizationVirtual && !bCullingEnable)
+			Process<FVoxelMeshVS<false, false>, FVoxelMeshPS<EVoxelMeshPSType::DensityVirtual>>(MeshBatch, BatchElementMask, PrimitiveSceneProxy, StaticMeshId, MaterialRenderProxy, Material, MeshFillMode, MeshCullMode);
 	}
 }
 
@@ -410,6 +422,7 @@ void AddHairStrandsRasterPass(
 	const EHairStrandsRasterPassType RasterPassType,
 	const FIntRect& ViewportRect,
 	const FVector4& HairRenderInfo,
+	const uint32 HairRenderInfoBits,
 	const FVector& RasterDirection,
 	TPassParameter* PassParameters)
 {
@@ -430,7 +443,7 @@ void AddHairStrandsRasterPass(
 		GetPassName(RasterPassType),
 		PassParameters,
 		ERDGPassFlags::Raster,
-		[PassParameters, Scene = Scene, ViewInfo, RasterPassType, &PrimitiveSceneInfos, ViewportRect, HairRenderInfo, RasterDirection](FRHICommandListImmediate& RHICmdList)
+		[PassParameters, Scene = Scene, ViewInfo, RasterPassType, &PrimitiveSceneInfos, ViewportRect, HairRenderInfo, HairRenderInfoBits, RasterDirection](FRHICommandListImmediate& RHICmdList)
 	{
 		check(RHICmdList.IsInsideRenderPass());
 		check(IsInRenderingThread());
@@ -442,6 +455,7 @@ void AddHairStrandsRasterPass(
 
 		TUniformBufferRef<FViewUniformShaderParameters> ViewUniformShaderParameters;
 		ViewInfo->CachedViewUniformShaderParameters->HairRenderInfo = HairRenderInfo;
+		ViewInfo->CachedViewUniformShaderParameters->HairRenderInfoBits = HairRenderInfoBits;
 
 		const FVector SavedViewForward = ViewInfo->CachedViewUniformShaderParameters->ViewForward;	
 		ViewInfo->CachedViewUniformShaderParameters->ViewForward = RasterDirection;
@@ -491,11 +505,12 @@ void AddHairStrandsRasterPass(
 
 		FHairRasterMeshProcessor HairRasterMeshProcessor(Scene, ViewInfo /* is a SceneView */, DrawRenderState, &ShadowContext, RasterPassType);
 
-		for (const FHairStrandsMacroGroupData::PrimitiveInfo& PrimitveInfo : PrimitiveSceneInfos)
+		for (const FHairStrandsMacroGroupData::PrimitiveInfo& PrimitiveInfo : PrimitiveSceneInfos)
 		{
-			const FMeshBatch& MeshBatch = *PrimitveInfo.MeshBatchAndRelevance.Mesh;
+			const bool bCullingEnable = PrimitiveInfo.IsCullingEnable();
+			const FMeshBatch& MeshBatch = *PrimitiveInfo.MeshBatchAndRelevance.Mesh;
 			const uint64 BatchElementMask = ~0ull;
-			HairRasterMeshProcessor.AddMeshBatch(MeshBatch, BatchElementMask, PrimitveInfo.MeshBatchAndRelevance.PrimitiveSceneProxy);
+			HairRasterMeshProcessor.AddMeshBatch(MeshBatch, BatchElementMask, PrimitiveInfo.MeshBatchAndRelevance.PrimitiveSceneProxy, -1 , bCullingEnable);
 		}
 
 		if (VisibleMeshDrawCommands.Num() > 0)
@@ -515,6 +530,7 @@ void AddHairDeepShadowRasterPass(
 	const EHairStrandsRasterPassType PassType,
 	const FIntRect& ViewportRect,
 	const FVector4& HairRenderInfo,
+	const uint32 HairRenderInfoBits,
 	const FVector& LightDirection,
 	FHairDeepShadowRasterPassParameters* PassParameters)
 {
@@ -528,6 +544,7 @@ void AddHairDeepShadowRasterPass(
 		PassType, 
 		ViewportRect, 
 		HairRenderInfo, 
+		HairRenderInfoBits,
 		LightDirection,
 		PassParameters);
 }
@@ -540,6 +557,7 @@ void AddHairVoxelizationRasterPass(
 	const EHairStrandsRasterPassType PassType,
 	const FIntRect& ViewportRect,
 	const FVector4& HairRenderInfo,
+	const uint32 HairRenderInfoBits,
 	const FVector& RasterDirection,
 	FHairVoxelizationRasterPassParameters* PassParameters)
 {
@@ -553,6 +571,7 @@ void AddHairVoxelizationRasterPass(
 		PassType, 
 		ViewportRect, 
 		HairRenderInfo, 
+		HairRenderInfoBits,
 		RasterDirection,
 		PassParameters);
-}
+}
\ No newline at end of file
diff --git a/Engine/Source/Runtime/Renderer/Private/HairStrands/HairStrandsRasterCommon.h b/Engine/Source/Runtime/Renderer/Private/HairStrands/HairStrandsRasterCommon.h
--- a/Engine/Source/Runtime/Renderer/Private/HairStrands/HairStrandsRasterCommon.h
+++ b/Engine/Source/Runtime/Renderer/Private/HairStrands/HairStrandsRasterCommon.h
@@ -42,6 +42,7 @@ void AddHairDeepShadowRasterPass(
 	const EHairStrandsRasterPassType ShadowPassType,
 	const FIntRect& ViewportRect,
 	const FVector4& HairRenderInfo,
+	const uint32 HairRenderInfoBits,
 	const FVector& LightDirection,
 	FHairDeepShadowRasterPassParameters* PassParameters);
 
@@ -74,5 +75,6 @@ void AddHairVoxelizationRasterPass(
 	const EHairStrandsRasterPassType ShadowPassType,
 	const FIntRect& ViewportRect,
 	const FVector4& HairRenderInfo,
+	const uint32 HairRenderInfoBits,
 	const FVector& RasterDirection,
 	FHairVoxelizationRasterPassParameters* PassParameters);
diff --git a/Engine/Source/Runtime/Renderer/Private/HairStrands/HairStrandsRendering.h b/Engine/Source/Runtime/Renderer/Private/HairStrands/HairStrandsRendering.h
--- a/Engine/Source/Runtime/Renderer/Private/HairStrands/HairStrandsRendering.h
+++ b/Engine/Source/Runtime/Renderer/Private/HairStrands/HairStrandsRendering.h
@@ -39,11 +39,8 @@ enum class EHairStrandsInterpolationType
 void RunHairStrandsInterpolation(
 	FRHICommandListImmediate& RHICmdList, 
 	EWorldType::Type WorldType, 
+	const class FGPUSkinCache* GPUSkinCache,
 	const struct FShaderDrawDebugData* DebugShaderData,
 	FGlobalShaderMap* ShaderMap, 
 	EHairStrandsInterpolationType Type,
-	FHairStrandClusterData* ClusterData);
-
-void RunHairStrandsBindingQueries(
-	FRHICommandListImmediate& RHICmdList,
-	TShaderMap<FGlobalShaderType>* ShaderMap);
\ No newline at end of file
+	FHairStrandClusterData* ClusterData);
\ No newline at end of file
diff --git a/Engine/Source/Runtime/Renderer/Private/HairStrands/HairStrandsUtils.cpp b/Engine/Source/Runtime/Renderer/Private/HairStrands/HairStrandsUtils.cpp
--- a/Engine/Source/Runtime/Renderer/Private/HairStrands/HairStrandsUtils.cpp
+++ b/Engine/Source/Runtime/Renderer/Private/HairStrands/HairStrandsUtils.cpp
@@ -19,6 +19,9 @@ static FAutoConsoleVariableRef CVarHairLocalScattering(TEXT("r.HairStrands.Compo
 static float GStrandHairRasterizationScale = 0.5f; // For no AA without TAA, a good value is: 1.325f (Empirical)
 static FAutoConsoleVariableRef CVarStrandHairRasterizationScale(TEXT("r.HairStrands.RasterizationScale"), GStrandHairRasterizationScale, TEXT("Rasterization scale to snap strand to pixel"), ECVF_Scalability | ECVF_RenderThreadSafe);
 
+static float GStrandHairStableRasterizationScale = 1.0f; // For no AA without TAA, a good value is: 1.325f (Empirical)
+static FAutoConsoleVariableRef CVarStrandHairStableRasterizationScale(TEXT("r.HairStrands.StableRasterizationScale"), GStrandHairStableRasterizationScale, TEXT("Rasterization scale to snap strand to pixel for 'stable' hair option. This value can't go below 1."), ECVF_Scalability | ECVF_RenderThreadSafe);
+
 static float GStrandHairVelocityRasterizationScale = 1.5f; // Tuned based on heavy motion example (e.g., head shaking)
 static FAutoConsoleVariableRef CVarStrandHairMaxRasterizationScale(TEXT("r.HairStrands.VelocityRasterizationScale"), GStrandHairVelocityRasterizationScale, TEXT("Rasterization scale to snap strand to pixel under high velocity"), ECVF_Scalability | ECVF_RenderThreadSafe);
 
@@ -89,8 +92,10 @@ FMinHairRadiusAtDepth1 ComputeMinStrandRadiusAtDepth1(
 	// Scales strand to covers a bit more than a pixel and insure at least one sample point is hit
 	const float PrimaryRasterizationScale = OverrideStrandHairRasterizationScale > 0 ? OverrideStrandHairRasterizationScale : GStrandHairRasterizationScale;
 	const float VelocityRasterizationScale = OverrideStrandHairRasterizationScale > 0 ? OverrideStrandHairRasterizationScale : GStrandHairVelocityRasterizationScale;
+	const float StableRasterizationScale = FMath::Max(1.f, GStrandHairStableRasterizationScale);
 	Out.Primary = InternalMinRadiusAtDepth1(PrimaryRasterizationScale);
 	Out.Velocity = InternalMinRadiusAtDepth1(VelocityRasterizationScale);
+	Out.Stable = InternalMinRadiusAtDepth1(StableRasterizationScale);
 
 	return Out;
 }
@@ -110,7 +115,7 @@ void ComputeWorldToLightClip(
 	const float MaxZ = FMath::Max(0.2f, FVector::Distance(LightPosition, SphereBound.Center)) + SphereBound.W;
 
 	const float StrandHairRasterizationScale = GStrandHairShadowRasterizationScale ? GStrandHairShadowRasterizationScale : GStrandHairRasterizationScale;
-
+	const float StrandHairStableRasterizationScale = FMath::Max(GStrandHairStableRasterizationScale, 1.0f);
 	MinStrandRadiusAtDepth1 = FMinHairRadiusAtDepth1();
 	WorldToClipTransform = FMatrix::Identity;
 	if (LightType == LightType_Directional)
@@ -119,7 +124,10 @@ void ComputeWorldToLightClip(
 		FReversedZOrthoMatrix OrthoMatrix(SphereRadius, SphereRadius, 1.f / (2 * SphereRadius), 0);
 		FLookAtMatrix LookAt(SphereBound.Center - LightDirection * SphereRadius, SphereBound.Center, FVector(0, 0, 1));
 		WorldToClipTransform = LookAt * OrthoMatrix;
-		MinStrandRadiusAtDepth1.Primary = StrandHairRasterizationScale * SphereRadius / FMath::Min(ShadowResolution.X, ShadowResolution.Y);
+
+		const float RadiusAtDepth1 = SphereRadius / FMath::Min(ShadowResolution.X, ShadowResolution.Y);
+		MinStrandRadiusAtDepth1.Stable = RadiusAtDepth1 * StrandHairStableRasterizationScale;
+		MinStrandRadiusAtDepth1.Primary = RadiusAtDepth1 * StrandHairRasterizationScale;
 		MinStrandRadiusAtDepth1.Velocity = MinStrandRadiusAtDepth1.Primary;
 	}
 	else if (LightType == LightType_Spot || LightType == LightType_Point)
@@ -245,24 +253,27 @@ FIntPoint GetVendorOptimalGroupSize2D()
 
 FVector4 PackHairRenderInfo(
 	float PrimaryRadiusAtDepth1,
+	float StableRadiusAtDepth1,
 	float VelocityRadiusAtDepth1,
-	float VelocityMagnitudeScale,
-	bool bIsOrtho,
-	bool bIsGPUDriven)
+	float VelocityMagnitudeScale)
 {
-	uint32 BitField = 0;
-	BitField |= bIsOrtho ? 0x1 : 0;
-	BitField |= bIsGPUDriven ? 0x2 : 0;
-
 	FVector4 Out;
 	Out.X = PrimaryRadiusAtDepth1;
-	Out.Y = VelocityRadiusAtDepth1;
-	Out.Z = VelocityMagnitudeScale;
-	Out.W = *((float*)(&BitField));
-
+	Out.Y = StableRadiusAtDepth1;
+	Out.Z = VelocityRadiusAtDepth1;
+	Out.W = VelocityMagnitudeScale;
 	return Out;
 }
 
+uint32 PackHairRenderInfoBits(
+	bool  bIsOrtho,
+	bool  bIsGPUDriven)
+{
+	uint32 BitField = 0;
+	BitField |= bIsOrtho ? 0x1 : 0;
+	BitField |= bIsGPUDriven ? 0x2 : 0;
+	return BitField;
+}
 FHairStrandsProjectionMeshData ExtractMeshData(FSkeletalMeshRenderData* RenderData)
 {
 	FHairStrandsProjectionMeshData MeshData;
diff --git a/Engine/Source/Runtime/Renderer/Private/HairStrands/HairStrandsUtils.h b/Engine/Source/Runtime/Renderer/Private/HairStrands/HairStrandsUtils.h
--- a/Engine/Source/Runtime/Renderer/Private/HairStrands/HairStrandsUtils.h
+++ b/Engine/Source/Runtime/Renderer/Private/HairStrands/HairStrandsUtils.h
@@ -56,7 +56,10 @@ RENDERER_API bool IsHairStrandsSupported(const EShaderPlatform Platform);
 
 FVector4 PackHairRenderInfo(
 	float PrimaryRadiusAtDepth1,
+	float StableRadiusAtDepth1,
 	float VelocityRadiusAtDepth1,
-	float VelocityScale,
-	bool  bIsOrtho,
-	bool  bIsGPUDriven);
\ No newline at end of file
+	float VelocityScale);
+
+uint32 PackHairRenderInfoBits(
+	bool bIsOrtho,
+	bool bIsGPUDriven);
\ No newline at end of file
