Integrate hair changes (rendering+simulation)

diff --git a/Engine/Source/Runtime/Renderer/Private/HairStrands/HairStrandsVisibility.cpp b/Engine/Source/Runtime/Renderer/Private/HairStrands/HairStrandsVisibility.cpp
--- a/Engine/Source/Runtime/Renderer/Private/HairStrands/HairStrandsVisibility.cpp
+++ b/Engine/Source/Runtime/Renderer/Private/HairStrands/HairStrandsVisibility.cpp
@@ -30,7 +30,7 @@ static FAutoConsoleVariableRef CVarHairStrandsMaterialCompactionDepthThreshold(T
 static FAutoConsoleVariableRef CVarHairStrandsMaterialCompactionTangentThreshold(TEXT("r.HairStrands.MaterialCompaction.TangentThreshold"), GHairStrandsMaterialCompactionTangentThreshold, TEXT("Compaciton threshold for tangent value for material compaction (in degrees). Default 10 deg."));
 
 static int32 GHairVisibilitySampleCount = 8;
-static FAutoConsoleVariableRef CVarHairVisibilitySampleCount(TEXT("r.HairStrands.VisibilitySampleCount"), GHairVisibilitySampleCount, TEXT("Hair strands visibility sample count"));
+static FAutoConsoleVariableRef CVarHairVisibilitySampleCount(TEXT("r.HairStrands.VisibilitySampleCount"), GHairVisibilitySampleCount, TEXT("Hair strands visibility sample count (4 or 8)"));
 
 static int32 GHairClearVisibilityBuffer = 0;
 static FAutoConsoleVariableRef CVarHairClearVisibilityBuffer(TEXT("r.HairStrands.VisibilityClear"), GHairClearVisibilityBuffer, TEXT("Clear hair strands visibility buffer"));
@@ -50,12 +50,10 @@ static FAutoConsoleVariableRef CVarGHairVisibilityPPLLMeanListElementCountPerPix
 static int32 GHairVisibilityPPLLMaxRenderNodePerPixel = 16;
 static FAutoConsoleVariableRef CVarGHairVisibilityPPLLMeanNodeCountPerPixel(TEXT("r.HairStrands.VisibilityPPLLMaxRenderNodePerPixel"), GHairVisibilityPPLLMaxRenderNodePerPixel, TEXT("The maximum number of node allowed to be independently shaded and composited per pixel. Total amount of node will be width*height*VisibilityPPLLMaxRenderNodePerPixel. The last node is used to aggregate all furthest strands to shade into a single one."));
 
-static int32 GHairStrandsVisibilityMaterialPass = 0;
+static int32 GHairStrandsVisibilityMaterialPass = 1;
 static FAutoConsoleVariableRef CVarHairStrandsVisibilityMaterialPass(TEXT("r.HairStrands.Visibility.MaterialPass"), GHairStrandsVisibilityMaterialPass, TEXT("Enable the deferred material pass evaluation after the hair visibility is resolved."));
 
-static int32 GHairStrandsViewHairCount = 1;
 static float GHairStrandsViewHairCountDepthDistanceThreshold = 30.f;
-static FAutoConsoleVariableRef CVarHairStrandsViewHairCount(TEXT("r.HairStrands.Visibility.HairCount"), GHairStrandsViewHairCount, TEXT("Enable the computation of 'view-hair-count' during the transmission pass."));
 static FAutoConsoleVariableRef CVarHairStrandsViewHairCountDepthDistanceThreshold(TEXT("r.HairStrands.Visibility.HairCount.DistanceThreshold"), GHairStrandsViewHairCountDepthDistanceThreshold, TEXT("Distance threshold defining if opaque depth get injected into the 'view-hair-count' buffer."));
 
 /////////////////////////////////////////////////////////////////////////////////////////
@@ -91,21 +89,22 @@ enum EHairVisibilityRenderMode
 	HairVisibilityRenderModeCount
 };
 
-EHairVisibilityRenderMode GetHairVisibilityRenderMode()
+static EHairVisibilityRenderMode GetHairVisibilityRenderMode()
 {
 	return GHairVisibilityPPLL > 0 ? HairVisibilityRenderMode_PPLL : HairVisibilityRenderMode_MSAA;
 }
 
-uint32 GetPPLLMeanListElementCountPerPixel()
+static uint32 GetPPLLMeanListElementCountPerPixel()
 {
 	return GHairVisibilityPPLLMeanListElementCountPerPixel;
 }
-uint32 GetPPLLMaxTotalListElementCount(FIntPoint Resolution)
+
+static uint32 GetPPLLMaxTotalListElementCount(FIntPoint Resolution)
 {
 	return Resolution.X * Resolution.Y * GetPPLLMeanListElementCountPerPixel();
 }
 
-uint32 GetPPLLMaxRenderNodePerPixel()
+static uint32 GetPPLLMaxRenderNodePerPixel()
 {
 	// The following must match the FPPLL permutation of FHairVisibilityPrimitiveIdCompactionCS.
 	if (GHairVisibilityPPLLMaxRenderNodePerPixel == 0)
@@ -126,29 +125,28 @@ uint32 GetPPLLMaxRenderNodePerPixel()
 	}
 	// If more is needed: please check out EncodeNodeDesc from HairStrandsVisibilityCommon.ush to verify node count representation limitations.
 }
-uint32 GetPPLLMaxTotalRenderNode(FIntPoint Resolution)
-{
-	return Resolution.X * Resolution.Y * GetPPLLMaxRenderNodePerPixel();
-}
 
-uint32 GetHairVisibilitySampleCount()
+static uint32 GetMSAASampleCount()
 {
-	return GetHairVisibilityRenderMode() == HairVisibilityRenderMode_MSAA ? FMath::Clamp(GHairVisibilitySampleCount,1,16) : 1;
+	// Only support 4 or 8 at the moment
+	check(GetHairVisibilityRenderMode() == HairVisibilityRenderMode_MSAA);
+	return GHairVisibilitySampleCount <= 4 ? 4 : 8;
 }
 
-void SetUpViewHairRenderInfo(const FViewInfo& ViewInfo, bool bEnableMSAA, FVector4& OutHairRenderInfo)
+void SetUpViewHairRenderInfo(const FViewInfo& ViewInfo, bool bEnableMSAA, FVector4& OutHairRenderInfo, uint32& OutHairRenderInfoBits)
 {
 	FVector2D PixelVelocity(1.f / (ViewInfo.ViewRect.Width() * 2), 1.f / (ViewInfo.ViewRect.Height() * 2));
 	const float VelocityMagnitudeScale = FMath::Clamp(CVarHairVelocityMagnitudeScale.GetValueOnAnyThread(), 0, 512) * FMath::Min(PixelVelocity.X, PixelVelocity.Y);
 
 	// In the case we render coverage, we need to override some view uniform shader parameters to account for the change in MSAA sample count.
-	const uint32 HairVisibilitySampleCount = bEnableMSAA ? GetHairVisibilitySampleCount() : 1;	// The coverage pass does not use MSAA
+	const uint32 HairVisibilitySampleCount = bEnableMSAA ? GetMSAASampleCount() : 1;	// The coverage pass does not use MSAA
 	const float RasterizationScaleOverride = 0.0f;	// no override
 	FMinHairRadiusAtDepth1 MinHairRadius = ComputeMinStrandRadiusAtDepth1(
 		FIntPoint(ViewInfo.UnconstrainedViewRect.Width(), ViewInfo.UnconstrainedViewRect.Height()), ViewInfo.FOV, HairVisibilitySampleCount, RasterizationScaleOverride);
 
 	TUniformBufferRef<FViewUniformShaderParameters> ViewUniformShaderParameters;
-	OutHairRenderInfo = PackHairRenderInfo(MinHairRadius.Primary, MinHairRadius.Velocity, VelocityMagnitudeScale, !ViewInfo.IsPerspectiveProjection(), false);
+	OutHairRenderInfo = PackHairRenderInfo(MinHairRadius.Primary, MinHairRadius.Stable, MinHairRadius.Velocity, VelocityMagnitudeScale);
+	OutHairRenderInfoBits = PackHairRenderInfoBits(!ViewInfo.IsPerspectiveProjection(), false);
 }
 
 static bool IsCompatibleWithHairVisibility(const FMeshMaterialShaderPermutationParameters& Parameters)
@@ -166,6 +164,8 @@ BEGIN_GLOBAL_SHADER_PARAMETER_STRUCT(FMaterialPassParameters, )
 	SHADER_PARAMETER_SRV(StructuredBuffer<uint>, NodeCoord)
 	SHADER_PARAMETER_SRV(StructuredBuffer<FNodeVis>, NodeVis)
 	SHADER_PARAMETER_SRV(Buffer<uint>, IndirectArgs)
+	SHADER_PARAMETER_UAV(RWStructuredBuffer<FPackedHairSample>, OutNodeData)
+	SHADER_PARAMETER_UAV(RWBuffer<float2>, OutNodeVelocity)
 END_GLOBAL_SHADER_PARAMETER_STRUCT()
 IMPLEMENT_GLOBAL_SHADER_PARAMETER_STRUCT(FMaterialPassParameters, "MaterialPassParameters");
 
@@ -224,11 +224,6 @@ public:
 		MaterialPass_MaterialId.Bind(Initializer.ParameterMap, TEXT("MaterialPass_MaterialId"));
 		MaterialPass_PrimitiveId.Bind(Initializer.ParameterMap, TEXT("MaterialPass_PrimitiveId"));
 		MaterialPass_LightChannelMask.Bind(Initializer.ParameterMap, TEXT("MaterialPass_LightChannelMask"));
-
-		// Fake binding for avoiding error from the shader compiler. These output are actually properly bound by the rendering 
-		// pass, and not directly by this material shader
-		OutNodeData.Bind(Initializer.ParameterMap, TEXT("OutNodeData"));
-		OutNodeVelocity.Bind(Initializer.ParameterMap, TEXT("OutNodeVelocity"));
 	}
 
 	FHairMaterialPS() {}
@@ -266,8 +261,6 @@ private:
 	LAYOUT_FIELD(FShaderParameter, MaterialPass_MaterialId);
 	LAYOUT_FIELD(FShaderParameter, MaterialPass_PrimitiveId);
 	LAYOUT_FIELD(FShaderParameter, MaterialPass_LightChannelMask);
-	LAYOUT_FIELD(FShaderParameter, OutNodeData);
-	LAYOUT_FIELD(FShaderParameter, OutNodeVelocity);
 };
 IMPLEMENT_MATERIAL_SHADER_TYPE(, FHairMaterialPS, TEXT("/Engine/Private/HairStrands/HairStrandsMaterialPS.usf"), TEXT("Main"), SF_Pixel);
 
@@ -402,10 +395,10 @@ FHairMaterialProcessor::FHairMaterialProcessor(
 
 BEGIN_SHADER_PARAMETER_STRUCT(FVisibilityMaterialPassParameters, )
 	SHADER_PARAMETER_RDG_TEXTURE(Texture2D<uint>, NodeIndex)
-	SHADER_PARAMETER_RDG_BUFFER_SRV(StructuredBuffer<uint>,	NodeCoord)
+	SHADER_PARAMETER_RDG_BUFFER_SRV(StructuredBuffer<uint>, NodeCoord)
 	SHADER_PARAMETER_RDG_BUFFER_SRV(StructuredBuffer<FNodeVis>, NodeVis)
 	SHADER_PARAMETER_RDG_BUFFER_SRV(Buffer<uint>, IndirectArgs)
-	SHADER_PARAMETER_RDG_BUFFER_UAV(RWStructuredBuffer<FNodeData>, OutNodeData)
+	SHADER_PARAMETER_RDG_BUFFER_UAV(RWStructuredBuffer<FPackedHairSample>, OutNodeData)
 	SHADER_PARAMETER_RDG_BUFFER_UAV(RWBuffer<float2>, OutNodeVelocity)
 	RENDER_TARGET_BINDING_SLOTS()
 END_SHADER_PARAMETER_STRUCT()
@@ -479,13 +472,17 @@ static FMaterialPassOutput AddHairMaterialPass(
 		MaterialPassParameters.NodeCoord		= PassParameters->NodeCoord->GetRHI();
 		MaterialPassParameters.NodeVis			= PassParameters->NodeVis->GetRHI();
 		MaterialPassParameters.IndirectArgs		= PassParameters->IndirectArgs->GetRHI();
+		MaterialPassParameters.OutNodeData 		= PassParameters->OutNodeData->GetRHI();
+		MaterialPassParameters.OutNodeVelocity 	= PassParameters->OutNodeVelocity->GetRHI();
 		TUniformBufferRef<FMaterialPassParameters> MaterialPassParametersBuffer = TUniformBufferRef<FMaterialPassParameters>::CreateUniformBufferImmediate(MaterialPassParameters, UniformBuffer_SingleFrame);
 
 		FMeshPassProcessorRenderState DrawRenderState(*ViewInfo, MaterialPassParametersBuffer);
+		// Note: this reference needs to persistent until SubmitMeshDrawCommands() is called, as DrawRenderState does not ref count 
+		// the view uniform buffer (raw pointer). It is only within the MeshProcessor that the uniform buffer get reference
+		TUniformBufferRef<FViewUniformShaderParameters> ViewUniformShaderParameters;
 		{
-			TUniformBufferRef<FViewUniformShaderParameters> ViewUniformShaderParameters;
 			const bool bEnableMSAA = false;
-			SetUpViewHairRenderInfo(*ViewInfo, bEnableMSAA, ViewInfo->CachedViewUniformShaderParameters->HairRenderInfo);
+			SetUpViewHairRenderInfo(*ViewInfo, bEnableMSAA, ViewInfo->CachedViewUniformShaderParameters->HairRenderInfo, ViewInfo->CachedViewUniformShaderParameters->HairRenderInfoBits);
 			ViewUniformShaderParameters = TUniformBufferRef<FViewUniformShaderParameters>::CreateUniformBufferImmediate(*ViewInfo->CachedViewUniformShaderParameters, UniformBuffer_SingleFrame);
 			DrawRenderState.SetViewUniformBuffer(ViewUniformShaderParameters);
 		}
@@ -675,21 +672,64 @@ static FRDGTextureRef AddHairLightChannelMaskPass(
 /////////////////////////////////////////////////////////////////////////////////////////
 BEGIN_GLOBAL_SHADER_PARAMETER_STRUCT(FVisibilityPassGlobalParameters, )
 	SHADER_PARAMETER(uint32, MaxPPLLNodeCount)
+	SHADER_PARAMETER_UAV(RWTexture2D<uint>, PPLLCounter)
+	SHADER_PARAMETER_UAV(RWTexture2D<uint>, PPLLNodeIndex)
+	SHADER_PARAMETER_UAV(RWStructuredBuffer<FPPLLNodeData>, PPLLNodeData)
 END_GLOBAL_SHADER_PARAMETER_STRUCT()
-IMPLEMENT_GLOBAL_SHADER_PARAMETER_STRUCT(FVisibilityPassGlobalParameters, "VisibilityPassGlobalParameters");
+IMPLEMENT_GLOBAL_SHADER_PARAMETER_STRUCT(FVisibilityPassGlobalParameters, "HairVisibilityPass");
 
 BEGIN_SHADER_PARAMETER_STRUCT(FVisibilityPassParameters, )
-	SHADER_PARAMETER(uint32, HairVisibilityPass_MaxPPLLNodeCount)
-	SHADER_PARAMETER_RDG_TEXTURE_UAV(Texture2D, PPLLCounter)
-	SHADER_PARAMETER_RDG_TEXTURE_UAV(Texture2D, PPLLNodeIndex)
-	SHADER_PARAMETER_RDG_BUFFER_UAV(StructuredBuffer, PPLLNodeData)
+	SHADER_PARAMETER(uint32, MaxPPLLNodeCount)
+	SHADER_PARAMETER_RDG_TEXTURE_UAV(RWTexture2D<uint>, PPLLCounter)
+	SHADER_PARAMETER_RDG_TEXTURE_UAV(RWTexture2D<uint>, PPLLNodeIndex)
+	SHADER_PARAMETER_RDG_BUFFER_UAV(RWStructuredBuffer<FPPLLNodeData, PPLLNodeData)
 	RENDER_TARGET_BINDING_SLOTS()
 END_SHADER_PARAMETER_STRUCT()
 
-template<EHairVisibilityRenderMode RenderMode>
+static FVisibilityPassGlobalParameters ConvertToGlobalPassParameter(const FVisibilityPassParameters* In)
+{
+	FVisibilityPassGlobalParameters Out;
+	Out.MaxPPLLNodeCount = In->MaxPPLLNodeCount;
+	Out.PPLLCounter = In->PPLLCounter->GetRHI();
+	Out.PPLLNodeIndex = In->PPLLNodeIndex->GetRHI();
+	Out.PPLLNodeData = In->PPLLNodeData->GetRHI();
+	return Out;
+}
+
+// Example: 28bytes * 8spp = 224bytes per pixel = 442Mb @ 1080p
+struct PPLLNodeData
+{
+	uint32 Depth;
+	uint32 PrimitiveId_MacroGroupId;
+	uint32 Tangent_Coverage;
+	uint32 BaseColor_Roughness;
+	uint32 Specular;
+	uint32 NextNodeIndex;
+	uint32 PackedVelocity;
+};
+
+void CreatePassDummyTextures(FRDGBuilder& GraphBuilder, FVisibilityPassParameters* PassParameters)
+{
+	FRDGTextureDesc Desc;
+	Desc.Extent.X = 1;
+	Desc.Extent.Y = 1;
+	Desc.Depth = 0;
+	Desc.Format = PF_R32_UINT;
+	Desc.NumMips = 1;
+	Desc.NumSamples = 1;
+	Desc.Flags = TexCreate_None;
+	Desc.TargetableFlags = TexCreate_UAV | TexCreate_ShaderResource;
+	Desc.ClearValue = FClearValueBinding(0);
+
+	PassParameters->PPLLCounter		= GraphBuilder.CreateUAV(GraphBuilder.CreateTexture(Desc, TEXT("HairVisibilityPPLLNodeIndex")));
+	PassParameters->PPLLNodeIndex	= GraphBuilder.CreateUAV(GraphBuilder.CreateTexture(Desc, TEXT("HairVisibilityPPLLNodeIndex")));
+	PassParameters->PPLLNodeData	= GraphBuilder.CreateUAV(GraphBuilder.CreateBuffer(FRDGBufferDesc::CreateStructuredDesc(sizeof(PPLLNodeData), 1), TEXT("DummyPPLLNodeData")));
+}
+
+template<EHairVisibilityRenderMode RenderMode, bool bCullingEnable>
 class FHairVisibilityVS : public FMeshMaterialShader
 {
-	DECLARE_SHADER_TYPE(FHairVisibilityVS,MeshMaterial);
+	DECLARE_SHADER_TYPE(FHairVisibilityVS, MeshMaterial);
 
 protected:
 
@@ -713,25 +753,33 @@ protected:
 		FMaterialShader::ModifyCompilationEnvironment(Parameters, OutEnvironment);
 		const uint32 RenderModeValue = uint32(RenderMode);
 		OutEnvironment.SetDefine(TEXT("HAIR_RENDER_MODE"), RenderModeValue);
-		OutEnvironment.SetDefine(TEXT("USE_CULLED_CLUSTER"), 1);
+		OutEnvironment.SetDefine(TEXT("USE_CULLED_CLUSTER"), bCullingEnable ? 1 : 0);
 	}
 };
-IMPLEMENT_MATERIAL_SHADER_TYPE(template<>, FHairVisibilityVS<HairVisibilityRenderMode_MSAA_Visibility>, TEXT("/Engine/Private/HairStrands/HairStrandsVisibilityVS.usf"), TEXT("Main"), SF_Vertex);
-IMPLEMENT_MATERIAL_SHADER_TYPE(template<>, FHairVisibilityVS<HairVisibilityRenderMode_MSAA>, TEXT("/Engine/Private/HairStrands/HairStrandsVisibilityVS.usf"), TEXT("Main"), SF_Vertex);
-IMPLEMENT_MATERIAL_SHADER_TYPE(template<>, FHairVisibilityVS<HairVisibilityRenderMode_Transmittance>, TEXT("/Engine/Private/HairStrands/HairStrandsVisibilityVS.usf"), TEXT("Main"), SF_Vertex);
-IMPLEMENT_MATERIAL_SHADER_TYPE(template<>, FHairVisibilityVS<HairVisibilityRenderMode_TransmittanceAndHairCount>, TEXT("/Engine/Private/HairStrands/HairStrandsVisibilityVS.usf"), TEXT("Main"), SF_Vertex);
-IMPLEMENT_MATERIAL_SHADER_TYPE(template<>, FHairVisibilityVS<HairVisibilityRenderMode_PPLL>, TEXT("/Engine/Private/HairStrands/HairStrandsVisibilityVS.usf"), TEXT("Main"), SF_Vertex);
+
+typedef FHairVisibilityVS<HairVisibilityRenderMode_MSAA_Visibility, false >				THairVisiblityVS_MSAAVisibility_NoCulling;
+typedef FHairVisibilityVS<HairVisibilityRenderMode_MSAA_Visibility, true >				THairVisiblityVS_MSAAVisibility_Culling;
+typedef FHairVisibilityVS<HairVisibilityRenderMode_MSAA, true >							THairVisiblityVS_MSAA;
+typedef FHairVisibilityVS<HairVisibilityRenderMode_Transmittance, true >				THairVisiblityVS_Transmittance;
+typedef FHairVisibilityVS<HairVisibilityRenderMode_TransmittanceAndHairCount, true >	THairVisiblityVS_TransmittanceAndHairCount;
+typedef FHairVisibilityVS<HairVisibilityRenderMode_PPLL, true >							THairVisiblityVS_PPLL;
+
+IMPLEMENT_MATERIAL_SHADER_TYPE(template<>, THairVisiblityVS_MSAAVisibility_NoCulling,	TEXT("/Engine/Private/HairStrands/HairStrandsVisibilityVS.usf"), TEXT("Main"), SF_Vertex);
+IMPLEMENT_MATERIAL_SHADER_TYPE(template<>, THairVisiblityVS_MSAAVisibility_Culling,		TEXT("/Engine/Private/HairStrands/HairStrandsVisibilityVS.usf"), TEXT("Main"), SF_Vertex);
+IMPLEMENT_MATERIAL_SHADER_TYPE(template<>, THairVisiblityVS_MSAA,						TEXT("/Engine/Private/HairStrands/HairStrandsVisibilityVS.usf"), TEXT("Main"), SF_Vertex);
+IMPLEMENT_MATERIAL_SHADER_TYPE(template<>, THairVisiblityVS_Transmittance,				TEXT("/Engine/Private/HairStrands/HairStrandsVisibilityVS.usf"), TEXT("Main"), SF_Vertex);
+IMPLEMENT_MATERIAL_SHADER_TYPE(template<>, THairVisiblityVS_TransmittanceAndHairCount,	TEXT("/Engine/Private/HairStrands/HairStrandsVisibilityVS.usf"), TEXT("Main"), SF_Vertex);
+IMPLEMENT_MATERIAL_SHADER_TYPE(template<>, THairVisiblityVS_PPLL,						TEXT("/Engine/Private/HairStrands/HairStrandsVisibilityVS.usf"), TEXT("Main"), SF_Vertex);
 
 /////////////////////////////////////////////////////////////////////////////////////////
 
 class FHairVisibilityShaderElementData : public FMeshMaterialShaderElementData
 {
 public:
-	FHairVisibilityShaderElementData(uint32 InHairMacroGroupId, uint32 InHairMaterialId, uint32 InLightChannelMask, uint32 InMaxPPLLNodeCount) : HairMacroGroupId(InHairMacroGroupId), HairMaterialId(InHairMaterialId), LightChannelMask(InLightChannelMask), MaxPPLLNodeCount(InMaxPPLLNodeCount) { }
+	FHairVisibilityShaderElementData(uint32 InHairMacroGroupId, uint32 InHairMaterialId, uint32 InLightChannelMask) : HairMacroGroupId(InHairMacroGroupId), HairMaterialId(InHairMaterialId), LightChannelMask(InLightChannelMask) { }
 	uint32 HairMacroGroupId;
 	uint32 HairMaterialId;
 	uint32 LightChannelMask;
-	uint32 MaxPPLLNodeCount;
 };
 
 template<EHairVisibilityRenderMode RenderMode>
@@ -750,13 +798,6 @@ public:
 		HairVisibilityPass_HairMacroGroupIndex.Bind(Initializer.ParameterMap, TEXT("HairVisibilityPass_HairMacroGroupIndex"));
 		HairVisibilityPass_HairMaterialId.Bind(Initializer.ParameterMap, TEXT("HairVisibilityPass_HairMaterialId"));
 		HairVisibilityPass_LightChannelMask.Bind(Initializer.ParameterMap, TEXT("HairVisibilityPass_LightChannelMask"));
-		HairVisibilityPass_MaxPPLLNodeCount.Bind(Initializer.ParameterMap, TEXT("HairVisibilityPass_MaxPPLLNodeCount"));
-		if (RenderMode == EHairVisibilityRenderMode::HairVisibilityRenderMode_PPLL)
-		{
-			PPLLCounter.Bind(Initializer.ParameterMap, TEXT("PPLLCounter"));
-			PPLLNodeIndex.Bind(Initializer.ParameterMap, TEXT("PPLLNodeIndex"));
-			PPLLNodes.Bind(Initializer.ParameterMap, TEXT("PPLLNodes"));
-		}
 	}
 
 	FHairVisibilityPS() {}
@@ -787,17 +828,11 @@ public:
 		ShaderBindings.Add(HairVisibilityPass_HairMacroGroupIndex, ShaderElementData.HairMacroGroupId);
 		ShaderBindings.Add(HairVisibilityPass_HairMaterialId, ShaderElementData.HairMaterialId);
 		ShaderBindings.Add(HairVisibilityPass_LightChannelMask, ShaderElementData.LightChannelMask);
-		ShaderBindings.Add(HairVisibilityPass_MaxPPLLNodeCount, ShaderElementData.MaxPPLLNodeCount);
 	}
 
 	LAYOUT_FIELD(FShaderParameter, HairVisibilityPass_HairMacroGroupIndex);
 	LAYOUT_FIELD(FShaderParameter, HairVisibilityPass_HairMaterialId);
 	LAYOUT_FIELD(FShaderParameter, HairVisibilityPass_LightChannelMask);
-	LAYOUT_FIELD(FShaderParameter, HairVisibilityPass_MaxPPLLNodeCount);
-	// This is on FVisibilityPassParameters but needed to avoid the compiler to assert with unbound parameters 
-	LAYOUT_FIELD(FShaderParameter, PPLLCounter);
-	LAYOUT_FIELD(FShaderParameter, PPLLNodeIndex);
-	LAYOUT_FIELD(FShaderParameter, PPLLNodes);
 };
 IMPLEMENT_MATERIAL_SHADER_TYPE(template<>, FHairVisibilityPS<HairVisibilityRenderMode_MSAA_Visibility>, TEXT("/Engine/Private/HairStrands/HairStrandsVisibilityPS.usf"), TEXT("MainVisibility"), SF_Pixel);
 IMPLEMENT_MATERIAL_SHADER_TYPE(template<>, FHairVisibilityPS<HairVisibilityRenderMode_MSAA>, TEXT("/Engine/Private/HairStrands/HairStrandsVisibilityPS.usf"), TEXT("MainVisibility"), SF_Pixel);
@@ -818,12 +853,10 @@ public:
 		FDynamicPassMeshDrawListContext* InDrawListContext);
 
 	virtual void AddMeshBatch(const FMeshBatch& RESTRICT MeshBatch, uint64 BatchElementMask, const FPrimitiveSceneProxy* RESTRICT PrimitiveSceneProxy, int32 StaticMeshId = -1) override final;
-	void AddMeshBatch(const FMeshBatch& RESTRICT MeshBatch, uint64 BatchElementMask, const FPrimitiveSceneProxy* RESTRICT PrimitiveSceneProxy, int32 StaticMeshId, uint32 HairMacroGroupId, uint32 HairMaterialId);
-
-	void SetMaxPPLLNodeCount(uint32 NewValue) { MaxPPLLNodeCount = NewValue; }
+	void AddMeshBatch(const FMeshBatch& RESTRICT MeshBatch, uint64 BatchElementMask, const FPrimitiveSceneProxy* RESTRICT PrimitiveSceneProxy, int32 StaticMeshId, uint32 HairMacroGroupId, uint32 HairMaterialId, bool bCullingEnable);
 
 private:
-	template<EHairVisibilityRenderMode RenderMode>
+	template<EHairVisibilityRenderMode RenderMode, bool bCullingEnable=true>
 	void Process(
 		const FMeshBatch& MeshBatch,
 		uint64 BatchElementMask,
@@ -839,15 +872,14 @@ private:
 
 	const EHairVisibilityRenderMode RenderMode;
 	FMeshPassProcessorRenderState PassDrawRenderState;
-	uint32 MaxPPLLNodeCount = 0;
 };
 
 void FHairVisibilityProcessor::AddMeshBatch(const FMeshBatch& RESTRICT MeshBatch, uint64 BatchElementMask, const FPrimitiveSceneProxy* RESTRICT PrimitiveSceneProxy, int32 StaticMeshId)
 {
-	AddMeshBatch(MeshBatch, BatchElementMask, PrimitiveSceneProxy, StaticMeshId, 0, 0);
+	AddMeshBatch(MeshBatch, BatchElementMask, PrimitiveSceneProxy, StaticMeshId, 0, 0, false);
 }
 
-void FHairVisibilityProcessor::AddMeshBatch(const FMeshBatch& RESTRICT MeshBatch, uint64 BatchElementMask, const FPrimitiveSceneProxy* RESTRICT PrimitiveSceneProxy, int32 StaticMeshId, uint32 HairMacroGroupId, uint32 HairMaterialId)
+void FHairVisibilityProcessor::AddMeshBatch(const FMeshBatch& RESTRICT MeshBatch, uint64 BatchElementMask, const FPrimitiveSceneProxy* RESTRICT PrimitiveSceneProxy, int32 StaticMeshId, uint32 HairMacroGroupId, uint32 HairMaterialId, bool bCullingEnable)
 {
 	static const FVertexFactoryType* CompatibleVF = FVertexFactoryType::GetVFByName(TEXT("FHairStrandsVertexFactory"));
 
@@ -868,20 +900,22 @@ void FHairVisibilityProcessor::AddMeshBatch(const FMeshBatch& RESTRICT MeshBatch
 		const FMeshDrawingPolicyOverrideSettings OverrideSettings = ComputeMeshOverrideSettings(MeshBatch);
 		const ERasterizerFillMode MeshFillMode = ComputeMeshFillMode(MeshBatch, Material, OverrideSettings);
 		const ERasterizerCullMode MeshCullMode = ComputeMeshCullMode(MeshBatch, Material, OverrideSettings);
-		if (RenderMode == HairVisibilityRenderMode_MSAA_Visibility)
-			Process<HairVisibilityRenderMode_MSAA_Visibility>(MeshBatch, BatchElementMask, PrimitiveSceneProxy, StaticMeshId, MaterialRenderProxy, Material, HairMacroGroupId, HairMaterialId, LightChannelMask, MeshFillMode, MeshCullMode);
-		if (RenderMode == HairVisibilityRenderMode_MSAA)
+		if (RenderMode == HairVisibilityRenderMode_MSAA_Visibility && bCullingEnable)
+			Process<HairVisibilityRenderMode_MSAA_Visibility, true>(MeshBatch, BatchElementMask, PrimitiveSceneProxy, StaticMeshId, MaterialRenderProxy, Material, HairMacroGroupId, HairMaterialId, LightChannelMask, MeshFillMode, MeshCullMode);
+		else if (RenderMode == HairVisibilityRenderMode_MSAA_Visibility && !bCullingEnable)
+			Process<HairVisibilityRenderMode_MSAA_Visibility, false>(MeshBatch, BatchElementMask, PrimitiveSceneProxy, StaticMeshId, MaterialRenderProxy, Material, HairMacroGroupId, HairMaterialId, LightChannelMask, MeshFillMode, MeshCullMode);
+		else if (RenderMode == HairVisibilityRenderMode_MSAA)
 			Process<HairVisibilityRenderMode_MSAA>(MeshBatch, BatchElementMask, PrimitiveSceneProxy, StaticMeshId, MaterialRenderProxy, Material, HairMacroGroupId, HairMaterialId, LightChannelMask, MeshFillMode, MeshCullMode);
-		if (RenderMode == HairVisibilityRenderMode_Transmittance)
+		else if (RenderMode == HairVisibilityRenderMode_Transmittance)
 			Process<HairVisibilityRenderMode_Transmittance>(MeshBatch, BatchElementMask, PrimitiveSceneProxy, StaticMeshId, MaterialRenderProxy, Material, HairMacroGroupId, HairMaterialId, LightChannelMask, MeshFillMode, MeshCullMode);
-		if (RenderMode == HairVisibilityRenderMode_TransmittanceAndHairCount)
+		else if (RenderMode == HairVisibilityRenderMode_TransmittanceAndHairCount)
 			Process<HairVisibilityRenderMode_TransmittanceAndHairCount>(MeshBatch, BatchElementMask, PrimitiveSceneProxy, StaticMeshId, MaterialRenderProxy, Material, HairMacroGroupId, HairMaterialId, LightChannelMask, MeshFillMode, MeshCullMode);
-		if (RenderMode == HairVisibilityRenderMode_PPLL)
+		else if (RenderMode == HairVisibilityRenderMode_PPLL)
 			Process<HairVisibilityRenderMode_PPLL>(MeshBatch, BatchElementMask, PrimitiveSceneProxy, StaticMeshId, MaterialRenderProxy, Material, HairMacroGroupId, HairMaterialId, LightChannelMask, MeshFillMode, MeshCullMode);
 	}
 }
 
-template<EHairVisibilityRenderMode TRenderMode>
+template<EHairVisibilityRenderMode TRenderMode, bool bCullingEnable>
 void FHairVisibilityProcessor::Process(
 	const FMeshBatch& MeshBatch, 
 	uint64 BatchElementMask, 
@@ -898,18 +932,18 @@ void FHairVisibilityProcessor::Process(
 	const FVertexFactory* VertexFactory = MeshBatch.VertexFactory;
 
 	TMeshProcessorShaders<
-		FHairVisibilityVS<TRenderMode>,
+		FHairVisibilityVS<TRenderMode, bCullingEnable>,
 		FMeshMaterialShader,
 		FMeshMaterialShader,
 		FHairVisibilityPS<TRenderMode>> PassShaders;
 	{
 		FVertexFactoryType* VertexFactoryType = VertexFactory->GetType();
-		PassShaders.VertexShader = MaterialResource.GetShader<FHairVisibilityVS<TRenderMode>>(VertexFactoryType);
+		PassShaders.VertexShader = MaterialResource.GetShader<FHairVisibilityVS<TRenderMode,bCullingEnable>>(VertexFactoryType);
 		PassShaders.PixelShader  = MaterialResource.GetShader<FHairVisibilityPS<TRenderMode>>(VertexFactoryType);
 	}
 
 	FMeshPassProcessorRenderState DrawRenderState(PassDrawRenderState);
-	FHairVisibilityShaderElementData ShaderElementData(HairMacroGroupId, HairMaterialId, LightChannelMask, MaxPPLLNodeCount);
+	FHairVisibilityShaderElementData ShaderElementData(HairMacroGroupId, HairMaterialId, LightChannelMask);
 	ShaderElementData.InitializeMeshMaterialData(ViewIfDynamicMeshCommand, PrimitiveSceneProxy, MeshBatch, StaticMeshId, false);
 
 	BuildMeshDrawCommands(
@@ -1012,31 +1046,6 @@ static void AddClearGraphicPass(
 	});
 }
 
-BEGIN_SHADER_PARAMETER_STRUCT(FClearUAVTextureParameters, )
-	SHADER_PARAMETER_RDG_TEXTURE_UAV(RWTexture2D, TextureUAV)
-END_SHADER_PARAMETER_STRUCT()
-
-void AddClearUAVPass(
-	FRDGBuilder& GraphBuilder,
-	FRDGEventName&& PassName,
-	FRDGTextureRef Texture,
-	uint32 Value)
-{	
-	FClearUAVTextureParameters* Parameters = GraphBuilder.AllocParameters< FClearUAVTextureParameters >();
-	Parameters->TextureUAV = GraphBuilder.CreateUAV(Texture);
-
-	GraphBuilder.AddPass(
-		Forward<FRDGEventName>(PassName),
-		Parameters,
-		ERDGPassFlags::Compute,
-		[Parameters, Texture, Value](FRHICommandList& RHICmdList)
-	{
-		FUintVector4 ClearValue(Value, Value, Value, Value);
-		FRHIUnorderedAccessView* GlobalCounterUAV = Parameters->TextureUAV->GetRHI();
-		RHICmdList.ClearUAVUint(GlobalCounterUAV, ClearValue);
-	});
-}
-
 ///////////////////////////////////////////////////////////////////////////////////////////////////
 // Copy dispatch count into an indirect buffer 
 class FCopyIndirectBufferCS : public FGlobalShader
@@ -1095,15 +1104,16 @@ class FHairVisibilityPrimitiveIdCompactionCS : public FGlobalShader
 	class FVendor		: SHADER_PERMUTATION_INT("PERMUTATION_VENDOR", HairVisibilityVendorCount);
 	class FVelocity		: SHADER_PERMUTATION_INT("PERMUTATION_VELOCITY", 4);
 	class FViewTransmittance : SHADER_PERMUTATION_INT("PERMUTATION_VIEWTRANSMITTANCE", 2);
-	class FMaterial		: SHADER_PERMUTATION_INT("PERMUTATION_MATERIAL_COMPACTION", 2);
-	class FPPLL			: SHADER_PERMUTATION_SPARSE_INT("PERMUTATION_PPLL", 0, 8, 16, 32); // See GetPPLLMaxRenderNodePerPixel
-	class FVisibility	: SHADER_PERMUTATION_INT("PERMUTATION_VISIBILITY", 2);
-	using FPermutationDomain = TShaderPermutationDomain<FVendor, FVelocity, FViewTransmittance, FMaterial, FPPLL, FVisibility>;
+	class FMaterial 	: SHADER_PERMUTATION_INT("PERMUTATION_MATERIAL_COMPACTION", 2);
+	class FPPLL 		: SHADER_PERMUTATION_SPARSE_INT("PERMUTATION_PPLL", 0, 8, 16, 32); // See GetPPLLMaxRenderNodePerPixel
+	class FVisibility 	: SHADER_PERMUTATION_INT("PERMUTATION_VISIBILITY", 2);
+	class FMSAACount 	: SHADER_PERMUTATION_SPARSE_INT("PERMUTATION_MSAACOUNT", 4, 8);
+	using FPermutationDomain = TShaderPermutationDomain<FVendor, FVelocity, FViewTransmittance, FMaterial, FPPLL, FVisibility, FMSAACount>;
 
 	BEGIN_SHADER_PARAMETER_STRUCT(FParameters, )
 		SHADER_PARAMETER(FIntPoint, OutputResolution)
 		SHADER_PARAMETER(uint32, MaxNodeCount)
-		SHADER_PARAMETER(uint32, HairVisibilityMSAASampleCount)
+		SHADER_PARAMETER(uint32, MSAASampleCount)
 		SHADER_PARAMETER(FIntPoint, ResolutionOffset)
 		SHADER_PARAMETER(float, DepthTheshold)
 		SHADER_PARAMETER(float, CosTangentThreshold)
@@ -1120,7 +1130,7 @@ class FHairVisibilityPrimitiveIdCompactionCS : public FGlobalShader
 		SHADER_PARAMETER_RDG_BUFFER_SRV(StructuredBuffer, PPLLNodeData)
 
 		SHADER_PARAMETER_RDG_TEXTURE(Texture2D, ViewTransmittanceTexture)
-		
+
 		SHADER_PARAMETER_RDG_TEXTURE_UAV(Texture2D, OutCompactNodeCounter)
 		SHADER_PARAMETER_RDG_TEXTURE_UAV(Texture2D, OutCompactNodeIndex)
 		SHADER_PARAMETER_RDG_TEXTURE_UAV(Texture2D, OutCategorizationTexture)
@@ -1139,7 +1149,12 @@ public:
 		if (PermutationVector.Get<FPPLL>() > 0)
 		{
 			PermutationVector.Set<FViewTransmittance>(0);
-			PermutationVector.Set<FVisibility>(0); 
+			PermutationVector.Set<FVisibility>(0);
+			PermutationVector.Set<FMSAACount>(4);
+		}
+		else if (PermutationVector.Get<FViewTransmittance>() == 1)
+		{
+			PermutationVector.Set<FMSAACount>(4);
 		}
 		return PermutationVector;
 	}
@@ -1155,6 +1170,14 @@ public:
 		{
 			return false;
 		}
+		if (PermutationVector.Get<FPPLL>() > 0 && PermutationVector.Get<FMSAACount>() == 8)
+		{
+			return false;
+		}
+		if (PermutationVector.Get<FViewTransmittance>() > 0 && PermutationVector.Get<FMSAACount>() == 8)
+		{
+			return false;
+		}
 		return IsHairStrandsSupported(Parameters.Platform);
 	}
 };
@@ -1186,6 +1209,8 @@ static void AddHairVisibilityPrimitiveIdCompactionPass(
 	}
 	else
 	{
+		check(PassParameters->MSAA_DepthTexture->Desc.NumSamples == GetMSAASampleCount());
+
 		if (bUseVisibility)
 		{
 			check(PassParameters->MSAA_DepthTexture);
@@ -1202,7 +1227,7 @@ static void AddHairVisibilityPrimitiveIdCompactionPass(
 		}
 	}
 
-	
+
 	FRDGTextureRef CompactCounter;
 	{
 		FRDGTextureDesc Desc;
@@ -1239,16 +1264,16 @@ static void AddHairVisibilityPrimitiveIdCompactionPass(
 		OutputDesc.TargetableFlags = TexCreate_UAV | TexCreate_ShaderResource;
 		OutCategorizationTexture = GraphBuilder.CreateTexture(OutputDesc, TEXT("CategorizationTexture"));
 	}
-	
-	AddClearUAVPass(GraphBuilder, RDG_EVENT_NAME("HairStrandsClearCompactionCounter"), CompactCounter, 0);
-	AddClearUAVPass(GraphBuilder, RDG_EVENT_NAME("HairStrandsClearCompactionOffsetAndCount"), OutCompactNodeIndex, 0);
-	AddClearUAVPass(GraphBuilder, RDG_EVENT_NAME("HairStrandsClearCategorizationTexture"), OutCategorizationTexture, 0);
+
+	const uint32 ClearValues[4] = { 0u,0u,0u,0u };
+	AddClearUAVPass(GraphBuilder, GraphBuilder.CreateUAV(CompactCounter), ClearValues);
+	AddClearUAVPass(GraphBuilder, GraphBuilder.CreateUAV(OutCompactNodeIndex), ClearValues);
+	AddClearUAVPass(GraphBuilder, GraphBuilder.CreateUAV(OutCategorizationTexture), ClearValues);
 
 	// Select render node count according to current mode
-	const uint32 HairVisibilityMSAASampleCount = GetHairVisibilitySampleCount();
-	const uint32 SampleCount = FMath::RoundUpToPowerOfTwo(HairVisibilityMSAASampleCount);
+	const uint32 MSAASampleCount = GetMSAASampleCount(); check(MSAASampleCount == 4 || MSAASampleCount == 8);
 	const uint32 PPLLMaxRenderNodePerPixel = GetPPLLMaxRenderNodePerPixel();
-	const uint32 MaxRenderNodeCount = Resolution.X * Resolution.Y * (GetHairVisibilityRenderMode() == HairVisibilityRenderMode_MSAA ? SampleCount : PPLLMaxRenderNodePerPixel);
+	const uint32 MaxRenderNodeCount = Resolution.X * Resolution.Y * (GetHairVisibilityRenderMode() == HairVisibilityRenderMode_MSAA ? MSAASampleCount : PPLLMaxRenderNodePerPixel);
 	OutCompactNodeData = GraphBuilder.CreateBuffer(FRDGBufferDesc::CreateStructuredDesc(bUseVisibility ? sizeof(HairStrandsVisibilityInternal::NodeVis) : sizeof(HairStrandsVisibilityInternal::NodeData), MaxRenderNodeCount), TEXT("HairVisibilityPrimitiveIdCompactNodeData"));
 
 	{
@@ -1268,14 +1293,14 @@ static void AddHairVisibilityPrimitiveIdCompactionPass(
 	PermutationVector.Set<FHairVisibilityPrimitiveIdCompactionCS::FViewTransmittance>(PassParameters->ViewTransmittanceTexture ? 1 : 0);
 	PermutationVector.Set<FHairVisibilityPrimitiveIdCompactionCS::FMaterial>(GHairStrandsMaterialCompactionEnable ? 1 : 0);
 	PermutationVector.Set<FHairVisibilityPrimitiveIdCompactionCS::FPPLL>(bUsePPLL ? PPLLMaxRenderNodePerPixel : 0);
-	PermutationVector.Set<FHairVisibilityPrimitiveIdCompactionCS::FVisibility>(bUseVisibility ? 1 : 0);
+	PermutationVector.Set<FHairVisibilityPrimitiveIdCompactionCS::FVisibility>(bUseVisibility ? 1 : 0);	
+	PermutationVector.Set<FHairVisibilityPrimitiveIdCompactionCS::FMSAACount>(MSAASampleCount == 4 ? 4 : 8);
 	PermutationVector = FHairVisibilityPrimitiveIdCompactionCS::RemapPermutation(PermutationVector);
 
 	PassParameters->OutputResolution = Resolution;
 	PassParameters->MaxNodeCount = MaxRenderNodeCount;
 	PassParameters->DepthTheshold = FMath::Clamp(GHairStrandsMaterialCompactionDepthThreshold, 0.f, 100.f);
 	PassParameters->CosTangentThreshold = FMath::Cos(FMath::DegreesToRadians(FMath::Clamp(GHairStrandsMaterialCompactionTangentThreshold, 0.f, 90.f)));
-	PassParameters->HairVisibilityMSAASampleCount = HairVisibilityMSAASampleCount;
 	PassParameters->SceneTexturesStruct = CreateUniformBufferImmediate(SceneTextures, EUniformBufferUsage::UniformBuffer_SingleDraw);
 	PassParameters->ViewUniformBuffer = View.ViewUniformBuffer;
 	PassParameters->OutCompactNodeCounter = GraphBuilder.CreateUAV(CompactCounter);
@@ -1297,7 +1322,7 @@ static void AddHairVisibilityPrimitiveIdCompactionPass(
 	TotalRect.Min.Y = FMath::FloorToInt(float(TotalRect.Min.Y) / float(GroupSize.Y)) * GroupSize.Y;
 	TotalRect.Max.X = FMath::CeilToInt(float(TotalRect.Max.X) / float(GroupSize.X)) * GroupSize.X;
 	TotalRect.Max.Y = FMath::CeilToInt(float(TotalRect.Max.Y) / float(GroupSize.Y)) * GroupSize.Y;
-	
+
 	FIntPoint RectResolution(TotalRect.Width(), TotalRect.Height());
 	PassParameters->ResolutionOffset = FIntPoint(TotalRect.Min.X, TotalRect.Min.Y);
 
@@ -1332,7 +1357,7 @@ class FHairGenerateTileCS : public FGlobalShader
 	END_SHADER_PARAMETER_STRUCT()
 
 public:
-	
+
 	static bool ShouldCompilePermutation(const FGlobalShaderPermutationParameters& Parameters)
 	{
 		return IsHairStrandsSupported(Parameters.Platform);
@@ -1383,10 +1408,11 @@ static void AddGenerateTilePass(
 		Desc.ClearValue = FClearValueBinding(0);
 		OutTileIndexTexture = GraphBuilder.CreateTexture(Desc, TEXT("HairTileIndexTexture"));
 	}
-	
+
 	OutTileBuffer = GraphBuilder.CreateBuffer(FRDGBufferDesc::CreateBufferDesc(sizeof(uint32), TileCount), TEXT("HairTileBuffer"));
 
-	AddClearUAVPass(GraphBuilder, RDG_EVENT_NAME("HairTileCounter"), TileCounter, 0);
+	uint32 ClearValues[4] = { 0,0,0,0 };
+	AddClearUAVPass(GraphBuilder, GraphBuilder.CreateUAV(TileCounter), ClearValues);
 
 	FHairGenerateTileCS::FPermutationDomain PermutationVector;
 	PermutationVector.Set<FHairGenerateTileCS::FTileSize>(0);
@@ -1394,7 +1420,7 @@ static void AddGenerateTilePass(
 	FHairGenerateTileCS::FParameters* PassParameters = GraphBuilder.AllocParameters<FHairGenerateTileCS::FParameters>();
 	PassParameters->Resolution = Resolution;
 	PassParameters->TileResolution = TileResolution;
-	PassParameters->CategorizationTexture = CategorizationTexture;	
+	PassParameters->CategorizationTexture = CategorizationTexture;
 	PassParameters->OutTileCounter = GraphBuilder.CreateUAV(TileCounter);
 	PassParameters->OutTileIndexTexture = GraphBuilder.CreateUAV(OutTileIndexTexture);
 	PassParameters->OutTileBuffer = GraphBuilder.CreateUAV(OutTileBuffer, PF_R16G16_UINT);
@@ -1440,7 +1466,7 @@ static FRDGTextureRef AddHairVisibilityFillOpaqueDepth(
 {
 	FRDGTextureRef OutVisibilityDepthTexture;
 	{
-		const uint32 MSAASampleCount = FMath::RoundUpToPowerOfTwo(FMath::Clamp(GHairVisibilitySampleCount, 1, 16));
+		check(GetHairVisibilityRenderMode() == HairVisibilityRenderMode_MSAA);
 
 		FRDGTextureDesc Desc;
 		Desc.Extent.X = Resolution.X;
@@ -1448,7 +1474,7 @@ static FRDGTextureRef AddHairVisibilityFillOpaqueDepth(
 		Desc.Depth = 0;
 		Desc.Format = PF_DepthStencil;
 		Desc.NumMips = 1;
-		Desc.NumSamples = MSAASampleCount;
+		Desc.NumSamples = GetMSAASampleCount();
 		Desc.Flags = TexCreate_None;
 		Desc.TargetableFlags = TexCreate_DepthStencilTargetable | TexCreate_ShaderResource;
 		Desc.ClearValue = FClearValueBinding::DepthFar;
@@ -1531,16 +1557,14 @@ static FRDGTextureRef AddHairVisibilityFillOpaqueDepth(
 
 ///////////////////////////////////////////////////////////////////////////////////////////////////
 
-template <typename TPassParameters>
 static void AddHairVisibilityCommonPass(
 	FRDGBuilder& GraphBuilder,
 	const FScene* Scene,
 	const FViewInfo* ViewInfo,
 	const FHairStrandsMacroGroupDatas& MacroGroupDatas,
 	const EHairVisibilityRenderMode RenderMode,
-	TPassParameters* PassParameters)
+	FVisibilityPassParameters* PassParameters)
 {
-
 	auto GetPassName = [RenderMode]()
 	{
 		switch (RenderMode)
@@ -1555,7 +1579,7 @@ static void AddHairVisibilityCommonPass(
 	};
 
 	GraphBuilder.AddPass(
-		GetPassName(), 
+		GetPassName(),
 		PassParameters,
 		ERDGPassFlags::Raster,
 		[PassParameters, Scene = Scene, ViewInfo, &MacroGroupDatas, RenderMode](FRHICommandListImmediate& RHICmdList)
@@ -1563,18 +1587,23 @@ static void AddHairVisibilityCommonPass(
 		check(RHICmdList.IsInsideRenderPass());
 		check(IsInRenderingThread());
 
-		FMeshPassProcessorRenderState DrawRenderState(*ViewInfo);
+		FVisibilityPassGlobalParameters GlobalPassParameters = ConvertToGlobalPassParameter(PassParameters);
+		TUniformBufferRef<FVisibilityPassGlobalParameters> GlobalPassParametersBuffer = TUniformBufferRef<FVisibilityPassGlobalParameters>::CreateUniformBufferImmediate(GlobalPassParameters, UniformBuffer_SingleFrame);
+
+		FMeshPassProcessorRenderState DrawRenderState(*ViewInfo, GlobalPassParametersBuffer);
 
+		// Note: this reference needs to persistent until SubmitMeshDrawCommands() is called, as DrawRenderState does not ref count 
+		// the view uniform buffer (raw pointer). It is only within the MeshProcessor that the uniform buffer get reference
+		TUniformBufferRef<FViewUniformShaderParameters> ViewUniformShaderParameters;
 		if(RenderMode == HairVisibilityRenderMode_Transmittance || RenderMode == HairVisibilityRenderMode_TransmittanceAndHairCount || RenderMode == HairVisibilityRenderMode_PPLL)
 		{
-			TUniformBufferRef<FViewUniformShaderParameters> ViewUniformShaderParameters;
 			const bool bEnableMSAA = false;
-			SetUpViewHairRenderInfo(*ViewInfo, bEnableMSAA, ViewInfo->CachedViewUniformShaderParameters->HairRenderInfo);
+			SetUpViewHairRenderInfo(*ViewInfo, bEnableMSAA, ViewInfo->CachedViewUniformShaderParameters->HairRenderInfo, ViewInfo->CachedViewUniformShaderParameters->HairRenderInfoBits);
 			// Create and set the uniform buffer
 			ViewUniformShaderParameters = TUniformBufferRef<FViewUniformShaderParameters>::CreateUniformBufferImmediate(*ViewInfo->CachedViewUniformShaderParameters, UniformBuffer_SingleFrame);
 			DrawRenderState.SetViewUniformBuffer(ViewUniformShaderParameters);
 		}
-		
+
 		{
 			RHICmdList.SetViewport(0, 0, 0.0f, ViewInfo->ViewRect.Width(), ViewInfo->ViewRect.Height(), 1.0f);
 			if (RenderMode == HairVisibilityRenderMode_MSAA)
@@ -1584,7 +1613,7 @@ static void AddHairVisibilityCommonPass(
 					CW_RGBA, BO_Add, BF_One, BF_Zero, BO_Add, BF_One, BF_Zero>::GetRHI());
 				DrawRenderState.SetDepthStencilState(TStaticDepthStencilState<true, CF_DepthNearOrEqual>::GetRHI());
 			}
-			else if(RenderMode == HairVisibilityRenderMode_MSAA_Visibility)
+			else if (RenderMode == HairVisibilityRenderMode_MSAA_Visibility)
 			{
 				DrawRenderState.SetBlendState(TStaticBlendState<CW_RGBA, BO_Add, BF_One, BF_Zero, BO_Add, BF_One, BF_Zero>::GetRHI());
 				DrawRenderState.SetDepthStencilState(TStaticDepthStencilState<true, CF_DepthNearOrEqual>::GetRHI());
@@ -1613,20 +1642,14 @@ static void AddHairVisibilityCommonPass(
 			bool NeedsShaderInitialization;
 			FDynamicPassMeshDrawListContext ShadowContext(DynamicMeshDrawCommandStorage, VisibleMeshDrawCommands, PipelineStateSet, NeedsShaderInitialization);
 			FHairVisibilityProcessor MeshProcessor(Scene, ViewInfo, DrawRenderState, RenderMode, &ShadowContext);
-			if (RenderMode == HairVisibilityRenderMode_PPLL)
-			{
-				// Work around because the value on the PassParamter is not taken into account for the PS global constant buffer.
-				// TODO fix that by having FVisibilityPassParameters be a global parameter structure.
-				MeshProcessor.SetMaxPPLLNodeCount(PassParameters->HairVisibilityPass_MaxPPLLNodeCount);
-			}
-
+			
 			for (const FHairStrandsMacroGroupData& MacroGroupData : MacroGroupDatas.Datas)
 			{
 				for (const FHairStrandsMacroGroupData::PrimitiveInfo& PrimitiveInfo : MacroGroupData.PrimitivesInfos)
 				{
 					const FMeshBatch& MeshBatch = *PrimitiveInfo.MeshBatchAndRelevance.Mesh;
 					const uint64 BatchElementMask = ~0ull;
-					MeshProcessor.AddMeshBatch(MeshBatch, BatchElementMask, PrimitiveInfo.MeshBatchAndRelevance.PrimitiveSceneProxy, -1, MacroGroupData.MacroGroupId, PrimitiveInfo.MaterialId);
+					MeshProcessor.AddMeshBatch(MeshBatch, BatchElementMask, PrimitiveInfo.MeshBatchAndRelevance.PrimitiveSceneProxy, -1, MacroGroupData.MacroGroupId, PrimitiveInfo.MaterialId, PrimitiveInfo.IsCullingEnable());
 				}
 			}
 
@@ -1653,7 +1676,7 @@ static void AddHairVisibilityMSAAPass(
 	FRDGTextureRef& OutVisibilityVelocityTexture,
 	FRDGTextureRef& OutVisibilityDepthTexture)
 {
-	const uint32 MSAASampleCount = FMath::RoundUpToPowerOfTwo(FMath::Clamp(GHairVisibilitySampleCount, 1, 16));
+	const uint32 MSAASampleCount = GetMSAASampleCount();
 
 	if (bUseVisibility)
 	{
@@ -1677,6 +1700,7 @@ static void AddHairVisibilityMSAAPass(
 		AddClearGraphicPass(GraphBuilder, RDG_EVENT_NAME("HairStrandsClearVisibilityMSAAIdTexture"), ViewInfo, 0xFFFFFFFF, OutVisibilityIdTexture);
 
 		FVisibilityPassParameters* PassParameters = GraphBuilder.AllocParameters<FVisibilityPassParameters>();
+		CreatePassDummyTextures(GraphBuilder, PassParameters);
 		PassParameters->RenderTargets[0] = FRenderTargetBinding(OutVisibilityIdTexture, ERenderTargetLoadAction::ELoad, 0);
 		PassParameters->RenderTargets.DepthStencil = FDepthStencilBinding(
 			OutVisibilityDepthTexture,
@@ -1692,7 +1716,7 @@ static void AddHairVisibilityMSAAPass(
 			Desc.Extent.X = Resolution.X;
 			Desc.Extent.Y = Resolution.Y;
 			Desc.Depth = 0;
-			Desc.Format = PF_R16G16B16A16_UINT;
+			Desc.Format = PF_R32G32_UINT;
 			Desc.NumMips = 1;
 			Desc.NumSamples = MSAASampleCount;
 			Desc.Flags = TexCreate_None;
@@ -1760,10 +1784,11 @@ static void AddHairVisibilityMSAAPass(
 		}
 
 		FVisibilityPassParameters* PassParameters = GraphBuilder.AllocParameters<FVisibilityPassParameters>();
+		CreatePassDummyTextures(GraphBuilder, PassParameters);
 		PassParameters->RenderTargets[0] = FRenderTargetBinding(OutVisibilityIdTexture, ERenderTargetLoadAction::ELoad, 0);
-		PassParameters->RenderTargets[1] = FRenderTargetBinding(OutVisibilityMaterialTexture,  LoadAction, 0);
+		PassParameters->RenderTargets[1] = FRenderTargetBinding(OutVisibilityMaterialTexture, LoadAction, 0);
 		PassParameters->RenderTargets[2] = FRenderTargetBinding(OutVisibilityAttributeTexture, LoadAction, 0);
-		PassParameters->RenderTargets[3] = FRenderTargetBinding(OutVisibilityVelocityTexture,  LoadAction, 0);
+		PassParameters->RenderTargets[3] = FRenderTargetBinding(OutVisibilityVelocityTexture, LoadAction, 0);
 
 		PassParameters->RenderTargets.DepthStencil = FDepthStencilBinding(
 			OutVisibilityDepthTexture,
@@ -1815,28 +1840,18 @@ static void AddHairVisibilityPPLLPass(
 
 	const uint32 PPLLMaxTotalListElementCount = GetPPLLMaxTotalListElementCount(Resolution);
 	{
-		// Example: 28bytes * 8spp = 224bytes per pixel = 442Mb @ 1080p
-		struct PPLLNodeData
-		{
-			uint32 Depth;
-			uint32 PrimitiveId_MacroGroupId;
-			uint32 Tangent_Coverage;
-			uint32 BaseColor_Roughness;
-			uint32 Specular;
-			uint32 NextNodeIndex;
-			uint32 PackedVelocity;
-		};
-
 		OutVisibilityPPLLNodeData = GraphBuilder.CreateBuffer(FRDGBufferDesc::CreateStructuredDesc(sizeof(PPLLNodeData), PPLLMaxTotalListElementCount), TEXT("HairVisibilityPPLLNodeData"));
 	}
-	AddClearUAVPass(GraphBuilder, RDG_EVENT_NAME("ClearHairVisibilityPPLLCounter"), OutVisibilityPPLLNodeCounter, 0);
-	AddClearUAVPass(GraphBuilder, RDG_EVENT_NAME("ClearHairVisibilityPPLLNodeIndex"), OutVisibilityPPLLNodeIndex, 0xFFFFFFFF);
+	const uint32 ClearValue0[4] = { 0,0,0,0 };
+	const uint32 ClearValueInvalid[4] = { 0xFFFFFFFF, 0xFFFFFFFF, 0xFFFFFFFF, 0xFFFFFFFF };
+	AddClearUAVPass(GraphBuilder, GraphBuilder.CreateUAV(OutVisibilityPPLLNodeCounter), ClearValue0);
+	AddClearUAVPass(GraphBuilder, GraphBuilder.CreateUAV(OutVisibilityPPLLNodeIndex), ClearValueInvalid);
 
 	FVisibilityPassParameters* PassParameters = GraphBuilder.AllocParameters<FVisibilityPassParameters>();
 	PassParameters->PPLLCounter = GraphBuilder.CreateUAV(FRDGTextureUAVDesc(OutVisibilityPPLLNodeCounter, 0));
 	PassParameters->PPLLNodeIndex = GraphBuilder.CreateUAV(FRDGTextureUAVDesc(OutVisibilityPPLLNodeIndex, 0));
 	PassParameters->PPLLNodeData = GraphBuilder.CreateUAV(FRDGBufferUAVDesc(OutVisibilityPPLLNodeData));
-	PassParameters->HairVisibilityPass_MaxPPLLNodeCount = PPLLMaxTotalListElementCount;
+	PassParameters->MaxPPLLNodeCount = PPLLMaxTotalListElementCount;
 	PassParameters->RenderTargets.DepthStencil = FDepthStencilBinding(InViewZDepthTexture, ERenderTargetLoadAction::ELoad, ERenderTargetLoadAction::ENoAction, FExclusiveDepthStencil::DepthRead_StencilNop);
 	AddHairVisibilityCommonPass(GraphBuilder, Scene, ViewInfo, MacroGroupDatas, HairVisibilityRenderMode_PPLL, PassParameters);
 }
@@ -1874,6 +1889,7 @@ static FHairPrimaryTransmittance AddHairViewTransmittancePass(
 	Desc.ClearValue = FClearValueBinding(FLinearColor(1.0f, 1.0f, 1.0f, 1.0f)); // Clear to transmittance 1
 
 	FVisibilityPassParameters* PassParameters = GraphBuilder.AllocParameters<FVisibilityPassParameters>();
+	CreatePassDummyTextures(GraphBuilder, PassParameters);
 	FHairPrimaryTransmittance Out;
 
 	Out.TransmittanceTexture = GraphBuilder.CreateTexture(Desc, TEXT("HairViewTransmittanceTexture"));
@@ -1906,7 +1922,7 @@ class FHairViewTransmittanceDepthPS : public FGlobalShader
 	BEGIN_SHADER_PARAMETER_STRUCT(FParameters, )
 		SHADER_PARAMETER(float, DistanceThreshold)
 		SHADER_PARAMETER_RDG_TEXTURE(Texture2D, SceneDepthTexture)
-		SHADER_PARAMETER_RDG_TEXTURE(Texture2D, CategorizationTexture)		
+		SHADER_PARAMETER_RDG_TEXTURE(Texture2D, CategorizationTexture)
 		SHADER_PARAMETER_STRUCT_REF(FViewUniformShaderParameters, ViewUniformBuffer)
 		RENDER_TARGET_BINDING_SLOTS()
 	END_SHADER_PARAMETER_STRUCT()
@@ -1934,7 +1950,7 @@ static void AddHairViewTransmittanceDepthPass(
 	Parameters->SceneDepthTexture = SceneDepthTexture;
 	Parameters->ViewUniformBuffer = View.ViewUniformBuffer;
 	Parameters->RenderTargets[0] = FRenderTargetBinding(HairCountTexture, ERenderTargetLoadAction::ELoad);
-	
+
 	TShaderMapRef<FPostProcessVS> VertexShader(View.ShaderMap);
 	TShaderMapRef<FHairViewTransmittanceDepthPS> PixelShader(View.ShaderMap);
 	const FGlobalShaderMap* GlobalShaderMap = View.ShaderMap;
@@ -1992,7 +2008,7 @@ class FHairVisibilityDepthPS : public FGlobalShader
 	END_SHADER_PARAMETER_STRUCT()
 
 public:
-		
+
 	static FPermutationDomain RemapPermutation(FPermutationDomain PermutationVector)
 	{
 		return PermutationVector;
@@ -2072,11 +2088,12 @@ static void AddHairVisibilityColorAndDepthPatchPass(
 }
 
 ///////////////////////////////////////////////////////////////////////////////////////////////////
+bool GetHairStrandsSkyLightingEnable();
 
 FHairStrandsVisibilityViews RenderHairStrandsVisibilityBuffer(
-	FRHICommandListImmediate& RHICmdList, 
-	const FScene* Scene, 
-	const TArray<FViewInfo>& Views, 
+	FRHICommandListImmediate& RHICmdList,
+	const FScene* Scene,
+	const TArray<FViewInfo>& Views,
 	TRefCountPtr<IPooledRenderTarget> InSceneGBufferBTexture,
 	TRefCountPtr<IPooledRenderTarget> InSceneColorTexture,
 	TRefCountPtr<IPooledRenderTarget> InSceneDepthTexture,
@@ -2113,20 +2130,22 @@ FHairStrandsVisibilityViews RenderHairStrandsVisibilityBuffer(
 			FRDGTextureRef SceneColorTexture = GraphBuilder.RegisterExternalTexture(InSceneColorTexture, TEXT("SceneColorTexture"));
 			FRDGTextureRef SceneDepthTexture = GraphBuilder.RegisterExternalTexture(InSceneDepthTexture, TEXT("SceneDepthTexture"));
 			FRDGTextureRef SceneVelocityTexture = InSceneVelocityTexture ? GraphBuilder.RegisterExternalTexture(InSceneVelocityTexture, TEXT("SceneVelocityTexture")) : nullptr;
-			
+
 			const EHairVisibilityRenderMode RenderMode = GetHairVisibilityRenderMode();
+			check(RenderMode == HairVisibilityRenderMode_MSAA || RenderMode == HairVisibilityRenderMode_PPLL);
 
 			// Run the view transmittance pass if needed (not in PPLL mode that is already a high quality render path)
 			FHairPrimaryTransmittance ViewTransmittance;
 			if (GHairStrandsViewTransmittancePassEnable > 0 && RenderMode != HairVisibilityRenderMode_PPLL)
 			{
-				const bool bOutputHairCount = GHairStrandsViewHairCount > 0;
+				// Note: Hair count is required for the sky lighting at the moment as it is used for the TT term
+				const bool bOutputHairCount = GetHairStrandsSkyLightingEnable();
 				ViewTransmittance = AddHairViewTransmittancePass(
 					GraphBuilder,
 					Scene,
 					&View,
 					MacroGroupDatas,
-					Resolution, 
+					Resolution,
 					bOutputHairCount,
 					SceneDepthTexture);
 			}
@@ -2146,7 +2165,7 @@ FHairStrandsVisibilityViews RenderHairStrandsVisibilityBuffer(
 					FRDGTextureRef AttributeTexture;
 					FRDGTextureRef VelocityTexture;
 				} MsaaVisibilityResources;
-				
+
 				MsaaVisibilityResources.DepthTexture = AddHairVisibilityFillOpaqueDepth(
 					GraphBuilder,
 					View,
@@ -2156,18 +2175,19 @@ FHairStrandsVisibilityViews RenderHairStrandsVisibilityBuffer(
 
 				AddHairVisibilityMSAAPass(
 					bIsVisiblityEnable,
-					GraphBuilder, 
-					Scene, 
-					&View, 
+					GraphBuilder,
+					Scene,
+					&View,
 					MacroGroupDatas,
 					Resolution,
-					MsaaVisibilityResources.IdTexture, 
-					MsaaVisibilityResources.MaterialTexture, 
+					MsaaVisibilityResources.IdTexture,
+					MsaaVisibilityResources.MaterialTexture,
 					MsaaVisibilityResources.AttributeTexture,
 					MsaaVisibilityResources.VelocityTexture,
-					MsaaVisibilityResources.DepthTexture); 
+					MsaaVisibilityResources.DepthTexture);
 
 				// This is used when compaction is not enabled.
+				VisibilityData.MaxSampleCount = MsaaVisibilityResources.IdTexture->Desc.NumSamples;
 				GraphBuilder.QueueTextureExtraction(MsaaVisibilityResources.IdTexture, &VisibilityData.IDTexture);
 				GraphBuilder.QueueTextureExtraction(MsaaVisibilityResources.DepthTexture, &VisibilityData.DepthTexture);
 				if (!bIsVisiblityEnable)
@@ -2301,13 +2321,15 @@ FHairStrandsVisibilityViews RenderHairStrandsVisibilityBuffer(
 						CategorizationTexture,
 						SceneVelocityTexture,
 						IndirectArgsBuffer);
+
+					VisibilityData.MaxSampleCount = GetPPLLMaxRenderNodePerPixel();
 					GraphBuilder.QueueTextureExtraction(CompactNodeIndex, &VisibilityData.NodeIndex);
 					GraphBuilder.QueueTextureExtraction(CategorizationTexture, &VisibilityData.CategorizationTexture);
 					GraphBuilder.QueueBufferExtraction(CompactNodeData, &VisibilityData.NodeData, FRDGResourceState::EAccess::Read, FRDGResourceState::EPipeline::Graphics);
 					GraphBuilder.QueueBufferExtraction(CompactNodeCoord, &VisibilityData.NodeCoord, FRDGResourceState::EAccess::Read, FRDGResourceState::EPipeline::Graphics);
 					GraphBuilder.QueueBufferExtraction(IndirectArgsBuffer, &VisibilityData.NodeIndirectArg, FRDGResourceState::EAccess::Read, FRDGResourceState::EPipeline::Compute);
 				}
-		
+
 				AddHairVisibilityColorAndDepthPatchPass(
 					GraphBuilder,
 					View,
diff --git a/Engine/Source/Runtime/Renderer/Private/HairStrands/HairStrandsVisibility.h b/Engine/Source/Runtime/Renderer/Private/HairStrands/HairStrandsVisibility.h
--- a/Engine/Source/Runtime/Renderer/Private/HairStrands/HairStrandsVisibility.h
+++ b/Engine/Source/Runtime/Renderer/Private/HairStrands/HairStrandsVisibility.h
@@ -28,6 +28,8 @@ struct FHairStrandsVisibilityData
 	TRefCountPtr<IPooledRenderTarget> PPLLNodeCounterTexture;
 	TRefCountPtr<IPooledRenderTarget> PPLLNodeIndexTexture;
 	TRefCountPtr<FPooledRDGBuffer>	  PPLLNodeDataBuffer;
+	uint32							  MaxPPLLNodePerPixelCount = 0;
+	uint32							  MaxPPLLNodeCount = 0;
 
 	TRefCountPtr<IPooledRenderTarget> TileIndexTexture;
 	TRefCountPtr<FPooledRDGBuffer>	  TileBuffer;
@@ -35,6 +37,7 @@ struct FHairStrandsVisibilityData
 	const uint32					  TileSize = 8;
 	const uint32					  TileThreadGroupSize = 32;
 
+	uint32							  MaxSampleCount = 8;
 	TRefCountPtr<IPooledRenderTarget> NodeIndex;
 	TRefCountPtr<FPooledRDGBuffer>	  NodeData;
 	FShaderResourceViewRHIRef		  NodeDataSRV;
@@ -59,9 +62,5 @@ FHairStrandsVisibilityViews RenderHairStrandsVisibilityBuffer(
 	TRefCountPtr<IPooledRenderTarget> VelocityTexture,
 	const struct FHairStrandsMacroGroupViews& MacroGroupViews);
 
-void SetUpViewHairRenderInfo(const FViewInfo& ViewInfo, bool bEnableMSAA, FVector4& OutHairRenderInfo);
-
-
-uint32 GetPPLLMeanListElementCountPerPixel();
-uint32 GetPPLLMaxTotalListElementCount(FIntPoint Resolution);
+void SetUpViewHairRenderInfo(const FViewInfo& ViewInfo, bool bEnableMSAA, FVector4& OutHairRenderInfo, uint32& OutHairRenderInfoBits);
 
diff --git a/Engine/Source/Runtime/Renderer/Private/HairStrands/HairStrandsVoxelization.cpp b/Engine/Source/Runtime/Renderer/Private/HairStrands/HairStrandsVoxelization.cpp
--- a/Engine/Source/Runtime/Renderer/Private/HairStrands/HairStrandsVoxelization.cpp
+++ b/Engine/Source/Runtime/Renderer/Private/HairStrands/HairStrandsVoxelization.cpp
@@ -69,8 +69,20 @@ static int32 GHairVirtualVoxelGPUDrivenMaxPageIndexRes = 32;
 static FAutoConsoleVariableRef CVarHairVirtualVoxelGPUDriven(TEXT("r.HairStrands.Voxelization.GPUDriven"), GHairVirtualVoxelGPUDriven, TEXT("Enable GPU driven voxelization."));
 static FAutoConsoleVariableRef CVarHairVirtualVoxelGPUDrivenMaxPageIndexRes(TEXT("r.HairStrands.Voxelization.GPUDriven.MaxPageIndexResolution"), GHairVirtualVoxelGPUDrivenMaxPageIndexRes, TEXT("Max resolution of the page index. This is used for allocating a conservative page index buffer when GPU driven allocation is enabled."));
 
+static int32 GHairVirtualVoxelUseIndirectScatterPageAllocation = 1;
+static FAutoConsoleVariableRef CVarHairVirtualVoxelUseIndirectScatterPageAllocation(TEXT("r.HairStrands.Voxelization.UseIndiretScatterPageAllocate"), GHairVirtualVoxelUseIndirectScatterPageAllocation, TEXT("Enable indirect scatter page allocation (faster)."));
+
 static const FIntPoint GPUDrivenViewportResolution = FIntPoint(4096, 4096);
 
+static int32 GHairVirtualVoxelInvalidEmptyPageIndex = 1;
+static FAutoConsoleVariableRef CVarHairVirtualVoxelInvalidEmptyPageIndex(TEXT("r.HairStrands.Voxelization.Virtual.InvalidateEmptyPageIndex"), GHairVirtualVoxelInvalidEmptyPageIndex, TEXT("Invalid voxel page index which does not contain any voxelized data."));
+
+static int32 GHairStrandsVoxelComputeRaster = 1;
+static FAutoConsoleVariableRef CVarHairStrandsVoxelComputeRaster(TEXT("r.HairStrands.Voxelization.Virtual.ComputeRaster"), GHairStrandsVoxelComputeRaster, TEXT("Use compute for rasterizing voxeliation (faster)."));
+
+static int32 GHairStrandsVoxelComputeRasterMaxVoxelCount = 32;
+static FAutoConsoleVariableRef CVarHairStrandsVoxelComputeRasterMaxVoxelCount(TEXT("r.HairStrands.Voxelization.Virtual.ComputeRasterMaxVoxelCount"), GHairStrandsVoxelComputeRasterMaxVoxelCount, TEXT("Max number of voxel which are rasterized for a given hair segment. This is for debug purpose only."));
+
 bool IsHairStrandsVoxelizationEnable()
 {
 	return GHairVoxelizationEnable > 0;
@@ -322,6 +334,73 @@ public:
 	}
 };
 
+class FVoxelMarkValidPageIndex_PrepareCS : public FGlobalShader
+{
+	DECLARE_GLOBAL_SHADER(FVoxelMarkValidPageIndex_PrepareCS);
+	SHADER_USE_PARAMETER_STRUCT(FVoxelMarkValidPageIndex_PrepareCS, FGlobalShader);
+
+	BEGIN_SHADER_PARAMETER_STRUCT(FParameters, )
+		SHADER_PARAMETER(uint32, MaxClusterCount)
+		SHADER_PARAMETER(uint32, MacroGroupId)
+		SHADER_PARAMETER(uint32, MaxScatterAllocationCount)
+		SHADER_PARAMETER_SRV(Buffer, ClusterAABBsBuffer)
+		SHADER_PARAMETER_RDG_BUFFER_SRV(Buffer, MacroGroupAABBBuffer)
+		SHADER_PARAMETER_RDG_BUFFER_SRV(Buffer, PageIndexResolutionAndOffsetBuffer)
+		SHADER_PARAMETER_RDG_BUFFER_UAV(Buffer<uint>, OutValidPageIndexBuffer)
+		SHADER_PARAMETER_RDG_BUFFER_UAV(RWBuffer<uint>, OutDeferredScatterCounter)
+		SHADER_PARAMETER_RDG_BUFFER_UAV(RWBuffer<uint2>, OutDeferredScatterBuffer)
+	END_SHADER_PARAMETER_STRUCT()
+
+public:
+	static bool ShouldCompilePermutation(const FGlobalShaderPermutationParameters& Parameters) { return IsHairStrandsSupported(Parameters.Platform); }
+	static void ModifyCompilationEnvironment(const FGlobalShaderPermutationParameters& Parameters, FShaderCompilerEnvironment& OutEnvironment)
+	{
+		FGlobalShader::ModifyCompilationEnvironment(Parameters, OutEnvironment);
+		OutEnvironment.SetDefine(TEXT("SHADER_MARKVALID_PREPARE"), 1);
+	}
+};
+
+class FVoxelMarkValidPageIndex_IndirectArgsCS : public FGlobalShader
+{
+	DECLARE_GLOBAL_SHADER(FVoxelMarkValidPageIndex_IndirectArgsCS);
+	SHADER_USE_PARAMETER_STRUCT(FVoxelMarkValidPageIndex_IndirectArgsCS, FGlobalShader);
+
+	BEGIN_SHADER_PARAMETER_STRUCT(FParameters, )
+		SHADER_PARAMETER_RDG_BUFFER_SRV(Buffer<uint>, DeferredScatterCounter)
+		SHADER_PARAMETER_RDG_BUFFER_UAV(RWBuffer<uint>, OutIndirectArgsBuffer)
+	END_SHADER_PARAMETER_STRUCT()
+
+public:
+	static bool ShouldCompilePermutation(const FGlobalShaderPermutationParameters& Parameters) { return IsHairStrandsSupported(Parameters.Platform); }
+	static void ModifyCompilationEnvironment(const FGlobalShaderPermutationParameters& Parameters, FShaderCompilerEnvironment& OutEnvironment)
+	{
+		FGlobalShader::ModifyCompilationEnvironment(Parameters, OutEnvironment);
+		OutEnvironment.SetDefine(TEXT("SHADER_MARKVALID_INDIRECTARG"), 1);
+	}
+};
+
+class FVoxelMarkValidPageIndex_ScatterCS : public FGlobalShader
+{
+	DECLARE_GLOBAL_SHADER(FVoxelMarkValidPageIndex_ScatterCS);
+	SHADER_USE_PARAMETER_STRUCT(FVoxelMarkValidPageIndex_ScatterCS, FGlobalShader);
+
+	BEGIN_SHADER_PARAMETER_STRUCT(FParameters, )
+		SHADER_PARAMETER_RDG_BUFFER(Buffer<int>,		IndirectBufferArgs)
+		SHADER_PARAMETER_RDG_BUFFER_SRV(Buffer<uint4>,	PageIndexResolutionAndOffsetBuffer)
+		SHADER_PARAMETER_RDG_BUFFER_SRV(Buffer<uint>,	DeferredScatterCounter)
+		SHADER_PARAMETER_RDG_BUFFER_SRV(Buffer<uint2>,  DeferredScatterBuffer)
+		SHADER_PARAMETER_RDG_BUFFER_UAV(RWBuffer<uint>, OutValidPageIndexBuffer)
+	END_SHADER_PARAMETER_STRUCT()
+
+public:
+	static bool ShouldCompilePermutation(const FGlobalShaderPermutationParameters& Parameters) { return IsHairStrandsSupported(Parameters.Platform); }
+	static void ModifyCompilationEnvironment(const FGlobalShaderPermutationParameters& Parameters, FShaderCompilerEnvironment& OutEnvironment)
+	{
+		FGlobalShader::ModifyCompilationEnvironment(Parameters, OutEnvironment);
+		OutEnvironment.SetDefine(TEXT("SHADER_MARKVALID_SCATTER"), 1);
+	}
+};
+
 class FVoxelMarkValidPageIndexCS : public FGlobalShader
 {
 	DECLARE_GLOBAL_SHADER(FVoxelMarkValidPageIndexCS);
@@ -370,6 +449,7 @@ class FVoxelAllocateVoxelPageCS : public FGlobalShader
 		SHADER_PARAMETER_RDG_BUFFER_SRV(Buffer, PageIndexResolutionAndOffsetBuffer)
 		SHADER_PARAMETER_RDG_BUFFER_UAV(RWBuffer, PageIndexGlobalCounter)
 		SHADER_PARAMETER_RDG_BUFFER_UAV(RWBuffer, PageIndexBuffer)
+		SHADER_PARAMETER_RDG_BUFFER_UAV(RWBuffer, PageToPageIndexBuffer)
 		SHADER_PARAMETER_RDG_BUFFER_UAV(RWBuffer, PageIndexCoordBuffer)
 	END_SHADER_PARAMETER_STRUCT()
 
@@ -473,6 +553,10 @@ public:
 	}
 };
 
+IMPLEMENT_GLOBAL_SHADER(FVoxelMarkValidPageIndex_PrepareCS, "/Engine/Private/HairStrands/HairStrandsVoxelPageAllocation.usf", "MarkValid_PrepareCS", SF_Compute);
+IMPLEMENT_GLOBAL_SHADER(FVoxelMarkValidPageIndex_IndirectArgsCS, "/Engine/Private/HairStrands/HairStrandsVoxelPageAllocation.usf", "MarkValid_BuildIndirectArgCS", SF_Compute);
+IMPLEMENT_GLOBAL_SHADER(FVoxelMarkValidPageIndex_ScatterCS, "/Engine/Private/HairStrands/HairStrandsVoxelPageAllocation.usf", "MarkValid_ScatterCS", SF_Compute);
+
 IMPLEMENT_GLOBAL_SHADER(FVoxelAllocatePageIndexCS, "/Engine/Private/HairStrands/HairStrandsVoxelPageAllocation.usf", "AllocatePageIndex", SF_Compute);
 IMPLEMENT_GLOBAL_SHADER(FVoxelMarkValidPageIndexCS, "/Engine/Private/HairStrands/HairStrandsVoxelPageAllocation.usf", "MarkValidCS", SF_Compute);
 IMPLEMENT_GLOBAL_SHADER(FVoxelAllocateVoxelPageCS, "/Engine/Private/HairStrands/HairStrandsVoxelPageAllocation.usf", "AllocateCS", SF_Compute);
@@ -498,6 +582,7 @@ static void AddAllocateVoxelPagesPass(
 	const uint32 IndirectDispatchGroupSize,
 	uint32& OutTotalPageIndexCount,
 	FRDGBufferRef& OutPageIndexBuffer,
+	FRDGBufferRef& OutPageToPageIndexBuffer,
 	FRDGBufferRef& OutPageIndexCoordBuffer,
 	FRDGBufferRef& OutNodeDescBuffer,
 	FRDGBufferRef& OutIndirectArgsBuffer,
@@ -566,6 +651,9 @@ static void AddAllocateVoxelPagesPass(
 	FRDGBufferRef NodeDescBuffer = GraphBuilder.CreateBuffer(FRDGBufferDesc::CreateStructuredDesc(sizeof(FPackedVirtualVoxelNodeDesc), MacroGroupCount), TEXT("VirtualVoxelNodeDescBuffer"));
 	FRDGBufferRef IndirectArgsBuffer = GraphBuilder.CreateBuffer(FRDGBufferDesc::CreateIndirectDesc<FRHIDispatchIndirectParameters>(MacroGroupCount), TEXT("VirtualVoxelIndirectArgsBuffer"));
 
+	const uint32 TotalPageCount = PageCountResolution.X * PageCountResolution.Y * PageCountResolution.Z;
+	FRDGBufferRef PageToPageIndexBuffer = GraphBuilder.CreateBuffer(FRDGBufferDesc::CreateBufferDesc(sizeof(uint32), TotalPageCount), TEXT("PageToPageIndexBuffer"));
+
 	FRDGBufferUAVRef PageIndexBufferUAV = GraphBuilder.CreateUAV(PageIndexBuffer, PF_R32_UINT);
 	FRDGBufferUAVRef PageIndexGlobalCounterUAV = GraphBuilder.CreateUAV(PageIndexGlobalCounter, PF_R32_UINT);
 	
@@ -612,6 +700,20 @@ static void AddAllocateVoxelPagesPass(
 			Parameters,
 			FIntVector(1,1,1));
 	}
+	FRDGBufferSRVRef PageIndexResolutionAndOffsetBufferSRV = GraphBuilder.CreateSRV(PageIndexResolutionBuffer, PF_R32G32B32A32_UINT);
+
+	uint32 TotalClusterCount = 0;
+	for (uint32 MacroGroupIt = 0; MacroGroupIt < MacroGroupCount; ++MacroGroupIt)
+	{
+		const FHairStrandsMacroGroupData& MacroGroup = MacroGroups.Datas[MacroGroupIt];
+		for (const FHairStrandsMacroGroupData::PrimitiveGroup& PrimitiveGroup : MacroGroup.PrimitivesGroups)
+		{
+			const FHairStrandsPrimitiveResources& Resources = GetHairStandsPrimitiveResources(PrimitiveGroup.ResourceId);
+			check(PrimitiveGroup.GroupIndex < uint32(Resources.Groups.Num()));
+			const FHairStrandsPrimitiveResources::FHairGroup& GroupResources = Resources.Groups[PrimitiveGroup.GroupIndex];
+			TotalClusterCount += GroupResources.ClusterCount;
+		}
+	}
 
 	// Mark valid page index
 	for (uint32 MacroGroupIt=0; MacroGroupIt <MacroGroupCount;++MacroGroupIt)
@@ -623,40 +725,119 @@ static void AddAllocateVoxelPagesPass(
 		const FHairStrandsMacroGroupData& MacroGroup = MacroGroups.Datas[MacroGroupIt];
 		FCPUMacroGroupAllocation& CPUAllocationDesc = CPUAllocationDescs[MacroGroupIt];
 
-		for (const FHairStrandsMacroGroupData::PrimitiveGroup& PrimitiveGroup : MacroGroup.PrimitivesGroups)
-		{
-			const FHairStrandsPrimitiveResources& Resources = GetHairStandsPrimitiveResources(PrimitiveGroup.ResourceId);
-			check(PrimitiveGroup.GroupIndex < uint32(Resources.Groups.Num()));
-			const FHairStrandsPrimitiveResources::FHairGroup& GroupResources = Resources.Groups[PrimitiveGroup.GroupIndex];
+		const bool bUseIndirectScatter = GHairVirtualVoxelUseIndirectScatterPageAllocation>0 && bIsGPUDriven;
+		if (bUseIndirectScatter)
+		{			
+			const uint32 AverageWorkItemPerCluster = 8; // Arbitrary/Guess number
+			const uint32 MaxAllocationCount = TotalClusterCount * AverageWorkItemPerCluster;
+			FRDGBufferRef ScatterCounter = GraphBuilder.CreateBuffer(FRDGBufferDesc::CreateBufferDesc(sizeof(uint32), 1), TEXT("PageScatterCounter"));
+			FRDGBufferRef ScatterBuffer = GraphBuilder.CreateBuffer(FRDGBufferDesc::CreateBufferDesc(2 * sizeof(uint32), MaxAllocationCount), TEXT("PageScatterBuffer"));
 
-			FVoxelMarkValidPageIndexCS::FParameters* Parameters = GraphBuilder.AllocParameters<FVoxelMarkValidPageIndexCS::FParameters>();
-			Parameters->MacroGroupId = MacroGroup.MacroGroupId;
-			Parameters->MaxClusterCount = GroupResources.ClusterCount;
-			Parameters->CPU_PageIndexResolution = CPUAllocationDesc.PageIndexResolution;
-			Parameters->CPU_PageIndexOffset = CPUAllocationDesc.PageIndexOffset;
-			Parameters->CPU_MinAABB = CPUAllocationDesc.MinAABB;
-			Parameters->CPU_MaxAABB = CPUAllocationDesc.MaxAABB;
-			Parameters->ClusterAABBsBuffer = GroupResources.ClusterAABBBuffer->SRV;
-			Parameters->OutValidPageIndexBuffer = PageIndexBufferUAV;
+			FRDGBufferUAVRef ScatterCounterUAV = GraphBuilder.CreateUAV(ScatterCounter, PF_R32_UINT);
+			FRDGBufferUAVRef ScatterBufferUAV  = GraphBuilder.CreateUAV(ScatterBuffer, PF_R32G32_UINT);
 
-			if (bIsGPUDriven)
+			AddClearUAVPass(GraphBuilder, ScatterCounterUAV, 0);
+
+			// Prepare
+			for (const FHairStrandsMacroGroupData::PrimitiveGroup& PrimitiveGroup : MacroGroup.PrimitivesGroups)
 			{
-				Parameters->MacroGroupAABBBuffer = GraphBuilder.CreateSRV(MacroGroupAABB, PF_R32_SINT);
-				Parameters->PageIndexResolutionAndOffsetBuffer = GraphBuilder.CreateSRV(PageIndexResolutionBuffer, PF_R32G32B32A32_UINT);
+				const FHairStrandsPrimitiveResources& Resources = GetHairStandsPrimitiveResources(PrimitiveGroup.ResourceId);
+				check(PrimitiveGroup.GroupIndex < uint32(Resources.Groups.Num()));
+				const FHairStrandsPrimitiveResources::FHairGroup& GroupResources = Resources.Groups[PrimitiveGroup.GroupIndex];
+
+				FVoxelMarkValidPageIndex_PrepareCS::FParameters* Parameters = GraphBuilder.AllocParameters<FVoxelMarkValidPageIndex_PrepareCS::FParameters>();
+				Parameters->MaxClusterCount			= GroupResources.ClusterCount;
+				Parameters->MacroGroupId			= MacroGroup.MacroGroupId;
+				Parameters->MaxScatterAllocationCount= MaxAllocationCount;
+
+				Parameters->ClusterAABBsBuffer		= GroupResources.ClusterAABBBuffer->SRV;
+				Parameters->MacroGroupAABBBuffer	= GraphBuilder.CreateSRV(MacroGroupAABB, PF_R32_SINT);
+				Parameters->PageIndexResolutionAndOffsetBuffer = PageIndexResolutionAndOffsetBufferSRV;
+
+				Parameters->OutDeferredScatterCounter	= ScatterCounterUAV;
+				Parameters->OutDeferredScatterBuffer	= ScatterBufferUAV;
+				Parameters->OutValidPageIndexBuffer		= PageIndexBufferUAV;
+
+				FIntVector DispatchCount((GroupResources.ClusterCount + GroupSize - 1) / GroupSize, 1, 1);
+				check(DispatchCount.X < 65535);
+				TShaderMapRef<FVoxelMarkValidPageIndex_PrepareCS> ComputeShader(View.ShaderMap);
+				FComputeShaderUtils::AddPass(
+					GraphBuilder,
+					RDG_EVENT_NAME("HairStrandsMarkValidPageIndex_Prepare"),
+					ComputeShader,
+					Parameters,
+					DispatchCount);
 			}
+			
+			FRDGBufferSRVRef ScatterCounterSRV = GraphBuilder.CreateSRV(ScatterCounter, PF_R32_UINT);
+
+			// Build indirect buffer args
+			FRDGBufferRef ScatterIndirectArgsBuffer = GraphBuilder.CreateBuffer(FRDGBufferDesc::CreateIndirectDesc<FRHIDispatchIndirectParameters>(1), TEXT("PageScatterIndirectArgs"));
+			{
+				check(MacroGroup.MacroGroupId < MacroGroupCount);
+
+				FVoxelMarkValidPageIndex_IndirectArgsCS::FParameters* Parameters = GraphBuilder.AllocParameters<FVoxelMarkValidPageIndex_IndirectArgsCS::FParameters>();
+				Parameters->DeferredScatterCounter = ScatterCounterSRV;
+				Parameters->OutIndirectArgsBuffer = GraphBuilder.CreateUAV(ScatterIndirectArgsBuffer);
 
-			FVoxelMarkValidPageIndexCS::FPermutationDomain PermutationVector;
-			PermutationVector.Set<FVoxelMarkValidPageIndexCS::FGPUDriven>(bIsGPUDriven ? 1 : 0);
-
-			FIntVector DispatchCount((GroupResources.ClusterCount + GroupSize - 1) / GroupSize, 1, 1);
-			check(DispatchCount.X < 65535);
-			TShaderMapRef<FVoxelMarkValidPageIndexCS> ComputeShader(View.ShaderMap, PermutationVector);
-			FComputeShaderUtils::AddPass(
-				GraphBuilder,
-				RDG_EVENT_NAME("HairStrandsMarkValidPageIndex"),
-				ComputeShader,
-				Parameters,
-				DispatchCount);
+				const FIntVector DispatchCount(1, 1, 1);
+				TShaderMapRef<FVoxelMarkValidPageIndex_IndirectArgsCS> ComputeShader(View.ShaderMap);
+				FComputeShaderUtils::AddPass(GraphBuilder, RDG_EVENT_NAME("HairStrandsMarkValidPageIndex_IndirectArgs"), ComputeShader, Parameters, DispatchCount);
+			}
+
+			// Scatter
+			{
+				FRDGBufferSRVRef ScatterBufferSRV = GraphBuilder.CreateSRV(ScatterBuffer, PF_R32G32_UINT);
+				check(MacroGroup.MacroGroupId < MacroGroupCount);
+
+				FVoxelMarkValidPageIndex_ScatterCS::FParameters* Parameters = GraphBuilder.AllocParameters<FVoxelMarkValidPageIndex_ScatterCS::FParameters>();
+				Parameters->IndirectBufferArgs = ScatterIndirectArgsBuffer;
+				Parameters->PageIndexResolutionAndOffsetBuffer = PageIndexResolutionAndOffsetBufferSRV;
+				Parameters->DeferredScatterCounter = ScatterCounterSRV;
+				Parameters->DeferredScatterBuffer = ScatterBufferSRV;
+				Parameters->OutValidPageIndexBuffer = PageIndexBufferUAV;
+
+				TShaderMapRef<FVoxelMarkValidPageIndex_ScatterCS> ComputeShader(View.ShaderMap);
+				FComputeShaderUtils::AddPass(GraphBuilder, RDG_EVENT_NAME("HairStrandsMarkValidPageIndex_Scatter"), ComputeShader, Parameters, ScatterIndirectArgsBuffer, 0);
+			}
+		}
+		else
+		{		
+			for (const FHairStrandsMacroGroupData::PrimitiveGroup& PrimitiveGroup : MacroGroup.PrimitivesGroups)
+			{
+				const FHairStrandsPrimitiveResources& Resources = GetHairStandsPrimitiveResources(PrimitiveGroup.ResourceId);
+				check(PrimitiveGroup.GroupIndex < uint32(Resources.Groups.Num()));
+				const FHairStrandsPrimitiveResources::FHairGroup& GroupResources = Resources.Groups[PrimitiveGroup.GroupIndex];
+
+				FVoxelMarkValidPageIndexCS::FParameters* Parameters = GraphBuilder.AllocParameters<FVoxelMarkValidPageIndexCS::FParameters>();
+				Parameters->MacroGroupId = MacroGroup.MacroGroupId;
+				Parameters->MaxClusterCount = GroupResources.ClusterCount;
+				Parameters->CPU_PageIndexResolution = CPUAllocationDesc.PageIndexResolution;
+				Parameters->CPU_PageIndexOffset = CPUAllocationDesc.PageIndexOffset;
+				Parameters->CPU_MinAABB = CPUAllocationDesc.MinAABB;
+				Parameters->CPU_MaxAABB = CPUAllocationDesc.MaxAABB;
+				Parameters->ClusterAABBsBuffer = GroupResources.ClusterAABBBuffer->SRV;
+				Parameters->OutValidPageIndexBuffer = PageIndexBufferUAV;
+
+				if (bIsGPUDriven)
+				{
+					Parameters->MacroGroupAABBBuffer = GraphBuilder.CreateSRV(MacroGroupAABB, PF_R32_SINT);
+					Parameters->PageIndexResolutionAndOffsetBuffer = PageIndexResolutionAndOffsetBufferSRV;
+				}
+
+				FVoxelMarkValidPageIndexCS::FPermutationDomain PermutationVector;
+				PermutationVector.Set<FVoxelMarkValidPageIndexCS::FGPUDriven>(bIsGPUDriven ? 1 : 0);
+
+				FIntVector DispatchCount((GroupResources.ClusterCount + GroupSize - 1) / GroupSize, 1, 1);
+				check(DispatchCount.X < 65535);
+				TShaderMapRef<FVoxelMarkValidPageIndexCS> ComputeShader(View.ShaderMap, PermutationVector);
+				FComputeShaderUtils::AddPass(
+					GraphBuilder,
+					RDG_EVENT_NAME("HairStrandsMarkValidPageIndex"),
+					ComputeShader,
+					Parameters,
+					DispatchCount);
+			}
 		}
 
 		// Fill in hair-macro-group information.
@@ -676,7 +857,7 @@ static void AddAllocateVoxelPagesPass(
 			if (bIsGPUDriven)
 			{
 				Parameters->MacroGroupAABBBuffer = GraphBuilder.CreateSRV(MacroGroupAABB, PF_R32_SINT);
-				Parameters->PageIndexResolutionAndOffsetBuffer = GraphBuilder.CreateSRV(PageIndexResolutionBuffer, PF_R32G32B32A32_UINT);
+				Parameters->PageIndexResolutionAndOffsetBuffer = PageIndexResolutionAndOffsetBufferSRV;
 			}
 
 			FVoxelAddNodeDescCS::FPermutationDomain PermutationVector;
@@ -697,6 +878,7 @@ static void AddAllocateVoxelPagesPass(
 			Parameters->CPU_PageIndexOffset = CPUAllocationDesc.PageIndexOffset;
 			Parameters->PageIndexGlobalCounter = PageIndexGlobalCounterUAV;
 			Parameters->PageIndexBuffer = PageIndexBufferUAV;
+			Parameters->PageToPageIndexBuffer = GraphBuilder.CreateUAV(PageToPageIndexBuffer, PF_R32_UINT);
 			Parameters->PageIndexCoordBuffer = GraphBuilder.CreateUAV(PageIndexCoordBuffer, PF_R8G8B8A8_UINT);
 
 			FVoxelAllocateVoxelPageCS::FPermutationDomain PermutationVector;
@@ -705,7 +887,7 @@ static void AddAllocateVoxelPagesPass(
 
 			if (bIsGPUDriven)
 			{
-				Parameters->PageIndexResolutionAndOffsetBuffer = GraphBuilder.CreateSRV(PageIndexResolutionBuffer, PF_R32G32B32A32_UINT);
+				Parameters->PageIndexResolutionAndOffsetBuffer = PageIndexResolutionAndOffsetBufferSRV;
 				Parameters->IndirectBufferArgs = PageIndexAllocationIndirectBufferArgs;
 
 				const uint32 ArgsOffset = sizeof(uint32) * 3 * MacroGroup.MacroGroupId;
@@ -740,6 +922,7 @@ static void AddAllocateVoxelPagesPass(
 	}
 
 	OutPageIndexBuffer = PageIndexBuffer;
+	OutPageToPageIndexBuffer = PageToPageIndexBuffer;
 	OutPageIndexCoordBuffer = PageIndexCoordBuffer;
 	OutNodeDescBuffer = NodeDescBuffer;
 	OutIndirectArgsBuffer = IndirectArgsBuffer;
@@ -750,10 +933,16 @@ static void AddAllocateVoxelPagesPass(
 FVirtualVoxelResources AllocateVirtualVoxelResources(
 	FRHICommandListImmediate& RHICmdList,
 	const FViewInfo& View,
-	FHairStrandsMacroGroupDatas& MacroGroups)
+	FHairStrandsMacroGroupDatas& MacroGroups, 
+	TRefCountPtr<FPooledRDGBuffer>& PageToPageIndexBuffer)
 {
+	DECLARE_GPU_STAT(HairStrandsVoxelPageAllocation);
+	SCOPED_DRAW_EVENT(RHICmdList, HairStrandsVoxelPageAllocation);
+	SCOPED_GPU_STAT(RHICmdList, HairStrandsVoxelPageAllocation);
+
 	FRDGBuilder GraphBuilder(RHICmdList);
 	FRDGBufferRef OutPageIndexBuffer = nullptr;
+	FRDGBufferRef OutPageToPageIndexBuffer = nullptr;
 	FRDGBufferRef OutPageIndexCoordBuffer = nullptr;
 	FRDGBufferRef OutNodeDescBuffer = nullptr;
 	FRDGBufferRef OutIndirectArgsBuffer = nullptr;
@@ -770,6 +959,7 @@ FVirtualVoxelResources AllocateVirtualVoxelResources(
 	Out.Parameters.Common.DensityScale				= GetHairStrandsVoxelizationDensityScale();
 	Out.Parameters.Common.DepthBiasScale			= GetHairStrandsVoxelizationDepthBiasScale();
 	Out.Parameters.Common.SteppingScale				= FMath::Clamp(GHairStransVoxelRaymarchingSteppingScale, 1.f, 10.f);
+	Out.Parameters.Common.NodeDescCount				= MacroGroups.Datas.Num();
 	Out.Parameters.Common.IndirectDispatchGroupSize = 64;
 
 	AddAllocateVoxelPagesPass(
@@ -784,6 +974,7 @@ FVirtualVoxelResources AllocateVirtualVoxelResources(
 		Out.Parameters.Common.IndirectDispatchGroupSize,
 		Out.Parameters.Common.PageIndexCount,
 		OutPageIndexBuffer, 
+		OutPageToPageIndexBuffer,
 		OutPageIndexCoordBuffer,
 		OutNodeDescBuffer,
 		OutIndirectArgsBuffer,
@@ -795,6 +986,11 @@ FVirtualVoxelResources AllocateVirtualVoxelResources(
 		GraphBuilder.QueueBufferExtraction(OutPageIndexBuffer, &Out.PageIndexBuffer, FRDGResourceState::EAccess::Read, FRDGResourceState::EPipeline::Graphics);
 	}
 
+	if (OutPageToPageIndexBuffer)
+	{
+		GraphBuilder.QueueBufferExtraction(OutPageToPageIndexBuffer, &PageToPageIndexBuffer, FRDGResourceState::EAccess::Read, FRDGResourceState::EPipeline::Graphics);
+	}
+
 	if (OutPageIndexCoordBuffer)
 	{
 		GraphBuilder.QueueBufferExtraction(OutPageIndexCoordBuffer, &Out.PageIndexCoordBuffer, FRDGResourceState::EAccess::Read, FRDGResourceState::EPipeline::Compute);
@@ -922,6 +1118,92 @@ static FRDGBufferRef IndirectVoxelPageClear(
 	return ClearIndArgsBuffer;
 }
 
+///////////////////////////////////////////////////////////////////////////////////////////////////
+
+class FVoxelRasterComputeCS : public FGlobalShader
+{
+	DECLARE_GLOBAL_SHADER(FVoxelRasterComputeCS);
+	SHADER_USE_PARAMETER_STRUCT(FVoxelRasterComputeCS, FGlobalShader);
+
+	BEGIN_SHADER_PARAMETER_STRUCT(FParameters, )
+		SHADER_PARAMETER_STRUCT(FVirtualVoxelCommonParameters, VirtualVoxel)
+		SHADER_PARAMETER(uint32, MacroGroupId)
+		SHADER_PARAMETER(uint32, DispatchCountX)
+		SHADER_PARAMETER(uint32, MaxRasterCount)
+		SHADER_PARAMETER(uint32, FrameIdMod8)
+		SHADER_PARAMETER(float,	  HairStrandsVF_Density)
+		SHADER_PARAMETER(float,   HairStrandsVF_Radius)
+		SHADER_PARAMETER(float,	  HairStrandsVF_Length)
+		SHADER_PARAMETER(FVector, HairStrandsVF_PositionOffset)
+		SHADER_PARAMETER(uint32,  HairStrandsVF_VertexCount)
+		SHADER_PARAMETER(FMatrix, HairStrandsVF_LocalToWorldPrimitiveTransform)
+		SHADER_PARAMETER_SRV(Buffer, HairStrandsVF_PositionBuffer)
+		SHADER_PARAMETER_RDG_BUFFER_SRV(Buffer, VoxelizationViewInfoBuffer)
+		SHADER_PARAMETER_RDG_TEXTURE_UAV(RWTexture3D, OutPageTexture)
+	END_SHADER_PARAMETER_STRUCT()
+
+public:
+	static bool ShouldCompilePermutation(const FGlobalShaderPermutationParameters& Parameters) { return IsHairStrandsSupported(Parameters.Platform); }
+	static void ModifyCompilationEnvironment(const FGlobalShaderPermutationParameters& Parameters, FShaderCompilerEnvironment& OutEnvironment)
+	{
+		FGlobalShader::ModifyCompilationEnvironment(Parameters, OutEnvironment);
+		OutEnvironment.SetDefine(TEXT("SHADER_RASTERCOMPUTE"), 1);
+	}
+};
+
+IMPLEMENT_GLOBAL_SHADER(FVoxelRasterComputeCS, "/Engine/Private/HairStrands/HairStrandsVoxelRasterCompute.usf", "MainCS", SF_Compute);
+
+static void AddVirtualVoxelizationComputeRasterPass(
+	FRDGBuilder& GraphBuilder,
+	const FViewInfo* ViewInfo,
+	FVirtualVoxelResources& VoxelResources,
+	FHairStrandsMacroGroupData& MacroGroup)
+{
+	const bool bIsGPUDriven = GHairVirtualVoxelGPUDriven > 0;
+	if (!bIsGPUDriven)
+		return;
+
+	const FHairStrandsMacroGroupData::TPrimitiveInfos& PrimitiveSceneInfos = MacroGroup.PrimitivesInfos;
+
+	FRDGBufferRef VoxelizationViewInfoBuffer = GraphBuilder.RegisterExternalBuffer(VoxelResources.VoxelizationViewInfoBuffer);
+	FRDGBufferSRVRef VoxelizationViewInfoBufferSRV = GraphBuilder.CreateSRV(VoxelizationViewInfoBuffer);
+	FRDGTextureRef PageTexture = GraphBuilder.RegisterExternalTexture(VoxelResources.PageTexture);
+	FRDGTextureUAVRef PageTextureUAV = GraphBuilder.CreateUAV(PageTexture);
+
+	const uint32 FrameIdMode8 = ViewInfo && ViewInfo->ViewState ? (ViewInfo->ViewState->GetFrameIndex() % 8) : 0;
+	const uint32 GroupSize = 32;
+	const uint32 DispatchCountX = 64;
+	TShaderMapRef<FVoxelRasterComputeCS> ComputeShader(ViewInfo->ShaderMap);
+	for (const FHairStrandsMacroGroupData::PrimitiveInfo& PrimitiveInfo : PrimitiveSceneInfos)
+	{
+		FVoxelRasterComputeCS::FParameters* PassParameters = GraphBuilder.AllocParameters<FVoxelRasterComputeCS::FParameters>();
+		PassParameters->MaxRasterCount				= FMath::Clamp(GHairStrandsVoxelComputeRasterMaxVoxelCount, 1, 256);
+		PassParameters->VirtualVoxel				= VoxelResources.Parameters.Common;
+		PassParameters->MacroGroupId				= MacroGroup.MacroGroupId;
+		PassParameters->VoxelizationViewInfoBuffer	= VoxelizationViewInfoBufferSRV;
+		PassParameters->DispatchCountX				= DispatchCountX;
+		PassParameters->OutPageTexture				= PageTextureUAV;
+		PassParameters->FrameIdMod8					= FrameIdMode8;
+
+		check(PrimitiveInfo.MeshBatchAndRelevance.Mesh && PrimitiveInfo.MeshBatchAndRelevance.Mesh->Elements.Num() > 0);
+		const FHairGroupPublicData* HairGroupPublicData = reinterpret_cast<const FHairGroupPublicData*>(PrimitiveInfo.MeshBatchAndRelevance.Mesh->Elements[0].VertexFactoryUserData);
+		const FHairGroupPublicData::VertexFactoryInput& VFInput = HairGroupPublicData->VFInput;
+		PassParameters->HairStrandsVF_PositionBuffer = VFInput.HairPositionBuffer;
+		PassParameters->HairStrandsVF_PositionOffset = VFInput.HairPositionOffset;
+		PassParameters->HairStrandsVF_VertexCount	 = VFInput.VertexCount;
+		PassParameters->HairStrandsVF_Radius		 = VFInput.HairRadius;
+		PassParameters->HairStrandsVF_Length		 = VFInput.HairLength;
+		PassParameters->HairStrandsVF_Density		 = VFInput.HairDensity;
+		PassParameters->HairStrandsVF_LocalToWorldPrimitiveTransform = VFInput.LocalToWorldTransform.ToMatrixWithScale();
+
+		const uint32 DispatchCountY = FMath::CeilToInt(PassParameters->HairStrandsVF_VertexCount / float(GroupSize * DispatchCountX));
+		const FIntVector DispatchCount(DispatchCountX, DispatchCountY, 1);
+		FComputeShaderUtils::AddPass(GraphBuilder, RDG_EVENT_NAME("HairStrandsVoxelComputeRaster"), ComputeShader, PassParameters, DispatchCount);
+	}
+}
+
+///////////////////////////////////////////////////////////////////////////////////////////////////
+
 static void AddVirtualVoxelizationRasterPass(
 	FRDGBuilder& GraphBuilder,
 	const FScene* Scene,
@@ -994,7 +1276,8 @@ static void AddVirtualVoxelizationRasterPass(
 
 	const float RadiusAtDepth1 = GStrandHairVoxelizationRasterizationScale * VoxelResources.Parameters.Common.VoxelWorldSize;
 	const bool bIsOrtho = true;
-	const FVector4 HairRenderInfo = PackHairRenderInfo(RadiusAtDepth1, RadiusAtDepth1, 1, bIsOrtho, bIsGPUDriven);
+	const FVector4 HairRenderInfo = PackHairRenderInfo(RadiusAtDepth1, RadiusAtDepth1, RadiusAtDepth1, 1);
+	const uint32 HairRenderInfoBits = PackHairRenderInfoBits(bIsOrtho, bIsGPUDriven);
 
 	FMatrix WorldToClip;
 	{
@@ -1007,6 +1290,13 @@ static void AddVirtualVoxelizationRasterPass(
 	FRDGBufferRef VoxelizationViewInfoBuffer = GraphBuilder.RegisterExternalBuffer(VoxelResources.VoxelizationViewInfoBuffer);
 	FRDGTextureRef PageTexture = GraphBuilder.RegisterExternalTexture(VoxelResources.PageTexture);
 
+	const bool bUseComputeRaster = GHairStrandsVoxelComputeRaster > 0;
+	if (bIsGPUDriven && bUseComputeRaster)
+	{	
+		AddVirtualVoxelizationComputeRasterPass(GraphBuilder, ViewInfo, VoxelResources, MacroGroup);
+		return;
+	}
+
 	FHairVoxelizationRasterPassParameters* PassParameters = GraphBuilder.AllocParameters<FHairVoxelizationRasterPassParameters>();
 	PassParameters->VirtualVoxel = VoxelResources.Parameters.Common;
 	PassParameters->WorldToClipMatrix = WorldToClip;
@@ -1018,16 +1308,25 @@ static void AddVirtualVoxelizationRasterPass(
 	PassParameters->VoxelizationViewInfoBuffer = GraphBuilder.CreateSRV(VoxelizationViewInfoBuffer);
 	PassParameters->DensityTexture = GraphBuilder.CreateUAV(PageTexture);
 
+	const FPooledRenderTargetDesc DummyDesc = FPooledRenderTargetDesc::CreateVolumeDesc(1, 1, 1, PF_R32_UINT, FClearValueBinding::Black, TexCreate_None, TexCreate_UAV | TexCreate_ShaderResource, false, 1);
+	FRDGTextureRef DummyTangentXTexture = GraphBuilder.CreateTexture(DummyDesc, TEXT("DummyTexture"));
+	FRDGTextureRef DummyTangentYTexture = GraphBuilder.CreateTexture(DummyDesc, TEXT("DummyTexture"));
+	FRDGTextureRef DummyTangentZTexture = GraphBuilder.CreateTexture(DummyDesc, TEXT("DummyTexture"));
+	FRDGTextureRef DummyMaterialTexture = GraphBuilder.CreateTexture(DummyDesc, TEXT("DummyTexture"));
+	PassParameters->TangentXTexture = GraphBuilder.CreateUAV(DummyTangentXTexture);
+	PassParameters->TangentYTexture = GraphBuilder.CreateUAV(DummyTangentYTexture);
+	PassParameters->TangentZTexture = GraphBuilder.CreateUAV(DummyTangentZTexture);
+	PassParameters->MaterialTexture = GraphBuilder.CreateUAV(DummyMaterialTexture);
+
 	// For debug purpose
 	#if 0
-	FRDGTextureRef DummyTexture = GraphBuilder.CreateTexture(FPooledRenderTargetDesc::Create2DDesc(RasterResolution, PF_R32_UINT, FClearValueBinding::Black, TexCreate_None, TexCreate_RenderTargetable, false), TEXT("DummyTexture"));
-	PassParameters->RenderTargets[0] = FRenderTargetBinding(DummyTexture, ERenderTargetLoadAction::EClear);
+	FRDGTextureRef DebugOutputTexture = GraphBuilder.CreateTexture(FPooledRenderTargetDesc::Create2DDesc(RasterResolution, PF_R32_UINT, FClearValueBinding::Black, TexCreate_None, TexCreate_RenderTargetable, false), TEXT("DummyTexture"));
+	PassParameters->RenderTargets[0] = FRenderTargetBinding(DebugOutputTexture, ERenderTargetLoadAction::EClear);
 	#endif
 
-	AddHairVoxelizationRasterPass(GraphBuilder, Scene, ViewInfo, PrimitiveSceneInfo, EHairStrandsRasterPassType::VoxelizationVirtual, ViewportRect, HairRenderInfo, RasterDirection, PassParameters);
+	AddHairVoxelizationRasterPass(GraphBuilder, Scene, ViewInfo, PrimitiveSceneInfo, EHairStrandsRasterPassType::VoxelizationVirtual, ViewportRect, HairRenderInfo, HairRenderInfoBits, RasterDirection, PassParameters);
 }
 
-
 ///////////////////////////////////////////////////////////////////////////////////////////////////
 class FVoxelFilterDepthCS : public FGlobalShader
 {
@@ -1112,7 +1411,7 @@ public:
 	static void ModifyCompilationEnvironment(const FGlobalShaderPermutationParameters& Parameters, FShaderCompilerEnvironment& OutEnvironment)
 	{
 		FGlobalShader::ModifyCompilationEnvironment(Parameters, OutEnvironment);
-		OutEnvironment.SetDefine(TEXT("SHADER_VOXEL"), 1);
+		OutEnvironment.SetDefine(TEXT("SHADER_MIP_VOXEL"), 1);
 	}
 };
 
@@ -1197,11 +1496,10 @@ public:
 	static void ModifyCompilationEnvironment(const FGlobalShaderPermutationParameters& Parameters, FShaderCompilerEnvironment& OutEnvironment)
 	{
 		FGlobalShader::ModifyCompilationEnvironment(Parameters, OutEnvironment);
-		OutEnvironment.SetDefine(TEXT("SHADER_VIRTUALVOXEL"), 1);
+		OutEnvironment.SetDefine(TEXT("SHADER_MIP_VIRTUALVOXEL"), 1);
 	}
 };
 
-
 class FVirtualVoxelIndirectArgMipCS : public FGlobalShader
 {
 	DECLARE_GLOBAL_SHADER(FVirtualVoxelIndirectArgMipCS);
@@ -1220,19 +1518,49 @@ public:
 	static void ModifyCompilationEnvironment(const FGlobalShaderPermutationParameters& Parameters, FShaderCompilerEnvironment& OutEnvironment)
 	{
 		FGlobalShader::ModifyCompilationEnvironment(Parameters, OutEnvironment);
-		OutEnvironment.SetDefine(TEXT("SHADER_INDIRECTARGS"), 1);
+		OutEnvironment.SetDefine(TEXT("SHADER_MIP_INDIRECTARGS"), 1);
+	}
+};
+
+class FVirtualVoxelPatchPageIndexWithMipDataCS : public FGlobalShader
+{
+	DECLARE_GLOBAL_SHADER(FVirtualVoxelPatchPageIndexWithMipDataCS);
+	SHADER_USE_PARAMETER_STRUCT(FVirtualVoxelPatchPageIndexWithMipDataCS, FGlobalShader);
+
+	BEGIN_SHADER_PARAMETER_STRUCT(FParameters, )
+		SHADER_PARAMETER(FIntVector, PageCountResolution)
+		SHADER_PARAMETER(uint32, PageResolution)
+		SHADER_PARAMETER(uint32, bUpdatePageIndex)
+		SHADER_PARAMETER(uint32, MipIt)
+
+		SHADER_PARAMETER_RDG_TEXTURE(Texture3D, DensityTexture)
+		SHADER_PARAMETER_RDG_BUFFER(Buffer, IndirectDispatchArgs)
+		SHADER_PARAMETER_RDG_BUFFER_SRV(Buffer<uint>, PageToPageIndexBuffer)
+		SHADER_PARAMETER_RDG_BUFFER_UAV(RWBuffer<uint>, OutPageIndexBuffer)
+
+	END_SHADER_PARAMETER_STRUCT()
+
+public:
+	static bool ShouldCompilePermutation(const FGlobalShaderPermutationParameters& Parameters) { return IsHairStrandsSupported(Parameters.Platform); }
+	static void ModifyCompilationEnvironment(const FGlobalShaderPermutationParameters& Parameters, FShaderCompilerEnvironment& OutEnvironment)
+	{
+		FGlobalShader::ModifyCompilationEnvironment(Parameters, OutEnvironment);
+		OutEnvironment.SetDefine(TEXT("SHADER_UPDATE_PAGEINDEX"), 1);
 	}
 };
 
+
 IMPLEMENT_GLOBAL_SHADER(FVirtualVoxelGenerateMipCS, "/Engine/Private/HairStrands/HairStrandsVoxelMip.usf", "MainCS", SF_Compute);
 IMPLEMENT_GLOBAL_SHADER(FVirtualVoxelIndirectArgMipCS, "/Engine/Private/HairStrands/HairStrandsVoxelMip.usf", "MainCS", SF_Compute);
+IMPLEMENT_GLOBAL_SHADER(FVirtualVoxelPatchPageIndexWithMipDataCS, "/Engine/Private/HairStrands/HairStrandsVoxelMip.usf", "MainCS", SF_Compute);
 
 
 static void AddVirtualVoxelGenerateMipPass(
 	FRDGBuilder& GraphBuilder,
 	const FViewInfo& View,
 	FHairStrandsMacroGroupDatas& MacroGroups,
-	FRDGBufferRef IndirectArgsBuffer)
+	FRDGBufferRef IndirectArgsBuffer, 
+	TRefCountPtr<FPooledRDGBuffer>& InPageToPageIndexBuffer)
 {
 	if (!MacroGroups.VirtualVoxelResources.IsValid())
 		return;
@@ -1245,6 +1573,7 @@ static void AddVirtualVoxelGenerateMipPass(
 
 	const uint32 MipCount = VoxelResources.PageTexture->GetDesc().NumMips;
 	FRDGTextureRef VoxelDensityTexture = GraphBuilder.RegisterExternalTexture(VoxelResources.PageTexture, TEXT("HairVirtualVoxelDensityTexture"));
+	FRDGBufferRef PageIndexBuffer = GraphBuilder.RegisterExternalBuffer(VoxelResources.PageIndexBuffer, TEXT("HairVirtualVoxelDensityTexturePageIndex"));
 
 	// Prepare indirect dispatch for all the pages this frame (allocated linearly in 3D DensityTexture)
 	TArray<FRDGBufferRef> MipIndirectArgsBuffers;
@@ -1292,6 +1621,35 @@ static void AddVirtualVoxelGenerateMipPass(
 			FComputeShaderUtils::DispatchIndirect(RHICmdList, ComputeShader, *Parameters, Parameters->IndirectDispatchArgs->GetIndirectRHICallBuffer(), 0);
 		});
 	}
+
+	// Patch the page index buffer with page whose voxels are empty after the voxelization is done
+	FRDGBufferRef PageToPageIndexBuffer = GraphBuilder.RegisterExternalBuffer(InPageToPageIndexBuffer, TEXT("PageToPageIndexBuffer"));
+	FRDGBufferSRVRef PageToPageIndexBufferSRV = GraphBuilder.CreateSRV(PageToPageIndexBuffer, PF_R32_UINT);
+	FRDGBufferUAVRef PageIndexBufferUAV = GraphBuilder.CreateUAV(PageIndexBuffer, PF_R32_UINT);
+	const bool bPatchEmptyPage = GHairVirtualVoxelInvalidEmptyPageIndex > 0;
+	if (bPatchEmptyPage)
+	{
+		const uint32 LastMipIt = MipCount - 1;
+		FVirtualVoxelPatchPageIndexWithMipDataCS::FParameters* Parameters = GraphBuilder.AllocParameters<FVirtualVoxelPatchPageIndexWithMipDataCS::FParameters>();
+		Parameters->MipIt = LastMipIt;
+		Parameters->PageResolution = VoxelResources.Parameters.Common.PageResolution;
+		Parameters->PageCountResolution = VoxelResources.Parameters.Common.PageCountResolution;
+		Parameters->DensityTexture = VoxelDensityTexture;
+		Parameters->PageToPageIndexBuffer = PageToPageIndexBufferSRV;
+		Parameters->OutPageIndexBuffer = PageIndexBufferUAV;
+		Parameters->IndirectDispatchArgs = MipIndirectArgsBuffers[LastMipIt-1];
+
+		TShaderMapRef<FVirtualVoxelPatchPageIndexWithMipDataCS> ComputeShader(View.ShaderMap);
+		ClearUnusedGraphResources(ComputeShader, Parameters);
+		GraphBuilder.AddPass(
+			RDG_EVENT_NAME("HairStrandsPatchPageIndexWithMip"),
+			Parameters,
+			ERDGPassFlags::Compute,
+			[Parameters, ComputeShader](FRHICommandList& RHICmdList)
+		{
+			FComputeShaderUtils::DispatchIndirect(RHICmdList, ComputeShader, *Parameters, Parameters->IndirectDispatchArgs->GetIndirectRHICallBuffer(), 0);
+		});
+	}
 }
 /////////////////////////////////////////////////////////////////////////////////////////
 
@@ -1348,7 +1706,8 @@ static void AddVoxelizationRasterPass(
 	const float RadiusAtDepth1 = GStrandHairVoxelizationRasterizationScale * SphereRadius / FMath::Min(Resolution.X, Resolution.Y);
 	const bool bIsOrtho = true;
 	const bool bIsGPUDriven = false;
-	const FVector4 HairRenderInfo = PackHairRenderInfo(RadiusAtDepth1, RadiusAtDepth1, 1, bIsOrtho, bIsGPUDriven);
+	const FVector4 HairRenderInfo = PackHairRenderInfo(RadiusAtDepth1, RadiusAtDepth1, RadiusAtDepth1, 1);
+	const uint32 HairRenderInfoBits = PackHairRenderInfoBits(bIsOrtho, bIsGPUDriven);
 
 	{
 		FReversedZOrthoMatrix OrthoMatrix(SphereRadius, SphereRadius, 1.f / (2 * SphereRadius), 0);
@@ -1383,6 +1742,7 @@ static void AddVoxelizationRasterPass(
 		bVoxelizeMaterial ? EHairStrandsRasterPassType::VoxelizationMaterial : EHairStrandsRasterPassType::Voxelization,
 		ViewportRect,
 		HairRenderInfo,
+		HairRenderInfoBits,
 		RasterDirection,
 		PassParameters);
 
@@ -1430,7 +1790,8 @@ void VoxelizeHairStrands(
 				// Toto moves this function into the render graph. At the moment this is not possible as this functions 
 				// generates internally a non-transient constant buffer which initialized VirtualVoxelResources. This 
 				// needs to be rewritten/worked out.
-				MacroGroupDatas.VirtualVoxelResources = AllocateVirtualVoxelResources(RHICmdList, View, MacroGroupDatas);
+				TRefCountPtr<FPooledRDGBuffer> PageToPageIndexBuffer;
+				MacroGroupDatas.VirtualVoxelResources = AllocateVirtualVoxelResources(RHICmdList, View, MacroGroupDatas, PageToPageIndexBuffer);
 
 				FRDGBuilder GraphBuilder(RHICmdList);
 				FRDGBufferRef ClearIndArgsBuffer = IndirectVoxelPageClear(GraphBuilder, View, MacroGroupDatas.VirtualVoxelResources);
@@ -1448,7 +1809,7 @@ void VoxelizeHairStrands(
 					}					
 				}
 
-				AddVirtualVoxelGenerateMipPass(GraphBuilder, View, MacroGroupDatas, ClearIndArgsBuffer);
+				AddVirtualVoxelGenerateMipPass(GraphBuilder, View, MacroGroupDatas, ClearIndArgsBuffer, PageToPageIndexBuffer);
 
 				GraphBuilder.Execute();
 			}
diff --git a/Engine/Source/Runtime/Renderer/Private/HairStrands/HairStrandsVoxelization.h b/Engine/Source/Runtime/Renderer/Private/HairStrands/HairStrandsVoxelization.h
--- a/Engine/Source/Runtime/Renderer/Private/HairStrands/HairStrandsVoxelization.h
+++ b/Engine/Source/Runtime/Renderer/Private/HairStrands/HairStrandsVoxelization.h
@@ -51,6 +51,7 @@ BEGIN_SHADER_PARAMETER_STRUCT(FVirtualVoxelCommonParameters, )
 	SHADER_PARAMETER(uint32, PageResolution)
 	SHADER_PARAMETER(uint32, PageIndexCount)
 	SHADER_PARAMETER(uint32, IndirectDispatchGroupSize)
+	SHADER_PARAMETER(uint32, NodeDescCount)
 	SHADER_PARAMETER(float, DensityScale)
 	SHADER_PARAMETER(float, DepthBiasScale)
 	SHADER_PARAMETER(float, SteppingScale)
@@ -99,4 +100,12 @@ void VoxelizeHairStrands(
 	FRHICommandListImmediate& RHICmdList,
 	const class FScene* Scene,
 	const TArray<FViewInfo>& Views,
-	struct FHairStrandsMacroGroupViews& MacroGroupViews);
\ No newline at end of file
+	struct FHairStrandsMacroGroupViews& MacroGroupViews);
+
+
+struct FHairStrandsOcclusionResources
+{
+	FRDGTextureRef CategorizationTexture = nullptr;
+	FRDGTextureRef LightChannelMaskTexture = nullptr;
+	const FVirtualVoxelResources* VoxelResources = nullptr;
+};
\ No newline at end of file
