Integrate hair changes (rendering+simulation)

diff --git a/Engine/Source/Runtime/Renderer/Private/LightRendering.cpp b/Engine/Source/Runtime/Renderer/Private/LightRendering.cpp
--- a/Engine/Source/Runtime/Renderer/Private/LightRendering.cpp
+++ b/Engine/Source/Runtime/Renderer/Private/LightRendering.cpp
@@ -908,14 +908,9 @@ void FDeferredShadingSceneRenderer::RenderLights(FRHICommandListImmediate& RHICm
 		return HairDatas && ViewIndex < HairDatas->MacroGroupsPerViews.Views.Num() && HairDatas->MacroGroupsPerViews.Views[ViewIndex].Datas.Num() > 0;
 	};
 
-	struct FHairStrandsTextures
+	auto GetHairStrandsResources = [HairDatas](int32 ViewIndex, FRDGBuilder& GraphBuilder)
 	{
-		FRDGTextureRef CategorizationTexture = nullptr;
-		FRDGTextureRef LightChannelMaskTexture = nullptr;
-	};
-	auto GetHairStrandsTextures = [HairDatas](int32 ViewIndex, FRDGBuilder& GraphBuilder)
-	{
-		FHairStrandsTextures Out;
+		FHairStrandsOcclusionResources Out;
 		if (HairDatas && ViewIndex < HairDatas->HairVisibilityViews.HairDatas.Num())
 		{
 			if (HairDatas->HairVisibilityViews.HairDatas[ViewIndex].CategorizationTexture)
@@ -926,6 +921,8 @@ void FDeferredShadingSceneRenderer::RenderLights(FRHICommandListImmediate& RHICm
 			{
 				Out.LightChannelMaskTexture = GraphBuilder.RegisterExternalTexture(HairDatas->HairVisibilityViews.HairDatas[ViewIndex].LightChannelMaskTexture);
 			}
+
+			Out.VoxelResources = &HairDatas->MacroGroupsPerViews.Views[ViewIndex].VirtualVoxelResources;
 		}
 		return Out;
 	};
@@ -1172,12 +1169,12 @@ void FDeferredShadingSceneRenderer::RenderLights(FRHICommandListImmediate& RHICm
 			const bool bDoShadowDenoisingBatching = DenoiserMode != 0 && MaxDenoisingBatchSize > 1;
 
 			//#dxr_todo: support multiview for the batching case
-			const bool dDoShadowBatching = (bDoShadowDenoisingBatching || MaxRTShadowBatchSize > 1) && Views.Num() == 1;
+			const bool bDoShadowBatching = (bDoShadowDenoisingBatching || MaxRTShadowBatchSize > 1) && Views.Num() == 1;
 
 			// Optimisations: batches all shadow ray tracing denoising. Definitely could be smarter to avoid high VGPR pressure if this entire
 			// function was converted to render graph, and want least intrusive change as possible. So right not it trades render target memory pressure
 			// for denoising perf.
-			if (RHI_RAYTRACING && dDoShadowBatching)
+			if (RHI_RAYTRACING && bDoShadowBatching)
 			{
 				const uint32 ViewIndex = 0;
 				FViewInfo& View = Views[ViewIndex];
@@ -1255,7 +1252,7 @@ void FDeferredShadingSceneRenderer::RenderLights(FRHICommandListImmediate& RHICm
 						const bool bDenoiserCompatible = !LightRequiresDenosier(LightSceneInfo) || IScreenSpaceDenoiser::EShadowRequirements::PenumbraAndClosestOccluder == DenoiserToUse->GetShadowRequirements(View, LightSceneInfo, RayTracingConfig);
 
 						const bool bWantsBatchedShadow = OcclusionType == FLightOcclusionType::Raytraced && 
-							dDoShadowBatching &&
+							bDoShadowBatching &&
 							bDenoiserCompatible &&
 							SortedLightInfo.SortKey.Fields.bShadowed && 
 							!ShouldRenderRayTracingStochasticRectLight(LightSceneInfo);
@@ -1375,12 +1372,13 @@ void FDeferredShadingSceneRenderer::RenderLights(FRHICommandListImmediate& RHICm
 									RayDistanceTexture = GraphBuilder.CreateTexture(Desc, TEXT("RayTracingOcclusionDistance"));
 								}
 
+								FString BatchLightNameWithLevel;
+								GetLightNameForDrawEvent(BatchLightSceneInfo.Proxy, BatchLightNameWithLevel);
+
 								FRDGTextureUAV* RayTracingShadowMaskUAV = GraphBuilder.CreateUAV(FRDGTextureUAVDesc(RayTracingShadowMaskTexture));
 								FRDGTextureUAV* RayHitDistanceUAV = GraphBuilder.CreateUAV(FRDGTextureUAVDesc(RayDistanceTexture));
-								FHairStrandsTextures HairTextures = GetHairStrandsTextures(ViewIndex, GraphBuilder);
+								FHairStrandsOcclusionResources HairResources = GetHairStrandsResources(ViewIndex, GraphBuilder);
 								{
-									FString BatchLightNameWithLevel;
-									GetLightNameForDrawEvent(BatchLightSceneInfo.Proxy, BatchLightNameWithLevel);
 									RDG_EVENT_SCOPE(GraphBuilder, "%s", *BatchLightNameWithLevel);
 
 									RenderRayTracingShadows(
@@ -1391,8 +1389,7 @@ void FDeferredShadingSceneRenderer::RenderLights(FRHICommandListImmediate& RHICm
 										BatchRayTracingConfig,
 										DenoiserRequirements,
 										false,
-										HairTextures.CategorizationTexture,
-										HairTextures.LightChannelMaskTexture,
+										&HairResources,
 										RayTracingShadowMaskUAV,
 										RayHitDistanceUAV);
 								}
@@ -1437,8 +1434,7 @@ void FDeferredShadingSceneRenderer::RenderLights(FRHICommandListImmediate& RHICm
 										BatchRayTracingConfig,
 										IScreenSpaceDenoiser::EShadowRequirements::Bailout,
 										true,
-										HairTextures.CategorizationTexture,
-										HairTextures.LightChannelMaskTexture,
+										&HairResources,
 										SubPixelRayTracingShadowMaskUAV,
 										SubPixelRayHitDistanceUAV);
 
@@ -1603,7 +1599,7 @@ void FDeferredShadingSceneRenderer::RenderLights(FRHICommandListImmediate& RHICm
 								DenoiserRequirements = DenoiserToUse->GetShadowRequirements(View, LightSceneInfo, RayTracingConfig);
 							}
 
-							FHairStrandsTextures HairTextures = GetHairStrandsTextures(ViewIndex, GraphBuilder);
+							FHairStrandsOcclusionResources HairResources = GetHairStrandsResources(ViewIndex, GraphBuilder);
 							RenderRayTracingShadows(
 								GraphBuilder,
 								SceneTextures,
@@ -1612,8 +1608,7 @@ void FDeferredShadingSceneRenderer::RenderLights(FRHICommandListImmediate& RHICm
 								RayTracingConfig,
 								DenoiserRequirements,
 								false,
-								HairTextures.CategorizationTexture,
-								HairTextures.LightChannelMaskTexture,
+								&HairResources,
 								RayTracingShadowMaskUAV,
 								RayHitDistanceUAV);
 
@@ -1671,8 +1666,7 @@ void FDeferredShadingSceneRenderer::RenderLights(FRHICommandListImmediate& RHICm
 									RayTracingConfig,
 									IScreenSpaceDenoiser::EShadowRequirements::Bailout,
 									true,
-									HairTextures.CategorizationTexture,
-									HairTextures.LightChannelMaskTexture,
+									&HairResources,
 									SubPixelRayTracingShadowMaskUAV,
 									SubPixelRayHitDistanceUAV);
 
diff --git a/Engine/Source/Runtime/Renderer/Private/RayTracing/RayTracingShadows.cpp b/Engine/Source/Runtime/Renderer/Private/RayTracing/RayTracingShadows.cpp
--- a/Engine/Source/Runtime/Renderer/Private/RayTracing/RayTracingShadows.cpp
+++ b/Engine/Source/Runtime/Renderer/Private/RayTracing/RayTracingShadows.cpp
@@ -85,12 +85,14 @@ class FOcclusionRGS : public FGlobalShader
 	}
 
 	BEGIN_SHADER_PARAMETER_STRUCT(FParameters, )
+		SHADER_PARAMETER_STRUCT_INCLUDE(ShaderDrawDebug::FShaderDrawDebugParameters, ShaderDrawParameters)
 		SHADER_PARAMETER(uint32, SamplesPerPixel)
 		SHADER_PARAMETER(float, NormalBias)
 		SHADER_PARAMETER(uint32, LightingChannelMask)
 		SHADER_PARAMETER(uint32, ShadowMaskType)
 		SHADER_PARAMETER(FIntRect, LightScissor)
 		SHADER_PARAMETER(FIntPoint, PixelOffset)
+		SHADER_PARAMETER(uint32, bUseHairVoxel)
 
 		SHADER_PARAMETER_STRUCT(FLightShaderParameters, Light)
 		SHADER_PARAMETER_STRUCT_INCLUDE(FSceneTextureParameters, SceneTextures)
@@ -101,6 +103,7 @@ class FOcclusionRGS : public FGlobalShader
 		SHADER_PARAMETER_RDG_TEXTURE_UAV(RWTexture2D<float4>, RWOcclusionMaskUAV)
 		SHADER_PARAMETER_RDG_TEXTURE_UAV(RWTexture2D<float>, RWRayDistanceUAV)
 		SHADER_PARAMETER_STRUCT_REF(FViewUniformShaderParameters, ViewUniformBuffer)
+		SHADER_PARAMETER_STRUCT_REF(FVirtualVoxelParameters, VirtualVoxel)
 	END_SHADER_PARAMETER_STRUCT()
 };
 
@@ -155,8 +158,7 @@ void FDeferredShadingSceneRenderer::RenderRayTracingShadows(
 	const IScreenSpaceDenoiser::FShadowRayTracingConfig& RayTracingConfig,
 	const IScreenSpaceDenoiser::EShadowRequirements DenoiserRequirements,
 	const bool bSubPixelShadowMask,
-	FRDGTextureRef HairCategorizationTexture,
-	FRDGTextureRef HairLightChannelMaskTexture,
+	const FHairStrandsOcclusionResources* HairResources,
 	FRDGTextureUAV* OutShadowMaskUAV,
 	FRDGTextureUAV* OutRayHitDistanceUAV)
 #if RHI_RAYTRACING
@@ -216,7 +218,12 @@ void FDeferredShadingSceneRenderer::RenderRayTracingShadows(
 
 	// Ray generation pass for shadow occlusion.
 	{
-		const bool bUseHairLighting = HairCategorizationTexture != nullptr;
+		const bool bUseHairLighting =
+			HairResources != nullptr &&
+			HairResources->CategorizationTexture != nullptr &&
+			HairResources->LightChannelMaskTexture != nullptr &&
+			HairResources->VoxelResources != nullptr;
+
 		FOcclusionRGS::FParameters* PassParameters = GraphBuilder.AllocParameters<FOcclusionRGS::FParameters>();
 		PassParameters->RWOcclusionMaskUAV = OutShadowMaskUAV;
 		PassParameters->RWRayDistanceUAV = OutRayHitDistanceUAV;
@@ -229,16 +236,18 @@ void FDeferredShadingSceneRenderer::RenderRayTracingShadows(
 		PassParameters->SceneTextures = SceneTextures;
 		PassParameters->LightScissor = ScissorRect;
 		PassParameters->PixelOffset = PixelOffset;
-		PassParameters->ShadowMaskType = bUseHairLighting ? (bSubPixelShadowMask ? 1 : 0) : 0;
-		PassParameters->HairCategorizationTexture = HairCategorizationTexture;
-		if (!PassParameters->HairCategorizationTexture)
-		{
-			PassParameters->HairCategorizationTexture = GraphBuilder.RegisterExternalTexture(GSystemTextures.BlackDummy);
-		}
-		PassParameters->HairLightChannelMaskTexture = HairLightChannelMaskTexture;
-		if (!PassParameters->HairLightChannelMaskTexture)
+		PassParameters->ShadowMaskType = bUseHairLighting && bSubPixelShadowMask ? 1 : 0;
+		if (bUseHairLighting)
 		{
-			PassParameters->HairLightChannelMaskTexture = GraphBuilder.RegisterExternalTexture(GSystemTextures.BlackDummy);
+			PassParameters->bUseHairVoxel = IsHairRayTracingEnabled() ? 0 : 1;
+			PassParameters->HairCategorizationTexture = HairResources->CategorizationTexture;
+			PassParameters->HairLightChannelMaskTexture = HairResources->LightChannelMaskTexture;
+			PassParameters->VirtualVoxel = HairResources->VoxelResources->UniformBuffer;
+
+			if (ShaderDrawDebug::IsShaderDrawDebugEnabled(View))
+			{
+				ShaderDrawDebug::SetParameters(GraphBuilder, View.ShaderDrawData, PassParameters->ShaderDrawParameters);
+			}
 		}
 		FOcclusionRGS::FPermutationDomain PermutationVector;
 		PermutationVector.Set<FOcclusionRGS::FLightTypeDim>(LightSceneProxy->GetLightType());
diff --git a/Engine/Source/Runtime/Renderer/Private/SceneRendering.cpp b/Engine/Source/Runtime/Renderer/Private/SceneRendering.cpp
--- a/Engine/Source/Runtime/Renderer/Private/SceneRendering.cpp
+++ b/Engine/Source/Runtime/Renderer/Private/SceneRendering.cpp
@@ -1668,7 +1668,7 @@ void FViewInfo::SetupUniformBufferParameters(
 	// Deep opacity maps info
 	{
 		const bool bEnableMSAA = true;
-		SetUpViewHairRenderInfo(*this, bEnableMSAA, ViewUniformShaderParameters.HairRenderInfo);
+		SetUpViewHairRenderInfo(*this, bEnableMSAA, ViewUniformShaderParameters.HairRenderInfo, ViewUniformShaderParameters.HairRenderInfoBits);
 	}
 
 	ViewUniformShaderParameters.VTFeedbackBuffer = SceneContext.GetVirtualTextureFeedbackUAV();
diff --git a/Engine/Source/Runtime/Renderer/Private/SceneVisibility.cpp b/Engine/Source/Runtime/Renderer/Private/SceneVisibility.cpp
--- a/Engine/Source/Runtime/Renderer/Private/SceneVisibility.cpp
+++ b/Engine/Source/Runtime/Renderer/Private/SceneVisibility.cpp
@@ -3060,7 +3060,7 @@ void FSceneRenderer::PreVisibilityFrameSetup(FRHICommandListImmediate& RHICmdLis
 		const EWorldType::Type WorldType = Views[0].Family->Scene->GetWorld()->WorldType;
 		const FShaderDrawDebugData* ShaderDrawData = &Views[0].ShaderDrawData;
 		auto ShaderMap = GetGlobalShaderMap(FeatureLevel);
-		RunHairStrandsInterpolation(RHICmdList, WorldType, ShaderDrawData, ShaderMap, EHairStrandsInterpolationType::SimulationStrands, nullptr);
+		RunHairStrandsInterpolation(RHICmdList, WorldType, Scene->GetGPUSkinCache(), ShaderDrawData, ShaderMap, EHairStrandsInterpolationType::SimulationStrands, nullptr);
 	}
 
 	// Notify the FX system that the scene is about to perform visibility checks.
diff --git a/Engine/Source/Runtime/Renderer/Public/HairStrandsInterface.h b/Engine/Source/Runtime/Renderer/Public/HairStrandsInterface.h
--- a/Engine/Source/Runtime/Renderer/Public/HairStrandsInterface.h
+++ b/Engine/Source/Runtime/Renderer/Public/HairStrandsInterface.h
@@ -12,6 +12,7 @@
 #include "Engine/EngineTypes.h"
 #include "Shader.h"
 #include "RenderResource.h"
+#include "RenderGraphResources.h"
 
 ////////////////////////////////////////////////////////////////////////////////////////////////////
 // Misc/Helpers
@@ -36,13 +37,11 @@ enum class EHairStrandsDebugMode : uint8
 /// Return the active debug view mode
 RENDERER_API EHairStrandsDebugMode GetHairStrandsDebugStrandsMode();
 
-/// Return the number of sample subsample count used for the visibility pass
-RENDERER_API uint32 GetHairVisibilitySampleCount();
-
 struct FMinHairRadiusAtDepth1
 {
 	float Primary = 1;
 	float Velocity = 1;
+	float Stable = 1;
 };
 
 /// Compute the strand radius at a distance of 1 meter
@@ -86,6 +85,18 @@ public:
 	bool GetCullingResultAvailable() const { return bCullingResultAvailable; }
 	void SetCullingResultAvailable(bool b) { bCullingResultAvailable = b; }
 
+	struct VertexFactoryInput 
+	{
+		FShaderResourceViewRHIRef HairPositionBuffer = nullptr;
+		FVector HairPositionOffset = FVector::ZeroVector;
+		uint32 VertexCount = 0;
+		float HairRadius = 0;
+		float HairLength = 0;
+		float HairDensity = 0;
+		FTransform LocalToWorldTransform;
+	};
+	VertexFactoryInput VFInput;
+
 private:
 
 	uint32 GroupInstanceVertexCount;
@@ -155,6 +166,7 @@ typedef void(*THairStrandsResetInterpolationFunction)(
 typedef void (*THairStrandsInterpolationFunction)(
 	FRHICommandListImmediate& RHICmdList, 
 	const struct FShaderDrawDebugData* ShaderDrawData,
+	const FTransform& LocalToWorld,
 	struct FHairStrandsInterpolationInput* Input, 
 	struct FHairStrandsInterpolationOutput* Output, 
 	struct FHairStrandsProjectionHairData& RenHairProjection,
@@ -221,6 +233,15 @@ struct FHairStrandsProjectionHairData
 		FRWBuffer* DeformedRootTrianglePosition0Buffer = nullptr;
 		FRWBuffer* DeformedRootTrianglePosition1Buffer = nullptr;
 		FRWBuffer* DeformedRootTrianglePosition2Buffer = nullptr;
+
+		// Samples to be used for rbf mesh interpolatrion
+		uint32 SampleCount = 0;
+		FRWBuffer* MeshInterpolationWeightsBuffer = nullptr;
+		FRWBuffer* MeshSampleIndicesBuffer = nullptr;
+
+		FRWBuffer* RestSamplePositionsBuffer = nullptr;
+		FRWBuffer* DeformedSamplePositionsBuffer = nullptr;
+		FRWBuffer* MeshSampleWeightsBuffer = nullptr;
 	};
 
 	struct HairGroup
@@ -251,14 +272,14 @@ struct FHairStrandsPrimitiveResources
 	TArray<FHairGroup> Groups;
 };
 
-FHairStrandsPrimitiveResources GetHairStandsPrimitiveResources(uint64 Id);
+FHairStrandsPrimitiveResources GetHairStandsPrimitiveResources(uint32 ComponentId);
 
 ////////////////////////////////////////////////////////////////////////////////////////////////////
 // Debug infos
 
 struct FHairStrandsDebugInfo
 {
-	uint64 Id = 0;
+	uint32 ComponentId = 0;
 	EWorldType::Type WorldType = EWorldType::None;
 
 	struct HairGroup
@@ -269,6 +290,8 @@ struct FHairStrandsDebugInfo
 		uint32 CurveCount = 0;
 
 		bool bHasSkinInterpolation = false;
+		bool bHasBinding = false;
+		bool bHasSimulation = false;
 		uint32 LODCount = 0;
 	};
 	TArray<HairGroup> HairGroups;
@@ -300,7 +323,8 @@ struct FHairStrandsProjectionDebugInfo
 // Registrations
 
 RENDERER_API void RegisterHairStrands(
-	uint64 Id,
+	uint32 ComponentId,
+	uint32 SkeletalComponentId,
 	EWorldType::Type WorldType,
 	const FHairStrandsInterpolationData& E,
 	const FHairStrandsProjectionHairData& RenProjection,
@@ -309,29 +333,30 @@ RENDERER_API void RegisterHairStrands(
 	const FHairStrandsDebugInfo& DebugInfo,
 	const FHairStrandsProjectionDebugInfo& DebugProjectionInfo);
 
-RENDERER_API void UnregisterHairStrands(uint64 Id);
+RENDERER_API void UnregisterHairStrands(uint32 ComponentId);
 
 RENDERER_API bool UpdateHairStrands(
-	uint64 Id,
+	uint32 ComponentId,
 	EWorldType::Type NewWorldType);
 
 RENDERER_API bool UpdateHairStrands(
-	uint64 Id,
+	uint32 ComponentId,
 	EWorldType::Type WorldType,
 	const FTransform& HairLocalToWorld,
 	const FHairStrandsProjectionHairData& RenProjection,
 	const FHairStrandsProjectionHairData& SimProjection);
 
 RENDERER_API bool UpdateHairStrands(
-	uint64 Id,
+	uint32 ComponentId,
 	EWorldType::Type WorldType,
 	const FTransform& HairLocalToWorld,
 	const FTransform& MeshLocalToWorld);
 
-RENDERER_API bool UpdateHairStrands(
-	uint64 Id, 
-	EWorldType::Type WorldType, 
-	const class FSkeletalMeshObject* MeshObject);
+RENDERER_API bool UpdateHairStrandsDebugInfo(
+	uint32 ComponentId,
+	EWorldType::Type WorldType,
+	const uint32 GroupIt,
+	const bool bSimulationEnable);
 
 RENDERER_API bool IsHairStrandsSupported(const EShaderPlatform Platform);
 bool IsHairStrandsEnable(EShaderPlatform Platform);
@@ -369,11 +394,23 @@ RENDERER_API void RunProjection(
 	FHairStrandsProjectionHairData& RenProjectionHairData,
 	FHairStrandsProjectionHairData& SimProjectionHairData);
 
+RENDERER_API FHairStrandsProjectionMeshData ExtractMeshData(class FSkeletalMeshRenderData* RenderData);
+FHairStrandsProjectionHairData::HairGroup ToProjectionHairData(struct FHairStrandsRootResource* In);
+
 typedef void (*TBindingProcess)(FRHICommandListImmediate& RHICmdList, void* Asset);
 RENDERER_API void EnqueueGroomBindingQuery(void* Asset, TBindingProcess BindingProcess);
-void RunHairStrandsBindingQueries(FRHICommandListImmediate& RHICmdList, class FGlobalShaderMap* ShaderMap);
-bool HasHairStrandsProjectionQuery(EShaderPlatform Platform);
 
-RENDERER_API FHairStrandsProjectionMeshData ExtractMeshData(class FSkeletalMeshRenderData* RenderData);
 
-FHairStrandsProjectionHairData::HairGroup ToProjectionHairData(struct FHairStrandsRootResource* In);
\ No newline at end of file
+struct FFollicleInfo
+{
+	enum EChannel {R = 0, G = 1, B = 2, A = 3};
+
+	uint32 GroomId = ~0;
+	EChannel Channel = R;
+	uint32 KernelSizeInPixels = 0;
+};
+
+RENDERER_API void EnqueueFollicleMaskUpdateQuery(const TArray<FFollicleInfo>& Infos, class UTexture2D* OutTexture);
+
+void RunHairStrandsProcess(FRHICommandListImmediate& RHICmdList, class FGlobalShaderMap* ShaderMap); 
+bool HasHairStrandsProcess(EShaderPlatform Platform);
