Move mobileDepthPrepass Uniform buffer update to the start of the frame

diff --git a/Engine/Source/Runtime/Renderer/Private/DepthRendering.cpp b/Engine/Source/Runtime/Renderer/Private/DepthRendering.cpp
--- a/Engine/Source/Runtime/Renderer/Private/DepthRendering.cpp
+++ b/Engine/Source/Runtime/Renderer/Private/DepthRendering.cpp
@@ -788,24 +788,22 @@ void FMobileSceneRenderer::RenderPrePass(FRHICommandListImmediate& RHICmdList)
 
 			SCOPED_GPU_MASK(RHICmdList, !View.IsInstancedStereoPass() ? View.GPUMask : (Views[0].GPUMask | Views[1].GPUMask));
 			SCOPED_CONDITIONAL_DRAW_EVENTF(RHICmdList, EventView, Views.Num() > 1, TEXT("View%d"), ViewIndex);
+			if (!View.ShouldRenderView())
+			{
+				continue;
+			}
 
-			TUniformBufferRef<FSceneTexturesUniformParameters> PassUniformBuffer;
-			CreateDepthPassUniformBuffer(RHICmdList, View, PassUniformBuffer);
-
-			FMeshPassProcessorRenderState DrawRenderState(View, PassUniformBuffer);
+			if (Scene->UniformBuffers.UpdateViewUniformBuffer(View))
+			{
+				UpdateDepthPrepassUniformBuffer(RHICmdList, View);
+			}
 
+			FMeshPassProcessorRenderState DrawRenderState(View, Scene->UniformBuffers.DepthPassUniformBuffer);
 			SetupDepthPassState(DrawRenderState);
 
-			if (View.ShouldRenderView())
-			{
-				Scene->UniformBuffers.UpdateViewUniformBuffer(View);
+			SetupPrePassView(RHICmdList, View, this);
 
-				{
-					SetupPrePassView(RHICmdList, View, this);
-
-					View.ParallelMeshDrawCommandPasses[EMeshPass::DepthPass].DispatchDraw(nullptr, RHICmdList);
-				}
-			}
+			View.ParallelMeshDrawCommandPasses[EMeshPass::DepthPass].DispatchDraw(nullptr, RHICmdList);
 		}
 	}
 }
diff --git a/Engine/Source/Runtime/Renderer/Private/MobileShadingRenderer.cpp b/Engine/Source/Runtime/Renderer/Private/MobileShadingRenderer.cpp
--- a/Engine/Source/Runtime/Renderer/Private/MobileShadingRenderer.cpp
+++ b/Engine/Source/Runtime/Renderer/Private/MobileShadingRenderer.cpp
@@ -352,6 +352,7 @@ void FMobileSceneRenderer::InitViews(FRHICommandListImmediate& RHICmdList)
 		const FViewInfo& View = Views[0];
 		// We want to wait for the extension jobs only when the view is being actually rendered for the first time
 		Scene->UniformBuffers.UpdateViewUniformBuffer(View, false);
+		UpdateDepthPrepassUniformBuffer(RHICmdList, View);
 		UpdateOpaqueBasePassUniformBuffer(RHICmdList, View);
 		UpdateTranslucentBasePassUniformBuffer(RHICmdList, View);
 		UpdateDirectionalLightUniformBuffers(RHICmdList, View);
@@ -1100,6 +1101,14 @@ void FMobileSceneRenderer::UpdateSkyReflectionUniformBuffer()
 	Scene->UniformBuffers.MobileSkyReflectionUniformBuffer.UpdateUniformBufferImmediate(Parameters);
 }
 
+void FMobileSceneRenderer::UpdateDepthPrepassUniformBuffer(FRHICommandListImmediate& RHICmdList, const FViewInfo& View)
+{
+	FSceneTexturesUniformParameters SceneTextureParameters;
+	FSceneRenderTargets& SceneContext = FSceneRenderTargets::Get(RHICmdList);
+	SetupSceneTextureUniformParameters(SceneContext, View.FeatureLevel, ESceneTextureSetupMode::None, SceneTextureParameters);
+	Scene->UniformBuffers.DepthPassUniformBuffer.UpdateUniformBufferImmediate(SceneTextureParameters);
+}
+
 void FMobileSceneRenderer::CreateDirectionalLightUniformBuffers(FViewInfo& View)
 {
 	bool bDynamicShadows = ViewFamily.EngineShowFlags.DynamicShadows;
diff --git a/Engine/Source/Runtime/Renderer/Private/SceneRendering.h b/Engine/Source/Runtime/Renderer/Private/SceneRendering.h
--- a/Engine/Source/Runtime/Renderer/Private/SceneRendering.h
+++ b/Engine/Source/Runtime/Renderer/Private/SceneRendering.h
@@ -1892,6 +1892,7 @@ protected:
 	void UpdateTranslucentBasePassUniformBuffer(FRHICommandListImmediate& RHICmdList, const FViewInfo& View);
 	void UpdateDirectionalLightUniformBuffers(FRHICommandListImmediate& RHICmdList, const FViewInfo& View);
 	void UpdateSkyReflectionUniformBuffer();
+	void UpdateDepthPrepassUniformBuffer(FRHICommandListImmediate& RHICmdList, const FViewInfo& View);
 	
 private:
 	bool bModulatedShadowsInUse;
