Fix Hair PPLL mode.

diff --git a/Engine/Source/Runtime/Renderer/Private/HairStrands/HairStrandsVisibility.cpp b/Engine/Source/Runtime/Renderer/Private/HairStrands/HairStrandsVisibility.cpp
--- a/Engine/Source/Runtime/Renderer/Private/HairStrands/HairStrandsVisibility.cpp
+++ b/Engine/Source/Runtime/Renderer/Private/HairStrands/HairStrandsVisibility.cpp
@@ -133,7 +133,7 @@ static uint32 GetMSAASampleCount()
 	return GHairVisibilitySampleCount <= 4 ? 4 : 8;
 }
 
-void SetUpViewHairRenderInfo(const FViewInfo& ViewInfo, bool bEnableMSAA, FVector4& OutHairRenderInfo, uint32& OutHairRenderInfoBits)
+static void SetUpViewHairRenderInfo(const FViewInfo& ViewInfo, bool bEnableMSAA, FVector4& OutHairRenderInfo, uint32& OutHairRenderInfoBits)
 {
 	FVector2D PixelVelocity(1.f / (ViewInfo.ViewRect.Width() * 2), 1.f / (ViewInfo.ViewRect.Height() * 2));
 	const float VelocityMagnitudeScale = FMath::Clamp(CVarHairVelocityMagnitudeScale.GetValueOnAnyThread(), 0, 512) * FMath::Min(PixelVelocity.X, PixelVelocity.Y);
@@ -144,11 +144,16 @@ void SetUpViewHairRenderInfo(const FViewInfo& ViewInfo, bool bEnableMSAA, FVecto
 	FMinHairRadiusAtDepth1 MinHairRadius = ComputeMinStrandRadiusAtDepth1(
 		FIntPoint(ViewInfo.UnconstrainedViewRect.Width(), ViewInfo.UnconstrainedViewRect.Height()), ViewInfo.FOV, HairVisibilitySampleCount, RasterizationScaleOverride);
 
-	TUniformBufferRef<FViewUniformShaderParameters> ViewUniformShaderParameters;
 	OutHairRenderInfo = PackHairRenderInfo(MinHairRadius.Primary, MinHairRadius.Stable, MinHairRadius.Velocity, VelocityMagnitudeScale);
 	OutHairRenderInfoBits = PackHairRenderInfoBits(!ViewInfo.IsPerspectiveProjection(), false);
 }
 
+void SetUpViewHairRenderInfo(const FViewInfo& ViewInfo, FVector4& OutHairRenderInfo, uint32& OutHairRenderInfoBits)
+{
+	const bool bMsaaEnable = GetHairVisibilityRenderMode() == HairVisibilityRenderMode_MSAA;
+	SetUpViewHairRenderInfo(ViewInfo, bMsaaEnable, OutHairRenderInfo, OutHairRenderInfoBits);
+}
+
 static bool IsCompatibleWithHairVisibility(const FMeshMaterialShaderPermutationParameters& Parameters)
 {
 	return IsCompatibleWithHairStrands(Parameters.Platform, Parameters.MaterialParameters);
@@ -1271,7 +1276,7 @@ static void AddHairVisibilityPrimitiveIdCompactionPass(
 	AddClearUAVPass(GraphBuilder, GraphBuilder.CreateUAV(OutCategorizationTexture), ClearValues);
 
 	// Select render node count according to current mode
-	const uint32 MSAASampleCount = GetMSAASampleCount(); check(MSAASampleCount == 4 || MSAASampleCount == 8);
+	const uint32 MSAASampleCount = GetHairVisibilityRenderMode() == HairVisibilityRenderMode_MSAA ? GetMSAASampleCount() : 1;
 	const uint32 PPLLMaxRenderNodePerPixel = GetPPLLMaxRenderNodePerPixel();
 	const uint32 MaxRenderNodeCount = Resolution.X * Resolution.Y * (GetHairVisibilityRenderMode() == HairVisibilityRenderMode_MSAA ? MSAASampleCount : PPLLMaxRenderNodePerPixel);
 	OutCompactNodeData = GraphBuilder.CreateBuffer(FRDGBufferDesc::CreateStructuredDesc(bUseVisibility ? sizeof(HairStrandsVisibilityInternal::NodeVis) : sizeof(HairStrandsVisibilityInternal::NodeData), MaxRenderNodeCount), TEXT("HairVisibilityPrimitiveIdCompactNodeData"));
diff --git a/Engine/Source/Runtime/Renderer/Private/HairStrands/HairStrandsVisibility.h b/Engine/Source/Runtime/Renderer/Private/HairStrands/HairStrandsVisibility.h
--- a/Engine/Source/Runtime/Renderer/Private/HairStrands/HairStrandsVisibility.h
+++ b/Engine/Source/Runtime/Renderer/Private/HairStrands/HairStrandsVisibility.h
@@ -62,5 +62,5 @@ FHairStrandsVisibilityViews RenderHairStrandsVisibilityBuffer(
 	TRefCountPtr<IPooledRenderTarget> VelocityTexture,
 	const struct FHairStrandsMacroGroupViews& MacroGroupViews);
 
-void SetUpViewHairRenderInfo(const FViewInfo& ViewInfo, bool bEnableMSAA, FVector4& OutHairRenderInfo, uint32& OutHairRenderInfoBits);
+void SetUpViewHairRenderInfo(const FViewInfo& ViewInfo, FVector4& OutHairRenderInfo, uint32& OutHairRenderInfoBits);
 
diff --git a/Engine/Source/Runtime/Renderer/Private/SceneRendering.cpp b/Engine/Source/Runtime/Renderer/Private/SceneRendering.cpp
--- a/Engine/Source/Runtime/Renderer/Private/SceneRendering.cpp
+++ b/Engine/Source/Runtime/Renderer/Private/SceneRendering.cpp
@@ -1665,11 +1665,8 @@ void FViewInfo::SetupUniformBufferParameters(
 		}
 	}
 
-	// Deep opacity maps info
-	{
-		const bool bEnableMSAA = true;
-		SetUpViewHairRenderInfo(*this, bEnableMSAA, ViewUniformShaderParameters.HairRenderInfo, ViewUniformShaderParameters.HairRenderInfoBits);
-	}
+	// Default values
+	SetUpViewHairRenderInfo(*this, ViewUniformShaderParameters.HairRenderInfo, ViewUniformShaderParameters.HairRenderInfoBits);
 
 	ViewUniformShaderParameters.VTFeedbackBuffer = SceneContext.GetVirtualTextureFeedbackUAV();
 	ViewUniformShaderParameters.QuadOverdraw = SceneContext.GetQuadOverdrawBufferUAV();
