Support LightShaft work with MSAA on IOS.

Make all postprocess passes including postprocess materials work properly with MSAA on IOS.

diff --git a/Engine/Shaders/Private/MaterialTemplate.ush b/Engine/Shaders/Private/MaterialTemplate.ush
--- a/Engine/Shaders/Private/MaterialTemplate.ush
+++ b/Engine/Shaders/Private/MaterialTemplate.ush
@@ -1636,6 +1636,12 @@ float EyeAdaptationLookup()
 }
 #endif
 
+//The post-process material needs to decode the scene color since it's encoded at PreTonemapMSAA if MSAA enabled on MetalMobilePlatorm 
+//The POST_PROCESS_MATERIAL_BEFORE_TONEMAP is 1 for both BL_BeforeTranslucency and BL_BeforeTonemapping post-process materials
+#if FEATURE_LEVEL <= FEATURE_LEVEL_ES3_1 && POST_PROCESS_MATERIAL && POST_PROCESS_MATERIAL_BEFORE_TONEMAP && METAL_PROFILE
+uint bMetalMSAAHDRDecode;
+#endif
+
 #if NEEDS_SCENE_TEXTURES
 
 #if SHADING_PATH_MOBILE
@@ -1660,8 +1666,12 @@ MaterialFloat4 MobileSceneTextureLookup(inout FMaterialPixelParameters Parameter
 		MaterialFloat4 Input0 = Texture2DSample(PostProcessInput_0_Texture, PostProcessInput_0_SharedSampler, UV);
 		#if POST_PROCESS_MATERIAL_BEFORE_TONEMAP
 			#if METAL_PROFILE
-				// If MSAA is in use undo PreTonemapMSAA, it is reapplied at the end of post process material
-				Input0.rgb = lerp(Input0.rgb, Input0.rgb * rcp(Input0.r*(-0.299) + Input0.g*(-0.587) + Input0.b*(-0.114) + 1.0), ResolvedView.IsMobileMSAA);
+				// Decode the input color since the color is encoded for MSAA 
+				// The decode instructions might be able to skip with dynamic branch
+				if (bMetalMSAAHDRDecode)
+				{
+					Input0.rgb = Input0.rgb * rcp(Input0.r*(-0.299) + Input0.g*(-0.587) + Input0.b*(-0.114) + 1.0);
+				}
 			#endif
 		#endif
 		// We need to preserve original SceneColor Alpha as it's used by tonemapper on mobile
diff --git a/Engine/Shaders/Private/PostProcessMaterialShaders.usf b/Engine/Shaders/Private/PostProcessMaterialShaders.usf
--- a/Engine/Shaders/Private/PostProcessMaterialShaders.usf
+++ b/Engine/Shaders/Private/PostProcessMaterialShaders.usf
@@ -381,9 +381,5 @@ void MainPS(
 #endif
 #endif
 	OutColor = FullSceneColor;
-#if POST_PROCESS_MATERIAL_BEFORE_TONEMAP && METAL_PROFILE
-	// If MSAA is in use we reapply the on-chip pre-tonemap since the tonemapper will be expecting it.
-	OutColor.rgb = lerp(OutColor.rgb, OutColor.rgb* rcp(OutColor.r*0.299 + OutColor.g*0.587 + OutColor.b*0.114 + 1.0), ResolvedView.IsMobileMSAA);
-#endif
 }
 #endif //(FEATURE_LEVEL > FEATURE_LEVEL_ES3_1)
diff --git a/Engine/Shaders/Private/PostProcessMobile.usf b/Engine/Shaders/Private/PostProcessMobile.usf
--- a/Engine/Shaders/Private/PostProcessMobile.usf
+++ b/Engine/Shaders/Private/PostProcessMobile.usf
@@ -109,10 +109,14 @@ void SunMaskPS_ES2(
 		half InDepth = CalcSceneDepth(InUVPos.xy);
 	#else
 		half InDepth = SceneColor.a;
-		OutColor = SceneColor;
 	#endif
 	
 	#if ES2_USE_SUN
+
+		#if METAL_MSAA_HDR_DECODE
+			SceneColor.rgb *= rcp(SceneColor.r*(-0.299) + SceneColor.g*(-0.587) + SceneColor.b*(-0.114) + 1.0);
+		#endif
+
 		half2 DepthMaskScaleBias = SunConstDepthMaskScaleBias();
 		half FarAmount = saturate(InDepth * DepthMaskScaleBias.x + DepthMaskScaleBias.y);
 		half3 SunAmount = SceneColor.rgb * SunColorApertureDiv2.rgb;
@@ -131,6 +135,10 @@ void SunMaskPS_ES2(
 	#if ES2_USE_DOF
 		OutSunShaftAndDof += Coc2(InDepth);
 	#endif
+
+	#if !ES2_USE_DEPTHTEXTURE
+		OutColor = SceneColor;
+	#endif
 }
 
 //
@@ -149,6 +157,21 @@ void PreTonemapMSAA_ES2(
 #endif
 }
 
+Texture2D InputTexture;
+SamplerState InputSampler;
+
+void MSAADecodeAndCopyRectPS(
+	noperspective float4 UVAndScreenPos : TEXCOORD0,
+	out half4 OutColor : SV_Target0
+	)
+{
+#if (METAL_PROFILE && !MAC)
+	float2 UV = UVAndScreenPos.xy;
+	OutColor = Texture2DSample(InputTexture, InputSampler, UV);
+	OutColor.rgb *= rcp(OutColor.r*(-0.299) + OutColor.g*(-0.587) + OutColor.b*(-0.114) + 1.0);
+#endif
+}
+
 // 
 // Bloom Setup - Mask Bloom and Downsample 1/16 Area
 //
@@ -228,7 +251,7 @@ void BloomPS_ES2(
 	// Output color is average.
 	AverageColor.rgb = (C0.rgb * 0.25) + (C1.rgb * 0.25) + (C2.rgb * 0.25) + (C3.rgb * 0.25);
 
-	#if ES2_USE_MSAA
+	#if METAL_MSAA_HDR_DECODE
 		// This should really happen before the average, instead doing after average as optimization.
 		AverageColor.rgb *= rcp(AverageColor.r*(-0.299) + AverageColor.g*(-0.587) + AverageColor.b*(-0.114) + 1.0);
 	#endif
@@ -1119,9 +1142,6 @@ void AaPS_ES2(
 
 Buffer<float4> EyeAdaptationBuffer;
 
-Texture2D InputTexture;
-SamplerState InputSampler;
-
 static const float FLOAT_PRECISION = 1e+5;
 
 static const float INV_FLOAT_PRECISION = 1e-5;
diff --git a/Engine/Shaders/Private/PostProcessTonemap.usf b/Engine/Shaders/Private/PostProcessTonemap.usf
--- a/Engine/Shaders/Private/PostProcessTonemap.usf
+++ b/Engine/Shaders/Private/PostProcessTonemap.usf
@@ -516,7 +516,7 @@ float4 TonemapCommonPS(
 
 	#else 
 
-		#if USE_MSAA
+		#if METAL_MSAA_HDR_DECODE
 			// look for PreTonemapMSAA 
 			SceneColor.rgb *= rcp(SceneColor.r*(-0.299) + SceneColor.g*(-0.587) + SceneColor.b*(-0.114) + 1.0);
 			// Try to kill negatives and NaNs here
@@ -891,7 +891,7 @@ void MainPS_ES2(
 			#endif
 		#endif
 
-		#if USE_MSAA
+		#if METAL_MSAA_HDR_DECODE
 			// Do after jitter for grain as an optimization.
 			SceneColor.rgb *= rcp(SceneColor.r*(-0.299) + SceneColor.g*(-0.587) + SceneColor.b*(-0.114) + 1.0);
 			// Try to kill negatives and NaNs here
diff --git a/Engine/Source/Runtime/Engine/Public/SceneView.h b/Engine/Source/Runtime/Engine/Public/SceneView.h
--- a/Engine/Source/Runtime/Engine/Public/SceneView.h
+++ b/Engine/Source/Runtime/Engine/Public/SceneView.h
@@ -697,7 +697,6 @@ enum ETranslucencyVolumeCascade
 	VIEW_UNIFORM_BUFFER_MEMBER(FLinearColor, SkyLightColor) \
 	VIEW_UNIFORM_BUFFER_MEMBER_ARRAY(FVector4, SkyIrradianceEnvironmentMap, [7]) \
 	VIEW_UNIFORM_BUFFER_MEMBER(float, MobilePreviewMode) \
-	VIEW_UNIFORM_BUFFER_MEMBER(float, IsMobileMSAA) \
 	VIEW_UNIFORM_BUFFER_MEMBER(float, HMDEyePaddingOffset) \
 	VIEW_UNIFORM_BUFFER_MEMBER_EX(float, ReflectionCubemapMaxMip, EShaderPrecisionModifier::Half) \
 	VIEW_UNIFORM_BUFFER_MEMBER(float, ShowDecalsMask) \
diff --git a/Engine/Source/Runtime/Renderer/Private/PostProcess/PostProcessMaterial.cpp b/Engine/Source/Runtime/Renderer/Private/PostProcess/PostProcessMaterial.cpp
--- a/Engine/Source/Runtime/Renderer/Private/PostProcess/PostProcessMaterial.cpp
+++ b/Engine/Source/Runtime/Renderer/Private/PostProcess/PostProcessMaterial.cpp
@@ -333,11 +333,31 @@ static void AddCopyAndFlipTexturePass(FRDGBuilder& GraphBuilder, const FViewInfo
 	AddDrawScreenPass(GraphBuilder, RDG_EVENT_NAME("DrawTexture"), View, OutputViewport, InputViewport, PixelShader, Parameters, EScreenPassDrawFlags::FlipYAxis);
 }
 
+void MobileMSAADecodeAndDrawTexturePass(
+	FRDGBuilder& GraphBuilder,
+	const FViewInfo& View,
+	FScreenPassTexture Input,
+	FScreenPassRenderTarget Output)
+{
+	const FScreenPassTextureViewport InputViewport(Input);
+	const FScreenPassTextureViewport OutputViewport(Output);
+
+	TShaderMapRef<FMSAADecodeAndCopyRectPS_ES2> PixelShader(View.ShaderMap);
+
+	FMSAADecodeAndCopyRectPS_ES2::FParameters* Parameters = GraphBuilder.AllocParameters<FMSAADecodeAndCopyRectPS_ES2::FParameters>();
+	Parameters->InputTexture = Input.Texture;
+	Parameters->InputSampler = TStaticSamplerState<>::GetRHI();
+	Parameters->RenderTargets[0] = Output.GetRenderTargetBinding();
+
+	AddDrawScreenPass(GraphBuilder, RDG_EVENT_NAME("MobileMSAADecodeAndDrawTexture"), View, OutputViewport, InputViewport, PixelShader, Parameters);
+}
+
 FScreenPassTexture AddPostProcessMaterialPass(
 	FRDGBuilder& GraphBuilder,
 	const FViewInfo& View,
 	const FPostProcessMaterialInputs& Inputs,
-	const UMaterialInterface* MaterialInterface)
+	const UMaterialInterface* MaterialInterface,
+	const bool bMetalMSAAHDRDecode)
 {
 	Inputs.Validate();
 
@@ -388,9 +408,12 @@ FScreenPassTexture AddPostProcessMaterialPass(
 	// The other case when we must render to an intermediate target is when we have to flip the image vertically because we're the last postprocess pass on mobile OpenGL.
 	// We can't simply output a flipped image, because the parts of the input image which show through the stencil mask or are blended in must also be flipped. In that case,
 	// we render normally to the intermediate target and flip the image when we copy to the output target.
+	//
+	// We need to decode the target color for blending material, force it rendering to an intermediate render target and decode the color.
 	const bool bForceIntermediateRT =
 		(DepthStencilTexture != nullptr && !GRHISupportsBackBufferWithCustomDepthStencil && Inputs.OverrideOutput.IsValid()) ||
-		(bIsCompositeWithInput && Inputs.bFlipYAxis);
+		(bIsCompositeWithInput && Inputs.bFlipYAxis) ||
+		(bMetalMSAAHDRDecode && bIsCompositeWithInput);
 
 	FScreenPassRenderTarget Output = Inputs.OverrideOutput;
 
@@ -421,7 +444,14 @@ FScreenPassTexture AddPostProcessMaterialPass(
 		if (bPrimeOutputColor || bForceIntermediateRT)
 		{
 			// Copy existing contents to new output and use load-action to preserve untouched pixels.
-			AddDrawTexturePass(GraphBuilder, View, SceneColor, Output);
+			if (bMetalMSAAHDRDecode)
+			{
+				MobileMSAADecodeAndDrawTexturePass(GraphBuilder, View, SceneColor, Output);
+			}
+			else
+			{
+				AddDrawTexturePass(GraphBuilder, View, SceneColor, Output);
+			}
 			Output.LoadAction = ERenderTargetLoadAction::ELoad;
 		}
 	}
@@ -440,6 +470,8 @@ FScreenPassTexture AddPostProcessMaterialPass(
 	PostProcessMaterialParameters->MobileCustomStencilTextureSampler = TStaticSamplerState<SF_Point, AM_Clamp, AM_Clamp, AM_Clamp>::GetRHI();
 	PostProcessMaterialParameters->MobileStencilValueRef = MaterialStencilRef;
 	PostProcessMaterialParameters->RenderTargets[0] = Output.GetRenderTargetBinding();
+	// The target color will be decoded if bForceIntermediateRT is true in any case, but we might still need to decode the input color
+	PostProcessMaterialParameters->bMetalMSAAHDRDecode = bMetalMSAAHDRDecode ? 1 : 0;
 
 	bool bMobilePlatform = IsMobilePlatform(View.GetShaderPlatform());
 
@@ -550,7 +582,7 @@ FScreenPassTexture AddPostProcessMaterialPass(
 		}
 	});
 
-	if (bForceIntermediateRT)
+	if (bForceIntermediateRT && !(bMetalMSAAHDRDecode && bIsCompositeWithInput))
 	{
 		if (!Inputs.bFlipYAxis)
 		{
@@ -680,6 +712,7 @@ FScreenPassTexture AddPostProcessMaterialChain(
 FRenderingCompositePass* AddPostProcessMaterialPass(
 	const FPostprocessContext& PostProcessContext,
 	const UMaterialInterface* MaterialInterface,
+	const bool bMetalMSAAHDRDecode,
 	EPixelFormat OverrideOutputFormat)
 {
 	const FMaterial* Material = nullptr;
@@ -697,7 +730,7 @@ FRenderingCompositePass* AddPostProcessMaterialPass(
 	}
 
 	return PostProcessContext.Graph.RegisterPass(new(FMemStack::Get()) TRCPassForRDG<kPostProcessMaterialInputCountMax, 1>(
-		[MaterialInterface, Material, OverrideOutputFormat](FRenderingCompositePass* Pass, FRenderingCompositePassContext& InContext)
+		[MaterialInterface, Material, OverrideOutputFormat, bMetalMSAAHDRDecode](FRenderingCompositePass* Pass, FRenderingCompositePassContext& InContext)
 	{
 		FRDGBuilder GraphBuilder(InContext.RHICmdList);
 
@@ -736,7 +769,7 @@ FRenderingCompositePass* AddPostProcessMaterialPass(
 			Inputs.CustomDepthTexture = GraphBuilder.RegisterExternalTexture(CustomDepthTarget, TEXT("CustomDepth"));
 		}
 
-		FScreenPassTexture Outputs = AddPostProcessMaterialPass(GraphBuilder, InContext.View, Inputs, MaterialInterface);
+		FScreenPassTexture Outputs = AddPostProcessMaterialPass(GraphBuilder, InContext.View, Inputs, MaterialInterface, bMetalMSAAHDRDecode);
 
 		Pass->ExtractRDGTextureForOutput(GraphBuilder, ePId_Output0, Outputs.Texture);
 
@@ -752,6 +785,7 @@ FRenderingCompositePass* AddPostProcessMaterialPass(
 FRenderingCompositeOutputRef AddPostProcessMaterialChain(
 	FPostprocessContext& Context,
 	EBlendableLocation Location,
+	bool& bMetalMSAAHDRDecode,
 	FRenderingCompositeOutputRef SeparateTranslucency,
 	FRenderingCompositeOutputRef PreTonemapHDRColor,
 	FRenderingCompositeOutputRef PostTonemapHDRColor,
@@ -765,7 +799,11 @@ FRenderingCompositeOutputRef AddPostProcessMaterialChain(
 
 	for (const UMaterialInterface* MaterialInterface : MaterialChain)
 	{
-		FRenderingCompositePass* Pass = AddPostProcessMaterialPass(Context, MaterialInterface);
+		FRenderingCompositePass* Pass = AddPostProcessMaterialPass(Context, MaterialInterface, bMetalMSAAHDRDecode);
+
+		// For solid material, we decode the input color and output the linear color
+		// For blend material, we force it rendering to an intermediate render target and decode there
+		bMetalMSAAHDRDecode = false;
 
 		Pass->SetInput(EPassInputId(EPostProcessMaterialInput::SceneColor), LastOutput);
 		Pass->SetInput(EPassInputId(EPostProcessMaterialInput::SeparateTranslucency), SeparateTranslucency);
diff --git a/Engine/Source/Runtime/Renderer/Private/PostProcess/PostProcessMaterial.h b/Engine/Source/Runtime/Renderer/Private/PostProcess/PostProcessMaterial.h
--- a/Engine/Source/Runtime/Renderer/Private/PostProcess/PostProcessMaterial.h
+++ b/Engine/Source/Runtime/Renderer/Private/PostProcess/PostProcessMaterial.h
@@ -44,6 +44,7 @@ BEGIN_SHADER_PARAMETER_STRUCT(FPostProcessMaterialParameters, )
 	SHADER_PARAMETER_SAMPLER(SamplerState, MobileCustomStencilTextureSampler)
 	SHADER_PARAMETER(int32, MobileStencilValueRef)
 	SHADER_PARAMETER(uint32, bFlipYAxis)
+	SHADER_PARAMETER(uint32, bMetalMSAAHDRDecode)
 	RENDER_TARGET_BINDING_SLOTS()
 END_SHADER_PARAMETER_STRUCT()
 
@@ -111,7 +112,9 @@ FScreenPassTexture AddPostProcessMaterialPass(
 	FRDGBuilder& GraphBuilder,
 	const FViewInfo& View,
 	const FPostProcessMaterialInputs& Inputs,
-	const UMaterialInterface* MaterialInterface);
+	const UMaterialInterface* MaterialInterface,
+	//bMetalMSAAHDRDecode has to be set only on mobile platform, the other platforms could use the default value.
+	const bool bMetalMSAAHDRDecode = false);
 
 FScreenPassTexture AddPostProcessMaterialChain(
 	FRDGBuilder& GraphBuilder,
@@ -144,11 +147,13 @@ FScreenPassTexture AddHighResolutionScreenshotMaskPass(
 FRenderingCompositePass* AddPostProcessMaterialPass(
 	const FPostprocessContext& Context,
 	const UMaterialInterface* MaterialInterface,
+	const bool bMetalMSAAHDRDecode,
 	EPixelFormat OutputFormat = PF_Unknown);
 
 FRenderingCompositeOutputRef AddPostProcessMaterialChain(
 	FPostprocessContext& Context,
 	EBlendableLocation InLocation,
+	bool& bMetalMSAAHDRDecode,
 	FRenderingCompositeOutputRef SeparateTranslucency = FRenderingCompositeOutputRef(),
 	FRenderingCompositeOutputRef PreTonemapHDRColor = FRenderingCompositeOutputRef(),
 	FRenderingCompositeOutputRef PostTonemapHDRColor = FRenderingCompositeOutputRef(),
diff --git a/Engine/Source/Runtime/Renderer/Private/PostProcess/PostProcessMobile.cpp b/Engine/Source/Runtime/Renderer/Private/PostProcess/PostProcessMobile.cpp
--- a/Engine/Source/Runtime/Renderer/Private/PostProcess/PostProcessMobile.cpp
+++ b/Engine/Source/Runtime/Renderer/Private/PostProcess/PostProcessMobile.cpp
@@ -19,7 +19,7 @@
 static TAutoConsoleVariable<int32> CVarMobileSupportBloomSetupRareCases(
 	TEXT("r.Mobile.MobileSupportBloomSetupRareCases"),
 	0,
-	TEXT("0: Don't generate permutations for BloomSetup rare cases. (default, like Sun+MSAA, Dof+MSAA, EyeAdaptaion+MSAA, and any of their combinations)\n")
+	TEXT("0: Don't generate permutations for BloomSetup rare cases. (default, like Sun+MetalMSAAHDRDecode, Dof+MetalMSAAHDRDecode, EyeAdaptaion+MetalMSAAHDRDecode, and any of their combinations)\n")
 	TEXT("1: Generate permutations for BloomSetup rare cases. "),
 	ECVF_ReadOnly);
 
@@ -31,6 +31,8 @@ static TAutoConsoleVariable<int32> CVarMobileEyeAdaptation(
 	TEXT(" 1: Enabled (Default)"),
 	ECVF_RenderThreadSafe);
 
+IMPLEMENT_GLOBAL_SHADER(FMSAADecodeAndCopyRectPS_ES2, "/Engine/Private/PostProcessMobile.usf", "MSAADecodeAndCopyRectPS", SF_Pixel);
+
 static EPixelFormat GetHDRPixelFormat()
 {
 	return PF_FloatRGBA;
@@ -70,14 +72,10 @@ bool IsMobileEyeAdaptationEnabled(const FViewInfo& View)
 // 8 = EyeAdaptation
 
 //Following variations should only be generated on IOS, only IOS has to do PreTonemapMSAA if MSAA is enabled.
-// 17 = Bloom + MSAA
-// 19 = Bloom + SunShaft + MSAA
-// 21 = Bloom + Dof + MSAA
-// 23 = Bloom + Dof + SunShaft + MSAA
-// 25 = Bloom + EyeAdaptation + MSAA
-// 27 = Bloom + SunShaft + EyeAdaptation + MSAA
-// 29 = Bloom + Dof + EyeAdaptation + MSAA
-// 31 = Bloom + SunShaft + Dof + EyeAdaptation + MSAA
+// 17 = Bloom + MetalMSAAHDRDecode
+// 21 = Bloom + Dof + MetalMSAAHDRDecode
+// 25 = Bloom + EyeAdaptation + MetalMSAAHDRDecode
+// 29 = Bloom + Dof + EyeAdaptation + MetalMSAAHDRDecode
 
 //Following variations are rare cases, depends on CVarMobileSupportBloomSetupRareCases
 // 2 = SunShaft
@@ -88,6 +86,21 @@ bool IsMobileEyeAdaptationEnabled(const FViewInfo& View)
 // 12 = Dof + EyeAdaptation
 // 14 = SunShaft + Dof + EyeAdaptation
 
+// 20 = Dof + MetalMSAAHDRDecode
+// 24 = EyeAdaptation + MetalMSAAHDRDecode
+// 28 = Dof + EyeAdaptation + MetalMSAAHDRDecode
+
+//Any variations with SunShaft + MetalMSAAHDRDecode should be not generated, because SceneColor has been decoded at SunMask pass
+// 19 = Bloom + SunShaft + MetalMSAAHDRDecode
+// 23 = Bloom + Dof + SunShaft + MetalMSAAHDRDecode
+// 27 = Bloom + SunShaft + EyeAdaptation + MetalMSAAHDRDecode
+// 31 = Bloom + SunShaft + Dof + EyeAdaptation + MetalMSAAHDRDecode
+
+// 18 = SunShaft + MetalMSAAHDRDecode
+// 22 = Dof + SunShaft + MetalMSAAHDRDecode
+// 26 = SunShaft + EyeAdaptation + MetalMSAAHDRDecode
+// 30 = SunShaft + Dof + EyeAdaptation + MetalMSAAHDRDecode
+
 
 // Remove the variation from this list if it should not be a rare case or enable the CVarMobileSupportBloomSetupRareCases for full cases.
 bool IsValidBloomSetupVariation(uint32 Variation)
@@ -96,25 +109,33 @@ bool IsValidBloomSetupVariation(uint32 Variation)
 		Variation == 2 ||
 		Variation == 4 ||
 		Variation == 6 ||
+
 		Variation == 10 ||
 		Variation == 12 ||
-		Variation == 14;
+		Variation == 14 ||
+		
+		Variation == 20 || 
+		Variation == 24 || 
+		Variation == 28;
 
 	return !bIsRareCases || CVarMobileSupportBloomSetupRareCases.GetValueOnAnyThread() != 0;
 }
 
-bool IsValidBloomSetupVariation(uint32 UseBloom, uint32 UseSun, uint32 UseDof, uint32 UseEyeAdaptation)
+bool IsValidBloomSetupVariation(bool bUseBloom, bool bUseSun, bool bUseDof, bool bUseEyeAdaptation)
 {
-	uint32 Variation = UseBloom | (UseSun << 1) | (UseDof << 2) | (UseEyeAdaptation << 3);
+	uint32 Variation = bUseBloom	? 1 << 0 : 0;
+	Variation |= bUseSun			? 1 << 1 : 0;
+	Variation |= bUseDof			? 1 << 2 : 0;
+	Variation |= bUseEyeAdaptation	? 1 << 3 : 0;
 	return IsValidBloomSetupVariation(Variation);
 }
 
-uint32 GetBloomSetupOutputNum(uint32 UseBloom, uint32 UseSun, uint32 UseDof, uint32 UseEyeAdaptation)
+uint32 GetBloomSetupOutputNum(bool bUseBloom, bool bUseSun, bool bUseDof, bool bUseEyeAdaptation)
 {
-	bool bValidVariation = IsValidBloomSetupVariation(UseBloom, UseSun, UseDof, UseEyeAdaptation);
+	bool bValidVariation = IsValidBloomSetupVariation(bUseBloom, bUseSun, bUseDof, bUseEyeAdaptation);
 
 	//if the variation is invalid, always use bloom permutation
-	return ((!bValidVariation || UseBloom) ? 1 : 0) + ((UseSun || UseDof) ? 1 : 0) + (UseEyeAdaptation ? 1 : 0);
+	return ((!bValidVariation || bUseBloom) ? 1 : 0) + ((bUseSun || bUseDof) ? 1 : 0) + (bUseEyeAdaptation ? 1 : 0);
 }
 
 //
@@ -155,113 +176,117 @@ public:
 
 IMPLEMENT_SHADER_TYPE(, FPostProcessBloomSetupVS_ES2, TEXT("/Engine/Private/PostProcessMobile.usf"), TEXT("BloomVS_ES2"), SF_Vertex);
 
-template <uint32 Variation>
+
 class FPostProcessBloomSetupPS_ES2 : public FGlobalShader
 {
-	DECLARE_SHADER_TYPE(FPostProcessBloomSetupPS_ES2, Global);
+public:
+	DECLARE_GLOBAL_SHADER(FPostProcessBloomSetupPS_ES2);
+	SHADER_USE_PARAMETER_STRUCT(FPostProcessBloomSetupPS_ES2, FGlobalShader);
+
+	class FUseBloomDim :				SHADER_PERMUTATION_BOOL("ES2_USE_BLOOM");
+	class FUseSunDim :					SHADER_PERMUTATION_BOOL("ES2_USE_SUN");
+	class FUseDofDim :					SHADER_PERMUTATION_BOOL("ES2_USE_DOF");
+	class FUseEyeAdaptationDim :		SHADER_PERMUTATION_BOOL("ES2_USE_EYEADAPTATION");
+	class FUseMetalMSAAHDRDecodeDim :	SHADER_PERMUTATION_BOOL("METAL_MSAA_HDR_DECODE");
+
+	using FPermutationDomain = TShaderPermutationDomain<
+		FUseBloomDim,
+		FUseSunDim,
+		FUseDofDim,
+		FUseEyeAdaptationDim,
+		FUseMetalMSAAHDRDecodeDim>;
 
 	BEGIN_SHADER_PARAMETER_STRUCT(FParameters, )
 		SHADER_PARAMETER(float, BloomThreshold)
 		SHADER_PARAMETER_STRUCT_REF(FViewUniformShaderParameters, View)
 		SHADER_PARAMETER_STRUCT(FEyeAdaptationParameters, EyeAdaptation)
+		SHADER_PARAMETER_TEXTURE(Texture2D, PostprocessInput0)
+		SHADER_PARAMETER_SAMPLER(SamplerState, PostprocessInput0Sampler)
+		SHADER_PARAMETER_TEXTURE(Texture2D, PostprocessInput1)
+		SHADER_PARAMETER_SAMPLER(SamplerState, PostprocessInput1Sampler)
 	END_SHADER_PARAMETER_STRUCT()
 
 	static bool ShouldCompilePermutation(const FGlobalShaderPermutationParameters& Parameters)
 	{
-		bool bValidVariation = IsValidBloomSetupVariation(Variation);
+		FPermutationDomain PermutationVector(Parameters.PermutationId);
+
+		auto bUseBloomDim = PermutationVector.Get<FUseBloomDim>();
+
+		auto bUseSunDim = PermutationVector.Get<FUseSunDim>();
+
+		auto bUseDofDim = PermutationVector.Get<FUseDofDim>();
+
+		auto bUseEyeAdaptationDim = PermutationVector.Get<FUseEyeAdaptationDim>();
 
-		return !IsConsolePlatform(Parameters.Platform) && 
+		auto bUseMetalMSAAHDRDecodeDim = PermutationVector.Get<FUseMetalMSAAHDRDecodeDim>();
+
+		bool bValidVariation = IsValidBloomSetupVariation(bUseBloomDim, bUseSunDim, bUseDofDim, bUseEyeAdaptationDim);
+
+		return IsMobilePlatform(Parameters.Platform) && 
 			// Exclude rare cases if CVarMobileSupportBloomSetupRareCases is 0
-			(bValidVariation || CVarMobileSupportBloomSetupRareCases.GetValueOnAnyThread() != 0) && 
-			// IOS should generate all valid variations, other mobile platform should exclude MSAA permu
-			(IsMetalMobilePlatform(Parameters.Platform) || (Variation&16) == 0);
+			(bValidVariation) && 
+			// IOS should generate all valid variations except SunShaft + MetalMSAAHDRDecode, other mobile platform should exclude MetalMSAAHDRDecode permutation
+			(!bUseMetalMSAAHDRDecodeDim || (IsMetalMobilePlatform(Parameters.Platform) && !bUseSunDim));
 	}
 
 	static void ModifyCompilationEnvironment(const FGlobalShaderPermutationParameters& Parameters, FShaderCompilerEnvironment& OutEnvironment)
 	{
 		FGlobalShader::ModifyCompilationEnvironment(Parameters, OutEnvironment);
-
-		CA_SUPPRESS(6313);
-		OutEnvironment.SetDefine(TEXT("ES2_USE_MSAA"), (Variation & 16) ? (uint32)1 : (uint32)0);
-		CA_SUPPRESS(6313);
-		OutEnvironment.SetDefine(TEXT("ES2_USE_EYEADAPTATION"), (Variation & 8) ? (uint32)1 : (uint32)0);
-		CA_SUPPRESS(6313);
-		OutEnvironment.SetDefine(TEXT("ES2_USE_DOF"), (Variation & 4) ? (uint32)1 : (uint32)0);
-		CA_SUPPRESS(6313);
-		OutEnvironment.SetDefine(TEXT("ES2_USE_SUN"), (Variation & 2) ? (uint32)1 : (uint32)0);
-		CA_SUPPRESS(6313);
-		OutEnvironment.SetDefine(TEXT("ES2_USE_BLOOM"), (Variation & 1) ? (uint32)1 : (uint32)0);
 	}
 
-	FPostProcessBloomSetupPS_ES2() {}
-
-public:
-	LAYOUT_FIELD(FPostProcessPassParameters, PostprocessParameter);
+	static FPermutationDomain RemapPermutationVector(FPermutationDomain PermutationVector, bool bValidVariation)
+	{
+		if (!bValidVariation)
+		{
+			//Use the permutation with Bloom
+			PermutationVector.Set<FUseBloomDim>(true);
+		}
+		return PermutationVector;
+	}
 
-	/** Initialization constructor. */
-	FPostProcessBloomSetupPS_ES2(const ShaderMetaType::CompiledShaderInitializerType& Initializer)
-		: FGlobalShader(Initializer)
+	static FPermutationDomain BuildPermutationVector(bool bInUseBloom, bool bInUseSun, bool bInUseDof, bool bInUseEyeAdaptation, bool bInUseMetalMSAAHDRDecode)
 	{
-		Bindings.BindForLegacyShaderParameters(this, Initializer.ParameterMap, *FParameters::FTypeInfo::GetStructMetadata());
-		PostprocessParameter.Bind(Initializer.ParameterMap);
+		FPermutationDomain PermutationVector;
+		PermutationVector.Set<FUseBloomDim>(bInUseBloom);
+		PermutationVector.Set<FUseSunDim>(bInUseSun);
+		PermutationVector.Set<FUseDofDim>(bInUseDof);
+		PermutationVector.Set<FUseEyeAdaptationDim>(bInUseEyeAdaptation);
+		PermutationVector.Set<FUseMetalMSAAHDRDecodeDim>(bInUseMetalMSAAHDRDecode);
+		return RemapPermutationVector(PermutationVector, IsValidBloomSetupVariation(bInUseBloom, bInUseSun, bInUseDof, bInUseEyeAdaptation));
 	}
+public:
 
-	void SetPS(const FRenderingCompositePassContext& Context, const TShaderRef<FPostProcessBloomSetupPS_ES2>& Shader)
+	void SetPS(const FRenderingCompositePassContext& Context, const TShaderRef<FPostProcessBloomSetupPS_ES2>& Shader, FRHITexture* PostprocessInput0, FRHITexture* PostprocessInput1)
 	{
 		FRHIPixelShader* ShaderRHI = Context.RHICmdList.GetBoundPixelShader();
 
 		const FPostProcessSettings& Settings = Context.View.FinalPostProcessSettings;
 
-		PostprocessParameter.SetPS(Context.RHICmdList, ShaderRHI, Context, TStaticSamplerState<SF_Bilinear, AM_Clamp, AM_Clamp, AM_Clamp>::GetRHI());
-
 		FParameters ShaderParameters;
 
 		ShaderParameters.EyeAdaptation = GetEyeAdaptationParameters(Context.View, ERHIFeatureLevel::ES3_1);
 		ShaderParameters.BloomThreshold = Settings.BloomThreshold;
 		ShaderParameters.View = Context.View.ViewUniformBuffer;
 
+		ShaderParameters.PostprocessInput0 = PostprocessInput0;
+		ShaderParameters.PostprocessInput0Sampler = TStaticSamplerState<SF_Bilinear, AM_Clamp, AM_Clamp, AM_Clamp>::GetRHI();
+
+		ShaderParameters.PostprocessInput1 = PostprocessInput1;
+		ShaderParameters.PostprocessInput1Sampler = TStaticSamplerState<SF_Bilinear, AM_Clamp, AM_Clamp, AM_Clamp>::GetRHI();
+
 		SetShaderParameters(Context.RHICmdList, Shader, ShaderRHI, ShaderParameters);
 	}
 };
 
-#define BLOOMSETUPPS_VARIATION(A) \
-typedef FPostProcessBloomSetupPS_ES2<A> FPostProcessBloomSetupPS_ES2_##A; \
-IMPLEMENT_SHADER_TYPE(template<>,FPostProcessBloomSetupPS_ES2_##A,TEXT("/Engine/Private/PostProcessMobile.usf"),TEXT("BloomPS_ES2"),SF_Pixel);
-
-BLOOMSETUPPS_VARIATION(1)
-BLOOMSETUPPS_VARIATION(3)
-BLOOMSETUPPS_VARIATION(5)
-BLOOMSETUPPS_VARIATION(7)
-BLOOMSETUPPS_VARIATION(9)
-BLOOMSETUPPS_VARIATION(11)
-BLOOMSETUPPS_VARIATION(13)
-BLOOMSETUPPS_VARIATION(15)
-BLOOMSETUPPS_VARIATION(17)
-BLOOMSETUPPS_VARIATION(2)
-BLOOMSETUPPS_VARIATION(4)
-BLOOMSETUPPS_VARIATION(6)
-BLOOMSETUPPS_VARIATION(8)
-BLOOMSETUPPS_VARIATION(10)
-BLOOMSETUPPS_VARIATION(12)
-BLOOMSETUPPS_VARIATION(14)
-BLOOMSETUPPS_VARIATION(19)
-BLOOMSETUPPS_VARIATION(21)
-BLOOMSETUPPS_VARIATION(23)
-BLOOMSETUPPS_VARIATION(25)
-BLOOMSETUPPS_VARIATION(27)
-BLOOMSETUPPS_VARIATION(29)
-BLOOMSETUPPS_VARIATION(31)
-
-#undef BLOOMSETUPPS_VARIATION
-
-template <uint32 Variation>
-void FRCPassPostProcessBloomSetupES2::SetShaderAndExecute(FRenderingCompositePassContext& Context)
+IMPLEMENT_GLOBAL_SHADER(FPostProcessBloomSetupPS_ES2, "/Engine/Private/PostProcessMobile.usf", "BloomPS_ES2", SF_Pixel);
+
+void FRCPassPostProcessBloomSetupES2::Process(FRenderingCompositePassContext& Context)
 {
 	SCOPED_DRAW_EVENT(Context.RHICmdList, PostProcessBloomSetup);
 
 	const FPooledRenderTargetDesc* InputDesc = GetInputDesc(ePId_Input0);
 	FIntPoint PrePostSourceViewportSize = PrePostSourceViewportRect.Size();
-	
+
 	uint32 DstX = FMath::DivideAndRoundUp(PrePostSourceViewportSize.X, 4);
 	uint32 DstY = FMath::DivideAndRoundUp(PrePostSourceViewportSize.Y, 4);
 
@@ -275,7 +300,7 @@ void FRCPassPostProcessBloomSetupES2::SetShaderAndExecute(FRenderingCompositePas
 
 	FIntPoint SrcSize;
 	FIntRect SrcRect;
-	if(bUseViewRectSource)
+	if (bUseViewRectSource)
 	{
 		// Mobile with framebuffer fetch uses view rect as source.
 		const FViewInfo& View = Context.View;
@@ -296,7 +321,7 @@ void FRCPassPostProcessBloomSetupES2::SetShaderAndExecute(FRenderingCompositePas
 	const FSceneRenderTargetItem* DestRenderTarget1 = nullptr;
 	const FSceneRenderTargetItem* DestRenderTarget2 = nullptr;
 
-	uint32 OutputNum = GetBloomSetupOutputNum(UseBloom, UseSun, UseDof, UseEyeAdaptation);
+	uint32 OutputNum = GetBloomSetupOutputNum(bUseBloom, bUseSun, bUseDof, bUseEyeAdaptation);
 	if (OutputNum > 1)
 	{
 		DestRenderTarget1 = &PassOutputs[1].RequestSurface(Context);
@@ -317,13 +342,17 @@ void FRCPassPostProcessBloomSetupES2::SetShaderAndExecute(FRenderingCompositePas
 
 	FRHIRenderPassInfo RPInfo(NumRenderTargets, RenderTargets, ERenderTargetActions::DontLoad_Store);
 
-	bool bIsValidVariation = IsValidBloomSetupVariation(UseBloom, UseSun, UseDof, UseEyeAdaptation);
+	bool bIsValidVariation = IsValidBloomSetupVariation(bUseBloom, bUseSun, bUseDof, bUseEyeAdaptation);
 
 	if (!bIsValidVariation)
 	{
 		RPInfo.ColorRenderTargets[0].Action = ERenderTargetActions::DontLoad_DontStore;
 	}
 
+	FRHITexture* InputRenderTarget0 = GetInput(ePId_Input0)->GetOutput()->PooledRenderTarget->GetRenderTargetItem().ShaderResourceTexture;
+
+	FRHITexture* InputRenderTarget1 = GetInput(ePId_Input1)->IsValid() ? GetInput(ePId_Input1)->GetOutput()->PooledRenderTarget->GetRenderTargetItem().ShaderResourceTexture : nullptr;
+
 	Context.RHICmdList.BeginRenderPass(RPInfo, TEXT("PostProcessBloomSetupES2"));
 	{
 		Context.SetViewportAndCallRHI(0, 0, 0.0f, DstX, DstY, 1.0f);
@@ -335,7 +364,8 @@ void FRCPassPostProcessBloomSetupES2::SetShaderAndExecute(FRenderingCompositePas
 		GraphicsPSOInit.DepthStencilState = TStaticDepthStencilState<false, CF_Always>::GetRHI();
 
 		TShaderMapRef<FPostProcessBloomSetupVS_ES2> VertexShader(Context.GetShaderMap());
-		TShaderMapRef<FPostProcessBloomSetupPS_ES2<Variation> > PixelShader(Context.GetShaderMap());
+		auto ShaderPermutationVector = FPostProcessBloomSetupPS_ES2::BuildPermutationVector(bUseBloom, bUseSun, bUseDof, bUseEyeAdaptation, bUseMetalMSAAHDRDecode);
+		TShaderMapRef<FPostProcessBloomSetupPS_ES2> PixelShader(Context.GetShaderMap(), ShaderPermutationVector);
 
 		GraphicsPSOInit.BoundShaderState.VertexDeclarationRHI = GFilterVertexDeclaration.VertexDeclarationRHI;
 		GraphicsPSOInit.BoundShaderState.VertexShaderRHI = VertexShader.GetVertexShader();
@@ -345,7 +375,7 @@ void FRCPassPostProcessBloomSetupES2::SetShaderAndExecute(FRenderingCompositePas
 		SetGraphicsPipelineState(Context.RHICmdList, GraphicsPSOInit);
 
 		VertexShader->SetVS(Context);
-		PixelShader->SetPS(Context, PixelShader);
+		PixelShader->SetPS(Context, PixelShader, InputRenderTarget0, InputRenderTarget1);
 
 		DrawRectangle(
 			Context.RHICmdList,
@@ -369,69 +399,9 @@ void FRCPassPostProcessBloomSetupES2::SetShaderAndExecute(FRenderingCompositePas
 	}
 }
 
-void FRCPassPostProcessBloomSetupES2::Process(FRenderingCompositePassContext& Context)
-{
-	SCOPED_DRAW_EVENT(Context.RHICmdList, PostProcessBloomSetup);
-
-	uint32 Variation = UseBloom | (UseSun << 1) | (UseDof << 2) | (UseEyeAdaptation << 3) | (UseMSAA << 4);
-	bool bMobileSupportBloomSetupRareCasesOff = CVarMobileSupportBloomSetupRareCases.GetValueOnRenderThread() == 0;
-
-
-#define SET_SHADER_SINGLE_CASE(ShaderCase) \
-		case ShaderCase: \
-			SetShaderAndExecute<ShaderCase>(Context);\
-		break
-
-#define SET_SHADER_CASE(ShaderCase, FallBackCase) \
-		case ShaderCase: \
-		if (bMobileSupportBloomSetupRareCasesOff)		\
-		{						\
-			SetShaderAndExecute<FallBackCase>(Context);\
-		} \
-		else \
-		{ \
-			SetShaderAndExecute<ShaderCase>(Context);\
-		} \
-		break
-
-	switch (Variation)
-	{
-		SET_SHADER_SINGLE_CASE(1); // Bloom
-		SET_SHADER_SINGLE_CASE(3); // Bloom + SunShaft
-		SET_SHADER_SINGLE_CASE(5); // Bloom + Dof
-		SET_SHADER_SINGLE_CASE(7); // Bloom + Dof + SunShaft
-		SET_SHADER_SINGLE_CASE(9); // Bloom + EyeAdaptation
-		SET_SHADER_SINGLE_CASE(11); // Bloom + SunShaft + EyeAdaptation
-		SET_SHADER_SINGLE_CASE(13); // Bloom + Dof + EyeAdaptation
-		SET_SHADER_SINGLE_CASE(15); // Bloom + SunShaft + Dof + EyeAdaptation
-		SET_SHADER_SINGLE_CASE(8); // EyeAdaptation
-
-		//Following cases only for IOS platform
-		SET_SHADER_SINGLE_CASE(17); // Bloom + MSAA
-		SET_SHADER_SINGLE_CASE(19); // Bloom + SunShaft + MSAA
-		SET_SHADER_SINGLE_CASE(21); // Bloom + Dof + MSAA
-		SET_SHADER_SINGLE_CASE(23); // Bloom + Dof + SunShaft + MSAA
-		SET_SHADER_SINGLE_CASE(25); // Bloom + EyeAdaptation + MSAA
-		SET_SHADER_SINGLE_CASE(27); // Bloom + SunShaft + EyeAdaptation + MSAA
-		SET_SHADER_SINGLE_CASE(29); // Bloom + Dof + EyeAdaptation + MSAA
-		SET_SHADER_SINGLE_CASE(31); // Bloom + SunShaft + Dof + EyeAdaptation + MSAA
-
-		//Following rare cases could fall back to Bloom, depends on CVarMobileSupportBloomSetupRareCases
-		SET_SHADER_CASE(2, 3); // SunShaft
-		SET_SHADER_CASE(4, 5); // Dof
-		SET_SHADER_CASE(6, 7); // SunShaft + Dof
-		
-		SET_SHADER_CASE(10, 11); // SunShaft + EyeAdaptation
-		SET_SHADER_CASE(12, 13); // Dof + EyeAdaptation
-		SET_SHADER_CASE(14, 15); // SunShaft + Dof + EyeAdaptation
-	}
-#undef SET_SHADER_CASE
-#undef SET_SHADER_SINGLE_CASE
-}
-
 FPooledRenderTargetDesc FRCPassPostProcessBloomSetupES2::ComputeOutputDesc(EPassOutputId InPassOutputId) const
 {
-	bool bIsValidVariation = IsValidBloomSetupVariation(UseBloom, UseSun, UseDof, UseEyeAdaptation);
+	bool bIsValidVariation = IsValidBloomSetupVariation(bUseBloom, bUseSun, bUseDof, bUseEyeAdaptation);
 
 	FPooledRenderTargetDesc Ret;
 	Ret.Depth = 0;
@@ -447,7 +417,7 @@ FPooledRenderTargetDesc FRCPassPostProcessBloomSetupES2::ComputeOutputDesc(EPass
 
 	Ret.bForceSeparateTargetAndShaderResource = false;
 	
-	if (!bIsValidVariation || UseBloom)
+	if (!bIsValidVariation || bUseBloom)
 	{
 		Ret.Format = InPassOutputId == ePId_Output0 ? PF_FloatR11G11B10 : PF_R16F;
 	}
@@ -799,32 +769,51 @@ FPooledRenderTargetDesc FRCPassPostProcessBloomUpES2::ComputeOutputDesc(EPassOut
 // SUN MASK
 //
 
-template <uint32 UseSunDofDepth>
 class FPostProcessSunMaskPS_ES2 : public FGlobalShader
 {
-	DECLARE_SHADER_TYPE(FPostProcessSunMaskPS_ES2, Global);
+public:
+	DECLARE_GLOBAL_SHADER(FPostProcessSunMaskPS_ES2);
+	SHADER_USE_PARAMETER_STRUCT(FPostProcessSunMaskPS_ES2, FGlobalShader);
+
+	class FUseSunDim :					SHADER_PERMUTATION_BOOL("ES2_USE_SUN");
+	class FUseDofDim :					SHADER_PERMUTATION_BOOL("ES2_USE_DOF");
+	class FUseDepthTextureDim :			SHADER_PERMUTATION_BOOL("ES2_USE_DEPTHTEXTURE");
+	class FUseMetalMSAAHDRDecodeDim :	SHADER_PERMUTATION_BOOL("METAL_MSAA_HDR_DECODE");
+
+	using FPermutationDomain = TShaderPermutationDomain<
+		FUseSunDim,
+		FUseDofDim, 
+		FUseDepthTextureDim,
+		FUseMetalMSAAHDRDecodeDim>;
 
 	BEGIN_SHADER_PARAMETER_STRUCT(FParameters, )
 		SHADER_PARAMETER(FVector4, SunColorApertureDiv2)
 		SHADER_PARAMETER_STRUCT_REF(FMobileSceneTextureUniformParameters, SceneTextures)
 		SHADER_PARAMETER_STRUCT_REF(FViewUniformShaderParameters, View)
-		SHADER_PARAMETER_UAV(RWTexture2D<half4>, OutputTexture)
+		SHADER_PARAMETER_TEXTURE(Texture2D, PostprocessInput0)
+		SHADER_PARAMETER_SAMPLER(SamplerState, PostprocessInput0Sampler)
 	END_SHADER_PARAMETER_STRUCT()
 
 	static bool ShouldCompilePermutation(const FGlobalShaderPermutationParameters& Parameters)
 	{
-		return !IsConsolePlatform(Parameters.Platform);
+		FPermutationDomain PermutationVector(Parameters.PermutationId);
+
+		auto bUseSunDim = PermutationVector.Get<FUseSunDim>();
+
+		auto bUseDofDim = PermutationVector.Get<FUseDofDim>();
+
+		auto bUseMetalMSAAHDRDecodeDim = PermutationVector.Get<FUseMetalMSAAHDRDecodeDim>();
+
+		return IsMobilePlatform(Parameters.Platform) && 
+				// Only generate shaders with SunShaft and/or Dof
+				(bUseSunDim || bUseDofDim) && 
+				// Only generated MetalMSAAHDRDecode shaders for SunShaft
+				(!bUseMetalMSAAHDRDecodeDim || (bUseSunDim && IsMetalMobilePlatform(Parameters.Platform)));
 	}	
 
 	static void ModifyCompilationEnvironment(const FGlobalShaderPermutationParameters& Parameters, FShaderCompilerEnvironment& OutEnvironment)
 	{
 		FGlobalShader::ModifyCompilationEnvironment(Parameters, OutEnvironment);
-		CA_SUPPRESS(6313);
-		OutEnvironment.SetDefine(TEXT("ES2_USE_DEPTHTEXTURE"), (UseSunDofDepth & 4) ? (uint32)1 : (uint32)0);
-		CA_SUPPRESS(6313);
-		OutEnvironment.SetDefine(TEXT("ES2_USE_SUN"), (UseSunDofDepth & 2) ? (uint32)1 : (uint32)0);
-		CA_SUPPRESS(6313);
-		OutEnvironment.SetDefine(TEXT("ES2_USE_DOF"), (UseSunDofDepth & 1) ? (uint32)1 : (uint32)0);
 		
 		// This post-processor has a 1-Dimensional color attachment for SV_Target0
 		OutEnvironment.SetDefine(TEXT("SUBPASS_COLOR0_ATTACHMENT_DIM"), 1);
@@ -838,24 +827,35 @@ class FPostProcessSunMaskPS_ES2 : public FGlobalShader
 		}
 	}
 
-	FPostProcessSunMaskPS_ES2() {}
+	static FPermutationDomain RemapPermutationVector(FPermutationDomain PermutationVector)
+	{
+		auto UseSunDim = PermutationVector.Get<FUseSunDim>();
 
-public:
-	LAYOUT_FIELD(FPostProcessPassParameters, PostprocessParameter);
+		if (!UseSunDim)
+		{
+			// Don't use MetalMSAAHDRDecode permutation without SunShaft
+			PermutationVector.Set<FUseMetalMSAAHDRDecodeDim>(false);
+		}
 
-	FPostProcessSunMaskPS_ES2(const ShaderMetaType::CompiledShaderInitializerType& Initializer)
-		: FGlobalShader(Initializer)
+		return PermutationVector;
+	}
+
+	static FPermutationDomain BuildPermutationVector(bool bInUseSun, bool bInUseDof, bool bInUseDepthTexture, bool bInUseMetalMSAAHDRDecode)
 	{
-		PostprocessParameter.Bind(Initializer.ParameterMap);
-		Bindings.BindForLegacyShaderParameters(this, Initializer.ParameterMap, *FParameters::FTypeInfo::GetStructMetadata());
+		FPermutationDomain PermutationVector;
+		PermutationVector.Set<FUseSunDim>(bInUseSun);
+		PermutationVector.Set<FUseDofDim>(bInUseDof);
+		PermutationVector.Set<FUseDepthTextureDim>(bInUseDepthTexture);
+		PermutationVector.Set<FUseMetalMSAAHDRDecodeDim>(bInUseMetalMSAAHDRDecode);
+		return RemapPermutationVector(PermutationVector);
 	}
 
-	void SetPS(const FRenderingCompositePassContext& Context, const TShaderRef<FPostProcessSunMaskPS_ES2>& Shader)
+public:
+
+	void SetPS(const FRenderingCompositePassContext& Context, const TShaderRef<FPostProcessSunMaskPS_ES2>& Shader, FRHITexture* PostprocessInput0)
 	{
 		FRHIPixelShader* ShaderRHI = Context.RHICmdList.GetBoundPixelShader();
 
-		PostprocessParameter.SetPS(Context.RHICmdList, ShaderRHI, Context, TStaticSamplerState<SF_Bilinear,AM_Clamp,AM_Clamp,AM_Clamp>::GetRHI());
-
 		FParameters ShaderParameters;
 
 		ShaderParameters.View = Context.View.ViewUniformBuffer;
@@ -867,16 +867,14 @@ public:
 		ShaderParameters.SunColorApertureDiv2.Z = Context.View.LightShaftColorMask.B;
 		ShaderParameters.SunColorApertureDiv2.W = GetMobileDepthOfFieldScale(Context.View) * 0.5f;
 
+		ShaderParameters.PostprocessInput0 = PostprocessInput0;
+		ShaderParameters.PostprocessInput0Sampler = TStaticSamplerState<SF_Bilinear, AM_Clamp, AM_Clamp, AM_Clamp>::GetRHI();
+
 		SetShaderParameters(Context.RHICmdList, Shader, ShaderRHI, ShaderParameters);
 	}
 };
 
-// #define avoids a lot of code duplication
-#define SUNMASK_PS_ES2(A) typedef FPostProcessSunMaskPS_ES2<A> FPostProcessSunMaskPS_ES2_##A; \
-	IMPLEMENT_SHADER_TYPE(template<>,FPostProcessSunMaskPS_ES2_##A,TEXT("/Engine/Private/PostProcessMobile.usf"),TEXT("SunMaskPS_ES2"), SF_Pixel);
-
-SUNMASK_PS_ES2(0)  SUNMASK_PS_ES2(1)  SUNMASK_PS_ES2(2)  SUNMASK_PS_ES2(3) SUNMASK_PS_ES2(4)  SUNMASK_PS_ES2(5)  SUNMASK_PS_ES2(6)  SUNMASK_PS_ES2(7)
-#undef SUNMASK_PS_ES2
+IMPLEMENT_GLOBAL_SHADER(FPostProcessSunMaskPS_ES2, "/Engine/Private/PostProcessMobile.usf", "SunMaskPS_ES2", SF_Pixel);
 
 
 class FPostProcessSunMaskVS_ES2 : public FGlobalShader
@@ -909,9 +907,10 @@ public:
 
 IMPLEMENT_SHADER_TYPE(,FPostProcessSunMaskVS_ES2,TEXT("/Engine/Private/PostProcessMobile.usf"),TEXT("SunMaskVS_ES2"),SF_Vertex);
 
-template <uint32 UseSunDofDepth>
-void FRCPassPostProcessSunMaskES2::SetShaderAndExecute(const FRenderingCompositePassContext& Context)
+void FRCPassPostProcessSunMaskES2::Process(FRenderingCompositePassContext& Context)
 {
+	SCOPED_DRAW_EVENT(Context.RHICmdList, PostProcessSunMask);
+
 	const FPooledRenderTargetDesc* InputDesc = GetInputDesc(ePId_Input0);
 
 	uint32 DstX = FMath::Max(1, PrePostSourceViewportSize.X);
@@ -935,15 +934,12 @@ void FRCPassPostProcessSunMaskES2::SetShaderAndExecute(const FRenderingComposite
 // TODO: This won't work with scaled views.
 	SrcRect = View.ViewRect;
 
-	CA_SUPPRESS(6313);
-	uint32 UseDepth = (UseSunDofDepth & 4) ? 1 : 0;
-
 	const FSceneRenderTargetItem& DestRenderTarget0 = PassOutputs[0].RequestSurface(Context);
 	const FSceneRenderTargetItem* DestRenderTarget1 = nullptr;
 
 	int32 NumRenderTargets = 1;
 
-	if (!UseDepth)
+	if (!bUseDepthTexture)
 	{
 		DestRenderTarget1 = &PassOutputs[1].RequestSurface(Context);
 		NumRenderTargets += 1;
@@ -955,6 +951,8 @@ void FRCPassPostProcessSunMaskES2::SetShaderAndExecute(const FRenderingComposite
 		DestRenderTarget1 != nullptr ? DestRenderTarget1->TargetableTexture : nullptr
 	};
 
+	FRHITexture* InputRenderTarget = GetInput(ePId_Input0)->GetOutput()->PooledRenderTarget->GetRenderTargetItem().ShaderResourceTexture;
+
 	FRHIRenderPassInfo RPInfo(NumRenderTargets, RenderTargets, ERenderTargetActions::DontLoad_Store);
 
 	Context.RHICmdList.BeginRenderPass(RPInfo, TEXT("PostProcessSunMaskES2"));
@@ -966,7 +964,8 @@ void FRCPassPostProcessSunMaskES2::SetShaderAndExecute(const FRenderingComposite
 	GraphicsPSOInit.DepthStencilState = TStaticDepthStencilState<false, CF_Always>::GetRHI();
 
 	TShaderMapRef<FPostProcessSunMaskVS_ES2> VertexShader(Context.GetShaderMap());
-	TShaderMapRef<FPostProcessSunMaskPS_ES2<UseSunDofDepth> > PixelShader(Context.GetShaderMap());
+	auto ShaderPermutationVector = FPostProcessSunMaskPS_ES2::BuildPermutationVector(bUseSun, bUseDof, bUseDepthTexture, bUseMetalMSAAHDRDecode);
+	TShaderMapRef<FPostProcessSunMaskPS_ES2> PixelShader(Context.GetShaderMap(), ShaderPermutationVector);
 
 	GraphicsPSOInit.BoundShaderState.VertexDeclarationRHI = GFilterVertexDeclaration.VertexDeclarationRHI;
 	GraphicsPSOInit.BoundShaderState.VertexShaderRHI = VertexShader.GetVertexShader();
@@ -976,7 +975,7 @@ void FRCPassPostProcessSunMaskES2::SetShaderAndExecute(const FRenderingComposite
 	SetGraphicsPipelineState(Context.RHICmdList, GraphicsPSOInit);
 
 	VertexShader->SetVS(Context);
-	PixelShader->SetPS(Context, PixelShader);
+	PixelShader->SetPS(Context, PixelShader, InputRenderTarget);
 
 	DrawRectangle(
 		Context.RHICmdList,
@@ -1000,31 +999,6 @@ void FRCPassPostProcessSunMaskES2::SetShaderAndExecute(const FRenderingComposite
 	}
 }
 
-void FRCPassPostProcessSunMaskES2::Process(FRenderingCompositePassContext& Context)
-{
-	SCOPED_DRAW_EVENT(Context.RHICmdList, PostProcessSunMask);
-
-	const FPooledRenderTargetDesc* InputDesc = GetInputDesc(ePId_Input0);
-
-	const FViewInfo& View = Context.View;
-	uint32 UseSun = Context.View.bLightShaftUse ? 1 : 0;
-	uint32 UseDof = (GetMobileDepthOfFieldScale(Context.View) > 0.0f) ? 1 : 0;
-	uint32 UseDepth = (InputDesc != NULL && InputDesc->Format == PF_FloatR11G11B10) ? 1 : 0;
-	uint32 UseSunDofDepth = (UseDepth << 2) | (UseSun << 1) | UseDof;
-
-	switch (UseSunDofDepth)
-	{
-	case 0: SetShaderAndExecute<0>(Context); break;
-	case 1: SetShaderAndExecute<1>(Context); break;
-	case 2: SetShaderAndExecute<2>(Context); break;
-	case 3: SetShaderAndExecute<3>(Context); break;
-	case 4: SetShaderAndExecute<4>(Context); break;
-	case 5: SetShaderAndExecute<5>(Context); break;
-	case 6: SetShaderAndExecute<6>(Context); break;
-	case 7: SetShaderAndExecute<7>(Context); break;
-	}
-}
-
 FPooledRenderTargetDesc FRCPassPostProcessSunMaskES2::ComputeOutputDesc(EPassOutputId InPassOutputId) const
 {
 	FPooledRenderTargetDesc Ret;
diff --git a/Engine/Source/Runtime/Renderer/Private/PostProcess/PostProcessMobile.h b/Engine/Source/Runtime/Renderer/Private/PostProcess/PostProcessMobile.h
--- a/Engine/Source/Runtime/Renderer/Private/PostProcess/PostProcessMobile.h
+++ b/Engine/Source/Runtime/Renderer/Private/PostProcess/PostProcessMobile.h
@@ -22,14 +22,14 @@ bool ShouldMobilePassFlipVerticalAxis(const FRenderingCompositePassContext& Cont
 class FRCPassPostProcessBloomSetupES2 : public TRenderingCompositePassBase<2, 3>
 {
 public:
-	FRCPassPostProcessBloomSetupES2(FIntRect InPrePostSourceViewportRect, bool bInUseViewRectSource, uint32 InUseBloom, uint32 InUseSun, uint32 InUseDof, uint32 InUseEyeAdaptation, uint32 InUseMSAA)
+	FRCPassPostProcessBloomSetupES2(FIntRect InPrePostSourceViewportRect, bool bInUseViewRectSource, bool bInUseBloom, bool bInUseSun, bool bInUseDof, bool bInUseEyeAdaptation, bool bInUseMetalMSAAHDRDecode)
 	: PrePostSourceViewportRect(InPrePostSourceViewportRect)
 	, bUseViewRectSource(bInUseViewRectSource)
-	, UseBloom(InUseBloom)
-	, UseSun(InUseSun)
-	, UseDof(InUseDof)
-	, UseEyeAdaptation(InUseEyeAdaptation)
-	, UseMSAA(InUseMSAA)
+	, bUseBloom(bInUseBloom)
+	, bUseSun(bInUseSun)
+	, bUseDof(bInUseDof)
+	, bUseEyeAdaptation(bInUseEyeAdaptation)
+	, bUseMetalMSAAHDRDecode(bInUseMetalMSAAHDRDecode)
 	{ }
 	virtual void Process(FRenderingCompositePassContext& Context) override;
 	virtual FPooledRenderTargetDesc ComputeOutputDesc(EPassOutputId InPassOutputId) const override;
@@ -37,14 +37,11 @@ public:
 private:
 	FIntRect PrePostSourceViewportRect;
 	bool bUseViewRectSource;
-	uint32 UseBloom;
-	uint32 UseSun;
-	uint32 UseDof;
-	uint32 UseEyeAdaptation;
-	uint32 UseMSAA;
-	
-	template <uint32 Variation>
-	void SetShaderAndExecute(FRenderingCompositePassContext& Context);
+	bool bUseBloom;
+	bool bUseSun;
+	bool bUseDof;
+	bool bUseEyeAdaptation;
+	bool bUseMetalMSAAHDRDecode;
 };
 
 class FRCPassPostProcessDofNearES2 : public TRenderingCompositePassBase<1, 1>
@@ -112,15 +109,22 @@ private:
 class FRCPassPostProcessSunMaskES2 : public TRenderingCompositePassBase<1, 2>
 {
 public:
-	FRCPassPostProcessSunMaskES2(FIntPoint InPrePostSourceViewportSize) : PrePostSourceViewportSize(InPrePostSourceViewportSize) { }
+	FRCPassPostProcessSunMaskES2(FIntPoint InPrePostSourceViewportSize, bool bInUseSun, bool bInUseDof, bool bInUseDepthTexture, bool bInUseMetalMSAAHDRDecode)
+	: PrePostSourceViewportSize(InPrePostSourceViewportSize)
+	, bUseSun(bInUseSun)
+	, bUseDof(bInUseDof)
+	, bUseDepthTexture(bInUseDepthTexture)
+	, bUseMetalMSAAHDRDecode(bInUseMetalMSAAHDRDecode)
+	{ }
 	virtual void Process(FRenderingCompositePassContext& Context) override;
 	virtual FPooledRenderTargetDesc ComputeOutputDesc(EPassOutputId InPassOutputId) const override;
 	virtual void Release() override { delete this; }
 private:
 	FIntPoint PrePostSourceViewportSize;
-
-	template <uint32 bUseDepthTexture>
-	void SetShaderAndExecute(const FRenderingCompositePassContext& Context);
+	bool bUseSun;
+	bool bUseDof;
+	bool bUseDepthTexture;
+	bool bUseMetalMSAAHDRDecode;
 };
 
 class FRCPassPostProcessSunAlphaES2 : public TRenderingCompositePassBase<1, 1>
@@ -218,3 +222,21 @@ public:
 	virtual void Release() override { delete this; }
 	virtual FPooledRenderTargetDesc ComputeOutputDesc(EPassOutputId InPassOutputId) const override;
 };
+
+/** Pixel shader to decode the input color and  copy pixels from src to dst only for mobile metal platform. */
+class FMSAADecodeAndCopyRectPS_ES2 : public FGlobalShader
+{
+	DECLARE_GLOBAL_SHADER(FMSAADecodeAndCopyRectPS_ES2);
+	SHADER_USE_PARAMETER_STRUCT(FMSAADecodeAndCopyRectPS_ES2, FGlobalShader);
+
+	static bool ShouldCompilePermutation(const FGlobalShaderPermutationParameters& Parameters)
+	{
+		return IsMetalMobilePlatform(Parameters.Platform);
+	}
+
+	BEGIN_SHADER_PARAMETER_STRUCT(FParameters, )
+		SHADER_PARAMETER_RDG_TEXTURE(Texture2D, InputTexture)
+		SHADER_PARAMETER_SAMPLER(SamplerState, InputSampler)
+		RENDER_TARGET_BINDING_SLOTS()
+	END_SHADER_PARAMETER_STRUCT()
+};
\ No newline at end of file
diff --git a/Engine/Source/Runtime/Renderer/Private/PostProcess/PostProcessTonemap.cpp b/Engine/Source/Runtime/Renderer/Private/PostProcess/PostProcessTonemap.cpp
--- a/Engine/Source/Runtime/Renderer/Private/PostProcess/PostProcessTonemap.cpp
+++ b/Engine/Source/Runtime/Renderer/Private/PostProcess/PostProcessTonemap.cpp
@@ -106,7 +106,7 @@ class FTonemapperVignetteDim       : SHADER_PERMUTATION_BOOL("USE_VIGNETTE");
 class FTonemapperSharpenDim        : SHADER_PERMUTATION_BOOL("USE_SHARPEN");
 class FTonemapperGrainJitterDim    : SHADER_PERMUTATION_BOOL("USE_GRAIN_JITTER");
 class FTonemapperSwitchAxis        : SHADER_PERMUTATION_BOOL("NEEDTOSWITCHVERTICLEAXIS");
-class FTonemapperMsaaDim           : SHADER_PERMUTATION_BOOL("USE_MSAA");
+class FTonemapperMsaaDim           : SHADER_PERMUTATION_BOOL("METAL_MSAA_HDR_DECODE");
 class FTonemapperEyeAdaptationDim  : SHADER_PERMUTATION_BOOL("EYEADAPTATION_EXPOSURE_FIX");
 
 using FCommonDomain = TShaderPermutationDomain<
@@ -147,7 +147,7 @@ bool ShouldCompileCommonPermutation(const FGlobalShaderPermutationParameters& Pa
 }
 
 // Common conversion of engine settings into.
-FCommonDomain BuildCommonPermutationDomain(const FViewInfo& View, bool bGammaOnly, bool bSwitchVerticalAxis)
+FCommonDomain BuildCommonPermutationDomain(const FViewInfo& View, bool bGammaOnly, bool bSwitchVerticalAxis, bool bMetalMSAAHDRDecode)
 {
 	const FSceneViewFamily* Family = View.Family;
 
@@ -169,14 +169,7 @@ FCommonDomain BuildCommonPermutationDomain(const FViewInfo& View, bool bGammaOnl
 	PermutationVector.Set<FTonemapperGrainJitterDim>(Settings.GrainJitter > 0.0f);
 	PermutationVector.Set<FTonemapperSharpenDim>(CVarTonemapperSharpen.GetValueOnRenderThread() > 0.0f);	
 	PermutationVector.Set<FTonemapperSwitchAxis>(bSwitchVerticalAxis);
-
-	if (IsMetalMobilePlatform(View.GetShaderPlatform()))
-	{
-		static const auto CVarMobileMSAA = IConsoleManager::Get().FindTConsoleVariableDataInt(TEXT("r.MobileMSAA"));
-		bool bMSAA = (CVarMobileMSAA ? CVarMobileMSAA->GetValueOnAnyThread() > 1 : false);
-
-		PermutationVector.Set<FTonemapperMsaaDim>(bMSAA);
-	}
+	PermutationVector.Set<FTonemapperMsaaDim>(bMetalMSAAHDRDecode);
 
 	return PermutationVector;
 }
@@ -807,7 +800,7 @@ FScreenPassTexture AddTonemapPass(FRDGBuilder& GraphBuilder, const FViewInfo& Vi
 	TonemapperPermutation::FDesktopDomain DesktopPermutationVector;
 
 	{
-		TonemapperPermutation::FCommonDomain CommonDomain = TonemapperPermutation::BuildCommonPermutationDomain(View, Inputs.bGammaOnly, Inputs.bFlipYAxis);
+		TonemapperPermutation::FCommonDomain CommonDomain = TonemapperPermutation::BuildCommonPermutationDomain(View, Inputs.bGammaOnly, Inputs.bFlipYAxis, Inputs.bMetalMSAAHDRDecode);
 		DesktopPermutationVector.Set<TonemapperPermutation::FCommonDomain>(CommonDomain);
 
 		if (!CommonDomain.Get<TonemapperPermutation::FTonemapperGammaOnlyDim>())
@@ -892,11 +885,12 @@ FScreenPassTexture AddTonemapPass(FRDGBuilder& GraphBuilder, const FViewInfo& Vi
 //////////////////////////////////////////////////////////////////////////
 //! DEPRECATED: This is only used by mobile until fully ported to RDG.
 
-FRCPassPostProcessTonemap::FRCPassPostProcessTonemap(bool bInDoGammaOnly, bool bInDoEyeAdaptation, bool bInHDROutput)
+FRCPassPostProcessTonemap::FRCPassPostProcessTonemap(bool bInDoGammaOnly, bool bInDoEyeAdaptation, bool bInHDROutput, bool bInMetalMSAAHDRDecode)
 	: bDoGammaOnly(bInDoGammaOnly)
 	, bDoScreenPercentageInTonemapper(false)
 	, bDoEyeAdaptation(bInDoEyeAdaptation)
 	, bHDROutput(bInHDROutput)
+	, bMetalMSAAHDRDecode(bInMetalMSAAHDRDecode)
 {}
 
 void FRCPassPostProcessTonemap::Process(FRenderingCompositePassContext& Context)
@@ -979,6 +973,7 @@ void FRCPassPostProcessTonemap::Process(FRenderingCompositePassContext& Context)
 	Inputs.bWriteAlphaChannel = View.AntiAliasingMethod == AAM_FXAA || IsPostProcessingWithAlphaChannelSupported();
 	Inputs.bFlipYAxis = ShouldMobilePassFlipVerticalAxis(Context, this);
 	Inputs.bGammaOnly = bDoGammaOnly;
+	Inputs.bMetalMSAAHDRDecode = bMetalMSAAHDRDecode;
 	Inputs.EyeAdaptationBuffer = bDoEyeAdaptation && View.GetEyeAdaptationBuffer() ? View.GetEyeAdaptationBuffer()->SRV : nullptr;
 
 	AddTonemapPass(GraphBuilder, View, Inputs);
@@ -1097,9 +1092,9 @@ class FPostProcessTonemapPS_ES2 : public FGlobalShader
 		return RemappedPermutationVector;
 	}
 
-	static FPermutationDomain BuildPermutationVector(const FViewInfo& View, bool bNeedsToSwitchVerticalAxis)
+	static FPermutationDomain BuildPermutationVector(const FViewInfo& View, bool bNeedsToSwitchVerticalAxis, bool bMetalMSAAHDRDecode)
 	{
-		TonemapperPermutation::FCommonDomain CommonPermutationVector = TonemapperPermutation::BuildCommonPermutationDomain(View, /* bGammaOnly = */ false, bNeedsToSwitchVerticalAxis);
+		TonemapperPermutation::FCommonDomain CommonPermutationVector = TonemapperPermutation::BuildCommonPermutationDomain(View, /* bGammaOnly = */ false, bNeedsToSwitchVerticalAxis, bMetalMSAAHDRDecode);
 
 		FPostProcessTonemapPS_ES2::FPermutationDomain MobilePermutationVector;
 		MobilePermutationVector.Set<TonemapperPermutation::FCommonDomain>(CommonPermutationVector);
@@ -1127,14 +1122,6 @@ class FPostProcessTonemapPS_ES2 : public FGlobalShader
 		MobilePermutationVector.Set<FTonemapperShadowTintDim>(Settings.FilmShadowTintAmount > 0.0f);
 		MobilePermutationVector.Set<FTonemapperContrastDim>(Settings.FilmContrast > 0.0f);
 
-		static const auto CVarMobileMSAA = IConsoleManager::Get().FindTConsoleVariableDataInt(TEXT("r.MobileMSAA"));
-		const EShaderPlatform ShaderPlatform = View.GetShaderPlatform();
-		if (IsMetalMobilePlatform(ShaderPlatform))
-		{
-			bool bMSAA = (CVarMobileMSAA ? CVarMobileMSAA->GetValueOnAnyThread() > 1 : false);
-			CommonPermutationVector.Set<TonemapperPermutation::FTonemapperMsaaDim>(bMSAA);
-		}
-
 		if (GSupportsRenderTargetFormat_PF_FloatRGBA)
 		{
 			bool bUseDof = GetMobileDepthOfFieldScale(View) > 0.0f && (!Settings.bMobileHQGaussian || (View.GetFeatureLevel() < ERHIFeatureLevel::ES3_1));
@@ -1418,12 +1405,13 @@ IMPLEMENT_GLOBAL_SHADER(FPostProcessTonemapVS_ES2, "/Engine/Private/PostProcessT
 IMPLEMENT_GLOBAL_SHADER(FPostProcessTonemapPS_ES2, "/Engine/Private/PostProcessTonemap.usf", "MainPS_ES2", SF_Pixel);
 
 
-FRCPassPostProcessTonemapES2::FRCPassPostProcessTonemapES2(const FViewInfo& InView, bool bInUsedFramebufferFetch, bool bInSRGBAwareTarget, bool bInDoEyeAdaptation)
+FRCPassPostProcessTonemapES2::FRCPassPostProcessTonemapES2(const FViewInfo& InView, bool bInUsedFramebufferFetch, bool bInSRGBAwareTarget, bool bInDoEyeAdaptation, bool bInMetalMSAAHDRDecode)
 	: bDoScreenPercentageInTonemapper(false)
 	, View(InView)
 	, bUsedFramebufferFetch(bInUsedFramebufferFetch)
 	, bSRGBAwareTarget(bInSRGBAwareTarget)
 	, bDoEyeAdaptation(bInDoEyeAdaptation)
+	, bMetalMSAAHDRDecode(bInMetalMSAAHDRDecode)
 { }
 
 void FRCPassPostProcessTonemapES2::Process(FRenderingCompositePassContext& Context)
@@ -1464,7 +1452,7 @@ void FRCPassPostProcessTonemapES2::Process(FRenderingCompositePassContext& Conte
 		bool bNeedsToSwitchVerticalAxis = ShouldMobilePassFlipVerticalAxis(Context, this);
 
 		auto VertexShaderPermutationVector = FPostProcessTonemapVS_ES2::BuildPermutationVector(bNeedsToSwitchVerticalAxis, bDoEyeAdaptation);
-		auto PixelShaderPermutationVector = FPostProcessTonemapPS_ES2::BuildPermutationVector(View, bNeedsToSwitchVerticalAxis);
+		auto PixelShaderPermutationVector = FPostProcessTonemapPS_ES2::BuildPermutationVector(View, bNeedsToSwitchVerticalAxis, bMetalMSAAHDRDecode);
 
 		TShaderMapRef<FPostProcessTonemapVS_ES2> VertexShader(Context.GetShaderMap(), VertexShaderPermutationVector);
 		TShaderMapRef<FPostProcessTonemapPS_ES2> PixelShader(Context.GetShaderMap(), PixelShaderPermutationVector);
diff --git a/Engine/Source/Runtime/Renderer/Private/PostProcess/PostProcessTonemap.h b/Engine/Source/Runtime/Renderer/Private/PostProcess/PostProcessTonemap.h
--- a/Engine/Source/Runtime/Renderer/Private/PostProcess/PostProcessTonemap.h
+++ b/Engine/Source/Runtime/Renderer/Private/PostProcess/PostProcessTonemap.h
@@ -81,6 +81,8 @@ struct FTonemapInputs
 	// Whether to leave the final output in HDR.
 	bool bOutputInHDR = false;
 
+	bool bMetalMSAAHDRDecode = false;
+
 	FRHIShaderResourceView* EyeAdaptationBuffer = nullptr;
 };
 
@@ -94,7 +96,7 @@ FScreenPassTexture AddTonemapPass(FRDGBuilder& GraphBuilder, const FViewInfo& Vi
 class FRCPassPostProcessTonemap : public TRenderingCompositePassBase<4, 1>
 {
 public:
-	FRCPassPostProcessTonemap(bool bInDoGammaOnly, bool bDoEyeAdaptation, bool bHDROutput);
+	FRCPassPostProcessTonemap(bool bInDoGammaOnly, bool bDoEyeAdaptation, bool bHDROutput, bool bMetalMSAAHDRDecode);
 
 	virtual void Process(FRenderingCompositePassContext& Context) override;
 	virtual void Release() override { delete this; }
@@ -105,6 +107,7 @@ public:
 private:
 	bool bDoEyeAdaptation;
 	bool bHDROutput;
+	bool bMetalMSAAHDRDecode;
 };
 
 // derives from TRenderingCompositePassBase<InputCount, OutputCount>
@@ -114,7 +117,7 @@ private:
 class FRCPassPostProcessTonemapES2 : public TRenderingCompositePassBase<4, 1>
 {
 public:
-	FRCPassPostProcessTonemapES2(const FViewInfo& View, bool bInUsedFramebufferFetch, bool bInSRGBAwareTarget, bool bInDoEyeAdaptation);
+	FRCPassPostProcessTonemapES2(const FViewInfo& View, bool bInUsedFramebufferFetch, bool bInSRGBAwareTarget, bool bInDoEyeAdaptation, bool bMetalMSAAHDRDecode);
 
 	// interface FRenderingCompositePass ---------
 
@@ -130,4 +133,5 @@ private:
 	bool bUsedFramebufferFetch;
 	bool bSRGBAwareTarget;
 	bool bDoEyeAdaptation;
+	bool bMetalMSAAHDRDecode;
 };
diff --git a/Engine/Source/Runtime/Renderer/Private/PostProcess/PostProcessing.cpp b/Engine/Source/Runtime/Renderer/Private/PostProcess/PostProcessing.cpp
--- a/Engine/Source/Runtime/Renderer/Private/PostProcess/PostProcessing.cpp
+++ b/Engine/Source/Runtime/Renderer/Private/PostProcess/PostProcessing.cpp
@@ -54,7 +54,7 @@ FPostProcessing GPostProcessing;
 
 bool IsMobileEyeAdaptationEnabled(const FViewInfo& View);
 
-bool IsValidBloomSetupVariation(uint32 UseBloom, uint32 UseSun, uint32 UseDof, uint32 UseEyeAdaptation);
+bool IsValidBloomSetupVariation(bool bUseBloom, bool bUseSun, bool bUseDof, bool bUseEyeAdaptation);
 
 namespace
 {
@@ -1212,7 +1212,8 @@ static FRCPassPostProcessTonemap* AddTonemapper(
 	const FRenderingCompositeOutputRef& EyeAdaptation,
 	const EAutoExposureMethod& EyeAdapationMethodId,
 	const bool bDoGammaOnly,
-	const bool bHDRTonemapperOutput)
+	const bool bHDRTonemapperOutput,
+	const bool bMetalMSAAHDRDecode)
 {
 	const FViewInfo& View = Context.View;
 	const EStereoscopicPass StereoPass = View.StereoPass;
@@ -1224,7 +1225,7 @@ static FRCPassPostProcessTonemap* AddTonemapper(
 	}
 
 	const bool bDoEyeAdaptation = IsAutoExposureMethodSupported(View.GetFeatureLevel(), EyeAdapationMethodId) && EyeAdaptation.IsValid();
-	FRCPassPostProcessTonemap* PostProcessTonemap = Context.Graph.RegisterPass(new(FMemStack::Get()) FRCPassPostProcessTonemap(bDoGammaOnly, bDoEyeAdaptation, bHDRTonemapperOutput));
+	FRCPassPostProcessTonemap* PostProcessTonemap = Context.Graph.RegisterPass(new(FMemStack::Get()) FRCPassPostProcessTonemap(bDoGammaOnly, bDoEyeAdaptation, bHDRTonemapperOutput, bMetalMSAAHDRDecode));
 
 	PostProcessTonemap->SetInput(ePId_Input0, Context.FinalOutput);
 	PostProcessTonemap->SetInput(ePId_Input1, BloomOutputCombined);
@@ -1238,7 +1239,7 @@ static FRCPassPostProcessTonemap* AddTonemapper(
 
 void FPostProcessing::AddGammaOnlyTonemapper(FPostprocessContext& Context)
 {
-	FRenderingCompositePass* PostProcessTonemap = Context.Graph.RegisterPass(new(FMemStack::Get()) FRCPassPostProcessTonemap(true, false/*eye*/, false));
+	FRenderingCompositePass* PostProcessTonemap = Context.Graph.RegisterPass(new(FMemStack::Get()) FRCPassPostProcessTonemap(true, false/*eye*/, false, false));
 
 	PostProcessTonemap->SetInput(ePId_Input0, Context.FinalOutput);
 
@@ -1308,6 +1309,11 @@ void FPostProcessing::ProcessES2(FRHICommandListImmediate& RHICmdList, FScene* S
 		const EAutoExposureMethod AutoExposureMethod = GetAutoExposureMethod(View);
 		const bool bUseEyeAdaptation = IsMobileEyeAdaptationEnabled(View);
 
+		static const auto CVarMobileMSAA = IConsoleManager::Get().FindTConsoleVariableDataInt(TEXT("r.MobileMSAA"));
+		
+		//The input scene color has been encoded to non-linear space and needs to decode somewhere if MSAA enabled on Metal platform
+		bool bMetalMSAAHDRDecode = GSupportsShaderFramebufferFetch && IsMetalMobilePlatform(View.GetShaderPlatform()) && CVarMobileMSAA && CVarMobileMSAA->GetValueOnRenderThread() > 1;
+
 		// add the passes we want to add to the graph (commenting a line means the pass is not inserted into the graph) ---------
 		FRenderingCompositeOutputRef PostProcessEyeAdaptation;
 		if( View.Family->EngineShowFlags.PostProcessing && bAllowFullPostProcess)
@@ -1315,25 +1321,22 @@ void FPostProcessing::ProcessES2(FRHICommandListImmediate& RHICmdList, FScene* S
 			const EMobileHDRMode HDRMode = GetMobileHDRMode();
 			bool bHDRModeAllowsPost = HDRMode == EMobileHDRMode::EnabledFloat16;
 
-			uint32 bUseSun = View.bLightShaftUse ? 1 : 0;
-			uint32 bUseDof = GetMobileDepthOfFieldScale(View) > 0.0f && !Context.View.Family->EngineShowFlags.VisualizeDOF ? 1 : 0;
-			uint32 bUseBloom = View.FinalPostProcessSettings.BloomIntensity > 0.0f ? 1 : 0;
-			uint32 bUseVignette = View.FinalPostProcessSettings.VignetteIntensity > 0.0f ? 1 : 0;
+			bool bUseSun = View.bLightShaftUse;
+			bool bUseDof = GetMobileDepthOfFieldScale(View) > 0.0f && !Context.View.Family->EngineShowFlags.VisualizeDOF;
+			bool bUseBloom = View.FinalPostProcessSettings.BloomIntensity > 0.0f;
+			bool bUseVignette = View.FinalPostProcessSettings.VignetteIntensity > 0.0f;
 
-			uint32 bUseBasicEyeAdaptation = bUseEyeAdaptation && (AutoExposureMethod == EAutoExposureMethod::AEM_Basic) &&
+			bool bUseBasicEyeAdaptation = bUseEyeAdaptation && (AutoExposureMethod == EAutoExposureMethod::AEM_Basic) &&
 				// Skip if we don't have any exposure range to generate (eye adaptation will clamp).
-				View.FinalPostProcessSettings.AutoExposureMinBrightness < View.FinalPostProcessSettings.AutoExposureMaxBrightness ? 1 : 0;
+				View.FinalPostProcessSettings.AutoExposureMinBrightness < View.FinalPostProcessSettings.AutoExposureMaxBrightness;
 
-			uint32 bUseHistogramEyeAdaptation = bUseEyeAdaptation && (AutoExposureMethod == EAutoExposureMethod::AEM_Histogram) &&
+			bool bUseHistogramEyeAdaptation = bUseEyeAdaptation && (AutoExposureMethod == EAutoExposureMethod::AEM_Histogram) &&
 				// Skip if we don't have any exposure range to generate (eye adaptation will clamp).
-				View.FinalPostProcessSettings.AutoExposureMinBrightness < View.FinalPostProcessSettings.AutoExposureMaxBrightness ? 1 : 0;
-
-			static const auto CVarMobileMSAA = IConsoleManager::Get().FindTConsoleVariableDataInt(TEXT("r.MobileMSAA"));
-			uint32 bUseMSAA = GSupportsShaderFramebufferFetch && IsMetalMobilePlatform(View.GetShaderPlatform()) && CVarMobileMSAA && CVarMobileMSAA->GetValueOnRenderThread() > 1 ? 1 : 0;
+				View.FinalPostProcessSettings.AutoExposureMinBrightness < View.FinalPostProcessSettings.AutoExposureMaxBrightness;
 
 			// Use original mobile Dof on ES2 devices regardless of bMobileHQGaussian.
 			// HQ gaussian 
-			uint32 bUseMobileDof = bUseDof && (!View.FinalPostProcessSettings.bMobileHQGaussian || (Context.View.GetFeatureLevel() < ERHIFeatureLevel::ES3_1)) ? 1 : 0;
+			bool bUseMobileDof = bUseDof && (!View.FinalPostProcessSettings.bMobileHQGaussian || (Context.View.GetFeatureLevel() < ERHIFeatureLevel::ES3_1));
 
 			bool bUsePost = bHDRModeAllowsPost && IsMobileHDR();
 			
@@ -1352,8 +1355,10 @@ void FPostProcessing::ProcessES2(FRHICommandListImmediate& RHICmdList, FScene* S
 			// Always evaluate custom post processes
 			if (bUsePost)
 			{
-				Context.FinalOutput = AddPostProcessMaterialChain(Context, BL_BeforeTranslucency, nullptr);
-				Context.FinalOutput = AddPostProcessMaterialChain(Context, BL_BeforeTonemapping, nullptr);
+				// The scene color will be decoded at the first post-process material and output linear color space for the following passes
+				// bMetalMSAAHDRDecode will be set to false if there is any post-process material exist
+				Context.FinalOutput = AddPostProcessMaterialChain(Context, BL_BeforeTranslucency, bMetalMSAAHDRDecode, nullptr);
+				Context.FinalOutput = AddPostProcessMaterialChain(Context, BL_BeforeTonemapping, bMetalMSAAHDRDecode, nullptr);
 			}
 
 			// Optional fixed pass processes
@@ -1362,15 +1367,20 @@ void FPostProcessing::ProcessES2(FRHICommandListImmediate& RHICmdList, FScene* S
 				FRenderingCompositeOutputRef PostProcessSunShaftAndDof;
 				if (bUseSun || bUseDof)
 				{
+					bool bUseDepthTexture = Context.FinalOutput.GetOutput()->RenderTargetDesc.Format == PF_FloatR11G11B10;
 					// Convert depth to {circle of confusion, sun shaft intensity}
 				//	FRenderingCompositePass* PostProcessSunMask = Context.Graph.RegisterPass(new(FMemStack::Get()) FRCPassPostProcessSunMaskES2(PrePostSourceViewportSize, false));
-					FRenderingCompositePass* PostProcessSunMask = Context.Graph.RegisterPass(new(FMemStack::Get()) FRCPassPostProcessSunMaskES2(SceneColorSize));
+					FRenderingCompositePass* PostProcessSunMask = Context.Graph.RegisterPass(new(FMemStack::Get()) FRCPassPostProcessSunMaskES2(SceneColorSize, bUseSun, bUseDof, bUseDepthTexture, bMetalMSAAHDRDecode));
 					PostProcessSunMask->SetInput(ePId_Input0, Context.FinalOutput);
 					PostProcessSunShaftAndDof = FRenderingCompositeOutputRef(PostProcessSunMask, ePId_Output0);
-					if (Context.FinalOutput.GetOutput()->RenderTargetDesc.Format != PF_FloatR11G11B10)
+					if (!bUseDepthTexture)
 					{
 						Context.FinalOutput = FRenderingCompositeOutputRef(PostProcessSunMask, ePId_Output1);
 					}
+
+					// The scene color will be decoded after sun mask pass and output to linear color space for following passes if sun shaft enabled
+					// set bMetalMSAAHDRDecode to false if sun shaft enabled
+					bMetalMSAAHDRDecode = (bMetalMSAAHDRDecode && !bUseSun);
 					//@todo Ronin sunmask pass isnt clipping to image only.
 				}
 
@@ -1379,9 +1389,9 @@ void FPostProcessing::ProcessES2(FRHICommandListImmediate& RHICmdList, FScene* S
 				FRenderingCompositeOutputRef PostProcessBloomSetup_EyeAdaptation;
 				if (bUseSun || bUseMobileDof || bUseBloom || bUseBasicEyeAdaptation || bUseHistogramEyeAdaptation)
 				{
-					uint32 bHasEyeAdaptationPass = (bUseBasicEyeAdaptation || bUseHistogramEyeAdaptation) ? 1 : 0;
+					bool bHasEyeAdaptationPass = (bUseBasicEyeAdaptation || bUseHistogramEyeAdaptation);
 
-					FRenderingCompositePass* Pass = Context.Graph.RegisterPass(new(FMemStack::Get()) FRCPassPostProcessBloomSetupES2(FinalOutputViewRect, bViewRectSource, bUseBloom, bUseSun, bUseMobileDof, bHasEyeAdaptationPass, bUseMSAA));
+					FRenderingCompositePass* Pass = Context.Graph.RegisterPass(new(FMemStack::Get()) FRCPassPostProcessBloomSetupES2(FinalOutputViewRect, bViewRectSource, bUseBloom, bUseSun, bUseMobileDof, bHasEyeAdaptationPass, bMetalMSAAHDRDecode));
 					Pass->SetInput(ePId_Input0, Context.FinalOutput);
 					Pass->SetInput(ePId_Input1, PostProcessSunShaftAndDof);
 
@@ -1646,7 +1656,7 @@ void FPostProcessing::ProcessES2(FRHICommandListImmediate& RHICmdList, FScene* S
 			if (bUseTonemapperFilm)
 			{
 				//@todo Ronin Set to EAutoExposureMethod::AEM_Basic for PC vk crash.
-				FRCPassPostProcessTonemap* PostProcessTonemap = AddTonemapper(Context, BloomOutput, PostProcessEyeAdaptation, AutoExposureMethod, false, false);
+				FRCPassPostProcessTonemap* PostProcessTonemap = AddTonemapper(Context, BloomOutput, PostProcessEyeAdaptation, AutoExposureMethod, false, false, bMetalMSAAHDRDecode);
 				// remember the tonemapper pass so we can check if it's last
 				TonemapperPass = PostProcessTonemap;
 
@@ -1656,7 +1666,7 @@ void FPostProcessing::ProcessES2(FRHICommandListImmediate& RHICmdList, FScene* S
 			else
 			{
 				// Must run to blit to back buffer even if post processing is off.
-				FRCPassPostProcessTonemapES2* PostProcessTonemap = (FRCPassPostProcessTonemapES2*)Context.Graph.RegisterPass(new(FMemStack::Get()) FRCPassPostProcessTonemapES2(Context.View, bViewRectSource, bSRGBAwareTarget, bUseEyeAdaptation));
+				FRCPassPostProcessTonemapES2* PostProcessTonemap = (FRCPassPostProcessTonemapES2*)Context.Graph.RegisterPass(new(FMemStack::Get()) FRCPassPostProcessTonemapES2(Context.View, bViewRectSource, bSRGBAwareTarget, bUseEyeAdaptation, bMetalMSAAHDRDecode));
 				// remember the tonemapper pass so we can check if it's last
 				TonemapperPass = PostProcessTonemap;
 
@@ -1680,6 +1690,9 @@ void FPostProcessing::ProcessES2(FRHICommandListImmediate& RHICmdList, FScene* S
 				DoScreenPercentageInTonemapperPtr = &PostProcessTonemap->bDoScreenPercentageInTonemapper;
 			}
 			SetMobilePassFlipVerticalAxis(TonemapperPass);
+			
+			//The output color should been decoded to linear space after tone mapper apparently
+			bMetalMSAAHDRDecode = false;
 		}
 
 		// if Context.FinalOutput was the clipped result of sunmask stage then this stage also restores Context.FinalOutput back original target size.
@@ -1689,7 +1702,7 @@ void FPostProcessing::ProcessES2(FRHICommandListImmediate& RHICmdList, FScene* S
 		{
 			if (IsMobileHDR())
 			{
-				Context.FinalOutput = AddPostProcessMaterialChain(Context, BL_AfterTonemapping, nullptr);
+				Context.FinalOutput = AddPostProcessMaterialChain(Context, BL_AfterTonemapping, bMetalMSAAHDRDecode, nullptr);
 			}
 			SetMobilePassFlipVerticalAxis(Context.FinalOutput.GetPass());
 
diff --git a/Engine/Source/Runtime/Renderer/Private/SceneRendering.cpp b/Engine/Source/Runtime/Renderer/Private/SceneRendering.cpp
--- a/Engine/Source/Runtime/Renderer/Private/SceneRendering.cpp
+++ b/Engine/Source/Runtime/Renderer/Private/SceneRendering.cpp
@@ -1557,12 +1557,6 @@ void FViewInfo::SetupUniformBufferParameters(
 		(RHIFeatureLevel == ERHIFeatureLevel::ES3_1) &&
 		GMaxRHIFeatureLevel > ERHIFeatureLevel::ES3_1) ? 1.0f : 0.0f;
 
-	static const auto MobileMSAACVar = IConsoleManager::Get().FindTConsoleVariableDataInt(TEXT("r.MobileMSAA"));
-	bool bMobileMSAA = Scene && IsMetalMobilePlatform(Scene->GetShaderPlatform())
-		&& MobileMSAACVar
-		&& MobileMSAACVar->GetValueOnAnyThread() > 1;
-	ViewUniformShaderParameters.IsMobileMSAA = bMobileMSAA ? 1.0f : 0.0f;
-
 	// Padding between the left and right eye may be introduced by an HMD, which instanced stereo needs to account for.
 	if ((IStereoRendering::IsStereoEyePass(StereoPass)) && (Family->Views.Num() > 1))
 	{
diff --git a/Engine/Source/Runtime/Renderer/Private/SceneVisibility.cpp b/Engine/Source/Runtime/Renderer/Private/SceneVisibility.cpp
--- a/Engine/Source/Runtime/Renderer/Private/SceneVisibility.cpp
+++ b/Engine/Source/Runtime/Renderer/Private/SceneVisibility.cpp
@@ -4118,9 +4118,6 @@ void FSceneRenderer::PostVisibilityFrameSetup(FILCUpdatePrimTaskData& OutILCTask
 			{
 				VisibleLightViewInfo.bInViewFrustum = true;
 
-				static const auto MobileMSAAVar = IConsoleManager::Get().FindTConsoleVariableDataInt(TEXT("r.MobileMSAA"));
-				bool bNotMobileMSAA = !(MobileMSAAVar ? MobileMSAAVar->GetValueOnRenderThread() > 1 : false);
-
 				// Setup single sun-shaft from direction lights for mobile.
 				if(bCheckLightShafts && LightSceneInfo->bEnableLightShaftBloom)
 				{
@@ -4141,7 +4138,7 @@ void FSceneRenderer::PostVisibilityFrameSetup(FILCUpdatePrimTaskData& OutILCTask
 						&& (LightSceneInfo->Proxy->GetLightType() == LightType_Directional 
 							|| DistanceToBlurOrigin < LightSceneInfo->Proxy->GetRadius() * PointLightRadiusFadeFactor))
 					{
-						View.bLightShaftUse = bNotMobileMSAA;
+						View.bLightShaftUse = true;
 						View.LightShaftCenter.X = ProjectedBlurOrigin.X / ProjectedBlurOrigin.W;
 						View.LightShaftCenter.Y = ProjectedBlurOrigin.Y / ProjectedBlurOrigin.W;
 						// TODO: Might want to hookup different colors for these.
