Integrate fix for issue with custom depth comparisons on mobile when no custom depth primitives are rendered
#rnx

diff --git a/Engine/Source/Runtime/Renderer/Private/PostProcess/SceneRenderTargets.cpp b/Engine/Source/Runtime/Renderer/Private/PostProcess/SceneRenderTargets.cpp
--- a/Engine/Source/Runtime/Renderer/Private/PostProcess/SceneRenderTargets.cpp
+++ b/Engine/Source/Runtime/Renderer/Private/PostProcess/SceneRenderTargets.cpp
@@ -3419,6 +3419,7 @@ void SetupMobileSceneTextureUniformParameters(
 	FMobileSceneTextureUniformParameters& SceneTextureParameters)
 {
 	FRHITexture* BlackDefault2D = GSystemTextures.BlackDummy->GetRenderTargetItem().ShaderResourceTexture;
+	FRHITexture* MaxFP16Depth2D = GSystemTextures.MaxFP16Depth->GetRenderTargetItem().ShaderResourceTexture;
 	FRHITexture* DepthDefault = GSystemTextures.DepthDummy->GetRenderTargetItem().ShaderResourceTexture;
 
 	SceneTextureParameters.SceneColorTexture = bSceneTexturesValid ? SceneContext.GetSceneColorTexture().GetReference() : BlackDefault2D;
@@ -3431,14 +3432,13 @@ void SetupMobileSceneTextureUniformParameters(
 	SceneTextureParameters.SceneAlphaCopyTexture = bSceneTexturesValid && SceneContext.HasSceneAlphaCopyTexture() ? SceneContext.GetSceneAlphaCopyTexture() : BlackDefault2D;
 	SceneTextureParameters.SceneAlphaCopyTextureSampler = TStaticSamplerState<SF_Point,AM_Clamp,AM_Clamp,AM_Clamp>::GetRHI();
 
-	FRHITexture* CustomDepth = BlackDefault2D;
+	FRHITexture* CustomDepth = MaxFP16Depth2D;
 
 	// if there is no custom depth it's better to have the far distance there
 	// we should update all pass uniform buffers at the start of frame on mobile, SceneContext.bCustomDepthIsValid is invalid at InitView, so pass the parameter from View.bUsesCustomDepthStencil
-	IPooledRenderTarget* CustomDepthTarget = bCustomDepthIsValid ? SceneContext.MobileCustomDepth.GetReference() : 0;
-	if (CustomDepthTarget)
+	if (bCustomDepthIsValid && SceneContext.MobileCustomDepth.IsValid())
 	{
-		CustomDepth = CustomDepthTarget->GetRenderTargetItem().ShaderResourceTexture;
+		CustomDepth = SceneContext.MobileCustomDepth->GetRenderTargetItem().ShaderResourceTexture;
 	}
 
 	SceneTextureParameters.CustomDepthTexture = CustomDepth;
@@ -3446,7 +3446,7 @@ void SetupMobileSceneTextureUniformParameters(
 
 	FRHITexture* MobileCustomStencil = BlackDefault2D;
 
-	if (SceneContext.MobileCustomStencil.IsValid())
+	if (bCustomDepthIsValid && SceneContext.MobileCustomStencil.IsValid())
 	{
 		MobileCustomStencil = SceneContext.MobileCustomStencil->GetRenderTargetItem().ShaderResourceTexture;
 	}
diff --git a/Engine/Source/Runtime/Renderer/Private/SystemTextures.cpp b/Engine/Source/Runtime/Renderer/Private/SystemTextures.cpp
--- a/Engine/Source/Runtime/Renderer/Private/SystemTextures.cpp
+++ b/Engine/Source/Runtime/Renderer/Private/SystemTextures.cpp
@@ -130,7 +130,7 @@ void FSystemTextures::InitializeCommonTextures(FRHICommandListImmediate& RHICmdL
 			RHICmdList.UnlockTexture2D((FTexture2DRHIRef&)PerlinNoiseGradient->GetRenderTargetItem().ShaderResourceTexture, 0, false);
 		}
 
-	if (!GSupportsShaderFramebufferFetch && GPixelFormats[PF_FloatRGBA].Supported)
+	if (GPixelFormats[PF_FloatRGBA].Supported)
 	{
 		FPooledRenderTargetDesc Desc(FPooledRenderTargetDesc::Create2DDesc(FIntPoint(1, 1), PF_FloatRGBA, FClearValueBinding(FLinearColor(65000.0f, 65000.0f, 65000.0f, 65000.0f)), TexCreate_HideInVisualizeTexture, TexCreate_RenderTargetable | TexCreate_NoFastClear | TexCreate_ShaderResource, false));
 		Desc.AutoWritable = false;
