Fix CIS

diff --git a/Engine/Source/Runtime/Renderer/Private/HairStrands/HairStrandsVoxelization.cpp b/Engine/Source/Runtime/Renderer/Private/HairStrands/HairStrandsVoxelization.cpp
--- a/Engine/Source/Runtime/Renderer/Private/HairStrands/HairStrandsVoxelization.cpp
+++ b/Engine/Source/Runtime/Renderer/Private/HairStrands/HairStrandsVoxelization.cpp
@@ -1163,42 +1163,45 @@ static void AddVirtualVoxelizationComputeRasterPass(
 	if (!bIsGPUDriven)
 		return;
 
-	const FHairStrandsMacroGroupData::TPrimitiveInfos& PrimitiveSceneInfos = MacroGroup.PrimitivesInfos;
-
-	FRDGBufferRef VoxelizationViewInfoBuffer = GraphBuilder.RegisterExternalBuffer(VoxelResources.VoxelizationViewInfoBuffer);
-	FRDGBufferSRVRef VoxelizationViewInfoBufferSRV = GraphBuilder.CreateSRV(VoxelizationViewInfoBuffer);
-	FRDGTextureRef PageTexture = GraphBuilder.RegisterExternalTexture(VoxelResources.PageTexture);
-	FRDGTextureUAVRef PageTextureUAV = GraphBuilder.CreateUAV(PageTexture);
-
-	const uint32 FrameIdMode8 = ViewInfo && ViewInfo->ViewState ? (ViewInfo->ViewState->GetFrameIndex() % 8) : 0;
-	const uint32 GroupSize = 32;
-	const uint32 DispatchCountX = 64;
-	TShaderMapRef<FVoxelRasterComputeCS> ComputeShader(ViewInfo->ShaderMap);
-	for (const FHairStrandsMacroGroupData::PrimitiveInfo& PrimitiveInfo : PrimitiveSceneInfos)
+	if (ViewInfo)
 	{
-		FVoxelRasterComputeCS::FParameters* PassParameters = GraphBuilder.AllocParameters<FVoxelRasterComputeCS::FParameters>();
-		PassParameters->MaxRasterCount				= FMath::Clamp(GHairStrandsVoxelComputeRasterMaxVoxelCount, 1, 256);
-		PassParameters->VirtualVoxel				= VoxelResources.Parameters.Common;
-		PassParameters->MacroGroupId				= MacroGroup.MacroGroupId;
-		PassParameters->VoxelizationViewInfoBuffer	= VoxelizationViewInfoBufferSRV;
-		PassParameters->DispatchCountX				= DispatchCountX;
-		PassParameters->OutPageTexture				= PageTextureUAV;
-		PassParameters->FrameIdMod8					= FrameIdMode8;
-
-		check(PrimitiveInfo.MeshBatchAndRelevance.Mesh && PrimitiveInfo.MeshBatchAndRelevance.Mesh->Elements.Num() > 0);
-		const FHairGroupPublicData* HairGroupPublicData = reinterpret_cast<const FHairGroupPublicData*>(PrimitiveInfo.MeshBatchAndRelevance.Mesh->Elements[0].VertexFactoryUserData);
-		const FHairGroupPublicData::VertexFactoryInput& VFInput = HairGroupPublicData->VFInput;
-		PassParameters->HairStrandsVF_PositionBuffer = VFInput.HairPositionBuffer;
-		PassParameters->HairStrandsVF_PositionOffset = VFInput.HairPositionOffset;
-		PassParameters->HairStrandsVF_VertexCount	 = VFInput.VertexCount;
-		PassParameters->HairStrandsVF_Radius		 = VFInput.HairRadius;
-		PassParameters->HairStrandsVF_Length		 = VFInput.HairLength;
-		PassParameters->HairStrandsVF_Density		 = VFInput.HairDensity;
-		PassParameters->HairStrandsVF_LocalToWorldPrimitiveTransform = VFInput.LocalToWorldTransform.ToMatrixWithScale();
-
-		const uint32 DispatchCountY = FMath::CeilToInt(PassParameters->HairStrandsVF_VertexCount / float(GroupSize * DispatchCountX));
-		const FIntVector DispatchCount(DispatchCountX, DispatchCountY, 1);
-		FComputeShaderUtils::AddPass(GraphBuilder, RDG_EVENT_NAME("HairStrandsVoxelComputeRaster"), ComputeShader, PassParameters, DispatchCount);
+		const FHairStrandsMacroGroupData::TPrimitiveInfos& PrimitiveSceneInfos = MacroGroup.PrimitivesInfos;
+
+		FRDGBufferRef VoxelizationViewInfoBuffer = GraphBuilder.RegisterExternalBuffer(VoxelResources.VoxelizationViewInfoBuffer);
+		FRDGBufferSRVRef VoxelizationViewInfoBufferSRV = GraphBuilder.CreateSRV(VoxelizationViewInfoBuffer);
+		FRDGTextureRef PageTexture = GraphBuilder.RegisterExternalTexture(VoxelResources.PageTexture);
+		FRDGTextureUAVRef PageTextureUAV = GraphBuilder.CreateUAV(PageTexture);
+
+		const uint32 FrameIdMode8 = ViewInfo && ViewInfo->ViewState ? (ViewInfo->ViewState->GetFrameIndex() % 8) : 0;
+		const uint32 GroupSize = 32;
+		const uint32 DispatchCountX = 64;
+		TShaderMapRef<FVoxelRasterComputeCS> ComputeShader(ViewInfo->ShaderMap);
+		for (const FHairStrandsMacroGroupData::PrimitiveInfo& PrimitiveInfo : PrimitiveSceneInfos)
+		{
+			FVoxelRasterComputeCS::FParameters* PassParameters = GraphBuilder.AllocParameters<FVoxelRasterComputeCS::FParameters>();
+			PassParameters->MaxRasterCount = FMath::Clamp(GHairStrandsVoxelComputeRasterMaxVoxelCount, 1, 256);
+			PassParameters->VirtualVoxel = VoxelResources.Parameters.Common;
+			PassParameters->MacroGroupId = MacroGroup.MacroGroupId;
+			PassParameters->VoxelizationViewInfoBuffer = VoxelizationViewInfoBufferSRV;
+			PassParameters->DispatchCountX = DispatchCountX;
+			PassParameters->OutPageTexture = PageTextureUAV;
+			PassParameters->FrameIdMod8 = FrameIdMode8;
+
+			check(PrimitiveInfo.MeshBatchAndRelevance.Mesh && PrimitiveInfo.MeshBatchAndRelevance.Mesh->Elements.Num() > 0);
+			const FHairGroupPublicData* HairGroupPublicData = reinterpret_cast<const FHairGroupPublicData*>(PrimitiveInfo.MeshBatchAndRelevance.Mesh->Elements[0].VertexFactoryUserData);
+			const FHairGroupPublicData::VertexFactoryInput& VFInput = HairGroupPublicData->VFInput;
+			PassParameters->HairStrandsVF_PositionBuffer = VFInput.HairPositionBuffer;
+			PassParameters->HairStrandsVF_PositionOffset = VFInput.HairPositionOffset;
+			PassParameters->HairStrandsVF_VertexCount = VFInput.VertexCount;
+			PassParameters->HairStrandsVF_Radius = VFInput.HairRadius;
+			PassParameters->HairStrandsVF_Length = VFInput.HairLength;
+			PassParameters->HairStrandsVF_Density = VFInput.HairDensity;
+			PassParameters->HairStrandsVF_LocalToWorldPrimitiveTransform = VFInput.LocalToWorldTransform.ToMatrixWithScale();
+
+			const uint32 DispatchCountY = FMath::CeilToInt(PassParameters->HairStrandsVF_VertexCount / float(GroupSize * DispatchCountX));
+			const FIntVector DispatchCount(DispatchCountX, DispatchCountY, 1);
+			FComputeShaderUtils::AddPass(GraphBuilder, RDG_EVENT_NAME("HairStrandsVoxelComputeRaster"), ComputeShader, PassParameters, DispatchCount);
+		}
 	}
 }
 
