Fix crash related to IsTranslucencyWritingCustomDepth() introduced in CL# 11586937

diff --git a/Engine/Source/Runtime/Renderer/Private/DepthRendering.h b/Engine/Source/Runtime/Renderer/Private/DepthRendering.h
--- a/Engine/Source/Runtime/Renderer/Private/DepthRendering.h
+++ b/Engine/Source/Runtime/Renderer/Private/DepthRendering.h
@@ -71,9 +71,9 @@ public:
 			return Parameters.VertexFactoryType->SupportsPositionOnly() && Parameters.MaterialParameters.bIsSpecialEngineMaterial;
 		}
 		
-		if (IsTranslucentBlendMode(Parameters.MaterialParameters.BlendMode) && !Parameters.MaterialParameters.bIsTranslucencyWritingCustomDepth)
+		if (IsTranslucentBlendMode(Parameters.MaterialParameters.BlendMode))
 		{
-			return false;
+			return Parameters.MaterialParameters.bIsTranslucencyWritingCustomDepth;
 		}
 
 		// Only compile for the default material and masked materials
@@ -150,9 +150,9 @@ public:
 
 	static bool ShouldCompilePermutation(const FMeshMaterialShaderPermutationParameters& Parameters)
 	{
-		if (IsTranslucentBlendMode(Parameters.MaterialParameters.BlendMode) && !Parameters.MaterialParameters.bIsTranslucencyWritingCustomDepth)
+		if (IsTranslucentBlendMode(Parameters.MaterialParameters.BlendMode))
 		{
-			return false;
+			return Parameters.MaterialParameters.bIsTranslucencyWritingCustomDepth;
 		}
 		
 		return
