Fix a bug of invalid custom depth is used in PP material.

diff --git a/Engine/Source/Runtime/Renderer/Private/ShaderBaseClasses.cpp b/Engine/Source/Runtime/Renderer/Private/ShaderBaseClasses.cpp
--- a/Engine/Source/Runtime/Renderer/Private/ShaderBaseClasses.cpp
+++ b/Engine/Source/Runtime/Renderer/Private/ShaderBaseClasses.cpp
@@ -361,7 +361,7 @@ void FMaterialShader::SetParameters(
 		{
 			FSceneRenderTargets& SceneContext = FSceneRenderTargets::Get(RHICmdList);
 			FMobileSceneTextureUniformParameters UniformParameters;
-			SetupMobileSceneTextureUniformParameters(SceneContext, View.FeatureLevel, true, true, UniformParameters);
+			SetupMobileSceneTextureUniformParameters(SceneContext, View.FeatureLevel, true, SceneContext.bCustomDepthIsValid, UniformParameters);
 			if (View.GetEyeAdaptationBuffer())
 			{
 				UniformParameters.EyeAdaptationBuffer = View.GetEyeAdaptationBuffer()->SRV;
diff --git a/Engine/Source/Runtime/Renderer/Private/SystemTextures.cpp b/Engine/Source/Runtime/Renderer/Private/SystemTextures.cpp
--- a/Engine/Source/Runtime/Renderer/Private/SystemTextures.cpp
+++ b/Engine/Source/Runtime/Renderer/Private/SystemTextures.cpp
@@ -132,7 +132,7 @@ void FSystemTextures::InitializeCommonTextures(FRHICommandListImmediate& RHICmdL
 
 	if (GPixelFormats[PF_FloatRGBA].Supported)
 	{
-		FPooledRenderTargetDesc Desc(FPooledRenderTargetDesc::Create2DDesc(FIntPoint(1, 1), PF_FloatRGBA, FClearValueBinding(FLinearColor(65000.0f, 65000.0f, 65000.0f, 65000.0f)), TexCreate_HideInVisualizeTexture, TexCreate_RenderTargetable | TexCreate_NoFastClear | TexCreate_ShaderResource, false));
+		FPooledRenderTargetDesc Desc(FPooledRenderTargetDesc::Create2DDesc(FIntPoint(1, 1), PF_FloatRGBA, FClearValueBinding(FLinearColor(65500.0f, 65500.0f, 65500.0f, 65500.0f)), TexCreate_HideInVisualizeTexture, TexCreate_RenderTargetable | TexCreate_NoFastClear | TexCreate_ShaderResource, false));
 		Desc.AutoWritable = false;
 		GRenderTargetPool.FindFreeElement(RHICmdList, Desc, MaxFP16Depth, TEXT("MaxFP16Depth"), true, ERenderTargetTransience::NonTransient);
 
