Add initial version of hair compute raster for primary visibility.

This is add a raster compute pass for computing hair transmittance. The output of this pass is currently not used. This is step toward actual visibility rasterization.

diff --git a/Engine/Shaders/Private/HairStrands/HairStrandsDebug.usf b/Engine/Shaders/Private/HairStrands/HairStrandsDebug.usf
--- a/Engine/Shaders/Private/HairStrands/HairStrandsDebug.usf
+++ b/Engine/Shaders/Private/HairStrands/HairStrandsDebug.usf
@@ -17,6 +17,9 @@ Texture2D<uint> NodeIndex;
 StructuredBuffer<FPackedHairSample> NodeData;
 SamplerState LinearSampler;
 
+Texture2D<uint> HairCountUintTexture;
+Texture2D<float> HairCountTexture;
+
 uint TileSize;
 Texture2D<uint> TileIndexTexture;
 
@@ -38,6 +41,20 @@ void MainPS(
 	const float2 UV = Input.UV;
 	const float2 PixelPos = UV * OutputResolution;
 
+	if (DebugMode == DEBUG_MODE_SAMPLE_COUNT)
+	{
+		const float HairCount = HairCountUintTexture.Load(uint3(PixelPos, 0)) / 1000.f;
+		//const float HairCount = HairCountTexture.Load(uint3(PixelPos, 0));
+
+		if (HairCount==0)
+			discard;
+		OutColor = float4(float3(1, 0, 0), 1);
+
+		OutColor = float4(GetColorCode(HairCount / float(32)), 1);
+
+		return;
+	}
+
 	const FCategorizationData CatData = DecodeCategorizationData(CategorizationTexture.Load(uint3(PixelPos,0)));
 	const bool bIsValid = CatData.PixelCoverage > 0 || DebugMode == DEBUG_MODE_TILE;
 	if (!bIsValid)
diff --git a/Engine/Shaders/Private/HairStrands/HairStrandsDebugPrint.usf b/Engine/Shaders/Private/HairStrands/HairStrandsDebugPrint.usf
--- a/Engine/Shaders/Private/HairStrands/HairStrandsDebugPrint.usf
+++ b/Engine/Shaders/Private/HairStrands/HairStrandsDebugPrint.usf
@@ -15,6 +15,7 @@ uint MaxSampleCount;
 uint2 GroupSize; 
 
 Texture2D<float> HairCountTexture;
+Texture2D<uint> HairCountUintTexture;
 Texture2D<uint4> CategorizationTexture;
 Texture2D<uint2> DepthStencilTexture;
 SamplerState LinearSampler;
@@ -52,6 +53,7 @@ void MainCS(uint GroupIndex : SV_GroupIndex, uint3 DispatchThreadId : SV_Dispatc
 		return;
 
 	const float ViewHairCount = HairCountTexture.Load(uint3(PixelPos, 0));
+	const float ViewHairCountUint = HairCountUintTexture.Load(uint3(PixelPos, 0)) /  1000.f;
 
 	const float2 OriginalPos = PixelPos / float2(MaxResolution);
 	float2 Pos = OriginalPos;
@@ -143,6 +145,28 @@ void MainCS(uint GroupIndex : SV_GroupIndex, uint3 DispatchThreadId : SV_Dispatc
 	Pos = ShaderPrintNewline(Pos);
 	Pos.x = OriginalPos.x;
 
+	// Accurate coverage (if coverage pass is activated) uint
+	Pos = ShaderPrintSymbol(Pos, _H_);
+	Pos = ShaderPrintSymbol(Pos, _A_);
+	Pos = ShaderPrintSymbol(Pos, _I_);
+	Pos = ShaderPrintSymbol(Pos, _R_);
+	Pos = ShaderPrintSymbol(Pos, _SPC_);
+	Pos = ShaderPrintSymbol(Pos, _C_);
+	Pos = ShaderPrintSymbol(Pos, _O_);
+	Pos = ShaderPrintSymbol(Pos, _U_);
+	Pos = ShaderPrintSymbol(Pos, _N_);
+	Pos = ShaderPrintSymbol(Pos, _T_);
+	Pos = ShaderPrintSymbol(Pos, _SPC_);
+	Pos = ShaderPrintSymbol(Pos, _U_);
+	Pos = ShaderPrintSymbol(Pos, _I_);
+	Pos = ShaderPrintSymbol(Pos, _N_);
+	Pos = ShaderPrintSymbol(Pos, _T_);
+	Pos = ShaderPrintSymbol(Pos, _SPC_);
+	Pos = ShaderPrint(Pos, ViewHairCountUint);
+
+	Pos = ShaderPrintNewline(Pos);
+	Pos.x = OriginalPos.x;
+	
 	// Accurate coverage (if coverage pass is activated)
 	Pos = ShaderPrintSymbol(Pos, _C_);
 	Pos = ShaderPrintSymbol(Pos, _O_);
diff --git a/Engine/Shaders/Private/HairStrands/HairStrandsVisibilityRasterCompute.usf b/Engine/Shaders/Private/HairStrands/HairStrandsVisibilityRasterCompute.usf
new file mode 100644
--- /dev/null
+++ b/Engine/Shaders/Private/HairStrands/HairStrandsVisibilityRasterCompute.usf
@@ -0,0 +1,252 @@
+// Copyright Epic Games, Inc. All Rights Reserved.
+
+#include "../Common.ush"
+#include "HairStrandsVertexFactoryCommon.ush"
+#include "HairStrandsVisibilityCommon.ush"
+
+///////////////////////////////////////////////////////////////////////////
+uint			DispatchCountX;
+uint			MaxRasterCount;
+uint			FrameIdMod8;
+uint			MacroGroupId;
+int2			OutputResolution;
+float			HairStrandsVF_Radius;
+float			HairStrandsVF_Length;
+float			HairStrandsVF_Density;
+uint			HairStrandsVF_bUseStableRasterization;
+float3			HairStrandsVF_PositionOffset;
+uint			HairStrandsVF_VertexCount;
+float4x4		HairStrandsVF_LocalToWorldPrimitiveTransform;
+Buffer<uint4>	HairStrandsVF_PositionBuffer;
+
+Texture2D<float> SceneDepthTexture;
+RWTexture2D<uint> OutTexture;
+
+///////////////////////////////////////////////////////////////////////////
+// Helping functions 
+
+float3 NDCToPixelCoord(float4 InDC)
+{
+	const float3 NDC = InDC.xyz / InDC.w;
+	float2 UV = NDC.xy * ResolvedView.ScreenPositionScaleBias.xy + ResolvedView.ScreenPositionScaleBias.wz;
+	return float3(UV * OutputResolution, NDC.z);
+}
+
+float ToCoverage(FHairControlPoint CP, float4 NDC, float RadiusAtDepth1Primary)
+{
+//	const float DeviceZ = NDC.w;
+//	const float SceneDepth = ConvertFromDeviceZ(DeviceZ);
+	const float SceneDepth = NDC.w;
+	const float PixelRadius = SceneDepth * RadiusAtDepth1Primary;
+	return CP.WorldRadius / max(CP.WorldRadius, PixelRadius);
+}
+
+void Swap(inout float V0, inout float V1)
+{
+	float T = V0;
+	V0 = V1;
+	V1 = T;
+}
+
+void Swap(inout float3 V0, inout float3 V1)
+{
+	float3 T = V0;
+	V0 = V1;
+	V1 = T;
+}
+
+bool IsVisible(float3 SP)
+{
+	const float DeviceZ = SP.z;
+	return DeviceZ > SceneDepthTexture.Load(uint3(SP.xy, 0)); // InverseZ
+}
+
+bool IsInScreen(int2 C)
+{
+	return C.x >= 0 && C.x < OutputResolution.x && C.y >= 0 && C.y < OutputResolution.y;
+}
+
+void Rasterize(int2 Coords, float Ratio, float Depth, float Coverage)
+{
+	if (IsInScreen(Coords))
+	{
+		const float CoveragePrecision = 1000.f;
+		const uint IntCoverage = Ratio * Coverage * CoveragePrecision;
+
+		InterlockedAdd(OutTexture[uint2(Coords)], IntCoverage);
+	}
+}
+
+//////////////////////////////////////////////////////////////////////////////////////////
+// Queued work
+
+struct FRasterInfo
+{
+	bool bIsSteep;
+	uint P0x;
+	uint P1x;
+	float DeltaY;
+	float DeltaZ;
+	float GradientY;
+	float GradientZ;
+};
+
+uint4 PackRasterInfo(FRasterInfo In)
+{
+	uint4 Out;
+	return Out;
+}
+
+FRasterInfo UnpackRasterInfo(uint4 In)
+{
+	return (FRasterInfo)0;
+}
+
+//////////////////////////////////////////////////////////////////////////////////////////
+
+#define GROUP_SIZE 32
+[numthreads(GROUP_SIZE, 1, 1)]
+void MainCS(uint3 DispatchThreadID : SV_DispatchThreadID)
+{
+	const uint VertexIndex = DispatchThreadID.x + DispatchThreadID.y * GROUP_SIZE * DispatchCountX;
+	if (VertexIndex >= HairStrandsVF_VertexCount)
+		return;
+
+	ResolvedView = ResolveView();
+
+	const FHairControlPoint CP0 = UnpackHairControlPoint(
+		HairStrandsVF_PositionBuffer[VertexIndex],
+		HairStrandsVF_PositionOffset,
+		HairStrandsVF_Radius,
+		HairStrandsVF_Length);
+
+	if (CP0.Type == HAIR_CONTROLPOINT_END)
+		return;
+
+	const uint VertexIndex1 = min(VertexIndex + 1, HairStrandsVF_VertexCount - 1);
+	const FHairControlPoint CP1 = UnpackHairControlPoint(
+		HairStrandsVF_PositionBuffer[VertexIndex1],
+		HairStrandsVF_PositionOffset,
+		HairStrandsVF_Radius,
+		HairStrandsVF_Length);
+
+
+	const FHairRenderInfo HairRenderInfo = GetHairRenderInfo(ResolvedView.HairRenderInfo, ResolvedView.HairRenderInfoBits, HairStrandsVF_bUseStableRasterization > 0);
+
+	const float VelocityScale = 0;
+	FHairViewInfo HairViewInfo;
+	HairViewInfo.TranslatedWorldCameraOrigin = ResolvedView.TranslatedWorldCameraOrigin;
+	HairViewInfo.ViewForward = ResolvedView.ViewForward;
+	HairViewInfo.RadiusAtDepth1 = lerp(HairRenderInfo.RadiusAtDepth1Primary, HairRenderInfo.RadiusAtDepth1Velocity, VelocityScale);
+	HairViewInfo.bIsOrthoView = false;
+	HairViewInfo.RadiusAtDepth1 = lerp(HairRenderInfo.RadiusAtDepth1Primary, HairRenderInfo.RadiusAtDepth1Velocity, VelocityScale);
+
+
+	const float3 WP0  = mul(float4(CP0.Position, 1), HairStrandsVF_LocalToWorldPrimitiveTransform).xyz;
+	float4 NDC0 = mul(float4(WP0, 1), View.WorldToClip);
+
+	const float3 WP1  = mul(float4(CP1.Position, 1), HairStrandsVF_LocalToWorldPrimitiveTransform).xyz;
+	const float4 NDC1 = mul(float4(WP1, 1), View.WorldToClip);
+	
+	float3 SP0 = NDCToPixelCoord(NDC0);
+	float3 SP1 = NDCToPixelCoord(NDC1);
+
+	// TODO: do actual clipping
+	const bool bInClip0 = IsInScreen(SP0.xy);
+	const bool bInClip1 = IsInScreen(SP1.xy);
+	if (!bInClip0 || !bInClip1)
+	{
+		return;
+	}
+
+	const bool bIsVisible0 = IsVisible(SP0);
+	const bool bIsVisible1 = IsVisible(SP1);
+	if (!bIsVisible0 || !bIsVisible1)
+	{
+		return;
+	}
+
+	const float Coverage0 = ToCoverage(CP0, NDC0, HairRenderInfo.RadiusAtDepth1Primary);
+	const float Coverage1 = ToCoverage(CP1, NDC1, HairRenderInfo.RadiusAtDepth1Primary);
+
+	// Conservative line raster (Wu's Algo)
+	// https://en.wikipedia.org/wiki/Xiaolin_Wu%27s_line_algorithm
+	const bool bIsSteep = abs(SP1.y - SP0.y) > abs(SP1.x - SP0.x);
+	if (bIsSteep)
+	{
+		Swap(SP0.x, SP0.y);
+		Swap(SP1.x, SP1.y);
+	}
+	if (SP0.x > SP1.x)
+	{
+		Swap(SP0, SP1);
+	}
+
+	const float3 D = SP1 - SP0;
+	const float GradientY = abs(D.x) < 0.00001f ? 1.0f : D.y / D.x;
+	const float GradientZ = abs(D.x) < 0.00001f ? 1.0f : D.z / D.x;
+
+	float DeltaY = 0;
+	float DeltaZ = 0;
+
+	// First point
+	int2 PX0;
+	{
+		float2 End;
+		End.x = round(SP0.x);
+		End.y = SP0.y + GradientY * (End.x - SP0.x);
+
+		const float XGap = 1.f - frac(SP0.x + 0.5f);
+
+		PX0.x = End.x;
+		PX0.y = floor(End.y);
+
+		const float FracY = frac(End.y);
+		const bool bLower = FracY < 0.5f;
+		int2 Coord = PX0 + (bLower ? int2(0, 0) : int2(0, 1));
+		Coord = bIsSteep ? Coord.yx : Coord.xy;
+		const float Ratio = bLower ? 1 - FracY : FracY;
+		Rasterize(Coord, Ratio * XGap, SP0.z, Coverage0);
+
+		DeltaY = End.y + GradientY; // first y-intersection for the main loop
+		DeltaZ = SP1.z + GradientY; // This is not perspective correct. Interpolate linear depth instead
+	}
+
+	// Second point
+	int2 PX1;
+	{
+		float2 End;
+		End.x = floor(SP1.x + 0.5); // Round
+		End.y = SP1.y + GradientY * (End.x - SP1.x);
+		const float XGap = 1 - frac(SP1.x + 0.5);
+
+		PX1.x = End.x;
+		PX1.y = floor(End.y);
+
+		const float FracY = frac(End.y);
+		const bool bLower = FracY < 0.5f;
+		int2 Coord = PX1 + (bLower ? int2(0, 0) : int2(0, 1));
+		Coord = bIsSteep ? Coord.yx : Coord.xy;
+		const float Ratio = bLower ? 1 - FracY : FracY;
+		Rasterize(Coord, Ratio * XGap, SP1.z, Coverage1);
+	}
+
+	// Inner points
+	const uint MaxIteration = 64;	
+	const float PixelCount = PX1.x - (PX0.x + 1);
+	for (int i = PX0.x + 1, j = 0; i < PX1.x && j < MaxIteration; ++i, ++j)
+	{
+		const float S = float(j) / PixelCount;
+		const float FracY = frac(DeltaY);
+		const bool bLower = FracY < 0.5;
+		const float Ratio = bLower ? 1 - FracY : FracY;
+		const float Coverage = lerp(Coverage0, Coverage1, S);
+
+		int2 Coord = int2(i, floor(DeltaY)) + (bLower ? int2(0, 0) : int2(0, 1));
+		Coord = bIsSteep ? Coord.yx : Coord.xy;
+		Rasterize(Coord, Ratio, DeltaZ, Coverage);
+
+		DeltaY += GradientY;
+		DeltaZ += GradientZ;
+	}
+}
\ No newline at end of file
diff --git a/Engine/Source/Runtime/Renderer/Private/HairStrands/HairStrandsDebug.cpp b/Engine/Source/Runtime/Renderer/Private/HairStrands/HairStrandsDebug.cpp
--- a/Engine/Source/Runtime/Renderer/Private/HairStrands/HairStrandsDebug.cpp
+++ b/Engine/Source/Runtime/Renderer/Private/HairStrands/HairStrandsDebug.cpp
@@ -243,6 +243,7 @@ class FHairDebugPrintCS : public FGlobalShader
 		SHADER_PARAMETER(uint32, HairMacroGroupCount)
 		SHADER_PARAMETER(uint32, MaxSampleCount)
 		SHADER_PARAMETER_RDG_TEXTURE(Texture2D, HairCountTexture)
+		SHADER_PARAMETER_RDG_TEXTURE(Texture2D, HairCountUintTexture)
 		SHADER_PARAMETER_RDG_TEXTURE(Texture2D, CategorizationTexture)
 		SHADER_PARAMETER_RDG_TEXTURE(Texture2D, HairVisibilityNodeOffsetAndCount)
 		SHADER_PARAMETER_RDG_BUFFER_SRV(StructuredBuffer, HairVisibilityNodeData)
@@ -276,6 +277,7 @@ static void AddDebugHairPrintPass(
 	if (!InCategorizationTexture || !InNodeIndex || !InNodeData || !InDepthStencilTexture) return;
 
 	FRDGTextureRef ViewHairCountTexture = InViewHairCountTexture ? GraphBuilder.RegisterExternalTexture(InViewHairCountTexture, TEXT("ViewHairCountTexture")) : GSystemTextures.GetBlackDummy(GraphBuilder);
+	FRDGTextureRef ViewHairCountUintTexture = VisibilityData.ViewHairCountUintTexture ? GraphBuilder.RegisterExternalTexture(VisibilityData.ViewHairCountUintTexture, TEXT("ViewHairCountUintTexture")) : GSystemTextures.GetBlackDummy(GraphBuilder);
 	FRDGTextureRef CategorizationTexture = InCategorizationTexture ? GraphBuilder.RegisterExternalTexture(InCategorizationTexture, TEXT("CategorizationTexture")) : nullptr;
 	FRDGTextureRef NodeIndex = InNodeIndex ? GraphBuilder.RegisterExternalTexture(InNodeIndex, TEXT("NodeIndex")) : nullptr;
 	FRDGBufferRef  NodeData  = InNodeData ? GraphBuilder.RegisterExternalBuffer(InNodeData, TEXT("NodeData")) : nullptr;
@@ -291,6 +293,7 @@ static void AddDebugHairPrintPass(
 	Parameters->FastResolveMask = STENCIL_TEMPORAL_RESPONSIVE_AA_MASK;
 	Parameters->CategorizationTexture = CategorizationTexture;
 	Parameters->HairCountTexture = ViewHairCountTexture;
+	Parameters->HairCountUintTexture = ViewHairCountUintTexture;
 	Parameters->HairVisibilityNodeData = GraphBuilder.CreateSRV(NodeData);
 	Parameters->HairVisibilityNodeOffsetAndCount = NodeIndex;
 	Parameters->DepthStencilTexture = InDepthStencilTexture;
@@ -325,6 +328,8 @@ class FHairDebugPS : public FGlobalShader
 		SHADER_PARAMETER(int32, SampleIndex)
 		SHADER_PARAMETER(uint32, TileSize)
 		SHADER_PARAMETER(uint32, MaxSampleCount)
+		SHADER_PARAMETER_RDG_TEXTURE(Texture2D, HairCountTexture)
+		SHADER_PARAMETER_RDG_TEXTURE(Texture2D, HairCountUintTexture)
 		SHADER_PARAMETER_RDG_TEXTURE(Texture2D, CategorizationTexture)
 		SHADER_PARAMETER_RDG_TEXTURE(Texture2D, NodeIndex)
 		SHADER_PARAMETER_RDG_TEXTURE(Texture2D, TileIndexTexture)
@@ -374,6 +379,8 @@ static void AddDebugHairPass(
 	FRDGBufferRef  NodeData = InNodeData ? GraphBuilder.RegisterExternalBuffer(InNodeData, TEXT("NodeData")) : nullptr;
 	FRDGTextureRef TileIndexTexture = InTileIndexTexture ? GraphBuilder.RegisterExternalTexture(InTileIndexTexture, TEXT("TileIndexTexture")) : nullptr;
 
+	FRDGTextureRef HairCountTexture = VisibilityData.ViewHairCountTexture ? GraphBuilder.RegisterExternalTexture(VisibilityData.ViewHairCountTexture, TEXT("HairCount")) : GSystemTextures.GetBlackDummy(GraphBuilder);
+	FRDGTextureRef HairCountUintTexture = VisibilityData.ViewHairCountUintTexture ? GraphBuilder.RegisterExternalTexture(VisibilityData.ViewHairCountUintTexture, TEXT("HairCountUint")) : GSystemTextures.GetBlackDummy(GraphBuilder); 
 	const FIntRect Viewport = View->ViewRect;
 	const FIntPoint Resolution(Viewport.Width(), Viewport.Height());
 
@@ -397,6 +404,8 @@ static void AddDebugHairPass(
 	Parameters->OutputResolution = Resolution;
 	Parameters->FastResolveMask = STENCIL_TEMPORAL_RESPONSIVE_AA_MASK;
 	Parameters->CategorizationTexture = CategorizationTexture;
+	Parameters->HairCountTexture = HairCountTexture;
+	Parameters->HairCountUintTexture = HairCountUintTexture;
 	Parameters->NodeIndex = NodeIndex;
 	Parameters->NodeData = GraphBuilder.CreateSRV(NodeData);
 	Parameters->TileIndexTexture = TileIndexTexture;
diff --git a/Engine/Source/Runtime/Renderer/Private/HairStrands/HairStrandsVisibility.cpp b/Engine/Source/Runtime/Renderer/Private/HairStrands/HairStrandsVisibility.cpp
--- a/Engine/Source/Runtime/Renderer/Private/HairStrands/HairStrandsVisibility.cpp
+++ b/Engine/Source/Runtime/Renderer/Private/HairStrands/HairStrandsVisibility.cpp
@@ -56,6 +56,11 @@ static FAutoConsoleVariableRef CVarHairStrandsVisibilityMaterialPass(TEXT("r.Hai
 static float GHairStrandsViewHairCountDepthDistanceThreshold = 30.f;
 static FAutoConsoleVariableRef CVarHairStrandsViewHairCountDepthDistanceThreshold(TEXT("r.HairStrands.Visibility.HairCount.DistanceThreshold"), GHairStrandsViewHairCountDepthDistanceThreshold, TEXT("Distance threshold defining if opaque depth get injected into the 'view-hair-count' buffer."));
 
+static int32 GHairStrandsVisibilityComputeRaster = 0;
+static int32 GHairStrandsVisibilityComputeRasterMaxPixelCount = 8;
+static FAutoConsoleVariableRef CVarHairStrandsVisibilityComputeRaster(TEXT("r.HairStrands.Visibility.ComputeRaster"), GHairStrandsVisibilityComputeRaster, TEXT("Define the maximal length rasterize in compute."));
+static FAutoConsoleVariableRef CVarHairStrandsVisibilityComputeRasterMaxPixelCount(TEXT("r.HairStrands.Visibility.ComputeRaster.MaxPixelCount"), GHairStrandsVisibilityComputeRasterMaxPixelCount, TEXT("Define the maximal length rasterize in compute."));
+
 /////////////////////////////////////////////////////////////////////////////////////////
 
 namespace HairStrandsVisibilityInternal
@@ -1867,6 +1872,7 @@ struct FHairPrimaryTransmittance
 {
 	FRDGTextureRef TransmittanceTexture = nullptr;
 	FRDGTextureRef HairCountTexture = nullptr;
+	FRDGTextureRef HairCountTextureUint = nullptr;
 };
 
 static FHairPrimaryTransmittance AddHairViewTransmittancePass(
@@ -2092,6 +2098,115 @@ static void AddHairVisibilityColorAndDepthPatchPass(
 	}
 }
 
+///////////////////////////////////////////////////////////////////////////////////////////////////
+
+class FVisiblityRasterComputeCS : public FGlobalShader
+{
+	DECLARE_GLOBAL_SHADER(FVisiblityRasterComputeCS);
+	SHADER_USE_PARAMETER_STRUCT(FVisiblityRasterComputeCS, FGlobalShader);
+
+	BEGIN_SHADER_PARAMETER_STRUCT(FParameters, )
+		SHADER_PARAMETER(uint32, MacroGroupId)
+		SHADER_PARAMETER(uint32, DispatchCountX)
+		SHADER_PARAMETER(uint32, MaxRasterCount)
+		SHADER_PARAMETER(uint32, FrameIdMod8)
+		SHADER_PARAMETER(FIntPoint, OutputResolution)
+		SHADER_PARAMETER(float, HairStrandsVF_Density)
+		SHADER_PARAMETER(float, HairStrandsVF_Radius)
+		SHADER_PARAMETER(float, HairStrandsVF_Length)
+		SHADER_PARAMETER(uint32, HairStrandsVF_bUseStableRasterization)
+		SHADER_PARAMETER(FVector, HairStrandsVF_PositionOffset)
+		SHADER_PARAMETER(uint32, HairStrandsVF_VertexCount)
+		SHADER_PARAMETER(FMatrix, HairStrandsVF_LocalToWorldPrimitiveTransform)
+		SHADER_PARAMETER_SRV(Buffer, HairStrandsVF_PositionBuffer)
+		SHADER_PARAMETER_RDG_TEXTURE(Texture2D, SceneDepthTexture)
+		SHADER_PARAMETER_RDG_TEXTURE_UAV(RWTexture2D, OutTexture)
+		SHADER_PARAMETER_STRUCT_REF(FViewUniformShaderParameters, ViewUniformBuffer)
+		END_SHADER_PARAMETER_STRUCT()
+
+public:
+	static bool ShouldCompilePermutation(const FGlobalShaderPermutationParameters& Parameters) { return IsHairStrandsSupported(Parameters.Platform); }
+	static void ModifyCompilationEnvironment(const FGlobalShaderPermutationParameters& Parameters, FShaderCompilerEnvironment& OutEnvironment)
+	{
+		FGlobalShader::ModifyCompilationEnvironment(Parameters, OutEnvironment);
+		OutEnvironment.SetDefine(TEXT("SHADER_RASTERCOMPUTE"), 1);
+	}
+};
+
+IMPLEMENT_GLOBAL_SHADER(FVisiblityRasterComputeCS, "/Engine/Private/HairStrands/HairStrandsVisibilityRasterCompute.usf", "MainCS", SF_Compute);
+
+static FRDGTextureRef AddVisibilityComputeRasterPass(
+	FRDGBuilder& GraphBuilder,
+	const FViewInfo* ViewInfo,
+	const FHairStrandsMacroGroupDatas& MacroGroupDatas,
+	const FIntPoint& Resolution,
+	const FRDGTextureRef SceneDepthTexture)
+{
+	FRDGTextureDesc Desc;
+	Desc.Extent.X = Resolution.X;
+	Desc.Extent.Y = Resolution.Y;
+	Desc.Depth = 0;
+	Desc.Format = PF_R32_UINT;
+	Desc.NumMips = 1;
+	Desc.NumSamples = 1;
+	Desc.Flags = TexCreate_None;
+	Desc.TargetableFlags = TexCreate_UAV | TexCreate_ShaderResource | TexCreate_RenderTargetable;
+	Desc.bForceSharedTargetAndShaderResource = true;
+	Desc.ClearValue = FClearValueBinding(FLinearColor(1.0f, 1.0f, 1.0f, 1.0f)); // Clear to transmittance 1
+	FRDGTextureRef  OutTexture = GraphBuilder.CreateTexture(Desc, TEXT("HairViewTransmittanceTexture"));
+	FRDGTextureUAVRef OutTextureUAV = GraphBuilder.CreateUAV(OutTexture);
+
+	uint32 ClearValues[4] = { 0,0,0,0 };
+	AddClearUAVPass(GraphBuilder, OutTextureUAV, ClearValues);
+
+	// Create and set the uniform buffer
+	const bool bEnableMSAA = false;
+	TUniformBufferRef<FViewUniformShaderParameters> ViewUniformShaderParameters;
+	SetUpViewHairRenderInfo(*ViewInfo, bEnableMSAA, ViewInfo->CachedViewUniformShaderParameters->HairRenderInfo, ViewInfo->CachedViewUniformShaderParameters->HairRenderInfoBits);
+	ViewUniformShaderParameters = TUniformBufferRef<FViewUniformShaderParameters>::CreateUniformBufferImmediate(*ViewInfo->CachedViewUniformShaderParameters, UniformBuffer_SingleFrame);
+
+	const uint32 FrameIdMode8 = ViewInfo && ViewInfo->ViewState ? (ViewInfo->ViewState->GetFrameIndex() % 8) : 0;
+	const uint32 GroupSize = 32;
+	const uint32 DispatchCountX = 64;
+	TShaderMapRef<FVisiblityRasterComputeCS> ComputeShader(ViewInfo->ShaderMap);
+
+	for (const FHairStrandsMacroGroupData& MacroGroup : MacroGroupDatas.Datas)
+	{
+		const FHairStrandsMacroGroupData::TPrimitiveInfos& PrimitiveSceneInfos = MacroGroup.PrimitivesInfos;
+
+		for (const FHairStrandsMacroGroupData::PrimitiveInfo& PrimitiveInfo : PrimitiveSceneInfos)
+		{
+			FVisiblityRasterComputeCS::FParameters* PassParameters = GraphBuilder.AllocParameters<FVisiblityRasterComputeCS::FParameters>();
+			PassParameters->OutputResolution = Resolution;
+			PassParameters->MacroGroupId = MacroGroup.MacroGroupId;
+			PassParameters->DispatchCountX = DispatchCountX;
+			PassParameters->MaxRasterCount = FMath::Clamp(GHairStrandsVisibilityComputeRasterMaxPixelCount, 1, 256);
+			PassParameters->FrameIdMod8 = FrameIdMode8;
+			PassParameters->ViewUniformBuffer = ViewUniformShaderParameters;
+			PassParameters->SceneDepthTexture = SceneDepthTexture;
+			PassParameters->OutTexture = OutTextureUAV;
+
+			check(PrimitiveInfo.MeshBatchAndRelevance.Mesh && PrimitiveInfo.MeshBatchAndRelevance.Mesh->Elements.Num() > 0);
+			const FHairGroupPublicData* HairGroupPublicData = reinterpret_cast<const FHairGroupPublicData*>(PrimitiveInfo.MeshBatchAndRelevance.Mesh->Elements[0].VertexFactoryUserData);
+			const FHairGroupPublicData::VertexFactoryInput& VFInput = HairGroupPublicData->VFInput;
+			PassParameters->HairStrandsVF_PositionBuffer = VFInput.HairPositionBuffer;
+			PassParameters->HairStrandsVF_PositionOffset = VFInput.HairPositionOffset;
+			PassParameters->HairStrandsVF_VertexCount = VFInput.VertexCount;
+			PassParameters->HairStrandsVF_Radius = VFInput.HairRadius;
+			PassParameters->HairStrandsVF_Length = VFInput.HairLength;
+			PassParameters->HairStrandsVF_bUseStableRasterization = VFInput.bUseStableRasterization ? 1 : 0;
+			PassParameters->HairStrandsVF_Density = VFInput.HairDensity;
+			PassParameters->HairStrandsVF_LocalToWorldPrimitiveTransform = VFInput.LocalToWorldTransform.ToMatrixWithScale();
+
+			const uint32 DispatchCountY = FMath::CeilToInt(PassParameters->HairStrandsVF_VertexCount / float(GroupSize * DispatchCountX));
+			const FIntVector DispatchCount(DispatchCountX, DispatchCountY, 1);
+			FComputeShaderUtils::AddPass(GraphBuilder, RDG_EVENT_NAME("HairStrandsVisibilityComputeRaster"), *ComputeShader, PassParameters, DispatchCount);
+		}
+	}
+
+	return OutTexture; 
+}
+
 ///////////////////////////////////////////////////////////////////////////////////////////////////
 bool GetHairStrandsSkyLightingEnable();
 
@@ -2153,6 +2268,18 @@ FHairStrandsVisibilityViews RenderHairStrandsVisibilityBuffer(
 					Resolution,
 					bOutputHairCount,
 					SceneDepthTexture);
+
+				const bool bUseRasterCompute = GHairStrandsVisibilityComputeRaster > 0;
+				if (bUseRasterCompute)
+				{
+					ViewTransmittance.HairCountTextureUint = AddVisibilityComputeRasterPass(
+						GraphBuilder,
+						&View,
+						MacroGroupDatas,
+						Resolution,
+						SceneDepthTexture);
+				}
+
 			}
 
 			FRDGTextureRef CategorizationTexture = nullptr;
@@ -2275,6 +2402,11 @@ FHairStrandsVisibilityViews RenderHairStrandsVisibilityBuffer(
 					GraphBuilder.QueueTextureExtraction(ViewTransmittance.HairCountTexture, &VisibilityData.ViewHairCountTexture);
 				}
 
+				if (ViewTransmittance.HairCountTextureUint)
+				{
+					GraphBuilder.QueueTextureExtraction(ViewTransmittance.HairCountTextureUint, &VisibilityData.ViewHairCountUintTexture);
+				}
+
 				// For fully covered pixels, write: 
 				// * black color into the scene color
 				// * closest depth
diff --git a/Engine/Source/Runtime/Renderer/Private/HairStrands/HairStrandsVisibility.h b/Engine/Source/Runtime/Renderer/Private/HairStrands/HairStrandsVisibility.h
--- a/Engine/Source/Runtime/Renderer/Private/HairStrands/HairStrandsVisibility.h
+++ b/Engine/Source/Runtime/Renderer/Private/HairStrands/HairStrandsVisibility.h
@@ -23,6 +23,7 @@ struct FHairStrandsVisibilityData
 	TRefCountPtr<IPooledRenderTarget> VelocityTexture;
 	TRefCountPtr<IPooledRenderTarget> CategorizationTexture;
 	TRefCountPtr<IPooledRenderTarget> ViewHairCountTexture;
+	TRefCountPtr<IPooledRenderTarget> ViewHairCountUintTexture;
 	TRefCountPtr<IPooledRenderTarget> LightChannelMaskTexture;
 
 	TRefCountPtr<IPooledRenderTarget> PPLLNodeCounterTexture;
