Add compute raster depth output.

diff --git a/Engine/Shaders/Private/HairStrands/HairStrandsVisibilityRasterCompute.usf b/Engine/Shaders/Private/HairStrands/HairStrandsVisibilityRasterCompute.usf
--- a/Engine/Shaders/Private/HairStrands/HairStrandsVisibilityRasterCompute.usf
+++ b/Engine/Shaders/Private/HairStrands/HairStrandsVisibilityRasterCompute.usf
@@ -20,7 +20,8 @@ float4x4		HairStrandsVF_LocalToWorldPrimitiveTransform;
 Buffer<uint4>	HairStrandsVF_PositionBuffer;
 
 Texture2D<float> SceneDepthTexture;
-RWTexture2D<uint> OutTexture;
+RWTexture2D<uint> OutHairCountTexture;
+RWTexture2D<uint> OutDepthTexture;
 
 ///////////////////////////////////////////////////////////////////////////
 // Helping functions 
@@ -73,7 +74,8 @@ void Rasterize(int2 Coords, float Ratio, float Depth, float Coverage)
 		const float CoveragePrecision = 1000.f;
 		const uint IntCoverage = Ratio * Coverage * CoveragePrecision;
 
-		InterlockedAdd(OutTexture[uint2(Coords)], IntCoverage);
+		InterlockedAdd(OutHairCountTexture[uint2(Coords)], IntCoverage);
+		InterlockedMax(OutDepthTexture[uint2(Coords)], asuint(Depth));
 	}
 }
 
@@ -143,7 +145,7 @@ void MainCS(uint3 DispatchThreadID : SV_DispatchThreadID)
 
 
 	const float3 WP0  = mul(float4(CP0.Position, 1), HairStrandsVF_LocalToWorldPrimitiveTransform).xyz;
-	float4 NDC0 = mul(float4(WP0, 1), View.WorldToClip);
+	const float4 NDC0 = mul(float4(WP0, 1), View.WorldToClip);
 
 	const float3 WP1  = mul(float4(CP1.Position, 1), HairStrandsVF_LocalToWorldPrimitiveTransform).xyz;
 	const float4 NDC1 = mul(float4(WP1, 1), View.WorldToClip);
@@ -209,7 +211,7 @@ void MainCS(uint3 DispatchThreadID : SV_DispatchThreadID)
 		Rasterize(Coord, Ratio * XGap, SP0.z, Coverage0);
 
 		DeltaY = End.y + GradientY; // first y-intersection for the main loop
-		DeltaZ = SP1.z + GradientY; // This is not perspective correct. Interpolate linear depth instead
+		DeltaZ = SP1.z + GradientZ; // This is not perspective correct. Interpolate linear depth instead
 	}
 
 	// Second point
@@ -233,7 +235,7 @@ void MainCS(uint3 DispatchThreadID : SV_DispatchThreadID)
 
 	// Inner points
 	const uint MaxIteration = 64;	
-	const float PixelCount = PX1.x - (PX0.x + 1);
+	const float PixelCount = abs(PX1.x - (PX0.x + 1));
 	for (int i = PX0.x + 1, j = 0; i < PX1.x && j < MaxIteration; ++i, ++j)
 	{
 		const float S = float(j) / PixelCount;
@@ -242,9 +244,13 @@ void MainCS(uint3 DispatchThreadID : SV_DispatchThreadID)
 		const float Ratio = bLower ? 1 - FracY : FracY;
 		const float Coverage = lerp(Coverage0, Coverage1, S);
 
+		const float DC_Z = lerp(NDC0.z, NDC1.z, S);
+		const float DC_W = lerp(NDC0.w, NDC1.w, S);
+		const float Depth = DC_Z / DC_W;
+
 		int2 Coord = int2(i, floor(DeltaY)) + (bLower ? int2(0, 0) : int2(0, 1));
 		Coord = bIsSteep ? Coord.yx : Coord.xy;
-		Rasterize(Coord, Ratio, DeltaZ, Coverage);
+		Rasterize(Coord, Ratio, Depth /*DeltaZ*/, Coverage);
 
 		DeltaY += GradientY;
 		DeltaZ += GradientZ;
diff --git a/Engine/Source/Runtime/Renderer/Private/HairStrands/HairStrandsVisibility.cpp b/Engine/Source/Runtime/Renderer/Private/HairStrands/HairStrandsVisibility.cpp
--- a/Engine/Source/Runtime/Renderer/Private/HairStrands/HairStrandsVisibility.cpp
+++ b/Engine/Source/Runtime/Renderer/Private/HairStrands/HairStrandsVisibility.cpp
@@ -1872,7 +1872,9 @@ struct FHairPrimaryTransmittance
 {
 	FRDGTextureRef TransmittanceTexture = nullptr;
 	FRDGTextureRef HairCountTexture = nullptr;
+
 	FRDGTextureRef HairCountTextureUint = nullptr;
+	FRDGTextureRef DepthTextureUint = nullptr;
 };
 
 static FHairPrimaryTransmittance AddHairViewTransmittancePass(
@@ -2120,7 +2122,8 @@ class FVisiblityRasterComputeCS : public FGlobalShader
 		SHADER_PARAMETER(FMatrix, HairStrandsVF_LocalToWorldPrimitiveTransform)
 		SHADER_PARAMETER_SRV(Buffer, HairStrandsVF_PositionBuffer)
 		SHADER_PARAMETER_RDG_TEXTURE(Texture2D, SceneDepthTexture)
-		SHADER_PARAMETER_RDG_TEXTURE_UAV(RWTexture2D, OutTexture)
+		SHADER_PARAMETER_RDG_TEXTURE_UAV(RWTexture2D, OutHairCountTexture)
+		SHADER_PARAMETER_RDG_TEXTURE_UAV(RWTexture2D, OutDepthTexture)
 		SHADER_PARAMETER_STRUCT_REF(FViewUniformShaderParameters, ViewUniformBuffer)
 		END_SHADER_PARAMETER_STRUCT()
 
@@ -2135,29 +2138,56 @@ public:
 
 IMPLEMENT_GLOBAL_SHADER(FVisiblityRasterComputeCS, "/Engine/Private/HairStrands/HairStrandsVisibilityRasterCompute.usf", "MainCS", SF_Compute);
 
-static FRDGTextureRef AddVisibilityComputeRasterPass(
+struct FRasterComputeOutput
+{
+	FRDGTextureRef HairCountTexture;
+	FRDGTextureRef DepthTexture;
+};
+
+static FRasterComputeOutput AddVisibilityComputeRasterPass(
 	FRDGBuilder& GraphBuilder,
 	const FViewInfo* ViewInfo,
 	const FHairStrandsMacroGroupDatas& MacroGroupDatas,
 	const FIntPoint& Resolution,
 	const FRDGTextureRef SceneDepthTexture)
 {
-	FRDGTextureDesc Desc;
-	Desc.Extent.X = Resolution.X;
-	Desc.Extent.Y = Resolution.Y;
-	Desc.Depth = 0;
-	Desc.Format = PF_R32_UINT;
-	Desc.NumMips = 1;
-	Desc.NumSamples = 1;
-	Desc.Flags = TexCreate_None;
-	Desc.TargetableFlags = TexCreate_UAV | TexCreate_ShaderResource | TexCreate_RenderTargetable;
-	Desc.bForceSharedTargetAndShaderResource = true;
-	Desc.ClearValue = FClearValueBinding(FLinearColor(1.0f, 1.0f, 1.0f, 1.0f)); // Clear to transmittance 1
-	FRDGTextureRef  OutTexture = GraphBuilder.CreateTexture(Desc, TEXT("HairViewTransmittanceTexture"));
-	FRDGTextureUAVRef OutTextureUAV = GraphBuilder.CreateUAV(OutTexture);
+	FRasterComputeOutput Out;
+
+	{
+		FRDGTextureDesc Desc;
+		Desc.Extent.X = Resolution.X;
+		Desc.Extent.Y = Resolution.Y;
+		Desc.Depth = 0;
+		Desc.Format = PF_R32_UINT;
+		Desc.NumMips = 1;
+		Desc.NumSamples = 1;
+		Desc.Flags = TexCreate_None;
+		Desc.TargetableFlags = TexCreate_UAV | TexCreate_ShaderResource | TexCreate_RenderTargetable;
+		Desc.bForceSharedTargetAndShaderResource = true;
+		Desc.ClearValue = FClearValueBinding(FLinearColor(1.0f, 1.0f, 1.0f, 1.0f)); // Clear to transmittance 1
+		Out.HairCountTexture = GraphBuilder.CreateTexture(Desc, TEXT("HairViewTransmittanceTexture"));
+	}
+	FRDGTextureUAVRef HairCountTextureUAV = GraphBuilder.CreateUAV(Out.HairCountTexture);
+
+	{
+		FRDGTextureDesc Desc;
+		Desc.Extent.X = Resolution.X;
+		Desc.Extent.Y = Resolution.Y;
+		Desc.Depth = 0;
+		Desc.Format = PF_R32_UINT;
+		Desc.NumMips = 1;
+		Desc.NumSamples = 1;
+		Desc.Flags = TexCreate_None;
+		Desc.TargetableFlags = TexCreate_UAV | TexCreate_ShaderResource | TexCreate_RenderTargetable;
+		Desc.bForceSharedTargetAndShaderResource = true;
+		Desc.ClearValue = FClearValueBinding(FLinearColor(1.0f, 1.0f, 1.0f, 1.0f)); // Clear to transmittance 1
+		Out.DepthTexture = GraphBuilder.CreateTexture(Desc, TEXT("HairDepthTexture"));
+	}
+	FRDGTextureUAVRef DepthTextureUAV = GraphBuilder.CreateUAV(Out.DepthTexture);
 
 	uint32 ClearValues[4] = { 0,0,0,0 };
-	AddClearUAVPass(GraphBuilder, OutTextureUAV, ClearValues);
+	AddClearUAVPass(GraphBuilder, HairCountTextureUAV, ClearValues);
+	AddClearUAVPass(GraphBuilder, DepthTextureUAV, ClearValues);
 
 	// Create and set the uniform buffer
 	const bool bEnableMSAA = false;
@@ -2184,7 +2214,8 @@ static FRDGTextureRef AddVisibilityComputeRasterPass(
 			PassParameters->FrameIdMod8 = FrameIdMode8;
 			PassParameters->ViewUniformBuffer = ViewUniformShaderParameters;
 			PassParameters->SceneDepthTexture = SceneDepthTexture;
-			PassParameters->OutTexture = OutTextureUAV;
+			PassParameters->OutHairCountTexture = HairCountTextureUAV;
+			PassParameters->OutDepthTexture = DepthTextureUAV;
 
 			check(PrimitiveInfo.MeshBatchAndRelevance.Mesh && PrimitiveInfo.MeshBatchAndRelevance.Mesh->Elements.Num() > 0);
 			const FHairGroupPublicData* HairGroupPublicData = reinterpret_cast<const FHairGroupPublicData*>(PrimitiveInfo.MeshBatchAndRelevance.Mesh->Elements[0].VertexFactoryUserData);
@@ -2204,7 +2235,7 @@ static FRDGTextureRef AddVisibilityComputeRasterPass(
 		}
 	}
 
-	return OutTexture; 
+	return Out;
 }
 
 ///////////////////////////////////////////////////////////////////////////////////////////////////
@@ -2272,12 +2303,15 @@ FHairStrandsVisibilityViews RenderHairStrandsVisibilityBuffer(
 				const bool bUseRasterCompute = GHairStrandsVisibilityComputeRaster > 0;
 				if (bUseRasterCompute)
 				{
-					ViewTransmittance.HairCountTextureUint = AddVisibilityComputeRasterPass(
+					FRasterComputeOutput RasterOutput = AddVisibilityComputeRasterPass(
 						GraphBuilder,
 						&View,
 						MacroGroupDatas,
 						Resolution,
 						SceneDepthTexture);
+
+					ViewTransmittance.HairCountTextureUint = RasterOutput.HairCountTexture;
+					ViewTransmittance.DepthTextureUint = RasterOutput.DepthTexture;
 				}
 
 			}
@@ -2407,6 +2441,11 @@ FHairStrandsVisibilityViews RenderHairStrandsVisibilityBuffer(
 					GraphBuilder.QueueTextureExtraction(ViewTransmittance.HairCountTextureUint, &VisibilityData.ViewHairCountUintTexture);
 				}
 
+				if (ViewTransmittance.DepthTextureUint)
+				{
+					GraphBuilder.QueueTextureExtraction(ViewTransmittance.DepthTextureUint, &VisibilityData.DepthTextureUint);
+				}
+
 				// For fully covered pixels, write: 
 				// * black color into the scene color
 				// * closest depth
diff --git a/Engine/Source/Runtime/Renderer/Private/HairStrands/HairStrandsVisibility.h b/Engine/Source/Runtime/Renderer/Private/HairStrands/HairStrandsVisibility.h
--- a/Engine/Source/Runtime/Renderer/Private/HairStrands/HairStrandsVisibility.h
+++ b/Engine/Source/Runtime/Renderer/Private/HairStrands/HairStrandsVisibility.h
@@ -24,6 +24,7 @@ struct FHairStrandsVisibilityData
 	TRefCountPtr<IPooledRenderTarget> CategorizationTexture;
 	TRefCountPtr<IPooledRenderTarget> ViewHairCountTexture;
 	TRefCountPtr<IPooledRenderTarget> ViewHairCountUintTexture;
+	TRefCountPtr<IPooledRenderTarget> DepthTextureUint;
 	TRefCountPtr<IPooledRenderTarget> LightChannelMaskTexture;
 
 	TRefCountPtr<IPooledRenderTarget> PPLLNodeCounterTexture;
