Add visibility texture output and super resolution.

diff --git a/Engine/Shaders/Private/HairStrands/HairStrandsVisibilityRasterCompute.usf b/Engine/Shaders/Private/HairStrands/HairStrandsVisibilityRasterCompute.usf
--- a/Engine/Shaders/Private/HairStrands/HairStrandsVisibilityRasterCompute.usf
+++ b/Engine/Shaders/Private/HairStrands/HairStrandsVisibilityRasterCompute.usf
@@ -9,7 +9,9 @@ uint			DispatchCountX;
 uint			MaxRasterCount;
 uint			FrameIdMod8;
 uint			MacroGroupId;
+uint			ResolutionMultiplier;
 int2			OutputResolution;
+uint			HairMaterialId;
 float			HairStrandsVF_Radius;
 float			HairStrandsVF_Length;
 float			HairStrandsVF_Density;
@@ -22,6 +24,11 @@ Buffer<uint4>	HairStrandsVF_PositionBuffer;
 Texture2D<float> SceneDepthTexture;
 RWTexture2D<uint> OutHairCountTexture;
 RWTexture2D<uint> OutDepthTexture;
+RWTexture2D<uint2> OutVisibilityTexture;
+
+// Add permutation for Nvidia/AMD/other
+#define NV_SHADER_EXTN_SLOT u7
+#include "/Engine/Shared/ThirdParty/NVIDIA/nvHLSLExtns.h"
 
 ///////////////////////////////////////////////////////////////////////////
 // Helping functions 
@@ -59,7 +66,8 @@ void Swap(inout float3 V0, inout float3 V1)
 bool IsVisible(float3 SP)
 {
 	const float DeviceZ = SP.z;
-	return DeviceZ > SceneDepthTexture.Load(uint3(SP.xy, 0)); // InverseZ
+	const uint2 SceneCoord = SP.xy / ResolutionMultiplier;
+	return DeviceZ > SceneDepthTexture.Load(uint3(SceneCoord, 0)); // InverseZ
 }
 
 bool IsInScreen(int2 C)
@@ -67,7 +75,7 @@ bool IsInScreen(int2 C)
 	return C.x >= 0 && C.x < OutputResolution.x && C.y >= 0 && C.y < OutputResolution.y;
 }
 
-void Rasterize(int2 Coords, float Ratio, float Depth, float Coverage)
+void Rasterize(int2 Coords, float Ratio, float Depth, uint PrimitiveId, float Coverage)
 {
 	if (IsInScreen(Coords))
 	{
@@ -75,7 +83,12 @@ void Rasterize(int2 Coords, float Ratio, float Depth, float Coverage)
 		const uint IntCoverage = Ratio * Coverage * CoveragePrecision;
 
 		InterlockedAdd(OutHairCountTexture[uint2(Coords)], IntCoverage);
-		InterlockedMax(OutDepthTexture[uint2(Coords)], asuint(Depth));
+//		InterlockedMax(OutDepthTexture[uint2(Coords)], asuint(Depth));
+
+		uint2 PackedVisibility;
+		PackedVisibility.x = PackHairVisSample(PrimitiveId, HairMaterialId);
+		PackedVisibility.y = asuint(Depth);
+		NvInterlockedMaxUint64(OutVisibilityTexture, Coords, PackedVisibility);
 	}
 }
 
@@ -208,7 +221,7 @@ void MainCS(uint3 DispatchThreadID : SV_DispatchThreadID)
 		int2 Coord = PX0 + (bLower ? int2(0, 0) : int2(0, 1));
 		Coord = bIsSteep ? Coord.yx : Coord.xy;
 		const float Ratio = bLower ? 1 - FracY : FracY;
-		Rasterize(Coord, Ratio * XGap, SP0.z, Coverage0);
+		Rasterize(Coord, Ratio * XGap, SP0.z, VertexIndex, Coverage0);
 
 		DeltaY = End.y + GradientY; // first y-intersection for the main loop
 		DeltaZ = SP1.z + GradientZ; // This is not perspective correct. Interpolate linear depth instead
@@ -230,7 +243,7 @@ void MainCS(uint3 DispatchThreadID : SV_DispatchThreadID)
 		int2 Coord = PX1 + (bLower ? int2(0, 0) : int2(0, 1));
 		Coord = bIsSteep ? Coord.yx : Coord.xy;
 		const float Ratio = bLower ? 1 - FracY : FracY;
-		Rasterize(Coord, Ratio * XGap, SP1.z, Coverage1);
+		Rasterize(Coord, Ratio * XGap, SP1.z, VertexIndex, Coverage1);
 	}
 
 	// Inner points
@@ -250,9 +263,9 @@ void MainCS(uint3 DispatchThreadID : SV_DispatchThreadID)
 
 		int2 Coord = int2(i, floor(DeltaY)) + (bLower ? int2(0, 0) : int2(0, 1));
 		Coord = bIsSteep ? Coord.yx : Coord.xy;
-		Rasterize(Coord, Ratio, Depth /*DeltaZ*/, Coverage);
+		Rasterize(Coord, Ratio, Depth /*DeltaZ*/, VertexIndex, Coverage);
 
 		DeltaY += GradientY;
 		DeltaZ += GradientZ;
 	}
-}
\ No newline at end of file
+}
diff --git a/Engine/Source/Runtime/Renderer/Private/HairStrands/HairStrandsVisibility.cpp b/Engine/Source/Runtime/Renderer/Private/HairStrands/HairStrandsVisibility.cpp
--- a/Engine/Source/Runtime/Renderer/Private/HairStrands/HairStrandsVisibility.cpp
+++ b/Engine/Source/Runtime/Renderer/Private/HairStrands/HairStrandsVisibility.cpp
@@ -58,8 +58,10 @@ static FAutoConsoleVariableRef CVarHairStrandsViewHairCountDepthDistanceThreshol
 
 static int32 GHairStrandsVisibilityComputeRaster = 0;
 static int32 GHairStrandsVisibilityComputeRasterMaxPixelCount = 8;
+static int32 GHairStrandsVisibilityComputeRasterSampleCount = 1;
 static FAutoConsoleVariableRef CVarHairStrandsVisibilityComputeRaster(TEXT("r.HairStrands.Visibility.ComputeRaster"), GHairStrandsVisibilityComputeRaster, TEXT("Define the maximal length rasterize in compute."));
 static FAutoConsoleVariableRef CVarHairStrandsVisibilityComputeRasterMaxPixelCount(TEXT("r.HairStrands.Visibility.ComputeRaster.MaxPixelCount"), GHairStrandsVisibilityComputeRasterMaxPixelCount, TEXT("Define the maximal length rasterize in compute."));
+static FAutoConsoleVariableRef CVarHairStrandsVisibilityComputeRasterSampleCount(TEXT("r.HairStrands.Visibility.ComputeRaster.SampleCount"), GHairStrandsVisibilityComputeRasterSampleCount, TEXT("Define sample count used in rasterize in compute."));
 
 /////////////////////////////////////////////////////////////////////////////////////////
 
@@ -2112,6 +2114,8 @@ class FVisiblityRasterComputeCS : public FGlobalShader
 		SHADER_PARAMETER(uint32, DispatchCountX)
 		SHADER_PARAMETER(uint32, MaxRasterCount)
 		SHADER_PARAMETER(uint32, FrameIdMod8)
+		SHADER_PARAMETER(uint32, HairMaterialId)
+		SHADER_PARAMETER(uint32, ResolutionMultiplier)
 		SHADER_PARAMETER(FIntPoint, OutputResolution)
 		SHADER_PARAMETER(float, HairStrandsVF_Density)
 		SHADER_PARAMETER(float, HairStrandsVF_Radius)
@@ -2124,11 +2128,19 @@ class FVisiblityRasterComputeCS : public FGlobalShader
 		SHADER_PARAMETER_RDG_TEXTURE(Texture2D, SceneDepthTexture)
 		SHADER_PARAMETER_RDG_TEXTURE_UAV(RWTexture2D, OutHairCountTexture)
 		SHADER_PARAMETER_RDG_TEXTURE_UAV(RWTexture2D, OutDepthTexture)
+		SHADER_PARAMETER_RDG_TEXTURE_UAV(RWTexture2D, OutVisibilityTexture)
 		SHADER_PARAMETER_STRUCT_REF(FViewUniformShaderParameters, ViewUniformBuffer)
 		END_SHADER_PARAMETER_STRUCT()
 
 public:
-	static bool ShouldCompilePermutation(const FGlobalShaderPermutationParameters& Parameters) { return IsHairStrandsSupported(Parameters.Platform); }
+	static bool ShouldCompilePermutation(const FGlobalShaderPermutationParameters& Parameters) 
+	{ 
+		// TODO:
+		//if (!FDataDrivenShaderPlatformInfo::GetInfo(Parameters.Platform).bSupportsUInt64ImageAtomics))
+		//	return false;
+
+		return IsHairStrandsSupported(Parameters.Platform); 
+	}
 	static void ModifyCompilationEnvironment(const FGlobalShaderPermutationParameters& Parameters, FShaderCompilerEnvironment& OutEnvironment)
 	{
 		FGlobalShader::ModifyCompilationEnvironment(Parameters, OutEnvironment);
@@ -2138,25 +2150,46 @@ public:
 
 IMPLEMENT_GLOBAL_SHADER(FVisiblityRasterComputeCS, "/Engine/Private/HairStrands/HairStrandsVisibilityRasterCompute.usf", "MainCS", SF_Compute);
 
+
+static bool DoesSupportRasterCompute()
+{
+#if PLATFORM_WINDOWS
+	return IsRHIDeviceNVIDIA() && GRHISupportsAtomicUInt64;
+#else
+	return false;
+#endif
+}
+
 struct FRasterComputeOutput
 {
+	FIntPoint BaseResolution;
+	FIntPoint SuperResolution;
+	uint32 ResolutionMultiplier = 1;
+
 	FRDGTextureRef HairCountTexture;
 	FRDGTextureRef DepthTexture;
+	FRDGTextureRef VisibilityTexture;
 };
 
 static FRasterComputeOutput AddVisibilityComputeRasterPass(
 	FRDGBuilder& GraphBuilder,
 	const FViewInfo* ViewInfo,
 	const FHairStrandsMacroGroupDatas& MacroGroupDatas,
-	const FIntPoint& Resolution,
+	const FIntPoint& InResolution,
 	const FRDGTextureRef SceneDepthTexture)
 {
+	check(DoesSupportRasterCompute());
+
 	FRasterComputeOutput Out;
 
+	Out.ResolutionMultiplier = FMath::Clamp(GHairStrandsVisibilityComputeRasterSampleCount, 1, 4);
+	Out.BaseResolution		 = InResolution;
+	Out.SuperResolution		 = InResolution * Out.ResolutionMultiplier;
+
 	{
 		FRDGTextureDesc Desc;
-		Desc.Extent.X = Resolution.X;
-		Desc.Extent.Y = Resolution.Y;
+		Desc.Extent.X = Out.SuperResolution.X;
+		Desc.Extent.Y = Out.SuperResolution.Y;
 		Desc.Depth = 0;
 		Desc.Format = PF_R32_UINT;
 		Desc.NumMips = 1;
@@ -2171,8 +2204,8 @@ static FRasterComputeOutput AddVisibilityComputeRasterPass(
 
 	{
 		FRDGTextureDesc Desc;
-		Desc.Extent.X = Resolution.X;
-		Desc.Extent.Y = Resolution.Y;
+		Desc.Extent.X = Out.SuperResolution.X;
+		Desc.Extent.Y = Out.SuperResolution.Y;
 		Desc.Depth = 0;
 		Desc.Format = PF_R32_UINT;
 		Desc.NumMips = 1;
@@ -2185,9 +2218,26 @@ static FRasterComputeOutput AddVisibilityComputeRasterPass(
 	}
 	FRDGTextureUAVRef DepthTextureUAV = GraphBuilder.CreateUAV(Out.DepthTexture);
 
+	{
+		FRDGTextureDesc Desc;
+		Desc.Extent.X = Out.SuperResolution.X;
+		Desc.Extent.Y = Out.SuperResolution.Y;
+		Desc.Depth = 0;
+		Desc.Format = PF_R32G32_UINT;
+		Desc.NumMips = 1;
+		Desc.NumSamples = 1;
+		Desc.Flags = TexCreate_None;
+		Desc.TargetableFlags = TexCreate_UAV | TexCreate_ShaderResource | TexCreate_RenderTargetable;
+		Desc.bForceSharedTargetAndShaderResource = true;
+		Desc.ClearValue = FClearValueBinding(FLinearColor(1.0f, 1.0f, 1.0f, 1.0f)); // Clear to transmittance 1
+		Out.VisibilityTexture = GraphBuilder.CreateTexture(Desc, TEXT("HairVisibilityTexture"));
+	}
+	FRDGTextureUAVRef VisibilityTextureUAV = GraphBuilder.CreateUAV(Out.VisibilityTexture);
+
 	uint32 ClearValues[4] = { 0,0,0,0 };
 	AddClearUAVPass(GraphBuilder, HairCountTextureUAV, ClearValues);
 	AddClearUAVPass(GraphBuilder, DepthTextureUAV, ClearValues);
+	AddClearUAVPass(GraphBuilder, VisibilityTextureUAV, ClearValues);
 
 	// Create and set the uniform buffer
 	const bool bEnableMSAA = false;
@@ -2207,15 +2257,18 @@ static FRasterComputeOutput AddVisibilityComputeRasterPass(
 		for (const FHairStrandsMacroGroupData::PrimitiveInfo& PrimitiveInfo : PrimitiveSceneInfos)
 		{
 			FVisiblityRasterComputeCS::FParameters* PassParameters = GraphBuilder.AllocParameters<FVisiblityRasterComputeCS::FParameters>();
-			PassParameters->OutputResolution = Resolution;
+			PassParameters->OutputResolution = Out.SuperResolution;
+			PassParameters->ResolutionMultiplier = Out.ResolutionMultiplier;
 			PassParameters->MacroGroupId = MacroGroup.MacroGroupId;
 			PassParameters->DispatchCountX = DispatchCountX;
 			PassParameters->MaxRasterCount = FMath::Clamp(GHairStrandsVisibilityComputeRasterMaxPixelCount, 1, 256);
 			PassParameters->FrameIdMod8 = FrameIdMode8;
+			PassParameters->HairMaterialId = PrimitiveInfo.MaterialId;
 			PassParameters->ViewUniformBuffer = ViewUniformShaderParameters;
 			PassParameters->SceneDepthTexture = SceneDepthTexture;
 			PassParameters->OutHairCountTexture = HairCountTextureUAV;
 			PassParameters->OutDepthTexture = DepthTextureUAV;
+			PassParameters->OutVisibilityTexture = VisibilityTextureUAV;
 
 			check(PrimitiveInfo.MeshBatchAndRelevance.Mesh && PrimitiveInfo.MeshBatchAndRelevance.Mesh->Elements.Num() > 0);
 			const FHairGroupPublicData* HairGroupPublicData = reinterpret_cast<const FHairGroupPublicData*>(PrimitiveInfo.MeshBatchAndRelevance.Mesh->Elements[0].VertexFactoryUserData);
@@ -2300,7 +2353,7 @@ FHairStrandsVisibilityViews RenderHairStrandsVisibilityBuffer(
 					bOutputHairCount,
 					SceneDepthTexture);
 
-				const bool bUseRasterCompute = GHairStrandsVisibilityComputeRaster > 0;
+				const bool bUseRasterCompute = GHairStrandsVisibilityComputeRaster > 0 && DoesSupportRasterCompute();
 				if (bUseRasterCompute)
 				{
 					FRasterComputeOutput RasterOutput = AddVisibilityComputeRasterPass(
