Changed scope name for hair raytracing shadows for better profiling.

diff --git a/Engine/Source/Runtime/Renderer/Private/LightRendering.cpp b/Engine/Source/Runtime/Renderer/Private/LightRendering.cpp
--- a/Engine/Source/Runtime/Renderer/Private/LightRendering.cpp
+++ b/Engine/Source/Runtime/Renderer/Private/LightRendering.cpp
@@ -1426,6 +1426,8 @@ void FDeferredShadingSceneRenderer::RenderLights(FRHICommandListImmediate& RHICm
 									FRDGTextureUAV* SubPixelRayTracingShadowMaskUAV = GraphBuilder.CreateUAV(FRDGTextureUAVDesc(SubPixelRayTracingShadowMaskTexture));
 									FRDGTextureUAV* SubPixelRayHitDistanceUAV = GraphBuilder.CreateUAV(FRDGTextureUAVDesc(SubPixelRayDistanceTexture));
 
+									RDG_EVENT_SCOPE(GraphBuilder, "%s Hair Strands", *BatchLightNameWithLevel);
+
 									RenderRayTracingShadows(
 										GraphBuilder,
 										SceneTextures,
