Removed allocation of unnecessary UAVs  in the raytracing shadows pass (they were already allocated before calling the function)

diff --git a/Engine/Source/Runtime/Renderer/Private/RayTracing/RayTracingShadows.cpp b/Engine/Source/Runtime/Renderer/Private/RayTracing/RayTracingShadows.cpp
--- a/Engine/Source/Runtime/Renderer/Private/RayTracing/RayTracingShadows.cpp
+++ b/Engine/Source/Runtime/Renderer/Private/RayTracing/RayTracingShadows.cpp
@@ -166,31 +166,6 @@ void FDeferredShadingSceneRenderer::RenderRayTracingShadows(
 	FLightSceneProxy* LightSceneProxy = LightSceneInfo.Proxy;
 	check(LightSceneProxy);
 
-	// Render targets
-	FRDGTextureRef ScreenShadowMaskTexture;
-	{
-		FRDGTextureDesc Desc = FRDGTextureDesc::Create2DDesc(
-			SceneTextures.SceneDepthBuffer->Desc.Extent,
-			PF_FloatRGBA,
-			FClearValueBinding::Black,
-			TexCreate_None,
-			TexCreate_ShaderResource | TexCreate_RenderTargetable | TexCreate_UAV,
-			/* bInForceSeparateTargetAndShaderResource = */ false);
-		ScreenShadowMaskTexture = GraphBuilder.CreateTexture(Desc, TEXT("RayTracingOcclusion"));
-	}
-
-	FRDGTextureRef RayDistanceTexture;
-	{
-		FRDGTextureDesc Desc = FRDGTextureDesc::Create2DDesc(
-			SceneTextures.SceneDepthBuffer->Desc.Extent,
-			PF_R16F,
-			FClearValueBinding::Black,
-			TexCreate_None,
-			TexCreate_ShaderResource | TexCreate_RenderTargetable | TexCreate_UAV,
-			/* bInForceSeparateTargetAndShaderResource = */ false);
-		RayDistanceTexture = GraphBuilder.CreateTexture(Desc, TEXT("RayTracingOcclusionDistance"));
-	}
-
 	FIntRect ScissorRect = View.ViewRect;
 	FIntPoint PixelOffset = { 0, 0 };
 
