Fix CIS issue.

diff --git a/Engine/Source/Runtime/Renderer/Private/HairStrands/HairStrandsVisibility.cpp b/Engine/Source/Runtime/Renderer/Private/HairStrands/HairStrandsVisibility.cpp
--- a/Engine/Source/Runtime/Renderer/Private/HairStrands/HairStrandsVisibility.cpp
+++ b/Engine/Source/Runtime/Renderer/Private/HairStrands/HairStrandsVisibility.cpp
@@ -2284,7 +2284,7 @@ static FRasterComputeOutput AddVisibilityComputeRasterPass(
 
 			const uint32 DispatchCountY = FMath::CeilToInt(PassParameters->HairStrandsVF_VertexCount / float(GroupSize * DispatchCountX));
 			const FIntVector DispatchCount(DispatchCountX, DispatchCountY, 1);
-			FComputeShaderUtils::AddPass(GraphBuilder, RDG_EVENT_NAME("HairStrandsVisibilityComputeRaster"), *ComputeShader, PassParameters, DispatchCount);
+			FComputeShaderUtils::AddPass(GraphBuilder, RDG_EVENT_NAME("HairStrandsVisibilityComputeRaster"), ComputeShader, PassParameters, DispatchCount);
 		}
 	}
 
