Fix missing permutation.

diff --git a/Engine/Source/Runtime/Renderer/Private/HairStrands/HairStrandsVisibility.cpp b/Engine/Source/Runtime/Renderer/Private/HairStrands/HairStrandsVisibility.cpp
--- a/Engine/Source/Runtime/Renderer/Private/HairStrands/HairStrandsVisibility.cpp
+++ b/Engine/Source/Runtime/Renderer/Private/HairStrands/HairStrandsVisibility.cpp
@@ -1164,10 +1164,6 @@ public:
 			PermutationVector.Set<FVisibility>(0);
 			PermutationVector.Set<FMSAACount>(4);
 		}
-		else if (PermutationVector.Get<FViewTransmittance>() == 1)
-		{
-			PermutationVector.Set<FMSAACount>(4);
-		}
 		return PermutationVector;
 	}
 
@@ -1186,10 +1182,6 @@ public:
 		{
 			return false;
 		}
-		if (PermutationVector.Get<FViewTransmittance>() > 0 && PermutationVector.Get<FMSAACount>() == 8)
-		{
-			return false;
-		}
 		return IsHairStrandsSupported(Parameters.Platform);
 	}
 };
