Add variable to turn off voxel traversal in RT shadow for profiling purpose.

diff --git a/Engine/Source/Runtime/Renderer/Private/RayTracing/RayTracingShadows.cpp b/Engine/Source/Runtime/Renderer/Private/RayTracing/RayTracingShadows.cpp
--- a/Engine/Source/Runtime/Renderer/Private/RayTracing/RayTracingShadows.cpp
+++ b/Engine/Source/Runtime/Renderer/Private/RayTracing/RayTracingShadows.cpp
@@ -41,6 +41,13 @@ static TAutoConsoleVariable<int32> CVarRayTracingShadowsEnableTwoSidedGeometry(
 	ECVF_RenderThreadSafe
 );
 
+static TAutoConsoleVariable<int32> CVarRayTracingShadowsEnableHairVoxel(
+	TEXT("r.RayTracing.Shadows.EnableHairVoxel"),
+	1,
+	TEXT("Enables use of hair voxel data for tracing shadow (default = 1)"),
+	ECVF_RenderThreadSafe
+);
+
 bool EnableRayTracingShadowTwoSidedGeometry()
 {
 	return CVarRayTracingShadowsEnableTwoSidedGeometry.GetValueOnRenderThread() != 0;
@@ -221,7 +228,8 @@ void FDeferredShadingSceneRenderer::RenderRayTracingShadows(
 		PassParameters->SSProfilesTexture = GraphBuilder.RegisterExternalTexture(View.RayTracingSubSurfaceProfileTexture);
 		if (bUseHairLighting)
 		{
-			PassParameters->bUseHairVoxel = IsHairRayTracingEnabled() ? 0 : 1;
+			const bool bUseHairVoxel = CVarRayTracingShadowsEnableHairVoxel.GetValueOnRenderThread() > 0;
+			PassParameters->bUseHairVoxel = (bUseHairVoxel && !IsHairRayTracingEnabled()) ? 1 : 0;
 			PassParameters->HairCategorizationTexture = HairResources->CategorizationTexture;
 			PassParameters->HairLightChannelMaskTexture = HairResources->LightChannelMaskTexture;
 			PassParameters->VirtualVoxel = HairResources->VoxelResources->UniformBuffer;
