Add pass for patching buffer data with hair data.

This allows to solve post process edge blur filter not behaving correctly on hair.

diff --git a/Engine/Shaders/Private/HairStrands/HairStrandsVisibilityComposeSubPixelPS.usf b/Engine/Shaders/Private/HairStrands/HairStrandsVisibilityComposeSubPixelPS.usf
--- a/Engine/Shaders/Private/HairStrands/HairStrandsVisibilityComposeSubPixelPS.usf
+++ b/Engine/Shaders/Private/HairStrands/HairStrandsVisibilityComposeSubPixelPS.usf
@@ -61,4 +61,25 @@ void FastResolvePS(in FScreenVertexOutput Input)
 		discard;
 }
 
-#endif // SHADER_FASTRESOLVE
\ No newline at end of file
+#endif // SHADER_FASTRESOLVE
+
+#if SHADER_PATCH
+
+Texture2D<uint4> CategorisationTexture;
+
+void MainPS(
+	in FScreenVertexOutput Input,
+	out float4 OutGBufferA : SV_Target0,
+	out float4 OutGBufferB : SV_Target1)
+{
+	const uint2 PixelCoord = floor(Input.Position.xy);
+	const FCategorizationData CatData = DecodeCategorizationData(CategorisationTexture.Load(uint3(Input.Position.xy, 0)));
+	if (CatData.PixelCoverage == 0)
+		discard;
+
+	const float3 Tangent = float3(0, 0, 1);
+	OutGBufferA = float4(Tangent, 0);
+	OutGBufferB = float4(0, 0, 0, SHADINGMODELID_HAIR);
+}
+
+#endif // SHADER_PATCH
\ No newline at end of file
diff --git a/Engine/Source/Runtime/Renderer/Private/HairStrands/HairStrandsComposition.cpp b/Engine/Source/Runtime/Renderer/Private/HairStrands/HairStrandsComposition.cpp
--- a/Engine/Source/Runtime/Renderer/Private/HairStrands/HairStrandsComposition.cpp
+++ b/Engine/Source/Runtime/Renderer/Private/HairStrands/HairStrandsComposition.cpp
@@ -16,6 +16,9 @@
 static int32 GHairFastResolveVelocityThreshold = 1;
 static FAutoConsoleVariableRef CVarHairFastResolveVelocityThreshold(TEXT("r.HairStrands.VelocityThreshold"), GHairFastResolveVelocityThreshold, TEXT("Threshold value (in pixel) above which a pixel is forced to be resolve with responsive AA (in order to avoid smearing). Default is 3."));
 
+static int32 GHairPatchBufferDataBeforePostProcessing = 1;
+static FAutoConsoleVariableRef CVarHairPatchBufferDataBeforePostProcessing(TEXT("r.HairStrands.PatchMaterialData"), GHairPatchBufferDataBeforePostProcessing, TEXT("Patch the buffer with hair material data before post processing run. (default 1)."));
+
 ///////////////////////////////////////////////////////////////////////////////////////////////////
 class FHairVisibilityComposeSubPixelPS : public FGlobalShader
 {
@@ -215,6 +218,82 @@ static void AddHairVisibilityFastResolvePass(
 	}
 }
 
+
+///////////////////////////////////////////////////////////////////////////////////////////////////
+class FHairPatchGbufferDataPS : public FGlobalShader
+{
+	DECLARE_GLOBAL_SHADER(FHairPatchGbufferDataPS);
+	SHADER_USE_PARAMETER_STRUCT(FHairPatchGbufferDataPS, FGlobalShader);
+
+	BEGIN_SHADER_PARAMETER_STRUCT(FParameters, )
+		SHADER_PARAMETER_RDG_TEXTURE(Texture2D, CategorisationTexture)
+		RENDER_TARGET_BINDING_SLOTS()
+	END_SHADER_PARAMETER_STRUCT()
+
+public:
+	static bool ShouldCompilePermutation(const FGlobalShaderPermutationParameters& Parameters) { return Parameters.Platform == SP_PCD3D_SM5; }
+	static void ModifyCompilationEnvironment(const FGlobalShaderPermutationParameters& Parameters, FShaderCompilerEnvironment& OutEnvironment)
+	{
+		FGlobalShader::ModifyCompilationEnvironment(Parameters, OutEnvironment);
+		OutEnvironment.SetDefine(TEXT("SHADER_PATCH"), 1);
+	}
+};
+
+IMPLEMENT_GLOBAL_SHADER(FHairPatchGbufferDataPS, "/Engine/Private/HairStrands/HairStrandsVisibilityComposeSubPixelPS.usf", "MainPS", SF_Pixel); 
+
+static void AddPatchGbufferDataPass(
+	FRDGBuilder& GraphBuilder,
+	const FViewInfo& View,
+	FRDGTextureRef HairCategorizationTexture,
+	FRDGTextureRef OutGbufferATexture,
+	FRDGTextureRef OutGbufferBTexture)
+{
+	const FIntPoint Resolution = OutGbufferATexture->Desc.Extent;
+	FHairPatchGbufferDataPS::FParameters* Parameters = GraphBuilder.AllocParameters<FHairPatchGbufferDataPS::FParameters>();
+	Parameters->CategorisationTexture = HairCategorizationTexture;
+	Parameters->RenderTargets[0] = FRenderTargetBinding(OutGbufferATexture, ERenderTargetLoadAction::ELoad);
+	Parameters->RenderTargets[1] = FRenderTargetBinding(OutGbufferBTexture, ERenderTargetLoadAction::ELoad);
+
+	TShaderMapRef<FPostProcessVS> VertexShader(View.ShaderMap);
+	TShaderMapRef<FHairPatchGbufferDataPS> PixelShader(View.ShaderMap);
+	const FIntRect Viewport = View.ViewRect;
+
+	const FViewInfo* CapturedView = &View;
+
+	ClearUnusedGraphResources(PixelShader, Parameters);
+
+	GraphBuilder.AddPass(
+		RDG_EVENT_NAME("HairPatchGbufferData"),
+		Parameters,
+		ERDGPassFlags::Raster,
+		[Parameters, VertexShader, PixelShader, Viewport, Resolution, CapturedView](FRHICommandList& RHICmdList)
+	{
+		FGraphicsPipelineStateInitializer GraphicsPSOInit;
+		RHICmdList.ApplyCachedRenderTargets(GraphicsPSOInit);
+		GraphicsPSOInit.BlendState = TStaticBlendState<CW_RGBA, BO_Add, BF_One, BF_Zero, BO_Add, BF_One, BF_Zero>::GetRHI();
+		GraphicsPSOInit.RasterizerState = TStaticRasterizerState<>::GetRHI();
+		GraphicsPSOInit.DepthStencilState = TStaticDepthStencilState<false, CF_Always>::GetRHI();
+		GraphicsPSOInit.BoundShaderState.VertexDeclarationRHI = GFilterVertexDeclaration.VertexDeclarationRHI;
+		GraphicsPSOInit.BoundShaderState.VertexShaderRHI = VertexShader.GetVertexShader();
+		GraphicsPSOInit.BoundShaderState.PixelShaderRHI = PixelShader.GetPixelShader();
+		GraphicsPSOInit.PrimitiveType = PT_TriangleList;
+		SetGraphicsPipelineState(RHICmdList, GraphicsPSOInit);
+
+		VertexShader->SetParameters(RHICmdList, CapturedView->ViewUniformBuffer);
+		RHICmdList.SetViewport(Viewport.Min.X, Viewport.Min.Y, 0.0f, Viewport.Max.X, Viewport.Max.Y, 1.0f);
+		SetShaderParameters(RHICmdList, PixelShader, PixelShader.GetPixelShader(), *Parameters);
+		DrawRectangle(
+			RHICmdList,
+			0, 0,
+			Viewport.Width(), Viewport.Height(),
+			Viewport.Min.X, Viewport.Min.Y,
+			Viewport.Width(), Viewport.Height(),
+			Viewport.Size(),
+			Resolution,
+			VertexShader,
+			EDRF_UseTriangleOptimization);
+	});
+}
 ///////////////////////////////////////////////////////////////////////////////////////////////////
 
 void RenderHairComposeSubPixel(
@@ -283,6 +362,20 @@ void RenderHairComposeSubPixel(
 						RDGHairVisibilityVelocityTexture,
 						SceneColorDepth);
 				}
+
+				const bool bPatchBufferData = GHairPatchBufferDataBeforePostProcessing > 0;
+				if (bPatchBufferData)
+				{
+					FSceneRenderTargets& SceneContext = FSceneRenderTargets::Get(RHICmdList);				
+					const FRDGTextureRef GBufferATexture = GraphBuilder.RegisterExternalTexture(SceneContext.GBufferA , TEXT("GBufferA"));
+					const FRDGTextureRef GBufferBTexture = GraphBuilder.RegisterExternalTexture(SceneContext.GBufferB, TEXT("GBufferB"));
+					AddPatchGbufferDataPass(
+						GraphBuilder,
+						View,
+						RDGCategorisationTexture,
+						GBufferATexture,
+						GBufferBTexture);
+				}
 			}
 		}
 	}
