Deferred lighting now runs at sample rate for hair.

This solves depth bound check not being enabled on opaque pixe and improve overal performance (-2ms in test content)

diff --git a/Engine/Shaders/Private/DeferredLightPixelShaders.usf b/Engine/Shaders/Private/DeferredLightPixelShaders.usf
--- a/Engine/Shaders/Private/DeferredLightPixelShaders.usf
+++ b/Engine/Shaders/Private/DeferredLightPixelShaders.usf
@@ -102,6 +102,17 @@ uint GetLightingChannelMask(float2 UV)
 	return LightingChannelsTexture.Load(uint3(IntegerUV, 0)).x;
 }
 
+float GetExposure()
+{
+#if USE_PREEXPOSURE
+	return View.PreExposure;
+#else
+	return 0;
+#endif
+}
+
+////////////////////////////////////////////////////////////////////////////////////////////////////
+
 #if USE_HAIR_LIGHTING
 
 Texture2D<uint4> HairCategorizationTexture;
@@ -145,8 +156,7 @@ float4 PerceptualToLinear(float4 Color)
 void ComputeHairLighting(
 	in FInputParams InputParams,
 	in FNodeDesc NodeDesc,
-	in FCategorizationData CatData,
-	inout float OutHairCoverage,
+	in float PixelCoverage,
 	inout float4 OutColor, 
 	inout float4 OutColorSubPixel)
 {
@@ -224,15 +234,12 @@ void ComputeHairLighting(
 		const float LocalCoverage = From8bitCoverage(Sample.Coverage8bit);
 		LocalColor += PercepualColor * LocalCoverage;
 	}
-	const float PixelCoverage = saturate(CatData.PixelCoverage);
-	OutHairCoverage = PixelCoverage;
-
 	OutColor = 0;
 	OutColorSubPixel = 0;
-	if (OutHairCoverage < 1)
+	if (PixelCoverage < 1)
 	{
 		OutColorSubPixel = LocalColor * PixelCoverage;
-		OutColorSubPixel.a = OutHairCoverage;
+		OutColorSubPixel.a = PixelCoverage;
 	}
 	else
 	{
@@ -260,6 +267,10 @@ float3 GetAtmosphericLightTransmittance(in float4 SVPos, in float2 ScreenUV, in
 
 #endif // USE_ATMOSPHERE_TRANSMITTANCE
 
+////////////////////////////////////////////////////////////////////////////////////////////////////
+
+#if USE_HAIR_LIGHTING == 0 || USE_HAIR_LIGHTING == 1 || USE_HAIR_LIGHTING == 3
+
 void DeferredLightPixelMain(
 #if LIGHT_SOURCE_SHAPE > 0
 	float4 InScreenPosition : TEXCOORD0,
@@ -269,7 +280,7 @@ void DeferredLightPixelMain(
 #endif
 	float4 SVPos			: SV_POSITION,
 	out float4 OutColor		: SV_Target0
-#if USE_HAIR_LIGHTING
+#if USE_HAIR_LIGHTING == 1
 ,	out float4 OutColorSubPixel	: SV_Target1
 #endif
 	)
@@ -292,20 +303,20 @@ void DeferredLightPixelMain(
 
 #if USE_HAIR_LIGHTING
 	const uint3 PixelCoord = uint3(floor(InputParams.PixelPos), 0);
-	const FNodeDesc NodeDesc = DecodeNodeDesc(HairVisibilityNodeOffsetAndCount.Load(PixelCoord));
-	const FCategorizationData CatData = DecodeCategorizationData(HairCategorizationTexture.Load(PixelCoord));
-	
+	const FCategorizationData CatData = DecodeCategorizationData(HairCategorizationTexture.Load(PixelCoord));	
+	const float HairCoverage = saturate(CatData.PixelCoverage);
+
+	#if USE_HAIR_LIGHTING == 1
 	// Only light pixels marked as using deferred shading
 	OutColorSubPixel = 0;
-	float HairCoverage = 0;
+	const FNodeDesc NodeDesc = DecodeNodeDesc(HairVisibilityNodeOffsetAndCount.Load(PixelCoord));
 	BRANCH 
 	if (NodeDesc.Count > 0)
 	{
-		ComputeHairLighting(InputParams, NodeDesc, CatData, HairCoverage, OutColor, OutColorSubPixel);
-	#if USE_PREEXPOSURE
-		OutColorSubPixel.rgb *= View.PreExposure;
-	#endif
+		ComputeHairLighting(InputParams, NodeDesc, HairCoverage, OutColor, OutColorSubPixel);
+		OutColorSubPixel.rgb *= GetExposure();
 	}
+	#endif
 #endif
 
 	FScreenSpaceData ScreenSpaceData = GetScreenSpaceData(InputParams.ScreenUV);
@@ -342,10 +353,108 @@ void DeferredLightPixelMain(
 #endif
 	}
 
-#if USE_PREEXPOSURE
 	// RGB:SceneColor Specular and Diffuse
 	// A:Non Specular SceneColor Luminance
 	// So we need PreExposure for both color and alpha
-	OutColor.rgba *= View.PreExposure;
+	OutColor.rgba *= GetExposure();
+}
 #endif
+
+////////////////////////////////////////////////////////////////////////////////////////////////////
+
+#if USE_HAIR_LIGHTING == 2
+
+StructuredBuffer<uint>	HairVisibilityNodeCoords;
+
+void DeferredLightPixelMain(
+	float4 SVPos : SV_POSITION,
+	uint NodeCount  : DISPATCH_NODECOUNT,
+	uint2 Resolution : DISPATCH_RESOLUTION,
+	out float4 OutColor : SV_Target0)
+{
+	OutColor = 0;
+
+	const uint2 InCoord = uint2(SVPos.xy);
+	const uint SampleIndex = InCoord.x + InCoord.y * Resolution.x;
+	if (SampleIndex >= NodeCount)
+	{
+		return;
+	}
+
+	const uint PackedCoord = HairVisibilityNodeCoords[SampleIndex];
+	const uint2 PixelCoord = uint2((PackedCoord & 0xFFFF), ((PackedCoord >> 16) & 0xFFFF));
+	const float2 UV = (PixelCoord + float2(0.5f, 0.5f)) / float2(View.BufferSizeAndInvSize.xy);
+	const float2 ScreenPosition = (UV - View.ScreenPositionScaleBias.wz) / View.ScreenPositionScaleBias.xy;
+	
+	const FPackedHairSample PackedSample = HairVisibilityNodeData[SampleIndex];
+	const FHairSample Sample = UnpackHairSample(PackedSample);
+
+	#if USE_LIGHTING_CHANNELS
+	if (!(Sample.LightChannelMask & DeferredLightUniforms.LightingChannelMask))
+	{
+		return;
+	}
+	#endif
+
+	// Inject material data from the visibility/mini-gbuffer for hair
+	const float SceneDepth = ConvertFromDeviceZ(Sample.Depth);
+	FScreenSpaceData HairScreenSpaceData;
+	HairScreenSpaceData.AmbientOcclusion = 1;
+	HairScreenSpaceData.GBuffer.WorldNormal = Sample.Tangent;
+	HairScreenSpaceData.GBuffer.BaseColor = Sample.BaseColor;
+	HairScreenSpaceData.GBuffer.Roughness = Sample.Roughness;
+	HairScreenSpaceData.GBuffer.ShadingModelID = SHADINGMODELID_HAIR;
+	HairScreenSpaceData.GBuffer.DiffuseColor = 0;
+	HairScreenSpaceData.GBuffer.SpecularColor = 0;
+	HairScreenSpaceData.GBuffer.Specular = Sample.Specular;
+	HairScreenSpaceData.GBuffer.Metallic = 0;
+	HairScreenSpaceData.GBuffer.Depth = SceneDepth;
+	HairScreenSpaceData.GBuffer.GBufferAO = 1;
+	HairScreenSpaceData.GBuffer.CustomData = float4(HairDualScatteringRoughnessOverride, 0, Sample.Backlit, 0);
+	HairScreenSpaceData.GBuffer.IndirectIrradiance = 0;
+	HairScreenSpaceData.GBuffer.PrecomputedShadowFactors = 1;
+
+	const float3 WorldPosition = mul(float4(ScreenPosition * HairScreenSpaceData.GBuffer.Depth, HairScreenSpaceData.GBuffer.Depth, 1), View.ScreenToWorld).xyz;
+	const float3 CameraVector = normalize(WorldPosition - View.WorldCameraOrigin);
+
+	FDeferredLightData LightData = SetupLightDataForStandardDeferred();
+	LightData.HairTransmittance.ScatteringComponent = HairComponents;
+
+	// todo add light culling/early out
+
+	// Compute Hair transmittance
+	{
+		float3 L = float3(0, 0, 1);
+		if (LightData.bRadialLight)
+			L = normalize(LightData.Position - WorldPosition);
+		else
+			L = LightData.Direction;
+
+		const float3 V = normalize(-CameraVector);
+		LightData.HairTransmittance = GetHairTransmittance(
+			V,
+			L,
+			HairScreenSpaceData.GBuffer,
+			SampleIndex,
+			HairTransmittanceBufferMaxCount,
+			HairTransmittanceBuffer,
+			HairLUTTexture,
+			HairLUTSampler,
+			HairComponents);
+	}
+
+	const float Dither = InterleavedGradientNoise(PixelCoord, View.StateFrameIndexMod8);
+	float4 LightAttenuation = HairShadowMaskValid ? ScreenShadowMaskSubPixelTexture.Load(uint3(PixelCoord,0)) : 1;
+	LightAttenuation = min(LightAttenuation, LightData.HairTransmittance.OpaqueVisibility.xxxx);
+
+	FRectTexture RectTexture = InitRectTexture(DeferredLightUniforms.SourceTexture);
+	float SurfaceShadow = 1.0f;
+	const float4 Radiance = GetDynamicLighting(WorldPosition, CameraVector, HairScreenSpaceData.GBuffer, HairScreenSpaceData.AmbientOcclusion, HairScreenSpaceData.GBuffer.ShadingModelID, LightData, LightAttenuation, Dither, PixelCoord, RectTexture, SurfaceShadow);
+	const float  Attenuation = ComputeLightProfileMultiplier(WorldPosition, DeferredLightUniforms.Position, -DeferredLightUniforms.Direction, DeferredLightUniforms.Tangent);
+
+	const float LocalCoverage = From8bitCoverage(Sample.Coverage8bit);
+
+	OutColor.rgb = Radiance.xyz * Attenuation * LocalCoverage * GetExposure();
+	OutColor.a = LocalCoverage;
 }
+#endif
\ No newline at end of file
diff --git a/Engine/Shaders/Private/DeferredLightVertexShaders.usf b/Engine/Shaders/Private/DeferredLightVertexShaders.usf
--- a/Engine/Shaders/Private/DeferredLightVertexShaders.usf
+++ b/Engine/Shaders/Private/DeferredLightVertexShaders.usf
@@ -11,6 +11,7 @@ float4 StencilingConeParameters;	// .x NumSides (0 if not cone), .y NumSlices, .
 float4x4 StencilingConeTransform;
 float3 StencilingPreViewTranslation;
 
+#if SHADER_RADIAL_LIGHT == 0
 /** Vertex shader for rendering a directional light using a full screen quad. */
 void DirectionalVertexMain(
 	in float2 InPosition : ATTRIBUTE0,
@@ -23,8 +24,9 @@ void DirectionalVertexMain(
 	DrawRectangle(float4(InPosition.xy, 0, 1), InUV, OutPosition, OutTexCoord);
 	OutScreenVector = mul(float4(OutPosition.xy, 1, 0), View.ScreenToTranslatedWorld).xyz;
 }
+#endif
 
-#if FEATURE_LEVEL >= FEATURE_LEVEL_SM4
+#if FEATURE_LEVEL >= FEATURE_LEVEL_SM4 && SHADER_RADIAL_LIGHT == 1
 /** Vertex shader for rendering a point or spot light using approximate bounding geometry. */
 void RadialVertexMain(
 	in uint InVertexId : SV_VertexID,
@@ -93,3 +95,28 @@ void RadialVertexMain(
 	OutScreenPosition = OutPosition = mul(float4(WorldPosition, 1), View.TranslatedWorldToClip);
 }
 #endif
+
+#if SHADER_HAIR == 1
+
+uint NodeGroupSize;
+int2 MaxResolution;
+Texture2D<uint> HairVisibilityNodeCount;
+
+void HairVertexMain(
+	in uint InVertexId : SV_VertexID,
+	out float4 OutPosition : SV_POSITION,
+	nointerpolation out uint OutNodeCount : DISPATCH_NODECOUNT,
+	nointerpolation out uint2 OutResolution : DISPATCH_RESOLUTION)
+{
+	OutNodeCount = HairVisibilityNodeCount.Load(uint3(0,0,0));
+	OutResolution.x = ceil(sqrt(OutNodeCount));
+	OutResolution.y = OutResolution.x;
+
+	const float2 ClipCoord = ((OutResolution + 0.5f) / float2(MaxResolution)) * 2;
+
+	const float2 UV = float2(InVertexId & 1, InVertexId >> 1);
+	const float2 Pos = float2(-UV.x, UV.y) * 4 + float2(-1, +1) + float2(ClipCoord.x, -ClipCoord.y);
+	OutPosition = float4(Pos, 0.5f, 1);
+}
+
+#endif
\ No newline at end of file
diff --git a/Engine/Shaders/Private/HairStrands/HairStrandsEnvironmentLighting.usf b/Engine/Shaders/Private/HairStrands/HairStrandsEnvironmentLighting.usf
--- a/Engine/Shaders/Private/HairStrands/HairStrandsEnvironmentLighting.usf
+++ b/Engine/Shaders/Private/HairStrands/HairStrandsEnvironmentLighting.usf
@@ -553,6 +553,9 @@ float3 ComputeThinTransmission(FGBufferData GBuffer, float2 ScreenPosition, floa
 
 ////////////////////////////////////////////////////////////////////////////////////////////////////////////////////////////////////////////
 
+
+#if LIGHTING_CS
+
 uint MaxVisibilityNodeCount;
 StructuredBuffer<uint>	HairVisibilityNodeCoord;
 Buffer<uint4>			IndirectArgsBuffer;
@@ -623,3 +626,107 @@ void MainCS(uint3 DispatchThreadId : SV_DispatchThreadID)
 		OutLightingBuffer[LocalOffset] = 0;
 	}
 }
+
+#endif // LIGHTING_CS
+
+
+
+#if LIGHTING_VS
+
+int2 MaxViewportResolution;
+Texture2D<uint> HairVisibilityNodeCount;
+
+void MainVS(
+	in uint InVertexId : SV_VertexID,
+	out float4 OutPosition : SV_POSITION,
+	nointerpolation out uint OutNodeCount : DISPATCH_NODECOUNT,
+	nointerpolation out uint2 OutResolution : DISPATCH_RESOLUTION)
+{
+	OutNodeCount = HairVisibilityNodeCount.Load(uint3(0, 0, 0));
+	OutResolution.x = ceil(sqrt(OutNodeCount));
+	OutResolution.y = OutResolution.x;
+
+	const float2 ClipCoord = ((OutResolution + 0.5f) / float2(MaxViewportResolution)) * 2;
+
+	const float2 UV = float2(InVertexId & 1, InVertexId >> 1);
+	const float2 Pos = float2(-UV.x, UV.y) * 4 + float2(-1, +1) + float2(ClipCoord.x, -ClipCoord.y);
+	OutPosition = float4(Pos, 0.5f, 1);
+}
+
+#endif // LIGHTING_VS
+
+#if LIGHTING_PS
+
+uint MaxVisibilityNodeCount;
+StructuredBuffer<uint>	HairVisibilityNodeCoord;
+Buffer<uint4>			IndirectArgsBuffer;
+
+#if PERMUTATION_GROUP_SIZE == 0
+#define LIGHTING_GROUP_SIZE 64
+#elif PERMUTATION_GROUP_SIZE == 1
+#define LIGHTING_GROUP_SIZE 32
+#else
+#error Unknown group size	
+#endif
+
+void MainPS(
+	float4 SVPos : SV_POSITION,
+	uint NodeCount : DISPATCH_NODECOUNT,
+	uint2 Resolution : DISPATCH_RESOLUTION,
+	out float4 OutColor : SV_Target0)
+{
+	OutColor = 0;
+
+	const uint2 InCoord = uint2(SVPos.xy);
+	const uint LocalOffset = InCoord.x + InCoord.y * Resolution.x;
+	if (LocalOffset >= NodeCount)
+	{
+		return;
+	}
+
+	const uint PackedCoord = HairVisibilityNodeCoord[LocalOffset];
+	const uint2 PixelCoord = uint2((PackedCoord & 0xFFFF), ((PackedCoord >> 16) & 0xFFFF));
+	const float2 SvPosition = PixelCoord + float2(0.5f, 0.5f);
+
+	const float2 UV = (PixelCoord + float2(0.5f, 0.5f)) / float2(View.BufferSizeAndInvSize.xy);
+	const float2 ScreenPosition = (UV - View.ScreenPositionScaleBias.wz) / View.ScreenPositionScaleBias.xy;
+
+	if (LocalOffset > MaxVisibilityNodeCount)
+		return;
+
+	const FPackedHairSample PackedSample = HairVisibilityNodeData[LocalOffset];
+	const FHairSample Sample = UnpackHairSample(PackedSample);
+
+	FGBufferData GBuffer;
+	GBuffer.ShadingModelID = SHADINGMODELID_HAIR;
+	GBuffer.WorldNormal = Sample.Tangent;
+	GBuffer.BaseColor = Sample.BaseColor;
+	GBuffer.Roughness = Sample.Roughness;
+	GBuffer.Specular = Sample.Specular;
+	GBuffer.Depth = ConvertFromDeviceZ(Sample.Depth);
+	GBuffer.GBufferAO = 1;
+	GBuffer.CustomData = float4(HairDualScatteringRoughnessOverride, 0, Sample.Backlit, 0);
+	GBuffer.IndirectIrradiance = 1000000;
+
+	bool bEnabled = true;
+	const bool bIsMacroGroupCompatible = Sample.MacroGroupId == Voxel_MacroGroupId;
+#if USE_SCENE_COLOR
+	bEnabled = Sample.bScatterSceneLighting;
+#endif
+	if (bIsMacroGroupCompatible && bEnabled)
+	{
+		const float SampleCoverage = saturate(From8bitCoverage(Sample.Coverage8bit));
+#if USE_SCENE_COLOR
+		const float3 SkyLighting = ComputeThinTransmission(GBuffer, ScreenPosition, SvPosition, Sample.PrimitiveId);
+#else
+		const float3 SkyLighting = ReflectionEnvironment(GBuffer, ScreenPosition, SvPosition, Sample.PrimitiveId);
+#endif
+
+		FLightAccumulator LightAccumulator = (FLightAccumulator)0;
+		LightAccumulator_Add(LightAccumulator, SkyLighting, SkyLighting, 1.0f, false);
+		OutColor.xyz = LightAccumulator_GetResult(LightAccumulator) * SampleCoverage;
+		OutColor.a = SampleCoverage;
+	}
+}
+
+#endif // LIGHTING_PS
diff --git a/Engine/Shaders/Private/HairStrands/HairStrandsVisibilityComposeSubPixelPS.usf b/Engine/Shaders/Private/HairStrands/HairStrandsVisibilityComposeSubPixelPS.usf
--- a/Engine/Shaders/Private/HairStrands/HairStrandsVisibilityComposeSubPixelPS.usf
+++ b/Engine/Shaders/Private/HairStrands/HairStrandsVisibilityComposeSubPixelPS.usf
@@ -14,9 +14,8 @@ Texture2D<uint4> CategorisationTexture;
 
 void SubColorPS(
 	in FScreenVertexOutput Input,
-	out float4 OutColor  : SV_Target0
-	, out float  OutDepth : SV_Depth
-)
+	out float4 OutColor  : SV_Target0, 
+	out float  OutDepth : SV_Depth)
 {
 	const uint2 PixelCoord = floor(Input.Position.xy);
 	OutColor = SubPixelColorTexture.Load(uint3(PixelCoord, 0));
@@ -82,4 +81,45 @@ void MainPS(
 	OutGBufferB = float4(0, 0, 0, SHADINGMODELID_HAIR);
 }
 
-#endif // SHADER_PATCH
\ No newline at end of file
+#endif // SHADER_PATCH
+
+#if SHADER_COMPOSE_SAMPLE
+
+Texture2D<uint>		HairSampleCount;
+Texture2D<uint4>	HairCategorizationTexture;
+Texture2D<uint>		HairVisibilityNodeOffsetAndCount;
+Texture2D<float4>	HairLightingSampleBuffer;
+
+void ComposeSamplePS(
+	in FScreenVertexOutput Input,
+	out float4 OutColor : SV_Target0,
+	out float  OutDepth : SV_Depth)
+{
+	OutColor = 0.0f;
+	const uint3 PixelCoord = uint3(floor(Input.Position.xy), 0);
+
+	FNodeDesc NodeDesc = DecodeNodeDesc(HairVisibilityNodeOffsetAndCount.Load(PixelCoord));
+	if (NodeDesc.Count == 0)
+	{
+		discard;
+	}
+
+	const FCategorizationData CatData = DecodeCategorizationData(HairCategorizationTexture.Load(PixelCoord));
+	const float PixelCoverage = min(CatData.PixelCoverage, 1);
+
+	const uint Resolution = ceil(sqrt(HairSampleCount.Load(uint3(0, 0, 0))));
+
+	LOOP
+	for (uint SampleIt = 0; SampleIt < NodeDesc.Count; SampleIt++)
+	{
+		const uint LocalOffset = NodeDesc.Offset + SampleIt;
+		const uint2 Coord = uint2(LocalOffset % Resolution, LocalOffset / Resolution);
+		const float4 LightingSample = HairLightingSampleBuffer.Load(uint3(Coord,0));
+		OutColor.xyz += LightingSample.xyz * PixelCoverage;
+	}
+	OutColor.a = PixelCoverage;
+	OutDepth = CatData.ClosestDepth;
+
+}
+
+#endif // SHADER_COMPOSE_SAMPLE
diff --git a/Engine/Source/Runtime/Renderer/Private/DeferredShadingRenderer.h b/Engine/Source/Runtime/Renderer/Private/DeferredShadingRenderer.h
--- a/Engine/Source/Runtime/Renderer/Private/DeferredShadingRenderer.h
+++ b/Engine/Source/Runtime/Renderer/Private/DeferredShadingRenderer.h
@@ -412,7 +412,16 @@ private:
 	  * @param LightIndex The light's index into FScene::Lights
 	  * @return true if anything got rendered
 	  */
-	void RenderLight(FRHICommandList& RHICmdList, const FLightSceneInfo* LightSceneInfo, IPooledRenderTarget* ScreenShadowMaskTexture, IPooledRenderTarget* ScreenShadowMaskSubPixelTexture, struct FHairStrandsTransmittanceMaskData* TransmittanceMaskData, const struct FHairStrandsVisibilityViews* InHairVisibilityViews, bool bRenderOverlap, bool bIssueDrawEvent);
+	void RenderLight(FRHICommandList& RHICmdList, const FLightSceneInfo* LightSceneInfo, IPooledRenderTarget* ScreenShadowMaskTexture, IPooledRenderTarget* ScreenShadowMaskSubPixelTexture, struct FHairStrandsTransmittanceMaskData* TransmittanceMaskData, const struct FHairStrandsVisibilityViews* InHairVisibilityViews, bool bRenderOverlap, bool bIssueDrawEvent, bool bSupportHair);
+
+	/** Specialized version of RenderLight for hair (run lighting evaluation on at sub-pixel rate, without depth bound) */
+	void RenderLightForHair(
+		FRHICommandList& RHICmdList,
+		const FLightSceneInfo* LightSceneInfo,
+		IPooledRenderTarget* ScreenShadowMaskSubPixelTexture,
+		FHairStrandsTransmittanceMaskData* InTransmittanceMaskData,
+		const FHairStrandsVisibilityViews* InHairVisibilityViews,
+		bool bIssueDrawEvent);
 
 	/** Renders an array of simple lights using standard deferred shading. */
 	void RenderSimpleLightsStandardDeferred(FRHICommandListImmediate& RHICmdList, const FSimpleLightArray& SimpleLights);
diff --git a/Engine/Source/Runtime/Renderer/Private/HairStrands/HairStrandsComposition.cpp b/Engine/Source/Runtime/Renderer/Private/HairStrands/HairStrandsComposition.cpp
--- a/Engine/Source/Runtime/Renderer/Private/HairStrands/HairStrandsComposition.cpp
+++ b/Engine/Source/Runtime/Renderer/Private/HairStrands/HairStrandsComposition.cpp
@@ -19,6 +19,94 @@ static FAutoConsoleVariableRef CVarHairFastResolveVelocityThreshold(TEXT("r.Hair
 static int32 GHairPatchBufferDataBeforePostProcessing = 1;
 static FAutoConsoleVariableRef CVarHairPatchBufferDataBeforePostProcessing(TEXT("r.HairStrands.PatchMaterialData"), GHairPatchBufferDataBeforePostProcessing, TEXT("Patch the buffer with hair material data before post processing run. (default 1)."));
 
+class FHairVisibilityComposeSamplePS : public FGlobalShader
+{
+	DECLARE_GLOBAL_SHADER(FHairVisibilityComposeSamplePS);
+	SHADER_USE_PARAMETER_STRUCT(FHairVisibilityComposeSamplePS, FGlobalShader);
+
+	BEGIN_SHADER_PARAMETER_STRUCT(FParameters, )
+		SHADER_PARAMETER_RDG_TEXTURE(Texture2D, HairSampleCount)
+		SHADER_PARAMETER_RDG_TEXTURE(Texture2D, HairCategorizationTexture)
+		SHADER_PARAMETER_RDG_TEXTURE(Texture2D, HairVisibilityNodeOffsetAndCount)
+		SHADER_PARAMETER_RDG_TEXTURE(Texture2D, HairLightingSampleBuffer)
+		RENDER_TARGET_BINDING_SLOTS()
+	END_SHADER_PARAMETER_STRUCT()
+
+public:
+	static bool ShouldCompilePermutation(const FGlobalShaderPermutationParameters& Parameters) { return Parameters.Platform == SP_PCD3D_SM5; }
+	static void ModifyCompilationEnvironment(const FGlobalShaderPermutationParameters& Parameters, FShaderCompilerEnvironment& OutEnvironment)
+	{
+		FGlobalShader::ModifyCompilationEnvironment(Parameters, OutEnvironment);
+		OutEnvironment.SetDefine(TEXT("SHADER_COMPOSE_SAMPLE"), 1);
+	}
+};
+
+IMPLEMENT_GLOBAL_SHADER(FHairVisibilityComposeSamplePS, "/Engine/Private/HairStrands/HairStrandsVisibilityComposeSubPixelPS.usf", "ComposeSamplePS", SF_Pixel);
+
+static void AddHairVisibilityComposeSamplePass(
+	FRDGBuilder& GraphBuilder,
+	const FViewInfo& View,
+	const FHairStrandsVisibilityData& VisibilityData,
+	const FRDGTextureRef& CategorizationTexture,
+	FRDGTextureRef& OutColorTexture,
+	FRDGTextureRef& OutDepthTexture)
+{
+	check(VisibilityData.SampleLightingBuffer);
+
+	FRDGTextureRef SampleLightingBuffer = GraphBuilder.RegisterExternalTexture(VisibilityData.SampleLightingBuffer);
+	FRDGTextureRef NodeCount = GraphBuilder.RegisterExternalTexture(VisibilityData.NodeCount);
+
+	FHairVisibilityComposeSamplePS::FParameters* Parameters = GraphBuilder.AllocParameters<FHairVisibilityComposeSamplePS::FParameters>();
+	Parameters->HairSampleCount = NodeCount;
+	Parameters->HairCategorizationTexture = CategorizationTexture;
+	Parameters->HairVisibilityNodeOffsetAndCount = GraphBuilder.RegisterExternalTexture(VisibilityData.NodeIndex);
+	Parameters->HairLightingSampleBuffer = SampleLightingBuffer;
+	Parameters->RenderTargets[0] = FRenderTargetBinding(OutColorTexture, ERenderTargetLoadAction::ELoad);
+	Parameters->RenderTargets.DepthStencil = FDepthStencilBinding(OutDepthTexture, ERenderTargetLoadAction::ELoad, ERenderTargetLoadAction::ENoAction, FExclusiveDepthStencil::DepthWrite_StencilNop);
+
+	TShaderMapRef<FPostProcessVS> VertexShader(View.ShaderMap);
+	TShaderMapRef<FHairVisibilityComposeSamplePS> PixelShader(View.ShaderMap);
+	FGlobalShaderMap* GlobalShaderMap = View.ShaderMap;
+	const FIntRect Viewport = View.ViewRect;
+	const FIntPoint Resolution = OutColorTexture->Desc.Extent;
+	const FViewInfo* CapturedView = &View;
+
+	ClearUnusedGraphResources(PixelShader, Parameters);
+
+	GraphBuilder.AddPass(
+		RDG_EVENT_NAME("HairStrandsComposeSample"),
+		Parameters,
+		ERDGPassFlags::Raster,
+		[Parameters, VertexShader, PixelShader, Viewport, Resolution, CapturedView](FRHICommandList& RHICmdList)
+	{
+		FGraphicsPipelineStateInitializer GraphicsPSOInit;
+		RHICmdList.ApplyCachedRenderTargets(GraphicsPSOInit);
+//		GraphicsPSOInit.BlendState = TStaticBlendState<CW_RGBA, BO_Add, BF_One, BF_One, BO_Add, BF_One, BF_One>::GetRHI();
+		GraphicsPSOInit.BlendState = TStaticBlendState<CW_RGBA, BO_Add, BF_One, BF_InverseSourceAlpha, BO_Max, BF_SourceAlpha, BF_DestAlpha>::GetRHI();
+		GraphicsPSOInit.RasterizerState = TStaticRasterizerState<>::GetRHI();
+		GraphicsPSOInit.DepthStencilState = TStaticDepthStencilState<true, CF_DepthNearOrEqual>::GetRHI();
+		GraphicsPSOInit.BoundShaderState.VertexDeclarationRHI = GFilterVertexDeclaration.VertexDeclarationRHI;
+		GraphicsPSOInit.BoundShaderState.VertexShaderRHI = VertexShader.GetVertexShader();
+		GraphicsPSOInit.BoundShaderState.PixelShaderRHI = PixelShader.GetPixelShader();
+		GraphicsPSOInit.PrimitiveType = PT_TriangleList;
+		SetGraphicsPipelineState(RHICmdList, GraphicsPSOInit);
+
+		VertexShader->SetParameters(RHICmdList, CapturedView->ViewUniformBuffer);
+		RHICmdList.SetViewport(Viewport.Min.X, Viewport.Min.Y, 0.0f, Viewport.Max.X, Viewport.Max.Y, 1.0f);
+		SetShaderParameters(RHICmdList, PixelShader, PixelShader.GetPixelShader(), *Parameters);
+		DrawRectangle(
+			RHICmdList,
+			0, 0,
+			Viewport.Width(), Viewport.Height(),
+			Viewport.Min.X, Viewport.Min.Y,
+			Viewport.Width(), Viewport.Height(),
+			Viewport.Size(),
+			Resolution,
+			VertexShader,
+			EDRF_UseTriangleOptimization);
+	});
+}
+
 ///////////////////////////////////////////////////////////////////////////////////////////////////
 class FHairVisibilityComposeSubPixelPS : public FGlobalShader
 {
@@ -342,16 +430,29 @@ void RenderHairComposeSubPixel(
 					SceneColorSubPixelTexture,
 					SceneColorTexture);
 
-				// #hair_todo : compose partially covered hair with transparent surface: this can be done by 
-				// rendering quad(s) covering the hair at the correct depth. This will be sorted with other 
-				// transparent surface, which should make the overall sorting workable
-				AddHairVisibilityComposeSubPixelPass(
-					GraphBuilder,
-					View,
-					SceneColorSubPixelTexture,
-					RDGCategorisationTexture,
-					SceneColorTexture,
-					SceneColorDepth);
+				if (VisibilityData.SampleLightingBuffer)
+				{
+					AddHairVisibilityComposeSamplePass(
+						GraphBuilder,
+						View,
+						VisibilityData,
+						RDGCategorisationTexture,
+						SceneColorTexture,
+						SceneColorDepth);
+				}
+				else
+				{
+					// #hair_todo : compose partially covered hair with transparent surface: this can be done by 
+					// rendering quad(s) covering the hair at the correct depth. This will be sorted with other 
+					// transparent surface, which should make the overall sorting workable
+					AddHairVisibilityComposeSubPixelPass(
+						GraphBuilder,
+						View,
+						SceneColorSubPixelTexture,
+						RDGCategorisationTexture,
+						SceneColorTexture,
+						SceneColorDepth);
+				}
 
 				if (HairVisibilityViews.HairDatas[ViewIndex].VelocityTexture)
 				{
diff --git a/Engine/Source/Runtime/Renderer/Private/HairStrands/HairStrandsEnvironment.cpp b/Engine/Source/Runtime/Renderer/Private/HairStrands/HairStrandsEnvironment.cpp
--- a/Engine/Source/Runtime/Renderer/Private/HairStrands/HairStrandsEnvironment.cpp
+++ b/Engine/Source/Runtime/Renderer/Private/HairStrands/HairStrandsEnvironment.cpp
@@ -259,11 +259,9 @@ static void AddHairEnvironmentLightingComposePass(
 
 ///////////////////////////////////////////////////////////////////////////////////////////////////
 
-class FHairEnvironmentLightingCS : public FGlobalShader
+class FHairEnvironmentLighting
 {
-	DECLARE_GLOBAL_SHADER(FHairEnvironmentLightingCS);
-	SHADER_USE_PARAMETER_STRUCT(FHairEnvironmentLightingCS, FGlobalShader)
-
+public:
 	class FGroupSize : SHADER_PERMUTATION_INT("PERMUTATION_GROUP_SIZE", 2);
 	class FSampleSet : SHADER_PERMUTATION_INT("PERMUTATION_SAMPLESET", 2);
 	class FVirtualVoxel : SHADER_PERMUTATION_INT("VOXEL_TRAVERSAL_TYPE", 5);
@@ -321,10 +319,12 @@ class FHairEnvironmentLightingCS : public FGlobalShader
 		SHADER_PARAMETER(uint32, JitterSphericalIntegration)
 		SHADER_PARAMETER(float, HairDistanceThreshold)
 		SHADER_PARAMETER(uint32, bHairUseViewHairCount)
+		SHADER_PARAMETER(FIntPoint, MaxViewportResolution)
 
 		SHADER_PARAMETER_TEXTURE(Texture2D, PreIntegratedGF)
 		SHADER_PARAMETER_SAMPLER(SamplerState, PreIntegratedGFSampler)
 
+		SHADER_PARAMETER_RDG_TEXTURE(Texture2D, HairVisibilityNodeCount)
 		SHADER_PARAMETER_RDG_BUFFER_SRV(StructuredBuffer, HairVisibilityNodeData)
 		SHADER_PARAMETER_RDG_BUFFER_SRV(StructuredBuffer, HairVisibilityNodeCoord)
 
@@ -351,10 +351,64 @@ class FHairEnvironmentLightingCS : public FGlobalShader
 
 		SHADER_PARAMETER_STRUCT_INCLUDE(ShaderDrawDebug::FShaderDrawDebugParameters, ShaderDrawParameters)
 
-		END_SHADER_PARAMETER_STRUCT()
+	END_SHADER_PARAMETER_STRUCT()
+};
+
+class FHairEnvironmentLightingCS : public FGlobalShader
+{
+	DECLARE_GLOBAL_SHADER(FHairEnvironmentLightingCS);
+	SHADER_USE_PARAMETER_STRUCT(FHairEnvironmentLightingCS, FGlobalShader)
+	using FPermutationDomain = FHairEnvironmentLighting::FPermutationDomain;
+
+	BEGIN_SHADER_PARAMETER_STRUCT(FParameters, )
+		SHADER_PARAMETER_STRUCT_INCLUDE(FHairEnvironmentLighting::FParameters, Common)
+	END_SHADER_PARAMETER_STRUCT()
+
+	static void ModifyCompilationEnvironment(const FGlobalShaderPermutationParameters& Parameters, FShaderCompilerEnvironment& OutEnvironment)
+	{
+		FHairEnvironmentLighting::ModifyCompilationEnvironment(Parameters, OutEnvironment);
+		OutEnvironment.SetDefine(TEXT("LIGHTING_CS"), 1);
+	}
+};
+
+class FHairEnvironmentLightingVS : public FGlobalShader
+{
+	DECLARE_GLOBAL_SHADER(FHairEnvironmentLightingVS);
+	SHADER_USE_PARAMETER_STRUCT(FHairEnvironmentLightingVS, FGlobalShader)
+	using FPermutationDomain = FHairEnvironmentLighting::FPermutationDomain;
+
+	BEGIN_SHADER_PARAMETER_STRUCT(FParameters, )
+		SHADER_PARAMETER_STRUCT_INCLUDE(FHairEnvironmentLighting::FParameters, Common)
+	END_SHADER_PARAMETER_STRUCT()
+
+	static void ModifyCompilationEnvironment(const FGlobalShaderPermutationParameters& Parameters, FShaderCompilerEnvironment& OutEnvironment)
+	{
+		FHairEnvironmentLighting::ModifyCompilationEnvironment(Parameters, OutEnvironment);
+		OutEnvironment.SetDefine(TEXT("LIGHTING_VS"), 1);
+	}
+};
+
+class FHairEnvironmentLightingPS : public FGlobalShader
+{
+	DECLARE_GLOBAL_SHADER(FHairEnvironmentLightingPS);
+	SHADER_USE_PARAMETER_STRUCT(FHairEnvironmentLightingPS, FGlobalShader)
+	using FPermutationDomain = FHairEnvironmentLighting::FPermutationDomain;
+
+	BEGIN_SHADER_PARAMETER_STRUCT(FParameters, )
+		SHADER_PARAMETER_STRUCT_INCLUDE(FHairEnvironmentLighting::FParameters, Common)
+		RENDER_TARGET_BINDING_SLOTS()
+	END_SHADER_PARAMETER_STRUCT()
+
+	static void ModifyCompilationEnvironment(const FGlobalShaderPermutationParameters& Parameters, FShaderCompilerEnvironment& OutEnvironment)
+	{
+		FHairEnvironmentLighting::ModifyCompilationEnvironment(Parameters, OutEnvironment);
+		OutEnvironment.SetDefine(TEXT("LIGHTING_PS"), 1);
+	}
 };
 
 IMPLEMENT_GLOBAL_SHADER(FHairEnvironmentLightingCS, "/Engine/Private/HairStrands/HairStrandsEnvironmentLighting.usf", "MainCS", SF_Compute);
+IMPLEMENT_GLOBAL_SHADER(FHairEnvironmentLightingPS, "/Engine/Private/HairStrands/HairStrandsEnvironmentLighting.usf", "MainPS", SF_Pixel);
+IMPLEMENT_GLOBAL_SHADER(FHairEnvironmentLightingVS, "/Engine/Private/HairStrands/HairStrandsEnvironmentLighting.usf", "MainVS", SF_Vertex);
 
 static FRDGBufferRef AddHairStrandsEnvironmentLightingPassCS(
 	FRDGBuilder& GraphBuilder,
@@ -380,7 +434,111 @@ static FRDGBufferRef AddHairStrandsEnvironmentLightingPassCS(
 	// Render the reflection environment with tiled deferred culling
 	const bool bHasBoxCaptures = (View.NumBoxReflectionCaptures > 0);
 	const bool bHasSphereCaptures = (View.NumSphereReflectionCaptures > 0);
-	FHairEnvironmentLightingCS::FParameters* PassParameters = GraphBuilder.AllocParameters<FHairEnvironmentLightingCS::FParameters>();
+	FHairEnvironmentLighting::FParameters PassParameters;
+
+	const FHairLUT InHairLUT = GetHairLUT(GraphBuilder.RHICmdList, View);
+	PassParameters.HairEnergyLUTTexture = GraphBuilder.RegisterExternalTexture(InHairLUT.Textures[HairLUTType_MeanEnergy], TEXT("HairMeanEnergyLUTTexture"));;
+	PassParameters.HairScatteringLUTTexture = GraphBuilder.RegisterExternalTexture(InHairLUT.Textures[HairLUTType_DualScattering], TEXT("HairScatteringEnergyLUTTexture"));
+	PassParameters.HairLUTSampler = TStaticSamplerState<SF_Bilinear, AM_Clamp, AM_Clamp, AM_Clamp>::GetRHI();
+	PassParameters.Voxel_MacroGroupId = MacroGroupData.MacroGroupId;
+	if (!bUseVirtualVoxel)
+	{
+		PassParameters.Voxel_MinAABB = MacroGroupData.GetMinBound();
+		PassParameters.Voxel_MaxAABB = MacroGroupData.GetMaxBound();
+		PassParameters.Voxel_Resolution = MacroGroupData.GetResolution();
+		PassParameters.Voxel_DensityTexture = GraphBuilder.RegisterExternalTexture(MacroGroupData.VoxelResources.DensityTexture);
+	}
+	const bool bUseSceneColor = SceneColorTexture != nullptr;
+	if (bUseSceneColor)
+	{
+		PassParameters.SceneColorTexture = SceneColorTexture;
+		PassParameters.HairCategorizationTexture = GraphBuilder.RegisterExternalTexture(VisibilityData.CategorizationTexture ? VisibilityData.CategorizationTexture : GSystemTextures.BlackDummy);
+	}
+	PassParameters.MaxViewportResolution = VisibilityData.SampleLightingViewportResolution;
+	PassParameters.HairVisibilityNodeCount = GraphBuilder.RegisterExternalTexture(VisibilityData.NodeCount);
+	PassParameters.Voxel_DensityScale = GetHairStrandsVoxelizationDensityScale();
+	PassParameters.Voxel_DepthBiasScale = GetHairStrandsVoxelizationDepthBiasScale();
+	PassParameters.Voxel_TanConeAngle = FMath::Tan(FMath::DegreesToRadians(GetHairStrandsSkyLightingConeAngle()));
+	PassParameters.HairDistanceThreshold = FMath::Max(GHairStrandsSkyLightingDistanceThreshold, 1.f);
+	PassParameters.bHairUseViewHairCount = VisibilityData.ViewHairCountTexture ? 1 : 0;
+	PassParameters.MaxVisibilityNodeCount = VisibilityData.NodeData->Desc.NumElements;
+	PassParameters.MultipleScatterSampleCount = FMath::Max(uint32(GHairStrandsSkyLightingSampleCount), 1u);
+	PassParameters.JitterSphericalIntegration = GHairStrandsSkyLightingJitterSphericalIntegration ? 1 : 0;
+	PassParameters.HairComponents = ToBitfield(GetHairComponents());
+	PassParameters.HairDualScatteringRoughnessOverride = GetHairDualScatteringRoughnessOverride();
+	PassParameters.PreIntegratedGF = GSystemTextures.PreintegratedGF->GetRenderTargetItem().ShaderResourceTexture;
+	PassParameters.PreIntegratedGFSampler = TStaticSamplerState<SF_Bilinear, AM_Clamp, AM_Clamp, AM_Clamp>::GetRHI();
+	PassParameters.HairCountTexture = GraphBuilder.RegisterExternalTexture(VisibilityData.ViewHairCountTexture ? VisibilityData.ViewHairCountTexture : GSystemTextures.BlackDummy);
+	PassParameters.SceneTextures = SceneTextures;
+	PassParameters.VirtualVoxel = MacroGroupDatas.VirtualVoxelResources.UniformBuffer;
+	SetupSceneTextureSamplers(&PassParameters.SceneTextureSamplers);
+	PassParameters.ViewUniformBuffer = View.ViewUniformBuffer;
+	PassParameters.ReflectionCaptureData = View.ReflectionCaptureUniformBuffer;
+	{
+		FReflectionUniformParameters ReflectionUniformParameters;
+		SetupReflectionUniformParameters(View, ReflectionUniformParameters);
+		PassParameters.ReflectionsParameters = CreateUniformBufferImmediate(ReflectionUniformParameters, UniformBuffer_SingleDraw);
+	}
+	PassParameters.ForwardLightData = View.ForwardLightingResources->ForwardLightDataUniformBuffer;
+	PassParameters.IndirectArgsBuffer = IndirectArgsBuffer;
+	PassParameters.OutLightingBuffer = GraphBuilder.CreateUAV(FRDGBufferUAVDesc(OutBuffer, PF_FloatRGBA));
+
+	if (ShaderDrawDebug::IsShaderDrawDebugEnabled(View))
+	{
+		ShaderDrawDebug::SetParameters(GraphBuilder, View.ShaderDrawData, PassParameters.ShaderDrawParameters);
+	}
+
+	// Bind hair data
+	FRDGBufferRef InHairVisibilityNodeData = GraphBuilder.RegisterExternalBuffer(VisibilityData.NodeData, TEXT("HairVisibilityNodeData"));
+	FRDGBufferRef InHairVisibilityNodeCoord = GraphBuilder.RegisterExternalBuffer(VisibilityData.NodeCoord, TEXT("HairVisibilityNodeCoord"));
+	PassParameters.HairVisibilityNodeData = GraphBuilder.CreateSRV(InHairVisibilityNodeData);
+	PassParameters.HairVisibilityNodeCoord = GraphBuilder.CreateSRV(InHairVisibilityNodeCoord);
+
+
+	check(NodeGroupSize == 64 || NodeGroupSize == 32);
+	FHairEnvironmentLightingCS::FPermutationDomain PermutationVector;
+	PermutationVector.Set<FHairEnvironmentLighting::FGroupSize>(NodeGroupSize == 64 ? 0 : (NodeGroupSize == 32 ? 1 : 2));
+	PermutationVector.Set<FHairEnvironmentLighting::FSampleSet>(PassParameters.MultipleScatterSampleCount <= 16 ? 0 : 1);
+	PermutationVector.Set<FHairEnvironmentLighting::FVirtualVoxel>(bUseVirtualVoxel ? VoxelTraversalType : 0);
+	PermutationVector.Set<FHairEnvironmentLighting::FUseSceneColor>(bUseSceneColor ? 1 : 0);
+	PermutationVector = FHairEnvironmentLighting::RemapPermutation(PermutationVector);
+
+	FHairEnvironmentLightingCS::FParameters* Parameters = GraphBuilder.AllocParameters<FHairEnvironmentLightingCS::FParameters>();
+	Parameters->Common = PassParameters;
+
+	TShaderMapRef<FHairEnvironmentLightingCS> ComputeShader(View.ShaderMap, PermutationVector);
+	FComputeShaderUtils::AddPass(
+		GraphBuilder,
+		RDG_EVENT_NAME("HairStrandsEnvironmentCS %dx%d", View.ViewRect.Width(), View.ViewRect.Height()),
+		ComputeShader,
+		Parameters,
+		IndirectArgsBuffer,
+		0);
+
+	return OutBuffer;
+}
+
+static void AddHairStrandsEnvironmentLightingPassPS(
+	FRDGBuilder& GraphBuilder,
+	const FViewInfo& View,
+	const FHairStrandsVisibilityData& VisibilityData,
+	const FHairStrandsMacroGroupDatas& MacroGroupDatas,
+	const FHairStrandsMacroGroupData& MacroGroupData,
+	const FRDGTextureRef SceneColorTexture,
+	const uint32 NodeGroupSize,
+	FRDGBufferRef IndirectArgsBuffer)
+{
+	FSceneTextureParameters SceneTextures;
+	SetupSceneTextureParameters(GraphBuilder, &SceneTextures);
+
+	const bool bUseVirtualVoxel = MacroGroupDatas.VirtualVoxelResources.IsValid();
+	const int32 VoxelTraversalType = FMath::Clamp(GHairSkylightingVirtualVoxel, 0, 4);
+
+	// Render the reflection environment with tiled deferred culling
+	const bool bHasBoxCaptures = (View.NumBoxReflectionCaptures > 0);
+	const bool bHasSphereCaptures = (View.NumSphereReflectionCaptures > 0);
+	FHairEnvironmentLightingPS::FParameters* ParametersPS = GraphBuilder.AllocParameters<FHairEnvironmentLightingPS::FParameters>();
+	FHairEnvironmentLighting::FParameters* PassParameters = &ParametersPS->Common;
 
 	const FHairLUT InHairLUT = GetHairLUT(GraphBuilder.RHICmdList, View);
 	PassParameters->HairEnergyLUTTexture = GraphBuilder.RegisterExternalTexture(InHairLUT.Textures[HairLUTType_MeanEnergy], TEXT("HairMeanEnergyLUTTexture"));;
@@ -400,6 +558,8 @@ static FRDGBufferRef AddHairStrandsEnvironmentLightingPassCS(
 		PassParameters->SceneColorTexture = SceneColorTexture;
 		PassParameters->HairCategorizationTexture = GraphBuilder.RegisterExternalTexture(VisibilityData.CategorizationTexture ? VisibilityData.CategorizationTexture : GSystemTextures.BlackDummy);
 	}
+	PassParameters->MaxViewportResolution = VisibilityData.SampleLightingViewportResolution;
+	PassParameters->HairVisibilityNodeCount = GraphBuilder.RegisterExternalTexture(VisibilityData.NodeCount);
 	PassParameters->Voxel_DensityScale = GetHairStrandsVoxelizationDensityScale();
 	PassParameters->Voxel_DepthBiasScale = GetHairStrandsVoxelizationDepthBiasScale();
 	PassParameters->Voxel_TanConeAngle = FMath::Tan(FMath::DegreesToRadians(GetHairStrandsSkyLightingConeAngle()));
@@ -425,7 +585,7 @@ static FRDGBufferRef AddHairStrandsEnvironmentLightingPassCS(
 	}
 	PassParameters->ForwardLightData = View.ForwardLightingResources->ForwardLightDataUniformBuffer;
 	PassParameters->IndirectArgsBuffer = IndirectArgsBuffer;
-	PassParameters->OutLightingBuffer = GraphBuilder.CreateUAV(FRDGBufferUAVDesc(OutBuffer, PF_FloatRGBA));
+	PassParameters->OutLightingBuffer = nullptr;
 
 	if (ShaderDrawDebug::IsShaderDrawDebugEnabled(View))
 	{
@@ -439,23 +599,47 @@ static FRDGBufferRef AddHairStrandsEnvironmentLightingPassCS(
 	PassParameters->HairVisibilityNodeCoord = GraphBuilder.CreateSRV(InHairVisibilityNodeCoord);
 
 	check(NodeGroupSize == 64 || NodeGroupSize == 32);
-	FHairEnvironmentLightingCS::FPermutationDomain PermutationVector;
-	PermutationVector.Set<FHairEnvironmentLightingCS::FGroupSize>(NodeGroupSize == 64 ? 0 : (NodeGroupSize == 32 ? 1 : 2));
-	PermutationVector.Set<FHairEnvironmentLightingCS::FSampleSet>(PassParameters->MultipleScatterSampleCount <= 16 ? 0 : 1);
-	PermutationVector.Set<FHairEnvironmentLightingCS::FVirtualVoxel>(bUseVirtualVoxel ? VoxelTraversalType : 0);
-	PermutationVector.Set<FHairEnvironmentLightingCS::FUseSceneColor>(bUseSceneColor ? 1 : 0);
-	PermutationVector = FHairEnvironmentLightingCS::RemapPermutation(PermutationVector);
+	FHairEnvironmentLightingPS::FPermutationDomain PermutationVector;
+	PermutationVector.Set<FHairEnvironmentLighting::FGroupSize>(NodeGroupSize == 64 ? 0 : (NodeGroupSize == 32 ? 1 : 2));
+	PermutationVector.Set<FHairEnvironmentLighting::FSampleSet>(PassParameters->MultipleScatterSampleCount <= 16 ? 0 : 1);
+	PermutationVector.Set<FHairEnvironmentLighting::FVirtualVoxel>(bUseVirtualVoxel ? VoxelTraversalType : 0);
+	PermutationVector.Set<FHairEnvironmentLighting::FUseSceneColor>(bUseSceneColor ? 1 : 0);
+	PermutationVector = FHairEnvironmentLighting::RemapPermutation(PermutationVector);
+
+	FIntPoint ViewportResolution = VisibilityData.SampleLightingViewportResolution;
+	const FViewInfo* CapturedView = &View;
+	TShaderMapRef<FHairEnvironmentLightingVS> VertexShader(View.ShaderMap, PermutationVector);
+	TShaderMapRef<FHairEnvironmentLightingPS> PixelShader(View.ShaderMap, PermutationVector);
 
-	TShaderMapRef<FHairEnvironmentLightingCS> ComputeShader(View.ShaderMap, PermutationVector);
-	FComputeShaderUtils::AddPass(
-		GraphBuilder,
-		RDG_EVENT_NAME("HairStrandsEnvironmentCS %dx%d", View.ViewRect.Width(), View.ViewRect.Height()),
-		ComputeShader,
-		PassParameters,
-		IndirectArgsBuffer,
-		0);
+	ParametersPS->RenderTargets[0] = FRenderTargetBinding(GraphBuilder.RegisterExternalTexture(VisibilityData.SampleLightingBuffer), ERenderTargetLoadAction::ELoad);
 
-	return OutBuffer;
+	GraphBuilder.AddPass(
+		RDG_EVENT_NAME("HairEnvLightingPS"),
+		ParametersPS,
+		ERDGPassFlags::Raster,
+		[ParametersPS, VertexShader, PixelShader, ViewportResolution, CapturedView](FRHICommandList& RHICmdList)
+	{
+		FHairEnvironmentLightingVS::FParameters ParametersVS;
+		ParametersVS.Common = ParametersPS->Common;
+
+		FGraphicsPipelineStateInitializer GraphicsPSOInit;
+		RHICmdList.ApplyCachedRenderTargets(GraphicsPSOInit);
+		GraphicsPSOInit.BlendState = TStaticBlendState<CW_RGBA, BO_Add, BF_One, BF_One, BO_Max, BF_SourceAlpha, BF_DestAlpha>::GetRHI();
+		GraphicsPSOInit.RasterizerState = TStaticRasterizerState<>::GetRHI();
+		GraphicsPSOInit.DepthStencilState = TStaticDepthStencilState<false, CF_Always>::GetRHI();
+		GraphicsPSOInit.BoundShaderState.VertexDeclarationRHI = GFilterVertexDeclaration.VertexDeclarationRHI;
+		GraphicsPSOInit.BoundShaderState.VertexShaderRHI = VertexShader.GetVertexShader();
+		GraphicsPSOInit.BoundShaderState.PixelShaderRHI = PixelShader.GetPixelShader();
+		GraphicsPSOInit.PrimitiveType = PT_TriangleList;
+		SetGraphicsPipelineState(RHICmdList, GraphicsPSOInit);
+
+		SetShaderParameters(RHICmdList, VertexShader, VertexShader.GetVertexShader(), ParametersVS);
+		SetShaderParameters(RHICmdList, PixelShader, PixelShader.GetPixelShader(), *ParametersPS);
+
+		RHICmdList.SetViewport(0, 0, 0.0f, ViewportResolution.X, ViewportResolution.Y, 1.0f);
+		RHICmdList.SetStreamSource(0, nullptr, 0);
+		RHICmdList.DrawPrimitive(0, 1, 1);
+	});
 }
 
 ///////////////////////////////////////////////////////////////////////////////////////////////////
@@ -490,10 +674,21 @@ void RenderHairStrandsSceneColorScattering(
 		FRDGTextureRef SceneColorSubPixelTexture = GraphBuilder.RegisterExternalTexture(SceneContext.GetSceneColorSubPixel());
 		FRDGBufferRef NodeIndirectArgBuffer = GraphBuilder.RegisterExternalBuffer(VisibilityData.NodeIndirectArg, TEXT("HairNodeIndirectArgBuffer"));
 
-		for (const FHairStrandsMacroGroupData& MacroGroupData : HairDatas->MacroGroupsPerViews.Views[ViewIndex].Datas)
+		const bool bUsePixelAccumulation = HairDatas->HairVisibilityViews.HairDatas[ViewIndex].SampleLightingBuffer != nullptr;
+		if (bUsePixelAccumulation)
 		{
-			FRDGBufferRef SampleLightingBuffer = AddHairStrandsEnvironmentLightingPassCS(GraphBuilder, View, VisibilityData, MacroGroupDatas, MacroGroupData, SceneColorTexture, VisibilityData.NodeGroupSize, NodeIndirectArgBuffer);
-			AddHairEnvironmentLightingComposePass(GraphBuilder, View, VisibilityData, MacroGroupData, SampleLightingBuffer, SceneColorTexture, SceneColorSubPixelTexture);
+			for (const FHairStrandsMacroGroupData& MacroGroupData : HairDatas->MacroGroupsPerViews.Views[ViewIndex].Datas)
+			{
+				AddHairStrandsEnvironmentLightingPassPS(GraphBuilder, View, VisibilityData, MacroGroupDatas, MacroGroupData, SceneColorTexture, VisibilityData.NodeGroupSize, NodeIndirectArgBuffer);
+			}
+		}
+		else
+		{
+			for (const FHairStrandsMacroGroupData& MacroGroupData : HairDatas->MacroGroupsPerViews.Views[ViewIndex].Datas)
+			{
+				FRDGBufferRef SampleLightingBuffer = AddHairStrandsEnvironmentLightingPassCS(GraphBuilder, View, VisibilityData, MacroGroupDatas, MacroGroupData, SceneColorTexture, VisibilityData.NodeGroupSize, NodeIndirectArgBuffer);
+				AddHairEnvironmentLightingComposePass(GraphBuilder, View, VisibilityData, MacroGroupData, SampleLightingBuffer, SceneColorTexture, SceneColorSubPixelTexture);
+			}
 		}
 		GraphBuilder.Execute();
 	}
@@ -522,11 +717,23 @@ void RenderHairStrandsEnvironmentLighting(
 	FRDGBufferRef NodeIndirectArgBuffer = GraphBuilder.RegisterExternalBuffer(VisibilityData.NodeIndirectArg, TEXT("HairNodeIndirectArgBuffer"));
 
 	const FViewInfo& View = Views[ViewIndex];
+
 	const FHairStrandsMacroGroupDatas& MacroGroupDatas = HairDatas->MacroGroupsPerViews.Views[ViewIndex];
-	for (const FHairStrandsMacroGroupData& MacroGroupData : HairDatas->MacroGroupsPerViews.Views[ViewIndex].Datas)
+	const bool bUsePixelAccumulation = HairDatas->HairVisibilityViews.HairDatas[ViewIndex].SampleLightingBuffer != nullptr;
+	if (bUsePixelAccumulation)
 	{
-		FRDGBufferRef SampleLightingBuffer = AddHairStrandsEnvironmentLightingPassCS(GraphBuilder, View, VisibilityData, MacroGroupDatas, MacroGroupData, SceneColorTexture, VisibilityData.NodeGroupSize, NodeIndirectArgBuffer);
-		AddHairEnvironmentLightingComposePass(GraphBuilder, View, VisibilityData, MacroGroupData, SampleLightingBuffer, SceneColorTexture, SceneColorSubPixelTexture);
+		for (const FHairStrandsMacroGroupData& MacroGroupData : HairDatas->MacroGroupsPerViews.Views[ViewIndex].Datas)
+		{
+			AddHairStrandsEnvironmentLightingPassPS(GraphBuilder, View, VisibilityData, MacroGroupDatas, MacroGroupData, nullptr /*SceneColorTexture*/, VisibilityData.NodeGroupSize, NodeIndirectArgBuffer);
+		}
+	}
+	else
+	{
+		for (const FHairStrandsMacroGroupData& MacroGroupData : HairDatas->MacroGroupsPerViews.Views[ViewIndex].Datas)
+		{
+			FRDGBufferRef SampleLightingBuffer = AddHairStrandsEnvironmentLightingPassCS(GraphBuilder, View, VisibilityData, MacroGroupDatas, MacroGroupData, nullptr /*SceneColorTexture*/, VisibilityData.NodeGroupSize, NodeIndirectArgBuffer);
+			AddHairEnvironmentLightingComposePass(GraphBuilder, View, VisibilityData, MacroGroupData, SampleLightingBuffer, SceneColorTexture, SceneColorSubPixelTexture);
+		}
 	}
 }
 
diff --git a/Engine/Source/Runtime/Renderer/Private/HairStrands/HairStrandsVisibility.cpp b/Engine/Source/Runtime/Renderer/Private/HairStrands/HairStrandsVisibility.cpp
--- a/Engine/Source/Runtime/Renderer/Private/HairStrands/HairStrandsVisibility.cpp
+++ b/Engine/Source/Runtime/Renderer/Private/HairStrands/HairStrandsVisibility.cpp
@@ -66,6 +66,11 @@ static FAutoConsoleVariableRef CVarHairStrandsVisibilityComputeRasterSampleCount
 static float GHairStrandsFullCoverageThreshold = 0.98f;
 static FAutoConsoleVariableRef CVarHairStrandsFullCoverageThreshold(TEXT("r.HairStrands.Visibility.FullCoverageThreshold"), GHairStrandsFullCoverageThreshold, TEXT("Define the coverage threshold at which a pixel is considered fully covered."));
 
+static TAutoConsoleVariable<int32> CVarHairStrandsSplitLighting(
+	TEXT("r.HairStrands.SplitLighting"), 0,
+	TEXT(""),
+	ECVF_RenderThreadSafe);
+
 /////////////////////////////////////////////////////////////////////////////////////////
 
 namespace HairStrandsVisibilityInternal
@@ -1199,12 +1204,14 @@ static void AddHairVisibilityPrimitiveIdCompactionPass(
 	const FHairStrandsMacroGroupDatas& MacroGroupDatas,
 	const uint32 NodeGroupSize,
 	FHairVisibilityPrimitiveIdCompactionCS::FParameters* PassParameters,
+	FRDGTextureRef& OutCompactCounter,
 	FRDGTextureRef& OutCompactNodeIndex,
 	FRDGBufferRef& OutCompactNodeData,
 	FRDGBufferRef& OutCompactNodeCoord,
 	FRDGTextureRef& OutCategorizationTexture,
 	FRDGTextureRef& OutVelocityTexture,
-	FRDGBufferRef& OutIndirectArgsBuffer)
+	FRDGBufferRef& OutIndirectArgsBuffer,
+	uint32& OutMaxRenderNodeCount)
 {
 	FIntPoint Resolution;
 	if (bUsePPLL)
@@ -1234,8 +1241,6 @@ static void AddHairVisibilityPrimitiveIdCompactionPass(
 		}
 	}
 
-
-	FRDGTextureRef CompactCounter;
 	{
 		FRDGTextureDesc Desc;
 		Desc.Extent.X = 1;
@@ -1247,7 +1252,7 @@ static void AddHairVisibilityPrimitiveIdCompactionPass(
 		Desc.Flags = TexCreate_None;
 		Desc.TargetableFlags = TexCreate_UAV | TexCreate_ShaderResource;
 		Desc.ClearValue = FClearValueBinding(0);
-		CompactCounter = GraphBuilder.CreateTexture(Desc, TEXT("HairVisibilityCompactCounter"));
+		OutCompactCounter = GraphBuilder.CreateTexture(Desc, TEXT("HairVisibilityCompactCounter"));
 	}
 
 	{
@@ -1273,7 +1278,7 @@ static void AddHairVisibilityPrimitiveIdCompactionPass(
 	}
 
 	const uint32 ClearValues[4] = { 0u,0u,0u,0u };
-	AddClearUAVPass(GraphBuilder, GraphBuilder.CreateUAV(CompactCounter), ClearValues);
+	AddClearUAVPass(GraphBuilder, GraphBuilder.CreateUAV(OutCompactCounter), ClearValues);
 	AddClearUAVPass(GraphBuilder, GraphBuilder.CreateUAV(OutCompactNodeIndex), ClearValues);
 	AddClearUAVPass(GraphBuilder, GraphBuilder.CreateUAV(OutCategorizationTexture), ClearValues);
 
@@ -1311,7 +1316,7 @@ static void AddHairVisibilityPrimitiveIdCompactionPass(
 	PassParameters->CosTangentThreshold = FMath::Cos(FMath::DegreesToRadians(FMath::Clamp(GHairStrandsMaterialCompactionTangentThreshold, 0.f, 90.f)));
 	PassParameters->SceneTexturesStruct = CreateUniformBufferImmediate(SceneTextures, EUniformBufferUsage::UniformBuffer_SingleDraw);
 	PassParameters->ViewUniformBuffer = View.ViewUniformBuffer;
-	PassParameters->OutCompactNodeCounter = GraphBuilder.CreateUAV(CompactCounter);
+	PassParameters->OutCompactNodeCounter = GraphBuilder.CreateUAV(OutCompactCounter);
 	PassParameters->OutCompactNodeIndex = GraphBuilder.CreateUAV(OutCompactNodeIndex);
 	PassParameters->OutCompactNodeData = GraphBuilder.CreateUAV(OutCompactNodeData);
 	PassParameters->OutCompactNodeCoord = GraphBuilder.CreateUAV(OutCompactNodeCoord);
@@ -1342,7 +1347,8 @@ static void AddHairVisibilityPrimitiveIdCompactionPass(
 		PassParameters,
 		FComputeShaderUtils::GetGroupCount(RectResolution, GroupSize));
 
-	OutIndirectArgsBuffer = AddCopyIndirectArgPass(GraphBuilder, &View, NodeGroupSize, 1, CompactCounter);
+	OutIndirectArgsBuffer = AddCopyIndirectArgPass(GraphBuilder, &View, NodeGroupSize, 1, OutCompactCounter);
+	OutMaxRenderNodeCount = MaxRenderNodeCount;
 }
 
 
@@ -2368,6 +2374,7 @@ FHairStrandsVisibilityViews RenderHairStrandsVisibilityBuffer(
 			FRDGTextureRef CategorizationTexture = nullptr;
 			FRDGTextureRef CompactNodeIndex = nullptr;
 			FRDGBufferRef  CompactNodeData = nullptr;
+			FRDGTextureRef NodeCounter = nullptr;
 			if (RenderMode == HairVisibilityRenderMode_MSAA)
 			{
 				const bool bIsVisiblityEnable = GHairStrandsVisibilityMaterialPass > 0;
@@ -2431,12 +2438,14 @@ FHairStrandsVisibilityViews RenderHairStrandsVisibilityBuffer(
 						MacroGroupDatas,
 						VisibilityData.NodeGroupSize,
 						PassParameters,
+						NodeCounter,
 						CompactNodeIndex,
 						CompactNodeData,
 						CompactNodeCoord,
 						CategorizationTexture,
 						SceneVelocityTexture,
-						IndirectArgsBuffer);
+						IndirectArgsBuffer,
+						VisibilityData.MaxNodeCount);
 
 					if (bIsVisiblityEnable)
 					{
@@ -2471,6 +2480,7 @@ FHairStrandsVisibilityViews RenderHairStrandsVisibilityBuffer(
 					GraphBuilder.QueueBufferExtraction(CompactNodeData,			&VisibilityData.NodeData,					FRDGResourceState::EAccess::Read, FRDGResourceState::EPipeline::Graphics);
 					GraphBuilder.QueueBufferExtraction(CompactNodeCoord,		&VisibilityData.NodeCoord,					FRDGResourceState::EAccess::Read, FRDGResourceState::EPipeline::Graphics);
 					GraphBuilder.QueueBufferExtraction(IndirectArgsBuffer,		&VisibilityData.NodeIndirectArg,			FRDGResourceState::EAccess::Read, FRDGResourceState::EPipeline::Compute);
+					GraphBuilder.QueueTextureExtraction(NodeCounter, 			&VisibilityData.NodeCount);
 				}
 
 				// View transmittance depth test needs to happen before the scene depth is patched with the hair depth (for fully-covered-by-hair pixels)
@@ -2540,12 +2550,14 @@ FHairStrandsVisibilityViews RenderHairStrandsVisibilityBuffer(
 						MacroGroupDatas,
 						VisibilityData.NodeGroupSize,
 						PassParameters,
+						NodeCounter,
 						CompactNodeIndex,
 						CompactNodeData,
 						CompactNodeCoord,
 						CategorizationTexture,
 						SceneVelocityTexture,
-						IndirectArgsBuffer);
+						IndirectArgsBuffer,
+						VisibilityData.MaxNodeCount);
 
 					VisibilityData.MaxSampleCount = GetPPLLMaxRenderNodePerPixel();
 					GraphBuilder.QueueTextureExtraction(CompactNodeIndex, &VisibilityData.NodeIndex);
@@ -2553,6 +2565,7 @@ FHairStrandsVisibilityViews RenderHairStrandsVisibilityBuffer(
 					GraphBuilder.QueueBufferExtraction(CompactNodeData, &VisibilityData.NodeData, FRDGResourceState::EAccess::Read, FRDGResourceState::EPipeline::Graphics);
 					GraphBuilder.QueueBufferExtraction(CompactNodeCoord, &VisibilityData.NodeCoord, FRDGResourceState::EAccess::Read, FRDGResourceState::EPipeline::Graphics);
 					GraphBuilder.QueueBufferExtraction(IndirectArgsBuffer, &VisibilityData.NodeIndirectArg, FRDGResourceState::EAccess::Read, FRDGResourceState::EPipeline::Compute);
+					GraphBuilder.QueueTextureExtraction(NodeCounter, &VisibilityData.NodeCount);
 				}
 
 				AddHairVisibilityColorAndDepthPatchPass(
@@ -2598,12 +2611,36 @@ FHairStrandsVisibilityViews RenderHairStrandsVisibilityBuffer(
 
 			GraphBuilder.Execute();
 
+			const bool bUseSplitLighting = CVarHairStrandsSplitLighting.GetValueOnRenderThread() >  0;
+			if (bUseSplitLighting)
+			{			
+				const uint32 SampleTextureResolution = FMath::CeilToInt(FMath::Sqrt(VisibilityData.MaxNodeCount));
+				VisibilityData.SampleLightingViewportResolution = FIntPoint(SampleTextureResolution, SampleTextureResolution);
+				FPooledRenderTargetDesc Desc = FPooledRenderTargetDesc::Create2DDesc(
+					VisibilityData.SampleLightingViewportResolution,
+					PF_FloatRGBA,
+					FClearValueBinding::Black,
+					TexCreate_None, 
+					TexCreate_UAV | TexCreate_ShaderResource | TexCreate_RenderTargetable, false);
+
+				GRenderTargetPool.FindFreeElement(RHICmdList, Desc, VisibilityData.SampleLightingBuffer, TEXT("HairSampleLightingBuffer"));
+
+				// TODO: do an indirect clear of the VisibilityData.SampleLightingBuffer ...
+				FRHIRenderPassInfo RPInfo(VisibilityData.SampleLightingBuffer->GetRenderTargetItem().TargetableTexture, MakeRenderTargetActions(ERenderTargetLoadAction::EClear, ERenderTargetStoreAction::EStore));
+				RHICmdList.BeginRenderPass(RPInfo, TEXT("HairSampleLightingClear"));
+				RHICmdList.EndRenderPass();
+			}
+
 			// #hair_todo: is there a better way to get SRV view of a RDG buffer? should work as long as there is not reuse between the pass
 			if (VisibilityData.NodeData)
+			{
 				VisibilityData.NodeDataSRV = RHICreateShaderResourceView(VisibilityData.NodeData->StructuredBuffer);
+			}
 
 			if (VisibilityData.NodeCoord)
+			{
 				VisibilityData.NodeCoordSRV = RHICreateShaderResourceView(VisibilityData.NodeCoord->StructuredBuffer);
+			}
 		}
 	}
 
diff --git a/Engine/Source/Runtime/Renderer/Private/HairStrands/HairStrandsVisibility.h b/Engine/Source/Runtime/Renderer/Private/HairStrands/HairStrandsVisibility.h
--- a/Engine/Source/Runtime/Renderer/Private/HairStrands/HairStrandsVisibility.h
+++ b/Engine/Source/Runtime/Renderer/Private/HairStrands/HairStrandsVisibility.h
@@ -40,6 +40,8 @@ struct FHairStrandsVisibilityData
 	const uint32					  TileThreadGroupSize = 32;
 
 	uint32							  MaxSampleCount = 8;
+	uint32							  MaxNodeCount = 0;
+	TRefCountPtr<IPooledRenderTarget> NodeCount;
 	TRefCountPtr<IPooledRenderTarget> NodeIndex;
 	TRefCountPtr<FPooledRDGBuffer>	  NodeData;
 	FShaderResourceViewRHIRef		  NodeDataSRV;
@@ -47,6 +49,14 @@ struct FHairStrandsVisibilityData
 	FShaderResourceViewRHIRef		  NodeCoordSRV;
 	TRefCountPtr<FPooledRDGBuffer>	  NodeIndirectArg;
 	uint32							  NodeGroupSize = 0;
+
+	// Hair lighting is accumulated within this buffer
+	// Allocated conservatively
+	// User indirect dispatch for accumulating contribution
+	FIntPoint SampleLightingViewportResolution;
+	TRefCountPtr<IPooledRenderTarget> SampleLightingBuffer;
+//	TRefCountPtr<IPooledRenderTarget> PixelLightingBuffer;
+
 };
 
 struct FHairStrandsVisibilityViews
diff --git a/Engine/Source/Runtime/Renderer/Private/LightRendering.cpp b/Engine/Source/Runtime/Renderer/Private/LightRendering.cpp
--- a/Engine/Source/Runtime/Renderer/Private/LightRendering.cpp
+++ b/Engine/Source/Runtime/Renderer/Private/LightRendering.cpp
@@ -217,11 +217,67 @@ struct FRenderLightParams
 	// Visibility buffer data
 	IPooledRenderTarget* HairCategorizationTexture = nullptr;
 	IPooledRenderTarget* HairVisibilityNodeOffsetAndCount = nullptr;
+	IPooledRenderTarget* HairVisibilityNodeCount = nullptr;
+	FShaderResourceViewRHIRef HairVisibilityNodeCoordsSRV = nullptr;
 	FShaderResourceViewRHIRef HairVisibilityNodeDataSRV = nullptr;
 
 	IPooledRenderTarget* ScreenShadowMaskSubPixelTexture = nullptr;
 };
 
+
+class TDeferredLightHairVS : public FGlobalShader
+{
+	DECLARE_SHADER_TYPE(TDeferredLightHairVS, Global);
+public:
+
+	static bool ShouldCompilePermutation(const FGlobalShaderPermutationParameters& Parameters)
+	{
+		return IsFeatureLevelSupported(Parameters.Platform, ERHIFeatureLevel::SM5);
+	}
+
+	static void ModifyCompilationEnvironment(const FGlobalShaderPermutationParameters& Parameters, FShaderCompilerEnvironment& OutEnvironment)
+	{
+		FGlobalShader::ModifyCompilationEnvironment(Parameters, OutEnvironment);
+		OutEnvironment.SetDefine(TEXT("SHADER_HAIR"), 1);
+	}
+
+	TDeferredLightHairVS() {}
+	TDeferredLightHairVS(const ShaderMetaType::CompiledShaderInitializerType& Initializer) :
+		FGlobalShader(Initializer)
+	{
+		MaxResolution.Bind(Initializer.ParameterMap, TEXT("MaxResolution"));
+		NodeGroupSize.Bind(Initializer.ParameterMap, TEXT("NodeGroupSize"));
+		HairVisibilityNodeCount.Bind(Initializer.ParameterMap, TEXT("HairVisibilityNodeCount"));
+	}
+
+	void SetParameters(FRHICommandList& RHICmdList, const FViewInfo& View, const FHairStrandsVisibilityData* VisibilityData)
+	{
+		FRHIPixelShader* ShaderRHI = RHICmdList.GetBoundPixelShader();
+		FGlobalShader::SetParameters<FViewUniformShaderParameters>(RHICmdList, ShaderRHI, View.ViewUniformBuffer);
+
+		if (!VisibilityData)
+		{
+			return;
+		}
+
+		if (HairVisibilityNodeCount.IsBound() && VisibilityData->NodeCount)
+		{
+			SetTextureParameter(RHICmdList,	ShaderRHI, HairVisibilityNodeCount, VisibilityData->NodeCount->GetRenderTargetItem().ShaderResourceTexture);
+		}
+
+		SetShaderValue(RHICmdList, ShaderRHI, NodeGroupSize, VisibilityData->NodeGroupSize);
+		SetShaderValue(RHICmdList, ShaderRHI, MaxResolution, VisibilityData->SampleLightingViewportResolution);
+	}
+
+private:
+	LAYOUT_FIELD(FShaderParameter, MaxResolution);
+	LAYOUT_FIELD(FShaderParameter, NodeGroupSize);
+	LAYOUT_FIELD(FShaderResourceParameter, HairVisibilityNodeCount);
+};
+
+IMPLEMENT_SHADER_TYPE(, TDeferredLightHairVS, TEXT("/Engine/Private/DeferredLightVertexShaders.usf"), TEXT("HairVertexMain"), SF_Vertex);
+
+
 enum class ELightSourceShape
 {
 	Directional,
@@ -243,7 +299,7 @@ class FDeferredLightPS : public FGlobalShader
 	class FVisualizeCullingDim	: SHADER_PERMUTATION_BOOL("VISUALIZE_LIGHT_CULLING");
 	class FLightingChannelsDim	: SHADER_PERMUTATION_BOOL("USE_LIGHTING_CHANNELS");
 	class FTransmissionDim		: SHADER_PERMUTATION_BOOL("USE_TRANSMISSION");
-	class FHairLighting			: SHADER_PERMUTATION_BOOL("USE_HAIR_LIGHTING");
+	class FHairLighting			: SHADER_PERMUTATION_INT("USE_HAIR_LIGHTING", 4);
 	class FAtmosphereTransmittance: SHADER_PERMUTATION_BOOL("USE_ATMOSPHERE_TRANSMITTANCE");
 
 	using FPermutationDomain = TShaderPermutationDomain<
@@ -333,6 +389,7 @@ class FDeferredLightPS : public FGlobalShader
 
 		HairCategorizationTexture.Bind(Initializer.ParameterMap, TEXT("HairCategorizationTexture"));
 		HairVisibilityNodeOffsetAndCount.Bind(Initializer.ParameterMap, TEXT("HairVisibilityNodeOffsetAndCount"));
+		HairVisibilityNodeCoords.Bind(Initializer.ParameterMap, TEXT("HairVisibilityNodeCoords"));
 		HairVisibilityNodeData.Bind(Initializer.ParameterMap, TEXT("HairVisibilityNodeData"));
 	}
 
@@ -514,6 +571,19 @@ private:
 					RenderLightParams->HairVisibilityNodeOffsetAndCount->GetRenderTargetItem().TargetableTexture);
 			}
 		}
+		
+		if (HairVisibilityNodeCoords.IsBound())
+		{
+			if (RenderLightParams && RenderLightParams->HairVisibilityNodeCoordsSRV)
+			{
+				FShaderResourceViewRHIRef SRV = RenderLightParams->HairVisibilityNodeCoordsSRV;
+				SetSRVParameter(
+					RHICmdList, 
+					ShaderRHI, 
+					HairVisibilityNodeCoords,
+					SRV);
+			}
+		}
 
 		if (HairVisibilityNodeData.IsBound())
 		{
@@ -580,6 +650,7 @@ private:
 	LAYOUT_FIELD(FShaderResourceParameter, HairTransmittanceBuffer);
 	LAYOUT_FIELD(FShaderResourceParameter, HairCategorizationTexture);
 	LAYOUT_FIELD(FShaderResourceParameter, HairVisibilityNodeOffsetAndCount);
+	LAYOUT_FIELD(FShaderResourceParameter, HairVisibilityNodeCoords);
 	LAYOUT_FIELD(FShaderResourceParameter, HairVisibilityNodeData);
 	LAYOUT_FIELD(FShaderResourceParameter, ScreenShadowMaskSubPixelTexture);
 
@@ -1029,7 +1100,7 @@ void FDeferredShadingSceneRenderer::RenderLights(FRHICommandListImmediate& RHICm
 					const FLightSceneInfo* const LightSceneInfo = SortedLightInfo.LightSceneInfo;
 
 					// Render the light to the scene color buffer, using a 1x1 white texture as input
-					RenderLight(RHICmdList, LightSceneInfo, nullptr, nullptr, nullptr, nullptr, false, false);
+					RenderLight(RHICmdList, LightSceneInfo, nullptr, nullptr, nullptr, nullptr, false, false, false);
 				}
 
 				SceneContext.FinishRenderingSceneColor(RHICmdList);
@@ -1054,7 +1125,7 @@ void FDeferredShadingSceneRenderer::RenderLights(FRHICommandListImmediate& RHICm
 
 					// Render the light to the scene color buffer, using a 1x1 white texture as input
 					SceneContext.BeginRenderingSceneColor(RHICmdList, ESimpleRenderTargetMode::EExistingColorAndDepth, FExclusiveDepthStencil::DepthRead_StencilWrite, true, true);
-					RenderLight(RHICmdList, LightSceneInfo, nullptr, nullptr, &TransmittanceMaskData, InHairVisibilityViews, false, false);
+					RenderLight(RHICmdList, LightSceneInfo, nullptr, nullptr, &TransmittanceMaskData, InHairVisibilityViews, false, false, false);
 					SceneContext.FinishRenderingSceneColor(RHICmdList);
 				}
 
@@ -1772,21 +1843,30 @@ void FDeferredShadingSceneRenderer::RenderLights(FRHICommandListImmediate& RHICm
 					SCOPED_DRAW_EVENT(RHICmdList, StandardDeferredLighting);
 					SceneContext.BeginRenderingSceneColor(RHICmdList, ESimpleRenderTargetMode::EExistingColorAndDepth, FExclusiveDepthStencil::DepthRead_StencilWrite, true, true);
 
+					// When split lighting is enable, the hair pixel are rendered into a separate target
+					const bool bUseHairSplitLighting = HairDatas && Views.Num() > 0 && HairDatas->HairVisibilityViews.HairDatas[0].SampleLightingBuffer != nullptr;
+
+					// ScreenShadowMaskTexture might have been created for a previous light, but only use it if we wrote valid data into it for this light
+					IPooledRenderTarget* LightShadowMaskTexture = nullptr;
+					IPooledRenderTarget* LightShadowMaskSubPixelTexture = nullptr;
+					if (bUsedShadowMaskTexture)
+					{
+						LightShadowMaskTexture = ScreenShadowMaskTexture;
+						LightShadowMaskSubPixelTexture = ScreenShadowMaskSubPixelTexture;
+					}
+
 					// Render the light to the scene color buffer, conditionally using the attenuation buffer or a 1x1 white texture as input 
 					if (bDirectLighting)
 					{
-						// ScreenShadowMaskTexture might have been created for a previous light, but only use it if we wrote valid data into it for this light
-						IPooledRenderTarget* LightShadowMaskTexture = nullptr;
-						IPooledRenderTarget* LightShadowMaskSubPixelTexture = nullptr;
-						if (bUsedShadowMaskTexture)
-						{
-							LightShadowMaskTexture = ScreenShadowMaskTexture;
-							LightShadowMaskSubPixelTexture = ScreenShadowMaskSubPixelTexture;
-						}
-						RenderLight(RHICmdList, &LightSceneInfo, LightShadowMaskTexture, LightShadowMaskSubPixelTexture, &TransmittanceMaskData, InHairVisibilityViews, false, true);
+						RenderLight(RHICmdList, &LightSceneInfo, LightShadowMaskTexture, LightShadowMaskSubPixelTexture, &TransmittanceMaskData, InHairVisibilityViews, false, true, !bUseHairSplitLighting);
 					}
 
 					SceneContext.FinishRenderingSceneColor(RHICmdList);
+
+					if (bDirectLighting && bUseHairSplitLighting)
+					{
+						RenderLightForHair(RHICmdList, &LightSceneInfo, LightShadowMaskSubPixelTexture, &TransmittanceMaskData, InHairVisibilityViews, true);
+					}
 				}
 			}
 		}
@@ -1820,7 +1900,7 @@ void FDeferredShadingSceneRenderer::RenderLightArrayForOverlapViewmode(FRHIComma
 			&& !LightSceneInfo->Proxy->HasStaticLighting()
 			&& LightSceneInfo->Proxy->CastsStaticShadow())
 		{
-			RenderLight(RHICmdList, LightSceneInfo, nullptr, nullptr, nullptr, nullptr, true, false);
+			RenderLight(RHICmdList, LightSceneInfo, nullptr, nullptr, nullptr, nullptr, true, false, false);
 		}
 	}
 }
@@ -1886,7 +1966,7 @@ static void SetShaderTemplLightingSimple(
 	PermutationVector.Set< FDeferredLightPS::FVisualizeCullingDim >( View.Family->EngineShowFlags.VisualizeLightCulling );
 	PermutationVector.Set< FDeferredLightPS::FLightingChannelsDim >( false );
 	PermutationVector.Set< FDeferredLightPS::FTransmissionDim >( false );
-	PermutationVector.Set< FDeferredLightPS::FHairLighting>( false );
+	PermutationVector.Set< FDeferredLightPS::FHairLighting>( 0 );
 	PermutationVector.Set < FDeferredLightPS::FAtmosphereTransmittance >( false );
 
 	TShaderMapRef< FDeferredLightPS > PixelShader( View.ShaderMap, PermutationVector );
@@ -1935,7 +2015,7 @@ void CalculateLightNearFarDepthFromBounds(const FViewInfo& View, const FSphere &
  * @param LightIndex The light's index into FScene::Lights
  * @return true if anything got rendered
  */
-void FDeferredShadingSceneRenderer::RenderLight(FRHICommandList& RHICmdList, const FLightSceneInfo* LightSceneInfo, IPooledRenderTarget* ScreenShadowMaskTexture, IPooledRenderTarget* ScreenShadowMaskSubPixelTexture, FHairStrandsTransmittanceMaskData* InTransmittanceMaskData, const FHairStrandsVisibilityViews* InHairVisibilityViews, bool bRenderOverlap, bool bIssueDrawEvent)
+void FDeferredShadingSceneRenderer::RenderLight(FRHICommandList& RHICmdList, const FLightSceneInfo* LightSceneInfo, IPooledRenderTarget* ScreenShadowMaskTexture, IPooledRenderTarget* ScreenShadowMaskSubPixelTexture, FHairStrandsTransmittanceMaskData* InTransmittanceMaskData, const FHairStrandsVisibilityViews* InHairVisibilityViews, bool bRenderOverlap, bool bIssueDrawEvent, bool bSupportHair)
 {
 	SCOPE_CYCLE_COUNTER(STAT_DirectLightRenderingTime);
 	INC_DWORD_STAT(STAT_NumLightsUsingStandardDeferred);
@@ -1944,8 +2024,9 @@ void FDeferredShadingSceneRenderer::RenderLight(FRHICommandList& RHICmdList, con
 	FGraphicsPipelineStateInitializer GraphicsPSOInit;
 	RHICmdList.ApplyCachedRenderTargets(GraphicsPSOInit);
 
-	const bool bHairRenderingEnabled = InTransmittanceMaskData && InHairVisibilityViews && (LightSceneInfo->Proxy->CastsHairStrandsDeepShadow() || IsHairStrandsVoxelizationEnable());
-	if (bHairRenderingEnabled)
+	const bool bHairAttunationEnabled = !bSupportHair && InTransmittanceMaskData && InHairVisibilityViews && (LightSceneInfo->Proxy->CastsHairStrandsDeepShadow() || IsHairStrandsVoxelizationEnable());
+	const bool bHairLightingEnabled = bSupportHair && InTransmittanceMaskData && InHairVisibilityViews && (LightSceneInfo->Proxy->CastsHairStrandsDeepShadow() || IsHairStrandsVoxelizationEnable());
+	if (bHairLightingEnabled)
 	{
 		GraphicsPSOInit.BlendState = TStaticBlendState<
 			CW_RGBA, BO_Add, BF_One, BF_One, BO_Add, BF_One, BF_One,
@@ -1983,7 +2064,7 @@ void FDeferredShadingSceneRenderer::RenderLight(FRHICommandList& RHICmdList, con
 		RHICmdList.SetViewport(View.ViewRect.Min.X, View.ViewRect.Min.Y, 0.0f, View.ViewRect.Max.X, View.ViewRect.Max.Y, 1.0f);
 
 		FRenderLightParams RenderLightParams;
-		if (bHairRenderingEnabled)
+		if (bHairLightingEnabled || bHairAttunationEnabled)
 		{
 			RenderLightParams.DeepShadow_TransmittanceMaskBuffer = InTransmittanceMaskData ? InTransmittanceMaskData->TransmittanceMaskSRV : nullptr;
 			RenderLightParams.DeepShadow_TransmittanceMaskBufferMaxCount = InTransmittanceMaskData && InTransmittanceMaskData->TransmittanceMask ? InTransmittanceMaskData->TransmittanceMask->Desc.NumElements : 0;
@@ -1993,10 +2074,13 @@ void FDeferredShadingSceneRenderer::RenderLight(FRHICommandList& RHICmdList, con
 				const FHairStrandsVisibilityData& HairVisibilityData = InHairVisibilityViews->HairDatas[ViewIndex];
 				RenderLightParams.HairVisibilityNodeOffsetAndCount = HairVisibilityData.NodeIndex;
 				RenderLightParams.HairVisibilityNodeDataSRV = HairVisibilityData.NodeDataSRV;
+				RenderLightParams.HairVisibilityNodeCoordsSRV = RenderLightParams.HairVisibilityNodeCoordsSRV;
 				RenderLightParams.HairCategorizationTexture = HairVisibilityData.CategorizationTexture;
 			}
 		}
 
+		const uint32 HairPermutation = bHairAttunationEnabled ? 3 : (bHairLightingEnabled ? 1 : 0);
+
 		if (LightSceneInfo->Proxy->GetLightType() == LightType_Directional)
 		{
 			// Turn DBT back off
@@ -2026,7 +2110,7 @@ void FDeferredShadingSceneRenderer::RenderLight(FRHICommandList& RHICmdList, con
 				PermutationVector.Set< FDeferredLightPS::FVisualizeCullingDim >( View.Family->EngineShowFlags.VisualizeLightCulling );
 				PermutationVector.Set< FDeferredLightPS::FLightingChannelsDim >( View.bUsesLightingChannels );
 				PermutationVector.Set< FDeferredLightPS::FTransmissionDim >( bTransmission );
-				PermutationVector.Set< FDeferredLightPS::FHairLighting>(bHairRenderingEnabled);
+				PermutationVector.Set< FDeferredLightPS::FHairLighting>(HairPermutation);
 				// Only directional lights are rendered in this path, so we only need to check if it is use to light the atmosphere
 				PermutationVector.Set< FDeferredLightPS::FAtmosphereTransmittance >(bAtmospherePerPixelTransmittance);
 
@@ -2036,7 +2120,7 @@ void FDeferredShadingSceneRenderer::RenderLight(FRHICommandList& RHICmdList, con
 				GraphicsPSOInit.BoundShaderState.PixelShaderRHI = PixelShader.GetPixelShader();
 
 				SetGraphicsPipelineState(RHICmdList, GraphicsPSOInit);
-				PixelShader->SetParameters(RHICmdList, View, LightSceneInfo, ScreenShadowMaskTexture, bHairRenderingEnabled ? &RenderLightParams : nullptr);
+				PixelShader->SetParameters(RHICmdList, View, LightSceneInfo, ScreenShadowMaskTexture, (bHairLightingEnabled || bHairAttunationEnabled) ? &RenderLightParams : nullptr);
 			}
 
 			VertexShader->SetParameters(RHICmdList, View, LightSceneInfo);
@@ -2057,7 +2141,7 @@ void FDeferredShadingSceneRenderer::RenderLight(FRHICommandList& RHICmdList, con
 		{
 			// Use DBT to allow work culling on shadow lights
 			// Disable depth bound when hair rendering is enabled as this rejects partially covered pixel write (with opaque background)
-			GraphicsPSOInit.bDepthBounds = !bHairRenderingEnabled && (GSupportsDepthBoundsTest && GAllowDepthBoundsTest != 0);
+			GraphicsPSOInit.bDepthBounds = !bHairLightingEnabled && (GSupportsDepthBoundsTest && GAllowDepthBoundsTest != 0);
 
 			TShaderMapRef<TDeferredLightVS<true> > VertexShader(View.ShaderMap);
 
@@ -2083,7 +2167,7 @@ void FDeferredShadingSceneRenderer::RenderLight(FRHICommandList& RHICmdList, con
 				PermutationVector.Set< FDeferredLightPS::FVisualizeCullingDim >( View.Family->EngineShowFlags.VisualizeLightCulling );
 				PermutationVector.Set< FDeferredLightPS::FLightingChannelsDim >( View.bUsesLightingChannels );
 				PermutationVector.Set< FDeferredLightPS::FTransmissionDim >( bTransmission );
-				PermutationVector.Set< FDeferredLightPS::FHairLighting>(bHairRenderingEnabled);
+				PermutationVector.Set< FDeferredLightPS::FHairLighting>(HairPermutation);
 				PermutationVector.Set < FDeferredLightPS::FAtmosphereTransmittance >(false);
 
 				TShaderMapRef< FDeferredLightPS > PixelShader( View.ShaderMap, PermutationVector );
@@ -2130,6 +2214,109 @@ void FDeferredShadingSceneRenderer::RenderLight(FRHICommandList& RHICmdList, con
 	}
 }
 
+void FDeferredShadingSceneRenderer::RenderLightForHair(
+	FRHICommandList& RHICmdList, 
+	const FLightSceneInfo* LightSceneInfo, 
+	IPooledRenderTarget* HairShadowMaskTexture, 
+	FHairStrandsTransmittanceMaskData* InTransmittanceMaskData, 
+	const FHairStrandsVisibilityViews* InHairVisibilityViews, 
+	bool bIssueDrawEvent)
+{
+	SCOPE_CYCLE_COUNTER(STAT_DirectLightRenderingTime);
+	INC_DWORD_STAT(STAT_NumLightsUsingStandardDeferred);
+	SCOPED_CONDITIONAL_DRAW_EVENT(RHICmdList, StandardDeferredLighting, bIssueDrawEvent);
+	
+	FGraphicsPipelineStateInitializer GraphicsPSOInit;
+	RHICmdList.ApplyCachedRenderTargets(GraphicsPSOInit);
+
+	const bool bHairRenderingEnabled = InTransmittanceMaskData && InHairVisibilityViews && (LightSceneInfo->Proxy->CastsHairStrandsDeepShadow() || IsHairStrandsVoxelizationEnable());
+	if (!bHairRenderingEnabled)
+	{
+		return;
+	}
+
+	GraphicsPSOInit.BlendState = TStaticBlendState<CW_RGBA, BO_Add, BF_One, BF_One, BO_Max, BF_SourceAlpha, BF_DestAlpha>::GetRHI();
+
+	GraphicsPSOInit.PrimitiveType = PT_TriangleList;
+
+	const FSphere LightBounds = LightSceneInfo->Proxy->GetBoundingSphere();
+	const bool bTransmission = LightSceneInfo->Proxy->Transmission();
+
+	for (int32 ViewIndex = 0; ViewIndex < Views.Num(); ViewIndex++)
+	{
+		FViewInfo& View = Views[ViewIndex];
+
+		// Ensure the light is valid for this view
+		if (!LightSceneInfo->ShouldRenderLight(View) || ViewIndex >= InHairVisibilityViews->HairDatas.Num())
+		{
+			continue;
+		}
+
+		check(RHICmdList.IsOutsideRenderPass());
+
+		const FHairStrandsVisibilityData& HairVisibilityData = InHairVisibilityViews->HairDatas[ViewIndex];
+		
+		FRenderLightParams RenderLightParams;
+		RenderLightParams.DeepShadow_TransmittanceMaskBuffer = InTransmittanceMaskData ? InTransmittanceMaskData->TransmittanceMaskSRV : nullptr;
+		RenderLightParams.DeepShadow_TransmittanceMaskBufferMaxCount = InTransmittanceMaskData && InTransmittanceMaskData->TransmittanceMask ? InTransmittanceMaskData->TransmittanceMask->Desc.NumElements : 0;
+		RenderLightParams.ScreenShadowMaskSubPixelTexture = HairShadowMaskTexture;
+		RenderLightParams.HairVisibilityNodeOffsetAndCount = HairVisibilityData.NodeIndex;
+		RenderLightParams.HairVisibilityNodeDataSRV = HairVisibilityData.NodeDataSRV;
+		RenderLightParams.HairVisibilityNodeCoordsSRV = HairVisibilityData.NodeCoordSRV;
+		RenderLightParams.HairCategorizationTexture = HairVisibilityData.CategorizationTexture;
+
+		FRHIRenderPassInfo RPInfo(HairVisibilityData.SampleLightingBuffer->GetRenderTargetItem().TargetableTexture, MakeRenderTargetActions(ERenderTargetLoadAction::ELoad, ERenderTargetStoreAction::EStore));
+		RHICmdList.BeginRenderPass(RPInfo, TEXT("HairLighting"));
+		RHICmdList.SetViewport(0, 0, 0.0f, HairVisibilityData.SampleLightingViewportResolution.X, HairVisibilityData.SampleLightingViewportResolution.Y, 1.0f);
+
+		FDeferredLightPS::FPermutationDomain PermutationVector;
+		if (LightSceneInfo->Proxy->GetLightType() == LightType_Directional)
+		{
+			PermutationVector.Set< FDeferredLightPS::FSourceShapeDim >(ELightSourceShape::Directional);
+			PermutationVector.Set< FDeferredLightPS::FIESProfileDim >(false);
+			PermutationVector.Set< FDeferredLightPS::FInverseSquaredDim >(false);
+			PermutationVector.Set< FDeferredLightPS::FVisualizeCullingDim >(View.Family->EngineShowFlags.VisualizeLightCulling);
+			PermutationVector.Set< FDeferredLightPS::FLightingChannelsDim >(View.bUsesLightingChannels);
+			PermutationVector.Set< FDeferredLightPS::FTransmissionDim >(false);
+			PermutationVector.Set< FDeferredLightPS::FHairLighting>(2);
+		}
+		else
+		{
+			const bool bUseIESTexture = View.Family->EngineShowFlags.TexturedLightProfiles && LightSceneInfo->Proxy->GetIESTextureResource() != 0;
+			PermutationVector.Set< FDeferredLightPS::FSourceShapeDim >(LightSceneInfo->Proxy->IsRectLight() ? ELightSourceShape::Rect : ELightSourceShape::Capsule);
+			PermutationVector.Set< FDeferredLightPS::FSourceTextureDim >(LightSceneInfo->Proxy->IsRectLight() && LightSceneInfo->Proxy->HasSourceTexture());
+			PermutationVector.Set< FDeferredLightPS::FIESProfileDim >(bUseIESTexture);
+			PermutationVector.Set< FDeferredLightPS::FInverseSquaredDim >(LightSceneInfo->Proxy->IsInverseSquared());
+			PermutationVector.Set< FDeferredLightPS::FVisualizeCullingDim >(View.Family->EngineShowFlags.VisualizeLightCulling);
+			PermutationVector.Set< FDeferredLightPS::FLightingChannelsDim >(View.bUsesLightingChannels);
+			PermutationVector.Set< FDeferredLightPS::FTransmissionDim >(false);
+			PermutationVector.Set< FDeferredLightPS::FHairLighting>(2);
+		}
+
+		TShaderMapRef<TDeferredLightHairVS> VertexShader(View.ShaderMap);
+		TShaderMapRef<FDeferredLightPS> PixelShader(View.ShaderMap, PermutationVector);
+
+		GraphicsPSOInit.BoundShaderState.VertexDeclarationRHI = GetVertexDeclarationFVector4();
+		GraphicsPSOInit.bDepthBounds = false;
+		GraphicsPSOInit.RasterizerState = TStaticRasterizerState<FM_Solid, CM_None>::GetRHI();
+		GraphicsPSOInit.DepthStencilState = TStaticDepthStencilState<false, CF_Always>::GetRHI();
+		GraphicsPSOInit.BoundShaderState.VertexShaderRHI = VertexShader.GetVertexShader();
+		GraphicsPSOInit.BoundShaderState.PixelShaderRHI = PixelShader.GetPixelShader();
+		GraphicsPSOInit.PrimitiveType = PT_TriangleList;
+		SetGraphicsPipelineState(RHICmdList, GraphicsPSOInit);
+
+
+		VertexShader->SetParameters(RHICmdList, View, &HairVisibilityData);
+		PixelShader->SetParameters(RHICmdList, View,  LightSceneInfo, HairShadowMaskTexture, &RenderLightParams);
+		PixelShader->SetParameters(RHICmdList, View,  LightSceneInfo, HairShadowMaskTexture, bHairRenderingEnabled ? &RenderLightParams : nullptr);
+
+		RHICmdList.SetStreamSource(0, nullptr, 0);
+		RHICmdList.DrawPrimitive(0, 1, 1);
+
+		RHICmdList.EndRenderPass();
+	}
+}
+
 void FDeferredShadingSceneRenderer::RenderSimpleLightsStandardDeferred(FRHICommandListImmediate& RHICmdList, const FSimpleLightArray& SimpleLights)
 {
 	SCOPE_CYCLE_COUNTER(STAT_DirectLightRenderingTime);
diff --git a/Engine/Source/Runtime/Renderer/Private/LightRendering.h b/Engine/Source/Runtime/Renderer/Private/LightRendering.h
--- a/Engine/Source/Runtime/Renderer/Private/LightRendering.h
+++ b/Engine/Source/Runtime/Renderer/Private/LightRendering.h
@@ -513,6 +513,7 @@ public:
 	static void ModifyCompilationEnvironment(const FGlobalShaderPermutationParameters& Parameters, FShaderCompilerEnvironment& OutEnvironment)
 	{
 		FGlobalShader::ModifyCompilationEnvironment(Parameters, OutEnvironment);
+		OutEnvironment.SetDefine(TEXT("SHADER_RADIAL_LIGHT"), bRadialLight ? 1 : 0);
 	}
 
 	TDeferredLightVS()	{}
