Remove sub color buffer.

This removes the subpixel color buffer, hair lighting is not stored all light sample into a separate buffer.

diff --git a/Engine/Shaders/Private/DeferredLightPixelShaders.usf b/Engine/Shaders/Private/DeferredLightPixelShaders.usf
--- a/Engine/Shaders/Private/DeferredLightPixelShaders.usf
+++ b/Engine/Shaders/Private/DeferredLightPixelShaders.usf
@@ -16,14 +16,16 @@
 #include "DeferredLightingCommon.ush"
 
 #if USE_HAIR_LIGHTING
-	#include "HairStrands/HairStrandsDeepShadowCommon.ush"
-	#include "HairStrands/HairStrandsDeepTransmittanceCommon.ush"
-	#include "HairStrands/HairStrandsDeepTransmittanceDualScattering.ush"
 	#include "HairStrands/HairStrandsVisibilityCommon.ush"
 	#if HAIR_CUSTOM_BSDF == 1
 	#include "HairBsdf.ush"
 	#endif
 #endif
+#if USE_HAIR_LIGHTING == 2
+	#include "HairStrands/HairStrandsDeepShadowCommon.ush"
+	#include "HairStrands/HairStrandsDeepTransmittanceCommon.ush"
+	#include "HairStrands/HairStrandsDeepTransmittanceDualScattering.ush"
+#endif
 
 #if USE_ATMOSPHERE_TRANSMITTANCE
 #include "/Engine/Private/SkyAtmosphereCommon.ush"
@@ -113,142 +115,6 @@ float GetExposure()
 
 ////////////////////////////////////////////////////////////////////////////////////////////////////
 
-#if USE_HAIR_LIGHTING
-
-Texture2D<uint4> HairCategorizationTexture;
-Texture2D<uint> HairVisibilityNodeOffsetAndCount;
-StructuredBuffer<FPackedHairSample> HairVisibilityNodeData;
-
-Texture2D<float4> ScreenShadowMaskSubPixelTexture;
-
-uint HairTransmittanceBufferMaxCount;
-StructuredBuffer<FPackedHairTransmittanceMask> HairTransmittanceBuffer;
-Texture3D<float4> HairLUTTexture;
-SamplerState HairLUTSampler;
-uint HairShadowMaskValid;
-uint HairComponents; // R, TT, TRT, LocalScattering, GlobalScattering, MultipleScattering
-float HairDualScatteringRoughnessOverride;
-
-float Luma(float3 Color)
-{
-	return 0.25f * ((Color.g * 2.0) + (Color.r + Color.b));
-}
-
-#define RESOLVE_IN_PERCEPTUAL_SPACE 0
-	
-float4 LinearToPerceptual(float4 Color)
-{
-#if RESOLVE_IN_PERCEPTUAL_SPACE
-	return float4(Color.xyz / (1 + Luma(Color.xyz)), Color.a);
-#else
-	return Color;
-#endif
-}
-float4 PerceptualToLinear(float4 Color) 
-{
-#if RESOLVE_IN_PERCEPTUAL_SPACE
-	return float4(Color.xyz / (1 - Luma(Color.xyz)), Color.a);
-#else
-	return Color;
-#endif
-}
-
-void ComputeHairLighting(
-	in FInputParams InputParams,
-	in FNodeDesc NodeDesc,
-	in float PixelCoverage,
-	inout float4 OutColor, 
-	inout float4 OutColorSubPixel)
-{
-	float4 LocalColor = 0;
-	LOOP
-	for (uint SampleIt = 0; SampleIt < NodeDesc.Count; SampleIt++)
-	{
-		const uint LocalOffset = NodeDesc.Offset + SampleIt;
-		const FPackedHairSample NodeData = HairVisibilityNodeData[LocalOffset];
-		const FHairSample Sample = UnpackHairSample(NodeData);
-
-	#if USE_LIGHTING_CHANNELS
-		if (!(Sample.LightChannelMask & DeferredLightUniforms.LightingChannelMask))
-		{
-			continue;
-		}
-	#endif
-
-		// Inject material data from the visibility/mini-gbuffer for hair
-		const float SceneDepth = ConvertFromDeviceZ(Sample.Depth);
-		FScreenSpaceData HairScreenSpaceData;
-		HairScreenSpaceData.AmbientOcclusion = 1;
-		HairScreenSpaceData.GBuffer.WorldNormal = Sample.Tangent;
-		HairScreenSpaceData.GBuffer.BaseColor = Sample.BaseColor;
-		HairScreenSpaceData.GBuffer.Roughness = Sample.Roughness;
-		HairScreenSpaceData.GBuffer.ShadingModelID = SHADINGMODELID_HAIR;
-		HairScreenSpaceData.GBuffer.DiffuseColor = 0;
-		HairScreenSpaceData.GBuffer.SpecularColor = 0;
-		HairScreenSpaceData.GBuffer.Specular = Sample.Specular;
-		HairScreenSpaceData.GBuffer.Metallic = 0;
-		HairScreenSpaceData.GBuffer.Depth = SceneDepth;
-		HairScreenSpaceData.GBuffer.GBufferAO = 1;
-		HairScreenSpaceData.GBuffer.CustomData = float4(HairDualScatteringRoughnessOverride, 0, Sample.Backlit, 0);
-		HairScreenSpaceData.GBuffer.IndirectIrradiance = 0;
-		HairScreenSpaceData.GBuffer.PrecomputedShadowFactors = 1;
-
-		const FDerivedParams DerivedParams = GetDerivedParams(InputParams, SceneDepth);
-		FDeferredLightData LightData = SetupLightDataForStandardDeferred();
-		LightData.HairTransmittance.ScatteringComponent = HairComponents;
-
-		// Compute Hair transmittance
-		{
-			const float3 LightPosition = LightData.Position;
-			float3 L = float3(0, 0, 1);
-			if (LightData.bRadialLight)
-				L = normalize(LightPosition - DerivedParams.WorldPosition);
-			else
-				L = LightData.Direction;
-
-			const float3 V = normalize(-DerivedParams.CameraVector);
-			LightData.HairTransmittance = GetHairTransmittance(
-				V, 
-				L,
-				HairScreenSpaceData.GBuffer, 
-				LocalOffset,
-				HairTransmittanceBufferMaxCount,
-				HairTransmittanceBuffer,
-				HairLUTTexture,
-				HairLUTSampler,
-				HairComponents);
-		}
-
-		const float Dither = InterleavedGradientNoise(InputParams.PixelPos, View.StateFrameIndexMod8);
-		const uint3 PixelCoord = uint3(floor(InputParams.PixelPos), 0);
-		float4 LightAttenuation = HairShadowMaskValid ? ScreenShadowMaskSubPixelTexture.Load(PixelCoord) : 1;
-		LightAttenuation = min(LightAttenuation, LightData.HairTransmittance.OpaqueVisibility.xxxx);
-
-		FRectTexture RectTexture = InitRectTexture(DeferredLightUniforms.SourceTexture);
-		float SurfaceShadow = 1.0f;
-		const float4 Radiance = GetDynamicLighting(DerivedParams.WorldPosition, DerivedParams.CameraVector, HairScreenSpaceData.GBuffer, HairScreenSpaceData.AmbientOcclusion, HairScreenSpaceData.GBuffer.ShadingModelID, LightData, LightAttenuation, Dither, uint2(InputParams.PixelPos), RectTexture, SurfaceShadow);
-		const float  Attenuation = ComputeLightProfileMultiplier(DerivedParams.WorldPosition, DeferredLightUniforms.Position, -DeferredLightUniforms.Direction, DeferredLightUniforms.Tangent);
-
-		const float4 LinearColor = Radiance * Attenuation;
-		const float4 PercepualColor = LinearToPerceptual(LinearColor);
-		const float LocalCoverage = From8bitCoverage(Sample.Coverage8bit);
-		LocalColor += PercepualColor * LocalCoverage;
-	}
-	OutColor = 0;
-	OutColorSubPixel = 0;
-	if (PixelCoverage < 1)
-	{
-		OutColorSubPixel = LocalColor * PixelCoverage;
-		OutColorSubPixel.a = PixelCoverage;
-	}
-	else
-	{
-		OutColor = LocalColor;
-	}
-}
-
-#endif // USE_HAIR_LIGHTING
-
 #if USE_ATMOSPHERE_TRANSMITTANCE
 
 float3 GetAtmosphericLightTransmittance(in float4 SVPos, in float2 ScreenUV, in float3 LightDirection)
@@ -269,7 +135,11 @@ float3 GetAtmosphericLightTransmittance(in float4 SVPos, in float2 ScreenUV, in
 
 ////////////////////////////////////////////////////////////////////////////////////////////////////
 
-#if USE_HAIR_LIGHTING == 0 || USE_HAIR_LIGHTING == 1 || USE_HAIR_LIGHTING == 3
+#if USE_HAIR_LIGHTING == 1
+Texture2D<uint4> HairCategorizationTexture;
+#endif
+
+#if USE_HAIR_LIGHTING == 0 || USE_HAIR_LIGHTING == 1
 
 void DeferredLightPixelMain(
 #if LIGHT_SOURCE_SHAPE > 0
@@ -280,9 +150,6 @@ void DeferredLightPixelMain(
 #endif
 	float4 SVPos			: SV_POSITION,
 	out float4 OutColor		: SV_Target0
-#if USE_HAIR_LIGHTING == 1
-,	out float4 OutColorSubPixel	: SV_Target1
-#endif
 	)
 {
 	const float2 PixelPos = SVPos.xy;
@@ -301,24 +168,6 @@ void DeferredLightPixelMain(
 	InputParams.ScreenVector	= ScreenVector;
 #endif
 
-#if USE_HAIR_LIGHTING
-	const uint3 PixelCoord = uint3(floor(InputParams.PixelPos), 0);
-	const FCategorizationData CatData = DecodeCategorizationData(HairCategorizationTexture.Load(PixelCoord));	
-	const float HairCoverage = saturate(CatData.PixelCoverage);
-
-	#if USE_HAIR_LIGHTING == 1
-	// Only light pixels marked as using deferred shading
-	OutColorSubPixel = 0;
-	const FNodeDesc NodeDesc = DecodeNodeDesc(HairVisibilityNodeOffsetAndCount.Load(PixelCoord));
-	BRANCH 
-	if (NodeDesc.Count > 0)
-	{
-		ComputeHairLighting(InputParams, NodeDesc, HairCoverage, OutColor, OutColorSubPixel);
-		OutColorSubPixel.rgb *= GetExposure();
-	}
-	#endif
-#endif
-
 	FScreenSpaceData ScreenSpaceData = GetScreenSpaceData(InputParams.ScreenUV);
 
 	// Only light pixels marked as using deferred shading
@@ -328,6 +177,15 @@ void DeferredLightPixelMain(
 #endif
 		)
 	{
+	#if USE_HAIR_LIGHTING
+		const uint3 PixelCoord = uint3(floor(InputParams.PixelPos), 0);
+		const FCategorizationData CatData = DecodeCategorizationData(HairCategorizationTexture.Load(PixelCoord));	
+		const float HairCoverage = saturate(CatData.PixelCoverage);
+		const float OpaqueVisibility = 1.0f - HairCoverage;
+	#else
+		const float OpaqueVisibility = 1.0f;
+	#endif
+
 		const float SceneDepth = CalcSceneDepth(InputParams.ScreenUV);
 		const FDerivedParams DerivedParams = GetDerivedParams(InputParams, SceneDepth);
 
@@ -340,17 +198,11 @@ void DeferredLightPixelMain(
 		const float4 Radiance = GetDynamicLighting(DerivedParams.WorldPosition, DerivedParams.CameraVector, ScreenSpaceData.GBuffer, ScreenSpaceData.AmbientOcclusion, ScreenSpaceData.GBuffer.ShadingModelID, LightData, GetPerPixelLightAttenuation(InputParams.ScreenUV), Dither, uint2(InputParams.PixelPos), RectTexture, SurfaceShadow);
 		const float  Attenuation = ComputeLightProfileMultiplier(DerivedParams.WorldPosition, DeferredLightUniforms.Position, -DeferredLightUniforms.Direction, DeferredLightUniforms.Tangent);
 
-#if USE_HAIR_LIGHTING
-		float OpaqueVisibility = 1.0f - HairCoverage;
-		OutColor += LinearToPerceptual(Radiance * Attenuation) * OpaqueVisibility;
-		OutColor  = PerceptualToLinear(OutColor);
-#else
-		OutColor += (Radiance * Attenuation);
-#endif
+		OutColor += (Radiance * Attenuation) * OpaqueVisibility;
 
-#if USE_ATMOSPHERE_TRANSMITTANCE
+	#if USE_ATMOSPHERE_TRANSMITTANCE
 		OutColor.rgb *= GetAtmosphericLightTransmittance(SVPos, InputParams.ScreenUV, DeferredLightUniforms.Direction.xyz);
-#endif
+	#endif
 	}
 
 	// RGB:SceneColor Specular and Diffuse
@@ -364,8 +216,20 @@ void DeferredLightPixelMain(
 
 #if USE_HAIR_LIGHTING == 2
 
+Texture2D<uint> HairVisibilityNodeOffsetAndCount;
+StructuredBuffer<FPackedHairSample> HairVisibilityNodeData;
 StructuredBuffer<uint>	HairVisibilityNodeCoords;
 
+Texture2D<float4> ScreenShadowMaskSubPixelTexture;
+
+uint HairTransmittanceBufferMaxCount;
+StructuredBuffer<FPackedHairTransmittanceMask> HairTransmittanceBuffer;
+Texture3D<float4> HairLUTTexture;
+SamplerState HairLUTSampler;
+uint HairShadowMaskValid;
+uint HairComponents; // R, TT, TRT, LocalScattering, GlobalScattering, MultipleScattering
+float HairDualScatteringRoughnessOverride;
+
 void DeferredLightPixelMain(
 	float4 SVPos : SV_POSITION,
 	uint NodeCount  : DISPATCH_NODECOUNT,
@@ -454,7 +318,8 @@ void DeferredLightPixelMain(
 
 	const float LocalCoverage = From8bitCoverage(Sample.Coverage8bit);
 
+	// TODO: Add atomspheric contribution
 	OutColor.rgb = Radiance.xyz * Attenuation * LocalCoverage * GetExposure();
 	OutColor.a = LocalCoverage;
 }
-#endif
\ No newline at end of file
+#endif
diff --git a/Engine/Shaders/Private/HairStrands/HairScatterCompose.usf b/Engine/Shaders/Private/HairStrands/HairScatterCompose.usf
--- a/Engine/Shaders/Private/HairStrands/HairScatterCompose.usf
+++ b/Engine/Shaders/Private/HairStrands/HairScatterCompose.usf
@@ -8,9 +8,9 @@
 #include "HairStrandsVisibilityCommon.ush"
 
 Texture2D<float4> SceneColorTexture;
-Texture2D<float4> SceneColorSubPixelTexture;
 Texture2D<uint4>  CategorizationTexture;
 
+// todo aggregate sample
 void MainPS(
 	in FScreenVertexOutput Input,
 	out float4 OutColor  : SV_Target0)
@@ -23,12 +23,12 @@ void MainPS(
 	const bool bIsCoveredByHair = CatData.PixelCoverage > 0;
 	if (bIsCoveredByHair)
 	{
-		if (CatData.PixelCoverage < 1)
-		{
-			const float4 Color = SceneColorSubPixelTexture.Load(uint3(PixelCoord, 0));
-			OutColor.xyz = Color.xyz / Color.w;
-			OutColor.w = Color.w;
-		}
+		//if (CatData.PixelCoverage < 1)
+		//{
+		//	const float4 Color = SceneColorSubPixelTexture.Load(uint3(PixelCoord, 0));
+		//	OutColor.xyz = Color.xyz / Color.w;
+		//	OutColor.w = Color.w;
+		//}
 
 		if (CatData.PixelCoverage >= 1)
 		{
diff --git a/Engine/Shaders/Private/HairStrands/HairStrandsEnvironmentLighting.usf b/Engine/Shaders/Private/HairStrands/HairStrandsEnvironmentLighting.usf
--- a/Engine/Shaders/Private/HairStrands/HairStrandsEnvironmentLighting.usf
+++ b/Engine/Shaders/Private/HairStrands/HairStrandsEnvironmentLighting.usf
@@ -553,84 +553,6 @@ float3 ComputeThinTransmission(FGBufferData GBuffer, float2 ScreenPosition, floa
 
 ////////////////////////////////////////////////////////////////////////////////////////////////////////////////////////////////////////////
 
-
-#if LIGHTING_CS
-
-uint MaxVisibilityNodeCount;
-StructuredBuffer<uint>	HairVisibilityNodeCoord;
-Buffer<uint4>			IndirectArgsBuffer;
-RWBuffer<float4>		OutLightingBuffer;
-
-#if PERMUTATION_GROUP_SIZE == 0
-#define LIGHTING_GROUP_SIZE 64
-#elif PERMUTATION_GROUP_SIZE == 1
-#define LIGHTING_GROUP_SIZE 32
-#else
-#error Unknown group size	
-#endif
-
-[numthreads(LIGHTING_GROUP_SIZE, 1, 1)]
-void MainCS(uint3 DispatchThreadId : SV_DispatchThreadID)
-{
-	// Note: Use a fixed group count width (HAIR_VISIBILITY_GROUP_COUNT_WIDTH) for avoiding loading the indirect args (to avoid dep. memory fetch)
-	const uint LocalOffset = DispatchThreadId.x + DispatchThreadId.y * LIGHTING_GROUP_SIZE * HAIR_VISIBILITY_GROUP_COUNT_WIDTH;
-
-	const uint PackedCoord = HairVisibilityNodeCoord[LocalOffset];
-	const uint2 PixelCoord = uint2((PackedCoord & 0xFFFF), ((PackedCoord >> 16) & 0xFFFF));
-	const float2 SvPosition = PixelCoord + float2(0.5f, 0.5f);
-
-	const float2 UV = (PixelCoord + float2(0.5f, 0.5f)) / float2(View.BufferSizeAndInvSize.xy);
-	const float2 ScreenPosition = (UV - View.ScreenPositionScaleBias.wz) / View.ScreenPositionScaleBias.xy;
-
-	if (LocalOffset > MaxVisibilityNodeCount)
-		return;
-
-	const FPackedHairSample PackedSample = HairVisibilityNodeData[LocalOffset];
-	const FHairSample Sample = UnpackHairSample(PackedSample);
-
-	FGBufferData GBuffer;
-	GBuffer.ShadingModelID = SHADINGMODELID_HAIR;
-	GBuffer.WorldNormal = Sample.Tangent;
-	GBuffer.BaseColor = Sample.BaseColor;
-	GBuffer.Roughness = Sample.Roughness;
-	GBuffer.Specular = Sample.Specular;
-	GBuffer.Depth = ConvertFromDeviceZ(Sample.Depth);
-	GBuffer.GBufferAO = 1;
-	GBuffer.CustomData = float4(HairDualScatteringRoughnessOverride,0,Sample.Backlit,0);
-	GBuffer.IndirectIrradiance = 1000000;
-
-	bool bEnabled = true;
-	const bool bIsMacroGroupCompatible = Sample.MacroGroupId == Voxel_MacroGroupId;
-#if USE_SCENE_COLOR
-	bEnabled = Sample.bScatterSceneLighting;
-#endif
-	if (bIsMacroGroupCompatible && bEnabled)
-	{
-		const float SampleCoverage = saturate(From8bitCoverage(Sample.Coverage8bit));
-	#if USE_SCENE_COLOR
-		const float3 SkyLighting = ComputeThinTransmission(GBuffer, ScreenPosition, SvPosition, Sample.PrimitiveId);
-	#else
-		const float3 SkyLighting = ReflectionEnvironment(GBuffer, ScreenPosition, SvPosition, Sample.PrimitiveId);
-	#endif
-
-		FLightAccumulator LightAccumulator = (FLightAccumulator)0;
-		LightAccumulator_Add(LightAccumulator, SkyLighting, SkyLighting, 1.0f, false);
-		float4 OutColor = LightAccumulator_GetResult(LightAccumulator);
-		OutColor.xyz *= SampleCoverage;
-		OutColor.a = SampleCoverage;
-
-		OutLightingBuffer[LocalOffset] = OutColor;
-	}
-	else
-	{
-		OutLightingBuffer[LocalOffset] = 0;
-	}
-}
-
-#endif // LIGHTING_CS
-
-
-
 #if LIGHTING_VS
 
 int2 MaxViewportResolution;
@@ -661,14 +583,6 @@ uint MaxVisibilityNodeCount;
 StructuredBuffer<uint>	HairVisibilityNodeCoord;
 Buffer<uint4>			IndirectArgsBuffer;
 
-#if PERMUTATION_GROUP_SIZE == 0
-#define LIGHTING_GROUP_SIZE 64
-#elif PERMUTATION_GROUP_SIZE == 1
-#define LIGHTING_GROUP_SIZE 32
-#else
-#error Unknown group size	
-#endif
-
 void MainPS(
 	float4 SVPos : SV_POSITION,
 	uint NodeCount : DISPATCH_NODECOUNT,
diff --git a/Engine/Shaders/Private/HairStrands/HairStrandsEnvironmentLightingCompose.usf b/Engine/Shaders/Private/HairStrands/HairStrandsEnvironmentLightingCompose.usf
deleted file mode 100644
diff --git a/Engine/Source/Runtime/Renderer/Private/DeferredShadingRenderer.cpp b/Engine/Source/Runtime/Renderer/Private/DeferredShadingRenderer.cpp
--- a/Engine/Source/Runtime/Renderer/Private/DeferredShadingRenderer.cpp
+++ b/Engine/Source/Runtime/Renderer/Private/DeferredShadingRenderer.cpp
@@ -2621,7 +2621,7 @@ void FDeferredShadingSceneRenderer::Render(FRHICommandListImmediate& RHICmdList)
 		SCOPED_GPU_STAT(RHICmdList, HairRendering);
 		if (HairDatas)
 		{
-			RenderHairComposeSubPixel(RHICmdList, Views, HairDatas);
+			RenderHairComposition(RHICmdList, Views, HairDatas);
 		}
 
 		if (IsHairStrandsEnable(Scene->GetShaderPlatform()))
diff --git a/Engine/Source/Runtime/Renderer/Private/DeferredShadingRenderer.h b/Engine/Source/Runtime/Renderer/Private/DeferredShadingRenderer.h
--- a/Engine/Source/Runtime/Renderer/Private/DeferredShadingRenderer.h
+++ b/Engine/Source/Runtime/Renderer/Private/DeferredShadingRenderer.h
@@ -412,16 +412,10 @@ private:
 	  * @param LightIndex The light's index into FScene::Lights
 	  * @return true if anything got rendered
 	  */
-	void RenderLight(FRHICommandList& RHICmdList, const FLightSceneInfo* LightSceneInfo, IPooledRenderTarget* ScreenShadowMaskTexture, IPooledRenderTarget* ScreenShadowMaskSubPixelTexture, struct FHairStrandsTransmittanceMaskData* TransmittanceMaskData, const struct FHairStrandsVisibilityViews* InHairVisibilityViews, bool bRenderOverlap, bool bIssueDrawEvent, bool bSupportHair);
+	void RenderLight(FRHICommandList& RHICmdList, const FLightSceneInfo* LightSceneInfo, IPooledRenderTarget* ScreenShadowMaskTexture, const struct FHairStrandsVisibilityViews* InHairVisibilityViews, bool bRenderOverlap, bool bIssueDrawEvent);
 
 	/** Specialized version of RenderLight for hair (run lighting evaluation on at sub-pixel rate, without depth bound) */
-	void RenderLightForHair(
-		FRHICommandList& RHICmdList,
-		const FLightSceneInfo* LightSceneInfo,
-		IPooledRenderTarget* ScreenShadowMaskSubPixelTexture,
-		FHairStrandsTransmittanceMaskData* InTransmittanceMaskData,
-		const FHairStrandsVisibilityViews* InHairVisibilityViews,
-		bool bIssueDrawEvent);
+	void RenderLightForHair(FRHICommandList& RHICmdList, const FLightSceneInfo* LightSceneInfo, IPooledRenderTarget* ScreenShadowMaskSubPixelTexture, struct FHairStrandsTransmittanceMaskData* InTransmittanceMaskData, const struct FHairStrandsVisibilityViews* InHairVisibilityViews);
 
 	/** Renders an array of simple lights using standard deferred shading. */
 	void RenderSimpleLightsStandardDeferred(FRHICommandListImmediate& RHICmdList, const FSimpleLightArray& SimpleLights);
diff --git a/Engine/Source/Runtime/Renderer/Private/HairStrands/HairStrandsComposition.cpp b/Engine/Source/Runtime/Renderer/Private/HairStrands/HairStrandsComposition.cpp
--- a/Engine/Source/Runtime/Renderer/Private/HairStrands/HairStrandsComposition.cpp
+++ b/Engine/Source/Runtime/Renderer/Private/HairStrands/HairStrandsComposition.cpp
@@ -107,91 +107,6 @@ static void AddHairVisibilityComposeSamplePass(
 	});
 }
 
-///////////////////////////////////////////////////////////////////////////////////////////////////
-class FHairVisibilityComposeSubPixelPS : public FGlobalShader
-{
-	DECLARE_GLOBAL_SHADER(FHairVisibilityComposeSubPixelPS);
-	SHADER_USE_PARAMETER_STRUCT(FHairVisibilityComposeSubPixelPS, FGlobalShader);
-
-	using FPermutationDomain = TShaderPermutationDomain<>;
-
-	BEGIN_SHADER_PARAMETER_STRUCT(FParameters, )
-		SHADER_PARAMETER_RDG_TEXTURE(Texture2D, SubPixelColorTexture)
-		SHADER_PARAMETER_RDG_TEXTURE(Texture2D, CategorisationTexture)
-		RENDER_TARGET_BINDING_SLOTS()
-	END_SHADER_PARAMETER_STRUCT()
-
-public:
-	static bool ShouldCompilePermutation(const FGlobalShaderPermutationParameters& Parameters) { return Parameters.Platform == SP_PCD3D_SM5; }
-	static void ModifyCompilationEnvironment(const FGlobalShaderPermutationParameters& Parameters, FShaderCompilerEnvironment& OutEnvironment)
-	{
-		FGlobalShader::ModifyCompilationEnvironment(Parameters, OutEnvironment);
-		OutEnvironment.SetDefine(TEXT("SHADER_SUBCOLOR"), 1);
-	}
-};
-
-IMPLEMENT_GLOBAL_SHADER(FHairVisibilityComposeSubPixelPS, "/Engine/Private/HairStrands/HairStrandsVisibilityComposeSubPixelPS.usf", "SubColorPS", SF_Pixel);
-
-static void AddHairVisibilityComposeSubPixelPass(
-	FRDGBuilder& GraphBuilder,
-	const FViewInfo& View,
-	const FRDGTextureRef& SubPixelColorTexture,
-	const FRDGTextureRef CategorisationTexture,
-	FRDGTextureRef& OutColorTexture,
-	FRDGTextureRef& OutDepthTexture)
-{
-	FHairVisibilityComposeSubPixelPS::FParameters* Parameters = GraphBuilder.AllocParameters<FHairVisibilityComposeSubPixelPS::FParameters>();
-	Parameters->SubPixelColorTexture = SubPixelColorTexture;
-	Parameters->CategorisationTexture = CategorisationTexture;
-	Parameters->RenderTargets[0] = FRenderTargetBinding(OutColorTexture, ERenderTargetLoadAction::ELoad);
-	Parameters->RenderTargets.DepthStencil = FDepthStencilBinding(OutDepthTexture, ERenderTargetLoadAction::ELoad, ERenderTargetLoadAction::ENoAction, FExclusiveDepthStencil::DepthWrite_StencilNop);
-
-	FHairVisibilityComposeSubPixelPS::FPermutationDomain PermutationVector;
-	TShaderMapRef<FHairVisibilityComposeSubPixelPS> PixelShader(View.ShaderMap, PermutationVector);
-	TShaderMapRef<FPostProcessVS> VertexShader(View.ShaderMap);
-	const FIntRect Viewport = View.ViewRect;
-	const FIntPoint Resolution = OutColorTexture->Desc.Extent;
-	const FViewInfo* CapturedView = &View;
-
-	{
-		ClearUnusedGraphResources(PixelShader, Parameters);
-
-		GraphBuilder.AddPass(
-			RDG_EVENT_NAME("HairStrandsVisibilityComposeSubSPixel"),
-			Parameters,
-			ERDGPassFlags::Raster,
-			[Parameters, VertexShader, PixelShader, Viewport, Resolution, CapturedView](FRHICommandList& RHICmdList)
-		{
-			FGraphicsPipelineStateInitializer GraphicsPSOInit;
-			RHICmdList.ApplyCachedRenderTargets(GraphicsPSOInit);
-			GraphicsPSOInit.BlendState = TStaticBlendState<CW_RGBA, BO_Add, BF_One, BF_InverseSourceAlpha, BO_Add, BF_Zero, BF_InverseSourceAlpha>::GetRHI();
-			GraphicsPSOInit.RasterizerState = TStaticRasterizerState<>::GetRHI();
-
-			// Write stencil value for partially covered pixel, in order to run responsive AA on them
-			GraphicsPSOInit.DepthStencilState = TStaticDepthStencilState<true, CF_DepthNearOrEqual>::GetRHI();
-			GraphicsPSOInit.BoundShaderState.VertexDeclarationRHI = GFilterVertexDeclaration.VertexDeclarationRHI;
-			GraphicsPSOInit.BoundShaderState.VertexShaderRHI = VertexShader.GetVertexShader();
-			GraphicsPSOInit.BoundShaderState.PixelShaderRHI = PixelShader.GetPixelShader();
-			GraphicsPSOInit.PrimitiveType = PT_TriangleList;
-			SetGraphicsPipelineState(RHICmdList, GraphicsPSOInit);
-
-			VertexShader->SetParameters(RHICmdList, CapturedView->ViewUniformBuffer);
-			RHICmdList.SetViewport(Viewport.Min.X, Viewport.Min.Y, 0.0f, Viewport.Max.X, Viewport.Max.Y, 1.0f);
-			SetShaderParameters(RHICmdList, PixelShader, PixelShader.GetPixelShader(), *Parameters);
-			DrawRectangle(
-				RHICmdList,
-				0, 0,
-				Viewport.Width(), Viewport.Height(),
-				Viewport.Min.X, Viewport.Min.Y,
-				Viewport.Width(), Viewport.Height(),
-				Viewport.Size(),
-				Resolution,
-				VertexShader,
-				EDRF_UseTriangleOptimization);
-		});
-	}
-}
-
 ///////////////////////////////////////////////////////////////////////////////////////////////////
 class FHairVisibilityFastResolvePS : public FGlobalShader
 {
@@ -384,7 +299,7 @@ static void AddPatchGbufferDataPass(
 }
 ///////////////////////////////////////////////////////////////////////////////////////////////////
 
-void RenderHairComposeSubPixel(
+void RenderHairComposition(
 	FRHICommandListImmediate& RHICmdList,
 	const TArray<FViewInfo>& Views,
 	const FHairStrandsDatas* HairDatas)
@@ -394,14 +309,13 @@ void RenderHairComposeSubPixel(
 
 	const FHairStrandsVisibilityViews& HairVisibilityViews = HairDatas->HairVisibilityViews;
 
-	DECLARE_GPU_STAT(HairStrandsComposeSubPixel);
-	SCOPED_DRAW_EVENT(RHICmdList, HairStrandsComposeSubPixel);
-	SCOPED_GPU_STAT(RHICmdList, HairStrandsComposeSubPixel);
+	DECLARE_GPU_STAT(HairStrandsComposition);
+	SCOPED_DRAW_EVENT(RHICmdList, HairStrandsComposition);
+	SCOPED_GPU_STAT(RHICmdList, HairStrandsComposition);
 
 	FSceneRenderTargets& SceneTargets = FSceneRenderTargets::Get(RHICmdList);
 	FRDGBuilder GraphBuilder(RHICmdList);
 
-	FRDGTextureRef SceneColorSubPixelTexture = GraphBuilder.RegisterExternalTexture(SceneTargets.SceneColorSubPixel, TEXT("SceneColorSubPixelTexture"));
 	FRDGTextureRef SceneColorTexture = GraphBuilder.RegisterExternalTexture(SceneTargets.GetSceneColor(), TEXT("SceneColorTexture"));
 	FRDGTextureRef SceneColorDepth = GraphBuilder.RegisterExternalTexture(SceneTargets.SceneDepthZ, TEXT("SceneDepthTexture"));
 	for (int32 ViewIndex = 0; ViewIndex < Views.Num(); ViewIndex++)
@@ -421,38 +335,23 @@ void RenderHairComposeSubPixel(
 				}
 				const FRDGTextureRef RDGCategorisationTexture = CategorisationTexture ? GraphBuilder.RegisterExternalTexture(CategorisationTexture, TEXT("HairVisibilityCategorisationTexture")) : nullptr;
 
-				AddHairDiffusionPass(
+				// todo: rehook the diffusion pass
+				//AddHairDiffusionPass(
+				//	GraphBuilder,
+				//	View,
+				//	VisibilityData,
+				//	MacroGroupDatas.VirtualVoxelResources,
+				//	SceneColorDepth,
+				//	SceneColorSubPixelTexture,
+				//	SceneColorTexture);
+
+				AddHairVisibilityComposeSamplePass(
 					GraphBuilder,
 					View,
 					VisibilityData,
-					MacroGroupDatas.VirtualVoxelResources,
-					SceneColorDepth,
-					SceneColorSubPixelTexture,
-					SceneColorTexture);
-
-				if (VisibilityData.SampleLightingBuffer)
-				{
-					AddHairVisibilityComposeSamplePass(
-						GraphBuilder,
-						View,
-						VisibilityData,
-						RDGCategorisationTexture,
-						SceneColorTexture,
-						SceneColorDepth);
-				}
-				else
-				{
-					// #hair_todo : compose partially covered hair with transparent surface: this can be done by 
-					// rendering quad(s) covering the hair at the correct depth. This will be sorted with other 
-					// transparent surface, which should make the overall sorting workable
-					AddHairVisibilityComposeSubPixelPass(
-						GraphBuilder,
-						View,
-						SceneColorSubPixelTexture,
-						RDGCategorisationTexture,
-						SceneColorTexture,
-						SceneColorDepth);
-				}
+					RDGCategorisationTexture,
+					SceneColorTexture,
+					SceneColorDepth);
 
 				if (HairVisibilityViews.HairDatas[ViewIndex].VelocityTexture)
 				{
diff --git a/Engine/Source/Runtime/Renderer/Private/HairStrands/HairStrandsComposition.h b/Engine/Source/Runtime/Renderer/Private/HairStrands/HairStrandsComposition.h
--- a/Engine/Source/Runtime/Renderer/Private/HairStrands/HairStrandsComposition.h
+++ b/Engine/Source/Runtime/Renderer/Private/HairStrands/HairStrandsComposition.h
@@ -9,7 +9,7 @@
 #include "CoreMinimal.h"
 #include "RendererInterface.h"
 
-void RenderHairComposeSubPixel(
+void RenderHairComposition(
 	FRHICommandListImmediate& RHICmdList, 
 	const TArray<FViewInfo>& Views,
 	const struct FHairStrandsDatas* HairDatas);
diff --git a/Engine/Source/Runtime/Renderer/Private/HairStrands/HairStrandsEnvironment.cpp b/Engine/Source/Runtime/Renderer/Private/HairStrands/HairStrandsEnvironment.cpp
--- a/Engine/Source/Runtime/Renderer/Private/HairStrands/HairStrandsEnvironment.cpp
+++ b/Engine/Source/Runtime/Renderer/Private/HairStrands/HairStrandsEnvironment.cpp
@@ -176,87 +176,6 @@ static void AddHairStrandsEnvironmentAOPass(
 	});
 }
 
-///////////////////////////////////////////////////////////////////////////////////////////////////
-
-class FHairEnvironmentLightingComposePS : public FGlobalShader
-{
-	DECLARE_GLOBAL_SHADER(FHairEnvironmentLightingComposePS);
-	SHADER_USE_PARAMETER_STRUCT(FHairEnvironmentLightingComposePS, FGlobalShader);
-
-	BEGIN_SHADER_PARAMETER_STRUCT(FParameters, )
-		SHADER_PARAMETER_RDG_TEXTURE(Texture2D, HairCategorizationTexture)
-		SHADER_PARAMETER_RDG_TEXTURE(Texture2D, HairVisibilityNodeOffsetAndCount)
-		SHADER_PARAMETER_RDG_BUFFER_SRV(Buffer, HairLightingSampleBuffer)
-		RENDER_TARGET_BINDING_SLOTS()
-	END_SHADER_PARAMETER_STRUCT()
-
-public:
-	static bool ShouldCompilePermutation(const FGlobalShaderPermutationParameters& Parameters) { return Parameters.Platform == SP_PCD3D_SM5; }
-};
-
-IMPLEMENT_GLOBAL_SHADER(FHairEnvironmentLightingComposePS, "/Engine/Private/HairStrands/HairStrandsEnvironmentLightingCompose.usf", "MainPS", SF_Pixel);
-
-static void AddHairEnvironmentLightingComposePass(
-	FRDGBuilder& GraphBuilder,
-	const FViewInfo& View,
-	const FHairStrandsVisibilityData& VisibilityData,
-	const FHairStrandsMacroGroupData& MacroGroupData,
-	const FRDGBufferRef SampleLightingBuffer,
-	FRDGTextureRef& OutColorTexture,
-	FRDGTextureRef& OutSubColorTexture)
-{
-	FHairEnvironmentLightingComposePS::FParameters* Parameters = GraphBuilder.AllocParameters<FHairEnvironmentLightingComposePS::FParameters>();
-	Parameters->HairCategorizationTexture = GraphBuilder.RegisterExternalTexture(VisibilityData.CategorizationTexture);
-	Parameters->HairVisibilityNodeOffsetAndCount = GraphBuilder.RegisterExternalTexture(VisibilityData.NodeIndex);
-	Parameters->HairLightingSampleBuffer  = GraphBuilder.CreateSRV(FRDGBufferSRVDesc(SampleLightingBuffer, EPixelFormat::PF_FloatRGBA));
-	Parameters->RenderTargets[0] = FRenderTargetBinding(OutColorTexture, ERenderTargetLoadAction::ELoad);
-	Parameters->RenderTargets[1] = FRenderTargetBinding(OutSubColorTexture, ERenderTargetLoadAction::ELoad);
-
-	TShaderMapRef<FPostProcessVS> VertexShader(View.ShaderMap);
-	TShaderMapRef<FHairEnvironmentLightingComposePS> PixelShader(View.ShaderMap);
-	const FGlobalShaderMap* GlobalShaderMap = View.ShaderMap;
-	const FIntRect Viewport = View.ViewRect;
-	const FIntPoint Resolution = OutColorTexture->Desc.Extent;
-	const FViewInfo* CapturedView = &View;
-
-	ClearUnusedGraphResources(PixelShader, Parameters);
-
-	GraphBuilder.AddPass(
-		RDG_EVENT_NAME("HairStrandsEnvironmentLightingCompose"),
-		Parameters,
-		ERDGPassFlags::Raster,
-		[Parameters, VertexShader, PixelShader, Viewport, Resolution, CapturedView](FRHICommandList& RHICmdList)
-	{
-		FGraphicsPipelineStateInitializer GraphicsPSOInit;
-		RHICmdList.ApplyCachedRenderTargets(GraphicsPSOInit);
-		GraphicsPSOInit.BlendState = TStaticBlendState<
-			CW_RGBA, BO_Add, BF_One, BF_One, BO_Add, BF_One, BF_One,
-			CW_RGBA, BO_Add, BF_One, BF_One, BO_Max, BF_SourceAlpha, BF_DestAlpha>::GetRHI();
-		GraphicsPSOInit.RasterizerState = TStaticRasterizerState<>::GetRHI();
-		GraphicsPSOInit.DepthStencilState = TStaticDepthStencilState<false, CF_Always>::GetRHI();
-		GraphicsPSOInit.BoundShaderState.VertexDeclarationRHI = GFilterVertexDeclaration.VertexDeclarationRHI;
-		GraphicsPSOInit.BoundShaderState.VertexShaderRHI = VertexShader.GetVertexShader();
-		GraphicsPSOInit.BoundShaderState.PixelShaderRHI = PixelShader.GetPixelShader();
-		GraphicsPSOInit.PrimitiveType = PT_TriangleList;
-		SetGraphicsPipelineState(RHICmdList, GraphicsPSOInit);
-
-		VertexShader->SetParameters(RHICmdList, CapturedView->ViewUniformBuffer);
-		RHICmdList.SetViewport(Viewport.Min.X, Viewport.Min.Y, 0.0f, Viewport.Max.X, Viewport.Max.Y, 1.0f);
-		SetShaderParameters(RHICmdList, PixelShader, PixelShader.GetPixelShader(), *Parameters);
-		DrawRectangle(
-			RHICmdList,
-			0, 0,
-			Viewport.Width(), Viewport.Height(),
-			Viewport.Min.X,   Viewport.Min.Y,
-			Viewport.Width(), Viewport.Height(),
-			Viewport.Size(),
-			Resolution,
-			VertexShader,
-			EDRF_UseTriangleOptimization);
-	});
-}
-
-
 ///////////////////////////////////////////////////////////////////////////////////////////////////
 
 class FHairEnvironmentLighting
@@ -266,7 +185,7 @@ public:
 	class FSampleSet : SHADER_PERMUTATION_INT("PERMUTATION_SAMPLESET", 2);
 	class FVirtualVoxel : SHADER_PERMUTATION_INT("VOXEL_TRAVERSAL_TYPE", 5);
 	class FUseSceneColor : SHADER_PERMUTATION_INT("USE_SCENE_COLOR", 2);
-	using FPermutationDomain = TShaderPermutationDomain<FGroupSize, FSampleSet, FVirtualVoxel, FUseSceneColor>;
+	using FPermutationDomain = TShaderPermutationDomain<FSampleSet, FVirtualVoxel, FUseSceneColor>;
 
 	static void ModifyCompilationEnvironment(const FGlobalShaderPermutationParameters& Parameters, FShaderCompilerEnvironment& OutEnvironment)
 	{
@@ -339,7 +258,6 @@ public:
 		SHADER_PARAMETER_SAMPLER(SamplerState, HairLUTSampler)
 
 		SHADER_PARAMETER_RDG_TEXTURE(Texture3D, Voxel_DensityTexture)
-		SHADER_PARAMETER_RDG_BUFFER(StructuredBuffer, IndirectArgsBuffer)
 		SHADER_PARAMETER_RDG_BUFFER_UAV(RWBuffer, OutLightingBuffer)
 
 		SHADER_PARAMETER_STRUCT_REF(FViewUniformShaderParameters, ViewUniformBuffer)
@@ -354,23 +272,6 @@ public:
 	END_SHADER_PARAMETER_STRUCT()
 };
 
-class FHairEnvironmentLightingCS : public FGlobalShader
-{
-	DECLARE_GLOBAL_SHADER(FHairEnvironmentLightingCS);
-	SHADER_USE_PARAMETER_STRUCT(FHairEnvironmentLightingCS, FGlobalShader)
-	using FPermutationDomain = FHairEnvironmentLighting::FPermutationDomain;
-
-	BEGIN_SHADER_PARAMETER_STRUCT(FParameters, )
-		SHADER_PARAMETER_STRUCT_INCLUDE(FHairEnvironmentLighting::FParameters, Common)
-	END_SHADER_PARAMETER_STRUCT()
-
-	static void ModifyCompilationEnvironment(const FGlobalShaderPermutationParameters& Parameters, FShaderCompilerEnvironment& OutEnvironment)
-	{
-		FHairEnvironmentLighting::ModifyCompilationEnvironment(Parameters, OutEnvironment);
-		OutEnvironment.SetDefine(TEXT("LIGHTING_CS"), 1);
-	}
-};
-
 class FHairEnvironmentLightingVS : public FGlobalShader
 {
 	DECLARE_GLOBAL_SHADER(FHairEnvironmentLightingVS);
@@ -406,127 +307,16 @@ class FHairEnvironmentLightingPS : public FGlobalShader
 	}
 };
 
-IMPLEMENT_GLOBAL_SHADER(FHairEnvironmentLightingCS, "/Engine/Private/HairStrands/HairStrandsEnvironmentLighting.usf", "MainCS", SF_Compute);
 IMPLEMENT_GLOBAL_SHADER(FHairEnvironmentLightingPS, "/Engine/Private/HairStrands/HairStrandsEnvironmentLighting.usf", "MainPS", SF_Pixel);
 IMPLEMENT_GLOBAL_SHADER(FHairEnvironmentLightingVS, "/Engine/Private/HairStrands/HairStrandsEnvironmentLighting.usf", "MainVS", SF_Vertex);
 
-static FRDGBufferRef AddHairStrandsEnvironmentLightingPassCS(
-	FRDGBuilder& GraphBuilder,
-	const FViewInfo& View,
-	const FHairStrandsVisibilityData& VisibilityData,
-	const FHairStrandsMacroGroupDatas& MacroGroupDatas,
-	const FHairStrandsMacroGroupData& MacroGroupData,
-	const FRDGTextureRef SceneColorTexture,
-	const uint32 NodeGroupSize,
-	FRDGBufferRef IndirectArgsBuffer)
-{
-	FRDGBufferRef OutBuffer = GraphBuilder.CreateBuffer(FRDGBufferDesc::CreateBufferDesc(
-		4 * sizeof(uint16),
-		VisibilityData.NodeData->Desc.NumElements),
-		TEXT("HairSkyLightingNodeData"));
-
-	FSceneTextureParameters SceneTextures;
-	SetupSceneTextureParameters(GraphBuilder, &SceneTextures);
-
-	const bool bUseVirtualVoxel = MacroGroupDatas.VirtualVoxelResources.IsValid();
-	const int32 VoxelTraversalType = FMath::Clamp(GHairSkylightingVirtualVoxel, 0, 4);
-
-	// Render the reflection environment with tiled deferred culling
-	const bool bHasBoxCaptures = (View.NumBoxReflectionCaptures > 0);
-	const bool bHasSphereCaptures = (View.NumSphereReflectionCaptures > 0);
-	FHairEnvironmentLighting::FParameters PassParameters;
-
-	const FHairLUT InHairLUT = GetHairLUT(GraphBuilder.RHICmdList, View);
-	PassParameters.HairEnergyLUTTexture = GraphBuilder.RegisterExternalTexture(InHairLUT.Textures[HairLUTType_MeanEnergy], TEXT("HairMeanEnergyLUTTexture"));;
-	PassParameters.HairScatteringLUTTexture = GraphBuilder.RegisterExternalTexture(InHairLUT.Textures[HairLUTType_DualScattering], TEXT("HairScatteringEnergyLUTTexture"));
-	PassParameters.HairLUTSampler = TStaticSamplerState<SF_Bilinear, AM_Clamp, AM_Clamp, AM_Clamp>::GetRHI();
-	PassParameters.Voxel_MacroGroupId = MacroGroupData.MacroGroupId;
-	if (!bUseVirtualVoxel)
-	{
-		PassParameters.Voxel_MinAABB = MacroGroupData.GetMinBound();
-		PassParameters.Voxel_MaxAABB = MacroGroupData.GetMaxBound();
-		PassParameters.Voxel_Resolution = MacroGroupData.GetResolution();
-		PassParameters.Voxel_DensityTexture = GraphBuilder.RegisterExternalTexture(MacroGroupData.VoxelResources.DensityTexture);
-	}
-	const bool bUseSceneColor = SceneColorTexture != nullptr;
-	if (bUseSceneColor)
-	{
-		PassParameters.SceneColorTexture = SceneColorTexture;
-		PassParameters.HairCategorizationTexture = GraphBuilder.RegisterExternalTexture(VisibilityData.CategorizationTexture ? VisibilityData.CategorizationTexture : GSystemTextures.BlackDummy);
-	}
-	PassParameters.MaxViewportResolution = VisibilityData.SampleLightingViewportResolution;
-	PassParameters.HairVisibilityNodeCount = GraphBuilder.RegisterExternalTexture(VisibilityData.NodeCount);
-	PassParameters.Voxel_DensityScale = GetHairStrandsVoxelizationDensityScale();
-	PassParameters.Voxel_DepthBiasScale = GetHairStrandsVoxelizationDepthBiasScale();
-	PassParameters.Voxel_TanConeAngle = FMath::Tan(FMath::DegreesToRadians(GetHairStrandsSkyLightingConeAngle()));
-	PassParameters.HairDistanceThreshold = FMath::Max(GHairStrandsSkyLightingDistanceThreshold, 1.f);
-	PassParameters.bHairUseViewHairCount = VisibilityData.ViewHairCountTexture ? 1 : 0;
-	PassParameters.MaxVisibilityNodeCount = VisibilityData.NodeData->Desc.NumElements;
-	PassParameters.MultipleScatterSampleCount = FMath::Max(uint32(GHairStrandsSkyLightingSampleCount), 1u);
-	PassParameters.JitterSphericalIntegration = GHairStrandsSkyLightingJitterSphericalIntegration ? 1 : 0;
-	PassParameters.HairComponents = ToBitfield(GetHairComponents());
-	PassParameters.HairDualScatteringRoughnessOverride = GetHairDualScatteringRoughnessOverride();
-	PassParameters.PreIntegratedGF = GSystemTextures.PreintegratedGF->GetRenderTargetItem().ShaderResourceTexture;
-	PassParameters.PreIntegratedGFSampler = TStaticSamplerState<SF_Bilinear, AM_Clamp, AM_Clamp, AM_Clamp>::GetRHI();
-	PassParameters.HairCountTexture = GraphBuilder.RegisterExternalTexture(VisibilityData.ViewHairCountTexture ? VisibilityData.ViewHairCountTexture : GSystemTextures.BlackDummy);
-	PassParameters.SceneTextures = SceneTextures;
-	PassParameters.VirtualVoxel = MacroGroupDatas.VirtualVoxelResources.UniformBuffer;
-	SetupSceneTextureSamplers(&PassParameters.SceneTextureSamplers);
-	PassParameters.ViewUniformBuffer = View.ViewUniformBuffer;
-	PassParameters.ReflectionCaptureData = View.ReflectionCaptureUniformBuffer;
-	{
-		FReflectionUniformParameters ReflectionUniformParameters;
-		SetupReflectionUniformParameters(View, ReflectionUniformParameters);
-		PassParameters.ReflectionsParameters = CreateUniformBufferImmediate(ReflectionUniformParameters, UniformBuffer_SingleDraw);
-	}
-	PassParameters.ForwardLightData = View.ForwardLightingResources->ForwardLightDataUniformBuffer;
-	PassParameters.IndirectArgsBuffer = IndirectArgsBuffer;
-	PassParameters.OutLightingBuffer = GraphBuilder.CreateUAV(FRDGBufferUAVDesc(OutBuffer, PF_FloatRGBA));
-
-	if (ShaderDrawDebug::IsShaderDrawDebugEnabled(View))
-	{
-		ShaderDrawDebug::SetParameters(GraphBuilder, View.ShaderDrawData, PassParameters.ShaderDrawParameters);
-	}
-
-	// Bind hair data
-	FRDGBufferRef InHairVisibilityNodeData = GraphBuilder.RegisterExternalBuffer(VisibilityData.NodeData, TEXT("HairVisibilityNodeData"));
-	FRDGBufferRef InHairVisibilityNodeCoord = GraphBuilder.RegisterExternalBuffer(VisibilityData.NodeCoord, TEXT("HairVisibilityNodeCoord"));
-	PassParameters.HairVisibilityNodeData = GraphBuilder.CreateSRV(InHairVisibilityNodeData);
-	PassParameters.HairVisibilityNodeCoord = GraphBuilder.CreateSRV(InHairVisibilityNodeCoord);
-
-
-	check(NodeGroupSize == 64 || NodeGroupSize == 32);
-	FHairEnvironmentLightingCS::FPermutationDomain PermutationVector;
-	PermutationVector.Set<FHairEnvironmentLighting::FGroupSize>(NodeGroupSize == 64 ? 0 : (NodeGroupSize == 32 ? 1 : 2));
-	PermutationVector.Set<FHairEnvironmentLighting::FSampleSet>(PassParameters.MultipleScatterSampleCount <= 16 ? 0 : 1);
-	PermutationVector.Set<FHairEnvironmentLighting::FVirtualVoxel>(bUseVirtualVoxel ? VoxelTraversalType : 0);
-	PermutationVector.Set<FHairEnvironmentLighting::FUseSceneColor>(bUseSceneColor ? 1 : 0);
-	PermutationVector = FHairEnvironmentLighting::RemapPermutation(PermutationVector);
-
-	FHairEnvironmentLightingCS::FParameters* Parameters = GraphBuilder.AllocParameters<FHairEnvironmentLightingCS::FParameters>();
-	Parameters->Common = PassParameters;
-
-	TShaderMapRef<FHairEnvironmentLightingCS> ComputeShader(View.ShaderMap, PermutationVector);
-	FComputeShaderUtils::AddPass(
-		GraphBuilder,
-		RDG_EVENT_NAME("HairStrandsEnvironmentCS %dx%d", View.ViewRect.Width(), View.ViewRect.Height()),
-		ComputeShader,
-		Parameters,
-		IndirectArgsBuffer,
-		0);
-
-	return OutBuffer;
-}
-
 static void AddHairStrandsEnvironmentLightingPassPS(
 	FRDGBuilder& GraphBuilder,
 	const FViewInfo& View,
 	const FHairStrandsVisibilityData& VisibilityData,
 	const FHairStrandsMacroGroupDatas& MacroGroupDatas,
 	const FHairStrandsMacroGroupData& MacroGroupData,
-	const FRDGTextureRef SceneColorTexture,
-	const uint32 NodeGroupSize,
-	FRDGBufferRef IndirectArgsBuffer)
+	const FRDGTextureRef SceneColorTexture)
 {
 	FSceneTextureParameters SceneTextures;
 	SetupSceneTextureParameters(GraphBuilder, &SceneTextures);
@@ -584,7 +374,6 @@ static void AddHairStrandsEnvironmentLightingPassPS(
 		PassParameters->ReflectionsParameters = CreateUniformBufferImmediate(ReflectionUniformParameters, UniformBuffer_SingleDraw);
 	}
 	PassParameters->ForwardLightData = View.ForwardLightingResources->ForwardLightDataUniformBuffer;
-	PassParameters->IndirectArgsBuffer = IndirectArgsBuffer;
 	PassParameters->OutLightingBuffer = nullptr;
 
 	if (ShaderDrawDebug::IsShaderDrawDebugEnabled(View))
@@ -598,9 +387,7 @@ static void AddHairStrandsEnvironmentLightingPassPS(
 	PassParameters->HairVisibilityNodeData = GraphBuilder.CreateSRV(InHairVisibilityNodeData);
 	PassParameters->HairVisibilityNodeCoord = GraphBuilder.CreateSRV(InHairVisibilityNodeCoord);
 
-	check(NodeGroupSize == 64 || NodeGroupSize == 32);
 	FHairEnvironmentLightingPS::FPermutationDomain PermutationVector;
-	PermutationVector.Set<FHairEnvironmentLighting::FGroupSize>(NodeGroupSize == 64 ? 0 : (NodeGroupSize == 32 ? 1 : 2));
 	PermutationVector.Set<FHairEnvironmentLighting::FSampleSet>(PassParameters->MultipleScatterSampleCount <= 16 ? 0 : 1);
 	PermutationVector.Set<FHairEnvironmentLighting::FVirtualVoxel>(bUseVirtualVoxel ? VoxelTraversalType : 0);
 	PermutationVector.Set<FHairEnvironmentLighting::FUseSceneColor>(bUseSceneColor ? 1 : 0);
@@ -611,6 +398,7 @@ static void AddHairStrandsEnvironmentLightingPassPS(
 	TShaderMapRef<FHairEnvironmentLightingVS> VertexShader(View.ShaderMap, PermutationVector);
 	TShaderMapRef<FHairEnvironmentLightingPS> PixelShader(View.ShaderMap, PermutationVector);
 
+	check(VisibilityData.SampleLightingBuffer);
 	ParametersPS->RenderTargets[0] = FRenderTargetBinding(GraphBuilder.RegisterExternalTexture(VisibilityData.SampleLightingBuffer), ERenderTargetLoadAction::ELoad);
 
 	GraphBuilder.AddPass(
@@ -671,24 +459,9 @@ void RenderHairStrandsSceneColorScattering(
 
 		FSceneRenderTargets& SceneContext = FSceneRenderTargets::Get(RHICmdList);
 		FRDGTextureRef SceneColorTexture = GraphBuilder.RegisterExternalTexture(SceneContext.GetSceneColor());
-		FRDGTextureRef SceneColorSubPixelTexture = GraphBuilder.RegisterExternalTexture(SceneContext.GetSceneColorSubPixel());
-		FRDGBufferRef NodeIndirectArgBuffer = GraphBuilder.RegisterExternalBuffer(VisibilityData.NodeIndirectArg, TEXT("HairNodeIndirectArgBuffer"));
-
-		const bool bUsePixelAccumulation = HairDatas->HairVisibilityViews.HairDatas[ViewIndex].SampleLightingBuffer != nullptr;
-		if (bUsePixelAccumulation)
-		{
-			for (const FHairStrandsMacroGroupData& MacroGroupData : HairDatas->MacroGroupsPerViews.Views[ViewIndex].Datas)
-			{
-				AddHairStrandsEnvironmentLightingPassPS(GraphBuilder, View, VisibilityData, MacroGroupDatas, MacroGroupData, SceneColorTexture, VisibilityData.NodeGroupSize, NodeIndirectArgBuffer);
-			}
-		}
-		else
+		for (const FHairStrandsMacroGroupData& MacroGroupData : HairDatas->MacroGroupsPerViews.Views[ViewIndex].Datas)
 		{
-			for (const FHairStrandsMacroGroupData& MacroGroupData : HairDatas->MacroGroupsPerViews.Views[ViewIndex].Datas)
-			{
-				FRDGBufferRef SampleLightingBuffer = AddHairStrandsEnvironmentLightingPassCS(GraphBuilder, View, VisibilityData, MacroGroupDatas, MacroGroupData, SceneColorTexture, VisibilityData.NodeGroupSize, NodeIndirectArgBuffer);
-				AddHairEnvironmentLightingComposePass(GraphBuilder, View, VisibilityData, MacroGroupData, SampleLightingBuffer, SceneColorTexture, SceneColorSubPixelTexture);
-			}
+			AddHairStrandsEnvironmentLightingPassPS(GraphBuilder, View, VisibilityData, MacroGroupDatas, MacroGroupData, SceneColorTexture);
 		}
 		GraphBuilder.Execute();
 	}
@@ -698,9 +471,7 @@ void RenderHairStrandsEnvironmentLighting(
 	FRDGBuilder& GraphBuilder,
 	const uint32 ViewIndex,
 	const TArray<FViewInfo>& Views, 
-	const FHairStrandsDatas* HairDatas,
-	FRDGTextureRef SceneColorTexture,
-	FRDGTextureRef SceneColorSubPixelTexture)
+	const FHairStrandsDatas* HairDatas)
 {
 	if (!GetHairStrandsSkyLightingEnable() || !HairDatas)
 		return;
@@ -713,27 +484,12 @@ void RenderHairStrandsEnvironmentLighting(
 	{
 		return;
 	}
-	
-	FRDGBufferRef NodeIndirectArgBuffer = GraphBuilder.RegisterExternalBuffer(VisibilityData.NodeIndirectArg, TEXT("HairNodeIndirectArgBuffer"));
 
 	const FViewInfo& View = Views[ViewIndex];
-
 	const FHairStrandsMacroGroupDatas& MacroGroupDatas = HairDatas->MacroGroupsPerViews.Views[ViewIndex];
-	const bool bUsePixelAccumulation = HairDatas->HairVisibilityViews.HairDatas[ViewIndex].SampleLightingBuffer != nullptr;
-	if (bUsePixelAccumulation)
-	{
-		for (const FHairStrandsMacroGroupData& MacroGroupData : HairDatas->MacroGroupsPerViews.Views[ViewIndex].Datas)
-		{
-			AddHairStrandsEnvironmentLightingPassPS(GraphBuilder, View, VisibilityData, MacroGroupDatas, MacroGroupData, nullptr /*SceneColorTexture*/, VisibilityData.NodeGroupSize, NodeIndirectArgBuffer);
-		}
-	}
-	else
+	for (const FHairStrandsMacroGroupData& MacroGroupData : HairDatas->MacroGroupsPerViews.Views[ViewIndex].Datas)
 	{
-		for (const FHairStrandsMacroGroupData& MacroGroupData : HairDatas->MacroGroupsPerViews.Views[ViewIndex].Datas)
-		{
-			FRDGBufferRef SampleLightingBuffer = AddHairStrandsEnvironmentLightingPassCS(GraphBuilder, View, VisibilityData, MacroGroupDatas, MacroGroupData, nullptr /*SceneColorTexture*/, VisibilityData.NodeGroupSize, NodeIndirectArgBuffer);
-			AddHairEnvironmentLightingComposePass(GraphBuilder, View, VisibilityData, MacroGroupData, SampleLightingBuffer, SceneColorTexture, SceneColorSubPixelTexture);
-		}
+		AddHairStrandsEnvironmentLightingPassPS(GraphBuilder, View, VisibilityData, MacroGroupDatas, MacroGroupData, nullptr);
 	}
 }
 
diff --git a/Engine/Source/Runtime/Renderer/Private/HairStrands/HairStrandsEnvironment.h b/Engine/Source/Runtime/Renderer/Private/HairStrands/HairStrandsEnvironment.h
--- a/Engine/Source/Runtime/Renderer/Private/HairStrands/HairStrandsEnvironment.h
+++ b/Engine/Source/Runtime/Renderer/Private/HairStrands/HairStrandsEnvironment.h
@@ -20,9 +20,7 @@ void RenderHairStrandsEnvironmentLighting(
 	FRDGBuilder& GraphBuilder,
 	const uint32 ViewIndex,
 	const TArray<FViewInfo>& Views,
-	const struct FHairStrandsDatas* HairDatas,
-	FRDGTextureRef SceneColorTexture,
-	FRDGTextureRef SceneColorSubPixelTexture);
+	const struct FHairStrandsDatas* HairDatas);
 
 void RenderHairStrandsSceneColorScattering(
 	FRHICommandListImmediate& RHICmdList,
diff --git a/Engine/Source/Runtime/Renderer/Private/HairStrands/HairStrandsScatter.cpp b/Engine/Source/Runtime/Renderer/Private/HairStrands/HairStrandsScatter.cpp
--- a/Engine/Source/Runtime/Renderer/Private/HairStrands/HairStrandsScatter.cpp
+++ b/Engine/Source/Runtime/Renderer/Private/HairStrands/HairStrandsScatter.cpp
@@ -33,7 +33,6 @@ class FHairComposePS : public FGlobalShader
 	BEGIN_SHADER_PARAMETER_STRUCT(FParameters, )
 		SHADER_PARAMETER_RDG_TEXTURE(Texture2D, CategorizationTexture)
 		SHADER_PARAMETER_RDG_TEXTURE(Texture2D, SceneColorTexture)
-		SHADER_PARAMETER_RDG_TEXTURE(Texture2D, SceneColorSubPixelTexture)
 		RENDER_TARGET_BINDING_SLOTS()
 	END_SHADER_PARAMETER_STRUCT()
 
@@ -47,8 +46,7 @@ static FRDGTextureRef AddPreScatterComposePass(
 	FRDGBuilder& GraphBuilder,
 	const FViewInfo& View,
 	const FRDGTextureRef& InCategorizationTexture,
-	const FRDGTextureRef& InSceneColorTexture,
-	const FRDGTextureRef& InSceneColorSubPixelTexture)
+	const FRDGTextureRef& InSceneColorTexture)
 {
 	const FIntPoint Resolution = InSceneColorTexture->Desc.Extent;
 	FRDGTextureRef OutputTexture;
@@ -70,7 +68,6 @@ static FRDGTextureRef AddPreScatterComposePass(
 	FHairComposePS::FParameters* Parameters = GraphBuilder.AllocParameters<FHairComposePS::FParameters>();
 	Parameters->CategorizationTexture = InCategorizationTexture;
 	Parameters->SceneColorTexture = InSceneColorTexture;
-	Parameters->SceneColorSubPixelTexture = InSceneColorSubPixelTexture;
 	Parameters->RenderTargets[0] = FRenderTargetBinding(OutputTexture, ERenderTargetLoadAction::ENoAction);
 
 	TShaderMapRef<FPostProcessVS> VertexShader(View.ShaderMap);
@@ -161,8 +158,7 @@ static FRDGTextureRef AddScatterPass(
 	const FRDGBufferRef&  InVisibilityNodeData,
 	const FRDGTextureRef& InCategorizationTexture,
 	const FRDGTextureRef& InDiffusionInput,
-	const FRDGTextureRef& OutSceneColorTexture,
-	const FRDGTextureRef& OutSceneColorSubPixelTexture)
+	const FRDGTextureRef& OutSceneColorTexture)
 {
 	if (!VoxelResources.IsValid())
 		return nullptr;
@@ -202,7 +198,7 @@ static FRDGTextureRef AddScatterPass(
 		ShaderDrawDebug::SetParameters(GraphBuilder, View.ShaderDrawData, Parameters->ShaderDrawParameters);
 	}
 	Parameters->RenderTargets[0] = FRenderTargetBinding(OutSceneColorTexture, ERenderTargetLoadAction::ELoad);
-	Parameters->RenderTargets[1] = FRenderTargetBinding(OutSceneColorSubPixelTexture, ERenderTargetLoadAction::ELoad);
+//	Parameters->RenderTargets[1] = FRenderTargetBinding(OutSceneColorSubPixelTexture, ERenderTargetLoadAction::ELoad); // TODO
 	Parameters->RenderTargets[2] = FRenderTargetBinding(OutDiffusionOutput, ERenderTargetLoadAction::ENoAction);
 
 	TShaderMapRef<FPostProcessVS> VertexShader(View.ShaderMap);
@@ -270,8 +266,7 @@ void AddHairDiffusionPass(
 	const FHairStrandsVisibilityData& VisibilityData,
 	const FVirtualVoxelResources& VoxelResources,
 	const FRDGTextureRef SceneColorDepth,
-	FRDGTextureRef OutSceneColorSubPixelTexture,
-	FRDGTextureRef OutSceneColorTexture)
+	FRDGTextureRef OutLightSampleTexture)
 {
 	const uint32 DiffusionPassCount = FMath::Clamp(GHairStrandsScatter_PassCount, 0, 8);
 	const bool bIsEnabled = DiffusionPassCount > 0 &&
@@ -291,8 +286,7 @@ void AddHairDiffusionPass(
 		GraphBuilder,
 		View,
 		CategorisationTexture,
-		OutSceneColorTexture,
-		OutSceneColorSubPixelTexture);
+		OutLightSampleTexture);
 
 	for (uint32 DiffusionPassIt = 0; DiffusionPassIt < DiffusionPassCount; ++DiffusionPassIt)
 	{
@@ -304,8 +298,8 @@ void AddHairDiffusionPass(
 			NodedData,
 			CategorisationTexture,
 			DiffusionInput,
-			OutSceneColorTexture,
-			OutSceneColorSubPixelTexture);
+			OutLightSampleTexture);
+
 		DiffusionInput = DiffusionOutput;
 	}
 }
\ No newline at end of file
diff --git a/Engine/Source/Runtime/Renderer/Private/HairStrands/HairStrandsScatter.h b/Engine/Source/Runtime/Renderer/Private/HairStrands/HairStrandsScatter.h
--- a/Engine/Source/Runtime/Renderer/Private/HairStrands/HairStrandsScatter.h
+++ b/Engine/Source/Runtime/Renderer/Private/HairStrands/HairStrandsScatter.h
@@ -12,5 +12,4 @@ void AddHairDiffusionPass(
 	const struct FHairStrandsVisibilityData& VisibilityData,
 	const struct FVirtualVoxelResources& VoxelResources,
 	const FRDGTextureRef SceneColorDepth,
-	FRDGTextureRef OutSceneColorSubPixelTexture,
 	FRDGTextureRef OutSceneColorTexture);
\ No newline at end of file
diff --git a/Engine/Source/Runtime/Renderer/Private/HairStrands/HairStrandsVisibility.cpp b/Engine/Source/Runtime/Renderer/Private/HairStrands/HairStrandsVisibility.cpp
--- a/Engine/Source/Runtime/Renderer/Private/HairStrands/HairStrandsVisibility.cpp
+++ b/Engine/Source/Runtime/Renderer/Private/HairStrands/HairStrandsVisibility.cpp
@@ -66,11 +66,6 @@ static FAutoConsoleVariableRef CVarHairStrandsVisibilityComputeRasterSampleCount
 static float GHairStrandsFullCoverageThreshold = 0.98f;
 static FAutoConsoleVariableRef CVarHairStrandsFullCoverageThreshold(TEXT("r.HairStrands.Visibility.FullCoverageThreshold"), GHairStrandsFullCoverageThreshold, TEXT("Define the coverage threshold at which a pixel is considered fully covered."));
 
-static TAutoConsoleVariable<int32> CVarHairStrandsSplitLighting(
-	TEXT("r.HairStrands.SplitLighting"), 0,
-	TEXT(""),
-	ECVF_RenderThreadSafe);
-
 /////////////////////////////////////////////////////////////////////////////////////////
 
 namespace HairStrandsVisibilityInternal
@@ -2611,8 +2606,7 @@ FHairStrandsVisibilityViews RenderHairStrandsVisibilityBuffer(
 
 			GraphBuilder.Execute();
 
-			const bool bUseSplitLighting = CVarHairStrandsSplitLighting.GetValueOnRenderThread() >  0;
-			if (bUseSplitLighting)
+			// Allocate buffer for storing all the light samples
 			{			
 				const uint32 SampleTextureResolution = FMath::CeilToInt(FMath::Sqrt(VisibilityData.MaxNodeCount));
 				VisibilityData.SampleLightingViewportResolution = FIntPoint(SampleTextureResolution, SampleTextureResolution);
diff --git a/Engine/Source/Runtime/Renderer/Private/IndirectLightRendering.cpp b/Engine/Source/Runtime/Renderer/Private/IndirectLightRendering.cpp
--- a/Engine/Source/Runtime/Renderer/Private/IndirectLightRendering.cpp
+++ b/Engine/Source/Runtime/Renderer/Private/IndirectLightRendering.cpp
@@ -653,11 +653,6 @@ void FDeferredShadingSceneRenderer::RenderDeferredReflectionsAndSkyLighting(FRHI
 	FRDGBuilder GraphBuilder(RHICmdList);
 
 	FRDGTextureRef SceneColorTexture = GraphBuilder.RegisterExternalTexture(SceneContext.GetSceneColor());
-	FRDGTextureRef SceneColorSubPixelTexture = nullptr;
-	if (HairDatas)
-	{
-		SceneColorSubPixelTexture = GraphBuilder.RegisterExternalTexture(SceneContext.GetSceneColorSubPixel());
-	}
 	FRDGTextureRef AmbientOcclusionTexture = GraphBuilder.RegisterExternalTexture(SceneContext.bScreenSpaceAOIsValid ? SceneContext.ScreenSpaceAO : GSystemTextures.WhiteDummy);
 	FRDGTextureRef DynamicBentNormalAOTexture = GraphBuilder.RegisterExternalTexture(DynamicBentNormalAO ? DynamicBentNormalAO : GSystemTextures.WhiteDummy);
 
@@ -921,19 +916,13 @@ void FDeferredShadingSceneRenderer::RenderDeferredReflectionsAndSkyLighting(FRHI
 		const bool bIsHairSkyLightingEnabled = HairDatas && (bSkyLight || bDynamicSkyLight || bReflectionEnv);
 		if (bIsHairSkyLightingEnabled)
 		{
-			RenderHairStrandsEnvironmentLighting(GraphBuilder, CurrentViewIndex, Views, HairDatas, SceneColorTexture, SceneColorSubPixelTexture);
+			RenderHairStrandsEnvironmentLighting(GraphBuilder, CurrentViewIndex, Views, HairDatas);
 		}
 	} // for (FViewInfo& View : Views)
 
 	TRefCountPtr<IPooledRenderTarget> OutSceneColor;
 	GraphBuilder.QueueTextureExtraction(SceneColorTexture, &OutSceneColor);
 
-	TRefCountPtr<IPooledRenderTarget> OutSceneSubpixelColor;
-	if (SceneColorSubPixelTexture)
-	{
-		GraphBuilder.QueueTextureExtraction(SceneColorSubPixelTexture, &OutSceneSubpixelColor);
-	}
-
 	GraphBuilder.Execute();
 
 	ResolveSceneColor(RHICmdList);
diff --git a/Engine/Source/Runtime/Renderer/Private/LightRendering.cpp b/Engine/Source/Runtime/Renderer/Private/LightRendering.cpp
--- a/Engine/Source/Runtime/Renderer/Private/LightRendering.cpp
+++ b/Engine/Source/Runtime/Renderer/Private/LightRendering.cpp
@@ -299,7 +299,7 @@ class FDeferredLightPS : public FGlobalShader
 	class FVisualizeCullingDim	: SHADER_PERMUTATION_BOOL("VISUALIZE_LIGHT_CULLING");
 	class FLightingChannelsDim	: SHADER_PERMUTATION_BOOL("USE_LIGHTING_CHANNELS");
 	class FTransmissionDim		: SHADER_PERMUTATION_BOOL("USE_TRANSMISSION");
-	class FHairLighting			: SHADER_PERMUTATION_INT("USE_HAIR_LIGHTING", 4);
+	class FHairLighting			: SHADER_PERMUTATION_INT("USE_HAIR_LIGHTING", 3);
 	class FAtmosphereTransmittance: SHADER_PERMUTATION_BOOL("USE_ATMOSPHERE_TRANSMITTANCE");
 
 	using FPermutationDomain = TShaderPermutationDomain<
@@ -349,6 +349,13 @@ class FDeferredLightPS : public FGlobalShader
 			return false;
 		}
 
+		if (PermutationVector.Get< FHairLighting >() == 2 && (
+			PermutationVector.Get< FVisualizeCullingDim >() ||
+			PermutationVector.Get< FTransmissionDim >()))
+		{
+			return false;
+		}
+
 		/*if( PermutationVector.Get< FVisualizeCullingDim >() && (
 			PermutationVector.Get< FSourceShapeDim >() == ELightSourceShape::Rect ||
 			PermutationVector.Get< FIESProfileDim >() ||
@@ -1026,14 +1033,6 @@ void FDeferredShadingSceneRenderer::RenderLights(FRHICommandListImmediate& RHICm
 			RHICmdList.WaitComputeFence(TranslucencyLightingVolumeClearEndFence);
 		}
 
-		// Clear sub pixel color when there is no light
-		if (bUseHairLighting)
-		{
-			FRHIRenderPassInfo RPInfo(SceneContext.SceneColorSubPixel->GetRenderTargetItem().TargetableTexture, MakeRenderTargetActions(ERenderTargetLoadAction::EClear, ERenderTargetStoreAction::EStore));
-			RHICmdList.BeginRenderPass(RPInfo, TEXT("ClearSubSpixelColor"));
-			RHICmdList.EndRenderPass();
-		}
-
 		if(ViewFamily.EngineShowFlags.DirectLighting)
 		{
 			SCOPED_DRAW_EVENT(RHICmdList, NonShadowedLights);
@@ -1100,7 +1099,7 @@ void FDeferredShadingSceneRenderer::RenderLights(FRHICommandListImmediate& RHICm
 					const FLightSceneInfo* const LightSceneInfo = SortedLightInfo.LightSceneInfo;
 
 					// Render the light to the scene color buffer, using a 1x1 white texture as input
-					RenderLight(RHICmdList, LightSceneInfo, nullptr, nullptr, nullptr, nullptr, false, false, false);
+					RenderLight(RHICmdList, LightSceneInfo, nullptr, nullptr, false, false);
 				}
 
 				SceneContext.FinishRenderingSceneColor(RHICmdList);
@@ -1124,8 +1123,8 @@ void FDeferredShadingSceneRenderer::RenderLights(FRHICommandListImmediate& RHICm
 					TransmittanceMaskData = RenderHairStrandsTransmittanceMask(RHICmdList, Views, LightSceneInfo, HairDatas, NullScreenShadowMaskSubPixelTexture);
 
 					// Render the light to the scene color buffer, using a 1x1 white texture as input
-					SceneContext.BeginRenderingSceneColor(RHICmdList, ESimpleRenderTargetMode::EExistingColorAndDepth, FExclusiveDepthStencil::DepthRead_StencilWrite, true, true);
-					RenderLight(RHICmdList, LightSceneInfo, nullptr, nullptr, &TransmittanceMaskData, InHairVisibilityViews, false, false, false);
+					SceneContext.BeginRenderingSceneColor(RHICmdList, ESimpleRenderTargetMode::EExistingColorAndDepth, FExclusiveDepthStencil::DepthRead_StencilWrite, true);
+					RenderLight(RHICmdList, LightSceneInfo, nullptr, InHairVisibilityViews, false, false);
 					SceneContext.FinishRenderingSceneColor(RHICmdList);
 				}
 
@@ -1841,10 +1840,7 @@ void FDeferredShadingSceneRenderer::RenderLights(FRHICommandListImmediate& RHICm
 				else
 				{
 					SCOPED_DRAW_EVENT(RHICmdList, StandardDeferredLighting);
-					SceneContext.BeginRenderingSceneColor(RHICmdList, ESimpleRenderTargetMode::EExistingColorAndDepth, FExclusiveDepthStencil::DepthRead_StencilWrite, true, true);
-
-					// When split lighting is enable, the hair pixel are rendered into a separate target
-					const bool bUseHairSplitLighting = HairDatas && Views.Num() > 0 && HairDatas->HairVisibilityViews.HairDatas[0].SampleLightingBuffer != nullptr;
+					SceneContext.BeginRenderingSceneColor(RHICmdList, ESimpleRenderTargetMode::EExistingColorAndDepth, FExclusiveDepthStencil::DepthRead_StencilWrite, true);
 
 					// ScreenShadowMaskTexture might have been created for a previous light, but only use it if we wrote valid data into it for this light
 					IPooledRenderTarget* LightShadowMaskTexture = nullptr;
@@ -1858,14 +1854,14 @@ void FDeferredShadingSceneRenderer::RenderLights(FRHICommandListImmediate& RHICm
 					// Render the light to the scene color buffer, conditionally using the attenuation buffer or a 1x1 white texture as input 
 					if (bDirectLighting)
 					{
-						RenderLight(RHICmdList, &LightSceneInfo, LightShadowMaskTexture, LightShadowMaskSubPixelTexture, &TransmittanceMaskData, InHairVisibilityViews, false, true, !bUseHairSplitLighting);
+						RenderLight(RHICmdList, &LightSceneInfo, LightShadowMaskTexture, InHairVisibilityViews, false, true);
 					}
 
 					SceneContext.FinishRenderingSceneColor(RHICmdList);
 
-					if (bDirectLighting && bUseHairSplitLighting)
+					if (bUseHairLighting)
 					{
-						RenderLightForHair(RHICmdList, &LightSceneInfo, LightShadowMaskSubPixelTexture, &TransmittanceMaskData, InHairVisibilityViews, true);
+						RenderLightForHair(RHICmdList, &LightSceneInfo, LightShadowMaskSubPixelTexture, &TransmittanceMaskData, InHairVisibilityViews);
 					}
 				}
 			}
@@ -1900,7 +1896,7 @@ void FDeferredShadingSceneRenderer::RenderLightArrayForOverlapViewmode(FRHIComma
 			&& !LightSceneInfo->Proxy->HasStaticLighting()
 			&& LightSceneInfo->Proxy->CastsStaticShadow())
 		{
-			RenderLight(RHICmdList, LightSceneInfo, nullptr, nullptr, nullptr, nullptr, true, false, false);
+			RenderLight(RHICmdList, LightSceneInfo, nullptr, nullptr, true, false);
 		}
 	}
 }
@@ -2015,7 +2011,8 @@ void CalculateLightNearFarDepthFromBounds(const FViewInfo& View, const FSphere &
  * @param LightIndex The light's index into FScene::Lights
  * @return true if anything got rendered
  */
-void FDeferredShadingSceneRenderer::RenderLight(FRHICommandList& RHICmdList, const FLightSceneInfo* LightSceneInfo, IPooledRenderTarget* ScreenShadowMaskTexture, IPooledRenderTarget* ScreenShadowMaskSubPixelTexture, FHairStrandsTransmittanceMaskData* InTransmittanceMaskData, const FHairStrandsVisibilityViews* InHairVisibilityViews, bool bRenderOverlap, bool bIssueDrawEvent, bool bSupportHair)
+
+void FDeferredShadingSceneRenderer::RenderLight(FRHICommandList& RHICmdList, const FLightSceneInfo* LightSceneInfo, IPooledRenderTarget* ScreenShadowMaskTexture, const FHairStrandsVisibilityViews* InHairVisibilityViews,  bool bRenderOverlap, bool bIssueDrawEvent)
 {
 	SCOPE_CYCLE_COUNTER(STAT_DirectLightRenderingTime);
 	INC_DWORD_STAT(STAT_NumLightsUsingStandardDeferred);
@@ -2024,19 +2021,7 @@ void FDeferredShadingSceneRenderer::RenderLight(FRHICommandList& RHICmdList, con
 	FGraphicsPipelineStateInitializer GraphicsPSOInit;
 	RHICmdList.ApplyCachedRenderTargets(GraphicsPSOInit);
 
-	const bool bHairAttunationEnabled = !bSupportHair && InTransmittanceMaskData && InHairVisibilityViews && (LightSceneInfo->Proxy->CastsHairStrandsDeepShadow() || IsHairStrandsVoxelizationEnable());
-	const bool bHairLightingEnabled = bSupportHair && InTransmittanceMaskData && InHairVisibilityViews && (LightSceneInfo->Proxy->CastsHairStrandsDeepShadow() || IsHairStrandsVoxelizationEnable());
-	if (bHairLightingEnabled)
-	{
-		GraphicsPSOInit.BlendState = TStaticBlendState<
-			CW_RGBA, BO_Add, BF_One, BF_One, BO_Add, BF_One, BF_One,
-			CW_RGBA, BO_Add, BF_One, BF_One, BO_Max, BF_SourceAlpha, BF_DestAlpha
-		>::GetRHI();
-	}
-	else
-	{
-		GraphicsPSOInit.BlendState = TStaticBlendState<CW_RGBA, BO_Add, BF_One, BF_One, BO_Add, BF_One, BF_One>::GetRHI();
-	}
+	GraphicsPSOInit.BlendState = TStaticBlendState<CW_RGBA, BO_Add, BF_One, BF_One, BO_Add, BF_One, BF_One>::GetRHI();
 
 	GraphicsPSOInit.PrimitiveType = PT_TriangleList;
 
@@ -2064,23 +2049,11 @@ void FDeferredShadingSceneRenderer::RenderLight(FRHICommandList& RHICmdList, con
 		RHICmdList.SetViewport(View.ViewRect.Min.X, View.ViewRect.Min.Y, 0.0f, View.ViewRect.Max.X, View.ViewRect.Max.Y, 1.0f);
 
 		FRenderLightParams RenderLightParams;
-		if (bHairLightingEnabled || bHairAttunationEnabled)
+		const bool bHairLighting = InHairVisibilityViews && ViewIndex < InHairVisibilityViews->HairDatas.Num() && InHairVisibilityViews->HairDatas[ViewIndex].CategorizationTexture != nullptr;
+		if (bHairLighting)
 		{
-			RenderLightParams.DeepShadow_TransmittanceMaskBuffer = InTransmittanceMaskData ? InTransmittanceMaskData->TransmittanceMaskSRV : nullptr;
-			RenderLightParams.DeepShadow_TransmittanceMaskBufferMaxCount = InTransmittanceMaskData && InTransmittanceMaskData->TransmittanceMask ? InTransmittanceMaskData->TransmittanceMask->Desc.NumElements : 0;
-			RenderLightParams.ScreenShadowMaskSubPixelTexture = ScreenShadowMaskSubPixelTexture;
-			if (InHairVisibilityViews && ViewIndex < InHairVisibilityViews->HairDatas.Num())
-			{
-				const FHairStrandsVisibilityData& HairVisibilityData = InHairVisibilityViews->HairDatas[ViewIndex];
-				RenderLightParams.HairVisibilityNodeOffsetAndCount = HairVisibilityData.NodeIndex;
-				RenderLightParams.HairVisibilityNodeDataSRV = HairVisibilityData.NodeDataSRV;
-				RenderLightParams.HairVisibilityNodeCoordsSRV = RenderLightParams.HairVisibilityNodeCoordsSRV;
-				RenderLightParams.HairCategorizationTexture = HairVisibilityData.CategorizationTexture;
-			}
+			RenderLightParams.HairCategorizationTexture = InHairVisibilityViews->HairDatas[ViewIndex].CategorizationTexture;
 		}
-
-		const uint32 HairPermutation = bHairAttunationEnabled ? 3 : (bHairLightingEnabled ? 1 : 0);
-
 		if (LightSceneInfo->Proxy->GetLightType() == LightType_Directional)
 		{
 			// Turn DBT back off
@@ -2110,7 +2083,7 @@ void FDeferredShadingSceneRenderer::RenderLight(FRHICommandList& RHICmdList, con
 				PermutationVector.Set< FDeferredLightPS::FVisualizeCullingDim >( View.Family->EngineShowFlags.VisualizeLightCulling );
 				PermutationVector.Set< FDeferredLightPS::FLightingChannelsDim >( View.bUsesLightingChannels );
 				PermutationVector.Set< FDeferredLightPS::FTransmissionDim >( bTransmission );
-				PermutationVector.Set< FDeferredLightPS::FHairLighting>(HairPermutation);
+				PermutationVector.Set< FDeferredLightPS::FHairLighting>(bHairLighting ? 1 : 0);
 				// Only directional lights are rendered in this path, so we only need to check if it is use to light the atmosphere
 				PermutationVector.Set< FDeferredLightPS::FAtmosphereTransmittance >(bAtmospherePerPixelTransmittance);
 
@@ -2120,7 +2093,7 @@ void FDeferredShadingSceneRenderer::RenderLight(FRHICommandList& RHICmdList, con
 				GraphicsPSOInit.BoundShaderState.PixelShaderRHI = PixelShader.GetPixelShader();
 
 				SetGraphicsPipelineState(RHICmdList, GraphicsPSOInit);
-				PixelShader->SetParameters(RHICmdList, View, LightSceneInfo, ScreenShadowMaskTexture, (bHairLightingEnabled || bHairAttunationEnabled) ? &RenderLightParams : nullptr);
+				PixelShader->SetParameters(RHICmdList, View, LightSceneInfo, ScreenShadowMaskTexture, (bHairLighting) ? &RenderLightParams : nullptr);
 			}
 
 			VertexShader->SetParameters(RHICmdList, View, LightSceneInfo);
@@ -2141,7 +2114,7 @@ void FDeferredShadingSceneRenderer::RenderLight(FRHICommandList& RHICmdList, con
 		{
 			// Use DBT to allow work culling on shadow lights
 			// Disable depth bound when hair rendering is enabled as this rejects partially covered pixel write (with opaque background)
-			GraphicsPSOInit.bDepthBounds = !bHairLightingEnabled && (GSupportsDepthBoundsTest && GAllowDepthBoundsTest != 0);
+			GraphicsPSOInit.bDepthBounds = GSupportsDepthBoundsTest && GAllowDepthBoundsTest != 0;
 
 			TShaderMapRef<TDeferredLightVS<true> > VertexShader(View.ShaderMap);
 
@@ -2167,7 +2140,7 @@ void FDeferredShadingSceneRenderer::RenderLight(FRHICommandList& RHICmdList, con
 				PermutationVector.Set< FDeferredLightPS::FVisualizeCullingDim >( View.Family->EngineShowFlags.VisualizeLightCulling );
 				PermutationVector.Set< FDeferredLightPS::FLightingChannelsDim >( View.bUsesLightingChannels );
 				PermutationVector.Set< FDeferredLightPS::FTransmissionDim >( bTransmission );
-				PermutationVector.Set< FDeferredLightPS::FHairLighting>(HairPermutation);
+				PermutationVector.Set< FDeferredLightPS::FHairLighting>(bHairLighting ? 1 : 0);
 				PermutationVector.Set < FDeferredLightPS::FAtmosphereTransmittance >(false);
 
 				TShaderMapRef< FDeferredLightPS > PixelShader( View.ShaderMap, PermutationVector );
@@ -2176,7 +2149,7 @@ void FDeferredShadingSceneRenderer::RenderLight(FRHICommandList& RHICmdList, con
 				GraphicsPSOInit.BoundShaderState.PixelShaderRHI = PixelShader.GetPixelShader();
 
 				SetGraphicsPipelineState(RHICmdList, GraphicsPSOInit);
-				PixelShader->SetParameters(RHICmdList, View, LightSceneInfo, ScreenShadowMaskTexture, &RenderLightParams);
+				PixelShader->SetParameters(RHICmdList, View, LightSceneInfo, ScreenShadowMaskTexture, (bHairLighting) ? &RenderLightParams : nullptr);
 			}
 
 			VertexShader->SetParameters(RHICmdList, View, LightSceneInfo);
@@ -2219,12 +2192,11 @@ void FDeferredShadingSceneRenderer::RenderLightForHair(
 	const FLightSceneInfo* LightSceneInfo, 
 	IPooledRenderTarget* HairShadowMaskTexture, 
 	FHairStrandsTransmittanceMaskData* InTransmittanceMaskData, 
-	const FHairStrandsVisibilityViews* InHairVisibilityViews, 
-	bool bIssueDrawEvent)
+	const FHairStrandsVisibilityViews* InHairVisibilityViews)
 {
 	SCOPE_CYCLE_COUNTER(STAT_DirectLightRenderingTime);
 	INC_DWORD_STAT(STAT_NumLightsUsingStandardDeferred);
-	SCOPED_CONDITIONAL_DRAW_EVENT(RHICmdList, StandardDeferredLighting, bIssueDrawEvent);
+	SCOPED_CONDITIONAL_DRAW_EVENT(RHICmdList, StandardDeferredLighting, true);
 	
 	FGraphicsPipelineStateInitializer GraphicsPSOInit;
 	RHICmdList.ApplyCachedRenderTargets(GraphicsPSOInit);
@@ -2255,7 +2227,11 @@ void FDeferredShadingSceneRenderer::RenderLightForHair(
 		check(RHICmdList.IsOutsideRenderPass());
 
 		const FHairStrandsVisibilityData& HairVisibilityData = InHairVisibilityViews->HairDatas[ViewIndex];
-		
+		if (!HairVisibilityData.SampleLightingBuffer)
+		{
+			continue;
+		}
+
 		FRenderLightParams RenderLightParams;
 		RenderLightParams.DeepShadow_TransmittanceMaskBuffer = InTransmittanceMaskData ? InTransmittanceMaskData->TransmittanceMaskSRV : nullptr;
 		RenderLightParams.DeepShadow_TransmittanceMaskBufferMaxCount = InTransmittanceMaskData && InTransmittanceMaskData->TransmittanceMask ? InTransmittanceMaskData->TransmittanceMask->Desc.NumElements : 0;
@@ -2273,12 +2249,9 @@ void FDeferredShadingSceneRenderer::RenderLightForHair(
 		if (LightSceneInfo->Proxy->GetLightType() == LightType_Directional)
 		{
 			PermutationVector.Set< FDeferredLightPS::FSourceShapeDim >(ELightSourceShape::Directional);
+			PermutationVector.Set< FDeferredLightPS::FSourceTextureDim >(false);
 			PermutationVector.Set< FDeferredLightPS::FIESProfileDim >(false);
 			PermutationVector.Set< FDeferredLightPS::FInverseSquaredDim >(false);
-			PermutationVector.Set< FDeferredLightPS::FVisualizeCullingDim >(View.Family->EngineShowFlags.VisualizeLightCulling);
-			PermutationVector.Set< FDeferredLightPS::FLightingChannelsDim >(View.bUsesLightingChannels);
-			PermutationVector.Set< FDeferredLightPS::FTransmissionDim >(false);
-			PermutationVector.Set< FDeferredLightPS::FHairLighting>(2);
 		}
 		else
 		{
@@ -2287,11 +2260,11 @@ void FDeferredShadingSceneRenderer::RenderLightForHair(
 			PermutationVector.Set< FDeferredLightPS::FSourceTextureDim >(LightSceneInfo->Proxy->IsRectLight() && LightSceneInfo->Proxy->HasSourceTexture());
 			PermutationVector.Set< FDeferredLightPS::FIESProfileDim >(bUseIESTexture);
 			PermutationVector.Set< FDeferredLightPS::FInverseSquaredDim >(LightSceneInfo->Proxy->IsInverseSquared());
-			PermutationVector.Set< FDeferredLightPS::FVisualizeCullingDim >(View.Family->EngineShowFlags.VisualizeLightCulling);
-			PermutationVector.Set< FDeferredLightPS::FLightingChannelsDim >(View.bUsesLightingChannels);
-			PermutationVector.Set< FDeferredLightPS::FTransmissionDim >(false);
-			PermutationVector.Set< FDeferredLightPS::FHairLighting>(2);
 		}
+		PermutationVector.Set< FDeferredLightPS::FLightingChannelsDim >(View.bUsesLightingChannels);
+		PermutationVector.Set< FDeferredLightPS::FVisualizeCullingDim >(false);
+		PermutationVector.Set< FDeferredLightPS::FTransmissionDim >(false);
+		PermutationVector.Set< FDeferredLightPS::FHairLighting>(2);
 
 		TShaderMapRef<TDeferredLightHairVS> VertexShader(View.ShaderMap);
 		TShaderMapRef<FDeferredLightPS> PixelShader(View.ShaderMap, PermutationVector);
diff --git a/Engine/Source/Runtime/Renderer/Private/PostProcess/SceneRenderTargets.cpp b/Engine/Source/Runtime/Renderer/Private/PostProcess/SceneRenderTargets.cpp
--- a/Engine/Source/Runtime/Renderer/Private/PostProcess/SceneRenderTargets.cpp
+++ b/Engine/Source/Runtime/Renderer/Private/PostProcess/SceneRenderTargets.cpp
@@ -252,7 +252,6 @@ static void SnapshotArray(TArray<TRefCountPtr<IPooledRenderTarget>, TInlineAlloc
 FSceneRenderTargets::FSceneRenderTargets(const FViewInfo& View, const FSceneRenderTargets& SnapshotSource)
 	: LightAttenuation(GRenderTargetPool.MakeSnapshot(SnapshotSource.LightAttenuation))
 	, LightAccumulation(GRenderTargetPool.MakeSnapshot(SnapshotSource.LightAccumulation))
-	, SceneColorSubPixel(GRenderTargetPool.MakeSnapshot(SnapshotSource.SceneColorSubPixel))
 	, DirectionalOcclusion(GRenderTargetPool.MakeSnapshot(SnapshotSource.DirectionalOcclusion))
 	, SceneDepthZ(GRenderTargetPool.MakeSnapshot(SnapshotSource.SceneDepthZ))	
 	, SceneVelocity(GRenderTargetPool.MakeSnapshot(SnapshotSource.SceneVelocity))
@@ -702,7 +701,7 @@ void FSceneRenderTargets::Allocate(FRHICommandListImmediate& RHICmdList, const F
 	AllocateRenderTargets(RHICmdList, ViewFamily.Views.Num());
 }
 
-void FSceneRenderTargets::BeginRenderingSceneColor(FRHICommandList& RHICmdList, ESimpleRenderTargetMode RenderTargetMode/*=EUninitializedColorExistingDepth*/, FExclusiveDepthStencil DepthStencilAccess, bool bTransitionWritable, bool bBindSubPixelRT)
+void FSceneRenderTargets::BeginRenderingSceneColor(FRHICommandList& RHICmdList, ESimpleRenderTargetMode RenderTargetMode/*=EUninitializedColorExistingDepth*/, FExclusiveDepthStencil DepthStencilAccess, bool bTransitionWritable)
 {
 	check(RHICmdList.IsOutsideRenderPass());
 
@@ -718,12 +717,6 @@ void FSceneRenderTargets::BeginRenderingSceneColor(FRHICommandList& RHICmdList,
 	RPInfo.DepthStencilRenderTarget.DepthStencilTarget = GetSceneDepthSurface();
 	RPInfo.DepthStencilRenderTarget.ExclusiveDepthStencil = DepthStencilAccess;
 
-	if (bBindSubPixelRT)
-	{
-		RPInfo.ColorRenderTargets[1].Action = MakeRenderTargetActions(ERenderTargetLoadAction::ELoad, ERenderTargetStoreAction::EStore);
-		RPInfo.ColorRenderTargets[1].RenderTarget = SceneColorSubPixel->GetRenderTargetItem().TargetableTexture;
-	}
-
 	if (bTransitionWritable)
 	{
 		TransitionRenderPassTargets(RHICmdList, RPInfo);
@@ -1391,17 +1384,6 @@ TRefCountPtr<IPooledRenderTarget>& FSceneRenderTargets::GetSceneColor()
 	return GetSceneColorForCurrentShadingPath();
 }
 
-TRefCountPtr<IPooledRenderTarget>& FSceneRenderTargets::GetSceneColorSubPixel()
-{
-	if (!GetSceneColorForCurrentShadingPath())
-	{
-		return GSystemTextures.BlackDummy;
-	}
-
-	return SceneColorSubPixel;
-}
-
-
 void FSceneRenderTargets::SetSceneColor(IPooledRenderTarget* In)
 {
 	check(CurrentShadingPath < EShadingPath::Num);
@@ -2691,12 +2673,6 @@ void FSceneRenderTargets::AllocateDeferredShadingPathRenderTargets(FRHICommandLi
 		GRenderTargetPool.FindFreeElement(RHICmdList, Desc, LightAccumulation, TEXT("LightAccumulation"), true, ERenderTargetTransience::NonTransient);
 	}
 
-	if (CurrentFeatureLevel >= ERHIFeatureLevel::SM5)
-	{
-		FPooledRenderTargetDesc Desc(FPooledRenderTargetDesc::Create2DDesc(BufferSize, PF_FloatRGBA,  FClearValueBinding::Transparent, TexCreate_None, TexCreate_RenderTargetable | TexCreate_ShaderResource, false));
-		GRenderTargetPool.FindFreeElement(RHICmdList, Desc, SceneColorSubPixel, TEXT("SceneColorSubPixel"), true);
-	}
-
 	AllocateDebugViewModeTargets(RHICmdList);
 
 	if (bAllocateVelocityGBuffer)
@@ -2942,7 +2918,6 @@ void FSceneRenderTargets::ReleaseAllTargets()
 	QuadOverdrawBuffer.SafeRelease();
 	LightAttenuation.SafeRelease();
 	LightAccumulation.SafeRelease();
-	SceneColorSubPixel.SafeRelease();
 	DirectionalOcclusion.SafeRelease();
 	CustomDepth.SafeRelease();
 	MobileCustomDepth.SafeRelease();
diff --git a/Engine/Source/Runtime/Renderer/Private/PostProcess/SceneRenderTargets.h b/Engine/Source/Runtime/Renderer/Private/PostProcess/SceneRenderTargets.h
--- a/Engine/Source/Runtime/Renderer/Private/PostProcess/SceneRenderTargets.h
+++ b/Engine/Source/Runtime/Renderer/Private/PostProcess/SceneRenderTargets.h
@@ -238,7 +238,7 @@ public:
 	/**
 	 * Sets the scene color target and restores its contents if necessary
 	 */
-	void BeginRenderingSceneColor(FRHICommandList& FRHICommandListImmediate, ESimpleRenderTargetMode RenderTargetMode = ESimpleRenderTargetMode::EUninitializedColorExistingDepth, FExclusiveDepthStencil DepthStencilAccess = FExclusiveDepthStencil::DepthWrite_StencilWrite, bool bTransitionWritable = true, bool bBindSubPixelRT = false);
+	void BeginRenderingSceneColor(FRHICommandList& FRHICommandListImmediate, ESimpleRenderTargetMode RenderTargetMode = ESimpleRenderTargetMode::EUninitializedColorExistingDepth, FExclusiveDepthStencil DepthStencilAccess = FExclusiveDepthStencil::DepthWrite_StencilWrite, bool bTransitionWritable = true);
 	void FinishRenderingSceneColor(FRHICommandList& RHICmdList);
 
 	// @return true: call FinishRenderingCustomDepth after rendering, false: don't render it, feature is disabled
@@ -472,7 +472,6 @@ public:
 	const TRefCountPtr<IPooledRenderTarget>& GetSceneColor() const;
 
 	TRefCountPtr<IPooledRenderTarget>& GetSceneColor();
-	TRefCountPtr<IPooledRenderTarget>& GetSceneColorSubPixel();
 
 	EPixelFormat GetSceneColorFormat(ERHIFeatureLevel::Type InFeatureLevel) const;
 	EPixelFormat GetSceneColorFormat() const;
@@ -542,9 +541,6 @@ public:
 	// Light Accumulation is a high precision scratch pad matching the size of the scene color buffer used by many passes.
 	TRefCountPtr<IPooledRenderTarget> LightAccumulation;
 
-	// Secondary scene color target for storing subpixel color details (e.g. hair pixels)
-	TRefCountPtr<IPooledRenderTarget> SceneColorSubPixel;
-
 	// Reflection Environment: Bringing back light accumulation buffer to apply indirect reflections
 	TRefCountPtr<IPooledRenderTarget> DirectionalOcclusion;
 	
