Fix directional shadow with hair.

diff --git a/Engine/Shaders/Private/ShadowProjectionPixelShader.usf b/Engine/Shaders/Private/ShadowProjectionPixelShader.usf
--- a/Engine/Shaders/Private/ShadowProjectionPixelShader.usf
+++ b/Engine/Shaders/Private/ShadowProjectionPixelShader.usf
@@ -89,6 +89,7 @@ float3 ModulatedShadowBlendOp(float3 Source)
 
 #if SUBPIXEL_SHADOW
 #include "HairStrands/HairStrandsVisibilityCommon.ush"
+float2 ShadowNearAndFarDepth;
 Texture2D<uint4> HairCategorizationTexture;
 #endif
 
@@ -123,6 +124,11 @@ void Main(
 	{
 		discard;
 	}
+
+	if (CategorizationData.ClosestDepth > ShadowNearAndFarDepth.x || CategorizationData.ClosestDepth < ShadowNearAndFarDepth.y) // Inverse-Z
+	{
+		discard;
+	}
 #endif
 
 	float4 ScreenPosition = float4(((ScreenUV.xy - View.ScreenPositionScaleBias.wz ) / View.ScreenPositionScaleBias.xy) * SceneW, SceneW, 1);
diff --git a/Engine/Source/Runtime/Renderer/Private/ShadowRendering.h b/Engine/Source/Runtime/Renderer/Private/ShadowRendering.h
--- a/Engine/Source/Runtime/Renderer/Private/ShadowRendering.h
+++ b/Engine/Source/Runtime/Renderer/Private/ShadowRendering.h
@@ -835,6 +835,7 @@ public:
 		HairCategorizationTexture.Bind(ParameterMap, TEXT("HairCategorizationTexture"));
 		PerObjectShadowFadeStart.Bind(ParameterMap, TEXT("PerObjectShadowFadeStart"));
 		InvPerObjectShadowFadeLength.Bind(ParameterMap, TEXT("InvPerObjectShadowFadeLength"));
+		ShadowNearAndFarDepth.Bind(ParameterMap, TEXT("ShadowNearAndFarDepth"));
 	}
 
 	void Set(FRHICommandList& RHICmdList, FShader* Shader, const FSceneView& View, const FProjectedShadowInfo* ShadowInfo, const FHairStrandsVisibilityData* HairVisibilityData, bool bModulatedShadows)
@@ -929,6 +930,16 @@ public:
 		if (HairVisibilityData && HairVisibilityData->CategorizationTexture)
 		{
 			SetTextureParameter(RHICmdList, ShaderRHI, HairCategorizationTexture, HairVisibilityData->CategorizationTexture->GetRenderTargetItem().ShaderResourceTexture);
+
+			FVector4 Near = View.ViewMatrices.GetProjectionMatrix().TransformFVector4(FVector4(0, 0, ShadowInfo->CascadeSettings.SplitNear));
+			FVector4 Far = View.ViewMatrices.GetProjectionMatrix().TransformFVector4(FVector4(0, 0, ShadowInfo->CascadeSettings.SplitFar));
+			const float DeviceZNear = Near.Z / Near.W;
+			const float DeviceZFar = Far.Z / Far.W;
+
+			FVector2D SliceNearAndFarDepth;
+			SliceNearAndFarDepth.X = DeviceZNear;
+			SliceNearAndFarDepth.Y = DeviceZFar;
+			SetShaderValue(RHICmdList, ShaderRHI, ShadowNearAndFarDepth, SliceNearAndFarDepth);
 		}
 
 		SetShaderValue(RHICmdList, ShaderRHI, PerObjectShadowFadeStart, ShadowInfo->PerObjectShadowFadeStart);
@@ -952,6 +963,7 @@ public:
 		Ar << P.HairCategorizationTexture;
 		Ar << P.PerObjectShadowFadeStart;
 		Ar << P.InvPerObjectShadowFadeLength;
+		Ar << P.ShadowNearAndFarDepth;
 		return Ar;
 	}*/
 
@@ -970,6 +982,7 @@ private:
 	LAYOUT_FIELD(FShaderResourceParameter, HairCategorizationTexture);
 	LAYOUT_FIELD(FShaderParameter, PerObjectShadowFadeStart);
 	LAYOUT_FIELD(FShaderParameter, InvPerObjectShadowFadeLength);
+	LAYOUT_FIELD(FShaderParameter, ShadowNearAndFarDepth);
 };
 
 /**
