Add indirect clear pass for light sample.

diff --git a/Engine/Shaders/Private/HairStrands/HairStrandsLightSample.usf b/Engine/Shaders/Private/HairStrands/HairStrandsLightSample.usf
new file mode 100644
--- /dev/null
+++ b/Engine/Shaders/Private/HairStrands/HairStrandsLightSample.usf
@@ -0,0 +1,46 @@
+// Copyright Epic Games, Inc. All Rights Reserved.
+
+#include "../Common.ush"
+
+///////////////////////////////////////////////////////////////////////////
+
+#if SHADER_VERTEX
+
+int2 MaxViewportResolution;
+Texture2D<uint> HairNodeCountTexture;
+
+void MainVS(
+	in uint InVertexId : SV_VertexID,
+	out float4 OutPosition : SV_POSITION,
+	nointerpolation out uint OutNodeCount : DISPATCH_NODECOUNT,
+	nointerpolation out uint2 OutResolution : DISPATCH_RESOLUTION)
+{
+	OutNodeCount = HairNodeCountTexture.Load(uint3(0, 0, 0));
+	OutResolution.x = ceil(sqrt(OutNodeCount));
+	OutResolution.y = OutResolution.x;
+
+	const float2 ClipCoord = ((OutResolution + 0.5f) / float2(MaxViewportResolution)) * 2;
+
+	const float2 UV = float2(InVertexId & 1, InVertexId >> 1);
+	const float2 Pos = float2(-UV.x, UV.y) * 4 + float2(-1, +1) + float2(ClipCoord.x, -ClipCoord.y);
+	OutPosition = float4(Pos, 0.5f, 1);
+}
+
+#endif
+
+
+#if SHADER_CLEAR
+
+void ClearPS(
+	float4 SVPos : SV_POSITION,
+	uint NodeCount : DISPATCH_NODECOUNT,
+	uint2 Resolution : DISPATCH_RESOLUTION,
+	out float4 OutColor : SV_Target0)
+{
+	const uint2 PixelCoord = uint2(SVPos.xy);
+	if (any(PixelCoord > Resolution))
+		discard;
+
+	OutColor = 0;
+}
+#endif
\ No newline at end of file
diff --git a/Engine/Source/Runtime/Renderer/Private/HairStrands/HairStrandsVisibility.cpp b/Engine/Source/Runtime/Renderer/Private/HairStrands/HairStrandsVisibility.cpp
--- a/Engine/Source/Runtime/Renderer/Private/HairStrands/HairStrandsVisibility.cpp
+++ b/Engine/Source/Runtime/Renderer/Private/HairStrands/HairStrandsVisibility.cpp
@@ -169,6 +169,108 @@ static bool IsCompatibleWithHairVisibility(const FMeshMaterialShaderPermutationP
 	return IsCompatibleWithHairStrands(Parameters.Platform, Parameters.MaterialParameters);
 }
 
+
+///////////////////////////////////////////////////////////////////////////////////////////////////
+
+class FHairLightSampleClearVS : public FGlobalShader
+{
+	DECLARE_GLOBAL_SHADER(FHairLightSampleClearVS);
+	SHADER_USE_PARAMETER_STRUCT(FHairLightSampleClearVS, FGlobalShader)
+
+	BEGIN_SHADER_PARAMETER_STRUCT(FParameters, )
+		SHADER_PARAMETER(FIntPoint, MaxViewportResolution)
+		SHADER_PARAMETER_RDG_TEXTURE(Texture2D, HairNodeCountTexture)
+	END_SHADER_PARAMETER_STRUCT()
+
+	static bool ShouldCompilePermutation(const FGlobalShaderPermutationParameters& Parameters) { return Parameters.Platform == SP_PCD3D_SM5; }
+	static void ModifyCompilationEnvironment(const FGlobalShaderPermutationParameters& Parameters, FShaderCompilerEnvironment& OutEnvironment)
+	{
+		FGlobalShader::ModifyCompilationEnvironment(Parameters, OutEnvironment);
+		OutEnvironment.SetDefine(TEXT("SHADER_VERTEX"), 1);
+	}
+};
+
+class FHairLightSampleClearPS : public FGlobalShader
+{
+	DECLARE_GLOBAL_SHADER(FHairLightSampleClearPS);
+	SHADER_USE_PARAMETER_STRUCT(FHairLightSampleClearPS, FGlobalShader)
+
+	BEGIN_SHADER_PARAMETER_STRUCT(FParameters, )
+		SHADER_PARAMETER(FIntPoint, MaxViewportResolution)
+		SHADER_PARAMETER_RDG_TEXTURE(Texture2D, HairNodeCountTexture)
+		RENDER_TARGET_BINDING_SLOTS()
+	END_SHADER_PARAMETER_STRUCT()
+
+	static bool ShouldCompilePermutation(const FGlobalShaderPermutationParameters& Parameters) { return Parameters.Platform == SP_PCD3D_SM5; }
+	static void ModifyCompilationEnvironment(const FGlobalShaderPermutationParameters& Parameters, FShaderCompilerEnvironment& OutEnvironment)
+	{
+		FGlobalShader::ModifyCompilationEnvironment(Parameters, OutEnvironment);
+		OutEnvironment.SetDefine(TEXT("SHADER_CLEAR"), 1);
+	}
+};
+
+IMPLEMENT_GLOBAL_SHADER(FHairLightSampleClearVS, "/Engine/Private/HairStrands/HairStrandsLightSample.usf", "MainVS", SF_Vertex);
+IMPLEMENT_GLOBAL_SHADER(FHairLightSampleClearPS, "/Engine/Private/HairStrands/HairStrandsLightSample.usf", "ClearPS", SF_Pixel);
+
+static FRDGTextureRef AddClearLightSamplePass(
+	FRDGBuilder& GraphBuilder,
+	const FViewInfo* View,
+	const uint32 MaxNodeCount,
+	const FRDGTextureRef NodeCounter)
+{	
+	const uint32 SampleTextureResolution = FMath::CeilToInt(FMath::Sqrt(MaxNodeCount));
+	FRDGTextureDesc Desc;
+	Desc.Extent.X = SampleTextureResolution;
+	Desc.Extent.Y = SampleTextureResolution;
+	Desc.Depth = 0;
+	Desc.Format = PF_FloatRGBA;
+	Desc.NumMips = 1;
+	Desc.Flags = 0;
+	Desc.TargetableFlags = TexCreate_UAV | TexCreate_ShaderResource | TexCreate_RenderTargetable;
+	FRDGTextureRef Output = GraphBuilder.CreateTexture(Desc, TEXT("HairLightSample"));
+
+	FHairLightSampleClearPS::FParameters* ParametersPS = GraphBuilder.AllocParameters<FHairLightSampleClearPS::FParameters>();
+	ParametersPS->MaxViewportResolution = Desc.Extent;
+	ParametersPS->HairNodeCountTexture = NodeCounter;
+	
+	const FIntPoint ViewportResolution = Desc.Extent;
+	TShaderMapRef<FHairLightSampleClearVS> VertexShader(View->ShaderMap);
+	TShaderMapRef<FHairLightSampleClearPS> PixelShader(View->ShaderMap);
+
+	ParametersPS->RenderTargets[0] = FRenderTargetBinding(Output, ERenderTargetLoadAction::ENoAction);
+
+	GraphBuilder.AddPass(
+		RDG_EVENT_NAME("HairLightSampleClearPS"),
+		ParametersPS,
+		ERDGPassFlags::Raster,
+		[ParametersPS, VertexShader, PixelShader, ViewportResolution](FRHICommandList& RHICmdList)
+	{
+		FHairLightSampleClearVS::FParameters ParametersVS;
+		ParametersVS.MaxViewportResolution = ParametersPS->MaxViewportResolution;
+		ParametersVS.HairNodeCountTexture = ParametersPS->HairNodeCountTexture;
+
+		FGraphicsPipelineStateInitializer GraphicsPSOInit;
+		RHICmdList.ApplyCachedRenderTargets(GraphicsPSOInit);
+		GraphicsPSOInit.BlendState = TStaticBlendState<CW_RGBA, BO_Add, BF_One, BF_Zero, BO_Add, BF_One, BF_Zero>::GetRHI();
+		GraphicsPSOInit.RasterizerState = TStaticRasterizerState<>::GetRHI();
+		GraphicsPSOInit.DepthStencilState = TStaticDepthStencilState<false, CF_Always>::GetRHI();
+		GraphicsPSOInit.BoundShaderState.VertexDeclarationRHI = GFilterVertexDeclaration.VertexDeclarationRHI;
+		GraphicsPSOInit.BoundShaderState.VertexShaderRHI = VertexShader.GetVertexShader();
+		GraphicsPSOInit.BoundShaderState.PixelShaderRHI = PixelShader.GetPixelShader();
+		GraphicsPSOInit.PrimitiveType = PT_TriangleList;
+		SetGraphicsPipelineState(RHICmdList, GraphicsPSOInit);
+
+		SetShaderParameters(RHICmdList, VertexShader, VertexShader.GetVertexShader(), ParametersVS);
+		SetShaderParameters(RHICmdList, PixelShader, PixelShader.GetPixelShader(), *ParametersPS);
+
+		RHICmdList.SetViewport(0, 0, 0.0f, ViewportResolution.X, ViewportResolution.Y, 1.0f);
+		RHICmdList.SetStreamSource(0, nullptr, 0);
+		RHICmdList.DrawPrimitive(0, 1, 1);
+	});
+
+	return Output;
+}
+
 /////////////////////////////////////////////////////////////////////////////////////////
 
 BEGIN_GLOBAL_SHADER_PARAMETER_STRUCT(FMaterialPassParameters, )
@@ -2470,6 +2572,11 @@ FHairStrandsVisibilityViews RenderHairStrandsVisibilityBuffer(
 						CompactNodeData = PassOutput.NodeData;
 					}
 
+					// Allocate buffer for storing all the light samples
+					FRDGTextureRef SampleLightingBuffer = AddClearLightSamplePass(GraphBuilder, &View, VisibilityData.MaxNodeCount, NodeCounter);
+					VisibilityData.SampleLightingViewportResolution = SampleLightingBuffer->Desc.Extent;
+
+					GraphBuilder.QueueTextureExtraction(SampleLightingBuffer, 	&VisibilityData.SampleLightingBuffer);
 					GraphBuilder.QueueTextureExtraction(CompactNodeIndex,		&VisibilityData.NodeIndex);
 					GraphBuilder.QueueTextureExtraction(CategorizationTexture,	&VisibilityData.CategorizationTexture);
 					GraphBuilder.QueueBufferExtraction(CompactNodeData,			&VisibilityData.NodeData,					FRDGResourceState::EAccess::Read, FRDGResourceState::EPipeline::Graphics);
@@ -2571,6 +2678,11 @@ FHairStrandsVisibilityViews RenderHairStrandsVisibilityBuffer(
 					SceneColorTexture,
 					SceneDepthTexture);
 
+				// Allocate buffer for storing all the light samples
+				FRDGTextureRef SampleLightingBuffer = AddClearLightSamplePass(GraphBuilder, &View, VisibilityData.MaxNodeCount, NodeCounter);
+				VisibilityData.SampleLightingViewportResolution = SampleLightingBuffer->Desc.Extent;
+				GraphBuilder.QueueTextureExtraction(SampleLightingBuffer, &VisibilityData.SampleLightingBuffer);
+
 #if WITH_EDITOR
 				// Extract texture for debug visualization
 				GraphBuilder.QueueTextureExtraction(PPLLNodeCounterTexture, &VisibilityData.PPLLNodeCounterTexture);
@@ -2606,25 +2718,6 @@ FHairStrandsVisibilityViews RenderHairStrandsVisibilityBuffer(
 
 			GraphBuilder.Execute();
 
-			// Allocate buffer for storing all the light samples
-			{			
-				const uint32 SampleTextureResolution = FMath::CeilToInt(FMath::Sqrt(VisibilityData.MaxNodeCount));
-				VisibilityData.SampleLightingViewportResolution = FIntPoint(SampleTextureResolution, SampleTextureResolution);
-				FPooledRenderTargetDesc Desc = FPooledRenderTargetDesc::Create2DDesc(
-					VisibilityData.SampleLightingViewportResolution,
-					PF_FloatRGBA,
-					FClearValueBinding::Black,
-					TexCreate_None, 
-					TexCreate_UAV | TexCreate_ShaderResource | TexCreate_RenderTargetable, false);
-
-				GRenderTargetPool.FindFreeElement(RHICmdList, Desc, VisibilityData.SampleLightingBuffer, TEXT("HairSampleLightingBuffer"));
-
-				// TODO: do an indirect clear of the VisibilityData.SampleLightingBuffer ...
-				FRHIRenderPassInfo RPInfo(VisibilityData.SampleLightingBuffer->GetRenderTargetItem().TargetableTexture, MakeRenderTargetActions(ERenderTargetLoadAction::EClear, ERenderTargetStoreAction::EStore));
-				RHICmdList.BeginRenderPass(RPInfo, TEXT("HairSampleLightingClear"));
-				RHICmdList.EndRenderPass();
-			}
-
 			// #hair_todo: is there a better way to get SRV view of a RDG buffer? should work as long as there is not reuse between the pass
 			if (VisibilityData.NodeData)
 			{
