Sorted reflection material gather rays now also use RAY_FLAG_FORCE_OPAQUE.

Also added few compile-time hints for ray traced reflection shaders to signal 1D dispatch and trivial CHS and miss shader where appropriate.

diff --git a/Engine/Shaders/Private/RayTracing/RayTracingReflections.usf b/Engine/Shaders/Private/RayTracing/RayTracingReflections.usf
--- a/Engine/Shaders/Private/RayTracing/RayTracingReflections.usf
+++ b/Engine/Shaders/Private/RayTracing/RayTracingReflections.usf
@@ -449,7 +449,9 @@ RAY_TRACING_ENTRY_RAYGEN(RayTracingReflectionsRGS)
 				{
 					if (DIM_DEFERRED_MATERIAL_MODE == DEFERRED_MATERIAL_MODE_GATHER)
 					{
-						uint RayFlags = 0;
+						// Material gather rays are traced as opaque to find the closest hit.
+						// Shading pass then traces a shortened ray from position right before first opaque hit to handle masked materials, translucency, etc.
+						uint RayFlags = RAY_FLAG_FORCE_OPAQUE;
 						DeferredMaterialPayload.SortKey = RAY_TRACING_DEFERRED_MATERIAL_KEY_RAY_MISS;
 						TraceRay(
 							TLAS,								// AccelerationStructure
diff --git a/Engine/Source/Runtime/Renderer/Private/RayTracing/RayTracingDeferredMaterials.cpp b/Engine/Source/Runtime/Renderer/Private/RayTracing/RayTracingDeferredMaterials.cpp
--- a/Engine/Source/Runtime/Renderer/Private/RayTracing/RayTracingDeferredMaterials.cpp
+++ b/Engine/Source/Runtime/Renderer/Private/RayTracing/RayTracingDeferredMaterials.cpp
@@ -16,11 +16,17 @@ class FRayTracingDeferredMaterialCHS : public FGlobalShader
 	DECLARE_GLOBAL_SHADER(FRayTracingDeferredMaterialCHS)
 	SHADER_USE_ROOT_PARAMETER_STRUCT(FRayTracingDeferredMaterialCHS, FGlobalShader)
 
-		static bool ShouldCompilePermutation(const FGlobalShaderPermutationParameters& Parameters)
+	static bool ShouldCompilePermutation(const FGlobalShaderPermutationParameters& Parameters)
 	{
 		return ShouldCompileRayTracingShadersForProject(Parameters.Platform);
 	}
 
+	static void ModifyCompilationEnvironment(const FGlobalShaderPermutationParameters& Parameters, FShaderCompilerEnvironment& OutEnvironment)
+	{
+		FGlobalShader::ModifyCompilationEnvironment(Parameters, OutEnvironment);
+		OutEnvironment.SetDefine(TEXT("UE_RAY_TRACING_LIGHTWEIGHT_CLOSEST_HIT_SHADER"), 1);
+	}
+
 	using FParameters = FEmptyShaderParameters;
 };
 
diff --git a/Engine/Source/Runtime/Renderer/Private/RayTracing/RayTracingReflections.cpp b/Engine/Source/Runtime/Renderer/Private/RayTracing/RayTracingReflections.cpp
--- a/Engine/Source/Runtime/Renderer/Private/RayTracing/RayTracingReflections.cpp
+++ b/Engine/Source/Runtime/Renderer/Private/RayTracing/RayTracingReflections.cpp
@@ -283,6 +283,25 @@ class FRayTracingReflectionsRGS : public FGlobalShader
 	{
 		return ShouldCompileRayTracingShadersForProject(Parameters.Platform);
 	}
+
+	static void ModifyCompilationEnvironment(const FGlobalShaderPermutationParameters& Parameters, FShaderCompilerEnvironment& OutEnvironment)
+	{
+		FGlobalShader::ModifyCompilationEnvironment(Parameters, OutEnvironment);
+
+		FPermutationDomain PermutationVector(Parameters.PermutationId);
+
+		OutEnvironment.SetDefine(TEXT("UE_RAY_TRACING_DYNAMIC_MISS_SHADER"), PermutationVector.Get<FMissShaderLighting>());
+
+		if (PermutationVector.Get<FDeferredMaterialMode>() == EDeferredMaterialMode::Shade)
+		{
+			OutEnvironment.SetDefine(TEXT("UE_RAY_TRACING_DISPATCH_1D"), 1);
+		}
+
+		if (PermutationVector.Get<FDeferredMaterialMode>() == EDeferredMaterialMode::Gather)
+		{
+			OutEnvironment.SetDefine(TEXT("UE_RAY_TRACING_LIGHTWEIGHT_CLOSEST_HIT_SHADER"), 1);
+		}
+	}
 };
 
 IMPLEMENT_GLOBAL_SHADER(FRayTracingReflectionsRGS, "/Engine/Private/RayTracing/RayTracingReflections.usf", "RayTracingReflectionsRGS", SF_RayGen);
