Fix hair lighting.

diff --git a/Engine/Shaders/Private/DeferredLightVertexShaders.usf b/Engine/Shaders/Private/DeferredLightVertexShaders.usf
--- a/Engine/Shaders/Private/DeferredLightVertexShaders.usf
+++ b/Engine/Shaders/Private/DeferredLightVertexShaders.usf
@@ -6,12 +6,14 @@
 
 #include "Common.ush"
 
+#if defined(SHADER_RADIAL_LIGHT)
 float4 StencilingGeometryPosAndScale;
 float4 StencilingConeParameters;	// .x NumSides (0 if not cone), .y NumSlices, .z ConeAngle, .w ConeSphereRadius
 float4x4 StencilingConeTransform;
 float3 StencilingPreViewTranslation;
+#endif
 
-#if SHADER_RADIAL_LIGHT == 0
+#if defined(SHADER_RADIAL_LIGHT) && SHADER_RADIAL_LIGHT == 0
 /** Vertex shader for rendering a directional light using a full screen quad. */
 void DirectionalVertexMain(
 	in float2 InPosition : ATTRIBUTE0,
@@ -96,10 +98,9 @@ void RadialVertexMain(
 }
 #endif
 
-#if SHADER_HAIR == 1
+#if defined(SHADER_HAIR) && SHADER_HAIR == 1
 
-uint NodeGroupSize;
-int2 MaxResolution;
+int2 MaxViewportResolution;
 Texture2D<uint> HairVisibilityNodeCount;
 
 void HairVertexMain(
@@ -112,7 +113,7 @@ void HairVertexMain(
 	OutResolution.x = ceil(sqrt(OutNodeCount));
 	OutResolution.y = OutResolution.x;
 
-	const float2 ClipCoord = ((OutResolution + 0.5f) / float2(MaxResolution)) * 2;
+	const float2 ClipCoord = ((OutResolution + 0.5f) / float2(MaxViewportResolution)) * 2;
 
 	const float2 UV = float2(InVertexId & 1, InVertexId >> 1);
 	const float2 Pos = float2(-UV.x, UV.y) * 4 + float2(-1, +1) + float2(ClipCoord.x, -ClipCoord.y);
diff --git a/Engine/Source/Runtime/Renderer/Private/LightRendering.cpp b/Engine/Source/Runtime/Renderer/Private/LightRendering.cpp
--- a/Engine/Source/Runtime/Renderer/Private/LightRendering.cpp
+++ b/Engine/Source/Runtime/Renderer/Private/LightRendering.cpp
@@ -245,14 +245,13 @@ public:
 	TDeferredLightHairVS(const ShaderMetaType::CompiledShaderInitializerType& Initializer) :
 		FGlobalShader(Initializer)
 	{
-		MaxResolution.Bind(Initializer.ParameterMap, TEXT("MaxResolution"));
-		NodeGroupSize.Bind(Initializer.ParameterMap, TEXT("NodeGroupSize"));
+		MaxViewportResolution.Bind(Initializer.ParameterMap, TEXT("MaxViewportResolution"));
 		HairVisibilityNodeCount.Bind(Initializer.ParameterMap, TEXT("HairVisibilityNodeCount"));
 	}
 
 	void SetParameters(FRHICommandList& RHICmdList, const FViewInfo& View, const FHairStrandsVisibilityData* VisibilityData)
 	{
-		FRHIPixelShader* ShaderRHI = RHICmdList.GetBoundPixelShader();
+		FRHIVertexShader* ShaderRHI = RHICmdList.GetBoundVertexShader();
 		FGlobalShader::SetParameters<FViewUniformShaderParameters>(RHICmdList, ShaderRHI, View.ViewUniformBuffer);
 
 		if (!VisibilityData)
@@ -265,13 +264,11 @@ public:
 			SetTextureParameter(RHICmdList,	ShaderRHI, HairVisibilityNodeCount, VisibilityData->NodeCount->GetRenderTargetItem().ShaderResourceTexture);
 		}
 
-		SetShaderValue(RHICmdList, ShaderRHI, NodeGroupSize, VisibilityData->NodeGroupSize);
-		SetShaderValue(RHICmdList, ShaderRHI, MaxResolution, VisibilityData->SampleLightingViewportResolution);
+		SetShaderValue(RHICmdList, ShaderRHI, MaxViewportResolution, VisibilityData->SampleLightingViewportResolution);
 	}
 
 private:
-	LAYOUT_FIELD(FShaderParameter, MaxResolution);
-	LAYOUT_FIELD(FShaderParameter, NodeGroupSize);
+	LAYOUT_FIELD(FShaderParameter, MaxViewportResolution);
 	LAYOUT_FIELD(FShaderResourceParameter, HairVisibilityNodeCount);
 };
 
@@ -2278,9 +2275,7 @@ void FDeferredShadingSceneRenderer::RenderLightForHair(
 		GraphicsPSOInit.PrimitiveType = PT_TriangleList;
 		SetGraphicsPipelineState(RHICmdList, GraphicsPSOInit);
 
-
 		VertexShader->SetParameters(RHICmdList, View, &HairVisibilityData);
-		PixelShader->SetParameters(RHICmdList, View,  LightSceneInfo, HairShadowMaskTexture, &RenderLightParams);
 		PixelShader->SetParameters(RHICmdList, View,  LightSceneInfo, HairShadowMaskTexture, bHairRenderingEnabled ? &RenderLightParams : nullptr);
 
 		RHICmdList.SetStreamSource(0, nullptr, 0);
