Fix CIS static analysis issue.

diff --git a/Engine/Source/Runtime/Renderer/Private/HairStrands/HairStrandsVisibility.cpp b/Engine/Source/Runtime/Renderer/Private/HairStrands/HairStrandsVisibility.cpp
--- a/Engine/Source/Runtime/Renderer/Private/HairStrands/HairStrandsVisibility.cpp
+++ b/Engine/Source/Runtime/Renderer/Private/HairStrands/HairStrandsVisibility.cpp
@@ -2265,14 +2265,14 @@ struct FRasterComputeOutput
 	FIntPoint SuperResolution;
 	uint32 ResolutionMultiplier = 1;
 
-	FRDGTextureRef HairCountTexture;
-	FRDGTextureRef DepthTexture;
-	FRDGTextureRef VisibilityTexture;
+	FRDGTextureRef HairCountTexture = nullptr;
+	FRDGTextureRef DepthTexture = nullptr;
+	FRDGTextureRef VisibilityTexture = nullptr;
 };
 
 static FRasterComputeOutput AddVisibilityComputeRasterPass(
 	FRDGBuilder& GraphBuilder,
-	const FViewInfo* ViewInfo,
+	const FViewInfo& ViewInfo,
 	const FHairStrandsMacroGroupDatas& MacroGroupDatas,
 	const FIntPoint& InResolution,
 	const FRDGTextureRef SceneDepthTexture)
@@ -2341,13 +2341,13 @@ static FRasterComputeOutput AddVisibilityComputeRasterPass(
 	// Create and set the uniform buffer
 	const bool bEnableMSAA = false;
 	TUniformBufferRef<FViewUniformShaderParameters> ViewUniformShaderParameters;
-	SetUpViewHairRenderInfo(*ViewInfo, bEnableMSAA, ViewInfo->CachedViewUniformShaderParameters->HairRenderInfo, ViewInfo->CachedViewUniformShaderParameters->HairRenderInfoBits);
-	ViewUniformShaderParameters = TUniformBufferRef<FViewUniformShaderParameters>::CreateUniformBufferImmediate(*ViewInfo->CachedViewUniformShaderParameters, UniformBuffer_SingleFrame);
+	SetUpViewHairRenderInfo(ViewInfo, bEnableMSAA, ViewInfo.CachedViewUniformShaderParameters->HairRenderInfo, ViewInfo.CachedViewUniformShaderParameters->HairRenderInfoBits);
+	ViewUniformShaderParameters = TUniformBufferRef<FViewUniformShaderParameters>::CreateUniformBufferImmediate(*ViewInfo.CachedViewUniformShaderParameters, UniformBuffer_SingleFrame);
 
-	const uint32 FrameIdMode8 = ViewInfo && ViewInfo->ViewState ? (ViewInfo->ViewState->GetFrameIndex() % 8) : 0;
+	const uint32 FrameIdMode8 = ViewInfo.ViewState ? (ViewInfo.ViewState->GetFrameIndex() % 8) : 0;
 	const uint32 GroupSize = 32;
 	const uint32 DispatchCountX = 64;
-	TShaderMapRef<FVisiblityRasterComputeCS> ComputeShader(ViewInfo->ShaderMap);
+	TShaderMapRef<FVisiblityRasterComputeCS> ComputeShader(ViewInfo.ShaderMap);
 
 	for (const FHairStrandsMacroGroupData& MacroGroup : MacroGroupDatas.Datas)
 	{
@@ -2457,7 +2457,7 @@ FHairStrandsVisibilityViews RenderHairStrandsVisibilityBuffer(
 				{
 					FRasterComputeOutput RasterOutput = AddVisibilityComputeRasterPass(
 						GraphBuilder,
-						&View,
+						View,
 						MacroGroupDatas,
 						Resolution,
 						SceneDepthTexture);
