Make sure thread-safe access to cvars in new shadow code

diff --git a/Engine/Source/Runtime/RenderCore/Private/RenderCore.cpp b/Engine/Source/Runtime/RenderCore/Private/RenderCore.cpp
--- a/Engine/Source/Runtime/RenderCore/Private/RenderCore.cpp
+++ b/Engine/Source/Runtime/RenderCore/Private/RenderCore.cpp
@@ -267,6 +267,19 @@ RENDERCORE_API int32 GetCVarForceLOD()
 	return Ret;
 }
 
+RENDERCORE_API int32 GetCVarForceLOD_AnyThread()
+{
+	int32 Ret = -1;
+
+#if EXPOSE_FORCE_LOD
+	{
+		Ret = CVarForceLOD.GetValueOnAnyThread();
+	}
+#endif // EXPOSE_FORCE_LOD
+
+	return Ret;
+}
+
 RENDERCORE_API int32 GetCVarForceLODShadow()
 {
 	int32 Ret = -1;
@@ -280,6 +293,19 @@ RENDERCORE_API int32 GetCVarForceLODShadow()
 	return Ret;
 }
 
+RENDERCORE_API int32 GetCVarForceLODShadow_AnyThread()
+{
+	int32 Ret = -1;
+
+#if EXPOSE_FORCE_LOD
+	{
+		Ret = CVarForceLODShadow.GetValueOnAnyThread();
+	}
+#endif // EXPOSE_FORCE_LOD
+
+	return Ret;
+}
+
 RENDERCORE_API bool IsHDREnabled()
 {
 	static const auto CVarHDROutputEnabled = IConsoleManager::Get().FindTConsoleVariableDataInt(TEXT("r.HDR.EnableHDROutput"));
diff --git a/Engine/Source/Runtime/RenderCore/Public/RenderCore.h b/Engine/Source/Runtime/RenderCore/Public/RenderCore.h
--- a/Engine/Source/Runtime/RenderCore/Public/RenderCore.h
+++ b/Engine/Source/Runtime/RenderCore/Public/RenderCore.h
@@ -287,8 +287,10 @@ extern RENDERCORE_API uint32 GSwapBufferTime;
 
 // shared by renderer and engine, compiles down to a constant in final release
 RENDERCORE_API int32 GetCVarForceLOD();
+RENDERCORE_API int32 GetCVarForceLOD_AnyThread();
 
 RENDERCORE_API int32 GetCVarForceLODShadow();
+RENDERCORE_API int32 GetCVarForceLODShadow_AnyThread();
 
 RENDERCORE_API bool IsHDREnabled();
 
diff --git a/Engine/Source/Runtime/Renderer/Private/SceneRendering.cpp b/Engine/Source/Runtime/Renderer/Private/SceneRendering.cpp
--- a/Engine/Source/Runtime/Renderer/Private/SceneRendering.cpp
+++ b/Engine/Source/Runtime/Renderer/Private/SceneRendering.cpp
@@ -75,6 +75,11 @@ bool UseCachedMeshDrawCommands()
 	return CVarCachedMeshDrawCommands.GetValueOnRenderThread() > 0;
 }
 
+bool UseCachedMeshDrawCommands_AnyThread()
+{
+	return CVarCachedMeshDrawCommands.GetValueOnAnyThread() > 0;
+}
+
 static TAutoConsoleVariable<int32> CVarMeshDrawCommandsDynamicInstancing(
 	TEXT("r.MeshDrawCommands.DynamicInstancing"),
 	1,
diff --git a/Engine/Source/Runtime/Renderer/Private/SceneRendering.h b/Engine/Source/Runtime/Renderer/Private/SceneRendering.h
--- a/Engine/Source/Runtime/Renderer/Private/SceneRendering.h
+++ b/Engine/Source/Runtime/Renderer/Private/SceneRendering.h
@@ -2049,6 +2049,7 @@ private:
 extern FFastVramConfig GFastVRamConfig;
 
 extern bool UseCachedMeshDrawCommands();
+extern bool UseCachedMeshDrawCommands_AnyThread();
 extern bool IsDynamicInstancingEnabled(ERHIFeatureLevel::Type FeatureLevel);
 
 enum class EGPUSkinCacheTransition
diff --git a/Engine/Source/Runtime/Renderer/Private/ShadowSetup.cpp b/Engine/Source/Runtime/Renderer/Private/ShadowSetup.cpp
--- a/Engine/Source/Runtime/Renderer/Private/ShadowSetup.cpp
+++ b/Engine/Source/Runtime/Renderer/Private/ShadowSetup.cpp
@@ -1091,7 +1091,7 @@ void FProjectedShadowInfo::AddCachedMeshDrawCommands_AnyThread(
 {
 	const EMeshPass::Type PassType = EMeshPass::CSMShadowDepth;
 	const EShadingPath ShadingPath = Scene->GetShadingPath();
-	const bool bUseCachedMeshCommand = UseCachedMeshDrawCommands()
+	const bool bUseCachedMeshCommand = UseCachedMeshDrawCommands_AnyThread()
 		&& !!(FPassProcessorManager::GetPassFlags(ShadingPath, PassType) & EMeshPassFlags::CachedMeshCommands)
 		&& StaticMeshRelevance.bSupportsCachingMeshDrawCommands;
 
@@ -1276,7 +1276,7 @@ bool FProjectedShadowInfo::ShouldDrawStaticMeshes_AnyThread(
 	const int32 PrimitiveId = PrimitiveSceneInfo->GetIndex();
 
 	{
-		const int32 ForcedLOD = CurrentView.Family->EngineShowFlags.LOD ? (GetCVarForceLODShadow() != -1 ? GetCVarForceLODShadow() : GetCVarForceLOD()) : -1;
+		const int32 ForcedLOD = CurrentView.Family->EngineShowFlags.LOD ? (GetCVarForceLODShadow_AnyThread() != -1 ? GetCVarForceLODShadow_AnyThread() : GetCVarForceLOD_AnyThread()) : -1;
 		const FLODMask* VisibilePrimitiveLODMask = nullptr;
 
 		if (CurrentView.PrimitivesLODMask[PrimitiveId].ContainsLOD(MAX_int8)) // only calculate it if it's not set
@@ -1597,11 +1597,11 @@ uint64 FProjectedShadowInfo::AddSubjectPrimitive_AnyThread(
 		static auto* CVarMobileEnableStaticAndCSMShadowReceivers = IConsoleManager::Get().FindTConsoleVariableDataInt(TEXT("r.Mobile.EnableStaticAndCSMShadowReceivers"));
 		static auto* CVarMobileEnableMovableLightCSMShaderCulling = IConsoleManager::Get().FindTConsoleVariableDataInt(TEXT("r.Mobile.EnableMovableLightCSMShaderCulling"));
 		static auto* CVarMobileCSMShaderCullingMethod = IConsoleManager::Get().FindTConsoleVariableDataInt(TEXT("r.Mobile.Shadow.CSMShaderCullingMethod"));
-		const uint32 MobileCSMCullingMode = CVarMobileCSMShaderCullingMethod->GetValueOnRenderThread() & 0xF;
+		const uint32 MobileCSMCullingMode = CVarMobileCSMShaderCullingMethod->GetValueOnAnyThread() & 0xF;
 		const bool bRecordShadowSubjectsForMobile =
 			(MobileCSMCullingMode == 2 || MobileCSMCullingMode == 3)
-			&& ((CVarMobileEnableMovableLightCSMShaderCulling->GetValueOnRenderThread() && GetLightSceneInfo().Proxy->IsMovable() && GetLightSceneInfo().ShouldRenderViewIndependentWholeSceneShadows())
-				|| (CVarMobileEnableStaticAndCSMShadowReceivers->GetValueOnRenderThread() && GetLightSceneInfo().Proxy->UseCSMForDynamicObjects()));
+			&& ((CVarMobileEnableMovableLightCSMShaderCulling->GetValueOnAnyThread() && GetLightSceneInfo().Proxy->IsMovable() && GetLightSceneInfo().ShouldRenderViewIndependentWholeSceneShadows())
+				|| (CVarMobileEnableStaticAndCSMShadowReceivers->GetValueOnAnyThread() && GetLightSceneInfo().Proxy->UseCSMForDynamicObjects()));
 
 		if (bRecordShadowSubjectsForMobile)
 		{
