UE4.25 - Fix LPV crash
#rnx

diff --git a/Engine/Source/Runtime/Renderer/Private/LightPropagationVolume.cpp b/Engine/Source/Runtime/Renderer/Private/LightPropagationVolume.cpp
--- a/Engine/Source/Runtime/Renderer/Private/LightPropagationVolume.cpp
+++ b/Engine/Source/Runtime/Renderer/Private/LightPropagationVolume.cpp
@@ -979,10 +979,6 @@ void FLightPropagationVolume::Clear(FRHICommandListImmediate& RHICmdList, FViewI
 		Shader->UnbindBuffers(RHICmdList, ShaderParams);
 	}
 	RHICmdList.EndUAVOverlap();
-
-	RHICmdList.SetUAVParameter( FComputeShaderRHIRef(), 7, mVplListBuffer->UAV, 0 );
-	RHICmdList.SetUAVParameter( FComputeShaderRHIRef(), 7, GvListBuffer->UAV, 0 );
-	RHICmdList.SetUAVParameter( FComputeShaderRHIRef(), 7, nullptr, 0 );
 }
 
 /**
