Compiles FFT bloom only on PC

diff --git a/Engine/Source/Runtime/Renderer/Private/GPUFastFourierTransform.cpp b/Engine/Source/Runtime/Renderer/Private/GPUFastFourierTransform.cpp
--- a/Engine/Source/Runtime/Renderer/Private/GPUFastFourierTransform.cpp
+++ b/Engine/Source/Runtime/Renderer/Private/GPUFastFourierTransform.cpp
@@ -200,7 +200,7 @@ namespace GPUFFT
 		static bool ShouldCompilePermutation(const FGlobalShaderPermutationParameters& Parameters)
 		{
 			// @todo MetalMRT: Metal MRT can't cope with the threadgroup storage requirements for these shaders right now
-			return IsFeatureLevelSupported(Parameters.Platform, ERHIFeatureLevel::SM5) && !IsMetalMRTPlatform(Parameters.Platform);
+			return IsFeatureLevelSupported(Parameters.Platform, ERHIFeatureLevel::SM5) && !IsMetalMRTPlatform(Parameters.Platform) && IsPCPlatform(Parameters.Platform);
 		}
 
 		static void ModifyCompilationEnvironment(const FGlobalShaderPermutationParameters& Parameters, FShaderCompilerEnvironment& OutEnvironment)
@@ -291,7 +291,7 @@ namespace GPUFFT
 		static bool ShouldCompilePermutation(const FGlobalShaderPermutationParameters& Parameters)
 		{
 			// @todo MetalMRT: Metal MRT can't cope with the threadgroup storage requirements for these shaders right now
-			return IsFeatureLevelSupported(Parameters.Platform, ERHIFeatureLevel::SM5) && !IsMetalMRTPlatform(Parameters.Platform);
+			return IsFeatureLevelSupported(Parameters.Platform, ERHIFeatureLevel::SM5) && !IsMetalMRTPlatform(Parameters.Platform) && IsPCPlatform(Parameters.Platform);
 		}
 
 		static void ModifyCompilationEnvironment(const FGlobalShaderPermutationParameters& Parameters, FShaderCompilerEnvironment& OutEnvironment)
@@ -379,7 +379,7 @@ namespace GPUFFT
 		static bool ShouldCompilePermutation(const FGlobalShaderPermutationParameters& Parameters)
 		{
 			// @todo MetalMRT: Metal MRT can't cope with the threadgroup storage requirements for these shaders right now
-			return IsFeatureLevelSupported(Parameters.Platform, ERHIFeatureLevel::SM5) && !IsMetalMRTPlatform(Parameters.Platform);
+			return IsFeatureLevelSupported(Parameters.Platform, ERHIFeatureLevel::SM5) && !IsMetalMRTPlatform(Parameters.Platform) && IsPCPlatform(Parameters.Platform);
 		}
 
 		static void ModifyCompilationEnvironment(const FGlobalShaderPermutationParameters& Parameters, FShaderCompilerEnvironment& OutEnvironment)
@@ -465,7 +465,7 @@ namespace GPUFFT
 		static bool ShouldCompilePermutation(const FGlobalShaderPermutationParameters& Parameters)
 		{
 			// @todo MetalMRT: Metal MRT can't cope with the threadgroup storage requirements for these shaders right now
-			return IsFeatureLevelSupported(Parameters.Platform, ERHIFeatureLevel::SM5) && !IsMetalMRTPlatform(Parameters.Platform);
+			return IsFeatureLevelSupported(Parameters.Platform, ERHIFeatureLevel::SM5) && !IsMetalMRTPlatform(Parameters.Platform) && IsPCPlatform(Parameters.Platform);
 		}
 
 		static void ModifyCompilationEnvironment(const FGlobalShaderPermutationParameters& Parameters, FShaderCompilerEnvironment& OutEnvironment)
@@ -539,7 +539,7 @@ namespace GPUFFT
 		static bool ShouldCompilePermutation(const FGlobalShaderPermutationParameters& Parameters)
 		{
 			// @todo MetalMRT: Metal MRT can't cope with the threadgroup storage requirements for these shaders right now
-			return IsFeatureLevelSupported(Parameters.Platform, ERHIFeatureLevel::SM5) && !IsMetalMRTPlatform(Parameters.Platform);
+			return IsFeatureLevelSupported(Parameters.Platform, ERHIFeatureLevel::SM5) && !IsMetalMRTPlatform(Parameters.Platform) && IsPCPlatform(Parameters.Platform);
 		}
 
 		static void ModifyCompilationEnvironment(const FGlobalShaderPermutationParameters& Parameters, FShaderCompilerEnvironment& OutEnvironment)
@@ -611,7 +611,7 @@ namespace GPUFFT
 		static bool ShouldCompilePermutation(const FGlobalShaderPermutationParameters& Parameters)
 		{
 			// @todo MetalMRT: Metal MRT can't cope with the threadgroup storage requirements for these shaders right now
-			return IsFeatureLevelSupported(Parameters.Platform, ERHIFeatureLevel::SM5) && !IsMetalMRTPlatform(Parameters.Platform);
+			return IsFeatureLevelSupported(Parameters.Platform, ERHIFeatureLevel::SM5) && !IsMetalMRTPlatform(Parameters.Platform) && IsPCPlatform(Parameters.Platform);
 		}
 
 		static void ModifyCompilationEnvironment(const FGlobalShaderPermutationParameters& Parameters, FShaderCompilerEnvironment& OutEnvironment)
@@ -740,7 +740,7 @@ namespace GPUFFT
 		static bool ShouldCompilePermutation(const FGlobalShaderPermutationParameters& Parameters)
 		{
 			// @todo MetalMRT: Metal MRT can't cope with the threadgroup storage requirements for these shaders right now
-			return IsFeatureLevelSupported(Parameters.Platform, ERHIFeatureLevel::SM5) && !IsMetalMRTPlatform(Parameters.Platform);
+			return IsFeatureLevelSupported(Parameters.Platform, ERHIFeatureLevel::SM5) && !IsMetalMRTPlatform(Parameters.Platform) && IsPCPlatform(Parameters.Platform);
 		}
 
 
@@ -777,7 +777,7 @@ namespace GPUFFT
 		static bool ShouldCompilePermutation(const FGlobalShaderPermutationParameters& Parameters)
 		{
 			// @todo MetalMRT: Metal MRT can't cope with the threadgroup storage requirements for these shaders right now
-			return IsFeatureLevelSupported(Parameters.Platform, ERHIFeatureLevel::SM5) && !IsMetalMRTPlatform(Parameters.Platform);
+			return IsFeatureLevelSupported(Parameters.Platform, ERHIFeatureLevel::SM5) && !IsMetalMRTPlatform(Parameters.Platform) && IsPCPlatform(Parameters.Platform);
 		}
 
 
@@ -927,7 +927,7 @@ namespace GPUFFT
 		static bool ShouldCompilePermutation(const FGlobalShaderPermutationParameters& Parameters)
 		{
 			// @todo MetalMRT: Metal MRT can't cope with the threadgroup storage requirements for these shaders right now
-			return IsFeatureLevelSupported(Parameters.Platform, ERHIFeatureLevel::SM5) && !IsMetalMRTPlatform(Parameters.Platform);
+			return IsFeatureLevelSupported(Parameters.Platform, ERHIFeatureLevel::SM5) && !IsMetalMRTPlatform(Parameters.Platform) && IsPCPlatform(Parameters.Platform);
 		}
 
 
diff --git a/Engine/Source/Runtime/Renderer/Private/PostProcess/PostProcessFFTBloom.cpp b/Engine/Source/Runtime/Renderer/Private/PostProcess/PostProcessFFTBloom.cpp
--- a/Engine/Source/Runtime/Renderer/Private/PostProcess/PostProcessFFTBloom.cpp
+++ b/Engine/Source/Runtime/Renderer/Private/PostProcess/PostProcessFFTBloom.cpp
@@ -23,7 +23,7 @@ public:
 	static bool ShouldCompilePermutation(const FGlobalShaderPermutationParameters& Parameters)
 	{
 		// @todo MetalMRT: Metal MRT can't cope with the threadgroup storage requirements for these shaders right now
-		return IsFeatureLevelSupported(Parameters.Platform, ERHIFeatureLevel::SM5) && !IsMetalMRTPlatform(Parameters.Platform);
+		return IsFeatureLevelSupported(Parameters.Platform, ERHIFeatureLevel::SM5) && !IsMetalMRTPlatform(Parameters.Platform) && IsPCPlatform(Parameters.Platform);
 	}
 
 	FFFTShader() = default;
@@ -185,7 +185,7 @@ bool IsFFTBloomPhysicalKernelReady(const FViewInfo& View)
 
 bool IsFFTBloomEnabled(const FViewInfo& View)
 {
-	const bool bOldMetalNoFFT = IsMetalPlatform(View.GetShaderPlatform()) && (RHIGetShaderLanguageVersion(View.GetShaderPlatform()) < 4);
+	const bool bOldMetalNoFFT = IsMetalPlatform(View.GetShaderPlatform()) && (RHIGetShaderLanguageVersion(View.GetShaderPlatform()) < 4) && IsPCPlatform(View.GetShaderPlatform());
 	const bool bUseFFTBloom = View.FinalPostProcessSettings.BloomMethod == EBloomMethod::BM_FFT && View.ViewState != nullptr;
 
 	static bool bWarnAboutOldMetalFFTOnce = false;
