Enable voxel tracing by default when using raytracing.

diff --git a/Engine/Source/Runtime/Renderer/Private/HairStrands/HairStrandsInterface.cpp b/Engine/Source/Runtime/Renderer/Private/HairStrands/HairStrandsInterface.cpp
--- a/Engine/Source/Runtime/Renderer/Private/HairStrands/HairStrandsInterface.cpp
+++ b/Engine/Source/Runtime/Renderer/Private/HairStrands/HairStrandsInterface.cpp
@@ -18,7 +18,7 @@ DEFINE_LOG_CATEGORY_STATIC(LogHairRendering, Log, All);
 static int32 GHairStrandsRenderingEnable = 1;
 static FAutoConsoleVariableRef CVarHairStrandsRenderingEnable(TEXT("r.HairStrands.Enable"), GHairStrandsRenderingEnable, TEXT("Enable/Disable hair strands rendering"));
 
-static int32 GHairStrandsRaytracingEnable = 1;
+static int32 GHairStrandsRaytracingEnable = 0;
 static FAutoConsoleVariableRef CVarHairStrandsRaytracingEnable(TEXT("r.HairStrands.Raytracing"), GHairStrandsRaytracingEnable, TEXT("Enable/Disable hair strands raytracing (fallback onto raster techniques"));
 
 ////////////////////////////////////////////////////////////////////////////////////////////////////////////////////////
