Enable GPU driven voxelization for hair by default and change default hair voxel resolution to be 0.3 as the page table get filled too quickly otherwise.

diff --git a/Engine/Source/Runtime/Renderer/Private/HairStrands/HairStrandsVoxelization.cpp b/Engine/Source/Runtime/Renderer/Private/HairStrands/HairStrandsVoxelization.cpp
--- a/Engine/Source/Runtime/Renderer/Private/HairStrands/HairStrandsVoxelization.cpp
+++ b/Engine/Source/Runtime/Renderer/Private/HairStrands/HairStrandsVoxelization.cpp
@@ -56,7 +56,7 @@ static int32 GHairForVoxelTransmittanceAndShadow = 0;
 static FAutoConsoleVariableRef CVarHairForVoxelTransmittanceAndShadow(TEXT("r.HairStrands.Voxelization.ForceTransmittanceAndShadow"), GHairForVoxelTransmittanceAndShadow, TEXT("For transmittance and shadow to be computed with density volume. This requires voxelization is enabled."));
 
 static int32 GHairVirtualVoxel = 1;
-static float GHairVirtualVoxel_VoxelWorldSize = 0.15f; // 1.5mm
+static float GHairVirtualVoxel_VoxelWorldSize = 0.3f; // 3.0mm
 static int32 GHairVirtualVoxel_PageResolution = 32;
 static int32 GHairVirtualVoxel_PageCountPerDim = 14;
 static FAutoConsoleVariableRef CVarHairVirtualVoxel(TEXT("r.HairStrands.Voxelization.Virtual"), GHairVirtualVoxel, TEXT("Enable the two voxel hierachy."));
@@ -64,7 +64,7 @@ static FAutoConsoleVariableRef CVarHairVirtualVoxel_VoxelWorldSize(TEXT("r.HairS
 static FAutoConsoleVariableRef CVarHairVirtualVoxel_VoxelPageResolution(TEXT("r.HairStrands.Voxelization.Virtual.VoxelPageResolution"), GHairVirtualVoxel_PageResolution, TEXT("Resolution of a voxel page."));
 static FAutoConsoleVariableRef CVarHairVirtualVoxel_VoxelPageCountPerDim(TEXT("r.HairStrands.Voxelization.Virtual.VoxelPageCountPerDim"), GHairVirtualVoxel_PageCountPerDim, TEXT("Number of voxel pages per texture dimension. The voxel page memory is allocated with a 3D texture. This value provide the resolution of this texture."));
 
-static int32 GHairVirtualVoxelGPUDriven = 0;
+static int32 GHairVirtualVoxelGPUDriven = 1;
 static int32 GHairVirtualVoxelGPUDrivenMaxPageIndexRes = 32;
 static FAutoConsoleVariableRef CVarHairVirtualVoxelGPUDriven(TEXT("r.HairStrands.Voxelization.GPUDriven"), GHairVirtualVoxelGPUDriven, TEXT("Enable GPU driven voxelization."));
 static FAutoConsoleVariableRef CVarHairVirtualVoxelGPUDrivenMaxPageIndexRes(TEXT("r.HairStrands.Voxelization.GPUDriven.MaxPageIndexResolution"), GHairVirtualVoxelGPUDrivenMaxPageIndexRes, TEXT("Max resolution of the page index. This is used for allocating a conservative page index buffer when GPU driven allocation is enabled."));
