Fix skin cache not running when strands hair is not enabled.

diff --git a/Engine/Source/Runtime/Renderer/Private/SceneRendering.cpp b/Engine/Source/Runtime/Renderer/Private/SceneRendering.cpp
--- a/Engine/Source/Runtime/Renderer/Private/SceneRendering.cpp
+++ b/Engine/Source/Runtime/Renderer/Private/SceneRendering.cpp
@@ -4215,7 +4215,9 @@ void RunGPUSkinCacheTransition(FRHICommandList& RHICmdList, FScene* Scene, EGPUS
 	// * When hair strands is enabled, the skin cache sync point is run earlier, during 
 	//   the init views pass, as the output of the skin cached is used by Niagara
 	const bool bHairStrandsEnabled = IsHairStrandsEnable(Scene->GetShaderPlatform());
-	const bool bRun = bHairStrandsEnabled && EGPUSkinCacheTransition::FrameSetup == Type;
+	const bool bRun = 
+		(bHairStrandsEnabled && EGPUSkinCacheTransition::FrameSetup == Type) || 
+		(!bHairStrandsEnabled && EGPUSkinCacheTransition::FrameSetup != Type);
 	if (bRun)
 	{
 		if (FGPUSkinCache* GPUSkinCache = Scene->GetGPUSkinCache())
