Add (temp) CVar for controlling hair density for transmission evaluation and enabling/disabling hair count texture for transmitted light evaluation.

diff --git a/Engine/Shaders/Private/HairStrands/HairStrandsEnvironmentLighting.usf b/Engine/Shaders/Private/HairStrands/HairStrandsEnvironmentLighting.usf
--- a/Engine/Shaders/Private/HairStrands/HairStrandsEnvironmentLighting.usf
+++ b/Engine/Shaders/Private/HairStrands/HairStrandsEnvironmentLighting.usf
@@ -96,6 +96,7 @@ uint  MultipleScatterSampleCount;
 uint  JitterSphericalIntegration;
 uint  HairComponents;
 float HairDualScatteringRoughnessOverride;
+float TransmissionDensityScaleFactor;
 
 uint  Voxel_MacroGroupId;
 float Voxel_TanConeAngle;
@@ -459,16 +460,15 @@ float3 ReflectionEnvironment(FGBufferData GBuffer, float2 ScreenPosition, float2
 			// Compute the number of hair count between light & shading point
 			// For TT, we don't want to abort the hair counting, as this is important to get the proper transmission
 			const float DistanceThresholdTT = 100000;
-			const float DensityScaleFactor = 2;
 			FHairTraversalResult Result;
 			if (bHairUseViewHairCount)
 			{
-				Result.HairCount = HairCountTexture.Load(uint3(SvPosition.xy, 0)) * DensityScaleFactor;
+				Result.HairCount = HairCountTexture.Load(uint3(SvPosition.xy, 0)) * max(TransmissionDensityScaleFactor, 0);
 				Result.Visibility = 1;
 			}
 			else
 			{
-				Result = InternalComputeHairCountVoxel(SampleWorldPosition, SampleLightPosition, Voxel_MinAABB, Voxel_MaxAABB, Voxel_Resolution, Voxel_DensityTexture, Voxel_DensityScale * DensityScaleFactor, 0, DistanceThresholdTT, false, bDebugEnabled);
+				Result = InternalComputeHairCountVoxel(SampleWorldPosition, SampleLightPosition, Voxel_MinAABB, Voxel_MaxAABB, Voxel_Resolution, Voxel_DensityTexture, Voxel_DensityScale * TransmissionDensityScaleFactor, 0, DistanceThresholdTT, false, bDebugEnabled);
 			}
 			TransmittanceMask = ComputeTransmittanceFromHairCount(Result.HairCount, Result.Visibility, GBuffer, SinLightAngle, HairScatteringLUTTexture, HairLUTSampler);
 		}
diff --git a/Engine/Source/Runtime/Renderer/Private/HairStrands/HairStrandsEnvironment.cpp b/Engine/Source/Runtime/Renderer/Private/HairStrands/HairStrandsEnvironment.cpp
--- a/Engine/Source/Runtime/Renderer/Private/HairStrands/HairStrandsEnvironment.cpp
+++ b/Engine/Source/Runtime/Renderer/Private/HairStrands/HairStrandsEnvironment.cpp
@@ -49,6 +49,12 @@ static FAutoConsoleVariableRef CVarHairStrandsSkyLightingJitterSphericalIntegrat
 static int32 GHairStrandsSkyAOSampleCount = 16;
 static FAutoConsoleVariableRef CVarHairStrandsSkyAOSampleCount(TEXT("r.HairStrands.SkyAO.SampleCount"), GHairStrandsSkyAOSampleCount, TEXT("Number of samples used for evaluating hair AO (default is set to 16)."), ECVF_Scalability | ECVF_RenderThreadSafe);
 
+static float GHairStrandsTransmissionDensityScaleFactor = 4;
+static FAutoConsoleVariableRef CVarHairStrandsTransmissionDensityScaleFactor(TEXT("r.HairStrands.SkyLighting.TransmissionDensityScale"), GHairStrandsTransmissionDensityScaleFactor, TEXT("Density scale for controlling how much sky lighting is transmitted."), ECVF_Scalability | ECVF_RenderThreadSafe);
+
+static int32 GHairStrandsSkyLightingUseHairCountTexture = 1;
+static FAutoConsoleVariableRef CVarHairStrandsSkyLightingUseHairCountTexture(TEXT("r.HairStrands.SkyLighting.UseViewHairCount"), GHairStrandsSkyLightingUseHairCountTexture, TEXT("Use the view hair count texture for estimating background transmitted light (enabled by default)."), ECVF_Scalability | ECVF_RenderThreadSafe);
+
 static float GHairStrandsSkyAODistanceThreshold = 10;
 static float GHairStrandsSkyLightingDistanceThreshold = 10;
 static FAutoConsoleVariableRef CVarHairStrandsSkyAOThreshold(TEXT("r.HairStrands.SkyAO.DistanceThreshold"), GHairStrandsSkyAODistanceThreshold, TEXT("Max distance for occlusion search."), ECVF_Scalability | ECVF_RenderThreadSafe);
@@ -235,6 +241,7 @@ public:
 		SHADER_PARAMETER(uint32, MultipleScatterSampleCount)
 		SHADER_PARAMETER(uint32, HairComponents)
 		SHADER_PARAMETER(float,  HairDualScatteringRoughnessOverride)
+		SHADER_PARAMETER(float, TransmissionDensityScaleFactor)
 		SHADER_PARAMETER(uint32, JitterSphericalIntegration)
 		SHADER_PARAMETER(float, HairDistanceThreshold)
 		SHADER_PARAMETER(uint32, bHairUseViewHairCount)
@@ -364,12 +371,13 @@ static void AddHairStrandsEnvironmentLightingPassPS(
 	PassParameters->Voxel_DepthBiasScale = GetHairStrandsVoxelizationDepthBiasScale();
 	PassParameters->Voxel_TanConeAngle = FMath::Tan(FMath::DegreesToRadians(GetHairStrandsSkyLightingConeAngle()));
 	PassParameters->HairDistanceThreshold = FMath::Max(GHairStrandsSkyLightingDistanceThreshold, 1.f);
-	PassParameters->bHairUseViewHairCount = VisibilityData.ViewHairCountTexture ? 1 : 0;
+	PassParameters->bHairUseViewHairCount = VisibilityData.ViewHairCountTexture && GHairStrandsSkyLightingUseHairCountTexture ? 1 : 0;
 	PassParameters->MaxVisibilityNodeCount = VisibilityData.NodeData->Desc.NumElements;
 	PassParameters->MultipleScatterSampleCount = FMath::Max(uint32(GHairStrandsSkyLightingSampleCount), 1u);
 	PassParameters->JitterSphericalIntegration = GHairStrandsSkyLightingJitterSphericalIntegration ? 1 : 0;
 	PassParameters->HairComponents = ToBitfield(GetHairComponents());
 	PassParameters->HairDualScatteringRoughnessOverride = GetHairDualScatteringRoughnessOverride();
+	PassParameters->TransmissionDensityScaleFactor = FMath::Max(0.f, GHairStrandsTransmissionDensityScaleFactor);
 	PassParameters->PreIntegratedGF = GSystemTextures.PreintegratedGF->GetRenderTargetItem().ShaderResourceTexture;
 	PassParameters->PreIntegratedGFSampler = TStaticSamplerState<SF_Bilinear, AM_Clamp, AM_Clamp, AM_Clamp>::GetRHI();
 	PassParameters->HairCountTexture = GraphBuilder.RegisterExternalTexture(VisibilityData.ViewHairCountTexture ? VisibilityData.ViewHairCountTexture : GSystemTextures.BlackDummy);
