Add opaque shadow mask permutation for hair, for evaluating the transmittance instead of the front depth.

This is more correct and allow to scale the shadow based on the hair density and thus attenuate shadow cast by groom when the hair strands are small/thin.

diff --git a/Engine/Shaders/Private/HairStrands/HairStrandsDeepShadowMask.usf b/Engine/Shaders/Private/HairStrands/HairStrandsDeepShadowMask.usf
--- a/Engine/Shaders/Private/HairStrands/HairStrandsDeepShadowMask.usf
+++ b/Engine/Shaders/Private/HairStrands/HairStrandsDeepShadowMask.usf
@@ -9,6 +9,12 @@
 #include "HairStrandsVisibilityCommon.ush"
 #include "HairStrandsVoxelPageCommon.ush"
 
+#include "HairStrandsCommon.ush"
+#include "HairStrandsVoxelPageCommon.ush"
+#include "HairStrandsDeepShadowCommon.ush"
+#include "HairStrandsDebugCommon.ush"
+#include "HairStrandsDeepTransmittanceCommon.ush"
+
 #define OPAQUE_MASK_TYPE_DEEP_SHADOW 0
 #define OPAQUE_MASK_TYPE_VOXEL 1
 #define OPAQUE_MASK_TYPE_VIRTUAL_VOXEL 2
@@ -28,6 +34,8 @@ uint				DeepShadow_IsWholeSceneLight;
 int2				DeepShadow_SlotResolution;
 int2				DeepShadow_SlotOffset;
 uint				DeepShadow_KernelType;
+float				DeepShadow_DepthBiasScale;
+float				DeepShadow_DensityScale;
 
 float4				Voxel_LightPosition;
 float3				Voxel_LightDirection;
@@ -40,7 +48,9 @@ Texture3D<uint> 	Voxel_DensityTexture;
 
 Texture2D<float4>	RayMarchMaskTexture;
 
-Texture2D<float4>	DeepShadow_FrontDepthTexture;
+Texture2D<float>	DeepShadow_FrontDepthTexture;
+Texture2D<float4>	DeepShadow_DomTexture;
+
 Texture2D<uint4>	CategorizationTexture;
 RWTexture2D<float4>	OutputColor;
 
@@ -51,6 +61,7 @@ SamplerState 		ShadowSampler;
 #define KERNEL_TYPE_4x4 1
 #define KERNEL_TYPE_Gaussian8 2
 #define KERNEL_TYPE_Gaussian16 3
+#define KERNEL_TYPE_GaussianTransmission8 4
 
 float3 TotWorldPosition(float2 UV, float Depth)
 {
@@ -82,7 +93,8 @@ bool IsCloser(float DepthA, float DepthB)
 
 struct FHairSamplingSettings
 {
-	Texture2D		ShadowDepthTexture;
+	Texture2D<float>ShadowDepthTexture;
+	Texture2D<float4>ShadowDomTexture;
 	SamplerState	ShadowDepthTextureSampler;
 	float2			ShadowAtlasInvResolution;
 	float2			ShadowSlotResolution;
@@ -285,6 +297,44 @@ float InternalGaussianFilter(float2 ShadowPosition, FHairSamplingSettings Settin
 	return AccShadow;
 }
 
+float InternalGaussianTransmissionFilter(float3 ShadowPosition, FHairSamplingSettings Settings, uint SampleCount)
+{
+	// Poisson disk position http://developer.download.nvidia.com/whitepapers/2008/PCSS_Integration.pdf
+	float2 PoissonDisk[16] =
+	{
+			float2(-0.94201624, -0.39906216),
+			float2(0.94558609, -0.76890725),
+			float2(-0.094184101, -0.92938870),
+			float2(0.34495938, 0.29387760),
+			float2(-0.91588581, 0.45771432),
+			float2(-0.81544232, -0.87912464),
+			float2(-0.38277543, 0.27676845),
+			float2(0.97484398, 0.75648379),
+			float2(0.44323325, -0.97511554),
+			float2(0.53742981, -0.47373420),
+			float2(-0.26496911, -0.41893023),
+			float2(0.79197514, 0.19090188),
+			float2(-0.24188840, 0.99706507),
+			float2(-0.81409955, 0.91437590),
+			float2(0.19984126, 0.78641367),
+			float2(0.14383161, -0.14100790)
+	};
+
+	const float SampleRadius = 3.f / Settings.ShadowSlotResolution.x;
+	float AccHairCount = 0;
+	for (uint SampleIt = 0; SampleIt < SampleCount; ++SampleIt)
+	{
+		const float2 Offset = PoissonDisk[SampleIt] * SampleRadius;
+		const float DepthBias = LAYER_DEPTH0 * DeepShadow_DepthBiasScale;
+		const float3 LightSpacePositionInAtlas = (ShadowPosition+float3(Offset,0)) * float3(DeepShadow_SlotResolution, 1) + float3(DeepShadow_SlotOffset, 0);
+		const float HairCount = SampleDOM_PCF2x2(LightSpacePositionInAtlas.xyz, DepthBias, Settings.ShadowDepthTexture, Settings.ShadowDomTexture);
+		AccHairCount += saturate(HairCount);
+	}
+
+	AccHairCount /= SampleCount;
+	return AccHairCount;
+}
+
 void MainPS(
 	FScreenVertexOutput Input,
 	out float4 OutColor : SV_Target0)
@@ -308,6 +358,7 @@ void MainPS(
 	#if OPAQUE_PCF == 1
 		FHairSamplingSettings ShadowSettings;
 		ShadowSettings.ShadowDepthTexture = DeepShadow_FrontDepthTexture;
+		ShadowSettings.ShadowDomTexture = DeepShadow_DomTexture;
 		ShadowSettings.ShadowDepthTextureSampler = ShadowSampler;
 		ShadowSettings.ShadowAtlasInvResolution = 1 / float2(ShadowAtlasResolution);
 		ShadowSettings.ShadowSlotResolution = DeepShadow_SlotResolution;
@@ -319,6 +370,8 @@ void MainPS(
 			Visibility = 1 - InternalGaussianFilter(LightSpacePosition.xy, ShadowSettings, 8);
 		#elif PERMUTATION_KERNEL_TYPE == KERNEL_TYPE_Gaussian16
 			Visibility = 1 - InternalGaussianFilter(LightSpacePosition.xy, ShadowSettings, 16);
+		#elif PERMUTATION_KERNEL_TYPE == KERNEL_TYPE_GaussianTransmission8
+			Visibility = 1 - InternalGaussianTransmissionFilter(LightSpacePosition.xyz, ShadowSettings, 8);
 		#else//PERMUTATION_KERNEL_TYPE == KERNEL_TYPE_2x2
 			Visibility = 1 - InternalManual2x2PCF(LightSpacePosition.xy, ShadowSettings);
 		#endif	
diff --git a/Engine/Source/Runtime/Renderer/Private/HairStrands/HairStrandsTransmittance.cpp b/Engine/Source/Runtime/Renderer/Private/HairStrands/HairStrandsTransmittance.cpp
--- a/Engine/Source/Runtime/Renderer/Private/HairStrands/HairStrandsTransmittance.cpp
+++ b/Engine/Source/Runtime/Renderer/Private/HairStrands/HairStrandsTransmittance.cpp
@@ -32,7 +32,7 @@ static uint32 GetDeepShadowKernelType() { return uint32(FMath::Max(0, GDeepShado
 static float GetDeepShadowKernelAperture() { return GDeepShadowKernelAperture; }
 
 static int32 GStrandHairShadowMaskKernelType = 2;
-static FAutoConsoleVariableRef GVarDeepShadowShadowMaskKernelType(TEXT("r.HairStrands.DeepShadow.ShadowMaskKernelType"), GStrandHairShadowMaskKernelType, TEXT("Set the kernel type for filtering shadow cast by hair on opaque geometry (0:2x2, 1:4x4, 2:Gaussian8, 3:Gaussian16). Default is 0"));
+static FAutoConsoleVariableRef GVarDeepShadowShadowMaskKernelType(TEXT("r.HairStrands.DeepShadow.ShadowMaskKernelType"), GStrandHairShadowMaskKernelType, TEXT("Set the kernel type for filtering shadow cast by hair on opaque geometry (0:2x2, 1:4x4, 2:Gaussian8, 3:Gaussian16, 4:Gaussian8 with transmittance. Default is 0"));
 
 static float GDeepShadowDensityScale = 2;	// Default is arbitrary, based on Mike asset
 static float GDeepShadowDepthBiasScale = 2;	// Default is arbitrary, based on content test
@@ -304,7 +304,7 @@ class FDeepShadowMaskPS : public FGlobalShader
 	SHADER_USE_PARAMETER_STRUCT(FDeepShadowMaskPS, FGlobalShader);
 
 	class FOpaqueMaskType : SHADER_PERMUTATION_INT("PERMUTATION_OPAQUEMASK_TYPE", FHairOpaqueMaskTypeCount);
-	class FKernelType : SHADER_PERMUTATION_INT("PERMUTATION_KERNEL_TYPE", 4);
+	class FKernelType : SHADER_PERMUTATION_INT("PERMUTATION_KERNEL_TYPE", 5);
 	using FPermutationDomain = TShaderPermutationDomain<FOpaqueMaskType, FKernelType>;
 
 	BEGIN_SHADER_PARAMETER_STRUCT(FParameters, )
@@ -316,6 +316,8 @@ class FDeepShadowMaskPS : public FGlobalShader
 		SHADER_PARAMETER(FIntPoint, DeepShadow_SlotResolution)
 		SHADER_PARAMETER(FMatrix, DeepShadow_WorldToLightTransform)
 		SHADER_PARAMETER(uint32, DeepShadow_IsWholeSceneLight)
+		SHADER_PARAMETER(float, DeepShadow_DepthBiasScale)
+		SHADER_PARAMETER(float, DeepShadow_DensityScale)
 
 		SHADER_PARAMETER(FVector4, Voxel_LightPosition)
 		SHADER_PARAMETER(FVector, Voxel_LightDirection)
@@ -329,6 +331,7 @@ class FDeepShadowMaskPS : public FGlobalShader
 
 		SHADER_PARAMETER_RDG_TEXTURE(Texture3D, Voxel_DensityTexture)
 		SHADER_PARAMETER_RDG_TEXTURE(Texture2D, DeepShadow_FrontDepthTexture)
+		SHADER_PARAMETER_RDG_TEXTURE(Texture2D, DeepShadow_DomTexture)
 		SHADER_PARAMETER_RDG_TEXTURE(Texture2D, CategorizationTexture)
 		SHADER_PARAMETER_SAMPLER(SamplerState, LinearSampler)
 		SHADER_PARAMETER_SAMPLER(SamplerState, ShadowSampler)
@@ -352,7 +355,10 @@ struct FDeepShadowOpaqueParams
 	FMatrix			DeepShadow_WorldToLightTransform;
 	FIntRect		DeepShadow_AtlasRect;
 	FRDGTextureRef	DeepShadow_FrontDepthTexture = nullptr;
+	FRDGTextureRef	DeepShadow_LayerTexture = nullptr;
 	bool			DeepShadow_bIsWholeSceneLight = false;
+	float			DeepShadow_DepthBiasScale = 1;
+	float			DeepShadow_DensityScale = 1;
 
 	FVector			Voxel_LightDirection = FVector::ZeroVector;
 	FVector4		Voxel_LightPosition = FVector4(0, 0, 0, 0);
@@ -382,10 +388,13 @@ static void AddDeepShadowOpaqueMaskPass(
 	Parameters->SceneTextures = SceneTextures;
 	Parameters->DeepShadow_WorldToLightTransform = Params.DeepShadow_WorldToLightTransform;
 	Parameters->DeepShadow_FrontDepthTexture = Params.DeepShadow_FrontDepthTexture;
+	Parameters->DeepShadow_DomTexture = Params.DeepShadow_LayerTexture;
 	Parameters->CategorizationTexture = Params.CategorizationTexture;
 	Parameters->LinearSampler = TStaticSamplerState<SF_Bilinear, AM_Clamp, AM_Clamp, AM_Clamp>::GetRHI();
 	Parameters->ShadowSampler = TStaticSamplerState<SF_Point, AM_Clamp, AM_Clamp, AM_Clamp>::GetRHI();
 	Parameters->DeepShadow_IsWholeSceneLight = Params.DeepShadow_bIsWholeSceneLight ? 1 : 0;
+	Parameters->DeepShadow_DepthBiasScale = Params.DeepShadow_DepthBiasScale;
+	Parameters->DeepShadow_DensityScale = Params.DeepShadow_DensityScale;
 	Parameters->RenderTargets[0] = FRenderTargetBinding(OutShadowMask, ERenderTargetLoadAction::ELoad);
 	Parameters->DeepShadow_SlotOffset = FIntPoint(Params.DeepShadow_AtlasRect.Min.X, Params.DeepShadow_AtlasRect.Min.Y);
 	Parameters->DeepShadow_SlotResolution = FIntPoint(Params.DeepShadow_AtlasRect.Max.X - Params.DeepShadow_AtlasRect.Min.X, Params.DeepShadow_AtlasRect.Max.Y - Params.DeepShadow_AtlasRect.Min.Y);
@@ -425,7 +434,7 @@ static void AddDeepShadowOpaqueMaskPass(
 
 	FDeepShadowMaskPS::FPermutationDomain PermutationVector;
 	PermutationVector.Set<FDeepShadowMaskPS::FOpaqueMaskType>(HairOpaqueMaskType);
-	PermutationVector.Set<FDeepShadowMaskPS::FKernelType>(FMath::Clamp(GStrandHairShadowMaskKernelType, 0, 3));
+	PermutationVector.Set<FDeepShadowMaskPS::FKernelType>(FMath::Clamp(GStrandHairShadowMaskKernelType, 0, 4));
 
 	const FIntPoint OutputResolution = SceneTextures.SceneDepthBuffer->Desc.Extent;
 	TShaderMapRef<FPostProcessVS> VertexShader(View.ShaderMap);
@@ -685,6 +694,7 @@ static void RenderHairStrandsShadowMask(
 
 			bHasDeepShadow = true;
 			FRDGTextureRef DeepShadowDepth = GraphBuilder.RegisterExternalTexture(DomData.DepthTexture, TEXT("DeepShadowDepthTexture"));
+			FRDGTextureRef DeepShadowLayers = GraphBuilder.RegisterExternalTexture(DomData.LayersTexture, TEXT("DeepShadowLayerTexture"));
 			const bool bIsWholeSceneLight = LightSceneInfo->Proxy->GetLightType() == LightType_Directional;
 
 			FDeepShadowOpaqueParams Params;
@@ -692,7 +702,10 @@ static void RenderHairStrandsShadowMask(
 			Params.DeepShadow_WorldToLightTransform = DomData.WorldToLightTransform;
 			Params.DeepShadow_AtlasRect = DomData.AtlasRect;
 			Params.DeepShadow_FrontDepthTexture = DeepShadowDepth;
+			Params.DeepShadow_LayerTexture = DeepShadowLayers;
 			Params.DeepShadow_bIsWholeSceneLight = bIsWholeSceneLight;
+			Params.DeepShadow_DepthBiasScale = GetDeepShadowDepthBiasScale();
+			Params.DeepShadow_DensityScale = GetDeepShadowDensityScale();
 
 			AddDeepShadowOpaqueMaskPass(
 				GraphBuilder,
