Reduce deep shadow bias.

Large value are no longer need due to various fixes which have been done during the release.

diff --git a/Engine/Source/Runtime/Renderer/Private/HairStrands/HairStrandsTransmittance.cpp b/Engine/Source/Runtime/Renderer/Private/HairStrands/HairStrandsTransmittance.cpp
--- a/Engine/Source/Runtime/Renderer/Private/HairStrands/HairStrandsTransmittance.cpp
+++ b/Engine/Source/Runtime/Renderer/Private/HairStrands/HairStrandsTransmittance.cpp
@@ -35,7 +35,7 @@ static int32 GStrandHairShadowMaskKernelType = 2;
 static FAutoConsoleVariableRef GVarDeepShadowShadowMaskKernelType(TEXT("r.HairStrands.DeepShadow.ShadowMaskKernelType"), GStrandHairShadowMaskKernelType, TEXT("Set the kernel type for filtering shadow cast by hair on opaque geometry (0:2x2, 1:4x4, 2:Gaussian8, 3:Gaussian16, 4:Gaussian8 with transmittance. Default is 0"));
 
 static float GDeepShadowDensityScale = 2;	// Default is arbitrary, based on Mike asset
-static float GDeepShadowDepthBiasScale = 2;	// Default is arbitrary, based on content test
+static float GDeepShadowDepthBiasScale = 0.05;
 static FAutoConsoleVariableRef CVarDeepShadowDensityScale(TEXT("r.HairStrands.DeepShadow.DensityScale"), GDeepShadowDensityScale, TEXT("Set density scale for compensating the lack of hair fiber in an asset"));
 static FAutoConsoleVariableRef CVarDeepShadowDepthBiasScale(TEXT("r.HairStrands.DeepShadow.DepthBiasScale"), GDeepShadowDepthBiasScale, TEXT("Set depth bias scale for transmittance computation"));
 
