Fix point light shadow with hair rendering.

diff --git a/Engine/Source/Runtime/Renderer/Private/ShadowRendering.h b/Engine/Source/Runtime/Renderer/Private/ShadowRendering.h
--- a/Engine/Source/Runtime/Renderer/Private/ShadowRendering.h
+++ b/Engine/Source/Runtime/Renderer/Private/ShadowRendering.h
@@ -931,10 +931,15 @@ public:
 		{
 			SetTextureParameter(RHICmdList, ShaderRHI, HairCategorizationTexture, HairVisibilityData->CategorizationTexture->GetRenderTargetItem().ShaderResourceTexture);
 
-			FVector4 Near = View.ViewMatrices.GetProjectionMatrix().TransformFVector4(FVector4(0, 0, ShadowInfo->CascadeSettings.SplitNear));
-			FVector4 Far = View.ViewMatrices.GetProjectionMatrix().TransformFVector4(FVector4(0, 0, ShadowInfo->CascadeSettings.SplitFar));
-			const float DeviceZNear = Near.Z / Near.W;
-			const float DeviceZFar = Far.Z / Far.W;
+			float DeviceZNear = 1;
+			float DeviceZFar = 0;
+			if (ShadowInfo->bDirectionalLight)
+			{
+				FVector4 Near = View.ViewMatrices.GetProjectionMatrix().TransformFVector4(FVector4(0, 0, ShadowInfo->CascadeSettings.SplitNear));
+				FVector4 Far = View.ViewMatrices.GetProjectionMatrix().TransformFVector4(FVector4(0, 0, ShadowInfo->CascadeSettings.SplitFar));
+				DeviceZNear = Near.Z / Near.W;
+				DeviceZFar = Far.Z / Far.W;
+			}
 
 			FVector2D SliceNearAndFarDepth;
 			SliceNearAndFarDepth.X = DeviceZNear;
