Initial base implementation of a new ray traced reflection technique: ray traced deferred reflections. This is an experimental alternative reflection code path that favors performance over features. It can be enabled via r.RayTracing.Reflections.ExperimentalDeferred (default 0).

The aim is to create a minimal implementation of the main RT reflection technique with material sorting (r.RayTracing.Reflections.SortMaterials 1), without any features that aren't absolutely necessary to avoid as much overhead as possible in the ray generation shader.

The algorithm consists of several separate stages:
- Generate reflection rays based on GBuffer (sorted in tiles by direction). Sorting may be optional in the future, based on performance measurements.
- Trace screen space reflection rays and output validity mask to avoid tracing/shading full rays [TODO; currently always tracing full rays]
- Trace reflection rays using lightweight RayGen shader and output material IDs
- Sort material IDs
- Execute material shaders and produce "Reflection GBuffer"
- Apply lighting to produce the final reflection buffer [TODO; currently done in material eval RGS]

Other features that are currently not implemented, but may be in the future:
- Roughness threshold
- Forced mirror-like reflections (similar to SSR low quality profile)
- Alpha masked materials
- Reflection capture for multi-bounce fallback
- Shadow maps instead of ray traced shadows

Features that will never be supported due to performance:
- Multi-bounce
- Multi-SPP
- Clearcoat
- Translucency

diff --git a/Engine/Shaders/Private/RayTracing/RayTracingDeferredMaterials.ush b/Engine/Shaders/Private/RayTracing/RayTracingDeferredMaterials.ush
new file mode 100644
--- /dev/null
+++ b/Engine/Shaders/Private/RayTracing/RayTracingDeferredMaterials.ush
@@ -0,0 +1,29 @@
+// Copyright Epic Games, Inc. All Rights Reserved.
+
+#pragma once
+
+#include "RayTracingCommon.ush"
+
+#ifdef OVERRIDE_RAYTRACINGDEFERREDMATERIALS_USH
+#include "/Platform/Private/RayTracingDeferredMaterials.ush"
+#else // OVERRIDE_RAYTRACINGDEFERREDMATERIALS_USH
+
+struct FRayIntersectionBookmark
+{
+	uint Data[2];
+};
+
+#define TraceDeferredMaterialGatherRay(TLAS, RayFlags, InstanceInclusionMask, Ray, OutBookmark, InOutPayload){\
+	TraceRay(TLAS, RayFlags, InstanceInclusionMask, RAY_TRACING_SHADER_SLOT_MATERIAL, RAY_TRACING_NUM_SHADER_SLOTS, 0, Ray, InOutPayload);\
+	OutBookmark.Data[0] = asuint(InOutPayload.HitT);\
+	OutBookmark.Data[1] = asuint(InOutPayload.HitT);\
+}
+
+#define TraceDeferredMaterialShadingRay(TLAS, RayFlags, InstanceInclusionMask, Ray, Bookmark, InOutPayload){\
+	RayDesc BookmarkRay = Ray;\
+	BookmarkRay.TMin = asfloat(Bookmark.Data[0]) - 0.5f;\
+	BookmarkRay.TMax = asfloat(Bookmark.Data[1]) + 0.5f;\
+	TraceRay(TLAS, RayFlags, InstanceInclusionMask, RAY_TRACING_SHADER_SLOT_MATERIAL, RAY_TRACING_NUM_SHADER_SLOTS, 0, BookmarkRay, InOutPayload);\
+}
+
+#endif // OVERRIDE_RAYTRACINGDEFERREDMATERIALS_USH
diff --git a/Engine/Shaders/Private/RayTracing/RayTracingDeferredReflections.usf b/Engine/Shaders/Private/RayTracing/RayTracingDeferredReflections.usf
new file mode 100644
--- /dev/null
+++ b/Engine/Shaders/Private/RayTracing/RayTracingDeferredReflections.usf
@@ -0,0 +1,166 @@
+// Copyright Epic Games, Inc. All Rights Reserved.
+
+// EDeferredMaterialMode
+#define DEFERRED_MATERIAL_MODE_NONE   0
+#define DEFERRED_MATERIAL_MODE_GATHER 1
+#define DEFERRED_MATERIAL_MODE_SHADE  2
+
+#ifndef DIM_DEFERRED_MATERIAL_MODE
+	#define DIM_DEFERRED_MATERIAL_MODE DEFERRED_MATERIAL_MODE_NONE
+#endif
+
+#include "../Common.ush"
+
+#define SUPPORT_CONTACT_SHADOWS		0
+#define USE_SOURCE_TEXTURE			1
+#define USE_SOURCE_TEXTURE_ARRAY	1
+
+#define LTCMatTexture				RaytracingLightsDataPacked.LTCMatTexture
+#define LTCMatSampler				RaytracingLightsDataPacked.LTCMatSampler
+#define LTCAmpTexture				RaytracingLightsDataPacked.LTCAmpTexture
+#define LTCAmpSampler				RaytracingLightsDataPacked.LTCAmpSampler
+#define PreIntegratedGF				ReflectionStruct.PreIntegratedGF			
+#define PreIntegratedGFSampler		ReflectionStruct.PreIntegratedGFSampler	
+
+#include "../DeferredShadingCommon.ush"
+#include "../DeferredLightingCommon.ush"
+#include "../ReflectionEnvironmentShared.ush"
+#include "../Montecarlo.ush"
+#include "../PathTracing/Utilities/PathTracingRandomSequence.ush" 
+#include "../HeightFogCommon.ush"
+#include "../ScreenSpaceDenoise/SSDPublic.ush"
+#include "../SobolRandom.ush"
+#include "../SceneTextureParameters.ush"
+
+#include "RayTracingCommon.ush"
+#include "RayTracingDeferredShadingCommon.ush"
+#include "RayTracingLightingCommon.ush"
+#include "RayTracingReflectionsCommon.ush"
+#include "RayTracingDeferredMaterials.ush"
+
+uint2 RayTracingResolution;
+uint2 TileAlignedResolution;
+float ReflectionMaxNormalBias;
+
+RWStructuredBuffer<FSortedReflectionRay>        RayBuffer;
+RWStructuredBuffer<FDeferredMaterialPayload>    MaterialBuffer;
+RWTexture2D<float4>                             ColorOutput;
+RWTexture2D<float>                              RayHitDistanceOutput;
+RaytracingAccelerationStructure                 TLAS;
+
+// Buffer of ray intersections aquired during Gather phase and 
+// used during to reduce ray intersection overhead during Shading.
+RWStructuredBuffer<FRayIntersectionBookmark>    BookmarkBuffer;
+
+RAY_TRACING_ENTRY_RAYGEN(RayTracingDeferredReflectionsRGS)
+{
+	const uint DispatchThreadId = DispatchRaysIndex().x; // This shader is dispatched in 1D configuration
+
+	FDeferredMaterialPayload DeferredMaterialPayload = (FDeferredMaterialPayload)0;
+	DeferredMaterialPayload.SortKey = RAY_TRACING_DEFERRED_MATERIAL_KEY_INVALID;
+
+#if DIM_DEFERRED_MATERIAL_MODE == DEFERRED_MATERIAL_MODE_GATHER
+
+	FSortedReflectionRay SortedRay = RayBuffer[DispatchThreadId];
+	DeferredMaterialPayload.PixelCoordinates = DispatchThreadId; // Store the sorted ray index, not packed 2D pixel coordinates
+	const uint2 PixelPos = UnpackPixelCoordinates(SortedRay.PixelCoordinates);
+	if (all(PixelPos < RayTracingResolution))
+	{
+		RayDesc Ray;
+		Ray.TMin      = 0;
+		Ray.TMax      = 1e9;
+		Ray.Origin    = SortedRay.Origin;
+		Ray.Direction = SortedRay.Direction;
+
+		DeferredMaterialPayload.SortKey = RAY_TRACING_DEFERRED_MATERIAL_KEY_RAY_MISS;
+		FRayIntersectionBookmark Bookmark = (FRayIntersectionBookmark)0;
+
+		TraceDeferredMaterialGatherRay(TLAS,
+			RAY_FLAG_FORCE_OPAQUE, // Masked and transparent materials are currently not supported
+			RAY_TRACING_MASK_OPAQUE,
+			Ray,
+			Bookmark,
+			DeferredMaterialPayload
+		);
+
+		BookmarkBuffer[DispatchThreadId] = Bookmark;
+	}
+
+	MaterialBuffer[DispatchThreadId] = DeferredMaterialPayload;
+
+#elif DIM_DEFERRED_MATERIAL_MODE == DEFERRED_MATERIAL_MODE_SHADE
+
+	DeferredMaterialPayload = MaterialBuffer[DispatchThreadId];
+
+	// Linear ray index is stored in PixelCoordinates (see DEFERRED_MATERIAL_MODE_GATHER above)
+	const uint SortedRayIndex = DeferredMaterialPayload.PixelCoordinates;
+
+	const FSortedReflectionRay SortedRay = RayBuffer[SortedRayIndex];
+	const uint2 PixelPos = UnpackPixelCoordinates(SortedRay.PixelCoordinates);
+	if (all(PixelPos < RayTracingResolution) && DeferredMaterialPayload.SortKey < RAY_TRACING_DEFERRED_MATERIAL_KEY_RAY_MISS)
+	{
+		RayDesc Ray;
+		Ray.TMin      = 0;
+		Ray.TMax      = 1e9;
+		Ray.Origin    = SortedRay.Origin;
+		Ray.Direction = SortedRay.Direction;
+
+		FRayIntersectionBookmark Bookmark = BookmarkBuffer[DeferredMaterialPayload.PixelCoordinates];
+
+		FPackedMaterialClosestHitPayload PackedPayload = (FPackedMaterialClosestHitPayload)0;
+
+		TraceDeferredMaterialShadingRay(
+			TLAS,
+			RAY_FLAG_FORCE_OPAQUE,
+			RAY_TRACING_MASK_OPAQUE,
+			Ray,
+			Bookmark,
+			PackedPayload);
+
+		float4 ResultColor = (float4)0;
+		float ResultDistance = -1;
+
+		if (PackedPayload.IsHit())
+		{
+			ResultDistance = PackedPayload.HitT;
+
+			bool bShouldDoDirectLighting               = true;
+			bool bShouldDoEmissiveAndIndirectLighting  = true;
+			bool bTopLayerRayTraceSkyLightContribution = false;
+			bool bDecoupleSampleGeneration             = false;
+			uint ReflectedShadowsType                  = 1; // Hard shadows
+			
+			RandomSequence RandSequence; // Not actually used, since hard shadows are forced
+			RandomSequence_Initialize(RandSequence, DispatchThreadId, 0);
+
+			float3 LightingViewDirection = Ray.Direction;
+			float3 RayHitWorldPos = Ray.Origin + Ray.Direction * PackedPayload.HitT;
+			FRayCone RayCone = (FRayCone)0; // Not used
+
+			AccumulateResults(
+				PackedPayload,
+				RayHitWorldPos,
+				LightingViewDirection,
+				TLAS,
+				RandSequence,
+				PixelPos,
+				ReflectionMaxNormalBias,
+				ReflectedShadowsType,
+				bShouldDoDirectLighting,
+				bShouldDoEmissiveAndIndirectLighting,
+				bTopLayerRayTraceSkyLightContribution,
+				bDecoupleSampleGeneration,
+				RayCone,
+				ResultColor.rgb);
+
+			ResultColor.rgb *= View.PreExposure;
+			ResultColor.a = 1;
+		}
+
+		ColorOutput[PixelPos] 			= ResultColor;
+		RayHitDistanceOutput[PixelPos] 	= ResultDistance;
+	}
+
+#endif // DIM_DEFERRED_MATERIAL_MODE
+}
+
diff --git a/Engine/Shaders/Private/RayTracing/RayTracingReflectionsCommon.ush b/Engine/Shaders/Private/RayTracing/RayTracingReflectionsCommon.ush
--- a/Engine/Shaders/Private/RayTracing/RayTracingReflectionsCommon.ush
+++ b/Engine/Shaders/Private/RayTracing/RayTracingReflectionsCommon.ush
@@ -19,3 +19,25 @@ struct FImaginaryReflectionGBufferData
 	// Required for averaging to account for the fact that not all samples may hit a valid imaginary surface
 	uint ValidSamples;
 };
+
+struct FSortedReflectionRay
+{
+	float3	Origin;
+	uint	PixelCoordinates;
+	float3	Direction;
+	uint	DebugSortKey;
+};
+
+uint PackPixelCoordinates(uint2 PixelCoordinates)
+{
+	return (PixelCoordinates.x & 0xFFFF)
+	     | (PixelCoordinates.y << 16);
+}
+
+uint2 UnpackPixelCoordinates(uint PixelCoordinates)
+{
+	return uint2(
+		PixelCoordinates & 0xFFFF,
+		PixelCoordinates >> 16);
+}
+
diff --git a/Engine/Shaders/Private/RayTracing/RayTracingReflectionsGenerateRaysCS.usf b/Engine/Shaders/Private/RayTracing/RayTracingReflectionsGenerateRaysCS.usf
new file mode 100644
--- /dev/null
+++ b/Engine/Shaders/Private/RayTracing/RayTracingReflectionsGenerateRaysCS.usf
@@ -0,0 +1,184 @@
+// Copyright Epic Games, Inc. All Rights Reserved.
+
+#include "../Common.ush"
+#include "../DeferredShadingCommon.ush"
+#include "../SceneTextureParameters.ush"
+#include "../ScreenSpaceDenoise/SSDPublic.ush"
+#include "../MortonCode.ush"
+
+#include "RayTracingCommon.ush"
+#include "RayTracingReflectionsCommon.ush"
+
+#define REFLECTION_RAY_DIRECTION_SORT_TILE_SIZE 32
+
+#define NUM_DIRECTION_BINS_X REFLECTION_RAY_DIRECTION_SORT_TILE_SIZE
+#define NUM_DIRECTION_BINS_Y REFLECTION_RAY_DIRECTION_SORT_TILE_SIZE
+#define NUM_DIRECTION_BINS (NUM_DIRECTION_BINS_X*NUM_DIRECTION_BINS_Y)
+
+#if THREADGROUP_SIZE != 1024
+#error This shader requires group size of exactly 1024 threads
+#endif
+
+#if THREADGROUP_SIZE != REFLECTION_RAY_DIRECTION_SORT_TILE_SIZE*REFLECTION_RAY_DIRECTION_SORT_TILE_SIZE
+#error Total number of rays peer tile must be equal to thread group size
+#endif
+
+#if THREADGROUP_SIZE < NUM_DIRECTION_BINS
+#error Total number of bins must be less or equal to thread group size
+#endif
+
+#ifndef DIM_WAVE_OPS
+	#define DIM_WAVE_OPS 0
+#endif
+
+#if DIM_WAVE_OPS
+#define MIN_SUPPORTED_WAVE_SIZE 32
+#define NUM_WAVE_SUMS ((NUM_DIRECTION_BINS+MIN_SUPPORTED_WAVE_SIZE-1) / MIN_SUPPORTED_WAVE_SIZE)
+groupshared uint WaveSums[NUM_WAVE_SUMS];
+#endif
+
+groupshared uint Bins[NUM_DIRECTION_BINS];
+
+// Inputs
+uint2 RayTracingResolution;
+uint2 TileAlignedResolution;
+float ReflectionMaxNormalBias;
+
+// Outputs
+RWStructuredBuffer<FSortedReflectionRay> RayBuffer;
+
+// #todo: refactor ApplyCameraRelativeDepthBias() in RayTracingCommon.ush so that it could be used in this compute shader.
+float ComputeBiasForPosition(float3 V)
+{
+	V = abs(V);
+	return max(max(V.x, V.y), V.z) * 0.01;
+}
+
+uint RayDirectionToBin(float3 Direction, uint2 NumBins)
+{
+	float2 OctDirection = UnitVectorToOctahedron(Direction);
+	uint2 Bin2D = min((uint2)floor(OctDirection * NumBins), NumBins-1);
+	return MortonEncode(Bin2D);
+}
+
+
+[numthreads(THREADGROUP_SIZE, 1, 1)]
+void GenerateReflectionRaysCS(
+	uint DispatchThreadId : SV_DispatchThreadID,
+	uint GroupThreadId : SV_GroupThreadID,
+	uint GroupId : SV_GroupID)
+{
+	if (GroupThreadId < NUM_DIRECTION_BINS)
+	{
+		Bins[GroupThreadId] = 0;
+	}
+
+	GroupMemoryBarrierWithGroupSync();
+
+	const uint2 TileSize         = uint2(REFLECTION_RAY_DIRECTION_SORT_TILE_SIZE, REFLECTION_RAY_DIRECTION_SORT_TILE_SIZE);
+	const uint  TileIndex        = GroupId;
+	const uint  RayIndexInTile   = GroupThreadId;
+	const uint2 NumTiles         = TileAlignedResolution / TileSize;
+	const uint2 TileBasePixelPos = uint2(TileIndex % NumTiles.x, TileIndex / NumTiles.x) * TileSize;
+	const uint2 PixelPos         = TileBasePixelPos + uint2(RayIndexInTile % TileSize.x, RayIndexInTile / TileSize.x);
+	const float2 ScreenPos       = ViewportUVToScreenPos((PixelPos.xy - View.ViewRectMin.xy) * View.ViewSizeAndInvSize.zw);
+	const float2 UV              = PixelPos.xy * View.BufferSizeAndInvSize.zw;
+	const FGBufferData GBuffer   = GetGBufferDataFromSceneTextures(UV);
+	const float SceneDepth       = GBuffer.Depth;
+
+	const float3 PositionTranslatedWorld = mul(float4(ScreenPos * SceneDepth, SceneDepth, 1), View.ScreenToTranslatedWorld).xyz;
+	const float3 V                       = normalize(View.TranslatedWorldCameraOrigin - PositionTranslatedWorld);
+
+
+#if GBUFFER_HAS_TANGENT
+	ModifyGGXAnisotropicNormalRoughness(GBuffer.WorldTangent, GBuffer.Anisotropy, GBuffer.Roughness, GBuffer.WorldNormal, V);
+#endif
+	
+	float2 E = Rand1SPPDenoiserInput(PixelPos);
+	
+	const bool bOutputForDenoiser = true; // #todo: pass this in via constants when denoiser is on
+	if (bOutputForDenoiser)
+	{
+		E.y *= 1 - GGX_IMPORTANT_SAMPLE_BIAS;
+	}
+	
+	float3x3 TangentBasis = GetTangentBasis(GBuffer.WorldNormal);
+	float3 TangentV       = mul(TangentBasis, V);
+
+	float3 L = (float3)0;
+
+	bool bGlossy = true; // #todo: pass this in via constants (change based on required quality, similar to SSR)
+	if (bGlossy)
+	{
+		float3 H = mul(ImportanceSampleVisibleGGX(UniformSampleDisk(E), Pow4(GBuffer.Roughness), TangentV).xyz, TangentBasis);
+		L = 2 * dot(V, H) * H - V;
+	}
+	else
+	{
+		L = reflect(-V, GBuffer.WorldNormal);
+	}
+
+	FSortedReflectionRay Ray;
+	Ray.Origin = View.WorldCameraOrigin + PositionTranslatedWorld;
+	Ray.Origin += GBuffer.WorldNormal * ComputeBiasForPosition(PositionTranslatedWorld);	
+	Ray.Direction = L;
+
+	uint Bin = RayDirectionToBin(Ray.Direction, uint2(NUM_DIRECTION_BINS_X,NUM_DIRECTION_BINS_Y));
+	if (any(PixelPos >= RayTracingResolution))
+	{
+		Bin = NUM_DIRECTION_BINS - 1;
+	}
+
+	uint SlotForThisThread = 0;
+	InterlockedAdd(Bins[Bin], 1, SlotForThisThread);
+
+	GroupMemoryBarrierWithGroupSync();
+
+#if DIM_WAVE_OPS
+	{
+		const uint WaveIndex = GroupThreadId / WaveGetLaneCount();
+		const uint LaneIndex = WaveGetLaneIndex();
+		uint Value = 0;
+		if (GroupThreadId < NUM_DIRECTION_BINS)
+		{
+			Value = Bins[GroupThreadId];
+			uint ThisWaveSum = WaveActiveSum(Value);
+			if (LaneIndex == 0)
+			{
+				WaveSums[WaveIndex] = ThisWaveSum;
+			}
+		}
+		GroupMemoryBarrierWithGroupSync();
+		if (WaveIndex == 0 && LaneIndex < NUM_WAVE_SUMS)
+		{
+			WaveSums[LaneIndex] = WavePrefixSum(WaveSums[LaneIndex]);
+		}
+		GroupMemoryBarrierWithGroupSync();
+		if (GroupThreadId < NUM_DIRECTION_BINS)
+		{
+			Bins[GroupThreadId] = WaveSums[WaveIndex] + WavePrefixSum(Value);
+		}
+	}
+#else // DIM_WAVE_OPS
+	if (GroupThreadId == 0)
+	{
+		uint Counter = 0;
+		for (uint i = 0; i < NUM_DIRECTION_BINS; i++)
+		{
+			uint NextCounter = Counter + Bins[i];
+			Bins[i] = Counter;
+			Counter = NextCounter;
+		}
+	}
+#endif // DIM_WAVE_OPS
+
+	GroupMemoryBarrierWithGroupSync();
+
+	Ray.PixelCoordinates = PackPixelCoordinates(PixelPos);
+	Ray.DebugSortKey = Bin;
+
+	uint StoreIndex = GroupId * (TileSize.x*TileSize.y) + Bins[Bin] + SlotForThisThread;
+
+	RayBuffer[StoreIndex] = Ray;
+}
+
diff --git a/Engine/Source/Runtime/Renderer/Private/DeferredShadingRenderer.h b/Engine/Source/Runtime/Renderer/Private/DeferredShadingRenderer.h
--- a/Engine/Source/Runtime/Renderer/Private/DeferredShadingRenderer.h
+++ b/Engine/Source/Runtime/Renderer/Private/DeferredShadingRenderer.h
@@ -524,6 +524,12 @@ private:
 		float ResolutionFraction,
 		IScreenSpaceDenoiser::FReflectionsInputs* OutDenoiserInputs);
 
+	void RenderRayTracingDeferredReflections(
+		FRDGBuilder& GraphBuilder,
+		const FSceneTextureParameters& SceneTextures,
+		const FViewInfo& View,
+		IScreenSpaceDenoiser::FReflectionsInputs* OutDenoiserInputs);
+
 	void RenderRayTracingShadows(
 		FRDGBuilder& GraphBuilder,
 		const FSceneTextureParameters& SceneTextures,
@@ -634,6 +640,7 @@ private:
 
 	// #dxr_todo: UE-72565: refactor ray tracing effects to not be member functions of DeferredShadingRenderer. Register each effect at startup and just loop over them automatically
 	static void PrepareRayTracingReflections(const FViewInfo& View, const FScene& Scene, TArray<FRHIRayTracingShader*>& OutRayGenShaders);
+	static void PrepareRayTracingDeferredReflections(const FViewInfo& View, const FScene& Scene, TArray<FRHIRayTracingShader*>& OutRayGenShaders);
 	static void PrepareRayTracingShadows(const FViewInfo& View, TArray<FRHIRayTracingShader*>& OutRayGenShaders);
 	static void PrepareRayTracingAmbientOcclusion(const FViewInfo& View, TArray<FRHIRayTracingShader*>& OutRayGenShaders);
 	static void PrepareRayTracingSkyLight(const FViewInfo& View, TArray<FRHIRayTracingShader*>& OutRayGenShaders);
diff --git a/Engine/Source/Runtime/Renderer/Private/RayTracing/RayTracingDeferredReflections.cpp b/Engine/Source/Runtime/Renderer/Private/RayTracing/RayTracingDeferredReflections.cpp
new file mode 100644
--- /dev/null
+++ b/Engine/Source/Runtime/Renderer/Private/RayTracing/RayTracingDeferredReflections.cpp
@@ -0,0 +1,317 @@
+// Copyright Epic Games, Inc. All Rights Reserved.
+
+/**
+ * Implements an experimental ray tracing reflection rendering algorithm based on ray and material sorting.
+ * The algorithm consists of several separate stages:
+ * - Generate reflection rays based on GBuffer (sorted in tiles by direction). Sorting may be optional in the future, based on performance measurements.
+ * - Trace screen space reflection rays and output validity mask to avoid tracing/shading full rays [TODO; currently always tracing full rays]
+ * - Trace reflection rays using lightweight RayGen shader and output material IDs
+ * - Sort material IDs
+ * - Execute material shaders and produce "Reflection GBuffer"
+ * - Apply lighting to produce the final reflection buffer [TODO; currently done in material eval RGS]
+ * 
+ * Other features that are currently not implemented, but may be in the future:
+ * - Roughness threshold
+ * - Forced mirror-like reflections (similar to SSR low quality profile)
+ * - Alpha masked materials
+ * - Reflection capture for multi-bounce fallback
+ * - Shadow maps instead of ray traced shadows
+ * 
+ * Features that will never be supported due to performance:
+ * - Multi-bounce
+ * - Multi-SPP
+ * - Clearcoat
+ * - Translucency
+ **/
+
+#include "RayTracing/RayTracingLighting.h"
+#include "RayTracing/RayTracingDeferredMaterials.h"
+#include "RendererPrivate.h"
+#include "GlobalShader.h"
+#include "DeferredShadingRenderer.h"
+#include "SceneTextureParameters.h"
+#include "ReflectionEnvironment.h"
+
+#if RHI_RAYTRACING
+
+namespace 
+{
+	struct FSortedReflectionRay
+	{
+		float  Origin[3];
+		uint32 PixelCoordinates; // X in low 16 bits, Y in high 16 bits
+		float  Direction[3];
+		uint32 DebugSortKey;
+	};
+
+	struct FRayIntersectionBookmark
+	{
+		uint32 Data[2];
+	};
+} // anon namespace
+
+class FGenerateReflectionRaysCS : public FGlobalShader
+{
+	DECLARE_GLOBAL_SHADER(FGenerateReflectionRaysCS);
+	SHADER_USE_PARAMETER_STRUCT(FGenerateReflectionRaysCS, FGlobalShader);
+
+	class FWaveOps : SHADER_PERMUTATION_BOOL("DIM_WAVE_OPS");
+	using FPermutationDomain = TShaderPermutationDomain<FWaveOps>;
+
+	BEGIN_SHADER_PARAMETER_STRUCT(FParameters, )
+	SHADER_PARAMETER(FIntPoint, RayTracingResolution)
+	SHADER_PARAMETER(FIntPoint, TileAlignedResolution)
+	SHADER_PARAMETER(float, ReflectionMaxNormalBias)
+	SHADER_PARAMETER_STRUCT_REF(FViewUniformShaderParameters, ViewUniformBuffer)
+	SHADER_PARAMETER_STRUCT_INCLUDE(FSceneTextureParameters, SceneTextures)
+	SHADER_PARAMETER_RDG_BUFFER_UAV(RWStructuredBuffer<FSortedReflectionRay>, RayBuffer)
+	END_SHADER_PARAMETER_STRUCT()
+
+	static bool ShouldCompilePermutation(const FGlobalShaderPermutationParameters& Parameters)
+	{
+		FPermutationDomain PermutationVector(Parameters.PermutationId);
+		if (PermutationVector.Get<FWaveOps>() && !RHISupportsWaveOperations(Parameters.Platform))
+		{
+			return false;
+		}
+
+		return ShouldCompileRayTracingShadersForProject(Parameters.Platform);
+	}
+
+	static uint32 GetGroupSize()
+	{
+		return 1024; // this shader generates rays and sorts them in 32x32 tiles using LDS
+	}
+
+	static void ModifyCompilationEnvironment(const FGlobalShaderPermutationParameters& Parameters, FShaderCompilerEnvironment& OutEnvironment)
+	{
+		FGlobalShader::ModifyCompilationEnvironment(Parameters, OutEnvironment);
+		OutEnvironment.SetDefine(TEXT("THREADGROUP_SIZE"), GetGroupSize());
+
+		FPermutationDomain PermutationVector(Parameters.PermutationId);
+
+		if (PermutationVector.Get<FWaveOps>())
+		{
+			OutEnvironment.CompilerFlags.Add(CFLAG_WaveOperations);
+		}
+	}
+
+};
+IMPLEMENT_GLOBAL_SHADER(FGenerateReflectionRaysCS, "/Engine/Private/RayTracing/RayTracingReflectionsGenerateRaysCS.usf", "GenerateReflectionRaysCS", SF_Compute);
+
+class FRayTracingDeferredReflectionsRGS : public FGlobalShader
+{
+	DECLARE_GLOBAL_SHADER(FRayTracingDeferredReflectionsRGS)
+	SHADER_USE_ROOT_PARAMETER_STRUCT(FRayTracingDeferredReflectionsRGS, FGlobalShader)
+
+	class FDeferredMaterialMode : SHADER_PERMUTATION_ENUM_CLASS("DIM_DEFERRED_MATERIAL_MODE", EDeferredMaterialMode);
+	using FPermutationDomain = TShaderPermutationDomain<FDeferredMaterialMode>;
+
+	BEGIN_SHADER_PARAMETER_STRUCT(FParameters, )
+		SHADER_PARAMETER(FIntPoint, RayTracingResolution)
+		SHADER_PARAMETER(FIntPoint, TileAlignedResolution)
+		SHADER_PARAMETER(float, ReflectionMaxNormalBias)
+		SHADER_PARAMETER_SRV(RaytracingAccelerationStructure, TLAS)
+		SHADER_PARAMETER_RDG_BUFFER_UAV(RWStructuredBuffer<FSortedReflectionRay>, RayBuffer)
+		SHADER_PARAMETER_RDG_BUFFER_UAV(RWStructuredBuffer<FRayIntersectionBookmark>, BookmarkBuffer)
+		SHADER_PARAMETER_RDG_BUFFER_UAV(RWStructuredBuffer<FDeferredMaterialPayload>, MaterialBuffer)
+		SHADER_PARAMETER_STRUCT_INCLUDE(FSceneTextureParameters, SceneTextures)
+		SHADER_PARAMETER_STRUCT_INCLUDE(FSceneTextureSamplerParameters, SceneTextureSamplers)
+		SHADER_PARAMETER_SRV(StructuredBuffer<FRTLightingData>, LightDataBuffer)
+		SHADER_PARAMETER_RDG_TEXTURE(Texture2D, SSProfilesTexture)
+		SHADER_PARAMETER_RDG_TEXTURE_UAV(RWTexture2D<float4>, ColorOutput)
+		SHADER_PARAMETER_RDG_TEXTURE_UAV(RWTexture2D<float>, RayHitDistanceOutput)
+		SHADER_PARAMETER_STRUCT_REF(FViewUniformShaderParameters, ViewUniformBuffer)
+		SHADER_PARAMETER_STRUCT_REF(FReflectionUniformParameters, ReflectionStruct)
+		SHADER_PARAMETER_STRUCT_REF(FRaytracingLightDataPacked, LightDataPacked)
+	END_SHADER_PARAMETER_STRUCT()
+
+	static bool ShouldCompilePermutation(const FGlobalShaderPermutationParameters& Parameters)
+	{
+		FPermutationDomain PermutationVector(Parameters.PermutationId);
+		if (PermutationVector.Get<FDeferredMaterialMode>() == EDeferredMaterialMode::None)
+		{
+			return false;
+		}
+
+		return ShouldCompileRayTracingShadersForProject(Parameters.Platform);
+	}
+
+	static void ModifyCompilationEnvironment(const FGlobalShaderPermutationParameters& Parameters, FShaderCompilerEnvironment& OutEnvironment)
+	{
+		FGlobalShader::ModifyCompilationEnvironment(Parameters, OutEnvironment);
+		OutEnvironment.SetDefine(TEXT("UE_RAY_TRACING_DISPATCH_1D"), 1); // Always using 1D dispatches
+	}
+};
+IMPLEMENT_GLOBAL_SHADER(FRayTracingDeferredReflectionsRGS, "/Engine/Private/RayTracing/RayTracingDeferredReflections.usf", "RayTracingDeferredReflectionsRGS", SF_RayGen);
+
+void FDeferredShadingSceneRenderer::PrepareRayTracingDeferredReflections(const FViewInfo& View, const FScene& Scene, TArray<FRHIRayTracingShader*>& OutRayGenShaders)
+{
+	FRayTracingDeferredReflectionsRGS::FPermutationDomain PermutationVector;
+
+	{
+		PermutationVector.Set<FRayTracingDeferredReflectionsRGS::FDeferredMaterialMode>(EDeferredMaterialMode::Gather);
+		auto RayGenShader = View.ShaderMap->GetShader<FRayTracingDeferredReflectionsRGS>(PermutationVector);
+		OutRayGenShaders.Add(RayGenShader.GetRayTracingShader());
+	}
+
+	{
+		PermutationVector.Set<FRayTracingDeferredReflectionsRGS::FDeferredMaterialMode>(EDeferredMaterialMode::Shade);
+		auto RayGenShader = View.ShaderMap->GetShader<FRayTracingDeferredReflectionsRGS>(PermutationVector);
+		OutRayGenShaders.Add(RayGenShader.GetRayTracingShader());
+	}
+}
+
+static void AddGenerateReflectionRaysPass(
+	FRDGBuilder& GraphBuilder,
+	const FViewInfo& View,
+	FRDGBufferRef RayBuffer,
+	const FRayTracingDeferredReflectionsRGS::FParameters& CommonParameters)
+{
+	FGenerateReflectionRaysCS::FParameters* PassParameters = GraphBuilder.AllocParameters<FGenerateReflectionRaysCS::FParameters>();
+	PassParameters->RayTracingResolution    = CommonParameters.RayTracingResolution;
+	PassParameters->TileAlignedResolution   = CommonParameters.TileAlignedResolution;
+	PassParameters->ReflectionMaxNormalBias = CommonParameters.ReflectionMaxNormalBias;
+	PassParameters->ViewUniformBuffer       = CommonParameters.ViewUniformBuffer;
+	PassParameters->SceneTextures           = CommonParameters.SceneTextures;
+	PassParameters->RayBuffer               = GraphBuilder.CreateUAV(RayBuffer);
+
+	FGenerateReflectionRaysCS::FPermutationDomain PermutationVector;
+	const bool bUseWaveOps = GRHISupportsWaveOperations && GRHIMinimumWaveSize >= 32 && RHISupportsWaveOperations(View.GetShaderPlatform());
+	PermutationVector.Set<FGenerateReflectionRaysCS::FWaveOps>(bUseWaveOps);
+
+	auto ComputeShader = View.ShaderMap->GetShader<FGenerateReflectionRaysCS>(PermutationVector);
+	ClearUnusedGraphResources(ComputeShader, PassParameters);
+
+	const uint32 NumRays = CommonParameters.TileAlignedResolution.X * CommonParameters.TileAlignedResolution.Y;
+	FIntVector GroupCount;
+	GroupCount.X = FMath::DivideAndRoundUp(NumRays, FGenerateReflectionRaysCS::GetGroupSize());
+	GroupCount.Y = 1;
+	GroupCount.Z = 1;
+	FComputeShaderUtils::AddPass(GraphBuilder, RDG_EVENT_NAME("GenerateReflectionRays"), ComputeShader, PassParameters, GroupCount);
+}
+
+void FDeferredShadingSceneRenderer::RenderRayTracingDeferredReflections(
+	FRDGBuilder& GraphBuilder,
+	const FSceneTextureParameters& SceneTextures,
+	const FViewInfo& View,
+	IScreenSpaceDenoiser::FReflectionsInputs* OutDenoiserInputs)
+{
+	FRDGTextureDesc OutputDesc = FPooledRenderTargetDesc::Create2DDesc(
+		FSceneRenderTargets::Get_FrameConstantsOnly().GetBufferSizeXY(),
+		PF_FloatRGBA, FClearValueBinding(FLinearColor(0, 0, 0, 0)),
+		TexCreate_None, TexCreate_ShaderResource | TexCreate_UAV,
+		false);
+
+	OutDenoiserInputs->Color          = GraphBuilder.CreateTexture(OutputDesc, TEXT("RTDReflections"));
+	OutputDesc.Format                 = PF_R16F;
+	OutDenoiserInputs->RayHitDistance = GraphBuilder.CreateTexture(OutputDesc, TEXT("RTDReflectionsHitDistance"));
+
+	const FIntPoint RayTracingResolution = View.ViewRect.Size();
+
+	const uint32 SortTileSize             = 64; // Ray sort tile is 32x32, material sort tile is 64x64, so we use 64 here (tile size is not configurable).
+	const FIntPoint TileAlignedResolution = FIntPoint::DivideAndRoundUp(RayTracingResolution, SortTileSize) * SortTileSize;
+
+	FRayTracingDeferredReflectionsRGS::FParameters CommonParameters;
+	CommonParameters.RayTracingResolution    = RayTracingResolution;
+	CommonParameters.TileAlignedResolution   = TileAlignedResolution;
+	CommonParameters.TLAS                    = View.RayTracingScene.RayTracingSceneRHI->GetShaderResourceView();
+	CommonParameters.SceneTextures           = SceneTextures;
+	SetupSceneTextureSamplers(&CommonParameters.SceneTextureSamplers);
+	CommonParameters.ViewUniformBuffer       = View.ViewUniformBuffer;
+	CommonParameters.LightDataPacked         = View.RayTracingLightingDataUniformBuffer;
+	CommonParameters.LightDataBuffer         = View.RayTracingLightingDataSRV;
+	CommonParameters.SSProfilesTexture       = GraphBuilder.RegisterExternalTexture(View.RayTracingSubSurfaceProfileTexture);
+	CommonParameters.ReflectionStruct        = CreateReflectionUniformBuffer(View, EUniformBufferUsage::UniformBuffer_SingleFrame);
+	CommonParameters.ReflectionMaxNormalBias = GetRaytracingMaxNormalBias();
+
+	// Generate sorted reflection rays
+
+	const uint32 TileAlignedNumRays          = TileAlignedResolution.X * TileAlignedResolution.Y;
+	const FRDGBufferDesc SortedRayBufferDesc = FRDGBufferDesc::CreateStructuredDesc(sizeof(FSortedReflectionRay), TileAlignedNumRays);
+	FRDGBufferRef SortedRayBuffer            = SortedRayBuffer = GraphBuilder.CreateBuffer(SortedRayBufferDesc, TEXT("ReflectionRayBuffer"));
+
+	const FRDGBufferDesc DeferredMaterialBufferDesc = FRDGBufferDesc::CreateStructuredDesc(sizeof(FDeferredMaterialPayload), TileAlignedNumRays);
+	FRDGBufferRef DeferredMaterialBuffer            = GraphBuilder.CreateBuffer(DeferredMaterialBufferDesc, TEXT("RayTracingReflectionsMaterialBuffer"));
+
+	const FRDGBufferDesc BookmarkBufferDesc = FRDGBufferDesc::CreateStructuredDesc(sizeof(FRayIntersectionBookmark), TileAlignedNumRays);
+	FRDGBufferRef BookmarkBuffer            = GraphBuilder.CreateBuffer(BookmarkBufferDesc, TEXT("RayTracingReflectionsBookmarkBuffer"));
+
+	AddGenerateReflectionRaysPass(GraphBuilder, View, SortedRayBuffer, CommonParameters);
+
+	// Trace reflection material gather rays
+
+	{
+		FRayTracingDeferredReflectionsRGS::FParameters& PassParameters = *GraphBuilder.AllocParameters<FRayTracingDeferredReflectionsRGS::FParameters>();
+		PassParameters                      = CommonParameters;
+		PassParameters.MaterialBuffer       = GraphBuilder.CreateUAV(DeferredMaterialBuffer);
+		PassParameters.RayBuffer            = GraphBuilder.CreateUAV(SortedRayBuffer);
+		PassParameters.BookmarkBuffer       = GraphBuilder.CreateUAV(BookmarkBuffer);
+		PassParameters.ColorOutput          = GraphBuilder.CreateUAV(OutDenoiserInputs->Color);
+		PassParameters.RayHitDistanceOutput = GraphBuilder.CreateUAV(OutDenoiserInputs->RayHitDistance);
+
+		FRayTracingDeferredReflectionsRGS::FPermutationDomain PermutationVector;
+		PermutationVector.Set<FRayTracingDeferredReflectionsRGS::FDeferredMaterialMode>(EDeferredMaterialMode::Gather);
+
+		auto RayGenShader = View.ShaderMap->GetShader<FRayTracingDeferredReflectionsRGS>(PermutationVector);
+		ClearUnusedGraphResources(RayGenShader, &PassParameters);
+
+		GraphBuilder.AddPass(
+			RDG_EVENT_NAME("RayTracingDeferredReflectionsGather %dx%d", RayTracingResolution.X, RayTracingResolution.Y),
+			&PassParameters,
+			ERDGPassFlags::Compute,
+		[&PassParameters, this, &View, TileAlignedNumRays, RayGenShader](FRHICommandList& RHICmdList)
+		{
+			FRayTracingPipelineState* Pipeline = BindRayTracingDeferredMaterialGatherPipeline(RHICmdList, View, RayGenShader.GetRayTracingShader());
+
+			FRayTracingShaderBindingsWriter GlobalResources;
+			SetShaderParameters(GlobalResources, RayGenShader, PassParameters);
+			FRHIRayTracingScene* RayTracingSceneRHI = View.RayTracingScene.RayTracingSceneRHI;
+			RHICmdList.RayTraceDispatch(Pipeline, RayGenShader.GetRayTracingShader(), RayTracingSceneRHI, GlobalResources, TileAlignedNumRays, 1);
+		});
+	}
+
+	// Sort hit points by material within 64x64 (4096 element) tiles
+
+	SortDeferredMaterials(GraphBuilder, View, 5, TileAlignedNumRays, DeferredMaterialBuffer);
+
+	// Shade reflection points
+
+	{
+		FRayTracingDeferredReflectionsRGS::FParameters& PassParameters = *GraphBuilder.AllocParameters<FRayTracingDeferredReflectionsRGS::FParameters>();
+		PassParameters                      = CommonParameters;
+		PassParameters.MaterialBuffer       = GraphBuilder.CreateUAV(DeferredMaterialBuffer);
+		PassParameters.RayBuffer            = GraphBuilder.CreateUAV(SortedRayBuffer);
+		PassParameters.BookmarkBuffer       = GraphBuilder.CreateUAV(BookmarkBuffer);
+		PassParameters.ColorOutput          = GraphBuilder.CreateUAV(OutDenoiserInputs->Color);
+		PassParameters.RayHitDistanceOutput = GraphBuilder.CreateUAV(OutDenoiserInputs->RayHitDistance);
+
+		FRayTracingDeferredReflectionsRGS::FPermutationDomain PermutationVector;
+		PermutationVector.Set<FRayTracingDeferredReflectionsRGS::FDeferredMaterialMode>(EDeferredMaterialMode::Shade);
+
+		auto RayGenShader = View.ShaderMap->GetShader<FRayTracingDeferredReflectionsRGS>(PermutationVector);
+		ClearUnusedGraphResources(RayGenShader, &PassParameters);
+
+		GraphBuilder.AddPass(
+			RDG_EVENT_NAME("RayTracingDeferredReflectionsShade %dx%d", RayTracingResolution.X, RayTracingResolution.Y),
+			&PassParameters,
+			ERDGPassFlags::Compute,
+		[&PassParameters, &View, TileAlignedNumRays, RayGenShader](FRHICommandList& RHICmdList)
+		{
+			FRayTracingShaderBindingsWriter GlobalResources;
+			SetShaderParameters(GlobalResources, RayGenShader, PassParameters);
+			FRHIRayTracingScene* RayTracingSceneRHI = View.RayTracingScene.RayTracingSceneRHI;
+			RHICmdList.RayTraceDispatch(View.RayTracingMaterialPipeline, RayGenShader.GetRayTracingShader(), RayTracingSceneRHI, GlobalResources, TileAlignedNumRays, 1);
+		});
+	}
+}
+#else // RHI_RAYTRACING
+void FDeferredShadingSceneRenderer::RenderRayTracingDeferredReflections(
+	FRDGBuilder& GraphBuilder,
+	const FSceneTextureParameters& SceneTextures,
+	const FViewInfo& View,
+	IScreenSpaceDenoiser::FReflectionsInputs* OutDenoiserInputs)
+{
+	checkNoEntry();
+}
+#endif // RHI_RAYTRACING
diff --git a/Engine/Source/Runtime/Renderer/Private/RayTracing/RayTracingReflections.cpp b/Engine/Source/Runtime/Renderer/Private/RayTracing/RayTracingReflections.cpp
--- a/Engine/Source/Runtime/Renderer/Private/RayTracing/RayTracingReflections.cpp
+++ b/Engine/Source/Runtime/Renderer/Private/RayTracing/RayTracingReflections.cpp
@@ -184,6 +184,11 @@ static TAutoConsoleVariable<int32> CVarRayTracingReflectionsRayTraceSkyLightCont
 	TEXT("Requests ray tracing reflections to use ray traced visibility rays for sky light contribution similar to sky light ray traced shadows. A Sky Light with ray traced shadows enabled must be present for this flag to take effect. (default = 0)"),
 	ECVF_RenderThreadSafe);
 
+static TAutoConsoleVariable<int32> CVarRayTracingReflectionsExperimentalDeferred(
+	TEXT("r.RayTracing.Reflections.ExperimentalDeferred"),
+	0,
+	TEXT("Whether to use the experimental deferred ray traced reflection rendering algorithm, which only supports a subset of features but runs faster. (default = 0)."),
+	ECVF_RenderThreadSafe);
 
 // ESamplePhase
 enum class ESamplePhase
@@ -383,6 +388,13 @@ void FDeferredShadingSceneRenderer::PrepareRayTracingReflections(const FViewInfo
 {
 	// Declare all RayGen shaders that require material closest hit shaders to be bound
 
+	if (CVarRayTracingReflectionsExperimentalDeferred.GetValueOnRenderThread())
+	{
+		// If deferred reflections technique is used, then we only need to gather its shaders and skip the rest.
+		PrepareRayTracingDeferredReflections(View, Scene, OutRayGenShaders);
+		return;
+	}
+
 	const bool bHybridReflections = ShouldRayTracedReflectionsUseHybridReflections();
 	const bool bSortMaterials = ShouldRayTracedReflectionsSortMaterials(View);
 	const bool bMissShaderLighting = CanUseRayTracingLightingMissShader(View.GetShaderPlatform());
@@ -458,6 +470,16 @@ void FDeferredShadingSceneRenderer::RenderRayTracingReflections(
 	IScreenSpaceDenoiser::FReflectionsInputs* OutDenoiserInputs)
 #if RHI_RAYTRACING
 {
+	if (CVarRayTracingReflectionsExperimentalDeferred.GetValueOnRenderThread())
+	{
+		RenderRayTracingDeferredReflections(
+			GraphBuilder,
+			SceneTextures,
+			View,
+			OutDenoiserInputs);
+		return;
+	}
+
 	const uint32 SortTileSize = CVarRayTracingReflectionsSortTileSize.GetValueOnRenderThread();
 	const uint32 EnableTranslucency = GRayTracingReflectionsTranslucency > -1 ? (uint32)GRayTracingReflectionsTranslucency : (uint32)View.FinalPostProcessSettings.RayTracingReflectionsTranslucency;
 
