Free reflection scratch cubemaps after use

diff --git a/Engine/Source/Runtime/Renderer/Private/PostProcess/SceneRenderTargets.h b/Engine/Source/Runtime/Renderer/Private/PostProcess/SceneRenderTargets.h
--- a/Engine/Source/Runtime/Renderer/Private/PostProcess/SceneRenderTargets.h
+++ b/Engine/Source/Runtime/Renderer/Private/PostProcess/SceneRenderTargets.h
@@ -280,6 +280,18 @@ public:
 			DownsampledTranslucencyDepthRT.SafeRelease();
 		}
 	}
+
+	void FreeReflectionScratchRenderTargets()
+	{
+		for (int32 Idx = 0; Idx < UE_ARRAY_COUNT(ReflectionColorScratchCubemap); ++Idx)
+		{
+			ReflectionColorScratchCubemap[Idx].SafeRelease();
+		}
+		for (int32 Idx = 0; Idx < UE_ARRAY_COUNT(DiffuseIrradianceScratchCubemap); ++Idx)
+		{
+			DiffuseIrradianceScratchCubemap[Idx].SafeRelease();
+		}
+	}
 	
 	void ResolveSceneDepthTexture(FRHICommandList& RHICmdList, const FResolveRect& ResolveRect);
 
diff --git a/Engine/Source/Runtime/Renderer/Private/ReflectionEnvironmentCapture.cpp b/Engine/Source/Runtime/Renderer/Private/ReflectionEnvironmentCapture.cpp
--- a/Engine/Source/Runtime/Renderer/Private/ReflectionEnvironmentCapture.cpp
+++ b/Engine/Source/Runtime/Renderer/Private/ReflectionEnvironmentCapture.cpp
@@ -97,6 +97,12 @@ static FAutoConsoleVariableRef CVarReleaseSRVsAfterCubeMipsGen(
 	TEXT("This reduces unnecessary memory allocations and copies when the item is copied around during snapshot creation."),
 	ECVF_RenderThreadSafe);
 
+static int32 GFreeReflectionScratchAfterUse = 0;
+static FAutoConsoleVariableRef CVarFreeReflectionScratchAfterUse(
+	TEXT("r.FreeReflectionScratchAfterUse"),
+	GFreeReflectionScratchAfterUse,
+	TEXT("Free reflection scratch render targets after use."));
+
 bool DoGPUArrayCopy()
 {
 	return GRHISupportsResolveCubemapFaces && CVarReflectionCaptureGPUArrayCopy.GetValueOnAnyThread();
@@ -1673,5 +1679,16 @@ void FScene::UpdateSkyCaptureContents(
 					}
 				});
 		}
+
+		if (!!GFreeReflectionScratchAfterUse)
+		{
+			ENQUEUE_RENDER_COMMAND(FreeReflectionScratch)(
+				[](FRHICommandListImmediate& RHICmdList)
+			{
+				FSceneRenderTargets& SceneContext = FSceneRenderTargets::Get(RHICmdList);
+				SceneContext.FreeReflectionScratchRenderTargets();
+				GRenderTargetPool.FreeUnusedResources();
+			});
+		}
 	}
 }
