Add option to sort hair sample by depth, and update their coverage according the their ordered transmittance (similar to PPLL method).

This reduces coverage artifacts.

diff --git a/Engine/Shaders/Private/HairStrands/HairStrandsMaterialPS.usf b/Engine/Shaders/Private/HairStrands/HairStrandsMaterialPS.usf
--- a/Engine/Shaders/Private/HairStrands/HairStrandsMaterialPS.usf
+++ b/Engine/Shaders/Private/HairStrands/HairStrandsMaterialPS.usf
@@ -69,6 +69,7 @@ FEvaluateOutput Evaluate(
 	float3 Velocity_World_P;
 	float3 Velocity_World_P_Prev;
 	float U = 0;
+	float WorldStrandRadius = 0;
 	FVertexFactoryInterpolantsVSToPS Interpolants;
 	{
 		const uint HairPrimitiveId0 = SamplePrimitiveId;
@@ -95,9 +96,9 @@ FEvaluateOutput Evaluate(
 		{
 			const float3 W = cross(A, B);
 			const bool bIsRight = dot(W, View.ViewForward) < 0;
-			const float WorldRadius	= lerp(Vertex0.Radius, Vertex1.Radius, SegmentU);
 			const float WorldV	= length(A - B * SegmentU);
-			SegmentV = WorldV / WorldRadius;
+			WorldStrandRadius	= lerp(Vertex0.Radius, Vertex1.Radius, SegmentU);
+			SegmentV = WorldV / WorldStrandRadius;
 			SegmentV = bIsRight ? SegmentV * 0.5f + 0.5f : (1 - SegmentV) * 0.5f;
 		}
 
@@ -131,32 +132,32 @@ FEvaluateOutput Evaluate(
 	// VS
 	const float4 SvPosition = float4(SamplePixelCoord, SampleDepth, 1);
 
-	// PS
-	FHairSample OutSample		= (FHairSample)0;
-	OutSample.Depth				= SampleDepth;
-	OutSample.Coverage8bit		= SampleCoverage8bit;
-	OutSample.PrimitiveId		= SamplePrimitiveId;
-	OutSample.MacroGroupId		= MaterialPass_MacroGroupId;
-	OutSample.LightChannelMask	= MaterialPass_LightChannelMask;
-	OutSample.bScatterSceneLighting = UseScatterSceneLighting();
-
 	// Coverage computation
 	// We don't use the coverage information
-	#if 0
 	float Coverage = 1;
+	if (MaterialPassParameters.bUpdateSampleCoverage)
 	{
 		const bool bUseStableRasterization = UseStableRasterization();
 
-		float WorldStrandRadius = VFIntermediates.HairDimensions.y; // Coming from the vertex factory evaluation
-
 		FHairRenderInfo HairRenderInfo = GetHairRenderInfo(ResolvedView.HairRenderInfo, ResolvedView.HairRenderInfoBits, bUseStableRasterization);
 		const float SceneDepth = ConvertFromDeviceZ(SvPosition.z); // Linear depth in world unit
 		const float PixelRadius = HairRenderInfo.bIsOrthoView ? HairRenderInfo.RadiusAtDepth1Primary : SceneDepth * HairRenderInfo.RadiusAtDepth1Primary; // Not correct but the coverage is not used (we count instead the number of sub-sample covered)
 		const float StrandRealRadius = WorldStrandRadius;
 		Coverage = StrandRealRadius / max(StrandRealRadius, PixelRadius);
+
+		SampleCoverage8bit = min(uint(Coverage * 0xFF), 0xFF);
 	}
-	float Transmittance = saturate(1.0f - Coverage);
-	#endif
+
+	// expressed the coverage relatively to the current pixel coverage?
+
+	// PS
+	FHairSample OutSample		= (FHairSample)0;
+	OutSample.Depth				= SampleDepth;
+	OutSample.Coverage8bit		= SampleCoverage8bit;
+	OutSample.PrimitiveId		= SamplePrimitiveId;
+	OutSample.MacroGroupId		= MaterialPass_MacroGroupId;
+	OutSample.LightChannelMask	= MaterialPass_LightChannelMask;
+	OutSample.bScatterSceneLighting = UseScatterSceneLighting();
 
 	// Material computation
 	{
diff --git a/Engine/Shaders/Private/HairStrands/HairStrandsVisibilityCommon.ush b/Engine/Shaders/Private/HairStrands/HairStrandsVisibilityCommon.ush
--- a/Engine/Shaders/Private/HairStrands/HairStrandsVisibilityCommon.ush
+++ b/Engine/Shaders/Private/HairStrands/HairStrandsVisibilityCommon.ush
@@ -219,6 +219,16 @@ FHairSample UnpackHairSample(const FPackedHairSample In)
 	return Out;
 }
 
+void PatchPackedHairSampleCoverage(inout FPackedHairSample Out, uint Coverage8bit)
+{
+	Out.Tangent_Coverage8bit = (Out.Tangent_Coverage8bit & 0x00FFFFFF) | (min(uint(0xFF), Coverage8bit) << 24);
+}
+
+uint GetPackedHairSampleCoverage8bit(in FPackedHairSample In)
+{
+	return 0xFF & (In.Tangent_Coverage8bit>>24);
+}
+
 ////////////////////////////////////////////////////////////////////////////////////////////////////////////////////////////////
 // Hair categorization
 
diff --git a/Engine/Shaders/Private/HairStrands/HairStrandsVisibilityCompaction.usf b/Engine/Shaders/Private/HairStrands/HairStrandsVisibilityCompaction.usf
--- a/Engine/Shaders/Private/HairStrands/HairStrandsVisibilityCompaction.usf
+++ b/Engine/Shaders/Private/HairStrands/HairStrandsVisibilityCompaction.usf
@@ -24,6 +24,7 @@
 int2 OutputResolution;
 int2 ResolutionOffset;
 uint MaxNodeCount;
+uint bSortSampleByDepth;
 float DepthTheshold; 
 float CosTangentThreshold;
 float CoverageThreshold; // Allow to scale the transmittance to compensate its asymptotic behavior 
@@ -58,6 +59,14 @@ float TransmittanceToCoverage(float InTransmittance)
 groupshared uint AllocationNodeCount;
 groupshared uint AllocationNodeOffset;
 
+
+struct FSampleSetDesc
+{
+	uint   UniqueSampleCount;
+	uint   ValidSampleCount;
+	uint   HairSampleCount;
+};
+
 ///////////////////////////////////////////////////////////////////////////////////////////////////////////////
 // PPLL
 ///////////////////////////////////////////////////////////////////////////////////////////////////////////////
@@ -260,28 +269,18 @@ void MainCS(uint GroupIndex : SV_GroupIndex, uint3 DispatchThreadId : SV_Dispatc
 	OutCategorizationTexture[PixelCoord] = EncodeCategorizationData(Categorization);
 }
 
-
-#else // PERMUTATION_PPLL
+#elif PERMUTATION_VISIBILITY // PERMUTATION_PPLL or PERMUTATION_VISIBILITY or PERMUTATION_MSAA
 
 ///////////////////////////////////////////////////////////////////////////////////////////////////////////////
 // MSAA visiblity buffer
 ///////////////////////////////////////////////////////////////////////////////////////////////////////////////
 
-struct FSampleSetDesc
-{
-	uint   UniqueSampleCount;
-	uint   ValidSampleCount;
-	uint   HairSampleCount;
-};
-
 #ifndef PERMUTATION_MSAACOUNT
 	#error PERMUTATION_MSAACOUNT is not defined
 #endif
 
-#if PERMUTATION_VISIBILITY
-
 void ComputeUniqueSamplesWithoutDepthTest(
-	inout uint3 OutSamples[PERMUTATION_MSAACOUNT],
+	inout uint4 OutSamples[PERMUTATION_MSAACOUNT],
 	inout FSampleSetDesc OutSet,
 	const uint2 PixelCoord,
 	const float OpaqueDepth,
@@ -294,7 +293,6 @@ void ComputeUniqueSamplesWithoutDepthTest(
 	OutSet.HairSampleCount = PERMUTATION_MSAACOUNT;
 	const uint InvalidSample = 0xFFFF;
 
-	float ClosestDepth = 0; // Inverse-Z
 	for (uint SampleIt = 0; SampleIt < OutSet.HairSampleCount; ++SampleIt)
 	{
 		// Note: InMSAA_IDTexture contains both the primitive ID and tha material ID. However 
@@ -315,28 +313,52 @@ void ComputeUniqueSamplesWithoutDepthTest(
 		for (uint UniqueIt = 0; UniqueIt < OutSet.UniqueSampleCount; ++UniqueIt)
 		{
 			const bool bIsSimilar = PrimitiveId == OutSamples[UniqueIt].x;
-
 			if (bIsSimilar)
 			{
 				OutSamples[UniqueIt].y += 1;
-				if (ClosestDepth < SampleDepth)
+
+				// Update the unique sample with the closest depth
+				const uint IntDepth = asuint(SampleDepth);
+				if (IntDepth > OutSamples[UniqueIt].w)
 				{
-					OutSamples[OutSet.UniqueSampleCount].z = SampleIt;
-					ClosestDepth = SampleDepth;
+					OutSamples[UniqueIt].z = SampleIt;
+					OutSamples[UniqueIt].w = asuint(SampleDepth);
 				}
+
 				bAlreadyExist = true;
 				break;
 			}
 		}
-		// todo: always takes the closest depth?
+		
 		if (!bAlreadyExist)
 		{
 			OutSamples[OutSet.UniqueSampleCount].x = PrimitiveId;
 			OutSamples[OutSet.UniqueSampleCount].y = 1;
 			OutSamples[OutSet.UniqueSampleCount].z = SampleIt;
+			OutSamples[OutSet.UniqueSampleCount].w = asuint(SampleDepth);
 			++OutSet.UniqueSampleCount;
 		}
 	}
+
+	// Sort sample from closer to further. This is used later for updating sample coverage 
+	// based on ordered transmittance. See HairStrandsVisibilityComputeSampleCoverage.usf for more details.
+	if (bSortSampleByDepth > 0)
+	{
+		for (uint i = 0; i < OutSet.UniqueSampleCount; ++i)
+		{
+			const uint DepthI = OutSamples[i].w;
+			for (uint j = i+1; j < OutSet.UniqueSampleCount; ++j)
+			{
+				const uint DepthJ = OutSamples[j].w;
+				if (DepthJ > DepthI)
+				{
+					uint4 Temp = OutSamples[i];
+					OutSamples[i] = OutSamples[j];
+					OutSamples[j] = Temp;
+				}
+			}
+		}
+	}
 }
 
 Texture2DMS<uint, PERMUTATION_MSAACOUNT> MSAA_IDTexture;
@@ -362,7 +384,7 @@ void MainCS(uint GroupIndex : SV_GroupIndex, uint3 DispatchThreadId : SV_Dispatc
 	FSampleSetDesc SampleDesc;
 
 	const float OpaqueDepth = SceneTexturesStruct.SceneDepthTexture.Load(uint3(PixelCoord, 0)).r;
-	uint3 Samples[PERMUTATION_MSAACOUNT];		// x:PrimitiveId|MaterialId, y:Weight, z:SampleIt
+	uint4 Samples[PERMUTATION_MSAACOUNT];		// x:PrimitiveId|MaterialId, y:Weight, z:SampleIt, w:Depth (as uint)
 	ComputeUniqueSamplesWithoutDepthTest(
 		Samples,
 		SampleDesc,
@@ -399,7 +421,7 @@ void MainCS(uint GroupIndex : SV_GroupIndex, uint3 DispatchThreadId : SV_Dispatc
 				// VisibilityData.Coverage8bit is a weight normalising to 1 the contribution of all the compacted samples. Because later it is weighted by Categorization.PixelCoverage.
 				FNodeVis OutNodeVis;
 				OutNodeVis.PrimitiveId	= VisSample.PrimitiveId;
-				OutNodeVis.Depth		= MSAA_DepthTexture.Load(PixelCoord, Samples[OutIndex].z);
+				OutNodeVis.Depth		= asfloat(Samples[OutIndex].w); // MSAA_DepthTexture.Load(PixelCoord, Samples[OutIndex].z);
 				OutNodeVis.Coverage8bit = To8bitCoverage(Samples[OutIndex].y / float(SampleDesc.ValidSampleCount));
 				OutNodeVis.MaterialId	= VisSample.MaterialId;
 
@@ -426,7 +448,7 @@ void MainCS(uint GroupIndex : SV_GroupIndex, uint3 DispatchThreadId : SV_Dispatc
 	OutCategorizationTexture[PixelCoord] = EncodeCategorizationData(Categorization);
 }
 
-#else // PERMUTATION_VISIBILITY
+#else // PERMUTATION_PPLL or PERMUTATION_VISIBILITY or PERMUTATION_MSAA
 
 ///////////////////////////////////////////////////////////////////////////////////////////////////////////////
 // MSAA buffer
@@ -667,7 +689,4 @@ void MainCS(uint GroupIndex : SV_GroupIndex, uint3 DispatchThreadId : SV_Dispatc
 	OutCategorizationTexture[PixelCoord] = EncodeCategorizationData(Categorization);
 }
 
-#endif  // PERMUTATION_VISIBILITY
-#endif // PERMUTATION_PPLL
-
-
+#endif // PERMUTATION_PPLL or PERMUTATION_VISIBILITY or PERMUTATION_MSAA
\ No newline at end of file
diff --git a/Engine/Shaders/Private/HairStrands/HairStrandsVisibilityComputeSampleCoverage.usf b/Engine/Shaders/Private/HairStrands/HairStrandsVisibilityComputeSampleCoverage.usf
new file mode 100644
--- /dev/null
+++ b/Engine/Shaders/Private/HairStrands/HairStrandsVisibilityComputeSampleCoverage.usf
@@ -0,0 +1,55 @@
+// Copyright Epic Games, Inc. All Rights Reserved.
+
+#pragma once
+
+#include "../Common.ush"
+#include "HairStrandsVisibilityCommon.ush" 
+
+#define MAX_SAMPLE_COUNT 8
+
+int2 Resolution;
+Texture2D<uint> NodeIndexAndOffset;
+StructuredBuffer<FPackedHairSample> InNodeDataBuffer;
+RWStructuredBuffer<FPackedHairSample> OutNodeDataBuffer;
+
+// TODO: Add permutation for group 32/64
+[numthreads(8, 4, 1)]
+void MainCS(uint3 DispatchThreadId : SV_DispatchThreadID)
+{
+	const uint2 PixelCoord = DispatchThreadId.xy;
+	if (any(PixelCoord.xy >= uint2(Resolution)))
+		return;
+
+	// All sample belonging to a pixel are already sorted by depth (closer to further)
+	const FNodeDesc NodeDesc = DecodeNodeDesc(NodeIndexAndOffset.Load(uint3(PixelCoord, 0)));
+
+	float ValidPixelSampleTotalCoverage = 0.0f;
+	float SortedCoverage[MAX_SAMPLE_COUNT];
+	float TotalSortedTransmittance = 1.0f;
+	for (uint i = 0; i < NodeDesc.Count; ++i)
+	{
+		const FPackedHairSample PackedNodeData = InNodeDataBuffer[NodeDesc.Offset + i];
+		const uint Coverage8bit = GetPackedHairSampleCoverage8bit(PackedNodeData);
+		const float Coverage = From8bitCoverage(Coverage8bit);
+
+		// Update current node coverage as a function of previous nodes coverage
+		SortedCoverage[i] = TotalSortedTransmittance * Coverage;
+
+		// Update transmittance for the next strands
+		TotalSortedTransmittance *= 1.0f - Coverage;
+
+		// Accumulate total coverage.
+		ValidPixelSampleTotalCoverage += SortedCoverage[i];
+	}
+
+	for (uint j = 0; j < NodeDesc.Count; ++j)
+	{
+		// Coverage8bit is a weight normalising to 1 the contribution of all the compacted samples. Because later it is weighted by Categorization.PixelCoverage.
+		// Patch the coverage on the out node
+		const uint PatchedCoverage8bit = To8bitCoverage(SortedCoverage[j] / float(ValidPixelSampleTotalCoverage));
+
+		FPackedHairSample PackedNodeData = InNodeDataBuffer[NodeDesc.Offset + j];
+		PatchPackedHairSampleCoverage(PackedNodeData, PatchedCoverage8bit);
+		OutNodeDataBuffer[NodeDesc.Offset + j] = PackedNodeData;
+	}
+}
\ No newline at end of file
diff --git a/Engine/Source/Runtime/Renderer/Private/HairStrands/HairStrandsVisibility.cpp b/Engine/Source/Runtime/Renderer/Private/HairStrands/HairStrandsVisibility.cpp
--- a/Engine/Source/Runtime/Renderer/Private/HairStrands/HairStrandsVisibility.cpp
+++ b/Engine/Source/Runtime/Renderer/Private/HairStrands/HairStrandsVisibility.cpp
@@ -66,6 +66,10 @@ static FAutoConsoleVariableRef CVarHairStrandsVisibilityComputeRasterSampleCount
 static float GHairStrandsFullCoverageThreshold = 0.98f;
 static FAutoConsoleVariableRef CVarHairStrandsFullCoverageThreshold(TEXT("r.HairStrands.Visibility.FullCoverageThreshold"), GHairStrandsFullCoverageThreshold, TEXT("Define the coverage threshold at which a pixel is considered fully covered."));
 
+
+static int32 GHairStrandsSortHairSampleByDepth = 0;
+static FAutoConsoleVariableRef CVarHairStrandsSortHairSampleByDepth(TEXT("r.HairStrands.Visibility.SortByDepth"), GHairStrandsSortHairSampleByDepth, TEXT("Sort hair fragment by depth and update their coverage based on ordered transmittance."));
+
 /////////////////////////////////////////////////////////////////////////////////////////
 
 namespace HairStrandsVisibilityInternal
@@ -277,6 +281,7 @@ BEGIN_GLOBAL_SHADER_PARAMETER_STRUCT(FMaterialPassParameters, )
 	SHADER_PARAMETER(FIntPoint, MaxResolution)
 	SHADER_PARAMETER(uint32, MaxSampleCount)
 	SHADER_PARAMETER(uint32, NodeGroupSize)
+	SHADER_PARAMETER(uint32, bUpdateSampleCoverage)
 	SHADER_PARAMETER_TEXTURE(Texture2D<uint>, NodeIndex)
 	SHADER_PARAMETER_SRV(StructuredBuffer<uint>, NodeCoord)
 	SHADER_PARAMETER_SRV(StructuredBuffer<FNodeVis>, NodeVis)
@@ -520,6 +525,58 @@ BEGIN_SHADER_PARAMETER_STRUCT(FVisibilityMaterialPassParameters, )
 	RENDER_TARGET_BINDING_SLOTS()
 END_SHADER_PARAMETER_STRUCT()
 
+///////////////////////////////////////////////////////////////////////////////////////////////////
+// Patch sample coverage
+class FUpdateSampleCoverageCS : public FGlobalShader
+{
+	DECLARE_GLOBAL_SHADER(FUpdateSampleCoverageCS);
+	SHADER_USE_PARAMETER_STRUCT(FUpdateSampleCoverageCS, FGlobalShader);
+
+	BEGIN_SHADER_PARAMETER_STRUCT(FParameters, )
+		SHADER_PARAMETER(FIntPoint, Resolution)
+		SHADER_PARAMETER_RDG_TEXTURE(Texture2D<uint>, NodeIndexAndOffset)
+		SHADER_PARAMETER_RDG_BUFFER_SRV(StructuredBuffer<FPackedHairSample>,  InNodeDataBuffer)
+		SHADER_PARAMETER_RDG_BUFFER_UAV(RWStructuredBuffer<FPackedHairSample>, OutNodeDataBuffer)
+	END_SHADER_PARAMETER_STRUCT()
+public:
+	static bool ShouldCompilePermutation(const FGlobalShaderPermutationParameters& Parameters) { return IsHairStrandsSupported(Parameters.Platform); }
+};
+
+IMPLEMENT_GLOBAL_SHADER(FUpdateSampleCoverageCS, "/Engine/Private/HairStrands/HairStrandsVisibilityComputeSampleCoverage.usf", "MainCS", SF_Compute);
+
+static FRDGBufferRef AddUpdateSampleCoveragePass(
+	FRDGBuilder& GraphBuilder,
+	const FViewInfo* View,
+	const FRDGTextureRef NodeIndexAndOffset,
+	const FRDGBufferRef InNodeDataBuffer)
+{
+	FRDGBufferRef OutNodeDataBuffer = GraphBuilder.CreateBuffer(FRDGBufferDesc::CreateStructuredDesc(InNodeDataBuffer->Desc.BytesPerElement, InNodeDataBuffer->Desc.NumElements), TEXT("HairCompactNodeData"));
+
+	FUpdateSampleCoverageCS::FParameters* Parameters = GraphBuilder.AllocParameters<FUpdateSampleCoverageCS::FParameters>();
+	Parameters->Resolution = NodeIndexAndOffset->Desc.Extent;
+	Parameters->NodeIndexAndOffset = NodeIndexAndOffset;
+	Parameters->InNodeDataBuffer = GraphBuilder.CreateSRV(InNodeDataBuffer);
+	Parameters->OutNodeDataBuffer = GraphBuilder.CreateUAV(OutNodeDataBuffer);
+
+	TShaderMapRef<FUpdateSampleCoverageCS> ComputeShader(View->ShaderMap);
+
+	// Add 64 threads permutation
+	const uint32 GroupSizeX = 8;
+	const uint32 GroupSizeY = 4;
+	const FIntVector DispatchCount = FIntVector(
+		(Parameters->Resolution.X + GroupSizeX-1) / GroupSizeX, 
+		(Parameters->Resolution.Y + GroupSizeY-1) / GroupSizeY, 
+		1);
+	FComputeShaderUtils::AddPass(
+		GraphBuilder,
+		RDG_EVENT_NAME("HairStrandsVisbilityUpdateCoverage"),
+		ComputeShader,
+		Parameters,
+		DispatchCount);
+
+	return OutNodeDataBuffer;
+}
+///////////////////////////////////////////////////////////////////////////////////////////////////
 struct FMaterialPassOutput
 {
 	static const EPixelFormat VelocityFormat = PF_G16R16;
@@ -531,6 +588,7 @@ static FMaterialPassOutput AddHairMaterialPass(
 	FRDGBuilder& GraphBuilder,
 	const FScene* Scene,
 	const FViewInfo* ViewInfo,
+	const bool bUpdateSampleCoverage,
 	const FHairStrandsMacroGroupDatas& MacroGroupDatas,
 	const uint32 NodeGroupSize,
 	FRDGTextureRef CompactNodeIndex,
@@ -576,12 +634,13 @@ static FMaterialPassOutput AddHairMaterialPass(
 		RDG_EVENT_NAME("HairStrandsMaterialPass"),
 		PassParameters,
 		ERDGPassFlags::Raster,
-		[PassParameters, Scene = Scene, ViewInfo, &MacroGroupDatas, MaxNodeCount, Resolution, NodeGroupSize](FRHICommandListImmediate& RHICmdList)
+		[PassParameters, Scene = Scene, ViewInfo, &MacroGroupDatas, MaxNodeCount, Resolution, NodeGroupSize, bUpdateSampleCoverage](FRHICommandListImmediate& RHICmdList)
 	{
 		check(RHICmdList.IsInsideRenderPass());
 		check(IsInRenderingThread());
 
 		FMaterialPassParameters MaterialPassParameters;
+		MaterialPassParameters.bUpdateSampleCoverage = bUpdateSampleCoverage ? 1 : 0;
 		MaterialPassParameters.MaxResolution	= Resolution;
 		MaterialPassParameters.NodeGroupSize	= NodeGroupSize;
 		MaterialPassParameters.MaxSampleCount	= MaxNodeCount;
@@ -1231,6 +1290,7 @@ class FHairVisibilityPrimitiveIdCompactionCS : public FGlobalShader
 		SHADER_PARAMETER(FIntPoint, OutputResolution)
 		SHADER_PARAMETER(FIntPoint, ResolutionOffset)
 		SHADER_PARAMETER(uint32, MaxNodeCount)
+		SHADER_PARAMETER(uint32, bSortSampleByDepth)
 		SHADER_PARAMETER(float, DepthTheshold)
 		SHADER_PARAMETER(float, CosTangentThreshold)
 		SHADER_PARAMETER(float, CoverageThreshold)
@@ -1408,6 +1468,7 @@ static void AddHairVisibilityPrimitiveIdCompactionPass(
 
 	PassParameters->OutputResolution = Resolution;
 	PassParameters->MaxNodeCount = MaxRenderNodeCount;
+	PassParameters->bSortSampleByDepth = GHairStrandsSortHairSampleByDepth > 0 ? 1 : 0;
 	PassParameters->CoverageThreshold = FMath::Clamp(GHairStrandsFullCoverageThreshold, 0.1f, 1.f);
 	PassParameters->DepthTheshold = FMath::Clamp(GHairStrandsMaterialCompactionDepthThreshold, 0.f, 100.f);
 	PassParameters->CosTangentThreshold = FMath::Cos(FMath::DegreesToRadians(FMath::Clamp(GHairStrandsMaterialCompactionTangentThreshold, 0.f, 90.f)));
@@ -2546,12 +2607,15 @@ FHairStrandsVisibilityViews RenderHairStrandsVisibilityBuffer(
 
 					if (bIsVisiblityEnable)
 					{
+						const bool bUpdateSampleCoverage = GHairStrandsSortHairSampleByDepth > 0;
+
 						// Evaluate material based on the visiblity pass result
 						// Output both complete sample data + per-sample velocity
 						FMaterialPassOutput PassOutput = AddHairMaterialPass(
 							GraphBuilder,
 							Scene,
 							&View,
+							bUpdateSampleCoverage,
 							MacroGroupDatas,
 							VisibilityData.NodeGroupSize,
 							CompactNodeIndex,
@@ -2569,6 +2633,15 @@ FHairStrandsVisibilityViews RenderHairStrandsVisibilityBuffer(
 							PassOutput.NodeVelocity,
 							SceneVelocityTexture);
 
+						if (bUpdateSampleCoverage)
+						{
+							PassOutput.NodeData = AddUpdateSampleCoveragePass(
+								GraphBuilder,
+								&View,
+								CompactNodeIndex,
+								PassOutput.NodeData);
+						}
+
 						CompactNodeData = PassOutput.NodeData;
 					}
 
