Restore the call to BeginRenderingSeparateTranslucency when TranslucencyPass is not either TPT_TranslucencyAfterDOF or TPT_TranslucencyAfterDOFModulate.  This ensures a matching Begin to the corresponding End in EndRenderPass() below.  The call to UpsampleTranslucency() also expects that a call to BeginRenderingSeparateTranslucency has been made.

#rnx

diff --git a/Engine/Source/Runtime/Renderer/Private/TranslucentRendering.cpp b/Engine/Source/Runtime/Renderer/Private/TranslucentRendering.cpp
--- a/Engine/Source/Runtime/Renderer/Private/TranslucentRendering.cpp
+++ b/Engine/Source/Runtime/Renderer/Private/TranslucentRendering.cpp
@@ -988,6 +988,10 @@ void FDeferredShadingSceneRenderer::RenderTranslucencyInner(FRHICommandListImmed
 				BeginTimingSeparateTranslucencyModulatePass(RHICmdList, View);
 				SceneContext.BeginRenderingSeparateTranslucencyModulate(RHICmdList, View, *this, NumProcessedViews == 0 || View.Family->bMultiGPUForkAndJoin);
 			}
+			else
+			{
+				SceneContext.BeginRenderingSeparateTranslucency(RHICmdList, View, *this, NumProcessedViews == 0 || View.Family->bMultiGPUForkAndJoin);
+			}
 
 			// Draw only translucent prims that are in the SeparateTranslucency pass
 			DrawRenderState.SetDepthStencilState(TStaticDepthStencilState<false, CF_DepthNearOrEqual>::GetRHI());
