Fix a crash bug of Histogram compute shader on A8 devcies.

diff --git a/Engine/Source/Developer/Apple/MetalShaderFormat/Private/MetalShaderFormat.cpp b/Engine/Source/Developer/Apple/MetalShaderFormat/Private/MetalShaderFormat.cpp
--- a/Engine/Source/Developer/Apple/MetalShaderFormat/Private/MetalShaderFormat.cpp
+++ b/Engine/Source/Developer/Apple/MetalShaderFormat/Private/MetalShaderFormat.cpp
@@ -40,7 +40,7 @@ class FMetalShaderFormat : public IShaderFormat
 public:
 	enum
 	{
-		HEADER_VERSION = 68,
+		HEADER_VERSION = 69,
 	};
 	
 	struct FVersion
diff --git a/Engine/Source/Runtime/Renderer/Private/PostProcess/PostProcessMobile.cpp b/Engine/Source/Runtime/Renderer/Private/PostProcess/PostProcessMobile.cpp
--- a/Engine/Source/Runtime/Renderer/Private/PostProcess/PostProcessMobile.cpp
+++ b/Engine/Source/Runtime/Renderer/Private/PostProcess/PostProcessMobile.cpp
@@ -2941,9 +2941,17 @@ class FPostProcessHistogramCS_ES2 : public FGlobalShader
 public:
 	// Changing these numbers requires PostProcessMobile.usf to be recompiled.
 	static const uint32 ThreadGroupSizeX = 8;
+#if PLATFORM_IOS || PLATFORM_TVOS	// the maximum total threadgroup memory allocation on A7 and A8 GPU is 16KB-32B, so it has to limit the thread group size on IOS/TVOS platform.
+	static const uint32 ThreadGroupSizeY = 4;
+#else
 	static const uint32 ThreadGroupSizeY = 8;
+#endif
 	static const uint32 LoopCountX = 2;
+#if PLATFORM_IOS || PLATFORM_TVOS
+	static const uint32 LoopCountY = 4;
+#else
 	static const uint32 LoopCountY = 2;
+#endif
 	static const uint32 HistogramSize = 64; // HistogramSize must be 64 and ThreadGroupSizeX * ThreadGroupSizeY must be larger than 32
 
 	// The number of texels on each axis processed by a single thread group.
