Fix a crash bug if mobile movable directional light is disabled.

diff --git a/Engine/Source/Runtime/Renderer/Private/MobileBasePass.cpp b/Engine/Source/Runtime/Renderer/Private/MobileBasePass.cpp
--- a/Engine/Source/Runtime/Renderer/Private/MobileBasePass.cpp
+++ b/Engine/Source/Runtime/Renderer/Private/MobileBasePass.cpp
@@ -334,7 +334,7 @@ ELightMapPolicyType MobileBasePass::SelectMeshLightmapPolicy(
 				}
 			}
 		}
-		else if (bMovableWithCSM || bUseStaticAndCSM)
+		else if (bMovableWithCSM || (bUseStaticAndCSM && ReadOnlyCVARCache.bMobileAllowMovableDirectionalLights))
 		{
 			// final determination of whether CSMs are rendered can be view dependent, thus we always need to clear the CSMs even if we're not going to render to them based on the condition below.
 			SelectedLightmapPolicy = LMP_MOBILE_MOVABLE_DIRECTIONAL_LIGHT_CSM;
