Fixed static analysis warning

diff --git a/Engine/Source/Runtime/Renderer/Private/RayTracing/RayTracingDeferredReflections.cpp b/Engine/Source/Runtime/Renderer/Private/RayTracing/RayTracingDeferredReflections.cpp
--- a/Engine/Source/Runtime/Renderer/Private/RayTracing/RayTracingDeferredReflections.cpp
+++ b/Engine/Source/Runtime/Renderer/Private/RayTracing/RayTracingDeferredReflections.cpp
@@ -229,7 +229,7 @@ void FDeferredShadingSceneRenderer::RenderRayTracingDeferredReflections(
 
 	const uint32 TileAlignedNumRays          = TileAlignedResolution.X * TileAlignedResolution.Y;
 	const FRDGBufferDesc SortedRayBufferDesc = FRDGBufferDesc::CreateStructuredDesc(sizeof(FSortedReflectionRay), TileAlignedNumRays);
-	FRDGBufferRef SortedRayBuffer            = SortedRayBuffer = GraphBuilder.CreateBuffer(SortedRayBufferDesc, TEXT("ReflectionRayBuffer"));
+	FRDGBufferRef SortedRayBuffer            = GraphBuilder.CreateBuffer(SortedRayBufferDesc, TEXT("ReflectionRayBuffer"));
 
 	const FRDGBufferDesc DeferredMaterialBufferDesc = FRDGBufferDesc::CreateStructuredDesc(sizeof(FDeferredMaterialPayload), TileAlignedNumRays);
 	FRDGBufferRef DeferredMaterialBuffer            = GraphBuilder.CreateBuffer(DeferredMaterialBufferDesc, TEXT("RayTracingReflectionsMaterialBuffer"));
