Tweaks for GPUBenchmark on Mac. On AMD use the timestamp render queries emulation (using the command buffer completion handlers). On Intel and NVIDIA we still hardcode values based on available video memory, but these were changed to set low settings on < 2GB devices and medium on < 4GB ones.

diff --git a/Engine/Source/Runtime/Apple/MetalRHI/Private/MetalRHI.cpp b/Engine/Source/Runtime/Apple/MetalRHI/Private/MetalRHI.cpp
--- a/Engine/Source/Runtime/Apple/MetalRHI/Private/MetalRHI.cpp
+++ b/Engine/Source/Runtime/Apple/MetalRHI/Private/MetalRHI.cpp
@@ -260,6 +260,9 @@ FMetalDynamicRHI::FMetalDynamicRHI(ERHIFeatureLevel::Type RequestedFeatureLevel)
 		bSupportsTiledReflections = true;
 		bSupportsDistanceFields = (FPlatformMisc::MacOSXVersionCompare(10,11,4) >= 0);
 		bSupportsRHIThread = true;
+		
+		// Current timestamp render query implementation using a completion handler only works reliably on AMD
+		GSupportsTimestampRenderQueries = true;
 	}
 	else if(GRHIAdapterName.Contains("Intel"))
 	{
diff --git a/Engine/Source/Runtime/Renderer/Private/GPUBenchmark.cpp b/Engine/Source/Runtime/Renderer/Private/GPUBenchmark.cpp
--- a/Engine/Source/Runtime/Renderer/Private/GPUBenchmark.cpp
+++ b/Engine/Source/Runtime/Renderer/Private/GPUBenchmark.cpp
@@ -486,25 +486,21 @@ void RendererGPUBenchmark(FRHICommandListImmediate& RHICmdList, FSynthBenchmarkR
 			UE_LOG(LogSynthBenchmark, Warning, TEXT("GPU driver does not support timer queries."));
 
 #else
-			// Workaround for Metal not having a timing API and some drivers not properly supporting command-buffer completion handler based implementation...
+			// Workaround for Metal not having a timing API and Intel and NVIDIA drivers not properly supporting command-buffer completion handler based implementation...
+			// On AMD GSupportsTimestampRenderQueries is true
 			FTextureMemoryStats MemStats;
 			RHIGetTextureMemoryStats(MemStats);
 			
 			float PerfScale = 1.0f;
-			if(MemStats.TotalGraphicsMemory < (2ll * 1024ll * 1024ll * 1024ll))
+			if (MemStats.TotalGraphicsMemory < (2ll * 1024ll * 1024ll * 1024ll))
 			{
-				// Assume Intel HD 5000, Iris, Iris Pro performance - not dreadful
-				PerfScale = 4.2f;
+				// Assume Intel HD 5000, Iris, Iris Pro performance, or low end NVIDIA - low settings
+				PerfScale = 6.2f;
 			}
-			else if(MemStats.TotalGraphicsMemory < (3ll * 1024ll * 1024ll * 1024ll))
+			else if (MemStats.TotalGraphicsMemory < (3ll * 1024ll * 1024ll * 1024ll))
 			{
-				// Assume Nvidia 6x0 & 7x0 series/AMD M370X or Radeon Pro 4x0 series - mostly OK
-				PerfScale = 2.0f;
-			}
-			else
-			{
-				// AMD 7xx0 & Dx00 series - should be pretty beefy
-				PerfScale = 1.2f;
+				// Assume NVIDIA 6x0 & 7x0 series - medium settings
+				PerfScale = 3.0f;
 			}
 
 			for (int32 Index = 0; Index < MethodCount; ++Index)
