Manually integrate CL 11225060

Ensure screen messages respect GAreScreenMessagesEnabled

#rnx

diff --git a/Engine/Source/Runtime/Renderer/Private/SceneRendering.cpp b/Engine/Source/Runtime/Renderer/Private/SceneRendering.cpp
--- a/Engine/Source/Runtime/Renderer/Private/SceneRendering.cpp
+++ b/Engine/Source/Runtime/Renderer/Private/SceneRendering.cpp
@@ -2746,7 +2746,7 @@ void FSceneRenderer::RenderFinish(FRHICommandListImmediate& RHICmdList)
 				FSceneViewState* ViewState = (FSceneViewState*)View.State;
 				bool bViewParentOrFrozen = ViewState && (ViewState->HasViewParent() || ViewState->bIsFrozen);
 				bool bLocked = View.bIsLocked;
-				if (bViewParentOrFrozen || bLocked || bAnyWarning)
+				if (GAreScreenMessagesEnabled && (bViewParentOrFrozen || bLocked || bAnyWarning))
 				{
 					SCOPED_CONDITIONAL_DRAW_EVENTF(RHICmdList, EventView, Views.Num() > 1, TEXT("View%d"), ViewIndex);
 
