Allow XR plugins to specify which view to use for LOD
Fix shadow setup for more than 2 views

diff --git a/Engine/Plugins/Runtime/OpenXR/Source/OpenXRHMD/Private/OpenXRHMD.cpp b/Engine/Plugins/Runtime/OpenXR/Source/OpenXRHMD/Private/OpenXRHMD.cpp
--- a/Engine/Plugins/Runtime/OpenXR/Source/OpenXRHMD/Private/OpenXRHMD.cpp
+++ b/Engine/Plugins/Runtime/OpenXR/Source/OpenXRHMD/Private/OpenXRHMD.cpp
@@ -715,6 +715,15 @@ uint32 FOpenXRHMD::GetViewIndexForPass(EStereoscopicPass StereoPassType) const
 	}
 }
 
+uint32 FOpenXRHMD::DeviceGetLODViewIndex() const
+{
+	if (SelectedViewConfigurationType == XR_VIEW_CONFIGURATION_TYPE_PRIMARY_QUAD_VARJO)
+	{
+		return GetViewIndexForPass(eSSP_LEFT_EYE_SIDE);
+	}
+	return IStereoRendering::DeviceGetLODViewIndex();
+}
+
 int32 FOpenXRHMD::GetDesiredNumberOfViews(bool bStereoRequested) const
 {
 	return bStereoRequested ? Views.Num() : 1; // FIXME: Monoscopic actually needs 2 views for quad vr
diff --git a/Engine/Plugins/Runtime/OpenXR/Source/OpenXRHMD/Private/OpenXRHMD.h b/Engine/Plugins/Runtime/OpenXR/Source/OpenXRHMD/Private/OpenXRHMD.h
--- a/Engine/Plugins/Runtime/OpenXR/Source/OpenXRHMD/Private/OpenXRHMD.h
+++ b/Engine/Plugins/Runtime/OpenXR/Source/OpenXRHMD/Private/OpenXRHMD.h
@@ -142,6 +142,7 @@ public:
 	virtual int32 GetDesiredNumberOfViews(bool bStereoRequested) const override;
 	virtual EStereoscopicPass GetViewPassForIndex(bool bStereoRequested, uint32 ViewIndex) const override;
 	virtual uint32 GetViewIndexForPass(EStereoscopicPass StereoPassType) const override;
+	virtual uint32 DeviceGetLODViewIndex() const override;
 	
 	virtual FMatrix GetStereoProjectionMatrix(const enum EStereoscopicPass StereoPassType) const override;
 	virtual void GetEyeRenderParams_RenderThread(const struct FRenderingCompositePassContext& Context, FVector2D& EyeToSrcUVScaleValue, FVector2D& EyeToSrcUVOffsetValue) const override;
diff --git a/Engine/Source/Runtime/Engine/Private/SceneManagement.cpp b/Engine/Source/Runtime/Engine/Private/SceneManagement.cpp
--- a/Engine/Source/Runtime/Engine/Private/SceneManagement.cpp
+++ b/Engine/Source/Runtime/Engine/Private/SceneManagement.cpp
@@ -528,9 +528,10 @@ int8 ComputeTemporalStaticMeshLOD( const FStaticMeshRenderData* RenderData, cons
 // Ensure we always use the left eye when selecting lods to avoid divergent selections in stereo
 const FSceneView& GetLODView(const FSceneView& InView)
 {
-	if (IStereoRendering::IsASecondaryView(InView) && InView.Family)
+	uint32 LODViewIndex = IStereoRendering::GetLODViewIndex();
+	if (InView.Family && InView.Family->Views.IsValidIndex(LODViewIndex))
 	{
-		return *InView.Family->Views[0];
+		return *InView.Family->Views[LODViewIndex];
 	}
 	else
 	{
diff --git a/Engine/Source/Runtime/Engine/Private/StereoRendering.cpp b/Engine/Source/Runtime/Engine/Private/StereoRendering.cpp
--- a/Engine/Source/Runtime/Engine/Private/StereoRendering.cpp
+++ b/Engine/Source/Runtime/Engine/Private/StereoRendering.cpp
@@ -105,3 +105,17 @@ bool IStereoRendering::DeviceIsAnAdditionalView(const FSceneView& View)
 {
 	return View.StereoPass > EStereoscopicPass::eSSP_RIGHT_EYE;
 }
+
+uint32 IStereoRendering::GetLODViewIndex()
+{
+	if (GEngine->StereoRenderingDevice.IsValid())
+	{
+		return GEngine->StereoRenderingDevice->DeviceGetLODViewIndex();
+	}
+	return 0;
+}
+
+uint32 IStereoRendering::DeviceGetLODViewIndex() const
+{
+	return 0;
+}
diff --git a/Engine/Source/Runtime/Engine/Public/StereoRendering.h b/Engine/Source/Runtime/Engine/Public/StereoRendering.h
--- a/Engine/Source/Runtime/Engine/Public/StereoRendering.h
+++ b/Engine/Source/Runtime/Engine/Public/StereoRendering.h
@@ -177,6 +177,16 @@ public:
 	*/
 	virtual bool DeviceIsAnAdditionalView(const FSceneView& View);
 
+	/**
+	 * Static helper. Passes request to the current stereo device
+	 */
+	static uint32 GetLODViewIndex();
+
+	/**
+	 * Return the index of the view that is used for selecting LODs
+	 */
+	virtual uint32 DeviceGetLODViewIndex() const;
+
 	/**
 	 * Adjusts the viewport rectangle for stereo, based on which eye pass is being rendered.
 	 */
diff --git a/Engine/Source/Runtime/Renderer/Private/ShadowSetup.cpp b/Engine/Source/Runtime/Renderer/Private/ShadowSetup.cpp
--- a/Engine/Source/Runtime/Renderer/Private/ShadowSetup.cpp
+++ b/Engine/Source/Runtime/Renderer/Private/ShadowSetup.cpp
@@ -3458,8 +3458,7 @@ void FSceneRenderer::InitProjectedShadowVisibility(FRHICommandListImmediate& RHI
 					// The view dependent projected shadow is valid for this view if it's the
 					// right eye and the projected shadow is being rendered for the left eye.
 					const bool bIsValidForView = IStereoRendering::IsASecondaryView(View)
-						&& Views.IsValidIndex(ViewIndex - 1)
-						&& IStereoRendering::IsAPrimaryView(Views[ViewIndex - 1])
+						&& IStereoRendering::IsAPrimaryView(*ProjectedShadowInfo.DependentView)
 						&& ProjectedShadowInfo.FadeAlphas.IsValidIndex(ViewIndex)
 						&& ProjectedShadowInfo.FadeAlphas[ViewIndex] == 1.0f;
 
@@ -4125,11 +4124,20 @@ void FSceneRenderer::AddViewDependentWholeSceneShadowsForView(
 		FadeAlphas.Init(0.0f, Views.Num());
 		FadeAlphas[ViewIndex] = LightShadowAmount;
 
-		if (IStereoRendering::IsAPrimaryView(View)
-			&& Views.IsValidIndex(ViewIndex + 1)
-			&& IStereoRendering::IsASecondaryView(Views[ViewIndex + 1]))
+		if (IStereoRendering::IsAPrimaryView(View))
 		{
-			FadeAlphas[ViewIndex + 1] = LightShadowAmount;
+			for (int FadeAlphaIndex = ViewIndex + 1; FadeAlphaIndex < Views.Num(); FadeAlphaIndex++)
+			{
+				if (Views.IsValidIndex(FadeAlphaIndex)
+					&& IStereoRendering::IsASecondaryView(Views[FadeAlphaIndex]))
+				{
+					FadeAlphas[FadeAlphaIndex] = LightShadowAmount;
+				}
+				else if (IStereoRendering::IsAPrimaryView(Views[FadeAlphaIndex]))
+				{
+					break;
+				}
+			}
 		}		
 		
 		// If rendering in stereo mode we render shadow depths only for the left eye, but project for both eyes!
