Fix missing deep opacity on part of meshes.

This CL address the issue by:
* Move deep shadow resources onto per view basis
* Add GPU AABB driven generation of the world to light transform and min strands radius.

diff --git a/Engine/Shaders/Private/HairStrands/HairStrandsAABBCommon.ush b/Engine/Shaders/Private/HairStrands/HairStrandsAABBCommon.ush
--- a/Engine/Shaders/Private/HairStrands/HairStrandsAABBCommon.ush
+++ b/Engine/Shaders/Private/HairStrands/HairStrandsAABBCommon.ush
@@ -8,6 +8,16 @@ struct FHairAABB
 	float3 Max;
 };
 
+float3 GetCenter(const FHairAABB In)
+{
+	return (In.Min + In.Max) * 0.5f;
+}
+
+float3 GetExtents(const FHairAABB In)
+{
+	return (In.Max - In.Min) * 0.5f;
+}
+
 FHairAABB InitHairAABB()
 {
 	FHairAABB O;
diff --git a/Engine/Shaders/Private/HairStrands/HairStrandsDeepShadowAllocation.usf b/Engine/Shaders/Private/HairStrands/HairStrandsDeepShadowAllocation.usf
new file mode 100644
--- /dev/null
+++ b/Engine/Shaders/Private/HairStrands/HairStrandsDeepShadowAllocation.usf
@@ -0,0 +1,124 @@
+// Copyright Epic Games, Inc. All Rights Reserved.
+
+//#include "HairStrandsVoxelPageCommonStruct.ush"
+
+#include "../Common.ush"
+#include "../Matrices.ush"
+#include "HairStrandsDeepShadowCommonStruct.ush"
+#include "HairStrandsAABBCommon.ush"
+
+#if SHADER_ALLOCATE
+
+#ifndef MAX_SLOT_COUNT
+#error MAX_SLOT_COUNT needs to be defined
+#endif
+
+float4				LightDirections[MAX_SLOT_COUNT];
+float4				LightPositions[MAX_SLOT_COUNT];
+uint4				MacroGroupIndices[MAX_SLOT_COUNT];
+
+// For testing parity with CPU version
+float3				CPU_MinAABB;
+uint				CPU_bUseCPUData;
+float3				CPU_MaxAABB;
+float				RasterizationScale;
+
+int2				SlotResolution;
+uint				SlotIndexCount;
+uint				MacroGroupCount;
+
+Buffer<int>								MacroGroupAABBBuffer;
+RWStructuredBuffer<FDeepShadowViewInfo>	OutShadowViewInfoBuffer;
+RWStructuredBuffer<FDeepShadowTransform>OutShadowWorldToLightTransformBuffer;
+
+
+float DegreesToRadians(float InDeg)
+{
+	return InDeg / 180.f * PI;
+}
+
+float ComputeMinStrandRadiusAtDepth1(const int2 Resolution, const float FOV, const float InRasterizationScale)
+{
+	const float DiameterToRadius = 0.5f;
+	const float vFOV = DegreesToRadians(FOV);
+	const float StrandDiameterAtDepth1 = tan(vFOV * 0.5f) / (0.5f * Resolution.y);
+	return DiameterToRadius * InRasterizationScale * StrandDiameterAtDepth1;
+}
+
+// This function is equivalent of the CPU version in HairStrandsDeepShadow.cpp
+// TODO: create a tighter bound
+void ComputeWorldToLightClip(
+	inout float4x4	OutWorldToClipTransform,
+	inout float		OutMinStrandRadiusAtDepth1,
+	inout float3	OutLightDirection,
+	const FHairAABB Aabb,
+	const uint		SlotIndex)
+{
+	const float3 Extents = GetExtents(Aabb);
+	const float3 Center  = GetCenter(Aabb);
+	const float  Radius  = length(Extents);
+
+	const float3 LightDirection = LightDirections[SlotIndex].xyz;
+	const float3 LightPosition  = LightPositions[SlotIndex].xyz;
+	const uint bIsDirectional	= LightPositions[SlotIndex].w == 0;
+
+	const float MinZ = max(0.1f, length(LightPosition - Center)) - Radius;
+	const float MaxZ = max(0.2f, length(LightPosition - Center)) + Radius;
+
+	const float StrandHairStableRasterizationScale = max(RasterizationScale, 1.0f);
+	OutMinStrandRadiusAtDepth1 = 1;
+	OutWorldToClipTransform = 0;
+
+	if (bIsDirectional) // (LightType == LightType_Directional)
+	{
+		OutLightDirection			= LightDirection;
+
+		const float4x4 ProjMatrix	= ReversedZOrthoMatrix(Radius, Radius, 1.f / (2 * Radius), 0);
+		const float4x4 WorldToLight	= LookAtMatrix(Center - LightDirection * Radius, Center, float3(0, 0, 1));
+		OutWorldToClipTransform		= mul(WorldToLight, ProjMatrix);
+
+		const float RadiusAtDepth1	= Radius / min(SlotResolution.x, SlotResolution.y);
+		OutMinStrandRadiusAtDepth1	= RadiusAtDepth1 * RasterizationScale;
+	}
+	else // if (LightType == LightType_Spot || LightType == LightType_Point || LightType == LightType_Rect)
+	{
+		OutLightDirection			= normalize(LightPosition - Center);
+
+		const float SphereDistance	= length(LightPosition - Center);
+		const float HalfFov			= asin(Radius / SphereDistance);
+
+		const float4x4 ProjMatrix	= ReversedZPerspectiveMatrix(HalfFov, 1, 1, MinZ, MaxZ);
+		const float4x4 WorldToLight = LookAtMatrix(LightPosition, Center, float3(0, 0, 1));
+		OutWorldToClipTransform		= mul(WorldToLight, ProjMatrix);
+
+		OutMinStrandRadiusAtDepth1 = ComputeMinStrandRadiusAtDepth1(SlotResolution, 2 * HalfFov, RasterizationScale);
+	}
+}
+
+// This code assume we have less than 32 macro group (which fit into a single CU/SM)
+[numthreads(MAX_SLOT_COUNT, 1, 1)]
+void CreateViewInfo(uint2 DispatchThreadId : SV_DispatchThreadID)
+{
+	const uint SlotIndex = DispatchThreadId.x;
+
+	if (SlotIndex > SlotIndexCount)
+		return;
+
+	const uint MacroGroupId = MacroGroupIndices[SlotIndex].x;
+
+	FHairAABB Bound = InitHairAABB();
+	if (MacroGroupId < MacroGroupCount)
+	{
+		Bound = ReadHairAABB(MacroGroupId, MacroGroupAABBBuffer);
+	}
+
+	FDeepShadowViewInfo ViewInfo;
+	ComputeWorldToLightClip(ViewInfo.WorldToClip, ViewInfo.MinRadiusAtDepth1, ViewInfo.ViewForward, Bound, SlotIndex);
+
+	FDeepShadowTransform Transform;
+	Transform.WorldToClip = ViewInfo.WorldToClip;
+
+	OutShadowViewInfoBuffer[SlotIndex] = ViewInfo;
+	OutShadowWorldToLightTransformBuffer[SlotIndex] = Transform;
+}
+#endif
diff --git a/Engine/Shaders/Private/HairStrands/HairStrandsDeepShadowCommonStruct.ush b/Engine/Shaders/Private/HairStrands/HairStrandsDeepShadowCommonStruct.ush
new file mode 100644
--- /dev/null
+++ b/Engine/Shaders/Private/HairStrands/HairStrandsDeepShadowCommonStruct.ush
@@ -0,0 +1,15 @@
+// Copyright Epic Games, Inc. All Rights Reserved.
+
+#pragma once
+
+struct FDeepShadowViewInfo
+{
+	float4x4 WorldToClip;
+	float3   ViewForward;
+	float    MinRadiusAtDepth1;
+}; 
+
+struct FDeepShadowTransform
+{
+	float4x4 WorldToClip; // Store the transform into a struct to avoid the matrix to be loaded as transposed
+};
\ No newline at end of file
diff --git a/Engine/Shaders/Private/HairStrands/HairStrandsDeepShadowDebug.usf b/Engine/Shaders/Private/HairStrands/HairStrandsDeepShadowDebug.usf
--- a/Engine/Shaders/Private/HairStrands/HairStrandsDeepShadowDebug.usf
+++ b/Engine/Shaders/Private/HairStrands/HairStrandsDeepShadowDebug.usf
@@ -4,6 +4,10 @@
 #include "../PostProcessCommon.ush"
 #include "HairStrandsDeepShadowCommon.ush"
 
+//////////////////////////////////////////////////////////////////////////////////////////
+
+#if SHADER_VISUALIZEDOM
+
 float DomScale = 1;
 float2 DomAtlasOffset;
 float2 DomAtlasScale;
@@ -90,7 +94,7 @@ void VisualizeDomPS(
 
 	OutColor = float4(Color, 1);
 }
-#endif
+#endif // PERMUTATION_OUTPUT_TYPE == 0
 
 // Display screen space rect for vizualizing projected viewrect of each hair macro group
 #if PERMUTATION_OUTPUT_TYPE == 1
@@ -114,4 +118,81 @@ void VisualizeDomPS(
 
 	OutColor = float4(Color, 1);
 }
-#endif
\ No newline at end of file
+#endif // PERMUTATION_OUTPUT_TYPE == 1
+
+#endif // SHADER_VISUALIZEDOM
+
+
+//////////////////////////////////////////////////////////////////////////////////////////
+
+#if SHADER_DOMINFO
+
+#include "../ShaderPrintCommon.ush"
+#include "../ShaderDrawDebug.ush"
+#include "../Matrices.ush"
+#include "HairStrandsDeepShadowCommonStruct.ush"
+#include "HairStrandsAABBCommon.ush"
+
+uint AllocatedSlotCount;
+uint MacroGroupCount;
+
+Buffer<int>	MacroGroupAABBBuffer;
+StructuredBuffer<FDeepShadowTransform> ShadowWorldToLightTransformBuffer;
+
+[numthreads(1, 1, 1)]
+void MainCS(uint3 DispatchThreadId : SV_DispatchThreadID)
+{
+	if (DispatchThreadId.x > 0)
+	{
+		return;
+	}
+
+	for (uint MacroGroupIt = 0; MacroGroupIt < MacroGroupCount; ++MacroGroupIt)
+	{
+		const FHairAABB Bound = ReadHairAABB(MacroGroupIt, MacroGroupAABBBuffer);
+		AddAABB(Bound.Min, Bound.Max, float4(1, 0.5f, 0, 1));
+	}
+
+	const float FrustumColor = float4(0.5f, 1, 0, 1);
+	for (uint SlotIt = 0; SlotIt < AllocatedSlotCount; ++SlotIt)
+	{
+		const float4x4 WorldToClip = ShadowWorldToLightTransformBuffer[SlotIt].WorldToClip;
+		const float4x4 ClipToWorld = Inverse(WorldToClip);
+
+		float4 P01 = float4(-1, -1, 0, 1);
+		float4 P11 = float4(-1,  1, 0, 1);
+		float4 P21 = float4( 1,  1, 0, 1);
+		float4 P31 = float4( 1, -1, 0, 1);
+
+		float4 P00 = float4(-1, -1, 1, 1);
+		float4 P10 = float4(-1,  1, 1, 1);
+		float4 P20 = float4( 1,  1, 1, 1);
+		float4 P30 = float4( 1, -1, 1, 1);
+
+		P00 = mul(P00, ClipToWorld); P00 /= P00.w; 
+		P10 = mul(P10, ClipToWorld); P10 /= P10.w; 
+		P20 = mul(P20, ClipToWorld); P20 /= P20.w; 
+		P30 = mul(P30, ClipToWorld); P30 /= P30.w; 
+
+		P01 = mul(P01, ClipToWorld); P01 /= P01.w;
+		P11 = mul(P11, ClipToWorld); P11 /= P11.w;
+		P21 = mul(P21, ClipToWorld); P21 /= P21.w;
+		P31 = mul(P31, ClipToWorld); P31 /= P31.w;
+
+		AddLine(P00, P10, FrustumColor, FrustumColor);
+		AddLine(P10, P20, FrustumColor, FrustumColor);
+		AddLine(P20, P30, FrustumColor, FrustumColor);
+		AddLine(P30, P00, FrustumColor, FrustumColor);
+
+		AddLine(P00, P01, FrustumColor, FrustumColor);
+		AddLine(P10, P11, FrustumColor, FrustumColor);
+		AddLine(P20, P21, FrustumColor, FrustumColor);
+		AddLine(P30, P31, FrustumColor, FrustumColor);
+
+		AddLine(P01, P11, FrustumColor, FrustumColor);
+		AddLine(P11, P21, FrustumColor, FrustumColor);
+		AddLine(P21, P31, FrustumColor, FrustumColor);
+		AddLine(P31, P01, FrustumColor, FrustumColor);
+	}
+}
+#endif // SHADER_DOMINFO
diff --git a/Engine/Shaders/Private/HairStrands/HairStrandsDeepShadowMask.usf b/Engine/Shaders/Private/HairStrands/HairStrandsDeepShadowMask.usf
--- a/Engine/Shaders/Private/HairStrands/HairStrandsDeepShadowMask.usf
+++ b/Engine/Shaders/Private/HairStrands/HairStrandsDeepShadowMask.usf
@@ -1,6 +1,7 @@
 // Copyright Epic Games, Inc. All Rights Reserved.
 
 #include "HairStrandsVoxelPageCommonStruct.ush"
+#include "HairStrandsDeepShadowCommonStruct.ush"
 
 #include "../Common.ush"
 #include "../CommonViewUniformBuffer.ush"
@@ -29,10 +30,14 @@
   #include "HairStrandsVoxelPageTraversal.ush"
 #endif
 
-float4x4			DeepShadow_WorldToLightTransform;
+StructuredBuffer<FDeepShadowTransform> DeepShadow_WorldToLightTransformBuffer;
+uint				DeepShadow_bIsGPUDriven;
+
+float4x4			DeepShadow_CPUWorldToLightTransform;
 uint				DeepShadow_IsWholeSceneLight;
 int2				DeepShadow_SlotResolution;
 int2				DeepShadow_SlotOffset;
+uint				DeepShadow_SlotIndex;
 uint				DeepShadow_KernelType;
 float				DeepShadow_DepthBiasScale;
 float				DeepShadow_DensityScale;
@@ -76,7 +81,17 @@ float3 TotWorldPosition(float2 UV, float Depth)
 
 float3 ToLightPosition(float3 WorldPosition)
 {
-    float4 LightPos = mul(float4(WorldPosition,1), DeepShadow_WorldToLightTransform);
+	float4x4 WorldToLightTransform = 0;
+	if (DeepShadow_bIsGPUDriven > 0)
+	{
+		WorldToLightTransform = DeepShadow_WorldToLightTransformBuffer.Load(DeepShadow_SlotIndex).WorldToClip;
+	}
+	else
+	{
+		WorldToLightTransform = DeepShadow_CPUWorldToLightTransform;
+	}
+
+    float4 LightPos = mul(float4(WorldPosition,1), WorldToLightTransform);
 	LightPos.xyz/= LightPos.w;
     const float2 LightUV = (LightPos.xy + float(1).xx) * 0.5f;
     return float3(LightUV.x, 1-LightUV.y, LightPos.z);
diff --git a/Engine/Shaders/Private/HairStrands/HairStrandsDeepShadowPS.usf b/Engine/Shaders/Private/HairStrands/HairStrandsDeepShadowPS.usf
--- a/Engine/Shaders/Private/HairStrands/HairStrandsDeepShadowPS.usf
+++ b/Engine/Shaders/Private/HairStrands/HairStrandsDeepShadowPS.usf
@@ -1,6 +1,7 @@
 // Copyright Epic Games, Inc. All Rights Reserved.
 
 #include "HairStrandsVoxelPageCommonStruct.ush"
+#include "HairStrandsDeepShadowCommonStruct.ush"
 
 #include "../Common.ush"
 
diff --git a/Engine/Shaders/Private/HairStrands/HairStrandsDeepShadowVS.usf b/Engine/Shaders/Private/HairStrands/HairStrandsDeepShadowVS.usf
--- a/Engine/Shaders/Private/HairStrands/HairStrandsDeepShadowVS.usf
+++ b/Engine/Shaders/Private/HairStrands/HairStrandsDeepShadowVS.usf
@@ -1,6 +1,8 @@
 // Copyright Epic Games, Inc. All Rights Reserved.
 
 #include "HairStrandsVoxelPageCommonStruct.ush"
+#include "HairStrandsDeepShadowCommonStruct.ush"
+
 #include "../Common.ush"
 #include "/Engine/Generated/Material.ush"
 #include "/Engine/Generated/VertexFactory.ush"
@@ -95,7 +97,7 @@ void Main(
 #if MESH_RENDER_MODE == MESH_RENDER_VOXEL
 	float4x4 WorldToClipMatrix = VoxelRasterPass.WorldToClipMatrix;
 #elif MESH_RENDER_MODE == MESH_RENDER_DEPTH || MESH_RENDER_MODE == MESH_RENDER_DOM
-	float4x4 WorldToClipMatrix = DeepRasterPass.WorldToClipMatrix;
+	float4x4 WorldToClipMatrix = DeepRasterPass.CPU_WorldToClipMatrix;
 #endif
 
 #ifdef VF_STRAND_HAIR	
@@ -119,6 +121,17 @@ void Main(
 	}
 	#endif
 
+	#if MESH_RENDER_MODE == MESH_RENDER_DEPTH || MESH_RENDER_MODE == MESH_RENDER_DOM
+	if (HairRenderInfo.bIsGPUDriven)
+	{
+		const uint AtlasSlotIndex = DeepRasterPass.AtlasSlotIndex;
+		FDeepShadowViewInfo DeepShadowViewInfo = DeepRasterPass.DeepShadowViewInfoBuffer.Load(AtlasSlotIndex);
+		WorldToClipMatrix			= DeepShadowViewInfo.WorldToClip;
+		HairViewInfo.ViewForward	= DeepShadowViewInfo.ViewForward;
+		HairViewInfo.RadiusAtDepth1 = DeepShadowViewInfo.MinRadiusAtDepth1;
+	}
+	#endif
+
 	float4 WorldPosition = VertexFactoryGetWorldPosition(Input, VFIntermediates, HairViewInfo);
 #else
 	float4 WorldPosition = VertexFactoryGetWorldPosition(Input, VFIntermediates);
diff --git a/Engine/Shaders/Private/HairStrands/HairStrandsDeepTransmittanceMask.usf b/Engine/Shaders/Private/HairStrands/HairStrandsDeepTransmittanceMask.usf
--- a/Engine/Shaders/Private/HairStrands/HairStrandsDeepTransmittanceMask.usf
+++ b/Engine/Shaders/Private/HairStrands/HairStrandsDeepTransmittanceMask.usf
@@ -1,6 +1,7 @@
 // Copyright Epic Games, Inc. All Rights Reserved.
 
 #include "HairStrandsVoxelPageCommonStruct.ush"
+#include "HairStrandsDeepShadowCommonStruct.ush"
 
 #include "../Common.ush"
 #include "../CommonViewUniformBuffer.ush"
@@ -49,8 +50,11 @@ float3 ToWorldPosition(float2 UV, float Depth)
 // Need to match FDeepShadowData::MaxMacroGroupCount in HairRendering.h
 #define MAX_HAIR_MACROGROUP_COUNT 16 
 
+StructuredBuffer<FDeepShadowTransform> DeepShadow_WorldToLightTransformBuffer;
+uint				DeepShadow_bIsGPUDriven;
+
 uint2				DeepShadow_Resolution;
-float4x4			DeepShadow_WorldToLightTransforms[MAX_HAIR_MACROGROUP_COUNT];
+float4x4			DeepShadow_CPUWorldToLightTransforms[MAX_HAIR_MACROGROUP_COUNT];
 float3				LightDirection;
 uint				LightChannelMask;
 uint				MaxVisibilityNodeCount;
@@ -61,7 +65,7 @@ float				LightRadius;
 float				DeepShadow_KernelAperture;
 uint				DeepShadow_KernelType;
 uint				DeepShadow_DebugMode;
-int4				DeepShadow_AtlasSlotOffsets[MAX_HAIR_MACROGROUP_COUNT];
+int4				DeepShadow_AtlasSlotOffsets_AtlasSlotIndex[MAX_HAIR_MACROGROUP_COUNT];
 Texture2D<float>	DeepShadow_FrontDepthTexture;
 Texture2D<float4>	DeepShadow_DomTexture;
 
@@ -130,11 +134,19 @@ void MainCS(uint2 DispatchThreadId : SV_DispatchThreadID)
 		GBuffer.Roughness = Sample.Roughness;
 
 		const uint MacroGroupIndex = clamp(Sample.MacroGroupId, 0, MAX_HAIR_MACROGROUP_COUNT - 1);
-		const int2 DeepShadowOffset = DeepShadow_AtlasSlotOffsets[MacroGroupIndex].xy;
-		const float4x4 WorldToLightTransform = DeepShadow_WorldToLightTransforms[MacroGroupIndex];
+		const int3 SlotOffset_SlotIndex = DeepShadow_AtlasSlotOffsets_AtlasSlotIndex[MacroGroupIndex].xyz;
+		float4x4 WorldToLightTransform = 0;
+		if (DeepShadow_bIsGPUDriven > 0)
+		{
+			WorldToLightTransform = DeepShadow_WorldToLightTransformBuffer.Load(SlotOffset_SlotIndex.z).WorldToClip;
+		}	
+		else
+		{
+			WorldToLightTransform = DeepShadow_CPUWorldToLightTransforms[MacroGroupIndex];
+		}
 		
 		FTransmittanceSettings Settings;
-		Settings.DeepShadowOffset = DeepShadowOffset;
+		Settings.DeepShadowOffset = SlotOffset_SlotIndex.xy;
 		Settings.DeepShadowResolution = DeepShadow_Resolution;
 		Settings.WorldToLightTransform = WorldToLightTransform;
 		Settings.LightDirection = LightDirection;
diff --git a/Engine/Shaders/Private/Matrices.ush b/Engine/Shaders/Private/Matrices.ush
--- a/Engine/Shaders/Private/Matrices.ush
+++ b/Engine/Shaders/Private/Matrices.ush
@@ -15,6 +15,19 @@ float4x4 ReversedZOrthoMatrix(float Width, float Height, float ZScale, float ZOf
 	return O;	
 }
 
+// Creates an perspective projection matrix
+//
+// Similar to Core\Public\Math\PerspectiveMatrix.h
+float4x4 ReversedZPerspectiveMatrix(float HalfFOV, float Width, float Height, float MinZ, float MaxZ)
+{
+	float4x4 O;
+	O[0] = float4(1.0f / tan(HalfFOV), 0.0f, 0.0f, 0.0f);
+	O[1] = float4(0.0f, Width / tan(HalfFOV) / Height, 0.0f, 0.0f);
+	O[2] = float4(0.0f, 0.0f, ((MinZ == MaxZ) ? 0.0f : MinZ / (MinZ - MaxZ)), 1.0f);
+	O[3] = float4(0.0f, 0.0f, ((MinZ == MaxZ) ? MinZ : -MaxZ * MinZ / (MinZ - MaxZ)), 0.0f);
+	return O;
+}
+
 // Creates a view matrix given an eye position, a position to look at, and an up vector. 
 // This does the same thing as D3DXMatrixLookAtLH.
 //
@@ -41,4 +54,44 @@ float4x4 LookAtMatrix(const float3 EyePosition, const float3 LookAtPosition, con
 	O[3][3] = 1.0f;
 
 	return O;
-}
\ No newline at end of file
+}
+
+float4x4 Inverse(float4x4 m)
+{
+	float n11 = m[0][0], n12 = m[1][0], n13 = m[2][0], n14 = m[3][0];
+	float n21 = m[0][1], n22 = m[1][1], n23 = m[2][1], n24 = m[3][1];
+	float n31 = m[0][2], n32 = m[1][2], n33 = m[2][2], n34 = m[3][2];
+	float n41 = m[0][3], n42 = m[1][3], n43 = m[2][3], n44 = m[3][3];
+
+	float t11 = n23 * n34 * n42 - n24 * n33 * n42 + n24 * n32 * n43 - n22 * n34 * n43 - n23 * n32 * n44 + n22 * n33 * n44;
+	float t12 = n14 * n33 * n42 - n13 * n34 * n42 - n14 * n32 * n43 + n12 * n34 * n43 + n13 * n32 * n44 - n12 * n33 * n44;
+	float t13 = n13 * n24 * n42 - n14 * n23 * n42 + n14 * n22 * n43 - n12 * n24 * n43 - n13 * n22 * n44 + n12 * n23 * n44;
+	float t14 = n14 * n23 * n32 - n13 * n24 * n32 - n14 * n22 * n33 + n12 * n24 * n33 + n13 * n22 * n34 - n12 * n23 * n34;
+
+	float det = n11 * t11 + n21 * t12 + n31 * t13 + n41 * t14;
+	float idet = 1.0f / det;
+
+	float4x4 ret;
+
+	ret[0][0] = t11 * idet;
+	ret[0][1] = (n24 * n33 * n41 - n23 * n34 * n41 - n24 * n31 * n43 + n21 * n34 * n43 + n23 * n31 * n44 - n21 * n33 * n44) * idet;
+	ret[0][2] = (n22 * n34 * n41 - n24 * n32 * n41 + n24 * n31 * n42 - n21 * n34 * n42 - n22 * n31 * n44 + n21 * n32 * n44) * idet;
+	ret[0][3] = (n23 * n32 * n41 - n22 * n33 * n41 - n23 * n31 * n42 + n21 * n33 * n42 + n22 * n31 * n43 - n21 * n32 * n43) * idet;
+
+	ret[1][0] = t12 * idet;
+	ret[1][1] = (n13 * n34 * n41 - n14 * n33 * n41 + n14 * n31 * n43 - n11 * n34 * n43 - n13 * n31 * n44 + n11 * n33 * n44) * idet;
+	ret[1][2] = (n14 * n32 * n41 - n12 * n34 * n41 - n14 * n31 * n42 + n11 * n34 * n42 + n12 * n31 * n44 - n11 * n32 * n44) * idet;
+	ret[1][3] = (n12 * n33 * n41 - n13 * n32 * n41 + n13 * n31 * n42 - n11 * n33 * n42 - n12 * n31 * n43 + n11 * n32 * n43) * idet;
+
+	ret[2][0] = t13 * idet;
+	ret[2][1] = (n14 * n23 * n41 - n13 * n24 * n41 - n14 * n21 * n43 + n11 * n24 * n43 + n13 * n21 * n44 - n11 * n23 * n44) * idet;
+	ret[2][2] = (n12 * n24 * n41 - n14 * n22 * n41 + n14 * n21 * n42 - n11 * n24 * n42 - n12 * n21 * n44 + n11 * n22 * n44) * idet;
+	ret[2][3] = (n13 * n22 * n41 - n12 * n23 * n41 - n13 * n21 * n42 + n11 * n23 * n42 + n12 * n21 * n43 - n11 * n22 * n43) * idet;
+
+	ret[3][0] = t14 * idet;
+	ret[3][1] = (n13 * n24 * n31 - n14 * n23 * n31 + n14 * n21 * n33 - n11 * n24 * n33 - n13 * n21 * n34 + n11 * n23 * n34) * idet;
+	ret[3][2] = (n14 * n22 * n31 - n12 * n24 * n31 - n14 * n21 * n32 + n11 * n24 * n32 + n12 * n21 * n34 - n11 * n22 * n34) * idet;
+	ret[3][3] = (n12 * n23 * n31 - n13 * n22 * n31 + n13 * n21 * n32 - n11 * n23 * n32 - n12 * n21 * n33 + n11 * n22 * n33) * idet;
+
+	return ret;
+}
diff --git a/Engine/Source/Runtime/Renderer/Private/DeferredShadingRenderer.cpp b/Engine/Source/Runtime/Renderer/Private/DeferredShadingRenderer.cpp
--- a/Engine/Source/Runtime/Renderer/Private/DeferredShadingRenderer.cpp
+++ b/Engine/Source/Runtime/Renderer/Private/DeferredShadingRenderer.cpp
@@ -2230,7 +2230,7 @@ void FDeferredShadingSceneRenderer::Render(FRHICommandListImmediate& RHICmdList)
 
 		// Voxelization and Deep Opacity Maps
 		VoxelizeHairStrands(RHICmdList, Scene, Views, HairDatasStorage.MacroGroupsPerViews);
-		HairDatasStorage.DeepShadowViews = RenderHairStrandsDeepShadows(RHICmdList, Scene, Views, HairDatasStorage.MacroGroupsPerViews);
+		RenderHairStrandsDeepShadows(RHICmdList, Scene, Views, HairDatasStorage.MacroGroupsPerViews);
 
 		ServiceLocalQueue();
 
diff --git a/Engine/Source/Runtime/Renderer/Private/HairStrands/HairStrandsCluster.cpp b/Engine/Source/Runtime/Renderer/Private/HairStrands/HairStrandsCluster.cpp
--- a/Engine/Source/Runtime/Renderer/Private/HairStrands/HairStrandsCluster.cpp
+++ b/Engine/Source/Runtime/Renderer/Private/HairStrands/HairStrandsCluster.cpp
@@ -248,6 +248,7 @@ FHairStrandsMacroGroupViews CreateHairStrandsMacroGroups(
 				{				
 					AddHairMacroGroupAABBPass(GraphBuilder, MacroGroup, MacroGroupAABBBufferUAV);
 				}
+				MacroGroups.MacroGroupResources.MacroGroupCount = MacroGroups.Datas.Num();
 				GraphBuilder.QueueBufferExtraction(MacroGroupAABBBuffer, &MacroGroups.MacroGroupResources.MacroGroupAABBsBuffer, FRDGResourceState::EAccess::Read, FRDGResourceState::EPipeline::Compute);
 				GraphBuilder.Execute();
 			}
diff --git a/Engine/Source/Runtime/Renderer/Private/HairStrands/HairStrandsCluster.h b/Engine/Source/Runtime/Renderer/Private/HairStrands/HairStrandsCluster.h
--- a/Engine/Source/Runtime/Renderer/Private/HairStrands/HairStrandsCluster.h
+++ b/Engine/Source/Runtime/Renderer/Private/HairStrands/HairStrandsCluster.h
@@ -9,6 +9,7 @@
 #include "CoreMinimal.h"
 #include "RendererInterface.h"
 #include "HairStrandsVoxelization.h"
+#include "HairStrandsDeepShadow.h"
 #include "SceneManagement.h"
 
 // -----------------------------
@@ -25,6 +26,7 @@
 
 struct FHairMacroGroupAABBData
 {
+	uint32 MacroGroupCount = 0;
 	TRefCountPtr<FPooledRDGBuffer>	MacroGroupAABBsBuffer;
 };
 
@@ -57,6 +59,7 @@ struct FHairStrandsMacroGroupData
 
 	FVirtualVoxelNodeDesc VirtualVoxelNodeDesc;
 	FHairStrandsVoxelResources VoxelResources;
+	FHairStrandsDeepShadowDatas DeepShadowDatas;
 	TPrimitiveInfos PrimitivesInfos;
 	TPrimitiveGroups PrimitivesGroups;
 	FBoxSphereBounds Bounds;
@@ -68,6 +71,7 @@ struct FHairStrandsMacroGroupData
 struct FHairStrandsMacroGroupDatas
 {
 	TArray<FHairStrandsMacroGroupData, SceneRenderingAllocator> Datas;
+	FDeepShadowResources DeepShadowResources;
 	FVirtualVoxelResources VirtualVoxelResources;
 	FHairMacroGroupAABBData MacroGroupResources;
 };
diff --git a/Engine/Source/Runtime/Renderer/Private/HairStrands/HairStrandsDebug.cpp b/Engine/Source/Runtime/Renderer/Private/HairStrands/HairStrandsDebug.cpp
--- a/Engine/Source/Runtime/Renderer/Private/HairStrands/HairStrandsDebug.cpp
+++ b/Engine/Source/Runtime/Renderer/Private/HairStrands/HairStrandsDebug.cpp
@@ -488,6 +488,11 @@ class FDeepShadowVisualizePS : public FGlobalShader
 
 public:
 	static bool ShouldCompilePermutation(const FGlobalShaderPermutationParameters& Parameters) { return IsHairStrandsSupported(Parameters.Platform); }
+	static void ModifyCompilationEnvironment(const FGlobalShaderPermutationParameters& Parameters, FShaderCompilerEnvironment& OutEnvironment)
+	{
+		FGlobalShader::ModifyCompilationEnvironment(Parameters, OutEnvironment);
+		OutEnvironment.SetDefine(TEXT("SHADER_VISUALIZEDOM"), 1);
+	}
 };
 
 IMPLEMENT_GLOBAL_SHADER(FDeepShadowVisualizePS, "/Engine/Private/HairStrands/HairStrandsDeepShadowDebug.usf", "VisualizeDomPS", SF_Pixel);
@@ -497,6 +502,7 @@ static void AddDebugDeepShadowTexturePass(
 	const FViewInfo* View,
 	const FIntRect& HairViewRect,
 	const FHairStrandsDeepShadowData* ShadowData,
+	const FDeepShadowResources* Resources,
 	FRDGTextureRef& OutTarget)
 {
 	check(OutTarget);
@@ -508,8 +514,8 @@ static void AddDebugDeepShadowTexturePass(
 	FVector2D AltasScale(0, 0);
 	if (ShadowData)
 	{
-		DeepShadowDepthTexture = GraphBuilder.RegisterExternalTexture(ShadowData->DepthTexture, TEXT("DOMDepthTexture"));
-		DeepShadowLayerTexture = GraphBuilder.RegisterExternalTexture(ShadowData->LayersTexture, TEXT("DOMLayerTexture"));
+		DeepShadowDepthTexture = GraphBuilder.RegisterExternalTexture(Resources->DepthAtlasTexture, TEXT("DOMDepthTexture"));
+		DeepShadowLayerTexture = GraphBuilder.RegisterExternalTexture(Resources->LayersAtlasTexture, TEXT("DOMLayerTexture"));
 
 		AtlasResolution = FIntPoint(DeepShadowDepthTexture->Desc.Extent.X, DeepShadowDepthTexture->Desc.Extent.Y);
 		AltasOffset = FVector2D(ShadowData->AtlasRect.Min.X / float(AtlasResolution.X), ShadowData->AtlasRect.Min.Y / float(AtlasResolution.Y));
@@ -571,6 +577,72 @@ static void AddDebugDeepShadowTexturePass(
 	});
 }
 
+///////////////////////////////////////////////////////////////////////////////////////////////////
+class FDeepShadowInfoCS : public FGlobalShader
+{
+	DECLARE_GLOBAL_SHADER(FDeepShadowInfoCS);
+	SHADER_USE_PARAMETER_STRUCT(FDeepShadowInfoCS, FGlobalShader);
+
+	BEGIN_SHADER_PARAMETER_STRUCT(FParameters, )
+		SHADER_PARAMETER_STRUCT_INCLUDE(FSceneTextureParameters, SceneTextures)
+		SHADER_PARAMETER_STRUCT_INCLUDE(ShaderDrawDebug::FShaderDrawDebugParameters, ShaderDrawParameters)
+		SHADER_PARAMETER_STRUCT_INCLUDE(ShaderPrint::FShaderParameters, ShaderPrintParameters)
+		SHADER_PARAMETER(FVector2D, OutputResolution)
+		SHADER_PARAMETER(uint32, AllocatedSlotCount)
+		SHADER_PARAMETER(uint32, MacroGroupCount)
+		SHADER_PARAMETER_STRUCT_REF(FViewUniformShaderParameters, ViewUniformBuffer)
+		SHADER_PARAMETER_RDG_BUFFER_SRV(Buffer, MacroGroupAABBBuffer)
+		SHADER_PARAMETER_RDG_BUFFER_SRV(StructuredBuffer, ShadowWorldToLightTransformBuffer)
+		SHADER_PARAMETER_RDG_TEXTURE_UAV(RWTexture2D, OutputTexture)
+		END_SHADER_PARAMETER_STRUCT()
+
+public:
+	static bool ShouldCompilePermutation(const FGlobalShaderPermutationParameters& Parameters) { return IsHairStrandsSupported(Parameters.Platform); }
+	static void ModifyCompilationEnvironment(const FGlobalShaderPermutationParameters& Parameters, FShaderCompilerEnvironment& OutEnvironment)
+	{
+		FGlobalShader::ModifyCompilationEnvironment(Parameters, OutEnvironment);
+		OutEnvironment.SetDefine(TEXT("SHADER_DOMINFO"), 1);
+	}
+};
+
+IMPLEMENT_GLOBAL_SHADER(FDeepShadowInfoCS, "/Engine/Private/HairStrands/HairStrandsDeepShadowDebug.usf", "MainCS", SF_Compute);
+
+static void AddDeepShadowInfoPass(
+	FRDGBuilder& GraphBuilder,
+	const FViewInfo& View,
+	const FHairStrandsMacroGroupDatas& MacroGroupDatas,
+	FRDGTextureRef& OutputTexture)
+{
+	if (MacroGroupDatas.DeepShadowResources.TotalAtlasSlotCount == 0)
+	{
+		return;
+	}
+
+	FSceneTextureParameters SceneTextures;
+	SetupSceneTextureParameters(GraphBuilder, &SceneTextures);
+
+	FRDGBufferRef MacroGroupAABBsBuffer = GraphBuilder.RegisterExternalBuffer(MacroGroupDatas.MacroGroupResources.MacroGroupAABBsBuffer);
+	FRDGBufferRef DeepShadowWorldToLightTransforms = GraphBuilder.RegisterExternalBuffer(MacroGroupDatas.DeepShadowResources.DeepShadowWorldToLightTransforms);
+
+	const FIntPoint Resolution(OutputTexture->Desc.Extent);
+	FDeepShadowInfoCS::FParameters* Parameters = GraphBuilder.AllocParameters<FDeepShadowInfoCS::FParameters>();
+	Parameters->ViewUniformBuffer = View.ViewUniformBuffer;
+	Parameters->OutputResolution = Resolution;
+	Parameters->AllocatedSlotCount = MacroGroupDatas.DeepShadowResources.TotalAtlasSlotCount;
+	Parameters->MacroGroupCount = MacroGroupDatas.MacroGroupResources.MacroGroupCount;
+	Parameters->SceneTextures = SceneTextures;
+	Parameters->MacroGroupAABBBuffer = GraphBuilder.CreateSRV(MacroGroupAABBsBuffer, PF_R32_SINT);
+	Parameters->ShadowWorldToLightTransformBuffer = GraphBuilder.CreateSRV(DeepShadowWorldToLightTransforms);
+	ShaderDrawDebug::SetParameters(GraphBuilder, View.ShaderDrawData, Parameters->ShaderDrawParameters);
+	ShaderPrint::SetParameters(View, Parameters->ShaderPrintParameters);
+	Parameters->OutputTexture = GraphBuilder.CreateUAV(OutputTexture);
+
+	TShaderMapRef<FDeepShadowInfoCS> ComputeShader(View.ShaderMap);
+
+	const FIntVector DispatchCount = DispatchCount.DivideAndRoundUp(FIntVector(OutputTexture->Desc.Extent.X, OutputTexture->Desc.Extent.Y, 1), FIntVector(8, 8, 1));
+	FComputeShaderUtils::AddPass(GraphBuilder, RDG_EVENT_NAME("HairStrandsDeepShadowDebugInfo"), ComputeShader, Parameters, DispatchCount);
+}
+
 ///////////////////////////////////////////////////////////////////////////////////////////////////
 class FVoxelRaymarchingPS : public FGlobalShader
 {
@@ -1431,7 +1503,6 @@ void RenderHairStrandsDebugInfo(
 	if (!HairDatas)
 		return;
 
-	const FHairStrandsDeepShadowViews& InDomViews = HairDatas->DeepShadowViews;
 	const FHairStrandsMacroGroupViews& InMacroGroupViews = HairDatas->MacroGroupsPerViews;
 
 	if (HairDebugMode == EHairDebugMode::MacroGroups)
@@ -1488,23 +1559,30 @@ void RenderHairStrandsDebugInfo(
 
 	if (HairDebugMode == EHairDebugMode::DeepOpacityMaps)
 	{
-		const uint32 DomIndex = GDeepShadowDebugIndex;
-		TRefCountPtr<IPooledRenderTarget> DepthTexture;
-		TRefCountPtr<IPooledRenderTarget> LayerTexture;
-		const FHairStrandsDeepShadowDatas& DeepShadoDatas = InDomViews.Views[ViewIndex];
-		const bool bIsValid = DomIndex < uint32(DeepShadoDatas.Datas.Num());
-		if (bIsValid)
-		{				
-			DepthTexture = DeepShadoDatas.Datas[DomIndex].DepthTexture;
-			LayerTexture = DeepShadoDatas.Datas[DomIndex].LayersTexture;
-		}
-
-		if (DepthTexture && LayerTexture)
+		if (ViewIndex < uint32(InMacroGroupViews.Views.Num()))
 		{
-			const FHairStrandsDeepShadowData& DeepShadoData = DeepShadoDatas.Datas[DomIndex];
 			FRDGBuilder GraphBuilder(RHICmdList);
-			FRDGTextureRef SceneColorTexture = GraphBuilder.RegisterExternalTexture(SceneTargets.GetSceneColor(), TEXT("SceneColorTexture"));
-			AddDebugDeepShadowTexturePass(GraphBuilder, &View, FIntRect(), &DeepShadoData, SceneColorTexture);
+			const FHairStrandsMacroGroupDatas& MacroGroupDatas = InMacroGroupViews.Views[ViewIndex];
+			for (const FHairStrandsMacroGroupData& MacroGroup : MacroGroupDatas.Datas)
+			{
+				TRefCountPtr<IPooledRenderTarget> DepthTexture = MacroGroupDatas.DeepShadowResources.DepthAtlasTexture;
+				TRefCountPtr<IPooledRenderTarget> LayerTexture = MacroGroupDatas.DeepShadowResources.LayersAtlasTexture;
+
+				if (!DepthTexture || !LayerTexture)
+				{
+					continue;
+				}
+
+				for (const FHairStrandsDeepShadowData& DeepShadowData : MacroGroup.DeepShadowDatas.Datas)
+				{
+					const uint32 DomIndex = GDeepShadowDebugIndex;
+					if (DeepShadowData.AtlasSlotIndex != DomIndex)
+						continue;
+
+					FRDGTextureRef SceneColorTexture = GraphBuilder.RegisterExternalTexture(SceneTargets.GetSceneColor(), TEXT("SceneColorTexture"));
+					AddDebugDeepShadowTexturePass(GraphBuilder, &View, FIntRect(), &DeepShadowData, &MacroGroupDatas.DeepShadowResources, SceneColorTexture);
+				}
+			}
 			GraphBuilder.Execute();
 		}
 	}
@@ -1519,41 +1597,69 @@ void RenderHairStrandsDebugInfo(
 			const FHairStrandsMacroGroupDatas& MacroGroupDatas = InMacroGroupViews.Views[ViewIndex];
 			for (const FHairStrandsMacroGroupData& MacroGroupData : MacroGroupDatas.Datas)
 			{
-				AddDebugDeepShadowTexturePass(GraphBuilder, &View, MacroGroupData.ScreenRect, nullptr, SceneColorTexture);
+				AddDebugDeepShadowTexturePass(GraphBuilder, &View, MacroGroupData.ScreenRect, nullptr, nullptr, SceneColorTexture);
 			}
 
 			const FIntRect TotalRect = ComputeVisibleHairStrandsMacroGroupsRect(View.ViewRect, MacroGroupDatas);
-			AddDebugDeepShadowTexturePass(GraphBuilder, &View, TotalRect, nullptr, SceneColorTexture);
+			AddDebugDeepShadowTexturePass(GraphBuilder, &View, TotalRect, nullptr, nullptr, SceneColorTexture);
 		}
 		GraphBuilder.Execute();
 	}
 	
 	const bool bIsVoxelMode = HairDebugMode == EHairDebugMode::VoxelsDensity || HairDebugMode == EHairDebugMode::VoxelsTangent || HairDebugMode == EHairDebugMode::VoxelsBaseColor || HairDebugMode == EHairDebugMode::VoxelsRoughness;
 
-	// Render Frustum for all lights & macro groups
+	// Render Frustum for all lights & macro groups 
 	{
+		if ((HairDebugMode == EHairDebugMode::LightBounds || HairDebugMode == EHairDebugMode::DeepOpacityMaps) && ViewIndex < uint32(InMacroGroupViews.Views.Num()))
+		{
+			if (ViewIndex < uint32(InMacroGroupViews.Views.Num()))
+			{
+				const FHairStrandsMacroGroupDatas& MacroGroupDatas = InMacroGroupViews.Views[ViewIndex];
+				if (MacroGroupDatas.DeepShadowResources.bIsGPUDriven)
+				{
+					FRDGBuilder GraphBuilder(RHICmdList);
+					FRDGTextureRef SceneColorTexture = GraphBuilder.RegisterExternalTexture(SceneTargets.GetSceneColor(), TEXT("SceneColorTexture"));
+					AddDeepShadowInfoPass(GraphBuilder, View, MacroGroupDatas, SceneColorTexture);
+					GraphBuilder.Execute();
+				}
+			}
+		}
+
 		FViewElementPDI ShadowFrustumPDI(&View, nullptr, nullptr);
 
 		// All DOMs
-		if (HairDebugMode == EHairDebugMode::LightBounds && ViewIndex < uint32(InDomViews.Views.Num()))
+		if (HairDebugMode == EHairDebugMode::LightBounds && ViewIndex < uint32(InMacroGroupViews.Views.Num()))
 		{
-			const FHairStrandsDeepShadowDatas& DOMs = InDomViews.Views[ViewIndex];
-			for (const FHairStrandsDeepShadowData& DomData : DOMs.Datas)
+			if (!InMacroGroupViews.Views[ViewIndex].DeepShadowResources.bIsGPUDriven)
 			{
-				DrawFrustumWireframe(&ShadowFrustumPDI, DomData.WorldToLightTransform.Inverse(), FColor::Emerald, 0);
-				DrawWireBox(&ShadowFrustumPDI, DomData.Bounds.GetBox(), FColor::Yellow, 0);
+				for (const FHairStrandsMacroGroupData& MacroGroupData : InMacroGroupViews.Views[ViewIndex].Datas)
+				{
+					for (const FHairStrandsDeepShadowData& DomData : MacroGroupData.DeepShadowDatas.Datas)
+					{
+						DrawFrustumWireframe(&ShadowFrustumPDI, DomData.CPU_WorldToLightTransform.Inverse(), FColor::Emerald, 0);
+						DrawWireBox(&ShadowFrustumPDI, DomData.Bounds.GetBox(), FColor::Yellow, 0);
+					}
+				}
 			}
 		}
 
 		// Current DOM
-		if (HairDebugMode == EHairDebugMode::DeepOpacityMaps && ViewIndex < uint32(InDomViews.Views.Num()))
+		if (HairDebugMode == EHairDebugMode::DeepOpacityMaps && ViewIndex < uint32(InMacroGroupViews.Views.Num()))
 		{
-			const int32 CurrentIndex = FMath::Max(0, GDeepShadowDebugIndex);
-			const FHairStrandsDeepShadowDatas& DOMs = InDomViews.Views[ViewIndex];
-			if (CurrentIndex < DOMs.Datas.Num())
+			if (!InMacroGroupViews.Views[ViewIndex].DeepShadowResources.bIsGPUDriven)
 			{
-				DrawFrustumWireframe(&ShadowFrustumPDI, DOMs.Datas[CurrentIndex].WorldToLightTransform.Inverse(), FColor::Emerald, 0);
-				DrawWireBox(&ShadowFrustumPDI, DOMs.Datas[CurrentIndex].Bounds.GetBox(), FColor::Yellow, 0);
+				const int32 CurrentIndex = FMath::Max(0, GDeepShadowDebugIndex);
+				for (const FHairStrandsMacroGroupData& MacroGroupData : InMacroGroupViews.Views[ViewIndex].Datas)
+				{
+					for (const FHairStrandsDeepShadowData& DomData : MacroGroupData.DeepShadowDatas.Datas)
+					{
+						if (DomData.AtlasSlotIndex == CurrentIndex)
+						{
+							DrawFrustumWireframe(&ShadowFrustumPDI, DomData.CPU_WorldToLightTransform.Inverse(), FColor::Emerald, 0);
+							DrawWireBox(&ShadowFrustumPDI, DomData.Bounds.GetBox(), FColor::Yellow, 0);
+						}
+					}
+				}
 			}
 		}
 
@@ -1579,7 +1685,7 @@ void RenderHairStrandsDebugInfo(
 			}
 		}
 	}
-
+	
 	const bool bRunDebugPass =
 		HairDebugMode == EHairDebugMode::TAAResolveType ||
 		HairDebugMode == EHairDebugMode::SamplePerPixel ||
@@ -1781,10 +1887,21 @@ void RenderHairStrandsDebugInfo(
 		FRenderTargetTemp TempRenderTarget(View, (const FTexture2DRHIRef&)SceneTargets.GetSceneColor()->GetRenderTargetItem().TargetableTexture);
 		FCanvas Canvas(&TempRenderTarget, NULL, ViewFamily.CurrentRealTime, ViewFamily.CurrentWorldTime, ViewFamily.DeltaWorldTime, View.FeatureLevel);
 
-		const FHairStrandsDeepShadowDatas& ViewData = InDomViews.Views[ViewIndex];
+		const uint32 AtlasTotalSlotCount = FDeepShadowResources::MaxAtlasSlotCount;
+		uint32 AtlasAllocatedSlot = 0;
+		FIntPoint AtlasSlotResolution = FIntPoint(0, 0);
+		FIntPoint AtlasResolution = FIntPoint(0, 0);
+		bool bIsGPUDriven = false;
+		if (ViewIndex < uint32(InMacroGroupViews.Views.Num()))
+		{
+			const FDeepShadowResources& Resources = InMacroGroupViews.Views[ViewIndex].DeepShadowResources;
+			AtlasResolution = Resources.DepthAtlasTexture ? Resources.DepthAtlasTexture->GetRenderTargetItem().TargetableTexture->GetTexture2D()->GetSizeXY() : FIntPoint(0,0);
+			AtlasAllocatedSlot = Resources.TotalAtlasSlotCount;
+			bIsGPUDriven = Resources.bIsGPUDriven;
+		}
+
 		const uint32 DomTextureIndex = GDeepShadowDebugIndex;
 
-		const FIntPoint AtlasResolution = ViewData.Datas.Num() > 0 && ViewData.Datas[0].DepthTexture ? ViewData.Datas[0].DepthTexture->GetDesc().Extent : FIntPoint(0, 0);
 		float X = 20;
 		float Y = 38;
 
@@ -1794,20 +1911,26 @@ void RenderHairStrandsDebugInfo(
 		const FHairComponent HairComponent = GetHairComponents();
 		Line = FString::Printf(TEXT("Hair Components : (R=%d, TT=%d, TRT=%d, GS=%d, LS=%d)"), HairComponent.R, HairComponent.TT, HairComponent.TRT, HairComponent.GlobalScattering, HairComponent.LocalScattering);
 		Canvas.DrawShadowedString(X, Y += YStep, *Line, GetStatsFont(), DebugColor);
-		Line = FString::Printf(TEXT("----------------------------------------------------------------"));				Canvas.DrawShadowedString(X, Y += YStep, *Line, GetStatsFont(), DebugColor);
-		Line = FString::Printf(TEXT("Debug strands mode : %s"), ToString(StrandsDebugMode));							Canvas.DrawShadowedString(X, Y += YStep, *Line, GetStatsFont(), DebugColor);
+		Line = FString::Printf(TEXT("----------------------------------------------------------------"));						Canvas.DrawShadowedString(X, Y += YStep, *Line, GetStatsFont(), DebugColor);
+		Line = FString::Printf(TEXT("Debug strands mode : %s"), ToString(StrandsDebugMode));									Canvas.DrawShadowedString(X, Y += YStep, *Line, GetStatsFont(), DebugColor);
 		Line = FString::Printf(TEXT("Voxelization : %s"), IsHairStrandsVoxelizationEnable() ? TEXT("On") : TEXT("Off"));		Canvas.DrawShadowedString(X, Y += YStep, *Line, GetStatsFont(), DebugColor);
 		Line = FString::Printf(TEXT("View rect optim.: %s"), IsHairStrandsViewRectOptimEnable() ? TEXT("On") : TEXT("Off"));	Canvas.DrawShadowedString(X, Y += YStep, *Line, GetStatsFont(), DebugColor);
-		Line = FString::Printf(TEXT("----------------------------------------------------------------"));				Canvas.DrawShadowedString(X, Y += YStep, *Line, GetStatsFont(), DebugColor);
-		Line = FString::Printf(TEXT("DOM Atlas resolution  : %d/%d"), AtlasResolution.X, AtlasResolution.Y);				Canvas.DrawShadowedString(X, Y += YStep, *Line, GetStatsFont(), DebugColor);
-		Line = FString::Printf(TEXT("DOM macro group count : %d"), ViewData.Datas.Num());									Canvas.DrawShadowedString(X, Y += YStep, *Line, GetStatsFont(), DebugColor);
-		Line = FString::Printf(TEXT("DOM Texture Index     : %d/%d"), DomTextureIndex, ViewData.Datas.Num());				Canvas.DrawShadowedString(X, Y += YStep, *Line, GetStatsFont(), DebugColor);
+		Line = FString::Printf(TEXT("----------------------------------------------------------------"));						Canvas.DrawShadowedString(X, Y += YStep, *Line, GetStatsFont(), DebugColor);
+		Line = FString::Printf(TEXT("DOM Atlas resolution  : %dx%d"), AtlasResolution.X, AtlasResolution.Y);					Canvas.DrawShadowedString(X, Y += YStep, *Line, GetStatsFont(), DebugColor);
+		Line = FString::Printf(TEXT("DOM Atlas slot        : %d/%d"), AtlasAllocatedSlot, AtlasTotalSlotCount);					Canvas.DrawShadowedString(X, Y += YStep, *Line, GetStatsFont(), DebugColor);
+		Line = FString::Printf(TEXT("DOM Texture Index     : %d/%d"), DomTextureIndex, AtlasAllocatedSlot);						Canvas.DrawShadowedString(X, Y += YStep, *Line, GetStatsFont(), DebugColor);
+		Line = FString::Printf(TEXT("DOM GPU driven        : %s"), bIsGPUDriven ? TEXT("On") : TEXT("Off"));					Canvas.DrawShadowedString(X, Y += YStep, *Line, GetStatsFont(), DebugColor);
 
-		uint32 BoundIndex = 0;
-		for (const FHairStrandsDeepShadowData& DomData : ViewData.Datas)
+		if (ViewIndex < uint32(InMacroGroupViews.Views.Num()))
 		{
-			Line = FString::Printf(TEXT(" %d - Bound Radus: %f.2m (%dx%d)"), BoundIndex++, DomData.Bounds.GetSphere().W / 10.f, DomData.ShadowResolution.X, DomData.ShadowResolution.Y);
-			Canvas.DrawShadowedString(X, Y += YStep, *Line, GetStatsFont(), DebugColor);
+			for (const FHairStrandsMacroGroupData& MacroGroupData : InMacroGroupViews.Views[ViewIndex].Datas)
+			{
+				for (const FHairStrandsDeepShadowData& DomData : MacroGroupData.DeepShadowDatas.Datas)
+				{
+					Line = FString::Printf(TEXT(" %d - Bound Radus: %f.2m (%dx%d)"), DomData.AtlasSlotIndex, DomData.Bounds.GetSphere().W / 10.f, DomData.ShadowResolution.X, DomData.ShadowResolution.Y);
+					Canvas.DrawShadowedString(X, Y += YStep, *Line, GetStatsFont(), DebugColor);
+				}
+			}
 		}
 
 		Canvas.Flush_RenderThread(RHICmdList);
diff --git a/Engine/Source/Runtime/Renderer/Private/HairStrands/HairStrandsDeepShadow.cpp b/Engine/Source/Runtime/Renderer/Private/HairStrands/HairStrandsDeepShadow.cpp
--- a/Engine/Source/Runtime/Renderer/Private/HairStrands/HairStrandsDeepShadow.cpp
+++ b/Engine/Source/Runtime/Renderer/Private/HairStrands/HairStrandsDeepShadow.cpp
@@ -11,6 +11,9 @@
 static int32 GDeepShadowResolution = 2048;
 static FAutoConsoleVariableRef CVarDeepShadowResolution(TEXT("r.HairStrands.DeepShadow.Resolution"), GDeepShadowResolution, TEXT("Shadow resolution for Deep Opacity Map rendering. (default = 2048)"));
 
+static int32 GDeepShadowGPUDriven = 1;
+static FAutoConsoleVariableRef CVarDeepShadowGPUDriven(TEXT("r.HairStrands.DeepShadow.GPUDriven"), GDeepShadowGPUDriven, TEXT("Enable deep shadow to be driven by GPU bounding box, rather CPU ones. This allows more robust behavior"));
+
 ///////////////////////////////////////////////////////////////////////////////////////////////////
 
 typedef TArray<const FLightSceneInfo*, SceneRenderingAllocator> FLightSceneInfos;
@@ -43,13 +46,52 @@ static FLightSceneInfosArray GetVisibleDeepShadowLights(const FScene* Scene, con
 	return VisibleLightsPerView;
 }
 
+///////////////////////////////////////////////////////////////////////////////////////////////////
+class FDeepShadowCreateViewInfoCS : public FGlobalShader
+{
+	DECLARE_GLOBAL_SHADER(FDeepShadowCreateViewInfoCS);
+	SHADER_USE_PARAMETER_STRUCT(FDeepShadowCreateViewInfoCS, FGlobalShader);
+
+	BEGIN_SHADER_PARAMETER_STRUCT(FParameters, )
+
+		SHADER_PARAMETER_ARRAY(FVector4,	LightDirections,	[FHairStrandsDeepShadowData::MaxMacroGroupCount])
+		SHADER_PARAMETER_ARRAY(FVector4,	LightPositions,		[FHairStrandsDeepShadowData::MaxMacroGroupCount])
+		SHADER_PARAMETER_ARRAY(FIntVector4,	MacroGroupIndices,	[FHairStrandsDeepShadowData::MaxMacroGroupCount])
+
+		SHADER_PARAMETER(FVector, CPU_MinAABB)
+		SHADER_PARAMETER(uint32, CPU_bUseCPUData)
+		SHADER_PARAMETER(FVector, CPU_MaxAABB)
+		SHADER_PARAMETER(float, RasterizationScale)
+
+		SHADER_PARAMETER(FIntPoint, SlotResolution)
+		SHADER_PARAMETER(uint32, SlotIndexCount)
+		SHADER_PARAMETER(uint32, MacroGroupCount)
+
+		SHADER_PARAMETER_RDG_BUFFER_SRV(Buffer<int>, MacroGroupAABBBuffer)
+		SHADER_PARAMETER_RDG_BUFFER_UAV(RWStructuredBuffer<FDeepShadowViewInfo>, OutShadowViewInfoBuffer)
+		SHADER_PARAMETER_RDG_BUFFER_UAV(RWStructuredBuffer<float4x4>, OutShadowWorldToLightTransformBuffer)
+	END_SHADER_PARAMETER_STRUCT()
+
+public:
+	static bool ShouldCompilePermutation(const FGlobalShaderPermutationParameters& Parameters) { return IsHairStrandsSupported(Parameters.Platform); }
+	static void ModifyCompilationEnvironment(const FGlobalShaderPermutationParameters& Parameters, FShaderCompilerEnvironment& OutEnvironment)
+	{
+		FGlobalShader::ModifyCompilationEnvironment(Parameters, OutEnvironment);
+		OutEnvironment.SetDefine(TEXT("SHADER_ALLOCATE"), 1);
+		OutEnvironment.SetDefine(TEXT("MAX_SLOT_COUNT"), FHairStrandsDeepShadowData::MaxMacroGroupCount);
+	}
+};
+
+IMPLEMENT_GLOBAL_SHADER(FDeepShadowCreateViewInfoCS, "/Engine/Private/HairStrands/HairStrandsDeepShadowAllocation.usf", "CreateViewInfo", SF_Compute);
 ///////////////////////////////////////////////////////////////////////////////////////////////////
 
-FHairStrandsDeepShadowViews RenderHairStrandsDeepShadows(
+float GetDeepShadowRasterizationScale();
+
+void RenderHairStrandsDeepShadows(
 	FRHICommandListImmediate& RHICmdList,
 	const FScene* Scene,
 	const TArray<FViewInfo>& Views,
-	const FHairStrandsMacroGroupViews& MacroGroupsViews)
+	FHairStrandsMacroGroupViews& MacroGroupsViews)
 {
 	QUICK_SCOPE_CYCLE_COUNTER(STAT_CLM_RenderDeepShadow);
 	DECLARE_GPU_STAT(HairStrandsDeepShadow);
@@ -58,174 +100,243 @@ FHairStrandsDeepShadowViews RenderHairStrandsDeepShadows(
 
 	FLightSceneInfosArray VisibleLightsPerView = GetVisibleDeepShadowLights(Scene, Views);
 
-	// Compute the number of DOM which need to be created and insert default value
-	FHairStrandsDeepShadowViews DeepShadowsPerView;
-	uint32 DOMSlotCount = 0;
 	for (int32 ViewIndex = 0; ViewIndex < Views.Num(); ViewIndex++)
 	{
 		const FViewInfo& View = Views[ViewIndex];
-		if (View.Family)
+		if (!View.Family || ViewIndex >= MacroGroupsViews.Views.Num())
 		{
-			const FHairStrandsMacroGroupDatas& MacroGroupDatas = MacroGroupsViews.Views[ViewIndex];
-			DeepShadowsPerView.Views.AddDefaulted();
+			continue;
+		}
 
-			if (MacroGroupDatas.Datas.Num() == 0 || 
-				VisibleLightsPerView[ViewIndex].Num() == 0 ||
-				IsHairStrandsForVoxelTransmittanceAndShadowEnable()) 
-			{
-				continue; 
-			}
+		FHairStrandsMacroGroupDatas& MacroGroupDatas = MacroGroupsViews.Views[ViewIndex];
+		if (MacroGroupDatas.Datas.Num() == 0 || 
+			VisibleLightsPerView[ViewIndex].Num() == 0 ||
+			IsHairStrandsForVoxelTransmittanceAndShadowEnable()) 
+		{
+			continue; 
+		}
 
-			for (int32 MacroGroupIt = 0; MacroGroupIt < MacroGroupDatas.Datas.Num(); ++MacroGroupIt)
-			{
-				const FHairStrandsMacroGroupData& MacroGroup = MacroGroupDatas.Datas[MacroGroupIt];
-				const FBoxSphereBounds MacroGroupBounds = MacroGroup.Bounds;
+		// Compute the number of DOM which need to be created and insert default value
+		uint32 DOMSlotCount = 0;
+		for (int32 MacroGroupIt = 0; MacroGroupIt < MacroGroupDatas.Datas.Num(); ++MacroGroupIt)
+		{
+			const FHairStrandsMacroGroupData& MacroGroup = MacroGroupDatas.Datas[MacroGroupIt];
+			const FBoxSphereBounds MacroGroupBounds = MacroGroup.Bounds;
 
-				for (const FLightSceneInfo* LightInfo : VisibleLightsPerView[ViewIndex])
+			for (const FLightSceneInfo* LightInfo : VisibleLightsPerView[ViewIndex])
+			{
+				const FLightSceneProxy* LightProxy = LightInfo->Proxy;
+				if (!LightProxy->AffectsBounds(MacroGroupBounds))
 				{
-					const FLightSceneProxy* LightProxy = LightInfo->Proxy;
-					if (!LightProxy->AffectsBounds(MacroGroupBounds))
-					{
-						continue;
-					}
+					continue;
+				}
 
-					DOMSlotCount++;
+				// Run out of atlas slot
+				if (DOMSlotCount >= FDeepShadowResources::MaxAtlasSlotCount)
+				{
+					continue;
 				}
+
+				DOMSlotCount++;
 			}
 		}
-	}
 
-	if (DOMSlotCount == 0)
-		return DeepShadowsPerView;
+		if (DOMSlotCount == 0)
+			continue;
 
-	const uint32 AtlasSlotX = FGenericPlatformMath::CeilToInt(FMath::Sqrt(DOMSlotCount));
-	const FIntPoint AtlasSlotDimension(AtlasSlotX, AtlasSlotX == DOMSlotCount ? 1 : AtlasSlotX);
-	const FIntPoint AtlasSlotResolution(GDeepShadowResolution, GDeepShadowResolution);
-	const FIntPoint AtlasResolution(AtlasSlotResolution.X * AtlasSlotDimension.X, AtlasSlotResolution.Y * AtlasSlotDimension.Y);
+		const uint32 AtlasSlotX = FGenericPlatformMath::CeilToInt(FMath::Sqrt(DOMSlotCount));
+		const FIntPoint AtlasSlotDimension(AtlasSlotX, AtlasSlotX == DOMSlotCount ? 1 : AtlasSlotX);
+		const FIntPoint AtlasSlotResolution(GDeepShadowResolution, GDeepShadowResolution);
+		const FIntPoint AtlasResolution(AtlasSlotResolution.X * AtlasSlotDimension.X, AtlasSlotResolution.Y * AtlasSlotDimension.Y);
 
-	FRDGBuilder GraphBuilder(RHICmdList);
+		FDeepShadowResources& Resources = MacroGroupDatas.DeepShadowResources;
+		Resources.TotalAtlasSlotCount = 0;
 
-	// Create Atlas resources for DOM
-	bool bClearFrontDepthAtlasTexture = true;
-	bool bClearLayerAtlasTexture = true;
-	FRDGTextureRef FrontDepthAtlasTexture = GraphBuilder.CreateTexture(FPooledRenderTargetDesc::Create2DDesc(AtlasResolution, PF_DepthStencil, FClearValueBinding::DepthFar, TexCreate_None, TexCreate_DepthStencilTargetable | TexCreate_ShaderResource, false), TEXT("ShadowDepth"));
-	FRDGTextureRef DeepShadowLayersAtlasTexture = GraphBuilder.CreateTexture(FPooledRenderTargetDesc::Create2DDesc(AtlasResolution, PF_FloatRGBA, FClearValueBinding::Transparent, TexCreate_None, TexCreate_RenderTargetable | TexCreate_ShaderResource, false), TEXT("DeepShadowLayers"));
+		FRDGBuilder GraphBuilder(RHICmdList);
 
-	uint32 AtlasSlotIndex = 0;
-	for (int32 ViewIndex = 0; ViewIndex < Views.Num(); ViewIndex++)
-	{
-		const FViewInfo& ViewInfo = Views[ViewIndex];
-		if (ViewInfo.Family)
+		// Create Atlas resources for DOM. It is shared for all lights, across all views
+		bool bClear = true;
+		FRDGTextureRef FrontDepthAtlasTexture = GraphBuilder.CreateTexture(FPooledRenderTargetDesc::Create2DDesc(AtlasResolution, PF_DepthStencil, FClearValueBinding::DepthFar, TexCreate_None, TexCreate_DepthStencilTargetable | TexCreate_ShaderResource, false), TEXT("ShadowDepth"));
+		FRDGTextureRef DeepShadowLayersAtlasTexture = GraphBuilder.CreateTexture(FPooledRenderTargetDesc::Create2DDesc(AtlasResolution, PF_FloatRGBA, FClearValueBinding::Transparent, TexCreate_None, TexCreate_RenderTargetable | TexCreate_ShaderResource, false), TEXT("DeepShadowLayers"));
+
+		// TODO add support for multiple view: need to deduplicate light which are visible accross several views
+		// Allocate atlas CPU slot
+		uint32 TotalAtlasSlotIndex = 0;
+		for (int32 MacroGroupIt = 0; MacroGroupIt < MacroGroupDatas.Datas.Num(); ++MacroGroupIt)
 		{
-			FHairStrandsDeepShadowDatas& DeepShadowDatas = DeepShadowsPerView.Views[ViewIndex];
+			FHairStrandsMacroGroupData& MacroGroup = MacroGroupDatas.Datas[MacroGroupIt];
 
-			const FHairStrandsMacroGroupDatas& MacroGroupDatas = MacroGroupsViews.Views[ViewIndex];
-			for (int32 MacroGroupIt = 0; MacroGroupIt < MacroGroupDatas.Datas.Num(); ++MacroGroupIt)
+			// List of all the light in the scene.
+			for (const FLightSceneInfo* LightInfo : VisibleLightsPerView[ViewIndex])
 			{
-				// List of all the light in the scene.
-				uint32 LightLocalIndex = 0;
-				for (const FLightSceneInfo* LightInfo : VisibleLightsPerView[ViewIndex])
+				FBoxSphereBounds MacroGroupBounds = MacroGroup.Bounds;
+
+				const FLightSceneProxy* LightProxy = LightInfo->Proxy;
+				if (!LightProxy->AffectsBounds(MacroGroupBounds))
 				{
-					const FHairStrandsMacroGroupData& MacroGroup = MacroGroupDatas.Datas[MacroGroupIt];
-					FBoxSphereBounds MacroGroupBounds = MacroGroup.Bounds;
-
-					const FLightSceneProxy* LightProxy = LightInfo->Proxy;
-					if (!LightProxy->AffectsBounds(MacroGroupBounds))
-					{
-						continue;
-					}
-
-					const ELightComponentType LightType = (ELightComponentType)LightProxy->GetLightType();
-
-					FMinHairRadiusAtDepth1 MinStrandRadiusAtDepth1;
-					const FIntPoint AtlasRectOffset(
-						(AtlasSlotIndex % AtlasSlotDimension.X) * AtlasSlotResolution.X,
-						(AtlasSlotIndex / AtlasSlotDimension.X) * AtlasSlotResolution.Y);
-
-					// Note: LightPosition.W is used in the transmittance mask shader to differentiate between directional and local lights.
-					FHairStrandsDeepShadowData& DomData = DeepShadowDatas.Datas.AddZeroed_GetRef();
-					ComputeWorldToLightClip(DomData.WorldToLightTransform, MinStrandRadiusAtDepth1, MacroGroupBounds, *LightProxy, LightType, AtlasSlotResolution);
-					DomData.LightDirection = LightProxy->GetDirection();
-					DomData.LightPosition = FVector4(LightProxy->GetPosition(), LightType == ELightComponentType::LightType_Directional ? 0 : 1);
-					DomData.LightLuminance = LightProxy->GetColor();
-					DomData.LightType = LightType;
-					DomData.LightId = LightInfo->Id;
-					DomData.ShadowResolution = AtlasSlotResolution;
-					DomData.Bounds = MacroGroupBounds;
-					DomData.AtlasRect = FIntRect(AtlasRectOffset, AtlasRectOffset + AtlasSlotResolution);
-					DomData.MacroGroupId = MacroGroup.MacroGroupId;					
-					AtlasSlotIndex++;
-
-					const bool bIsOrtho = LightType == ELightComponentType::LightType_Directional;
-					const FVector4 HairRenderInfo = PackHairRenderInfo(MinStrandRadiusAtDepth1.Primary, MinStrandRadiusAtDepth1.Stable, MinStrandRadiusAtDepth1.Primary, 1);
-					const uint32 HairRenderInfoBits = PackHairRenderInfoBits(bIsOrtho, false);
+					continue;
+				}
 					
-					// Front depth
-					{
-						DECLARE_GPU_STAT(HairStrandsDeepShadowFrontDepth);
-						SCOPED_DRAW_EVENT(GraphBuilder.RHICmdList, HairStrandsDeepShadowFrontDepth);
-						SCOPED_GPU_STAT(GraphBuilder.RHICmdList, HairStrandsDeepShadowFrontDepth);
-
-						FHairDeepShadowRasterPassParameters* PassParameters = GraphBuilder.AllocParameters<FHairDeepShadowRasterPassParameters>();
-						PassParameters->WorldToClipMatrix = DomData.WorldToLightTransform;;
-						PassParameters->SliceValue = FVector4(1, 1, 1, 1);
-						PassParameters->AtlasRect = DomData.AtlasRect;
-						PassParameters->ViewportResolution = AtlasSlotResolution;
-						PassParameters->RenderTargets.DepthStencil = FDepthStencilBinding(FrontDepthAtlasTexture, bClearFrontDepthAtlasTexture ? ERenderTargetLoadAction::EClear : ERenderTargetLoadAction::ELoad, ERenderTargetLoadAction::ENoAction, FExclusiveDepthStencil::DepthWrite_StencilNop);
-
-						AddHairDeepShadowRasterPass(
-							GraphBuilder,
-							Scene,
-							&ViewInfo,
-							MacroGroup.PrimitivesInfos,
-							EHairStrandsRasterPassType::FrontDepth,
-							DomData.AtlasRect,
-							HairRenderInfo,
-							HairRenderInfoBits,
-							DomData.LightDirection,
-							PassParameters);
-					}
-
-					// Deep layers
-					{
-						DECLARE_GPU_STAT(HairStrandsDeepShadowLayers);
-						SCOPED_DRAW_EVENT(GraphBuilder.RHICmdList, HairStrandsDeepShadowLayers);
-						SCOPED_GPU_STAT(GraphBuilder.RHICmdList, HairStrandsDeepShadowLayers);
-
-						FHairDeepShadowRasterPassParameters* PassParameters = GraphBuilder.AllocParameters<FHairDeepShadowRasterPassParameters>();
-						PassParameters->WorldToClipMatrix = DomData.WorldToLightTransform;;
-						PassParameters->SliceValue = FVector4(1, 1, 1, 1);
-						PassParameters->AtlasRect = DomData.AtlasRect;
-						PassParameters->ViewportResolution = AtlasSlotResolution;
-						PassParameters->FrontDepthTexture = FrontDepthAtlasTexture;
-						PassParameters->RenderTargets[0] = FRenderTargetBinding(DeepShadowLayersAtlasTexture, bClearLayerAtlasTexture ? ERenderTargetLoadAction::EClear : ERenderTargetLoadAction::ELoad, 0);
-
-						AddHairDeepShadowRasterPass(
-							GraphBuilder,
-							Scene,
-							&ViewInfo,
-							MacroGroup.PrimitivesInfos,
-							EHairStrandsRasterPassType::DeepOpacityMap,
-							DomData.AtlasRect,
-							HairRenderInfo,
-							HairRenderInfoBits,
-							DomData.LightDirection,
-							PassParameters);						
-					}
-
-					GraphBuilder.QueueTextureExtraction(FrontDepthAtlasTexture, &DomData.DepthTexture);
-					GraphBuilder.QueueTextureExtraction(DeepShadowLayersAtlasTexture, &DomData.LayersTexture);
-
-					bClearFrontDepthAtlasTexture = false;
-					bClearLayerAtlasTexture = false;
-
-					++LightLocalIndex;
+				if (TotalAtlasSlotIndex >= FDeepShadowResources::MaxAtlasSlotCount)
+				{
+					continue;
+				}
+
+				const ELightComponentType LightType = (ELightComponentType)LightProxy->GetLightType();
+
+				FMinHairRadiusAtDepth1 MinStrandRadiusAtDepth1;
+				const FIntPoint AtlasRectOffset(
+					(TotalAtlasSlotIndex % AtlasSlotDimension.X) * AtlasSlotResolution.X,
+					(TotalAtlasSlotIndex / AtlasSlotDimension.X) * AtlasSlotResolution.Y);
+
+				// Note: LightPosition.W is used in the transmittance mask shader to differentiate between directional and local lights.
+				FHairStrandsDeepShadowData& DomData = MacroGroup.DeepShadowDatas.Datas.AddZeroed_GetRef();
+				ComputeWorldToLightClip(DomData.CPU_WorldToLightTransform, MinStrandRadiusAtDepth1, MacroGroupBounds, *LightProxy, LightType, AtlasSlotResolution);
+				DomData.LightDirection = LightProxy->GetDirection();
+				DomData.LightPosition = FVector4(LightProxy->GetPosition(), LightType == ELightComponentType::LightType_Directional ? 0 : 1);
+				DomData.LightLuminance = LightProxy->GetColor();
+				DomData.LightType = LightType;
+				DomData.LightId = LightInfo->Id;
+				DomData.ShadowResolution = AtlasSlotResolution;
+				DomData.Bounds = MacroGroupBounds;
+				DomData.AtlasRect = FIntRect(AtlasRectOffset, AtlasRectOffset + AtlasSlotResolution);
+				DomData.MacroGroupId = MacroGroup.MacroGroupId;
+				DomData.CPU_MinStrandRadiusAtDepth1 = MinStrandRadiusAtDepth1;
+				DomData.AtlasSlotIndex = TotalAtlasSlotIndex;
+				TotalAtlasSlotIndex++;
+			}
+		}
+
+		// Sanity check
+		check(DOMSlotCount == TotalAtlasSlotIndex); 
+
+		Resources.TotalAtlasSlotCount = TotalAtlasSlotIndex;
+		Resources.AtlasSlotResolution = AtlasSlotResolution;
+
+		FRDGBufferRef DeepShadowViewInfoBuffer = GraphBuilder.CreateBuffer(FRDGBufferDesc::CreateStructuredDesc(20 * sizeof(float), FMath::Max(1u, TotalAtlasSlotIndex)), TEXT("DeepShadowViewInfo"));
+		FRDGBufferRef DeepShadowWorldToLightBuffer = GraphBuilder.CreateBuffer(FRDGBufferDesc::CreateStructuredDesc(16 * sizeof(float), FMath::Max(1u, TotalAtlasSlotIndex)), TEXT("DeepShadowWorldToLightTransform"));
+		FRDGBufferSRVRef DeepShadowViewInfoBufferSRV = GraphBuilder.CreateSRV(DeepShadowViewInfoBuffer);
+
+		Resources.bIsGPUDriven = GDeepShadowGPUDriven > 0;
+		{
+			FRDGBufferRef MacroGroupAABB = GraphBuilder.RegisterExternalBuffer(MacroGroupDatas.MacroGroupResources.MacroGroupAABBsBuffer, TEXT("HairInstanceGroupAABBs"));
+
+			check(TotalAtlasSlotIndex < FDeepShadowResources::MaxAtlasSlotCount);
+
+			// Allocate and create projection matrix and Min radius
+			// Stored FDeepShadowViewInfo structs
+			// See HairStrandsDeepShadowCommonStruct.ush for more details
+			FDeepShadowCreateViewInfoCS::FParameters* Parameters = GraphBuilder.AllocParameters<FDeepShadowCreateViewInfoCS::FParameters>();
+
+			for (FHairStrandsMacroGroupData& MacroGroup : MacroGroupDatas.Datas)
+			{
+				for (FHairStrandsDeepShadowData& DomData : MacroGroup.DeepShadowDatas.Datas)
+				{				
+					Parameters->LightDirections[DomData.AtlasSlotIndex]		= FVector4(DomData.LightDirection.X, DomData.LightDirection.Y, DomData.LightDirection.Z, 0);
+					Parameters->LightPositions[DomData.AtlasSlotIndex]		= FVector4(DomData.LightPosition.X, DomData.LightPosition.Y, DomData.LightPosition.Z, DomData.LightType == LightType_Directional ? 0 : 1);
+					Parameters->MacroGroupIndices[DomData.AtlasSlotIndex]	= FIntVector4(DomData.MacroGroupId, 0,0,0);
 				}
 			}
 
-			GraphBuilder.Execute();
+			Parameters->SlotResolution = Resources.AtlasSlotResolution;
+			Parameters->SlotIndexCount = Resources.TotalAtlasSlotCount;
+			Parameters->MacroGroupCount = MacroGroupDatas.Datas.Num();
+			Parameters->MacroGroupAABBBuffer = GraphBuilder.CreateSRV(MacroGroupAABB, PF_R32_SINT);
+			Parameters->OutShadowViewInfoBuffer = GraphBuilder.CreateUAV(DeepShadowViewInfoBuffer);
+			Parameters->OutShadowWorldToLightTransformBuffer = GraphBuilder.CreateUAV(DeepShadowWorldToLightBuffer);
+
+			Parameters->RasterizationScale = GetDeepShadowRasterizationScale();
+			Parameters->CPU_bUseCPUData	= 0;
+			Parameters->CPU_MinAABB		= FVector::ZeroVector;
+			Parameters->CPU_MaxAABB		= FVector::ZeroVector;
+
+			// Currently support only 32 instance group at max
+			TShaderMapRef<FDeepShadowCreateViewInfoCS> ComputeShader(View.ShaderMap);
+			FComputeShaderUtils::AddPass(
+				GraphBuilder,
+				RDG_EVENT_NAME("HairStrandsDeepShadowAllocate"),
+				ComputeShader,
+				Parameters,
+				FIntVector(1, 1, 1));
 		}
-	}
 
-	return DeepShadowsPerView;
+		// Render deep shadows
+		for (FHairStrandsMacroGroupData& MacroGroup : MacroGroupDatas.Datas)
+		{
+			for (FHairStrandsDeepShadowData& DomData : MacroGroup.DeepShadowDatas.Datas)
+			{
+				const bool bIsOrtho = DomData.LightType == ELightComponentType::LightType_Directional;
+				const FVector4 HairRenderInfo = PackHairRenderInfo(DomData.CPU_MinStrandRadiusAtDepth1.Primary, DomData.CPU_MinStrandRadiusAtDepth1.Stable, DomData.CPU_MinStrandRadiusAtDepth1.Primary, 1);
+				const uint32 HairRenderInfoBits = PackHairRenderInfoBits(bIsOrtho, Resources.bIsGPUDriven);
+					
+				// Front depth
+				{
+					DECLARE_GPU_STAT(HairStrandsDeepShadowFrontDepth);
+					SCOPED_DRAW_EVENT(GraphBuilder.RHICmdList, HairStrandsDeepShadowFrontDepth);
+					SCOPED_GPU_STAT(GraphBuilder.RHICmdList, HairStrandsDeepShadowFrontDepth);
+
+					FHairDeepShadowRasterPassParameters* PassParameters = GraphBuilder.AllocParameters<FHairDeepShadowRasterPassParameters>();
+					PassParameters->CPU_WorldToClipMatrix = DomData.CPU_WorldToLightTransform;;
+					PassParameters->SliceValue = FVector4(1, 1, 1, 1);
+					PassParameters->AtlasRect = DomData.AtlasRect;
+					PassParameters->AtlasSlotIndex = DomData.AtlasSlotIndex;
+					PassParameters->ViewportResolution = AtlasSlotResolution;
+					PassParameters->DeepShadowViewInfoBuffer = DeepShadowViewInfoBufferSRV;
+					PassParameters->RenderTargets.DepthStencil = FDepthStencilBinding(FrontDepthAtlasTexture, bClear ? ERenderTargetLoadAction::EClear : ERenderTargetLoadAction::ELoad, ERenderTargetLoadAction::ENoAction, FExclusiveDepthStencil::DepthWrite_StencilNop);
+
+					AddHairDeepShadowRasterPass(
+						GraphBuilder,
+						Scene,
+						&View,
+						MacroGroup.PrimitivesInfos,
+						EHairStrandsRasterPassType::FrontDepth,
+						DomData.AtlasRect,
+						HairRenderInfo,
+						HairRenderInfoBits,
+						DomData.LightDirection,
+						PassParameters);
+				}
+
+				// Deep layers
+				{
+					DECLARE_GPU_STAT(HairStrandsDeepShadowLayers);
+					SCOPED_DRAW_EVENT(GraphBuilder.RHICmdList, HairStrandsDeepShadowLayers);
+					SCOPED_GPU_STAT(GraphBuilder.RHICmdList, HairStrandsDeepShadowLayers);
+
+					FHairDeepShadowRasterPassParameters* PassParameters = GraphBuilder.AllocParameters<FHairDeepShadowRasterPassParameters>();
+					PassParameters->CPU_WorldToClipMatrix = DomData.CPU_WorldToLightTransform;;
+					PassParameters->SliceValue = FVector4(1, 1, 1, 1);
+					PassParameters->AtlasRect = DomData.AtlasRect;
+					PassParameters->AtlasSlotIndex = DomData.AtlasSlotIndex;
+					PassParameters->ViewportResolution = AtlasSlotResolution;
+					PassParameters->FrontDepthTexture = FrontDepthAtlasTexture;
+					PassParameters->DeepShadowViewInfoBuffer = DeepShadowViewInfoBufferSRV;
+					PassParameters->RenderTargets[0] = FRenderTargetBinding(DeepShadowLayersAtlasTexture, bClear ? ERenderTargetLoadAction::EClear : ERenderTargetLoadAction::ELoad, 0);
+
+					AddHairDeepShadowRasterPass(
+						GraphBuilder,
+						Scene,
+						&View,
+						MacroGroup.PrimitivesInfos,
+						EHairStrandsRasterPassType::DeepOpacityMap,
+						DomData.AtlasRect,
+						HairRenderInfo,
+						HairRenderInfoBits,
+						DomData.LightDirection,
+						PassParameters);						
+				}
+				bClear = false;
+			}
+		}
+
+		GraphBuilder.QueueTextureExtraction(FrontDepthAtlasTexture, &Resources.DepthAtlasTexture);
+		GraphBuilder.QueueTextureExtraction(DeepShadowLayersAtlasTexture, &Resources.LayersAtlasTexture);
+		GraphBuilder.QueueBufferExtraction(DeepShadowWorldToLightBuffer, &Resources.DeepShadowWorldToLightTransforms, FRDGResourceState::EAccess::Read, FRDGResourceState::EPipeline::Graphics);
+		GraphBuilder.Execute();
+
+		if (Resources.DeepShadowWorldToLightTransforms)
+		{
+			Resources.DeepShadowWorldToLightTransformsSRV = RHICreateShaderResourceView(Resources.DeepShadowWorldToLightTransforms->StructuredBuffer);
+		}
+	}
 }
diff --git a/Engine/Source/Runtime/Renderer/Private/HairStrands/HairStrandsDeepShadow.h b/Engine/Source/Runtime/Renderer/Private/HairStrands/HairStrandsDeepShadow.h
--- a/Engine/Source/Runtime/Renderer/Private/HairStrands/HairStrandsDeepShadow.h
+++ b/Engine/Source/Runtime/Renderer/Private/HairStrands/HairStrandsDeepShadow.h
@@ -8,6 +8,8 @@
 
 #include "CoreMinimal.h"
 #include "RendererInterface.h"
+#include "HairStrandsInterface.h"
+#include "RenderGraphResources.h"
 #include "SceneTypes.h"
 
 /// Hold deep shadow information for a given light.
@@ -15,11 +17,11 @@ struct FHairStrandsDeepShadowData
 {
 	static const uint32 MaxMacroGroupCount = 16u;
 
-	TRefCountPtr<IPooledRenderTarget> DepthTexture;
-	TRefCountPtr<IPooledRenderTarget> LayersTexture;
-	FMatrix WorldToLightTransform;
+	FMatrix CPU_WorldToLightTransform;
+	FMinHairRadiusAtDepth1 CPU_MinStrandRadiusAtDepth1;
 	FIntRect AtlasRect;
 	uint32 MacroGroupId = ~0;
+	uint32 AtlasSlotIndex = AtlasSlotIndex;
 
 	FIntPoint ShadowResolution = FIntPoint::ZeroValue;
 	uint32 LightId = ~0;
@@ -31,20 +33,31 @@ struct FHairStrandsDeepShadowData
 	FBoxSphereBounds Bounds;
 };
 
-/// Store all deep shadows infos for a given view
-struct FHairStrandsDeepShadowDatas
+struct FDeepShadowResources
 {
-	TArray<FHairStrandsDeepShadowData, SceneRenderingAllocator> Datas;
+	// Limit the number of atlas slot to 32, in order to create the view info per slot in single compute
+	// This limitation can be alleviate, and is just here for convenience (see FDeepShadowCreateViewInfoCS)
+	static const uint32 MaxAtlasSlotCount = 32u;
+
+	uint32 TotalAtlasSlotCount = 0;
+	FIntPoint AtlasSlotResolution;
+	bool bIsGPUDriven = false;
+
+	TRefCountPtr<IPooledRenderTarget> DepthAtlasTexture;
+	TRefCountPtr<IPooledRenderTarget> LayersAtlasTexture;
+
+	TRefCountPtr<FPooledRDGBuffer> DeepShadowWorldToLightTransforms;
+	TRefCountPtr<FRHIShaderResourceView> DeepShadowWorldToLightTransformsSRV;
 };
 
-/// Store all deep shadows info for all views
-struct FHairStrandsDeepShadowViews
+/// Store all deep shadows infos for a given view
+struct FHairStrandsDeepShadowDatas
 {
-	TArray<FHairStrandsDeepShadowDatas, SceneRenderingAllocator> Views;
+	TArray<FHairStrandsDeepShadowData, SceneRenderingAllocator> Datas;
 };
 
-FHairStrandsDeepShadowViews RenderHairStrandsDeepShadows(
+void RenderHairStrandsDeepShadows(
 	FRHICommandListImmediate& RHICmdList,
 	const class FScene* Scene,
 	const TArray<FViewInfo>& Views,
-	const struct FHairStrandsMacroGroupViews& MacroGroupsViews);
\ No newline at end of file
+	struct FHairStrandsMacroGroupViews& MacroGroupsViews);
\ No newline at end of file
diff --git a/Engine/Source/Runtime/Renderer/Private/HairStrands/HairStrandsRasterCommon.cpp b/Engine/Source/Runtime/Renderer/Private/HairStrands/HairStrandsRasterCommon.cpp
--- a/Engine/Source/Runtime/Renderer/Private/HairStrands/HairStrandsRasterCommon.cpp
+++ b/Engine/Source/Runtime/Renderer/Private/HairStrands/HairStrandsRasterCommon.cpp
@@ -14,22 +14,26 @@
 /////////////////////////////////////////////////////////////////////////////////////////
 // Deep shadow global parameters
 BEGIN_GLOBAL_SHADER_PARAMETER_STRUCT(FHairDeepShadowRasterGlobalParameters, )
-	SHADER_PARAMETER(FMatrix, WorldToClipMatrix)
+	SHADER_PARAMETER(FMatrix, CPU_WorldToClipMatrix)
 	SHADER_PARAMETER(FVector4, SliceValue)
 	SHADER_PARAMETER(FIntRect, AtlasRect)
 	SHADER_PARAMETER(FIntPoint, ViewportResolution)
+	SHADER_PARAMETER(uint32, AtlasSlotIndex)
 	SHADER_PARAMETER_TEXTURE(Texture2D<float>, FrontDepthTexture)
+	SHADER_PARAMETER_SRV(StructuredBuffer<FDeepShadowViewInfo>, DeepShadowViewInfoBuffer)
 	END_GLOBAL_SHADER_PARAMETER_STRUCT()
 IMPLEMENT_GLOBAL_SHADER_PARAMETER_STRUCT(FHairDeepShadowRasterGlobalParameters, "DeepRasterPass");
 
 static FHairDeepShadowRasterGlobalParameters ConvertToGlobalPassParameter(const FHairDeepShadowRasterPassParameters* In)
 {
 	FHairDeepShadowRasterGlobalParameters Out;
-	Out.WorldToClipMatrix	= In->WorldToClipMatrix;
-	Out.SliceValue			= In->SliceValue;
-	Out.AtlasRect			= In->AtlasRect;
-	Out.ViewportResolution	= In->ViewportResolution;
-	Out.FrontDepthTexture	= In->FrontDepthTexture ? In->FrontDepthTexture->GetRHI() : (FRHITexture*)(GSystemTextures.DepthDummy->GetRenderTargetItem().ShaderResourceTexture);
+	Out.CPU_WorldToClipMatrix	= In->CPU_WorldToClipMatrix;
+	Out.SliceValue				= In->SliceValue;
+	Out.AtlasRect				= In->AtlasRect;
+	Out.ViewportResolution		= In->ViewportResolution;
+	Out.AtlasSlotIndex			= In->AtlasSlotIndex;
+	Out.FrontDepthTexture		= In->FrontDepthTexture ? In->FrontDepthTexture->GetRHI() : (FRHITexture*)(GSystemTextures.DepthDummy->GetRenderTargetItem().ShaderResourceTexture);
+	Out.DeepShadowViewInfoBuffer= In->DeepShadowViewInfoBuffer->GetRHI();
 	return Out;
 }
 
diff --git a/Engine/Source/Runtime/Renderer/Private/HairStrands/HairStrandsRasterCommon.h b/Engine/Source/Runtime/Renderer/Private/HairStrands/HairStrandsRasterCommon.h
--- a/Engine/Source/Runtime/Renderer/Private/HairStrands/HairStrandsRasterCommon.h
+++ b/Engine/Source/Runtime/Renderer/Private/HairStrands/HairStrandsRasterCommon.h
@@ -26,11 +26,14 @@ enum class EHairStrandsRasterPassType : uint8
 // Deep shadow raster pass
 
 BEGIN_SHADER_PARAMETER_STRUCT(FHairDeepShadowRasterPassParameters, )
-	SHADER_PARAMETER(FMatrix, WorldToClipMatrix)
+	SHADER_PARAMETER(FMatrix, CPU_WorldToClipMatrix)
 	SHADER_PARAMETER(FVector4, SliceValue)
 	SHADER_PARAMETER(FIntRect, AtlasRect)
 	SHADER_PARAMETER(FIntPoint, ViewportResolution)
+	SHADER_PARAMETER(uint32, AtlasSlotIndex)
 	SHADER_PARAMETER_RDG_TEXTURE(Texture2D<float>, FrontDepthTexture)
+	SHADER_PARAMETER_RDG_BUFFER_SRV(StructuredBuffer<FDeepShadowViewInfo>, DeepShadowViewInfoBuffer)
+
 	RENDER_TARGET_BINDING_SLOTS()
 END_SHADER_PARAMETER_STRUCT()
 
diff --git a/Engine/Source/Runtime/Renderer/Private/HairStrands/HairStrandsRendering.h b/Engine/Source/Runtime/Renderer/Private/HairStrands/HairStrandsRendering.h
--- a/Engine/Source/Runtime/Renderer/Private/HairStrands/HairStrandsRendering.h
+++ b/Engine/Source/Runtime/Renderer/Private/HairStrands/HairStrandsRendering.h
@@ -25,7 +25,6 @@
 /// Hold all the hair strands data
 struct FHairStrandsDatas
 {
-	FHairStrandsDeepShadowViews DeepShadowViews;
 	FHairStrandsVisibilityViews HairVisibilityViews;
 	FHairStrandsMacroGroupViews MacroGroupsPerViews;
 };
diff --git a/Engine/Source/Runtime/Renderer/Private/HairStrands/HairStrandsTransmittance.cpp b/Engine/Source/Runtime/Renderer/Private/HairStrands/HairStrandsTransmittance.cpp
--- a/Engine/Source/Runtime/Renderer/Private/HairStrands/HairStrandsTransmittance.cpp
+++ b/Engine/Source/Runtime/Renderer/Private/HairStrands/HairStrandsTransmittance.cpp
@@ -49,18 +49,20 @@ float GetDeepShadowDensityScale() { return FMath::Max(0.0f, GDeepShadowDensitySc
 float GetDeepShadowDepthBiasScale() { return FMath::Max(0.0f, GDeepShadowDepthBiasScale); }
 ///////////////////////////////////////////////////////////////////////////////////////////////////
 
-static bool HasDeepShadowData(const FLightSceneInfo* LightSceneInfo, const FHairStrandsDeepShadowDatas& InDatas)
+static bool HasDeepShadowData(const FLightSceneInfo* LightSceneInfo, const FHairStrandsMacroGroupDatas& InDatas)
 {
-	for (const FHairStrandsDeepShadowData& DomData : InDatas.Datas)
+	for (const FHairStrandsMacroGroupData& MacroGroupData : InDatas.Datas)
 	{
-		if (DomData.LightId == LightSceneInfo->Id)
-			return true;
+		for (const FHairStrandsDeepShadowData& DomData : MacroGroupData.DeepShadowDatas.Datas)
+		{
+			if (DomData.LightId == LightSceneInfo->Id)
+				return true;
+		}
 	}
 
 	return false;
 }
 
-
 ///////////////////////////////////////////////////////////////////////////////////////////////////
 // Transmittance Mask from deep shadow
 
@@ -88,8 +90,8 @@ class FDeepTransmittanceMaskCS : public FGlobalShader
 		SHADER_PARAMETER_STRUCT_INCLUDE(ShaderDrawDebug::FShaderDrawDebugParameters, ShaderDrawParameters)
 		SHADER_PARAMETER_STRUCT_INCLUDE(ShaderPrint::FShaderParameters, ShaderPrintParameters)
 
-		SHADER_PARAMETER_ARRAY(FIntVector4, DeepShadow_AtlasSlotOffsets, [FHairStrandsDeepShadowData::MaxMacroGroupCount])
-		SHADER_PARAMETER_ARRAY(FMatrix, DeepShadow_WorldToLightTransforms, [FHairStrandsDeepShadowData::MaxMacroGroupCount])
+		SHADER_PARAMETER_ARRAY(FIntVector4, DeepShadow_AtlasSlotOffsets_AtlasSlotIndex, [FHairStrandsDeepShadowData::MaxMacroGroupCount])
+		SHADER_PARAMETER_ARRAY(FMatrix, DeepShadow_CPUWorldToLightTransforms, [FHairStrandsDeepShadowData::MaxMacroGroupCount])
 		SHADER_PARAMETER(FIntPoint, DeepShadow_Resolution)
 		SHADER_PARAMETER(float, LightRadius)
 		SHADER_PARAMETER(FVector, LightDirection)
@@ -102,6 +104,7 @@ class FDeepTransmittanceMaskCS : public FGlobalShader
 		SHADER_PARAMETER(uint32, DeepShadow_KernelType)
 		SHADER_PARAMETER(uint32, DeepShadow_DebugMode)
 		SHADER_PARAMETER(FMatrix, DeepShadow_ShadowToWorld)
+		SHADER_PARAMETER(uint32, DeepShadow_bIsGPUDriven)
 
 		SHADER_PARAMETER_ARRAY(FVector4, Voxel_MinAABBs, [FHairStrandsDeepShadowData::MaxMacroGroupCount])
 		SHADER_PARAMETER_ARRAY(FVector4, Voxel_MaxAABBs, [FHairStrandsDeepShadowData::MaxMacroGroupCount])
@@ -119,6 +122,8 @@ class FDeepTransmittanceMaskCS : public FGlobalShader
 
 		SHADER_PARAMETER_RDG_TEXTURE(Texture2D, DeepShadow_FrontDepthTexture)
 		SHADER_PARAMETER_RDG_TEXTURE(Texture2D, DeepShadow_DomTexture)
+		SHADER_PARAMETER_RDG_BUFFER_SRV(StructuredBuffer, DeepShadow_WorldToLightTransformBuffer)
+
 		SHADER_PARAMETER_RDG_TEXTURE(Texture3D, HairLUTTexture)
 		SHADER_PARAMETER_RDG_BUFFER_SRV(StructuredBuffer, HairVisibilityNodeData)
 		SHADER_PARAMETER_RDG_BUFFER_SRV(StructuredBuffer, HairVisibilityNodeCoord)
@@ -168,9 +173,11 @@ IMPLEMENT_GLOBAL_SHADER(FDeepTransmittanceMaskCS, "/Engine/Private/HairStrands/H
 
 struct FDeepShadowTransmittanceParams
 {
-	FIntVector4 DeepShadow_AtlasSlotOffsets[FHairStrandsDeepShadowData::MaxMacroGroupCount];
-	FMatrix DeepShadow_WorldToLightTransforms[FHairStrandsDeepShadowData::MaxMacroGroupCount];
+	FIntVector4 DeepShadow_AtlasSlotOffsets_AtlasSlotIndex[FHairStrandsDeepShadowData::MaxMacroGroupCount];
+	FMatrix DeepShadow_CPUWorldToLightTransforms[FHairStrandsDeepShadowData::MaxMacroGroupCount];
+	FRDGBufferSRVRef DeepShadow_WorldToLightTransformBuffer = nullptr;
 	FIntPoint DeepShadow_Resolution = FIntPoint(0, 0);
+	bool DeepShadow_bIsGPUDriven = false;
 	FVector LightDirection = FVector::ZeroVector;
 	FVector4 LightPosition = FVector4(0, 0, 0, 0);
 	uint32 LightChannelMask = 0;
@@ -178,7 +185,7 @@ struct FDeepShadowTransmittanceParams
 	float DepthBiasScale = 0;
 	float DensityScale = 0;
 	FMatrix DeepShadow_ShadowToWorld = FMatrix::Identity;
-
+	
 	FRDGTextureRef DeepShadow_FrontDepthTexture = nullptr;
 	FRDGTextureRef DeepShadow_DomTexture = nullptr;
 
@@ -231,6 +238,8 @@ static FRDGBufferRef AddDeepShadowTransmittanceMaskPass(
 	Parameters->DeepShadow_ShadowToWorld = Params.DeepShadow_ShadowToWorld;
 	Parameters->IndirectArgsBuffer = IndirectArgsBuffer;
 	Parameters->MaxVisibilityNodeCount = Params.HairVisibilityNodeData->Desc.NumElements;
+	Parameters->DeepShadow_bIsGPUDriven = Params.DeepShadow_bIsGPUDriven ? 1 : 0;;
+	Parameters->DeepShadow_WorldToLightTransformBuffer = Params.DeepShadow_WorldToLightTransformBuffer;
 
 	if (ShaderDrawDebug::IsShaderDrawDebugEnabled(View))
 	{
@@ -238,8 +247,8 @@ static FRDGBufferRef AddDeepShadowTransmittanceMaskPass(
 		ShaderPrint::SetParameters(View, Parameters->ShaderPrintParameters);
 	}
 
-	memcpy(&(Parameters->DeepShadow_AtlasSlotOffsets[0]), Params.DeepShadow_AtlasSlotOffsets, sizeof(FIntVector4) * FHairStrandsDeepShadowData::MaxMacroGroupCount);
-	memcpy(&(Parameters->DeepShadow_WorldToLightTransforms[0]), Params.DeepShadow_WorldToLightTransforms, sizeof(FMatrix) * FHairStrandsDeepShadowData::MaxMacroGroupCount);
+	memcpy(&(Parameters->DeepShadow_AtlasSlotOffsets_AtlasSlotIndex[0]), Params.DeepShadow_AtlasSlotOffsets_AtlasSlotIndex, sizeof(FIntVector4) * FHairStrandsDeepShadowData::MaxMacroGroupCount);
+	memcpy(&(Parameters->DeepShadow_CPUWorldToLightTransforms[0]), Params.DeepShadow_CPUWorldToLightTransforms, sizeof(FMatrix) * FHairStrandsDeepShadowData::MaxMacroGroupCount);
 
 	memcpy(&(Parameters->Voxel_MinAABBs[0]), Params.Voxel_MinAABBs, sizeof(FVector4) * FHairStrandsDeepShadowData::MaxMacroGroupCount);
 	memcpy(&(Parameters->Voxel_MaxAABBs[0]), Params.Voxel_MaxAABBs, sizeof(FVector4) * FHairStrandsDeepShadowData::MaxMacroGroupCount);
@@ -311,13 +320,15 @@ class FDeepShadowMaskPS : public FGlobalShader
 		SHADER_PARAMETER_STRUCT_INCLUDE(FSceneTextureParameters, SceneTextures)
 		SHADER_PARAMETER_STRUCT_INCLUDE(ShaderDrawDebug::FShaderDrawDebugParameters, ShaderDrawParameters)
 		SHADER_PARAMETER_STRUCT_INCLUDE(ShaderPrint::FShaderParameters, ShaderPrintParameters)
-
+		                         
 		SHADER_PARAMETER(FIntPoint, DeepShadow_SlotOffset)
+		SHADER_PARAMETER(uint32, DeepShadow_SlotIndex)
 		SHADER_PARAMETER(FIntPoint, DeepShadow_SlotResolution)
-		SHADER_PARAMETER(FMatrix, DeepShadow_WorldToLightTransform)
+		SHADER_PARAMETER(FMatrix, DeepShadow_CPUWorldToLightTransform)
 		SHADER_PARAMETER(uint32, DeepShadow_IsWholeSceneLight)
 		SHADER_PARAMETER(float, DeepShadow_DepthBiasScale)
 		SHADER_PARAMETER(float, DeepShadow_DensityScale)
+		SHADER_PARAMETER(uint32, DeepShadow_bIsGPUDriven)
 
 		SHADER_PARAMETER(FVector4, Voxel_LightPosition)
 		SHADER_PARAMETER(FVector, Voxel_LightDirection)
@@ -329,6 +340,7 @@ class FDeepShadowMaskPS : public FGlobalShader
 
 		SHADER_PARAMETER_RDG_TEXTURE(Texture2D, RayMarchMaskTexture)
 
+		SHADER_PARAMETER_RDG_BUFFER_SRV(StructuredBuffer, DeepShadow_WorldToLightTransformBuffer)
 		SHADER_PARAMETER_RDG_TEXTURE(Texture3D, Voxel_DensityTexture)
 		SHADER_PARAMETER_RDG_TEXTURE(Texture2D, DeepShadow_FrontDepthTexture)
 		SHADER_PARAMETER_RDG_TEXTURE(Texture2D, DeepShadow_DomTexture)
@@ -352,13 +364,16 @@ struct FDeepShadowOpaqueParams
 {
 	FRDGTextureRef	CategorizationTexture = nullptr;
 
-	FMatrix			DeepShadow_WorldToLightTransform;
+	FRDGBufferSRVRef DeepShadow_WorldToLightTransformBuffer = nullptr;
+	FMatrix			DeepShadow_CPUWorldToLightTransform;
 	FIntRect		DeepShadow_AtlasRect;
 	FRDGTextureRef	DeepShadow_FrontDepthTexture = nullptr;
 	FRDGTextureRef	DeepShadow_LayerTexture = nullptr;
 	bool			DeepShadow_bIsWholeSceneLight = false;
+	bool			DeepShadow_bIsGPUDriven = false;
 	float			DeepShadow_DepthBiasScale = 1;
 	float			DeepShadow_DensityScale = 1;
+	uint32			DeepShadow_AtlasSlotIndex = 0;
 
 	FVector			Voxel_LightDirection = FVector::ZeroVector;
 	FVector4		Voxel_LightPosition = FVector4(0, 0, 0, 0);
@@ -386,7 +401,7 @@ static void AddDeepShadowOpaqueMaskPass(
 	FDeepShadowMaskPS::FParameters* Parameters = GraphBuilder.AllocParameters<FDeepShadowMaskPS::FParameters>();
 	Parameters->ViewUniformBuffer = View.ViewUniformBuffer;
 	Parameters->SceneTextures = SceneTextures;
-	Parameters->DeepShadow_WorldToLightTransform = Params.DeepShadow_WorldToLightTransform;
+	Parameters->DeepShadow_CPUWorldToLightTransform = Params.DeepShadow_CPUWorldToLightTransform;
 	Parameters->DeepShadow_FrontDepthTexture = Params.DeepShadow_FrontDepthTexture;
 	Parameters->DeepShadow_DomTexture = Params.DeepShadow_LayerTexture;
 	Parameters->CategorizationTexture = Params.CategorizationTexture;
@@ -396,8 +411,11 @@ static void AddDeepShadowOpaqueMaskPass(
 	Parameters->DeepShadow_DepthBiasScale = Params.DeepShadow_DepthBiasScale;
 	Parameters->DeepShadow_DensityScale = Params.DeepShadow_DensityScale;
 	Parameters->RenderTargets[0] = FRenderTargetBinding(OutShadowMask, ERenderTargetLoadAction::ELoad);
+	Parameters->DeepShadow_SlotIndex = Params.DeepShadow_AtlasSlotIndex;
 	Parameters->DeepShadow_SlotOffset = FIntPoint(Params.DeepShadow_AtlasRect.Min.X, Params.DeepShadow_AtlasRect.Min.Y);
 	Parameters->DeepShadow_SlotResolution = FIntPoint(Params.DeepShadow_AtlasRect.Max.X - Params.DeepShadow_AtlasRect.Min.X, Params.DeepShadow_AtlasRect.Max.Y - Params.DeepShadow_AtlasRect.Min.Y);
+	Parameters->DeepShadow_WorldToLightTransformBuffer = Params.DeepShadow_WorldToLightTransformBuffer;
+	Parameters->DeepShadow_bIsGPUDriven = Params.DeepShadow_bIsGPUDriven ? 1 : 0;;
 
 	if (ShaderDrawDebug::IsShaderDrawDebugEnabled(View))
 	{
@@ -489,7 +507,6 @@ static FHairStrandsTransmittanceMaskData RenderHairStrandsTransmittanceMask(
 	FRHICommandListImmediate& RHICmdList,
 	const FViewInfo& View,
 	const FLightSceneInfo* LightSceneInfo,
-	const FHairStrandsDeepShadowDatas& DeepShadowDatas,
 	const FHairStrandsMacroGroupDatas& MacroGroupDatas,
 	const FHairStrandsVisibilityData& VisibilityData,
 	TRefCountPtr<IPooledRenderTarget>& ScreenShadowMaskSubPixelTexture)
@@ -497,7 +514,7 @@ static FHairStrandsTransmittanceMaskData RenderHairStrandsTransmittanceMask(
 	if (MacroGroupDatas.Datas.Num() == 0)
 		return FHairStrandsTransmittanceMaskData();
 
-	if (!HasDeepShadowData(LightSceneInfo, DeepShadowDatas) && !IsHairStrandsVoxelizationEnable())
+	if (!HasDeepShadowData(LightSceneInfo, MacroGroupDatas) && !IsHairStrandsVoxelizationEnable())
 		return FHairStrandsTransmittanceMaskData();
 
 	DECLARE_GPU_STAT(HairStrandsTransmittanceMask);
@@ -520,31 +537,45 @@ static FHairStrandsTransmittanceMaskData RenderHairStrandsTransmittanceMask(
 	Params.HairVisibilityNodeData = GraphBuilder.RegisterExternalBuffer(VisibilityData.NodeData, TEXT("HairVisibilityNodeData"));
 	Params.HairVisibilityNodeCoord = GraphBuilder.RegisterExternalBuffer(VisibilityData.NodeCoord, TEXT("HairVisibilityNodeCoord"));
 	Params.DensityScale = GetDeepShadowDensityScale();
-	memset(Params.DeepShadow_AtlasSlotOffsets, 0, sizeof(Params.DeepShadow_AtlasSlotOffsets));
-	memset(Params.DeepShadow_WorldToLightTransforms, 0, sizeof(Params.DeepShadow_WorldToLightTransforms));
+	memset(Params.DeepShadow_AtlasSlotOffsets_AtlasSlotIndex, 0, sizeof(Params.DeepShadow_AtlasSlotOffsets_AtlasSlotIndex));
+	memset(Params.DeepShadow_CPUWorldToLightTransforms, 0, sizeof(Params.DeepShadow_CPUWorldToLightTransforms));
+
 
 	FRDGBufferRef OutShadowMask = nullptr;
 	bool bHasFoundLight = false;
 	if (!IsHairStrandsForVoxelTransmittanceAndShadowEnable())
 	{
-		for (const FHairStrandsDeepShadowData& DeepShadowData : DeepShadowDatas.Datas)
+		FRDGBufferRef DeepShadow_WorldToLightTransformBuffer = nullptr;
+		FRDGBufferSRVRef DeepShadow_WorldToLightTransformBufferSRV = nullptr;
+		for (const FHairStrandsMacroGroupData& MacroGroupData : MacroGroupDatas.Datas)
 		{
-			if (DeepShadowData.LightId == LightSceneInfo->Id)
+			for (const FHairStrandsDeepShadowData& DeepShadowData : MacroGroupData.DeepShadowDatas.Datas)
 			{
-				bHasFoundLight = true;
-				if (!Params.DeepShadow_FrontDepthTexture)
+				if (DeepShadowData.LightId == LightSceneInfo->Id)
 				{
-					Params.DeepShadow_FrontDepthTexture = GraphBuilder.RegisterExternalTexture(DeepShadowData.DepthTexture, TEXT("DeepShadow_FrontDepthTexture"));
-					Params.DeepShadow_DomTexture = GraphBuilder.RegisterExternalTexture(DeepShadowData.LayersTexture, TEXT("DeepShadow_DomTexture"));
+					if (DeepShadow_WorldToLightTransformBuffer == nullptr)
+					{
+						DeepShadow_WorldToLightTransformBuffer = GraphBuilder.RegisterExternalBuffer(MacroGroupDatas.DeepShadowResources.DeepShadowWorldToLightTransforms);
+						DeepShadow_WorldToLightTransformBufferSRV = GraphBuilder.CreateSRV(DeepShadow_WorldToLightTransformBuffer);
+					}
+
+					bHasFoundLight = true;
+					if (!Params.DeepShadow_FrontDepthTexture)
+					{
+						Params.DeepShadow_FrontDepthTexture = GraphBuilder.RegisterExternalTexture(MacroGroupDatas.DeepShadowResources.DepthAtlasTexture, TEXT("DeepShadow_FrontDepthTexture"));
+						Params.DeepShadow_DomTexture = GraphBuilder.RegisterExternalTexture(MacroGroupDatas.DeepShadowResources.LayersAtlasTexture, TEXT("DeepShadow_DomTexture"));
+					}
+					Params.DeepShadow_Resolution = DeepShadowData.ShadowResolution;
+					Params.LightDirection = DeepShadowData.LightDirection;
+					Params.LightPosition = DeepShadowData.LightPosition;
+					Params.LightRadius = 0;
+					Params.LightChannelMask = LightSceneInfo->Proxy->GetLightingChannelMask();
+					Params.DepthBiasScale = GetDeepShadowDepthBiasScale();
+					Params.DeepShadow_AtlasSlotOffsets_AtlasSlotIndex[DeepShadowData.MacroGroupId] = FIntVector4(DeepShadowData.AtlasRect.Min.X, DeepShadowData.AtlasRect.Min.Y, DeepShadowData.AtlasSlotIndex, 0);
+					Params.DeepShadow_CPUWorldToLightTransforms[DeepShadowData.MacroGroupId] = DeepShadowData.CPU_WorldToLightTransform;
+					Params.DeepShadow_WorldToLightTransformBuffer = DeepShadow_WorldToLightTransformBufferSRV;
+					Params.DeepShadow_bIsGPUDriven = MacroGroupDatas.DeepShadowResources.bIsGPUDriven;
 				}
-				Params.DeepShadow_Resolution = DeepShadowData.ShadowResolution;
-				Params.LightDirection = DeepShadowData.LightDirection;
-				Params.LightPosition = DeepShadowData.LightPosition;
-				Params.LightRadius = 0;
-				Params.LightChannelMask = LightSceneInfo->Proxy->GetLightingChannelMask();
-				Params.DepthBiasScale = GetDeepShadowDepthBiasScale();
-				Params.DeepShadow_AtlasSlotOffsets[DeepShadowData.MacroGroupId] = FIntVector4(DeepShadowData.AtlasRect.Min.X, DeepShadowData.AtlasRect.Min.Y, DeepShadowData.AtlasRect.Width(), DeepShadowData.AtlasRect.Height());
-				Params.DeepShadow_WorldToLightTransforms[DeepShadowData.MacroGroupId] = DeepShadowData.WorldToLightTransform;
 			}
 		}
 
@@ -643,11 +674,10 @@ FHairStrandsTransmittanceMaskData RenderHairStrandsTransmittanceMask(
 		for (int32 ViewIndex = 0; ViewIndex < Views.Num(); ViewIndex++)
 		{
 			const FViewInfo& View = Views[ViewIndex];
-			const FHairStrandsDeepShadowDatas& InDeepShadowDatas = HairDatas->DeepShadowViews.Views[ViewIndex];
 			const FHairStrandsVisibilityData& InHairVisibilityData = HairDatas->HairVisibilityViews.HairDatas[ViewIndex];
 			const FHairStrandsMacroGroupDatas& InMacroGroupDatas = HairDatas->MacroGroupsPerViews.Views[ViewIndex];
 
-			TransmittanceMaskData = RenderHairStrandsTransmittanceMask(RHICmdList, View, LightSceneInfo, InDeepShadowDatas, InMacroGroupDatas, InHairVisibilityData, ScreenShadowMaskSubPixelTexture);
+			TransmittanceMaskData = RenderHairStrandsTransmittanceMask(RHICmdList, View, LightSceneInfo, InMacroGroupDatas, InHairVisibilityData, ScreenShadowMaskSubPixelTexture);
 		}
 	}
 	return TransmittanceMaskData;
@@ -660,14 +690,13 @@ static void RenderHairStrandsShadowMask(
 	const FViewInfo& View,
 	const FLightSceneInfo* LightSceneInfo,
 	IPooledRenderTarget* ScreenShadowMaskTexture,
-	const FHairStrandsDeepShadowDatas& DeepShadowDatas,
 	const FHairStrandsVisibilityData& InVisibilityData,
 	const FHairStrandsMacroGroupDatas& InMacroGroupDatas)
 {
 	if (InMacroGroupDatas.Datas.Num() == 0)
 		return;
 
-	if (!HasDeepShadowData(LightSceneInfo, DeepShadowDatas) && !IsHairStrandsVoxelizationEnable())
+	if (!HasDeepShadowData(LightSceneInfo, InMacroGroupDatas) && !IsHairStrandsVoxelizationEnable())
 		return;
 
 	DECLARE_GPU_STAT(HairStrandsOpaqueMask);
@@ -684,36 +713,52 @@ static void RenderHairStrandsShadowMask(
 	FRDGTextureRef Categorization = RegisterExternalTextureWithFallback(GraphBuilder, InVisibilityData.CategorizationTexture, GSystemTextures.BlackDummy);
 	FRDGTextureRef OutShadowMask = GraphBuilder.RegisterExternalTexture(ScreenShadowMaskTexture, TEXT("ScreenShadowMaskTexture"));
 
+
 	bool bHasDeepShadow = false;
 	if (!IsHairStrandsForVoxelTransmittanceAndShadowEnable())
 	{
-		for (const FHairStrandsDeepShadowData& DomData : DeepShadowDatas.Datas)
-		{
-			if (DomData.LightId != LightSceneInfo->Id)
-				continue;
+		FRDGBufferRef DeepShadow_WorldToLightTransformBuffer = nullptr;
+		FRDGBufferSRVRef DeepShadow_WorldToLightTransformBufferSRV = nullptr;
 
-			bHasDeepShadow = true;
-			FRDGTextureRef DeepShadowDepth = GraphBuilder.RegisterExternalTexture(DomData.DepthTexture, TEXT("DeepShadowDepthTexture"));
-			FRDGTextureRef DeepShadowLayers = GraphBuilder.RegisterExternalTexture(DomData.LayersTexture, TEXT("DeepShadowLayerTexture"));
-			const bool bIsWholeSceneLight = LightSceneInfo->Proxy->GetLightType() == LightType_Directional;
+		for (const FHairStrandsMacroGroupData& MacroGroupData : InMacroGroupDatas.Datas)
+		{
+			for (const FHairStrandsDeepShadowData& DomData : MacroGroupData.DeepShadowDatas.Datas)
+			{
+				if (DomData.LightId != LightSceneInfo->Id)
+					continue;
 
-			FDeepShadowOpaqueParams Params;
-			Params.CategorizationTexture = Categorization;
-			Params.DeepShadow_WorldToLightTransform = DomData.WorldToLightTransform;
-			Params.DeepShadow_AtlasRect = DomData.AtlasRect;
-			Params.DeepShadow_FrontDepthTexture = DeepShadowDepth;
-			Params.DeepShadow_LayerTexture = DeepShadowLayers;
-			Params.DeepShadow_bIsWholeSceneLight = bIsWholeSceneLight;
-			Params.DeepShadow_DepthBiasScale = GetDeepShadowDepthBiasScale();
-			Params.DeepShadow_DensityScale = GetDeepShadowDensityScale();
+				if (DeepShadow_WorldToLightTransformBuffer == nullptr)
+				{
+					DeepShadow_WorldToLightTransformBuffer = GraphBuilder.RegisterExternalBuffer(InMacroGroupDatas.DeepShadowResources.DeepShadowWorldToLightTransforms);
+					DeepShadow_WorldToLightTransformBufferSRV = GraphBuilder.CreateSRV(DeepShadow_WorldToLightTransformBuffer);
+				}
 
-			AddDeepShadowOpaqueMaskPass(
-				GraphBuilder,
-				SceneTextures,
-				View,
-				FHairOpaqueMaskType_DeepShadow,
-				Params,
-				OutShadowMask);
+				bHasDeepShadow = true;
+				FRDGTextureRef DeepShadowDepth = GraphBuilder.RegisterExternalTexture(InMacroGroupDatas.DeepShadowResources.DepthAtlasTexture, TEXT("DeepShadowDepthTexture"));
+				FRDGTextureRef DeepShadowLayers = GraphBuilder.RegisterExternalTexture(InMacroGroupDatas.DeepShadowResources.LayersAtlasTexture, TEXT("DeepShadowLayerTexture"));
+				const bool bIsWholeSceneLight = LightSceneInfo->Proxy->GetLightType() == LightType_Directional;
+
+				FDeepShadowOpaqueParams Params;
+				Params.CategorizationTexture = Categorization;
+				Params.DeepShadow_AtlasSlotIndex = DomData.AtlasSlotIndex;
+				Params.DeepShadow_WorldToLightTransformBuffer = DeepShadow_WorldToLightTransformBufferSRV;
+				Params.DeepShadow_bIsGPUDriven = InMacroGroupDatas.DeepShadowResources.bIsGPUDriven ? 1 : 0;
+				Params.DeepShadow_CPUWorldToLightTransform = DomData.CPU_WorldToLightTransform;
+				Params.DeepShadow_AtlasRect = DomData.AtlasRect;
+				Params.DeepShadow_FrontDepthTexture = DeepShadowDepth;
+				Params.DeepShadow_LayerTexture = DeepShadowLayers;
+				Params.DeepShadow_bIsWholeSceneLight = bIsWholeSceneLight;
+				Params.DeepShadow_DepthBiasScale = GetDeepShadowDepthBiasScale();
+				Params.DeepShadow_DensityScale = GetDeepShadowDensityScale();
+
+				AddDeepShadowOpaqueMaskPass(
+					GraphBuilder,
+					SceneTextures,
+					View,
+					FHairOpaqueMaskType_DeepShadow,
+					Params,
+					OutShadowMask);
+			}
 		}
 	}
 
@@ -763,18 +808,16 @@ void RenderHairStrandsShadowMask(
 	IPooledRenderTarget* ScreenShadowMaskTexture,
 	const FHairStrandsDatas* HairDatas)
 {
-	const FHairStrandsDeepShadowViews& DeepShadowViews = HairDatas->DeepShadowViews;
 	const FHairStrandsVisibilityViews& HairVisibilityViews = HairDatas->HairVisibilityViews;
 	const FHairStrandsMacroGroupViews& MacroGroupViews = HairDatas->MacroGroupsPerViews;
 
 	for (int32 ViewIndex = 0; ViewIndex < Views.Num(); ++ViewIndex)
 	{
-		if (ViewIndex < DeepShadowViews.Views.Num() && ViewIndex < HairVisibilityViews.HairDatas.Num() && ViewIndex < MacroGroupViews.Views.Num())
+		if (ViewIndex < HairVisibilityViews.HairDatas.Num() && ViewIndex < MacroGroupViews.Views.Num())
 		{
 			const FHairStrandsVisibilityData& HairVisibilityData = HairVisibilityViews.HairDatas[ViewIndex];
-			const FHairStrandsDeepShadowDatas& DeepShadowDatas = DeepShadowViews.Views[ViewIndex];
 			const FHairStrandsMacroGroupDatas& MacroGroupDatas = MacroGroupViews.Views[ViewIndex];
-			RenderHairStrandsShadowMask(RHICmdList, Views[ViewIndex], LightSceneInfo, ScreenShadowMaskTexture, DeepShadowDatas, HairVisibilityData, MacroGroupDatas);
+			RenderHairStrandsShadowMask(RHICmdList, Views[ViewIndex], LightSceneInfo, ScreenShadowMaskTexture, HairVisibilityData, MacroGroupDatas);
 		}
 	}
 }
diff --git a/Engine/Source/Runtime/Renderer/Private/HairStrands/HairStrandsUtils.cpp b/Engine/Source/Runtime/Renderer/Private/HairStrands/HairStrandsUtils.cpp
--- a/Engine/Source/Runtime/Renderer/Private/HairStrands/HairStrandsUtils.cpp
+++ b/Engine/Source/Runtime/Renderer/Private/HairStrands/HairStrandsUtils.cpp
@@ -72,6 +72,17 @@ uint32 ToBitfield(const FHairComponent& C)
 		(C.LocalScattering  ? 1u : 0u) << 3 |
 		(C.GlobalScattering ? 1u : 0u) << 4 ;
 }
+
+float GetPrimiartyRasterizationScale()
+{
+	return FMath::Max(0.f, GStrandHairRasterizationScale);
+}
+
+float GetDeepShadowRasterizationScale()
+{
+	return FMath::Max(0.f, GStrandHairShadowRasterizationScale ? GStrandHairShadowRasterizationScale : GStrandHairRasterizationScale);
+}
+
 FMinHairRadiusAtDepth1 ComputeMinStrandRadiusAtDepth1(
 	const FIntPoint& Resolution,
 	const float FOV,
@@ -114,7 +125,7 @@ void ComputeWorldToLightClip(
 	const float MinZ = FMath::Max(0.1f, FVector::Distance(LightPosition, SphereBound.Center)) - SphereBound.W;
 	const float MaxZ = FMath::Max(0.2f, FVector::Distance(LightPosition, SphereBound.Center)) + SphereBound.W;
 
-	const float StrandHairRasterizationScale = GStrandHairShadowRasterizationScale ? GStrandHairShadowRasterizationScale : GStrandHairRasterizationScale;
+	const float StrandHairRasterizationScale = GetDeepShadowRasterizationScale();
 	const float StrandHairStableRasterizationScale = FMath::Max(GStrandHairStableRasterizationScale, 1.0f);
 	MinStrandRadiusAtDepth1 = FMinHairRadiusAtDepth1();
 	WorldToClipTransform = FMatrix::Identity;
diff --git a/Engine/Source/Runtime/Renderer/Private/LightRendering.cpp b/Engine/Source/Runtime/Renderer/Private/LightRendering.cpp
--- a/Engine/Source/Runtime/Renderer/Private/LightRendering.cpp
+++ b/Engine/Source/Runtime/Renderer/Private/LightRendering.cpp
@@ -976,7 +976,6 @@ void FDeferredShadingSceneRenderer::GatherAndSortLights(FSortedLightSetSceneInfo
 void FDeferredShadingSceneRenderer::RenderLights(FRHICommandListImmediate& RHICmdList, FSortedLightSetSceneInfo &SortedLightSet, const FHairStrandsDatas* HairDatas)
 {
 	const bool bUseHairLighting = HairDatas != nullptr;
-	const FHairStrandsDeepShadowViews* InDeepShadowViews = HairDatas ? &HairDatas->DeepShadowViews : nullptr;
 	const FHairStrandsVisibilityViews* InHairVisibilityViews = HairDatas ? &HairDatas->HairVisibilityViews : nullptr;
 	auto HasHairStrandsClusters = [bUseHairLighting, HairDatas](int32 ViewIndex)
 	{
