Update to the way distance field warning messages are handled.

diff --git a/Engine/Source/Runtime/Renderer/Private/SceneRendering.cpp b/Engine/Source/Runtime/Renderer/Private/SceneRendering.cpp
--- a/Engine/Source/Runtime/Renderer/Private/SceneRendering.cpp
+++ b/Engine/Source/Runtime/Renderer/Private/SceneRendering.cpp
@@ -2704,8 +2704,11 @@ void FSceneRenderer::RenderFinish(FRHICommandListImmediate& RHICmdList)
 		}
 		const bool bShowSkinCacheOOM = CVarSkinCacheOOM != nullptr && GPUSkinCacheExtraRequiredMemory > 0;
 
+		static const auto* CVarGenerateMeshDistanceFields = IConsoleManager::Get().FindTConsoleVariableDataInt(TEXT("r.GenerateMeshDistanceFields"));
+		const bool bMeshDistanceFieldEnabled = CVarGenerateMeshDistanceFields != nullptr && CVarGenerateMeshDistanceFields->GetValueOnRenderThread() > 0;
 		extern bool UseDistanceFieldAO();
-		const bool bShowDFAODisabledWarning = !UseDistanceFieldAO() && (ViewFamily.EngineShowFlags.VisualizeMeshDistanceFields || ViewFamily.EngineShowFlags.VisualizeGlobalDistanceField || ViewFamily.EngineShowFlags.VisualizeDistanceFieldAO);
+		const bool bShowDFAODisabledWarning = !UseDistanceFieldAO() && (ViewFamily.EngineShowFlags.VisualizeDistanceFieldAO);
+		const bool bShowDFDisabledWarning = !bMeshDistanceFieldEnabled && (ViewFamily.EngineShowFlags.VisualizeMeshDistanceFields || ViewFamily.EngineShowFlags.VisualizeGlobalDistanceField || ViewFamily.EngineShowFlags.VisualizeDistanceFieldAO);
 
 		const bool bShowAtmosphericFogWarning = Scene->AtmosphericFog != nullptr && !ReadOnlyCVARCache.bEnableAtmosphericFog;
 
@@ -2735,7 +2738,7 @@ void FSceneRenderer::RenderFinish(FRHICommandListImmediate& RHICmdList)
 		
 		const bool bAnyWarning = bShowPrecomputedVisibilityWarning || bShowGlobalClipPlaneWarning || bShowAtmosphericFogWarning || bShowSkylightWarning || bShowPointLightWarning 
 			|| bShowDFAODisabledWarning || bShowShadowedLightOverflowWarning || bShowMobileDynamicCSMWarning || bShowMobileLowQualityLightmapWarning || bShowMobileMovableDirectionalLightWarning
-			|| bMobileShowVertexFogWarning || bShowSkinCacheOOM || bSingleLayerWaterWarning;
+			|| bMobileShowVertexFogWarning || bShowSkinCacheOOM || bSingleLayerWaterWarning || bShowDFDisabledWarning;
 
 		for(int32 ViewIndex = 0;ViewIndex < Views.Num();ViewIndex++)
 		{	
@@ -2788,6 +2791,12 @@ void FSceneRenderer::RenderFinish(FRHICommandListImmediate& RHICmdList)
 						Canvas.DrawShadowedText(10, Y, Message, GetStatsFont(), FLinearColor(1.0, 0.05, 0.05, 1.0));
 						Y += 14;
 					}
+					if (bShowDFDisabledWarning)
+					{
+						static const FText Message = NSLOCTEXT("Renderer", "DFDisabled", "Mesh distance fields generation is disabled by project settings, cannot visualize DFAO, mesh or global distance field.");
+						Canvas.DrawShadowedText(10, Y, Message, GetStatsFont(), FLinearColor(1.0, 0.05, 0.05, 1.0));
+						Y += 14;
+					}
 
 					if (bShowAtmosphericFogWarning)
 					{
