Clamp hardcoded mip level to one that works for the given texture size when creating texture in SSS pass

#rnx

diff --git a/Engine/Source/Runtime/Renderer/Private/PostProcess/PostProcessSubsurface.cpp b/Engine/Source/Runtime/Renderer/Private/PostProcess/PostProcessSubsurface.cpp
--- a/Engine/Source/Runtime/Renderer/Private/PostProcess/PostProcessSubsurface.cpp
+++ b/Engine/Source/Runtime/Renderer/Private/PostProcess/PostProcessSubsurface.cpp
@@ -806,6 +806,7 @@ void ComputeSubsurfaceForView(
 		TexCreate_RenderTargetable | TexCreate_ShaderResource | TexCreate_UAV,
 		false);
 
+	// Create texture desc with 6 mips if possible, otherwise clamp number of mips to match the viewport resolution
 	const FRDGTextureDesc SubsurfaceTextureWith6MipsDescriptor = FRDGTextureDesc::Create2DDesc(
 		SubsurfaceViewport.Extent,
 		PF_FloatRGBA,
@@ -813,7 +814,7 @@ void ComputeSubsurfaceForView(
 		TexCreate_None,
 		TexCreate_RenderTargetable | TexCreate_ShaderResource | TexCreate_UAV,
 		false,
-		6);
+		FMath::Min(6u, 1 + FMath::FloorLog2((uint32)SubsurfaceViewport.Extent.GetMin())));
 
 	const FSubsurfaceParameters SubsurfaceCommonParameters = GetSubsurfaceCommonParameters(GraphBuilder.RHICmdList, View);
 	const FScreenPassTextureViewportParameters SubsurfaceViewportParameters = GetScreenPassTextureViewportParameters(SubsurfaceViewport);
