Fixed crash in ray tracing translucency

diff --git a/Engine/Source/Runtime/Renderer/Private/RayTracing/RayTracingTranslucency.cpp b/Engine/Source/Runtime/Renderer/Private/RayTracing/RayTracingTranslucency.cpp
--- a/Engine/Source/Runtime/Renderer/Private/RayTracing/RayTracingTranslucency.cpp
+++ b/Engine/Source/Runtime/Renderer/Private/RayTracing/RayTracingTranslucency.cpp
@@ -161,28 +161,30 @@ void FDeferredShadingSceneRenderer::RenderRayTracingTranslucency(FRHICommandList
 	FSceneTextureParameters SceneTextures;
 	SetupSceneTextureParameters(GraphBuilder, &SceneTextures);
 
-	RDG_EVENT_SCOPE(GraphBuilder, "RayTracingTranslucency");
-	RDG_GPU_STAT_SCOPE(GraphBuilder, RayTracingTranslucency)
-
-	for (int32 ViewIndex = 0, Num = Views.Num(); ViewIndex < Num; ViewIndex++)
 	{
-		FViewInfo& View = Views[ViewIndex];
+		RDG_EVENT_SCOPE(GraphBuilder, "RayTracingTranslucency");
+		RDG_GPU_STAT_SCOPE(GraphBuilder, RayTracingTranslucency)
 
-		const FScreenPassTexture SceneColor(SceneColorTexture, View.ViewRect);
-		const FScreenPassRenderTarget Output(SceneColorTexture, View.ViewRect, ERenderTargetLoadAction::ELoad);
+		for (int32 ViewIndex = 0, Num = Views.Num(); ViewIndex < Num; ViewIndex++)
+		{
+			FViewInfo& View = Views[ViewIndex];
 
-		//#dxr_todo: UE-72581 do not use reflections denoiser structs but separated ones
-		IScreenSpaceDenoiser::FReflectionsInputs DenoiserInputs;
-		float ResolutionFraction = 1.0f;
-		int32 TranslucencySPP = GRayTracingTranslucencySamplesPerPixel > -1 ? GRayTracingTranslucencySamplesPerPixel : View.FinalPostProcessSettings.RayTracingTranslucencySamplesPerPixel;
-		
-		RenderRayTracingPrimaryRaysView(
-			GraphBuilder,
-			View, &DenoiserInputs.Color, &DenoiserInputs.RayHitDistance,
-			TranslucencySPP, GRayTracingTranslucencyHeightFog, ResolutionFraction, 
-			ERayTracingPrimaryRaysFlag::AllowSkipSkySample | ERayTracingPrimaryRaysFlag::UseGBufferForMaxDistance);
+			const FScreenPassRenderTarget Output(SceneColorTexture, View.ViewRect, ERenderTargetLoadAction::ELoad);
 
-		AddDrawTexturePass(GraphBuilder, View, SceneColor, Output);
+			//#dxr_todo: UE-72581 do not use reflections denoiser structs but separated ones
+			IScreenSpaceDenoiser::FReflectionsInputs DenoiserInputs;
+			float ResolutionFraction = 1.0f;
+			int32 TranslucencySPP = GRayTracingTranslucencySamplesPerPixel > -1 ? GRayTracingTranslucencySamplesPerPixel : View.FinalPostProcessSettings.RayTracingTranslucencySamplesPerPixel;
+		
+			RenderRayTracingPrimaryRaysView(
+				GraphBuilder,
+				View, &DenoiserInputs.Color, &DenoiserInputs.RayHitDistance,
+				TranslucencySPP, GRayTracingTranslucencyHeightFog, ResolutionFraction, 
+				ERayTracingPrimaryRaysFlag::AllowSkipSkySample | ERayTracingPrimaryRaysFlag::UseGBufferForMaxDistance);
+
+			const FScreenPassTexture SceneColor(DenoiserInputs.Color, View.ViewRect);
+			AddDrawTexturePass(GraphBuilder, View, SceneColor, Output);
+		}
 	}
 
 	GraphBuilder.Execute();
