UE4.25 - Unblock OpenGL
#rnx

diff --git a/Engine/Source/Runtime/Renderer/Private/GpuDebugRendering.cpp b/Engine/Source/Runtime/Renderer/Private/GpuDebugRendering.cpp
--- a/Engine/Source/Runtime/Renderer/Private/GpuDebugRendering.cpp
+++ b/Engine/Source/Runtime/Renderer/Private/GpuDebugRendering.cpp
@@ -50,7 +50,7 @@ namespace ShaderDrawDebug
 
 	static bool IsShaderDrawDebugEnabled(const EShaderPlatform Platform)
 	{
-		return IsFeatureLevelSupported(Platform, ERHIFeatureLevel::SM5) && IsPCPlatform(Platform);
+		return IsFeatureLevelSupported(Platform, ERHIFeatureLevel::SM5) && IsPCPlatform(Platform) && !IsOpenGLPlatform(Platform);
 	}
 
 	bool IsShaderDrawDebugEnabled(const FViewInfo& View)
