Fix typo causing CIS issue.

diff --git a/Engine/Source/Runtime/Renderer/Private/HairStrands/HairStrandsDeepShadow.h b/Engine/Source/Runtime/Renderer/Private/HairStrands/HairStrandsDeepShadow.h
--- a/Engine/Source/Runtime/Renderer/Private/HairStrands/HairStrandsDeepShadow.h
+++ b/Engine/Source/Runtime/Renderer/Private/HairStrands/HairStrandsDeepShadow.h
@@ -21,7 +21,7 @@ struct FHairStrandsDeepShadowData
 	FMinHairRadiusAtDepth1 CPU_MinStrandRadiusAtDepth1;
 	FIntRect AtlasRect;
 	uint32 MacroGroupId = ~0;
-	uint32 AtlasSlotIndex = AtlasSlotIndex;
+	uint32 AtlasSlotIndex = 0;
 
 	FIntPoint ShadowResolution = FIntPoint::ZeroValue;
 	uint32 LightId = ~0;
