Fixed ray tracing was not taking into account layer visibility to add/exlude objects from the ray tracing world.

diff --git a/Engine/Source/Runtime/Renderer/Private/DeferredShadingRenderer.cpp b/Engine/Source/Runtime/Renderer/Private/DeferredShadingRenderer.cpp
--- a/Engine/Source/Runtime/Renderer/Private/DeferredShadingRenderer.cpp
+++ b/Engine/Source/Runtime/Renderer/Private/DeferredShadingRenderer.cpp
@@ -765,6 +765,16 @@ bool FDeferredShadingSceneRenderer::GatherRayTracingWorldInstances(FRHICommandLi
 					continue;
 				}
 
+				if (View.HiddenPrimitives.Contains(SceneInfo->PrimitiveComponentId))
+				{
+					continue;
+				}
+
+				if (View.ShowOnlyPrimitives.IsSet() && !View.ShowOnlyPrimitives->Contains(SceneInfo->PrimitiveComponentId))
+				{
+					continue;
+				}
+
 				bool bShouldRayTraceSceneCapture = GRayTracingSceneCaptures > 0 || (GRayTracingSceneCaptures == -1 && View.bSceneCaptureUsesRayTracing);
 				if (View.bIsSceneCapture && (!bShouldRayTraceSceneCapture || !SceneInfo->bIsVisibleInReflectionCaptures))
 				{
