Remove plain voxelization support for hair from transmittance and shadow mask.

diff --git a/Engine/Shaders/Private/HairStrands/HairStrandsDeepShadowMask.usf b/Engine/Shaders/Private/HairStrands/HairStrandsDeepShadowMask.usf
--- a/Engine/Shaders/Private/HairStrands/HairStrandsDeepShadowMask.usf
+++ b/Engine/Shaders/Private/HairStrands/HairStrandsDeepShadowMask.usf
@@ -11,23 +11,17 @@
 #include "HairStrandsVoxelPageCommon.ush"
 
 #include "HairStrandsCommon.ush"
-#include "HairStrandsVoxelPageCommon.ush"
 #include "HairStrandsDeepShadowCommon.ush"
 #include "HairStrandsDebugCommon.ush"
 #include "HairStrandsDeepTransmittanceCommon.ush"
 
 #define OPAQUE_MASK_TYPE_DEEP_SHADOW 0
-#define OPAQUE_MASK_TYPE_VOXEL 1
-#define OPAQUE_MASK_TYPE_VIRTUAL_VOXEL 2
+#define OPAQUE_MASK_TYPE_VIRTUAL_VOXEL 1
 
 #define VOXEL_TRAVERSAL_DEBUG 0
-#if PERMUTATION_OPAQUEMASK_TYPE == OPAQUE_MASK_TYPE_VOXEL || PERMUTATION_OPAQUEMASK_TYPE == OPAQUE_MASK_TYPE_VIRTUAL_VOXEL
-  #if PERMUTATION_OPAQUEMASK_TYPE == OPAQUE_MASK_TYPE_VOXEL
-	#define VOXEL_TRAVERSAL_TYPE VOXEL_TRAVERSAL_LEGACY
-  #elif PERMUTATION_OPAQUEMASK_TYPE == OPAQUE_MASK_TYPE_VIRTUAL_VOXEL
+#if PERMUTATION_OPAQUEMASK_TYPE == OPAQUE_MASK_TYPE_VIRTUAL_VOXEL
 	#define VOXEL_TRAVERSAL_TYPE VOXEL_TRAVERSAL_LINEAR_MIPMAP
-  #endif
-  #include "HairStrandsVoxelPageTraversal.ush"
+	#include "HairStrandsVoxelPageTraversal.ush"
 #endif
 
 StructuredBuffer<FDeepShadowTransform> DeepShadow_WorldToLightTransformBuffer;
@@ -44,12 +38,7 @@ float				DeepShadow_DensityScale;
 
 float4				Voxel_LightPosition;
 float3				Voxel_LightDirection;
-float				Voxel_DensityScale;
-float3				Voxel_MinAABB;
-uint				Voxel_Resolution;
-float3				Voxel_MaxAABB;
 uint				Voxel_MacroGroupId;
-Texture3D<uint> 	Voxel_DensityTexture;
 
 Texture2D<float4>	RayMarchMaskTexture;
 
@@ -409,7 +398,7 @@ void MainPS(
 
 #endif
 
-#if PERMUTATION_OPAQUEMASK_TYPE == OPAQUE_MASK_TYPE_VOXEL || PERMUTATION_OPAQUEMASK_TYPE == OPAQUE_MASK_TYPE_VIRTUAL_VOXEL
+#if PERMUTATION_OPAQUEMASK_TYPE == OPAQUE_MASK_TYPE_VIRTUAL_VOXEL
 
 void MainPS(
 	FScreenVertexOutput Input,
@@ -436,13 +425,7 @@ void MainPS(
 		// Origin is shifted 2 voxels away towards the light + a constant bias of the size of the voxel
 		const float3 Random = GetHairVoxelJitter(PixelCoord.xy, View.StateFrameIndexMod8);
 		{
-		#if PERMUTATION_OPAQUEMASK_TYPE == OPAQUE_MASK_TYPE_VOXEL
-			const float3 VoxelSize = (Voxel_MaxAABB - Voxel_MinAABB) / Voxel_Resolution;
-			const float MaxVoxelSize = max(VoxelSize.x, max(VoxelSize.y, VoxelSize.z));
-		#elif PERMUTATION_OPAQUEMASK_TYPE == OPAQUE_MASK_TYPE_VIRTUAL_VOXEL
-			const float MaxVoxelSize = VirtualVoxel.VoxelWorldSize;
-		#endif
-			const float3 DepthBias = MaxVoxelSize * (2 * LightDirection + Random);
+			const float3 DepthBias = VirtualVoxel.VoxelWorldSize * (2 * LightDirection + Random);
 			WorldPosition += DepthBias;
 		}
 
@@ -451,27 +434,6 @@ void MainPS(
 		const float DistanceThreshold = 100000.0f;
 		const float CoverageThreshold = 0.995f; // When Coverage is high, we do not trace shadow on opaque since hair/fur is covering the background.
 
-	#if PERMUTATION_OPAQUEMASK_TYPE == OPAQUE_MASK_TYPE_VOXEL
-		if (RayMarchMaskTexture.Load(uint3(PixelCoord, 0)).r > 0.0f &&  Categorization.PixelCoverage < CoverageThreshold)
-		{
-			Result = ComputeHairCountVoxel(
-				WorldPosition,
-				LightPosition,
-				Voxel_MinAABB,
-				Voxel_MaxAABB,
-				Voxel_Resolution,
-				Voxel_DensityTexture,
-				Voxel_DensityScale,
-				1,
-				DistanceThreshold);
-		}
-		else
-		{
-			discard;
-		}
-	#endif
-
-	#if PERMUTATION_OPAQUEMASK_TYPE == OPAQUE_MASK_TYPE_VIRTUAL_VOXEL
 		FVirtualVoxelCommonDesc CommonDesc;
 		CommonDesc.PageCountResolution = VirtualVoxel.PageCountResolution;
 		CommonDesc.VoxelWorldSize = VirtualVoxel.VoxelWorldSize;
@@ -511,7 +473,6 @@ void MainPS(
 		{
 			discard;
 		}
-	#endif
 
 		Visibility = min(saturate(1 - Result.HairCount), Result.Visibility);
 	}
diff --git a/Engine/Shaders/Private/HairStrands/HairStrandsDeepTransmittanceMask.usf b/Engine/Shaders/Private/HairStrands/HairStrandsDeepTransmittanceMask.usf
--- a/Engine/Shaders/Private/HairStrands/HairStrandsDeepTransmittanceMask.usf
+++ b/Engine/Shaders/Private/HairStrands/HairStrandsDeepTransmittanceMask.usf
@@ -19,14 +19,11 @@
 #include "../ShaderPrintCommon.ush"
 
 #define PERMUTATION_TRANSMITTANCE_DEEP_SHADOW 0
-#define PERMUTATION_TRANSMITTANCE_VOXEL 1
-#define PERMUTATION_TRANSMITTANCE_VIRTUAL_VOXEL 2
+#define PERMUTATION_TRANSMITTANCE_VIRTUAL_VOXEL 1
 
 
 #if PERMUTATION_TRANSMITTANCE_TYPE == PERMUTATION_TRANSMITTANCE_DEEP_SHADOW
 	#define VOXEL_TRAVERSAL_TYPE VOXEL_TRAVERSAL_LEGACY
-#elif PERMUTATION_TRANSMITTANCE_TYPE == PERMUTATION_TRANSMITTANCE_VOXEL
-	#define VOXEL_TRAVERSAL_TYPE VOXEL_TRAVERSAL_LEGACY
 #elif PERMUTATION_TRANSMITTANCE_TYPE == PERMUTATION_TRANSMITTANCE_VIRTUAL_VOXEL
 	#if PERMUTATION_TRAVERSAL == 1
 		#define VOXEL_TRAVERSAL_TYPE VOXEL_TRAVERSAL_LINEAR_MIPMAP
@@ -69,18 +66,6 @@ int4				DeepShadow_AtlasSlotOffsets_AtlasSlotIndex[MAX_HAIR_MACROGROUP_COUNT];
 Texture2D<float>	DeepShadow_FrontDepthTexture;
 Texture2D<float4>	DeepShadow_DomTexture;
 
-float4				Voxel_MinAABBs[MAX_HAIR_MACROGROUP_COUNT];
-float4				Voxel_MaxAABBs[MAX_HAIR_MACROGROUP_COUNT];
-uint				Voxel_Resolution;
-Texture3D<uint> 	Voxel_DensityTexture0;
-Texture3D<uint> 	Voxel_DensityTexture1;
-Texture3D<uint> 	Voxel_DensityTexture2;
-Texture3D<uint> 	Voxel_DensityTexture3;
-Texture3D<uint> 	Voxel_DensityTexture4;
-Texture3D<uint> 	Voxel_DensityTexture5;
-Texture3D<uint> 	Voxel_DensityTexture6;
-Texture3D<uint> 	Voxel_DensityTexture7;
-
 Texture2D<float4> 	RayMarchMaskTexture;
 
 StructuredBuffer<FPackedHairSample> HairVisibilityNodeData;
@@ -168,7 +153,7 @@ void MainCS(uint2 DispatchThreadId : SV_DispatchThreadID)
 }
 #endif
 
-#if PERMUTATION_TRANSMITTANCE_TYPE == PERMUTATION_TRANSMITTANCE_VOXEL || PERMUTATION_TRANSMITTANCE_TYPE == PERMUTATION_TRANSMITTANCE_VIRTUAL_VOXEL
+#if PERMUTATION_TRANSMITTANCE_TYPE == PERMUTATION_TRANSMITTANCE_VIRTUAL_VOXEL
 
 #if VOXEL_TRAVERSAL_DEBUG
 void PrintInfo(uint Index, uint2 PixelCoord, FHairTransmittanceMask TransmittanceMask)
@@ -235,15 +220,7 @@ FHairTransmittanceMask ComputeTransmittanceVoxel(
 
 	// Depth bias
 	// Origin is shifted voxels away towards the light + a constant bias of the size of the voxel
-	float MaxVoxelSize;
-	{
-	#if PERMUTATION_TRANSMITTANCE_TYPE == PERMUTATION_TRANSMITTANCE_VOXEL
-		const float3 VoxelSize = (Voxel_MaxAABBs[MacroGroupId].xyz - Voxel_MinAABBs[MacroGroupId].xyz) / Voxel_Resolution;
-		MaxVoxelSize = max(VoxelSize.x, max(VoxelSize.y, VoxelSize.z));
-	#elif PERMUTATION_TRANSMITTANCE_TYPE == PERMUTATION_TRANSMITTANCE_VIRTUAL_VOXEL
-		MaxVoxelSize = VirtualVoxel.VoxelWorldSize;
-	#endif
-	}
+	const float MaxVoxelSize = VirtualVoxel.VoxelWorldSize;
 
 	// Compute the number of hair count between light & shading point
 	const float DistanceThreshold = 100000;
@@ -265,19 +242,6 @@ FHairTransmittanceMask ComputeTransmittanceVoxel(
 		const float3 DepthBias = MaxVoxelSize * (Settings.DeepShadowDepthBiasScale*LightDirection + SampleRandom);
 		const float3 SampleWorldPosition = WorldPosition + DepthBias;
 
-		#if PERMUTATION_TRANSMITTANCE_TYPE == PERMUTATION_TRANSMITTANCE_VOXEL
-		FHairTraversalResult Result = InitHairTraversalResult();
-		if (MacroGroupId == 0) Result = ComputeHairCountVoxel(SampleWorldPosition, LightPosition, Voxel_MinAABBs[0].xyz, Voxel_MaxAABBs[0].xyz, Voxel_Resolution, Voxel_DensityTexture0, Settings.DeepShadowDensityScale, 0, DistanceThreshold);
-		if (MacroGroupId == 1) Result = ComputeHairCountVoxel(SampleWorldPosition, LightPosition, Voxel_MinAABBs[1].xyz, Voxel_MaxAABBs[1].xyz, Voxel_Resolution, Voxel_DensityTexture1, Settings.DeepShadowDensityScale, 0, DistanceThreshold);
-		if (MacroGroupId == 2) Result = ComputeHairCountVoxel(SampleWorldPosition, LightPosition, Voxel_MinAABBs[2].xyz, Voxel_MaxAABBs[2].xyz, Voxel_Resolution, Voxel_DensityTexture2, Settings.DeepShadowDensityScale, 0, DistanceThreshold);
-		if (MacroGroupId == 3) Result = ComputeHairCountVoxel(SampleWorldPosition, LightPosition, Voxel_MinAABBs[3].xyz, Voxel_MaxAABBs[3].xyz, Voxel_Resolution, Voxel_DensityTexture3, Settings.DeepShadowDensityScale, 0, DistanceThreshold);
-		if (MacroGroupId == 4) Result = ComputeHairCountVoxel(SampleWorldPosition, LightPosition, Voxel_MinAABBs[4].xyz, Voxel_MaxAABBs[4].xyz, Voxel_Resolution, Voxel_DensityTexture4, Settings.DeepShadowDensityScale, 0, DistanceThreshold);
-		if (MacroGroupId == 5) Result = ComputeHairCountVoxel(SampleWorldPosition, LightPosition, Voxel_MinAABBs[5].xyz, Voxel_MaxAABBs[5].xyz, Voxel_Resolution, Voxel_DensityTexture5, Settings.DeepShadowDensityScale, 0, DistanceThreshold);
-		if (MacroGroupId == 6) Result = ComputeHairCountVoxel(SampleWorldPosition, LightPosition, Voxel_MinAABBs[6].xyz, Voxel_MaxAABBs[6].xyz, Voxel_Resolution, Voxel_DensityTexture6, Settings.DeepShadowDensityScale, 0, DistanceThreshold);
-		if (MacroGroupId == 7) Result = ComputeHairCountVoxel(SampleWorldPosition, LightPosition, Voxel_MinAABBs[7].xyz, Voxel_MaxAABBs[7].xyz, Voxel_Resolution, Voxel_DensityTexture7, Settings.DeepShadowDensityScale, 0, DistanceThreshold);
-		#endif
-
-		#if PERMUTATION_TRANSMITTANCE_TYPE == PERMUTATION_TRANSMITTANCE_VIRTUAL_VOXEL
 		FVirtualVoxelCommonDesc CommonDesc;
 		CommonDesc.PageCountResolution	= VirtualVoxel.PageCountResolution;
 		CommonDesc.VoxelWorldSize		= VirtualVoxel.VoxelWorldSize;
@@ -305,7 +269,6 @@ FHairTransmittanceMask ComputeTransmittanceVoxel(
 			VirtualVoxel.PageTexture,
 			TraversalSettings);
 		InitHairTraversalResult();
-		#endif
 
 		AvgHairCount += Result.HairCount;
 		AvgHairVisility += Result.Visibility;
diff --git a/Engine/Source/Runtime/Renderer/Private/HairStrands/HairStrandsTransmittance.cpp b/Engine/Source/Runtime/Renderer/Private/HairStrands/HairStrandsTransmittance.cpp
--- a/Engine/Source/Runtime/Renderer/Private/HairStrands/HairStrandsTransmittance.cpp
+++ b/Engine/Source/Runtime/Renderer/Private/HairStrands/HairStrandsTransmittance.cpp
@@ -45,8 +45,8 @@ static FAutoConsoleVariableRef CVarHairStrandsTransmittanceSuperSampling(TEXT("r
 static int32 GHairStrandsTransmittanceMaskUseMipTraversal = 1;
 static FAutoConsoleVariableRef CVarHairStrandsTransmittanceMaskUseMipTraversal(TEXT("r.HairStrands.DeepShadow.MipTraversal"), GHairStrandsTransmittanceMaskUseMipTraversal, TEXT("Evaluate transmittance using mip-map traversal (faster)."), ECVF_Scalability | ECVF_RenderThreadSafe);
 
-float GetDeepShadowDensityScale() { return FMath::Max(0.0f, GDeepShadowDensityScale); }
-float GetDeepShadowDepthBiasScale() { return FMath::Max(0.0f, GDeepShadowDepthBiasScale); }
+static float GetDeepShadowDensityScale() { return FMath::Max(0.0f, GDeepShadowDensityScale); }
+static float GetDeepShadowDepthBiasScale() { return FMath::Max(0.0f, GDeepShadowDepthBiasScale); }
 ///////////////////////////////////////////////////////////////////////////////////////////////////
 
 static bool HasDeepShadowData(const FLightSceneInfo* LightSceneInfo, const FHairStrandsMacroGroupDatas& InDatas)
@@ -69,7 +69,6 @@ static bool HasDeepShadowData(const FLightSceneInfo* LightSceneInfo, const FHair
 enum FHairTransmittanceType
 {
 	FHairTransmittanceType_DeepShadow,
-	FHairTransmittanceType_Voxel,
 	FHairTransmittanceType_VirtualVoxel,
 	FHairTransmittanceTypeCount
 };
@@ -106,18 +105,6 @@ class FDeepTransmittanceMaskCS : public FGlobalShader
 		SHADER_PARAMETER(FMatrix, DeepShadow_ShadowToWorld)
 		SHADER_PARAMETER(uint32, DeepShadow_bIsGPUDriven)
 
-		SHADER_PARAMETER_ARRAY(FVector4, Voxel_MinAABBs, [FHairStrandsDeepShadowData::MaxMacroGroupCount])
-		SHADER_PARAMETER_ARRAY(FVector4, Voxel_MaxAABBs, [FHairStrandsDeepShadowData::MaxMacroGroupCount])
-		SHADER_PARAMETER(uint32, Voxel_Resolution)
-		SHADER_PARAMETER_RDG_TEXTURE(Texture3D, Voxel_DensityTexture0)
-		SHADER_PARAMETER_RDG_TEXTURE(Texture3D, Voxel_DensityTexture1)
-		SHADER_PARAMETER_RDG_TEXTURE(Texture3D, Voxel_DensityTexture2)
-		SHADER_PARAMETER_RDG_TEXTURE(Texture3D, Voxel_DensityTexture3)
-		SHADER_PARAMETER_RDG_TEXTURE(Texture3D, Voxel_DensityTexture4)
-		SHADER_PARAMETER_RDG_TEXTURE(Texture3D, Voxel_DensityTexture5)
-		SHADER_PARAMETER_RDG_TEXTURE(Texture3D, Voxel_DensityTexture6)
-		SHADER_PARAMETER_RDG_TEXTURE(Texture3D, Voxel_DensityTexture7)
-
 		SHADER_PARAMETER_RDG_TEXTURE(Texture2D, RayMarchMaskTexture)
 
 		SHADER_PARAMETER_RDG_TEXTURE(Texture2D, DeepShadow_FrontDepthTexture)
@@ -147,10 +134,6 @@ public:
 			return false;
 		}
 
-		if (PermutationVector.Get<FTransmittanceType>() == FHairTransmittanceType_Voxel && PermutationVector.Get<FTraversal>() == 1)
-		{
-			return false;
-		}
 		return IsHairStrandsSupported(Parameters.Platform);
 	}
 
@@ -161,10 +144,6 @@ public:
 			PermutationVector.Set<FSuperSampling>(0);
 			PermutationVector.Set<FTraversal>(0);
 		}
-		else if (PermutationVector.Get<FTransmittanceType>() == FHairTransmittanceType_Voxel)
-		{
-			PermutationVector.Set<FTraversal>(0);
-		}
 		return PermutationVector;
 	}
 };
@@ -192,11 +171,6 @@ struct FDeepShadowTransmittanceParams
 	FRDGBufferRef HairVisibilityNodeData = nullptr;
 	FRDGBufferRef HairVisibilityNodeCoord = nullptr;
 
-	FRDGTextureRef Voxel_DensityTextures[FHairStrandsDeepShadowData::MaxMacroGroupCount];
-	FVector4 Voxel_MinAABBs[FHairStrandsDeepShadowData::MaxMacroGroupCount];
-	FVector4 Voxel_MaxAABBs[FHairStrandsDeepShadowData::MaxMacroGroupCount];
-	uint32   Voxel_Resolution;
-
 	const FVirtualVoxelResources* VirtualVoxelResources = nullptr;
 };
 
@@ -250,18 +224,6 @@ static FRDGBufferRef AddDeepShadowTransmittanceMaskPass(
 	memcpy(&(Parameters->DeepShadow_AtlasSlotOffsets_AtlasSlotIndex[0]), Params.DeepShadow_AtlasSlotOffsets_AtlasSlotIndex, sizeof(FIntVector4) * FHairStrandsDeepShadowData::MaxMacroGroupCount);
 	memcpy(&(Parameters->DeepShadow_CPUWorldToLightTransforms[0]), Params.DeepShadow_CPUWorldToLightTransforms, sizeof(FMatrix) * FHairStrandsDeepShadowData::MaxMacroGroupCount);
 
-	memcpy(&(Parameters->Voxel_MinAABBs[0]), Params.Voxel_MinAABBs, sizeof(FVector4) * FHairStrandsDeepShadowData::MaxMacroGroupCount);
-	memcpy(&(Parameters->Voxel_MaxAABBs[0]), Params.Voxel_MaxAABBs, sizeof(FVector4) * FHairStrandsDeepShadowData::MaxMacroGroupCount);
-	Parameters->Voxel_Resolution = Params.Voxel_Resolution;
-	Parameters->Voxel_DensityTexture0 = Params.Voxel_DensityTextures[0];
-	Parameters->Voxel_DensityTexture1 = Params.Voxel_DensityTextures[1];
-	Parameters->Voxel_DensityTexture2 = Params.Voxel_DensityTextures[2];
-	Parameters->Voxel_DensityTexture3 = Params.Voxel_DensityTextures[3];
-	Parameters->Voxel_DensityTexture4 = Params.Voxel_DensityTextures[4];
-	Parameters->Voxel_DensityTexture5 = Params.Voxel_DensityTextures[5];
-	Parameters->Voxel_DensityTexture6 = Params.Voxel_DensityTextures[6];
-	Parameters->Voxel_DensityTexture7 = Params.Voxel_DensityTextures[7];
-
 	Parameters->RayMarchMaskTexture = GraphBuilder.RegisterExternalTexture(ScreenShadowMaskSubPixelTexture.IsValid() ? ScreenShadowMaskSubPixelTexture : GSystemTextures.WhiteDummy);
 
 	bool bIsSuperSampled = false;
@@ -302,7 +264,6 @@ static FRDGBufferRef AddDeepShadowTransmittanceMaskPass(
 enum FHairOpaqueMaskType
 {
 	FHairOpaqueMaskType_DeepShadow,
-	FHairOpaqueMaskType_Voxel,
 	FHairOpaqueMaskType_VirtualVoxel,
 	FHairOpaqueMaskTypeCount
 };
@@ -332,16 +293,11 @@ class FDeepShadowMaskPS : public FGlobalShader
 
 		SHADER_PARAMETER(FVector4, Voxel_LightPosition)
 		SHADER_PARAMETER(FVector, Voxel_LightDirection)
-		SHADER_PARAMETER(float, Voxel_DensityScale)
-		SHADER_PARAMETER(FVector, Voxel_MinAABB)
-		SHADER_PARAMETER(uint32, Voxel_Resolution)
-		SHADER_PARAMETER(FVector, Voxel_MaxAABB)
 		SHADER_PARAMETER(uint32, Voxel_MacroGroupId)
 
 		SHADER_PARAMETER_RDG_TEXTURE(Texture2D, RayMarchMaskTexture)
 
 		SHADER_PARAMETER_RDG_BUFFER_SRV(StructuredBuffer, DeepShadow_WorldToLightTransformBuffer)
-		SHADER_PARAMETER_RDG_TEXTURE(Texture3D, Voxel_DensityTexture)
 		SHADER_PARAMETER_RDG_TEXTURE(Texture2D, DeepShadow_FrontDepthTexture)
 		SHADER_PARAMETER_RDG_TEXTURE(Texture2D, DeepShadow_DomTexture)
 		SHADER_PARAMETER_RDG_TEXTURE(Texture2D, CategorizationTexture)
@@ -377,13 +333,7 @@ struct FDeepShadowOpaqueParams
 
 	FVector			Voxel_LightDirection = FVector::ZeroVector;
 	FVector4		Voxel_LightPosition = FVector4(0, 0, 0, 0);
-	float			Voxel_DensityScale = 0;
-	FRDGTextureRef	Voxel_DensityTexture = nullptr;
-	FVector			Voxel_MinAABB;
-	FVector			Voxel_MaxAABB;
-	uint32			Voxel_Resolution;
 	uint32			Voxel_MacroGroupId;
-
 	const FVirtualVoxelResources* Voxel_VirtualVoxel = nullptr;
 };
 
@@ -424,7 +374,7 @@ static void AddDeepShadowOpaqueMaskPass(
 	}
 
 	FRDGTextureRef RayMarchMask = nullptr;
-	if (HairOpaqueMaskType == FHairOpaqueMaskType_VirtualVoxel || HairOpaqueMaskType == FHairOpaqueMaskType_Voxel)
+	if (HairOpaqueMaskType == FHairOpaqueMaskType_VirtualVoxel)
 	{
 		FRDGTextureDesc Desc = OutShadowMask->Desc;
 		Desc.TargetableFlags |= TexCreate_ShaderResource;
@@ -435,18 +385,12 @@ static void AddDeepShadowOpaqueMaskPass(
 	}
 	Parameters->RayMarchMaskTexture = RayMarchMask;
 
-	Parameters->Voxel_LightPosition = Params.Voxel_LightPosition;
-	Parameters->Voxel_LightDirection = Params.Voxel_LightDirection;
-	Parameters->Voxel_DensityScale = Params.Voxel_DensityScale;
-	Parameters->Voxel_MinAABB = Params.Voxel_MinAABB;
-	Parameters->Voxel_Resolution = Params.Voxel_Resolution;
-	Parameters->Voxel_MaxAABB = Params.Voxel_MaxAABB;
-	Parameters->Voxel_MacroGroupId = Params.Voxel_MacroGroupId;
-	Parameters->Voxel_DensityTexture = Params.Voxel_DensityTexture;
-
 	if (HairOpaqueMaskType == FHairOpaqueMaskType_VirtualVoxel)
-	{		
+	{
 		check(Params.Voxel_VirtualVoxel);
+		Parameters->Voxel_LightPosition = Params.Voxel_LightPosition;
+		Parameters->Voxel_LightDirection = Params.Voxel_LightDirection;
+		Parameters->Voxel_MacroGroupId = Params.Voxel_MacroGroupId;
 		Parameters->VirtualVoxel = Params.Voxel_VirtualVoxel->UniformBuffer;
 	}
 
@@ -596,21 +540,8 @@ static FHairStrandsTransmittanceMaskData RenderHairStrandsTransmittanceMask(
 		}
 	}
 
-	if (!bHasFoundLight && IsHairStrandsVoxelizationEnable())
+	if (!bHasFoundLight && MacroGroupDatas.VirtualVoxelResources.IsValid())
 	{
-		Params.Voxel_Resolution = 0;
-		memset(Params.Voxel_MinAABBs, 0, sizeof(FVector4) * FHairStrandsDeepShadowData::MaxMacroGroupCount);
-		memset(Params.Voxel_MaxAABBs, 0, sizeof(FVector4) * FHairStrandsDeepShadowData::MaxMacroGroupCount);
-
-		TRefCountPtr<IPooledRenderTarget> DummyVoxelResources;
-		FPooledRenderTargetDesc Desc(FPooledRenderTargetDesc::CreateVolumeDesc(1, 1, 1, PF_R32_UINT, FClearValueBinding::Black, TexCreate_None, TexCreate_UAV | TexCreate_ShaderResource, false, 1));
-		GRenderTargetPool.FindFreeElement(RHICmdList, Desc, DummyVoxelResources, TEXT("DummyDensityTexture"));
-		FRDGTextureRef DefaultDensityTexture = GraphBuilder.RegisterExternalTexture(DummyVoxelResources, TEXT("Voxel_DefaultDensityTexture"));
-		for (uint32 TexIt = 0; TexIt < FHairStrandsDeepShadowData::MaxMacroGroupCount; ++TexIt)
-		{
-			Params.Voxel_DensityTextures[TexIt] = DefaultDensityTexture;
-		}
-
 		FLightShaderParameters LightParameters;
 		LightSceneInfo->Proxy->GetLightShaderParameters(LightParameters);
 
@@ -622,22 +553,11 @@ static FHairStrandsTransmittanceMaskData RenderHairStrandsTransmittanceMask(
 		Params.LightRadius = FMath::Max(LightParameters.SourceLength, LightParameters.SourceRadius);
 		Params.VirtualVoxelResources = &MacroGroupDatas.VirtualVoxelResources;
 
-		const bool bUseVirtualVoxel = MacroGroupDatas.VirtualVoxelResources.IsValid();
-		const FHairTransmittanceType HairTransmittanceType = bUseVirtualVoxel ? FHairTransmittanceType_VirtualVoxel : FHairTransmittanceType_Voxel;
-		for (const FHairStrandsMacroGroupData& MacroGroupData : MacroGroupDatas.Datas)
-		{
-			Params.Voxel_Resolution = MacroGroupData.GetResolution();
-			Params.Voxel_MinAABBs[MacroGroupData.MacroGroupId] = MacroGroupData.GetMinBound();
-			Params.Voxel_MaxAABBs[MacroGroupData.MacroGroupId] = MacroGroupData.GetMaxBound();
-			if (MacroGroupData.VoxelResources.DensityTexture)
-				Params.Voxel_DensityTextures[MacroGroupData.MacroGroupId] = GraphBuilder.RegisterExternalTexture(MacroGroupData.VoxelResources.DensityTexture);
-		}
-
 		OutShadowMask = AddDeepShadowTransmittanceMaskPass(
 			GraphBuilder,
 			SceneTextures,
 			View,
-			HairTransmittanceType,
+			FHairTransmittanceType_VirtualVoxel,
 			Params,
 			VisibilityData.NodeGroupSize,
 			HairLUTTexture,
@@ -764,7 +684,7 @@ static void RenderHairStrandsShadowMask(
 
 	// Code is disabled for now until we have the full DOM/voxel fallback logic
 	// If there is no deep shadow for this light, fallback on the voxel representation
-	if (!bHasDeepShadow && IsHairStrandsVoxelizationEnable())
+	if (!bHasDeepShadow && InMacroGroupDatas.VirtualVoxelResources.IsValid())
 	{
 		// TODO: Change this to be a single pass with virtual voxel?
 		for (const FHairStrandsMacroGroupData& MacroGroupData : InMacroGroupDatas.Datas)
@@ -773,23 +693,16 @@ static void RenderHairStrandsShadowMask(
 
 			FDeepShadowOpaqueParams Params;
 			Params.CategorizationTexture = Categorization;
-			Params.Voxel_Resolution = MacroGroupData.GetResolution();
-			Params.Voxel_MinAABB = MacroGroupData.GetMinBound();
-			Params.Voxel_MaxAABB = MacroGroupData.GetMaxBound();
-			Params.Voxel_DensityTexture = GraphBuilder.RegisterExternalTexture(MacroGroupData.VoxelResources.DensityTexture ? MacroGroupData.VoxelResources.DensityTexture : GSystemTextures.WhiteDummy, TEXT("Voxel_DensityTexture"));
-			Params.Voxel_DensityScale = GetDeepShadowDensityScale();
 			Params.Voxel_LightDirection = LightSceneInfo->Proxy->GetDirection();
 			Params.Voxel_LightPosition = FVector4(LightSceneInfo->Proxy->GetPosition(), LightSceneInfo->Proxy->GetLightType() == ELightComponentType::LightType_Directional ? 0 : 1);
 			Params.Voxel_MacroGroupId = MacroGroupData.MacroGroupId;
 
-			const bool bUseVirtualVoxel = InMacroGroupDatas.VirtualVoxelResources.IsValid();
-			const FHairOpaqueMaskType HairOpaqueMaskType = bUseVirtualVoxel ? FHairOpaqueMaskType_VirtualVoxel : FHairOpaqueMaskType_Voxel;
-			Params.Voxel_VirtualVoxel = bUseVirtualVoxel ? &InMacroGroupDatas.VirtualVoxelResources : nullptr;
+			Params.Voxel_VirtualVoxel = &InMacroGroupDatas.VirtualVoxelResources;
 			AddDeepShadowOpaqueMaskPass(
 				GraphBuilder,
 				SceneTextures,
 				View,
-				HairOpaqueMaskType,
+				FHairOpaqueMaskType_VirtualVoxel,
 				Params,
 				OutShadowMask);
 		}
diff --git a/Engine/Source/Runtime/Renderer/Private/HairStrands/HairStrandsVoxelization.cpp b/Engine/Source/Runtime/Renderer/Private/HairStrands/HairStrandsVoxelization.cpp
--- a/Engine/Source/Runtime/Renderer/Private/HairStrands/HairStrandsVoxelization.cpp
+++ b/Engine/Source/Runtime/Renderer/Private/HairStrands/HairStrandsVoxelization.cpp
@@ -27,7 +27,7 @@ static FAutoConsoleVariableRef CVarHairVoxelizationMaterialEnable(TEXT("r.HairSt
 static float GHairVoxelizationAABBScale = 1.0f;
 static FAutoConsoleVariableRef CVarHairVoxelizationAABBScale(TEXT("r.HairStrands.Voxelization.AABBScale"), GHairVoxelizationAABBScale, TEXT("Scale the hair macro group bounding box"));
 
-static float GHairVoxelizationDensityScale = 1.0f;
+static float GHairVoxelizationDensityScale = 2.0f;
 static float GHairVoxelizationDepthBiasScale = 3.0f;
 static FAutoConsoleVariableRef CVarHairVoxelizationDensityScale(TEXT("r.HairStrands.Voxelization.DensityScale"), GHairVoxelizationDensityScale, TEXT("Scale the hair density when computing voxel transmittance. Default value is 2 (arbitraty)"));
 static FAutoConsoleVariableRef CVarHairVoxelizationDepthBiasScale(TEXT("r.HairStrands.Voxelization.DepthBiasScale"), GHairVoxelizationDepthBiasScale, TEXT("Set depth bias for voxel ray marching. Offset the origin position towards the light"));
