Remove solid voxelization for hair (overall code).

diff --git a/Engine/Shaders/Private/HairStrands/HairStrandsDeepShadowPS.usf b/Engine/Shaders/Private/HairStrands/HairStrandsDeepShadowPS.usf
--- a/Engine/Shaders/Private/HairStrands/HairStrandsDeepShadowPS.usf
+++ b/Engine/Shaders/Private/HairStrands/HairStrandsDeepShadowPS.usf
@@ -71,7 +71,6 @@ void MainVoxel(
 	const float HairCoverage = WorldPositon_Coverage.w;
 
 #if VF_STRAND_HAIR  
-  #if SUPPORT_SPARSE	
 	const FPackedVirtualVoxelNodeDesc PackedNode = VoxelRasterPass.VirtualVoxel.NodeDescBuffer[VoxelRasterPass.MacroGroupId];
 	const FVirtualVoxelNodeDesc NodeDesc = UnpackVoxelNode(PackedNode, VoxelRasterPass.VirtualVoxel.PageResolution);
 
@@ -112,7 +111,8 @@ void MainVoxel(
 			}
 		}
 	}
-  #else
+
+	#if 0 // This code is no longer used, but it is a reference for voxelizing material
 	const int3 VoxelCoord = (WorldPosition.xyz - VoxelRasterPass.VoxelMinAABB) / (VoxelRasterPass.VoxelMaxAABB - VoxelRasterPass.VoxelMinAABB) * VoxelRasterPass.VoxelResolution;
 		 
 	if (VoxelCoord.x >= 0 || VoxelCoord.y >= 0 || VoxelCoord.z >= 0 ||
@@ -157,6 +157,7 @@ void MainVoxel(
 	#if DEBUG_OUTPUT
 	OutColor = 0xFF;
 	#endif
-  #endif
+	#endif
+
 #endif
 }
\ No newline at end of file
diff --git a/Engine/Shaders/Private/HairStrands/HairStrandsDeepTransmittanceMask.usf b/Engine/Shaders/Private/HairStrands/HairStrandsDeepTransmittanceMask.usf
--- a/Engine/Shaders/Private/HairStrands/HairStrandsDeepTransmittanceMask.usf
+++ b/Engine/Shaders/Private/HairStrands/HairStrandsDeepTransmittanceMask.usf
@@ -56,8 +56,8 @@ float3				LightDirection;
 uint				LightChannelMask;
 uint				MaxVisibilityNodeCount;
 float4				LightPosition;
-float				DepthBiasScale;
-float				DensityScale;
+float				DeepShadow_DepthBiasScale;
+float				DeepShadow_DensityScale;
 float				LightRadius;
 float				DeepShadow_KernelAperture;
 uint				DeepShadow_KernelType;
@@ -136,8 +136,8 @@ void MainCS(uint2 DispatchThreadId : SV_DispatchThreadID)
 		Settings.WorldToLightTransform = WorldToLightTransform;
 		Settings.LightDirection = LightDirection;
 		Settings.LightPosition = LightPosition;
-		Settings.DeepShadowDepthBiasScale = DepthBiasScale;
-		Settings.DeepShadowDensityScale = DensityScale;
+		Settings.DeepShadowDepthBiasScale = DeepShadow_DepthBiasScale;
+		Settings.DeepShadowDensityScale = DeepShadow_DensityScale;
 		Settings.TransmittanceKernelApertureInDegree = DeepShadow_KernelAperture;
 		Settings.TransmittanceKernelType = DeepShadow_KernelType;
 		Settings.Random = TransmittanceRandom;
@@ -238,8 +238,8 @@ FHairTransmittanceMask ComputeTransmittanceVoxel(
 		#else
 		const float3 SampleRandom = GetHairVoxelJitter(PixelCoord, View.StateFrameIndexMod8);
 		#endif
-
-		const float3 DepthBias = MaxVoxelSize * (Settings.DeepShadowDepthBiasScale*LightDirection + SampleRandom);
+		
+		const float3 DepthBias = MaxVoxelSize * (VirtualVoxel.DepthBiasScale*LightDirection + SampleRandom);
 		const float3 SampleWorldPosition = WorldPosition + DepthBias;
 
 		FVirtualVoxelCommonDesc CommonDesc;
@@ -332,8 +332,8 @@ void MainCS(uint2 DispatchThreadId : SV_DispatchThreadID)
 		FTransmittanceSettings Settings;
 		Settings.LightDirection = LightDirection;
 		Settings.LightPosition = LightPosition;
-		Settings.DeepShadowDensityScale = DensityScale;
-		Settings.DeepShadowDepthBiasScale = DepthBiasScale;
+		Settings.DeepShadowDensityScale = 0;	// Use virtual voxel struct value
+		Settings.DeepShadowDepthBiasScale = 0;	// Use virtual voxel struct value
 		Settings.Random = GetHairVoxelJitter(PixelCoord.xy, View.StateFrameIndexMod8);
 
 		const float SceneDepth = ConvertFromDeviceZ(Sample.Depth);
diff --git a/Engine/Shaders/Private/HairStrands/HairStrandsVoxelMip.usf b/Engine/Shaders/Private/HairStrands/HairStrandsVoxelMip.usf
--- a/Engine/Shaders/Private/HairStrands/HairStrandsVoxelMip.usf
+++ b/Engine/Shaders/Private/HairStrands/HairStrandsVoxelMip.usf
@@ -38,126 +38,6 @@ uint ComputeMipDensity(
 	return min(VOXEL_HAIR_MASK, TotalHair) | OutTotalOpaque;
 }
 
-#if SHADER_MIP_VOXEL
-
-uint VoxelResolution;
-uint SourceMip;
-uint TargetMip;
-
-Texture3D<uint> InDensityTexture;
-RWTexture3D<uint> OutDensityTexture0;
-
-
-// Write two levels at a time, and store/cache thread result within LDS
-#if PERMUTATION_METHOD == 1
-RWTexture3D<uint> OutDensityTexture1;
-groupshared uint Values0[64];
-groupshared uint Values1[8];
-
-uint ToMorton(uint3 Index)
-{
-	return (Index.x & 0x1) | ((Index.y & 0x1)<<1) | ((Index.z& 0x1)<<2);
-}
-
-uint3 FromMorton(uint Index)
-{
-	return uint3((Index & 0x1) > 0 ? 1 : 0, (Index & 0x2) > 0 ? 1 : 0, (Index & 0x4) > 0 ? 1 : 0);
-}
-
-[numthreads(4, 4, 4)]
-void MainCS(uint3 DispatchThreadId : SV_DispatchThreadID, uint3 GroupThreadId : SV_GroupThreadID , uint GroupIndex : SV_GroupIndex, uint3 GroupId : SV_GroupId)
-{
-	{
-		const uint3 LocalBase	= uint3(GroupThreadId.x >> 1, GroupThreadId.y >> 1, GroupThreadId.z>> 1);
-		const uint3 LocalOffset = uint3(GroupThreadId.x % 2,  GroupThreadId.y % 2,  GroupThreadId.z % 2);
-
-		const uint Base   = ToMorton(LocalBase);
-		const uint Offset = ToMorton(LocalOffset);
-		const uint OutLocal = Base*8 + Offset;
-		Values0[OutLocal] = InDensityTexture[DispatchThreadId];
-	}
-
-	GroupMemoryBarrierWithGroupSync();
-
-	if (GroupIndex < 8)
-	{
-		const uint3 OutLocalCoord = FromMorton(GroupIndex);
-		const uint3 OutVoxelCoord = (GroupId << 1) + OutLocalCoord;
-
-		const uint Offset = GroupIndex * 8;
-		const uint Total = ComputeMipDensity(
-			Values0[Offset + 0],
-			Values0[Offset + 1],
-			Values0[Offset + 2],
-			Values0[Offset + 3],
-			Values0[Offset + 4],
-			Values0[Offset + 5],
-			Values0[Offset + 6],
-			Values0[Offset + 7]);
-		Values1[GroupIndex] = Total;
-		OutDensityTexture0[OutVoxelCoord] = Total;
-	}
-
-	GroupMemoryBarrierWithGroupSync();
-
-	if (GroupIndex == 0)
-	{
-		const uint3 OutVoxelCoord = GroupId;
-		const uint Total = ComputeMipDensity(
-			Values1[0],
-			Values1[1],
-			Values1[2],
-			Values1[3],
-			Values1[4],
-			Values1[5],
-			Values1[6],
-			Values1[7]);
-		OutDensityTexture1[OutVoxelCoord] = Total;
-	}
-}
-#endif
-
-#if PERMUTATION_METHOD == 0
-[numthreads(4, 4, 4)]
-void MainCS(uint3 DispatchThreadId : SV_DispatchThreadID)
-{
-	const uint3 OutVoxelCoord = DispatchThreadId;
-
-	const uint3 InVoxelCoord0 = uint3(OutVoxelCoord.x << 1, OutVoxelCoord.y << 1, OutVoxelCoord.z << 1);
-	const uint3 InVoxelCoord1 = InVoxelCoord0 + uint3(1,0,0);
-	const uint3 InVoxelCoord2 = InVoxelCoord0 + uint3(0,1,0);
-	const uint3 InVoxelCoord3 = InVoxelCoord0 + uint3(1,1,0);
-	const uint3 InVoxelCoord4 = InVoxelCoord0 + uint3(0,0,1);
-	const uint3 InVoxelCoord5 = InVoxelCoord0 + uint3(1,0,1);
-	const uint3 InVoxelCoord6 = InVoxelCoord0 + uint3(0,1,1);
-	const uint3 InVoxelCoord7 = InVoxelCoord0 + uint3(1,1,1);
-
-
-	const uint RawDensity0 = InDensityTexture[InVoxelCoord0];
-	const uint RawDensity1 = InDensityTexture[InVoxelCoord1];
-	const uint RawDensity2 = InDensityTexture[InVoxelCoord2];
-	const uint RawDensity3 = InDensityTexture[InVoxelCoord3];
-	const uint RawDensity4 = InDensityTexture[InVoxelCoord4];
-	const uint RawDensity5 = InDensityTexture[InVoxelCoord5];
-	const uint RawDensity6 = InDensityTexture[InVoxelCoord6];
-	const uint RawDensity7 = InDensityTexture[InVoxelCoord7];
-
-	const uint Total = ComputeMipDensity(
-		RawDensity0,
-		RawDensity1,
-		RawDensity2,
-		RawDensity3,
-		RawDensity4,
-		RawDensity5,
-		RawDensity6,
-		RawDensity7);
-
-	OutDensityTexture0[OutVoxelCoord] = Total;
-}
-#endif
-
-#endif
-
 #if SHADER_MIP_VIRTUALVOXEL
 int3 PageCountResolution;
 uint PageResolution;
@@ -245,7 +125,6 @@ void MainCS(uint3 DispatchThreadId : SV_DispatchThreadID)
 
 #endif
 
-
 #if SHADER_UPDATE_PAGEINDEX
 int3 PageCountResolution;
 uint PageResolution;
diff --git a/Engine/Shaders/Private/HairStrands/HairStrandsVoxelOpaque.usf b/Engine/Shaders/Private/HairStrands/HairStrandsVoxelOpaque.usf
--- a/Engine/Shaders/Private/HairStrands/HairStrandsVoxelOpaque.usf
+++ b/Engine/Shaders/Private/HairStrands/HairStrandsVoxelOpaque.usf
@@ -6,79 +6,6 @@
 #include "HairStrandsVoxelPageCommonStruct.ush"
 #include "HairStrandsVoxelPageCommon.ush"
 
-#if SHADER_INJECTOPAQUE_VOXEL
-float3 VoxelMinAABB;
-uint   VoxelResolution;
-float3 VoxelMaxAABB;
-uint   VoxelBiasCount;
-float2 OutputResolution;
-float2 SceneDepthResolution;
-uint   VoxelMarkCount;
-
-Texture2D<float> SceneDepthTexture;
-RWTexture3D<uint> DensityTexture;
-SamplerState LinearSampler;
-
-float3 ReconstructWorldPositionFromDepth(float2 UV, float Depth)
-{
-	float2 ScreenPosition = (UV - View.ScreenPositionScaleBias.wz) / View.ScreenPositionScaleBias.xy;
-	float4 HomogeneousWorldPosition = mul(float4(ScreenPosition * Depth, Depth, 1), View.ScreenToWorld);
-	float3 WorldPosition = HomogeneousWorldPosition.xyz / HomogeneousWorldPosition.w;
-
-	return WorldPosition;
-}
-
-// hair_todo: change to CS & optmize (naive code)
-void MainPS(
-	FScreenVertexOutput Input,
-	out float4  OutColor0 : SV_Target0)
-{
-	const float2 UV = Input.UV;
-	const float2 PixelCoord = UV * OutputResolution;
-
-	const float VoxelSize = (VoxelMaxAABB.x - VoxelMinAABB.x) / VoxelResolution;
-
-	uint3 VoxelPos = 0;
-	uint2 TileCount = OutputResolution / VoxelResolution;
-	VoxelPos.x = PixelCoord.x % VoxelResolution;
-	VoxelPos.y = PixelCoord.y % VoxelResolution;
-	VoxelPos.z = PixelCoord.x / VoxelResolution + PixelCoord.y / VoxelResolution * TileCount.x;
-
-	const float3 WorldPosition = VoxelPos / float(VoxelResolution) * (VoxelMaxAABB - VoxelMinAABB) + VoxelMinAABB;
-
-	float4 ClipPos = mul(float4(WorldPosition,1), View.WorldToClip);
-	ClipPos /= ClipPos.w;
-	const float DepthBias 			= VoxelBiasCount * VoxelSize;
-	const float  VoxelDepth 		= ConvertFromDeviceZ(ClipPos.z) - DepthBias;
-	float2 SceneUV 					= float2(0.5f*(ClipPos.x+1), 1-0.5f*(ClipPos.y+1));
-	SceneUV							= ViewportUVToBufferUV(SceneUV);
-	const float2 ScenePixelCoord 	= SceneUV * SceneDepthResolution;
-
-	OutColor0 = float4(VoxelPos.xy / float(VoxelResolution), 0, 1); 
-	const bool bIsOnScreen = SceneUV.x >= 0 && SceneUV.x < 1 && SceneUV.y >= 0 && SceneUV.y < 1;
-	if (!bIsOnScreen)
-		return;
-
-	const float ClosestDepth = ConvertFromDeviceZ(SceneDepthTexture.Load(uint3(ScenePixelCoord,0)));
-	const float3 SceneWorldPos = ReconstructWorldPositionFromDepth(SceneUV, ClosestDepth);
-
-	const bool bIsInVolume = 
-		SceneWorldPos.x >= VoxelMinAABB.x && SceneWorldPos.x < VoxelMaxAABB.x && 
-		SceneWorldPos.y >= VoxelMinAABB.y && SceneWorldPos.y < VoxelMaxAABB.y && 
-		SceneWorldPos.z >= VoxelMinAABB.z && SceneWorldPos.z < VoxelMaxAABB.z;
-
-	if (!bIsInVolume)
-		return;
-
-	// Inject opaque depth on a thin layer (Dist < DistThreshold) for avoiding weird projection
-	if (ClosestDepth < VoxelDepth && abs(ClosestDepth - VoxelDepth) < VoxelMarkCount * VoxelSize)
-	{
-		OutColor0 = float4(1,0,0,1);
-		InterlockedOr(DensityTexture[VoxelPos], VOXEL_OPAQUE_ADD);
-	}
-}
-#endif
-
 #if SHADER_INJECTOPAQUE_VIRTUALVOXEL
 uint3 DispatchedPageIndexResolution;
 uint MacroGroupId;
@@ -170,42 +97,4 @@ void MainCS(uint3 DispatchThreadId : SV_DispatchThreadID)
 		}
 	}
 }
-#endif
-
-
-#if SHADER_FILTERDEPTH_VOXEL
-
-uint				VoxelResolution;
-RWTexture3D<uint>	VoxelTexture;
-
-[numthreads(4, 4, 4)]
-void MainCS(uint3 DispatchThreadId : SV_DispatchThreadID)
-{
-	const uint3 VoxelCoord  = DispatchThreadId;
-	const uint  RawDensity	= VoxelTexture[VoxelCoord];
-	const bool  bHasOpaque	= GetInternalVoxelOpaqueVisibility(RawDensity) > 0;
-
-	// If the current voxel does not contains any opaque data, there is nothing to filter.
-	if (!bHasOpaque)
-	{
-		return;
-	}
-
-	const int FilterSize = 3;
-	for (int z = -FilterSize; z <= FilterSize; z++)
-	for (int y = -FilterSize; y <= FilterSize; y++)
-	for (int x = -FilterSize; x <= FilterSize; x++)
-	{
-		const int3 Offset = int3(x, y, z);
-		const int3 Sample = clamp(VoxelCoord + Offset, 0, VoxelResolution-1);
-		const uint Density = VoxelTexture[Sample] & VOXEL_HAIR_MASK;
-		if (Density > 0)
-		{
-			return;
-		}
-	}
-
-	// No hair density has been found, so this opaque voxel is a false positive
-	InterlockedAnd(VoxelTexture[VoxelCoord], VOXEL_HAIR_MASK);
-}
 #endif
\ No newline at end of file
diff --git a/Engine/Shaders/Private/HairStrands/HairStrandsVoxelRayMarching.usf b/Engine/Shaders/Private/HairStrands/HairStrandsVoxelRayMarching.usf
deleted file mode 100644
diff --git a/Engine/Source/Runtime/Renderer/Private/HairStrands/HairStrandsCluster.h b/Engine/Source/Runtime/Renderer/Private/HairStrands/HairStrandsCluster.h
--- a/Engine/Source/Runtime/Renderer/Private/HairStrands/HairStrandsCluster.h
+++ b/Engine/Source/Runtime/Renderer/Private/HairStrands/HairStrandsCluster.h
@@ -53,12 +53,7 @@ struct FHairStrandsMacroGroupData
 	typedef TArray<PrimitiveGroup, SceneRenderingAllocator> TPrimitiveGroups;
 
 
-	FVector GetMinBound() const { return VoxelResources.MinAABB; }
-	FVector GetMaxBound() const { return VoxelResources.MaxAABB; }
-	uint32  GetResolution() const { return VoxelResources.DensityTexture ? VoxelResources.DensityTexture->GetDesc().Extent.X : 0; }
-
 	FVirtualVoxelNodeDesc VirtualVoxelNodeDesc;
-	FHairStrandsVoxelResources VoxelResources;
 	FHairStrandsDeepShadowDatas DeepShadowDatas;
 	TPrimitiveInfos PrimitivesInfos;
 	TPrimitiveGroups PrimitivesGroups;
diff --git a/Engine/Source/Runtime/Renderer/Private/HairStrands/HairStrandsDebug.cpp b/Engine/Source/Runtime/Renderer/Private/HairStrands/HairStrandsDebug.cpp
--- a/Engine/Source/Runtime/Renderer/Private/HairStrands/HairStrandsDebug.cpp
+++ b/Engine/Source/Runtime/Renderer/Private/HairStrands/HairStrandsDebug.cpp
@@ -643,147 +643,6 @@ static void AddDeepShadowInfoPass(
 	FComputeShaderUtils::AddPass(GraphBuilder, RDG_EVENT_NAME("HairStrandsDeepShadowDebugInfo"), ComputeShader, Parameters, DispatchCount);
 }
 
-///////////////////////////////////////////////////////////////////////////////////////////////////
-class FVoxelRaymarchingPS : public FGlobalShader
-{
-	DECLARE_GLOBAL_SHADER(FVoxelRaymarchingPS);
-	SHADER_USE_PARAMETER_STRUCT(FVoxelRaymarchingPS, FGlobalShader);
-
-	class FDebugMode : SHADER_PERMUTATION_INT("PERMUTATION_DEBUG_MODE", 4);
-	using FPermutationDomain = TShaderPermutationDomain<FDebugMode>;
-
-	BEGIN_SHADER_PARAMETER_STRUCT(FParameters, )
-		SHADER_PARAMETER_STRUCT_INCLUDE(FSceneTextureParameters, SceneTextures)
-
-		SHADER_PARAMETER(FVector, VoxelMinAABB)
-		SHADER_PARAMETER(uint32, VoxelResolution)
-		SHADER_PARAMETER(FVector, VoxelMaxAABB)
-		SHADER_PARAMETER(float, DensityIsoline)
-		SHADER_PARAMETER(float, VoxelDensityScale)
-		SHADER_PARAMETER(FVector2D, OutputResolution)
-
-		SHADER_PARAMETER_RDG_TEXTURE(Texture3D, DensityTexture)
-		SHADER_PARAMETER_RDG_TEXTURE(Texture3D, TangentXTexture)
-		SHADER_PARAMETER_RDG_TEXTURE(Texture3D, TangentYTexture)
-		SHADER_PARAMETER_RDG_TEXTURE(Texture3D, TangentZTexture)
-		SHADER_PARAMETER_RDG_TEXTURE(Texture3D, MaterialTexture)
-		SHADER_PARAMETER_SAMPLER(SamplerState, LinearSampler)
-
-		SHADER_PARAMETER_STRUCT_REF(FViewUniformShaderParameters, ViewUniformBuffer)
-		RENDER_TARGET_BINDING_SLOTS()
-		END_SHADER_PARAMETER_STRUCT()
-
-public:
-	static bool ShouldCompilePermutation(const FGlobalShaderPermutationParameters& Parameters) { return IsHairStrandsSupported(Parameters.Platform); }
-};
-
-IMPLEMENT_GLOBAL_SHADER(FVoxelRaymarchingPS, "/Engine/Private/HairStrands/HairStrandsVoxelRayMarching.usf", "MainPS", SF_Pixel);
-
-static void AddVoxelRaymarchingPass(
-	FRDGBuilder& GraphBuilder,
-	const FViewInfo& View,
-	const EHairDebugMode DebugMode,
-	const FHairStrandsMacroGroupDatas& MacroGroupDatas,
-	FRDGTextureRef& OutputTexture)
-{
-	check(DebugMode == EHairDebugMode::VoxelsDensity || DebugMode == EHairDebugMode::VoxelsTangent || DebugMode == EHairDebugMode::VoxelsBaseColor || DebugMode == EHairDebugMode::VoxelsRoughness);
-
-	FSceneTextureParameters SceneTextures;
-	SetupSceneTextureParameters(GraphBuilder, &SceneTextures);
-
-	const FIntPoint Resolution(OutputTexture->Desc.Extent);
-	for (const FHairStrandsMacroGroupData& MacroGroupData : MacroGroupDatas.Datas)
-	{
-		if (DebugMode == EHairDebugMode::VoxelsDensity && !MacroGroupData.VoxelResources.DensityTexture)
-			return;
-
-		if (DebugMode == EHairDebugMode::VoxelsTangent && (!MacroGroupData.VoxelResources.TangentXTexture || !MacroGroupData.VoxelResources.TangentYTexture || !MacroGroupData.VoxelResources.TangentZTexture))
-			return;
-
-		if ((DebugMode == EHairDebugMode::VoxelsBaseColor || DebugMode == EHairDebugMode::VoxelsRoughness) && !MacroGroupData.VoxelResources.MaterialTexture)
-			return;
-
-		const FRDGTextureRef VoxelDensityTexture  = GraphBuilder.RegisterExternalTexture(MacroGroupData.VoxelResources.DensityTexture  ? MacroGroupData.VoxelResources.DensityTexture : GSystemTextures.BlackDummy, TEXT("HairVoxelDensityTexture"));
-		const FRDGTextureRef VoxelTangentXTexture = GraphBuilder.RegisterExternalTexture(MacroGroupData.VoxelResources.TangentXTexture ? MacroGroupData.VoxelResources.TangentXTexture : GSystemTextures.BlackDummy, TEXT("HairVoxelTangentXTexture"));
-		const FRDGTextureRef VoxelTangentYTexture = GraphBuilder.RegisterExternalTexture(MacroGroupData.VoxelResources.TangentYTexture ? MacroGroupData.VoxelResources.TangentYTexture : GSystemTextures.BlackDummy, TEXT("HairVoxelTangentYTexture"));
-		const FRDGTextureRef VoxelTangentZTexture = GraphBuilder.RegisterExternalTexture(MacroGroupData.VoxelResources.TangentZTexture ? MacroGroupData.VoxelResources.TangentZTexture : GSystemTextures.BlackDummy, TEXT("HairVoxelTangentZTexture"));
-		const FRDGTextureRef VoxelMaterialTexture = GraphBuilder.RegisterExternalTexture(MacroGroupData.VoxelResources.MaterialTexture ? MacroGroupData.VoxelResources.MaterialTexture : GSystemTextures.BlackDummy, TEXT("HairVoxelMaterialTexture"));
-
-		FVoxelRaymarchingPS::FParameters* Parameters = GraphBuilder.AllocParameters<FVoxelRaymarchingPS::FParameters>();
-		Parameters->ViewUniformBuffer = View.ViewUniformBuffer;
-		Parameters->OutputResolution = Resolution;
-		Parameters->SceneTextures = SceneTextures;
-		Parameters->DensityTexture = VoxelDensityTexture;
-		Parameters->TangentXTexture = VoxelTangentXTexture;
-		Parameters->TangentYTexture = VoxelTangentYTexture;
-		Parameters->TangentZTexture = VoxelTangentZTexture;
-		Parameters->MaterialTexture = VoxelMaterialTexture;
-		Parameters->VoxelMinAABB = MacroGroupData.GetMinBound();
-		Parameters->VoxelMaxAABB = MacroGroupData.GetMaxBound();
-		Parameters->VoxelResolution = MacroGroupData.GetResolution();
-		Parameters->VoxelDensityScale = GetHairStrandsVoxelizationDensityScale();
-		Parameters->DensityIsoline = 1;
-		Parameters->LinearSampler = TStaticSamplerState<SF_Trilinear, AM_Clamp, AM_Clamp, AM_Clamp>::GetRHI();
-		Parameters->RenderTargets[0] = FRenderTargetBinding(OutputTexture, ERenderTargetLoadAction::ELoad);
-
-		const FIntPoint OutputResolution = SceneTextures.SceneDepthBuffer->Desc.Extent;
-		TShaderMapRef<FPostProcessVS> VertexShader(View.ShaderMap);
-
-		FVoxelRaymarchingPS::FPermutationDomain PermutationVector;
-		uint32 DebugPermutation = 0;
-		switch (DebugMode)
-		{
-		case EHairDebugMode::VoxelsDensity:		DebugPermutation = 0; break;
-		case EHairDebugMode::VoxelsTangent:		DebugPermutation = 1; break;
-		case EHairDebugMode::VoxelsBaseColor:	DebugPermutation = 2; break;
-		case EHairDebugMode::VoxelsRoughness:	DebugPermutation = 3; break;
-		};
-		PermutationVector.Set<FVoxelRaymarchingPS::FDebugMode>(DebugPermutation);
-
-		TShaderMapRef<FVoxelRaymarchingPS> PixelShader(View.ShaderMap, PermutationVector);
-		const FGlobalShaderMap* GlobalShaderMap = View.ShaderMap;
-		const FIntRect Viewport = View.ViewRect;
-		const FViewInfo* CapturedView = &View;
-
-		ClearUnusedGraphResources(PixelShader, Parameters);
-
-		GraphBuilder.AddPass(
-			RDG_EVENT_NAME("HairStrandsVoxelRaymarching"),
-			Parameters,
-			ERDGPassFlags::Raster,
-			[Parameters, VertexShader, PixelShader, Viewport, Resolution, CapturedView](FRHICommandList& RHICmdList)
-		{
-			FGraphicsPipelineStateInitializer GraphicsPSOInit;
-			RHICmdList.ApplyCachedRenderTargets(GraphicsPSOInit);
-			GraphicsPSOInit.BlendState = TStaticBlendState<CW_RGBA, BO_Add, BF_One, BF_Zero, BO_Add, BF_One, BF_Zero>::GetRHI();
-			GraphicsPSOInit.RasterizerState = TStaticRasterizerState<>::GetRHI();
-			GraphicsPSOInit.DepthStencilState = TStaticDepthStencilState<false, CF_Always>::GetRHI();
-
-			GraphicsPSOInit.BoundShaderState.VertexDeclarationRHI = GFilterVertexDeclaration.VertexDeclarationRHI;
-			GraphicsPSOInit.BoundShaderState.VertexShaderRHI = VertexShader.GetVertexShader();
-			GraphicsPSOInit.BoundShaderState.PixelShaderRHI = PixelShader.GetPixelShader();
-			GraphicsPSOInit.PrimitiveType = PT_TriangleList;
-			SetGraphicsPipelineState(RHICmdList, GraphicsPSOInit);
-
-			VertexShader->SetParameters(RHICmdList, CapturedView->ViewUniformBuffer);
-			RHICmdList.SetViewport(Viewport.Min.X, Viewport.Min.Y, 0.0f, Viewport.Max.X, Viewport.Max.Y, 1.0f);
-			SetShaderParameters(RHICmdList, PixelShader, PixelShader.GetPixelShader(), *Parameters);
-
-			DrawRectangle(
-				RHICmdList,
-				0, 0,
-				Viewport.Width(), Viewport.Height(),
-				Viewport.Min.X, Viewport.Min.Y,
-				Viewport.Width(), Viewport.Height(),
-				Viewport.Size(),
-				Resolution,
-				VertexShader,
-				EDRF_UseTriangleOptimization);
-		});
-	}
-}
-
-	
 ///////////////////////////////////////////////////////////////////////////////////////////////////
 class FVoxelVirtualRaymarchingCS : public FGlobalShader
 {
@@ -1519,19 +1378,13 @@ void RenderHairStrandsDebugInfo(
 		// CPU bound of macro groups
 		FViewElementPDI ShadowFrustumPDI(&View, nullptr, nullptr);
 		const FHairStrandsMacroGroupDatas& MacroGroupDatas = InMacroGroupViews.Views[ViewIndex];
-		const bool bUseVirtualVoxel = MacroGroupDatas.VirtualVoxelResources.IsValid();
-		for (const FHairStrandsMacroGroupData& MacroGroupData : MacroGroupDatas.Datas)
+		if (MacroGroupDatas.VirtualVoxelResources.IsValid())
 		{
-			if (bUseVirtualVoxel)
+			for (const FHairStrandsMacroGroupData& MacroGroupData : MacroGroupDatas.Datas)
 			{
 				const FBox Bound(MacroGroupData.VirtualVoxelNodeDesc.WorldMinAABB, MacroGroupData.VirtualVoxelNodeDesc.WorldMaxAABB);
 				DrawWireBox(&ShadowFrustumPDI, Bound, FColor::Red, 0);
 			}
-			else
-			{
-				const FBox Bound(MacroGroupData.GetMinBound(), MacroGroupData.GetMaxBound());
-				DrawWireBox(&ShadowFrustumPDI, Bound, FColor::Red, 0);
-			}
 		}
 
 		FRenderTargetTemp TempRenderTarget(View, (const FTexture2DRHIRef&)SceneTargets.GetSceneColor()->GetRenderTargetItem().TargetableTexture);
@@ -1667,21 +1520,14 @@ void RenderHairStrandsDebugInfo(
 		if (bIsVoxelMode && ViewIndex < uint32(InMacroGroupViews.Views.Num()))
 		{
 			const FHairStrandsMacroGroupDatas& MacroGroupDatas = InMacroGroupViews.Views[ViewIndex];
-			const bool bUseVirtualVoxel = MacroGroupDatas.VirtualVoxelResources.IsValid();
-			for (const FHairStrandsMacroGroupData& MacroGroupData : MacroGroupDatas.Datas)
+			if (MacroGroupDatas.VirtualVoxelResources.IsValid())
 			{
-				if (bUseVirtualVoxel)
+				for (const FHairStrandsMacroGroupData& MacroGroupData : MacroGroupDatas.Datas)
 				{
 					const FBox Bound(MacroGroupData.VirtualVoxelNodeDesc.WorldMinAABB, MacroGroupData.VirtualVoxelNodeDesc.WorldMaxAABB);
 					DrawWireBox(&ShadowFrustumPDI, Bound, FColor::Red, 0);
 					DrawFrustumWireframe(&ShadowFrustumPDI, MacroGroupData.VirtualVoxelNodeDesc.WorldToClip.Inverse(), FColor::Purple, 0);
 				}
-				else
-				{
-					const FBox VoxelizationBox(MacroGroupData.GetMinBound(), MacroGroupData.GetMaxBound());
-					DrawWireBox(&ShadowFrustumPDI, VoxelizationBox, FColor::Red, 0);
-					DrawFrustumWireframe(&ShadowFrustumPDI, MacroGroupData.VoxelResources.WorldToClip.Inverse(), FColor::Purple, 0);
-				}
 			}
 		}
 	}
@@ -1717,21 +1563,13 @@ void RenderHairStrandsDebugInfo(
 		if (ViewIndex < uint32(InMacroGroupViews.Views.Num()))
 		{
 			const FHairStrandsMacroGroupDatas& MacroGroupDatas = InMacroGroupViews.Views[ViewIndex];
-			const bool bIsVirtualVoxelizationEnabled = MacroGroupDatas.VirtualVoxelResources.IsValid();
-			if (bIsVirtualVoxelizationEnabled)
+			if (MacroGroupDatas.VirtualVoxelResources.IsValid())
 			{
 				FRDGBuilder GraphBuilder(RHICmdList);
 				FRDGTextureRef SceneColorTexture = GraphBuilder.RegisterExternalTexture(SceneTargets.GetSceneColor(), TEXT("SceneColorTexture"));
 				AddVoxelPageRaymarchingPass(GraphBuilder, View, MacroGroupDatas, SceneColorTexture);
 				GraphBuilder.Execute();
 			}
-			else
-			{
-				FRDGBuilder GraphBuilder(RHICmdList);
-				FRDGTextureRef SceneColorTexture = GraphBuilder.RegisterExternalTexture(SceneTargets.GetSceneColor(), TEXT("SceneColorTexture"));
-				AddVoxelRaymarchingPass(GraphBuilder, View, HairDebugMode, MacroGroupDatas, SceneColorTexture);
-				GraphBuilder.Execute();
-			}
 		}
 	}
 
diff --git a/Engine/Source/Runtime/Renderer/Private/HairStrands/HairStrandsRasterCommon.cpp b/Engine/Source/Runtime/Renderer/Private/HairStrands/HairStrandsRasterCommon.cpp
--- a/Engine/Source/Runtime/Renderer/Private/HairStrands/HairStrandsRasterCommon.cpp
+++ b/Engine/Source/Runtime/Renderer/Private/HairStrands/HairStrandsRasterCommon.cpp
@@ -21,7 +21,7 @@ BEGIN_GLOBAL_SHADER_PARAMETER_STRUCT(FHairDeepShadowRasterGlobalParameters, )
 	SHADER_PARAMETER(uint32, AtlasSlotIndex)
 	SHADER_PARAMETER_TEXTURE(Texture2D<float>, FrontDepthTexture)
 	SHADER_PARAMETER_SRV(StructuredBuffer<FDeepShadowViewInfo>, DeepShadowViewInfoBuffer)
-	END_GLOBAL_SHADER_PARAMETER_STRUCT()
+END_GLOBAL_SHADER_PARAMETER_STRUCT()
 IMPLEMENT_GLOBAL_SHADER_PARAMETER_STRUCT(FHairDeepShadowRasterGlobalParameters, "DeepRasterPass");
 
 static FHairDeepShadowRasterGlobalParameters ConvertToGlobalPassParameter(const FHairDeepShadowRasterPassParameters* In)
@@ -48,12 +48,7 @@ BEGIN_GLOBAL_SHADER_PARAMETER_STRUCT(FHairVoxelizationRasterGlobalParameters, )
 	SHADER_PARAMETER(uint32, MacroGroupId)
 	SHADER_PARAMETER(FIntPoint, ViewportResolution)
 	SHADER_PARAMETER_SRV(StructuredBuffer<FVoxelizationViewInfo>, VoxelizationViewInfoBuffer)
-
 	SHADER_PARAMETER_UAV(RWTexture3D<uint>, DensityTexture)
-	SHADER_PARAMETER_UAV(RWTexture3D<uint>, TangentXTexture)
-	SHADER_PARAMETER_UAV(RWTexture3D<uint>, TangentYTexture)
-	SHADER_PARAMETER_UAV(RWTexture3D<uint>, TangentZTexture)
-	SHADER_PARAMETER_UAV(RWTexture3D<uint>, MaterialTexture)
 END_GLOBAL_SHADER_PARAMETER_STRUCT()
 IMPLEMENT_GLOBAL_SHADER_PARAMETER_STRUCT(FHairVoxelizationRasterGlobalParameters, "VoxelRasterPass");
 
@@ -68,13 +63,7 @@ static FHairVoxelizationRasterGlobalParameters ConvertToGlobalPassParameter(cons
 	Out.MacroGroupId	   = In->MacroGroupId;
 	Out.ViewportResolution = In->ViewportResolution;
 	Out.VoxelizationViewInfoBuffer = In->VoxelizationViewInfoBuffer->GetRHI();
-
 	Out.DensityTexture  = In->DensityTexture->GetRHI();
-	Out.TangentXTexture = In->TangentXTexture->GetRHI();
-	Out.TangentYTexture = In->TangentYTexture->GetRHI();
-	Out.TangentZTexture = In->TangentZTexture->GetRHI();
-	Out.MaterialTexture = In->MaterialTexture->GetRHI();
-
 	return Out;
 }
 
@@ -125,7 +114,6 @@ protected:
 	{
 		ERHIFeatureLevel::Type FeatureLevel = GetMaxSupportedFeatureLevel((EShaderPlatform)Initializer.Target.Platform);
 		check(FSceneInterface::GetShadingPath(FeatureLevel) != EShadingPath::Mobile);
-		// deferred
 		PassUniformBuffer.Bind(Initializer.ParameterMap, FHairDeepShadowRasterGlobalParameters::StaticStructMetadata.GetShaderVariableName());
 	}
 
@@ -160,7 +148,6 @@ protected:
 	{
 		ERHIFeatureLevel::Type FeatureLevel = GetMaxSupportedFeatureLevel((EShaderPlatform)Initializer.Target.Platform);
 		check(FSceneInterface::GetShadingPath(FeatureLevel) != EShadingPath::Mobile);
-		// deferred
 		PassUniformBuffer.Bind(Initializer.ParameterMap, FHairVoxelizationRasterGlobalParameters::StaticStructMetadata.GetShaderVariableName());
 	}
 
@@ -184,11 +171,9 @@ protected:
 
 typedef FVoxelMeshVS<false, false> TVoxelMeshVS_NoMaterial_NoCluster;
 typedef FVoxelMeshVS<false, true>  TVoxelMeshVS_NoMaterial_Cluster;
-typedef FVoxelMeshVS<true,  true>  TVoxelMeshVS_Material_Cluster;
 
 IMPLEMENT_MATERIAL_SHADER_TYPE(template<>, TVoxelMeshVS_NoMaterial_NoCluster, TEXT("/Engine/Private/HairStrands/HairStrandsDeepShadowVS.usf"), TEXT("Main"), SF_Vertex);
 IMPLEMENT_MATERIAL_SHADER_TYPE(template<>, TVoxelMeshVS_NoMaterial_Cluster, TEXT("/Engine/Private/HairStrands/HairStrandsDeepShadowVS.usf"), TEXT("Main"), SF_Vertex);
-IMPLEMENT_MATERIAL_SHADER_TYPE(template<>, TVoxelMeshVS_Material_Cluster, TEXT("/Engine/Private/HairStrands/HairStrandsDeepShadowVS.usf"), TEXT("Main"), SF_Vertex);
 
 /////////////////////////////////////////////////////////////////////////////////////////
 
@@ -245,8 +230,7 @@ IMPLEMENT_MATERIAL_SHADER_TYPE(, FDeepShadowDomMeshPS, TEXT("/Engine/Private/Hai
 enum class EVoxelMeshPSType
 {
 	Density,
-	Material,
-	DensityVirtual
+	Material
 };
 template<EVoxelMeshPSType VoxelizationType>
 class FVoxelMeshPS : public FMeshMaterialShader
@@ -275,12 +259,10 @@ public:
 		FMaterialShader::ModifyCompilationEnvironment(Parameters, OutEnvironment);
 		OutEnvironment.SetDefine(TEXT("SUPPORT_TANGENT_PROPERTY"), VoxelizationType == EVoxelMeshPSType::Material ? 1 : 0);
 		OutEnvironment.SetDefine(TEXT("SUPPORT_MATERIAL_PROPERTY"), VoxelizationType == EVoxelMeshPSType::Material ? 1 : 0);
-		OutEnvironment.SetDefine(TEXT("SUPPORT_SPARSE"), VoxelizationType == EVoxelMeshPSType::DensityVirtual ? 1 : 0);
 	}
 };
 IMPLEMENT_MATERIAL_SHADER_TYPE(template<>, FVoxelMeshPS<EVoxelMeshPSType::Density>, TEXT("/Engine/Private/HairStrands/HairStrandsDeepShadowPS.usf"), TEXT("MainVoxel"), SF_Pixel);
 IMPLEMENT_MATERIAL_SHADER_TYPE(template<>, FVoxelMeshPS<EVoxelMeshPSType::Material>, TEXT("/Engine/Private/HairStrands/HairStrandsDeepShadowPS.usf"), TEXT("MainVoxel"), SF_Pixel);
-IMPLEMENT_MATERIAL_SHADER_TYPE(template<>, FVoxelMeshPS<EVoxelMeshPSType::DensityVirtual>, TEXT("/Engine/Private/HairStrands/HairStrandsDeepShadowPS.usf"), TEXT("MainVoxel"), SF_Pixel);
 
 /////////////////////////////////////////////////////////////////////////////////////////
 
@@ -338,14 +320,10 @@ void FHairRasterMeshProcessor::AddMeshBatch(const FMeshBatch& RESTRICT MeshBatch
 			Process<FDeepShadowDepthMeshVS, FDeepShadowDepthMeshPS>(MeshBatch, BatchElementMask, PrimitiveSceneProxy, StaticMeshId, MaterialRenderProxy, Material, MeshFillMode, MeshCullMode);
 		else if (RasterPassType == EHairStrandsRasterPassType::DeepOpacityMap)
 			Process<FDeepShadowDomMeshVS, FDeepShadowDomMeshPS>(MeshBatch, BatchElementMask, PrimitiveSceneProxy, StaticMeshId, MaterialRenderProxy, Material, MeshFillMode, MeshCullMode);
-		else if (RasterPassType == EHairStrandsRasterPassType::Voxelization)
-			Process<FVoxelMeshVS<false, true>, FVoxelMeshPS<EVoxelMeshPSType::Density>>(MeshBatch, BatchElementMask, PrimitiveSceneProxy, StaticMeshId, MaterialRenderProxy, Material, MeshFillMode, MeshCullMode);
-		else if (RasterPassType == EHairStrandsRasterPassType::VoxelizationMaterial)
-			Process<FVoxelMeshVS<true, true>, FVoxelMeshPS<EVoxelMeshPSType::Material>>(MeshBatch, BatchElementMask, PrimitiveSceneProxy, StaticMeshId, MaterialRenderProxy, Material, MeshFillMode, MeshCullMode);
 		else if (RasterPassType == EHairStrandsRasterPassType::VoxelizationVirtual && bCullingEnable)
-			Process<FVoxelMeshVS<false, true>, FVoxelMeshPS<EVoxelMeshPSType::DensityVirtual>>(MeshBatch, BatchElementMask, PrimitiveSceneProxy, StaticMeshId, MaterialRenderProxy, Material, MeshFillMode, MeshCullMode);
+			Process<FVoxelMeshVS<false, true>, FVoxelMeshPS<EVoxelMeshPSType::Density>>(MeshBatch, BatchElementMask, PrimitiveSceneProxy, StaticMeshId, MaterialRenderProxy, Material, MeshFillMode, MeshCullMode);
 		else if (RasterPassType == EHairStrandsRasterPassType::VoxelizationVirtual && !bCullingEnable)
-			Process<FVoxelMeshVS<false, false>, FVoxelMeshPS<EVoxelMeshPSType::DensityVirtual>>(MeshBatch, BatchElementMask, PrimitiveSceneProxy, StaticMeshId, MaterialRenderProxy, Material, MeshFillMode, MeshCullMode);
+			Process<FVoxelMeshVS<false, false>, FVoxelMeshPS<EVoxelMeshPSType::Density>>(MeshBatch, BatchElementMask, PrimitiveSceneProxy, StaticMeshId, MaterialRenderProxy, Material, MeshFillMode, MeshCullMode);
 	}
 }
 
@@ -436,8 +414,6 @@ void AddHairStrandsRasterPass(
 		{
 		case EHairStrandsRasterPassType::DeepOpacityMap:		return RDG_EVENT_NAME("HairStrandsRasterDeepOpacityMap");
 		case EHairStrandsRasterPassType::FrontDepth:			return RDG_EVENT_NAME("HairStrandsRasterFrontDepth");
-		case EHairStrandsRasterPassType::Voxelization:			return RDG_EVENT_NAME("HairStrandsRasterVoxelization");
-		case EHairStrandsRasterPassType::VoxelizationMaterial:	return RDG_EVENT_NAME("HairStrandsRasterVoxelizationMaterial");
 		case EHairStrandsRasterPassType::VoxelizationVirtual:	return RDG_EVENT_NAME("HairStrandsRasterVoxelizationVirtual");
 		default:												return RDG_EVENT_NAME("Noname");
 		}
@@ -488,12 +464,6 @@ void AddHairStrandsRasterPass(
 				CW_RGBA, BO_Add, BF_One, BF_Zero, BO_Add, BF_One, BF_Zero>::GetRHI());
 			DrawRenderState.SetDepthStencilState(TStaticDepthStencilState<true, CF_DepthNearOrEqual>::GetRHI());
 		}
-		else if (RasterPassType == EHairStrandsRasterPassType::Voxelization || RasterPassType == EHairStrandsRasterPassType::VoxelizationMaterial)
-		{
-			DrawRenderState.SetBlendState(TStaticBlendState<
-				CW_RGBA, BO_Add, BF_One, BF_Zero, BO_Add, BF_One, BF_Zero>::GetRHI());
-			DrawRenderState.SetDepthStencilState(TStaticDepthStencilState<false, CF_Always>::GetRHI());
-		}
 		else if (RasterPassType == EHairStrandsRasterPassType::VoxelizationVirtual)
 		{
 			DrawRenderState.SetBlendState(TStaticBlendState<
@@ -558,21 +528,18 @@ void AddHairVoxelizationRasterPass(
 	const FScene* Scene,
 	const FViewInfo* ViewInfo,
 	const FHairStrandsMacroGroupData::TPrimitiveInfos& PrimitiveSceneInfos,
-	const EHairStrandsRasterPassType PassType,
 	const FIntRect& ViewportRect,
 	const FVector4& HairRenderInfo,
 	const uint32 HairRenderInfoBits,
 	const FVector& RasterDirection,
 	FHairVoxelizationRasterPassParameters* PassParameters)
 {
-	check(PassType == EHairStrandsRasterPassType::Voxelization || PassType == EHairStrandsRasterPassType::VoxelizationMaterial || PassType == EHairStrandsRasterPassType::VoxelizationVirtual);
-
 	AddHairStrandsRasterPass<FHairVoxelizationRasterPassParameters, FHairVoxelizationRasterGlobalParameters>(
 		GraphBuilder, 
 		Scene, 
 		ViewInfo, 
 		PrimitiveSceneInfos, 
-		PassType, 
+		EHairStrandsRasterPassType::VoxelizationVirtual,
 		ViewportRect, 
 		HairRenderInfo, 
 		HairRenderInfoBits,
diff --git a/Engine/Source/Runtime/Renderer/Private/HairStrands/HairStrandsRasterCommon.h b/Engine/Source/Runtime/Renderer/Private/HairStrands/HairStrandsRasterCommon.h
--- a/Engine/Source/Runtime/Renderer/Private/HairStrands/HairStrandsRasterCommon.h
+++ b/Engine/Source/Runtime/Renderer/Private/HairStrands/HairStrandsRasterCommon.h
@@ -17,8 +17,6 @@ enum class EHairStrandsRasterPassType : uint8
 {
 	FrontDepth,
 	DeepOpacityMap,
-	Voxelization,
-	VoxelizationMaterial,
 	VoxelizationVirtual
 };
 
@@ -60,12 +58,7 @@ BEGIN_SHADER_PARAMETER_STRUCT(FHairVoxelizationRasterPassParameters, )
 	SHADER_PARAMETER(uint32, MacroGroupId)
 	SHADER_PARAMETER(FIntPoint, ViewportResolution)
 	SHADER_PARAMETER_RDG_BUFFER_SRV(StructuredBuffer<FVoxelizationViewInfo>, VoxelizationViewInfoBuffer)
-
 	SHADER_PARAMETER_RDG_TEXTURE_UAV(RWTexture3D<uint>, DensityTexture)
-	SHADER_PARAMETER_RDG_TEXTURE_UAV(RWTexture3D<uint>, TangentXTexture)
-	SHADER_PARAMETER_RDG_TEXTURE_UAV(RWTexture3D<uint>, TangentYTexture)
-	SHADER_PARAMETER_RDG_TEXTURE_UAV(RWTexture3D<uint>, TangentZTexture)
-	SHADER_PARAMETER_RDG_TEXTURE_UAV(RWTexture3D<uint>, MaterialTexture)
 
 	RENDER_TARGET_BINDING_SLOTS()
 END_SHADER_PARAMETER_STRUCT()
@@ -75,7 +68,6 @@ void AddHairVoxelizationRasterPass(
 	const FScene* Scene,
 	const FViewInfo* ViewInfo,
 	const FHairStrandsMacroGroupData::TPrimitiveInfos& PrimitiveSceneInfos,
-	const EHairStrandsRasterPassType ShadowPassType,
 	const FIntRect& ViewportRect,
 	const FVector4& HairRenderInfo,
 	const uint32 HairRenderInfoBits,
diff --git a/Engine/Source/Runtime/Renderer/Private/HairStrands/HairStrandsTransmittance.cpp b/Engine/Source/Runtime/Renderer/Private/HairStrands/HairStrandsTransmittance.cpp
--- a/Engine/Source/Runtime/Renderer/Private/HairStrands/HairStrandsTransmittance.cpp
+++ b/Engine/Source/Runtime/Renderer/Private/HairStrands/HairStrandsTransmittance.cpp
@@ -97,8 +97,8 @@ class FDeepTransmittanceMaskCS : public FGlobalShader
 		SHADER_PARAMETER(uint32, MaxVisibilityNodeCount)
 		SHADER_PARAMETER(FVector4, LightPosition)
 		SHADER_PARAMETER(uint32, LightChannelMask)
-		SHADER_PARAMETER(float, DepthBiasScale)
-		SHADER_PARAMETER(float, DensityScale)
+		SHADER_PARAMETER(float, DeepShadow_DepthBiasScale)
+		SHADER_PARAMETER(float, DeepShadow_DensityScale)
 		SHADER_PARAMETER(float, DeepShadow_KernelAperture)
 		SHADER_PARAMETER(uint32, DeepShadow_KernelType)
 		SHADER_PARAMETER(uint32, DeepShadow_DebugMode)
@@ -161,8 +161,8 @@ struct FDeepShadowTransmittanceParams
 	FVector4 LightPosition = FVector4(0, 0, 0, 0);
 	uint32 LightChannelMask = 0;
 	float LightRadius = 0;
-	float DepthBiasScale = 0;
-	float DensityScale = 0;
+	float DeepShadow_DepthBiasScale = 0;
+	float DeepShadow_DensityScale = 0;
 	FMatrix DeepShadow_ShadowToWorld = FMatrix::Identity;
 	
 	FRDGTextureRef DeepShadow_FrontDepthTexture = nullptr;
@@ -204,8 +204,8 @@ static FRDGBufferRef AddDeepShadowTransmittanceMaskPass(
 	Parameters->LightDirection = Params.LightDirection;
 	Parameters->LightPosition = Params.LightPosition;
 	Parameters->LightRadius = Params.LightRadius;
-	Parameters->DepthBiasScale = Params.DepthBiasScale;
-	Parameters->DensityScale = Params.DensityScale;
+	Parameters->DeepShadow_DepthBiasScale = Params.DeepShadow_DepthBiasScale;
+	Parameters->DeepShadow_DensityScale = Params.DeepShadow_DensityScale;
 	Parameters->DeepShadow_KernelAperture = GetDeepShadowKernelAperture();
 	Parameters->DeepShadow_KernelType = GetDeepShadowKernelType();
 	Parameters->DeepShadow_DebugMode = GetDeepShadowDebugMode();
@@ -480,7 +480,8 @@ static FHairStrandsTransmittanceMaskData RenderHairStrandsTransmittanceMask(
 	FDeepShadowTransmittanceParams Params;
 	Params.HairVisibilityNodeData = GraphBuilder.RegisterExternalBuffer(VisibilityData.NodeData, TEXT("HairVisibilityNodeData"));
 	Params.HairVisibilityNodeCoord = GraphBuilder.RegisterExternalBuffer(VisibilityData.NodeCoord, TEXT("HairVisibilityNodeCoord"));
-	Params.DensityScale = GetDeepShadowDensityScale();
+	Params.DeepShadow_DensityScale = GetDeepShadowDensityScale();
+	Params.DeepShadow_DepthBiasScale = GetDeepShadowDepthBiasScale();
 	memset(Params.DeepShadow_AtlasSlotOffsets_AtlasSlotIndex, 0, sizeof(Params.DeepShadow_AtlasSlotOffsets_AtlasSlotIndex));
 	memset(Params.DeepShadow_CPUWorldToLightTransforms, 0, sizeof(Params.DeepShadow_CPUWorldToLightTransforms));
 
@@ -514,7 +515,6 @@ static FHairStrandsTransmittanceMaskData RenderHairStrandsTransmittanceMask(
 					Params.LightPosition = DeepShadowData.LightPosition;
 					Params.LightRadius = 0;
 					Params.LightChannelMask = LightSceneInfo->Proxy->GetLightingChannelMask();
-					Params.DepthBiasScale = GetDeepShadowDepthBiasScale();
 					Params.DeepShadow_AtlasSlotOffsets_AtlasSlotIndex[DeepShadowData.MacroGroupId] = FIntVector4(DeepShadowData.AtlasRect.Min.X, DeepShadowData.AtlasRect.Min.Y, DeepShadowData.AtlasSlotIndex, 0);
 					Params.DeepShadow_CPUWorldToLightTransforms[DeepShadowData.MacroGroupId] = DeepShadowData.CPU_WorldToLightTransform;
 					Params.DeepShadow_WorldToLightTransformBuffer = DeepShadow_WorldToLightTransformBufferSRV;
@@ -545,8 +545,8 @@ static FHairStrandsTransmittanceMaskData RenderHairStrandsTransmittanceMask(
 		FLightShaderParameters LightParameters;
 		LightSceneInfo->Proxy->GetLightShaderParameters(LightParameters);
 
-		Params.DensityScale = GetHairStrandsVoxelizationDensityScale();
-		Params.DepthBiasScale = GetHairStrandsVoxelizationDepthBiasScale();
+		Params.DeepShadow_DensityScale = 0;
+		Params.DeepShadow_DepthBiasScale = 0;
 		Params.LightDirection = LightSceneInfo->Proxy->GetDirection();
 		Params.LightPosition = FVector4(LightSceneInfo->Proxy->GetPosition(), LightSceneInfo->Proxy->GetLightType() == ELightComponentType::LightType_Directional ? 0 : 1);
 		Params.LightChannelMask = LightSceneInfo->Proxy->GetLightingChannelMask();
diff --git a/Engine/Source/Runtime/Renderer/Private/HairStrands/HairStrandsVoxelization.cpp b/Engine/Source/Runtime/Renderer/Private/HairStrands/HairStrandsVoxelization.cpp
--- a/Engine/Source/Runtime/Renderer/Private/HairStrands/HairStrandsVoxelization.cpp
+++ b/Engine/Source/Runtime/Renderer/Private/HairStrands/HairStrandsVoxelization.cpp
@@ -15,15 +15,9 @@
 static float GStrandHairVoxelizationRasterizationScale = 1.0f;
 static FAutoConsoleVariableRef CVarStrandHairVoxelizationRasterizationScale(TEXT("r.HairStrands.VoxelizationRasterizationScale"), GStrandHairVoxelizationRasterizationScale, TEXT("Rasterization scale to snap strand to pixel for voxelization"));
 
-static int32 GHairVoxelizationResolution = 256;
-static FAutoConsoleVariableRef CVarGHairVoxelizationResolution(TEXT("r.HairStrands.Voxelization.Resolution"), GHairVoxelizationResolution, TEXT("Change the resolution of the voxelization volume for hair strands"));
-
 static int32 GHairVoxelizationEnable = 1;
 static FAutoConsoleVariableRef CVarGHairVoxelizationEnable(TEXT("r.HairStrands.Voxelization"), GHairVoxelizationEnable, TEXT("Enable hair voxelization for transmittance evaluation"));
 
-static int32 GHairVoxelizationMaterialEnable = 0;
-static FAutoConsoleVariableRef CVarHairVoxelizationMaterialEnable(TEXT("r.HairStrands.VoxelizationMaterial"), GHairVoxelizationMaterialEnable, TEXT("Enable hair material voxelization for LOD evaluation"));
-
 static float GHairVoxelizationAABBScale = 1.0f;
 static FAutoConsoleVariableRef CVarHairVoxelizationAABBScale(TEXT("r.HairStrands.Voxelization.AABBScale"), GHairVoxelizationAABBScale, TEXT("Scale the hair macro group bounding box"));
 
@@ -35,12 +29,6 @@ static FAutoConsoleVariableRef CVarHairVoxelizationDepthBiasScale(TEXT("r.HairSt
 static int32 GHairVoxelInjectOpaqueDepthEnable = 1;
 static FAutoConsoleVariableRef CVarHairVoxelInjectOpaqueDepthEnable(TEXT("r.HairStrands.Voxelization.InjectOpaqueDepth"), GHairVoxelInjectOpaqueDepthEnable, TEXT("Inject opaque geometry depth into the voxel volume for acting as occluder."));
 
-static int32 GHairVoxelFilterOpaqueDepthEnable = 0;
-static FAutoConsoleVariableRef CVarHairVoxelFilterOpaqueDepthEnable(TEXT("r.HairStrands.Voxelization.FilterOpaqueDepth"), GHairVoxelFilterOpaqueDepthEnable, TEXT("Filter opaque geometry depth into the voxel volume for acting as occluder."));
-
-static int32 GHairStrandsVoxelMipMethod = 0;
-static FAutoConsoleVariableRef CVarHairVoxelMipMethod(TEXT("r.HairStrands.Voxelization.MipMethod"), GHairStrandsVoxelMipMethod, TEXT("Voxel mip methods (0 : one level per pass, 1 : two levels per pass."));
-
 static int32 GHairStransVoxelInjectOpaqueBiasCount = 3;
 static int32 GHairStransVoxelInjectOpaqueMarkCount = 6;
 static FAutoConsoleVariableRef CVarHairStransVoxelInjectOpaqueBiasCount(TEXT("r.HairStrands.Voxelization.InjectOpaque.BiasCount"), GHairStransVoxelInjectOpaqueBiasCount, TEXT("Bias, in number of voxel, at which opaque depth is injected."));
@@ -49,8 +37,8 @@ static FAutoConsoleVariableRef CVarHairStransVoxelInjectOpaqueMarkCount(TEXT("r.
 static float GHairStransVoxelRaymarchingSteppingScale = 1.15f;
 static FAutoConsoleVariableRef CVarHairStransVoxelRaymarchingSteppingScale(TEXT("r.HairStrands.Voxelization.Raymarching.SteppingScale"), GHairStransVoxelRaymarchingSteppingScale, TEXT("Stepping scale used for raymarching the voxel structure."));
 
-float GetHairStrandsVoxelizationDensityScale() { return FMath::Max(0.0f, GHairVoxelizationDensityScale); }
-float GetHairStrandsVoxelizationDepthBiasScale() { return FMath::Max(0.0f, GHairVoxelizationDepthBiasScale); }
+static float GetHairStrandsVoxelizationDensityScale() { return FMath::Max(0.0f, GHairVoxelizationDensityScale); }
+static float GetHairStrandsVoxelizationDepthBiasScale() { return FMath::Max(0.0f, GHairVoxelizationDepthBiasScale); }
 
 static int32 GHairForVoxelTransmittanceAndShadow = 0;
 static FAutoConsoleVariableRef CVarHairForVoxelTransmittanceAndShadow(TEXT("r.HairStrands.Voxelization.ForceTransmittanceAndShadow"), GHairForVoxelTransmittanceAndShadow, TEXT("For transmittance and shadow to be computed with density volume. This requires voxelization is enabled."));
@@ -97,133 +85,6 @@ bool IsHairStrandsForVoxelTransmittanceAndShadowEnable()
 
 IMPLEMENT_GLOBAL_SHADER_PARAMETER_STRUCT(FVirtualVoxelParameters, "VirtualVoxel");
 
-///////////////////////////////////////////////////////////////////////////////////////////////////
-class FVoxelInjectOpaquePS : public FGlobalShader
-{
-	DECLARE_GLOBAL_SHADER(FVoxelInjectOpaquePS);
-	SHADER_USE_PARAMETER_STRUCT(FVoxelInjectOpaquePS, FGlobalShader);
-
-	BEGIN_SHADER_PARAMETER_STRUCT(FParameters, )
-		SHADER_PARAMETER_STRUCT_INCLUDE(FSceneTextureParameters, SceneTextures)
-
-		SHADER_PARAMETER(FVector, VoxelMinAABB)
-		SHADER_PARAMETER(uint32, VoxelResolution)
-		SHADER_PARAMETER(FVector, VoxelMaxAABB)
-		SHADER_PARAMETER(uint32, VoxelBiasCount)
-		SHADER_PARAMETER(FVector2D, OutputResolution)
-		SHADER_PARAMETER(FVector2D, SceneDepthResolution)
-		SHADER_PARAMETER(uint32, VoxelMarkCount)
-
-		SHADER_PARAMETER_RDG_TEXTURE(Texture2D, SceneDepthTexture)
-		SHADER_PARAMETER_RDG_TEXTURE_UAV(RWTexture3D, DensityTexture)
-		SHADER_PARAMETER_SAMPLER(SamplerState, LinearSampler)
-
-		SHADER_PARAMETER_STRUCT_REF(FViewUniformShaderParameters, ViewUniformBuffer)
-		RENDER_TARGET_BINDING_SLOTS()
-	END_SHADER_PARAMETER_STRUCT()
-
-public:
-	static bool ShouldCompilePermutation(const FGlobalShaderPermutationParameters& Parameters) { return IsHairStrandsSupported(Parameters.Platform); }
-	static void ModifyCompilationEnvironment(const FGlobalShaderPermutationParameters& Parameters, FShaderCompilerEnvironment& OutEnvironment)
-	{
-		FGlobalShader::ModifyCompilationEnvironment(Parameters, OutEnvironment);
-		OutEnvironment.SetDefine(TEXT("SHADER_INJECTOPAQUE_VOXEL"), 1);
-	}
-};
-
-IMPLEMENT_GLOBAL_SHADER(FVoxelInjectOpaquePS, "/Engine/Private/HairStrands/HairStrandsVoxelOpaque.usf", "MainPS", SF_Pixel);
-
-static void AddVoxelInjectOpaquePass(
-	FRDGBuilder& GraphBuilder,
-	const FViewInfo& View,
-	FHairStrandsMacroGroupData& MacroGroup)
-{
-	FSceneTextureParameters SceneTextures;
-	SetupSceneTextureParameters(GraphBuilder, &SceneTextures);
-
-	// #hair_todo: change this to a CS. PS was for easing debugging
-	const uint32 LinearizeRes = FMath::Sqrt(MacroGroup.VoxelResources.DensityTexture->GetDesc().Depth);
-
-	FIntPoint Resolution;
-	Resolution.X = MacroGroup.VoxelResources.DensityTexture->GetDesc().Extent.X * LinearizeRes;
-	Resolution.Y = MacroGroup.VoxelResources.DensityTexture->GetDesc().Extent.Y * LinearizeRes;
-
-	if (!MacroGroup.VoxelResources.DensityTexture)
-		return;
-
-	const FRDGTextureRef VoxelDensityTexture = GraphBuilder.RegisterExternalTexture(MacroGroup.VoxelResources.DensityTexture, TEXT("HairVoxelDensityTexture"));
-
-	FRDGTextureDesc OutputDesc;
-	OutputDesc.Extent.X = Resolution.X;
-	OutputDesc.Extent.Y = Resolution.Y;
-	OutputDesc.Depth = 0;
-	OutputDesc.Format = PF_FloatRGBA;
-	OutputDesc.NumMips = 1;
-	OutputDesc.Flags = 0;
-	OutputDesc.TargetableFlags = TexCreate_RenderTargetable;
-	FRDGTextureRef DummyTexture = GraphBuilder.CreateTexture(OutputDesc, TEXT("HairVoxelInjectDepth")); // Dummy texture for debugging. Convert this pass into a compute shader.
-
-	FVoxelInjectOpaquePS::FParameters* Parameters = GraphBuilder.AllocParameters<FVoxelInjectOpaquePS::FParameters>();
-	Parameters->ViewUniformBuffer = View.ViewUniformBuffer;
-	Parameters->OutputResolution = Resolution;
-	Parameters->VoxelBiasCount = FMath::Max(0, GHairStransVoxelInjectOpaqueBiasCount);
-	Parameters->VoxelMarkCount = FMath::Max(0, GHairStransVoxelInjectOpaqueMarkCount);
-	Parameters->SceneDepthResolution = SceneTextures.SceneDepthBuffer->Desc.Extent;
-	Parameters->SceneDepthTexture = SceneTextures.SceneDepthBuffer;
-	Parameters->SceneTextures = SceneTextures;
-	Parameters->DensityTexture = GraphBuilder.CreateUAV(VoxelDensityTexture);
-	Parameters->VoxelMinAABB = MacroGroup.GetMinBound();
-	Parameters->VoxelMaxAABB = MacroGroup.GetMaxBound();
-	Parameters->VoxelResolution = MacroGroup.GetResolution();
-	Parameters->LinearSampler = TStaticSamplerState<SF_Trilinear, AM_Clamp, AM_Clamp, AM_Clamp>::GetRHI();
-	Parameters->RenderTargets[0] = FRenderTargetBinding(DummyTexture, ERenderTargetLoadAction::EClear);
-
-	const FIntPoint OutputResolution = SceneTextures.SceneDepthBuffer->Desc.Extent;
-	TShaderMapRef<FPostProcessVS> VertexShader(View.ShaderMap);
-	TShaderMapRef<FVoxelInjectOpaquePS> PixelShader(View.ShaderMap);
-	const FGlobalShaderMap* GlobalShaderMap = View.ShaderMap;
-	const FIntRect Viewport(0, 0, Resolution.X, Resolution.Y);
-	const FViewInfo* CapturedView = &View;
-	
-	ClearUnusedGraphResources(PixelShader, Parameters);
-
-	GraphBuilder.AddPass(
-		RDG_EVENT_NAME("HairStrandsVoxelInjectOpaque"),
-		Parameters,
-		ERDGPassFlags::Raster,
-		[Parameters, VertexShader, PixelShader, Viewport, Resolution, CapturedView](FRHICommandList& RHICmdList)
-	{
-		FGraphicsPipelineStateInitializer GraphicsPSOInit;
-		RHICmdList.ApplyCachedRenderTargets(GraphicsPSOInit);
-		GraphicsPSOInit.BlendState = TStaticBlendState<CW_RGBA, BO_Add, BF_One, BF_Zero, BO_Add, BF_One, BF_Zero>::GetRHI();
-		GraphicsPSOInit.RasterizerState = TStaticRasterizerState<>::GetRHI();
-		GraphicsPSOInit.DepthStencilState = TStaticDepthStencilState<false, CF_Always>::GetRHI();
-
-		GraphicsPSOInit.BoundShaderState.VertexDeclarationRHI = GFilterVertexDeclaration.VertexDeclarationRHI;
-		GraphicsPSOInit.BoundShaderState.VertexShaderRHI = VertexShader.GetVertexShader();
-		GraphicsPSOInit.BoundShaderState.PixelShaderRHI = PixelShader.GetPixelShader();
-		GraphicsPSOInit.PrimitiveType = PT_TriangleList;
-		SetGraphicsPipelineState(RHICmdList, GraphicsPSOInit);
-
-		VertexShader->SetParameters(RHICmdList, CapturedView->ViewUniformBuffer);
-		RHICmdList.SetViewport(Viewport.Min.X, Viewport.Min.Y, 0.0f, Viewport.Max.X, Viewport.Max.Y, 1.0f);
-		SetShaderParameters(RHICmdList, PixelShader, PixelShader.GetPixelShader(), *Parameters);
-
-		DrawRectangle(
-			RHICmdList,
-			0, 0,
-			Viewport.Width(), Viewport.Height(),
-			Viewport.Min.X, Viewport.Min.Y,
-			Viewport.Width(), Viewport.Height(),
-			Viewport.Size(),
-			Resolution,
-			VertexShader,
-			EDRF_UseTriangleOptimization);
-	});
-}
-
-
-
 ///////////////////////////////////////////////////////////////////////////////////////////////////
 class FVirtualVoxelInjectOpaqueCS : public FGlobalShader
 {
@@ -1271,7 +1132,6 @@ static void AddVirtualVoxelizationRasterPass(
 		RasterResolution = GPUDrivenViewportResolution;
 	}
 
-
 	const FBox RasterAABB(MacroGroup.VirtualVoxelNodeDesc.WorldMinAABB, MacroGroup.VirtualVoxelNodeDesc.WorldMaxAABB);
 	const FVector RasterAABBSize = RasterAABB.GetSize();
 	const FVector RasterAABBCenter = RasterAABB.GetCenter();
@@ -1311,170 +1171,15 @@ static void AddVirtualVoxelizationRasterPass(
 	PassParameters->VoxelizationViewInfoBuffer = GraphBuilder.CreateSRV(VoxelizationViewInfoBuffer);
 	PassParameters->DensityTexture = GraphBuilder.CreateUAV(PageTexture);
 
-	const FPooledRenderTargetDesc DummyDesc = FPooledRenderTargetDesc::CreateVolumeDesc(1, 1, 1, PF_R32_UINT, FClearValueBinding::Black, TexCreate_None, TexCreate_UAV | TexCreate_ShaderResource, false, 1);
-	FRDGTextureRef DummyTangentXTexture = GraphBuilder.CreateTexture(DummyDesc, TEXT("DummyTexture"));
-	FRDGTextureRef DummyTangentYTexture = GraphBuilder.CreateTexture(DummyDesc, TEXT("DummyTexture"));
-	FRDGTextureRef DummyTangentZTexture = GraphBuilder.CreateTexture(DummyDesc, TEXT("DummyTexture"));
-	FRDGTextureRef DummyMaterialTexture = GraphBuilder.CreateTexture(DummyDesc, TEXT("DummyTexture"));
-	PassParameters->TangentXTexture = GraphBuilder.CreateUAV(DummyTangentXTexture);
-	PassParameters->TangentYTexture = GraphBuilder.CreateUAV(DummyTangentYTexture);
-	PassParameters->TangentZTexture = GraphBuilder.CreateUAV(DummyTangentZTexture);
-	PassParameters->MaterialTexture = GraphBuilder.CreateUAV(DummyMaterialTexture);
-
 	// For debug purpose
 	#if 0
 	FRDGTextureRef DebugOutputTexture = GraphBuilder.CreateTexture(FPooledRenderTargetDesc::Create2DDesc(RasterResolution, PF_R32_UINT, FClearValueBinding::Black, TexCreate_None, TexCreate_RenderTargetable, false), TEXT("DummyTexture"));
 	PassParameters->RenderTargets[0] = FRenderTargetBinding(DebugOutputTexture, ERenderTargetLoadAction::EClear);
 	#endif
 
-	AddHairVoxelizationRasterPass(GraphBuilder, Scene, ViewInfo, PrimitiveSceneInfo, EHairStrandsRasterPassType::VoxelizationVirtual, ViewportRect, HairRenderInfo, HairRenderInfoBits, RasterDirection, PassParameters);
-}
-
-///////////////////////////////////////////////////////////////////////////////////////////////////
-class FVoxelFilterDepthCS : public FGlobalShader
-{
-	DECLARE_GLOBAL_SHADER(FVoxelFilterDepthCS);
-	SHADER_USE_PARAMETER_STRUCT(FVoxelFilterDepthCS, FGlobalShader);
-
-	BEGIN_SHADER_PARAMETER_STRUCT(FParameters, )
-		SHADER_PARAMETER_STRUCT_INCLUDE(FSceneTextureParameters, SceneTextures)
-
-		SHADER_PARAMETER(uint32, VoxelResolution)
-		SHADER_PARAMETER_RDG_TEXTURE_UAV(RWTexture3D, VoxelTexture)
-
-	END_SHADER_PARAMETER_STRUCT()
-
-public:
-	static bool ShouldCompilePermutation(const FGlobalShaderPermutationParameters& Parameters) { return IsHairStrandsSupported(Parameters.Platform); }
-	static void ModifyCompilationEnvironment(const FGlobalShaderPermutationParameters& Parameters, FShaderCompilerEnvironment& OutEnvironment)
-	{
-		FGlobalShader::ModifyCompilationEnvironment(Parameters, OutEnvironment);
-		OutEnvironment.SetDefine(TEXT("SHADER_FILTERDEPTH_VOXEL"), 1);
-	}
-};
-
-IMPLEMENT_GLOBAL_SHADER(FVoxelFilterDepthCS, "/Engine/Private/HairStrands/HairStrandsVoxelOpaque.usf", "MainCS", SF_Compute);
-
-static void AddFilterVoxelOpaqueDepthPass(
-	FRDGBuilder& GraphBuilder,
-	const FViewInfo& View,
-	FHairStrandsMacroGroupData& MacroGroup)
-{
-	if (!MacroGroup.VoxelResources.DensityTexture)
-		return;
-
-	TRefCountPtr<IPooledRenderTarget> DensityTexture = MacroGroup.VoxelResources.DensityTexture;
-	check(DensityTexture->GetDesc().Extent.X == DensityTexture->GetDesc().Extent.Y);
-
-	const uint32 VoxelResolution = DensityTexture->GetDesc().Extent.X;
-	FRDGTextureRef VoxelDensityTexture = GraphBuilder.RegisterExternalTexture(DensityTexture, TEXT("HairVoxelDensityTexture"));
-	FVoxelFilterDepthCS::FParameters* Parameters = GraphBuilder.AllocParameters<FVoxelFilterDepthCS::FParameters>();
-	Parameters->VoxelTexture = GraphBuilder.CreateUAV(FRDGTextureUAVDesc(VoxelDensityTexture));
-	Parameters->VoxelResolution = VoxelResolution;
-
-	TShaderMapRef<FVoxelFilterDepthCS> ComputeShader(View.ShaderMap);
-	const FGlobalShaderMap* GlobalShaderMap = View.ShaderMap;
-	const FIntVector DispatchCount = FComputeShaderUtils::GetGroupCount(FIntVector(VoxelResolution, VoxelResolution, VoxelResolution), FIntVector(4, 4, 4));
-
-	ClearUnusedGraphResources(ComputeShader, Parameters);
-	GraphBuilder.AddPass(
-		RDG_EVENT_NAME("HairStrandsVoxelFilterDepth"),
-		Parameters,
-		ERDGPassFlags::Compute,
-		[Parameters, ComputeShader, DispatchCount](FRHICommandList& RHICmdList)
-	{
-		FComputeShaderUtils::Dispatch(RHICmdList, ComputeShader, *Parameters, DispatchCount);
-	});
-}
-
-///////////////////////////////////////////////////////////////////////////////////////////////////
-class FVoxelGenerateMipCS : public FGlobalShader
-{
-	DECLARE_GLOBAL_SHADER(FVoxelGenerateMipCS);
-	SHADER_USE_PARAMETER_STRUCT(FVoxelGenerateMipCS, FGlobalShader);
-
-	class FMethod : SHADER_PERMUTATION_INT("PERMUTATION_METHOD", 2);
-	using FPermutationDomain = TShaderPermutationDomain<FMethod>;
-
-	BEGIN_SHADER_PARAMETER_STRUCT(FParameters, )
-		SHADER_PARAMETER_STRUCT_INCLUDE(FSceneTextureParameters, SceneTextures)
-
-		SHADER_PARAMETER(uint32, VoxelResolution)
-		SHADER_PARAMETER(uint32, SourceMip)
-		SHADER_PARAMETER(uint32, TargetMip)
-
-		SHADER_PARAMETER_RDG_TEXTURE_SRV(Texture3D, InDensityTexture)
-		SHADER_PARAMETER_RDG_TEXTURE_UAV(RWTexture3D, OutDensityTexture0)
-		SHADER_PARAMETER_RDG_TEXTURE_UAV(RWTexture3D, OutDensityTexture1)
-
-	END_SHADER_PARAMETER_STRUCT()
-
-public:
-	static bool ShouldCompilePermutation(const FGlobalShaderPermutationParameters& Parameters) { return IsHairStrandsSupported(Parameters.Platform); }
-	static void ModifyCompilationEnvironment(const FGlobalShaderPermutationParameters& Parameters, FShaderCompilerEnvironment& OutEnvironment)
-	{
-		FGlobalShader::ModifyCompilationEnvironment(Parameters, OutEnvironment);
-		OutEnvironment.SetDefine(TEXT("SHADER_MIP_VOXEL"), 1);
-	}
-};
-
-IMPLEMENT_GLOBAL_SHADER(FVoxelGenerateMipCS, "/Engine/Private/HairStrands/HairStrandsVoxelMip.usf", "MainCS", SF_Compute);
-
-static void AddVoxelGenerateMipPass(
-	FRDGBuilder& GraphBuilder,
-	const FViewInfo& View,
-	FHairStrandsMacroGroupData& MacroGroup)
-{
-	if (!MacroGroup.VoxelResources.DensityTexture)
-		return;
-
-	TRefCountPtr<IPooledRenderTarget> DensityTexture = MacroGroup.VoxelResources.DensityTexture;
-	check(DensityTexture->GetDesc().Extent.X == DensityTexture->GetDesc().Extent.Y);
-
-	const uint32 NumLevelPerPass = GHairStrandsVoxelMipMethod > 0 ? 2 : 1;
-
-	const uint32 VoxelResolution = DensityTexture->GetDesc().Extent.X;
-	const uint32 MipCount = DensityTexture->GetDesc().NumMips;
-	FRDGTextureRef VoxelDensityTexture = GraphBuilder.RegisterExternalTexture(DensityTexture, TEXT("HairVoxelDensityTexture"));
-	for (uint32 MipIt = 0; MipIt < MipCount-1; MipIt += NumLevelPerPass)
-	{
-		FVoxelGenerateMipCS::FParameters* Parameters = GraphBuilder.AllocParameters<FVoxelGenerateMipCS::FParameters>();
-		Parameters->InDensityTexture = GraphBuilder.CreateSRV(FRDGTextureSRVDesc::CreateForMipLevel(VoxelDensityTexture, MipIt));
-		Parameters->OutDensityTexture0 = GraphBuilder.CreateUAV(FRDGTextureUAVDesc(VoxelDensityTexture, MipIt+1));
-		if (NumLevelPerPass > 1)
-		{
-			Parameters->OutDensityTexture1 = GraphBuilder.CreateUAV(FRDGTextureUAVDesc(VoxelDensityTexture, MipIt+2));
-		}
-		Parameters->VoxelResolution = VoxelResolution;
-		Parameters->SourceMip = MipIt;
-		Parameters->TargetMip = MipIt+1;
-
-		const uint32 SourceResolution = VoxelResolution >> (MipIt);
-		const uint32 TargetResolution = VoxelResolution >> (MipIt + 1);
-
-		FVoxelGenerateMipCS::FPermutationDomain PermutationVector;
-		PermutationVector.Set<FVoxelGenerateMipCS::FMethod>(NumLevelPerPass == 2 ? 1 : 0);
-
-		TShaderMapRef<FVoxelGenerateMipCS> ComputeShader(View.ShaderMap, PermutationVector);
-		const FGlobalShaderMap* GlobalShaderMap = View.ShaderMap;
-		const FIntVector DispatchCount = 
-			NumLevelPerPass == 1 ? 
-			FComputeShaderUtils::GetGroupCount(FIntVector(TargetResolution, TargetResolution, TargetResolution), FIntVector(4, 4, 4)):
-			FComputeShaderUtils::GetGroupCount(FIntVector(SourceResolution, SourceResolution, SourceResolution), FIntVector(4, 4, 4));
-		
-		ClearUnusedGraphResources(ComputeShader, Parameters);
-		GraphBuilder.AddPass(
-			RDG_EVENT_NAME("HairStrandsVoxelMip"),
-			Parameters,
-			ERDGPassFlags::Compute | ERDGPassFlags::GenerateMips,
-			[Parameters, ComputeShader, DispatchCount](FRHICommandList& RHICmdList)
-		{
-			FComputeShaderUtils::Dispatch(RHICmdList, ComputeShader, *Parameters, DispatchCount);
-		});
-	}
+	AddHairVoxelizationRasterPass(GraphBuilder, Scene, ViewInfo, PrimitiveSceneInfo, ViewportRect, HairRenderInfo, HairRenderInfoBits, RasterDirection, PassParameters);
 }
 
-
 ///////////////////////////////////////////////////////////////////////////////////////////////////
 class FVirtualVoxelGenerateMipCS : public FGlobalShader
 {
@@ -1656,111 +1361,6 @@ static void AddVirtualVoxelGenerateMipPass(
 }
 /////////////////////////////////////////////////////////////////////////////////////////
 
-static void AddVoxelizationRasterPass(
-	FRDGBuilder& GraphBuilder,
-	const FScene* Scene,
-	const FViewInfo* ViewInfo,
-	FHairStrandsMacroGroupData& MacroGroup)
-{
-	DECLARE_GPU_STAT(HairStrandsVoxelize);
-	SCOPED_DRAW_EVENT(GraphBuilder.RHICmdList, HairStrandsVoxelize);
-	SCOPED_GPU_STAT(GraphBuilder.RHICmdList, HairStrandsVoxelize);
-
-	const FHairStrandsMacroGroupData::TPrimitiveInfos& PrimitiveSceneInfo = MacroGroup.PrimitivesInfos;
-	const FBoxSphereBounds& Bounds = MacroGroup.Bounds;
-	FHairStrandsVoxelResources& VoxelResources = MacroGroup.VoxelResources;
-
-	const uint32 ResolutionDim = FMath::RoundUpToPowerOfTwo(GHairVoxelizationResolution);
-	const EPixelFormat Format = PF_R32_UINT;
-	const FIntPoint Resolution(ResolutionDim, ResolutionDim);
-	uint32 MipCount = 1;
-	{
-		uint32 CurrentResolution = ResolutionDim;
-		while (CurrentResolution > 4)
-		{
-			++MipCount;
-			CurrentResolution = CurrentResolution >> 1;
-		}
-	}
-
-	const bool bVoxelizeMaterial = GHairVoxelizationMaterialEnable > 0;
-	FRDGTextureRef DensityTexture = nullptr;
-	FRDGTextureRef TangentXTexture = nullptr;
-	FRDGTextureRef TangentYTexture = nullptr;
-	FRDGTextureRef TangentZTexture = nullptr;
-	FRDGTextureRef MaterialTexture = nullptr;
-	{
-		FPooledRenderTargetDesc Desc = FPooledRenderTargetDesc::CreateVolumeDesc(ResolutionDim, ResolutionDim, ResolutionDim, PF_R32_UINT, FClearValueBinding::Black, TexCreate_None, TexCreate_UAV | TexCreate_ShaderResource, false, MipCount);
-		DensityTexture = GraphBuilder.CreateTexture(Desc, TEXT("DensityTexture"));
-		if (bVoxelizeMaterial)
-		{
-			TangentXTexture = GraphBuilder.CreateTexture(Desc, TEXT("TangentXTexture"));
-			TangentYTexture = GraphBuilder.CreateTexture(Desc, TEXT("TangentYTexture"));
-			TangentZTexture = GraphBuilder.CreateTexture(Desc, TEXT("TangentZTexture"));
-			MaterialTexture = GraphBuilder.CreateTexture(Desc, TEXT("MaterialTexture"));
-		}
-	}
-
-	const FIntRect ViewportRect(0, 0, ResolutionDim, ResolutionDim);
-	const FVector RasterDirection = FVector(1, 0, 0);
-	const FSphere SphereBound = Bounds.GetSphere();
-	const float SphereRadius = SphereBound.W * GHairVoxelizationAABBScale;
-
-	const float RadiusAtDepth1 = GStrandHairVoxelizationRasterizationScale * SphereRadius / FMath::Min(Resolution.X, Resolution.Y);
-	const bool bIsOrtho = true;
-	const bool bIsGPUDriven = false;
-	const FVector4 HairRenderInfo = PackHairRenderInfo(RadiusAtDepth1, RadiusAtDepth1, RadiusAtDepth1, 1);
-	const uint32 HairRenderInfoBits = PackHairRenderInfoBits(bIsOrtho, bIsGPUDriven);
-
-	{
-		FReversedZOrthoMatrix OrthoMatrix(SphereRadius, SphereRadius, 1.f / (2 * SphereRadius), 0);
-		FLookAtMatrix LookAt(SphereBound.Center - RasterDirection * SphereRadius, SphereBound.Center, FVector(0, 0, 1));
-
-		VoxelResources.WorldToClip = LookAt * OrthoMatrix;
-		VoxelResources.MinAABB = Bounds.GetSphere().Center - SphereRadius;
-		VoxelResources.MaxAABB = Bounds.GetSphere().Center + SphereRadius;
-	}
-
-	FHairVoxelizationRasterPassParameters* PassParameters = GraphBuilder.AllocParameters<FHairVoxelizationRasterPassParameters>();
-	PassParameters->WorldToClipMatrix = VoxelResources.WorldToClip;
-	PassParameters->VoxelMinAABB = VoxelResources.MinAABB;
-	PassParameters->VoxelMaxAABB = VoxelResources.MaxAABB;
-	PassParameters->VoxelResolution = FIntVector(ViewportRect.Width(), ViewportRect.Width(), ViewportRect.Width());
-	PassParameters->MacroGroupId = MacroGroup.MacroGroupId;
-	PassParameters->ViewportResolution = FIntPoint(ResolutionDim, ResolutionDim);
-	PassParameters->DensityTexture = GraphBuilder.CreateUAV(DensityTexture);
-	if (bVoxelizeMaterial)
-	{
-		PassParameters->TangentXTexture = GraphBuilder.CreateUAV(TangentXTexture);
-		PassParameters->TangentYTexture = GraphBuilder.CreateUAV(TangentYTexture);
-		PassParameters->TangentZTexture = GraphBuilder.CreateUAV(TangentZTexture);
-		PassParameters->MaterialTexture = GraphBuilder.CreateUAV(MaterialTexture);
-	}
-
-	AddHairVoxelizationRasterPass(
-		GraphBuilder,
-		Scene,
-		ViewInfo,
-		PrimitiveSceneInfo,
-		bVoxelizeMaterial ? EHairStrandsRasterPassType::VoxelizationMaterial : EHairStrandsRasterPassType::Voxelization,
-		ViewportRect,
-		HairRenderInfo,
-		HairRenderInfoBits,
-		RasterDirection,
-		PassParameters);
-
-
-	GraphBuilder.QueueTextureExtraction(DensityTexture, &VoxelResources.DensityTexture);
-	if (bVoxelizeMaterial)
-	{
-		GraphBuilder.QueueTextureExtraction(TangentXTexture, &VoxelResources.TangentXTexture);
-		GraphBuilder.QueueTextureExtraction(TangentYTexture, &VoxelResources.TangentYTexture);
-		GraphBuilder.QueueTextureExtraction(TangentZTexture, &VoxelResources.TangentZTexture);
-		GraphBuilder.QueueTextureExtraction(MaterialTexture, &VoxelResources.MaterialTexture);
-	}
-
-}
-
 void VoxelizeHairStrands(
 	FRHICommandListImmediate& RHICmdList, 
 	const FScene* Scene, 
@@ -1786,61 +1386,33 @@ void VoxelizeHairStrands(
 		SCOPED_DRAW_EVENT(RHICmdList, HairStrandsVoxelization);
 		SCOPED_GPU_STAT(RHICmdList, HairStrandsVoxelization);
 
-		if (GHairVirtualVoxel)
+		if (MacroGroupDatas.Datas.Num() > 0)
 		{
-			if (MacroGroupDatas.Datas.Num() > 0)
-			{
-				// Toto moves this function into the render graph. At the moment this is not possible as this functions 
-				// generates internally a non-transient constant buffer which initialized VirtualVoxelResources. This 
-				// needs to be rewritten/worked out.
-				TRefCountPtr<FPooledRDGBuffer> PageToPageIndexBuffer;
-				MacroGroupDatas.VirtualVoxelResources = AllocateVirtualVoxelResources(RHICmdList, View, MacroGroupDatas, PageToPageIndexBuffer);
-
-				FRDGBuilder GraphBuilder(RHICmdList);
-				FRDGBufferRef ClearIndArgsBuffer = IndirectVoxelPageClear(GraphBuilder, View, MacroGroupDatas.VirtualVoxelResources);
-
-				for (FHairStrandsMacroGroupData& MacroGroup : MacroGroupDatas.Datas)
-				{
-					AddVirtualVoxelizationRasterPass(GraphBuilder, Scene, &View, MacroGroupDatas.VirtualVoxelResources, MacroGroup);
-				}
-
-				if (GHairVoxelInjectOpaqueDepthEnable)
-				{
-					for (FHairStrandsMacroGroupData& MacroGroup : MacroGroupDatas.Datas)
-					{
-						AddVirtualVoxelInjectOpaquePass(GraphBuilder, View, MacroGroupDatas.VirtualVoxelResources, MacroGroup);
-					}					
-				}
+			// Toto moves this function into the render graph. At the moment this is not possible as this functions 
+			// generates internally a non-transient constant buffer which initialized VirtualVoxelResources. This 
+			// needs to be rewritten/worked out.
+			TRefCountPtr<FPooledRDGBuffer> PageToPageIndexBuffer;
+			MacroGroupDatas.VirtualVoxelResources = AllocateVirtualVoxelResources(RHICmdList, View, MacroGroupDatas, PageToPageIndexBuffer);
 
-				AddVirtualVoxelGenerateMipPass(GraphBuilder, View, MacroGroupDatas, ClearIndArgsBuffer, PageToPageIndexBuffer);
+			FRDGBuilder GraphBuilder(RHICmdList);
+			FRDGBufferRef ClearIndArgsBuffer = IndirectVoxelPageClear(GraphBuilder, View, MacroGroupDatas.VirtualVoxelResources);
 
-				GraphBuilder.Execute();
-			}
-		}
-		else
-		{
 			for (FHairStrandsMacroGroupData& MacroGroup : MacroGroupDatas.Datas)
 			{
-				FRDGBuilder GraphBuilder(RHICmdList);
-				AddVoxelizationRasterPass(
-					GraphBuilder,
-					Scene,
-					&View,
-					MacroGroup);
+				AddVirtualVoxelizationRasterPass(GraphBuilder, Scene, &View, MacroGroupDatas.VirtualVoxelResources, MacroGroup);
+			}
 
-				if (GHairVoxelInjectOpaqueDepthEnable)
+			if (GHairVoxelInjectOpaqueDepthEnable > 0)
+			{
+				for (FHairStrandsMacroGroupData& MacroGroup : MacroGroupDatas.Datas)
 				{
-					AddVoxelInjectOpaquePass(GraphBuilder, View, MacroGroup);
-					if (GHairVoxelFilterOpaqueDepthEnable)
-					{
-						AddFilterVoxelOpaqueDepthPass(GraphBuilder, View, MacroGroup);
-					}
-				}
+					AddVirtualVoxelInjectOpaquePass(GraphBuilder, View, MacroGroupDatas.VirtualVoxelResources, MacroGroup);
+				}					
+			}
 
-				AddVoxelGenerateMipPass(GraphBuilder, View, MacroGroup);
+			AddVirtualVoxelGenerateMipPass(GraphBuilder, View, MacroGroupDatas, ClearIndArgsBuffer, PageToPageIndexBuffer);
 
-				GraphBuilder.Execute();
-			}
+			GraphBuilder.Execute();
 		}
 	}
 }
diff --git a/Engine/Source/Runtime/Renderer/Private/HairStrands/HairStrandsVoxelization.h b/Engine/Source/Runtime/Renderer/Private/HairStrands/HairStrandsVoxelization.h
--- a/Engine/Source/Runtime/Renderer/Private/HairStrands/HairStrandsVoxelization.h
+++ b/Engine/Source/Runtime/Renderer/Private/HairStrands/HairStrandsVoxelization.h
@@ -10,18 +10,6 @@
 #include "RendererInterface.h"
 #include "RenderGraphResources.h"
 
-struct FHairStrandsVoxelResources
-{
-	TRefCountPtr<IPooledRenderTarget> DensityTexture;
-	TRefCountPtr<IPooledRenderTarget> TangentXTexture;
-	TRefCountPtr<IPooledRenderTarget> TangentYTexture;
-	TRefCountPtr<IPooledRenderTarget> TangentZTexture;
-	TRefCountPtr<IPooledRenderTarget> MaterialTexture;
-	FMatrix WorldToClip;
-	FVector MinAABB;
-	FVector MaxAABB;
-};
-
 struct FVirtualVoxelNodeDesc
 {
 	FVector WorldMinAABB = FVector::ZeroVector;
@@ -93,8 +81,6 @@ struct FVirtualVoxelResources
 /// Global enable/disable for hair voxelization
 bool IsHairStrandsVoxelizationEnable();
 bool IsHairStrandsForVoxelTransmittanceAndShadowEnable();
-float GetHairStrandsVoxelizationDensityScale();
-float GetHairStrandsVoxelizationDepthBiasScale();
 
 void VoxelizeHairStrands(
 	FRHICommandListImmediate& RHICmdList,
