Fix hair scene color scatter pass running with even without anything to do.

diff --git a/Engine/Plugins/Experimental/HairStrands/Source/HairStrandsCore/Private/HairStrandsRendering.cpp b/Engine/Plugins/Experimental/HairStrands/Source/HairStrandsCore/Private/HairStrandsRendering.cpp
--- a/Engine/Plugins/Experimental/HairStrands/Source/HairStrandsCore/Private/HairStrandsRendering.cpp
+++ b/Engine/Plugins/Experimental/HairStrands/Source/HairStrandsCore/Private/HairStrandsRendering.cpp
@@ -965,6 +965,7 @@ void ComputeHairStrandsInterpolation(
 			Output.HairGroupPublicData->VFInput.HairRadius = Output.VFInput.HairRadius;
 			Output.HairGroupPublicData->VFInput.HairLength = Output.VFInput.HairLength;
 			Output.HairGroupPublicData->VFInput.bUseStableRasterization = Output.VFInput.bUseStableRasterization;
+			Output.HairGroupPublicData->VFInput.bScatterSceneLighting = Output.VFInput.bScatterSceneLighting;
 			Output.HairGroupPublicData->VFInput.HairDensity = Output.VFInput.HairDensity;
 			Output.HairGroupPublicData->VFInput.LocalToWorldTransform = LocalToWorld;
 
diff --git a/Engine/Source/Runtime/Renderer/Private/HairStrands/HairStrandsCluster.cpp b/Engine/Source/Runtime/Renderer/Private/HairStrands/HairStrandsCluster.cpp
--- a/Engine/Source/Runtime/Renderer/Private/HairStrands/HairStrandsCluster.cpp
+++ b/Engine/Source/Runtime/Renderer/Private/HairStrands/HairStrandsCluster.cpp
@@ -126,6 +126,36 @@ static bool DoesGroupExists(uint32 ResourceId, uint32 GroupIndex, const FHairStr
 	return false;
 }
 
+static void InternalUpdateMacroGroup(FHairStrandsMacroGroupData& MacroGroup, int32& MaterialId, const FMeshBatchAndRelevance* MeshBatchAndRelevance, const FPrimitiveSceneProxy* Proxy)
+{
+	if (MeshBatchAndRelevance)
+	{
+		check(MeshBatchAndRelevance->Mesh);
+		check(MeshBatchAndRelevance->Mesh->Elements.Num() == 1);
+
+		const FHairGroupPublicData* HairGroupPublicData = reinterpret_cast<const FHairGroupPublicData*>(MeshBatchAndRelevance->Mesh->Elements[0].VertexFactoryUserData);
+		if (HairGroupPublicData->VFInput.bScatterSceneLighting)
+		{
+			MacroGroup.bNeedScatterSceneLighting = true;
+		}
+
+		FHairStrandsMacroGroupData::PrimitiveInfo& PrimitiveInfo = MacroGroup.PrimitivesInfos.AddZeroed_GetRef();
+		PrimitiveInfo.MeshBatchAndRelevance = *MeshBatchAndRelevance;
+		PrimitiveInfo.MaterialId = MaterialId++;
+		PrimitiveInfo.ResourceId = reinterpret_cast<uint64>(MeshBatchAndRelevance->Mesh->Elements[0].UserData);
+		PrimitiveInfo.GroupIndex = HairGroupPublicData->GetGroupIndex();
+		check(PrimitiveInfo.GroupIndex < 32); // Sanity check
+
+		const bool bAlreadyExists = DoesGroupExists(PrimitiveInfo.ResourceId, PrimitiveInfo.GroupIndex, MacroGroup.PrimitivesGroups);
+		if (!bAlreadyExists)
+		{
+			FHairStrandsMacroGroupData::PrimitiveGroup& PrimitiveGroup = MacroGroup.PrimitivesGroups.AddZeroed_GetRef();
+			PrimitiveGroup.GroupIndex = PrimitiveInfo.GroupIndex;
+			PrimitiveGroup.ResourceId = PrimitiveInfo.ResourceId;
+		}
+	}
+}
+
 FHairStrandsMacroGroupViews CreateHairStrandsMacroGroups(
 	FRHICommandListImmediate& RHICmdList,
 	const FScene* Scene,
@@ -164,29 +194,7 @@ FHairStrandsMacroGroupViews CreateHairStrandsMacroGroups(
 					if (bIntersect)
 					{
 						MacroGroup.Bounds = Union(MacroGroup.Bounds, PrimitiveBounds);
-
-						if (MeshBatchAndRelevance)
-						{
-							check(MeshBatchAndRelevance->Mesh);
-							check(MeshBatchAndRelevance->Mesh->Elements.Num() == 1);
-
-							const FHairGroupPublicData* HairGroupPublicData = reinterpret_cast<const FHairGroupPublicData*>(MeshBatchAndRelevance->Mesh->Elements[0].VertexFactoryUserData);
-
-							FHairStrandsMacroGroupData::PrimitiveInfo& PrimitiveInfo = MacroGroup.PrimitivesInfos.AddZeroed_GetRef();
-							PrimitiveInfo.MeshBatchAndRelevance = *MeshBatchAndRelevance;
-							PrimitiveInfo.MaterialId = MaterialId++;
-							PrimitiveInfo.ResourceId = reinterpret_cast<uint64>(MeshBatchAndRelevance->Mesh->Elements[0].UserData);
-							PrimitiveInfo.GroupIndex = HairGroupPublicData->GetGroupIndex();
-							check(PrimitiveInfo.GroupIndex < 32); // Sanity check
-
-							const bool bAlreadyExists = DoesGroupExists(PrimitiveInfo.ResourceId, PrimitiveInfo.GroupIndex, MacroGroup.PrimitivesGroups);
-							if (!bAlreadyExists)
-							{
-								FHairStrandsMacroGroupData::PrimitiveGroup& PrimitiveGroup = MacroGroup.PrimitivesGroups.AddZeroed_GetRef();
-								PrimitiveGroup.GroupIndex = PrimitiveInfo.GroupIndex;
-								PrimitiveGroup.ResourceId = PrimitiveInfo.ResourceId;
-							}
-						}
+						InternalUpdateMacroGroup(MacroGroup, MaterialId, MeshBatchAndRelevance, Proxy);
 						bFound = true;
 						break;
 					}
@@ -196,28 +204,7 @@ FHairStrandsMacroGroupViews CreateHairStrandsMacroGroups(
 				{
 					FHairStrandsMacroGroupData MacroGroup;
 					MacroGroup.MacroGroupId = MacroGroupId++;
-					if (MeshBatchAndRelevance)
-					{
-						check(MeshBatchAndRelevance->Mesh);
-						check(MeshBatchAndRelevance->Mesh->Elements.Num() == 1);
-
-						const FHairGroupPublicData* HairGroupPublicData  = reinterpret_cast<const FHairGroupPublicData*>(MeshBatchAndRelevance->Mesh->Elements[0].VertexFactoryUserData);
-
-						FHairStrandsMacroGroupData::PrimitiveInfo& PrimitiveInfo = MacroGroup.PrimitivesInfos.AddZeroed_GetRef();
-						PrimitiveInfo.MeshBatchAndRelevance = *MeshBatchAndRelevance;
-						PrimitiveInfo.MaterialId = MaterialId++;
-						PrimitiveInfo.ResourceId = reinterpret_cast<uint64>(MeshBatchAndRelevance->Mesh->Elements[0].UserData);
-						PrimitiveInfo.GroupIndex = HairGroupPublicData->GetGroupIndex();
-						check(PrimitiveInfo.GroupIndex < 32); // Sanity check
-
-						const bool bAlreadyExists = DoesGroupExists(PrimitiveInfo.ResourceId, PrimitiveInfo.GroupIndex, MacroGroup.PrimitivesGroups);
-						if (!bAlreadyExists)
-						{
-							FHairStrandsMacroGroupData::PrimitiveGroup& PrimitiveGroup = MacroGroup.PrimitivesGroups.AddZeroed_GetRef();
-							PrimitiveGroup.GroupIndex = PrimitiveInfo.GroupIndex;
-							PrimitiveGroup.ResourceId = PrimitiveInfo.ResourceId;
-						}
-					}
+					InternalUpdateMacroGroup(MacroGroup, MaterialId, MeshBatchAndRelevance, Proxy);
 					MacroGroup.Bounds = PrimitiveBounds;
 					MacroGroups.Datas.Add(MacroGroup);
 				}
diff --git a/Engine/Source/Runtime/Renderer/Private/HairStrands/HairStrandsCluster.h b/Engine/Source/Runtime/Renderer/Private/HairStrands/HairStrandsCluster.h
--- a/Engine/Source/Runtime/Renderer/Private/HairStrands/HairStrandsCluster.h
+++ b/Engine/Source/Runtime/Renderer/Private/HairStrands/HairStrandsCluster.h
@@ -60,6 +60,8 @@ struct FHairStrandsMacroGroupData
 	FBoxSphereBounds Bounds;
 	FIntRect ScreenRect;
 	uint32 MacroGroupId;
+
+	bool bNeedScatterSceneLighting = false;
 };
 
 /// Store all hair strands macro group infos for a given view
diff --git a/Engine/Source/Runtime/Renderer/Private/HairStrands/HairStrandsEnvironment.cpp b/Engine/Source/Runtime/Renderer/Private/HairStrands/HairStrandsEnvironment.cpp
--- a/Engine/Source/Runtime/Renderer/Private/HairStrands/HairStrandsEnvironment.cpp
+++ b/Engine/Source/Runtime/Renderer/Private/HairStrands/HairStrandsEnvironment.cpp
@@ -375,7 +375,7 @@ static void AddHairStrandsEnvironmentLightingPassPS(
 	ParametersPS->RenderTargets[0] = FRenderTargetBinding(GraphBuilder.RegisterExternalTexture(VisibilityData.SampleLightingBuffer), ERenderTargetLoadAction::ELoad);
 
 	GraphBuilder.AddPass(
-		RDG_EVENT_NAME("HairEnvLightingPS"),
+		bUseSceneColor ? RDG_EVENT_NAME("HairEnvSceneScatterPS") : RDG_EVENT_NAME("HairEnvLightingPS"),
 		ParametersPS,
 		ERDGPassFlags::Raster,
 		[ParametersPS, VertexShader, PixelShader, ViewportResolution, CapturedView](FRHICommandList& RHICmdList)
@@ -428,15 +428,31 @@ void RenderHairStrandsSceneColorScattering(
 
 		const FViewInfo& View = Views[ViewIndex];
 		const FHairStrandsMacroGroupDatas& MacroGroupDatas = HairDatas->MacroGroupsPerViews.Views[ViewIndex];
-		FRDGBuilder GraphBuilder(RHICmdList);
 
-		FSceneRenderTargets& SceneContext = FSceneRenderTargets::Get(RHICmdList);
-		FRDGTextureRef SceneColorTexture = GraphBuilder.RegisterExternalTexture(SceneContext.GetSceneColor());
+		bool bNeedScatterSceneLighting = false;
 		for (const FHairStrandsMacroGroupData& MacroGroupData : HairDatas->MacroGroupsPerViews.Views[ViewIndex].Datas)
 		{
-			AddHairStrandsEnvironmentLightingPassPS(GraphBuilder, View, VisibilityData, MacroGroupDatas, MacroGroupData, SceneColorTexture);
+			if (MacroGroupData.bNeedScatterSceneLighting)
+			{
+				bNeedScatterSceneLighting = true;
+				break;
+			}
+		}
+
+		if (bNeedScatterSceneLighting)
+		{
+			FRDGBuilder GraphBuilder(RHICmdList);
+			FSceneRenderTargets& SceneContext = FSceneRenderTargets::Get(RHICmdList);
+			FRDGTextureRef SceneColorTexture = GraphBuilder.RegisterExternalTexture(SceneContext.GetSceneColor());
+			for (const FHairStrandsMacroGroupData& MacroGroupData : HairDatas->MacroGroupsPerViews.Views[ViewIndex].Datas)
+			{
+				if (MacroGroupData.bNeedScatterSceneLighting)
+				{
+					AddHairStrandsEnvironmentLightingPassPS(GraphBuilder, View, VisibilityData, MacroGroupDatas, MacroGroupData, SceneColorTexture);
+				}
+			}
+			GraphBuilder.Execute();
 		}
-		GraphBuilder.Execute();
 	}
 }
 
diff --git a/Engine/Source/Runtime/Renderer/Public/HairStrandsInterface.h b/Engine/Source/Runtime/Renderer/Public/HairStrandsInterface.h
--- a/Engine/Source/Runtime/Renderer/Public/HairStrandsInterface.h
+++ b/Engine/Source/Runtime/Renderer/Public/HairStrandsInterface.h
@@ -94,6 +94,7 @@ public:
 		float HairLength = 0;
 		float HairDensity = 0;
 		bool bUseStableRasterization = false;
+		bool bScatterSceneLighting = false;
 		FTransform LocalToWorldTransform;
 	};
 	VertexFactoryInput VFInput;
