Fixed: Sample scene depth values from SceneColor alpha for mobile scene capture, as we dont resolve scene depth texture

diff --git a/Engine/Shaders/Private/MobileSceneCapture.usf b/Engine/Shaders/Private/MobileSceneCapture.usf
--- a/Engine/Shaders/Private/MobileSceneCapture.usf
+++ b/Engine/Shaders/Private/MobileSceneCapture.usf
@@ -33,6 +33,7 @@ void MainCopyPS(
     float4 EncodedColor = Texture2DSample(InTexture, InTextureSampler, Input.UV);
 	OutColor = EncodedColor;
 
+	// will sample SceneColor.A
 	float SceneW = CalcSceneDepth(Input.UV);
 
 #if SOURCE_MODE_SCENE_COLOR_AND_OPACITY
diff --git a/Engine/Source/Runtime/Renderer/Private/MobileSceneCaptureRendering.cpp b/Engine/Source/Runtime/Renderer/Private/MobileSceneCaptureRendering.cpp
--- a/Engine/Source/Runtime/Renderer/Private/MobileSceneCaptureRendering.cpp
+++ b/Engine/Source/Runtime/Renderer/Private/MobileSceneCaptureRendering.cpp
@@ -55,7 +55,7 @@ public:
 
 	static void ModifyCompilationEnvironment(const FGlobalShaderPermutationParameters& Parameters, FShaderCompilerEnvironment& OutEnvironment)
 	{
-		OutEnvironment.SetDefine(TEXT("MOBILE_FORCE_DEPTH_TEXTURE_READS"), 1u);
+		OutEnvironment.SetDefine(TEXT("POST_PROCESS_MATERIAL"), 1u); // this will force reading depth from a SceneColor.A
 		const TCHAR* DefineName = GShaderSourceModeDefineName[CaptureSource];
 		if (DefineName)
 		{
