Release uniformbuffers when the world is cleaned up. this releases a lot of rendertarget references.

diff --git a/Engine/Source/Runtime/Engine/Public/SceneInterface.h b/Engine/Source/Runtime/Engine/Public/SceneInterface.h
--- a/Engine/Source/Runtime/Engine/Public/SceneInterface.h
+++ b/Engine/Source/Runtime/Engine/Public/SceneInterface.h
@@ -349,6 +349,9 @@ public:
 		return NULL;
 	}
 
+	virtual void OnWorldCleanup()
+	{
+	}
 	virtual void UpdateSceneSettings(AWorldSettings* WorldSettings) {}
 
 	/**
diff --git a/Engine/Source/Runtime/Renderer/Private/Renderer.cpp b/Engine/Source/Runtime/Renderer/Private/Renderer.cpp
--- a/Engine/Source/Runtime/Renderer/Private/Renderer.cpp
+++ b/Engine/Source/Runtime/Renderer/Private/Renderer.cpp
@@ -73,14 +73,16 @@ void FRendererModule::ReallocateSceneRenderTargets()
 
 void FRendererModule::OnWorldCleanup(UWorld* World, bool bSessionEnded, bool bCleanupResources)
 {
-	if (bFlushRenderTargetsOnWorldCleanup > 0)
+	FSceneInterface* Scene = World->Scene;
+	ENQUEUE_RENDER_COMMAND(OnWorldCleanup)(
+	[Scene](FRHICommandListImmediate& RHICmdList)
 	{
-		ENQUEUE_RENDER_COMMAND(OnWorldCleanup)(
-			[](FRHICommandListImmediate& RHICmdList)
+		if(bFlushRenderTargetsOnWorldCleanup > 0)
 		{
 			GRenderTargetPool.FreeUnusedResources();
-		});
-	}
+		}
+		Scene->OnWorldCleanup();
+	});
 }
 
 void FRendererModule::SceneRenderTargetsSetBufferSize(uint32 SizeX, uint32 SizeY)
diff --git a/Engine/Source/Runtime/Renderer/Private/RendererScene.cpp b/Engine/Source/Runtime/Renderer/Private/RendererScene.cpp
--- a/Engine/Source/Runtime/Renderer/Private/RendererScene.cpp
+++ b/Engine/Source/Runtime/Renderer/Private/RendererScene.cpp
@@ -609,6 +609,11 @@ SIZE_T FScene::GetSizeBytes() const
 		+ PrimitiveOctree.GetSizeBytes();
 }
 
+void FScene::OnWorldCleanup()
+{
+	UniformBuffers.Clear();
+}
+
 void FScene::CheckPrimitiveArrays(int MaxTypeOffsetIndex)
 {
 	check(Primitives.Num() == PrimitiveTransforms.Num());
@@ -885,6 +890,45 @@ FORCEINLINE static void VerifyProperPIEScene(UPrimitiveComponent* Component, UWo
 #endif
 }
 
+void FPersistentUniformBuffers::Clear()
+{
+	ViewUniformBuffer.SafeRelease();
+	InstancedViewUniformBuffer.SafeRelease();
+	DepthPassUniformBuffer.SafeRelease();
+	OpaqueBasePassUniformBuffer.SafeRelease();
+	TranslucentBasePassUniformBuffer.SafeRelease();
+	ReflectionCaptureUniformBuffer.SafeRelease();
+	CSMShadowDepthViewUniformBuffer.SafeRelease();
+	CSMShadowDepthPassUniformBuffer.SafeRelease();
+	DistortionPassUniformBuffer.SafeRelease();
+	VelocityPassUniformBuffer.SafeRelease();
+	HitProxyPassUniformBuffer.SafeRelease();
+	MeshDecalPassUniformBuffer.SafeRelease();
+	LightmapDensityPassUniformBuffer.SafeRelease();
+	DebugViewModePassUniformBuffer.SafeRelease();
+	VoxelizeVolumePassUniformBuffer.SafeRelease();
+	VoxelizeVolumeViewUniformBuffer.SafeRelease();
+	ConvertToUniformMeshPassUniformBuffer.SafeRelease();
+	CustomDepthPassUniformBuffer.SafeRelease();
+	MobileCustomDepthPassUniformBuffer.SafeRelease();
+	CustomDepthViewUniformBuffer.SafeRelease();
+	InstancedCustomDepthViewUniformBuffer.SafeRelease();
+	VirtualTextureViewUniformBuffer.SafeRelease();
+	MobileOpaqueBasePassUniformBuffer.SafeRelease();
+	MobileTranslucentBasePassUniformBuffer.SafeRelease();
+	MobileCSMShadowDepthPassUniformBuffer.SafeRelease();
+	MobileDistortionPassUniformBuffer.SafeRelease();
+
+
+	for (auto& UniformBuffer : MobileDirectionalLightUniformBuffers)
+	{
+		UniformBuffer.SafeRelease();
+	}
+	MobileSkyReflectionUniformBuffer.SafeRelease();
+#if WITH_EDITOR
+	EditorSelectionPassUniformBuffer.SafeRelease();
+#endif
+}
 void FPersistentUniformBuffers::Initialize()
 {
 	FViewUniformShaderParameters ViewUniformBufferParameters;
diff --git a/Engine/Source/Runtime/Renderer/Private/ScenePrivate.h b/Engine/Source/Runtime/Renderer/Private/ScenePrivate.h
--- a/Engine/Source/Runtime/Renderer/Private/ScenePrivate.h
+++ b/Engine/Source/Runtime/Renderer/Private/ScenePrivate.h
@@ -2416,6 +2416,7 @@ public:
 	}
 
 	void Initialize();
+	void Clear();
 
 	/** Compares the provided view against the cached view and updates the view uniform buffer
 	 *  if the views differ. Returns whether uniform buffer was updated.
@@ -2872,6 +2873,8 @@ public:
 	{
 		return this;
 	}
+	virtual void OnWorldCleanup() override;
+
 
 	virtual void UpdateSceneSettings(AWorldSettings* WorldSettings) override;
 
