Fix crash on shutdown

diff --git a/Engine/Source/Runtime/Renderer/Private/Renderer.cpp b/Engine/Source/Runtime/Renderer/Private/Renderer.cpp
--- a/Engine/Source/Runtime/Renderer/Private/Renderer.cpp
+++ b/Engine/Source/Runtime/Renderer/Private/Renderer.cpp
@@ -81,8 +81,12 @@ void FRendererModule::OnWorldCleanup(UWorld* World, bool bSessionEnded, bool bCl
 		{
 			GRenderTargetPool.FreeUnusedResources();
 		}
-		Scene->OnWorldCleanup();
+		if(Scene)
+		{
+			Scene->OnWorldCleanup();
+		}
 	});
+
 }
 
 void FRendererModule::SceneRenderTargetsSetBufferSize(uint32 SizeX, uint32 SizeY)
