Add hair & skeletal mesh asset name to hair debug info.

diff --git a/Engine/Plugins/Experimental/HairStrands/Source/HairStrandsCore/Private/GroomComponent.cpp b/Engine/Plugins/Experimental/HairStrands/Source/HairStrandsCore/Private/GroomComponent.cpp
--- a/Engine/Plugins/Experimental/HairStrands/Source/HairStrandsCore/Private/GroomComponent.cpp
+++ b/Engine/Plugins/Experimental/HairStrands/Source/HairStrandsCore/Private/GroomComponent.cpp
@@ -1095,6 +1095,7 @@ void UGroomComponent::InitResources(bool bIsBindingReloading)
 	Interpolation.Function = ComputeHairStrandsInterpolation;
 	Interpolation.ResetFunction = ResetHairStrandsInterpolation;
 
+	UGroomAsset* LocalGroomAsset = GroomAsset;
 	USkeletalMesh* InSourceSkeletalMesh = SourceSkeletalMesh;
 	const bool bRunMeshProjection = bHasValidSketalMesh && (!BindingAsset || !bIsBindingCompatible);
 	FHairGroupResources* LocalResources = HairGroupResources;
@@ -1110,7 +1111,8 @@ void UGroomComponent::InitResources(bool bIsBindingReloading)
 			DebugGroupInfo,
 			bRunMeshProjection,
 			SkeletalMeshComponent,
-			InSourceSkeletalMesh
+			InSourceSkeletalMesh,
+			LocalGroomAsset
 		]
 		(FRHICommandListImmediate& RHICmdList)
 	{
@@ -1176,6 +1178,14 @@ void UGroomComponent::InitResources(bool bIsBindingReloading)
 		}
 
 		FHairStrandsProjectionDebugInfo HairProjectionDebugInfo;
+		#if WITH_EDITOR
+		HairProjectionDebugInfo.GroomAssetName = LocalGroomAsset->GetName();
+		if (SkeletalMeshComponent)
+		{
+			HairProjectionDebugInfo.SkeletalComponentName = SkeletalMeshComponent->GetPathName();
+		}
+		#endif
+
 		if (bRunMeshProjection)
 		{				
 			FSkeletalMeshRenderData* TargetRenderData = SkeletalMeshComponent->GetSkeletalMeshRenderData();
diff --git a/Engine/Source/Runtime/Renderer/Private/HairStrands/HairStrandsDebug.cpp b/Engine/Source/Runtime/Renderer/Private/HairStrands/HairStrandsDebug.cpp
--- a/Engine/Source/Runtime/Renderer/Private/HairStrands/HairStrandsDebug.cpp
+++ b/Engine/Source/Runtime/Renderer/Private/HairStrands/HairStrandsDebug.cpp
@@ -1319,8 +1319,9 @@ void RenderHairStrandsDebugInfo(
 
 		float X = 20;
 		float Y = ClusterY;
-		FLinearColor InactiveColor(0.5, 0.5, 0.5);
-		FLinearColor DebugColor(1, 1, 0);
+		const FLinearColor InactiveColor(0.5, 0.5, 0.5);
+		const FLinearColor DebugColor(1, 1, 0);
+		const FLinearColor DebugGroupColor(0.5f, 0, 0);
 		FString Line;
 
 		const FHairStrandsDebugInfos DebugInfos = GetHairStandsDebugInfos();
@@ -1336,7 +1337,7 @@ void RenderHairStrandsDebugInfo(
 			check(ViewFamily.Scene && ViewFamily.Scene->GetWorld());
 			const bool bIsActive = DebugInfo.WorldType == ViewFamily.Scene->GetWorld()->WorldType;
 
-			Line = FString::Printf(TEXT(" * Id:%d | WorldType:%s | Group count : %d"), DebugInfo.ComponentId, ToString(DebugInfo.WorldType), DebugInfo.HairGroups.Num());
+			Line = FString::Printf(TEXT(" * Id:%d | WorldType:%s | Group count : %d | Asset : %s | Skeletal : %s "), DebugInfo.ComponentId, ToString(DebugInfo.WorldType), DebugInfo.HairGroups.Num(), *DebugInfo.GroomAssetName, *DebugInfo.SkeletalComponentName);
 			Canvas.DrawShadowedString(X, Y += YStep, *Line, GetStatsFont(), bIsActive ? DebugColor : InactiveColor);
 
 			for (const FHairStrandsDebugInfo::HairGroup& DebugHairGroup : DebugInfo.HairGroups)
@@ -1350,7 +1351,7 @@ void RenderHairStrandsDebugInfo(
 					DebugHairGroup.bHasBinding ? TEXT("True") : TEXT("False"),
 					DebugHairGroup.bHasSimulation ? TEXT("True") : TEXT("False"),
 					DebugHairGroup.LODCount);
-				Canvas.DrawShadowedString(X, Y += YStep, *Line, GetStatsFont(), bIsActive ? DebugColor : InactiveColor);
+				Canvas.DrawShadowedString(X, Y += YStep, *Line, GetStatsFont(), bIsActive ? DebugGroupColor : InactiveColor);
 			}
 		}
 
diff --git a/Engine/Source/Runtime/Renderer/Private/HairStrands/HairStrandsInterface.cpp b/Engine/Source/Runtime/Renderer/Private/HairStrands/HairStrandsInterface.cpp
--- a/Engine/Source/Runtime/Renderer/Private/HairStrands/HairStrandsInterface.cpp
+++ b/Engine/Source/Runtime/Renderer/Private/HairStrands/HairStrandsInterface.cpp
@@ -518,6 +518,8 @@ FHairStrandsDebugInfos GetHairStandsDebugInfos()
 		Info = E.DebugInfo;
 		Info.ComponentId = E.ComponentId;
 		Info.WorldType = E.WorldType;
+		Info.GroomAssetName = E.DebugProjectionInfo.GroomAssetName;
+		Info.SkeletalComponentName = E.DebugProjectionInfo.SkeletalComponentName;
 
 		const uint32 GroupCount = Info.HairGroups.Num();
 		for (uint32 GroupIt=0; GroupIt < GroupCount; ++GroupIt)
diff --git a/Engine/Source/Runtime/Renderer/Public/HairStrandsInterface.h b/Engine/Source/Runtime/Renderer/Public/HairStrandsInterface.h
--- a/Engine/Source/Runtime/Renderer/Public/HairStrandsInterface.h
+++ b/Engine/Source/Runtime/Renderer/Public/HairStrandsInterface.h
@@ -283,6 +283,8 @@ struct FHairStrandsDebugInfo
 {
 	uint32 ComponentId = 0;
 	EWorldType::Type WorldType = EWorldType::None;
+	FString GroomAssetName;
+	FString SkeletalComponentName;
 
 	struct HairGroup
 	{
@@ -319,6 +321,9 @@ struct FHairStrandsProjectionDebugInfo
 	FHairStrandsProjectionMeshData SourceMeshData;
 	FHairStrandsProjectionMeshData TargetMeshData;
 	TArray<FRWBuffer> TransferredPositions;
+	FString GroomAssetName;
+	FString SkeletalComponentName;
+
 };
 
 ////////////////////////////////////////////////////////////////////////////////////////////////////
