Add missing buffer transition after interpolatino causing hair flickering.

diff --git a/Engine/Source/Runtime/Engine/Private/GPUSkinCache.cpp b/Engine/Source/Runtime/Engine/Private/GPUSkinCache.cpp
--- a/Engine/Source/Runtime/Engine/Private/GPUSkinCache.cpp
+++ b/Engine/Source/Runtime/Engine/Private/GPUSkinCache.cpp
@@ -667,11 +667,11 @@ void FGPUSkinCache::Cleanup()
 	ensure(Allocations.Num() == 0);
 }
 
-void FGPUSkinCache::TransitionAllToReadable(FRHICommandList& RHICmdList)
+void FGPUSkinCache::TransitionAllToReadable(FRHICommandList& RHICmdList, EResourceTransitionPipeline Pipeline)
 {
 	if (BuffersToTransition.Num() > 0)
 	{
-		RHICmdList.TransitionResources(EResourceTransitionAccess::EReadable, EResourceTransitionPipeline::EComputeToGfx, BuffersToTransition.GetData(), BuffersToTransition.Num());
+		RHICmdList.TransitionResources(EResourceTransitionAccess::EReadable, Pipeline, BuffersToTransition.GetData(), BuffersToTransition.Num());
 		BuffersToTransition.SetNum(0, false);
 	}
 }
diff --git a/Engine/Source/Runtime/Engine/Public/GPUSkinCache.h b/Engine/Source/Runtime/Engine/Public/GPUSkinCache.h
--- a/Engine/Source/Runtime/Engine/Public/GPUSkinCache.h
+++ b/Engine/Source/Runtime/Engine/Public/GPUSkinCache.h
@@ -331,7 +331,7 @@ public:
 		const FVertexBufferAndSRV* BoneBuffers[NUM_BUFFERS];
 	};
 
-	ENGINE_API void TransitionAllToReadable(FRHICommandList& RHICmdList);
+	ENGINE_API void TransitionAllToReadable(FRHICommandList& RHICmdList, EResourceTransitionPipeline Pipeline = EResourceTransitionPipeline::EComputeToGfx);
 
 #if RHI_RAYTRACING
 	void AddRayTracingGeometryToUpdate(FRayTracingGeometry* RayTracingGeometry)
diff --git a/Engine/Source/Runtime/Renderer/Private/HairStrands/HairStrandsInterface.cpp b/Engine/Source/Runtime/Renderer/Private/HairStrands/HairStrandsInterface.cpp
--- a/Engine/Source/Runtime/Renderer/Private/HairStrands/HairStrandsInterface.cpp
+++ b/Engine/Source/Runtime/Renderer/Private/HairStrands/HairStrandsInterface.cpp
@@ -267,6 +267,7 @@ void RunMeshTransfer(
 	const ERHIFeatureLevel::Type FeatureLevel = GMaxRHIFeatureLevel;
 	FGlobalShaderMap* ShaderMap = GetGlobalShaderMap(FeatureLevel);
 
+	FBufferTransitionQueue TransitionQueue;
 	FRDGBuilder GraphBuilder(RHICmdList);
 
 	const uint32 LODCount = TargetMeshData.LODs.Num();
@@ -276,10 +277,11 @@ void RunMeshTransfer(
 		check(TargetMeshData.LODs[LODIndex].Sections.Num() > 0);
 
 		OutTransferedPositions[LODIndex].Initialize(sizeof(float), TargetMeshData.LODs[LODIndex].Sections[0].TotalVertexCount * 3, PF_R32_FLOAT);
-		TransferMesh(GraphBuilder, ShaderMap, LODIndex, SourceMeshData, TargetMeshData, OutTransferedPositions[LODIndex]);
+		TransferMesh(GraphBuilder, ShaderMap, LODIndex, SourceMeshData, TargetMeshData, OutTransferedPositions[LODIndex], TransitionQueue);
 	}
 	
 	GraphBuilder.Execute();	
+	TransitBufferToReadable(RHICmdList, TransitionQueue);
 }
 
 void RunProjection(
@@ -292,7 +294,9 @@ void RunProjection(
 	const ERHIFeatureLevel::Type FeatureLevel = GMaxRHIFeatureLevel;
 	FGlobalShaderMap* ShaderMap = GetGlobalShaderMap(FeatureLevel);
 
-	auto Project = [&RHICmdList, ShaderMap, &TargetMeshData, LocalToWorld](FHairStrandsProjectionHairData& ProjectionHairData)
+	FBufferTransitionQueue TransitionQueue;
+
+	auto Project = [&RHICmdList, ShaderMap, &TargetMeshData, LocalToWorld, &TransitionQueue](FHairStrandsProjectionHairData& ProjectionHairData)
 	{
 		FRDGBuilder GraphBuilder(RHICmdList);
 		for (FHairStrandsProjectionHairData::HairGroup& HairGroup : ProjectionHairData.HairGroups)
@@ -301,8 +305,8 @@ void RunProjection(
 			{
 				const uint32 LODIndex = LODData.LODIndex;
 				HairGroup.LocalToWorld = LocalToWorld;
-				ProjectHairStrandsOntoMesh(GraphBuilder, ShaderMap, LODIndex, TargetMeshData, HairGroup);
-				UpdateHairStrandsMeshTriangles(GraphBuilder, ShaderMap, LODIndex, HairStrandsTriangleType::RestPose, TargetMeshData.LODs[LODIndex], HairGroup);
+				ProjectHairStrandsOntoMesh(GraphBuilder, ShaderMap, LODIndex, TargetMeshData, HairGroup, TransitionQueue);
+				UpdateHairStrandsMeshTriangles(GraphBuilder, ShaderMap, LODIndex, HairStrandsTriangleType::RestPose, TargetMeshData.LODs[LODIndex], HairGroup, TransitionQueue);
 			}
 		}
 		GraphBuilder.Execute();
@@ -310,6 +314,8 @@ void RunProjection(
 
 	Project(RenProjectionHairData);
 	Project(SimProjectionHairData);
+
+	TransitBufferToReadable(RHICmdList, TransitionQueue);
 }
 
 void RunHairStrandsInterpolation(
@@ -352,13 +358,15 @@ void RunHairStrandsInterpolation(
 			MeshDataLOD.Sections.Add(ConvertMeshSection(Section, E.SkeletalLocalToWorld));
 		}
 
+		FBufferTransitionQueue TransitionQueue;
+
 		FRDGBuilder GraphBuilder(RHICmdList);
 
 		for (FHairStrandsProjectionHairData::HairGroup& ProjectionHairData : E.RenProjectionHairDatas.HairGroups)
 		{
 			if (EHairStrandsInterpolationType::RenderStrands == Type && 0 <= E.FrameLODIndex && E.FrameLODIndex < ProjectionHairData.LODDatas.Num() && ProjectionHairData.LODDatas[E.FrameLODIndex].bIsValid)
 			{
-				UpdateHairStrandsMeshTriangles(GraphBuilder, ShaderMap, E.FrameLODIndex, HairStrandsTriangleType::DeformedPose, MeshDataLOD, ProjectionHairData);
+				UpdateHairStrandsMeshTriangles(GraphBuilder, ShaderMap, E.FrameLODIndex, HairStrandsTriangleType::DeformedPose, MeshDataLOD, ProjectionHairData, TransitionQueue);
 			}
 		}
 
@@ -366,7 +374,7 @@ void RunHairStrandsInterpolation(
 		{
 			if (EHairStrandsInterpolationType::SimulationStrands == Type && 0 <= E.FrameLODIndex && E.FrameLODIndex < ProjectionHairData.LODDatas.Num() && ProjectionHairData.LODDatas[E.FrameLODIndex].bIsValid)
 			{
-				UpdateHairStrandsMeshTriangles(GraphBuilder, ShaderMap, E.FrameLODIndex, HairStrandsTriangleType::DeformedPose, MeshDataLOD, ProjectionHairData);
+				UpdateHairStrandsMeshTriangles(GraphBuilder, ShaderMap, E.FrameLODIndex, HairStrandsTriangleType::DeformedPose, MeshDataLOD, ProjectionHairData, TransitionQueue);
 			}
 		}
 
@@ -374,8 +382,8 @@ void RunHairStrandsInterpolation(
 		{
 			if (EHairStrandsInterpolationType::SimulationStrands == Type && 0 <= E.FrameLODIndex && E.FrameLODIndex < ProjectionHairData.LODDatas.Num() && ProjectionHairData.LODDatas[E.FrameLODIndex].bIsValid)
 			{
-				InitHairStrandsMeshSamples(GraphBuilder, ShaderMap, E.FrameLODIndex, HairStrandsTriangleType::DeformedPose, MeshDataLOD, ProjectionHairData);
-				UpdateHairStrandsMeshSamples(GraphBuilder, ShaderMap, E.FrameLODIndex, MeshDataLOD, ProjectionHairData);
+				InitHairStrandsMeshSamples(GraphBuilder, ShaderMap, E.FrameLODIndex, HairStrandsTriangleType::DeformedPose, MeshDataLOD, ProjectionHairData, TransitionQueue);
+				UpdateHairStrandsMeshSamples(GraphBuilder, ShaderMap, E.FrameLODIndex, MeshDataLOD, ProjectionHairData, TransitionQueue);
 				//InterpolateHairStrandsMeshTriangles(GraphBuilder, ShaderMap, E.FrameLODIndex, MeshDataLOD, ProjectionHairData);
 			}
 		}
@@ -390,6 +398,7 @@ void RunHairStrandsInterpolation(
 			}
 		}*/
 		GraphBuilder.Execute();
+		TransitBufferToReadable(RHICmdList, TransitionQueue);
 	}
 
 	// Reset deformation
@@ -713,3 +722,12 @@ bool HasHairStrandsProcess(EShaderPlatform Platform)
 		GHairStrandsRenderingEnable == 1 &&
 		(!GBindingQueries.IsEmpty() || !GFollicleQueries.IsEmpty());
 }
+
+void TransitBufferToReadable(FRHICommandListImmediate& RHICmdList, FBufferTransitionQueue& BuffersToTransit)
+{
+	if (BuffersToTransit.Num())
+	{
+		RHICmdList.TransitionResources(EResourceTransitionAccess::EReadable, EResourceTransitionPipeline::EComputeToCompute, BuffersToTransit.GetData(), BuffersToTransit.Num());
+		BuffersToTransit.SetNum(0, false);
+	}
+}
\ No newline at end of file
diff --git a/Engine/Source/Runtime/Renderer/Private/HairStrands/HairStrandsMeshProjection.cpp b/Engine/Source/Runtime/Renderer/Private/HairStrands/HairStrandsMeshProjection.cpp
--- a/Engine/Source/Runtime/Renderer/Private/HairStrands/HairStrandsMeshProjection.cpp
+++ b/Engine/Source/Runtime/Renderer/Private/HairStrands/HairStrandsMeshProjection.cpp
@@ -123,7 +123,8 @@ static void AddMeshTransferPass(
 	const FHairStrandsProjectionMeshData::Section& SourceSectionData,
 	const FHairStrandsProjectionMeshData::Section& TargetSectionData,
 	FRDGBufferRef VertexSectionId, 
-	FRWBuffer& OutTargetRestPosition)
+	FRWBuffer& OutTargetRestPosition,
+	FBufferTransitionQueue& OutTransitionQueue)
 {
 	if (!SourceSectionData.IndexBuffer ||
 		!SourceSectionData.PositionBuffer ||
@@ -190,6 +191,8 @@ static void AddMeshTransferPass(
 			Parameters,
 			DispatchGroupCount);
 		bClear = false;
+
+		OutTransitionQueue.Add(OutTargetRestPosition.UAV);
 	}
 }
 
@@ -240,7 +243,8 @@ static void AddHairStrandMeshProjectionPass(
 	const int32 LODIndex,
 	const FHairStrandsProjectionMeshData::Section& MeshSectionData,
 	const FHairStrandsProjectionHairData::HairGroup& RootData,
-	FRDGBufferRef RootDistanceBuffer)
+	FRDGBufferRef RootDistanceBuffer,
+	FBufferTransitionQueue& OutTransitionQueue)
 {
 	if (!RootData.RootPositionBuffer ||
 		!RootData.RootNormalBuffer ||
@@ -308,6 +312,9 @@ static void AddHairStrandMeshProjectionPass(
 			ComputeShader,
 			Parameters,
 			DispatchGroupCount);
+
+		OutTransitionQueue.Add(RootData.LODDatas[LODIndex].RootTriangleIndexBuffer->UAV);
+		OutTransitionQueue.Add(RootData.LODDatas[LODIndex].RootTriangleBarycentricBuffer->UAV);
 	}
 }
 
@@ -316,7 +323,8 @@ void ProjectHairStrandsOntoMesh(
 	FGlobalShaderMap* ShaderMap,
 	const int32 LODIndex,
 	const FHairStrandsProjectionMeshData& ProjectionMeshData,
-	FHairStrandsProjectionHairData::HairGroup& ProjectionHairData)
+	FHairStrandsProjectionHairData::HairGroup& ProjectionHairData, 
+	FBufferTransitionQueue& TransitionQueue)
 {
 	if (LODIndex < 0 || LODIndex >= ProjectionHairData.LODDatas.Num())
 		return;
@@ -327,7 +335,7 @@ void ProjectHairStrandsOntoMesh(
 	for (const FHairStrandsProjectionMeshData::Section& MeshSection : ProjectionMeshData.LODs[LODIndex].Sections)
 	{
 		check(ProjectionHairData.LODDatas[LODIndex].LODIndex == LODIndex);
-		AddHairStrandMeshProjectionPass(GraphBuilder, ShaderMap, ClearDistance, LODIndex, MeshSection, ProjectionHairData, RootDistanceBuffer);
+		AddHairStrandMeshProjectionPass(GraphBuilder, ShaderMap, ClearDistance, LODIndex, MeshSection, ProjectionHairData, RootDistanceBuffer, TransitionQueue);
 		ProjectionHairData.LODDatas[LODIndex].bIsValid = true;
 		ClearDistance = false;
 	}
@@ -339,7 +347,8 @@ void TransferMesh(
 	const int32 LODIndex,
 	const FHairStrandsProjectionMeshData& SourceMeshData,
 	const FHairStrandsProjectionMeshData& TargetMeshData,
-	FRWBuffer& OutPositionBuffer)
+	FRWBuffer& OutPositionBuffer,
+	FBufferTransitionQueue& OutTransitionQueue)
 {
 	if (LODIndex < 0)
 		return;
@@ -359,7 +368,7 @@ void TransferMesh(
 	FRDGBufferRef VertexSectionId = AddMeshSectionId(GraphBuilder, ShaderMap, TargetMeshData.LODs[TargetLODIndex]);
 	const FHairStrandsProjectionMeshData::Section& SourceMeshSection = SourceMeshData.LODs[SourceLODIndex].Sections[SourceSectionIndex];
 	const FHairStrandsProjectionMeshData::Section& TargetMeshSection = TargetMeshData.LODs[TargetLODIndex].Sections[TargetSectionIndex];
-	AddMeshTransferPass(GraphBuilder, ShaderMap, true, SourceMeshSection, TargetMeshSection, VertexSectionId, OutPositionBuffer);
+	AddMeshTransferPass(GraphBuilder, ShaderMap, true, SourceMeshSection, TargetMeshSection, VertexSectionId, OutPositionBuffer, OutTransitionQueue);
 }
 
 ///////////////////////////////////////////////////////////////////////////////////////////////////
@@ -437,7 +446,8 @@ static void AddHairStrandUpdateMeshTrianglesPass(
 	const int32 LODIndex,
 	const HairStrandsTriangleType Type,
 	const FHairStrandsProjectionMeshData::LOD& MeshData,
-	FHairStrandsProjectionHairData::HairGroup& RootData)
+	FHairStrandsProjectionHairData::HairGroup& RootData, 
+	FBufferTransitionQueue& OutTransitionQueue)
 {
 	if (RootData.RootCount == 0 || LODIndex < 0 || LODIndex >= RootData.LODDatas.Num())
 	{
@@ -577,6 +587,10 @@ static void AddHairStrandUpdateMeshTrianglesPass(
 		Parameters->OutRootTrianglePosition0 = LODData.RestRootTrianglePosition0Buffer->UAV;
 		Parameters->OutRootTrianglePosition1 = LODData.RestRootTrianglePosition1Buffer->UAV;
 		Parameters->OutRootTrianglePosition2 = LODData.RestRootTrianglePosition2Buffer->UAV;
+
+		OutTransitionQueue.Add(LODData.RestRootTrianglePosition0Buffer->UAV);
+		OutTransitionQueue.Add(LODData.RestRootTrianglePosition1Buffer->UAV);
+		OutTransitionQueue.Add(LODData.RestRootTrianglePosition2Buffer->UAV);
 	}
 	else if (Type == HairStrandsTriangleType::DeformedPose)
 	{
@@ -584,6 +598,10 @@ static void AddHairStrandUpdateMeshTrianglesPass(
 		Parameters->OutRootTrianglePosition1 = LODData.DeformedRootTrianglePosition1Buffer->UAV;
 		Parameters->OutRootTrianglePosition2 = LODData.DeformedRootTrianglePosition2Buffer->UAV;
 		if (LODData.Status) (*LODData.Status) = FHairStrandsProjectionHairData::LODData::EStatus::Completed;
+
+		OutTransitionQueue.Add(LODData.DeformedRootTrianglePosition0Buffer->UAV);
+		OutTransitionQueue.Add(LODData.DeformedRootTrianglePosition1Buffer->UAV);
+		OutTransitionQueue.Add(LODData.DeformedRootTrianglePosition2Buffer->UAV);
 	}
 	else
 	{
@@ -611,9 +629,10 @@ void UpdateHairStrandsMeshTriangles(
 	const int32 LODIndex,
 	const HairStrandsTriangleType Type,
 	const FHairStrandsProjectionMeshData::LOD& ProjectionMeshData,
-	FHairStrandsProjectionHairData::HairGroup& ProjectionHairData)
+	FHairStrandsProjectionHairData::HairGroup& ProjectionHairData,
+	FBufferTransitionQueue& OutTransitionQueue)
 {
-	AddHairStrandUpdateMeshTrianglesPass(GraphBuilder, ShaderMap, LODIndex, Type, ProjectionMeshData, ProjectionHairData);
+	AddHairStrandUpdateMeshTrianglesPass(GraphBuilder, ShaderMap, LODIndex, Type, ProjectionMeshData, ProjectionHairData,OutTransitionQueue);
 }
 
 ///////////////////////////////////////////////////////////////////////////////////////////////////
@@ -652,12 +671,22 @@ public:
 
 IMPLEMENT_GLOBAL_SHADER(FHairInterpolateMeshTriangleCS, "/Engine/Private/HairStrands/HairStrandsMeshInterpolate.usf", "MainCS", SF_Compute);
 
+void TransitBufferToReadable(FRHICommandList& RHICmdList, FBufferTransitionQueue& BuffersToTransit)
+{
+	if (BuffersToTransit.Num())
+	{
+		RHICmdList.TransitionResources(EResourceTransitionAccess::EReadable, EResourceTransitionPipeline::EComputeToCompute, BuffersToTransit.GetData(), BuffersToTransit.Num());
+		BuffersToTransit.SetNum(0, false);
+	}
+}
+
 static void AddHairStrandInterpolateMeshTrianglesPass(
 	FRDGBuilder& GraphBuilder,
 	FGlobalShaderMap* ShaderMap,
 	const int32 LODIndex,
 	const FHairStrandsProjectionMeshData::LOD& MeshData,
-	FHairStrandsProjectionHairData::HairGroup& RootData)
+	FHairStrandsProjectionHairData::HairGroup& RootData,
+	FBufferTransitionQueue& OutTransitionQueue)
 {
 	if (RootData.RootCount == 0 || LODIndex < 0 || LODIndex >= RootData.LODDatas.Num())
 	{
@@ -690,6 +719,10 @@ static void AddHairStrandInterpolateMeshTrianglesPass(
 		ComputeShader,
 		Parameters,
 		DispatchGroupCount);
+
+	OutTransitionQueue.Add(LODData.DeformedRootTrianglePosition0Buffer->UAV);
+	OutTransitionQueue.Add(LODData.DeformedRootTrianglePosition1Buffer->UAV);
+	OutTransitionQueue.Add(LODData.DeformedRootTrianglePosition2Buffer->UAV);
 }
 
 void InterpolateHairStrandsMeshTriangles(
@@ -697,9 +730,10 @@ void InterpolateHairStrandsMeshTriangles(
 	FGlobalShaderMap* ShaderMap,
 	const int32 LODIndex,
 	const FHairStrandsProjectionMeshData::LOD& ProjectionMeshData,
-	FHairStrandsProjectionHairData::HairGroup& ProjectionHairData)
+	FHairStrandsProjectionHairData::HairGroup& ProjectionHairData,
+	FBufferTransitionQueue& OutTransitionQueue)
 {
-	AddHairStrandInterpolateMeshTrianglesPass(GraphBuilder, ShaderMap, LODIndex, ProjectionMeshData, ProjectionHairData);
+	AddHairStrandInterpolateMeshTrianglesPass(GraphBuilder, ShaderMap, LODIndex, ProjectionMeshData, ProjectionHairData, OutTransitionQueue);
 }
 
 ///////////////////////////////////////////////////////////////////////////////////////////////////
@@ -738,7 +772,8 @@ static void AddHairStrandInitMeshSamplesPass(
 	const int32 LODIndex,
 	const HairStrandsTriangleType Type,
 	const FHairStrandsProjectionMeshData::LOD& MeshData,
-	FHairStrandsProjectionHairData::HairGroup& RootData)
+	FHairStrandsProjectionHairData::HairGroup& RootData,
+	FBufferTransitionQueue& OutTransitionQueue)
 {
 	if (LODIndex < 0 || LODIndex >= RootData.LODDatas.Num())
 	{
@@ -760,10 +795,12 @@ static void AddHairStrandInitMeshSamplesPass(
 		if (Type == HairStrandsTriangleType::RestPose)
 		{
 			Parameters->OutSamplePositionsBuffer = LODData.RestSamplePositionsBuffer->UAV;
+			OutTransitionQueue.Add(LODData.RestSamplePositionsBuffer->UAV);
 		}
 		else if (Type == HairStrandsTriangleType::DeformedPose)
 		{
 			Parameters->OutSamplePositionsBuffer = LODData.DeformedSamplePositionsBuffer->UAV;
+			OutTransitionQueue.Add(LODData.DeformedSamplePositionsBuffer->UAV);
 		}
 		else
 		{
@@ -788,9 +825,10 @@ void InitHairStrandsMeshSamples(
 	const int32 LODIndex,
 	const HairStrandsTriangleType Type,
 	const FHairStrandsProjectionMeshData::LOD& ProjectionMeshData,
-	FHairStrandsProjectionHairData::HairGroup& ProjectionHairData)
+	FHairStrandsProjectionHairData::HairGroup& ProjectionHairData,
+	FBufferTransitionQueue& OutTransitionQueue)
 {
-	AddHairStrandInitMeshSamplesPass(GraphBuilder, ShaderMap, LODIndex, Type, ProjectionMeshData, ProjectionHairData);
+	AddHairStrandInitMeshSamplesPass(GraphBuilder, ShaderMap, LODIndex, Type, ProjectionMeshData, ProjectionHairData, OutTransitionQueue);
 }
 
 ///////////////////////////////////////////////////////////////////////////////////////////////////
@@ -828,7 +866,8 @@ static void AddHairStrandUpdateMeshSamplesPass(
 	FGlobalShaderMap* ShaderMap,
 	const int32 LODIndex,
 	const FHairStrandsProjectionMeshData::LOD& MeshData,
-	FHairStrandsProjectionHairData::HairGroup& RootData)
+	FHairStrandsProjectionHairData::HairGroup& RootData,
+	FBufferTransitionQueue& OutTransitionQueue)
 {
 	if (LODIndex < 0 || LODIndex >= RootData.LODDatas.Num())
 	{
@@ -858,6 +897,8 @@ static void AddHairStrandUpdateMeshSamplesPass(
 			ComputeShader,
 			Parameters,
 			DispatchGroupCount);
+
+		OutTransitionQueue.Add(LODData.MeshSampleWeightsBuffer->UAV);
 	}
 }
 
@@ -866,9 +907,10 @@ void UpdateHairStrandsMeshSamples(
 	FGlobalShaderMap* ShaderMap,
 	const int32 LODIndex,
 	const FHairStrandsProjectionMeshData::LOD& ProjectionMeshData,
-	FHairStrandsProjectionHairData::HairGroup& ProjectionHairData)
+	FHairStrandsProjectionHairData::HairGroup& ProjectionHairData,
+	FBufferTransitionQueue& OutTransitionQueue)
 {
-	AddHairStrandUpdateMeshSamplesPass(GraphBuilder, ShaderMap, LODIndex, ProjectionMeshData, ProjectionHairData);
+	AddHairStrandUpdateMeshSamplesPass(GraphBuilder, ShaderMap, LODIndex, ProjectionMeshData, ProjectionHairData, OutTransitionQueue);
 }
 
 ///////////////////////////////////////////////////////////////////////////////////////////////////
diff --git a/Engine/Source/Runtime/Renderer/Private/HairStrands/HairStrandsMeshProjection.h b/Engine/Source/Runtime/Renderer/Private/HairStrands/HairStrandsMeshProjection.h
--- a/Engine/Source/Runtime/Renderer/Private/HairStrands/HairStrandsMeshProjection.h
+++ b/Engine/Source/Runtime/Renderer/Private/HairStrands/HairStrandsMeshProjection.h
@@ -14,7 +14,8 @@ void ProjectHairStrandsOntoMesh(
 	FGlobalShaderMap* ShaderMap,
 	const int32 LODIndex,
 	const FHairStrandsProjectionMeshData& ProjectionMeshData,
-	FHairStrandsProjectionHairData::HairGroup& ProjectionHairData);
+	FHairStrandsProjectionHairData::HairGroup& ProjectionHairData, 
+	FBufferTransitionQueue& OutTransitionQueue);
 
 void TransferMesh(
 	FRDGBuilder& GraphBuilder,
@@ -22,7 +23,8 @@ void TransferMesh(
 	const int32 LODIndex,
 	const FHairStrandsProjectionMeshData& SourceMeshData,
 	const FHairStrandsProjectionMeshData& TargetMeshData,
-	struct FRWBuffer& OutPositionBuffer);
+	struct FRWBuffer& OutPositionBuffer,
+	FBufferTransitionQueue& OutTransitionQueue);
 
 enum class HairStrandsTriangleType
 {
@@ -37,7 +39,8 @@ void UpdateHairStrandsMeshTriangles(
 	const int32 LODIndex,
 	const HairStrandsTriangleType Type,
 	const FHairStrandsProjectionMeshData::LOD& ProjectionMeshData,
-	FHairStrandsProjectionHairData::HairGroup& ProjectionHairData);
+	FHairStrandsProjectionHairData::HairGroup& ProjectionHairData,
+	FBufferTransitionQueue& TransitionQueue);
 
 /* Init the samples information to be used for interpolation*/
 void InitHairStrandsMeshSamples(
@@ -46,7 +49,8 @@ void InitHairStrandsMeshSamples(
 	const int32 LODIndex,
 	const HairStrandsTriangleType Type,
 	const FHairStrandsProjectionMeshData::LOD& ProjectionMeshData,
-	FHairStrandsProjectionHairData::HairGroup& ProjectionHairData);
+	FHairStrandsProjectionHairData::HairGroup& ProjectionHairData,
+	FBufferTransitionQueue& TransitionQueue);
 
 /* Update the samples information to be used for interpolation*/
 void UpdateHairStrandsMeshSamples(
@@ -54,7 +58,8 @@ void UpdateHairStrandsMeshSamples(
 	FGlobalShaderMap* ShaderMap,
 	const int32 LODIndex,
 	const FHairStrandsProjectionMeshData::LOD& ProjectionMeshData,
-	FHairStrandsProjectionHairData::HairGroup& ProjectionHairData);
+	FHairStrandsProjectionHairData::HairGroup& ProjectionHairData,
+	FBufferTransitionQueue& TransitionQueue);
 
 void GenerateFolliculeMask(
 	FRDGBuilder& GraphBuilder,
@@ -77,5 +82,6 @@ void InterpolateHairStrandsMeshTriangles(
 	FGlobalShaderMap* ShaderMap,
 	const int32 LODIndex,
 	const FHairStrandsProjectionMeshData::LOD& ProjectionMeshData,
-	FHairStrandsProjectionHairData::HairGroup& ProjectionHairData);
+	FHairStrandsProjectionHairData::HairGroup& ProjectionHairData,
+	FBufferTransitionQueue& TransitionQueue);
 
diff --git a/Engine/Source/Runtime/Renderer/Private/SceneRendering.cpp b/Engine/Source/Runtime/Renderer/Private/SceneRendering.cpp
--- a/Engine/Source/Runtime/Renderer/Private/SceneRendering.cpp
+++ b/Engine/Source/Runtime/Renderer/Private/SceneRendering.cpp
@@ -4227,11 +4227,12 @@ void RunGPUSkinCacheTransition(FRHICommandList& RHICmdList, FScene* Scene, EGPUS
 	const bool bRun = 
 		(bHairStrandsEnabled && EGPUSkinCacheTransition::FrameSetup == Type) || 
 		(!bHairStrandsEnabled && EGPUSkinCacheTransition::FrameSetup != Type);
+	const EResourceTransitionPipeline TargetPipeline = bHairStrandsEnabled ? EResourceTransitionPipeline::EComputeToCompute : EResourceTransitionPipeline::EComputeToGfx;
 	if (bRun)
 	{
 		if (FGPUSkinCache* GPUSkinCache = Scene->GetGPUSkinCache())
 		{
-			GPUSkinCache->TransitionAllToReadable(RHICmdList);
+			GPUSkinCache->TransitionAllToReadable(RHICmdList, TargetPipeline);
 		}
 	}
 }
\ No newline at end of file
diff --git a/Engine/Source/Runtime/Renderer/Public/HairStrandsInterface.h b/Engine/Source/Runtime/Renderer/Public/HairStrandsInterface.h
--- a/Engine/Source/Runtime/Renderer/Public/HairStrandsInterface.h
+++ b/Engine/Source/Runtime/Renderer/Public/HairStrandsInterface.h
@@ -421,3 +421,7 @@ RENDERER_API void EnqueueFollicleMaskUpdateQuery(const TArray<FFollicleInfo>& In
 
 void RunHairStrandsProcess(FRHICommandListImmediate& RHICmdList, class FGlobalShaderMap* ShaderMap); 
 bool HasHairStrandsProcess(EShaderPlatform Platform);
+
+
+typedef TArray<FRHIUnorderedAccessView*> FBufferTransitionQueue;
+RENDERER_API void TransitBufferToReadable(FRHICommandListImmediate& RHICmdList, FBufferTransitionQueue& BuffersToTransit);
