Fix transition issue when patching Gbuffer with hair data.

diff --git a/Engine/Source/Runtime/Renderer/Private/HairStrands/HairStrandsComposition.cpp b/Engine/Source/Runtime/Renderer/Private/HairStrands/HairStrandsComposition.cpp
--- a/Engine/Source/Runtime/Renderer/Private/HairStrands/HairStrandsComposition.cpp
+++ b/Engine/Source/Runtime/Renderer/Private/HairStrands/HairStrandsComposition.cpp
@@ -375,6 +375,9 @@ void RenderHairComposition(
 						RDGCategorisationTexture,
 						GBufferATexture,
 						GBufferBTexture);
+
+					GraphBuilder.QueueTextureExtraction(GBufferATexture, &SceneContext.GBufferA, true);
+					GraphBuilder.QueueTextureExtraction(GBufferBTexture, &SceneContext.GBufferB, true);
 				}
 			}
 		}
