Fix transition for hair cluster.

diff --git a/Engine/Source/Runtime/Renderer/Private/HairStrands/HairStrandsClusters.cpp b/Engine/Source/Runtime/Renderer/Private/HairStrands/HairStrandsClusters.cpp
--- a/Engine/Source/Runtime/Renderer/Private/HairStrands/HairStrandsClusters.cpp
+++ b/Engine/Source/Runtime/Renderer/Private/HairStrands/HairStrandsClusters.cpp
@@ -348,7 +348,8 @@ static void AddClusterCullingPass(
 	const FHairCullingParams& CullingParameters,
 	const FHairHZBParameters& HZBParameters,
 	const int32 ClusterCullingLodMode,
-	FHairStrandClusterData::FHairGroup& ClusterData)
+	FHairStrandClusterData::FHairGroup& ClusterData, 
+	FBufferTransitionQueue& OutTransitionQueue)
 {
 	FRDGBufferRef DispatchIndirectParametersClusterCount = GraphBuilder.CreateBuffer(FRDGBufferDesc::CreateIndirectDesc<FRHIDispatchIndirectParameters>(), TEXT("HairDispatchIndirectParametersClusterCount"));
 	FRDGBufferRef DispatchIndirectParametersClusterCount2D = GraphBuilder.CreateBuffer(FRDGBufferDesc::CreateIndirectDesc<FRHIDispatchIndirectParameters>(), TEXT("HairDispatchIndirectParametersClusterCount2D"));
@@ -450,6 +451,7 @@ static void AddClusterCullingPass(
 			Parameters,
 			DispatchCount);
 	}
+	OutTransitionQueue.Add(DrawIndirectParametersBuffer.UAV);
 
 	/// Prepare some indirect draw buffers for specific compute group size
 	{
@@ -534,6 +536,9 @@ static void AddClusterCullingPass(
 			ComputeShader,
 			Parameters,
 			DispatchIndirectParametersClusterCount2D, 0); // DispatchIndirectParametersClusterCount2D is used to avoid having any dispatch dimension going above 65535.
+
+		OutTransitionQueue.Add(Parameters->CulledCompactedIndexBuffer);
+		OutTransitionQueue.Add(Parameters->CulledCompactedRadiusScaleBuffer);
 	}
 
 	// Should this be move onto the culling result?
@@ -554,7 +559,8 @@ static void AddClusterCullingPass(
 static void AddClusterResetLod0(
 	FRDGBuilder& GraphBuilder,
 	FGlobalShaderMap* ShaderMap,
-	FHairStrandClusterData::FHairGroup& ClusterData)
+	FHairStrandClusterData::FHairGroup& ClusterData,
+	FBufferTransitionQueue& OutTransitionQueue)
 {
 	// Set as culling result not available
 	ClusterData.SetCullingResultAvailable(false);
@@ -573,6 +579,8 @@ static void AddClusterResetLod0(
 		ComputeShader,
 		Parameters,
 		FIntVector(1, 1, 1));
+
+	OutTransitionQueue.Add(Parameters->DrawIndirectParameters);
 }
 
 ///////////////////////////////////////////////////////////////////////////////////////////////////
@@ -622,12 +630,13 @@ void ComputeHairStrandsClustersCulling(
 	const bool bCullingProcessSkipped = GHairStrandsClusterCulling <= 0 || (CullingParameters.bShadowViewMode && GStrandHairClusterCullingShadow <= 0);
 	for (const FViewInfo& View : Views)
 	{
+		FBufferTransitionQueue TransitionQueue;
 		for (FHairStrandClusterData::FHairGroup& ClusterData : ClusterDatas.HairGroups)		
 		{
 			if (bCullingProcessSkipped)
 			{
 				FRDGBuilder GraphBuilder(RHICmdList);
-				AddClusterResetLod0(GraphBuilder, &ShaderMap, ClusterData);
+				AddClusterResetLod0(GraphBuilder, &ShaderMap, ClusterData, TransitionQueue);
 				GraphBuilder.Execute();
 				continue;
 			}
@@ -642,11 +651,13 @@ void ComputeHairStrandsClustersCulling(
 				CullingParameters,
 				HZBParameters,
 				LODMode,
-				ClusterData);
+				ClusterData,
+				TransitionQueue);
 			GraphBuilder.Execute();
-
+			
 			ClusterData.SetCullingResultAvailable(true);
 		}
+		TransitBufferToReadable(RHICmdList, TransitionQueue);
 	}
 }
 
@@ -661,11 +672,13 @@ void ResetHairStrandsClusterToLOD0(
 	SCOPED_DRAW_EVENT(RHICmdList, HairStrandsResetLod0);
 	SCOPED_GPU_STAT(RHICmdList, HairStrandsResetLod0);
 
+	FBufferTransitionQueue TransitionQueue;
 	for (FHairStrandClusterData::FHairGroup& ClusterData : ClusterDatas.HairGroups)
 	{
 		// TODO use compute overlap (will need to split AddClusterCullingPass)
 		FRDGBuilder GraphBuilder(RHICmdList);
-		AddClusterResetLod0(GraphBuilder, &ShaderMap, ClusterData);
+		AddClusterResetLod0(GraphBuilder, &ShaderMap, ClusterData, TransitionQueue);
 		GraphBuilder.Execute();
 	}
+	TransitBufferToReadable(RHICmdList, TransitionQueue);
 }
