Fix hair crash when gbuffer texture are invalid.

diff --git a/Engine/Source/Runtime/Renderer/Private/HairStrands/HairStrandsComposition.cpp b/Engine/Source/Runtime/Renderer/Private/HairStrands/HairStrandsComposition.cpp
--- a/Engine/Source/Runtime/Renderer/Private/HairStrands/HairStrandsComposition.cpp
+++ b/Engine/Source/Runtime/Renderer/Private/HairStrands/HairStrandsComposition.cpp
@@ -367,17 +367,20 @@ void RenderHairComposition(
 				if (bPatchBufferData)
 				{
 					FSceneRenderTargets& SceneContext = FSceneRenderTargets::Get(RHICmdList);				
-					const FRDGTextureRef GBufferATexture = GraphBuilder.RegisterExternalTexture(SceneContext.GBufferA , TEXT("GBufferA"));
-					const FRDGTextureRef GBufferBTexture = GraphBuilder.RegisterExternalTexture(SceneContext.GBufferB, TEXT("GBufferB"));
-					AddPatchGbufferDataPass(
-						GraphBuilder,
-						View,
-						RDGCategorisationTexture,
-						GBufferATexture,
-						GBufferBTexture);
-
-					GraphBuilder.QueueTextureExtraction(GBufferATexture, &SceneContext.GBufferA, true);
-					GraphBuilder.QueueTextureExtraction(GBufferBTexture, &SceneContext.GBufferB, true);
+					const FRDGTextureRef GBufferATexture = GraphBuilder.TryRegisterExternalTexture(SceneContext.GBufferA , TEXT("GBufferA"));
+					const FRDGTextureRef GBufferBTexture = GraphBuilder.TryRegisterExternalTexture(SceneContext.GBufferB, TEXT("GBufferB"));
+					if (GBufferATexture && GBufferBTexture)
+					{
+						AddPatchGbufferDataPass(
+							GraphBuilder,
+							View,
+							RDGCategorisationTexture,
+							GBufferATexture,
+							GBufferBTexture);
+
+						GraphBuilder.QueueTextureExtraction(GBufferATexture, &SceneContext.GBufferA, true);
+						GraphBuilder.QueueTextureExtraction(GBufferBTexture, &SceneContext.GBufferB, true);
+					}
 				}
 			}
 		}
diff --git a/Engine/Source/Runtime/Renderer/Private/HairStrands/HairStrandsVisibility.cpp b/Engine/Source/Runtime/Renderer/Private/HairStrands/HairStrandsVisibility.cpp
--- a/Engine/Source/Runtime/Renderer/Private/HairStrands/HairStrandsVisibility.cpp
+++ b/Engine/Source/Runtime/Renderer/Private/HairStrands/HairStrandsVisibility.cpp
@@ -2205,6 +2205,11 @@ static void AddHairVisibilityColorAndDepthPatchPass(
 	FRDGTextureRef& OutColorTexture,
 	FRDGTextureRef& OutDepthTexture)
 {
+	if (!OutGBufferBTexture || !OutColorTexture || !OutDepthTexture)
+	{
+		return;
+	}
+
 	FHairVisibilityDepthPS::FParameters* Parameters = GraphBuilder.AllocParameters<FHairVisibilityDepthPS::FParameters>();
 	Parameters->CategorisationTexture = CategorisationTexture;
 	Parameters->RenderTargets[0] = FRenderTargetBinding(OutGBufferBTexture, ERenderTargetLoadAction::ELoad);
@@ -2490,7 +2495,7 @@ FHairStrandsVisibilityViews RenderHairStrandsVisibilityBuffer(
 
 			FRDGBuilder GraphBuilder(RHICmdList);
 			FSceneRenderTargets& SceneContext = FSceneRenderTargets::Get(RHICmdList);
-			FRDGTextureRef SceneGBufferBTexture = GraphBuilder.RegisterExternalTexture(InSceneGBufferBTexture, TEXT("SceneGBufferBTexture"));
+			FRDGTextureRef SceneGBufferBTexture = GraphBuilder.TryRegisterExternalTexture(InSceneGBufferBTexture, TEXT("SceneGBufferBTexture"));
 			FRDGTextureRef SceneColorTexture = GraphBuilder.RegisterExternalTexture(InSceneColorTexture, TEXT("SceneColorTexture"));
 			FRDGTextureRef SceneDepthTexture = GraphBuilder.RegisterExternalTexture(InSceneDepthTexture, TEXT("SceneDepthTexture"));
 			FRDGTextureRef SceneVelocityTexture = InSceneVelocityTexture ? GraphBuilder.RegisterExternalTexture(InSceneVelocityTexture, TEXT("SceneVelocityTexture")) : nullptr;
