Fixed crash in shader complexity pass when visualizing complexity in cooked build.

diff --git a/Engine/Source/Runtime/Renderer/Private/PostProcess/PostProcessVisualizeComplexity.cpp b/Engine/Source/Runtime/Renderer/Private/PostProcess/PostProcessVisualizeComplexity.cpp
--- a/Engine/Source/Runtime/Renderer/Private/PostProcess/PostProcessVisualizeComplexity.cpp
+++ b/Engine/Source/Runtime/Renderer/Private/PostProcess/PostProcessVisualizeComplexity.cpp
@@ -27,15 +27,29 @@ public:
 		RENDER_TARGET_BINDING_SLOTS()
 	END_SHADER_PARAMETER_STRUCT()
 
+	enum class EQuadOverdraw
+	{
+		Disable,
+		Enable,
+		MAX
+	};
+
+	class FQuadOverdraw : SHADER_PERMUTATION_ENUM_CLASS("READ_QUAD_OVERDRAW", EQuadOverdraw);
+	using FPermutationDomain = TShaderPermutationDomain<FQuadOverdraw>;
+
 	static bool ShouldCompilePermutation(const FGlobalShaderPermutationParameters& Parameters)
 	{
+		const FPermutationDomain PermutationVector(Parameters.PermutationId);
+		if (PermutationVector.Get<FVisualizeComplexityApplyPS::FQuadOverdraw>() == FVisualizeComplexityApplyPS::EQuadOverdraw::Enable)
+		{
+			return AllowDebugViewShaderMode(DVSM_QuadComplexity, Parameters.Platform, GetMaxSupportedFeatureLevel(Parameters.Platform));
+		}
 		return true;
 	}
 
 	static void ModifyCompilationEnvironment(const FGlobalShaderPermutationParameters& Parameters, FShaderCompilerEnvironment& OutEnvironment)
 	{
 		FGlobalShader::ModifyCompilationEnvironment(Parameters, OutEnvironment);
-		OutEnvironment.SetDefine(TEXT("READ_QUAD_OVERDRAW"), AllowDebugViewShaderMode(DVSM_QuadComplexity, Parameters.Platform, GetMaxSupportedFeatureLevel(Parameters.Platform)));
 		OutEnvironment.SetDefine(TEXT("MAX_NUM_COMPLEXITY_COLORS"), MaxNumShaderComplexityColors);
 		// EVisualizeComplexityColorSamplingMethod values
 		OutEnvironment.SetDefine(TEXT("CS_RAMP"), (uint32)FVisualizeComplexityInputs::EColorSamplingMethod::Ramp);
@@ -106,14 +120,14 @@ FScreenPassTexture AddVisualizeComplexityPass(FRDGBuilder& GraphBuilder, const F
 	const FSceneRenderTargets& SceneRenderTargets = FSceneRenderTargets::Get(GraphBuilder.RHICmdList);
 	const EDebugViewShaderMode DebugViewShaderMode = View.Family->GetDebugViewShaderMode();
 
+	PassParameters->DebugViewShaderMode = DVSM_ShaderComplexity;
+	FVisualizeComplexityApplyPS::EQuadOverdraw QuadOverdrawEnum = FVisualizeComplexityApplyPS::EQuadOverdraw::Disable;
+
 	if (FRDGTextureRef QuadOverdrawTexture = GraphBuilder.TryRegisterExternalTexture(SceneRenderTargets.QuadOverdrawBuffer))
 	{
 		PassParameters->QuadOverdrawTexture = QuadOverdrawTexture;
 		PassParameters->DebugViewShaderMode = DebugViewShaderMode;
-	}
-	else // Otherwise fallback to a complexity mode that does not require the QuadOverdraw resources.
-	{
-		PassParameters->DebugViewShaderMode = DVSM_ShaderComplexity;
+		QuadOverdrawEnum = FVisualizeComplexityApplyPS::EQuadOverdraw::Enable;
 	}
 
 	PassParameters->bLegend = Inputs.bDrawLegend;
@@ -122,7 +136,9 @@ FScreenPassTexture AddVisualizeComplexityPass(FRDGBuilder& GraphBuilder, const F
 	PassParameters->ComplexityScale = Inputs.ComplexityScale;
 	PassParameters->UsedQuadTextureSize = FIntPoint(View.ViewRect.Size() + FIntPoint(1, 1)) / 2;
 
-	TShaderMapRef<FVisualizeComplexityApplyPS> PixelShader(View.ShaderMap);
+	FVisualizeComplexityApplyPS::FPermutationDomain PermutationVector;
+	PermutationVector.Set<FVisualizeComplexityApplyPS::FQuadOverdraw>(QuadOverdrawEnum);
+	TShaderMapRef<FVisualizeComplexityApplyPS> PixelShader(View.ShaderMap, PermutationVector);
 
 	RDG_EVENT_SCOPE(GraphBuilder, "VisualizeComplexity");
 
diff --git a/Engine/Source/Runtime/Renderer/Private/ShaderComplexityRendering.cpp b/Engine/Source/Runtime/Renderer/Private/ShaderComplexityRendering.cpp
--- a/Engine/Source/Runtime/Renderer/Private/ShaderComplexityRendering.cpp
+++ b/Engine/Source/Runtime/Renderer/Private/ShaderComplexityRendering.cpp
@@ -66,9 +66,7 @@ static FAutoConsoleVariableRef CVarShaderComplexityBaselineDeferredUnlitPS(
 	ECVF_Default
 );
 
-IMPLEMENT_SHADER_TYPE(template<>,TComplexityAccumulatePS<false>,TEXT("/Engine/Private/ShaderComplexityAccumulatePixelShader.usf"),TEXT("Main"),SF_Pixel);
-IMPLEMENT_SHADER_TYPE(template<>,TComplexityAccumulatePS<true>,TEXT("/Engine/Private/QuadComplexityAccumulatePixelShader.usf"),TEXT("Main"),SF_Pixel);
-
+IMPLEMENT_SHADER_TYPE(,FComplexityAccumulatePS,TEXT("/Engine/Private/ShaderComplexityAccumulatePixelShader.usf"),TEXT("Main"),SF_Pixel);
 
 void FComplexityAccumulateInterface::GetDebugViewModeShaderBindings(
 	const FDebugViewModePS& BaseShader,
@@ -86,17 +84,17 @@ void FComplexityAccumulateInterface::GetDebugViewModeShaderBindings(
 	FMeshDrawSingleShaderBindings& ShaderBindings
 ) const
 {
+	const FComplexityAccumulatePS& Shader = static_cast<const FComplexityAccumulatePS&>(BaseShader);
+
 	// normalize the complexity so we can fit it in a low precision scene color which is necessary on some platforms
 	// late value is for overdraw which can be problematic with a low precision float format, at some point the precision isn't there any more and it doesn't accumulate
 	if (DebugViewMode == DVSM_QuadComplexity)
 	{
-		const TComplexityAccumulatePS<true>& Shader = static_cast<const TComplexityAccumulatePS<true>&>(BaseShader);
 		ShaderBindings.Add(Shader.NormalizedComplexity, FVector4(NormalizedQuadComplexityValue));
 		ShaderBindings.Add(Shader.ShowQuadOverdraw, 0);
 	}
 	else
 	{
-		const TComplexityAccumulatePS<false>& Shader = static_cast<const TComplexityAccumulatePS<false>&>(BaseShader);
 		const float NormalizeMul = 1.0f / GetMaxShaderComplexityCount(Material.GetFeatureLevel());
 		ShaderBindings.Add(Shader.NormalizedComplexity, FVector4(NumPSInstructions * NormalizeMul, NumVSInstructions * NormalizeMul, 1 / 32.0f));
 		ShaderBindings.Add(Shader.ShowQuadOverdraw, DebugViewMode != DVSM_ShaderComplexity ? 1 : 0);
@@ -105,16 +103,11 @@ void FComplexityAccumulateInterface::GetDebugViewModeShaderBindings(
 
 TShaderRef<FDebugViewModePS> FComplexityAccumulateInterface::GetPixelShader(const FMaterial* InMaterial, FVertexFactoryType* VertexFactoryType) const
 {
-	if (bShowQuadComplexity)
-	{
-		return InMaterial->GetShader<TComplexityAccumulatePS<true>>(VertexFactoryType);
-	}
-	else
-	{
-		return InMaterial->GetShader<TComplexityAccumulatePS<false>>(VertexFactoryType);
-	}
-
+	FComplexityAccumulatePS::FPermutationDomain PermutationVector;
+	PermutationVector.Set<FComplexityAccumulatePS::FQuadOverdraw>(bShowQuadComplexity ? FComplexityAccumulatePS::EQuadOverdraw::Enable : FComplexityAccumulatePS::EQuadOverdraw::Disable);
+	return InMaterial->GetShader<FComplexityAccumulatePS>(VertexFactoryType, PermutationVector);
 }
+
 void FComplexityAccumulateInterface::SetDrawRenderState(EBlendMode BlendMode, FRenderState& DrawRenderState, bool bHasDepthPrepassForMaskedMaterial) const
 {
 	if (BlendMode == BLEND_Opaque)
@@ -139,8 +132,4 @@ void FComplexityAccumulateInterface::SetDrawRenderState(EBlendMode BlendMode, FR
 	DrawRenderState.BlendState = TStaticBlendState<CW_RGBA, BO_Add, BF_One, BF_One, BO_Add, BF_Zero, BF_One>::GetRHI();
 }
 
-// Instantiate the template 
-template class TComplexityAccumulatePS<false>;
-template class TComplexityAccumulatePS<true>;
-
 #endif // !(UE_BUILD_SHIPPING || UE_BUILD_TEST)
diff --git a/Engine/Source/Runtime/Renderer/Private/ShaderComplexityRendering.h b/Engine/Source/Runtime/Renderer/Private/ShaderComplexityRendering.h
--- a/Engine/Source/Runtime/Renderer/Private/ShaderComplexityRendering.h
+++ b/Engine/Source/Runtime/Renderer/Private/ShaderComplexityRendering.h
@@ -29,16 +29,28 @@ extern int32 GShaderComplexityBaselineDeferredVS;
 extern int32 GShaderComplexityBaselineDeferredPS;
 extern int32 GShaderComplexityBaselineDeferredUnlitPS;
 
-
-template <bool bQuadComplexity>
-class TComplexityAccumulatePS : public FDebugViewModePS
+class FComplexityAccumulatePS : public FDebugViewModePS
 {
-	DECLARE_SHADER_TYPE(TComplexityAccumulatePS,MeshMaterial);
 public:
+	DECLARE_SHADER_TYPE(FComplexityAccumulatePS, MeshMaterial);
+
+	enum class EQuadOverdraw
+	{
+		Disable,
+		Enable,
+		MAX
+	};
+
+	class FQuadOverdraw : SHADER_PERMUTATION_ENUM_CLASS("OUTPUT_QUAD_OVERDRAW", EQuadOverdraw);
+	using FPermutationDomain = TShaderPermutationDomain<FQuadOverdraw>;
+
 	static bool ShouldCompilePermutation(const FMeshMaterialShaderPermutationParameters& Parameters)
 	{
+		const FPermutationDomain PermutationVector(Parameters.PermutationId);
+		const bool bQuadOverdraw = PermutationVector.Get<FComplexityAccumulatePS::FQuadOverdraw>() == FComplexityAccumulatePS::EQuadOverdraw::Enable;
+
 		// See FDebugViewModeMaterialProxy::GetFriendlyName()
-		if (AllowDebugViewShaderMode(bQuadComplexity ? DVSM_QuadComplexity : DVSM_ShaderComplexity, Parameters.Platform, Parameters.MaterialParameters.FeatureLevel))
+		if (AllowDebugViewShaderMode(bQuadOverdraw ? DVSM_QuadComplexity : DVSM_ShaderComplexity, Parameters.Platform, Parameters.MaterialParameters.FeatureLevel))
 		{
 			// If it comes from FDebugViewModeMaterialProxy, compile it.
 			if (Parameters.MaterialParameters.bMaterialIsComplexityAccumulate)
@@ -56,7 +68,7 @@ public:
 		return false;
 	}
 
-	TComplexityAccumulatePS(const ShaderMetaType::CompiledShaderInitializerType& Initializer):
+	FComplexityAccumulatePS(const ShaderMetaType::CompiledShaderInitializerType& Initializer):
 		FDebugViewModePS(Initializer)
 	{
 		NormalizedComplexity.Bind(Initializer.ParameterMap,TEXT("NormalizedComplexity"));
@@ -64,11 +76,10 @@ public:
 		QuadBufferUAV.Bind(Initializer.ParameterMap,TEXT("RWQuadBuffer"));
 	}
 
-	TComplexityAccumulatePS() {}
+	FComplexityAccumulatePS() {}
 
 	static void ModifyCompilationEnvironment(const FMaterialShaderPermutationParameters& Parameters, FShaderCompilerEnvironment& OutEnvironment)
 	{
-		OutEnvironment.SetDefine(TEXT("OUTPUT_QUAD_OVERDRAW"), AllowDebugViewShaderMode(DVSM_QuadComplexity, Parameters.Platform, Parameters.MaterialParameters.FeatureLevel));
 		TCHAR BufferRegister[] = { 'u', '0', 0 };
 		BufferRegister[1] += FSceneRenderTargets::GetQuadOverdrawUAVIndex(Parameters.Platform, Parameters.MaterialParameters.FeatureLevel);
 		OutEnvironment.SetDefine(TEXT("QUAD_BUFFER_REGISTER"), BufferRegister);
