Add optional pass for computing statistical coverage.

diff --git a/Engine/Shaders/Private/HairStrands/HairStrandsCoverage.usf b/Engine/Shaders/Private/HairStrands/HairStrandsCoverage.usf
new file mode 100644
--- /dev/null
+++ b/Engine/Shaders/Private/HairStrands/HairStrandsCoverage.usf
@@ -0,0 +1,35 @@
+// Copyright Epic Games, Inc. All Rights Reserved.
+
+#include "../Common.ush"
+
+int2 OutputResolution;
+float LUT_HairCount; 
+float LUT_HairRadiusCount;
+
+SamplerState LinearSampler;
+Texture2D<float> HairCoverageLUT;
+Texture2D<float2> HairCountTexture;
+RWTexture2D<float> OutputTexture;
+
+#define TILE_PIXEL_SIZE 8
+
+[numthreads(TILE_PIXEL_SIZE, TILE_PIXEL_SIZE, 1)]
+void MainCS(uint3 DispatchThreadId : SV_DispatchThreadID)
+{
+	const int2 PixelCoord = DispatchThreadId.xy;
+	if (any(PixelCoord > OutputResolution))
+	{
+		return;
+	}
+
+	const float2 Raw = HairCountTexture.Load(uint3(PixelCoord, 0)).xy;
+
+	const float RemapScale = 2; // Remap the radius from [0..0.5] -> [0..1], i.e. the mapping used by the LUT
+
+	const float HairCount = Raw.y;
+	const float HairRadius = saturate(Raw.x / HairCount * RemapScale);
+	const float2 UV = float2(HairRadius, (HairCount+0.5f) / LUT_HairCount);
+	const float Coverage = HairCoverageLUT.SampleLevel(LinearSampler, UV, 0);
+	
+	OutputTexture[PixelCoord] = saturate(1-Coverage);
+}
\ No newline at end of file
diff --git a/Engine/Shaders/Private/HairStrands/HairStrandsLUT.usf b/Engine/Shaders/Private/HairStrands/HairStrandsLUT.usf
--- a/Engine/Shaders/Private/HairStrands/HairStrandsLUT.usf
+++ b/Engine/Shaders/Private/HairStrands/HairStrandsLUT.usf
@@ -6,6 +6,8 @@
 #include "../DeferredShadingCommon.ush"
 #include "HairStrandsCommon.ush"
 
+#if SHADER_HAIRLUT
+
 #include "../ShadingModels.ush"
 
 #ifdef HAIR_CUSTOM_BSDF
@@ -30,9 +32,10 @@ RWTexture3D<float4>	OutputColor;
 	return float(bits) * 2.3283064365386963e-10; // / 0x100000000
  }
 
- float2 hammersley2d(uint i, uint N) {
-     return float2(float(i)/float(N), radicalInverse_VdC(i));
- }
+float2 hammersley2d(uint i, uint N) 
+{
+	return float2(float(i)/float(N), radicalInverse_VdC(i));
+}
 
 #define TILE_PIXEL_SIZE 8
 #define JITTER_VIEW 0
@@ -265,4 +268,29 @@ void MainCS(uint3 DispatchThreadId : SV_DispatchThreadID)
 		1);
 }
 
-#endif
\ No newline at end of file
+#endif
+
+#endif // SHADER_HAIRLUT
+
+
+#if SHADER_HAIRCOVERAGE
+
+int2 OutputResolution;
+StructuredBuffer<float> InputBuffer;
+RWTexture2D<float> OutputTexture;
+
+#define TILE_PIXEL_SIZE 8
+
+[numthreads(TILE_PIXEL_SIZE, TILE_PIXEL_SIZE, 1)]
+void MainCS(uint3 DispatchThreadId : SV_DispatchThreadID)
+{
+	const int2 PixelCoord = DispatchThreadId.xy;
+	if (any(PixelCoord > OutputResolution))
+	{
+		return;
+	}
+
+	const uint Index = PixelCoord.x + PixelCoord.y * OutputResolution.x;
+	OutputTexture[PixelCoord] = InputBuffer[Index];
+}
+#endif // SHADER_HAIRCOVERAGE
\ No newline at end of file
diff --git a/Engine/Shaders/Private/HairStrands/HairStrandsVisibilityPS.usf b/Engine/Shaders/Private/HairStrands/HairStrandsVisibilityPS.usf
--- a/Engine/Shaders/Private/HairStrands/HairStrandsVisibilityPS.usf
+++ b/Engine/Shaders/Private/HairStrands/HairStrandsVisibilityPS.usf
@@ -60,7 +60,7 @@ void MainVisibility(
 	, out float OutColor0 : SV_Target0
 #elif HAIR_RENDER_MODE == RENDER_MODE_TRANSMITTANCE_AND_HAIRCOUNT
 	, out float OutColor0 : SV_Target0
-	, out float OutColor1 : SV_Target1
+	, out float2 OutColor1 : SV_Target1
 #elif HAIR_RENDER_MODE == RENDER_MODE_MSAA
 	, out uint2 OutColor0	: SV_Target0
   #if SUPPORT_MATERIAL_PROPERTY
@@ -104,7 +104,7 @@ void MainVisibility(
 	OutColor0 = saturate(1.0f - Coverage);
 #elif HAIR_RENDER_MODE == RENDER_MODE_TRANSMITTANCE_AND_HAIRCOUNT
 	OutColor0 = saturate(1.0f - Coverage);
-	OutColor1 = Coverage;
+	OutColor1 = float2(Coverage, 1);
 #endif
 
 #if HAIR_RENDER_MODE == RENDER_MODE_MSAA || HAIR_RENDER_MODE == RENDER_MODE_PPLL
diff --git a/Engine/Source/Runtime/Renderer/Private/HairStrands/HairStrandsLUT.cpp b/Engine/Source/Runtime/Renderer/Private/HairStrands/HairStrandsLUT.cpp
--- a/Engine/Source/Runtime/Renderer/Private/HairStrands/HairStrandsLUT.cpp
+++ b/Engine/Source/Runtime/Renderer/Private/HairStrands/HairStrandsLUT.cpp
@@ -26,7 +26,7 @@ class FHairLUTCS : public FGlobalShader
 	DECLARE_GLOBAL_SHADER(FHairLUTCS);
 	SHADER_USE_PARAMETER_STRUCT(FHairLUTCS, FGlobalShader);
 
-	class FLUTType : SHADER_PERMUTATION_INT("PERMUTATION_LUT_TYPE", HairLUTTypeCount);
+	class FLUTType : SHADER_PERMUTATION_INT("PERMUTATION_LUT_TYPE", 2); // Permutation only for HairLUTType_DualScattering and HairLUTType_MeanEnergy
 	using FPermutationDomain = TShaderPermutationDomain<FLUTType>;
 
 	BEGIN_SHADER_PARAMETER_STRUCT(FParameters, )
@@ -40,6 +40,11 @@ class FHairLUTCS : public FGlobalShader
 
 public:
 	static bool ShouldCompilePermutation(const FGlobalShaderPermutationParameters& Parameters) { return IsHairStrandsSupported(Parameters.Platform); }
+	static void ModifyCompilationEnvironment(const FGlobalShaderPermutationParameters& Parameters, FShaderCompilerEnvironment& OutEnvironment)
+	{
+		FGlobalShader::ModifyCompilationEnvironment(Parameters, OutEnvironment);
+		OutEnvironment.SetDefine(TEXT("SHADER_HAIRLUT"), 1);
+	}
 };
 
 IMPLEMENT_GLOBAL_SHADER(FHairLUTCS, "/Engine/Private/HairStrands/HairStrandsLUT.usf", "MainCS", SF_Compute);
@@ -83,11 +88,187 @@ static FRDGTextureRef AddHairLUTPass(
 	return HairLUTTexture;
 }
 
+/////////////////////////////////////////////////////////////////////////////////////////
+
+struct FHairCountToCoverageData
+{
+	// The LUT data have been computed by rasterizing in 1D segments with various radius size 
+	// (normalized), which are aligned with the rasterization plane and ortho-projected. 
+	// The segments position are randomly position within this volume. The number represent the 
+	// coverage of these segments over (the result is averaged over multiple iteration)
+	//
+	//
+	//  | --.--				 |
+	//  |			--.--	 |
+	//  |	--.--			 |
+	//  |				--.--| <- segment with radius 2
+	//  |					 |
+	//	 ====================  <- Projection plane
+	//    -------   ---------  <- Projected coverage (18/20 = 0.9)
+	//
+	// LUT data
+	//
+	// Data are layout as:
+	// .----> Radius [0..0.5]
+	// :
+	// :
+	// v
+	// Count [0..64]
+	static const uint32 HairCount = 64;			// Hair count
+	static const uint32 HairRadiusCount = 64;	// Normalized hair radius [0,1]
+	float Data[HairCount * HairRadiusCount] =
+	{
+		0.000000, 0.000000, 0.000000, 0.000000, 0.000000, 0.000000, 0.000000, 0.000000, 0.000000, 0.000000, 0.000000, 0.000000, 0.000000, 0.000000, 0.000000, 0.000000, 0.000000, 0.000000, 0.000000, 0.000000, 0.000000, 0.000000, 0.000000, 0.000000, 0.000000, 0.000000, 0.000000, 0.000000, 0.000000, 0.000000, 0.000000, 0.000000, 0.000000, 0.000000, 0.000000, 0.000000, 0.000000, 0.000000, 0.000000, 0.000000, 0.000000, 0.000000, 0.000000, 0.000000, 0.000000, 0.000000, 0.000000, 0.000000, 0.000000, 0.000000, 0.000000, 0.000000, 0.000000, 0.000000, 0.000000, 0.000000, 0.000000, 0.000000, 0.000000, 0.000000, 0.000000, 0.000000, 0.000000, 0.000000,
+		0.008789, 0.024406, 0.039734, 0.055405, 0.069511, 0.084885, 0.098946, 0.113274, 0.129799, 0.144028, 0.155937, 0.173019, 0.188591, 0.196022, 0.218147, 0.228256, 0.238892, 0.258514, 0.264328, 0.280075, 0.295296, 0.314209, 0.324478, 0.327744, 0.347076, 0.360077, 0.379913, 0.376045, 0.393951, 0.407997, 0.412193, 0.434975, 0.440971, 0.461090, 0.466393, 0.484055, 0.489891, 0.498116, 0.525528, 0.526619, 0.542175, 0.557365, 0.571144, 0.565247, 0.571953, 0.580391, 0.578880, 0.617622, 0.622833, 0.630913, 0.652977, 0.645233, 0.666428, 0.660080, 0.683022, 0.682678, 0.685478, 0.693871, 0.684959, 0.705963, 0.709587, 0.712509, 0.727493, 0.751579,
+		0.017426, 0.048172, 0.077583, 0.107742, 0.135368, 0.161079, 0.192635, 0.210968, 0.246330, 0.274132, 0.303452, 0.310791, 0.336220, 0.358101, 0.382088, 0.398293, 0.406044, 0.438766, 0.461700, 0.480713, 0.508865, 0.533913, 0.542992, 0.545395, 0.567390, 0.563477, 0.583435, 0.612846, 0.631836, 0.647995, 0.648209, 0.685837, 0.669106, 0.712608, 0.713882, 0.719688, 0.757088, 0.747147, 0.718300, 0.774094, 0.784523, 0.760803, 0.800415, 0.786049, 0.776390, 0.820137, 0.835327, 0.832344, 0.836143, 0.810715, 0.842072, 0.853645, 0.857590, 0.865463, 0.896072, 0.855072, 0.894402, 0.885559, 0.901695, 0.874512, 0.921165, 0.907143, 0.906357, 0.942085,
+		0.026024, 0.070976, 0.112885, 0.156082, 0.196938, 0.235626, 0.270355, 0.312592, 0.337318, 0.370491, 0.405411, 0.425171, 0.459076, 0.488503, 0.512268, 0.544518, 0.556007, 0.568535, 0.610794, 0.621742, 0.631241, 0.653442, 0.674744, 0.688759, 0.732086, 0.735016, 0.750565, 0.761841, 0.775375, 0.784073, 0.799828, 0.823380, 0.803802, 0.826637, 0.844254, 0.845901, 0.847824, 0.867767, 0.860481, 0.865143, 0.896805, 0.895027, 0.894897, 0.899597, 0.900177, 0.921333, 0.917160, 0.923805, 0.911125, 0.926506, 0.922195, 0.931145, 0.951881, 0.946098, 0.936646, 0.940208, 0.951546, 0.954811, 0.957283, 0.958916, 0.971268, 0.953644, 0.970245, 0.966476,
+		0.034500, 0.092972, 0.149261, 0.204109, 0.255180, 0.294800, 0.338127, 0.383827, 0.427025, 0.466690, 0.490837, 0.531883, 0.552689, 0.579475, 0.609520, 0.636734, 0.656677, 0.711418, 0.730705, 0.725983, 0.742622, 0.763206, 0.774139, 0.806419, 0.820213, 0.828415, 0.839287, 0.824165, 0.864479, 0.863770, 0.903656, 0.863655, 0.894279, 0.894775, 0.914398, 0.906357, 0.934486, 0.918411, 0.923599, 0.927643, 0.938988, 0.951210, 0.947090, 0.936371, 0.950104, 0.939949, 0.967262, 0.967094, 0.963509, 0.950974, 0.947487, 0.964249, 0.975464, 0.979202, 0.968735, 0.970848, 0.984901, 0.993584, 0.988457, 0.978539, 0.984863, 0.986694, 0.987686, 0.988327,
+		0.043159, 0.115349, 0.185425, 0.246895, 0.301567, 0.363075, 0.408554, 0.460640, 0.496216, 0.541985, 0.575050, 0.627769, 0.626396, 0.643623, 0.706795, 0.727127, 0.741158, 0.776749, 0.780037, 0.780907, 0.823433, 0.832558, 0.837196, 0.873184, 0.878357, 0.907158, 0.891731, 0.902962, 0.907890, 0.919731, 0.917488, 0.904564, 0.938148, 0.946938, 0.947891, 0.944351, 0.930428, 0.960114, 0.954109, 0.964043, 0.957581, 0.974121, 0.967110, 0.966339, 0.972649, 0.978516, 0.973984, 0.978958, 0.978470, 0.976807, 0.980011, 0.988541, 0.988998, 0.985252, 0.992378, 0.989326, 0.994057, 0.993500, 0.992233, 0.997513, 0.994102, 0.993172, 0.991951, 0.994698,
+		0.051559, 0.136993, 0.217186, 0.290283, 0.352890, 0.410606, 0.464027, 0.511902, 0.551926, 0.609207, 0.658470, 0.676750, 0.696190, 0.728073, 0.766769, 0.786621, 0.781738, 0.822853, 0.826515, 0.854538, 0.868111, 0.889374, 0.890587, 0.899826, 0.894936, 0.929871, 0.924667, 0.925766, 0.950134, 0.953720, 0.950790, 0.963669, 0.956429, 0.958801, 0.968628, 0.969856, 0.971626, 0.966431, 0.981903, 0.980042, 0.986427, 0.976883, 0.980614, 0.992188, 0.995102, 0.987427, 0.990036, 0.987900, 0.984207, 0.995285, 0.995262, 0.994095, 0.994446, 0.997803, 0.991959, 0.997414, 0.998039, 0.995407, 0.996544, 0.998749, 0.998421, 0.997673, 0.993515, 0.998276,
+		0.060028, 0.157425, 0.246254, 0.323997, 0.392830, 0.466003, 0.523483, 0.579788, 0.612083, 0.654938, 0.696167, 0.738136, 0.759682, 0.806519, 0.823616, 0.823822, 0.859291, 0.867531, 0.874245, 0.888222, 0.904663, 0.910927, 0.927429, 0.932625, 0.934631, 0.947792, 0.953156, 0.957588, 0.963486, 0.967110, 0.963135, 0.970177, 0.972160, 0.983017, 0.973801, 0.984840, 0.977005, 0.984856, 0.980446, 0.984871, 0.986580, 0.989845, 0.993851, 0.990486, 0.991814, 0.994095, 0.992905, 0.994278, 0.998276, 0.992569, 0.998299, 0.998459, 0.999046, 0.998894, 0.996979, 0.998032, 0.996536, 0.999657, 0.999542, 0.998299, 1.000000, 0.998856, 0.998085, 1.000000,
+		0.067612, 0.178909, 0.279297, 0.365265, 0.445526, 0.513519, 0.570251, 0.630669, 0.675842, 0.696297, 0.751038, 0.777313, 0.800484, 0.814667, 0.850967, 0.869514, 0.885628, 0.888557, 0.894508, 0.897949, 0.944176, 0.945366, 0.944557, 0.956909, 0.947495, 0.957817, 0.968575, 0.976051, 0.972198, 0.966736, 0.983101, 0.978447, 0.987679, 0.986870, 0.983665, 0.985641, 0.992165, 0.994568, 0.992073, 0.992020, 0.992538, 0.994095, 0.993568, 0.995605, 0.993912, 0.994537, 0.998459, 0.996910, 0.996323, 0.996185, 0.994881, 0.999306, 0.998672, 0.998795, 0.998474, 1.000000, 0.997963, 0.998917, 0.998825, 0.998131, 1.000000, 0.999641, 0.998283, 0.999916,
+		0.075523, 0.196907, 0.304184, 0.402954, 0.484756, 0.550385, 0.614304, 0.653305, 0.704628, 0.741997, 0.785645, 0.810913, 0.837151, 0.861092, 0.883461, 0.902954, 0.922089, 0.931786, 0.929649, 0.937462, 0.949234, 0.953613, 0.958870, 0.967575, 0.969231, 0.971031, 0.977249, 0.979080, 0.984818, 0.981041, 0.988510, 0.982941, 0.993843, 0.993423, 0.988708, 0.993393, 0.993721, 0.989815, 0.997269, 0.994308, 0.994835, 0.998421, 0.998520, 0.995613, 0.998856, 0.997902, 0.997162, 1.000000, 0.999886, 0.997650, 0.995758, 0.997231, 0.999870, 0.999611, 0.998886, 1.000000, 1.000000, 0.999306, 0.999840, 0.999802, 1.000000, 1.000000, 1.000000, 1.000000,
+		0.083893, 0.218582, 0.339401, 0.433067, 0.513351, 0.586296, 0.637558, 0.701775, 0.758751, 0.779373, 0.829193, 0.852005, 0.873390, 0.879280, 0.915291, 0.924088, 0.933380, 0.944077, 0.937752, 0.967621, 0.956108, 0.970253, 0.978577, 0.972061, 0.975395, 0.978745, 0.985992, 0.979378, 0.984177, 0.986565, 0.986961, 0.994408, 0.991814, 0.994125, 0.990974, 0.998482, 0.996368, 0.995483, 0.998840, 0.998703, 0.997612, 0.998749, 0.999237, 0.999069, 0.996582, 0.999664, 0.999733, 0.999947, 0.999634, 0.998520, 0.999161, 0.999008, 1.000000, 0.999626, 1.000000, 0.999992, 1.000000, 1.000000, 0.999954, 1.000000, 1.000000, 1.000000, 1.000000, 1.000000,
+		0.092308, 0.235756, 0.361153, 0.458038, 0.553986, 0.618820, 0.681404, 0.744438, 0.788002, 0.830200, 0.856842, 0.866829, 0.894180, 0.909943, 0.937607, 0.934853, 0.951210, 0.956596, 0.957680, 0.968117, 0.969849, 0.978058, 0.975967, 0.981163, 0.987625, 0.986732, 0.990028, 0.990814, 0.993111, 0.994026, 0.993919, 0.996422, 0.992676, 0.995613, 0.998428, 0.998375, 0.996201, 0.997986, 0.998680, 0.995148, 0.998497, 0.998375, 0.999336, 0.998276, 1.000000, 0.997009, 0.999931, 0.999451, 0.999916, 1.000000, 1.000000, 1.000000, 0.999527, 0.999489, 0.999458, 0.999664, 1.000000, 1.000000, 1.000000, 1.000000, 0.999069, 1.000000, 0.999931, 1.000000,
+		0.100403, 0.251717, 0.386551, 0.494720, 0.576874, 0.652550, 0.717285, 0.760880, 0.794640, 0.848831, 0.862274, 0.888161, 0.913094, 0.928772, 0.926506, 0.944870, 0.962143, 0.966675, 0.968651, 0.972496, 0.976944, 0.984314, 0.979530, 0.989098, 0.986214, 0.988693, 0.992043, 0.993332, 0.991219, 0.991920, 0.994682, 0.997849, 0.998085, 0.995010, 0.999527, 0.999237, 0.995590, 0.996895, 0.997597, 0.998405, 1.000000, 0.996521, 0.999916, 0.999939, 0.999237, 0.999702, 1.000000, 0.999779, 0.999641, 1.000000, 1.000000, 0.999588, 1.000000, 1.000000, 0.998657, 1.000000, 1.000000, 1.000000, 1.000000, 1.000000, 1.000000, 1.000000, 0.999809, 1.000000,
+		0.108910, 0.270073, 0.401291, 0.516602, 0.611908, 0.686852, 0.741570, 0.801125, 0.828384, 0.862450, 0.885727, 0.914024, 0.921753, 0.933525, 0.954437, 0.963829, 0.969299, 0.970062, 0.972473, 0.983170, 0.986092, 0.987579, 0.990578, 0.992264, 0.997375, 0.995148, 0.994690, 0.995209, 0.996483, 0.994431, 0.995171, 0.997261, 0.992523, 0.997787, 0.999664, 0.998795, 0.998055, 0.999107, 0.999626, 1.000000, 0.998718, 1.000000, 0.999802, 0.999992, 1.000000, 0.999733, 1.000000, 1.000000, 0.999374, 1.000000, 0.999901, 0.999657, 0.999672, 1.000000, 1.000000, 1.000000, 1.000000, 1.000000, 1.000000, 1.000000, 1.000000, 1.000000, 1.000000, 1.000000,
+		0.116379, 0.289742, 0.434738, 0.544319, 0.636292, 0.704086, 0.775749, 0.819443, 0.862495, 0.880333, 0.912170, 0.930916, 0.939827, 0.939926, 0.969421, 0.967003, 0.970039, 0.980370, 0.979713, 0.987335, 0.988701, 0.991798, 0.991356, 0.995628, 0.991928, 0.993896, 0.995262, 0.995811, 0.996529, 0.996834, 0.997292, 0.998085, 0.997940, 0.996895, 0.998947, 0.998672, 1.000000, 0.998657, 1.000000, 0.999893, 0.999268, 1.000000, 0.999504, 1.000000, 1.000000, 0.999878, 1.000000, 0.999908, 0.999901, 1.000000, 1.000000, 1.000000, 1.000000, 0.999542, 1.000000, 1.000000, 0.999847, 1.000000, 1.000000, 1.000000, 1.000000, 1.000000, 1.000000, 1.000000,
+		0.124184, 0.307930, 0.453300, 0.571358, 0.658180, 0.735107, 0.783073, 0.835258, 0.890739, 0.901199, 0.924019, 0.952148, 0.941734, 0.961433, 0.961937, 0.975883, 0.973068, 0.981674, 0.982170, 0.988525, 0.991249, 0.989670, 0.994225, 0.995537, 0.995682, 0.994553, 0.997116, 0.997818, 0.997185, 0.998962, 0.999329, 0.998863, 0.999138, 0.999298, 0.999107, 0.998711, 0.999695, 1.000000, 0.999275, 0.999374, 0.998672, 0.999222, 0.999580, 1.000000, 1.000000, 1.000000, 1.000000, 0.999992, 1.000000, 1.000000, 1.000000, 1.000000, 1.000000, 1.000000, 1.000000, 1.000000, 1.000000, 1.000000, 1.000000, 1.000000, 1.000000, 1.000000, 1.000000, 1.000000,
+		0.131897, 0.325531, 0.474838, 0.594193, 0.684105, 0.756439, 0.811806, 0.852737, 0.885254, 0.912521, 0.926079, 0.946075, 0.958748, 0.973618, 0.970711, 0.979477, 0.980881, 0.984039, 0.992569, 0.989113, 0.994568, 0.996231, 0.993401, 0.996964, 0.997414, 0.999512, 0.996246, 0.997818, 0.996475, 0.999901, 0.998894, 0.999962, 1.000000, 0.999260, 0.998810, 0.998993, 1.000000, 1.000000, 0.998962, 1.000000, 0.999878, 1.000000, 1.000000, 1.000000, 1.000000, 1.000000, 1.000000, 1.000000, 1.000000, 1.000000, 1.000000, 1.000000, 1.000000, 1.000000, 1.000000, 1.000000, 1.000000, 1.000000, 1.000000, 1.000000, 1.000000, 1.000000, 1.000000, 1.000000,
+		0.138802, 0.342453, 0.497086, 0.612640, 0.699142, 0.781097, 0.830269, 0.874092, 0.911591, 0.924973, 0.944786, 0.950035, 0.967682, 0.976898, 0.978294, 0.979156, 0.984772, 0.992920, 0.991875, 0.993591, 0.994789, 0.995239, 0.998413, 0.996101, 0.994453, 0.997665, 0.999481, 0.998207, 0.998192, 0.998344, 0.999496, 0.999420, 0.999779, 1.000000, 0.999390, 1.000000, 0.999901, 0.998154, 1.000000, 0.999870, 1.000000, 1.000000, 0.999802, 1.000000, 0.999718, 1.000000, 1.000000, 1.000000, 1.000000, 1.000000, 1.000000, 1.000000, 1.000000, 1.000000, 1.000000, 1.000000, 1.000000, 1.000000, 1.000000, 1.000000, 1.000000, 1.000000, 1.000000, 1.000000,
+		0.146057, 0.357552, 0.514435, 0.646202, 0.726608, 0.796219, 0.838158, 0.883415, 0.902046, 0.934616, 0.953339, 0.954117, 0.971527, 0.981773, 0.983421, 0.984383, 0.987785, 0.987968, 0.994308, 0.992142, 0.992188, 0.996742, 0.995972, 0.998962, 0.999557, 0.999008, 0.998917, 0.998589, 0.999290, 0.998520, 0.999802, 0.999344, 0.999916, 0.998840, 0.999565, 1.000000, 1.000000, 1.000000, 0.999466, 1.000000, 1.000000, 1.000000, 1.000000, 1.000000, 1.000000, 1.000000, 1.000000, 1.000000, 1.000000, 1.000000, 1.000000, 1.000000, 1.000000, 1.000000, 1.000000, 1.000000, 1.000000, 1.000000, 1.000000, 1.000000, 1.000000, 1.000000, 1.000000, 1.000000,
+		0.154175, 0.375061, 0.531952, 0.655937, 0.745850, 0.821587, 0.854744, 0.899902, 0.919884, 0.937317, 0.957794, 0.961128, 0.976089, 0.984970, 0.982155, 0.986389, 0.990921, 0.996063, 0.997536, 0.995163, 0.997627, 0.997665, 0.997032, 0.998360, 0.997452, 0.998573, 0.999863, 0.999153, 0.998703, 0.999336, 0.999397, 0.998558, 0.999046, 0.999496, 0.999794, 1.000000, 1.000000, 0.999779, 0.998978, 1.000000, 1.000000, 1.000000, 1.000000, 1.000000, 1.000000, 1.000000, 1.000000, 1.000000, 1.000000, 1.000000, 1.000000, 1.000000, 1.000000, 1.000000, 1.000000, 1.000000, 1.000000, 1.000000, 1.000000, 1.000000, 1.000000, 1.000000, 1.000000, 1.000000,
+		0.161049, 0.389206, 0.548698, 0.685417, 0.767311, 0.835144, 0.874847, 0.916588, 0.929390, 0.954712, 0.966225, 0.980965, 0.981232, 0.979637, 0.985573, 0.990486, 0.994980, 0.996101, 0.996887, 0.996147, 0.997635, 0.998535, 0.997871, 0.999275, 0.998100, 0.999962, 0.999397, 1.000000, 1.000000, 1.000000, 0.999374, 1.000000, 0.998657, 0.999962, 0.999962, 1.000000, 1.000000, 1.000000, 1.000000, 1.000000, 1.000000, 1.000000, 0.999809, 1.000000, 1.000000, 1.000000, 1.000000, 1.000000, 1.000000, 1.000000, 1.000000, 1.000000, 1.000000, 1.000000, 1.000000, 1.000000, 1.000000, 1.000000, 1.000000, 1.000000, 1.000000, 1.000000, 1.000000, 1.000000,
+		0.170219, 0.404900, 0.569893, 0.692253, 0.783989, 0.854202, 0.883682, 0.922493, 0.950554, 0.952759, 0.965172, 0.977585, 0.981270, 0.987717, 0.991150, 0.993927, 0.991531, 0.996193, 0.992485, 0.994942, 0.996277, 0.998131, 0.999123, 0.998474, 0.998230, 0.998543, 0.999771, 0.999451, 0.999916, 0.999901, 0.999443, 1.000000, 0.999512, 1.000000, 0.999863, 1.000000, 1.000000, 1.000000, 1.000000, 1.000000, 1.000000, 1.000000, 1.000000, 1.000000, 1.000000, 1.000000, 1.000000, 1.000000, 1.000000, 1.000000, 1.000000, 1.000000, 1.000000, 1.000000, 1.000000, 1.000000, 1.000000, 1.000000, 1.000000, 1.000000, 1.000000, 1.000000, 1.000000, 1.000000,
+		0.175995, 0.417778, 0.587364, 0.704094, 0.794579, 0.848862, 0.906929, 0.925674, 0.950180, 0.960602, 0.969231, 0.975098, 0.984589, 0.989143, 0.993950, 0.993095, 0.994751, 0.996964, 0.995728, 0.997673, 0.997772, 0.996094, 0.998779, 0.998352, 0.998810, 1.000000, 0.999397, 0.999718, 0.999466, 0.999428, 0.999542, 1.000000, 0.999413, 1.000000, 0.999809, 1.000000, 1.000000, 0.999870, 1.000000, 1.000000, 1.000000, 1.000000, 1.000000, 1.000000, 1.000000, 1.000000, 1.000000, 1.000000, 1.000000, 1.000000, 1.000000, 1.000000, 1.000000, 1.000000, 1.000000, 1.000000, 1.000000, 1.000000, 1.000000, 1.000000, 1.000000, 1.000000, 1.000000, 1.000000,
+		0.183105, 0.431953, 0.608543, 0.722977, 0.815285, 0.879860, 0.906868, 0.940750, 0.962158, 0.968269, 0.977051, 0.986664, 0.987892, 0.992203, 0.997589, 0.994934, 0.998116, 0.998703, 0.996826, 0.999458, 0.997658, 0.999382, 0.998772, 0.997490, 0.999878, 1.000000, 0.999550, 1.000000, 0.999062, 1.000000, 1.000000, 1.000000, 0.999657, 1.000000, 1.000000, 1.000000, 1.000000, 1.000000, 1.000000, 1.000000, 1.000000, 1.000000, 1.000000, 1.000000, 1.000000, 1.000000, 1.000000, 1.000000, 1.000000, 1.000000, 1.000000, 1.000000, 1.000000, 1.000000, 1.000000, 1.000000, 1.000000, 1.000000, 1.000000, 1.000000, 1.000000, 1.000000, 1.000000, 1.000000,
+		0.189987, 0.440308, 0.616745, 0.747749, 0.818283, 0.876892, 0.920494, 0.950851, 0.962852, 0.966843, 0.973465, 0.990219, 0.986603, 0.993317, 0.995186, 0.996986, 0.996811, 0.998154, 0.997940, 0.998734, 0.998146, 0.998367, 0.999504, 0.999908, 0.999779, 0.997971, 0.998894, 0.999939, 1.000000, 1.000000, 0.999893, 1.000000, 1.000000, 1.000000, 1.000000, 1.000000, 0.999817, 1.000000, 1.000000, 1.000000, 1.000000, 1.000000, 1.000000, 1.000000, 1.000000, 1.000000, 1.000000, 1.000000, 1.000000, 1.000000, 1.000000, 1.000000, 1.000000, 1.000000, 1.000000, 1.000000, 1.000000, 1.000000, 1.000000, 1.000000, 1.000000, 1.000000, 1.000000, 1.000000,
+		0.198090, 0.458580, 0.638443, 0.752754, 0.839699, 0.888031, 0.927704, 0.946678, 0.959824, 0.978127, 0.979744, 0.987892, 0.987404, 0.995247, 0.992355, 0.995598, 0.997238, 0.998955, 0.997292, 0.999100, 0.999382, 0.999924, 0.999161, 0.999992, 0.999466, 1.000000, 0.999870, 1.000000, 1.000000, 1.000000, 0.999863, 1.000000, 1.000000, 1.000000, 1.000000, 1.000000, 1.000000, 1.000000, 1.000000, 1.000000, 1.000000, 1.000000, 1.000000, 1.000000, 1.000000, 1.000000, 1.000000, 1.000000, 1.000000, 1.000000, 1.000000, 1.000000, 1.000000, 1.000000, 1.000000, 1.000000, 1.000000, 1.000000, 1.000000, 1.000000, 1.000000, 1.000000, 1.000000, 1.000000,
+		0.204880, 0.471451, 0.648201, 0.774055, 0.847694, 0.896393, 0.931526, 0.954346, 0.968887, 0.976952, 0.983536, 0.989174, 0.995110, 0.995262, 0.995323, 0.997826, 0.997688, 0.997101, 0.998741, 0.999519, 0.999809, 0.999985, 0.999489, 1.000000, 1.000000, 0.999626, 0.999977, 1.000000, 1.000000, 0.999924, 0.999374, 1.000000, 1.000000, 0.999901, 1.000000, 1.000000, 1.000000, 1.000000, 1.000000, 1.000000, 1.000000, 1.000000, 1.000000, 1.000000, 1.000000, 1.000000, 1.000000, 1.000000, 1.000000, 1.000000, 1.000000, 1.000000, 1.000000, 1.000000, 1.000000, 1.000000, 1.000000, 1.000000, 1.000000, 1.000000, 1.000000, 1.000000, 1.000000, 1.000000,
+		0.211052, 0.481041, 0.659256, 0.774529, 0.850517, 0.902763, 0.947754, 0.962341, 0.973381, 0.981850, 0.984558, 0.990578, 0.992073, 0.991371, 0.996262, 0.998383, 0.995682, 0.999008, 0.999786, 0.998871, 0.999123, 0.999809, 0.999931, 0.999832, 1.000000, 0.999466, 0.999947, 1.000000, 1.000000, 1.000000, 1.000000, 1.000000, 1.000000, 1.000000, 1.000000, 1.000000, 0.999954, 1.000000, 1.000000, 1.000000, 1.000000, 1.000000, 1.000000, 1.000000, 1.000000, 1.000000, 1.000000, 1.000000, 1.000000, 1.000000, 1.000000, 1.000000, 1.000000, 1.000000, 1.000000, 1.000000, 1.000000, 1.000000, 1.000000, 1.000000, 1.000000, 1.000000, 1.000000, 1.000000,
+		0.220139, 0.503319, 0.671234, 0.787834, 0.860901, 0.918205, 0.943840, 0.962631, 0.979546, 0.988197, 0.988914, 0.992943, 0.995140, 0.996895, 0.997475, 0.999001, 0.998734, 0.999107, 0.998596, 0.999985, 0.999710, 0.999886, 0.999573, 1.000000, 1.000000, 1.000000, 0.999725, 1.000000, 1.000000, 1.000000, 1.000000, 1.000000, 1.000000, 1.000000, 1.000000, 1.000000, 1.000000, 1.000000, 1.000000, 1.000000, 1.000000, 1.000000, 1.000000, 1.000000, 1.000000, 1.000000, 1.000000, 1.000000, 1.000000, 1.000000, 1.000000, 1.000000, 1.000000, 1.000000, 1.000000, 1.000000, 1.000000, 1.000000, 1.000000, 1.000000, 1.000000, 1.000000, 1.000000, 1.000000,
+		0.225166, 0.509087, 0.683983, 0.804581, 0.873917, 0.912231, 0.947807, 0.968803, 0.978722, 0.985580, 0.992889, 0.996399, 0.994324, 0.998238, 0.997505, 0.997810, 0.999580, 0.999031, 0.999031, 1.000000, 0.999405, 0.999565, 0.999954, 1.000000, 0.999962, 1.000000, 1.000000, 1.000000, 1.000000, 1.000000, 1.000000, 1.000000, 1.000000, 1.000000, 1.000000, 1.000000, 1.000000, 1.000000, 1.000000, 1.000000, 1.000000, 1.000000, 1.000000, 1.000000, 1.000000, 1.000000, 1.000000, 1.000000, 1.000000, 1.000000, 1.000000, 1.000000, 1.000000, 1.000000, 1.000000, 1.000000, 1.000000, 1.000000, 1.000000, 1.000000, 1.000000, 1.000000, 1.000000, 1.000000,
+		0.230301, 0.524612, 0.698334, 0.822609, 0.886528, 0.937378, 0.956963, 0.971664, 0.981796, 0.988419, 0.992538, 0.995132, 0.994629, 0.998329, 0.999405, 0.999649, 0.999252, 0.999107, 0.999504, 0.999763, 0.999802, 0.999733, 0.999535, 1.000000, 1.000000, 1.000000, 1.000000, 0.999992, 1.000000, 1.000000, 1.000000, 1.000000, 0.999954, 1.000000, 1.000000, 0.999969, 1.000000, 1.000000, 1.000000, 1.000000, 1.000000, 1.000000, 1.000000, 1.000000, 1.000000, 1.000000, 1.000000, 1.000000, 1.000000, 1.000000, 1.000000, 1.000000, 1.000000, 1.000000, 1.000000, 1.000000, 1.000000, 1.000000, 1.000000, 1.000000, 1.000000, 1.000000, 1.000000, 1.000000,
+		0.238495, 0.530701, 0.707237, 0.817604, 0.888786, 0.931213, 0.962540, 0.976921, 0.982521, 0.982918, 0.990417, 0.996971, 0.997513, 0.997192, 0.998779, 0.999001, 0.999176, 0.999374, 0.999413, 1.000000, 0.998802, 1.000000, 1.000000, 1.000000, 1.000000, 1.000000, 1.000000, 1.000000, 1.000000, 1.000000, 1.000000, 1.000000, 1.000000, 1.000000, 1.000000, 1.000000, 1.000000, 1.000000, 1.000000, 1.000000, 1.000000, 0.999992, 1.000000, 1.000000, 1.000000, 1.000000, 1.000000, 1.000000, 1.000000, 1.000000, 1.000000, 1.000000, 1.000000, 1.000000, 1.000000, 1.000000, 1.000000, 1.000000, 1.000000, 1.000000, 1.000000, 1.000000, 1.000000, 1.000000,
+		0.245865, 0.542244, 0.719765, 0.835083, 0.905663, 0.943886, 0.955551, 0.968971, 0.987480, 0.991135, 0.992340, 0.997025, 0.996613, 0.998848, 0.999878, 0.999641, 0.998367, 0.999413, 1.000000, 0.999435, 0.999901, 1.000000, 1.000000, 1.000000, 1.000000, 1.000000, 1.000000, 1.000000, 1.000000, 1.000000, 1.000000, 1.000000, 1.000000, 1.000000, 1.000000, 1.000000, 1.000000, 1.000000, 1.000000, 1.000000, 1.000000, 1.000000, 1.000000, 1.000000, 1.000000, 1.000000, 1.000000, 1.000000, 1.000000, 1.000000, 1.000000, 1.000000, 1.000000, 1.000000, 1.000000, 1.000000, 1.000000, 1.000000, 1.000000, 1.000000, 1.000000, 1.000000, 1.000000, 1.000000,
+		0.250404, 0.554688, 0.736572, 0.845070, 0.908997, 0.945251, 0.969360, 0.978600, 0.987839, 0.993118, 0.991402, 0.995766, 0.995193, 0.998695, 0.996956, 0.999435, 0.998276, 0.999557, 1.000000, 1.000000, 1.000000, 0.999527, 0.999947, 0.999733, 0.999512, 0.999924, 1.000000, 1.000000, 0.999908, 1.000000, 1.000000, 1.000000, 1.000000, 1.000000, 1.000000, 1.000000, 1.000000, 1.000000, 1.000000, 1.000000, 1.000000, 1.000000, 1.000000, 1.000000, 1.000000, 1.000000, 1.000000, 1.000000, 1.000000, 1.000000, 1.000000, 1.000000, 1.000000, 1.000000, 1.000000, 1.000000, 1.000000, 1.000000, 1.000000, 1.000000, 1.000000, 1.000000, 1.000000, 1.000000,
+		0.257126, 0.566643, 0.750778, 0.854897, 0.914001, 0.953949, 0.969147, 0.984016, 0.991066, 0.991371, 0.995728, 0.995651, 0.998466, 0.999046, 0.999199, 0.998924, 0.999657, 1.000000, 0.999847, 0.999718, 0.999863, 1.000000, 1.000000, 1.000000, 1.000000, 1.000000, 1.000000, 0.999924, 1.000000, 1.000000, 1.000000, 1.000000, 1.000000, 1.000000, 1.000000, 1.000000, 1.000000, 1.000000, 1.000000, 1.000000, 1.000000, 1.000000, 1.000000, 1.000000, 1.000000, 1.000000, 1.000000, 1.000000, 1.000000, 1.000000, 1.000000, 1.000000, 1.000000, 1.000000, 1.000000, 1.000000, 1.000000, 1.000000, 1.000000, 1.000000, 1.000000, 1.000000, 1.000000, 1.000000,
+		0.266174, 0.571640, 0.755798, 0.860291, 0.915726, 0.956024, 0.974861, 0.982788, 0.987846, 0.995201, 0.996902, 0.996552, 0.998238, 0.998566, 0.999435, 0.999893, 0.998993, 0.997871, 0.999435, 1.000000, 1.000000, 0.999130, 1.000000, 0.999939, 1.000000, 1.000000, 1.000000, 1.000000, 1.000000, 1.000000, 1.000000, 1.000000, 1.000000, 1.000000, 1.000000, 1.000000, 1.000000, 1.000000, 1.000000, 1.000000, 1.000000, 1.000000, 1.000000, 1.000000, 1.000000, 1.000000, 1.000000, 1.000000, 1.000000, 1.000000, 1.000000, 1.000000, 1.000000, 1.000000, 1.000000, 1.000000, 1.000000, 1.000000, 1.000000, 1.000000, 1.000000, 1.000000, 1.000000, 1.000000,
+		0.269890, 0.594246, 0.772141, 0.867409, 0.922417, 0.958870, 0.971848, 0.984039, 0.991684, 0.994217, 0.995880, 0.998108, 0.998093, 0.998947, 0.999718, 0.999611, 0.998413, 1.000000, 1.000000, 0.999840, 1.000000, 1.000000, 1.000000, 1.000000, 1.000000, 1.000000, 1.000000, 1.000000, 1.000000, 1.000000, 1.000000, 1.000000, 1.000000, 1.000000, 1.000000, 1.000000, 1.000000, 1.000000, 1.000000, 1.000000, 1.000000, 1.000000, 1.000000, 1.000000, 1.000000, 1.000000, 1.000000, 1.000000, 1.000000, 1.000000, 1.000000, 1.000000, 1.000000, 1.000000, 1.000000, 1.000000, 1.000000, 1.000000, 1.000000, 1.000000, 1.000000, 1.000000, 1.000000, 1.000000,
+		0.278290, 0.597427, 0.779610, 0.880524, 0.927666, 0.957581, 0.975914, 0.985031, 0.989151, 0.993179, 0.999023, 0.999153, 0.998840, 0.997696, 0.999672, 1.000000, 0.999443, 0.999702, 0.999924, 1.000000, 1.000000, 1.000000, 0.999825, 1.000000, 1.000000, 1.000000, 1.000000, 1.000000, 1.000000, 1.000000, 1.000000, 1.000000, 1.000000, 1.000000, 1.000000, 1.000000, 1.000000, 1.000000, 1.000000, 1.000000, 1.000000, 1.000000, 1.000000, 1.000000, 1.000000, 1.000000, 1.000000, 1.000000, 1.000000, 1.000000, 1.000000, 1.000000, 1.000000, 1.000000, 1.000000, 1.000000, 1.000000, 1.000000, 1.000000, 1.000000, 1.000000, 1.000000, 1.000000, 1.000000,
+		0.284424, 0.607811, 0.780205, 0.878990, 0.933708, 0.958366, 0.982422, 0.990944, 0.995102, 0.995644, 0.997368, 0.998718, 0.999268, 0.999268, 0.999992, 0.999901, 0.999962, 1.000000, 1.000000, 0.999863, 1.000000, 1.000000, 1.000000, 1.000000, 1.000000, 1.000000, 1.000000, 1.000000, 1.000000, 1.000000, 1.000000, 1.000000, 1.000000, 1.000000, 1.000000, 1.000000, 1.000000, 1.000000, 1.000000, 1.000000, 1.000000, 1.000000, 1.000000, 1.000000, 1.000000, 1.000000, 1.000000, 1.000000, 1.000000, 1.000000, 1.000000, 1.000000, 1.000000, 1.000000, 1.000000, 1.000000, 1.000000, 1.000000, 1.000000, 1.000000, 1.000000, 1.000000, 1.000000, 1.000000,
+		0.290627, 0.613869, 0.785736, 0.891724, 0.937843, 0.969086, 0.979057, 0.990097, 0.994537, 0.994766, 0.996559, 0.998047, 0.998680, 0.999443, 0.999252, 0.999977, 1.000000, 0.999466, 1.000000, 1.000000, 1.000000, 1.000000, 1.000000, 1.000000, 0.999794, 1.000000, 1.000000, 1.000000, 1.000000, 1.000000, 1.000000, 1.000000, 1.000000, 1.000000, 1.000000, 1.000000, 1.000000, 1.000000, 1.000000, 1.000000, 1.000000, 1.000000, 1.000000, 1.000000, 1.000000, 1.000000, 1.000000, 1.000000, 1.000000, 1.000000, 1.000000, 1.000000, 1.000000, 1.000000, 1.000000, 1.000000, 1.000000, 1.000000, 1.000000, 1.000000, 1.000000, 1.000000, 1.000000, 1.000000,
+		0.297386, 0.629547, 0.803345, 0.891632, 0.949585, 0.971718, 0.984802, 0.989922, 0.995056, 0.998642, 0.998802, 0.999039, 0.999535, 0.999481, 1.000000, 0.999619, 0.999481, 0.999924, 1.000000, 0.999840, 1.000000, 1.000000, 1.000000, 1.000000, 1.000000, 1.000000, 1.000000, 1.000000, 1.000000, 1.000000, 1.000000, 1.000000, 1.000000, 1.000000, 1.000000, 1.000000, 1.000000, 1.000000, 1.000000, 1.000000, 1.000000, 1.000000, 1.000000, 1.000000, 1.000000, 1.000000, 1.000000, 1.000000, 1.000000, 1.000000, 1.000000, 1.000000, 1.000000, 1.000000, 1.000000, 1.000000, 1.000000, 1.000000, 1.000000, 1.000000, 1.000000, 1.000000, 1.000000, 1.000000,
+		0.303810, 0.631767, 0.804535, 0.902519, 0.948914, 0.969879, 0.988220, 0.991173, 0.997070, 0.996880, 0.997566, 0.999237, 0.999451, 0.999138, 1.000000, 0.999733, 0.999611, 1.000000, 1.000000, 0.999939, 0.999878, 1.000000, 1.000000, 1.000000, 1.000000, 1.000000, 1.000000, 1.000000, 1.000000, 0.999939, 1.000000, 1.000000, 1.000000, 1.000000, 1.000000, 1.000000, 1.000000, 1.000000, 1.000000, 1.000000, 1.000000, 1.000000, 1.000000, 1.000000, 1.000000, 1.000000, 1.000000, 1.000000, 1.000000, 1.000000, 1.000000, 1.000000, 1.000000, 1.000000, 1.000000, 1.000000, 1.000000, 1.000000, 1.000000, 1.000000, 1.000000, 1.000000, 1.000000, 1.000000,
+		0.312042, 0.638069, 0.817337, 0.912064, 0.948082, 0.976067, 0.986984, 0.991264, 0.997993, 0.998634, 0.998901, 0.999413, 0.999786, 0.999802, 0.999916, 1.000000, 0.999939, 0.999672, 1.000000, 0.999969, 0.999908, 1.000000, 1.000000, 1.000000, 1.000000, 0.999985, 1.000000, 1.000000, 1.000000, 1.000000, 1.000000, 1.000000, 1.000000, 1.000000, 1.000000, 1.000000, 1.000000, 1.000000, 1.000000, 1.000000, 1.000000, 1.000000, 1.000000, 1.000000, 1.000000, 1.000000, 1.000000, 1.000000, 1.000000, 1.000000, 1.000000, 1.000000, 1.000000, 1.000000, 1.000000, 1.000000, 1.000000, 1.000000, 1.000000, 1.000000, 1.000000, 1.000000, 1.000000, 1.000000,
+		0.317795, 0.651794, 0.820450, 0.910469, 0.957016, 0.979805, 0.990822, 0.991936, 0.995750, 0.997276, 0.998398, 0.999001, 0.999855, 0.999199, 0.999992, 1.000000, 1.000000, 1.000000, 1.000000, 1.000000, 1.000000, 1.000000, 1.000000, 1.000000, 1.000000, 1.000000, 1.000000, 1.000000, 1.000000, 1.000000, 1.000000, 1.000000, 1.000000, 1.000000, 1.000000, 1.000000, 1.000000, 1.000000, 1.000000, 1.000000, 1.000000, 1.000000, 1.000000, 1.000000, 1.000000, 1.000000, 1.000000, 1.000000, 1.000000, 1.000000, 1.000000, 1.000000, 1.000000, 1.000000, 1.000000, 1.000000, 1.000000, 1.000000, 1.000000, 1.000000, 1.000000, 1.000000, 1.000000, 1.000000,
+		0.321625, 0.660835, 0.828522, 0.910446, 0.960503, 0.975853, 0.987556, 0.993469, 0.996925, 0.998535, 0.998871, 0.999443, 0.999527, 0.999268, 0.999954, 0.999924, 1.000000, 0.999969, 0.999802, 0.999725, 1.000000, 1.000000, 1.000000, 1.000000, 1.000000, 1.000000, 1.000000, 1.000000, 1.000000, 1.000000, 1.000000, 1.000000, 1.000000, 1.000000, 1.000000, 1.000000, 1.000000, 1.000000, 1.000000, 1.000000, 1.000000, 1.000000, 1.000000, 1.000000, 1.000000, 1.000000, 1.000000, 1.000000, 1.000000, 1.000000, 1.000000, 1.000000, 1.000000, 1.000000, 1.000000, 1.000000, 1.000000, 1.000000, 1.000000, 1.000000, 1.000000, 1.000000, 1.000000, 1.000000,
+		0.324356, 0.663864, 0.845032, 0.921822, 0.960854, 0.980316, 0.989243, 0.993690, 0.998207, 0.998024, 0.997673, 0.998962, 0.999588, 1.000000, 0.999870, 1.000000, 0.999649, 0.999878, 1.000000, 1.000000, 1.000000, 0.999931, 1.000000, 1.000000, 1.000000, 1.000000, 1.000000, 1.000000, 1.000000, 1.000000, 1.000000, 1.000000, 1.000000, 1.000000, 1.000000, 1.000000, 1.000000, 1.000000, 1.000000, 1.000000, 1.000000, 1.000000, 1.000000, 1.000000, 1.000000, 1.000000, 1.000000, 1.000000, 1.000000, 1.000000, 1.000000, 1.000000, 1.000000, 1.000000, 1.000000, 1.000000, 1.000000, 1.000000, 1.000000, 1.000000, 1.000000, 1.000000, 1.000000, 1.000000,
+		0.330910, 0.676758, 0.841492, 0.921738, 0.962257, 0.980797, 0.990211, 0.994637, 0.997414, 0.999016, 0.998611, 0.999596, 0.999596, 0.999794, 0.999954, 1.000000, 0.999786, 1.000000, 1.000000, 1.000000, 0.999596, 1.000000, 1.000000, 1.000000, 1.000000, 1.000000, 1.000000, 1.000000, 1.000000, 1.000000, 1.000000, 1.000000, 1.000000, 1.000000, 1.000000, 1.000000, 1.000000, 1.000000, 1.000000, 1.000000, 1.000000, 1.000000, 1.000000, 1.000000, 1.000000, 1.000000, 1.000000, 1.000000, 1.000000, 1.000000, 1.000000, 1.000000, 1.000000, 1.000000, 1.000000, 1.000000, 1.000000, 1.000000, 1.000000, 1.000000, 1.000000, 1.000000, 1.000000, 1.000000,
+		0.340332, 0.687988, 0.851479, 0.919609, 0.966499, 0.984642, 0.986450, 0.996460, 0.997650, 0.997688, 0.999130, 1.000000, 0.999603, 1.000000, 0.999710, 0.999611, 1.000000, 1.000000, 1.000000, 1.000000, 1.000000, 1.000000, 1.000000, 1.000000, 1.000000, 1.000000, 1.000000, 1.000000, 1.000000, 1.000000, 1.000000, 1.000000, 1.000000, 1.000000, 1.000000, 1.000000, 1.000000, 1.000000, 1.000000, 1.000000, 1.000000, 1.000000, 1.000000, 1.000000, 1.000000, 1.000000, 1.000000, 1.000000, 1.000000, 1.000000, 1.000000, 1.000000, 1.000000, 1.000000, 1.000000, 1.000000, 1.000000, 1.000000, 1.000000, 1.000000, 1.000000, 1.000000, 1.000000, 1.000000,
+		0.348244, 0.686165, 0.859581, 0.931076, 0.967926, 0.983055, 0.993011, 0.993179, 0.998367, 0.998245, 0.998032, 0.998596, 0.999916, 1.000000, 0.999962, 0.999252, 0.999870, 1.000000, 1.000000, 1.000000, 1.000000, 1.000000, 1.000000, 1.000000, 1.000000, 1.000000, 1.000000, 1.000000, 1.000000, 1.000000, 1.000000, 1.000000, 1.000000, 1.000000, 1.000000, 1.000000, 1.000000, 1.000000, 1.000000, 1.000000, 1.000000, 1.000000, 1.000000, 1.000000, 1.000000, 1.000000, 1.000000, 1.000000, 1.000000, 1.000000, 1.000000, 1.000000, 1.000000, 1.000000, 1.000000, 1.000000, 1.000000, 1.000000, 1.000000, 1.000000, 1.000000, 1.000000, 1.000000, 1.000000,
+		0.350342, 0.698196, 0.854889, 0.928741, 0.965134, 0.986572, 0.994316, 0.996613, 0.996811, 0.997177, 0.999741, 0.999481, 0.999802, 0.999901, 1.000000, 1.000000, 1.000000, 1.000000, 0.999954, 1.000000, 1.000000, 1.000000, 1.000000, 1.000000, 1.000000, 1.000000, 1.000000, 1.000000, 1.000000, 1.000000, 1.000000, 1.000000, 1.000000, 1.000000, 1.000000, 1.000000, 1.000000, 1.000000, 1.000000, 1.000000, 1.000000, 1.000000, 1.000000, 1.000000, 1.000000, 1.000000, 1.000000, 1.000000, 1.000000, 1.000000, 1.000000, 1.000000, 1.000000, 1.000000, 1.000000, 1.000000, 1.000000, 1.000000, 1.000000, 1.000000, 1.000000, 1.000000, 1.000000, 1.000000,
+		0.356522, 0.709854, 0.865700, 0.937141, 0.976524, 0.983276, 0.993027, 0.996117, 0.999123, 0.998611, 0.999359, 0.999947, 0.999680, 0.999969, 1.000000, 1.000000, 1.000000, 1.000000, 1.000000, 1.000000, 1.000000, 1.000000, 1.000000, 1.000000, 1.000000, 1.000000, 1.000000, 1.000000, 1.000000, 1.000000, 1.000000, 1.000000, 1.000000, 1.000000, 1.000000, 1.000000, 1.000000, 1.000000, 1.000000, 1.000000, 1.000000, 1.000000, 1.000000, 1.000000, 1.000000, 1.000000, 1.000000, 1.000000, 1.000000, 1.000000, 1.000000, 1.000000, 1.000000, 1.000000, 1.000000, 1.000000, 1.000000, 1.000000, 1.000000, 1.000000, 1.000000, 1.000000, 1.000000, 1.000000,
+		0.360924, 0.714699, 0.870094, 0.942703, 0.973259, 0.986771, 0.993820, 0.997528, 0.998154, 0.999390, 0.999428, 0.999550, 0.999924, 0.999939, 0.999695, 1.000000, 0.999908, 1.000000, 1.000000, 1.000000, 1.000000, 1.000000, 1.000000, 1.000000, 1.000000, 1.000000, 1.000000, 1.000000, 1.000000, 1.000000, 1.000000, 1.000000, 1.000000, 1.000000, 1.000000, 1.000000, 1.000000, 1.000000, 1.000000, 1.000000, 1.000000, 1.000000, 1.000000, 1.000000, 1.000000, 1.000000, 1.000000, 1.000000, 1.000000, 1.000000, 1.000000, 1.000000, 1.000000, 1.000000, 1.000000, 1.000000, 1.000000, 1.000000, 1.000000, 1.000000, 1.000000, 1.000000, 1.000000, 1.000000,
+		0.369286, 0.722252, 0.877327, 0.944992, 0.977333, 0.990089, 0.995255, 0.997139, 0.998589, 0.999512, 0.998917, 0.999886, 1.000000, 0.999901, 1.000000, 0.999886, 1.000000, 1.000000, 1.000000, 1.000000, 1.000000, 1.000000, 1.000000, 1.000000, 1.000000, 1.000000, 1.000000, 1.000000, 1.000000, 1.000000, 1.000000, 1.000000, 1.000000, 1.000000, 1.000000, 1.000000, 1.000000, 1.000000, 1.000000, 1.000000, 1.000000, 1.000000, 1.000000, 1.000000, 1.000000, 1.000000, 1.000000, 1.000000, 1.000000, 1.000000, 1.000000, 1.000000, 1.000000, 1.000000, 1.000000, 1.000000, 1.000000, 1.000000, 1.000000, 1.000000, 1.000000, 1.000000, 1.000000, 1.000000,
+		0.374062, 0.725288, 0.884613, 0.952583, 0.975601, 0.989159, 0.996315, 0.998009, 0.998070, 0.999146, 0.999863, 0.999985, 0.999855, 0.999939, 1.000000, 1.000000, 1.000000, 1.000000, 1.000000, 1.000000, 1.000000, 1.000000, 1.000000, 1.000000, 1.000000, 1.000000, 1.000000, 1.000000, 1.000000, 1.000000, 1.000000, 1.000000, 1.000000, 1.000000, 1.000000, 1.000000, 1.000000, 1.000000, 1.000000, 1.000000, 1.000000, 1.000000, 1.000000, 1.000000, 1.000000, 1.000000, 1.000000, 1.000000, 1.000000, 1.000000, 1.000000, 1.000000, 1.000000, 1.000000, 1.000000, 1.000000, 1.000000, 1.000000, 1.000000, 1.000000, 1.000000, 1.000000, 1.000000, 1.000000,
+		0.379021, 0.733742, 0.890366, 0.949318, 0.983437, 0.993240, 0.993996, 0.997612, 0.999283, 0.999596, 0.999496, 0.999977, 0.999809, 0.999924, 0.999931, 1.000000, 1.000000, 1.000000, 1.000000, 1.000000, 1.000000, 1.000000, 1.000000, 1.000000, 1.000000, 1.000000, 1.000000, 1.000000, 1.000000, 1.000000, 1.000000, 1.000000, 1.000000, 1.000000, 1.000000, 1.000000, 1.000000, 1.000000, 1.000000, 1.000000, 1.000000, 1.000000, 1.000000, 1.000000, 1.000000, 1.000000, 1.000000, 1.000000, 1.000000, 1.000000, 1.000000, 1.000000, 1.000000, 1.000000, 1.000000, 1.000000, 1.000000, 1.000000, 1.000000, 1.000000, 1.000000, 1.000000, 1.000000, 1.000000,
+		0.382401, 0.740753, 0.896332, 0.950905, 0.975594, 0.990028, 0.995827, 0.998184, 0.999268, 0.999916, 0.999756, 0.999832, 0.999939, 1.000000, 1.000000, 0.999893, 1.000000, 1.000000, 0.999901, 1.000000, 1.000000, 1.000000, 1.000000, 1.000000, 1.000000, 1.000000, 1.000000, 1.000000, 1.000000, 1.000000, 1.000000, 1.000000, 1.000000, 1.000000, 1.000000, 1.000000, 1.000000, 1.000000, 1.000000, 1.000000, 1.000000, 1.000000, 1.000000, 1.000000, 1.000000, 1.000000, 1.000000, 1.000000, 1.000000, 1.000000, 1.000000, 1.000000, 1.000000, 1.000000, 1.000000, 1.000000, 1.000000, 1.000000, 1.000000, 1.000000, 1.000000, 1.000000, 1.000000, 1.000000,
+		0.393326, 0.748840, 0.893021, 0.960419, 0.982246, 0.990952, 0.997612, 0.998741, 0.999229, 0.999809, 0.999718, 0.999817, 1.000000, 1.000000, 1.000000, 1.000000, 1.000000, 1.000000, 1.000000, 1.000000, 1.000000, 1.000000, 1.000000, 1.000000, 1.000000, 1.000000, 1.000000, 1.000000, 1.000000, 1.000000, 1.000000, 1.000000, 1.000000, 1.000000, 1.000000, 1.000000, 1.000000, 1.000000, 1.000000, 1.000000, 1.000000, 1.000000, 1.000000, 1.000000, 1.000000, 1.000000, 1.000000, 1.000000, 1.000000, 1.000000, 1.000000, 1.000000, 1.000000, 1.000000, 1.000000, 1.000000, 1.000000, 1.000000, 1.000000, 1.000000, 1.000000, 1.000000, 1.000000, 1.000000,
+		0.391319, 0.747688, 0.894875, 0.959045, 0.983551, 0.993164, 0.996681, 0.998161, 0.999359, 0.999626, 0.998711, 0.999969, 1.000000, 1.000000, 1.000000, 1.000000, 1.000000, 1.000000, 1.000000, 1.000000, 1.000000, 1.000000, 1.000000, 1.000000, 1.000000, 1.000000, 1.000000, 1.000000, 1.000000, 1.000000, 1.000000, 1.000000, 1.000000, 1.000000, 1.000000, 1.000000, 1.000000, 1.000000, 1.000000, 1.000000, 1.000000, 1.000000, 1.000000, 1.000000, 1.000000, 1.000000, 1.000000, 1.000000, 1.000000, 1.000000, 1.000000, 1.000000, 1.000000, 1.000000, 1.000000, 1.000000, 1.000000, 1.000000, 1.000000, 1.000000, 1.000000, 1.000000, 1.000000, 1.000000,
+		0.397568, 0.760132, 0.901192, 0.959579, 0.985405, 0.992844, 0.996971, 0.999138, 0.998741, 0.999344, 0.999733, 0.999954, 0.999672, 1.000000, 1.000000, 1.000000, 1.000000, 1.000000, 0.999855, 1.000000, 1.000000, 1.000000, 1.000000, 1.000000, 1.000000, 1.000000, 1.000000, 1.000000, 1.000000, 1.000000, 1.000000, 1.000000, 1.000000, 1.000000, 1.000000, 1.000000, 1.000000, 1.000000, 1.000000, 1.000000, 1.000000, 1.000000, 1.000000, 1.000000, 1.000000, 1.000000, 1.000000, 1.000000, 1.000000, 1.000000, 1.000000, 1.000000, 1.000000, 1.000000, 1.000000, 1.000000, 1.000000, 1.000000, 1.000000, 1.000000, 1.000000, 1.000000, 1.000000, 1.000000,
+		0.405067, 0.760406, 0.905258, 0.962410, 0.987137, 0.993011, 0.997681, 0.999596, 0.998993, 0.999146, 0.999565, 1.000000, 1.000000, 1.000000, 1.000000, 1.000000, 1.000000, 1.000000, 1.000000, 1.000000, 1.000000, 1.000000, 1.000000, 1.000000, 1.000000, 1.000000, 1.000000, 1.000000, 1.000000, 1.000000, 1.000000, 1.000000, 1.000000, 1.000000, 1.000000, 1.000000, 1.000000, 1.000000, 1.000000, 1.000000, 1.000000, 1.000000, 1.000000, 1.000000, 1.000000, 1.000000, 1.000000, 1.000000, 1.000000, 1.000000, 1.000000, 1.000000, 1.000000, 1.000000, 1.000000, 1.000000, 1.000000, 1.000000, 1.000000, 1.000000, 1.000000, 1.000000, 1.000000, 1.000000,
+		0.408173, 0.773369, 0.905670, 0.963364, 0.987007, 0.994469, 0.997124, 0.998253, 0.999489, 0.999954, 0.999748, 0.999748, 1.000000, 1.000000, 1.000000, 1.000000, 1.000000, 1.000000, 1.000000, 1.000000, 1.000000, 1.000000, 1.000000, 1.000000, 1.000000, 1.000000, 1.000000, 1.000000, 1.000000, 1.000000, 1.000000, 1.000000, 1.000000, 1.000000, 1.000000, 1.000000, 1.000000, 1.000000, 1.000000, 1.000000, 1.000000, 1.000000, 1.000000, 1.000000, 1.000000, 1.000000, 1.000000, 1.000000, 1.000000, 1.000000, 1.000000, 1.000000, 1.000000, 1.000000, 1.000000, 1.000000, 1.000000, 1.000000, 1.000000, 1.000000, 1.000000, 1.000000, 1.000000, 1.000000,
+		0.418167, 0.771477, 0.912422, 0.965050, 0.988014, 0.995186, 0.998009, 0.999237, 0.999825, 0.999664, 1.000000, 0.999962, 0.999817, 1.000000, 1.000000, 1.000000, 1.000000, 1.000000, 1.000000, 0.999481, 1.000000, 1.000000, 1.000000, 1.000000, 1.000000, 1.000000, 1.000000, 1.000000, 1.000000, 1.000000, 1.000000, 1.000000, 1.000000, 1.000000, 1.000000, 1.000000, 1.000000, 1.000000, 1.000000, 1.000000, 1.000000, 1.000000, 1.000000, 1.000000, 1.000000, 1.000000, 1.000000, 1.000000, 1.000000, 1.000000, 1.000000, 1.000000, 1.000000, 1.000000, 1.000000, 1.000000, 1.000000, 1.000000, 1.000000, 1.000000, 1.000000, 1.000000, 1.000000, 1.000000,
+		0.419846, 0.781548, 0.921494, 0.965538, 0.985680, 0.996445, 0.997757, 0.999626, 0.999504, 0.999771, 1.000000, 0.999985, 1.000000, 1.000000, 0.999771, 1.000000, 1.000000, 0.999718, 1.000000, 1.000000, 1.000000, 1.000000, 1.000000, 1.000000, 1.000000, 1.000000, 1.000000, 1.000000, 1.000000, 1.000000, 1.000000, 1.000000, 1.000000, 1.000000, 1.000000, 1.000000, 1.000000, 1.000000, 1.000000, 1.000000, 1.000000, 1.000000, 1.000000, 1.000000, 1.000000, 1.000000, 1.000000, 1.000000, 1.000000, 1.000000, 1.000000, 1.000000, 1.000000, 1.000000, 1.000000, 1.000000, 1.000000, 1.000000, 1.000000, 1.000000, 1.000000, 1.000000, 1.000000, 1.000000,
+		0.426842, 0.782837, 0.919655, 0.968567, 0.988304, 0.995544, 0.998573, 0.998650, 0.999329, 0.999252, 1.000000, 0.999916, 0.999992, 1.000000, 1.000000, 1.000000, 1.000000, 1.000000, 1.000000, 1.000000, 1.000000, 1.000000, 1.000000, 1.000000, 1.000000, 1.000000, 1.000000, 1.000000, 1.000000, 1.000000, 1.000000, 1.000000, 1.000000, 1.000000, 1.000000, 1.000000, 1.000000, 1.000000, 1.000000, 1.000000, 1.000000, 1.000000, 1.000000, 1.000000, 1.000000, 1.000000, 1.000000, 1.000000, 1.000000, 1.000000, 1.000000, 1.000000, 1.000000, 1.000000, 1.000000, 1.000000, 1.000000, 1.000000, 1.000000, 1.000000, 1.000000, 1.000000, 1.000000, 1.000000,
+	};
+};
+
+/////////////////////////////////////////////////////////////////////////////////////////
+
+class FHairCoverageLUTCS : public FGlobalShader
+{
+	DECLARE_GLOBAL_SHADER(FHairCoverageLUTCS);
+	SHADER_USE_PARAMETER_STRUCT(FHairCoverageLUTCS, FGlobalShader); 
+
+	BEGIN_SHADER_PARAMETER_STRUCT(FParameters, )
+		SHADER_PARAMETER(FIntPoint, OutputResolution)
+		SHADER_PARAMETER_RDG_BUFFER_SRV(RWStructuredBuffer<float>, InputBuffer)
+		SHADER_PARAMETER_RDG_TEXTURE_UAV(RWTexture2D<float>, OutputTexture)
+	END_SHADER_PARAMETER_STRUCT()
+
+	BEGIN_SHADER_PARAMETER_STRUCT(FUploadParameters, )
+		SHADER_PARAMETER_RDG_BUFFER_UPLOAD(StructuredBuffer<float>, UploadBuffer)
+	END_SHADER_PARAMETER_STRUCT()
+
+public:
+	static bool ShouldCompilePermutation(const FGlobalShaderPermutationParameters& Parameters) { return IsHairStrandsSupported(Parameters.Platform); }
+	static void ModifyCompilationEnvironment(const FGlobalShaderPermutationParameters& Parameters, FShaderCompilerEnvironment& OutEnvironment)
+	{
+		FGlobalShader::ModifyCompilationEnvironment(Parameters, OutEnvironment);
+		OutEnvironment.SetDefine(TEXT("SHADER_HAIRCOVERAGE"), 1);
+	}
+};
+
+IMPLEMENT_GLOBAL_SHADER(FHairCoverageLUTCS, "/Engine/Private/HairStrands/HairStrandsLUT.usf", "MainCS", SF_Compute);
+
+static FRDGTextureRef AddHairCoverageLUTPass(FRDGBuilder& GraphBuilder, const FViewInfo& View)
+{
+	// Create buffer with upload data
+	FIntPoint OutputResolution(FHairCountToCoverageData::HairCount, FHairCountToCoverageData::HairRadiusCount);
+	FRDGBufferRef InputBuffer = nullptr;
+
+	const uint32 ElementCount = FHairCountToCoverageData::HairCount * FHairCountToCoverageData::HairRadiusCount;
+	const uint32 SizeInBytes = ElementCount * sizeof(float);
+	FRDGBufferRef UploadBuffer = GraphBuilder.CreateBuffer(FRDGBufferDesc::CreateStructuredDesc(sizeof(float), ElementCount), TEXT("HairCoverageData"));
+	FHairCoverageLUTCS::FUploadParameters* UploadParameters = GraphBuilder.AllocParameters<FHairCoverageLUTCS::FUploadParameters>();
+	UploadParameters->UploadBuffer = UploadBuffer;
+
+	GraphBuilder.AddPass(
+		RDG_EVENT_NAME("UploadHairCoverageBuffer"),
+		UploadParameters,
+		ERDGPassFlags::Copy,
+		[UploadParameters, SizeInBytes](FRHICommandListImmediate& RHICmdList)
+	{
+		FHairCountToCoverageData Source;
+		
+		void* Dest = RHILockStructuredBuffer(UploadParameters->UploadBuffer->GetRHIStructuredBuffer(), 0, SizeInBytes, RLM_WriteOnly);
+		FPlatformMemory::Memcpy(Dest, Source.Data, SizeInBytes);
+		RHIUnlockStructuredBuffer(UploadParameters->UploadBuffer->GetRHIStructuredBuffer());
+	});
+
+
+	FRDGTextureDesc OutputDesc;
+	OutputDesc.Extent.X = OutputResolution.X;
+	OutputDesc.Extent.Y = OutputResolution.Y;
+	OutputDesc.Format = PF_R32_FLOAT;
+	OutputDesc.NumMips = 1;
+	OutputDesc.Flags = 0;
+	OutputDesc.TargetableFlags = TexCreate_UAV | TexCreate_ShaderResource;
+	FRDGTextureRef HairLUTTexture = GraphBuilder.CreateTexture(OutputDesc, TEXT("HairCoverageLUT"));
+
+	FHairCoverageLUTCS::FParameters* Parameters = GraphBuilder.AllocParameters<FHairCoverageLUTCS::FParameters>();
+	Parameters->OutputResolution = OutputResolution;
+	Parameters->InputBuffer = GraphBuilder.CreateSRV(UploadBuffer);
+	Parameters->OutputTexture = GraphBuilder.CreateUAV(FRDGTextureUAVDesc(HairLUTTexture, 0));
+
+	TShaderMapRef<FHairCoverageLUTCS> ComputeShader(View.ShaderMap);
+	FComputeShaderUtils::AddPass(
+		GraphBuilder,
+		RDG_EVENT_NAME("HairStrandsCoverageLUT"),
+		ComputeShader,
+		Parameters,
+		FComputeShaderUtils::GetGroupCount(OutputResolution, FIntPoint(8,8)));
+
+	return HairLUTTexture;
+}
+
 FHairLUT GetHairLUT(FRHICommandListImmediate& RHICmdList, const FViewInfo& View)
 {
 	// Lazy LUT generation
 	const bool bNeedGenerate = 
-		GSystemTextures.HairLUT0.GetReference() == nullptr || GSystemTextures.HairLUT1.GetReference() == nullptr ||
+		GSystemTextures.HairLUT0.GetReference() == nullptr || GSystemTextures.HairLUT1.GetReference() == nullptr || GSystemTextures.HairLUT2.GetReference() == nullptr ||
 		GSystemTextures.HairLUT0.GetReference()->GetRenderTargetItem().ShaderResourceTexture->GetSizeXYZ() != FIntVector(GHairLUTIncidentAngleCount, GHairLUTRoughnessCount, GHairLUTAbsorptionCount);
 	if (bNeedGenerate)
 	{
@@ -99,11 +280,15 @@ FHairLUT GetHairLUT(FRHICommandListImmediate& RHICmdList, const FViewInfo& View)
 		FRDGTextureRef HairMeanEnergyLUTTexture		= AddHairLUTPass(GraphBuilder, View, HairLUTType_MeanEnergy);		
 		GraphBuilder.QueueTextureExtraction(HairMeanEnergyLUTTexture, &GSystemTextures.HairLUT1);
 
+		FRDGTextureRef HairCoverageLUTTexture = AddHairCoverageLUTPass(GraphBuilder, View);
+		GraphBuilder.QueueTextureExtraction(HairCoverageLUTTexture, &GSystemTextures.HairLUT2);
+
 		GraphBuilder.Execute();
 	}
 	FHairLUT HairLUTData;
 	HairLUTData.Textures[HairLUTType_DualScattering] = GSystemTextures.HairLUT0;
 	HairLUTData.Textures[HairLUTType_MeanEnergy] = GSystemTextures.HairLUT1;
+	HairLUTData.Textures[HairLUTType_Coverage] = GSystemTextures.HairLUT2;
 
 	return HairLUTData;
 }
diff --git a/Engine/Source/Runtime/Renderer/Private/HairStrands/HairStrandsLUT.h b/Engine/Source/Runtime/Renderer/Private/HairStrands/HairStrandsLUT.h
--- a/Engine/Source/Runtime/Renderer/Private/HairStrands/HairStrandsLUT.h
+++ b/Engine/Source/Runtime/Renderer/Private/HairStrands/HairStrandsLUT.h
@@ -13,6 +13,7 @@ enum FHairLUTType
 {
 	HairLUTType_DualScattering,
 	HairLUTType_MeanEnergy,
+	HairLUTType_Coverage,
 	HairLUTTypeCount
 };
 
diff --git a/Engine/Source/Runtime/Renderer/Private/HairStrands/HairStrandsVisibility.cpp b/Engine/Source/Runtime/Renderer/Private/HairStrands/HairStrandsVisibility.cpp
--- a/Engine/Source/Runtime/Renderer/Private/HairStrands/HairStrandsVisibility.cpp
+++ b/Engine/Source/Runtime/Renderer/Private/HairStrands/HairStrandsVisibility.cpp
@@ -66,10 +66,12 @@ static FAutoConsoleVariableRef CVarHairStrandsVisibilityComputeRasterSampleCount
 static float GHairStrandsFullCoverageThreshold = 0.98f;
 static FAutoConsoleVariableRef CVarHairStrandsFullCoverageThreshold(TEXT("r.HairStrands.Visibility.FullCoverageThreshold"), GHairStrandsFullCoverageThreshold, TEXT("Define the coverage threshold at which a pixel is considered fully covered."));
 
-
 static int32 GHairStrandsSortHairSampleByDepth = 0;
 static FAutoConsoleVariableRef CVarHairStrandsSortHairSampleByDepth(TEXT("r.HairStrands.Visibility.SortByDepth"), GHairStrandsSortHairSampleByDepth, TEXT("Sort hair fragment by depth and update their coverage based on ordered transmittance."));
 
+static int32 GHairStrandsHairCountToTransmittance = 0;
+static FAutoConsoleVariableRef CVarHairStrandsHairCountToTransmittance(TEXT("r.HairStrands.Visibility.UseCoverageMappping"), GHairStrandsHairCountToTransmittance, TEXT("Use hair count to coverage transfer function."));
+
 /////////////////////////////////////////////////////////////////////////////////////////
 
 namespace HairStrandsVisibilityInternal
@@ -1799,7 +1801,7 @@ static void AddHairVisibilityCommonPass(
 			{
 				DrawRenderState.SetBlendState(TStaticBlendState<
 					CW_RED, BO_Add, BF_DestColor, BF_Zero, BO_Add, BF_Zero, BF_Zero,
-					CW_RED, BO_Add, BF_One, BF_One, BO_Add, BF_Zero, BF_Zero>::GetRHI());
+					CW_RG, BO_Add, BF_One, BF_One, BO_Add, BF_Zero, BF_Zero>::GetRHI());
 				DrawRenderState.SetDepthStencilState(TStaticDepthStencilState<false, CF_DepthNearOrEqual>::GetRHI());
 			}
 			else if (RenderMode == HairVisibilityRenderMode_PPLL)
@@ -2072,6 +2074,7 @@ static FHairPrimaryTransmittance AddHairViewTransmittancePass(
 
 	if (RenderMode == HairVisibilityRenderMode_TransmittanceAndHairCount)
 	{
+		Desc.Format = PF_G32R32F;
 		Desc.ClearValue = FClearValueBinding(FLinearColor(0.0f, 0.0f, 0.0f, 0.0f));
 		Out.HairCountTexture = GraphBuilder.CreateTexture(Desc, TEXT("HairViewHairCountTexture"));
 		PassParameters->RenderTargets[1] = FRenderTargetBinding(Out.HairCountTexture, ERenderTargetLoadAction::EClear, 0);
@@ -2269,6 +2272,68 @@ static void AddHairVisibilityColorAndDepthPatchPass(
 
 ///////////////////////////////////////////////////////////////////////////////////////////////////
 
+class FHairCountToCoverageCS : public FGlobalShader
+{
+	DECLARE_GLOBAL_SHADER(FHairCountToCoverageCS);
+	SHADER_USE_PARAMETER_STRUCT(FHairCountToCoverageCS, FGlobalShader);
+
+	BEGIN_SHADER_PARAMETER_STRUCT(FParameters, )
+		SHADER_PARAMETER(FIntPoint, OutputResolution)
+		SHADER_PARAMETER(float, LUT_HairCount)
+		SHADER_PARAMETER(float, LUT_HairRadiusCount)
+		SHADER_PARAMETER_SAMPLER(SamplerState, LinearSampler)
+		SHADER_PARAMETER_RDG_TEXTURE(Texture2D, HairCoverageLUT)
+		SHADER_PARAMETER_RDG_TEXTURE(Texture2D, HairCountTexture)
+		SHADER_PARAMETER_RDG_TEXTURE_UAV(RWTexture2D, OutputTexture)
+	END_SHADER_PARAMETER_STRUCT()
+
+public:
+	static bool ShouldCompilePermutation(const FGlobalShaderPermutationParameters& Parameters)
+	{
+		return IsHairStrandsSupported(Parameters.Platform);
+	}
+};
+
+IMPLEMENT_GLOBAL_SHADER(FHairCountToCoverageCS, "/Engine/Private/HairStrands/HairStrandsCoverage.usf", "MainCS", SF_Compute);
+
+static FRDGTextureRef AddHairHairCountToTransmittancePass(
+	FRDGBuilder& GraphBuilder,
+	const FViewInfo& ViewInfo,
+	const FHairLUT& HairLUT,
+	const FRDGTextureRef HairCountTexture)
+{
+	const FIntPoint OutputResolution = HairCountTexture->Desc.Extent;
+
+	FRDGTextureDesc Desc;
+	Desc.Extent = OutputResolution;
+	Desc.Depth = 0;
+	Desc.Format = PF_R32_FLOAT;
+	Desc.NumMips = 1;
+	Desc.NumSamples = 1;
+	Desc.Flags = TexCreate_None;
+	Desc.TargetableFlags = TexCreate_UAV | TexCreate_ShaderResource | TexCreate_RenderTargetable;
+	Desc.bForceSharedTargetAndShaderResource = true;
+	Desc.ClearValue = FClearValueBinding(FLinearColor(0.0f, 0.0f, 0.0f, 0.0f));
+	FRDGTextureRef OutputTexture = GraphBuilder.CreateTexture(Desc, TEXT("HairVisibilityTexture"));
+	FRDGTextureRef HairCoverageLUT = GraphBuilder.RegisterExternalTexture(HairLUT.Textures[HairLUTType_Coverage], TEXT("HairCoverageLUT"));
+
+	FHairCountToCoverageCS::FParameters* PassParameters = GraphBuilder.AllocParameters<FHairCountToCoverageCS::FParameters>();
+	PassParameters->LUT_HairCount = HairCoverageLUT->Desc.Extent.X;
+	PassParameters->LUT_HairRadiusCount = HairCoverageLUT->Desc.Extent.Y;
+	PassParameters->OutputResolution = OutputResolution;
+	PassParameters->HairCoverageLUT = HairCoverageLUT;
+	PassParameters->HairCountTexture = HairCountTexture;
+	PassParameters->LinearSampler = TStaticSamplerState<SF_Bilinear, AM_Clamp, AM_Clamp, AM_Clamp>::GetRHI();
+	PassParameters->OutputTexture = GraphBuilder.CreateUAV(OutputTexture);
+
+	TShaderMapRef<FHairCountToCoverageCS> ComputeShader(ViewInfo.ShaderMap);
+	FComputeShaderUtils::AddPass(GraphBuilder, RDG_EVENT_NAME("HairStrandsVisibilityComputeRaster"), ComputeShader, PassParameters, FComputeShaderUtils::GetGroupCount(OutputResolution, FIntPoint(8,8)));
+
+	return OutputTexture;
+}
+
+///////////////////////////////////////////////////////////////////////////////////////////////////
+
 class FVisiblityRasterComputeCS : public FGlobalShader
 {
 	DECLARE_GLOBAL_SHADER(FVisiblityRasterComputeCS);
@@ -2480,6 +2545,8 @@ FHairStrandsVisibilityViews RenderHairStrandsVisibilityBuffer(
 		const FViewInfo& View = Views[ViewIndex];
 		if (View.Family)
 		{
+			FHairLUT HairLUT = GetHairLUT(RHICmdList, View);
+
 			FHairStrandsVisibilityData& VisibilityData = Output.HairDatas.AddDefaulted_GetRef();
 			VisibilityData.NodeGroupSize = GetVendorOptimalGroupSize1D();
 			const FHairStrandsMacroGroupDatas& MacroGroupDatas = MacroGroupViews.Views[ViewIndex];
@@ -2518,6 +2585,16 @@ FHairStrandsVisibilityViews RenderHairStrandsVisibilityBuffer(
 					bOutputHairCount,
 					SceneDepthTexture);
 
+				const bool bHairCountToTransmittance = GHairStrandsHairCountToTransmittance > 0;
+				if (bHairCountToTransmittance)
+				{
+					ViewTransmittance.TransmittanceTexture = AddHairHairCountToTransmittancePass(
+						GraphBuilder,
+						View,
+						HairLUT,
+						ViewTransmittance.HairCountTexture);
+				}
+
 				const bool bUseRasterCompute = GHairStrandsVisibilityComputeRaster > 0 && DoesSupportRasterCompute();
 				if (bUseRasterCompute)
 				{
diff --git a/Engine/Source/Runtime/Renderer/Private/SystemTextures.h b/Engine/Source/Runtime/Renderer/Private/SystemTextures.h
--- a/Engine/Source/Runtime/Renderer/Private/SystemTextures.h
+++ b/Engine/Source/Runtime/Renderer/Private/SystemTextures.h
@@ -73,6 +73,7 @@ public:
 	/** Hair BSDF LUT texture */
 	TRefCountPtr<IPooledRenderTarget> HairLUT0;
 	TRefCountPtr<IPooledRenderTarget> HairLUT1;
+	TRefCountPtr<IPooledRenderTarget> HairLUT2;
 	/** Linearly Transformed Cosines LUTs */
 	TRefCountPtr<IPooledRenderTarget> LTCMat;
 	TRefCountPtr<IPooledRenderTarget> LTCAmp;
