Fixed crash when changing the quality settings.

An FGlobalComponentRecreateRenderStateContext object was created immediately after changing the quality level, which ended up trying to fetch material shader maps for the new level before they are created.

Also found a number of places which were calling FPrimitiveSceneInfo::CacheMeshDrawCommands with nothing to do. That function does a non-trivial amount of work even if the array of scene infos is empty, so I added empty checks to guard against it.

#rnx

diff --git a/Engine/Source/Runtime/Engine/Private/UnrealEngine.cpp b/Engine/Source/Runtime/Engine/Private/UnrealEngine.cpp
--- a/Engine/Source/Runtime/Engine/Private/UnrealEngine.cpp
+++ b/Engine/Source/Runtime/Engine/Private/UnrealEngine.cpp
@@ -602,12 +602,14 @@ void ScalabilityCVarsSinkCallback()
 
 		if (bRecreateRenderstate || bCacheResourceShaders)
 		{
+			// Make the render state rebuild object before updating the cached values, because its constructor calls UpdateAllPrimitiveSceneInfos(), which
+			// may end up using the material shader maps for the quality level stored in GCachedScalabilityCVars.MaterialQualityLevel. If that value already
+			// points to the new quality level, the shader maps won't exist, and we'll crash.
+			FGlobalComponentRecreateRenderStateContext Recreate;
+
 			// after FlushRenderingCommands() to not have render thread pick up the data partially
 			GCachedScalabilityCVars = LocalScalabilityCVars;
 
-			// Note: constructor and destructor has side effect
-			FGlobalComponentRecreateRenderStateContext Recreate;
-
 			if (bCacheResourceShaders)
 			{
 				// For all materials, UMaterial::CacheResourceShadersForRendering
diff --git a/Engine/Source/Runtime/Renderer/Private/RendererScene.cpp b/Engine/Source/Runtime/Renderer/Private/RendererScene.cpp
--- a/Engine/Source/Runtime/Renderer/Private/RendererScene.cpp
+++ b/Engine/Source/Runtime/Renderer/Private/RendererScene.cpp
@@ -4115,12 +4115,18 @@ void FScene::UpdateAllPrimitiveSceneInfos(FRHICommandListImmediate& RHICmdList,
 		}
 
 		// Re-add the primitive to the scene with the new transform.
-		FPrimitiveSceneInfo::AddToScene(RHICmdList, this, UpdatedSceneInfosWithStaticDrawListUpdate, true, true, bAsyncCreateLPIs);
-		FPrimitiveSceneInfo::AddToScene(RHICmdList, this, UpdatedSceneInfosWithoutStaticDrawListUpdate, false, true, bAsyncCreateLPIs);
+		if (UpdatedSceneInfosWithStaticDrawListUpdate.Num() > 0)
+		{
+			FPrimitiveSceneInfo::AddToScene(RHICmdList, this, UpdatedSceneInfosWithStaticDrawListUpdate, true, true, bAsyncCreateLPIs);
+		}
 
-		for (FPrimitiveSceneInfo* PrimitiveSceneInfo : UpdatedSceneInfosWithoutStaticDrawListUpdate)
+		if (UpdatedSceneInfosWithoutStaticDrawListUpdate.Num() > 0)
 		{
-			PrimitiveSceneInfo->FlushRuntimeVirtualTexture();
+			FPrimitiveSceneInfo::AddToScene(RHICmdList, this, UpdatedSceneInfosWithoutStaticDrawListUpdate, false, true, bAsyncCreateLPIs);
+			for (FPrimitiveSceneInfo* PrimitiveSceneInfo : UpdatedSceneInfosWithoutStaticDrawListUpdate)
+			{
+				PrimitiveSceneInfo->FlushRuntimeVirtualTexture();
+			}
 		}
 
 		if (AsyncCreateLightPrimitiveInteractionsTask && AsyncCreateLightPrimitiveInteractionsTask->GetTask().HasPendingPrimitives())
diff --git a/Engine/Source/Runtime/Renderer/Private/SceneVisibility.cpp b/Engine/Source/Runtime/Renderer/Private/SceneVisibility.cpp
--- a/Engine/Source/Runtime/Renderer/Private/SceneVisibility.cpp
+++ b/Engine/Source/Runtime/Renderer/Private/SceneVisibility.cpp
@@ -3710,7 +3710,10 @@ void FSceneRenderer::ComputeViewVisibility(FRHICommandListImmediate& RHICmdList,
 				UpdatedSceneInfos.Add(Primitive);
 			}
 		}
-		FPrimitiveSceneInfo::UpdateStaticMeshes(RHICmdList, Scene, UpdatedSceneInfos);
+		if (UpdatedSceneInfos.Num() > 0)
+		{
+			FPrimitiveSceneInfo::UpdateStaticMeshes(RHICmdList, Scene, UpdatedSceneInfos);
+		}
 		Scene->PrimitivesNeedingStaticMeshUpdateWithoutVisibilityCheck.Reset();
 	}
 
@@ -3940,7 +3943,11 @@ void FSceneRenderer::ComputeViewVisibility(FRHICommandListImmediate& RHICmdList,
 				int32 PrimitiveIndex = BitIt.GetIndex();
 				AddedSceneInfos.Add(Scene->Primitives[PrimitiveIndex]);
 			}
-			FPrimitiveSceneInfo::UpdateStaticMeshes(RHICmdList, Scene, AddedSceneInfos);
+
+			if (AddedSceneInfos.Num() > 0)
+			{
+				FPrimitiveSceneInfo::UpdateStaticMeshes(RHICmdList, Scene, AddedSceneInfos);
+			}
 		}
 
 		// ISR views can't compute relevance until all views are frustum culled
