Fix for construction scripts adding / removing components interfering with deferred adds, causing components to be added multiple times.

diff --git a/Engine/Source/Runtime/Engine/Classes/Components/ActorComponent.h b/Engine/Source/Runtime/Engine/Classes/Components/ActorComponent.h
--- a/Engine/Source/Runtime/Engine/Classes/Components/ActorComponent.h
+++ b/Engine/Source/Runtime/Engine/Classes/Components/ActorComponent.h
@@ -22,10 +22,27 @@ class AActor;
 class UActorComponent;
 class UAssetUserData;
 class ULevel;
+class UWorld;
+class UPrimitiveComponent;
 
-struct FRegisterComponentContext
+class FRegisterComponentContext
 {
-	TArray<class UPrimitiveComponent*> AddPrimitiveBatches;
+public:
+	FRegisterComponentContext(UWorld* InWorld)
+		: World(InWorld)
+	{}
+
+	void AddPrimitive(UPrimitiveComponent* PrimitiveComponent)
+	{
+		checkSlow(!AddPrimitiveBatches.Contains(PrimitiveComponent));
+		AddPrimitiveBatches.Add(PrimitiveComponent);
+	}
+
+	void Process();
+
+private:
+	UWorld* World;
+	TArray<UPrimitiveComponent*> AddPrimitiveBatches;
 };
 
 #if WITH_EDITOR
diff --git a/Engine/Source/Runtime/Engine/Classes/Engine/Level.h b/Engine/Source/Runtime/Engine/Classes/Engine/Level.h
--- a/Engine/Source/Runtime/Engine/Classes/Engine/Level.h
+++ b/Engine/Source/Runtime/Engine/Classes/Engine/Level.h
@@ -31,6 +31,7 @@ class UTexture2D;
 struct FLevelCollection;
 class ULevelActorContainer;
 class FLevelPartitionOperationScope;
+class FRegisterComponentContext;
 
 UINTERFACE()
 class ULevelPartitionInterface : public UInterface
@@ -766,7 +767,7 @@ public:
 	 * creates the BSP model components.
 	 * @param bRerunConstructionScripts	If we want to rerun construction scripts on actors in level
 	 */
-	ENGINE_API void UpdateLevelComponents(bool bRerunConstructionScripts, struct FRegisterComponentContext* Context = nullptr);
+	ENGINE_API void UpdateLevelComponents(bool bRerunConstructionScripts, FRegisterComponentContext* Context = nullptr);
 
 	/**
 	 * Incrementally updates all components of actors associated with this level.
@@ -774,7 +775,7 @@ public:
 	 * @param NumComponentsToUpdate		Number of components to update in this run, 0 for all
 	 * @param bRerunConstructionScripts	If we want to rerun construction scripts on actors in level
 	 */
-	void IncrementalUpdateComponents( int32 NumComponentsToUpdate, bool bRerunConstructionScripts, struct FRegisterComponentContext* Context = nullptr);
+	void IncrementalUpdateComponents( int32 NumComponentsToUpdate, bool bRerunConstructionScripts, FRegisterComponentContext* Context = nullptr);
 
 	/**
 	* Incrementally unregisters all components of actors associated with this level.
diff --git a/Engine/Source/Runtime/Engine/Private/Components/ActorComponent.cpp b/Engine/Source/Runtime/Engine/Private/Components/ActorComponent.cpp
--- a/Engine/Source/Runtime/Engine/Private/Components/ActorComponent.cpp
+++ b/Engine/Source/Runtime/Engine/Private/Components/ActorComponent.cpp
@@ -32,6 +32,7 @@
 #include "HAL/LowLevelMemTracker.h"
 #include "Net/Core/PushModel/PushModel.h"
 #include "UObject/FrameworkObjectVersion.h"
+#include "Async/ParallelFor.h"
 
 #if WITH_EDITOR
 #include "Kismet2/ComponentEditorUtils.h"
@@ -78,6 +79,22 @@ FUObjectAnnotationSparseBool GSelectedComponentAnnotation;
 /** Static var indicating activity of reregister context */
 int32 FGlobalComponentReregisterContext::ActiveGlobalReregisterContextCount = 0;
 
+void FRegisterComponentContext::Process()
+{
+	FSceneInterface* Scene = World->Scene;
+	ParallelFor(AddPrimitiveBatches.Num(),
+		[&](int32 Index)
+		{
+			if (!AddPrimitiveBatches[Index]->IsPendingKill())
+			{
+				Scene->AddPrimitive(AddPrimitiveBatches[Index]);
+			}
+		},
+		!FApp::ShouldUseThreadingForPerformance()
+	);
+	AddPrimitiveBatches.Empty();
+}
+
 void UpdateAllPrimitiveSceneInfosForSingleComponent(UActorComponent* InComponent, TSet<FSceneInterface*>* InScenesToUpdateAllPrimitiveSceneInfosForBatching /* = nullptr*/)
 {
 	if (FSceneInterface* Scene = InComponent->GetScene())
diff --git a/Engine/Source/Runtime/Engine/Private/Components/PrimitiveComponent.cpp b/Engine/Source/Runtime/Engine/Private/Components/PrimitiveComponent.cpp
--- a/Engine/Source/Runtime/Engine/Private/Components/PrimitiveComponent.cpp
+++ b/Engine/Source/Runtime/Engine/Private/Components/PrimitiveComponent.cpp
@@ -555,7 +555,7 @@ void UPrimitiveComponent::CreateRenderState_Concurrent(FRegisterComponentContext
 	{
 		if (Context != nullptr)
 		{
-			Context->AddPrimitiveBatches.Add(this);
+			Context->AddPrimitive(this);
 		}
 		else
 		{
diff --git a/Engine/Source/Runtime/Engine/Private/Level.cpp b/Engine/Source/Runtime/Engine/Private/Level.cpp
--- a/Engine/Source/Runtime/Engine/Private/Level.cpp
+++ b/Engine/Source/Runtime/Engine/Private/Level.cpp
@@ -985,6 +985,13 @@ void ULevel::IncrementalUpdateComponents(int32 NumComponentsToUpdate, bool bReru
 #endif
 		if (bRerunConstructionScripts && !IsTemplate() && !GIsUCCMakeStandaloneHeaderGenerator)
 		{
+			// We need to process pending adds prior to rerunning the construction scripts, which may internally
+			// perform removals / adds themselves.
+			if (Context)
+			{
+				Context->Process();
+			}
+
 			// Don't rerun construction scripts until after all actors' components have been registered.  This
 			// is necessary because child attachment lists are populated during registration, and running construction
 			// scripts requires that the attachments are correctly initialized.
diff --git a/Engine/Source/Runtime/Engine/Private/UnrealEngine.cpp b/Engine/Source/Runtime/Engine/Private/UnrealEngine.cpp
--- a/Engine/Source/Runtime/Engine/Private/UnrealEngine.cpp
+++ b/Engine/Source/Runtime/Engine/Private/UnrealEngine.cpp
@@ -12896,21 +12896,12 @@ bool UEngine::LoadMap( FWorldContext& WorldContext, FURL URL, class UPendingNetG
 	// Note that AI system will be created only if ai-system-creation conditions are met
 	WorldContext.World()->CreateAISystem();
 
-	FRegisterComponentContext Context;
 	// Initialize gameplay for the level.
-	WorldContext.World()->InitializeActorsForPlay(URL, true, &Context);
-
-	UWorld* ParallelWorld = WorldContext.World();
-	ParallelFor(Context.AddPrimitiveBatches.Num(),
-		[&Context, ParallelWorld](int32 Index)
-		{
-			if (!Context.AddPrimitiveBatches[Index]->IsPendingKill())
-			{
-				ParallelWorld->Scene->AddPrimitive(Context.AddPrimitiveBatches[Index]);
-			}
-		},
-		!FApp::ShouldUseThreadingForPerformance()
-	);
+	{
+		FRegisterComponentContext Context(WorldContext.World());
+		WorldContext.World()->InitializeActorsForPlay(URL, true, &Context);
+		Context.Process();
+	}
 
 	// calling it after InitializeActorsForPlay has been called to have all potential bounding boxed initialized
 	FNavigationSystem::AddNavigationSystemToWorld(*WorldContext.World(), FNavigationSystemRunMode::GameMode);
diff --git a/Engine/Source/Runtime/Renderer/Private/RendererScene.cpp b/Engine/Source/Runtime/Renderer/Private/RendererScene.cpp
--- a/Engine/Source/Runtime/Renderer/Private/RendererScene.cpp
+++ b/Engine/Source/Runtime/Renderer/Private/RendererScene.cpp
@@ -1284,6 +1284,8 @@ void FScene::AddPrimitive(UPrimitiveComponent* Primitive)
 		Primitive->LastSubmitTime = WorldTime;
 	}
 
+	checkf(!Primitive->SceneProxy, TEXT("Primitive has already been added to the scene!"));
+
 	// Create the primitive's scene proxy.
 	FPrimitiveSceneProxy* PrimitiveSceneProxy = Primitive->CreateSceneProxy();
 	Primitive->SceneProxy = PrimitiveSceneProxy;
@@ -1618,6 +1620,15 @@ void FScene::RemovePrimitive( UPrimitiveComponent* Primitive )
 
 	FPrimitiveSceneProxy* PrimitiveSceneProxy = Primitive->SceneProxy;
 
+	// Send a command to the rendering thread to release the scene.
+	FScene* Scene = this;
+	ENQUEUE_RENDER_COMMAND(FReleaseCommand)(
+		[Scene](FRHICommandListImmediate& RHICmdList)
+	{
+		// Flush any remaining batched primitive update commands before deleting the scene.
+		Scene->UpdateAllPrimitiveSceneInfos(RHICmdList);
+	});
+
 	if(PrimitiveSceneProxy)
 	{
 		FPrimitiveSceneInfo* PrimitiveSceneInfo = PrimitiveSceneProxy->GetPrimitiveSceneInfo();
@@ -1626,7 +1637,7 @@ void FScene::RemovePrimitive( UPrimitiveComponent* Primitive )
 		Primitive->SceneProxy = NULL;
 
 		// Send a command to the rendering thread to remove the primitive from the scene.
-		FScene* Scene = this;
+		//FScene* Scene = this;
 		FThreadSafeCounter* AttachmentCounter = &Primitive->AttachmentCounter;
 		ENQUEUE_RENDER_COMMAND(FRemovePrimitiveCommand)(
 			[Scene, PrimitiveSceneInfo, AttachmentCounter](FRHICommandList&)
