Reverted debug code that shouldn't have been checked in.

diff --git a/Engine/Source/Runtime/Renderer/Private/RendererScene.cpp b/Engine/Source/Runtime/Renderer/Private/RendererScene.cpp
--- a/Engine/Source/Runtime/Renderer/Private/RendererScene.cpp
+++ b/Engine/Source/Runtime/Renderer/Private/RendererScene.cpp
@@ -1620,15 +1620,6 @@ void FScene::RemovePrimitive( UPrimitiveComponent* Primitive )
 
 	FPrimitiveSceneProxy* PrimitiveSceneProxy = Primitive->SceneProxy;
 
-	// Send a command to the rendering thread to release the scene.
-	FScene* Scene = this;
-	ENQUEUE_RENDER_COMMAND(FReleaseCommand)(
-		[Scene](FRHICommandListImmediate& RHICmdList)
-	{
-		// Flush any remaining batched primitive update commands before deleting the scene.
-		Scene->UpdateAllPrimitiveSceneInfos(RHICmdList);
-	});
-
 	if(PrimitiveSceneProxy)
 	{
 		FPrimitiveSceneInfo* PrimitiveSceneInfo = PrimitiveSceneProxy->GetPrimitiveSceneInfo();
@@ -1637,7 +1628,7 @@ void FScene::RemovePrimitive( UPrimitiveComponent* Primitive )
 		Primitive->SceneProxy = NULL;
 
 		// Send a command to the rendering thread to remove the primitive from the scene.
-		//FScene* Scene = this;
+		FScene* Scene = this;
 		FThreadSafeCounter* AttachmentCounter = &Primitive->AttachmentCounter;
 		ENQUEUE_RENDER_COMMAND(FRemovePrimitiveCommand)(
 			[Scene, PrimitiveSceneInfo, AttachmentCounter](FRHICommandList&)
