REL - This shallow copy of Create Snapshot causes the RobinHoodHashTable to leak it's support data structure as empty is not cleaning it up as it needs to be available to not polute the insertion loops with init guards.

diff --git a/Engine/Source/Runtime/Core/Public/Experimental/Containers/RobinHoodHashTable.h b/Engine/Source/Runtime/Core/Public/Experimental/Containers/RobinHoodHashTable.h
--- a/Engine/Source/Runtime/Core/Public/Experimental/Containers/RobinHoodHashTable.h
+++ b/Engine/Source/Runtime/Core/Public/Experimental/Containers/RobinHoodHashTable.h
@@ -501,7 +501,7 @@ namespace RobinHoodHashTable_Private
 				HashData.Reserve(SizePow2Minus1 + 1);  
 				HashData.AddUninitialized(SizePow2Minus1 + 1);
 
-				for (int Idx = 0; Idx <= SizePow2Minus1; Idx++)
+				for (uint32 Idx = 0; Idx <= SizePow2Minus1; Idx++)
 				{
 					IndexData[Idx] = Other.IndexData[Idx];
 					HashData[Idx] = Other.HashData[Idx];
diff --git a/Engine/Source/Runtime/Renderer/Private/MeshDrawCommands.h b/Engine/Source/Runtime/Renderer/Private/MeshDrawCommands.h
--- a/Engine/Source/Runtime/Renderer/Private/MeshDrawCommands.h
+++ b/Engine/Source/Runtime/Renderer/Private/MeshDrawCommands.h
@@ -172,6 +172,11 @@ public:
 		new (&TaskContext.MinimalPipelineStatePassSet) FGraphicsMinimalPipelineStateSet();
 	}
 
+	void FreeCreateSnapshot()
+	{
+		TaskContext.MinimalPipelineStatePassSet.~FGraphicsMinimalPipelineStateSet();
+	}
+
 private:
 	FPrimitiveIdVertexBufferPoolEntry PrimitiveIdVertexBufferPoolEntry;
 	FMeshDrawCommandPassSetupTaskContext TaskContext;
diff --git a/Engine/Source/Runtime/Renderer/Private/SceneRendering.cpp b/Engine/Source/Runtime/Renderer/Private/SceneRendering.cpp
--- a/Engine/Source/Runtime/Renderer/Private/SceneRendering.cpp
+++ b/Engine/Source/Runtime/Renderer/Private/SceneRendering.cpp
@@ -1797,6 +1797,11 @@ void FViewInfo::DestroyAllSnapshots()
 			Snapshot->ParallelMeshDrawCommandPasses[Index].WaitForTasksAndEmpty();
 		}
 
+		for (int i = 0; i < EMeshPass::Num; i++)
+		{
+			Snapshot->ParallelMeshDrawCommandPasses[i].FreeCreateSnapshot();
+		}
+
 		FreeViewInfoSnapshots.Add(Snapshot);
 	}
 	ViewInfoSnapshots.Reset();
@@ -4223,12 +4228,11 @@ void RunGPUSkinCacheTransition(FRHICommandList& RHICmdList, FScene* Scene, EGPUS
 	const bool bRun = 
 		(bHairStrandsEnabled && EGPUSkinCacheTransition::FrameSetup == Type) || 
 		(!bHairStrandsEnabled && EGPUSkinCacheTransition::FrameSetup != Type);
-	const EResourceTransitionPipeline TargetPipeline = bHairStrandsEnabled ? EResourceTransitionPipeline::EComputeToCompute : EResourceTransitionPipeline::EComputeToGfx;
 	if (bRun)
 	{
 		if (FGPUSkinCache* GPUSkinCache = Scene->GetGPUSkinCache())
 		{
-			GPUSkinCache->TransitionAllToReadable(RHICmdList, TargetPipeline);
+			GPUSkinCache->TransitionAllToReadable(RHICmdList);
 		}
 	}
 }
\ No newline at end of file
