Only dereference primitive pointers if the primitive passes screen size test to reduce some cache misses

diff --git a/Engine/Source/Runtime/Renderer/Private/ShadowSetup.cpp b/Engine/Source/Runtime/Renderer/Private/ShadowSetup.cpp
--- a/Engine/Source/Runtime/Renderer/Private/ShadowSetup.cpp
+++ b/Engine/Source/Runtime/Renderer/Private/ShadowSetup.cpp
@@ -3785,6 +3785,20 @@ struct FGatherShadowPrimitivesPacket
 		}
 	}
 
+	bool DoesPrimitiveCastInsetShadow(const FPrimitiveSceneInfo* PrimitiveSceneInfo, const FPrimitiveSceneProxy* PrimitiveProxy) const
+	{
+		// If light attachment root is valid, we're in a group and need to get the flag from the root.
+		if (PrimitiveSceneInfo->LightingAttachmentRoot.IsValid())
+		{
+			const FAttachmentGroupSceneInfo& AttachmentGroup = PrimitiveSceneInfo->Scene->AttachmentGroups.FindChecked(PrimitiveSceneInfo->LightingAttachmentRoot);
+			return AttachmentGroup.ParentSceneInfo && AttachmentGroup.ParentSceneInfo->Proxy->CastsInsetShadow();
+		}
+		else
+		{
+			return PrimitiveProxy->CastsInsetShadow();
+		}
+	}
+
 	void FilterPrimitiveForShadows(const FPrimitiveSceneInfoCompact& PrimitiveSceneInfoCompact)
 	{
 		const FPrimitiveFlagsCompact& PrimitiveFlagsCompact = PrimitiveSceneInfoCompact.PrimitiveFlagsCompact;
@@ -3856,14 +3870,6 @@ struct FGatherShadowPrimitivesPacket
 					bScreenSpaceSizeCulled = FMath::Square(PrimitiveBounds.SphereRadius) < FMath::Square(MinScreenRadiusForShadowCaster) * DistanceSquared * ProjectedShadowInfo->DependentView->LODDistanceFactorSquared;
 				}
 
-				bool bCastsInsetShadows = PrimitiveProxy->CastsInsetShadow();
-				// If light attachment root is valid, we're in a group and need to get the flag from the root.
-				if (PrimitiveSceneInfo->LightingAttachmentRoot.IsValid())
-				{
-					FAttachmentGroupSceneInfo& AttachmentGroup = PrimitiveSceneInfo->Scene->AttachmentGroups.FindChecked(PrimitiveSceneInfo->LightingAttachmentRoot);
-					bCastsInsetShadows = AttachmentGroup.ParentSceneInfo && AttachmentGroup.ParentSceneInfo->Proxy->CastsInsetShadow();
-				}
-
 				if (!bScreenSpaceSizeCulled
 					&& ProjectedShadowInfo->GetLightSceneInfoCompact().AffectsPrimitive(PrimitiveBounds, PrimitiveProxy)
 					// Include all primitives for movable lights, but only statically shadowed primitives from a light with static shadowing,
@@ -3872,7 +3878,7 @@ struct FGatherShadowPrimitivesPacket
 					// Only render primitives into a reflective shadowmap that are supposed to affect indirect lighting
 					&& !(ProjectedShadowInfo->bReflectiveShadowmap && !PrimitiveProxy->AffectsDynamicIndirectLighting())
 					// Exclude primitives that will create their own per-object shadow, except when rendering RSMs
-					&& (!bCastsInsetShadows || ProjectedShadowInfo->bReflectiveShadowmap)
+					&& (!DoesPrimitiveCastInsetShadow(PrimitiveSceneInfo, PrimitiveProxy) || ProjectedShadowInfo->bReflectiveShadowmap)
 					// Exclude primitives that will create a per-object shadow from a stationary light
 					&& !ShouldCreateObjectShadowForStationaryLight(&LightSceneInfo, PrimitiveProxy, true)
 					// Only render shadows from objects that use static lighting during a reflection capture, since the reflection capture doesn't update at runtime
