Fixed: iOS GPU crash wheh post-process material uses custom stencil while project sets r.CustomDepth=1

diff --git a/Engine/Source/Runtime/Renderer/Private/CustomDepthRendering.cpp b/Engine/Source/Runtime/Renderer/Private/CustomDepthRendering.cpp
--- a/Engine/Source/Runtime/Renderer/Private/CustomDepthRendering.cpp
+++ b/Engine/Source/Runtime/Renderer/Private/CustomDepthRendering.cpp
@@ -69,7 +69,7 @@ void FCustomDepthPassMeshProcessor::AddMeshBatch(const FMeshBatch& RESTRICT Mesh
 		const bool bIsTranslucent = IsTranslucentBlendMode(BlendMode);
 
 
-		const bool bWriteCustomStencilValues = FSceneRenderTargets::IsCustomDepthPassWritingStencil();
+		const bool bWriteCustomStencilValues = FSceneRenderTargets::IsCustomDepthPassWritingStencil(FeatureLevel);
 		float MobileColorValue = 0.0f;
 
 		if (bWriteCustomStencilValues)
diff --git a/Engine/Source/Runtime/Renderer/Private/PostProcess/PostProcessTonemap.cpp b/Engine/Source/Runtime/Renderer/Private/PostProcess/PostProcessTonemap.cpp
--- a/Engine/Source/Runtime/Renderer/Private/PostProcess/PostProcessTonemap.cpp
+++ b/Engine/Source/Runtime/Renderer/Private/PostProcess/PostProcessTonemap.cpp
@@ -983,7 +983,8 @@ void FRCPassPostProcessTonemap::Process(FRenderingCompositePassContext& Context)
 	GraphBuilder.Execute();
 
 	// We only release the SceneColor after the last view was processed (SplitScreen)
-	if(View.IsLastInFamily() && !GIsEditor)
+	if (View.IsLastInFamily() && !GIsEditor 
+		&& View.GetFeatureLevel() > ERHIFeatureLevel::ES3_1) // Mobile renderer may still use SceneColor.A to sample scene depth
 	{
 		// The RT should be released as early as possible to allow sharing of that memory for other purposes.
 		// This becomes even more important with some limited VRam (XBoxOne).
diff --git a/Engine/Source/Runtime/Renderer/Private/PostProcess/SceneRenderTargets.cpp b/Engine/Source/Runtime/Renderer/Private/PostProcess/SceneRenderTargets.cpp
--- a/Engine/Source/Runtime/Renderer/Private/PostProcess/SceneRenderTargets.cpp
+++ b/Engine/Source/Runtime/Renderer/Private/PostProcess/SceneRenderTargets.cpp
@@ -1464,7 +1464,7 @@ bool FSceneRenderTargets::BeginRenderingCustomDepth(FRHICommandListImmediate& RH
 
 		FRHIRenderPassInfo RPInfo;
 
-		const bool bWritesCustomStencilValues = IsCustomDepthPassWritingStencil();
+		const bool bWritesCustomStencilValues = IsCustomDepthPassWritingStencil(CurrentFeatureLevel);
 		const bool bRequiresStencilColorTarget = (CurrentFeatureLevel <= ERHIFeatureLevel::ES3_1);
 
 		int32 NumColorTargets = 0;
@@ -3036,7 +3036,7 @@ const FUnorderedAccessViewRHIRef& FSceneRenderTargets::GetSceneColorTextureUAV()
 IPooledRenderTarget* FSceneRenderTargets::RequestCustomDepth(FRHICommandListImmediate& RHICmdList, bool bPrimitives)
 {
 	int Value = CVarCustomDepth.GetValueOnRenderThread();
-	const bool bCustomDepthPassWritingStencil = IsCustomDepthPassWritingStencil();
+	const bool bCustomDepthPassWritingStencil = IsCustomDepthPassWritingStencil(CurrentFeatureLevel);
 	const bool bMobilePath = (CurrentFeatureLevel <= ERHIFeatureLevel::ES3_1);
 	int32 DownsampleFactor = bMobilePath && CVarMobileCustomDepthDownSample.GetValueOnRenderThread() > 0 ? 2 : 1;
 
@@ -3102,9 +3102,15 @@ IPooledRenderTarget* FSceneRenderTargets::RequestCustomDepth(FRHICommandListImme
 	return 0;
 }
 
-bool FSceneRenderTargets::IsCustomDepthPassWritingStencil()
+bool FSceneRenderTargets::IsCustomDepthPassWritingStencil(ERHIFeatureLevel::Type InFeatureLevel)
 {
-	return (CVarCustomDepth.GetValueOnRenderThread() == 3);
+	int32 CustomDepthValue = CVarCustomDepth.GetValueOnRenderThread();
+	// Mobile uses "On Demand" for both Depth and Stencil textures
+	if (CustomDepthValue == 3 || (CustomDepthValue == 1 && InFeatureLevel <= ERHIFeatureLevel::ES3_1))
+	{
+		return true;
+	}
+	return false;
 }
 
 /** Returns an index in the range [0, NumCubeShadowDepthSurfaces) given an input resolution. */
@@ -3421,7 +3427,8 @@ void SetupMobileSceneTextureUniformParameters(
 
 	FRHITexture* MobileCustomStencil = BlackDefault2D;
 
-	if (bCustomDepthIsValid && SceneContext.MobileCustomStencil.IsValid())
+	if (bCustomDepthIsValid && SceneContext.MobileCustomStencil.IsValid() 
+		&& (SceneContext.MobileCustomStencil->GetDesc().TargetableFlags & TexCreate_Memoryless) == 0)
 	{
 		MobileCustomStencil = SceneContext.MobileCustomStencil->GetRenderTargetItem().ShaderResourceTexture;
 	}
diff --git a/Engine/Source/Runtime/Renderer/Private/PostProcess/SceneRenderTargets.h b/Engine/Source/Runtime/Renderer/Private/PostProcess/SceneRenderTargets.h
--- a/Engine/Source/Runtime/Renderer/Private/PostProcess/SceneRenderTargets.h
+++ b/Engine/Source/Runtime/Renderer/Private/PostProcess/SceneRenderTargets.h
@@ -434,7 +434,7 @@ public:
 	// @return can be 0 if the feature is disabled
 	IPooledRenderTarget* RequestCustomDepth(FRHICommandListImmediate& RHICmdList, bool bPrimitives);
 
-	static bool IsCustomDepthPassWritingStencil();
+	static bool IsCustomDepthPassWritingStencil(ERHIFeatureLevel::Type InFeatureLevel);
 
 	// ---
 
