More optimal handling of scene capture and reflection capture when in VR Mode (using Grow rather than resize in cases where it makes sense).

diff --git a/Engine/Source/Runtime/Renderer/Private/PostProcess/SceneRenderTargets.cpp b/Engine/Source/Runtime/Renderer/Private/PostProcess/SceneRenderTargets.cpp
--- a/Engine/Source/Runtime/Renderer/Private/PostProcess/SceneRenderTargets.cpp
+++ b/Engine/Source/Runtime/Renderer/Private/PostProcess/SceneRenderTargets.cpp
@@ -293,6 +293,7 @@ FSceneRenderTargets::FSceneRenderTargets(const FViewInfo& View, const FSceneRend
 	, bSeparateTranslucencyPass(SnapshotSource.bSeparateTranslucencyPass)
 	, BufferSize(SnapshotSource.BufferSize)
 	, SeparateTranslucencyBufferSize(SnapshotSource.SeparateTranslucencyBufferSize)
+	, LastStereoSize(SnapshotSource.LastStereoSize)
 	, SeparateTranslucencyScale(SnapshotSource.SeparateTranslucencyScale)
 	, SmallColorDepthDownsampleFactor(SnapshotSource.SmallColorDepthDownsampleFactor)
 	, bUseDownsizedOcclusionQueries(SnapshotSource.bUseDownsizedOcclusionQueries)
@@ -404,15 +405,62 @@ FIntPoint FSceneRenderTargets::ComputeDesiredSize(const FSceneViewFamily& ViewFa
 		bIsVRScene |= (IStereoRendering::IsStereoEyeView(*View) && GEngine->XRSystem.IsValid());
 	}
 
+	FIntPoint DesiredBufferSize = FIntPoint::ZeroValue;
+	FIntPoint DesiredFamilyBufferSize = FSceneRenderer::GetDesiredInternalBufferSize(ViewFamily);
+
 	{
 		bool bUseResizeMethodCVar = true;
 
 		if (CVarSceneTargetsResizeMethodForceOverride.GetValueOnRenderThread() != 1)
 		{
-			if (!FPlatformProperties::SupportsWindowedMode() || (bIsVRScene && !bIsSceneCapture))
+			if (!FPlatformProperties::SupportsWindowedMode() || bIsVRScene)
 			{
-				// Force ScreenRes on non windowed platforms.
-				SceneTargetsSizingMethod = RequestedSize;
+				if (bIsVRScene)
+				{
+					if (!bIsSceneCapture && !bIsReflectionCapture)
+					{
+						// If this isn't a scene capture, and it's a VR scene, and the size has changed since the last time we
+						// rendered a VR scene (or this is the first time), use the requested size method.
+						if (DesiredFamilyBufferSize.X != LastStereoSize.X || DesiredFamilyBufferSize.Y != LastStereoSize.Y)
+						{
+							LastStereoSize = DesiredFamilyBufferSize;
+							SceneTargetsSizingMethod = RequestedSize;
+							UE_LOG(LogRenderer, Warning, TEXT("Resizing VR buffer to %d by %d"), DesiredFamilyBufferSize.X, DesiredFamilyBufferSize.Y);
+						}
+						else
+						{
+							// Otherwise use the grow method.
+							SceneTargetsSizingMethod = Grow;
+						}
+					}
+					else
+					{
+						// If this is a scene capture, and it's smaller than the VR view size, then don't re-allocate buffers, just use the "grow" method.
+						// If it's bigger than the VR view, then log a warning, and use resize method.
+						if (DesiredFamilyBufferSize.X > LastStereoSize.X || DesiredFamilyBufferSize.Y > LastStereoSize.Y)
+						{
+							if (LastStereoSize.X > 0 && bIsSceneCapture)
+							{
+								static bool DisplayedCaptureSizeWarning = false;
+								if (!DisplayedCaptureSizeWarning)
+								{
+									DisplayedCaptureSizeWarning = true;
+									UE_LOG(LogRenderer, Warning, TEXT("Scene capture of %d by %d is larger than the current VR target. If this is deliberate for a capture that is being done for multiple frames, consider the performance and memory implications. To disable this warning and ensure optimal behavior with this path, set r.SceneRenderTargetResizeMethod to 2, and r.SceneRenderTargetResizeMethodForceOverride to 1."), DesiredFamilyBufferSize.X, DesiredFamilyBufferSize.Y);
+								}
+							}
+							SceneTargetsSizingMethod = RequestedSize;
+						}
+						else
+						{
+							SceneTargetsSizingMethod = Grow;
+						}
+					}
+				}
+				else
+				{
+					// Force ScreenRes on non windowed platforms.
+					SceneTargetsSizingMethod = RequestedSize;
+				}
 				bUseResizeMethodCVar = false;
 			}
 			else if (GIsEditor)
@@ -430,8 +478,6 @@ FIntPoint FSceneRenderTargets::ComputeDesiredSize(const FSceneViewFamily& ViewFa
 		}
 	}
 
-	FIntPoint DesiredBufferSize = FIntPoint::ZeroValue;
-	FIntPoint DesiredFamilyBufferSize = FSceneRenderer::GetDesiredInternalBufferSize(ViewFamily);
 	switch (SceneTargetsSizingMethod)
 	{
 		case RequestedSize:
diff --git a/Engine/Source/Runtime/Renderer/Private/PostProcess/SceneRenderTargets.h b/Engine/Source/Runtime/Renderer/Private/PostProcess/SceneRenderTargets.h
--- a/Engine/Source/Runtime/Renderer/Private/PostProcess/SceneRenderTargets.h
+++ b/Engine/Source/Runtime/Renderer/Private/PostProcess/SceneRenderTargets.h
@@ -168,6 +168,7 @@ protected:
 		bSeparateTranslucencyPass(false),
 		BufferSize(0, 0),
 		SeparateTranslucencyBufferSize(0, 0),
+		LastStereoSize(0, 0),
 		SeparateTranslucencyScale(1),
 		SmallColorDepthDownsampleFactor(2),
 		bUseDownsizedOcclusionQueries(true),
@@ -762,6 +763,8 @@ private:
 	/* Size of the first view, used for multiview rendertargets */
 	FIntPoint View0Size;
 	FIntPoint SeparateTranslucencyBufferSize;
+	/* Size of the stereo view, if we're an XR device. */
+	FIntPoint LastStereoSize;
 	float SeparateTranslucencyScale;
 	/** e.g. 2 */
 	uint32 SmallColorDepthDownsampleFactor;
