Handle RT bindings for decals that write or read to normal Gbuffer.

diff --git a/Engine/Source/Runtime/Renderer/Private/CompositionLighting/PostProcessMeshDecals.cpp b/Engine/Source/Runtime/Renderer/Private/CompositionLighting/PostProcessMeshDecals.cpp
--- a/Engine/Source/Runtime/Renderer/Private/CompositionLighting/PostProcessMeshDecals.cpp
+++ b/Engine/Source/Runtime/Renderer/Private/CompositionLighting/PostProcessMeshDecals.cpp
@@ -262,26 +262,31 @@ void FMeshDecalMeshProcessor::AddMeshBatch(const FMeshBatch& RESTRICT MeshBatch,
 						FDecalRenderingCommon::ComputeFinalDecalBlendMode(ShaderPlatform, (EDecalBlendMode)Material->GetDecalBlendMode(), bHasNormal),
 						PassDecalStage);
 
-					if (ViewIfDynamicMeshCommand->Family->UseDebugViewPS())
+					FDecalRenderingCommon::ERenderTargetMode DecalRenderTargetMode = FDecalRenderingCommon::ComputeRenderTargetMode(ShaderPlatform, DecalBlendMode, bHasNormal);
+
+					if (DecalRenderTargetMode == RenderTargetMode)
 					{
-						// Deferred decals can only use translucent blend mode
-						if (ViewIfDynamicMeshCommand->Family->EngineShowFlags.ShaderComplexity)
+						if (ViewIfDynamicMeshCommand->Family->UseDebugViewPS())
 						{
-							// If we are in the translucent pass then override the blend mode, otherwise maintain additive blending.
-							PassDrawRenderState.SetBlendState(TStaticBlendState<CW_RGBA, BO_Add, BF_One, BF_One, BO_Add, BF_Zero, BF_One>::GetRHI());
+							// Deferred decals can only use translucent blend mode
+							if (ViewIfDynamicMeshCommand->Family->EngineShowFlags.ShaderComplexity)
+							{
+								// If we are in the translucent pass then override the blend mode, otherwise maintain additive blending.
+								PassDrawRenderState.SetBlendState(TStaticBlendState<CW_RGBA, BO_Add, BF_One, BF_One, BO_Add, BF_Zero, BF_One>::GetRHI());
+							}
+							else if (ViewIfDynamicMeshCommand->Family->GetDebugViewShaderMode() != DVSM_OutputMaterialTextureScales)
+							{
+								// Otherwise, force translucent blend mode (shaders will use an hardcoded alpha).
+								PassDrawRenderState.SetBlendState(TStaticBlendState<CW_RGBA, BO_Add, BF_SourceAlpha, BF_InverseSourceAlpha, BO_Add, BF_Zero, BF_InverseSourceAlpha>::GetRHI());
+							}
 						}
-						else if (ViewIfDynamicMeshCommand->Family->GetDebugViewShaderMode() != DVSM_OutputMaterialTextureScales)
+						else
 						{
-							// Otherwise, force translucent blend mode (shaders will use an hardcoded alpha).
-							PassDrawRenderState.SetBlendState(TStaticBlendState<CW_RGBA, BO_Add, BF_SourceAlpha, BF_InverseSourceAlpha, BO_Add, BF_Zero, BF_InverseSourceAlpha>::GetRHI());
+							PassDrawRenderState.SetBlendState(GetDecalBlendState(FeatureLevel, PassDecalStage, DecalBlendMode, bHasNormal));
 						}
-					}
-					else
-					{
-						PassDrawRenderState.SetBlendState(GetDecalBlendState(FeatureLevel, PassDecalStage, DecalBlendMode, bHasNormal));
-					}
 
-					Process(MeshBatch, BatchElementMask, StaticMeshId, PrimitiveSceneProxy, *MaterialRenderProxy, *Material, MeshFillMode, MeshCullMode);
+						Process(MeshBatch, BatchElementMask, StaticMeshId, PrimitiveSceneProxy, *MaterialRenderProxy, *Material, MeshFillMode, MeshCullMode);
+					}
 				}
 			}
 		}
@@ -358,9 +363,13 @@ void DrawDecalMeshCommands(FRenderingCompositePassContext& Context, EDecalRender
 
 	const bool bPerPixelDBufferMask = IsUsingPerPixelDBufferMask(View.GetShaderPlatform());
 
+	bool bHasNormalRenderTarget = FDecalRenderingCommon::RTM_GBufferNormal == RenderTargetMode ||
+		FDecalRenderingCommon::RTM_SceneColorAndGBufferDepthWriteWithNormal == RenderTargetMode ||
+		FDecalRenderingCommon::RTM_SceneColorAndGBufferWithNormal == RenderTargetMode;
+
 	FGraphicsPipelineStateInitializer GraphicsPSOInit;
 	FDecalRenderTargetManager RenderTargetManager(Context.RHICmdList, Context.GetShaderPlatform(), Context.GetFeatureLevel(), CurrentDecalStage);
-	RenderTargetManager.SetRenderTargetMode(RenderTargetMode, true, bPerPixelDBufferMask);
+	RenderTargetManager.SetRenderTargetMode(RenderTargetMode, bHasNormalRenderTarget, bPerPixelDBufferMask);
 	Context.SetViewportAndCallRHI(Context.View.ViewRect);
 	RHICmdList.ApplyCachedRenderTargets(GraphicsPSOInit);
 
@@ -414,6 +423,7 @@ void RenderMeshDecals(FRenderingCompositePassContext& Context, EDecalRenderStage
 		case DRS_BeforeLighting:
 			DrawDecalMeshCommands(Context, CurrentDecalStage, FDecalRenderingCommon::RTM_GBufferNormal);
 			DrawDecalMeshCommands(Context, CurrentDecalStage, FDecalRenderingCommon::RTM_SceneColorAndGBufferWithNormal);
+			DrawDecalMeshCommands(Context, CurrentDecalStage, FDecalRenderingCommon::RTM_SceneColorAndGBufferNoNormal);
 			break;
 
 		case DRS_Mobile:
