Make AO decals respect meshes that disable receiving decals.

diff --git a/Engine/Source/Runtime/Renderer/Private/CompositionLighting/PostProcessDeferredDecals.cpp b/Engine/Source/Runtime/Renderer/Private/CompositionLighting/PostProcessDeferredDecals.cpp
--- a/Engine/Source/Runtime/Renderer/Private/CompositionLighting/PostProcessDeferredDecals.cpp
+++ b/Engine/Source/Runtime/Renderer/Private/CompositionLighting/PostProcessDeferredDecals.cpp
@@ -326,7 +326,10 @@ static FDecalDepthState ComputeDecalDepthState(EDecalRenderStage LocalDecalStage
 		return Ret;
 	}
 
-	const bool bUseDecalMask = LocalDecalStage == DRS_BeforeLighting || LocalDecalStage == DRS_Emissive;
+	const bool bUseDecalMask = 
+		LocalDecalStage == DRS_BeforeLighting || 
+		LocalDecalStage == DRS_Emissive || 
+		LocalDecalStage == DRS_AmbientOcclusion;
 
 	if (bInsideDecal)
 	{
@@ -780,7 +783,7 @@ void FRCPassPostProcessDeferredDecals::Process(FRenderingCompositePassContext& C
 			// This stops the targets from being resolved and decoded until the last view is rendered.
 			// This is done so as to not run eliminate fast clear on the views before the end.
 			bool bLastView = Context.View.Family->Views.Last() == &Context.View;
-			if ((Scene.Decals.Num() > 0) && bLastView && CurrentStage == DRS_Emissive)
+			if ((Scene.Decals.Num() > 0) && bLastView && CurrentStage == DRS_AmbientOcclusion)
 			{
 				// we don't modify stencil but if out input was having stencil for us (after base pass - we need to clear)
 				// Clear stencil to 0, which is the assumed default by other passes
