Add hair sky lighting sample debug output.

diff --git a/Engine/Shaders/Private/HairStrands/HairStrandsBsdfPlot.usf b/Engine/Shaders/Private/HairStrands/HairStrandsBsdfPlot.usf
--- a/Engine/Shaders/Private/HairStrands/HairStrandsBsdfPlot.usf
+++ b/Engine/Shaders/Private/HairStrands/HairStrandsBsdfPlot.usf
@@ -10,6 +10,10 @@
 #include "HairStrandsVisibilityCommon.ush"
 #include "HairStrandsEnvironmentLightingCommon.ush"
 
+/////////////////////////////////////////////////////////////////////////////////////////////////////////
+
+#if SHADER_PLOTBSDF
+
 int2 InputCoord;
 int2 OutputOffset;
 int2 OutputResolution;
@@ -119,3 +123,149 @@ void MainPS(
 	//OutColor = float4(ColorV, 1);
 	//OutColor = float4(UV, 0, 1);
 }
+
+#endif // SHADER_PLOTBSDF
+
+
+//////////////////////////////////////////////////////////////////////////////////////////////////////////////////////////////
+
+#if SHADER_PLOTSAMPLE
+
+#include "HairStrandsDebugCommonStruct.ush"
+
+StructuredBuffer<FDebugShadingInfo>	Debug_ShadingPointBuffer;
+Buffer<uint>						Debug_ShadingPointCounter;
+StructuredBuffer<FDebugSample>		Debug_SampleBuffer;
+Buffer<uint>						Debug_SampleCounter;
+uint								Debug_MaxShadingPointCount;
+uint								Debug_MaxSampleCount;
+
+int2 OutputOffset;
+int2 OutputResolution;
+int2 MaxResolution;
+uint HairComponents;
+float Exposure;
+
+struct FPlotRect
+{
+	int2 Offset;
+	int2 Resolution;
+};
+
+bool IsInside(const int2 Coord, const FPlotRect Rect)
+{
+	return all(Coord >= Rect.Offset) && all(Coord < Rect.Offset + Rect.Resolution);
+}
+
+float2 ToUV(const int2 Coord, const FPlotRect Rect)
+{
+	return ((Coord - Rect.Offset) + float2(0.5f, 0.5f)) / float2(OutputResolution);
+}
+
+RWTexture2D<float4> OutputTexture;
+
+float2 CartesianToUnwrap(float3 D)
+{
+	const float Theta = asin(D.x);
+	const float Phi = atan2(D.z, D.y);
+	return float2(Theta, Phi);
+}
+
+float2 CartesianToLatlong(float3 D)
+{
+	const float Theta = asin(D.x);
+	const float Phi = atan2(D.z, D.y);
+	return float2(((Phi/PI)+1)*0.5f, Theta/PI);
+}
+
+void MainPS(
+	in FScreenVertexOutput Input,
+	out float4 OutColor : SV_Target0)
+{
+	OutColor = float4(0, 0, 0, 1);
+
+	// Prune all pixel which are not within the band horizontal band
+	const int2 PixelCoord = Input.Position.xy;
+	if (PixelCoord.y < OutputOffset.y || PixelCoord.y >= OutputOffset.x + OutputResolution.y)
+	{
+		discard;
+	}
+
+	const float CosAngleThreshold = 0.999f;
+	const float PointSize = 3.f / OutputResolution.x;
+
+	// Local 
+	// Plot sample points
+	bool bIsValid = false;
+	const uint ShadingPointCount = Debug_ShadingPointCounter[0];
+	for (uint PointIt = 0; PointIt < ShadingPointCount; PointIt++)
+	{
+		FPlotRect Rect;
+		Rect.Offset = OutputOffset + PointIt * int2(OutputResolution.x + 10, 0);
+		Rect.Resolution = OutputResolution;
+
+		if (!IsInside(PixelCoord, Rect))
+			continue;
+
+		bIsValid = true;
+		const float2 UV = ToUV(PixelCoord, Rect);
+
+		const float  Pixel_Phi	= (UV.x * 2 - 1) * PI + PI*0.5; //[3PI/2, 0, 2PI, 3PI/2 ], so that the hair fiber is aligned on the center
+		const float  Pixel_Theta= (UV.y * 2 - 1) * PI * 0.5f;
+		const float3 PixelDir_Local	= float3(sin(Pixel_Theta), cos(Pixel_Theta)*cos(Pixel_Phi), cos(Pixel_Theta)*sin(Pixel_Phi));
+
+		const FDebugShadingInfo P = Debug_ShadingPointBuffer[PointIt];
+
+		// Shading frame T, V
+		const float3 Local_X = P.T;
+		const float3 Local_Y = normalize(cross(Local_X, P.V));
+		const float3 Local_Z = normalize(cross(Local_Y, Local_X));
+		const float3x3 WorldToLocal = transpose(float3x3(Local_X, Local_Y, Local_Z));
+
+		const float3 V_Local = mul(P.V, WorldToLocal);
+		const float3 T_Local = mul(P.T, WorldToLocal); // Should be float3(1,0,0)
+
+		{
+			FHairTransmittanceData TransmittanceData = (FHairTransmittanceData)0;
+			TransmittanceData.LocalScattering = 0;
+			TransmittanceData.GlobalScattering = 1;
+			TransmittanceData.ScatteringComponent = HairComponents;
+			TransmittanceData.OpaqueVisibility = 1;
+
+			FGBufferData HairGBuffer = (FGBufferData)0;
+			HairGBuffer.WorldNormal = P.T;
+			HairGBuffer.BaseColor = P.BaseColor;
+			HairGBuffer.Roughness = P.Roughness;
+			HairGBuffer.ShadingModelID = SHADINGMODELID_HAIR;
+			HairGBuffer.DiffuseColor = 0;
+			HairGBuffer.SpecularColor = 0;
+			HairGBuffer.Specular = 0.5f; //TODO
+			HairGBuffer.Metallic = 0;
+			HairGBuffer.Depth = 0.5f;
+			HairGBuffer.GBufferAO = 1;
+			HairGBuffer.CustomData = 0;
+			HairGBuffer.IndirectIrradiance = 1000000;
+
+			OutColor.xyz = HairShading(HairGBuffer, PixelDir_Local, V_Local, T_Local, 1, TransmittanceData, 1, 0, uint2(0, 0)) * Exposure;
+		}
+
+		for (uint SampleIt = 0; SampleIt < P.SampleCount; SampleIt++)
+		{
+			const FDebugSample S_World = Debug_SampleBuffer[P.SampleOffset + SampleIt];
+
+			const float3 S_Local = mul(S_World.Direction, WorldToLocal);
+
+			if (dot(PixelDir_Local, S_Local) > CosAngleThreshold)
+			{
+				OutColor.xyz = float3(1, 0, 0);
+			}
+		}
+	}
+
+	if (!bIsValid)
+	{
+		discard;
+	}
+}
+
+#endif // SHADER_PLOTSAMPLE
diff --git a/Engine/Shaders/Private/HairStrands/HairStrandsDebugCommonStruct.ush b/Engine/Shaders/Private/HairStrands/HairStrandsDebugCommonStruct.ush
new file mode 100644
--- /dev/null
+++ b/Engine/Shaders/Private/HairStrands/HairStrandsDebugCommonStruct.ush
@@ -0,0 +1,21 @@
+// Copyright Epic Games, Inc. All Rights Reserved.
+
+#pragma once
+
+struct FDebugShadingInfo
+{
+	float3	BaseColor;
+	float	Roughness;
+	float3	T;
+	uint	SampleCount;
+	float3	V;
+	uint	SampleOffset;
+};
+
+struct FDebugSample
+{
+	float3	Direction;
+	float	Pdf;
+	float3	Weight;
+	float	Pad;
+};
\ No newline at end of file
diff --git a/Engine/Shaders/Private/HairStrands/HairStrandsEnvironmentLighting.usf b/Engine/Shaders/Private/HairStrands/HairStrandsEnvironmentLighting.usf
--- a/Engine/Shaders/Private/HairStrands/HairStrandsEnvironmentLighting.usf
+++ b/Engine/Shaders/Private/HairStrands/HairStrandsEnvironmentLighting.usf
@@ -96,7 +96,58 @@ StructuredBuffer<FPackedHairSample> HairVisibilityNodeData;
 
 Texture3D<float4>			HairScatteringLUTTexture;
 Texture3D<float4>			HairEnergyLUTTexture;
-SamplerState 				HairLUTSampler;
+SamplerState 				HairLUTSampler; 
+
+//////////////////////////////////////////////////////////////////////////////////////////////////////////////////////////////
+// Debug info
+
+#if PERMUTATION_DEBUG
+#include "HairStrandsDebugCommonStruct.ush"
+
+RWStructuredBuffer<FDebugShadingInfo>	Debug_ShadingPointBuffer;
+RWBuffer<uint>							Debug_ShadingPointCounter;
+RWStructuredBuffer<FDebugSample>		Debug_SampleBuffer;
+RWBuffer<uint>							Debug_SampleCounter;
+uint									Debug_MaxShadingPointCount;
+uint									Debug_MaxSampleCount;
+
+uint AllocateSample(uint SampleCount)
+{
+}
+
+uint AddShadingPoint(const FGBufferData GB, const float3 V, uint SampleCount)
+{
+	uint SampleOffset = 0;
+	InterlockedAdd(Debug_SampleCounter[0], SampleCount, SampleOffset);
+
+	uint PointOffset = 0;
+	InterlockedAdd(Debug_ShadingPointCounter[0], 1, PointOffset);
+	if (PointOffset < Debug_MaxShadingPointCount)
+	{
+		FDebugShadingInfo Out;
+		Out.BaseColor = GB.BaseColor;
+		Out.Roughness = GB.Roughness;
+		Out.V = V;
+		Out.T = GB.WorldNormal;
+		Out.SampleCount = SampleCount;
+		Out.SampleOffset = SampleOffset;
+		Debug_ShadingPointBuffer[PointOffset] = Out;
+	}
+
+	return SampleOffset;
+}
+
+void AddSample(uint SampleOffset, uint SampleIt, const float3 D, const float3 W)
+{
+	FDebugSample Out = (FDebugSample)0;
+	Out.Direction = D;
+	Out.Weight = W;
+	Out.Pdf = 1;
+	Debug_SampleBuffer[SampleOffset + SampleIt] = Out;
+}
+#endif //PERMUTATION_DEBUG
+
+//////////////////////////////////////////////////////////////////////////////////////////////////////////////////////////////
 
 float HairDistanceThreshold;
 uint bHairUseViewHairCount;
@@ -165,9 +216,9 @@ float3 ReflectionEnvironment(FGBufferData GBuffer, float2 ScreenPosition, float2
 	const uint2  PixelPosition	= SvPosition - ResolvedView.ViewRectMin.xy;
 	const float  PixelSceneDepth= GBuffer.Depth;
 
+	const float VoxelMaxSize = VirtualVoxel.VoxelWorldSize;
 	const uint PrimitiveIdMod8 = (PrimitiveId % 8);
-	const float  VoxelMaxSize = VirtualVoxel.VoxelWorldSize;
-	const float3 VoxelOffset = VoxelMaxSize * GetHairVoxelJitter(SvPosition, View.StateFrameIndexMod8 + PrimitiveIdMod8);
+	const float3 VoxelOffset =  GetHairVoxelJitter(SvPosition, View.StateFrameIndexMod8 + PrimitiveIdMod8);
 
 	const float3 V = -CameraToPixel;
 	const float3 T = GBuffer.WorldNormal;
@@ -496,8 +547,14 @@ float3 ReflectionEnvironmentUniformSampling(
 	const uint2  PixelPosition = SvPosition - ResolvedView.ViewRectMin.xy;
 	const float  PixelSceneDepth = GBuffer.Depth;
 
+#if VOXEL_TRAVERSAL_TYPE != VOXEL_TRAVERSAL_LEGACY
+	const float3 VoxelSize = VirtualVoxel.VoxelWorldSize; // Use the correct voxel size when we use virtual voxel volume.
+#else
+	const float3 VoxelSize = (Voxel_MaxAABB - Voxel_MinAABB) / Voxel_Resolution;
+#endif
+
 	const uint PrimitiveIdMod8 = (PrimitiveId % 8);
-	const float  VoxelMaxSize = VirtualVoxel.VoxelWorldSize;
+	const float VoxelMaxSize = VirtualVoxel.VoxelWorldSize;
 	const float3 VoxelOffset = VoxelMaxSize * GetHairVoxelJitter(SvPosition, View.StateFrameIndexMod8 + PrimitiveIdMod8);
 
 	const float3 V = -CameraToPixel;
@@ -539,6 +596,14 @@ float3 ReflectionEnvironmentUniformSampling(
 			InterleavedGradientNoise(SvPosition, PrimitiveIdMod8 * 7));
 	}
 
+	#if PERMUTATION_DEBUG
+	uint SampleOffset = 0;
+	if (bDebugEnabled)
+	{
+		SampleOffset = AddShadingPoint(GBuffer, V, SampleCount);
+	}
+	#endif
+
 	float3 AccLighting = 0;
 	for (uint SampleIt = 0; SampleIt < SampleCount; ++SampleIt)
 	{
@@ -555,6 +620,13 @@ float3 ReflectionEnvironmentUniformSampling(
 		const float3 SampleL = UniformSampleSphere(frac(Rand2)).xyz;
 		const float3 SampleLightPosition = WorldPosition + SampleL * MaxSearchDistance;
 
+		#if PERMUTATION_DEBUG
+		if (bDebugEnabled)
+		{
+			AddSample(SampleOffset, SampleIt, SampleL, 1);
+		}
+		#endif
+
 		const float SinLightAngle = dot(SampleL, T);
 		const FHairTraversalResult Result = InternalComputeHairCountVoxel(
 			SampleWorldPosition, 
@@ -588,6 +660,13 @@ float3 ReflectionEnvironmentUniformSampling(
 			Result.Visibility * 
 			PreExposure * 
 			GatherRadiance(SampleL, 0.1f, PixelPosition, PixelSceneDepth, WorldPosition);
+
+		#if PERMUTATION_DEBUG
+		if (bDebugEnabled)
+		{
+			AddLine(SampleWorldPosition, SampleLightPosition, float4(1, 0, 1, 1), float4(1, 0, 1, 1));
+		}
+		#endif
 	}
 	const float SphereInt = 4 * PI;
 	const float3 Color = SphereInt * AccLighting / SampleCount;
@@ -707,6 +786,12 @@ void MainPS(
 	GBuffer.CustomData = float4(HairDualScatteringRoughnessOverride, 0, Sample.Backlit, 0);
 	GBuffer.IndirectIrradiance = 1000000;
 
+#if PERMUTATION_DEBUG
+	const bool bDebugEnabled = all(int2(PixelCoord) == GetCursorPos());
+#else
+	const bool bDebugEnabled = false;
+#endif
+
 	const bool bIsMacroGroupCompatible = Sample.MacroGroupId == Voxel_MacroGroupId;
 #if PERMUTATION_INTEGRATION_TYPE == INTEGRATION_SCENECOLOR
 	const bool bEnabled = Sample.bScatterSceneLighting;
@@ -722,7 +807,7 @@ void MainPS(
 		#elif PERMUTATION_INTEGRATION_TYPE == INTEGRATION_ADHOC
 			ReflectionEnvironment(GBuffer, ScreenPosition, SvPosition, Sample.PrimitiveId);
 		#elif PERMUTATION_INTEGRATION_TYPE == INTEGRATION_UNIFORM
-			ReflectionEnvironmentUniformSampling(GBuffer, ScreenPosition, SvPosition, Sample.PrimitiveId, false);
+			ReflectionEnvironmentUniformSampling(GBuffer, ScreenPosition, SvPosition, Sample.PrimitiveId, bDebugEnabled);
 		#endif
 
 		const float SampleCoverage = saturate(From8bitCoverage(Sample.Coverage8bit));
diff --git a/Engine/Source/Runtime/Renderer/Private/DeferredShadingRenderer.h b/Engine/Source/Runtime/Renderer/Private/DeferredShadingRenderer.h
--- a/Engine/Source/Runtime/Renderer/Private/DeferredShadingRenderer.h
+++ b/Engine/Source/Runtime/Renderer/Private/DeferredShadingRenderer.h
@@ -227,7 +227,7 @@ private:
 	void RenderDiffuseIndirectAndAmbientOcclusion(FRHICommandListImmediate& RHICmdList);
 
 	/** Renders sky lighting and reflections that can be done in a deferred pass. */
-	void RenderDeferredReflectionsAndSkyLighting(FRHICommandListImmediate& RHICmdList, TRefCountPtr<IPooledRenderTarget>& DynamicBentNormalAO, TRefCountPtr<IPooledRenderTarget>& VelocityRT, const struct FHairStrandsDatas* HairDatas);
+	void RenderDeferredReflectionsAndSkyLighting(FRHICommandListImmediate& RHICmdList, TRefCountPtr<IPooledRenderTarget>& DynamicBentNormalAO, TRefCountPtr<IPooledRenderTarget>& VelocityRT, struct FHairStrandsDatas* HairDatas);
 
 	/** Computes DFAO, modulates it to scene color (which is assumed to contain diffuse indirect lighting), and stores the output bent normal for use occluding specular. */
 	void RenderDFAOAsIndirectShadowing(
diff --git a/Engine/Source/Runtime/Renderer/Private/HairStrands/HairStrandsDebug.cpp b/Engine/Source/Runtime/Renderer/Private/HairStrands/HairStrandsDebug.cpp
--- a/Engine/Source/Runtime/Renderer/Private/HairStrands/HairStrandsDebug.cpp
+++ b/Engine/Source/Runtime/Renderer/Private/HairStrands/HairStrandsDebug.cpp
@@ -107,6 +107,56 @@ FHairCullInfo GetHairStrandsCullInfo()
 	return Out;
 }
 
+FHairStrandsDebugData::Data FHairStrandsDebugData::CreateData(FRDGBuilder& GraphBuilder)
+{
+	FHairStrandsDebugData::Data Out;
+	Out.ShadingPointBuffer = GraphBuilder.CreateBuffer(FRDGBufferDesc::CreateStructuredDesc(sizeof(ShadingInfo), MaxShadingPointCount), TEXT("HairDebugShadingPoint"));
+	Out.ShadingPointCounter = GraphBuilder.CreateBuffer(FRDGBufferDesc::CreateBufferDesc(sizeof(uint32), 1), TEXT("HairDebugShadingPointCounter"));
+	Out.SampleBuffer = GraphBuilder.CreateBuffer(FRDGBufferDesc::CreateStructuredDesc(sizeof(Sample), MaxSampleCount), TEXT("HairDebugSample"));
+	Out.SampleCounter = GraphBuilder.CreateBuffer(FRDGBufferDesc::CreateBufferDesc(sizeof(uint32), 1), TEXT("HairDebugSampleCounter"));
+	AddClearUAVPass(GraphBuilder, GraphBuilder.CreateUAV(Out.ShadingPointCounter, PF_R32_UINT), 0u);
+	AddClearUAVPass(GraphBuilder, GraphBuilder.CreateUAV(Out.SampleCounter, PF_R32_UINT), 0u);
+	return Out;
+}
+
+FHairStrandsDebugData::Data FHairStrandsDebugData::ImportData(FRDGBuilder& GraphBuilder, const FHairStrandsDebugData& In)
+{
+	FHairStrandsDebugData::Data Out;
+	Out.ShadingPointBuffer = GraphBuilder.RegisterExternalBuffer(In.ShadingPointBuffer, TEXT("HairDebugShadingPoint"));
+	Out.ShadingPointCounter = GraphBuilder.RegisterExternalBuffer(In.ShadingPointCounter, TEXT("HairDebugShadingPointCounter"));
+	Out.SampleBuffer = GraphBuilder.RegisterExternalBuffer(In.SampleBuffer, TEXT("HairDebugSample"));
+	Out.SampleCounter = GraphBuilder.RegisterExternalBuffer(In.SampleCounter, TEXT("HairDebugSampleCounter"));
+	return Out;
+}
+
+void FHairStrandsDebugData::ExtractData(FRDGBuilder& GraphBuilder, FHairStrandsDebugData::Data& In, FHairStrandsDebugData& Out)
+{
+	GraphBuilder.QueueBufferExtraction(In.ShadingPointBuffer, &Out.ShadingPointBuffer, FRDGResourceState::EAccess::Read, FRDGResourceState::EPipeline::Compute);
+	GraphBuilder.QueueBufferExtraction(In.ShadingPointCounter, &Out.ShadingPointCounter, FRDGResourceState::EAccess::Read, FRDGResourceState::EPipeline::Compute);
+	GraphBuilder.QueueBufferExtraction(In.SampleBuffer, &Out.SampleBuffer, FRDGResourceState::EAccess::Read, FRDGResourceState::EPipeline::Compute);
+	GraphBuilder.QueueBufferExtraction(In.SampleCounter, &Out.SampleCounter, FRDGResourceState::EAccess::Read, FRDGResourceState::EPipeline::Compute);
+}
+
+void FHairStrandsDebugData::SetParameters(FRDGBuilder& GraphBuilder, FHairStrandsDebugData::Data& In, FHairStrandsDebugData::FWriteParameters& Out)
+{
+	Out.Debug_MaxSampleCount = FHairStrandsDebugData::MaxSampleCount;
+	Out.Debug_MaxShadingPointCount = FHairStrandsDebugData::MaxShadingPointCount;
+	Out.Debug_ShadingPointBuffer = GraphBuilder.CreateUAV(In.ShadingPointBuffer);
+	Out.Debug_ShadingPointCounter = GraphBuilder.CreateUAV(In.ShadingPointCounter, PF_R32_UINT);
+	Out.Debug_SampleBuffer = GraphBuilder.CreateUAV(In.SampleBuffer);
+	Out.Debug_SampleCounter = GraphBuilder.CreateUAV(In.SampleCounter, PF_R32_UINT);
+}
+
+void FHairStrandsDebugData::SetParameters(FRDGBuilder& GraphBuilder, FHairStrandsDebugData::Data& In, FHairStrandsDebugData::FReadParameters& Out)
+{
+	Out.Debug_MaxSampleCount = FHairStrandsDebugData::MaxSampleCount;
+	Out.Debug_MaxShadingPointCount = FHairStrandsDebugData::MaxShadingPointCount;
+	Out.Debug_ShadingPointBuffer = GraphBuilder.CreateSRV(In.ShadingPointBuffer);
+	Out.Debug_ShadingPointCounter = GraphBuilder.CreateSRV(In.ShadingPointCounter, PF_R32_UINT);
+	Out.Debug_SampleBuffer = GraphBuilder.CreateSRV(In.SampleBuffer);
+	Out.Debug_SampleCounter = GraphBuilder.CreateSRV(In.SampleCounter, PF_R32_UINT);
+}
+
 enum class EHairDebugMode : uint8
 {
 	None,
@@ -709,10 +759,10 @@ static void AddVoxelPageRaymarchingPass(
 }
 
 ///////////////////////////////////////////////////////////////////////////////////////////////////
-class FHairStrandsBSDFPlotPS : public FGlobalShader
+class FHairStrandsPlotBSDFPS : public FGlobalShader
 {
-	DECLARE_GLOBAL_SHADER(FHairStrandsBSDFPlotPS);
-	SHADER_USE_PARAMETER_STRUCT(FHairStrandsBSDFPlotPS, FGlobalShader);
+	DECLARE_GLOBAL_SHADER(FHairStrandsPlotBSDFPS);
+	SHADER_USE_PARAMETER_STRUCT(FHairStrandsPlotBSDFPS, FGlobalShader);
 
 	BEGIN_SHADER_PARAMETER_STRUCT(FParameters, )
 		SHADER_PARAMETER(FIntPoint, InputCoord)
@@ -728,9 +778,14 @@ class FHairStrandsBSDFPlotPS : public FGlobalShader
 
 public:
 	static bool ShouldCompilePermutation(const FGlobalShaderPermutationParameters& Parameters) { return IsHairStrandsSupported(Parameters.Platform); }
+	static void ModifyCompilationEnvironment(const FGlobalShaderPermutationParameters& Parameters, FShaderCompilerEnvironment& OutEnvironment)
+	{
+		FGlobalShader::ModifyCompilationEnvironment(Parameters, OutEnvironment);
+		OutEnvironment.SetDefine(TEXT("SHADER_PLOTBSDF"), 1);
+	}
 };
 
-IMPLEMENT_GLOBAL_SHADER(FHairStrandsBSDFPlotPS, "/Engine/Private/HairStrands/HairStrandsBsdfPlot.usf", "MainPS", SF_Pixel);
+IMPLEMENT_GLOBAL_SHADER(FHairStrandsPlotBSDFPS, "/Engine/Private/HairStrands/HairStrandsBsdfPlot.usf", "MainPS", SF_Pixel);
 
 static void AddPlotBSDFPass(
 	FRDGBuilder& GraphBuilder,
@@ -742,9 +797,9 @@ static void AddPlotBSDFPass(
 	SetupSceneTextureParameters(GraphBuilder, &SceneTextures);
 
 	const FIntPoint Resolution(OutputTexture->Desc.Extent);
-	FHairStrandsBSDFPlotPS::FParameters* Parameters = GraphBuilder.AllocParameters<FHairStrandsBSDFPlotPS::FParameters>();
+	FHairStrandsPlotBSDFPS::FParameters* Parameters = GraphBuilder.AllocParameters<FHairStrandsPlotBSDFPS::FParameters>();
 	Parameters->InputCoord = View.CursorPos;
-	Parameters->OutputOffset = FIntPoint(100,100);
+	Parameters->OutputOffset = FIntPoint(10,100);
 	Parameters->OutputResolution = FIntPoint(256, 256);
 	Parameters->MaxResolution = OutputTexture->Desc.Extent;
 	Parameters->HairComponents = ToBitfield(GetHairComponents());
@@ -755,7 +810,7 @@ static void AddPlotBSDFPass(
 
 	const FIntPoint OutputResolution = SceneTextures.SceneDepthBuffer->Desc.Extent;
 	TShaderMapRef<FPostProcessVS> VertexShader(View.ShaderMap);
-	TShaderMapRef<FHairStrandsBSDFPlotPS> PixelShader(View.ShaderMap);
+	TShaderMapRef<FHairStrandsPlotBSDFPS> PixelShader(View.ShaderMap);
 	const FGlobalShaderMap* GlobalShaderMap = View.ShaderMap;
 	const FIntRect Viewport = View.ViewRect;
 	const FViewInfo* CapturedView = &View;
@@ -796,7 +851,99 @@ static void AddPlotBSDFPass(
 			EDRF_UseTriangleOptimization);
 	});
 }
-	
+
+///////////////////////////////////////////////////////////////////////////////////////////////////
+class FHairStrandsPlotSamplePS : public FGlobalShader
+{
+	DECLARE_GLOBAL_SHADER(FHairStrandsPlotSamplePS);
+	SHADER_USE_PARAMETER_STRUCT(FHairStrandsPlotSamplePS, FGlobalShader);
+
+	BEGIN_SHADER_PARAMETER_STRUCT(FParameters, )
+		SHADER_PARAMETER_STRUCT_INCLUDE(FHairStrandsDebugData::FReadParameters, DebugData)
+		SHADER_PARAMETER(FIntPoint, OutputOffset)
+		SHADER_PARAMETER(FIntPoint, OutputResolution)
+		SHADER_PARAMETER(FIntPoint, MaxResolution)
+		SHADER_PARAMETER(uint32, HairComponents)
+		SHADER_PARAMETER(float, Exposure)
+		RENDER_TARGET_BINDING_SLOTS()
+	END_SHADER_PARAMETER_STRUCT()
+
+public:
+	static bool ShouldCompilePermutation(const FGlobalShaderPermutationParameters& Parameters) { return IsHairStrandsSupported(Parameters.Platform); }
+	static void ModifyCompilationEnvironment(const FGlobalShaderPermutationParameters& Parameters, FShaderCompilerEnvironment& OutEnvironment)
+	{
+		FGlobalShader::ModifyCompilationEnvironment(Parameters, OutEnvironment);
+		OutEnvironment.SetDefine(TEXT("SHADER_PLOTSAMPLE"), 1);
+	}
+};
+
+IMPLEMENT_GLOBAL_SHADER(FHairStrandsPlotSamplePS, "/Engine/Private/HairStrands/HairStrandsBsdfPlot.usf", "MainPS", SF_Pixel);
+
+static void AddPlotSamplePass(
+	FRDGBuilder& GraphBuilder,
+	const FViewInfo& View,
+	FHairStrandsDebugData::Data& DebugData,
+	FRDGTextureRef& OutputTexture)
+{
+
+	FSceneTextureParameters SceneTextures;
+	SetupSceneTextureParameters(GraphBuilder, &SceneTextures);
+
+	const FIntPoint Resolution(OutputTexture->Desc.Extent);
+	FHairStrandsPlotSamplePS::FParameters* Parameters = GraphBuilder.AllocParameters<FHairStrandsPlotSamplePS::FParameters>();
+
+	FHairStrandsDebugData::SetParameters(GraphBuilder, DebugData, Parameters->DebugData);
+	Parameters->OutputOffset = FIntPoint(100, 100);
+	Parameters->OutputResolution = FIntPoint(256, 256);
+	Parameters->MaxResolution = OutputTexture->Desc.Extent;
+	Parameters->HairComponents = ToBitfield(GetHairComponents());
+	Parameters->Exposure = GHairStrandsDebugPlotBsdfExposure;
+	Parameters->RenderTargets[0] = FRenderTargetBinding(OutputTexture, ERenderTargetLoadAction::ELoad);
+
+	const FIntPoint OutputResolution = SceneTextures.SceneDepthBuffer->Desc.Extent;
+	TShaderMapRef<FPostProcessVS> VertexShader(View.ShaderMap);
+	TShaderMapRef<FHairStrandsPlotSamplePS> PixelShader(View.ShaderMap);
+	const FGlobalShaderMap* GlobalShaderMap = View.ShaderMap;
+	const FIntRect Viewport = View.ViewRect;
+	const FViewInfo* CapturedView = &View;
+
+	ClearUnusedGraphResources(PixelShader, Parameters);
+
+	GraphBuilder.AddPass(
+		RDG_EVENT_NAME("HairStrandsSamplePlot"),
+		Parameters,
+		ERDGPassFlags::Raster,
+		[Parameters, VertexShader, PixelShader, Viewport, Resolution, CapturedView](FRHICommandList& RHICmdList)
+	{
+		FGraphicsPipelineStateInitializer GraphicsPSOInit;
+		RHICmdList.ApplyCachedRenderTargets(GraphicsPSOInit);
+		GraphicsPSOInit.BlendState = TStaticBlendState<CW_RGBA, BO_Add, BF_One, BF_Zero, BO_Add, BF_One, BF_Zero>::GetRHI();
+		GraphicsPSOInit.RasterizerState = TStaticRasterizerState<>::GetRHI();
+		GraphicsPSOInit.DepthStencilState = TStaticDepthStencilState<false, CF_Always>::GetRHI();
+
+		GraphicsPSOInit.BoundShaderState.VertexDeclarationRHI = GFilterVertexDeclaration.VertexDeclarationRHI;
+		GraphicsPSOInit.BoundShaderState.VertexShaderRHI = VertexShader.GetVertexShader();
+		GraphicsPSOInit.BoundShaderState.PixelShaderRHI = PixelShader.GetPixelShader();
+		GraphicsPSOInit.PrimitiveType = PT_TriangleList;
+		SetGraphicsPipelineState(RHICmdList, GraphicsPSOInit);
+
+		VertexShader->SetParameters(RHICmdList, CapturedView->ViewUniformBuffer);
+		RHICmdList.SetViewport(Viewport.Min.X, Viewport.Min.Y, 0.0f, Viewport.Max.X, Viewport.Max.Y, 1.0f);
+		SetShaderParameters(RHICmdList, PixelShader, PixelShader.GetPixelShader(), *Parameters);
+
+		DrawRectangle(
+			RHICmdList,
+			0, 0,
+			Viewport.Width(), Viewport.Height(),
+			Viewport.Min.X, Viewport.Min.Y,
+			Viewport.Width(), Viewport.Height(),
+			Viewport.Size(),
+			Resolution,
+			VertexShader,
+			EDRF_UseTriangleOptimization);
+	});
+}
+
 ///////////////////////////////////////////////////////////////////////////////////////////////////
 BEGIN_SHADER_PARAMETER_STRUCT(FHairProjectionMeshDebugParameters, )
 	SHADER_PARAMETER(FMatrix, LocalToWorld)
@@ -1302,11 +1449,19 @@ void RenderHairStrandsDebugInfo(
 	const EHairStrandsDebugMode StrandsDebugMode = GetHairStrandsDebugStrandsMode();
 	const EHairDebugMode HairDebugMode = GetHairDebugMode();
 
-	if (GHairStrandsDebugPlotBsdf > 0)
+	if (HairDatas && (GHairStrandsDebugPlotBsdf > 0 || HairDatas->DebugData.IsValid()))
 	{
 		FRDGBuilder GraphBuilder(RHICmdList);
 		FRDGTextureRef SceneColorTexture = GraphBuilder.RegisterExternalTexture(SceneTargets.GetSceneColor(), TEXT("SceneColorTexture"));
-		AddPlotBSDFPass(GraphBuilder, View, SceneColorTexture);
+		if (GHairStrandsDebugPlotBsdf > 0)
+		{
+			AddPlotBSDFPass(GraphBuilder, View, SceneColorTexture);
+		}
+		if (HairDatas->DebugData.IsValid())
+		{
+			FHairStrandsDebugData::Data DebugData = FHairStrandsDebugData::ImportData(GraphBuilder, HairDatas->DebugData);
+			AddPlotSamplePass(GraphBuilder, View, DebugData, SceneColorTexture);
+		}
 		GraphBuilder.Execute();		
 	}
 
diff --git a/Engine/Source/Runtime/Renderer/Private/HairStrands/HairStrandsDebug.h b/Engine/Source/Runtime/Renderer/Private/HairStrands/HairStrandsDebug.h
--- a/Engine/Source/Runtime/Renderer/Private/HairStrands/HairStrandsDebug.h
+++ b/Engine/Source/Runtime/Renderer/Private/HairStrands/HairStrandsDebug.h
@@ -8,6 +8,76 @@
 
 #include "CoreMinimal.h"
 #include "RendererInterface.h"
+#include "RenderGraphResources.h"
+#include "ShaderParameterStruct.h"
+
+class FRDGBuilder;
+
+struct FHairStrandsDebugData
+{
+	BEGIN_SHADER_PARAMETER_STRUCT(FWriteParameters, )
+		SHADER_PARAMETER(uint32,							Debug_MaxShadingPointCount)
+		SHADER_PARAMETER(uint32,							Debug_MaxSampleCount)
+		SHADER_PARAMETER_RDG_BUFFER_UAV(RWStructuredBuffer, Debug_ShadingPointBuffer)
+		SHADER_PARAMETER_RDG_BUFFER_UAV(RWBuffer,			Debug_ShadingPointCounter)
+		SHADER_PARAMETER_RDG_BUFFER_UAV(RWStructuredBuffer, Debug_SampleBuffer)
+		SHADER_PARAMETER_RDG_BUFFER_UAV(RWBuffer,			Debug_SampleCounter)
+	END_SHADER_PARAMETER_STRUCT()
+
+	BEGIN_SHADER_PARAMETER_STRUCT(FReadParameters, )
+		SHADER_PARAMETER(uint32,							Debug_MaxShadingPointCount)
+		SHADER_PARAMETER(uint32,							Debug_MaxSampleCount)
+		SHADER_PARAMETER_RDG_BUFFER_SRV(RWStructuredBuffer, Debug_ShadingPointBuffer)
+		SHADER_PARAMETER_RDG_BUFFER_SRV(RWBuffer,			Debug_ShadingPointCounter)
+		SHADER_PARAMETER_RDG_BUFFER_SRV(RWStructuredBuffer, Debug_SampleBuffer)
+		SHADER_PARAMETER_RDG_BUFFER_SRV(RWBuffer,			Debug_SampleCounter)
+	END_SHADER_PARAMETER_STRUCT()
+
+	struct ShadingInfo
+	{
+		FVector BaseColor;
+		float	Roughness;
+		FVector T;
+		uint32	SampleCount;
+		FVector V;
+		float	SampleOffset;
+	};
+
+	struct Sample
+	{
+		FVector Direction;
+		float	Pdf;
+		FVector Weights;
+		float	Pad;
+	};
+	static const uint32 MaxShadingPointCount = 32;
+	static const uint32 MaxSampleCount = 1024 * 32;
+
+	struct Data
+	{
+		FRDGBufferRef ShadingPointBuffer = nullptr;
+		FRDGBufferRef ShadingPointCounter = nullptr;
+
+		FRDGBufferRef SampleBuffer = nullptr;
+		FRDGBufferRef SampleCounter = nullptr;
+	};
+
+	bool IsValid() const
+	{
+		return ShadingPointBuffer && ShadingPointCounter && SampleBuffer &&	SampleCounter;
+	}
+
+	static Data CreateData(FRDGBuilder& GraphBuilder);
+	static Data ImportData(FRDGBuilder& GraphBuilder, const FHairStrandsDebugData& In);
+	static void ExtractData(FRDGBuilder& GraphBuilder, Data& In, FHairStrandsDebugData& Out);
+	static void SetParameters(FRDGBuilder& GraphBuilder, Data& In, FWriteParameters& Out);
+	static void SetParameters(FRDGBuilder& GraphBuilder, Data& In, FReadParameters& Out);
+
+	TRefCountPtr<FPooledRDGBuffer> ShadingPointBuffer;
+	TRefCountPtr<FPooledRDGBuffer> ShadingPointCounter;
+	TRefCountPtr<FPooledRDGBuffer> SampleBuffer;
+	TRefCountPtr<FPooledRDGBuffer> SampleCounter;
+};
 
 void RenderHairStrandsDebugInfo(
 	FRHICommandListImmediate& RHICmdList,
diff --git a/Engine/Source/Runtime/Renderer/Private/HairStrands/HairStrandsEnvironment.cpp b/Engine/Source/Runtime/Renderer/Private/HairStrands/HairStrandsEnvironment.cpp
--- a/Engine/Source/Runtime/Renderer/Private/HairStrands/HairStrandsEnvironment.cpp
+++ b/Engine/Source/Runtime/Renderer/Private/HairStrands/HairStrandsEnvironment.cpp
@@ -57,6 +57,9 @@ static FAutoConsoleVariableRef CVarHairStrandsSkyLightingDistanceThreshold(TEXT(
 static int32 GHairStrandsSkyLighting_IntegrationType = 0;
 static FAutoConsoleVariableRef CVarHairStrandsSkyLighting_IntegrationType(TEXT("r.HairStrands.SkyLighting.IntegrationType"), GHairStrandsSkyLighting_IntegrationType, TEXT("Hair env. lighting integration type (0:Adhoc, 1:Uniform."), ECVF_Scalability | ECVF_RenderThreadSafe);
 
+static int32 GHairStrandsSkyLighting_DebugSample = 0;
+static FAutoConsoleVariableRef CVarHairStrandsSkyLighting_DebugSample(TEXT("r.HairStrands.SkyLighting.DebugSample"), GHairStrandsSkyLighting_DebugSample, TEXT("Enable debug view for visualizing sample used for the sky integration"), ECVF_Scalability | ECVF_RenderThreadSafe);
+
 ///////////////////////////////////////////////////////////////////////////////////////////////////
 
 enum class EHairLightingIntegrationType : uint8
@@ -173,10 +176,10 @@ static void AddHairStrandsEnvironmentAOPass(
 class FHairEnvironmentLighting
 {
 public:
-	class FGroupSize : SHADER_PERMUTATION_INT("PERMUTATION_GROUP_SIZE", 2);
-	class FSampleSet : SHADER_PERMUTATION_INT("PERMUTATION_SAMPLESET", 2);
-	class FIntegrationType : SHADER_PERMUTATION_INT("PERMUTATION_INTEGRATION_TYPE", 3);
-	using FPermutationDomain = TShaderPermutationDomain<FSampleSet, FIntegrationType>;
+	class FSampleSet		: SHADER_PERMUTATION_INT("PERMUTATION_SAMPLESET", 2);
+	class FIntegrationType	: SHADER_PERMUTATION_INT("PERMUTATION_INTEGRATION_TYPE", 3);
+	class FDebug			: SHADER_PERMUTATION_INT("PERMUTATION_DEBUG", 2); 
+	using FPermutationDomain = TShaderPermutationDomain<FSampleSet, FIntegrationType, FDebug>;
 
 	static void ModifyCompilationEnvironment(const FGlobalShaderPermutationParameters& Parameters, FShaderCompilerEnvironment& OutEnvironment)
 	{
@@ -250,9 +253,6 @@ public:
 		SHADER_PARAMETER_STRUCT_REF(FForwardLightData, ForwardLightData)
 
 		SHADER_PARAMETER_STRUCT_REF(FVirtualVoxelParameters, VirtualVoxel)
-
-		SHADER_PARAMETER_STRUCT_INCLUDE(ShaderDrawDebug::FShaderDrawDebugParameters, ShaderDrawParameters)
-
 	END_SHADER_PARAMETER_STRUCT()
 };
 
@@ -260,7 +260,7 @@ class FHairEnvironmentLightingVS : public FGlobalShader
 {
 	DECLARE_GLOBAL_SHADER(FHairEnvironmentLightingVS);
 	SHADER_USE_PARAMETER_STRUCT(FHairEnvironmentLightingVS, FGlobalShader)
-	using FPermutationDomain = FHairEnvironmentLighting::FPermutationDomain;
+	//using FPermutationDomain = FHairEnvironmentLighting::FPermutationDomain;
 
 	BEGIN_SHADER_PARAMETER_STRUCT(FParameters, )
 		SHADER_PARAMETER_STRUCT_INCLUDE(FHairEnvironmentLighting::FParameters, Common)
@@ -286,6 +286,8 @@ class FHairEnvironmentLightingPS : public FGlobalShader
 
 	BEGIN_SHADER_PARAMETER_STRUCT(FParameters, )
 		SHADER_PARAMETER_STRUCT_INCLUDE(FHairEnvironmentLighting::FParameters, Common)
+		SHADER_PARAMETER_STRUCT_INCLUDE(ShaderDrawDebug::FShaderDrawDebugParameters, ShaderDrawParameters)
+		SHADER_PARAMETER_STRUCT_INCLUDE(FHairStrandsDebugData::FWriteParameters, DebugData)
 		RENDER_TARGET_BINDING_SLOTS()
 	END_SHADER_PARAMETER_STRUCT()
 
@@ -310,7 +312,8 @@ static void AddHairStrandsEnvironmentLightingPassPS(
 	const FHairStrandsVisibilityData& VisibilityData,
 	const FHairStrandsMacroGroupDatas& MacroGroupDatas,
 	const FHairStrandsMacroGroupData& MacroGroupData,
-	const FRDGTextureRef SceneColorTexture)
+	const FRDGTextureRef SceneColorTexture,
+	FHairStrandsDebugData::Data* DebugData)
 {
 	FSceneTextureParameters SceneTextures;
 	SetupSceneTextureParameters(GraphBuilder, &SceneTextures);
@@ -375,7 +378,12 @@ static void AddHairStrandsEnvironmentLightingPassPS(
 
 	if (ShaderDrawDebug::IsShaderDrawDebugEnabled(View))
 	{
-		ShaderDrawDebug::SetParameters(GraphBuilder, View.ShaderDrawData, PassParameters->ShaderDrawParameters);
+		ShaderDrawDebug::SetParameters(GraphBuilder, View.ShaderDrawData, ParametersPS->ShaderDrawParameters);
+	}
+
+	if (DebugData)
+	{
+		FHairStrandsDebugData::SetParameters(GraphBuilder, *DebugData, ParametersPS->DebugData);
 	}
 
 	// Bind hair data
@@ -387,11 +395,12 @@ static void AddHairStrandsEnvironmentLightingPassPS(
 	FHairEnvironmentLightingPS::FPermutationDomain PermutationVector;
 	PermutationVector.Set<FHairEnvironmentLighting::FSampleSet>(PassParameters->MultipleScatterSampleCount <= 16 ? 0 : 1);
 	PermutationVector.Set<FHairEnvironmentLighting::FIntegrationType>(uint32(IntegrationType));
+	PermutationVector.Set<FHairEnvironmentLighting::FDebug>(DebugData ? 1 : 0);
 	PermutationVector = FHairEnvironmentLighting::RemapPermutation(PermutationVector);
 
 	FIntPoint ViewportResolution = VisibilityData.SampleLightingViewportResolution;
 	const FViewInfo* CapturedView = &View;
-	TShaderMapRef<FHairEnvironmentLightingVS> VertexShader(View.ShaderMap, PermutationVector);
+	TShaderMapRef<FHairEnvironmentLightingVS> VertexShader(View.ShaderMap);
 	TShaderMapRef<FHairEnvironmentLightingPS> PixelShader(View.ShaderMap, PermutationVector);
 
 	check(VisibilityData.SampleLightingBuffer);
@@ -431,7 +440,7 @@ static void AddHairStrandsEnvironmentLightingPassPS(
 void RenderHairStrandsSceneColorScattering(
 	FRHICommandListImmediate& RHICmdList,
 	const TArray<FViewInfo>& Views,
-	const struct FHairStrandsDatas* HairDatas)
+	FHairStrandsDatas* HairDatas)
 {
 	if (Views.Num() == 0 || !HairDatas || GHairScatterSceneLighting <= 0)
 		return;
@@ -471,7 +480,7 @@ void RenderHairStrandsSceneColorScattering(
 			{
 				if (MacroGroupData.bNeedScatterSceneLighting)
 				{
-					AddHairStrandsEnvironmentLightingPassPS(GraphBuilder, View, VisibilityData, MacroGroupDatas, MacroGroupData, SceneColorTexture);
+					AddHairStrandsEnvironmentLightingPassPS(GraphBuilder, View, VisibilityData, MacroGroupDatas, MacroGroupData, SceneColorTexture, nullptr);
 				}
 			}
 			GraphBuilder.Execute();
@@ -483,11 +492,23 @@ void RenderHairStrandsEnvironmentLighting(
 	FRDGBuilder& GraphBuilder,
 	const uint32 ViewIndex,
 	const TArray<FViewInfo>& Views, 
-	const FHairStrandsDatas* HairDatas)
+	FHairStrandsDatas* HairDatas)
 {
 	if (!GetHairStrandsSkyLightingEnable() || !HairDatas)
 		return;
 
+	struct FHairEnvDebugData
+	{
+		FRDGBufferRef ShadingPointBuffer = nullptr;
+		FRDGBufferRef SampleBuffer = nullptr;
+		FRDGBufferRef ShadingPointCounter = nullptr;
+
+		FRDGBufferUAVRef ShadingPointBufferUAV = nullptr;
+		FRDGBufferUAVRef SampleBufferUAV = nullptr;
+		FRDGBufferUAVRef ShadingPointCounterUAV = nullptr;
+	};
+
+
 	check(ViewIndex < uint32(Views.Num()));
 	check(ViewIndex < uint32(HairDatas->HairVisibilityViews.HairDatas.Num()));	
 	const FHairStrandsVisibilityData& VisibilityData = HairDatas->HairVisibilityViews.HairDatas[ViewIndex];
@@ -497,11 +518,23 @@ void RenderHairStrandsEnvironmentLighting(
 		return;
 	}
 
+	const bool bDebugSamplingEnable = GHairStrandsSkyLighting_DebugSample > 0;
+	FHairStrandsDebugData::Data DebugData;
+	if (bDebugSamplingEnable)
+	{
+		DebugData = FHairStrandsDebugData::CreateData(GraphBuilder);
+	}
+
 	const FViewInfo& View = Views[ViewIndex];
 	const FHairStrandsMacroGroupDatas& MacroGroupDatas = HairDatas->MacroGroupsPerViews.Views[ViewIndex];
 	for (const FHairStrandsMacroGroupData& MacroGroupData : HairDatas->MacroGroupsPerViews.Views[ViewIndex].Datas)
 	{
-		AddHairStrandsEnvironmentLightingPassPS(GraphBuilder, View, VisibilityData, MacroGroupDatas, MacroGroupData, nullptr);
+		AddHairStrandsEnvironmentLightingPassPS(GraphBuilder, View, VisibilityData, MacroGroupDatas, MacroGroupData, nullptr, bDebugSamplingEnable ? &DebugData : nullptr);
+	}
+
+	if (bDebugSamplingEnable)
+	{
+		FHairStrandsDebugData::ExtractData(GraphBuilder, DebugData, HairDatas->DebugData);
 	}
 }
 
diff --git a/Engine/Source/Runtime/Renderer/Private/HairStrands/HairStrandsEnvironment.h b/Engine/Source/Runtime/Renderer/Private/HairStrands/HairStrandsEnvironment.h
--- a/Engine/Source/Runtime/Renderer/Private/HairStrands/HairStrandsEnvironment.h
+++ b/Engine/Source/Runtime/Renderer/Private/HairStrands/HairStrandsEnvironment.h
@@ -20,9 +20,9 @@ void RenderHairStrandsEnvironmentLighting(
 	FRDGBuilder& GraphBuilder,
 	const uint32 ViewIndex,
 	const TArray<FViewInfo>& Views,
-	const struct FHairStrandsDatas* HairDatas);
+	struct FHairStrandsDatas* HairDatas);
 
 void RenderHairStrandsSceneColorScattering(
 	FRHICommandListImmediate& RHICmdList,
 	const TArray<FViewInfo>& Views,
-	const struct FHairStrandsDatas* HairDatas);
+	struct FHairStrandsDatas* HairDatas);
diff --git a/Engine/Source/Runtime/Renderer/Private/HairStrands/HairStrandsRendering.h b/Engine/Source/Runtime/Renderer/Private/HairStrands/HairStrandsRendering.h
--- a/Engine/Source/Runtime/Renderer/Private/HairStrands/HairStrandsRendering.h
+++ b/Engine/Source/Runtime/Renderer/Private/HairStrands/HairStrandsRendering.h
@@ -27,6 +27,7 @@ struct FHairStrandsDatas
 {
 	FHairStrandsVisibilityViews HairVisibilityViews;
 	FHairStrandsMacroGroupViews MacroGroupsPerViews;
+	FHairStrandsDebugData DebugData;
 };
 
 enum class EHairStrandsInterpolationType
diff --git a/Engine/Source/Runtime/Renderer/Private/IndirectLightRendering.cpp b/Engine/Source/Runtime/Renderer/Private/IndirectLightRendering.cpp
--- a/Engine/Source/Runtime/Renderer/Private/IndirectLightRendering.cpp
+++ b/Engine/Source/Runtime/Renderer/Private/IndirectLightRendering.cpp
@@ -595,7 +595,7 @@ void FDeferredShadingSceneRenderer::RenderDiffuseIndirectAndAmbientOcclusion(FRH
 	GraphBuilder.Execute();
 }
 
-void FDeferredShadingSceneRenderer::RenderDeferredReflectionsAndSkyLighting(FRHICommandListImmediate& RHICmdList, TRefCountPtr<IPooledRenderTarget>& DynamicBentNormalAO, TRefCountPtr<IPooledRenderTarget>& VelocityRT, const FHairStrandsDatas* HairDatas)
+void FDeferredShadingSceneRenderer::RenderDeferredReflectionsAndSkyLighting(FRHICommandListImmediate& RHICmdList, TRefCountPtr<IPooledRenderTarget>& DynamicBentNormalAO, TRefCountPtr<IPooledRenderTarget>& VelocityRT, FHairStrandsDatas* HairDatas)
 {
 	check(RHICmdList.IsOutsideRenderPass());
 
