Fix bad hair behaviour with spatial anti aliasing

diff --git a/Engine/Source/Runtime/Renderer/Private/DeferredShadingRenderer.cpp b/Engine/Source/Runtime/Renderer/Private/DeferredShadingRenderer.cpp
--- a/Engine/Source/Runtime/Renderer/Private/DeferredShadingRenderer.cpp
+++ b/Engine/Source/Runtime/Renderer/Private/DeferredShadingRenderer.cpp
@@ -1680,7 +1680,7 @@ void FDeferredShadingSceneRenderer::Render(FRHICommandListImmediate& RHICmdList)
 
 	// Interpolation needs to happen after the skin cache run as there is a dependency 
 	// on the skin cache output.
-	const bool bRunHairStrands = IsHairStrandsEnable(Scene->GetShaderPlatform()) && Views.Num() > 0;
+	const bool bRunHairStrands = IsHairStrandsEnable(Scene->GetShaderPlatform()) && (Views.Num() > 0) && !ViewFamily.bWorldIsPaused;
 	FHairStrandClusterData HairClusterData;
 	if (bRunHairStrands)
 	{
