Use consistent parameters when allocating the depth texture through stereo render target manager

diff --git a/Engine/Source/Runtime/Renderer/Private/PostProcess/SceneRenderTargets.cpp b/Engine/Source/Runtime/Renderer/Private/PostProcess/SceneRenderTargets.cpp
--- a/Engine/Source/Runtime/Renderer/Private/PostProcess/SceneRenderTargets.cpp
+++ b/Engine/Source/Runtime/Renderer/Private/PostProcess/SceneRenderTargets.cpp
@@ -2304,12 +2304,12 @@ void FSceneRenderTargets::AllocateMobileRenderTargets(FRHICommandListImmediate&
 	AllocateFoveationTexture(RHICmdList);
 
 	static const auto MobileMultiViewCVar = IConsoleManager::Get().FindTConsoleVariableDataInt(TEXT("vr.MobileMultiView"));
-	static const auto CVarMobileMultiViewDirect = IConsoleManager::Get().FindTConsoleVariableDataInt(TEXT("vr.MobileMultiView.Direct"));
 
 	const bool bIsUsingMobileMultiView = (GSupportsMobileMultiView || GRHISupportsArrayIndexFromAnyShader) && (MobileMultiViewCVar && MobileMultiViewCVar->GetValueOnAnyThread() != 0);
 
-	// TODO: Test platform support for direct
-	const bool bIsMobileMultiViewDirectEnabled = (CVarMobileMultiViewDirect && CVarMobileMultiViewDirect->GetValueOnAnyThread() != 0);
+	// If the plugin uses separate render targets it is required to support mobile multi-view direct
+	IStereoRenderTargetManager* const StereoRenderTargetManager = GEngine->StereoRenderingDevice.IsValid() ? GEngine->StereoRenderingDevice->GetRenderTargetManager() : nullptr;
+	const bool bIsMobileMultiViewDirectEnabled = StereoRenderTargetManager && StereoRenderTargetManager->ShouldUseSeparateRenderTarget();
 
 	if (bIsUsingMobileMultiView)
 	{
@@ -2456,7 +2456,7 @@ void FSceneRenderTargets::AllocateCommonDepthTargets(FRHICommandList& RHICmdList
 	if (!SceneDepthZ || GFastVRamConfig.bDirty)
 	{
 		FTexture2DRHIRef DepthTex, SRTex;
-		bHMDAllocatedDepthTarget = StereoRenderTargetManager && StereoRenderTargetManager->AllocateDepthTexture(0, BufferSize.X, BufferSize.Y, PF_X24_G8, 0, TexCreate_None, TexCreate_DepthStencilTargetable, DepthTex, SRTex, GetNumSceneColorMSAASamples(CurrentFeatureLevel));
+		bHMDAllocatedDepthTarget = StereoRenderTargetManager && StereoRenderTargetManager->AllocateDepthTexture(0, BufferSize.X, BufferSize.Y, PF_DepthStencil, 1, TexCreate_None, TexCreate_DepthStencilTargetable | TexCreate_ShaderResource | TexCreate_InputAttachmentRead, DepthTex, SRTex, GetNumSceneColorMSAASamples(CurrentFeatureLevel));
 
 		// Allow UAV depth?
 		const uint32 DepthUAVFlag = GRHISupportsDepthUAV ? TexCreate_UAV : 0;
