Fix a bug of HistogramCS for mobile.

diff --git a/Engine/Shaders/Private/PostProcessMobile.usf b/Engine/Shaders/Private/PostProcessMobile.usf
--- a/Engine/Shaders/Private/PostProcessMobile.usf
+++ b/Engine/Shaders/Private/PostProcessMobile.usf
@@ -1155,14 +1155,14 @@ static const float INV_FLOAT_PRECISION = 1e-5;
 uint NumEntries;
 
 uint ClearValue;
-RWTexture2D<uint> UAV;
+RWBuffer<uint> UAV;
 
 [numthreads(64, 1, 1)]
 void ClearUAVUIntCS(uint3 DispatchThreadId : SV_DispatchThreadID)
 {
 	if (DispatchThreadId.x < NumEntries)
 	{
-		UAV[DispatchThreadId.xy] = ClearValue;
+		UAV[DispatchThreadId.x] = ClearValue;
 	}
 }
 #endif
@@ -1309,7 +1309,7 @@ const static  float InvHistogramSizeMinusOne = 1.0f / (float)(HISTOGRAM_SIZE - 1
 
 
 #if HISTOGRAM_COMPUTE_SHADER
-// Output histogram texture (UAV)
+// Output histogram buffer (UAV)
 RWBuffer<uint> RWHistogramBuffer;
 
 const static uint QUARTER_HISTOGRAM_SIZE = HISTOGRAM_SIZE / 4;
@@ -1439,11 +1439,11 @@ void Histogram_MainCS(
 #endif
 
 #if HISTOGRAM_EYEADAPTATION_COMPUTE_SHADER
-groupshared float SharedHistogram[HISTOGRAM_SIZE];
+Buffer<uint> HistogramBuffer;
 
 float GetHistogramBucket_Mobile(uint BucketIndex)
 {
-	return SharedHistogram[BucketIndex];
+	return HistogramBuffer[BucketIndex];
 }
 
 float ComputeHistogramSum_Mobile()
@@ -1498,18 +1498,11 @@ float ComputeEyeAdaptationExposure_Mobile()
 	return LumAve;
 }
 
-Buffer<uint> HistogramBuffer;
-
 RWBuffer<float4> OutputBuffer;
 
 [numthreads(1, 1, 1)]
 void HistogramEyeAdaptationCS(uint2 DispatchThreadId : SV_DispatchThreadID)
 {
-	UNROLL for (uint i = 0; i < HISTOGRAM_SIZE; ++i)
-	{
-		SharedHistogram[i] = HistogramBuffer[i];
-	}
-
 	float4 OutColor = 0;
 
 	const float AverageSceneLuminance = ComputeEyeAdaptationExposure_Mobile();
diff --git a/Engine/Source/Runtime/Renderer/Private/PostProcess/PostProcessMobile.cpp b/Engine/Source/Runtime/Renderer/Private/PostProcess/PostProcessMobile.cpp
--- a/Engine/Source/Runtime/Renderer/Private/PostProcess/PostProcessMobile.cpp
+++ b/Engine/Source/Runtime/Renderer/Private/PostProcess/PostProcessMobile.cpp
@@ -2606,9 +2606,9 @@ class FClearUAVUIntCS_ES2 : public FGlobalShader
 		SHADER_PARAMETER_UAV(RWBuffer<uint>, UAV)
 		SHADER_PARAMETER(uint32, ClearValue)
 		SHADER_PARAMETER(uint32, NumEntries)
-		END_SHADER_PARAMETER_STRUCT()
+	END_SHADER_PARAMETER_STRUCT()
 
-		FClearUAVUIntCS_ES2() {}
+	FClearUAVUIntCS_ES2() {}
 public:
 
 	FClearUAVUIntCS_ES2(const ShaderMetaType::CompiledShaderInitializerType& Initializer)
@@ -2756,7 +2756,6 @@ void FRCPassPostProcessAverageLuminanceES2::Process(FRenderingCompositePassConte
 
 	{
 		// clear Average Luminance History
-		if (IsAndroidOpenGLESPlatform(View.GetShaderPlatform()))
 		{
 			TShaderMapRef<FClearUAVUIntCS_ES2> ClearShader(Context.GetShaderMap());
 
@@ -2768,10 +2767,6 @@ void FRCPassPostProcessAverageLuminanceES2::Process(FRenderingCompositePassConte
 
 			UnsetShaderUAVs(Context.RHICmdList, ClearShader, Context.RHICmdList.GetBoundComputeShader());
 		}
-		else
-		{
-			Context.RHICmdList.ClearUAVUint(GAverageLuminanceBuffer->UnorderedAccessViewRHI, FUintVector4(0, 0, 0, 0));
-		}
 
 		Context.RHICmdList.TransitionResource(EResourceTransitionAccess::ERWBarrier, EResourceTransitionPipeline::EComputeToCompute, GAverageLuminanceBuffer->UnorderedAccessViewRHI);
 
@@ -2940,18 +2935,21 @@ class FPostProcessHistogramCS_ES2 : public FGlobalShader
 {
 public:
 	// Changing these numbers requires PostProcessMobile.usf to be recompiled.
+	static const uint32 MetalThreadGroupSizeX = 8;
+	static const uint32 MetalThreadGroupSizeY = 4; // the maximum total threadgroup memory allocation on A7 and A8 GPU is 16KB-32B, so it has to limit the thread group size on IOS/TVOS platform.
+
+	static const uint32 MetalLoopCountX = 2;
+	static const uint32 MetalLoopCountY = 4;
+
+	// The number of texels on each axis processed by a single thread group.
+	static const FIntPoint MetalTexelsPerThreadGroup;
+
 	static const uint32 ThreadGroupSizeX = 8;
-#if PLATFORM_IOS || PLATFORM_TVOS	// the maximum total threadgroup memory allocation on A7 and A8 GPU is 16KB-32B, so it has to limit the thread group size on IOS/TVOS platform.
-	static const uint32 ThreadGroupSizeY = 4;
-#else
 	static const uint32 ThreadGroupSizeY = 8;
-#endif
+
 	static const uint32 LoopCountX = 2;
-#if PLATFORM_IOS || PLATFORM_TVOS
-	static const uint32 LoopCountY = 4;
-#else
 	static const uint32 LoopCountY = 2;
-#endif
+
 	static const uint32 HistogramSize = 64; // HistogramSize must be 64 and ThreadGroupSizeX * ThreadGroupSizeY must be larger than 32
 
 	// The number of texels on each axis processed by a single thread group.
@@ -2975,11 +2973,13 @@ public:
 
 	static void ModifyCompilationEnvironment(const FGlobalShaderPermutationParameters& Parameters, FShaderCompilerEnvironment& OutEnvironment)
 	{
+		bool bIsMetalMobilePlatform = IsMetalMobilePlatform(Parameters.Platform);
+
 		FGlobalShader::ModifyCompilationEnvironment(Parameters, OutEnvironment);
-		OutEnvironment.SetDefine(TEXT("THREADGROUP_SIZEX"), ThreadGroupSizeX);
-		OutEnvironment.SetDefine(TEXT("THREADGROUP_SIZEY"), ThreadGroupSizeY);
-		OutEnvironment.SetDefine(TEXT("LOOP_SIZEX"), LoopCountX);
-		OutEnvironment.SetDefine(TEXT("LOOP_SIZEY"), LoopCountY);
+		OutEnvironment.SetDefine(TEXT("THREADGROUP_SIZEX"), bIsMetalMobilePlatform ? MetalThreadGroupSizeX : ThreadGroupSizeX);
+		OutEnvironment.SetDefine(TEXT("THREADGROUP_SIZEY"), bIsMetalMobilePlatform ? MetalThreadGroupSizeY : ThreadGroupSizeY);
+		OutEnvironment.SetDefine(TEXT("LOOP_SIZEX"), bIsMetalMobilePlatform ? MetalLoopCountX : LoopCountX);
+		OutEnvironment.SetDefine(TEXT("LOOP_SIZEY"), bIsMetalMobilePlatform ? MetalLoopCountY : LoopCountY);
 		OutEnvironment.SetDefine(TEXT("HISTOGRAM_COMPUTE_SHADER"), 1u);
 		OutEnvironment.CompilerFlags.Add(CFLAG_StandardOptimization);
 	}
@@ -3014,6 +3014,8 @@ public:
 	}
 };
 
+const FIntPoint FPostProcessHistogramCS_ES2::MetalTexelsPerThreadGroup(MetalThreadGroupSizeX * MetalLoopCountX * 2, MetalThreadGroupSizeY * MetalLoopCountY * 2); // Multiply 2 because we use bilinear filter, to reduce the sample count
+
 const FIntPoint FPostProcessHistogramCS_ES2::TexelsPerThreadGroup(ThreadGroupSizeX * LoopCountX * 2, ThreadGroupSizeY * LoopCountY * 2); // Multiply 2 because we use bilinear filter, to reduce the sample count
 
 IMPLEMENT_SHADER_TYPE(, FPostProcessHistogramCS_ES2, TEXT("/Engine/Private/PostProcessMobile.usf"), TEXT("Histogram_MainCS"), SF_Compute);
@@ -3046,7 +3048,6 @@ void FRCPassPostProcessHistogramES2::Process(FRenderingCompositePassContext& Con
 
 	{
 		// clear Histogram History
-		if (IsAndroidOpenGLESPlatform(View.GetShaderPlatform()))
 		{
 			TShaderMapRef<FClearUAVUIntCS_ES2> ClearShader(Context.GetShaderMap());
 
@@ -3058,10 +3059,6 @@ void FRCPassPostProcessHistogramES2::Process(FRenderingCompositePassContext& Con
 
 			UnsetShaderUAVs(Context.RHICmdList, ClearShader, Context.RHICmdList.GetBoundComputeShader());
 		}
-		else
-		{
-			Context.RHICmdList.ClearUAVUint(GHistogramBuffer->UnorderedAccessViewRHI, FUintVector4(0, 0, 0, 0));
-		}
 
 		Context.RHICmdList.TransitionResource(EResourceTransitionAccess::ERWBarrier, EResourceTransitionPipeline::EComputeToCompute, GHistogramBuffer->UnorderedAccessViewRHI);
 
@@ -3073,7 +3070,7 @@ void FRCPassPostProcessHistogramES2::Process(FRenderingCompositePassContext& Con
 			Context.RHICmdList.SetComputeShader(ComputeShader.GetComputeShader());
 
 			const FIntPoint SrcRectExtent = InputDesc->Extent;
-			const FIntPoint ThreadGroupCount = FIntPoint::DivideAndRoundUp(SrcRectExtent, FPostProcessHistogramCS_ES2::TexelsPerThreadGroup);
+			const FIntPoint ThreadGroupCount = FIntPoint::DivideAndRoundUp(SrcRectExtent, IsMetalMobilePlatform(Context.View.GetShaderPlatform()) ? FPostProcessHistogramCS_ES2::MetalTexelsPerThreadGroup : FPostProcessHistogramCS_ES2::TexelsPerThreadGroup);
 
 			ComputeShader->SetCS(Context.RHICmdList, ComputeShader, Context, SrcRectExtent, InputRenderTarget.ShaderResourceTexture, GHistogramBuffer->UnorderedAccessViewRHI);
 
