Disable lightmapdensity shaders when StaticLighting is disabled.
Together with cl 11724687 this saves quite a lot of shaders

diff --git a/Engine/Source/Editor/UnrealEd/Private/SEditorViewportViewMenu.cpp b/Engine/Source/Editor/UnrealEd/Private/SEditorViewportViewMenu.cpp
--- a/Engine/Source/Editor/UnrealEd/Private/SEditorViewportViewMenu.cpp
+++ b/Engine/Source/Editor/UnrealEd/Private/SEditorViewportViewMenu.cpp
@@ -253,7 +253,10 @@ void SEditorViewportViewMenu::FillViewMenu(UToolMenu* Menu) const
 							if (FeatureLevel == ERHIFeatureLevel::SM5)
 							{
 								Section.AddMenuEntry(BaseViewportCommands.LightComplexityMode, UViewModeUtils::GetViewModeDisplayName(VMI_LightComplexity));
-								Section.AddMenuEntry(BaseViewportCommands.LightmapDensityMode, UViewModeUtils::GetViewModeDisplayName(VMI_LightmapDensity));
+								if (IConsoleManager::Get().FindTConsoleVariableDataInt(TEXT("r.AllowStaticLighting"))->GetValueOnAnyThread() != 0)
+								{
+									Section.AddMenuEntry(BaseViewportCommands.LightmapDensityMode, UViewModeUtils::GetViewModeDisplayName(VMI_LightmapDensity));
+								}
 								Section.AddMenuEntry(BaseViewportCommands.StationaryLightOverlapMode, UViewModeUtils::GetViewModeDisplayName(VMI_StationaryLightOverlap));
 							}
 
diff --git a/Engine/Source/Runtime/Renderer/Private/LightMapDensityRendering.h b/Engine/Source/Runtime/Renderer/Private/LightMapDensityRendering.h
--- a/Engine/Source/Runtime/Renderer/Private/LightMapDensityRendering.h
+++ b/Engine/Source/Runtime/Renderer/Private/LightMapDensityRendering.h
@@ -58,7 +58,10 @@ public:
 
 	static bool ShouldCompilePermutation(const FMeshMaterialShaderPermutationParameters& Parameters)
 	{
+		static const auto AllowStaticLightingVar = IConsoleManager::Get().FindTConsoleVariableDataInt(TEXT("r.AllowStaticLighting"));
+		bool bAllowStaticLighting = AllowStaticLightingVar->GetValueOnAnyThread() != 0;
 		return  AllowDebugViewmodes(Parameters.Platform) 
+				&& bAllowStaticLighting
 				&& (Parameters.MaterialParameters.bIsSpecialEngineMaterial || Parameters.MaterialParameters.bIsMasked || Parameters.MaterialParameters.bMaterialMayModifyMeshPosition)
 				&& LightMapPolicyType::ShouldCompilePermutation(Parameters)
 				&& IsFeatureLevelSupported(Parameters.Platform, ERHIFeatureLevel::SM5);
@@ -111,7 +114,12 @@ public:
 
 	static bool ShouldCompilePermutation(const FMeshMaterialShaderPermutationParameters& Parameters)
 	{
+
+		static const auto AllowStaticLightingVar = IConsoleManager::Get().FindTConsoleVariableDataInt(TEXT("r.AllowStaticLighting"));
+		bool bAllowStaticLighting = AllowStaticLightingVar->GetValueOnAnyThread() != 0;
+
 		return AllowDebugViewmodes(Parameters.Platform) 
+			&& bAllowStaticLighting
 			&& FBaseHS::ShouldCompilePermutation(Parameters)
 			&& TLightMapDensityVS<LightMapPolicyType>::ShouldCompilePermutation(Parameters);
 	}
@@ -142,7 +150,11 @@ public:
 
 	static bool ShouldCompilePermutation(const FMeshMaterialShaderPermutationParameters& Parameters)
 	{
+		static const auto AllowStaticLightingVar = IConsoleManager::Get().FindTConsoleVariableDataInt(TEXT("r.AllowStaticLighting"));
+		bool bAllowStaticLighting = AllowStaticLightingVar->GetValueOnAnyThread() != 0;
+
 		return AllowDebugViewmodes(Parameters.Platform) 
+			&& bAllowStaticLighting
 			&& FBaseDS::ShouldCompilePermutation(Parameters)
 			&& TLightMapDensityVS<LightMapPolicyType>::ShouldCompilePermutation(Parameters);		
 	}
@@ -174,7 +186,10 @@ public:
 
 	static bool ShouldCompilePermutation(const FMeshMaterialShaderPermutationParameters& Parameters)
 	{
+		static const auto AllowStaticLightingVar = IConsoleManager::Get().FindTConsoleVariableDataInt(TEXT("r.AllowStaticLighting"));
+		bool bAllowStaticLighting = AllowStaticLightingVar->GetValueOnAnyThread() != 0;
 		return	AllowDebugViewmodes(Parameters.Platform) 
+				&& bAllowStaticLighting
 				&& (Parameters.MaterialParameters.bIsSpecialEngineMaterial || Parameters.MaterialParameters.bIsMasked || Parameters.MaterialParameters.bMaterialMayModifyMeshPosition)
 				&& LightMapPolicyType::ShouldCompilePermutation(Parameters)
 				&& IsFeatureLevelSupported(Parameters.Platform, ERHIFeatureLevel::SM5);
