Add an option to not compile the surfel shaders.
This is a workaround for the pending removal of the code, which will happen on another branch.

Together with cl 11724687 this saves quite a lot of shaders

diff --git a/Engine/Source/Runtime/Renderer/Private/ConvertToUniformMesh.cpp b/Engine/Source/Runtime/Renderer/Private/ConvertToUniformMesh.cpp
--- a/Engine/Source/Runtime/Renderer/Private/ConvertToUniformMesh.cpp
+++ b/Engine/Source/Runtime/Renderer/Private/ConvertToUniformMesh.cpp
@@ -35,7 +35,9 @@ protected:
 
 	static bool ShouldCompilePermutation(const FMeshMaterialShaderPermutationParameters& Parameters)
 	{
+		bool bEnabled = false;
 		return IsFeatureLevelSupported(Parameters.Platform, ERHIFeatureLevel::SM5) 
+			&& bEnabled
 			&& DoesPlatformSupportDistanceFieldGI(Parameters.Platform)
 			&& (FCString::Strstr(Parameters.VertexFactoryType->GetName(), TEXT("LocalVertexFactory")) != NULL
 				|| FCString::Strstr(Parameters.VertexFactoryType->GetName(), TEXT("InstancedStaticMeshVertexFactory")) != NULL);
@@ -107,7 +109,9 @@ protected:
 
 	static bool ShouldCompilePermutation(const FMeshMaterialShaderPermutationParameters& Parameters)
 	{
+		bool bEnabled = false;
 		return IsFeatureLevelSupported(Parameters.Platform, ERHIFeatureLevel::SM5) 
+			&& bEnabled
 			&& DoesPlatformSupportDistanceFieldGI(Parameters.Platform)
 			&& (FCString::Strstr(Parameters.VertexFactoryType->GetName(), TEXT("LocalVertexFactory")) != NULL
 				|| FCString::Strstr(Parameters.VertexFactoryType->GetName(), TEXT("InstancedStaticMeshVertexFactory")) != NULL);
@@ -315,7 +319,9 @@ public:
 			return false;
 		}
 
-		return IsFeatureLevelSupported(Parameters.Platform, ERHIFeatureLevel::SM5) && DoesPlatformSupportDistanceFieldGI(Parameters.Platform);
+		bool bEnabled = false;
+
+		return IsFeatureLevelSupported(Parameters.Platform, ERHIFeatureLevel::SM5) && DoesPlatformSupportDistanceFieldGI(Parameters.Platform) && bEnabled;
 	}
 
 	static void ModifyCompilationEnvironment(const FMaterialShaderPermutationParameters& Parameters, FShaderCompilerEnvironment& OutEnvironment)
