Added options to turn off compilation of ray tracing material closest or any-hit shaders. This may be useful for titles that use custom shaders or don't require full material evaluation (i.e. only using ray tracing for shadows or AO).

New CVars:
- r.RayTracing.CompileMaterialCHS (default 1)
- r.RayTracing.CompileMaterialAHS (default 1)

The value is read-only and must match between cook-time and run-time.

diff --git a/Engine/Shaders/Private/RayTracing/RayTracingMaterialHitShaders.usf b/Engine/Shaders/Private/RayTracing/RayTracingMaterialHitShaders.usf
--- a/Engine/Shaders/Private/RayTracing/RayTracingMaterialHitShaders.usf
+++ b/Engine/Shaders/Private/RayTracing/RayTracingMaterialHitShaders.usf
@@ -271,12 +271,15 @@ RAY_TRACING_ENTRY_CLOSEST_HIT(MaterialCHS,
 	FPackedMaterialClosestHitPayload, PackedPayload,
 	FDefaultAttributes, Attributes)
 {
+	PackedPayload.HitT = RayTCurrent();
+
+#if USE_MATERIAL_CLOSEST_HIT_SHADER
+
 	if (PackedPayload.IsMinimalPayloadMode())
 	{
 		// Minimal payload mode only fills FMinimalPayload::HitT, skipping actual material evaluation.
 		// This mode is used when tracing shadow rays against masked geometry.
 		// Dynamic branch is used to avoid compiling an extra shader permutation.
-		PackedPayload.HitT = RayTCurrent();
 		return;
 	}
 
@@ -328,7 +331,6 @@ RAY_TRACING_ENTRY_CLOSEST_HIT(MaterialCHS,
 
 	Payload.BlendingMode = MaterialBlendingMode;
 
-	Payload.HitT = RayTCurrent();
 	if (HitKind() == HIT_KIND_TRIANGLE_FRONT_FACE)
 	{
 		Payload.SetFrontFace();
@@ -446,12 +448,22 @@ RAY_TRACING_ENTRY_CLOSEST_HIT(MaterialCHS,
 #endif
 
 	PackedPayload = PackRayTracingPayload(Payload, PropagatedCone);
+
+#endif // USE_MATERIAL_CLOSEST_HIT_SHADER
 }
 
 RAY_TRACING_ENTRY_ANY_HIT(MaterialAHS,
 	FPackedMaterialClosestHitPayload, PackedPayload,
 	FDefaultAttributes, Attributes)
 {
+#if USE_MATERIAL_ANY_HIT_SHADER
+
+#if MATERIALBLENDING_SOLID
+	// Continue traversal by simply returning from this shader
+	// Usually this path is inactive because opaque materials don't have AHS bound
+	return;
+#endif // MATERIALBLENDING_SOLID
+
 #if MATERIALBLENDING_MASKED
 	ResolvedView = ResolveView();
 
@@ -488,5 +500,13 @@ RAY_TRACING_ENTRY_ANY_HIT(MaterialAHS,
 		// Continue traversal by simply returning from this shader
 		return;
 	}
-#endif
+#endif // MATERIALBLENDING_MASKED
+
+	// For any other blending models, ignore hit in anyhit shader if we're tracing visibility rays until we have proper handling for translucent
+	if (PackedPayload.IsIgnoreTranslucentMaterials())
+	{
+		IgnoreHit();
+	}
+
+#endif // USE_MATERIAL_ANY_HIT_SHADER
 }  
diff --git a/Engine/Source/Runtime/RenderCore/Private/Shader.cpp b/Engine/Source/Runtime/RenderCore/Private/Shader.cpp
--- a/Engine/Source/Runtime/RenderCore/Private/Shader.cpp
+++ b/Engine/Source/Runtime/RenderCore/Private/Shader.cpp
@@ -1697,4 +1697,16 @@ void ShaderMapAppendKeyString(EShaderPlatform Platform, FString& KeyString)
 	{
 		KeyString += TEXT("_PPDBM");
 	}
+
+	if (ShouldCompileRayTracingShadersForProject(Platform))
+	{
+		static const auto CVarCompileCHS = IConsoleManager::Get().FindConsoleVariable(TEXT("r.RayTracing.CompileMaterialCHS"));
+		static const auto CVarCompileAHS = IConsoleManager::Get().FindConsoleVariable(TEXT("r.RayTracing.CompileMaterialAHS"));
+		static const auto CVarTextureLod = IConsoleManager::Get().FindConsoleVariable(TEXT("r.RayTracing.UseTextureLod"));
+
+		KeyString += FString::Printf(TEXT("_RAY-CHS%dAHS%dLOD%d"),
+			CVarCompileCHS && CVarCompileCHS->GetBool() ? 1 : 0,
+			CVarCompileAHS && CVarCompileAHS->GetBool() ? 1 : 0,
+			CVarTextureLod && CVarTextureLod->GetBool() ? 1 : 0);
+	}
 }
diff --git a/Engine/Source/Runtime/Renderer/Private/RayTracing/RayTracingMaterialHitShaders.cpp b/Engine/Source/Runtime/Renderer/Private/RayTracing/RayTracingMaterialHitShaders.cpp
--- a/Engine/Source/Runtime/Renderer/Private/RayTracing/RayTracingMaterialHitShaders.cpp
+++ b/Engine/Source/Runtime/Renderer/Private/RayTracing/RayTracingMaterialHitShaders.cpp
@@ -19,6 +19,24 @@ static FAutoConsoleVariableRef CVarEnableRayTracingMaterials(
 	ECVF_RenderThreadSafe
 );
 
+int32 GCompileRayTracingMaterialCHS = 1;
+static FAutoConsoleVariableRef CVarCompileRayTracingMaterialCHS(
+	TEXT("r.RayTracing.CompileMaterialCHS"),
+	GCompileRayTracingMaterialCHS,
+	TEXT(" 0: skip compilation of closest-hit shaders for materials (useful if only shadows or ambient occlusion effects are needed)\n")
+	TEXT(" 1: compile closest hit shaders for all ray tracing materials (default)\n"),
+	ECVF_ReadOnly
+);
+
+int32 GCompileRayTracingMaterialAHS = 1;
+static FAutoConsoleVariableRef CVarCompileRayTracingMaterialAHS(
+	TEXT("r.RayTracing.CompileMaterialAHS"),
+	GCompileRayTracingMaterialAHS,
+	TEXT(" 0: skip compilation of any-hit shaders for materials (useful if alpha masked or translucent materials are not needed)\n")
+	TEXT(" 1: compile any hit shaders for all ray tracing materials (default)\n"),
+	ECVF_ReadOnly
+);
+
 // CVar defined in DeferredShadingRenderer.cpp
 extern int32 GRayTracingUseTextureLod;
 
@@ -109,8 +127,16 @@ public:
 
 	static bool ShouldCompilePermutation(const FMeshMaterialShaderPermutationParameters& Parameters)
 	{
+		if (!GCompileRayTracingMaterialAHS && !GCompileRayTracingMaterialCHS)
+		{
+			return false;
+		}
+
+		const bool bWantAnyHitShader = GCompileRayTracingMaterialAHS 
+			&& (Parameters.MaterialParameters.bIsMasked || Parameters.MaterialParameters.BlendMode == BLEND_Translucent);
+
 		return IsSupportedVertexFactoryType(Parameters.VertexFactoryType)
-			&& (Parameters.MaterialParameters.bIsMasked == UseAnyHitShader)
+			&& (bWantAnyHitShader == UseAnyHitShader)
 			&& LightMapPolicyType::ShouldCompilePermutation(Parameters)
 			&& ShouldCompileRayTracingShadersForProject(Parameters.Platform)
 			&& (bool)GRayTracingUseTextureLod == UseRayConeTextureLod;
@@ -118,6 +144,10 @@ public:
 
 	static void ModifyCompilationEnvironment(const FMaterialShaderPermutationParameters& Parameters, FShaderCompilerEnvironment& OutEnvironment)
 	{
+		// NOTE: Any CVars that are used in this function must be handled in ShaderMapAppendKeyString() to ensure shaders are recompiled when necessary.
+
+		OutEnvironment.SetDefine(TEXT("USE_MATERIAL_CLOSEST_HIT_SHADER"), GCompileRayTracingMaterialCHS ? 1 : 0);
+		OutEnvironment.SetDefine(TEXT("USE_MATERIAL_ANY_HIT_SHADER"), GCompileRayTracingMaterialAHS ? 1 : 0);
 		OutEnvironment.SetDefine(TEXT("USE_RAYTRACED_TEXTURE_RAYCONE_LOD"), UseRayConeTextureLod ? 1 : 0);
 		OutEnvironment.SetDefine(TEXT("SCENE_TEXTURES_DISABLED"), 1);
 		LightMapPolicyType::ModifyCompilationEnvironment(Parameters, OutEnvironment);
@@ -147,6 +177,34 @@ public:
 	}
 };
 
+class FTrivialMaterialCHS : public FMaterialCHS
+{
+	DECLARE_SHADER_TYPE(FTrivialMaterialCHS, MeshMaterial);
+public:
+
+	FTrivialMaterialCHS(const FMeshMaterialShaderType::CompiledShaderInitializerType& Initializer)
+		: FMaterialCHS(Initializer)
+	{}
+
+	FTrivialMaterialCHS() {}
+
+	static bool ShouldCompilePermutation(const FMeshMaterialShaderPermutationParameters& Parameters)
+	{
+		return IsSupportedVertexFactoryType(Parameters.VertexFactoryType)
+			&& ShouldCompileRayTracingShadersForProject(Parameters.Platform);
+	}
+
+	static void ModifyCompilationEnvironment(const FMaterialShaderPermutationParameters& Parameters, FShaderCompilerEnvironment& OutEnvironment)
+	{
+	}
+
+	static bool ValidateCompiledResult(EShaderPlatform Platform, const FShaderParameterMap& ParameterMap, TArray<FString>& OutError)
+	{
+		return true;
+	}
+};
+
+IMPLEMENT_MATERIAL_SHADER_TYPE(, FTrivialMaterialCHS, TEXT("/Engine/Private/RayTracing/RayTracingMaterialDefaultHitShaders.usf"), TEXT("closesthit=OpaqueShadowCHS"), SF_RayHitGroup);
 
 #define IMPLEMENT_MATERIALCHS_TYPE(LightMapPolicyType, LightMapPolicyName, AnyHitShaderName) \
 	typedef TMaterialCHS<LightMapPolicyType, false, false> TMaterialCHS##LightMapPolicyName; \
@@ -170,7 +228,10 @@ IMPLEMENT_GLOBAL_SHADER(FOpaqueShadowHitGroup, "/Engine/Private/RayTracing/RayTr
 template<typename LightMapPolicyType>
 static TShaderRef<FMaterialCHS> GetMaterialHitShader(const FMaterial& RESTRICT MaterialResource, const FVertexFactory* VertexFactory, bool UseTextureLod)
 {
-	if (MaterialResource.IsMasked())
+	const bool bMaterialsCompiled = GCompileRayTracingMaterialAHS || GCompileRayTracingMaterialCHS;
+	checkf(bMaterialsCompiled, TEXT(""));
+
+	if ((MaterialResource.IsMasked() || MaterialResource.GetBlendMode() == BLEND_Translucent) && GCompileRayTracingMaterialAHS)
 	{
 		if(UseTextureLod)
 		{ 
@@ -261,6 +322,9 @@ void FRayTracingMeshProcessor::Process(
 	const FUniformLightMapPolicy& RESTRICT LightMapPolicy,
 	const typename FUniformLightMapPolicy::ElementDataType& RESTRICT LightMapElementData)
 {
+
+	const bool bMaterialsCompiled = GCompileRayTracingMaterialAHS || GCompileRayTracingMaterialCHS;
+
 	const FVertexFactory* VertexFactory = MeshBatch.VertexFactory;
 
 	TMeshProcessorShaders<
@@ -271,27 +335,34 @@ void FRayTracingMeshProcessor::Process(
 		FMeshMaterialShader,
 		FMaterialCHS> RayTracingShaders;
 
-	const bool bUseTextureLOD = bool(GRayTracingUseTextureLod);
+	if (bMaterialsCompiled)
+	{
+		const bool bUseTextureLOD = bool(GRayTracingUseTextureLod);
 
-	switch (LightMapPolicy.GetIndirectPolicy())
+		switch (LightMapPolicy.GetIndirectPolicy())
+		{
+		case LMP_PRECOMPUTED_IRRADIANCE_VOLUME_INDIRECT_LIGHTING:
+			RayTracingShaders.RayHitGroupShader = GetMaterialHitShader<TUniformLightMapPolicy<LMP_PRECOMPUTED_IRRADIANCE_VOLUME_INDIRECT_LIGHTING>>(MaterialResource, VertexFactory, bUseTextureLOD);
+			break;
+		case LMP_LQ_LIGHTMAP:
+			RayTracingShaders.RayHitGroupShader = GetMaterialHitShader<TUniformLightMapPolicy<LMP_LQ_LIGHTMAP>>(MaterialResource, VertexFactory, bUseTextureLOD);
+			break;
+		case LMP_HQ_LIGHTMAP:
+			RayTracingShaders.RayHitGroupShader = GetMaterialHitShader<TUniformLightMapPolicy<LMP_HQ_LIGHTMAP>>(MaterialResource, VertexFactory, bUseTextureLOD);
+			break;
+		case LMP_DISTANCE_FIELD_SHADOWS_AND_HQ_LIGHTMAP:
+			RayTracingShaders.RayHitGroupShader = GetMaterialHitShader<TUniformLightMapPolicy<LMP_DISTANCE_FIELD_SHADOWS_AND_HQ_LIGHTMAP>>(MaterialResource, VertexFactory, bUseTextureLOD);
+			break;
+		case LMP_NO_LIGHTMAP:
+			RayTracingShaders.RayHitGroupShader = GetMaterialHitShader<TUniformLightMapPolicy<LMP_NO_LIGHTMAP>>(MaterialResource, VertexFactory, bUseTextureLOD);
+			break;
+		default:
+			check(false);
+		}
+	}
+	else
 	{
-	case LMP_PRECOMPUTED_IRRADIANCE_VOLUME_INDIRECT_LIGHTING:
-		RayTracingShaders.RayHitGroupShader = GetMaterialHitShader<TUniformLightMapPolicy<LMP_PRECOMPUTED_IRRADIANCE_VOLUME_INDIRECT_LIGHTING>>(MaterialResource, VertexFactory, bUseTextureLOD);
-		break;
-	case LMP_LQ_LIGHTMAP:
-		RayTracingShaders.RayHitGroupShader = GetMaterialHitShader<TUniformLightMapPolicy<LMP_LQ_LIGHTMAP>>(MaterialResource, VertexFactory, bUseTextureLOD);
-		break;
-	case LMP_HQ_LIGHTMAP:
-		RayTracingShaders.RayHitGroupShader = GetMaterialHitShader<TUniformLightMapPolicy<LMP_HQ_LIGHTMAP>>(MaterialResource, VertexFactory, bUseTextureLOD);
-		break;
-	case LMP_DISTANCE_FIELD_SHADOWS_AND_HQ_LIGHTMAP:
-		RayTracingShaders.RayHitGroupShader = GetMaterialHitShader<TUniformLightMapPolicy<LMP_DISTANCE_FIELD_SHADOWS_AND_HQ_LIGHTMAP>>(MaterialResource, VertexFactory, bUseTextureLOD);
-		break;
-	case LMP_NO_LIGHTMAP:
-		RayTracingShaders.RayHitGroupShader = GetMaterialHitShader<TUniformLightMapPolicy<LMP_NO_LIGHTMAP>>(MaterialResource, VertexFactory, bUseTextureLOD);
-		break;
-	default:
-		check(false);
+		RayTracingShaders.RayHitGroupShader = MaterialResource.GetShader<FTrivialMaterialCHS>(VertexFactory->GetType());
 	}
 
 	FMeshPassProcessorRenderState PassDrawRenderState(Scene->UniformBuffers.ViewUniformBuffer, Scene->UniformBuffers.OpaqueBasePassUniformBuffer);
@@ -478,7 +549,10 @@ FRayTracingPipelineState* FDeferredShadingSceneRenderer::BindRayTracingMaterialP
 
 	Initializer.SetRayGenShaderTable(RayGenShaderTable);
 
-	const bool bEnableMaterials = GEnableRayTracingMaterials != 0;
+	const bool bMaterialsCompiled = GCompileRayTracingMaterialAHS || GCompileRayTracingMaterialCHS;
+	const bool bEnableMaterials = bMaterialsCompiled && GEnableRayTracingMaterials != 0;
+	static auto CVarEnableShadowMaterials = IConsoleManager::Get().FindConsoleVariable(TEXT("r.RayTracing.Shadows.EnableMaterials"));
+	const bool bEnableShadowMaterials = bMaterialsCompiled && (CVarEnableShadowMaterials ? CVarEnableShadowMaterials->GetInt() != 0 : true);
 
 	TArray<FRHIRayTracingShader*> RayTracingMaterialLibrary;
 
@@ -500,9 +574,6 @@ FRayTracingPipelineState* FDeferredShadingSceneRenderer::BindRayTracingMaterialP
 
 	FViewInfo& ReferenceView = Views[0];
 
-	static auto CVarEnableShadowMaterials = IConsoleManager::Get().FindConsoleVariable(TEXT("r.RayTracing.Shadows.EnableMaterials"));
-	bool bEnableShadowMaterials = CVarEnableShadowMaterials ? CVarEnableShadowMaterials->GetInt() != 0 : true;
-
 	const uint32 NumTotalMeshCommands = ReferenceView.VisibleRayTracingMeshCommands.Num();
 	const uint32 TargetCommandsPerTask = 4096; // Granularity chosen based on profiling Infiltrator scene to balance wall time speedup and total CPU thread time.
 	const uint32 NumTasks = FMath::Max(1u, FMath::DivideAndRoundUp(NumTotalMeshCommands, TargetCommandsPerTask));
