fix compile errors caused by merge

diff --git a/Engine/Source/Runtime/Renderer/Private/BasePassRendering.h b/Engine/Source/Runtime/Renderer/Private/BasePassRendering.h
--- a/Engine/Source/Runtime/Renderer/Private/BasePassRendering.h
+++ b/Engine/Source/Runtime/Renderer/Private/BasePassRendering.h
@@ -517,7 +517,7 @@ public:
 };
 
 //Alternative base pass PS for 128 bit canvas render targets that need to be set at shader compilation time.
-class F128BitRTBasePassPS : public TBasePassPS<TUniformLightMapPolicy<LMP_NO_LIGHTMAP>, false>
+class F128BitRTBasePassPS : public TBasePassPS<TUniformLightMapPolicy<LMP_NO_LIGHTMAP>>
 {
 	DECLARE_SHADER_TYPE(F128BitRTBasePassPS, MeshMaterial);
 public:
@@ -535,7 +535,7 @@ public:
 
 	/** Initialization constructor. */
 	F128BitRTBasePassPS(const ShaderMetaType::CompiledShaderInitializerType& Initializer): 
-		TBasePassPS<TUniformLightMapPolicy<LMP_NO_LIGHTMAP>, false>(Initializer)
+		TBasePassPS<TUniformLightMapPolicy<LMP_NO_LIGHTMAP>>(Initializer)
 	{}
 
 	/** Default constructor. */
