FGenerateReflectionRaysCS is not compiled when ray tracing is disabled

diff --git a/Engine/Source/Runtime/Renderer/Private/RayTracing/RayTracingDeferredReflections.cpp b/Engine/Source/Runtime/Renderer/Private/RayTracing/RayTracingDeferredReflections.cpp
--- a/Engine/Source/Runtime/Renderer/Private/RayTracing/RayTracingDeferredReflections.cpp
+++ b/Engine/Source/Runtime/Renderer/Private/RayTracing/RayTracingDeferredReflections.cpp
@@ -128,6 +128,11 @@ class FRayTracingDeferredReflectionsRGS : public FGlobalShader
 
 	static bool ShouldCompilePermutation(const FGlobalShaderPermutationParameters& Parameters)
 	{
+		if (!ShouldCompileRayTracingShadersForProject(Parameters.Platform))
+		{
+			return false;
+		}
+
 		FPermutationDomain PermutationVector(Parameters.PermutationId);
 		if (PermutationVector.Get<FDeferredMaterialMode>() == EDeferredMaterialMode::None)
 		{
