Forcing pre-exposure to always be enabled in the shaders.

diff --git a/Engine/Source/Runtime/Engine/Private/ShaderCompiler/ShaderCompiler.cpp b/Engine/Source/Runtime/Engine/Private/ShaderCompiler/ShaderCompiler.cpp
--- a/Engine/Source/Runtime/Engine/Private/ShaderCompiler/ShaderCompiler.cpp
+++ b/Engine/Source/Runtime/Engine/Private/ShaderCompiler/ShaderCompiler.cpp
@@ -53,8 +53,6 @@
 
 DEFINE_LOG_CATEGORY(LogShaderCompilers);
 
-RENDERCORE_API bool UsePreExposure(EShaderPlatform Platform);
-
 #if ENABLE_COOK_STATS
 namespace GlobalShaderCookStats
 {
@@ -3688,7 +3686,8 @@ void GlobalBeginCompileShader(
 	}
 
 	{
-		Input.Environment.SetDefine(TEXT("USE_PREEXPOSURE"), UsePreExposure((EShaderPlatform)Target.Platform) ? 1 : 0);
+		// PreExposure is now always enabled in the shaders.
+		Input.Environment.SetDefine(TEXT("USE_PREEXPOSURE"), 1);
 	}
 
 	{
diff --git a/Engine/Source/Runtime/RenderCore/Private/Shader.cpp b/Engine/Source/Runtime/RenderCore/Private/Shader.cpp
--- a/Engine/Source/Runtime/RenderCore/Private/Shader.cpp
+++ b/Engine/Source/Runtime/RenderCore/Private/Shader.cpp
@@ -42,12 +42,6 @@ IMPLEMENT_TYPE_LAYOUT(FShaderParameterMapInfo);
 IMPLEMENT_EXPORTED_INTRINSIC_TYPE_LAYOUT(TIndexedPtr<FShaderType>);
 IMPLEMENT_EXPORTED_INTRINSIC_TYPE_LAYOUT(TIndexedPtr<FVertexFactoryType>);
 
-RENDERCORE_API bool UsePreExposure(EShaderPlatform Platform)
-{
-	static const auto CVarUsePreExposure = IConsoleManager::Get().FindTConsoleVariableDataInt(TEXT("r.UsePreExposure"));
-	return CVarUsePreExposure->GetValueOnAnyThread() != 0;
-}
-
 static TAutoConsoleVariable<int32> CVarUsePipelines(
 	TEXT("r.ShaderPipelines"),
 	1,
@@ -1341,7 +1335,8 @@ void ShaderMapAppendKeyString(EShaderPlatform Platform, FString& KeyString)
 	}
 
 	{
-		KeyString += UsePreExposure(Platform) ? TEXT("_PreExp") : TEXT("");
+		// PreExposure is always used
+		KeyString += TEXT("_PreExp");
 	}
 
 	{
diff --git a/Engine/Source/Runtime/Renderer/Private/PostProcess/PostProcessEyeAdaptation.cpp b/Engine/Source/Runtime/Renderer/Private/PostProcess/PostProcessEyeAdaptation.cpp
--- a/Engine/Source/Runtime/Renderer/Private/PostProcess/PostProcessEyeAdaptation.cpp
+++ b/Engine/Source/Runtime/Renderer/Private/PostProcess/PostProcessEyeAdaptation.cpp
@@ -4,8 +4,6 @@
 #include "RHIGPUReadback.h"
 #include "Curves/CurveFloat.h"
 
-RENDERCORE_API bool UsePreExposure(EShaderPlatform Platform);
-
 bool IsMobileEyeAdaptationEnabled(const FViewInfo& View);
 
 namespace
@@ -75,6 +73,13 @@ static TAutoConsoleVariable<int32> CVarEnablePreExposureOnlyInTheEditor(
 const ERHIFeatureLevel::Type BasicEyeAdaptationMinFeatureLevel = ERHIFeatureLevel::SM5;
 }
 
+// Function is static because EyeAdaptation is the only system that should be checking this value.
+static bool UsePreExposureEnabled()
+{
+	static const auto CVarUsePreExposure = IConsoleManager::Get().FindTConsoleVariableDataInt(TEXT("r.UsePreExposure"));
+	return CVarUsePreExposure->GetValueOnAnyThread() != 0;
+}
+
 bool IsAutoExposureMethodSupported(ERHIFeatureLevel::Type FeatureLevel, EAutoExposureMethod AutoExposureMethodId)
 {
 	switch (AutoExposureMethodId)
@@ -895,7 +900,7 @@ void FSceneViewState::UpdatePreExposure(FViewInfo& View)
 	
 	if (bIsPreExposureRelevant && bEnableAutoExposure)
 	{
-		if (UsePreExposure(View.GetShaderPlatform()))
+		if (UsePreExposureEnabled())
 		{
 			const float PreExposureOverride = CVarEyeAdaptationPreExposureOverride.GetValueOnRenderThread();
 			const float LastExposure = View.GetLastEyeAdaptationExposure();
@@ -910,10 +915,16 @@ void FSceneViewState::UpdatePreExposure(FViewInfo& View)
 
 			bUpdateLastExposure = true;
 		}
-		// The exposure compensation curves require the scene average luminance
-		else if (View.FinalPostProcessSettings.AutoExposureBiasCurve)
+		else
 		{
-			bUpdateLastExposure = true;
+			if (View.FinalPostProcessSettings.AutoExposureBiasCurve)
+			{
+				// The exposure compensation curves require the scene average luminance
+				bUpdateLastExposure = true;
+			}
+
+			// Whe PreExposure is not used, we will override the PreExposure value to 1.0, which simulates PreExposure being off.
+			PreExposure = 1.0f;
 		}
 	}
 
