Change fallback texture to 'UInt' format, to match shader expectation

diff --git a/Engine/Source/Runtime/Renderer/Private/ScreenSpaceDenoise.cpp b/Engine/Source/Runtime/Renderer/Private/ScreenSpaceDenoise.cpp
--- a/Engine/Source/Runtime/Renderer/Private/ScreenSpaceDenoise.cpp
+++ b/Engine/Source/Runtime/Renderer/Private/ScreenSpaceDenoise.cpp
@@ -1858,7 +1858,7 @@ static void DenoiseSignalAtConstantPixelDensity(
 		if (CompressedMetadataLayout == ECompressedMetadataLayout::DepthAndViewNormal)
 		{
 			PassParameters->PrevCompressedMetadata[0] = RegisterExternalTextureWithFallback(
-				GraphBuilder, ViewInfoPooledRenderTargets.PrevCompressedDepthViewNormal, GSystemTextures.BlackDummy);
+				GraphBuilder, ViewInfoPooledRenderTargets.PrevCompressedDepthViewNormal, GSystemTextures.ZeroUIntDummy);
 		}
 
 		FScreenSpaceDenoiserHistory DummyPrevFrameHistory;
