Fix for Deferred Decals after base pass not being rendered when GTAO is enabled. This was due to them not being registered in the render graph as a dependancy when rendering the GTAO.
Also removed some commented out code

diff --git a/Engine/Source/Runtime/Renderer/Private/CompositionLighting/CompositionLighting.cpp b/Engine/Source/Runtime/Renderer/Private/CompositionLighting/CompositionLighting.cpp
--- a/Engine/Source/Runtime/Renderer/Private/CompositionLighting/CompositionLighting.cpp
+++ b/Engine/Source/Runtime/Renderer/Private/CompositionLighting/CompositionLighting.cpp
@@ -221,35 +221,14 @@ static FRenderingCompositeOutputRef AddPostProcessingGTAOAllPasses(FRHICommandLi
 		GRenderTargetPool.FindFreeElement(RHICmdList, Desc, SceneContext.ScreenSpaceGTAODepths, TEXT("ScreenSpaceGTAODepths"));
 	}
 
-	//if(CVarGTAOCombined.GetValueOnRenderThread() ==1 || (GTAOType== EGTAOType::ESplitNoNormalAsync) )
 	{
 		FRenderingCompositePass* AmbientOcclusionGTAO = Context.Graph.RegisterPass(new(FMemStack::Get()) FRCPassPostProcessAmbientOcclusion_GTAOHorizonSearchIntegrate(Context.View, DownsampleFactor, false, GTAOType));
 		AmbientOcclusionGTAO->SetInput(ePId_Input0, Context.SceneDepth);
 		AmbientOcclusionGTAO->SetInput(ePId_Input1, HZBInput);
 		FinalOutputPass = AmbientOcclusionGTAO;
+		// Add dependancy for Decals
+		AmbientOcclusionGTAO->AddDependency(Context.FinalOutput);
 	}
-/*
-	else
-	{
-		FPooledRenderTargetDesc Desc(FPooledRenderTargetDesc::Create2DDesc(HorizonBufferSize, PF_R8G8, FClearValueBinding::White, TexCreate_None, TexCreate_RenderTargetable, false));
-		if (SceneContext.GetCurrentFeatureLevel() >= ERHIFeatureLevel::SM5)
-		{
-			Desc.TargetableFlags |= TexCreate_UAV;
-		}
-		GRenderTargetPool.FindFreeElement(RHICmdList, Desc, SceneContext.ScreenSpaceGTAOHorizons, TEXT("ScreenSpaceGTAOHorizons"));
-		FRenderingCompositePass* AmbientOcclusionHorizonSearch = Context.Graph.RegisterPass(new(FMemStack::Get()) FRCPassPostProcessAmbientOcclusion_HorizonSearch(Context.View, DownsampleFactor, EGTAOType::ESplitUseNormalAsync));
-
-		AmbientOcclusionHorizonSearch->SetInput(ePId_Input0, Context.SceneDepth);
-		AmbientOcclusionHorizonSearch->SetInput(ePId_Input1, HZBInput);
-
-		FinalOutputPass = AmbientOcclusionHorizonSearch;
-
-		FRenderingCompositePass* AmbientOcclusionInnerIntegrate = Context.Graph.RegisterPass(new(FMemStack::Get()) FRCPassPostProcessAmbientOcclusion_InnerIntegrate(Context.View, DownsampleFactor, false));
-		AmbientOcclusionInnerIntegrate->SetInput(ePId_Input0, Context.SceneDepth);
-		AmbientOcclusionInnerIntegrate->SetInput(ePId_Input1, FinalOutputPass);
-		FinalOutputPass = AmbientOcclusionInnerIntegrate;
-	}
-*/
 
 	SceneContext.bScreenSpaceAOIsValid = true;
 
@@ -311,13 +290,14 @@ static FRenderingCompositeOutputRef AddPostProcessingGTAOPostAsync(FRHICommandLi
 		AmbientOcclusionInnerIntegrate->SetInput(ePId_Input0, Context.SceneDepth);
 		FinalOutputPass = AmbientOcclusionInnerIntegrate;
 
+		// Add dependancy for Decals
+		AmbientOcclusionInnerIntegrate->AddDependency(Context.FinalOutput);
+
 		// Add spatial Filter
 		if (CVarGTAOSpatialFilter.GetValueOnRenderThread() == 1)
 		{
 			FRenderingCompositePass* SpatialPass;
 			SpatialPass = Context.Graph.RegisterPass(new (FMemStack::Get()) FRCPassPostProcessAmbientOcclusion_GTAO_SpatialFilter(Context.View, DownsampleFactor, GTAOType));
-			// TODO . If the 
-
 			SpatialPass->SetInput(ePId_Input0, FinalOutputPass);
 			SpatialPass->SetInput(ePId_Input1, HZBInput);
 			FinalOutputPass = SpatialPass;
