Runtime virtual texture world height packing now depends on the volume extents.
Fixes precision issues raised by a number of licensees (for example UDN558460).

diff --git a/Engine/Shaders/Private/VirtualTextureCommon.ush b/Engine/Shaders/Private/VirtualTextureCommon.ush
--- a/Engine/Shaders/Private/VirtualTextureCommon.ush
+++ b/Engine/Shaders/Private/VirtualTextureCommon.ush
@@ -556,7 +556,7 @@ float3 VirtualTextureUnpackBaseColorYCoCg(in float4 PackedValue)
 /** Generic normal unpack funtion. */
 float3 VirtualTextureUnpackNormal(in float2 PackedXY, in float PackedSignZ)
 {
-	float2 NormalXY = PackedXY * 2.f - 1.f;
+	float2 NormalXY = PackedXY * (255.f / 127.f) - 1.f;
 	float SignZ = PackedSignZ * 2.f - 1.f;
 	float NormalZ = sqrt(saturate(1.0f - dot(NormalXY, NormalXY))) * SignZ;
 	return float3(NormalXY, NormalZ);
@@ -587,11 +587,9 @@ float3 VirtualTextureUnpackNormalBC5BC1(in float4 PackedValue0, in float4 Packed
 }
 
 /** Unpack 16 bit height (with some hardcoded range scale/bias). */
-float VirtualTextureUnpackHeightR16(in float4 PackedValue)
+float VirtualTextureUnpackHeight(in float4 PackedValue, in float2 UnpackHeightScaleBias)
 {
-	float Height = PackedValue.r * 65535.f - 32768.f;
-	Height *= 2.f;
-	return Height;
+	return PackedValue.r * UnpackHeightScaleBias.x + UnpackHeightScaleBias.y;
 }
 
 #endif //VirtualTextureCommon_ush__
diff --git a/Engine/Shaders/Private/VirtualTextureMaterial.usf b/Engine/Shaders/Private/VirtualTextureMaterial.usf
--- a/Engine/Shaders/Private/VirtualTextureMaterial.usf
+++ b/Engine/Shaders/Private/VirtualTextureMaterial.usf
@@ -38,16 +38,15 @@ void MainVS(
 }
 
 // Prepare for VirtualTextureUnpackNormal() in VirtualTextureCommon.h
-float3 PackNormal( float3 N )
+float3 PackNormal( in float3 N )
 {
-	return normalize(N) * 0.5f + 0.5f;
+	return normalize(N) * (127.f / 255.f) + (127.f / 255.f);
 }
 
-// Prepare for VirtualTextureUnpackHeightR16() in VirtualTextureCommon.ush
-float PackWorldHeight( float Height )
+// Prepare for VirtualTextureUnpackHeight() in VirtualTextureCommon.ush
+float PackWorldHeight( in float Height, in float2 PackHeightScaleBias )
 {
-	Height *= 0.5f;
-	return saturate((Height + 32768.f) / 65535.f);
+	return saturate(Height * PackHeightScaleBias.x + PackHeightScaleBias.y);
 }
 
 void FPixelShaderInOut_MainPS(
@@ -87,7 +86,7 @@ void FPixelShaderInOut_MainPS(
 #endif
 
 	// Apply debug shading
-	BaseColor = lerp(BaseColor, half3(In.SvPosition.xy * ResolvedView.ViewSizeAndInvSize.zw, 0), ResolvedView.VirtualTextureParams.y * 0.25);
+	BaseColor = lerp(BaseColor, half3(In.SvPosition.xy * ResolvedView.ViewSizeAndInvSize.zw, 0), ResolvedView.RuntimeVirtualTextureDebugParams.x * 0.25);
 
 #if defined(OUT_BASECOLOR)
 	Out.MRT[0] = float4(BaseColor, 1.f) * Opacity;
@@ -97,7 +96,8 @@ void FPixelShaderInOut_MainPS(
 	Out.MRT[1] = float4(PackedNormal.xy, Mask, 1.f) * Opacity;
 	Out.MRT[2] = float4(Specular, Roughness, PackedNormal.z, 1.f) * Opacity;
 #elif defined(OUT_WORLDHEIGHT)
-	Out.MRT[0] = float4(PackWorldHeight(WorldHeight), 0, 0, 1);
+	float PackedHeight = PackWorldHeight(WorldHeight, ResolvedView.RuntimeVirtualTexturePackHeight);
+	Out.MRT[0] = float4(PackedHeight, 0, 0, 1);
 #endif
 }
 
diff --git a/Engine/Source/Editor/VirtualTexturingEditor/Private/RuntimeVirtualTextureBuild.cpp b/Engine/Source/Editor/VirtualTexturingEditor/Private/RuntimeVirtualTextureBuild.cpp
--- a/Engine/Source/Editor/VirtualTexturingEditor/Private/RuntimeVirtualTextureBuild.cpp
+++ b/Engine/Source/Editor/VirtualTexturingEditor/Private/RuntimeVirtualTextureBuild.cpp
@@ -126,6 +126,7 @@ namespace RuntimeVirtualTexture
 		FSceneInterface* Scene = InComponent->GetScene();
 		const uint32 VirtualTextureSceneIndex = RuntimeVirtualTexture::GetRuntimeVirtualTextureSceneIndex_GameThread(InComponent);
 		const FTransform Transform = InComponent->GetVirtualTextureTransform();
+		const FBox Bounds = InComponent->Bounds.GetBox();
 
 		FVTProducerDescription VTDesc;
 		RuntimeVirtualTexture->GetProducerDescription(VTDesc, Transform);
@@ -176,7 +177,7 @@ namespace RuntimeVirtualTexture
 					Scene, VirtualTextureSceneIndex, 
 					&RenderTileResources,
 					MaterialType, NumLayers,
-					Transform, UVRange,
+					Transform, Bounds, UVRange,
 					RenderLevel, MaxLevel, 
 					TileX, TileY,
 					TileSize, ImageSizeX, ImageSizeY, 
@@ -197,6 +198,7 @@ namespace RuntimeVirtualTexture
 					Desc.Scene = Scene->GetRenderScene();
 					Desc.RuntimeVirtualTextureMask = 1 << VirtualTextureSceneIndex;
 					Desc.UVToWorld = Transform;
+					Desc.WorldBounds = Bounds;
 					Desc.MaterialType = MaterialType;
 					Desc.MaxLevel = MaxLevel;
 					Desc.bClearTextures = true;
diff --git a/Engine/Source/Editor/VirtualTexturingEditor/Private/RuntimeVirtualTextureThumbnailRenderer.cpp b/Engine/Source/Editor/VirtualTexturingEditor/Private/RuntimeVirtualTextureThumbnailRenderer.cpp
--- a/Engine/Source/Editor/VirtualTexturingEditor/Private/RuntimeVirtualTextureThumbnailRenderer.cpp
+++ b/Engine/Source/Editor/VirtualTexturingEditor/Private/RuntimeVirtualTextureThumbnailRenderer.cpp
@@ -69,6 +69,7 @@ void URuntimeVirtualTextureThumbnailRenderer::Draw(UObject* Object, int32 X, int
 
 	const FBox2D DestBox = FBox2D(FVector2D(X, Y), FVector2D(Width, Height));
 	const FTransform Transform = RuntimeVirtualTextureComponent->GetVirtualTextureTransform();
+	const FBox Bounds = RuntimeVirtualTextureComponent->Bounds.GetBox();
 	const uint32 VirtualTextureSceneIndex = RuntimeVirtualTexture::GetRuntimeVirtualTextureSceneIndex_GameThread(RuntimeVirtualTextureComponent);
 	const ERuntimeVirtualTextureMaterialType MaterialType = RuntimeVirtualTexture->GetMaterialType();
 
@@ -77,7 +78,7 @@ void URuntimeVirtualTextureThumbnailRenderer::Draw(UObject* Object, int32 X, int
 	const int32 MaxLevel = (int32)FMath::CeilLogTwo(FMath::Max(VTDesc.BlockWidthInTiles, VTDesc.BlockHeightInTiles));
 
 	ENQUEUE_RENDER_COMMAND(BakeStreamingTextureTileCommand)(
-		[Scene, VirtualTextureSceneIndex, MaterialType, RenderTarget, DestBox, Transform, MaxLevel](FRHICommandListImmediate& RHICmdList)
+		[Scene, VirtualTextureSceneIndex, MaterialType, RenderTarget, DestBox, Transform, Bounds, MaxLevel](FRHICommandListImmediate& RHICmdList)
 	{
 		FMaterialRenderProxy::UpdateDeferredCachedUniformExpressions();
 
@@ -85,6 +86,7 @@ void URuntimeVirtualTextureThumbnailRenderer::Draw(UObject* Object, int32 X, int
 		Desc.Scene = Scene->GetRenderScene();
 		Desc.RuntimeVirtualTextureMask = 1 << VirtualTextureSceneIndex;
 		Desc.UVToWorld = Transform;
+		Desc.WorldBounds = Bounds;
 		Desc.MaterialType = MaterialType;
 		Desc.MaxLevel = MaxLevel;
 		Desc.bClearTextures = true;
diff --git a/Engine/Source/Runtime/Engine/Classes/Components/RuntimeVirtualTextureComponent.h b/Engine/Source/Runtime/Engine/Classes/Components/RuntimeVirtualTextureComponent.h
--- a/Engine/Source/Runtime/Engine/Classes/Components/RuntimeVirtualTextureComponent.h
+++ b/Engine/Source/Runtime/Engine/Classes/Components/RuntimeVirtualTextureComponent.h
@@ -52,12 +52,16 @@ public:
 
 protected:
 	//~ Begin UActorComponent Interface
-	virtual bool IsVisible() const override;
 	virtual void CreateRenderState_Concurrent(FRegisterComponentContext* Context) override;
 	virtual void SendRenderTransform_Concurrent() override;
 	virtual void DestroyRenderState_Concurrent() override;
 	//~ End UActorComponent Interface
 
+	//~ Begin USceneComponent Interface
+	virtual bool IsVisible() const override;
+	virtual FBoxSphereBounds CalcBounds(const FTransform& LocalToWorld) const override;
+	//~ End USceneComponent Interface
+
 public:
 	/** Scene proxy object. Managed by the scene but stored here. */
 	class FRuntimeVirtualTextureSceneProxy* SceneProxy;
diff --git a/Engine/Source/Runtime/Engine/Classes/VT/RuntimeVirtualTexture.h b/Engine/Source/Runtime/Engine/Classes/VT/RuntimeVirtualTexture.h
--- a/Engine/Source/Runtime/Engine/Classes/VT/RuntimeVirtualTexture.h
+++ b/Engine/Source/Runtime/Engine/Classes/VT/RuntimeVirtualTexture.h
@@ -130,7 +130,7 @@ public:
 	bool GetClearTextures() const { return bClearTextures; }
 
 	/** (Re)Initialize this object. Call this whenever we modify the producer or transform. */
-	void Initialize(IVirtualTexture* InProducer, FTransform const& VolumeToWorld);
+	void Initialize(IVirtualTexture* InProducer, FTransform const& VolumeToWorld, FBox const& WorldBounds);
 
 	/** Release the resources for this object This will need to be called if our producer becomes stale and we aren't doing a full reinit with a new producer. */
 	void Release();
@@ -176,4 +176,6 @@ private:
 
 	/** Material uniform parameters to support transform from world to UV coordinates. */
 	FVector4 WorldToUVTransformParameters[3];
+	/** Material uniform parameter used to pack world height. */
+	FVector4 WorldHeightUnpackParameter;
 };
diff --git a/Engine/Source/Runtime/Engine/Private/Components/RuntimeVirtualTextureComponent.cpp b/Engine/Source/Runtime/Engine/Private/Components/RuntimeVirtualTextureComponent.cpp
--- a/Engine/Source/Runtime/Engine/Private/Components/RuntimeVirtualTextureComponent.cpp
+++ b/Engine/Source/Runtime/Engine/Private/Components/RuntimeVirtualTextureComponent.cpp
@@ -49,6 +49,12 @@ void URuntimeVirtualTextureComponent::DestroyRenderState_Concurrent()
 	Super::DestroyRenderState_Concurrent();
 }
 
+FBoxSphereBounds URuntimeVirtualTextureComponent::CalcBounds(const FTransform& LocalToWorld) const
+{
+	// Bounds are based on the unit box centered on the origin
+	return FBoxSphereBounds(FBox(FVector(-.5f, -.5f, -1.f), FVector(.5f, .5f, 1.f))).TransformBy(LocalToWorld);
+}
+
 FTransform URuntimeVirtualTextureComponent::GetVirtualTextureTransform() const
 {
 	// Transform is based on bottom left of the URuntimeVirtualTextureComponent unit box (which is centered on the origin)
diff --git a/Engine/Source/Runtime/Engine/Private/Materials/HLSLMaterialTranslator.cpp b/Engine/Source/Runtime/Engine/Private/Materials/HLSLMaterialTranslator.cpp
--- a/Engine/Source/Runtime/Engine/Private/Materials/HLSLMaterialTranslator.cpp
+++ b/Engine/Source/Runtime/Engine/Private/Materials/HLSLMaterialTranslator.cpp
@@ -2991,8 +2991,8 @@ int32 FHLSLMaterialTranslator::ViewProperty(EMaterialExposedViewProperty Propert
 		{MEVP_TemporalSampleCount, MCT_Float1, TEXT("View.TemporalAAParams.y"), nullptr},
 		{MEVP_TemporalSampleIndex, MCT_Float1, TEXT("View.TemporalAAParams.x"), nullptr},
 		{MEVP_TemporalSampleOffset, MCT_Float2, TEXT("View.TemporalAAParams.zw"), nullptr},
-		{MEVP_RuntimeVirtualTextureOutputLevel, MCT_Float1, TEXT("View.VirtualTextureParams.x"), nullptr},
-		{MEVP_RuntimeVirtualTextureOutputDerivative, MCT_Float2, TEXT("View.VirtualTextureParams.zw"), nullptr},
+		{MEVP_RuntimeVirtualTextureOutputLevel, MCT_Float1, TEXT("View.RuntimeVirtualTextureMipLevel.x"), nullptr},
+		{MEVP_RuntimeVirtualTextureOutputDerivative, MCT_Float2, TEXT("View.RuntimeVirtualTextureMipLevel.zw"), nullptr},
 		{MEVP_PreExposure, MCT_Float1, TEXT("View.PreExposure.x"), TEXT("View.OneOverPreExposure.x")},
 	};
 	static_assert((sizeof(ViewPropertyMetaArray) / sizeof(ViewPropertyMetaArray[0])) == MEVP_MAX, "incoherency between EMaterialExposedViewProperty and ViewPropertyMetaArray");
@@ -5004,7 +5004,7 @@ int32 FHLSLMaterialTranslator::VirtualTextureWorldToUV(int32 WorldPositionIndex,
 	return AddInlinedCodeChunk(MCT_Float2, *SampleCode, *GetParameterCode(WorldPositionIndex), *GetParameterCode(P0), *GetParameterCode(P1), *GetParameterCode(P2));
 }
 
-int32 FHLSLMaterialTranslator::VirtualTextureUnpack(int32 CodeIndex0, int32 CodeIndex1, int32 CodeIndex2, EVirtualTextureUnpackType UnpackType)
+int32 FHLSLMaterialTranslator::VirtualTextureUnpack(int32 CodeIndex0, int32 CodeIndex1, int32 CodeIndex2, int32 P0, EVirtualTextureUnpackType UnpackType)
 {
 	if (UnpackType == EVirtualTextureUnpackType::BaseColorYCoCg)
 	{
@@ -5033,8 +5033,8 @@ int32 FHLSLMaterialTranslator::VirtualTextureUnpack(int32 CodeIndex0, int32 Code
 	}
 	else if (UnpackType == EVirtualTextureUnpackType::HeightR16)
 	{
-		FString	SampleCode(TEXT("VirtualTextureUnpackHeightR16(%s)"));
-		return CodeIndex0 == INDEX_NONE ? INDEX_NONE : AddCodeChunk(MCT_Float, *SampleCode, *GetParameterCode(CodeIndex0));
+		FString	SampleCode(TEXT("VirtualTextureUnpackHeight(%s, %s)"));
+		return CodeIndex0 == INDEX_NONE ? INDEX_NONE : AddCodeChunk(MCT_Float, *SampleCode, *GetParameterCode(CodeIndex0), *GetParameterCode(P0));
 	}
 
 	return CodeIndex0;
diff --git a/Engine/Source/Runtime/Engine/Private/Materials/HLSLMaterialTranslator.h b/Engine/Source/Runtime/Engine/Private/Materials/HLSLMaterialTranslator.h
--- a/Engine/Source/Runtime/Engine/Private/Materials/HLSLMaterialTranslator.h
+++ b/Engine/Source/Runtime/Engine/Private/Materials/HLSLMaterialTranslator.h
@@ -609,7 +609,7 @@ protected:
 	virtual int32 VirtualTextureUniform(int32 TextureIndex, int32 VectorIndex) override;
 	virtual int32 VirtualTextureUniform(FName ParameterName, int32 TextureIndex, int32 VectorIndex) override;
 	virtual int32 VirtualTextureWorldToUV(int32 WorldPositionIndex, int32 P0, int32 P1, int32 P2) override;
-	virtual int32 VirtualTextureUnpack(int32 CodeIndex0, int32 CodeIndex1, int32 CodeIndex2, EVirtualTextureUnpackType UnpackType) override;
+	virtual int32 VirtualTextureUnpack(int32 CodeIndex0, int32 CodeIndex1, int32 CodeIndex2, int32 P0, EVirtualTextureUnpackType UnpackType) override;
 
 	virtual int32 ExternalTexture(const FGuid& ExternalTextureGuid) override;
 	virtual int32 ExternalTexture(UTexture* InTexture, int32& TextureReferenceIndex) override;
diff --git a/Engine/Source/Runtime/Engine/Private/Materials/MaterialExpressions.cpp b/Engine/Source/Runtime/Engine/Private/Materials/MaterialExpressions.cpp
--- a/Engine/Source/Runtime/Engine/Private/Materials/MaterialExpressions.cpp
+++ b/Engine/Source/Runtime/Engine/Private/Materials/MaterialExpressions.cpp
@@ -2301,6 +2301,20 @@ int32 UMaterialExpressionRuntimeVirtualTextureSample::Compile(class FMaterialCom
 		}
 	}
 
+	// Compile the runtime virtual texture uniforms
+	int32 Uniforms[ERuntimeVirtualTextureShaderUniform_Count];
+	for (int32 UniformIndex = 0; UniformIndex < ERuntimeVirtualTextureShaderUniform_Count; ++UniformIndex)
+	{
+		if (bIsParameter)
+		{
+			Uniforms[UniformIndex] = Compiler->VirtualTextureUniform(GetParameterName(), TextureReferenceIndex[0], UniformIndex);
+		}
+		else
+		{
+			Uniforms[UniformIndex] = Compiler->VirtualTextureUniform(TextureReferenceIndex[0], UniformIndex);
+		}
+	}
+
 	// Compile the coordinates
 	// We use the virtual texture world space transform by default
 	int32 CoordinateIndex = INDEX_NONE;
@@ -2326,20 +2340,9 @@ int32 UMaterialExpressionRuntimeVirtualTextureSample::Compile(class FMaterialCom
 			WorldPositionIndex = Compiler->WorldPosition(WPT_Default);
 		}
 		
-		int32 P0, P1, P2;
-		if (bIsParameter)
-		{
-			P0 = Compiler->VirtualTextureUniform(GetParameterName(), TextureReferenceIndex[0], 0);
-			P1 = Compiler->VirtualTextureUniform(GetParameterName(), TextureReferenceIndex[0], 1);
-			P2 = Compiler->VirtualTextureUniform(GetParameterName(), TextureReferenceIndex[0], 2);
-		}
-		else
-		{
-			P0 = Compiler->VirtualTextureUniform(TextureReferenceIndex[0], 0);
-			P1 = Compiler->VirtualTextureUniform(TextureReferenceIndex[0], 1);
-			P2 = Compiler->VirtualTextureUniform(TextureReferenceIndex[0], 2);
-		}
-
+		const int32 P0 = Uniforms[ERuntimeVirtualTextureShaderUniform_WorldToUVTransform0];
+		const int32 P1 = Uniforms[ERuntimeVirtualTextureShaderUniform_WorldToUVTransform1];
+		const int32 P2 = Uniforms[ERuntimeVirtualTextureShaderUniform_WorldToUVTransform2];
 		CoordinateIndex = Compiler->VirtualTextureWorldToUV(WorldPositionIndex, P0, P1, P2);
 	}
 	
@@ -2376,7 +2379,8 @@ int32 UMaterialExpressionRuntimeVirtualTextureSample::Compile(class FMaterialCom
 	int32 UnpackCodeIndex = INDEX_NONE;
 	if (UnpackType != EVirtualTextureUnpackType::None)
 	{
-		UnpackCodeIndex = Compiler->VirtualTextureUnpack(SampleCodeIndex[0], SampleCodeIndex[1], SampleCodeIndex[2], UnpackType);
+		int32 P0 = Uniforms[ERuntimeVirtualTextureShaderUniform_WorldHeightUnpack];
+		UnpackCodeIndex = Compiler->VirtualTextureUnpack(SampleCodeIndex[0], SampleCodeIndex[1], SampleCodeIndex[2], P0, UnpackType);
 	}
 	else
 	{
diff --git a/Engine/Source/Runtime/Engine/Private/VT/RuntimeVirtualTexture.cpp b/Engine/Source/Runtime/Engine/Private/VT/RuntimeVirtualTexture.cpp
--- a/Engine/Source/Runtime/Engine/Private/VT/RuntimeVirtualTexture.cpp
+++ b/Engine/Source/Runtime/Engine/Private/VT/RuntimeVirtualTexture.cpp
@@ -399,18 +399,30 @@ IAllocatedVirtualTexture* URuntimeVirtualTexture::GetAllocatedVirtualTexture() c
 
 FVector4 URuntimeVirtualTexture::GetUniformParameter(int32 Index) const
 {
-	check(Index >= 0 && Index < sizeof(WorldToUVTransformParameters)/sizeof(WorldToUVTransformParameters[0]));
+	switch (Index)
+	{
+	case ERuntimeVirtualTextureShaderUniform_WorldToUVTransform0: return WorldToUVTransformParameters[0];
+	case ERuntimeVirtualTextureShaderUniform_WorldToUVTransform1: return WorldToUVTransformParameters[1];
+	case ERuntimeVirtualTextureShaderUniform_WorldToUVTransform2: return WorldToUVTransformParameters[2];
+	case ERuntimeVirtualTextureShaderUniform_WorldHeightUnpack: return WorldHeightUnpackParameter;
+	default:
+		break;
+	}
 	
-	return WorldToUVTransformParameters[Index];
+	check(0);
+	return FVector4(ForceInitToZero);
 }
 
-void URuntimeVirtualTexture::Initialize(IVirtualTexture* InProducer, FTransform const& VolumeToWorld)
+void URuntimeVirtualTexture::Initialize(IVirtualTexture* InProducer, FTransform const& VolumeToWorld, FBox const& WorldBounds)
 {
 	//todo[vt]: possible issues with precision in large worlds here it might be better to calculate/upload camera space relative transform per frame?
 	WorldToUVTransformParameters[0] = VolumeToWorld.GetTranslation();
 	WorldToUVTransformParameters[1] = VolumeToWorld.GetUnitAxis(EAxis::X) * 1.f / VolumeToWorld.GetScale3D().X;
 	WorldToUVTransformParameters[2] = VolumeToWorld.GetUnitAxis(EAxis::Y) * 1.f / VolumeToWorld.GetScale3D().Y;
 
+	const float HeightRange = FMath::Max(WorldBounds.Max.Z - WorldBounds.Min.Z, 1.f);
+	WorldHeightUnpackParameter = FVector4(HeightRange, WorldBounds.Min.Z, 0.f, 0.f);
+
 	InitResource(InProducer, VolumeToWorld);
 }
 
diff --git a/Engine/Source/Runtime/Engine/Public/MaterialCompiler.h b/Engine/Source/Runtime/Engine/Public/MaterialCompiler.h
--- a/Engine/Source/Runtime/Engine/Public/MaterialCompiler.h
+++ b/Engine/Source/Runtime/Engine/Public/MaterialCompiler.h
@@ -209,7 +209,7 @@ public:
 	virtual int32 VirtualTextureUniform(int32 TextureIndex, int32 VectorIndex) = 0;
 	virtual int32 VirtualTextureUniform(FName ParameterName, int32 TextureIndex, int32 VectorIndex) = 0;
 	virtual int32 VirtualTextureWorldToUV(int32 WorldPositionIndex, int32 P0, int32 P1, int32 P2) = 0;
-	virtual int32 VirtualTextureUnpack(int32 CodeIndex0, int32 CodeIndex1, int32 CodeIndex2, EVirtualTextureUnpackType UnpackType) = 0;
+	virtual int32 VirtualTextureUnpack(int32 CodeIndex0, int32 CodeIndex1, int32 CodeIndex2, int32 P0, EVirtualTextureUnpackType UnpackType) = 0;
 
 	virtual int32 ExternalTexture(const FGuid& ExternalTextureGuid) = 0;
 	virtual int32 ExternalTexture(UTexture* InTexture, int32& TextureReferenceIndex) = 0;
@@ -505,7 +505,7 @@ public:
 	virtual int32 VirtualTextureUniform(int32 TextureIndex, int32 VectorIndex) override { return Compiler->VirtualTextureUniform(TextureIndex, VectorIndex); }
 	virtual int32 VirtualTextureUniform(FName ParameterName, int32 TextureIndex, int32 VectorIndex) override { return Compiler->VirtualTextureUniform(ParameterName, TextureIndex, VectorIndex); }
 	virtual int32 VirtualTextureWorldToUV(int32 WorldPositionIndex, int32 P0, int32 P1, int32 P2) override { return Compiler->VirtualTextureWorldToUV(WorldPositionIndex, P0, P1, P2); }
-	virtual int32 VirtualTextureUnpack(int32 CodeIndex0, int32 CodeIndex1, int32 CodeIndex2, EVirtualTextureUnpackType UnpackType) override { return Compiler->VirtualTextureUnpack(CodeIndex0, CodeIndex1, CodeIndex2, UnpackType); }
+	virtual int32 VirtualTextureUnpack(int32 CodeIndex0, int32 CodeIndex1, int32 CodeIndex2, int32 P0, EVirtualTextureUnpackType UnpackType) override { return Compiler->VirtualTextureUnpack(CodeIndex0, CodeIndex1, CodeIndex2, P0, UnpackType); }
 
 	virtual int32 ExternalTexture(const FGuid& ExternalTextureGuid) override { return Compiler->ExternalTexture(ExternalTextureGuid); }
 	virtual int32 ExternalTexture(UTexture* InTexture, int32& TextureReferenceIndex) override { return Compiler->ExternalTexture(InTexture, TextureReferenceIndex); }
diff --git a/Engine/Source/Runtime/Engine/Public/SceneView.h b/Engine/Source/Runtime/Engine/Public/SceneView.h
--- a/Engine/Source/Runtime/Engine/Public/SceneView.h
+++ b/Engine/Source/Runtime/Engine/Public/SceneView.h
@@ -724,9 +724,11 @@ enum ETranslucencyVolumeCascade
 	VIEW_UNIFORM_BUFFER_MEMBER(float, IndirectLightingCacheShowFlag) \
 	VIEW_UNIFORM_BUFFER_MEMBER(float, EyeToPixelSpreadAngle) \
 	VIEW_UNIFORM_BUFFER_MEMBER(FMatrix, WorldToVirtualTexture) \
-	VIEW_UNIFORM_BUFFER_MEMBER(FVector4, VirtualTextureParams) \
 	VIEW_UNIFORM_BUFFER_MEMBER_ARRAY(FVector4, XRPassthroughCameraUVs, [2]) \
 	VIEW_UNIFORM_BUFFER_MEMBER(uint32, VirtualTextureFeedbackStride) \
+	VIEW_UNIFORM_BUFFER_MEMBER(FVector4, RuntimeVirtualTextureMipLevel) \
+	VIEW_UNIFORM_BUFFER_MEMBER(FVector2D, RuntimeVirtualTexturePackHeight) \
+	VIEW_UNIFORM_BUFFER_MEMBER(FVector4, RuntimeVirtualTextureDebugParams) \
 	VIEW_UNIFORM_BUFFER_MEMBER(int32, FarShadowStaticMeshLODBias) \
 	VIEW_UNIFORM_BUFFER_MEMBER(float, MinRoughness) \
 	VIEW_UNIFORM_BUFFER_MEMBER(FVector4, HairRenderInfo) \
diff --git a/Engine/Source/Runtime/Engine/Public/VT/RuntimeVirtualTextureEnum.h b/Engine/Source/Runtime/Engine/Public/VT/RuntimeVirtualTextureEnum.h
--- a/Engine/Source/Runtime/Engine/Public/VT/RuntimeVirtualTextureEnum.h
+++ b/Engine/Source/Runtime/Engine/Public/VT/RuntimeVirtualTextureEnum.h
@@ -59,6 +59,16 @@ enum class ERuntimeVirtualTextureMainPassType : uint8
 	Always UMETA(DisplayName = "Virtual Texture AND Main Pass"),
 };
 
+/** Enumeration of runtime virtual texture shader uniforms. */
+enum ERuntimeVirtualTextureShaderUniform
+{
+	ERuntimeVirtualTextureShaderUniform_WorldToUVTransform0,
+	ERuntimeVirtualTextureShaderUniform_WorldToUVTransform1,
+	ERuntimeVirtualTextureShaderUniform_WorldToUVTransform2,
+	ERuntimeVirtualTextureShaderUniform_WorldHeightUnpack,
+	ERuntimeVirtualTextureShaderUniform_Count,
+};
+
 /** Enumeration of runtime virtual texture debug modes. */
 enum class ERuntimeVirtualTextureDebugType
 {
diff --git a/Engine/Source/Runtime/Renderer/Private/SceneRendering.cpp b/Engine/Source/Runtime/Renderer/Private/SceneRendering.cpp
--- a/Engine/Source/Runtime/Renderer/Private/SceneRendering.cpp
+++ b/Engine/Source/Runtime/Renderer/Private/SceneRendering.cpp
@@ -1625,8 +1625,10 @@ void FViewInfo::SetupUniformBufferParameters(
 
 	ViewUniformShaderParameters.PreIntegratedBRDF = GEngine->PreIntegratedSkinBRDFTexture->Resource->TextureRHI;
 
-	ViewUniformShaderParameters.VirtualTextureParams = FVector4(ForceInitToZero);
 	ViewUniformShaderParameters.VirtualTextureFeedbackStride = SceneContext.VirtualTextureFeedback.GetFeedbackStride();
+	ViewUniformShaderParameters.RuntimeVirtualTextureMipLevel = FVector4(ForceInitToZero);
+	ViewUniformShaderParameters.RuntimeVirtualTexturePackHeight = FVector2D(ForceInitToZero);
+	ViewUniformShaderParameters.RuntimeVirtualTextureDebugParams = FVector4(ForceInitToZero);
 	
 	if (UseGPUScene(GMaxRHIShaderPlatform, RHIFeatureLevel))
 	{
diff --git a/Engine/Source/Runtime/Renderer/Private/VT/RuntimeVirtualTextureProducer.cpp b/Engine/Source/Runtime/Renderer/Private/VT/RuntimeVirtualTextureProducer.cpp
--- a/Engine/Source/Runtime/Renderer/Private/VT/RuntimeVirtualTextureProducer.cpp
+++ b/Engine/Source/Runtime/Renderer/Private/VT/RuntimeVirtualTextureProducer.cpp
@@ -8,7 +8,14 @@
 #include "VT/RuntimeVirtualTextureSceneProxy.h"
 
 
-FRuntimeVirtualTextureFinalizer::FRuntimeVirtualTextureFinalizer(FVTProducerDescription const& InDesc, uint32 InProducerId, ERuntimeVirtualTextureMaterialType InMaterialType, bool InClearTextures, FSceneInterface* InScene, FTransform const& InUVToWorld)
+FRuntimeVirtualTextureFinalizer::FRuntimeVirtualTextureFinalizer(
+	FVTProducerDescription const& InDesc, 
+	uint32 InProducerId, 
+	ERuntimeVirtualTextureMaterialType InMaterialType, 
+	bool InClearTextures, 
+	FSceneInterface* InScene, 
+	FTransform const& InUVToWorld,
+	FBox const& InWorldBounds)
 	: Desc(InDesc)
 	, ProducerId(InProducerId)
 	, RuntimeVirtualTextureMask(0)
@@ -16,6 +23,7 @@ FRuntimeVirtualTextureFinalizer::FRuntimeVirtualTextureFinalizer(FVTProducerDesc
 	, bClearTextures(InClearTextures)
 	, Scene(InScene)
 	, UVToWorld(InUVToWorld)
+	, WorldBounds(InWorldBounds)
 {
 }
 
@@ -59,6 +67,7 @@ void FRuntimeVirtualTextureFinalizer::Finalize(FRHICommandListImmediate& RHICmdL
 	RenderPageBatchDesc.Scene = Scene->GetRenderScene();
 	RenderPageBatchDesc.RuntimeVirtualTextureMask = RuntimeVirtualTextureMask;
 	RenderPageBatchDesc.UVToWorld = UVToWorld;
+	RenderPageBatchDesc.WorldBounds = WorldBounds;
 	RenderPageBatchDesc.MaterialType = MaterialType;
 	RenderPageBatchDesc.MaxLevel = Desc.MaxLevel;
 	RenderPageBatchDesc.bClearTextures = bClearTextures;
@@ -129,8 +138,15 @@ void FRuntimeVirtualTextureFinalizer::Finalize(FRHICommandListImmediate& RHICmdL
 	Tiles.SetNumUnsafeInternal(0);
 }
 
-FRuntimeVirtualTextureProducer::FRuntimeVirtualTextureProducer(FVTProducerDescription const& InDesc, uint32 InProducerId, ERuntimeVirtualTextureMaterialType InMaterialType, bool InClearTextures, FSceneInterface* InScene, FTransform const& InUVToWorld)
-	: Finalizer(InDesc, InProducerId, InMaterialType, InClearTextures, InScene, InUVToWorld)
+FRuntimeVirtualTextureProducer::FRuntimeVirtualTextureProducer(
+	FVTProducerDescription const& InDesc, 
+	uint32 InProducerId, 
+	ERuntimeVirtualTextureMaterialType InMaterialType, 
+	bool InClearTextures, 
+	FSceneInterface* InScene, 
+	FTransform const& InUVToWorld,
+	FBox const& InWorldBounds)
+	: Finalizer(InDesc, InProducerId, InMaterialType, InClearTextures, InScene, InUVToWorld, InWorldBounds)
 {
 }
 
diff --git a/Engine/Source/Runtime/Renderer/Private/VT/RuntimeVirtualTextureProducer.h b/Engine/Source/Runtime/Renderer/Private/VT/RuntimeVirtualTextureProducer.h
--- a/Engine/Source/Runtime/Renderer/Private/VT/RuntimeVirtualTextureProducer.h
+++ b/Engine/Source/Runtime/Renderer/Private/VT/RuntimeVirtualTextureProducer.h
@@ -15,7 +15,15 @@ class FSceneInterface;
 class FRuntimeVirtualTextureFinalizer : public IVirtualTextureFinalizer
 {
 public:
-	FRuntimeVirtualTextureFinalizer(FVTProducerDescription const& InDesc, uint32 InProducerId, ERuntimeVirtualTextureMaterialType InMaterialType, bool InClearTextures, FSceneInterface* InScene, FTransform const& InUVToWorld);
+	FRuntimeVirtualTextureFinalizer(
+		FVTProducerDescription const& InDesc, 
+		uint32 InProducerId, 
+		ERuntimeVirtualTextureMaterialType InMaterialType, 
+		bool InClearTextures, 
+		FSceneInterface* InScene, 
+		FTransform const& InUVToWorld,
+		FBox const& InWorldBounds);
+
 	virtual ~FRuntimeVirtualTextureFinalizer() {}
 
 	/** A description for a single tile to render. */
@@ -54,6 +62,8 @@ private:
 	FSceneInterface* Scene;
 	/** Transform from UV space to world space. */
 	FTransform UVToWorld;
+	/** Bounds of runtime virtual texture volume in world space. */
+	FBox WorldBounds;
 	/** Array of tiles in the queue to finalize. */
 	TArray<FTileEntry> Tiles;
 };
@@ -62,7 +72,15 @@ private:
 class FRuntimeVirtualTextureProducer : public IVirtualTexture
 {
 public:
-	RENDERER_API FRuntimeVirtualTextureProducer(FVTProducerDescription const& InDesc, uint32 InProducerId, ERuntimeVirtualTextureMaterialType InMaterialType, bool InClearTextures, FSceneInterface* InScene, FTransform const& InUVToWorld);
+	RENDERER_API FRuntimeVirtualTextureProducer(
+		FVTProducerDescription const& InDesc, 
+		uint32 InProducerId, 
+		ERuntimeVirtualTextureMaterialType InMaterialType, 
+		bool InClearTextures, 
+		FSceneInterface* InScene, 
+		FTransform const& InUVToWorld,
+		FBox const& InWorldBounds);
+	
 	RENDERER_API virtual ~FRuntimeVirtualTextureProducer() {}
 
 	//~ Begin IVirtualTexture Interface.
diff --git a/Engine/Source/Runtime/Renderer/Private/VT/RuntimeVirtualTextureRender.cpp b/Engine/Source/Runtime/Renderer/Private/VT/RuntimeVirtualTextureRender.cpp
--- a/Engine/Source/Runtime/Renderer/Private/VT/RuntimeVirtualTextureRender.cpp
+++ b/Engine/Source/Runtime/Renderer/Private/VT/RuntimeVirtualTextureRender.cpp
@@ -873,6 +873,7 @@ namespace RuntimeVirtualTexture
 		FRHIUnorderedAccessView* OutputUAV2,
 		FBox2D const& DestBox2,
 		FTransform const& UVToWorld,
+		FBox const& WorldBounds,
 		FBox2D const& UVRange,
 		uint8 vLevel,
 		uint8 MaxLevel,
@@ -915,6 +916,11 @@ namespace RuntimeVirtualTexture
 		const float ZOffset = -NearPlane;
 		ViewInitOptions.ProjectionMatrix = FReversedZOrthoMatrix(OrthoWidth, OrthoHeight, ZScale, ZOffset);
 
+		const FVector MipLevelParameter = FVector4((float)vLevel, 0.f, OrthoWidth / (float)TextureSize.X, OrthoHeight / (float)TextureSize.Y);
+		
+		const float HeightRange = FMath::Max(WorldBounds.Max.Z - WorldBounds.Min.Z, 1.f);
+		const FVector2D WorldHeightPackParameter = FVector2D(1.f / HeightRange, -WorldBounds.Min.Z / HeightRange);
+
 		ViewInitOptions.BackgroundColor = FLinearColor::Black;
 		ViewInitOptions.OverlayColor = FLinearColor::White;
 
@@ -926,7 +932,9 @@ namespace RuntimeVirtualTexture
 		View->CachedViewUniformShaderParameters = MakeUnique<FViewUniformShaderParameters>();
 		View->SetupUniformBufferParameters(SceneContext, nullptr, 0, *View->CachedViewUniformShaderParameters);
 		View->CachedViewUniformShaderParameters->WorldToVirtualTexture = WorldToUVRotate.ToMatrixNoScale();
-		View->CachedViewUniformShaderParameters->VirtualTextureParams = FVector4((float)vLevel, DebugType == ERuntimeVirtualTextureDebugType::Debug ? 1.f : 0.f, OrthoWidth / (float)TextureSize.X, OrthoHeight / (float)TextureSize.Y);
+		View->CachedViewUniformShaderParameters->RuntimeVirtualTextureMipLevel = MipLevelParameter;
+		View->CachedViewUniformShaderParameters->RuntimeVirtualTexturePackHeight = WorldHeightPackParameter;
+		View->CachedViewUniformShaderParameters->RuntimeVirtualTextureDebugParams = FVector4(DebugType == ERuntimeVirtualTextureDebugType::Debug ? 1.f : 0.f, 0.f, 0.f, 0.f);
 		View->ViewUniformBuffer = TUniformBufferRef<FViewUniformShaderParameters>::CreateUniformBufferImmediate(*View->CachedViewUniformShaderParameters, UniformBuffer_SingleFrame);
 		UploadDynamicPrimitiveShaderDataForView(RHICmdList, *(const_cast<FScene*>(Scene)), *View);
 		Scene->UniformBuffers.VirtualTextureViewUniformBuffer.UpdateUniformBufferImmediate(*View->CachedViewUniformShaderParameters);
@@ -1093,6 +1101,7 @@ namespace RuntimeVirtualTexture
 				InDesc.Targets[1].Texture, InDesc.Targets[1].UAV, PageDesc.DestBox[1],
 				InDesc.Targets[2].Texture, InDesc.Targets[2].UAV, PageDesc.DestBox[2],
 				InDesc.UVToWorld,
+				InDesc.WorldBounds,
 				PageDesc.UVRange,
 				PageDesc.vLevel,
 				InDesc.MaxLevel,
diff --git a/Engine/Source/Runtime/Renderer/Private/VT/RuntimeVirtualTextureSceneProxy.cpp b/Engine/Source/Runtime/Renderer/Private/VT/RuntimeVirtualTextureSceneProxy.cpp
--- a/Engine/Source/Runtime/Renderer/Private/VT/RuntimeVirtualTextureSceneProxy.cpp
+++ b/Engine/Source/Runtime/Renderer/Private/VT/RuntimeVirtualTextureSceneProxy.cpp
@@ -23,6 +23,7 @@ FRuntimeVirtualTextureSceneProxy::FRuntimeVirtualTextureSceneProxy(URuntimeVirtu
 
 		VirtualTexture = InComponent->GetVirtualTexture();
 		Transform = InComponent->GetVirtualTextureTransform();
+		const FBox Bounds = InComponent->Bounds.GetBox();
 
 		// The producer description is calculated using the transform to determine the aspect ratio
 		FVTProducerDescription Desc;
@@ -35,7 +36,7 @@ FRuntimeVirtualTextureSceneProxy::FRuntimeVirtualTextureSceneProxy(URuntimeVirtu
 		const bool bClearTextures = VirtualTexture->GetClearTextures();
 
 		// The Producer object created here will be passed into the virtual texture system which will take ownership.
-		IVirtualTexture* Producer = new FRuntimeVirtualTextureProducer(Desc, ProducerId, MaterialType, bClearTextures, InComponent->GetScene(), Transform);
+		IVirtualTexture* Producer = new FRuntimeVirtualTextureProducer(Desc, ProducerId, MaterialType, bClearTextures, InComponent->GetScene(), Transform, Bounds);
 
 		if (InComponent->IsStreamingLowMips() && VirtualTexture->GetStreamLowMips() > 0)
 		{
@@ -49,7 +50,7 @@ FRuntimeVirtualTextureSceneProxy::FRuntimeVirtualTextureSceneProxy(URuntimeVirtu
 		}
 
 		// The Initialize() call will allocate the virtual texture by spawning work on the render thread.
-		VirtualTexture->Initialize(Producer, Transform);
+		VirtualTexture->Initialize(Producer, Transform, Bounds);
 	}
 }
 
diff --git a/Engine/Source/Runtime/Renderer/Public/VT/RuntimeVirtualTextureRender.h b/Engine/Source/Runtime/Renderer/Public/VT/RuntimeVirtualTextureRender.h
--- a/Engine/Source/Runtime/Renderer/Public/VT/RuntimeVirtualTextureRender.h
+++ b/Engine/Source/Runtime/Renderer/Public/VT/RuntimeVirtualTextureRender.h
@@ -54,6 +54,8 @@ namespace RuntimeVirtualTexture
 		uint32 RuntimeVirtualTextureMask;
 		/** Virtual texture UV space to world space transform. */
 		FTransform UVToWorld;
+		/** Virtual texture world space bounds. */
+		FBox WorldBounds;
 		/** Material type of the runtime virtual texture that we are rendering. */
 		ERuntimeVirtualTextureMaterialType MaterialType;
 		/** Max mip level of the runtime virtual texture that we are rendering. */
