Fix for incorrectly using the same rendertargets multiple times with NoAction, when Load is the proper action to keep the data of the other view intact.

#rnx

diff --git a/Engine/Source/Runtime/Renderer/Private/SingleLayerWaterRendering.cpp b/Engine/Source/Runtime/Renderer/Private/SingleLayerWaterRendering.cpp
--- a/Engine/Source/Runtime/Renderer/Private/SingleLayerWaterRendering.cpp
+++ b/Engine/Source/Runtime/Renderer/Private/SingleLayerWaterRendering.cpp
@@ -475,11 +475,11 @@ void FDeferredShadingSceneRenderer::CopySingleLayerWaterTextures(FRHICommandList
 	    PassParameters->SceneDepthCopyDownsampleSampler = TStaticSamplerState<SF_Bilinear, AM_Clamp, AM_Clamp, AM_Clamp>::GetRHI();
 	    PassParameters->SVPositionToSourceTextureUV = FVector2D(RefractionDownsampleFactor / float(SceneContext.GetBufferSizeXY().X), RefractionDownsampleFactor / float(SceneContext.GetBufferSizeXY().Y));
     
-	    PassParameters->RenderTargets[0] = FRenderTargetBinding(SceneDepthWithoutSingleLayerWaterTexture, ERenderTargetLoadAction::ENoAction);
+	    PassParameters->RenderTargets[0] = FRenderTargetBinding(SceneDepthWithoutSingleLayerWaterTexture, (ViewIndex == 0) ? ERenderTargetLoadAction::ENoAction : ERenderTargetLoadAction::ELoad);
 
 	    if (bCopyColor)
 	    {
-		    PassParameters->RenderTargets[1] = FRenderTargetBinding(SceneColorWithoutSingleLayerWaterTexture, ERenderTargetLoadAction::ENoAction);
+		    PassParameters->RenderTargets[1] = FRenderTargetBinding(SceneColorWithoutSingleLayerWaterTexture, (ViewIndex == 0) ? ERenderTargetLoadAction::ENoAction : ERenderTargetLoadAction::ELoad);
 	    }
 
 	    FWaterRefractionCopyPS::FPermutationDomain PermutationVector;
