fix compile errors caused by merge

diff --git a/Engine/Source/Runtime/Renderer/Private/BasePassRendering.h b/Engine/Source/Runtime/Renderer/Private/BasePassRendering.h
--- a/Engine/Source/Runtime/Renderer/Private/BasePassRendering.h
+++ b/Engine/Source/Runtime/Renderer/Private/BasePassRendering.h
@@ -519,7 +519,7 @@ public:
 };
 
 //Alternative base pass PS for 128 bit canvas render targets that need to be set at shader compilation time.
-class F128BitRTBasePassPS : public TBasePassPS<TUniformLightMapPolicy<LMP_NO_LIGHTMAP>>
+class F128BitRTBasePassPS : public TBasePassPS<TUniformLightMapPolicy<LMP_NO_LIGHTMAP>, false>
 {
 	DECLARE_SHADER_TYPE(F128BitRTBasePassPS, MeshMaterial);
 public:
@@ -537,7 +537,7 @@ public:
 
 	/** Initialization constructor. */
 	F128BitRTBasePassPS(const ShaderMetaType::CompiledShaderInitializerType& Initializer): 
-		TBasePassPS<TUniformLightMapPolicy<LMP_NO_LIGHTMAP>>(Initializer)
+		TBasePassPS<TUniformLightMapPolicy<LMP_NO_LIGHTMAP>, false>(Initializer)
 	{}
 
 	/** Default constructor. */
