Fix incorrect voxel radius normalization.

diff --git a/Engine/Source/Runtime/Renderer/Private/HairStrands/HairStrandsVoxelization.cpp b/Engine/Source/Runtime/Renderer/Private/HairStrands/HairStrandsVoxelization.cpp
--- a/Engine/Source/Runtime/Renderer/Private/HairStrands/HairStrandsVoxelization.cpp
+++ b/Engine/Source/Runtime/Renderer/Private/HairStrands/HairStrandsVoxelization.cpp
@@ -1141,7 +1141,7 @@ static void AddVirtualVoxelizationRasterPass(
 	const FVector RasterAABBCenter = RasterAABB.GetCenter();
 	const FIntRect ViewportRect = FIntRect(0, 0, RasterResolution.X, RasterResolution.Y);
 
-	const float RadiusAtDepth1 = GStrandHairVoxelizationRasterizationScale * VoxelResources.Parameters.Common.VoxelWorldSize;
+	const float RadiusAtDepth1 = GStrandHairVoxelizationRasterizationScale * VoxelResources.Parameters.Common.VoxelWorldSize * 0.5f;
 	const bool bIsOrtho = true;
 	const FVector4 HairRenderInfo = PackHairRenderInfo(RadiusAtDepth1, RadiusAtDepth1, RadiusAtDepth1, 1);
 	const uint32 HairRenderInfoBits = PackHairRenderInfoBits(bIsOrtho, bIsGPUDriven);
