Fixed crash in PostProcessVisualizeComplexity when in ES3.1 preview mode.

diff --git a/Engine/Source/Runtime/Renderer/Private/PostProcess/PostProcessVisualizeComplexity.cpp b/Engine/Source/Runtime/Renderer/Private/PostProcess/PostProcessVisualizeComplexity.cpp
--- a/Engine/Source/Runtime/Renderer/Private/PostProcess/PostProcessVisualizeComplexity.cpp
+++ b/Engine/Source/Runtime/Renderer/Private/PostProcess/PostProcessVisualizeComplexity.cpp
@@ -123,7 +123,9 @@ FScreenPassTexture AddVisualizeComplexityPass(FRDGBuilder& GraphBuilder, const F
 	PassParameters->DebugViewShaderMode = DVSM_ShaderComplexity;
 	FVisualizeComplexityApplyPS::EQuadOverdraw QuadOverdrawEnum = FVisualizeComplexityApplyPS::EQuadOverdraw::Disable;
 
-	if (FRDGTextureRef QuadOverdrawTexture = GraphBuilder.TryRegisterExternalTexture(SceneRenderTargets.QuadOverdrawBuffer))
+	FRDGTextureRef QuadOverdrawTexture = GraphBuilder.TryRegisterExternalTexture(SceneRenderTargets.QuadOverdrawBuffer);
+
+	if (QuadOverdrawTexture && AllowDebugViewShaderMode(DVSM_QuadComplexity, View.GetShaderPlatform(), View.FeatureLevel))
 	{
 		PassParameters->QuadOverdrawTexture = QuadOverdrawTexture;
 		PassParameters->DebugViewShaderMode = DebugViewShaderMode;
