Fix for limit of 256 texture array layers on some Android platforms

diff --git a/Engine/Source/Runtime/Renderer/Private/ReflectionEnvironmentCapture.cpp b/Engine/Source/Runtime/Renderer/Private/ReflectionEnvironmentCapture.cpp
--- a/Engine/Source/Runtime/Renderer/Private/ReflectionEnvironmentCapture.cpp
+++ b/Engine/Source/Runtime/Renderer/Private/ReflectionEnvironmentCapture.cpp
@@ -882,6 +882,8 @@ void FScene::AllocateReflectionCaptures(const TArray<UReflectionCaptureComponent
 	{
 		if (GetFeatureLevel() >= ERHIFeatureLevel::SM5)
 		{
+			int32_t PlatformMaxNumReflectionCaptures = FMath::Min(FMath::FloorToInt(GMaxTextureArrayLayers / 6.0f), GMaxNumReflectionCaptures);
+
 			for (int32 CaptureIndex = 0; CaptureIndex < NewCaptures.Num(); CaptureIndex++)
 			{
 				bool bAlreadyExists = false;
@@ -896,9 +898,9 @@ void FScene::AllocateReflectionCaptures(const TArray<UReflectionCaptureComponent
 						bAlreadyExists = true;
 					}
 				}
-
+				
 				// Add the capture to the allocated list
-				if (!bAlreadyExists && ReflectionSceneData.AllocatedReflectionCapturesGameThread.Num() < GMaxNumReflectionCaptures)
+				if (!bAlreadyExists && ReflectionSceneData.AllocatedReflectionCapturesGameThread.Num() < PlatformMaxNumReflectionCaptures)
 				{
 					ReflectionSceneData.AllocatedReflectionCapturesGameThread.Add(NewCaptures[CaptureIndex]);
 				}
@@ -917,7 +919,7 @@ void FScene::AllocateReflectionCaptures(const TArray<UReflectionCaptureComponent
 				DesiredMaxCubemaps = FMath::Pow(MaxCubemapsRoundUpBase, FMath::TruncToInt(Exponent) + 1);
 			}
 
-			DesiredMaxCubemaps = FMath::Min(DesiredMaxCubemaps, GMaxNumReflectionCaptures);
+			DesiredMaxCubemaps = FMath::Min(DesiredMaxCubemaps, PlatformMaxNumReflectionCaptures);
 
 			const int32 ReflectionCaptureSize = UReflectionCaptureComponent::GetReflectionCaptureSize();
 			bool bNeedsUpdateAllCaptures = DesiredMaxCubemaps != ReflectionSceneData.MaxAllocatedReflectionCubemapsGameThread || ReflectionCaptureSize != ReflectionSceneData.CubemapArray.GetCubemapSize();
diff --git a/Engine/Source/Runtime/Renderer/Private/SceneVisibility.cpp b/Engine/Source/Runtime/Renderer/Private/SceneVisibility.cpp
--- a/Engine/Source/Runtime/Renderer/Private/SceneVisibility.cpp
+++ b/Engine/Source/Runtime/Renderer/Private/SceneVisibility.cpp
@@ -3551,10 +3551,11 @@ void UpdateReflectionSceneData(FScene* Scene)
 	ReflectionSceneData.NumSphereCaptures = 0;
 
 	const int32 MaxCubemaps = ReflectionSceneData.CubemapArray.GetMaxCubemaps();
+	int32_t PlatformMaxNumReflectionCaptures = FMath::Min(FMath::FloorToInt(GMaxTextureArrayLayers / 6.0f), GMaxNumReflectionCaptures);
 
 	// Pack visible reflection captures into the uniform buffer, each with an index to its cubemap array entry.
 	// GPUScene primitive data stores closest reflection capture as index into this buffer, so this index which must be invalidate every time OutSortData contents change.
-	for (int32 ReflectionProxyIndex = 0; ReflectionProxyIndex < ReflectionSceneData.RegisteredReflectionCaptures.Num() && ReflectionSceneData.SortedCaptures.Num() < GMaxNumReflectionCaptures; ReflectionProxyIndex++)
+	for (int32 ReflectionProxyIndex = 0; ReflectionProxyIndex < ReflectionSceneData.RegisteredReflectionCaptures.Num() && ReflectionSceneData.SortedCaptures.Num() < PlatformMaxNumReflectionCaptures; ReflectionProxyIndex++)
 	{
 		FReflectionCaptureProxy* CurrentCapture = ReflectionSceneData.RegisteredReflectionCaptures[ReflectionProxyIndex];
 
