Fixed visualize motion blur to accept the last pass override.

diff --git a/Engine/Source/Runtime/Renderer/Private/PostProcess/PostProcessMotionBlur.cpp b/Engine/Source/Runtime/Renderer/Private/PostProcess/PostProcessMotionBlur.cpp
--- a/Engine/Source/Runtime/Renderer/Private/PostProcess/PostProcessMotionBlur.cpp
+++ b/Engine/Source/Runtime/Renderer/Private/PostProcess/PostProcessMotionBlur.cpp
@@ -146,6 +146,11 @@ EMotionBlurQuality GetMotionBlurQuality()
 	return static_cast<EMotionBlurQuality>(Quality - 1);
 }
 
+EMotionBlurFilter GetMotionBlurFilter()
+{
+	return CVarMotionBlurSeparable.GetValueOnRenderThread() != 0 ? EMotionBlurFilter::Separable : EMotionBlurFilter::Unified;
+}
+
 FRHISamplerState* GetMotionBlurColorSampler()
 {
 	bool bFiltered = false;
@@ -401,8 +406,8 @@ enum class EMotionBlurFilterPass : uint32
 struct FMotionBlurViewports
 {
 	FMotionBlurViewports(
-		const FScreenPassTextureViewport& InColorViewport,
-		const FScreenPassTextureViewport& InVelocityViewport)
+		FScreenPassTextureViewport InColorViewport,
+		FScreenPassTextureViewport InVelocityViewport)
 	{
 		Color = InColorViewport;
 		Velocity = InVelocityViewport;
@@ -727,34 +732,28 @@ FRDGTextureRef AddMotionBlurFilterPass(
 	return ColorTextureOutput;
 }
 
-FRDGTextureRef AddVisualizeMotionBlurPass(
-	FRDGBuilder& GraphBuilder,
-	const FViewInfo& View,
-	FIntRect ColorViewportRect,
-	FIntRect VelocityViewportRect,
-	FRDGTextureRef ColorTexture,
-	FRDGTextureRef DepthTexture,
-	FRDGTextureRef VelocityTexture)
+FScreenPassTexture AddVisualizeMotionBlurPass(FRDGBuilder& GraphBuilder, const FViewInfo& View, const FMotionBlurInputs& Inputs)
 {
-	const FScreenPassTextureViewport ColorViewport(ColorTexture, ColorViewportRect);
-	const FScreenPassTextureViewport VelocityViewport(VelocityTexture, VelocityViewportRect);
-	const FMotionBlurViewports Viewports(ColorViewport, VelocityViewport);
+	check(Inputs.SceneColor.IsValid());
+	check(Inputs.SceneDepth.IsValid());
+	check(Inputs.SceneVelocity.IsValid());
+	checkf(Inputs.SceneDepth.ViewRect == Inputs.SceneVelocity.ViewRect, TEXT("The implementation requires that depth and velocity have the same viewport."));
+
+	FScreenPassRenderTarget Output = Inputs.OverrideOutput;
 
-	FRDGTextureDesc OutputDesc = ColorTexture->Desc;
-	OutputDesc.TargetableFlags &= ~(TexCreate_RenderTargetable | TexCreate_UAV);
-	OutputDesc.TargetableFlags |= TexCreate_RenderTargetable;
+	if (!Output.IsValid())
+	{
+		Output = FScreenPassRenderTarget::CreateFromInput(GraphBuilder, Inputs.SceneColor, View.GetOverwriteLoadAction(), TEXT("VisualizeMotionBlur"));
+	}
 
-	FScreenPassRenderTarget Output;
-	Output.Texture = GraphBuilder.CreateTexture(OutputDesc, TEXT("MotionBlurVisualize"));
-	Output.ViewRect = ColorViewport.Rect;
-	Output.LoadAction = View.GetOverwriteLoadAction();
+	const FMotionBlurViewports Viewports(FScreenPassTextureViewport(Inputs.SceneColor), FScreenPassTextureViewport(Inputs.SceneVelocity));
 
 	FMotionBlurVisualizePS::FParameters* PassParameters = GraphBuilder.AllocParameters<FMotionBlurVisualizePS::FParameters>();
 	PassParameters->WorldToClipPrev = GetPreviousWorldToClipMatrix(View);
 	PassParameters->View = View.ViewUniformBuffer;
-	PassParameters->ColorTexture = ColorTexture;
-	PassParameters->DepthTexture = DepthTexture;
-	PassParameters->VelocityTexture = VelocityTexture;
+	PassParameters->ColorTexture = Inputs.SceneColor.Texture;
+	PassParameters->DepthTexture = Inputs.SceneDepth.Texture;
+	PassParameters->VelocityTexture = Inputs.SceneVelocity.Texture;
 	PassParameters->Color = Viewports.ColorParameters;
 	PassParameters->Velocity = Viewports.VelocityParameters;
 	PassParameters->ColorToVelocity = Viewports.ColorToVelocityTransform;
@@ -765,7 +764,7 @@ FRDGTextureRef AddVisualizeMotionBlurPass(
 
 	TShaderMapRef<FMotionBlurVisualizePS> PixelShader(View.ShaderMap);
 
-	AddDrawScreenPass(GraphBuilder, RDG_EVENT_NAME("Visualizer"), View, ColorViewport, ColorViewport, PixelShader, PassParameters);
+	AddDrawScreenPass(GraphBuilder, RDG_EVENT_NAME("Visualizer"), View, Viewports.Color, Viewports.Color, PixelShader, PassParameters);
 
 	Output.LoadAction = ERenderTargetLoadAction::ELoad;
 
@@ -808,25 +807,18 @@ FRDGTextureRef AddVisualizeMotionBlurPass(
 		Canvas.DrawShadowedString(X + ColumnWidth, Y, *Line, GetStatsFont(), FLinearColor(1, 1, 0));
 	});
 
-	return Output.Texture;
+	return Output;
 }
 
-FRDGTextureRef AddMotionBlurPass(
-	FRDGBuilder& GraphBuilder,
-	const FViewInfo& View,
-	FIntRect ColorViewportRect,
-	FIntRect VelocityViewportRect,
-	FRDGTextureRef ColorTexture,
-	FRDGTextureRef DepthTexture,
-	FRDGTextureRef VelocityTexture)
+FScreenPassTexture AddMotionBlurPass(FRDGBuilder& GraphBuilder, const FViewInfo& View, const FMotionBlurInputs& Inputs)
 {
-	check(ColorTexture);
-	check(DepthTexture);
-	check(VelocityTexture);
+	check(Inputs.SceneColor.IsValid());
+	check(Inputs.SceneDepth.IsValid());
+	check(Inputs.SceneVelocity.IsValid());
+	checkf(Inputs.SceneDepth.ViewRect == Inputs.SceneVelocity.ViewRect, TEXT("The motion blur depth and velocity must have the same viewport."));
+	checkf(!Inputs.OverrideOutput.IsValid(), TEXT("The motion blur override output support is unimplemented."));
 
-	const FMotionBlurViewports Viewports(
-		FScreenPassTextureViewport(ColorTexture, ColorViewportRect),
-		FScreenPassTextureViewport(VelocityTexture, VelocityViewportRect));
+	const FMotionBlurViewports Viewports(FScreenPassTextureViewport(Inputs.SceneColor), FScreenPassTextureViewport(Inputs.SceneVelocity));
 
 	RDG_EVENT_SCOPE(GraphBuilder, "MotionBlur");
 
@@ -836,29 +828,28 @@ FRDGTextureRef AddMotionBlurPass(
 		GraphBuilder,
 		View,
 		Viewports,
-		ColorTexture,
-		DepthTexture,
-		VelocityTexture,
+		Inputs.SceneColor.Texture,
+		Inputs.SceneDepth.Texture,
+		Inputs.SceneVelocity.Texture,
 		&VelocityFlatTexture,
 		&VelocityTileTexture);
 
-	const EMotionBlurQuality MotionBlurQuality = GetMotionBlurQuality();
+	FScreenPassTexture Output;
+	Output.ViewRect = Viewports.Color.Rect;
 
-	const bool bUseSeparableFilter = CVarMotionBlurSeparable.GetValueOnRenderThread() != 0;
-
-	if (bUseSeparableFilter)
+	if (Inputs.Filter == EMotionBlurFilter::Separable)
 	{
 		FRDGTextureRef MotionBlurFilterTexture = AddMotionBlurFilterPass(
 			GraphBuilder,
 			View,
 			Viewports,
-			ColorTexture,
+			Inputs.SceneColor.Texture,
 			VelocityFlatTexture,
 			VelocityTileTexture,
 			EMotionBlurFilterPass::Separable0,
-			MotionBlurQuality);
+			Inputs.Quality);
 
-		return AddMotionBlurFilterPass(
+		Output.Texture = AddMotionBlurFilterPass(
 			GraphBuilder,
 			View,
 			Viewports,
@@ -866,18 +857,20 @@ FRDGTextureRef AddMotionBlurPass(
 			VelocityFlatTexture,
 			VelocityTileTexture,
 			EMotionBlurFilterPass::Separable1,
-			MotionBlurQuality);
+			Inputs.Quality);
 	}
 	else
 	{
-		return AddMotionBlurFilterPass(
+		Output.Texture = AddMotionBlurFilterPass(
 			GraphBuilder,
 			View,
 			Viewports,
-			ColorTexture,
+			Inputs.SceneColor.Texture,
 			VelocityFlatTexture,
 			VelocityTileTexture,
 			EMotionBlurFilterPass::Unified,
-			MotionBlurQuality);
+			Inputs.Quality);
 	}
+
+	return Output;
 }
\ No newline at end of file
diff --git a/Engine/Source/Runtime/Renderer/Private/PostProcess/PostProcessMotionBlur.h b/Engine/Source/Runtime/Renderer/Private/PostProcess/PostProcessMotionBlur.h
--- a/Engine/Source/Runtime/Renderer/Private/PostProcess/PostProcessMotionBlur.h
+++ b/Engine/Source/Runtime/Renderer/Private/PostProcess/PostProcessMotionBlur.h
@@ -24,23 +24,38 @@ enum class EMotionBlurQuality : uint32
 	MAX
 };
 
-// Returns the motion blur quality level set by the 'r.MotionBlurQuality' CVar.
+enum class EMotionBlurFilter : uint32
+{
+	Unified,
+	Separable
+};
+
+// Returns the global setting for motion blur quality.
 EMotionBlurQuality GetMotionBlurQuality();
 
-FRDGTextureRef AddMotionBlurPass(
-	FRDGBuilder& GraphBuilder,
-	const FViewInfo& View,
-	FIntRect ColorViewportRect,
-	FIntRect VelocityViewportRect,
-	FRDGTextureRef ColorTexture,
-	FRDGTextureRef DepthTexture,
-	FRDGTextureRef VelocityTexture);
-
-FRDGTextureRef AddVisualizeMotionBlurPass(
-	FRDGBuilder& GraphBuilder,
-	const FViewInfo& View,
-	FIntRect ColorViewportRect,
-	FIntRect VelocityViewportRect,
-	FRDGTextureRef ColorTexture,
-	FRDGTextureRef DepthTexture,
-	FRDGTextureRef VelocityTexture);
\ No newline at end of file
+// Returns the global setting for motion blur filter.
+EMotionBlurFilter GetMotionBlurFilter();
+
+struct FMotionBlurInputs
+{
+	// [Optional] Render to the specified output. If invalid, a new texture is created and returned.
+	FScreenPassRenderTarget OverrideOutput;
+
+	// [Required] The input scene color and view rect.
+	FScreenPassTexture SceneColor;
+
+	// [Required] The input scene depth and view rect.
+	FScreenPassTexture SceneDepth;
+
+	// [Required] The input scene velocity and view rect.
+	FScreenPassTexture SceneVelocity;
+
+	// [Required] Quality to use when processing motion blur.
+	EMotionBlurQuality Quality = EMotionBlurQuality::VeryHigh;
+
+	// [Required] Filter to use when processing motion blur.
+	EMotionBlurFilter Filter = EMotionBlurFilter::Separable;
+};
+
+FScreenPassTexture AddMotionBlurPass(FRDGBuilder& GraphBuilder, const FViewInfo& View, const FMotionBlurInputs& Inputs);
+FScreenPassTexture AddVisualizeMotionBlurPass(FRDGBuilder& GraphBuilder, const FViewInfo& View, const FMotionBlurInputs& Inputs);
\ No newline at end of file
diff --git a/Engine/Source/Runtime/Renderer/Private/PostProcess/PostProcessing.cpp b/Engine/Source/Runtime/Renderer/Private/PostProcess/PostProcessing.cpp
--- a/Engine/Source/Runtime/Renderer/Private/PostProcess/PostProcessing.cpp
+++ b/Engine/Source/Runtime/Renderer/Private/PostProcess/PostProcessing.cpp
@@ -244,6 +244,7 @@ void AddPostProcessingPasses(FRDGBuilder& GraphBuilder, const FViewInfo& View, c
 
 	enum class EPass : uint32
 	{
+		MotionBlur,
 		Tonemap,
 		FXAA,
 		PostProcessMaterialAfterTonemapping,
@@ -267,6 +268,7 @@ void AddPostProcessingPasses(FRDGBuilder& GraphBuilder, const FViewInfo& View, c
 
 	const TCHAR* PassNames[] =
 	{
+		TEXT("MotionBlur"),
 		TEXT("Tonemap"),
 		TEXT("FXAA"),
 		TEXT("PostProcessMaterial (AfterTonemapping)"),
@@ -365,6 +367,7 @@ void AddPostProcessingPasses(FRDGBuilder& GraphBuilder, const FViewInfo& View, c
 
 		const FPostProcessMaterialChain PostProcessMaterialAfterTonemappingChain = GetPostProcessMaterialChain(View, BL_AfterTonemapping);
 
+		PassSequence.SetEnabled(EPass::MotionBlur, bVisualizeMotionBlur || bMotionBlurEnabled);
 		PassSequence.SetEnabled(EPass::Tonemap, bTonemapEnabled);
 		PassSequence.SetEnabled(EPass::FXAA, AntiAliasingMethod == AAM_FXAA);
 		PassSequence.SetEnabled(EPass::PostProcessMaterialAfterTonemapping, PostProcessMaterialAfterTonemappingChain.Num() != 0);
@@ -467,16 +470,25 @@ void AddPostProcessingPasses(FRDGBuilder& GraphBuilder, const FViewInfo& View, c
 			}
 		}
 
-		// Motion blur visualization replaces motion blur when enabled.
-		if (bVisualizeMotionBlur)
+		if (PassSequence.IsEnabled(EPass::MotionBlur))
 		{
-			check(Velocity.ViewRect == SceneDepth.ViewRect);
-			SceneColor.Texture = AddVisualizeMotionBlurPass(GraphBuilder, View, SceneColor.ViewRect, SceneDepth.ViewRect, SceneColor.Texture, SceneDepth.Texture, Velocity.Texture);
-		}
-		else if (bMotionBlurEnabled)
-		{
-			check(Velocity.ViewRect == SceneDepth.ViewRect);
-			SceneColor.Texture = AddMotionBlurPass(GraphBuilder, View, SceneColor.ViewRect, SceneDepth.ViewRect, SceneColor.Texture, SceneDepth.Texture, Velocity.Texture);
+			FMotionBlurInputs PassInputs;
+			PassSequence.AcceptOverrideIfLastPass(EPass::MotionBlur, PassInputs.OverrideOutput);
+			PassInputs.SceneColor = SceneColor;
+			PassInputs.SceneDepth = SceneDepth;
+			PassInputs.SceneVelocity = Velocity;
+			PassInputs.Quality = GetMotionBlurQuality();
+			PassInputs.Filter = GetMotionBlurFilter();
+
+			// Motion blur visualization replaces motion blur when enabled.
+			if (bVisualizeMotionBlur)
+			{
+				SceneColor = AddVisualizeMotionBlurPass(GraphBuilder, View, PassInputs);
+			}
+			else
+			{
+				SceneColor = AddMotionBlurPass(GraphBuilder, View, PassInputs);
+			}
 		}
 
 		// If TAA didn't do it, downsample the scene color texture by half.
@@ -656,6 +668,7 @@ void AddPostProcessingPasses(FRDGBuilder& GraphBuilder, const FViewInfo& View, c
 	// Minimal PostProcessing - Separate translucency composition and gamma-correction only.
 	else
 	{
+		PassSequence.SetEnabled(EPass::MotionBlur, false);
 		PassSequence.SetEnabled(EPass::Tonemap, true);
 		PassSequence.SetEnabled(EPass::FXAA, false);
 		PassSequence.SetEnabled(EPass::PostProcessMaterialAfterTonemapping, false);
