Fixed build break.

diff --git a/Engine/Source/Runtime/Renderer/Private/PostProcess/PostProcessMotionBlur.cpp b/Engine/Source/Runtime/Renderer/Private/PostProcess/PostProcessMotionBlur.cpp
--- a/Engine/Source/Runtime/Renderer/Private/PostProcess/PostProcessMotionBlur.cpp
+++ b/Engine/Source/Runtime/Renderer/Private/PostProcess/PostProcessMotionBlur.cpp
@@ -807,7 +807,7 @@ FScreenPassTexture AddVisualizeMotionBlurPass(FRDGBuilder& GraphBuilder, const F
 		Canvas.DrawShadowedString(X + ColumnWidth, Y, *Line, GetStatsFont(), FLinearColor(1, 1, 0));
 	});
 
-	return Output;
+	return MoveTemp(Output);
 }
 
 FScreenPassTexture AddMotionBlurPass(FRDGBuilder& GraphBuilder, const FViewInfo& View, const FMotionBlurInputs& Inputs)
