Fix frequent crash happening during hair projection due to a limit on the mesh section count.

Increase the limit of mesh section from 16 to 64.

diff --git a/Engine/Shaders/Private/HairStrands/HairStrandsMeshProjectionCommon.ush b/Engine/Shaders/Private/HairStrands/HairStrandsMeshProjectionCommon.ush
--- a/Engine/Shaders/Private/HairStrands/HairStrandsMeshProjectionCommon.ush
+++ b/Engine/Shaders/Private/HairStrands/HairStrandsMeshProjectionCommon.ush
@@ -247,13 +247,13 @@ struct FHairTriangleIndex
 
 uint EncodeTriangleIndex(uint TriangleIndex, uint SectionIndex)
 {
-	return ((SectionIndex & 0xF)<<28) | (TriangleIndex & 0xFFFFFFF);
+	return ((SectionIndex & 0x3F)<<26) | (TriangleIndex & 0x3FFFFFF);
 }
 
 FHairTriangleIndex DecodeTriangleIndex(uint Encoded)
 {
 	FHairTriangleIndex Out;
-	Out.SectionIndex = (Encoded>>28) & 0xF;
-	Out.TriangleIndex = Encoded & 0xFFFFFFF;
+	Out.SectionIndex = (Encoded>>26) & 0x3F;
+	Out.TriangleIndex = Encoded & 0x3FFFFFF;
 	return Out;
 }
\ No newline at end of file
diff --git a/Engine/Source/Runtime/Renderer/Private/HairStrands/HairStrandsMeshProjection.cpp b/Engine/Source/Runtime/Renderer/Private/HairStrands/HairStrandsMeshProjection.cpp
--- a/Engine/Source/Runtime/Renderer/Private/HairStrands/HairStrandsMeshProjection.cpp
+++ b/Engine/Source/Runtime/Renderer/Private/HairStrands/HairStrandsMeshProjection.cpp
@@ -262,11 +262,11 @@ static void AddHairStrandMeshProjectionPass(
 	// The current shader code HairStrandsMeshProjection.usf encode the section ID onto the highest 4bits of a 32bits uint. 
 	// This limits the number of section to 16. See EncodeTriangleIndex & DecodeTriangleIndex functions in 
 	// HairStarndsMeshProjectionCommon.ush for mode details.
-	// This means that the mesh needs to have less than 285M triangles (since triangle ID is stored onto 28bits).
+	// This means that the mesh needs to have less than 67M triangles (since triangle ID is stored onto 26bits).
 	//
 	// This could be increase if necessary.
-	check(MeshSectionData.SectionIndex < 16);
-	check(MeshSectionData.NumPrimitives < ((1<<28)-1))
+	check(MeshSectionData.SectionIndex < 64);
+	check(MeshSectionData.NumPrimitives < ((1<<26)-1))
 
 	// For projecting hair onto a skeletal mesh, 1 thread is spawn for each hair which iterates over all triangles.
 	// To avoid TDR, we split projection into multiple passes when the mesh is too large.
