Fix forward rendering support for hair rendering.

diff --git a/Engine/Shaders/Private/HairStrands/HairStrandsDeepShadowMask.usf b/Engine/Shaders/Private/HairStrands/HairStrandsDeepShadowMask.usf
--- a/Engine/Shaders/Private/HairStrands/HairStrandsDeepShadowMask.usf
+++ b/Engine/Shaders/Private/HairStrands/HairStrandsDeepShadowMask.usf
@@ -5,7 +5,6 @@
 
 #include "../Common.ush"
 #include "../CommonViewUniformBuffer.ush"
-#include "../SceneTextureParameters.ush"
 #include "../DeferredShadingCommon.ush"
 #include "HairStrandsVisibilityCommon.ush"
 #include "HairStrandsVoxelPageCommon.ush"
@@ -28,7 +27,7 @@ StructuredBuffer<FDeepShadowTransform> DeepShadow_WorldToLightTransformBuffer;
 uint				DeepShadow_bIsGPUDriven;
 
 float4x4			DeepShadow_CPUWorldToLightTransform;
-uint				DeepShadow_IsWholeSceneLight;
+uint				bIsWholeSceneLight;
 int2				DeepShadow_SlotResolution;
 int2				DeepShadow_SlotOffset;
 uint				DeepShadow_SlotIndex;
@@ -45,6 +44,7 @@ Texture2D<float4>	RayMarchMaskTexture;
 Texture2D<float>	DeepShadow_FrontDepthTexture;
 Texture2D<float4>	DeepShadow_DomTexture;
 
+Texture2D<float>	SceneDepthTexture;
 Texture2D<uint4>	CategorizationTexture;
 RWTexture2D<float4>	OutputColor;
 
@@ -345,11 +345,11 @@ void MainPS(
 {
 	const float2 UV = Input.UV;
 	const uint2 PixelCoord = floor(Input.Position.xy);
-	FGBufferData GBuffer = GetGBufferDataFromSceneTextures(UV);
+	const float SceneDepth = ConvertFromDeviceZ(SceneDepthTexture.Load(uint3(PixelCoord, 0)));
 	const FCategorizationData Categorization = DecodeCategorizationData(CategorizationTexture.Load(uint3(PixelCoord, 0)));
 	const bool bIsHairPixel = Categorization.PixelCoverage >= 1;
 
-	const float3 WorldPosition		= TotWorldPosition(UV, GBuffer.Depth);
+	const float3 WorldPosition		= TotWorldPosition(UV, SceneDepth);
 	const float3 LightSpacePosition = ToLightPosition(WorldPosition);
 
 	// PCF or simple point sampler
@@ -393,7 +393,7 @@ void MainPS(
 
 	// Depending of if the light is direction (i.e., WholeSceneLight) or a local light (Point/Spot/Rect), 
 	// we change the output layout
-    OutColor = DeepShadow_IsWholeSceneLight ? float4(Visibility, 1, 1, 1) : float4(1, 1, Visibility, 1);
+    OutColor = bIsWholeSceneLight ? float4(Visibility, 1, 1, 1) : float4(1, 1, Visibility, 1);
 }
 
 #endif
@@ -406,11 +406,12 @@ void MainPS(
 {
 	const float2 UV = Input.UV;
 	const uint2 PixelCoord = floor(Input.Position.xy);
-	FGBufferData GBuffer = GetGBufferDataFromSceneTextures(UV);
+	const float SceneDepth = ConvertFromDeviceZ(SceneDepthTexture.Load(uint3(PixelCoord, 0)));
+
 	const FCategorizationData Categorization = DecodeCategorizationData(CategorizationTexture.Load(uint3(PixelCoord, 0)));
 	const bool bIsHairPixel = Categorization.PixelCoverage >= 1;
 
-	float3 WorldPosition = TotWorldPosition(UV, GBuffer.Depth);
+	float3 WorldPosition = TotWorldPosition(UV, SceneDepth);
 
 	// PCF or simple point sampler
 	float Visibility = 1;
@@ -456,7 +457,7 @@ void MainPS(
 		TraversalSettings.bDebugEnabled = bDebugEnabled;
 		TraversalSettings.SteppingScale = VirtualVoxel.SteppingScale;
 		TraversalSettings.Random = Random;
-		TraversalSettings.PixelRadius = GBuffer.Depth * VirtualVoxel.HairCoveragePixelRadiusAtDepth1;
+		TraversalSettings.PixelRadius = SceneDepth * VirtualVoxel.HairCoveragePixelRadiusAtDepth1;
 		
 		if (RayMarchMaskTexture.Load(uint3(PixelCoord, 0)).r > 0.0f &&  Categorization.PixelCoverage < CoverageThreshold)
 		{
@@ -480,7 +481,7 @@ void MainPS(
 
 	// Depending of if the light is direction (i.e., WholeSceneLight) or a local light (Point/Spot/Rect), 
 	// we change the output layout
-	OutColor = DeepShadow_IsWholeSceneLight ? float4(Visibility, 1, 1, 1) : float4(1, 1, Visibility, 1);
+	OutColor = bIsWholeSceneLight ? float4(Visibility, 1, 1, 1) : float4(1, 1, Visibility, 1);
 }
 
 #endif
\ No newline at end of file
diff --git a/Engine/Shaders/Private/HairStrands/HairStrandsDeepTransmittanceCommon.ush b/Engine/Shaders/Private/HairStrands/HairStrandsDeepTransmittanceCommon.ush
--- a/Engine/Shaders/Private/HairStrands/HairStrandsDeepTransmittanceCommon.ush
+++ b/Engine/Shaders/Private/HairStrands/HairStrandsDeepTransmittanceCommon.ush
@@ -278,11 +278,7 @@ FHairTransmittanceMask ComputeTransmittance(
 	Texture3D<float4>	HairLUTTexture,
 	SamplerState		LinearSampler)
 {
-	const bool bIsHairPixel = GBuffer.ShadingModelID == SHADINGMODELID_HAIR;
-
 	FHairTransmittanceMask Out = InitHairTransmittanceMask();
-
-	if (bIsHairPixel)
 	{
 		const bool bIsDirectionalLight = Settings.LightPosition.w == 0;
 
diff --git a/Engine/Shaders/Private/ShadowProjectionPixelShader.usf b/Engine/Shaders/Private/ShadowProjectionPixelShader.usf
--- a/Engine/Shaders/Private/ShadowProjectionPixelShader.usf
+++ b/Engine/Shaders/Private/ShadowProjectionPixelShader.usf
@@ -65,9 +65,10 @@ float InvPerObjectShadowFadeLength;
 	float4 LightPositionOrDirection;
 	float ShadowReceiverBias;
 
-#if USE_FADE_PLANE
+#if USE_FADE_PLANE || SUBPIXEL_SHADOW
 	float FadePlaneOffset;
 	float InvFadePlaneLength;
+	uint bCascadeUseFadePlane;
 #endif
 
 #if USE_PCSS
@@ -107,12 +108,17 @@ void Main(
 	out float4 OutColor : SV_Target0
 	)
 {
+	#if USE_FADE_PLANE
+	const bool bUseFadePlane = true;
+	#endif
+
 	const FPQMPContext PQMPContext = PQMPInit(SVPos.xy);
 	float2 ScreenUV = float2(SVPos.xy * View.BufferSizeAndInvSize.zw);
 	float SceneW = CalcSceneDepth(ScreenUV);
 
 	bool bIsSubsurfaceCompatible = true;
 #if SUBPIXEL_SHADOW
+	const bool bUseFadePlane = bCascadeUseFadePlane > 0;
 	const uint3 PixelCoord = uint3(floor(SVPos.xy), 0);
 	const FCategorizationData CategorizationData = DecodeCategorizationData(HairCategorizationTexture.Load(PixelCoord));
 	if (CategorizationData.PixelCoverage > 0)
@@ -128,7 +134,7 @@ void Main(
 	if (CategorizationData.ClosestDepth > ShadowNearAndFarDepth.x || CategorizationData.ClosestDepth < ShadowNearAndFarDepth.y) // Inverse-Z
 	{
 		discard;
-	}
+	};
 #endif
 
 	float4 ScreenPosition = float4(((ScreenUV.xy - View.ScreenPositionScaleBias.wz ) / View.ScreenPositionScaleBias.xy) * SceneW, SceneW, 1);
@@ -242,9 +248,12 @@ void Main(
 	}
 	#endif // !USE_PCSS
 		
-	#if USE_FADE_PLANE
+	#if USE_FADE_PLANE || SUBPIXEL_SHADOW
+	if (bUseFadePlane)
+	{
 		// Create a blend factor which is one before and at the fade plane, and lerps to zero at the far plane.
 		BlendFactor = 1.0f - saturate((SceneW - FadePlaneOffset) * InvFadePlaneLength);
+	}
 	#endif
 
 	#if FEATURE_LEVEL >= FEATURE_LEVEL_SM4 && !FORWARD_SHADING && !APPLY_TRANSLUCENCY_SHADOWS
@@ -326,9 +335,12 @@ void Main(
 	OutColor = EncodeLightAttenuation(half4(FadedShadow, FadedSSSShadow, FadedShadow, FadedSSSShadow));
 #endif
 
-	#if USE_FADE_PLANE
-		// When the fade plane is in use for CSMs, we output the fade value in the alpha channel for blending.
+	#if USE_FADE_PLANE || SUBPIXEL_SHADOW
+	// When the fade plane is in use for CSMs, we output the fade value in the alpha channel for blending.
+	if (bUseFadePlane)
+	{
 		OutColor.a = BlendFactor;
+	}
 	#endif
 
 #if MODULATED_SHADOWS
diff --git a/Engine/Source/Runtime/Renderer/Private/DeferredShadingRenderer.cpp b/Engine/Source/Runtime/Renderer/Private/DeferredShadingRenderer.cpp
--- a/Engine/Source/Runtime/Renderer/Private/DeferredShadingRenderer.cpp
+++ b/Engine/Source/Runtime/Renderer/Private/DeferredShadingRenderer.cpp
@@ -1823,11 +1823,22 @@ void FDeferredShadingSceneRenderer::Render(FRHICommandListImmediate& RHICmdList)
 		ComputeVolumetricFog(RHICmdList);
 	}
 
-	TRefCountPtr<IPooledRenderTarget> ForwardScreenSpaceShadowMask;
+	FHairStrandsDatas* HairDatas = nullptr;
+	FHairStrandsDatas HairDatasStorage;
+	const bool bIsViewCompatible = Views.Num() > 0 && Views[0].Family->ViewMode == VMI_Lit;
+	const bool bHairEnable = IsHairStrandsEnable(Scene->GetShaderPlatform()) && bIsViewCompatible;
 
+	TRefCountPtr<IPooledRenderTarget> ForwardScreenSpaceShadowMask;
+	TRefCountPtr<IPooledRenderTarget> ForwardScreenSpaceShadowMask_Hair;
 	if (IsForwardShadingEnabled(ShaderPlatform))
 	{
-		RenderForwardShadingShadowProjections(RHICmdList, ForwardScreenSpaceShadowMask);
+		if (bHairEnable)
+		{
+			RenderHairBasePass(RHICmdList, Scene, SceneContext, Views, HairClusterData, HairDatasStorage);
+			HairDatas = &HairDatasStorage;
+		}
+
+		RenderForwardShadingShadowProjections(RHICmdList, ForwardScreenSpaceShadowMask, ForwardScreenSpaceShadowMask_Hair, HairDatas);
 	}
 
 	// only temporarily available after early z pass and until base pass
@@ -2230,49 +2241,11 @@ void FDeferredShadingSceneRenderer::Render(FRHICommandListImmediate& RHICmdList)
 
 	checkSlow(RHICmdList.IsOutsideRenderPass());
 
-	// #hair_todo: Add multi-view
-	FHairStrandsDatas* HairDatas = nullptr;
-	FHairStrandsDatas HairDatasStorage;
-	const bool bIsViewCompatible = Views.Num() > 0 && Views[0].Family->ViewMode == VMI_Lit; 
-	if (IsHairStrandsEnable(Scene->GetShaderPlatform()) && bIsViewCompatible)
+	// Hair base pass for deferred shading
+	if (bHairEnable && !HairDatas)
 	{
-		SCOPED_GPU_STAT(RHICmdList, HairRendering);
-		HairDatasStorage.MacroGroupsPerViews = CreateHairStrandsMacroGroups(RHICmdList, Scene, Views);
-
-		// Culling/LOD pass for DOM and Voxelisation altogether
-		FHairCullingParams CullingParams;
-		CullingParams.bShadowViewMode = true;
-		CullingParams.bCullingProcessSkipped = false;
-		ComputeHairStrandsClustersCulling(RHICmdList, *GetGlobalShaderMap(FeatureLevel), Views, CullingParams, HairClusterData);
-
-		ServiceLocalQueue();
-
-		// Voxelization and Deep Opacity Maps
-		VoxelizeHairStrands(RHICmdList, Scene, Views, HairDatasStorage.MacroGroupsPerViews);
-		RenderHairStrandsDeepShadows(RHICmdList, Scene, Views, HairDatasStorage.MacroGroupsPerViews);
-
-		ServiceLocalQueue();
-
-		// Culling/LOD pass for visibility (must be done after HZB is generated)
-		CullingParams.bShadowViewMode = false;
-		ComputeHairStrandsClustersCulling(RHICmdList, *GetGlobalShaderMap(FeatureLevel), Views, CullingParams, HairClusterData);
-		// Hair visibility pass
-		HairDatasStorage.HairVisibilityViews = RenderHairStrandsVisibilityBuffer(RHICmdList, Scene, Views, SceneContext.GBufferB, SceneContext.GetSceneColor(), SceneContext.SceneDepthZ, SceneContext.SceneVelocity, HairDatasStorage.MacroGroupsPerViews);
-		// Reset indirect draw buffer
-		ResetHairStrandsClusterToLOD0(RHICmdList, *GetGlobalShaderMap(FeatureLevel), HairClusterData);
-
-		ServiceLocalQueue();
-
+		RenderHairBasePass(RHICmdList, Scene, SceneContext, Views, HairClusterData, HairDatasStorage);
 		HairDatas = &HairDatasStorage;
-
-		if (SceneContext.bScreenSpaceAOIsValid && SceneContext.ScreenSpaceAO)
-		{
-			RenderHairStrandsAmbientOcclusion(
-				RHICmdList,
-				Views,
-				HairDatas,
-				SceneContext.ScreenSpaceAO);
-		}
 	}
 
 	// Render lighting.
@@ -2381,7 +2354,11 @@ void FDeferredShadingSceneRenderer::Render(FRHICommandListImmediate& RHICmdList)
 #endif // RHI_RAYTRACING
 		ServiceLocalQueue();
 	}
-
+	else if (HairDatas)
+	{
+		RenderLightsForHair(RHICmdList, SortedLightSet, HairDatas, ForwardScreenSpaceShadowMask_Hair);
+		RenderDeferredReflectionsAndSkyLightingHair(RHICmdList, HairDatas);
+	}
 	checkSlow(RHICmdList.IsOutsideRenderPass());
 
 	const bool bShouldRenderSingleLayerWater = ShouldRenderSingleLayerWater(Views, ViewFamily.EngineShowFlags);
diff --git a/Engine/Source/Runtime/Renderer/Private/DeferredShadingRenderer.h b/Engine/Source/Runtime/Renderer/Private/DeferredShadingRenderer.h
--- a/Engine/Source/Runtime/Renderer/Private/DeferredShadingRenderer.h
+++ b/Engine/Source/Runtime/Renderer/Private/DeferredShadingRenderer.h
@@ -228,6 +228,7 @@ private:
 
 	/** Renders sky lighting and reflections that can be done in a deferred pass. */
 	void RenderDeferredReflectionsAndSkyLighting(FRHICommandListImmediate& RHICmdList, TRefCountPtr<IPooledRenderTarget>& DynamicBentNormalAO, TRefCountPtr<IPooledRenderTarget>& VelocityRT, struct FHairStrandsDatas* HairDatas);
+	void RenderDeferredReflectionsAndSkyLightingHair(FRHICommandListImmediate& RHICmdList, struct FHairStrandsDatas* HairDatas);
 
 	/** Computes DFAO, modulates it to scene color (which is assumed to contain diffuse indirect lighting), and stores the output bent normal for use occluding specular. */
 	void RenderDFAOAsIndirectShadowing(
@@ -382,7 +383,7 @@ private:
 	bool RenderShadowProjections(FRHICommandListImmediate& RHICmdList, const FLightSceneInfo* LightSceneInfo, IPooledRenderTarget* ScreenShadowMaskTexture, IPooledRenderTarget* ScreenShadowMaskSubPixelTexture, const struct FHairStrandsDatas* HairDatas, bool& bInjectedTranslucentVolume);
 
 	/** Render shadow projections when forward rendering. */
-	void RenderForwardShadingShadowProjections(FRHICommandListImmediate& RHICmdList, TRefCountPtr<IPooledRenderTarget>& ForwardScreenSpaceShadowMask);
+	void RenderForwardShadingShadowProjections(FRHICommandListImmediate& RHICmdList, TRefCountPtr<IPooledRenderTarget>& ForwardScreenSpaceShadowMask, TRefCountPtr<IPooledRenderTarget>& ForwardScreenSpaceShadowMaskSubPixel, const struct FHairStrandsDatas* InHairDatas);
 
 	/**
 	  * Used by RenderLights to render a light function to the attenuation buffer.
@@ -413,6 +414,7 @@ private:
 	  * @return true if anything got rendered
 	  */
 	void RenderLight(FRHICommandList& RHICmdList, const FLightSceneInfo* LightSceneInfo, IPooledRenderTarget* ScreenShadowMaskTexture, const struct FHairStrandsVisibilityViews* InHairVisibilityViews, bool bRenderOverlap, bool bIssueDrawEvent);
+	void RenderLightsForHair(FRHICommandListImmediate& RHICmdList, FSortedLightSetSceneInfo &SortedLightSet, const FHairStrandsDatas* HairDatas, TRefCountPtr<IPooledRenderTarget>& InScreenShadowMaskSubPixelTexture);
 
 	/** Specialized version of RenderLight for hair (run lighting evaluation on at sub-pixel rate, without depth bound) */
 	void RenderLightForHair(FRHICommandList& RHICmdList, const FLightSceneInfo* LightSceneInfo, IPooledRenderTarget* ScreenShadowMaskSubPixelTexture, struct FHairStrandsTransmittanceMaskData* InTransmittanceMaskData, const struct FHairStrandsVisibilityViews* InHairVisibilityViews);
diff --git a/Engine/Source/Runtime/Renderer/Private/HairStrands/HairStrandsRendering.cpp b/Engine/Source/Runtime/Renderer/Private/HairStrands/HairStrandsRendering.cpp
--- a/Engine/Source/Runtime/Renderer/Private/HairStrands/HairStrandsRendering.cpp
+++ b/Engine/Source/Runtime/Renderer/Private/HairStrands/HairStrandsRendering.cpp
@@ -1,3 +1,61 @@
 // Copyright Epic Games, Inc. All Rights Reserved.
 
-#include "HairStrandsRendering.h"
\ No newline at end of file
+#include "HairStrandsRendering.h"
+#include "SceneRendering.h"
+#include "ScenePrivate.h"
+
+void ServiceLocalQueue();
+void RenderHairBasePass(
+	FRHICommandListImmediate& RHICmdList,
+	FScene* Scene,
+	FSceneRenderTargets& SceneContext,
+	TArray<FViewInfo>& Views,
+	FHairStrandClusterData HairClusterData,
+	FHairStrandsDatas& OutHairDatas)
+{
+	// #hair_todo: Add multi-view
+	const bool bIsViewCompatible = Views.Num() > 0 && Views[0].Family->ViewMode == VMI_Lit;
+	if (IsHairStrandsEnable(Scene->GetShaderPlatform()) && bIsViewCompatible)
+	{
+		const ERHIFeatureLevel::Type FeatureLevel = Scene->GetFeatureLevel();
+
+		//SCOPED_GPU_STAT(RHICmdList, HairRendering);
+		OutHairDatas.MacroGroupsPerViews = CreateHairStrandsMacroGroups(RHICmdList, Scene, Views);
+
+		// Culling/LOD pass for DOM and Voxelisation altogether
+		FHairCullingParams CullingParams;
+		CullingParams.bShadowViewMode = true;
+		CullingParams.bCullingProcessSkipped = false;
+		ComputeHairStrandsClustersCulling(RHICmdList, *GetGlobalShaderMap(FeatureLevel), Views, CullingParams, HairClusterData);
+
+		ServiceLocalQueue();
+
+		// Voxelization and Deep Opacity Maps
+		VoxelizeHairStrands(RHICmdList, Scene, Views, OutHairDatas.MacroGroupsPerViews);
+		RenderHairStrandsDeepShadows(RHICmdList, Scene, Views, OutHairDatas.MacroGroupsPerViews);
+
+		ServiceLocalQueue();
+
+		// Culling/LOD pass for visibility (must be done after HZB is generated)
+		CullingParams.bShadowViewMode = false;
+		ComputeHairStrandsClustersCulling(RHICmdList, *GetGlobalShaderMap(FeatureLevel), Views, CullingParams, HairClusterData);
+
+		// Hair visibility pass
+		TRefCountPtr<IPooledRenderTarget> SceneColor = SceneContext.IsSceneColorAllocated() ? SceneContext.GetSceneColor() : nullptr;
+		OutHairDatas.HairVisibilityViews = RenderHairStrandsVisibilityBuffer(RHICmdList, Scene, Views, SceneContext.GBufferB, SceneColor, SceneContext.SceneDepthZ, SceneContext.SceneVelocity, OutHairDatas.MacroGroupsPerViews);
+
+		// Reset indirect draw buffer
+		ResetHairStrandsClusterToLOD0(RHICmdList, *GetGlobalShaderMap(FeatureLevel), HairClusterData);
+
+		ServiceLocalQueue();
+
+		if (SceneContext.bScreenSpaceAOIsValid && SceneContext.ScreenSpaceAO)
+		{
+			RenderHairStrandsAmbientOcclusion(
+				RHICmdList,
+				Views,
+				&OutHairDatas,
+				SceneContext.ScreenSpaceAO);
+		}
+	}
+}
diff --git a/Engine/Source/Runtime/Renderer/Private/HairStrands/HairStrandsRendering.h b/Engine/Source/Runtime/Renderer/Private/HairStrands/HairStrandsRendering.h
--- a/Engine/Source/Runtime/Renderer/Private/HairStrands/HairStrandsRendering.h
+++ b/Engine/Source/Runtime/Renderer/Private/HairStrands/HairStrandsRendering.h
@@ -43,4 +43,12 @@ void RunHairStrandsInterpolation(
 	const struct FShaderDrawDebugData* DebugShaderData,
 	FGlobalShaderMap* ShaderMap, 
 	EHairStrandsInterpolationType Type,
-	FHairStrandClusterData* ClusterData);
\ No newline at end of file
+	FHairStrandClusterData* ClusterData);
+
+void RenderHairBasePass(
+	FRHICommandListImmediate& RHICmdList,
+	FScene* Scene,
+	FSceneRenderTargets& SceneContext,
+	TArray<FViewInfo>& Views,
+	FHairStrandClusterData HairClusterData,
+	FHairStrandsDatas& OutHairDatas);
\ No newline at end of file
diff --git a/Engine/Source/Runtime/Renderer/Private/HairStrands/HairStrandsTransmittance.cpp b/Engine/Source/Runtime/Renderer/Private/HairStrands/HairStrandsTransmittance.cpp
--- a/Engine/Source/Runtime/Renderer/Private/HairStrands/HairStrandsTransmittance.cpp
+++ b/Engine/Source/Runtime/Renderer/Private/HairStrands/HairStrandsTransmittance.cpp
@@ -278,7 +278,6 @@ class FDeepShadowMaskPS : public FGlobalShader
 	using FPermutationDomain = TShaderPermutationDomain<FOpaqueMaskType, FKernelType>;
 
 	BEGIN_SHADER_PARAMETER_STRUCT(FParameters, )
-		SHADER_PARAMETER_STRUCT_INCLUDE(FSceneTextureParameters, SceneTextures)
 		SHADER_PARAMETER_STRUCT_INCLUDE(ShaderDrawDebug::FShaderDrawDebugParameters, ShaderDrawParameters)
 		SHADER_PARAMETER_STRUCT_INCLUDE(ShaderPrint::FShaderParameters, ShaderPrintParameters)
 		                         
@@ -286,7 +285,7 @@ class FDeepShadowMaskPS : public FGlobalShader
 		SHADER_PARAMETER(uint32, DeepShadow_SlotIndex)
 		SHADER_PARAMETER(FIntPoint, DeepShadow_SlotResolution)
 		SHADER_PARAMETER(FMatrix, DeepShadow_CPUWorldToLightTransform)
-		SHADER_PARAMETER(uint32, DeepShadow_IsWholeSceneLight)
+		SHADER_PARAMETER(uint32, bIsWholeSceneLight)
 		SHADER_PARAMETER(float, DeepShadow_DepthBiasScale)
 		SHADER_PARAMETER(float, DeepShadow_DensityScale)
 		SHADER_PARAMETER(uint32, DeepShadow_bIsGPUDriven)
@@ -301,6 +300,7 @@ class FDeepShadowMaskPS : public FGlobalShader
 		SHADER_PARAMETER_RDG_TEXTURE(Texture2D, DeepShadow_FrontDepthTexture)
 		SHADER_PARAMETER_RDG_TEXTURE(Texture2D, DeepShadow_DomTexture)
 		SHADER_PARAMETER_RDG_TEXTURE(Texture2D, CategorizationTexture)
+		SHADER_PARAMETER_RDG_TEXTURE(Texture2D, SceneDepthTexture)
 		SHADER_PARAMETER_SAMPLER(SamplerState, LinearSampler)
 		SHADER_PARAMETER_SAMPLER(SamplerState, ShadowSampler)
 
@@ -319,13 +319,13 @@ IMPLEMENT_GLOBAL_SHADER(FDeepShadowMaskPS, "/Engine/Private/HairStrands/HairStra
 struct FDeepShadowOpaqueParams
 {
 	FRDGTextureRef	CategorizationTexture = nullptr;
+	bool			bIsWholeSceneLight = false;
 
 	FRDGBufferSRVRef DeepShadow_WorldToLightTransformBuffer = nullptr;
 	FMatrix			DeepShadow_CPUWorldToLightTransform;
 	FIntRect		DeepShadow_AtlasRect;
 	FRDGTextureRef	DeepShadow_FrontDepthTexture = nullptr;
 	FRDGTextureRef	DeepShadow_LayerTexture = nullptr;
-	bool			DeepShadow_bIsWholeSceneLight = false;
 	bool			DeepShadow_bIsGPUDriven = false;
 	float			DeepShadow_DepthBiasScale = 1;
 	float			DeepShadow_DensityScale = 1;
@@ -340,7 +340,7 @@ struct FDeepShadowOpaqueParams
 // Opaque mask
 static void AddDeepShadowOpaqueMaskPass(
 	FRDGBuilder& GraphBuilder,
-	const FSceneTextureParameters& SceneTextures,
+	const FRDGTextureRef& SceneDepthTexture,
 	const FViewInfo& View,
 	const FHairOpaqueMaskType HairOpaqueMaskType,
 	const FDeepShadowOpaqueParams& Params,
@@ -350,14 +350,14 @@ static void AddDeepShadowOpaqueMaskPass(
 
 	FDeepShadowMaskPS::FParameters* Parameters = GraphBuilder.AllocParameters<FDeepShadowMaskPS::FParameters>();
 	Parameters->ViewUniformBuffer = View.ViewUniformBuffer;
-	Parameters->SceneTextures = SceneTextures;
+	Parameters->SceneDepthTexture = SceneDepthTexture;
 	Parameters->DeepShadow_CPUWorldToLightTransform = Params.DeepShadow_CPUWorldToLightTransform;
 	Parameters->DeepShadow_FrontDepthTexture = Params.DeepShadow_FrontDepthTexture;
 	Parameters->DeepShadow_DomTexture = Params.DeepShadow_LayerTexture;
 	Parameters->CategorizationTexture = Params.CategorizationTexture;
 	Parameters->LinearSampler = TStaticSamplerState<SF_Bilinear, AM_Clamp, AM_Clamp, AM_Clamp>::GetRHI();
 	Parameters->ShadowSampler = TStaticSamplerState<SF_Point, AM_Clamp, AM_Clamp, AM_Clamp>::GetRHI();
-	Parameters->DeepShadow_IsWholeSceneLight = Params.DeepShadow_bIsWholeSceneLight ? 1 : 0;
+	Parameters->bIsWholeSceneLight = Params.bIsWholeSceneLight ? 1 : 0;
 	Parameters->DeepShadow_DepthBiasScale = Params.DeepShadow_DepthBiasScale;
 	Parameters->DeepShadow_DensityScale = Params.DeepShadow_DensityScale;
 	Parameters->RenderTargets[0] = FRenderTargetBinding(OutShadowMask, ERenderTargetLoadAction::ELoad);
@@ -398,7 +398,7 @@ static void AddDeepShadowOpaqueMaskPass(
 	PermutationVector.Set<FDeepShadowMaskPS::FOpaqueMaskType>(HairOpaqueMaskType);
 	PermutationVector.Set<FDeepShadowMaskPS::FKernelType>(FMath::Clamp(GStrandHairShadowMaskKernelType, 0, 4));
 
-	const FIntPoint OutputResolution = SceneTextures.SceneDepthBuffer->Desc.Extent;
+	const FIntPoint OutputResolution = SceneDepthTexture->Desc.Extent;
 	TShaderMapRef<FPostProcessVS> VertexShader(View.ShaderMap);
 	TShaderMapRef<FDeepShadowMaskPS> PixelShader(View.ShaderMap, PermutationVector);
 	const FGlobalShaderMap* GlobalShaderMap = View.ShaderMap;
@@ -626,14 +626,10 @@ static void RenderHairStrandsShadowMask(
 	FSceneRenderTargets& SceneTargets = FSceneRenderTargets::Get(RHICmdList);
 	FRDGBuilder GraphBuilder(RHICmdList);
 
-	// Note: GbufferB.a store the shading model on the 4 lower bits (MATERIAL_SHADINGMODEL_HAIR)
-	FSceneTextureParameters SceneTextures;
-	SetupSceneTextureParameters(GraphBuilder, &SceneTextures);
-
+	FRDGTextureRef SceneDepthTexture = GraphBuilder.RegisterExternalTexture(SceneTargets.SceneDepthZ, TEXT("SceneDephtTexture"));
 	FRDGTextureRef Categorization = RegisterExternalTextureWithFallback(GraphBuilder, InVisibilityData.CategorizationTexture, GSystemTextures.BlackDummy);
 	FRDGTextureRef OutShadowMask = GraphBuilder.RegisterExternalTexture(ScreenShadowMaskTexture, TEXT("ScreenShadowMaskTexture"));
 
-
 	bool bHasDeepShadow = false;
 	if (!IsHairStrandsForVoxelTransmittanceAndShadowEnable())
 	{
@@ -667,13 +663,13 @@ static void RenderHairStrandsShadowMask(
 				Params.DeepShadow_AtlasRect = DomData.AtlasRect;
 				Params.DeepShadow_FrontDepthTexture = DeepShadowDepth;
 				Params.DeepShadow_LayerTexture = DeepShadowLayers;
-				Params.DeepShadow_bIsWholeSceneLight = bIsWholeSceneLight;
+				Params.bIsWholeSceneLight = bIsWholeSceneLight;
 				Params.DeepShadow_DepthBiasScale = GetDeepShadowDepthBiasScale();
 				Params.DeepShadow_DensityScale = GetDeepShadowDensityScale();
 
 				AddDeepShadowOpaqueMaskPass(
 					GraphBuilder,
-					SceneTextures,
+					SceneDepthTexture,
 					View,
 					FHairOpaqueMaskType_DeepShadow,
 					Params,
@@ -696,11 +692,12 @@ static void RenderHairStrandsShadowMask(
 			Params.Voxel_LightDirection = LightSceneInfo->Proxy->GetDirection();
 			Params.Voxel_LightPosition = FVector4(LightSceneInfo->Proxy->GetPosition(), LightSceneInfo->Proxy->GetLightType() == ELightComponentType::LightType_Directional ? 0 : 1);
 			Params.Voxel_MacroGroupId = MacroGroupData.MacroGroupId;
+			Params.bIsWholeSceneLight = bIsWholeSceneLight ? 1 : 0;
 
 			Params.Voxel_VirtualVoxel = &InMacroGroupDatas.VirtualVoxelResources;
 			AddDeepShadowOpaqueMaskPass(
 				GraphBuilder,
-				SceneTextures,
+				SceneDepthTexture,
 				View,
 				FHairOpaqueMaskType_VirtualVoxel,
 				Params,
diff --git a/Engine/Source/Runtime/Renderer/Private/HairStrands/HairStrandsVisibility.cpp b/Engine/Source/Runtime/Renderer/Private/HairStrands/HairStrandsVisibility.cpp
--- a/Engine/Source/Runtime/Renderer/Private/HairStrands/HairStrandsVisibility.cpp
+++ b/Engine/Source/Runtime/Renderer/Private/HairStrands/HairStrandsVisibility.cpp
@@ -2539,7 +2539,6 @@ FHairStrandsVisibilityViews RenderHairStrandsVisibilityBuffer(
 	SCOPED_GPU_STAT(RHICmdList, HairStrandsVisibility);
 
 	FHairStrandsVisibilityViews Output;
-
 	for (int32 ViewIndex = 0; ViewIndex < Views.Num(); ViewIndex++)
 	{
 		const FViewInfo& View = Views[ViewIndex];
@@ -2557,16 +2556,18 @@ FHairStrandsVisibilityViews RenderHairStrandsVisibilityBuffer(
 			// Use the scene color for computing target resolution as the View.ViewRect, 
 			// doesn't include the actual resolution padding which make buffer size 
 			// mismatch, and create artifact (e.g. velocity computation)
-			const FIntPoint Resolution = InSceneColorTexture->GetDesc().Extent;
+			check(InSceneDepthTexture);
+			const FIntPoint Resolution = InSceneDepthTexture->GetDesc().Extent;
 
 
 			FRDGBuilder GraphBuilder(RHICmdList);
 			FSceneRenderTargets& SceneContext = FSceneRenderTargets::Get(RHICmdList);
 			FRDGTextureRef SceneGBufferBTexture = GraphBuilder.TryRegisterExternalTexture(InSceneGBufferBTexture, TEXT("SceneGBufferBTexture"));
-			FRDGTextureRef SceneColorTexture = GraphBuilder.RegisterExternalTexture(InSceneColorTexture, TEXT("SceneColorTexture"));
+			FRDGTextureRef SceneColorTexture = GraphBuilder.TryRegisterExternalTexture(InSceneColorTexture, TEXT("SceneColorTexture"));
 			FRDGTextureRef SceneDepthTexture = GraphBuilder.RegisterExternalTexture(InSceneDepthTexture, TEXT("SceneDepthTexture"));
-			FRDGTextureRef SceneVelocityTexture = InSceneVelocityTexture ? GraphBuilder.RegisterExternalTexture(InSceneVelocityTexture, TEXT("SceneVelocityTexture")) : nullptr;
+			FRDGTextureRef SceneVelocityTexture = GraphBuilder.TryRegisterExternalTexture(InSceneVelocityTexture, TEXT("SceneVelocityTexture"));
 
+			const bool bRunColorAndDepthPatching = SceneGBufferBTexture && SceneColorTexture;
 			const EHairVisibilityRenderMode RenderMode = GetHairVisibilityRenderMode();
 			check(RenderMode == HairVisibilityRenderMode_MSAA || RenderMode == HairVisibilityRenderMode_PPLL);
 
@@ -2766,13 +2767,16 @@ FHairStrandsVisibilityViews RenderHairStrandsVisibilityBuffer(
 				// * black color into the scene color
 				// * closest depth
 				// * unlit shading model ID 
-				AddHairVisibilityColorAndDepthPatchPass(
-					GraphBuilder,
-					View,
-					CategorizationTexture,
-					SceneGBufferBTexture,
-					SceneColorTexture,
-					SceneDepthTexture);
+				if (bRunColorAndDepthPatching)
+				{
+					AddHairVisibilityColorAndDepthPatchPass(
+						GraphBuilder,
+						View,
+						CategorizationTexture,
+						SceneGBufferBTexture,
+						SceneColorTexture,
+						SceneDepthTexture);
+				}
 			}
 			else if (RenderMode == HairVisibilityRenderMode_PPLL)
 			{
@@ -2825,13 +2829,16 @@ FHairStrandsVisibilityViews RenderHairStrandsVisibilityBuffer(
 					GraphBuilder.QueueTextureExtraction(NodeCounter, &VisibilityData.NodeCount);
 				}
 
-				AddHairVisibilityColorAndDepthPatchPass(
-					GraphBuilder,
-					View,
-					CategorizationTexture,
-					SceneGBufferBTexture,
-					SceneColorTexture,
-					SceneDepthTexture);
+				if (bRunColorAndDepthPatching)
+				{
+					AddHairVisibilityColorAndDepthPatchPass(
+						GraphBuilder,
+						View,
+						CategorizationTexture,
+						SceneGBufferBTexture,
+						SceneColorTexture,
+						SceneDepthTexture);
+				}
 
 				// Allocate buffer for storing all the light samples
 				FRDGTextureRef SampleLightingBuffer = AddClearLightSamplePass(GraphBuilder, &View, VisibilityData.MaxNodeCount, NodeCounter);
diff --git a/Engine/Source/Runtime/Renderer/Private/IndirectLightRendering.cpp b/Engine/Source/Runtime/Renderer/Private/IndirectLightRendering.cpp
--- a/Engine/Source/Runtime/Renderer/Private/IndirectLightRendering.cpp
+++ b/Engine/Source/Runtime/Renderer/Private/IndirectLightRendering.cpp
@@ -927,3 +927,52 @@ void FDeferredShadingSceneRenderer::RenderDeferredReflectionsAndSkyLighting(FRHI
 
 	ResolveSceneColor(RHICmdList);
 }
+
+void FDeferredShadingSceneRenderer::RenderDeferredReflectionsAndSkyLightingHair(
+	FRHICommandListImmediate& RHICmdList,
+	FHairStrandsDatas* HairDatas)
+{
+	check(RHICmdList.IsOutsideRenderPass());
+
+	if (ViewFamily.EngineShowFlags.VisualizeLightCulling || !ViewFamily.EngineShowFlags.Lighting)
+	{
+		return;
+	}
+
+	// If we're currently capturing a reflection capture, output SpecularColor * IndirectIrradiance for metals so they are not black in reflections,
+	// Since we don't have multiple bounce specular reflections
+	bool bReflectionCapture = false;
+	for (int32 ViewIndex = 0, Num = Views.Num(); ViewIndex < Num; ViewIndex++)
+	{
+		const FViewInfo& View = Views[ViewIndex];
+		bReflectionCapture = bReflectionCapture || View.bIsReflectionCapture;
+	}
+
+	if (bReflectionCapture)
+	{
+		// if we are rendering a reflection capture then we can skip this pass entirely (no reflection and no sky contribution evaluated in this pass)
+		return;
+	}
+
+	// The specular sky light contribution is also needed by RT Reflections as a fallback.
+	const bool bSkyLight = Scene->SkyLight
+		&& Scene->SkyLight->ProcessedTexture
+		&& !Scene->SkyLight->bHasStaticLighting;
+
+	bool bDynamicSkyLight = ShouldRenderDeferredDynamicSkyLight(Scene, ViewFamily);
+	bool bApplySkyShadowing = false;
+	const bool bReflectionEnv = ShouldDoReflectionEnvironment();
+
+	uint32 ViewIndex = 0;
+	for (FViewInfo& View : Views)
+	{
+		const uint32 CurrentViewIndex = ViewIndex++;
+		const bool bIsHairSkyLightingEnabled = HairDatas && (bSkyLight || bDynamicSkyLight || bReflectionEnv);
+		if (bIsHairSkyLightingEnabled)
+		{
+			FRDGBuilder GraphBuilder(RHICmdList);
+			RenderHairStrandsEnvironmentLighting(GraphBuilder, CurrentViewIndex, Views, HairDatas);
+			GraphBuilder.Execute();
+		}
+	}
+}
\ No newline at end of file
diff --git a/Engine/Source/Runtime/Renderer/Private/LightGridInjection.cpp b/Engine/Source/Runtime/Renderer/Private/LightGridInjection.cpp
--- a/Engine/Source/Runtime/Renderer/Private/LightGridInjection.cpp
+++ b/Engine/Source/Runtime/Renderer/Private/LightGridInjection.cpp
@@ -745,10 +745,11 @@ void FDeferredShadingSceneRenderer::ComputeLightGrid(FRHICommandListImmediate& R
 
 
 
-void FDeferredShadingSceneRenderer::RenderForwardShadingShadowProjections(FRHICommandListImmediate& RHICmdList, TRefCountPtr<IPooledRenderTarget>& ForwardScreenSpaceShadowMask)
+void FDeferredShadingSceneRenderer::RenderForwardShadingShadowProjections(FRHICommandListImmediate& RHICmdList, TRefCountPtr<IPooledRenderTarget>& ForwardScreenSpaceShadowMask, TRefCountPtr<IPooledRenderTarget>& ForwardScreenSpaceShadowMaskSubPixel, const FHairStrandsDatas* InHairDatas)
 {
 	check(RHICmdList.IsOutsideRenderPass());
 
+	const bool bIsHairEnable = InHairDatas;
 	bool bScreenShadowMaskNeeded = false;
 
 	for (TSparseArray<FLightSceneInfoCompact>::TConstIterator LightIt(Scene->Lights); LightIt; ++LightIt)
@@ -765,15 +766,29 @@ void FDeferredShadingSceneRenderer::RenderForwardShadingShadowProjections(FRHICo
 		CSV_SCOPED_TIMING_STAT_EXCLUSIVE(RenderForwardShadingShadowProjections);
 		FSceneRenderTargets& SceneRenderTargets = FSceneRenderTargets::Get(RHICmdList);
 		SceneRenderTargets.AllocateScreenShadowMask(RHICmdList, ForwardScreenSpaceShadowMask);
+		if (bIsHairEnable)
+		{
+			SceneRenderTargets.AllocateScreenShadowMask(RHICmdList, ForwardScreenSpaceShadowMaskSubPixel);
+		}
 
 		SCOPED_DRAW_EVENT(RHICmdList, ShadowProjectionOnOpaque);
 		SCOPED_GPU_STAT(RHICmdList, ShadowProjection);
 
 		// All shadows render with min blending
-		FRHIRenderPassInfo RPInfo(ForwardScreenSpaceShadowMask->GetRenderTargetItem().TargetableTexture, ERenderTargetActions::Clear_Store);
-		TransitionRenderPassTargets(RHICmdList, RPInfo);
-		RHICmdList.BeginRenderPass(RPInfo, TEXT("RenderForwardShadingShadowProjectionsClear"));
-		RHICmdList.EndRenderPass();
+		{
+			FRHIRenderPassInfo RPInfo(ForwardScreenSpaceShadowMask->GetRenderTargetItem().TargetableTexture, ERenderTargetActions::Clear_Store);
+			TransitionRenderPassTargets(RHICmdList, RPInfo);
+			RHICmdList.BeginRenderPass(RPInfo, TEXT("RenderForwardShadingShadowProjectionsClear"));
+			RHICmdList.EndRenderPass();
+		}
+
+		if (bIsHairEnable)
+		{
+			FRHIRenderPassInfo RPInfo(ForwardScreenSpaceShadowMaskSubPixel->GetRenderTargetItem().TargetableTexture, ERenderTargetActions::Clear_Store);
+			TransitionRenderPassTargets(RHICmdList, RPInfo);
+			RHICmdList.BeginRenderPass(RPInfo, TEXT("RenderForwardShadingShadowSubPixelProjectionsClear"));
+			RHICmdList.EndRenderPass();
+		}
 
 		// Note: all calls here will set up renderpasses internally.
 		// #todo-renderpasses might be worth refactoring all this and splitting into lists of draws for each renderpass
@@ -792,7 +807,12 @@ void FDeferredShadingSceneRenderer::RenderForwardShadingShadowProjections(FRHICo
 
 				if (VisibleLightInfo.ShadowsToProject.Num() > 0)
 				{
-					FSceneRenderer::RenderShadowProjections(RHICmdList, LightSceneInfo, ForwardScreenSpaceShadowMask, nullptr, true, false, nullptr);
+					FSceneRenderer::RenderShadowProjections(RHICmdList, LightSceneInfo, ForwardScreenSpaceShadowMask, ForwardScreenSpaceShadowMaskSubPixel, true, false, bIsHairEnable ? &InHairDatas->HairVisibilityViews : nullptr);
+
+					if (InHairDatas)
+					{
+						RenderHairStrandsShadowMask(RHICmdList, Views, LightSceneInfo, ForwardScreenSpaceShadowMask, InHairDatas);
+					}
 				}
 
 				RenderCapsuleDirectShadows(RHICmdList, *LightSceneInfo, ForwardScreenSpaceShadowMask, VisibleLightInfo.CapsuleShadowsToProject, true);
@@ -802,8 +822,12 @@ void FDeferredShadingSceneRenderer::RenderForwardShadingShadowProjections(FRHICo
 					RenderLightFunction(RHICmdList, LightSceneInfo, ForwardScreenSpaceShadowMask, true, true);
 				}
 			}
+			RHICmdList.CopyToResolveTarget(ForwardScreenSpaceShadowMask->GetRenderTargetItem().TargetableTexture, ForwardScreenSpaceShadowMask->GetRenderTargetItem().ShaderResourceTexture, FResolveParams(FResolveRect()));
+			if (bIsHairEnable)
+			{
+				RHICmdList.CopyToResolveTarget(ForwardScreenSpaceShadowMaskSubPixel->GetRenderTargetItem().TargetableTexture, ForwardScreenSpaceShadowMaskSubPixel->GetRenderTargetItem().ShaderResourceTexture, FResolveParams(FResolveRect()));
+			}
 		}
-		RHICmdList.CopyToResolveTarget(ForwardScreenSpaceShadowMask->GetRenderTargetItem().TargetableTexture, ForwardScreenSpaceShadowMask->GetRenderTargetItem().ShaderResourceTexture, FResolveParams(FResolveRect()));
 	}
 }
 
diff --git a/Engine/Source/Runtime/Renderer/Private/LightRendering.cpp b/Engine/Source/Runtime/Renderer/Private/LightRendering.cpp
--- a/Engine/Source/Runtime/Renderer/Private/LightRendering.cpp
+++ b/Engine/Source/Runtime/Renderer/Private/LightRendering.cpp
@@ -2284,6 +2284,45 @@ void FDeferredShadingSceneRenderer::RenderLightForHair(
 	}
 }
 
+// Forward lighting version for hair
+void FDeferredShadingSceneRenderer::RenderLightsForHair(FRHICommandListImmediate& RHICmdList, FSortedLightSetSceneInfo &SortedLightSet, const FHairStrandsDatas* HairDatas, TRefCountPtr<IPooledRenderTarget>& InScreenShadowMaskSubPixelTexture)
+{
+	const FSimpleLightArray &SimpleLights = SortedLightSet.SimpleLights;
+	const TArray<FSortedLightSceneInfo, SceneRenderingAllocator> &SortedLights = SortedLightSet.SortedLights;
+	const int32 AttenuationLightStart = SortedLightSet.AttenuationLightStart;
+	const int32 SimpleLightsEnd = SortedLightSet.SimpleLightsEnd;
+
+	const bool bUseHairLighting = HairDatas != nullptr;
+	if (ViewFamily.EngineShowFlags.DirectLighting && bUseHairLighting)
+	{
+		SCOPED_DRAW_EVENT(RHICmdList, DirectLighting);
+
+		//FSceneRenderTargets& SceneContext = FSceneRenderTargets::Get(RHICmdList);
+		for (int32 LightIndex = AttenuationLightStart; LightIndex < SortedLights.Num(); LightIndex++)
+		{
+			const FSortedLightSceneInfo& SortedLightInfo = SortedLights[LightIndex];
+			const FLightSceneInfo& LightSceneInfo = *SortedLightInfo.LightSceneInfo;
+			if (LightSceneInfo.Proxy)
+			{
+				TRefCountPtr<IPooledRenderTarget> ScreenShadowMaskSubPixelTexture = InScreenShadowMaskSubPixelTexture;
+
+				const bool bDrawDeepShadow = SortedLightInfo.SortKey.Fields.bShadowed;
+				FHairStrandsTransmittanceMaskData TransmittanceMaskData;
+				if (bDrawDeepShadow)
+				{
+					TransmittanceMaskData = RenderHairStrandsTransmittanceMask(RHICmdList, Views, &LightSceneInfo, HairDatas, ScreenShadowMaskSubPixelTexture);
+				}
+
+				RenderLightForHair(RHICmdList, 
+					&LightSceneInfo,
+					ScreenShadowMaskSubPixelTexture,
+					bDrawDeepShadow ? &TransmittanceMaskData : nullptr,
+					&HairDatas->HairVisibilityViews);
+			}
+		}
+	}
+}
+
 void FDeferredShadingSceneRenderer::RenderSimpleLightsStandardDeferred(FRHICommandListImmediate& RHICmdList, const FSimpleLightArray& SimpleLights)
 {
 	SCOPE_CYCLE_COUNTER(STAT_DirectLightRenderingTime);
diff --git a/Engine/Source/Runtime/Renderer/Private/ShadowRendering.h b/Engine/Source/Runtime/Renderer/Private/ShadowRendering.h
--- a/Engine/Source/Runtime/Renderer/Private/ShadowRendering.h
+++ b/Engine/Source/Runtime/Renderer/Private/ShadowRendering.h
@@ -836,9 +836,10 @@ public:
 		PerObjectShadowFadeStart.Bind(ParameterMap, TEXT("PerObjectShadowFadeStart"));
 		InvPerObjectShadowFadeLength.Bind(ParameterMap, TEXT("InvPerObjectShadowFadeLength"));
 		ShadowNearAndFarDepth.Bind(ParameterMap, TEXT("ShadowNearAndFarDepth"));
+		bCascadeUseFadePlane.Bind(ParameterMap, TEXT("bCascadeUseFadePlane"));
 	}
 
-	void Set(FRHICommandList& RHICmdList, FShader* Shader, const FSceneView& View, const FProjectedShadowInfo* ShadowInfo, const FHairStrandsVisibilityData* HairVisibilityData, bool bModulatedShadows)
+	void Set(FRHICommandList& RHICmdList, FShader* Shader, const FSceneView& View, const FProjectedShadowInfo* ShadowInfo, const FHairStrandsVisibilityData* HairVisibilityData, bool bModulatedShadows, bool bUseFadePlane)
 	{
 		FRHIPixelShader* ShaderRHI = RHICmdList.GetBoundPixelShader();
 
@@ -913,7 +914,7 @@ public:
 		SetShaderValue(RHICmdList, ShaderRHI, ProjectionDepthBias, FVector4(ShadowInfo->GetShaderDepthBias(), ShadowInfo->GetShaderSlopeDepthBias(), ShadowInfo->GetShaderReceiverDepthBias(), ShadowInfo->MaxSubjectZ - ShadowInfo->MinSubjectZ));
 		SetShaderValue(RHICmdList, ShaderRHI, FadePlaneOffset, ShadowInfo->CascadeSettings.FadePlaneOffset);
 
-		if(InvFadePlaneLength.IsBound())
+		if(InvFadePlaneLength.IsBound() && bUseFadePlane)
 		{
 			check(ShadowInfo->CascadeSettings.FadePlaneLength > 0);
 			SetShaderValue(RHICmdList, ShaderRHI, InvFadePlaneLength, 1.0f / ShadowInfo->CascadeSettings.FadePlaneLength);
@@ -945,6 +946,7 @@ public:
 			SliceNearAndFarDepth.X = DeviceZNear;
 			SliceNearAndFarDepth.Y = DeviceZFar;
 			SetShaderValue(RHICmdList, ShaderRHI, ShadowNearAndFarDepth, SliceNearAndFarDepth);
+			SetShaderValue(RHICmdList, ShaderRHI, bCascadeUseFadePlane, bUseFadePlane ? 1 : 0);
 		}
 
 		SetShaderValue(RHICmdList, ShaderRHI, PerObjectShadowFadeStart, ShadowInfo->PerObjectShadowFadeStart);
@@ -988,6 +990,7 @@ private:
 	LAYOUT_FIELD(FShaderParameter, PerObjectShadowFadeStart);
 	LAYOUT_FIELD(FShaderParameter, InvPerObjectShadowFadeLength);
 	LAYOUT_FIELD(FShaderParameter, ShadowNearAndFarDepth);
+	LAYOUT_FIELD(FShaderParameter, bCascadeUseFadePlane);
 };
 
 /**
@@ -1054,7 +1057,9 @@ public:
 
 		FShadowProjectionPixelShaderInterface::SetParameters(RHICmdList, ViewIndex, View, HairVisibilityData, ShadowInfo);
 
-		ProjectionParameters.Set(RHICmdList, this, View, ShadowInfo, HairVisibilityData, bModulatedShadows);
+		const bool bUseFadePlaneEnable = ShadowInfo->CascadeSettings.FadePlaneLength > 0;
+
+		ProjectionParameters.Set(RHICmdList, this, View, ShadowInfo, HairVisibilityData, bModulatedShadows, bUseFadePlaneEnable);
 		const FLightSceneProxy& LightProxy = *(ShadowInfo->GetLightSceneInfo().Proxy);
 
 		SetShaderValue(RHICmdList, ShaderRHI, ShadowFadeFraction, ShadowInfo->FadeAlphas[ViewIndex] );
