Fix CIS issue.

diff --git a/Engine/Source/Runtime/Renderer/Private/LightGridInjection.cpp b/Engine/Source/Runtime/Renderer/Private/LightGridInjection.cpp
--- a/Engine/Source/Runtime/Renderer/Private/LightGridInjection.cpp
+++ b/Engine/Source/Runtime/Renderer/Private/LightGridInjection.cpp
@@ -749,7 +749,7 @@ void FDeferredShadingSceneRenderer::RenderForwardShadingShadowProjections(FRHICo
 {
 	check(RHICmdList.IsOutsideRenderPass());
 
-	const bool bIsHairEnable = InHairDatas;
+	const bool bIsHairEnable = InHairDatas != nullptr;
 	bool bScreenShadowMaskNeeded = false;
 
 	for (TSparseArray<FLightSceneInfoCompact>::TConstIterator LightIt(Scene->Lights); LightIt; ++LightIt)
