Add back the UAV clear of VT feedback buffer that got lost somewhere in an old merge.
Fixes oversubscription of VT physical pages coming from page requests staying in the feedback buffer.

diff --git a/Engine/Source/Runtime/Renderer/Private/VT/VirtualTextureFeedback.cpp b/Engine/Source/Runtime/Renderer/Private/VT/VirtualTextureFeedback.cpp
--- a/Engine/Source/Runtime/Renderer/Private/VT/VirtualTextureFeedback.cpp
+++ b/Engine/Source/Runtime/Renderer/Private/VT/VirtualTextureFeedback.cpp
@@ -175,6 +175,11 @@ void FVirtualTextureFeedback::CreateResourceGPU( FRHICommandListImmediate& RHICm
 		FeedbackBufferUAV = RHICreateUnorderedAccessView(FeedbackBuffer, /*Format=*/ PF_R32_UINT);
 	}
 	
+	// Clear to default value
+	RHICmdList.TransitionResource(EResourceTransitionAccess::ERWNoBarrier, EResourceTransitionPipeline::EGfxToCompute, FeedbackBufferUAV);
+	RHICmdList.ClearUAVUint(FeedbackBufferUAV.GetReference(), FUintVector4(~0u, ~0u, ~0u, ~0u));
+	RHICmdList.TransitionResource(EResourceTransitionAccess::ERWNoBarrier, EResourceTransitionPipeline::EComputeToCompute, FeedbackBufferUAV);
+
 	FeedBackFences->Init(RHICmdList);
 }
 
