Readd support for 16bit virtual texture page tables.
This had been broken by a previous changelist (9751925)

diff --git a/Engine/Config/BaseEngine.ini b/Engine/Config/BaseEngine.ini
--- a/Engine/Config/BaseEngine.ini
+++ b/Engine/Config/BaseEngine.ini
@@ -2404,4 +2404,4 @@ DefaultChaosSolverActorClass=/Script/ChaosSolverEngine.ChaosSolverActor
 ; "r.VT.ListPhysicalPool" list in-depth details on the allocated physical pools
 ; The upper allocation limit for a pool is caluclated by matching a size and format from the list below or falling back to the default
 DefaultSizeInMegabyte=64
-+Pools=(SizeInMegabyte=64, MinTileSize=0, MaxTileSize=9999, Formats=(PF_DXT1, PF_DXT5))
++Pools=(SizeInMegabyte=64, MinTileSize=0, MaxTileSize=9999, Formats=(PF_DXT5, PF_DXT5))
diff --git a/Engine/Source/Runtime/Renderer/Private/VT/AllocatedVirtualTexture.cpp b/Engine/Source/Runtime/Renderer/Private/VT/AllocatedVirtualTexture.cpp
--- a/Engine/Source/Runtime/Renderer/Private/VT/AllocatedVirtualTexture.cpp
+++ b/Engine/Source/Runtime/Renderer/Private/VT/AllocatedVirtualTexture.cpp
@@ -85,13 +85,25 @@ FAllocatedVirtualTexture::FAllocatedVirtualTexture(FVirtualTextureSystem* InSyst
 		}
 	}
 
+	// Use 16bit page table entries if all physical spaces are small enough
+	bool bSupport16BitPageTable = true;
+	for (int32 Index = 0; Index < UniquePageTableLayers.Num(); ++Index)
+	{
+		if (!UniquePageTableLayers[Index].PhysicalSpace->DoesSupport16BitPageTable())
+		{
+			bSupport16BitPageTable = false;
+			break;
+		}
+	}
+
 	FVTSpaceDescription SpaceDesc;
 	SpaceDesc.Dimensions = InDesc.Dimensions;
 	SpaceDesc.NumPageTableLayers = UniquePageTableLayers.Num();
 	SpaceDesc.TileSize = InDesc.TileSize;
 	SpaceDesc.TileBorderSize = InDesc.TileBorderSize;
 	SpaceDesc.bPrivateSpace = InDesc.bPrivateSpace;
-	SpaceDesc.PageTableFormat = EVTPageTableFormat::UInt32;
+	SpaceDesc.PageTableFormat = bSupport16BitPageTable ? EVTPageTableFormat::UInt16 : EVTPageTableFormat::UInt32;
+
 	Space = InSystem->AcquireSpace(SpaceDesc, this);
 	SpaceID = Space->GetID();
 	PageTableFormat = Space->GetPageTableFormat();
