Do not cache mesh draw commands for ray tracing in preview worlds as we will not render ray tracing anyway. This fixes crashes with landscape thumbnails.

diff --git a/Engine/Source/Runtime/Renderer/Private/PrimitiveSceneInfo.cpp b/Engine/Source/Runtime/Renderer/Private/PrimitiveSceneInfo.cpp
--- a/Engine/Source/Runtime/Renderer/Private/PrimitiveSceneInfo.cpp
+++ b/Engine/Source/Runtime/Renderer/Private/PrimitiveSceneInfo.cpp
@@ -394,7 +394,7 @@ void FPrimitiveSceneInfo::CacheMeshDrawCommands(FRHICommandListImmediate& RHICmd
 	FGraphicsMinimalPipelineStateId::InitializePersistentIds();
 					
 #if RHI_RAYTRACING
-	if (IsRayTracingEnabled())
+	if (IsRayTracingEnabled() && !(Scene->World->WorldType == EWorldType::EditorPreview || Scene->World->WorldType == EWorldType::GamePreview))
 	{
 		FCachedRayTracingMeshCommandContext CommandContext(Scene->CachedRayTracingMeshCommands);
 		FRayTracingMeshProcessor RayTracingMeshProcessor(&CommandContext, Scene, nullptr);
