Fix incorrect shader usage for creating base color only runtime virtual texture thumbnail

diff --git a/Engine/Shaders/Private/VirtualTextureCompress.usf b/Engine/Shaders/Private/VirtualTextureCompress.usf
--- a/Engine/Shaders/Private/VirtualTextureCompress.usf
+++ b/Engine/Shaders/Private/VirtualTextureCompress.usf
@@ -177,7 +177,7 @@ void CompressBaseColorCS(
 	OutCompressTexture0_u2[ThreadId.xy + DestRect.xy] = CompressBC1Block_SRGB(BlockBaseColor);
 }
 
-/** Simple copy used for thumbnails. */
+/** Simple copy of base color used for thumbnails. */
 void CopyBaseColorPS(
 	in float4 InPosition : SV_POSITION, 
 	in noperspective float2 InTexCoord : TEXCOORD0, 
diff --git a/Engine/Source/Runtime/Renderer/Private/VT/RuntimeVirtualTextureRender.cpp b/Engine/Source/Runtime/Renderer/Private/VT/RuntimeVirtualTextureRender.cpp
--- a/Engine/Source/Runtime/Renderer/Private/VT/RuntimeVirtualTextureRender.cpp
+++ b/Engine/Source/Runtime/Renderer/Private/VT/RuntimeVirtualTextureRender.cpp
@@ -635,7 +635,7 @@ namespace RuntimeVirtualTexture
 		{}
 	};
 
-	IMPLEMENT_SHADER_TYPE(template<>, FShader_VirtualTextureCopy_PS< ERuntimeVirtualTextureMaterialType::BaseColor >, TEXT("/Engine/Private/VirtualTextureCompress.usf"), TEXT("CopyBaseColorNormalSpecularPS"), SF_Pixel);
+	IMPLEMENT_SHADER_TYPE(template<>, FShader_VirtualTextureCopy_PS< ERuntimeVirtualTextureMaterialType::BaseColor >, TEXT("/Engine/Private/VirtualTextureCompress.usf"), TEXT("CopyBaseColorPS"), SF_Pixel);
 	IMPLEMENT_SHADER_TYPE(template<>, FShader_VirtualTextureCopy_PS< ERuntimeVirtualTextureMaterialType::BaseColor_Normal_Specular >, TEXT("/Engine/Private/VirtualTextureCompress.usf"), TEXT("CopyBaseColorNormalSpecularPS"), SF_Pixel);
 	IMPLEMENT_SHADER_TYPE(template<>, FShader_VirtualTextureCopy_PS< ERuntimeVirtualTextureMaterialType::BaseColor_Normal_Specular_YCoCg >, TEXT("/Engine/Private/VirtualTextureCompress.usf"), TEXT("CopyBaseColorNormalSpecularYCoCgPS"), SF_Pixel);
 	IMPLEMENT_SHADER_TYPE(template<>, FShader_VirtualTextureCopy_PS< ERuntimeVirtualTextureMaterialType::BaseColor_Normal_Specular_Mask_YCoCg >, TEXT("/Engine/Private/VirtualTextureCompress.usf"), TEXT("CopyBaseColorNormalSpecularMaskYCoCgPS"), SF_Pixel);
