Enable SingleLayerWaterIndirectDraw for vulkan desktop.

diff --git a/Engine/Source/Runtime/Renderer/Private/SingleLayerWaterRendering.cpp b/Engine/Source/Runtime/Renderer/Private/SingleLayerWaterRendering.cpp
--- a/Engine/Source/Runtime/Renderer/Private/SingleLayerWaterRendering.cpp
+++ b/Engine/Source/Runtime/Renderer/Private/SingleLayerWaterRendering.cpp
@@ -131,7 +131,7 @@ bool ShouldRenderSingleLayerWaterSkippedRenderEditorNotification(const TArray<FV
 bool UseSingleLayerWaterIndirectDraw(EShaderPlatform ShaderPlatform)
 {
 	return IsFeatureLevelSupported(ShaderPlatform, ERHIFeatureLevel::SM5)
-		&& !IsSwitchPlatform(ShaderPlatform) && !IsVulkanPlatform(ShaderPlatform) && !IsMetalPlatform(ShaderPlatform); // Switch does not use tiling, Vulkan gives error with WaterTileCatergorisationCS usage of atomic, and Metal does not play nice, either.
+		&& !IsSwitchPlatform(ShaderPlatform) && !IsVulkanMobilePlatform(ShaderPlatform) && !IsMetalPlatform(ShaderPlatform); // Switch does not use tiling, Vulkan gives error with WaterTileCatergorisationCS usage of atomic, and Metal does not play nice, either.
 }
 
 
