Fix for sublevel loading.

diff --git a/Engine/Source/Runtime/Renderer/Private/RendererScene.cpp b/Engine/Source/Runtime/Renderer/Private/RendererScene.cpp
--- a/Engine/Source/Runtime/Renderer/Private/RendererScene.cpp
+++ b/Engine/Source/Runtime/Renderer/Private/RendererScene.cpp
@@ -3517,6 +3517,8 @@ void FScene::OnLevelAddedToWorld(FName LevelAddedName, UWorld* InWorld, bool bIs
 void FScene::OnLevelAddedToWorld_RenderThread(FName InLevelName)
 {
 	// Mark level primitives
+	TArray<FPrimitiveSceneInfo*> PrimitivesToAdd;
+
 	for (auto It = Primitives.CreateIterator(); It; ++It)
 	{
 		FPrimitiveSceneProxy* Proxy = (*It)->Proxy;
@@ -3528,10 +3530,12 @@ void FScene::OnLevelAddedToWorld_RenderThread(FName InLevelName)
 				// The only type of SceneProxy using this is landscape
 				(*It)->RemoveStaticMeshes();
 				Proxy->OnLevelAddedToWorld();
-				(*It)->AddStaticMeshes(FRHICommandListExecutor::GetImmediateCommandList(), this, Primitives);
+				PrimitivesToAdd.Add(*It);
 			}
 		}
 	}
+
+	FPrimitiveSceneInfo::AddStaticMeshes(FRHICommandListExecutor::GetImmediateCommandList(), this, PrimitivesToAdd);
 }
 
 void FScene::ProcessAtmosphereLightAddition_RenderThread(FLightSceneInfo* LightSceneInfo)
