Adding ray-traced reflections to single-layer water.

Same as CL 12214472.

diff --git a/Engine/Shaders/Private/RayTracing/RayTracingReflections.usf b/Engine/Shaders/Private/RayTracing/RayTracingReflections.usf
--- a/Engine/Shaders/Private/RayTracing/RayTracingReflections.usf
+++ b/Engine/Shaders/Private/RayTracing/RayTracingReflections.usf
@@ -93,6 +93,7 @@ uint UseReflectionCaptures;
 uint ReflectedShadowsType;
 uint ShouldDoDirectLighting;
 uint ShouldDoEmissiveAndIndirectLighting;
+uint ShouldReflectOnlyWater;
 uint UpscaleFactor;
 uint SortTileSize;
 uint2 RayTracingResolution;
@@ -330,7 +331,11 @@ RAY_TRACING_ENTRY_RAYGEN(RayTracingReflectionsRGS)
 	// Cutoff
 	// * Roughness value
 	// * Hair material is not for RT reflection at the moment. Their indirect contributuion is based on diffuse SH evaluation in ReflectionEnvironemenPixelShader.usf
-    const bool bIsValidPixel = (RoughnessFade > 0 || ClearCoatRoughnessFade > 0) && GBufferData.ShadingModelID != SHADINGMODELID_HAIR;
+    bool bIsValidPixel = (RoughnessFade > 0 || ClearCoatRoughnessFade > 0) && GBufferData.ShadingModelID != SHADINGMODELID_HAIR;
+	if (ShouldReflectOnlyWater)
+	{
+		bIsValidPixel = bIsValidPixel && (GBufferData.ShadingModelID == SHADINGMODELID_SINGLELAYERWATER);
+	}
 
     const uint LocalMaxBounces = DIM_HYBRID ? 1 : MaxBounces;
 
diff --git a/Engine/Source/Runtime/Renderer/Private/DeferredShadingRenderer.h b/Engine/Source/Runtime/Renderer/Private/DeferredShadingRenderer.h
--- a/Engine/Source/Runtime/Renderer/Private/DeferredShadingRenderer.h
+++ b/Engine/Source/Runtime/Renderer/Private/DeferredShadingRenderer.h
@@ -523,6 +523,7 @@ private:
 		const FSceneTextureParameters& SceneTextures,
 		const FViewInfo& View,
 		int32 SamplePerPixel,
+		bool bReflectWaterOnly,
 		float ResolutionFraction,
 		IScreenSpaceDenoiser::FReflectionsInputs* OutDenoiserInputs);
 
diff --git a/Engine/Source/Runtime/Renderer/Private/IndirectLightRendering.cpp b/Engine/Source/Runtime/Renderer/Private/IndirectLightRendering.cpp
--- a/Engine/Source/Runtime/Renderer/Private/IndirectLightRendering.cpp
+++ b/Engine/Source/Runtime/Renderer/Private/IndirectLightRendering.cpp
@@ -698,11 +698,14 @@ void FDeferredShadingSceneRenderer::RenderDeferredReflectionsAndSkyLighting(FRHI
 					RayTracingConfig.ResolutionFraction = 1.0f;
 				}
 
+				bool bReflectOnlyWater = false;
 				RenderRayTracingReflections(
 					GraphBuilder,
 					SceneTextures,
 					View,
-					RayTracingConfig.RayCountPerPixel, RayTracingConfig.ResolutionFraction,
+					RayTracingConfig.RayCountPerPixel,
+					bReflectOnlyWater,
+					RayTracingConfig.ResolutionFraction,
 					&DenoiserInputs);
 			}
 			else if (bScreenSpaceReflections)
diff --git a/Engine/Source/Runtime/Renderer/Private/RayTracing/RayTracingReflections.cpp b/Engine/Source/Runtime/Renderer/Private/RayTracing/RayTracingReflections.cpp
--- a/Engine/Source/Runtime/Renderer/Private/RayTracing/RayTracingReflections.cpp
+++ b/Engine/Source/Runtime/Renderer/Private/RayTracing/RayTracingReflections.cpp
@@ -239,6 +239,7 @@ class FRayTracingReflectionsRGS : public FGlobalShader
 		SHADER_PARAMETER(int32, ShouldDoDirectLighting)
 		SHADER_PARAMETER(int32, ReflectedShadowsType)
 		SHADER_PARAMETER(int32, ShouldDoEmissiveAndIndirectLighting)
+		SHADER_PARAMETER(int32, ShouldReflectOnlyWater)
 		SHADER_PARAMETER(int32, UpscaleFactor)
 		SHADER_PARAMETER(int32, SortTileSize)
 		SHADER_PARAMETER(FIntPoint, RayTracingResolution)
@@ -466,6 +467,7 @@ void FDeferredShadingSceneRenderer::RenderRayTracingReflections(
 	const FSceneTextureParameters& SceneTextures,
 	const FViewInfo& View,
 	int32 SamplePerPixel,
+	bool bReflectOnlyWater,
 	float ResolutionFraction,
 	IScreenSpaceDenoiser::FReflectionsInputs* OutDenoiserInputs)
 #if RHI_RAYTRACING
@@ -553,6 +555,7 @@ void FDeferredShadingSceneRenderer::RenderRayTracingReflections(
 	CommonParameters.ShouldDoDirectLighting = GRayTracingReflectionsDirectLighting;
 	CommonParameters.ReflectedShadowsType = GRayTracingReflectionsShadows > -1 ? GRayTracingReflectionsShadows : (int32)View.FinalPostProcessSettings.RayTracingReflectionsShadows;
 	CommonParameters.ShouldDoEmissiveAndIndirectLighting = GRayTracingReflectionsEmissiveAndIndirectLighting;
+	CommonParameters.ShouldReflectOnlyWater = bReflectOnlyWater;
 	CommonParameters.UpscaleFactor = UpscaleFactor;
 	CommonParameters.ReflectionMinRayDistance = FMath::Min(GRayTracingReflectionsMinRayDistance, GRayTracingReflectionsMaxRayDistance);
 	CommonParameters.ReflectionMaxRayDistance = GRayTracingReflectionsMaxRayDistance;
diff --git a/Engine/Source/Runtime/Renderer/Private/SceneRendering.h b/Engine/Source/Runtime/Renderer/Private/SceneRendering.h
--- a/Engine/Source/Runtime/Renderer/Private/SceneRendering.h
+++ b/Engine/Source/Runtime/Renderer/Private/SceneRendering.h
@@ -808,6 +808,7 @@ struct FPreviousViewInfo
 	
 	// Temporal AA history for SSR
 	FTemporalAAHistory SSRHistory;
+	FTemporalAAHistory WaterSSRHistory;
 
 	// Scene color input for SSR, that can be different from TemporalAAHistory.RT[0] if there is a SSR
 	// input post process material.
@@ -815,6 +816,7 @@ struct FPreviousViewInfo
 
 	// History for the reflections
 	FScreenSpaceDenoiserHistory ReflectionsHistory;
+	FScreenSpaceDenoiserHistory WaterReflectionsHistory;
 	
 	// History for the ambient occlusion
 	FScreenSpaceDenoiserHistory AmbientOcclusionHistory;
diff --git a/Engine/Source/Runtime/Renderer/Private/ScreenSpaceDenoise.cpp b/Engine/Source/Runtime/Renderer/Private/ScreenSpaceDenoise.cpp
--- a/Engine/Source/Runtime/Renderer/Private/ScreenSpaceDenoise.cpp
+++ b/Engine/Source/Runtime/Renderer/Private/ScreenSpaceDenoise.cpp
@@ -2453,6 +2453,52 @@ public:
 		ReflectionsOutput.Color = SignalOutput.Textures[0];
 		return ReflectionsOutput;
 	}
+
+	FReflectionsOutputs DenoiseWaterReflections(
+		FRDGBuilder& GraphBuilder,
+		const FViewInfo& View,
+		FPreviousViewInfo* PreviousViewInfos,
+		const FSceneTextureParameters& SceneTextures,
+		const FReflectionsInputs& ReflectionInputs,
+		const FReflectionsRayTracingConfig RayTracingConfig) const override
+	{
+		RDG_GPU_STAT_SCOPE(GraphBuilder, ReflectionsDenoiser);
+
+		// Imaginary depth is only used for Nvidia denoiser.
+		// TODO: permutation to not generate it?
+		if (ReflectionInputs.RayImaginaryDepth)
+			GraphBuilder.RemoveUnusedTextureWarning(ReflectionInputs.RayImaginaryDepth);
+
+		FSSDSignalTextures InputSignal;
+		InputSignal.Textures[0] = ReflectionInputs.Color;
+		InputSignal.Textures[1] = ReflectionInputs.RayHitDistance;
+
+		FSSDConstantPixelDensitySettings Settings;
+		Settings.FullResViewport = View.ViewRect;
+		Settings.SignalProcessing = ESignalProcessing::Reflections; // TODO: water reflection to denoise only water pixels
+		Settings.InputResolutionFraction = RayTracingConfig.ResolutionFraction;
+		Settings.ReconstructionSamples = CVarReflectionReconstructionSampleCount.GetValueOnRenderThread();
+		Settings.bUseTemporalAccumulation = CVarReflectionTemporalAccumulation.GetValueOnRenderThread() != 0;
+		Settings.HistoryConvolutionSampleCount = CVarReflectionHistoryConvolutionSampleCount.GetValueOnRenderThread();
+		Settings.MaxInputSPP = RayTracingConfig.RayCountPerPixel;
+
+		TStaticArray<FScreenSpaceDenoiserHistory*, IScreenSpaceDenoiser::kMaxBatchSize> PrevHistories;
+		TStaticArray<FScreenSpaceDenoiserHistory*, IScreenSpaceDenoiser::kMaxBatchSize> NewHistories;
+		PrevHistories[0] = &PreviousViewInfos->WaterReflectionsHistory;
+		NewHistories[0] = View.ViewState ? &View.ViewState->PrevFrameViewInfo.WaterReflectionsHistory : nullptr;
+
+		FSSDSignalTextures SignalOutput;
+		DenoiseSignalAtConstantPixelDensity(
+			GraphBuilder, View, SceneTextures,
+			InputSignal, Settings,
+			PrevHistories,
+			NewHistories,
+			&SignalOutput);
+
+		FReflectionsOutputs ReflectionsOutput;
+		ReflectionsOutput.Color = SignalOutput.Textures[0];
+		return ReflectionsOutput;
+	}
 	
 	FAmbientOcclusionOutputs DenoiseAmbientOcclusion(
 		FRDGBuilder& GraphBuilder,
diff --git a/Engine/Source/Runtime/Renderer/Private/ScreenSpaceDenoise.h b/Engine/Source/Runtime/Renderer/Private/ScreenSpaceDenoise.h
--- a/Engine/Source/Runtime/Renderer/Private/ScreenSpaceDenoise.h
+++ b/Engine/Source/Runtime/Renderer/Private/ScreenSpaceDenoise.h
@@ -208,6 +208,15 @@ public:
 		const FSceneTextureParameters& SceneTextures,
 		const FReflectionsInputs& ReflectionInputs,
 		const FReflectionsRayTracingConfig RayTracingConfig) const = 0;
+
+	/** Entry point to denoise water reflections. */
+	virtual FReflectionsOutputs DenoiseWaterReflections(
+		FRDGBuilder& GraphBuilder,
+		const FViewInfo& View,
+		FPreviousViewInfo* PreviousViewInfos,
+		const FSceneTextureParameters& SceneTextures,
+		const FReflectionsInputs& ReflectionInputs,
+		const FReflectionsRayTracingConfig RayTracingConfig) const = 0;
 	
 	/** Entry point to denoise reflections. */
 	virtual FAmbientOcclusionOutputs DenoiseAmbientOcclusion(
diff --git a/Engine/Source/Runtime/Renderer/Private/SingleLayerWaterRendering.cpp b/Engine/Source/Runtime/Renderer/Private/SingleLayerWaterRendering.cpp
--- a/Engine/Source/Runtime/Renderer/Private/SingleLayerWaterRendering.cpp
+++ b/Engine/Source/Runtime/Renderer/Private/SingleLayerWaterRendering.cpp
@@ -12,6 +12,8 @@
 #include "PostProcess/PostProcessSubsurface.h"
 #include "PostProcess/SceneRenderTargets.h"
 #include "PostProcessTemporalAA.h"
+#include "RayTracing/RaytracingOptions.h"
+#include "RayTracing/RayTracingReflections.h"
 #include "RenderGraph.h"
 #include "ScenePrivate.h"
 #include "SceneRendering.h"
@@ -21,6 +23,7 @@
 
 
 DECLARE_GPU_STAT(SingleLayerWater);
+DECLARE_GPU_STAT_NAMED(RayTracingWaterReflections, TEXT("Ray Tracing Water Reflections"));
 
 
 
@@ -62,6 +65,10 @@ static TAutoConsoleVariable<int32> CVarWaterSingleLayerSSR(
 	TEXT("Enable SSR for the single water renderring system."),
 	ECVF_RenderThreadSafe | ECVF_Scalability);
 
+static TAutoConsoleVariable<int32> CVarWaterSingleLayerRTR(
+	TEXT("r.Water.SingleLayer.RTR"), 1,
+	TEXT("Enable RTR for the single water renderring system."),
+	ECVF_RenderThreadSafe | ECVF_Scalability);
 
 static TAutoConsoleVariable<int32> CVarWaterSingleLayerSSRTAA(
 	TEXT("r.Water.SingleLayer.SSRTAA"), 1,
@@ -667,7 +674,67 @@ void FDeferredShadingSceneRenderer::RenderSingleLayerWaterReflections(FRHIComman
 		}
 
  		const bool bEnableSSR = CVarWaterSingleLayerSSR.GetValueOnRenderThread() != 0 && ShouldRenderScreenSpaceReflections(View);
-		if (bEnableSSR)
+		const bool bEnableRTR = CVarWaterSingleLayerRTR.GetValueOnRenderThread() != 0 && ShouldRenderRayTracingReflections(View);
+		if (bEnableRTR)
+		{
+			RDG_EVENT_SCOPE(GraphBuilder, "RayTracingWaterReflections");
+			RDG_GPU_STAT_SCOPE(GraphBuilder, RayTracingWaterReflections);
+
+			IScreenSpaceDenoiser::FReflectionsInputs DenoiserInputs;
+			IScreenSpaceDenoiser::FReflectionsRayTracingConfig RayTracingConfig;
+
+			//RayTracingConfig.ResolutionFraction = FMath::Clamp(GetRayTracingReflectionsScreenPercentage() / 100.0f, 0.25f, 1.0f);
+			RayTracingConfig.ResolutionFraction = 1.0f;
+			//RayTracingConfig.RayCountPerPixel = GetRayTracingReflectionsSamplesPerPixel(View) > -1 ? GetRayTracingReflectionsSamplesPerPixel(View) : View.FinalPostProcessSettings.RayTracingReflectionsSamplesPerPixel;
+			RayTracingConfig.RayCountPerPixel = 1;
+
+			// Water is assumed to have zero roughness and is not currently denoised.
+			//int32 DenoiserMode = GetReflectionsDenoiserMode();
+			int DenoiserMode = 0;
+			bool bDenoise = DenoiserMode != 0;
+
+			if (!bDenoise)
+			{
+				RayTracingConfig.ResolutionFraction = 1.0f;
+			}
+
+			bool bReflectOnlyWater = true;
+			RenderRayTracingReflections(
+				GraphBuilder,
+				SceneTextures,
+				View,
+				RayTracingConfig.RayCountPerPixel,
+				bReflectOnlyWater,
+				RayTracingConfig.ResolutionFraction,
+				&DenoiserInputs);
+
+			if (bDenoise)
+			{
+				const IScreenSpaceDenoiser* DefaultDenoiser = IScreenSpaceDenoiser::GetDefaultDenoiser();
+				const IScreenSpaceDenoiser* DenoiserToUse = DenoiserMode == 1 ? DefaultDenoiser : GScreenSpaceDenoiser;
+
+				// Standard event scope for denoiser to have all profiling information not matter what, and with explicit detection of third party.
+				RDG_EVENT_SCOPE(GraphBuilder, "%s%s(WaterReflections) %dx%d",
+					DenoiserToUse != DefaultDenoiser ? TEXT("ThirdParty ") : TEXT(""),
+					DenoiserToUse->GetDebugName(),
+					View.ViewRect.Width(), View.ViewRect.Height());
+
+				IScreenSpaceDenoiser::FReflectionsOutputs DenoiserOutputs = DenoiserToUse->DenoiseWaterReflections(
+					GraphBuilder,
+					View,
+					&View.PrevViewInfo,
+					SceneTextures,
+					DenoiserInputs,
+					RayTracingConfig);
+
+				ReflectionsColor = DenoiserOutputs.Color;
+			}
+			else
+			{
+				ReflectionsColor = DenoiserInputs.Color;
+			}
+		}
+		else if (bEnableSSR)
 		{
 			// RUN SSR
 			// Uses the water GBuffer (depth, ABCDEF) to know how to start tracing.
@@ -701,6 +768,8 @@ void FDeferredShadingSceneRenderer::RenderSingleLayerWaterReflections(FRHIComman
 					SceneTextures, 
 					View,
 					TAASettings,
+					//View.PrevViewInfo.WaterSSRHistory,
+					//&View.ViewState->PrevFrameViewInfo.WaterSSRHistory);
 					View.PrevViewInfo.SSRHistory,
 					&View.ViewState->PrevFrameViewInfo.SSRHistory);
 
