Missing declaration header..

diff --git a/Engine/Source/Runtime/Renderer/Private/RayTracing/RayTracingReflections.h b/Engine/Source/Runtime/Renderer/Private/RayTracing/RayTracingReflections.h
--- a/Engine/Source/Runtime/Renderer/Private/RayTracing/RayTracingReflections.h
+++ b/Engine/Source/Runtime/Renderer/Private/RayTracing/RayTracingReflections.h
@@ -16,4 +16,11 @@ bool ShouldRayTracedReflectionsUseHybridReflections();
 bool ShouldRayTracedReflectionsSortMaterials(const FViewInfo& View);
 bool ShouldRayTracedReflectionsRayTraceSkyLightContribution(const FScene& Scene);
 
+#else
+
+FORCEINLINE int32 GetRayTracingReflectionsSamplesPerPixel(const FViewInfo& View)
+{
+	return 0;
+}
+
 #endif // RHI_RAYTRACING
