Fix compilation error.

diff --git a/Engine/Source/Runtime/Renderer/Private/ScreenSpaceDenoise.cpp b/Engine/Source/Runtime/Renderer/Private/ScreenSpaceDenoise.cpp
--- a/Engine/Source/Runtime/Renderer/Private/ScreenSpaceDenoise.cpp
+++ b/Engine/Source/Runtime/Renderer/Private/ScreenSpaceDenoise.cpp
@@ -2468,6 +2468,9 @@ public:
 		// TODO: permutation to not generate it?
 		if (ReflectionInputs.RayImaginaryDepth)
 			GraphBuilder.RemoveUnusedTextureWarning(ReflectionInputs.RayImaginaryDepth);
+		
+		FViewInfoPooledRenderTargets ViewInfoPooledRenderTargets;
+		SetupSceneViewInfoPooledRenderTargets(View, &ViewInfoPooledRenderTargets);
 
 		FSSDSignalTextures InputSignal;
 		InputSignal.Textures[0] = ReflectionInputs.Color;
