Added ShadowCacheusageMB CSV stat

#rnx

diff --git a/Engine/Source/Runtime/Renderer/Private/ShadowSetup.cpp b/Engine/Source/Runtime/Renderer/Private/ShadowSetup.cpp
--- a/Engine/Source/Runtime/Renderer/Private/ShadowSetup.cpp
+++ b/Engine/Source/Runtime/Renderer/Private/ShadowSetup.cpp
@@ -5043,6 +5043,7 @@ void FSceneRenderer::InitDynamicShadows(FRHICommandListImmediate& RHICmdList, FG
 		}
 
 		CSV_CUSTOM_STAT(LightCount, UpdatedShadowMaps, float(NumPointShadowCachesUpdatedThisFrame + NumSpotShadowCachesUpdatedThisFrame), ECsvCustomStatOp::Set);
+		CSV_CUSTOM_STAT_GLOBAL(ShadowCacheUsageMB, (float(Scene->GetCachedWholeSceneShadowMapsSize()) / 1024) / 1024, ECsvCustomStatOp::Set);
 
 		// Calculate visibility of the projected shadows.
 		InitProjectedShadowVisibility(RHICmdList);
