Additional change to likely fix CIS
#rnx

diff --git a/Engine/Source/Runtime/Renderer/Private/ScreenSpaceDenoise.cpp b/Engine/Source/Runtime/Renderer/Private/ScreenSpaceDenoise.cpp
--- a/Engine/Source/Runtime/Renderer/Private/ScreenSpaceDenoise.cpp
+++ b/Engine/Source/Runtime/Renderer/Private/ScreenSpaceDenoise.cpp
@@ -2493,6 +2493,7 @@ public:
 		FSSDSignalTextures SignalOutput;
 		DenoiseSignalAtConstantPixelDensity(
 			GraphBuilder, View, SceneTextures,
+			ViewInfoPooledRenderTargets,
 			InputSignal, Settings,
 			PrevHistories,
 			NewHistories,
