GetTonemappingLUT needs to check bNeedFloatOutput

diff --git a/Engine/Source/Runtime/Renderer/Private/ScenePrivate.h b/Engine/Source/Runtime/Renderer/Private/ScenePrivate.h
--- a/Engine/Source/Runtime/Renderer/Private/ScenePrivate.h
+++ b/Engine/Source/Runtime/Renderer/Private/ScenePrivate.h
@@ -1278,7 +1278,8 @@ public:
 		if (CombinedLUTRenderTarget.IsValid() == false || 
 			CombinedLUTRenderTarget->GetDesc().Extent.Y != LUTSize ||
 			((CombinedLUTRenderTarget->GetDesc().Depth != 0) != bUseVolumeLUT) ||
-			!!(CombinedLUTRenderTarget->GetDesc().TargetableFlags & TexCreate_UAV) != bNeedUAV)
+			!!(CombinedLUTRenderTarget->GetDesc().TargetableFlags & TexCreate_UAV) != bNeedUAV ||
+			(CombinedLUTRenderTarget->GetDesc().Format == PF_FloatRGBA) != bNeedFloatOutput)
 		{
 			// Create the texture needed for the tonemapping LUT
 			FPooledRenderTargetDesc Desc = CreateLUTRenderTarget(LUTSize, bUseVolumeLUT, bNeedUAV, bNeedFloatOutput);
