Fix CIS

diff --git a/Engine/Source/Runtime/Renderer/Private/ShadowSetup.cpp b/Engine/Source/Runtime/Renderer/Private/ShadowSetup.cpp
--- a/Engine/Source/Runtime/Renderer/Private/ShadowSetup.cpp
+++ b/Engine/Source/Runtime/Renderer/Private/ShadowSetup.cpp
@@ -1293,8 +1293,10 @@ bool FProjectedShadowInfo::ShouldDrawStaticMeshes_AnyThread(
 
 			if (PrimitiveSceneInfo->bIsUsingCustomLODRules)
 			{
+				PRAGMA_DISABLE_DEPRECATION_WARNINGS
 				ViewLODToRender = Proxy->GetCustomLOD(CurrentView, CurrentView.LODDistanceFactor, ForcedLOD, MeshScreenSizeSquared);
 				ViewLODToRender.ClampToFirstLOD(CurFirstLODIdx);
+				PRAGMA_ENABLE_DEPRECATION_WARNINGS
 			}
 			else
 			{
@@ -1347,7 +1349,9 @@ bool FProjectedShadowInfo::ShouldDrawStaticMeshes_AnyThread(
 			{
 				const bool bHasShelfShadow = !bReflectiveShadowmap && !bPreShadow;
 				const float ShadowMapTextureResolution = BorderSize * 2 + ResolutionX; // We assume Shadow Map texture to be squared (current design)
+				PRAGMA_DISABLE_DEPRECATION_WARNINGS
 				ShadowLODToRender = Proxy->GetCustomWholeSceneShadowLOD(CurrentView, CurrentView.LODDistanceFactor, ForcedLOD, *VisibilePrimitiveLODMask, ShadowMapTextureResolution, ShadowBounds.W * 2.0f, ShadowId, bHasShelfShadow);
+				PRAGMA_ENABLE_DEPRECATION_WARNINGS
 			}
 
 			// Don't cache if it requires per view per mesh state for distance cull fade.
