Fix for UDN with TDR on DX12
- Add extra profile tags
- Make sure GPU command lists are (request) kicked before flushing the Rendering & RHI thread

diff --git a/Engine/Source/Runtime/Renderer/Private/GPUBenchmark.cpp b/Engine/Source/Runtime/Renderer/Private/GPUBenchmark.cpp
--- a/Engine/Source/Runtime/Renderer/Private/GPUBenchmark.cpp
+++ b/Engine/Source/Runtime/Renderer/Private/GPUBenchmark.cpp
@@ -537,10 +537,14 @@ void RendererGPUBenchmark(FRHICommandListImmediate& RHICmdList, FSynthBenchmarkR
 		}
 
 		RHICmdList.EndRenderQuery(TimerQueries[0].GetQuery());
+		
+		SCOPED_DRAW_EVENTF(RHICmdList, Benchmark, TEXT("Scale:%f"), WorkScale);
 
 		// multiple iterations to see how trust able the values are
 		for(uint32 Iteration = 0; Iteration < IterationCount; ++Iteration)
 		{
+			SCOPED_DRAW_EVENTF(RHICmdList, Benchmark, TEXT("Iteration:%d"), Iteration);
+
 			for(uint32 MethodIterator = 0; MethodIterator < MethodCount; ++MethodIterator)
 			{
 				// alternate between forward and backward (should give the same number)
@@ -597,6 +601,7 @@ void RendererGPUBenchmark(FRHICommandListImmediate& RHICmdList, FSynthBenchmarkR
 		{
 			uint64 OldAbsTime = 0;
 			// flushes the RHI thread to make sure all RHICmdList.EndRenderQuery() commands got executed.
+			RHICmdList.SubmitCommandsHint();
 			RHICmdList.ImmediateFlush(EImmediateFlushType::FlushRHIThread);
 			RHICmdList.GetRenderQueryResult(TimerQueries[0].GetQuery(), OldAbsTime, true);
 
