Fix deprecation warnings
#rnx

diff --git a/Engine/Source/Runtime/Renderer/Private/ShadowSetup.cpp b/Engine/Source/Runtime/Renderer/Private/ShadowSetup.cpp
--- a/Engine/Source/Runtime/Renderer/Private/ShadowSetup.cpp
+++ b/Engine/Source/Runtime/Renderer/Private/ShadowSetup.cpp
@@ -1345,6 +1345,7 @@ bool FProjectedShadowInfo::ShouldDrawStaticMeshes_AnyThread(
 
 		if (WholeSceneDirectionalShadow)
 		{
+			PRAGMA_DISABLE_DEPRECATION_WARNINGS
 			if (PrimitiveSceneInfo->bIsUsingCustomWholeSceneShadowLODRules)
 			{
 				const bool bHasShelfShadow = !bReflectiveShadowmap && !bPreShadow;
@@ -1353,6 +1354,7 @@ bool FProjectedShadowInfo::ShouldDrawStaticMeshes_AnyThread(
 				ShadowLODToRender = Proxy->GetCustomWholeSceneShadowLOD(CurrentView, CurrentView.LODDistanceFactor, ForcedLOD, *VisibilePrimitiveLODMask, ShadowMapTextureResolution, ShadowBounds.W * 2.0f, ShadowId, bHasShelfShadow);
 				PRAGMA_ENABLE_DEPRECATION_WARNINGS
 			}
+			PRAGMA_ENABLE_DEPRECATION_WARNINGS
 
 			// Don't cache if it requires per view per mesh state for distance cull fade.
 			const bool bCanCache = !bMayBeFading && !bNeedUpdateStaticMeshes;
