Fix CIS issues
#rnx

diff --git a/Engine/Source/Runtime/Renderer/Private/ShadowSetup.cpp b/Engine/Source/Runtime/Renderer/Private/ShadowSetup.cpp
--- a/Engine/Source/Runtime/Renderer/Private/ShadowSetup.cpp
+++ b/Engine/Source/Runtime/Renderer/Private/ShadowSetup.cpp
@@ -1350,9 +1350,7 @@ bool FProjectedShadowInfo::ShouldDrawStaticMeshes_AnyThread(
 			{
 				const bool bHasShelfShadow = !bReflectiveShadowmap && !bPreShadow;
 				const float ShadowMapTextureResolution = BorderSize * 2 + ResolutionX; // We assume Shadow Map texture to be squared (current design)
-				PRAGMA_DISABLE_DEPRECATION_WARNINGS
 				ShadowLODToRender = Proxy->GetCustomWholeSceneShadowLOD(CurrentView, CurrentView.LODDistanceFactor, ForcedLOD, *VisibilePrimitiveLODMask, ShadowMapTextureResolution, ShadowBounds.W * 2.0f, ShadowId, bHasShelfShadow);
-				PRAGMA_ENABLE_DEPRECATION_WARNINGS
 			}
 			PRAGMA_ENABLE_DEPRECATION_WARNINGS
 
