Fix static analysis warning.

diff --git a/Engine/Source/Runtime/Renderer/Private/SingleLayerWaterRendering.cpp b/Engine/Source/Runtime/Renderer/Private/SingleLayerWaterRendering.cpp
--- a/Engine/Source/Runtime/Renderer/Private/SingleLayerWaterRendering.cpp
+++ b/Engine/Source/Runtime/Renderer/Private/SingleLayerWaterRendering.cpp
@@ -690,8 +690,9 @@ void FDeferredShadingSceneRenderer::RenderSingleLayerWaterReflections(FRHIComman
 
 			// Water is assumed to have zero roughness and is not currently denoised.
 			//int32 DenoiserMode = GetReflectionsDenoiserMode();
+			//bool bDenoise = DenoiserMode != 0;
 			int DenoiserMode = 0;
-			bool bDenoise = DenoiserMode != 0;
+			bool bDenoise = false;
 
 			if (!bDenoise)
 			{
