Better messaging when ray tracing reaches the maximum number of sections supported per mesh

diff --git a/Engine/Source/Runtime/Renderer/Private/MeshPassProcessor.cpp b/Engine/Source/Runtime/Renderer/Private/MeshPassProcessor.cpp
--- a/Engine/Source/Runtime/Renderer/Private/MeshPassProcessor.cpp
+++ b/Engine/Source/Runtime/Renderer/Private/MeshPassProcessor.cpp
@@ -281,7 +281,7 @@ void FMeshDrawShaderBindings::SetRayTracingShaderBindingsForHitGroup(
 	uint32 ShaderSlot) const
 {
 	check(ShaderLayouts.Num() == 1);
-	check(SegmentIndex < 0xFF);
+	checkf(SegmentIndex < 0xFF, TEXT("Ray Tracing does not support meshes with more than 256 sections."));
 
 	FReadOnlyMeshDrawSingleShaderBindings SingleShaderBindings(ShaderLayouts[0], GetData());
 
