Cherry-pick 10755012: Fixes half pixel shift regression in Mitchell Netraveli that was making the output blurier

Regression caused by CL 7579472

diff --git a/Engine/Source/Runtime/Renderer/Private/PostProcess/PostProcessMitchellNetravali.cpp b/Engine/Source/Runtime/Renderer/Private/PostProcess/PostProcessMitchellNetravali.cpp
--- a/Engine/Source/Runtime/Renderer/Private/PostProcess/PostProcessMitchellNetravali.cpp
+++ b/Engine/Source/Runtime/Renderer/Private/PostProcess/PostProcessMitchellNetravali.cpp
@@ -62,8 +62,8 @@ FRDGTextureRef ComputeMitchellNetravaliDownsample(
 	// Scale / Bias factor to map the dispatch thread id to the input texture UV.
 	PassParameters->DispatchThreadToInputUVScale.X = Input.ViewRect.Width()  / float(OutputViewport.Rect.Width()  * Input.Texture->Desc.Extent.X);
 	PassParameters->DispatchThreadToInputUVScale.Y = Input.ViewRect.Height() / float(OutputViewport.Rect.Height() * Input.Texture->Desc.Extent.Y);
-	PassParameters->DispatchThreadToInputUVBias.X = PassParameters->DispatchThreadToInputUVScale.X * (0.5f * Input.ViewRect.Min.X);
-	PassParameters->DispatchThreadToInputUVBias.Y = PassParameters->DispatchThreadToInputUVScale.Y * (0.5f * Input.ViewRect.Min.Y);
+	PassParameters->DispatchThreadToInputUVBias.X = PassParameters->DispatchThreadToInputUVScale.X * (0.5f + Input.ViewRect.Min.X);
+	PassParameters->DispatchThreadToInputUVBias.Y = PassParameters->DispatchThreadToInputUVScale.Y * (0.5f + Input.ViewRect.Min.Y);
 
 	TShaderMapRef<FMitchellNetravaliDownsampleCS> ComputeShader(View.ShaderMap);
 
