Adding ray traced reflections to single-layer water.

diff --git a/Engine/Source/Runtime/Renderer/Private/IndirectLightRendering.cpp b/Engine/Source/Runtime/Renderer/Private/IndirectLightRendering.cpp
--- a/Engine/Source/Runtime/Renderer/Private/IndirectLightRendering.cpp
+++ b/Engine/Source/Runtime/Renderer/Private/IndirectLightRendering.cpp
@@ -671,7 +671,7 @@ void FDeferredShadingSceneRenderer::RenderDeferredReflectionsAndSkyLighting(FRHI
 		FRDGTextureRef ReflectionsColor = nullptr;
 		if (bRayTracedReflections || bScreenSpaceReflections)
 		{
-			int32 DenoiserMode = CVarUseReflectionDenoiser.GetValueOnRenderThread();
+			int32 DenoiserMode = GetReflectionsDenoiserMode();
 
 			bool bDenoise = false;
 			bool bTemporalFilter = false;
diff --git a/Engine/Source/Runtime/Renderer/Private/ScreenSpaceDenoise.cpp b/Engine/Source/Runtime/Renderer/Private/ScreenSpaceDenoise.cpp
--- a/Engine/Source/Runtime/Renderer/Private/ScreenSpaceDenoise.cpp
+++ b/Engine/Source/Runtime/Renderer/Private/ScreenSpaceDenoise.cpp
@@ -40,6 +40,15 @@ static TAutoConsoleVariable<int32> CVarShadowHistoryConvolutionSampleCount(
 	TEXT("Number of samples to use to convolve the history over time."),
 	ECVF_RenderThreadSafe);
 
+static TAutoConsoleVariable<int32> CVarUseReflectionDenoiser(
+	TEXT("r.Reflections.Denoiser"),
+	2,
+	TEXT("Choose the denoising algorithm.\n")
+	TEXT(" 0: Disabled;\n")
+	TEXT(" 1: Forces the default denoiser of the renderer;\n")
+	TEXT(" 2: GScreenSpaceDenoiser which may be overriden by a third party plugin (default)."),
+	ECVF_RenderThreadSafe);
+
 static TAutoConsoleVariable<int32> CVarReflectionReconstructionSampleCount(
 	TEXT("r.Reflections.Denoiser.ReconstructionSamples"), 16,
 	TEXT("Maximum number of samples for the reconstruction pass (default = 16)."),
@@ -2789,3 +2798,8 @@ const IScreenSpaceDenoiser* IScreenSpaceDenoiser::GetDefaultDenoiser()
 	static IScreenSpaceDenoiser* GDefaultDenoiser = new FDefaultScreenSpaceDenoiser;
 	return GDefaultDenoiser;
 }
+
+int GetReflectionsDenoiserMode()
+{
+	return CVarUseReflectionDenoiser.GetValueOnRenderThread();
+}
\ No newline at end of file
diff --git a/Engine/Source/Runtime/Renderer/Private/ScreenSpaceDenoise.h b/Engine/Source/Runtime/Renderer/Private/ScreenSpaceDenoise.h
--- a/Engine/Source/Runtime/Renderer/Private/ScreenSpaceDenoise.h
+++ b/Engine/Source/Runtime/Renderer/Private/ScreenSpaceDenoise.h
@@ -281,3 +281,5 @@ public:
 
 // The interface for the renderer to denoise what it needs, Plugins can come and point this to custom interface.
 extern RENDERER_API const IScreenSpaceDenoiser* GScreenSpaceDenoiser;
+
+extern int GetReflectionsDenoiserMode();
\ No newline at end of file
diff --git a/Engine/Source/Runtime/Renderer/Private/SingleLayerWaterRendering.cpp b/Engine/Source/Runtime/Renderer/Private/SingleLayerWaterRendering.cpp
--- a/Engine/Source/Runtime/Renderer/Private/SingleLayerWaterRendering.cpp
+++ b/Engine/Source/Runtime/Renderer/Private/SingleLayerWaterRendering.cpp
@@ -12,6 +12,7 @@
 #include "PostProcess/PostProcessSubsurface.h"
 #include "PostProcess/SceneRenderTargets.h"
 #include "PostProcessTemporalAA.h"
+#include "RayTracing/RayTracingOptions.h"
 #include "RayTracing/RaytracingOptions.h"
 #include "RayTracing/RayTracingReflections.h"
 #include "RenderGraph.h"
@@ -26,7 +27,6 @@ DECLARE_GPU_STAT(SingleLayerWater);
 DECLARE_GPU_STAT_NAMED(RayTracingWaterReflections, TEXT("Ray Tracing Water Reflections"));
 
 
-
 static TAutoConsoleVariable<int32> CVarWaterSingleLayer(
 	TEXT("r.Water.SingleLayer"), 1,
 	TEXT("Enable the single water renderring system."),
@@ -769,10 +769,8 @@ void FDeferredShadingSceneRenderer::RenderSingleLayerWaterReflections(FRHIComman
 					SceneTextures, 
 					View,
 					TAASettings,
-					//View.PrevViewInfo.WaterSSRHistory,
-					//&View.ViewState->PrevFrameViewInfo.WaterSSRHistory);
-					View.PrevViewInfo.SSRHistory,
-					&View.ViewState->PrevFrameViewInfo.SSRHistory);
+					View.PrevViewInfo.WaterSSRHistory,
+					&View.ViewState->PrevFrameViewInfo.WaterSSRHistory);
 
 				ReflectionsColor = TAAOutputs.SceneColor;
 			}
