Correcting case-sensitivity for Linux.

diff --git a/Engine/Source/Runtime/Renderer/Private/SingleLayerWaterRendering.cpp b/Engine/Source/Runtime/Renderer/Private/SingleLayerWaterRendering.cpp
--- a/Engine/Source/Runtime/Renderer/Private/SingleLayerWaterRendering.cpp
+++ b/Engine/Source/Runtime/Renderer/Private/SingleLayerWaterRendering.cpp
@@ -12,7 +12,6 @@
 #include "PostProcess/PostProcessSubsurface.h"
 #include "PostProcess/SceneRenderTargets.h"
 #include "PostProcessTemporalAA.h"
-#include "RayTracing/RayTracingOptions.h"
 #include "RayTracing/RaytracingOptions.h"
 #include "RayTracing/RayTracingReflections.h"
 #include "RenderGraph.h"
