Created ray tracing reflections options construct.

diff --git a/Engine/Source/Runtime/Renderer/Private/DeferredShadingRenderer.cpp b/Engine/Source/Runtime/Renderer/Private/DeferredShadingRenderer.cpp
--- a/Engine/Source/Runtime/Renderer/Private/DeferredShadingRenderer.cpp
+++ b/Engine/Source/Runtime/Renderer/Private/DeferredShadingRenderer.cpp
@@ -1110,6 +1110,7 @@ bool FDeferredShadingSceneRenderer::DispatchRayTracingWorldUpdates(FRHICommandLi
 		// #dxr_todo: UE-72565: refactor ray tracing effects to not be member functions of DeferredShadingRenderer. register each effect at startup and just loop over them automatically to gather all required shaders
 		TArray<FRHIRayTracingShader*> RayGenShaders;
 		PrepareRayTracingReflections(View, *Scene, RayGenShaders);
+		PrepareSingleLayerWaterRayTracingReflections(View, *Scene, RayGenShaders);
 		PrepareRayTracingShadows(View, RayGenShaders);
 		PrepareRayTracingAmbientOcclusion(View, RayGenShaders);
 		PrepareRayTracingSkyLight(View, RayGenShaders);
diff --git a/Engine/Source/Runtime/Renderer/Private/DeferredShadingRenderer.h b/Engine/Source/Runtime/Renderer/Private/DeferredShadingRenderer.h
--- a/Engine/Source/Runtime/Renderer/Private/DeferredShadingRenderer.h
+++ b/Engine/Source/Runtime/Renderer/Private/DeferredShadingRenderer.h
@@ -27,6 +27,7 @@ class FExponentialHeightFogSceneInfo;
 class FRaytracingLightDataPacked;
 struct FSingleLayerWaterPassData;
 struct FHeightFogRenderingParameters;
+struct FRayTracingReflectionOptions;
 
 class FLightShaftsOutput
 {
@@ -522,15 +523,14 @@ private:
 		FRDGBuilder& GraphBuilder,
 		const FSceneTextureParameters& SceneTextures,
 		const FViewInfo& View,
-		int32 SamplePerPixel,
-		bool bReflectWaterOnly,
-		float ResolutionFraction,
+		const FRayTracingReflectionOptions& Options,
 		IScreenSpaceDenoiser::FReflectionsInputs* OutDenoiserInputs);
 
 	void RenderRayTracingDeferredReflections(
 		FRDGBuilder& GraphBuilder,
 		const FSceneTextureParameters& SceneTextures,
 		const FViewInfo& View,
+		const FRayTracingReflectionOptions& Options,
 		IScreenSpaceDenoiser::FReflectionsInputs* OutDenoiserInputs);
 
 	void RenderRayTracingShadows(
@@ -644,6 +644,7 @@ private:
 	// #dxr_todo: UE-72565: refactor ray tracing effects to not be member functions of DeferredShadingRenderer. Register each effect at startup and just loop over them automatically
 	static void PrepareRayTracingReflections(const FViewInfo& View, const FScene& Scene, TArray<FRHIRayTracingShader*>& OutRayGenShaders);
 	static void PrepareRayTracingDeferredReflections(const FViewInfo& View, const FScene& Scene, TArray<FRHIRayTracingShader*>& OutRayGenShaders);
+	static void PrepareSingleLayerWaterRayTracingReflections(const FViewInfo& View, const FScene& Scene, TArray<FRHIRayTracingShader*>& OutRayGenShaders);
 	static void PrepareRayTracingShadows(const FViewInfo& View, TArray<FRHIRayTracingShader*>& OutRayGenShaders);
 	static void PrepareRayTracingAmbientOcclusion(const FViewInfo& View, TArray<FRHIRayTracingShader*>& OutRayGenShaders);
 	static void PrepareRayTracingSkyLight(const FViewInfo& View, TArray<FRHIRayTracingShader*>& OutRayGenShaders);
diff --git a/Engine/Source/Runtime/Renderer/Private/IndirectLightRendering.cpp b/Engine/Source/Runtime/Renderer/Private/IndirectLightRendering.cpp
--- a/Engine/Source/Runtime/Renderer/Private/IndirectLightRendering.cpp
+++ b/Engine/Source/Runtime/Renderer/Private/IndirectLightRendering.cpp
@@ -698,14 +698,24 @@ void FDeferredShadingSceneRenderer::RenderDeferredReflectionsAndSkyLighting(FRHI
 					RayTracingConfig.ResolutionFraction = 1.0f;
 				}
 
-				bool bReflectOnlyWater = false;
+				FRayTracingReflectionOptions Options;
+				if (ShouldRayTracedReflectionsUseSortedDeferredAlgorithm(View))
+				{
+					Options.Algorithm = FRayTracingReflectionOptions::SortedDeferred;
+				}
+				else if (ShouldRayTracedReflectionsSortMaterials(View))
+				{
+					Options.Algorithm = FRayTracingReflectionOptions::Sorted;
+				}
+				Options.SamplesPerPixel = RayTracingConfig.RayCountPerPixel;
+				Options.ResolutionFraction = RayTracingConfig.ResolutionFraction;
+				Options.bReflectOnlyWater = false;
+
 				RenderRayTracingReflections(
 					GraphBuilder,
 					SceneTextures,
 					View,
-					RayTracingConfig.RayCountPerPixel,
-					bReflectOnlyWater,
-					RayTracingConfig.ResolutionFraction,
+					Options,
 					&DenoiserInputs);
 			}
 			else if (bScreenSpaceReflections)
@@ -978,4 +988,4 @@ void FDeferredShadingSceneRenderer::RenderDeferredReflectionsAndSkyLightingHair(
 			GraphBuilder.Execute();
 		}
 	}
-}
\ No newline at end of file
+}
diff --git a/Engine/Source/Runtime/Renderer/Private/RayTracing/RayTracingDeferredReflections.cpp b/Engine/Source/Runtime/Renderer/Private/RayTracing/RayTracingDeferredReflections.cpp
--- a/Engine/Source/Runtime/Renderer/Private/RayTracing/RayTracingDeferredReflections.cpp
+++ b/Engine/Source/Runtime/Renderer/Private/RayTracing/RayTracingDeferredReflections.cpp
@@ -200,6 +200,7 @@ void FDeferredShadingSceneRenderer::RenderRayTracingDeferredReflections(
 	FRDGBuilder& GraphBuilder,
 	const FSceneTextureParameters& SceneTextures,
 	const FViewInfo& View,
+	const FRayTracingReflectionOptions& Options,
 	IScreenSpaceDenoiser::FReflectionsInputs* OutDenoiserInputs)
 {
 	FRDGTextureDesc OutputDesc = FPooledRenderTargetDesc::Create2DDesc(
@@ -315,6 +316,7 @@ void FDeferredShadingSceneRenderer::RenderRayTracingDeferredReflections(
 	FRDGBuilder& GraphBuilder,
 	const FSceneTextureParameters& SceneTextures,
 	const FViewInfo& View,
+	const FRayTracingReflectionOptions& Options,
 	IScreenSpaceDenoiser::FReflectionsInputs* OutDenoiserInputs)
 {
 	checkNoEntry();
diff --git a/Engine/Source/Runtime/Renderer/Private/RayTracing/RayTracingReflections.cpp b/Engine/Source/Runtime/Renderer/Private/RayTracing/RayTracingReflections.cpp
--- a/Engine/Source/Runtime/Renderer/Private/RayTracing/RayTracingReflections.cpp
+++ b/Engine/Source/Runtime/Renderer/Private/RayTracing/RayTracingReflections.cpp
@@ -379,6 +379,13 @@ bool ShouldRayTracedReflectionsSortMaterials(const FViewInfo& View)
 	return (ShouldRayTracedReflectionsUseHybridReflections() || CVarRayTracingReflectionsSortMaterials.GetValueOnRenderThread() != 0) && !bIsMultiviewSecondary;
 }
 
+bool ShouldRayTracedReflectionsUseSortedDeferredAlgorithm(const FViewInfo& View)
+{
+	const bool bIsMultiviewSecondary = View.ViewRect.Min.X > 0 || View.ViewRect.Min.Y > 0;
+
+	return (CVarRayTracingReflectionsExperimentalDeferred.GetValueOnRenderThread() != 0) && !bIsMultiviewSecondary;
+}
+
 bool ShouldRayTracedReflectionsRayTraceSkyLightContribution(const FScene& Scene)
 {
 	// Only ray trace sky light contribution when the ray traced sky light should be rendered in normal conditions (sky light exists, ray traced shadows enabled)
@@ -437,6 +444,20 @@ void FDeferredShadingSceneRenderer::PrepareRayTracingReflections(const FViewInfo
 	}
 }
 
+void FDeferredShadingSceneRenderer::PrepareSingleLayerWaterRayTracingReflections(const FViewInfo& View, const FScene& Scene, TArray<FRHIRayTracingShader*>& OutRayGenShaders)
+{
+	const bool bMissShaderLighting = CanUseRayTracingLightingMissShader(View.GetShaderPlatform());
+	const bool bRayTraceSkyLightContribution = ShouldRayTracedReflectionsRayTraceSkyLightContribution(Scene);
+
+	FRayTracingReflectionsRGS::FPermutationDomain PermutationVector;
+	PermutationVector.Set<FRayTracingReflectionsRGS::FEnableTwoSidedGeometryForShadowDim>(EnableRayTracingShadowTwoSidedGeometry());
+	PermutationVector.Set<FRayTracingReflectionsRGS::FDeferredMaterialMode>(EDeferredMaterialMode::None);
+	PermutationVector.Set<FRayTracingReflectionsRGS::FMissShaderLighting>(bMissShaderLighting);
+	PermutationVector.Set<FRayTracingReflectionsRGS::FRayTraceSkyLightContribution>(bRayTraceSkyLightContribution);
+	auto RayGenShader = View.ShaderMap->GetShader<FRayTracingReflectionsRGS>(PermutationVector);
+	OutRayGenShaders.Add(RayGenShader.GetRayTracingShader());
+}
+
 #endif // RHI_RAYTRACING
 
 void FDeferredShadingSceneRenderer::SetupImaginaryReflectionTextureParameters(
@@ -466,18 +487,17 @@ void FDeferredShadingSceneRenderer::RenderRayTracingReflections(
 	FRDGBuilder& GraphBuilder,
 	const FSceneTextureParameters& SceneTextures,
 	const FViewInfo& View,
-	int32 SamplePerPixel,
-	bool bReflectOnlyWater,
-	float ResolutionFraction,
+	const FRayTracingReflectionOptions& Options,
 	IScreenSpaceDenoiser::FReflectionsInputs* OutDenoiserInputs)
 #if RHI_RAYTRACING
 {
-	if (CVarRayTracingReflectionsExperimentalDeferred.GetValueOnRenderThread())
+	if (Options.Algorithm == FRayTracingReflectionOptions::SortedDeferred)
 	{
 		RenderRayTracingDeferredReflections(
 			GraphBuilder,
 			SceneTextures,
 			View,
+			Options,
 			OutDenoiserInputs);
 		return;
 	}
@@ -486,15 +506,15 @@ void FDeferredShadingSceneRenderer::RenderRayTracingReflections(
 	const uint32 EnableTranslucency = GRayTracingReflectionsTranslucency > -1 ? (uint32)GRayTracingReflectionsTranslucency : (uint32)View.FinalPostProcessSettings.RayTracingReflectionsTranslucency;
 
 	const bool bHybridReflections = ShouldRayTracedReflectionsUseHybridReflections();
-	const bool bSortMaterials = ShouldRayTracedReflectionsSortMaterials(View);
+	const bool bSortMaterials = Options.Algorithm == FRayTracingReflectionOptions::Sorted;
 	const bool bLightingMissShader = CanUseRayTracingLightingMissShader(View.GetShaderPlatform());
 	const bool bRayTraceSkyLightContribution = ShouldRayTracedReflectionsRayTraceSkyLightContribution(*Scene);
 
 	FSceneRenderTargets& SceneContext = FSceneRenderTargets::Get(GraphBuilder.RHICmdList);
 	FSceneViewState* SceneViewState = (FSceneViewState*)View.State;
 
-	int32 UpscaleFactor = int32(1.0f / ResolutionFraction);
-	ensure(ResolutionFraction == 1.0 / UpscaleFactor);
+	int32 UpscaleFactor = int32(1.0f / Options.ResolutionFraction);
+	ensure(Options.ResolutionFraction == 1.0 / UpscaleFactor);
 	ensureMsgf(FComputeShaderUtils::kGolden2DGroupSize % UpscaleFactor == 0, TEXT("Reflection ray tracing will have uv misalignement."));
 	FIntPoint RayTracingResolution = FIntPoint::DivideAndRoundUp(View.ViewRect.Size(), UpscaleFactor);
 	FIntPoint RayTracingBufferSize = SceneTextures.SceneDepthBuffer->Desc.Extent / UpscaleFactor;
@@ -548,14 +568,14 @@ void FDeferredShadingSceneRenderer::RenderRayTracingReflections(
 
 	FRayTracingReflectionsRGS::FParameters CommonParameters;
 
-	CommonParameters.SamplesPerPixel = SamplePerPixel;
+	CommonParameters.SamplesPerPixel = Options.SamplesPerPixel;
 	CommonParameters.MaxBounces = FMath::Max(1, GRayTracingReflectionsMaxBounces > -1? GRayTracingReflectionsMaxBounces : View.FinalPostProcessSettings.RayTracingReflectionsMaxBounces);
 	CommonParameters.HeightFog = GRayTracingReflectionsHeightFog;
 	CommonParameters.UseReflectionCaptures = GRayTracingReflectionsCaptures;
 	CommonParameters.ShouldDoDirectLighting = GRayTracingReflectionsDirectLighting;
 	CommonParameters.ReflectedShadowsType = GRayTracingReflectionsShadows > -1 ? GRayTracingReflectionsShadows : (int32)View.FinalPostProcessSettings.RayTracingReflectionsShadows;
 	CommonParameters.ShouldDoEmissiveAndIndirectLighting = GRayTracingReflectionsEmissiveAndIndirectLighting;
-	CommonParameters.ShouldReflectOnlyWater = bReflectOnlyWater;
+	CommonParameters.ShouldReflectOnlyWater = Options.bReflectOnlyWater;
 	CommonParameters.UpscaleFactor = UpscaleFactor;
 	CommonParameters.ReflectionMinRayDistance = FMath::Min(GRayTracingReflectionsMinRayDistance, GRayTracingReflectionsMaxRayDistance);
 	CommonParameters.ReflectionMaxRayDistance = GRayTracingReflectionsMaxRayDistance;
@@ -610,13 +630,13 @@ void FDeferredShadingSceneRenderer::RenderRayTracingReflections(
 		SetupSkyLightQuasiRandomParameters(*Scene, View, BlueNoiseDimensions, &CommonParameters.SkyLightQuasiRandomData);
 	}
 
-	for (int32 SamplePassIndex = 0; SamplePassIndex < SamplePerPixel; SamplePassIndex++)
+	for (int32 SamplePassIndex = 0; SamplePassIndex < Options.SamplesPerPixel; SamplePassIndex++)
 	{
-		if (SamplePerPixel > 1)
+		if (Options.SamplesPerPixel > 1)
 		{
 			CommonParameters.SampleMode = (int32)ESamplePhase::Accum;
 			CommonParameters.SampleMode = SamplePassIndex == 0 ? (int32)ESamplePhase::Init : CommonParameters.SampleMode;
-			CommonParameters.SampleMode = SamplePassIndex == (SamplePerPixel - 1) ? (int32)ESamplePhase::Resolve : CommonParameters.SampleMode;
+			CommonParameters.SampleMode = SamplePassIndex == (Options.SamplesPerPixel - 1) ? (int32)ESamplePhase::Resolve : CommonParameters.SampleMode;
 		}
 		CommonParameters.SampleOffset = SamplePassIndex;
 
@@ -679,7 +699,7 @@ void FDeferredShadingSceneRenderer::RenderRayTracingReflections(
 				if (NumPasses > 1 || RenderTileSize <= 0)
 				{
 					GraphBuilder.AddPass(
-						RDG_EVENT_NAME("ReflectionRayTracing(spp=%d) %dx%d", SamplePerPixel, RayTracingResolution.X, RayTracingResolution.Y),
+						RDG_EVENT_NAME("ReflectionRayTracing(spp=%d) %dx%d", Options.SamplesPerPixel, RayTracingResolution.X, RayTracingResolution.Y),
 						PassParameters,
 						ERDGPassFlags::Compute,
 						[PassParameters, this, &View, RayGenShader, RayTracingResolution, DeferredMaterialBufferNumElements, DeferredMaterialMode](FRHICommandList& RHICmdList)
@@ -722,7 +742,7 @@ void FDeferredShadingSceneRenderer::RenderRayTracingReflections(
 							int32 DispatchSizeX = FMath::Min<int32>(RenderTileSize, RayTracingResolution.X - TilePassParameters->RenderTileOffsetX);
 							int32 DispatchSizeY = FMath::Min<int32>(RenderTileSize, RayTracingResolution.Y - TilePassParameters->RenderTileOffsetY);
 							GraphBuilder.AddPass(
-								RDG_EVENT_NAME("ReflectionRayTracing(spp=%d) %dx%d", SamplePerPixel, DispatchSizeX, DispatchSizeY),
+								RDG_EVENT_NAME("ReflectionRayTracing(spp=%d) %dx%d", Options.SamplesPerPixel, DispatchSizeX, DispatchSizeY),
 								TilePassParameters,
 								ERDGPassFlags::Compute,
 								[TilePassParameters, this, &View, RayGenShader, DispatchSizeX, DispatchSizeY, DeferredMaterialBufferNumElements, DeferredMaterialMode](FRHICommandList& RHICmdList)
diff --git a/Engine/Source/Runtime/Renderer/Private/RayTracing/RayTracingReflections.h b/Engine/Source/Runtime/Renderer/Private/RayTracing/RayTracingReflections.h
--- a/Engine/Source/Runtime/Renderer/Private/RayTracing/RayTracingReflections.h
+++ b/Engine/Source/Runtime/Renderer/Private/RayTracing/RayTracingReflections.h
@@ -7,6 +7,21 @@
 class FViewInfo;
 class FScene;
 
+struct FRayTracingReflectionOptions
+{
+	enum EAlgorithm
+	{
+		BruteForce,
+		Sorted,
+		SortedDeferred
+	};
+
+	EAlgorithm Algorithm = EAlgorithm::Sorted;
+	int SamplesPerPixel = 1;
+	float ResolutionFraction = 1.0;
+	bool bReflectOnlyWater = false;
+};
+
 #if RHI_RAYTRACING
 
 int32 GetRayTracingReflectionsSamplesPerPixel(const FViewInfo& View);
@@ -14,6 +29,7 @@ int32 GetRayTracingReflectionsSamplesPerPixel(const FViewInfo& View);
 bool ShouldRenderRayTracingReflections(const FViewInfo& View);
 bool ShouldRayTracedReflectionsUseHybridReflections();
 bool ShouldRayTracedReflectionsSortMaterials(const FViewInfo& View);
+bool ShouldRayTracedReflectionsUseSortedDeferredAlgorithm(const FViewInfo& View);
 bool ShouldRayTracedReflectionsRayTraceSkyLightContribution(const FScene& Scene);
 
 #else
diff --git a/Engine/Source/Runtime/Renderer/Private/SingleLayerWaterRendering.cpp b/Engine/Source/Runtime/Renderer/Private/SingleLayerWaterRendering.cpp
--- a/Engine/Source/Runtime/Renderer/Private/SingleLayerWaterRendering.cpp
+++ b/Engine/Source/Runtime/Renderer/Private/SingleLayerWaterRendering.cpp
@@ -690,7 +690,7 @@ void FDeferredShadingSceneRenderer::RenderSingleLayerWaterReflections(FRHIComman
 			// Water is assumed to have zero roughness and is not currently denoised.
 			//int32 DenoiserMode = GetReflectionsDenoiserMode();
 			//bool bDenoise = DenoiserMode != 0;
-			int DenoiserMode = 0;
+			int32 DenoiserMode = 0;
 			bool bDenoise = false;
 
 			if (!bDenoise)
@@ -698,14 +698,35 @@ void FDeferredShadingSceneRenderer::RenderSingleLayerWaterReflections(FRHIComman
 				RayTracingConfig.ResolutionFraction = 1.0f;
 			}
 
+			FRayTracingReflectionOptions Options;
+			Options.Algorithm = FRayTracingReflectionOptions::BruteForce;
+			Options.SamplesPerPixel = 1;
+			Options.ResolutionFraction = 1.0;
+			Options.bReflectOnlyWater = true;
+
+			{
+				float UpscaleFactor = 1.0;
+				FRDGTextureDesc Desc = FRDGTextureDesc::Create2DDesc(
+					SceneTextures.SceneDepthBuffer->Desc.Extent / UpscaleFactor,
+					PF_FloatRGBA,
+					FClearValueBinding::None,
+					/* InFlags = */ TexCreate_None,
+					/* InTargetableFlags = */ TexCreate_ShaderResource | TexCreate_RenderTargetable | TexCreate_UAV,
+					/* bInForceSeparateTargetAndShaderResource = */ false);
+
+				DenoiserInputs.Color = GraphBuilder.CreateTexture(Desc, TEXT("RayTracingReflections"));
+
+				Desc.Format = PF_R16F;
+				DenoiserInputs.RayHitDistance = GraphBuilder.CreateTexture(Desc, TEXT("RayTracingReflectionsHitDistance"));
+				DenoiserInputs.RayImaginaryDepth = GraphBuilder.CreateTexture(Desc, TEXT("RayTracingReflectionsImaginaryDepth"));
+			}
+
 			bool bReflectOnlyWater = true;
 			RenderRayTracingReflections(
 				GraphBuilder,
 				SceneTextures,
 				View,
-				RayTracingConfig.RayCountPerPixel,
-				bReflectOnlyWater,
-				RayTracingConfig.ResolutionFraction,
+				Options,
 				&DenoiserInputs);
 
 			if (bDenoise)
