Add roughness threshold to Ray Traced Deferred Reflections

diff --git a/Engine/Shaders/Private/RayTracing/RayTracingDeferredReflections.usf b/Engine/Shaders/Private/RayTracing/RayTracingDeferredReflections.usf
--- a/Engine/Shaders/Private/RayTracing/RayTracingDeferredReflections.usf
+++ b/Engine/Shaders/Private/RayTracing/RayTracingDeferredReflections.usf
@@ -41,6 +41,7 @@
 uint2 RayTracingResolution;
 uint2 TileAlignedResolution;
 float ReflectionMaxNormalBias;
+float ReflectionMaxRoughness;
 
 RWStructuredBuffer<FSortedReflectionRay>        RayBuffer;
 RWStructuredBuffer<FDeferredMaterialPayload>    MaterialBuffer;
@@ -75,13 +76,19 @@ RAY_TRACING_ENTRY_RAYGEN(RayTracingDeferredReflectionsRGS)
 		DeferredMaterialPayload.SortKey = RAY_TRACING_DEFERRED_MATERIAL_KEY_RAY_MISS;
 		FRayIntersectionBookmark Bookmark = (FRayIntersectionBookmark)0;
 
-		TraceDeferredMaterialGatherRay(TLAS,
-			RAY_FLAG_FORCE_OPAQUE, // Masked and transparent materials are currently not supported
-			RAY_TRACING_MASK_OPAQUE,
-			Ray,
-			Bookmark,
-			DeferredMaterialPayload
-		);
+		FGBufferData GBuffer = GetGBufferDataFromSceneTexturesLoad(PixelPos);
+		float RoughnessFade = GetRoughnessFade(GBuffer.Roughness, ReflectionMaxRoughness);
+		bool bIsValidPixel = RoughnessFade > 0;
+		if (bIsValidPixel)
+		{
+			TraceDeferredMaterialGatherRay(TLAS,
+				RAY_FLAG_FORCE_OPAQUE, // Masked and transparent materials are currently not supported
+				RAY_TRACING_MASK_OPAQUE,
+				Ray,
+				Bookmark,
+				DeferredMaterialPayload
+			);
+		}
 
 		BookmarkBuffer[DispatchThreadId] = Bookmark;
 	}
diff --git a/Engine/Shaders/Private/RayTracing/RayTracingLightingCommon.ush b/Engine/Shaders/Private/RayTracing/RayTracingLightingCommon.ush
--- a/Engine/Shaders/Private/RayTracing/RayTracingLightingCommon.ush
+++ b/Engine/Shaders/Private/RayTracing/RayTracingLightingCommon.ush
@@ -180,12 +180,6 @@ FRectTexture GetRayTracingRectTextureData()
 	return RectTexture;
 }
 
-float GetRoughnessFade(float Roughness, float MaxRoughness)
-{
-	float RoughnessMaskScale = -2.0 / MaxRoughness;
-	return saturate(Roughness * RoughnessMaskScale + 2.0);
-}
-
 float3 GenerateReflectedRayDirection(
 	float3 IncidentDirection,
 	float3 WorldNormal,
diff --git a/Engine/Shaders/Private/RayTracing/RayTracingReflectionsCommon.ush b/Engine/Shaders/Private/RayTracing/RayTracingReflectionsCommon.ush
--- a/Engine/Shaders/Private/RayTracing/RayTracingReflectionsCommon.ush
+++ b/Engine/Shaders/Private/RayTracing/RayTracingReflectionsCommon.ush
@@ -41,3 +41,8 @@ uint2 UnpackPixelCoordinates(uint PixelCoordinates)
 		PixelCoordinates >> 16);
 }
 
+float GetRoughnessFade(float Roughness, float MaxRoughness)
+{
+	float RoughnessMaskScale = -2.0 / MaxRoughness;
+	return saturate(Roughness * RoughnessMaskScale + 2.0);
+}
diff --git a/Engine/Source/Runtime/Renderer/Private/RayTracing/RayTracingDeferredReflections.cpp b/Engine/Source/Runtime/Renderer/Private/RayTracing/RayTracingDeferredReflections.cpp
--- a/Engine/Source/Runtime/Renderer/Private/RayTracing/RayTracingDeferredReflections.cpp
+++ b/Engine/Source/Runtime/Renderer/Private/RayTracing/RayTracingDeferredReflections.cpp
@@ -26,6 +26,7 @@
 
 #include "RayTracing/RayTracingLighting.h"
 #include "RayTracing/RayTracingDeferredMaterials.h"
+#include "RayTracing/RayTracingReflections.h"
 #include "RendererPrivate.h"
 #include "GlobalShader.h"
 #include "DeferredShadingRenderer.h"
@@ -62,6 +63,7 @@ class FGenerateReflectionRaysCS : public FGlobalShader
 	SHADER_PARAMETER(FIntPoint, RayTracingResolution)
 	SHADER_PARAMETER(FIntPoint, TileAlignedResolution)
 	SHADER_PARAMETER(float, ReflectionMaxNormalBias)
+	SHADER_PARAMETER(float, ReflectionMaxRoughness)
 	SHADER_PARAMETER_STRUCT_REF(FViewUniformShaderParameters, ViewUniformBuffer)
 	SHADER_PARAMETER_STRUCT_INCLUDE(FSceneTextureParameters, SceneTextures)
 	SHADER_PARAMETER_RDG_BUFFER_UAV(RWStructuredBuffer<FSortedReflectionRay>, RayBuffer)
@@ -111,6 +113,7 @@ class FRayTracingDeferredReflectionsRGS : public FGlobalShader
 		SHADER_PARAMETER(FIntPoint, RayTracingResolution)
 		SHADER_PARAMETER(FIntPoint, TileAlignedResolution)
 		SHADER_PARAMETER(float, ReflectionMaxNormalBias)
+		SHADER_PARAMETER(float, ReflectionMaxRoughness)
 		SHADER_PARAMETER_SRV(RaytracingAccelerationStructure, TLAS)
 		SHADER_PARAMETER_RDG_BUFFER_UAV(RWStructuredBuffer<FSortedReflectionRay>, RayBuffer)
 		SHADER_PARAMETER_RDG_BUFFER_UAV(RWStructuredBuffer<FRayIntersectionBookmark>, BookmarkBuffer)
@@ -177,6 +180,7 @@ static void AddGenerateReflectionRaysPass(
 	PassParameters->RayTracingResolution    = CommonParameters.RayTracingResolution;
 	PassParameters->TileAlignedResolution   = CommonParameters.TileAlignedResolution;
 	PassParameters->ReflectionMaxNormalBias = CommonParameters.ReflectionMaxNormalBias;
+	PassParameters->ReflectionMaxRoughness  = CommonParameters.ReflectionMaxRoughness;
 	PassParameters->ViewUniformBuffer       = CommonParameters.ViewUniformBuffer;
 	PassParameters->SceneTextures           = CommonParameters.SceneTextures;
 	PassParameters->RayBuffer               = GraphBuilder.CreateUAV(RayBuffer);
@@ -221,6 +225,7 @@ void FDeferredShadingSceneRenderer::RenderRayTracingDeferredReflections(
 	FRayTracingDeferredReflectionsRGS::FParameters CommonParameters;
 	CommonParameters.RayTracingResolution    = RayTracingResolution;
 	CommonParameters.TileAlignedResolution   = TileAlignedResolution;
+	CommonParameters.ReflectionMaxRoughness  = GetRayTracingReflectionsMaxRoughness(View);
 	CommonParameters.TLAS                    = View.RayTracingScene.RayTracingSceneRHI->GetShaderResourceView();
 	CommonParameters.SceneTextures           = SceneTextures;
 	SetupSceneTextureSamplers(&CommonParameters.SceneTextureSamplers);
diff --git a/Engine/Source/Runtime/Renderer/Private/RayTracing/RayTracingReflections.cpp b/Engine/Source/Runtime/Renderer/Private/RayTracing/RayTracingReflections.cpp
--- a/Engine/Source/Runtime/Renderer/Private/RayTracing/RayTracingReflections.cpp
+++ b/Engine/Source/Runtime/Renderer/Private/RayTracing/RayTracingReflections.cpp
@@ -356,6 +356,11 @@ int32 GetRayTracingReflectionsSamplesPerPixel(const FViewInfo& View)
 	return GRayTracingReflectionsSamplesPerPixel >= 0 ? GRayTracingReflectionsSamplesPerPixel : View.FinalPostProcessSettings.RayTracingReflectionsSamplesPerPixel;
 }
 
+float GetRayTracingReflectionsMaxRoughness(const FViewInfo& View)
+{
+	return FMath::Clamp(GRayTracingReflectionsMaxRoughness >= 0 ? GRayTracingReflectionsMaxRoughness : View.FinalPostProcessSettings.RayTracingReflectionsMaxRoughness, 0.01f, 1.0f);
+}
+
 bool ShouldRenderRayTracingReflections(const FViewInfo& View)
 {
 	bool bThisViewHasRaytracingReflections = View.FinalPostProcessSettings.ReflectionsType == EReflectionsType::RayTracing;
@@ -579,7 +584,7 @@ void FDeferredShadingSceneRenderer::RenderRayTracingReflections(
 	CommonParameters.UpscaleFactor = UpscaleFactor;
 	CommonParameters.ReflectionMinRayDistance = FMath::Min(GRayTracingReflectionsMinRayDistance, GRayTracingReflectionsMaxRayDistance);
 	CommonParameters.ReflectionMaxRayDistance = GRayTracingReflectionsMaxRayDistance;
-	CommonParameters.ReflectionMaxRoughness = FMath::Clamp(GRayTracingReflectionsMaxRoughness >= 0 ? GRayTracingReflectionsMaxRoughness : View.FinalPostProcessSettings.RayTracingReflectionsMaxRoughness, 0.01f, 1.0f);
+	CommonParameters.ReflectionMaxRoughness = GetRayTracingReflectionsMaxRoughness(View);
 	CommonParameters.ReflectionMaxNormalBias = GetRaytracingMaxNormalBias();
 	CommonParameters.RayTracingResolution = RayTracingResolution;
 	CommonParameters.TileAlignedResolution = TileAlignedResolution;
diff --git a/Engine/Source/Runtime/Renderer/Private/RayTracing/RayTracingReflections.h b/Engine/Source/Runtime/Renderer/Private/RayTracing/RayTracingReflections.h
--- a/Engine/Source/Runtime/Renderer/Private/RayTracing/RayTracingReflections.h
+++ b/Engine/Source/Runtime/Renderer/Private/RayTracing/RayTracingReflections.h
@@ -25,6 +25,7 @@ struct FRayTracingReflectionOptions
 #if RHI_RAYTRACING
 
 int32 GetRayTracingReflectionsSamplesPerPixel(const FViewInfo& View);
+float GetRayTracingReflectionsMaxRoughness(const FViewInfo& View);
 
 bool ShouldRenderRayTracingReflections(const FViewInfo& View);
 bool ShouldRayTracedReflectionsUseHybridReflections();
