Adding appropriate #if RHI_RAYTRACING guards to new reflections algorithm functions when ray tracing is not defined.

diff --git a/Engine/Source/Runtime/Renderer/Private/RayTracing/RayTracingReflections.h b/Engine/Source/Runtime/Renderer/Private/RayTracing/RayTracingReflections.h
--- a/Engine/Source/Runtime/Renderer/Private/RayTracing/RayTracingReflections.h
+++ b/Engine/Source/Runtime/Renderer/Private/RayTracing/RayTracingReflections.h
@@ -40,4 +40,14 @@ FORCEINLINE int32 GetRayTracingReflectionsSamplesPerPixel(const FViewInfo& View)
 	return 0;
 }
 
+FORCEINLINE bool ShouldRayTracedReflectionsSortMaterials(const FViewInfo& View)
+{
+	return false;
+}
+
+FORCEINLINE bool ShouldRayTracedReflectionsUseSortedDeferredAlgorithm(const FViewInfo& View)
+{
+	return false;
+}
+
 #endif // RHI_RAYTRACING
