Switched to double-sided shadow rays in deferred reflection shader. Renamed new reflection pass output resources to match the vanilla reflection pass.

This fixes `vis RayTracingReflections` command and missing shadows in reflections in some scenes.

diff --git a/Engine/Shaders/Private/RayTracing/RayTracingDeferredReflections.usf b/Engine/Shaders/Private/RayTracing/RayTracingDeferredReflections.usf
--- a/Engine/Shaders/Private/RayTracing/RayTracingDeferredReflections.usf
+++ b/Engine/Shaders/Private/RayTracing/RayTracingDeferredReflections.usf
@@ -1,5 +1,7 @@
 // Copyright Epic Games, Inc. All Rights Reserved.
 
+#define RAY_TRACING_DEFERRED_REFLECTIONS_SHADER_VERSION 0x672EBB8E  // Change to force shader compilation of this shader
+
 // EDeferredMaterialMode
 #define DEFERRED_MATERIAL_MODE_NONE   0
 #define DEFERRED_MATERIAL_MODE_GATHER 1
diff --git a/Engine/Source/Runtime/Renderer/Private/RayTracing/RayTracingDeferredReflections.cpp b/Engine/Source/Runtime/Renderer/Private/RayTracing/RayTracingDeferredReflections.cpp
--- a/Engine/Source/Runtime/Renderer/Private/RayTracing/RayTracingDeferredReflections.cpp
+++ b/Engine/Source/Runtime/Renderer/Private/RayTracing/RayTracingDeferredReflections.cpp
@@ -149,6 +149,7 @@ class FRayTracingDeferredReflectionsRGS : public FGlobalShader
 	{
 		FGlobalShader::ModifyCompilationEnvironment(Parameters, OutEnvironment);
 		OutEnvironment.SetDefine(TEXT("UE_RAY_TRACING_DISPATCH_1D"), 1); // Always using 1D dispatches
+		OutEnvironment.SetDefine(TEXT("ENABLE_TWO_SIDED_GEOMETRY"), 1); // Always using double-sided ray tracing for shadow rays
 	}
 };
 IMPLEMENT_GLOBAL_SHADER(FRayTracingDeferredReflectionsRGS, "/Engine/Private/RayTracing/RayTracingDeferredReflections.usf", "RayTracingDeferredReflectionsRGS", SF_RayGen);
@@ -213,9 +214,9 @@ void FDeferredShadingSceneRenderer::RenderRayTracingDeferredReflections(
 		TexCreate_None, TexCreate_ShaderResource | TexCreate_UAV,
 		false);
 
-	OutDenoiserInputs->Color          = GraphBuilder.CreateTexture(OutputDesc, TEXT("RTDReflections"));
+	OutDenoiserInputs->Color          = GraphBuilder.CreateTexture(OutputDesc, TEXT("RayTracingReflections"));
 	OutputDesc.Format                 = PF_R16F;
-	OutDenoiserInputs->RayHitDistance = GraphBuilder.CreateTexture(OutputDesc, TEXT("RTDReflectionsHitDistance"));
+	OutDenoiserInputs->RayHitDistance = GraphBuilder.CreateTexture(OutputDesc, TEXT("RayTracingReflectionsHitDistance"));
 
 	const FIntPoint RayTracingResolution = View.ViewRect.Size();
 
