Consolidated ray traced reflection options gathering into a single helper function GetRayTracingReflectionOptions().

diff --git a/Engine/Source/Runtime/Renderer/Private/IndirectLightRendering.cpp b/Engine/Source/Runtime/Renderer/Private/IndirectLightRendering.cpp
--- a/Engine/Source/Runtime/Renderer/Private/IndirectLightRendering.cpp
+++ b/Engine/Source/Runtime/Renderer/Private/IndirectLightRendering.cpp
@@ -22,12 +22,6 @@ static TAutoConsoleVariable<int32> CVarDiffuseIndirectDenoiser(
 	TEXT("Denoising options (default = 1)"),
 	ECVF_RenderThreadSafe);
 
-static TAutoConsoleVariable<float> CVarReflectionScreenPercentage(
-	TEXT("r.RayTracing.Reflections.ScreenPercentage"),
-	100.0f,
-	TEXT("Screen percentage the reflections should be ray traced at (default = 100)."),
-	ECVF_RenderThreadSafe);
-
 static TAutoConsoleVariable<int32> CVarUseReflectionDenoiser(
 	TEXT("r.Reflections.Denoiser"),
 	2,
@@ -663,13 +657,14 @@ void FDeferredShadingSceneRenderer::RenderDeferredReflectionsAndSkyLighting(FRHI
 	for (FViewInfo& View : Views)
 	{
 		const uint32 CurrentViewIndex = ViewIndex++;
-		const bool bRayTracedReflections = ShouldRenderRayTracingReflections(View);
-		const bool bScreenSpaceReflections = !bRayTracedReflections && ShouldRenderScreenSpaceReflections(View);
 
-		const bool bComposePlanarReflections = !bRayTracedReflections && HasDeferredPlanarReflections(View);
+		const FRayTracingReflectionOptions RayTracingReflectionOptions = GetRayTracingReflectionOptions(View, *Scene);
+
+		const bool bScreenSpaceReflections = !RayTracingReflectionOptions.bEnabled && ShouldRenderScreenSpaceReflections(View);
+		const bool bComposePlanarReflections = !RayTracingReflectionOptions.bEnabled && HasDeferredPlanarReflections(View);
 
 		FRDGTextureRef ReflectionsColor = nullptr;
-		if (bRayTracedReflections || bScreenSpaceReflections)
+		if (RayTracingReflectionOptions.bEnabled || bScreenSpaceReflections)
 		{
 			int32 DenoiserMode = GetReflectionsDenoiserMode();
 
@@ -678,44 +673,29 @@ void FDeferredShadingSceneRenderer::RenderDeferredReflectionsAndSkyLighting(FRHI
 
 			// Traces the reflections, either using screen space reflection, or ray tracing.
 			IScreenSpaceDenoiser::FReflectionsInputs DenoiserInputs;
-			IScreenSpaceDenoiser::FReflectionsRayTracingConfig RayTracingConfig;
-			if (bRayTracedReflections)
+			IScreenSpaceDenoiser::FReflectionsRayTracingConfig DenoiserConfig;
+			if (RayTracingReflectionOptions.bEnabled)
 			{
 				RDG_EVENT_SCOPE(GraphBuilder, "RayTracingReflections");
 				RDG_GPU_STAT_SCOPE(GraphBuilder, RayTracingReflections);
 
 				bDenoise = DenoiserMode != 0;
 
-				RayTracingConfig.ResolutionFraction = FMath::Clamp(CVarReflectionScreenPercentage.GetValueOnRenderThread() / 100.0f, 0.25f, 1.0f);
-#if RHI_RAYTRACING
-				RayTracingConfig.RayCountPerPixel = GetRayTracingReflectionsSamplesPerPixel(View);
-#else
-				RayTracingConfig.RayCountPerPixel = 0;
-#endif
+				DenoiserConfig.ResolutionFraction = RayTracingReflectionOptions.ResolutionFraction;
+				DenoiserConfig.RayCountPerPixel = RayTracingReflectionOptions.SamplesPerPixel;
 
 				if (!bDenoise)
 				{
-					RayTracingConfig.ResolutionFraction = 1.0f;
+					DenoiserConfig.ResolutionFraction = 1.0f;
 				}
 
-				FRayTracingReflectionOptions Options;
-				if (ShouldRayTracedReflectionsUseSortedDeferredAlgorithm(View))
-				{
-					Options.Algorithm = FRayTracingReflectionOptions::SortedDeferred;
-				}
-				else if (ShouldRayTracedReflectionsSortMaterials(View))
-				{
-					Options.Algorithm = FRayTracingReflectionOptions::Sorted;
-				}
-				Options.SamplesPerPixel = RayTracingConfig.RayCountPerPixel;
-				Options.ResolutionFraction = RayTracingConfig.ResolutionFraction;
-				Options.bReflectOnlyWater = false;
+				check(RayTracingReflectionOptions.bReflectOnlyWater == false);
 
 				RenderRayTracingReflections(
 					GraphBuilder,
 					SceneTextures,
 					View,
-					Options,
+					RayTracingReflectionOptions,
 					&DenoiserInputs);
 			}
 			else if (bScreenSpaceReflections)
@@ -726,7 +706,7 @@ void FDeferredShadingSceneRenderer::RenderDeferredReflectionsAndSkyLighting(FRHI
 				FRDGTextureRef CurrentSceneColor = GraphBuilder.RegisterExternalTexture(SceneContext.GetSceneColor());
 
 				ESSRQuality SSRQuality;
-				GetSSRQualityForView(View, &SSRQuality, &RayTracingConfig);
+				GetSSRQualityForView(View, &SSRQuality, &DenoiserConfig);
 
 				RDG_EVENT_SCOPE(GraphBuilder, "ScreenSpaceReflections(Quality=%d)", int32(SSRQuality));
 
@@ -755,7 +735,7 @@ void FDeferredShadingSceneRenderer::RenderDeferredReflectionsAndSkyLighting(FRHI
 					&View.PrevViewInfo,
 					SceneTextures,
 					DenoiserInputs,
-					RayTracingConfig);
+					DenoiserConfig);
 
 				ReflectionsColor = DenoiserOutputs.Color;
 			}
@@ -779,7 +759,7 @@ void FDeferredShadingSceneRenderer::RenderDeferredReflectionsAndSkyLighting(FRHI
 			}
 			else
 			{
-				if (bRayTracedReflections && DenoiserInputs.RayHitDistance)
+				if (RayTracingReflectionOptions.bEnabled && DenoiserInputs.RayHitDistance)
 				{
 					// The performance of ray tracing does not allow to run without a denoiser in real time.
 					// Multiple rays per pixel is unsupported by the denoiser that will most likely more bound by to
@@ -790,11 +770,11 @@ void FDeferredShadingSceneRenderer::RenderDeferredReflectionsAndSkyLighting(FRHI
 
 				ReflectionsColor = DenoiserInputs.Color;
 			}
-		} // if (bRayTracedReflections || bScreenSpaceReflections)
+		} // if (RayTracingReflectionOptions.bEnabled || bScreenSpaceReflections)
 
 		if (bComposePlanarReflections)
 		{
-			check(!bRayTracedReflections);
+			check(!RayTracingReflectionOptions.bEnabled);
 			RenderDeferredPlanarReflections(GraphBuilder, SceneTextures, View, /* inout */ ReflectionsColor);
 		}
 
@@ -886,7 +866,9 @@ void FDeferredShadingSceneRenderer::RenderDeferredReflectionsAndSkyLighting(FRHI
 			// Bind hair data
 			const bool bCheckerboardSubsurfaceRendering = IsSubsurfaceCheckerboardFormat(SceneContext.GetSceneColorFormat());
 			auto PermutationVector = FReflectionEnvironmentSkyLightingPS::BuildPermutationVector(
-				View, bHasBoxCaptures, bHasSphereCaptures, DynamicBentNormalAO != NULL, bSkyLight, bDynamicSkyLight, bApplySkyShadowing, bRayTracedReflections);
+				View, bHasBoxCaptures, bHasSphereCaptures, DynamicBentNormalAO != nullptr,
+				bSkyLight, bDynamicSkyLight, bApplySkyShadowing,
+				RayTracingReflectionOptions.bEnabled);
 
 			TShaderMapRef<FReflectionEnvironmentSkyLightingPS> PixelShader(View.ShaderMap, PermutationVector);
 			ClearUnusedGraphResources(PixelShader, PassParameters);
diff --git a/Engine/Source/Runtime/Renderer/Private/RayTracing/RayTracingDeferredReflections.cpp b/Engine/Source/Runtime/Renderer/Private/RayTracing/RayTracingDeferredReflections.cpp
--- a/Engine/Source/Runtime/Renderer/Private/RayTracing/RayTracingDeferredReflections.cpp
+++ b/Engine/Source/Runtime/Renderer/Private/RayTracing/RayTracingDeferredReflections.cpp
@@ -226,7 +226,7 @@ void FDeferredShadingSceneRenderer::RenderRayTracingDeferredReflections(
 	FRayTracingDeferredReflectionsRGS::FParameters CommonParameters;
 	CommonParameters.RayTracingResolution    = RayTracingResolution;
 	CommonParameters.TileAlignedResolution   = TileAlignedResolution;
-	CommonParameters.ReflectionMaxRoughness  = GetRayTracingReflectionsMaxRoughness(View);
+	CommonParameters.ReflectionMaxRoughness  = Options.MaxRoughness;
 	CommonParameters.TLAS                    = View.RayTracingScene.RayTracingSceneRHI->GetShaderResourceView();
 	CommonParameters.SceneTextures           = SceneTextures;
 	SetupSceneTextureSamplers(&CommonParameters.SceneTextureSamplers);
diff --git a/Engine/Source/Runtime/Renderer/Private/RayTracing/RayTracingReflections.cpp b/Engine/Source/Runtime/Renderer/Private/RayTracing/RayTracingReflections.cpp
--- a/Engine/Source/Runtime/Renderer/Private/RayTracing/RayTracingReflections.cpp
+++ b/Engine/Source/Runtime/Renderer/Private/RayTracing/RayTracingReflections.cpp
@@ -33,6 +33,12 @@ static FAutoConsoleVariableRef CVarReflectionsMethod(
 	TEXT("1: use ray traced reflections\n")
 );
 
+static TAutoConsoleVariable<float> CVarReflectionScreenPercentage(
+	TEXT("r.RayTracing.Reflections.ScreenPercentage"),
+	100.0f,
+	TEXT("Screen percentage the reflections should be ray traced at (default = 100)."),
+	ECVF_RenderThreadSafe);
+
 static int32 GRayTracingReflectionsSamplesPerPixel = -1;
 static FAutoConsoleVariableRef CVarRayTracingReflectionsSamplesPerPixel(
 	TEXT("r.RayTracing.Reflections.SamplesPerPixel"),
@@ -351,12 +357,12 @@ class FSplitImaginaryReflectionGBufferCS : public FGlobalShader
 
 IMPLEMENT_GLOBAL_SHADER(FSplitImaginaryReflectionGBufferCS, "/Engine/Private/RayTracing/SplitImaginaryReflectionGBufferCS.usf", "MainCS", SF_Compute);
 
-int32 GetRayTracingReflectionsSamplesPerPixel(const FViewInfo& View)
+static int32 GetRayTracingReflectionsSamplesPerPixel(const FViewInfo& View)
 {
 	return GRayTracingReflectionsSamplesPerPixel >= 0 ? GRayTracingReflectionsSamplesPerPixel : View.FinalPostProcessSettings.RayTracingReflectionsSamplesPerPixel;
 }
 
-float GetRayTracingReflectionsMaxRoughness(const FViewInfo& View)
+static float GetRayTracingReflectionsMaxRoughness(const FViewInfo& View)
 {
 	return FMath::Clamp(GRayTracingReflectionsMaxRoughness >= 0 ? GRayTracingReflectionsMaxRoughness : View.FinalPostProcessSettings.RayTracingReflectionsMaxRoughness, 0.01f, 1.0f);
 }
@@ -372,26 +378,26 @@ bool ShouldRenderRayTracingReflections(const FViewInfo& View)
 	return IsRayTracingEnabled() && bReflectionPassEnabled;
 }
 
-bool ShouldRayTracedReflectionsUseHybridReflections()
+static bool ShouldRayTracedReflectionsUseHybridReflections()
 {
 	return CVarRayTracingReflectionsHybrid.GetValueOnRenderThread() != 0;
 }
 
-bool ShouldRayTracedReflectionsSortMaterials(const FViewInfo& View)
+static bool ShouldRayTracedReflectionsSortMaterials(const FViewInfo& View)
 {
 	const bool bIsMultiviewSecondary = View.ViewRect.Min.X > 0 || View.ViewRect.Min.Y > 0;
 
 	return (ShouldRayTracedReflectionsUseHybridReflections() || CVarRayTracingReflectionsSortMaterials.GetValueOnRenderThread() != 0) && !bIsMultiviewSecondary;
 }
 
-bool ShouldRayTracedReflectionsUseSortedDeferredAlgorithm(const FViewInfo& View)
+static bool ShouldRayTracedReflectionsUseSortedDeferredAlgorithm(const FViewInfo& View)
 {
 	const bool bIsMultiviewSecondary = View.ViewRect.Min.X > 0 || View.ViewRect.Min.Y > 0;
 
 	return (CVarRayTracingReflectionsExperimentalDeferred.GetValueOnRenderThread() != 0) && !bIsMultiviewSecondary;
 }
 
-bool ShouldRayTracedReflectionsRayTraceSkyLightContribution(const FScene& Scene)
+static bool ShouldRayTracedReflectionsRayTraceSkyLightContribution(const FScene& Scene)
 {
 	// Only ray trace sky light contribution when the ray traced sky light should be rendered in normal conditions (sky light exists, ray traced shadows enabled)
 	return CVarRayTracingReflectionsRayTraceSkyLightContribution.GetValueOnRenderThread() && ShouldRenderRayTracingSkyLight(Scene.SkyLight);
@@ -825,3 +831,38 @@ void FDeferredShadingSceneRenderer::RenderRayTracingReflections(
 	check(0);
 }
 #endif
+
+FRayTracingReflectionOptions GetRayTracingReflectionOptions(const FViewInfo& View, const FScene& Scene)
+{
+	FRayTracingReflectionOptions Result;
+
+#if RHI_RAYTRACING
+
+	Result.bEnabled = ShouldRenderRayTracingReflections(View);
+
+	if (ShouldRayTracedReflectionsUseSortedDeferredAlgorithm(View))
+	{
+		Result.Algorithm = FRayTracingReflectionOptions::SortedDeferred;
+	}
+	else if (ShouldRayTracedReflectionsSortMaterials(View))
+	{
+		Result.Algorithm = FRayTracingReflectionOptions::Sorted;
+	}
+	else
+	{
+		Result.Algorithm = FRayTracingReflectionOptions::EAlgorithm::BruteForce;
+	}
+
+	Result.SamplesPerPixel = GetRayTracingReflectionsSamplesPerPixel(View);
+	Result.ResolutionFraction = FMath::Clamp(CVarReflectionScreenPercentage.GetValueOnRenderThread() / 100.0f, 0.25f, 1.0f);
+	Result.MaxRoughness = GetRayTracingReflectionsMaxRoughness(View);
+	Result.bSkyLight = ShouldRayTracedReflectionsRayTraceSkyLightContribution(Scene);
+
+#else // RHI_RAYTRACING
+
+	Result.bEnabled = false;
+
+#endif // RHI_RAYTRACING
+
+	return Result;
+}
diff --git a/Engine/Source/Runtime/Renderer/Private/RayTracing/RayTracingReflections.h b/Engine/Source/Runtime/Renderer/Private/RayTracing/RayTracingReflections.h
--- a/Engine/Source/Runtime/Renderer/Private/RayTracing/RayTracingReflections.h
+++ b/Engine/Source/Runtime/Renderer/Private/RayTracing/RayTracingReflections.h
@@ -16,38 +16,13 @@ struct FRayTracingReflectionOptions
 		SortedDeferred
 	};
 
+	bool bEnabled = true;
 	EAlgorithm Algorithm = EAlgorithm::Sorted;
-	int SamplesPerPixel = 1;
-	float ResolutionFraction = 1.0;
+	int32 SamplesPerPixel = 1;
+	float ResolutionFraction = 1.0f;
+	float MaxRoughness = 1.0f;
 	bool bReflectOnlyWater = false;
+	bool bSkyLight = true;
 };
 
-#if RHI_RAYTRACING
-
-int32 GetRayTracingReflectionsSamplesPerPixel(const FViewInfo& View);
-float GetRayTracingReflectionsMaxRoughness(const FViewInfo& View);
-
-bool ShouldRenderRayTracingReflections(const FViewInfo& View);
-bool ShouldRayTracedReflectionsUseHybridReflections();
-bool ShouldRayTracedReflectionsSortMaterials(const FViewInfo& View);
-bool ShouldRayTracedReflectionsUseSortedDeferredAlgorithm(const FViewInfo& View);
-bool ShouldRayTracedReflectionsRayTraceSkyLightContribution(const FScene& Scene);
-
-#else
-
-FORCEINLINE int32 GetRayTracingReflectionsSamplesPerPixel(const FViewInfo& View)
-{
-	return 0;
-}
-
-FORCEINLINE bool ShouldRayTracedReflectionsSortMaterials(const FViewInfo& View)
-{
-	return false;
-}
-
-FORCEINLINE bool ShouldRayTracedReflectionsUseSortedDeferredAlgorithm(const FViewInfo& View)
-{
-	return false;
-}
-
-#endif // RHI_RAYTRACING
+FRayTracingReflectionOptions GetRayTracingReflectionOptions(const FViewInfo& View, const FScene& Scene);
