Use the same reflection environment shader permutation for deferred reflections as for SSR. Fixes deferred reflections on clearcoat materials.

diff --git a/Engine/Source/Runtime/Renderer/Private/IndirectLightRendering.cpp b/Engine/Source/Runtime/Renderer/Private/IndirectLightRendering.cpp
--- a/Engine/Source/Runtime/Renderer/Private/IndirectLightRendering.cpp
+++ b/Engine/Source/Runtime/Renderer/Private/IndirectLightRendering.cpp
@@ -865,10 +865,16 @@ void FDeferredShadingSceneRenderer::RenderDeferredReflectionsAndSkyLighting(FRHI
 
 			// Bind hair data
 			const bool bCheckerboardSubsurfaceRendering = IsSubsurfaceCheckerboardFormat(SceneContext.GetSceneColorFormat());
+
+			// ScreenSpace and SortedDeferred ray traced reflections use the same reflection environment shader,
+			// but main RT reflection shader requires a custom path as it evaluates the clear coat BRDF differently.
+			const bool bRequiresSpecializedReflectionEnvironmentShader = RayTracingReflectionOptions.bEnabled
+				&& RayTracingReflectionOptions.Algorithm != FRayTracingReflectionOptions::EAlgorithm::SortedDeferred;
+
 			auto PermutationVector = FReflectionEnvironmentSkyLightingPS::BuildPermutationVector(
 				View, bHasBoxCaptures, bHasSphereCaptures, DynamicBentNormalAO != nullptr,
 				bSkyLight, bDynamicSkyLight, bApplySkyShadowing,
-				RayTracingReflectionOptions.bEnabled);
+				bRequiresSpecializedReflectionEnvironmentShader);
 
 			TShaderMapRef<FReflectionEnvironmentSkyLightingPS> PixelShader(View.ShaderMap, PermutationVector);
 			ClearUnusedGraphResources(PixelShader, PassParameters);
