Implement a mini depth pass to mark pixels fading out in dithered LOD transition, so that ray tracing passes can use the mask to workaround/suppress dithered LOD transition shadow artifacts

diff --git a/Engine/Shaders/Private/RayTracing/RayTracingOcclusionRGS.usf b/Engine/Shaders/Private/RayTracing/RayTracingOcclusionRGS.usf
--- a/Engine/Shaders/Private/RayTracing/RayTracingOcclusionRGS.usf
+++ b/Engine/Shaders/Private/RayTracing/RayTracingOcclusionRGS.usf
@@ -413,7 +413,12 @@ FOcclusionResult ComputeOcclusion(
 			/* out */ Ray.Direction,
 			/* out */ Ray.TMin,
 			/* out */ Ray.TMax);
-		ApplyCameraRelativeDepthBias(Ray, PixelCoord, DeviceZ, WorldNormal, NormalBias);
+			
+
+		uint Stencil = SceneStencilBuffer.Load(int3(PixelCoord, 0)) STENCIL_COMPONENT_SWIZZLE;
+		bool bDitheredLODFadingOut = Stencil & 1;
+		
+		ApplyCameraRelativeDepthBias(Ray, PixelCoord, DeviceZ, WorldNormal, NormalBias + (bDitheredLODFadingOut ? 0.8f : 0.0f));
 
 		BRANCH
 		if (!bIsValidRay)
diff --git a/Engine/Shaders/Private/SceneTextureParameters.ush b/Engine/Shaders/Private/SceneTextureParameters.ush
--- a/Engine/Shaders/Private/SceneTextureParameters.ush
+++ b/Engine/Shaders/Private/SceneTextureParameters.ush
@@ -10,6 +10,7 @@ uint bIsSceneLightingChannelsValid;
 
 // Matches FSceneTextureParameters
 Texture2D SceneDepthBuffer;
+Texture2D<uint2> SceneStencilBuffer;
 Texture2D SceneVelocityBuffer;
 Texture2D SceneGBufferA;
 Texture2D SceneGBufferB;
diff --git a/Engine/Source/Runtime/Engine/Public/SceneManagement.h b/Engine/Source/Runtime/Engine/Public/SceneManagement.h
--- a/Engine/Source/Runtime/Engine/Public/SceneManagement.h
+++ b/Engine/Source/Runtime/Engine/Public/SceneManagement.h
@@ -2949,12 +2949,12 @@ struct FLODMask
 	//LodFade value within the BVH based on the LodFadeMask in the GBuffer
 	bool ContainsRayTracedLOD(int32 LODIndex) const
 	{
-		return DitheredLODIndices[0] == LODIndex;
+		return DitheredLODIndices[1] == LODIndex;
 	}
 
 	int8 GetRayTracedLOD()
 	{
-		return DitheredLODIndices[0];
+		return DitheredLODIndices[1];
 	}
 
 	bool IsDithered() const
diff --git a/Engine/Source/Runtime/Renderer/Private/DeferredShadingRenderer.cpp b/Engine/Source/Runtime/Renderer/Private/DeferredShadingRenderer.cpp
--- a/Engine/Source/Runtime/Renderer/Private/DeferredShadingRenderer.cpp
+++ b/Engine/Source/Runtime/Renderer/Private/DeferredShadingRenderer.cpp
@@ -875,7 +875,7 @@ bool FDeferredShadingSceneRenderer::GatherRayTracingWorldInstances(FRHICommandLi
 					else
 					{
 						float LODScale = LODScaleCVarValue * View.LODDistanceFactor;
-						LODToRender = ComputeLODForMeshes(SceneInfo->StaticMeshRelevances, View, Bounds.BoxSphereBounds.Origin, Bounds.BoxSphereBounds.SphereRadius, ForcedLODLevel, MeshScreenSizeSquared, CurFirstLODIdx, LODScale, false);
+						LODToRender = ComputeLODForMeshes(SceneInfo->StaticMeshRelevances, View, Bounds.BoxSphereBounds.Origin, Bounds.BoxSphereBounds.SphereRadius, ForcedLODLevel, MeshScreenSizeSquared, CurFirstLODIdx, LODScale, true);
 					}
 
 					FRHIRayTracingGeometry* RayTracingGeometryInstance = SceneInfo->GetStaticRayTracingGeometryInstance(LODToRender.GetRayTracedLOD());
@@ -1683,7 +1683,7 @@ void FDeferredShadingSceneRenderer::Render(FRHICommandListImmediate& RHICmdList)
 
 	// Interpolation needs to happen after the skin cache run as there is a dependency 
 	// on the skin cache output.
-	const bool bRunHairStrands = IsHairStrandsEnable(Scene->GetShaderPlatform()) && (Views.Num() > 0) && !ViewFamily.bWorldIsPaused;
+	const bool bRunHairStrands = IsHairStrandsEnable(Scene->GetShaderPlatform()) && Views.Num() > 0;
 	FHairStrandClusterData HairClusterData;
 	if (bRunHairStrands)
 	{
@@ -2285,6 +2285,13 @@ void FDeferredShadingSceneRenderer::Render(FRHICommandListImmediate& RHICmdList)
 			
 		}
 
+#if RHI_RAYTRACING
+		if (IsRayTracingEnabled())
+		{
+			RenderDitheredLODFadingOutMask(RHICmdList, Views[0]);
+		}
+#endif
+
 		RHICmdList.SetCurrentStat(GET_STATID(STAT_CLM_Lighting));
 		{
 			RenderLights(RHICmdList, SortedLightSet, HairDatas);
diff --git a/Engine/Source/Runtime/Renderer/Private/DeferredShadingRenderer.h b/Engine/Source/Runtime/Renderer/Private/DeferredShadingRenderer.h
--- a/Engine/Source/Runtime/Renderer/Private/DeferredShadingRenderer.h
+++ b/Engine/Source/Runtime/Renderer/Private/DeferredShadingRenderer.h
@@ -533,6 +533,8 @@ private:
 		const FRayTracingReflectionOptions& Options,
 		IScreenSpaceDenoiser::FReflectionsInputs* OutDenoiserInputs);
 
+	void RenderDitheredLODFadingOutMask(FRHICommandListImmediate& RHICmdList, const FViewInfo& View);
+
 	void RenderRayTracingShadows(
 		FRDGBuilder& GraphBuilder,
 		const FSceneTextureParameters& SceneTextures,
diff --git a/Engine/Source/Runtime/Renderer/Private/DepthRendering.cpp b/Engine/Source/Runtime/Renderer/Private/DepthRendering.cpp
--- a/Engine/Source/Runtime/Renderer/Private/DepthRendering.cpp
+++ b/Engine/Source/Runtime/Renderer/Private/DepthRendering.cpp
@@ -577,6 +577,7 @@ void FDeferredShadingSceneRenderer::RenderPrePassEditorPrimitives(FRHICommandLis
 					bRespectUseAsOccluderFlag,
 					DepthDrawingMode,
 					false,
+					false,
 					DynamicMeshPassContext);
 
 				const uint64 DefaultBatchElementMask = ~0ull;
@@ -601,6 +602,7 @@ void FDeferredShadingSceneRenderer::RenderPrePassEditorPrimitives(FRHICommandLis
 					bRespectUseAsOccluderFlag,
 					DepthDrawingMode,
 					false,
+					false,
 					DynamicMeshPassContext);
 
 				const uint64 DefaultBatchElementMask = ~0ull;
@@ -901,7 +903,10 @@ void FDepthPassMeshProcessor::Process(
 
 	FMeshPassProcessorRenderState DrawRenderState(PassDrawRenderState);
 
-	SetDepthPassDitheredLODTransitionState(ViewIfDynamicMeshCommand, MeshBatch, StaticMeshId, DrawRenderState);
+	if (!bDitheredLODFadingOutMaskPass)
+	{
+		SetDepthPassDitheredLODTransitionState(ViewIfDynamicMeshCommand, MeshBatch, StaticMeshId, DrawRenderState);
+	}
 
 	FDepthOnlyShaderElementData ShaderElementData(0.0f);
 	ShaderElementData.InitializeMeshMaterialData(ViewIfDynamicMeshCommand, PrimitiveSceneProxy, MeshBatch, StaticMeshId, true);
@@ -1008,11 +1013,13 @@ FDepthPassMeshProcessor::FDepthPassMeshProcessor(const FScene* Scene,
 	const bool InbRespectUseAsOccluderFlag,
 	const EDepthDrawingMode InEarlyZPassMode,
 	const bool InbEarlyZPassMovable,
+	const bool bDitheredLODFadingOutMaskPass,
 	FMeshPassDrawListContext* InDrawListContext)
 	: FMeshPassProcessor(Scene, Scene->GetFeatureLevel(), InViewIfDynamicMeshCommand, InDrawListContext)
 	, bRespectUseAsOccluderFlag(InbRespectUseAsOccluderFlag)
 	, EarlyZPassMode(InEarlyZPassMode)
 	, bEarlyZPassMovable(InbEarlyZPassMovable)
+	, bDitheredLODFadingOutMaskPass(bDitheredLODFadingOutMaskPass)
 {
 	PassDrawRenderState = InPassDrawRenderState;
 	PassDrawRenderState.SetViewUniformBuffer(Scene->UniformBuffers.ViewUniformBuffer);
@@ -1024,8 +1031,26 @@ FMeshPassProcessor* CreateDepthPassProcessor(const FScene* Scene, const FSceneVi
 {
 	FMeshPassProcessorRenderState DepthPassState;
 	SetupDepthPassState(DepthPassState);
-	return new(FMemStack::Get()) FDepthPassMeshProcessor(Scene, InViewIfDynamicMeshCommand, DepthPassState, true, Scene->EarlyZPassMode, Scene->bEarlyZPassMovable, InDrawListContext);
+	return new(FMemStack::Get()) FDepthPassMeshProcessor(Scene, InViewIfDynamicMeshCommand, DepthPassState, true, Scene->EarlyZPassMode, Scene->bEarlyZPassMovable, false, InDrawListContext);
 }
 
 FRegisterPassProcessorCreateFunction RegisterDepthPass(&CreateDepthPassProcessor, EShadingPath::Deferred, EMeshPass::DepthPass, EMeshPassFlags::CachedMeshCommands | EMeshPassFlags::MainView);
-FRegisterPassProcessorCreateFunction RegisterMobileDepthPass(&CreateDepthPassProcessor, EShadingPath::Mobile, EMeshPass::DepthPass, EMeshPassFlags::CachedMeshCommands | EMeshPassFlags::MainView);
\ No newline at end of file
+FRegisterPassProcessorCreateFunction RegisterMobileDepthPass(&CreateDepthPassProcessor, EShadingPath::Mobile, EMeshPass::DepthPass, EMeshPassFlags::CachedMeshCommands | EMeshPassFlags::MainView);
+
+FMeshPassProcessor* CreateDitheredLODFadingOutMaskPassProcessor(const FScene* Scene, const FSceneView* InViewIfDynamicMeshCommand, FMeshPassDrawListContext* InDrawListContext)
+{
+	FMeshPassProcessorRenderState DrawRenderState;
+
+	DrawRenderState.SetBlendState(TStaticBlendState<CW_NONE>::GetRHI());
+	DrawRenderState.SetDepthStencilState(
+		TStaticDepthStencilState<true, CF_Equal,
+		true, CF_Always, SO_Keep, SO_Keep, SO_Replace,
+		false, CF_Always, SO_Keep, SO_Keep, SO_Keep,
+		STENCIL_SANDBOX_MASK, STENCIL_SANDBOX_MASK
+		>::GetRHI());
+	DrawRenderState.SetStencilRef(STENCIL_SANDBOX_MASK);
+
+	return new(FMemStack::Get()) FDepthPassMeshProcessor(Scene, InViewIfDynamicMeshCommand, DrawRenderState, true, Scene->EarlyZPassMode, Scene->bEarlyZPassMovable, true, InDrawListContext);
+}
+
+FRegisterPassProcessorCreateFunction RegisterDitheredLODFadingOutMaskPass(&CreateDitheredLODFadingOutMaskPassProcessor, EShadingPath::Deferred, EMeshPass::DitheredLODFadingOutMaskPass, EMeshPassFlags::MainView);
diff --git a/Engine/Source/Runtime/Renderer/Private/DepthRendering.h b/Engine/Source/Runtime/Renderer/Private/DepthRendering.h
--- a/Engine/Source/Runtime/Renderer/Private/DepthRendering.h
+++ b/Engine/Source/Runtime/Renderer/Private/DepthRendering.h
@@ -233,6 +233,8 @@ public:
 		const bool InbRespectUseAsOccluderFlag,
 		const EDepthDrawingMode InEarlyZPassMode,
 		const bool InbEarlyZPassMovable,
+		/** Whether this mesh processor is being reused for rendering a pass that marks all fading out pixels on the screen */
+		const bool bDitheredLODFadingOutMaskPass,
 		FMeshPassDrawListContext* InDrawListContext);
 
 	virtual void AddMeshBatch(const FMeshBatch& RESTRICT MeshBatch, uint64 BatchElementMask, const FPrimitiveSceneProxy* RESTRICT PrimitiveSceneProxy, int32 StaticMeshId = -1) override final;
@@ -256,6 +258,42 @@ private:
 	const bool bRespectUseAsOccluderFlag;
 	const EDepthDrawingMode EarlyZPassMode;
 	const bool bEarlyZPassMovable;
+	const bool bDitheredLODFadingOutMaskPass;
 };
 
 extern void SetupDepthPassState(FMeshPassProcessorRenderState& DrawRenderState);
+
+class FRayTracingDitheredLODMeshProcessor : public FMeshPassProcessor
+{
+public:
+
+	FRayTracingDitheredLODMeshProcessor(const FScene* Scene,
+		const FSceneView* InViewIfDynamicMeshCommand,
+		const FMeshPassProcessorRenderState& InPassDrawRenderState,
+		const bool InbRespectUseAsOccluderFlag,
+		const EDepthDrawingMode InEarlyZPassMode,
+		const bool InbEarlyZPassMovable,
+		FMeshPassDrawListContext* InDrawListContext);
+
+	virtual void AddMeshBatch(const FMeshBatch& RESTRICT MeshBatch, uint64 BatchElementMask, const FPrimitiveSceneProxy* RESTRICT PrimitiveSceneProxy, int32 StaticMeshId = -1) override final;
+
+private:
+
+	template<bool bPositionOnly>
+	void Process(
+		const FMeshBatch& MeshBatch,
+		uint64 BatchElementMask,
+		int32 StaticMeshId,
+		EBlendMode BlendMode,
+		const FPrimitiveSceneProxy* RESTRICT PrimitiveSceneProxy,
+		const FMaterialRenderProxy& RESTRICT MaterialRenderProxy,
+		const FMaterial& RESTRICT MaterialResource,
+		ERasterizerFillMode MeshFillMode,
+		ERasterizerCullMode MeshCullMode);
+
+	FMeshPassProcessorRenderState PassDrawRenderState;
+
+	const bool bRespectUseAsOccluderFlag;
+	const EDepthDrawingMode EarlyZPassMode;
+	const bool bEarlyZPassMovable;
+};
diff --git a/Engine/Source/Runtime/Renderer/Private/RayTracing/RayTracingShadows.cpp b/Engine/Source/Runtime/Renderer/Private/RayTracing/RayTracingShadows.cpp
--- a/Engine/Source/Runtime/Renderer/Private/RayTracing/RayTracingShadows.cpp
+++ b/Engine/Source/Runtime/Renderer/Private/RayTracing/RayTracingShadows.cpp
@@ -310,3 +310,17 @@ void FDeferredShadingSceneRenderer::RenderRayTracingShadows(
 	unimplemented();
 }
 #endif
+
+void FDeferredShadingSceneRenderer::RenderDitheredLODFadingOutMask(FRHICommandListImmediate& RHICmdList, const FViewInfo& View)
+{
+	FSceneRenderTargets& SceneContext = FSceneRenderTargets::Get(RHICmdList);
+
+	SceneContext.BeginRenderingPrePass(RHICmdList, false);
+
+	RHICmdList.SetScissorRect(false, 0, 0, 0, 0);
+	RHICmdList.SetViewport(View.ViewRect.Min.X, View.ViewRect.Min.Y, 0.0f, View.ViewRect.Max.X, View.ViewRect.Max.Y, 1.0f);
+
+	View.ParallelMeshDrawCommandPasses[EMeshPass::DitheredLODFadingOutMaskPass].DispatchDraw(nullptr, RHICmdList);
+
+	SceneContext.FinishRenderingPrePass(RHICmdList);
+}
diff --git a/Engine/Source/Runtime/Renderer/Private/SceneTextureParameters.cpp b/Engine/Source/Runtime/Renderer/Private/SceneTextureParameters.cpp
--- a/Engine/Source/Runtime/Renderer/Private/SceneTextureParameters.cpp
+++ b/Engine/Source/Runtime/Renderer/Private/SceneTextureParameters.cpp
@@ -20,6 +20,7 @@ void SetupSceneTextureParameters(
 
 	// Should always have a depth buffer around allocated, since early z-pass is first.
 	OutTextures->SceneDepthBuffer = GraphBuilder.RegisterExternalTexture(SceneContext.SceneDepthZ, TEXT("SceneDepthZ"));
+	OutTextures->SceneStencilBuffer = GraphBuilder.CreateSRV(FRDGTextureSRVDesc::CreateWithPixelFormat(OutTextures->SceneDepthBuffer, PF_X24_G8));
 
 	// Registers all the scene texture from the scene context. No fallback is provided to catch mistake at shader parameter validation time
 	// when a pass is trying to access a resource before any other pass actually created it.
diff --git a/Engine/Source/Runtime/Renderer/Private/SceneTextureParameters.h b/Engine/Source/Runtime/Renderer/Private/SceneTextureParameters.h
--- a/Engine/Source/Runtime/Renderer/Private/SceneTextureParameters.h
+++ b/Engine/Source/Runtime/Renderer/Private/SceneTextureParameters.h
@@ -20,6 +20,7 @@ BEGIN_SHADER_PARAMETER_STRUCT(FSceneTextureParameters, )
 	SHADER_PARAMETER(uint32, bIsSceneLightingChannelsValid)
 
 	SHADER_PARAMETER_RDG_TEXTURE(Texture2D, SceneDepthBuffer)
+	SHADER_PARAMETER_RDG_TEXTURE_SRV(Texture2D<uint2>, SceneStencilBuffer)
 	SHADER_PARAMETER_RDG_TEXTURE(Texture2D, SceneVelocityBuffer)
 	SHADER_PARAMETER_RDG_TEXTURE(Texture2D, SceneGBufferA)
 	SHADER_PARAMETER_RDG_TEXTURE(Texture2D, SceneGBufferB)
diff --git a/Engine/Source/Runtime/Renderer/Private/SceneVisibility.cpp b/Engine/Source/Runtime/Renderer/Private/SceneVisibility.cpp
--- a/Engine/Source/Runtime/Renderer/Private/SceneVisibility.cpp
+++ b/Engine/Source/Runtime/Renderer/Private/SceneVisibility.cpp
@@ -2275,6 +2275,16 @@ struct FRelevancePacket
 							if ((StaticMeshRelevance.bUseForDepthPass && bDrawDepthOnly && ShadingPath != EShadingPath::Mobile) || bMobileIsInDepthPassMaskedMesh)
 							{
 								DrawCommandPacket.AddCommandsForMesh(PrimitiveIndex, PrimitiveSceneInfo, StaticMeshRelevance, StaticMesh, Scene, bCanCache, EMeshPass::DepthPass);
+
+#if RHI_RAYTRACING
+								if (IsRayTracingEnabled())
+								{
+									if (MarkMask & EMarkMaskBits::StaticMeshFadeOutDitheredLODMapMask)
+									{
+										DrawCommandPacket.AddCommandsForMesh(PrimitiveIndex, PrimitiveSceneInfo, StaticMeshRelevance, StaticMesh, Scene, bCanCache, EMeshPass::DitheredLODFadingOutMaskPass);
+									}
+								}
+#endif
 							}
 
 							// Mark static mesh as visible for rendering
diff --git a/Engine/Source/Runtime/Renderer/Private/ShadowSetup.cpp b/Engine/Source/Runtime/Renderer/Private/ShadowSetup.cpp
--- a/Engine/Source/Runtime/Renderer/Private/ShadowSetup.cpp
+++ b/Engine/Source/Runtime/Renderer/Private/ShadowSetup.cpp
@@ -2001,6 +2001,7 @@ void FProjectedShadowInfo::SetupMeshDrawCommandsForProjectionStenciling(FSceneRe
 				false,
 				DDM_AllOccluders,
 				false,
+				false,
 				&ProjectionStencilingContext);
 
 			// Pre-shadows mask by receiver elements, self-shadow mask by subject elements.
diff --git a/Engine/Source/Runtime/Renderer/Public/MeshPassProcessor.h b/Engine/Source/Runtime/Renderer/Public/MeshPassProcessor.h
--- a/Engine/Source/Runtime/Renderer/Public/MeshPassProcessor.h
+++ b/Engine/Source/Runtime/Renderer/Public/MeshPassProcessor.h
@@ -39,6 +39,7 @@ namespace EMeshPass
 		MobileBasePassCSM,  /** Mobile base pass with CSM shading enabled */
 		MobileInverseOpacity,  /** Mobile specific scene capture, Non-cached */
 		VirtualTexture,
+		DitheredLODFadingOutMaskPass, /** A mini depth pass used to mark pixels with dithered LOD fading out. Currently only used by ray tracing shadows. */
 
 #if WITH_EDITOR
 		HitProxy,
