Mobile DoF integration for FilmGrain Tonemap.TonemapperFilm set by default

diff --git a/Engine/Shaders/Private/PostProcessCommon.ush b/Engine/Shaders/Private/PostProcessCommon.ush
--- a/Engine/Shaders/Private/PostProcessCommon.ush
+++ b/Engine/Shaders/Private/PostProcessCommon.ush
@@ -216,3 +216,30 @@ bool IsComputeUVOutOfBounds(in float2 UV)
 	return (max(CenterDist.x, CenterDist.y) >= 0.5f);
 }
 
+// Constants for DOF blend in.
+half CocMaxRadiusInPixelsRcp()
+{
+	half2 MaxOffset = half2(-2.125, -0.50) * 2.0;
+	return rcp(sqrt(dot(MaxOffset, MaxOffset)));
+}
+
+half2 CocBlendScaleBias()
+{
+	half Start = 0.25 * CocMaxRadiusInPixelsRcp();
+	half End = 1.0 * CocMaxRadiusInPixelsRcp();
+	half2 ScaleBias;
+	ScaleBias.x = 1.0 / (End - Start);
+	ScaleBias.y = (-Start) * ScaleBias.x;
+	return ScaleBias;
+}
+
+half2 CocBlendScaleBiasFine()
+{
+	half Start = 0.0 * CocMaxRadiusInPixelsRcp();
+	half End = 0.5 * CocMaxRadiusInPixelsRcp();
+	half2 ScaleBias;
+	ScaleBias.x = 1.0 / (End - Start);
+	ScaleBias.y = (-Start) * ScaleBias.x;
+	return ScaleBias;
+}
+
diff --git a/Engine/Shaders/Private/PostProcessMobile.usf b/Engine/Shaders/Private/PostProcessMobile.usf
--- a/Engine/Shaders/Private/PostProcessMobile.usf
+++ b/Engine/Shaders/Private/PostProcessMobile.usf
@@ -774,6 +774,51 @@ void DofBlurPS_ES2(
 }
 
 
+// Integrate DOF
+
+void IntegrateDOFVS_ES2(
+	in float4 InPosition : ATTRIBUTE0,
+	in float2 InTexCoord : ATTRIBUTE1,
+	out float4 OutTexCoords : TEXCOORD0,
+	out float2 OutFineDofGrain : TEXCOORD1,
+	out float4 OutPosition : SV_POSITION
+)
+{
+	DrawRectangle(InPosition, InTexCoord, OutPosition, OutTexCoords.xy);
+
+	// Fine adjustment is inside the possible non-full viewport in the full resolution texture.
+	OutFineDofGrain.xy = OutTexCoords.xy + PostprocessInput0Size.zw * float2(-0.5, 0.5);
+	// Want grain and a second UV based on the knowledge that the source texture has a full viewport.
+	float2 FullViewUV = OutPosition.xy * float2(0.5, -0.5) + 0.5;
+	// For DOF attempt to undo sampling bias for the first transition region.
+	// This is better for the fine transition, breaks down for the larger bokeh.
+	// This is the best compromise for mobile using 4 bilinear taps only.
+	OutTexCoords.zw = FullViewUV.xy + PostprocessInput1Size.zw * float2(0.25, -0.5);
+}
+
+void IntegrateDOFPS_ES2(
+	in float4 TexCoords : TEXCOORD0,
+	in float2 FineDofGrain : TEXCOORD1,
+	out half4 OutColor : SV_Target0
+)
+{
+	half4 SceneColor = Texture2DSample(PostprocessInput0, PostprocessInput0Sampler, TexCoords.xy);
+
+	half4 DofFine = Texture2DSample(PostprocessInput0, PostprocessInput0Sampler, FineDofGrain.xy);
+	half4 Dof = Texture2DSample(PostprocessInput1, PostprocessInput1Sampler, TexCoords.zw);
+	// Convert alpha back into circle of confusion.
+	OutColor.a = max(Dof.a, abs(SceneColor.a * 2.0 - 1.0));
+	// Convert circle of confusion into blend factors.		
+	half2 ScaleBias = CocBlendScaleBias(); // Constant.
+	half DofAmount = saturate(SceneColor.a * ScaleBias.x + ScaleBias.y);
+	half2 ScaleBias2 = CocBlendScaleBiasFine(); // Constant.
+	half DofAmountFine = saturate(SceneColor.a * ScaleBias2.x + ScaleBias2.y);
+	// Blend in fine DOF.
+	OutColor.rgb = lerp(SceneColor.rgb, DofFine.rgb, DofAmountFine);
+	// Blend in coarse DOF.
+	OutColor.rgb = lerp(SceneColor.rgb, Dof.rgb, DofAmount);
+}
+
 //
 // First sun shaft blur and move sun intensity from alpha to single channel output.
 //
diff --git a/Engine/Shaders/Private/PostProcessTonemap.usf b/Engine/Shaders/Private/PostProcessTonemap.usf
--- a/Engine/Shaders/Private/PostProcessTonemap.usf
+++ b/Engine/Shaders/Private/PostProcessTonemap.usf
@@ -777,33 +777,6 @@ void MainVS_ES2(
 	}
 }
 
-// Constants for DOF blend in.
-half CocMaxRadiusInPixelsRcp() 
-{ 
-	half2 MaxOffset = half2(-2.125,-0.50)*2.0; 
-	return rcp(sqrt(dot(MaxOffset, MaxOffset))); 
-}
-
-half2 CocBlendScaleBias() 
-{
-	half Start = 0.25 * CocMaxRadiusInPixelsRcp();
-	half End = 1.0 * CocMaxRadiusInPixelsRcp();
-	half2 ScaleBias;
-	ScaleBias.x = 1.0/(End-Start);
-	ScaleBias.y = (-Start)*ScaleBias.x;
-	return ScaleBias;
-}
-
-half2 CocBlendScaleBiasFine() 
-{
-	half Start = 0.0 * CocMaxRadiusInPixelsRcp();
-	half End = 0.5 * CocMaxRadiusInPixelsRcp();
-	half2 ScaleBias;
-	ScaleBias.x = 1.0/(End-Start);
-	ScaleBias.y = (-Start)*ScaleBias.x;
-	return ScaleBias;
-}
-
 float4 OverlayColor;
 float SRGBAwareTarget;
 
diff --git a/Engine/Source/Runtime/Renderer/Private/PostProcess/PostProcessMobile.cpp b/Engine/Source/Runtime/Renderer/Private/PostProcess/PostProcessMobile.cpp
--- a/Engine/Source/Runtime/Renderer/Private/PostProcess/PostProcessMobile.cpp
+++ b/Engine/Source/Runtime/Renderer/Private/PostProcess/PostProcessMobile.cpp
@@ -2164,8 +2164,129 @@ FPooledRenderTargetDesc FRCPassPostProcessDofBlurES2::ComputeOutputDesc(EPassOut
 }
 
 
+class FPostProcessIntegrateDofPS_ES2 : public FGlobalShader
+{
+	DECLARE_SHADER_TYPE(FPostProcessIntegrateDofPS_ES2, Global);
+
+	static bool ShouldCompilePermutation(const FGlobalShaderPermutationParameters& Parameters)
+	{
+		return !IsConsolePlatform(Parameters.Platform);
+	}
+
+	FPostProcessIntegrateDofPS_ES2() {}
+
+public:
+	LAYOUT_FIELD(FPostProcessPassParameters, PostprocessParameter);
+
+	FPostProcessIntegrateDofPS_ES2(const ShaderMetaType::CompiledShaderInitializerType& Initializer)
+		: FGlobalShader(Initializer)
+	{
+		PostprocessParameter.Bind(Initializer.ParameterMap);
+	}
+
+	void SetPS(const FRenderingCompositePassContext& Context)
+	{
+		FRHIPixelShader* ShaderRHI = Context.RHICmdList.GetBoundPixelShader();
+		const FPostProcessSettings& Settings = Context.View.FinalPostProcessSettings;
+
+		FGlobalShader::SetParameters<FViewUniformShaderParameters>(Context.RHICmdList, ShaderRHI, Context.View.ViewUniformBuffer);
+		PostprocessParameter.SetPS(Context.RHICmdList, ShaderRHI, Context, TStaticSamplerState<SF_Bilinear, AM_Clamp, AM_Clamp, AM_Clamp>::GetRHI());
+	}
+};
+
+IMPLEMENT_SHADER_TYPE(, FPostProcessIntegrateDofPS_ES2, TEXT("/Engine/Private/PostProcessMobile.usf"), TEXT("IntegrateDOFPS_ES2"), SF_Pixel);
+
+
+class FPostProcessIntegrateDofVS_ES2 : public FGlobalShader
+{
+	DECLARE_SHADER_TYPE(FPostProcessIntegrateDofVS_ES2, Global);
+
+	static bool ShouldCompilePermutation(const FGlobalShaderPermutationParameters& Parameters)
+	{
+		return !IsConsolePlatform(Parameters.Platform);
+	}
+
+	FPostProcessIntegrateDofVS_ES2() {}
+
+public:
+	LAYOUT_FIELD(FPostProcessPassParameters, PostprocessParameter);
+
+	FPostProcessIntegrateDofVS_ES2(const ShaderMetaType::CompiledShaderInitializerType& Initializer)
+		: FGlobalShader(Initializer)
+	{
+		PostprocessParameter.Bind(Initializer.ParameterMap);
+	}
 
+	void SetVS(const FRenderingCompositePassContext& Context)
+	{
+		FRHIVertexShader* ShaderRHI = Context.RHICmdList.GetBoundVertexShader();
 
+		FGlobalShader::SetParameters<FViewUniformShaderParameters>(Context.RHICmdList, ShaderRHI, Context.View.ViewUniformBuffer);
+		PostprocessParameter.SetVS(ShaderRHI, Context, TStaticSamplerState<SF_Point, AM_Clamp, AM_Clamp, AM_Clamp>::GetRHI());
+	}
+};
+
+IMPLEMENT_SHADER_TYPE(, FPostProcessIntegrateDofVS_ES2, TEXT("/Engine/Private/PostProcessMobile.usf"), TEXT("IntegrateDOFVS_ES2"), SF_Vertex);
+
+void FRCPassIntegrateDofES2::Process(FRenderingCompositePassContext& Context)
+{
+	SCOPED_DRAW_EVENT(Context.RHICmdList, PostProcessIntegrateDof);
+
+	FIntRect DstRect;
+	DstRect.Min.X = 0;
+	DstRect.Min.Y = 0;
+	DstRect.Max.X = PrePostSourceViewportSize.X;
+	DstRect.Max.Y = PrePostSourceViewportSize.Y;
+
+	const FSceneRenderTargetItem& DestRenderTarget = PassOutputs[0].RequestSurface(Context);
+
+	ERenderTargetLoadAction LoadAction = ERenderTargetLoadAction::EClear;
+	FRHIRenderPassInfo RPInfo(DestRenderTarget.TargetableTexture, MakeRenderTargetActions(LoadAction, ERenderTargetStoreAction::EStore));
+	Context.RHICmdList.BeginRenderPass(RPInfo, TEXT("IntegrateDof"));
+	{
+
+		Context.SetViewportAndCallRHI(0, 0, 0.0f, PrePostSourceViewportSize.X, PrePostSourceViewportSize.Y, 1.0f);
+
+		FGraphicsPipelineStateInitializer GraphicsPSOInit;
+		Context.RHICmdList.ApplyCachedRenderTargets(GraphicsPSOInit);
+		GraphicsPSOInit.BlendState = TStaticBlendState<>::GetRHI();
+		GraphicsPSOInit.RasterizerState = TStaticRasterizerState<>::GetRHI();
+		GraphicsPSOInit.DepthStencilState = TStaticDepthStencilState<false, CF_Always>::GetRHI();
+
+		TShaderMapRef<FPostProcessIntegrateDofVS_ES2> VertexShader(Context.GetShaderMap());
+		TShaderMapRef<FPostProcessIntegrateDofPS_ES2> PixelShader(Context.GetShaderMap());
+
+		GraphicsPSOInit.BoundShaderState.VertexDeclarationRHI = GFilterVertexDeclaration.VertexDeclarationRHI;
+		GraphicsPSOInit.BoundShaderState.VertexShaderRHI = VertexShader.GetVertexShader();
+		GraphicsPSOInit.BoundShaderState.PixelShaderRHI = PixelShader.GetPixelShader();
+		GraphicsPSOInit.PrimitiveType = PT_TriangleList;
+
+		SetGraphicsPipelineState(Context.RHICmdList, GraphicsPSOInit);
+
+		VertexShader->SetVS(Context);
+		PixelShader->SetPS(Context);
+
+		DrawRectangle(
+			Context.RHICmdList,
+			0, 0,
+			PrePostSourceViewportSize.X, PrePostSourceViewportSize.Y,
+			0, 0,
+			PrePostSourceViewportSize.X, PrePostSourceViewportSize.Y,
+			PrePostSourceViewportSize,
+			PrePostSourceViewportSize,
+			VertexShader,
+			EDRF_UseTriangleOptimization);
+	}
+	Context.RHICmdList.EndRenderPass();
+	Context.RHICmdList.CopyToResolveTarget(DestRenderTarget.TargetableTexture, DestRenderTarget.ShaderResourceTexture, FResolveParams());
+}
+
+FPooledRenderTargetDesc FRCPassIntegrateDofES2::ComputeOutputDesc(EPassOutputId InPassOutputId) const
+{
+	FPooledRenderTargetDesc Ret = GetInput(ePId_Input0)->GetOutput()->RenderTargetDesc;
+	Ret.DebugName = TEXT("IntegrateDof");
+	return Ret;
+}
 
 
 //
diff --git a/Engine/Source/Runtime/Renderer/Private/PostProcess/PostProcessMobile.h b/Engine/Source/Runtime/Renderer/Private/PostProcess/PostProcessMobile.h
--- a/Engine/Source/Runtime/Renderer/Private/PostProcess/PostProcessMobile.h
+++ b/Engine/Source/Runtime/Renderer/Private/PostProcess/PostProcessMobile.h
@@ -80,6 +80,17 @@ private:
 	FIntPoint PrePostSourceViewportSize;
 };
 
+class FRCPassIntegrateDofES2 : public TRenderingCompositePassBase<2, 1>
+{
+public:
+	FRCPassIntegrateDofES2(FIntPoint InPrePostSourceViewportSize) : PrePostSourceViewportSize(InPrePostSourceViewportSize) { }
+	virtual void Process(FRenderingCompositePassContext& Context) override;
+	virtual FPooledRenderTargetDesc ComputeOutputDesc(EPassOutputId InPassOutputId) const override;
+	virtual void Release() override { delete this; }
+private:
+	FIntPoint PrePostSourceViewportSize;
+};
+
 class FRCPassPostProcessBloomDownES2 : public TRenderingCompositePassBase<1, 1>
 {
 public:
diff --git a/Engine/Source/Runtime/Renderer/Private/PostProcess/PostProcessTonemap.cpp b/Engine/Source/Runtime/Renderer/Private/PostProcess/PostProcessTonemap.cpp
--- a/Engine/Source/Runtime/Renderer/Private/PostProcess/PostProcessTonemap.cpp
+++ b/Engine/Source/Runtime/Renderer/Private/PostProcess/PostProcessTonemap.cpp
@@ -89,7 +89,7 @@ TAutoConsoleVariable<float> CVarGamma(
 
 TAutoConsoleVariable<int32> CVarMobileTonemapperFilm(
 	TEXT("r.Mobile.TonemapperFilm"),
-	0,
+	1,
 	TEXT("Whether mobile platforms should use new film tone mapper"),
 	ECVF_RenderThreadSafe);
 
diff --git a/Engine/Source/Runtime/Renderer/Private/PostProcess/PostProcessing.cpp b/Engine/Source/Runtime/Renderer/Private/PostProcess/PostProcessing.cpp
--- a/Engine/Source/Runtime/Renderer/Private/PostProcess/PostProcessing.cpp
+++ b/Engine/Source/Runtime/Renderer/Private/PostProcess/PostProcessing.cpp
@@ -1470,6 +1470,17 @@ void FPostProcessing::ProcessES2(FRHICommandListImmediate& RHICmdList, FScene* S
 							PostProcessDofBlur = FRenderingCompositeOutputRef(Pass);
 							DofOutput = PostProcessDofBlur;
 						}
+
+						if (bUseTonemapperFilm)
+						{
+							FRenderingCompositeOutputRef PostProcessIntegrateDof;
+							{
+								FRenderingCompositePass* Pass = Context.Graph.RegisterPass(new(FMemStack::Get()) FRCPassIntegrateDofES2(FinalOutputViewRect.Size()));
+								Pass->SetInput(ePId_Input0, Context.FinalOutput);
+								Pass->SetInput(ePId_Input1, DofOutput);
+								Context.FinalOutput = FRenderingCompositeOutputRef(Pass);
+							}
+						}
 					}
 					else
 					{
