VT feedback refactor to allow multiple contributors per frame

diff --git a/Engine/Source/Runtime/Renderer/Private/CompositionLighting/PostProcessDeferredDecals.cpp b/Engine/Source/Runtime/Renderer/Private/CompositionLighting/PostProcessDeferredDecals.cpp
--- a/Engine/Source/Runtime/Renderer/Private/CompositionLighting/PostProcessDeferredDecals.cpp
+++ b/Engine/Source/Runtime/Renderer/Private/CompositionLighting/PostProcessDeferredDecals.cpp
@@ -1071,11 +1071,6 @@ void FDecalRenderTargetManager::SetRenderTargetMode(FDecalRenderingCommon::ERend
 	RPInfo.DepthStencilRenderTarget.Action = MakeDepthStencilTargetActions(DepthTargetActions, ERenderTargetActions::Load_Store);
 	RPInfo.DepthStencilRenderTarget.ExclusiveDepthStencil = DepthStencilAccess;
 
-	if (UseVirtualTexturing(FeatureLevel))
-	{
-		SceneContext.BindVirtualTextureFeedbackUAV(RPInfo);
-	}
-
 	if (TargetsToTransitionWritable[CurrentRenderTargetMode])
 	{
 		TransitionRenderPassTargets(RHICmdList, RPInfo);
diff --git a/Engine/Source/Runtime/Renderer/Private/DeferredShadingRenderer.cpp b/Engine/Source/Runtime/Renderer/Private/DeferredShadingRenderer.cpp
--- a/Engine/Source/Runtime/Renderer/Private/DeferredShadingRenderer.cpp
+++ b/Engine/Source/Runtime/Renderer/Private/DeferredShadingRenderer.cpp
@@ -28,6 +28,7 @@
 #include "PipelineStateCache.h"
 #include "ClearQuad.h"
 #include "RendererModule.h"
+#include "VT/VirtualTextureFeedback.h"
 #include "VT/VirtualTextureSystem.h"
 #include "GPUScene.h"
 #include "RayTracing/RayTracingMaterialHitShaders.h"
@@ -1431,6 +1432,11 @@ void FDeferredShadingSceneRenderer::Render(FRHICommandListImmediate& RHICmdList)
 		{
 			SCOPED_GPU_STAT(RHICmdList, VirtualTextureUpdate);
 			FVirtualTextureSystem::Get().Update(RHICmdList, FeatureLevel, Scene);
+
+			// Clear virtual texture feedback to default value
+			FUnorderedAccessViewRHIRef FeedbackUAV = SceneContext.GetVirtualTextureFeedbackUAV();
+			RHICmdList.ClearUAVUint(FeedbackUAV, FUintVector4(~0u, ~0u, ~0u, ~0u));
+			RHICmdList.TransitionResource(EResourceTransitionAccess::ERWNoBarrier, EResourceTransitionPipeline::EGfxToGfx, FeedbackUAV);
 		}
 
 		for (int32 ViewIndex = 0; ViewIndex < Views.Num(); ViewIndex++)
@@ -2655,14 +2661,21 @@ void FDeferredShadingSceneRenderer::Render(FRHICommandListImmediate& RHICmdList)
 	if (bUseVirtualTexturing)
 	{
 		SCOPED_GPU_STAT(RHICmdList, VirtualTextureUpdate);
-		// No pass after this can make VT page requests
-		TArray<FIntRect, TInlineAllocator<FVirtualTextureFeedback::MaxRectPerTarget>> ViewRects;
+		
+		// No pass after this should make VT page requests
+		RHICmdList.TransitionResource(EResourceTransitionAccess::EReadable, EResourceTransitionPipeline::EGfxToGfx, SceneContext.VirtualTextureFeedbackUAV);
+
+		TArray<FIntRect, TInlineAllocator<4>> ViewRects;
 		ViewRects.AddUninitialized(Views.Num());
 		for (int32 ViewIndex = 0; ViewIndex < Views.Num(); ++ViewIndex)
 		{
 			ViewRects[ViewIndex] = Views[ViewIndex].ViewRect;
 		}
-		SceneContext.VirtualTextureFeedback.TransferGPUToCPU(RHICmdList, ViewRects);
+
+		FVirtualTextureFeedback::FBufferDesc Desc;
+		Desc.Init2D(SceneContext.GetBufferSizeXY(), ViewRects, SceneContext.GetVirtualTextureFeedbackScale());
+
+		GVirtualTextureFeedback.TransferGPUToCPU(RHICmdList, SceneContext.VirtualTextureFeedback, Desc);
 	}
 
 #if RHI_RAYTRACING
diff --git a/Engine/Source/Runtime/Renderer/Private/MobileShadingRenderer.cpp b/Engine/Source/Runtime/Renderer/Private/MobileShadingRenderer.cpp
--- a/Engine/Source/Runtime/Renderer/Private/MobileShadingRenderer.cpp
+++ b/Engine/Source/Runtime/Renderer/Private/MobileShadingRenderer.cpp
@@ -48,6 +48,7 @@
 #include "DebugViewModeRendering.h"
 #include "SkyAtmosphereRendering.h"
 #include "VisualizeTexture.h"
+#include "VT/VirtualTextureFeedback.h"
 #include "VT/VirtualTextureSystem.h"
 #include "GPUSortManager.h"
 
@@ -601,12 +602,6 @@ void FMobileSceneRenderer::Render(FRHICommandListImmediate& RHICmdList)
 	SceneColorRenderPassInfo.bOcclusionQueries = SceneColorRenderPassInfo.NumOcclusionQueries != 0;
 	SceneColorRenderPassInfo.bMultiviewPass = View.bIsMobileMultiViewEnabled;
 
-	// TODO: required only for DX11 ?
-	if (UseVirtualTexturing(ViewFeatureLevel) && !IsHlslccShaderPlatform(View.GetShaderPlatform()))
-	{
-		SceneContext.BindVirtualTextureFeedbackUAV(SceneColorRenderPassInfo);
-	}
-	
 	RHICmdList.BeginRenderPass(SceneColorRenderPassInfo, TEXT("SceneColorRendering"));
 
 	RHICmdList.SetCurrentStat(GET_STATID(STAT_CLM_MobilePrePass));
@@ -787,15 +782,18 @@ void FMobileSceneRenderer::Render(FRHICommandListImmediate& RHICmdList)
 
 	if (bUseVirtualTexturing)
 	{	
-		// No pass after this can make VT page requests
-		TArray<FIntRect, TInlineAllocator<FVirtualTextureFeedback::MaxRectPerTarget>> ViewRects;
+		// No pass after this should make VT page requests
+		TArray<FIntRect, TInlineAllocator<4>> ViewRects;
 		ViewRects.AddUninitialized(Views.Num());
 		for (int32 ViewIndex = 0; ViewIndex < Views.Num(); ++ViewIndex)
 		{
 			ViewRects[ViewIndex] = Views[ViewIndex].ViewRect;
 		}
 		
-		SceneContext.VirtualTextureFeedback.TransferGPUToCPU(RHICmdList, ViewRects);
+		FVirtualTextureFeedback::FBufferDesc Desc;
+		Desc.Init2D(SceneContext.GetBufferSizeXY(), ViewRects, SceneContext.GetVirtualTextureFeedbackScale());
+
+		GVirtualTextureFeedback.TransferGPUToCPU(RHICmdList, SceneContext.VirtualTextureFeedback, Desc);
 	}
 	
 	if (ViewFamily.bResolveScene)
diff --git a/Engine/Source/Runtime/Renderer/Private/PostProcess/SceneRenderTargets.cpp b/Engine/Source/Runtime/Renderer/Private/PostProcess/SceneRenderTargets.cpp
--- a/Engine/Source/Runtime/Renderer/Private/PostProcess/SceneRenderTargets.cpp
+++ b/Engine/Source/Runtime/Renderer/Private/PostProcess/SceneRenderTargets.cpp
@@ -875,6 +875,21 @@ int32 FSceneRenderTargets::GetGBufferRenderTargets(FRDGBuilder& GraphBuilder, ER
 	return MRTCount;
 }
 
+FUnorderedAccessViewRHIRef FSceneRenderTargets::GetVirtualTextureFeedbackUAV()
+{
+	return VirtualTextureFeedbackUAV.IsValid() ? VirtualTextureFeedbackUAV : GEmptyVertexBufferWithUAV->UnorderedAccessViewRHI;
+}
+
+int32 FSceneRenderTargets::GetVirtualTextureFeedbackScale()
+{
+	return GVirtualTextureFeedbackFactor;
+}
+
+FIntPoint FSceneRenderTargets::GetVirtualTextureFeedbackBufferSize()
+{
+	return FIntPoint::DivideAndRoundUp(BufferSize, FMath::Max(GVirtualTextureFeedbackFactor, 1));
+}
+
 #if !(UE_BUILD_SHIPPING || UE_BUILD_TEST)
 int32 FSceneRenderTargets::GetQuadOverdrawUAVIndex(EShaderPlatform Platform, ERHIFeatureLevel::Type FeatureLevel)
 {
@@ -926,23 +941,6 @@ FUnorderedAccessViewRHIRef FSceneRenderTargets::GetQuadOverdrawBufferUAV()
 	return GBlackTextureWithUAV->UnorderedAccessViewRHI;
 }
 
-
-
-FUnorderedAccessViewRHIRef FSceneRenderTargets::GetVirtualTextureFeedbackUAV()
-{
-	if (!VirtualTextureFeedback.FeedbackBufferUAV)
-	{
-		return GEmptyVertexBufferWithUAV->UnorderedAccessViewRHI;
-	}
-	
-	return VirtualTextureFeedback.FeedbackBufferUAV;
-}
-
-void FSceneRenderTargets::BindVirtualTextureFeedbackUAV(FRHIRenderPassInfo& RPInfo)
-{
-	return;
-}
-
 void FSceneRenderTargets::BeginRenderingGBuffer(FRHICommandList& RHICmdList, ERenderTargetLoadAction ColorLoadAction, ERenderTargetLoadAction DepthLoadAction, FExclusiveDepthStencil DepthStencilAccess, bool bBindQuadOverdrawBuffers, bool bClearQuadOverdrawBuffers, const FLinearColor& ClearColor/*=(0,0,0,1)*/, bool bIsWireframe)
 {
 	check(RHICmdList.IsOutsideRenderPass());
@@ -986,8 +984,6 @@ void FSceneRenderTargets::BeginRenderingGBuffer(FRHICommandList& RHICmdList, ERe
 		bSceneDepthCleared = true;
 	}
 
-	SetQuadOverdrawUAV(RHICmdList, bBindQuadOverdrawBuffers, bClearQuadOverdrawBuffers, RPInfo);
-
 	// Stencil always has to be store or certain VK drivers will leave the attachment in an undefined state.
 	RPInfo.DepthStencilRenderTarget.Action = MakeDepthStencilTargetActions(MakeRenderTargetActions(DepthLoadAction, DepthStoreAction), MakeRenderTargetActions(DepthLoadAction, ERenderTargetStoreAction::EStore));
 	RPInfo.DepthStencilRenderTarget.DepthStencilTarget = GetSceneDepthSurface();
@@ -995,10 +991,7 @@ void FSceneRenderTargets::BeginRenderingGBuffer(FRHICommandList& RHICmdList, ERe
 
 	RHICmdList.TransitionResource(DepthStencilAccess, RPInfo.DepthStencilRenderTarget.DepthStencilTarget);
 
-	if(UseVirtualTexturing(CurrentFeatureLevel) && !bBindQuadOverdrawBuffers )
-	{
-		BindVirtualTextureFeedbackUAV(RPInfo);
-	}	
+	SetQuadOverdrawUAV(RHICmdList, bBindQuadOverdrawBuffers, bClearQuadOverdrawBuffers, RPInfo);
 
 	RHICmdList.BeginRenderPass(RPInfo, TEXT("GBuffer"));
 
@@ -1765,11 +1758,6 @@ void FSceneRenderTargets::BeginRenderingTranslucency(FRHICommandList& RHICmdList
 	RPInfo.DepthStencilRenderTarget.DepthStencilTarget = GetSceneDepthSurface();
 	RPInfo.DepthStencilRenderTarget.ExclusiveDepthStencil = FExclusiveDepthStencil::DepthRead_StencilWrite;
 
-	if (UseVirtualTexturing(CurrentFeatureLevel))
-	{
-		BindVirtualTextureFeedbackUAV(RPInfo);
-	}
-
 	TransitionRenderPassTargets(RHICmdList, RPInfo);
 
 	RHICmdList.BeginRenderPass(RPInfo, TEXT("BeginRenderingTranslucency"));
@@ -1849,11 +1837,6 @@ void FSceneRenderTargets::BeginRenderingSeparateTranslucency(FRHICommandList& RH
 		RPInfo.ColorRenderTargets[0].Action = ERenderTargetActions::Clear_Store;
 	}
 	
-	if (UseVirtualTexturing(CurrentFeatureLevel))
-	{
-		BindVirtualTextureFeedbackUAV(RPInfo);
-	}
-
 	RHICmdList.TransitionResource(EResourceTransitionAccess::EWritable, RPInfo.ColorRenderTargets[0].RenderTarget);
 	RHICmdList.BeginRenderPass(RPInfo, TEXT("BeginRenderingSeparateTranslucency"));
 
@@ -1964,11 +1947,6 @@ void FSceneRenderTargets::BeginRenderingSeparateTranslucencyModulate(FRHICommand
 		// Clear the color target the first pass through and re-use after
 		RPInfo.ColorRenderTargets[0].Action = ERenderTargetActions::Clear_Store;
 	}
-	
-	if (UseVirtualTexturing(CurrentFeatureLevel))
-	{
-		BindVirtualTextureFeedbackUAV(RPInfo);
-	}
 
 	RHICmdList.TransitionResource(EResourceTransitionAccess::EWritable, RPInfo.ColorRenderTargets[0].RenderTarget);
 	RHICmdList.BeginRenderPass(RPInfo, TEXT("BeginRenderingSeparateTranslucencyModulate"));
@@ -2319,17 +2297,13 @@ void FSceneRenderTargets::AllocateMobileRenderTargets(FRHICommandListImmediate&
 		AllocMobileMultiViewDepth(RHICmdList, ScaleFactor);
 	}
 
+	AllocateVirtualTextureFeedbackBuffer(RHICmdList);
+
 	AllocateDebugViewModeTargets(RHICmdList);
 
 	EPixelFormat Format = GetSceneColor()->GetDesc().Format;
 
 	SceneAlphaCopy = GSystemTextures.MaxFP16Depth;
-	
-	if (UseVirtualTexturing(CurrentFeatureLevel))
-	{
-		FIntPoint FeedbackSize = FIntPoint::DivideAndRoundUp(BufferSize, FMath::Max(GVirtualTextureFeedbackFactor, 1));
-		VirtualTextureFeedback.CreateResourceGPU(RHICmdList, FeedbackSize);
-	}
 }
 
 // This is a helper class. It generates and provides N names with
@@ -2416,6 +2390,19 @@ void FSceneRenderTargets::AllocateReflectionTargets(FRHICommandList& RHICmdList,
 	}
 }
 
+void FSceneRenderTargets::AllocateVirtualTextureFeedbackBuffer(FRHICommandList& RHICmdList)
+{
+	if (UseVirtualTexturing(CurrentFeatureLevel))
+	{
+		const FIntPoint FeedbackSize = GetVirtualTextureFeedbackBufferSize();
+		const int32 FeedbackSizeBytes = FeedbackSize.X * FeedbackSize.Y * sizeof(uint32);
+
+		FRHIResourceCreateInfo CreateInfo(TEXT("VirtualTextureFeedbackGPU"));
+		VirtualTextureFeedback = RHICreateVertexBuffer(FeedbackSizeBytes, BUF_Static | BUF_ShaderResource | BUF_UnorderedAccess | BUF_SourceCopy, CreateInfo);
+		VirtualTextureFeedbackUAV = RHICreateUnorderedAccessView(VirtualTextureFeedback, PF_R32_UINT);
+	}
+}
+
 void FSceneRenderTargets::AllocateDebugViewModeTargets(FRHICommandList& RHICmdList)
 {
 	// If the shader/quad complexity shader need a quad overdraw buffer to be bind, allocate it.
@@ -2719,8 +2706,6 @@ void FSceneRenderTargets::AllocateDeferredShadingPathRenderTargets(FRHICommandLi
 		GRenderTargetPool.FindFreeElement(RHICmdList, Desc, LightAccumulation, TEXT("LightAccumulation"), true, ERenderTargetTransience::NonTransient);
 	}
 
-	AllocateDebugViewModeTargets(RHICmdList);
-
 	if (bAllocateVelocityGBuffer)
 	{
 		FPooledRenderTargetDesc VelocityRTDesc = FVelocityRendering::GetRenderTargetDesc();
@@ -2728,11 +2713,9 @@ void FSceneRenderTargets::AllocateDeferredShadingPathRenderTargets(FRHICommandLi
 		GRenderTargetPool.FindFreeElement(RHICmdList, VelocityRTDesc, SceneVelocity, TEXT("GBufferVelocity"));
 	}
 
-	if (UseVirtualTexturing(CurrentFeatureLevel))
-	{
-		FIntPoint FeedbackSize = FIntPoint::DivideAndRoundUp(BufferSize, FMath::Max(GVirtualTextureFeedbackFactor, 1));
-		VirtualTextureFeedback.CreateResourceGPU(RHICmdList, FeedbackSize);
-	}
+	AllocateVirtualTextureFeedbackBuffer(RHICmdList);
+
+	AllocateDebugViewModeTargets(RHICmdList);
 }
 
 EPixelFormat FSceneRenderTargets::GetDesiredMobileSceneColorFormat() const
@@ -2969,6 +2952,8 @@ void FSceneRenderTargets::ReleaseAllTargets()
 	MobileCustomDepth.SafeRelease();
 	MobileCustomStencil.SafeRelease();
 	CustomStencilSRV.SafeRelease();
+	VirtualTextureFeedback.SafeRelease();
+	VirtualTextureFeedbackUAV.SafeRelease();
 
 	for (int32 i = 0; i < UE_ARRAY_COUNT(OptionalShadowDepthColor); i++)
 	{
@@ -3001,8 +2986,6 @@ void FSceneRenderTargets::ReleaseAllTargets()
 	EditorPrimitivesDepth.SafeRelease();
 
 	FoveationTexture.SafeRelease();
-
-	VirtualTextureFeedback.ReleaseResources();
 }
 
 void FSceneRenderTargets::ReleaseDynamicRHI()
@@ -3482,14 +3465,7 @@ void SetupMobileSceneTextureUniformParameters(
 	SceneTextureParameters.MobileCustomStencilTexture = MobileCustomStencil;
 	SceneTextureParameters.MobileCustomStencilTextureSampler = TStaticSamplerState<SF_Point, AM_Clamp, AM_Clamp, AM_Clamp>::GetRHI();
 	
-	if (SceneContext.VirtualTextureFeedback.FeedbackBufferUAV.IsValid())
-	{
-		SceneTextureParameters.VirtualTextureFeedbackUAV = SceneContext.VirtualTextureFeedback.FeedbackBufferUAV;
-	}
-	else
-	{
-		SceneTextureParameters.VirtualTextureFeedbackUAV = GVirtualTextureFeedbackDummyResource.UAV;
-	}
+	SceneTextureParameters.VirtualTextureFeedbackUAV = SceneContext.GetVirtualTextureFeedbackUAV();
 
 	SceneTextureParameters.EyeAdaptationBuffer = GWhiteVertexBufferWithSRV->ShaderResourceViewRHI;
 }
diff --git a/Engine/Source/Runtime/Renderer/Private/PostProcess/SceneRenderTargets.h b/Engine/Source/Runtime/Renderer/Private/PostProcess/SceneRenderTargets.h
--- a/Engine/Source/Runtime/Renderer/Private/PostProcess/SceneRenderTargets.h
+++ b/Engine/Source/Runtime/Renderer/Private/PostProcess/SceneRenderTargets.h
@@ -15,7 +15,6 @@
 #include "SceneInterface.h"
 #include "SceneView.h"
 #include "RendererInterface.h"
-#include "../VT/VirtualTextureFeedback.h"
 
 class FViewInfo;
 class FRDGBuilder;
@@ -228,10 +227,11 @@ public:
 #endif
 
 	void SetQuadOverdrawUAV(FRHICommandList& RHICmdList, bool bBindQuadOverdrawBuffers, bool bClearQuadOverdrawBuffers, FRHIRenderPassInfo& Info);
-
 	FUnorderedAccessViewRHIRef GetQuadOverdrawBufferUAV();
+	
 	FUnorderedAccessViewRHIRef GetVirtualTextureFeedbackUAV();
-	void BindVirtualTextureFeedbackUAV(FRHIRenderPassInfo& RPInfo);
+	int32 GetVirtualTextureFeedbackScale();
+	FIntPoint GetVirtualTextureFeedbackBufferSize();
 
 	void BeginRenderingGBuffer(FRHICommandList& RHICmdList, ERenderTargetLoadAction ColorLoadAction, ERenderTargetLoadAction DepthLoadAction, FExclusiveDepthStencil DepthStencilAccess, bool bBindQuadOverdrawBuffers, bool bClearQuadOverdrawBuffers = false, const FLinearColor& ClearColor = FLinearColor(0, 0, 0, 1), bool bIsWireframe=false);
 	void FinishGBufferPassAndResolve(FRHICommandListImmediate& RHICmdList, FExclusiveDepthStencil DepthStencilAccess);
@@ -527,6 +527,8 @@ public:
 
 	void AllocateLightingChannelTexture(FRHICommandList& RHICmdList);
 
+	void AllocateVirtualTextureFeedbackBuffer(FRHICommandList& RHICmdList);
+
 	void AllocateDebugViewModeTargets(FRHICommandList& RHICmdList);
 
 	void AllocateScreenShadowMask(FRHICommandList& RHICmdList, TRefCountPtr<IPooledRenderTarget>& ScreenShadowMaskTexture, bool bCreateShaderResourceView = false);
@@ -632,15 +634,16 @@ public:
 	/** Texture to control variable resolution rendering */
 	TRefCountPtr<IPooledRenderTarget> FoveationTexture;
 
+	/** Virtual Texture feedback buffer bound as UAV during the base pass */
+	FVertexBufferRHIRef VirtualTextureFeedback;
+	FUnorderedAccessViewRHIRef VirtualTextureFeedbackUAV;
+
 	// todo: free ScreenSpaceAO so pool can reuse
 	bool bScreenSpaceAOIsValid;
 
 	// todo: free ScreenSpaceAO so pool can reuse
 	bool bCustomDepthIsValid;
 
-	/**	Virtual texture feedback buffer bound as UAV during basepass*/
-	FVirtualTextureFeedback VirtualTextureFeedback;
-
 private:
 	/** used by AdjustGBufferRefCount */
 	int32 GBufferRefCount;
diff --git a/Engine/Source/Runtime/Renderer/Private/SceneRendering.cpp b/Engine/Source/Runtime/Renderer/Private/SceneRendering.cpp
--- a/Engine/Source/Runtime/Renderer/Private/SceneRendering.cpp
+++ b/Engine/Source/Runtime/Renderer/Private/SceneRendering.cpp
@@ -1611,7 +1611,7 @@ void FViewInfo::SetupUniformBufferParameters(
 
 	ViewUniformShaderParameters.PreIntegratedBRDF = GEngine->PreIntegratedSkinBRDFTexture->Resource->TextureRHI;
 
-	ViewUniformShaderParameters.VirtualTextureFeedbackStride = SceneContext.VirtualTextureFeedback.GetFeedbackStride();
+	ViewUniformShaderParameters.VirtualTextureFeedbackStride = SceneContext.GetVirtualTextureFeedbackBufferSize().X;
 	ViewUniformShaderParameters.RuntimeVirtualTextureMipLevel = FVector4(ForceInitToZero);
 	ViewUniformShaderParameters.RuntimeVirtualTexturePackHeight = FVector2D(ForceInitToZero);
 	ViewUniformShaderParameters.RuntimeVirtualTextureDebugParams = FVector4(ForceInitToZero);
diff --git a/Engine/Source/Runtime/Renderer/Private/SingleLayerWaterRendering.cpp b/Engine/Source/Runtime/Renderer/Private/SingleLayerWaterRendering.cpp
--- a/Engine/Source/Runtime/Renderer/Private/SingleLayerWaterRendering.cpp
+++ b/Engine/Source/Runtime/Renderer/Private/SingleLayerWaterRendering.cpp
@@ -548,10 +548,6 @@ void FDeferredShadingSceneRenderer::BeginRenderingWaterGBuffer(FRHICommandList&
 	// Set other UAVs
 	const bool bClearQuadOverdrawBuffers = false;
 	SceneContext.SetQuadOverdrawUAV(RHICmdList, bBindQuadOverdrawBuffers, bClearQuadOverdrawBuffers, RPInfo);
-	if (UseVirtualTexturing(CurrentFeatureLevel) && !bBindQuadOverdrawBuffers)
-	{
-		SceneContext.BindVirtualTextureFeedbackUAV(RPInfo);
-	}
 
 	// Make the render targets writable
 	FRHITexture* TransitionRTs[MaxSimultaneousRenderTargets];
diff --git a/Engine/Source/Runtime/Renderer/Private/VT/VirtualTextureFeedback.cpp b/Engine/Source/Runtime/Renderer/Private/VT/VirtualTextureFeedback.cpp
--- a/Engine/Source/Runtime/Renderer/Private/VT/VirtualTextureFeedback.cpp
+++ b/Engine/Source/Runtime/Renderer/Private/VT/VirtualTextureFeedback.cpp
@@ -2,29 +2,33 @@
 
 #include "VirtualTextureFeedback.h"
 
-#include "ClearQuad.h"
-#include "VisualizeTexture.h"
-
 #if PLATFORM_WINDOWS
-// Temporary use of Querys until RHI has a good fence on D3D11
+// Use Query objects until RHI has a good fence on D3D11
 #define USE_RHI_FENCES 0
 #else
 #define USE_RHI_FENCES 1
 #endif
 
 #if USE_RHI_FENCES
-class FFeedbackFences
+
+/** Container for GPU fences. */
+class FFeedbackGPUFencePool
 {
 public:
-	FGPUFenceRHIRef Fences[FVirtualTextureFeedback::TargetCapacity];
+	TArray<FGPUFenceRHIRef> Fences;
+
+	FFeedbackGPUFencePool(int32 NumFences)
+	{
+		Fences.AddDefaulted(NumFences);
+	}
 
-	void Init(FRHICommandListImmediate& RHICmdList)
+	void InitRHI()
 	{
 	}
 
-	void ReleaseAll()
+	void ReleaseRHI()
 	{
-		for (int i = 0; i < FVirtualTextureFeedback::TargetCapacity; ++i)
+		for (int i = 0; i < Fences.Num(); ++i)
 		{
 			Fences[i].SafeRelease();
 		}
@@ -59,31 +63,40 @@ public:
 		Fences[Index].SafeRelease();
 	}
 };
-#else
-class FFeedbackFences
+
+#else // USE_RHI_FENCES
+
+/** Container for GPU fences. Implemented as GPU Queries. */
+class FFeedbackGPUFencePool
 {
 public:
 	FRenderQueryPoolRHIRef FenceQueryPool;
 	FGPUFenceRHIRef DummyFence;
-	FRHIPooledRenderQuery Fences[FVirtualTextureFeedback::TargetCapacity];
+	TArray<FRHIPooledRenderQuery> Fences;
+	bool bDummyFenceWritten = false;
 
-	void Init(FRHICommandListImmediate& RHICmdList)
+	FFeedbackGPUFencePool(int32 InSize)
+	{
+		Fences.AddDefaulted(InSize);
+	}
+
+	void InitRHI()
 	{
 		if (!FenceQueryPool.IsValid())
 		{
-			FenceQueryPool = RHICreateRenderQueryPool(RQT_AbsoluteTime, FVirtualTextureFeedback::TargetCapacity);
+			FenceQueryPool = RHICreateRenderQueryPool(RQT_AbsoluteTime, Fences.Num());
 		}
 
 		if (!DummyFence.IsValid())
 		{
-			DummyFence = RHICmdList.CreateGPUFence(FName());
-			RHICmdList.WriteGPUFence(DummyFence);
+			DummyFence = RHICreateGPUFence(FName());
+			bDummyFenceWritten = false;
 		}
 	}
 
-	void ReleaseAll()
+	void ReleaseRHI()
 	{
-		for (int i = 0; i < FVirtualTextureFeedback::TargetCapacity; ++i)
+		for (int i = 0; i < Fences.Num(); ++i)
 		{
 			if (Fences[i].IsValid())
 			{
@@ -92,6 +105,7 @@ public:
 		}
 
 		DummyFence.SafeRelease();
+		bDummyFenceWritten = false;
 
 		FenceQueryPool.SafeRelease();
 	}
@@ -102,6 +116,14 @@ public:
 		{
 			Fences[Index] = FenceQueryPool->AllocateQuery();
 		}
+
+		if (!bDummyFenceWritten && DummyFence.IsValid())
+		{
+			// Write dummy fence one time on first Allocate
+			// After that we want it to always Poll() true
+			RHICmdList.WriteGPUFence(DummyFence);
+			bDummyFenceWritten = true;
+		}
 	}
 	
 	void Write(FRHICommandListImmediate& RHICmdList, int32 Index)
@@ -125,146 +147,249 @@ public:
 		Fences[Index].ReleaseQuery();
 	}
 };
-#endif
 
-#include "RHIGPUReadback.h"
+#endif // USE_RHI_FENCES
 
-TGlobalResource<FVirtualTextureFeedbackDummyResource> GVirtualTextureFeedbackDummyResource;
+void FVirtualTextureFeedback::FBufferDesc::Init(int32 InBufferSize)
+{
+	BufferSize = FIntPoint(InBufferSize, 1);
+	NumRects = 0;
+	TotalReadSize = InBufferSize;
+}
+
+void FVirtualTextureFeedback::FBufferDesc::Init2D(FIntPoint InBufferSize)
+{
+	BufferSize = InBufferSize;
+	NumRects = 0;
+	TotalReadSize = BufferSize.X * BufferSize.Y;
+}
+
+void FVirtualTextureFeedback::FBufferDesc::Init2D(FIntPoint InUnscaledBufferSize, TArrayView<FIntRect> const& InUnscaledViewRects, int32 InBufferScale)
+{
+	const int32 BufferScale = FMath::Max(InBufferScale, 1);
+
+	BufferSize = FIntPoint::DivideAndRoundUp(InUnscaledBufferSize, BufferScale);
+	NumRects = 0;
+	TotalReadSize = BufferSize.X * BufferSize.Y;
+
+	if (InUnscaledViewRects.Num() > 0 && InUnscaledViewRects[0].Size() != InUnscaledBufferSize)
+	{
+		NumRects = FMath::Min((int32)MaxRectPerTransfer, InUnscaledViewRects.Num());
+		TotalReadSize = 0;
+
+		for (int32 RectIndex = 0; RectIndex < NumRects; ++RectIndex)
+		{
+			FIntRect const& Rect = InUnscaledViewRects[RectIndex];
+			Rects[RectIndex].Min = FIntPoint::DivideAndRoundDown(Rect.Min, BufferScale);
+			Rects[RectIndex].Max = FIntPoint::DivideAndRoundUp(Rect.Max, BufferScale);
+			TotalReadSize += Rects[RectIndex].Area();
+		}
+	}
+}
 
 FVirtualTextureFeedback::FVirtualTextureFeedback()
-	: Size(0, 0)
-	, NumBytes(0)
-	, GPUWriteIndex(0)
-	, CPUReadIndex(0)
-	, PendingTargetCount(0)
+	: NumPending(0)
+	, WriteIndex(0)
+	, ReadIndex(0)
 {
-	FeedBackFences = new FFeedbackFences;
+	Fences = new FFeedbackGPUFencePool(MaxTransfers);
 }
 
 FVirtualTextureFeedback::~FVirtualTextureFeedback()
 {
-	delete FeedBackFences;
+	delete Fences;
 }
 
-void FVirtualTextureFeedback::ReleaseResources()
+void FVirtualTextureFeedback::InitRHI()
 {
-	FeedbackBufferUAV.SafeRelease();
-	FeedbackBuffer.SafeRelease();
-	
-	for (int i = 0; i < TargetCapacity; ++i)
+	for (int32 Index = 0; Index < MaxTransfers; ++Index)
 	{
-		FeedbackCPU[i].ReadbackBuffer.SafeRelease();
+		FeedbackItems[Index].StagingBuffer = RHICreateStagingBuffer();
 	}
 
-	FeedBackFences->ReleaseAll();
-
-	CPUReadIndex = 0u;
-	GPUWriteIndex = 0u;
-	PendingTargetCount = 0u;
+	Fences->InitRHI();
 }
 
-void FVirtualTextureFeedback::CreateResourceGPU( FRHICommandListImmediate& RHICmdList, FIntPoint InSize)
+void FVirtualTextureFeedback::ReleaseRHI()
 {
-	if (Size != InSize || !FeedbackBuffer.IsValid())
+	for (int32 Index = 0; Index < MaxTransfers; ++Index)
 	{
-		Size = InSize;
-		NumBytes = Size.X * Size.Y * sizeof(uint32);
-
-		FRHIResourceCreateInfo CreateInfo(TEXT("VTFeedbackGPU"));
-		FeedbackBuffer = RHICreateVertexBuffer(NumBytes, BUF_Static | BUF_ShaderResource | BUF_UnorderedAccess | BUF_SourceCopy, CreateInfo);
-		FeedbackBufferUAV = RHICreateUnorderedAccessView(FeedbackBuffer, /*Format=*/ PF_R32_UINT);
+		FeedbackItems[Index].StagingBuffer.SafeRelease();
 	}
-	
-	// Clear to default value
-	RHICmdList.TransitionResource(EResourceTransitionAccess::ERWNoBarrier, EResourceTransitionPipeline::EGfxToCompute, FeedbackBufferUAV);
-	RHICmdList.ClearUAVUint(FeedbackBufferUAV.GetReference(), FUintVector4(~0u, ~0u, ~0u, ~0u));
-	RHICmdList.TransitionResource(EResourceTransitionAccess::ERWNoBarrier, EResourceTransitionPipeline::EComputeToCompute, FeedbackBufferUAV);
 
-	FeedBackFences->Init(RHICmdList);
+	Fences->ReleaseRHI();
 }
 
-void FVirtualTextureFeedback::TransferGPUToCPU( FRHICommandListImmediate& RHICmdList, TArrayView<FIntRect> const& ViewRects)
+void FVirtualTextureFeedback::TransferGPUToCPU(FRHICommandListImmediate& RHICmdList, FVertexBufferRHIRef const& Buffer, FBufferDesc const& Desc)
 {
 	QUICK_SCOPE_CYCLE_COUNTER(STAT_VirtualTextureFeedback_TransferGPUToCPU);
 
-	RHICmdList.TransitionResource( EResourceTransitionAccess::EReadable, EResourceTransitionPipeline::EGfxToGfx, FeedbackBufferUAV );
-
-	if (PendingTargetCount >= TargetCapacity)
+	if (NumPending >= MaxTransfers)
 	{
-		// If we have too many pending transfers, start throwing away the oldest
-		// We will need to allocate a new fence, since the previous fence will still be set on the old CopyToResolveTarget command (which we will now ignore/discard)
-		FeedBackFences->Release(CPUReadIndex);
-		--PendingTargetCount;
-		CPUReadIndex = (CPUReadIndex + 1) % TargetCapacity;
+		// If we have too many pending transfers, start throwing away the oldest in the ring buffer.
+		// We will need to allocate a new fence, since the previous fence will still be set on the old CopyToResolveTarget command (which we will now ignore/discard).
+		Fences->Release(ReadIndex);
+		NumPending --;
+		ReadIndex = (ReadIndex + 1) % MaxTransfers;
 	}
 
-	FFeedBackItem& FeedbackEntryCPU = FeedbackCPU[GPUWriteIndex];
-
-	FeedbackEntryCPU.NumRects = FMath::Min((int32)MaxRectPerTarget, ViewRects.Num());
-	for (int32 RectIndex = 0; RectIndex < FeedbackEntryCPU.NumRects; ++RectIndex)
-	{
-		FIntRect const& Rect = ViewRects[RectIndex];
-		//todo[vt]: Value of 16 has to match r.vt.FeedbackFactor
-		FeedbackEntryCPU.Rects[RectIndex].Min = FIntPoint(Rect.Min.X / 16, Rect.Min.Y / 16);
-		FeedbackEntryCPU.Rects[RectIndex].Max = FIntPoint((Rect.Max.X + 15) / 16, (Rect.Max.Y + 15) / 16);
-	}
-	
-	FeedBackFences->Allocate(RHICmdList, GPUWriteIndex);
+	FFeedbackItem& FeedbackItem = FeedbackItems[WriteIndex];
+	FeedbackItem.Desc = Desc;
 
 	// We only need to transfer 1 copy of the data, so restrict mask to the first active GPU.
-	FeedbackEntryCPU.GPUMask = FRHIGPUMask::FromIndex(RHICmdList.GetGPUMask().GetFirstIndex());
-	SCOPED_GPU_MASK(RHICmdList, FeedbackEntryCPU.GPUMask);
+	FeedbackItem.GPUMask = FRHIGPUMask::FromIndex(RHICmdList.GetGPUMask().GetFirstIndex());
+	SCOPED_GPU_MASK(RHICmdList, FeedbackItem.GPUMask);
 
-	FeedbackEntryCPU.ReadbackBuffer = RHICreateStagingBuffer();
-	RHICmdList.CopyToStagingBuffer(FeedbackBuffer, FeedbackEntryCPU.ReadbackBuffer, 0, NumBytes);
+	RHICmdList.CopyToStagingBuffer(Buffer, FeedbackItem.StagingBuffer, 0, Desc.BufferSize.X * Desc.BufferSize.Y * sizeof(uint32));
 
-	FeedBackFences->Write(RHICmdList, GPUWriteIndex);
+	Fences->Allocate(RHICmdList, WriteIndex);
+	Fences->Write(RHICmdList, WriteIndex);
 
-	GPUWriteIndex = (GPUWriteIndex + 1) % TargetCapacity;
-	++PendingTargetCount;
+	// Increment the ring buffer write position.
+	WriteIndex = (WriteIndex + 1) % MaxTransfers;
+	++NumPending;
 }
 
 bool FVirtualTextureFeedback::CanMap(FRHICommandListImmediate& RHICmdList)
 {
-	const FFeedBackItem& FeedbackEntryCPU = FeedbackCPU[CPUReadIndex];
-	return 
-		PendingTargetCount > 0u && 
-		FeedbackEntryCPU.ReadbackBuffer->IsValid() && 
-		FeedBackFences->Poll(RHICmdList, CPUReadIndex);
+	return NumPending > 0u && Fences->Poll(RHICmdList, ReadIndex);
 }
 
-bool FVirtualTextureFeedback::Map(FRHICommandListImmediate& RHICmdList, MapResult& OutResult)
+FVirtualTextureFeedback::FMapResult FVirtualTextureFeedback::Map(FRHICommandListImmediate& RHICmdList, int32 MaxTransfersToMap)
 {
 	QUICK_SCOPE_CYCLE_COUNTER(STAT_VirtualTextureFeedback_Map);
-		
-	FFeedBackItem& FeedbackEntryCPU = FeedbackCPU[CPUReadIndex];
-		
-	if (PendingTargetCount > 0u &&
-		FeedbackEntryCPU.ReadbackBuffer->IsValid() &&
-		FeedBackFences->Poll(RHICmdList, CPUReadIndex))
+
+	FVirtualTextureFeedback::FMapResult MapResult;
+
+	// Calculate number and size of available results.
+	int32 NumResults = 0;
+	int32 NumRects = 0;
+	int32 TotalReadSize = 0;
+	for (int32 ResultIndex = 0; ResultIndex < MaxTransfersToMap && ResultIndex < NumPending; ++ResultIndex)
 	{
-		SCOPED_GPU_MASK(RHICmdList, FeedbackEntryCPU.GPUMask);
-	
-		OutResult.NumRects = FeedbackEntryCPU.NumRects;
-		for (int32 i = 0; i < FeedbackEntryCPU.NumRects; ++i)
+		const int32 FeedbackIndex = (ReadIndex + ResultIndex) % MaxTransfers;
+		FBufferDesc const& FeedbackItemDesc = FeedbackItems[FeedbackIndex].Desc;
+
+		if (!Fences->Poll(RHICmdList, FeedbackIndex))
+		{
+			break;
+		}
+
+		NumResults ++;
+		NumRects += FeedbackItemDesc.NumRects;
+		TotalReadSize += FeedbackItemDesc.TotalReadSize;
+	}
+
+	// Fetch the valid results.
+	if (NumResults > 0)
+	{
+		// Get a FMapResources object to store anything that will need cleaning up on Unmap()
+		MapResult.MapHandle = FreeMapResources.Num() ? FreeMapResources.Pop() : MapResources.AddDefaulted();
+
+		if (NumResults == 1 && NumRects == 0)
+		{
+			// If only one target with no rectangles then fast path is to return the locked buffer.
+			const int32 FeedbackIndex = ReadIndex;
+			FBufferDesc const& FeedbackItemDesc = FeedbackItems[FeedbackIndex].Desc;
+			FRHIGPUMask GPUMask = FeedbackItems[FeedbackIndex].GPUMask;
+			FStagingBufferRHIRef StagingBuffer = FeedbackItems[FeedbackIndex].StagingBuffer;
+			
+			SCOPED_GPU_MASK(RHICmdList, GPUMask);
+			const int32 BufferSize = FeedbackItemDesc.BufferSize.X * FeedbackItemDesc.BufferSize.Y;
+			MapResult.Data = (uint32*)RHICmdList.LockStagingBuffer(StagingBuffer, Fences->GetMapFence(FeedbackIndex), 0, BufferSize * sizeof(uint32));
+			MapResult.Size = BufferSize;
+
+			// Store index so that we can unlock staging buffer when we call Unmap().
+			MapResources[MapResult.MapHandle].FeedbackItemToUnlockIndex = FeedbackIndex;
+		}
+		else
 		{
-			OutResult.Rects[i] = FeedbackEntryCPU.Rects[i];
+			// Concatenate the results to a single buffer (stored in the MapResources) and return that.
+			MapResources[MapResult.MapHandle].ResultData.SetNumUninitialized(TotalReadSize, false);
+			MapResult.Data = MapResources[MapResult.MapHandle].ResultData.GetData();
+			MapResult.Size = 0;
+
+			for (int32 ResultIndex = 0; ResultIndex < NumResults; ++ResultIndex)
+			{
+				const int32 FeedbackIndex = (ReadIndex + ResultIndex) % MaxTransfers;
+				FBufferDesc const& FeedbackItemDesc = FeedbackItems[FeedbackIndex].Desc;
+				FRHIGPUMask GPUMask = FeedbackItems[FeedbackIndex].GPUMask;
+				FStagingBufferRHIRef StagingBuffer = FeedbackItems[FeedbackIndex].StagingBuffer;
+
+				SCOPED_GPU_MASK(RHICmdList, GPUMask);
+				const int32 BufferSize = FeedbackItemDesc.BufferSize.X * FeedbackItemDesc.BufferSize.Y;
+				uint32 const* Data = (uint32*)RHICmdList.LockStagingBuffer(StagingBuffer, Fences->GetMapFence(FeedbackIndex), 0, BufferSize * sizeof(uint32));
+
+				if (FeedbackItemDesc.NumRects == 0)
+				{
+					// Copy full buffer
+					FMemory::Memcpy(MapResult.Data + MapResult.Size, Data, BufferSize * sizeof(uint32));
+					MapResult.Size += BufferSize;
+				}
+				else
+				{
+					// Copy individual rectangles from the buffer
+					const int32 BufferWidth = FeedbackItemDesc.BufferSize.X;
+					for (int32 RectIndex = 0; RectIndex < FeedbackItemDesc.NumRects; ++RectIndex)
+					{
+						const FIntRect Rect = FeedbackItemDesc.Rects[RectIndex];
+						const int32 RectWidth = Rect.Width();
+						const int32 RectHeight = Rect.Height();
+
+						uint32 const* Src = Data + Rect.Min.Y * BufferWidth + Rect.Min.X;
+						uint32* Dst = MapResult.Data + MapResult.Size;
+
+						for (int32 y = 0; y < RectHeight; ++y)
+						{
+							FMemory::Memcpy(Dst, Src, RectWidth * sizeof(uint32));
+
+							Src += BufferWidth;
+							Dst += RectWidth;
+						}
+
+						MapResult.Size += RectWidth * RectHeight;
+					}
+				}
+
+				RHICmdList.UnlockStagingBuffer(StagingBuffer);
+			}
 		}
-		OutResult.Pitch = Size.X;
-		
-		void* MappedMem = (uint32*)RHICmdList.LockStagingBuffer(FeedbackEntryCPU.ReadbackBuffer, FeedBackFences->GetMapFence(CPUReadIndex), 0, NumBytes);
-
-		OutResult.Buffer.SetNumUninitialized(NumBytes / sizeof(uint32));
-		FMemory::Memcpy(OutResult.Buffer.GetData(), MappedMem, NumBytes);
-
-		RHICmdList.UnlockStagingBuffer(FeedbackEntryCPU.ReadbackBuffer);
-		FeedbackEntryCPU.ReadbackBuffer.SafeRelease();
-		
-		check(PendingTargetCount > 0u);
-		--PendingTargetCount;
-		CPUReadIndex = (CPUReadIndex + 1) % TargetCapacity;
-		return true;
+
+		check(MapResult.Size == TotalReadSize)
+	
+		// Increment the ring buffer read position.
+		NumPending -= NumResults;
+		ReadIndex = (ReadIndex + NumResults) % MaxTransfers;
 	}
 
-	return false;
+	return MapResult;
+}
+
+FVirtualTextureFeedback::FMapResult FVirtualTextureFeedback::Map(FRHICommandListImmediate& RHICmdList)
+{
+	return Map(RHICmdList, MaxTransfers);
+}
+
+void FVirtualTextureFeedback::Unmap(FRHICommandListImmediate& RHICmdList, int32 MapHandle)
+{
+	if (MapHandle >= 0)
+	{
+		FMapResources& Resources = MapResources[MapHandle];
+
+		// Do any required buffer Unlock.
+		if (Resources.FeedbackItemToUnlockIndex >= 0)
+		{
+			SCOPED_GPU_MASK(RHICmdList, FeedbackItems[Resources.FeedbackItemToUnlockIndex].GPUMask);
+			RHICmdList.UnlockStagingBuffer(FeedbackItems[Resources.FeedbackItemToUnlockIndex].StagingBuffer);
+			Resources.FeedbackItemToUnlockIndex = -1;
+		}
+
+		// Reset any allocated data buffer.
+		Resources.ResultData.Reset();
+
+		// Return to free list.
+		FreeMapResources.Add(MapHandle);
+	}
 }
+
+TGlobalResource< FVirtualTextureFeedback > GVirtualTextureFeedback;
diff --git a/Engine/Source/Runtime/Renderer/Private/VT/VirtualTextureFeedback.h b/Engine/Source/Runtime/Renderer/Private/VT/VirtualTextureFeedback.h
--- a/Engine/Source/Runtime/Renderer/Private/VT/VirtualTextureFeedback.h
+++ b/Engine/Source/Runtime/Renderer/Private/VT/VirtualTextureFeedback.h
@@ -3,89 +3,108 @@
 #pragma once
 
 #include "CoreMinimal.h"
-#include "RenderCore/Public/RenderTargetPool.h"
-
-/*
-====================================
-	Manages GPU and CPU buffers for VT feedback.
-	Shared for all views, not per view.
-
-	Should use append buffer but that requires RHI ability to copy
-	a GPU structured buffer to a CPU read only version.
-====================================
-*/
-class FVirtualTextureFeedback
+#include "RenderResource.h"
+
+/** 
+ * Class that handles the read back of feedback buffers from the GPU.
+ * Multiple buffers can be transferred per frame using TransferGPUToCPU().
+ * RHIBuffers that are passed in are expected to have been transitioned to a state for reading.
+ * Use Map()/Unmap() to return the data. It will only return data that is ready to fetch without stalling the GPU.
+ * All calls are expected to be from the single render thread only.
+ */
+class FVirtualTextureFeedback : public FRenderResource
 {
 public:
 	FVirtualTextureFeedback();
 	~FVirtualTextureFeedback();
 
-	static const uint32 TargetCapacity = 4u;
-	static const uint32 MaxRectPerTarget = 4u;
-
-	struct MapResult
+	/** 
+	 * Description of how to interpret a Buffer that is being transferred.
+	 * For example a buffer may be a simple flat buffer, or a 2D screen-space buffer with rectangles representing multiple player viewports.
+	 * In the future we may also want to support append style buffers containing buffer size etc.
+	 */
+	struct FBufferDesc
 	{
-		TArray<uint32> Buffer;
-		int32 Pitch;
-		int32 NumRects;
-		FIntRect Rects[MaxRectPerTarget];
+		/** Initialize for a flat "1D" buffer. */
+		void Init(int32 InBufferSize);
+		/** Initialize for a "2D" buffer. */
+		void Init2D(FIntPoint InBufferSize);
+		/** Initialize for a "2D" buffer. Pass in view rectangles that define some subset of the buffer to actually analyze. */
+		void Init2D(FIntPoint InUnscaledBufferSize, TArrayView<FIntRect> const& InUnscaledViewRects, int32 InBufferScale);
+
+		/** Size of buffer. 1D buffers have Y=1. */
+		FIntPoint BufferSize = 0;
+		/** The maximum number of rectangles to read from a 2D buffer. */
+		static const uint32 MaxRectPerTransfer = 4u;
+		/** The number of rectangles to read from the buffer. */
+		int32 NumRects = 0;
+		/** Rectangles to read from the buffer. */
+		FIntRect Rects[MaxRectPerTransfer];
+		/** Number of buffer elements to actually read (calculated from the Rects). */
+		int32 TotalReadSize = 0;
 	};
 
-	FVertexBufferRHIRef FeedbackBuffer;
-	FUnorderedAccessViewRHIRef FeedbackBufferUAV;
+	/** Commit a RHIBuffer to be transferred for later CPU analysis. */
+	void TransferGPUToCPU(FRHICommandListImmediate& RHICmdList, FVertexBufferRHIRef const& Buffer, FBufferDesc const& Desc);
 
-	void			CreateResourceGPU(FRHICommandListImmediate& RHICmdList, FIntPoint InSize);
-	void			ReleaseResources();
+	/** Returns true if there are any pending transfer results that are ready so that we can call Map(). */
+	bool CanMap(FRHICommandListImmediate& RHICmdList);
 
-	void			TransferGPUToCPU(FRHICommandListImmediate& RHICmdList, TArrayView<FIntRect> const& ViewRects);
+	/** Structure returned by Map() containing the feedback data. */
+	struct FMapResult
+	{
+		uint32* RESTRICT Data = nullptr;
+		uint32 Size = 0;
+		int32 MapHandle = -1;
+	};
 
-	uint32			GetPendingTargetCount() const { return PendingTargetCount; }
-	uint32			GetFeedbackStride() const { return (uint32)Size.X; }
+	/** Fetch all pending results into a flat output buffer for analysis. */
+	FMapResult Map(FRHICommandListImmediate& RHICmdList);
+	/** Fetch up to MaxTransfersToMap pending results into a flat output buffer for analysis. */
+	FMapResult Map(FRHICommandListImmediate& RHICmdList, int32 MaxTransfersToMap);
+	
+	/** Always call Unmap() after finishing processing of the returned data. This releases any resources allocated for the Map(). */
+	void Unmap(FRHICommandListImmediate& RHICmdList, int32 MapHandle);
 
-	bool			CanMap(FRHICommandListImmediate& RHICmdList);
-	bool			Map(FRHICommandListImmediate& RHICmdList, MapResult& OutResult);
+protected:
+	//~ Begin FRenderResource Interface
+	virtual void InitRHI() override;
+	virtual void ReleaseRHI() override;
+	//~ End FRenderResource Interface
 
 private:
-	struct FFeedBackItem
+	/** Description of a pending feedback transfer. */
+	struct FFeedbackItem
 	{
-		int32 NumRects;
-		FIntRect Rects[MaxRectPerTarget];
+		FBufferDesc Desc;
 		FRHIGPUMask GPUMask;
-		FStagingBufferRHIRef ReadbackBuffer;
+		FStagingBufferRHIRef StagingBuffer;
 	};
 
-	FFeedBackItem FeedbackCPU[TargetCapacity];
+	/** The maximum number of pending feedback transfers that can be held before we start dropping them. */
+	static const uint32 MaxTransfers = 8u;
+	/** Pending feedback transfers are stored as a ring buffer. */
+	FFeedbackItem FeedbackItems[MaxTransfers];
 
-	class FFeedbackFences* FeedBackFences;
+	/** GPU fence pool. Contains a fence array that is kept in sync with the FeedbackItems ring buffer. Fences are used to know when a transfer is ready to Map() without stalling. */
+	class FFeedbackGPUFencePool* Fences;
 
-	FIntPoint Size;
-	uint32 NumBytes;
+	int32 NumPending;
+	uint32 WriteIndex;
+	uint32 ReadIndex;
 
-	uint32 GPUWriteIndex;
-	uint32 CPUReadIndex;
-	uint32 PendingTargetCount;
-};
-
-
-/* Dummy resource for VT feedback, used to bypass UB validation errors when VT is disabled */
-class FVirtualTextureFeedbackDummyResource : public FRenderResource
-{
-public:
-	FVertexBufferRHIRef			VertexBufferRHI;
-	FUnorderedAccessViewRHIRef	UAV;
-
-	virtual void InitRHI() override
-	{
-		FRHIResourceCreateInfo CreateInfo(TEXT("VTFeedbackDummy"));
-		VertexBufferRHI = RHICreateVertexBuffer(4, BUF_Static | BUF_ShaderResource | BUF_UnorderedAccess, CreateInfo);
-		UAV = RHICreateUnorderedAccessView(VertexBufferRHI, /*Format=*/ PF_R32_UINT);
-	}
-	virtual void ReleaseRHI() override
+	/** Structure describing resources associated with a Map() that will need freeing on UnMap() */
+	struct FMapResources
 	{
-		UAV.SafeRelease();
-		VertexBufferRHI.SafeRelease();
-	}
-};
+		int32 FeedbackItemToUnlockIndex = -1;
+		TArray<uint32> ResultData;
+	};
 
-extern TGlobalResource<FVirtualTextureFeedbackDummyResource> GVirtualTextureFeedbackDummyResource;
+	/** Array of FMapResources. Will always be size 1 unless we need to Map() multiple individual buffers at once. */
+	TArray<FMapResources> MapResources;
+	/** Free indices in the MapResources array. */
+	TArray<int32> FreeMapResources;
+};
 
+/** Global object for handling the read back of Virtual Texture Feedback. */
+extern TGlobalResource< FVirtualTextureFeedback > GVirtualTextureFeedback;
diff --git a/Engine/Source/Runtime/Renderer/Private/VT/VirtualTextureSystem.cpp b/Engine/Source/Runtime/Renderer/Private/VT/VirtualTextureSystem.cpp
--- a/Engine/Source/Runtime/Renderer/Private/VT/VirtualTextureSystem.cpp
+++ b/Engine/Source/Runtime/Renderer/Private/VT/VirtualTextureSystem.cpp
@@ -118,9 +118,7 @@ struct FFeedbackAnalysisParameters
 	FVirtualTextureSystem* System = nullptr;
 	const uint32* FeedbackBuffer = nullptr;
 	FUniquePageList* UniquePageList = nullptr;
-	uint32 FeedbackWidth = 0u;
-	uint32 FeedbackHeight = 0u;
-	uint32 FeedbackPitch = 0u;
+	uint32 FeedbackSize = 0u;
 };
 
 struct FGatherRequestsParameters
@@ -856,34 +854,28 @@ void FVirtualTextureSystem::FeedbackAnalysisTask(const FFeedbackAnalysisParamete
 {
 	FUniquePageList* RESTRICT RequestedPageList = Parameters.UniquePageList;
 	const uint32* RESTRICT Buffer = Parameters.FeedbackBuffer;
-	const uint32 Width = Parameters.FeedbackWidth;
-	const uint32 Height = Parameters.FeedbackHeight;
-	const uint32 Pitch = Parameters.FeedbackPitch;
+	const uint32 BufferSize = Parameters.FeedbackSize;
 
 	// Combine simple runs of identical requests
 	uint32 LastPixel = 0xffffffff;
 	uint32 LastCount = 0;
 
-	for (uint32 y = 0; y < Height; y++)
+	for (uint32 Index = 0; Index < BufferSize; Index++)
 	{
-		const uint32* RESTRICT BufferRow = Buffer + y * Pitch;
-		for (uint32 x = 0; x < Width; x++)
+		const uint32 Pixel = Buffer[Index];
+		if (Pixel == LastPixel)
 		{
-			const uint32 Pixel = BufferRow[x];
-			if (Pixel == LastPixel)
-			{
-				LastCount++;
-				continue;
-			}
-
-			if (LastPixel != 0xffffffff)
-			{
-				RequestedPageList->Add(LastPixel, LastCount);
-			}
+			LastCount++;
+			continue;
+		}
 
-			LastPixel = Pixel;
-			LastCount = 1;
+		if (LastPixel != 0xffffffff)
+		{
+			RequestedPageList->Add(LastPixel, LastCount);
 		}
+
+		LastPixel = Pixel;
+		LastCount = 1;
 	}
 
 	if (LastPixel != 0xffffffff)
@@ -934,8 +926,7 @@ void FVirtualTextureSystem::Update(FRHICommandListImmediate& RHICmdList, ERHIFea
 	{
 		// Only flush if we know that there is GPU feedback available to refill the visible data this frame
 		// This prevents bugs when low frame rate causes feedback buffer to stall so that the physical cache isn't filled immediately which causes visible glitching
-		FSceneRenderTargets& SceneContext = FSceneRenderTargets::Get(RHICmdList);
-		if (SceneContext.VirtualTextureFeedback.CanMap(RHICmdList))
+		if (GVirtualTextureFeedback.CanMap(RHICmdList))
 		{
 			// Each RVT will call FVirtualTextureSystem::FlushCache()
 			Scene->FlushDirtyRuntimeVirtualTextures();
@@ -946,85 +937,46 @@ void FVirtualTextureSystem::Update(FRHICommandListImmediate& RHICmdList, ERHIFea
 	FMemMark Mark(MemStack);
 	FUniquePageList* MergedUniquePageList = new(MemStack) FUniquePageList;
 	MergedUniquePageList->Initialize();
+	
+	if (CVarVTEnableFeedBack.GetValueOnRenderThread())
 	{
 		FMemMark FeedbackMark(MemStack);
 
-		FSceneRenderTargets& SceneContext = FSceneRenderTargets::Get(RHICmdList);
-
-		// Gather all outstanding feedback buffers
-		uint32 FeedbackBufferCount = 0;
-		uint32 FeedbackRectCount = 0;
-		FVirtualTextureFeedback::MapResult MappedFeedbackBuffers[FVirtualTextureFeedback::TargetCapacity];
+		// Fetch feedback for analysis
+		FVirtualTextureFeedback::FMapResult FeedbackResult;
 
-		if (CVarVTEnableFeedBack.GetValueOnRenderThread())
 		{
 			SCOPE_CYCLE_COUNTER(STAT_FeedbackMap);
-
-#if WITH_EDITOR
-			uint32 MaxFeedbackTargetCount = FVirtualTextureFeedback::TargetCapacity;
-#else
-			uint32 MaxFeedbackTargetCount = 1;
-#endif
-
-			while (FeedbackBufferCount < MaxFeedbackTargetCount)
-			{
-				if (!SceneContext.VirtualTextureFeedback.Map(RHICmdList, MappedFeedbackBuffers[FeedbackBufferCount]))
-				{
-					break;
-				}
-				
-				FeedbackRectCount += MappedFeedbackBuffers[FeedbackBufferCount].NumRects;
-				FeedbackBufferCount++;
-			}
+			FeedbackResult = GVirtualTextureFeedback.Map(RHICmdList);
 		}
 
-		// Create tasks to read all the buffers
-		uint32 NumFeedbackTasks = 0u;
-
+		// Create tasks to read the feedback data
+		// Give each task a section of the feedback buffer to analyze
 		FFeedbackAnalysisParameters FeedbackAnalysisParameters[MaxNumTasks];
-		static_assert(MaxNumTasks >= FVirtualTextureFeedback::TargetCapacity * FVirtualTextureFeedback::MaxRectPerTarget, "MaxNumTasks is too small");
 
 		const uint32 MaxNumFeedbackTasks = FMath::Clamp((uint32)CVarVTNumFeedbackTasks.GetValueOnRenderThread(), 1u, MaxNumTasks);
-		const uint32 FeedbackTasksPerRect = FMath::Max(MaxNumFeedbackTasks / FMath::Max(FeedbackRectCount, 1u), 1u);
+		const uint32 FeedbackSizePerTask = FMath::DivideAndRoundUp(FeedbackResult.Size, MaxNumFeedbackTasks);
 
-		for (uint32 FeedbackBufferIndex = 0; FeedbackBufferIndex < FeedbackBufferCount; ++FeedbackBufferIndex)
+		uint32 NumFeedbackTasks = 0;
+		uint32 CurrentOffset = 0;
+		while (CurrentOffset < FeedbackResult.Size)
 		{
-			FVirtualTextureFeedback::MapResult const& FeedbackBuffer = MappedFeedbackBuffers[FeedbackBufferIndex];
-
-			for (int32 RectIndex = 0; RectIndex < FeedbackBuffer.NumRects; ++RectIndex)
+			const uint32 TaskIndex = NumFeedbackTasks++;
+			FFeedbackAnalysisParameters& Params = FeedbackAnalysisParameters[TaskIndex];
+			Params.System = this;
+			if (TaskIndex == 0u)
 			{
-				FIntRect const& Rect = FeedbackBuffer.Rects[RectIndex];
-
-				// Give each task a section of a feedback rect to analyze
-				//todo[vt]: For buffers/rects of different sizes we will have different task payload sizes which is not efficient
-				const uint32 FeedbackRowsPerTask = FMath::DivideAndRoundUp((uint32)Rect.Size().Y, FeedbackTasksPerRect);
-				const uint32 NumRows = (uint32)Rect.Size().Y;
-
-				uint32 CurrentRow = 0;
-				while (CurrentRow < NumRows)
-				{
-					const uint32 CurrentHeight = FMath::Min(FeedbackRowsPerTask, NumRows - CurrentRow);
-					if (CurrentHeight > 0u)
-					{
-						const uint32 TaskIndex = NumFeedbackTasks++;
-						FFeedbackAnalysisParameters& Params = FeedbackAnalysisParameters[TaskIndex];
-						Params.System = this;
-						if (TaskIndex == 0u)
-						{
-							Params.UniquePageList = MergedUniquePageList;
-						}
-						else
-						{
-							Params.UniquePageList = new(MemStack) FUniquePageList;
-						}
-						Params.FeedbackBuffer = FeedbackBuffer.Buffer.GetData() + (Rect.Min.Y + CurrentRow) * FeedbackBuffer.Pitch + Rect.Min.X;
-						Params.FeedbackWidth = Rect.Size().X;
-						Params.FeedbackHeight = CurrentHeight;
-						Params.FeedbackPitch = FeedbackBuffer.Pitch;
-						CurrentRow += CurrentHeight;
-					}
-				}
+				Params.UniquePageList = MergedUniquePageList;
 			}
+			else
+			{
+				Params.UniquePageList = new(MemStack) FUniquePageList;
+			}
+			Params.FeedbackBuffer = FeedbackResult.Data + CurrentOffset;
+
+			const uint32 Size = FMath::Min(FeedbackSizePerTask, FeedbackResult.Size - CurrentOffset);
+			Params.FeedbackSize = Size;
+			CurrentOffset += Size;
 		}
 
 		// Kick the tasks
@@ -1033,7 +985,7 @@ void FVirtualTextureSystem::Update(FRHICommandListImmediate& RHICmdList, ERHIFea
 		const int32 WorkerFeedbackTaskCount = NumFeedbackTasks - LocalFeedbackTaskCount;
 
 		FGraphEventArray Tasks;
-		if(WorkerFeedbackTaskCount > 0)
+		if (WorkerFeedbackTaskCount > 0)
 		{
 			SCOPE_CYCLE_COUNTER(STAT_ProcessRequests_SubmitTasks);
 			Tasks.Reserve(WorkerFeedbackTaskCount);
@@ -1059,7 +1011,7 @@ void FVirtualTextureSystem::Update(FRHICommandListImmediate& RHICmdList, ERHIFea
 			}
 		}
 
-		if(NumFeedbackTasks > 1u)
+		if (NumFeedbackTasks > 1u)
 		{
 			SCOPE_CYCLE_COUNTER(STAT_ProcessRequests_MergePages);
 			for (uint32 TaskIndex = 1u; TaskIndex < NumFeedbackTasks; ++TaskIndex)
@@ -1067,6 +1019,8 @@ void FVirtualTextureSystem::Update(FRHICommandListImmediate& RHICmdList, ERHIFea
 				MergedUniquePageList->MergePages(FeedbackAnalysisParameters[TaskIndex].UniquePageList);
 			}
 		}
+
+		GVirtualTextureFeedback.Unmap(RHICmdList, FeedbackResult.MapHandle);
 	}
 
 	FUniqueRequestList* MergedRequestList = new(MemStack) FUniqueRequestList(MemStack);
