Fix for incorrect clear on primary view in AA when rendering multiple viewports on mobile

diff --git a/Engine/Source/Runtime/Renderer/Private/PostProcess/PostProcessMobile.cpp b/Engine/Source/Runtime/Renderer/Private/PostProcess/PostProcessMobile.cpp
--- a/Engine/Source/Runtime/Renderer/Private/PostProcess/PostProcessMobile.cpp
+++ b/Engine/Source/Runtime/Renderer/Private/PostProcess/PostProcessMobile.cpp
@@ -2644,7 +2644,8 @@ void FRCPassPostProcessAaES2::Process(FRenderingCompositePassContext& Context)
 
 	ERenderTargetActions LoadStoreAction = ERenderTargetActions::Load_Store;
 	//#todo-rv-vr
-	if (!IStereoRendering::IsASecondaryView(Context.View))
+	if ((!IStereoRendering::IsASecondaryView(Context.View) && IStereoRendering::IsStereoEyeView(Context.View)) ||
+		Context.View.Family->Views.Num() == 1)
 	{
 		// Full clear to avoid restore
 		LoadStoreAction = ERenderTargetActions::Clear_Store;
