Fixed duplicated Cvar was causing error in non unity builds

diff --git a/Engine/Source/Runtime/Renderer/Private/IndirectLightRendering.cpp b/Engine/Source/Runtime/Renderer/Private/IndirectLightRendering.cpp
--- a/Engine/Source/Runtime/Renderer/Private/IndirectLightRendering.cpp
+++ b/Engine/Source/Runtime/Renderer/Private/IndirectLightRendering.cpp
@@ -22,15 +22,6 @@ static TAutoConsoleVariable<int32> CVarDiffuseIndirectDenoiser(
 	TEXT("Denoising options (default = 1)"),
 	ECVF_RenderThreadSafe);
 
-static TAutoConsoleVariable<int32> CVarUseReflectionDenoiser(
-	TEXT("r.Reflections.Denoiser"),
-	2,
-	TEXT("Choose the denoising algorithm.\n")
-	TEXT(" 0: Disabled;\n")
-	TEXT(" 1: Forces the default denoiser of the renderer;\n")
-	TEXT(" 2: GScreenSpaceDenoiser which may be overriden by a third party plugin (default)."),
-	ECVF_RenderThreadSafe);
-
 static TAutoConsoleVariable<int32> CVarDenoiseSSR(
 	TEXT("r.SSR.ExperimentalDenoiser"), 0,
 	TEXT("Replace SSR's TAA pass with denoiser."),
