UE4.25 - Deprecate UnbindRenderTargets

diff --git a/Engine/Source/Runtime/Engine/Private/Particles/ParticleGpuSimulation.cpp b/Engine/Source/Runtime/Engine/Private/Particles/ParticleGpuSimulation.cpp
--- a/Engine/Source/Runtime/Engine/Private/Particles/ParticleGpuSimulation.cpp
+++ b/Engine/Source/Runtime/Engine/Private/Particles/ParticleGpuSimulation.cpp
@@ -4981,7 +4981,9 @@ void FFXSystem::SimulateGPUParticles(
 		RHICmdList.TransitionResources(EResourceTransitionAccess::EReadable, VisualizeStateRHIs, 2);		
 	}
 
+	PRAGMA_DISABLE_DEPRECATION_WARNINGS
 	UnbindRenderTargets(RHICmdList);
+	PRAGMA_ENABLE_DEPRECATION_WARNINGS
 
 	// Stats.
 	if (Phase == GetLastParticleSimulationPhase(GetShaderPlatform()))
diff --git a/Engine/Source/Runtime/Engine/Private/Slate/SceneViewport.cpp b/Engine/Source/Runtime/Engine/Private/Slate/SceneViewport.cpp
--- a/Engine/Source/Runtime/Engine/Private/Slate/SceneViewport.cpp
+++ b/Engine/Source/Runtime/Engine/Private/Slate/SceneViewport.cpp
@@ -1735,7 +1735,9 @@ void FSceneViewport::EndRenderFrame(FRHICommandListImmediate& RHICmdList, bool b
 		if (bShouldUnsetTargets)
 		{
 			// Set the active render target(s) to nothing to release references in the case that the viewport is resized by slate before we draw again
+			PRAGMA_DISABLE_DEPRECATION_WARNINGS
 			UnbindRenderTargets(RHICmdList);
+			PRAGMA_ENABLE_DEPRECATION_WARNINGS
 		}
 		// Note: this releases our reference but does not release the resource as it is owned by slate (this is intended)
 		RenderTargetTextureRenderThreadRHI.SafeRelease();
diff --git a/Engine/Source/Runtime/RHI/Public/RHIUtilities.h b/Engine/Source/Runtime/RHI/Public/RHIUtilities.h
--- a/Engine/Source/Runtime/RHI/Public/RHIUtilities.h
+++ b/Engine/Source/Runtime/RHI/Public/RHIUtilities.h
@@ -563,7 +563,7 @@ inline void TransitionRenderPassTargets(FRHICommandList& RHICmdList, const FRHIR
 	RHICmdList.TransitionResources(EResourceTransitionAccess::EWritable, Transitions, TransitionIndex);
 }
 
-// Will soon be deprecated as well...
+UE_DEPRECATED(4.25, "UnbindRenderTargets API is deprecated; EndRenderPass implies UnbindRenderTargets.")
 inline void UnbindRenderTargets(FRHICommandList& RHICmdList)
 {
 	check(RHICmdList.IsOutsideRenderPass());
diff --git a/Engine/Source/Runtime/RenderCore/Private/RenderGraphBuilder.cpp b/Engine/Source/Runtime/RenderCore/Private/RenderGraphBuilder.cpp
--- a/Engine/Source/Runtime/RenderCore/Private/RenderGraphBuilder.cpp
+++ b/Engine/Source/Runtime/RenderCore/Private/RenderGraphBuilder.cpp
@@ -719,7 +719,9 @@ void FRDGBuilder::ExecutePass(const FRDGPass* Pass)
 	}
 	else
 	{
+		PRAGMA_DISABLE_DEPRECATION_WARNINGS
 		UnbindRenderTargets(RHICmdList);
+		PRAGMA_DISABLE_DEPRECATION_WARNINGS
 	}
 
 	Pass->Execute(RHICmdList);
diff --git a/Engine/Source/Runtime/Renderer/Private/CapsuleShadowRendering.cpp b/Engine/Source/Runtime/Renderer/Private/CapsuleShadowRendering.cpp
--- a/Engine/Source/Runtime/Renderer/Private/CapsuleShadowRendering.cpp
+++ b/Engine/Source/Runtime/Renderer/Private/CapsuleShadowRendering.cpp
@@ -802,8 +802,9 @@ bool FDeferredShadingSceneRenderer::RenderCapsuleDirectShadows(
 				FPlatformMemory::Memcpy(CapsuleShapeLockedData, CapsuleShapeData.GetData(), DataSize);
 				RHIUnlockVertexBuffer(LightSceneInfo.ShadowCapsuleShapesVertexBuffer);
 
-				// #todo-renderpasses Remove once everything is converted to renderpasses
+				PRAGMA_DISABLE_DEPRECATION_WARNINGS
 				UnbindRenderTargets(RHICmdList);
+				PRAGMA_ENABLE_DEPRECATION_WARNINGS
 
 				const bool bDirectionalLight = LightSceneInfo.Proxy->GetLightType() == LightType_Directional;
 				FIntRect ScissorRect;
@@ -1319,8 +1320,9 @@ void FDeferredShadingSceneRenderer::RenderIndirectCapsuleShadows(
 					{
 						check(IndirectShadowLightDirectionSRV);
 
-						// #todo-renderpasses remove once everything is converted to renderpasses
+						PRAGMA_DISABLE_DEPRECATION_WARNINGS
 						UnbindRenderTargets(RHICmdList);
+						PRAGMA_ENABLE_DEPRECATION_WARNINGS
 
 						FIntRect ScissorRect = View.ViewRect;
 
@@ -1547,8 +1549,9 @@ void FDeferredShadingSceneRenderer::RenderCapsuleShadowsForMovableSkylight(FRHIC
 					{
 						check(IndirectShadowLightDirectionSRV);
 
-						// #todo-renderpasses remove once everything is converted to renderpasses
+						PRAGMA_DISABLE_DEPRECATION_WARNINGS
 						UnbindRenderTargets(RHICmdList);
+						PRAGMA_DISABLE_DEPRECATION_WARNINGS
 
 						FIntRect ScissorRect = View.ViewRect;
 
diff --git a/Engine/Source/Runtime/Renderer/Private/CompositionLighting/PostProcessAmbientOcclusion.cpp b/Engine/Source/Runtime/Renderer/Private/CompositionLighting/PostProcessAmbientOcclusion.cpp
--- a/Engine/Source/Runtime/Renderer/Private/CompositionLighting/PostProcessAmbientOcclusion.cpp
+++ b/Engine/Source/Runtime/Renderer/Private/CompositionLighting/PostProcessAmbientOcclusion.cpp
@@ -694,7 +694,9 @@ void FRCPassPostProcessAmbientOcclusionSmooth::Process(FRenderingCompositePassCo
 {
 	SCOPED_GPU_STAT(Context.RHICmdList, SSAOSmooth);
 
+	PRAGMA_DISABLE_DEPRECATION_WARNINGS
 	UnbindRenderTargets(Context.RHICmdList);
+	PRAGMA_ENABLE_DEPRECATION_WARNINGS
 	Context.SetViewportAndCallRHI(Context.View.ViewRect);
 
 	FSceneRenderTargets& SceneContext = FSceneRenderTargets::Get(Context.RHICmdList);
@@ -984,8 +986,9 @@ void FRCPassPostProcessAmbientOcclusion::ProcessCS(FRenderingCompositePassContex
 	} \
 	break
 
-	// #todo-renderpasses remove once everything is renderpasses
+	PRAGMA_DISABLE_DEPRECATION_WARNINGS
 	UnbindRenderTargets(Context.RHICmdList);
+	PRAGMA_ENABLE_DEPRECATION_WARNINGS
 	Context.SetViewportAndCallRHI(ViewRect, 0.0f, 1.0f);
 
 	//for async compute we need to set up a fence to make sure the resource is ready before we start.
@@ -1831,7 +1834,9 @@ void FRCPassPostProcessAmbientOcclusion_GTAOHorizonSearchIntegrate::Process(FRen
 	if (1) //CVarCompute->GetValueOnRenderThread() >= 1)
 	{
 		// Compute  Version
+		PRAGMA_DISABLE_DEPRECATION_WARNINGS
 		UnbindRenderTargets(Context.RHICmdList);
+		PRAGMA_ENABLE_DEPRECATION_WARNINGS
 
 #define SET_SHADER_CASE_CS(RHICmdList,ShaderQualityCase) \
 		case ShaderQualityCase: \
@@ -1994,7 +1999,9 @@ void FRCPassPostProcessAmbientOcclusion_GTAOInnerIntegrate::Process(FRenderingCo
 		// Compute  Version
 		Context.RHICmdList.TransitionResource(EResourceTransitionAccess::ERWBarrier, EResourceTransitionPipeline::EGfxToCompute, DestRenderTarget.UAV, nullptr);
 
+		PRAGMA_DISABLE_DEPRECATION_WARNINGS
 		UnbindRenderTargets(Context.RHICmdList);
+		PRAGMA_ENABLE_DEPRECATION_WARNINGS
 
 		TShaderMapRef<FPostProcessGTAOInnerIntegratePSandCS<true>> ComputeShader(Context.GetShaderMap());
 		Context.RHICmdList.SetComputeShader(ComputeShader.GetComputeShader());
@@ -2149,7 +2156,9 @@ void FRCPassPostProcessAmbientOcclusion_HorizonSearch::Process(FRenderingComposi
 		RHICmdListComputeImmediate.WaitComputeFence(AsyncStartFence);
 
 		// Compute  Version
+		PRAGMA_DISABLE_DEPRECATION_WARNINGS
 		UnbindRenderTargets(Context.RHICmdList);
+		PRAGMA_ENABLE_DEPRECATION_WARNINGS
 
 		// 0..4, 0:low 4:high
 		switch (ShaderQuality)
@@ -2521,7 +2530,9 @@ void FRCPassPostProcessAmbientOcclusion_GTAO_TemporalFilter::Process(FRenderingC
 	if (1)
 	{
 		// Compute  Version
+		PRAGMA_DISABLE_DEPRECATION_WARNINGS
 		UnbindRenderTargets(Context.RHICmdList);
+		PRAGMA_ENABLE_DEPRECATION_WARNINGS
 		TShaderMapRef<FPostProcessGTAOTemporalFilterPSandCS<true>> ComputeShader(Context.GetShaderMap());
 		uint32 GroupSizeX = FMath::DivideAndRoundUp(OutputViewRect.Width(), 8);
 		uint32 GroupSizeY = FMath::DivideAndRoundUp(OutputViewRect.Height(), 8);
@@ -2758,7 +2769,9 @@ void FRCPassPostProcessAmbientOcclusion_GTAO_SpatialFilter::Process(FRenderingCo
 	FIntRect OutputViewRect = InputViewRect;
 
 	// Compute  Version
+	PRAGMA_DISABLE_DEPRECATION_WARNINGS
 	UnbindRenderTargets(Context.RHICmdList);
+	PRAGMA_ENABLE_DEPRECATION_WARNINGS
 
 	// Get filter size
 	uint32 FilterWidth = CVarGTAOFilterWidth.GetValueOnRenderThread();
@@ -2966,7 +2979,9 @@ void FRCPassPostProcessAmbientOcclusion_GTAO_Upsample::Process(FRenderingComposi
 
 	if (0)
 	{
+		PRAGMA_DISABLE_DEPRECATION_WARNINGS
 		UnbindRenderTargets(Context.RHICmdList);
+		PRAGMA_ENABLE_DEPRECATION_WARNINGS
 		TShaderMapRef<FPostProcessGTAOUpsamplePSandCS<true>> ComputeShader(Context.GetShaderMap());
 
 		uint32 GroupSizeX = FMath::DivideAndRoundUp(View.ViewRect.Width(), 8);
diff --git a/Engine/Source/Runtime/Renderer/Private/ConvertToUniformMesh.cpp b/Engine/Source/Runtime/Renderer/Private/ConvertToUniformMesh.cpp
--- a/Engine/Source/Runtime/Renderer/Private/ConvertToUniformMesh.cpp
+++ b/Engine/Source/Runtime/Renderer/Private/ConvertToUniformMesh.cpp
@@ -259,7 +259,9 @@ int32 FUniformMeshConverter::Convert(
 			GUniformMeshTemporaryBuffers.Initialize();
 		}
 
+		PRAGMA_DISABLE_DEPRECATION_WARNINGS
 		UnbindRenderTargets(RHICmdList);
+		PRAGMA_ENABLE_DEPRECATION_WARNINGS
 
 		uint32 Offsets[1] = {0};
 		FRHIVertexBuffer* const StreamOutTargets[1] = {GUniformMeshTemporaryBuffers.TriangleData.GetReference()};
diff --git a/Engine/Source/Runtime/Renderer/Private/DeferredShadingRenderer.cpp b/Engine/Source/Runtime/Renderer/Private/DeferredShadingRenderer.cpp
--- a/Engine/Source/Runtime/Renderer/Private/DeferredShadingRenderer.cpp
+++ b/Engine/Source/Runtime/Renderer/Private/DeferredShadingRenderer.cpp
@@ -2533,7 +2533,9 @@ void FDeferredShadingSceneRenderer::Render(FRHICommandListImmediate& RHICmdList)
 	}
 	checkSlow(RHICmdList.IsOutsideRenderPass());
 	// Unbind everything in case FX has to read.
+	PRAGMA_DISABLE_DEPRECATION_WARNINGS
 	UnbindRenderTargets(RHICmdList);
+	PRAGMA_ENABLE_DEPRECATION_WARNINGS
 
 	// Notify the FX system that opaque primitives have been rendered and we now have a valid depth buffer.
 	if (Scene->FXSystem && Views.IsValidIndex(0))
diff --git a/Engine/Source/Runtime/Renderer/Private/DistanceFieldAmbientOcclusion.cpp b/Engine/Source/Runtime/Renderer/Private/DistanceFieldAmbientOcclusion.cpp
--- a/Engine/Source/Runtime/Renderer/Private/DistanceFieldAmbientOcclusion.cpp
+++ b/Engine/Source/Runtime/Renderer/Private/DistanceFieldAmbientOcclusion.cpp
@@ -383,8 +383,9 @@ void ComputeDistanceFieldNormal(FRHICommandListImmediate& RHICmdList, const TArr
 {
 	if (GAOComputeShaderNormalCalculation)
 	{
-		// #todo-renderpasses remove once everything is converted to renderpasses
+		PRAGMA_DISABLE_DEPRECATION_WARNINGS
 		UnbindRenderTargets(RHICmdList);
+		PRAGMA_ENABLE_DEPRECATION_WARNINGS
 
 		for (int32 ViewIndex = 0; ViewIndex < Views.Num(); ViewIndex++)
 		{
diff --git a/Engine/Source/Runtime/Renderer/Private/DistanceFieldObjectCulling.cpp b/Engine/Source/Runtime/Renderer/Private/DistanceFieldObjectCulling.cpp
--- a/Engine/Source/Runtime/Renderer/Private/DistanceFieldObjectCulling.cpp
+++ b/Engine/Source/Runtime/Renderer/Private/DistanceFieldObjectCulling.cpp
@@ -590,8 +590,9 @@ void ScatterTilesToObjects(FRHICommandListImmediate& RHICmdList, const FViewInfo
 	}
 	RHICmdList.EndRenderPass();
 
-	// #todo-renderpasses Remove once SetRenderTargets is completely removed.
+	PRAGMA_DISABLE_DEPRECATION_WARNINGS
 	UnbindRenderTargets(RHICmdList);
+	PRAGMA_ENABLE_DEPRECATION_WARNINGS
 	RHICmdList.TransitionResources(EResourceTransitionAccess::ERWBarrier, EResourceTransitionPipeline::EGfxToCompute, UAVs.GetData(), UAVs.Num());
 }
 
@@ -605,7 +606,9 @@ FIntPoint GetTileListGroupSizeForView(const FViewInfo& View)
 void BuildTileObjectLists(FRHICommandListImmediate& RHICmdList, FScene* Scene, TArray<FViewInfo>& Views, FSceneRenderTargetItem& DistanceFieldNormal, const FDistanceFieldAOParameters& Parameters)
 {
 	SCOPED_DRAW_EVENT(RHICmdList, BuildTileList);
+	PRAGMA_DISABLE_DEPRECATION_WARNINGS
 	UnbindRenderTargets(RHICmdList);
+	PRAGMA_ENABLE_DEPRECATION_WARNINGS
 
 	for (int32 ViewIndex = 0; ViewIndex < Views.Num(); ViewIndex++)
 	{
diff --git a/Engine/Source/Runtime/Renderer/Private/DistanceFieldScreenGridLighting.cpp b/Engine/Source/Runtime/Renderer/Private/DistanceFieldScreenGridLighting.cpp
--- a/Engine/Source/Runtime/Renderer/Private/DistanceFieldScreenGridLighting.cpp
+++ b/Engine/Source/Runtime/Renderer/Private/DistanceFieldScreenGridLighting.cpp
@@ -535,7 +535,9 @@ void FDeferredShadingSceneRenderer::RenderDistanceFieldAOScreenGrid(
 	}
 	ScreenGridResources->AcquireTransientResource();
 
+	PRAGMA_DISABLE_DEPRECATION_WARNINGS
 	UnbindRenderTargets(RHICmdList);
+	PRAGMA_ENABLE_DEPRECATION_WARNINGS
 
 	if (bUseGlobalDistanceField)
 	{
diff --git a/Engine/Source/Runtime/Renderer/Private/DistanceFieldShadowing.cpp b/Engine/Source/Runtime/Renderer/Private/DistanceFieldShadowing.cpp
--- a/Engine/Source/Runtime/Renderer/Private/DistanceFieldShadowing.cpp
+++ b/Engine/Source/Runtime/Renderer/Private/DistanceFieldShadowing.cpp
@@ -862,8 +862,9 @@ void ScatterObjectsToShadowTiles(
 	}
 	RHICmdList.EndRenderPass();
 
-	// #todo-renderpasses Needed for now. Once SetRenderTarget is removed completely this will be redundant.
+	PRAGMA_DISABLE_DEPRECATION_WARNINGS
 	UnbindRenderTargets(RHICmdList);
+	PRAGMA_ENABLE_DEPRECATION_WARNINGS
 	RHICmdList.TransitionResources(EResourceTransitionAccess::EReadable, EResourceTransitionPipeline::EGfxToCompute, UAVs.GetData(), UAVs.Num());
 }
 
@@ -1207,7 +1208,9 @@ void FProjectedShadowInfo::BeginRenderRayTracedDistanceFieldProjection(FRHIComma
 		{
 			check(!Scene->DistanceFieldSceneData.HasPendingOperations());
 
+			PRAGMA_DISABLE_DEPRECATION_WARNINGS
 			UnbindRenderTargets(RHICmdList);
+			PRAGMA_ENABLE_DEPRECATION_WARNINGS
 
 			int32 NumPlanes = 0;
 			const FPlane* PlaneData = NULL;
@@ -1256,7 +1259,9 @@ void FProjectedShadowInfo::BeginRenderRayTracedDistanceFieldProjection(FRHIComma
 			}
 
 			SCOPED_DRAW_EVENT(RHICmdList, RayTraceShadows);
+			PRAGMA_DISABLE_DEPRECATION_WARNINGS
 			UnbindRenderTargets(RHICmdList);
+			PRAGMA_ENABLE_DEPRECATION_WARNINGS
 
 			RayTraceShadows(RHICmdList, View, this, TileIntersectionResources);
 
@@ -1278,7 +1283,9 @@ void FProjectedShadowInfo::BeginRenderRayTracedDistanceFieldProjection(FRHIComma
 
 		check(!Scene->DistanceFieldSceneData.HasPendingHeightFieldOperations());
 
+		PRAGMA_DISABLE_DEPRECATION_WARNINGS
 		UnbindRenderTargets(RHICmdList);
+		PRAGMA_ENABLE_DEPRECATION_WARNINGS
 
 		const int32 NumPlanes = CascadeSettings.ShadowBoundsAccurate.Planes.Num();
 		const FPlane* PlaneData = CascadeSettings.ShadowBoundsAccurate.Planes.GetData();
@@ -1315,7 +1322,9 @@ void FProjectedShadowInfo::BeginRenderRayTracedDistanceFieldProjection(FRHIComma
 		}
 
 		SCOPED_DRAW_EVENT(RHICmdList, RayTraceHeightFieldShadows);
+		PRAGMA_DISABLE_DEPRECATION_WARNINGS
 		UnbindRenderTargets(RHICmdList);
+		PRAGMA_ENABLE_DEPRECATION_WARNINGS
 
 		RayTraceHeightFieldShadows(RHICmdList, View, this, TileIntersectionResources, bHasPrevOutput, PrevOutputRT ? PrevOutputRT->GetRenderTargetItem().ShaderResourceTexture : nullptr);
 
diff --git a/Engine/Source/Runtime/Renderer/Private/DistanceFieldVisualization.cpp b/Engine/Source/Runtime/Renderer/Private/DistanceFieldVisualization.cpp
--- a/Engine/Source/Runtime/Renderer/Private/DistanceFieldVisualization.cpp
+++ b/Engine/Source/Runtime/Renderer/Private/DistanceFieldVisualization.cpp
@@ -187,7 +187,9 @@ void FDeferredShadingSceneRenderer::RenderMeshDistanceFieldVisualization(FRHICom
 			}
 
 			{
+				PRAGMA_DISABLE_DEPRECATION_WARNINGS
 				UnbindRenderTargets(RHICmdList);
+				PRAGMA_ENABLE_DEPRECATION_WARNINGS
 
 				for (int32 ViewIndex = 0; ViewIndex < Views.Num(); ViewIndex++)
 				{
diff --git a/Engine/Source/Runtime/Renderer/Private/GPUFastFourierTransform.cpp b/Engine/Source/Runtime/Renderer/Private/GPUFastFourierTransform.cpp
--- a/Engine/Source/Runtime/Renderer/Private/GPUFastFourierTransform.cpp
+++ b/Engine/Source/Runtime/Renderer/Private/GPUFastFourierTransform.cpp
@@ -1182,8 +1182,9 @@ namespace
 		// Get pointer to the shader
 		TShaderRef<FCopyWindowCS> ComputeShader = GetCopyWindowCS(ShaderMap);
 
-		// #todo-renderpasses remove once everything is converted
+		PRAGMA_DISABLE_DEPRECATION_WARNINGS
 		UnbindRenderTargets(RHICmdList);
+		PRAGMA_ENABLE_DEPRECATION_WARNINGS
 		RHICmdList.SetComputeShader(ComputeShader.GetComputeShader());
 
 		// Bind output
@@ -1243,8 +1244,9 @@ namespace
 		// Get pointer to the shader
 		TShaderRef<FComplexMultiplyImagesCS> ComputeShader = GetComplexMultiplyImagesCS(ShaderMap);
 
-		// #todo-renderpasses remove once everything is converted
+		PRAGMA_DISABLE_DEPRECATION_WARNINGS
 		UnbindRenderTargets(RHICmdList);
+		PRAGMA_ENABLE_DEPRECATION_WARNINGS
 		RHICmdList.SetComputeShader(ComputeShader.GetComputeShader());
 
 		// Bind output
@@ -1312,8 +1314,9 @@ namespace
 		// Get pointer to the shader
 		TShaderRef<FPackTwoForOneFFTPassCS> ComputeShader = GetPackTwoForOneFFTPassCS(ShaderMap, TransformLength);
 
-		// #todo-renderpasses remove once everything is converted
+		PRAGMA_DISABLE_DEPRECATION_WARNINGS
 		UnbindRenderTargets(RHICmdList);
+		PRAGMA_ENABLE_DEPRECATION_WARNINGS
 		RHICmdList.SetComputeShader(ComputeShader.GetComputeShader());
 
 		// Bind output
@@ -1363,8 +1366,9 @@ namespace
 		TShaderRef<FComplexFFTPassCS> ComputeShader = GetComplexFFTPassCS(ShaderMap, TransformLength);
 
 
-		// #todo-renderpasses remove once everything is converted
+		PRAGMA_DISABLE_DEPRECATION_WARNINGS
 		UnbindRenderTargets(RHICmdList);
+		PRAGMA_ENABLE_DEPRECATION_WARNINGS
 		RHICmdList.SetComputeShader(ComputeShader.GetComputeShader());
 
 		// Bind output
@@ -1415,9 +1419,9 @@ namespace
 		// Get a base pointer to the shader
 		TShaderRef<FReorderFFTPassCS> ComputeShader = GetReorderFFTPassCS(ShaderMap);
 
-
-		// #todo-renderpasses remove once everything is converted
+		PRAGMA_DISABLE_DEPRECATION_WARNINGS
 		UnbindRenderTargets(RHICmdList);
+		PRAGMA_ENABLE_DEPRECATION_WARNINGS
 		RHICmdList.SetComputeShader(ComputeShader.GetComputeShader());
 
 		// Bind output
@@ -1478,9 +1482,9 @@ namespace
 		// Get a base pointer to the shader
 		TShaderRef<FGroupShardSubFFTPassCS> ComputeShader = GetGroupSharedSubFFTPassCS(ShaderMap, TransformLength);
 
-
-		// #todo-renderpasses remove once everything is converted
+		PRAGMA_DISABLE_DEPRECATION_WARNINGS
 		UnbindRenderTargets(RHICmdList);
+		PRAGMA_ENABLE_DEPRECATION_WARNINGS
 		RHICmdList.SetComputeShader(ComputeShader.GetComputeShader());
 
 		// Bind output
@@ -1526,9 +1530,9 @@ namespace
 		// Get a base pointer to the shader
 		TShaderRef<FGSComplexTransformBaseCS> ComputeShader = GetComplexFFTCS(ShaderMap, TransformLength);
 
-
-		// #todo-renderpasses remove once everything is converted
+		PRAGMA_DISABLE_DEPRECATION_WARNINGS
 		UnbindRenderTargets(RHICmdList);
+		PRAGMA_ENABLE_DEPRECATION_WARNINGS
 		RHICmdList.SetComputeShader(ComputeShader.GetComputeShader());
 
 		// Bind output
@@ -1578,9 +1582,9 @@ namespace
 		// Get a base pointer to the shader
 		TShaderRef<FGSComplexTransformBaseCS> ComputeShader = GetTwoForOneFFTCS(ShaderMap, TransformLength);
 
-		// DJH - do we need this SetRenderTarget?
-		// #todo-renderpasses remove once everything is converted
+		PRAGMA_DISABLE_DEPRECATION_WARNINGS
 		UnbindRenderTargets(RHICmdList);
+		PRAGMA_ENABLE_DEPRECATION_WARNINGS
 		RHICmdList.SetComputeShader(ComputeShader.GetComputeShader());
 
 
@@ -1636,9 +1640,9 @@ namespace
 		// Get a base pointer to the shader
 		TShaderRef<FGSConvolutionWithTextureKernelBaseCS> ComputeShader = GetConvolutionWithTextureKernelCS(ShaderMap, SignalLength);
 
-
-		// #todo-renderpasses remove once everything is converted
+		PRAGMA_DISABLE_DEPRECATION_WARNINGS
 		UnbindRenderTargets(RHICmdList);
+		PRAGMA_ENABLE_DEPRECATION_WARNINGS
 		RHICmdList.SetComputeShader(ComputeShader.GetComputeShader());
 		
 		FScopedUAVBind ScopedBind = FScopedUAVBind::BindOutput(RHICmdList, ComputeShader, DstUAV);
diff --git a/Engine/Source/Runtime/Renderer/Private/PostProcess/PostProcessDOF.cpp b/Engine/Source/Runtime/Renderer/Private/PostProcess/PostProcessDOF.cpp
--- a/Engine/Source/Runtime/Renderer/Private/PostProcess/PostProcessDOF.cpp
+++ b/Engine/Source/Runtime/Renderer/Private/PostProcess/PostProcessDOF.cpp
@@ -272,7 +272,9 @@ void FRCPassPostProcessDOFSetup::Process(FRenderingCompositePassContext& Context
 
 	// #todo-rco: needed to avoid multiple resolves clearing the RT with VK.
 	// #todo mattc this is probably busted now since the resolves previously happened after this SetRenderTarget.
+	PRAGMA_DISABLE_DEPRECATION_WARNINGS
 	UnbindRenderTargets(Context.RHICmdList);
+	PRAGMA_ENABLE_DEPRECATION_WARNINGS
 }
 
 FPooledRenderTargetDesc FRCPassPostProcessDOFSetup::ComputeOutputDesc(EPassOutputId InPassOutputId) const
diff --git a/Engine/Source/Runtime/Renderer/Private/PostProcess/PostProcessFFTBloom.cpp b/Engine/Source/Runtime/Renderer/Private/PostProcess/PostProcessFFTBloom.cpp
--- a/Engine/Source/Runtime/Renderer/Private/PostProcess/PostProcessFFTBloom.cpp
+++ b/Engine/Source/Runtime/Renderer/Private/PostProcess/PostProcessFFTBloom.cpp
@@ -242,7 +242,9 @@ void ResizeAndCenterTexture(
 	ClampedImageCenterUV.Y = FMath::Clamp(SrcImageCenterUV.Y, 0.f, 1.f);
 
 	check(DstUAV);
+	PRAGMA_DISABLE_DEPRECATION_WARNINGS
 	UnbindRenderTargets(RHICmdList);
+	PRAGMA_ENABLE_DEPRECATION_WARNINGS
 	RHICmdList.TransitionResource(EResourceTransitionAccess::ERWBarrier, EResourceTransitionPipeline::EGfxToCompute, DstUAV);
 
 	{
@@ -294,7 +296,9 @@ void CaptureKernelWeight(
 	FSceneRenderTargetItem& DstTargetItem = CenterWeightRT->GetRenderTargetItem();
 
 	check(DstTargetItem.UAV);
+	PRAGMA_DISABLE_DEPRECATION_WARNINGS
 	UnbindRenderTargets(RHICmdList);
+	PRAGMA_ENABLE_DEPRECATION_WARNINGS
 	RHICmdList.TransitionResource(EResourceTransitionAccess::ERWBarrier, EResourceTransitionPipeline::EGfxToCompute, DstTargetItem.UAV);
 
 	{
@@ -348,7 +352,9 @@ void BlendLowRes(
 
 	// set destination
 	check(DstUAV);
+	PRAGMA_DISABLE_DEPRECATION_WARNINGS
 	UnbindRenderTargets(RHICmdList);
+	PRAGMA_ENABLE_DEPRECATION_WARNINGS
 	RHICmdList.TransitionResource(EResourceTransitionAccess::ERWBarrier, EResourceTransitionPipeline::EComputeToCompute, DstUAV);
 
 	{
@@ -396,7 +402,9 @@ void CopyImageRect(
 
 	// set destination
 	check(DstUAV);
+	PRAGMA_DISABLE_DEPRECATION_WARNINGS
 	UnbindRenderTargets(RHICmdList);
+	PRAGMA_ENABLE_DEPRECATION_WARNINGS
 	RHICmdList.TransitionResource(EResourceTransitionAccess::ERWBarrier, EResourceTransitionPipeline::EGfxToCompute, DstUAV);
 
 	{
diff --git a/Engine/Source/Runtime/Renderer/Private/SceneCaptureRendering.cpp b/Engine/Source/Runtime/Renderer/Private/SceneCaptureRendering.cpp
--- a/Engine/Source/Runtime/Renderer/Private/SceneCaptureRendering.cpp
+++ b/Engine/Source/Runtime/Renderer/Private/SceneCaptureRendering.cpp
@@ -447,8 +447,9 @@ static void UpdateSceneCaptureContent_RenderThread(
 			checkNoEntry();
 			break;
 	}
-	// Unbind everything in case FX has to read.
+	PRAGMA_DISABLE_DEPRECATION_WARNINGS
 	UnbindRenderTargets(RHICmdList);
+	PRAGMA_ENABLE_DEPRECATION_WARNINGS
 }
 
 void BuildProjectionMatrix(FIntPoint RenderTargetSize, ECameraProjectionMode::Type ProjectionType, float FOV, float InOrthoWidth, float InNearClippingPlane, FMatrix& ProjectionMatrix)
diff --git a/Engine/Source/Runtime/Renderer/Private/TiledDeferredLightRendering.cpp b/Engine/Source/Runtime/Renderer/Private/TiledDeferredLightRendering.cpp
--- a/Engine/Source/Runtime/Renderer/Private/TiledDeferredLightRendering.cpp
+++ b/Engine/Source/Runtime/Renderer/Private/TiledDeferredLightRendering.cpp
@@ -361,7 +361,9 @@ void FDeferredShadingSceneRenderer::RenderTiledDeferredLighting(FRHICommandListI
 		INC_DWORD_STAT_BY(STAT_NumLightsUsingSimpleTiledDeferred, SimpleLights.InstanceData.Num());
 		SCOPE_CYCLE_COUNTER(STAT_DirectLightRenderingTime);
 
+		PRAGMA_DISABLE_DEPRECATION_WARNINGS
 		UnbindRenderTargets(RHICmdList);
+		PRAGMA_ENABLE_DEPRECATION_WARNINGS
 
 		// Determine how many compute shader passes will be needed to process all the lights
 		const int32 NumPassesNeeded = FMath::DivideAndRoundUp(NumLightsToRender, GMaxNumTiledDeferredLights);
diff --git a/Engine/Source/Runtime/Renderer/Private/VolumetricFog.cpp b/Engine/Source/Runtime/Renderer/Private/VolumetricFog.cpp
--- a/Engine/Source/Runtime/Renderer/Private/VolumetricFog.cpp
+++ b/Engine/Source/Runtime/Renderer/Private/VolumetricFog.cpp
@@ -1097,7 +1097,10 @@ void FDeferredShadingSceneRenderer::ComputeVolumetricFog(FRHICommandListImmediat
 					ERDGPassFlags::Compute,
 					[PassParameters, ComputeShader, &View, this, FogInfo, bUseTemporalReprojection, VolumetricFogGridSize, IntegrationData, bUseDirectionalLightShadowing, bUseDistanceFieldSkyOcclusion, LightFunctionWorldToShadow](FRHICommandListImmediate& RHICmdList)
 				{
+					PRAGMA_DISABLE_DEPRECATION_WARNINGS
 					UnbindRenderTargets(RHICmdList);
+					PRAGMA_ENABLE_DEPRECATION_WARNINGS
+
 					const FIntVector NumGroups = FIntVector::DivideAndRoundUp(VolumetricFogGridSize, FIntVector(VolumetricFogLightScatteringGroupSizeX, VolumetricFogLightScatteringGroupSizeY, VolumetricFogLightScatteringGroupSizeZ));
 
 					RHICmdList.SetComputeShader(ComputeShader.GetComputeShader());
