Add support for parallel AddMeshBatch for DynamicRT meshes

diff --git a/Engine/Source/Runtime/Renderer/Private/DeferredShadingRenderer.cpp b/Engine/Source/Runtime/Renderer/Private/DeferredShadingRenderer.cpp
--- a/Engine/Source/Runtime/Renderer/Private/DeferredShadingRenderer.cpp
+++ b/Engine/Source/Runtime/Renderer/Private/DeferredShadingRenderer.cpp
@@ -180,6 +180,20 @@ static TAutoConsoleVariable<int32> CVarRayTracingAsyncBuild(
 	ECVF_RenderThreadSafe
 );
 
+static int32 GRayTracingParallelMeshBatchSetup = 1;
+static FAutoConsoleVariableRef CRayTracingParallelMeshBatchSetup(
+	TEXT("r.RayTracing.ParallelMeshBatchSetup"),
+	GRayTracingParallelMeshBatchSetup,
+	TEXT("Whether to setup ray tracing materials via parallel jobs."),
+	ECVF_RenderThreadSafe);
+
+static int32 GRayTracingParallelMeshBatchSize = 128;
+static FAutoConsoleVariableRef CRayTracingParallelMeshBatchSize(
+	TEXT("r.RayTracing.ParallelMeshBatchSize"),
+	GRayTracingParallelMeshBatchSize,
+	TEXT("Batch size for ray tracing materials parallel jobs."),
+	ECVF_RenderThreadSafe);
+
 #if !UE_BUILD_SHIPPING
 static TAutoConsoleVariable<int32> CVarForceBlackVelocityBuffer(
 	TEXT("r.Test.ForceBlackVelocityBuffer"), 0,
@@ -431,6 +445,8 @@ DECLARE_CYCLE_STAT(TEXT("AfterTranslucentVelocity"), STAT_CLM_AfterTranslucentVe
 DECLARE_CYCLE_STAT(TEXT("RenderFinish"), STAT_CLM_RenderFinish, STATGROUP_CommandListMarkers);
 DECLARE_CYCLE_STAT(TEXT("AfterFrame"), STAT_CLM_AfterFrame, STATGROUP_CommandListMarkers);
 
+DECLARE_CYCLE_STAT(TEXT("Wait RayTracing Add Mesh Batch"), STAT_WaitRayTracingAddMesh, STATGROUP_SceneRendering);
+
 FGraphEventRef FDeferredShadingSceneRenderer::TranslucencyTimestampQuerySubmittedFence[FOcclusionQueryHelpers::MaxBufferedOcclusionFrames + 1];
 FGlobalDynamicIndexBuffer FDeferredShadingSceneRenderer::DynamicIndexBufferForInitViews;
 FGlobalDynamicIndexBuffer FDeferredShadingSceneRenderer::DynamicIndexBufferForInitShadows;
@@ -990,13 +1006,20 @@ bool FDeferredShadingSceneRenderer::GatherRayTracingWorldInstances(FRHICommandLi
 						Views[ViewIndex].RayTracingGeometryInstances.Add(RayTracingInstance);
 					}
 
-					for (int32 SegmentIndex = 0; SegmentIndex < Instance.Materials.Num(); SegmentIndex++)
+					if (GRayTracingParallelMeshBatchSetup && FApp::ShouldUseThreadingForPerformance())
 					{
-						FMeshBatch& MeshBatch = Instance.Materials[SegmentIndex];
-						FDynamicRayTracingMeshCommandContext CommandContext(ReferenceView.DynamicRayTracingMeshCommandStorage, ReferenceView.VisibleRayTracingMeshCommands, SegmentIndex, InstanceIndex);
-						FRayTracingMeshProcessor RayTracingMeshProcessor(&CommandContext, Scene, &ReferenceView);
+						ReferenceView.AddRayTracingMeshBatchData.Emplace(Instance.Materials, SceneProxy, InstanceIndex);
+					}
+					else
+					{
+						for (int32 SegmentIndex = 0; SegmentIndex < Instance.Materials.Num(); SegmentIndex++)
+						{
+							FMeshBatch& MeshBatch = Instance.Materials[SegmentIndex];
+							FDynamicRayTracingMeshCommandContext CommandContext(ReferenceView.DynamicRayTracingMeshCommandStorage, ReferenceView.VisibleRayTracingMeshCommands, SegmentIndex, InstanceIndex);
+							FRayTracingMeshProcessor RayTracingMeshProcessor(&CommandContext, Scene, &ReferenceView);
 
-						RayTracingMeshProcessor.AddMeshBatch(MeshBatch, 1, SceneProxy);
+							RayTracingMeshProcessor.AddMeshBatch(MeshBatch, 1, SceneProxy);
+						}
 					}
 				}
 			}
@@ -1005,6 +1028,47 @@ bool FDeferredShadingSceneRenderer::GatherRayTracingWorldInstances(FRHICommandLi
 
 	//
 
+	if (ReferenceView.AddRayTracingMeshBatchData.Num() > 0)
+	{
+		const uint32 NumTotalItems = ReferenceView.AddRayTracingMeshBatchData.Num();
+		const uint32 TargetItemsPerTask = GRayTracingParallelMeshBatchSize;//128; // Granularity based on profiling several scenes
+		const uint32 NumTasks = FMath::Max(1u, FMath::DivideAndRoundUp(NumTotalItems, TargetItemsPerTask));
+		const uint32 BatchSize = FMath::DivideAndRoundUp(NumTotalItems, NumTasks);
+
+		ReferenceView.DynamicRayTracingMeshCommandStorageParallel.Init(FDynamicRayTracingMeshCommandStorage(), BatchSize);
+		ReferenceView.VisibleRayTracingMeshCommandsParallel.Init(FRayTracingMeshCommandOneFrameArray(), BatchSize);
+
+		for (uint32 Batch = 0; Batch < NumTasks; Batch++)
+		{
+			uint32 BatchStart = Batch * BatchSize;
+			uint32 BatchEnd = FMath::Min(BatchStart + BatchSize, (uint32) ReferenceView.AddRayTracingMeshBatchData.Num());
+
+			ReferenceView.DynamicRayTracingMeshCommandStorageParallel[Batch].RayTracingMeshCommands.Reserve(Scene->Primitives.Num());
+			ReferenceView.VisibleRayTracingMeshCommandsParallel[Batch].Reserve(Scene->Primitives.Num());
+
+			ReferenceView.AddRayTracingMeshBatchTaskList.Add(FFunctionGraphTask::CreateAndDispatchWhenReady(
+				[&ReferenceView, Scene = this->Scene, Batch, BatchStart, BatchEnd]()
+			{
+				for (uint32 Index = BatchStart; Index < BatchEnd; Index++)
+				{
+					FRayTracingMeshBatchWorkItem& MeshBatchJob = ReferenceView.AddRayTracingMeshBatchData[Index];
+					for (uint32 SegmentIndex = 0; SegmentIndex < (uint32) MeshBatchJob.MeshBatches.Num(); SegmentIndex++)
+					{
+						FMeshBatch& MeshBatch = MeshBatchJob.MeshBatches[SegmentIndex];
+						const FPrimitiveSceneProxy* SceneProxy = MeshBatchJob.SceneProxy;
+						const uint32 InstanceIndex = MeshBatchJob.InstanceIndex;
+						FDynamicRayTracingMeshCommandContext CommandContext(ReferenceView.DynamicRayTracingMeshCommandStorageParallel[Batch], ReferenceView.VisibleRayTracingMeshCommandsParallel[Batch], SegmentIndex, InstanceIndex);
+						FRayTracingMeshProcessor RayTracingMeshProcessor(&CommandContext, Scene, &ReferenceView);
+						RayTracingMeshProcessor.AddMeshBatch(MeshBatch, 1, SceneProxy);
+					}
+				}
+			},
+				TStatId(), nullptr, ENamedThreads::AnyThread));
+		}
+	}
+
+	//
+
 	{
 		TRACE_CPUPROFILER_EVENT_SCOPE(GatherRayTracingWorldInstances_AddInstances);
 
@@ -1095,6 +1159,22 @@ bool FDeferredShadingSceneRenderer::DispatchRayTracingWorldUpdates(FRHICommandLi
 
 	TRACE_CPUPROFILER_EVENT_SCOPE(FDeferredShadingSceneRenderer::DispatchRayTracingWorldUpdates);
 
+	FViewInfo& ReferenceView = Views[0];
+	if (ReferenceView.AddRayTracingMeshBatchTaskList.Num() > 0)
+	{
+		SCOPE_CYCLE_COUNTER(STAT_WaitRayTracingAddMesh);
+
+		FTaskGraphInterface::Get().WaitUntilTasksComplete(ReferenceView.AddRayTracingMeshBatchTaskList, ENamedThreads::GetRenderThread_Local());
+
+		for (int32 Batch = 0; Batch < ReferenceView.AddRayTracingMeshBatchTaskList.Num(); Batch++)
+		{
+			ReferenceView.VisibleRayTracingMeshCommands.Append(ReferenceView.VisibleRayTracingMeshCommandsParallel[Batch]);
+		}
+
+		ReferenceView.AddRayTracingMeshBatchTaskList.Empty();
+		ReferenceView.AddRayTracingMeshBatchData.Empty();
+	}
+
 	bool bAsyncUpdateGeometry = (CVarRayTracingAsyncBuild.GetValueOnRenderThread() != 0)
 							  && GRHISupportsRayTracingAsyncBuildAccelerationStructure;
 
diff --git a/Engine/Source/Runtime/Renderer/Private/SceneRendering.h b/Engine/Source/Runtime/Renderer/Private/SceneRendering.h
--- a/Engine/Source/Runtime/Renderer/Private/SceneRendering.h
+++ b/Engine/Source/Runtime/Renderer/Private/SceneRendering.h
@@ -858,6 +858,22 @@ public:
 
 typedef TArray<FViewCommands, TInlineAllocator<4>> FViewVisibleCommandsPerView;
 
+#if RHI_RAYTRACING
+struct FRayTracingMeshBatchWorkItem
+{
+	FRayTracingMeshBatchWorkItem(TArray<FMeshBatch>& InBatches, FPrimitiveSceneProxy* InSceneProxy, uint32 InInstanceIndex) :
+		SceneProxy(InSceneProxy),
+		InstanceIndex(InInstanceIndex)
+	{
+		Swap(MeshBatches, InBatches);
+	}
+
+	TArray<FMeshBatch> MeshBatches;
+	FPrimitiveSceneProxy* SceneProxy;
+	uint32 InstanceIndex;
+};
+#endif
+
 /** A FSceneView with additional state used by the scene renderer. */
 class FViewInfo : public FSceneView
 {
@@ -998,10 +1014,14 @@ public:
 	
 #if RHI_RAYTRACING
 	TUniquePtr<FRayTracingMeshResourceCollector> RayTracingMeshResourceCollector;
-
 	FRayTracingMeshCommandOneFrameArray VisibleRayTracingMeshCommands;
-
 	FDynamicRayTracingMeshCommandStorage DynamicRayTracingMeshCommandStorage;
+
+	FGraphEventArray AddRayTracingMeshBatchTaskList;
+	TArray<FRayTracingMeshBatchWorkItem> AddRayTracingMeshBatchData;
+
+	TArray<FRayTracingMeshCommandOneFrameArray> VisibleRayTracingMeshCommandsParallel;
+	TArray<FDynamicRayTracingMeshCommandStorage> DynamicRayTracingMeshCommandStorageParallel;
 #endif
 
 	// Used by mobile renderer to determine whether static meshes will be rendered with CSM shaders or not.
