Fix a bug of missing dof coc info in TonemapperES2

diff --git a/Engine/Shaders/Private/PostProcessMobile.usf b/Engine/Shaders/Private/PostProcessMobile.usf
--- a/Engine/Shaders/Private/PostProcessMobile.usf
+++ b/Engine/Shaders/Private/PostProcessMobile.usf
@@ -806,8 +806,9 @@ void IntegrateDOFPS_ES2(
 
 	half4 DofFine = Texture2DSample(PostprocessInput0, PostprocessInput0Sampler, FineDofGrain.xy);
 	half4 Dof = Texture2DSample(PostprocessInput1, PostprocessInput1Sampler, TexCoords.zw);
+	half DofCoc = Texture2DSample(PostprocessInput2, PostprocessInput2Sampler, TexCoords.xy).r;
 	// Convert alpha back into circle of confusion.
-	OutColor.a = max(Dof.a, abs(SceneColor.a * 2.0 - 1.0));
+	SceneColor.a = max(Dof.a, abs(DofCoc * 2.0 - 1.0));
 	// Convert circle of confusion into blend factors.		
 	half2 ScaleBias = CocBlendScaleBias(); // Constant.
 	half DofAmount = saturate(SceneColor.a * ScaleBias.x + ScaleBias.y);
@@ -816,7 +817,8 @@ void IntegrateDOFPS_ES2(
 	// Blend in fine DOF.
 	OutColor.rgb = lerp(SceneColor.rgb, DofFine.rgb, DofAmountFine);
 	// Blend in coarse DOF.
-	OutColor.rgb = lerp(SceneColor.rgb, Dof.rgb, DofAmount);
+	OutColor.rgb = lerp(OutColor.rgb, Dof.rgb, DofAmount);
+	OutColor.a = DofCoc;
 }
 
 //
diff --git a/Engine/Shaders/Private/PostProcessTonemap.usf b/Engine/Shaders/Private/PostProcessTonemap.usf
--- a/Engine/Shaders/Private/PostProcessTonemap.usf
+++ b/Engine/Shaders/Private/PostProcessTonemap.usf
@@ -836,8 +836,9 @@ void MainPS_ES2(
 		#if USE_DOF
 			half4 DofFine = Texture2DSample(PostprocessInput0, PostprocessInput0Sampler, FineDofGrain.xy);
 			half4 Dof = Texture2DSample(PostprocessInput2, PostprocessInput2Sampler, FullViewUV.zw);
+			half DofCoc = Texture2DSample(PostprocessInput4, PostprocessInput4Sampler, UV).r;
 			// Convert alpha back into circle of confusion.
-			SceneColor.a = max(Dof.a, abs(SceneColor.a * 2.0 - 1.0));
+			SceneColor.a = max(Dof.a, abs(DofCoc * 2.0 - 1.0));
 			// Convert circle of confusion into blend factors.		
 			half2 ScaleBias = CocBlendScaleBias(); // Constant.
 			half DofAmount = saturate(SceneColor.a * ScaleBias.x + ScaleBias.y);
diff --git a/Engine/Source/Runtime/Renderer/Private/PostProcess/PostProcessMobile.h b/Engine/Source/Runtime/Renderer/Private/PostProcess/PostProcessMobile.h
--- a/Engine/Source/Runtime/Renderer/Private/PostProcess/PostProcessMobile.h
+++ b/Engine/Source/Runtime/Renderer/Private/PostProcess/PostProcessMobile.h
@@ -80,7 +80,7 @@ private:
 	FIntPoint PrePostSourceViewportSize;
 };
 
-class FRCPassIntegrateDofES2 : public TRenderingCompositePassBase<2, 1>
+class FRCPassIntegrateDofES2 : public TRenderingCompositePassBase<3, 1>
 {
 public:
 	FRCPassIntegrateDofES2(FIntPoint InPrePostSourceViewportSize) : PrePostSourceViewportSize(InPrePostSourceViewportSize) { }
diff --git a/Engine/Source/Runtime/Renderer/Private/PostProcess/PostProcessTonemap.h b/Engine/Source/Runtime/Renderer/Private/PostProcess/PostProcessTonemap.h
--- a/Engine/Source/Runtime/Renderer/Private/PostProcess/PostProcessTonemap.h
+++ b/Engine/Source/Runtime/Renderer/Private/PostProcess/PostProcessTonemap.h
@@ -114,7 +114,7 @@ private:
 // ePId_Input0: SceneColor
 // ePId_Input1: BloomCombined (not needed for bDoGammaOnly)
 // ePId_Input2: Dof (not needed for bDoGammaOnly)
-class FRCPassPostProcessTonemapES2 : public TRenderingCompositePassBase<4, 1>
+class FRCPassPostProcessTonemapES2 : public TRenderingCompositePassBase<5, 1>
 {
 public:
 	FRCPassPostProcessTonemapES2(const FViewInfo& View, bool bInUsedFramebufferFetch, bool bInSRGBAwareTarget, bool bInDoEyeAdaptation, bool bMetalMSAAHDRDecode);
diff --git a/Engine/Source/Runtime/Renderer/Private/PostProcess/PostProcessing.cpp b/Engine/Source/Runtime/Renderer/Private/PostProcess/PostProcessing.cpp
--- a/Engine/Source/Runtime/Renderer/Private/PostProcess/PostProcessing.cpp
+++ b/Engine/Source/Runtime/Renderer/Private/PostProcess/PostProcessing.cpp
@@ -1286,6 +1286,7 @@ void FPostProcessing::ProcessES2(FRHICommandListImmediate& RHICmdList, FScene* S
 		FPostprocessContext Context(RHICmdList, CompositeContext.Graph, View);
 		FRenderingCompositeOutputRef BloomOutput;
 		FRenderingCompositeOutputRef DofOutput;
+		FRenderingCompositeOutputRef PostProcessSunShaftAndDof;
 
 		bool bUseAa = View.AntiAliasingMethod == AAM_TemporalAA;
 
@@ -1377,7 +1378,6 @@ void FPostProcessing::ProcessES2(FRHICommandListImmediate& RHICmdList, FScene* S
 			// Optional fixed pass processes
 			if (bUsePost && (bUseSun | bUseDof | bUseBloom | bUseVignette | bUseBasicEyeAdaptation | bUseHistogramEyeAdaptation))
 			{
-				FRenderingCompositeOutputRef PostProcessSunShaftAndDof;
 				if (bUseSun || bUseDof)
 				{
 					bool bUseDepthTexture = Context.FinalOutput.GetOutput()->RenderTargetDesc.Format == PF_FloatR11G11B10;
@@ -1478,6 +1478,7 @@ void FPostProcessing::ProcessES2(FRHICommandListImmediate& RHICmdList, FScene* S
 								FRenderingCompositePass* Pass = Context.Graph.RegisterPass(new(FMemStack::Get()) FRCPassIntegrateDofES2(FinalOutputViewRect.Size()));
 								Pass->SetInput(ePId_Input0, Context.FinalOutput);
 								Pass->SetInput(ePId_Input1, DofOutput);
+								Pass->SetInput(ePId_Input2, PostProcessSunShaftAndDof);
 								Context.FinalOutput = FRenderingCompositeOutputRef(Pass);
 							}
 						}
@@ -1707,6 +1708,7 @@ void FPostProcessing::ProcessES2(FRHICommandListImmediate& RHICmdList, FScene* S
 				}
 				PostProcessTonemap->SetInput(ePId_Input2, DofOutput);
 				PostProcessTonemap->SetInput(ePId_Input3, PostProcessEyeAdaptation);
+				PostProcessTonemap->SetInput(ePId_Input4, PostProcessSunShaftAndDof);
 
 				Context.FinalOutput = FRenderingCompositeOutputRef(PostProcessTonemap);
 
