UE4.25 - Copy/paste fix
#rnx

diff --git a/Engine/Source/Runtime/RenderCore/Private/RenderGraphBuilder.cpp b/Engine/Source/Runtime/RenderCore/Private/RenderGraphBuilder.cpp
--- a/Engine/Source/Runtime/RenderCore/Private/RenderGraphBuilder.cpp
+++ b/Engine/Source/Runtime/RenderCore/Private/RenderGraphBuilder.cpp
@@ -721,7 +721,7 @@ void FRDGBuilder::ExecutePass(const FRDGPass* Pass)
 	{
 		PRAGMA_DISABLE_DEPRECATION_WARNINGS
 		UnbindRenderTargets(RHICmdList);
-		PRAGMA_DISABLE_DEPRECATION_WARNINGS
+		PRAGMA_ENABLE_DEPRECATION_WARNINGS
 	}
 
 	Pass->Execute(RHICmdList);
diff --git a/Engine/Source/Runtime/Renderer/Private/CapsuleShadowRendering.cpp b/Engine/Source/Runtime/Renderer/Private/CapsuleShadowRendering.cpp
--- a/Engine/Source/Runtime/Renderer/Private/CapsuleShadowRendering.cpp
+++ b/Engine/Source/Runtime/Renderer/Private/CapsuleShadowRendering.cpp
@@ -1551,7 +1551,7 @@ void FDeferredShadingSceneRenderer::RenderCapsuleShadowsForMovableSkylight(FRHIC
 
 						PRAGMA_DISABLE_DEPRECATION_WARNINGS
 						UnbindRenderTargets(RHICmdList);
-						PRAGMA_DISABLE_DEPRECATION_WARNINGS
+						PRAGMA_ENABLE_DEPRECATION_WARNINGS
 
 						FIntRect ScissorRect = View.ViewRect;
 
