Fix VLM crash when there's no valid indirection texture data but sublevel brick data (usually because a sublevel is loaded as standalone)

diff --git a/Engine/Source/Runtime/Renderer/Private/ScenePrivate.h b/Engine/Source/Runtime/Renderer/Private/ScenePrivate.h
--- a/Engine/Source/Runtime/Renderer/Private/ScenePrivate.h
+++ b/Engine/Source/Runtime/Renderer/Private/ScenePrivate.h
@@ -1632,7 +1632,7 @@ class FVolumetricLightmapSceneData
 public:
 
 	FVolumetricLightmapSceneData() { GlobalVolumetricLightmap.Data = &GlobalVolumetricLightmapData; }
-	bool HasData() const { return LevelVolumetricLightmaps.Num() > 0; }
+	bool HasData() const { return GlobalVolumetricLightmapData.IndirectionTexture.Data.Num() && LevelVolumetricLightmaps.Num() > 0; }
 	void AddLevelVolume(const class FPrecomputedVolumetricLightmap* InVolume, EShadingPath ShadingPath, bool bIsPersistentLevel);
 	void RemoveLevelVolume(const class FPrecomputedVolumetricLightmap* InVolume);
 	const FPrecomputedVolumetricLightmap* GetLevelVolumetricLightmap() const 
