FNEM - only clear out the padding when a parameter is set.

diff --git a/Engine/Source/Runtime/Renderer/Public/MeshDrawShaderBindings.h b/Engine/Source/Runtime/Renderer/Public/MeshDrawShaderBindings.h
--- a/Engine/Source/Runtime/Renderer/Public/MeshDrawShaderBindings.h
+++ b/Engine/Source/Runtime/Renderer/Public/MeshDrawShaderBindings.h
@@ -103,7 +103,6 @@ public:
 		FMeshDrawShaderBindingsLayout(InLayout)
 	{
 		Data = InData;
-		memset(GetLooseDataStart(), 0, GetLooseDataSizeBytes());
 	}
 
 	template<typename UniformBufferStructType>
@@ -205,6 +204,8 @@ public:
 							ensureMsgf(sizeof(ParameterType) == Parameter.GetNumBytes(), TEXT("Attempted to set fewer bytes than the shader required.  Setting %u bytes on loose parameter at BaseIndex %u, Size %u.  This can cause GPU hangs, depending on usage."), sizeof(ParameterType), Parameter.GetBaseIndex(), Parameter.GetNumBytes());
 							const int32 NumBytesToSet = FMath::Min<int32>(sizeof(ParameterType), Parameter.GetNumBytes());
 							FMemory::Memcpy(LooseDataOffset, &Value, NumBytesToSet);
+							const int32 NumBytesToClear = FMath::Min<int32>(0, Parameter.GetNumBytes() - NumBytesToSet);
+							FMemory::Memset(LooseDataOffset + NumBytesToSet, 0x00, NumBytesToClear);
 							bFoundParameter = true;
 							break;
 						}
