Fixed potential crash when there is no velocity data and debug visualization modes are set through Cvars in game mode

diff --git a/Engine/Source/Runtime/Renderer/Private/ScreenSpaceDenoise.cpp b/Engine/Source/Runtime/Renderer/Private/ScreenSpaceDenoise.cpp
--- a/Engine/Source/Runtime/Renderer/Private/ScreenSpaceDenoise.cpp
+++ b/Engine/Source/Runtime/Renderer/Private/ScreenSpaceDenoise.cpp
@@ -1478,7 +1478,7 @@ static void DenoiseSignalAtConstantPixelDensity(
 		CommonParameters.EyeAdaptation = GetEyeAdaptationTexture(GraphBuilder, View);
 
 		// Remove dependency of the velocity buffer on camera cut, given it's going to be ignored by the shaders.
-		if (View.bCameraCut)
+		if (View.bCameraCut || !CommonParameters.SceneTextures.SceneVelocityBuffer)
 		{
 			CommonParameters.SceneTextures.SceneVelocityBuffer = GraphBuilder.RegisterExternalTexture(GSystemTextures.BlackDummy);
 		}
