DXC shader compilation validation error fixes

- FilterVertexShader.usf, FilterPixelShader.usf
  - DXC is unable to determine conditional and associated indexing are compile time constant so refactor to give it some help
  - unroll a small loop in the VS, otherwise DXC is unhappy with the indexing

- SingleLayerWaterRendering.cpp
  - ForwardLightData uniform parameters were used but unbound in FSingleLayerWaterCompositePS
    - added the parameter macro and bound the associated global uniform buffer from the view

- LightRendering.cpp
  - CapsuleLightIntegrate.ush declares a texture2d called DummyRectLightTextureForCapsuleCompilerWarning
    - its used to initalize the texture member in an FAreaLight struct but that texture is unused in this scenario
    - I assume this was added to work around a warning that the FAreaLight was not fully initialized
      - unfortunately, unoptimized DXC does not strip this parameter and UE throws an error that we didn't bind anything to it

#rnx

diff --git a/Engine/Shaders/Private/FilterPixelShader.usf b/Engine/Shaders/Private/FilterPixelShader.usf
--- a/Engine/Shaders/Private/FilterPixelShader.usf
+++ b/Engine/Shaders/Private/FilterPixelShader.usf
@@ -120,19 +120,17 @@ void MainPS(
 		Color += SampleFilterTexture(InOffsetUVs[i]) * SampleWeights[i];
 	}
 #else
-	int SampleIndex = 0;
-
-	UNROLL for (; SampleIndex < STATIC_SAMPLE_COUNT - 1; SampleIndex += 2)
+	UNROLL for (int SampleIndex = 0; SampleIndex < STATIC_SAMPLE_COUNT - 1; SampleIndex += 2)
 	{
 		float4 UVUV = InOffsetUVs[SampleIndex / 2];
 		Color += SampleFilterTexture(UVUV.xy) * SampleWeights[SampleIndex + 0];
 		Color += SampleFilterTexture(UVUV.zw) * SampleWeights[SampleIndex + 1];
 	}
 
-	FLATTEN if (SampleIndex < STATIC_SAMPLE_COUNT)
+	FLATTEN if (STATIC_SAMPLE_COUNT & 1)
 	{
-		float2 UV = InOffsetUVs[SampleIndex / 2].xy;
-		Color += SampleFilterTexture(UV) * SampleWeights[SampleIndex + 0];
+		float2 UV = InOffsetUVs[PACKED_STATIC_SAMPLE_COUNT - 1].xy;
+		Color += SampleFilterTexture(UV) * SampleWeights[STATIC_SAMPLE_COUNT - 1];
 	}
 #endif
 
diff --git a/Engine/Shaders/Private/FilterVertexShader.usf b/Engine/Shaders/Private/FilterVertexShader.usf
--- a/Engine/Shaders/Private/FilterVertexShader.usf
+++ b/Engine/Shaders/Private/FilterVertexShader.usf
@@ -27,19 +27,19 @@ void MainVS(
 {
 	DrawRectangle(InPosition, UV, OutPosition, UV);
 
-	int SampleIndex = 0;
+	
 
-	for (;SampleIndex < STATIC_SAMPLE_COUNT - 1; SampleIndex += 2)
+	for (int SampleIndex = 0; SampleIndex < STATIC_SAMPLE_COUNT - 1; SampleIndex += 2)
 	{
 		half4 OffsetUVUV = SampleOffsets[SampleIndex / 2];
 		OutOffsetUVs[SampleIndex + 0] = UV.xy + OffsetUVUV.xy;
 		OutOffsetUVs[SampleIndex + 1] = UV.xy + OffsetUVUV.zw;
 	}
 
-	if (SampleIndex < STATIC_SAMPLE_COUNT)
+	if (STATIC_SAMPLE_COUNT & 1)
 	{
-		half4 OffsetUVUV = SampleOffsets[SampleIndex / 2];
-		OutOffsetUVs[SampleIndex] = UV.xy + OffsetUVUV.xy;
+		half4 OffsetUVUV = SampleOffsets[PACKED_STATIC_SAMPLE_COUNT - 1];
+		OutOffsetUVs[STATIC_SAMPLE_COUNT - 1] = UV.xy + OffsetUVUV.xy;
 	}
 
 	OutUV = UV;
@@ -57,7 +57,7 @@ void MainVS(
 {
 	DrawRectangle(InPosition, UV, OutPosition, UV);
 
-	for(int OffsetIndex = 0; OffsetIndex < PACKED_STATIC_SAMPLE_COUNT; ++OffsetIndex)
+	UNROLL for(int OffsetIndex = 0; OffsetIndex < PACKED_STATIC_SAMPLE_COUNT; ++OffsetIndex)
 	{
 		OutOffsetUVs[OffsetIndex] = UV.xyxy + SampleOffsets[OffsetIndex];
 	}
diff --git a/Engine/Source/Runtime/Renderer/Private/LightRendering.cpp b/Engine/Source/Runtime/Renderer/Private/LightRendering.cpp
--- a/Engine/Source/Runtime/Renderer/Private/LightRendering.cpp
+++ b/Engine/Source/Runtime/Renderer/Private/LightRendering.cpp
@@ -395,6 +395,8 @@ class FDeferredLightPS : public FGlobalShader
 		HairVisibilityNodeOffsetAndCount.Bind(Initializer.ParameterMap, TEXT("HairVisibilityNodeOffsetAndCount"));
 		HairVisibilityNodeCoords.Bind(Initializer.ParameterMap, TEXT("HairVisibilityNodeCoords"));
 		HairVisibilityNodeData.Bind(Initializer.ParameterMap, TEXT("HairVisibilityNodeData"));
+
+		DummyRectLightTextureForCapsuleCompilerWarning.Bind(Initializer.ParameterMap, TEXT("DummyRectLightTextureForCapsuleCompilerWarning"));
 	}
 
 	FDeferredLightPS()
@@ -633,7 +635,18 @@ private:
 				HairDualScatteringRoughnessOverride,
 				DualScatteringRoughness);
 		}
-		
+
+		if (DummyRectLightTextureForCapsuleCompilerWarning.IsBound())
+		{
+			SetTextureParameter(
+				RHICmdList,
+				ShaderRHI,
+				DummyRectLightTextureForCapsuleCompilerWarning,
+				LTCMatSampler,
+				TStaticSamplerState<SF_Bilinear, AM_Clamp, AM_Clamp, AM_Clamp>::GetRHI(),
+				GSystemTextures.DepthDummy->GetRenderTargetItem().ShaderResourceTexture
+			);
+		}
 	}
 
 	LAYOUT_FIELD(FSceneTextureShaderParameters, SceneTextureParameters);
@@ -663,6 +676,8 @@ private:
 	LAYOUT_FIELD(FShaderParameter, HairComponents);
 	LAYOUT_FIELD(FShaderParameter, HairShadowMaskValid);
 	LAYOUT_FIELD(FShaderParameter, HairDualScatteringRoughnessOverride);
+
+	LAYOUT_FIELD(FShaderResourceParameter, DummyRectLightTextureForCapsuleCompilerWarning);
 };
 
 IMPLEMENT_GLOBAL_SHADER(FDeferredLightPS, "/Engine/Private/DeferredLightPixelShaders.usf", "DeferredLightPixelMain", SF_Pixel);
diff --git a/Engine/Source/Runtime/Renderer/Private/SingleLayerWaterRendering.cpp b/Engine/Source/Runtime/Renderer/Private/SingleLayerWaterRendering.cpp
--- a/Engine/Source/Runtime/Renderer/Private/SingleLayerWaterRendering.cpp
+++ b/Engine/Source/Runtime/Renderer/Private/SingleLayerWaterRendering.cpp
@@ -263,6 +263,7 @@ BEGIN_SHADER_PARAMETER_STRUCT(FSingleLayerWaterCommonShaderParameters, )
 	SHADER_PARAMETER_STRUCT_REF(FViewUniformShaderParameters, ViewUniformBuffer)
 	SHADER_PARAMETER_STRUCT_REF(FReflectionCaptureShaderData, ReflectionCaptureData)
 	SHADER_PARAMETER_STRUCT_REF(FReflectionUniformParameters, ReflectionsParameters)
+	SHADER_PARAMETER_STRUCT_REF(FForwardLightData, ForwardLightData)
 END_SHADER_PARAMETER_STRUCT()
 
 class FSingleLayerWaterCompositePS : public FGlobalShader
@@ -636,6 +637,7 @@ void FDeferredShadingSceneRenderer::RenderSingleLayerWaterReflections(FRHIComman
 				SetupReflectionUniformParameters(View, ReflectionUniformParameters);
 				Parameters.ReflectionsParameters = CreateUniformBufferImmediate(ReflectionUniformParameters, UniformBuffer_SingleDraw);
 			}
+			Parameters.ForwardLightData = View.ForwardLightingResources->ForwardLightDataUniformBuffer;
 		};
 
 		const bool bRunTiled = UseSingleLayerWaterIndirectDraw(View.GetShaderPlatform()) && CVarWaterSingleLayerTiledComposite.GetValueOnRenderThread();
