Remove some shader compilation workarounds added for previous DXC version

#rnx

diff --git a/Engine/Source/Runtime/Renderer/Private/DistanceFieldShadowing.cpp b/Engine/Source/Runtime/Renderer/Private/DistanceFieldShadowing.cpp
--- a/Engine/Source/Runtime/Renderer/Private/DistanceFieldShadowing.cpp
+++ b/Engine/Source/Runtime/Renderer/Private/DistanceFieldShadowing.cpp
@@ -463,9 +463,6 @@ public:
 		OutEnvironment.SetDefine(TEXT("DISTANCEFIELD_PRIMITIVE_TYPE"), (int32)PrimitiveType);
 		OutEnvironment.SetDefine(TEXT("HAS_PREVIOUS_OUTPUT"), (int32)bHasPrevOutput);
 		OutEnvironment.SetDefine(TEXT("PLATFORM_SUPPORTS_TYPED_UAV_LOAD"), (int32)RHISupports4ComponentUAVReadWrite(Parameters.Platform));
-
-		// temporary workaround for DXC
-		OutEnvironment.CompilerFlags.AddUnique(CFLAG_SkipOptimizationsDXC);
 	}
 
 	/** Default constructor. */
diff --git a/Engine/Source/Runtime/Renderer/Private/GPUFastFourierTransform.cpp b/Engine/Source/Runtime/Renderer/Private/GPUFastFourierTransform.cpp
--- a/Engine/Source/Runtime/Renderer/Private/GPUFastFourierTransform.cpp
+++ b/Engine/Source/Runtime/Renderer/Private/GPUFastFourierTransform.cpp
@@ -936,7 +936,6 @@ namespace GPUFFT
 			FGlobalShader::ModifyCompilationEnvironment(Parameters, OutEnvironment);
 			OutEnvironment.SetDefine(TEXT("INCLUDE_GROUP_SHARED_CONVOLUTION_WITH_TEXTURE"), 1);
 			OutEnvironment.SetDefine(TEXT("SCAN_LINE_LENGTH"), PowRadixSignalLength);
-			OutEnvironment.CompilerFlags.Add(CFLAG_ForceOptimization);
 		}
 
 		/** Default constructor **/
