Do not render raytraced DF shadow if the light should not be rendered

diff --git a/Engine/Source/Runtime/Renderer/Private/ShadowSetup.cpp b/Engine/Source/Runtime/Renderer/Private/ShadowSetup.cpp
--- a/Engine/Source/Runtime/Renderer/Private/ShadowSetup.cpp
+++ b/Engine/Source/Runtime/Renderer/Private/ShadowSetup.cpp
@@ -4162,7 +4162,8 @@ void FSceneRenderer::AddViewDependentWholeSceneShadowsForView(
 		// If rendering in stereo mode we render shadow depths only for the left eye, but project for both eyes!
 		if (IStereoRendering::IsAPrimaryView(View))
 		{
-			const bool bExtraDistanceFieldCascade = LightSceneInfo.Proxy->ShouldCreateRayTracedCascade(View.GetFeatureLevel(), LightSceneInfo.IsPrecomputedLightingValid(), View.MaxShadowCascades);
+			const bool bExtraDistanceFieldCascade = LightSceneInfo.ShouldRenderLightViewIndependent()
+				&& LightSceneInfo.Proxy->ShouldCreateRayTracedCascade(View.GetFeatureLevel(), LightSceneInfo.IsPrecomputedLightingValid(), View.MaxShadowCascades);
 
 			const int32 ProjectionCount = LightSceneInfo.Proxy->GetNumViewDependentWholeSceneShadows(View, LightSceneInfo.IsPrecomputedLightingValid()) + (bExtraDistanceFieldCascade?1:0);
 
