Fixed crash when HighFrequencyNoiseTexture is missing by removing its binding as it is not used in the tonemap shader anymore.

diff --git a/Engine/Shaders/Private/PostProcessTonemap.usf b/Engine/Shaders/Private/PostProcessTonemap.usf
--- a/Engine/Shaders/Private/PostProcessTonemap.usf
+++ b/Engine/Shaders/Private/PostProcessTonemap.usf
@@ -39,10 +39,6 @@ SCREEN_PASS_TEXTURE_VIEWPORT(Bloom)
 SCREEN_PASS_TEXTURE_VIEWPORT(Output)
 SCREEN_PASS_TEXTURE_VIEWPORT_TRANSFORM(ColorToBloom)
 
-// 64x64 tiling noise texture, optimized to have mostly high frequency content
-Texture2D NoiseTexture;
-SamplerState NoiseSampler;
-
 Texture2D ColorTexture;
 SamplerState ColorSampler;
 
@@ -133,14 +129,6 @@ half3 ColorLookupTable( half3 LinearColor )
 	return OutDeviceColor * 1.05;
 }
 
-// xy should be a integer position (e.g. pixel position on the screen)
-// use the PseudoRandom() function if you have ALU performance left
-// and this if you have TEX performance left.
-float PseudoRandomTex(float2 xy)
-{
-	return Texture2DSample(NoiseTexture, NoiseSampler, xy / 64.0f).r;
-}
-
 float SwitchVerticalAxis;
 
 #ifndef NEEDTOSWITCHVERTICLEAXIS
diff --git a/Engine/Source/Runtime/Renderer/Private/PostProcess/PostProcessTonemap.cpp b/Engine/Source/Runtime/Renderer/Private/PostProcess/PostProcessTonemap.cpp
--- a/Engine/Source/Runtime/Renderer/Private/PostProcess/PostProcessTonemap.cpp
+++ b/Engine/Source/Runtime/Renderer/Private/PostProcess/PostProcessTonemap.cpp
@@ -521,11 +521,9 @@ BEGIN_SHADER_PARAMETER_STRUCT(FTonemapParameters, )
 	// SM5 and above use Texture2D for EyeAdaptation
 	SHADER_PARAMETER_RDG_TEXTURE(Texture2D, EyeAdaptation)
 	SHADER_PARAMETER_RDG_TEXTURE(, ColorGradingLUT)
-	SHADER_PARAMETER_TEXTURE(Texture2D, NoiseTexture)
 	SHADER_PARAMETER_TEXTURE(Texture2D, BloomDirtMaskTexture)
 	SHADER_PARAMETER_SAMPLER(SamplerState, ColorSampler)
 	SHADER_PARAMETER_SAMPLER(SamplerState, BloomSampler)
-	SHADER_PARAMETER_SAMPLER(SamplerState, NoiseSampler)
 	SHADER_PARAMETER_SAMPLER(SamplerState, ColorGradingLUTSampler)
 	SHADER_PARAMETER_SAMPLER(SamplerState, BloomDirtMaskSampler)
 	SHADER_PARAMETER(FVector4, ColorScale0)
@@ -685,18 +683,6 @@ FScreenPassTexture AddTonemapPass(FRDGBuilder& GraphBuilder, const FViewInfo& Vi
 
 	const FScreenPassTextureViewport OutputViewport(Output);
 
-	FRHITexture* NoiseTexture = nullptr;
-
-	{
-		check(GEngine);
-		const auto* HighFrequencyNoiseTexture = GEngine->HighFrequencyNoiseTexture;
-		check(HighFrequencyNoiseTexture);
-		const auto* Resource = HighFrequencyNoiseTexture->Resource;
-		check(Resource);
-		NoiseTexture = Resource->TextureRHI;
-		check(NoiseTexture);
-	}
-
 	FRHITexture* BloomDirtMaskTexture = GSystemTextures.BlackDummy->GetRenderTargetItem().TargetableTexture;
 
 	if (PostProcessSettings.BloomDirtMask && PostProcessSettings.BloomDirtMask->Resource)
@@ -766,11 +752,9 @@ FScreenPassTexture AddTonemapPass(FRDGBuilder& GraphBuilder, const FViewInfo& Vi
 	CommonParameters.BloomTexture = Inputs.Bloom.Texture;
 	CommonParameters.EyeAdaptation = Inputs.EyeAdaptationTexture;
 	CommonParameters.ColorGradingLUT = Inputs.ColorGradingTexture;
-	CommonParameters.NoiseTexture = NoiseTexture;
 	CommonParameters.BloomDirtMaskTexture = BloomDirtMaskTexture;
 	CommonParameters.ColorSampler = BilinearClampSampler;
 	CommonParameters.BloomSampler = BilinearClampSampler;
-	CommonParameters.NoiseSampler = PointClampSampler;
 	CommonParameters.ColorGradingLUTSampler = BilinearClampSampler;
 	CommonParameters.BloomDirtMaskSampler = BilinearClampSampler;
 	CommonParameters.ColorScale0 = PostProcessSettings.SceneColorTint;
