Fix deep shadow map support when raytracing is enabled.

diff --git a/Engine/Source/Runtime/Renderer/Private/HairStrands/HairStrandsTransmittance.cpp b/Engine/Source/Runtime/Renderer/Private/HairStrands/HairStrandsTransmittance.cpp
--- a/Engine/Source/Runtime/Renderer/Private/HairStrands/HairStrandsTransmittance.cpp
+++ b/Engine/Source/Runtime/Renderer/Private/HairStrands/HairStrandsTransmittance.cpp
@@ -606,12 +606,12 @@ FHairStrandsTransmittanceMaskData RenderHairStrandsTransmittanceMask(
 ///////////////////////////////////////////////////////////////////////////////////////////////////
 
 static void RenderHairStrandsShadowMask(
-	FRHICommandListImmediate& RHICmdList,
+	FRDGBuilder& GraphBuilder,
 	const FViewInfo& View,
 	const FLightSceneInfo* LightSceneInfo,
-	IPooledRenderTarget* ScreenShadowMaskTexture,
 	const FHairStrandsVisibilityData& InVisibilityData,
-	const FHairStrandsMacroGroupDatas& InMacroGroupDatas)
+	const FHairStrandsMacroGroupDatas& InMacroGroupDatas,
+	FRDGTextureRef OutShadowMask)
 {
 	if (InMacroGroupDatas.Datas.Num() == 0)
 		return;
@@ -620,15 +620,13 @@ static void RenderHairStrandsShadowMask(
 		return;
 
 	DECLARE_GPU_STAT(HairStrandsOpaqueMask);
-	SCOPED_DRAW_EVENT(RHICmdList, HairStrandsOpaqueMask);
-	SCOPED_GPU_STAT(RHICmdList, HairStrandsOpaqueMask);
+	SCOPED_DRAW_EVENT(GraphBuilder.RHICmdList, HairStrandsOpaqueMask);
+	SCOPED_GPU_STAT(GraphBuilder.RHICmdList, HairStrandsOpaqueMask);
 
-	FSceneRenderTargets& SceneTargets = FSceneRenderTargets::Get(RHICmdList);
-	FRDGBuilder GraphBuilder(RHICmdList);
+	FSceneRenderTargets& SceneTargets = FSceneRenderTargets::Get(GraphBuilder.RHICmdList);
 
 	FRDGTextureRef SceneDepthTexture = GraphBuilder.RegisterExternalTexture(SceneTargets.SceneDepthZ, TEXT("SceneDephtTexture"));
 	FRDGTextureRef Categorization = RegisterExternalTextureWithFallback(GraphBuilder, InVisibilityData.CategorizationTexture, GSystemTextures.BlackDummy);
-	FRDGTextureRef OutShadowMask = GraphBuilder.RegisterExternalTexture(ScreenShadowMaskTexture, TEXT("ScreenShadowMaskTexture"));
 
 	bool bHasDeepShadow = false;
 	if (!IsHairStrandsForVoxelTransmittanceAndShadowEnable())
@@ -704,20 +702,20 @@ static void RenderHairStrandsShadowMask(
 				OutShadowMask);
 		}
 	}
-
-	TRefCountPtr<IPooledRenderTarget> LocalOutput = GSystemTextures.BlackDummy;
-	GraphBuilder.QueueTextureExtraction(OutShadowMask, &LocalOutput);
-
-	GraphBuilder.Execute();
 }
 
 void RenderHairStrandsShadowMask(
-	FRHICommandListImmediate& RHICmdList,
+	FRDGBuilder& GraphBuilder,
 	const TArray<FViewInfo>& Views,
 	const FLightSceneInfo* LightSceneInfo,
-	IPooledRenderTarget* ScreenShadowMaskTexture,
-	const FHairStrandsDatas* HairDatas)
+	const FHairStrandsDatas* HairDatas,
+	FRDGTextureRef OutShadowMask)
 {
+	if (Views.Num() == 0 || HairDatas == nullptr || OutShadowMask == nullptr)
+	{
+		return;
+	}
+
 	const FHairStrandsVisibilityViews& HairVisibilityViews = HairDatas->HairVisibilityViews;
 	const FHairStrandsMacroGroupViews& MacroGroupViews = HairDatas->MacroGroupsPerViews;
 
@@ -727,7 +725,27 @@ void RenderHairStrandsShadowMask(
 		{
 			const FHairStrandsVisibilityData& HairVisibilityData = HairVisibilityViews.HairDatas[ViewIndex];
 			const FHairStrandsMacroGroupDatas& MacroGroupDatas = MacroGroupViews.Views[ViewIndex];
-			RenderHairStrandsShadowMask(RHICmdList, Views[ViewIndex], LightSceneInfo, ScreenShadowMaskTexture, HairVisibilityData, MacroGroupDatas);
+			RenderHairStrandsShadowMask(GraphBuilder, Views[ViewIndex], LightSceneInfo, HairVisibilityData, MacroGroupDatas, OutShadowMask);
 		}
 	}
 }
+
+void RenderHairStrandsShadowMask(
+	FRHICommandListImmediate& RHICmdList,
+	const TArray<FViewInfo>& Views,
+	const FLightSceneInfo* LightSceneInfo,
+	const FHairStrandsDatas* HairDatas,
+	IPooledRenderTarget* ScreenShadowMaskTexture)
+{
+	if (Views.Num() == 0 || HairDatas == nullptr || ScreenShadowMaskTexture == nullptr)
+	{
+		return;
+	}
+
+	FRDGBuilder GraphBuilder(RHICmdList);
+
+	FRDGTextureRef OutShadowMask = GraphBuilder.RegisterExternalTexture(ScreenShadowMaskTexture, TEXT("ScreenShadowMaskTexture"));
+	RenderHairStrandsShadowMask(GraphBuilder, Views, LightSceneInfo, HairDatas, OutShadowMask);
+
+	GraphBuilder.Execute();
+}
diff --git a/Engine/Source/Runtime/Renderer/Private/HairStrands/HairStrandsTransmittance.h b/Engine/Source/Runtime/Renderer/Private/HairStrands/HairStrandsTransmittance.h
--- a/Engine/Source/Runtime/Renderer/Private/HairStrands/HairStrandsTransmittance.h
+++ b/Engine/Source/Runtime/Renderer/Private/HairStrands/HairStrandsTransmittance.h
@@ -17,12 +17,19 @@ struct FHairStrandsTransmittanceMaskData
 };
 
 /// Write opaque hair shadow onto screen shadow mask to have fine hair details cast onto opaque geometries
+void RenderHairStrandsShadowMask(
+	FRDGBuilder& GraphBuilder,
+	const TArray<FViewInfo>& Views,
+	const FLightSceneInfo* LightSceneInfo,
+	const struct FHairStrandsDatas* HairDatas,
+	FRDGTextureRef ScreenShadowMaskTexture); 
+
 void RenderHairStrandsShadowMask(
 	FRHICommandListImmediate& RHICmdList,
 	const TArray<FViewInfo>& Views,
 	const class FLightSceneInfo* LightSceneInfo,
-	IPooledRenderTarget* ScreenShadowMaskTexture,
-	const struct FHairStrandsDatas* Hairdatas);
+	const struct FHairStrandsDatas* Hairdatas,
+	IPooledRenderTarget* ScreenShadowMaskTexture);
 
 /// Write hair transmittance onto screen shadow mask
 FHairStrandsTransmittanceMaskData RenderHairStrandsTransmittanceMask(
diff --git a/Engine/Source/Runtime/Renderer/Private/HairStrands/HairStrandsVoxelization.h b/Engine/Source/Runtime/Renderer/Private/HairStrands/HairStrandsVoxelization.h
--- a/Engine/Source/Runtime/Renderer/Private/HairStrands/HairStrandsVoxelization.h
+++ b/Engine/Source/Runtime/Renderer/Private/HairStrands/HairStrandsVoxelization.h
@@ -100,4 +100,5 @@ struct FHairStrandsOcclusionResources
 	FRDGTextureRef CategorizationTexture = nullptr;
 	FRDGTextureRef LightChannelMaskTexture = nullptr;
 	const FVirtualVoxelResources* VoxelResources = nullptr;
+	bool bUseHairVoxel = false;
 };
\ No newline at end of file
diff --git a/Engine/Source/Runtime/Renderer/Private/LightGridInjection.cpp b/Engine/Source/Runtime/Renderer/Private/LightGridInjection.cpp
--- a/Engine/Source/Runtime/Renderer/Private/LightGridInjection.cpp
+++ b/Engine/Source/Runtime/Renderer/Private/LightGridInjection.cpp
@@ -811,7 +811,7 @@ void FDeferredShadingSceneRenderer::RenderForwardShadingShadowProjections(FRHICo
 
 					if (InHairDatas)
 					{
-						RenderHairStrandsShadowMask(RHICmdList, Views, LightSceneInfo, ForwardScreenSpaceShadowMask, InHairDatas);
+						RenderHairStrandsShadowMask(RHICmdList, Views, LightSceneInfo, InHairDatas, ForwardScreenSpaceShadowMask);
 					}
 				}
 
diff --git a/Engine/Source/Runtime/Renderer/Private/LightRendering.cpp b/Engine/Source/Runtime/Renderer/Private/LightRendering.cpp
--- a/Engine/Source/Runtime/Renderer/Private/LightRendering.cpp
+++ b/Engine/Source/Runtime/Renderer/Private/LightRendering.cpp
@@ -985,36 +985,35 @@ void FDeferredShadingSceneRenderer::GatherAndSortLights(FSortedLightSetSceneInfo
 	check(OutSortedLights.AttenuationLightStart >= OutSortedLights.ClusteredSupportedEnd);
 }
 
+static bool HasHairStrandsClusters(int32 ViewIndex, const FHairStrandsDatas* HairDatas)
+{
+	return HairDatas && ViewIndex < HairDatas->MacroGroupsPerViews.Views.Num() && HairDatas->MacroGroupsPerViews.Views[ViewIndex].Datas.Num() > 0;
+};
 
+static FHairStrandsOcclusionResources GetHairStrandsResources(int32 ViewIndex, FRDGBuilder& GraphBuilder, const FHairStrandsDatas* HairDatas)
+{
+	FHairStrandsOcclusionResources Out;
+	if (HairDatas && ViewIndex < HairDatas->HairVisibilityViews.HairDatas.Num())
+	{
+		if (HairDatas->HairVisibilityViews.HairDatas[ViewIndex].CategorizationTexture)
+		{
+			Out.CategorizationTexture = GraphBuilder.RegisterExternalTexture(HairDatas->HairVisibilityViews.HairDatas[ViewIndex].CategorizationTexture);
+		}
+		if (HairDatas->HairVisibilityViews.HairDatas[ViewIndex].LightChannelMaskTexture)
+		{
+			Out.LightChannelMaskTexture = GraphBuilder.RegisterExternalTexture(HairDatas->HairVisibilityViews.HairDatas[ViewIndex].LightChannelMaskTexture);
+		}
+
+		Out.VoxelResources = &HairDatas->MacroGroupsPerViews.Views[ViewIndex].VirtualVoxelResources;
+	}
+	return Out;
+}
 
 /** Renders the scene's lighting. */
 void FDeferredShadingSceneRenderer::RenderLights(FRHICommandListImmediate& RHICmdList, FSortedLightSetSceneInfo &SortedLightSet, const FHairStrandsDatas* HairDatas)
 {
 	const bool bUseHairLighting = HairDatas != nullptr;
 	const FHairStrandsVisibilityViews* InHairVisibilityViews = HairDatas ? &HairDatas->HairVisibilityViews : nullptr;
-	auto HasHairStrandsClusters = [bUseHairLighting, HairDatas](int32 ViewIndex)
-	{
-		return HairDatas && ViewIndex < HairDatas->MacroGroupsPerViews.Views.Num() && HairDatas->MacroGroupsPerViews.Views[ViewIndex].Datas.Num() > 0;
-	};
-
-	auto GetHairStrandsResources = [HairDatas](int32 ViewIndex, FRDGBuilder& GraphBuilder)
-	{
-		FHairStrandsOcclusionResources Out;
-		if (HairDatas && ViewIndex < HairDatas->HairVisibilityViews.HairDatas.Num())
-		{
-			if (HairDatas->HairVisibilityViews.HairDatas[ViewIndex].CategorizationTexture)
-			{
-				Out.CategorizationTexture = GraphBuilder.RegisterExternalTexture(HairDatas->HairVisibilityViews.HairDatas[ViewIndex].CategorizationTexture);
-			}
-			if (HairDatas->HairVisibilityViews.HairDatas[ViewIndex].LightChannelMaskTexture)
-			{
-				Out.LightChannelMaskTexture = GraphBuilder.RegisterExternalTexture(HairDatas->HairVisibilityViews.HairDatas[ViewIndex].LightChannelMaskTexture);
-			}
-
-			Out.VoxelResources = &HairDatas->MacroGroupsPerViews.Views[ViewIndex].VirtualVoxelResources;
-		}
-		return Out;
-	};
 
 	check(RHICmdList.IsOutsideRenderPass());
 
@@ -1273,7 +1272,7 @@ void FDeferredShadingSceneRenderer::RenderLights(FRHICommandListImmediate& RHICm
 
 				PreprocessedShadowMaskTextures.SetNum(SortedLights.Num());
 
-				if (HasHairStrandsClusters(ViewIndex))
+				if (HasHairStrandsClusters(ViewIndex, HairDatas))
 				{ 
 					PreprocessedShadowMaskSubPixelTextures.SetNum(SortedLights.Num());
 				}
@@ -1297,7 +1296,8 @@ void FDeferredShadingSceneRenderer::RenderLights(FRHICommandListImmediate& RHICm
 				bool bDrawPreviewIndicator = ViewFamily.EngineShowFlags.PreviewShadowsIndicator && !LightSceneInfo.IsPrecomputedLightingValid() && LightSceneInfo.Proxy->HasStaticShadowing();
 				bool bInjectedTranslucentVolume = false;
 				bool bUsedShadowMaskTexture = false;
-				const bool bDrawDeepShadow = bDrawShadows && bUseHairLighting;
+				const bool bDrawHairShadow = bDrawShadows && bUseHairLighting;
+				const bool bUseHairDeepShadow = bDrawShadows && bUseHairLighting && LightSceneInfo.Proxy->CastsHairStrandsDeepShadow();
 
 				FScopeCycleCounter Context(LightSceneInfo.Proxy->GetStatId());
 
@@ -1473,7 +1473,8 @@ void FDeferredShadingSceneRenderer::RenderLights(FRHICommandListImmediate& RHICm
 
 								FRDGTextureUAV* RayTracingShadowMaskUAV = GraphBuilder.CreateUAV(FRDGTextureUAVDesc(RayTracingShadowMaskTexture));
 								FRDGTextureUAV* RayHitDistanceUAV = GraphBuilder.CreateUAV(FRDGTextureUAVDesc(RayDistanceTexture));
-								FHairStrandsOcclusionResources HairResources = GetHairStrandsResources(ViewIndex, GraphBuilder);
+								FHairStrandsOcclusionResources HairResources = GetHairStrandsResources(ViewIndex, GraphBuilder, HairDatas);
+								HairResources.bUseHairVoxel = !BatchLightSceneInfo.Proxy->CastsHairStrandsDeepShadow();
 								{
 									RDG_EVENT_SCOPE(GraphBuilder, "%s", *BatchLightNameWithLevel);
 
@@ -1491,7 +1492,7 @@ void FDeferredShadingSceneRenderer::RenderLights(FRHICommandListImmediate& RHICm
 										RayHitDistanceUAV,
 										SubPixelRayTracingShadowMaskUAV);
 									
-									if (HasHairStrandsClusters(ViewIndex))
+									if (HasHairStrandsClusters(ViewIndex, HairDatas))
 									{
 										TRefCountPtr<IPooledRenderTarget>* RefDestination = &PreprocessedShadowMaskSubPixelTextures[LightBatchIndex - AttenuationLightStart];
 										check(*RefDestination == nullptr);
@@ -1565,6 +1566,12 @@ void FDeferredShadingSceneRenderer::RenderLights(FRHICommandListImmediate& RHICm
 							ScreenShadowMaskSubPixelTexture = PreprocessedShadowMaskSubPixelTextures[ShadowMaskIndex];
 							PreprocessedShadowMaskSubPixelTextures[ShadowMaskIndex] = nullptr;
 						}
+
+						// Inject deep shadow mask if the light supports it
+						if (bUseHairDeepShadow)
+						{
+							RenderHairStrandsShadowMask(RHICmdList, Views, &LightSceneInfo, HairDatas, ScreenShadowMaskTexture);
+						}
 					}
 					else if (OcclusionType == FLightOcclusionType::Raytraced)
 					{
@@ -1643,7 +1650,9 @@ void FDeferredShadingSceneRenderer::RenderLights(FRHICommandListImmediate& RHICm
 								DenoiserRequirements = DenoiserToUse->GetShadowRequirements(View, LightSceneInfo, RayTracingConfig);
 							}
 
-							FHairStrandsOcclusionResources HairResources = GetHairStrandsResources(ViewIndex, GraphBuilder);
+							FHairStrandsOcclusionResources HairResources = GetHairStrandsResources(ViewIndex, GraphBuilder, HairDatas);
+							HairResources.bUseHairVoxel = !bUseHairDeepShadow;
+
 							RenderRayTracingShadows(
 								GraphBuilder,
 								SceneTextures,
@@ -1698,13 +1707,19 @@ void FDeferredShadingSceneRenderer::RenderLights(FRHICommandListImmediate& RHICm
 								GraphBuilder.QueueTextureExtraction(RayTracingShadowMaskTexture, &ScreenShadowMaskTexture);
 							}
 
-							if (HasHairStrandsClusters(ViewIndex))
+							if (HasHairStrandsClusters(ViewIndex, HairDatas))
 							{
 								GraphBuilder.QueueTextureExtraction(SubPixelRayTracingShadowMaskTexture, &ScreenShadowMaskSubPixelTexture);
 							}
 						}
 
 						GraphBuilder.Execute();
+
+						// Inject deep shadow mask if the light supports it
+						if (HairDatas && bUseHairDeepShadow)
+						{
+							RenderHairStrandsShadowMask(RHICmdList, Views, &LightSceneInfo, HairDatas, ScreenShadowMaskTexture);
+						}
 					}
 					else // (OcclusionType == FOcclusionType::Shadowmap)
 					{
@@ -1835,7 +1850,7 @@ void FDeferredShadingSceneRenderer::RenderLights(FRHICommandListImmediate& RHICm
 				}
 
 				FHairStrandsTransmittanceMaskData TransmittanceMaskData;
-				if (bDrawDeepShadow)
+				if (bDrawHairShadow)
 				{
 					TransmittanceMaskData = RenderHairStrandsTransmittanceMask(RHICmdList, Views, &LightSceneInfo, HairDatas, ScreenShadowMaskSubPixelTexture);
 				}
@@ -2321,9 +2336,9 @@ void FDeferredShadingSceneRenderer::RenderLightsForHair(FRHICommandListImmediate
 			{
 				TRefCountPtr<IPooledRenderTarget> ScreenShadowMaskSubPixelTexture = InScreenShadowMaskSubPixelTexture;
 
-				const bool bDrawDeepShadow = SortedLightInfo.SortKey.Fields.bShadowed;
+				const bool bDrawHairShadow = SortedLightInfo.SortKey.Fields.bShadowed;
 				FHairStrandsTransmittanceMaskData TransmittanceMaskData;
-				if (bDrawDeepShadow)
+				if (bDrawHairShadow)
 				{
 					TransmittanceMaskData = RenderHairStrandsTransmittanceMask(RHICmdList, Views, &LightSceneInfo, HairDatas, ScreenShadowMaskSubPixelTexture);
 				}
@@ -2331,7 +2346,7 @@ void FDeferredShadingSceneRenderer::RenderLightsForHair(FRHICommandListImmediate
 				RenderLightForHair(RHICmdList, 
 					&LightSceneInfo,
 					ScreenShadowMaskSubPixelTexture,
-					bDrawDeepShadow ? &TransmittanceMaskData : nullptr,
+					bDrawHairShadow ? &TransmittanceMaskData : nullptr,
 					&HairDatas->HairVisibilityViews);
 			}
 		}
diff --git a/Engine/Source/Runtime/Renderer/Private/RayTracing/RayTracingShadows.cpp b/Engine/Source/Runtime/Renderer/Private/RayTracing/RayTracingShadows.cpp
--- a/Engine/Source/Runtime/Renderer/Private/RayTracing/RayTracingShadows.cpp
+++ b/Engine/Source/Runtime/Renderer/Private/RayTracing/RayTracingShadows.cpp
@@ -229,7 +229,7 @@ void FDeferredShadingSceneRenderer::RenderRayTracingShadows(
 		if (bUseHairLighting)
 		{
 			const bool bUseHairVoxel = CVarRayTracingShadowsEnableHairVoxel.GetValueOnRenderThread() > 0;
-			PassParameters->bUseHairVoxel = (bUseHairVoxel && !IsHairRayTracingEnabled()) ? 1 : 0;
+			PassParameters->bUseHairVoxel = (HairResources->bUseHairVoxel && bUseHairVoxel && !IsHairRayTracingEnabled()) ? 1 : 0;
 			PassParameters->HairCategorizationTexture = HairResources->CategorizationTexture;
 			PassParameters->HairLightChannelMaskTexture = HairResources->LightChannelMaskTexture;
 			PassParameters->VirtualVoxel = HairResources->VoxelResources->UniformBuffer;
diff --git a/Engine/Source/Runtime/Renderer/Private/ShadowRendering.cpp b/Engine/Source/Runtime/Renderer/Private/ShadowRendering.cpp
--- a/Engine/Source/Runtime/Renderer/Private/ShadowRendering.cpp
+++ b/Engine/Source/Runtime/Renderer/Private/ShadowRendering.cpp
@@ -1767,7 +1767,7 @@ bool FDeferredShadingSceneRenderer::RenderShadowProjections(FRHICommandListImmed
 	// Inject deep shadow mask
 	if (HairDatas)
 	{
-		RenderHairStrandsShadowMask(RHICmdList, Views, LightSceneInfo, ScreenShadowMaskTexture, HairDatas);
+		RenderHairStrandsShadowMask(RHICmdList, Views, LightSceneInfo, HairDatas, ScreenShadowMaskTexture);
 	}
 
 	return true;
