Fix parallel rendering issue with vkCmdDrawIndexed outside rendering pass

diff --git a/Engine/Source/Runtime/Renderer/Private/TranslucentRendering.cpp b/Engine/Source/Runtime/Renderer/Private/TranslucentRendering.cpp
--- a/Engine/Source/Runtime/Renderer/Private/TranslucentRendering.cpp
+++ b/Engine/Source/Runtime/Renderer/Private/TranslucentRendering.cpp
@@ -401,6 +401,12 @@ void RenderViewTranslucencyInner(FRHICommandListImmediate& RHICmdList, const FVi
 		// editor and debug rendering
 		if (View.bHasTranslucentViewMeshElements)
 		{
+			if (ParallelCommandListSet)
+			{
+				FSceneRenderTargets& SceneContext = FSceneRenderTargets::Get(RHICmdList);
+				SceneContext.BeginRenderingPrePass(RHICmdList, false);
+			}
+
 			{
 				QUICK_SCOPE_CYCLE_COUNTER(RenderTranslucencyParallel_SDPG_World);
 
@@ -450,6 +456,11 @@ void RenderViewTranslucencyInner(FRHICommandListImmediate& RHICmdList, const FVi
 					}
 				});
 			}
+
+			if (ParallelCommandListSet)
+			{
+				RHICmdList.EndRenderPass();
+			}
 		}
 
 		const FSceneViewState* ViewState = (const FSceneViewState*)View.State;
