Update translucent rendering fix (CL 12868321) to use BeginRenderingTranslucency

diff --git a/Engine/Source/Runtime/Renderer/Private/TranslucentRendering.cpp b/Engine/Source/Runtime/Renderer/Private/TranslucentRendering.cpp
--- a/Engine/Source/Runtime/Renderer/Private/TranslucentRendering.cpp
+++ b/Engine/Source/Runtime/Renderer/Private/TranslucentRendering.cpp
@@ -404,7 +404,7 @@ void RenderViewTranslucencyInner(FRHICommandListImmediate& RHICmdList, const FVi
 			if (ParallelCommandListSet)
 			{
 				FSceneRenderTargets& SceneContext = FSceneRenderTargets::Get(RHICmdList);
-				SceneContext.BeginRenderingPrePass(RHICmdList, false);
+				SceneContext.BeginRenderingTranslucency(RHICmdList, View, *ParallelCommandListSet->SceneRenderer, false);
 			}
 
 			{
