Converted the LODColoration viewmode to use the debug view mode framework.
This fixes tessellation, vertex displacement and to some extent translucency rendering.

diff --git a/Engine/Shaders/Private/ViewMode/LODColorationPixelShader.usf b/Engine/Shaders/Private/ViewMode/LODColorationPixelShader.usf
new file mode 100644
--- /dev/null
+++ b/Engine/Shaders/Private/ViewMode/LODColorationPixelShader.usf
@@ -0,0 +1,79 @@
+// Copyright Epic Games, Inc. All Rights Reserved.
+
+/*=============================================================================
+	LODColorationPixelSahder.usf: LOD Color.
+=============================================================================*/
+
+#define TEX_COORD_SCALE_ANALYSIS 1
+#define DEBUG_MATERIAL_PARAMETERS 1
+
+#undef SCENE_TEXTURES_DISABLED
+#define SCENE_TEXTURES_DISABLED 1
+
+#undef EYE_ADAPTATION_DISABLED
+#define EYE_ADAPTATION_DISABLED 1
+
+#include "../Common.ush"
+#include "../SHCommon.ush"
+
+struct FTexCoordScalesParams { float Dummy; };
+MaterialFloat StoreTexCoordScale(in out FTexCoordScalesParams Params, float2 UV, int TextureReferenceIndex) { return 1.f; }
+MaterialFloat StoreTexSample(in out FTexCoordScalesParams Params, float4 C, int TextureReferenceIndex) { return 1.f; }
+
+#include "/Engine/Generated/Material.ush"
+#include "../DebugViewModeCommon.ush"
+
+int LODIndex;
+
+#if LIGHTMAP_VT_ENABLED || NUM_VIRTUALTEXTURE_SAMPLES
+EARLYDEPTHSTENCIL
+#endif
+void Main(
+	in FDebugPSIn DebugInputs
+	OPTIONAL_IsFrontFace,
+	out float4 OutColor : SV_Target0
+	)
+{
+#if INSTANCED_STEREO
+	ResolvedView = ResolveView(DebugInputs.EyeIndex);
+#else
+	ResolvedView = ResolveView();
+#endif
+
+	float3 Result = DebugViewModePass.LODColors[LODIndex].rgb;
+
+	FMaterialPixelParameters MaterialParameters = GetMaterialPixelParameters(DebugInputs, DebugInputs.SvPosition);
+	MaterialParameters.TexCoordScalesParams.Dummy = 0;
+
+	half3 BaseColor;
+	half Opacity;
+	{
+		FPixelMaterialInputs PixelMaterialInputs;
+		CalcMaterialParameters(MaterialParameters, PixelMaterialInputs, DebugInputs.SvPosition, bIsFrontFace);
+
+		// Sample material properties. The results are not used, but the calls to StoreTexCoordScale are still be made.
+		BaseColor = GetMaterialBaseColorRaw(PixelMaterialInputs);
+		Opacity = GetMaterialOpacityRaw(PixelMaterialInputs);
+
+		half3 Emissive = GetMaterialEmissiveRaw(PixelMaterialInputs);
+#if MATERIALBLENDING_ADDITIVE
+		Opacity = (.05 + .95 * Luminance(Emissive));
+#else
+		BaseColor += Emissive;
+#endif
+	}
+
+	Result *= (.05 + .95 * Luminance(BaseColor));
+
+	OutColor = RETURN_COLOR(float4(Result, Opacity));
+	
+#if NUM_VIRTUALTEXTURE_SAMPLES || LIGHTMAP_VT_ENABLED
+	FinalizeVirtualTextureFeedback(
+		MaterialParameters.VirtualTextureFeedback,
+		MaterialParameters.SvPosition,
+		1.0,
+		View.FrameNumber,
+		View.VTFeedbackBuffer
+	);
+#endif
+}
diff --git a/Engine/Source/Editor/UnrealEd/Private/EditorBuildUtils.cpp b/Engine/Source/Editor/UnrealEd/Private/EditorBuildUtils.cpp
--- a/Engine/Source/Editor/UnrealEd/Private/EditorBuildUtils.cpp
+++ b/Engine/Source/Editor/UnrealEd/Private/EditorBuildUtils.cpp
@@ -1104,6 +1104,8 @@ EDebugViewShaderMode ViewModeIndexToDebugViewShaderMode(EViewModeIndex SelectedV
 		return DVSM_RequiredTextureResolution;
 	case VMI_RayTracingDebug:
 		return DVSM_RayTracingDebug;
+	case VMI_LODColoration:
+		return DVSM_LODColoration;
 	case VMI_Unknown:
 	default :
 		return DVSM_None;
diff --git a/Engine/Source/Runtime/Engine/Private/DebugViewModeHelpers.cpp b/Engine/Source/Runtime/Engine/Private/DebugViewModeHelpers.cpp
--- a/Engine/Source/Runtime/Engine/Private/DebugViewModeHelpers.cpp
+++ b/Engine/Source/Runtime/Engine/Private/DebugViewModeHelpers.cpp
@@ -40,6 +40,7 @@ bool AllowDebugViewShaderMode(EDebugViewShaderMode ShaderMode, EShaderPlatform P
 	case DVSM_None:
 		return false;
 	case DVSM_ShaderComplexity:
+	case DVSM_LODColoration:
 		return IsPCPlatform(Platform);
 	case DVSM_ShaderComplexityContainedQuadOverhead:
 	case DVSM_ShaderComplexityBleedingQuadOverhead:
diff --git a/Engine/Source/Runtime/Engine/Private/PrimitiveDrawingUtils.cpp b/Engine/Source/Runtime/Engine/Private/PrimitiveDrawingUtils.cpp
--- a/Engine/Source/Runtime/Engine/Private/PrimitiveDrawingUtils.cpp
+++ b/Engine/Source/Runtime/Engine/Private/PrimitiveDrawingUtils.cpp
@@ -1445,24 +1445,6 @@ void ApplyViewModeOverrides(
 			Collector.RegisterOneFrameMaterialProxy(WireframeMaterialInstance);
 		}
 	}
-	else if (EngineShowFlags.LODColoration)
-	{
-		if (!Mesh.IsTranslucent(FeatureLevel) && GEngine->LODColorationColors.Num()  > 0)
-		{
-			int32 lodColorationIndex = FMath::Clamp((int32)Mesh.VisualizeLODIndex, 0, GEngine->LODColorationColors.Num() - 1);
-	
-			bool bLit = Mesh.MaterialRenderProxy->GetMaterial(FeatureLevel)->GetShadingModels().IsLit();
-			const UMaterial* LODColorationMaterial = (bLit && EngineShowFlags.Lighting) ? GEngine->LevelColorationLitMaterial : GEngine->LevelColorationUnlitMaterial;
-
-			auto LODColorationMaterialInstance = new FColoredMaterialRenderProxy(
-				LODColorationMaterial->GetRenderProxy(),
-				GetSelectionColor(GEngine->LODColorationColors[lodColorationIndex], bSelected, PrimitiveSceneProxy->IsHovered() )
-				);
-
-			Mesh.MaterialRenderProxy = LODColorationMaterialInstance;
-			Collector.RegisterOneFrameMaterialProxy(LODColorationMaterialInstance);
-		}
-	}
 	else if (EngineShowFlags.HLODColoration)
 	{
 		if (!Mesh.IsTranslucent(FeatureLevel) && GEngine->HLODColorationColors.Num() > 0)
diff --git a/Engine/Source/Runtime/Engine/Private/SceneView.cpp b/Engine/Source/Runtime/Engine/Private/SceneView.cpp
--- a/Engine/Source/Runtime/Engine/Private/SceneView.cpp
+++ b/Engine/Source/Runtime/Engine/Private/SceneView.cpp
@@ -2608,7 +2608,10 @@ EDebugViewShaderMode FSceneViewFamily::ChooseDebugViewShaderMode() const
 	{
 		return DVSM_RayTracingDebug;
 	}
-
+	else if (EngineShowFlags.LODColoration)
+	{
+		return DVSM_LODColoration;
+	}
 	return DVSM_None;
 }
 
diff --git a/Engine/Source/Runtime/Engine/Private/ShowFlags.cpp b/Engine/Source/Runtime/Engine/Private/ShowFlags.cpp
--- a/Engine/Source/Runtime/Engine/Private/ShowFlags.cpp
+++ b/Engine/Source/Runtime/Engine/Private/ShowFlags.cpp
@@ -264,6 +264,7 @@ void ApplyViewMode(EViewModeIndex ViewModeIndex, bool bPerspective, FEngineShowF
 		case VMI_MeshUVDensityAccuracy:
 		case VMI_MaterialTextureScaleAccuracy:
 		case VMI_RequiredTextureResolution:
+		case VMI_LODColoration:
 			bPostProcessing = false;
 			break;
 		case VMI_StationaryLightOverlap:
@@ -285,7 +286,6 @@ void ApplyViewMode(EViewModeIndex ViewModeIndex, bool bPerspective, FEngineShowF
 		case VMI_CollisionVisibility:
 			bPostProcessing = false;
 			break;
-		case VMI_LODColoration:
 		case VMI_HLODColoration:
 			bPostProcessing = true;
 			break;
@@ -438,6 +438,7 @@ void EngineShowFlagOverride(EShowFlagInitMode ShowFlagInitMode, EViewModeIndex V
 			ViewModeIndex == VMI_MeshUVDensityAccuracy ||
 			ViewModeIndex == VMI_MaterialTextureScaleAccuracy ||
 			ViewModeIndex == VMI_RequiredTextureResolution ||
+			ViewModeIndex == VMI_LODColoration ||
 			ViewModeIndex == VMI_LightmapDensity)
 		{
 			EngineShowFlags.SetLighting(false);
@@ -466,7 +467,7 @@ void EngineShowFlagOverride(EShowFlagInitMode ShowFlagInitMode, EViewModeIndex V
 			EngineShowFlags.Atmosphere = 0;
 		}
 
-		if (ViewModeIndex == VMI_LODColoration || ViewModeIndex == VMI_HLODColoration)
+		if (ViewModeIndex == VMI_HLODColoration)
 		{
 			EngineShowFlags.SetLighting(true);	// Best currently otherwise the image becomes hard to read.
 			EngineShowFlags.Fog = 0;			// Removed fog to improve color readability.
@@ -483,6 +484,10 @@ void EngineShowFlagOverride(EShowFlagInitMode ShowFlagInitMode, EViewModeIndex V
 			EngineShowFlags.Particles = 0; // FX are fully streamed.
 			EngineShowFlags.Fog = 0;
 		}
+		if (ViewModeIndex == VMI_LODColoration)
+		{
+			EngineShowFlags.Decals = 0; // Decals require the use of FDebugPSInLean.
+		}
 
 		if (ViewModeIndex == VMI_PathTracing)
 		{
diff --git a/Engine/Source/Runtime/Engine/Public/DebugViewModeHelpers.h b/Engine/Source/Runtime/Engine/Public/DebugViewModeHelpers.h
--- a/Engine/Source/Runtime/Engine/Public/DebugViewModeHelpers.h
+++ b/Engine/Source/Runtime/Engine/Public/DebugViewModeHelpers.h
@@ -29,6 +29,7 @@ enum EDebugViewShaderMode
 	DVSM_OutputMaterialTextureScales, // Outputs the material texture scales.
 	DVSM_RequiredTextureResolution, // Visualize the accuracy of the material texture scales used for texture streaming.
 	DVSM_RayTracingDebug,			// Visualize ray tracing debug modes.
+	DVSM_LODColoration,				// Visualize primitive LOD .
 	DVSM_MAX
 };
 
diff --git a/Engine/Source/Runtime/Engine/Public/MaterialShared.h b/Engine/Source/Runtime/Engine/Public/MaterialShared.h
--- a/Engine/Source/Runtime/Engine/Public/MaterialShared.h
+++ b/Engine/Source/Runtime/Engine/Public/MaterialShared.h
@@ -2975,6 +2975,7 @@ struct FMaterialShaderParameters
 			uint64 bMaterialIsPrimitiveDistanceAccuracy : 1;
 			uint64 bMaterialIsRequiredTextureResolution : 1;
 			uint64 bMaterialIsComplexityAccumulate : 1;
+			uint64 bMaterialIsLODColoration : 1;
 			uint64 bIsUsedWithLidarPointCloud : 1;
 			uint64 bIsStencilTestEnabled : 1;
 		};
@@ -3043,6 +3044,7 @@ struct FMaterialShaderParameters
 		bMaterialIsPrimitiveDistanceAccuracy = FriendlyName.Contains(TEXT("PrimitiveDistanceAccuracy"));
 		bMaterialIsRequiredTextureResolution = FriendlyName.Contains(TEXT("RequiredTextureResolution"));
 		bMaterialIsComplexityAccumulate = FriendlyName.Contains(TEXT("ComplexityAccumulate"));
+		bMaterialIsLODColoration = FriendlyName.Contains(TEXT("LODColoration"));
 	}
 };
 
diff --git a/Engine/Source/Runtime/Renderer/Private/DebugViewModeRendering.cpp b/Engine/Source/Runtime/Renderer/Private/DebugViewModeRendering.cpp
--- a/Engine/Source/Runtime/Renderer/Private/DebugViewModeRendering.cpp
+++ b/Engine/Source/Runtime/Renderer/Private/DebugViewModeRendering.cpp
@@ -11,6 +11,7 @@ DebugViewModeRendering.cpp: Contains definitions for rendering debug viewmodes.
 #include "MeshTexCoordSizeAccuracyRendering.h"
 #include "MaterialTexCoordScalesRendering.h"
 #include "RequiredTextureResolutionRendering.h"
+#include "ViewMode/LODColorationRendering.h"
 #include "PrimitiveSceneInfo.h"
 #include "ScenePrivate.h"
 #include "PostProcessing.h"
@@ -34,15 +35,31 @@ void SetupDebugViewModePassUniformBuffer(FSceneRenderTargets& SceneContext, ERHI
 {
 	SetupSceneTextureUniformParameters(SceneContext, FeatureLevel, ESceneTextureSetupMode::None, PassParameters.SceneTextures);
 
-	const int32 NumEngineColors = FMath::Min<int32>(GEngine->StreamingAccuracyColors.Num(), NumStreamingAccuracyColors);
-	int32 ColorIndex = 0;
-	for (; ColorIndex < NumEngineColors; ++ColorIndex)
+	// Accuracy colors
 	{
-		PassParameters.AccuracyColors[ColorIndex] = GEngine->StreamingAccuracyColors[ColorIndex];
+		const int32 NumEngineColors = FMath::Min<int32>(GEngine->StreamingAccuracyColors.Num(), NumStreamingAccuracyColors);
+		int32 ColorIndex = 0;
+		for (; ColorIndex < NumEngineColors; ++ColorIndex)
+		{
+			PassParameters.AccuracyColors[ColorIndex] = GEngine->StreamingAccuracyColors[ColorIndex];
+		}
+		for (; ColorIndex < NumStreamingAccuracyColors; ++ColorIndex)
+		{
+			PassParameters.AccuracyColors[ColorIndex] = FLinearColor::Black;
+		}
 	}
-	for (; ColorIndex < NumStreamingAccuracyColors; ++ColorIndex)
+	// LOD colors
 	{
-		PassParameters.AccuracyColors[ColorIndex] = FLinearColor::Black;
+		const int32 NumEngineColors = FMath::Min<int32>(GEngine->LODColorationColors.Num(), NumLODColorationColors);
+		int32 ColorIndex = 0;
+		for (; ColorIndex < NumEngineColors; ++ColorIndex)
+		{
+			PassParameters.LODColors[ColorIndex] = GEngine->LODColorationColors[ColorIndex];
+		}
+		for (; ColorIndex < NumLODColorationColors; ++ColorIndex)
+		{
+			PassParameters.LODColors[ColorIndex] = NumEngineColors > 0 ? GEngine->LODColorationColors.Last() : FLinearColor::Black;
+		}
 	}
 }
 
@@ -353,6 +370,8 @@ void InitDebugViewModeInterfaces()
 	FDebugViewModeInterface::SetInterface(DVSM_MaterialTextureScaleAccuracy, new FMaterialTexCoordScaleAccuracyInterface());
 	FDebugViewModeInterface::SetInterface(DVSM_OutputMaterialTextureScales, new FOutputMaterialTexCoordScaleInterface());
 	FDebugViewModeInterface::SetInterface(DVSM_RequiredTextureResolution, new FRequiredTextureResolutionInterface());
+
+	FDebugViewModeInterface::SetInterface(DVSM_LODColoration, new FLODColorationInterface());
 }
 
 #else // !(UE_BUILD_SHIPPING || UE_BUILD_TEST)
diff --git a/Engine/Source/Runtime/Renderer/Private/DebugViewModeRendering.h b/Engine/Source/Runtime/Renderer/Private/DebugViewModeRendering.h
--- a/Engine/Source/Runtime/Renderer/Private/DebugViewModeRendering.h
+++ b/Engine/Source/Runtime/Renderer/Private/DebugViewModeRendering.h
@@ -22,11 +22,13 @@ struct FMeshDrawingRenderState;
 class FDebugViewModeInterface;
 
 static const int32 NumStreamingAccuracyColors = 5;
+static const int32 NumLODColorationColors = 8;
 static const float UndefinedStreamingAccuracyIntensity = .015f;
 
 BEGIN_GLOBAL_SHADER_PARAMETER_STRUCT(FDebugViewModePassPassUniformParameters, )
 	SHADER_PARAMETER_STRUCT(FSceneTexturesUniformParameters, SceneTextures)
 	SHADER_PARAMETER_ARRAY(FLinearColor, AccuracyColors, [NumStreamingAccuracyColors])
+	SHADER_PARAMETER_ARRAY(FLinearColor, LODColors, [NumLODColorationColors])
 END_GLOBAL_SHADER_PARAMETER_STRUCT()
 
 #if !(UE_BUILD_SHIPPING || UE_BUILD_TEST)
diff --git a/Engine/Source/Runtime/Renderer/Private/PostProcess/PostProcessing.cpp b/Engine/Source/Runtime/Renderer/Private/PostProcess/PostProcessing.cpp
--- a/Engine/Source/Runtime/Renderer/Private/PostProcess/PostProcessing.cpp
+++ b/Engine/Source/Runtime/Renderer/Private/PostProcess/PostProcessing.cpp
@@ -1019,6 +1019,8 @@ void AddDebugPostProcessingPasses(FRDGBuilder& GraphBuilder, const FViewInfo& Vi
 
 			break;
 		}
+		case DVSM_LODColoration:
+			break;
 		default:
 			ensure(false);
 			break;
diff --git a/Engine/Source/Runtime/Renderer/Private/ViewMode/LODColorationRendering.cpp b/Engine/Source/Runtime/Renderer/Private/ViewMode/LODColorationRendering.cpp
new file mode 100644
--- /dev/null
+++ b/Engine/Source/Runtime/Renderer/Private/ViewMode/LODColorationRendering.cpp
@@ -0,0 +1,39 @@
+// Copyright Epic Games, Inc. All Rights Reserved.
+
+/*=============================================================================
+LODColorationRendering.cpp: Contains definitions for rendering the LOD coloration viewmode.
+=============================================================================*/
+
+#include "LODColorationRendering.h"
+#include "PrimitiveSceneProxy.h"
+#include "EngineGlobals.h"
+#include "MeshBatch.h"
+#include "Engine/Engine.h"
+
+#if !(UE_BUILD_SHIPPING || UE_BUILD_TEST)
+
+IMPLEMENT_MATERIAL_SHADER_TYPE(,FLODColorationPS,TEXT("/Engine/Private/ViewMode/LODColorationPixelShader.usf"),TEXT("Main"),SF_Pixel);
+
+void FLODColorationInterface::GetDebugViewModeShaderBindings(
+	const FDebugViewModePS& ShaderBase,
+	const FPrimitiveSceneProxy* RESTRICT PrimitiveSceneProxy,
+	const FMaterialRenderProxy& RESTRICT MaterialRenderProxy,
+	const FMaterial& RESTRICT Material,
+	EDebugViewShaderMode DebugViewMode,
+	const FVector& ViewOrigin,
+	int32 VisualizeLODIndex,
+	int32 VisualizeElementIndex,
+	int32 NumVSInstructions,
+	int32 NumPSInstructions,
+	int32 ViewModeParam,
+	FName ViewModeParamName,
+	FMeshDrawSingleShaderBindings& ShaderBindings
+) const
+{
+	const FLODColorationPS& Shader = static_cast<const FLODColorationPS&>(ShaderBase);
+	const int32 LODIndex = FMath::Clamp(VisualizeLODIndex, 0, NumLODColorationColors - 1);
+
+	ShaderBindings.Add(Shader.LODIndexParameter, LODIndex);
+}
+
+#endif // !(UE_BUILD_SHIPPING || UE_BUILD_TEST)
diff --git a/Engine/Source/Runtime/Renderer/Private/ViewMode/LODColorationRendering.h b/Engine/Source/Runtime/Renderer/Private/ViewMode/LODColorationRendering.h
new file mode 100644
--- /dev/null
+++ b/Engine/Source/Runtime/Renderer/Private/ViewMode/LODColorationRendering.h
@@ -0,0 +1,76 @@
+// Copyright Epic Games, Inc. All Rights Reserved.
+
+/*=============================================================================
+LODColorationRendering.h: Declarations used for the viewmode.
+=============================================================================*/
+
+#pragma once
+
+#include "CoreMinimal.h"
+#include "ShaderParameters.h"
+#include "Shader.h"
+#include "GlobalShader.h"
+#include "DebugViewModeRendering.h"
+#include "DebugViewModeInterface.h"
+
+class FPrimitiveSceneProxy;
+struct FMeshBatchElement;
+struct FMeshDrawingRenderState;
+
+#if !(UE_BUILD_SHIPPING || UE_BUILD_TEST)
+
+/**
+* Pixel shader that renders texture streamer wanted mips accuracy.
+*/
+class FLODColorationPS : public FDebugViewModePS
+{
+	DECLARE_SHADER_TYPE(FLODColorationPS,MeshMaterial);
+
+public:
+
+	static bool ShouldCompilePermutation(const FMeshMaterialShaderPermutationParameters& Parameters)
+	{
+		// See FDebugViewModeMaterialProxy::GetFriendlyName()
+		return AllowDebugViewShaderMode(DVSM_LODColoration, Parameters.Platform, Parameters.MaterialParameters.FeatureLevel) && Parameters.MaterialParameters.bMaterialIsLODColoration;
+	}
+
+	FLODColorationPS(const ShaderMetaType::CompiledShaderInitializerType& Initializer):
+		FDebugViewModePS(Initializer)
+	{
+		LODIndexParameter.Bind(Initializer.ParameterMap, TEXT("LODIndex"));
+	}
+
+	FLODColorationPS() {}
+
+	static void ModifyCompilationEnvironment(const FMaterialShaderPermutationParameters& Parameters, FShaderCompilerEnvironment& OutEnvironment)
+	{
+	}
+
+	LAYOUT_FIELD(FShaderParameter, LODIndexParameter);
+};
+
+class FLODColorationInterface : public FDebugViewModeInterface
+{
+public:
+
+	FLODColorationInterface() : FDebugViewModeInterface(TEXT("LODColoration"), false, true, false) {}
+	virtual TShaderRef<FDebugViewModePS> GetPixelShader(const FMaterial* InMaterial, FVertexFactoryType* VertexFactoryType) const override { return InMaterial->GetShader<FLODColorationPS>(VertexFactoryType); }
+
+	virtual void GetDebugViewModeShaderBindings(
+		const FDebugViewModePS& ShaderBase,
+		const FPrimitiveSceneProxy* RESTRICT PrimitiveSceneProxy,
+		const FMaterialRenderProxy& RESTRICT MaterialRenderProxy,
+		const FMaterial& RESTRICT Material,
+		EDebugViewShaderMode DebugViewMode,
+		const FVector& ViewOrigin,
+		int32 VisualizeLODIndex,
+		int32 VisualizeElementIndex,
+		int32 NumVSInstructions,
+		int32 NumPSInstructions,
+		int32 ViewModeParam,
+		FName ViewModeParamName,
+		FMeshDrawSingleShaderBindings& ShaderBindings
+	) const override;
+};
+
+#endif // !(UE_BUILD_SHIPPING || UE_BUILD_TEST)
