Refactor VT feedback system to provide public interface for submitting buffers from external libs.

diff --git a/Engine/Source/Runtime/Renderer/Private/DeferredShadingRenderer.cpp b/Engine/Source/Runtime/Renderer/Private/DeferredShadingRenderer.cpp
--- a/Engine/Source/Runtime/Renderer/Private/DeferredShadingRenderer.cpp
+++ b/Engine/Source/Runtime/Renderer/Private/DeferredShadingRenderer.cpp
@@ -2762,10 +2762,10 @@ void FDeferredShadingSceneRenderer::Render(FRHICommandListImmediate& RHICmdList)
 			ViewRects[ViewIndex] = Views[ViewIndex].ViewRect;
 		}
 
-		FVirtualTextureFeedback::FBufferDesc Desc;
+		FVirtualTextureFeedbackBufferDesc Desc;
 		Desc.Init2D(SceneContext.GetBufferSizeXY(), ViewRects, SceneContext.GetVirtualTextureFeedbackScale());
 
-		GVirtualTextureFeedback.TransferGPUToCPU(RHICmdList, SceneContext.VirtualTextureFeedback, Desc);
+		SubmitVirtualTextureFeedbackBuffer(RHICmdList, SceneContext.VirtualTextureFeedback, Desc);
 	}
 
 #if RHI_RAYTRACING
diff --git a/Engine/Source/Runtime/Renderer/Private/MobileShadingRenderer.cpp b/Engine/Source/Runtime/Renderer/Private/MobileShadingRenderer.cpp
--- a/Engine/Source/Runtime/Renderer/Private/MobileShadingRenderer.cpp
+++ b/Engine/Source/Runtime/Renderer/Private/MobileShadingRenderer.cpp
@@ -790,10 +790,10 @@ void FMobileSceneRenderer::Render(FRHICommandListImmediate& RHICmdList)
 			ViewRects[ViewIndex] = Views[ViewIndex].ViewRect;
 		}
 		
-		FVirtualTextureFeedback::FBufferDesc Desc;
+		FVirtualTextureFeedbackBufferDesc Desc;
 		Desc.Init2D(SceneContext.GetBufferSizeXY(), ViewRects, SceneContext.GetVirtualTextureFeedbackScale());
 
-		GVirtualTextureFeedback.TransferGPUToCPU(RHICmdList, SceneContext.VirtualTextureFeedback, Desc);
+		SubmitVirtualTextureFeedbackBuffer(RHICmdList, SceneContext.VirtualTextureFeedback, Desc);
 	}
 	
 	if (ViewFamily.bResolveScene)
diff --git a/Engine/Source/Runtime/Renderer/Private/VT/VirtualTextureFeedback.cpp b/Engine/Source/Runtime/Renderer/Private/VT/VirtualTextureFeedback.cpp
--- a/Engine/Source/Runtime/Renderer/Private/VT/VirtualTextureFeedback.cpp
+++ b/Engine/Source/Runtime/Renderer/Private/VT/VirtualTextureFeedback.cpp
@@ -1,6 +1,6 @@
 // Copyright Epic Games, Inc. All Rights Reserved.
 
-#include "VirtualTextureFeedback.h"
+#include "VT/VirtualTextureFeedback.h"
 
 #if PLATFORM_WINDOWS
 // Use Query objects until RHI has a good fence on D3D11
@@ -150,42 +150,6 @@ public:
 
 #endif // USE_RHI_FENCES
 
-void FVirtualTextureFeedback::FBufferDesc::Init(int32 InBufferSize)
-{
-	BufferSize = FIntPoint(InBufferSize, 1);
-	NumRects = 0;
-	TotalReadSize = InBufferSize;
-}
-
-void FVirtualTextureFeedback::FBufferDesc::Init2D(FIntPoint InBufferSize)
-{
-	BufferSize = InBufferSize;
-	NumRects = 0;
-	TotalReadSize = BufferSize.X * BufferSize.Y;
-}
-
-void FVirtualTextureFeedback::FBufferDesc::Init2D(FIntPoint InUnscaledBufferSize, TArrayView<FIntRect> const& InUnscaledViewRects, int32 InBufferScale)
-{
-	const int32 BufferScale = FMath::Max(InBufferScale, 1);
-
-	BufferSize = FIntPoint::DivideAndRoundUp(InUnscaledBufferSize, BufferScale);
-	NumRects = 0;
-	TotalReadSize = BufferSize.X * BufferSize.Y;
-
-	if (InUnscaledViewRects.Num() > 0 && InUnscaledViewRects[0].Size() != InUnscaledBufferSize)
-	{
-		NumRects = FMath::Min((int32)MaxRectPerTransfer, InUnscaledViewRects.Num());
-		TotalReadSize = 0;
-
-		for (int32 RectIndex = 0; RectIndex < NumRects; ++RectIndex)
-		{
-			FIntRect const& Rect = InUnscaledViewRects[RectIndex];
-			Rects[RectIndex].Min = FIntPoint::DivideAndRoundDown(Rect.Min, BufferScale);
-			Rects[RectIndex].Max = FIntPoint::DivideAndRoundUp(Rect.Max, BufferScale);
-			TotalReadSize += Rects[RectIndex].Area();
-		}
-	}
-}
 
 FVirtualTextureFeedback::FVirtualTextureFeedback()
 	: NumPending(0)
@@ -220,7 +184,7 @@ void FVirtualTextureFeedback::ReleaseRHI()
 	Fences->ReleaseRHI();
 }
 
-void FVirtualTextureFeedback::TransferGPUToCPU(FRHICommandListImmediate& RHICmdList, FVertexBufferRHIRef const& Buffer, FBufferDesc const& Desc)
+void FVirtualTextureFeedback::TransferGPUToCPU(FRHICommandListImmediate& RHICmdList, FVertexBufferRHIRef const& Buffer, FVirtualTextureFeedbackBufferDesc const& Desc)
 {
 	QUICK_SCOPE_CYCLE_COUNTER(STAT_VirtualTextureFeedback_TransferGPUToCPU);
 
@@ -268,7 +232,7 @@ FVirtualTextureFeedback::FMapResult FVirtualTextureFeedback::Map(FRHICommandList
 	for (int32 ResultIndex = 0; ResultIndex < MaxTransfersToMap && ResultIndex < NumPending; ++ResultIndex)
 	{
 		const int32 FeedbackIndex = (ReadIndex + ResultIndex) % MaxTransfers;
-		FBufferDesc const& FeedbackItemDesc = FeedbackItems[FeedbackIndex].Desc;
+		FVirtualTextureFeedbackBufferDesc const& FeedbackItemDesc = FeedbackItems[FeedbackIndex].Desc;
 
 		if (!Fences->Poll(RHICmdList, FeedbackIndex))
 		{
@@ -290,7 +254,7 @@ FVirtualTextureFeedback::FMapResult FVirtualTextureFeedback::Map(FRHICommandList
 		{
 			// If only one target with no rectangles then fast path is to return the locked buffer.
 			const int32 FeedbackIndex = ReadIndex;
-			FBufferDesc const& FeedbackItemDesc = FeedbackItems[FeedbackIndex].Desc;
+			FVirtualTextureFeedbackBufferDesc const& FeedbackItemDesc = FeedbackItems[FeedbackIndex].Desc;
 			FRHIGPUMask GPUMask = FeedbackItems[FeedbackIndex].GPUMask;
 			FStagingBufferRHIRef StagingBuffer = FeedbackItems[FeedbackIndex].StagingBuffer;
 			
@@ -312,7 +276,7 @@ FVirtualTextureFeedback::FMapResult FVirtualTextureFeedback::Map(FRHICommandList
 			for (int32 ResultIndex = 0; ResultIndex < NumResults; ++ResultIndex)
 			{
 				const int32 FeedbackIndex = (ReadIndex + ResultIndex) % MaxTransfers;
-				FBufferDesc const& FeedbackItemDesc = FeedbackItems[FeedbackIndex].Desc;
+				FVirtualTextureFeedbackBufferDesc const& FeedbackItemDesc = FeedbackItems[FeedbackIndex].Desc;
 				FRHIGPUMask GPUMask = FeedbackItems[FeedbackIndex].GPUMask;
 				FStagingBufferRHIRef StagingBuffer = FeedbackItems[FeedbackIndex].StagingBuffer;
 
diff --git a/Engine/Source/Runtime/Renderer/Private/VT/VirtualTextureFeedback.h b/Engine/Source/Runtime/Renderer/Private/VT/VirtualTextureFeedback.h
--- a/Engine/Source/Runtime/Renderer/Private/VT/VirtualTextureFeedback.h
+++ b/Engine/Source/Runtime/Renderer/Private/VT/VirtualTextureFeedback.h
@@ -4,11 +4,11 @@
 
 #include "CoreMinimal.h"
 #include "RenderResource.h"
+#include "VT/VirtualTextureFeedbackBuffer.h"
 
 /** 
  * Class that handles the read back of feedback buffers from the GPU.
  * Multiple buffers can be transferred per frame using TransferGPUToCPU().
- * RHIBuffers that are passed in are expected to have been transitioned to a state for reading.
  * Use Map()/Unmap() to return the data. It will only return data that is ready to fetch without stalling the GPU.
  * All calls are expected to be from the single render thread only.
  */
@@ -18,34 +18,8 @@ public:
 	FVirtualTextureFeedback();
 	~FVirtualTextureFeedback();
 
-	/** 
-	 * Description of how to interpret a Buffer that is being transferred.
-	 * For example a buffer may be a simple flat buffer, or a 2D screen-space buffer with rectangles representing multiple player viewports.
-	 * In the future we may also want to support append style buffers containing buffer size etc.
-	 */
-	struct FBufferDesc
-	{
-		/** Initialize for a flat "1D" buffer. */
-		void Init(int32 InBufferSize);
-		/** Initialize for a "2D" buffer. */
-		void Init2D(FIntPoint InBufferSize);
-		/** Initialize for a "2D" buffer. Pass in view rectangles that define some subset of the buffer to actually analyze. */
-		void Init2D(FIntPoint InUnscaledBufferSize, TArrayView<FIntRect> const& InUnscaledViewRects, int32 InBufferScale);
-
-		/** Size of buffer. 1D buffers have Y=1. */
-		FIntPoint BufferSize = 0;
-		/** The maximum number of rectangles to read from a 2D buffer. */
-		static const uint32 MaxRectPerTransfer = 4u;
-		/** The number of rectangles to read from the buffer. */
-		int32 NumRects = 0;
-		/** Rectangles to read from the buffer. */
-		FIntRect Rects[MaxRectPerTransfer];
-		/** Number of buffer elements to actually read (calculated from the Rects). */
-		int32 TotalReadSize = 0;
-	};
-
 	/** Commit a RHIBuffer to be transferred for later CPU analysis. */
-	void TransferGPUToCPU(FRHICommandListImmediate& RHICmdList, FVertexBufferRHIRef const& Buffer, FBufferDesc const& Desc);
+	void TransferGPUToCPU(FRHICommandListImmediate& RHICmdList, FVertexBufferRHIRef const& Buffer, FVirtualTextureFeedbackBufferDesc const& Desc);
 
 	/** Returns true if there are any pending transfer results that are ready so that we can call Map(). */
 	bool CanMap(FRHICommandListImmediate& RHICmdList);
@@ -76,7 +50,7 @@ private:
 	/** Description of a pending feedback transfer. */
 	struct FFeedbackItem
 	{
-		FBufferDesc Desc;
+		FVirtualTextureFeedbackBufferDesc Desc;
 		FRHIGPUMask GPUMask;
 		FStagingBufferRHIRef StagingBuffer;
 	};
diff --git a/Engine/Source/Runtime/Renderer/Private/VT/VirtualTextureFeedbackBuffer.cpp b/Engine/Source/Runtime/Renderer/Private/VT/VirtualTextureFeedbackBuffer.cpp
new file mode 100644
--- /dev/null
+++ b/Engine/Source/Runtime/Renderer/Private/VT/VirtualTextureFeedbackBuffer.cpp
@@ -0,0 +1,46 @@
+// Copyright Epic Games, Inc. All Rights Reserved.
+
+#include "VT/VirtualTextureFeedbackBuffer.h"
+#include "VT/VirtualTextureFeedback.h"
+
+void FVirtualTextureFeedbackBufferDesc::Init(int32 InBufferSize)
+{
+	BufferSize = FIntPoint(InBufferSize, 1);
+	NumRects = 0;
+	TotalReadSize = InBufferSize;
+}
+
+void FVirtualTextureFeedbackBufferDesc::Init2D(FIntPoint InBufferSize)
+{
+	BufferSize = InBufferSize;
+	NumRects = 0;
+	TotalReadSize = BufferSize.X * BufferSize.Y;
+}
+
+void FVirtualTextureFeedbackBufferDesc::Init2D(FIntPoint InUnscaledBufferSize, TArrayView<FIntRect> const& InUnscaledViewRects, int32 InBufferScale)
+{
+	const int32 BufferScale = FMath::Max(InBufferScale, 1);
+
+	BufferSize = FIntPoint::DivideAndRoundUp(InUnscaledBufferSize, BufferScale);
+	NumRects = 0;
+	TotalReadSize = BufferSize.X * BufferSize.Y;
+
+	if (InUnscaledViewRects.Num() > 0 && InUnscaledViewRects[0].Size() != InUnscaledBufferSize)
+	{
+		NumRects = FMath::Min((int32)MaxRectPerTransfer, InUnscaledViewRects.Num());
+		TotalReadSize = 0;
+
+		for (int32 RectIndex = 0; RectIndex < NumRects; ++RectIndex)
+		{
+			FIntRect const& Rect = InUnscaledViewRects[RectIndex];
+			Rects[RectIndex].Min = FIntPoint::DivideAndRoundDown(Rect.Min, BufferScale);
+			Rects[RectIndex].Max = FIntPoint::DivideAndRoundUp(Rect.Max, BufferScale);
+			TotalReadSize += Rects[RectIndex].Area();
+		}
+	}
+}
+
+void SubmitVirtualTextureFeedbackBuffer(FRHICommandListImmediate& RHICmdList, FVertexBufferRHIRef const& Buffer, FVirtualTextureFeedbackBufferDesc const& Desc)
+{
+	GVirtualTextureFeedback.TransferGPUToCPU(RHICmdList, Buffer, Desc);
+}
diff --git a/Engine/Source/Runtime/Renderer/Public/VT/VirtualTextureFeedbackBuffer.h b/Engine/Source/Runtime/Renderer/Public/VT/VirtualTextureFeedbackBuffer.h
new file mode 100644
--- /dev/null
+++ b/Engine/Source/Runtime/Renderer/Public/VT/VirtualTextureFeedbackBuffer.h
@@ -0,0 +1,40 @@
+// Copyright Epic Games, Inc. All Rights Reserved.
+
+#pragma once
+
+#include "CoreMinimal.h"
+
+/** 
+ * Description of how to interpret an RHIBuffer that is being fed to the virtual texture feedback system.
+ * For example a buffer may be a simple flat buffer, or a 2D screen-space buffer with rectangles representing multiple player viewports.
+ * In the future we may also want to support append style buffers containing buffer size etc.
+*/
+struct FVirtualTextureFeedbackBufferDesc
+{
+	/** Initialize for a flat "1D" buffer. */
+	RENDERER_API void Init(int32 InBufferSize);
+	/** Initialize for a "2D" buffer. */
+	RENDERER_API void Init2D(FIntPoint InBufferSize);
+	/** Initialize for a "2D" buffer. Pass in view rectangles that define some subset of the buffer to actually analyze. */
+	RENDERER_API void Init2D(FIntPoint InUnscaledBufferSize, TArrayView<FIntRect> const& InUnscaledViewRects, int32 InBufferScale);
+
+	/** Size of buffer. 1D buffers have Y=1. */
+	FIntPoint BufferSize = 0;
+	/** The maximum number of rectangles to read from a 2D buffer. */
+	static const uint32 MaxRectPerTransfer = 4u;
+	/** The number of rectangles to read from the buffer. */
+	int32 NumRects = 0;
+	/** Rectangles to read from the buffer. */
+	FIntRect Rects[MaxRectPerTransfer];
+	/** Number of buffer elements to actually read (calculated from the Rects). */
+	int32 TotalReadSize = 0;
+};
+
+/** 
+ * Submit an RHIBuffer containing virtual texture feedback data to the virtual texture system.
+ * The buffer is internally copied to the CPU and parsed to determine which virtual texture pages need to be mapped.
+ * RHIBuffers that are passed in are expected to have been transitioned to a state for reading.
+ * Multiple buffers can be transferred per frame using this function.
+ * The function can be called from the render thread only.
+*/
+RENDERER_API void SubmitVirtualTextureFeedbackBuffer(FRHICommandListImmediate& RHICmdList, FVertexBufferRHIRef const& Buffer, FVirtualTextureFeedbackBufferDesc const& Desc);
