CSI fixes
#rnx

diff --git a/Engine/Source/Editor/DetailCustomizations/Private/BoundsCopyComponentDetails.cpp b/Engine/Source/Editor/DetailCustomizations/Private/BoundsCopyComponentDetails.cpp
--- a/Engine/Source/Editor/DetailCustomizations/Private/BoundsCopyComponentDetails.cpp
+++ b/Engine/Source/Editor/DetailCustomizations/Private/BoundsCopyComponentDetails.cpp
@@ -1,4 +1,4 @@
-// Copyright 1998-2019 Epic Games, Inc. All Rights Reserved.
+// Copyright Epic Games, Inc. All Rights Reserved.
 
 #include "BoundsCopyComponentDetails.h"
 
diff --git a/Engine/Source/Editor/DetailCustomizations/Private/BoundsCopyComponentDetails.h b/Engine/Source/Editor/DetailCustomizations/Private/BoundsCopyComponentDetails.h
--- a/Engine/Source/Editor/DetailCustomizations/Private/BoundsCopyComponentDetails.h
+++ b/Engine/Source/Editor/DetailCustomizations/Private/BoundsCopyComponentDetails.h
@@ -1,4 +1,4 @@
-// Copyright 1998-2019 Epic Games, Inc. All Rights Reserved.
+// Copyright Epic Games, Inc. All Rights Reserved.
 
 #pragma once
 
diff --git a/Engine/Source/Runtime/Engine/Classes/Components/BoundsCopyComponent.h b/Engine/Source/Runtime/Engine/Classes/Components/BoundsCopyComponent.h
--- a/Engine/Source/Runtime/Engine/Classes/Components/BoundsCopyComponent.h
+++ b/Engine/Source/Runtime/Engine/Classes/Components/BoundsCopyComponent.h
@@ -1,4 +1,4 @@
-// Copyright 1998-2019 Epic Games, Inc. All Rights Reserved.
+// Copyright Epic Games, Inc. All Rights Reserved.
 
 #pragma once
 
diff --git a/Engine/Source/Runtime/Engine/Private/Components/BoundsCopyComponent.cpp b/Engine/Source/Runtime/Engine/Private/Components/BoundsCopyComponent.cpp
--- a/Engine/Source/Runtime/Engine/Private/Components/BoundsCopyComponent.cpp
+++ b/Engine/Source/Runtime/Engine/Private/Components/BoundsCopyComponent.cpp
@@ -1,4 +1,4 @@
-// Copyright 1998-2019 Epic Games, Inc. All Rights Reserved.
+// Copyright Epic Games, Inc. All Rights Reserved.
 
 #include "Components/BoundsCopyComponent.h"
 
diff --git a/Engine/Source/Runtime/Renderer/Public/VT/VirtualTextureFeedbackBuffer.h b/Engine/Source/Runtime/Renderer/Public/VT/VirtualTextureFeedbackBuffer.h
--- a/Engine/Source/Runtime/Renderer/Public/VT/VirtualTextureFeedbackBuffer.h
+++ b/Engine/Source/Runtime/Renderer/Public/VT/VirtualTextureFeedbackBuffer.h
@@ -3,6 +3,7 @@
 #pragma once
 
 #include "CoreMinimal.h"
+#include "RHIResources.h"
 
 /** 
  * Description of how to interpret an RHIBuffer that is being fed to the virtual texture feedback system.
@@ -37,4 +38,4 @@ struct FVirtualTextureFeedbackBufferDesc
  * Multiple buffers can be transferred per frame using this function.
  * The function can be called from the render thread only.
 */
-RENDERER_API void SubmitVirtualTextureFeedbackBuffer(FRHICommandListImmediate& RHICmdList, FVertexBufferRHIRef const& Buffer, FVirtualTextureFeedbackBufferDesc const& Desc);
+RENDERER_API void SubmitVirtualTextureFeedbackBuffer(class FRHICommandListImmediate& RHICmdList, FVertexBufferRHIRef const& Buffer, FVirtualTextureFeedbackBufferDesc const& Desc);
