Fix for incorrect image loads

diff --git a/Engine/Source/Runtime/Renderer/Private/AtmosphereRendering.cpp b/Engine/Source/Runtime/Renderer/Private/AtmosphereRendering.cpp
--- a/Engine/Source/Runtime/Renderer/Private/AtmosphereRendering.cpp
+++ b/Engine/Source/Runtime/Renderer/Private/AtmosphereRendering.cpp
@@ -471,7 +471,7 @@ void FDeferredShadingSceneRenderer::RenderAtmosphere(FRHICommandListImmediate& R
 	{
 		FSceneRenderTargets& SceneContext = FSceneRenderTargets::Get(RHICmdList);
 		
-		SceneContext.BeginRenderingSceneColor(RHICmdList, ESimpleRenderTargetMode::EUninitializedColorExistingDepth, FExclusiveDepthStencil::DepthRead_StencilWrite);
+        SceneContext.BeginRenderingSceneColor(RHICmdList, ESimpleRenderTargetMode::EExistingColorAndDepth, FExclusiveDepthStencil::DepthRead_StencilWrite);
 		
 		FGraphicsPipelineStateInitializer GraphicsPSOInit;
 		RHICmdList.ApplyCachedRenderTargets(GraphicsPSOInit);
diff --git a/Engine/Source/Runtime/Renderer/Private/DeferredShadingRenderer.cpp b/Engine/Source/Runtime/Renderer/Private/DeferredShadingRenderer.cpp
--- a/Engine/Source/Runtime/Renderer/Private/DeferredShadingRenderer.cpp
+++ b/Engine/Source/Runtime/Renderer/Private/DeferredShadingRenderer.cpp
@@ -2018,8 +2018,7 @@ void FDeferredShadingSceneRenderer::Render(FRHICommandListImmediate& RHICmdList)
 		// if we didn't to the prepass above, then we will need to clear now, otherwise, it's already been cleared and rendered to
 		ERenderTargetLoadAction ColorLoadAction = ERenderTargetLoadAction::ELoad;
 
-		bool bRequiresLoad = IsMetalPlatform(ShaderPlatform) || IsVulkanMobileSM5Platform(ShaderPlatform);
-		ERenderTargetLoadAction DepthLoadAction = bClearDepth ? ERenderTargetLoadAction::EClear : (!bRequiresLoad ? ERenderTargetLoadAction::ENoAction : ERenderTargetLoadAction::ELoad);
+		ERenderTargetLoadAction DepthLoadAction = bClearDepth ? ERenderTargetLoadAction::EClear : ERenderTargetLoadAction::ELoad;
 
 		const bool bClearBlack = ViewFamily.EngineShowFlags.ShaderComplexity || ViewFamily.EngineShowFlags.StationaryLightOverlap;
 		const float ClearAlpha = GetSceneColorClearAlpha();
@@ -2074,8 +2073,7 @@ void FDeferredShadingSceneRenderer::Render(FRHICommandListImmediate& RHICmdList)
 		// Make sure we have began the renderpass
 		ERenderTargetLoadAction DepthLoadAction = bDepthWasCleared ? ERenderTargetLoadAction::ELoad : ERenderTargetLoadAction::EClear;
 
-		bool bRequiresLoad = IsMetalPlatform(ShaderPlatform) || IsVulkanMobileSM5Platform(ShaderPlatform);
-		SceneContext.BeginRenderingGBuffer(RHICmdList, (!bRequiresLoad ? ERenderTargetLoadAction::ENoAction : ERenderTargetLoadAction::ELoad), DepthLoadAction, BasePassDepthStencilAccess, ViewFamily.EngineShowFlags.ShaderComplexity);
+		SceneContext.BeginRenderingGBuffer(RHICmdList, ERenderTargetLoadAction::ELoad, DepthLoadAction, BasePassDepthStencilAccess, ViewFamily.EngineShowFlags.ShaderComplexity);
 	}
 	// Wait for Async SSAO before rendering base pass with forward rendering
 	if (IsForwardShadingEnabled(ShaderPlatform))
@@ -2608,8 +2606,7 @@ void FDeferredShadingSceneRenderer::Render(FRHICommandListImmediate& RHICmdList)
 		SetupSceneTextureUniformParameters(SceneContext, FeatureLevel, ESceneTextureSetupMode::SceneDepth | ESceneTextureSetupMode::GBuffers, SceneTextureParameters);
 		TUniformBufferRef<FSceneTexturesUniformParameters> SceneTextureUniformBuffer = TUniformBufferRef<FSceneTexturesUniformParameters>::CreateUniformBufferImmediate(SceneTextureParameters, UniformBuffer_SingleFrame);
 
-		bool bRequiresLoad = IsMetalPlatform(ShaderPlatform) || IsVulkanMobileSM5Platform(ShaderPlatform);
-        SceneContext.BeginRenderingSceneColor(RHICmdList, (!bRequiresLoad ? ESimpleRenderTargetMode::EUninitializedColorExistingDepth : ESimpleRenderTargetMode::EExistingColorAndDepth));
+        SceneContext.BeginRenderingSceneColor(RHICmdList, ESimpleRenderTargetMode::EExistingColorAndDepth);
 		for (int32 ViewIndex = 0; ViewIndex < Views.Num(); ++ViewIndex)
 		{
 			const FViewInfo& View = Views[ViewIndex];
diff --git a/Engine/Source/Runtime/Renderer/Private/LightShaftRendering.cpp b/Engine/Source/Runtime/Renderer/Private/LightShaftRendering.cpp
--- a/Engine/Source/Runtime/Renderer/Private/LightShaftRendering.cpp
+++ b/Engine/Source/Runtime/Renderer/Private/LightShaftRendering.cpp
@@ -917,7 +917,7 @@ void ApplyLightShaftBloom(FRHICommandListImmediate& RHICmdList, const FViewInfo&
 	}
 	else
 	{
-		SceneContext.BeginRenderingSceneColor(RHICmdList, ESimpleRenderTargetMode::EUninitializedColorExistingDepth, FExclusiveDepthStencil::DepthRead_StencilWrite);
+		SceneContext.BeginRenderingSceneColor(RHICmdList, ESimpleRenderTargetMode::EExistingColorAndDepth, FExclusiveDepthStencil::DepthRead_StencilWrite);
 	}
 
 	FGraphicsPipelineStateInitializer GraphicsPSOInit;
