M - missing Pre-Declarations

diff --git a/Engine/Source/Runtime/Engine/Classes/Engine/World.h b/Engine/Source/Runtime/Engine/Classes/Engine/World.h
--- a/Engine/Source/Runtime/Engine/Classes/Engine/World.h
+++ b/Engine/Source/Runtime/Engine/Classes/Engine/World.h
@@ -62,7 +62,7 @@ class FSceneView;
 struct FUniqueNetIdRepl;
 struct FEncryptionKeyResponse;
 
-template<typename,typename> class TOctree;
+template<typename,typename> class TOctree2;
 
 /**
  * Misc. Iterator types
diff --git a/Engine/Source/Runtime/Renderer/Public/PrimitiveSceneInfo.h b/Engine/Source/Runtime/Renderer/Public/PrimitiveSceneInfo.h
--- a/Engine/Source/Runtime/Renderer/Public/PrimitiveSceneInfo.h
+++ b/Engine/Source/Runtime/Renderer/Public/PrimitiveSceneInfo.h
@@ -24,7 +24,7 @@ class FScene;
 class FViewInfo;
 class UPrimitiveComponent;
 class FIndirectLightingCacheUniformParameters;
-template<typename ElementType,typename OctreeSemantics> class TOctree;
+template<typename ElementType,typename OctreeSemantics> class TOctree2;
 
 /** Data used to track a primitive's allocation in the volume texture atlas that stores indirect lighting. */
 class FIndirectLightingCacheAllocation
