Batched ray tracing material gather pipeline using batched / parallelized API. Saves ~50% of the RHI thread critical path time.

diff --git a/Engine/Source/Runtime/Renderer/Private/RayTracing/RayTracingDeferredMaterials.cpp b/Engine/Source/Runtime/Renderer/Private/RayTracing/RayTracingDeferredMaterials.cpp
--- a/Engine/Source/Runtime/Renderer/Private/RayTracing/RayTracingDeferredMaterials.cpp
+++ b/Engine/Source/Runtime/Renderer/Private/RayTracing/RayTracingDeferredMaterials.cpp
@@ -52,26 +52,34 @@ FRayTracingPipelineState* FDeferredShadingSceneRenderer::BindRayTracingDeferredM
 
 	const FViewInfo& ReferenceView = Views[0];
 
+	const int32 NumTotalBindings = ReferenceView.VisibleRayTracingMeshCommands.Num();
+
+	FRayTracingLocalShaderBindings* Bindings = (FRayTracingLocalShaderBindings*)RHICmdList.Alloc(
+		sizeof(FRayTracingLocalShaderBindings) * NumTotalBindings,
+		alignof(FRayTracingLocalShaderBindings));
+
+	uint32 BindingIndex = 0;
+
 	for (const FVisibleRayTracingMeshCommand VisibleMeshCommand : ReferenceView.VisibleRayTracingMeshCommands)
 	{
 		const FRayTracingMeshCommand& MeshCommand = *VisibleMeshCommand.RayTracingMeshCommand;
 
-		const uint32 HitGroupIndex = 0; // Force the default CHS to be used on all geometry
-
-		const uint32 ShaderSlot = 0; // Multiple shader slots can be used for different ray types. Slot 0 is the primary material slot.
-		const uint32 MaterialIndexInUserData = MeshCommand.MaterialShaderIndex;
-		RHICmdList.SetRayTracingHitGroup(
-			View.RayTracingScene.RayTracingSceneRHI,
-			VisibleMeshCommand.InstanceIndex, 
-			MeshCommand.GeometrySegmentIndex, 
-			ShaderSlot,
-			PipelineState,
-			HitGroupIndex,
-			0, nullptr, // uniform buffers
-			0, nullptr, // loose data
-			MaterialIndexInUserData);
+		FRayTracingLocalShaderBindings Binding = {};
+		Binding.InstanceIndex = VisibleMeshCommand.InstanceIndex;
+		Binding.SegmentIndex = MeshCommand.GeometrySegmentIndex;
+		Binding.UserData = MeshCommand.MaterialShaderIndex;
+
+		Bindings[BindingIndex] = Binding;
+		BindingIndex++;
 	}
 
+	const bool bCopyDataToInlineStorage = false; // Storage is already allocated from RHICmdList, no extra copy necessary
+	RHICmdList.SetRayTracingHitGroups(
+		View.RayTracingScene.RayTracingSceneRHI,
+		PipelineState,
+		NumTotalBindings, Bindings,
+		bCopyDataToInlineStorage);
+
 	return PipelineState;
 }
 
