Fixed asserts due to using RHICmdList.Alloc() in bypass mode. FMemStack is used instead when necessary.

#[fyi] Rolando.Caloca

diff --git a/Engine/Source/Runtime/Renderer/Private/DeferredShadingRenderer.cpp b/Engine/Source/Runtime/Renderer/Private/DeferredShadingRenderer.cpp
--- a/Engine/Source/Runtime/Renderer/Private/DeferredShadingRenderer.cpp
+++ b/Engine/Source/Runtime/Renderer/Private/DeferredShadingRenderer.cpp
@@ -1301,9 +1301,10 @@ void FDeferredShadingSceneRenderer::WaitForRayTracingScene(FRHICommandListImmedi
 				}
 			}
 
-			FRayTracingLocalShaderBindings* MergedBindings = (FRayTracingLocalShaderBindings*)RHICmdList.Alloc(
-				sizeof(FRayTracingLocalShaderBindings) * NumTotalBindings,
-				alignof(FRayTracingLocalShaderBindings));
+			const uint32 MergedBindingsSize = sizeof(FRayTracingLocalShaderBindings) * NumTotalBindings;
+			FRayTracingLocalShaderBindings* MergedBindings = (FRayTracingLocalShaderBindings*) (RHICmdList.Bypass() 
+				? FMemStack::Get().Alloc(MergedBindingsSize, alignof(FRayTracingLocalShaderBindings))
+				: RHICmdList.Alloc(MergedBindingsSize, alignof(FRayTracingLocalShaderBindings)));
 
 			uint32 MergedBindingIndex = 0;
 			for (FRayTracingLocalShaderBindingWriter* BindingWriter : View.RayTracingMaterialBindings)
diff --git a/Engine/Source/Runtime/Renderer/Private/RayTracing/RayTracingDeferredMaterials.cpp b/Engine/Source/Runtime/Renderer/Private/RayTracing/RayTracingDeferredMaterials.cpp
--- a/Engine/Source/Runtime/Renderer/Private/RayTracing/RayTracingDeferredMaterials.cpp
+++ b/Engine/Source/Runtime/Renderer/Private/RayTracing/RayTracingDeferredMaterials.cpp
@@ -54,9 +54,10 @@ FRayTracingPipelineState* FDeferredShadingSceneRenderer::BindRayTracingDeferredM
 
 	const int32 NumTotalBindings = ReferenceView.VisibleRayTracingMeshCommands.Num();
 
-	FRayTracingLocalShaderBindings* Bindings = (FRayTracingLocalShaderBindings*)RHICmdList.Alloc(
-		sizeof(FRayTracingLocalShaderBindings) * NumTotalBindings,
-		alignof(FRayTracingLocalShaderBindings));
+	const uint32 MergedBindingsSize = sizeof(FRayTracingLocalShaderBindings) * NumTotalBindings;
+	FRayTracingLocalShaderBindings* Bindings = (FRayTracingLocalShaderBindings*)(RHICmdList.Bypass()
+		? FMemStack::Get().Alloc(MergedBindingsSize, alignof(FRayTracingLocalShaderBindings))
+		: RHICmdList.Alloc(MergedBindingsSize, alignof(FRayTracingLocalShaderBindings)));
 
 	uint32 BindingIndex = 0;
 
