Fixes temporal reprojection of denoiser when using dynamic resolution

diff --git a/Engine/Shaders/Private/ScreenSpaceDenoise/SSDTemporalAccumulation.usf b/Engine/Shaders/Private/ScreenSpaceDenoise/SSDTemporalAccumulation.usf
--- a/Engine/Shaders/Private/ScreenSpaceDenoise/SSDTemporalAccumulation.usf
+++ b/Engine/Shaders/Private/ScreenSpaceDenoise/SSDTemporalAccumulation.usf
@@ -318,7 +318,10 @@ int bCameraCut[MAX_SIGNAL_BATCH_SIZE];
 
 float HistoryPreExposureCorrection;
 
-float4 HistoryBufferUVMinMax[CONFIG_SIGNAL_BATCH_SIZE];
+float4 ScreenPosToHistoryBufferUV;
+float4 HistoryBufferSizeAndInvSize;
+float4 HistoryBufferUVMinMax;
+float4 HistoryBufferScissorUVMinMax[CONFIG_SIGNAL_BATCH_SIZE];
 
 
 #if !defined(CONFIG_INPUT_TEXTURE_COUNT)
@@ -523,8 +526,8 @@ void TemporallyAccumulate(
 	// Sample the history buffer.
 	FSSDSignalArray HistorySamples = CreateSignalArrayFromScalarValue(0.0);
 	{
-		float2 HistoryBufferUV = HistoryScreenPosition.xy * View.ScreenPositionScaleBias.xy + View.ScreenPositionScaleBias.wz;
-		float2 ClampedHistoryBufferUV = clamp(HistoryBufferUV, BufferBilinearUVMinMax.xy, BufferBilinearUVMinMax.zw);
+		float2 HistoryBufferUV = HistoryScreenPosition.xy * ScreenPosToHistoryBufferUV.xy + ScreenPosToHistoryBufferUV.wz;
+		float2 ClampedHistoryBufferUV = clamp(HistoryBufferUV, HistoryBufferUVMinMax.xy, HistoryBufferUVMinMax.zw);
 		bool bIsPreviousFrameOffscreen = any(HistoryBufferUV != ClampedHistoryBufferUV);
 
 		// TODO(Denoiser): global camera cut to save perf that is going to be required for panic post filtering.
@@ -557,8 +560,8 @@ void TemporallyAccumulate(
 			SetBilateralPreset(CONFIG_HISTORY_BILATERAL_PRESET, /* inout */ KernelConfig);
 
 			// SGPR configuration of the kernel.
-			KernelConfig.BufferSizeAndInvSize = BufferSizeAndInvSize;
-			KernelConfig.BufferBilinearUVMinMax = BufferBilinearUVMinMax;
+			KernelConfig.BufferSizeAndInvSize = HistoryBufferSizeAndInvSize;
+			KernelConfig.BufferBilinearUVMinMax = HistoryBufferUVMinMax;
 			
 			#if CONFIG_CLAMP_UV_PER_SIGNAL
 			{
@@ -566,7 +569,7 @@ void TemporallyAccumulate(
 				for (uint BatchedSignalId = 0; BatchedSignalId < CONFIG_SIGNAL_BATCH_SIZE; BatchedSignalId++)
 				{
 					uint MultiplexId = BatchedSignalId / CONFIG_MULTIPLEXED_SIGNALS_PER_SIGNAL_DOMAIN;
-					KernelConfig.PerSignalUVMinMax[MultiplexId] = HistoryBufferUVMinMax[MultiplexId];
+					KernelConfig.PerSignalUVMinMax[MultiplexId] = HistoryBufferScissorUVMinMax[MultiplexId];
 				}
 			}
 			#endif
diff --git a/Engine/Source/Runtime/Renderer/Private/SceneRendering.h b/Engine/Source/Runtime/Renderer/Private/SceneRendering.h
--- a/Engine/Source/Runtime/Renderer/Private/SceneRendering.h
+++ b/Engine/Source/Runtime/Renderer/Private/SceneRendering.h
@@ -775,6 +775,9 @@ struct FGTAOTAAHistory
 // Structure that hold all information related to previous frame.
 struct FPreviousViewInfo
 {
+	// View rect
+	FIntRect ViewRect;
+
 	// View matrices.
 	FViewMatrices ViewMatrices;
 
diff --git a/Engine/Source/Runtime/Renderer/Private/SceneVisibility.cpp b/Engine/Source/Runtime/Renderer/Private/SceneVisibility.cpp
--- a/Engine/Source/Runtime/Renderer/Private/SceneVisibility.cpp
+++ b/Engine/Source/Runtime/Renderer/Private/SceneVisibility.cpp
@@ -3346,6 +3346,7 @@ void FSceneRenderer::PreVisibilityFrameSetup(FRHICommandListImmediate& RHICmdLis
 		// Setup a new FPreviousViewInfo from current frame infos.
 		FPreviousViewInfo NewPrevViewInfo;
 		{
+			NewPrevViewInfo.ViewRect = View.ViewRect;
 			NewPrevViewInfo.ViewMatrices = View.ViewMatrices;
 		}
 
diff --git a/Engine/Source/Runtime/Renderer/Private/ScreenSpaceDenoise.cpp b/Engine/Source/Runtime/Renderer/Private/ScreenSpaceDenoise.cpp
--- a/Engine/Source/Runtime/Renderer/Private/ScreenSpaceDenoise.cpp
+++ b/Engine/Source/Runtime/Renderer/Private/ScreenSpaceDenoise.cpp
@@ -1141,8 +1141,10 @@ class FSSDTemporalAccumulationCS : public FGlobalShader
 	BEGIN_SHADER_PARAMETER_STRUCT(FParameters, )
 		SHADER_PARAMETER_ARRAY(int32, bCameraCut, [IScreenSpaceDenoiser::kMaxBatchSize])
 		SHADER_PARAMETER(float, HistoryPreExposureCorrection)
-
-		SHADER_PARAMETER_ARRAY(FVector4, HistoryBufferUVMinMax, [IScreenSpaceDenoiser::kMaxBatchSize])
+		SHADER_PARAMETER(FVector4, ScreenPosToHistoryBufferUV)
+		SHADER_PARAMETER(FVector4, HistoryBufferSizeAndInvSize)
+		SHADER_PARAMETER(FVector4, HistoryBufferUVMinMax)
+		SHADER_PARAMETER_ARRAY(FVector4, HistoryBufferScissorUVMinMax, [IScreenSpaceDenoiser::kMaxBatchSize])
 
 		SHADER_PARAMETER_STRUCT_INCLUDE(FSSDCommonParameters, CommonParameters)
 		SHADER_PARAMETER_STRUCT_INCLUDE(FSSDConvolutionMetaData, ConvolutionMetaData)
@@ -1852,7 +1854,6 @@ static void DenoiseSignalAtConstantPixelDensity(
 		PassParameters->CommonParameters = CommonParameters;
 		PassParameters->ConvolutionMetaData = ConvolutionMetaData;
 		PassParameters->HistoryPreExposureCorrection = View.PreExposure / View.PrevViewInfo.SceneColorPreExposure;
-
 		PassParameters->SignalInput = SignalHistory;
 		PassParameters->HistoryRejectionSignal = RejectionPreConvolutionSignal;
 		PassParameters->SignalHistoryOutput = CreateMultiplexedUAVs(GraphBuilder, SignalOutput);
@@ -1868,6 +1869,48 @@ static void DenoiseSignalAtConstantPixelDensity(
 				GraphBuilder, ViewInfoPooledRenderTargets.PrevCompressedDepthViewNormal, GSystemTextures.ZeroUIntDummy);
 		}
 
+		FIntPoint PrevFrameBufferExtent;
+
+		bool bGlobalCameraCut = !View.PrevViewInfo.DepthBuffer.IsValid() && !View.PrevViewInfo.CompressedDepthViewNormal.IsValid();
+		if (bGlobalCameraCut)
+		{
+			PassParameters->ScreenPosToHistoryBufferUV = FVector4(1.0f, 1.0f, 1.0f, 1.0f);
+			PassParameters->HistoryBufferUVMinMax = FVector4(0.0f, 0.0f, 0.0f, 0.0f);
+			PassParameters->HistoryBufferSizeAndInvSize = FVector4(1.0f, 1.0f, 1.0f, 1.0f);
+			PrevFrameBufferExtent = FIntPoint(1, 1);
+		}
+		else
+		{
+			FIntPoint ViewportOffset = View.PrevViewInfo.ViewRect.Min;
+			FIntPoint ViewportExtent = View.PrevViewInfo.ViewRect.Size();
+
+			if (PassParameters->PrevCompressedMetadata[0])
+			{
+				PrevFrameBufferExtent = PassParameters->PrevCompressedMetadata[0]->Desc.Extent;
+			}
+			else
+			{
+				PrevFrameBufferExtent = PassParameters->PrevDepthBuffer->Desc.Extent;
+			}
+
+			float InvBufferSizeX = 1.f / float(PrevFrameBufferExtent.X);
+			float InvBufferSizeY = 1.f / float(PrevFrameBufferExtent.Y);
+
+			PassParameters->ScreenPosToHistoryBufferUV = FVector4(
+				ViewportExtent.X * 0.5f * InvBufferSizeX,
+				-ViewportExtent.Y * 0.5f * InvBufferSizeY,
+				(ViewportExtent.X * 0.5f + ViewportOffset.X) * InvBufferSizeX,
+				(ViewportExtent.Y * 0.5f + ViewportOffset.Y) * InvBufferSizeY);
+
+			PassParameters->HistoryBufferUVMinMax = FVector4(
+				(ViewportOffset.X + 0.5f) * InvBufferSizeX,
+				(ViewportOffset.Y + 0.5f) * InvBufferSizeY,
+				(ViewportOffset.X + ViewportExtent.X - 0.5f) * InvBufferSizeX,
+				(ViewportOffset.Y + ViewportExtent.Y - 0.5f) * InvBufferSizeY);
+
+			PassParameters->HistoryBufferSizeAndInvSize = FVector4(PrevFrameBufferExtent.X, PrevFrameBufferExtent.Y, InvBufferSizeX, InvBufferSizeY);
+		}
+
 		FScreenSpaceDenoiserHistory DummyPrevFrameHistory;
 
 		// Setup signals' previous frame historu buffers.
@@ -1877,7 +1920,7 @@ static void DenoiseSignalAtConstantPixelDensity(
 
 			PassParameters->bCameraCut[BatchedSignalId] = !PrevFrameHistory->IsValid();
 
-			if (!(View.ViewState && Settings.bUseTemporalAccumulation))
+			if (!(View.ViewState && Settings.bUseTemporalAccumulation) || bGlobalCameraCut)
 			{
 				PassParameters->bCameraCut[BatchedSignalId] = true;
 			}
@@ -1889,11 +1932,11 @@ static void DenoiseSignalAtConstantPixelDensity(
 					GraphBuilder, PrevFrameHistory->RT[BufferId], GSystemTextures.BlackDummy);
 			}
 
-			PassParameters->HistoryBufferUVMinMax[BatchedSignalId] = FVector4(
-				float(PrevFrameHistory->Scissor.Min.X + 0.5f) / float(BufferExtent.X),
-				float(PrevFrameHistory->Scissor.Min.Y + 0.5f) / float(BufferExtent.Y),
-				float(PrevFrameHistory->Scissor.Max.X - 0.5f) / float(BufferExtent.X),
-				float(PrevFrameHistory->Scissor.Max.Y - 0.5f) / float(BufferExtent.Y));
+			PassParameters->HistoryBufferScissorUVMinMax[BatchedSignalId] = FVector4(
+				float(PrevFrameHistory->Scissor.Min.X + 0.5f) / float(PrevFrameBufferExtent.X),
+				float(PrevFrameHistory->Scissor.Min.Y + 0.5f) / float(PrevFrameBufferExtent.Y),
+				float(PrevFrameHistory->Scissor.Max.X - 0.5f) / float(PrevFrameBufferExtent.X),
+				float(PrevFrameHistory->Scissor.Max.Y - 0.5f) / float(PrevFrameBufferExtent.Y));
 
 			// Releases the reference on previous frame so the history's render target can be reused ASAP.
 			PrevFrameHistory->SafeRelease();
