Adding transmission override for ray tracing global illumination.

r.RayTracing.GlobalIllumination.EnableTransmission (default = 1)

diff --git a/Engine/Shaders/Private/PathTracing/Light/PathTracingDirectionalLight.ush b/Engine/Shaders/Private/PathTracing/Light/PathTracingDirectionalLight.ush
--- a/Engine/Shaders/Private/PathTracing/Light/PathTracingDirectionalLight.ush
+++ b/Engine/Shaders/Private/PathTracing/Light/PathTracingDirectionalLight.ush
@@ -52,7 +52,7 @@ void DirectionalLight_EstimateLight(
 	float3 LightDirection = GetNormal(LightId);
 	float3 SurfaceNormal = Payload.WorldNormal;
 	float NoL = dot(SurfaceNormal, LightDirection);
-	NoL = HasTransmission(LightId) ? abs(NoL) : saturate(NoL);
+	NoL = ENABLE_TRANSMISSION && HasTransmission(LightId) ? abs(NoL) : saturate(NoL);
 
 	float LightPower = Luminance(GetColor(LightId));
 	OutIrradiance = LightPower * NoL;
diff --git a/Engine/Shaders/Private/PathTracing/Light/PathTracingLightSampling.ush b/Engine/Shaders/Private/PathTracing/Light/PathTracingLightSampling.ush
--- a/Engine/Shaders/Private/PathTracing/Light/PathTracingLightSampling.ush
+++ b/Engine/Shaders/Private/PathTracing/Light/PathTracingLightSampling.ush
@@ -6,6 +6,10 @@ LightSampling.usf: Light sampling for next-event estimation
 
 #pragma once
 
+#ifndef ENABLE_TRANSMISSION
+	#define ENABLE_TRANSMISSION 1
+#endif
+
 #include "PathTracingLightCommon.ush"
 #include "PathTracingDirectionalLight.ush"
 #include "PathTracingPointLight.ush"
diff --git a/Engine/Shaders/Private/PathTracing/Light/PathTracingPointLight.ush b/Engine/Shaders/Private/PathTracing/Light/PathTracingPointLight.ush
--- a/Engine/Shaders/Private/PathTracing/Light/PathTracingPointLight.ush
+++ b/Engine/Shaders/Private/PathTracing/Light/PathTracingPointLight.ush
@@ -65,7 +65,7 @@ void PointLight_EstimateLight(
 	float3 SurfaceNormal = Payload.WorldNormal;
 	// dxr_todo: incorporate sphere light radius
 	float NoL = dot(SurfaceNormal, LightDirection);
-	NoL = HasTransmission(LightId) ? abs(NoL) : saturate(NoL);
+	NoL = ENABLE_TRANSMISSION && HasTransmission(LightId) ? abs(NoL) : saturate(NoL);
 
 	float LightPower = Luminance(GetColor(LightId));
 	float Falloff = ComputeAttenuationFalloff(sqrt(LightDistanceSquared), GetAttenuation(LightId));
diff --git a/Engine/Shaders/Private/PathTracing/Light/PathTracingRectLight.ush b/Engine/Shaders/Private/PathTracing/Light/PathTracingRectLight.ush
--- a/Engine/Shaders/Private/PathTracing/Light/PathTracingRectLight.ush
+++ b/Engine/Shaders/Private/PathTracing/Light/PathTracingRectLight.ush
@@ -73,7 +73,7 @@ void RectLight_EstimateLight(
 	// Approximate geometric term
 	float3 SurfaceNormal = Payload.WorldNormal;
 	float NoL = dot(SurfaceNormal, LightDirection);
-	NoL = HasTransmission(LightId) ? abs(NoL) : saturate(NoL);
+	NoL = ENABLE_TRANSMISSION && HasTransmission(LightId) ? abs(NoL) : saturate(NoL);
 	float Area = GetWidth(LightId) * GetHeight(LightId);
 
 	float LightPower = Luminance(GetColor(LightId));
diff --git a/Engine/Shaders/Private/PathTracing/Light/PathTracingSpotLight.ush b/Engine/Shaders/Private/PathTracing/Light/PathTracingSpotLight.ush
--- a/Engine/Shaders/Private/PathTracing/Light/PathTracingSpotLight.ush
+++ b/Engine/Shaders/Private/PathTracing/Light/PathTracingSpotLight.ush
@@ -145,7 +145,7 @@ void SpotLight_EstimateLight(
 		Falloff *= Falloff;
 		Falloff *= ComputeAttenuationFalloff(sqrt(LightDistanceSquared), GetAttenuation(LightId));
 		float CosTheta = dot(Payload.WorldNormal, LightDirection);
-		CosTheta = HasTransmission(LightId) ? abs(CosTheta) : saturate(CosTheta);
+		CosTheta = ENABLE_TRANSMISSION && HasTransmission(LightId) ? abs(CosTheta) : saturate(CosTheta);
 		OutIrradiance = LightPower * CosTheta * Falloff / LightDistanceSquared;
 		// #dxr_todo: UE-72534 apply inner core
 	}
diff --git a/Engine/Source/Runtime/Renderer/Private/RayTracing/RayTracingGlobalIllumination.cpp b/Engine/Source/Runtime/Renderer/Private/RayTracing/RayTracingGlobalIllumination.cpp
--- a/Engine/Source/Runtime/Renderer/Private/RayTracing/RayTracingGlobalIllumination.cpp
+++ b/Engine/Source/Runtime/Renderer/Private/RayTracing/RayTracingGlobalIllumination.cpp
@@ -121,6 +121,13 @@ static TAutoConsoleVariable<int32> CVarRayTracingGlobalIlluminationEnableTwoSide
 	ECVF_RenderThreadSafe
 );
 
+static TAutoConsoleVariable<int32> CVarRayTracingGlobalIlluminationEnableTransmission(
+	TEXT("r.RayTracing.GlobalIllumination.EnableTransmission"),
+	1,
+	TEXT("Enables transmission when tracing GI rays (default = 1)"),
+	ECVF_RenderThreadSafe
+);
+
 static int32 GRayTracingGlobalIlluminationRenderTileSize = 0;
 static FAutoConsoleVariableRef CVarRayTracingGlobalIlluminationRenderTileSize(
 	TEXT("r.RayTracing.GlobalIllumination.RenderTileSize"),
@@ -333,8 +340,9 @@ class FGlobalIlluminationRGS : public FGlobalShader
 
 	class FUseAttenuationTermDim : SHADER_PERMUTATION_BOOL("USE_ATTENUATION_TERM");
 	class FEnableTwoSidedGeometryDim : SHADER_PERMUTATION_BOOL("ENABLE_TWO_SIDED_GEOMETRY");
+	class FEnableTransmissionDim : SHADER_PERMUTATION_INT("ENABLE_TRANSMISSION", 2);
 
-	using FPermutationDomain = TShaderPermutationDomain<FUseAttenuationTermDim, FEnableTwoSidedGeometryDim>;
+	using FPermutationDomain = TShaderPermutationDomain<FUseAttenuationTermDim, FEnableTwoSidedGeometryDim, FEnableTransmissionDim>;
 
 	static bool ShouldCompilePermutation(const FGlobalShaderPermutationParameters& Parameters)
 	{
@@ -388,11 +396,12 @@ class FRayTracingGlobalIlluminationCreateGatherPointsRGS : public FGlobalShader
 	DECLARE_GLOBAL_SHADER(FRayTracingGlobalIlluminationCreateGatherPointsRGS)
 	SHADER_USE_ROOT_PARAMETER_STRUCT(FRayTracingGlobalIlluminationCreateGatherPointsRGS, FGlobalShader)
 
-		class FUseAttenuationTermDim : SHADER_PERMUTATION_BOOL("USE_ATTENUATION_TERM");
+	class FUseAttenuationTermDim : SHADER_PERMUTATION_BOOL("USE_ATTENUATION_TERM");
 	class FEnableTwoSidedGeometryDim : SHADER_PERMUTATION_BOOL("ENABLE_TWO_SIDED_GEOMETRY");
 	class FDeferredMaterialMode : SHADER_PERMUTATION_ENUM_CLASS("DIM_DEFERRED_MATERIAL_MODE", EDeferredMaterialMode);
+	class FEnableTransmissionDim : SHADER_PERMUTATION_INT("ENABLE_TRANSMISSION", 2);
 
-	using FPermutationDomain = TShaderPermutationDomain<FUseAttenuationTermDim, FEnableTwoSidedGeometryDim, FDeferredMaterialMode>;
+	using FPermutationDomain = TShaderPermutationDomain<FUseAttenuationTermDim, FEnableTwoSidedGeometryDim, FDeferredMaterialMode, FEnableTransmissionDim>;
 
 	static bool ShouldCompilePermutation(const FGlobalShaderPermutationParameters& Parameters)
 	{
@@ -451,7 +460,7 @@ class FRayTracingGlobalIlluminationCreateGatherPointsTraceRGS : public FGlobalSh
 	DECLARE_GLOBAL_SHADER(FRayTracingGlobalIlluminationCreateGatherPointsTraceRGS)
 	SHADER_USE_ROOT_PARAMETER_STRUCT(FRayTracingGlobalIlluminationCreateGatherPointsTraceRGS, FGlobalShader)
 
-		class FUseAttenuationTermDim : SHADER_PERMUTATION_BOOL("USE_ATTENUATION_TERM");
+	class FUseAttenuationTermDim : SHADER_PERMUTATION_BOOL("USE_ATTENUATION_TERM");
 	class FEnableTwoSidedGeometryDim : SHADER_PERMUTATION_BOOL("ENABLE_TWO_SIDED_GEOMETRY");
 	class FDeferredMaterialMode : SHADER_PERMUTATION_ENUM_CLASS("DIM_DEFERRED_MATERIAL_MODE", EDeferredMaterialMode);
 
@@ -514,7 +523,7 @@ class FRayTracingGlobalIlluminationFinalGatherRGS : public FGlobalShader
 	DECLARE_GLOBAL_SHADER(FRayTracingGlobalIlluminationFinalGatherRGS)
 	SHADER_USE_ROOT_PARAMETER_STRUCT(FRayTracingGlobalIlluminationFinalGatherRGS, FGlobalShader)
 
-		class FUseAttenuationTermDim : SHADER_PERMUTATION_BOOL("USE_ATTENUATION_TERM");
+	class FUseAttenuationTermDim : SHADER_PERMUTATION_BOOL("USE_ATTENUATION_TERM");
 	class FEnableTwoSidedGeometryDim : SHADER_PERMUTATION_BOOL("ENABLE_TWO_SIDED_GEOMETRY");
 
 	using FPermutationDomain = TShaderPermutationDomain<FUseAttenuationTermDim, FEnableTwoSidedGeometryDim>;
@@ -558,6 +567,7 @@ IMPLEMENT_GLOBAL_SHADER(FRayTracingGlobalIlluminationFinalGatherRGS, "/Engine/Pr
 void FDeferredShadingSceneRenderer::PrepareRayTracingGlobalIllumination(const FViewInfo& View, TArray<FRHIRayTracingShader*>& OutRayGenShaders)
 {
 	const bool bSortMaterials = CVarRayTracingGlobalIlluminationFinalGatherSortMaterials.GetValueOnRenderThread() != 0;
+	int EnableTransmission = CVarRayTracingGlobalIlluminationEnableTransmission.GetValueOnRenderThread();
 
 	// Declare all RayGen shaders that require material closest hit shaders to be bound
 	for (int UseAttenuationTerm = 0; UseAttenuationTerm < 2; ++UseAttenuationTerm)
@@ -567,6 +577,7 @@ void FDeferredShadingSceneRenderer::PrepareRayTracingGlobalIllumination(const FV
 			FGlobalIlluminationRGS::FPermutationDomain PermutationVector;
 			PermutationVector.Set<FGlobalIlluminationRGS::FUseAttenuationTermDim>(CVarRayTracingGlobalIlluminationEnableLightAttenuation.GetValueOnRenderThread() != 0);
 			PermutationVector.Set<FGlobalIlluminationRGS::FEnableTwoSidedGeometryDim>(EnableTwoSidedGeometry == 1);
+			PermutationVector.Set<FGlobalIlluminationRGS::FEnableTransmissionDim>(EnableTransmission);
 			TShaderMapRef<FGlobalIlluminationRGS> RayGenerationShader(View.ShaderMap, PermutationVector);
 			OutRayGenShaders.Add(RayGenerationShader.GetRayTracingShader());
 
@@ -588,6 +599,7 @@ void FDeferredShadingSceneRenderer::PrepareRayTracingGlobalIllumination(const FV
 					CreateGatherPointsPermutationVector.Set<FRayTracingGlobalIlluminationCreateGatherPointsRGS::FUseAttenuationTermDim>(UseAttenuationTerm == 1);
 					CreateGatherPointsPermutationVector.Set<FRayTracingGlobalIlluminationCreateGatherPointsRGS::FEnableTwoSidedGeometryDim>(EnableTwoSidedGeometry == 1);
 					CreateGatherPointsPermutationVector.Set<FRayTracingGlobalIlluminationCreateGatherPointsRGS::FDeferredMaterialMode>(EDeferredMaterialMode::Shade);
+					CreateGatherPointsPermutationVector.Set<FRayTracingGlobalIlluminationCreateGatherPointsRGS::FEnableTransmissionDim>(EnableTransmission);
 					TShaderMapRef<FRayTracingGlobalIlluminationCreateGatherPointsRGS> CreateGatherPointsRayGenerationShader(View.ShaderMap, CreateGatherPointsPermutationVector);
 					OutRayGenShaders.Add(CreateGatherPointsRayGenerationShader.GetRayTracingShader());
 				}
@@ -598,6 +610,7 @@ void FDeferredShadingSceneRenderer::PrepareRayTracingGlobalIllumination(const FV
 				CreateGatherPointsPermutationVector.Set<FRayTracingGlobalIlluminationCreateGatherPointsRGS::FUseAttenuationTermDim>(UseAttenuationTerm == 1);
 				CreateGatherPointsPermutationVector.Set<FRayTracingGlobalIlluminationCreateGatherPointsRGS::FEnableTwoSidedGeometryDim>(EnableTwoSidedGeometry == 1);
 				CreateGatherPointsPermutationVector.Set<FRayTracingGlobalIlluminationCreateGatherPointsRGS::FDeferredMaterialMode>(EDeferredMaterialMode::None);
+				CreateGatherPointsPermutationVector.Set < FRayTracingGlobalIlluminationCreateGatherPointsRGS::FEnableTransmissionDim>(EnableTransmission);
 				TShaderMapRef<FRayTracingGlobalIlluminationCreateGatherPointsRGS> CreateGatherPointsRayGenerationShader(View.ShaderMap, CreateGatherPointsPermutationVector);
 				OutRayGenShaders.Add(CreateGatherPointsRayGenerationShader.GetRayTracingShader());
 			}
@@ -827,6 +840,7 @@ void FDeferredShadingSceneRenderer::RayTracingGlobalIlluminationCreateGatherPoin
 	FRayTracingGlobalIlluminationCreateGatherPointsRGS::FPermutationDomain PermutationVector;
 	PermutationVector.Set<FRayTracingGlobalIlluminationCreateGatherPointsRGS::FUseAttenuationTermDim>(true);
 	PermutationVector.Set<FRayTracingGlobalIlluminationCreateGatherPointsRGS::FEnableTwoSidedGeometryDim>(CVarRayTracingGlobalIlluminationEnableTwoSidedGeometry.GetValueOnRenderThread() != 0);
+	PermutationVector.Set<FRayTracingGlobalIlluminationCreateGatherPointsRGS::FEnableTransmissionDim>(CVarRayTracingGlobalIlluminationEnableTransmission.GetValueOnRenderThread());
 	TShaderMapRef<FRayTracingGlobalIlluminationCreateGatherPointsRGS> RayGenerationShader(GetGlobalShaderMap(FeatureLevel), PermutationVector);
 	ClearUnusedGraphResources(RayGenerationShader, PassParameters);
 
@@ -902,6 +916,7 @@ void FDeferredShadingSceneRenderer::RayTracingGlobalIlluminationCreateGatherPoin
 			PermutationVector.Set<FRayTracingGlobalIlluminationCreateGatherPointsRGS::FUseAttenuationTermDim>(true);
 			PermutationVector.Set<FRayTracingGlobalIlluminationCreateGatherPointsRGS::FEnableTwoSidedGeometryDim>(CVarRayTracingGlobalIlluminationEnableTwoSidedGeometry.GetValueOnRenderThread() != 0);
 			PermutationVector.Set<FRayTracingGlobalIlluminationCreateGatherPointsRGS::FDeferredMaterialMode>(EDeferredMaterialMode::Shade);
+			PermutationVector.Set<FRayTracingGlobalIlluminationCreateGatherPointsRGS::FEnableTransmissionDim>(CVarRayTracingGlobalIlluminationEnableTransmission.GetValueOnRenderThread());
 			TShaderMapRef<FRayTracingGlobalIlluminationCreateGatherPointsRGS> RayGenerationShader(GetGlobalShaderMap(FeatureLevel), PermutationVector);
 			ClearUnusedGraphResources(RayGenerationShader, PassParameters);
 
@@ -1098,6 +1113,7 @@ void FDeferredShadingSceneRenderer::RenderRayTracingGlobalIlluminationBruteForce
 	FGlobalIlluminationRGS::FPermutationDomain PermutationVector;
 	PermutationVector.Set<FGlobalIlluminationRGS::FUseAttenuationTermDim>(true);
 	PermutationVector.Set<FGlobalIlluminationRGS::FEnableTwoSidedGeometryDim>(CVarRayTracingGlobalIlluminationEnableTwoSidedGeometry.GetValueOnRenderThread() != 0);
+	PermutationVector.Set<FGlobalIlluminationRGS::FEnableTransmissionDim>(CVarRayTracingGlobalIlluminationEnableTransmission.GetValueOnRenderThread());
 	TShaderMapRef<FGlobalIlluminationRGS> RayGenerationShader(GetGlobalShaderMap(FeatureLevel), PermutationVector);
 	ClearUnusedGraphResources(RayGenerationShader, PassParameters);
 
