Fix crash when invoking GI without a SkyLight.

diff --git a/Engine/Source/Runtime/Renderer/Private/RayTracing/RayTracingGlobalIllumination.cpp b/Engine/Source/Runtime/Renderer/Private/RayTracing/RayTracingGlobalIllumination.cpp
--- a/Engine/Source/Runtime/Renderer/Private/RayTracing/RayTracingGlobalIllumination.cpp
+++ b/Engine/Source/Runtime/Renderer/Private/RayTracing/RayTracingGlobalIllumination.cpp
@@ -190,11 +190,12 @@ void SetupLightParameters(
 
 	FSkyLightSceneProxy* SkyLight = Scene.SkyLight;
 	FVector SkyLightColor = FVector(0.0f, 0.0f, 0.0f);
+	uint32 SkyLightTransmission = 0;
 	if (SkyLight && SkyLight->bAffectGlobalIllumination)
 	{
 		SkyLightColor = FVector(SkyLight->GetEffectiveLightColor());
+		SkyLightTransmission = SkyLight->bTransmission;
 	}
-	uint32 SkyLightTransmission = SkyLight->bTransmission;
 
 	// Prepend SkyLight to light buffer
 	// WARNING: Until ray payload encodes Light data buffer, the execution depends on this ordering!
