Fix CIS.

diff --git a/Engine/Source/Runtime/Renderer/Private/DeferredShadingRenderer.cpp b/Engine/Source/Runtime/Renderer/Private/DeferredShadingRenderer.cpp
--- a/Engine/Source/Runtime/Renderer/Private/DeferredShadingRenderer.cpp
+++ b/Engine/Source/Runtime/Renderer/Private/DeferredShadingRenderer.cpp
@@ -2915,6 +2915,8 @@ void FDeferredShadingSceneRenderer::Render(FRHICommandListImmediate& RHICmdList)
 			{
 				FViewInfo& View = Views[ViewIndex];
 
+#if !(UE_BUILD_SHIPPING || UE_BUILD_TEST)
+
 				if (IsPostProcessVisualizeDebugMaterialEnabled(View))
 				{
 					const UMaterialInterface* DebugMaterialInterface = GetPostProcessVisualizeDebugMaterialInterface(View);
@@ -2927,10 +2929,14 @@ void FDeferredShadingSceneRenderer::Render(FRHICommandListImmediate& RHICmdList)
 				}
 				else
 				{
+#endif
 					SCOPED_GPU_MASK(RHICmdList, View.GPUMask);
 					RDG_EVENT_SCOPE_CONDITIONAL(GraphBuilder, Views.Num() > 1, "View%d", ViewIndex);
 					AddPostProcessingPasses(GraphBuilder, View, PostProcessingInputs);
+
+#if !(UE_BUILD_SHIPPING || UE_BUILD_TEST)
 				}
+#endif
 			}
 		}
 
