EM - Fix Edge case in the loop not checking the divisor again

diff --git a/Engine/Source/Runtime/Renderer/Public/MeshDrawShaderBindings.h b/Engine/Source/Runtime/Renderer/Public/MeshDrawShaderBindings.h
--- a/Engine/Source/Runtime/Renderer/Public/MeshDrawShaderBindings.h
+++ b/Engine/Source/Runtime/Renderer/Public/MeshDrawShaderBindings.h
@@ -276,7 +276,7 @@ private:
 		}
 
 		//small size array optimization
-		const int ArrayEnd = Index + Size;
+		const int ArrayEnd = FMath::Min(Index + Size + 1, Array.Num() - Index);
 		while (Index < ArrayEnd)
 		{
 			if (Array[Index].BaseIndex == BaseIndex)
