Implements reflection denoiser's preconvolution pass

diff --git a/Engine/Shaders/Private/RayTracing/RayTracingReflections.usf b/Engine/Shaders/Private/RayTracing/RayTracingReflections.usf
--- a/Engine/Shaders/Private/RayTracing/RayTracingReflections.usf
+++ b/Engine/Shaders/Private/RayTracing/RayTracingReflections.usf
@@ -974,12 +974,17 @@ RAY_TRACING_ENTRY_RAYGEN(RayTracingReflectionsRGS)
 			CurrentImaginaryReflectionGBufferData.ValidSamples = 1;
 		}
 
+		float ConfusionFactor = ComputeDenoiserConfusionFactor(
+			ClosestHitDistance != DENOISER_INVALID_HIT_DISTANCE,
+			length(WorldPosition - View.WorldCameraOrigin),
+			ClosestHitDistance);
+
 		// Accumulate output information if it has been set by the initial sample
 		if (SampleMode == SAMPLE_PHASE_ACCUM || SampleMode == SAMPLE_PHASE_RESOLVE)
 		{
 			// Read the previous sample's output information
 			float4 PreviousColor = ColorOutput[DispatchThreadId];
-			float PreviousClosestHitDistance = RayHitDistanceOutput[DispatchThreadId];
+			float PreviousConfusionFactor = RayHitDistanceOutput[DispatchThreadId];
 			FImaginaryReflectionGBufferData PreviousImaginaryReflectionGBufferData;
 
 #if DIM_RAY_TRACE_SKY_LIGHT_CONTRIBUTION
@@ -997,13 +1002,13 @@ RAY_TRACING_ENTRY_RAYGEN(RayTracingReflectionsRGS)
 				ReflectedColor += PreviousColor;
 				ReflectedWeight += PreviousColor.a;
 
-				if (ClosestHitDistance > 0.0 && PreviousClosestHitDistance > 0.0)
+				if (ConfusionFactor >= 0.0 && PreviousConfusionFactor >= 0.0)
 				{
-					ClosestHitDistance = min(ClosestHitDistance, PreviousClosestHitDistance);
+					ConfusionFactor = min(ConfusionFactor, PreviousConfusionFactor);
 				}
 				else
 				{
-					ClosestHitDistance = max(ClosestHitDistance, PreviousClosestHitDistance);
+					ConfusionFactor = max(ConfusionFactor, PreviousConfusionFactor);
 				}
 
 				if (DIM_RAY_TRACE_SKY_LIGHT_CONTRIBUTION)
@@ -1019,7 +1024,7 @@ RAY_TRACING_ENTRY_RAYGEN(RayTracingReflectionsRGS)
 				// no data this pass, use accumulated data
 				ReflectedColor = PreviousColor;
 				ReflectedWeight = PreviousColor.a;
-				ClosestHitDistance = PreviousClosestHitDistance;
+				ConfusionFactor = PreviousConfusionFactor;
 
 				if (DIM_RAY_TRACE_SKY_LIGHT_CONTRIBUTION)
 				{
@@ -1059,9 +1064,11 @@ RAY_TRACING_ENTRY_RAYGEN(RayTracingReflectionsRGS)
 		ReflectedColor.xyz *= View.PreExposure;
 	#endif
 
+		float DistanceCameraToPixel = 
+
 		// Output required information
 		ColorOutput[DispatchThreadId] = ClampToHalfFloatRange(ReflectedColor);
-		RayHitDistanceOutput[DispatchThreadId] = ClosestHitDistance;
+		RayHitDistanceOutput[DispatchThreadId] = ConfusionFactor;
 		RayImaginaryDepthOutput[DispatchThreadId] = ImaginaryClosestSVDepth; // TODO(Denoiser): shader permutation?
 
 #if DIM_RAY_TRACE_SKY_LIGHT_CONTRIBUTION
diff --git a/Engine/Shaders/Private/ScreenSpaceDenoise/SSDDefinitions.ush b/Engine/Shaders/Private/ScreenSpaceDenoise/SSDDefinitions.ush
--- a/Engine/Shaders/Private/ScreenSpaceDenoise/SSDDefinitions.ush
+++ b/Engine/Shaders/Private/ScreenSpaceDenoise/SSDDefinitions.ush
@@ -58,6 +58,11 @@
 #define WORLD_RADIUS_INVALID -1
 #define WORLD_RADIUS_MISS (INFINITE_FLOAT)
 
+// World bluring radius for a miss.
+#define CONFUSION_FACTOR_INVALID -1
+#define CONFUSION_FACTOR_MISS 1
+
+
 // Mip level for sample that are invalid or missed.
 #define MIP_LEVEL_INVALID -1
 #define MIP_LEVEL_MISS 1024
@@ -116,9 +121,6 @@
 	/** Buffer layout for the reflection given by the ray generation shader. */
 	#define SIGNAL_BUFFER_LAYOUT_REFLECTIONS_INPUT 3
 
-	/** Buffer layout for the rejection output. */
-	#define SIGNAL_BUFFER_LAYOUT_REFLECTIONS_REJECTION 9
-
 	/** Buffer layout for the reflection output. */
 	#define SIGNAL_BUFFER_LAYOUT_REFLECTIONS_HISTORY 4
 
diff --git a/Engine/Shaders/Private/ScreenSpaceDenoise/SSDPublic.ush b/Engine/Shaders/Private/ScreenSpaceDenoise/SSDPublic.ush
--- a/Engine/Shaders/Private/ScreenSpaceDenoise/SSDPublic.ush
+++ b/Engine/Shaders/Private/ScreenSpaceDenoise/SSDPublic.ush
@@ -16,6 +16,9 @@
 /** Special value to be encoded in denoiser input when no ray intersection has been found. */
 #define DENOISER_MISS_HIT_DISTANCE -1.0
 
+/** Special value to be encoded in denoiser input when no ray has been shot. */
+#define DENOISER_INVALID_CONFUSION_FACTOR -1.0
+
 
 /** Random generator used for 1 sample per pixel for denoiser input. */
 float2 Rand1SPPDenoiserInput(uint2 PixelPos) // TODO(Denoiser): kill
@@ -38,3 +41,9 @@ float2 Rand1SPPDenoiserInput(uint2 PixelPos) // TODO(Denoiser): kill
 
 	return E;
 }
+
+/** Compute the factor to apply to inifinity bluring radius based on the sample's hit distance. */
+float ComputeDenoiserConfusionFactor(bool bValid, float DistanceCameraToPixel, float ClosestHitDistance)
+{
+	return bValid ? (ClosestHitDistance / (DistanceCameraToPixel + ClosestHitDistance)) : DENOISER_INVALID_CONFUSION_FACTOR;
+}
diff --git a/Engine/Shaders/Private/ScreenSpaceDenoise/SSDSignalAccumulator.ush b/Engine/Shaders/Private/ScreenSpaceDenoise/SSDSignalAccumulator.ush
--- a/Engine/Shaders/Private/ScreenSpaceDenoise/SSDSignalAccumulator.ush
+++ b/Engine/Shaders/Private/ScreenSpaceDenoise/SSDSignalAccumulator.ush
@@ -117,6 +117,7 @@ FSSDSignalAccumulator CreateSignalAccumulator()
 	#if COMPILE_MIN_FREQUENCY_ACCUMULATOR
 		Accumulator.MinFrequency.ClosestHitDistance = INFINITE_FLOAT;
 		Accumulator.MinFrequency.WorldBluringRadius = WORLD_RADIUS_MISS;
+		Accumulator.MinFrequency.ConfusionFactor = CONFUSION_FACTOR_MISS;
 	#endif
 
 	// DRB initialization.
diff --git a/Engine/Shaders/Private/ScreenSpaceDenoise/SSDSignalBufferEncoding.ush b/Engine/Shaders/Private/ScreenSpaceDenoise/SSDSignalBufferEncoding.ush
--- a/Engine/Shaders/Private/ScreenSpaceDenoise/SSDSignalBufferEncoding.ush
+++ b/Engine/Shaders/Private/ScreenSpaceDenoise/SSDSignalBufferEncoding.ush
@@ -283,31 +283,12 @@ void DecodeMultiplexedSignalsFromFloat4(
 	#if COMPILE_SIGNAL_COLOR
 	else if (SignalBufferLayout == SIGNAL_BUFFER_LAYOUT_REFLECTIONS_INPUT)
 	{
-		OutSamples.Array[0].SampleCount = RawSample[1].r == DENOISER_INVALID_HIT_DISTANCE ? 0 : 1;
-		OutSamples.Array[0].MissCount = OutSamples.Array[0].SampleCount * (1 - RawSample[0].a);
+		OutSamples.Array[0].SampleCount = RawSample[1].r == DENOISER_INVALID_CONFUSION_FACTOR ? 0 : 1;
+		OutSamples.Array[0].MissCount = RawSample[1].r == DENOISER_INVALID_CONFUSION_FACTOR ? 0 : (1 - RawSample[0].a);
 		OutSamples.Array[0].SceneColor = RawSample[0];
 
 		// max() because ray hit distance is < 0 when miss / not ray has been shot.
-		OutSampleFrequencies.Array[0].ClosestHitDistance = RawSample[1].r;
-		//OutSampleFrequencies.Array[0].ClosestHitDistance = ((RawSample[1].r == DENOISER_MISS_HIT_DISTANCE) ? 1.0e20 : RawSample[1].r);
-	}
-	else if (SignalBufferLayout == SIGNAL_BUFFER_LAYOUT_REFLECTIONS_REJECTION)
-	{
-		if (bNormalizeSample)
-		{
-			OutSamples.Array[0].SampleCount = RawSample[1].r > 0 ? 1 : 0;
-		}
-		else
-		{
-			OutSamples.Array[0].SampleCount = RawSample[1].r;
-		}
-		
-		OutSamples.Array[0].MissCount = OutSamples.Array[0].SampleCount * (1 - RawSample[0].a);
-		OutSamples.Array[0].SceneColor = OutSamples.Array[0].SampleCount * RawSample[0];
-
-		OutSamples.Array[1].SampleCount = RawSample[1].g;
-		OutSamples.Array[1].MissCount = OutSamples.Array[0].SampleCount * (1 - RawSample[2].a);
-		OutSamples.Array[1].SceneColor = OutSamples.Array[0].SampleCount * RawSample[2];
+		OutSampleFrequencies.Array[0].ConfusionFactor = RawSample[1].r;
 	}
 	else if (SignalBufferLayout == SIGNAL_BUFFER_LAYOUT_REFLECTIONS_HISTORY)
 	{
@@ -323,13 +304,7 @@ void DecodeMultiplexedSignalsFromFloat4(
 		OutSamples.Array[0].MissCount = OutSamples.Array[0].SampleCount * (1 - RawSample[0].a);
 		OutSamples.Array[0].SceneColor = OutSamples.Array[0].SampleCount * RawSample[0];
 
-		// Temporal analysis of the history.
-		#if SIGNAL_ARRAY_SIZE >= 3
-			OutSamples.Array[1].SampleCount = OutSamples.Array[0].SampleCount;
-			OutSamples.Array[1].SceneColor.x = OutSamples.Array[1].SampleCount * RawSample[1].g;
-			OutSamples.Array[2].SampleCount = OutSamples.Array[0].SampleCount;
-			OutSamples.Array[2].SceneColor.x = OutSamples.Array[2].SampleCount * RawSample[1].b;
-		#endif
+		OutSampleFrequencies.Array[0].ConfusionFactor = RawSample[1].g;
 	}
 	#endif
 	#if 1
@@ -704,20 +679,6 @@ void EncodeMultiplexedSignals(
 	}
 	#endif
 	#if COMPILE_SIGNAL_COLOR
-	else if (SignalBufferLayout == SIGNAL_BUFFER_LAYOUT_REFLECTIONS_REJECTION)
-	{
-		// Outputs scene color to be compatible with a SSR output.
-		OutRawSample[0] = ClampColorForEncoding(Sample[0].SceneColor * (Sample[0].SampleCount > 0 ? rcp(Sample[0].SampleCount) : 0));
-		OutRawSample[1].r = Sample[0].SampleCount;
-		OutBufferCount = 2;
-
-		if (MultiplexCount == 2)
-		{
-			OutRawSample[2] = ClampColorForEncoding(Sample[1].SceneColor * (Sample[1].SampleCount > 0 ? rcp(Sample[1].SampleCount) : 0));
-			OutRawSample[1].g = Sample[1].SampleCount;
-			OutBufferCount = 3;
-		}
-	}
 	else if (SignalBufferLayout == SIGNAL_BUFFER_LAYOUT_REFLECTIONS_HISTORY)
 	{
 		float NormalizationFactor = SafeRcp(Sample[0].SampleCount);
@@ -725,17 +686,8 @@ void EncodeMultiplexedSignals(
 		// Outputs scene color to be compatible with a SSR output.
 		OutRawSample[0] = ClampColorForEncoding(Sample[0].SceneColor * NormalizationFactor);
 		OutRawSample[1].r = Sample[0].SampleCount;
+		OutRawSample[1].g = SampleFrequency[0].ConfusionFactor;
 		OutBufferCount = 2;
-
-		// Temporal analysis.
-		#if SIGNAL_ARRAY_SIZE >= 3
-		if (MultiplexCount == 3)
-		{
-			OutRawSample[1].g = Sample[1].SceneColor.x * NormalizationFactor;
-			OutRawSample[1].b = Sample[2].SceneColor.x * NormalizationFactor;
-			OutBufferCount = 3;
-		}
-		#endif
 	}
 	#endif
 	#if 1
diff --git a/Engine/Shaders/Private/ScreenSpaceDenoise/SSDSignalCore.ush b/Engine/Shaders/Private/ScreenSpaceDenoise/SSDSignalCore.ush
--- a/Engine/Shaders/Private/ScreenSpaceDenoise/SSDSignalCore.ush
+++ b/Engine/Shaders/Private/ScreenSpaceDenoise/SSDSignalCore.ush
@@ -65,6 +65,9 @@ struct FSSDSignalFrequency
 
 	// Blur radius allowed for this sample.
 	float WorldBluringRadius;
+
+	// Circle confusion multiplier.
+	float ConfusionFactor;
 };
 
 /** Array of signal samples. Technically should be a typedef, but there is a bug in HLSL with the out operator of an array. */
@@ -119,6 +122,7 @@ FSSDSignalFrequency CreateInvalidSignalFrequency()
 	FSSDSignalFrequency Frequency;
 	Frequency.ClosestHitDistance = DENOISER_INVALID_HIT_DISTANCE;
 	Frequency.WorldBluringRadius = WORLD_RADIUS_INVALID;
+	Frequency.ConfusionFactor = CONFUSION_FACTOR_INVALID;
 	return Frequency;
 }
 
@@ -477,6 +481,7 @@ FSSDSignalFrequency MinSignalFrequency(FSSDSignalFrequency FrequencyA, FSSDSigna
 	FSSDSignalFrequency OutFrequency;
 	OutFrequency.ClosestHitDistance = min(FrequencyA.ClosestHitDistance, FrequencyB.ClosestHitDistance);
 	OutFrequency.WorldBluringRadius = min(FrequencyA.WorldBluringRadius, FrequencyB.WorldBluringRadius);
+	OutFrequency.ConfusionFactor = min(FrequencyA.ConfusionFactor, FrequencyB.ConfusionFactor);
 	return OutFrequency;
 }
 
@@ -485,5 +490,6 @@ FSSDSignalFrequency WaveBroadcastSignalFrequency(const FWaveBroadcastSettings Br
 	FSSDSignalFrequency OutFrequency;
 	OutFrequency.ClosestHitDistance = WaveBroadcast(BroadcastSettings, Frequency.ClosestHitDistance);
 	OutFrequency.WorldBluringRadius = WaveBroadcast(BroadcastSettings, Frequency.WorldBluringRadius);
+	OutFrequency.ConfusionFactor = WaveBroadcast(BroadcastSettings, Frequency.ConfusionFactor);
 	return OutFrequency;
 }
diff --git a/Engine/Shaders/Private/ScreenSpaceDenoise/SSDSignalFramework.ush b/Engine/Shaders/Private/ScreenSpaceDenoise/SSDSignalFramework.ush
--- a/Engine/Shaders/Private/ScreenSpaceDenoise/SSDSignalFramework.ush
+++ b/Engine/Shaders/Private/ScreenSpaceDenoise/SSDSignalFramework.ush
@@ -452,28 +452,6 @@ void ProjectSpecularLobeToScreenSpace(
 	}
 }
 
-/** Return the size of the bluring radius caused by the specular lobe in ViewportUV.x space. */
-float ComputeInfinityBluringRadius(FSSDSampleSceneInfos SceneMetadata)
-{
-	float2 NormalizedScreenMajorAxis;
-	float InifinityMajorViewportRadius;
-	float InifinityMinorViewportRadius;
-	ProjectSpecularLobeToScreenSpace(
-		SceneMetadata,
-		/* out */ NormalizedScreenMajorAxis,
-		/* out */ InifinityMajorViewportRadius,
-		/* out */ InifinityMinorViewportRadius);
-	
-	// To avoid overbluring killing anisotropy, use minor viewport radius that is smaller.
-	return InifinityMinorViewportRadius;
-}
-
-/** Return the size of bluring radius caused by specular lobe in ViewportUV.x space. */
-float ComputeBluringRadiusOfNormalizedSample(FSSDSignalFrequency SampleFrequency, FSSDSampleSceneInfos SceneMetadata)
-{
-	return ComputeInfinityBluringRadius(SceneMetadata) * ComputeBluringOutOfFocusOfNormalizedSample(SampleFrequency, SceneMetadata);
-}
-
 
 //------------------------------------------------------- MAIN FUNCTIONS TO IMPLEMENT FOR THE PASSES
 
diff --git a/Engine/Shaders/Private/ScreenSpaceDenoise/SSDSpatialAccumulation.usf b/Engine/Shaders/Private/ScreenSpaceDenoise/SSDSpatialAccumulation.usf
--- a/Engine/Shaders/Private/ScreenSpaceDenoise/SSDSpatialAccumulation.usf
+++ b/Engine/Shaders/Private/ScreenSpaceDenoise/SSDSpatialAccumulation.usf
@@ -14,9 +14,6 @@
 
 /** Policy to use to change the size of kernel. */
 	#define SAMPLE_COUNT_POLICY_DISABLED 0
-	#define SAMPLE_COUNT_POLICY_METADATA_BASED 1
-	#define SAMPLE_COUNT_POLICY_HIT_DISTANCE_BASED 2
-	#define SAMPLE_COUNT_POLICY_ACCUMULATION_BASED 3
 	#define SAMPLE_COUNT_POLICY_SAMPLE_ACCUMULATION_BASED 4
 
 /** What signal should be outputed. */
@@ -158,65 +155,31 @@
 #elif CONFIG_SIGNAL_PROCESSING == SIGNAL_PROCESSING_REFLECTIONS
 	// Denoise only specular.
 	#define MAX_SIGNAL_BATCH_SIZE 1
+	#define SIGNAL_ARRAY_SIZE 1
 
 	#define COMPILE_SIGNAL_COLOR 1
+	#define CONFIG_BILATERAL_PRESET BILATERAL_PRESET_REFLECTIONS
 
-	#if DIM_STAGE == STAGE_RECONSTRUCTION
+	#if DIM_STAGE == STAGE_RECONSTRUCTION || DIM_STAGE == STAGE_PRE_CONVOLUTION
 		#define SIGNAL_ARRAY_SIZE 1
 
 		// Input and output layout.
+		#if DIM_STAGE == STAGE_RECONSTRUCTION 
+			#define CONFIG_INPUT_TEXTURE_COUNT 2
 			#define CONFIG_SIGNAL_INPUT_LAYOUT  SIGNAL_BUFFER_LAYOUT_REFLECTIONS_INPUT
+		#else
+			#define CONFIG_INPUT_TEXTURE_COUNT 2
+			#define CONFIG_SIGNAL_INPUT_LAYOUT  SIGNAL_BUFFER_LAYOUT_REFLECTIONS_HISTORY
+		#endif
 		#define CONFIG_SIGNAL_OUTPUT_LAYOUT SIGNAL_BUFFER_LAYOUT_REFLECTIONS_HISTORY
+		#define CONFIG_OUTPUT_TEXTURE_COUNT 2
 
-		// Configure the bilateral therm.
 		#define CONFIG_SAMPLE_SET           SAMPLE_SET_DIRECTIONAL_ELLIPSE
 
-		#define CONFIG_BILATERAL_PRESET BILATERAL_PRESET_REFLECTIONS
-		//#define CONFIG_BILATERAL_PRESET BILATERAL_PRESET_REFLECTIONS_1SPP
-		//#define CONFIG_BILATERAL_PRESET BILATERAL_PRESET_DISABLED
-
-		//#define CONFIG_SAMPLE_COUNT_POLICY  SAMPLE_COUNT_POLICY_DISABLED
-
 		// Do color accumulation with karis weighting to avoid flickering specular highlight to show up the kernel pattern.
 		// TODO(Denoiser): This is a bit agressive.
 		#define CONFIG_ACCUMULATION_COLOR_SPACE (COLOR_SPACE_RGB | COLOR_SPACE_KARIS_WEIGHTING)
 
-		#define CONFIG_INPUT_TEXTURE_COUNT 2
-		#define CONFIG_OUTPUT_TEXTURE_COUNT 2
-
-	#elif DIM_STAGE == STAGE_REJECTION_PRE_CONVOLUTION
-		#define SIGNAL_ARRAY_SIZE 2 // first and second momment to measure variance in temporal accumulation.
-
-		// Input and output layout.
-		#define CONFIG_SIGNAL_INPUT_LAYOUT  SIGNAL_BUFFER_LAYOUT_REFLECTIONS_HISTORY
-		#define CONFIG_SIGNAL_OUTPUT_LAYOUT SIGNAL_BUFFER_LAYOUT_REFLECTIONS_REJECTION
-
-		#define CONFIG_BILATERAL_PRESET BILATERAL_PRESET_REFLECTIONS
-
-		#define CONFIG_ACCUMULATION_COLOR_SPACE REFLECTIONS_REJECTION_COLOR_SPACE
-		#define CONFIG_OUTPUT_COLOR_SPACE       REFLECTIONS_REJECTION_COLOR_SPACE
-
-		#define CONFIG_INPUT_TEXTURE_COUNT 2
-		#define CONFIG_OUTPUT_TEXTURE_COUNT 3
-
-	#elif DIM_STAGE == STAGE_POST_FILTERING
-		#define SIGNAL_ARRAY_SIZE 1
-
-		// Input and output layout.
-		#define CONFIG_SIGNAL_INPUT_LAYOUT  SIGNAL_BUFFER_LAYOUT_REFLECTIONS_HISTORY
-		#define CONFIG_SIGNAL_OUTPUT_LAYOUT SIGNAL_BUFFER_LAYOUT_REFLECTIONS_HISTORY
-
-		//#define CONFIG_SAMPLE_SET           SAMPLE_SET_3X3
-		#define CONFIG_SAMPLE_SET           SAMPLE_SET_STACKOWIAK_4_SETS
-		#define CONFIG_SAMPLE_SUBSET        1
-		#define CONFIG_SAMPLE_COUNT_POLICY  SAMPLE_COUNT_POLICY_ACCUMULATION_BASED
-		//#define CONFIG_SAMPLE_COUNT_POLICY  SAMPLE_COUNT_POLICY_DISABLED
-
-		#define CONFIG_BILATERAL_PRESET BILATERAL_PRESET_REFLECTIONS
-
-		#define CONFIG_INPUT_TEXTURE_COUNT 2
-		#define CONFIG_OUTPUT_TEXTURE_COUNT 2
-
 	#else
 		#error Unexpected stage.
 	#endif
@@ -534,6 +497,7 @@
 //------------------------------------------------------- PARAMETERS
 
 uint MaxSampleCount;
+uint PreviousCumulativeMaxSampleCount;
 uint UpscaleFactor;
 
 #if !CONFIG_UPSCALE && CONFIG_CUSTOM_SPREAD_FACTOR
@@ -700,32 +664,6 @@ void MainCS(
 		RequestedSampleCount = 1;
 	#elif CONFIG_SAMPLE_COUNT_POLICY == SAMPLE_COUNT_POLICY_DISABLED
 		// NOP
-	#elif CONFIG_SAMPLE_COUNT_POLICY == SAMPLE_COUNT_POLICY_METADATA_BASED
-	{
-		float RefBluringRadius = ComputeInfinityBluringRadius(RefSceneMetadata) * View.ViewSizeAndInvSize.x * rcp(KernelSpreadFactor);
-		RequestedSampleCount = PI * RefBluringRadius * RefBluringRadius;
-	}
-	#elif CONFIG_SAMPLE_COUNT_POLICY == SAMPLE_COUNT_POLICY_HIT_DISTANCE_BASED
-	{
-		#if CONFIG_SIGNAL_BATCH_SIZE != 1
-			#error Unable to support more than one signal.
-		#endif
-		float RefBluringRadius = ComputeBluringRadiusOfNormalizedSample(RefFrequencies, RefSceneMetadata) * View.ViewSizeAndInvSize.x * rcp(KernelSpreadFactor);
-		RequestedSampleCount = PI * RefBluringRadius * RefBluringRadius;
-	}
-	#elif CONFIG_SAMPLE_COUNT_POLICY == SAMPLE_COUNT_POLICY_ACCUMULATION_BASED
-	{
-		#if CONFIG_SIGNAL_BATCH_SIZE != 1
-			#error Unable to support more than one signal.
-		#endif
-		// TOOD: use ComputeBluringRadiusOfNormalizedSample instead.
-		float RefBluringRadius = ComputeInfinityBluringRadius(RefSceneMetadata) * View.ViewSizeAndInvSize.x * rcp(KernelSpreadFactor);
-		float IdealSampleCount = PI * RefBluringRadius * RefBluringRadius;
-
-		IdealSampleCount = min(IdealSampleCount, TARGETED_SAMPLE_COUNT);
-
-		RequestedSampleCount = IdealSampleCount / clamp(RefSamples.Array[0].SampleCount, 1, MaxSampleCount);
-	}
 	#elif CONFIG_SAMPLE_COUNT_POLICY == SAMPLE_COUNT_POLICY_SAMPLE_ACCUMULATION_BASED
 	{
 		#if CONFIG_SIGNAL_BATCH_SIZE != 1
@@ -956,25 +894,50 @@ void MainCS(
 					/* out */ InifinityMajorViewportRadius,
 					/* out */ InifinityMinorViewportRadius);
 
-				float OutOfFocus = ComputeBluringOutOfFocusOfNormalizedSample(RefFrequencies.Array[0], RefSceneMetadata);
+				float ConfusionFactor = saturate(RefFrequencies.Array[0].ConfusionFactor);
+				float AspectRatio = InifinityMinorViewportRadius / InifinityMajorViewportRadius;
 
-				KernelConfig.MajorAxis = NormalizedScreenMajorAxis * float2(1, -1);
-				KernelConfig.MajorPixelRadius = InifinityMajorViewportRadius * OutOfFocus * View.ViewSizeAndInvSize.x;
+				float PreviousMaxPixelDiameter = sqrt(rcp(TargetSamplePerPixel) * PreviousCumulativeMaxSampleCount / AspectRatio);
+				float MaxPixelDiameter = sqrt(rcp(TargetSamplePerPixel) * MaxSampleCount * PreviousCumulativeMaxSampleCount / AspectRatio);
 
-				float AspectRatio = InifinityMinorViewportRadius / InifinityMajorViewportRadius;
+				KernelConfig.MajorAxis = NormalizedScreenMajorAxis * float2(1, -1);
+				KernelConfig.MajorPixelRadius = InifinityMajorViewportRadius * ConfusionFactor * View.ViewSizeAndInvSize.x - PreviousMaxPixelDiameter;
 
-				float MaxPixelDiameter = sqrt(rcp(TargetSamplePerPixel) * MaxSampleCount / AspectRatio);
 				float MaxPixelRadius = 0.5 * MaxPixelDiameter;
 
-
 				KernelConfig.MajorPixelRadius = clamp(KernelConfig.MajorPixelRadius, 0, MaxPixelRadius);
 				KernelConfig.MinorPixelRadius = AspectRatio * KernelConfig.MajorPixelRadius;
 
 				// *4 to multiply from radii area to diameters area.
 				float ConvolutionArea = 4.0 * max(KernelConfig.MajorPixelRadius, MinimalPixelRadius) * max(KernelConfig.MinorPixelRadius, MinimalPixelRadius);
 
-				KernelConfig.SampleCount = clamp(ConvolutionArea * TargetSamplePerPixel, 0, MaxSampleCount);
+				KernelConfig.SampleCount = clamp(ConvolutionArea * TargetSamplePerPixel * rcp(PreviousCumulativeMaxSampleCount), 0, MaxSampleCount);
 			
+				#if 0
+				{
+					DebugOutput[DispatchThreadId] = float4(
+						KernelConfig.SampleCount,
+						KernelConfig.MajorPixelRadius,
+						KernelConfig.MinorPixelRadius,
+						ConfusionFactor);
+				}
+				#elif 0
+				{
+					// DebugOutput[DispatchThreadId] = float4(
+					// 	KernelConfig.SampleCount,
+					// 	InifinityMajorViewportRadius,
+					// 	InifinityMinorViewportRadius,
+					// 	AspectRatio);
+				}
+				#elif 0
+				{
+					DebugOutput[DispatchThreadId] = float4(
+						//GetWorldNormal(RefSceneMetadata) * 0.5 + 0.5,
+						abs(GetTranslatedWorldPosition(RefSceneMetadata) * 0.001),
+						RefSceneMetadata.WorldDepth);
+				}
+				#endif
+
 				// DebugOutput[DispatchThreadId] = float4(
 				// 	KernelConfig.SampleCount,
 				// 	KernelConfig.MajorPixelRadius,
@@ -988,15 +951,26 @@ void MainCS(
 		#endif // CONFIG_SAMPLE_SET == SAMPLE_SET_DIRECTIONAL_*
 		
 		FSSDCompressedSignalAccumulatorArray CompressedAccumulators = CompressAccumulatorArray(UncompressedAccumulators, CONFIG_ACCUMULATOR_VGPR_COMPRESSION);
-
-		AccumulateKernel(
-			KernelConfig,
-			Signal_Textures_0,
-			Signal_Textures_1,
-			Signal_Textures_2,
-			Signal_Textures_3,
-			/* inout */ UncompressedAccumulators,
-			/* inout */ CompressedAccumulators);
+		
+		#if CONFIG_SIGNAL_PROCESSING == SIGNAL_PROCESSING_VIRTUAL_SHADOW_MAP_MASK
+		bool bRefHasSamples = false;
+		for (uint SignalMultiplexId = 0; SignalMultiplexId < SIGNAL_ARRAY_SIZE; SignalMultiplexId++)
+		{
+			bRefHasSamples = bRefHasSamples || (RefSamples.Array[SignalMultiplexId].SampleCount > 0);
+		}
+		BRANCH
+		if (bRefHasSamples)
+		#endif
+		{
+			AccumulateKernel(
+				KernelConfig,
+				Signal_Textures_0,
+				Signal_Textures_1,
+				Signal_Textures_2,
+				Signal_Textures_3,
+				/* inout */ UncompressedAccumulators,
+				/* inout */ CompressedAccumulators);
+		}
 
 		// When doing history rejection pre convolution, could still have no information found with the 3x3 + kernel,
 		// therefore dynamically complete to form an entire 3x3 convolution.
diff --git a/Engine/Shaders/Private/ScreenSpaceDenoise/SSDTemporalAccumulation.usf b/Engine/Shaders/Private/ScreenSpaceDenoise/SSDTemporalAccumulation.usf
--- a/Engine/Shaders/Private/ScreenSpaceDenoise/SSDTemporalAccumulation.usf
+++ b/Engine/Shaders/Private/ScreenSpaceDenoise/SSDTemporalAccumulation.usf
@@ -9,12 +9,6 @@
 	#define REJECTION_INPUT_MODE_1UNNORMALIZED 0
 	#define REJECTION_INPUT_MODE_2PRETRANSFORMED_MOMMENTS 1
 
-/** What set of sample should be used for history rejection. */
-	#define REJECTION_SAMPLE_SET_BOX 0
-	// [ Sobek 2018, "Real-time Reflections_in Mafia III" ]
-	#define REJECTION_SAMPLE_SET_SOBEK2018 1
-
-
 /** How to perform history rejection. */
 	#define HISTORY_REJECTION_DISABLED 0
 	#define HISTORY_REJECTION_MINMAX_BOUNDARIES 1
@@ -96,22 +90,12 @@
 	#define CONFIG_BILATERAL_PRESET BILATERAL_PRESET_REFLECTIONS_TAA
 	#define CONFIG_HISTORY_BILATERAL_PRESET BILATERAL_PRESET_REFLECTIONS_TAA
 
+	#define CONFIG_REJECTION_SAMPLE_SET SAMPLE_SET_3X3_PLUS
+
 	// Use variance based rejection.
 	#define CONFIG_HISTORY_REJECTION             HISTORY_REJECTION_VAR_BOUNDARIES
 	#define CONFIG_HISTORY_REJECTION_COLOR_SPACE REFLECTIONS_REJECTION_COLOR_SPACE
 
-	// Use pre transformed rejection buffer that contains pre transformed momment 1 & 2.
-	#if 1
-		#define CONFIG_USE_REJECTION_BUFFER               1
-		#define CONFIG_SIGNAL_HISTORY_REJECTION_LAYOUT    SIGNAL_BUFFER_LAYOUT_REFLECTIONS_REJECTION
-		
-		#define CONFIG_REJECTION_TEXTURE_COUNT 3 
-
-		#define CONFIG_REJECTION_BUFFER_COLOR_SPACE       REFLECTIONS_REJECTION_COLOR_SPACE
-		#define CONFIG_REJECTION_INPUT_MODE               REJECTION_INPUT_MODE_2PRETRANSFORMED_MOMMENTS
-		#define CONFIG_REJECTION_SAMPLE_SET               REJECTION_SAMPLE_SET_SOBEK2018
-	#endif
-
 #elif CONFIG_SIGNAL_PROCESSING == SIGNAL_PROCESSING_AO
 	#define MAX_SIGNAL_BATCH_SIZE 1
 	#define SIGNAL_ARRAY_SIZE 2 // First and Second moment in rejection pre convolution.
@@ -135,7 +119,7 @@
 		#define CONFIG_REJECTION_TEXTURE_COUNT            1
 
 		#define CONFIG_REJECTION_INPUT_MODE               REJECTION_INPUT_MODE_2PRETRANSFORMED_MOMMENTS
-		#define CONFIG_REJECTION_SAMPLE_SET               REJECTION_SAMPLE_SET_SOBEK2018
+		#define CONFIG_REJECTION_SAMPLE_SET               SAMPLE_SET_3X3_SOBEK2018
 	#endif
 
 #elif CONFIG_SIGNAL_PROCESSING == SIGNAL_PROCESSING_DIFFUSE_INDIRECT_AND_AO
@@ -202,6 +186,8 @@
 	// Reject using variance to smooth highlights further on history rejections.
 	#define CONFIG_HISTORY_REJECTION HISTORY_REJECTION_VAR_BOUNDARIES
 
+	#define CONFIG_REJECTION_SAMPLE_SET SAMPLE_SET_3X3_PLUS
+
 	// Use position and normal based rejection ideal for diffuse.
 	#define CONFIG_BILATERAL_PRESET BILATERAL_PRESET_DIFFUSE
 	#define CONFIG_HISTORY_BILATERAL_PRESET BILATERAL_PRESET_DIFFUSE
@@ -236,7 +222,7 @@
 
 /** Sample set to use for history rejection. */
 #ifndef CONFIG_REJECTION_SAMPLE_SET
-	#define CONFIG_REJECTION_SAMPLE_SET REJECTION_SAMPLE_SET_BOX
+	#define CONFIG_REJECTION_SAMPLE_SET SAMPLE_SET_NXN
 #endif
 
 /** The color space to use for history rejection. */
@@ -508,9 +494,19 @@ void TemporallyAccumulate(
 
 		#if CONFIG_SIGNAL_PROCESSING == SIGNAL_PROCESSING_REFLECTIONS
 		{
-			float SignalPixelBluringRadius = ComputeInfinityBluringRadius(RefSceneMetadata) * View.ViewSizeAndInvSize.x;
+			float2 NormalizedScreenMajorAxis;
+			float InifinityMajorViewportRadius;
+			float InifinityMinorViewportRadius;
+			ProjectSpecularLobeToScreenSpace(
+				RefSceneMetadata,
+				/* out */ NormalizedScreenMajorAxis,
+				/* out */ InifinityMajorViewportRadius,
+				/* out */ InifinityMinorViewportRadius);
+
+			InifinityMajorViewportRadius *= View.ViewSizeAndInvSize.x;
+			InifinityMinorViewportRadius *= View.ViewSizeAndInvSize.x;
 
-			TargetedSampleCount = PI * SignalPixelBluringRadius * SignalPixelBluringRadius;
+			TargetedSampleCount = PI * InifinityMajorViewportRadius * InifinityMinorViewportRadius;
 
 			// TODO(Denoiser): try to make this better? * 0.125 to avoid banding caused by floating point precision.
 			TargetedSampleCount = clamp(TargetedSampleCount, 1, TARGETED_SAMPLE_COUNT);
@@ -690,23 +686,16 @@ void TemporallyAccumulate(
 				KernelConfig.BilateralDistanceComputation = SIGNAL_WORLD_FREQUENCY_PRECOMPUTED_BLURING_RADIUS;
 			KernelConfig.WorldBluringDistanceMultiplier = CONFIG_BILATERAL_DISTANCE_MULTIPLIER;
 			
-			#if CONFIG_REJECTION_SAMPLE_SET == REJECTION_SAMPLE_SET_SOBEK2018
-			{
-				KernelConfig.SampleSet = SAMPLE_SET_3X3_SOBEK2018;
-				//KernelConfig.SampleSet = SAMPLE_SET_NXN;
-				//KernelConfig.BoxKernelRadius = 3;
-				KernelConfig.bUnroll = true;
-			}
-			#elif CONFIG_SIGNAL_PROCESSING == SIGNAL_PROCESSING_SSGI
+			#if CONFIG_REJECTION_SAMPLE_SET == SAMPLE_SET_NXN
 			{
-				KernelConfig.SampleSet = SAMPLE_SET_3X3_PLUS;
-				KernelConfig.bUnroll = true;
+				KernelConfig.SampleSet = SAMPLE_SET_NXN;
+				KernelConfig.BoxKernelRadius = 3;
+				KernelConfig.bUnroll = false;
 			}
 			#else
 			{
-				KernelConfig.SampleSet = CONFIG_SIGNAL_PROCESSING == SIGNAL_PROCESSING_SHADOW_VISIBILITY_MASK ? SAMPLE_SET_3X3 : SAMPLE_SET_NXN;
-				KernelConfig.BoxKernelRadius = 3;
-				KernelConfig.bUnroll = false;
+				KernelConfig.SampleSet = CONFIG_REJECTION_SAMPLE_SET;
+				KernelConfig.bUnroll = true;
 			}
 			#endif
 
diff --git a/Engine/Source/Runtime/Renderer/Private/ScreenSpaceDenoise.cpp b/Engine/Source/Runtime/Renderer/Private/ScreenSpaceDenoise.cpp
--- a/Engine/Source/Runtime/Renderer/Private/ScreenSpaceDenoise.cpp
+++ b/Engine/Source/Runtime/Renderer/Private/ScreenSpaceDenoise.cpp
@@ -50,8 +50,13 @@ static TAutoConsoleVariable<int32> CVarUseReflectionDenoiser(
 	ECVF_RenderThreadSafe);
 
 static TAutoConsoleVariable<int32> CVarReflectionReconstructionSampleCount(
-	TEXT("r.Reflections.Denoiser.ReconstructionSamples"), 16,
-	TEXT("Maximum number of samples for the reconstruction pass (default = 16)."),
+	TEXT("r.Reflections.Denoiser.ReconstructionSamples"), 8,
+	TEXT("Maximum number of samples for the reconstruction pass (default = 8)."),
+	ECVF_RenderThreadSafe);
+
+static TAutoConsoleVariable<int32> CVarReflectionPreConvolutionCount(
+	TEXT("r.Reflections.Denoiser.PreConvolution"), 1,
+	TEXT("Number of pre-convolution passes (default = 1)."),
 	ECVF_RenderThreadSafe);
 
 static TAutoConsoleVariable<int32> CVarReflectionTemporalAccumulation(
@@ -59,11 +64,6 @@ static TAutoConsoleVariable<int32> CVarReflectionTemporalAccumulation(
 	TEXT("Accumulates the samples over multiple frames."),
 	ECVF_RenderThreadSafe);
 
-static TAutoConsoleVariable<int32> CVarReflectionHistoryConvolutionSampleCount(
-	TEXT("r.Reflections.Denoiser.HistoryConvolution.SampleCount"), 1,
-	TEXT("Number of samples to use for history post filter (default = 1)."),
-	ECVF_RenderThreadSafe);
-
 static TAutoConsoleVariable<int32> CVarAOReconstructionSampleCount(
 	TEXT("r.AmbientOcclusion.Denoiser.ReconstructionSamples"), 16,
 	TEXT("Maximum number of samples for the reconstruction pass (default = 16)."),
@@ -225,6 +225,7 @@ static bool SignalUsesPreConvolution(ESignalProcessing SignalProcessing)
 {
 	return
 		SignalProcessing == ESignalProcessing::ShadowVisibilityMask ||
+		SignalProcessing == ESignalProcessing::Reflections ||
 		SignalProcessing == ESignalProcessing::DiffuseAndAmbientOcclusion;
 }
 
@@ -233,7 +234,7 @@ static bool SignalUsesRejectionPreConvolution(ESignalProcessing SignalProcessing
 {
 	return (
 		//SignalProcessing == ESignalProcessing::ShadowVisibilityMask ||
-		SignalProcessing == ESignalProcessing::Reflections ||
+		//SignalProcessing == ESignalProcessing::Reflections ||
 		SignalProcessing == ESignalProcessing::AmbientOcclusion);
 }
 
@@ -242,7 +243,6 @@ static bool SignalUsesPostConvolution(ESignalProcessing SignalProcessing)
 {
 	return (
 		SignalProcessing == ESignalProcessing::ShadowVisibilityMask ||
-		SignalProcessing == ESignalProcessing::Reflections ||
 		SignalProcessing == ESignalProcessing::AmbientOcclusion ||
 		SignalProcessing == ESignalProcessing::DiffuseAndAmbientOcclusion);
 }
@@ -420,8 +420,8 @@ const TCHAR* const kPreConvolutionResourceNames[] = {
 	nullptr,
 
 	// Reflections
-	nullptr,
-	nullptr,
+	TEXT("ReflectionsPreConvolution0"),
+	TEXT("ReflectionsPreConvolution1"),
 	nullptr,
 	nullptr,
 
@@ -988,6 +988,7 @@ class FSSDSpatialAccumulationCS : public FGlobalShader
 		SHADER_PARAMETER_ARRAY(FVector4, InputBufferUVMinMax, [IScreenSpaceDenoiser::kMaxBatchSize])
 
 		SHADER_PARAMETER(uint32, MaxSampleCount)
+		SHADER_PARAMETER(uint32, PreviousCumulativeMaxSampleCount)
 		SHADER_PARAMETER(int32, UpscaleFactor)
 		SHADER_PARAMETER(float, KernelSpreadFactor)
 		SHADER_PARAMETER(float, HarmonicPeriode)
@@ -1268,7 +1269,7 @@ static void DenoiseSignalAtConstantPixelDensity(
 		else if (Settings.SignalProcessing == ESignalProcessing::Reflections)
 		{
 			ReconstructionDescs[0].Format = HistoryDescs[0].Format = PF_FloatRGBA;
-			ReconstructionDescs[1].Format = HistoryDescs[1].Format = PF_R16F;
+			ReconstructionDescs[1].Format = HistoryDescs[1].Format = PF_G16R16F;
 			ReconstructionTextureCount = HistoryTextureCountPerSignal = 2;
 			bHasReconstructionLayoutDifferentFromHistory = false;
 		}
@@ -1541,7 +1542,8 @@ static void DenoiseSignalAtConstantPixelDensity(
 				float(SignalScissor.Max.Y - 0.5f) / float(BufferExtent.Y));
 		}
 
-		PassParameters->MaxSampleCount = FMath::Clamp(Settings.ReconstructionSamples, 1, kStackowiakMaxSampleCountPerSet);
+		PassParameters->MaxSampleCount = Settings.ReconstructionSamples;
+		PassParameters->PreviousCumulativeMaxSampleCount = 1;
 		PassParameters->UpscaleFactor = int32(Settings.DenoisingResolutionFraction / Settings.InputResolutionFraction);
 		PassParameters->HarmonicPeriode = Settings.HarmonicPeriode;
 		PassParameters->CommonParameters = CommonParameters;
@@ -1588,6 +1590,9 @@ static void DenoiseSignalAtConstantPixelDensity(
 		FSSDSpatialAccumulationCS::FParameters* PassParameters = GraphBuilder.AllocParameters<FSSDSpatialAccumulationCS::FParameters>();
 		PassParameters->CommonParameters = CommonParameters;
 		PassParameters->ConvolutionMetaData = ConvolutionMetaData;
+
+		PassParameters->MaxSampleCount = Settings.ReconstructionSamples;
+		PassParameters->PreviousCumulativeMaxSampleCount = FMath::Pow(PassParameters->MaxSampleCount, 1 + PreConvolutionId);
 		PassParameters->KernelSpreadFactor = 8 * (1 << PreConvolutionId);
 		PassParameters->SignalInput = SignalHistory;
 		PassParameters->SignalOutput = CreateMultiplexedUAVs(GraphBuilder, NewSignalOutput);
@@ -1603,7 +1608,10 @@ static void DenoiseSignalAtConstantPixelDensity(
 		TShaderMapRef<FSSDSpatialAccumulationCS> ComputeShader(View.ShaderMap, PermutationVector);
 		FComputeShaderUtils::AddPass(
 			GraphBuilder,
-			RDG_EVENT_NAME("SSD PreConvolution(MaxSamples=7 Spread=%d)", PassParameters->KernelSpreadFactor),
+			RDG_EVENT_NAME(
+				"SSD PreConvolution(MaxSamples=%d Spread=%f)", 
+				PassParameters->MaxSampleCount,
+				PassParameters->KernelSpreadFactor),
 			ComputeShader,
 			PassParameters,
 			FComputeShaderUtils::GetGroupCount(Viewport.Size(), FSSDSpatialAccumulationCS::kGroupSize));
@@ -1646,13 +1654,6 @@ static void DenoiseSignalAtConstantPixelDensity(
 					}
 					RejectionTextureCount = Settings.SignalBatchSize;
 				}
-				else if (Settings.SignalProcessing == ESignalProcessing::Reflections)
-				{
-					RejectionSignalProcessingDescs[0].Format = PF_FloatRGBA;
-					RejectionSignalProcessingDescs[1].Format = PF_G16R16F;
-					RejectionSignalProcessingDescs[2].Format = PF_FloatRGBA;
-					RejectionTextureCount = 3;
-				}
 				else if (Settings.SignalProcessing == ESignalProcessing::AmbientOcclusion)
 				{
 					RejectionSignalProcessingDescs[0].Format = PF_FloatRGBA;
@@ -2046,7 +2047,7 @@ public:
 		FSSDConstantPixelDensitySettings Settings;
 		Settings.SignalProcessing = ESignalProcessing::ShadowVisibilityMask;
 		Settings.InputResolutionFraction = 1.0f;
-		Settings.ReconstructionSamples = CVarShadowReconstructionSampleCount.GetValueOnRenderThread();
+		Settings.ReconstructionSamples = FMath::Clamp(CVarShadowReconstructionSampleCount.GetValueOnRenderThread(), 1, kStackowiakMaxSampleCountPerSet);
 		Settings.PreConvolutionCount = CVarShadowPreConvolutionCount.GetValueOnRenderThread();
 		Settings.bUseTemporalAccumulation = CVarShadowTemporalAccumulation.GetValueOnRenderThread() != 0;
 		Settings.HistoryConvolutionSampleCount = CVarShadowHistoryConvolutionSampleCount.GetValueOnRenderThread();
@@ -2334,8 +2335,8 @@ public:
 		Settings.SignalProcessing = ESignalProcessing::Reflections;
 		Settings.InputResolutionFraction = RayTracingConfig.ResolutionFraction;
 		Settings.ReconstructionSamples = CVarReflectionReconstructionSampleCount.GetValueOnRenderThread();
+		Settings.PreConvolutionCount = CVarReflectionPreConvolutionCount.GetValueOnRenderThread();
 		Settings.bUseTemporalAccumulation = CVarReflectionTemporalAccumulation.GetValueOnRenderThread() != 0;
-		Settings.HistoryConvolutionSampleCount = CVarReflectionHistoryConvolutionSampleCount.GetValueOnRenderThread();
 		Settings.MaxInputSPP = RayTracingConfig.RayCountPerPixel;
 
 		TStaticArray<FScreenSpaceDenoiserHistory*, IScreenSpaceDenoiser::kMaxBatchSize> PrevHistories;
@@ -2383,8 +2384,8 @@ public:
 		Settings.SignalProcessing = ESignalProcessing::Reflections; // TODO: water reflection to denoise only water pixels
 		Settings.InputResolutionFraction = RayTracingConfig.ResolutionFraction;
 		Settings.ReconstructionSamples = CVarReflectionReconstructionSampleCount.GetValueOnRenderThread();
+		Settings.PreConvolutionCount = CVarReflectionPreConvolutionCount.GetValueOnRenderThread();
 		Settings.bUseTemporalAccumulation = CVarReflectionTemporalAccumulation.GetValueOnRenderThread() != 0;
-		Settings.HistoryConvolutionSampleCount = CVarReflectionHistoryConvolutionSampleCount.GetValueOnRenderThread();
 		Settings.MaxInputSPP = RayTracingConfig.RayCountPerPixel;
 
 		TStaticArray<FScreenSpaceDenoiserHistory*, IScreenSpaceDenoiser::kMaxBatchSize> PrevHistories;
@@ -2427,7 +2428,7 @@ public:
 		Settings.FullResViewport = View.ViewRect;
 		Settings.SignalProcessing = ESignalProcessing::AmbientOcclusion;
 		Settings.InputResolutionFraction = RayTracingConfig.ResolutionFraction;
-		Settings.ReconstructionSamples = CVarAOReconstructionSampleCount.GetValueOnRenderThread();
+		Settings.ReconstructionSamples = FMath::Clamp(CVarAOReconstructionSampleCount.GetValueOnRenderThread(), 1, kStackowiakMaxSampleCountPerSet);
 		Settings.bUseTemporalAccumulation = CVarAOTemporalAccumulation.GetValueOnRenderThread() != 0;
 		Settings.HistoryConvolutionSampleCount = CVarAOHistoryConvolutionSampleCount.GetValueOnRenderThread();
 		Settings.HistoryConvolutionKernelSpreadFactor = CVarAOHistoryConvolutionKernelSpreadFactor.GetValueOnRenderThread();
@@ -2472,7 +2473,7 @@ public:
 		Settings.FullResViewport = View.ViewRect;
 		Settings.SignalProcessing = ESignalProcessing::DiffuseAndAmbientOcclusion;
 		Settings.InputResolutionFraction = Config.ResolutionFraction;
-		Settings.ReconstructionSamples = CVarGIReconstructionSampleCount.GetValueOnRenderThread();
+		Settings.ReconstructionSamples = FMath::Clamp(CVarGIReconstructionSampleCount.GetValueOnRenderThread(), 1, kStackowiakMaxSampleCountPerSet);
 		Settings.PreConvolutionCount = CVarGIPreConvolutionCount.GetValueOnRenderThread();
 		Settings.bUseTemporalAccumulation = CVarGITemporalAccumulation.GetValueOnRenderThread() != 0;
 		Settings.HistoryConvolutionSampleCount = CVarGIHistoryConvolutionSampleCount.GetValueOnRenderThread();
@@ -2518,7 +2519,7 @@ public:
 		Settings.FullResViewport = View.ViewRect;
 		Settings.SignalProcessing = ESignalProcessing::DiffuseAndAmbientOcclusion;
 		Settings.InputResolutionFraction = Config.ResolutionFraction;
-		Settings.ReconstructionSamples = CVarGIReconstructionSampleCount.GetValueOnRenderThread();
+		Settings.ReconstructionSamples = FMath::Clamp(CVarGIReconstructionSampleCount.GetValueOnRenderThread(), 1, kStackowiakMaxSampleCountPerSet);
 		Settings.PreConvolutionCount = CVarGIPreConvolutionCount.GetValueOnRenderThread();
 		Settings.bUseTemporalAccumulation = CVarGITemporalAccumulation.GetValueOnRenderThread() != 0;
 		Settings.HistoryConvolutionSampleCount = CVarGIHistoryConvolutionSampleCount.GetValueOnRenderThread();
@@ -2564,7 +2565,7 @@ public:
 		Settings.FullResViewport = View.ViewRect;
 		Settings.SignalProcessing = ESignalProcessing::DiffuseAndAmbientOcclusion;
 		Settings.InputResolutionFraction = Config.ResolutionFraction;
-		Settings.ReconstructionSamples = CVarGIReconstructionSampleCount.GetValueOnRenderThread();
+		Settings.ReconstructionSamples = FMath::Clamp(CVarGIReconstructionSampleCount.GetValueOnRenderThread(), 1, kStackowiakMaxSampleCountPerSet);
 		Settings.PreConvolutionCount = CVarGIPreConvolutionCount.GetValueOnRenderThread();
 		Settings.bUseTemporalAccumulation = CVarGITemporalAccumulation.GetValueOnRenderThread() != 0;
 		Settings.HistoryConvolutionSampleCount = CVarGIHistoryConvolutionSampleCount.GetValueOnRenderThread();
