Fix hair strands velocity rendering (crash).

Velocity texture format change in raytracing is enabled/supported.

diff --git a/Engine/Shaders/Private/HairStrands/HairStrandsVelocity.usf b/Engine/Shaders/Private/HairStrands/HairStrandsVelocity.usf
--- a/Engine/Shaders/Private/HairStrands/HairStrandsVelocity.usf
+++ b/Engine/Shaders/Private/HairStrands/HairStrandsVelocity.usf
@@ -32,7 +32,17 @@ StructuredBuffer<FNodeVis> NodeVis;
 #define VELOCITY_TYPE_CLOSEST 2
 #define VELOCITY_TYPE_MAX 3
 #if PERMUTATION_VELOCITY != VELOCITY_TYPE_NONE
+  #if PERMUTATION_OUTPUT_FORMAT == 0
 	RWTexture2D<float2> OutVelocityTexture;
+  #else
+	RWTexture2D<float4> OutVelocityTexture;
+  #endif
+#endif
+
+#if PERMUTATION_OUTPUT_FORMAT == 0
+float2 ToOutput(float2 In) { return In; }
+#else
+float4 ToOutput(float2 In) { return float4(In, 0, 0); }
 #endif
 
 [numthreads(TILE_PIXEL_SIZE_X, TILE_PIXEL_SIZE_Y, 1)]
@@ -87,15 +97,15 @@ void MainCS(uint GroupIndex : SV_GroupIndex, uint3 DispatchThreadId : SV_Dispatc
 		}
 
 		#if PERMUTATION_VELOCITY == VELOCITY_TYPE_AVG
-		OutVelocityTexture[PixelCoord] = EncodeVelocityToTexture(float3(AverageVelocity / max(NodeDesc.Count, 1), 0.0)).xy;
+		OutVelocityTexture[PixelCoord] = ToOutput(EncodeVelocityToTexture(float3(AverageVelocity / max(NodeDesc.Count, 1), 0.0)).xy);
 		#endif		
 
 		#if PERMUTATION_VELOCITY == VELOCITY_TYPE_CLOSEST
-		OutVelocityTexture[PixelCoord] = ClosestEncodedVelocity;
+		OutVelocityTexture[PixelCoord] = ToOutput(ClosestEncodedVelocity);
 		#endif		
 
 		#if PERMUTATION_VELOCITY == VELOCITY_TYPE_MAX
-		OutVelocityTexture[PixelCoord] = MaxEncodedVelocity;
+		OutVelocityTexture[PixelCoord] = ToOutput(MaxEncodedVelocity);
 		#endif		
 	}
 }
diff --git a/Engine/Source/Runtime/Renderer/Private/HairStrands/HairStrandsVisibility.cpp b/Engine/Source/Runtime/Renderer/Private/HairStrands/HairStrandsVisibility.cpp
--- a/Engine/Source/Runtime/Renderer/Private/HairStrands/HairStrandsVisibility.cpp
+++ b/Engine/Source/Runtime/Renderer/Private/HairStrands/HairStrandsVisibility.cpp
@@ -708,7 +708,8 @@ class FHairVelocityCS: public FGlobalShader
 
 	class FVendor : SHADER_PERMUTATION_INT("PERMUTATION_VENDOR", HairVisibilityVendorCount);
 	class FVelocity : SHADER_PERMUTATION_INT("PERMUTATION_VELOCITY", 4);
-	using FPermutationDomain = TShaderPermutationDomain<FVendor, FVelocity>;
+	class FOuputFormat : SHADER_PERMUTATION_INT("PERMUTATION_OUTPUT_FORMAT", 2);
+	using FPermutationDomain = TShaderPermutationDomain<FVendor, FVelocity, FOuputFormat>;
 
 	BEGIN_SHADER_PARAMETER_STRUCT(FParameters, )
 		SHADER_PARAMETER(FIntPoint, ResolutionOffset)
@@ -743,7 +744,8 @@ static void AddHairVelocityPass(
 	if (!bWriteOutVelocity)
 		return;
 
-	check(OutVelocityTexture->Desc.Format == FMaterialPassOutput::VelocityFormat);
+	check(OutVelocityTexture->Desc.Format == PF_G16R16 || OutVelocityTexture->Desc.Format == PF_A16B16G16R16);
+	const bool bTwoChannelsOutput = OutVelocityTexture->Desc.Format == PF_G16R16;
 
 	FSceneRenderTargets& SceneContext = FSceneRenderTargets::Get(GraphBuilder.RHICmdList);
 	FSceneTexturesUniformParameters SceneTextures;
@@ -752,6 +754,7 @@ static void AddHairVelocityPass(
 	FHairVelocityCS::FPermutationDomain PermutationVector;
 	PermutationVector.Set<FHairVelocityCS::FVendor>(GetVendor());
 	PermutationVector.Set<FHairVelocityCS::FVelocity>(bWriteOutVelocity ? FMath::Clamp(GHairVelocityType + 1, 0, 3) : 0);
+	PermutationVector.Set<FHairVelocityCS::FOuputFormat>(bTwoChannelsOutput ? 0 : 1);
 
 	FHairVelocityCS::FParameters* PassParameters = GraphBuilder.AllocParameters<FHairVelocityCS::FParameters>();
 	PassParameters->SceneTexturesStruct = CreateUniformBufferImmediate(SceneTextures, EUniformBufferUsage::UniformBuffer_SingleDraw);
