Moved RVT TextureCopy operations into RDG graph

diff --git a/Engine/Source/Runtime/Renderer/Private/VT/RuntimeVirtualTextureRender.cpp b/Engine/Source/Runtime/Renderer/Private/VT/RuntimeVirtualTextureRender.cpp
--- a/Engine/Source/Runtime/Renderer/Private/VT/RuntimeVirtualTextureRender.cpp
+++ b/Engine/Source/Runtime/Renderer/Private/VT/RuntimeVirtualTextureRender.cpp
@@ -709,6 +709,34 @@ namespace RuntimeVirtualTexture
 		}
 	}
 
+	BEGIN_SHADER_PARAMETER_STRUCT(FCopyToOutputParameters, )
+		SHADER_PARAMETER_RDG_TEXTURE(, Input)
+	END_SHADER_PARAMETER_STRUCT()
+
+	/** Set up the copy to final output physical texture. */
+	void AddCopyToOutputPass(FRDGBuilder& GraphBuilder, FRDGTextureRef InputTexture, FRHITexture2D* OutputTexture, FBox2D const& DestBox)
+	{
+		FRHICopyTextureInfo CopyInfo;
+		CopyInfo.Size = InputTexture->Desc.GetSize();
+		CopyInfo.DestPosition = FIntVector(DestBox.Min.X, DestBox.Min.Y, 0);
+
+		FCopyToOutputParameters* Parameters = GraphBuilder.AllocParameters<FCopyToOutputParameters>();
+		Parameters->Input = InputTexture;
+
+		GraphBuilder.AddPass(
+			RDG_EVENT_NAME("VirtualTextureCopyToOutput"),
+			Parameters,
+			ERDGPassFlags::Copy,
+			[InputTexture, OutputTexture, CopyInfo](FRHICommandList& RHICmdList)
+			{
+				InputTexture->MarkResourceAsUsed();
+
+				RHICmdList.TransitionResource(EResourceTransitionAccess::EWritable, OutputTexture);
+				RHICmdList.CopyTexture(InputTexture->GetRHI(), OutputTexture, CopyInfo);
+				RHICmdList.TransitionResource(EResourceTransitionAccess::EReadable, OutputTexture);
+			});
+	}
+
 
 	/** Structure to localize the setup of our render graph based on the virtual texture setup. */
 	struct FRenderGraphSetup
@@ -1008,72 +1036,23 @@ namespace RuntimeVirtualTexture
 			}
 		}
 
-		// Set up the output to capture
-		TRefCountPtr<IPooledRenderTarget> GraphOutputTexture0;
-		FIntVector GraphOutputSize0;
+		// Copy to Output for each output texture
 		if (GraphSetup.OutputAlias0 != nullptr)
 		{
-			GraphBuilder.QueueTextureExtraction(GraphSetup.OutputAlias0, &GraphOutputTexture0);
-			GraphOutputSize0 = GraphSetup.OutputAlias0->Desc.GetSize();
+			AddCopyToOutputPass(GraphBuilder, GraphSetup.OutputAlias0, OutputTexture0, DestBox0);
 		}
-
-		TRefCountPtr<IPooledRenderTarget> GraphOutputTexture1;
-		FIntVector GraphOutputSize1;
 		if (GraphSetup.OutputAlias1 != nullptr)
 		{
-			GraphBuilder.QueueTextureExtraction(GraphSetup.OutputAlias1, &GraphOutputTexture1);
-			GraphOutputSize1 = GraphSetup.OutputAlias1->Desc.GetSize();
+			AddCopyToOutputPass(GraphBuilder, GraphSetup.OutputAlias1, OutputTexture1, DestBox1);
 		}
-
-		TRefCountPtr<IPooledRenderTarget> GraphOutputTexture2;
-		FIntVector GraphOutputSize2;
 		if (GraphSetup.OutputAlias2 != nullptr)
 		{
-			GraphBuilder.QueueTextureExtraction(GraphSetup.OutputAlias2, &GraphOutputTexture2);
-			GraphOutputSize2 = GraphSetup.OutputAlias2->Desc.GetSize();
+			AddCopyToOutputPass(GraphBuilder, GraphSetup.OutputAlias2, OutputTexture2, DestBox2);
 		}
 
 		// Execute the graph
 		GraphBuilder.Execute();
 
-		FRHITexture* TexturesToTransition[] =
-		{
-			OutputTexture0 ? OutputTexture0->GetTexture2D() : nullptr,
-			OutputTexture1 ? OutputTexture1->GetTexture2D() : nullptr,
-			OutputTexture2 ? OutputTexture2->GetTexture2D() : nullptr
-		};
-		RHICmdList.TransitionResources(EResourceTransitionAccess::EWritable, TexturesToTransition, UE_ARRAY_COUNT(TexturesToTransition));
-
-		// Copy to final destination
-		if (GraphSetup.OutputAlias0 != nullptr && OutputTexture0 != nullptr)
-		{
-			FRHICopyTextureInfo Info;
-			Info.Size = GraphOutputSize0;
-			Info.DestPosition = FIntVector(DestBox0.Min.X, DestBox0.Min.Y, 0);
-
-			RHICmdList.CopyTexture(GraphOutputTexture0->GetRenderTargetItem().ShaderResourceTexture->GetTexture2D(), OutputTexture0->GetTexture2D(), Info);
-		}
-
-		if (GraphSetup.OutputAlias1 != nullptr && OutputTexture1 != nullptr)
-		{
-			FRHICopyTextureInfo Info;
-			Info.Size = GraphOutputSize1;
-			Info.DestPosition = FIntVector(DestBox1.Min.X, DestBox1.Min.Y, 0);
-
-			RHICmdList.CopyTexture(GraphOutputTexture1->GetRenderTargetItem().ShaderResourceTexture->GetTexture2D(), OutputTexture1->GetTexture2D(), Info);
-		}
-
-		if (GraphSetup.OutputAlias2 != nullptr && OutputTexture2 != nullptr)
-		{
-			FRHICopyTextureInfo Info;
-			Info.Size = GraphOutputSize2;
-			Info.DestPosition = FIntVector(DestBox2.Min.X, DestBox2.Min.Y, 0);
-
-			RHICmdList.CopyTexture(GraphOutputTexture2->GetRenderTargetItem().ShaderResourceTexture->GetTexture2D(), OutputTexture2->GetTexture2D(), Info);
-		}
-
-		RHICmdList.TransitionResources(EResourceTransitionAccess::EReadable, TexturesToTransition, UE_ARRAY_COUNT(TexturesToTransition));
-
 		View->CachedViewUniformShaderParameters.Reset();
 	}
 
