Adding LightingChannelMask support to path tracing and ray tracing global illumination passes.

diff --git a/Engine/Shaders/Private/PathTracing/Light/PathTracingLightCommon.ush b/Engine/Shaders/Private/PathTracing/Light/PathTracingLightCommon.ush
--- a/Engine/Shaders/Private/PathTracing/Light/PathTracingLightCommon.ush
+++ b/Engine/Shaders/Private/PathTracing/Light/PathTracingLightCommon.ush
@@ -135,4 +135,9 @@ float ComputeAttenuationFalloff(float Distance, float Attenuation)
 bool HasTransmission(int LightId)
 {
 	return SceneLightsData.Flags[LightId] & 0x01;
+}
+
+uint GetLightingChannelMask(int LightId)
+{
+	return (SceneLightsData.Flags[LightId] >> 1) & 0x7;
 }
\ No newline at end of file
diff --git a/Engine/Shaders/Private/PathTracing/Light/PathTracingLightSampling.ush b/Engine/Shaders/Private/PathTracing/Light/PathTracingLightSampling.ush
--- a/Engine/Shaders/Private/PathTracing/Light/PathTracingLightSampling.ush
+++ b/Engine/Shaders/Private/PathTracing/Light/PathTracingLightSampling.ush
@@ -52,6 +52,7 @@ void EstimateLight(
 )
 {
 	OutIrradiance = 0.0f;
+	if ((GetLightingChannelMask(LightId) & Payload.PrimitiveLightingChannelMask) == 0) return;
 
 	if (IsPointLight(LightId))
 	{
diff --git a/Engine/Source/Runtime/Renderer/Private/PathTracing.cpp b/Engine/Source/Runtime/Renderer/Private/PathTracing.cpp
--- a/Engine/Source/Runtime/Renderer/Private/PathTracing.cpp
+++ b/Engine/Source/Runtime/Renderer/Private/PathTracing.cpp
@@ -208,7 +208,9 @@ public:
 				FLightShaderParameters LightParameters;
 				Light.LightSceneInfo->Proxy->GetLightShaderParameters(LightParameters);
 				uint32 Transmission = Light.LightSceneInfo->Proxy->Transmission();
+				uint8 LightingChannelMask = Light.LightSceneInfo->Proxy->GetLightingChannelMask();
 				LightData.Flags[LightData.Count] = Transmission & 0x01;
+				LightData.Flags[LightData.Count] |= (LightingChannelMask & 0x7) << 1;
 
 				ELightComponentType LightComponentType = (ELightComponentType)Light.LightSceneInfo->Proxy->GetLightType();
 				switch (LightComponentType)
diff --git a/Engine/Source/Runtime/Renderer/Private/RayTracing/RayTracingGlobalIllumination.cpp b/Engine/Source/Runtime/Renderer/Private/RayTracing/RayTracingGlobalIllumination.cpp
--- a/Engine/Source/Runtime/Renderer/Private/RayTracing/RayTracingGlobalIllumination.cpp
+++ b/Engine/Source/Runtime/Renderer/Private/RayTracing/RayTracingGlobalIllumination.cpp
@@ -191,6 +191,8 @@ void SetupLightParameters(
 	FSkyLightSceneProxy* SkyLight = Scene.SkyLight;
 	FVector SkyLightColor = FVector(0.0f, 0.0f, 0.0f);
 	uint32 SkyLightTransmission = 0;
+	// SkyLight does not have a LightingChannelMask
+	uint8 SkyLightLightingChannelMask = 0xFF;
 	if (SkyLight && SkyLight->bAffectGlobalIllumination)
 	{
 		SkyLightColor = FVector(SkyLight->GetEffectiveLightColor());
@@ -202,7 +204,9 @@ void SetupLightParameters(
 	uint32 SkyLightIndex = 0;
 	LightParameters->Type[SkyLightIndex] = 0;
 	LightParameters->Color[SkyLightIndex] = SkyLightColor;
-	LightParameters->Flags[SkyLightIndex] = SkyLightTransmission & 0x01;
+	LightParameters->Flags[SkyLightIndex]  = SkyLightTransmission & 0x01;
+	LightParameters->Flags[SkyLightIndex] |= (SkyLightLightingChannelMask & 0x7) << 1;
+	
 	LightParameters->Count++;
 
 	uint32 MaxLightCount = FMath::Min(CVarRayTracingGlobalIlluminationMaxLightCount.GetValueOnRenderThread(), RAY_TRACING_LIGHT_COUNT_MAXIMUM);
@@ -215,8 +219,11 @@ void SetupLightParameters(
 
 		FLightShaderParameters LightShaderParameters;
 		Light.LightSceneInfo->Proxy->GetLightShaderParameters(LightShaderParameters);
+
 		uint32 Transmission = Light.LightSceneInfo->Proxy->Transmission();
-		LightParameters->Flags[LightParameters->Count] = Transmission & 0x01;
+		uint8 LightingChannelMask = Light.LightSceneInfo->Proxy->GetLightingChannelMask();
+		LightParameters->Flags[LightParameters->Count]  = Transmission & 0x01;
+		LightParameters->Flags[LightParameters->Count] |= (LightingChannelMask & 0x7) << 1;
 
 		ELightComponentType LightComponentType = (ELightComponentType)Light.LightSceneInfo->Proxy->GetLightType();
 		switch (LightComponentType)
