Fixed a crash when turning off all ray tracing effects during run-time.

diff --git a/Engine/Source/Runtime/Renderer/Private/DeferredShadingRenderer.cpp b/Engine/Source/Runtime/Renderer/Private/DeferredShadingRenderer.cpp
--- a/Engine/Source/Runtime/Renderer/Private/DeferredShadingRenderer.cpp
+++ b/Engine/Source/Runtime/Renderer/Private/DeferredShadingRenderer.cpp
@@ -1284,6 +1284,12 @@ void FDeferredShadingSceneRenderer::WaitForRayTracingScene(FRHICommandListImmedi
 	for (int32 ViewIndex = 0; ViewIndex < Views.Num(); ++ViewIndex)
 	{
 		FViewInfo& View = Views[ViewIndex];
+		if (!View.RayTracingMaterialPipeline)
+		{
+			check(View.RayTracingMaterialBindings.Num() == 0);
+			continue;
+		}
+
 		if (View.RayTracingMaterialBindings.Num())
 		{
 			FTaskGraphInterface::Get().WaitUntilTaskCompletes(View.RayTracingMaterialBindingsTask, ENamedThreads::GetRenderThread_Local());
