Fix for incorrectly rendering separate translucency when camera is under water

john.hable

diff --git a/Engine/Source/Runtime/Renderer/Private/TranslucentRendering.cpp b/Engine/Source/Runtime/Renderer/Private/TranslucentRendering.cpp
--- a/Engine/Source/Runtime/Renderer/Private/TranslucentRendering.cpp
+++ b/Engine/Source/Runtime/Renderer/Private/TranslucentRendering.cpp
@@ -330,10 +330,15 @@ class FTranslucencyPassParallelCommandListSet : public FParallelCommandListSet
 {
 	ETranslucencyPass::Type TranslucencyPass;
 
+	// Keep track of if separate translucency render state should be setup for this commandlist or not. 
+	// There are cases where we might want to render the primitives from the separate translucency buckets in the regular translucency scene buffer
+	bool bRenderInSeparateTranslucency;
+
 public:
 	FTranslucencyPassParallelCommandListSet(const FViewInfo& InView, const FSceneRenderer* InSceneRenderer, FRHICommandListImmediate& InParentCmdList, bool bInParallelExecute, bool bInCreateSceneContext, const FMeshPassProcessorRenderState& InDrawRenderState, ETranslucencyPass::Type InTranslucencyPass, bool InRenderInSeparateTranslucency)
 		: FParallelCommandListSet(GET_STATID(STAT_CLP_Translucency), InView, InSceneRenderer, InParentCmdList, bInParallelExecute, bInCreateSceneContext, InDrawRenderState)
 		, TranslucencyPass(InTranslucencyPass)
+		, bRenderInSeparateTranslucency(InRenderInSeparateTranslucency)
 	{
 	}
 
@@ -347,11 +352,11 @@ public:
 		// Never needs clear here as it is already done in RenderTranslucency.
 		FParallelCommandListSet::SetStateOnCommandList(CmdList);
 		FSceneRenderTargets& SceneContext = FSceneRenderTargets::Get(CmdList);
-		if (TranslucencyPass == ETranslucencyPass::TPT_TranslucencyAfterDOF)
+		if (TranslucencyPass == ETranslucencyPass::TPT_TranslucencyAfterDOF && bRenderInSeparateTranslucency)
 		{
 			SceneContext.BeginRenderingSeparateTranslucency(CmdList, View, *SceneRenderer, false);
 		}
-		else if (TranslucencyPass == ETranslucencyPass::TPT_TranslucencyAfterDOFModulate)
+		else if (TranslucencyPass == ETranslucencyPass::TPT_TranslucencyAfterDOFModulate && bRenderInSeparateTranslucency)
 		{
 			SceneContext.BeginRenderingSeparateTranslucencyModulate(CmdList, View, *SceneRenderer, false);
 		}
