Fixing SkyLight lighting channel mask for path tracing

diff --git a/Engine/Source/Runtime/Renderer/Private/PathTracing.cpp b/Engine/Source/Runtime/Renderer/Private/PathTracing.cpp
--- a/Engine/Source/Runtime/Renderer/Private/PathTracing.cpp
+++ b/Engine/Source/Runtime/Renderer/Private/PathTracing.cpp
@@ -194,9 +194,11 @@ public:
 			// Prepend SkyLight to light buffer
 			// WARNING: Until ray payload encodes Light data buffer, the execution depends on this ordering!
 			uint32 SkyLightIndex = 0;
+			uint8 SkyLightLightingChannelMask = 0xFF;
 			LightData.Type[SkyLightIndex] = 0;
 			LightData.Color[SkyLightIndex] = FVector(SkyLightData.Color);
 			LightData.Flags[SkyLightIndex] = SkyLightData.bTransmission & 0x01;
+			LightData.Flags[SkyLightIndex] |= (SkyLightLightingChannelMask & 0x7) << 1;
 			LightData.Count++;
 
 			for (auto Light : Lights)
