Add continuous update flag to RVT asset
Add VT scalability setting for throttling continuous updates

diff --git a/Engine/Source/Runtime/Engine/Classes/VT/RuntimeVirtualTexture.h b/Engine/Source/Runtime/Engine/Classes/VT/RuntimeVirtualTexture.h
--- a/Engine/Source/Runtime/Engine/Classes/VT/RuntimeVirtualTexture.h
+++ b/Engine/Source/Runtime/Engine/Classes/VT/RuntimeVirtualTexture.h
@@ -33,6 +33,10 @@ protected:
 	UPROPERTY(EditAnywhere, AdvancedDisplay, Category = Layout, meta = (DisplayName = "Enable clear before render"))
 	bool bClearTextures = true;
 
+	/** Enable continuous update of the virtual texture pages. This round-robin updates already mapped pages and can help fix pages that are mapped before dependent textures are fully streamed in.  */
+	UPROPERTY(EditAnywhere, AdvancedDisplay, Category = Layout, meta = (DisplayName = "Enable continuous page updates"))
+	bool bContinuousUpdate = false;
+
 	/** Enable page table channel packing. This reduces page table memory and update cost but can reduce the ability to share physical memory with other virtual textures.  */
 	UPROPERTY(EditAnywhere, AdvancedDisplay, Category = Layout, meta = (DisplayName = "Enable packed page table"))
 	bool bSinglePhysicalSpace = true;
diff --git a/Engine/Source/Runtime/Engine/Private/VT/RuntimeVirtualTexture.cpp b/Engine/Source/Runtime/Engine/Private/VT/RuntimeVirtualTexture.cpp
--- a/Engine/Source/Runtime/Engine/Private/VT/RuntimeVirtualTexture.cpp
+++ b/Engine/Source/Runtime/Engine/Private/VT/RuntimeVirtualTexture.cpp
@@ -225,6 +225,8 @@ void URuntimeVirtualTexture::GetProducerDescription(FVTProducerDescription& OutD
 	OutDesc.BlockHeightInTiles = HeightInTiles;
 	OutDesc.MaxLevel = FMath::Max((int32)FMath::CeilLogTwo(FMath::Max(OutDesc.BlockWidthInTiles, OutDesc.BlockHeightInTiles)) - GetRemoveLowMips(), 0);
 
+	OutDesc.bContinuousUpdate = bContinuousUpdate;
+
 	OutDesc.NumTextureLayers = GetLayerCount();
 	OutDesc.NumPhysicalGroups = bSinglePhysicalSpace ? 1 : GetLayerCount();
 
diff --git a/Engine/Source/Runtime/Engine/Private/VT/VirtualTextureScalability.cpp b/Engine/Source/Runtime/Engine/Private/VT/VirtualTextureScalability.cpp
--- a/Engine/Source/Runtime/Engine/Private/VT/VirtualTextureScalability.cpp
+++ b/Engine/Source/Runtime/Engine/Private/VT/VirtualTextureScalability.cpp
@@ -27,6 +27,13 @@ namespace VirtualTextureScalability
 		ECVF_RenderThreadSafe | ECVF_Scalability
 	);
 
+	static TAutoConsoleVariable<int32> CVarVTMaxContinuousUpdatesPerFrame(
+		TEXT("r.VT.MaxContinuousUpdatesPerFrame"),
+		1,
+		TEXT("Max number of page uploads for pages that are already mapped."),
+		ECVF_RenderThreadSafe | ECVF_Scalability
+	);
+
 	static TAutoConsoleVariable<float> CVarVTPoolSizeScale(
 		TEXT("r.VT.PoolSizeScale"),
 		1.f,
@@ -124,6 +131,11 @@ namespace VirtualTextureScalability
 #endif
 	}
 
+	int32 GetMaxContinuousUpdatesPerFrame()
+	{
+		return CVarVTMaxContinuousUpdatesPerFrame.GetValueOnAnyThread();
+	}
+
 	float GetPoolSizeScale()
 	{
 		return CVarVTPoolSizeScale.GetValueOnAnyThread();
diff --git a/Engine/Source/Runtime/Engine/Public/VT/VirtualTextureScalability.h b/Engine/Source/Runtime/Engine/Public/VT/VirtualTextureScalability.h
--- a/Engine/Source/Runtime/Engine/Public/VT/VirtualTextureScalability.h
+++ b/Engine/Source/Runtime/Engine/Public/VT/VirtualTextureScalability.h
@@ -8,6 +8,8 @@ namespace VirtualTextureScalability
 {
 	/** Get max upload rate to virtual textures. */
 	ENGINE_API int32 GetMaxUploadsPerFrame();
+	/** Get max update rate of already mapped virtual texture pages. */
+	ENGINE_API int32 GetMaxContinuousUpdatesPerFrame();
 	/** Get scale factor for virtual texture physical pool sizes. */
 	ENGINE_API float GetPoolSizeScale();
 	/** Get resolution bias for runtime virtual textures. */
diff --git a/Engine/Source/Runtime/Renderer/Private/VT/UniqueRequestList.h b/Engine/Source/Runtime/Renderer/Private/VT/UniqueRequestList.h
--- a/Engine/Source/Runtime/Renderer/Private/VT/UniqueRequestList.h
+++ b/Engine/Source/Runtime/Renderer/Private/VT/UniqueRequestList.h
@@ -68,12 +68,14 @@ public:
 		, LoadRequests(new(MemStack) FVirtualTextureLocalTile[LoadRequestCapacity])
 		, MappingRequests(new(MemStack) FMappingRequest[MappingRequestCapacity])
 		, DirectMappingRequests(new(MemStack) FDirectMappingRequest[DirectMappingRequestCapacity])
+		, ContinuousUpdateRequests(new(MemStack) FVirtualTextureLocalTile[LoadRequestCapacity])
 		, LoadRequestCount(new(MemStack) uint16[LoadRequestCapacity])
 		, LoadRequestGroupMask(new(MemStack) uint8[LoadRequestCapacity])
 		, NumLoadRequests(0u)
 		, NumLockRequests(0u)
 		, NumMappingRequests(0u)
 		, NumDirectMappingRequests(0u)
+		, NumContinuousUpdateRequests(0u)
 	{
 	}
 
@@ -82,15 +84,19 @@ public:
 		LoadRequestHash.Clear();
 		MappingRequestHash.Clear();
 		DirectMappingRequestHash.Clear();
+		ContinuousUpdateRequestHash.Clear();
 	}
 
 	inline uint32 GetNumLoadRequests() const { return NumLoadRequests; }
 	inline uint32 GetNumMappingRequests() const { return NumMappingRequests; }
-	inline uint32 GetNumDirectMappingRequests() { return NumDirectMappingRequests; }
+	inline uint32 GetNumDirectMappingRequests() const { return NumDirectMappingRequests; }
+	inline uint32 GetNumContinuousUpdateRequests() const { return NumContinuousUpdateRequests; }
 
 	inline const FVirtualTextureLocalTile& GetLoadRequest(uint32 i) const { checkSlow(i < NumLoadRequests); return LoadRequests[i]; }
 	inline const FMappingRequest& GetMappingRequest(uint32 i) const { checkSlow(i < NumMappingRequests); return MappingRequests[i]; }
 	inline const FDirectMappingRequest& GetDirectMappingRequest(uint32 i) const { checkSlow(i < NumDirectMappingRequests); return DirectMappingRequests[i]; }
+	inline const FVirtualTextureLocalTile& GetContinuousUpdateRequest(uint32 i) const { checkSlow(i < NumContinuousUpdateRequests); return ContinuousUpdateRequests[i]; }
+	
 	inline uint8 GetGroupMask(uint32 i) const { checkSlow(i < NumLoadRequests); return LoadRequestGroupMask[i]; }
 	inline bool IsLocked(uint32 i) const { checkSlow(i < NumLoadRequests); return i < NumLockRequests; }
 
@@ -102,6 +108,8 @@ public:
 	void AddDirectMappingRequest(uint8 InSpaceID, uint16 InPhysicalSpaceID, uint8 InPageTableLayerIndex, uint8 InLogSize, uint32 InAddress, uint8 InLevel, uint16 InPhysicalAddress);
 	void AddDirectMappingRequest(const FDirectMappingRequest& Request);
 
+	void AddContinuousUpdateRequest(const FVirtualTextureLocalTile& Request);
+
 	void MergeRequests(const FUniqueRequestList* RESTRICT Other, FMemStack& MemStack);
 
 	void SortRequests(FVirtualTextureProducerCollection& Producers, FMemStack& MemStack, uint32 MaxNumRequests);
@@ -110,14 +118,17 @@ private:
 	static const uint32 LoadRequestCapacity = 4u * 1024;
 	static const uint32 MappingRequestCapacity = 8u * 1024 - 256u;
 	static const uint32 DirectMappingRequestCapacity = MappingRequestCapacity;
+	static const uint32 ContinuousUpdateRequestCapacity = LoadRequestCapacity;
 
 	TStaticHashTable<1024u, LoadRequestCapacity> LoadRequestHash;
 	TStaticHashTable<1024u, MappingRequestCapacity> MappingRequestHash;
 	TStaticHashTable<512u, DirectMappingRequestCapacity> DirectMappingRequestHash;
+	TStaticHashTable<1024u, ContinuousUpdateRequestCapacity> ContinuousUpdateRequestHash;
 
 	FVirtualTextureLocalTile* LoadRequests;
 	FMappingRequest* MappingRequests;
 	FDirectMappingRequest* DirectMappingRequests;
+	FVirtualTextureLocalTile* ContinuousUpdateRequests;
 	uint16* LoadRequestCount;
 	uint8* LoadRequestGroupMask;
 
@@ -125,6 +136,7 @@ private:
 	uint32 NumLockRequests;
 	uint32 NumMappingRequests;
 	uint32 NumDirectMappingRequests;
+	uint32 NumContinuousUpdateRequests;
 };
 
 
@@ -240,6 +252,25 @@ inline void FUniqueRequestList::AddDirectMappingRequest(const FDirectMappingRequ
 	}
 }
 
+inline void FUniqueRequestList::AddContinuousUpdateRequest(const FVirtualTextureLocalTile& Request)
+{
+	const uint16 Hash = MurmurFinalize64(Request.PackedValue);
+	for (uint16 Index = ContinuousUpdateRequestHash.First(Hash); ContinuousUpdateRequestHash.IsValid(Index); Index = ContinuousUpdateRequestHash.Next(Index))
+	{
+		if (Request == ContinuousUpdateRequests[Index])
+		{
+			return;
+		}
+	}
+
+	if (ensure(NumContinuousUpdateRequests < ContinuousUpdateRequestCapacity))
+	{
+		const uint32 Index = NumContinuousUpdateRequests++;
+		ContinuousUpdateRequestHash.Add(Hash, Index);
+		ContinuousUpdateRequests[Index] = Request;
+	}
+}
+
 inline void FUniqueRequestList::MergeRequests(const FUniqueRequestList* RESTRICT Other, FMemStack& MemStack)
 {
 	FMemMark Mark(MemStack);
@@ -272,6 +303,11 @@ inline void FUniqueRequestList::MergeRequests(const FUniqueRequestList* RESTRICT
 	{
 		AddDirectMappingRequest(Other->GetDirectMappingRequest(Index));
 	}
+
+	for (uint32 Index = 0u; Index < Other->NumContinuousUpdateRequests; ++Index)
+	{
+		AddContinuousUpdateRequest(Other->GetContinuousUpdateRequest(Index));
+	}
 }
 
 inline void FUniqueRequestList::SortRequests(FVirtualTextureProducerCollection& Producers, FMemStack& MemStack, uint32 MaxNumRequests)
diff --git a/Engine/Source/Runtime/Renderer/Private/VT/VirtualTextureSystem.cpp b/Engine/Source/Runtime/Renderer/Private/VT/VirtualTextureSystem.cpp
--- a/Engine/Source/Runtime/Renderer/Private/VT/VirtualTextureSystem.cpp
+++ b/Engine/Source/Runtime/Renderer/Private/VT/VirtualTextureSystem.cpp
@@ -42,6 +42,7 @@ DECLARE_DWORD_COUNTER_STAT(TEXT("Num page visible resident"), STAT_NumPageVisibl
 DECLARE_DWORD_COUNTER_STAT(TEXT("Num page visible not resident"), STAT_NumPageVisibleNotResident, STATGROUP_VirtualTexturing);
 DECLARE_DWORD_COUNTER_STAT(TEXT("Num page prefetch"), STAT_NumPagePrefetch, STATGROUP_VirtualTexturing);
 DECLARE_DWORD_COUNTER_STAT(TEXT("Num page update"), STAT_NumPageUpdate, STATGROUP_VirtualTexturing);
+DECLARE_DWORD_COUNTER_STAT(TEXT("Num mapped page update"), STAT_NumMappedPageUpdate, STATGROUP_VirtualTexturing);
 DECLARE_DWORD_COUNTER_STAT(TEXT("Num continuous page update"), STAT_NumContinuousPageUpdate, STATGROUP_VirtualTexturing);
 
 DECLARE_DWORD_COUNTER_STAT(TEXT("Num stacks requested"), STAT_NumStacksRequested, STATGROUP_VirtualTexturing);
@@ -1101,6 +1102,14 @@ void FVirtualTextureSystem::Update(FRHICommandListImmediate& RHICmdList, ERHIFea
 		MergedRequestList->SortRequests(Producers, MemStack, MaxRequestUploads);
 	}
 
+	{
+		// After sorting and clamping the load requests, if we still have unused upload bandwidth then use it to add some continous updates
+		const int32 MaxNumUploads = VirtualTextureScalability::GetMaxUploadsPerFrame();
+		const int32 MaxTilesToProduce = FMath::Max(MaxNumUploads - MappedTilesToProduce.Num() - (int32)MergedRequestList->GetNumLoadRequests(), 0);
+
+		GetContinuousUpdatesToProduce(MergedRequestList, MaxTilesToProduce);
+	}
+
 	// Submit the requests to produce pages that are already mapped
 	SubmitPreMappedRequests(RHICmdList, FeatureLevel);
 	// Submit the merged requests
@@ -1255,10 +1264,6 @@ void FVirtualTextureSystem::GatherRequestsTask(const FGatherRequestsParameters&
 	uint32 NumNonResidentPages = 0u;
 	uint32 NumPrefetchPages = 0u;
 
-#if WITH_EDITOR
-	TSet<FVirtualTextureLocalTile> ContinuousUpdateTilesToProduceThreadLocal;
-#endif
-
 	for (uint32 i = Parameters.PageStartIndex; i < PageEndIndex; ++i)
 	{
 		const uint32 PageEncoded = UniquePageList->GetPage(i);
@@ -1303,16 +1308,14 @@ void FVirtualTextureSystem::GatherRequestsTask(const FGatherRequestsParameters&
 					// Page is already resident, just need to update LRU free list
 					AddPageUpdate(PageUpdateBuffers, PageUpdateFlushCount, PhysicalSpaceIDAndAddress.PhysicalSpaceID, PhysicalSpaceIDAndAddress.pAddress);
 
-				#if WITH_EDITOR
+					// If continuous update flag is set then add this to pages which can be potentially updated if we have spare upload bandwidth
+					//todo[vt]: Would be better to test continuous update flag *per producer*, but this would require extra indirection so need to profile first
+					if (GetPhysicalSpace(PhysicalSpaceIDAndAddress.PhysicalSpaceID)->GetDescription().bContinuousUpdate)
 					{
-						if (GetPhysicalSpace(PhysicalSpaceIDAndAddress.PhysicalSpaceID)->GetDescription().bContinuousUpdate)
-						{
-							FTexturePagePool& RESTRICT PagePool = GetPhysicalSpace(PhysicalSpaceIDAndAddress.PhysicalSpaceID)->GetPagePool();
-
-							ContinuousUpdateTilesToProduceThreadLocal.Add(PagePool.GetLocalTileFromPhysicalAddress(PhysicalSpaceIDAndAddress.pAddress));
-						}
+						FTexturePagePool& RESTRICT PagePool = GetPhysicalSpace(PhysicalSpaceIDAndAddress.PhysicalSpaceID)->GetPagePool();
+						const FVirtualTextureLocalTile LocalTile = PagePool.GetLocalTileFromPhysicalAddress(PhysicalSpaceIDAndAddress.pAddress);
+						RequestList->AddContinuousUpdateRequest(LocalTile);
 					}
-				#endif
 
 					++PageUpdateBuffers[PhysicalSpaceIDAndAddress.PhysicalSpaceID].WorkingSetSize;
 					++NumResidentPages;
@@ -1630,14 +1633,6 @@ void FVirtualTextureSystem::GatherRequestsTask(const FGatherRequestsParameters&
 			{
 				PagePool.UpdateUsage(Frame, Buffer.PhysicalAddresses[i]);
 			}
-
-		#if WITH_EDITOR
-			if (PhysicalSpace->GetDescription().bContinuousUpdate)
-			{
-				FScopeLock ScopeLock(&ContinuousUpdateTilesToProduceCS);
-				ContinuousUpdateTilesToProduce.Append(ContinuousUpdateTilesToProduceThreadLocal);
-			}
-		#endif
 		}
 		
 		INC_DWORD_STAT_BY(STAT_NumPageUpdate, Buffer.NumPageUpdates);
@@ -1649,6 +1644,21 @@ void FVirtualTextureSystem::GatherRequestsTask(const FGatherRequestsParameters&
 	INC_DWORD_STAT_BY(STAT_NumPagePrefetch, NumPrefetchPages);
 }
 
+void FVirtualTextureSystem::GetContinuousUpdatesToProduce(FUniqueRequestList const* RequestList, int32 MaxTilesToProduce)
+{
+	const int32 NumContinuousUpdateRequests = (int32)RequestList->GetNumContinuousUpdateRequests();
+	const int32 MaxContinousUpdates = FMath::Min(VirtualTextureScalability::GetMaxContinuousUpdatesPerFrame(), NumContinuousUpdateRequests);
+
+	int32 NumContinuousUpdates = 0;
+	while (NumContinuousUpdates < MaxContinousUpdates && ContinuousUpdateTilesToProduce.Num() < MaxTilesToProduce)
+	{
+		// Note it's possible that we add a duplicate value to the TSet here, and so MappedTilesToProduce doesn't grow.
+		// But ending up with fewer continuous updates then the maximum is OK.
+		int32 RandomIndex = FMath::Rand() % NumContinuousUpdateRequests;
+		ContinuousUpdateTilesToProduce.Add(RequestList->GetContinuousUpdateRequest(RandomIndex));
+		NumContinuousUpdates++;
+	}
+}
 
 void FVirtualTextureSystem::UpdateCSVStats() const
 {
@@ -1743,6 +1753,7 @@ void FVirtualTextureSystem::SubmitRequestsFromLocalTileList(const TSet<FVirtualT
 void FVirtualTextureSystem::SubmitPreMappedRequests(FRHICommandListImmediate& RHICmdList, ERHIFeatureLevel::Type FeatureLevel)
 {
 	{
+		INC_DWORD_STAT_BY(STAT_NumMappedPageUpdate, MappedTilesToProduce.Num());
 		SubmitRequestsFromLocalTileList(MappedTilesToProduce, EVTProducePageFlags::None, RHICmdList, FeatureLevel);
 		MappedTilesToProduce.Reset();
 	}
diff --git a/Engine/Source/Runtime/Renderer/Private/VT/VirtualTextureSystem.h b/Engine/Source/Runtime/Renderer/Private/VT/VirtualTextureSystem.h
--- a/Engine/Source/Runtime/Renderer/Private/VT/VirtualTextureSystem.h
+++ b/Engine/Source/Runtime/Renderer/Private/VT/VirtualTextureSystem.h
@@ -91,6 +91,8 @@ private:
 	void GatherRequestsTask(const FGatherRequestsParameters& Parameters);
 	void FeedbackAnalysisTask(const FFeedbackAnalysisParameters& Parameters);
 
+	void GetContinuousUpdatesToProduce(FUniqueRequestList const* RequestList, int32 MaxTilesToProduce);
+
 	void UpdateCSVStats() const;
 
 	uint32	Frame;
@@ -122,7 +124,6 @@ private:
 	TArray<FVirtualTextureLocalTile> TilesToLock;
 	FTexturePageLocks TileLocks;
 
-	FCriticalSection ContinuousUpdateTilesToProduceCS;
 	TSet<FVirtualTextureLocalTile> ContinuousUpdateTilesToProduce;
 	TSet<FVirtualTextureLocalTile> MappedTilesToProduce;
 	TArray<FVirtualTextureLocalTile> TransientCollectedPages;
