REL - Remove check that will miss removal in shipping builds.

diff --git a/Engine/Source/Runtime/Renderer/Private/PrimitiveSceneInfo.cpp b/Engine/Source/Runtime/Renderer/Private/PrimitiveSceneInfo.cpp
--- a/Engine/Source/Runtime/Renderer/Private/PrimitiveSceneInfo.cpp
+++ b/Engine/Source/Runtime/Renderer/Private/PrimitiveSceneInfo.cpp
@@ -454,7 +454,7 @@ void FPrimitiveSceneInfo::RemoveCachedMeshDrawCommands()
 			StateBucketCount.Num--;
 			if (StateBucketCount.Num == 0)
 			{
-				check(Scene->CachedMeshDrawCommandStateBuckets[PassIndex].RemoveByElementId(CachedCommand.StateBucketId));
+				Scene->CachedMeshDrawCommandStateBuckets[PassIndex].RemoveByElementId(CachedCommand.StateBucketId);
 			}
 
 			FGraphicsMinimalPipelineStateId::RemovePersistentId(CachedPipelineId);
