Fixed wrong Cvar string

diff --git a/Engine/Source/Runtime/Renderer/Private/DeferredShadingRenderer.cpp b/Engine/Source/Runtime/Renderer/Private/DeferredShadingRenderer.cpp
--- a/Engine/Source/Runtime/Renderer/Private/DeferredShadingRenderer.cpp
+++ b/Engine/Source/Runtime/Renderer/Private/DeferredShadingRenderer.cpp
@@ -810,7 +810,7 @@ bool FDeferredShadingSceneRenderer::GatherRayTracingWorldInstances(FRHICommandLi
 
 				//#dxr_todo UE-68621  The Raytracing code path does not support ShowFlags since data moved to the SceneInfo. 
 				//Touching the SceneProxy to determine this would simply cost too much
-				static const auto RayTracingStaticMeshesCVar = IConsoleManager::Get().FindTConsoleVariableDataInt(TEXT("r.RayTracing.StaticMeshes"));
+				static const auto RayTracingStaticMeshesCVar = IConsoleManager::Get().FindTConsoleVariableDataInt(TEXT("r.RayTracing.Geometry.StaticMeshes"));
 
 				if (SceneInfo->bIsRayTracingStaticRelevant && View.Family->EngineShowFlags.StaticMeshes && RayTracingStaticMeshesCVar && RayTracingStaticMeshesCVar->GetValueOnRenderThread() > 0)
 				{
