Fix crash if physical texture in AllocatedVT is null, this can happen with async texture creation if VTs haven't finished background creation

Edigrated from CL 13818555

diff --git a/Engine/Source/Runtime/Renderer/Private/VT/AllocatedVirtualTexture.cpp b/Engine/Source/Runtime/Renderer/Private/VT/AllocatedVirtualTexture.cpp
--- a/Engine/Source/Runtime/Renderer/Private/VT/AllocatedVirtualTexture.cpp
+++ b/Engine/Source/Runtime/Renderer/Private/VT/AllocatedVirtualTexture.cpp
@@ -89,7 +89,8 @@ FAllocatedVirtualTexture::FAllocatedVirtualTexture(FVirtualTextureSystem* InSyst
 	bool bSupport16BitPageTable = true;
 	for (int32 Index = 0; Index < UniquePageTableLayers.Num(); ++Index)
 	{
-		if (!UniquePageTableLayers[Index].PhysicalSpace->DoesSupport16BitPageTable())
+		const FVirtualTexturePhysicalSpace* PhysicalSpace = UniquePageTableLayers[Index].PhysicalSpace;
+		if (PhysicalSpace && !PhysicalSpace->DoesSupport16BitPageTable())
 		{
 			bSupport16BitPageTable = false;
 			break;
