Fixed duplicate symbol triggering in StressTestUnity build mode

diff --git a/Engine/Source/Runtime/Renderer/Private/DistanceFieldObjectManagement.cpp b/Engine/Source/Runtime/Renderer/Private/DistanceFieldObjectManagement.cpp
--- a/Engine/Source/Runtime/Renderer/Private/DistanceFieldObjectManagement.cpp
+++ b/Engine/Source/Runtime/Renderer/Private/DistanceFieldObjectManagement.cpp
@@ -63,19 +63,19 @@ void TDistanceFieldObjectBuffers<PrimitiveType>::Initialize()
 
 const uint32 UpdateObjectsGroupSize = 64;
 
-struct FParallelUpdateRange
+struct FParallelUpdateRangeDFO
 {
 	int32 ItemStart;
 	int32 ItemCount;
 };
 
-struct FParallelUpdateRanges
+struct FParallelUpdateRangesDFO
 {
-	FParallelUpdateRange Range[4];
+	FParallelUpdateRangeDFO Range[4];
 };
 
 // TODO: Improve and move to shared utility location.
-static int32 PartitionUpdateRanges(FParallelUpdateRanges& Ranges, int32 ItemCount, bool bAllowParallel)
+static int32 PartitionUpdateRangesDFO(FParallelUpdateRangesDFO& Ranges, int32 ItemCount, bool bAllowParallel)
 {
 	if (ItemCount < 256 || !bAllowParallel)
 	{
@@ -556,9 +556,9 @@ void FDeferredShadingSceneRenderer::UpdateGlobalDistanceFieldObjectBuffers(FRHIC
 
 					const TArray<FPrimitiveBounds>& PrimitiveBounds = Scene->PrimitiveBounds;
 
-					FParallelUpdateRanges ParallelRanges;
+					FParallelUpdateRangesDFO ParallelRanges;
 
-					int32 RangeCount = PartitionUpdateRanges(ParallelRanges, DistanceFieldSceneData.IndicesToUpdateInObjectBuffers.Num(), bExecuteInParallel);
+					int32 RangeCount = PartitionUpdateRangesDFO(ParallelRanges, DistanceFieldSceneData.IndicesToUpdateInObjectBuffers.Num(), bExecuteInParallel);
 
 					ParallelFor(RangeCount,
 						[&DistanceFieldSceneData, &ParallelRanges, PrimitiveBounds, InvTextureDim, bAnyViewEnabledDistanceCulling, RangeCount](int32 RangeIndex)
