Fix static analysis-flagged merge issue with an unnecessary ternary operation.

diff --git a/Engine/Source/Runtime/Renderer/Private/SceneTextureReductions.cpp b/Engine/Source/Runtime/Renderer/Private/SceneTextureReductions.cpp
--- a/Engine/Source/Runtime/Renderer/Private/SceneTextureReductions.cpp
+++ b/Engine/Source/Runtime/Renderer/Private/SceneTextureReductions.cpp
@@ -189,7 +189,7 @@ void BuildHZB(
 					PermutationVector.Get<FHZBBuildCS::FDimVisBufferFormat>() != 0 ? TEXT(" ReadVisBuffer") : TEXT(""),
 					DstSize.X, DstSize.Y),
 				PassParameters,
-				StartDestMip ? (ERDGPassFlags::Compute) : ERDGPassFlags::Compute,
+				ERDGPassFlags::Compute,
 				[PassParameters, ComputeShader, DstSize](FRHICommandList& RHICmdList)
 			{
 				FComputeShaderUtils::Dispatch(RHICmdList, ComputeShader, *PassParameters, FComputeShaderUtils::GetGroupCount(DstSize, 8));
@@ -211,7 +211,7 @@ void BuildHZB(
 			GraphBuilder.AddPass(
 				RDG_EVENT_NAME("DownsampleHZB(mip=%d) %dx%d", StartDestMip, DstSize.X, DstSize.Y),
 				PassParameters,
-				StartDestMip ? (ERDGPassFlags::Raster) : ERDGPassFlags::Raster,
+				ERDGPassFlags::Raster,
 				[PassParameters, GlobalShaderMap, PixelShader, DstSize](FRHICommandList& RHICmdList)
 			{
 				FPixelShaderUtils::DrawFullscreenPixelShader(RHICmdList, GlobalShaderMap, PixelShader, *PassParameters, FIntRect(0, 0, DstSize.X, DstSize.Y));
