Made C++ uses of Buffer/Structured buffer match how it is used in the shader.

diff --git a/Engine/Source/Runtime/Renderer/Private/VirtualShadowMaps/VirtualShadowMapArray.cpp b/Engine/Source/Runtime/Renderer/Private/VirtualShadowMaps/VirtualShadowMapArray.cpp
--- a/Engine/Source/Runtime/Renderer/Private/VirtualShadowMaps/VirtualShadowMapArray.cpp
+++ b/Engine/Source/Runtime/Renderer/Private/VirtualShadowMaps/VirtualShadowMapArray.cpp
@@ -203,10 +203,10 @@ static void SetCacheDataShaderParameters(FRDGBuilder& GraphBuilder, const TArray
 		}
 	}
 	CacheDataParameters.ShadowMapCacheData = GraphBuilder.CreateSRV(CreateStructuredBuffer(GraphBuilder, TEXT("ShadowMapCacheData"), ShadowMapCacheData));
-	CacheDataParameters.PrevPageFlags = GraphBuilder.CreateSRV(GraphBuilder.RegisterExternalBuffer(VirtualShadowMapArrayCacheManager->PrevPageFlags, TEXT("PrevPageFlags")), PF_R32_UINT);
-	CacheDataParameters.PrevPageTable = GraphBuilder.CreateSRV(GraphBuilder.RegisterExternalBuffer(VirtualShadowMapArrayCacheManager->PrevPageTable, TEXT("PrevPageTable")), PF_R32_UINT);
-	CacheDataParameters.PrevPhysicalPageMetaData = GraphBuilder.CreateSRV(GraphBuilder.RegisterExternalBuffer(VirtualShadowMapArrayCacheManager->PrevPhysicalPageMetaData, TEXT("PrevPhysicalPageMetaData")), PF_R32G32_UINT);
-	CacheDataParameters.PrevDynamicCasterPageFlags = GraphBuilder.CreateSRV(GraphBuilder.RegisterExternalBuffer(VirtualShadowMapArrayCacheManager->PrevDynamicCasterPageFlags, TEXT("PrevDynamicCasterPageFlags")), PF_R32_UINT);
+	CacheDataParameters.PrevPageFlags = GraphBuilder.CreateSRV(GraphBuilder.RegisterExternalBuffer(VirtualShadowMapArrayCacheManager->PrevPageFlags, TEXT("PrevPageFlags")));
+	CacheDataParameters.PrevPageTable = GraphBuilder.CreateSRV(GraphBuilder.RegisterExternalBuffer(VirtualShadowMapArrayCacheManager->PrevPageTable, TEXT("PrevPageTable")));
+	CacheDataParameters.PrevPhysicalPageMetaData = GraphBuilder.CreateSRV(GraphBuilder.RegisterExternalBuffer(VirtualShadowMapArrayCacheManager->PrevPhysicalPageMetaData, TEXT("PrevPhysicalPageMetaData")));
+	CacheDataParameters.PrevDynamicCasterPageFlags = GraphBuilder.CreateSRV(GraphBuilder.RegisterExternalBuffer(VirtualShadowMapArrayCacheManager->PrevDynamicCasterPageFlags, TEXT("PrevDynamicCasterPageFlags")));
 
 }
 
@@ -440,13 +440,13 @@ void FVirtualShadowMapArray::GenerateIdentityPageTables(FRDGBuilder& GraphBuilde
 		FRDGBufferRef PageTableRDG = GraphBuilder.CreateBuffer(FRDGBufferDesc::CreateStructuredDesc(8, ShadowMaps.Num() * CommonParameters.PageTableSize), TEXT("PageTable"));
 
 		const uint32 NumPageFlags = ShadowMaps.Num() * CommonParameters.PageTableSize;
-		FRDGBufferRef PageFlagsRDG = GraphBuilder.CreateBuffer(FRDGBufferDesc::CreateBufferDesc(sizeof(uint32), NumPageFlags), TEXT("PageFlags"));
-		AddClearUAVPass(GraphBuilder, GraphBuilder.CreateUAV(PageFlagsRDG, PF_R32_UINT), 0);
+		FRDGBufferRef PageFlagsRDG = GraphBuilder.CreateBuffer(FRDGBufferDesc::CreateStructuredDesc(sizeof(uint32), NumPageFlags), TEXT("PageFlags"));
+		AddClearUAVPass(GraphBuilder, GraphBuilder.CreateUAV(PageFlagsRDG), 0);
 
 		// Total storage for Hierarchical page tables for all virtual shadow maps
 		const uint32 NumHPageFlags = ShadowMaps.Num() * CommonParameters.HPageTableSize;
-		FRDGBufferRef HPageFlagsRDG = GraphBuilder.CreateBuffer(FRDGBufferDesc::CreateBufferDesc(sizeof(uint32), NumHPageFlags), TEXT("HPageFlags"));
-		AddClearUAVPass(GraphBuilder, GraphBuilder.CreateUAV(HPageFlagsRDG, PF_R32_UINT), 0);
+		FRDGBufferRef HPageFlagsRDG = GraphBuilder.CreateBuffer(FRDGBufferDesc::CreateStructuredDesc(sizeof(uint32), NumHPageFlags), TEXT("HPageFlags"));
+		AddClearUAVPass(GraphBuilder, GraphBuilder.CreateUAV(HPageFlagsRDG), 0);
 
 		uint32 MinPhysicalMipLevel = FMath::CeilLogTwo(FVirtualShadowMap::VirtualMaxResolutionXY) - FMath::CeilLogTwo(GetPhysicalPoolSize().X);
 		uint32 ClampedMipLevel = MipLevel == INDEX_NONE ? INDEX_NONE : FMath::Max(MinPhysicalMipLevel, MipLevel);
@@ -456,8 +456,8 @@ void FVirtualShadowMapArray::GenerateIdentityPageTables(FRDGBuilder& GraphBuilde
 			FGenerateIdentityMappingCS::FParameters* PassParameters = GraphBuilder.AllocParameters< FGenerateIdentityMappingCS::FParameters >();
 			PassParameters->CommonParameters = CommonParameters;
 			PassParameters->VirtualShadowMapId = SM->ID;
-			PassParameters->OutPageTable = GraphBuilder.CreateUAV(PageTableRDG, PF_R32G32_UINT);
-			PassParameters->OutPageRequestFlags = GraphBuilder.CreateUAV(PageFlagsRDG, PF_R32_UINT);
+			PassParameters->OutPageTable = GraphBuilder.CreateUAV(PageTableRDG);
+			PassParameters->OutPageRequestFlags = GraphBuilder.CreateUAV(PageFlagsRDG);
 			PassParameters->TargetMipLevel = ClampedMipLevel;
 
 			auto ComputeShader = GetGlobalShaderMap(GMaxRHIFeatureLevel)->GetShader<FGenerateIdentityMappingCS>();
@@ -476,8 +476,8 @@ void FVirtualShadowMapArray::GenerateIdentityPageTables(FRDGBuilder& GraphBuilde
 			// Run pass building hierarchical page flags to make culling acceptable performance.
 			FGenerateHierarchicalPageFlagsCS::FParameters* PassParameters = GraphBuilder.AllocParameters< FGenerateHierarchicalPageFlagsCS::FParameters >();
 			PassParameters->CommonParameters = CommonParameters;
-			PassParameters->OutHPageFlags = GraphBuilder.CreateUAV(HPageFlagsRDG, PF_R32_UINT);
-			PassParameters->PageFlags = GraphBuilder.CreateSRV(PageFlagsRDG, PF_R32_UINT);
+			PassParameters->OutHPageFlags = GraphBuilder.CreateUAV(HPageFlagsRDG);
+			PassParameters->PageFlags = GraphBuilder.CreateSRV(PageFlagsRDG);
 
 			auto ComputeShader = GetGlobalShaderMap(GMaxRHIFeatureLevel)->GetShader<FGenerateHierarchicalPageFlagsCS>();
 			FComputeShaderUtils::AddPass(
@@ -510,8 +510,8 @@ void FVirtualShadowMapArray::ClearPhysicalMemory(FRDGBuilder& GraphBuilder, FRDG
 		FRDGBufferRef IndirectArgsBuffer = GraphBuilder.CreateBuffer(FRDGBufferDesc::CreateIndirectDesc(4), TEXT("IndirectArgs"));
 		{
 			FInitClearPhysicalPagesArgsCS::FParameters* PassParameters = GraphBuilder.AllocParameters<FInitClearPhysicalPagesArgsCS::FParameters>();
-			PassParameters->NumAllocatedPhysicalPages = GraphBuilder.CreateSRV(GraphBuilder.RegisterExternalBuffer(AllocatedPagesOffset), PF_R32_UINT);
-			PassParameters->ClearPhysicalPagesArgs = GraphBuilder.CreateUAV(IndirectArgsBuffer);
+			PassParameters->NumAllocatedPhysicalPages = GraphBuilder.CreateSRV(GraphBuilder.RegisterExternalBuffer(AllocatedPagesOffset));
+			PassParameters->ClearPhysicalPagesArgs = GraphBuilder.CreateUAV(IndirectArgsBuffer, PF_R32_UINT);
 
 			auto ComputeShader = GetGlobalShaderMap(GMaxRHIFeatureLevel)->GetShader<FInitClearPhysicalPagesArgsCS>();
 			FComputeShaderUtils::AddPass(
@@ -535,9 +535,9 @@ void FVirtualShadowMapArray::ClearPhysicalMemory(FRDGBuilder& GraphBuilder, FRDG
 			{
 				PassParameters->CachedPhysicalPagesTexture = RegisterExternalTextureWithFallback(GraphBuilder, VirtualShadowMapArrayCacheManager ? VirtualShadowMapArrayCacheManager->PrevPhysicalPagePool : nullptr, GSystemTextures.BlackDummy, TEXT("PrevPhysicalPagePool"));
 				PassParameters->CachedPageInfos = GraphBuilder.CreateSRV(GraphBuilder.RegisterExternalBuffer(CachedPageInfos));
-				PassParameters->PrevPhysicalPageMetaData = GraphBuilder.CreateSRV(GraphBuilder.RegisterExternalBuffer(VirtualShadowMapArrayCacheManager->PrevPhysicalPageMetaData, TEXT("PrevPhysicalPageMetaData")), PF_R32G32_UINT);
+				PassParameters->PrevPhysicalPageMetaData = GraphBuilder.CreateSRV(GraphBuilder.RegisterExternalBuffer(VirtualShadowMapArrayCacheManager->PrevPhysicalPageMetaData, TEXT("PrevPhysicalPageMetaData")));
 			}
-			PassParameters->PhysicalPageMetaDataOut = GraphBuilder.CreateUAV(PhysicalPageMetaDataRDG, PF_R32G32_UINT);
+			PassParameters->PhysicalPageMetaDataOut = GraphBuilder.CreateUAV(PhysicalPageMetaDataRDG);
 
 			FClearPhysicalPagesCS::FPermutationDomain PermutationVector;
 			PermutationVector.Set<FClearPhysicalPagesCS::FHasCacheDataDim>(bCacheDataAvailable);
@@ -596,8 +596,8 @@ void FVirtualShadowMapArray::MarkPhysicalPagesRendered(FRDGBuilder& GraphBuilder
 		FMarkRenderedPhysicalPagesCS::FParameters* PassParameters = GraphBuilder.AllocParameters< FMarkRenderedPhysicalPagesCS::FParameters >();
 		PassParameters->CommonParameters = CommonParameters;
 		PassParameters->VirtualShadowMapFlags = GraphBuilder.CreateSRV(VirtualShadowMapFlagsRDG);
-		PassParameters->PageTable = GraphBuilder.CreateSRV(PageTableRDG, PF_R32G32_UINT);
-		PassParameters->InOutPhysicalPageMetaData = GraphBuilder.CreateUAV(PhysicalPageMetaDataRDG, PF_R32G32_UINT);
+		PassParameters->PageTable = GraphBuilder.CreateSRV(PageTableRDG);
+		PassParameters->InOutPhysicalPageMetaData = GraphBuilder.CreateUAV(PhysicalPageMetaDataRDG);
 
 		auto ComputeShader = GetGlobalShaderMap(GMaxRHIFeatureLevel)->GetShader<FMarkRenderedPhysicalPagesCS>();
 
@@ -675,22 +675,22 @@ void FVirtualShadowMapArray::GeneratePageFlagsFromLightGrid(FRDGBuilder& GraphBu
 		FRDGBufferRef StatsBufferRDG = nullptr;
 		if (CVarShowStats.GetValueOnRenderThread() || (VirtualShadowMapArrayCacheManager && VirtualShadowMapArrayCacheManager->IsAccumulatingStats()))
 		{
-			StatsBufferRDG = GraphBuilder.CreateBuffer(FRDGBufferDesc::CreateBufferDesc(sizeof(uint32), NumStats), TEXT("StatsBuffer"));
-			AddClearUAVPass(GraphBuilder, GraphBuilder.CreateUAV(StatsBufferRDG, PF_R32_UINT), 0);
+			StatsBufferRDG = GraphBuilder.CreateBuffer(FRDGBufferDesc::CreateStructuredDesc(sizeof(uint32), NumStats), TEXT("StatsBuffer"));
+			AddClearUAVPass(GraphBuilder, GraphBuilder.CreateUAV(StatsBufferRDG), 0);
 		}
 
 		CommonParameters.NumShadowMaps = ShadowMaps.Num();
 		// Create and clear the requested page flags
 		const uint32 NumPageFlags = ShadowMaps.Num() * CommonParameters.PageTableSize;
-		FRDGBufferRef PageRequestFlagsRDG = GraphBuilder.CreateBuffer(FRDGBufferDesc::CreateBufferDesc(sizeof(uint32), NumPageFlags), TEXT("PageRequestFlags"));
-		AddClearUAVPass(GraphBuilder, GraphBuilder.CreateUAV(PageRequestFlagsRDG, PF_R32_UINT), 0);
-		FRDGBufferRef DynamicCasterPageFlagsRDG = GraphBuilder.CreateBuffer(FRDGBufferDesc::CreateBufferDesc(sizeof(uint32), NumPageFlags), TEXT("DynamicCasterPageFlags"));
-		AddClearUAVPass(GraphBuilder, GraphBuilder.CreateUAV(DynamicCasterPageFlagsRDG, PF_R32_UINT), 0);
+		FRDGBufferRef PageRequestFlagsRDG = GraphBuilder.CreateBuffer(FRDGBufferDesc::CreateStructuredDesc(sizeof(uint32), NumPageFlags), TEXT("PageRequestFlags"));
+		AddClearUAVPass(GraphBuilder, GraphBuilder.CreateUAV(PageRequestFlagsRDG), 0);
+		FRDGBufferRef DynamicCasterPageFlagsRDG = GraphBuilder.CreateBuffer(FRDGBufferDesc::CreateStructuredDesc(sizeof(uint32), NumPageFlags), TEXT("DynamicCasterPageFlags"));
+		AddClearUAVPass(GraphBuilder, GraphBuilder.CreateUAV(DynamicCasterPageFlagsRDG), 0);
 
 		// Total storage for Hierarchical page tables for all virtual shadow maps
 		const uint32 NumHPageFlags = ShadowMaps.Num() * CommonParameters.HPageTableSize;
-		FRDGBufferRef HPageFlagsRDG = GraphBuilder.CreateBuffer(FRDGBufferDesc::CreateBufferDesc(sizeof(uint32), NumHPageFlags), TEXT("HPageFlags"));
-		AddClearUAVPass(GraphBuilder, GraphBuilder.CreateUAV(HPageFlagsRDG, PF_R32_UINT), 0);
+		FRDGBufferRef HPageFlagsRDG = GraphBuilder.CreateBuffer(FRDGBufferDesc::CreateStructuredDesc(sizeof(uint32), NumHPageFlags), TEXT("HPageFlags"));
+		AddClearUAVPass(GraphBuilder, GraphBuilder.CreateUAV(HPageFlagsRDG), 0);
 
 		const uint32 NumPageRects = CommonParameters.NumShadowMaps * FVirtualShadowMap::MaxMipLevels;
 		FRDGBufferRef PageRectBoundsRDG = GraphBuilder.CreateBuffer(FRDGBufferDesc::CreateStructuredDesc(sizeof(FIntVector4), NumPageRects), TEXT("PageRectBounds"));
@@ -804,7 +804,7 @@ void FVirtualShadowMapArray::GeneratePageFlagsFromLightGrid(FRDGBuilder& GraphBu
 				
 				PassParameters->VisBuffer64 = VisBuffer64;
 				PassParameters->View = View.ViewUniformBuffer;
-				PassParameters->OutPageRequestFlags = GraphBuilder.CreateUAV(PageRequestFlagsRDG, PF_R32_UINT);
+				PassParameters->OutPageRequestFlags = GraphBuilder.CreateUAV(PageRequestFlagsRDG);
 				PassParameters->ForwardLightData = View.ForwardLightingResources->ForwardLightDataUniformBuffer;
 				PassParameters->VirtualShadowMapIdRemap = GraphBuilder.CreateSRV(VirtualShadowMapIdRemapRDG);
 				PassParameters->ShadowMapProjectionData = GraphBuilder.CreateSRV(ShadowMapProjectionDataRDG);
@@ -849,7 +849,7 @@ void FVirtualShadowMapArray::GeneratePageFlagsFromLightGrid(FRDGBuilder& GraphBu
 				PassParameters->CommonParameters = CommonParameters;
 				PassParameters->NumVirtualShadowMaps = ForceTopMipVisible.Num();
 				PassParameters->ShadowMapHasFallback = GraphBuilder.CreateSRV(ShadowMapHasFallbackRDG);
-				PassParameters->OutPageRequestFlags = GraphBuilder.CreateUAV(PageRequestFlagsRDG, PF_R32_UINT);
+				PassParameters->OutPageRequestFlags = GraphBuilder.CreateUAV(PageRequestFlagsRDG);
 				auto ComputeShader = Views[0].ShaderMap->GetShader<FMarkFallbackPageFlagsCS>( );
 
 				FComputeShaderUtils::AddPass(
@@ -866,14 +866,14 @@ void FVirtualShadowMapArray::GeneratePageFlagsFromLightGrid(FRDGBuilder& GraphBu
 		FRDGBufferRef HZBPageTableRDG		= GraphBuilder.CreateBuffer( FRDGBufferDesc::CreateStructuredDesc( 8, NumPageFlags ), TEXT("HZBPageTable") );
 
 		// Note: these are passed to the rendering and are not identical to the PageRequest flags coming in from GeneratePageFlagsFromPixels 
-		FRDGBufferRef PageFlagsRDG = GraphBuilder.CreateBuffer(FRDGBufferDesc::CreateBufferDesc(sizeof(uint32), NumPageFlags), TEXT("PageFlags"));
+		FRDGBufferRef PageFlagsRDG = GraphBuilder.CreateBuffer(FRDGBufferDesc::CreateStructuredDesc(sizeof(uint32), NumPageFlags), TEXT("PageFlags"));
 
-		FRDGBufferRef HInvalidPageFlagsRDG = GraphBuilder.CreateBuffer(FRDGBufferDesc::CreateBufferDesc(sizeof(uint32), NumHPageFlags), TEXT("HInvalidPageFlags"));
-		AddClearUAVPass(GraphBuilder, GraphBuilder.CreateUAV(HInvalidPageFlagsRDG, PF_R32_UINT), 0);
+		FRDGBufferRef HInvalidPageFlagsRDG = GraphBuilder.CreateBuffer(FRDGBufferDesc::CreateStructuredDesc(sizeof(uint32), NumHPageFlags), TEXT("HInvalidPageFlags"));
+		AddClearUAVPass(GraphBuilder, GraphBuilder.CreateUAV(HInvalidPageFlagsRDG), 0);
 
 		// Create and clear the counter / page offset, it gets atomically incremented to allocate the physical pages
-		FRDGBufferRef AllocatedPagesOffsetRDG = GraphBuilder.CreateBuffer( FRDGBufferDesc::CreateIndirectDesc( 4 ), TEXT("AllocatedPagesOffset") );
-		AddClearUAVPass( GraphBuilder, GraphBuilder.CreateUAV( AllocatedPagesOffsetRDG, PF_R32_UINT ), 0 );
+		FRDGBufferRef AllocatedPagesOffsetRDG = GraphBuilder.CreateBuffer( FRDGBufferDesc::CreateStructuredDesc(sizeof(uint32),  1), TEXT("AllocatedPagesOffset") );
+		AddClearUAVPass( GraphBuilder, GraphBuilder.CreateUAV(AllocatedPagesOffsetRDG), 0 );
 
 		// Enough space for all physical pages that might be allocated
 		FRDGBufferRef CachedPageInfosRDG = GraphBuilder.CreateBuffer(FRDGBufferDesc::CreateStructuredDesc(sizeof(FCachedPageInfo), CommonParameters.MaxPhysicalPages), TEXT("CachedPageInfos"));
@@ -897,27 +897,27 @@ void FVirtualShadowMapArray::GeneratePageFlagsFromLightGrid(FRDGBuilder& GraphBu
 
 		{
 			// Note: does not actually need mip0 so can be trimmed down a bit
-			FRDGBufferRef CoverageSummary = GraphBuilder.CreateBuffer( FRDGBufferDesc::CreateBufferDesc( 4, NumPageFlags ), TEXT("CoverageSummary") );
+			FRDGBufferRef CoverageSummary = GraphBuilder.CreateBuffer( FRDGBufferDesc::CreateStructuredDesc( 4, NumPageFlags ), TEXT("CoverageSummary") );
 			
-			AddClearUAVPass( GraphBuilder, GraphBuilder.CreateUAV( CoverageSummary, PF_R32_UINT ), 0 );
+			AddClearUAVPass( GraphBuilder, GraphBuilder.CreateUAV( CoverageSummary ), 0 );
 			
 			// Run a pass to create page mappings
 			FCreatePageMappingsCS::FParameters* PassParameters = GraphBuilder.AllocParameters< FCreatePageMappingsCS::FParameters >();
 			PassParameters->CommonParameters = CommonParameters;
 
-			PassParameters->PageRequestFlags = GraphBuilder.CreateSRV(PageRequestFlagsRDG, PF_R32_UINT );
-			PassParameters->CoverageSummaryInOut	= GraphBuilder.CreateUAV( CoverageSummary, PF_R32_UINT );
-			PassParameters->AllocatedPagesOffset	= GraphBuilder.CreateUAV( AllocatedPagesOffsetRDG, PF_R32_UINT );
-			PassParameters->OutPageTable			= GraphBuilder.CreateUAV( PageTableRDG, PF_R32G32_UINT );
-			PassParameters->OutCachedPageInfos = GraphBuilder.CreateUAV( CachedPageInfosRDG );
-			PassParameters->OutPageFlags		= GraphBuilder.CreateUAV( PageFlagsRDG, PF_R32_UINT);
+			PassParameters->PageRequestFlags		= GraphBuilder.CreateSRV(PageRequestFlagsRDG );
+			PassParameters->CoverageSummaryInOut	= GraphBuilder.CreateUAV( CoverageSummary );
+			PassParameters->AllocatedPagesOffset	= GraphBuilder.CreateUAV( AllocatedPagesOffsetRDG);
+			PassParameters->OutPageTable			= GraphBuilder.CreateUAV( PageTableRDG );
+			PassParameters->OutCachedPageInfos		= GraphBuilder.CreateUAV( CachedPageInfosRDG );
+			PassParameters->OutPageFlags			= GraphBuilder.CreateUAV( PageFlagsRDG );
 
 			bool bCacheDataAvailable = VirtualShadowMapArrayCacheManager && VirtualShadowMapArrayCacheManager->IsValid();
 			if (bCacheDataAvailable)
 			{
 				SetCacheDataShaderParameters(GraphBuilder, ShadowMaps, VirtualShadowMapArrayCacheManager, PassParameters->CacheDataParameters);
 			}
-			PassParameters->OutStatsBuffer = StatsBufferRDG ? GraphBuilder.CreateUAV(StatsBufferRDG, PF_R32_UINT) : nullptr;
+			PassParameters->OutStatsBuffer = StatsBufferRDG ? GraphBuilder.CreateUAV(StatsBufferRDG) : nullptr;
 
 			// Invoked one CS thread per 2x2 mip0 texels (i.e. one per mip1 texel)
 			const uint32 DispatchWidthThreads = FVirtualShadowMap::Level0DimPagesXY >> 1;
@@ -939,8 +939,8 @@ void FVirtualShadowMapArray::GeneratePageFlagsFromLightGrid(FRDGBuilder& GraphBu
 			// Run pass building hierarchical page flags to make culling acceptable performance.
 			FGenerateHierarchicalPageFlagsCS::FParameters* PassParameters = GraphBuilder.AllocParameters< FGenerateHierarchicalPageFlagsCS::FParameters >();
 			PassParameters->CommonParameters = CommonParameters;
-			PassParameters->OutHPageFlags = GraphBuilder.CreateUAV(HPageFlagsRDG, PF_R32_UINT);
-			PassParameters->PageFlags = GraphBuilder.CreateSRV(PageFlagsRDG, PF_R32_UINT);
+			PassParameters->OutHPageFlags = GraphBuilder.CreateUAV(HPageFlagsRDG);
+			PassParameters->PageFlags = GraphBuilder.CreateSRV(PageFlagsRDG);
 			PassParameters->PageRectBoundsOut = GraphBuilder.CreateUAV(PageRectBoundsRDG);
 
 			auto ComputeShader = Views[0].ShaderMap->GetShader<FGenerateHierarchicalPageFlagsCS>();
@@ -1040,8 +1040,8 @@ void FVirtualShadowMapArray::RenderDebugInfo(FRDGBuilder& GraphBuilder, FVirtual
 		FDebugVisualizeVirtualSmCS::FParameters* PassParameters = GraphBuilder.AllocParameters<FDebugVisualizeVirtualSmCS::FParameters>();
 		PassParameters->CommonParameters = CommonParameters;
 
-		PassParameters->PageFlags			= GraphBuilder.CreateSRV( GraphBuilder.RegisterExternalBuffer( PageFlags ), PF_R32_UINT );
-		PassParameters->HPageFlags			= GraphBuilder.CreateSRV( GraphBuilder.RegisterExternalBuffer( HPageFlags ), PF_R32_UINT );
+		PassParameters->PageFlags			= GraphBuilder.CreateSRV( GraphBuilder.RegisterExternalBuffer( PageFlags ) );
+		PassParameters->HPageFlags			= GraphBuilder.CreateSRV( GraphBuilder.RegisterExternalBuffer( HPageFlags ) );
 
 		PassParameters->PhysicalPagePool	= RegisterExternalTextureWithFallback( GraphBuilder, PhysicalPagePool, GSystemTextures.BlackDummy, TEXT("PhysicalPagePool") );
 		PassParameters->PageTable			= GraphBuilder.CreateSRV( GraphBuilder.RegisterExternalBuffer( PageTable, TEXT("PageTable") ) );
@@ -1103,7 +1103,7 @@ void FVirtualShadowMapArray::PrintStats(FRDGBuilder& GraphBuilder, const FViewIn
 			FVirtualSmPrintStatsCS::FParameters* PassParameters = GraphBuilder.AllocParameters<FVirtualSmPrintStatsCS::FParameters>();
 
 			ShaderPrint::SetParameters(View, PassParameters->ShaderPrintStruct);
-			PassParameters->InStatsBuffer = GraphBuilder.CreateSRV(GraphBuilder.RegisterExternalBuffer(StatsBufferRef), PF_R32_UINT);
+			PassParameters->InStatsBuffer = GraphBuilder.CreateSRV(GraphBuilder.RegisterExternalBuffer(StatsBufferRef));
 			PassParameters->CommonParameters = CommonParameters;
 
 			auto ComputeShader = View.ShaderMap->GetShader<FVirtualSmPrintStatsCS>();
