Fixes RTGI composition to the scene color

diff --git a/Engine/Shaders/Private/DiffuseIndirectComposite.usf b/Engine/Shaders/Private/DiffuseIndirectComposite.usf
--- a/Engine/Shaders/Private/DiffuseIndirectComposite.usf
+++ b/Engine/Shaders/Private/DiffuseIndirectComposite.usf
@@ -26,6 +26,15 @@
 	#define CONFIG_SIGNAL_INPUT_LAYOUT  SIGNAL_BUFFER_LAYOUT_DIFFUSE_PROBE_HIERARCHY_HISTORY
 	#define CONFIG_INPUT_TEXTURE_COUNT 3
 
+#elif DIM_APPLY_DIFFUSE_INDIRECT == 3 // RTGI
+	#define CONFIG_SIGNAL_PROCESSING SIGNAL_PROCESSING_DIFFUSE_INDIRECT_AND_AO
+	#define COMPILE_SIGNAL_COLOR 1
+	#define MAX_SIGNAL_BATCH_SIZE 1
+	#define SIGNAL_ARRAY_SIZE 1
+
+	#define CONFIG_SIGNAL_INPUT_LAYOUT  SIGNAL_BUFFER_LAYOUT_DIFFUSE_INDIRECT_AND_AO_HISTORY
+	#define CONFIG_INPUT_TEXTURE_COUNT 2
+
 #else // !DIM_APPLY_DIFFUSE_INDIRECT
 	// NOP
 #endif // !DIM_APPLY_DIFFUSE_INDIRECT
@@ -185,7 +194,7 @@ void MainPS(
 				Sample = UncompressedAccumulators.Array[0].Moment1;
 			#endif
 				
-			#if DIM_APPLY_DIFFUSE_INDIRECT == 1
+			#if DIM_APPLY_DIFFUSE_INDIRECT == 1 || DIM_APPLY_DIFFUSE_INDIRECT == 3
 			{
 				DiffuseIndirectLighting = Sample.SceneColor.rgb;
 			}
diff --git a/Engine/Source/Runtime/Renderer/Private/IndirectLightRendering.cpp b/Engine/Source/Runtime/Renderer/Private/IndirectLightRendering.cpp
--- a/Engine/Source/Runtime/Renderer/Private/IndirectLightRendering.cpp
+++ b/Engine/Source/Runtime/Renderer/Private/IndirectLightRendering.cpp
@@ -41,11 +41,6 @@ static TAutoConsoleVariable<int32> CVarDiffuseIndirectDenoiser(
 	TEXT("Denoising options (default = 1)"),
 	ECVF_RenderThreadSafe);
 
-static TAutoConsoleVariable<int32> CVarDiffuseIndirectLumenDenoiser(
-	TEXT("r.DiffuseIndirect.LumenDenoiser"), 0,
-	TEXT("Whether lumen should use SSD for denoising."),
-	ECVF_RenderThreadSafe);
-
 static TAutoConsoleVariable<int32> CVarLumenProbeHierarchy(
 	TEXT("r.Lumen.ProbeHierarchy"), 0,
 	TEXT("Whether to use probe based denoiser for all indirect lighting."),
@@ -98,7 +93,7 @@ class FDiffuseIndirectCompositePS : public FGlobalShader
 	DECLARE_GLOBAL_SHADER(FDiffuseIndirectCompositePS)
 	SHADER_USE_PARAMETER_STRUCT(FDiffuseIndirectCompositePS, FGlobalShader)
 
-	class FApplyDiffuseIndirectDim : SHADER_PERMUTATION_INT("DIM_APPLY_DIFFUSE_INDIRECT", 3);
+	class FApplyDiffuseIndirectDim : SHADER_PERMUTATION_INT("DIM_APPLY_DIFFUSE_INDIRECT", 4);
 	class FUpscaleDiffuseIndirectDim : SHADER_PERMUTATION_INT("DIM_UPSCALE_DIFFUSE_INDIRECT", 3);
 
 	using FPermutationDomain = TShaderPermutationDomain<FApplyDiffuseIndirectDim, FUpscaleDiffuseIndirectDim>;
@@ -112,14 +107,8 @@ class FDiffuseIndirectCompositePS : public FGlobalShader
 
 		FPermutationDomain PermutationVector(Parameters.PermutationId);
 
-		// Upscale diffuse indirect only if applied.
-		if (!PermutationVector.Get<FApplyDiffuseIndirectDim>() && PermutationVector.Get<FUpscaleDiffuseIndirectDim>())
-		{
-			return false;
-		}
-
-		// Don't upscale probe hierarchy
-		if (PermutationVector.Get<FApplyDiffuseIndirectDim>() == 2 && PermutationVector.Get<FUpscaleDiffuseIndirectDim>())
+		// Only upscale SSGI
+		if (PermutationVector.Get<FApplyDiffuseIndirectDim>() != 1 && PermutationVector.Get<FUpscaleDiffuseIndirectDim>())
 		{
 			return false;
 		}
@@ -316,16 +305,8 @@ void FDeferredShadingSceneRenderer::CommitIndirectLightingState()
 		}
 		else if (ShouldRenderLumenDiffuseGI(View))
 		{
-			if (CVarLumenProbeHierarchy.GetValueOnRenderThread())
-			{
-				DiffuseIndirectMethod = EDiffuseIndirectMethod::Lumen;
-				bUseLumenProbeHierarchy = true;
-			}
-			else
-			{
-				DiffuseIndirectMethod = EDiffuseIndirectMethod::Lumen;
-				DiffuseIndirectDenoiser = IScreenSpaceDenoiser::GetDenoiserMode(CVarDiffuseIndirectLumenDenoiser);
-			}
+			DiffuseIndirectMethod = EDiffuseIndirectMethod::Lumen;
+			bUseLumenProbeHierarchy = CVarLumenProbeHierarchy.GetValueOnRenderThread() != 0;
 		}
 		
 		if (bUseLumenProbeHierarchy)
@@ -1003,16 +984,6 @@ void FDeferredShadingSceneRenderer::RenderDiffuseIndirectAndAmbientOcclusion(
 
 				AmbientOcclusionMask = DenoiserOutputs.Textures[1];
 			}
-			else if (ViewPipelineState.DiffuseIndirectMethod == EDiffuseIndirectMethod::Lumen)
-			{
-				DenoiserOutputs = DenoiserToUse->DenoiseDiffuseIndirectHarmonic(
-					GraphBuilder,
-					View,
-					&View.PrevViewInfo,
-					SceneTextures,
-					DenoiserSphericalHarmonicInputs,
-					CommonDiffuseParameters);
-			}
 			else if (ViewPipelineState.bEnableSSGI)
 			{
 				DenoiserOutputs = DenoiserToUse->DenoiseScreenSpaceDiffuseIndirect(
@@ -1137,10 +1108,15 @@ void FDeferredShadingSceneRenderer::RenderDiffuseIndirectAndAmbientOcclusion(
 					PermutationVector.Set<FDiffuseIndirectCompositePS::FApplyDiffuseIndirectDim>(2);
 					DiffuseIndirectSampling = TEXT("ProbeHierarchy");
 				}
+				else if (ViewPipelineState.DiffuseIndirectMethod == EDiffuseIndirectMethod::RTGI)
+				{
+					PermutationVector.Set<FDiffuseIndirectCompositePS::FApplyDiffuseIndirectDim>(3);
+					DiffuseIndirectSampling = TEXT("RTGI");
+				}
 				else
 				{
 					PermutationVector.Set<FDiffuseIndirectCompositePS::FApplyDiffuseIndirectDim>(1);
-					DiffuseIndirectSampling = TEXT("LightingBuffer");
+					DiffuseIndirectSampling = TEXT("SSGI");
 				}
 
 				PermutationVector.Set<FDiffuseIndirectCompositePS::FUpscaleDiffuseIndirectDim>(DenoiserOutputs.Textures[0]->Desc.Extent != SceneColor->Desc.Extent);
