Makes static analysis happy in the denoiser

diff --git a/Engine/Source/Runtime/Renderer/Private/ScreenSpaceDenoise.cpp b/Engine/Source/Runtime/Renderer/Private/ScreenSpaceDenoise.cpp
--- a/Engine/Source/Runtime/Renderer/Private/ScreenSpaceDenoise.cpp
+++ b/Engine/Source/Runtime/Renderer/Private/ScreenSpaceDenoise.cpp
@@ -2201,7 +2201,7 @@ static void DenoiseSignalAtConstantPixelDensity(
 					{
 						Dest = &View.ViewState->PrevFrameViewInfo.CompressedOpaqueDepth;
 					}
-					else if (i == 1)
+					else // if (i == 1)
 					{
 						Dest = &View.ViewState->PrevFrameViewInfo.CompressedOpaqueShadingModel;
 					}
