Added a parallel rendering thread task scope tag to fix IsInParallelRenderingThread() asserts on Frosty startup (from Yuriy)

diff --git a/Engine/Source/Runtime/Renderer/Private/DeferredShadingRenderer.cpp b/Engine/Source/Runtime/Renderer/Private/DeferredShadingRenderer.cpp
--- a/Engine/Source/Runtime/Renderer/Private/DeferredShadingRenderer.cpp
+++ b/Engine/Source/Runtime/Renderer/Private/DeferredShadingRenderer.cpp
@@ -1086,6 +1086,7 @@ bool FDeferredShadingSceneRenderer::GatherRayTracingWorldInstances(FRHICommandLi
 			ReferenceView.AddRayTracingMeshBatchTaskList.Add(FFunctionGraphTask::CreateAndDispatchWhenReady(
 				[&ReferenceView, Scene = this->Scene, Batch, BatchStart, BatchEnd]()
 			{
+				FTaskTagScope TaskTagScope(ETaskTag::EParallelRenderingThread);
 				TRACE_CPUPROFILER_EVENT_SCOPE(RayTracingMeshBatchTask);
 
 				for (uint32 Index = BatchStart; Index < BatchEnd; Index++)
