Move virtual shadow map HZB into cache manager instead of a function static variable.

diff --git a/Engine/Source/Runtime/Renderer/Private/ScenePrivate.h b/Engine/Source/Runtime/Renderer/Private/ScenePrivate.h
--- a/Engine/Source/Runtime/Renderer/Private/ScenePrivate.h
+++ b/Engine/Source/Runtime/Renderer/Private/ScenePrivate.h
@@ -2365,14 +2365,6 @@ public:
 	{}
 };
 
-struct FVirtualShadowMapPrevHZB
-{
-	FViewMatrices ViewMatrices;
-	FIntRect	  ViewRect;
-	uint32		  TargetLayerIndex = INDEX_NONE;
-	uint32		  FrameNumber = 0;
-};
-
 #if WITH_EDITOR
 class FPixelInspectorData
 {
@@ -2658,9 +2650,6 @@ public:
 	/** Map from light id to the cached shadowmap data for that light. */
 	TMap<int32, FCachedShadowMapData> CachedShadowMaps;
 
-	/** Map from light id to previous virtual shadow map HZB info. */
-	TMap< int32, FVirtualShadowMapPrevHZB > PrevVirtualShadowMapHZBs;
-
 	TRefCountPtr<IPooledRenderTarget> PreShadowCacheDepthZ;
 
 	/** Preshadows that are currently cached in the PreshadowCache render target. */
diff --git a/Engine/Source/Runtime/Renderer/Private/ShadowDepthRendering.cpp b/Engine/Source/Runtime/Renderer/Private/ShadowDepthRendering.cpp
--- a/Engine/Source/Runtime/Renderer/Private/ShadowDepthRendering.cpp
+++ b/Engine/Source/Runtime/Renderer/Private/ShadowDepthRendering.cpp
@@ -41,6 +41,7 @@
 #include "GPUScene.h"
 #include "SceneTextureReductions.h"
 #include "RendererModule.h"
+#include "VirtualShadowMaps/VirtualShadowMapCacheManager.h"
 #include "VirtualShadowMaps/VirtualShadowMapClipmap.h"
 
 DECLARE_GPU_DRAWCALL_STAT_NAMED(ShadowDepths, TEXT("Shadow Depths"));
@@ -1837,18 +1838,16 @@ void FSceneRenderer::RenderShadowDepthMaps(FRHICommandListImmediate& RHICmdList)
 	if (SortedShadowsForShadowDepthPass.VirtualShadowMapShadows.Num() > 0 ||
 		SortedShadowsForShadowDepthPass.VirtualShadowMapClipmaps.Num() > 0)
 	{
-		static TRefCountPtr<IPooledRenderTarget>	HZBPhysical = nullptr;
-		static TRefCountPtr<FPooledRDGBuffer>		HZBPageTable = nullptr;
-		static uint32								HZBFrameNumber = 0;
-
 		if( CVarNaniteShadowsUseHZB.GetValueOnRenderThread() )
 		{
-			VirtualShadowMapArray.HZBPhysical	= HZBPhysical;
-			VirtualShadowMapArray.HZBPageTable	= HZBPageTable;
+			VirtualShadowMapArray.HZBPhysical	= Scene->VirtualShadowMapArrayCacheManager->HZBPhysical;
+			VirtualShadowMapArray.HZBPageTable	= Scene->VirtualShadowMapArrayCacheManager->HZBPageTable;
 		}
 
-		uint32 CachedFrameNumber = HZBFrameNumber++;
-
+		FVirtualShadowMapArrayCacheManager *CacheManager = Scene->VirtualShadowMapArrayCacheManager;
+		const uint32 CachedFrameNumber = CacheManager->HZBFrameNumber;
+		const uint32 CurrentFrameNumber = ++CacheManager->HZBFrameNumber;
+		
 		FRDGBuilder GraphBuilder(RHICmdList);
 		{
 			RDG_EVENT_SCOPE(GraphBuilder, "Render Virtual Shadow Maps");
@@ -1869,6 +1868,8 @@ void FSceneRenderer::RenderShadowDepthMaps(FRHICommandListImmediate& RHICmdList)
 				&GraphBuilder,
 				bUpdateStreaming,
 				Scene = Scene,
+				CacheManager = CacheManager,
+				CurrentFrameNumber,
 				CachedFrameNumber
 			](bool bShouldClampToNearPlane, const FString &VirtualFilterName)
 			{
@@ -1900,10 +1901,10 @@ void FSceneRenderer::RenderShadowDepthMaps(FRHICommandListImmediate& RHICmdList)
 							Params.TargetLayerIndex = Clipmap->GetVirtualShadowMap(ClipmapLevelIndex)->ID;
 							Params.ViewMatrices = Clipmap->GetViewMatrices(ClipmapLevelIndex);
 
-							// TODO: Clean this up
+							// TODO: Clean this up - could be stored in a single structure for the whole clipmap
+							int32 AbsoluteClipmapLevel = Clipmap->GetClipmapLevel(ClipmapLevelIndex);		// NOTE: Can be negative!
 							int32 HZBKey = Clipmap->GetLightSceneInfo().Id;
-							HZBKey += Clipmap->GetClipmapLevel(ClipmapLevelIndex) << 26;
-							FVirtualShadowMapPrevHZB& PrevHZB = Scene->PrevVirtualShadowMapHZBs.FindOrAdd(HZBKey);
+							FVirtualShadowMapHZBMetadata& PrevHZB = CacheManager->HZBMetadata.FindOrAdd(HZBKey);
 							if (PrevHZB.FrameNumber == CachedFrameNumber)
 							{
 								Params.PrevTargetLayerIndex = PrevHZB.TargetLayerIndex;
@@ -1916,7 +1917,7 @@ void FSceneRenderer::RenderShadowDepthMaps(FRHICommandListImmediate& RHICmdList)
 							}
 
 							PrevHZB.TargetLayerIndex = Params.TargetLayerIndex;
-							PrevHZB.FrameNumber = HZBFrameNumber;
+							PrevHZB.FrameNumber = CurrentFrameNumber;
 							PrevHZB.ViewMatrices = Params.ViewMatrices;
 
 							Nanite::FPackedView View = Nanite::CreatePackedView(Params);
@@ -1926,13 +1927,10 @@ void FSceneRenderer::RenderShadowDepthMaps(FRHICommandListImmediate& RHICmdList)
 					}
 				}
 
-				for (int32 VirtualShadowMapIndex = 0; VirtualShadowMapIndex < SortedShadowsForShadowDepthPass.VirtualShadowMapShadows.Num(); ++VirtualShadowMapIndex)
+				for (FProjectedShadowInfo* ProjectedShadowInfo : SortedShadowsForShadowDepthPass.VirtualShadowMapShadows)
 				{
-					FProjectedShadowInfo *ProjectedShadowInfo = SortedShadowsForShadowDepthPass.VirtualShadowMapShadows[VirtualShadowMapIndex];
-					if (ProjectedShadowInfo->ShouldClampToNearPlane() == bShouldClampToNearPlane)
+					if (ProjectedShadowInfo->ShouldClampToNearPlane() == bShouldClampToNearPlane && ProjectedShadowInfo->VirtualShadowMap)
 					{
-						ensure(ProjectedShadowInfo->VirtualShadowMap);
-						
 						Nanite::FPackedViewParams Params;
 						Params.ViewMatrices = ProjectedShadowInfo->ShadowDepthView->ViewMatrices;
 						Params.ViewRect = ProjectedShadowInfo->GetViewRectForView();
@@ -1944,9 +1942,8 @@ void FSceneRenderer::RenderShadowDepthMaps(FRHICommandListImmediate& RHICmdList)
 						Params.TargetMipCount = FVirtualShadowMap::MaxMipLevels;
 
 						int32 HZBKey = ProjectedShadowInfo->GetLightSceneInfo().Id;
-						HZBKey += ProjectedShadowInfo->CascadeSettings.ShadowSplitIndex << 28;
-						FVirtualShadowMapPrevHZB& PrevHZB = Scene->PrevVirtualShadowMapHZBs.FindOrAdd( HZBKey );
-
+						HZBKey += FMath::Max(0, ProjectedShadowInfo->CascadeSettings.ShadowSplitIndex) << 28;
+						FVirtualShadowMapHZBMetadata& PrevHZB = CacheManager->HZBMetadata.FindOrAdd(HZBKey);
 						if( PrevHZB.FrameNumber == CachedFrameNumber )
 						{
 							Params.PrevTargetLayerIndex = PrevHZB.TargetLayerIndex;
@@ -1959,7 +1956,7 @@ void FSceneRenderer::RenderShadowDepthMaps(FRHICommandListImmediate& RHICmdList)
 						}
 						
 						PrevHZB.TargetLayerIndex = Params.TargetLayerIndex;
-						PrevHZB.FrameNumber = HZBFrameNumber;
+						PrevHZB.FrameNumber = CurrentFrameNumber;
 						PrevHZB.ViewMatrices = Params.ViewMatrices;
 
 						Nanite::FPackedView View = Nanite::CreatePackedView(Params);
@@ -1967,6 +1964,7 @@ void FSceneRenderer::RenderShadowDepthMaps(FRHICommandListImmediate& RHICmdList)
 						VirtualShadowMapFlags[ProjectedShadowInfo->VirtualShadowMap->ID] = 1;
 					}
 				}
+
 				if (VirtualShadowViews.Num() > 0)
 				{
 					// Update page state for all virtual shadow maps included in this call, it is a bit crap but...
@@ -2043,8 +2041,8 @@ void FSceneRenderer::RenderShadowDepthMaps(FRHICommandListImmediate& RHICmdList)
 		}
 		GraphBuilder.Execute();
 
-		HZBPhysical		= VirtualShadowMapArray.HZBPhysical;
-		HZBPageTable	= VirtualShadowMapArray.PageTable;
+		Scene->VirtualShadowMapArrayCacheManager->HZBPhysical  = VirtualShadowMapArray.HZBPhysical;
+		Scene->VirtualShadowMapArrayCacheManager->HZBPageTable = VirtualShadowMapArray.PageTable;
 
 		//GVisualizeTexture.SetCheckPoint(RHICmdList, VirtualShadowMapArray.PhysicalPagePool);
 	}
diff --git a/Engine/Source/Runtime/Renderer/Private/ShadowSetup.cpp b/Engine/Source/Runtime/Renderer/Private/ShadowSetup.cpp
--- a/Engine/Source/Runtime/Renderer/Private/ShadowSetup.cpp
+++ b/Engine/Source/Runtime/Renderer/Private/ShadowSetup.cpp
@@ -3141,11 +3141,11 @@ void FSceneRenderer::CreateWholeSceneProjectedShadow(
 				FWholeSceneProjectedShadowInitializer& ProjectedShadowInitializer = ProjectedShadowInitializers[ShadowIndex];
 				const bool bCopyLightFromVirtual = bCopyFallbackFromVirtual && !ProjectedShadowInitializer.bOnePassPointLightShadow;
 
-			// Modify resolution if using virtual shadow maps
+				// Modify resolution if using virtual shadow maps
 				if( bNeedsVirtualShadowMap )
-			{
+				{
 					MaxDesiredResolution = FVirtualShadowMap::PageSize;
-			}
+				}
 
 				// Round down to the nearest power of two so that resolution changes are always doubling or halving the resolution, which increases filtering stability
 				// Use the max resolution if the desired resolution is larger than that
@@ -3282,8 +3282,6 @@ void FSceneRenderer::CreateWholeSceneProjectedShadow(
 
 						const FLightSceneProxy& LightProxy = *(ProjectedShadowInfo->GetLightSceneInfo().Proxy);
 
-						const FMatrix FaceProjection = FPerspectiveMatrix(PI / 4.0f, 1, 1, 1, LightProxy.GetRadius());
-
 						ProjectedShadowInfo->OnePassShadowFaceProjectionMatrix = FReversedZPerspectiveMatrix(PI / 4.0f, 1, 1, 1, LightProxy.GetRadius());
 
 						// Light projection and bounding volume is set up relative to the light position
@@ -4873,9 +4871,7 @@ void FSceneRenderer::AllocateOnePassPointLightDepthTargets(FRHICommandListImmedi
 				ProjectedShadowInfo->bAllocated = true;
 				ProjectedShadowInfo->RenderTargets.CopyReferencesFromRenderTargets(ShadowMapCubemap.RenderTargets);
 				
-				// Set up a single depth view for one pass cubemap rendering
 				ProjectedShadowInfo->SetupShadowDepthView(RHICmdList, this);
-
 				ShadowMapCubemap.Shadows.Add(ProjectedShadowInfo);
 			}
 		}
diff --git a/Engine/Source/Runtime/Renderer/Private/VirtualShadowMaps/VirtualShadowMapCacheManager.h b/Engine/Source/Runtime/Renderer/Private/VirtualShadowMaps/VirtualShadowMapCacheManager.h
--- a/Engine/Source/Runtime/Renderer/Private/VirtualShadowMaps/VirtualShadowMapCacheManager.h
+++ b/Engine/Source/Runtime/Renderer/Private/VirtualShadowMaps/VirtualShadowMapCacheManager.h
@@ -72,6 +72,14 @@ public:
 };
 
 
+struct FVirtualShadowMapHZBMetadata
+{
+	FViewMatrices ViewMatrices;
+	FIntRect	  ViewRect;
+	uint32		  TargetLayerIndex = INDEX_NONE;
+	uint32		  FrameNumber = 0;
+};
+
 class FVirtualShadowMapArrayCacheManager
 {
 public:
@@ -121,4 +129,10 @@ public:
 	TRefCountPtr<FPooledRDGBuffer>		AccumulatedStatsBuffer;
 	bool bAccumulatingStats = false;
 	FRHIGPUBufferReadback *GPUBufferReadback = nullptr;
+
+	// HZB of physical pages from previous frame
+	TRefCountPtr<IPooledRenderTarget>			HZBPhysical = nullptr;
+	TRefCountPtr<FPooledRDGBuffer>				HZBPageTable = nullptr;
+	uint32										HZBFrameNumber = 0;
+	TMap<int32, FVirtualShadowMapHZBMetadata>	HZBMetadata;
 };
