RayTracing cvars are ECVF_Scalability

diff --git a/Engine/Source/Runtime/Renderer/Private/LightRendering.cpp b/Engine/Source/Runtime/Renderer/Private/LightRendering.cpp
--- a/Engine/Source/Runtime/Renderer/Private/LightRendering.cpp
+++ b/Engine/Source/Runtime/Renderer/Private/LightRendering.cpp
@@ -63,7 +63,7 @@ static FAutoConsoleVariableRef CVarRayTracingOcclusion(
 	GRayTracingShadows,
 	TEXT("0: use traditional rasterized shadow map\n")
 	TEXT("1: use ray tracing shadows (default)"),
-	ECVF_RenderThreadSafe
+	ECVF_RenderThreadSafe | ECVF_Scalability
 );
 
 static int32 GShadowRayTracingSamplesPerPixel = 1;
diff --git a/Engine/Source/Runtime/Renderer/Private/RayTracing/RayTracingAmbientOcclusion.cpp b/Engine/Source/Runtime/Renderer/Private/RayTracing/RayTracingAmbientOcclusion.cpp
--- a/Engine/Source/Runtime/Renderer/Private/RayTracing/RayTracingAmbientOcclusion.cpp
+++ b/Engine/Source/Runtime/Renderer/Private/RayTracing/RayTracingAmbientOcclusion.cpp
@@ -26,7 +26,7 @@ static FAutoConsoleVariableRef CVarRayTracingAmbientOcclusion(
 	TEXT("-1: Value driven by postprocess volume (default) \n")
 	TEXT(" 0: ray tracing ambient occlusion off \n")
 	TEXT(" 1: ray tracing ambient occlusion enabled"),
-	ECVF_RenderThreadSafe
+	ECVF_RenderThreadSafe | ECVF_Scalability
 );
 
 static TAutoConsoleVariable<int32> CVarUseAODenoiser(
diff --git a/Engine/Source/Runtime/Renderer/Private/RayTracing/RayTracingGlobalIllumination.cpp b/Engine/Source/Runtime/Renderer/Private/RayTracing/RayTracingGlobalIllumination.cpp
--- a/Engine/Source/Runtime/Renderer/Private/RayTracing/RayTracingGlobalIllumination.cpp
+++ b/Engine/Source/Runtime/Renderer/Private/RayTracing/RayTracingGlobalIllumination.cpp
@@ -30,7 +30,7 @@ static TAutoConsoleVariable<int32> CVarRayTracingGlobalIllumination(
 	TEXT("-1: Value driven by postprocess volume (default) \n")
 	TEXT(" 0: ray tracing global illumination off \n")
 	TEXT(" 1: ray tracing global illumination enabled"),
-	ECVF_RenderThreadSafe
+	ECVF_RenderThreadSafe | ECVF_Scalability
 );
 
 static int32 GRayTracingGlobalIlluminationSamplesPerPixel = -1;
diff --git a/Engine/Source/Runtime/Renderer/Private/RayTracing/RayTracingReflections.cpp b/Engine/Source/Runtime/Renderer/Private/RayTracing/RayTracingReflections.cpp
--- a/Engine/Source/Runtime/Renderer/Private/RayTracing/RayTracingReflections.cpp
+++ b/Engine/Source/Runtime/Renderer/Private/RayTracing/RayTracingReflections.cpp
@@ -31,7 +31,7 @@ static FAutoConsoleVariableRef CVarReflectionsMethod(
 	TEXT("-1: Value driven by postprocess volume (default) \n")
 	TEXT("0: use traditional rasterized SSR\n")
 	TEXT("1: use ray traced reflections\n"),
-	ECVF_RenderThreadSafe
+	ECVF_RenderThreadSafe | ECVF_Scalability
 );
 
 static TAutoConsoleVariable<float> CVarReflectionScreenPercentage(
diff --git a/Engine/Source/Runtime/Renderer/Private/RayTracing/RayTracingTranslucency.cpp b/Engine/Source/Runtime/Renderer/Private/RayTracing/RayTracingTranslucency.cpp
--- a/Engine/Source/Runtime/Renderer/Private/RayTracing/RayTracingTranslucency.cpp
+++ b/Engine/Source/Runtime/Renderer/Private/RayTracing/RayTracingTranslucency.cpp
@@ -26,7 +26,7 @@ static TAutoConsoleVariable<int32> CVarRayTracingTranslucency(
 	TEXT("-1: Value driven by postprocess volume (default) \n")
 	TEXT(" 0: ray tracing translucency off (use raster) \n")
 	TEXT(" 1: ray tracing translucency enabled"),
-	ECVF_RenderThreadSafe);
+	ECVF_RenderThreadSafe | ECVF_Scalability);
 
 static float GRayTracingTranslucencyMaxRoughness = -1;
 static FAutoConsoleVariableRef CVarRayTracingTranslucencyMaxRoughness(
diff --git a/Engine/Source/Runtime/Renderer/Private/RayTracing/RaytracingSkylight.cpp b/Engine/Source/Runtime/Renderer/Private/RayTracing/RaytracingSkylight.cpp
--- a/Engine/Source/Runtime/Renderer/Private/RayTracing/RaytracingSkylight.cpp
+++ b/Engine/Source/Runtime/Renderer/Private/RayTracing/RaytracingSkylight.cpp
@@ -33,7 +33,8 @@ static int32 GRayTracingSkyLight = -1;
 static FAutoConsoleVariableRef CVarRayTracingSkyLight(
 	TEXT("r.RayTracing.SkyLight"),
 	GRayTracingSkyLight,
-	TEXT("Enables ray tracing SkyLight (default = 0)")
+	TEXT("Enables ray tracing SkyLight (default = 0)"),
+	ECVF_Scalability
 );
 
 static int32 GRayTracingSkyLightSamplesPerPixel = -1;
