Fix for check firing on non power of two UDIM textures.

diff --git a/Engine/Source/Runtime/Renderer/Private/VT/TexturePagePool.cpp b/Engine/Source/Runtime/Renderer/Private/VT/TexturePagePool.cpp
--- a/Engine/Source/Runtime/Renderer/Private/VT/TexturePagePool.cpp
+++ b/Engine/Source/Runtime/Renderer/Private/VT/TexturePagePool.cpp
@@ -61,7 +61,7 @@ void FTexturePagePool::EvictAllPages(FVirtualTextureSystem* System)
 void FTexturePagePool::UnmapAllPagesForSpace(FVirtualTextureSystem* System, uint8 SpaceID, uint32 vAddress, uint32 Size, uint32 MaxLevel)
 {
 	check(Size > 0u);
-	checkf((Size >> MaxLevel) > 0u, TEXT("Size %d is too small for max level %d"), Size, MaxLevel);
+	checkf(FMath::CeilLogTwo(Size) >= MaxLevel, TEXT("Size %d is too small for max level %d"), Size, MaxLevel);
 	checkf((vAddress & (0xffffffff << (MaxLevel * 2u))) == vAddress, TEXT("vAddress %08X is not aligned to max level %d"), vAddress, MaxLevel);
 
 	const uint32 vAddressMax = vAddress + FMath::Square(Size);
