Refactored ShadowProjectionData to not have any view dependence
- simple matter of moving out view.PreTranslation handling

diff --git a/Engine/Shaders/Private/VirtualShadowMaps/PageManagement.usf b/Engine/Shaders/Private/VirtualShadowMaps/PageManagement.usf
--- a/Engine/Shaders/Private/VirtualShadowMaps/PageManagement.usf
+++ b/Engine/Shaders/Private/VirtualShadowMaps/PageManagement.usf
@@ -210,7 +210,9 @@ struct FScreenPixelData
 	float DeviceZ;
 	float SceneDepth;
 	float3 ViewPosition;
-	float3 TranslatedWorldPosition;
+	
+	// Note: Actual world position (not pre-translated)
+	float3 WorldPosition;
 
 	// NOTE: This data is only valid when executing after the base pass (bPostBasePass is true)
 	FGBufferData GBufferData;
@@ -256,7 +258,9 @@ FScreenPixelData GetScreenPixelData(uint2 PixelPos, float DeviceZ)
 	Data.ViewPosition = ViewH.xyz / ViewH.w;
 
 	float4 TranslatedWorldH = mul(Data.ClipPosition, View.ClipToTranslatedWorld);
-	Data.TranslatedWorldPosition = TranslatedWorldH.xyz / TranslatedWorldH.w;
+	float3 TranslatedWorldPosition = TranslatedWorldH.xyz / TranslatedWorldH.w;
+	// Subtract the view pre-translation to get back to world space
+	Data.WorldPosition = TranslatedWorldPosition - View.PreViewTranslation;
 
 	if (bPostBasePass != 0)
 	{
@@ -287,7 +291,8 @@ uint GetMipLevel(uint VirtualShadowMapId, FScreenPixelData Pixel)
 		float RadiusWorld = DepthScale * RadiusScreen;
 
 		// NOTE: Virtual shadow maps are currently always the same size in each axis (VSM_VIRTUAL_MAX_RESOLUTION_XY), so ShadowProj._m00 == ShadowProj._m11
-		float4 ShadowViewPosition = mul(float4(Pixel.TranslatedWorldPosition, 1.0f), ShadowMapProjectionData[VirtualShadowMapId].TranslatedWorldToShadowViewMatrix);
+		float3 ShadowViewTranslatedWorldPosition = Pixel.WorldPosition + ShadowMapProjectionData[VirtualShadowMapId].ShadowPreViewTranslation;
+		float4 ShadowViewPosition = mul(float4(ShadowViewTranslatedWorldPosition, 1.0f), ShadowMapProjectionData[VirtualShadowMapId].TranslatedWorldToShadowViewMatrix);
 		float4 RadiusClipH = mul(float4(RadiusWorld, 0.0f, ShadowViewPosition.z, 1.0f), ShadowMapProjectionData[VirtualShadowMapId].ShadowViewToClipMatrix);
 		float RadiusClip = abs(RadiusClipH.x / RadiusClipH.w);
 		Footprint = RadiusClip * float(2 * VSM_VIRTUAL_MAX_RESOLUTION_XY);
@@ -304,7 +309,7 @@ uint CalcPageOffsetFromUVMip(uint VirtualShadowMapId, float2 ShadowUv, uint MipL
 
 void MarkPage(uint VirtualShadowMapId, uint MipLevel, FScreenPixelData Pixel)
 {
-	FShadowPosition ShadowPosition = TranslatedWorldToShadow(ShadowMapProjectionData[VirtualShadowMapId], Pixel.TranslatedWorldPosition);
+	FShadowPosition ShadowPosition = WorldToShadow(ShadowMapProjectionData[VirtualShadowMapId], Pixel.WorldPosition);
 
 	// Check overlap vs the shadow map space
 	// NOTE: XY test not really needed anymore with the precise cone test in the caller, but we'll leave it for the moment
@@ -442,7 +447,7 @@ void GeneratePageFlagsFromPixels(uint3 GroupId : SV_GroupID, uint GroupIndex : S
 			const FLocalLightData LightData = GetLocalLightData(CulledLightGridData.DataStartIndex + Index, EyeIndex);
 			const uint LightGridLightIndex = ForwardLightData.CulledLightDataGrid[CulledLightGridData.DataStartIndex + Index];
 
-			float3 ToLight = normalize(LightData.LightPositionAndInvRadius.xyz - (Pixel.TranslatedWorldPosition - View.PreViewTranslation));
+			float3 ToLight = normalize(LightData.LightPositionAndInvRadius.xyz - Pixel.WorldPosition);
 			// Also do precise cone test, since froxels are pretty coarse at times.
 			if (dot(ToLight, LightData.LightDirectionAndShadowMask.xyz) < LightData.SpotAnglesAndSourceRadiusPacked.x)
 			{
diff --git a/Engine/Shaders/Private/VirtualShadowMaps/ProjectionCommon.ush b/Engine/Shaders/Private/VirtualShadowMaps/ProjectionCommon.ush
--- a/Engine/Shaders/Private/VirtualShadowMaps/ProjectionCommon.ush
+++ b/Engine/Shaders/Private/VirtualShadowMaps/ProjectionCommon.ush
@@ -8,9 +8,12 @@ ProjectionCommon.ush:
 // Must match C++ structure in VirtualShadowMapArray.h
 struct FVirtualShadowMapProjectionShaderData
 {
+	// From shadow-pretranslated world-space to shadow view space, example use: (WorldSpacePos + ShadowPreViewTranslation) * TranslatedWorldToShadowViewMatrix
 	float4x4 TranslatedWorldToShadowViewMatrix;
 	float4x4 ShadowViewToClipMatrix;
 	float4x4 TranslatedWorldToShadowUvNormalMatrix;
+	float3 ShadowPreViewTranslation;
+	float Dummy;
 	uint VirtualShadowMapId;	
 	int ClipmapLevel;
 	int ClipmapLevelCount;
@@ -46,11 +49,11 @@ struct FShadowPosition
 	float Depth;
 };
 
-FShadowPosition TranslatedWorldToShadow(FVirtualShadowMapProjectionShaderData ProjectionData, float3 TranslatedWorldPosition)
+FShadowPosition WorldToShadow(FVirtualShadowMapProjectionShaderData ProjectionData, float3 WorldPosition)
 {
 	FShadowPosition ShadowPosition;
-
-	ShadowPosition.View = mul(float4(TranslatedWorldPosition, 1.0f), ProjectionData.TranslatedWorldToShadowViewMatrix).xyz;
+	// 
+	ShadowPosition.View = mul(float4(WorldPosition + ProjectionData.ShadowPreViewTranslation, 1.0f), ProjectionData.TranslatedWorldToShadowViewMatrix).xyz;
 
 	ShadowPosition.ClipH = mul(float4(ShadowPosition.View, 1.0f), ProjectionData.ShadowViewToClipMatrix);
 
diff --git a/Engine/Shaders/Private/VirtualShadowMaps/VirtualShadowMapProjection.usf b/Engine/Shaders/Private/VirtualShadowMaps/VirtualShadowMapProjection.usf
--- a/Engine/Shaders/Private/VirtualShadowMaps/VirtualShadowMapProjection.usf
+++ b/Engine/Shaders/Private/VirtualShadowMaps/VirtualShadowMapProjection.usf
@@ -299,7 +299,7 @@ void VirtualShadowMapProjectionPS(
 
 	FVirtualShadowMapProjectionShaderData ProjectionData = ShadowMapProjectionData[ProjectionDataIndex];
 
-	FShadowPosition ShadowPosition = TranslatedWorldToShadow(ProjectionData, TranslatedWorldPosition);
+	FShadowPosition ShadowPosition = WorldToShadow(ProjectionData, TranslatedWorldPosition - View.PreViewTranslation);
 	
 	// TODO: SampleVirtualSmLevel for clipmap
 	FVirtualSmSample SmSample;
diff --git a/Engine/Source/Runtime/Renderer/Private/VirtualShadowMaps/VirtualShadowMapArray.cpp b/Engine/Source/Runtime/Renderer/Private/VirtualShadowMaps/VirtualShadowMapArray.cpp
--- a/Engine/Source/Runtime/Renderer/Private/VirtualShadowMaps/VirtualShadowMapArray.cpp
+++ b/Engine/Source/Runtime/Renderer/Private/VirtualShadowMaps/VirtualShadowMapArray.cpp
@@ -95,7 +95,7 @@ FMatrix CalcTranslatedWorldToShadowUvNormalMatrix(
 	return TranslatedWorldToShadowUv.GetTransposed().Inverse();
 }
 
-FVirtualShadowMapProjectionShaderData GetVirtualShadowMapProjectionShaderData(const FViewInfo& View, const FProjectedShadowInfo* ShadowInfo)
+FVirtualShadowMapProjectionShaderData GetVirtualShadowMapProjectionShaderData(const FProjectedShadowInfo* ShadowInfo)
 {
 	check(ShadowInfo->HasVirtualShadowMap());
 
@@ -112,6 +112,7 @@ FVirtualShadowMapProjectionShaderData GetVirtualShadowMapProjectionShaderData(co
 		check(ShadowViewRect.Max.Y == FVirtualShadowMap::VirtualMaxResolutionXY);
 	}
 
+	// TODO: Probably safer to use float max etc?
 	float MinSceneDepth = -HALF_WORLD_MAX;
 	float MaxSceneDepth = HALF_WORLD_MAX;
 	if (ShadowInfo->bDirectionalLight)
@@ -120,17 +121,14 @@ FVirtualShadowMapProjectionShaderData GetVirtualShadowMapProjectionShaderData(co
 		MaxSceneDepth = ShadowInfo->CascadeSettings.UnfadedSplitFar;
 	}
 
-	FMatrix TranslatedWorldToShadowView =
-		FTranslationMatrix(ShadowInfo->PreShadowTranslation - View.ViewMatrices.GetPreViewTranslation()) *
-		ShadowInfo->TranslatedWorldToView;
-
 	FMatrix ViewToClip = ShadowInfo->ViewToClip;
 	uint32 VirtualShadowMapId = ShadowInfo->VirtualShadowMap->ID;
 		
 	FVirtualShadowMapProjectionShaderData Data;
-	Data.TranslatedWorldToShadowViewMatrix = TranslatedWorldToShadowView;
+	Data.TranslatedWorldToShadowViewMatrix = ShadowInfo->TranslatedWorldToView;
 	Data.ShadowViewToClipMatrix = ViewToClip;
-	Data.TranslatedWorldToShadowUvNormalMatrix = CalcTranslatedWorldToShadowUvNormalMatrix(TranslatedWorldToShadowView, ViewToClip);
+	Data.TranslatedWorldToShadowUvNormalMatrix = CalcTranslatedWorldToShadowUvNormalMatrix(ShadowInfo->TranslatedWorldToView, ViewToClip);
+	Data.ShadowPreViewTranslation = FVector4(ShadowInfo->PreShadowTranslation, 666.0f);
 	Data.VirtualShadowMapId = VirtualShadowMapId;
 	Data.MinSceneDepth = MinSceneDepth;
 	Data.MaxSceneDepth = MaxSceneDepth;
@@ -699,44 +697,46 @@ void FVirtualShadowMapArray::GeneratePageFlagsFromLightGrid(FRDGBuilder& GraphBu
 		TArray<uint32, SceneRenderingAllocator> ForceTopMipVisible;
 		ForceTopMipVisible.AddDefaulted(ShadowMaps.Num());
 
-		for (int32 ViewIndex = 0; ViewIndex < Views.Num(); ++ViewIndex)
-		{
-			const FViewInfo &View = Views[ViewIndex];
-			const Nanite::FRasterResults &NaniteRasterResult = NaniteRasterResults[ViewIndex];
 
-			// Store shadow map projection data for each virtual shadow map
-			TArray<FVirtualShadowMapProjectionShaderData, SceneRenderingAllocator> ShadowMapProjectionData;
-			ShadowMapProjectionData.AddDefaulted(ShadowMaps.Num());
-			
-			// Gather directional light virtual shadow maps
-			TArray<uint32, SceneRenderingAllocator> DirectionalLightSmInds;
-			for (const FVisibleLightInfo& VisibleLightInfo : VisibleLightInfos)
+		// Store shadow map projection data for each virtual shadow map
+		TArray<FVirtualShadowMapProjectionShaderData, SceneRenderingAllocator> ShadowMapProjectionData;
+		ShadowMapProjectionData.AddDefaulted(ShadowMaps.Num());
+
+		// Gather directional light virtual shadow maps
+		TArray<uint32, SceneRenderingAllocator> DirectionalLightSmInds;
+		for (const FVisibleLightInfo& VisibleLightInfo : VisibleLightInfos)
+		{
+			for (const TSharedPtr<FVirtualShadowMapClipmap>& Clipmap : VisibleLightInfo.VirtualShadowMapClipmaps)
 			{
-				for (const TSharedPtr<FVirtualShadowMapClipmap>& Clipmap : VisibleLightInfo.VirtualShadowMapClipmaps)
+				for (int32 ClipmapLevel = 0; ClipmapLevel < Clipmap->GetLevelCount(); ++ClipmapLevel)
 				{
-					for (int32 ClipmapLevel = 0; ClipmapLevel < Clipmap->GetLevelCount(); ++ClipmapLevel)
-					{
-						uint32 ID = Clipmap->GetVirtualShadowMap(ClipmapLevel)->ID;
-						ShadowMapProjectionData[ID] = Clipmap->GetProjectionShaderData(View.ViewMatrices, ClipmapLevel);
-						DirectionalLightSmInds.Add(ID);
-					}
+					uint32 ID = Clipmap->GetVirtualShadowMap(ClipmapLevel)->ID;
+					ShadowMapProjectionData[ID] = Clipmap->GetProjectionShaderData(ClipmapLevel);
+					DirectionalLightSmInds.Add(ID);
 				}
+			}
 
-				for (FProjectedShadowInfo *ProjectedShadowInfo : VisibleLightInfo.AllProjectedShadows)
+			for (FProjectedShadowInfo* ProjectedShadowInfo : VisibleLightInfo.AllProjectedShadows)
+			{
+				if (ProjectedShadowInfo->HasVirtualShadowMap())
 				{
-					if (ProjectedShadowInfo->HasVirtualShadowMap())
+					uint32 ID = ProjectedShadowInfo->VirtualShadowMap->ID;
+					ShadowMapProjectionData[ID] = GetVirtualShadowMapProjectionShaderData(ProjectedShadowInfo);
+					ForceTopMipVisible[ID] = ProjectedShadowInfo->bForceTopMipVisible ? 1 : 0;
+
+					if (ProjectedShadowInfo->bDirectionalLight)
 					{
-						uint32 ID = ProjectedShadowInfo->VirtualShadowMap->ID;
-						ShadowMapProjectionData[ID] = GetVirtualShadowMapProjectionShaderData(View, ProjectedShadowInfo);
-						ForceTopMipVisible[ID] = ProjectedShadowInfo->bForceTopMipVisible ? 1 : 0;
-
-						if (ProjectedShadowInfo->bDirectionalLight)
-						{
-							DirectionalLightSmInds.Add(ID);
-						}
+						DirectionalLightSmInds.Add(ID);
 					}
 				}
 			}
+		}
+		FRDGBufferRef ShadowMapProjectionDataRDG = CreateStructuredBuffer(GraphBuilder, TEXT("ShadowMapProjectionData"), ShadowMapProjectionData);
+
+		for (int32 ViewIndex = 0; ViewIndex < Views.Num(); ++ViewIndex)
+		{
+			const FViewInfo &View = Views[ViewIndex];
+			const Nanite::FRasterResults &NaniteRasterResult = NaniteRasterResults[ViewIndex];
 
 			// This view contained no local lights (that were stored in the light grid), and no directional lights, so nothing to do.
 			if (View.ForwardLightingResources->LocalLightVisibleLightInfosIndex.Num() + DirectionalLightSmInds.Num() == 0)
@@ -746,7 +746,6 @@ void FVirtualShadowMapArray::GeneratePageFlagsFromLightGrid(FRDGBuilder& GraphBu
 
 			// Build light-index-in-light-grid => virtual-shadow-map-index remap, must be built for each view since they have different sub-sets of lights
 			// TODO: change this engine behaviour and instead upload lights once and for all, such that all indexes can refer to the same light set.
-			// TODO: in the mean time this could likely merge with the above loop over all visibile lights
 			TArray<uint32, SceneRenderingAllocator> VirtualShadowMapIdRemap = DirectionalLightSmInds;
 			// Note: the remap for the local lights is stored after the directional lights, such that this array is always non-empty...
 			VirtualShadowMapIdRemap.AddDefaulted(View.ForwardLightingResources->LocalLightVisibleLightInfosIndex.Num());
@@ -778,8 +777,6 @@ void FVirtualShadowMapArray::GeneratePageFlagsFromLightGrid(FRDGBuilder& GraphBu
 			
 			// Project Pixels onto SMs
 			{
-				FRDGBufferRef ShadowMapProjectionDataRDG = CreateStructuredBuffer(GraphBuilder, TEXT("ShadowMapProjectionData"), ShadowMapProjectionData);
-
 				FGeneratePageFlagsFromPixelsCS::FPermutationDomain PermutationVector;
 				PermutationVector.Set<FGeneratePageFlagsFromPixelsCS::FNaniteDepthBufferDim>(NaniteRasterResult.VisBuffer64 != nullptr && !bPostBasePass);			
 
@@ -957,6 +954,7 @@ void FVirtualShadowMapArray::GeneratePageFlagsFromLightGrid(FRDGBuilder& GraphBu
 		GraphBuilder.QueueBufferExtraction( HPageFlagsRDG,			&HPageFlags );
 		GraphBuilder.QueueBufferExtraction( PageTableRDG,			&PageTable );
 		GraphBuilder.QueueBufferExtraction( AllocatedPagesOffsetRDG,&AllocatedPagesOffset );
+		GraphBuilder.QueueBufferExtraction( ShadowMapProjectionDataRDG, &ShadowMapProjectionDataBuffer);
 		
 		// Extract page count & other stuff...
 		if (StatsBufferRDG)
diff --git a/Engine/Source/Runtime/Renderer/Private/VirtualShadowMaps/VirtualShadowMapArray.h b/Engine/Source/Runtime/Renderer/Private/VirtualShadowMaps/VirtualShadowMapArray.h
--- a/Engine/Source/Runtime/Renderer/Private/VirtualShadowMaps/VirtualShadowMapArray.h
+++ b/Engine/Source/Runtime/Renderer/Private/VirtualShadowMaps/VirtualShadowMapArray.h
@@ -65,11 +65,20 @@ public:
 };
 
 // Useful data for both the page mapping shader and the projection shader
+// as well as cached shadow maps
 struct FVirtualShadowMapProjectionShaderData
 {
+	/**
+	 * Transform from shadow-pre-translated world space to shadow view space, example use: (WorldSpacePos + ShadowPreViewTranslation) * TranslatedWorldToShadowViewMatrix
+	 * TODO: Why don't we call it a rotation and store in a 3x3? Does it ever have translation in?
+	 */
 	FMatrix TranslatedWorldToShadowViewMatrix;
 	FMatrix ShadowViewToClipMatrix;
 	FMatrix TranslatedWorldToShadowUvNormalMatrix;
+	/**
+	 * Translation from world space to shadow space (add before transform by TranslatedWorldToShadowViewMatrix).
+	 */
+	FVector4 ShadowPreViewTranslation;
 	uint32 VirtualShadowMapId;
 	
 	// These could be per-light (first/count), but convenient here and not much overhead
@@ -83,8 +92,9 @@ struct FVirtualShadowMapProjectionShaderData
 	// Seems the FMatrix forces 16-byte alignment
 	float Padding[2];
 };
+static_assert((sizeof(FVirtualShadowMapProjectionShaderData) % 16) == 0, "FVirtualShadowMapProjectionShaderData size should be a multiple of 16-bytes for alignment.");
 
-FVirtualShadowMapProjectionShaderData GetVirtualShadowMapProjectionShaderData(const FViewInfo& View, const FProjectedShadowInfo* ShadowInfo);
+FVirtualShadowMapProjectionShaderData GetVirtualShadowMapProjectionShaderData(const FProjectedShadowInfo* ShadowInfo);
 
 FMatrix CalcTranslatedWorldToShadowUvNormalMatrix(const FMatrix& TranslatedWorldToShadowView, const FMatrix& ViewToClip);
 
@@ -184,4 +194,7 @@ public:
 	// uint4 buffer with one rect for each mip level in all SMs, calculated to bound committed pages
 	// Used to clip the rect size of clusters during culling.
 	TRefCountPtr<FPooledRDGBuffer>		PageRectBounds;
+
+
+	TRefCountPtr<FPooledRDGBuffer>		ShadowMapProjectionDataBuffer;
 };
diff --git a/Engine/Source/Runtime/Renderer/Private/VirtualShadowMaps/VirtualShadowMapClipmap.cpp b/Engine/Source/Runtime/Renderer/Private/VirtualShadowMaps/VirtualShadowMapClipmap.cpp
--- a/Engine/Source/Runtime/Renderer/Private/VirtualShadowMaps/VirtualShadowMapClipmap.cpp
+++ b/Engine/Source/Runtime/Renderer/Private/VirtualShadowMaps/VirtualShadowMapClipmap.cpp
@@ -164,19 +164,16 @@ FViewMatrices FVirtualShadowMapClipmap::GetViewMatrices(int32 ClipmapIndex) cons
 	return FViewMatrices(Initializer);
 }
 
-FVirtualShadowMapProjectionShaderData FVirtualShadowMapClipmap::GetProjectionShaderData(const FViewMatrices& CameraViewMatrices, int32 ClipmapIndex) const
+FVirtualShadowMapProjectionShaderData FVirtualShadowMapClipmap::GetProjectionShaderData(int32 ClipmapIndex) const
 {
 	check(ClipmapIndex >= 0 && ClipmapIndex < LevelData.Num());
 	const FLevelData& Level = LevelData[ClipmapIndex];
 
-	FMatrix TranslatedWorldToShadowView =
-		FTranslationMatrix(-Level.WorldCenter - CameraViewMatrices.GetPreViewTranslation()) *
-		TranslatedWorldToView;
-
 	FVirtualShadowMapProjectionShaderData Data;
-	Data.TranslatedWorldToShadowViewMatrix = TranslatedWorldToShadowView;
+	Data.TranslatedWorldToShadowViewMatrix = TranslatedWorldToView;
 	Data.ShadowViewToClipMatrix = Level.ViewToClip;
-	Data.TranslatedWorldToShadowUvNormalMatrix = CalcTranslatedWorldToShadowUvNormalMatrix(TranslatedWorldToShadowView, Level.ViewToClip);
+	Data.TranslatedWorldToShadowUvNormalMatrix = CalcTranslatedWorldToShadowUvNormalMatrix(TranslatedWorldToView, Level.ViewToClip);
+	Data.ShadowPreViewTranslation = FVector4(-Level.WorldCenter, 666.0f);
 	Data.VirtualShadowMapId = Level.VirtualShadowMap->ID;
 	Data.ClipmapLevel = FirstLevel + ClipmapIndex;
 	Data.ClipmapLevelCount = LevelData.Num();
diff --git a/Engine/Source/Runtime/Renderer/Private/VirtualShadowMaps/VirtualShadowMapClipmap.h b/Engine/Source/Runtime/Renderer/Private/VirtualShadowMaps/VirtualShadowMapClipmap.h
--- a/Engine/Source/Runtime/Renderer/Private/VirtualShadowMaps/VirtualShadowMapClipmap.h
+++ b/Engine/Source/Runtime/Renderer/Private/VirtualShadowMaps/VirtualShadowMapClipmap.h
@@ -52,7 +52,6 @@ public:
 	}
 
 	FVirtualShadowMapProjectionShaderData GetProjectionShaderData(
-		const FViewMatrices& CameraViewMatrices,
 		int32 ClipmapIndex) const;
 
 private:
diff --git a/Engine/Source/Runtime/Renderer/Private/VirtualShadowMaps/VirtualShadowMapProjection.cpp b/Engine/Source/Runtime/Renderer/Private/VirtualShadowMaps/VirtualShadowMapProjection.cpp
--- a/Engine/Source/Runtime/Renderer/Private/VirtualShadowMaps/VirtualShadowMapProjection.cpp
+++ b/Engine/Source/Runtime/Renderer/Private/VirtualShadowMaps/VirtualShadowMapProjection.cpp
@@ -223,7 +223,6 @@ static void RenderVirtualShadowMapProjectionCommon(
 }
 
 static void GetVirtualShadowMapProjectionData(
-	const FViewInfo& View,
 	FProjectedShadowInfo** ShadowInfo,
 	uint32 ShadowInfoCount,
 	TArray<FVirtualShadowMapProjectionShaderData, SceneRenderingAllocator>& OutShadowMapProjectionData)
@@ -235,18 +234,17 @@ static void GetVirtualShadowMapProjectionData(
 		OutShadowMapProjectionData.AddDefaulted(ShadowInfoCount);
 		for (uint32 CascadeIndex = 0; CascadeIndex < ShadowInfoCount; ++CascadeIndex)
 		{
-			OutShadowMapProjectionData[CascadeIndex] = GetVirtualShadowMapProjectionShaderData(View, ShadowInfo[CascadeIndex]);
+			OutShadowMapProjectionData[CascadeIndex] = GetVirtualShadowMapProjectionShaderData(ShadowInfo[CascadeIndex]);
 		}
 	}
 	else
 	{
 		// Single projection data
-		OutShadowMapProjectionData.Add(GetVirtualShadowMapProjectionShaderData(View, ShadowInfo[0]));
+		OutShadowMapProjectionData.Add(GetVirtualShadowMapProjectionShaderData(ShadowInfo[0]));
 	}
 }
 
 static void GetVirtualShadowMapProjectionData(
-	const FViewInfo& View,
 	const TSharedPtr<FVirtualShadowMapClipmap>& Clipmap,
 	TArray<FVirtualShadowMapProjectionShaderData, SceneRenderingAllocator>& OutShadowMapProjectionData)
 {	
@@ -257,7 +255,7 @@ static void GetVirtualShadowMapProjectionData(
 	OutShadowMapProjectionData.AddDefaulted(ClipmapLevelCount);
 	for (int32 ClipmapLevel = 0; ClipmapLevel < ClipmapLevelCount; ++ClipmapLevel)
 	{
-		OutShadowMapProjectionData[ClipmapLevel] = Clipmap->GetProjectionShaderData(View.ViewMatrices, ClipmapLevel);
+		OutShadowMapProjectionData[ClipmapLevel] = Clipmap->GetProjectionShaderData(ClipmapLevel);
 	}
 }
 
@@ -271,7 +269,7 @@ void RenderVirtualShadowMapProjectionForDenoising(
 	FRDGTextureRef SignalTexture)
 {
 	TArray<FVirtualShadowMapProjectionShaderData, SceneRenderingAllocator> ShadowMapProjectionData;
-	GetVirtualShadowMapProjectionData(View, ShadowInfo, ShadowInfoCount, ShadowMapProjectionData);
+	GetVirtualShadowMapProjectionData(ShadowInfo, ShadowInfoCount, ShadowMapProjectionData);
 
 	FRenderTargetBinding RenderTargetBinding(SignalTexture, ERenderTargetLoadAction::EClear);
 	RenderVirtualShadowMapProjectionCommon(
@@ -294,7 +292,7 @@ void RenderVirtualShadowMapProjectionForDenoising(
 	FRDGTextureRef SignalTexture)
 {
 	TArray<FVirtualShadowMapProjectionShaderData, SceneRenderingAllocator> ShadowMapProjectionData;
-	GetVirtualShadowMapProjectionData(View, Clipmap, ShadowMapProjectionData);
+	GetVirtualShadowMapProjectionData(Clipmap, ShadowMapProjectionData);
 
 	FRenderTargetBinding RenderTargetBinding(SignalTexture, ERenderTargetLoadAction::EClear);
 	RenderVirtualShadowMapProjectionCommon(
@@ -320,7 +318,7 @@ void RenderVirtualShadowMapProjection(
 	bool bProjectingForForwardShading)
 {
 	TArray<FVirtualShadowMapProjectionShaderData, SceneRenderingAllocator> ShadowMapProjectionData;
-	GetVirtualShadowMapProjectionData(View, ShadowInfo, ShadowInfoCount, ShadowMapProjectionData);
+	GetVirtualShadowMapProjectionData(ShadowInfo, ShadowInfoCount, ShadowMapProjectionData);
 
 	FRenderTargetBinding RenderTargetBinding(ScreenShadowMaskTexture, ERenderTargetLoadAction::ELoad);
 	// Blend state should not vary per cascade
@@ -346,7 +344,7 @@ void RenderVirtualShadowMapProjection(
 	bool bProjectingForForwardShading)
 {
 	TArray<FVirtualShadowMapProjectionShaderData, SceneRenderingAllocator> ShadowMapProjectionData;
-	GetVirtualShadowMapProjectionData(View, Clipmap, ShadowMapProjectionData);
+	GetVirtualShadowMapProjectionData(Clipmap, ShadowMapProjectionData);
 
 	FRenderTargetBinding RenderTargetBinding(ScreenShadowMaskTexture, ERenderTargetLoadAction::ELoad);
 	
