Visualize Lumen Diffuse Indirect now replaces DiffuseColor with .18, and disables SSGI since screen traces need direct lighting to be visible to work.

diff --git a/Engine/Shaders/Private/DiffuseIndirectComposite.usf b/Engine/Shaders/Private/DiffuseIndirectComposite.usf
--- a/Engine/Shaders/Private/DiffuseIndirectComposite.usf
+++ b/Engine/Shaders/Private/DiffuseIndirectComposite.usf
@@ -208,7 +208,7 @@ void MainPS(
 			#endif
 		}
 
-		float3 DiffuseColor = GBuffer.DiffuseColor;
+		float3 DiffuseColor = bVisualizeDiffuseIndirect ? float3(.18f, .18f, .18f) : GBuffer.DiffuseColor;
 		float3 SpecularColor = GBuffer.SpecularColor;
 
 		if (bVisualizeDiffuseIndirect)
diff --git a/Engine/Source/Runtime/Renderer/Private/IndirectLightRendering.cpp b/Engine/Source/Runtime/Renderer/Private/IndirectLightRendering.cpp
--- a/Engine/Source/Runtime/Renderer/Private/IndirectLightRendering.cpp
+++ b/Engine/Source/Runtime/Renderer/Private/IndirectLightRendering.cpp
@@ -309,9 +309,13 @@ void FDeferredShadingSceneRenderer::CommitIndirectLightingState()
 			bUseLumenProbeHierarchy = CVarLumenProbeHierarchy.GetValueOnRenderThread() != 0;
 		}
 		
-		if (bUseLumenProbeHierarchy)
+		if (DiffuseIndirectMethod == EDiffuseIndirectMethod::Lumen)
 		{
-			// NOP
+			if (ViewFamily.EngineShowFlags.VisualizeLumenIndirectDiffuse)
+			{
+				// SSGI can only give reasonable results with direct lighting onscreen
+				bEnableSSGI = false;
+			}
 		}
 		else if (bEnableSSGI && DiffuseIndirectMethod == EDiffuseIndirectMethod::Disabled)
 		{
