Various nanite and virtual shadow refactoring:
- Remove virtual shadow cascades path and enable clipmaps instead
- Restore full resolution dense shadow maps even when virtual shadow maps are enabled. Combine dense and virtual shadow maps in the screen space shadow mask during projection.
- Remove "copy to fallback" path for virtual shadows; to be replaced by a separate explicit fallback map or similar. This will break Reverb indirect occlusion with Nanite geometry when virtual shadow maps are on for now. For functional testing, disabling virtual shadow maps will still work.
- Fix a bug with point light shadow projection matrices
- Change VirtualShadowMapId's to "int" since we set them to INDEX_NONE (-1) in a bunch of places

diff --git a/Engine/Shaders/Private/VirtualShadowMaps/PageAccessCommon.ush b/Engine/Shaders/Private/VirtualShadowMaps/PageAccessCommon.ush
--- a/Engine/Shaders/Private/VirtualShadowMaps/PageAccessCommon.ush
+++ b/Engine/Shaders/Private/VirtualShadowMaps/PageAccessCommon.ush
@@ -61,7 +61,6 @@ StructuredBuffer<uint> PrevDynamicCasterPageFlags;
 
 
 // Set by common parameters
-uint LevelOffsets[VSM_MAX_MIP_LEVELS];
 uint HPageFlagLevelOffsets[VSM_MAX_MIP_LEVELS];
 
 // Size, in entries, of one page table (could be a constant?) Also probably better to stride this over Mip levels such that offsets are more POT(?)
@@ -98,18 +97,6 @@ uint CalcLevelDimsTexels(uint Level)
 
 uint CalcLevelOffsets(uint Level)
 {
-	// return LevelOffsets[Level];
-#if 0
-	uint Offset = 0;
-	for( uint i = 0; i < Level; i++ )
-	{
-		uint LevelPageDim = VSM_LEVEL0_DIM_PAGES_XY >> i;
-		Offset += LevelPageDim * LevelPageDim;
-	}
-	return Offset;
-#else
-	// Faster version of the loop above.
-
 	// VSM_LEVEL0_DIM_PAGES_XY is a power of two, so the footprint of each mip level MipSize_i=(VSM_LEVEL0_DIM_PAGES_XY>>i)^2 is also a power of two.
 	// The binary representation of a mip size is just a single bit: 1 << log2(MipSize_i) = (1 << (2 * (VSM_LOG2_LEVEL0_DIM_PAGES_XY - i))).
 	
@@ -132,7 +119,11 @@ uint CalcLevelOffsets(uint Level)
 #endif
 
 	return 0x55555555u & Mask;
-#endif
+}
+
+uint CalcHPageFlagLevelOffsets(uint Level)
+{
+	return HPageFlagLevelOffsets[Level];
 }
 
 /**
@@ -232,7 +223,7 @@ bool OverlapsAnyValidPage( uint ShadowMapID, uint MipLevel, int2 ViewPortRectMin
 		// Root level size in pages
 		uint HLevelSizePages = VSM_LEVEL0_DIM_PAGES_XY >> (MipLevel + HMipLevel);
 		
-		int HPageLevelOffset = ShadowMapID * HPageTableSize + HPageFlagLevelOffsets[MipLevel]
+		int HPageLevelOffset = ShadowMapID * HPageTableSize + CalcHPageFlagLevelOffsets(MipLevel)
 			+ CalcLevelOffsets( HMipLevel + MipLevel)
 			- CalcLevelOffsets( MipLevel + 1U );
 
diff --git a/Engine/Shaders/Private/VirtualShadowMaps/PageManagement.usf b/Engine/Shaders/Private/VirtualShadowMaps/PageManagement.usf
--- a/Engine/Shaders/Private/VirtualShadowMaps/PageManagement.usf
+++ b/Engine/Shaders/Private/VirtualShadowMaps/PageManagement.usf
@@ -194,12 +194,11 @@ void GenerateIdentityMapping(uint3 PageAddressXY_LevelZ : SV_DispatchThreadID)
 
 Texture2D<UlongType> VisBuffer64;
 StructuredBuffer<FVirtualShadowMapProjectionShaderData> ShadowMapProjectionData;
-StructuredBuffer<uint> VirtualShadowMapIdRemap;
+StructuredBuffer<int> VirtualShadowMapIdRemap;
 uint NumDirectionalLightSmInds;
 uint bPostBasePass;
 float LodFootprintScale;
 uint LodPixelCountThreshold;
-int Clipmap;
 
 // Convenience to group up a bunch of data computed for a screen pixel
 struct FScreenPixelData
@@ -275,7 +274,7 @@ FScreenPixelData GetScreenPixelData(uint2 PixelPos)
 	return GetScreenPixelData(PixelPos, GetDeviceZFromPixel(PixelPos));
 }
 
-uint GetMipLevel(uint VirtualShadowMapId, FScreenPixelData Pixel)
+uint GetMipLevel(int VirtualShadowMapId, FScreenPixelData Pixel)
 {
 	float Footprint = float(VSM_VIRTUAL_MAX_RESOLUTION_XY);
 
@@ -390,47 +389,27 @@ void GeneratePageFlagsFromPixels(uint3 GroupId : SV_GroupID, uint GroupIndex : S
 	// Directional lights
 	if (NumDirectionalLightSmInds > 0)
 	{
-		if (Clipmap > 0)
+		// Mark one page per directional light. Lights can have varying numbers of clipmap levels allocated which
+		// will all be contiguous in the array, so we hop between them using the counts.
+		int ProjectionIndex = 0;
+		while (ProjectionIndex < NumDirectionalLightSmInds)
 		{
-			// Mark one page per directional light. Lights can have varying numbers of clipmap levels allocated
-			// so we hop between them using the counts.
-			int ProjectionIndex = 0;
-			while (ProjectionIndex < NumDirectionalLightSmInds)
-			{
-				// NOTE: We don't need the VirtualShadowMapIdRemap here since directional lights are
-				// put contiguously at the start of the array.
-				const int FirstClipmapLevel          = ShadowMapProjectionData[ProjectionIndex].ClipmapLevel;
-				const int ClipmapLevelCount          = ShadowMapProjectionData[ProjectionIndex].ClipmapLevelCount;
-				const float ClipmapResolutionLodBias = ShadowMapProjectionData[ProjectionIndex].ClipmapResolutionLodBias;
-
-				// Possibly use a single shadow projection data for the whole clipmap with just scales/biases for
-				// each level (to handle snapping). For now we maintain one per level.
-				int ClipmapLevel = CalcClipmapLevel(Pixel.ViewPosition, ClipmapResolutionLodBias);
-				int ClipmapIndex = max(0, ClipmapLevel - FirstClipmapLevel);				
-				if (ClipmapIndex < ClipmapLevelCount)
-				{
-					MarkPage(ShadowMapProjectionData[ProjectionIndex + ClipmapIndex].VirtualShadowMapId, 0, Pixel);
-				}
+			int ClipmapStartId = VirtualShadowMapIdRemap[ProjectionIndex];
 
-				ProjectionIndex += max(1, ClipmapLevelCount);		// Just to ensure termination even with bad data
-			}
-		}
-		else
-		{
-			LOOP
-			for (uint Index = 0; Index < NumDirectionalLightSmInds; ++Index)
-			{
-				uint VirtualShadowMapId = VirtualShadowMapIdRemap[Index];
-
-				float MinSmSceneDepth = ShadowMapProjectionData[VirtualShadowMapId].MinSceneDepth;
-				float MaxSmSceneDepth = ShadowMapProjectionData[VirtualShadowMapId].MaxSceneDepth;
+			const int FirstClipmapLevel          = ShadowMapProjectionData[ClipmapStartId].ClipmapLevel;
+			const int ClipmapLevelCount          = ShadowMapProjectionData[ClipmapStartId].ClipmapLevelCount;
+			const float ClipmapResolutionLodBias = ShadowMapProjectionData[ClipmapStartId].ClipmapResolutionLodBias;
 
-				if (Pixel.SceneDepth >= MinSmSceneDepth && Pixel.SceneDepth < MaxSmSceneDepth)
-				{
-					uint MipLevel = GetMipLevel(VirtualShadowMapId, Pixel);
-					MarkPage(VirtualShadowMapId, MipLevel, Pixel);
-				}
+			// Possibly use a single shadow projection data for the whole clipmap with just scales/biases for
+			// each level (to handle snapping). For now we maintain one per level.
+			int ClipmapLevel = CalcClipmapLevel(Pixel.ViewPosition, ClipmapResolutionLodBias);
+			int ClipmapIndex = max(0, ClipmapLevel - FirstClipmapLevel);
+			if (ClipmapIndex < ClipmapLevelCount)
+			{
+				MarkPage(ClipmapStartId + ClipmapIndex, 0, Pixel);
 			}
+
+			ProjectionIndex += max(1, ClipmapLevelCount);		// Max just to ensure termination even with bad data
 		}
 	}
 
@@ -455,7 +434,7 @@ void GeneratePageFlagsFromPixels(uint3 GroupId : SV_GroupID, uint GroupIndex : S
 			}
 
 			// The Virtual Shadow Remap stores directional lights first
-			uint VirtualShadowMapId = VirtualShadowMapIdRemap[NumDirectionalLightSmInds + LightGridLightIndex];
+			int VirtualShadowMapId = VirtualShadowMapIdRemap[NumDirectionalLightSmInds + LightGridLightIndex];
 			if (VirtualShadowMapId != INDEX_NONE)
 			{
 				uint MipLevel = GetMipLevel(VirtualShadowMapId, Pixel);
@@ -537,7 +516,7 @@ void GenerateHierarchicalPageFlags(uint2 ThreadId : SV_DispatchThreadID, uint2 G
 			PX >>= 1U;
 			PY >>= 1U;
 			uint HPageFlagOffset = ShadowMapID * HPageTableSize 
-				+ HPageFlagLevelOffsets[MipLevel] 
+				+ CalcHPageFlagLevelOffsets(MipLevel)
 				+ CalcLevelOffsets(HMipLevel + MipLevel) 
 				- CalcLevelOffsets(MipLevel + 1U) 
 				+ PY * HLevelSizePages + PX;
@@ -554,21 +533,6 @@ void GenerateHierarchicalPageFlags(uint2 ThreadId : SV_DispatchThreadID, uint2 G
 	}
 }
 
-
-uint NumVirtualShadowMaps;
-StructuredBuffer<uint> ShadowMapHasFallback;
-
-[numthreads( 64, 1, 1 )]
-void MarkFallbackPageFlags( uint InIndex : SV_DispatchThreadID )
-{
-	uint VirtualShadowMapId = InIndex;
-	if( VirtualShadowMapId >= NumVirtualShadowMaps || ShadowMapHasFallback[VirtualShadowMapId] == 0 )
-		return;
-
-	OutPageRequestFlags[CalcPageTableLevelOffset(VirtualShadowMapId, VSM_MAX_MIP_LEVELS-1)] = 1;
-}
-
-
 RWStructuredBuffer<uint> CoverageSummaryInOut;
 StructuredBuffer<uint> PageRequestFlags;
 
diff --git a/Engine/Shaders/Private/VirtualShadowMaps/ProjectionCommon.ush b/Engine/Shaders/Private/VirtualShadowMaps/ProjectionCommon.ush
--- a/Engine/Shaders/Private/VirtualShadowMaps/ProjectionCommon.ush
+++ b/Engine/Shaders/Private/VirtualShadowMaps/ProjectionCommon.ush
@@ -14,16 +14,13 @@ struct FVirtualShadowMapProjectionShaderData
 	float4x4 TranslatedWorldToShadowUvNormalMatrix;
 	float3 ShadowPreViewTranslation;
 	float Dummy;
-	uint VirtualShadowMapId;	
+	int VirtualShadowMapId;	
 	int ClipmapLevel;
 	int ClipmapLevelCount;
 	float ClipmapResolutionLodBias;
-
-	float MinSceneDepth;
-	float MaxSceneDepth;
-	
+		
 	// Seems the FMatrix forces 16-byte alignment, so we need to pad
-	float Padding[2];
+	//float Padding[2];
 };
 
 uint CalcMipLevelFromFootprint(float Footprint, uint MipLevelCount = VSM_MAX_MIP_LEVELS)
diff --git a/Engine/Shaders/Private/VirtualShadowMaps/VirtualShadowMapProjection.usf b/Engine/Shaders/Private/VirtualShadowMaps/VirtualShadowMapProjection.usf
--- a/Engine/Shaders/Private/VirtualShadowMaps/VirtualShadowMapProjection.usf
+++ b/Engine/Shaders/Private/VirtualShadowMaps/VirtualShadowMapProjection.usf
@@ -24,7 +24,7 @@
 
 #define MIN_OCCLUDER_DISTANCE 1e-6f
 
-// One per cascade for directional lights
+// One per clipmap level for directional lights
 StructuredBuffer< FVirtualShadowMapProjectionShaderData > ShadowMapProjectionData;
 Texture2D< uint > PhysicalPagePool;
 
@@ -33,10 +33,6 @@ uint ProjectionCount;
 int DebugOutputType;
 float ContactShadowLength;
 
-// Could be per-directional-light parameters, but global is fine for now
-int Clipmap;
-
-
 // TODO: Share LOD visualization color logic between Debug.usf and here (and possibly Nanite debug)
 float3 IntToColor(uint Index)
 {
@@ -211,26 +207,6 @@ float ComputeOccluderDistance(float4x4 ShadowViewToClip, float OccluderDepth, fl
 	return max(MIN_OCCLUDER_DISTANCE, Result);
 }
 
-// Returns negative number if pixel is not in *any* of the cascades
-int GetCascadeIndexForPixel(float Depth)
-{
-	// Determine which cascade we fall into
-	// NOTE: Since we are currently not using derivatives/filtered texture lookups there is no need to force this to remain
-	// consisent across the pixel quad.
-
-	// NOTE: Cascades are sorted in reverse order, so we'll greedily pick from the end of the list in case there is any overlap
-	int CascadeIndex;
-	for (CascadeIndex = ProjectionCount - 1; CascadeIndex >= 0; --CascadeIndex)
-	{
-		if (Depth >= ShadowMapProjectionData[CascadeIndex].MinSceneDepth &&
-			Depth < ShadowMapProjectionData[CascadeIndex].MaxSceneDepth)
-		{
-			break;
-		}
-	}
-	return CascadeIndex;
-}
-
 void VirtualShadowMapProjectionPS(
 	in float4 SvPosition : SV_Position,
 	out float4 OutputTarget : SV_Target
@@ -244,7 +220,6 @@ void VirtualShadowMapProjectionPS(
 	float3 TranslatedWorldPosition = SvPositionToTranslatedWorld(float4(SvPosition.xy, DeviceZ, 1.0f));
 		
 	// TODO: Early out on sky, etc.
-	bool bClipmap = Clipmap > 0;
 	bool bInLightRegion = true;
 	int ProjectionDataIndex = 0;
 
@@ -261,22 +236,14 @@ void VirtualShadowMapProjectionPS(
 	}
 	else if (LightType == LIGHT_TYPE_DIRECTIONAL)
 	{
-		if (bClipmap)
-		{
-			float3 ViewPosition = mul(float4(TranslatedWorldPosition, 1.0f), View.TranslatedWorldToView).xyz;
-
-			const int FirstClipmapLevel = ShadowMapProjectionData[0].ClipmapLevel;
-			const int ClipmapLevelCount = ShadowMapProjectionData[0].ClipmapLevelCount;
-			const float ClipmapResolutionLodBias = ShadowMapProjectionData[0].ClipmapResolutionLodBias;
-			int ClipmapLevel = CalcClipmapLevel(ViewPosition, ClipmapResolutionLodBias);
-			ProjectionDataIndex = max(0, ClipmapLevel - FirstClipmapLevel);
-			bInLightRegion = (ProjectionDataIndex < ClipmapLevelCount);
-		}
-		else
-		{
-			ProjectionDataIndex = GetCascadeIndexForPixel(SceneDepth);
-			bInLightRegion = (ProjectionDataIndex >= 0);
-		}
+		float3 ViewPosition = mul(float4(TranslatedWorldPosition, 1.0f), View.TranslatedWorldToView).xyz;
+
+		const int FirstClipmapLevel = ShadowMapProjectionData[0].ClipmapLevel;
+		const int ClipmapLevelCount = ShadowMapProjectionData[0].ClipmapLevelCount;
+		const float ClipmapResolutionLodBias = ShadowMapProjectionData[0].ClipmapResolutionLodBias;
+		int ClipmapLevel = CalcClipmapLevel(ViewPosition, ClipmapResolutionLodBias);
+		ProjectionDataIndex = max(0, ClipmapLevel - FirstClipmapLevel);
+		bInLightRegion = (ProjectionDataIndex < ClipmapLevelCount);
 	}
 	if (!bInLightRegion)
 	{
@@ -298,12 +265,10 @@ void VirtualShadowMapProjectionPS(
 	// NOTE: Could skip most of the rest of the shader if contact shadows say we're shadowed, but likely to be small, incoherent branching
 
 	FVirtualShadowMapProjectionShaderData ProjectionData = ShadowMapProjectionData[ProjectionDataIndex];
-
 	FShadowPosition ShadowPosition = WorldToShadow(ProjectionData, TranslatedWorldPosition - View.PreViewTranslation);
 	
-	// TODO: SampleVirtualSmLevel for clipmap
 	FVirtualSmSample SmSample;
-	if (bClipmap)
+	if (LightType == LIGHT_TYPE_DIRECTIONAL)
 	{
 		SmSample = SampleVirtualSmLevel(ProjectionData.VirtualShadowMapId, ShadowPosition.UV, 0);
 	}
@@ -389,16 +354,15 @@ void VirtualShadowMapProjectionPS(
 	// Output debug for "vis".
 	float3 Color = float3(0, 0, 0);
 
-	if (DebugOutputType == 1)		// Cascade/clipmap
+	if (DebugOutputType == 1)		// Clipmap level
 	{
-		int Index = bClipmap ? (ProjectionDataIndex) : (ProjectionCount - ProjectionDataIndex - 1);
-		Color = lerp(IntToColor(Index), OutputSignal.rrr, 0.4f);
+		Color = lerp(IntToColor(ProjectionDataIndex), OutputSignal.rrr, 0.4f);
 	}
-	else if (DebugOutputType == 2)			// Mip level/LOD
+	else if (DebugOutputType == 2)	// Mip level/LOD
 	{
 		Color = lerp(IntToColor(SmSample.MipLevel + 1), OutputSignal.rrr, 0.0f);
 	}
-	else if (DebugOutputType == 3)		// Invalid page
+	else if (DebugOutputType == 3)	// Invalid page
 	{
 		Color = SmSample.bValid ? (0.25f * OutputSignal.rrr) : float3(1, 0, 1);
 	}
diff --git a/Engine/Source/Runtime/Engine/Private/Components/DirectionalLightComponent.cpp b/Engine/Source/Runtime/Engine/Private/Components/DirectionalLightComponent.cpp
--- a/Engine/Source/Runtime/Engine/Private/Components/DirectionalLightComponent.cpp
+++ b/Engine/Source/Runtime/Engine/Private/Components/DirectionalLightComponent.cpp
@@ -832,10 +832,8 @@ private:
 
 		// Determine start and end distances to the current cascade's split planes
 		// Presence of the ray traced cascade does not change depth ranges for the shadow-mapped cascades
-		const float UnfadedSplitNear = GetSplitDistance(View, ShadowSplitIndex, bPrecomputedLightingIsValid, bIsRayTracedCascade);
-		const float UnfadedSplitFar = GetSplitDistance(View, ShadowSplitIndex + 1, bPrecomputedLightingIsValid, bIsRayTracedCascade);
-		float SplitNear = UnfadedSplitNear;
-		float SplitFar = UnfadedSplitFar;
+		float SplitNear = GetSplitDistance(View, ShadowSplitIndex, bPrecomputedLightingIsValid, bIsRayTracedCascade);
+		float SplitFar = GetSplitDistance(View, ShadowSplitIndex + 1, bPrecomputedLightingIsValid, bIsRayTracedCascade);
 
 		float FadePlane = SplitFar;
 
@@ -884,8 +882,6 @@ private:
 			OutCascadeSettings->FadePlaneLength = SplitFar - FadePlane;
 			OutCascadeSettings->CascadeBiasDistribution = ShadowCascadeBiasDistribution;
 			OutCascadeSettings->ShadowSplitIndex = (int32)ShadowSplitIndex;
-			OutCascadeSettings->UnfadedSplitNear = UnfadedSplitNear;
-			OutCascadeSettings->UnfadedSplitFar = UnfadedSplitFar;
 		}
 
 		const FSphere CascadeSphere = FDirectionalLightSceneProxy::GetShadowSplitBoundsDepthRange(View, View.ViewMatrices.GetViewOrigin(), SplitNear, SplitFar, OutCascadeSettings);
diff --git a/Engine/Source/Runtime/Engine/Public/SceneManagement.h b/Engine/Source/Runtime/Engine/Public/SceneManagement.h
--- a/Engine/Source/Runtime/Engine/Public/SceneManagement.h
+++ b/Engine/Source/Runtime/Engine/Public/SceneManagement.h
@@ -927,10 +927,6 @@ public:
 
 	/** Strength of depth bias across cascades. */
 	float CascadeBiasDistribution;
-	
-	// Near and far split as computed by the split distribution function without modifications for near and far planes.
-	float UnfadedSplitNear;
-	float UnfadedSplitFar;
 
 	FShadowCascadeSettings()
 		: SplitNear(0.0f)
@@ -942,8 +938,6 @@ public:
 		, bFarShadowCascade(false)
 		, ShadowSplitIndex(INDEX_NONE)
 		, CascadeBiasDistribution(1)
-		, UnfadedSplitNear(0.0f)
-		, UnfadedSplitFar(WORLD_MAX)
 	{
 	}
 };
diff --git a/Engine/Source/Runtime/Renderer/Private/SceneOcclusion.cpp b/Engine/Source/Runtime/Renderer/Private/SceneOcclusion.cpp
--- a/Engine/Source/Runtime/Renderer/Private/SceneOcclusion.cpp
+++ b/Engine/Source/Runtime/Renderer/Private/SceneOcclusion.cpp
@@ -1125,7 +1125,7 @@ void FSceneRenderer::BeginOcclusionTests(FRHICommandListImmediate& RHICmdList, b
 								continue;
 							}
 
-							if (ProjectedShadowInfo.HasVirtualShadowMap() || ProjectedShadowInfo.VirtualShadowMapToCopyFrom != nullptr)
+							if (ProjectedShadowInfo.HasVirtualShadowMap())
 							{
 								// Skip virtual SMs (and their fallback maps), as they overlay the same light as physical SM and have their own culling
 								// methods anyway.
diff --git a/Engine/Source/Runtime/Renderer/Private/ShadowDepthRendering.cpp b/Engine/Source/Runtime/Renderer/Private/ShadowDepthRendering.cpp
--- a/Engine/Source/Runtime/Renderer/Private/ShadowDepthRendering.cpp
+++ b/Engine/Source/Runtime/Renderer/Private/ShadowDepthRendering.cpp
@@ -1614,12 +1614,7 @@ void FSceneRenderer::RenderShadowDepthMapAtlases(FRHICommandListImmediate& RHICm
 
 				ProjectedShadowInfo->SetupShadowUniformBuffers(RHICmdList, Scene);
 				ProjectedShadowInfo->TransitionCachedShadowmap(RHICmdList, Scene);
-
-				const bool bRender = ProjectedShadowInfo->VirtualShadowMapToCopyFrom == nullptr;	// Don't render non-Nanite to fallback maps. // TODO: Bail out at an earlier point
-				if (bRender)
-				{
-					ProjectedShadowInfo->RenderDepth(RHICmdList, this, BeginShadowRenderPass, true);
-				}
+				ProjectedShadowInfo->RenderDepth(RHICmdList, this, BeginShadowRenderPass, true);
 			}
 		}
 
@@ -1687,11 +1682,7 @@ void FSceneRenderer::RenderShadowDepthMapAtlases(FRHICommandListImmediate& RHICm
 					BeginShadowRenderPass(RHICmdList, ShadowIndex == 0);
 				}
 
-				const bool bRender = ProjectedShadowInfo->VirtualShadowMapToCopyFrom == nullptr;	// Don't render non-Nanite to fallback maps. // TODO: Bail out at an earlier point
-				if (bRender)
-				{
-					ProjectedShadowInfo->RenderDepth(RHICmdList, this, BeginShadowRenderPass, false);
-				}
+				ProjectedShadowInfo->RenderDepth(RHICmdList, this, BeginShadowRenderPass, false);
 
 				if (!bForceSingleRenderPass)
 				{
@@ -1742,23 +1733,6 @@ void FSceneRenderer::RenderShadowDepthMapAtlases(FRHICommandListImmediate& RHICm
 
 					const FIntRect AtlasViewRect = ProjectedShadowInfo->GetViewRectForView();
 
-					if( ProjectedShadowInfo->VirtualShadowMapToCopyFrom )
-					{
-						check(ProjectedShadowInfo->ResolutionX == FVirtualShadowMap::PageSize);
-						check(ProjectedShadowInfo->ResolutionY == FVirtualShadowMap::PageSize);
-
-						Nanite::EmitFallbackShadowMapFromVSM(
-							GraphBuilder,
-							VirtualShadowMapArray,
-							ProjectedShadowInfo->VirtualShadowMapToCopyFrom->ID,
-							ShadowMap,
-							AtlasViewRect,
-							ProjectedShadowInfo->ShadowDepthView->ViewMatrices.GetProjectionMatrix(),
-							ProjectedShadowInfo->GetShaderDepthBias(),
-							ProjectedShadowInfo->bDirectionalLight
-						);
-					}
-					else
 					{
 						// We render the Nanite geometry into a viewport-sized raster context, then emit into the right atlas location
 						FIntRect NaniteViewRect(FIntPoint(0, 0), AtlasViewRect.Size());
@@ -1808,7 +1782,7 @@ void FSceneRenderer::RenderShadowDepthMapAtlases(FRHICommandListImmediate& RHICm
 							ProjectedShadowInfo->GetShaderDepthBias(),
 							ProjectedShadowInfo->bDirectionalLight
 						);
-					}	// don't copy from fallback
+					}
 				}
 			}
 
@@ -1845,7 +1819,9 @@ void FSceneRenderer::RenderShadowDepthMaps(FRHICommandListImmediate& RHICmdList)
 
 	if (bNaniteEnabled && (bHasVSMShadows || bHasVSMClipMaps))
 	{
-		if (CVarNaniteShadowsUseHZB.GetValueOnRenderThread())
+		bool bUseHZB = (CVarNaniteShadowsUseHZB.GetValueOnRenderThread() != 0);
+
+		if (bUseHZB)
 		{
 			VirtualShadowMapArray.HZBPhysical	= Scene->VirtualShadowMapArrayCacheManager->HZBPhysical;
 			VirtualShadowMapArray.HZBPageTable	= Scene->VirtualShadowMapArrayCacheManager->HZBPageTable;
@@ -1936,7 +1912,7 @@ void FSceneRenderer::RenderShadowDepthMaps(FRHICommandListImmediate& RHICmdList)
 
 				for (FProjectedShadowInfo* ProjectedShadowInfo : SortedShadowsForShadowDepthPass.VirtualShadowMapShadows)
 				{
-					if (ProjectedShadowInfo->ShouldClampToNearPlane() == bShouldClampToNearPlane && ProjectedShadowInfo->VirtualShadowMap)
+					if (ProjectedShadowInfo->ShouldClampToNearPlane() == bShouldClampToNearPlane && ProjectedShadowInfo->HasVirtualShadowMap())
 					{
 						Nanite::FPackedViewParams Params;
 						Params.ViewMatrices = ProjectedShadowInfo->ShadowDepthView->ViewMatrices;
@@ -2023,7 +1999,7 @@ void FSceneRenderer::RenderShadowDepthMaps(FRHICommandListImmediate& RHICmdList)
 			}
 
 
-			if( CVarNaniteShadowsUseHZB.GetValueOnRenderThread() )
+			if( bUseHZB )
 			{
 				RDG_EVENT_SCOPE(GraphBuilder, "BuildShadowHZB");
 
@@ -2179,7 +2155,7 @@ void FSceneRenderer::RenderShadowDepthMaps(FRHICommandListImmediate& RHICmdList)
 					{
 						FViewMatrices::FMinimalInitializer MatricesInitializer;
 						MatricesInitializer.ViewOrigin = -ProjectedShadowInfo->PreShadowTranslation;
-						MatricesInitializer.ViewRotationMatrix = FTranslationMatrix(-ProjectedShadowInfo->PreShadowTranslation) * ProjectedShadowInfo->OnePassShadowViewMatrices[CubemapFaceIndex];
+						MatricesInitializer.ViewRotationMatrix = ProjectedShadowInfo->OnePassShadowViewMatrices[CubemapFaceIndex];
 						MatricesInitializer.ProjectionMatrix = ProjectedShadowInfo->OnePassShadowFaceProjectionMatrix;
 						MatricesInitializer.ConstrainedViewRect = ShadowViewRect;
 
diff --git a/Engine/Source/Runtime/Renderer/Private/ShadowRendering.cpp b/Engine/Source/Runtime/Renderer/Private/ShadowRendering.cpp
--- a/Engine/Source/Runtime/Renderer/Private/ShadowRendering.cpp
+++ b/Engine/Source/Runtime/Renderer/Private/ShadowRendering.cpp
@@ -1763,11 +1763,6 @@ bool FSceneRenderer::RenderShadowProjections(FRHICommandListImmediate& RHICmdLis
 		check(RHICmdList.IsOutsideRenderPass());
 		check(!bProjectingForForwardShading);		// Not yet implemented/tested
 		check(Views.Num() == 1);					// TODO: Test multiview
-
-		// Right now we should only see one or the other type of virtual shadow map projection
-		// And currently only 1 clipmap (if present). Multi-view stuff TBD.
-		check(VisibleLightInfo.VirtualShadowMapClipmaps.Num() < 2);
-		//check(VirtualShadowMaps.Num() == 0 || VisibleLightInfo.VirtualShadowMapClipmaps.Num() == 0);
 		
 		FRDGBuilder GraphBuilder(RHICmdList);
 		FRDGTextureRef RDGScreenShadowMaskTexture = GraphBuilder.RegisterExternalTexture(ScreenShadowMaskTexture, TEXT("ScreenShadowMaskTexture"));
@@ -1812,6 +1807,7 @@ bool FSceneRenderer::RenderShadowProjections(FRHICommandListImmediate& RHICmdLis
 						// Project virtual shadow maps
 						if (VisibleLightInfo.VirtualShadowMapClipmaps.Num() > 0)
 						{
+							check(VisibleLightInfo.VirtualShadowMapClipmaps.Num() == 1);
 							RenderVirtualShadowMapProjectionForDenoising(
 								VisibleLightInfo.VirtualShadowMapClipmaps[0],
 								GraphBuilder,
@@ -1822,8 +1818,9 @@ bool FSceneRenderer::RenderShadowProjections(FRHICommandListImmediate& RHICmdLis
 						}
 						else
 						{
+							check(VirtualShadowMaps.Num() == 1);
 							RenderVirtualShadowMapProjectionForDenoising(
-								VirtualShadowMaps.GetData(), VirtualShadowMaps.Num(),
+								VirtualShadowMaps[0],
 								GraphBuilder,
 								View,
 								VirtualShadowMapArray,
@@ -1858,6 +1855,7 @@ bool FSceneRenderer::RenderShadowProjections(FRHICommandListImmediate& RHICmdLis
 					{
 						if (VisibleLightInfo.VirtualShadowMapClipmaps.Num() > 0)
 						{
+							check(VisibleLightInfo.VirtualShadowMapClipmaps.Num() == 1);
 							RenderVirtualShadowMapProjection(
 								VisibleLightInfo.VirtualShadowMapClipmaps[0],
 								GraphBuilder,
@@ -1869,8 +1867,9 @@ bool FSceneRenderer::RenderShadowProjections(FRHICommandListImmediate& RHICmdLis
 						}
 						else
 						{
+							check(VirtualShadowMaps.Num() == 1);
 							RenderVirtualShadowMapProjection(
-								VirtualShadowMaps.GetData(), VirtualShadowMaps.Num(),
+								VirtualShadowMaps[0],
 								GraphBuilder,
 								View,
 								VirtualShadowMapArray,
diff --git a/Engine/Source/Runtime/Renderer/Private/ShadowRendering.h b/Engine/Source/Runtime/Renderer/Private/ShadowRendering.h
--- a/Engine/Source/Runtime/Renderer/Private/ShadowRendering.h
+++ b/Engine/Source/Runtime/Renderer/Private/ShadowRendering.h
@@ -365,9 +365,6 @@ public:
 	/** Whether to render Nanite geometry into this shadow map. */
 	uint32 bNaniteGeometry : 1;
 
-	/** Whether top mip of virtual shadow map should be forced visible */
-	uint32 bForceTopMipVisible : 1;
-
 	/** Whether to include this shadow in the direct lighting shadow mask.
 	 * Currently this is used to disable projection of some low-resolution proxy shadow maps when using other high resolution virtual shadow maps
 	 * for direct lighting visibility. */
@@ -377,9 +374,6 @@ public:
 	/** */
 	FVirtualShadowMap *VirtualShadowMap;
 
-	/** Virtual shadow map to copy contents from */
-	FVirtualShadowMap *VirtualShadowMapToCopyFrom;
-
 	/** View projection matrices for each cubemap face, used by one pass point light shadows. */
 	TArray<FMatrix> OnePassShadowViewProjectionMatrices;
 
diff --git a/Engine/Source/Runtime/Renderer/Private/ShadowSetup.cpp b/Engine/Source/Runtime/Renderer/Private/ShadowSetup.cpp
--- a/Engine/Source/Runtime/Renderer/Private/ShadowSetup.cpp
+++ b/Engine/Source/Runtime/Renderer/Private/ShadowSetup.cpp
@@ -242,13 +242,6 @@ FAutoConsoleVariableRef CVarUseOctreeForShadowCulling(
 	ECVF_Scalability | ECVF_RenderThreadSafe
 	);
 
-static TAutoConsoleVariable<int32> CVarCopyFallbackFromVirtual(
-	TEXT( "r.Shadow.v.CopyFallbackFromVSM" ),
-	1,
-	TEXT( "When enabled, Nanite part of fallback shadows are copied from the corresponding virtual shadow map instead of being rendered." ),
-	ECVF_RenderThreadSafe
-);
-
 CSV_DECLARE_CATEGORY_EXTERN(LightCount);
 
 #if !UE_BUILD_SHIPPING
@@ -493,10 +486,8 @@ FProjectedShadowInfo::FProjectedShadowInfo()
 	, bTransmission(false)
 	, bHairStrandsDeepShadow(false)
 	, bNaniteGeometry(true)
-	, bForceTopMipVisible(false)
 	, bIncludeInScreenSpaceShadowMask(true)
 	, VirtualShadowMap(nullptr)
-	, VirtualShadowMapToCopyFrom(nullptr)
 	, PerObjectShadowFadeStart(WORLD_MAX)
 	, InvPerObjectShadowFadeLength(0.0f)
 	, LightSceneInfo(0)
@@ -3042,8 +3033,7 @@ void FSceneRenderer::CreateWholeSceneProjectedShadow(
 	// TODO: Base this off of a light/editor parameter; for now just all spot lights get virtual shadow maps when the CVar is enabled
 	static const auto EnableVirtualSMCVar = IConsoleManager::Get().FindTConsoleVariableDataInt(TEXT("r.Shadow.v.Enable"));
 	const bool bNeedsVirtualShadowMap = (EnableVirtualSMCVar->GetValueOnRenderThread() != 0 && LightSceneInfo->Proxy->GetLightType() == LightType_Spot);
-	const bool bCopyFallbackFromVirtual = bNeedsVirtualShadowMap && CVarCopyFallbackFromVirtual.GetValueOnRenderThread() != 0;
-
+	
 	// Try to create a whole-scene projected shadow initializer for the light.
 	TArray<FWholeSceneProjectedShadowInitializer, TInlineAllocator<6> > ProjectedShadowInitializers;
 	if (LightSceneInfo->Proxy->GetWholeSceneProjectedShadowInitializer(ViewFamily, ProjectedShadowInitializers))
@@ -3139,20 +3129,13 @@ void FSceneRenderer::CreateWholeSceneProjectedShadow(
 			for (int32 ShadowIndex = 0, ShadowCount = ProjectedShadowInitializers.Num(); ShadowIndex < ShadowCount; ShadowIndex++)
 			{
 				FWholeSceneProjectedShadowInitializer& ProjectedShadowInitializer = ProjectedShadowInitializers[ShadowIndex];
-				const bool bCopyLightFromVirtual = bCopyFallbackFromVirtual && !ProjectedShadowInitializer.bOnePassPointLightShadow;
-
-				// Modify resolution if using virtual shadow maps
-				if( bNeedsVirtualShadowMap )
-				{
-					MaxDesiredResolution = FVirtualShadowMap::PageSize;
-				}
-
+				
 				// Round down to the nearest power of two so that resolution changes are always doubling or halving the resolution, which increases filtering stability
 				// Use the max resolution if the desired resolution is larger than that
 				// FMath::CeilLogTwo(MaxDesiredResolution + 1.0f) instead of FMath::CeilLogTwo(MaxDesiredResolution) because FMath::CeilLogTwo takes
 				// an uint32 as argument and this causes MaxDesiredResolution get truncated. For example, if MaxDesiredResolution is 256.1f,
 				// FMath::CeilLogTwo returns 8 but the next line of code expects a 9 to work correctly
-				int32 RoundedDesiredResolution = FMath::Max<int32>((1 << (FMath::CeilLogTwo(MaxDesiredResolution + 1.0f) - 1)) - (!bCopyLightFromVirtual ? ShadowBorder * 2 : 0), 1);
+				int32 RoundedDesiredResolution = FMath::Max<int32>((1 << (FMath::CeilLogTwo(MaxDesiredResolution + 1.0f) - 1)) - (ShadowBorder * 2), 1);
 				int32 SizeX = MaxDesiredResolution >= MaxShadowResolution ? MaxShadowResolution : RoundedDesiredResolution;
 				int32 SizeY = MaxDesiredResolution >= MaxShadowResolutionY ? MaxShadowResolutionY : RoundedDesiredResolution;
 
@@ -3169,8 +3152,6 @@ void FSceneRenderer::CreateWholeSceneProjectedShadow(
 				SizeX = ShadowMapSize.X - ShadowBorder * 2;
 				SizeY = ShadowMapSize.Y - ShadowBorder * 2;
 
-
-				FVirtualShadowMap* VirtualShadowMapToCopyFrom = nullptr;
 				if (bNeedsVirtualShadowMap)
 				{
 					// Create the projected shadow info.
@@ -3179,11 +3160,6 @@ void FSceneRenderer::CreateWholeSceneProjectedShadow(
 					VisibleLightInfo.MemStackProjectedShadows.Add(ProjectedShadowInfo);
 
 					ProjectedShadowInfo->VirtualShadowMap = VirtualShadowMapArray.Allocate();
-					if( bCopyLightFromVirtual )
-					{
-						ProjectedShadowInfo->bForceTopMipVisible = true;
-						VirtualShadowMapToCopyFrom = ProjectedShadowInfo->VirtualShadowMap;
-					}
 
 					// Rescale size to fit whole virtual SM but keeping aspect ratio
 					int32 VirtualSizeX = SizeX >= SizeY ? FVirtualShadowMap::VirtualMaxResolutionXY : (FVirtualShadowMap::VirtualMaxResolutionXY * SizeX) / SizeY;
@@ -3249,13 +3225,9 @@ void FSceneRenderer::CreateWholeSceneProjectedShadow(
 						false	// no RSM
 						);
 
-					ProjectedShadowInfo->VirtualShadowMapToCopyFrom = VirtualShadowMapToCopyFrom;
 					ProjectedShadowInfo->CacheMode = CacheMode[CacheModeIndex];
 					ProjectedShadowInfo->FadeAlphas = FadeAlphas;
 
-					// When using virtual shadow maps, we don't project the fallback shadow map or else we'd get double shadowing
-					ProjectedShadowInfo->bIncludeInScreenSpaceShadowMask = !bNeedsVirtualShadowMap;
-
 					VisibleLightInfo.MemStackProjectedShadows.Add(ProjectedShadowInfo);
 
 					if (ProjectedShadowInitializer.bOnePassPointLightShadow)
@@ -3307,6 +3279,12 @@ void FSceneRenderer::CreateWholeSceneProjectedShadow(
 						ProjectedShadowInfo->CasterFrustum.Init();
 					}
 
+					// If we have a virtual shadow map, disable nanite rendering into the regular shadow map or else we'd get double-shadowing
+					if (bNeedsVirtualShadowMap)
+					{
+						ProjectedShadowInfo->bNaniteGeometry = false;
+					}
+
 					bool bContainsNaniteSubjects = false;
 
 					// Ray traced shadows use the GPU managed distance field object buffers, no CPU culling should be used
@@ -4073,13 +4051,10 @@ void FSceneRenderer::AddViewDependentWholeSceneShadowsForView(
 {
 	SCOPE_CYCLE_COUNTER(STAT_AddViewDependentWholeSceneShadowsForView);
 
-	// TODO: Sort out a fallback solution for clipmaps; for now we have no way to copy "matching" data so let it render separately for consistency
 	static const auto EnableVirtualSMCVar = IConsoleManager::Get().FindTConsoleVariableDataInt(TEXT("r.Shadow.v.Enable"));
 	const bool bNeedsVirtualShadowMap = EnableVirtualSMCVar->GetValueOnRenderThread() != 0;
 	const bool bDirectionalLight        = LightSceneInfo.Proxy->GetLightType() == LightType_Directional;
-	const bool bVirtualShadowMapClipmap = bNeedsVirtualShadowMap && bDirectionalLight && FVirtualShadowMapClipmap::IsEnabled();	
-	const bool bCopyFallbackFromVirtual = bNeedsVirtualShadowMap && !bVirtualShadowMapClipmap && CVarCopyFallbackFromVirtual.GetValueOnRenderThread() != 0;
-
+	const bool bVirtualShadowMapClipmap = bNeedsVirtualShadowMap && bDirectionalLight;
 
 	// Allow each view to create a whole scene view dependent shadow
 	for (int32 ViewIndex = 0; ViewIndex < Views.Num(); ViewIndex++)
@@ -4105,7 +4080,7 @@ void FSceneRenderer::AddViewDependentWholeSceneShadowsForView(
 					break;
 				}
 			}
-		}		
+		}
 		
 		// If rendering in stereo mode we render shadow depths only for the left eye, but project for both eyes!
 		if (IStereoRendering::IsAPrimaryView(View))
@@ -4136,10 +4111,9 @@ void FSceneRenderer::AddViewDependentWholeSceneShadowsForView(
 
 				if (LightSceneInfo.Proxy->GetViewDependentWholeSceneProjectedShadowInitializer(View, LocalIndex, LightSceneInfo.IsPrecomputedLightingValid(), ProjectedShadowInitializer))
 				{
-					MaxShadowCascadeDistance = FMath::Max(MaxShadowCascadeDistance, ProjectedShadowInitializer.CascadeSettings.UnfadedSplitFar);
+					MaxShadowCascadeDistance = FMath::Max(MaxShadowCascadeDistance, ProjectedShadowInitializer.CascadeSettings.SplitFar);
 
 					// Tack on a virtual SM alongside if enabled, the regular one is still there to take care of non-nanite geo.
-					FProjectedShadowInfo* ShadowMapInfoToCopyFrom = nullptr;
 					if( bNeedsVirtualShadowMap && !bVirtualShadowMapClipmap )
 					{
 						// Create the projected shadow info.
@@ -4149,13 +4123,6 @@ void FSceneRenderer::AddViewDependentWholeSceneShadowsForView(
 
 						ProjectedShadowInfo->VirtualShadowMap = VirtualShadowMapArray.Allocate();
 						
-						if( bCopyFallbackFromVirtual )
-						{
-							ensure(ProjectedShadowInfo->VirtualShadowMap);
-							ProjectedShadowInfo->bForceTopMipVisible = true;
-							ShadowMapInfoToCopyFrom = ProjectedShadowInfo;
-						}
-
 						uint32 VirtualSMRes = FVirtualShadowMap::VirtualMaxResolutionXY;
 						TSharedPtr<FVirtualShadowMapCacheEntry> VirtualSmCacheEntry;
 						if (Scene->VirtualShadowMapArrayCacheManager)
@@ -4165,7 +4132,6 @@ void FSceneRenderer::AddViewDependentWholeSceneShadowsForView(
 						
 						ProjectedShadowInfo->VirtualShadowMap->VirtualShadowMapCacheEntry = VirtualSmCacheEntry;
 
-
 						// Set up projection as per normal
 						ProjectedShadowInfo->SetupWholeSceneProjection(
 							&LightSceneInfo,
@@ -4194,20 +4160,6 @@ void FSceneRenderer::AddViewDependentWholeSceneShadowsForView(
 						// Modify resolution if using virtual shadow maps
 						FIntPoint ShadowBufferResolution;
 						FIntPoint ShadowBufferSnapResolution;
-						if (bNeedsVirtualShadowMap)
-						{
-							// TODO: when using caching ensure the viewport is correct (since it is no longer the full page size).
-							ShadowBufferResolution = FIntPoint( FVirtualShadowMap::PageSize, FVirtualShadowMap::PageSize );
-							ShadowBufferSnapResolution = ShadowBufferResolution;
-
-							if( bCopyFallbackFromVirtual )
-							{
-								// Force snapping to VSM resolution for consistency.
-								// If we don't do this and copy the data from VSM the shadow map and the projection will not match.
-								ShadowBufferSnapResolution = FIntPoint( FVirtualShadowMap::VirtualMaxResolutionXY, FVirtualShadowMap::VirtualMaxResolutionXY );
-						}
-						}
-						else
 						{
 							const int32 MaxCSMResolution = GetCachedScalabilityCVars().MaxCSMShadowResolution;
 							ShadowBufferResolution = FIntPoint(	FMath::Clamp( MaxCSMResolution, 1, (int32)GMaxShadowDepthBufferSizeX ) - ShadowBorder * 2,
@@ -4217,41 +4169,29 @@ void FSceneRenderer::AddViewDependentWholeSceneShadowsForView(
 						
 						// Create the projected shadow info.
 						FProjectedShadowInfo* ProjectedShadowInfo = new(FMemStack::Get(), 1, 16) FProjectedShadowInfo;
-
-						// If caching is active, we should use the setup for that.
-						if (ShadowMapInfoToCopyFrom && ShadowMapInfoToCopyFrom->VirtualShadowMap->VirtualShadowMapCacheEntry)
-						{
-							*ProjectedShadowInfo = *ShadowMapInfoToCopyFrom;
-							ProjectedShadowInfo->ResolutionX = FVirtualShadowMap::PageSize;
-							ProjectedShadowInfo->ResolutionY = FVirtualShadowMap::PageSize;
-							ProjectedShadowInfo->BorderSize = ShadowBorder;
-							ProjectedShadowInfo->VirtualShadowMap = nullptr;
-							ProjectedShadowInfo->UpdateShaderDepthBias();
-						}
-						else
-						{
-							ProjectedShadowInfo->SetupWholeSceneProjection(
-								&LightSceneInfo,
-								&View,
-								ProjectedShadowInitializer,
-								ShadowBufferResolution.X,
-								ShadowBufferResolution.Y,
-								ShadowBufferSnapResolution.X,
-								ShadowBufferSnapResolution.Y,
-								ShadowBorder,
-								false	// no RSM
-							);
-						}
+						ProjectedShadowInfo->SetupWholeSceneProjection(
+							&LightSceneInfo,
+							&View,
+							ProjectedShadowInitializer,
+							ShadowBufferResolution.X,
+							ShadowBufferResolution.Y,
+							ShadowBufferSnapResolution.X,
+							ShadowBufferSnapResolution.Y,
+							ShadowBorder,
+							false	// no RSM
+						);
 						ProjectedShadowInfo->FadeAlphas = FadeAlphas;
-						ProjectedShadowInfo->VirtualShadowMapToCopyFrom = ShadowMapInfoToCopyFrom ? ShadowMapInfoToCopyFrom->VirtualShadowMap : nullptr;
-
-						// When using virtual shadow maps, we don't project the fallback shadow map or else we'd get double shadowing
-						ProjectedShadowInfo->bIncludeInScreenSpaceShadowMask = !bNeedsVirtualShadowMap;
 
 						VisibleLightInfo.MemStackProjectedShadows.Add(ProjectedShadowInfo);
 						VisibleLightInfo.AllProjectedShadows.Add(ProjectedShadowInfo);
 						ShadowInfos.Add(ProjectedShadowInfo);
 
+						// If we have a virtual shadow map, disable nanite rendering into the regular shadow map or else we'd get double-shadowing
+						if (bNeedsVirtualShadowMap)
+						{
+							ProjectedShadowInfo->bNaniteGeometry = false;
+						}
+
 						// Ray traced shadows use the GPU managed distance field object buffers, no CPU culling needed
 						if (!ProjectedShadowInfo->bRayTracedDistanceField)
 						{
@@ -4262,9 +4202,7 @@ void FSceneRenderer::AddViewDependentWholeSceneShadowsForView(
 			}
 
 			if (bVirtualShadowMapClipmap && MaxShadowCascadeDistance > 0.0f)
-			{				
-				check(!bCopyFallbackFromVirtual);		// TODO: Support in some different form
-
+			{
 				FMatrix WorldToLight = FInverseRotationMatrix(LightSceneInfo.Proxy->GetDirection().GetSafeNormal().Rotation());
 
 				TSharedPtr<FVirtualShadowMapClipmap> VirtualShadowMapClipmap = TSharedPtr<FVirtualShadowMapClipmap>(new FVirtualShadowMapClipmap(
diff --git a/Engine/Source/Runtime/Renderer/Private/VirtualShadowMaps/VirtualShadowMapArray.cpp b/Engine/Source/Runtime/Renderer/Private/VirtualShadowMaps/VirtualShadowMapArray.cpp
--- a/Engine/Source/Runtime/Renderer/Private/VirtualShadowMaps/VirtualShadowMapArray.cpp
+++ b/Engine/Source/Runtime/Renderer/Private/VirtualShadowMaps/VirtualShadowMapArray.cpp
@@ -98,6 +98,7 @@ FMatrix CalcTranslatedWorldToShadowUvNormalMatrix(
 FVirtualShadowMapProjectionShaderData GetVirtualShadowMapProjectionShaderData(const FProjectedShadowInfo* ShadowInfo)
 {
 	check(ShadowInfo->HasVirtualShadowMap());
+	check(ShadowInfo->CascadeSettings.ShadowSplitIndex == INDEX_NONE);		// We use clipmaps for virtual shadow maps, not cascades
 
 	// NOTE: Virtual shadow maps are never atlased, but verify our assumptions
 	{
@@ -112,46 +113,27 @@ FVirtualShadowMapProjectionShaderData GetVirtualShadowMapProjectionShaderData(co
 		check(ShadowViewRect.Max.Y == FVirtualShadowMap::VirtualMaxResolutionXY);
 	}
 
-	// TODO: Probably safer to use float max etc?
-	float MinSceneDepth = -HALF_WORLD_MAX;
-	float MaxSceneDepth = HALF_WORLD_MAX;
-	if (ShadowInfo->bDirectionalLight)
-	{
-		MinSceneDepth = ShadowInfo->CascadeSettings.UnfadedSplitNear;
-		MaxSceneDepth = ShadowInfo->CascadeSettings.UnfadedSplitFar;
-	}
-
 	FMatrix ViewToClip = ShadowInfo->ViewToClip;
-	uint32 VirtualShadowMapId = ShadowInfo->VirtualShadowMap->ID;
 		
 	FVirtualShadowMapProjectionShaderData Data;
 	Data.TranslatedWorldToShadowViewMatrix = ShadowInfo->TranslatedWorldToView;
 	Data.ShadowViewToClipMatrix = ViewToClip;
 	Data.TranslatedWorldToShadowUvNormalMatrix = CalcTranslatedWorldToShadowUvNormalMatrix(ShadowInfo->TranslatedWorldToView, ViewToClip);
 	Data.ShadowPreViewTranslation = FVector4(ShadowInfo->PreShadowTranslation, 666.0f);
-	Data.VirtualShadowMapId = VirtualShadowMapId;
-	Data.MinSceneDepth = MinSceneDepth;
-	Data.MaxSceneDepth = MaxSceneDepth;
+	Data.VirtualShadowMapId = ShadowInfo->VirtualShadowMap->ID;
 
 	return Data;
 }
 
 FVirtualShadowMapArray::FVirtualShadowMapArray()
 {
-	uint32 Offset = 0;
-	for (uint32 Level = 0; Level < FVirtualShadowMap::MaxMipLevels; ++Level)
-	{
-		CommonParameters.LevelOffsets[Level] = Offset;
-		uint32 LevelPageDim = FVirtualShadowMap::Level0DimPagesXY >> Level;
-		Offset += LevelPageDim * LevelPageDim;
-	}
-	CommonParameters.PageTableSize = Offset;
+	CommonParameters.PageTableSize = FVirtualShadowMap::PageTableSize;		// TODO: Define
 
 	uint32 HPageFlagOffset = 0;
 	for (uint32 Level = 0; Level < FVirtualShadowMap::MaxMipLevels - 1; ++Level)
 	{
 		CommonParameters.HPageFlagLevelOffsets[Level] = HPageFlagOffset;
-		HPageFlagOffset += CommonParameters.PageTableSize - CommonParameters.LevelOffsets[Level + 1];
+		HPageFlagOffset += CommonParameters.PageTableSize - CalcVirtualShadowMapLevelOffsets(Level + 1, FVirtualShadowMap::Log2Level0DimPagesXY);
 	}
 	// The last mip level is 1x1 and thus does not have any H levels possible.
 	CommonParameters.HPageFlagLevelOffsets[FVirtualShadowMap::MaxMipLevels - 1] = 0;
@@ -308,7 +290,6 @@ class FGeneratePageFlagsFromPixelsCS : public FVirtualPageManagementShader
 		SHADER_PARAMETER(uint32, bPostBasePass)
 		SHADER_PARAMETER(float, LodFootprintScale)
 		SHADER_PARAMETER(uint32, LodPixelCountThreshold)
-		SHADER_PARAMETER(int32, Clipmap)
 	END_SHADER_PARAMETER_STRUCT()
 };
 IMPLEMENT_GLOBAL_SHADER(FGeneratePageFlagsFromPixelsCS, "/Engine/Private/VirtualShadowMaps/PageManagement.usf", "GeneratePageFlagsFromPixels", SF_Compute);
@@ -366,20 +347,6 @@ class FCreatePageMappingsCS : public FVirtualPageManagementShader
 };
 IMPLEMENT_GLOBAL_SHADER(FCreatePageMappingsCS, "/Engine/Private/VirtualShadowMaps/PageManagement.usf", "CreatePageMappings", SF_Compute);
 
-class FMarkFallbackPageFlagsCS : public FVirtualPageManagementShader
-{
-	DECLARE_GLOBAL_SHADER( FMarkFallbackPageFlagsCS );
-	SHADER_USE_PARAMETER_STRUCT( FMarkFallbackPageFlagsCS, FVirtualPageManagementShader )
-
-	BEGIN_SHADER_PARAMETER_STRUCT( FParameters, )
-		SHADER_PARAMETER_STRUCT_INCLUDE( FVirtualShadowMapCommonParameters, CommonParameters )
-		SHADER_PARAMETER(uint32, NumVirtualShadowMaps)
-		SHADER_PARAMETER_RDG_BUFFER_SRV(StructuredBuffer<uint>, ShadowMapHasFallback)
-		SHADER_PARAMETER_RDG_BUFFER_UAV(RWStructuredBuffer<uint>, OutPageRequestFlags)
-	END_SHADER_PARAMETER_STRUCT()
-};
-IMPLEMENT_GLOBAL_SHADER( FMarkFallbackPageFlagsCS, "/Engine/Private/VirtualShadowMaps/PageManagement.usf", "MarkFallbackPageFlags", SF_Compute );
-
 class FInitClearPhysicalPagesArgsCS : public FVirtualPageManagementShader
 {
 	DECLARE_GLOBAL_SHADER(FInitClearPhysicalPagesArgsCS);
@@ -677,24 +644,20 @@ void FVirtualShadowMapArray::GeneratePageFlagsFromLightGrid(FRDGBuilder& GraphBu
 		FRDGBufferRef PageRectBoundsRDG = GraphBuilder.CreateBuffer(FRDGBufferDesc::CreateStructuredDesc(sizeof(FIntVector4), NumPageRects), TEXT("PageRectBounds"));
 		AddInitPageRectsPass(GraphBuilder, CommonParameters, PageRectBoundsRDG);
 
-		// Used to mark any shadow maps that require fallback page rendering as we go through the lights so we can force the top mip to be mapped later
-		TArray<uint32, SceneRenderingAllocator> ForceTopMipVisible;
-		ForceTopMipVisible.AddDefaulted(ShadowMaps.Num());
-
-
 		// Store shadow map projection data for each virtual shadow map
 		TArray<FVirtualShadowMapProjectionShaderData, SceneRenderingAllocator> ShadowMapProjectionData;
 		ShadowMapProjectionData.AddDefaulted(ShadowMaps.Num());
 
 		// Gather directional light virtual shadow maps
-		TArray<uint32, SceneRenderingAllocator> DirectionalLightSmInds;
+		TArray<int32, SceneRenderingAllocator> DirectionalLightSmInds;
 		for (const FVisibleLightInfo& VisibleLightInfo : VisibleLightInfos)
 		{
 			for (const TSharedPtr<FVirtualShadowMapClipmap>& Clipmap : VisibleLightInfo.VirtualShadowMapClipmaps)
 			{
+				// NOTE: Shader assumes all levels from a given clipmap are contiguous in both the remap and projection arrays
 				for (int32 ClipmapLevel = 0; ClipmapLevel < Clipmap->GetLevelCount(); ++ClipmapLevel)
 				{
-					uint32 ID = Clipmap->GetVirtualShadowMap(ClipmapLevel)->ID;
+					int32 ID = Clipmap->GetVirtualShadowMap(ClipmapLevel)->ID;
 					ShadowMapProjectionData[ID] = Clipmap->GetProjectionShaderData(ClipmapLevel);
 					DirectionalLightSmInds.Add(ID);
 				}
@@ -704,9 +667,8 @@ void FVirtualShadowMapArray::GeneratePageFlagsFromLightGrid(FRDGBuilder& GraphBu
 			{
 				if (ProjectedShadowInfo->HasVirtualShadowMap())
 				{
-					uint32 ID = ProjectedShadowInfo->VirtualShadowMap->ID;
+					int32 ID = ProjectedShadowInfo->VirtualShadowMap->ID;
 					ShadowMapProjectionData[ID] = GetVirtualShadowMapProjectionShaderData(ProjectedShadowInfo);
-					ForceTopMipVisible[ID] = ProjectedShadowInfo->bForceTopMipVisible ? 1 : 0;
 
 					if (ProjectedShadowInfo->bDirectionalLight)
 					{
@@ -730,7 +692,7 @@ void FVirtualShadowMapArray::GeneratePageFlagsFromLightGrid(FRDGBuilder& GraphBu
 
 			// Build light-index-in-light-grid => virtual-shadow-map-index remap, must be built for each view since they have different sub-sets of lights
 			// TODO: change this engine behaviour and instead upload lights once and for all, such that all indexes can refer to the same light set.
-			TArray<uint32, SceneRenderingAllocator> VirtualShadowMapIdRemap = DirectionalLightSmInds;
+			TArray<int32, SceneRenderingAllocator> VirtualShadowMapIdRemap = DirectionalLightSmInds;
 			// Note: the remap for the local lights is stored after the directional lights, such that this array is always non-empty...
 			VirtualShadowMapIdRemap.AddDefaulted(View.ForwardLightingResources->LocalLightVisibleLightInfosIndex.Num());
 			for (int32 L = 0; L < View.ForwardLightingResources->LocalLightVisibleLightInfosIndex.Num(); ++L)
@@ -791,8 +753,6 @@ void FVirtualShadowMapArray::GeneratePageFlagsFromLightGrid(FRDGBuilder& GraphBu
 				PassParameters->NumDirectionalLightSmInds = DirectionalLightSmInds.Num();
 				PassParameters->LodFootprintScale = LodFootprintScale;
 				PassParameters->LodPixelCountThreshold = LodPixelCountThreshold;
-				
-				PassParameters->Clipmap = FVirtualShadowMapClipmap::IsEnabled() ? 1 : 0;
 
 				auto ComputeShader = View.ShaderMap->GetShader<FGeneratePageFlagsFromPixelsCS>(PermutationVector);
 
@@ -809,39 +769,6 @@ void FVirtualShadowMapArray::GeneratePageFlagsFromLightGrid(FRDGBuilder& GraphBu
 			}
 		}
 
-		{
-			bool bAnyForceTopMipVisible = false;
-			for (uint32 bForceTopMipVisible : ForceTopMipVisible)
-			{
-				if (bForceTopMipVisible > 0)
-				{
-					bAnyForceTopMipVisible = true;
-					break;
-				}
-			}
-
-			if (bAnyForceTopMipVisible)
-			{
-				// TODO: Combine this with clearing of PageFlags to get rid of a pass?
-				FRDGBufferRef ShadowMapHasFallbackRDG = CreateStructuredBuffer( GraphBuilder, TEXT( "ShadowMapHasFallback" ), ForceTopMipVisible );
-
-				FMarkFallbackPageFlagsCS::FParameters* PassParameters = GraphBuilder.AllocParameters< FMarkFallbackPageFlagsCS::FParameters >();
-				PassParameters->CommonParameters = CommonParameters;
-				PassParameters->NumVirtualShadowMaps = ForceTopMipVisible.Num();
-				PassParameters->ShadowMapHasFallback = GraphBuilder.CreateSRV(ShadowMapHasFallbackRDG);
-				PassParameters->OutPageRequestFlags = GraphBuilder.CreateUAV(PageRequestFlagsRDG);
-				auto ComputeShader = Views[0].ShaderMap->GetShader<FMarkFallbackPageFlagsCS>( );
-
-				FComputeShaderUtils::AddPass(
-					GraphBuilder,
-					RDG_EVENT_NAME( "MarkFallbackPages" ),
-					ComputeShader,
-					PassParameters,
-					FComputeShaderUtils::GetGroupCount( ForceTopMipVisible.Num(), 64 )
-				);
-			}
-		}
-
 		FRDGBufferRef PageTableRDG			= GraphBuilder.CreateBuffer( FRDGBufferDesc::CreateStructuredDesc( 8, NumPageFlags ), TEXT("PageTable") );
 		FRDGBufferRef HZBPageTableRDG		= GraphBuilder.CreateBuffer( FRDGBufferDesc::CreateStructuredDesc( 8, NumPageFlags ), TEXT("HZBPageTable") );
 
diff --git a/Engine/Source/Runtime/Renderer/Private/VirtualShadowMaps/VirtualShadowMapArray.h b/Engine/Source/Runtime/Renderer/Private/VirtualShadowMaps/VirtualShadowMapArray.h
--- a/Engine/Source/Runtime/Renderer/Private/VirtualShadowMaps/VirtualShadowMapArray.h
+++ b/Engine/Source/Runtime/Renderer/Private/VirtualShadowMaps/VirtualShadowMapArray.h
@@ -20,6 +20,14 @@ constexpr uint32 ILog2Const(uint32 n)
 	return (n > 1) ? 1 + ILog2Const(n / 2) : 0;
 }
 
+// See CalcLevelOffsets in PageAccessCommon.ush for some details on this logic
+constexpr uint32 CalcVirtualShadowMapLevelOffsets(uint32 Level, uint32 Log2Level0DimPagesXY)
+{
+	uint32 NumBits = Level << 1;
+	uint32 StartBit = (2U * Log2Level0DimPagesXY + 2U) - NumBits;
+	uint32 Mask = ((1U << NumBits) - 1U) << StartBit;
+	return 0x55555555U & Mask;
+}
 
 class FVirtualShadowMap
 {
@@ -43,6 +51,8 @@ public:
 	static constexpr uint32 Log2Level0DimPagesXY = ILog2Const(Level0DimPagesXY);
 	static constexpr uint32 MaxMipLevels = Log2Level0DimPagesXY + 1U;
 
+	static constexpr uint32 PageTableSize = CalcVirtualShadowMapLevelOffsets(MaxMipLevels, Log2Level0DimPagesXY);
+
 	static constexpr uint32 VirtualMaxResolutionXY = Level0DimPagesXY * PageSize;
 	
 	static constexpr uint32 PhysicalPageAddressBits = 16U;
@@ -50,17 +60,16 @@ public:
 	static constexpr uint32 MaxPhysicalTextureDimTexels = MaxPhysicalTextureDimPages * PageSize;
 
 	static constexpr uint32 RasterWindowPages = 4u;
-	
-	
+		
 	// Something large (we're using ints at the moment...)
 	// TODO: Fix this when tweaking data sizes of page table entries to e.g., 2x8 bits
-	static constexpr uint32 InvalidPhysicalPageAddress = 65535U; 
-
+	static constexpr uint32 InvalidPhysicalPageAddress = 65535U;
+	
 	FVirtualShadowMap(uint32 InID) : ID(InID)
 	{
 	}
 
-	uint32 ID = INDEX_NONE;
+	int32 ID = INDEX_NONE;
 	TSharedPtr<FVirtualShadowMapCacheEntry> VirtualShadowMapCacheEntry;
 };
 
@@ -79,18 +88,15 @@ struct FVirtualShadowMapProjectionShaderData
 	 * Translation from world space to shadow space (add before transform by TranslatedWorldToShadowViewMatrix).
 	 */
 	FVector4 ShadowPreViewTranslation;
-	uint32 VirtualShadowMapId;
+	int32 VirtualShadowMapId;
 	
 	// These could be per-light (first/count), but convenient here and not much overhead
 	int32 ClipmapLevel = 0;					// "Absolute" level, can be negative
 	int32 ClipmapLevelCount = 0;
 	float ClipmapResolutionLodBias = 0.0f;
 
-	float MinSceneDepth = -HALF_WORLD_MAX;
-	float MaxSceneDepth =  HALF_WORLD_MAX;
-
 	// Seems the FMatrix forces 16-byte alignment
-	float Padding[2];
+	//float Padding[2];
 };
 static_assert((sizeof(FVirtualShadowMapProjectionShaderData) % 16) == 0, "FVirtualShadowMapProjectionShaderData size should be a multiple of 16-bytes for alignment.");
 
@@ -99,7 +105,6 @@ FVirtualShadowMapProjectionShaderData GetVirtualShadowMapProjectionShaderData(co
 FMatrix CalcTranslatedWorldToShadowUvNormalMatrix(const FMatrix& TranslatedWorldToShadowView, const FMatrix& ViewToClip);
 
 BEGIN_SHADER_PARAMETER_STRUCT(FVirtualShadowMapCommonParameters, )
-	SHADER_PARAMETER_ARRAY(uint32, LevelOffsets, [FVirtualShadowMap::MaxMipLevels])
 	SHADER_PARAMETER_ARRAY(uint32, HPageFlagLevelOffsets, [FVirtualShadowMap::MaxMipLevels])
 	SHADER_PARAMETER(uint32, PageTableSize)
 	SHADER_PARAMETER(uint32, HPageTableSize)
diff --git a/Engine/Source/Runtime/Renderer/Private/VirtualShadowMaps/VirtualShadowMapCacheManager.cpp b/Engine/Source/Runtime/Renderer/Private/VirtualShadowMaps/VirtualShadowMapCacheManager.cpp
--- a/Engine/Source/Runtime/Renderer/Private/VirtualShadowMaps/VirtualShadowMapCacheManager.cpp
+++ b/Engine/Source/Runtime/Renderer/Private/VirtualShadowMaps/VirtualShadowMapCacheManager.cpp
@@ -26,7 +26,7 @@ static TAutoConsoleVariable<int32> CVarCacheVirtualSMs(
 );
 
 
-void FVirtualShadowMapCacheEntry::UpdateClipmap(uint32 VirtualShadowMapId, const FMatrix &WorldToLight, FIntPoint PageSpaceLocation, float GlobalDepth)
+void FVirtualShadowMapCacheEntry::UpdateClipmap(int32 VirtualShadowMapId, const FMatrix &WorldToLight, FIntPoint PageSpaceLocation, float GlobalDepth)
 {
 	// Swap previous frame data over.
 	PrevVirtualShadowMapId = CurrentVirtualShadowMapId;
@@ -44,7 +44,7 @@ void FVirtualShadowMapCacheEntry::UpdateClipmap(uint32 VirtualShadowMapId, const
 	CurrentShadowMapGlobalDepth = GlobalDepth;
 }
 
-void FVirtualShadowMapCacheEntry::Update(uint32 VirtualShadowMapId, const FMatrix &ShadowPreTranslatedWorldToShadowClip, 
+void FVirtualShadowMapCacheEntry::Update(int32 VirtualShadowMapId, const FMatrix &ShadowPreTranslatedWorldToShadowClip, 
 	const FVector &SubjectWorldSpacePosition, bool bIsViewDependent, const FWholeSceneProjectedShadowInitializer &InCacheValidKey, 
 	FVector &SnappedSubjectWorldSpacePosition)
 {
diff --git a/Engine/Source/Runtime/Renderer/Private/VirtualShadowMaps/VirtualShadowMapCacheManager.h b/Engine/Source/Runtime/Renderer/Private/VirtualShadowMaps/VirtualShadowMapCacheManager.h
--- a/Engine/Source/Runtime/Renderer/Private/VirtualShadowMaps/VirtualShadowMapCacheManager.h
+++ b/Engine/Source/Runtime/Renderer/Private/VirtualShadowMaps/VirtualShadowMapCacheManager.h
@@ -22,12 +22,12 @@ public:
 	bool IsValid() { return PrevVirtualShadowMapId != INDEX_NONE; }
 
 
-	void Update(uint32 VirtualShadowMapId, const FMatrix &ShadowPreTranslatedWorldToShadowClip, 
+	void Update(int32 VirtualShadowMapId, const FMatrix &ShadowPreTranslatedWorldToShadowClip, 
 		const FVector &SubjectWorldSpacePosition, bool bIsViewDependent, 
 		const FWholeSceneProjectedShadowInitializer &InCacheValidKey, 
 		FVector &SnappedSubjectWorldSpacePosition);
 
-	void UpdateClipmap(uint32 VirtualShadowMapId, const FMatrix &WorldToLight, FIntPoint PageSpaceLocation, float GlobalDepth);
+	void UpdateClipmap(int32 VirtualShadowMapId, const FMatrix &WorldToLight, FIntPoint PageSpaceLocation, float GlobalDepth);
 
 	/**
 	 * Returns the PrevVirtualShadowMapId if cached data is valid (bValidData), or INDEX_NONE otherwise. 
@@ -48,7 +48,7 @@ public:
 	FIntPoint PrevPageSpaceLocation = FIntPoint(0,0);
 
 	// Set to INDEX_NONE when cache data is invalidated by an external change, like movement.
-	uint32 PrevVirtualShadowMapId = INDEX_NONE;
+	int32 PrevVirtualShadowMapId = INDEX_NONE;
 
 	// Depth of the world-space origin of the shadow map in shadow maps space.
 	// Used to offset the depth of pages as they are copied when the light moves
@@ -58,7 +58,7 @@ public:
 	
 	// Aligned location in pages after update.
 	FIntPoint CurrentPageSpaceLocation = FIntPoint(0, 0);
-	uint32 CurrentVirtualShadowMapId = INDEX_NONE;
+	int32 CurrentVirtualShadowMapId = INDEX_NONE;
 	float CurrentShadowMapGlobalDepth = 0.0f;
 
 	// TODO: Potentially refactor this to decouple the cache key details
diff --git a/Engine/Source/Runtime/Renderer/Private/VirtualShadowMaps/VirtualShadowMapClipmap.cpp b/Engine/Source/Runtime/Renderer/Private/VirtualShadowMaps/VirtualShadowMapClipmap.cpp
--- a/Engine/Source/Runtime/Renderer/Private/VirtualShadowMaps/VirtualShadowMapClipmap.cpp
+++ b/Engine/Source/Runtime/Renderer/Private/VirtualShadowMaps/VirtualShadowMapClipmap.cpp
@@ -11,13 +11,6 @@ VirtualShadowMapClipmap.cpp
 #include "VirtualShadowMapArray.h"
 #include "VirtualShadowMapCacheManager.h"
 
-static TAutoConsoleVariable<int32> CVarVirtualShadowMapClipmap(
-	TEXT( "r.Shadow.v.Clipmap" ),
-	0,
-	TEXT( "Enable clipmap for directional light virtual shadow maps instead of conventional cascades." ),
-	ECVF_RenderThreadSafe
-);
-
 static TAutoConsoleVariable<float> CVarVirtualShadowMapClipmapResolutionLodBias(
 	TEXT( "r.Shadow.v.Clipmap.ResolutionLodBias" ),
 	3.0f,
@@ -32,11 +25,6 @@ static TAutoConsoleVariable<float> CVarVirtualShadowMapClipmapFirstLevel(
 );
 
 
-bool FVirtualShadowMapClipmap::IsEnabled()
-{
-	return CVarVirtualShadowMapClipmap.GetValueOnRenderThread() != 0;
-}
-
 // "Virtual" clipmap level to clipmap radius
 static float GetLevelRadius(int32 Level)
 {
diff --git a/Engine/Source/Runtime/Renderer/Private/VirtualShadowMaps/VirtualShadowMapClipmap.h b/Engine/Source/Runtime/Renderer/Private/VirtualShadowMaps/VirtualShadowMapClipmap.h
--- a/Engine/Source/Runtime/Renderer/Private/VirtualShadowMaps/VirtualShadowMapClipmap.h
+++ b/Engine/Source/Runtime/Renderer/Private/VirtualShadowMaps/VirtualShadowMapClipmap.h
@@ -16,9 +16,6 @@ class FVirtualShadowMapArrayCacheManager;
 class FVirtualShadowMapClipmap : FRefCountedObject
 {
 public:	
-	// Use clipmap instead of cascades for virtual shadow maps
-	static bool IsEnabled();
-
 	FVirtualShadowMapClipmap(
 		FVirtualShadowMapArray& VirtualShadowMapArray,
 		FVirtualShadowMapArrayCacheManager* VirtualShadowMapArrayCacheManager,
diff --git a/Engine/Source/Runtime/Renderer/Private/VirtualShadowMaps/VirtualShadowMapProjection.cpp b/Engine/Source/Runtime/Renderer/Private/VirtualShadowMaps/VirtualShadowMapProjection.cpp
--- a/Engine/Source/Runtime/Renderer/Private/VirtualShadowMaps/VirtualShadowMapProjection.cpp
+++ b/Engine/Source/Runtime/Renderer/Private/VirtualShadowMaps/VirtualShadowMapProjection.cpp
@@ -47,7 +47,6 @@ BEGIN_SHADER_PARAMETER_STRUCT(FVirtualShadowMapProjectionParameters, )
 	SHADER_PARAMETER(uint32, ProjectionCount)
 	SHADER_PARAMETER(float, ContactShadowLength)
 	SHADER_PARAMETER(int32, DebugOutputType)
-	SHADER_PARAMETER(int32, Clipmap)
 	SHADER_PARAMETER_RDG_BUFFER_SRV(StructuredBuffer<uint2>, PageTable)
 	SHADER_PARAMETER_RDG_TEXTURE(Texture2D<uint>, PhysicalPagePool)
 	SHADER_PARAMETER_RDG_BUFFER_SRV(StructuredBuffer< FVirtualShadowMapProjectionShaderData >, ShadowMapProjectionData)
@@ -106,7 +105,6 @@ IMPLEMENT_GLOBAL_SHADER(FVirtualShadowMapProjectionCompositePS, "/Engine/Private
 static void AddPass_RenderVirtualShadowMapProjection(
 	const FLightSceneProxy* LightProxy,
 	const TArray<FVirtualShadowMapProjectionShaderData, SceneRenderingAllocator>& ShadowMapProjectionData,
-	bool bClipmap,
 	FRDGBuilder& GraphBuilder,
 	const FViewInfo& View, 
 	FVirtualShadowMapArray& VirtualShadowMapArray,
@@ -134,7 +132,6 @@ static void AddPass_RenderVirtualShadowMapProjection(
 	PassParameters->DebugOutputType = CVarVirtualShadowMapDebugProjection.GetValueOnRenderThread();
 	PassParameters->ContactShadowLength = CVarContactShadowLength.GetValueOnRenderThread();
 	PassParameters->ProjectionCount = ShadowMapProjectionData.Num();
-	PassParameters->Clipmap = bClipmap ? 1 : 0;
 
 	FRDGBufferRef ShadowMapProjectionDataRDG = CreateStructuredBuffer(
 		GraphBuilder, 
@@ -181,7 +178,6 @@ static FRDGTextureRef CreateDebugOutput(FRDGBuilder& GraphBuilder, FIntPoint Ext
 static void RenderVirtualShadowMapProjectionCommon(
 	const FLightSceneProxy* LightProxy,
 	const TArray<FVirtualShadowMapProjectionShaderData, SceneRenderingAllocator>& ShadowMapProjectionData,
-	bool bClipmap,
 	FRDGBuilder& GraphBuilder,
 	const FViewInfo& View, 
 	FVirtualShadowMapArray& VirtualShadowMapArray,
@@ -192,7 +188,7 @@ static void RenderVirtualShadowMapProjectionCommon(
 {
 	// Main Pass
 	AddPass_RenderVirtualShadowMapProjection(
-		LightProxy, ShadowMapProjectionData, bClipmap,
+		LightProxy, ShadowMapProjectionData,
 		GraphBuilder,
 		View,
 		VirtualShadowMapArray,
@@ -209,7 +205,7 @@ static void RenderVirtualShadowMapProjectionCommon(
 		FRenderTargetBinding DebugRenderTargetBinding(DebugOutput, ERenderTargetLoadAction::EClear);
 
 		AddPass_RenderVirtualShadowMapProjection(
-			LightProxy, ShadowMapProjectionData, bClipmap,
+			LightProxy, ShadowMapProjectionData,
 			GraphBuilder,
 			View,
 			VirtualShadowMapArray,
@@ -222,28 +218,6 @@ static void RenderVirtualShadowMapProjectionCommon(
 	}
 }
 
-static void GetVirtualShadowMapProjectionData(
-	FProjectedShadowInfo** ShadowInfo,
-	uint32 ShadowInfoCount,
-	TArray<FVirtualShadowMapProjectionShaderData, SceneRenderingAllocator>& OutShadowMapProjectionData)
-{	
-	OutShadowMapProjectionData.Empty();
-	if (ShadowInfoCount > 1)
-	{
-		// One projection data per cascade
-		OutShadowMapProjectionData.AddDefaulted(ShadowInfoCount);
-		for (uint32 CascadeIndex = 0; CascadeIndex < ShadowInfoCount; ++CascadeIndex)
-		{
-			OutShadowMapProjectionData[CascadeIndex] = GetVirtualShadowMapProjectionShaderData(ShadowInfo[CascadeIndex]);
-		}
-	}
-	else
-	{
-		// Single projection data
-		OutShadowMapProjectionData.Add(GetVirtualShadowMapProjectionShaderData(ShadowInfo[0]));
-	}
-}
-
 static void GetVirtualShadowMapProjectionData(
 	const TSharedPtr<FVirtualShadowMapClipmap>& Clipmap,
 	TArray<FVirtualShadowMapProjectionShaderData, SceneRenderingAllocator>& OutShadowMapProjectionData)
@@ -260,8 +234,7 @@ static void GetVirtualShadowMapProjectionData(
 }
 
 void RenderVirtualShadowMapProjectionForDenoising(
-	FProjectedShadowInfo** ShadowInfo,
-	uint32 ShadowInfoCount,
+	FProjectedShadowInfo* ShadowInfo,
 	FRDGBuilder& GraphBuilder,
 	const FViewInfo& View,
 	FVirtualShadowMapArray& VirtualShadowMapArray,
@@ -269,11 +242,11 @@ void RenderVirtualShadowMapProjectionForDenoising(
 	FRDGTextureRef SignalTexture)
 {
 	TArray<FVirtualShadowMapProjectionShaderData, SceneRenderingAllocator> ShadowMapProjectionData;
-	GetVirtualShadowMapProjectionData(ShadowInfo, ShadowInfoCount, ShadowMapProjectionData);
+	ShadowMapProjectionData.Add(GetVirtualShadowMapProjectionShaderData(ShadowInfo));
 
 	FRenderTargetBinding RenderTargetBinding(SignalTexture, ERenderTargetLoadAction::EClear);
 	RenderVirtualShadowMapProjectionCommon(
-		ShadowInfo[0]->GetLightSceneInfo().Proxy, ShadowMapProjectionData, false,
+		ShadowInfo->GetLightSceneInfo().Proxy, ShadowMapProjectionData,
 		GraphBuilder,
 		View,
 		VirtualShadowMapArray,
@@ -296,7 +269,7 @@ void RenderVirtualShadowMapProjectionForDenoising(
 
 	FRenderTargetBinding RenderTargetBinding(SignalTexture, ERenderTargetLoadAction::EClear);
 	RenderVirtualShadowMapProjectionCommon(
-		Clipmap->GetLightSceneInfo().Proxy, ShadowMapProjectionData, true,
+		Clipmap->GetLightSceneInfo().Proxy, ShadowMapProjectionData,
 		GraphBuilder,
 		View,
 		VirtualShadowMapArray,
@@ -308,8 +281,7 @@ void RenderVirtualShadowMapProjectionForDenoising(
 
 
 void RenderVirtualShadowMapProjection(
-	FProjectedShadowInfo** ShadowInfo,
-	uint32 ShadowInfoCount,
+	FProjectedShadowInfo* ShadowInfo,
 	FRDGBuilder& GraphBuilder,
 	const FViewInfo& View,
 	FVirtualShadowMapArray& VirtualShadowMapArray,
@@ -318,13 +290,13 @@ void RenderVirtualShadowMapProjection(
 	bool bProjectingForForwardShading)
 {
 	TArray<FVirtualShadowMapProjectionShaderData, SceneRenderingAllocator> ShadowMapProjectionData;
-	GetVirtualShadowMapProjectionData(ShadowInfo, ShadowInfoCount, ShadowMapProjectionData);
+	ShadowMapProjectionData.Add(GetVirtualShadowMapProjectionShaderData(ShadowInfo));
 
 	FRenderTargetBinding RenderTargetBinding(ScreenShadowMaskTexture, ERenderTargetLoadAction::ELoad);
 	// Blend state should not vary per cascade
-	FRHIBlendState* BlendState = ShadowInfo[0]->GetBlendStateForProjection(bProjectingForForwardShading, false);
+	FRHIBlendState* BlendState = ShadowInfo->GetBlendStateForProjection(bProjectingForForwardShading, false);
 	RenderVirtualShadowMapProjectionCommon(
-		ShadowInfo[0]->GetLightSceneInfo().Proxy, ShadowMapProjectionData, false,
+		ShadowInfo->GetLightSceneInfo().Proxy, ShadowMapProjectionData,
 		GraphBuilder,
 		View,
 		VirtualShadowMapArray,
@@ -352,7 +324,7 @@ void RenderVirtualShadowMapProjection(
 	FRHIBlendState* BlendState = TStaticBlendState<CW_BA, BO_Min, BF_One, BF_One, BO_Min, BF_One, BF_One>::GetRHI();
 
 	RenderVirtualShadowMapProjectionCommon(
-		Clipmap->GetLightSceneInfo().Proxy, ShadowMapProjectionData, true,
+		Clipmap->GetLightSceneInfo().Proxy, ShadowMapProjectionData,
 		GraphBuilder,
 		View,
 		VirtualShadowMapArray,
diff --git a/Engine/Source/Runtime/Renderer/Private/VirtualShadowMaps/VirtualShadowMapProjection.h b/Engine/Source/Runtime/Renderer/Private/VirtualShadowMaps/VirtualShadowMapProjection.h
--- a/Engine/Source/Runtime/Renderer/Private/VirtualShadowMaps/VirtualShadowMapProjection.h
+++ b/Engine/Source/Runtime/Renderer/Private/VirtualShadowMaps/VirtualShadowMapProjection.h
@@ -21,8 +21,7 @@ void RenderVirtualShadowMapProjectionForDenoising(
 	FRDGTextureRef SignalTexture);
 
 void RenderVirtualShadowMapProjectionForDenoising(
-	FProjectedShadowInfo** ShadowInfo,
-	uint32 ShadowInfoCount,
+	FProjectedShadowInfo* ShadowInfo,
 	FRDGBuilder& GraphBuilder,
 	const FViewInfo& View,
 	FVirtualShadowMapArray& VirtualShadowMapArray,
@@ -39,8 +38,7 @@ void RenderVirtualShadowMapProjection(
 	bool bProjectingForForwardShading);
 
 void RenderVirtualShadowMapProjection(
-	FProjectedShadowInfo** ShadowInfo,
-	uint32 ShadowInfoCount,
+	FProjectedShadowInfo* ShadowInfo,
 	FRDGBuilder& GraphBuilder,
 	const FViewInfo& View,
 	FVirtualShadowMapArray& VirtualShadowMapArray,
