Makes TAAU plugable

diff --git a/Engine/Source/Editor/LevelEditor/Private/SLevelViewportToolBar.cpp b/Engine/Source/Editor/LevelEditor/Private/SLevelViewportToolBar.cpp
--- a/Engine/Source/Editor/LevelEditor/Private/SLevelViewportToolBar.cpp
+++ b/Engine/Source/Editor/LevelEditor/Private/SLevelViewportToolBar.cpp
@@ -39,6 +39,8 @@
 #include "SEditorViewportViewMenuContext.h"
 #include "Bookmarks/IBookmarkTypeTools.h"
 #include "ToolMenu.h"
+#include "CustomStaticScreenPercentage.h"
+#include "CustomEditorStaticScreenPercentage.h"
 
 #define LOCTEXT_NAMESPACE "LevelViewportToolBar"
 
@@ -663,7 +665,34 @@ void SLevelViewportToolBar::FillOptionsMenu(UToolMenu* Menu)
 				Section.AddEntry(FToolMenuEntry::InitWidget("FarViewPlane", GenerateFarViewPlaneMenu(), LOCTEXT("FarViewPlane", "Far View Plane")));
 			}
 
-			Section.AddEntry(FToolMenuEntry::InitWidget("ScreenPercentage", GenerateScreenPercentageMenu(), LOCTEXT("ScreenPercentage", "Screen Percentage")));
+
+			// for the TemporalUpscaler plugin to inject its UI
+
+			bool bAddDefaultScreenPercentageSlider = true;
+
+			if (GCustomEditorStaticScreenPercentage)
+			{
+				ICustomEditorStaticScreenPercentage::FViewportMenuEntryArguments Arguments
+				{
+					&Section,
+					this,
+					Viewport.Pin()
+				};
+
+				bAddDefaultScreenPercentageSlider = !GCustomEditorStaticScreenPercentage->GenerateEditorViewportOptionsMenuEntry(Arguments);
+			}
+
+			if (bAddDefaultScreenPercentageSlider)
+			{
+				// if we have a custom screen percentage, but it's not active currently then we need to clamp the screen percentage to the range of the built-in TAA
+				if (GCustomStaticScreenPercentage)
+				{
+					const int32 PreviewScreenPercentageMin = FSceneViewScreenPercentageConfig::kMinTAAUpsampleResolutionFraction * 100.0f;
+					const int32 PreviewScreenPercentageMax = FSceneViewScreenPercentageConfig::kMaxTAAUpsampleResolutionFraction * 100.0f;
+					this->OnScreenPercentageValueChanged(FMath::Clamp(this->OnGetScreenPercentageValue(), PreviewScreenPercentageMin, PreviewScreenPercentageMax));
+				}
+				Section.AddEntry(FToolMenuEntry::InitWidget("ScreenPercentage", GenerateScreenPercentageMenu(), LOCTEXT("ScreenPercentage", "Screen Percentage")));
+			}
 		}
 
 		{
diff --git a/Engine/Source/Editor/UnrealEd/Private/EditorViewportClient.cpp b/Engine/Source/Editor/UnrealEd/Private/EditorViewportClient.cpp
--- a/Engine/Source/Editor/UnrealEd/Private/EditorViewportClient.cpp
+++ b/Engine/Source/Editor/UnrealEd/Private/EditorViewportClient.cpp
@@ -56,6 +56,10 @@
 #include "UnrealEngine.h"
 #include "BufferVisualizationData.h"
 
+#include "CustomEditorStaticScreenPercentage.h"
+
+ICustomEditorStaticScreenPercentage* GCustomEditorStaticScreenPercentage = nullptr;
+
 #define LOCTEXT_NAMESPACE "EditorViewportClient"
 
 const EViewModeIndex FEditorViewportClient::DefaultPerspectiveViewMode = VMI_Lit;
@@ -431,6 +435,7 @@ FEditorViewportClient::FEditorViewportClient(FEditorModeTools* InModeTools, FPre
 	, MovingPreviewLightTimer(0.0f)
 	, bLockFlightCamera(false)
 	, PreviewResolutionFraction(1.0f)
+	, bPreviewCustomTemporalUpscaler(false)
 	, SceneDPIMode(ESceneDPIMode::EditorDefault)
 	, PerspViewModeIndex(DefaultPerspectiveViewMode)
 	, OrthoViewModeIndex(DefaultOrthoViewMode)
@@ -2679,6 +2684,16 @@ void FEditorViewportClient::SetPreviewScreenPercentage(int32 PreviewScreenPercen
 	PreviewResolutionFraction = PreviewScreenPercentage / 100.0f;
 }
 
+bool FEditorViewportClient::GetPreviewCustomTemporalUpscaler() const
+{
+	return bPreviewCustomTemporalUpscaler;
+}
+
+void FEditorViewportClient::SetPreviewCustomTemporalUpscaler(bool PreviewCustomTemporalUpscaler)
+{
+	bPreviewCustomTemporalUpscaler = PreviewCustomTemporalUpscaler;
+}
+
 bool FEditorViewportClient::SupportsLowDPIPreview() const
 {
 	return GetDPIDerivedResolutionFraction() < 1.0f;
@@ -3835,6 +3850,48 @@ void FEditorViewportClient::Draw(FViewport* InViewport, FCanvas* Canvas)
 
 	ViewFamily.LandscapeLODOverride = LandscapeLODOverride;
 
+	// Setup the screen percentage and upscaling method for the view family.
+	bool bFinalScreenPercentageShowFlag;
+	{
+		checkf(ViewFamily.GetScreenPercentageInterface() == nullptr,
+			TEXT("Some code has tried to set up an alien screen percentage driver, that could be wrong if not supported very well by the RHI."));
+
+		// If not doing VR rendering, apply DPI derived resolution fraction even if show flag is disabled
+		if (!bStereoRendering && SupportsLowDPIPreview() && IsLowDPIPreview() && ViewFamily.SupportsScreenPercentage())
+		{
+			ViewFamily.SecondaryViewFraction = GetDPIDerivedResolutionFraction();
+		}
+
+		// Setup custom upscaler and screen percentage.
+		if (GCustomEditorStaticScreenPercentage)
+		{
+			GCustomEditorStaticScreenPercentage->SetupEditorViewFamily(ViewFamily, PreviewResolutionFraction, bPreviewCustomTemporalUpscaler);
+		}
+
+		// If a screen percentage interface was not set by one of the view extension, then set the legacy one.
+		if (ViewFamily.GetScreenPercentageInterface() == nullptr)
+		{
+			float GlobalResolutionFraction = 1.0f;
+
+			// If not doing VR rendering, apply preview resolution fraction.
+			if (!bStereoRendering && SupportsPreviewResolutionFraction() && ViewFamily.SupportsScreenPercentage())
+			{
+				GlobalResolutionFraction = PreviewResolutionFraction;
+
+				// Force screen percentage's engine show flag to be turned on for preview screen percentage.
+				ViewFamily.EngineShowFlags.ScreenPercentage = (GlobalResolutionFraction != 1.0);
+			}
+
+			// In editor viewport, we ignore r.ScreenPercentage and FPostProcessSettings::ScreenPercentage by design.
+			ViewFamily.SetScreenPercentageInterface(new FLegacyScreenPercentageDriver(
+				ViewFamily, GlobalResolutionFraction, /* AllowPostProcessSettingsScreenPercentage = */ false));
+		}
+
+		check(ViewFamily.GetScreenPercentageInterface() != nullptr);
+
+		bFinalScreenPercentageShowFlag = ViewFamily.EngineShowFlags.ScreenPercentage;
+	}
+
 	FSceneView* View = nullptr;
 
 	// Stereo rendering
@@ -3862,30 +3919,8 @@ void FEditorViewportClient::Draw(FViewport* InViewport, FCanvas* Canvas)
 		Canvas->Clear(FLinearColor::Black);
 	}
 
-	// If not doing VR rendering, apply DPI derived resolution fraction even if show flag is disabled
-	if (!bStereoRendering && SupportsLowDPIPreview() && IsLowDPIPreview() && ViewFamily.SupportsScreenPercentage())
-	{
-		ViewFamily.SecondaryViewFraction = GetDPIDerivedResolutionFraction();
-	}
-
-	// If a screen percentage interface was not set by one of the view extension, then set the legacy one.
-	if (ViewFamily.GetScreenPercentageInterface() == nullptr)
-	{
-		float GlobalResolutionFraction = 1.0f;
-
-		// If not doing VR rendering, apply preview resolution fraction.
-		if (!bStereoRendering && SupportsPreviewResolutionFraction() && ViewFamily.SupportsScreenPercentage())
-		{
-			GlobalResolutionFraction = PreviewResolutionFraction;
-
-			// Force screen percentage's engine show flag to be turned on for preview screen percentage.
-			ViewFamily.EngineShowFlags.ScreenPercentage = (GlobalResolutionFraction != 1.0);
-		}
-
-		// In editor viewport, we ignore r.ScreenPercentage and FPostProcessSettings::ScreenPercentage by design.
-		ViewFamily.SetScreenPercentageInterface(new FLegacyScreenPercentageDriver(
-			ViewFamily, GlobalResolutionFraction, /* AllowPostProcessSettingsScreenPercentage = */ false));
-	}
+	// Make sure the engine show flag for screen percentage is still what it was when setting up the screen percentage interface
+	ViewFamily.EngineShowFlags.ScreenPercentage = bFinalScreenPercentageShowFlag;
 
 	// Draw the 3D scene
 	GetRendererModule().BeginRenderingViewFamily(Canvas,&ViewFamily);
diff --git a/Engine/Source/Editor/UnrealEd/Public/CustomEditorStaticScreenPercentage.h b/Engine/Source/Editor/UnrealEd/Public/CustomEditorStaticScreenPercentage.h
new file mode 100644
--- /dev/null
+++ b/Engine/Source/Editor/UnrealEd/Public/CustomEditorStaticScreenPercentage.h
@@ -0,0 +1,26 @@
+#pragma once
+#include "CoreMinimal.h"
+
+class FSceneViewFamily;
+
+class SLevelViewportToolBar;
+struct FToolMenuSection;
+class SLevelViewport;
+
+class UNREALED_API ICustomEditorStaticScreenPercentage
+{
+public:
+	virtual ~ICustomEditorStaticScreenPercentage() {};
+
+	virtual void SetupEditorViewFamily(FSceneViewFamily& ViewFamily, float PreviewResolutionFraction, bool bPreviewCustomTemporalUpscaler) = 0;
+	
+	struct FViewportMenuEntryArguments
+	{
+		FToolMenuSection* Section;
+		SLevelViewportToolBar* ToolBar;
+		TSharedPtr<SLevelViewport> Viewport;
+	};
+	virtual bool GenerateEditorViewportOptionsMenuEntry(const FViewportMenuEntryArguments& Arguments) = 0;
+};
+
+extern UNREALED_API ICustomEditorStaticScreenPercentage* GCustomEditorStaticScreenPercentage;
diff --git a/Engine/Source/Editor/UnrealEd/Public/EditorViewportClient.h b/Engine/Source/Editor/UnrealEd/Public/EditorViewportClient.h
--- a/Engine/Source/Editor/UnrealEd/Public/EditorViewportClient.h
+++ b/Engine/Source/Editor/UnrealEd/Public/EditorViewportClient.h
@@ -1246,6 +1246,12 @@ public:
 	/** Set preview screen percentage on UI behalf. */
 	void SetPreviewScreenPercentage(int32 PreviewScreenPercentage);
 
+	/** @return preview custom temporal upscaler  for UI. */
+	bool GetPreviewCustomTemporalUpscaler() const;
+
+	/** Set preview temporal upscaler mode for UI. */
+	void SetPreviewCustomTemporalUpscaler(bool PreviewCustomTemporalUpscaler);
+
 	/** @return True if DPI preview is supported. */
 	bool SupportsLowDPIPreview() const;
 
@@ -1763,6 +1769,9 @@ private:
 	/** Controles resolution fraction for previewing in editor viewport at different screen percentage. */
 	float PreviewResolutionFraction;
 
+	/** Controls whether a custom upscaler from a plugin is used*/
+	bool bPreviewCustomTemporalUpscaler;
+
 	// DPI mode for scene rendering.
 	enum class ESceneDPIMode
 	{
diff --git a/Engine/Source/Runtime/D3D12RHI/Private/D3D12DescriptorCache.cpp b/Engine/Source/Runtime/D3D12RHI/Private/D3D12DescriptorCache.cpp
--- a/Engine/Source/Runtime/D3D12RHI/Private/D3D12DescriptorCache.cpp
+++ b/Engine/Source/Runtime/D3D12RHI/Private/D3D12DescriptorCache.cpp
@@ -270,7 +270,7 @@ bool FD3D12DescriptorCache::SetDescriptorHeaps()
 }
 
 
-void FD3D12DescriptorCache::SetCurrentCommandList(const FD3D12CommandListHandle& CommandListHandle)
+D3D12RHI_API void FD3D12DescriptorCache::SetCurrentCommandList(const FD3D12CommandListHandle& CommandListHandle)
 {
 	// Clear the previous heap pointers (since it's a new command list) and then set the current descriptor heaps.
 	pPreviousViewHeap = nullptr;
diff --git a/Engine/Source/Runtime/D3D12RHI/Private/D3D12DescriptorCache.h b/Engine/Source/Runtime/D3D12RHI/Private/D3D12DescriptorCache.h
--- a/Engine/Source/Runtime/D3D12RHI/Private/D3D12DescriptorCache.h
+++ b/Engine/Source/Runtime/D3D12RHI/Private/D3D12DescriptorCache.h
@@ -600,7 +600,7 @@ public:
 	// Notify the descriptor cache every time you start recording a command list.
 	// This sets descriptor heaps on the command list and indicates the current fence value which allows
 	// us to avoid querying DX12 for that value thousands of times per frame, which can be costly.
-	void SetCurrentCommandList(const FD3D12CommandListHandle& CommandListHandle);
+	D3D12RHI_API void SetCurrentCommandList(const FD3D12CommandListHandle& CommandListHandle);
 
 	// ------------------------------------------------------
 	// end Descriptor Slot Reservation stuff
diff --git a/Engine/Source/Runtime/D3D12RHI/Private/D3D12Device.h b/Engine/Source/Runtime/D3D12RHI/Private/D3D12Device.h
--- a/Engine/Source/Runtime/D3D12RHI/Private/D3D12Device.h
+++ b/Engine/Source/Runtime/D3D12RHI/Private/D3D12Device.h
@@ -111,8 +111,8 @@ public:
 
 	TArray<FD3D12CommandListHandle> PendingCommandLists;
 
-	void RegisterGPUWork(uint32 NumPrimitives = 0, uint32 NumVertices = 0);
-	void RegisterGPUDispatch(FIntVector GroupCount);
+	D3D12RHI_API void RegisterGPUWork(uint32 NumPrimitives = 0, uint32 NumVertices = 0);
+	D3D12RHI_API void RegisterGPUDispatch(FIntVector GroupCount);
 	
 	FD3D12SamplerState* CreateSampler(const FSamplerStateInitializerRHI& Initializer);
 	void CreateSamplerInternal(const D3D12_SAMPLER_DESC& Desc, D3D12_CPU_DESCRIPTOR_HANDLE Descriptor);
diff --git a/Engine/Source/Runtime/Engine/Private/CustomStaticScreenPercentage.cpp b/Engine/Source/Runtime/Engine/Private/CustomStaticScreenPercentage.cpp
new file mode 100644
--- /dev/null
+++ b/Engine/Source/Runtime/Engine/Private/CustomStaticScreenPercentage.cpp
@@ -0,0 +1,4 @@
+
+#include "CustomStaticScreenPercentage.h"
+
+ICustomStaticScreenPercentage* GCustomStaticScreenPercentage = nullptr;
diff --git a/Engine/Source/Runtime/Engine/Private/GameViewportClient.cpp b/Engine/Source/Runtime/Engine/Private/GameViewportClient.cpp
--- a/Engine/Source/Runtime/Engine/Private/GameViewportClient.cpp
+++ b/Engine/Source/Runtime/Engine/Private/GameViewportClient.cpp
@@ -61,7 +61,7 @@
 #include "IImageWrapper.h"
 #include "IImageWrapperModule.h"
 #include "HAL/PlatformApplicationMisc.h"
-
+#include "CustomStaticScreenPercentage.h"
 
 #if WITH_EDITOR
 #include "Settings/LevelEditorPlaySettings.h"
@@ -1306,6 +1306,104 @@ void UGameViewportClient::Draw(FViewport* InViewport, FCanvas* SceneCanvas)
 		}
 	}
 
+	// Setup the screen percentage and upscaling method for the view family.
+	bool bFinalScreenPercentageShowFlag;
+	bool bUsesDynamicResolution = false;
+	{
+		checkf(ViewFamily.GetScreenPercentageInterface() == nullptr,
+			TEXT("Some code has tried to set up an alien screen percentage driver, that could be wrong if not supported very well by the RHI."));
+
+		// Force screen percentage show flag to be turned off if not supported.
+		if (!ViewFamily.SupportsScreenPercentage())
+		{
+			ViewFamily.EngineShowFlags.ScreenPercentage = false;
+		}
+
+		// Set up secondary resolution fraction for the view family.
+		if (!bStereoRendering && ViewFamily.SupportsScreenPercentage())
+		{
+			float CustomSecondaruScreenPercentage = CVarSecondaryScreenPercentage.GetValueOnGameThread();
+
+			if (CustomSecondaruScreenPercentage > 0.0)
+			{
+				// Override secondary resolution fraction with CVar.
+				ViewFamily.SecondaryViewFraction = FMath::Min(CustomSecondaruScreenPercentage / 100.0f, 1.0f);
+			}
+			else
+			{
+				// Automatically compute secondary resolution fraction from DPI.
+				ViewFamily.SecondaryViewFraction = GetDPIDerivedResolutionFraction();
+			}
+
+			check(ViewFamily.SecondaryViewFraction > 0.0f);
+		}
+
+		// Setup main view family with screen percentage interface by dynamic resolution if screen percentage is enabled.
+		#if WITH_DYNAMIC_RESOLUTION
+		if (ViewFamily.EngineShowFlags.ScreenPercentage)
+		{
+			FDynamicResolutionStateInfos DynamicResolutionStateInfos;
+			GEngine->GetDynamicResolutionCurrentStateInfos(/* out */ DynamicResolutionStateInfos);
+
+			// Do not allow dynamic resolution to touch the view family if not supported to ensure there is no possibility to ruin
+			// game play experience on platforms that does not support it, but have it enabled by mistake.
+			if (DynamicResolutionStateInfos.Status == EDynamicResolutionStatus::Enabled)
+			{
+				GEngine->EmitDynamicResolutionEvent(EDynamicResolutionStateEvent::BeginDynamicResolutionRendering);
+				GEngine->GetDynamicResolutionState()->SetupMainViewFamily(ViewFamily);
+
+				bUsesDynamicResolution = ViewFamily.GetScreenPercentageInterface() != nullptr;
+			}
+			else if (DynamicResolutionStateInfos.Status == EDynamicResolutionStatus::DebugForceEnabled)
+			{
+				GEngine->EmitDynamicResolutionEvent(EDynamicResolutionStateEvent::BeginDynamicResolutionRendering);
+				ViewFamily.SetScreenPercentageInterface(new FLegacyScreenPercentageDriver(
+					ViewFamily,
+					DynamicResolutionStateInfos.ResolutionFractionApproximation,
+					/* AllowPostProcessSettingsScreenPercentage = */ false,
+					DynamicResolutionStateInfos.ResolutionFractionUpperBound));
+
+				bUsesDynamicResolution = true;
+			}
+
+			#if CSV_PROFILER
+			if (DynamicResolutionStateInfos.ResolutionFractionApproximation >= 0.0f)
+			{
+				CSV_CUSTOM_STAT_GLOBAL(DynamicResolutionPercentage, DynamicResolutionStateInfos.ResolutionFractionApproximation * 100.0f, ECsvCustomStatOp::Set);
+			}
+			#endif
+		}
+		#endif
+
+		if (GCustomStaticScreenPercentage)
+		{
+			GCustomStaticScreenPercentage->SetupMainGameViewFamily(ViewFamily);
+		}
+
+		// If a screen percentage interface was not set by dynamic resolution, then create one matching legacy behavior.
+		if (ViewFamily.GetScreenPercentageInterface() == nullptr)
+		{
+			bool AllowPostProcessSettingsScreenPercentage = false;
+			float GlobalResolutionFraction = 1.0f;
+
+			if (ViewFamily.EngineShowFlags.ScreenPercentage)
+			{
+				// Allow FPostProcessSettings::ScreenPercentage.
+				AllowPostProcessSettingsScreenPercentage = true;
+
+				// Get global view fraction set by r.ScreenPercentage.
+				GlobalResolutionFraction = FLegacyScreenPercentageDriver::GetCVarResolutionFraction();
+			}
+
+			ViewFamily.SetScreenPercentageInterface(new FLegacyScreenPercentageDriver(
+				ViewFamily, GlobalResolutionFraction, AllowPostProcessSettingsScreenPercentage));
+		}
+
+		check(ViewFamily.GetScreenPercentageInterface() != nullptr);
+
+		bFinalScreenPercentageShowFlag = ViewFamily.EngineShowFlags.ScreenPercentage;
+	}
+
 	TMap<ULocalPlayer*,FSceneView*> PlayerViewMap;
 	TArray<FSceneView*> Views;
 
@@ -1503,98 +1601,23 @@ void UGameViewportClient::Draw(FViewport* InViewport, FCanvas* SceneCanvas)
 		bBufferCleared = true;
 	}
 
-	// Force screen percentage show flag to be turned off if not supported.
-	if (!ViewFamily.SupportsScreenPercentage())
 	{
-		ViewFamily.EngineShowFlags.ScreenPercentage = false;
-	}
+		// Make sure the engine show flag for screen percentage is still what it was when setting up the screen percentage interface
+		ViewFamily.EngineShowFlags.ScreenPercentage = bFinalScreenPercentageShowFlag;
 
-	// Set up secondary resolution fraction for the view family.
-	if (!bStereoRendering && ViewFamily.SupportsScreenPercentage())
-	{
-		float CustomSecondaruScreenPercentage = CVarSecondaryScreenPercentage.GetValueOnGameThread();
-
-		if (CustomSecondaruScreenPercentage > 0.0)
+		if (bStereoRendering && bUsesDynamicResolution)
 		{
-			// Override secondary resolution fraction with CVar.
-			ViewFamily.SecondaryViewFraction = FMath::Min(CustomSecondaruScreenPercentage / 100.0f, 1.0f);
-		}
-		else
-		{
-			// Automatically compute secondary resolution fraction from DPI.
-			ViewFamily.SecondaryViewFraction = GetDPIDerivedResolutionFraction();
-		}
-
-		check(ViewFamily.SecondaryViewFraction > 0.0f);
-	}
-
-	checkf(ViewFamily.GetScreenPercentageInterface() == nullptr,
-		TEXT("Some code has tried to set up an alien screen percentage driver, that could be wrong if not supported very well by the RHI."));
-
-	// Setup main view family with screen percentage interface by dynamic resolution if screen percentage is enabled.
-	#if WITH_DYNAMIC_RESOLUTION
-	if (ViewFamily.EngineShowFlags.ScreenPercentage)
-	{
-		FDynamicResolutionStateInfos DynamicResolutionStateInfos;
-		GEngine->GetDynamicResolutionCurrentStateInfos(/* out */ DynamicResolutionStateInfos);
-
-		// Do not allow dynamic resolution to touch the view family if not supported to ensure there is no possibility to ruin
-		// game play experience on platforms that does not support it, but have it enabled by mistake.
-		if (DynamicResolutionStateInfos.Status == EDynamicResolutionStatus::Enabled)
-		{
-			GEngine->EmitDynamicResolutionEvent(EDynamicResolutionStateEvent::BeginDynamicResolutionRendering);
-			GEngine->GetDynamicResolutionState()->SetupMainViewFamily(ViewFamily);
-		}
-		else if (DynamicResolutionStateInfos.Status == EDynamicResolutionStatus::DebugForceEnabled)
-		{
-			GEngine->EmitDynamicResolutionEvent(EDynamicResolutionStateEvent::BeginDynamicResolutionRendering);
-			ViewFamily.SetScreenPercentageInterface(new FLegacyScreenPercentageDriver(
-				ViewFamily,
-				DynamicResolutionStateInfos.ResolutionFractionApproximation,
-				/* AllowPostProcessSettingsScreenPercentage = */ false,
-				DynamicResolutionStateInfos.ResolutionFractionUpperBound));
-		}
-
-		#if CSV_PROFILER
-		if (DynamicResolutionStateInfos.ResolutionFractionApproximation >= 0.0f)
-		{
-			CSV_CUSTOM_STAT_GLOBAL(DynamicResolutionPercentage, DynamicResolutionStateInfos.ResolutionFractionApproximation * 100.0f, ECsvCustomStatOp::Set);
-		}
-		#endif
-	}
-	#endif
-
-	// If a screen percentage interface was not set by dynamic resolution, then create one matching legacy behavior.
-	if (ViewFamily.GetScreenPercentageInterface() == nullptr)
-	{
-		bool AllowPostProcessSettingsScreenPercentage = false;
-		float GlobalResolutionFraction = 1.0f;
-
-		if (ViewFamily.EngineShowFlags.ScreenPercentage)
-		{
-			// Allow FPostProcessSettings::ScreenPercentage.
-			AllowPostProcessSettingsScreenPercentage = true;
-
-			// Get global view fraction set by r.ScreenPercentage.
-			GlobalResolutionFraction = FLegacyScreenPercentageDriver::GetCVarResolutionFraction();
-		}
-
-		ViewFamily.SetScreenPercentageInterface(new FLegacyScreenPercentageDriver(
-			ViewFamily, GlobalResolutionFraction, AllowPostProcessSettingsScreenPercentage));
-	}
-	else if (bStereoRendering)
-	{
-		// Change screen percentage method to raw output when doing dynamic resolution with VR if not using TAA upsample.
-		for (FSceneView* View : Views)
-		{
-			if (View->PrimaryScreenPercentageMethod == EPrimaryScreenPercentageMethod::SpatialUpscale)
+			// Change screen percentage method to raw output when doing dynamic resolution with VR if not using TAA upsample.
+			for (FSceneView* View : Views)
 			{
-				View->PrimaryScreenPercentageMethod = EPrimaryScreenPercentageMethod::RawOutput;
+				if (View->PrimaryScreenPercentageMethod == EPrimaryScreenPercentageMethod::SpatialUpscale)
+				{
+					View->PrimaryScreenPercentageMethod = EPrimaryScreenPercentageMethod::RawOutput;
+				}
 			}
 		}
 	}
 
-	
 	ViewFamily.bIsHDR = GetWindow().IsValid() ? GetWindow().Get()->GetIsHDR() : false;
 
 	// Draw the player views.
diff --git a/Engine/Source/Runtime/Engine/Private/SceneView.cpp b/Engine/Source/Runtime/Engine/Private/SceneView.cpp
--- a/Engine/Source/Runtime/Engine/Private/SceneView.cpp
+++ b/Engine/Source/Runtime/Engine/Private/SceneView.cpp
@@ -806,9 +806,9 @@ FSceneView::FSceneView(const FSceneViewInitOptions& InitOptions)
 
 	SetupAntiAliasingMethod();
 
-	if (CVarEnableTemporalUpsample.GetValueOnAnyThread() && AntiAliasingMethod == AAM_TemporalAA && GetFeatureLevel() >= ERHIFeatureLevel::SM5)
+	if ((AntiAliasingMethod == AAM_TemporalAA && GetFeatureLevel() >= ERHIFeatureLevel::SM5) &&
+		(CVarEnableTemporalUpsample.GetValueOnAnyThread() || Family->GetTemporalUpscalerInterface() != nullptr))
 	{
-		// The renderer will automatically fallback to SpatialUpscale if not using TemporalAA anti aliasing method.
 		PrimaryScreenPercentageMethod = EPrimaryScreenPercentageMethod::TemporalUpscale;
 	}
 
@@ -2483,7 +2483,8 @@ FSceneViewFamily::FSceneViewFamily(const ConstructionValues& CVS)
 	GammaCorrection(CVS.GammaCorrection),
 	SecondaryViewFraction(1.0f),
 	SecondaryScreenPercentageMethod(ESecondaryScreenPercentageMethod::LowerPixelDensitySimulation),
-	ScreenPercentageInterface(nullptr)
+	ScreenPercentageInterface(nullptr),
+	TemporalUpscalerInterface(nullptr)
 {
 	// If we do not pass a valid scene pointer then SetWorldTimes must be called to initialized with valid times.
 	ensure(CVS.bTimesSet);
diff --git a/Engine/Source/Runtime/Engine/Public/CustomStaticScreenPercentage.h b/Engine/Source/Runtime/Engine/Public/CustomStaticScreenPercentage.h
new file mode 100644
--- /dev/null
+++ b/Engine/Source/Runtime/Engine/Public/CustomStaticScreenPercentage.h
@@ -0,0 +1,17 @@
+#pragma once
+#include "CoreMinimal.h"
+class FSceneViewFamily;
+
+class ENGINE_API ICustomStaticScreenPercentage
+{
+public:
+	virtual ~ICustomStaticScreenPercentage() {};
+
+	virtual void SetupMainGameViewFamily(FSceneViewFamily& ViewFamily) = 0;
+
+	virtual float GetMinUpsampleResolutionFraction() const = 0;
+	virtual float GetMaxUpsampleResolutionFraction() const = 0;
+
+};
+
+extern ENGINE_API ICustomStaticScreenPercentage* GCustomStaticScreenPercentage;
diff --git a/Engine/Source/Runtime/Engine/Public/SceneView.h b/Engine/Source/Runtime/Engine/Public/SceneView.h
--- a/Engine/Source/Runtime/Engine/Public/SceneView.h
+++ b/Engine/Source/Runtime/Engine/Public/SceneView.h
@@ -27,6 +27,7 @@ class ISceneViewExtension;
 class FSceneViewFamily;
 class FVolumetricFogViewResources;
 class FIESLightProfileResource;
+class ITemporalUpscaler;
 
 enum class ERayTracingRenderMode
 {
@@ -1773,13 +1774,28 @@ public:
 		: FSceneViewFamily(static_cast<const FSceneViewFamily&>(InViewFamily))
 	{
 		check(ScreenPercentageInterface == nullptr);
+		check(TemporalUpscalerInterface == nullptr);
 	}
 
 
+	FORCEINLINE void SetTemporalUpscalerInterface(const ITemporalUpscaler* InTemporalUpscalerInterface)
+	{
+		check(InTemporalUpscalerInterface);
+		checkf(TemporalUpscalerInterface == nullptr, TEXT("View family already had a temporal upscaler assigned."));
+		TemporalUpscalerInterface = InTemporalUpscalerInterface;
+	}
+
+	FORCEINLINE const ITemporalUpscaler* GetTemporalUpscalerInterface() const
+	{
+		return TemporalUpscalerInterface;
+	}
+
 private:
 	/** Interface to handle screen percentage of the views of the family. */
 	ISceneViewFamilyScreenPercentage* ScreenPercentageInterface;
 
+	const ITemporalUpscaler* TemporalUpscalerInterface;
+
 	// Only FSceneRenderer can copy a view family.
 	FSceneViewFamily(const FSceneViewFamily&) = default;
 
diff --git a/Engine/Source/Runtime/RenderCore/Private/CustomResourcePool.cpp b/Engine/Source/Runtime/RenderCore/Private/CustomResourcePool.cpp
new file mode 100644
--- /dev/null
+++ b/Engine/Source/Runtime/RenderCore/Private/CustomResourcePool.cpp
@@ -0,0 +1,12 @@
+
+#include "CustomResourcePool.h"
+
+ICustomResourcePool* GCustomResourcePool = nullptr;
+
+void ICustomResourcePool::TickPoolElements()
+{
+	if (GCustomResourcePool)
+	{
+		GCustomResourcePool->Tick();
+	}
+}
diff --git a/Engine/Source/Runtime/RenderCore/Public/CustomResourcePool.h b/Engine/Source/Runtime/RenderCore/Public/CustomResourcePool.h
new file mode 100644
--- /dev/null
+++ b/Engine/Source/Runtime/RenderCore/Public/CustomResourcePool.h
@@ -0,0 +1,12 @@
+#pragma once
+#include "CoreMinimal.h"
+
+class RENDERCORE_API ICustomResourcePool
+{
+public:
+	virtual ~ICustomResourcePool() {}
+	virtual void Tick() = 0;
+
+	static void TickPoolElements();
+};
+extern RENDERCORE_API ICustomResourcePool* GCustomResourcePool;
diff --git a/Engine/Source/Runtime/Renderer/Private/IndirectLightRendering.cpp b/Engine/Source/Runtime/Renderer/Private/IndirectLightRendering.cpp
--- a/Engine/Source/Runtime/Renderer/Private/IndirectLightRendering.cpp
+++ b/Engine/Source/Runtime/Renderer/Private/IndirectLightRendering.cpp
@@ -1457,6 +1457,8 @@ void FDeferredShadingSceneRenderer::RenderDeferredReflectionsAndSkyLighting(
 				check(View.ViewState);
 				FTAAPassParameters TAASettings(View);
 				TAASettings.Pass = ETAAPassConfig::ScreenSpaceReflections;
+				TAASettings.SceneDepthTexture = SceneTextures.SceneDepthBuffer;
+				TAASettings.SceneVelocityTexture = SceneTextures.SceneVelocityBuffer;
 				TAASettings.SceneColorInput = DenoiserInputs.Color;
 				TAASettings.bOutputRenderTargetable = (
 					ViewPipelineState.bComposePlanarReflections ||
@@ -1465,7 +1467,6 @@ void FDeferredShadingSceneRenderer::RenderDeferredReflectionsAndSkyLighting(
 
 				FTAAOutputs TAAOutputs = AddTemporalAAPass(
 					GraphBuilder,
-					SceneTextures,
 					View,
 					TAASettings,
 					View.PrevViewInfo.SSRHistory,
diff --git a/Engine/Source/Runtime/Renderer/Private/LightShaftRendering.cpp b/Engine/Source/Runtime/Renderer/Private/LightShaftRendering.cpp
--- a/Engine/Source/Runtime/Renderer/Private/LightShaftRendering.cpp
+++ b/Engine/Source/Runtime/Renderer/Private/LightShaftRendering.cpp
@@ -581,12 +581,13 @@ void ApplyTemporalAA(
 
 		FTAAPassParameters TAAParameters(View);
 		TAAParameters.Pass = ETAAPassConfig::LightShaft;
+		TAAParameters.SceneDepthTexture = SceneTextures.SceneDepthBuffer;
+		TAAParameters.SceneVelocityTexture = SceneTextures.SceneVelocityBuffer;
 		TAAParameters.SetupViewRect(View, /** ResolutionDivisor = */ 2);
 		TAAParameters.SceneColorInput = LightShaftSetup;
 
 		FTAAOutputs Outputs = AddTemporalAAPass(
 			GraphBuilder,
-			SceneTextures,
 			View,
 			TAAParameters,
 			*HistoryState,
diff --git a/Engine/Source/Runtime/Renderer/Private/PostProcess/DiaphragmDOF.cpp b/Engine/Source/Runtime/Renderer/Private/PostProcess/DiaphragmDOF.cpp
--- a/Engine/Source/Runtime/Renderer/Private/PostProcess/DiaphragmDOF.cpp
+++ b/Engine/Source/Runtime/Renderer/Private/PostProcess/DiaphragmDOF.cpp
@@ -1617,12 +1617,13 @@ FRDGTextureRef DiaphragmDOF::AddPasses(
 	// TAA the setup for the convolution to be temporally stable.
 	if (View.AntiAliasingMethod == AAM_TemporalAA && ViewState)
 	{
+		TAAParameters.SceneDepthTexture = SceneTextures.SceneDepthBuffer;
+		TAAParameters.SceneVelocityTexture = SceneTextures.SceneVelocityBuffer;
 		TAAParameters.SceneColorInput = HalfResGatherInputTextures.SceneColor;
 		TAAParameters.SceneMetadataInput = HalfResGatherInputTextures.SeparateCoc;
 
 		FTAAOutputs TAAOutputs = AddTemporalAAPass(
 			GraphBuilder,
-			SceneTextures,
 			View,
 			TAAParameters,
 			View.PrevViewInfo.DOFSetupHistory,
diff --git a/Engine/Source/Runtime/Renderer/Private/PostProcess/PostProcessing.cpp b/Engine/Source/Runtime/Renderer/Private/PostProcess/PostProcessing.cpp
--- a/Engine/Source/Runtime/Renderer/Private/PostProcess/PostProcessing.cpp
+++ b/Engine/Source/Runtime/Renderer/Private/PostProcess/PostProcessing.cpp
@@ -477,30 +477,39 @@ void AddPostProcessingPasses(FRDGBuilder& GraphBuilder, const FViewInfo& View, c
 				// TemporalAA is only able to match the low quality mode (box filter).
 				GetDownsampleQuality() == EDownsampleQuality::Low;
 
-			if (CVarTAAAlgorithm.GetValueOnRenderThread())
-			{
-				AddTemporalAA2Passes(
-					GraphBuilder,
-					SceneTextures,
-					View,
-					SceneColor.Texture,
-					&SceneColor.Texture,
-					&SecondaryViewRect);
-			}
-			else
-			{
-				AddTemporalAAPass(
-					GraphBuilder,
-					SceneTextures,
-					View,
-					bAllowSceneDownsample,
-					DownsampleOverrideFormat,
-					SceneColor.Texture,
-					&SceneColor.Texture,
-					&SecondaryViewRect,
-					&DownsampledSceneColor.Texture,
-					&DownsampledSceneColor.ViewRect);
-			}
+			int32 UpscaleMode = ITemporalUpscaler::GetTemporalUpscalerMode();
+
+			const ITemporalUpscaler* DefaultTemporalUpscaler = ITemporalUpscaler::GetDefaultTemporalUpscaler();
+			const ITemporalUpscaler* UpscalerToUse = ( UpscaleMode == 0 || !View.Family->GetTemporalUpscalerInterface())? DefaultTemporalUpscaler : View.Family->GetTemporalUpscalerInterface();
+
+			const TCHAR* UpscalerName = UpscalerToUse->GetDebugName();
+
+			// Standard event scope for temporal upscaler to have all profiling information not matter what, and with explicit detection of third party.
+			RDG_EVENT_SCOPE_CONDITIONAL(
+				GraphBuilder,
+				UpscalerToUse != DefaultTemporalUpscaler,
+				"ThirdParty %s %dx%d -> %dx%d",
+				UpscalerToUse->GetDebugName(),
+				View.ViewRect.Width(), View.ViewRect.Height(),
+				View.GetSecondaryViewRectSize().X, View.GetSecondaryViewRectSize().Y);
+
+			ITemporalUpscaler::FPassInputs UpscalerPassInputs;
+
+			UpscalerPassInputs.bAllowDownsampleSceneColor = bAllowSceneDownsample;
+			UpscalerPassInputs.DownsampleOverrideFormat = DownsampleOverrideFormat;
+			UpscalerPassInputs.SceneColorTexture = SceneColor.Texture;
+			UpscalerPassInputs.SceneDepthTexture = SceneDepth.Texture;
+			UpscalerPassInputs.SceneVelocityTexture = Velocity.Texture;
+			UpscalerPassInputs.EyeAdaptationTexture = GetEyeAdaptationTexture(GraphBuilder, View);
+
+			UpscalerToUse->AddPasses(
+				GraphBuilder,
+				View,
+				UpscalerPassInputs,
+				&SceneColor.Texture,
+				&SecondaryViewRect,
+				&DownsampledSceneColor.Texture,
+				&DownsampledSceneColor.ViewRect);
 		}
 		else if (ScreenSpaceRayTracing::ShouldRenderScreenSpaceReflections(View))
 		{
@@ -1095,12 +1104,14 @@ void AddDebugViewPostProcessingPasses(FRDGBuilder& GraphBuilder, const FViewInfo
 		case DVSM_RayTracingDebug:
 		{
 			FTAAPassParameters Parameters(View);
+			Parameters.SceneDepthTexture = SceneTextures.SceneDepthBuffer;
+			Parameters.SceneVelocityTexture = SceneTextures.SceneVelocityBuffer;
 			Parameters.SceneColorInput = SceneColor.Texture;
 
 			const FTemporalAAHistory& InputHistory = View.PrevViewInfo.TemporalAAHistory;
 			FTemporalAAHistory* OutputHistory = &View.ViewState->PrevFrameViewInfo.TemporalAAHistory;
 
-			FTAAOutputs Outputs = AddTemporalAAPass(GraphBuilder, SceneTextures, View, Parameters, InputHistory, OutputHistory);
+			FTAAOutputs Outputs = AddTemporalAAPass(GraphBuilder, View, Parameters, InputHistory, OutputHistory);
 			SceneColor.Texture = Outputs.SceneColor;
 
 			break;
@@ -2013,11 +2024,15 @@ void FPostProcessing::ProcessPlanarReflection(FRHICommandListImmediate& RHICmdLi
 		FTemporalAAHistory* OutputHistory = &ViewState->PrevFrameViewInfo.TemporalAAHistory;
 
 		FTAAPassParameters Parameters(View);
+		Parameters.SceneDepthTexture = SceneTextures.SceneDepthBuffer;
+
+		// Planar reflections don't support velocity.
+		Parameters.SceneVelocityTexture = nullptr;
+
 		Parameters.SceneColorInput = GraphBuilder.RegisterExternalTexture(SceneContext.GetSceneColor(), TEXT("SceneColor"));
 
 		FTAAOutputs PassOutputs = AddTemporalAAPass(
 			GraphBuilder,
-			SceneTextures,
 			View,
 			Parameters,
 			InputHistory,
diff --git a/Engine/Source/Runtime/Renderer/Private/PostProcess/TemporalAA.cpp b/Engine/Source/Runtime/Renderer/Private/PostProcess/TemporalAA.cpp
--- a/Engine/Source/Runtime/Renderer/Private/PostProcess/TemporalAA.cpp
+++ b/Engine/Source/Runtime/Renderer/Private/PostProcess/TemporalAA.cpp
@@ -69,6 +69,14 @@ TAutoConsoleVariable<int32> CVarTAANyquistHistory(
 	TEXT(""),
 	ECVF_RenderThreadSafe);
 
+static TAutoConsoleVariable<int32> CVarUseTemporalAAUpscaler(
+	TEXT("r.TemporalAA.Upscaler"),
+	1,
+	TEXT("Choose the upscaling algorithm.\n")
+	TEXT(" 0: Forces the default temporal upscaler of the renderer;\n")
+	TEXT(" 1: GTemporalUpscaler which may be overridden by a third party plugin (default)."),
+	ECVF_RenderThreadSafe);
+
 inline bool DoesPlatformSupportTemporalHistoryUpscale(EShaderPlatform Platform)
 {
 	return (IsPCPlatform(Platform) || FDataDrivenShaderPlatformInfo::GetSupportsTemporalHistoryUpscale(Platform))
@@ -172,8 +180,6 @@ class FTAAStandaloneCS : public FGlobalShader
 		SHADER_PARAMETER(FVector4, HistoryBufferUVMinMax)
 		SHADER_PARAMETER(FVector4, ScreenPosToHistoryBufferUV)
 
-		SHADER_PARAMETER_STRUCT_INCLUDE(FSceneTextureParameters, SceneTextures)
-
 		SHADER_PARAMETER_RDG_TEXTURE(Texture2D, EyeAdaptation)
 
 		// Inputs
@@ -186,7 +192,9 @@ class FTAAStandaloneCS : public FGlobalShader
 		SHADER_PARAMETER_RDG_TEXTURE_ARRAY(Texture2D, HistoryBuffer, [FTemporalAAHistory::kRenderTargetCount])
 		SHADER_PARAMETER_SAMPLER_ARRAY(SamplerState, HistoryBufferSampler, [FTemporalAAHistory::kRenderTargetCount])
 
+		SHADER_PARAMETER_RDG_TEXTURE(Texture2D, SceneDepthBuffer)
 		SHADER_PARAMETER_SAMPLER(SamplerState, SceneDepthBufferSampler)
+		SHADER_PARAMETER_RDG_TEXTURE(Texture2D, SceneVelocityBuffer)
 		SHADER_PARAMETER_SAMPLER(SamplerState, SceneVelocityBufferSampler)
 
 		SHADER_PARAMETER_RDG_TEXTURE_SRV(Texture2D, StencilTexture)
@@ -590,7 +598,6 @@ bool FTAAPassParameters::Validate() const
 
 FTAAOutputs AddTemporalAAPass(
 	FRDGBuilder& GraphBuilder,
-	const FSceneTextureParameters& SceneTextures,
 	const FViewInfo& View,
 	const FTAAPassParameters& Inputs,
 	const FTemporalAAHistory& InputHistory,
@@ -726,16 +733,18 @@ FTAAOutputs AddTemporalAAPass(
 		PassParameters->CurrentFrameWeight = CVarTemporalAACurrentFrameWeight.GetValueOnRenderThread();
 		PassParameters->bCameraCut = bCameraCut;
 
-		PassParameters->SceneTextures = SceneTextures;
+		PassParameters->SceneDepthBuffer = Inputs.SceneDepthTexture;
+		PassParameters->SceneVelocityBuffer = Inputs.SceneVelocityTexture;
+
 		PassParameters->SceneDepthBufferSampler = TStaticSamplerState<SF_Point>::GetRHI();
 		PassParameters->SceneVelocityBufferSampler = TStaticSamplerState<SF_Point>::GetRHI();
 
-		PassParameters->StencilTexture = GraphBuilder.CreateSRV(FRDGTextureSRVDesc::CreateWithPixelFormat(SceneTextures.SceneDepthBuffer, PF_X24_G8));
+		PassParameters->StencilTexture = GraphBuilder.CreateSRV(FRDGTextureSRVDesc::CreateWithPixelFormat(Inputs.SceneDepthTexture, PF_X24_G8));
 
 		// We need a valid velocity buffer texture. Use black (no velocity) if none exists.
-		if (!PassParameters->SceneTextures.SceneVelocityBuffer)
+		if (!PassParameters->SceneVelocityBuffer)
 		{
-			PassParameters->SceneTextures.SceneVelocityBuffer = GraphBuilder.RegisterExternalTexture(GSystemTextures.BlackDummy);;
+			PassParameters->SceneVelocityBuffer = GraphBuilder.RegisterExternalTexture(GSystemTextures.BlackDummy);;
 		}
 
 		// Input buffer shader parameters
@@ -774,7 +783,7 @@ FTAAOutputs AddTemporalAAPass(
 				}
 
 				// Remove dependency of the velocity buffer on camera cut, given it's going to be ignored by the shader.
-				PassParameters->SceneTextures.SceneVelocityBuffer = BlackDummy;
+				PassParameters->SceneVelocityBuffer = BlackDummy;
 			}
 			else
 			{
@@ -1384,13 +1393,10 @@ void AddTemporalAA2Passes(
 	*OutSceneColorViewRect = OutputRect;
 } // AddTemporalAA2Passes()
 
-void AddTemporalAAPass(
+static void AddTemporalAAPass(
 	FRDGBuilder& GraphBuilder,
-	const FSceneTextureParameters& SceneTextures,
 	const FViewInfo& View,
-	const bool bAllowDownsampleSceneColor,
-	const EPixelFormat DownsampleOverrideFormat,
-	FRDGTextureRef InSceneColorTexture,
+	const ITemporalUpscaler::FPassInputs& PassInputs,
 	FRDGTextureRef* OutSceneColorTexture,
 	FIntRect* OutSceneColorViewRect,
 	FRDGTextureRef* OutSceneColorHalfResTexture,
@@ -1431,19 +1437,21 @@ void AddTemporalAAPass(
 		TAAParameters.OutputViewRect.Max = HistoryViewSize;
 	}
 
-	TAAParameters.DownsampleOverrideFormat = DownsampleOverrideFormat;
+	TAAParameters.DownsampleOverrideFormat = PassInputs.DownsampleOverrideFormat;
 
-	TAAParameters.bDownsample = bAllowDownsampleSceneColor && TAAParameters.bUseFast;
+	TAAParameters.bDownsample = PassInputs.bAllowDownsampleSceneColor && TAAParameters.bUseFast;
 
-	TAAParameters.SceneColorInput = InSceneColorTexture;
+	TAAParameters.SceneDepthTexture = PassInputs.SceneDepthTexture;
+	TAAParameters.SceneVelocityTexture = PassInputs.SceneVelocityTexture;
+	TAAParameters.SceneColorInput = PassInputs.SceneColorTexture;
 
 	const FTemporalAAHistory& InputHistory = View.PrevViewInfo.TemporalAAHistory;
 
 	FTemporalAAHistory& OutputHistory = View.ViewState->PrevFrameViewInfo.TemporalAAHistory;
 
-	const FTAAOutputs TAAOutputs = AddTemporalAAPass(
+
+	const FTAAOutputs TAAOutputs = ::AddTemporalAAPass(
 		GraphBuilder,
-		SceneTextures,
 		View,
 		TAAParameters,
 		InputHistory,
@@ -1461,8 +1469,8 @@ void AddTemporalAAPass(
 
 		FScreenPassTextureViewport OutputViewport;
 		OutputViewport.Rect = SecondaryViewRect;
-		OutputViewport.Extent.X = FMath::Max(InSceneColorTexture->Desc.Extent.X, QuantizedOutputSize.X);
-		OutputViewport.Extent.Y = FMath::Max(InSceneColorTexture->Desc.Extent.Y, QuantizedOutputSize.Y);
+		OutputViewport.Extent.X = FMath::Max(PassInputs.SceneColorTexture->Desc.Extent.X, QuantizedOutputSize.X);
+		OutputViewport.Extent.Y = FMath::Max(PassInputs.SceneColorTexture->Desc.Extent.Y, QuantizedOutputSize.Y);
 
 		SceneColorTexture = ComputeMitchellNetravaliDownsample(GraphBuilder, View, FScreenPassTexture(SceneColorTexture, InputViewport), OutputViewport);
 	}
@@ -1471,96 +1479,47 @@ void AddTemporalAAPass(
 	*OutSceneColorViewRect = SecondaryViewRect;
 	*OutSceneColorHalfResTexture = TAAOutputs.DownsampledSceneColor;
 	*OutSceneColorHalfResViewRect = FIntRect::DivideAndRoundUp(SecondaryViewRect, 2);
-}
+} // AddTemporalAAPass()
 
-//////////////////////////////////////////////////////////////////////////
-//! Legacy - Only used by DebugViewModeRendering. Remove after porting.
+const ITemporalUpscaler* GTemporalUpscaler = nullptr;
 
-class FRCPassPostProcessTemporalAA : public TRenderingCompositePassBase<3, 3>
-{
+class FDefaultTemporalUpscaler : public ITemporalUpscaler
+{ 
 public:
-	FRCPassPostProcessTemporalAA(
-		const FPostprocessContext& Context,
-		const FTAAPassParameters& InParameters,
-		const FTemporalAAHistory& InInputHistory,
-		FTemporalAAHistory* OutOutputHistory)
-		: SavedParameters(InParameters)
-		, InputHistory(InInputHistory)
-		, OutputHistory(OutOutputHistory)
-	{
-		check(InParameters.Pass == ETAAPassConfig::Main);
-		check(SavedParameters.Validate());
-
-		bIsComputePass = true;
-		bPreferAsyncCompute = false;
-	}
 
-	void Process(FRenderingCompositePassContext& Context) override
+	virtual const TCHAR* GetDebugName() const
 	{
-		WaitForInputPassComputeFences(Context.RHICmdList);
-
-		AsyncEndFence = FComputeFenceRHIRef();
-
-		FRDGBuilder GraphBuilder(Context.RHICmdList);
-
-		FSceneTextureParameters SceneTextures;
-		SetupSceneTextureParameters(GraphBuilder, &SceneTextures);
-
-		FTAAPassParameters Parameters = SavedParameters;
-		Parameters.SceneColorInput = CreateRDGTextureForRequiredInput(GraphBuilder, ePId_Input0, TEXT("SceneColor"));
-
-#if WITH_MGPU
-		static const FName NameForTemporalEffect("RCPassPostProcessTemporalAA");
-		GraphBuilder.SetNameForTemporalEffect(FName(NameForTemporalEffect, Context.View.ViewState->UniqueID));
-#endif
-		FTAAOutputs Outputs = AddTemporalAAPass(
-			GraphBuilder,
-			SceneTextures,
-			Context.View,
-			Parameters,
-			InputHistory,
-			/* out */ OutputHistory);
-
-		ExtractRDGTextureForOutput(GraphBuilder, ePId_Output0, Outputs.SceneColor);
-
-		GraphBuilder.Execute();
+		return TEXT("FDefaultTemporalUpscaler");
 	}
 
-	void Release() override { delete this; }
-
-	FPooledRenderTargetDesc ComputeOutputDesc(EPassOutputId InPassOutputId) const override
+	virtual void AddPasses(
+		FRDGBuilder& GraphBuilder,
+		const FViewInfo& View,
+		const FPassInputs& PassInputs,
+		FRDGTextureRef* OutSceneColorTexture,
+		FIntRect* OutSceneColorViewRect,
+		FRDGTextureRef* OutSceneColorHalfResTexture,
+		FIntRect* OutSceneColorHalfResViewRect) const final
 	{
-		// ExtractRDGTextureForOutput() is doing this work for us already.
-		return FPooledRenderTargetDesc();
+		return AddTemporalAAPass(
+			GraphBuilder,
+			View,
+			PassInputs,
+			OutSceneColorTexture,
+			OutSceneColorViewRect,
+			OutSceneColorHalfResTexture,
+			OutSceneColorHalfResViewRect);
 	}
-
-	FRHIComputeFence* GetComputePassEndFence() const override { return AsyncEndFence; }
-
-private:
-	const FTAAPassParameters SavedParameters;
-
-	FComputeFenceRHIRef AsyncEndFence;
-
-	const FTemporalAAHistory& InputHistory;
-	FTemporalAAHistory* OutputHistory;
 };
 
-FRenderingCompositeOutputRef AddTemporalAADebugViewPass(FPostprocessContext& Context)
+// static
+const ITemporalUpscaler* ITemporalUpscaler::GetDefaultTemporalUpscaler()
 {
-	check(Context.View.ViewState);
-
-	FSceneViewState* ViewState = Context.View.ViewState;
-
-	FTAAPassParameters Parameters(Context.View);
-
-	FRCPassPostProcessTemporalAA* TemporalAAPass = Context.Graph.RegisterPass(new(FMemStack::Get()) FRCPassPostProcessTemporalAA(
-		Context,
-		Parameters,
-		Context.View.PrevViewInfo.TemporalAAHistory,
-		&ViewState->PrevFrameViewInfo.TemporalAAHistory));
-
-	TemporalAAPass->SetInput(ePId_Input0, Context.FinalOutput);
-	return FRenderingCompositeOutputRef(TemporalAAPass);
+	static FDefaultTemporalUpscaler DefaultTemporalUpscaler;
+	return &DefaultTemporalUpscaler;
 }
 
-//////////////////////////////////////////////////////////////////////////
\ No newline at end of file
+int ITemporalUpscaler::GetTemporalUpscalerMode()
+{
+	return CVarUseTemporalAAUpscaler.GetValueOnRenderThread();
+}
diff --git a/Engine/Source/Runtime/Renderer/Private/PostProcess/TemporalAA.h b/Engine/Source/Runtime/Renderer/Private/PostProcess/TemporalAA.h
--- a/Engine/Source/Runtime/Renderer/Private/PostProcess/TemporalAA.h
+++ b/Engine/Source/Runtime/Renderer/Private/PostProcess/TemporalAA.h
@@ -4,7 +4,6 @@
 
 #include "ScreenPass.h"
 
-class FSceneTextureParameters;
 struct FTemporalAAHistory;
 
 /** List of TAA configurations. */
@@ -83,6 +82,10 @@ struct FTAAPassParameters
 	// Resolution divisor.
 	int32 ResolutionDivisor = 1;
 
+	// Full resolution depth and velocity textures to reproject the history.
+	FRDGTexture* SceneDepthTexture = nullptr;
+	FRDGTexture* SceneVelocityTexture = nullptr;
+
 	// Anti aliased scene color.
 	// Can have alpha channel, or CoC for DOF.
 	FRDGTexture* SceneColorInput = nullptr;
@@ -133,26 +136,49 @@ struct FTAAPassParameters
 };
 
 /** Temporal AA pass which emits a filtered scene color and new history. */
-FTAAOutputs AddTemporalAAPass(
+extern RENDERER_API FTAAOutputs AddTemporalAAPass(
 	FRDGBuilder& GraphBuilder,
-	const FSceneTextureParameters& SceneTextures,
 	const FViewInfo& View,
 	const FTAAPassParameters& Inputs,
 	const FTemporalAAHistory& InputHistory,
 	FTemporalAAHistory* OutputHistory);
 
-/** Temporal AA helper method which performs filtering on the main pass scene color. Supports upsampled history and,
- *  if requested, will attempt to perform the scene color downsample. Returns the filtered scene color, the downsampled
- *  scene color (or null if it was not performed), and the secondary view rect.
- */
-void AddTemporalAAPass(
-	FRDGBuilder& GraphBuilder,
-	const FSceneTextureParameters& SceneTextures,
-	const FViewInfo& View,
-	const bool bAllowDownsampleSceneColor,
-	const EPixelFormat DownsampleOverrideFormat,
-	FRDGTextureRef InSceneColorTexture,
-	FRDGTextureRef* OutSceneColorTexture,
-	FIntRect* OutSceneColorViewRect,
-	FRDGTextureRef* OutSceneColorHalfResTexture,
-	FIntRect* OutSceneColorHalfResViewRect);
\ No newline at end of file
+/** Interface for the main temporal upscaling algorithm. */
+class RENDERER_API ITemporalUpscaler
+{
+public:
+
+	struct FPassInputs
+	{
+		bool bAllowDownsampleSceneColor;
+		EPixelFormat DownsampleOverrideFormat;
+		FRDGTextureRef SceneColorTexture;
+		FRDGTextureRef SceneDepthTexture;
+		FRDGTextureRef SceneVelocityTexture;
+		FRDGTextureRef EyeAdaptationTexture;
+	};
+
+	virtual ~ITemporalUpscaler() {};
+
+	virtual const TCHAR* GetDebugName() const = 0;
+
+	///** Temporal AA helper method which performs filtering on the main pass scene color. Supports upsampled history and,
+	// *  if requested, will attempt to perform the scene color downsample. Returns the filtered scene color, the downsampled
+	// *  scene color (or null if it was not performed), and the secondary view rect.
+	// */
+
+	virtual void AddPasses(
+		FRDGBuilder& GraphBuilder,
+		const FViewInfo& View,
+		const FPassInputs& PassInputs,
+		FRDGTextureRef* OutSceneColorTexture,
+		FIntRect* OutSceneColorViewRect,
+		FRDGTextureRef* OutSceneColorHalfResTexture,
+		FIntRect* OutSceneColorHalfResViewRect) const = 0;
+
+	static const ITemporalUpscaler* GetDefaultTemporalUpscaler();
+	static int GetTemporalUpscalerMode();
+}; 
+
+extern RENDERER_API const ITemporalUpscaler* GTemporalUpscaler;
+
diff --git a/Engine/Source/Runtime/Renderer/Private/Renderer.cpp b/Engine/Source/Runtime/Renderer/Private/Renderer.cpp
--- a/Engine/Source/Runtime/Renderer/Private/Renderer.cpp
+++ b/Engine/Source/Runtime/Renderer/Private/Renderer.cpp
@@ -33,7 +33,7 @@
 #include "VisualizeTexturePresent.h"
 #include "ScreenSpaceDenoise.h"
 #include "VT/VirtualTextureSystem.h"
-
+#include "PostProcess/TemporalAA.h"
 DEFINE_LOG_CATEGORY(LogRenderer);
 
 IMPLEMENT_MODULE(FRendererModule, Renderer);
@@ -53,6 +53,7 @@ FAutoConsoleVariableRef CVarFlushRenderTargetsOnWorldCleanup(TEXT("r.bFlushRende
 void FRendererModule::StartupModule()
 {
 	GScreenSpaceDenoiser = IScreenSpaceDenoiser::GetDefaultDenoiser();
+	GTemporalUpscaler = ITemporalUpscaler::GetDefaultTemporalUpscaler();
 
 	FVirtualTextureSystem::Initialize();
 }
diff --git a/Engine/Source/Runtime/Renderer/Private/SceneRendering.h b/Engine/Source/Runtime/Renderer/Private/SceneRendering.h
--- a/Engine/Source/Runtime/Renderer/Private/SceneRendering.h
+++ b/Engine/Source/Runtime/Renderer/Private/SceneRendering.h
@@ -811,6 +811,12 @@ struct FGTAOTAAHistory
 };
 
 
+// Plugins can derive from this and use it for their own purposes
+class RENDERER_API ICustomTemporalAAHistory : public IRefCountedObject
+{
+public:
+	virtual ~ICustomTemporalAAHistory() {}
+};
 
 // Structure that hold all information related to previous frame.
 struct FPreviousViewInfo
@@ -852,6 +858,9 @@ struct FPreviousViewInfo
 	// Temporal AA result of last frame
 	FTemporalAAHistory TemporalAAHistory;
 
+	// Custom Temporal AA result of last frame, used by plugins
+	TRefCountPtr<ICustomTemporalAAHistory> CustomTemporalAAHistory;
+
 	// Half resolution version temporal AA result of last frame
 	TRefCountPtr<IPooledRenderTarget> HalfResTemporalAAHistory;
 
@@ -1297,8 +1306,8 @@ public:
 #endif
 
 	/** Returns the size of view rect after primary upscale ( == only with secondary screen percentage). */
-	FIntPoint GetSecondaryViewRectSize() const;
-
+	RENDERER_API FIntPoint GetSecondaryViewRectSize() const;
+	
 	/** Returns whether the view requires a secondary upscale. */
 	bool RequiresSecondaryUpscale() const
 	{
diff --git a/Engine/Source/Runtime/Renderer/Private/SceneTextureParameters.h b/Engine/Source/Runtime/Renderer/Private/SceneTextureParameters.h
--- a/Engine/Source/Runtime/Renderer/Private/SceneTextureParameters.h
+++ b/Engine/Source/Runtime/Renderer/Private/SceneTextureParameters.h
@@ -58,4 +58,4 @@ void SetupSceneTextureSamplers(FSceneTextureSamplerParameters* OutSamplers);
 
 /** Returns a render graph texture resource reference onto the eye adaptation or fallback.
  */
-FRDGTextureRef GetEyeAdaptationTexture(FRDGBuilder& GraphBuilder, const FViewInfo& View);
+FRDGTextureRef RENDERER_API GetEyeAdaptationTexture(FRDGBuilder& GraphBuilder, const FViewInfo& View);
diff --git a/Engine/Source/Runtime/Renderer/Private/ScreenPass.h b/Engine/Source/Runtime/Renderer/Private/ScreenPass.h
--- a/Engine/Source/Runtime/Renderer/Private/ScreenPass.h
+++ b/Engine/Source/Runtime/Renderer/Private/ScreenPass.h
@@ -176,7 +176,7 @@ BEGIN_SHADER_PARAMETER_STRUCT(FScreenPassTextureViewportParameters, )
 	SHADER_PARAMETER(FVector2D, UVViewportBilinearMax)
 END_SHADER_PARAMETER_STRUCT()
 
-FScreenPassTextureViewportParameters GetScreenPassTextureViewportParameters(const FScreenPassTextureViewport& InViewport);
+FScreenPassTextureViewportParameters RENDERER_API GetScreenPassTextureViewportParameters(const FScreenPassTextureViewport& InViewport);
 
 /** Contains a transform that maps UV coordinates from one screen pass texture viewport to another.
  *  Assumes normalized UV coordinates [0, 0]x[1, 1] where [0, 0] maps to the source view min
@@ -419,7 +419,7 @@ void AddDrawScreenPass(
  *  (CopyTexture); otherwise, it rasterizes using a pixel shader. Use this method if the two
  *  textures may have different formats.
  */
-void AddDrawTexturePass(
+void RENDERER_API AddDrawTexturePass(
 	FRDGBuilder& GraphBuilder,
 	const FViewInfo& View,
 	FRDGTextureRef InputTexture,
@@ -439,7 +439,7 @@ inline void AddDrawTexturePass(
 	AddDrawTexturePass(GraphBuilder, View, InputTexture, OutputTexture, ViewportRect.Min, ViewportRect.Min, ViewportRect.Size());
 }
 
-void AddDrawTexturePass(
+void RENDERER_API AddDrawTexturePass(
 	FRDGBuilder& GraphBuilder,
 	const FViewInfo& View,
 	FScreenPassTexture Input,
diff --git a/Engine/Source/Runtime/Renderer/Private/SingleLayerWaterRendering.cpp b/Engine/Source/Runtime/Renderer/Private/SingleLayerWaterRendering.cpp
--- a/Engine/Source/Runtime/Renderer/Private/SingleLayerWaterRendering.cpp
+++ b/Engine/Source/Runtime/Renderer/Private/SingleLayerWaterRendering.cpp
@@ -755,13 +755,14 @@ void FDeferredShadingSceneRenderer::RenderSingleLayerWaterReflections(FRHIComman
 			{
 				check(View.ViewState);
 				FTAAPassParameters TAASettings(View);
+				TAASettings.SceneDepthTexture = SceneTextures.SceneDepthBuffer;
+				TAASettings.SceneVelocityTexture = SceneTextures.SceneVelocityBuffer;
 				TAASettings.Pass = ETAAPassConfig::ScreenSpaceReflections;
 				TAASettings.SceneColorInput = DenoiserInputs.Color;
 				TAASettings.bOutputRenderTargetable = true;
 
 				FTAAOutputs TAAOutputs = AddTemporalAAPass(
 					GraphBuilder,
-					SceneTextures, 
 					View,
 					TAASettings,
 					View.PrevViewInfo.WaterSSRHistory,
diff --git a/Engine/Source/Runtime/VulkanRHI/Private/VulkanLayers.cpp b/Engine/Source/Runtime/VulkanRHI/Private/VulkanLayers.cpp
--- a/Engine/Source/Runtime/VulkanRHI/Private/VulkanLayers.cpp
+++ b/Engine/Source/Runtime/VulkanRHI/Private/VulkanLayers.cpp
@@ -7,7 +7,14 @@
 #include "VulkanRHIPrivate.h"
 #include "IHeadMountedDisplayModule.h"
 #include "IHeadMountedDisplayVulkanExtensions.h"
-
+#include "VulkanRHIBridge.h"
+namespace VulkanRHIBridge
+{
+	extern TArray<const ANSICHAR*> InstanceExtensions;
+	extern TArray<const ANSICHAR*> InstanceLayers;
+	extern TArray<const ANSICHAR*> DeviceExtensions;
+	extern TArray<const ANSICHAR*> DeviceLayers;
+}
 #if VULKAN_HAS_DEBUGGING_ENABLED
 bool GRenderDocFound = false;
 #endif
@@ -403,6 +410,10 @@ void FVulkanDynamicRHI::GetInstanceLayersAndExtensions(TArray<const ANSICHAR*>&
 		}
 	}
 
+	// plugins might have used the VulkanRHIBridge to enable additional extensions
+	OutInstanceExtensions.Append(VulkanRHIBridge::InstanceExtensions);
+	OutInstanceLayers.Append(VulkanRHIBridge::InstanceLayers);
+
 	TArray<const ANSICHAR*> PlatformExtensions;
 	FVulkanPlatform::GetInstanceExtensions(PlatformExtensions);
 
@@ -589,6 +600,10 @@ void FVulkanDevice::GetDeviceExtensionsAndLayers(VkPhysicalDevice Gpu, EGpuVendo
 		}
 	}
 
+	// plugins might have used the VulkanRHIBridge to enable additional extensions
+	OutDeviceExtensions.Append(VulkanRHIBridge::DeviceExtensions);
+	OutDeviceLayers.Append(VulkanRHIBridge::DeviceLayers);
+
 	// Now gather the actually used extensions based on the enabled layers
 	TArray<const ANSICHAR*> AvailableExtensions;
 	{
diff --git a/Engine/Source/Runtime/VulkanRHI/Private/VulkanPendingState.h b/Engine/Source/Runtime/VulkanRHI/Private/VulkanPendingState.h
--- a/Engine/Source/Runtime/VulkanRHI/Private/VulkanPendingState.h
+++ b/Engine/Source/Runtime/VulkanRHI/Private/VulkanPendingState.h
@@ -29,6 +29,13 @@ public:
 
 	~FVulkanPendingComputeState();
 
+	void Reset()
+	{
+		UAVListForAutoFlush.SetNum(0,false);
+		CurrentPipeline = nullptr;
+		CurrentState = nullptr;
+	}
+
 	void SetComputePipeline(FVulkanComputePipeline* InComputePipeline)
 	{
 		if (InComputePipeline != CurrentPipeline)
diff --git a/Engine/Source/Runtime/VulkanRHI/Private/VulkanUtil.cpp b/Engine/Source/Runtime/VulkanRHI/Private/VulkanUtil.cpp
--- a/Engine/Source/Runtime/VulkanRHI/Private/VulkanUtil.cpp
+++ b/Engine/Source/Runtime/VulkanRHI/Private/VulkanUtil.cpp
@@ -646,6 +646,12 @@ void FVulkanGPUProfiler::DumpCrashMarkers(void* BufferData)
 #include "VulkanRHIBridge.h"
 namespace VulkanRHIBridge
 {
+	TArray<const ANSICHAR*> InstanceExtensions;
+	TArray<const ANSICHAR*> InstanceLayers;
+	TArray<const ANSICHAR*> DeviceExtensions;
+	TArray<const ANSICHAR*> DeviceLayers;
+
+
 	uint64 GetInstance(FVulkanDynamicRHI* RHI)
 	{
 		return (uint64)RHI->GetInstance();
@@ -667,6 +673,20 @@ namespace VulkanRHIBridge
 	{
 		return (uint64)Device->GetPhysicalHandle();
 	}
+
+	void AddEnabledInstanceExtensionsAndLayers(const TArray<const ANSICHAR*>& InInstanceExtensions, const TArray<const ANSICHAR*>& InInstanceLayers)
+	{
+		checkf(!GVulkanRHI, TEXT("AddEnabledInstanceExtensionsAndLayers should be called before the VulkanRHI has been created"));
+		InstanceExtensions.Append(InInstanceExtensions);
+		InstanceLayers.Append(InInstanceLayers);
+	}
+
+	void AddEnabledDeviceExtensionsAndLayers(const TArray<const ANSICHAR*>& InDeviceExtensions, const TArray<const ANSICHAR*>& InDeviceLayers)
+	{
+		checkf(!GVulkanRHI, TEXT("AddEnabledDeviceExtensionsAndLayers should be called before the VulkanRHI has been created"));
+		DeviceExtensions.Append(InDeviceExtensions);
+		DeviceLayers.Append(InDeviceLayers);
+	}
 }
 
 
diff --git a/Engine/Source/Runtime/VulkanRHI/Public/VulkanConfiguration.h b/Engine/Source/Runtime/VulkanRHI/Public/VulkanConfiguration.h
--- a/Engine/Source/Runtime/VulkanRHI/Public/VulkanConfiguration.h
+++ b/Engine/Source/Runtime/VulkanRHI/Public/VulkanConfiguration.h
@@ -261,7 +261,7 @@ typedef struct VkCommandBufferInheritanceRenderPassTransformInfoQCOM {
 
 
 
-DECLARE_LOG_CATEGORY_EXTERN(LogVulkanRHI, Log, All);
+VULKANRHI_API DECLARE_LOG_CATEGORY_EXTERN(LogVulkanRHI, Log, All);
 
 #if VULKAN_CUSTOM_MEMORY_MANAGER_ENABLED
 	#define VULKAN_CPU_ALLOCATOR								VulkanRHI::GetMemoryAllocator(nullptr)
diff --git a/Engine/Source/Runtime/VulkanRHI/Public/VulkanRHIBridge.h b/Engine/Source/Runtime/VulkanRHI/Public/VulkanRHIBridge.h
--- a/Engine/Source/Runtime/VulkanRHI/Public/VulkanRHIBridge.h
+++ b/Engine/Source/Runtime/VulkanRHI/Public/VulkanRHIBridge.h
@@ -11,6 +11,11 @@ class FVulkanDevice;
 
 namespace VulkanRHIBridge
 {
+
+	VULKANRHI_API void AddEnabledInstanceExtensionsAndLayers(const TArray<const ANSICHAR*>& InInstanceExtensions, const TArray<const ANSICHAR*>& InInstanceLayers);
+	VULKANRHI_API void AddEnabledDeviceExtensionsAndLayers(const TArray<const ANSICHAR*>& InDeviceExtensions, const TArray<const ANSICHAR*>& InDeviceLayers);
+
+
 	VULKANRHI_API uint64 GetInstance(FVulkanDynamicRHI*);
 
 	VULKANRHI_API FVulkanDevice* GetDevice(FVulkanDynamicRHI*);
diff --git a/Engine/Source/Runtime/Windows/D3D11RHI/Private/D3D11RHIPrivate.h b/Engine/Source/Runtime/Windows/D3D11RHI/Private/D3D11RHIPrivate.h
--- a/Engine/Source/Runtime/Windows/D3D11RHI/Private/D3D11RHIPrivate.h
+++ b/Engine/Source/Runtime/Windows/D3D11RHI/Private/D3D11RHIPrivate.h
@@ -964,6 +964,17 @@ protected:
 	TMultiMap<ID3D11Resource*, FUnresolvedRTInfo> UnresolvedTargets;
 #endif
 
+public:
+	void RegisterGPUWork(uint32 NumPrimitives = 0, uint32 NumVertices = 0)
+	{
+		GPUProfilingData.RegisterGPUWork(NumPrimitives, NumVertices);
+	}
+	void RegisterGPUDispatch(FIntVector GroupCount)
+	{
+		GPUProfilingData.RegisterGPUDispatch(GroupCount);
+	}
+
+protected:
 	FD3DGPUProfiler GPUProfilingData;
 	// >= 0, was computed before, unless hardware was changed during engine init it should be the same
 	int32 ChosenAdapter;
