Fixes -Werror in TAA

diff --git a/Engine/Source/Runtime/Renderer/Private/PostProcess/TemporalAA.cpp b/Engine/Source/Runtime/Renderer/Private/PostProcess/TemporalAA.cpp
--- a/Engine/Source/Runtime/Renderer/Private/PostProcess/TemporalAA.cpp
+++ b/Engine/Source/Runtime/Renderer/Private/PostProcess/TemporalAA.cpp
@@ -1510,7 +1510,7 @@ public:
 		if (CVarTAAAlgorithm.GetValueOnRenderThread())
 		{
 			*OutSceneColorHalfResTexture = nullptr;
-			*OutSceneColorHalfResViewRect; // TODO.
+			//*OutSceneColorHalfResViewRect; // TODO.
 
 			return AddGen5MainTemporalAAPasses(
 				GraphBuilder,
