Stabilises high frequency details on smooth gradient by adding stocastic quantization in decimate and filtering passes

diff --git a/Engine/Shaders/Private/TemporalAA/TAACommon.ush b/Engine/Shaders/Private/TemporalAA/TAACommon.ush
--- a/Engine/Shaders/Private/TemporalAA/TAACommon.ush
+++ b/Engine/Shaders/Private/TemporalAA/TAACommon.ush
@@ -4,6 +4,7 @@
 #include "../Random.ush"
 #include "../TextureSampling.ush"
 #include "../FastMath.ush"
+#include "../MonteCarlo.ush"
 
 
 //------------------------------------------------------- COMPILER CONFIG
@@ -22,6 +23,8 @@
 	#define CONFIG_COMPILE_FP16 0
 #endif
 
+#define CONFIG_ENABLE_STOCASTIC_QUANTIZATION 1
+
 
 //------------------------------------------------------- CONSTANTS
 
diff --git a/Engine/Shaders/Private/TemporalAA/TAADecimateHistory.usf b/Engine/Shaders/Private/TemporalAA/TAADecimateHistory.usf
--- a/Engine/Shaders/Private/TemporalAA/TAADecimateHistory.usf
+++ b/Engine/Shaders/Private/TemporalAA/TAADecimateHistory.usf
@@ -12,12 +12,13 @@
 
 //------------------------------------------------------- PARAMETERS
 
-Texture2D<float2> DilatedVelocityTexture;
-Texture2D<taa_half> ParallaxRejectionMaskTexture;
-
+float3 OutputQuantizationError;
 float HistoryPreExposureCorrection;
 uint bCameraCut;
 
+Texture2D<float2> DilatedVelocityTexture;
+Texture2D<taa_half> ParallaxRejectionMaskTexture;
+
 Texture2D<taa_half3> PrevHistory_Textures_0;
 Texture2D<taa_half3> PrevHistory_Textures_1;
 Texture2D<taa_half2> PrevHistory_Textures_2;
@@ -235,6 +236,16 @@ void MainCS(
 		#error Unknown kernel
 	#endif
 
+
+	#if CONFIG_ENABLE_STOCASTIC_QUANTIZATION
+	{
+		uint2 Random = Rand3DPCG16(int3(InputPixelPos, View.StateFrameIndexMod8)).xy;
+		float2 E = Hammersley16(0, 1, Random);
+
+		DecimatedColor += DecimatedColor * (E.x * OutputQuantizationError);
+	}
+	#endif
+
 	// Immediately reject based on the rejection mask
 	FLATTEN
 	if ((PrevOutputScreenPos.x < -1.0) ||
diff --git a/Engine/Shaders/Private/TemporalAA/TAAFilterFrequencies.usf b/Engine/Shaders/Private/TemporalAA/TAAFilterFrequencies.usf
--- a/Engine/Shaders/Private/TemporalAA/TAAFilterFrequencies.usf
+++ b/Engine/Shaders/Private/TemporalAA/TAAFilterFrequencies.usf
@@ -88,6 +88,8 @@ static const float2 kStackowiakSampleSet0[56 * 4] =
 
 //------------------------------------------------------- PARAMETERS
 
+float3 OutputQuantizationError;
+
 Texture2D<taa_half3> InputTexture;
 Texture2D<taa_half3> PredictionSceneColorTexture;
 Texture2D<taa_half> PredictionInfoTexture;
@@ -265,6 +267,17 @@ void MainCS(
 		FilteredInput *= SafeRcp(FilteredInputWeight);
 		FilteredHistory *= SafeRcp(FilteredHistoryWeight);
 	}
+	
+	#if CONFIG_ENABLE_STOCASTIC_QUANTIZATION
+	{
+		uint2 Random = Rand3DPCG16(int3(InputPixelPos, View.StateFrameIndexMod8)).xy;
+		float2 E = Hammersley16(0, 1, Random);
+
+		FilteredInput += FilteredInput * (E.x * OutputQuantizationError);
+		FilteredHistory += FilteredHistory * (E.x * OutputQuantizationError);
+	}
+	#endif
+
 
 	if (all(InputPixelPos < InputInfo_ViewportSize))
 	{
diff --git a/Engine/Shaders/Private/TemporalAA/TAAUpdateHistory.usf b/Engine/Shaders/Private/TemporalAA/TAAUpdateHistory.usf
--- a/Engine/Shaders/Private/TemporalAA/TAAUpdateHistory.usf
+++ b/Engine/Shaders/Private/TemporalAA/TAAUpdateHistory.usf
@@ -2,7 +2,6 @@
 
 #include "TAACommon.ush"
 #include "../Random.ush"
-#include "../MonteCarlo.ush"
 
 
 //------------------------------------------------------- CONFIG
@@ -628,7 +627,7 @@ void MainCS(
 		{
 			// Stocastically round up or down using the hardware RWTexture2D truncation unit to take into precision
 			// loss due to pixel format encoding.
-			#if 1
+			#if CONFIG_ENABLE_STOCASTIC_QUANTIZATION
 			{
 				uint2 Random = Rand3DPCG16(int3(LocalHistoryPixelPos, View.StateFrameIndexMod8)).xy;
 				float2 E = Hammersley16(0, 1, Random);
diff --git a/Engine/Source/Runtime/Renderer/Private/PostProcess/TemporalAA.cpp b/Engine/Source/Runtime/Renderer/Private/PostProcess/TemporalAA.cpp
--- a/Engine/Source/Runtime/Renderer/Private/PostProcess/TemporalAA.cpp
+++ b/Engine/Source/Runtime/Renderer/Private/PostProcess/TemporalAA.cpp
@@ -352,12 +352,13 @@ class FTAA2DecimateHistoryCS : public FTAAGen5Shader
 
 	BEGIN_SHADER_PARAMETER_STRUCT(FParameters, )
 		SHADER_PARAMETER_STRUCT_INCLUDE(FTAA2CommonParameters, CommonParameters)
-		SHADER_PARAMETER_RDG_TEXTURE(Texture2D, DilatedVelocityTexture)
-		SHADER_PARAMETER_RDG_TEXTURE(Texture2D, ParallaxRejectionMaskTexture)
-
+		SHADER_PARAMETER(FVector, OutputQuantizationError)
 		SHADER_PARAMETER(float, HistoryPreExposureCorrection)
 		SHADER_PARAMETER(int32, bCameraCut)
 
+		SHADER_PARAMETER_RDG_TEXTURE(Texture2D, DilatedVelocityTexture)
+		SHADER_PARAMETER_RDG_TEXTURE(Texture2D, ParallaxRejectionMaskTexture)
+
 		SHADER_PARAMETER_STRUCT(FScreenPassTextureViewportParameters, PrevHistoryInfo)
 		SHADER_PARAMETER_STRUCT(FTAA2HistoryTextures, PrevHistory)
 
@@ -374,9 +375,12 @@ class FTAA2FilterFrequenciesCS : public FTAAGen5Shader
 
 	BEGIN_SHADER_PARAMETER_STRUCT(FParameters, )
 		SHADER_PARAMETER_STRUCT_INCLUDE(FTAA2CommonParameters, CommonParameters)
+		SHADER_PARAMETER(FVector, OutputQuantizationError)
+
 		SHADER_PARAMETER_RDG_TEXTURE(Texture2D, InputTexture)
 		SHADER_PARAMETER_RDG_TEXTURE(Texture2D, PredictionSceneColorTexture)
 		SHADER_PARAMETER_RDG_TEXTURE(Texture2D, PredictionInfoTexture)
+
 		SHADER_PARAMETER_RDG_TEXTURE_UAV(RWTexture2D, FilteredInputOutput)
 		SHADER_PARAMETER_RDG_TEXTURE_UAV(RWTexture2D, FilteredPredictionSceneColorOutput)
 		SHADER_PARAMETER_RDG_TEXTURE_UAV(RWTexture2D, DebugOutput)
@@ -456,6 +460,25 @@ float CatmullRom(float x)
 		return (1.5f * ax - 2.5f) * ax*ax + 1.0f;
 }
 
+FVector ComputePixelFormatQuantizationError(EPixelFormat PixelFormat)
+{
+	FVector Error;
+	if (PixelFormat == PF_FloatRGBA || PixelFormat == PF_FloatR11G11B10)
+	{
+		FIntVector HistoryColorMantissaBits = PixelFormat == PF_FloatR11G11B10 ? FIntVector(6, 6, 5) : FIntVector(10, 10, 10);
+
+		Error.X = FMath::Pow(0.5f, HistoryColorMantissaBits.X);
+		Error.Y = FMath::Pow(0.5f, HistoryColorMantissaBits.Y);
+		Error.Z = FMath::Pow(0.5f, HistoryColorMantissaBits.Z);
+	}
+	else
+	{
+		check(0);
+	}
+
+	return Error;
+}
+
 void SetupSampleWeightParameters(FTAAStandaloneCS::FParameters* OutTAAParameters, const FTAAPassParameters& PassParameters, FVector2D TemporalJitterPixels)
 {
 	float JitterX = TemporalJitterPixels.X;
@@ -1194,11 +1217,13 @@ static void AddGen5MainTemporalAAPasses(
 
 		FTAA2DecimateHistoryCS::FParameters* PassParameters = GraphBuilder.AllocParameters<FTAA2DecimateHistoryCS::FParameters>();
 		PassParameters->CommonParameters = CommonParameters;
-		PassParameters->DilatedVelocityTexture = DilatedVelocityTexture;
-		PassParameters->ParallaxRejectionMaskTexture = ParallaxRejectionMaskTexture;
+		PassParameters->OutputQuantizationError = ComputePixelFormatQuantizationError(PredictionSceneColorTexture->Desc.Format);
 		PassParameters->HistoryPreExposureCorrection = View.PreExposure / View.PrevViewInfo.SceneColorPreExposure;
 		PassParameters->bCameraCut = bCameraCut;
 
+		PassParameters->DilatedVelocityTexture = DilatedVelocityTexture;
+		PassParameters->ParallaxRejectionMaskTexture = ParallaxRejectionMaskTexture;\
+
 		PassParameters->PrevHistoryInfo = PrevHistoryInfo;
 		PassParameters->PrevHistory = PrevHistory;
 
@@ -1237,6 +1262,8 @@ static void AddGen5MainTemporalAAPasses(
 
 			FTAA2FilterFrequenciesCS::FParameters* PassParameters = GraphBuilder.AllocParameters<FTAA2FilterFrequenciesCS::FParameters>();
 			PassParameters->CommonParameters = CommonParameters;
+			PassParameters->OutputQuantizationError = ComputePixelFormatQuantizationError(FilteredInputTexture->Desc.Format);
+
 			PassParameters->InputTexture = PassInputs.SceneColorTexture;
 			PassParameters->PredictionSceneColorTexture = PredictionSceneColorTexture;
 			PassParameters->PredictionInfoTexture = PredictionInfoTexture;
@@ -1332,12 +1359,7 @@ static void AddGen5MainTemporalAAPasses(
 		PassParameters->DilatedVelocityTexture = DilatedVelocityTexture;
 		PassParameters->ParallaxRejectionMaskTexture = ParallaxRejectionMaskTexture;
 
-		{
-			FIntVector HistoryColorMantissaBits = bR11G11B10History ? FIntVector(6, 6, 5) : FIntVector(10, 10, 10);
-			PassParameters->HistoryQuantizationError.X = FMath::Pow(0.5f, HistoryColorMantissaBits.X);
-			PassParameters->HistoryQuantizationError.Y = FMath::Pow(0.5f, HistoryColorMantissaBits.Y);
-			PassParameters->HistoryQuantizationError.Z = FMath::Pow(0.5f, HistoryColorMantissaBits.Z);
-		}
+		PassParameters->HistoryQuantizationError = ComputePixelFormatQuantizationError(History.Textures[0]->Desc.Format);
 		PassParameters->HistoryPreExposureCorrection = View.PreExposure / View.PrevViewInfo.SceneColorPreExposure;
 		PassParameters->bCameraCut = bCameraCut;
 
