Fixes redondant declaration of CVarTAAAlgorithm

diff --git a/Engine/Source/Runtime/Renderer/Private/PostProcess/PostProcessing.cpp b/Engine/Source/Runtime/Renderer/Private/PostProcess/PostProcessing.cpp
--- a/Engine/Source/Runtime/Renderer/Private/PostProcess/PostProcessing.cpp
+++ b/Engine/Source/Runtime/Renderer/Private/PostProcess/PostProcessing.cpp
@@ -103,13 +103,6 @@ TAutoConsoleVariable<int32> CVarPostProcessingForceAsyncDispatch(
 	ECVF_RenderThreadSafe);
 #endif
 
-TAutoConsoleVariable<int32> CVarTAAAlgorithm(
-	TEXT("r.TemporalAA.Algorithm"), 0,
-	TEXT("Select the TAA algorithm.\n")
-	TEXT(" 0: TAAU 1.0 (default)\n")
-	TEXT(" 1: TAAU 2.0 (experimental)"),
-	ECVF_RenderThreadSafe);
-
 } //! namespace
 
 bool IsPostProcessingWithComputeEnabled(ERHIFeatureLevel::Type FeatureLevel)
