Fallbacks to Gen4 TAAU when Gen5 is not supported

diff --git a/Engine/Source/Runtime/Renderer/Private/PostProcess/TemporalAA.cpp b/Engine/Source/Runtime/Renderer/Private/PostProcess/TemporalAA.cpp
--- a/Engine/Source/Runtime/Renderer/Private/PostProcess/TemporalAA.cpp
+++ b/Engine/Source/Runtime/Renderer/Private/PostProcess/TemporalAA.cpp
@@ -1528,7 +1528,7 @@ public:
 		FRDGTextureRef* OutSceneColorHalfResTexture,
 		FIntRect* OutSceneColorHalfResViewRect) const final
 	{
-		if (CVarTAAAlgorithm.GetValueOnRenderThread())
+		if (CVarTAAAlgorithm.GetValueOnRenderThread() && DoesPlatformSupportGen5TAA(View.GetShaderPlatform()))
 		{
 			*OutSceneColorHalfResTexture = nullptr;
 			//*OutSceneColorHalfResViewRect; // TODO.
