Fix UE5 build on Mac: Initialize main thread with ETaskTag::ENone, remove reference to CVar "r.Metal.ForceDXC" to enable tessellation, and disable FDiffuseIndirectCompositePS shader on Mac.

#rnx

diff --git a/Engine/Source/Runtime/Apple/MetalRHI/Private/MetalCommandQueue.cpp b/Engine/Source/Runtime/Apple/MetalRHI/Private/MetalCommandQueue.cpp
--- a/Engine/Source/Runtime/Apple/MetalRHI/Private/MetalCommandQueue.cpp
+++ b/Engine/Source/Runtime/Apple/MetalRHI/Private/MetalCommandQueue.cpp
@@ -284,13 +284,7 @@ FMetalCommandQueue::FMetalCommandQueue(mtlpp::Device InDevice, uint32 const MaxN
     
     if(Device.SupportsFeatureSet(mtlpp::FeatureSet::macOS_GPUFamily1_v3) && FPlatformMisc::MacOSXVersionCompare(10,13,0) >= 0)
     {
-        Features |= EMetalFeaturesMultipleViewports | EMetalFeaturesPipelineBufferMutability | EMetalFeaturesGPUCaptureManager;
-		
-		static const auto CVar = IConsoleManager::Get().FindConsoleVariable(TEXT("r.Metal.ForceDXC"));
-		if (CVar && CVar->GetInt() != 0)
-		{
-			Features |= EMetalFeaturesSeparateTessellation;
-		}
+        Features |= EMetalFeaturesMultipleViewports | EMetalFeaturesPipelineBufferMutability | EMetalFeaturesGPUCaptureManager | EMetalFeaturesSeparateTessellation;
 		
 		if (FParse::Param(FCommandLine::Get(),TEXT("metalfence")))
 		{
diff --git a/Engine/Source/Runtime/Core/Public/CoreGlobals.h b/Engine/Source/Runtime/Core/Public/CoreGlobals.h
--- a/Engine/Source/Runtime/Core/Public/CoreGlobals.h
+++ b/Engine/Source/Runtime/Core/Public/CoreGlobals.h
@@ -498,6 +498,11 @@ class FTaskTagScope
 	ETaskTag Tag;
 	bool TagOnlyIfNone;
 
+public:
+	/**
+	 * Clear the ETaskTag::StaticInit tag so that the Main OS Thread has no active tag.
+	 * This will allow functions such as IsInGameThread() to function properly when called from the Main Thread.
+	 */
 	static void CORE_API SetTagNone();
 
 protected:
diff --git a/Engine/Source/Runtime/Launch/Private/Mac/LaunchMac.cpp b/Engine/Source/Runtime/Launch/Private/Mac/LaunchMac.cpp
--- a/Engine/Source/Runtime/Launch/Private/Mac/LaunchMac.cpp
+++ b/Engine/Source/Runtime/Launch/Private/Mac/LaunchMac.cpp
@@ -339,6 +339,8 @@ extern bool GIsConsoleExecutable;
 
 INT32_MAIN_INT32_ARGC_TCHAR_ARGV()
 {
+	FTaskTagScope::SetTagNone();
+	
 	for (int32 Option = 1; Option < ArgC; Option++)
 	{
 		GSavedCommandLine += TEXT(" ");
diff --git a/Engine/Source/Runtime/Renderer/Private/IndirectLightRendering.cpp b/Engine/Source/Runtime/Renderer/Private/IndirectLightRendering.cpp
--- a/Engine/Source/Runtime/Renderer/Private/IndirectLightRendering.cpp
+++ b/Engine/Source/Runtime/Renderer/Private/IndirectLightRendering.cpp
@@ -1044,6 +1044,7 @@ void FDeferredShadingSceneRenderer::RenderDiffuseIndirectAndAmbientOcclusion(
 		}
 
 		// Applies diffuse indirect and ambient occlusion to the scene color.
+#if !PLATFORM_MAC
 		if ((DenoiserOutputs.Textures[0] || AmbientOcclusionMask) && (!bIsVisualizePass || ViewPipelineState.DiffuseIndirectDenoiser != IScreenSpaceDenoiser::EMode::Disabled || ViewPipelineState.bUseLumenProbeHierarchy))
 		{
 			FDiffuseIndirectCompositePS::FParameters* PassParameters = GraphBuilder.AllocParameters<FDiffuseIndirectCompositePS::FParameters>();
@@ -1148,6 +1149,7 @@ void FDeferredShadingSceneRenderer::RenderDiffuseIndirectAndAmbientOcclusion(
 				View.ViewRect,
 				BlendState);
 		} // if (DenoiserOutputs.Color || bApplySSAO)
+#endif // PLATFORM_MAC
 
 		if (ViewPipelineState.DiffuseIndirectMethod == EDiffuseIndirectMethod::Lumen && !ViewPipelineState.bUseLumenProbeHierarchy)
 		{
