Replace "#if !PLATFORM_MAC" directive.

#rnx

diff --git a/Engine/Source/Runtime/Renderer/Private/IndirectLightRendering.cpp b/Engine/Source/Runtime/Renderer/Private/IndirectLightRendering.cpp
--- a/Engine/Source/Runtime/Renderer/Private/IndirectLightRendering.cpp
+++ b/Engine/Source/Runtime/Renderer/Private/IndirectLightRendering.cpp
@@ -100,7 +100,7 @@ class FDiffuseIndirectCompositePS : public FGlobalShader
 
 	static bool ShouldCompilePermutation(const FGlobalShaderPermutationParameters& Parameters)
 	{
-		if (Parameters.Platform == EShaderPlatform::SP_METAL_SM5)
+		if (IsMetalPlatform(Parameters.Platform))
 		{
 			return false;
 		}
@@ -1044,8 +1044,8 @@ void FDeferredShadingSceneRenderer::RenderDiffuseIndirectAndAmbientOcclusion(
 		}
 
 		// Applies diffuse indirect and ambient occlusion to the scene color.
-#if !PLATFORM_MAC
-		if ((DenoiserOutputs.Textures[0] || AmbientOcclusionMask) && (!bIsVisualizePass || ViewPipelineState.DiffuseIndirectDenoiser != IScreenSpaceDenoiser::EMode::Disabled || ViewPipelineState.bUseLumenProbeHierarchy))
+		if ((DenoiserOutputs.Textures[0] || AmbientOcclusionMask) && (!bIsVisualizePass || ViewPipelineState.DiffuseIndirectDenoiser != IScreenSpaceDenoiser::EMode::Disabled || ViewPipelineState.bUseLumenProbeHierarchy)
+			&& !IsMetalPlatform(ShaderPlatform))
 		{
 			FDiffuseIndirectCompositePS::FParameters* PassParameters = GraphBuilder.AllocParameters<FDiffuseIndirectCompositePS::FParameters>();
 			
@@ -1149,7 +1149,6 @@ void FDeferredShadingSceneRenderer::RenderDiffuseIndirectAndAmbientOcclusion(
 				View.ViewRect,
 				BlendState);
 		} // if (DenoiserOutputs.Color || bApplySSAO)
-#endif // PLATFORM_MAC
 
 		if (ViewPipelineState.DiffuseIndirectMethod == EDiffuseIndirectMethod::Lumen && !ViewPipelineState.bUseLumenProbeHierarchy)
 		{
