Fixes caching of whole-scene shadow maps when resolution changes (but is within threshold for re-use)

diff --git a/Engine/Source/Runtime/Renderer/Private/ShadowSetup.cpp b/Engine/Source/Runtime/Renderer/Private/ShadowSetup.cpp
--- a/Engine/Source/Runtime/Renderer/Private/ShadowSetup.cpp
+++ b/Engine/Source/Runtime/Renderer/Private/ShadowSetup.cpp
@@ -3148,9 +3148,28 @@ void FSceneRenderer::CreateWholeSceneProjectedShadow(
 				int32 NumShadowMaps = 1;
 				EShadowDepthCacheMode CacheMode[2] = { SDCM_Uncached, SDCM_Uncached };
 
-				FIntPoint ShadowMapSize( SizeX + ShadowBorder * 2, SizeY + ShadowBorder * 2 );
-				SizeX = ShadowMapSize.X - ShadowBorder * 2;
-				SizeY = ShadowMapSize.Y - ShadowBorder * 2;
+				if (!bAnyViewIsSceneCapture && !ProjectedShadowInitializer.bRayTracedDistanceField)
+				{
+					FIntPoint ShadowMapSize( SizeX + ShadowBorder * 2, SizeY + ShadowBorder * 2 );
+
+					ComputeWholeSceneShadowCacheModes(
+						LightSceneInfo,
+						ProjectedShadowInitializer.bOnePassPointLightShadow,
+						ViewFamily.CurrentRealTime,
+						MaxDesiredResolution,
+						FIntPoint(MaxShadowResolution, MaxShadowResolutionY),
+						Scene,
+						// Below are in-out or out parameters. They can change
+						ProjectedShadowInitializer,
+						ShadowMapSize,
+						InOutNumPointShadowCachesUpdatedThisFrame,
+						InOutNumSpotShadowCachesUpdatedThisFrame,
+						NumShadowMaps,
+						CacheMode);
+
+					SizeX = ShadowMapSize.X - ShadowBorder * 2;
+					SizeY = ShadowMapSize.Y - ShadowBorder * 2;
+				}
 
 				if (bNeedsVirtualShadowMap)
 				{
@@ -3164,7 +3183,7 @@ void FSceneRenderer::CreateWholeSceneProjectedShadow(
 					// Rescale size to fit whole virtual SM but keeping aspect ratio
 					int32 VirtualSizeX = SizeX >= SizeY ? FVirtualShadowMap::VirtualMaxResolutionXY : (FVirtualShadowMap::VirtualMaxResolutionXY * SizeX) / SizeY;
 					int32 VirtualSizeY = SizeY >= SizeX ? FVirtualShadowMap::VirtualMaxResolutionXY : (FVirtualShadowMap::VirtualMaxResolutionXY * SizeY) / SizeX;
-					
+
 					TSharedPtr<FVirtualShadowMapCacheEntry> VirtualSmCacheEntry;
 					if (Scene->VirtualShadowMapArrayCacheManager)
 					{
@@ -3189,24 +3208,6 @@ void FSceneRenderer::CreateWholeSceneProjectedShadow(
 
 					VisibleLightInfo.AllProjectedShadows.Add(ProjectedShadowInfo);
 				}
-
-				if (!bAnyViewIsSceneCapture && !ProjectedShadowInitializer.bRayTracedDistanceField)
-				{
-					ComputeWholeSceneShadowCacheModes(
-						LightSceneInfo,
-						ProjectedShadowInitializer.bOnePassPointLightShadow,
-						ViewFamily.CurrentRealTime,
-						MaxDesiredResolution,
-						FIntPoint(MaxShadowResolution, MaxShadowResolutionY),
-						Scene,
-						// Below are in-out or out parameters. They can change
-						ProjectedShadowInitializer,
-						ShadowMapSize,
-						InOutNumPointShadowCachesUpdatedThisFrame,
-						InOutNumSpotShadowCachesUpdatedThisFrame,
-						NumShadowMaps,
-						CacheMode);
-				}
 				
 				for (int32 CacheModeIndex = 0; CacheModeIndex < NumShadowMaps; CacheModeIndex++)
 				{
