Fix crash in Reverb when starting the flight sequence (Velocity rendering lacking a global uniform buffer)

diff --git a/Engine/Source/Runtime/Renderer/Private/VelocityRendering.cpp b/Engine/Source/Runtime/Renderer/Private/VelocityRendering.cpp
--- a/Engine/Source/Runtime/Renderer/Private/VelocityRendering.cpp
+++ b/Engine/Source/Runtime/Renderer/Private/VelocityRendering.cpp
@@ -358,12 +358,17 @@ void FDeferredShadingSceneRenderer::RenderVelocitiesInner(FRHICommandListImmedia
 		SetupSceneTextureUniformParameters(SceneContext, View.FeatureLevel, VelocityPass == EVelocityPass::Opaque ? ESceneTextureSetupMode::None : ESceneTextureSetupMode::All, SceneTextureParameters);
 		Scene->UniformBuffers.VelocityPassUniformBuffer.UpdateUniformBufferImmediate(SceneTextureParameters);
 
+		FUniformBufferRHIRef PassUniformBuffer = CreateSceneTextureUniformBufferDependentOnShadingPath(SceneContext, View.FeatureLevel, VelocityPass == EVelocityPass::Opaque ? ESceneTextureSetupMode::None : ESceneTextureSetupMode::All, UniformBuffer_SingleFrame);
+
 		FMeshPassProcessorRenderState DrawRenderState(View, Scene->UniformBuffers.VelocityPassUniformBuffer);
 
 		if (View.ShouldRenderView())
 		{
 			SCOPED_GPU_MASK(RHICmdList, View.GPUMask);
 
+			FUniformBufferStaticBindings GlobalUniformBuffers(PassUniformBuffer);
+			SCOPED_UNIFORM_BUFFER_GLOBAL_BINDINGS(RHICmdList, GlobalUniformBuffers);
+
 			Scene->UniformBuffers.UpdateViewUniformBuffer(View);
 
 			SetVelocitiesState(RHICmdList, View, this, DrawRenderState, VelocityRT, VelocityPass);
