Use a single Nanite (multi-view) pass for rendering all shadows to a given atlas.

diff --git a/Engine/Source/Runtime/Renderer/Private/ShadowDepthRendering.cpp b/Engine/Source/Runtime/Renderer/Private/ShadowDepthRendering.cpp
--- a/Engine/Source/Runtime/Renderer/Private/ShadowDepthRendering.cpp
+++ b/Engine/Source/Runtime/Renderer/Private/ShadowDepthRendering.cpp
@@ -1705,88 +1705,99 @@ void FSceneRenderer::RenderShadowDepthMapAtlases(FRHICommandListImmediate& RHICm
 
 		if (bNaniteEnabled && CVarNaniteShadows.GetValueOnRenderThread())
 		{
-			FRDGBuilder GraphBuilder( RHICmdList );
-
-			FRDGTextureRef ShadowMap = GraphBuilder.RegisterExternalTexture( ShadowMapAtlas.RenderTargets.DepthTarget, TEXT("DepthBuffer") );
-
+			bool bWantsNearClip = false;
+			bool bWantsNoNearClip = false;
+			TArray<Nanite::FPackedView, SceneRenderingAllocator> PackedViews;
+			TArray<FProjectedShadowInfo*, SceneRenderingAllocator> ShadowsToEmit;
+			for (int32 ShadowIndex = 0; ShadowIndex < ShadowMapAtlas.Shadows.Num(); ShadowIndex++)
 			{
-				RDG_EVENT_SCOPE( GraphBuilder, "Nanite Shadows" );
+				FProjectedShadowInfo* ProjectedShadowInfo = ShadowMapAtlas.Shadows[ShadowIndex];
 
-				for (int32 ShadowIndex = 0; ShadowIndex < ShadowMapAtlas.Shadows.Num(); ShadowIndex++)
+				// TODO: We avoid rendering Nanite geometry into both movable AND static cached shadows, but has a side effect
+				// that if there is *only* a movable cached shadow map (and not static), it won't rendering anything.
+				// Logic around Nanite and the cached shadows is fuzzy in a bunch of places and the whole thing needs some rethinking
+				// so leaving this like this for now as it is unlikely to happen in realistic scenes.
+				if (!ProjectedShadowInfo->bNaniteGeometry ||
+					ProjectedShadowInfo->CacheMode == SDCM_MovablePrimitivesOnly)
 				{
-					FProjectedShadowInfo* ProjectedShadowInfo = ShadowMapAtlas.Shadows[ShadowIndex];
-
-					// FIXME: We avoid rendering Nanite geometry into both movable AND static cached shadows, but has a side effect
-					// that if there is *only* a movable cached shadow map (and not static), it won't rendering anything.
-					// Logic around Nanite and the cached shadows is fuzzy in a bunch of places and the whole thing needs some rethinking
-					// so leaving this like this for now as it is unlikely to happen in realistic scenes.
-					if (!ProjectedShadowInfo->bNaniteGeometry ||
-						ProjectedShadowInfo->CacheMode == SDCM_MovablePrimitivesOnly)
-					{
-						continue;
-					}
-
-					FString LightName;
-					GetLightNameForDrawEvent( ProjectedShadowInfo->GetLightSceneInfo().Proxy, LightName );
-
-					RDG_EVENT_SCOPE( GraphBuilder, "%s %ux%u", *LightName, ProjectedShadowInfo->ResolutionX, ProjectedShadowInfo->ResolutionY );
-
-					const FIntRect AtlasViewRect = ProjectedShadowInfo->GetViewRectForView();
+					continue;
+				}
 
-					{
-						// We render the Nanite geometry into a viewport-sized raster context, then emit into the right atlas location
-						FIntRect NaniteViewRect(FIntPoint(0, 0), AtlasViewRect.Size());
-						FIntPoint RasterContextSize = NaniteViewRect.Size();
+				// Orthographic shadow projections want depth clamping rather than clipping
+				if (ProjectedShadowInfo->ShouldClampToNearPlane())
+				{
+					bWantsNoNearClip = true;
+				}
+				else
+				{
+					bWantsNearClip = true;
+				}
 
-						// Orthographic shadow projections want depth clamping rather than clipping
-						Nanite::FRasterState RasterState;
-						RasterState.bNearClip = !ProjectedShadowInfo->ShouldClampToNearPlane();
+				// Only add the view if it has instances
+				const FIntRect AtlasViewRect = ProjectedShadowInfo->GetViewRectForView();
 
-						const bool bUpdateStreaming = CVarNaniteShadowsUpdateStreaming.GetValueOnRenderThread() != 0;
-						Nanite::FCullingContext CullingContext = Nanite::InitCullingContext( GraphBuilder, *Scene, nullptr, NaniteViewRect, true, bUpdateStreaming, false, false );
+				Nanite::FPackedViewParams Initializer;
+				Initializer.ViewMatrices = ProjectedShadowInfo->ShadowDepthView->ViewMatrices;
+				Initializer.PrevViewMatrices = Initializer.ViewMatrices;
+				Initializer.ViewRect = AtlasViewRect;
+				Initializer.RasterContextSize = AtlasSize;
+				Initializer.LODScaleFactor = ComputeNaniteShadowsLODScaleFactor();
+				PackedViews.Add(Nanite::CreatePackedView(Initializer));
 
-						Nanite::FRasterContext RasterContext = Nanite::InitRasterContext(GraphBuilder, RasterContextSize, Nanite::EOutputBufferMode::DepthOnly);
+				ShadowsToEmit.Add(ProjectedShadowInfo);
+			}
+				
+			if (PackedViews.Num() > 0)
+			{
+				FRDGBuilder GraphBuilder(RHICmdList);
+				FRDGTextureRef ShadowMap = GraphBuilder.RegisterExternalTexture(ShadowMapAtlas.RenderTargets.DepthTarget, TEXT("DepthBuffer"));
 
-						// Setup packed view
-						TArray<Nanite::FPackedView, SceneRenderingAllocator> PackedViews;
-						{
-							Nanite::FPackedViewParams Initializer;
-							Initializer.ViewMatrices = ProjectedShadowInfo->ShadowDepthView->ViewMatrices;
-							Initializer.PrevViewMatrices = Initializer.ViewMatrices;
-							Initializer.ViewRect = NaniteViewRect;
-							Initializer.RasterContextSize = RasterContextSize;
-							Initializer.LODScaleFactor = ComputeNaniteShadowsLODScaleFactor();
-							PackedViews.Add(Nanite::CreatePackedView(Initializer));
-						}
+				{
+					RDG_EVENT_SCOPE(GraphBuilder, "Nanite Shadows");
+
+					// NOTE: Rendering into an atlas like this is not going to work properly with HZB, but we are not currently using HZB here.
+					// It might be worthwhile going through the virtual SM rendering path even for "dense" cases even just for proper handling of all the details.
+					FIntRect FullAtlasViewRect(FIntPoint(0, 0), AtlasSize);
+					const bool bUpdateStreaming = CVarNaniteShadowsUpdateStreaming.GetValueOnRenderThread() != 0;
+					Nanite::FCullingContext CullingContext = Nanite::InitCullingContext(GraphBuilder, *Scene, nullptr, FullAtlasViewRect, true, bUpdateStreaming, false, false);
+					Nanite::FRasterContext RasterContext = Nanite::InitRasterContext(GraphBuilder, AtlasSize, Nanite::EOutputBufferMode::DepthOnly);
+
+					// We can't currently do both at once, but this shouldn't currently happen since directional lights go into different atlases
+					// If we ever need this we can do two passes like virtual shadow maps
+					check(!bWantsNearClip || !bWantsNoNearClip);
+					Nanite::FRasterState RasterState;
+					RasterState.bNearClip = bWantsNearClip;
 
-						const bool bExtractStats = Nanite::IsStatFilterActiveForLight(ProjectedShadowInfo->GetLightSceneInfo().Proxy);
+					Nanite::CullRasterize(
+						GraphBuilder,
+						*Scene,
+						PackedViews,
+						CullingContext,
+						RasterContext,
+						RasterState,
+						nullptr,	// InstanceDraws
+						false		// bExtractStats TODO
+					);
 
-						Nanite::CullRasterize(
-							GraphBuilder,
-							*Scene,
-							PackedViews,
-							CullingContext,
-							RasterContext,
-							RasterState,
-							nullptr,
-							bExtractStats
-						);
+					for (FProjectedShadowInfo* ProjectedShadowInfo : ShadowsToEmit)
+					{
+						const FIntRect AtlasViewRect = ProjectedShadowInfo->GetViewRectForView();
 
 						Nanite::EmitShadowMap(
 							GraphBuilder,
 							RasterContext,
 							ShadowMap,
-							NaniteViewRect,
+							AtlasViewRect,
 							AtlasViewRect.Min,
 							ProjectedShadowInfo->ShadowDepthView->ViewMatrices.GetProjectionMatrix(),
 							ProjectedShadowInfo->GetShaderDepthBias(),
 							ProjectedShadowInfo->bDirectionalLight
-						);
+							);
 					}
 				}
-			}
 
-			GraphBuilder.Execute();
+				GraphBuilder.Execute();
+			}
 		}
 
 		RHICmdList.TransitionResource(EResourceTransitionAccess::EReadable, RenderTarget.TargetableTexture);
