Enable global distance field shader compilation on Mac (disabled by mistake in Reverb stream) to fix missing shader crash there

diff --git a/Engine/Source/Runtime/Renderer/Private/GlobalDistanceField.cpp b/Engine/Source/Runtime/Renderer/Private/GlobalDistanceField.cpp
--- a/Engine/Source/Runtime/Renderer/Private/GlobalDistanceField.cpp
+++ b/Engine/Source/Runtime/Renderer/Private/GlobalDistanceField.cpp
@@ -184,7 +184,7 @@ constexpr float MESH_SDF_APPROX_MARGIN = 0.7f;
 
 bool ShouldCompileGlobalDistanceFieldShader(const FGlobalShaderPermutationParameters& Parameters)
 {
-	return IsFeatureLevelSupported(Parameters.Platform, ERHIFeatureLevel::SM5) && DoesPlatformSupportDistanceFieldAO(Parameters.Platform) && !IsMetalPlatform(Parameters.Platform) && IsUsingDistanceFields(Parameters.Platform);
+	return IsFeatureLevelSupported(Parameters.Platform, ERHIFeatureLevel::SM5) && DoesPlatformSupportDistanceFieldAO(Parameters.Platform) && IsUsingDistanceFields(Parameters.Platform);
 }
 
 float GetMinMeshSDFRadius(float VoxelWorldSize)
