Handle situations where direction and spot lights are mixed in the same shadow atlas (preshadows) by doing a second Nanite pass.

diff --git a/Engine/Source/Runtime/Renderer/Private/ShadowDepthRendering.cpp b/Engine/Source/Runtime/Renderer/Private/ShadowDepthRendering.cpp
--- a/Engine/Source/Runtime/Renderer/Private/ShadowDepthRendering.cpp
+++ b/Engine/Source/Runtime/Renderer/Private/ShadowDepthRendering.cpp
@@ -1708,6 +1708,7 @@ void FSceneRenderer::RenderShadowDepthMapAtlases(FRHICommandListImmediate& RHICm
 			bool bWantsNearClip = false;
 			bool bWantsNoNearClip = false;
 			TArray<Nanite::FPackedView, SceneRenderingAllocator> PackedViews;
+			TArray<Nanite::FPackedView, SceneRenderingAllocator> PackedViewsNoNearClip;
 			TArray<FProjectedShadowInfo*, SceneRenderingAllocator> ShadowsToEmit;
 			for (int32 ShadowIndex = 0; ShadowIndex < ShadowMapAtlas.Shadows.Num(); ShadowIndex++)
 			{
@@ -1723,17 +1724,6 @@ void FSceneRenderer::RenderShadowDepthMapAtlases(FRHICommandListImmediate& RHICm
 					continue;
 				}
 
-				// Orthographic shadow projections want depth clamping rather than clipping
-				if (ProjectedShadowInfo->ShouldClampToNearPlane())
-				{
-					bWantsNoNearClip = true;
-				}
-				else
-				{
-					bWantsNearClip = true;
-				}
-
-				// Only add the view if it has instances
 				const FIntRect AtlasViewRect = ProjectedShadowInfo->GetViewRectForView();
 
 				Nanite::FPackedViewParams Initializer;
@@ -1742,12 +1732,22 @@ void FSceneRenderer::RenderShadowDepthMapAtlases(FRHICommandListImmediate& RHICm
 				Initializer.ViewRect = AtlasViewRect;
 				Initializer.RasterContextSize = AtlasSize;
 				Initializer.LODScaleFactor = ComputeNaniteShadowsLODScaleFactor();
-				PackedViews.Add(Nanite::CreatePackedView(Initializer));
+
+				// Orthographic shadow projections want depth clamping rather than clipping
+				if (ProjectedShadowInfo->ShouldClampToNearPlane())
+				{
+					PackedViewsNoNearClip.Add(Nanite::CreatePackedView(Initializer));
+				}
+				else
+				{
+					PackedViews.Add(Nanite::CreatePackedView(Initializer));
+				}
+				
 
 				ShadowsToEmit.Add(ProjectedShadowInfo);
 			}
 				
-			if (PackedViews.Num() > 0)
+			if (PackedViews.Num() > 0 || PackedViewsNoNearClip.Num() > 0)
 			{
 				FRDGBuilder GraphBuilder(RHICmdList);
 				FRDGTextureRef ShadowMap = GraphBuilder.RegisterExternalTexture(ShadowMapAtlas.RenderTargets.DepthTarget, TEXT("DepthBuffer"));
@@ -1755,29 +1755,50 @@ void FSceneRenderer::RenderShadowDepthMapAtlases(FRHICommandListImmediate& RHICm
 				{
 					RDG_EVENT_SCOPE(GraphBuilder, "Nanite Shadows");
 
+					// Need separate passes for near clip on/off currently
+					const bool bSupportsMultiplePasses = (PackedViews.Num() > 0 && PackedViewsNoNearClip.Num() > 0);
+
 					// NOTE: Rendering into an atlas like this is not going to work properly with HZB, but we are not currently using HZB here.
 					// It might be worthwhile going through the virtual SM rendering path even for "dense" cases even just for proper handling of all the details.
 					FIntRect FullAtlasViewRect(FIntPoint(0, 0), AtlasSize);
 					const bool bUpdateStreaming = CVarNaniteShadowsUpdateStreaming.GetValueOnRenderThread() != 0;
-					Nanite::FCullingContext CullingContext = Nanite::InitCullingContext(GraphBuilder, *Scene, nullptr, FullAtlasViewRect, true, bUpdateStreaming, false, false);
+					Nanite::FCullingContext CullingContext = Nanite::InitCullingContext(GraphBuilder, *Scene, nullptr, FullAtlasViewRect, true, bUpdateStreaming, bSupportsMultiplePasses, false);
 					Nanite::FRasterContext RasterContext = Nanite::InitRasterContext(GraphBuilder, AtlasSize, Nanite::EOutputBufferMode::DepthOnly);
+							
+					
+					if (PackedViews.Num() > 0)
+					{
+						Nanite::FRasterState RasterState;
+						RasterState.bNearClip = true;
 
-					// We can't currently do both at once, but this shouldn't currently happen since directional lights go into different atlases
-					// If we ever need this we can do two passes like virtual shadow maps
-					check(!bWantsNearClip || !bWantsNoNearClip);
-					Nanite::FRasterState RasterState;
-					RasterState.bNearClip = bWantsNearClip;
+						Nanite::CullRasterize(
+							GraphBuilder,
+							*Scene,
+							PackedViews,
+							CullingContext,
+							RasterContext,
+							RasterState,
+							nullptr,	// InstanceDraws
+							false		// bExtractStats TODO
+						);
+					}
 
-					Nanite::CullRasterize(
-						GraphBuilder,
-						*Scene,
-						PackedViews,
-						CullingContext,
-						RasterContext,
-						RasterState,
-						nullptr,	// InstanceDraws
-						false		// bExtractStats TODO
-					);
+					if (PackedViewsNoNearClip.Num() > 0)
+					{
+						Nanite::FRasterState RasterState;
+						RasterState.bNearClip = false;
+
+						Nanite::CullRasterize(
+							GraphBuilder,
+							*Scene,
+							PackedViewsNoNearClip,
+							CullingContext,
+							RasterContext,
+							RasterState,
+							nullptr,	// InstanceDraws
+							false		// bExtractStats TODO
+						);
+					}
 
 					for (FProjectedShadowInfo* ProjectedShadowInfo : ShadowsToEmit)
 					{
