Fix assert when view frustum contains a far culling plane.
Add a comment to try and warn about any changes to the frustum plane extraction code in the future.

diff --git a/Engine/Source/Runtime/Engine/Private/ConvexVolume.cpp b/Engine/Source/Runtime/Engine/Private/ConvexVolume.cpp
--- a/Engine/Source/Runtime/Engine/Private/ConvexVolume.cpp
+++ b/Engine/Source/Runtime/Engine/Private/ConvexVolume.cpp
@@ -451,6 +451,9 @@ void GetViewFrustumBoundsInternal(FConvexVolume& OutResult, const FMatrix& ViewP
 	OutResult.Planes.Empty(6);
 	FPlane Temp;
 
+	// NOTE: Be careful changing anything here! Some callers make assumptions about the order of the planes returned.
+	// See for instance BuildLightViewFrustumConvexHull in ShadowSetup.cpp
+
 	// Near clipping plane.
 	if (bUseNearPlane && ViewProjectionMatrix.GetFrustumNearPlane(Temp))
 	{
diff --git a/Engine/Source/Runtime/Renderer/Private/ShadowSetup.cpp b/Engine/Source/Runtime/Renderer/Private/ShadowSetup.cpp
--- a/Engine/Source/Runtime/Renderer/Private/ShadowSetup.cpp
+++ b/Engine/Source/Runtime/Renderer/Private/ShadowSetup.cpp
@@ -2940,11 +2940,12 @@ typedef TArray<FConvexVolume, TInlineAllocator<8>> FLightViewFrustumConvexHulls;
 // Based on: https://udn.unrealengine.com/questions/410475/improved-culling-of-shadow-casters-for-spotlights.html
 void BuildLightViewFrustumConvexHull(const FVector& LightOrigin, const FConvexVolume& Frustum, FConvexVolume& ConvexHull)
 {
-	// This function assumes that there are 4 planes
-	// If this isn't the case, we should really know about it, so assert.
+	// This function works with 4 frustum planes
+	// We may occasionally see a 5th if OverrideFarClippingPlaneDistance is used such as in the Editor, but we ignore that 
+	// virtual plane (which we assume is the last in the list) for the purposes of culling here.
 	const int32 EdgeCount = 4;
 	const int32 PlaneCount = 4;
-	check(Frustum.Planes.Num() == PlaneCount);
+	check(Frustum.Planes.Num() == 4 || Frustum.Planes.Num() == 5);
 
 	enum EFrustumPlanes
 	{
