Minor precision improvement to clipmap math far away from the origin.
Remove some of the logical asserts as they can trip due to precision in these cases, but it is not a catestrophic error, and the checks were mostly useful when developing the initial logic.

diff --git a/Engine/Source/Runtime/Renderer/Private/VirtualShadowMaps/VirtualShadowMapClipmap.cpp b/Engine/Source/Runtime/Renderer/Private/VirtualShadowMaps/VirtualShadowMapClipmap.cpp
--- a/Engine/Source/Runtime/Renderer/Private/VirtualShadowMaps/VirtualShadowMapClipmap.cpp
+++ b/Engine/Source/Runtime/Renderer/Private/VirtualShadowMaps/VirtualShadowMapClipmap.cpp
@@ -37,18 +37,22 @@ FVirtualShadowMapClipmap::FVirtualShadowMapClipmap(
 	FVirtualShadowMapArray& VirtualShadowMapArray,
 	FVirtualShadowMapArrayCacheManager* VirtualShadowMapArrayCacheManager,
 	const FLightSceneInfo& InLightSceneInfo,
-	const FMatrix& WorldToLight,
+	const FMatrix& WorldToLightRotationMatrix,
 	const FViewMatrices& CameraViewMatrices,
 	float MaxRadius)
 	: LightSceneInfo(InLightSceneInfo)
 {
+	check(WorldToLightRotationMatrix.GetOrigin() == FVector(0, 0, 0));	// Should not contain translation or scaling
+
 	const FMatrix FaceMatrix(
 		FPlane( 0, 0, 1, 0 ),
 		FPlane( 0, 1, 0, 0 ),
 		FPlane(-1, 0, 0, 0 ),
 		FPlane( 0, 0, 0, 1 ));
 
-	TranslatedWorldToView = WorldToLight * FaceMatrix;
+	WorldToViewRotationMatrix = WorldToLightRotationMatrix * FaceMatrix;
+	// Pure rotation matrix
+	FMatrix ViewToWorldRotationMatrix = WorldToViewRotationMatrix.GetTransposed();
 	
 	// NOTE: Rotational (roll) invariance of the directional light depends on square pixels so we just base everything on the camera X scales/resolution
 	float LodScale = 1.0f / CameraViewMatrices.GetProjectionScale().X;
@@ -90,23 +94,16 @@ FVirtualShadowMapClipmap::FVirtualShadowMapClipmap(
 
 		// NOTE: We may eventually want to snap/quantize Z as well, but right now we just adjust it each time
 		// we update the cache page. With clipmaps this could ertainly be simplified.
-		FVector ViewCenter = TranslatedWorldToView.TransformPosition(WorldOrigin);
+		FVector ViewCenter = WorldToViewRotationMatrix.TransformPosition(WorldOrigin);
 		FIntPoint PageSpaceCenter(
 			FMath::RoundToInt(ViewCenter.X / PageSizeInWorldSpace),
 			FMath::RoundToInt(ViewCenter.Y / PageSizeInWorldSpace));
 		ViewCenter.X = PageSpaceCenter.X * PageSizeInWorldSpace;
 		ViewCenter.Y = PageSpaceCenter.Y * PageSizeInWorldSpace;
 
-		const FVector SnappedWorldCenter = TranslatedWorldToView.InverseTransformPosition(ViewCenter);
+		const FVector SnappedWorldCenter = ViewToWorldRotationMatrix.TransformPosition(ViewCenter);
 		const float SnappedLevelRadius = 0.5f * (PageSizeInWorldSpace * DimPages);
-			
-		// Sanity check that we didn't snap further away than we thought we were going to
-		{
-			float SnapError = (SnappedWorldCenter - WorldOrigin).Size();
-			check(SnapError <= (UE_HALF_SQRT_2 * PageSizeInWorldSpace));
-			check((RawLevelRadius + SnapError) <= SnappedLevelRadius);
-		}
-		
+				
 		// NOTE: Any ZScale/ZOffset change probably needs to invalidate the page cache...
 		// For now they should be stable for a given (absolute) clipmap level
 		const float ZScale = 0.5f / RawLevelRadius;
@@ -125,7 +122,7 @@ FVirtualShadowMapClipmap::FVirtualShadowMapClipmap(
 				PageOffset.Y = -PageOffset.Y;		// Viewport
 				float DepthOffset = -ViewCenter.Z * ZScale;
 
-				CacheEntry->UpdateClipmap(Level.VirtualShadowMap->ID, WorldToLight, PageOffset, DepthOffset);
+				CacheEntry->UpdateClipmap(Level.VirtualShadowMap->ID, WorldToLightRotationMatrix, PageOffset, DepthOffset);
 
 				Level.VirtualShadowMap->VirtualShadowMapCacheEntry = CacheEntry;
 			}
@@ -143,7 +140,7 @@ FViewMatrices FVirtualShadowMapClipmap::GetViewMatrices(int32 ClipmapIndex) cons
 	// NOTE: Be careful here! There's special logic in FViewMatrices around ViewOrigin for ortho projections we need to bypass...
 	// There's also the fact that some of this data is going to be "wrong", due to the "overridden" matrix thing that shadows do
 	Initializer.ViewOrigin = Level.WorldCenter;
-	Initializer.ViewRotationMatrix = TranslatedWorldToView;
+	Initializer.ViewRotationMatrix = WorldToViewRotationMatrix;
 	Initializer.ProjectionMatrix = Level.ViewToClip;
 
 	// TODO: This is probably unused in the shadows/nanite path, but coupling here is not ideal
@@ -158,9 +155,9 @@ FVirtualShadowMapProjectionShaderData FVirtualShadowMapClipmap::GetProjectionSha
 	const FLevelData& Level = LevelData[ClipmapIndex];
 
 	FVirtualShadowMapProjectionShaderData Data;
-	Data.TranslatedWorldToShadowViewMatrix = TranslatedWorldToView;
+	Data.TranslatedWorldToShadowViewMatrix = WorldToViewRotationMatrix;
 	Data.ShadowViewToClipMatrix = Level.ViewToClip;
-	Data.TranslatedWorldToShadowUvNormalMatrix = CalcTranslatedWorldToShadowUvNormalMatrix(TranslatedWorldToView, Level.ViewToClip);
+	Data.TranslatedWorldToShadowUvNormalMatrix = CalcTranslatedWorldToShadowUvNormalMatrix(WorldToViewRotationMatrix, Level.ViewToClip);
 	Data.ShadowPreViewTranslation = -Level.WorldCenter;
 	Data.VirtualShadowMapId = Level.VirtualShadowMap->ID;
 	Data.LightType = ELightComponentType::LightType_Directional;
diff --git a/Engine/Source/Runtime/Renderer/Private/VirtualShadowMaps/VirtualShadowMapClipmap.h b/Engine/Source/Runtime/Renderer/Private/VirtualShadowMaps/VirtualShadowMapClipmap.h
--- a/Engine/Source/Runtime/Renderer/Private/VirtualShadowMaps/VirtualShadowMapClipmap.h
+++ b/Engine/Source/Runtime/Renderer/Private/VirtualShadowMaps/VirtualShadowMapClipmap.h
@@ -20,7 +20,7 @@ public:
 		FVirtualShadowMapArray& VirtualShadowMapArray,
 		FVirtualShadowMapArrayCacheManager* VirtualShadowMapArrayCacheManager,
 		const FLightSceneInfo& InLightSceneInfo,
-		const FMatrix& WorldToLight,
+		const FMatrix& WorldToLightRotationMatrix,
 		const FViewMatrices& CameraViewMatrices,
 		float MaxRadius		// Maximum radius the clipmap must cover from the center point; used to compute level count
 	);
@@ -65,7 +65,8 @@ private:
 	* */
 	FVector WorldOrigin;
 
-	FMatrix TranslatedWorldToView;
+	/** Directional light rotation matrix (no translation) */
+	FMatrix WorldToViewRotationMatrix;
 
 	int32 FirstLevel;
 	float ResolutionLodBias;
