Guard access to GPUScene with UseGPUScene() in FPrimitiveSceneInfo::CacheMeshDrawCommands, otherwise it will crash on platforms that don't support GPUScene

diff --git a/Engine/Source/Runtime/Renderer/Private/PrimitiveSceneInfo.cpp b/Engine/Source/Runtime/Renderer/Private/PrimitiveSceneInfo.cpp
--- a/Engine/Source/Runtime/Renderer/Private/PrimitiveSceneInfo.cpp
+++ b/Engine/Source/Runtime/Renderer/Private/PrimitiveSceneInfo.cpp
@@ -440,25 +440,28 @@ void FPrimitiveSceneInfo::CacheMeshDrawCommands(FRHICommandListImmediate& RHICmd
 	}
 #endif
 
-	SCOPE_CYCLE_COUNTER(STAT_UpdateGPUSceneTime);
-	for (FPrimitiveSceneInfo* SceneInfo : SceneInfos)
+	if (UseGPUScene(GMaxRHIShaderPlatform, Scene->GetFeatureLevel()))
 	{
-		check(SceneInfo->InstanceDataOffset == INDEX_NONE);
-		check(SceneInfo->NumInstanceDataEntries == 0);
-		if (SceneInfo->Proxy->SupportsInstanceDataBuffer() && UseGPUScene(GMaxRHIShaderPlatform, Scene->GetFeatureLevel()))
+		SCOPE_CYCLE_COUNTER(STAT_UpdateGPUSceneTime);
+		for (FPrimitiveSceneInfo* SceneInfo : SceneInfos)
 		{
-			if (const TArray<FPrimitiveInstance>* PrimitiveInstances = SceneInfo->Proxy->GetPrimitiveInstances())
+			check(SceneInfo->InstanceDataOffset == INDEX_NONE);
+			check(SceneInfo->NumInstanceDataEntries == 0);
+			if (SceneInfo->Proxy->SupportsInstanceDataBuffer())
 			{
-				SceneInfo->InstanceDataOffset = Scene->GPUScene.AllocateInstanceSlots(PrimitiveInstances->Num());
-				SceneInfo->NumInstanceDataEntries = PrimitiveInstances->Num();
+				if (const TArray<FPrimitiveInstance>* PrimitiveInstances = SceneInfo->Proxy->GetPrimitiveInstances())
+				{
+					SceneInfo->InstanceDataOffset = Scene->GPUScene.AllocateInstanceSlots(PrimitiveInstances->Num());
+					SceneInfo->NumInstanceDataEntries = PrimitiveInstances->Num();
+				}
 			}
-		}
 
-		// Force a primitive update in the GPU scene
-		if (!Scene->GPUScene.PrimitivesMarkedToUpdate[SceneInfo->PackedIndex])
-		{
-			Scene->GPUScene.PrimitivesToUpdate.Add(SceneInfo->PackedIndex);
-			Scene->GPUScene.PrimitivesMarkedToUpdate[SceneInfo->PackedIndex] = true;
+			// Force a primitive update in the GPU scene
+			if (!Scene->GPUScene.PrimitivesMarkedToUpdate[SceneInfo->PackedIndex])
+			{
+				Scene->GPUScene.PrimitivesToUpdate.Add(SceneInfo->PackedIndex);
+				Scene->GPUScene.PrimitivesMarkedToUpdate[SceneInfo->PackedIndex] = true;
+			}
 		}
 	}
 }
