Fixes crashes with sg.PostProcessQuality=0 in QAGame

diff --git a/Engine/Source/Runtime/Renderer/Private/IndirectLightRendering.cpp b/Engine/Source/Runtime/Renderer/Private/IndirectLightRendering.cpp
--- a/Engine/Source/Runtime/Renderer/Private/IndirectLightRendering.cpp
+++ b/Engine/Source/Runtime/Renderer/Private/IndirectLightRendering.cpp
@@ -844,11 +844,16 @@ void FDeferredShadingSceneRenderer::RenderDiffuseIndirectAndAmbientOcclusion(
 		}
 		else if (ViewPipelineState.AmbientOcclusionMethod == EAmbientOcclusionMethod::SSAO)
 		{
-			// TODO: Convert SSAO to RDG and setup it up here.
-			check(SceneContext.bScreenSpaceAOIsValid);
-
-			// Fetch result of SSAO that was done earlier.
-			AmbientOcclusionMask = GraphBuilder.RegisterExternalTexture(SceneContext.ScreenSpaceAO);
+			// TODO: Calls SSAO here once converted to RDG.
+			if (SceneContext.bScreenSpaceAOIsValid)
+			{
+				// Fetch result of SSAO that was done earlier.
+				AmbientOcclusionMask = GraphBuilder.RegisterExternalTexture(SceneContext.ScreenSpaceAO);
+			}
+			else
+			{
+				AmbientOcclusionMask = GraphBuilder.RegisterExternalTexture(GSystemTextures.WhiteDummy);
+			}
 		}
 		else
 		{
