Removed dead code from IndirectLightingRendering

diff --git a/Engine/Source/Runtime/Renderer/Private/IndirectLightRendering.cpp b/Engine/Source/Runtime/Renderer/Private/IndirectLightRendering.cpp
--- a/Engine/Source/Runtime/Renderer/Private/IndirectLightRendering.cpp
+++ b/Engine/Source/Runtime/Renderer/Private/IndirectLightRendering.cpp
@@ -691,9 +691,6 @@ void FDeferredShadingSceneRenderer::RenderDiffuseIndirectAndAmbientOcclusion(
 
 	for (FViewInfo& View : Views)
 	{
-		FRDGTexture* FurthestHZBTexture = GraphBuilder.RegisterExternalTexture(View.HZB);
-		FRDGTexture* ClosestHZBTexture = GraphBuilder.TryRegisterExternalTexture(View.ClosestHZB);
-
 		const FPerViewPipelineState& ViewPipelineState = GetViewPipelineState(View);
 
 		int32 DenoiseMode = CVarDiffuseIndirectDenoiser.GetValueOnRenderThread();
