Don't use FRenderQueryPoolRHIRef for fencing because it's not thread safe - alloc and free one query at a time

diff --git a/Engine/Source/Runtime/Renderer/Private/VT/VirtualTextureFeedback.cpp b/Engine/Source/Runtime/Renderer/Private/VT/VirtualTextureFeedback.cpp
--- a/Engine/Source/Runtime/Renderer/Private/VT/VirtualTextureFeedback.cpp
+++ b/Engine/Source/Runtime/Renderer/Private/VT/VirtualTextureFeedback.cpp
@@ -70,23 +70,17 @@ public:
 class FFeedbackGPUFencePool
 {
 public:
-	FRenderQueryPoolRHIRef FenceQueryPool;
 	FGPUFenceRHIRef DummyFence;
-	TArray<FRHIPooledRenderQuery> Fences;
+	TArray<FRenderQueryRHIRef> Fences;
 	bool bDummyFenceWritten = false;
 
 	FFeedbackGPUFencePool(int32 InSize)
 	{
-		Fences.AddDefaulted(InSize);
+		Fences.Reserve(InSize);
 	}
 
 	void InitRHI()
 	{
-		if (!FenceQueryPool.IsValid())
-		{
-			FenceQueryPool = RHICreateRenderQueryPool(RQT_AbsoluteTime, Fences.Num());
-		}
-
 		if (!DummyFence.IsValid())
 		{
 			DummyFence = RHICreateGPUFence(FName());
@@ -96,26 +90,19 @@ public:
 
 	void ReleaseRHI()
 	{
-		for (int i = 0; i < Fences.Num(); ++i)
-		{
-			if (Fences[i].IsValid())
-			{
-				Fences[i].ReleaseQuery();
-			}
-		}
+		Fences.Empty(Fences.Num());
 
 		DummyFence.SafeRelease();
 		bDummyFenceWritten = false;
-
-		FenceQueryPool.SafeRelease();
 	}
 
 	void Allocate(FRHICommandListImmediate& RHICmdList, int32 Index)
 	{
-		if (!Fences[Index].IsValid())
+		if (Fences[Index].IsValid())
 		{
-			Fences[Index] = FenceQueryPool->AllocateQuery();
+			Fences[Index].SafeRelease();
 		}
+		Fences[Index] = GDynamicRHI->RHICreateRenderQuery(RQT_AbsoluteTime);
 
 		if (!bDummyFenceWritten && DummyFence.IsValid())
 		{
@@ -128,13 +115,13 @@ public:
 	
 	void Write(FRHICommandListImmediate& RHICmdList, int32 Index)
 	{
-		RHICmdList.EndRenderQuery(Fences[Index].GetQuery());
+		RHICmdList.EndRenderQuery(Fences[Index]);
 	}
 
 	bool Poll(FRHICommandListImmediate& RHICmdList, int32 Index)
 	{
 		uint64 Dummy;
-		return RHICmdList.GetRenderQueryResult(Fences[Index].GetQuery(), Dummy, false, RHICmdList.GetGPUMask().ToIndex());
+		return RHICmdList.GetRenderQueryResult(Fences[Index], Dummy, false, RHICmdList.GetGPUMask().ToIndex());
 	}
 
 	FGPUFenceRHIRef GetMapFence(int32 Index)
@@ -144,7 +131,7 @@ public:
 
 	void Release(int32 Index)
 	{
-		Fences[Index].ReleaseQuery();
+		Fences[Index].SafeRelease();
 	}
 };
 
