Fix correct initialize size setup of array - wrong last minute change before checkin

diff --git a/Engine/Source/Runtime/Renderer/Private/VT/VirtualTextureFeedback.cpp b/Engine/Source/Runtime/Renderer/Private/VT/VirtualTextureFeedback.cpp
--- a/Engine/Source/Runtime/Renderer/Private/VT/VirtualTextureFeedback.cpp
+++ b/Engine/Source/Runtime/Renderer/Private/VT/VirtualTextureFeedback.cpp
@@ -76,7 +76,7 @@ public:
 
 	FFeedbackGPUFencePool(int32 InSize)
 	{
-		Fences.Reserve(InSize);
+		Fences.AddDefaulted(InSize);
 	}
 
 	void InitRHI()
@@ -90,7 +90,13 @@ public:
 
 	void ReleaseRHI()
 	{
-		Fences.Empty(Fences.Num());
+		for (int i = 0; i < Fences.Num(); ++i)
+		{
+			if (Fences[i].IsValid())
+			{
+				Fences[i].SafeRelease();
+			}
+		}
 
 		DummyFence.SafeRelease();
 		bDummyFenceWritten = false;
