Move complete shadow maps to separate atlases so they can be copied instead of re-rendered (upcoming change)
Trim non-CSM shadow atlases to avoid allocating unused space (as is done with the CSM atlases)

diff --git a/Engine/Source/Runtime/Renderer/Private/SceneRendering.h b/Engine/Source/Runtime/Renderer/Private/SceneRendering.h
--- a/Engine/Source/Runtime/Renderer/Private/SceneRendering.h
+++ b/Engine/Source/Runtime/Renderer/Private/SceneRendering.h
@@ -1829,7 +1829,7 @@ protected:
 
 	void AllocateShadowDepthTargets(FRHICommandListImmediate& RHICmdList);
 	
-	void AllocatePerObjectShadowDepthTargets(FRHICommandListImmediate& RHICmdList, TArray<FProjectedShadowInfo*, SceneRenderingAllocator>& Shadows);
+	void AllocateAtlasedShadowDepthTargets(FRHICommandListImmediate& RHICmdList, TArray<FProjectedShadowInfo*, SceneRenderingAllocator>& Shadows);
 
 	void AllocateCachedSpotlightShadowDepthTargets(FRHICommandListImmediate& RHICmdList, TArray<FProjectedShadowInfo*, SceneRenderingAllocator>& CachedShadows);
 
diff --git a/Engine/Source/Runtime/Renderer/Private/ShadowSetup.cpp b/Engine/Source/Runtime/Renderer/Private/ShadowSetup.cpp
--- a/Engine/Source/Runtime/Renderer/Private/ShadowSetup.cpp
+++ b/Engine/Source/Runtime/Renderer/Private/ShadowSetup.cpp
@@ -262,6 +262,13 @@ static TAutoConsoleVariable<int32> CVarCompleteShadowMapResolution(
 	ECVF_RenderThreadSafe
 );
 
+static TAutoConsoleVariable<int32> CVarAlwaysAllocateMaxResolutionAtlases(
+	TEXT("r.Shadow.AlwaysAllocateMaxResolutionAtlases"),
+	1,
+	TEXT("If enabled, will always allocate shadow map atlases at the maximum resolution rather than trimming unused space. This will waste more memory but can possibly reduce render target pool fragmentation and thrash."),
+	ECVF_RenderThreadSafe
+);
+
 
 
 CSV_DECLARE_CATEGORY_EXTERN(LightCount);
@@ -4315,6 +4322,7 @@ void FSceneRenderer::AllocateShadowDepthTargets(FRHICommandListImmediate& RHICmd
 	// Sort visible shadows based on their allocation needs
 	// 2d shadowmaps for this frame only that can be atlased across lights
 	TArray<FProjectedShadowInfo*, SceneRenderingAllocator> Shadows;
+	TArray<FProjectedShadowInfo*, SceneRenderingAllocator> CompleteShadows;
 	// 2d shadowmaps that will persist across frames, can't be atlased
 	TArray<FProjectedShadowInfo*, SceneRenderingAllocator> CachedSpotlightShadows;
 	TArray<FProjectedShadowInfo*, SceneRenderingAllocator> TranslucentShadows;
@@ -4332,6 +4340,7 @@ void FSceneRenderer::AllocateShadowDepthTargets(FRHICommandListImmediate& RHICmd
 
 		// All cascades for a light need to be in the same texture
 		TArray<FProjectedShadowInfo*, SceneRenderingAllocator> WholeSceneDirectionalShadows;
+		TArray<FProjectedShadowInfo*, SceneRenderingAllocator> WholeSceneCompleteDirectionalShadows;
 
 		for (int32 ShadowIndex = 0; ShadowIndex < VisibleLightInfo.AllProjectedShadows.Num(); ShadowIndex++)
 		{
@@ -4448,7 +4457,14 @@ void FSceneRenderer::AllocateShadowDepthTargets(FRHICommandListImmediate& RHICmd
 					}
 					else if (ProjectedShadowInfo->bDirectionalLight && ProjectedShadowInfo->bWholeSceneShadow)
 					{
-						WholeSceneDirectionalShadows.Add(ProjectedShadowInfo);
+						if (ProjectedShadowInfo->bCompleteShadowMap)
+						{
+							WholeSceneCompleteDirectionalShadows.Add(ProjectedShadowInfo);
+						}
+						else
+						{
+							WholeSceneDirectionalShadows.Add(ProjectedShadowInfo);
+						}
 					}
 					else if (ProjectedShadowInfo->bOnePassPointLightShadow)
 					{
@@ -4465,7 +4481,14 @@ void FSceneRenderer::AllocateShadowDepthTargets(FRHICommandListImmediate& RHICmd
 					}
 					else
 					{
-						Shadows.Add(ProjectedShadowInfo);
+						if (ProjectedShadowInfo->bCompleteShadowMap)
+						{
+							CompleteShadows.Add(ProjectedShadowInfo);
+						}
+						else
+						{
+							Shadows.Add(ProjectedShadowInfo);
+						}
 					}
 				}
 			}
@@ -4477,6 +4500,7 @@ void FSceneRenderer::AllocateShadowDepthTargets(FRHICommandListImmediate& RHICmd
 		VisibleLightInfo.CompleteProjectedShadows.Sort(FCompareFProjectedShadowInfoBySplitIndex());
 
 		AllocateCSMDepthTargets(RHICmdList, WholeSceneDirectionalShadows);
+		AllocateCSMDepthTargets(RHICmdList, WholeSceneCompleteDirectionalShadows);
 	}
 
 	if (CachedPreShadows.Num() > 0)
@@ -4505,7 +4529,8 @@ void FSceneRenderer::AllocateShadowDepthTargets(FRHICommandListImmediate& RHICmd
 	AllocateOnePassPointLightDepthTargets(RHICmdList, WholeScenePointShadows);
 	AllocateRSMDepthTargets(RHICmdList, RSMShadows);
 	AllocateCachedSpotlightShadowDepthTargets(RHICmdList, CachedSpotlightShadows);
-	AllocatePerObjectShadowDepthTargets(RHICmdList, Shadows);
+	AllocateAtlasedShadowDepthTargets(RHICmdList, Shadows);
+	AllocateAtlasedShadowDepthTargets(RHICmdList, CompleteShadows);
 	AllocateTranslucentShadowDepthTargets(RHICmdList, TranslucentShadows);
 
 	// Update translucent shadow map uniform buffers.
@@ -4543,47 +4568,56 @@ void FSceneRenderer::AllocateShadowDepthTargets(FRHICommandListImmediate& RHICmd
 	SET_MEMORY_STAT(STAT_ShadowmapAtlasMemory, SortedShadowsForShadowDepthPass.ComputeMemorySize());
 }
 
-void FSceneRenderer::AllocatePerObjectShadowDepthTargets(FRHICommandListImmediate& RHICmdList, TArray<FProjectedShadowInfo*, SceneRenderingAllocator>& Shadows)
+struct FLayoutAndAssignedShadows
 {
-	if (Shadows.Num() > 0)
-	{
-		FSceneRenderTargets& SceneContext = FSceneRenderTargets::Get(RHICmdList);
-		const FIntPoint ShadowBufferResolution = SceneContext.GetShadowDepthTextureResolution();
+	FLayoutAndAssignedShadows(int32 MaxTextureSize) :
+		TextureLayout(1, 1, MaxTextureSize, MaxTextureSize, false, ETextureLayoutAspectRatio::None, false)
+	{}
 
-		int32 OriginalNumAtlases = SortedShadowsForShadowDepthPass.ShadowMapAtlases.Num();
+	FLayoutAndAssignedShadows(FIntPoint MaxTextureSize) :
+		TextureLayout(1, 1, MaxTextureSize.X, MaxTextureSize.Y, false, ETextureLayoutAspectRatio::None, false)
+	{}
 
-		FTextureLayout CurrentShadowLayout(1, 1, ShadowBufferResolution.X, ShadowBufferResolution.Y, false, ETextureLayoutAspectRatio::None, false);
-		FPooledRenderTargetDesc ShadowMapDesc2D = FPooledRenderTargetDesc::Create2DDesc(ShadowBufferResolution, PF_ShadowDepth, FClearValueBinding::DepthOne, TexCreate_None, TexCreate_DepthStencilTargetable | TexCreate_ShaderResource, false);
-		ShadowMapDesc2D.Flags |= GFastVRamConfig.ShadowPerObject;
+	FTextureLayout TextureLayout;
+	TArray<FProjectedShadowInfo*, SceneRenderingAllocator> Shadows;
+};
 
-		// Sort the projected shadows by resolution.
-		Shadows.Sort(FCompareFProjectedShadowInfoByResolution());
+void FSceneRenderer::AllocateAtlasedShadowDepthTargets(FRHICommandListImmediate& RHICmdList, TArray<FProjectedShadowInfo*, SceneRenderingAllocator>& Shadows)
+{
+	if (Shadows.Num() == 0)
+	{
+		return;
+	}
 
-		for (int32 ShadowIndex = 0; ShadowIndex < Shadows.Num(); ShadowIndex++)
-		{
-			FProjectedShadowInfo* ProjectedShadowInfo = Shadows[ShadowIndex];
+	FSceneRenderTargets& SceneContext = FSceneRenderTargets::Get(RHICmdList);
+	const FIntPoint MaxTextureSize = SceneContext.GetShadowDepthTextureResolution();
 
-			// Atlased shadows need a border
-			check(ProjectedShadowInfo->BorderSize != 0);
-			check(!ProjectedShadowInfo->bAllocated);
+	TArray<FLayoutAndAssignedShadows, SceneRenderingAllocator> Layouts;
+	Layouts.Add(FLayoutAndAssignedShadows(MaxTextureSize));
 
-			if (ProjectedShadowInfo->CacheMode == SDCM_MovablePrimitivesOnly && !ProjectedShadowInfo->HasSubjectPrims())
-			{
-				FCachedShadowMapData& CachedShadowMapData = Scene->CachedShadowMaps.FindChecked(ProjectedShadowInfo->GetLightSceneInfo().Id);
-				ProjectedShadowInfo->X = ProjectedShadowInfo->Y = 0;
-				ProjectedShadowInfo->bAllocated = true;
-				// Skip the shadow depth pass since there are no movable primitives to composite, project from the cached shadowmap directly which contains static primitive depths
-				ProjectedShadowInfo->RenderTargets.CopyReferencesFromRenderTargets(CachedShadowMapData.ShadowMap);
-			}
-			else
-			{
-				if (SortedShadowsForShadowDepthPass.ShadowMapAtlases.Num() == OriginalNumAtlases)
-				{
-					// Start with an empty atlas for per-object shadows (don't allow packing object shadows into the CSM atlas atm)
-					SortedShadowsForShadowDepthPass.ShadowMapAtlases.AddDefaulted();
-				}
+	for (int32 ShadowIndex = 0; ShadowIndex < Shadows.Num(); ShadowIndex++)
+	{
+		FProjectedShadowInfo* ProjectedShadowInfo = Shadows[ShadowIndex];
 
-				if (CurrentShadowLayout.AddElement(
+		// Atlased shadows need a border
+		check(ProjectedShadowInfo->BorderSize != 0);
+		check(!ProjectedShadowInfo->bAllocated);
+
+		if (ProjectedShadowInfo->CacheMode == SDCM_MovablePrimitivesOnly && !ProjectedShadowInfo->HasSubjectPrims())
+		{
+			FCachedShadowMapData& CachedShadowMapData = Scene->CachedShadowMaps.FindChecked(ProjectedShadowInfo->GetLightSceneInfo().Id);
+			ProjectedShadowInfo->X = ProjectedShadowInfo->Y = 0;
+			ProjectedShadowInfo->bAllocated = true;
+			// Skip the shadow depth pass since there are no movable primitives to composite, project from the cached shadowmap directly which contains static primitive depths
+			ProjectedShadowInfo->RenderTargets.CopyReferencesFromRenderTargets(CachedShadowMapData.ShadowMap);
+		}
+		else
+		{
+			// Avoid infinite loop if texture cannot be allocated even on a fresh atlas
+			// This should not occur, but good to have a safeguard and will still trigger the check() below
+			for (int32 Attempt = 0; Attempt < 2; ++Attempt)
+			{
+				if (Layouts.Last().TextureLayout.AddElement(
 					ProjectedShadowInfo->X,
 					ProjectedShadowInfo->Y,
 					ProjectedShadowInfo->ResolutionX + ProjectedShadowInfo->BorderSize * 2,
@@ -4591,32 +4625,45 @@ void FSceneRenderer::AllocatePerObjectShadowDepthTargets(FRHICommandListImmediat
 					)
 				{
 					ProjectedShadowInfo->bAllocated = true;
-				}
-				else
-				{
-					CurrentShadowLayout = FTextureLayout(1, 1, ShadowBufferResolution.X, ShadowBufferResolution.Y, false, ETextureLayoutAspectRatio::None, false);
-					SortedShadowsForShadowDepthPass.ShadowMapAtlases.AddDefaulted();
-
-					if (CurrentShadowLayout.AddElement(
-						ProjectedShadowInfo->X,
-						ProjectedShadowInfo->Y,
-						ProjectedShadowInfo->ResolutionX + ProjectedShadowInfo->BorderSize * 2,
-						ProjectedShadowInfo->ResolutionY + ProjectedShadowInfo->BorderSize * 2)
-						)
-					{
-						ProjectedShadowInfo->bAllocated = true;
-					}
+					Layouts.Last().Shadows.Add(ProjectedShadowInfo);
+					break;
 				}
 
-				check(ProjectedShadowInfo->bAllocated);
+				// Out of space, add a new atlas and try again
+				Layouts.Add(FLayoutAndAssignedShadows(MaxTextureSize));
+			}
 
-				FSortedShadowMapAtlas& ShadowMapAtlas = SortedShadowsForShadowDepthPass.ShadowMapAtlases.Last();
+			check(ProjectedShadowInfo->bAllocated);
+		}
+	}
 
-				if (!ShadowMapAtlas.RenderTargets.IsValid() || GFastVRamConfig.bDirty)
-				{
-					GRenderTargetPool.FindFreeElement(RHICmdList, ShadowMapDesc2D, ShadowMapAtlas.RenderTargets.DepthTarget, TEXT("ShadowDepthAtlas"), true, ERenderTargetTransience::NonTransient);
-				}
+	for (int32 LayoutIndex = 0; LayoutIndex < Layouts.Num(); LayoutIndex++)
+	{
+		const FLayoutAndAssignedShadows& CurrentLayout = Layouts[LayoutIndex];
 
+		SortedShadowsForShadowDepthPass.ShadowMapAtlases.AddDefaulted();
+		FSortedShadowMapAtlas& ShadowMapAtlas = SortedShadowsForShadowDepthPass.ShadowMapAtlases.Last();
+
+		// Snap atlas sizes to powers of 2 to keep them more consistent for the render target pool
+		FIntPoint AtlasSize(
+			FMath::Min((uint32)MaxTextureSize.X, FMath::RoundUpToPowerOfTwo(CurrentLayout.TextureLayout.GetSizeX())),
+			FMath::Min((uint32)MaxTextureSize.Y, FMath::RoundUpToPowerOfTwo(CurrentLayout.TextureLayout.GetSizeY())));
+
+		if (CVarAlwaysAllocateMaxResolutionAtlases.GetValueOnRenderThread() != 0)
+		{
+			AtlasSize = MaxTextureSize;
+		}
+
+		FPooledRenderTargetDesc ShadowMapDesc2D = FPooledRenderTargetDesc::Create2DDesc(AtlasSize, PF_ShadowDepth, FClearValueBinding::DepthOne, TexCreate_None, TexCreate_DepthStencilTargetable | TexCreate_ShaderResource, false);
+		ShadowMapDesc2D.Flags |= GFastVRamConfig.ShadowPerObject;
+		GRenderTargetPool.FindFreeElement(RHICmdList, ShadowMapDesc2D, ShadowMapAtlas.RenderTargets.DepthTarget, TEXT("ShadowDepthAtlas"), true, ERenderTargetTransience::NonTransient);
+
+		for (int32 ShadowIndex = 0; ShadowIndex < CurrentLayout.Shadows.Num(); ShadowIndex++)
+		{
+			FProjectedShadowInfo* ProjectedShadowInfo = CurrentLayout.Shadows[ShadowIndex];
+
+			if (ProjectedShadowInfo->bAllocated)
+			{
 				ProjectedShadowInfo->RenderTargets.CopyReferencesFromRenderTargets(ShadowMapAtlas.RenderTargets);
 				ProjectedShadowInfo->SetupShadowDepthView(RHICmdList, this);
 				ShadowMapAtlas.Shadows.Add(ProjectedShadowInfo);
@@ -4672,16 +4719,6 @@ const TCHAR* GetCSMRenderTargetName(int32 ShadowMapIndex)
 	return **ShadowmapNames[ShadowMapIndex];
 }
 
-struct FLayoutAndAssignedShadows
-{
-	FLayoutAndAssignedShadows(int32 MaxTextureSize) :
-		TextureLayout(1, 1, MaxTextureSize, MaxTextureSize, false, ETextureLayoutAspectRatio::None, false)
-	{}
-
-	FTextureLayout TextureLayout;
-	TArray<FProjectedShadowInfo*, SceneRenderingAllocator> Shadows;
-};
-
 void FSceneRenderer::AllocateCSMDepthTargets(FRHICommandListImmediate& RHICmdList, const TArray<FProjectedShadowInfo*, SceneRenderingAllocator>& WholeSceneDirectionalShadows)
 {
 	if (WholeSceneDirectionalShadows.Num() > 0)
