Fix shadow jittering issues in forward shading
- Disabling Nanite if forward shading is enabled
- Nanite disabled early shadow path, making shadow projection use last frame (or nothing)
- Added ensure to detect shadows being sampled before being rendered again.

diff --git a/Engine/Source/Runtime/Renderer/Private/DeferredShadingRenderer.cpp b/Engine/Source/Runtime/Renderer/Private/DeferredShadingRenderer.cpp
--- a/Engine/Source/Runtime/Renderer/Private/DeferredShadingRenderer.cpp
+++ b/Engine/Source/Runtime/Renderer/Private/DeferredShadingRenderer.cpp
@@ -244,11 +244,13 @@ namespace Nanite
 	extern void ListStatFilters(FSceneRenderer* SceneRenderer);
 }
 
-inline bool IsNaniteEnabled()
+// Note this is used also in  ShadowDepthRendering.cpp, so can't be inlined
+bool IsNaniteEnabled(EShaderPlatform ShaderPlatform)
 {
 	static const auto EnableNaniteCVar = IConsoleManager::Get().FindConsoleVariable(TEXT("r.Nanite"));
-	static bool bNaniteSupported = DoesPlatformSupportNanite(GMaxRHIShaderPlatform);
-	return bNaniteSupported /*&& GRHISupportsAtomicUInt64*/ && EnableNaniteCVar->GetInt() > 0;
+	static bool bNaniteSupported = DoesPlatformSupportNanite(ShaderPlatform);
+	static bool bForwardShadingEnabled = IsForwardShadingEnabled(ShaderPlatform);
+	return bNaniteSupported /*&& GRHISupportsAtomicUInt64*/ && !bForwardShadingEnabled && EnableNaniteCVar->GetInt() > 0;
 }
 
 DECLARE_CYCLE_STAT(TEXT("PostInitViews FlushDel"), STAT_PostInitViews_FlushDel, STATGROUP_InitViews);
@@ -1530,7 +1532,7 @@ void FDeferredShadingSceneRenderer::CommitFinalPipelineState()
 
 	// Family pipeline state
 	{
-		FamilyPipelineState.Set(&FFamilyPipelineState::bNanite, IsNaniteEnabled()); // TODO: Should this respect ViewFamily.EngineShowFlags.NaniteMeshes?
+		FamilyPipelineState.Set(&FFamilyPipelineState::bNanite, IsNaniteEnabled(ShaderPlatform)); // TODO: Should this respect ViewFamily.EngineShowFlags.NaniteMeshes?
 
 		static const auto ICVarHZBOcc = IConsoleManager::Get().FindConsoleVariable(TEXT("r.HZBOcclusion"));
 		FamilyPipelineState.Set(&FFamilyPipelineState::bHZBOcclusion, ICVarHZBOcc->GetInt() != 0);	
@@ -1575,7 +1577,7 @@ void FDeferredShadingSceneRenderer::CommitFinalPipelineState()
 
 void FDeferredShadingSceneRenderer::Render(FRHICommandListImmediate& RHICmdList)
 {
-	static const bool bNaniteEnabled = IsNaniteEnabled() && ViewFamily.EngineShowFlags.NaniteMeshes;
+	static const bool bNaniteEnabled = IsNaniteEnabled(ShaderPlatform) && ViewFamily.EngineShowFlags.NaniteMeshes;
 
 	Scene->UpdateAllPrimitiveSceneInfos(RHICmdList, true);
 
diff --git a/Engine/Source/Runtime/Renderer/Private/SceneRendering.cpp b/Engine/Source/Runtime/Renderer/Private/SceneRendering.cpp
--- a/Engine/Source/Runtime/Renderer/Private/SceneRendering.cpp
+++ b/Engine/Source/Runtime/Renderer/Private/SceneRendering.cpp
@@ -2185,6 +2185,7 @@ FSceneRenderer::FSceneRenderer(const FSceneViewFamily* InViewFamily,FHitProxyCon
 ,	InstancedStereoWidth(0)
 ,	RootMark(nullptr)
 ,	FamilySize(0, 0)
+,	bShadowDepthRenderCompleted(false)
 {
 	check(Scene != NULL);
 
diff --git a/Engine/Source/Runtime/Renderer/Private/SceneRendering.h b/Engine/Source/Runtime/Renderer/Private/SceneRendering.h
--- a/Engine/Source/Runtime/Renderer/Private/SceneRendering.h
+++ b/Engine/Source/Runtime/Renderer/Private/SceneRendering.h
@@ -1966,6 +1966,7 @@ private:
 	/** Dump all UPrimitiveComponents in the Scene to a CSV file */
 	void DumpPrimitives(const FViewCommands& ViewCommands);
 #endif
+	bool bShadowDepthRenderCompleted;
 };
 
 /**
diff --git a/Engine/Source/Runtime/Renderer/Private/ShadowDepthRendering.cpp b/Engine/Source/Runtime/Renderer/Private/ShadowDepthRendering.cpp
--- a/Engine/Source/Runtime/Renderer/Private/ShadowDepthRendering.cpp
+++ b/Engine/Source/Runtime/Renderer/Private/ShadowDepthRendering.cpp
@@ -130,7 +130,7 @@ namespace Nanite
 	extern FString GetFilterNameForLight(const FLightSceneProxy* LightProxy);
 }
 
-extern bool IsNaniteEnabled();
+extern bool IsNaniteEnabled(EShaderPlatform ShaderPlatform);
 
 // Multiply PackedView.LODScale by return value when rendering Nanite shadows
 static float ComputeNaniteShadowsLODScaleFactor()
@@ -1514,7 +1514,7 @@ void FSceneRenderer::RenderShadowDepthMapAtlases(FRHICommandListImmediate& RHICm
 	bool bCanUseParallelDispatch = RHICmdList.IsImmediate() &&  // translucent shadows are draw on the render thread, using a recursive cmdlist (which is not immediate)
 		GRHICommandList.UseParallelAlgorithms() && CVarParallelShadows.GetValueOnRenderThread();
 
-	const bool bNaniteEnabled = IsNaniteEnabled() && ViewFamily.EngineShowFlags.NaniteMeshes;
+	const bool bNaniteEnabled = IsNaniteEnabled(ShaderPlatform) && ViewFamily.EngineShowFlags.NaniteMeshes;
 
 	for (int32 AtlasIndex = 0; AtlasIndex < SortedShadowsForShadowDepthPass.ShadowMapAtlases.Num(); AtlasIndex++)
 	{
@@ -1848,7 +1848,7 @@ void FSceneRenderer::RenderShadowDepthMaps(FRHICommandListImmediate& RHICmdList)
 
 	const bool bHasVSMShadows = SortedShadowsForShadowDepthPass.VirtualShadowMapShadows.Num() > 0;
 	const bool bHasVSMClipMaps = SortedShadowsForShadowDepthPass.VirtualShadowMapClipmaps.Num() > 0;
-	const bool bNaniteEnabled = IsNaniteEnabled() && ViewFamily.EngineShowFlags.NaniteMeshes;
+	const bool bNaniteEnabled = IsNaniteEnabled(ShaderPlatform) && ViewFamily.EngineShowFlags.NaniteMeshes;
 
 	if (bNaniteEnabled && (bHasVSMShadows || bHasVSMClipMaps))
 	{
@@ -2435,6 +2435,7 @@ void FSceneRenderer::RenderShadowDepthMaps(FRHICommandListImmediate& RHICmdList)
 	}
 
 	checkSlow(RHICmdList.IsOutsideRenderPass());
+	bShadowDepthRenderCompleted = true;
 }
 
 template<bool bRenderReflectiveShadowMap>
diff --git a/Engine/Source/Runtime/Renderer/Private/ShadowRendering.cpp b/Engine/Source/Runtime/Renderer/Private/ShadowRendering.cpp
--- a/Engine/Source/Runtime/Renderer/Private/ShadowRendering.cpp
+++ b/Engine/Source/Runtime/Renderer/Private/ShadowRendering.cpp
@@ -1622,6 +1622,8 @@ bool FDeferredShadingSceneRenderer::InjectReflectiveShadowMaps(FRHICommandListIm
 
 bool FSceneRenderer::RenderShadowProjections(FRHICommandListImmediate& RHICmdList, const FLightSceneInfo* LightSceneInfo, IPooledRenderTarget* ScreenShadowMaskTexture, IPooledRenderTarget* ScreenShadowMaskSubPixelTexture, bool bProjectingForForwardShading, bool bMobileModulatedProjections, const FHairStrandsVisibilityViews* InHairVisibilityViews)
 {
+	ensureMsgf(bShadowDepthRenderCompleted, TEXT("Shadow depth rendering was not done before shadow projections, this will cause severe shadow artifacts and indicates an engine bug (pass ordering)"));
+
 	FVisibleLightInfo& VisibleLightInfo = VisibleLightInfos[LightSceneInfo->Id];
 	FSceneRenderTargets& SceneContext = FSceneRenderTargets::Get(RHICmdList);
 
