Fix missing shader parameter assert in outline rendering due to inconsistent nanite usage checks

diff --git a/Engine/Source/Runtime/Renderer/Private/PostProcess/PostProcessSelectionOutline.cpp b/Engine/Source/Runtime/Renderer/Private/PostProcess/PostProcessSelectionOutline.cpp
--- a/Engine/Source/Runtime/Renderer/Private/PostProcess/PostProcessSelectionOutline.cpp
+++ b/Engine/Source/Runtime/Renderer/Private/PostProcess/PostProcessSelectionOutline.cpp
@@ -44,7 +44,7 @@ FScreenPassTexture AddSelectionOutlinePass(FRDGBuilder& GraphBuilder, const FVie
 	check(Inputs.SceneColor.IsValid());
 	check(Inputs.SceneDepth.IsValid());
 
-	const bool bNaniteEnabled = DoesPlatformSupportNanite(GMaxRHIShaderPlatform); // TODO: Respect r.Nanite
+	const bool bNaniteEnabled = NaniteRasterResults != nullptr;
 
 	RDG_EVENT_SCOPE(GraphBuilder, "EditorSelectionOutlines");
 
