Enable global distance field shader for Metal RHI (SM5).

#rnx

diff --git a/Engine/Source/Runtime/Renderer/Private/GlobalDistanceField.cpp b/Engine/Source/Runtime/Renderer/Private/GlobalDistanceField.cpp
--- a/Engine/Source/Runtime/Renderer/Private/GlobalDistanceField.cpp
+++ b/Engine/Source/Runtime/Renderer/Private/GlobalDistanceField.cpp
@@ -200,6 +200,11 @@ FAutoConsoleVariableRef CVarAOGlobalDistanceFieldAverageCulledObjectsPerPage(
 // Approximation of margin added to SDF objects during generation inside FMeshUtilities::GenerateSignedDistanceFieldVolumeData.
 constexpr float MESH_SDF_APPROX_MARGIN = 0.7f;
 
+bool ShouldCompileGlobalDistanceFieldShader(const FGlobalShaderPermutationParameters& Parameters)
+{
+	return IsFeatureLevelSupported(Parameters.Platform, ERHIFeatureLevel::SM5) && DoesPlatformSupportDistanceFieldAO(Parameters.Platform) && IsUsingDistanceFields(Parameters.Platform);
+}
+
 float GetMinMeshSDFRadius(float VoxelWorldSize)
 {
 	float MinRadius = GAOGlobalDistanceFieldMinMeshSDFRadius * (GAOGlobalDistanceFieldFastCameraMode ? 10.0f : 1.0f);
