Complete shadow maps: copy/resample non-nanite geometry instead of re-rendering it
- Refactor shadow depth rendering high level to clean up passes a bit
- Add back in SnapResolution to allow the complete shadow maps to snap consistently with the shadow that is their copy source (otherwise the projections differ)
- NOTE: Shadow bias for directional lights is technically dependent on resolution, and it is baked into the high resolution map. Thus the copy path has a minor difference in biasing here but it's not major concern for the uses of the complete shadow map.

diff --git a/Engine/Shaders/Private/VirtualShadowMaps/CopyCompleteShadows.usf b/Engine/Shaders/Private/VirtualShadowMaps/CopyCompleteShadows.usf
new file mode 100644
--- /dev/null
+++ b/Engine/Shaders/Private/VirtualShadowMaps/CopyCompleteShadows.usf
@@ -0,0 +1,18 @@
+// Copyright Epic Games, Inc. All Rights Reserved.
+
+#include "../Common.ush"
+
+Texture2D<float> SourceBuffer;
+float4 SourceScaleOffset;
+float4 DestInvTargetExtent;
+
+void CopyToCompleteShadowMapPS(
+	in float4 SvPosition	: SV_Position,
+	out float OutDepth		: SV_Depth
+)
+{
+	float2 DestPos = SvPosition.xy;
+	float2 SourcePos = (DestPos * SourceScaleOffset.xy + SourceScaleOffset.zw);
+	float Depth = SourceBuffer.Load(int3(SourcePos, 0));
+	OutDepth = Depth;
+}
diff --git a/Engine/Source/Runtime/Renderer/Private/SceneRendering.h b/Engine/Source/Runtime/Renderer/Private/SceneRendering.h
--- a/Engine/Source/Runtime/Renderer/Private/SceneRendering.h
+++ b/Engine/Source/Runtime/Renderer/Private/SceneRendering.h
@@ -1606,6 +1606,8 @@ struct FSortedShadowMaps
 
 	TArray<TSharedPtr<FVirtualShadowMapClipmap>, SceneRenderingAllocator> VirtualShadowMapClipmaps;
 
+	TArray<FSortedShadowMapAtlas,SceneRenderingAllocator> CompleteShadowMapAtlases;
+
 	void Release();
 
 	int64 ComputeMemorySize() const
@@ -1830,11 +1832,11 @@ protected:
 
 	void AllocateShadowDepthTargets(FRHICommandListImmediate& RHICmdList);
 	
-	void AllocateAtlasedShadowDepthTargets(FRHICommandListImmediate& RHICmdList, TArray<FProjectedShadowInfo*, SceneRenderingAllocator>& Shadows);
+	void AllocateAtlasedShadowDepthTargets(FRHICommandListImmediate& RHICmdList, TArray<FProjectedShadowInfo*, SceneRenderingAllocator>& Shadows, TArray<FSortedShadowMapAtlas,SceneRenderingAllocator>& OutAtlases);
 
 	void AllocateCachedSpotlightShadowDepthTargets(FRHICommandListImmediate& RHICmdList, TArray<FProjectedShadowInfo*, SceneRenderingAllocator>& CachedShadows);
 
-	void AllocateCSMDepthTargets(FRHICommandListImmediate& RHICmdList, const TArray<FProjectedShadowInfo*, SceneRenderingAllocator>& WholeSceneDirectionalShadows);
+	void AllocateCSMDepthTargets(FRHICommandListImmediate& RHICmdList, const TArray<FProjectedShadowInfo*, SceneRenderingAllocator>& WholeSceneDirectionalShadows, TArray<FSortedShadowMapAtlas,SceneRenderingAllocator>& OutAtlases);
 
 	void AllocateRSMDepthTargets(FRHICommandListImmediate& RHICmdList, const TArray<FProjectedShadowInfo*, SceneRenderingAllocator>& RSMShadows);
 
diff --git a/Engine/Source/Runtime/Renderer/Private/ShadowDepthRendering.cpp b/Engine/Source/Runtime/Renderer/Private/ShadowDepthRendering.cpp
--- a/Engine/Source/Runtime/Renderer/Private/ShadowDepthRendering.cpp
+++ b/Engine/Source/Runtime/Renderer/Private/ShadowDepthRendering.cpp
@@ -41,6 +41,7 @@
 #include "GPUScene.h"
 #include "SceneTextureReductions.h"
 #include "RendererModule.h"
+#include "PixelShaderUtils.h"
 #include "VirtualShadowMaps/VirtualShadowMapCacheManager.h"
 #include "VirtualShadowMaps/VirtualShadowMapClipmap.h"
 
@@ -49,13 +50,6 @@ DECLARE_GPU_DRAWCALL_STAT_NAMED(ShadowDepths, TEXT("Shadow Depths"));
 IMPLEMENT_GLOBAL_SHADER_PARAMETER_STRUCT(FShadowDepthPassUniformParameters, "ShadowDepthPass");
 IMPLEMENT_GLOBAL_SHADER_PARAMETER_STRUCT(FMobileShadowDepthPassUniformParameters, "MobileShadowDepthPass");
 
-TAutoConsoleVariable<int32> CVarNaniteShadows(
-	TEXT("r.Shadow.Nanite"),
-	1,
-	TEXT("Enables shadows from Nanite meshes."),
-	ECVF_RenderThreadSafe);
-
-
 template<bool bUsingVertexLayers = false>
 class TScreenVSForGS : public FScreenVS
 {
@@ -94,6 +88,12 @@ static TAutoConsoleVariable<int32> CVarShadowForceSerialSingleRenderPass(
 	TEXT("Force Serial shadow passes to render in 1 pass."),
 	ECVF_RenderThreadSafe);
 
+TAutoConsoleVariable<int32> CVarNaniteShadows(
+	TEXT("r.Shadow.Nanite"),
+	1,
+	TEXT("Enables shadows from Nanite meshes."),
+	ECVF_RenderThreadSafe);
+
 static TAutoConsoleVariable<int32> CVarNaniteShadowsUseHZB(
 	TEXT("r.Shadow.NaniteUseHZB"),
 	1,
@@ -1499,6 +1499,178 @@ FRHIGPUMask FSceneRenderer::GetGPUMaskForShadow(FProjectedShadowInfo* ProjectedS
 }
 #endif // WITH_MGPU
 
+
+static void RenderShadowDepthAtlasNanite(
+	FRDGBuilder& GraphBuilder,
+	const FScene& Scene,
+	const FSortedShadowMapAtlas& ShadowMapAtlas)
+{
+	const FIntPoint AtlasSize = ShadowMapAtlas.RenderTargets.DepthTarget->GetDesc().Extent;
+
+	bool bWantsNearClip = false;
+	bool bWantsNoNearClip = false;
+	TArray<Nanite::FPackedView, SceneRenderingAllocator> PackedViews;
+	TArray<Nanite::FPackedView, SceneRenderingAllocator> PackedViewsNoNearClip;
+	TArray<FProjectedShadowInfo*, SceneRenderingAllocator> ShadowsToEmit;
+	for (int32 ShadowIndex = 0; ShadowIndex < ShadowMapAtlas.Shadows.Num(); ShadowIndex++)
+	{
+		FProjectedShadowInfo* ProjectedShadowInfo = ShadowMapAtlas.Shadows[ShadowIndex];
+
+		// TODO: We avoid rendering Nanite geometry into both movable AND static cached shadows, but has a side effect
+		// that if there is *only* a movable cached shadow map (and not static), it won't rendering anything.
+		// Logic around Nanite and the cached shadows is fuzzy in a bunch of places and the whole thing needs some rethinking
+		// so leaving this like this for now as it is unlikely to happen in realistic scenes.
+		if (!ProjectedShadowInfo->bNaniteGeometry ||
+			ProjectedShadowInfo->CacheMode == SDCM_MovablePrimitivesOnly)
+		{
+			continue;
+		}
+
+		const FIntRect AtlasViewRect = ProjectedShadowInfo->GetViewRectForView();
+
+		Nanite::FPackedViewParams Initializer;
+		Initializer.ViewMatrices = ProjectedShadowInfo->ShadowDepthView->ViewMatrices;
+		Initializer.PrevViewMatrices = Initializer.ViewMatrices;
+		Initializer.ViewRect = AtlasViewRect;
+		Initializer.RasterContextSize = AtlasSize;
+		Initializer.LODScaleFactor = ComputeNaniteShadowsLODScaleFactor();
+
+		// Orthographic shadow projections want depth clamping rather than clipping
+		if (ProjectedShadowInfo->ShouldClampToNearPlane())
+		{
+			PackedViewsNoNearClip.Add(Nanite::CreatePackedView(Initializer));
+		}
+		else
+		{
+			PackedViews.Add(Nanite::CreatePackedView(Initializer));
+		}
+
+
+		ShadowsToEmit.Add(ProjectedShadowInfo);
+	}
+
+	if (PackedViews.Num() > 0 || PackedViewsNoNearClip.Num() > 0)
+	{
+		RDG_EVENT_SCOPE(GraphBuilder, "Nanite Shadows");
+
+		// Need separate passes for near clip on/off currently
+		const bool bSupportsMultiplePasses = (PackedViews.Num() > 0 && PackedViewsNoNearClip.Num() > 0);
+		const bool bPrimaryContext = false;
+
+		// NOTE: Rendering into an atlas like this is not going to work properly with HZB, but we are not currently using HZB here.
+		// It might be worthwhile going through the virtual SM rendering path even for "dense" cases even just for proper handling of all the details.
+		FIntRect FullAtlasViewRect(FIntPoint(0, 0), AtlasSize);
+		const bool bUpdateStreaming = CVarNaniteShadowsUpdateStreaming.GetValueOnRenderThread() != 0;
+		Nanite::FCullingContext CullingContext = Nanite::InitCullingContext(GraphBuilder, Scene, nullptr, FullAtlasViewRect, true, bUpdateStreaming, bSupportsMultiplePasses, false, bPrimaryContext);
+		Nanite::FRasterContext RasterContext = Nanite::InitRasterContext(GraphBuilder, AtlasSize, Nanite::EOutputBufferMode::DepthOnly);
+
+		if (PackedViews.Num() > 0)
+		{
+			Nanite::FRasterState RasterState;
+			RasterState.bNearClip = true;
+
+			Nanite::CullRasterize(
+				GraphBuilder,
+				Scene,
+				PackedViews,
+				CullingContext,
+				RasterContext,
+				RasterState,
+				nullptr,	// InstanceDraws
+				false		// bExtractStats TODO
+			);
+		}
+
+		if (PackedViewsNoNearClip.Num() > 0)
+		{
+			Nanite::FRasterState RasterState;
+			RasterState.bNearClip = false;
+
+			Nanite::CullRasterize(
+				GraphBuilder,
+				Scene,
+				PackedViewsNoNearClip,
+				CullingContext,
+				RasterContext,
+				RasterState,
+				nullptr,	// InstanceDraws
+				false		// bExtractStats TODO
+			);
+		}
+
+		FRDGTextureRef ShadowMap = GraphBuilder.RegisterExternalTexture(ShadowMapAtlas.RenderTargets.DepthTarget, TEXT("DepthBuffer"));
+		for (FProjectedShadowInfo* ProjectedShadowInfo : ShadowsToEmit)
+		{
+			const FIntRect AtlasViewRect = ProjectedShadowInfo->GetViewRectForView();
+
+			Nanite::EmitShadowMap(
+				GraphBuilder,
+				RasterContext,
+				ShadowMap,
+				AtlasViewRect,
+				AtlasViewRect.Min,
+				ProjectedShadowInfo->ShadowDepthView->ViewMatrices.GetProjectionMatrix(),
+				ProjectedShadowInfo->GetShaderDepthBias(),
+				ProjectedShadowInfo->bDirectionalLight
+			);
+		}
+	}
+}
+
+class FCopyToCompleteShadowMapPS : public FGlobalShader
+{
+	DECLARE_GLOBAL_SHADER(FCopyToCompleteShadowMapPS);
+	SHADER_USE_PARAMETER_STRUCT(FCopyToCompleteShadowMapPS, FGlobalShader);
+
+	BEGIN_SHADER_PARAMETER_STRUCT(FParameters, )
+		SHADER_PARAMETER(FVector4, SourceScaleOffset)
+		SHADER_PARAMETER_RDG_TEXTURE(Texture2D<uint>, SourceBuffer)
+		RENDER_TARGET_BINDING_SLOTS()
+	END_SHADER_PARAMETER_STRUCT()
+
+	static bool ShouldCompilePermutation(const FGlobalShaderPermutationParameters& Parameters)
+	{
+		return IsFeatureLevelSupported(Parameters.Platform, ERHIFeatureLevel::SM5);
+	}
+};
+IMPLEMENT_GLOBAL_SHADER(FCopyToCompleteShadowMapPS, "/Engine/Private/VirtualShadowMaps/CopyCompleteShadows.usf", "CopyToCompleteShadowMapPS", SF_Pixel);
+
+static void CopyToCompleteShadowMap(
+	FRDGBuilder& GraphBuilder,
+	const FRDGTextureRef SourceBuffer,
+	const FRDGTextureRef DestBuffer,
+	const FIntRect& SourceRect,
+	const FIntRect& DestRect,
+	ERenderTargetLoadAction LoadAction
+	)
+{
+	FVector4 SourceScaleOffset;
+	SourceScaleOffset.X = static_cast<float>(SourceRect.Width() ) / DestRect.Width();
+	SourceScaleOffset.Y = static_cast<float>(SourceRect.Height()) / DestRect.Height();
+	SourceScaleOffset.Z = SourceRect.Min.X - (SourceScaleOffset.X * DestRect.Min.X);
+	SourceScaleOffset.W = SourceRect.Min.Y - (SourceScaleOffset.Y * DestRect.Min.Y);
+
+	auto* PassParameters = GraphBuilder.AllocParameters<FCopyToCompleteShadowMapPS::FParameters>();
+	PassParameters->SourceBuffer = SourceBuffer;
+	PassParameters->SourceScaleOffset = SourceScaleOffset;
+	PassParameters->RenderTargets.DepthStencil = FDepthStencilBinding(DestBuffer, LoadAction, FExclusiveDepthStencil::DepthWrite_StencilNop);
+
+	auto ShaderMap = GetGlobalShaderMap(GMaxRHIFeatureLevel);
+	auto PixelShader = ShaderMap->GetShader<FCopyToCompleteShadowMapPS>();
+
+	FPixelShaderUtils::AddFullscreenPass(
+		GraphBuilder,
+		ShaderMap,
+		RDG_EVENT_NAME("CopyToCompleteShadowMap"),
+		PixelShader,
+		PassParameters,
+		DestRect,
+		nullptr,
+		nullptr,
+		TStaticDepthStencilState<true, CF_Always>::GetRHI()
+	);
+}
+
 void FSceneRenderer::RenderShadowDepthMapAtlases(FRHICommandListImmediate& RHICmdList)
 {
 	check(RHICmdList.IsOutsideRenderPass());
@@ -1512,7 +1684,10 @@ void FSceneRenderer::RenderShadowDepthMapAtlases(FRHICommandListImmediate& RHICm
 	bool bCanUseParallelDispatch = RHICmdList.IsImmediate() &&  // translucent shadows are draw on the render thread, using a recursive cmdlist (which is not immediate)
 		GRHICommandList.UseParallelAlgorithms() && CVarParallelShadows.GetValueOnRenderThread();
 
-	const bool bNaniteEnabled = UseNanite(ShaderPlatform) && ViewFamily.EngineShowFlags.NaniteMeshes;
+	const bool bNaniteEnabled = 
+		UseNanite(ShaderPlatform) &&
+		ViewFamily.EngineShowFlags.NaniteMeshes &&
+		CVarNaniteShadows.GetValueOnRenderThread() != 0;
 
 	for (int32 AtlasIndex = 0; AtlasIndex < SortedShadowsForShadowDepthPass.ShadowMapAtlases.Num(); AtlasIndex++)
 	{
@@ -1701,126 +1876,61 @@ void FSceneRenderer::RenderShadowDepthMapAtlases(FRHICommandListImmediate& RHICm
 			CurrentLightForDrawEvent = NULL;
 		}
 
-		if (bNaniteEnabled && CVarNaniteShadows.GetValueOnRenderThread())
+		if (bNaniteEnabled)
 		{
-			bool bWantsNearClip = false;
-			bool bWantsNoNearClip = false;
-			TArray<Nanite::FPackedView, SceneRenderingAllocator> PackedViews;
-			TArray<Nanite::FPackedView, SceneRenderingAllocator> PackedViewsNoNearClip;
-			TArray<FProjectedShadowInfo*, SceneRenderingAllocator> ShadowsToEmit;
-			for (int32 ShadowIndex = 0; ShadowIndex < ShadowMapAtlas.Shadows.Num(); ShadowIndex++)
-			{
-				FProjectedShadowInfo* ProjectedShadowInfo = ShadowMapAtlas.Shadows[ShadowIndex];
+			FRDGBuilder GraphBuilder(RHICmdList);
+			RenderShadowDepthAtlasNanite(GraphBuilder, *Scene, ShadowMapAtlas);
+			GraphBuilder.Execute();
+		}
 
-				// TODO: We avoid rendering Nanite geometry into both movable AND static cached shadows, but has a side effect
-				// that if there is *only* a movable cached shadow map (and not static), it won't rendering anything.
-				// Logic around Nanite and the cached shadows is fuzzy in a bunch of places and the whole thing needs some rethinking
-				// so leaving this like this for now as it is unlikely to happen in realistic scenes.
-				if (!ProjectedShadowInfo->bNaniteGeometry ||
-					ProjectedShadowInfo->CacheMode == SDCM_MovablePrimitivesOnly)
-				{
-					continue;
-				}
+		RHICmdList.TransitionResource(EResourceTransitionAccess::EReadable, RenderTarget.TargetableTexture);
+	}
 
-				const FIntRect AtlasViewRect = ProjectedShadowInfo->GetViewRectForView();
+	// Copy/resample shadow maps into "complete" shadow maps and add Nanite geometry
+	if (SortedShadowsForShadowDepthPass.CompleteShadowMapAtlases.Num() > 0)
+	{
+		FRDGBuilder GraphBuilder(RHICmdList);
 
-				Nanite::FPackedViewParams Initializer;
-				Initializer.ViewMatrices = ProjectedShadowInfo->ShadowDepthView->ViewMatrices;
-				Initializer.PrevViewMatrices = Initializer.ViewMatrices;
-				Initializer.ViewRect = AtlasViewRect;
-				Initializer.RasterContextSize = AtlasSize;
-				Initializer.LODScaleFactor = ComputeNaniteShadowsLODScaleFactor();
+		for (const FSortedShadowMapAtlas& ShadowMapAtlas : SortedShadowsForShadowDepthPass.CompleteShadowMapAtlases)
+		{
+			FSceneRenderTargetItem& RenderTarget = ShadowMapAtlas.RenderTargets.DepthTarget->GetRenderTargetItem();
+			FIntPoint AtlasSize = ShadowMapAtlas.RenderTargets.DepthTarget->GetDesc().Extent;
+		
+			FRDGTextureRef DestShadowMap = GraphBuilder.RegisterExternalTexture(ShadowMapAtlas.RenderTargets.DepthTarget, TEXT("DepthBuffer"));
 
-				// Orthographic shadow projections want depth clamping rather than clipping
-				if (ProjectedShadowInfo->ShouldClampToNearPlane())
-				{
-					PackedViewsNoNearClip.Add(Nanite::CreatePackedView(Initializer));
-				}
-				else
-				{
-					PackedViews.Add(Nanite::CreatePackedView(Initializer));
-				}
-				
+			RDG_EVENT_SCOPE(GraphBuilder, "Complete Atlas %ux%u", AtlasSize.X, AtlasSize.Y);
 
-				ShadowsToEmit.Add(ProjectedShadowInfo);
-			}
-				
-			if (PackedViews.Num() > 0 || PackedViewsNoNearClip.Num() > 0)
+			bool bCleared = false;
+			for (int32 ShadowIndex = 0; ShadowIndex < ShadowMapAtlas.Shadows.Num(); ShadowIndex++)
 			{
-				FRDGBuilder GraphBuilder(RHICmdList);
-				FRDGTextureRef ShadowMap = GraphBuilder.RegisterExternalTexture(ShadowMapAtlas.RenderTargets.DepthTarget, TEXT("DepthBuffer"));
+				FProjectedShadowInfo* ProjectedShadowInfo = ShadowMapAtlas.Shadows[ShadowIndex];
+				FProjectedShadowInfo* SourceShadowInfo = ProjectedShadowInfo->CompleteShadowMapCopySource;
 
+				check(SourceShadowInfo);
+				if (SourceShadowInfo && SourceShadowInfo->RenderTargets.DepthTarget != nullptr)
 				{
-					RDG_EVENT_SCOPE(GraphBuilder, "Nanite Shadows");
-
-					// Need separate passes for near clip on/off currently
-					const bool bSupportsMultiplePasses = (PackedViews.Num() > 0 && PackedViewsNoNearClip.Num() > 0);
-					const bool bPrimaryContext = false;
-
-					// NOTE: Rendering into an atlas like this is not going to work properly with HZB, but we are not currently using HZB here.
-					// It might be worthwhile going through the virtual SM rendering path even for "dense" cases even just for proper handling of all the details.
-					FIntRect FullAtlasViewRect(FIntPoint(0, 0), AtlasSize);
-					const bool bUpdateStreaming = CVarNaniteShadowsUpdateStreaming.GetValueOnRenderThread() != 0;
-					Nanite::FCullingContext CullingContext = Nanite::InitCullingContext(GraphBuilder, *Scene, nullptr, FullAtlasViewRect, true, bUpdateStreaming, bSupportsMultiplePasses, false, bPrimaryContext);
-					Nanite::FRasterContext RasterContext = Nanite::InitRasterContext(GraphBuilder, AtlasSize, Nanite::EOutputBufferMode::DepthOnly);
-							
-					
-					if (PackedViews.Num() > 0)
-					{
-						Nanite::FRasterState RasterState;
-						RasterState.bNearClip = true;
-
-						Nanite::CullRasterize(
-							GraphBuilder,
-							*Scene,
-							PackedViews,
-							CullingContext,
-							RasterContext,
-							RasterState,
-							nullptr,	// InstanceDraws
-							false		// bExtractStats TODO
-						);
-					}
-
-					if (PackedViewsNoNearClip.Num() > 0)
-					{
-						Nanite::FRasterState RasterState;
-						RasterState.bNearClip = false;
-
-						Nanite::CullRasterize(
-							GraphBuilder,
-							*Scene,
-							PackedViewsNoNearClip,
-							CullingContext,
-							RasterContext,
-							RasterState,
-							nullptr,	// InstanceDraws
-							false		// bExtractStats TODO
-						);
-					}
-
-					for (FProjectedShadowInfo* ProjectedShadowInfo : ShadowsToEmit)
-					{
-						const FIntRect AtlasViewRect = ProjectedShadowInfo->GetViewRectForView();
-
-						Nanite::EmitShadowMap(
-							GraphBuilder,
-							RasterContext,
-							ShadowMap,
-							AtlasViewRect,
-							AtlasViewRect.Min,
-							ProjectedShadowInfo->ShadowDepthView->ViewMatrices.GetProjectionMatrix(),
-							ProjectedShadowInfo->GetShaderDepthBias(),
-							ProjectedShadowInfo->bDirectionalLight
-							);
-					}
+					FRDGTextureRef SourceShadowMap = GraphBuilder.RegisterExternalTexture(SourceShadowInfo->RenderTargets.DepthTarget, TEXT("SourceDepthBuffer"));
+					const FIntRect SourceRect = SourceShadowInfo->GetViewRectForView();
+					const FIntRect DestRect = ProjectedShadowInfo->GetViewRectForView();
+					CopyToCompleteShadowMap(GraphBuilder, SourceShadowMap, DestShadowMap, SourceRect, DestRect, bCleared ? ERenderTargetLoadAction::ELoad : ERenderTargetLoadAction::EClear);
+					bCleared = true;
 				}
+			}
 
-				GraphBuilder.Execute();
+			if (!bCleared)
+			{
+				// If nothing cleared it, ensure it's done before nanite rendering at least
+				AddClearDepthStencilPass(GraphBuilder, DestShadowMap, true, 1.0f, false, 0);
+				bCleared = true;
+			}
+
+			if (bNaniteEnabled)
+			{
+				RenderShadowDepthAtlasNanite(GraphBuilder, *Scene, ShadowMapAtlas);
 			}
 		}
 
-		RHICmdList.TransitionResource(EResourceTransitionAccess::EReadable, RenderTarget.TargetableTexture);
+		GraphBuilder.Execute();
 	}
 }
 
diff --git a/Engine/Source/Runtime/Renderer/Private/ShadowRendering.h b/Engine/Source/Runtime/Renderer/Private/ShadowRendering.h
--- a/Engine/Source/Runtime/Renderer/Private/ShadowRendering.h
+++ b/Engine/Source/Runtime/Renderer/Private/ShadowRendering.h
@@ -373,6 +373,9 @@ public:
 	/** Defaults to false, used to flag a fallback map. */
 	uint32 bCompleteShadowMap : 1;
 
+	/** Regular shadow map to resample for use as a complete shadow map */
+	FProjectedShadowInfo *CompleteShadowMapCopySource;
+
 	/** */
 	FVirtualShadowMap *VirtualShadowMap;
 
@@ -422,6 +425,8 @@ public:
 		const FWholeSceneProjectedShadowInitializer& Initializer,
 		uint32 InResolutionX,
 		uint32 InResolutionY,
+		uint32 InSnapResolutionX,
+		uint32 InSnapResolutionY,
 		uint32 InBorderSize,
 		bool bInReflectiveShadowMap,
 		TSharedPtr<FVirtualShadowMapCacheEntry> VirtualSmCacheEntry = TSharedPtr<FVirtualShadowMapCacheEntry>()
diff --git a/Engine/Source/Runtime/Renderer/Private/ShadowSetup.cpp b/Engine/Source/Runtime/Renderer/Private/ShadowSetup.cpp
--- a/Engine/Source/Runtime/Renderer/Private/ShadowSetup.cpp
+++ b/Engine/Source/Runtime/Renderer/Private/ShadowSetup.cpp
@@ -262,9 +262,16 @@ static TAutoConsoleVariable<int32> CVarCompleteShadowMapResolution(
 	ECVF_RenderThreadSafe
 );
 
+static TAutoConsoleVariable<int32> CVarCopyToCompleteShadowMaps(
+	TEXT("r.Shadow.CopyToCompleteShadowMaps"),
+	1,
+	TEXT("Copy non-Nanite geometry from high resolution shadow map instead of re-rendering it into the complete shadow map."),
+	ECVF_RenderThreadSafe
+);
+
 static TAutoConsoleVariable<int32> CVarAlwaysAllocateMaxResolutionAtlases(
 	TEXT("r.Shadow.AlwaysAllocateMaxResolutionAtlases"),
-	1,
+	0,
 	TEXT("If enabled, will always allocate shadow map atlases at the maximum resolution rather than trimming unused space. This will waste more memory but can possibly reduce render target pool fragmentation and thrash."),
 	ECVF_RenderThreadSafe
 );
@@ -517,6 +524,7 @@ FProjectedShadowInfo::FProjectedShadowInfo()
 	, bNaniteGeometry(true)
 	, bIncludeInScreenSpaceShadowMask(true)
 	, bCompleteShadowMap(false)
+	, CompleteShadowMapCopySource(nullptr)
 	, VirtualShadowMap(nullptr)
 	, PerObjectShadowFadeStart(WORLD_MAX)
 	, InvPerObjectShadowFadeLength(0.0f)
@@ -684,6 +692,8 @@ void FProjectedShadowInfo::SetupWholeSceneProjection(
 	const FWholeSceneProjectedShadowInitializer& Initializer,
 	uint32 InResolutionX,
 	uint32 InResolutionY,
+	uint32 InSnapResolutionX,
+	uint32 InSnapResolutionY,
 	uint32 InBorderSize,
 	bool bInReflectiveShadowMap,
 	TSharedPtr<FVirtualShadowMapCacheEntry> VirtualSmCacheEntry)
@@ -742,8 +752,8 @@ void FProjectedShadowInfo::SetupWholeSceneProjection(
 		// This corresponds to the maximum kernel filter size used by subsurface shadows in ShadowProjectionPixelShader.usf
 		const int32 MaxDownsampleFactor = 4;
 		// Determine the distance necessary to snap the shadow's position to the nearest texel
-		const float SnapX = FMath::Fmod(TransformedPosition.X, 2.0f * MaxDownsampleFactor / InResolutionX);
-		const float SnapY = FMath::Fmod(TransformedPosition.Y, 2.0f * MaxDownsampleFactor / InResolutionY);
+		const float SnapX = FMath::Fmod(TransformedPosition.X, 2.0f * MaxDownsampleFactor / InSnapResolutionX);
+		const float SnapY = FMath::Fmod(TransformedPosition.Y, 2.0f * MaxDownsampleFactor / InSnapResolutionY);
 		// Snap the shadow's position and transform it back into world space
 		// This snapping prevents sub-texel camera movements which removes view dependent aliasing from the final shadow result
 		// This only maintains stable shadows under camera translation and rotation
@@ -3284,6 +3294,8 @@ void FSceneRenderer::CreateWholeSceneProjectedShadow(
 							ProjectedShadowInitializer,
 							VirtualSizeX,
 							VirtualSizeY,
+							VirtualSizeX,
+							VirtualSizeY,
 							0, // no border
 							false,	// no RSM
 							VirtualSmCacheEntry
@@ -3291,41 +3303,6 @@ void FSceneRenderer::CreateWholeSceneProjectedShadow(
 
 						VisibleLightInfo.AllProjectedShadows.Add(ProjectedShadowInfo);
 					}
-
-					// Also set up "complete" shadow map (one with lower resolution and containing both nanite and non-nanite)
-					const int32 CompleteShadowMapSize = CVarCompleteShadowMapResolution.GetValueOnRenderThread() - ShadowBorder * 2;
-					if (CompleteShadowMapSize > 0)
-					{
-						// Create the projected shadow info.
-						FProjectedShadowInfo* CompleteProjectedShadowInfo = new(FMemStack::Get(), 1, 16) FProjectedShadowInfo;
-						// Add to remember-to-call-dtor list
-						VisibleLightInfo.MemStackProjectedShadows.Add(CompleteProjectedShadowInfo);
-
-						// Rescale size to fit whole virtual SM but keeping aspect ratio
-						int32 CompleteShadowMapSizeX = SizeX >= SizeY ? CompleteShadowMapSize : (CompleteShadowMapSize * SizeX) / SizeY;
-						int32 CompleteShadowMapSizeY = SizeY >= SizeX ? CompleteShadowMapSize : (CompleteShadowMapSize * SizeY) / SizeX;
-
-						// Set up projection as per normal
-						CompleteProjectedShadowInfo->SetupWholeSceneProjection(
-							LightSceneInfo,
-							NULL,
-							ProjectedShadowInitializer,
-							CompleteShadowMapSizeX,
-							CompleteShadowMapSizeY,
-							ShadowBorder,
-							false	// no RSM
-							);
-
-						CompleteProjectedShadowInfo->bCompleteShadowMap = true;
-						CompleteProjectedShadowInfo->bIncludeInScreenSpaceShadowMask = false;
-
-						bool bContainsNaniteSubjects = false;
-						AddInteractingPrimitives(LightSceneInfo->GetDynamicInteractionOftenMovingPrimitiveList(false), CompleteProjectedShadowInfo, bContainsNaniteSubjects);
-						AddInteractingPrimitives(LightSceneInfo->GetDynamicInteractionStaticPrimitiveList(false), CompleteProjectedShadowInfo, bContainsNaniteSubjects);
-
-						VisibleLightInfo.AllProjectedShadows.Add(CompleteProjectedShadowInfo);
-						VisibleLightInfo.CompleteProjectedShadows.Add(CompleteProjectedShadowInfo);
-					}
 				}
 				
 				for (int32 CacheModeIndex = 0; CacheModeIndex < NumShadowMaps; CacheModeIndex++)
@@ -3339,6 +3316,8 @@ void FSceneRenderer::CreateWholeSceneProjectedShadow(
 						ProjectedShadowInitializer,
 						SizeX,
 						SizeY,
+						SizeX,
+						SizeY,
 						ShadowBorder,
 						false	// no RSM
 						);
@@ -3401,6 +3380,49 @@ void FSceneRenderer::CreateWholeSceneProjectedShadow(
 					if (bNeedsVirtualShadowMap)
 					{
 						ProjectedShadowInfo->bNaniteGeometry = false;
+
+						// If this is not a static shadow cache, also set up "complete" shadow map (one with lower resolution and containing both nanite and non-nanite)
+						const int32 CompleteShadowMapSize = CVarCompleteShadowMapResolution.GetValueOnRenderThread() - ShadowBorder * 2;
+						if (CacheMode[CacheModeIndex] != SDCM_StaticPrimitivesOnly && CompleteShadowMapSize > 0)
+						{
+							// Create the projected shadow info.
+							FProjectedShadowInfo* CompleteProjectedShadowInfo = new(FMemStack::Get(), 1, 16) FProjectedShadowInfo;
+							// Add to remember-to-call-dtor list
+							VisibleLightInfo.MemStackProjectedShadows.Add(CompleteProjectedShadowInfo);
+
+							// Rescale size to fit whole virtual SM but keeping aspect ratio
+							int32 CompleteShadowMapSizeX = SizeX >= SizeY ? CompleteShadowMapSize : (CompleteShadowMapSize * SizeX) / SizeY;
+							int32 CompleteShadowMapSizeY = SizeY >= SizeX ? CompleteShadowMapSize : (CompleteShadowMapSize * SizeY) / SizeX;
+
+							// Never make it larger than the original shadow map
+							CompleteShadowMapSizeX = FMath::Min(CompleteShadowMapSizeX, SizeX);
+							CompleteShadowMapSizeY = FMath::Min(CompleteShadowMapSizeY, SizeY);
+
+							// Set up projection as per normal
+							CompleteProjectedShadowInfo->SetupWholeSceneProjection(
+								LightSceneInfo,
+								NULL,
+								ProjectedShadowInitializer,
+								CompleteShadowMapSizeX,
+								CompleteShadowMapSizeY,
+								// Snap to original shadow map resolution
+								ProjectedShadowInfo->ResolutionX,
+								ProjectedShadowInfo->ResolutionY,
+								ShadowBorder,
+								false	// no RSM
+							);
+
+							CompleteProjectedShadowInfo->bCompleteShadowMap = true;
+							CompleteProjectedShadowInfo->bIncludeInScreenSpaceShadowMask = false;
+							CompleteProjectedShadowInfo->CompleteShadowMapCopySource = ProjectedShadowInfo;
+
+							bool bContainsNaniteSubjects = false;
+							AddInteractingPrimitives(LightSceneInfo->GetDynamicInteractionOftenMovingPrimitiveList(false), CompleteProjectedShadowInfo, bContainsNaniteSubjects);
+							AddInteractingPrimitives(LightSceneInfo->GetDynamicInteractionStaticPrimitiveList(false), CompleteProjectedShadowInfo, bContainsNaniteSubjects);
+
+							VisibleLightInfo.AllProjectedShadows.Add(CompleteProjectedShadowInfo);
+							VisibleLightInfo.CompleteProjectedShadows.Add(CompleteProjectedShadowInfo);
+						}
 					}
 
 					bool bContainsNaniteSubjects = false;
@@ -4191,6 +4213,8 @@ void FSceneRenderer::AddViewDependentWholeSceneShadowsForView(
 						ProjectedShadowInitializer,
 						ShadowBufferResolution.X,
 						ShadowBufferResolution.Y,
+						ShadowBufferResolution.X,
+						ShadowBufferResolution.Y,
 						ShadowBorder,
 						false	// no RSM
 					);
@@ -4225,12 +4249,16 @@ void FSceneRenderer::AddViewDependentWholeSceneShadowsForView(
 								ProjectedShadowInitializer,
 								CompleteShadowMapSizeX,
 								CompleteShadowMapSizeY,
+								// Snap consistently with the original resolution shadow map
+								ProjectedShadowInfo->ResolutionX,
+								ProjectedShadowInfo->ResolutionY,
 								ShadowBorder,
 								false	// no RSM
 								);
 
 							CompleteProjectedShadowInfo->bCompleteShadowMap = true;
 							CompleteProjectedShadowInfo->bIncludeInScreenSpaceShadowMask = false;
+							CompleteProjectedShadowInfo->CompleteShadowMapCopySource = ProjectedShadowInfo;
 
 							VisibleLightInfo.AllProjectedShadows.Add(CompleteProjectedShadowInfo);
 							VisibleLightInfo.CompleteProjectedShadows.Add(CompleteProjectedShadowInfo);
@@ -4299,6 +4327,8 @@ void FSceneRenderer::AddViewDependentWholeSceneShadowsForView(
 							ProjectedShadowInitializer,
 							ShadowBufferResolution,
 							ShadowBufferResolution,
+							ShadowBufferResolution,
+							ShadowBufferResolution,
 							0,
 							true);		// RSM
 
@@ -4335,6 +4365,8 @@ void FSceneRenderer::AllocateShadowDepthTargets(FRHICommandListImmediate& RHICmd
 	// Cubemaps, can't be atlased
 	TArray<FProjectedShadowInfo*, SceneRenderingAllocator> WholeScenePointShadows;
 
+	bool bCopyToCompleteShadowMaps = CVarCopyToCompleteShadowMaps.GetValueOnRenderThread() != 0;
+
 	for (TSparseArray<FLightSceneInfoCompact>::TConstIterator LightIt(Scene->Lights); LightIt; ++LightIt)
 	{
 		const FLightSceneInfoCompact& LightSceneInfoCompact = *LightIt;
@@ -4502,8 +4534,9 @@ void FSceneRenderer::AllocateShadowDepthTargets(FRHICommandListImmediate& RHICmd
 		VisibleLightInfo.RSMsToProject.Sort(FCompareFProjectedShadowInfoBySplitIndex());
 		VisibleLightInfo.CompleteProjectedShadows.Sort(FCompareFProjectedShadowInfoBySplitIndex());
 
-		AllocateCSMDepthTargets(RHICmdList, WholeSceneDirectionalShadows);
-		AllocateCSMDepthTargets(RHICmdList, WholeSceneCompleteDirectionalShadows);
+		AllocateCSMDepthTargets(RHICmdList, WholeSceneDirectionalShadows, SortedShadowsForShadowDepthPass.ShadowMapAtlases);
+		AllocateCSMDepthTargets(RHICmdList, WholeSceneCompleteDirectionalShadows,
+			bCopyToCompleteShadowMaps ? SortedShadowsForShadowDepthPass.CompleteShadowMapAtlases : SortedShadowsForShadowDepthPass.ShadowMapAtlases);
 	}
 
 	if (CachedPreShadows.Num() > 0)
@@ -4532,8 +4565,9 @@ void FSceneRenderer::AllocateShadowDepthTargets(FRHICommandListImmediate& RHICmd
 	AllocateOnePassPointLightDepthTargets(RHICmdList, WholeScenePointShadows);
 	AllocateRSMDepthTargets(RHICmdList, RSMShadows);
 	AllocateCachedSpotlightShadowDepthTargets(RHICmdList, CachedSpotlightShadows);
-	AllocateAtlasedShadowDepthTargets(RHICmdList, Shadows);
-	AllocateAtlasedShadowDepthTargets(RHICmdList, CompleteShadows);
+	AllocateAtlasedShadowDepthTargets(RHICmdList, Shadows, SortedShadowsForShadowDepthPass.ShadowMapAtlases);
+	AllocateAtlasedShadowDepthTargets(RHICmdList, CompleteShadows,
+		bCopyToCompleteShadowMaps ? SortedShadowsForShadowDepthPass.CompleteShadowMapAtlases : SortedShadowsForShadowDepthPass.ShadowMapAtlases);
 	AllocateTranslucentShadowDepthTargets(RHICmdList, TranslucentShadows);
 
 	// Update translucent shadow map uniform buffers.
@@ -4585,7 +4619,10 @@ struct FLayoutAndAssignedShadows
 	TArray<FProjectedShadowInfo*, SceneRenderingAllocator> Shadows;
 };
 
-void FSceneRenderer::AllocateAtlasedShadowDepthTargets(FRHICommandListImmediate& RHICmdList, TArray<FProjectedShadowInfo*, SceneRenderingAllocator>& Shadows)
+void FSceneRenderer::AllocateAtlasedShadowDepthTargets(
+	FRHICommandListImmediate& RHICmdList,
+	TArray<FProjectedShadowInfo*, SceneRenderingAllocator>& Shadows,
+	TArray<FSortedShadowMapAtlas,SceneRenderingAllocator>& OutAtlases)
 {
 	if (Shadows.Num() == 0)
 	{
@@ -4644,8 +4681,8 @@ void FSceneRenderer::AllocateAtlasedShadowDepthTargets(FRHICommandListImmediate&
 	{
 		const FLayoutAndAssignedShadows& CurrentLayout = Layouts[LayoutIndex];
 
-		SortedShadowsForShadowDepthPass.ShadowMapAtlases.AddDefaulted();
-		FSortedShadowMapAtlas& ShadowMapAtlas = SortedShadowsForShadowDepthPass.ShadowMapAtlases.Last();
+		OutAtlases.AddDefaulted();
+		FSortedShadowMapAtlas& ShadowMapAtlas = OutAtlases.Last();
 
 		// Snap atlas sizes to powers of 2 to keep them more consistent for the render target pool
 		FIntPoint AtlasSize(
@@ -4722,7 +4759,11 @@ const TCHAR* GetCSMRenderTargetName(int32 ShadowMapIndex)
 	return **ShadowmapNames[ShadowMapIndex];
 }
 
-void FSceneRenderer::AllocateCSMDepthTargets(FRHICommandListImmediate& RHICmdList, const TArray<FProjectedShadowInfo*, SceneRenderingAllocator>& WholeSceneDirectionalShadows)
+void FSceneRenderer::AllocateCSMDepthTargets(
+	FRHICommandListImmediate& RHICmdList,
+	const TArray<FProjectedShadowInfo*,	SceneRenderingAllocator>& WholeSceneDirectionalShadows,
+	TArray<FSortedShadowMapAtlas, SceneRenderingAllocator>& OutAtlases
+	)
 {
 	if (WholeSceneDirectionalShadows.Num() > 0)
 	{
@@ -4760,25 +4801,20 @@ void FSceneRenderer::AllocateCSMDepthTargets(FRHICommandListImmediate& RHICmdLis
 		for (int32 LayoutIndex = 0; LayoutIndex < Layouts.Num(); LayoutIndex++)
 		{
 			const FLayoutAndAssignedShadows& CurrentLayout = Layouts[LayoutIndex];
-
-			SortedShadowsForShadowDepthPass.ShadowMapAtlases.AddDefaulted();
-			FSortedShadowMapAtlas& ShadowMapAtlas = SortedShadowsForShadowDepthPass.ShadowMapAtlases.Last();
+			OutAtlases.AddDefaulted();
+			FSortedShadowMapAtlas& ShadowMapAtlas = OutAtlases.Last();
 
 			FIntPoint WholeSceneAtlasSize(CurrentLayout.TextureLayout.GetSizeX(), CurrentLayout.TextureLayout.GetSizeY());
 			FPooledRenderTargetDesc WholeSceneShadowMapDesc2D(FPooledRenderTargetDesc::Create2DDesc(WholeSceneAtlasSize, PF_ShadowDepth, FClearValueBinding::DepthOne, TexCreate_None, TexCreate_DepthStencilTargetable | TexCreate_ShaderResource, false));
 			WholeSceneShadowMapDesc2D.Flags |= GFastVRamConfig.ShadowCSM;
 			GRenderTargetPool.FindFreeElement(RHICmdList, WholeSceneShadowMapDesc2D, ShadowMapAtlas.RenderTargets.DepthTarget, GetCSMRenderTargetName(LayoutIndex));
 
-			for (int32 ShadowIndex = 0; ShadowIndex < CurrentLayout.Shadows.Num(); ShadowIndex++)
+			for (FProjectedShadowInfo* ProjectedShadowInfo : CurrentLayout.Shadows)
 			{
-				FProjectedShadowInfo* ProjectedShadowInfo = CurrentLayout.Shadows[ShadowIndex];
-
-				if (ProjectedShadowInfo->bAllocated)
-				{
-					ProjectedShadowInfo->RenderTargets.CopyReferencesFromRenderTargets(ShadowMapAtlas.RenderTargets);
-					ProjectedShadowInfo->SetupShadowDepthView(RHICmdList, this);
-					ShadowMapAtlas.Shadows.Add(ProjectedShadowInfo);
-				}
+				check(ProjectedShadowInfo->bAllocated);
+				ProjectedShadowInfo->RenderTargets.CopyReferencesFromRenderTargets(ShadowMapAtlas.RenderTargets);
+				ProjectedShadowInfo->SetupShadowDepthView(RHICmdList, this);
+				ShadowMapAtlas.Shadows.Add(ProjectedShadowInfo);
 			}
 		}
 	}
diff --git a/Engine/Source/Runtime/Renderer/Private/ShadowSetupMobile.cpp b/Engine/Source/Runtime/Renderer/Private/ShadowSetupMobile.cpp
--- a/Engine/Source/Runtime/Renderer/Private/ShadowSetupMobile.cpp
+++ b/Engine/Source/Runtime/Renderer/Private/ShadowSetupMobile.cpp
@@ -341,6 +341,8 @@ bool BuildSingleCascadeShadowInfo(FViewInfo &View, TArray<FVisibleLightInfo, Sce
 				WholeSceneInitializer,
 				ShadowBufferResolution.X,
 				ShadowBufferResolution.Y,
+				ShadowBufferResolution.X,
+				ShadowBufferResolution.Y,
 				ShadowBorder,
 				false	// no RSM
 			);
diff --git a/Engine/Source/Runtime/Renderer/Private/VolumetricFogLightFunction.cpp b/Engine/Source/Runtime/Renderer/Private/VolumetricFogLightFunction.cpp
--- a/Engine/Source/Runtime/Renderer/Private/VolumetricFogLightFunction.cpp
+++ b/Engine/Source/Runtime/Renderer/Private/VolumetricFogLightFunction.cpp
@@ -192,6 +192,8 @@ void FDeferredShadingSceneRenderer::RenderLightFunctionForVolumetricFog(
 				ShadowInitializer,
 				LightFunctionResolution.X,
 				LightFunctionResolution.Y,
+				LightFunctionResolution.X,
+				LightFunctionResolution.Y,
 				0,
 				false
 			);
