First Strata checkin
- added cvar to enable it
- exemple of buffer and texture created and written to from BasePass

diff --git a/Engine/Source/Runtime/Engine/Classes/Engine/RendererSettings.h b/Engine/Source/Runtime/Engine/Classes/Engine/RendererSettings.h
--- a/Engine/Source/Runtime/Engine/Classes/Engine/RendererSettings.h
+++ b/Engine/Source/Runtime/Engine/Classes/Engine/RendererSettings.h
@@ -719,6 +719,15 @@ class ENGINE_API URendererSettings : public UDeveloperSettings
 		ConfigRestartRequired = true))
 		uint32 bSupportSkyAtmosphereAffectsHeightFog : 1;
 
+	/**
+	"Strata settings."
+	*/
+	UPROPERTY(config, EditAnywhere, Category = Strata, meta = (
+		ConsoleVariable = "r.Strata", DisplayName = "Strata materials",
+		ToolTip = "Enable the new strata material worfklow and evaluation.",
+		ConfigRestartRequired = true))
+		uint32 bEnableStrata : 1;
+
 	/**
 	"Skin cache allows a compute shader to skin once each vertex, save those results into a new buffer and reuse those calculations when later running the depth, base and velocity passes. This also allows opting into the 'recompute tangents' for skinned mesh instance feature. Disabling will reduce the number of shader permutations required per material. Changing this setting requires restarting the editor."
 	*/
diff --git a/Engine/Source/Runtime/Engine/Private/ShaderCompiler/ShaderCompiler.cpp b/Engine/Source/Runtime/Engine/Private/ShaderCompiler/ShaderCompiler.cpp
--- a/Engine/Source/Runtime/Engine/Private/ShaderCompiler/ShaderCompiler.cpp
+++ b/Engine/Source/Runtime/Engine/Private/ShaderCompiler/ShaderCompiler.cpp
@@ -3975,6 +3975,12 @@ void GlobalBeginCompileShader(
 		Input.Environment.SetDefine(TEXT("PROJECT_SUPPORT_SKY_ATMOSPHERE"), bSupportSkyAtmosphere ? 1 : 0);
 	}
 
+	{
+		static IConsoleVariable* CVar = IConsoleManager::Get().FindConsoleVariable(TEXT("r.Strata"));
+		bool bStrata = CVar && CVar->GetInt() != 0;
+		Input.Environment.SetDefine(TEXT("PROJECT_STRATA"), bStrata ? 1 : 0);
+	}
+
 	{
 		static IConsoleVariable* CVar = IConsoleManager::Get().FindConsoleVariable(TEXT("r.SupportSkyAtmosphereAffectsHeightFog"));
 		Input.Environment.SetDefine(TEXT("PROJECT_SUPPORT_SKY_ATMOSPHERE_AFFECTS_HEIGHFOG"), (CVar && bSupportSkyAtmosphere) ? (CVar->GetInt() != 0) : 0);
diff --git a/Engine/Source/Runtime/RenderCore/Private/Shader.cpp b/Engine/Source/Runtime/RenderCore/Private/Shader.cpp
--- a/Engine/Source/Runtime/RenderCore/Private/Shader.cpp
+++ b/Engine/Source/Runtime/RenderCore/Private/Shader.cpp
@@ -1662,6 +1662,14 @@ void ShaderMapAppendKeyString(EShaderPlatform Platform, FString& KeyString)
 		}
 	}
 
+	{
+		static const auto CVar = IConsoleManager::Get().FindTConsoleVariableDataInt(TEXT("r.Strata"));
+		if (CVar && CVar->GetValueOnAnyThread() > 0)
+		{
+			KeyString += TEXT("_STRATA");
+		}
+	}
+
 	{
 		if (MaskedInEarlyPass(Platform))
 		{
diff --git a/Engine/Source/Runtime/Renderer/Private/BasePassRendering.cpp b/Engine/Source/Runtime/Renderer/Private/BasePassRendering.cpp
--- a/Engine/Source/Runtime/Renderer/Private/BasePassRendering.cpp
+++ b/Engine/Source/Runtime/Renderer/Private/BasePassRendering.cpp
@@ -649,6 +649,9 @@ void CreateOpaqueBasePassUniformBuffer(
 	BasePassParameters.PreIntegratedGFSampler = TStaticSamplerState<SF_Bilinear, AM_Clamp, AM_Clamp, AM_Clamp>::GetRHI();
 	SetupDistortionParams(BasePassParameters.DistortionParams, View);
 
+	// Strata
+	BindStrataBasePassUniformParameters(View, BasePassParameters.Strata);
+
 	// Misc
 	BasePassParameters.EyeAdaptation = GetEyeAdaptation(View);
 
diff --git a/Engine/Source/Runtime/Renderer/Private/BasePassRendering.h b/Engine/Source/Runtime/Renderer/Private/BasePassRendering.h
--- a/Engine/Source/Runtime/Renderer/Private/BasePassRendering.h
+++ b/Engine/Source/Runtime/Renderer/Private/BasePassRendering.h
@@ -25,6 +25,7 @@
 #include "PlanarReflectionRendering.h"
 #include "UnrealEngine.h"
 #include "ReflectionEnvironment.h"
+#include "Strata/Strata.h"
 
 class FScene;
 
@@ -33,6 +34,8 @@ template<typename TBufferStruct> class TUniformBufferRef;
 class FViewInfo;
 class UMaterialExpressionSingleLayerWaterMaterialOutput;
 
+struct FStrataData;
+
 /** Whether to allow the indirect lighting cache to be applied to dynamic objects. */
 extern int32 GIndirectLightingCache;
 
@@ -59,6 +62,7 @@ END_GLOBAL_SHADER_PARAMETER_STRUCT()
 
 BEGIN_GLOBAL_SHADER_PARAMETER_STRUCT(FOpaqueBasePassUniformParameters,)
 	SHADER_PARAMETER_STRUCT(FSharedBasePassUniformParameters, Shared)
+	SHADER_PARAMETER_STRUCT(FStrataOpaquePassUniformParameters, Strata)
 	// Forward shading 
 	SHADER_PARAMETER(int32, UseForwardScreenSpaceShadowMask)
 	SHADER_PARAMETER_TEXTURE(Texture2D, ForwardScreenSpaceShadowMaskTexture)
diff --git a/Engine/Source/Runtime/Renderer/Private/DeferredShadingRenderer.cpp b/Engine/Source/Runtime/Renderer/Private/DeferredShadingRenderer.cpp
--- a/Engine/Source/Runtime/Renderer/Private/DeferredShadingRenderer.cpp
+++ b/Engine/Source/Runtime/Renderer/Private/DeferredShadingRenderer.cpp
@@ -51,6 +51,7 @@
 #include "Rendering/NaniteStreamingManager.h"
 #include "SceneTextureReductions.h"
 #include "VirtualShadowMaps/VirtualShadowMapCacheManager.h"
+#include "Strata/Strata.h"
 
 extern int32 GNaniteDebugFlags;
 extern int32 GNaniteShowStats;
@@ -2303,6 +2304,13 @@ void FDeferredShadingSceneRenderer::Render(FRHICommandListImmediate& RHICmdList)
 		Scene->AllocateAndCaptureFrameSkyEnvMap(RHICmdList, *this, MainView, bShouldRenderSkyAtmosphere, bShouldRenderVolumetricCloud);
 	}
 
+	// Strata initialisation
+	{
+		FRDGBuilder GraphBuilder(RHICmdList);
+		InitialiseStrataFrameSceneData(*this, GraphBuilder);
+		GraphBuilder.Execute();
+	}
+
 	// Clear LPVs for all views
 	if (FeatureLevel >= ERHIFeatureLevel::SM5)
 	{
diff --git a/Engine/Source/Runtime/Renderer/Private/DeferredShadingRenderer.h b/Engine/Source/Runtime/Renderer/Private/DeferredShadingRenderer.h
--- a/Engine/Source/Runtime/Renderer/Private/DeferredShadingRenderer.h
+++ b/Engine/Source/Runtime/Renderer/Private/DeferredShadingRenderer.h
@@ -38,6 +38,7 @@ namespace LumenRadianceCache
 struct FSingleLayerWaterPassData;
 struct FHeightFogRenderingParameters;
 struct FRayTracingReflectionOptions;
+struct FStrataData;
 
 class FLightShaftsOutput
 {
diff --git a/Engine/Source/Runtime/Renderer/Private/ScenePrivate.h b/Engine/Source/Runtime/Renderer/Private/ScenePrivate.h
--- a/Engine/Source/Runtime/Renderer/Private/ScenePrivate.h
+++ b/Engine/Source/Runtime/Renderer/Private/ScenePrivate.h
@@ -2710,6 +2710,9 @@ public:
 	/** The runtime virtual textures in the scene. */
 	TSparseArray<FRuntimeVirtualTextureSceneProxy*> RuntimeVirtualTextures;
 
+	/** Strata data shared between all views. */
+	FStrataData StrataData;
+
 	float DefaultMaxDistanceFieldOcclusionDistance;
 
 	float GlobalDistanceFieldViewDistance;
diff --git a/Engine/Source/Runtime/Renderer/Private/SceneRendering.cpp b/Engine/Source/Runtime/Renderer/Private/SceneRendering.cpp
--- a/Engine/Source/Runtime/Renderer/Private/SceneRendering.cpp
+++ b/Engine/Source/Runtime/Renderer/Private/SceneRendering.cpp
@@ -929,6 +929,8 @@ void FViewInfo::Init()
 
 	NumVisibleDynamicPrimitives = 0;
 	NumVisibleDynamicEditorPrimitives = 0;
+
+	StrataData = nullptr;
 }
 
 FViewInfo::~FViewInfo()
diff --git a/Engine/Source/Runtime/Renderer/Private/SceneRendering.h b/Engine/Source/Runtime/Renderer/Private/SceneRendering.h
--- a/Engine/Source/Runtime/Renderer/Private/SceneRendering.h
+++ b/Engine/Source/Runtime/Renderer/Private/SceneRendering.h
@@ -48,6 +48,7 @@ struct FExposureBufferData;
 struct FCloudRenderContext;
 struct FSingleLayerWaterPassData;
 class FVirtualShadowMapClipmap;
+struct FStrataData;
 
 DECLARE_STATS_GROUP(TEXT("Command List Markers"), STATGROUP_CommandListMarkers, STATCAT_Advanced);
 
@@ -1129,6 +1130,8 @@ public:
 	// Used by mobile renderer to determine whether static meshes will be rendered with CSM shaders or not.
 	FMobileCSMVisibilityInfo MobileCSMVisibilityInfo;
 
+	FStrataData* StrataData;
+
 	/** Parameters for exponential height fog. */
 	FVector4 ExponentialFogParameters;
 	FVector4 ExponentialFogParameters2;
diff --git a/Engine/Source/Runtime/Renderer/Private/Strata/Strata.cpp b/Engine/Source/Runtime/Renderer/Private/Strata/Strata.cpp
new file mode 100644
--- /dev/null
+++ b/Engine/Source/Runtime/Renderer/Private/Strata/Strata.cpp
@@ -0,0 +1,84 @@
+// Copyright Epic Games, Inc. All Rights Reserved.
+
+#include "Strata.h"
+#include "HAL/IConsoleManager.h"
+#include "SceneView.h"
+#include "ScenePrivate.h"
+#include "SceneRendering.h"
+#include "RendererInterface.h"
+
+
+
+// The project setting for Strata
+static TAutoConsoleVariable<int32> CVarStrata(
+	TEXT("r.Strata"),
+	0,
+	TEXT("Enables Strata."),
+	ECVF_ReadOnly | ECVF_RenderThreadSafe);
+
+
+
+bool IsStrataEnabled()
+{
+	return CVarStrata.GetValueOnRenderThread() > 0;
+}
+
+
+
+void InitialiseStrataFrameSceneData(FSceneRenderer& SceneRenderer, FRDGBuilder& GraphBuilder)
+{
+	FStrataData& StrataData = SceneRenderer.Scene->StrataData;
+
+	uint32 ResolutionX = 1;
+	uint32 ResolutionY = 1;
+
+	if (IsStrataEnabled())
+	{
+		FSceneRenderTargets& SceneContext = FSceneRenderTargets::Get(GraphBuilder.RHICmdList);
+		FIntPoint BufferSizeXY = SceneContext.GetBufferSizeXY();
+		
+		ResolutionX = BufferSizeXY.X;
+		ResolutionY = BufferSizeXY.Y;
+
+		StrataData.MaxBytesPerPixel = 16;
+
+	}
+	else
+	{
+		StrataData.MaxBytesPerPixel = 1;
+	}
+
+	FRDGTextureRef MaterialLobesTexture = GraphBuilder.CreateTexture(
+		FRDGTextureDesc::Create2DDesc(FIntPoint(ResolutionX, ResolutionY), PF_R16F, FClearValueBinding::None,
+		TexCreate_None, TexCreate_RenderTargetable | TexCreate_ShaderResource | TexCreate_UAV, false, 1), TEXT("MaterialLobesTexture"));
+	AddClearRenderTargetPass(GraphBuilder, MaterialLobesTexture, FLinearColor::Black);
+	GraphBuilder.QueueTextureExtraction(MaterialLobesTexture, &StrataData.MaterialLobesTexture);
+
+	const uint32 DesiredBufferSize = FMath::Max(1u, ResolutionX * ResolutionY * StrataData.MaxBytesPerPixel);
+	if (StrataData.MaterialLobesBuffer.NumBytes < DesiredBufferSize)
+	{
+		if (StrataData.MaterialLobesBuffer.NumBytes > 0)
+		{
+			StrataData.MaterialLobesBuffer.Release();
+		}
+		StrataData.MaterialLobesBuffer.Initialize(DesiredBufferSize, BUF_Static, TEXT("MaterialLobesBuffer"));
+	}
+
+	// Set reference to the Strata data from each view
+	for (int32 ViewIndex = 0; ViewIndex < SceneRenderer.Views.Num(); ViewIndex++)
+	{
+		FViewInfo& View = SceneRenderer.Views[ViewIndex];
+		View.StrataData = &SceneRenderer.Scene->StrataData;
+	}
+}
+
+
+
+void BindStrataBasePassUniformParameters(const FViewInfo& View, FStrataOpaquePassUniformParameters& OutStrataUniformParameters)
+{
+	OutStrataUniformParameters.MaxBytesPerPixel = View.StrataData->MaxBytesPerPixel;
+	OutStrataUniformParameters.MaterialLobesTextureUAV = View.StrataData->MaterialLobesTexture->GetRenderTargetItem().UAV;
+	OutStrataUniformParameters.MaterialLobesBufferUAV = View.StrataData->MaterialLobesBuffer.UAV;
+}
+
+
diff --git a/Engine/Source/Runtime/Renderer/Private/Strata/Strata.h b/Engine/Source/Runtime/Renderer/Private/Strata/Strata.h
new file mode 100644
--- /dev/null
+++ b/Engine/Source/Runtime/Renderer/Private/Strata/Strata.h
@@ -0,0 +1,45 @@
+// Copyright Epic Games, Inc. All Rights Reserved.
+
+#pragma once
+
+#include "CoreMinimal.h"
+#include "ShaderParameterMacros.h"
+#include "RenderGraphResources.h"
+#include "MeshPassProcessor.h"
+#include "UnifiedBuffer.h"
+#include "RHIUtilities.h"
+
+// Forward declarations.
+class FScene;
+class FRDGBuilder;
+
+
+
+BEGIN_SHADER_PARAMETER_STRUCT(FStrataOpaquePassUniformParameters, )
+	SHADER_PARAMETER(uint32, MaxBytesPerPixel)
+	SHADER_PARAMETER_UAV(RWByteAddressBuffer, MaterialLobesBufferUAV)
+	SHADER_PARAMETER_UAV(RWTexture2D<float>, MaterialLobesTextureUAV)
+END_SHADER_PARAMETER_STRUCT()
+
+
+
+struct FStrataData
+{
+	uint32 MaxBytesPerPixel;
+	TRefCountPtr<IPooledRenderTarget> MaterialLobesTexture; // This should be a RDG resource when the refactoring gets in
+	FRWByteAddressBuffer MaterialLobesBuffer;							// This should be a RDG resource	"		"		"
+
+	FStrataData()
+	{
+	}
+};
+
+
+
+bool IsStrataEnabled();
+
+void InitialiseStrataFrameSceneData(FSceneRenderer& SceneRenderer, FRDGBuilder& GraphBuilder);
+
+void BindStrataBasePassUniformParameters(const FViewInfo& View, FStrataOpaquePassUniformParameters& OutStrataUniformParameters);
+
+
