Fixed crash when Strata is disabled

diff --git a/Engine/Source/Runtime/Renderer/Private/Strata/Strata.cpp b/Engine/Source/Runtime/Renderer/Private/Strata/Strata.cpp
--- a/Engine/Source/Runtime/Renderer/Private/Strata/Strata.cpp
+++ b/Engine/Source/Runtime/Renderer/Private/Strata/Strata.cpp
@@ -54,7 +54,7 @@ void InitialiseStrataFrameSceneData(FSceneRenderer& SceneRenderer, FRDGBuilder&
 	AddClearRenderTargetPass(GraphBuilder, MaterialLobesTexture, FLinearColor::Black);
 	GraphBuilder.QueueTextureExtraction(MaterialLobesTexture, &StrataData.MaterialLobesTexture);
 
-	const uint32 DesiredBufferSize = FMath::Max(1u, ResolutionX * ResolutionY * StrataData.MaxBytesPerPixel);
+	const uint32 DesiredBufferSize = FMath::Max(4u, ResolutionX * ResolutionY * StrataData.MaxBytesPerPixel);
 	if (StrataData.MaterialLobesBuffer.NumBytes < DesiredBufferSize)
 	{
 		if (StrataData.MaterialLobesBuffer.NumBytes > 0)
