Fixed crash when opening material editor.

diff --git a/Engine/Source/Runtime/Renderer/Private/Strata/Strata.cpp b/Engine/Source/Runtime/Renderer/Private/Strata/Strata.cpp
--- a/Engine/Source/Runtime/Renderer/Private/Strata/Strata.cpp
+++ b/Engine/Source/Runtime/Renderer/Private/Strata/Strata.cpp
@@ -76,9 +76,18 @@ void InitialiseStrataFrameSceneData(FSceneRenderer& SceneRenderer, FRDGBuilder&
 
 void BindStrataBasePassUniformParameters(const FViewInfo& View, FStrataOpaquePassUniformParameters& OutStrataUniformParameters)
 {
-	OutStrataUniformParameters.MaxBytesPerPixel = View.StrataData->MaxBytesPerPixel;
-	OutStrataUniformParameters.MaterialLobesTextureUAV = View.StrataData->MaterialLobesTexture->GetRenderTargetItem().UAV;
-	OutStrataUniformParameters.MaterialLobesBufferUAV = View.StrataData->MaterialLobesBuffer.UAV;
+	if (View.StrataData)
+	{
+		OutStrataUniformParameters.MaxBytesPerPixel = View.StrataData->MaxBytesPerPixel;
+		OutStrataUniformParameters.MaterialLobesTextureUAV = View.StrataData->MaterialLobesTexture->GetRenderTargetItem().UAV;
+		OutStrataUniformParameters.MaterialLobesBufferUAV = View.StrataData->MaterialLobesBuffer.UAV;
+	}
+	else
+	{
+		OutStrataUniformParameters.MaxBytesPerPixel = 0;
+		OutStrataUniformParameters.MaterialLobesTextureUAV = GEmptyVertexBufferWithUAV->UnorderedAccessViewRHI;
+		OutStrataUniformParameters.MaterialLobesBufferUAV = GEmptyVertexBufferWithUAV->UnorderedAccessViewRHI;
+	}
 }
 
 
