Added Strata namespace.

diff --git a/Engine/Source/Runtime/Renderer/Private/BasePassRendering.cpp b/Engine/Source/Runtime/Renderer/Private/BasePassRendering.cpp
--- a/Engine/Source/Runtime/Renderer/Private/BasePassRendering.cpp
+++ b/Engine/Source/Runtime/Renderer/Private/BasePassRendering.cpp
@@ -650,7 +650,7 @@ void CreateOpaqueBasePassUniformBuffer(
 	SetupDistortionParams(BasePassParameters.DistortionParams, View);
 
 	// Strata
-	BindStrataBasePassUniformParameters(View, BasePassParameters.Strata);
+	Strata::BindStrataBasePassUniformParameters(View, BasePassParameters.Strata);
 
 	// Misc
 	BasePassParameters.EyeAdaptation = GetEyeAdaptation(View);
diff --git a/Engine/Source/Runtime/Renderer/Private/DeferredShadingRenderer.cpp b/Engine/Source/Runtime/Renderer/Private/DeferredShadingRenderer.cpp
--- a/Engine/Source/Runtime/Renderer/Private/DeferredShadingRenderer.cpp
+++ b/Engine/Source/Runtime/Renderer/Private/DeferredShadingRenderer.cpp
@@ -2307,7 +2307,7 @@ void FDeferredShadingSceneRenderer::Render(FRHICommandListImmediate& RHICmdList)
 	// Strata initialisation
 	{
 		FRDGBuilder GraphBuilder(RHICmdList);
-		InitialiseStrataFrameSceneData(*this, GraphBuilder);
+		Strata::InitialiseStrataFrameSceneData(*this, GraphBuilder);
 		GraphBuilder.Execute();
 	}
 
diff --git a/Engine/Source/Runtime/Renderer/Private/Strata/Strata.cpp b/Engine/Source/Runtime/Renderer/Private/Strata/Strata.cpp
--- a/Engine/Source/Runtime/Renderer/Private/Strata/Strata.cpp
+++ b/Engine/Source/Runtime/Renderer/Private/Strata/Strata.cpp
@@ -18,13 +18,14 @@ static TAutoConsoleVariable<int32> CVarStrata(
 
 
 
+namespace Strata
+{
+
 bool IsStrataEnabled()
 {
 	return CVarStrata.GetValueOnRenderThread() > 0;
 }
 
-
-
 void InitialiseStrataFrameSceneData(FSceneRenderer& SceneRenderer, FRDGBuilder& GraphBuilder)
 {
 	FStrataData& StrataData = SceneRenderer.Scene->StrataData;
@@ -72,8 +73,6 @@ void InitialiseStrataFrameSceneData(FSceneRenderer& SceneRenderer, FRDGBuilder&
 	}
 }
 
-
-
 void BindStrataBasePassUniformParameters(const FViewInfo& View, FStrataOpaquePassUniformParameters& OutStrataUniformParameters)
 {
 	if (View.StrataData)
@@ -90,4 +89,6 @@ void BindStrataBasePassUniformParameters(const FViewInfo& View, FStrataOpaquePas
 	}
 }
 
+} // namespace Strata
+
 
diff --git a/Engine/Source/Runtime/Renderer/Private/Strata/Strata.h b/Engine/Source/Runtime/Renderer/Private/Strata/Strata.h
--- a/Engine/Source/Runtime/Renderer/Private/Strata/Strata.h
+++ b/Engine/Source/Runtime/Renderer/Private/Strata/Strata.h
@@ -27,7 +27,7 @@ struct FStrataData
 {
 	uint32 MaxBytesPerPixel;
 	TRefCountPtr<IPooledRenderTarget> MaterialLobesTexture; // This should be a RDG resource when the refactoring gets in
-	FRWByteAddressBuffer MaterialLobesBuffer;							// This should be a RDG resource	"		"		"
+	FRWByteAddressBuffer MaterialLobesBuffer;				// This should be a RDG resource	"		"		"
 
 	FStrataData()
 	{
@@ -36,10 +36,16 @@ struct FStrataData
 
 
 
+
+namespace Strata
+{
+
 bool IsStrataEnabled();
 
 void InitialiseStrataFrameSceneData(FSceneRenderer& SceneRenderer, FRDGBuilder& GraphBuilder);
 
 void BindStrataBasePassUniformParameters(const FViewInfo& View, FStrataOpaquePassUniformParameters& OutStrataUniformParameters);
 
+};
+
 
