Fixed crash with volumetric fog conservative depth.
Tested on PC

diff --git a/Engine/Source/Runtime/Renderer/Private/VolumetricFog.cpp b/Engine/Source/Runtime/Renderer/Private/VolumetricFog.cpp
--- a/Engine/Source/Runtime/Renderer/Private/VolumetricFog.cpp
+++ b/Engine/Source/Runtime/Renderer/Private/VolumetricFog.cpp
@@ -1172,7 +1172,7 @@ void FDeferredShadingSceneRenderer::ComputeVolumetricFog(FRDGBuilder& GraphBuild
 			{
 				FIntPoint ConservativeDepthTextureSize = FIntPoint(VolumetricFogGridSize.X, VolumetricFogGridSize.Y);
 				ConservativeDepthTexture = GraphBuilder.CreateTexture(FPooledRenderTargetDesc::Create2DDesc(ConservativeDepthTextureSize, PF_R16F,
-					FClearValueBinding::Black, 0, TexCreate_ShaderResource | TexCreate_UAV, false), TEXT("ConservativeDepthTexture"));
+					FClearValueBinding::Black, 0, TexCreate_RenderTargetable | TexCreate_ShaderResource | TexCreate_UAV, false), TEXT("ConservativeDepthTexture"));
 				AddGenerateConservativeDepthBufferPass(View, GraphBuilder, ConservativeDepthTexture, GVolumetricFogGridPixelSize);
 			}
 			else
