Added check to prevent calling CreateLightPrimitiveInteraction for a Primitive that is being removed
- at this point the proxy contains potentially dangling pointers to data owned by the Component and is a mine field.

diff --git a/Engine/Source/Runtime/Renderer/Private/LightSceneInfo.cpp b/Engine/Source/Runtime/Renderer/Private/LightSceneInfo.cpp
--- a/Engine/Source/Runtime/Renderer/Private/LightSceneInfo.cpp
+++ b/Engine/Source/Runtime/Renderer/Private/LightSceneInfo.cpp
@@ -133,7 +133,7 @@ void FLightSceneInfo::AddToScene()
  */
 void FLightSceneInfo::CreateLightPrimitiveInteraction(const FLightSceneInfoCompact& LightSceneInfoCompact, const FPrimitiveSceneInfoCompact& PrimitiveSceneInfoCompact)
 {
-	if(	LightSceneInfoCompact.AffectsPrimitive(PrimitiveSceneInfoCompact.Bounds, PrimitiveSceneInfoCompact.Proxy))
+	if(	!Scene->IsPrimitiveBeingRemoved(PrimitiveSceneInfoCompact.PrimitiveSceneInfo) && LightSceneInfoCompact.AffectsPrimitive(PrimitiveSceneInfoCompact.Bounds, PrimitiveSceneInfoCompact.Proxy))
 	{
 		// create light interaction and add to light/primitive lists
 		FLightPrimitiveInteraction::Create(this,PrimitiveSceneInfoCompact.PrimitiveSceneInfo);
diff --git a/Engine/Source/Runtime/Renderer/Private/RendererScene.cpp b/Engine/Source/Runtime/Renderer/Private/RendererScene.cpp
--- a/Engine/Source/Runtime/Renderer/Private/RendererScene.cpp
+++ b/Engine/Source/Runtime/Renderer/Private/RendererScene.cpp
@@ -1104,6 +1104,7 @@ public:
 
 	void DoWork()
 	{
+		FOptionalTaskTagScope Scope(ETaskTag::EParallelRenderingThread);
 		for (TSet<FPrimitiveSceneInfo*>::TConstIterator It(PendingPrimitives); It; ++It)
 		{
 			FMemMark MemStackMark(FMemStack::Get());
@@ -4395,6 +4396,12 @@ void FScene::FlushAsyncLightPrimitiveInteractionCreation() const
 	}
 }
 
+bool FScene::IsPrimitiveBeingRemoved(FPrimitiveSceneInfo* PrimitiveSceneInfo) const
+{
+	check(IsInRenderingThread() || IsInParallelRenderingThread());
+	return RemovedPrimitiveSceneInfos.Contains(PrimitiveSceneInfo);
+}
+
 /**
  * Dummy NULL scene interface used by dedicated servers.
  */
diff --git a/Engine/Source/Runtime/Renderer/Private/ScenePrivate.h b/Engine/Source/Runtime/Renderer/Private/ScenePrivate.h
--- a/Engine/Source/Runtime/Renderer/Private/ScenePrivate.h
+++ b/Engine/Source/Runtime/Renderer/Private/ScenePrivate.h
@@ -3016,6 +3016,8 @@ public:
 
 	void FlushAsyncLightPrimitiveInteractionCreation() const;
 
+	bool IsPrimitiveBeingRemoved(FPrimitiveSceneInfo* PrimitiveSceneInfo) const;
+
 private:
 
 	/**
