Fixed: QAGame TM-ShaderModels crash when switching to mobile preview (merged from UE4)

diff --git a/Engine/Source/Runtime/Renderer/Private/LightRendering.h b/Engine/Source/Runtime/Renderer/Private/LightRendering.h
--- a/Engine/Source/Runtime/Renderer/Private/LightRendering.h
+++ b/Engine/Source/Runtime/Renderer/Private/LightRendering.h
@@ -454,7 +454,12 @@ public:
 
 	static bool ShouldCompilePermutation(const FGlobalShaderPermutationParameters& Parameters)
 	{
-		return IsFeatureLevelSupported(Parameters.Platform, ERHIFeatureLevel::SM5) || IsMobileDeferredShadingEnabled(Parameters.Platform);
+		if (bRadialLight)
+		{
+			return IsFeatureLevelSupported(Parameters.Platform, ERHIFeatureLevel::SM5) || IsMobileDeferredShadingEnabled(Parameters.Platform);
+		}
+		// used with FPrefilterPlanarReflectionPS on mobile
+		return true;
 	}
 
 	static void ModifyCompilationEnvironment(const FGlobalShaderPermutationParameters& Parameters, FShaderCompilerEnvironment& OutEnvironment)
