Fix light grid injection to use same conditions as before merge from main
- should be enabled for lumen and virtual shadow maps
- also ensure to detect if the light grid was not built before virtual shadow map setup.

diff --git a/Engine/Source/Runtime/Renderer/Private/ClusteredDeferredShadingPass.cpp b/Engine/Source/Runtime/Renderer/Private/ClusteredDeferredShadingPass.cpp
--- a/Engine/Source/Runtime/Renderer/Private/ClusteredDeferredShadingPass.cpp
+++ b/Engine/Source/Runtime/Renderer/Private/ClusteredDeferredShadingPass.cpp
@@ -49,9 +49,9 @@ bool FDeferredShadingSceneRenderer::ShouldUseClusteredDeferredShading() const
 	return GUseClusteredDeferredShading != 0 && Scene->GetFeatureLevel() >= ERHIFeatureLevel::SM5;
 }
 
-bool FDeferredShadingSceneRenderer::AreClusteredLightsInLightGrid() const
+bool FDeferredShadingSceneRenderer::AreLightsInLightGrid() const
 {
-	return bClusteredShadingLightsInLightGrid;
+	return bAreLightsInLightGrid;
 }
 
 
diff --git a/Engine/Source/Runtime/Renderer/Private/DeferredShadingRenderer.cpp b/Engine/Source/Runtime/Renderer/Private/DeferredShadingRenderer.cpp
--- a/Engine/Source/Runtime/Renderer/Private/DeferredShadingRenderer.cpp
+++ b/Engine/Source/Runtime/Renderer/Private/DeferredShadingRenderer.cpp
@@ -299,7 +299,7 @@ FDeferredShadingSceneRenderer::FDeferredShadingSceneRenderer(const FSceneViewFam
 	: FSceneRenderer(InViewFamily, HitProxyConsumer)
 	, EarlyZPassMode(Scene ? Scene->EarlyZPassMode : DDM_None)
 	, bEarlyZPassMovable(Scene ? Scene->bEarlyZPassMovable : false)
-	, bClusteredShadingLightsInLightGrid(false)
+	, bAreLightsInLightGrid(false)
 {
 	static const auto StencilLODDitherCVar = IConsoleManager::Get().FindTConsoleVariableDataInt(TEXT("r.StencilForLODDither"));
 	bDitheredLODTransitionsUseStencil = StencilLODDitherCVar->GetValueOnAnyThread() != 0;
@@ -1981,7 +1981,8 @@ void FDeferredShadingSceneRenderer::Render(FRHICommandListImmediate& RHICmdList)
 		&VisibleLightInfos=VisibleLightInfos,
 		&VirtualShadowMapArray=VirtualShadowMapArray,
 		&Views= Views, 
-		VirtualShadowMapArrayCacheManager = Scene->VirtualShadowMapArrayCacheManager](bool bPostBasePass)
+		VirtualShadowMapArrayCacheManager = Scene->VirtualShadowMapArrayCacheManager,
+		this](bool bPostBasePass)
 	{
 		static const auto CVarEnableVirtualSM = IConsoleManager::Get().FindTConsoleVariableDataInt(TEXT("r.Shadow.v.Enable"));
 
@@ -1993,7 +1994,7 @@ void FDeferredShadingSceneRenderer::Render(FRHICommandListImmediate& RHICmdList)
 			}
 			else
 			{
-				// ensure(LightGridWasBuilt)
+				ensureMsgf(AreLightsInLightGrid(), TEXT("Virtual shadow map setup requires local lights to be injected into the light grid (this may be caused by 'r.LightCulling.Quality=0')."));
 				// ensure(ShadowMapSetupDone)
 				{
 					VirtualShadowMapArray.GeneratePageFlagsFromLightGrid(GraphBuilder, Views, SortedLightSet, VisibleLightInfos, NaniteRasterResults, bPostBasePass, VirtualShadowMapArrayCacheManager);
diff --git a/Engine/Source/Runtime/Renderer/Private/DeferredShadingRenderer.h b/Engine/Source/Runtime/Renderer/Private/DeferredShadingRenderer.h
--- a/Engine/Source/Runtime/Renderer/Private/DeferredShadingRenderer.h
+++ b/Engine/Source/Runtime/Renderer/Private/DeferredShadingRenderer.h
@@ -590,9 +590,9 @@ private:
 	bool ShouldUseClusteredDeferredShading() const;
 
 	/**
-	 * Have the requisite lights been injected into the light grid, AKA can we run the shading pass?
+	 * Have the lights been injected into the light grid?
 	 */
-	bool AreClusteredLightsInLightGrid() const;
+	bool AreLightsInLightGrid() const;
 
 
 	/** Add a clustered deferred shading lighting render pass.	Note: in the future it should take the RenderGraph builder as argument */
@@ -1118,8 +1118,8 @@ private:
 
 #endif // RHI_RAYTRACING
 
-	/** Set to true if the lights needed for clustered shading have been injected in the light grid (set in GatherLightsAndComputeLightGrid). */
-	bool bClusteredShadingLightsInLightGrid;
+	/** Set to true if lights were injected into the light grid (this controlled by somewhat complex logic, this flag is used to cross-check). */
+	bool bAreLightsInLightGrid;
 };
 
 DECLARE_CYCLE_STAT_EXTERN(TEXT("PrePass"), STAT_CLM_PrePass, STATGROUP_CommandListMarkers, );
\ No newline at end of file
diff --git a/Engine/Source/Runtime/Renderer/Private/LightGridInjection.cpp b/Engine/Source/Runtime/Renderer/Private/LightGridInjection.cpp
--- a/Engine/Source/Runtime/Renderer/Private/LightGridInjection.cpp
+++ b/Engine/Source/Runtime/Renderer/Private/LightGridInjection.cpp
@@ -770,11 +770,15 @@ void FDeferredShadingSceneRenderer::GatherLightsAndComputeLightGrid(FRDGBuilder&
 		bAnyViewUsesForwardLighting |= View.bTranslucentSurfaceLighting || ShouldRenderVolumetricFog() || View.bHasSingleLayerWaterMaterial;
 	}
 	
-	const bool bCullLightsToGrid = GLightCullingQuality 
-		&& (ViewFamily.EngineShowFlags.DirectLighting 
-		&& (IsForwardShadingEnabled(ShaderPlatform) || bAnyViewUsesForwardLighting || IsRayTracingEnabled() || ShouldUseClusteredDeferredShading()));
-	// Store this flag if lights are injected in the grids, check with 'AreClusteredLightsInLightGrid()'
-	bClusteredShadingLightsInLightGrid = bCullLightsToGrid;
+	// TODO: Move decision logic to separate function for readability
+	static const auto EnableVirtualSMCVar = IConsoleManager::Get().FindTConsoleVariableDataInt(TEXT("r.Shadow.v.Enable"));
+
+	const bool bCullLightsToGrid = GLightCullingQuality
+		&& (ViewFamily.EngineShowFlags.DirectLighting
+			&& (IsForwardShadingEnabled(ShaderPlatform) || bAnyViewUsesForwardLighting || IsRayTracingEnabled() || ShouldUseClusteredDeferredShading() || GAllowLumenScene || EnableVirtualSMCVar->GetValueOnRenderThread() != 0));
+
+	// Store this flag if lights are injected in the grids, check with 'AreLightsInLightGrid()'
+	bAreLightsInLightGrid = bCullLightsToGrid;
 	
 	ComputeLightGrid(GraphBuilder, bCullLightsToGrid, SortedLightSet);
 }
diff --git a/Engine/Source/Runtime/Renderer/Private/LightRendering.cpp b/Engine/Source/Runtime/Renderer/Private/LightRendering.cpp
--- a/Engine/Source/Runtime/Renderer/Private/LightRendering.cpp
+++ b/Engine/Source/Runtime/Renderer/Private/LightRendering.cpp
@@ -1262,11 +1262,11 @@ FRDGTextureRef FDeferredShadingSceneRenderer::RenderLights(
 
 			bool bRenderSimpleLightsStandardDeferred = SortedLightSet.SimpleLights.InstanceData.Num() > 0;
 
-			UE_CLOG(ShouldUseClusteredDeferredShading() && !AreClusteredLightsInLightGrid(), LogRenderer, Warning,
+			UE_CLOG(ShouldUseClusteredDeferredShading() && !AreLightsInLightGrid(), LogRenderer, Warning,
 				TEXT("Clustered deferred shading is enabled, but lights were not injected in grid, falling back to other methods (hint 'r.LightCulling.Quality' may cause this)."));
 
 			// True if the clustered shading is enabled and the feature level is there, and that the light grid had lights injected.
-			if (ShouldUseClusteredDeferredShading() && AreClusteredLightsInLightGrid())
+			if (ShouldUseClusteredDeferredShading() && AreLightsInLightGrid())
 			{
 				// Tell the trad. deferred that the clustered deferred capable lights are taken care of.
 				// This includes the simple lights
