Remove commented out code in mesh pass processor. Fixed intendation.

diff --git a/Engine/Source/Runtime/Renderer/Private/MeshPassProcessor.cpp b/Engine/Source/Runtime/Renderer/Private/MeshPassProcessor.cpp
--- a/Engine/Source/Runtime/Renderer/Private/MeshPassProcessor.cpp
+++ b/Engine/Source/Runtime/Renderer/Private/MeshPassProcessor.cpp
@@ -281,8 +281,8 @@ void FMeshDrawShaderBindings::SetRayTracingShaderBindingsForHitGroup(
 	{
 		FShaderParameterInfo Parameter = UniformBufferParameters[UniformBufferIndex];
 		const FRHIUniformBuffer* UniformBuffer = UniformBufferBindings[UniformBufferIndex];
-			MaxUniformBufferUsed = FMath::Max((int32)Parameter.BaseIndex, MaxUniformBufferUsed);
-		}
+		MaxUniformBufferUsed = FMath::Max((int32)Parameter.BaseIndex, MaxUniformBufferUsed);
+	}
 
 	const uint32 NumUniformBuffersToSet = MaxUniformBufferUsed + 1;
 
@@ -353,9 +353,9 @@ FGraphicsMinimalPipelineStateId FGraphicsMinimalPipelineStateId::GetPersistentId
 		TableId = PersistentIdTable.FindOrAddIdByHash(hash, InPipelineState, FRefCountedGraphicsMinimalPipelineState());
 		FRefCountedGraphicsMinimalPipelineState& Value = PersistentIdTable.GetByElementId(TableId).Value;
 		if (Value.RefNum == 0 && !NeedsShaderInitialisation)
-	{
+		{
 			NeedsShaderInitialisation = true;
-	}
+		}
 		Value.RefNum++;
 	}
 
@@ -391,10 +391,10 @@ void FGraphicsMinimalPipelineStateId::RemovePersistentId(FGraphicsMinimalPipelin
 		FScopeLock Lock(&PersistentIdTableLock);
 		FRefCountedGraphicsMinimalPipelineState& RefCountedStateInitializer = PersistentIdTable.GetByElementId(Id.SetElementIndex).Value;
 
-	check(RefCountedStateInitializer.RefNum > 0);
-	--RefCountedStateInitializer.RefNum;
+		check(RefCountedStateInitializer.RefNum > 0);
+		--RefCountedStateInitializer.RefNum;
 		if (RefCountedStateInitializer.RefNum == 0)
-	{
+		{
 			PersistentIdTable.RemoveByElementId(Id.SetElementIndex);
 		}
 	}
@@ -406,33 +406,17 @@ FGraphicsMinimalPipelineStateId FGraphicsMinimalPipelineStateId::GetPipelineStat
 	Ret.bValid = 1;
 	Ret.bComesFromLocalPipelineStateSet = 0;
 
-	Experimental::FHashElementId TableIndex;
-
-	// Disable getting entries from the persistent table for now because this is currently used for dynamic MDCs and View overrides without performing an AddRef on the persistent entry
-// If the persistent entries gets freed between building the MDC and draw time then the draw command will use the wrong PipelineState object. This can happen if FPrimitiveSceneInfo::UpdateStaticMeshes is 
-// called for example.
-// TODO: Better solutions are deferred deletion of persistent table entries until a fixed point when no dynamic MDCs are used anymore or not allowing deleting of persistent entries at all after dynamic MDCs or view overrides are applied.
-/*
-	{
-		FScopeLock Lock(&PersistentIdTableLock);
-		TableIndex = PersistentIdTable.FindId(InPipelineState);
-	}
-*/
-
-	if (!TableIndex.IsValid())
-	{
-		Ret.bComesFromLocalPipelineStateSet = 1;
+	Ret.bComesFromLocalPipelineStateSet = 1;
 #if UE_BUILD_DEBUG
-		FGraphicsMinimalPipelineStateInitializer PipelineStateDebug = FGraphicsMinimalPipelineStateInitializer(InPipelineState);
-		check(GetTypeHash(PipelineStateDebug) == GetTypeHash(InPipelineState));
-		check(PipelineStateDebug == InPipelineState);
+	FGraphicsMinimalPipelineStateInitializer PipelineStateDebug = FGraphicsMinimalPipelineStateInitializer(InPipelineState);
+	check(GetTypeHash(PipelineStateDebug) == GetTypeHash(InPipelineState));
+	check(PipelineStateDebug == InPipelineState);
 #endif
-		TableIndex = InOutPassSet.FindOrAddId(InPipelineState);
+	Experimental::FHashElementId TableIndex = InOutPassSet.FindOrAddId(InPipelineState);
 #if UE_BUILD_DEBUG
-		check(InOutPassSet.GetByElementId(TableIndex) == InPipelineState);
+	check(InOutPassSet.GetByElementId(TableIndex) == InPipelineState);
 #endif
-		InNeedsShaderInitialisation = InNeedsShaderInitialisation || InPipelineState.BoundShaderState.NeedsShaderInitialisation();
-	}
+	InNeedsShaderInitialisation = InNeedsShaderInitialisation || InPipelineState.BoundShaderState.NeedsShaderInitialisation();
 
 	checkf(TableIndex.GetIndex() < (MAX_uint32 >> 2), TEXT("One frame FGraphicsMinimalPipelineStateId table overflow!"));
 
@@ -666,7 +650,7 @@ void FMeshDrawShaderBindings::CopyFrom(const FMeshDrawShaderBindings& Other)
 	}
 	else
 	{
-	FPlatformMemory::Memcpy(GetData(), Other.GetData(), Size);
+		FPlatformMemory::Memcpy(GetData(), Other.GetData(), Size);
 	}
 
 #if VALIDATE_UNIFORM_BUFFER_LIFETIME
@@ -1402,7 +1386,7 @@ void FCachedPassMeshDrawListContext::FinalizeCommand(
 #if MESH_DRAW_COMMAND_DEBUG_DATA
 			if (CachedMeshDrawCommandStateBuckets.GetByElementId(SetId).Value.Num == 1)
 			{
-			MeshDrawCommand.ClearDebugPrimitiveSceneProxy(); //When using State Buckets multiple PrimitiveSceneProxies use the same MeshDrawCommand, so The PrimitiveSceneProxy pointer can't be stored.
+				MeshDrawCommand.ClearDebugPrimitiveSceneProxy(); //When using State Buckets multiple PrimitiveSceneProxies use the same MeshDrawCommand, so The PrimitiveSceneProxy pointer can't be stored.
 			}
 #endif
 		}
