Restored DiffuseIndirectAndAO scope

diff --git a/Engine/Source/Runtime/Renderer/Private/IndirectLightRendering.cpp b/Engine/Source/Runtime/Renderer/Private/IndirectLightRendering.cpp
--- a/Engine/Source/Runtime/Renderer/Private/IndirectLightRendering.cpp
+++ b/Engine/Source/Runtime/Renderer/Private/IndirectLightRendering.cpp
@@ -685,6 +685,8 @@ void FDeferredShadingSceneRenderer::RenderDiffuseIndirectAndAmbientOcclusion(
 		return;
 	}
 
+	RDG_EVENT_SCOPE(GraphBuilder, "DiffuseIndirectAndAO");
+
 	FSceneRenderTargets& SceneContext = FSceneRenderTargets::Get(GraphBuilder.RHICmdList);
 
 	FSceneTextureParameters SceneTextures = GetSceneTextureParameters(GraphBuilder, SceneTexturesUniformBuffer);
