Fixed global distance field visualization. Legacacy binding's don't support float3 array well.

diff --git a/Engine/Shaders/Private/GlobalDistanceFieldShared.ush b/Engine/Shaders/Private/GlobalDistanceFieldShared.ush
--- a/Engine/Shaders/Private/GlobalDistanceFieldShared.ush
+++ b/Engine/Shaders/Private/GlobalDistanceFieldShared.ush
@@ -11,9 +11,9 @@
 #endif
 
 #if DISTANCE_FIELD_IN_VIEW_UB
-#define MaxGlobalDFAOConeDistance View.MaxGlobalDFAOConeDistance
+	#define MaxGlobalDFAOConeDistance View.MaxGlobalDFAOConeDistance
 #else
-float MaxGlobalDFAOConeDistance;
+	float MaxGlobalDFAOConeDistance;
 #endif
 
 #if PLATFORM_SUPPORTS_DISTANCE_FIELDS
@@ -46,8 +46,8 @@ float MaxGlobalDFAOConeDistance;
 	Texture3D GlobalDistanceFieldMipTexture;
 	float4 GlobalVolumeCenterAndExtent[MAX_GLOBAL_DF_CLIPMAPS];
 	float4 GlobalVolumeWorldToUVAddAndMul[MAX_GLOBAL_DF_CLIPMAPS];
-	float3 GlobalDistanceFieldMipWorldToUVScale[MAX_GLOBAL_DF_CLIPMAPS];
-	float3 GlobalDistanceFieldMipWorldToUVBias[MAX_GLOBAL_DF_CLIPMAPS];
+	float4 GlobalDistanceFieldMipWorldToUVScale[MAX_GLOBAL_DF_CLIPMAPS];
+	float4 GlobalDistanceFieldMipWorldToUVBias[MAX_GLOBAL_DF_CLIPMAPS];
 	float GlobalDistanceFieldMipFactor;
 	float GlobalDistanceFieldMipTransition;
 	float3 GlobalDistanceFieldInvPageAtlasSize;
@@ -118,7 +118,7 @@ float3 ComputeGlobalUV(float3 WorldPosition, uint ClipmapIndex)
 
 float3 ComputeGlobalMipUV(float3 WorldPosition, uint ClipmapIndex)
 {
-	return saturate(WorldPosition * GlobalDistanceFieldMipWorldToUVScale[ClipmapIndex] + GlobalDistanceFieldMipWorldToUVBias[ClipmapIndex]);
+	return saturate(WorldPosition * GlobalDistanceFieldMipWorldToUVScale[ClipmapIndex].xyz + GlobalDistanceFieldMipWorldToUVBias[ClipmapIndex].xyz);
 }
 
 float GetDistanceToNearestSurfaceGlobal(float3 WorldPosition)
diff --git a/Engine/Source/Runtime/Engine/Public/SceneView.h b/Engine/Source/Runtime/Engine/Public/SceneView.h
--- a/Engine/Source/Runtime/Engine/Public/SceneView.h
+++ b/Engine/Source/Runtime/Engine/Public/SceneView.h
@@ -772,8 +772,8 @@ enum ETranslucencyVolumeCascade
 	VIEW_UNIFORM_BUFFER_MEMBER(int32, StereoPassIndex) \
 	VIEW_UNIFORM_BUFFER_MEMBER_ARRAY(FVector4, GlobalVolumeCenterAndExtent, [GMaxGlobalDistanceFieldClipmaps]) \
 	VIEW_UNIFORM_BUFFER_MEMBER_ARRAY(FVector4, GlobalVolumeWorldToUVAddAndMul, [GMaxGlobalDistanceFieldClipmaps]) \
-	VIEW_UNIFORM_BUFFER_MEMBER_ARRAY(FVector, GlobalDistanceFieldMipWorldToUVScale, [GMaxGlobalDistanceFieldClipmaps]) \
-	VIEW_UNIFORM_BUFFER_MEMBER_ARRAY(FVector, GlobalDistanceFieldMipWorldToUVBias, [GMaxGlobalDistanceFieldClipmaps]) \
+	VIEW_UNIFORM_BUFFER_MEMBER_ARRAY(FVector4, GlobalDistanceFieldMipWorldToUVScale, [GMaxGlobalDistanceFieldClipmaps]) \
+	VIEW_UNIFORM_BUFFER_MEMBER_ARRAY(FVector4, GlobalDistanceFieldMipWorldToUVBias, [GMaxGlobalDistanceFieldClipmaps]) \
 	VIEW_UNIFORM_BUFFER_MEMBER(float, GlobalDistanceFieldMipFactor) \
 	VIEW_UNIFORM_BUFFER_MEMBER(float, GlobalDistanceFieldMipTransition) \
 	VIEW_UNIFORM_BUFFER_MEMBER(int32, GlobalDistanceFieldClipmapSizeInPages) \
diff --git a/Engine/Source/Runtime/Renderer/Private/GlobalDistanceField.cpp b/Engine/Source/Runtime/Renderer/Private/GlobalDistanceField.cpp
--- a/Engine/Source/Runtime/Renderer/Private/GlobalDistanceField.cpp
+++ b/Engine/Source/Runtime/Renderer/Private/GlobalDistanceField.cpp
@@ -1038,8 +1038,8 @@ void FViewInfo::SetupDefaultGlobalDistanceFieldUniformBufferParameters(FViewUnif
 	{
 		ViewUniformShaderParameters.GlobalVolumeCenterAndExtent[Index] = FVector4(0);
 		ViewUniformShaderParameters.GlobalVolumeWorldToUVAddAndMul[Index] = FVector4(0);
-		ViewUniformShaderParameters.GlobalDistanceFieldMipWorldToUVScale[Index] = FVector(0);
-		ViewUniformShaderParameters.GlobalDistanceFieldMipWorldToUVBias[Index] = FVector(0);
+		ViewUniformShaderParameters.GlobalDistanceFieldMipWorldToUVScale[Index] = FVector4(0);
+		ViewUniformShaderParameters.GlobalDistanceFieldMipWorldToUVBias[Index] = FVector4(0);
 	}
 	ViewUniformShaderParameters.GlobalDistanceFieldMipFactor = 1.0f;
 	ViewUniformShaderParameters.GlobalDistanceFieldMipTransition = 0.0f;
diff --git a/Engine/Source/Runtime/Renderer/Public/GlobalDistanceFieldParameters.h b/Engine/Source/Runtime/Renderer/Public/GlobalDistanceFieldParameters.h
--- a/Engine/Source/Runtime/Renderer/Public/GlobalDistanceFieldParameters.h
+++ b/Engine/Source/Runtime/Renderer/Public/GlobalDistanceFieldParameters.h
@@ -27,8 +27,8 @@ public:
 
 	FVector4 CenterAndExtent[GMaxGlobalDistanceFieldClipmaps];
 	FVector4 WorldToUVAddAndMul[GMaxGlobalDistanceFieldClipmaps];
-	FVector MipWorldToUVScale[GMaxGlobalDistanceFieldClipmaps];
-	FVector MipWorldToUVBias[GMaxGlobalDistanceFieldClipmaps];
+	FVector4 MipWorldToUVScale[GMaxGlobalDistanceFieldClipmaps];
+	FVector4 MipWorldToUVBias[GMaxGlobalDistanceFieldClipmaps];
 	float MipFactor;
 	float MipTransition;
 	FVector PageTableScrollOffset[GMaxGlobalDistanceFieldClipmaps];
@@ -76,12 +76,12 @@ public:
 		Ar << Parameters.GlobalDistanceFieldPageAtlasTexture;
 		Ar << Parameters.GlobalDistanceFieldPageTableTexture;
 		Ar << Parameters.GlobalDistanceFieldMipTexture;
-		Ar << Parameters.GlobalDistanceFieldMipFactor;
-		Ar << Parameters.GlobalDistanceFieldMipTransition;
 		Ar << Parameters.GlobalVolumeCenterAndExtent;
 		Ar << Parameters.GlobalVolumeWorldToUVAddAndMul;
 		Ar << Parameters.GlobalDistanceFieldMipWorldToUVScale;
 		Ar << Parameters.GlobalDistanceFieldMipWorldToUVBias;
+		Ar << Parameters.GlobalDistanceFieldMipFactor;
+		Ar << Parameters.GlobalDistanceFieldMipTransition;
 		Ar << Parameters.GlobalDistanceFieldClipmapSizeInPages;
 		Ar << Parameters.GlobalDistanceFieldInvPageAtlasSize;
 		Ar << Parameters.GlobalVolumeDimension;
@@ -120,12 +120,12 @@ private:
 	LAYOUT_FIELD(FShaderResourceParameter, GlobalDistanceFieldPageAtlasTexture)
 	LAYOUT_FIELD(FShaderResourceParameter, GlobalDistanceFieldPageTableTexture)
 	LAYOUT_FIELD(FShaderResourceParameter, GlobalDistanceFieldMipTexture)
-	LAYOUT_FIELD(FShaderParameter, GlobalDistanceFieldMipFactor)
-	LAYOUT_FIELD(FShaderParameter, GlobalDistanceFieldMipTransition)
 	LAYOUT_FIELD(FShaderParameter, GlobalVolumeCenterAndExtent)
 	LAYOUT_FIELD(FShaderParameter, GlobalVolumeWorldToUVAddAndMul)
 	LAYOUT_FIELD(FShaderParameter, GlobalDistanceFieldMipWorldToUVScale)
 	LAYOUT_FIELD(FShaderParameter, GlobalDistanceFieldMipWorldToUVBias)
+	LAYOUT_FIELD(FShaderParameter, GlobalDistanceFieldMipFactor)
+	LAYOUT_FIELD(FShaderParameter, GlobalDistanceFieldMipTransition)
 	LAYOUT_FIELD(FShaderParameter, GlobalDistanceFieldClipmapSizeInPages)	
 	LAYOUT_FIELD(FShaderParameter, GlobalDistanceFieldInvPageAtlasSize)
 	LAYOUT_FIELD(FShaderParameter, GlobalVolumeDimension)
