Fix crashes related to GAsyncCreateLightPrimitiveInteractions when threading is disabled.

diff --git a/Engine/Source/Runtime/Renderer/Private/RendererScene.cpp b/Engine/Source/Runtime/Renderer/Private/RendererScene.cpp
--- a/Engine/Source/Runtime/Renderer/Private/RendererScene.cpp
+++ b/Engine/Source/Runtime/Renderer/Private/RendererScene.cpp
@@ -3758,7 +3758,7 @@ void FScene::UpdateAllPrimitiveSceneInfos(FRHICommandListImmediate& RHICmdList,
 	TSet<FPrimitiveSceneInfo*> DeletedSceneInfos;
 	DeletedSceneInfos.Reserve(RemovedLocalPrimitiveSceneInfos.Num());
 
-	if (!!GAsyncCreateLightPrimitiveInteractions && !AsyncCreateLightPrimitiveInteractionsTask)
+	if (!!GAsyncCreateLightPrimitiveInteractions && !AsyncCreateLightPrimitiveInteractionsTask && FApp::ShouldUseThreadingForPerformance())
 	{
 		AsyncCreateLightPrimitiveInteractionsTask = new FAsyncTask<FAsyncCreateLightPrimitiveInteractionsTask>();
 	}
@@ -4314,6 +4314,7 @@ void FScene::CreateLightPrimitiveInteractionsForPrimitive(FPrimitiveSceneInfo* P
 {
 	if (!bAsyncCreateLPIs
 		|| !GAsyncCreateLightPrimitiveInteractions
+		|| !FApp::ShouldUseThreadingForPerformance()
 		|| GetShadingPath() == EShadingPath::Mobile
 		|| PrimitiveInfo->LODParentComponentId.IsValid()
 		|| !PrimitiveInfo->Proxy->IsMeshShapeOftenMoving())
@@ -4344,7 +4345,7 @@ void FScene::CreateLightPrimitiveInteractionsForPrimitive(FPrimitiveSceneInfo* P
 
 void FScene::FlushAsyncLightPrimitiveInteractionCreation() const
 {
-	if (!GAsyncCreateLightPrimitiveInteractions)
+	if (!GAsyncCreateLightPrimitiveInteractions || !FApp::ShouldUseThreadingForPerformance())
 	{
 		check(!AsyncCreateLightPrimitiveInteractionsTask || AsyncCreateLightPrimitiveInteractionsTask->IsIdle());
 	}
