Removed PREVENT_RENDERTARGET_SIZE_THRASHING fully to fix rapid OOM on PS5, but also clean up the scene render target code, since all our platforms opt-in to this codepath.

Originally introduced in CL 4118928, but all platforms except Mac opted-in over time (even including iOS).

diff --git a/Engine/Source/Runtime/Renderer/Private/PostProcess/SceneRenderTargets.cpp b/Engine/Source/Runtime/Renderer/Private/PostProcess/SceneRenderTargets.cpp
--- a/Engine/Source/Runtime/Renderer/Private/PostProcess/SceneRenderTargets.cpp
+++ b/Engine/Source/Runtime/Renderer/Private/PostProcess/SceneRenderTargets.cpp
@@ -295,9 +295,8 @@ FSceneRenderTargets::FSceneRenderTargets(const FViewInfo& View, const FSceneRend
 	, bAllocatedFoveationTexture(SnapshotSource.bAllocatedFoveationTexture)
 {
 	FMemory::Memcpy(LargestDesiredSizes, SnapshotSource.LargestDesiredSizes);
-#if PREVENT_RENDERTARGET_SIZE_THRASHING
 	FMemory::Memcpy(HistoryFlags, SnapshotSource.HistoryFlags, sizeof(SnapshotSource.HistoryFlags));
-#endif
+
 	SnapshotArray(SceneColor, SnapshotSource.SceneColor);
 	SnapshotArray(ReflectionColorScratchCubemap, SnapshotSource.ReflectionColorScratchCubemap);
 	SnapshotArray(DiffuseIrradianceScratchCubemap, SnapshotSource.DiffuseIrradianceScratchCubemap);
@@ -317,8 +316,6 @@ inline const TCHAR* GetSceneColorTargetName(EShadingPath ShadingPath)
 	return SceneColorNames[(uint32)ShadingPath];
 }
 
-#if PREVENT_RENDERTARGET_SIZE_THRASHING
-
 namespace ERenderTargetHistory
 {
 	enum Type
@@ -348,17 +345,6 @@ static bool AnyCaptureRenderedRecently(const uint8* HistoryFlags, uint8 Mask)
 	return Result != 0;
 }
 
-#define UPDATE_HISTORY_FLAGS(Flags, bIsSceneCapture, bIsReflectionCapture, bIsHighResScreenShot) UpdateHistoryFlags(Flags, bIsSceneCapture, bIsReflectionCapture, bIsHighResScreenShot)
-#define ANY_CAPTURE_RENDERED_RECENTLY(HistoryFlags, NumEntries) AnyCaptureRenderedRecently<NumEntries>(HistoryFlags, ERenderTargetHistory::RTH_MaskAll)
-#define ANY_HIGHRES_CAPTURE_RENDERED_RECENTLY(HistoryFlags, NumEntries) AnyCaptureRenderedRecently<NumEntries>(HistoryFlags, ERenderTargetHistory::RTH_HighresScreenshot)
-
-#else
-
-#define UPDATE_HISTORY_FLAGS(Flags, bIsSceneCapture, bIsReflectionCapture, bIsHighResScreenShot)
-#define ANY_CAPTURE_RENDERED_RECENTLY(HistoryFlags, NumEntries) (false)
-#define ANY_HIGHRES_CAPTURE_RENDERED_RECENTLY(HistoryFlags, NumEntries) (false)
-#endif
-
 FIntPoint FSceneRenderTargets::ComputeDesiredSize(const FSceneViewFamily& ViewFamily)
 {
 	enum ESizingMethods { RequestedSize, ScreenRes, Grow, VisibleSizingMethodsCount };
@@ -480,18 +466,16 @@ FIntPoint FSceneRenderTargets::ComputeDesiredSize(const FSceneViewFamily& ViewFa
 		}
 		// this allows the BufferSize to shrink each frame (in game)
 		LargestDesiredSizes[CurrentDesiredSizeIndex] = FIntPoint::ZeroValue;
-#if PREVENT_RENDERTARGET_SIZE_THRASHING
 		HistoryFlags[CurrentDesiredSizeIndex] = 0;
-#endif
 	}
 
 	// this allows The BufferSize to not grow below the SceneCapture requests (happen before scene rendering, in the same frame with a Grow request)
 	FIntPoint& LargestDesiredSizeThisFrame = LargestDesiredSizes[CurrentDesiredSizeIndex];
 	LargestDesiredSizeThisFrame = LargestDesiredSizeThisFrame.ComponentMax(DesiredBufferSize);
 	bool bIsHighResScreenshot = GIsHighResScreenshot;
-	UPDATE_HISTORY_FLAGS(HistoryFlags[CurrentDesiredSizeIndex], bIsSceneCapture, bIsReflectionCapture, bIsHighResScreenshot);
+	UpdateHistoryFlags(HistoryFlags[CurrentDesiredSizeIndex], bIsSceneCapture, bIsReflectionCapture, bIsHighResScreenshot);
 
-	// we want to shrink the buffer but as we can have multiple scenecaptures per frame we have to delay that a frame to get all size requests
+	// we want to shrink the buffer but as we can have multiple scene captures per frame we have to delay that a frame to get all size requests
 	// Don't save buffer size in history while making high-res screenshot.
 	// We have to use the requested size when allocating an hmd depth target to ensure it matches the hmd allocated render target size.
 	bool bAllowDelayResize = !GIsHighResScreenshot && !bIsVRScene;
@@ -500,7 +484,7 @@ FIntPoint FSceneRenderTargets::ComputeDesiredSize(const FSceneViewFamily& ViewFa
 	// This prevents problems when orientation changes on mobile in particular.
 	// bIsReflectionCapture is explicitly checked on all platforms to prevent aspect ratio change detection from forcing the immediate buffer resize.
 	// This ensures that 1) buffers are not resized spuriously during reflection rendering 2) all cubemap faces use the same render target size.
-	if (bAllowDelayResize && !bIsReflectionCapture && !ANY_CAPTURE_RENDERED_RECENTLY(HistoryFlags, FrameSizeHistoryCount))
+	if (bAllowDelayResize && !bIsReflectionCapture && !AnyCaptureRenderedRecently<FrameSizeHistoryCount>(HistoryFlags, ERenderTargetHistory::RTH_MaskAll))
 	{
 		const bool bAspectRatioChanged =
 			!BufferSize.Y ||
@@ -519,19 +503,17 @@ FIntPoint FSceneRenderTargets::ComputeDesiredSize(const FSceneViewFamily& ViewFa
 			for (int32 i = 0; i < FrameSizeHistoryCount; ++i)
 			{
 				LargestDesiredSizes[i] = FIntPoint::ZeroValue;
-#if PREVENT_RENDERTARGET_SIZE_THRASHING
 				HistoryFlags[i] = 0;
-#endif
 			}
 		}
 	}
-	const bool bAnyHighresScreenshotRecently = ANY_HIGHRES_CAPTURE_RENDERED_RECENTLY(HistoryFlags, FrameSizeHistoryCount);
-	if(bAnyHighresScreenshotRecently != GIsHighResScreenshot)
+	const bool bAnyHighresScreenshotRecently = AnyCaptureRenderedRecently<FrameSizeHistoryCount>(HistoryFlags, ERenderTargetHistory::RTH_HighresScreenshot);
+	if (bAnyHighresScreenshotRecently != GIsHighResScreenshot)
 	{
 		bAllowDelayResize = false;
 	}
 
-	if(bAllowDelayResize)
+	if (bAllowDelayResize)
 	{
 		for (int32 i = 0; i < FrameSizeHistoryCount; ++i)
 		{
diff --git a/Engine/Source/Runtime/Renderer/Private/PostProcess/SceneRenderTargets.h b/Engine/Source/Runtime/Renderer/Private/PostProcess/SceneRenderTargets.h
--- a/Engine/Source/Runtime/Renderer/Private/PostProcess/SceneRenderTargets.h
+++ b/Engine/Source/Runtime/Renderer/Private/PostProcess/SceneRenderTargets.h
@@ -125,10 +125,6 @@ static const int32 NumTranslucencyShadowSurfaces = 2;
 // Store shading model into stencil [1-3] bits
 #define GET_STENCIL_MOBILE_SM_MASK(Value) uint8((Value & 0x7) << 1)
 
-#if !defined(PREVENT_RENDERTARGET_SIZE_THRASHING)
-	#define PREVENT_RENDERTARGET_SIZE_THRASHING (PLATFORM_DESKTOP || PLATFORM_PS4 || PLATFORM_ANDROID || PLATFORM_IOS || PLATFORM_SWITCH || PLATFORM_UNIX)
-#endif
-
 enum class ESceneColorFormatType
 {
 	Mobile,
@@ -198,14 +194,14 @@ protected:
 		bHMDAllocatedDepthTarget(false),
 		bKeepDepthContent(true),
 		bAllocatedFoveationTexture(false)
-		{
-			FMemory::Memset(LargestDesiredSizes, 0);
-#if PREVENT_RENDERTARGET_SIZE_THRASHING
-			FMemory::Memset(HistoryFlags, 0, sizeof(HistoryFlags));
-#endif
-		}
+	{
+		FMemory::Memset(LargestDesiredSizes, 0);
+		FMemory::Memset(HistoryFlags, 0, sizeof(HistoryFlags));
+	}
+	
 	/** Constructor that creates snapshot */
 	FSceneRenderTargets(const FViewInfo& InView, const FSceneRenderTargets& SnapshotSource);
+
 public:
 
 	bool IsShadingPathValid() const
@@ -546,11 +542,10 @@ private:
 	/** as we might get multiple BufferSize requests each frame for SceneCaptures and we want to avoid reallocations we can only go as low as the largest request */
 	static const uint32 FrameSizeHistoryCount = 3;
 	FIntPoint LargestDesiredSizes[FrameSizeHistoryCount];
-#if PREVENT_RENDERTARGET_SIZE_THRASHING
+
 	// bit 0 - whether any scene capture rendered
 	// bit 1 - whether any reflection capture rendered
 	uint8 HistoryFlags[FrameSizeHistoryCount];
-#endif
 
 	/** to detect when LargestDesiredSizeThisFrame is outdated */
 	uint32 ThisFrameNumber;
