Ported depth prepass to RDG.
 - Removed awkward AfterTasks lambda. Looks vestigial and isn't worth the code complexity.
 - Consolidated depth-related config data into a struct instead of loose variables on the deferred shading renderer.
 - Moved some internal functions to static instead of class members.

diff --git a/Engine/Shaders/Private/DitheredTransitionStencil.usf b/Engine/Shaders/Private/DitheredTransitionStencil.usf
--- a/Engine/Shaders/Private/DitheredTransitionStencil.usf
+++ b/Engine/Shaders/Private/DitheredTransitionStencil.usf
@@ -25,10 +25,8 @@ half GetDitherValue(float2 InPosition)
 	return (RetVal - 0.001f);
 }
 
-void Main(
-	FScreenVertexOutput Interpolators)
+void Main(float4 SvPosition : SV_POSITION)
 {
-	float4 SvPosition = Interpolators.Position;
 	if (abs(DitheredTransitionFactor) > 0.001f)
 	{
 		clip(GetDitherValue(SvPosition.xy));
diff --git a/Engine/Source/Runtime/RenderCore/Private/RenderGraphUtils.cpp b/Engine/Source/Runtime/RenderCore/Private/RenderGraphUtils.cpp
--- a/Engine/Source/Runtime/RenderCore/Private/RenderGraphUtils.cpp
+++ b/Engine/Source/Runtime/RenderCore/Private/RenderGraphUtils.cpp
@@ -557,6 +557,13 @@ void AddClearDepthStencilPass(
 	});
 }
 
+void AddClearDepthStencilPass(FRDGBuilder& GraphBuilder, FRDGTextureRef Texture, ERenderTargetLoadAction DepthLoadAction, ERenderTargetLoadAction StencilLoadAction)
+{
+	auto* PassParameters = GraphBuilder.AllocParameters<FRenderTargetParameters>();
+	PassParameters->RenderTargets.DepthStencil = FDepthStencilBinding(Texture, DepthLoadAction, StencilLoadAction, FExclusiveDepthStencil::DepthWrite_StencilWrite);
+	GraphBuilder.AddPass(RDG_EVENT_NAME("ClearDepthStencil (%s)", Texture->Name), PassParameters, ERDGPassFlags::Raster, [](FRHICommandList&) {});
+}
+
 void AddClearStencilPass(FRDGBuilder& GraphBuilder, FRDGTextureRef Texture)
 {
 	auto* PassParameters = GraphBuilder.AllocParameters<FRenderTargetParameters>();
diff --git a/Engine/Source/Runtime/RenderCore/Public/RenderGraphUtils.h b/Engine/Source/Runtime/RenderCore/Public/RenderGraphUtils.h
--- a/Engine/Source/Runtime/RenderCore/Public/RenderGraphUtils.h
+++ b/Engine/Source/Runtime/RenderCore/Public/RenderGraphUtils.h
@@ -478,6 +478,13 @@ RENDERCORE_API void AddClearDepthStencilPass(
 	bool bClearStencil,
 	uint8 Stencil);
 
+/** Adds a render graph pass to clear a depth stencil target to its optimized clear value using a raster pass. */
+RENDERCORE_API void AddClearDepthStencilPass(
+	FRDGBuilder& GraphBuilder,
+	FRDGTextureRef Texture,
+	ERenderTargetLoadAction DepthLoadAction = ERenderTargetLoadAction::EClear,
+	ERenderTargetLoadAction StencilLoadAction = ERenderTargetLoadAction::EClear);
+
 /** Adds a render graph pass to clear the stencil portion of a depth / stencil target to its fast clear value. */
 RENDERCORE_API void AddClearStencilPass(FRDGBuilder& GraphBuilder, FRDGTextureRef Texture);
 
diff --git a/Engine/Source/Runtime/Renderer/Private/DeferredShadingRenderer.cpp b/Engine/Source/Runtime/Renderer/Private/DeferredShadingRenderer.cpp
--- a/Engine/Source/Runtime/Renderer/Private/DeferredShadingRenderer.cpp
+++ b/Engine/Source/Runtime/Renderer/Private/DeferredShadingRenderer.cpp
@@ -60,23 +60,6 @@
 extern int32 GNaniteDebugFlags;
 extern int32 GNaniteShowStats;
 
-static TAutoConsoleVariable<int32> CVarStencilForLODDither(
-	TEXT("r.StencilForLODDither"),
-	0,
-	TEXT("Whether to use stencil tests in the prepass, and depth-equal tests in the base pass to implement LOD dithering.\n")
-	TEXT("If disabled, LOD dithering will be done through clip() instructions in the prepass and base pass, which disables EarlyZ.\n")
-	TEXT("Forces a full prepass when enabled."),
-	ECVF_RenderThreadSafe | ECVF_ReadOnly);
-
-static TAutoConsoleVariable<int32> CVarStencilLODDitherMode(
-	TEXT("r.StencilLODMode"),
-	2,
-	TEXT("Specifies the dither LOD stencil mode.\n")
-	TEXT(" 0: Graphics pass.\n")
-	TEXT(" 1: Compute pass (on supported platforms).\n")
-	TEXT(" 2: Compute async pass (on supported platforms)."),
-	ECVF_RenderThreadSafe);
-
 TAutoConsoleVariable<int32> CVarCustomDepthOrder(
 	TEXT("r.CustomDepth.Order"),
 	1,	
@@ -299,13 +282,9 @@ const TCHAR* GetDepthPassReason(bool bDitheredLODTransitionsUseStencil, EShaderP
 
 FDeferredShadingSceneRenderer::FDeferredShadingSceneRenderer(const FSceneViewFamily* InViewFamily,FHitProxyConsumer* HitProxyConsumer)
 	: FSceneRenderer(InViewFamily, HitProxyConsumer)
-	, EarlyZPassMode(Scene ? Scene->EarlyZPassMode : DDM_None)
-	, bEarlyZPassMovable(Scene ? Scene->bEarlyZPassMovable : false)
+	, DepthPass(GetDepthPassInfo(Scene))
 	, bAreLightsInLightGrid(false)
-{
-	static const auto StencilLODDitherCVar = IConsoleManager::Get().FindTConsoleVariableDataInt(TEXT("r.StencilForLODDither"));
-	bDitheredLODTransitionsUseStencil = StencilLODDitherCVar->GetValueOnAnyThread() != 0;
-}
+{}
 
 /** 
 * Renders the view family. 
@@ -344,7 +323,7 @@ TGlobalResource<FGlobalDynamicReadBuffer> FDeferredShadingSceneRenderer::Dynamic
 static FORCEINLINE bool NeedsPrePass(const FDeferredShadingSceneRenderer* Renderer)
 {
 	return (RHIHasTiledGPU(Renderer->ViewFamily.GetShaderPlatform()) == false) && 
-		(Renderer->EarlyZPassMode != DDM_None || Renderer->bEarlyZPassMovable != 0);
+		(Renderer->DepthPass.EarlyZPassMode != DDM_None || Renderer->DepthPass.bEarlyZPassMovable != 0);
 }
 
 bool FDeferredShadingSceneRenderer::RenderHzb(
@@ -1512,7 +1491,7 @@ void FDeferredShadingSceneRenderer::Render(FRHICommandListImmediate& RHICmdList)
 		&& !ViewFamily.EngineShowFlags.LightMapDensity;
 
 	// Use read-only depth in the base pass if we have a full depth prepass.
-	const bool bAllowReadOnlyDepthBasePass = EarlyZPassMode == DDM_AllOpaque
+	const bool bAllowReadOnlyDepthBasePass = DepthPass.EarlyZPassMode == DDM_AllOpaque
 		&& !ViewFamily.EngineShowFlags.ShaderComplexity
 		&& !ViewFamily.UseDebugViewPS()
 		&& !ViewFamily.EngineShowFlags.Wireframe
@@ -1631,6 +1610,14 @@ void FDeferredShadingSceneRenderer::Render(FRHICommandListImmediate& RHICmdList)
 		}
 	}
 
+	FRDGBuilder GraphBuilder(RHICmdList, RDG_EVENT_NAME("DeferredShadingRenderer_Render(ViewFamily=%s)", ViewFamily.bResolveScene ? TEXT("Primary") : TEXT("Auxiliary")));
+
+	const FIntPoint SceneTextureExtent = SceneContext.GetBufferSizeXY();
+
+	FRDGTextureMSAA SceneDepthTexture = RegisterExternalTextureMSAA(GraphBuilder, SceneContext.SceneDepthZ);
+	FRDGTextureRef SmallDepthTexture = GraphBuilder.RegisterExternalTexture(SceneContext.SmallDepthZ, ERenderTargetTexture::Targetable);
+	FRDGTextureSRVRef SceneStencilTexture = GraphBuilder.CreateSRV(FRDGTextureSRVDesc::CreateWithPixelFormat(SceneDepthTexture.Target, PF_X24_G8));
+
 	const bool bUseGBuffer = IsUsingGBuffers(ShaderPlatform);
 	bool bCanOverlayRayTracingOutput = CanOverlayRayTracingOutput(Views[0]);// #dxr_todo: UE-72557 multi-view case
 	
@@ -1681,6 +1668,8 @@ void FDeferredShadingSceneRenderer::Render(FRHICommandListImmediate& RHICmdList)
 
 	// Dynamic vertex and index buffers need to be committed before rendering.
 	GEngine->GetPreRenderDelegate().Broadcast();
+
+	AddPass(GraphBuilder, [this, bDoInitViewAftersPrepass](FRHICommandList&)
 	{
 		SCOPE_CYCLE_COUNTER(STAT_FDeferredShadingSceneRenderer_FGlobalDynamicVertexBuffer_Commit);
 		DynamicIndexBufferForInitViews.Commit();
@@ -1693,116 +1682,36 @@ void FDeferredShadingSceneRenderer::Render(FRHICommandListImmediate& RHICmdList)
 			DynamicIndexBufferForInitShadows.Commit();
 			DynamicReadBufferForInitShadows.Commit();
 		}
-	}
+	});
 
-	StencilLODMode = CVarStencilLODDitherMode.GetValueOnRenderThread();
-	if (!GRHISupportsDepthUAV)
-	{
-		// RHI doesn't support depth/stencil UAVs - enforce graphics path
-		StencilLODMode = 0;
-	}
-	else if (StencilLODMode == 2 && !GSupportsEfficientAsyncCompute)
-	{
-		// Async compute is not supported, fall back to compute path (on graphics queue)
-		StencilLODMode = 1;
-	}
-	else if (IsHMDHiddenAreaMaskActive())
+	if (DepthPass.IsComputeStencilDitherEnabled())
 	{
-		// Unsupported mode for compute path - enforce graphics path on VR
-		StencilLODMode = 0;
-	}
-
-	const bool bStencilLODCompute = (StencilLODMode == 1 || StencilLODMode == 2);
-	const bool bStencilLODComputeAsync = StencilLODMode == 2;
-
-	const FRHITransition* AsyncDitherLODEndTransition = nullptr;
-	if (bStencilLODCompute && bDitheredLODTransitionsUseStencil)
-	{
-		// Either compute pass will happen prior to the prepass, and the
-		// stencil clear will be skipped there.
-		FUnorderedAccessViewRHIRef StencilTextureUAV = RHICreateUnorderedAccessViewStencil(SceneContext.GetSceneDepthSurface(), 0 /* Mip Level */);
-
-		FRHITransitionInfo StencilUAVTransition(StencilTextureUAV, ERHIAccess::Unknown, ERHIAccess::ERWBarrier);
-
-		if (bStencilLODComputeAsync)
-		{
-			const FRHITransition* GFxToAsyncTransition = RHICreateTransition(ERHIPipeline::Graphics, ERHIPipeline::AsyncCompute, ERHICreateTransitionFlags::None, MakeArrayView(&StencilUAVTransition, 1));
-			AsyncDitherLODEndTransition = RHICreateTransition(ERHIPipeline::AsyncCompute, ERHIPipeline::Graphics, ERHICreateTransitionFlags::None, MakeArrayView(&StencilUAVTransition, 1));
-
-			RHICmdList.BeginTransition(GFxToAsyncTransition);
-
-			// Wait for the gfx transition to finish
-			FRHIAsyncComputeCommandListImmediate& RHICmdListComputeImmediate = FRHICommandListExecutor::GetImmediateAsyncComputeCommandList();
-			SCOPED_COMPUTE_EVENT(RHICmdListComputeImmediate, AsyncSSAO);
-			RHICmdListComputeImmediate.EndTransition(GFxToAsyncTransition);
-			
-			PreRenderDitherFill(RHICmdListComputeImmediate, SceneContext, StencilTextureUAV);
-
-			RHICmdListComputeImmediate.BeginTransition(AsyncDitherLODEndTransition);
-			FRHIAsyncComputeCommandListImmediate::ImmediateDispatch(RHICmdListComputeImmediate);
-		}
-		else
-		{
-			// Transition without pipeline crossing
-			RHICmdList.Transition(StencilUAVTransition);
-			PreRenderDitherFill(RHICmdList, SceneContext, StencilTextureUAV);
-		}
+		AddDitheredStencilFillPass(GraphBuilder, Views, SceneDepthTexture.Target, DepthPass);
 	}
 
 	// Notify the FX system that the scene is about to be rendered.
 	if (FXSystem && Views.IsValidIndex(0))
 	{
 		SCOPE_CYCLE_COUNTER(STAT_FDeferredShadingSceneRenderer_FXSystem_PreRender);
-		RHICmdList.SetCurrentStat(GET_STATID(STAT_CLM_FXPreRender));
-		FXSystem->PreRender(RHICmdList, Views[0].ViewUniformBuffer, &Views[0].GlobalDistanceFieldInfo.ParameterData, Views[0].AllowGPUParticleUpdate());
-		if (FGPUSortManager* GPUSortManager = FXSystem->GetGPUSortManager())
-		{
-			GPUSortManager->OnPreRender(RHICmdList);
-		}
-	}
-
-	if (AsyncDitherLODEndTransition)
-	{
-		RHICmdList.EndTransition(AsyncDitherLODEndTransition);
-	}
-
-	bool bDidAfterTaskWork = false;
-	auto AfterTasksAreStarted = [&bDidAfterTaskWork, bDoInitViewAftersPrepass, this, &RHICmdList, &ILCTaskData]()
-	{
-		if (!bDidAfterTaskWork)
+		AddSetCurrentStatPass(GraphBuilder, GET_STATID(STAT_CLM_FXPreRender));
+		AddPass(GraphBuilder, [this](FRHICommandListImmediate& RHICmdList)
 		{
-			QUICK_SCOPE_CYCLE_COUNTER(STAT_AfterPrepassTasksWork);
-			bDidAfterTaskWork = true; // only do this once
-			if (bDoInitViewAftersPrepass)
+			FXSystem->PreRender(RHICmdList, Views[0].ViewUniformBuffer, &Views[0].GlobalDistanceFieldInfo.ParameterData, Views[0].AllowGPUParticleUpdate());
+			if (FGPUSortManager* GPUSortManager = FXSystem->GetGPUSortManager())
 			{
-				{
-					SCOPED_GPU_STAT(RHICmdList, VisibilityCommands);
-					InitViewsPossiblyAfterPrepass(RHICmdList, ILCTaskData);
-				}
-				
-				{
-					SCOPED_GPU_STAT(RHICmdList, GPUSceneUpdate);
-					PrepareDistanceFieldScene(RHICmdList, false);
-				}
-
-				{
-					SCOPED_GPU_STAT(RHICmdList, UploadDynamicBuffers);
-					SCOPE_CYCLE_COUNTER(STAT_FDeferredShadingSceneRenderer_FGlobalDynamicVertexBuffer_Commit);
-					DynamicVertexBufferForInitShadows.Commit();
-					DynamicIndexBufferForInitShadows.Commit();
-					DynamicReadBufferForInitShadows.Commit();
-				}
-
-				ServiceLocalQueue();
+				GPUSortManager->OnPreRender(RHICmdList);
 			}
-		}
-	};
+		});
+	}
 
-	RunGPUSkinCacheTransition(RHICmdList, Scene, EGPUSkinCacheTransition::Renderer);
+	AddPass(GraphBuilder, [this](FRHICommandList& RHICmdList)
+	{
+		RunGPUSkinCacheTransition(RHICmdList, Scene, EGPUSkinCacheTransition::Renderer);
+	});
 
 	FHairStrandsBookmarkParameters HairStrandsBookmarkParameters;
+
 	{
-		FRDGBuilder GraphBuilder(RHICmdList, RDG_EVENT_NAME("HairStrandsCullingAndInterpolation(ViewFamily=%s)", ViewFamily.bResolveScene ? TEXT("Primary") : TEXT("Auxiliary")));
 		if (IsHairStrandsEnabled(EHairStrandsShaderType::All, Scene->GetShaderPlatform()))
 		{
 			HairStrandsBookmarkParameters = CreateHairStrandsBookmarkParameters(Views[0]);
@@ -1826,55 +1735,83 @@ void FDeferredShadingSceneRenderer::Render(FRHICommandListImmediate& RHICmdList)
 
 			RunHairStrandsBookmark(GraphBuilder, EHairStrandsBookmark::ProcessStrandsInterpolation, HairStrandsBookmarkParameters);
 		}
-		GraphBuilder.Execute();
 	}
 
-	checkSlow(RHICmdList.IsOutsideRenderPass());
-
 	if (bNaniteEnabled)
 	{
 		Nanite::ListStatFilters(this);
-		Nanite::GStreamingManager.EndAsyncUpdate(RHICmdList);	// Must happen before any Nanite rendering in the frame
-	}
 
-	// The Z-prepass
-
-	// Draw the scene pre-pass / early z pass, populating the scene depth buffer and HiZ
-	GRenderTargetPool.AddPhaseEvent(TEXT("EarlyZPass"));
-	bool bDepthWasCleared = false;
-	if (bNeedsPrePass)
-	{
-		bDepthWasCleared = RenderPrePass(RHICmdList, AfterTasksAreStarted);
+		AddPass(GraphBuilder, [](FRHICommandListImmediate& RHICmdList)
+		{
+			// Must happen before any Nanite rendering in the frame
+			Nanite::GStreamingManager.EndAsyncUpdate(RHICmdList); 
+		});
 	}
-	else
+
 	{
-		// we didn't do the prepass, but we still want the HMD mask if there is one
-		AfterTasksAreStarted();
-		RHICmdList.SetCurrentStat(GET_STATID(STAT_CLM_PrePass));
-		bDepthWasCleared = RenderPrePassHMD(RHICmdList);
-	}
-	check(bDidAfterTaskWork);
-	RHICmdList.SetCurrentStat(GET_STATID(STAT_CLM_AfterPrePass));
-	ServiceLocalQueue();
+		AddSetCurrentStatPass(GraphBuilder, GET_STATID(STAT_CLM_PrePass));
 
-#if RHI_RAYTRACING
-	// Must be done after FGlobalDynamicVertexBuffer::Get().Commit() for dynamic geometries to be updated
-	DispatchRayTracingWorldUpdates(RHICmdList);
-#endif
+		// Both compute approaches run earlier, so skip clearing stencil here, just load existing.
+		const ERenderTargetLoadAction StencilLoadAction = DepthPass.IsComputeStencilDitherEnabled()
+			? ERenderTargetLoadAction::ELoad
+			: ERenderTargetLoadAction::EClear;
 
-	// Z-Prepass End
+		const ERenderTargetLoadAction DepthLoadAction = ERenderTargetLoadAction::EClear;
+		AddClearDepthStencilPass(GraphBuilder, SceneDepthTexture.Target, DepthLoadAction, StencilLoadAction);
 
+		// Draw the scene pre-pass / early z pass, populating the scene depth buffer and HiZ
+		if (bNeedsPrePass)
+		{
+			RenderPrePass(GraphBuilder, SceneDepthTexture.Target);
+		}
+		else
+		{
+			// We didn't do the prepass, but we still want the HMD mask if there is one
+			RenderPrePassHMD(GraphBuilder, SceneDepthTexture.Target);
+		}
+		AddSetCurrentStatPass(GraphBuilder, GET_STATID(STAT_CLM_AfterPrePass));
+		AddServiceLocalQueuePass(GraphBuilder);
 
-	FRDGBuilder GraphBuilder(RHICmdList, RDG_EVENT_NAME("DeferredShadingRenderer_Render(ViewFamily=%s)", ViewFamily.bResolveScene ? TEXT("Primary") : TEXT("Auxiliary")));
+		if (bDoInitViewAftersPrepass)
+		{
+			{
+				RDG_GPU_STAT_SCOPE(GraphBuilder, VisibilityCommands);
+				AddPass(GraphBuilder, [this, &ILCTaskData](FRHICommandListImmediate& RHICmdList)
+				{
+					InitViewsPossiblyAfterPrepass(RHICmdList, ILCTaskData);
+				});
+			}
 
-	const FIntPoint SceneTextureExtent = SceneContext.GetBufferSizeXY();
+			{
+				RDG_GPU_STAT_SCOPE(GraphBuilder, GPUSceneUpdate);
+				AddPass(GraphBuilder, [this](FRHICommandListImmediate& RHICmdList)
+				{
+					PrepareDistanceFieldScene(RHICmdList, false);
+				});
+			}
 
-	FRDGTextureMSAA SceneDepthTexture = RegisterExternalTextureMSAA(GraphBuilder, SceneContext.SceneDepthZ);
-	FRDGTextureRef SmallDepthTexture = GraphBuilder.RegisterExternalTexture(SceneContext.SmallDepthZ, ERenderTargetTexture::Targetable);
-	FRDGTextureSRVRef SceneStencilTexture = GraphBuilder.CreateSRV(FRDGTextureSRVDesc::CreateWithPixelFormat(SceneDepthTexture.Target, PF_X24_G8));
+			{
+				RDG_GPU_STAT_SCOPE(GraphBuilder, UploadDynamicBuffers);
+				AddPass(GraphBuilder, [this](FRHICommandList&)
+				{
+					SCOPE_CYCLE_COUNTER(STAT_FDeferredShadingSceneRenderer_FGlobalDynamicVertexBuffer_Commit);
+					DynamicVertexBufferForInitShadows.Commit();
+					DynamicIndexBufferForInitShadows.Commit();
+					DynamicReadBufferForInitShadows.Commit();
+				});
+			}
 
-	ESceneTextureSetupMode SceneTexturesSetupMode = ESceneTextureSetupMode::SceneDepth;
-	TRDGUniformBufferRef<FSceneTextureUniformParameters> SceneTextures = CreateSceneTextureUniformBuffer(GraphBuilder, FeatureLevel, SceneTexturesSetupMode);
+			AddServiceLocalQueuePass(GraphBuilder);
+		}
+	}
+
+#if RHI_RAYTRACING
+	AddPass(GraphBuilder, [this](FRHICommandListImmediate& RHICmdList)
+	{
+		// Must be done after FGlobalDynamicVertexBuffer::Get().Commit() for dynamic geometries to be updated
+		DispatchRayTracingWorldUpdates(RHICmdList);
+	});
+#endif
 
 	TArray<Nanite::FRasterResults, TInlineAllocator<2>> NaniteRasterResults;
 	if (bNaniteEnabled && Views.Num() > 0)
@@ -1908,7 +1845,7 @@ void FDeferredShadingSceneRenderer::Render(FRHICommandListImmediate& RHICmdList)
 				Nanite::FCullingContext CullingContext = Nanite::InitCullingContext(
 					GraphBuilder,
 					*Scene,
-					EarlyZPassMode != DDM_AllOpaque ? View.PrevViewInfo.NaniteHZB : View.PrevViewInfo.HZB,
+					DepthPass.EarlyZPassMode != DDM_AllOpaque ? View.PrevViewInfo.NaniteHZB : View.PrevViewInfo.HZB,
 					View.PrevViewInfo.ViewRect,
 					bTwoPassOcclusion,
 					bUpdateStreaming,
@@ -1932,7 +1869,7 @@ void FDeferredShadingSceneRenderer::Render(FRHICommandListImmediate& RHICmdList)
 					bExtractStats
 				);
 
-				if( EarlyZPassMode != DDM_AllOpaque && bTwoPassOcclusion && View.ViewState )
+				if (DepthPass.EarlyZPassMode != DDM_AllOpaque && bTwoPassOcclusion && View.ViewState)
 				{
 					// Won't have a complete SceneDepth for post pass so can't use complete HZB for main pass or it will poke holes in the post pass HZB killing occlusion culling.
 					RDG_EVENT_SCOPE(GraphBuilder, "BuildNaniteHZB");
@@ -1961,6 +1898,9 @@ void FDeferredShadingSceneRenderer::Render(FRHICommandListImmediate& RHICmdList)
 		}
 	}
 
+	ESceneTextureSetupMode SceneTexturesSetupMode = ESceneTextureSetupMode::SceneDepth;
+	TRDGUniformBufferRef<FSceneTextureUniformParameters> SceneTextures = CreateSceneTextureUniformBuffer(GraphBuilder, FeatureLevel, SceneTexturesSetupMode);
+
 	const bool bShouldRenderVelocities = ShouldRenderVelocities();
 	const bool bBasePassCanOutputVelocity = FVelocityRendering::BasePassCanOutputVelocity(FeatureLevel);
 	const bool bUseSelectiveBasePassOutputs = IsUsingSelectiveBasePassOutputs(ShaderPlatform);
@@ -2013,8 +1953,9 @@ void FDeferredShadingSceneRenderer::Render(FRHICommandListImmediate& RHICmdList)
 			}
 		}
 	};
+
 	// Early occlusion queries
-	const bool bOcclusionBeforeBasePass = !bNaniteEnabled && ((EarlyZPassMode == EDepthDrawingMode::DDM_AllOccluders) || (EarlyZPassMode == EDepthDrawingMode::DDM_AllOpaque));
+	const bool bOcclusionBeforeBasePass = !bNaniteEnabled && ((DepthPass.EarlyZPassMode == EDepthDrawingMode::DDM_AllOccluders) || (DepthPass.EarlyZPassMode == EDepthDrawingMode::DDM_AllOpaque));
 
 	if (bOcclusionBeforeBasePass)
 	{
@@ -2146,20 +2087,6 @@ void FDeferredShadingSceneRenderer::Render(FRHICommandListImmediate& RHICmdList)
 				SSAOLevels = 0;
 			}
 
-#if 0
-			// Disable Nanite Pre Pass to prevent DBuffer decal rendering on Nanite meshes.
-			// Enable once all platforms support pre pass depth output.
-			if (bNaniteEnabled && 0) // TODO: Nanite dbuffer decal support
-			{
-				Nanite::DrawPrePass(
-					RHICmdList,
-					*Scene,
-					Views[ViewIndex],
-					NaniteRasterResults[ViewIndex]
-					);
-			}
-#endif
-
 			GCompositionLighting.ProcessBeforeBasePass(GraphBuilder, Scene->UniformBuffers, View, SceneTextures, bDBuffer, SSAOLevels);
 		}
 		AddServiceLocalQueuePass(GraphBuilder);
@@ -2208,8 +2135,7 @@ void FDeferredShadingSceneRenderer::Render(FRHICommandListImmediate& RHICmdList)
 	FRDGTextureMSAA SceneColorTexture = RegisterExternalTextureMSAA(GraphBuilder, SceneContext.GetSceneColor());
 
 	{
-		const ERenderTargetLoadAction DepthLoadAction = bDepthWasCleared ? ERenderTargetLoadAction::ELoad : ERenderTargetLoadAction::EClear;
-		RenderBasePass(GraphBuilder, BasePassDepthStencilAccess, SceneColorTexture.Target, SceneDepthTexture.Target, DepthLoadAction, ForwardScreenSpaceShadowMaskTexture);
+		RenderBasePass(GraphBuilder, BasePassDepthStencilAccess, SceneColorTexture.Target, SceneDepthTexture.Target, ERenderTargetLoadAction::ELoad, ForwardScreenSpaceShadowMaskTexture);
 		AddServiceLocalQueuePass(GraphBuilder);
 		
 		if (bNaniteEnabled && bShouldApplyNaniteMaterials)
diff --git a/Engine/Source/Runtime/Renderer/Private/DeferredShadingRenderer.h b/Engine/Source/Runtime/Renderer/Private/DeferredShadingRenderer.h
--- a/Engine/Source/Runtime/Renderer/Private/DeferredShadingRenderer.h
+++ b/Engine/Source/Runtime/Renderer/Private/DeferredShadingRenderer.h
@@ -162,10 +162,7 @@ class FDeferredShadingSceneRenderer : public FSceneRenderer
 {
 public:
 	/** Defines which objects we want to render in the EarlyZPass. */
-	EDepthDrawingMode EarlyZPassMode;
-	bool bEarlyZPassMovable;
-	bool bDitheredLODTransitionsUseStencil;
-	int32 StencilLODMode = 0;
+	FDepthPassInfo DepthPass;
 	
 	const FRHITransition* TranslucencyLightingVolumeClearEndTransition = nullptr;
 
@@ -403,31 +400,10 @@ private:
 
 	void ComputeLumenTranslucencyGIVolume(FRDGBuilder& GraphBuilder, FLumenCardTracingInputs& TracingInputs, FGlobalShaderMap* GlobalShaderMap);
 
-		
 	void CreateIndirectCapsuleShadows();
 
-	/**
-	* Setup the prepass. This is split out so that in parallel we can do the fx prerender after we start the parallel tasks
-	* @return true if the depth was cleared
-	*/
-	bool PreRenderPrePass(FRHICommandListImmediate& RHICmdList);
-
-	void PreRenderDitherFill(FRHIAsyncComputeCommandListImmediate& RHICmdList, FSceneRenderTargets& SceneContext, FRHIUnorderedAccessView* StencilTextureUAV);
-	void PreRenderDitherFill(FRHICommandListImmediate& RHICmdList, FSceneRenderTargets& SceneContext, FRHIUnorderedAccessView* StencilTextureUAV);
-
-	void RenderPrePassEditorPrimitives(FRHICommandList& RHICmdList, const FViewInfo& View, const FMeshPassProcessorRenderState& DrawRenderState, EDepthDrawingMode DepthDrawingMode, bool bRespectUseAsOccluderFlag);
-
-	/**
-	 * Renders the scene's prepass and occlusion queries.
-	 * @return true if the depth was cleared
-	 */
-	bool RenderPrePass(FRHICommandListImmediate& RHICmdList, TFunctionRef<void()> AfterTasksAreStarted);
-
-	/**
-	 * Renders the active HMD's hidden area mask as a depth prepass, if available.
-	 * @return true if depth is cleared
-	 */
-	bool RenderPrePassHMD(FRHICommandListImmediate& RHICmdList);
+	void RenderPrePass(FRDGBuilder& GraphBuilder, FRDGTextureRef SceneDepthTexture);
+	void RenderPrePassHMD(FRDGBuilder& GraphBuilder, FRDGTextureRef SceneDepthTexture);
 
 	void RenderFog(
 		FRDGBuilder& GraphBuilder,
diff --git a/Engine/Source/Runtime/Renderer/Private/DepthRendering.cpp b/Engine/Source/Runtime/Renderer/Private/DepthRendering.cpp
--- a/Engine/Source/Runtime/Renderer/Private/DepthRendering.cpp
+++ b/Engine/Source/Runtime/Renderer/Private/DepthRendering.cpp
@@ -28,6 +28,8 @@
 #include "ClearQuad.h"
 #include "GPUSkinCache.h"
 #include "MeshPassProcessor.inl"
+#include "PixelShaderUtils.h"
+#include "RenderGraphUtils.h"
 
 static TAutoConsoleVariable<int32> CVarParallelPrePass(
 	TEXT("r.ParallelPrePass"),
@@ -48,6 +50,60 @@ static FAutoConsoleVariableRef CVarSortPrepassMasked(
 	ECVF_Default
 );
 
+static TAutoConsoleVariable<int32> CVarStencilLODDitherMode(
+	TEXT("r.StencilLODMode"),
+	2,
+	TEXT("Specifies the dither LOD stencil mode.\n")
+	TEXT(" 0: Graphics pass.\n")
+	TEXT(" 1: Compute pass (on supported platforms).\n")
+	TEXT(" 2: Compute async pass (on supported platforms)."),
+	ECVF_RenderThreadSafe);
+
+static TAutoConsoleVariable<int32> CVarStencilForLODDither(
+	TEXT("r.StencilForLODDither"),
+	0,
+	TEXT("Whether to use stencil tests in the prepass, and depth-equal tests in the base pass to implement LOD dithering.\n")
+	TEXT("If disabled, LOD dithering will be done through clip() instructions in the prepass and base pass, which disables EarlyZ.\n")
+	TEXT("Forces a full prepass when enabled."),
+	ECVF_RenderThreadSafe | ECVF_ReadOnly);
+
+extern bool IsHMDHiddenAreaMaskActive();
+
+FDepthPassInfo GetDepthPassInfo(const FScene* Scene)
+{
+	FDepthPassInfo Info;
+	Info.EarlyZPassMode = Scene ? Scene->EarlyZPassMode : DDM_None;
+	Info.bEarlyZPassMovable = Scene ? Scene->bEarlyZPassMovable : false;
+	Info.bDitheredLODTransitionsUseStencil = CVarStencilForLODDither.GetValueOnAnyThread() > 0;
+	Info.StencilDitherPassFlags = ERDGPassFlags::Raster;
+
+	if (GRHISupportsDepthUAV && !IsHMDHiddenAreaMaskActive())
+	{
+		switch (CVarStencilLODDitherMode.GetValueOnAnyThread())
+		{
+		case 1:
+			Info.StencilDitherPassFlags = ERDGPassFlags::Compute;
+			break;
+		case 2:
+			Info.StencilDitherPassFlags = ERDGPassFlags::AsyncCompute;
+			break;
+		}
+	}
+
+	return Info;
+}
+
+BEGIN_SHADER_PARAMETER_STRUCT(FDepthPassParameters, )
+	RENDER_TARGET_BINDING_SLOTS()
+END_SHADER_PARAMETER_STRUCT()
+
+FDepthPassParameters* GetDepthPassParameters(FRDGBuilder& GraphBuilder, FRDGTextureRef DepthTexture)
+{
+	auto* PassParameters = GraphBuilder.AllocParameters<FDepthPassParameters>();
+	PassParameters->RenderTargets.DepthStencil = FDepthStencilBinding(DepthTexture, ERenderTargetLoadAction::ELoad, ERenderTargetLoadAction::ELoad, FExclusiveDepthStencil::DepthWrite_StencilWrite);
+	return PassParameters;
+}
+
 const TCHAR* GetDepthDrawingModeString(EDepthDrawingMode Mode)
 {
 	switch (Mode)
@@ -82,6 +138,11 @@ IMPLEMENT_SHADERPIPELINE_TYPE_VS(DepthPosOnlyNoPixelPipeline, TDepthOnlyVS<true>
 IMPLEMENT_SHADERPIPELINE_TYPE_VSPS(DepthNoColorOutputPipeline, TDepthOnlyVS<false>, FDepthOnlyPS<false>, true);
 IMPLEMENT_SHADERPIPELINE_TYPE_VSPS(DepthWithColorOutputPipeline, TDepthOnlyVS<false>, FDepthOnlyPS<true>, true);
 
+static bool IsDepthPassWaitForTasksEnabled()
+{
+	return CVarRHICmdFlushRenderThreadTasksPrePass.GetValueOnRenderThread() > 0 || CVarRHICmdFlushRenderThreadTasks.GetValueOnRenderThread() > 0;
+}
+
 static FORCEINLINE bool UseShaderPipelines(ERHIFeatureLevel::Type InFeatureLevel)
 {
 	static const auto* CVar = IConsoleManager::Get().FindTConsoleVariableDataInt(TEXT("r.ShaderPipelines"));
@@ -213,348 +274,145 @@ void SetDepthPassDitheredLODTransitionState(const FSceneView* SceneView, const F
 	}
 }
 
-static void SetupPrePassView(FRHICommandList& RHICmdList, const FViewInfo& View, const FSceneRenderer* SceneRenderer, const bool bIsEditorPrimitivePass = false)
-{
-	RHICmdList.SetScissorRect(false, 0, 0, 0, 0);
-
-	if (bIsEditorPrimitivePass)
-	{
-		RHICmdList.SetViewport(View.ViewRect.Min.X, View.ViewRect.Min.Y, 0.0f, View.ViewRect.Max.X, View.ViewRect.Max.Y, 1.0f);
-	}
-	else
-	{
-		SceneRenderer->SetStereoViewport(RHICmdList, View);
-	}
-}
-
-static void RenderHiddenAreaMaskView(FRHICommandList& RHICmdList, FGraphicsPipelineStateInitializer& GraphicsPSOInit, const FViewInfo& View)
-{
-	const auto FeatureLevel = GMaxRHIFeatureLevel;
-	const auto ShaderMap = GetGlobalShaderMap(FeatureLevel);
-	TShaderMapRef<TOneColorVS<true> > VertexShader(ShaderMap);
-
-	extern TGlobalResource<FFilterVertexDeclaration> GFilterVertexDeclaration;
-	GraphicsPSOInit.BoundShaderState.VertexDeclarationRHI = GFilterVertexDeclaration.VertexDeclarationRHI;
-	GraphicsPSOInit.BoundShaderState.VertexShaderRHI = VertexShader.GetVertexShader();
-	GraphicsPSOInit.PrimitiveType = PT_TriangleList;
-
-	SetGraphicsPipelineState(RHICmdList, GraphicsPSOInit);
-	VertexShader->SetDepthParameter(RHICmdList, 1.0f);
-
-	if (GEngine->XRSystem->GetHMDDevice())
-	{
-		GEngine->XRSystem->GetHMDDevice()->DrawHiddenAreaMesh_RenderThread(RHICmdList, View.StereoPass);
-	}
-}
-
-void FDeferredShadingSceneRenderer::RenderPrePassView(FRHICommandList& RHICmdList, const FViewInfo& View)
-{
-	SetupPrePassView(RHICmdList, View, this);
-
-	View.ParallelMeshDrawCommandPasses[EMeshPass::DepthPass].DispatchDraw(nullptr, RHICmdList);
-}
-
 DECLARE_CYCLE_STAT(TEXT("Prepass"), STAT_CLP_Prepass, STATGROUP_ParallelCommandListMarkers);
 
-class FPrePassParallelCommandListSet : public FParallelCommandListSet
-{
-public:
-	FPrePassParallelCommandListSet(FRHICommandListImmediate& InParentCmdList, const FSceneRenderer& InSceneRenderer, const FViewInfo& InView, bool bInCreateSceneContext)
-		: FParallelCommandListSet(GET_STATID(STAT_CLP_Prepass), InView, InParentCmdList, bInCreateSceneContext)
-		, SceneRenderer(InSceneRenderer)
-	{
-		// Do not copy-paste. this is a very unusual FParallelCommandListSet because it is a prepass and we want to do some work after starting some tasks
-	}
-
-	virtual ~FPrePassParallelCommandListSet()
-	{
-		// Do not copy-paste. this is a very unusual FParallelCommandListSet because it is a prepass and we want to do some work after starting some tasks
-		Dispatch(true);
-	}
-
-	virtual void SetStateOnCommandList(FRHICommandList& CmdList) override
-	{
-		FParallelCommandListSet::SetStateOnCommandList(CmdList);
-		FSceneRenderTargets::Get(CmdList).BeginRenderingPrePass(CmdList, false);
-		SetupPrePassView(CmdList, View, &SceneRenderer);
-	}
-
-private:
-	const FSceneRenderer& SceneRenderer;
-};
-
-bool FDeferredShadingSceneRenderer::RenderPrePassViewParallel(const FViewInfo& View, FRHICommandListImmediate& ParentCmdList, TFunctionRef<void()> AfterTasksAreStarted, bool bDoPrePre)
-{
-	bool bDepthWasCleared = false;
-
-	check(ParentCmdList.IsOutsideRenderPass());
-
-	{
-		FPrePassParallelCommandListSet ParallelCommandListSet(ParentCmdList, *this, View,
-			CVarRHICmdFlushRenderThreadTasksPrePass.GetValueOnRenderThread() == 0 && CVarRHICmdFlushRenderThreadTasks.GetValueOnRenderThread() == 0);
-
-		View.ParallelMeshDrawCommandPasses[EMeshPass::DepthPass].DispatchDraw(&ParallelCommandListSet, ParentCmdList);
-
-		if (bDoPrePre)
-		{
-			bDepthWasCleared = PreRenderPrePass(ParentCmdList);
-		}
-	}
-
-	if (bDoPrePre)
-	{
-		AfterTasksAreStarted();
-	}
-
-	return bDepthWasCleared;
-}
-
 /** A pixel shader used to fill the stencil buffer with the current dithered transition mask. */
 class FDitheredTransitionStencilPS : public FGlobalShader
 {
-	DECLARE_SHADER_TYPE(FDitheredTransitionStencilPS, Global);
-
 public:
+	DECLARE_GLOBAL_SHADER(FDitheredTransitionStencilPS);
+	SHADER_USE_PARAMETER_STRUCT(FDitheredTransitionStencilPS, FGlobalShader);
+
+	BEGIN_SHADER_PARAMETER_STRUCT(FParameters, )
+		SHADER_PARAMETER_STRUCT_REF(FViewUniformShaderParameters, View)
+		SHADER_PARAMETER(float, DitheredTransitionFactor)
+		RENDER_TARGET_BINDING_SLOTS()
+	END_SHADER_PARAMETER_STRUCT()
 
 	static bool ShouldCompilePermutation(const FGlobalShaderPermutationParameters& Parameters)
 	{ 
 		return IsFeatureLevelSupported(Parameters.Platform, ERHIFeatureLevel::SM5);
 	}
-
-	FDitheredTransitionStencilPS(const ShaderMetaType::CompiledShaderInitializerType& Initializer)
-	: FGlobalShader(Initializer)
-	{
-		DitheredTransitionFactorParameter.Bind(Initializer.ParameterMap, TEXT("DitheredTransitionFactor"), EShaderParameterFlags::SPF_Mandatory);
-	}
-
-	FDitheredTransitionStencilPS()
-	{
-	}
-
-	void SetParameters(FRHICommandList& RHICmdList, const FSceneView& View)
-	{
-		FGlobalShader::SetParameters<FViewUniformShaderParameters>(RHICmdList, RHICmdList.GetBoundPixelShader(), View.ViewUniformBuffer);
-
-		const float DitherFactor = View.GetTemporalLODTransition();
-		SetShaderValue(RHICmdList, RHICmdList.GetBoundPixelShader(), DitheredTransitionFactorParameter, DitherFactor);
-	}
-
-	LAYOUT_FIELD(FShaderParameter, DitheredTransitionFactorParameter);
 };
-IMPLEMENT_SHADER_TYPE(, FDitheredTransitionStencilPS, TEXT("/Engine/Private/DitheredTransitionStencil.usf"), TEXT("Main"), SF_Pixel);
+
+IMPLEMENT_GLOBAL_SHADER(FDitheredTransitionStencilPS, "/Engine/Private/DitheredTransitionStencil.usf", "Main", SF_Pixel);
 
 /** A compute shader used to fill the stencil buffer with the current dithered transition mask. */
 class FDitheredTransitionStencilCS : public FGlobalShader
 {
-	DECLARE_SHADER_TYPE(FDitheredTransitionStencilCS, Global);
-
 public:
+	DECLARE_GLOBAL_SHADER(FDitheredTransitionStencilCS);
+	SHADER_USE_PARAMETER_STRUCT(FDitheredTransitionStencilCS, FGlobalShader);
+
+	BEGIN_SHADER_PARAMETER_STRUCT(FParameters, )
+		SHADER_PARAMETER_STRUCT_REF(FViewUniformShaderParameters, View)
+		SHADER_PARAMETER_RDG_TEXTURE_UAV(RWTexture2D<uint>, StencilOutput)
+		SHADER_PARAMETER(float, DitheredTransitionFactor)
+		SHADER_PARAMETER(FIntVector4, StencilOffsetAndValues)
+	END_SHADER_PARAMETER_STRUCT()
 
 	static bool ShouldCompilePermutation(const FGlobalShaderPermutationParameters& Parameters)
 	{
 		return IsFeatureLevelSupported(Parameters.Platform, ERHIFeatureLevel::SM5);
 	}
-
-	FDitheredTransitionStencilCS(const ShaderMetaType::CompiledShaderInitializerType& Initializer)
-		: FGlobalShader(Initializer)
-	{
-		DitheredTransitionFactorParameter.Bind(Initializer.ParameterMap, TEXT("DitheredTransitionFactor"), EShaderParameterFlags::SPF_Mandatory);
-		StencilOffsetAndValuesParameter.Bind(Initializer.ParameterMap, TEXT("StencilOffsetAndValues"), EShaderParameterFlags::SPF_Mandatory);
-		StencilOutputParameter.Bind(Initializer.ParameterMap, TEXT("StencilOutput"), EShaderParameterFlags::SPF_Mandatory);
-	}
-
-	FDitheredTransitionStencilCS()
-	{
-	}
-
-	template <typename TRHICmdList>
-	void SetParameters(TRHICmdList& RHICmdList, const FSceneView& View, FRHIUnorderedAccessView* StencilOutputUAV, FIntPoint BufferSizeXY, FIntPoint ViewOffsetXY, uint32 StencilValue)
-	{
-		FRHIComputeShader* ComputeShader = RHICmdList.GetBoundComputeShader();
-
-		FGlobalShader::SetParameters<FViewUniformShaderParameters>(RHICmdList, ComputeShader, View.ViewUniformBuffer);
-
-		const float DitherFactor = View.GetTemporalLODTransition();
-		SetShaderValue(RHICmdList, ComputeShader, DitheredTransitionFactorParameter, DitherFactor);
-
-		const uint32 MaskedValue = (StencilValue & 0xFF);
-		const uint32 ClearedValue = 0;
-
-		const FIntVector4 StencilOffsetAndValues(
-			ViewOffsetXY.X,
-			ViewOffsetXY.Y,
-			int32(MaskedValue),
-			int32(ClearedValue));
-
-		SetShaderValue(RHICmdList, ComputeShader, StencilOffsetAndValuesParameter, StencilOffsetAndValues);
-		SetUAVParameter(RHICmdList, ComputeShader, StencilOutputParameter, StencilOutputUAV);
-	}
-
-	template <typename TRHICmdList>
-	void UnsetParameters(TRHICmdList& RHICmdList)
-	{
-		FRHIComputeShader* ComputeShader = RHICmdList.GetBoundComputeShader();
-		if (StencilOutputParameter.IsBound())
-		{
-			RHICmdList.SetUAVParameter(ComputeShader, StencilOutputParameter.GetBaseIndex(), nullptr);
-		}
-	}
-
-	LAYOUT_FIELD(FShaderParameter, DitheredTransitionFactorParameter);
-	LAYOUT_FIELD(FShaderParameter, StencilOffsetAndValuesParameter);
-	LAYOUT_FIELD(FShaderResourceParameter, StencilOutputParameter);
 };
-IMPLEMENT_SHADER_TYPE(, FDitheredTransitionStencilCS, TEXT("/Engine/Private/DitheredTransitionStencil.usf"), TEXT("MainCS"), SF_Compute);
-
-/** Possibly do the FX prerender and setup the prepass*/
-bool FDeferredShadingSceneRenderer::PreRenderPrePass(FRHICommandListImmediate& RHICmdList)
-{
-	// This can be called from within RenderPrePassViewParallel, so we need to reset
-	// the current GPU mask to the AllViews mask before iterating over Views again.
-	// Otherwise emulate stereo gets broken.
-	SCOPED_GPU_MASK(RHICmdList, AllViewsGPUMask);
-
-	RHICmdList.SetCurrentStat(GET_STATID(STAT_CLM_PrePass));
-
-	// RenderPrePassHMD clears the depth buffer. If this changes we must change RenderPrePass to maintain the correct behavior!
-	bool bDepthWasCleared = RenderPrePassHMD(RHICmdList);
-
-	FSceneRenderTargets& SceneContext = FSceneRenderTargets::Get(RHICmdList);
-
-	// Both compute approaches run earlier, so skip clearing stencil here, just load existing.
-	const bool bNoStencilClear = bDitheredLODTransitionsUseStencil && (StencilLODMode == 1 || StencilLODMode == 2);
 
-	SceneContext.BeginRenderingPrePass(RHICmdList, !bDepthWasCleared, !bNoStencilClear);
-	bDepthWasCleared = true;
+IMPLEMENT_GLOBAL_SHADER(FDitheredTransitionStencilCS, "/Engine/Private/DitheredTransitionStencil.usf", "MainCS", SF_Compute);
 
-	// Dithered transition stencil mask fill (graphics path)
-	if (bDitheredLODTransitionsUseStencil && StencilLODMode == 0)
-	{
-		PreRenderDitherFill(RHICmdList, SceneContext, nullptr);
-	}
-
-	// Need to close the render pass here since we may call BeginRenderingPrePass later
-	RHICmdList.EndRenderPass();
-
-	return bDepthWasCleared;
-}
-
-void FDeferredShadingSceneRenderer::PreRenderDitherFill(FRHIAsyncComputeCommandListImmediate& RHICmdList, FSceneRenderTargets& SceneContext, FRHIUnorderedAccessView* StencilTextureUAV)
+void AddDitheredStencilFillPass(FRDGBuilder& GraphBuilder, TConstArrayView<FViewInfo> Views, FRDGTextureRef DepthTexture, const FDepthPassInfo& DepthPass)
 {
-	SCOPED_GPU_EVENT(RHICmdList, DitheredStencilPrePass);
+	RDG_EVENT_SCOPE(GraphBuilder, "DitheredStencilPrePass");
 
-	FIntPoint BufferSizeXY = SceneContext.GetBufferSizeXY();
+	checkf(EnumHasAnyFlags(DepthPass.StencilDitherPassFlags, ERDGPassFlags::Raster | ERDGPassFlags::Compute | ERDGPassFlags::AsyncCompute), TEXT("Stencil dither fill pass flags are invalid."));
 
-	for (int32 ViewIndex = 0; ViewIndex < Views.Num(); ++ViewIndex)
+	if (DepthPass.StencilDitherPassFlags == ERDGPassFlags::Raster)
 	{
-		SCOPED_CONDITIONAL_GPU_EVENTF(RHICmdList, EventView, Views.Num() > 1, TEXT("View%d"), ViewIndex);
-
-		FViewInfo& View = Views[ViewIndex];
-
-		TShaderMapRef<FDitheredTransitionStencilCS> ComputeShader(View.ShaderMap);
-		RHICmdList.SetComputeShader(ComputeShader.GetComputeShader());
-		ComputeShader->SetParameters(RHICmdList, View, StencilTextureUAV, BufferSizeXY, View.ViewRect.Min, STENCIL_SANDBOX_MASK);
-		const int32 SubWidth = FMath::Min(BufferSizeXY.X, View.ViewRect.Width());
-		const int32 SubHeight = FMath::Min(BufferSizeXY.Y, View.ViewRect.Height());
-		check(SubWidth > 0 && SubHeight > 0);
-
-		DispatchComputeShader(RHICmdList, ComputeShader.GetShader(), FMath::DivideAndRoundUp(SubWidth, 8), FMath::DivideAndRoundUp(SubHeight, 8), 1);
-		ComputeShader->UnsetParameters(RHICmdList);
-	}
-}
-
+		FRHIDepthStencilState* DepthStencilState = TStaticDepthStencilState<false, CF_Always,
+			true, CF_Always, SO_Keep, SO_Keep, SO_Replace,
+			false, CF_Always, SO_Keep, SO_Keep, SO_Keep,
+			STENCIL_SANDBOX_MASK, STENCIL_SANDBOX_MASK>::GetRHI();
 
-void FDeferredShadingSceneRenderer::PreRenderDitherFill(FRHICommandListImmediate& RHICmdList, FSceneRenderTargets& SceneContext, FRHIUnorderedAccessView* StencilTextureUAV)
-{
-	SCOPED_DRAW_EVENT(RHICmdList, DitheredStencilPrePass);
+		const uint32 StencilRef = STENCIL_SANDBOX_MASK;
 
-	FIntPoint BufferSizeXY = SceneContext.GetBufferSizeXY();
-	if (StencilLODMode == 1 || StencilLODMode == 2)
-	{
 		for (int32 ViewIndex = 0; ViewIndex < Views.Num(); ++ViewIndex)
 		{
-			SCOPED_CONDITIONAL_DRAW_EVENTF(RHICmdList, EventView, Views.Num() > 1, TEXT("View%d"), ViewIndex);
+			const FViewInfo& View = Views[ViewIndex];
+			RDG_GPU_MASK_SCOPE(GraphBuilder, View.GPUMask);
+			RDG_EVENT_SCOPE_CONDITIONAL(GraphBuilder, Views.Num() > 1, "View%d", ViewIndex);
+
+			TShaderMapRef<FDitheredTransitionStencilPS> PixelShader(View.ShaderMap);
 
-			FViewInfo& View = Views[ViewIndex];
+			auto* PassParameters = GraphBuilder.AllocParameters<FDitheredTransitionStencilPS::FParameters>();
+			PassParameters->View = View.ViewUniformBuffer;
+			PassParameters->DitheredTransitionFactor = View.GetTemporalLODTransition();
+			PassParameters->RenderTargets.DepthStencil = FDepthStencilBinding(DepthTexture, ERenderTargetLoadAction::ELoad, ERenderTargetLoadAction::ELoad, FExclusiveDepthStencil::DepthWrite_StencilWrite);
 
-			TShaderMapRef<FDitheredTransitionStencilCS> ComputeShader(View.ShaderMap);
-			RHICmdList.SetComputeShader(ComputeShader.GetComputeShader());
-			ComputeShader->SetParameters(RHICmdList, View, StencilTextureUAV, BufferSizeXY, View.ViewRect.Min, STENCIL_SANDBOX_MASK);
-			const int32 SubWidth = FMath::Min(BufferSizeXY.X, View.ViewRect.Width());
-			const int32 SubHeight = FMath::Min(BufferSizeXY.Y, View.ViewRect.Height());
-			check(SubWidth > 0 && SubHeight > 0);
-
-			DispatchComputeShader(RHICmdList, ComputeShader.GetShader(), FMath::DivideAndRoundUp(SubWidth, 8), FMath::DivideAndRoundUp(SubHeight, 8), 1);
-			ComputeShader->UnsetParameters(RHICmdList);
+			FPixelShaderUtils::AddFullscreenPass(GraphBuilder, View.ShaderMap, {}, PixelShader, PassParameters, View.ViewRect, nullptr, nullptr, DepthStencilState, StencilRef);
 		}
 	}
 	else
 	{
-		FGraphicsPipelineStateInitializer GraphicsPSOInit;
-		RHICmdList.ApplyCachedRenderTargets(GraphicsPSOInit);
-		GraphicsPSOInit.RasterizerState = TStaticRasterizerState<>::GetRHI();
-		GraphicsPSOInit.BlendState = TStaticBlendState<>::GetRHI();
-		GraphicsPSOInit.DepthStencilState = TStaticDepthStencilState<false, CF_Always,
-			true, CF_Always, SO_Keep, SO_Keep, SO_Replace,
-			false, CF_Always, SO_Keep, SO_Keep, SO_Keep,
-			STENCIL_SANDBOX_MASK, STENCIL_SANDBOX_MASK>::GetRHI();
+		const int32 MaskedValue = (STENCIL_SANDBOX_MASK & 0xFF);
+		const int32 ClearedValue = 0;
 
 		for (int32 ViewIndex = 0; ViewIndex < Views.Num(); ++ViewIndex)
 		{
-			FViewInfo& View = Views[ViewIndex];
-
-			SCOPED_GPU_MASK(RHICmdList, View.GPUMask);
-			SCOPED_CONDITIONAL_DRAW_EVENTF(RHICmdList, EventView, Views.Num() > 1, TEXT("View%d"), ViewIndex);
-
-			RHICmdList.SetViewport(View.ViewRect.Min.X, View.ViewRect.Min.Y, 0.0f, View.ViewRect.Max.X, View.ViewRect.Max.Y, 1.0f);
+			const FViewInfo& View = Views[ViewIndex];
+			RDG_GPU_MASK_SCOPE(GraphBuilder, View.GPUMask);
+			RDG_EVENT_SCOPE_CONDITIONAL(GraphBuilder, Views.Num() > 1, "View%d", ViewIndex);
 
-			// Set shaders, states
-			TShaderMapRef<FScreenVS> ScreenVertexShader(View.ShaderMap);
-			TShaderMapRef<FDitheredTransitionStencilPS> PixelShader(View.ShaderMap);
+			TShaderMapRef<FDitheredTransitionStencilCS> ComputeShader(View.ShaderMap);
 
-			extern TGlobalResource<FFilterVertexDeclaration> GFilterVertexDeclaration;
-			GraphicsPSOInit.BoundShaderState.VertexDeclarationRHI = GFilterVertexDeclaration.VertexDeclarationRHI;
-			GraphicsPSOInit.BoundShaderState.VertexShaderRHI = ScreenVertexShader.GetVertexShader();
-			GraphicsPSOInit.BoundShaderState.PixelShaderRHI = PixelShader.GetPixelShader();
-			GraphicsPSOInit.PrimitiveType = PT_TriangleList;
-
-			SetGraphicsPipelineState(RHICmdList, GraphicsPSOInit);
-			RHICmdList.SetStencilRef(STENCIL_SANDBOX_MASK);
-
-			PixelShader->SetParameters(RHICmdList, View);
-
-			DrawRectangle(
-				RHICmdList,
-				0, 0,
-				BufferSizeXY.X, BufferSizeXY.Y,
-				View.ViewRect.Min.X, View.ViewRect.Min.Y,
-				View.ViewRect.Width(), View.ViewRect.Height(),
-				BufferSizeXY,
-				BufferSizeXY,
-				ScreenVertexShader,
-				EDRF_UseTriangleOptimization);
+			auto* PassParameters = GraphBuilder.AllocParameters<FDitheredTransitionStencilCS::FParameters>();
+			PassParameters->View = View.ViewUniformBuffer;
+			PassParameters->StencilOutput = GraphBuilder.CreateUAV(FRDGTextureUAVDesc::CreateForMetaData(DepthTexture, ERDGTextureMetaDataAccess::Stencil));
+			PassParameters->DitheredTransitionFactor = View.GetTemporalLODTransition();
+			PassParameters->StencilOffsetAndValues = FIntVector4(View.ViewRect.Min.X, View.ViewRect.Min.Y, MaskedValue, ClearedValue);
+
+			const FIntPoint SubExtent(
+				FMath::Min(DepthTexture->Desc.Extent.X, View.ViewRect.Width()),
+				FMath::Min(DepthTexture->Desc.Extent.Y, View.ViewRect.Height()));
+			check(SubExtent.X > 0 && SubExtent.Y > 0);
+
+			FComputeShaderUtils::AddPass(
+				GraphBuilder,
+				{},
+				DepthPass.StencilDitherPassFlags,
+				ComputeShader,
+				PassParameters,
+				FComputeShaderUtils::GetGroupCount(SubExtent, FComputeShaderUtils::kGolden2DGroupSize));
 		}
 	}
 }
 
-void FDeferredShadingSceneRenderer::RenderPrePassEditorPrimitives(FRHICommandList& RHICmdList, const FViewInfo& View, const FMeshPassProcessorRenderState& DrawRenderState, EDepthDrawingMode DepthDrawingMode, bool bRespectUseAsOccluderFlag) 
+static void RenderPrePassEditorPrimitives(
+	FRDGBuilder& GraphBuilder,
+	const FViewInfo& View,
+	FDepthPassParameters* PassParameters,
+	const FMeshPassProcessorRenderState& DrawRenderState,
+	EDepthDrawingMode DepthDrawingMode)
 {
-	SetupPrePassView(RHICmdList, View, this, true);
+	GraphBuilder.AddPass(
+		RDG_EVENT_NAME("EditorPrimitives"),
+		PassParameters,
+		ERDGPassFlags::Raster,
+		[&View, DrawRenderState, DepthDrawingMode](FRHICommandList& RHICmdList)
+	{
+		const bool bRespectUseAsOccluderFlag = true;
 
-	View.SimpleElementCollector.DrawBatchedElements(RHICmdList, DrawRenderState, View, EBlendModeFilter::OpaqueAndMasked, SDPG_World);
-	View.SimpleElementCollector.DrawBatchedElements(RHICmdList, DrawRenderState, View, EBlendModeFilter::OpaqueAndMasked, SDPG_Foreground);
+		RHICmdList.SetViewport(View.ViewRect.Min.X, View.ViewRect.Min.Y, 0.0f, View.ViewRect.Max.X, View.ViewRect.Max.Y, 1.0f);
 
-	bool bDirty = false;
-	if (!View.Family->EngineShowFlags.CompositeEditorPrimitives)
-	{
-		const bool bNeedToSwitchVerticalAxis = RHINeedsToSwitchVerticalAxis(ShaderPlatform);
-		const FScene* LocalScene = Scene;
+		View.SimpleElementCollector.DrawBatchedElements(RHICmdList, DrawRenderState, View, EBlendModeFilter::OpaqueAndMasked, SDPG_World);
+		View.SimpleElementCollector.DrawBatchedElements(RHICmdList, DrawRenderState, View, EBlendModeFilter::OpaqueAndMasked, SDPG_Foreground);
+
+		if (!View.Family->EngineShowFlags.CompositeEditorPrimitives)
+		{
+			const bool bNeedToSwitchVerticalAxis = RHINeedsToSwitchVerticalAxis(View.GetShaderPlatform());
 
-		DrawDynamicMeshPass(View, RHICmdList,
-			[&View, &DrawRenderState, LocalScene, DepthDrawingMode, bRespectUseAsOccluderFlag](FDynamicPassMeshDrawListContext* DynamicMeshPassContext)
+			DrawDynamicMeshPass(View, RHICmdList, [&](FDynamicPassMeshDrawListContext* DynamicMeshPassContext)
 			{
 				FDepthPassMeshProcessor PassMeshProcessor(
-					LocalScene,
+					View.Family->Scene->GetRenderScene(),
 					&View,
 					DrawRenderState,
 					bRespectUseAsOccluderFlag,
@@ -564,7 +422,7 @@ void FDeferredShadingSceneRenderer::RenderPrePassEditorPrimitives(FRHICommandLis
 					DynamicMeshPassContext);
 
 				const uint64 DefaultBatchElementMask = ~0ull;
-					
+
 				for (int32 MeshIndex = 0; MeshIndex < View.ViewMeshElements.Num(); MeshIndex++)
 				{
 					const FMeshBatch& MeshBatch = View.ViewMeshElements[MeshIndex];
@@ -572,14 +430,13 @@ void FDeferredShadingSceneRenderer::RenderPrePassEditorPrimitives(FRHICommandLis
 				}
 			});
 
-		// Draw the view's batched simple elements(lines, sprites, etc).
-		bDirty |= View.BatchedViewElements.Draw(RHICmdList, DrawRenderState, FeatureLevel, bNeedToSwitchVerticalAxis, View, false) || bDirty;
+			// Draw the view's batched simple elements(lines, sprites, etc).
+			View.BatchedViewElements.Draw(RHICmdList, DrawRenderState, View.FeatureLevel, bNeedToSwitchVerticalAxis, View, false);
 
-		DrawDynamicMeshPass(View, RHICmdList,
-			[&View, &DrawRenderState, LocalScene, DepthDrawingMode, bRespectUseAsOccluderFlag](FDynamicPassMeshDrawListContext* DynamicMeshPassContext)
+			DrawDynamicMeshPass(View, RHICmdList, [&](FDynamicPassMeshDrawListContext* DynamicMeshPassContext)
 			{
 				FDepthPassMeshProcessor PassMeshProcessor(
-					LocalScene,
+					View.Family->Scene->GetRenderScene(),
 					&View,
 					DrawRenderState,
 					bRespectUseAsOccluderFlag,
@@ -589,7 +446,7 @@ void FDeferredShadingSceneRenderer::RenderPrePassEditorPrimitives(FRHICommandLis
 					DynamicMeshPassContext);
 
 				const uint64 DefaultBatchElementMask = ~0ull;
-					
+
 				for (int32 MeshIndex = 0; MeshIndex < View.TopViewMeshElements.Num(); MeshIndex++)
 				{
 					const FMeshBatch& MeshBatch = View.TopViewMeshElements[MeshIndex];
@@ -597,9 +454,10 @@ void FDeferredShadingSceneRenderer::RenderPrePassEditorPrimitives(FRHICommandLis
 				}
 			});
 
-		// Draw the view's batched simple elements(lines, sprites, etc).
-		bDirty |= View.TopBatchedViewElements.Draw(RHICmdList, DrawRenderState, FeatureLevel, bNeedToSwitchVerticalAxis, View, false) || bDirty;
-	}
+			// Draw the view's batched simple elements(lines, sprites, etc).
+			View.TopBatchedViewElements.Draw(RHICmdList, DrawRenderState, View.FeatureLevel, bNeedToSwitchVerticalAxis, View, false);
+		}
+	});
 }
 
 void SetupDepthPassState(FMeshPassProcessorRenderState& DrawRenderState)
@@ -609,120 +467,115 @@ void SetupDepthPassState(FMeshPassProcessorRenderState& DrawRenderState)
 	DrawRenderState.SetDepthStencilState(TStaticDepthStencilState<true, CF_DepthNearOrEqual>::GetRHI());
 }
 
-bool FDeferredShadingSceneRenderer::RenderPrePass(FRHICommandListImmediate& RHICmdList, TFunctionRef<void()> AfterTasksAreStarted)
+void FDeferredShadingSceneRenderer::RenderPrePass(FRDGBuilder& GraphBuilder, FRDGTextureRef SceneDepthTexture)
 {
-	check(RHICmdList.IsOutsideRenderPass());
-
-	SCOPED_NAMED_EVENT(FDeferredShadingSceneRenderer_RenderPrePass, FColor::Emerald);
-	bool bDepthWasCleared = false;
-
 	extern const TCHAR* GetDepthPassReason(bool bDitheredLODTransitionsUseStencil, EShaderPlatform ShaderPlatform);
-	SCOPED_DRAW_EVENTF(RHICmdList, PrePass, TEXT("PrePass %s %s"), GetDepthDrawingModeString(EarlyZPassMode), GetDepthPassReason(bDitheredLODTransitionsUseStencil, ShaderPlatform));
+	RDG_EVENT_SCOPE(GraphBuilder, "PrePass %s %s", GetDepthDrawingModeString(DepthPass.EarlyZPassMode), GetDepthPassReason(DepthPass.bDitheredLODTransitionsUseStencil, ShaderPlatform));
+	RDG_CSV_STAT_EXCLUSIVE_SCOPE(GraphBuilder, RenderPrePass);
+	RDG_GPU_STAT_SCOPE(GraphBuilder, Prepass);
 
+	SCOPED_NAMED_EVENT(FDeferredShadingSceneRenderer_RenderPrePass, FColor::Emerald);
 	SCOPE_CYCLE_COUNTER(STAT_DepthDrawTime);
-	CSV_SCOPED_TIMING_STAT_EXCLUSIVE(RenderPrePass);
-	SCOPED_GPU_STAT(RHICmdList, Prepass);
 
-	bool bDidPrePre = false;
-	FSceneRenderTargets& SceneContext = FSceneRenderTargets::Get(RHICmdList);
+	const bool bParallelDepthPass = GRHICommandList.UseParallelAlgorithms() && CVarParallelPrePass.GetValueOnRenderThread();
 
-	bool bParallel = GRHICommandList.UseParallelAlgorithms() && CVarParallelPrePass.GetValueOnRenderThread();
+	RenderPrePassHMD(GraphBuilder, SceneDepthTexture);
 
-	if (!bParallel)
-	{
-		// nothing to be gained by delaying this.
-		AfterTasksAreStarted();
-		// Note: the depth buffer will be cleared under PreRenderPrePass.
-		bDepthWasCleared = PreRenderPrePass(RHICmdList);
-		bDidPrePre = true;
-
-		// PreRenderPrePass will end up clearing the depth buffer so do not clear it again.
-		SceneContext.BeginRenderingPrePass(RHICmdList, false);
-	}
-	else
+	if (DepthPass.IsRasterStencilDitherEnabled())
 	{
-		SceneContext.GetSceneDepthSurface(); // this probably isn't needed, but if there was some lazy allocation of the depth surface going on, we want it allocated now before we go wide. We may not have called BeginRenderingPrePass yet if bDoFXPrerender is true
+		AddDitheredStencilFillPass(GraphBuilder, Views, SceneDepthTexture, DepthPass);
 	}
 
+	FDepthPassParameters* PassParameters = GetDepthPassParameters(GraphBuilder, SceneDepthTexture);
+
 	// Draw a depth pass to avoid overdraw in the other passes.
-	if(EarlyZPassMode != DDM_None)
+	if (DepthPass.EarlyZPassMode != DDM_None)
 	{
-		const bool bWaitForTasks = bParallel && (CVarRHICmdFlushRenderThreadTasksPrePass.GetValueOnRenderThread() > 0 || CVarRHICmdFlushRenderThreadTasks.GetValueOnRenderThread() > 0);
-		FScopedCommandListWaitForTasks Flusher(bWaitForTasks, RHICmdList);
-
-		for(int32 ViewIndex = 0;ViewIndex < Views.Num();ViewIndex++)
+		if (bParallelDepthPass)
 		{
-			const FViewInfo& View = Views[ViewIndex];
+			RDG_WAIT_FOR_TASKS_CONDITIONAL(GraphBuilder, IsDepthPassWaitForTasksEnabled());
 
-			SCOPED_GPU_MASK(RHICmdList, !View.IsInstancedStereoPass() ? View.GPUMask : (Views[0].GPUMask | Views[1].GPUMask));
-			SCOPED_CONDITIONAL_DRAW_EVENTF(RHICmdList, EventView, Views.Num() > 1, TEXT("View%d"), ViewIndex);
+			for (int32 ViewIndex = 0; ViewIndex < Views.Num(); ++ViewIndex)
+			{
+				FViewInfo& View = Views[ViewIndex];
+				RDG_GPU_MASK_SCOPE(GraphBuilder, View.GPUMask);
+				RDG_EVENT_SCOPE_CONDITIONAL(GraphBuilder, Views.Num() > 1, "View%d", ViewIndex);
 
-			FMeshPassProcessorRenderState DrawRenderState(View);
+				FMeshPassProcessorRenderState DrawRenderState(View);
+				SetupDepthPassState(DrawRenderState);
 
-			SetupDepthPassState(DrawRenderState);
+				const bool bShouldRenderView = View.ShouldRenderView();
+				if (bShouldRenderView)
+				{
+					GraphBuilder.AddPass(
+						RDG_EVENT_NAME("DepthPassParallel"),
+						PassParameters,
+						ERDGPassFlags::Raster | ERDGPassFlags::SkipRenderPass,
+						[this, &View, PassParameters](FRHICommandListImmediate& RHICmdList)
+					{
+						Scene->UniformBuffers.UpdateViewUniformBuffer(View);
+						FRDGParallelCommandListSet ParallelCommandListSet(RHICmdList, GET_STATID(STAT_CLP_Prepass), *this, View, FParallelCommandListBindings(PassParameters));
+						ParallelCommandListSet.SetHighPriority();
 
-			if (View.ShouldRenderView())
-			{
-				Scene->UniformBuffers.UpdateViewUniformBuffer(View);
+						View.ParallelMeshDrawCommandPasses[EMeshPass::DepthPass].DispatchDraw(&ParallelCommandListSet, RHICmdList);
+					});
 
-				if (bParallel)
-				{
-					check(RHICmdList.IsOutsideRenderPass());
-					bDepthWasCleared = RenderPrePassViewParallel(View, RHICmdList, AfterTasksAreStarted, !bDidPrePre) || bDepthWasCleared;
-					bDidPrePre = true;
-				}
-				else
-				{
-					RenderPrePassView(RHICmdList, View);
+					RenderPrePassEditorPrimitives(GraphBuilder, View, PassParameters, DrawRenderState, DepthPass.EarlyZPassMode);
 				}
 			}
-
-			// Parallel rendering has self contained renderpasses so we need a new one for editor primitives.
-			if (bParallel)
-			{
-				SceneContext.BeginRenderingPrePass(RHICmdList, false);
-			}
-			RenderPrePassEditorPrimitives(RHICmdList, View, DrawRenderState, EarlyZPassMode, true);
-			if (bParallel)
+		}
+		else
+		{
+			for (int32 ViewIndex = 0; ViewIndex < Views.Num(); ++ViewIndex)
 			{
-				RHICmdList.EndRenderPass();
+				FViewInfo& View = Views[ViewIndex];
+				RDG_GPU_MASK_SCOPE(GraphBuilder, View.GPUMask);
+				RDG_EVENT_SCOPE_CONDITIONAL(GraphBuilder, Views.Num() > 1, "View%d", ViewIndex);
+
+				FMeshPassProcessorRenderState DrawRenderState(View);
+				SetupDepthPassState(DrawRenderState);
+
+				const bool bShouldRenderView = View.ShouldRenderView();
+				if (bShouldRenderView)
+				{
+					GraphBuilder.AddPass(
+						RDG_EVENT_NAME("DepthPass"),
+						PassParameters,
+						ERDGPassFlags::Raster,
+						[this, &View, PassParameters](FRHICommandList& RHICmdList)
+					{
+						Scene->UniformBuffers.UpdateViewUniformBuffer(View);
+						SetStereoViewport(RHICmdList, View, 1.0f);
+						View.ParallelMeshDrawCommandPasses[EMeshPass::DepthPass].DispatchDraw(nullptr, RHICmdList);
+					});
+
+					RenderPrePassEditorPrimitives(GraphBuilder, View, PassParameters, DrawRenderState, DepthPass.EarlyZPassMode);
+				}
 			}
 		}
 	}
-	if (!bDidPrePre)
-	{
-		// Only parallel rendering with all views marked as not-to-be-rendered will get here.
-		// For some reason we haven't done this yet. Best do it now for consistency with the old code.
-		AfterTasksAreStarted();
-		bDepthWasCleared = PreRenderPrePass(RHICmdList);
-		bDidPrePre = true;
-	}
-
-	if (bParallel)
-	{
-		// In parallel mode there will be no renderpass here. Need to restart.
-		SceneContext.BeginRenderingPrePass(RHICmdList, false);
-	}
 
 	// Dithered transition stencil mask clear, accounting for all active viewports
-	if (bDitheredLODTransitionsUseStencil)
+	if (DepthPass.bDitheredLODTransitionsUseStencil)
 	{
-		if (Views.Num() > 1)
+		GraphBuilder.AddPass(
+			RDG_EVENT_NAME("DitherStencilClear"),
+			PassParameters,
+			ERDGPassFlags::Raster,
+			[this](FRHICommandList& RHICmdList)
 		{
-			FIntRect FullViewRect = Views[0].ViewRect;
-			for (int32 ViewIndex = 1; ViewIndex < Views.Num(); ++ViewIndex)
+			if (Views.Num() > 1)
 			{
-				FullViewRect.Union(Views[ViewIndex].ViewRect);
+				FIntRect FullViewRect = Views[0].ViewRect;
+				for (int32 ViewIndex = 1; ViewIndex < Views.Num(); ++ViewIndex)
+				{
+					FullViewRect.Union(Views[ViewIndex].ViewRect);
+				}
+				RHICmdList.SetViewport(FullViewRect.Min.X, FullViewRect.Min.Y, 0, FullViewRect.Max.X, FullViewRect.Max.Y, 1);
 			}
-			RHICmdList.SetViewport(FullViewRect.Min.X, FullViewRect.Min.Y, 0, FullViewRect.Max.X, FullViewRect.Max.Y, 1);
-		}
-		DrawClearQuad(RHICmdList, false, FLinearColor::Transparent, false, 0, true, 0);
+			DrawClearQuad(RHICmdList, false, FLinearColor::Transparent, false, 0, true, 0);
+		});
 	}
-
-	// Now we are finally finished.
-	SceneContext.FinishRenderingPrePass(RHICmdList);
-
-	return bDepthWasCleared;
 }
 
 void FMobileSceneRenderer::RenderPrePass(FRHICommandListImmediate& RHICmdList)
@@ -753,51 +606,61 @@ void FMobileSceneRenderer::RenderPrePass(FRHICommandListImmediate& RHICmdList)
 
 			Scene->UniformBuffers.UpdateViewUniformBuffer(View);
 
-			SetupPrePassView(RHICmdList, View, this);
+			SetStereoViewport(RHICmdList, View);
 
 			View.ParallelMeshDrawCommandPasses[EMeshPass::DepthPass].DispatchDraw(nullptr, RHICmdList);
 		}
 	}
 }
 
-extern bool IsHMDHiddenAreaMaskActive();
-
-bool FDeferredShadingSceneRenderer::RenderPrePassHMD(FRHICommandListImmediate& RHICmdList)
+void FDeferredShadingSceneRenderer::RenderPrePassHMD(FRDGBuilder& GraphBuilder, FRDGTextureRef DepthTexture)
 {
 	// Early out before we change any state if there's not a mask to render
 	if (!IsHMDHiddenAreaMaskActive())
 	{
-		return false;
+		return;
 	}
 
-	// This is the only place the depth buffer is cleared. If this changes we MUST change RenderPrePass and others to maintain the behavior.
-	FSceneRenderTargets& SceneContext = FSceneRenderTargets::Get(RHICmdList);
-	SceneContext.BeginRenderingPrePass(RHICmdList, true);
+	auto* HMDDevice = GEngine->XRSystem->GetHMDDevice();
+	if (!HMDDevice)
+	{
+		return;
+	}
 
+	FDepthPassParameters* PassParameters = GetDepthPassParameters(GraphBuilder, DepthTexture);
 
-	FGraphicsPipelineStateInitializer GraphicsPSOInit;
-	RHICmdList.ApplyCachedRenderTargets(GraphicsPSOInit);
+	GraphBuilder.AddPass(
+		RDG_EVENT_NAME("HiddenAreaMask"),
+		PassParameters,
+		ERDGPassFlags::Raster,
+		[this, HMDDevice](FRHICommandList& RHICmdList)
+	{
+		extern TGlobalResource<FFilterVertexDeclaration> GFilterVertexDeclaration;
 
-	GraphicsPSOInit.BlendState = TStaticBlendState<CW_NONE>::GetRHI();
-	GraphicsPSOInit.DepthStencilState = TStaticDepthStencilState<true, CF_DepthNearOrEqual>::GetRHI();
-	GraphicsPSOInit.RasterizerState = TStaticRasterizerState<FM_Solid, CM_None>::GetRHI();
+		TShaderMapRef<TOneColorVS<true>> VertexShader(GetGlobalShaderMap(GMaxRHIFeatureLevel));
 
-	RHICmdList.SetScissorRect(false, 0, 0, 0, 0);
+		FGraphicsPipelineStateInitializer GraphicsPSOInit;
+		GraphicsPSOInit.BlendState = TStaticBlendState<CW_NONE>::GetRHI();
+		GraphicsPSOInit.DepthStencilState = TStaticDepthStencilState<true, CF_DepthNearOrEqual>::GetRHI();
+		GraphicsPSOInit.RasterizerState = TStaticRasterizerState<FM_Solid, CM_None>::GetRHI();
+		GraphicsPSOInit.BoundShaderState.VertexShaderRHI = VertexShader.GetVertexShader();
+		GraphicsPSOInit.BoundShaderState.VertexDeclarationRHI = GFilterVertexDeclaration.VertexDeclarationRHI;
+		GraphicsPSOInit.PrimitiveType = PT_TriangleList;
+		RHICmdList.ApplyCachedRenderTargets(GraphicsPSOInit);
 
-	for (int32 ViewIndex = 0; ViewIndex < Views.Num(); ++ViewIndex)
-	{
-		const FViewInfo& View = Views[ViewIndex];
-		if (IStereoRendering::IsStereoEyeView(View))
+		for (const FViewInfo& View : Views)
 		{
-			SCOPED_GPU_MASK(RHICmdList, View.GPUMask);
-			RHICmdList.SetViewport(View.ViewRect.Min.X, View.ViewRect.Min.Y, 0.0f, View.ViewRect.Max.X, View.ViewRect.Max.Y, 1.0f);
-			RenderHiddenAreaMaskView(RHICmdList, GraphicsPSOInit, View);
-		}
-	}
-
-	SceneContext.FinishRenderingPrePass(RHICmdList);
+			if (IStereoRendering::IsStereoEyeView(View))
+			{
+				SCOPED_GPU_MASK(RHICmdList, View.GPUMask);
+				RHICmdList.SetViewport(View.ViewRect.Min.X, View.ViewRect.Min.Y, 0.0f, View.ViewRect.Max.X, View.ViewRect.Max.Y, 1.0f);
 
-	return true;
+				SetGraphicsPipelineState(RHICmdList, GraphicsPSOInit);
+				VertexShader->SetDepthParameter(RHICmdList, 1.0f);
+				HMDDevice->DrawHiddenAreaMesh_RenderThread(RHICmdList, View.StereoPass);
+			}
+		}
+	});
 }
 
 FMeshDrawCommandSortKey CalculateDepthPassMeshStaticSortKey(EBlendMode BlendMode, const FMeshMaterialShader* VertexShader, const FMeshMaterialShader* PixelShader)
diff --git a/Engine/Source/Runtime/Renderer/Private/DepthRendering.h b/Engine/Source/Runtime/Renderer/Private/DepthRendering.h
--- a/Engine/Source/Runtime/Renderer/Private/DepthRendering.h
+++ b/Engine/Source/Runtime/Renderer/Private/DepthRendering.h
@@ -33,6 +33,28 @@ enum EDepthDrawingMode
 
 extern const TCHAR* GetDepthDrawingModeString(EDepthDrawingMode Mode);
 
+struct FDepthPassInfo
+{
+	bool IsComputeStencilDitherEnabled() const
+	{
+		return StencilDitherPassFlags != ERDGPassFlags::Raster && bDitheredLODTransitionsUseStencil;
+	}
+
+	bool IsRasterStencilDitherEnabled() const
+	{
+		return StencilDitherPassFlags == ERDGPassFlags::Raster && bDitheredLODTransitionsUseStencil;
+	}
+
+	EDepthDrawingMode EarlyZPassMode = DDM_None;
+	bool bEarlyZPassMovable = false;
+	bool bDitheredLODTransitionsUseStencil = false;
+	ERDGPassFlags StencilDitherPassFlags = ERDGPassFlags::Raster;
+};
+
+extern FDepthPassInfo GetDepthPassInfo(const FScene* Scene);
+
+void AddDitheredStencilFillPass(FRDGBuilder& GraphBuilder, TConstArrayView<FViewInfo> Views, FRDGTextureRef DepthTexture, const FDepthPassInfo& DepthPass);
+
 class FDepthOnlyShaderElementData : public FMeshMaterialShaderElementData
 {
 public:
diff --git a/Engine/Source/Runtime/Renderer/Private/PostProcess/SceneRenderTargets.cpp b/Engine/Source/Runtime/Renderer/Private/PostProcess/SceneRenderTargets.cpp
--- a/Engine/Source/Runtime/Renderer/Private/PostProcess/SceneRenderTargets.cpp
+++ b/Engine/Source/Runtime/Renderer/Private/PostProcess/SceneRenderTargets.cpp
@@ -1220,48 +1220,6 @@ void FSceneRenderTargets::AdjustGBufferRefCount(FRHICommandList& RHICmdList, int
 	}	
 }
 
-void FSceneRenderTargets::BeginRenderingPrePass(FRHICommandList& RHICmdList, bool bPerformClear, bool bStencilClear)
-{
-	check(RHICmdList.IsOutsideRenderPass());
-
-	SCOPED_DRAW_EVENT(RHICmdList, BeginRenderingPrePass);
-
-	FTexture2DRHIRef DepthTarget = GetSceneDepthSurface();
-
-	RHICmdList.Transition(FRHITransitionInfo(DepthTarget, ERHIAccess::Unknown, ERHIAccess::DSVWrite | ERHIAccess::DSVRead));
-	
-	// No color target bound for the prepass.
-	FRHIRenderPassInfo RPInfo;
-	RPInfo.DepthStencilRenderTarget.DepthStencilTarget = DepthTarget;
-	RPInfo.DepthStencilRenderTarget.ExclusiveDepthStencil = FExclusiveDepthStencil::DepthWrite_StencilWrite;
-	
-	if (bPerformClear)
-	{				
-		// Clear the depth buffer.
-		// Note, this is a reversed Z depth surface, so 0.0f is the far plane.
-		ERenderTargetActions StencilAction = bStencilClear ? ERenderTargetActions::Clear_Store : ERenderTargetActions::Load_Store;
-		RPInfo.DepthStencilRenderTarget.Action = MakeDepthStencilTargetActions(ERenderTargetActions::Clear_Store, StencilAction);
-	}
-	else
-	{
-		// Set the scene depth surface and a DUMMY buffer as color buffer
-		// (as long as it's the same dimension as the depth buffer),	
-		RPInfo.DepthStencilRenderTarget.Action = MakeDepthStencilTargetActions(ERenderTargetActions::Load_Store, ERenderTargetActions::Load_Store);
-	}
-		
-	RHICmdList.BeginRenderPass(RPInfo, TEXT("BeginRenderingPrePass"));
-}
-
-void FSceneRenderTargets::FinishRenderingPrePass(FRHICommandListImmediate& RHICmdList)
-{
-	check(RHICmdList.IsInsideRenderPass());
-
-	SCOPED_DRAW_EVENT(RHICmdList, FinishRenderingPrePass);
-	RHICmdList.EndRenderPass();
-
-	GVisualizeTexture.SetCheckPoint(RHICmdList, SceneDepthZ);
-}
-
 void FSceneRenderTargets::CleanUpEditorPrimitiveTargets()
 {
 	EditorPrimitivesDepth.SafeRelease();
diff --git a/Engine/Source/Runtime/Renderer/Private/PostProcess/SceneRenderTargets.h b/Engine/Source/Runtime/Renderer/Private/PostProcess/SceneRenderTargets.h
--- a/Engine/Source/Runtime/Renderer/Private/PostProcess/SceneRenderTargets.h
+++ b/Engine/Source/Runtime/Renderer/Private/PostProcess/SceneRenderTargets.h
@@ -254,9 +254,6 @@ public:
 		}
 	}
 
-	void BeginRenderingPrePass(FRHICommandList& RHICmdList, bool bPerformClear, bool bStencilClear = true);
-	void FinishRenderingPrePass(FRHICommandListImmediate& RHICmdList);
-
 	void SetDefaultColorClear(const FClearValueBinding ColorClear)
 	{
 		DefaultColorClear = ColorClear;
diff --git a/Engine/Source/Runtime/Renderer/Private/SceneRendering.h b/Engine/Source/Runtime/Renderer/Private/SceneRendering.h
--- a/Engine/Source/Runtime/Renderer/Private/SceneRendering.h
+++ b/Engine/Source/Runtime/Renderer/Private/SceneRendering.h
@@ -532,15 +532,21 @@ public:
 
 	~FRDGParallelCommandListSet() override
 	{
-		Dispatch();
+		Dispatch(bHighPriority);
 	}
 
 	void SetStateOnCommandList(FRHICommandList& RHICmdList) override;
 
+	void SetHighPriority()
+	{
+		bHighPriority = true;
+	}
+
 private:
 	const FSceneRenderer& SceneRenderer;
 	FParallelCommandListBindings Bindings;
 	float ViewportScale;
+	bool bHighPriority = false;
 };
 
 enum EVolumeUpdateType
diff --git a/Engine/Source/Runtime/Renderer/Private/SceneVisibility.cpp b/Engine/Source/Runtime/Renderer/Private/SceneVisibility.cpp
--- a/Engine/Source/Runtime/Renderer/Private/SceneVisibility.cpp
+++ b/Engine/Source/Runtime/Renderer/Private/SceneVisibility.cpp
@@ -4410,7 +4410,7 @@ void FDeferredShadingSceneRenderer::PreVisibilityFrameSetup(FRHICommandListImmed
 	for (int32 ViewIndex = 0; ViewIndex < Views.Num(); ViewIndex++)
 	{
 		FViewInfo& View = Views[ViewIndex];
-		View.bAllowStencilDither = bDitheredLODTransitionsUseStencil;
+		View.bAllowStencilDither = DepthPass.bDitheredLODTransitionsUseStencil;
 	}
 
 	FSceneRenderer::PreVisibilityFrameSetup(RHICmdList);
