Fixed ray tracing shadows of directional lights were overriden by invalid shadow maps in some cases

diff --git a/Engine/Source/Runtime/Renderer/Private/LightRendering.cpp b/Engine/Source/Runtime/Renderer/Private/LightRendering.cpp
--- a/Engine/Source/Runtime/Renderer/Private/LightRendering.cpp
+++ b/Engine/Source/Runtime/Renderer/Private/LightRendering.cpp
@@ -174,11 +174,13 @@ FDeferredLightUniformStruct GetDeferredLightParameters(const FSceneView& View, c
 	FDeferredLightUniformStruct Parameters;
 	LightSceneInfo.Proxy->GetLightShaderParameters(Parameters.LightParameters);
 
-	const FVector2D FadeParams = LightSceneInfo.Proxy->GetDirectionalLightDistanceFadeParameters(View.GetFeatureLevel(), LightSceneInfo.IsPrecomputedLightingValid(), View.MaxShadowCascades);
+	const bool bIsRayTracedLight = ShouldRenderRayTracingShadows(*LightSceneInfo.Proxy);
 
+	const FVector2D FadeParams = LightSceneInfo.Proxy->GetDirectionalLightDistanceFadeParameters(View.GetFeatureLevel(), !bIsRayTracedLight && LightSceneInfo.IsPrecomputedLightingValid(), View.MaxShadowCascades);
+	
 	// use MAD for efficiency in the shader
 	Parameters.DistanceFadeMAD = FVector2D(FadeParams.Y, -FadeParams.X * FadeParams.Y);
-
+	
 	int32 ShadowMapChannel = LightSceneInfo.Proxy->GetShadowMapChannel();
 
 	static const auto AllowStaticLightingVar = IConsoleManager::Get().FindTConsoleVariableDataInt(TEXT("r.AllowStaticLighting"));
diff --git a/Engine/Source/Runtime/Renderer/Private/RayTracing/RayTracingLighting.cpp b/Engine/Source/Runtime/Renderer/Private/RayTracing/RayTracingLighting.cpp
--- a/Engine/Source/Runtime/Renderer/Private/RayTracing/RayTracingLighting.cpp
+++ b/Engine/Source/Runtime/Renderer/Private/RayTracing/RayTracingLighting.cpp
@@ -287,7 +287,8 @@ static void SetupRaytracingLightDataPacked(
 			LightDataElement.Tangent[Element] = LightParameters.Tangent[Element];
 		}
 
-		const FVector2D FadeParams = Light.LightSceneInfo->Proxy->GetDirectionalLightDistanceFadeParameters(View.GetFeatureLevel(), Light.LightSceneInfo->IsPrecomputedLightingValid(), View.MaxShadowCascades);
+		// Ray tracing should compute fade parameters ignoring lightmaps
+		const FVector2D FadeParams = Light.LightSceneInfo->Proxy->GetDirectionalLightDistanceFadeParameters(View.GetFeatureLevel(), false, View.MaxShadowCascades);
 		const FVector2D DistanceFadeMAD = { FadeParams.Y, -FadeParams.X * FadeParams.Y };
 
 		for (int32 Element = 0; Element < 2; Element++)
