Fixed crash when raytracing is enabled.

#rnx

diff --git a/Engine/Source/Runtime/Renderer/Private/DeferredShadingRenderer.cpp b/Engine/Source/Runtime/Renderer/Private/DeferredShadingRenderer.cpp
--- a/Engine/Source/Runtime/Renderer/Private/DeferredShadingRenderer.cpp
+++ b/Engine/Source/Runtime/Renderer/Private/DeferredShadingRenderer.cpp
@@ -1039,7 +1039,7 @@ bool FDeferredShadingSceneRenderer::GatherRayTracingWorldInstances(FRHICommandLi
 	return true;
 }
 
-bool FDeferredShadingSceneRenderer::DispatchRayTracingWorldUpdates(FRHICommandListImmediate& RHICmdList)
+bool FDeferredShadingSceneRenderer::DispatchRayTracingWorldUpdates(FRDGBuilder& GraphBuilder)
 {
 	if (!IsRayTracingEnabled() || Views.Num() == 0)
 	{
@@ -1083,9 +1083,9 @@ bool FDeferredShadingSceneRenderer::DispatchRayTracingWorldUpdates(FRHICommandLi
 	}
 
 	bool bAsyncUpdateGeometry = (CVarRayTracingAsyncBuild.GetValueOnRenderThread() != 0)
-							  && GRHISupportsRayTracingAsyncBuildAccelerationStructure;
+		&& GRHISupportsRayTracingAsyncBuildAccelerationStructure;
 
-	SCOPED_GPU_MASK(RHICmdList, FRHIGPUMask::All());
+	RDG_GPU_MASK_SCOPE(GraphBuilder, FRHIGPUMask::All());
 
 	for (int32 ViewIndex = 0; ViewIndex < Views.Num(); ++ViewIndex)
 	{
@@ -1118,8 +1118,8 @@ bool FDeferredShadingSceneRenderer::DispatchRayTracingWorldUpdates(FRHICommandLi
 		PrepareRayTracingTranslucency(View, RayGenShaders);
 		PrepareRayTracingDebug(View, RayGenShaders);
 		PreparePathTracing(View, RayGenShaders);
-		PrepareRayTracingLumenDirectLighting(View, *Scene,RayGenShaders);
-		PrepareRayTracingScreenProbeGather(View,RayGenShaders);
+		PrepareRayTracingLumenDirectLighting(View, *Scene, RayGenShaders);
+		PrepareRayTracingScreenProbeGather(View, RayGenShaders);
 
 		View.RayTracingScene.RayTracingSceneRHI = RHICreateRayTracingScene(SceneInitializer);
 
@@ -1127,7 +1127,7 @@ bool FDeferredShadingSceneRenderer::DispatchRayTracingWorldUpdates(FRHICommandLi
 		{
 			auto DefaultHitShader = View.ShaderMap->GetShader<FOpaqueShadowHitGroup>().GetRayTracingShader();
 
-			View.RayTracingMaterialPipeline = BindRayTracingMaterialPipeline(RHICmdList, View,
+			View.RayTracingMaterialPipeline = BindRayTracingMaterialPipeline(GraphBuilder.RHICmdList, View,
 				RayGenShaders,
 				DefaultHitShader
 			);
@@ -1135,7 +1135,7 @@ bool FDeferredShadingSceneRenderer::DispatchRayTracingWorldUpdates(FRHICommandLi
 
 		// Initialize common resources used for lighting in ray tracing effects
 
-		View.RayTracingSubSurfaceProfileTexture = GetSubsufaceProfileTexture_RT(RHICmdList);
+		View.RayTracingSubSurfaceProfileTexture = GetSubsufaceProfileTexture_RT(GraphBuilder.RHICmdList);
 		if (!View.RayTracingSubSurfaceProfileTexture)
 		{
 			View.RayTracingSubSurfaceProfileTexture = GSystemTextures.BlackDummy;
@@ -1143,59 +1143,62 @@ bool FDeferredShadingSceneRenderer::DispatchRayTracingWorldUpdates(FRHICommandLi
 
 		View.RayTracingSubSurfaceProfileSRV = RHICreateShaderResourceView(View.RayTracingSubSurfaceProfileTexture->GetRenderTargetItem().ShaderResourceTexture, 0);
 
-		View.RayTracingLightData = CreateRayTracingLightData(RHICmdList, Scene->Lights, View,
+		View.RayTracingLightData = CreateRayTracingLightData(GraphBuilder.RHICmdList, Scene->Lights, View,
 			EUniformBufferUsage::UniformBuffer_SingleFrame);
 	}
 
 	if (!bAsyncUpdateGeometry)
 	{
-		SCOPED_GPU_STAT(RHICmdList, RayTracingAS);
-		SCOPED_GPU_STAT(RHICmdList, RayTracingDynamicGeom);
-
-		Scene->GetRayTracingDynamicGeometryCollection()->DispatchUpdates(RHICmdList);
+		RDG_GPU_STAT_SCOPE(GraphBuilder, RayTracingAS);
+		RDG_GPU_STAT_SCOPE(GraphBuilder, RayTracingDynamicGeom);
 
+		AddPass(GraphBuilder, [this](FRHICommandListImmediate& RHICmdList)
 		{
-			SCOPED_DRAW_EVENT(RHICmdList, BuildRayTracingScene);
+			Scene->GetRayTracingDynamicGeometryCollection()->DispatchUpdates(RHICmdList);
+		});
 
+		AddPass(GraphBuilder, RDG_EVENT_NAME("BuildRayTracingScene"), [this](FRHICommandList& RHICmdList)
+		{
 			for (int32 ViewIndex = 0; ViewIndex < Views.Num(); ++ViewIndex)
 			{
 				FViewInfo& View = Views[ViewIndex];
 				RHICmdList.BuildAccelerationStructure(View.RayTracingScene.RayTracingSceneRHI);
 			}
-		}
+		});
 	}
 	else
 	{
-		// Build 'invalid' transition on empty resources to force transition to and from async compute
-		// RT structures currently have no state tracking yet, so this would need to be added first
-		FRHITransitionInfo RWNoBarrierTransition((FRHIUnorderedAccessView*)nullptr, ERHIAccess::Unknown, ERHIAccess::ERWNoBarrier);
-		FRHITransitionInfo RWBarrierTransition((FRHIUnorderedAccessView*)nullptr, ERHIAccess::Unknown, ERHIAccess::ERWBarrier);
-
-		check(RayTracingDynamicGeometryUpdateEndTransition == nullptr);
-		const FRHITransition* RayTracingDynamicGeometryUpdateBeginTransition = RHICreateTransition(ERHIPipeline::Graphics, ERHIPipeline::AsyncCompute, ERHICreateTransitionFlags::None, MakeArrayView(&RWNoBarrierTransition, 1));
-		RayTracingDynamicGeometryUpdateEndTransition = RHICreateTransition(ERHIPipeline::AsyncCompute, ERHIPipeline::Graphics, ERHICreateTransitionFlags::None, MakeArrayView(&RWBarrierTransition, 1));
+		AddPass(GraphBuilder, [this](FRHICommandList& RHICmdList)
+		{
+			check(RayTracingDynamicGeometryUpdateEndTransition == nullptr);
+			const FRHITransition* RayTracingDynamicGeometryUpdateBeginTransition = RHICreateTransition(ERHIPipeline::Graphics, ERHIPipeline::AsyncCompute, ERHICreateTransitionFlags::None, {});
+			RayTracingDynamicGeometryUpdateEndTransition = RHICreateTransition(ERHIPipeline::AsyncCompute, ERHIPipeline::Graphics, ERHICreateTransitionFlags::None, {});
 
-		FRHIAsyncComputeCommandListImmediate& RHIAsyncCmdList = FRHICommandListExecutor::GetImmediateAsyncComputeCommandList();
+			FRHIAsyncComputeCommandListImmediate& RHIAsyncCmdList = FRHICommandListExecutor::GetImmediateAsyncComputeCommandList();
 
-		// TArray<FRHIUnorderedAccessView*> VertexBuffers;
-		// Scene->GetRayTracingDynamicGeometryCollection()->GetVertexBufferUAVs(VertexBuffers);
+			// TArray<FRHIUnorderedAccessView*> VertexBuffers;
+			// Scene->GetRayTracingDynamicGeometryCollection()->GetVertexBufferUAVs(VertexBuffers);
 
-		RHICmdList.BeginTransition(RayTracingDynamicGeometryUpdateBeginTransition);
-		RHIAsyncCmdList.EndTransition(RayTracingDynamicGeometryUpdateBeginTransition);
+			RHICmdList.BeginTransition(RayTracingDynamicGeometryUpdateBeginTransition);
+			RHIAsyncCmdList.EndTransition(RayTracingDynamicGeometryUpdateBeginTransition);
 
-		Scene->GetRayTracingDynamicGeometryCollection()->DispatchUpdates(RHIAsyncCmdList);
+			Scene->GetRayTracingDynamicGeometryCollection()->DispatchUpdates(RHIAsyncCmdList);
 
-		for (int32 ViewIndex = 0; ViewIndex < Views.Num(); ++ViewIndex)
-		{
-			FViewInfo& View = Views[ViewIndex];
-			RHIAsyncCmdList.BuildAccelerationStructure(View.RayTracingScene.RayTracingSceneRHI);
-		}
+			for (int32 ViewIndex = 0; ViewIndex < Views.Num(); ++ViewIndex)
+			{
+				FViewInfo& View = Views[ViewIndex];
+				RHIAsyncCmdList.BuildAccelerationStructure(View.RayTracingScene.RayTracingSceneRHI);
+			}
 
-		RHIAsyncCmdList.BeginTransition(RayTracingDynamicGeometryUpdateEndTransition);
-		FRHIAsyncComputeCommandListImmediate::ImmediateDispatch(RHIAsyncCmdList);
+			RHIAsyncCmdList.BeginTransition(RayTracingDynamicGeometryUpdateEndTransition);
+			FRHIAsyncComputeCommandListImmediate::ImmediateDispatch(RHIAsyncCmdList);
+		});
 	}
 
-	Scene->GetRayTracingDynamicGeometryCollection()->EndUpdate(RHICmdList);
+	AddPass(GraphBuilder, [this](FRHICommandListImmediate& RHICmdList)
+	{
+		Scene->GetRayTracingDynamicGeometryCollection()->EndUpdate(RHICmdList);
+	});
 
 	return true;
 }
@@ -1806,11 +1809,8 @@ void FDeferredShadingSceneRenderer::Render(FRHICommandListImmediate& RHICmdList)
 	}
 
 #if RHI_RAYTRACING
-	AddPass(GraphBuilder, [this](FRHICommandListImmediate& RHICmdList)
-	{
-		// Must be done after FGlobalDynamicVertexBuffer::Get().Commit() for dynamic geometries to be updated
-		DispatchRayTracingWorldUpdates(RHICmdList);
-	});
+	// Must be done after FGlobalDynamicVertexBuffer::Get().Commit() for dynamic geometries to be updated
+	DispatchRayTracingWorldUpdates(GraphBuilder);
 #endif
 
 	TArray<Nanite::FRasterResults, TInlineAllocator<2>> NaniteRasterResults;
diff --git a/Engine/Source/Runtime/Renderer/Private/DeferredShadingRenderer.h b/Engine/Source/Runtime/Renderer/Private/DeferredShadingRenderer.h
--- a/Engine/Source/Runtime/Renderer/Private/DeferredShadingRenderer.h
+++ b/Engine/Source/Runtime/Renderer/Private/DeferredShadingRenderer.h
@@ -1065,7 +1065,7 @@ private:
 	void RenderRayTracingBarycentrics(FRDGBuilder& GraphBuilder, const FViewInfo& View, FRDGTextureRef SceneColorOutputTexture);
 
 	bool GatherRayTracingWorldInstances(FRHICommandListImmediate& RHICmdList);
-	bool DispatchRayTracingWorldUpdates(FRHICommandListImmediate& RHICmdList);
+	bool DispatchRayTracingWorldUpdates(FRDGBuilder& GraphBuilder);
 	FRayTracingPipelineState* BindRayTracingMaterialPipeline(FRHICommandList& RHICmdList, FViewInfo& View, const TArrayView<FRHIRayTracingShader*>& RayGenShaderTable, FRHIRayTracingShader* DefaultClosestHitShader);
 	FRayTracingPipelineState* BindRayTracingDeferredMaterialGatherPipeline(FRHICommandList& RHICmdList, const FViewInfo& View, const TArrayView<FRHIRayTracingShader*>& RayGenShaderTable);
 
