Fixed client shipping compilation error with Strata.

diff --git a/Engine/Source/Runtime/Renderer/Private/Strata/Strata.cpp b/Engine/Source/Runtime/Renderer/Private/Strata/Strata.cpp
--- a/Engine/Source/Runtime/Renderer/Private/Strata/Strata.cpp
+++ b/Engine/Source/Runtime/Renderer/Private/Strata/Strata.cpp
@@ -169,7 +169,7 @@ void AddVisualizeMaterialPasses(FRDGBuilder& GraphBuilder, const TArray<FViewInf
 		FVisualizeMaterialPS::FPermutationDomain PermutationVector;
 		TShaderMapRef<FVisualizeMaterialPS> PixelShader(View.ShaderMap, PermutationVector);
 
-		if (View.Family->EngineShowFlags.VisualizeStrataMaterial > 0)
+		if (View.Family->EngineShowFlags.VisualizeStrataMaterial)
 		{
 			FVisualizeMaterialPS::FParameters* PassParameters = GraphBuilder.AllocParameters<FVisualizeMaterialPS::FParameters>();
 			PassParameters->ViewUniformBuffer = View.ViewUniformBuffer;
