Removed deprecated LPV feature from the engine.

diff --git a/Engine/Plugins/Blendables/LightPropagationVolume/LightPropagationVolume.uplugin b/Engine/Plugins/Blendables/LightPropagationVolume/LightPropagationVolume.uplugin
deleted file mode 100644
diff --git a/Engine/Plugins/Blendables/LightPropagationVolume/Source/LightPropagationVolumeEditor/Classes/LightPropagationVolumeBlendableFactory.h b/Engine/Plugins/Blendables/LightPropagationVolume/Source/LightPropagationVolumeEditor/Classes/LightPropagationVolumeBlendableFactory.h
deleted file mode 100644
diff --git a/Engine/Plugins/Blendables/LightPropagationVolume/Source/LightPropagationVolumeEditor/LightPropagationVolumeEditor.Build.cs b/Engine/Plugins/Blendables/LightPropagationVolume/Source/LightPropagationVolumeEditor/LightPropagationVolumeEditor.Build.cs
deleted file mode 100644
diff --git a/Engine/Plugins/Blendables/LightPropagationVolume/Source/LightPropagationVolumeEditor/Private/AssetTypeActions_LightPropagationVolumeBlendable.h b/Engine/Plugins/Blendables/LightPropagationVolume/Source/LightPropagationVolumeEditor/Private/AssetTypeActions_LightPropagationVolumeBlendable.h
deleted file mode 100644
diff --git a/Engine/Plugins/Blendables/LightPropagationVolume/Source/LightPropagationVolumeEditor/Private/LightPropagationVolumeEditorModule.cpp b/Engine/Plugins/Blendables/LightPropagationVolume/Source/LightPropagationVolumeEditor/Private/LightPropagationVolumeEditorModule.cpp
deleted file mode 100644
diff --git a/Engine/Plugins/Blendables/LightPropagationVolume/Source/LightPropagationVolumeRuntime/LightPropagationVolumeRuntime.Build.cs b/Engine/Plugins/Blendables/LightPropagationVolume/Source/LightPropagationVolumeRuntime/LightPropagationVolumeRuntime.Build.cs
deleted file mode 100644
diff --git a/Engine/Plugins/Blendables/LightPropagationVolume/Source/LightPropagationVolumeRuntime/Private/LightPropagationVolumeRuntimeModule.cpp b/Engine/Plugins/Blendables/LightPropagationVolume/Source/LightPropagationVolumeRuntime/Private/LightPropagationVolumeRuntimeModule.cpp
deleted file mode 100644
diff --git a/Engine/Plugins/Blendables/LightPropagationVolume/Source/LightPropagationVolumeRuntime/Public/LightPropagationVolumeBlendable.h b/Engine/Plugins/Blendables/LightPropagationVolume/Source/LightPropagationVolumeRuntime/Public/LightPropagationVolumeBlendable.h
deleted file mode 100644
diff --git a/Engine/Plugins/FX/Niagara/Source/Niagara/Private/NiagaraComponent.cpp b/Engine/Plugins/FX/Niagara/Source/Niagara/Private/NiagaraComponent.cpp
--- a/Engine/Plugins/FX/Niagara/Source/Niagara/Private/NiagaraComponent.cpp
+++ b/Engine/Plugins/FX/Niagara/Source/Niagara/Private/NiagaraComponent.cpp
@@ -335,7 +335,6 @@ FRHIUniformBuffer* FNiagaraSceneProxy::GetUniformBufferNoVelocity() const
 				bHasPrecomputedVolumetricLightmap,
 				DrawsVelocity(),
 				GetLightingChannelMask(),
-				LpvBiasMultiplier,
 				LocalPrimitiveSceneInfo ? LocalPrimitiveSceneInfo->GetLightmapDataOffset() : 0,
 				SingleCaptureIndex,
 				false,
diff --git a/Engine/Shaders/Private/LPVBuildGeometryVolume.usf b/Engine/Shaders/Private/LPVBuildGeometryVolume.usf
deleted file mode 100644
diff --git a/Engine/Shaders/Private/LPVClear.usf b/Engine/Shaders/Private/LPVClear.usf
deleted file mode 100644
diff --git a/Engine/Shaders/Private/LPVClearLists.usf b/Engine/Shaders/Private/LPVClearLists.usf
deleted file mode 100644
diff --git a/Engine/Shaders/Private/LPVCommon.ush b/Engine/Shaders/Private/LPVCommon.ush
deleted file mode 100644
diff --git a/Engine/Shaders/Private/LPVDirectLightInject.usf b/Engine/Shaders/Private/LPVDirectLightInject.usf
deleted file mode 100644
diff --git a/Engine/Shaders/Private/LPVDirectionalOcclusion.usf b/Engine/Shaders/Private/LPVDirectionalOcclusion.usf
deleted file mode 100644
diff --git a/Engine/Shaders/Private/LPVFinalPass.ush b/Engine/Shaders/Private/LPVFinalPass.ush
deleted file mode 100644
diff --git a/Engine/Shaders/Private/LPVGeometryVolumeCommon.ush b/Engine/Shaders/Private/LPVGeometryVolumeCommon.ush
deleted file mode 100644
diff --git a/Engine/Shaders/Private/LPVInject_AccumulateVplLists.usf b/Engine/Shaders/Private/LPVInject_AccumulateVplLists.usf
deleted file mode 100644
diff --git a/Engine/Shaders/Private/LPVInject_GenerateVplLists.usf b/Engine/Shaders/Private/LPVInject_GenerateVplLists.usf
deleted file mode 100644
diff --git a/Engine/Shaders/Private/LPVListEntry.ush b/Engine/Shaders/Private/LPVListEntry.ush
deleted file mode 100644
diff --git a/Engine/Shaders/Private/LPVPacking.ush b/Engine/Shaders/Private/LPVPacking.ush
deleted file mode 100644
diff --git a/Engine/Shaders/Private/LPVPropagate.usf b/Engine/Shaders/Private/LPVPropagate.usf
deleted file mode 100644
diff --git a/Engine/Shaders/Private/LPVVisualise.usf b/Engine/Shaders/Private/LPVVisualise.usf
deleted file mode 100644
diff --git a/Engine/Shaders/Private/LPVWriteCommon.ush b/Engine/Shaders/Private/LPVWriteCommon.ush
deleted file mode 100644
diff --git a/Engine/Shaders/Private/LPVWriteVplCommon.ush b/Engine/Shaders/Private/LPVWriteVplCommon.ush
deleted file mode 100644
diff --git a/Engine/Shaders/Private/PostProcessLpvIndirect.usf b/Engine/Shaders/Private/PostProcessLpvIndirect.usf
deleted file mode 100644
diff --git a/Engine/Shaders/Private/ShadowDepthCommon.ush b/Engine/Shaders/Private/ShadowDepthCommon.ush
--- a/Engine/Shaders/Private/ShadowDepthCommon.ush
+++ b/Engine/Shaders/Private/ShadowDepthCommon.ush
@@ -8,12 +8,8 @@
 #define PERSPECTIVE_CORRECT_DEPTH 0
 #endif
 
-#ifndef REFLECTIVE_SHADOW_MAP
-#define REFLECTIVE_SHADOW_MAP 0
-#endif
-
-#define INTERPOLATE_VF_ATTRIBUTES (!MATERIALBLENDING_SOLID || REFLECTIVE_SHADOW_MAP || MATERIALBLENDING_MASKED || (REQUIRES_VF_ATTRIBUTES_FOR_CLIPPING && !POSITION_ONLY))
-#define INTERPOLATE_POSITION ((REFLECTIVE_SHADOW_MAP && !ONEPASS_POINTLIGHT_SHADOW) || MATERIALBLENDING_MASKED)
+#define INTERPOLATE_VF_ATTRIBUTES (!MATERIALBLENDING_SOLID || MATERIALBLENDING_MASKED || (REQUIRES_VF_ATTRIBUTES_FOR_CLIPPING && !POSITION_ONLY))
+#define INTERPOLATE_POSITION (MATERIALBLENDING_MASKED)
 
 /** Data passed from the vertex shader to the pixel shader for the shadow depth pass. */
 struct FShadowDepthVSToPS
diff --git a/Engine/Shaders/Private/ShadowDepthPixelShader.usf b/Engine/Shaders/Private/ShadowDepthPixelShader.usf
--- a/Engine/Shaders/Private/ShadowDepthPixelShader.usf
+++ b/Engine/Shaders/Private/ShadowDepthPixelShader.usf
@@ -8,7 +8,6 @@
 #define SHADOW_DEPTH_SHADER 1
 #define USE_STENCIL_LOD_DITHER	0
 
-#include "LPVListEntry.ush"
 #include "Common.ush"
 
 // Reroute SceneTexturesStruct uniform buffer references to the shadow depth pass uniform buffer
@@ -26,33 +25,9 @@
 
 #define SECONDARY_OCCLUSION			1
 
-#define NO_LPV						COMPILER_HLSLCC
-
-// With this enabled, very large or one-sided objects (e.g. terrain) inject an additional layer of GV cells 
-// to increase AO contribution. The AO calculation works best with closed geometry and small-medium scale objects
-#define LPV_GV_BOOST				1
-
-#if REFLECTIVE_SHADOW_MAP && !NO_LPV
-	// Reroute LpvWrite uniform buffer references to the shadow depth pass uniform buffer
-	#define LpvWrite PassStruct.LPV
-	#include "LPVWriteVplCommon.ush"
-#endif
-
-
-float3 GetMaterialBounceColor(FPixelMaterialInputs PixelMaterialInputs)
-{
-	half3 BaseColor = GetMaterialBaseColor(PixelMaterialInputs);
-	return BaseColor;
-}
-
-
 void Main( 
 	FShadowDepthVSToPS Inputs,
-	in float4 SvPosition : SV_Position,		// after all interpolators
-	out float4 OutColor : SV_Target0		// OutReflectiveShadowMapNormal, when REFLECTIVE_SHADOW_MAP is defined.
-#if REFLECTIVE_SHADOW_MAP
-	,out float4 OutReflectiveShadowMapDiffuse : SV_Target1
-#endif
+	in float4 SvPosition : SV_Position		// after all interpolators
 #if PERSPECTIVE_CORRECT_DEPTH
 	,out float OutDepth : SV_DEPTH
 #endif
@@ -89,75 +64,5 @@ void Main(
 	Inputs.ShadowDepth += Inputs.DepthBias;
 
 	OutDepth = saturate(Inputs.ShadowDepth);
-#endif
-
-	OutColor = 0;
-
-#if REFLECTIVE_SHADOW_MAP
-	OutReflectiveShadowMapDiffuse = 0;
-	float3 normal = normalize( MaterialParameters.TangentToWorld[2] ); // Smooth normal
-	float3 flux = GetMaterialBounceColor(PixelMaterialInputs);
-
-#if MATERIAL_INJECT_EMISSIVE_INTO_LPV
-	float3 emissive = GetMaterialEmissive(PixelMaterialInputs);
-
-	// Remove directionality from the light in the VPL.
-	normal = float3 (0.0f, 0.0f, 0.0f);
-#endif
-
-#if MATERIAL_INJECT_EMISSIVE_INTO_LPV || MATERIAL_BLOCK_GI
-	float opacity = GetMaterialOpacity(PixelMaterialInputs);
-#endif
-
-#if (SECONDARY_OCCLUSION || MATERIAL_INJECT_EMISSIVE_INTO_LPV) && !NO_LPV
-	// Optimisation: only add to GV/VPL if the fragment is within the grid
-	float3 grid = WorldToGrid( MaterialParameters.AbsoluteWorldPosition.xyz );
-	float minGrid = min( grid.x, min(grid.y,grid.z ) );
-	float maxGrid = max( grid.x, max(grid.y,grid.z ) );
-
-	[branch]
-	if ( minGrid > -0.5f && maxGrid <= 32.5f ) // Account for 0.5 cell offset
-	{
-#if SECONDARY_OCCLUSION && (MATERIALBLENDING_MASKED || MATERIALBLENDING_SOLID) && !MATERIAL_BLOCK_GI
-		float3 GvInjectPosition = MaterialParameters.AbsoluteWorldPosition - normal * LpvWrite.GeometryVolumeInjectionBias;
-#if LPV_GV_BOOST
-		if ( GetPrimitiveData(MaterialParameters.PrimitiveId).LpvBiasMultiplier < 0.25 ) 
-		{
-			// Add an additional layer of GV cells facing the opposite direction to improve AO contribution
-			// TODO: add per-primitive AO-boost tickbox instead of hijacking LpvBiasMultiplier like this
-			// NOTE: The cost of this additional branch is minimal (on AMD at least)
-			float offset = ( LpvWrite.LpvScale * 3.0f );
-			AddToGvList( GvInjectPosition - normal * offset, flux, -normal );
-		}
-#endif
-		AddToGvList( GvInjectPosition, flux, normal ); 
-
-#endif
-
-#if MATERIAL_INJECT_EMISSIVE_INTO_LPV && !MATERIAL_BLOCK_GI
-		AddToVplList( MaterialParameters.AbsoluteWorldPosition, emissive * LpvWrite.EmissiveInjectionMultiplier, normal, true );
-#endif
-	}
-	else
-	{
-		OutColor = 0;
-		OutReflectiveShadowMapDiffuse = 0;
-		return;
-	}
-#endif
-
-#if MATERIAL_BLOCK_GI
-	OutColor = 1.0f;
-	OutReflectiveShadowMapDiffuse = opacity;
-#elif MATERIALBLENDING_TRANSLUCENT || MATERIALBLENDING_ALPHACOMPOSITE || MATERIALBLENDING_ALPHAHOLDOUT
-	OutColor = 0;
-	OutReflectiveShadowMapDiffuse = 0;
-	clip(-1.0);
-#else
-	// Pass the LPV bias multiplier in the alpha channel of the normal
-	OutColor = float4(normal * .5 + .5, GetPrimitiveData(MaterialParameters.PrimitiveId).LpvBiasMultiplier / 32.0f );
-	OutReflectiveShadowMapDiffuse = float4(flux, 0);
-#endif
-
 #endif
 }
diff --git a/Engine/Shaders/Private/ShadowDepthVertexShader.usf b/Engine/Shaders/Private/ShadowDepthVertexShader.usf
--- a/Engine/Shaders/Private/ShadowDepthVertexShader.usf
+++ b/Engine/Shaders/Private/ShadowDepthVertexShader.usf
@@ -8,7 +8,6 @@
 #define SHADOW_DEPTH_SHADER 1
 #define USE_STENCIL_LOD_DITHER	0
 
-#include "LPVListEntry.ush"
 #include "Common.ush"
 
 // Reroute SceneTexturesStruct uniform buffer references to the shadow depth pass uniform buffer
diff --git a/Engine/Source/Editor/UnrealEd/Private/StaticMeshEdit.cpp b/Engine/Source/Editor/UnrealEd/Private/StaticMeshEdit.cpp
--- a/Engine/Source/Editor/UnrealEd/Private/StaticMeshEdit.cpp
+++ b/Engine/Source/Editor/UnrealEd/Private/StaticMeshEdit.cpp
@@ -1186,7 +1186,6 @@ TSharedPtr<FExistingStaticMeshData> SaveExistingStaticMeshData(UStaticMesh* Exis
 
 		ExistingMeshDataPtr->ExistingBodySetup = ExistingMesh->BodySetup;
 
-		ExistingMeshDataPtr->LpvBiasMultiplier = ExistingMesh->LpvBiasMultiplier;
 		ExistingMeshDataPtr->bHasNavigationData = ExistingMesh->bHasNavigationData;
 		ExistingMeshDataPtr->LODGroup = ExistingMesh->LODGroup;
 		ExistingMeshDataPtr->MinLOD = ExistingMesh->MinLOD;
@@ -1720,7 +1719,6 @@ void RestoreExistingMeshData(TSharedPtr<const FExistingStaticMeshData> ExistingM
 		}
 	}
 
-	NewMesh->LpvBiasMultiplier = ExistingMeshDataPtr->LpvBiasMultiplier;
 	NewMesh->bHasNavigationData = ExistingMeshDataPtr->bHasNavigationData;
 	NewMesh->LODGroup = ExistingMeshDataPtr->LODGroup;
 
diff --git a/Engine/Source/Runtime/Engine/Classes/Components/PrimitiveComponent.h b/Engine/Source/Runtime/Engine/Classes/Components/PrimitiveComponent.h
--- a/Engine/Source/Runtime/Engine/Classes/Components/PrimitiveComponent.h
+++ b/Engine/Source/Runtime/Engine/Classes/Components/PrimitiveComponent.h
@@ -647,14 +647,6 @@ public:
 	/** Used by the renderer, to identify a component across re-registers. */
 	FPrimitiveComponentId ComponentId;
 
-	/**
-	 * Multiplier used to scale the Light Propagation Volume light injection bias, to reduce light bleeding. 
-	 * Set to 0 for no bias, 1 for default or higher for increased biasing (e.g. for 
-	 * thin geometry such as walls)
-	 */
-	UPROPERTY(EditAnywhere, BlueprintReadOnly, AdvancedDisplay, Category=Rendering, meta=(UIMin = "0.0", UIMax = "3.0"))
-	float LpvBiasMultiplier;
-
 	/**
 	* Incremented by the main thread before being attached to the scene, decremented
 	* by the rendering thread after removal. This counter exists to assert that 
diff --git a/Engine/Source/Runtime/Engine/Classes/Engine/RendererSettings.h b/Engine/Source/Runtime/Engine/Classes/Engine/RendererSettings.h
--- a/Engine/Source/Runtime/Engine/Classes/Engine/RendererSettings.h
+++ b/Engine/Source/Runtime/Engine/Classes/Engine/RendererSettings.h
@@ -815,12 +815,6 @@ class ENGINE_API URendererSettings : public UDeveloperSettings
 		ConfigRestartRequired = true))
 		uint32 bSupportReversedIndexBuffers : 1;
 
-	UPROPERTY(config, EditAnywhere, Category = Lighting, meta = (
-		ConsoleVariable = "r.LightPropagationVolume", DisplayName = "Light Propagation Volumes",
-		ToolTip = "Whether to allow the usage and compilation of Light Propagation Volumes.",
-		ConfigRestartRequired = true))
-		uint32 bLPV : 1;
-
 	UPROPERTY(config, EditAnywhere, Category = Mobile, meta = (
 		ConsoleVariable = "r.Mobile.AmbientOcclusion", DisplayName = "Mobile Ambient Occlusion",
 		ToolTip = "Mobile Ambient Occlusion. Causion: An extra sampler will be occupied in mobile base pass pixel shader after enable the mobile ambient occlusion. Changing this setting requires restarting the editor.",
diff --git a/Engine/Source/Runtime/Engine/Classes/Engine/Scene.h b/Engine/Source/Runtime/Engine/Classes/Engine/Scene.h
--- a/Engine/Source/Runtime/Engine/Classes/Engine/Scene.h
+++ b/Engine/Source/Runtime/Engine/Classes/Engine/Scene.h
@@ -1012,51 +1012,6 @@ struct FPostProcessSettings
 	UPROPERTY(EditAnywhere, BlueprintReadWrite, Category = Overrides, meta = (PinHiddenByDefault, InlineEditConditionToggle))
 	uint32 bOverride_RayTracingAORadius : 1;
 
-	UPROPERTY(EditAnywhere, BlueprintReadWrite, Category=Overrides, meta=(PinHiddenByDefault, InlineEditConditionToggle))
-	uint8 bOverride_LPVIntensity:1;
-
-	UPROPERTY(EditAnywhere, Category=Overrides, meta=(InlineEditConditionToggle))
-	uint8 bOverride_LPVDirectionalOcclusionIntensity:1;
-
-	UPROPERTY(EditAnywhere, Category=Overrides, meta=(InlineEditConditionToggle))
-	uint8 bOverride_LPVDirectionalOcclusionRadius:1;
-
-	UPROPERTY(EditAnywhere, Category=Overrides, meta=(InlineEditConditionToggle))
-	uint8 bOverride_LPVDiffuseOcclusionExponent:1;
-
-	UPROPERTY(EditAnywhere, Category=Overrides, meta=(InlineEditConditionToggle))
-	uint8 bOverride_LPVSpecularOcclusionExponent:1;
-
-	UPROPERTY(EditAnywhere, Category=Overrides, meta=(InlineEditConditionToggle))
-	uint8 bOverride_LPVDiffuseOcclusionIntensity:1;
-
-	UPROPERTY(EditAnywhere, Category=Overrides, meta=(InlineEditConditionToggle))
-	uint8 bOverride_LPVSpecularOcclusionIntensity:1;
-
-	UPROPERTY(EditAnywhere, BlueprintReadWrite, Category=Overrides, meta=(PinHiddenByDefault, InlineEditConditionToggle))
-	uint8 bOverride_LPVSize:1;
-
-	UPROPERTY(EditAnywhere, BlueprintReadWrite, Category=Overrides, meta=(PinHiddenByDefault, InlineEditConditionToggle))
-	uint8 bOverride_LPVSecondaryOcclusionIntensity:1;
-
-	UPROPERTY(EditAnywhere, BlueprintReadWrite, Category=Overrides, meta=(PinHiddenByDefault, InlineEditConditionToggle))
-	uint8 bOverride_LPVSecondaryBounceIntensity:1;
-
-	UPROPERTY(EditAnywhere, BlueprintReadWrite, Category=Overrides, meta=(PinHiddenByDefault, InlineEditConditionToggle))
-	uint8 bOverride_LPVGeometryVolumeBias:1;
-
-	UPROPERTY(EditAnywhere, BlueprintReadWrite, Category=Overrides, meta=(PinHiddenByDefault, InlineEditConditionToggle))
-	uint8 bOverride_LPVVplInjectionBias:1;
-
-	UPROPERTY(EditAnywhere, BlueprintReadWrite, Category=Overrides, meta=(PinHiddenByDefault, InlineEditConditionToggle))
-	uint8 bOverride_LPVEmissiveInjectionIntensity:1;
-
-	UPROPERTY(EditAnywhere, BlueprintReadWrite, Category=Overrides, meta=(PinHiddenByDefault, InlineEditConditionToggle))
-	uint8 bOverride_LPVFadeRange : 1;
-
-	UPROPERTY(EditAnywhere, BlueprintReadWrite, Category=Overrides, meta = (PinHiddenByDefault, InlineEditConditionToggle))
-	uint8 bOverride_LPVDirectionalOcclusionFadeRange : 1;
-
 	UPROPERTY(EditAnywhere, BlueprintReadWrite, Category=Overrides, meta=(PinHiddenByDefault, InlineEditConditionToggle))
 	uint8 bOverride_IndirectLightingColor:1;
 
@@ -1807,57 +1762,6 @@ struct FPostProcessSettings
 	UPROPERTY(interp, BlueprintReadWrite, Category="Rendering Features|Motion Blur", meta=(ClampMin = "0.0", UIMax = "100.0", editcondition = "bOverride_MotionBlurPerObjectSize", DisplayName = "Per Object Size"))
 	float MotionBlurPerObjectSize;
 
-	/** How strong the dynamic GI from the LPV should be. 0.0 is off, 1.0 is the "normal" value, but higher values can be used to boost the effect*/
-	UPROPERTY(interp, BlueprintReadWrite, Category="Rendering Features|Light Propagation Volume", meta=(editcondition = "bOverride_LPVIntensity", UIMin = "0", UIMax = "20", DisplayName = "Intensity"))
-	float LPVIntensity;
-
-	/** Bias applied to light injected into the LPV in cell units. Increase to reduce bleeding through thin walls*/
-	UPROPERTY(interp, BlueprintReadWrite, Category="Rendering Features|Light Propagation Volume", AdvancedDisplay, meta=(editcondition = "bOverride_LPVVplInjectionBias", UIMin = "0", UIMax = "2", DisplayName = "Light Injection Bias"))
-	float LPVVplInjectionBias;
-
-	/** The size of the LPV volume, in Unreal units*/
-	UPROPERTY(EditAnywhere, BlueprintReadWrite, Category="Rendering Features|Light Propagation Volume", meta=(editcondition = "bOverride_LPVSize", UIMin = "100", UIMax = "20000", DisplayName = "Size"))
-	float LPVSize;
-
-	/** Secondary occlusion strength (bounce light shadows). Set to 0 to disable*/
-	UPROPERTY(interp, BlueprintReadWrite, Category="Rendering Features|Light Propagation Volume", meta=(editcondition = "bOverride_LPVSecondaryOcclusionIntensity", UIMin = "0", UIMax = "1", DisplayName = "Secondary Occlusion Intensity"))
-	float LPVSecondaryOcclusionIntensity;
-
-	/** Secondary bounce light strength (bounce light shadows). Set to 0 to disable*/
-	UPROPERTY(interp, BlueprintReadWrite, Category="Rendering Features|Light Propagation Volume", AdvancedDisplay, meta=(editcondition = "bOverride_LPVSecondaryBounceIntensity", UIMin = "0", UIMax = "1", DisplayName = "Secondary Bounce Intensity"))
-	float LPVSecondaryBounceIntensity;
-
-	/** Bias applied to the geometry volume in cell units. Increase to reduce darkening due to secondary occlusion */
-	UPROPERTY(interp, BlueprintReadWrite, Category="Rendering Features|Light Propagation Volume", AdvancedDisplay, meta=(editcondition = "bOverride_LPVGeometryVolumeBias", UIMin = "0", UIMax = "2", DisplayName = "Geometry Volume Bias"))
-	float LPVGeometryVolumeBias;
-
-	UPROPERTY(interp, BlueprintReadWrite, Category="Rendering Features|Light Propagation Volume", AdvancedDisplay, meta=(editcondition = "bOverride_LPVEmissiveInjectionIntensity", UIMin = "0", UIMax = "20", DisplayName = "Emissive Injection Intensity"))
-	float LPVEmissiveInjectionIntensity;
-
-	/** Controls the amount of directional occlusion. Requires LPV. Values very close to 1.0 are recommended */
-	UPROPERTY(interp, BlueprintReadWrite, Category="Rendering Features|Light Propagation Volume", meta=(editcondition = "bOverride_LPVDirectionalOcclusionIntensity", UIMin = "0", UIMax = "1", DisplayName = "Occlusion Intensity"))
-	float LPVDirectionalOcclusionIntensity;
-
-	/** Occlusion Radius - 16 is recommended for most scenes */
-	UPROPERTY(interp, BlueprintReadWrite, Category="Rendering Features|Light Propagation Volume", AdvancedDisplay, meta=(editcondition = "bOverride_LPVDirectionalOcclusionRadius", UIMin = "1", UIMax = "16", DisplayName = "Occlusion Radius"))
-	float LPVDirectionalOcclusionRadius;
-
-	/** Diffuse occlusion exponent - increase for more contrast. 1 to 2 is recommended */
-	UPROPERTY(interp, BlueprintReadWrite, Category="Rendering Features|Light Propagation Volume", meta=(editcondition = "bOverride_LPVDiffuseOcclusionExponent", UIMin = "0.5", UIMax = "5", DisplayName = "Diffuse occlusion exponent"))
-	float LPVDiffuseOcclusionExponent;
-
-	/** Specular occlusion exponent - increase for more contrast. 6 to 9 is recommended */
-	UPROPERTY(interp, BlueprintReadWrite, Category="Rendering Features|Light Propagation Volume", meta=(editcondition = "bOverride_LPVSpecularOcclusionExponent", UIMin = "1", UIMax = "16", DisplayName = "Specular occlusion exponent"))
-	float LPVSpecularOcclusionExponent;
-
-	/** Diffuse occlusion intensity - higher values provide increased diffuse occlusion.*/
-	UPROPERTY(interp, BlueprintReadWrite, Category="Rendering Features|Light Propagation Volume", AdvancedDisplay, meta=(editcondition = "bOverride_LPVDiffuseOcclusionIntensity", UIMin = "0", UIMax = "4", DisplayName = "Diffuse occlusion intensity"))
-	float LPVDiffuseOcclusionIntensity;
-
-	/** Specular occlusion intensity - higher values provide increased specular occlusion.*/
-	UPROPERTY(interp, BlueprintReadWrite, Category="Rendering Features|Light Propagation Volume", AdvancedDisplay, meta=(editcondition = "bOverride_LPVSpecularOcclusionIntensity", UIMin = "0", UIMax = "4", DisplayName = "Specular occlusion intensity"))
-	float LPVSpecularOcclusionIntensity;
-
 	/** Sets the reflections type */
 	UPROPERTY(EditAnywhere, BlueprintReadWrite, Category = "Rendering Features|Reflections", meta = (editcondition = "bOverride_ReflectionsType", DisplayName = "Type"))
 	EReflectionsType ReflectionsType;
@@ -1929,16 +1833,6 @@ struct FPostProcessSettings
 	UPROPERTY(EditAnywhere, BlueprintReadWrite, Category = "Rendering Features|PathTracing", meta = (ClampMin = "1", ClampMax = "64", editcondition = "bOverride_PathTracingSamplesPerPixel", DisplayName = "Samples Per Pixel"))
 	int32 PathTracingSamplesPerPixel;
 
-
-
-	/** LPV Fade range - increase to fade more gradually towards the LPV edges.*/
-	UPROPERTY(interp, BlueprintReadWrite, Category = "Rendering Features|Light Propagation Volume", AdvancedDisplay, meta = (editcondition = "bOverride_LPVFadeRange", UIMin = "0", UIMax = "9", DisplayName = "Fade range"))
-	float LPVFadeRange;
-
-	/** LPV Directional Occlusion Fade range - increase to fade more gradually towards the LPV edges.*/
-	UPROPERTY(interp, BlueprintReadWrite, Category = "Rendering Features|Light Propagation Volume", AdvancedDisplay, meta = (editcondition = "bOverride_LPVDirectionalOcclusionFadeRange", UIMin = "0", UIMax = "9", DisplayName = "DO Fade range"))
-	float LPVDirectionalOcclusionFadeRange;
-
 	/**
 	* To render with lower or high resolution than it is presented,
 	* controlled by console variable,
diff --git a/Engine/Source/Runtime/Engine/Classes/Engine/StaticMesh.h b/Engine/Source/Runtime/Engine/Classes/Engine/StaticMesh.h
--- a/Engine/Source/Runtime/Engine/Classes/Engine/StaticMesh.h
+++ b/Engine/Source/Runtime/Engine/Classes/Engine/StaticMesh.h
@@ -614,10 +614,6 @@ class UStaticMesh : public UStreamableRenderAsset, public IInterface_CollisionDa
 		return INDEX_NONE;
 	}
 
-	/** Bias multiplier for Light Propagation Volume lighting */
-	UPROPERTY(EditAnywhere, BlueprintReadOnly, Category=StaticMesh, meta=(UIMin = "0.0", UIMax = "3.0"))
-	float LpvBiasMultiplier;
-
 	UPROPERTY(EditAnywhere, BlueprintReadWrite, Category = StaticMesh)
 	TArray<FStaticMaterial> StaticMaterials;
 
diff --git a/Engine/Source/Runtime/Engine/Classes/Materials/Material.h b/Engine/Source/Runtime/Engine/Classes/Materials/Material.h
--- a/Engine/Source/Runtime/Engine/Classes/Materials/Material.h
+++ b/Engine/Source/Runtime/Engine/Classes/Materials/Material.h
@@ -593,12 +593,6 @@ public:
 	UPROPERTY(EditAnywhere, BlueprintReadWrite, Category=Material, meta=(DisplayName = "Emissive (Dynamic Area Light)"), AdvancedDisplay)
 	uint8 bUseEmissiveForDynamicAreaLighting : 1;
 
-	/**
-	 * If enabled, the material's opacity defines how much GI is blocked when using the LightPropagationVolume feature
-	 */
-	UPROPERTY(EditAnywhere, BlueprintReadWrite, Category=Material, meta=(DisplayName = "Block Global Illumination"), AdvancedDisplay)
-	uint8 bBlockGI : 1;
-
 	/** 
 	 * This is a special usage flag that allows a material to be assignable to any primitive type.
 	 * This is useful for materials used by code to implement certain viewmodes, for example the default material or lighting only material.
diff --git a/Engine/Source/Runtime/Engine/Private/Components/DirectionalLightComponent.cpp b/Engine/Source/Runtime/Engine/Private/Components/DirectionalLightComponent.cpp
--- a/Engine/Source/Runtime/Engine/Private/Components/DirectionalLightComponent.cpp
+++ b/Engine/Source/Runtime/Engine/Private/Components/DirectionalLightComponent.cpp
@@ -378,48 +378,6 @@ public:
 		return true;
 	}
 
-	// reflective shadow map for Light Propagation Volume
-	virtual bool GetViewDependentRsmWholeSceneProjectedShadowInitializer(
-		const class FSceneView& View, 
-		const FBox& LightPropagationVolumeBounds,
-	    class FWholeSceneProjectedShadowInitializer& OutInitializer ) const override
-	{
-		float LpvExtent = LightPropagationVolumeBounds.GetExtent().X; // LPV is a cube, so this should be valid
-
-		OutInitializer.PreShadowTranslation = -LightPropagationVolumeBounds.GetCenter();
-		OutInitializer.WorldToLight = FInverseRotationMatrix(GetDirection().GetSafeNormal().Rotation());
-		OutInitializer.Scales = FVector2D(1.0f / LpvExtent,1.0f / LpvExtent);
-		OutInitializer.SubjectBounds = FBoxSphereBounds( FVector::ZeroVector, LightPropagationVolumeBounds.GetExtent(), FMath::Sqrt( LpvExtent * LpvExtent * 3.0f ) );
-		OutInitializer.WAxis = FVector4(0,0,0,1);
-		// Use the minimum of half the world, things further away do not cast shadows,
-		// However, if the cascade bounds are larger, then extend the casting distance far enough to encompass the cascade.
-		OutInitializer.MinLightW = FMath::Min<float>(-HALF_WORLD_MAX, -OutInitializer.SubjectBounds.SphereRadius);
-		// Range must extend to end of cascade bounds
-		const float MaxLightW = OutInitializer.SubjectBounds.SphereRadius;
-		OutInitializer.MaxDistanceToCastInLightW = MaxLightW - OutInitializer.MinLightW;
-
-		// Compute the RSM bounds
-		{
-			FVector Centre = LightPropagationVolumeBounds.GetCenter();
-			FVector Extent = LightPropagationVolumeBounds.GetExtent();
-			FVector CascadeFrustumVerts[8];
-			CascadeFrustumVerts[0] = Centre + Extent * FVector( -1.0f,-1.0f, 1.0f ); // 0 Near Top    Right
-			CascadeFrustumVerts[1] = Centre + Extent * FVector( -1.0f,-1.0f,-1.0f ); // 1 Near Bottom Right
-			CascadeFrustumVerts[2] = Centre + Extent * FVector(  1.0f,-1.0f, 1.0f ); // 2 Near Top    Left
-			CascadeFrustumVerts[3] = Centre + Extent * FVector(  1.0f,-1.0f,-1.0f ); // 3 Near Bottom Left
-			CascadeFrustumVerts[4] = Centre + Extent * FVector( -1.0f, 1.0f, 1.0f ); // 4 Far  Top    Right
-			CascadeFrustumVerts[5] = Centre + Extent * FVector( -1.0f, 1.0f,-1.0f ); // 5 Far  Bottom Right
-			CascadeFrustumVerts[6] = Centre + Extent * FVector(  1.0f, 1.0f, 1.0f ); // 6 Far  Top    Left
-			CascadeFrustumVerts[7] = Centre + Extent * FVector(  1.0f, 1.0f,-1.0f ); // 7 Far  Bottom Left
-
-			FPlane Far;
-			FPlane Near;
-			const FVector LightDirection = -GetDirection();
-			ComputeShadowCullingVolume(View.bReverseCulling, CascadeFrustumVerts, LightDirection, OutInitializer.CascadeSettings.ShadowBoundsAccurate, Near, Far);
-		}
-		return true;
-	}
-
 	virtual FVector2D GetDirectionalLightDistanceFadeParameters(ERHIFeatureLevel::Type InFeatureLevel, bool bPrecomputedLightingIsValid, int32 MaxNearCascades) const override
 	{
 		float FarDistance = GetCSMMaxDistance(bPrecomputedLightingIsValid, MaxNearCascades);
diff --git a/Engine/Source/Runtime/Engine/Private/Components/PrimitiveComponent.cpp b/Engine/Source/Runtime/Engine/Private/Components/PrimitiveComponent.cpp
--- a/Engine/Source/Runtime/Engine/Private/Components/PrimitiveComponent.cpp
+++ b/Engine/Source/Runtime/Engine/Private/Components/PrimitiveComponent.cpp
@@ -327,7 +327,6 @@ UPrimitiveComponent::UPrimitiveComponent(const FObjectInitializer& ObjectInitial
 	bCastDynamicShadow = true;
 	bAffectDynamicIndirectLighting = true;
 	bAffectDistanceFieldLighting = true;
-	LpvBiasMultiplier = 1.0f;
 	bCastStaticShadow = true;
 	bCastVolumetricTranslucentShadow = false;
 	bCastContactShadow = true;
diff --git a/Engine/Source/Runtime/Engine/Private/DynamicMeshBuilder.cpp b/Engine/Source/Runtime/Engine/Private/DynamicMeshBuilder.cpp
--- a/Engine/Source/Runtime/Engine/Private/DynamicMeshBuilder.cpp
+++ b/Engine/Source/Runtime/Engine/Private/DynamicMeshBuilder.cpp
@@ -867,7 +867,6 @@ void FDynamicMeshBuilder::GetMesh(
 				false,
 				/* bDrawsVelocity = */ true,
 				GetDefaultLightingChannelMask(),
-				0,
 				INDEX_NONE,
 				INDEX_NONE,
 				false
@@ -960,7 +959,6 @@ void FDynamicMeshBuilder::GetMeshElement(const FMatrix& LocalToWorld, const FMat
 				false,
 				/* bDrawsVelocity = */ true,
 				GetDefaultLightingChannelMask(),
-				0,
 				INDEX_NONE,
 				INDEX_NONE,
 				false
@@ -1040,7 +1038,6 @@ void FDynamicMeshBuilder::Draw(FPrimitiveDrawInterface* PDI,const FMatrix& Local
 			false,
 			/* bDrawsVelocity = */ true,
 			GetDefaultLightingChannelMask(),
-			0,
 			INDEX_NONE,
 			INDEX_NONE,
 			false
diff --git a/Engine/Source/Runtime/Engine/Private/Materials/Material.cpp b/Engine/Source/Runtime/Engine/Private/Materials/Material.cpp
--- a/Engine/Source/Runtime/Engine/Private/Materials/Material.cpp
+++ b/Engine/Source/Runtime/Engine/Private/Materials/Material.cpp
@@ -980,7 +980,6 @@ PRAGMA_ENABLE_DEPRECATION_WARNINGS
 	bEnableStencilTest = false;
 
 	bUseEmissiveForDynamicAreaLighting = false;
-	bBlockGI = false;
 	RefractionDepthBias = 0.0f;
 	MaterialDecalResponse = MDR_ColorNormalRoughness;
 
diff --git a/Engine/Source/Runtime/Engine/Private/Materials/MaterialShared.cpp b/Engine/Source/Runtime/Engine/Private/Materials/MaterialShared.cpp
--- a/Engine/Source/Runtime/Engine/Private/Materials/MaterialShared.cpp
+++ b/Engine/Source/Runtime/Engine/Private/Materials/MaterialShared.cpp
@@ -1058,8 +1058,6 @@ void FMaterial::DiscardShaderMap()
 
 EMaterialDomain FMaterialResource::GetMaterialDomain() const { return Material->MaterialDomain; }
 bool FMaterialResource::IsTangentSpaceNormal() const { return Material->bTangentSpaceNormal || (!Material->Normal.IsConnected() && !Material->bUseMaterialAttributes); }
-bool FMaterialResource::ShouldInjectEmissiveIntoLPV() const { return Material->bUseEmissiveForDynamicAreaLighting; }
-bool FMaterialResource::ShouldBlockGI() const { return Material->bBlockGI; }
 bool FMaterialResource::ShouldGenerateSphericalParticleNormals() const { return Material->bGenerateSphericalParticleNormals; }
 bool FMaterialResource::ShouldDisableDepthTest() const { return Material->bDisableDepthTest; }
 bool FMaterialResource::ShouldWriteOnlyAlpha() const { return Material->bWriteOnlyAlpha; }
@@ -1768,10 +1766,8 @@ void FMaterial::SetupMaterialEnvironment(
 	OutEnvironment.SetDefine(TEXT("MATERIAL_PLANAR_FORWARD_REFLECTIONS"), IsUsingPlanarForwardReflections());
 	OutEnvironment.SetDefine(TEXT("MATERIAL_NONMETAL"), IsNonmetal());
 	OutEnvironment.SetDefine(TEXT("MATERIAL_USE_LM_DIRECTIONALITY"), UseLmDirectionality());
-	OutEnvironment.SetDefine(TEXT("MATERIAL_INJECT_EMISSIVE_INTO_LPV"), ShouldInjectEmissiveIntoLPV());
 	OutEnvironment.SetDefine(TEXT("MATERIAL_SSR"), ShouldDoSSR() && IsTranslucentBlendMode(GetBlendMode()));
 	OutEnvironment.SetDefine(TEXT("MATERIAL_CONTACT_SHADOWS"), ShouldDoContactShadows() && IsTranslucentBlendMode(GetBlendMode()));
-	OutEnvironment.SetDefine(TEXT("MATERIAL_BLOCK_GI"), ShouldBlockGI());
 	OutEnvironment.SetDefine(TEXT("MATERIAL_DITHER_OPACITY_MASK"), IsDitherMasked());
 	OutEnvironment.SetDefine(TEXT("MATERIAL_NORMAL_CURVATURE_TO_ROUGHNESS"), UseNormalCurvatureToRoughness() ? TEXT("1") : TEXT("0"));
 	OutEnvironment.SetDefine(TEXT("MATERIAL_ALLOW_NEGATIVE_EMISSIVECOLOR"), AllowNegativeEmissiveColor());
diff --git a/Engine/Source/Runtime/Engine/Private/Particles/ParticleSystemRender.cpp b/Engine/Source/Runtime/Engine/Private/Particles/ParticleSystemRender.cpp
--- a/Engine/Source/Runtime/Engine/Private/Particles/ParticleSystemRender.cpp
+++ b/Engine/Source/Runtime/Engine/Private/Particles/ParticleSystemRender.cpp
@@ -7208,7 +7208,6 @@ void FParticleSystemSceneProxy::UpdateWorldSpacePrimitiveUniformBuffer() const
 			GetScene().HasPrecomputedVolumetricLightmap_RenderThread(),
 			DrawsVelocity(),
 			GetLightingChannelMask(),
-			0,
 			INDEX_NONE,
 			INDEX_NONE,
 			AlwaysHasVelocity()
diff --git a/Engine/Source/Runtime/Engine/Private/PrimitiveSceneProxy.cpp b/Engine/Source/Runtime/Engine/Private/PrimitiveSceneProxy.cpp
--- a/Engine/Source/Runtime/Engine/Private/PrimitiveSceneProxy.cpp
+++ b/Engine/Source/Runtime/Engine/Private/PrimitiveSceneProxy.cpp
@@ -176,7 +176,6 @@ FPrimitiveSceneProxy::FPrimitiveSceneProxy(const UPrimitiveComponent* InComponen
 ,	VirtualTextureLodBias(InComponent->VirtualTextureLodBias)
 ,	VirtualTextureCullMips(InComponent->VirtualTextureCullMips)
 ,	VirtualTextureMinCoverage(InComponent->VirtualTextureMinCoverage)
-,	LpvBiasMultiplier(InComponent->LpvBiasMultiplier)
 ,	DynamicIndirectShadowMinVisibility(0)
 ,	PrimitiveComponentId(InComponent->ComponentId)
 ,	Scene(InComponent->GetScene())
@@ -396,7 +395,6 @@ void FPrimitiveSceneProxy::UpdateUniformBuffer()
 				bHasPrecomputedVolumetricLightmap,
 				DrawsVelocity(), 
 				GetLightingChannelMask(),
-				LpvBiasMultiplier,
 				PrimitiveSceneInfo ? PrimitiveSceneInfo->GetLightmapDataOffset() : 0,
 				SingleCaptureIndex,
 				bOutputVelocity || AlwaysHasVelocity(),
diff --git a/Engine/Source/Runtime/Engine/Private/PrimitiveUniformShaderParameters.cpp b/Engine/Source/Runtime/Engine/Private/PrimitiveUniformShaderParameters.cpp
--- a/Engine/Source/Runtime/Engine/Private/PrimitiveUniformShaderParameters.cpp
+++ b/Engine/Source/Runtime/Engine/Private/PrimitiveUniformShaderParameters.cpp
@@ -128,7 +128,6 @@ FPrimitiveSceneShaderData::FPrimitiveSceneShaderData(const FPrimitiveSceneProxy*
 		bHasPrecomputedVolumetricLightmap,
 		Proxy->DrawsVelocity(), 
 		Proxy->GetLightingChannelMask(),
-		Proxy->GetLpvBiasMultiplier(),
 		Proxy->GetPrimitiveSceneInfo()->GetLightmapDataOffset(),
 		SingleCaptureIndex,
         bOutputVelocity,
@@ -164,7 +163,7 @@ void FPrimitiveSceneShaderData::Setup(const FPrimitiveUniformShaderParameters& P
 	Data[17] = *(const FVector4*)&PrimitiveUniformShaderParameters.PreviousWorldToLocal.M[3][0];
 
 	Data[18] = FVector4(PrimitiveUniformShaderParameters.ActorWorldPosition, PrimitiveUniformShaderParameters.UseSingleSampleShadowFromStationaryLights);
-	Data[19] = FVector4(PrimitiveUniformShaderParameters.ObjectBounds, PrimitiveUniformShaderParameters.LpvBiasMultiplier);
+	Data[19] = FVector4(PrimitiveUniformShaderParameters.ObjectBounds, 0.0f);
 
 	Data[20] = FVector4(
 		PrimitiveUniformShaderParameters.DecalReceiverMask, 
diff --git a/Engine/Source/Runtime/Engine/Private/RendererSettings.cpp b/Engine/Source/Runtime/Engine/Private/RendererSettings.cpp
--- a/Engine/Source/Runtime/Engine/Private/RendererSettings.cpp
+++ b/Engine/Source/Runtime/Engine/Private/RendererSettings.cpp
@@ -73,7 +73,6 @@ URendererSettings::URendererSettings(const FObjectInitializer& ObjectInitializer
 	GPUSimulationTextureSizeY = 1024;
 	bEnableRayTracing = 0;
 	bEnableRayTracingTextureLOD = 0;
-	bLPV = 1;
 }
 
 void URendererSettings::PostInitProperties()
diff --git a/Engine/Source/Runtime/Engine/Private/Scene.cpp b/Engine/Source/Runtime/Engine/Private/Scene.cpp
--- a/Engine/Source/Runtime/Engine/Private/Scene.cpp
+++ b/Engine/Source/Runtime/Engine/Private/Scene.cpp
@@ -451,13 +451,6 @@ FPostProcessSettings::FPostProcessSettings()
 	BloomDirtMaskTint = FLinearColor(0.5f, 0.5f, 0.5f);
 	AmbientCubemapIntensity = 1.0f;
 	AmbientCubemapTint = FLinearColor(1, 1, 1);
-	LPVIntensity = 1.0f;
-	LPVSize = 5312.0f;
-	LPVSecondaryOcclusionIntensity = 0.0f;
-	LPVSecondaryBounceIntensity = 0.0f;
-	LPVVplInjectionBias = 0.64f;
-	LPVGeometryVolumeBias = 0.384f;
-	LPVEmissiveInjectionIntensity = 1.0f;
 	CameraShutterSpeed = 60.f;
 	CameraISO = 100.f;
 	AutoExposureCalibrationConstant_DEPRECATED = 16.f;
@@ -493,15 +486,6 @@ FPostProcessSettings::FPostProcessSettings()
 
 	AutoExposureApplyPhysicalCameraExposure = 1;
 
-	LPVDirectionalOcclusionIntensity = 0.0f;
-	LPVDirectionalOcclusionRadius = 8.0f;
-	LPVDiffuseOcclusionExponent = 1.0f;
-	LPVSpecularOcclusionExponent = 7.0f;
-	LPVDiffuseOcclusionIntensity = 1.0f;
-	LPVSpecularOcclusionIntensity = 1.0f;
-	LPVFadeRange = 0.0f;
-	LPVDirectionalOcclusionFadeRange = 0.0f;
-
 	// next value might get overwritten by r.DefaultFeature.LensFlare
 	LensFlareIntensity = 1.0f;
 	LensFlareTint = FLinearColor(1.0f, 1.0f, 1.0f);
@@ -720,21 +704,6 @@ FPostProcessSettings::FPostProcessSettings(const FPostProcessSettings& Settings)
 	, bOverride_RayTracingAOSamplesPerPixel(Settings.bOverride_RayTracingAOSamplesPerPixel)
 	, bOverride_RayTracingAOIntensity(Settings.bOverride_RayTracingAOIntensity)
 	, bOverride_RayTracingAORadius(Settings.bOverride_RayTracingAORadius)
-	, bOverride_LPVIntensity(Settings.bOverride_LPVIntensity)
-	, bOverride_LPVDirectionalOcclusionIntensity(Settings.bOverride_LPVDirectionalOcclusionIntensity)
-	, bOverride_LPVDirectionalOcclusionRadius(Settings.bOverride_LPVDirectionalOcclusionRadius)
-	, bOverride_LPVDiffuseOcclusionExponent(Settings.bOverride_WhiteTemp)
-	, bOverride_LPVSpecularOcclusionExponent(Settings.bOverride_LPVSpecularOcclusionExponent)
-	, bOverride_LPVDiffuseOcclusionIntensity(Settings.bOverride_LPVDiffuseOcclusionIntensity)
-	, bOverride_LPVSpecularOcclusionIntensity(Settings.bOverride_LPVSpecularOcclusionIntensity)
-	, bOverride_LPVSize(Settings.bOverride_LPVSize)
-	, bOverride_LPVSecondaryOcclusionIntensity(Settings.bOverride_LPVSecondaryOcclusionIntensity)
-	, bOverride_LPVSecondaryBounceIntensity(Settings.bOverride_LPVSecondaryBounceIntensity)
-	, bOverride_LPVGeometryVolumeBias(Settings.bOverride_LPVGeometryVolumeBias)
-	, bOverride_LPVVplInjectionBias(Settings.bOverride_LPVVplInjectionBias)
-	, bOverride_LPVEmissiveInjectionIntensity(Settings.bOverride_LPVEmissiveInjectionIntensity)
-	, bOverride_LPVFadeRange(Settings.bOverride_LPVFadeRange)
-	, bOverride_LPVDirectionalOcclusionFadeRange(Settings.bOverride_LPVDirectionalOcclusionFadeRange)
 	, bOverride_IndirectLightingColor(Settings.bOverride_IndirectLightingColor)
 	, bOverride_IndirectLightingIntensity(Settings.bOverride_IndirectLightingIntensity)
 	, bOverride_ColorGradingIntensity(Settings.bOverride_ColorGradingIntensity)
@@ -929,19 +898,6 @@ FPostProcessSettings::FPostProcessSettings(const FPostProcessSettings& Settings)
 	, MotionBlurMax(Settings.MotionBlurMax)
 	, MotionBlurTargetFPS(Settings.MotionBlurTargetFPS)
 	, MotionBlurPerObjectSize(Settings.MotionBlurPerObjectSize)
-	, LPVIntensity(Settings.LPVIntensity)
-	, LPVVplInjectionBias(Settings.LPVVplInjectionBias)
-	, LPVSize(Settings.LPVSize)
-	, LPVSecondaryOcclusionIntensity(Settings.LPVSecondaryOcclusionIntensity)
-	, LPVSecondaryBounceIntensity(Settings.LPVSecondaryBounceIntensity)
-	, LPVGeometryVolumeBias(Settings.LPVGeometryVolumeBias)
-	, LPVEmissiveInjectionIntensity(Settings.LPVEmissiveInjectionIntensity)
-	, LPVDirectionalOcclusionIntensity(Settings.LPVDirectionalOcclusionIntensity)
-	, LPVDirectionalOcclusionRadius(Settings.LPVDirectionalOcclusionRadius)
-	, LPVDiffuseOcclusionExponent(Settings.WhiteTemp)
-	, LPVSpecularOcclusionExponent(Settings.LPVSpecularOcclusionExponent)
-	, LPVDiffuseOcclusionIntensity(Settings.LPVDiffuseOcclusionIntensity)
-	, LPVSpecularOcclusionIntensity(Settings.LPVSpecularOcclusionIntensity)
 
 	, ReflectionsType(Settings.ReflectionsType)
 	, ScreenSpaceReflectionIntensity(Settings.ScreenSpaceReflectionIntensity)
@@ -964,9 +920,6 @@ FPostProcessSettings::FPostProcessSettings(const FPostProcessSettings& Settings)
 	, PathTracingMaxBounces(Settings.PathTracingMaxBounces)
 	, PathTracingSamplesPerPixel(Settings.PathTracingSamplesPerPixel)
 
-	, LPVFadeRange(Settings.LPVFadeRange)
-	, LPVDirectionalOcclusionFadeRange(Settings.LPVDirectionalOcclusionFadeRange)
-
 	, ScreenPercentage(Settings.ScreenPercentage)
 
 	, WeightedBlendables(Settings.WeightedBlendables)
diff --git a/Engine/Source/Runtime/Engine/Private/SceneManagement.cpp b/Engine/Source/Runtime/Engine/Private/SceneManagement.cpp
--- a/Engine/Source/Runtime/Engine/Private/SceneManagement.cpp
+++ b/Engine/Source/Runtime/Engine/Private/SceneManagement.cpp
@@ -318,7 +318,6 @@ void FDynamicPrimitiveUniformBuffer::Set(
 			bHasPrecomputedVolumetricLightmap,
 			bDrawsVelocity,
 			GetDefaultLightingChannelMask(),
-			1.0f,
 			INDEX_NONE,
 			INDEX_NONE,
 			bOutputVelocity,
diff --git a/Engine/Source/Runtime/Engine/Private/SceneView.cpp b/Engine/Source/Runtime/Engine/Private/SceneView.cpp
--- a/Engine/Source/Runtime/Engine/Private/SceneView.cpp
+++ b/Engine/Source/Runtime/Engine/Private/SceneView.cpp
@@ -22,7 +22,6 @@
 #include "IHeadMountedDisplay.h"
 #include "IXRTrackingSystem.h"
 #include "Engine/RendererSettings.h"
-#include "LightPropagationVolumeSettings.h"
 #include "Materials/MaterialInstanceDynamic.h"
 #include "HighResScreenshot.h"
 #include "Slate/SceneViewport.h"
@@ -1659,31 +1658,6 @@ void FSceneView::OverridePostProcessSettings(const FPostProcessSettings& Src, fl
 		}
 	}
 
-	// will be deprecated soon, use the new asset LightPropagationVolumeBlendable instead
-	{
-		FLightPropagationVolumeSettings& Dest = FinalPostProcessSettings.BlendableManager.GetSingleFinalData<FLightPropagationVolumeSettings>();
-
-		LERP_PP(LPVIntensity);
-		LERP_PP(LPVSecondaryOcclusionIntensity);
-		LERP_PP(LPVSecondaryBounceIntensity);
-		LERP_PP(LPVVplInjectionBias);
-		LERP_PP(LPVGeometryVolumeBias);
-		LERP_PP(LPVEmissiveInjectionIntensity);
-		LERP_PP(LPVDirectionalOcclusionIntensity);
-		LERP_PP(LPVDirectionalOcclusionRadius);
-		LERP_PP(LPVDiffuseOcclusionExponent);
-		LERP_PP(LPVSpecularOcclusionExponent);
-		LERP_PP(LPVDiffuseOcclusionIntensity);
-		LERP_PP(LPVSpecularOcclusionIntensity);
-		LERP_PP(LPVFadeRange);
-		LERP_PP(LPVDirectionalOcclusionFadeRange);
-
-		if (Src.bOverride_LPVSize)
-		{
-			Dest.LPVSize = Src.LPVSize;
-		}
-	}
-
 	// Blendable objects
 	{
 		uint32 Count = Src.WeightedBlendables.Array.Num();
@@ -1796,26 +1770,6 @@ void FSceneView::EndFinalPostprocessSettings(const FSceneViewInitOptions& ViewIn
 {
 	const auto SceneViewFeatureLevel = GetFeatureLevel();
 
-	// will be deprecated soon, use the new asset LightPropagationVolumeBlendable instead
-	{
-		FLightPropagationVolumeSettings& Dest = FinalPostProcessSettings.BlendableManager.GetSingleFinalData<FLightPropagationVolumeSettings>();
-
-		if(Dest.LPVDirectionalOcclusionIntensity < 0.001f)
-		{
-			Dest.LPVDirectionalOcclusionIntensity = 0.0f;
-		}
-
-		if (Dest.LPVIntensity < 0.001f)
-		{
-			Dest.LPVIntensity = 0.0f;
-		}
-
-		if(!Family->EngineShowFlags.GlobalIllumination)
-		{
-			Dest.LPVIntensity = 0.0f;
-		}
-	}
-
 	{
 		static const auto SceneColorFringeQualityCVar = IConsoleManager::Get().FindTConsoleVariableDataInt(TEXT("r.SceneColorFringeQuality"));
 
diff --git a/Engine/Source/Runtime/Engine/Private/StaticMesh.cpp b/Engine/Source/Runtime/Engine/Private/StaticMesh.cpp
--- a/Engine/Source/Runtime/Engine/Private/StaticMesh.cpp
+++ b/Engine/Source/Runtime/Engine/Private/StaticMesh.cpp
@@ -2778,7 +2778,6 @@ UStaticMesh::UStaticMesh(const FObjectInitializer& ObjectInitializer)
 	NumStreamedLODs.Default = -1;
 #endif // #if WITH_EDITORONLY_DATA
 	LightMapResolution = 4;
-	LpvBiasMultiplier = 1.0f;
 	MinLOD.Default = 0;
 
 	bSupportUniformlyDistributedSampling = false;
diff --git a/Engine/Source/Runtime/Engine/Private/StaticMeshRender.cpp b/Engine/Source/Runtime/Engine/Private/StaticMeshRender.cpp
--- a/Engine/Source/Runtime/Engine/Private/StaticMeshRender.cpp
+++ b/Engine/Source/Runtime/Engine/Private/StaticMeshRender.cpp
@@ -287,8 +287,6 @@ FStaticMeshSceneProxy::FStaticMeshSceneProxy(UStaticMeshComponent* InComponent,
 	// We always use local vertex factory, which gets its primitive data from GPUScene, so we can skip expensive primitive uniform buffer updates
 	bVFRequiresPrimitiveUniformBuffer = !UseGPUScene(GMaxRHIShaderPlatform, FeatureLevel);
 
-	LpvBiasMultiplier = FMath::Min( InComponent->GetStaticMesh()->LpvBiasMultiplier * InComponent->LpvBiasMultiplier, 3.0f );
-
 #if STATICMESH_ENABLE_DEBUG_RENDERING
 	if( GIsEditor )
 	{
@@ -1057,29 +1055,6 @@ HHitProxy* FStaticMeshSceneProxy::CreateHitProxies(UPrimitiveComponent* Componen
 }
 #endif // WITH_EDITOR
 
-// use for render thread only
-bool UseLightPropagationVolumeRT2(ERHIFeatureLevel::Type InFeatureLevel)
-{
-	if (InFeatureLevel < ERHIFeatureLevel::SM5)
-	{
-		return false;
-	}
-
-	// todo: better we get the engine LPV state not the cvar (later we want to make it changeable at runtime)
-	static const auto* CVar = IConsoleManager::Get().FindTConsoleVariableDataInt(TEXT("r.LightPropagationVolume"));
-	check(CVar);
-
-	int32 Value = CVar->GetValueOnRenderThread();
-
-	return Value != 0;
-}
-
-inline bool AllowShadowOnlyMesh(ERHIFeatureLevel::Type InFeatureLevel)
-{
-	// todo: later we should refine that (only if occlusion feature in LPV is on, only if inside a cascade, if shadow casting is disabled it should look at bUseEmissiveForDynamicAreaLighting)
-	return !UseLightPropagationVolumeRT2(InFeatureLevel);
-}
-
 inline void SetupMeshBatchForRuntimeVirtualTexture(FMeshBatch& MeshBatch)
 {
 	MeshBatch.CastShadow = 0;
@@ -1164,8 +1139,7 @@ void FStaticMeshSceneProxy::DrawStaticElements(FStaticPrimitiveDrawInterface* PD
 					const FLODInfo& ProxyLODInfo = LODs[LODIndex];
 
 					// The shadow-only mesh can be used only if all elements cast shadows and use opaque materials with no vertex modification.
-					// In some cases (e.g. LPV) we don't want the optimization
-					bool bSafeToUseUnifiedMesh = AllowShadowOnlyMesh(FeatureLevel);
+					bool bSafeToUseUnifiedMesh = true;
 
 					bool bAnySectionUsesDitheredLODTransition = false;
 					bool bAllSectionsUseDitheredLODTransition = true;
diff --git a/Engine/Source/Runtime/Engine/Public/MaterialShared.h b/Engine/Source/Runtime/Engine/Public/MaterialShared.h
--- a/Engine/Source/Runtime/Engine/Public/MaterialShared.h
+++ b/Engine/Source/Runtime/Engine/Public/MaterialShared.h
@@ -1541,8 +1541,6 @@ public:
 	virtual bool IsTranslucencyWritingCustomDepth() const { return false; }
 	virtual bool IsTranslucencyWritingVelocity() const { return false; }
 	virtual bool IsTangentSpaceNormal() const { return false; }
-	virtual bool ShouldInjectEmissiveIntoLPV() const { return false; }
-	virtual bool ShouldBlockGI() const { return false; }
 	virtual bool ShouldGenerateSphericalParticleNormals() const { return false; }
 	virtual	bool ShouldDisableDepthTest() const { return false; }
 	virtual	bool ShouldWriteOnlyAlpha() const { return false; }
@@ -2381,8 +2379,6 @@ public:
 	ENGINE_API virtual bool IsTranslucencyWritingCustomDepth() const override;
 	ENGINE_API virtual bool IsTranslucencyWritingVelocity() const override;
 	ENGINE_API virtual bool IsTangentSpaceNormal() const override;
-	ENGINE_API virtual bool ShouldInjectEmissiveIntoLPV() const override;
-	ENGINE_API virtual bool ShouldBlockGI() const override;
 	ENGINE_API virtual bool ShouldGenerateSphericalParticleNormals() const override;
 	ENGINE_API virtual bool ShouldDisableDepthTest() const override;
 	ENGINE_API virtual bool ShouldWriteOnlyAlpha() const override;
@@ -3015,8 +3011,6 @@ struct FMaterialShaderParameters
 			uint64 bIsTwoSided : 1;
 			uint64 bIsDistorted : 1;
 			uint64 bShouldCastDynamicShadows : 1;
-			uint64 bShouldInjectEmissiveIntoLPV : 1;
-			uint64 bShouldBlockGI : 1;
 			uint64 bWritesEveryPixel : 1;
 			uint64 bWritesEveryPixelShadowPass : 1;
 			uint64 bHasNormalConnected : 1;
@@ -3080,8 +3074,6 @@ struct FMaterialShaderParameters
 		bIsTwoSided = InMaterial->IsTwoSided();
 		bIsDistorted = InMaterial->IsDistorted();
 		bShouldCastDynamicShadows = InMaterial->ShouldCastDynamicShadows();
-		bShouldInjectEmissiveIntoLPV = InMaterial->ShouldInjectEmissiveIntoLPV();
-		bShouldBlockGI = InMaterial->ShouldBlockGI();
 		bWritesEveryPixel = InMaterial->WritesEveryPixel(false);
 		bWritesEveryPixelShadowPass = InMaterial->WritesEveryPixel(true);
 		bHasNormalConnected = InMaterial->HasNormalConnected();
diff --git a/Engine/Source/Runtime/Engine/Public/PrimitiveSceneProxy.h b/Engine/Source/Runtime/Engine/Public/PrimitiveSceneProxy.h
--- a/Engine/Source/Runtime/Engine/Public/PrimitiveSceneProxy.h
+++ b/Engine/Source/Runtime/Engine/Public/PrimitiveSceneProxy.h
@@ -526,7 +526,6 @@ public:
 	inline bool WritesVirtualTexture(URuntimeVirtualTexture* VirtualTexture) const { return RuntimeVirtualTextures.Find(VirtualTexture) != INDEX_NONE; }
 	inline bool AffectsDynamicIndirectLighting() const { return bAffectDynamicIndirectLighting; }
 	inline bool AffectsDistanceFieldLighting() const { return bAffectDistanceFieldLighting; }
-	inline float GetLpvBiasMultiplier() const { return LpvBiasMultiplier; }
 	inline EIndirectLightingCacheQuality GetIndirectLightingCacheQuality() const { return IndirectLightingCacheQuality; }
 	inline bool CastsVolumetricTranslucentShadow() const { return bCastVolumetricTranslucentShadow; }
 	inline bool CastsContactShadow() const { return bCastContactShadow; }
@@ -1012,9 +1011,6 @@ protected:
 	/** Log2 of minimum estimated pixel coverage before culling from runtime virtual texture. */
 	int8 VirtualTextureMinCoverage;
 
-	/** The bias applied to LPV injection */
-	float LpvBiasMultiplier;
-
 	/** Min visibility for capsule shadows. */
 	float DynamicIndirectShadowMinVisibility;
 
diff --git a/Engine/Source/Runtime/Engine/Public/PrimitiveUniformShaderParameters.h b/Engine/Source/Runtime/Engine/Public/PrimitiveUniformShaderParameters.h
--- a/Engine/Source/Runtime/Engine/Public/PrimitiveUniformShaderParameters.h
+++ b/Engine/Source/Runtime/Engine/Public/PrimitiveUniformShaderParameters.h
@@ -29,7 +29,6 @@ BEGIN_GLOBAL_SHADER_PARAMETER_STRUCT(FPrimitiveUniformShaderParameters,ENGINE_AP
 	SHADER_PARAMETER_EX(float,UseSingleSampleShadowFromStationaryLights,EShaderPrecisionModifier::Half)	
 
 	SHADER_PARAMETER(FVector,ObjectBounds)		// only needed for editor/development
-	SHADER_PARAMETER(float,LpvBiasMultiplier)
 
 	SHADER_PARAMETER_EX(float,DecalReceiverMask,EShaderPrecisionModifier::Half)
 	SHADER_PARAMETER_EX(float,PerObjectGBufferData,EShaderPrecisionModifier::Half)		// 0..1, 2 bits, bCastContactShadow, bHeightfieldRepresentation
@@ -65,7 +64,6 @@ inline FPrimitiveUniformShaderParameters GetPrimitiveUniformShaderParameters(
 	bool bUseVolumetricLightmap,
 	bool bDrawsVelocity,
 	uint32 LightingChannelMask,
-	float LpvBiasMultiplier,
 	uint32 LightmapDataIndex,
 	int32 SingleCaptureIndex,
 	bool bOutputVelocity,
@@ -87,7 +85,6 @@ inline FPrimitiveUniformShaderParameters GetPrimitiveUniformShaderParameters(
 	Result.ObjectOrientation = LocalToWorld.GetUnitAxis( EAxis::Z );
 	Result.ActorWorldPosition = ActorPosition;
 	Result.LightingChannelMask = LightingChannelMask;
-	Result.LpvBiasMultiplier = LpvBiasMultiplier;
 
 	{
 		// Extract per axis scales from LocalToWorld transform
@@ -142,7 +139,6 @@ inline FPrimitiveUniformShaderParameters GetPrimitiveUniformShaderParameters(
 	bool bUseVolumetricLightmap,
 	bool bDrawsVelocity,
 	uint32 LightingChannelMask,
-	float LpvBiasMultiplier,
 	uint32 LightmapDataIndex,
 	int32 SingleCaptureIndex,
     bool bOutputVelocity,
@@ -163,8 +159,7 @@ inline FPrimitiveUniformShaderParameters GetPrimitiveUniformShaderParameters(
 		bUseSingleSampleShadowFromStationaryLights, 
 		bUseVolumetricLightmap, 
 		bDrawsVelocity, 
-		LightingChannelMask, 
-		LpvBiasMultiplier, 
+		LightingChannelMask,
 		LightmapDataIndex, 
 		SingleCaptureIndex,
 		bOutputVelocity,
@@ -179,8 +174,7 @@ inline TUniformBufferRef<FPrimitiveUniformShaderParameters> CreatePrimitiveUnifo
 	const FBoxSphereBounds& LocalBounds,
 	const FBoxSphereBounds& PreSkinnedLocalBounds,
 	bool bReceivesDecals,
-	bool bDrawsVelocity,
-	float LpvBiasMultiplier = 1.0f
+	bool bDrawsVelocity
 	)
 {
 	check(IsInRenderingThread());
@@ -199,7 +193,6 @@ inline TUniformBufferRef<FPrimitiveUniformShaderParameters> CreatePrimitiveUnifo
 			false,
 			bDrawsVelocity,
 			GetDefaultLightingChannelMask(),
-			LpvBiasMultiplier,
 			INDEX_NONE,
 			INDEX_NONE,
 			false,
@@ -226,7 +219,6 @@ inline FPrimitiveUniformShaderParameters GetIdentityPrimitiveParameters()
 		false,
 		/* bDrawsVelocity = */ true,
 		GetDefaultLightingChannelMask(),
-		1.0f,		// LPV bias
 		INDEX_NONE,
 		INDEX_NONE,
 		false,
diff --git a/Engine/Source/Runtime/Engine/Public/SceneManagement.h b/Engine/Source/Runtime/Engine/Public/SceneManagement.h
--- a/Engine/Source/Runtime/Engine/Public/SceneManagement.h
+++ b/Engine/Source/Runtime/Engine/Public/SceneManagement.h
@@ -1371,18 +1371,6 @@ public:
 		return false;
 	}
 
-	/**
-	 * Sets up a projected shadow initializer for a reflective shadow map that's dependent on the current view for shadows from the entire scene.
-	 * @return True if the whole-scene projected shadow should be used.
-	 */
-	virtual bool GetViewDependentRsmWholeSceneProjectedShadowInitializer(
-		const class FSceneView& View, 
-		const FBox& LightPropagationVolumeBounds,
-		class FWholeSceneProjectedShadowInitializer& OutInitializer ) const
-	{
-		return false;
-	}
-
 	/**
 	 * Sets up a projected shadow initializer for the given subject.
 	 * @param SubjectBounds - The bounding volume of the subject.
diff --git a/Engine/Source/Runtime/Engine/Public/ShowFlags.h b/Engine/Source/Runtime/Engine/Public/ShowFlags.h
--- a/Engine/Source/Runtime/Engine/Public/ShowFlags.h
+++ b/Engine/Source/Runtime/Engine/Public/ShowFlags.h
@@ -382,7 +382,6 @@ private:
 		SetLightMapDensity(false);
 		SetLODColoration(false);
 		SetHLODColoration(false);
-		SetVisualizeLPV(false);
 		SetStreamingBounds(false);
 		SetHISMCOcclusionBounds(false);
 		SetHISMCClusterTree(false);
diff --git a/Engine/Source/Runtime/Engine/Public/ShowFlagsValues.inl b/Engine/Source/Runtime/Engine/Public/ShowFlagsValues.inl
--- a/Engine/Source/Runtime/Engine/Public/ShowFlagsValues.inl
+++ b/Engine/Source/Runtime/Engine/Public/ShowFlagsValues.inl
@@ -287,8 +287,6 @@ SHOWFLAG_FIXED_IN_SHIPPING(0, VisualizeLightCulling, SFG_Hidden, NSLOCTEXT("Unre
 SHOWFLAG_FIXED_IN_SHIPPING(1, PrecomputedVisibility, SFG_Advanced, NSLOCTEXT("UnrealEd", "PrecomputedVisibilitySF", "Precomputed Visibility"))
 /** Contribution from sky light, for now SHOWFLAG_ALWAYS_ACCESSIBLE because it's exposed in SceneCapture */
 SHOWFLAG_ALWAYS_ACCESSIBLE(SkyLighting, SFG_LightTypes, NSLOCTEXT("UnrealEd", "SkyLightingSF", "Sky Lighting"))
-/** Visualize Light Propagation Volume, for developer (by default off): */
-SHOWFLAG_FIXED_IN_SHIPPING(0, VisualizeLPV, SFG_Visualize, NSLOCTEXT("UnrealEd", "VisualizeLPVSF", "Light Propagation Volume"))
 /** Visualize preview shadow indicator */
 SHOWFLAG_FIXED_IN_SHIPPING(0, PreviewShadowsIndicator, SFG_Visualize, NSLOCTEXT("UnrealEd", "PreviewShadowIndicatorSF", "Preview Shadows Indicator"))
 /** Visualize precomputed visibility cells */
diff --git a/Engine/Source/Runtime/Landscape/Private/Landscape.cpp b/Engine/Source/Runtime/Landscape/Private/Landscape.cpp
--- a/Engine/Source/Runtime/Landscape/Private/Landscape.cpp
+++ b/Engine/Source/Runtime/Landscape/Private/Landscape.cpp
@@ -188,8 +188,6 @@ ULandscapeComponent::ULandscapeComponent(const FObjectInitializer& ObjectInitial
 	EditToolRenderData = FLandscapeEditToolRenderData();
 #endif
 
-	LpvBiasMultiplier = 0.0f; // Bias is 0 for landscape, since it's single sided
-
 	// We don't want to load this on the server, this component is for graphical purposes only
 	AlwaysLoadOnServer = false;
 
diff --git a/Engine/Source/Runtime/Landscape/Private/LandscapeRender.cpp b/Engine/Source/Runtime/Landscape/Private/LandscapeRender.cpp
--- a/Engine/Source/Runtime/Landscape/Private/LandscapeRender.cpp
+++ b/Engine/Source/Runtime/Landscape/Private/LandscapeRender.cpp
@@ -3787,10 +3787,8 @@ public:
 			FName(TEXT("TBasePassPSFNoLightMapPolicy")),
 			FName(TEXT("TBasePassVSFCachedPointIndirectLightingPolicy")),
 			FName(TEXT("TBasePassPSFCachedPointIndirectLightingPolicy")),
-			FName(TEXT("TShadowDepthVSVertexShadowDepth_OutputDepthfalse")),
-			FName(TEXT("TShadowDepthVSVertexShadowDepth_OutputDepthtrue")), // used by LPV
-			FName(TEXT("TShadowDepthPSPixelShadowDepth_NonPerspectiveCorrectfalse")),
-			FName(TEXT("TShadowDepthPSPixelShadowDepth_NonPerspectiveCorrecttrue")), // used by LPV
+			FName(TEXT("TShadowDepthVSVertexShadowDepth_OutputDepth")),
+			FName(TEXT("TShadowDepthPSPixelShadowDepth_NonPerspectiveCorrect")),
 			FName(TEXT("TBasePassPSFSimpleDirectionalLightLightingPolicy")),
 			FName(TEXT("TBasePassPSFSimpleDirectionalLightLightingPolicySkylight")),
 			FName(TEXT("TBasePassVSFSimpleDirectionalLightLightingPolicy")),
diff --git a/Engine/Source/Runtime/RenderCore/Private/RenderCore.cpp b/Engine/Source/Runtime/RenderCore/Private/RenderCore.cpp
--- a/Engine/Source/Runtime/RenderCore/Private/RenderCore.cpp
+++ b/Engine/Source/Runtime/RenderCore/Private/RenderCore.cpp
@@ -189,9 +189,6 @@ DEFINE_STAT(STAT_Scene_UpdateStaticDrawLists_RT);
 DEFINE_STAT(STAT_Scene_UpdateStaticDrawListsForMaterials_RT);
 DEFINE_STAT(STAT_GameToRendererMallocTotal);
 
-DEFINE_STAT(STAT_UpdateLPVs);
-DEFINE_STAT(STAT_ReflectiveShadowMaps);
-DEFINE_STAT(STAT_ReflectiveShadowMapDrawTime);
 DEFINE_STAT(STAT_NumReflectiveShadowMapLights);
 
 DEFINE_STAT(STAT_ShadowmapAtlasMemory);
diff --git a/Engine/Source/Runtime/RenderCore/Public/RenderCore.h b/Engine/Source/Runtime/RenderCore/Public/RenderCore.h
--- a/Engine/Source/Runtime/RenderCore/Public/RenderCore.h
+++ b/Engine/Source/Runtime/RenderCore/Public/RenderCore.h
@@ -32,7 +32,6 @@ DECLARE_CYCLE_STAT_EXTERN(TEXT("Water pass drawing"), STAT_WaterPassDrawTime, ST
 DECLARE_CYCLE_STAT_EXTERN(TEXT("Lighting drawing"),STAT_LightingDrawTime,STATGROUP_SceneRendering, RENDERCORE_API);
 DECLARE_CYCLE_STAT_EXTERN(TEXT("Proj Shadow drawing"),STAT_ProjectedShadowDrawTime,STATGROUP_SceneRendering, RENDERCORE_API);
 DECLARE_CYCLE_STAT_EXTERN(TEXT("Reflective Shadow map drawing"),STAT_ReflectiveShadowMapDrawTime,STATGROUP_SceneRendering, RENDERCORE_API);
-DECLARE_CYCLE_STAT_EXTERN(TEXT("Update LPVs"),STAT_UpdateLPVs,STATGROUP_SceneRendering, RENDERCORE_API);
 DECLARE_CYCLE_STAT_EXTERN(TEXT("Translucency drawing"),STAT_TranslucencyDrawTime,STATGROUP_SceneRendering, RENDERCORE_API);
 DECLARE_CYCLE_STAT_EXTERN(TEXT("Dynamic Primitive drawing"),STAT_DynamicPrimitiveDrawTime,STATGROUP_SceneRendering, RENDERCORE_API);
 DECLARE_CYCLE_STAT_EXTERN(TEXT("StaticDrawList drawing"),STAT_StaticDrawListDrawTime,STATGROUP_SceneRendering, RENDERCORE_API);
diff --git a/Engine/Source/Runtime/Renderer/Private/CompositionLighting/CompositionLighting.cpp b/Engine/Source/Runtime/Renderer/Private/CompositionLighting/CompositionLighting.cpp
--- a/Engine/Source/Runtime/Renderer/Private/CompositionLighting/CompositionLighting.cpp
+++ b/Engine/Source/Runtime/Renderer/Private/CompositionLighting/CompositionLighting.cpp
@@ -9,11 +9,9 @@
 #include "PostProcess/RenderingCompositionGraph.h"
 #include "PostProcess/PostProcessInput.h"
 #include "PostProcess/PostProcessing.h"
-#include "CompositionLighting/PostProcessLpvIndirect.h"
 #include "CompositionLighting/PostProcessAmbientOcclusion.h"
 #include "CompositionLighting/PostProcessDeferredDecals.h"
 #include "PostProcess/PostProcessSubsurface.h"
-#include "LightPropagationVolumeSettings.h"
 #include "DecalRenderingShared.h"
 #include "VisualizeTexture.h"
 #include "RayTracing/RaytracingOptions.h"
@@ -25,7 +23,6 @@ FCompositionLighting GCompositionLighting;
 
 DECLARE_GPU_STAT_NAMED(CompositionBeforeBasePass, TEXT("Composition BeforeBasePass") );
 DECLARE_GPU_STAT_NAMED(CompositionPreLighting, TEXT("Composition PreLighting") );
-DECLARE_GPU_STAT_NAMED(CompositionLpvIndirect, TEXT("Composition LpvIndirect") );
 DECLARE_GPU_STAT_NAMED(CompositionPostLighting, TEXT("Composition PostLighting") );
 
 static TAutoConsoleVariable<int32> CVarSSAOSmoothPass(
@@ -73,33 +70,6 @@ bool IsAmbientCubemapPassRequired(const FSceneView& View)
 	return View.FinalPostProcessSettings.ContributingCubemaps.Num() != 0 && IsUsingGBuffers(View.GetShaderPlatform());
 }
 
-bool IsLpvIndirectPassRequired(const FViewInfo& View)
-{
-	FScene* Scene = (FScene*)View.Family->Scene;
-
-	const FSceneViewState* ViewState = (FSceneViewState*)View.State;
-
-	if(ViewState)
-	{
-		// This check should be inclusive to stereo views
-		const bool bIncludeStereoViews = true;
-
-		FLightPropagationVolume* LightPropagationVolume = ViewState->GetLightPropagationVolume(View.GetFeatureLevel(), bIncludeStereoViews);
-
-		if(LightPropagationVolume)
-		{
-			const FLightPropagationVolumeSettings& LPVSettings = View.FinalPostProcessSettings.BlendableManager.GetSingleFinalDataConst<FLightPropagationVolumeSettings>();
-
-			if(LPVSettings.LPVIntensity > 0.0f)
-			{
-				return true;
-			}
-		}
-	}
-
-	return false;
-}
-
 static bool IsReflectionEnvironmentActive(const FSceneView& View)
 {
 	FScene* Scene = (FScene*)View.Family->Scene;
@@ -125,15 +95,12 @@ bool ShouldRenderScreenSpaceAmbientOcclusion(const FViewInfo& View)
 {
 	bool bEnabled = true;
 
-	if (!IsLpvIndirectPassRequired(View))
-	{
-		bEnabled = View.FinalPostProcessSettings.AmbientOcclusionIntensity > 0
-			&& View.Family->EngineShowFlags.Lighting
-			&& View.FinalPostProcessSettings.AmbientOcclusionRadius >= 0.1f
-			&& !View.Family->UseDebugViewPS()
-			&& (FSSAOHelper::IsBasePassAmbientOcclusionRequired(View) || IsAmbientCubemapPassRequired(View) || IsReflectionEnvironmentActive(View) || IsSkylightActive(View) || View.Family->EngineShowFlags.VisualizeBuffer)
-			&& !IsSimpleForwardShadingEnabled(View.GetShaderPlatform());
-	}
+	bEnabled = View.FinalPostProcessSettings.AmbientOcclusionIntensity > 0
+		&& View.Family->EngineShowFlags.Lighting
+		&& View.FinalPostProcessSettings.AmbientOcclusionRadius >= 0.1f
+		&& !View.Family->UseDebugViewPS()
+		&& (FSSAOHelper::IsBasePassAmbientOcclusionRequired(View) || IsAmbientCubemapPassRequired(View) || IsReflectionEnvironmentActive(View) || IsSkylightActive(View) || View.Family->EngineShowFlags.VisualizeBuffer)
+		&& !IsSimpleForwardShadingEnabled(View.GetShaderPlatform());
 #if RHI_RAYTRACING
 	bEnabled &= !ShouldRenderRayTracingAmbientOcclusion(View);
 #endif
@@ -676,37 +643,6 @@ void FCompositionLighting::ProcessAfterBasePass(
 	}
 }
 
-
-void FCompositionLighting::ProcessLpvIndirect(FRHICommandListImmediate& RHICmdList, FViewInfo& View)
-{
-	check(IsInRenderingThread());
-	
-	FMemMark Mark(FMemStack::Get());
-	FRenderingCompositePassContext CompositeContext(RHICmdList, View);
-	FPostprocessContext Context(RHICmdList, CompositeContext.Graph, View);
-
-	{
-		FSceneRenderTargets& SceneContext = FSceneRenderTargets::Get(RHICmdList);
-
-		FRenderingCompositePass* SSAO = Context.Graph.RegisterPass(new FRCPassPostProcessInput(SceneContext.ScreenSpaceAO));
-
-		FRenderingCompositePass* Pass = Context.Graph.RegisterPass(new FRCPassPostProcessLpvIndirect());
-		Pass->SetInput(ePId_Input0, Context.FinalOutput);
-		Pass->SetInput(ePId_Input1, SSAO );
-
-		Context.FinalOutput = FRenderingCompositeOutputRef(Pass);
-	}
-
-	// The graph setup should be finished before this line ----------------------------------------
-
-	SCOPED_DRAW_EVENT(RHICmdList, CompositionLpvIndirect);
-	SCOPED_GPU_STAT(RHICmdList, CompositionLpvIndirect);
-
-	// we don't replace the final element with the scenecolor because this is what those passes should do by themself
-
-	CompositeContext.Process(Context.FinalOutput.GetPass(), TEXT("CompositionLighting"));
-}
-
 bool FCompositionLighting::CanProcessAsyncSSAO(const TArray<FViewInfo>& Views)
 {
 	bool bAnyAsyncSSAO = true;
diff --git a/Engine/Source/Runtime/Renderer/Private/CompositionLighting/CompositionLighting.h b/Engine/Source/Runtime/Renderer/Private/CompositionLighting/CompositionLighting.h
--- a/Engine/Source/Runtime/Renderer/Private/CompositionLighting/CompositionLighting.h
+++ b/Engine/Source/Runtime/Renderer/Private/CompositionLighting/CompositionLighting.h
@@ -32,9 +32,6 @@ public:
 		TRDGUniformBufferRef<FSceneTextureUniformParameters> SceneTexturesUniformBuffer,
 		bool bEnableSSAO);
 
-	// only call if LPV is enabled
-	void ProcessLpvIndirect(FRHICommandListImmediate& RHICmdList, FViewInfo& View);
-
 	bool CanProcessAsyncSSAO(const TArray<FViewInfo>& Views);
 
 	void ProcessAsyncSSAO(
diff --git a/Engine/Source/Runtime/Renderer/Private/CompositionLighting/PostProcessLpvIndirect.cpp b/Engine/Source/Runtime/Renderer/Private/CompositionLighting/PostProcessLpvIndirect.cpp
deleted file mode 100644
diff --git a/Engine/Source/Runtime/Renderer/Private/CompositionLighting/PostProcessLpvIndirect.h b/Engine/Source/Runtime/Renderer/Private/CompositionLighting/PostProcessLpvIndirect.h
deleted file mode 100644
diff --git a/Engine/Source/Runtime/Renderer/Private/DeferredShadingRenderer.cpp b/Engine/Source/Runtime/Renderer/Private/DeferredShadingRenderer.cpp
--- a/Engine/Source/Runtime/Renderer/Private/DeferredShadingRenderer.cpp
+++ b/Engine/Source/Runtime/Renderer/Private/DeferredShadingRenderer.cpp
@@ -220,7 +220,6 @@ DECLARE_CYCLE_STAT(TEXT("DeferredShadingSceneRenderer DistanceFieldAO Init"), ST
 DECLARE_CYCLE_STAT(TEXT("DeferredShadingSceneRenderer FGlobalDynamicVertexBuffer Commit"), STAT_FDeferredShadingSceneRenderer_FGlobalDynamicVertexBuffer_Commit, STATGROUP_SceneRendering);
 DECLARE_CYCLE_STAT(TEXT("DeferredShadingSceneRenderer FXSystem PreRender"), STAT_FDeferredShadingSceneRenderer_FXSystem_PreRender, STATGROUP_SceneRendering);
 DECLARE_CYCLE_STAT(TEXT("DeferredShadingSceneRenderer AllocGBufferTargets"), STAT_FDeferredShadingSceneRenderer_AllocGBufferTargets, STATGROUP_SceneRendering);
-DECLARE_CYCLE_STAT(TEXT("DeferredShadingSceneRenderer ClearLPVs"), STAT_FDeferredShadingSceneRenderer_ClearLPVs, STATGROUP_SceneRendering);
 DECLARE_CYCLE_STAT(TEXT("DeferredShadingSceneRenderer DBuffer"), STAT_FDeferredShadingSceneRenderer_DBuffer, STATGROUP_SceneRendering);
 DECLARE_CYCLE_STAT(TEXT("DeferredShadingSceneRenderer ResolveDepth After Basepass"), STAT_FDeferredShadingSceneRenderer_ResolveDepth_After_Basepass, STATGROUP_SceneRendering);
 DECLARE_CYCLE_STAT(TEXT("DeferredShadingSceneRenderer Resolve After Basepass"), STAT_FDeferredShadingSceneRenderer_Resolve_After_Basepass, STATGROUP_SceneRendering);
@@ -1323,8 +1322,6 @@ void FDeferredShadingSceneRenderer::WaitForRayTracingScene(FRDGBuilder& GraphBui
 
 #endif // RHI_RAYTRACING
 
-extern bool IsLpvIndirectPassRequired(const FViewInfo& View);
-
 static TAutoConsoleVariable<float> CVarStallInitViews(
 	TEXT("CriticalPathStall.AfterInitViews"),
 	0.0f,
@@ -2013,15 +2010,7 @@ void FDeferredShadingSceneRenderer::Render(FRHICommandListImmediate& RHICmdList)
 	const bool bStrataEnabled = Strata::IsStrataEnabled();
 	if (bStrataEnabled)
 	{		
-		Strata::InitialiseStrataFrameSceneData(*this, GraphBuilder);		
-	}
-
-	// Clear LPVs for all views
-	if (FeatureLevel >= ERHIFeatureLevel::SM5)
-	{
-		SCOPE_CYCLE_COUNTER(STAT_FDeferredShadingSceneRenderer_ClearLPVs);
-		ClearLPVs(GraphBuilder);
-		AddServiceLocalQueuePass(GraphBuilder);
+		Strata::InitialiseStrataFrameSceneData(*this, GraphBuilder);
 	}
 
 	if(GetCustomDepthPassLocation() == 0)
@@ -2410,24 +2399,6 @@ void FDeferredShadingSceneRenderer::Render(FRHICommandListImmediate& RHICmdList)
 		}
 		AddServiceLocalQueuePass(GraphBuilder);
 
-		// Pre-lighting composition lighting stage
-		// e.g. LPV indirect
-		for (int32 ViewIndex = 0; ViewIndex < Views.Num(); ++ViewIndex)
-		{
-			FViewInfo& View = Views[ViewIndex]; 
-
-			if(IsLpvIndirectPassRequired(View))
-			{
-				AddUntrackedAccessPass(GraphBuilder, [this, &View, ViewIndex](FRHICommandListImmediate& InRHICmdList)
-				{
-					SCOPED_GPU_MASK(InRHICmdList, View.GPUMask);
-					SCOPED_CONDITIONAL_DRAW_EVENTF(InRHICmdList, EventView, Views.Num() > 1, TEXT("View%d"), ViewIndex);
-					GCompositionLighting.ProcessLpvIndirect(InRHICmdList, View);
-					ServiceLocalQueue();
-				});
-			}
-		}
-
 		// Render diffuse sky lighting and reflections that only operate on opaque pixels
 		RenderDeferredReflectionsAndSkyLighting(GraphBuilder, SceneTextures, SceneColorTexture, DynamicBentNormalAOTexture, VelocityTexture, HairDatas);
 
diff --git a/Engine/Source/Runtime/Renderer/Private/DeferredShadingRenderer.h b/Engine/Source/Runtime/Renderer/Private/DeferredShadingRenderer.h
--- a/Engine/Source/Runtime/Renderer/Private/DeferredShadingRenderer.h
+++ b/Engine/Source/Runtime/Renderer/Private/DeferredShadingRenderer.h
@@ -179,12 +179,6 @@ public:
 	/** Clears a view */
 	void ClearView(FRHICommandListImmediate& RHICmdList);
 
-	/** Clears LPVs for all views */
-	void ClearLPVs(FRDGBuilder& GraphBuilder);
-
-	/** Propagates LPVs for all views */
-	void UpdateLPVs(FRHICommandListImmediate& RHICmdList);
-
 	/**
 	 * Renders the scene's prepass for a particular view
 	 * @return true if anything was rendered
@@ -669,9 +663,6 @@ private:
 
 	FRDGTextureRef CopyStencilToLightingChannelTexture(FRDGBuilder& GraphBuilder, FRDGTextureSRVRef SceneStencilTexture);
 
-	/** Injects reflective shadowmaps into LPVs */
-	bool InjectReflectiveShadowMaps(FRHICommandListImmediate& RHICmdList, const FLightSceneInfo* LightSceneInfo);
-
 	/** Renders capsule shadows for all per-object shadows using it for the given light. */
 	bool RenderCapsuleDirectShadows(
 		FRDGBuilder& GraphBuilder,
diff --git a/Engine/Source/Runtime/Renderer/Private/LightPropagationVolume.cpp b/Engine/Source/Runtime/Renderer/Private/LightPropagationVolume.cpp
deleted file mode 100644
diff --git a/Engine/Source/Runtime/Renderer/Private/LightPropagationVolume.h b/Engine/Source/Runtime/Renderer/Private/LightPropagationVolume.h
deleted file mode 100644
diff --git a/Engine/Source/Runtime/Renderer/Private/LightPropagationVolumeVisualisation.cpp b/Engine/Source/Runtime/Renderer/Private/LightPropagationVolumeVisualisation.cpp
deleted file mode 100644
diff --git a/Engine/Source/Runtime/Renderer/Private/LightRendering.cpp b/Engine/Source/Runtime/Renderer/Private/LightRendering.cpp
--- a/Engine/Source/Runtime/Renderer/Private/LightRendering.cpp
+++ b/Engine/Source/Runtime/Renderer/Private/LightRendering.cpp
@@ -7,7 +7,6 @@
 #include "LightRendering.h"
 #include "RendererModule.h"
 #include "DeferredShadingRenderer.h"
-#include "LightPropagationVolume.h"
 #include "ScenePrivate.h"
 #include "PostProcess/SceneFilterRendering.h"
 #include "PipelineStateCache.h"
@@ -1349,78 +1348,6 @@ FRDGTextureRef FDeferredShadingSceneRenderer::RenderLights(
 			}
 		}
 
-		if ( IsFeatureLevelSupported(ShaderPlatformForFeatureLevel, ERHIFeatureLevel::SM5) )
-		{
-			RDG_EVENT_SCOPE(GraphBuilder, "IndirectLighting");
-			bool bRenderedRSM = false;
-			// Render Reflective shadow maps
-			// Draw shadowed and light function lights
-			for (int32 LightIndex = AttenuationLightStart; LightIndex < SortedLights.Num(); LightIndex++)
-			{
-				const FSortedLightSceneInfo& SortedLightInfo = SortedLights[LightIndex];
-				const FLightSceneInfo& LightSceneInfo = *SortedLightInfo.LightSceneInfo;
-				// Render any reflective shadow maps (if necessary)
-				if ( LightSceneInfo.Proxy && LightSceneInfo.Proxy->AffectsDynamicIndirectLighting() )
-				{
-					if ( LightSceneInfo.Proxy->GetLightType() == LightType_Directional )
-					{
-						INC_DWORD_STAT(STAT_NumReflectiveShadowMapLights);
-						AddUntrackedAccessPass(GraphBuilder, [this, &LightSceneInfo](FRHICommandListImmediate& RHICmdList)
-						{
-							InjectReflectiveShadowMaps(RHICmdList, &LightSceneInfo);
-						});
-						bRenderedRSM = true;
-					}
-				}
-			}
-
-			// LPV Direct Light Injection
-			if (bRenderedRSM)
-			{
-				for (int32 LightIndex = SimpleLightsEnd; LightIndex < SortedLights.Num(); LightIndex++)
-				{
-					const FSortedLightSceneInfo& SortedLightInfo = SortedLights[LightIndex];
-					const FLightSceneInfo* const LightSceneInfo = SortedLightInfo.LightSceneInfo;
-
-					// Render any reflective shadow maps (if necessary)
-					if ( LightSceneInfo && LightSceneInfo->Proxy && LightSceneInfo->Proxy->AffectsDynamicIndirectLighting() )
-					{
-						if ( LightSceneInfo->Proxy->GetLightType() != LightType_Directional )
-						{
-							// Inject the light directly into all relevant LPVs
-							for (int32 ViewIndex = 0; ViewIndex < Views.Num(); ViewIndex++)
-							{
-								FViewInfo& View = Views[ViewIndex];
-								RDG_GPU_MASK_SCOPE(GraphBuilder, View.GPUMask);
-
-								if (LightSceneInfo->ShouldRenderLight(View))
-								{
-									FSceneViewState* ViewState = (FSceneViewState*)View.State;
-									if (ViewState)
-									{
-										FLightPropagationVolume* Lpv = ViewState->GetLightPropagationVolume(View.GetFeatureLevel());
-										if (Lpv && LightSceneInfo->Proxy)
-										{
-											AddUntrackedAccessPass(GraphBuilder, [Lpv, LightSceneInfo, &View](FRHICommandListImmediate& RHICmdList)
-											{
-												Lpv->InjectLightDirect(RHICmdList, *LightSceneInfo->Proxy, View);
-											});
-										}
-									}
-								}
-							}
-						}
-					}
-				}
-			}
-
-			// Kickoff the LPV update (asynchronously if possible)
-			AddUntrackedAccessPass(GraphBuilder, [this](FRHICommandListImmediate& RHICmdList)
-			{
-				UpdateLPVs(RHICmdList);
-			});
-		}
-
 		{
 			RDG_EVENT_SCOPE(GraphBuilder, "ShadowedLights");
 
diff --git a/Engine/Source/Runtime/Renderer/Private/PostProcess/PostProcessing.cpp b/Engine/Source/Runtime/Renderer/Private/PostProcess/PostProcessing.cpp
--- a/Engine/Source/Runtime/Renderer/Private/PostProcess/PostProcessing.cpp
+++ b/Engine/Source/Runtime/Renderer/Private/PostProcess/PostProcessing.cpp
@@ -38,7 +38,6 @@
 #include "PostProcess/PostProcessStreamingAccuracyLegend.h"
 #include "PostProcess/PostProcessSubsurface.h"
 #include "CompositionLighting/PostProcessPassThrough.h"
-#include "CompositionLighting/PostProcessLpvIndirect.h"
 #include "ShaderPrint.h"
 #include "GpuDebugRendering.h"
 #include "HighResScreenshot.h"
@@ -891,11 +890,6 @@ void AddPostProcessingPasses(FRDGBuilder& GraphBuilder, const FViewInfo& View, c
 		SceneColor = AddVisualizeComplexityPass(GraphBuilder, View, PassInputs);
 	}
 
-	if (EngineShowFlags.VisualizeLPV)
-	{
-		AddVisualizeLPVPass(GraphBuilder, View, SceneColor);
-	}
-
 #if WITH_EDITOR
 	if (PassSequence.IsEnabled(EPass::VisualizeLevelInstance))
 	{
diff --git a/Engine/Source/Runtime/Renderer/Private/PostProcess/SceneRenderTargets.cpp b/Engine/Source/Runtime/Renderer/Private/PostProcess/SceneRenderTargets.cpp
--- a/Engine/Source/Runtime/Renderer/Private/PostProcess/SceneRenderTargets.cpp
+++ b/Engine/Source/Runtime/Renderer/Private/PostProcess/SceneRenderTargets.cpp
@@ -12,7 +12,6 @@
 #include "SceneTextureParameters.h"
 #include "VelocityRendering.h"
 #include "RendererModule.h"
-#include "LightPropagationVolume.h"
 #include "ScenePrivate.h"
 #include "HdrCustomResolveShaders.h"
 #include "WideCustomResolveShaders.h"
@@ -29,11 +28,6 @@
 #include "VT/VirtualTextureFeedback.h"
 #include "VisualizeTexture.h"
 #include "GpuDebugRendering.h"
-static TAutoConsoleVariable<int32> CVarRSMResolution(
-	TEXT("r.LPV.RSMResolution"),
-	360,
-	TEXT("Reflective Shadow Map resolution (used for LPV) - higher values result in less aliasing artifacts, at the cost of performance"),
-	ECVF_Scalability | ECVF_RenderThreadSafe);
 
 /*-----------------------------------------------------------------------------
 FSceneRenderTargets
@@ -236,7 +230,6 @@ static void SnapshotArray(TArray<TRefCountPtr<IPooledRenderTarget>, TInlineAlloc
 
 FSceneRenderTargets::FSceneRenderTargets(const FViewInfo& View, const FSceneRenderTargets& SnapshotSource)
 	: LightAccumulation(GRenderTargetPool.MakeSnapshot(SnapshotSource.LightAccumulation))
-	, DirectionalOcclusion(GRenderTargetPool.MakeSnapshot(SnapshotSource.DirectionalOcclusion))
 	, SceneDepthZ(GRenderTargetPool.MakeSnapshot(SnapshotSource.SceneDepthZ))	
 	, SceneVelocity(GRenderTargetPool.MakeSnapshot(SnapshotSource.SceneVelocity))
 	, SmallDepthZ(GRenderTargetPool.MakeSnapshot(SnapshotSource.SmallDepthZ))
@@ -278,7 +271,6 @@ FSceneRenderTargets::FSceneRenderTargets(const FViewInfo& View, const FSceneRend
 	, CurrentMobileSceneColorFormat(SnapshotSource.CurrentMobileSceneColorFormat)
 	, bAllowStaticLighting(SnapshotSource.bAllowStaticLighting)
 	, CurrentMaxShadowResolution(SnapshotSource.CurrentMaxShadowResolution)
-	, CurrentRSMResolution(SnapshotSource.CurrentRSMResolution)
 	, CurrentTranslucencyLightingVolumeDim(SnapshotSource.CurrentTranslucencyLightingVolumeDim)
 	, CurrentMSAACount(SnapshotSource.CurrentMSAACount)
 	, CurrentMinShadowResolution(SnapshotSource.CurrentMinShadowResolution)
@@ -652,8 +644,6 @@ void FSceneRenderTargets::Allocate(FRHICommandListImmediate& RHICmdList, const F
 
 	int32 MaxShadowResolution = GetCachedScalabilityCVars().MaxShadowResolution;
 
-	int32 RSMResolution = FMath::Clamp(CVarRSMResolution.GetValueOnRenderThread(), 1, 2048);
-
 	if (ViewFamily.Scene->GetShadingPath() == EShadingPath::Mobile)
 	{
 		// ensure there is always enough space for mobile renderer's tiled shadow maps
@@ -669,8 +659,6 @@ void FSceneRenderTargets::Allocate(FRHICommandListImmediate& RHICmdList, const F
 
 	int32 MSAACount = GetNumSceneColorMSAASamples(NewFeatureLevel);
 
-	bool bLightPropagationVolume = UseLightPropagationVolumeRT(NewFeatureLevel);
-
 	uint32 MinShadowResolution;
 	{
 		static const auto CVar = IConsoleManager::Get().FindTConsoleVariableDataInt(TEXT("r.Shadow.MinResolution"));
@@ -686,10 +674,8 @@ void FSceneRenderTargets::Allocate(FRHICommandListImmediate& RHICmdList, const F
 		(bAllowStaticLighting != bNewAllowStaticLighting) ||
 		(bUseDownsizedOcclusionQueries != bDownsampledOcclusionQueries) ||
 		(CurrentMaxShadowResolution != MaxShadowResolution) ||
- 		(CurrentRSMResolution != RSMResolution) ||
 		(CurrentTranslucencyLightingVolumeDim != TranslucencyLightingVolumeDim) ||
 		(CurrentMSAACount != MSAACount) ||
-		(bCurrentLightPropagationVolume != bLightPropagationVolume) ||
 		(CurrentMinShadowResolution != MinShadowResolution))
 	{
 		CurrentGBufferFormat = GBufferFormat;
@@ -698,11 +684,9 @@ void FSceneRenderTargets::Allocate(FRHICommandListImmediate& RHICmdList, const F
 		bAllowStaticLighting = bNewAllowStaticLighting;
 		bUseDownsizedOcclusionQueries = bDownsampledOcclusionQueries;
 		CurrentMaxShadowResolution = MaxShadowResolution;
-		CurrentRSMResolution = RSMResolution;
 		CurrentTranslucencyLightingVolumeDim = TranslucencyLightingVolumeDim;
 		CurrentMSAACount = MSAACount;
 		CurrentMinShadowResolution = MinShadowResolution;
-		bCurrentLightPropagationVolume = bLightPropagationVolume;
 
 		// Reinitialize the render targets for the given size.
 		SetBufferSize(DesiredBufferSize.X, DesiredBufferSize.Y);
@@ -1721,13 +1705,6 @@ void FSceneRenderTargets::AllocateDeferredShadingPathRenderTargets(FRHICommandLi
 		}
 	}
 
-	// LPV : Dynamic directional occlusion for diffuse and specular
-	if(UseLightPropagationVolumeRT(CurrentFeatureLevel))
-	{
-		FPooledRenderTargetDesc Desc(FPooledRenderTargetDesc::Create2DDesc(BufferSize, PF_R8G8, FClearValueBinding::Transparent, TexCreate_None, TexCreate_RenderTargetable | TexCreate_ShaderResource, false));
-		GRenderTargetPool.FindFreeElement(RHICmdList, Desc, DirectionalOcclusion, TEXT("DirectionalOcclusion"));
-	}
-
 	if (CurrentFeatureLevel >= ERHIFeatureLevel::SM5) 
 	{
 		FPooledRenderTargetDesc Desc(FPooledRenderTargetDesc::Create2DDesc(BufferSize, PF_FloatRGBA, FClearValueBinding::Black, TexCreate_None, TexCreate_RenderTargetable | TexCreate_ShaderResource, false));
@@ -2001,7 +1978,6 @@ void FSceneRenderTargets::ReleaseAllTargets()
 	ScreenSpaceGTAOHorizons.SafeRelease();
 	QuadOverdrawBuffer.SafeRelease();
 	LightAccumulation.SafeRelease();
-	DirectionalOcclusion.SafeRelease();
 	CustomDepth.SafeRelease();
 	MobileCustomDepth.SafeRelease();
 	MobileCustomStencil.SafeRelease();
diff --git a/Engine/Source/Runtime/Renderer/Private/PostProcess/SceneRenderTargets.h b/Engine/Source/Runtime/Renderer/Private/PostProcess/SceneRenderTargets.h
--- a/Engine/Source/Runtime/Renderer/Private/PostProcess/SceneRenderTargets.h
+++ b/Engine/Source/Runtime/Renderer/Private/PostProcess/SceneRenderTargets.h
@@ -181,9 +181,7 @@ protected:
 		CurrentMobileSceneColorFormat(EPixelFormat::PF_Unknown),
 		bAllowStaticLighting(true),
 		CurrentMaxShadowResolution(0),
-		CurrentRSMResolution(0),
 		CurrentTranslucencyLightingVolumeDim(64),
-		bCurrentLightPropagationVolume(false),
 		CurrentFeatureLevel(ERHIFeatureLevel::Num),
 		CurrentShadingPath(EShadingPath::Num),
 		bRequireSceneColorAlpha(false),
@@ -334,11 +332,6 @@ public:
 	const FTexture2DRHIRef& GetSmallDepthSurface() const							{ return (const FTexture2DRHIRef&)SmallDepthZ->GetRenderTargetItem().TargetableTexture; }
 	const FTexture2DRHIRef& GetOptionalShadowDepthColorSurface(FRHICommandList& RHICmdList, int32 Width, int32 Height) const;
 
-	const FTexture2DRHIRef& GetDirectionalOcclusionTexture() const 
-	{	
-		return (const FTexture2DRHIRef&)DirectionalOcclusion->GetRenderTargetItem().TargetableTexture; 
-	}
-
 	// @return can be empty if the feature is disabled
 	FCustomDepthTextures RequestCustomDepth(FRDGBuilder& GraphBuilder, bool bPrimitives);
 
@@ -388,9 +381,6 @@ public:
 	FIntPoint GetTranslucentShadowDepthTextureResolution() const;
 	int32 GetTranslucentShadowDownsampleFactor() const { return 2; }
 
-	/** Returns the size of the RSM buffer, taking into account platform limitations and game specific resolution limits. */
-	inline int32 GetReflectiveShadowMapResolution() const { return CurrentRSMResolution; }
-
 	int32 GetNumGBufferTargets() const;
 
 	int32 GetMSAACount() const { return CurrentMSAACount; }
@@ -452,9 +442,6 @@ public:
 	// Light Accumulation is a high precision scratch pad matching the size of the scene color buffer used by many passes.
 	TRefCountPtr<IPooledRenderTarget> LightAccumulation;
 
-	// Reflection Environment: Bringing back light accumulation buffer to apply indirect reflections
-	TRefCountPtr<IPooledRenderTarget> DirectionalOcclusion;
-	
 	// Scene depth and stencil.
 	TRefCountPtr<IPooledRenderTarget> SceneDepthZ;
 	TRefCountPtr<FRHIShaderResourceView> SceneStencilSRV;
@@ -667,8 +654,6 @@ private:
 	bool bAllowStaticLighting;
 	/** To detect a change of the CVar r.Shadow.MaxResolution */
 	int32 CurrentMaxShadowResolution;
-	/** To detect a change of the CVar r.Shadow.RsmResolution*/
-	int32 CurrentRSMResolution;
 	/** To detect a change of the CVar r.TranslucencyLightingVolumeDim */
 	int32 CurrentTranslucencyLightingVolumeDim;
 	/** To detect a change of the CVar r.MobileMSAA or r.MSAA */
diff --git a/Engine/Source/Runtime/Renderer/Private/ReflectionEnvironment.cpp b/Engine/Source/Runtime/Renderer/Private/ReflectionEnvironment.cpp
--- a/Engine/Source/Runtime/Renderer/Private/ReflectionEnvironment.cpp
+++ b/Engine/Source/Runtime/Renderer/Private/ReflectionEnvironment.cpp
@@ -25,7 +25,6 @@
 #include "PostProcess/PostProcessing.h"
 #include "PostProcess/PostProcessSubsurface.h"
 #include "LightRendering.h"
-#include "LightPropagationVolumeSettings.h"
 #include "PipelineStateCache.h"
 #include "DistanceFieldAmbientOcclusion.h"
 #include "SceneTextureParameters.h"
@@ -37,9 +36,6 @@
 #include "SceneTextureParameters.h"
 #include "HairStrands/HairStrandsRendering.h"
 
-
-extern TAutoConsoleVariable<int32> CVarLPVMixing;
-
 static TAutoConsoleVariable<int32> CVarReflectionEnvironment(
 	TEXT("r.ReflectionEnvironment"),
 	1,
diff --git a/Engine/Source/Runtime/Renderer/Private/RendererPrivate.h b/Engine/Source/Runtime/Renderer/Private/RendererPrivate.h
--- a/Engine/Source/Runtime/Renderer/Private/RendererPrivate.h
+++ b/Engine/Source/Runtime/Renderer/Private/RendererPrivate.h
@@ -16,7 +16,6 @@
 #include "LightMapDensityRendering.h"
 #include "ClearQuad.h"
 #include "AtmosphereRendering.h"
-#include "LightPropagationVolume.h"
 #include "ScenePrivate.h"
 #include "RenderTargetTemp.h"
 #include "RendererModule.h"
diff --git a/Engine/Source/Runtime/Renderer/Private/RendererScene.cpp b/Engine/Source/Runtime/Renderer/Private/RendererScene.cpp
--- a/Engine/Source/Runtime/Renderer/Private/RendererScene.cpp
+++ b/Engine/Source/Runtime/Renderer/Private/RendererScene.cpp
@@ -41,7 +41,6 @@
 #include "SkyAtmosphereRendering.h"
 #include "BasePassRendering.h"
 #include "MobileBasePassRendering.h"
-#include "LightPropagationVolume.h"
 #include "ScenePrivate.h"
 #include "RendererModule.h"
 #include "StaticMeshResources.h"
@@ -201,8 +200,6 @@ FSceneViewState::FSceneViewState()
 
 	SequencerState = ESS_None;
 
-	LightPropagationVolume = NULL; 
-
 	bIsStereoView = false;
 
 	bRoundRobinOcclusionEnabled = false;
@@ -288,7 +285,6 @@ FSceneViewState::~FSceneViewState()
 	AOTileIntersectionResources = NULL;
 	DestroyRenderResource(AOScreenGridResources);
 	AOScreenGridResources = NULL;
-	DestroyLightPropagationVolume();
 
 #if RHI_RAYTRACING
 	DestroyRWBuffer(VarianceMipTree);
diff --git a/Engine/Source/Runtime/Renderer/Private/ScenePrivate.h b/Engine/Source/Runtime/Renderer/Private/ScenePrivate.h
--- a/Engine/Source/Runtime/Renderer/Private/ScenePrivate.h
+++ b/Engine/Source/Runtime/Renderer/Private/ScenePrivate.h
@@ -65,7 +65,6 @@ DECLARE_GPU_STAT_NAMED_EXTERN(ShadowProjection, TEXT("Shadow Projection"));
 
 class AWorldSettings;
 class FAtmosphericFogSceneInfo;
-class FLightPropagationVolume;
 class FMaterialParameterCollectionInstanceResource;
 class FPrecomputedLightVolume;
 class FScene;
@@ -929,9 +928,6 @@ private:
 	// counts up by one each frame, warped in 0..3 range, ResetViewState() puts it back to 0
 	int32 DistanceFieldTemporalSampleIndex;
 
-	// light propagation volume used in this view
-	TRefCountPtr<FLightPropagationVolume> LightPropagationVolume;
-
 	// whether this view is a stereo counterpart to a primary view
 	bool bIsStereoView;
 
@@ -1182,21 +1178,9 @@ public:
 		return DistanceFieldTemporalSampleIndex;
 	}
 
-
-	// call only if not yet created
-	void SetupLightPropagationVolume(FSceneView& View, FSceneViewFamily& ViewFamily);
-
-	/**
-	 * @return can return 0
-	 * @param bIncludeStereo - specifies whether the getter should include stereo views in its returned value
-	 */
-	FLightPropagationVolume* GetLightPropagationVolume(ERHIFeatureLevel::Type InFeatureLevel, bool bIncludeStereo = false) const;
-
 	/** Default constructor. */
 	FSceneViewState();
 
-	void DestroyLightPropagationVolume();
-
 	virtual ~FSceneViewState();
 
 	// called every frame after the view state was updated
@@ -1429,11 +1413,6 @@ public:
 	void OnStartRender(FViewInfo& View, FSceneViewFamily& ViewFamily)
 	{
 		check(IsInRenderingThread());
-
-		if(!(View.FinalPostProcessSettings.IndirectLightingColor * View.FinalPostProcessSettings.IndirectLightingIntensity).IsAlmostBlack())
-		{
-			SetupLightPropagationVolume(View, ViewFamily);
-		}
 		ConditionallyAllocateSceneSoftwareOcclusion(View.GetFeatureLevel());
 	}
 
diff --git a/Engine/Source/Runtime/Renderer/Private/SceneRendering.cpp b/Engine/Source/Runtime/Renderer/Private/SceneRendering.cpp
--- a/Engine/Source/Runtime/Renderer/Private/SceneRendering.cpp
+++ b/Engine/Source/Runtime/Renderer/Private/SceneRendering.cpp
@@ -392,7 +392,6 @@ FASTVRAM_CVAR(GBufferVelocity, 0);
 FASTVRAM_CVAR(HZB, 1);
 FASTVRAM_CVAR(SceneDepth, 1);
 FASTVRAM_CVAR(SceneColor, 1);
-FASTVRAM_CVAR(LPV, 1);
 FASTVRAM_CVAR(BokehDOF, 1);
 FASTVRAM_CVAR(CircleDOF, 1);
 FASTVRAM_CVAR(CombineLUTs, 1);
@@ -534,7 +533,6 @@ void FFastVramConfig::Update()
 	bDirty |= UpdateTextureFlagFromCVar(CVarFastVRam_HZB, HZB);
 	bDirty |= UpdateTextureFlagFromCVar(CVarFastVRam_SceneDepth, SceneDepth);
 	bDirty |= UpdateTextureFlagFromCVar(CVarFastVRam_SceneColor, SceneColor);
-	bDirty |= UpdateTextureFlagFromCVar(CVarFastVRam_LPV, LPV);
 	bDirty |= UpdateTextureFlagFromCVar(CVarFastVRam_BokehDOF, BokehDOF);
 	bDirty |= UpdateTextureFlagFromCVar(CVarFastVRam_CircleDOF, CircleDOF);
 	bDirty |= UpdateTextureFlagFromCVar(CVarFastVRam_CombineLUTs, CombineLUTs);
@@ -2171,11 +2169,6 @@ void FSortedShadowMaps::Release()
 		ShadowMapAtlases[AtlasIndex].RenderTargets.Release();
 	}
 
-	for (int32 AtlasIndex = 0; AtlasIndex < RSMAtlases.Num(); AtlasIndex++)
-	{
-		RSMAtlases[AtlasIndex].RenderTargets.Release();
-	}
-
 	for (int32 AtlasIndex = 0; AtlasIndex < ShadowMapCubemaps.Num(); AtlasIndex++)
 	{
 		ShadowMapCubemaps[AtlasIndex].RenderTargets.Release();
diff --git a/Engine/Source/Runtime/Renderer/Private/SceneRendering.h b/Engine/Source/Runtime/Renderer/Private/SceneRendering.h
--- a/Engine/Source/Runtime/Renderer/Private/SceneRendering.h
+++ b/Engine/Source/Runtime/Renderer/Private/SceneRendering.h
@@ -169,7 +169,6 @@ public:
 	/** Shadows to project for each feature that needs special handling. */
 	TArray<FProjectedShadowInfo*,SceneRenderingAllocator> ShadowsToProject;
 	TArray<FProjectedShadowInfo*,SceneRenderingAllocator> CapsuleShadowsToProject;
-	TArray<FProjectedShadowInfo*,SceneRenderingAllocator> RSMsToProject;
 
 	/** All visible projected preshdows.  These are not allocated on the mem stack so they are refcounted. */
 	TArray<TRefCountPtr<FProjectedShadowInfo>,SceneRenderingAllocator> ProjectedPreShadows;
@@ -1638,8 +1637,6 @@ struct FSortedShadowMaps
 	/** Visible shadows sorted by their shadow depth map render target. */
 	TArray<FSortedShadowMapAtlas,SceneRenderingAllocator> ShadowMapAtlases;
 
-	TArray<FSortedShadowMapAtlas,SceneRenderingAllocator> RSMAtlases;
-
 	TArray<FSortedShadowMapAtlas,SceneRenderingAllocator> ShadowMapCubemaps;
 
 	FSortedShadowMapAtlas PreshadowCache;
@@ -1663,11 +1660,6 @@ struct FSortedShadowMaps
 			MemorySize += ShadowMapAtlases[i].RenderTargets.ComputeMemorySize();
 		}
 
-		for (int i = 0; i < RSMAtlases.Num(); i++)
-		{
-			MemorySize += RSMAtlases[i].RenderTargets.ComputeMemorySize();
-		}
-
 		for (int i = 0; i < ShadowMapCubemaps.Num(); i++)
 		{
 			MemorySize += ShadowMapCubemaps[i].RenderTargets.ComputeMemorySize();
@@ -1897,8 +1889,6 @@ protected:
 
 	void AllocateCSMDepthTargets(FRHICommandListImmediate& RHICmdList, const TArray<FProjectedShadowInfo*, SceneRenderingAllocator>& WholeSceneDirectionalShadows, TArray<FSortedShadowMapAtlas,SceneRenderingAllocator>& OutAtlases);
 
-	void AllocateRSMDepthTargets(FRHICommandListImmediate& RHICmdList, const TArray<FProjectedShadowInfo*, SceneRenderingAllocator>& RSMShadows);
-
 	void AllocateOnePassPointLightDepthTargets(FRHICommandListImmediate& RHICmdList, const TArray<FProjectedShadowInfo*, SceneRenderingAllocator>& WholeScenePointShadows);
 
 	void AllocateTranslucentShadowDepthTargets(FRHICommandListImmediate& RHICmdList, TArray<FProjectedShadowInfo*, SceneRenderingAllocator>& TranslucentShadows);
@@ -2257,7 +2247,6 @@ struct FFastVramConfig
 	ETextureCreateFlags HZB;
 	ETextureCreateFlags SceneDepth;
 	ETextureCreateFlags SceneColor;
-	ETextureCreateFlags LPV;
 	ETextureCreateFlags BokehDOF;
 	ETextureCreateFlags CircleDOF;
 	ETextureCreateFlags CombineLUTs;
diff --git a/Engine/Source/Runtime/Renderer/Private/ShadowDepthRendering.cpp b/Engine/Source/Runtime/Renderer/Private/ShadowDepthRendering.cpp
--- a/Engine/Source/Runtime/Renderer/Private/ShadowDepthRendering.cpp
+++ b/Engine/Source/Runtime/Renderer/Private/ShadowDepthRendering.cpp
@@ -30,7 +30,6 @@
 #include "ShaderBaseClasses.h"
 #include "ShadowRendering.h"
 #include "SceneRendering.h"
-#include "LightPropagationVolume.h"
 #include "ScenePrivate.h"
 #include "PostProcess/SceneFilterRendering.h"
 #include "ScreenRendering.h"
@@ -140,9 +139,7 @@ void SetupShadowDepthPassUniformBuffer(
 	const FProjectedShadowInfo* ShadowInfo,
 	FRHICommandListImmediate& RHICmdList,
 	const FViewInfo& View,
-	FShadowDepthPassUniformParameters& ShadowDepthPassParameters,
-	FLightPropagationVolume* LPV = nullptr
-)
+	FShadowDepthPassUniformParameters& ShadowDepthPassParameters)
 {
 	FSceneRenderTargets& SceneRenderTargets = FSceneRenderTargets::Get(RHICmdList);
 	SetupSceneTextureUniformParameters(SceneRenderTargets, View.FeatureLevel, ESceneTextureSetupMode::None, ShadowDepthPassParameters.SceneTextures);
@@ -164,32 +161,6 @@ void SetupShadowDepthPassUniformBuffer(
 			ShadowDepthPassParameters.ShadowViewMatrices[FaceIndex] = Translation * ShadowInfo->OnePassShadowViewMatrices[FaceIndex];
 		}
 	}
-
-
-	ShadowDepthPassParameters.RWGvListBuffer = GBlackTextureWithUAV->UnorderedAccessViewRHI;
-	ShadowDepthPassParameters.RWGvListHeadBuffer = GBlackTextureWithUAV->UnorderedAccessViewRHI;
-	ShadowDepthPassParameters.RWVplListBuffer = GBlackTextureWithUAV->UnorderedAccessViewRHI;
-	ShadowDepthPassParameters.RWVplListHeadBuffer = GBlackTextureWithUAV->UnorderedAccessViewRHI;
-
-	if (ShadowInfo->bReflectiveShadowmap)
-	{
-		const FSceneViewState* ViewState = (const FSceneViewState*)View.State;
-
-		if (ViewState)
-		{
-			const FLightPropagationVolume* Lpv = ViewState->GetLightPropagationVolume(View.GetFeatureLevel());
-
-			if (Lpv)
-			{
-				ShadowDepthPassParameters.LPV = Lpv->GetWriteUniformBufferParams();
-			}
-
-			ShadowDepthPassParameters.RWGvListBuffer = LPV->GetGvListBufferUav();
-			ShadowDepthPassParameters.RWGvListHeadBuffer = LPV->GetGvListHeadBufferUav();
-			ShadowDepthPassParameters.RWVplListBuffer = LPV->GetVplListBufferUav();
-			ShadowDepthPassParameters.RWVplListHeadBuffer = LPV->GetVplListHeadBufferUav();
-		}
-	}
 }
 
 void SetupShadowDepthPassUniformBuffer(
@@ -285,7 +256,7 @@ static TAutoConsoleVariable<int32> CVarSupportPointLightWholeSceneShadows(
 /**
 * A vertex shader for rendering the depth of a mesh.
 */
-template <EShadowDepthVertexShaderMode ShaderMode, bool bRenderReflectiveShadowMap, bool bUsePositionOnlyStream, bool bIsForGeometryShader = false>
+template <EShadowDepthVertexShaderMode ShaderMode, bool bUsePositionOnlyStream, bool bIsForGeometryShader = false>
 class TShadowDepthVS : public FShadowDepthVS
 {
 	DECLARE_SHADER_TYPE(TShadowDepthVS, MeshMaterial);
@@ -318,35 +289,15 @@ public:
 		}
 
 		//Note: This logic needs to stay in sync with OverrideWithDefaultMaterialForShadowDepth!
-		// Compile for special engine materials.
-		if (bRenderReflectiveShadowMap)
-		{
-			static const auto SupportLPV = IConsoleManager::Get().FindTConsoleVariableDataInt(TEXT("r.LightPropagationVolume"));
-			if (SupportLPV && SupportLPV->GetValueOnAnyThread() == 0)
-			{
-				return false;
-			}
-			else
-			{
-				// Reflective shadow map shaders must be compiled for every material because they access the material normal
-				return !bUsePositionOnlyStream
-					// Don't render ShadowDepth for translucent unlit materials, unless we're injecting emissive
-					&& (Parameters.MaterialParameters.bShouldCastDynamicShadows || Parameters.MaterialParameters.bShouldInjectEmissiveIntoLPV || Parameters.MaterialParameters.bShouldBlockGI)
-					&& IsFeatureLevelSupported(Platform, ERHIFeatureLevel::SM5);
-			}
-		}
-		else
-		{
-			return (Parameters.MaterialParameters.bIsSpecialEngineMaterial
-				// Masked and WPO materials need their shaders but cannot be used with a position only stream.
-				|| ((!Parameters.MaterialParameters.bWritesEveryPixelShadowPass || Parameters.MaterialParameters.bMaterialMayModifyMeshPosition) && !bUsePositionOnlyStream))
-				// Only compile one pass point light shaders for feature levels >= SM5
-				&& (ShaderMode != VertexShadowDepth_OnePassPointLight || IsFeatureLevelSupported(Platform, ERHIFeatureLevel::SM5))
-				// Only compile position-only shaders for vertex factories that support it. (Note: this assumes that a vertex factor which supports PositionOnly, supports also PositionAndNormalOnly)
-				&& (!bUsePositionOnlyStream || Parameters.VertexFactoryType->SupportsPositionOnly())
-				// Don't render ShadowDepth for translucent unlit materials
-				&& Parameters.MaterialParameters.bShouldCastDynamicShadows;
-		}
+		return (Parameters.MaterialParameters.bIsSpecialEngineMaterial
+			// Masked and WPO materials need their shaders but cannot be used with a position only stream.
+			|| ((!Parameters.MaterialParameters.bWritesEveryPixelShadowPass || Parameters.MaterialParameters.bMaterialMayModifyMeshPosition) && !bUsePositionOnlyStream))
+			// Only compile one pass point light shaders for feature levels >= SM5
+			&& (ShaderMode != VertexShadowDepth_OnePassPointLight || IsFeatureLevelSupported(Platform, ERHIFeatureLevel::SM5))
+			// Only compile position-only shaders for vertex factories that support it. (Note: this assumes that a vertex factor which supports PositionOnly, supports also PositionAndNormalOnly)
+			&& (!bUsePositionOnlyStream || Parameters.VertexFactoryType->SupportsPositionOnly())
+			// Don't render ShadowDepth for translucent unlit materials
+			&& Parameters.MaterialParameters.bShouldCastDynamicShadows;
 	}
 
 	static void ModifyCompilationEnvironment(const FMaterialShaderPermutationParameters& Parameters, FShaderCompilerEnvironment& OutEnvironment)
@@ -355,7 +306,6 @@ public:
 		OutEnvironment.SetDefine(TEXT("PERSPECTIVE_CORRECT_DEPTH"), (uint32)(ShaderMode == VertexShadowDepth_PerspectiveCorrect));
 		OutEnvironment.SetDefine(TEXT("ONEPASS_POINTLIGHT_SHADOW"), (uint32)(ShaderMode == VertexShadowDepth_OnePassPointLight || ShaderMode == VertexShadowDepth_VSLayer));
 		OutEnvironment.SetDefine(TEXT("USING_VERTEX_SHADER_LAYER"), (uint32)(ShaderMode == VertexShadowDepth_VSLayer));
-		OutEnvironment.SetDefine(TEXT("REFLECTIVE_SHADOW_MAP"), (uint32)bRenderReflectiveShadowMap);
 		OutEnvironment.SetDefine(TEXT("POSITION_ONLY"), (uint32)bUsePositionOnlyStream);
 		OutEnvironment.SetDefine(TEXT("IS_FOR_GEOMETRY_SHADER"), (uint32)bIsForGeometryShader);
 
@@ -374,7 +324,7 @@ public:
 /**
 * A Hull shader for rendering the depth of a mesh.
 */
-template <EShadowDepthVertexShaderMode ShaderMode, bool bRenderReflectiveShadowMap>
+template <EShadowDepthVertexShaderMode ShaderMode>
 class TShadowDepthHS : public FBaseHS
 {
 	DECLARE_SHADER_TYPE(TShadowDepthHS, MeshMaterial);
@@ -391,21 +341,21 @@ public:
 	{
 		// Re-use ShouldCache from vertex shader
 		return FBaseHS::ShouldCompilePermutation(Parameters)
-			&& TShadowDepthVS<ShaderMode, bRenderReflectiveShadowMap, false>::ShouldCompilePermutation(Parameters);
+			&& TShadowDepthVS<ShaderMode, false>::ShouldCompilePermutation(Parameters);
 	}
 
 	static void ModifyCompilationEnvironment(const FMaterialShaderPermutationParameters& Parameters, FShaderCompilerEnvironment& OutEnvironment)
 	{
 		// Re-use compilation env from vertex shader
 
-		TShadowDepthVS<ShaderMode, bRenderReflectiveShadowMap, false>::ModifyCompilationEnvironment(Parameters, OutEnvironment);
+		TShadowDepthVS<ShaderMode, false>::ModifyCompilationEnvironment(Parameters, OutEnvironment);
 	}
 };
 
 /**
 * A Domain shader for rendering the depth of a mesh.
 */
-template <EShadowDepthVertexShaderMode ShaderMode, bool bRenderReflectiveShadowMap>
+template <EShadowDepthVertexShaderMode ShaderMode>
 class TShadowDepthDS : public FBaseDS
 {
 	DECLARE_SHADER_TYPE(TShadowDepthDS, MeshMaterial);
@@ -433,13 +383,13 @@ public:
 	{
 		// Re-use ShouldCache from vertex shader
 		return FBaseDS::ShouldCompilePermutation(Parameters)
-			&& TShadowDepthVS<ShaderMode, bRenderReflectiveShadowMap, false>::ShouldCompilePermutation(Parameters);
+			&& TShadowDepthVS<ShaderMode, false>::ShouldCompilePermutation(Parameters);
 	}
 
 	static void ModifyCompilationEnvironment(const FMaterialShaderPermutationParameters& Parameters, FShaderCompilerEnvironment& OutEnvironment)
 	{
 		// Re-use compilation env from vertex shader
-		TShadowDepthVS<ShaderMode, bRenderReflectiveShadowMap, false>::ModifyCompilationEnvironment(Parameters, OutEnvironment);
+		TShadowDepthVS<ShaderMode, false>::ModifyCompilationEnvironment(Parameters, OutEnvironment);
 	}
 };
 
@@ -451,14 +401,14 @@ public:
 
 	static bool ShouldCompilePermutation(const FMeshMaterialShaderPermutationParameters& Parameters)
 	{
-		return RHISupportsGeometryShaders(Parameters.Platform) && TShadowDepthVS<VertexShadowDepth_OnePassPointLight, false, false, true>::ShouldCompilePermutation(Parameters);
+		return RHISupportsGeometryShaders(Parameters.Platform) && TShadowDepthVS<VertexShadowDepth_OnePassPointLight, false, true>::ShouldCompilePermutation(Parameters);
 	}
 
 	static void ModifyCompilationEnvironment(const FMaterialShaderPermutationParameters& Parameters, FShaderCompilerEnvironment& OutEnvironment)
 	{
 		FMeshMaterialShader::ModifyCompilationEnvironment(Parameters, OutEnvironment);
 		OutEnvironment.SetDefine(TEXT("ONEPASS_POINTLIGHT_SHADOW"), 1);
-		TShadowDepthVS<VertexShadowDepth_OnePassPointLight, false, false, true>::ModifyCompilationEnvironment(Parameters, OutEnvironment);
+		TShadowDepthVS<VertexShadowDepth_OnePassPointLight, false, true>::ModifyCompilationEnvironment(Parameters, OutEnvironment);
 	}
 
 	FOnePassPointShadowDepthGS(const ShaderMetaType::CompiledShaderInitializerType& Initializer)
@@ -482,42 +432,40 @@ public:
 
 IMPLEMENT_SHADER_TYPE(, FOnePassPointShadowDepthGS, TEXT("/Engine/Private/ShadowDepthVertexShader.usf"), TEXT("MainOnePassPointLightGS"), SF_Geometry);
 
-#define IMPLEMENT_SHADOW_DEPTH_SHADERMODE_SHADERS(ShaderMode,bRenderReflectiveShadowMap) \
-	typedef TShadowDepthVS<ShaderMode, bRenderReflectiveShadowMap, false> TShadowDepthVS##ShaderMode##bRenderReflectiveShadowMap; \
-	typedef TShadowDepthHS<ShaderMode, bRenderReflectiveShadowMap       > TShadowDepthHS##ShaderMode##bRenderReflectiveShadowMap; \
-	typedef TShadowDepthDS<ShaderMode, bRenderReflectiveShadowMap       > TShadowDepthDS##ShaderMode##bRenderReflectiveShadowMap; \
-	IMPLEMENT_MATERIAL_SHADER_TYPE(template<>, TShadowDepthVS##ShaderMode##bRenderReflectiveShadowMap, TEXT("/Engine/Private/ShadowDepthVertexShader.usf"), TEXT("Main"),       SF_Vertex); \
-	IMPLEMENT_MATERIAL_SHADER_TYPE(template<>, TShadowDepthHS##ShaderMode##bRenderReflectiveShadowMap, TEXT("/Engine/Private/ShadowDepthVertexShader.usf"), TEXT("MainHull"),   SF_Hull  ); \
-	IMPLEMENT_MATERIAL_SHADER_TYPE(template<>, TShadowDepthDS##ShaderMode##bRenderReflectiveShadowMap, TEXT("/Engine/Private/ShadowDepthVertexShader.usf"), TEXT("MainDomain"), SF_Domain);
+#define IMPLEMENT_SHADOW_DEPTH_SHADERMODE_SHADERS(ShaderMode) \
+	typedef TShadowDepthVS<ShaderMode, false> TShadowDepthVS##ShaderMode; \
+	typedef TShadowDepthHS<ShaderMode       > TShadowDepthHS##ShaderMode; \
+	typedef TShadowDepthDS<ShaderMode       > TShadowDepthDS##ShaderMode; \
+	IMPLEMENT_MATERIAL_SHADER_TYPE(template<>, TShadowDepthVS##ShaderMode, TEXT("/Engine/Private/ShadowDepthVertexShader.usf"), TEXT("Main"),       SF_Vertex); \
+	IMPLEMENT_MATERIAL_SHADER_TYPE(template<>, TShadowDepthHS##ShaderMode, TEXT("/Engine/Private/ShadowDepthVertexShader.usf"), TEXT("MainHull"),   SF_Hull  ); \
+	IMPLEMENT_MATERIAL_SHADER_TYPE(template<>, TShadowDepthDS##ShaderMode, TEXT("/Engine/Private/ShadowDepthVertexShader.usf"), TEXT("MainDomain"), SF_Domain);
 
-IMPLEMENT_SHADOW_DEPTH_SHADERMODE_SHADERS(VertexShadowDepth_PerspectiveCorrect, true);
-IMPLEMENT_SHADOW_DEPTH_SHADERMODE_SHADERS(VertexShadowDepth_PerspectiveCorrect, false);
-IMPLEMENT_SHADOW_DEPTH_SHADERMODE_SHADERS(VertexShadowDepth_OutputDepth, true);
-IMPLEMENT_SHADOW_DEPTH_SHADERMODE_SHADERS(VertexShadowDepth_OutputDepth, false);
-IMPLEMENT_SHADOW_DEPTH_SHADERMODE_SHADERS(VertexShadowDepth_OnePassPointLight, false);
+IMPLEMENT_SHADOW_DEPTH_SHADERMODE_SHADERS(VertexShadowDepth_PerspectiveCorrect);
+IMPLEMENT_SHADOW_DEPTH_SHADERMODE_SHADERS(VertexShadowDepth_OutputDepth);
+IMPLEMENT_SHADOW_DEPTH_SHADERMODE_SHADERS(VertexShadowDepth_OnePassPointLight);
 
 // Position only vertex shaders.
-typedef TShadowDepthVS<VertexShadowDepth_PerspectiveCorrect, false, true> TShadowDepthVSVertexShadowDepth_PerspectiveCorrectPositionOnly;
-typedef TShadowDepthVS<VertexShadowDepth_OutputDepth,        false, true> TShadowDepthVSVertexShadowDepth_OutputDepthPositionOnly;
-typedef TShadowDepthVS<VertexShadowDepth_OnePassPointLight,  false, true> TShadowDepthVSVertexShadowDepth_OnePassPointLightPositionOnly;
+typedef TShadowDepthVS<VertexShadowDepth_PerspectiveCorrect, true> TShadowDepthVSVertexShadowDepth_PerspectiveCorrectPositionOnly;
+typedef TShadowDepthVS<VertexShadowDepth_OutputDepth,        true> TShadowDepthVSVertexShadowDepth_OutputDepthPositionOnly;
+typedef TShadowDepthVS<VertexShadowDepth_OnePassPointLight,  true> TShadowDepthVSVertexShadowDepth_OnePassPointLightPositionOnly;
 IMPLEMENT_MATERIAL_SHADER_TYPE(template<>, TShadowDepthVSVertexShadowDepth_PerspectiveCorrectPositionOnly, TEXT("/Engine/Private/ShadowDepthVertexShader.usf"), TEXT("PositionOnlyMain"), SF_Vertex);
 IMPLEMENT_MATERIAL_SHADER_TYPE(template<>, TShadowDepthVSVertexShadowDepth_OutputDepthPositionOnly,        TEXT("/Engine/Private/ShadowDepthVertexShader.usf"), TEXT("PositionOnlyMain"), SF_Vertex);
 IMPLEMENT_MATERIAL_SHADER_TYPE(template<>, TShadowDepthVSVertexShadowDepth_OnePassPointLightPositionOnly,  TEXT("/Engine/Private/ShadowDepthVertexShader.usf"), TEXT("PositionOnlyMain"), SF_Vertex);
 
 // One pass point light VS for GS shaders.
-typedef TShadowDepthVS<VertexShadowDepth_OnePassPointLight, false, false, true> TShadowDepthVSForGSVertexShadowDepth_OnePassPointLight;
-typedef TShadowDepthVS<VertexShadowDepth_OnePassPointLight, false, true,  true> TShadowDepthVSForGSVertexShadowDepth_OnePassPointLightPositionOnly;
+typedef TShadowDepthVS<VertexShadowDepth_OnePassPointLight, false, true> TShadowDepthVSForGSVertexShadowDepth_OnePassPointLight;
+typedef TShadowDepthVS<VertexShadowDepth_OnePassPointLight, true,  true> TShadowDepthVSForGSVertexShadowDepth_OnePassPointLightPositionOnly;
 IMPLEMENT_MATERIAL_SHADER_TYPE(template<>, TShadowDepthVSForGSVertexShadowDepth_OnePassPointLight,             TEXT("/Engine/Private/ShadowDepthVertexShader.usf"), TEXT("MainForGS"),             SF_Vertex);
 IMPLEMENT_MATERIAL_SHADER_TYPE(template<>, TShadowDepthVSForGSVertexShadowDepth_OnePassPointLightPositionOnly, TEXT("/Engine/Private/ShadowDepthVertexShader.usf"), TEXT("PositionOnlyMainForGS"), SF_Vertex);
 
-// One pass point light with vertex shader layer output.     bRenderReflectiveShadowMap
-//                                                                | bUsePositionOnlyStream
-//                                                                |      | bIsForGeometryShader
-//                                                                |      |      |
-typedef TShadowDepthVS<VertexShadowDepth_VSLayer, false, false, false> TShadowDepthVSVertexShadowDepth_VSLayer;
-typedef TShadowDepthVS<VertexShadowDepth_VSLayer, false, true,  false> TShadowDepthVSVertexShadowDepth_VSLayerPositionOnly;
-typedef TShadowDepthVS<VertexShadowDepth_VSLayer, false, false, true>  TShadowDepthVSVertexShadowDepth_VSLayerGS; // not used
-typedef TShadowDepthVS<VertexShadowDepth_VSLayer, false, true,  true>  TShadowDepthVSVertexShadowDepth_VSLayerGSPositionOnly; // not used
+// One pass point light with vertex shader layer output.
+//                                                       bUsePositionOnlyStream
+//                                                            | bIsForGeometryShader
+//                                                            |      |
+typedef TShadowDepthVS<VertexShadowDepth_VSLayer, false, false> TShadowDepthVSVertexShadowDepth_VSLayer;
+typedef TShadowDepthVS<VertexShadowDepth_VSLayer, true,  false> TShadowDepthVSVertexShadowDepth_VSLayerPositionOnly;
+typedef TShadowDepthVS<VertexShadowDepth_VSLayer, false, true>  TShadowDepthVSVertexShadowDepth_VSLayerGS; // not used
+typedef TShadowDepthVS<VertexShadowDepth_VSLayer, true,  true>  TShadowDepthVSVertexShadowDepth_VSLayerGSPositionOnly; // not used
 IMPLEMENT_MATERIAL_SHADER_TYPE(template<>, TShadowDepthVSVertexShadowDepth_VSLayer,               TEXT("/Engine/Private/ShadowDepthVertexShader.usf"), TEXT("Main"),             SF_Vertex);
 IMPLEMENT_MATERIAL_SHADER_TYPE(template<>, TShadowDepthVSVertexShadowDepth_VSLayerPositionOnly,   TEXT("/Engine/Private/ShadowDepthVertexShader.usf"), TEXT("PositionOnlyMain"), SF_Vertex);
 IMPLEMENT_MATERIAL_SHADER_TYPE(template<>, TShadowDepthVSVertexShadowDepth_VSLayerGS,             TEXT("/Engine/Private/ShadowDepthVertexShader.usf"), TEXT("Main"),             SF_Vertex); // not used
@@ -534,11 +482,6 @@ public:
 	FShadowDepthBasePS(const ShaderMetaType::CompiledShaderInitializerType& Initializer)
 		: FMeshMaterialShader(Initializer)
 	{
-		GvListBuffer.Bind(Initializer.ParameterMap, TEXT("RWGvListBuffer"));
-		GvListHeadBuffer.Bind(Initializer.ParameterMap, TEXT("RWGvListHeadBuffer"));
-		VplListBuffer.Bind(Initializer.ParameterMap, TEXT("RWVplListBuffer"));
-		VplListHeadBuffer.Bind(Initializer.ParameterMap, TEXT("RWVplListHeadBuffer"));
-
 		const ERHIFeatureLevel::Type FeatureLevel = GetMaxSupportedFeatureLevel((EShaderPlatform)Initializer.Target.Platform);
 
 		if (FSceneInterface::GetShadingPath(FeatureLevel) == EShadingPath::Deferred)
@@ -552,14 +495,7 @@ public:
 		}
 	}
 
-	FShadowDepthBasePS() {}
-
-private:
-
-	LAYOUT_FIELD(FRWShaderParameter, GvListBuffer);
-	LAYOUT_FIELD(FRWShaderParameter, GvListHeadBuffer);
-	LAYOUT_FIELD(FRWShaderParameter, VplListBuffer);
-	LAYOUT_FIELD(FRWShaderParameter, VplListHeadBuffer);
+	FShadowDepthBasePS() = default;
 };
 
 enum EShadowDepthPixelShaderMode
@@ -569,7 +505,7 @@ enum EShadowDepthPixelShaderMode
 	PixelShadowDepth_OnePassPointLight
 };
 
-template <EShadowDepthPixelShaderMode ShaderMode, bool bRenderReflectiveShadowMap>
+template <EShadowDepthPixelShaderMode ShaderMode>
 class TShadowDepthPS : public FShadowDepthBasePS
 {
 	DECLARE_SHADER_TYPE(TShadowDepthPS, MeshMaterial);
@@ -589,42 +525,21 @@ public:
 				|| (ShaderMode == PixelShadowDepth_PerspectiveCorrect && Parameters.MaterialParameters.bMaterialMayModifyMeshPosition))
 				&& ShaderMode != PixelShadowDepth_OnePassPointLight
 				// Don't render ShadowDepth for translucent unlit materials
-				&& Parameters.MaterialParameters.bShouldCastDynamicShadows
-				&& !bRenderReflectiveShadowMap;
+				&& Parameters.MaterialParameters.bShouldCastDynamicShadows;
 		}
 
-		if (bRenderReflectiveShadowMap)
-		{
-			static const auto SupportLPV = IConsoleManager::Get().FindTConsoleVariableDataInt(TEXT("r.LightPropagationVolume"));
-			if (SupportLPV && SupportLPV->GetValueOnAnyThread() == 0)
-			{
-				return false;
-			}
-			else
-			{
-				//Note: This logic needs to stay in sync with OverrideWithDefaultMaterialForShadowDepth!
-				// Reflective shadow map shaders must be compiled for every material because they access the material normal
-				return
-					// Only compile one pass point light shaders for feature levels >= SM4
-					(Parameters.MaterialParameters.bShouldCastDynamicShadows || Parameters.MaterialParameters.bShouldInjectEmissiveIntoLPV || Parameters.MaterialParameters.bShouldBlockGI)
-					&& IsFeatureLevelSupported(Platform, ERHIFeatureLevel::SM5);
-			}
-		}
-		else
-		{
-			//Note: This logic needs to stay in sync with OverrideWithDefaultMaterialForShadowDepth!
-			return (Parameters.MaterialParameters.bIsSpecialEngineMaterial
-				// Only compile for masked or lit translucent materials
-				|| !Parameters.MaterialParameters.bWritesEveryPixelShadowPass
-				|| (Parameters.MaterialParameters.bMaterialMayModifyMeshPosition && Parameters.MaterialParameters.bIsUsedWithInstancedStaticMeshes)
-				// Perspective correct rendering needs a pixel shader and WPO materials can't be overridden with default material.
-				|| (ShaderMode == PixelShadowDepth_PerspectiveCorrect && Parameters.MaterialParameters.bMaterialMayModifyMeshPosition))
-				// Only compile one pass point light shaders for feature levels >= SM5
-				&& (ShaderMode != PixelShadowDepth_OnePassPointLight || IsFeatureLevelSupported(Platform, ERHIFeatureLevel::SM5))
-				// Don't render ShadowDepth for translucent unlit materials
-				&& Parameters.MaterialParameters.bShouldCastDynamicShadows
-				&& IsFeatureLevelSupported(Platform, ERHIFeatureLevel::SM5);
-		}
+		//Note: This logic needs to stay in sync with OverrideWithDefaultMaterialForShadowDepth!
+		return (Parameters.MaterialParameters.bIsSpecialEngineMaterial
+			// Only compile for masked or lit translucent materials
+			|| !Parameters.MaterialParameters.bWritesEveryPixelShadowPass
+			|| (Parameters.MaterialParameters.bMaterialMayModifyMeshPosition && Parameters.MaterialParameters.bIsUsedWithInstancedStaticMeshes)
+			// Perspective correct rendering needs a pixel shader and WPO materials can't be overridden with default material.
+			|| (ShaderMode == PixelShadowDepth_PerspectiveCorrect && Parameters.MaterialParameters.bMaterialMayModifyMeshPosition))
+			// Only compile one pass point light shaders for feature levels >= SM5
+			&& (ShaderMode != PixelShadowDepth_OnePassPointLight || IsFeatureLevelSupported(Platform, ERHIFeatureLevel::SM5))
+			// Don't render ShadowDepth for translucent unlit materials
+			&& Parameters.MaterialParameters.bShouldCastDynamicShadows
+			&& IsFeatureLevelSupported(Platform, ERHIFeatureLevel::SM5);
 	}
 
 	static void ModifyCompilationEnvironment(const FMaterialShaderPermutationParameters& Parameters, FShaderCompilerEnvironment& OutEnvironment)
@@ -632,7 +547,6 @@ public:
 		FShadowDepthBasePS::ModifyCompilationEnvironment(Parameters, OutEnvironment);
 		OutEnvironment.SetDefine(TEXT("PERSPECTIVE_CORRECT_DEPTH"), (uint32)(ShaderMode == PixelShadowDepth_PerspectiveCorrect));
 		OutEnvironment.SetDefine(TEXT("ONEPASS_POINTLIGHT_SHADOW"), (uint32)(ShaderMode == PixelShadowDepth_OnePassPointLight));
-		OutEnvironment.SetDefine(TEXT("REFLECTIVE_SHADOW_MAP"), (uint32)bRenderReflectiveShadowMap);
 	}
 
 	TShadowDepthPS()
@@ -646,16 +560,13 @@ public:
 };
 
 // typedef required to get around macro expansion failure due to commas in template argument list for TShadowDepthPixelShader
-#define IMPLEMENT_SHADOWDEPTHPASS_PIXELSHADER_TYPE(ShaderMode, bRenderReflectiveShadowMap) \
-	typedef TShadowDepthPS<ShaderMode, bRenderReflectiveShadowMap> TShadowDepthPS##ShaderMode##bRenderReflectiveShadowMap; \
-	IMPLEMENT_MATERIAL_SHADER_TYPE(template<>,TShadowDepthPS##ShaderMode##bRenderReflectiveShadowMap,TEXT("/Engine/Private/ShadowDepthPixelShader.usf"),TEXT("Main"),SF_Pixel);
+#define IMPLEMENT_SHADOWDEPTHPASS_PIXELSHADER_TYPE(ShaderMode) \
+	typedef TShadowDepthPS<ShaderMode> TShadowDepthPS##ShaderMode; \
+	IMPLEMENT_MATERIAL_SHADER_TYPE(template<>,TShadowDepthPS##ShaderMode,TEXT("/Engine/Private/ShadowDepthPixelShader.usf"),TEXT("Main"),SF_Pixel);
 
-IMPLEMENT_SHADOWDEPTHPASS_PIXELSHADER_TYPE(PixelShadowDepth_NonPerspectiveCorrect, true);
-IMPLEMENT_SHADOWDEPTHPASS_PIXELSHADER_TYPE(PixelShadowDepth_NonPerspectiveCorrect, false);
-IMPLEMENT_SHADOWDEPTHPASS_PIXELSHADER_TYPE(PixelShadowDepth_PerspectiveCorrect, true);
-IMPLEMENT_SHADOWDEPTHPASS_PIXELSHADER_TYPE(PixelShadowDepth_PerspectiveCorrect, false);
-IMPLEMENT_SHADOWDEPTHPASS_PIXELSHADER_TYPE(PixelShadowDepth_OnePassPointLight, true);
-IMPLEMENT_SHADOWDEPTHPASS_PIXELSHADER_TYPE(PixelShadowDepth_OnePassPointLight, false);
+IMPLEMENT_SHADOWDEPTHPASS_PIXELSHADER_TYPE(PixelShadowDepth_NonPerspectiveCorrect);
+IMPLEMENT_SHADOWDEPTHPASS_PIXELSHADER_TYPE(PixelShadowDepth_PerspectiveCorrect);
+IMPLEMENT_SHADOWDEPTHPASS_PIXELSHADER_TYPE(PixelShadowDepth_OnePassPointLight);
 
 /**
 * Overrides a material used for shadow depth rendering with the default material when appropriate.
@@ -665,13 +576,11 @@ IMPLEMENT_SHADOWDEPTHPASS_PIXELSHADER_TYPE(PixelShadowDepth_OnePassPointLight, f
 void OverrideWithDefaultMaterialForShadowDepth(
 	const FMaterialRenderProxy*& InOutMaterialRenderProxy,
 	const FMaterial*& InOutMaterialResource,
-	bool bReflectiveShadowmap,
 	ERHIFeatureLevel::Type InFeatureLevel)
 {
 	// Override with the default material when possible.
 	if (InOutMaterialResource->WritesEveryPixel(true) &&						// Don't override masked materials.
-		!InOutMaterialResource->MaterialModifiesMeshPosition_RenderThread() &&	// Don't override materials using world position offset.
-		!bReflectiveShadowmap)													// Don't override when rendering reflective shadow maps.
+		!InOutMaterialResource->MaterialModifiesMeshPosition_RenderThread())	// Don't override materials using world position offset.
 	{
 		const FMaterialRenderProxy* DefaultProxy = UMaterial::GetDefaultMaterial(MD_Surface)->GetRenderProxy();
 		const FMaterial* DefaultMaterialResource = DefaultProxy->GetMaterial(InFeatureLevel);
@@ -682,7 +591,6 @@ void OverrideWithDefaultMaterialForShadowDepth(
 	}
 }
 
-template <bool bRenderingReflectiveShadowMaps>
 void GetShadowDepthPassShaders(
 	const FMaterial& Material,
 	const FVertexFactory* VertexFactory,
@@ -696,8 +604,6 @@ void GetShadowDepthPassShaders(
 	TShaderRef<FShadowDepthBasePS>& PixelShader,
 	TShaderRef<FOnePassPointShadowDepthGS>& GeometryShader)
 {
-	check(!bOnePassPointLightShadow || !bRenderingReflectiveShadowMaps);
-
 	// Use perspective correct shadow depths for shadow types which typically render low poly meshes into the shadow depth buffer.
 	// Depth will be interpolated to the pixel shader and written out, which disables HiZ and double speed Z.
 	// Directional light shadows use an ortho projection and can use the non-perspective correct path without artifacts.
@@ -722,11 +628,11 @@ void GetShadowDepthPassShaders(
 		{
 			if (bPositionOnlyVS)
 			{
-				VertexShader = Material.GetShader<TShadowDepthVS<VertexShadowDepth_OnePassPointLight, false, true, true> >(VFType);
+				VertexShader = Material.GetShader<TShadowDepthVS<VertexShadowDepth_OnePassPointLight, true, true> >(VFType);
 			}
 			else
 			{
-				VertexShader = Material.GetShader<TShadowDepthVS<VertexShadowDepth_OnePassPointLight, false, false, true> >(VFType);
+				VertexShader = Material.GetShader<TShadowDepthVS<VertexShadowDepth_OnePassPointLight, false, true> >(VFType);
 			}
 
 			if (RHISupportsGeometryShaders(GShaderPlatformForFeatureLevel[FeatureLevel]))
@@ -737,79 +643,59 @@ void GetShadowDepthPassShaders(
 
 			if (bInitializeTessellationShaders)
 			{
-				HullShader = Material.GetShader<TShadowDepthHS<VertexShadowDepth_OnePassPointLight, false> >(VFType);
-				DomainShader = Material.GetShader<TShadowDepthDS<VertexShadowDepth_OnePassPointLight, false> >(VFType);
+				HullShader = Material.GetShader<TShadowDepthHS<VertexShadowDepth_OnePassPointLight>>(VFType);
+				DomainShader = Material.GetShader<TShadowDepthDS<VertexShadowDepth_OnePassPointLight>>(VFType);
 			}
 		}
 		else
 		{
 			if (bPositionOnlyVS)
 			{
-				VertexShader = Material.GetShader<TShadowDepthVS<VertexShadowDepth_VSLayer, false, true, false> >(VFType);
+				VertexShader = Material.GetShader<TShadowDepthVS<VertexShadowDepth_VSLayer, true, false> >(VFType);
 			}
 			else
 			{
-				VertexShader = Material.GetShader<TShadowDepthVS<VertexShadowDepth_VSLayer, false, false, false> >(VFType);
+				VertexShader = Material.GetShader<TShadowDepthVS<VertexShadowDepth_VSLayer, false, false> >(VFType);
 			}
 		}
 	}
 	else if (bUsePerspectiveCorrectShadowDepths)
 	{
-		if (bRenderingReflectiveShadowMaps)
+		if (bPositionOnlyVS)
 		{
-			VertexShader = Material.GetShader<TShadowDepthVS<VertexShadowDepth_PerspectiveCorrect, true, false> >(VFType);
+			VertexShader = Material.GetShader<TShadowDepthVS<VertexShadowDepth_PerspectiveCorrect, true> >(VFType);
 		}
 		else
 		{
-			if (bPositionOnlyVS)
-			{
-				VertexShader = Material.GetShader<TShadowDepthVS<VertexShadowDepth_PerspectiveCorrect, false, true> >(VFType);
-			}
-			else
-			{
-				VertexShader = Material.GetShader<TShadowDepthVS<VertexShadowDepth_PerspectiveCorrect, false, false> >(VFType);
-			}
+			VertexShader = Material.GetShader<TShadowDepthVS<VertexShadowDepth_PerspectiveCorrect, false> >(VFType);
 		}
 
 		if (bInitializeTessellationShaders)
 		{
-			HullShader = Material.GetShader<TShadowDepthHS<VertexShadowDepth_PerspectiveCorrect, bRenderingReflectiveShadowMaps> >(VFType);
-			DomainShader = Material.GetShader<TShadowDepthDS<VertexShadowDepth_PerspectiveCorrect, bRenderingReflectiveShadowMaps> >(VFType);
+			HullShader = Material.GetShader<TShadowDepthHS<VertexShadowDepth_PerspectiveCorrect> >(VFType);
+			DomainShader = Material.GetShader<TShadowDepthDS<VertexShadowDepth_PerspectiveCorrect> >(VFType);
 		}
 	}
 	else
 	{
-		if (bRenderingReflectiveShadowMaps)
+		if (bPositionOnlyVS)
 		{
-			VertexShader = Material.GetShader<TShadowDepthVS<VertexShadowDepth_OutputDepth, true, false> >(VFType);
-
-			if (bInitializeTessellationShaders)
-			{
-				HullShader = Material.GetShader<TShadowDepthHS<VertexShadowDepth_OutputDepth, true> >(VFType);
-				DomainShader = Material.GetShader<TShadowDepthDS<VertexShadowDepth_OutputDepth, true> >(VFType);
-			}
+			VertexShader = Material.GetShader<TShadowDepthVS<VertexShadowDepth_OutputDepth, true> >(VFType);
 		}
 		else
 		{
-			if (bPositionOnlyVS)
-			{
-				VertexShader = Material.GetShader<TShadowDepthVS<VertexShadowDepth_OutputDepth, false, true> >(VFType);
-			}
-			else
-			{
-				VertexShader = Material.GetShader<TShadowDepthVS<VertexShadowDepth_OutputDepth, false, false> >(VFType);
-			}
+			VertexShader = Material.GetShader<TShadowDepthVS<VertexShadowDepth_OutputDepth, false> >(VFType);
+		}
 
-			if (bInitializeTessellationShaders)
-			{
-				HullShader = Material.GetShader<TShadowDepthHS<VertexShadowDepth_OutputDepth, false> >(VFType);
-				DomainShader = Material.GetShader<TShadowDepthDS<VertexShadowDepth_OutputDepth, false> >(VFType);
-			}
+		if (bInitializeTessellationShaders)
+		{
+			HullShader = Material.GetShader<TShadowDepthHS<VertexShadowDepth_OutputDepth> >(VFType);
+			DomainShader = Material.GetShader<TShadowDepthDS<VertexShadowDepth_OutputDepth> >(VFType);
 		}
 	}
 
 	// Pixel shaders
-	if (Material.WritesEveryPixel(true) && !bUsePerspectiveCorrectShadowDepths && !bRenderingReflectiveShadowMaps && VertexFactory->SupportsNullPixelShader())
+	if (Material.WritesEveryPixel(true) && !bUsePerspectiveCorrectShadowDepths && VertexFactory->SupportsNullPixelShader())
 	{
 		// No pixel shader necessary.
 		PixelShader.Reset();
@@ -818,15 +704,15 @@ void GetShadowDepthPassShaders(
 	{
 		if (bUsePerspectiveCorrectShadowDepths)
 		{
-			PixelShader = Material.GetShader<TShadowDepthPS<PixelShadowDepth_PerspectiveCorrect, bRenderingReflectiveShadowMaps> >(VFType, false);
+			PixelShader = Material.GetShader<TShadowDepthPS<PixelShadowDepth_PerspectiveCorrect> >(VFType, false);
 		}
 		else if (bOnePassPointLightShadow)
 		{
-			PixelShader = Material.GetShader<TShadowDepthPS<PixelShadowDepth_OnePassPointLight, false> >(VFType, false);
+			PixelShader = Material.GetShader<TShadowDepthPS<PixelShadowDepth_OnePassPointLight> >(VFType, false);
 		}
 		else
 		{
-			PixelShader = Material.GetShader<TShadowDepthPS<PixelShadowDepth_NonPerspectiveCorrect, bRenderingReflectiveShadowMaps> >(VFType, false);
+			PixelShader = Material.GetShader<TShadowDepthPS<PixelShadowDepth_NonPerspectiveCorrect> >(VFType, false);
 		}
 	}
 }
@@ -835,11 +721,11 @@ void GetShadowDepthPassShaders(
 FProjectedShadowInfo
 -----------------------------------------------------------------------------*/
 
-static void CheckShadowDepthMaterials(const FMaterialRenderProxy* InRenderProxy, const FMaterial* InMaterial, bool bReflectiveShadowmap, ERHIFeatureLevel::Type InFeatureLevel)
+static void CheckShadowDepthMaterials(const FMaterialRenderProxy* InRenderProxy, const FMaterial* InMaterial, ERHIFeatureLevel::Type InFeatureLevel)
 {
 	const FMaterialRenderProxy* RenderProxy = InRenderProxy;
 	const FMaterial* Material = InMaterial;
-	OverrideWithDefaultMaterialForShadowDepth(RenderProxy, Material, bReflectiveShadowmap, InFeatureLevel);
+	OverrideWithDefaultMaterialForShadowDepth(RenderProxy, Material, InFeatureLevel);
 	check(RenderProxy == InRenderProxy);
 	check(Material == InMaterial);
 }
@@ -862,22 +748,10 @@ void FProjectedShadowInfo::ClearDepth(FRHICommandList& RHICmdList, class FSceneR
 	// Translucent shadows use draw call clear
 	check(!bTranslucentShadow);
 
-	if (bReflectiveShadowmap)
-	{
-		// Clear color and depth targets			
-		bClearColor = true;
-		Colors[0] = FLinearColor(0, 0, 1, 0);
-		Colors[1] = FLinearColor(0, 0, 0, 0);
-
-		NumClearColors = FMath::Min(2, NumColorTextures);
-	}
-	else
-	{
-		// Clear depth only.
-		bClearColor = false;
-		Colors[0] = FLinearColor::White;
-		NumClearColors = FMath::Min(1, NumColorTextures);
-	}
+	// Clear depth only.
+	bClearColor = false;
+	Colors[0] = FLinearColor::White;
+	NumClearColors = FMath::Min(1, NumColorTextures);
 
 	if (bPerformClear)
 	{
@@ -908,18 +782,10 @@ void FProjectedShadowInfo::SetStateForView(FRHICommandList& RHICmdList) const
 	);
 }
 
-void SetStateForShadowDepth(bool bReflectiveShadowmap, bool bOnePassPointLightShadow, FMeshPassProcessorRenderState& DrawRenderState)
+void SetStateForShadowDepth(bool bOnePassPointLightShadow, FMeshPassProcessorRenderState& DrawRenderState)
 {
-	if (bReflectiveShadowmap && !bOnePassPointLightShadow)
-	{
-		// Enable color writes to the reflective shadow map targets with opaque blending
-		DrawRenderState.SetBlendState(TStaticBlendStateWriteMask<CW_RGBA, CW_RGBA>::GetRHI());
-	}
-	else
-	{
-		// Disable color writes
-		DrawRenderState.SetBlendState(TStaticBlendState<CW_NONE>::GetRHI());
-	}
+	// Disable color writes
+	DrawRenderState.SetBlendState(TStaticBlendState<CW_NONE>::GetRHI());
 
 	if (bOnePassPointLightShadow)
 	{
@@ -1184,15 +1050,15 @@ void FProjectedShadowInfo::CopyCachedShadowMap(FRHICommandList& RHICmdList, cons
 }
 
 
-void FProjectedShadowInfo::SetupShadowUniformBuffers(FRHICommandListImmediate& RHICmdList, FScene* Scene, FLightPropagationVolume* LPV)
+void FProjectedShadowInfo::SetupShadowUniformBuffers(FRHICommandListImmediate& RHICmdList, FScene* Scene)
 {
 	const ERHIFeatureLevel::Type FeatureLevel = ShadowDepthView->FeatureLevel;
 	if (FSceneInterface::GetShadingPath(FeatureLevel) == EShadingPath::Deferred)
 	{
 		FShadowDepthPassUniformParameters ShadowDepthPassParameters;
-		SetupShadowDepthPassUniformBuffer(this, RHICmdList, *ShadowDepthView, ShadowDepthPassParameters, LPV);
+		SetupShadowDepthPassUniformBuffer(this, RHICmdList, *ShadowDepthView, ShadowDepthPassParameters);
 
-		if (IsWholeSceneDirectionalShadow() && !bReflectiveShadowmap)
+		if (IsWholeSceneDirectionalShadow())
 		{
 			check(GetShadowDepthType() == CSMShadowDepthType);
 			Scene->UniformBuffers.CSMShadowDepthPassUniformBuffer.UpdateUniformBufferImmediate(ShadowDepthPassParameters);
@@ -1261,7 +1127,7 @@ void FProjectedShadowInfo::RenderDepthInner(FRHICommandListImmediate& RHICmdList
 	}
 
 	FMeshPassProcessorRenderState DrawRenderState(*ShadowDepthView, PassUniformBuffer);
-	SetStateForShadowDepth(bReflectiveShadowmap, bOnePassPointLightShadow, DrawRenderState);
+	SetStateForShadowDepth(bOnePassPointLightShadow, DrawRenderState);
 	SetStateForView(RHICmdList);
 
 	if (CacheMode == SDCM_MovablePrimitivesOnly)
@@ -2410,117 +2276,16 @@ void FSceneRenderer::RenderShadowDepthMaps(FRDGBuilder& GraphBuilder)
 			}
 			RHICmdList.EndRenderPass();
 
-			RHICmdList.Transition(FRHITransitionInfo(ColorTarget0.TargetableTexture, ERHIAccess::Unknown, ERHIAccess::SRVMask));
-			RHICmdList.Transition(FRHITransitionInfo(ColorTarget1.TargetableTexture, ERHIAccess::Unknown, ERHIAccess::SRVMask));
-		}
-
-		// Get a copy of LpvWriteUniformBufferParams for parallel RSM draw-call submission
-		{
-			for (int32 ViewIdx = 0; ViewIdx < Views.Num(); ++ViewIdx)
-			{
-				FViewInfo& View = Views[ViewIdx];
-				FSceneViewState* ViewState = View.ViewState;
-
-				if (ViewState)
-				{
-					FLightPropagationVolume* Lpv = ViewState->GetLightPropagationVolume(FeatureLevel);
-
-					if (Lpv)
-					{
-						Lpv->SetRsmUniformBuffer();
-					}
-				}
-			}
+			RHICmdList.Transition({
+				FRHITransitionInfo(ColorTarget0.TargetableTexture, ERHIAccess::Unknown, ERHIAccess::SRVMask),
+				FRHITransitionInfo(ColorTarget1.TargetableTexture, ERHIAccess::Unknown, ERHIAccess::SRVMask)
+			});
 		}
 
-		for (int32 AtlasIndex = 0; AtlasIndex < SortedShadowsForShadowDepthPass.RSMAtlases.Num(); AtlasIndex++)
-		{
-			checkSlow(RHICmdList.IsOutsideRenderPass());
-
-			const FSortedShadowMapAtlas& ShadowMapAtlas = SortedShadowsForShadowDepthPass.RSMAtlases[AtlasIndex];
-			FSceneRenderTargetItem ColorTarget0 = ShadowMapAtlas.RenderTargets.ColorTargets[0]->GetRenderTargetItem();
-			FSceneRenderTargetItem ColorTarget1 = ShadowMapAtlas.RenderTargets.ColorTargets[1]->GetRenderTargetItem();
-			FSceneRenderTargetItem DepthTarget = ShadowMapAtlas.RenderTargets.DepthTarget->GetRenderTargetItem();
-			FIntPoint TargetSize = ShadowMapAtlas.RenderTargets.DepthTarget->GetDesc().Extent;
-
-			SCOPED_DRAW_EVENTF(RHICmdList, EventShadowDepths, TEXT("RSM%u %ux%u"), AtlasIndex, TargetSize.X, TargetSize.Y);
-
-			for (int32 ShadowIndex = 0; ShadowIndex < ShadowMapAtlas.Shadows.Num(); ShadowIndex++)
-			{
-				FProjectedShadowInfo* ProjectedShadowInfo = ShadowMapAtlas.Shadows[ShadowIndex];
-				SCOPED_GPU_MASK(RHICmdList, GetGPUMaskForShadow(ProjectedShadowInfo));
-
-				const bool bDoParallelDispatch = RHICmdList.IsImmediate() &&  // translucent shadows are draw on the render thread, using a recursive cmdlist (which is not immediate)
-					GRHICommandList.UseParallelAlgorithms() && CVarParallelShadows.GetValueOnRenderThread() &&
-					(ProjectedShadowInfo->IsWholeSceneDirectionalShadow() || CVarParallelShadowsNonWholeScene.GetValueOnRenderThread());
-
-				FSceneViewState* ViewState = (FSceneViewState*)ProjectedShadowInfo->DependentView->State;
-				FLightPropagationVolume* LightPropagationVolume = ViewState->GetLightPropagationVolume(FeatureLevel);
-
-				ProjectedShadowInfo->SetupShadowUniformBuffers(RHICmdList, Scene, LightPropagationVolume);
-
-				auto BeginShadowRenderPass = [this, LightPropagationVolume, ProjectedShadowInfo, &ColorTarget0, &ColorTarget1, &DepthTarget](FRHICommandList& InRHICmdList, bool bPerformClear)
-				{
-					FRHITexture* RenderTargets[2];
-					RenderTargets[0] = ColorTarget0.TargetableTexture;
-					RenderTargets[1] = ColorTarget1.TargetableTexture;
-
-					// Hook up the geometry volume UAVs
-					FRHITransitionInfo UAVTransitions[] = {
-						FRHITransitionInfo(LightPropagationVolume->GetGvListBufferUav(), ERHIAccess::Unknown, ERHIAccess::UAVMask),
-						FRHITransitionInfo(LightPropagationVolume->GetGvListHeadBufferUav(), ERHIAccess::Unknown, ERHIAccess::UAVMask),
-						FRHITransitionInfo(LightPropagationVolume->GetVplListBufferUav(), ERHIAccess::Unknown, ERHIAccess::UAVMask),
-						FRHITransitionInfo(LightPropagationVolume->GetVplListHeadBufferUav(), ERHIAccess::Unknown, ERHIAccess::UAVMask)
-					};
-
-					FRHIRenderPassInfo RPInfo(UE_ARRAY_COUNT(RenderTargets), RenderTargets, ERenderTargetActions::Load_Store);
-					RPInfo.DepthStencilRenderTarget.Action = EDepthStencilTargetActions::LoadDepthNotStencil_StoreDepthNotStencil;
-					RPInfo.DepthStencilRenderTarget.DepthStencilTarget = DepthTarget.TargetableTexture;
-					RPInfo.DepthStencilRenderTarget.ExclusiveDepthStencil = FExclusiveDepthStencil::DepthWrite_StencilNop;
-
-					InRHICmdList.Transition(MakeArrayView(UAVTransitions, UE_ARRAY_COUNT(UAVTransitions)));
-					TransitionRenderPassTargets(InRHICmdList, RPInfo);
-					InRHICmdList.BeginRenderPass(RPInfo, TEXT("ShadowAtlas"));
-
-					ProjectedShadowInfo->ClearDepth(InRHICmdList, this, UE_ARRAY_COUNT(RenderTargets), RenderTargets, DepthTarget.TargetableTexture, bPerformClear);
-				};
-
-				{
-					SCOPED_DRAW_EVENT(RHICmdList, Clear);
-					BeginShadowRenderPass(RHICmdList, true);
-				}
-
-				// In parallel mode the first renderpass is just the clear.
-				if (bDoParallelDispatch)
-				{
-					RHICmdList.EndRenderPass();
-				}
-
-				ProjectedShadowInfo->RenderDepth(RHICmdList, this, BeginShadowRenderPass, bDoParallelDispatch);
-
-				if (!bDoParallelDispatch)
-				{
-					RHICmdList.EndRenderPass();
-				}
-				{
-					// Resolve the shadow depth z surface.
-					RHICmdList.CopyToResolveTarget(DepthTarget.TargetableTexture, DepthTarget.ShaderResourceTexture, FResolveParams());
-					RHICmdList.CopyToResolveTarget(ColorTarget0.TargetableTexture, ColorTarget0.ShaderResourceTexture, FResolveParams());
-					RHICmdList.CopyToResolveTarget(ColorTarget1.TargetableTexture, ColorTarget1.ShaderResourceTexture, FResolveParams());
-
-					FRHITransitionInfo UavsToReadable[2];
-					UavsToReadable[0] = FRHITransitionInfo(LightPropagationVolume->GetGvListBufferUav(), ERHIAccess::Unknown, ERHIAccess::SRVMask);
-					UavsToReadable[1] = FRHITransitionInfo(LightPropagationVolume->GetGvListHeadBufferUav(), ERHIAccess::Unknown, ERHIAccess::SRVMask);
-					RHICmdList.Transition(MakeArrayView(UavsToReadable, UE_ARRAY_COUNT(UavsToReadable)));
-				}
-				checkSlow(RHICmdList.IsOutsideRenderPass());
-			}
-		}
 		bShadowDepthRenderCompleted = true;
 	});
 }
 
-template<bool bRenderReflectiveShadowMap>
 void FShadowDepthPassMeshProcessor::Process(
 	const FMeshBatch& RESTRICT MeshBatch,
 	uint64 BatchElementMask,
@@ -2540,13 +2305,12 @@ void FShadowDepthPassMeshProcessor::Process(
 		FShadowDepthBasePS,
 		FOnePassPointShadowDepthGS> ShadowDepthPassShaders;
 
-	const bool bUsePositionOnlyVS = 
-		!bRenderReflectiveShadowMap
-		&& VertexFactory->SupportsPositionAndNormalOnlyStream()
+	const bool bUsePositionOnlyVS =
+		   VertexFactory->SupportsPositionAndNormalOnlyStream()
 		&& MaterialResource.WritesEveryPixel(true)
 		&& !MaterialResource.MaterialModifiesMeshPosition_RenderThread();
 
-	GetShadowDepthPassShaders<bRenderReflectiveShadowMap>(
+	GetShadowDepthPassShaders(
 		MaterialResource,
 		VertexFactory,
 		FeatureLevel,
@@ -2595,7 +2359,6 @@ void FShadowDepthPassMeshProcessor::AddMeshBatch(const FMeshBatch& RESTRICT Mesh
 
 		const FMaterialRenderProxy& MaterialRenderProxy = FallbackMaterialRenderProxyPtr ? *FallbackMaterialRenderProxyPtr : *MeshBatch.MaterialRenderProxy;
 		const EBlendMode BlendMode = Material.GetBlendMode();
-		const bool bReflectiveShadowmap = ShadowDepthType.bReflectiveShadowmap && !ShadowDepthType.bOnePassPointLightShadow;
 		const bool bShouldCastShadow = Material.ShouldCastDynamicShadows();
 
 		const FMeshDrawingPolicyOverrideSettings OverrideSettings = ComputeMeshOverrideSettings(MeshBatch);
@@ -2607,37 +2370,28 @@ void FShadowDepthPassMeshProcessor::AddMeshBatch(const FMeshBatch& RESTRICT Mesh
 			const ERasterizerCullMode MeshCullMode = ComputeMeshCullMode(MeshBatch, Material, OverrideSettings);
 
 			const bool bTwoSided = Material.IsTwoSided() || PrimitiveSceneProxy->CastsShadowAsTwoSided();
-			// @TODO: only render directional light shadows as two sided, and only when blocking is enabled (required by geometry volume injection)
-			const bool bEffectivelyTwoSided = ShadowDepthType.bReflectiveShadowmap ? true : bTwoSided;
 			// Invert culling order when mobile HDR == false.
 			auto ShaderPlatform = GShaderPlatformForFeatureLevel[FeatureLevel];
 			static auto* MobileHDRCvar = IConsoleManager::Get().FindTConsoleVariableDataInt(TEXT("r.MobileHDR"));
 			check(MobileHDRCvar);
 			const bool bPlatformReversesCulling = (RHINeedsToSwitchVerticalAxis(ShaderPlatform) && MobileHDRCvar->GetValueOnAnyThread() == 0);
 
-			const bool bRenderSceneTwoSided = bEffectivelyTwoSided;
+			const bool bRenderSceneTwoSided = bTwoSided;
 			const bool bReverseCullMode = XOR(bPlatformReversesCulling, ShadowDepthType.bOnePassPointLightShadow);
 
 			FinalCullMode = bRenderSceneTwoSided ? CM_None : bReverseCullMode ? InverseCullMode(MeshCullMode) : MeshCullMode;
 		}
 
-		if ((bShouldCastShadow || (bReflectiveShadowmap && (Material.ShouldInjectEmissiveIntoLPV() || Material.ShouldBlockGI())))
+		if (bShouldCastShadow
 			&& ShouldIncludeDomainInMeshPass(Material.GetMaterialDomain())
 			&& ShouldIncludeMaterialInDefaultOpaquePass(Material))
 		{
 			const FMaterialRenderProxy* EffectiveMaterialRenderProxy = &MaterialRenderProxy;
 			const FMaterial* EffectiveMaterial = &Material;
 
-			OverrideWithDefaultMaterialForShadowDepth(EffectiveMaterialRenderProxy, EffectiveMaterial, ShadowDepthType.bReflectiveShadowmap, FeatureLevel);
+			OverrideWithDefaultMaterialForShadowDepth(EffectiveMaterialRenderProxy, EffectiveMaterial, FeatureLevel);
 
-			if (ShadowDepthType.bReflectiveShadowmap)
-			{
-				Process<true>(MeshBatch, BatchElementMask, StaticMeshId, PrimitiveSceneProxy, *EffectiveMaterialRenderProxy, *EffectiveMaterial, MeshFillMode, FinalCullMode);
-			}
-			else
-			{
-				Process<false>(MeshBatch, BatchElementMask, StaticMeshId, PrimitiveSceneProxy, *EffectiveMaterialRenderProxy, *EffectiveMaterial, MeshFillMode, FinalCullMode);
-			}
+			Process(MeshBatch, BatchElementMask, StaticMeshId, PrimitiveSceneProxy, *EffectiveMaterialRenderProxy, *EffectiveMaterial, MeshFillMode, FinalCullMode);
 		}
 	}
 }
@@ -2653,10 +2407,10 @@ FShadowDepthPassMeshProcessor::FShadowDepthPassMeshProcessor(
 	, PassDrawRenderState(FMeshPassProcessorRenderState(InViewUniformBuffer, InPassUniformBuffer))
 	, ShadowDepthType(InShadowDepthType)
 {
-	SetStateForShadowDepth(ShadowDepthType.bReflectiveShadowmap, ShadowDepthType.bOnePassPointLightShadow, PassDrawRenderState);
+	SetStateForShadowDepth(ShadowDepthType.bOnePassPointLightShadow, PassDrawRenderState);
 }
 
-FShadowDepthType CSMShadowDepthType(true, false, false);
+FShadowDepthType CSMShadowDepthType(true, false);
 
 FMeshPassProcessor* CreateCSMShadowDepthPassProcessor(const FScene* Scene, const FSceneView* InViewIfDynamicMeshCommand, FMeshPassDrawListContext* InDrawListContext)
 {
diff --git a/Engine/Source/Runtime/Renderer/Private/ShadowRendering.cpp b/Engine/Source/Runtime/Renderer/Private/ShadowRendering.cpp
--- a/Engine/Source/Runtime/Renderer/Private/ShadowRendering.cpp
+++ b/Engine/Source/Runtime/Renderer/Private/ShadowRendering.cpp
@@ -9,7 +9,6 @@
 #include "DepthRendering.h"
 #include "SceneRendering.h"
 #include "DeferredShadingRenderer.h"
-#include "LightPropagationVolume.h"
 #include "ScenePrivate.h"
 #include "PipelineStateCache.h"
 #include "ClearQuad.h"
@@ -1583,43 +1582,6 @@ bool FSceneRenderer::CheckForProjectedShadows( const FLightSceneInfo* LightScene
 	return false;
 }
 
-bool FDeferredShadingSceneRenderer::InjectReflectiveShadowMaps(FRHICommandListImmediate& RHICmdList, const FLightSceneInfo* LightSceneInfo)
-{
-	FVisibleLightInfo& VisibleLightInfo = VisibleLightInfos[LightSceneInfo->Id];
-
-	// Inject the RSM into the LPVs
-	for (int32 ShadowIndex = 0; ShadowIndex < VisibleLightInfo.RSMsToProject.Num(); ShadowIndex++)
-	{
-		FProjectedShadowInfo* ProjectedShadowInfo = VisibleLightInfo.RSMsToProject[ShadowIndex];
-
-		check(ProjectedShadowInfo->bReflectiveShadowmap);
-
-		if (ProjectedShadowInfo->bAllocated && ProjectedShadowInfo->DependentView)
-		{
-			FSceneViewState* ViewState = (FSceneViewState*)ProjectedShadowInfo->DependentView->State;
-
-			FLightPropagationVolume* LightPropagationVolume = ViewState ? ViewState->GetLightPropagationVolume(FeatureLevel) : NULL;
-
-			if (LightPropagationVolume)
-			{
-				if (ProjectedShadowInfo->bWholeSceneShadow)
-				{
-					LightPropagationVolume->InjectDirectionalLightRSM( 
-						RHICmdList, 
-						*ProjectedShadowInfo->DependentView,
-						(const FTexture2DRHIRef&)ProjectedShadowInfo->RenderTargets.ColorTargets[0]->GetRenderTargetItem().ShaderResourceTexture,
-						(const FTexture2DRHIRef&)ProjectedShadowInfo->RenderTargets.ColorTargets[1]->GetRenderTargetItem().ShaderResourceTexture, 
-						(const FTexture2DRHIRef&)ProjectedShadowInfo->RenderTargets.DepthTarget->GetRenderTargetItem().ShaderResourceTexture,
-						*ProjectedShadowInfo, 
-						LightSceneInfo->Proxy->GetColor() );
-				}
-			}
-		}
-	}
-
-	return true;
-}
-
 void FSceneRenderer::RenderShadowProjections(
 	FRHICommandListImmediate& RHICmdList,
 	const FLightSceneProxy* LightSceneProxy,
diff --git a/Engine/Source/Runtime/Renderer/Private/ShadowRendering.h b/Engine/Source/Runtime/Renderer/Private/ShadowRendering.h
--- a/Engine/Source/Runtime/Renderer/Private/ShadowRendering.h
+++ b/Engine/Source/Runtime/Renderer/Private/ShadowRendering.h
@@ -27,7 +27,6 @@
 #include "SceneRenderTargetParameters.h"
 #include "ShaderParameterUtils.h"
 #include "LightRendering.h"
-#include "LightPropagationVolume.h"
 #include "HairStrands/HairStrandsRendering.h"
 
 ENGINE_API IPooledRenderTarget* GetSubsufaceProfileTexture_RT(FRHICommandListImmediate& RHICmdList);
@@ -54,7 +53,6 @@ class FShadowDepthBasePS;
 void OverrideWithDefaultMaterialForShadowDepth(
 	const FMaterialRenderProxy*& InOutMaterialRenderProxy, 
 	const FMaterial*& InOutMaterialResource,
-	bool bReflectiveShadowmap,
 	ERHIFeatureLevel::Type InFeatureLevel
 	);
 
@@ -75,19 +73,16 @@ class FShadowDepthType
 public:
 	bool bDirectionalLight;
 	bool bOnePassPointLightShadow;
-	bool bReflectiveShadowmap;
 
-	FShadowDepthType(bool bInDirectionalLight, bool bInOnePassPointLightShadow, bool bInReflectiveShadowmap) 
+	FShadowDepthType(bool bInDirectionalLight, bool bInOnePassPointLightShadow) 
 	: bDirectionalLight(bInDirectionalLight)
 	, bOnePassPointLightShadow(bInOnePassPointLightShadow)
-	, bReflectiveShadowmap(bInReflectiveShadowmap)
 	{}
 
 	inline bool operator==(const FShadowDepthType& rhs) const
 	{
-		if (bDirectionalLight != rhs.bDirectionalLight || 
-			bOnePassPointLightShadow != rhs.bOnePassPointLightShadow || 
-			bReflectiveShadowmap != rhs.bReflectiveShadowmap) 
+		if (bDirectionalLight != rhs.bDirectionalLight ||
+			bOnePassPointLightShadow != rhs.bOnePassPointLightShadow)
 		{
 			return false;
 		}
@@ -116,7 +111,6 @@ public:
 
 private:
 
-	template<bool bRenderReflectiveShadowMap>
 	void Process(
 		const FMeshBatch& RESTRICT MeshBatch,
 		uint64 BatchElementMask,
@@ -201,18 +195,12 @@ typedef TFunctionRef<void(FRHICommandList& RHICmdList, bool bFirst)> FBeginShado
 
 BEGIN_GLOBAL_SHADER_PARAMETER_STRUCT(FShadowDepthPassUniformParameters,)
 	SHADER_PARAMETER_STRUCT(FSceneTextureUniformParameters, SceneTextures)
-	SHADER_PARAMETER_STRUCT(FLpvWriteUniformBufferParameters, LPV)
 	SHADER_PARAMETER(FMatrix, ProjectionMatrix)
 	SHADER_PARAMETER(FMatrix, ViewMatrix)
 	SHADER_PARAMETER(FVector4, ShadowParams)
 	SHADER_PARAMETER(float, bClampToNearPlane)
 	SHADER_PARAMETER_ARRAY(FMatrix, ShadowViewProjectionMatrices, [6])
 	SHADER_PARAMETER_ARRAY(FMatrix, ShadowViewMatrices, [6])
-
-	SHADER_PARAMETER_UAV(RWStructuredBuffer<VplListEntry>, RWGvListBuffer)
-	SHADER_PARAMETER_UAV(RWByteAddressBuffer, RWGvListHeadBuffer)
-	SHADER_PARAMETER_UAV(RWStructuredBuffer<VplListEntry>, RWVplListBuffer)
-	SHADER_PARAMETER_UAV(RWByteAddressBuffer, RWVplListHeadBuffer)
 END_GLOBAL_SHADER_PARAMETER_STRUCT()
 
 BEGIN_GLOBAL_SHADER_PARAMETER_STRUCT(FMobileShadowDepthPassUniformParameters,)
@@ -335,9 +323,6 @@ public:
 	/** Whether this shadow affects the whole scene or only a group of objects. */
 	uint32 bWholeSceneShadow : 1;
 
-	/** Whether the shadow needs to render reflective shadow maps. */ 
-	uint32 bReflectiveShadowmap : 1; 
-
 	/** Whether this shadow should support casting shadows from translucent surfaces. */
 	uint32 bTranslucentShadow : 1;
 
@@ -424,7 +409,6 @@ public:
 		uint32 InSnapResolutionX,
 		uint32 InSnapResolutionY,
 		uint32 InBorderSize,
-		bool bInReflectiveShadowMap,
 		TSharedPtr<FVirtualShadowMapCacheEntry> VirtualSmCacheEntry = TSharedPtr<FVirtualShadowMapCacheEntry>()
 		);
 
@@ -610,13 +594,13 @@ public:
 
 	FShadowDepthType GetShadowDepthType() const 
 	{
-		return FShadowDepthType(bDirectionalLight, bOnePassPointLightShadow, bReflectiveShadowmap);
+		return FShadowDepthType(bDirectionalLight, bOnePassPointLightShadow);
 	}
 
 	/**
 	* Setup uniformbuffers and update Primitive Shader Data
 	*/
-	void SetupShadowUniformBuffers(FRHICommandListImmediate& RHICmdList, FScene* Scene, FLightPropagationVolume* LPV = nullptr);
+	void SetupShadowUniformBuffers(FRHICommandListImmediate& RHICmdList, FScene* Scene);
 
 	/**
 	* Ensure Cached Shadowmap is in readable state
diff --git a/Engine/Source/Runtime/Renderer/Private/ShadowSetup.cpp b/Engine/Source/Runtime/Renderer/Private/ShadowSetup.cpp
--- a/Engine/Source/Runtime/Renderer/Private/ShadowSetup.cpp
+++ b/Engine/Source/Runtime/Renderer/Private/ShadowSetup.cpp
@@ -25,10 +25,8 @@
 #include "TextureLayout.h"
 #include "SceneRendering.h"
 #include "DynamicPrimitiveDrawing.h"
-#include "LightPropagationVolume.h"
 #include "ScenePrivate.h"
 #include "RendererModule.h"
-#include "LightPropagationVolumeSettings.h"
 #include "CapsuleShadowRendering.h"
 #include "Async/ParallelFor.h"
 #include "VirtualShadowMaps/VirtualShadowMapCacheManager.h"
@@ -42,14 +40,6 @@ static FAutoConsoleVariableRef CVarMinScreenRadiusForShadowCaster(
 	ECVF_Scalability | ECVF_RenderThreadSafe
 	);
 
-static float GMinScreenRadiusForShadowCasterRSM = 0.06f;
-static FAutoConsoleVariableRef CVarMinScreenRadiusForShadowCasterRSM(
-	TEXT("r.Shadow.RadiusThresholdRSM"),
-	GMinScreenRadiusForShadowCasterRSM,
-	TEXT("Cull shadow casters in the RSM if they are too small, values is the minimal screen space bounding sphere radius\n")
-	TEXT("(default 0.06)")
-	);
-
 int32 GCacheWholeSceneShadows = 1;
 FAutoConsoleVariableRef CVarCacheWholeSceneShadows(
 	TEXT("r.Shadow.CacheWholeSceneShadows"),
@@ -522,7 +512,6 @@ FProjectedShadowInfo::FProjectedShadowInfo()
 	, bDirectionalLight(false)
 	, bOnePassPointLightShadow(false)
 	, bWholeSceneShadow(false)
-	, bReflectiveShadowmap(false)
 	, bTranslucentShadow(false)
 	, bRayTracedDistanceField(false)
 	, bCapsuleShadow(false)
@@ -705,7 +694,6 @@ void FProjectedShadowInfo::SetupWholeSceneProjection(
 	uint32 InSnapResolutionX,
 	uint32 InSnapResolutionY,
 	uint32 InBorderSize,
-	bool bInReflectiveShadowMap,
 	TSharedPtr<FVirtualShadowMapCacheEntry> VirtualSmCacheEntry)
 {	
 	LightSceneInfo = InLightSceneInfo;
@@ -721,7 +709,6 @@ void FProjectedShadowInfo::SetupWholeSceneProjection(
 	bWholeSceneShadow = true;
 	bTransmission = InLightSceneInfo->Proxy->Transmission();
 	bHairStrandsDeepShadow = InLightSceneInfo->Proxy->CastsHairStrandsDeepShadow();
-	bReflectiveShadowmap = bInReflectiveShadowMap; 
 	BorderSize = InBorderSize;
 
 	checkf(
@@ -1167,7 +1154,7 @@ bool FProjectedShadowInfo::ShouldDrawStaticMeshes(FViewInfo& InCurrentView, FPri
 		FLODMask ShadowLODToRender = *VisibilePrimitiveLODMask;
 
 		// Use lowest LOD for PreShadow
-		if (bReflectiveShadowmap || (bPreShadow && GPreshadowsForceLowestLOD))
+		if (bPreShadow && GPreshadowsForceLowestLOD)
 		{
 			int8 LODToRenderScan = -MAX_int8;
 			FLODMask LODToRender;
@@ -1292,7 +1279,7 @@ bool FProjectedShadowInfo::ShouldDrawStaticMeshes_AnyThread(
 		FLODMask ShadowLODToRender = *VisibilePrimitiveLODMask;
 
 		// Use lowest LOD for PreShadow
-		if (bReflectiveShadowmap || (bPreShadow && GPreshadowsForceLowestLOD))
+		if (bPreShadow && GPreshadowsForceLowestLOD)
 		{
 			int8 LODToRenderScan = -MAX_int8;
 			FLODMask LODToRender;
@@ -3311,7 +3298,6 @@ void FSceneRenderer::CreateWholeSceneProjectedShadow(
 							VirtualSizeX,
 							VirtualSizeY,
 							0, // no border
-							false,	// no RSM
 							VirtualSmCacheEntry
 						);
 
@@ -3332,8 +3318,7 @@ void FSceneRenderer::CreateWholeSceneProjectedShadow(
 						SizeY,
 						SizeX,
 						SizeY,
-						ShadowBorder,
-						false	// no RSM
+						ShadowBorder
 						);
 
 					ProjectedShadowInfo->CacheMode = CacheMode[CacheModeIndex];
@@ -3422,8 +3407,7 @@ void FSceneRenderer::CreateWholeSceneProjectedShadow(
 								// Snap to original shadow map resolution
 								ProjectedShadowInfo->ResolutionX,
 								ProjectedShadowInfo->ResolutionY,
-								ShadowBorder,
-								false	// no RSM
+								ShadowBorder
 							);
 
 							CompleteProjectedShadowInfo->bCompleteShadowMap = true;
@@ -3565,7 +3549,7 @@ void FSceneRenderer::InitProjectedShadowVisibility(FRHICommandListImmediate& RHI
 					}
 
 					// Draw the shadow frustum.
-					if(bPrimitiveIsShadowRelevant && !bShadowIsOccluded && !ProjectedShadowInfo.bReflectiveShadowmap)  
+					if(bPrimitiveIsShadowRelevant && !bShadowIsOccluded)
 					{
 						bool bDrawPreshadowFrustum = CVarDrawPreshadowFrustum.GetValueOnRenderThread() != 0;
 
@@ -3688,18 +3672,6 @@ void FSceneRenderer::GatherShadowDynamicMeshElements(FGlobalDynamicIndexBuffer&
 		}
 	}
 
-	for (int32 AtlasIndex = 0; AtlasIndex < SortedShadowsForShadowDepthPass.RSMAtlases.Num(); AtlasIndex++)
-	{
-		FSortedShadowMapAtlas& Atlas = SortedShadowsForShadowDepthPass.RSMAtlases[AtlasIndex];
-
-		for (int32 ShadowIndex = 0; ShadowIndex < Atlas.Shadows.Num(); ShadowIndex++)
-		{
-			FProjectedShadowInfo* ProjectedShadowInfo = Atlas.Shadows[ShadowIndex];
-			FVisibleLightInfo& VisibleLightInfo = VisibleLightInfos[ProjectedShadowInfo->GetLightSceneInfo().Id];
-			ProjectedShadowInfo->GatherDynamicMeshElements(*this, VisibleLightInfo, ReusedViewsArray, DynamicIndexBuffer, DynamicVertexBuffer, DynamicReadBuffer);
-		}
-	}
-
 	for (int32 AtlasIndex = 0; AtlasIndex < SortedShadowsForShadowDepthPass.ShadowMapCubemaps.Num(); AtlasIndex++)
 	{
 		FSortedShadowMapAtlas& Atlas = SortedShadowsForShadowDepthPass.ShadowMapCubemaps[AtlasIndex];
@@ -3918,7 +3890,7 @@ struct FGatherShadowPrimitivesPacket
 				&& ProjectedShadowInfo->CascadeSettings.ShadowBoundsAccurate.IntersectBox(PrimitiveBounds.Origin, PrimitiveBounds.BoxExtent))
 			{
 				// Distance culling for RSMs
-				const float MinScreenRadiusForShadowCaster = ProjectedShadowInfo->bReflectiveShadowmap ? GMinScreenRadiusForShadowCasterRSM : GMinScreenRadiusForShadowCaster;
+				const float MinScreenRadiusForShadowCaster = GMinScreenRadiusForShadowCaster;
 
 				bool bScreenSpaceSizeCulled = false;
 				check(ProjectedShadowInfo->DependentView);
@@ -3934,10 +3906,8 @@ struct FGatherShadowPrimitivesPacket
 					// Include all primitives for movable lights, but only statically shadowed primitives from a light with static shadowing,
 					// Since lights with static shadowing still create per-object shadows for primitives without static shadowing.
 					&& (!LightProxy.HasStaticLighting() || (!LightSceneInfo.IsPrecomputedLightingValid() || LightProxy.UseCSMForDynamicObjects()))
-					// Only render primitives into a reflective shadowmap that are supposed to affect indirect lighting
-					&& !(ProjectedShadowInfo->bReflectiveShadowmap && !PrimitiveProxy->AffectsDynamicIndirectLighting())
 					// Exclude primitives that will create their own per-object shadow, except when rendering RSMs
-					&& (!DoesPrimitiveCastInsetShadow(PrimitiveSceneInfo, PrimitiveProxy) || ProjectedShadowInfo->bReflectiveShadowmap)
+					&& (!DoesPrimitiveCastInsetShadow(PrimitiveSceneInfo, PrimitiveProxy))
 					// Exclude primitives that will create a per-object shadow from a stationary light
 					&& !ShouldCreateObjectShadowForStationaryLight(&LightSceneInfo, PrimitiveProxy, true)
 					// Only render shadows from objects that use static lighting during a reflection capture, since the reflection capture doesn't update at runtime
@@ -4230,8 +4200,7 @@ void FSceneRenderer::AddViewDependentWholeSceneShadowsForView(
 						ShadowBufferResolution.Y,
 						ShadowBufferResolution.X,
 						ShadowBufferResolution.Y,
-						ShadowBorder,
-						false	// no RSM
+						ShadowBorder
 					);
 					ProjectedShadowInfo->FadeAlphas = FadeAlphas;
 
@@ -4267,8 +4236,7 @@ void FSceneRenderer::AddViewDependentWholeSceneShadowsForView(
 								// Snap consistently with the original resolution shadow map
 								ProjectedShadowInfo->ResolutionX,
 								ProjectedShadowInfo->ResolutionY,
-								ShadowBorder,
-								false	// no RSM
+								ShadowBorder
 								);
 
 							CompleteProjectedShadowInfo->bCompleteShadowMap = true;
@@ -4306,59 +4274,6 @@ void FSceneRenderer::AddViewDependentWholeSceneShadowsForView(
 				VisibleLightInfo.VirtualShadowMapClipmaps.Add(VirtualShadowMapClipmap);
 				SortedShadowsForShadowDepthPass.VirtualShadowMapClipmaps.Add(VirtualShadowMapClipmap);
 			}
-
-			FSceneViewState* ViewState = (FSceneViewState*)View.State;
-			if (ViewState)
-			{
-				FLightPropagationVolume* LightPropagationVolume = ViewState->GetLightPropagationVolume(View.GetFeatureLevel());
-				
-				float LPVIntensity = 0.f;
-
-				if (LightPropagationVolume && LightPropagationVolume->bEnabled)
-				{
-					const FLightPropagationVolumeSettings& LPVSettings = View.FinalPostProcessSettings.BlendableManager.GetSingleFinalDataConst<FLightPropagationVolumeSettings>();
-					LPVIntensity = LPVSettings.LPVIntensity;
-				}
-
-				if (LPVIntensity > 0)
-				{
-					// Generate the RSM shadow info
-					FWholeSceneProjectedShadowInitializer ProjectedShadowInitializer;
-					FLightPropagationVolume& Lpv = *LightPropagationVolume;
-
-					if (LightSceneInfo.Proxy->GetViewDependentRsmWholeSceneProjectedShadowInitializer(View, Lpv.GetBoundingBox(), ProjectedShadowInitializer))
-					{
-						// moved out from the FProjectedShadowInfo constructor
-						ProjectedShadowInitializer.CascadeSettings.ShadowSplitIndex = 0;
-
-						const int32 ShadowBufferResolution = SceneContext_ConstantsOnly.GetReflectiveShadowMapResolution();
-
-						// Create the projected shadow info.
-						FProjectedShadowInfo* ProjectedShadowInfo = new(FMemStack::Get(), 1, 16) FProjectedShadowInfo;
-
-						ProjectedShadowInfo->SetupWholeSceneProjection(
-							&LightSceneInfo,
-							&View,
-							ProjectedShadowInitializer,
-							ShadowBufferResolution,
-							ShadowBufferResolution,
-							ShadowBufferResolution,
-							ShadowBufferResolution,
-							0,
-							true);		// RSM
-
-						VisibleLightInfo.MemStackProjectedShadows.Add(ProjectedShadowInfo);
-						VisibleLightInfo.AllProjectedShadows.Add(ProjectedShadowInfo);
-						ShadowInfos.Add(ProjectedShadowInfo); // or separate list?
-
-						// Ray traced shadows use the GPU managed distance field object buffers, no CPU culling needed
-						if (!ProjectedShadowInfo->bRayTracedDistanceField)
-						{
-							ShadowInfosThatNeedCulling.Add(ProjectedShadowInfo);
-						}
-					}
-				}
-			}
 		}
 	}
 }
@@ -4440,11 +4355,7 @@ void FSceneRenderer::AllocateShadowDepthTargets(FRHICommandListImmediate& RHICmd
 			if (bShadowIsVisible)
 			{
 				// Visible shadow stats
-				if (ProjectedShadowInfo->bReflectiveShadowmap)
-				{
-					INC_DWORD_STAT(STAT_ReflectiveShadowMaps);
-				}
-				else if (ProjectedShadowInfo->bWholeSceneShadow)
+				if (ProjectedShadowInfo->bWholeSceneShadow)
 				{
 					INC_DWORD_STAT(STAT_WholeSceneShadows);
 
@@ -4467,12 +4378,7 @@ void FSceneRenderer::AllocateShadowDepthTargets(FRHICommandListImmediate& RHICmd
 					&& (FeatureLevel >= ERHIFeatureLevel::SM5 || ProjectedShadowInfo->bPerObjectOpaqueShadow);
 
 				if (bNeedsProjection)
-				{
-					if (ProjectedShadowInfo->bReflectiveShadowmap)
-					{
-						VisibleLightInfo.RSMsToProject.Add(ProjectedShadowInfo);
-					}
-					else if (ProjectedShadowInfo->bCapsuleShadow)
+				{if (ProjectedShadowInfo->bCapsuleShadow)
 					{
 						VisibleLightInfo.CapsuleShadowsToProject.Add(ProjectedShadowInfo);
 					}
@@ -4491,11 +4397,6 @@ void FSceneRenderer::AllocateShadowDepthTargets(FRHICommandListImmediate& RHICmd
 						ProjectedShadowInfo->SetupShadowDepthView(RHICmdList, this);
 						SortedShadowsForShadowDepthPass.VirtualShadowMapShadows.Add(ProjectedShadowInfo);
 					}
-					else if (ProjectedShadowInfo->bReflectiveShadowmap)
-					{
-						check(ProjectedShadowInfo->bWholeSceneShadow);
-						RSMShadows.Add(ProjectedShadowInfo);
-					}
 					else if (ProjectedShadowInfo->bPreShadow && ProjectedShadowInfo->bAllocatedInPreshadowCache)
 					{
 						CachedPreShadows.Add(ProjectedShadowInfo);
@@ -4541,7 +4442,6 @@ void FSceneRenderer::AllocateShadowDepthTargets(FRHICommandListImmediate& RHICmd
 
 		// Sort cascades, this is needed for blending between cascades to work, and also to enable fetching the farthest cascade as index [0]
 		VisibleLightInfo.ShadowsToProject.Sort(FCompareFProjectedShadowInfoBySplitIndex());
-		VisibleLightInfo.RSMsToProject.Sort(FCompareFProjectedShadowInfoBySplitIndex());
 		VisibleLightInfo.CompleteProjectedShadows.Sort(FCompareFProjectedShadowInfoBySplitIndex());
 
 		if (!bMobile)
@@ -4602,7 +4502,6 @@ void FSceneRenderer::AllocateShadowDepthTargets(FRHICommandListImmediate& RHICmd
 	}
 
 	AllocateOnePassPointLightDepthTargets(RHICmdList, WholeScenePointShadows);
-	AllocateRSMDepthTargets(RHICmdList, RSMShadows);
 	AllocateCachedSpotlightShadowDepthTargets(RHICmdList, CachedSpotlightShadows);
 	AllocateAtlasedShadowDepthTargets(RHICmdList, Shadows, SortedShadowsForShadowDepthPass.ShadowMapAtlases);
 	AllocateAtlasedShadowDepthTargets(RHICmdList, CompleteShadows,
@@ -4860,68 +4759,6 @@ void FSceneRenderer::AllocateCSMDepthTargets(
 	}
 }
 
-void FSceneRenderer::AllocateRSMDepthTargets(FRHICommandListImmediate& RHICmdList, const TArray<FProjectedShadowInfo*, SceneRenderingAllocator>& RSMShadows)
-{
-	if (RSMShadows.Num() > 0
-		&& FeatureLevel >= ERHIFeatureLevel::SM5)
-	{
-		const int32 MaxTextureSize = 1 << (GMaxTextureMipCount - 1);
-		FTextureLayout ShadowLayout(1, 1, MaxTextureSize, MaxTextureSize, false, ETextureLayoutAspectRatio::None, false);
-
-		for (int32 ShadowIndex = 0; ShadowIndex < RSMShadows.Num(); ShadowIndex++)
-		{
-			FProjectedShadowInfo* ProjectedShadowInfo = RSMShadows[ShadowIndex];
-
-			check(ProjectedShadowInfo->BorderSize == 0);
-			check(!ProjectedShadowInfo->bAllocated);
-
-			if (ShadowLayout.AddElement(
-				ProjectedShadowInfo->X,
-				ProjectedShadowInfo->Y,
-				ProjectedShadowInfo->ResolutionX,
-				ProjectedShadowInfo->ResolutionY)
-				)
-			{
-				ProjectedShadowInfo->bAllocated = true;
-			}
-		}
-
-		SortedShadowsForShadowDepthPass.RSMAtlases.AddDefaulted();
-		FSortedShadowMapAtlas& ShadowMapAtlas = SortedShadowsForShadowDepthPass.RSMAtlases.Last();
-		ShadowMapAtlas.RenderTargets.ColorTargets.Empty(2);
-		ShadowMapAtlas.RenderTargets.ColorTargets.AddDefaulted(2);
-
-		FIntPoint WholeSceneAtlasSize(ShadowLayout.GetSizeX(), ShadowLayout.GetSizeY());
-
-		{
-			FPooledRenderTargetDesc Desc(FPooledRenderTargetDesc::Create2DDesc(WholeSceneAtlasSize, PF_R8G8B8A8, FClearValueBinding::None, TexCreate_None, TexCreate_RenderTargetable | TexCreate_ShaderResource, false));
-			GRenderTargetPool.FindFreeElement(RHICmdList, Desc, ShadowMapAtlas.RenderTargets.ColorTargets[0], TEXT("RSMNormal"), ERenderTargetTransience::NonTransient );
-		}
-
-		{
-			FPooledRenderTargetDesc Desc(FPooledRenderTargetDesc::Create2DDesc(WholeSceneAtlasSize, PF_FloatR11G11B10, FClearValueBinding::None, TexCreate_None, TexCreate_RenderTargetable | TexCreate_ShaderResource, false));
-			GRenderTargetPool.FindFreeElement(RHICmdList, Desc, ShadowMapAtlas.RenderTargets.ColorTargets[1], TEXT("RSMDiffuse"), ERenderTargetTransience::NonTransient);
-		}
-
-		{
-			FPooledRenderTargetDesc Desc(FPooledRenderTargetDesc::Create2DDesc(WholeSceneAtlasSize, PF_DepthStencil, FClearValueBinding::None, TexCreate_None, TexCreate_DepthStencilTargetable | TexCreate_ShaderResource, false));
-			GRenderTargetPool.FindFreeElement(RHICmdList, Desc, ShadowMapAtlas.RenderTargets.DepthTarget, TEXT("RSMDepth"), ERenderTargetTransience::NonTransient);
-		}
-
-		for (int32 ShadowIndex = 0; ShadowIndex < RSMShadows.Num(); ShadowIndex++)
-		{
-			FProjectedShadowInfo* ProjectedShadowInfo = RSMShadows[ShadowIndex];
-
-			if (ProjectedShadowInfo->bAllocated)
-			{
-				ProjectedShadowInfo->RenderTargets.CopyReferencesFromRenderTargets(ShadowMapAtlas.RenderTargets);
-				ProjectedShadowInfo->SetupShadowDepthView(RHICmdList, this);
-				ShadowMapAtlas.Shadows.Add(ProjectedShadowInfo);
-			}
-		}
-	}
-}
-
 void FSceneRenderer::AllocateOnePassPointLightDepthTargets(FRHICommandListImmediate& RHICmdList, const TArray<FProjectedShadowInfo*, SceneRenderingAllocator>& WholeScenePointShadows)
 {
 	if (FeatureLevel >= ERHIFeatureLevel::SM5)
diff --git a/Engine/Source/Runtime/Renderer/Private/ShadowSetupMobile.cpp b/Engine/Source/Runtime/Renderer/Private/ShadowSetupMobile.cpp
--- a/Engine/Source/Runtime/Renderer/Private/ShadowSetupMobile.cpp
+++ b/Engine/Source/Runtime/Renderer/Private/ShadowSetupMobile.cpp
@@ -108,12 +108,6 @@ static bool MobileDetermineStaticMeshesCSMVisibilityStateInner(
 	if (PrimitiveProxy->WillEverBeLit() && PrimitiveProxy->ShouldReceiveMobileCSMShadows()
 		&& (PrimitiveProxy->GetLightingChannelMask() & LightProxy->GetLightingChannelMask()) != 0)
 	{
-
-		if (ProjectedShadowInfo->bReflectiveShadowmap && !PrimitiveProxy->AffectsDynamicIndirectLighting())
-		{
-			return bFoundCSMReceiver;
-		}
-
 		const FVector LightDirection = LightProxy->GetDirection();
 		const FVector PrimitiveToShadowCenter = ProjectedShadowInfo->ShadowBounds.Center - PrimitiveBounds.Origin;
 		// Project the primitive's bounds origin onto the light vector
@@ -343,8 +337,7 @@ bool BuildSingleCascadeShadowInfo(FViewInfo &View, TArray<FVisibleLightInfo, Sce
 				ShadowBufferResolution.Y,
 				ShadowBufferResolution.X,
 				ShadowBufferResolution.Y,
-				ShadowBorder,
-				false	// no RSM
+				ShadowBorder
 			);
 			bSuccess = true;
 		}
diff --git a/Engine/Source/Runtime/Renderer/Private/TranslucentRendering.cpp b/Engine/Source/Runtime/Renderer/Private/TranslucentRendering.cpp
--- a/Engine/Source/Runtime/Renderer/Private/TranslucentRendering.cpp
+++ b/Engine/Source/Runtime/Renderer/Private/TranslucentRendering.cpp
@@ -564,11 +564,6 @@ bool FSceneRenderer::ShouldRenderTranslucency(ETranslucencyPass::Type Translucen
 	// Change this condition to control where simple elements should be rendered.
 	if (IsMainTranslucencyPass(TranslucencyPass))
 	{
-		if (ViewFamily.EngineShowFlags.VisualizeLPV)
-		{
-			return true;
-		}
-
 		for (const FViewInfo& View : Views)
 		{
 			if (View.bHasTranslucentViewMeshElements || View.SimpleElementCollector.BatchedElements.HasPrimsToDraw())
@@ -908,17 +903,6 @@ static void RenderViewTranslucencyInner(
 				RHICmdList.EndRenderPass();
 			}
 		}
-
-		const FSceneViewState* ViewState = (const FSceneViewState*)View.State;
-		if (ViewState && View.Family->EngineShowFlags.VisualizeLPV)
-		{
-			FLightPropagationVolume* LightPropagationVolume = ViewState->GetLightPropagationVolume(View.GetFeatureLevel());
-
-			if (LightPropagationVolume)
-			{
-				LightPropagationVolume->Visualise(RHICmdList, View);
-			}
-		}
 	}
 }
 
diff --git a/Engine/Source/Runtime/Renderer/Private/VolumetricFogLightFunction.cpp b/Engine/Source/Runtime/Renderer/Private/VolumetricFogLightFunction.cpp
--- a/Engine/Source/Runtime/Renderer/Private/VolumetricFogLightFunction.cpp
+++ b/Engine/Source/Runtime/Renderer/Private/VolumetricFogLightFunction.cpp
@@ -194,8 +194,7 @@ void FDeferredShadingSceneRenderer::RenderLightFunctionForVolumetricFog(
 				LightFunctionResolution.Y,
 				LightFunctionResolution.X,
 				LightFunctionResolution.Y,
-				0,
-				false
+				0
 			);
 		}
 
diff --git a/Engine/Source/Runtime/Renderer/Public/LightPropagationVolumeSettings.h b/Engine/Source/Runtime/Renderer/Public/LightPropagationVolumeSettings.h
deleted file mode 100644
