Removal of the deprecated rendering composition graph from the engine.

diff --git a/Engine/Plugins/Experimental/SimpleHMD/Source/SimpleHMD/Private/SimpleHMD.cpp b/Engine/Plugins/Experimental/SimpleHMD/Source/SimpleHMD/Private/SimpleHMD.cpp
--- a/Engine/Plugins/Experimental/SimpleHMD/Source/SimpleHMD/Private/SimpleHMD.cpp
+++ b/Engine/Plugins/Experimental/SimpleHMD/Source/SimpleHMD/Private/SimpleHMD.cpp
@@ -170,7 +170,7 @@ FQuat FSimpleHMD::GetBaseOrientation() const
 	return FQuat::Identity;
 }
 
-void FSimpleHMD::DrawDistortionMesh_RenderThread(struct FRenderingCompositePassContext& Context, const FIntPoint& TextureSize)
+void FSimpleHMD::DrawDistortionMesh_RenderThread(struct FHeadMountedDisplayPassContext& Context, const FIntPoint& TextureSize)
 {
 	float ClipSpaceQuadZ = 0.0f;
 	FMatrix QuadTexTransform = FMatrix::Identity;
@@ -274,7 +274,7 @@ FMatrix FSimpleHMD::GetStereoProjectionMatrix(const enum EStereoscopicPass Stere
 		* FTranslationMatrix(FVector(PassProjectionOffset,0,0));
 }
 
-void FSimpleHMD::GetEyeRenderParams_RenderThread(const FRenderingCompositePassContext& Context, FVector2D& EyeToSrcUVScaleValue, FVector2D& EyeToSrcUVOffsetValue) const
+void FSimpleHMD::GetEyeRenderParams_RenderThread(const FHeadMountedDisplayPassContext& Context, FVector2D& EyeToSrcUVScaleValue, FVector2D& EyeToSrcUVOffsetValue) const
 {
 	EyeToSrcUVOffsetValue = FVector2D::ZeroVector;
 	EyeToSrcUVScaleValue = FVector2D(1.0f, 1.0f);
diff --git a/Engine/Plugins/Experimental/SimpleHMD/Source/SimpleHMD/Private/SimpleHMD.h b/Engine/Plugins/Experimental/SimpleHMD/Source/SimpleHMD/Private/SimpleHMD.h
--- a/Engine/Plugins/Experimental/SimpleHMD/Source/SimpleHMD/Private/SimpleHMD.h
+++ b/Engine/Plugins/Experimental/SimpleHMD/Source/SimpleHMD/Private/SimpleHMD.h
@@ -68,7 +68,7 @@ public:
 	virtual bool GetHMDMonitorInfo(MonitorInfo&) override;
 	virtual void GetFieldOfView(float& OutHFOVInDegrees, float& OutVFOVInDegrees) const override;
 	virtual bool IsChromaAbCorrectionEnabled() const override;
-	virtual void DrawDistortionMesh_RenderThread(struct FRenderingCompositePassContext& Context, const FIntPoint& TextureSize) override;
+	virtual void DrawDistortionMesh_RenderThread(struct FHeadMountedDisplayPassContext& Context, const FIntPoint& TextureSize) override;
 
 	/** IStereoRendering interface */
 	virtual bool IsStereoEnabled() const override;
@@ -77,7 +77,7 @@ public:
 	virtual void CalculateStereoViewOffset(const enum EStereoscopicPass StereoPassType, FRotator& ViewRotation,
 		const float InWorldToMeters, FVector& ViewLocation) override;
 	virtual FMatrix GetStereoProjectionMatrix(const enum EStereoscopicPass StereoPassType) const override;
-	virtual void GetEyeRenderParams_RenderThread(const struct FRenderingCompositePassContext& Context, FVector2D& EyeToSrcUVScaleValue, FVector2D& EyeToSrcUVOffsetValue) const override;
+	virtual void GetEyeRenderParams_RenderThread(const struct FHeadMountedDisplayPassContext& Context, FVector2D& EyeToSrcUVScaleValue, FVector2D& EyeToSrcUVOffsetValue) const override;
 
 	/** ISceneViewExtension interface */
 	virtual void SetupViewFamily(FSceneViewFamily& InViewFamily) override;
diff --git a/Engine/Plugins/Lumin/MagicLeap/Source/MagicLeap/Private/MagicLeapHMD.cpp b/Engine/Plugins/Lumin/MagicLeap/Source/MagicLeap/Private/MagicLeapHMD.cpp
--- a/Engine/Plugins/Lumin/MagicLeap/Source/MagicLeap/Private/MagicLeapHMD.cpp
+++ b/Engine/Plugins/Lumin/MagicLeap/Source/MagicLeap/Private/MagicLeapHMD.cpp
@@ -2014,7 +2014,7 @@ bool FMagicLeapHMD::GetRelativeEyePose(int32 DeviceId, EStereoscopicPass Eye, FQ
 	return false;
 }
 
-void FMagicLeapHMD::GetEyeRenderParams_RenderThread(const FRenderingCompositePassContext& Context, FVector2D& EyeToSrcUVScaleValue, FVector2D& EyeToSrcUVOffsetValue) const
+void FMagicLeapHMD::GetEyeRenderParams_RenderThread(const FHeadMountedDisplayPassContext& Context, FVector2D& EyeToSrcUVScaleValue, FVector2D& EyeToSrcUVOffsetValue) const
 {
 	check(bDeviceInitialized);
 	check(IsInRenderingThread());
diff --git a/Engine/Plugins/Lumin/MagicLeap/Source/MagicLeap/Private/MagicLeapHMD.h b/Engine/Plugins/Lumin/MagicLeap/Source/MagicLeap/Private/MagicLeapHMD.h
--- a/Engine/Plugins/Lumin/MagicLeap/Source/MagicLeap/Private/MagicLeapHMD.h
+++ b/Engine/Plugins/Lumin/MagicLeap/Source/MagicLeap/Private/MagicLeapHMD.h
@@ -87,7 +87,7 @@ public:
 	virtual bool EnableStereo(bool bStereo = true) override;
 	virtual void AdjustViewRect(EStereoscopicPass StereoPass, int32& X, int32& Y, uint32& SizeX, uint32& SizeY) const override;
 	virtual void InitCanvasFromView(FSceneView* InView, UCanvas* Canvas) override;
-	virtual void GetEyeRenderParams_RenderThread(const FRenderingCompositePassContext& Context, FVector2D& EyeToSrcUVScaleValue, FVector2D& EyeToSrcUVOffsetValue) const override;
+	virtual void GetEyeRenderParams_RenderThread(const FHeadMountedDisplayPassContext& Context, FVector2D& EyeToSrcUVScaleValue, FVector2D& EyeToSrcUVOffsetValue) const override;
 	virtual void CalculateRenderTargetSize(const class FViewport& Viewport, uint32& InOutSizeX, uint32& InOutSizeY) override;
 	virtual bool NeedReAllocateViewportRenderTarget(const FViewport& Viewport) override;
 	virtual bool AllocateRenderTargetTexture(uint32 Index, uint32 SizeX, uint32 SizeY, uint8 Format, uint32 NumMips, ETextureCreateFlags Flags, ETextureCreateFlags TargetableTextureFlags, FTexture2DRHIRef& OutTargetableTexture, FTexture2DRHIRef& OutShaderResourceTexture, uint32 NumSamples = 1) override;
diff --git a/Engine/Plugins/Runtime/AR/AppleAR/AppleARKit/Source/AppleARKit/Private/AppleARKitVideoOverlay.cpp b/Engine/Plugins/Runtime/AR/AppleAR/AppleARKit/Source/AppleARKit/Private/AppleARKitVideoOverlay.cpp
--- a/Engine/Plugins/Runtime/AR/AppleAR/AppleARKit/Source/AppleARKit/Private/AppleARKitVideoOverlay.cpp
+++ b/Engine/Plugins/Runtime/AR/AppleAR/AppleARKit/Source/AppleARKit/Private/AppleARKitVideoOverlay.cpp
@@ -22,7 +22,6 @@
 #include "Containers/DynamicRHIResourceArray.h"
 #include "PostProcess/SceneFilterRendering.h"
 #include "PostProcess/PostProcessMaterial.h"
-#include "PostProcessParameters.h"
 #include "CommonRenderResources.h"
 #include "ARUtilitiesFunctionLibrary.h"
 
diff --git a/Engine/Plugins/Runtime/GoogleVR/GoogleVRHMD/Source/GoogleVRHMD/Private/GoogleVRHMD.cpp b/Engine/Plugins/Runtime/GoogleVR/GoogleVRHMD/Source/GoogleVRHMD/Private/GoogleVRHMD.cpp
--- a/Engine/Plugins/Runtime/GoogleVR/GoogleVRHMD/Source/GoogleVRHMD/Private/GoogleVRHMD.cpp
+++ b/Engine/Plugins/Runtime/GoogleVR/GoogleVRHMD/Source/GoogleVRHMD/Private/GoogleVRHMD.cpp
@@ -1524,7 +1524,7 @@ gvr_rectf FGoogleVRHMD::GetGVREyeFOV(int EyeIndex) const
 }
 #endif
 
-void FGoogleVRHMD::GetEyeRenderParams_RenderThread(const struct FRenderingCompositePassContext& Context, FVector2D& EyeToSrcUVScaleValue, FVector2D& EyeToSrcUVOffsetValue) const
+void FGoogleVRHMD::GetEyeRenderParams_RenderThread(const struct FHeadMountedDisplayPassContext& Context, FVector2D& EyeToSrcUVScaleValue, FVector2D& EyeToSrcUVOffsetValue) const
 {
 	if (Context.View.StereoPass == eSSP_LEFT_EYE)
 	{
diff --git a/Engine/Plugins/Runtime/GoogleVR/GoogleVRHMD/Source/GoogleVRHMD/Private/GoogleVRHMD.h b/Engine/Plugins/Runtime/GoogleVR/GoogleVRHMD/Source/GoogleVRHMD/Private/GoogleVRHMD.h
--- a/Engine/Plugins/Runtime/GoogleVR/GoogleVRHMD/Source/GoogleVRHMD/Private/GoogleVRHMD.h
+++ b/Engine/Plugins/Runtime/GoogleVR/GoogleVRHMD/Source/GoogleVRHMD/Private/GoogleVRHMD.h
@@ -514,7 +514,7 @@ public:
 	/**
 	 * Returns eye render params, used from PostProcessHMD, RenderThread.
 	 */
-	virtual void GetEyeRenderParams_RenderThread(const struct FRenderingCompositePassContext& Context, FVector2D& EyeToSrcUVScaleValue, FVector2D& EyeToSrcUVOffsetValue) const override;
+	virtual void GetEyeRenderParams_RenderThread(const struct FHeadMountedDisplayPassContext& Context, FVector2D& EyeToSrcUVScaleValue, FVector2D& EyeToSrcUVOffsetValue) const override;
 
 	// Renders texture into a backbuffer. Could be empty if no rendertarget texture is used, or if direct-rendering
 	// through RHI bridge is implemented.
@@ -742,7 +742,7 @@ public:
 	// Begin IHeadMountedDisplay Virtual Interface //
 	/////////////////////////////////////////////////
 
-	virtual void DrawDistortionMesh_RenderThread(struct FRenderingCompositePassContext& Context, const FIntPoint& TextureSize) override;
+	virtual void DrawDistortionMesh_RenderThread(struct FHeadMountedDisplayPassContext& Context, const FIntPoint& TextureSize) override;
 
 	virtual float GetPixelDenity() const override { return PixelDensity; }
 	virtual void SetPixelDensity(const float NewDensity) override;
diff --git a/Engine/Plugins/Runtime/GoogleVR/GoogleVRHMD/Source/GoogleVRHMD/Private/GoogleVRHMDRender.cpp b/Engine/Plugins/Runtime/GoogleVR/GoogleVRHMD/Source/GoogleVRHMD/Private/GoogleVRHMDRender.cpp
--- a/Engine/Plugins/Runtime/GoogleVR/GoogleVRHMD/Source/GoogleVRHMD/Private/GoogleVRHMDRender.cpp
+++ b/Engine/Plugins/Runtime/GoogleVR/GoogleVRHMD/Source/GoogleVRHMD/Private/GoogleVRHMDRender.cpp
@@ -102,7 +102,7 @@ void FGoogleVRHMD::GenerateDistortionCorrectionVertexBuffer(EStereoscopicPass Ey
 	RHIUnlockVertexBuffer(DistortionMeshVertices);
 }
 
-void FGoogleVRHMD::DrawDistortionMesh_RenderThread(struct FRenderingCompositePassContext& Context, const FIntPoint& TextureSize)
+void FGoogleVRHMD::DrawDistortionMesh_RenderThread(struct FHeadMountedDisplayPassContext& Context, const FIntPoint& TextureSize)
 {
 	const FViewInfo& View = Context.View;
 	FRHICommandListImmediate& RHICmdList = Context.RHICmdList;
diff --git a/Engine/Plugins/Runtime/Oculus/OculusVR/Source/OculusHMD/Private/OculusHMD.h b/Engine/Plugins/Runtime/Oculus/OculusVR/Source/OculusHMD/Private/OculusHMD.h
--- a/Engine/Plugins/Runtime/Oculus/OculusVR/Source/OculusHMD/Private/OculusHMD.h
+++ b/Engine/Plugins/Runtime/Oculus/OculusVR/Source/OculusHMD/Private/OculusHMD.h
@@ -135,7 +135,7 @@ public:
 	virtual bool HasVisibleAreaMesh() const override;
 	virtual void DrawHiddenAreaMesh_RenderThread(class FRHICommandList& RHICmdList, EStereoscopicPass StereoPass) const override;
 	virtual void DrawVisibleAreaMesh_RenderThread(class FRHICommandList& RHICmdList, EStereoscopicPass StereoPass) const override;
-	//virtual void DrawDistortionMesh_RenderThread(struct FRenderingCompositePassContext& Context, const FIntPoint& TextureSize) override;
+	//virtual void DrawDistortionMesh_RenderThread(struct FHeadMountedDisplayPassContext& Context, const FIntPoint& TextureSize) override;
 	//virtual void UpdateScreenSettings(const FViewport* InViewport) override;
 	//virtual void UpdatePostProcessSettings(FPostProcessSettings*) override;
 	//virtual FTexture* GetDistortionTextureLeft() const override;
@@ -165,7 +165,7 @@ public:
 	virtual void CalculateStereoViewOffset(const enum EStereoscopicPass StereoPassType, FRotator& ViewRotation, const float WorldToMeters, FVector& ViewLocation) override;
 	virtual FMatrix GetStereoProjectionMatrix(const enum EStereoscopicPass StereoPassType) const override;
 	virtual void InitCanvasFromView(class FSceneView* InView, class UCanvas* Canvas) override;
-	//virtual void GetEyeRenderParams_RenderThread(const struct FRenderingCompositePassContext& Context, FVector2D& EyeToSrcUVScaleValue, FVector2D& EyeToSrcUVOffsetValue) const override;
+	//virtual void GetEyeRenderParams_RenderThread(const struct FHeadMountedDisplayPassContext& Context, FVector2D& EyeToSrcUVScaleValue, FVector2D& EyeToSrcUVOffsetValue) const override;
 	virtual void RenderTexture_RenderThread(class FRHICommandListImmediate& RHICmdList, class FRHITexture2D* BackBuffer, class FRHITexture2D* SrcTexture, FVector2D WindowSize) const override;
 	//virtual void SetClippingPlanes(float NCP, float FCP) override;
 	virtual IStereoRenderTargetManager* GetRenderTargetManager() override { return this; }
diff --git a/Engine/Plugins/Runtime/OpenXR/Source/OpenXRHMD/Private/OpenXRHMD.cpp b/Engine/Plugins/Runtime/OpenXR/Source/OpenXRHMD/Private/OpenXRHMD.cpp
--- a/Engine/Plugins/Runtime/OpenXR/Source/OpenXRHMD/Private/OpenXRHMD.cpp
+++ b/Engine/Plugins/Runtime/OpenXR/Source/OpenXRHMD/Private/OpenXRHMD.cpp
@@ -1215,7 +1215,7 @@ FMatrix FOpenXRHMD::GetStereoProjectionMatrix(const enum EStereoscopicPass Stere
 	return Mat;
 }
 
-void FOpenXRHMD::GetEyeRenderParams_RenderThread(const FRenderingCompositePassContext& Context, FVector2D& EyeToSrcUVScaleValue, FVector2D& EyeToSrcUVOffsetValue) const
+void FOpenXRHMD::GetEyeRenderParams_RenderThread(const FHeadMountedDisplayPassContext& Context, FVector2D& EyeToSrcUVScaleValue, FVector2D& EyeToSrcUVOffsetValue) const
 {
 	EyeToSrcUVOffsetValue = FVector2D::ZeroVector;
 	EyeToSrcUVScaleValue = FVector2D(1.0f, 1.0f);
diff --git a/Engine/Plugins/Runtime/OpenXR/Source/OpenXRHMD/Private/OpenXRHMD.h b/Engine/Plugins/Runtime/OpenXR/Source/OpenXRHMD/Private/OpenXRHMD.h
--- a/Engine/Plugins/Runtime/OpenXR/Source/OpenXRHMD/Private/OpenXRHMD.h
+++ b/Engine/Plugins/Runtime/OpenXR/Source/OpenXRHMD/Private/OpenXRHMD.h
@@ -181,7 +181,7 @@ public:
 	virtual uint32 DeviceGetLODViewIndex() const override;
 	
 	virtual FMatrix GetStereoProjectionMatrix(const enum EStereoscopicPass StereoPassType) const override;
-	virtual void GetEyeRenderParams_RenderThread(const struct FRenderingCompositePassContext& Context, FVector2D& EyeToSrcUVScaleValue, FVector2D& EyeToSrcUVOffsetValue) const override;
+	virtual void GetEyeRenderParams_RenderThread(const struct FHeadMountedDisplayPassContext& Context, FVector2D& EyeToSrcUVScaleValue, FVector2D& EyeToSrcUVOffsetValue) const override;
 	virtual IStereoRenderTargetManager* GetRenderTargetManager() override;
 	virtual void RenderTexture_RenderThread(class FRHICommandListImmediate& RHICmdList, class FRHITexture2D* BackBuffer, class FRHITexture2D* SrcTexture, FVector2D WindowSize) const override;
 
diff --git a/Engine/Plugins/Runtime/Steam/SteamVR/Source/SteamVR/Private/SteamVRHMD.cpp b/Engine/Plugins/Runtime/Steam/SteamVR/Source/SteamVR/Private/SteamVRHMD.cpp
--- a/Engine/Plugins/Runtime/Steam/SteamVR/Source/SteamVR/Private/SteamVRHMD.cpp
+++ b/Engine/Plugins/Runtime/Steam/SteamVR/Source/SteamVR/Private/SteamVRHMD.cpp
@@ -1489,7 +1489,7 @@ FMatrix FSteamVRHMD::GetStereoProjectionMatrix(const enum EStereoscopicPass Ster
 	return Mat;
 }
 
-void FSteamVRHMD::GetEyeRenderParams_RenderThread(const FRenderingCompositePassContext& Context, FVector2D& EyeToSrcUVScaleValue, FVector2D& EyeToSrcUVOffsetValue) const
+void FSteamVRHMD::GetEyeRenderParams_RenderThread(const FHeadMountedDisplayPassContext& Context, FVector2D& EyeToSrcUVScaleValue, FVector2D& EyeToSrcUVOffsetValue) const
 {
 	if (Context.View.StereoPass == eSSP_LEFT_EYE)
 	{
diff --git a/Engine/Plugins/Runtime/Steam/SteamVR/Source/SteamVR/Private/SteamVRHMD.h b/Engine/Plugins/Runtime/Steam/SteamVR/Source/SteamVR/Private/SteamVRHMD.h
--- a/Engine/Plugins/Runtime/Steam/SteamVR/Source/SteamVR/Private/SteamVRHMD.h
+++ b/Engine/Plugins/Runtime/Steam/SteamVR/Source/SteamVR/Private/SteamVRHMD.h
@@ -146,7 +146,7 @@ public:
 	virtual bool HasVisibleAreaMesh() const override { return VisibleAreaMeshes[0].IsValid() && VisibleAreaMeshes[1].IsValid(); }
 	virtual void DrawVisibleAreaMesh_RenderThread(FRHICommandList& RHICmdList, EStereoscopicPass StereoPass) const override;
 
-	virtual void DrawDistortionMesh_RenderThread(struct FRenderingCompositePassContext& Context, const FIntPoint& TextureSize) override;
+	virtual void DrawDistortionMesh_RenderThread(struct FHeadMountedDisplayPassContext& Context, const FIntPoint& TextureSize) override;
 
 	virtual void UpdateScreenSettings(const FViewport* InViewport) override {}
 	
@@ -166,7 +166,7 @@ public:
 	virtual void CalculateStereoViewOffset(const EStereoscopicPass StereoPassType, FRotator& ViewRotation, const float MetersToWorld, FVector& ViewLocation) override;
 	virtual FMatrix GetStereoProjectionMatrix(const enum EStereoscopicPass StereoPassType) const override;
 	virtual void RenderTexture_RenderThread(FRHICommandListImmediate& RHICmdList, FRHITexture2D* BackBuffer, FRHITexture2D* SrcTexture, FVector2D WindowSize) const override;
-	virtual void GetEyeRenderParams_RenderThread(const FRenderingCompositePassContext& Context, FVector2D& EyeToSrcUVScaleValue, FVector2D& EyeToSrcUVOffsetValue) const override;
+	virtual void GetEyeRenderParams_RenderThread(const FHeadMountedDisplayPassContext& Context, FVector2D& EyeToSrcUVScaleValue, FVector2D& EyeToSrcUVOffsetValue) const override;
 	virtual IStereoRenderTargetManager* GetRenderTargetManager() override { return this; }
 	virtual IStereoLayers* GetStereoLayers() override;
 
diff --git a/Engine/Plugins/Runtime/Steam/SteamVR/Source/SteamVR/Private/SteamVRRender.cpp b/Engine/Plugins/Runtime/Steam/SteamVR/Source/SteamVR/Private/SteamVRRender.cpp
--- a/Engine/Plugins/Runtime/Steam/SteamVR/Source/SteamVR/Private/SteamVRRender.cpp
+++ b/Engine/Plugins/Runtime/Steam/SteamVR/Source/SteamVR/Private/SteamVRRender.cpp
@@ -35,7 +35,7 @@ static TAutoConsoleVariable<int> CVarEnableDepthSubmission(
 	TEXT("By default, depth is not passed through in SteamVR for devices that support depth. Set this flag to 1 to enable depth submission, 0 to disable."),
 	ECVF_Default);
 
-void FSteamVRHMD::DrawDistortionMesh_RenderThread(struct FRenderingCompositePassContext& Context, const FIntPoint& TextureSize)
+void FSteamVRHMD::DrawDistortionMesh_RenderThread(struct FHeadMountedDisplayPassContext& Context, const FIntPoint& TextureSize)
 {
 	check(0);
 }
diff --git a/Engine/Plugins/Runtime/WindowsMixedReality/Source/WindowsMixedRealityHMD/Private/WindowsMixedRealityHMD.cpp b/Engine/Plugins/Runtime/WindowsMixedReality/Source/WindowsMixedRealityHMD/Private/WindowsMixedRealityHMD.cpp
--- a/Engine/Plugins/Runtime/WindowsMixedReality/Source/WindowsMixedRealityHMD/Private/WindowsMixedRealityHMD.cpp
+++ b/Engine/Plugins/Runtime/WindowsMixedReality/Source/WindowsMixedRealityHMD/Private/WindowsMixedRealityHMD.cpp
@@ -1417,7 +1417,7 @@ namespace WindowsMixedReality
 	RenderTargetDescription wmrRenderTargets[3];
 
 	void FWindowsMixedRealityHMD::GetEyeRenderParams_RenderThread(
-		const FRenderingCompositePassContext& Context,
+		const FHeadMountedDisplayPassContext& Context,
 		FVector2D& EyeToSrcUVScaleValue,
 		FVector2D& EyeToSrcUVOffsetValue) const
 	{
diff --git a/Engine/Plugins/Runtime/WindowsMixedReality/Source/WindowsMixedRealityHMD/Private/WindowsMixedRealityHMD.h b/Engine/Plugins/Runtime/WindowsMixedReality/Source/WindowsMixedRealityHMD/Private/WindowsMixedRealityHMD.h
--- a/Engine/Plugins/Runtime/WindowsMixedReality/Source/WindowsMixedRealityHMD/Private/WindowsMixedRealityHMD.h
+++ b/Engine/Plugins/Runtime/WindowsMixedReality/Source/WindowsMixedRealityHMD/Private/WindowsMixedRealityHMD.h
@@ -291,7 +291,7 @@ namespace WindowsMixedReality
 
 		// Inherited via FXRRenderTargetManager
 		virtual void GetEyeRenderParams_RenderThread(
-			const struct FRenderingCompositePassContext& Context,
+			const struct FHeadMountedDisplayPassContext& Context,
 			FVector2D& EyeToSrcUVScaleValue,
 			FVector2D& EyeToSrcUVOffsetValue) const override;
 		virtual FIntPoint GetIdealRenderTargetSize() const override;
diff --git a/Engine/Source/Runtime/HeadMountedDisplay/Public/IHeadMountedDisplay.h b/Engine/Source/Runtime/HeadMountedDisplay/Public/IHeadMountedDisplay.h
--- a/Engine/Source/Runtime/HeadMountedDisplay/Public/IHeadMountedDisplay.h
+++ b/Engine/Source/Runtime/HeadMountedDisplay/Public/IHeadMountedDisplay.h
@@ -15,8 +15,24 @@ struct FPostProcessSettings;
 struct FWorldContext;
 class UTexture;
 class FSceneViewFamily;
+class FViewInfo;
+class FRHICommandListImmediate;
 class FTexture;
 
+struct FHeadMountedDisplayPassContext
+{
+	FHeadMountedDisplayPassContext(FRHICommandListImmediate& InRHICmdList, const FViewInfo& InView)
+		: RHICmdList(InRHICmdList)
+		, View(InView)
+	{}
+
+	FRHICommandListImmediate& RHICmdList;
+	const FViewInfo& View;
+};
+
+UE_DEPRECATED(5.0, "FRenderingCompositePassContext has been refactored to FHeadMountedDisplayPassContext.")
+typedef FHeadMountedDisplayPassContext FRenderingCompositePassContext;
+
 /**
  * HMD device interface
  */
@@ -88,7 +104,7 @@ public:
 	/**
 	 * Returns eye render params, used from PostProcessHMD, RenderThread.
 	 */
-	virtual void	GetEyeRenderParams_RenderThread(const struct FRenderingCompositePassContext& Context, FVector2D& EyeToSrcUVScaleValue, FVector2D& EyeToSrcUVOffsetValue) const {}
+	virtual void	GetEyeRenderParams_RenderThread(const struct FHeadMountedDisplayPassContext& Context, FVector2D& EyeToSrcUVScaleValue, FVector2D& EyeToSrcUVOffsetValue) const {}
 
 	/**
 	 * Accessors to modify the interpupillary distance (meters)
@@ -200,7 +216,7 @@ public:
 	*/
 	virtual void DrawVisibleAreaMesh_RenderThread(class FRHICommandList& RHICmdList, EStereoscopicPass StereoPass) const {};
 
-	virtual void DrawDistortionMesh_RenderThread(struct FRenderingCompositePassContext& Context, const FIntPoint& TextureSize) {}
+	virtual void DrawDistortionMesh_RenderThread(struct FHeadMountedDisplayPassContext& Context, const FIntPoint& TextureSize) {}
 
 	/**
 	 * This method is able to change screen settings right before any drawing occurs. 
diff --git a/Engine/Source/Runtime/Renderer/Private/CompositionLighting/CompositionLighting.cpp b/Engine/Source/Runtime/Renderer/Private/CompositionLighting/CompositionLighting.cpp
--- a/Engine/Source/Runtime/Renderer/Private/CompositionLighting/CompositionLighting.cpp
+++ b/Engine/Source/Runtime/Renderer/Private/CompositionLighting/CompositionLighting.cpp
@@ -6,8 +6,6 @@
 
 #include "CompositionLighting/CompositionLighting.h"
 #include "ScenePrivate.h"
-#include "PostProcess/RenderingCompositionGraph.h"
-#include "PostProcess/PostProcessInput.h"
 #include "PostProcess/PostProcessing.h"
 #include "CompositionLighting/PostProcessAmbientOcclusion.h"
 #include "CompositionLighting/PostProcessDeferredDecals.h"
diff --git a/Engine/Source/Runtime/Renderer/Private/CompositionLighting/PostProcessAmbientOcclusion.h b/Engine/Source/Runtime/Renderer/Private/CompositionLighting/PostProcessAmbientOcclusion.h
--- a/Engine/Source/Runtime/Renderer/Private/CompositionLighting/PostProcessAmbientOcclusion.h
+++ b/Engine/Source/Runtime/Renderer/Private/CompositionLighting/PostProcessAmbientOcclusion.h
@@ -9,7 +9,6 @@
 #include "CoreMinimal.h"
 #include "UniformBuffer.h"
 #include "RendererInterface.h"
-#include "PostProcess/RenderingCompositionGraph.h"
 
 class FViewInfo;
 
diff --git a/Engine/Source/Runtime/Renderer/Private/CompositionLighting/PostProcessDeferredDecals.h b/Engine/Source/Runtime/Renderer/Private/CompositionLighting/PostProcessDeferredDecals.h
--- a/Engine/Source/Runtime/Renderer/Private/CompositionLighting/PostProcessDeferredDecals.h
+++ b/Engine/Source/Runtime/Renderer/Private/CompositionLighting/PostProcessDeferredDecals.h
@@ -4,7 +4,6 @@
 
 #include "RenderGraph.h"
 #include "DecalRenderingCommon.h"
-#include "PostProcess/RenderingCompositionGraph.h"
 
 class FViewInfo;
 
diff --git a/Engine/Source/Runtime/Renderer/Private/CompositionLighting/PostProcessPassThrough.cpp b/Engine/Source/Runtime/Renderer/Private/CompositionLighting/PostProcessPassThrough.cpp
deleted file mode 100644
diff --git a/Engine/Source/Runtime/Renderer/Private/CompositionLighting/PostProcessPassThrough.h b/Engine/Source/Runtime/Renderer/Private/CompositionLighting/PostProcessPassThrough.h
deleted file mode 100644
diff --git a/Engine/Source/Runtime/Renderer/Private/DebugViewModeRendering.cpp b/Engine/Source/Runtime/Renderer/Private/DebugViewModeRendering.cpp
--- a/Engine/Source/Runtime/Renderer/Private/DebugViewModeRendering.cpp
+++ b/Engine/Source/Runtime/Renderer/Private/DebugViewModeRendering.cpp
@@ -17,12 +17,10 @@ DebugViewModeRendering.cpp: Contains definitions for rendering debug viewmodes.
 #include "PostProcessing.h"
 #include "PostProcess/PostProcessVisualizeComplexity.h"
 #include "PostProcess/PostProcessStreamingAccuracyLegend.h"
-#include "CompositionLighting/PostProcessPassThrough.h"
 #include "PostProcess/PostProcessSelectionOutline.h"
 #include "PostProcess/PostProcessCompositeEditorPrimitives.h"
 #include "PostProcess/PostProcessUpscale.h"
 #include "PostProcess/TemporalAA.h"
-#include "PostProcess/PostProcessInput.h"
 #include "SceneRendering.h"
 #include "DeferredShadingRenderer.h"
 #include "MeshPassProcessor.inl"
diff --git a/Engine/Source/Runtime/Renderer/Private/DistanceFieldLightingPost.cpp b/Engine/Source/Runtime/Renderer/Private/DistanceFieldLightingPost.cpp
--- a/Engine/Source/Runtime/Renderer/Private/DistanceFieldLightingPost.cpp
+++ b/Engine/Source/Runtime/Renderer/Private/DistanceFieldLightingPost.cpp
@@ -494,7 +494,6 @@ void GeometryAwareUpsample(FRDGBuilder& GraphBuilder, const FViewInfo& View, FRD
 		GraphicsPSOInit.BoundShaderState.PixelShaderRHI = PixelShader.GetPixelShader();
 		SetGraphicsPipelineState(RHICmdList, GraphicsPSOInit);
 		PixelShader->SetParameters(RHICmdList, View, DistanceFieldNormal->GetPooledRenderTarget()->GetRenderTargetItem(), BentNormalInterpolation->GetPooledRenderTarget()->GetRenderTargetItem(), Parameters);
-		VertexShader->SetParameters(RHICmdList, View.ViewUniformBuffer);
 
 		DrawRectangle(
 			RHICmdList,
@@ -592,8 +591,6 @@ void UpdateHistory(
 						VelocityTexture ? VelocityTexture->GetPooledRenderTarget() : nullptr,
 						Parameters);
 
-					VertexShader->SetParameters(RHICmdList, View.ViewUniformBuffer);
-
 					DrawRectangle(
 						RHICmdList,
 						0, 0,
@@ -661,8 +658,6 @@ void UpdateHistory(
 						PixelShader->SetParameters(RHICmdList, View, DistanceFieldNormal->GetPooledRenderTarget()->GetRenderTargetItem(), NewBentNormalHistory->GetPooledRenderTarget()->GetRenderTargetItem());
 					}
 
-					VertexShader->SetParameters(RHICmdList, View.ViewUniformBuffer);
-
 					DrawRectangle(
 						RHICmdList,
 						0, 0,
diff --git a/Engine/Source/Runtime/Renderer/Private/DistanceFieldShadowing.cpp b/Engine/Source/Runtime/Renderer/Private/DistanceFieldShadowing.cpp
--- a/Engine/Source/Runtime/Renderer/Private/DistanceFieldShadowing.cpp
+++ b/Engine/Source/Runtime/Renderer/Private/DistanceFieldShadowing.cpp
@@ -1441,7 +1441,6 @@ void FProjectedShadowInfo::RenderRayTracedDistanceFieldProjection(
 				TShaderMapRef<TDistanceFieldShadowingUpsamplePS<false> > PixelShader(View.ShaderMap);
 				GraphicsPSOInit.BoundShaderState.PixelShaderRHI = PixelShader.GetPixelShader();
 				SetGraphicsPipelineState(RHICmdList, GraphicsPSOInit);
-				VertexShader->SetParameters(RHICmdList, View.ViewUniformBuffer);
 				PixelShader->SetParameters(RHICmdList, View, this, ScissorRect, RayTracedShadowsTexture->GetPooledRenderTarget());
 			}
 			else
@@ -1449,7 +1448,6 @@ void FProjectedShadowInfo::RenderRayTracedDistanceFieldProjection(
 				TShaderMapRef<TDistanceFieldShadowingUpsamplePS<true> > PixelShader(View.ShaderMap);
 				GraphicsPSOInit.BoundShaderState.PixelShaderRHI = PixelShader.GetPixelShader();
 				SetGraphicsPipelineState(RHICmdList, GraphicsPSOInit);
-				VertexShader->SetParameters(RHICmdList, View.ViewUniformBuffer);
 				PixelShader->SetParameters(RHICmdList, View, this, ScissorRect, RayTracedShadowsTexture->GetPooledRenderTarget());
 			}
 
diff --git a/Engine/Source/Runtime/Renderer/Private/HairStrands/HairStrandsComposition.cpp b/Engine/Source/Runtime/Renderer/Private/HairStrands/HairStrandsComposition.cpp
--- a/Engine/Source/Runtime/Renderer/Private/HairStrands/HairStrandsComposition.cpp
+++ b/Engine/Source/Runtime/Renderer/Private/HairStrands/HairStrandsComposition.cpp
@@ -10,6 +10,7 @@
 #include "PostProcessing.h"
 #include "HairStrandsRendering.h"
 #include "HairStrandsScatter.h"
+#include "FogRendering.h"
 
 /////////////////////////////////////////////////////////////////////////////////////////
 
@@ -94,7 +95,6 @@ static void AddHairVisibilityComposeSamplePass(
 	FGlobalShaderMap* GlobalShaderMap = View.ShaderMap;
 	const FIntRect Viewport = View.ViewRect;
 	const FIntPoint Resolution = OutColorTexture->Desc.Extent;
-	const FViewInfo* CapturedView = &View;
 
 	ClearUnusedGraphResources(PixelShader, Parameters);
 
@@ -102,7 +102,7 @@ static void AddHairVisibilityComposeSamplePass(
 		RDG_EVENT_NAME("HairStrandsComposeSample"),
 		Parameters,
 		ERDGPassFlags::Raster,
-		[Parameters, VertexShader, PixelShader, Viewport, Resolution, CapturedView](FRHICommandList& RHICmdList)
+		[Parameters, VertexShader, PixelShader, Viewport, Resolution](FRHICommandList& RHICmdList)
 	{
 		FGraphicsPipelineStateInitializer GraphicsPSOInit;
 		RHICmdList.ApplyCachedRenderTargets(GraphicsPSOInit);
@@ -124,7 +124,6 @@ static void AddHairVisibilityComposeSamplePass(
 		GraphicsPSOInit.PrimitiveType = PT_TriangleList;
 		SetGraphicsPipelineState(RHICmdList, GraphicsPSOInit);
 
-		VertexShader->SetParameters(RHICmdList, CapturedView->ViewUniformBuffer);
 		RHICmdList.SetViewport(Viewport.Min.X, Viewport.Min.Y, 0.0f, Viewport.Max.X, Viewport.Max.Y, 1.0f);
 		SetShaderParameters(RHICmdList, PixelShader, PixelShader.GetPixelShader(), *Parameters);
 		DrawRectangle(
@@ -203,7 +202,6 @@ static FRDGTextureRef AddHairDOFDepthPass(
 	FGlobalShaderMap* GlobalShaderMap = View.ShaderMap;
 	const FIntPoint Resolution = OutputResolution;
 	const FIntRect Viewport = View.ViewRect;
-	const FViewInfo* CapturedView = &View;
 
 	ClearUnusedGraphResources(PixelShader, Parameters);
 
@@ -211,7 +209,7 @@ static FRDGTextureRef AddHairDOFDepthPass(
 		RDG_EVENT_NAME("HairStrandsDOFDepth"),
 		Parameters,
 		ERDGPassFlags::Raster,
-		[Parameters, VertexShader, PixelShader, Viewport, Resolution, CapturedView](FRHICommandList& RHICmdList)
+		[Parameters, VertexShader, PixelShader, Viewport, Resolution](FRHICommandList& RHICmdList)
 	{
 		FGraphicsPipelineStateInitializer GraphicsPSOInit;
 		RHICmdList.ApplyCachedRenderTargets(GraphicsPSOInit); 
@@ -224,7 +222,6 @@ static FRDGTextureRef AddHairDOFDepthPass(
 		GraphicsPSOInit.PrimitiveType = PT_TriangleList;
 		SetGraphicsPipelineState(RHICmdList, GraphicsPSOInit);
 
-		VertexShader->SetParameters(RHICmdList, CapturedView->ViewUniformBuffer);
 		RHICmdList.SetViewport(Viewport.Min.X, Viewport.Min.Y, 0.0f, Viewport.Max.X, Viewport.Max.Y, 1.0f);
 		SetShaderParameters(RHICmdList, PixelShader, PixelShader.GetPixelShader(), *Parameters);
 		DrawRectangle(
@@ -309,8 +306,6 @@ static void AddHairVisibilityFastResolveMSAAPass(
 	TShaderMapRef<FHairVisibilityFastResolvePS> PixelShader(View.ShaderMap, PermutationVector);
 	const FGlobalShaderMap* GlobalShaderMap = View.ShaderMap;
 	const FIntRect Viewport = View.ViewRect;
-	
-	const FViewInfo* CapturedView = &View;
 
 	{
 		ClearUnusedGraphResources(PixelShader, Parameters);
@@ -319,7 +314,7 @@ static void AddHairVisibilityFastResolveMSAAPass(
 			RDG_EVENT_NAME("HairStrandsVisibilityMarkTAAFastResolve"),
 			Parameters,
 			ERDGPassFlags::Raster,
-			[Parameters, VertexShader, PixelShader, Viewport, Resolution, CapturedView](FRHICommandList& RHICmdList)
+			[Parameters, VertexShader, PixelShader, Viewport, Resolution](FRHICommandList& RHICmdList)
 		{
 			FGraphicsPipelineStateInitializer GraphicsPSOInit;
 			RHICmdList.ApplyCachedRenderTargets(GraphicsPSOInit);
@@ -336,7 +331,6 @@ static void AddHairVisibilityFastResolveMSAAPass(
 			GraphicsPSOInit.PrimitiveType = PT_TriangleList;
 			SetGraphicsPipelineState(RHICmdList, GraphicsPSOInit);
 
-			VertexShader->SetParameters(RHICmdList, CapturedView->ViewUniformBuffer);
 			RHICmdList.SetStencilRef(STENCIL_TEMPORAL_RESPONSIVE_AA_MASK);
 			RHICmdList.SetViewport(Viewport.Min.X, Viewport.Min.Y, 0.0f, Viewport.Max.X, Viewport.Max.Y, 1.0f);
 			SetShaderParameters(RHICmdList, PixelShader, PixelShader.GetPixelShader(), *Parameters);
@@ -403,8 +397,6 @@ static void AddHairVisibilityFastResolveMaskPass(
 	const FGlobalShaderMap* GlobalShaderMap = View.ShaderMap;
 	const FIntRect Viewport = View.ViewRect;
 
-	const FViewInfo* CapturedView = &View;
-
 	{
 		ClearUnusedGraphResources(PixelShader, Parameters);
 
@@ -412,7 +404,7 @@ static void AddHairVisibilityFastResolveMaskPass(
 			RDG_EVENT_NAME("HairStrandsVisibilityMarkTAAFastResolve"),
 			Parameters,
 			ERDGPassFlags::Raster,
-			[Parameters, VertexShader, PixelShader, Viewport, Resolution, CapturedView](FRHICommandList& RHICmdList)
+			[Parameters, VertexShader, PixelShader, Viewport, Resolution](FRHICommandList& RHICmdList)
 			{
 				FGraphicsPipelineStateInitializer GraphicsPSOInit;
 				RHICmdList.ApplyCachedRenderTargets(GraphicsPSOInit);
@@ -429,7 +421,6 @@ static void AddHairVisibilityFastResolveMaskPass(
 				GraphicsPSOInit.PrimitiveType = PT_TriangleList;
 				SetGraphicsPipelineState(RHICmdList, GraphicsPSOInit);
 
-				VertexShader->SetParameters(RHICmdList, CapturedView->ViewUniformBuffer);
 				RHICmdList.SetStencilRef(STENCIL_TEMPORAL_RESPONSIVE_AA_MASK);
 				RHICmdList.SetViewport(Viewport.Min.X, Viewport.Min.Y, 0.0f, Viewport.Max.X, Viewport.Max.Y, 1.0f);
 				SetShaderParameters(RHICmdList, PixelShader, PixelShader.GetPixelShader(), *Parameters);
@@ -544,14 +535,13 @@ static void AddHairVisibilityGBufferWritePass(
 	const FGlobalShaderMap* GlobalShaderMap = View.ShaderMap;
 	const FIntRect Viewport = View.ViewRect;
 	const FIntPoint Resolution = OutGBufferATexture->Desc.Extent;
-	const FViewInfo* CapturedView = &View;
 	ClearUnusedGraphResources(PixelShader, Parameters);
 
 	GraphBuilder.AddPass(
 		RDG_EVENT_NAME("HairStrandsGBufferWrite"),
 		Parameters,
 		ERDGPassFlags::Raster,
-		[Parameters, VertexShader, PixelShader, Viewport, Resolution, CapturedView, bWriteDepth](FRHICommandList& RHICmdList)
+		[Parameters, VertexShader, PixelShader, Viewport, Resolution, bWriteDepth](FRHICommandList& RHICmdList)
 		{
 			FGraphicsPipelineStateInitializer GraphicsPSOInit;
 			RHICmdList.ApplyCachedRenderTargets(GraphicsPSOInit);
@@ -572,7 +562,6 @@ static void AddHairVisibilityGBufferWritePass(
 			GraphicsPSOInit.PrimitiveType = PT_TriangleList;
 			SetGraphicsPipelineState(RHICmdList, GraphicsPSOInit);
 
-			VertexShader->SetParameters(RHICmdList, CapturedView->ViewUniformBuffer);
 			RHICmdList.SetViewport(Viewport.Min.X, Viewport.Min.Y, 0.0f, Viewport.Max.X, Viewport.Max.Y, 1.0f);
 			SetShaderParameters(RHICmdList, PixelShader, PixelShader.GetPixelShader(), *Parameters);
 			DrawRectangle(
diff --git a/Engine/Source/Runtime/Renderer/Private/HairStrands/HairStrandsDebug.cpp b/Engine/Source/Runtime/Renderer/Private/HairStrands/HairStrandsDebug.cpp
--- a/Engine/Source/Runtime/Renderer/Private/HairStrands/HairStrandsDebug.cpp
+++ b/Engine/Source/Runtime/Renderer/Private/HairStrands/HairStrandsDebug.cpp
@@ -437,7 +437,7 @@ static void AddDebugHairPass(
 		RDG_EVENT_NAME("HairStrandsDebug"),
 		Parameters,
 		ERDGPassFlags::Raster,
-		[Parameters, VertexShader, PixelShader, Viewport, Resolution, View](FRHICommandList& RHICmdList)
+		[Parameters, VertexShader, PixelShader, Viewport, Resolution](FRHICommandList& RHICmdList)
 	{
 		FGraphicsPipelineStateInitializer GraphicsPSOInit;
 		RHICmdList.ApplyCachedRenderTargets(GraphicsPSOInit);
@@ -450,7 +450,6 @@ static void AddDebugHairPass(
 		GraphicsPSOInit.PrimitiveType = PT_TriangleList;
 		SetGraphicsPipelineState(RHICmdList, GraphicsPSOInit);
 
-		VertexShader->SetParameters(RHICmdList, View->ViewUniformBuffer);
 		RHICmdList.SetViewport(Viewport.Min.X, Viewport.Min.Y, 0.0f, Viewport.Max.X, Viewport.Max.Y, 1.0f);
 		SetShaderParameters(RHICmdList, PixelShader, PixelShader.GetPixelShader(), *Parameters);
 
@@ -549,7 +548,7 @@ static void AddDebugDeepShadowTexturePass(
 		ShadowData ? RDG_EVENT_NAME("DebugDeepShadowTexture") : RDG_EVENT_NAME("DebugHairViewRect"),
 		Parameters,
 		ERDGPassFlags::Raster,
-		[Parameters, VertexShader, PixelShader, Viewport, Resolution, View](FRHICommandList& RHICmdList)
+		[Parameters, VertexShader, PixelShader, Viewport, Resolution](FRHICommandList& RHICmdList)
 	{
 		FGraphicsPipelineStateInitializer GraphicsPSOInit;
 		RHICmdList.ApplyCachedRenderTargets(GraphicsPSOInit);
@@ -562,7 +561,6 @@ static void AddDebugDeepShadowTexturePass(
 		GraphicsPSOInit.PrimitiveType = PT_TriangleList;
 		SetGraphicsPipelineState(RHICmdList, GraphicsPSOInit);
 
-		VertexShader->SetParameters(RHICmdList, View->ViewUniformBuffer);
 		RHICmdList.SetViewport(Viewport.Min.X, Viewport.Min.Y, 0.0f, Viewport.Max.X, Viewport.Max.Y, 1.0f);
 		SetShaderParameters(RHICmdList, PixelShader, PixelShader.GetPixelShader(), *Parameters);
 
@@ -773,7 +771,6 @@ static void AddPlotBSDFPass(
 	TShaderMapRef<FHairStrandsPlotBSDFPS> PixelShader(View.ShaderMap);
 	const FGlobalShaderMap* GlobalShaderMap = View.ShaderMap;
 	const FIntRect Viewport = View.ViewRect;
-	const FViewInfo* CapturedView = &View;
 
 	ClearUnusedGraphResources(PixelShader, Parameters);
 
@@ -781,7 +778,7 @@ static void AddPlotBSDFPass(
 		RDG_EVENT_NAME("HairStrandsBsdfPlot"),
 		Parameters,
 		ERDGPassFlags::Raster,
-		[Parameters, VertexShader, PixelShader, Viewport, Resolution, CapturedView](FRHICommandList& RHICmdList)
+		[Parameters, VertexShader, PixelShader, Viewport, Resolution](FRHICommandList& RHICmdList)
 	{
 		FGraphicsPipelineStateInitializer GraphicsPSOInit;
 		RHICmdList.ApplyCachedRenderTargets(GraphicsPSOInit);
@@ -795,7 +792,6 @@ static void AddPlotBSDFPass(
 		GraphicsPSOInit.PrimitiveType = PT_TriangleList;
 		SetGraphicsPipelineState(RHICmdList, GraphicsPSOInit);
 
-		VertexShader->SetParameters(RHICmdList, CapturedView->ViewUniformBuffer);
 		RHICmdList.SetViewport(Viewport.Min.X, Viewport.Min.Y, 0.0f, Viewport.Max.X, Viewport.Max.Y, 1.0f);
 		SetShaderParameters(RHICmdList, PixelShader, PixelShader.GetPixelShader(), *Parameters);
 
@@ -872,7 +868,6 @@ static void AddPlotSamplePass(
 	TShaderMapRef<FHairStrandsPlotSamplePS> PixelShader(View.ShaderMap);
 	const FGlobalShaderMap* GlobalShaderMap = View.ShaderMap;
 	const FIntRect Viewport = View.ViewRect;
-	const FViewInfo* CapturedView = &View;
 
 	ClearUnusedGraphResources(PixelShader, Parameters);
 
@@ -880,7 +875,7 @@ static void AddPlotSamplePass(
 		RDG_EVENT_NAME("HairStrandsSamplePlot"),
 		Parameters,
 		ERDGPassFlags::Raster,
-		[Parameters, VertexShader, PixelShader, Viewport, Resolution, CapturedView](FRHICommandList& RHICmdList)
+		[Parameters, VertexShader, PixelShader, Viewport, Resolution](FRHICommandList& RHICmdList)
 	{
 		FGraphicsPipelineStateInitializer GraphicsPSOInit;
 		RHICmdList.ApplyCachedRenderTargets(GraphicsPSOInit);
@@ -894,7 +889,6 @@ static void AddPlotSamplePass(
 		GraphicsPSOInit.PrimitiveType = PT_TriangleList;
 		SetGraphicsPipelineState(RHICmdList, GraphicsPSOInit);
 
-		VertexShader->SetParameters(RHICmdList, CapturedView->ViewUniformBuffer);
 		RHICmdList.SetViewport(Viewport.Min.X, Viewport.Min.Y, 0.0f, Viewport.Max.X, Viewport.Max.Y, 1.0f);
 		SetShaderParameters(RHICmdList, PixelShader, PixelShader.GetPixelShader(), *Parameters);
 
diff --git a/Engine/Source/Runtime/Renderer/Private/HairStrands/HairStrandsScatter.cpp b/Engine/Source/Runtime/Renderer/Private/HairStrands/HairStrandsScatter.cpp
--- a/Engine/Source/Runtime/Renderer/Private/HairStrands/HairStrandsScatter.cpp
+++ b/Engine/Source/Runtime/Renderer/Private/HairStrands/HairStrandsScatter.cpp
@@ -71,8 +71,6 @@ static FRDGTextureRef AddPreScatterComposePass(
 	const FGlobalShaderMap* GlobalShaderMap = View.ShaderMap;
 	const FIntRect Viewport = View.ViewRect;
 
-	const FViewInfo* CapturedView = &View;
-
 	{
 		ClearUnusedGraphResources(PixelShader, Parameters);
 
@@ -80,7 +78,7 @@ static FRDGTextureRef AddPreScatterComposePass(
 			RDG_EVENT_NAME("HairCompose"),
 			Parameters,
 			ERDGPassFlags::Raster,
-			[Parameters, VertexShader, PixelShader, Viewport, Resolution, CapturedView](FRHICommandList& RHICmdList)
+			[Parameters, VertexShader, PixelShader, Viewport, Resolution](FRHICommandList& RHICmdList)
 		{
 			FGraphicsPipelineStateInitializer GraphicsPSOInit;
 			RHICmdList.ApplyCachedRenderTargets(GraphicsPSOInit);
@@ -93,7 +91,6 @@ static FRDGTextureRef AddPreScatterComposePass(
 			GraphicsPSOInit.PrimitiveType = PT_TriangleList;
 			SetGraphicsPipelineState(RHICmdList, GraphicsPSOInit);
 
-			VertexShader->SetParameters(RHICmdList, CapturedView->ViewUniformBuffer);
 			RHICmdList.SetViewport(Viewport.Min.X, Viewport.Min.Y, 0.0f, Viewport.Max.X, Viewport.Max.Y, 1.0f);
 			SetShaderParameters(RHICmdList, PixelShader, PixelShader.GetPixelShader(), *Parameters);
 			DrawRectangle(
@@ -202,8 +199,6 @@ static FRDGTextureRef AddScatterPass(
 	const FGlobalShaderMap* GlobalShaderMap = View.ShaderMap;
 	const FIntRect Viewport = View.ViewRect;
 
-	const FViewInfo* CapturedView = &View;
-
 	{
 		ClearUnusedGraphResources(PixelShader, Parameters);
 
@@ -211,7 +206,7 @@ static FRDGTextureRef AddScatterPass(
 			RDG_EVENT_NAME("HairScatter"),
 			Parameters,
 			ERDGPassFlags::Raster,
-			[Parameters, VertexShader, PixelShader, Viewport, Resolution, CapturedView](FRHICommandList& RHICmdList)
+			[Parameters, VertexShader, PixelShader, Viewport, Resolution](FRHICommandList& RHICmdList)
 		{
 			FGraphicsPipelineStateInitializer GraphicsPSOInit;
 			RHICmdList.ApplyCachedRenderTargets(GraphicsPSOInit);
@@ -235,7 +230,6 @@ static FRDGTextureRef AddScatterPass(
 			GraphicsPSOInit.PrimitiveType = PT_TriangleList;
 			SetGraphicsPipelineState(RHICmdList, GraphicsPSOInit);
 
-			VertexShader->SetParameters(RHICmdList, CapturedView->ViewUniformBuffer);
 			RHICmdList.SetViewport(Viewport.Min.X, Viewport.Min.Y, 0.0f, Viewport.Max.X, Viewport.Max.Y, 1.0f);
 			SetShaderParameters(RHICmdList, PixelShader, PixelShader.GetPixelShader(), *Parameters);
 			DrawRectangle(
diff --git a/Engine/Source/Runtime/Renderer/Private/HairStrands/HairStrandsTransmittance.cpp b/Engine/Source/Runtime/Renderer/Private/HairStrands/HairStrandsTransmittance.cpp
--- a/Engine/Source/Runtime/Renderer/Private/HairStrands/HairStrandsTransmittance.cpp
+++ b/Engine/Source/Runtime/Renderer/Private/HairStrands/HairStrandsTransmittance.cpp
@@ -16,6 +16,7 @@
 #include "PostProcessing.h"
 #include "GpuDebugRendering.h"
 #include "ShaderPrintParameters.h"
+#include "LightSceneInfo.h"
 
 ///////////////////////////////////////////////////////////////////////////////////////////////////
 
@@ -434,7 +435,6 @@ static void AddDeepShadowOpaqueMaskPass(
 	TShaderMapRef<FDeepShadowMaskPS> PixelShader(View.ShaderMap, PermutationVector);
 	const FGlobalShaderMap* GlobalShaderMap = View.ShaderMap;
 	const FIntRect Viewport = View.ViewRect;
-	const FViewInfo* CapturedView = &View;
 
 	{
 		ClearUnusedGraphResources(PixelShader, Parameters);
@@ -444,7 +444,7 @@ static void AddDeepShadowOpaqueMaskPass(
 			RDG_EVENT_NAME("HairStrandsShadowMask"),
 			Parameters,
 			ERDGPassFlags::Raster,
-			[Parameters, VertexShader, PixelShader, Viewport, Resolution, CapturedView](FRHICommandList& RHICmdList)
+			[Parameters, VertexShader, PixelShader, Viewport, Resolution](FRHICommandList& RHICmdList)
 		{
 			FGraphicsPipelineStateInitializer GraphicsPSOInit;
 			RHICmdList.ApplyCachedRenderTargets(GraphicsPSOInit);
@@ -458,7 +458,6 @@ static void AddDeepShadowOpaqueMaskPass(
 			GraphicsPSOInit.PrimitiveType = PT_TriangleList;
 			SetGraphicsPipelineState(RHICmdList, GraphicsPSOInit);
 
-			VertexShader->SetParameters(RHICmdList, CapturedView->ViewUniformBuffer);
 			RHICmdList.SetViewport(Viewport.Min.X, Viewport.Min.Y, 0.0f, Viewport.Max.X, Viewport.Max.Y, 1.0f);
 			SetShaderParameters(RHICmdList, PixelShader, PixelShader.GetPixelShader(), *Parameters);
 
diff --git a/Engine/Source/Runtime/Renderer/Private/HairStrands/HairStrandsVisibility.cpp b/Engine/Source/Runtime/Renderer/Private/HairStrands/HairStrandsVisibility.cpp
--- a/Engine/Source/Runtime/Renderer/Private/HairStrands/HairStrandsVisibility.cpp
+++ b/Engine/Source/Runtime/Renderer/Private/HairStrands/HairStrandsVisibility.cpp
@@ -13,6 +13,7 @@
 #include "RenderGraphUtils.h"
 #include "PostProcessing.h"
 #include "MeshPassProcessor.inl"
+#include "ScenePrivate.h"
 
 DECLARE_GPU_STAT(HairStrandsVisibility);
 
@@ -1729,7 +1730,7 @@ static void AddClearGraphicPass(
 		Forward<FRDGEventName>(PassName),
 		Parameters,
 		ERDGPassFlags::Raster,
-		[Parameters, VertexShader, PixelShader, Viewport, Resolution, View](FRHICommandList& RHICmdList)
+		[Parameters, VertexShader, PixelShader, Viewport, Resolution](FRHICommandList& RHICmdList)
 	{
 		FGraphicsPipelineStateInitializer GraphicsPSOInit;
 		RHICmdList.ApplyCachedRenderTargets(GraphicsPSOInit);
@@ -1742,7 +1743,6 @@ static void AddClearGraphicPass(
 		GraphicsPSOInit.PrimitiveType = PT_TriangleList;
 		SetGraphicsPipelineState(RHICmdList, GraphicsPSOInit);
 
-		VertexShader->SetParameters(RHICmdList, View->ViewUniformBuffer);
 		RHICmdList.SetViewport(Viewport.Min.X, Viewport.Min.Y, 0.0f, Viewport.Max.X, Viewport.Max.Y, 1.0f);
 		SetShaderParameters(RHICmdList, PixelShader, PixelShader.GetPixelShader(), *Parameters);
 
@@ -2279,7 +2279,6 @@ static FRDGTextureRef AddHairVisibilityFillOpaqueDepth(
 	TShaderMapRef<FHairVisibilityFillOpaqueDepthPS> PixelShader(View.ShaderMap);
 	const FGlobalShaderMap* GlobalShaderMap = View.ShaderMap;
 	const FIntRect Viewport = View.ViewRect;
-	const FViewInfo* CapturedView = &View;
 
 	TArray<FIntRect> MacroGroupRects;
 	if (IsHairStrandsViewRectOptimEnable())
@@ -2301,7 +2300,7 @@ static FRDGTextureRef AddHairVisibilityFillOpaqueDepth(
 			RDG_EVENT_NAME("HairStrandsVisibilityFillOpaqueDepth"),
 			Parameters,
 			ERDGPassFlags::Raster,
-			[Parameters, VertexShader, PixelShader, MacroGroupRects, Viewport, Resolution, CapturedView](FRHICommandList& RHICmdList)
+			[Parameters, VertexShader, PixelShader, MacroGroupRects, Viewport, Resolution](FRHICommandList& RHICmdList)
 		{
 			FGraphicsPipelineStateInitializer GraphicsPSOInit;
 			RHICmdList.ApplyCachedRenderTargets(GraphicsPSOInit);
@@ -2315,7 +2314,6 @@ static FRDGTextureRef AddHairVisibilityFillOpaqueDepth(
 			GraphicsPSOInit.PrimitiveType = PT_TriangleList;
 			SetGraphicsPipelineState(RHICmdList, GraphicsPSOInit);
 
-			VertexShader->SetParameters(RHICmdList, CapturedView->ViewUniformBuffer);
 			SetShaderParameters(RHICmdList, PixelShader, PixelShader.GetPixelShader(), *Parameters);
 
 			for (const FIntRect& ViewRect : MacroGroupRects)
@@ -2688,14 +2686,13 @@ static void AddHairViewTransmittanceDepthPass(
 	const FGlobalShaderMap* GlobalShaderMap = View.ShaderMap;
 	const FIntRect Viewport = View.ViewRect;
 	const FIntPoint Resolution = HairCountTexture->Desc.Extent;
-	const FViewInfo* CapturedView = &View;
 	ClearUnusedGraphResources(PixelShader, Parameters);
 
 	GraphBuilder.AddPass(
 		RDG_EVENT_NAME("HairStrandsViewTransmittanceDepth"),
 		Parameters,
 		ERDGPassFlags::Raster,
-		[Parameters, VertexShader, PixelShader, Viewport, Resolution, CapturedView](FRHICommandList& RHICmdList)
+		[Parameters, VertexShader, PixelShader, Viewport, Resolution](FRHICommandList& RHICmdList)
 	{
 		FGraphicsPipelineStateInitializer GraphicsPSOInit;
 		RHICmdList.ApplyCachedRenderTargets(GraphicsPSOInit);
@@ -2709,7 +2706,6 @@ static void AddHairViewTransmittanceDepthPass(
 		GraphicsPSOInit.PrimitiveType = PT_TriangleList;
 		SetGraphicsPipelineState(RHICmdList, GraphicsPSOInit);
 
-		VertexShader->SetParameters(RHICmdList, CapturedView->ViewUniformBuffer);
 		RHICmdList.SetViewport(Viewport.Min.X, Viewport.Min.Y, 0.0f, Viewport.Max.X, Viewport.Max.Y, 1.0f);
 		SetShaderParameters(RHICmdList, PixelShader, PixelShader.GetPixelShader(), *Parameters);
 		DrawRectangle(
@@ -2790,7 +2786,6 @@ static void AddHairVisibilityColorAndDepthPatchPass(
 	const FGlobalShaderMap* GlobalShaderMap = View.ShaderMap;
 	const FIntRect Viewport = View.ViewRect;
 	const FIntPoint Resolution = OutDepthTexture->Desc.Extent;
-	const FViewInfo* CapturedView = &View;
 
 	{
 		ClearUnusedGraphResources(PixelShader, Parameters);
@@ -2799,7 +2794,7 @@ static void AddHairVisibilityColorAndDepthPatchPass(
 			RDG_EVENT_NAME("HairStrandsVisibilityWriteColorAndDepth"),
 			Parameters,
 			ERDGPassFlags::Raster,
-			[Parameters, VertexShader, PixelShader, Viewport, Resolution, CapturedView](FRHICommandList& RHICmdList)
+			[Parameters, VertexShader, PixelShader, Viewport, Resolution](FRHICommandList& RHICmdList)
 		{
 			FGraphicsPipelineStateInitializer GraphicsPSOInit;
 			RHICmdList.ApplyCachedRenderTargets(GraphicsPSOInit);
@@ -2813,7 +2808,6 @@ static void AddHairVisibilityColorAndDepthPatchPass(
 			GraphicsPSOInit.PrimitiveType = PT_TriangleList;
 			SetGraphicsPipelineState(RHICmdList, GraphicsPSOInit);
 
-			VertexShader->SetParameters(RHICmdList, CapturedView->ViewUniformBuffer);
 			RHICmdList.SetViewport(Viewport.Min.X, Viewport.Min.Y, 0.0f, Viewport.Max.X, Viewport.Max.Y, 1.0f);
 			SetShaderParameters(RHICmdList, PixelShader, PixelShader.GetPixelShader(), *Parameters);
 			DrawRectangle(
diff --git a/Engine/Source/Runtime/Renderer/Private/HairStrands/HairStrandsVoxelization.cpp b/Engine/Source/Runtime/Renderer/Private/HairStrands/HairStrandsVoxelization.cpp
--- a/Engine/Source/Runtime/Renderer/Private/HairStrands/HairStrandsVoxelization.cpp
+++ b/Engine/Source/Runtime/Renderer/Private/HairStrands/HairStrandsVoxelization.cpp
@@ -4,6 +4,7 @@
 #include "HairStrandsRasterCommon.h"
 #include "HairStrandsCluster.h"
 #include "HairStrandsUtils.h"
+#include "HairStrandsLUT.h"
 #include "Shader.h"
 #include "GlobalShader.h"
 #include "ShaderParameters.h"
@@ -11,6 +12,7 @@
 #include "SceneTextureParameters.h"
 #include "RenderGraphUtils.h"
 #include "PostProcessing.h"
+#include "ScenePrivate.h"
 
 static float GStrandHairVoxelizationRasterizationScale = 1.0f;
 static FAutoConsoleVariableRef CVarStrandHairVoxelizationRasterizationScale(TEXT("r.HairStrands.VoxelizationRasterizationScale"), GStrandHairVoxelizationRasterizationScale, TEXT("Rasterization scale to snap strand to pixel for voxelization"));
diff --git a/Engine/Source/Runtime/Renderer/Private/PostProcess/PostProcessBlurGBuffer.h b/Engine/Source/Runtime/Renderer/Private/PostProcess/PostProcessBlurGBuffer.h
deleted file mode 100644
diff --git a/Engine/Source/Runtime/Renderer/Private/PostProcess/PostProcessDeviceEncodingOnly.h b/Engine/Source/Runtime/Renderer/Private/PostProcess/PostProcessDeviceEncodingOnly.h
--- a/Engine/Source/Runtime/Renderer/Private/PostProcess/PostProcessDeviceEncodingOnly.h
+++ b/Engine/Source/Runtime/Renderer/Private/PostProcess/PostProcessDeviceEncodingOnly.h
@@ -2,7 +2,6 @@
 
 #pragma once
 
-#include "PostProcess/RenderingCompositionGraph.h"
 #include "PostProcess/PostProcessEyeAdaptation.h"
 #include "ScreenPass.h"
 #include "OverridePassSequence.h"
diff --git a/Engine/Source/Runtime/Renderer/Private/PostProcess/PostProcessHMD.cpp b/Engine/Source/Runtime/Renderer/Private/PostProcess/PostProcessHMD.cpp
--- a/Engine/Source/Runtime/Renderer/Private/PostProcess/PostProcessHMD.cpp
+++ b/Engine/Source/Runtime/Renderer/Private/PostProcess/PostProcessHMD.cpp
@@ -5,7 +5,6 @@
 #include "Engine/Engine.h"
 #include "IHeadMountedDisplay.h"
 #include "IXRTrackingSystem.h"
-#include "RenderingCompositionGraph.h"
 
 /** The filter vertex declaration resource type. */
 class FDistortionVertexDeclaration : public FRenderResource
@@ -96,7 +95,7 @@ FScreenPassTexture AddDefaultHMDDistortionPass(FRDGBuilder& GraphBuilder, const
 	IHeadMountedDisplay* HMDDevice = GEngine->XRSystem->GetHMDDevice();
 
 	{
-		FRenderingCompositePassContext PassContext(GraphBuilder.RHICmdList, View);
+		FHeadMountedDisplayPassContext PassContext(GraphBuilder.RHICmdList, View);
 		FVector2D EyeToSrcUVScaleValue;
 		FVector2D EyeToSrcUVOffsetValue;
 		HMDDevice->GetEyeRenderParams_RenderThread(PassContext, EyeToSrcUVScaleValue, EyeToSrcUVOffsetValue);
@@ -124,7 +123,7 @@ FScreenPassTexture AddDefaultHMDDistortionPass(FRDGBuilder& GraphBuilder, const
 		SetShaderParameters(RHICmdList, VertexShader, VertexShader.GetVertexShader(), *PassParameters);
 		SetShaderParameters(RHICmdList, PixelShader, PixelShader.GetPixelShader(), *PassParameters);
 
-		FRenderingCompositePassContext PassContext(RHICmdList, View);
+		FHeadMountedDisplayPassContext PassContext(RHICmdList, View);
 		HMDDevice->DrawDistortionMesh_RenderThread(PassContext, PassParameters->InputTexture->Desc.Extent);
 	});
 
diff --git a/Engine/Source/Runtime/Renderer/Private/PostProcess/PostProcessHMD.h b/Engine/Source/Runtime/Renderer/Private/PostProcess/PostProcessHMD.h
--- a/Engine/Source/Runtime/Renderer/Private/PostProcess/PostProcessHMD.h
+++ b/Engine/Source/Runtime/Renderer/Private/PostProcess/PostProcessHMD.h
@@ -3,7 +3,6 @@
 #pragma once
 
 #include "ScreenPass.h"
-#include "RenderingCompositionGraph.h"
 
 /** The vertex data used to filter a texture. */
 struct FDistortionVertex
@@ -27,6 +26,4 @@ struct FHMDDistortionInputs
 
 FScreenPassTexture AddDefaultHMDDistortionPass(FRDGBuilder& GraphBuilder, const FViewInfo& View, const FHMDDistortionInputs& Inputs);
 
-FScreenPassTexture AddHMDDistortionPass(FRDGBuilder& GraphBuilder, const FViewInfo& View, const FHMDDistortionInputs& Inputs);
-
-FRenderingCompositeOutputRef AddHMDDistortionPass(FRenderingCompositionGraph& Graph, FRenderingCompositeOutputRef Input);
\ No newline at end of file
+FScreenPassTexture AddHMDDistortionPass(FRDGBuilder& GraphBuilder, const FViewInfo& View, const FHMDDistortionInputs& Inputs);
\ No newline at end of file
diff --git a/Engine/Source/Runtime/Renderer/Private/PostProcess/PostProcessInput.cpp b/Engine/Source/Runtime/Renderer/Private/PostProcess/PostProcessInput.cpp
deleted file mode 100644
diff --git a/Engine/Source/Runtime/Renderer/Private/PostProcess/PostProcessInput.h b/Engine/Source/Runtime/Renderer/Private/PostProcess/PostProcessInput.h
deleted file mode 100644
diff --git a/Engine/Source/Runtime/Renderer/Private/PostProcess/PostProcessMaterial.cpp b/Engine/Source/Runtime/Renderer/Private/PostProcess/PostProcessMaterial.cpp
--- a/Engine/Source/Runtime/Renderer/Private/PostProcess/PostProcessMaterial.cpp
+++ b/Engine/Source/Runtime/Renderer/Private/PostProcess/PostProcessMaterial.cpp
@@ -20,6 +20,7 @@
 #include "PostProcessMobile.h"
 #include "BufferVisualizationData.h"
 #include "SceneTextureParameters.h"
+#include "SystemTextures.h"
 
 namespace
 {
@@ -795,33 +796,4 @@ FScreenPassTexture AddHighResolutionScreenshotMaskPass(
 	}
 
 	return Output;
-}
-
-void AddHighResScreenshotMask(FPostprocessContext& Context)
-{
-	FRenderingCompositePass* Pass = Context.Graph.RegisterPass(
-		new(FMemStack::Get()) TRCPassForRDG<1, 1>(
-			[](FRenderingCompositePass* InPass, FRenderingCompositePassContext& InContext)
-	{
-		FRDGBuilder GraphBuilder(InContext.RHICmdList);
-
-		FHighResolutionScreenshotMaskInputs PassInputs;
-		PassInputs.SceneColor.Texture = InPass->CreateRDGTextureForRequiredInput(GraphBuilder, ePId_Input0, TEXT("SceneColor"));
-		PassInputs.SceneColor.ViewRect = InContext.SceneColorViewRect;
-
-		if (FRDGTextureRef OverrideOutputTexture = InPass->FindRDGTextureForOutput(GraphBuilder, ePId_Output0, TEXT("FrameBuffer")))
-		{
-			PassInputs.OverrideOutput.Texture = OverrideOutputTexture;
-			PassInputs.OverrideOutput.ViewRect = InContext.GetSceneColorDestRect(InPass);
-			PassInputs.OverrideOutput.LoadAction = InContext.View.IsFirstInFamily() ? ERenderTargetLoadAction::EClear : ERenderTargetLoadAction::ELoad;
-		}
-
-		FScreenPassTexture PassOutput = AddHighResolutionScreenshotMaskPass(GraphBuilder, InContext.View, PassInputs);
-
-		InPass->ExtractRDGTextureForOutput(GraphBuilder, ePId_Output0, PassOutput.Texture);
-
-		GraphBuilder.Execute();
-	}));
-	Pass->SetInput(ePId_Input0, Context.FinalOutput);
-	Context.FinalOutput = Pass;
 }
\ No newline at end of file
diff --git a/Engine/Source/Runtime/Renderer/Private/PostProcess/PostProcessNoiseBlur.cpp b/Engine/Source/Runtime/Renderer/Private/PostProcess/PostProcessNoiseBlur.cpp
deleted file mode 100644
diff --git a/Engine/Source/Runtime/Renderer/Private/PostProcess/PostProcessNoiseBlur.h b/Engine/Source/Runtime/Renderer/Private/PostProcess/PostProcessNoiseBlur.h
deleted file mode 100644
diff --git a/Engine/Source/Runtime/Renderer/Private/PostProcess/PostProcessOutput.cpp b/Engine/Source/Runtime/Renderer/Private/PostProcess/PostProcessOutput.cpp
deleted file mode 100644
diff --git a/Engine/Source/Runtime/Renderer/Private/PostProcess/PostProcessOutput.h b/Engine/Source/Runtime/Renderer/Private/PostProcess/PostProcessOutput.h
deleted file mode 100644
diff --git a/Engine/Source/Runtime/Renderer/Private/PostProcess/PostProcessSubsurface.h b/Engine/Source/Runtime/Renderer/Private/PostProcess/PostProcessSubsurface.h
--- a/Engine/Source/Runtime/Renderer/Private/PostProcess/PostProcessSubsurface.h
+++ b/Engine/Source/Runtime/Renderer/Private/PostProcess/PostProcessSubsurface.h
@@ -4,7 +4,6 @@
 
 #include "ScreenPass.h"
 #include "OverridePassSequence.h"
-#include "RenderingCompositionGraph.h"
 
 // Returns whether subsurface scattering is globally enabled.
 bool IsSubsurfaceEnabled();
diff --git a/Engine/Source/Runtime/Renderer/Private/PostProcess/PostProcessVisualizeLevelInstance.cpp b/Engine/Source/Runtime/Renderer/Private/PostProcess/PostProcessVisualizeLevelInstance.cpp
--- a/Engine/Source/Runtime/Renderer/Private/PostProcess/PostProcessVisualizeLevelInstance.cpp
+++ b/Engine/Source/Runtime/Renderer/Private/PostProcess/PostProcessVisualizeLevelInstance.cpp
@@ -153,41 +153,4 @@ FScreenPassTexture AddVisualizeLevelInstancePass(FRDGBuilder& GraphBuilder, cons
 	return MoveTemp(Output);
 }
 
-FRenderingCompositeOutputRef AddVisualizeLevelInstancePass(FRenderingCompositionGraph& Graph, FRenderingCompositeOutputRef Input, const Nanite::FRasterResults *NaniteRasterResults)
-{
-	FRenderingCompositePass* Pass = Graph.RegisterPass(
-		new(FMemStack::Get()) TRCPassForRDG<1, 1>(
-			[NaniteRasterResults](FRenderingCompositePass* InPass, FRenderingCompositePassContext& InContext)
-	{
-		FRDGBuilder GraphBuilder(InContext.RHICmdList);
-
-		FRDGTextureRef SceneColorTexture = InPass->CreateRDGTextureForRequiredInput(GraphBuilder, ePId_Input0, TEXT("SceneColor"));
-		const FIntRect SceneColorViewRect = InContext.GetSceneColorDestRect(InPass);
-
-		const FSceneRenderTargets& SceneContext = FSceneRenderTargets::Get(GraphBuilder.RHICmdList);
-		FRDGTextureRef SceneDepthTexture = GraphBuilder.RegisterExternalTexture(SceneContext.SceneDepthZ, TEXT("SceneDepthZ"));
-
-		FVisualizeLevelInstanceInputs Inputs;
-		Inputs.SceneColor.Texture = SceneColorTexture;
-		Inputs.SceneColor.ViewRect = SceneColorViewRect;
-		Inputs.SceneDepth.Texture = SceneDepthTexture;
-		Inputs.SceneDepth.ViewRect = InContext.View.ViewRect;
-
-		if (FRDGTextureRef OutputTexture = InPass->FindRDGTextureForOutput(GraphBuilder, ePId_Output0, TEXT("BackBuffer")))
-		{
-			Inputs.OverrideOutput.Texture = OutputTexture;
-			Inputs.OverrideOutput.ViewRect = InContext.GetSceneColorDestRect(InPass->GetOutput(ePId_Output0)->PooledRenderTarget->GetRenderTargetItem());
-			Inputs.OverrideOutput.LoadAction = InContext.View.IsFirstInFamily() ? ERenderTargetLoadAction::EClear : ERenderTargetLoadAction::ELoad;
-		}
-
-		FScreenPassTexture Outputs = AddVisualizeLevelInstancePass(GraphBuilder, InContext.View, Inputs, NaniteRasterResults);
-
-		InPass->ExtractRDGTextureForOutput(GraphBuilder, ePId_Output0, Outputs.Texture);
-
-		GraphBuilder.Execute();
-	}));
-	Pass->SetInput(ePId_Input0, Input);
-	return Pass;
-}
-
 #endif
\ No newline at end of file
diff --git a/Engine/Source/Runtime/Renderer/Private/PostProcess/PostProcessVisualizeLevelInstance.h b/Engine/Source/Runtime/Renderer/Private/PostProcess/PostProcessVisualizeLevelInstance.h
--- a/Engine/Source/Runtime/Renderer/Private/PostProcess/PostProcessVisualizeLevelInstance.h
+++ b/Engine/Source/Runtime/Renderer/Private/PostProcess/PostProcessVisualizeLevelInstance.h
@@ -6,7 +6,6 @@
 
 #include "ScreenPass.h"
 #include "OverridePassSequence.h"
-#include "PostProcess/RenderingCompositionGraph.h"
 
 namespace Nanite
 {
@@ -30,6 +29,4 @@ struct FVisualizeLevelInstanceInputs
 
 FScreenPassTexture AddVisualizeLevelInstancePass(FRDGBuilder& GraphBuilder, const FViewInfo& View, const FVisualizeLevelInstanceInputs& Inputs, const Nanite::FRasterResults *NaniteRasterResults);
 
-FRenderingCompositeOutputRef AddVisualizeLevelInstancePass(FRenderingCompositionGraph& Graph, FRenderingCompositeOutputRef Input, const Nanite::FRasterResults *NaniteRasterResults);
-
 #endif
diff --git a/Engine/Source/Runtime/Renderer/Private/PostProcess/PostProcessing.cpp b/Engine/Source/Runtime/Renderer/Private/PostProcess/PostProcessing.cpp
--- a/Engine/Source/Runtime/Renderer/Private/PostProcess/PostProcessing.cpp
+++ b/Engine/Source/Runtime/Renderer/Private/PostProcess/PostProcessing.cpp
@@ -1,7 +1,6 @@
 // Copyright Epic Games, Inc. All Rights Reserved.
 
 #include "PostProcess/PostProcessing.h"
-#include "PostProcess/PostProcessInput.h"
 #include "PostProcess/PostProcessAA.h"
 #if WITH_EDITOR
 	#include "PostProcess/PostProcessBufferInspector.h"
@@ -37,7 +36,6 @@
 #include "PostProcess/PostProcessFFTBloom.h"
 #include "PostProcess/PostProcessStreamingAccuracyLegend.h"
 #include "PostProcess/PostProcessSubsurface.h"
-#include "CompositionLighting/PostProcessPassThrough.h"
 #include "ShaderPrint.h"
 #include "GpuDebugRendering.h"
 #include "HighResScreenshot.h"
@@ -1318,42 +1316,6 @@ static bool IsGaussianActive(const FViewInfo& View)
 	return true;
 }
 
-FPostprocessContext::FPostprocessContext(FRHICommandListImmediate& InRHICmdList, FRenderingCompositionGraph& InGraph, const FViewInfo& InView)
-	: RHICmdList(InRHICmdList)
-	, Graph(InGraph)
-	, View(InView)
-	, SceneColor(0)
-	, SceneDepth(0)
-{
-	FSceneRenderTargets& SceneContext = FSceneRenderTargets::Get(InRHICmdList);
-	if (SceneContext.IsSceneColorAllocated())
-	{
-		SceneColor = Graph.RegisterPass(new(FMemStack::Get()) FRCPassPostProcessInput(SceneContext.GetSceneColor()));
-	}
-
-	SceneDepth = Graph.RegisterPass(new(FMemStack::Get()) FRCPassPostProcessInput(SceneContext.SceneDepthZ));
-
-	FinalOutput = FRenderingCompositeOutputRef(SceneColor);
-}
-
-// could be moved into the graph
-// allows for Framebuffer blending optimization with the composition graph
-void FPostProcessing::OverrideRenderTarget(FRenderingCompositeOutputRef It, TRefCountPtr<IPooledRenderTarget>& RT, FPooledRenderTargetDesc& Desc)
-{
-	for (;;)
-	{
-		It.GetOutput()->PooledRenderTarget = RT;
-		It.GetOutput()->RenderTargetDesc = Desc;
-
-		if (!It.GetPass()->FrameBufferBlendingWithInput0())
-		{
-			break;
-		}
-
-		It = *It.GetPass()->GetInput(ePId_Input0);
-	}
-}
-
 void AddMobilePostProcessingPasses(FRDGBuilder& GraphBuilder, const FViewInfo& View, const FMobilePostProcessingInputs& Inputs)
 {
 	CSV_SCOPED_TIMING_STAT_EXCLUSIVE(RenderPostProcessing);
diff --git a/Engine/Source/Runtime/Renderer/Private/PostProcess/PostProcessing.h b/Engine/Source/Runtime/Renderer/Private/PostProcess/PostProcessing.h
--- a/Engine/Source/Runtime/Renderer/Private/PostProcess/PostProcessing.h
+++ b/Engine/Source/Runtime/Renderer/Private/PostProcess/PostProcessing.h
@@ -4,7 +4,7 @@
 
 #include "ScreenPass.h"
 #include "TranslucentRendering.h"
-#include "PostProcess/RenderingCompositionGraph.h"
+#include "SystemTextures.h"
 
 class FSceneTextureParameters;
 
@@ -40,6 +40,8 @@ bool IsPostProcessingWithComputeEnabled(ERHIFeatureLevel::Type FeatureLevel);
 // Returns whether the post process pipeline supports propagating the alpha channel.
 bool IsPostProcessingWithAlphaChannelSupported();
 
+using FPostProcessVS = FScreenPassVS;
+
 struct FPostProcessingInputs
 {
 	TRDGUniformBufferRef<FSceneTextureUniformParameters> SceneTextures = nullptr;
@@ -81,37 +83,6 @@ void AddMobilePostProcessingPasses(FRDGBuilder& GraphBuilder, const FViewInfo& V
 
 void AddBasicPostProcessPasses(FRDGBuilder& GraphBuilder, const FViewInfo& View);
 
-// For compatibility with composition graph passes until they are ported to Render Graph.
-class RENDERER_API FPostProcessVS : public FScreenPassVS
-{
-public:
-	FPostProcessVS() = default;
-	FPostProcessVS(const ShaderMetaType::CompiledShaderInitializerType& Initializer)
-		: FScreenPassVS(Initializer)
-	{}
-
-	void SetParameters(const FRenderingCompositePassContext&) {}
-	void SetParameters(FRHICommandList&, FRHIUniformBuffer*) {}
-};
-
-/** The context used to setup a post-process pass. */
-class FPostprocessContext
-{
-public:
-	FPostprocessContext(FRHICommandListImmediate& InRHICmdList, FRenderingCompositionGraph& InGraph, const FViewInfo& InView);
-
-	FRHICommandListImmediate& RHICmdList;
-	FRenderingCompositionGraph& Graph;
-	const FViewInfo& View;
-
-	// 0 if there was no scene color available at constructor call time
-	FRenderingCompositePass* SceneColor;
-	// never 0
-	FRenderingCompositePass* SceneDepth;
-
-	FRenderingCompositeOutputRef FinalOutput;
-};
-
 /**
  * The center for all post processing activities.
  */
@@ -119,8 +90,6 @@ class FPostProcessing
 {
 public:
 	void ProcessPlanarReflection(FRHICommandListImmediate& RHICmdList, const FViewInfo& View, TRefCountPtr<IPooledRenderTarget>& OutFilteredSceneColor);
-
-	void OverrideRenderTarget(FRenderingCompositeOutputRef It, TRefCountPtr<IPooledRenderTarget>& RT, FPooledRenderTargetDesc& Desc);
 };
 
 /** The global used for post processing. */
diff --git a/Engine/Source/Runtime/Renderer/Private/PostProcess/RenderingCompositionGraph.cpp b/Engine/Source/Runtime/Renderer/Private/PostProcess/RenderingCompositionGraph.cpp
deleted file mode 100644
diff --git a/Engine/Source/Runtime/Renderer/Private/PostProcess/RenderingCompositionGraph.h b/Engine/Source/Runtime/Renderer/Private/PostProcess/RenderingCompositionGraph.h
deleted file mode 100644
diff --git a/Engine/Source/Runtime/Renderer/Private/RendererPrivate.h b/Engine/Source/Runtime/Renderer/Private/RendererPrivate.h
--- a/Engine/Source/Runtime/Renderer/Private/RendererPrivate.h
+++ b/Engine/Source/Runtime/Renderer/Private/RendererPrivate.h
@@ -7,7 +7,6 @@
 #pragma once
 
 #include "CoreMinimal.h"
-#include "PostProcessParameters.h"
 #include "ShaderComplexityRendering.h"
 #include "Async/ParallelFor.h"
 #include "DeferredShadingRenderer.h"
diff --git a/Engine/Source/Runtime/Renderer/Private/SceneRendering.cpp b/Engine/Source/Runtime/Renderer/Private/SceneRendering.cpp
--- a/Engine/Source/Runtime/Renderer/Private/SceneRendering.cpp
+++ b/Engine/Source/Runtime/Renderer/Private/SceneRendering.cpp
@@ -20,7 +20,6 @@
 #include "RendererModule.h"
 #include "ScenePrivate.h"
 #include "PostProcess/SceneFilterRendering.h"
-#include "PostProcess/RenderingCompositionGraph.h"
 #include "PostProcess/PostProcessEyeAdaptation.h"
 #include "PostProcess/PostProcessSubsurface.h"
 #include "PostProcess/PostProcessing.h"
@@ -3392,7 +3391,6 @@ void FSceneRenderer::OnStartRender(FRHICommandListImmediate& RHICmdList)
 	FSceneRenderTargets& SceneContext = FSceneRenderTargets::Get(RHICmdList);
 
 	FVisualizeTexturePresent::OnStartRender(Views[0]);
-	CompositionGraph_OnStartFrame();
 	SceneContext.bScreenSpaceAOIsValid = false;
 	SceneContext.bCustomDepthIsValid = false;
 
diff --git a/Engine/Source/Runtime/Renderer/Private/VolumetricFogLightFunction.cpp b/Engine/Source/Runtime/Renderer/Private/VolumetricFogLightFunction.cpp
--- a/Engine/Source/Runtime/Renderer/Private/VolumetricFogLightFunction.cpp
+++ b/Engine/Source/Runtime/Renderer/Private/VolumetricFogLightFunction.cpp
@@ -241,7 +241,6 @@ void FDeferredShadingSceneRenderer::RenderLightFunctionForVolumetricFog(
 
 				SetGraphicsPipelineState(RHICmdList, GraphicsPSOInit);
 
-				VertexShader->SetParameters(RHICmdList, View.ViewUniformBuffer);
 				PixelShader->SetParameters(RHICmdList, View, DirectionalLightSceneInfo, MaterialProxy, FVector2D(1.0f / LightFunctionResolution.X, 1.0f / LightFunctionResolution.Y), WorldToShadowValue.Inverse());
 
 				DrawRectangle(
diff --git a/Engine/Source/Runtime/Renderer/Public/PostProcessParameters.h b/Engine/Source/Runtime/Renderer/Public/PostProcessParameters.h
deleted file mode 100644
