Build fixes.

#rnx

diff --git a/Engine/Plugins/Runtime/TextureShare/Source/TextureShare/Private/TextureShareRHI.cpp b/Engine/Plugins/Runtime/TextureShare/Source/TextureShare/Private/TextureShareRHI.cpp
--- a/Engine/Plugins/Runtime/TextureShare/Source/TextureShare/Private/TextureShareRHI.cpp
+++ b/Engine/Plugins/Runtime/TextureShare/Source/TextureShare/Private/TextureShareRHI.cpp
@@ -18,7 +18,6 @@
 
 #include "ScreenRendering.h"
 #include "PostProcess/SceneFilterRendering.h"
-#include "PostProcessParameters.h"
 
 #include "RenderTargetPool.h"
 
diff --git a/Engine/Source/Runtime/Renderer/Private/CompositionLighting/PostProcessAmbientOcclusion.cpp b/Engine/Source/Runtime/Renderer/Private/CompositionLighting/PostProcessAmbientOcclusion.cpp
--- a/Engine/Source/Runtime/Renderer/Private/CompositionLighting/PostProcessAmbientOcclusion.cpp
+++ b/Engine/Source/Runtime/Renderer/Private/CompositionLighting/PostProcessAmbientOcclusion.cpp
@@ -7,6 +7,7 @@
 #include "CompositionLighting/PostProcessAmbientOcclusion.h"
 #include "CompositionLighting/CompositionLighting.h"
 #include "SceneTextureParameters.h"
+#include "ScenePrivate.h"
 
 DECLARE_GPU_STAT_NAMED(SSAOSetup, TEXT("ScreenSpace AO Setup") );
 DECLARE_GPU_STAT_NAMED(SSAO, TEXT("ScreenSpace AO") );
diff --git a/Engine/Source/Runtime/Renderer/Private/CompositionLighting/PostProcessAmbientOcclusion.h b/Engine/Source/Runtime/Renderer/Private/CompositionLighting/PostProcessAmbientOcclusion.h
--- a/Engine/Source/Runtime/Renderer/Private/CompositionLighting/PostProcessAmbientOcclusion.h
+++ b/Engine/Source/Runtime/Renderer/Private/CompositionLighting/PostProcessAmbientOcclusion.h
@@ -6,7 +6,7 @@
 
 #pragma once
 
-#include "CoreMinimal.h"
+#include "ScreenPass.h"
 #include "UniformBuffer.h"
 #include "RendererInterface.h"
 
diff --git a/Engine/Source/Runtime/Renderer/Private/CompositionLighting/PostProcessMeshDecals.cpp b/Engine/Source/Runtime/Renderer/Private/CompositionLighting/PostProcessMeshDecals.cpp
--- a/Engine/Source/Runtime/Renderer/Private/CompositionLighting/PostProcessMeshDecals.cpp
+++ b/Engine/Source/Runtime/Renderer/Private/CompositionLighting/PostProcessMeshDecals.cpp
@@ -16,6 +16,7 @@
 #include "DecalRenderingShared.h"
 #include "CompositionLighting/PostProcessDeferredDecals.h"
 #include "SceneRendering.h"
+#include "ScenePrivate.h"
 #include "UnrealEngine.h"
 #include "DebugViewModeRendering.h"
 #include "MeshPassProcessor.inl"
diff --git a/Engine/Source/Runtime/Renderer/Private/PostProcess/PostProcessSubsurface.cpp b/Engine/Source/Runtime/Renderer/Private/PostProcess/PostProcessSubsurface.cpp
--- a/Engine/Source/Runtime/Renderer/Private/PostProcess/PostProcessSubsurface.cpp
+++ b/Engine/Source/Runtime/Renderer/Private/PostProcess/PostProcessSubsurface.cpp
@@ -15,7 +15,7 @@
 #include "Engine/SubsurfaceProfile.h"
 #include "CanvasTypes.h"
 #include "RenderTargetTemp.h"
-#include "SystemTextures.h"
+#include "ScenePrivate.h"
 #include "GenerateMips.h"
 #include "ClearQuad.h"
 
