Disabled strata shader permutations compilation (global shaders) when Strata is not enabled for the project.

diff --git a/Engine/Source/Runtime/Renderer/Private/IndirectLightRendering.cpp b/Engine/Source/Runtime/Renderer/Private/IndirectLightRendering.cpp
--- a/Engine/Source/Runtime/Renderer/Private/IndirectLightRendering.cpp
+++ b/Engine/Source/Runtime/Renderer/Private/IndirectLightRendering.cpp
@@ -191,6 +191,11 @@ class FReflectionEnvironmentSkyLightingPS : public FGlobalShader
 			PermutationVector.Set<FSkyShadowing>(false);
 		}
 
+		if (PermutationVector.Get<FStrata>() && !Strata::IsStrataEnabled())
+		{
+			PermutationVector.Set<FStrata>(false);
+		}
+
 		return PermutationVector;
 	}
 
diff --git a/Engine/Source/Runtime/Renderer/Private/LightRendering.cpp b/Engine/Source/Runtime/Renderer/Private/LightRendering.cpp
--- a/Engine/Source/Runtime/Renderer/Private/LightRendering.cpp
+++ b/Engine/Source/Runtime/Renderer/Private/LightRendering.cpp
@@ -492,6 +492,11 @@ class FDeferredLightPS : public FGlobalShader
 			}
 		}
 
+		if (PermutationVector.Get<FStrata>() && !Strata::IsStrataEnabled())
+		{
+			return false;
+		}
+
 		return IsFeatureLevelSupported(Parameters.Platform, ERHIFeatureLevel::SM5);
 	}
 
