Fixing Strata cvar threading issues.

diff --git a/Engine/Source/Runtime/Renderer/Private/Strata/Strata.cpp b/Engine/Source/Runtime/Renderer/Private/Strata/Strata.cpp
--- a/Engine/Source/Runtime/Renderer/Private/Strata/Strata.cpp
+++ b/Engine/Source/Runtime/Renderer/Private/Strata/Strata.cpp
@@ -29,7 +29,7 @@ namespace Strata
 
 bool IsStrataEnabled()
 {
-	return CVarStrata.GetValueOnRenderThread() > 0;
+	return CVarStrata.GetValueOnAnyThread() > 0;
 }
 
 void InitialiseStrataFrameSceneData(FSceneRenderer& SceneRenderer, FRDGBuilder& GraphBuilder)
