Fixing crashes caused by out of order View and other data updates.

diff --git a/Engine/Source/Runtime/Renderer/Private/DeferredShadingRenderer.cpp b/Engine/Source/Runtime/Renderer/Private/DeferredShadingRenderer.cpp
--- a/Engine/Source/Runtime/Renderer/Private/DeferredShadingRenderer.cpp
+++ b/Engine/Source/Runtime/Renderer/Private/DeferredShadingRenderer.cpp
@@ -1776,19 +1776,13 @@ void FDeferredShadingSceneRenderer::Render(FRHICommandListImmediate& RHICmdList)
 		if (bDoInitViewAftersPrepass)
 		{
 			{
-				RDG_GPU_STAT_SCOPE(GraphBuilder, VisibilityCommands);
-				AddPass(GraphBuilder, [this, &ILCTaskData](FRHICommandListImmediate& RHICmdList)
-				{
-					InitViewsPossiblyAfterPrepass(RHICmdList, ILCTaskData);
-				});
+				SCOPED_GPU_STAT(GraphBuilder.RHICmdList, VisibilityCommands);
+				InitViewsPossiblyAfterPrepass(GraphBuilder.RHICmdList, ILCTaskData);
 			}
 
 			{
-				RDG_GPU_STAT_SCOPE(GraphBuilder, GPUSceneUpdate);
-				AddPass(GraphBuilder, [this](FRHICommandListImmediate& RHICmdList)
-				{
-					PrepareDistanceFieldScene(RHICmdList, false);
-				});
+				SCOPED_GPU_STAT(GraphBuilder.RHICmdList, GPUSceneUpdate);
+				PrepareDistanceFieldScene(GraphBuilder.RHICmdList, false);
 			}
 
 			{
