Set the IS_VIRTUAL_TEXTURE_MATERIAL define in Runtime Virtual Texture rendering.
Makes everything compatible with new shader generation work.

diff --git a/Engine/Source/Runtime/Renderer/Private/VT/RuntimeVirtualTextureRender.cpp b/Engine/Source/Runtime/Renderer/Private/VT/RuntimeVirtualTextureRender.cpp
--- a/Engine/Source/Runtime/Renderer/Private/VT/RuntimeVirtualTextureRender.cpp
+++ b/Engine/Source/Runtime/Renderer/Private/VT/RuntimeVirtualTextureRender.cpp
@@ -105,6 +105,7 @@ namespace RuntimeVirtualTexture
 		{
 			FMeshMaterialShader::ModifyCompilationEnvironment(Parameters, OutEnvironment);
 			OutEnvironment.SetDefine(TEXT("VIRTUAL_TEXTURE_PAGE_RENDER"), 1);
+			OutEnvironment.SetDefine(TEXT("IS_VIRTUAL_TEXTURE_MATERIAL"), 1);
 		}
 
 		FShader_VirtualTextureMaterialDraw()
