Fix logic used to wrap UDIM VT coordinates during CPU update
- Need to make sure wrapped vAddress is properly used in all cases, otherwise can end up with duplicate/inconsistent page mappings

diff --git a/Engine/Source/Runtime/RenderCore/Public/VirtualTexturing.h b/Engine/Source/Runtime/RenderCore/Public/VirtualTexturing.h
--- a/Engine/Source/Runtime/RenderCore/Public/VirtualTexturing.h
+++ b/Engine/Source/Runtime/RenderCore/Public/VirtualTexturing.h
@@ -316,6 +316,8 @@ public:
 		, SpaceID(~0u)
 		, MaxLevel(0u)
 		, VirtualAddress(~0u)
+		, VirtualPageX(~0u)
+		, VirtualPageY(~0u)
 	{}
 
 	virtual uint32 GetNumPageTableTextures() const = 0;
@@ -350,6 +352,8 @@ public:
 	inline uint32 GetDepthInPixels() const { return DepthInTiles * Description.TileSize; }
 	inline uint32 GetSpaceID() const { return SpaceID; }
 	inline uint32 GetVirtualAddress() const { return VirtualAddress; }
+	inline uint32 GetVirtualPageX() const { return VirtualPageX; }
+	inline uint32 GetVirtualPageY() const { return VirtualPageY; }
 	inline uint32 GetMaxLevel() const { return MaxLevel; }
 	inline EVTPageTableFormat GetPageTableFormat() const { return PageTableFormat; }
 
@@ -371,6 +375,8 @@ protected:
 	uint32 SpaceID;
 	uint32 MaxLevel;
 	uint32 VirtualAddress;
+	uint32 VirtualPageX;
+	uint32 VirtualPageY;
 };
 
 /** 
diff --git a/Engine/Source/Runtime/Renderer/Private/VT/AllocatedVirtualTexture.cpp b/Engine/Source/Runtime/Renderer/Private/VT/AllocatedVirtualTexture.cpp
--- a/Engine/Source/Runtime/Renderer/Private/VT/AllocatedVirtualTexture.cpp
+++ b/Engine/Source/Runtime/Renderer/Private/VT/AllocatedVirtualTexture.cpp
@@ -117,6 +117,15 @@ FAllocatedVirtualTexture::~FAllocatedVirtualTexture()
 {
 }
 
+void FAllocatedVirtualTexture::AssignVirtualAddress(uint32 vAddress)
+{
+	checkf(VirtualAddress == ~0u, TEXT("Trying to assign vAddress to AllocatedVT, already assigned"));
+	check(vAddress != ~0u);
+	VirtualAddress = vAddress;
+	VirtualPageX = FMath::ReverseMortonCode2(vAddress);
+	VirtualPageY = FMath::ReverseMortonCode2(vAddress >> 1);
+}
+
 void FAllocatedVirtualTexture::Destroy(FVirtualTextureSystem* System)
 {
 	const int32 NewRefCount = RefCount.Decrement();
diff --git a/Engine/Source/Runtime/Renderer/Private/VT/AllocatedVirtualTexture.h b/Engine/Source/Runtime/Renderer/Private/VT/AllocatedVirtualTexture.h
--- a/Engine/Source/Runtime/Renderer/Private/VT/AllocatedVirtualTexture.h
+++ b/Engine/Source/Runtime/Renderer/Private/VT/AllocatedVirtualTexture.h
@@ -34,6 +34,8 @@ public:
 	inline void IncrementRefCount() { RefCount.Increment(); }
 	inline uint32 GetFrameAllocated() const { return FrameAllocated; }
 
+	void AssignVirtualAddress(uint32 vAddress);
+
 	// begin IAllocatedVirtualTexture
 	virtual uint32 GetNumPageTableTextures() const override;
 	virtual FRHITexture* GetPageTableTexture(uint32 InPageTableIndex) const override;
diff --git a/Engine/Source/Runtime/Renderer/Private/VT/VirtualTextureAllocator.cpp b/Engine/Source/Runtime/Renderer/Private/VT/VirtualTextureAllocator.cpp
--- a/Engine/Source/Runtime/Renderer/Private/VT/VirtualTextureAllocator.cpp
+++ b/Engine/Source/Runtime/Renderer/Private/VT/VirtualTextureAllocator.cpp
@@ -128,7 +128,7 @@ int32 FVirtualTextureAllocator::AcquireBlock()
 	return FreeBlock;
 }
 
-uint32 FVirtualTextureAllocator::Find(uint32 vAddress) const
+uint32 FVirtualTextureAllocator::FindAddressBlock(uint32 vAddress) const
 {
 	uint32 Min = 0;
 	uint32 Max = SortedAddresses.Num();
@@ -148,9 +148,9 @@ uint32 FVirtualTextureAllocator::Find(uint32 vAddress) const
 	return Min;
 }
 
-FAllocatedVirtualTexture* FVirtualTextureAllocator::Find(uint32 vAddress, uint32& Local_vAddress) const
+FAllocatedVirtualTexture* FVirtualTextureAllocator::Find(uint32 vAddress, uint32& OutLocal_vAddress) const
 {
-	const uint32 SortedIndex = Find(vAddress);
+	const uint32 SortedIndex = FindAddressBlock(vAddress);
 
 	const uint16 Index = SortedIndices[SortedIndex];
 	const FAddressBlock& AddressBlock = AddressBlocks[Index];
@@ -160,7 +160,8 @@ FAllocatedVirtualTexture* FVirtualTextureAllocator::Find(uint32 vAddress, uint32
 	if (vAddress >= AddressBlock.vAddress &&
 		vAddress < AddressBlock.vAddress + BlockSize)
 	{
-		Local_vAddress = vAddress - AddressBlock.vAddress;
+		check(AddressBlock.VT->GetVirtualAddress() == AddressBlock.vAddress);
+		OutLocal_vAddress = vAddress - AddressBlock.vAddress;
 		// TODO mip bias
 		return AddressBlock.VT;
 	}
@@ -248,7 +249,7 @@ uint32 FVirtualTextureAllocator::Alloc(FAllocatedVirtualTexture* VT)
 				FAddressBlock* ParentBlock = &AddressBlocks[FreeIndex];
 				check(ParentBlock->vAddress == AllocBlock->vAddress);
 
-				const int32 SortedIndex = Find(AllocBlock->vAddress);
+				const int32 SortedIndex = FindAddressBlock(AllocBlock->vAddress);
 				check(AllocBlock->vAddress == SortedAddresses[SortedIndex]);
 
 				// Replace parent block with new allocated block
@@ -406,7 +407,7 @@ void FVirtualTextureAllocator::FreeAddressBlock(uint32 Index)
 		// Remove this block and its siblings from the sorted lists
 		// We can assume that the sibling blocks are sequential in the sorted list since they are free and so have no children
 		// FirstSibling will be the last in the range of siblings in the sorted lists 
-		const uint32 SortedIndexRangeEnd = Find(AddressBlocks[AddressBlock.FirstSibling].vAddress);
+		const uint32 SortedIndexRangeEnd = FindAddressBlock(AddressBlocks[AddressBlock.FirstSibling].vAddress);
 		check(SortedAddresses[SortedIndexRangeEnd] == AddressBlocks[AddressBlock.FirstSibling].vAddress);
 		const uint32 NumSiblings = 1 << vDimensions;
 		check(SortedIndexRangeEnd + 1 >= NumSiblings);
diff --git a/Engine/Source/Runtime/Renderer/Private/VT/VirtualTextureAllocator.h b/Engine/Source/Runtime/Renderer/Private/VT/VirtualTextureAllocator.h
--- a/Engine/Source/Runtime/Renderer/Private/VT/VirtualTextureAllocator.h
+++ b/Engine/Source/Runtime/Renderer/Private/VT/VirtualTextureAllocator.h
@@ -29,7 +29,8 @@ public:
 	 * Translate a virtual page address in the address space to a local page address within a virtual texture.
 	 * @return nullptr If there is no virtual texture allocated at this address.
 	 */
-	FAllocatedVirtualTexture* Find(uint32 vAddress, uint32& Local_vAddress) const;
+	FAllocatedVirtualTexture* Find(uint32 vAddress, uint32& OutLocal_vAddress) const;
+	inline FAllocatedVirtualTexture* Find(uint32 vAddress) const { uint32 UnusedLocal_vAddress = 0u; return Find(vAddress, UnusedLocal_vAddress); }
 
 	/**
 	 * Allocate address space for the virtual texture.
@@ -60,7 +61,7 @@ public:
 private:
 	int32 AcquireBlock();
 	void FreeAddressBlock(uint32 Index);
-	uint32 Find(uint32 vAddress) const;
+	uint32 FindAddressBlock(uint32 vAddress) const;
 
 	struct FAddressBlock
 	{
diff --git a/Engine/Source/Runtime/Renderer/Private/VT/VirtualTextureSystem.cpp b/Engine/Source/Runtime/Renderer/Private/VT/VirtualTextureSystem.cpp
--- a/Engine/Source/Runtime/Renderer/Private/VT/VirtualTextureSystem.cpp
+++ b/Engine/Source/Runtime/Renderer/Private/VT/VirtualTextureSystem.cpp
@@ -619,7 +619,7 @@ FVirtualTextureSpace* FVirtualTextureSystem::AcquireSpace(const FVTSpaceDescript
 					const uint32 vAddress = Space->AllocateVirtualTexture(AllocatedVT);
 					if (vAddress != ~0u)
 					{
-						AllocatedVT->VirtualAddress = vAddress;
+						AllocatedVT->AssignVirtualAddress(vAddress);
 						Space->AddRef();
 						return Space;
 					}
@@ -646,8 +646,7 @@ FVirtualTextureSpace* FVirtualTextureSystem::AcquireSpace(const FVTSpaceDescript
 				BeginInitResource(Space);
 
 				const uint32 vAddress = Space->AllocateVirtualTexture(AllocatedVT);
-				check(vAddress != ~0u);
-				AllocatedVT->VirtualAddress = vAddress;
+				AllocatedVT->AssignVirtualAddress(vAddress);
 
 				Space->AddRef();
 				return Space;
@@ -1366,7 +1365,7 @@ void FVirtualTextureSystem::GatherRequestsTask(const FGatherRequestsParameters&
 		// vPosition holds morton interleaved tileX/Y position, shifted down relative to current mip
 		// vAddress is the same quantity, but shifted to be relative to mip0
 		const uint32 vDimensions = Space->GetDimensions();
-		const uint32 vAddress = vPosition << (vLevel * vDimensions);
+		uint32 vAddress = vPosition << (vLevel * vDimensions);
 
 		const FAdaptiveVirtualTexture* RESTRICT AdaptiveVT = AdaptiveVTs[ID];
 		if (AdaptiveVT != nullptr && vLevelPlusOne <= 1)
@@ -1427,8 +1426,7 @@ void FVirtualTextureSystem::GatherRequestsTask(const FGatherRequestsParameters&
 		}
 
 		// Need to resolve AllocatedVT in order to determine which pages to load
-		uint32 AllocatedLocal_vAddress = 0;
-		const FAllocatedVirtualTexture* RESTRICT AllocatedVT = Space->GetAllocator().Find(vAddress, AllocatedLocal_vAddress);
+		const FAllocatedVirtualTexture* RESTRICT AllocatedVT = Space->GetAllocator().Find(vAddress);
 		if (!AllocatedVT)
 		{
 			if (CVarVTVerbose.GetValueOnAnyThread())
@@ -1480,6 +1478,8 @@ void FVirtualTextureSystem::GatherRequestsTask(const FGatherRequestsParameters&
 			}
 		}
 
+		check(vAddress >= AllocatedVT->GetVirtualAddress());
+
 		for (uint32 ProducerIndex = 0u; ProducerIndex < NumUniqueProducers; ++ProducerIndex)
 		{
 			uint8 GroupMaskToLoad = ProducerGroupMaskToLoad[ProducerIndex];
@@ -1497,11 +1497,10 @@ void FVirtualTextureSystem::GatherRequestsTask(const FGatherRequestsParameters&
 
 			const uint32 MaxLevel = FMath::Min(Producer->GetMaxLevel(), AllocatedVT->GetMaxLevel());
 			const uint32 ProducerMipBias = AllocatedVT->GetUniqueProducerMipBias(ProducerIndex);
-			uint32 Mapping_vLevel = FMath::Max(vLevel, ProducerMipBias);
 
-			// rescale vAddress to the correct tile within the given mip level
 			// here vLevel is clamped against ProducerMipBias, as ProducerMipBias represents the most detailed level of this producer, relative to the allocated VT
-			uint32 Local_vAddress = AllocatedLocal_vAddress >> (Mapping_vLevel * vDimensions);
+			// used to rescale vAddress to the correct tile within the given mip level
+			uint32 Mapping_vLevel = FMath::Max(vLevel, ProducerMipBias);
 
 			// Local_vLevel is the level within the producer that we want to allocate/map
 			// here we subtract ProducerMipBias (clamped to ensure we don't fall below 0),
@@ -1511,16 +1510,18 @@ void FVirtualTextureSystem::GatherRequestsTask(const FGatherRequestsParameters&
 			// Wrap Local_vAddress for the given producer
 			// For square textures, this could simply be (Local_vAddress % NumTilesInMip), but that doesn't work for non-square
 			// Possible there is a more clever approach to take here
+			uint32 Local_vAddress;
 			{
 				const uint32 MipScaleFactor = (1u << Local_vLevel);
-				const uint32 ProducerMipWidthInTiles = FMath::DivideAndRoundUp(Producer->GetWidthInTiles(), MipScaleFactor);
-				const uint32 ProducerMipHeightInTiles = FMath::DivideAndRoundUp(Producer->GetHeightInTiles(), MipScaleFactor);
-				uint32 Local_vTileX = FMath::ReverseMortonCode2(Local_vAddress);
-				uint32 Local_vTileY = FMath::ReverseMortonCode2(Local_vAddress >> 1);
-				Local_vTileX %= ProducerMipWidthInTiles;
-				Local_vTileY %= ProducerMipHeightInTiles;
-				Local_vAddress = FMath::MortonCode2(Local_vTileX) | (FMath::MortonCode2(Local_vTileY) << 1);
+				const uint32 ProducerMipWidthInPages = FMath::DivideAndRoundUp(Producer->GetWidthInTiles(), MipScaleFactor);
+				const uint32 ProducerMipHeightInPages = FMath::DivideAndRoundUp(Producer->GetHeightInTiles(), MipScaleFactor);
+				uint32 Local_vPageX = (vPageX - AllocatedVT->GetVirtualPageX()) >> Mapping_vLevel;
+				uint32 Local_vTileY = (vPageY - AllocatedVT->GetVirtualPageY()) >> Mapping_vLevel;
+				Local_vPageX %= ProducerMipWidthInPages;
+				Local_vTileY %= ProducerMipHeightInPages;
+				Local_vAddress = FMath::MortonCode2(Local_vPageX) | (FMath::MortonCode2(Local_vTileY) << 1);
 			}
+			vAddress = (Local_vAddress << (Mapping_vLevel * vDimensions)) + AllocatedVT->GetVirtualAddress();
 
 			const uint32 LocalMipBias = Producer->GetVirtualTexture()->GetLocalMipBias(Local_vLevel, Local_vAddress);
 			if (LocalMipBias > 0u)
@@ -1568,31 +1569,18 @@ void FVirtualTextureSystem::GatherRequestsTask(const FGatherRequestsParameters&
 							AllocatedVT->GetProducerIndexForPageTableLayer(PageTableLayerIndex) == ProducerIndex)
 						{
 							bool bPageWasMapped = false;
-							if (Allocated_vLevel != vLevel)
-							{
-								// if we found a lower resolution tile than was requested, it may have already been mapped, check for that first
-								// don't need to check this if the allocated page is at the level that was requested...if that was already mapped we wouldn't have gotten this far
-								const FTexturePageMap& PageMap = Space->GetPageMapForPageTableLayer(PageTableLayerIndex);
-								const FPhysicalSpaceIDAndAddress PrevPhysicalSpaceIDAndAddress = PageMap.FindPagePhysicalSpaceIDAndAddress(Allocated_vLevel, Allocated_vAddress);
-								if (PrevPhysicalSpaceIDAndAddress.Packed != ~0u)
-								{
-									// if this address was previously mapped, ensure that it was mapped by the same physical space
-									ensure(PrevPhysicalSpaceIDAndAddress.PhysicalSpaceID == PhysicalSpace->GetID());
-									// either it wasn't mapped, or it's mapped to the current physical address...
-									// otherwise that means that the same local tile is mapped to two separate physical addresses, which is an error
-									ensure(PrevPhysicalSpaceIDAndAddress.pAddress == pAddress);
-									bPageWasMapped = true;
-								}
-							}
-#if DO_GUARD_SLOW
-							else
+							// if we found a lower resolution tile than was requested, it may have already been mapped, check for that first
+							const FTexturePageMap& PageMap = Space->GetPageMapForPageTableLayer(PageTableLayerIndex);
+							const FPhysicalSpaceIDAndAddress PrevPhysicalSpaceIDAndAddress = PageMap.FindPagePhysicalSpaceIDAndAddress(Allocated_vLevel, Allocated_vAddress);
+							if (PrevPhysicalSpaceIDAndAddress.Packed != ~0u)
 							{
-								// verify our assumption that the page shouldn't be mapped yet
-								const FTexturePageMap& PageMap = Space->GetPageMapForPageTableLayer(PageTableLayerIndex);
-								const FPhysicalSpaceIDAndAddress PrevPhysicalSpaceIDAndAddress = PageMap.FindPagePhysicalSpaceIDAndAddress(Allocated_vLevel, Allocated_vAddress);
-								checkSlow(PrevPhysicalSpaceIDAndAddress.Packed == ~0u);
+								// if this address was previously mapped, ensure that it was mapped by the same physical space
+								ensure(PrevPhysicalSpaceIDAndAddress.PhysicalSpaceID == PhysicalSpace->GetID());
+								// either it wasn't mapped, or it's mapped to the current physical address...
+								// otherwise that means that the same local tile is mapped to two separate physical addresses, which is an error
+								ensure(PrevPhysicalSpaceIDAndAddress.pAddress == pAddress);
+								bPageWasMapped = true;
 							}
-#endif // DO_GUARD_SLOW
 
 							if (!bPageWasMapped)
 							{
