Fix bias handling for mapping resident pages of newly allocated VTs

diff --git a/Engine/Source/Runtime/Renderer/Private/VT/VirtualTextureAllocator.cpp b/Engine/Source/Runtime/Renderer/Private/VT/VirtualTextureAllocator.cpp
--- a/Engine/Source/Runtime/Renderer/Private/VT/VirtualTextureAllocator.cpp
+++ b/Engine/Source/Runtime/Renderer/Private/VT/VirtualTextureAllocator.cpp
@@ -160,7 +160,7 @@ FAllocatedVirtualTexture* FVirtualTextureAllocator::Find(uint32 vAddress, uint32
 	if (vAddress >= AddressBlock.vAddress &&
 		vAddress < AddressBlock.vAddress + BlockSize)
 	{
-		check(AddressBlock.VT->GetVirtualAddress() == AddressBlock.vAddress);
+		check(!AddressBlock.VT || AddressBlock.VT->GetVirtualAddress() == AddressBlock.vAddress);
 		OutLocal_vAddress = vAddress - AddressBlock.vAddress;
 		// TODO mip bias
 		return AddressBlock.VT;
diff --git a/Engine/Source/Runtime/Renderer/Private/VT/VirtualTextureSystem.cpp b/Engine/Source/Runtime/Renderer/Private/VT/VirtualTextureSystem.cpp
--- a/Engine/Source/Runtime/Renderer/Private/VT/VirtualTextureSystem.cpp
+++ b/Engine/Source/Runtime/Renderer/Private/VT/VirtualTextureSystem.cpp
@@ -2125,8 +2125,8 @@ void FVirtualTextureSystem::SubmitRequests(FRDGBuilder& GraphBuilder, ERHIFeatur
 		{
 			const FAllocatedVirtualTexture* AllocatedVT = AllocatedVTsToMap[Index];
 			const uint32 vDimensions = AllocatedVT->GetDimensions();
-			const uint32 BaseTileX = FMath::ReverseMortonCode2(AllocatedVT->GetVirtualAddress());
-			const uint32 BaseTileY = FMath::ReverseMortonCode2(AllocatedVT->GetVirtualAddress() >> 1);
+			const uint32 BaseTileX = AllocatedVT->GetVirtualPageX();
+			const uint32 BaseTileY = AllocatedVT->GetVirtualPageY();
 			FVirtualTextureSpace* Space = AllocatedVT->GetSpace();
 
 			uint32 NumFullyMappedLayers = 0u;
@@ -2146,18 +2146,21 @@ void FVirtualTextureSystem::SubmitRequests(FRDGBuilder& GraphBuilder, ERHIFeatur
 				const uint32 ProducerMipBias = AllocatedVT->GetUniqueProducerMipBias(ProducerIndex);
 				const uint32 WidthInTiles = Producer->GetWidthInTiles();
 				const uint32 HeightInTiles = Producer->GetHeightInTiles();
-				const uint32 MaxLevel = FMath::Min(Producer->GetMaxLevel(), AllocatedVT->GetMaxLevel());
+				const uint32 MaxLevel = FMath::Min(Producer->GetMaxLevel(), AllocatedVT->GetMaxLevel() - ProducerMipBias);
 
 				FVirtualTexturePhysicalSpace* PhysicalSpace = AllocatedVT->GetPhysicalSpaceForPageTableLayer(PageTableLayerIndex);
 				FTexturePagePool& PagePool = PhysicalSpace->GetPagePool();
 				FTexturePageMap& PageMap = Space->GetPageMapForPageTableLayer(PageTableLayerIndex);
 				
 				bool bIsLayerFullyMapped = false;
-				for (uint32 Local_vLevel = 0u; Local_vLevel <= MaxLevel; ++Local_vLevel)
+				for (uint32 Local_vLevel = 0; Local_vLevel <= MaxLevel; ++Local_vLevel)
 				{
 					const uint32 vLevel = Local_vLevel + ProducerMipBias;
-					const uint32 LevelWidthInTiles = FMath::Max(WidthInTiles >> vLevel, 1u);
-					const uint32 LevelHeightInTiles = FMath::Max(HeightInTiles >> vLevel, 1u);
+					check(vLevel <= AllocatedVT->GetMaxLevel());
+
+					const uint32 MipScaleFactor = (1u << Local_vLevel);
+					const uint32 LevelWidthInTiles = FMath::DivideAndRoundUp(WidthInTiles, MipScaleFactor);
+					const uint32 LevelHeightInTiles = FMath::DivideAndRoundUp(HeightInTiles, MipScaleFactor);
 
 					uint32 NumNonResidentPages = 0u;
 					for (uint32 TileY = 0u; TileY < LevelHeightInTiles; ++TileY)
