Path tracer - stop tracing rays when we hit the max number of samples per pixel

When this is combined with FrameIndependentTemporalSeed=0, it allows rendering as scene with a predictable noise pattern which is useful for consistent unit testing.

Invalidate the path traced image in more cases: object deletion, camera movement while not in path traced mode, random sequence changes and changing SamplesPerPixel

diff --git a/Engine/Source/Runtime/Renderer/Private/PathTracing.cpp b/Engine/Source/Runtime/Renderer/Private/PathTracing.cpp
--- a/Engine/Source/Runtime/Renderer/Private/PathTracing.cpp
+++ b/Engine/Source/Runtime/Renderer/Private/PathTracing.cpp
@@ -15,11 +15,11 @@
 #include "RayTracing/RaytracingOptions.h"
 #include "HAL/PlatformApplicationMisc.h"
 
-static int32 GPathTracingMaxBounces = -1;
-static FAutoConsoleVariableRef CVarPathTracingMaxBounces(
+TAutoConsoleVariable<int32> CVarPathTracingMaxBounces(
 	TEXT("r.PathTracing.MaxBounces"),
-	GPathTracingMaxBounces,
-	TEXT("Sets the maximum number of path tracing bounces (default = -1 (driven by postprocesing volume))")
+	-1,
+	TEXT("Sets the maximum number of path tracing bounces (default = -1 (driven by postprocesing volume))"),
+	ECVF_RenderThreadSafe
 );
 
 TAutoConsoleVariable<int32> CVarPathTracingSamplesPerPixel(
@@ -35,7 +35,8 @@ TAutoConsoleVariable<int32> CVarPathTracingMISMode(
 	TEXT("Selects the sampling techniques (default = 2 (MIS enabled))\n")
 	TEXT("0: Material sampling\n")
 	TEXT("1: Light sampling\n")
-	TEXT("2: MIS betwen material and light sampling (default)\n")
+	TEXT("2: MIS betwen material and light sampling (default)\n"),
+	ECVF_RenderThreadSafe
 );
 
 TAutoConsoleVariable<int32> CVarPathTracingVisibleLights(
@@ -43,13 +44,15 @@ TAutoConsoleVariable<int32> CVarPathTracingVisibleLights(
 	0,
 	TEXT("Should light sources be visible to camera rays? (default = 0 (off))\n")
 	TEXT("0: Hide lights from camera rays (default)\n")
-	TEXT("1: Make lights visible to camera\n")
+	TEXT("1: Make lights visible to camera\n"),
+	ECVF_RenderThreadSafe
 );
 
 TAutoConsoleVariable<int32> CVarPathTracingMaxPathIntensity(
 	TEXT("r.PathTracing.MaxPathIntensity"),
 	-1,
-	TEXT("When positive, light paths greater that this amount are clamped to prevent fireflies (default = -1 (off))")
+	TEXT("When positive, light paths greater that this amount are clamped to prevent fireflies (default = -1 (off))"),
+	ECVF_RenderThreadSafe
 );
 
 TAutoConsoleVariable<int32> CVarPathTracingFrameIndependentTemporalSeed(
@@ -63,11 +66,12 @@ TAutoConsoleVariable<int32> CVarPathTracingFrameIndependentTemporalSeed(
 
 TAutoConsoleVariable<int32> CVarPathTracingRandomSequence(
 	TEXT("r.PathTracing.RandomSequence"),
-	2,
+	3,
 	TEXT("Changes the underlying random sequence\n")
 	TEXT("0: LCG\n")
 	TEXT("1: Halton\n")
-	TEXT("2: Scrambled Halton (default)\n"),
+	TEXT("2: Scrambled Halton\n")
+	TEXT("3: Owen-Sobol (default)\n"),
 	ECVF_RenderThreadSafe
 );
 
@@ -120,6 +124,75 @@ IMPLEMENT_GLOBAL_SHADER_PARAMETER_STRUCT(FPathTracingData, "PathTracingData");
 IMPLEMENT_GLOBAL_SHADER_PARAMETER_STRUCT(FPathTracingLightData, "SceneLightsData");
 IMPLEMENT_GLOBAL_SHADER_PARAMETER_STRUCT(FPathTracingAdaptiveSamplingData, "AdaptiveSamplingData");
 
+// This function prepares the portion of shader arguments that may involve invalidating the path traced state
+static bool PrepareShaderArgs(const FViewInfo& View, FPathTracingData& PathTracingData, FPathTracingAdaptiveSamplingData& AdaptiveSamplingData) {
+	int32 PathTracingMaxBounces = CVarPathTracingMaxBounces.GetValueOnRenderThread();
+	if (PathTracingMaxBounces < 0)
+		PathTracingMaxBounces = View.FinalPostProcessSettings.PathTracingMaxBounces;
+	PathTracingData.MaxBounces = PathTracingMaxBounces;
+	PathTracingData.MaxNormalBias = GetRaytracingMaxNormalBias();
+	PathTracingData.MISMode = CVarPathTracingMISMode.GetValueOnRenderThread();
+	PathTracingData.VisibleLights = CVarPathTracingVisibleLights.GetValueOnRenderThread();
+	PathTracingData.MaxPathIntensity = CVarPathTracingMaxPathIntensity.GetValueOnRenderThread();
+
+	bool NeedInvalidation = false;
+
+	// If any of the parameters above changed since last time -- reset the accumulation
+	// FIXME: find something cleaner than just using static variables here. Should all
+	// the state used for comparison go into ViewState?
+	static uint32 PrevMaxBounces = PathTracingMaxBounces;
+	if (PathTracingData.MaxBounces != PrevMaxBounces)
+	{
+		NeedInvalidation = true;
+		PrevMaxBounces = PathTracingData.MaxBounces;
+	}
+
+	// Changing MIS mode requires starting over
+	static uint32 PreviousMISMode = PathTracingData.MISMode;
+	if (PreviousMISMode != PathTracingData.MISMode)
+	{
+		NeedInvalidation = true;
+		PreviousMISMode = PathTracingData.MISMode;
+	}
+
+	// Changing VisibleLights requires starting over
+	static uint32 PreviousVisibleLights = PathTracingData.VisibleLights;
+	if (PreviousVisibleLights != PathTracingData.VisibleLights)
+	{
+		NeedInvalidation = true;
+		PreviousVisibleLights = PathTracingData.VisibleLights;
+	}
+
+	// Changing MaxPathIntensity requires starting over
+	static float PreviousMaxPathIntensity = PathTracingData.MaxPathIntensity;
+	if (PreviousMaxPathIntensity != PathTracingData.MaxPathIntensity)
+	{
+		NeedInvalidation = true;
+		PreviousMaxPathIntensity = PathTracingData.MaxPathIntensity;
+	}
+
+	AdaptiveSamplingData.UseAdaptiveSampling = CVarPathTracingAdaptiveSampling.GetValueOnRenderThread();
+	AdaptiveSamplingData.RandomSequence = CVarPathTracingRandomSequence.GetValueOnRenderThread();
+
+	// Changing Adaptive sampling mode requires starting over
+	static uint32 PreviousUseAdaptiveSampling = AdaptiveSamplingData.UseAdaptiveSampling;
+	if (PreviousUseAdaptiveSampling != AdaptiveSamplingData.UseAdaptiveSampling)
+	{
+		NeedInvalidation = true;
+		PreviousUseAdaptiveSampling = AdaptiveSamplingData.UseAdaptiveSampling;
+	}
+
+	static uint32 PreviousRandomSequence = AdaptiveSamplingData.RandomSequence;
+	if (PreviousRandomSequence != AdaptiveSamplingData.RandomSequence)
+	{
+		NeedInvalidation = true;
+		PreviousRandomSequence = AdaptiveSamplingData.RandomSequence;
+	}
+
+	// the rest of PathTracingData and AdaptiveSamplingData is filled in by SetParameters below
+	return NeedInvalidation;
+}
+
 class FPathTracingRG : public FGlobalShader
 {
 	DECLARE_SHADER_TYPE(FPathTracingRG, Global);
@@ -163,6 +236,9 @@ public:
 		const FRayTracingScene& RayTracingScene,
 		FRHIUniformBuffer* ViewUniformBuffer,
 		FRHIUniformBuffer* SceneTexturesUniformBuffer,
+		// Shader arguments (expected to be filled in by PrepareShaderArgs ahead of time)
+		FPathTracingData& PathTracingData,
+		FPathTracingAdaptiveSamplingData& AdaptiveSamplingData,
 		// Light buffer
 		const TSparseArray<FLightSceneInfoCompact>& Lights,
 		// Adaptive sampling
@@ -183,47 +259,6 @@ public:
 
 		// Path tracing data
 		{
-			FPathTracingData PathTracingData;
-
-			int32 PathTracingMaxBounces = GPathTracingMaxBounces > -1 ? GPathTracingMaxBounces : View.FinalPostProcessSettings.PathTracingMaxBounces;
-			PathTracingData.MaxBounces = PathTracingMaxBounces;
-			PathTracingData.MaxNormalBias = GetRaytracingMaxNormalBias();
-			PathTracingData.MISMode = CVarPathTracingMISMode.GetValueOnRenderThread();
-			PathTracingData.VisibleLights = CVarPathTracingVisibleLights.GetValueOnRenderThread();
-			PathTracingData.MaxPathIntensity = CVarPathTracingMaxPathIntensity.GetValueOnRenderThread();
-
-			// If any of the parameters above changed since last time -- reset the accumulation
-			static uint32 PrevMaxBounces = PathTracingMaxBounces;
-			if (PathTracingData.MaxBounces != PrevMaxBounces)
-			{
-				Scene->bPathTracingNeedsInvalidation = true;
-				PrevMaxBounces = PathTracingData.MaxBounces;
-			}
-
-			// Changing MIS mode requires starting over
-			static uint32 PreviousMISMode = PathTracingData.MISMode;
-			if (PreviousMISMode != PathTracingData.MISMode)
-			{
-				Scene->bPathTracingNeedsInvalidation = true;
-				PreviousMISMode = PathTracingData.MISMode;
-			}
-
-			// Changing VisibleLights requires starting over
-			static uint32 PreviousVisibleLights = PathTracingData.VisibleLights;
-			if (PreviousVisibleLights != PathTracingData.VisibleLights)
-			{
-				Scene->bPathTracingNeedsInvalidation = true;
-				PreviousVisibleLights = PathTracingData.VisibleLights;
-			}
-
-			// Changing MaxPathIntensity requires starting over
-			static float PreviousMaxPathIntensity = PathTracingData.MaxPathIntensity;
-			if (PreviousMaxPathIntensity != PathTracingData.MaxPathIntensity)
-			{
-				Scene->bPathTracingNeedsInvalidation = true;
-				PreviousMaxPathIntensity = PathTracingData.MaxPathIntensity;
-			}
-
 			PathTracingData.TileOffset = TileOffset;
 
 			FUniformBufferRHIRef PathTracingDataUniformBuffer = RHICreateUniformBuffer(&PathTracingData, FPathTracingData::StaticStructMetadata.GetLayout(), EUniformBufferUsage::UniformBuffer_SingleDraw);
@@ -342,18 +377,6 @@ public:
 				TemporalSeed = FrameIndependentTemporalSeed;
 			}
 
-			FPathTracingAdaptiveSamplingData AdaptiveSamplingData;
-			AdaptiveSamplingData.UseAdaptiveSampling = CVarPathTracingAdaptiveSampling.GetValueOnRenderThread();
-			AdaptiveSamplingData.RandomSequence = CVarPathTracingRandomSequence.GetValueOnRenderThread();
-
-			// Changing Adaptive sampling mode requires starting over
-			static uint32 PreviousUseAdaptiveSampling = AdaptiveSamplingData.UseAdaptiveSampling;
-			if (PreviousUseAdaptiveSampling != AdaptiveSamplingData.UseAdaptiveSampling)
-			{
-				Scene->bPathTracingNeedsInvalidation = true;
-				PreviousUseAdaptiveSampling = AdaptiveSamplingData.UseAdaptiveSampling;
-			}
-
 			if (VarianceMipTree.NumBytes > 0)
 			{
 				AdaptiveSamplingData.Iteration = Iteration;
@@ -509,6 +532,15 @@ BEGIN_SHADER_PARAMETER_STRUCT(FPathTracingPassParameters, )
 	RENDER_TARGET_BINDING_SLOTS()
 END_SHADER_PARAMETER_STRUCT()
 
+void FSceneViewState::PathTracingInvalidate()
+{
+	PathTracingIrradianceRT.SafeRelease();
+	PathTracingSampleCountRT.SafeRelease();
+	VarianceMipTreeDimensions = FIntVector(0);
+	TotalRayCount = 0;
+	PathTracingSPP = 0;
+}
+
 void FDeferredShadingSceneRenderer::RenderPathTracing(
 	FRDGBuilder& GraphBuilder,
 	const FViewInfo& View,
@@ -529,14 +561,33 @@ void FDeferredShadingSceneRenderer::RenderPathTracing(
 		ERDGPassFlags::Compute | ERDGPassFlags::Raster | ERDGPassFlags::SkipRenderPass | ERDGPassFlags::UntrackedAccess,
 		[this, &View, SceneTexturesUniformBuffer, SceneColorOutputTexture](FRHICommandListImmediate& RHICmdList)
 	{
-		// The local iteration counter.
-		static int32 SPPCount = 0;
-		// The frame independent temporal seed, not reset at the beginning of each frame unlike SPPCount to allow for less temporal aliasing.
-		static uint32 FrameIndependentTemporalSeed = 0;
-
 		auto ViewSize = View.ViewRect.Size();
 		FSceneViewState* ViewState = (FSceneViewState*)View.State;
 
+		FPathTracingData PathTracingData;
+		FPathTracingAdaptiveSamplingData AdaptiveSamplingData;
+
+		bool bArgsChanged = PrepareShaderArgs(View, PathTracingData, AdaptiveSamplingData);
+
+        // Get current value of MaxSPP and reset render if it has changed
+		int32 SamplesPerPixelCVar = CVarPathTracingSamplesPerPixel.GetValueOnRenderThread();
+		uint32 MaxSPP = SamplesPerPixelCVar > -1 ? SamplesPerPixelCVar : View.FinalPostProcessSettings.PathTracingSamplesPerPixel;
+		static uint32 PreviousMaxSPP = MaxSPP;
+		if (PreviousMaxSPP != MaxSPP)
+		{
+			PreviousMaxSPP = MaxSPP;
+			bArgsChanged = true;
+		}
+
+		// If the scene has changed in some way (camera move, object movement, etc ...)
+		// we must invalidate the viewstate to start over from scratch
+		if (bArgsChanged || ViewState->PathTracingRect != View.ViewRect)
+		{
+			ViewState->PathTracingInvalidate();
+			ViewState->PathTracingRect = View.ViewRect;
+		}
+		
+
 		// Construct render targets for compositing
 		TRefCountPtr<IPooledRenderTarget> RadianceRT;
 		TRefCountPtr<IPooledRenderTarget> SampleCountRT;
@@ -580,7 +631,7 @@ void FDeferredShadingSceneRenderer::RenderPathTracing(
 
 		bool bDoMGPUPathTracing = GNumExplicitGPUsForRendering > 1 && GPUCount > 1;
 
-		if (bDoMGPUPathTracing)
+		if (bDoMGPUPathTracing && ViewState->PathTracingSPP < MaxSPP)
 		{
 			//#dxr-todo: Set minimum tile size for mGPU
 			int32 TileSizeX = FIntPoint::DivideAndRoundUp(View.ViewRect.Size(), GPUCount).X;
@@ -628,8 +679,14 @@ void FDeferredShadingSceneRenderer::RenderPathTracing(
 					View.RayTracingScene,
 					View.ViewUniformBuffer,
 					SceneTexturesUniformBufferRHI,
+					PathTracingData,
+					AdaptiveSamplingData,
 					Scene->Lights,
-					SPPCount, FrameIndependentTemporalSeed, ViewState->VarianceMipTreeDimensions, *ViewState->VarianceMipTree, TileOffset,
+					ViewState->PathTracingSPP,
+					ViewState->PathTracingFrameIndependentTemporalSeed,
+					ViewState->VarianceMipTreeDimensions,
+					*ViewState->VarianceMipTree,
+					TileOffset,
 					RadianceRT->GetRenderTargetItem().UAV,
 					SampleCountRT->GetRenderTargetItem().UAV,
 					PixelPositionRT->GetRenderTargetItem().UAV,
@@ -682,7 +739,7 @@ void FDeferredShadingSceneRenderer::RenderPathTracing(
 			}
 #endif
 		}
-		else
+		else if (ViewState->PathTracingSPP < MaxSPP)
 		{
 			FIntVector TileOffset;
 			TileOffset.X = bWiperMode > 0 ? WipeOffsetX : 0;
@@ -695,8 +752,14 @@ void FDeferredShadingSceneRenderer::RenderPathTracing(
 				View.RayTracingScene,
 				View.ViewUniformBuffer,
 				SceneTexturesUniformBufferRHI,
+				PathTracingData,
+				AdaptiveSamplingData,
 				Scene->Lights,
-				SPPCount, FrameIndependentTemporalSeed, ViewState->VarianceMipTreeDimensions, *ViewState->VarianceMipTree, TileOffset,
+				ViewState->PathTracingSPP,
+				ViewState->PathTracingFrameIndependentTemporalSeed,
+				ViewState->VarianceMipTreeDimensions,
+				*ViewState->VarianceMipTree,
+				TileOffset,
 				RadianceRT->GetRenderTargetItem().UAV,
 				SampleCountRT->GetRenderTargetItem().UAV,
 				PixelPositionRT->GetRenderTargetItem().UAV,
@@ -726,7 +789,7 @@ void FDeferredShadingSceneRenderer::RenderPathTracing(
 		// Single GPU for launching compute shaders
 		SCOPED_GPU_MASK(RHICmdList, FRHIGPUMask::FromIndex(MainGPUIndex));
 
-		if ((SPPCount % CVarPathTracingRayCountFrequency.GetValueOnRenderThread() == 0))
+		if ((ViewState->PathTracingSPP % CVarPathTracingRayCountFrequency.GetValueOnRenderThread() == 0))
 		{
 			ComputeRayCount(RHICmdList, View, RayCountPerPixelRT->GetRenderTargetItem().ShaderResourceTexture);
 		}
@@ -811,7 +874,6 @@ void FDeferredShadingSceneRenderer::RenderPathTracing(
 		GraphicsPSOInit.PrimitiveType = PT_TriangleList;
 		SetGraphicsPipelineState(RHICmdList, GraphicsPSOInit);
 
-		//for (int32 ViewIndex = 0; ViewIndex < Views.Num(); ++ViewIndex)
 		{
 			RHICmdList.SetViewport(View.ViewRect.Min.X, View.ViewRect.Min.Y, 0.0f, View.ViewRect.Max.X, View.ViewRect.Max.Y, 1.0f);
 			FTextureRHIRef RadianceRedTexture = RadianceSortedRedRT->GetRenderTargetItem().ShaderResourceTexture;
@@ -823,21 +885,16 @@ void FDeferredShadingSceneRenderer::RenderPathTracing(
 			FTextureRHIRef CumulativeRadianceTexture = GBlackTexture->TextureRHI;
 			FTextureRHIRef CumulativeSampleCount = GBlackTexture->TextureRHI;
 
-			int32 SamplesPerPixelCVar = CVarPathTracingSamplesPerPixel.GetValueOnRenderThread();
-			int32 PathTracingSamplesPerPixel = SamplesPerPixelCVar > -1 ? SamplesPerPixelCVar : View.FinalPostProcessSettings.PathTracingSamplesPerPixel;
-			if (ViewState->PathTracingIrradianceRT && SPPCount < PathTracingSamplesPerPixel)
+			if (ViewState->PathTracingIrradianceRT)
 			{
+				// If we still have a valid accumulated image - re-use it
 				CumulativeRadianceTexture = ViewState->PathTracingIrradianceRT->GetRenderTargetItem().ShaderResourceTexture;
 				CumulativeSampleCount = ViewState->PathTracingSampleCountRT->GetRenderTargetItem().ShaderResourceTexture;
-				SPPCount++;
 			}
 			else
 			{
-				SPPCount = 0;
+				// Accumulated image got flushed by invalidating, start over
 			}
-
-			++FrameIndependentTemporalSeed;
-
 			PixelShader->SetParameters(RHICmdList, View, RadianceRedTexture, RadianceGreenTexture, RadianceBlueTexture, RadianceAlphaTexture, SampleCountTexture, CumulativeRadianceTexture, CumulativeSampleCount);
 
 			int32 DispatchSizeX = View.ViewRect.Size().X;
@@ -854,6 +911,10 @@ void FDeferredShadingSceneRenderer::RenderPathTracing(
 		}
 		RHICmdList.EndRenderPass();
 
+		// Bump counters for next frame
+		++ViewState->PathTracingSPP;
+		++ViewState->PathTracingFrameIndependentTemporalSeed;
+
 		RHICmdList.CopyToResolveTarget(OutputRadianceRT->GetRenderTargetItem().TargetableTexture, OutputRadianceRT->GetRenderTargetItem().ShaderResourceTexture, FResolveParams());
 		RHICmdList.CopyToResolveTarget(OutputSampleCountRT->GetRenderTargetItem().TargetableTexture, OutputSampleCountRT->GetRenderTargetItem().ShaderResourceTexture, FResolveParams());
 		GVisualizeTexture.SetCheckPoint(RHICmdList, OutputRadianceRT);
@@ -864,7 +925,7 @@ void FDeferredShadingSceneRenderer::RenderPathTracing(
 		ViewState->PathTracingSampleCountRT = OutputSampleCountRT;
 
 		// Process variance mip for adaptive sampling
-		if (SPPCount % CVarPathTracingVarianceMapRebuildFrequency.GetValueOnRenderThread() == 0)
+		if (ViewState->PathTracingSPP % CVarPathTracingVarianceMapRebuildFrequency.GetValueOnRenderThread() == 0)
 		{
 			SCOPED_GPU_STAT(RHICmdList, Stat_GPU_PathTracingBuildVarianceMipTree);
 
diff --git a/Engine/Source/Runtime/Renderer/Private/RendererScene.cpp b/Engine/Source/Runtime/Renderer/Private/RendererScene.cpp
--- a/Engine/Source/Runtime/Renderer/Private/RendererScene.cpp
+++ b/Engine/Source/Runtime/Renderer/Private/RendererScene.cpp
@@ -230,6 +230,8 @@ FSceneViewState::FSceneViewState()
 	VarianceMipTreeDimensions = FIntVector(0);
 	VarianceMipTree = new FRWBuffer;
 	PathTracingRect = FIntRect(0, 0, 0, 0);
+	PathTracingSPP = 0;
+	PathTracingFrameIndependentTemporalSeed = 0;
 	TotalRayCount = 0;
 	TotalRayCountBuffer = new FRWBuffer;
 	if (GetFeatureLevel() >= ERHIFeatureLevel::SM5)
@@ -4168,6 +4170,8 @@ void FScene::UpdateAllPrimitiveSceneInfos(FRHICommandListImmediate& RHICmdList,
 				}
 
 				AddPrimitiveToUpdateGPU(*this, PrimitiveIndex);
+
+				// Invalidate PathTraced image because we added something to the scene
 				bPathTracingNeedsInvalidation = true;
 
 				DistanceFieldSceneData.AddPrimitive(PrimitiveSceneInfo);
@@ -4322,7 +4326,7 @@ void FScene::UpdateAllPrimitiveSceneInfos(FRHICommandListImmediate& RHICmdList,
 		SCOPED_NAMED_EVENT(FScene_DeletePrimitiveSceneInfo, FColor::Red);
 		for (FPrimitiveSceneInfo* PrimitiveSceneInfo : DeletedSceneInfos)
 		{
-			// It is possible that hte HitProxies list isn't empty if PrimitiveSceneInfo was Added/Removed in same frame
+			// It is possible that the HitProxies list isn't empty if PrimitiveSceneInfo was Added/Removed in same frame
 			// Delete the PrimitiveSceneInfo on the game thread after the rendering thread has processed its removal.
 			// This must be done on the game thread because the hit proxy references (and possibly other members) need to be freed on the game thread.
 			struct DeferDeleteHitProxies : FDeferredCleanupInterface
@@ -4335,6 +4339,9 @@ void FScene::UpdateAllPrimitiveSceneInfos(FRHICommandListImmediate& RHICmdList,
 			// free the primitive scene proxy.
 			delete PrimitiveSceneInfo->Proxy;
 			delete PrimitiveSceneInfo;
+
+			// Invalidate PathTraced image because we removed something from the scene
+			bPathTracingNeedsInvalidation = true;
 		}
 	}
 
diff --git a/Engine/Source/Runtime/Renderer/Private/ScenePrivate.h b/Engine/Source/Runtime/Renderer/Private/ScenePrivate.h
--- a/Engine/Source/Runtime/Renderer/Private/ScenePrivate.h
+++ b/Engine/Source/Runtime/Renderer/Private/ScenePrivate.h
@@ -1004,6 +1004,9 @@ public:
 	FSamplerStateRHIRef MaterialTextureBilinearClampedSamplerCache;
 
 #if RHI_RAYTRACING
+	// Invalidates cached results related to the path tracer so accumulated rendering can start over
+	void PathTracingInvalidate();
+
 	// Reference path tracing cached results
 	TRefCountPtr<IPooledRenderTarget> PathTracingIrradianceRT;
 	TRefCountPtr<IPooledRenderTarget> PathTracingSampleCountRT;
@@ -1011,6 +1014,13 @@ public:
 	FRWBuffer* VarianceMipTree;
 	FIntVector VarianceMipTreeDimensions;
 
+	// Path tracer samples/pixel
+	uint32_t PathTracingSPP;
+
+	// The frame independent temporal seed, not reset on invalidation unlike PathTracingSPP to avoid
+	// the "screen door" effect and reduce temporal aliasing
+	uint32_t PathTracingFrameIndependentTemporalSeed;
+
 	// Path tracer ray counter
 	uint32 TotalRayCount;
 	FRWBuffer* TotalRayCountBuffer;
diff --git a/Engine/Source/Runtime/Renderer/Private/SceneVisibility.cpp b/Engine/Source/Runtime/Renderer/Private/SceneVisibility.cpp
--- a/Engine/Source/Runtime/Renderer/Private/SceneVisibility.cpp
+++ b/Engine/Source/Runtime/Renderer/Private/SceneVisibility.cpp
@@ -3411,22 +3411,16 @@ void FSceneRenderer::PreVisibilityFrameSetup(FRHICommandListImmediate& RHICmdLis
 				ViewState->PrevFrameViewInfo.ViewMatrices.GetViewOrigin(),
 				0.18f /*degree*/, 0.1f /*cm*/);
 			const bool bIsProjMatrixDifferent = View.ViewMatrices.GetProjectionNoAAMatrix() != View.ViewState->PrevFrameViewInfo.ViewMatrices.GetProjectionNoAAMatrix();
-			const bool bInvalidatePathTracer = View.RayTracingRenderMode == ERayTracingRenderMode::PathTracing &&
-			(
+			const bool bInvalidatePathTracer = (
 				bResetCamera ||
 				Scene->bPathTracingNeedsInvalidation ||
-				View.ViewRect != ViewState->PathTracingRect ||
 				bIsProjMatrixDifferent ||
 				bIsThereALargeMomvement
 			);
 
 			if (bInvalidatePathTracer)
 			{
-				ViewState->PathTracingIrradianceRT.SafeRelease();
-				ViewState->PathTracingSampleCountRT.SafeRelease();
-				ViewState->VarianceMipTreeDimensions = FIntVector(0);
-				ViewState->PathTracingRect = View.ViewRect;
-				ViewState->TotalRayCount = 0;
+				ViewState->PathTracingInvalidate();
 				Scene->bPathTracingNeedsInvalidation = false;
 			}
 #endif // RHI_RAYTRACING
