SDF - use min non-uniform axis to maintain conservative SDF distance. Distances are now underestimated, instead of overestimated. This guarantees that rays will always hit the surface.

diff --git a/Engine/Source/Runtime/Core/Public/Math/Matrix.h b/Engine/Source/Runtime/Core/Public/Math/Matrix.h
--- a/Engine/Source/Runtime/Core/Public/Math/Matrix.h
+++ b/Engine/Source/Runtime/Core/Public/Math/Matrix.h
@@ -201,6 +201,9 @@ public:
 	/** Scale the translation part of the matrix by the supplied vector. */
 	inline void ScaleTranslation(const FVector& Scale3D);
 
+	/** @return the minimum magnitude of any row of the matrix. */
+	inline float GetMinimumAxisScale() const;
+
 	/** @return the maximum magnitude of any row of the matrix. */
 	inline float GetMaximumAxisScale() const;
 
diff --git a/Engine/Source/Runtime/Core/Public/Math/Matrix.inl b/Engine/Source/Runtime/Core/Public/Math/Matrix.inl
--- a/Engine/Source/Runtime/Core/Public/Math/Matrix.inl
+++ b/Engine/Source/Runtime/Core/Public/Math/Matrix.inl
@@ -500,6 +500,19 @@ inline bool FMatrix::ContainsNaN() const
 	return false;
 }
 
+/** @return the minimum magnitude of any row of the matrix. */
+inline float FMatrix::GetMinimumAxisScale() const
+{
+	const float MaxRowScaleSquared = FMath::Min(
+		GetScaledAxis(EAxis::X).SizeSquared(),
+		FMath::Min(
+			GetScaledAxis(EAxis::Y).SizeSquared(),
+			GetScaledAxis(EAxis::Z).SizeSquared()
+		)
+	);
+	return FMath::Sqrt(MaxRowScaleSquared);
+}
+
 /** @return the maximum magnitude of any row of the matrix. */
 inline float FMatrix::GetMaximumAxisScale() const
 {
diff --git a/Engine/Source/Runtime/Renderer/Private/DistanceFieldObjectManagement.cpp b/Engine/Source/Runtime/Renderer/Private/DistanceFieldObjectManagement.cpp
--- a/Engine/Source/Runtime/Renderer/Private/DistanceFieldObjectManagement.cpp
+++ b/Engine/Source/Runtime/Renderer/Private/DistanceFieldObjectManagement.cpp
@@ -595,8 +595,7 @@ void FDeferredShadingSceneRenderer::UpdateGlobalDistanceFieldObjectBuffers(FRDGB
 									bool bThrottled;
 									PrimAndInst.Primitive->Proxy->GetDistancefieldAtlasData(LocalVolumeBounds, DistanceMinMax, BlockMin, BlockSize, bBuiltAsIfTwoSided, SelfShadowBias, bThrottled);
 
-									FMatrix LocalToWorld = PrimAndInst.LocalToWorld;
-									const float MaxScale = LocalToWorld.GetMaximumAxisScale();
+									const FMatrix LocalToWorld = PrimAndInst.LocalToWorld;
 
 									const FMatrix VolumeToWorld = FScaleMatrix(LocalVolumeBounds.GetExtent())
 										* FTranslationMatrix(LocalVolumeBounds.GetCenter())
@@ -625,7 +624,7 @@ void FDeferredShadingSceneRenderer::UpdateGlobalDistanceFieldObjectBuffers(FRDGB
 									//float3 VolumeUV = (VolumePosition / LocalPositionExtent * .5f * UVScale + .5f * UVScale + UVAdd;
 									const FVector LocalPositionExtent = LocalVolumeBounds.GetExtent() / FVector(MaxExtent);
 									const FVector UVScale = FVector(BlockSize) * InvTextureDim;
-									const float VolumeScale = UniformScaleVolumeToWorld.GetMaximumAxisScale();
+									const float VolumeMinScale = UniformScaleVolumeToWorld.GetMinimumAxisScale();
 
 									const FMatrix WorldToVolumeT = UniformScaleVolumeToWorld.Inverse().GetTransposed();
 									// WorldToVolumeT
@@ -646,7 +645,7 @@ void FDeferredShadingSceneRenderer::UpdateGlobalDistanceFieldObjectBuffers(FRDGB
 
 									// UVScale, VolumeScale and sign gives bGeneratedAsTwoSided
 									const float WSign = bBuiltAsIfTwoSided ? -1 : 1;
-									UploadObjectData[4] = FVector4(FVector(BlockSize) * InvTextureDim * .5f / LocalPositionExtent, WSign * VolumeScale);
+									UploadObjectData[4] = FVector4(FVector(BlockSize) * InvTextureDim * .5f / LocalPositionExtent, WSign * VolumeMinScale);
 
 									// UVAdd
 									UploadObjectData[5] = FVector4(FVector(BlockMin) * InvTextureDim + .5f * UVScale, SelfShadowBias);
