Moving velocity pass earlier and adding the option to write depth after the depth pass.

diff --git a/Engine/Source/Runtime/Renderer/Private/DeferredShadingRenderer.cpp b/Engine/Source/Runtime/Renderer/Private/DeferredShadingRenderer.cpp
--- a/Engine/Source/Runtime/Renderer/Private/DeferredShadingRenderer.cpp
+++ b/Engine/Source/Runtime/Renderer/Private/DeferredShadingRenderer.cpp
@@ -1493,8 +1493,11 @@ void FDeferredShadingSceneRenderer::Render(FRHICommandListImmediate& RHICmdList)
 		&& !ViewFamily.EngineShowFlags.Wireframe
 		&& !ViewFamily.EngineShowFlags.LightMapDensity;
 
+	// if DDM_AllOpaqueNoVelocity was used, then velocity should have already been rendered as well
+	const bool bIsEarlyDepthComplete = (DepthPass.EarlyZPassMode == DDM_AllOpaque || DepthPass.EarlyZPassMode == DDM_AllOpaqueNoVelocity);
+
 	// Use read-only depth in the base pass if we have a full depth prepass.
-	const bool bAllowReadOnlyDepthBasePass = DepthPass.EarlyZPassMode == DDM_AllOpaque
+	const bool bAllowReadOnlyDepthBasePass = bIsEarlyDepthComplete
 		&& !ViewFamily.EngineShowFlags.ShaderComplexity
 		&& !ViewFamily.UseDebugViewPS()
 		&& !ViewFamily.EngineShowFlags.Wireframe
@@ -1754,6 +1757,21 @@ void FDeferredShadingSceneRenderer::Render(FRHICommandListImmediate& RHICmdList)
 		Nanite::GStreamingManager.EndAsyncUpdate(GraphBuilder);
 	}
 
+	FRDGTextureRef VelocityTexture = TryRegisterExternalTexture(GraphBuilder, SceneContext.SceneVelocity);
+
+	const bool bShouldRenderVelocities = ShouldRenderVelocities();
+	const bool bBasePassCanOutputVelocity = FVelocityRendering::BasePassCanOutputVelocity(FeatureLevel);
+	const bool bUseSelectiveBasePassOutputs = IsUsingSelectiveBasePassOutputs(ShaderPlatform);
+	const bool bIsViewCompatible = Views.Num() > 0;
+	const bool bHairEnable = HairStrandsBookmarkParameters.bHasElements && bIsViewCompatible && IsHairStrandsEnabled(EHairStrandsShaderType::Strands, Views[0].GetShaderPlatform());
+
+	{
+		// Even if !bShouldRenderVelocities, the velocity buffer must be bound because it's a compile time option for the shader. Also, this pass was modified
+		// to allocate the velocity texture, since we need it earlier in the frame and itss usage is separate from the rest of the GBuffer passes.
+		SceneContext.PreallocGBufferTargets();
+		SceneContext.AllocVelocityTarget(RHICmdList);
+	}
+
 	{
 		AddSetCurrentStatPass(GraphBuilder, GET_STATID(STAT_CLM_PrePass));
 
@@ -1775,9 +1793,20 @@ void FDeferredShadingSceneRenderer::Render(FRHICommandListImmediate& RHICmdList)
 			// We didn't do the prepass, but we still want the HMD mask if there is one
 			RenderPrePassHMD(GraphBuilder, SceneDepthTexture.Target);
 		}
+
 		AddSetCurrentStatPass(GraphBuilder, GET_STATID(STAT_CLM_AfterPrePass));
 		AddServiceLocalQueuePass(GraphBuilder);
 
+		// special pass for DDM_AllOpaqueNoVelocity, which uses the velocity pass to finish the early depth pass write
+		if (bShouldRenderVelocities && Scene->EarlyZPassMode == DDM_AllOpaqueNoVelocity)
+		{
+			// Render the velocities of movable objects
+			AddSetCurrentStatPass(GraphBuilder, GET_STATID(STAT_CLM_Velocity));
+			RenderVelocities(GraphBuilder, SceneDepthTexture.Resolve, VelocityTexture, nullptr, EVelocityPass::Opaque, bHairEnable);
+			AddSetCurrentStatPass(GraphBuilder, GET_STATID(STAT_CLM_AfterVelocity));
+			AddServiceLocalQueuePass(GraphBuilder);
+		}
+
 		if (bDoInitViewAftersPrepass)
 		{
 			{
@@ -1838,7 +1867,7 @@ void FDeferredShadingSceneRenderer::Render(FRHICommandListImmediate& RHICmdList)
 				Nanite::FCullingContext CullingContext = Nanite::InitCullingContext(
 					GraphBuilder,
 					*Scene,
-					DepthPass.EarlyZPassMode != DDM_AllOpaque ? View.PrevViewInfo.NaniteHZB : View.PrevViewInfo.HZB,
+					!bIsEarlyDepthComplete ? View.PrevViewInfo.NaniteHZB : View.PrevViewInfo.HZB,
 					View.PrevViewInfo.ViewRect,
 					bTwoPassOcclusion,
 					bUpdateStreaming,
@@ -1862,7 +1891,7 @@ void FDeferredShadingSceneRenderer::Render(FRHICommandListImmediate& RHICmdList)
 					bExtractStats
 				);
 
-				if (DepthPass.EarlyZPassMode != DDM_AllOpaque && bTwoPassOcclusion && View.ViewState)
+				if (!bIsEarlyDepthComplete && bTwoPassOcclusion && View.ViewState)
 				{
 					// Won't have a complete SceneDepth for post pass so can't use complete HZB for main pass or it will poke holes in the post pass HZB killing occlusion culling.
 					RDG_EVENT_SCOPE(GraphBuilder, "BuildNaniteHZB");
@@ -1894,10 +1923,6 @@ void FDeferredShadingSceneRenderer::Render(FRHICommandListImmediate& RHICmdList)
 	ESceneTextureSetupMode SceneTexturesSetupMode = ESceneTextureSetupMode::SceneDepth;
 	TRDGUniformBufferRef<FSceneTextureUniformParameters> SceneTextures = CreateSceneTextureUniformBuffer(GraphBuilder, FeatureLevel, SceneTexturesSetupMode);
 
-	const bool bShouldRenderVelocities = ShouldRenderVelocities();
-	const bool bBasePassCanOutputVelocity = FVelocityRendering::BasePassCanOutputVelocity(FeatureLevel);
-	const bool bUseSelectiveBasePassOutputs = IsUsingSelectiveBasePassOutputs(ShaderPlatform);
-
 	AddResolveSceneDepthPass(GraphBuilder, Views, SceneDepthTexture);
 
 	// NOTE: The ordering of the lights is used to select sub-sets for different purposes, e.g., those that support clustered deferred.
@@ -1913,7 +1938,6 @@ void FDeferredShadingSceneRenderer::Render(FRHICommandListImmediate& RHICmdList)
 
 	{
 		SCOPE_CYCLE_COUNTER(STAT_FDeferredShadingSceneRenderer_AllocGBufferTargets);
-		SceneContext.PreallocGBufferTargets(); // Even if !bShouldRenderVelocities, the velocity buffer must be bound because it's a compile time option for the shader.
 		SceneContext.AllocGBufferTargets(RHICmdList);
 	}
 
@@ -1948,7 +1972,7 @@ void FDeferredShadingSceneRenderer::Render(FRHICommandListImmediate& RHICmdList)
 	};
 
 	// Early occlusion queries
-	const bool bOcclusionBeforeBasePass = !bNaniteEnabled && ((DepthPass.EarlyZPassMode == EDepthDrawingMode::DDM_AllOccluders) || (DepthPass.EarlyZPassMode == EDepthDrawingMode::DDM_AllOpaque));
+	const bool bOcclusionBeforeBasePass = !bNaniteEnabled && ((DepthPass.EarlyZPassMode == EDepthDrawingMode::DDM_AllOccluders) || bIsEarlyDepthComplete);
 
 	if (bOcclusionBeforeBasePass)
 	{
@@ -2031,8 +2055,6 @@ void FDeferredShadingSceneRenderer::Render(FRHICommandListImmediate& RHICmdList)
 
 	FHairStrandsRenderingData* HairDatas = nullptr;
 	FHairStrandsRenderingData HairDatasStorage;
-	const bool bIsViewCompatible = Views.Num() > 0;
-	const bool bHairEnable = HairStrandsBookmarkParameters.bHasElements && bIsViewCompatible && IsHairStrandsEnabled(EHairStrandsShaderType::Strands, Views[0].GetShaderPlatform());
 
 	FRDGTextureRef ForwardScreenSpaceShadowMaskTexture = nullptr;
 	FRDGTextureRef ForwardScreenSpaceShadowMaskHairTexture = nullptr;
@@ -2139,8 +2161,6 @@ void FDeferredShadingSceneRenderer::Render(FRHICommandListImmediate& RHICmdList)
 		}
 	}
 
-	FRDGTextureRef VelocityTexture = TryRegisterExternalTexture(GraphBuilder, SceneContext.SceneVelocity);
-
 	// Rebuild scene textures to include GBuffers.
 	SceneTexturesSetupMode |= ESceneTextureSetupMode::GBuffers;
 	SceneTextures = CreateSceneTextureUniformBuffer(GraphBuilder, FeatureLevel, SceneTexturesSetupMode);
@@ -2241,8 +2261,11 @@ void FDeferredShadingSceneRenderer::Render(FRHICommandListImmediate& RHICmdList)
 	FRDGTextureRef SkyLightHitDistanceTexture = nullptr;
 #endif
 
+	// TODO: Keeping the velocities here for testing, but if that works, this pass will be remove and DDM_AllOpaqueNoVelocity will be the only option with
+	// DBuffer decals enabled.
+
 	// If bBasePassCanOutputVelocity is set, basepass fully writes the velocity buffer unless bUseSelectiveBasePassOutputs is enabled.
-	if (bShouldRenderVelocities && (!bBasePassCanOutputVelocity || bUseSelectiveBasePassOutputs))
+	if (bShouldRenderVelocities && (!bBasePassCanOutputVelocity || bUseSelectiveBasePassOutputs) && (Scene->EarlyZPassMode != DDM_AllOpaqueNoVelocity))
 	{
 		// Render the velocities of movable objects
 		AddSetCurrentStatPass(GraphBuilder, GET_STATID(STAT_CLM_Velocity));
diff --git a/Engine/Source/Runtime/Renderer/Private/DepthRendering.cpp b/Engine/Source/Runtime/Renderer/Private/DepthRendering.cpp
--- a/Engine/Source/Runtime/Renderer/Private/DepthRendering.cpp
+++ b/Engine/Source/Runtime/Renderer/Private/DepthRendering.cpp
@@ -67,6 +67,11 @@ static TAutoConsoleVariable<int32> CVarStencilForLODDither(
 	TEXT("Forces a full prepass when enabled."),
 	ECVF_RenderThreadSafe | ECVF_ReadOnly);
 
+static TAutoConsoleVariable<int32> CVarDepthPassMergedWithVelocity(
+	TEXT("r.DepthPassMergedWithVelocity"),
+	0,
+	TEXT("If enabled, and we are doing a full depth pass, then the depth pass will ignore movable objects and the velocity pass will write depth directly after the depth pass. After the velocity pass is finished, a full opaque depth-only texture is ready."));
+
 extern bool IsHMDHiddenAreaMaskActive();
 
 FDepthPassInfo GetDepthPassInfo(const FScene* Scene)
@@ -116,6 +121,8 @@ const TCHAR* GetDepthDrawingModeString(EDepthDrawingMode Mode)
 		return TEXT("DDM_AllOccluders");
 	case DDM_AllOpaque:
 		return TEXT("DDM_AllOpaque");
+	case DDM_AllOpaqueNoVelocity:
+		return TEXT("DDM_AllOpaqueNoVelocity");
 	default:
 		check(0);
 	}
@@ -829,6 +836,34 @@ void FDepthPassMeshProcessor::AddMeshBatch(const FMeshBatch& RESTRICT MeshBatch,
 		}
 	}
 
+	// if we are skipping movable objects in early Z, which can happen in DDM_AllOpaqueNoVelocity
+	if (EarlyZPassMode == DDM_AllOpaqueNoVelocity && PrimitiveSceneProxy&& ViewIfDynamicMeshCommand)
+	{
+		// We should ideally check to see if we are using the FOpaqueVelocityMeshProcessor or FTranslucentVelocityMeshProcessor. But for
+		// the object to get here, it would already be culled if it was translucent. We can safely use the FOpaqueVelocityMeshProcessor.
+
+		// This logic is copy/paste/modified from FOpaqueVelocityMeshProcessor::AddMeshBatch(), but ideally we should clean it up into
+		// a single function that is shared to avoid breakages from code changes.
+		EShaderPlatform ShaderPlatform = ViewIfDynamicMeshCommand->GetShaderPlatform();
+		if (!FOpaqueVelocityMeshProcessor::PrimitiveCanHaveVelocity(ShaderPlatform, PrimitiveSceneProxy))
+		{
+			bDraw = false;
+		}
+
+		if (!FOpaqueVelocityMeshProcessor::PrimitiveHasVelocityForFrame(PrimitiveSceneProxy))
+		{
+			bDraw = false;
+		}
+
+		checkSlow(ViewIfDynamicMeshCommand->bIsViewInfo);
+		FViewInfo* ViewInfo = (FViewInfo*)ViewIfDynamicMeshCommand;
+
+		if (!FOpaqueVelocityMeshProcessor::PrimitiveHasVelocityForView(*ViewInfo, PrimitiveSceneProxy))
+		{
+			bDraw = false;
+		}
+	}
+
 	if (bDraw)
 	{
 		// Determine the mesh's material and blend mode.
diff --git a/Engine/Source/Runtime/Renderer/Private/DepthRendering.h b/Engine/Source/Runtime/Renderer/Private/DepthRendering.h
--- a/Engine/Source/Runtime/Renderer/Private/DepthRendering.h
+++ b/Engine/Source/Runtime/Renderer/Private/DepthRendering.h
@@ -29,6 +29,8 @@ enum EDepthDrawingMode
 	DDM_AllOpaque		= 3,
 	// Masked materials only
 	DDM_MaskedOnly = 4,
+	// Full prepass, every object must be drawn and every pixel must match the base pass depth, except dynamic geometry which will render in the Velocity pass
+	DDM_AllOpaqueNoVelocity	= 5,
 };
 
 extern const TCHAR* GetDepthDrawingModeString(EDepthDrawingMode Mode);
diff --git a/Engine/Source/Runtime/Renderer/Private/MobileShadingRenderer.cpp b/Engine/Source/Runtime/Renderer/Private/MobileShadingRenderer.cpp
--- a/Engine/Source/Runtime/Renderer/Private/MobileShadingRenderer.cpp
+++ b/Engine/Source/Runtime/Renderer/Private/MobileShadingRenderer.cpp
@@ -529,6 +529,8 @@ void FMobileSceneRenderer::Render(FRHICommandListImmediate& RHICmdList)
 	if (bDeferredShading)
 	{
 		ETextureCreateFlags AddFlags = bRequiresMultiPass ? TexCreate_InputAttachmentRead : (TexCreate_InputAttachmentRead | TexCreate_Memoryless);
+
+		SceneContext.AllocVelocityTarget(RHICmdList);
 		SceneContext.AllocGBufferTargets(RHICmdList, AddFlags);
 	}
 
diff --git a/Engine/Source/Runtime/Renderer/Private/PostProcess/SceneRenderTargets.cpp b/Engine/Source/Runtime/Renderer/Private/PostProcess/SceneRenderTargets.cpp
--- a/Engine/Source/Runtime/Renderer/Private/PostProcess/SceneRenderTargets.cpp
+++ b/Engine/Source/Runtime/Renderer/Private/PostProcess/SceneRenderTargets.cpp
@@ -1331,6 +1331,18 @@ static FPooledRenderTargetDesc GetDescFromRenderTarget(FIntPoint BufferSize, con
 	return Desc;
 }
 
+void FSceneRenderTargets::AllocVelocityTarget(FRHICommandList& RHICmdList)
+{
+	const EShaderPlatform ShaderPlatform = GetFeatureLevelShaderPlatform(CurrentFeatureLevel);
+
+	if (bAllocateVelocityGBuffer)
+	{
+		FPooledRenderTargetDesc VelocityRTDesc = Translate(FVelocityRendering::GetRenderTargetDesc(ShaderPlatform));
+		VelocityRTDesc.Flags |= GFastVRamConfig.GBufferVelocity;
+		GRenderTargetPool.FindFreeElement(RHICmdList, VelocityRTDesc, SceneVelocity, TEXT("GBufferVelocity"));
+	}
+}
+
 void FSceneRenderTargets::AllocGBufferTargets(FRHICommandList& RHICmdList, ETextureCreateFlags AddTargetableFlags)
 {	
 	// AdjustGBufferRefCount +1 doesn't match -1 (within the same frame)
@@ -1358,13 +1370,6 @@ void FSceneRenderTargets::AllocGBufferTargets(FRHICommandList& RHICmdList, EText
 			Desc.Flags |= GFastVRamConfig.GBufferA;
 			GRenderTargetPool.FindFreeElement(RHICmdList, Desc, GBufferA, TEXT("GBufferA"));
 		}
-
-		if (bAllocateVelocityGBuffer)
-		{
-			FPooledRenderTargetDesc VelocityRTDesc = Translate(FVelocityRendering::GetRenderTargetDesc(ShaderPlatform));
-			VelocityRTDesc.Flags |= GFastVRamConfig.GBufferVelocity;
-			GRenderTargetPool.FindFreeElement(RHICmdList, VelocityRTDesc, SceneVelocity, TEXT("GBufferVelocity"));
-		}
 	}
 	else
 	{
@@ -1428,14 +1433,6 @@ void FSceneRenderTargets::AllocGBufferTargets(FRHICommandList& RHICmdList, EText
 		// otherwise we have a severe problem
 		check(GBufferA);
 
-		if (IndexVelocity >= 0)
-		{
-			//FPooledRenderTargetDesc VelocityRTDesc = Translate(FVelocityRendering::GetRenderTargetDesc(ShaderPlatform));
-			FPooledRenderTargetDesc Desc = GetDescFromRenderTarget(BufferSize, GBufferInfo.Targets[IndexVelocity], AddTargetableFlags);
-			Desc.Flags |= GFastVRamConfig.GBufferVelocity;
-			GRenderTargetPool.FindFreeElement(RHICmdList, Desc, SceneVelocity, TEXT("GBufferVelocity"));
-		}
-
 	}
 
 	FVelocityRendering::GetRenderTargetDesc(ShaderPlatform);
diff --git a/Engine/Source/Runtime/Renderer/Private/PostProcess/SceneRenderTargets.h b/Engine/Source/Runtime/Renderer/Private/PostProcess/SceneRenderTargets.h
--- a/Engine/Source/Runtime/Renderer/Private/PostProcess/SceneRenderTargets.h
+++ b/Engine/Source/Runtime/Renderer/Private/PostProcess/SceneRenderTargets.h
@@ -410,7 +410,7 @@ public:
 	void AdjustGBufferRefCount(FRHICommandList& RHICmdList, int Delta);
 
 	void PreallocGBufferTargets();
-	bool AllocateGBufferTarget(const FName& TargetName, ETextureCreateFlags TargetFlags) const;
+	bool AllocateGBufferTargets(const FName& TargetName, ETextureCreateFlags TargetFlags) const;
 
 	EPixelFormat GetGBufferAFormat() const;
 	EPixelFormat GetGBufferBFormat() const;
@@ -420,6 +420,7 @@ public:
 	EPixelFormat GetGBufferFFormat() const;
 	void AllocGBufferTargets(FRHICommandList& RHICmdList);
 	void AllocGBufferTargets(FRHICommandList& RHICmdList, ETextureCreateFlags AddTargetableFlags);
+	void AllocVelocityTarget(FRHICommandList& RHICmdList);
 
 	void AllocateReflectionTargets(FRHICommandList& RHICmdList, int32 TargetSize);
 
diff --git a/Engine/Source/Runtime/Renderer/Private/RendererScene.cpp b/Engine/Source/Runtime/Renderer/Private/RendererScene.cpp
--- a/Engine/Source/Runtime/Renderer/Private/RendererScene.cpp
+++ b/Engine/Source/Runtime/Renderer/Private/RendererScene.cpp
@@ -71,6 +71,7 @@
 #include "Rendering/StaticLightingSystemInterface.h"
 #endif
 
+
 /** Affects BasePassPixelShader.usf so must relaunch editor to recompile shaders. */
 static TAutoConsoleVariable<int32> CVarEarlyZPassOnlyMaterialMasking(
 	TEXT("r.EarlyZPassOnlyMaterialMasking"),
@@ -128,6 +129,8 @@ TGlobalResource< FGlobalDitherUniformBuffer > GDitherFadedInUniformBuffer;
 
 static FThreadSafeCounter FSceneViewState_UniqueID;
 
+extern bool IsVelocityMergedWithDepthPass();
+
 /**
  * Holds the info to update SpeedTree wind per unique tree object in the scene, instead of per instance
  */
@@ -3431,11 +3434,11 @@ void FScene::UpdateEarlyZPassMode()
 		if (ShouldForceFullDepthPass(ShaderPlatform))
 		{
 			// DBuffer decals and stencil LOD dithering force a full prepass
-			EarlyZPassMode = DDM_AllOpaque;
-			bEarlyZPassMovable = true;
+			EarlyZPassMode = IsVelocityMergedWithDepthPass() ? DDM_AllOpaqueNoVelocity : DDM_AllOpaque;
+			bEarlyZPassMovable = IsVelocityMergedWithDepthPass() ? false : true;
 		}
 
-		if (EarlyZPassMode == DDM_AllOpaque
+		if ((EarlyZPassMode == DDM_AllOpaque || EarlyZPassMode == DDM_AllOpaqueNoVelocity)
 			&& CVarBasePassWriteDepthEvenWithFullPrepass.GetValueOnAnyThread() == 0)
 		{
 			DefaultBasePassDepthStencilAccess = FExclusiveDepthStencil::DepthRead_StencilWrite;
diff --git a/Engine/Source/Runtime/Renderer/Private/VelocityRendering.cpp b/Engine/Source/Runtime/Renderer/Private/VelocityRendering.cpp
--- a/Engine/Source/Runtime/Renderer/Private/VelocityRendering.cpp
+++ b/Engine/Source/Runtime/Renderer/Private/VelocityRendering.cpp
@@ -61,6 +61,13 @@ bool IsVelocityWaitForTasksEnabled()
 	return IsParallelVelocity() && (CVarRHICmdFlushRenderThreadTasksVelocityPass.GetValueOnRenderThread() > 0 || CVarRHICmdFlushRenderThreadTasks.GetValueOnRenderThread() > 0);
 }
 
+bool IsVelocityMergedWithDepthPass()
+{
+	static const auto CVarDepthPassMergedWithVelocity = IConsoleManager::Get().FindTConsoleVariableDataInt(TEXT("r.DepthPassMergedWithVelocity"));
+	bool bIsMerged = (CVarDepthPassMergedWithVelocity && CVarDepthPassMergedWithVelocity->GetValueOnAnyThread() != 0);
+	return bIsMerged;
+}
+
 class FVelocityVS : public FMeshMaterialShader
 {
 public:
@@ -302,9 +309,9 @@ void FDeferredShadingSceneRenderer::RenderVelocities(
 				DepthTexture,
 				ERenderTargetLoadAction::ELoad,
 				ERenderTargetLoadAction::ELoad,
-				VelocityPass == EVelocityPass::Opaque
-				? FExclusiveDepthStencil::DepthRead_StencilWrite
-				: FExclusiveDepthStencil::DepthWrite_StencilWrite);
+				(VelocityPass == EVelocityPass::Opaque && !(Scene->EarlyZPassMode == DDM_AllOpaqueNoVelocity))
+					? FExclusiveDepthStencil::DepthRead_StencilWrite
+					: FExclusiveDepthStencil::DepthWrite_StencilWrite);
 
 			if (IsParallelVelocity())
 			{
@@ -741,7 +748,9 @@ FMeshPassProcessor* CreateVelocityPassProcessor(const FScene* Scene, const FScen
 {
 	FMeshPassProcessorRenderState VelocityPassState;
 	VelocityPassState.SetBlendState(TStaticBlendState<CW_RGBA>::GetRHI());
-	VelocityPassState.SetDepthStencilState(TStaticDepthStencilState<false, CF_DepthNearOrEqual>::GetRHI());
+	VelocityPassState.SetDepthStencilState((Scene->EarlyZPassMode == DDM_AllOpaqueNoVelocity) // if the depth mode is all opaque except velocity, it relies on velocity to write the depth of the remaining meshes
+										    ? TStaticDepthStencilState<true, CF_DepthNearOrEqual>::GetRHI()
+											: TStaticDepthStencilState<false, CF_DepthNearOrEqual>::GetRHI());
 
 	return new(FMemStack::Get()) FOpaqueVelocityMeshProcessor(Scene, InViewIfDynamicMeshCommand, VelocityPassState, InDrawListContext);
 }
