Fixed LumenIndirectComposite pass which was not writing alpha: that was preventing the write of scatterable light luminance in alpha required for subsequent sub surface passes. Then luminance to scatter during the subsurface setup pass was always 0.

diff --git a/Engine/Source/Runtime/Renderer/Private/IndirectLightRendering.cpp b/Engine/Source/Runtime/Renderer/Private/IndirectLightRendering.cpp
--- a/Engine/Source/Runtime/Renderer/Private/IndirectLightRendering.cpp
+++ b/Engine/Source/Runtime/Renderer/Private/IndirectLightRendering.cpp
@@ -988,7 +988,7 @@ void FDeferredShadingSceneRenderer::RenderDiffuseIndirectAndAmbientOcclusion(
 			TShaderMapRef<FDiffuseIndirectCompositePS> PixelShader(View.ShaderMap, PermutationVector);
 			ClearUnusedGraphResources(PixelShader, PassParameters);
 
-			FRHIBlendState* BlendState = TStaticBlendState<CW_RGB, BO_Add, BF_One, BF_Source1Color, BO_Add, BF_One, BF_Source1Alpha>::GetRHI();
+			FRHIBlendState* BlendState = TStaticBlendState<CW_RGBA, BO_Add, BF_One, BF_Source1Color, BO_Add, BF_One, BF_Source1Alpha>::GetRHI();
 
 			if (bIsVisualizePass)
 			{
