Make sure volelization material has correct domain, it's possible the default MD_Volume material may not have volume domain

diff --git a/Engine/Source/Runtime/Renderer/Private/VolumetricFogVoxelization.cpp b/Engine/Source/Runtime/Renderer/Private/VolumetricFogVoxelization.cpp
--- a/Engine/Source/Runtime/Renderer/Private/VolumetricFogVoxelization.cpp
+++ b/Engine/Source/Runtime/Renderer/Private/VolumetricFogVoxelization.cpp
@@ -486,7 +486,8 @@ void VoxelizeVolumePrimitive(FVoxelizeVolumeMeshProcessor& PassMeshProcessor,
 	const FPrimitiveSceneProxy* PrimitiveSceneProxy,
 	const FMeshBatch& OriginalMesh)
 {
-	const FMaterial& Material = OriginalMesh.MaterialRenderProxy->GetIncompleteMaterialWithFallback(View.GetFeatureLevel());
+	const FMaterialRenderProxy* MaterialProxy = OriginalMesh.MaterialRenderProxy;
+	const FMaterial& Material = OriginalMesh.MaterialRenderProxy->GetMaterialWithFallback(View.GetFeatureLevel(), MaterialProxy);
 
 	if (Material.GetMaterialDomain() == MD_Volume)
 	{
@@ -510,7 +511,7 @@ void VoxelizeVolumePrimitive(FVoxelizeVolumeMeshProcessor& PassMeshProcessor,
 				GQuadMeshVertexFactory->InitResource();
 			}
 			LocalQuadMesh.VertexFactory = GQuadMeshVertexFactory;
-			LocalQuadMesh.MaterialRenderProxy = OriginalMesh.MaterialRenderProxy;
+			LocalQuadMesh.MaterialRenderProxy = MaterialProxy;
 			LocalQuadMesh.Elements[0].IndexBuffer = &GQuadMeshIndexBuffer;
 			LocalQuadMesh.Elements[0].PrimitiveUniformBuffer = OriginalMesh.Elements[0].PrimitiveUniformBuffer;
 			LocalQuadMesh.Elements[0].FirstIndex = 0;
