Adjust some VSM clipmap settings to avoid running out of pages
- Don't bias the "first" clipmap level based on the resolution bias so that the effective "max" resolution at the near clipmap remains the same. There's no right/wrong on this one but this way is likely more intuitive.

diff --git a/Engine/Source/Runtime/Renderer/Private/VirtualShadowMaps/VirtualShadowMapArray.h b/Engine/Source/Runtime/Renderer/Private/VirtualShadowMaps/VirtualShadowMapArray.h
--- a/Engine/Source/Runtime/Renderer/Private/VirtualShadowMaps/VirtualShadowMapArray.h
+++ b/Engine/Source/Runtime/Renderer/Private/VirtualShadowMaps/VirtualShadowMapArray.h
@@ -43,8 +43,8 @@ public:
 	static constexpr uint32 PageSize = 128U;
 	static constexpr uint32 Level0DimPagesXY = 128U;
 
-	// With 128x128 pages, a 4k texture holds 1024 physical pages
-	static constexpr uint32 PhysicalPagePoolTexureSizeX = 4096U;
+	// With 128x128 pages, a 8k x 4k texture holds 2048 physical pages
+	static constexpr uint32 PhysicalPagePoolTexureSizeX = 8192U;
 	static constexpr uint32 PhysicalPagePoolTexureSizeY = 4096U;
 
 	static constexpr uint32 PageSizeMask = PageSize - 1U;
diff --git a/Engine/Source/Runtime/Renderer/Private/VirtualShadowMaps/VirtualShadowMapClipmap.cpp b/Engine/Source/Runtime/Renderer/Private/VirtualShadowMaps/VirtualShadowMapClipmap.cpp
--- a/Engine/Source/Runtime/Renderer/Private/VirtualShadowMaps/VirtualShadowMapClipmap.cpp
+++ b/Engine/Source/Runtime/Renderer/Private/VirtualShadowMaps/VirtualShadowMapClipmap.cpp
@@ -13,13 +13,13 @@ VirtualShadowMapClipmap.cpp
 
 static TAutoConsoleVariable<float> CVarVirtualShadowMapClipmapResolutionLodBias(
 	TEXT( "r.Shadow.v.Clipmap.ResolutionLodBias" ),
-	3.0f,
+	2.0f,
 	TEXT( "" ),
 	ECVF_RenderThreadSafe
 );
 static TAutoConsoleVariable<float> CVarVirtualShadowMapClipmapFirstLevel(
 	TEXT( "r.Shadow.v.Clipmap.FirstLevel" ),
-	5.0f,
+	8.0f,
 	TEXT( "First level of the virtual clipmap. Lower values allow higher resolution shadows closer to the camera." ),
 	ECVF_RenderThreadSafe
 );
@@ -64,7 +64,7 @@ FVirtualShadowMapClipmap::FVirtualShadowMapClipmap(
 	// For now we adjust resolution by just biasing the page we look up in. This is wasteful in terms of page table vs.
 	// just resizing the virtual shadow maps for each clipmap, but convenient for now. This means we need to additionally bias
 	// which levels are present.
-	FirstLevel = FMath::FloorToInt(CVarVirtualShadowMapClipmapFirstLevel.GetValueOnRenderThread() + ResolutionLodBias);
+	FirstLevel = FMath::FloorToInt(CVarVirtualShadowMapClipmapFirstLevel.GetValueOnRenderThread());
 	int32 LastLevel = FMath::FloorToInt(FMath::Log2(MaxRadius) + ResolutionLodBias);
 	LastLevel = FMath::Max(FirstLevel, LastLevel);
 	int32 LevelCount = LastLevel - FirstLevel + 1;
