Fixing local/global variable name clash.

diff --git a/Engine/Source/Runtime/Renderer/Private/VelocityRendering.cpp b/Engine/Source/Runtime/Renderer/Private/VelocityRendering.cpp
--- a/Engine/Source/Runtime/Renderer/Private/VelocityRendering.cpp
+++ b/Engine/Source/Runtime/Renderer/Private/VelocityRendering.cpp
@@ -63,8 +63,8 @@ bool IsVelocityWaitForTasksEnabled()
 
 bool IsVelocityMergedWithDepthPass()
 {
-	static const auto CVarDepthPassMergedWithVelocity = IConsoleManager::Get().FindTConsoleVariableDataInt(TEXT("r.DepthPassMergedWithVelocity"));
-	bool bIsMerged = (CVarDepthPassMergedWithVelocity && CVarDepthPassMergedWithVelocity->GetValueOnAnyThread() != 0);
+	static const auto CVarMergeDepth = IConsoleManager::Get().FindTConsoleVariableDataInt(TEXT("r.DepthPassMergedWithVelocity"));
+	bool bIsMerged = (CVarMergeDepth && CVarMergeDepth->GetValueOnAnyThread() != 0);
 	return bIsMerged;
 }
 
