Porting shadow depth and projections to RDG.

diff --git a/Engine/Shaders/Private/TranslucentLightingShaders.usf b/Engine/Shaders/Private/TranslucentLightingShaders.usf
--- a/Engine/Shaders/Private/TranslucentLightingShaders.usf
+++ b/Engine/Shaders/Private/TranslucentLightingShaders.usf
@@ -139,40 +139,6 @@ void FilterMainPS(
 	OutColor1 = TextureValue1 * InvWeight;
 }
 
-#if INJECTION_PIXEL_SHADER
-
-/** Parameters needed to access the shadow map of the light. */
-float4x4 WorldToShadowMatrix;
-float4 ShadowmapMinMax;
-
-/** Pixel shader that calculates per object shadowing from translucency for a volume texture and outputs the shadow factor. */
-void PerObjectShadowingMainPS(
-	FWriteToSliceGeometryOutput Input,
-	out float4 OutColor : SV_Target0
-	)
-{
-	float ZPosition = View.TranslucencyLightingVolumeMin[VolumeCascadeIndex].z + (Input.LayerIndex + .5f) * View.TranslucencyLightingVolumeInvSize[VolumeCascadeIndex].w;
-	float3 WorldPosition = float3(View.TranslucencyLightingVolumeMin[VolumeCascadeIndex].xy + Input.Vertex.UV / View.TranslucencyLightingVolumeInvSize[VolumeCascadeIndex].xy - .5f * View.TranslucencyLightingVolumeInvSize[VolumeCascadeIndex].w, ZPosition);
-	float VoxelSize = View.TranslucencyLightingVolumeInvSize[VolumeCascadeIndex].w;
-
-	float ShadowFactor = 1;
-
-	float3 WorldPositionForShadowing = WorldPosition;
-	// Transform the world position into shadowmap space
-	float4 HomogeneousShadowPosition = mul(float4(WorldPositionForShadowing, 1), WorldToShadowMatrix);
-	float2 ShadowUVs = HomogeneousShadowPosition.xy / HomogeneousShadowPosition.w;
-
-	// Treat as unshadowed if the voxel is outside of the shadow map
-	if (all(ShadowUVs >= ShadowmapMinMax.xy && ShadowUVs <= ShadowmapMinMax.zw))
-	{
-		ShadowFactor = CalculateTranslucencyShadowing(ShadowUVs, HomogeneousShadowPosition.z);
-	}
-
-	OutColor = float4(ShadowFactor.xxx, 0);
-}
-
-#endif // #if INJECTION_PIXEL_SHADER
-
 #include "CubemapCommon.ush"
 
 /** Add AmbientCubemap color to the volume. */
diff --git a/Engine/Source/Runtime/RenderCore/Public/RenderGraphBuilder.h b/Engine/Source/Runtime/RenderCore/Public/RenderGraphBuilder.h
--- a/Engine/Source/Runtime/RenderCore/Public/RenderGraphBuilder.h
+++ b/Engine/Source/Runtime/RenderCore/Public/RenderGraphBuilder.h
@@ -194,6 +194,12 @@ public:
 	/** Flag a buffer that is produced by a pass but never used or extracted to not emit an 'unused' warning. */
 	void RemoveUnusedBufferWarning(FRDGBufferRef Buffer);
 
+	/** Manually begins a new GPU event scope. */
+	void BeginEventScope(FRDGEventName&& Name);
+
+	/** Manually ends the current GPU event scope. */
+	void EndEventScope();
+
 	/** Executes the queued passes, managing setting of render targets (RHI RenderPasses), resource transitions and queued texture extraction. */
 	void Execute();
 
diff --git a/Engine/Source/Runtime/RenderCore/Public/RenderGraphBuilder.inl b/Engine/Source/Runtime/RenderCore/Public/RenderGraphBuilder.inl
--- a/Engine/Source/Runtime/RenderCore/Public/RenderGraphBuilder.inl
+++ b/Engine/Source/Runtime/RenderCore/Public/RenderGraphBuilder.inl
@@ -393,4 +393,18 @@ inline void FRDGBuilder::RemoveUnusedBufferWarning(FRDGBufferRef Buffer)
 	UserValidation.ExecuteGuard(TEXT("RemoveUnusedBufferWarning"), Buffer->Name);
 	UserValidation.RemoveUnusedWarning(Buffer);
 #endif
+}
+
+inline void FRDGBuilder::BeginEventScope(FRDGEventName&& ScopeName)
+{
+#if RDG_GPU_SCOPES
+	GPUScopeStacks.BeginEventScope(MoveTemp(ScopeName));
+#endif
+}
+
+inline void FRDGBuilder::EndEventScope()
+{
+#if RDG_GPU_SCOPES
+	GPUScopeStacks.EndEventScope();
+#endif
 }
\ No newline at end of file
diff --git a/Engine/Source/Runtime/Renderer/Private/DeferredShadingRenderer.h b/Engine/Source/Runtime/Renderer/Private/DeferredShadingRenderer.h
--- a/Engine/Source/Runtime/Renderer/Private/DeferredShadingRenderer.h
+++ b/Engine/Source/Runtime/Renderer/Private/DeferredShadingRenderer.h
@@ -813,9 +813,6 @@ private:
 	/** Clears the volume texture used to accumulate per object shadows for translucency. */
 	void ClearTranslucentVolumePerObjectShadowing(FRHICommandList& RHICmdList, const int32 ViewIndex);
 
-	/** Accumulates the per object shadow's contribution for translucency. */
-	void AccumulateTranslucentVolumeObjectShadowing(FRHICommandList& RHICmdList, const FProjectedShadowInfo* InProjectedShadowInfo, bool bClearVolume, const FViewInfo& View, const int32 ViewIndex);
-
 	/** Accumulates direct lighting for the given light.  InProjectedShadowInfo can be NULL in which case the light will be unshadowed. */
 	void InjectTranslucentVolumeLighting(FRDGBuilder& GraphBuilder, const FLightSceneInfo& LightSceneInfo, const FProjectedShadowInfo* InProjectedShadowInfo, const FViewInfo& View, int32 ViewIndex);
 
diff --git a/Engine/Source/Runtime/Renderer/Private/RendererScene.cpp b/Engine/Source/Runtime/Renderer/Private/RendererScene.cpp
--- a/Engine/Source/Runtime/Renderer/Private/RendererScene.cpp
+++ b/Engine/Source/Runtime/Renderer/Private/RendererScene.cpp
@@ -897,14 +897,12 @@ void FPersistentUniformBuffers::Clear()
 	InstancedViewUniformBuffer.SafeRelease();
 	ReflectionCaptureUniformBuffer.SafeRelease();
 	CSMShadowDepthViewUniformBuffer.SafeRelease();
-	CSMShadowDepthPassUniformBuffer.SafeRelease();
 	VoxelizeVolumeViewUniformBuffer.SafeRelease();
 	CustomDepthViewUniformBuffer.SafeRelease();
 	InstancedCustomDepthViewUniformBuffer.SafeRelease();
 	VirtualTextureViewUniformBuffer.SafeRelease();
 	MobileOpaqueBasePassUniformBuffer.SafeRelease();
 	MobileTranslucentBasePassUniformBuffer.SafeRelease();
-	MobileCSMShadowDepthPassUniformBuffer.SafeRelease();
 
 
 	for (auto& UniformBuffer : MobileDirectionalLightUniformBuffers)
@@ -931,9 +929,6 @@ void FPersistentUniformBuffers::Initialize()
 
 	CSMShadowDepthViewUniformBuffer = TUniformBufferRef<FViewUniformShaderParameters>::CreateUniformBufferImmediate(ViewUniformBufferParameters, UniformBuffer_MultiFrame, EUniformBufferValidation::None);
 
-	FShadowDepthPassUniformParameters CSMShadowDepthPassParameters;
-	CSMShadowDepthPassUniformBuffer = TUniformBufferRef<FShadowDepthPassUniformParameters>::CreateUniformBufferImmediate(CSMShadowDepthPassParameters, UniformBuffer_MultiFrame, EUniformBufferValidation::None);
-
 	VoxelizeVolumeViewUniformBuffer = TUniformBufferRef<FViewUniformShaderParameters>::CreateUniformBufferImmediate(ViewUniformBufferParameters, UniformBuffer_MultiFrame, EUniformBufferValidation::None);
 
 	CustomDepthViewUniformBuffer = TUniformBufferRef<FViewUniformShaderParameters>::CreateUniformBufferImmediate(ViewUniformBufferParameters, UniformBuffer_MultiFrame, EUniformBufferValidation::None);
@@ -943,9 +938,6 @@ void FPersistentUniformBuffers::Initialize()
 
 	LumenCardCaptureViewUniformBuffer = TUniformBufferRef<FViewUniformShaderParameters>::CreateUniformBufferImmediate(ViewUniformBufferParameters, UniformBuffer_MultiFrame, EUniformBufferValidation::None);
 
-	FMobileShadowDepthPassUniformParameters MobileCSMShadowDepthPassParameters;
-	MobileCSMShadowDepthPassUniformBuffer = TUniformBufferRef<FMobileShadowDepthPassUniformParameters>::CreateUniformBufferImmediate(MobileCSMShadowDepthPassParameters, UniformBuffer_MultiFrame, EUniformBufferValidation::None);
-
 	FMobileBasePassUniformParameters MobileBasePassUniformParameters;
 	MobileOpaqueBasePassUniformBuffer = TUniformBufferRef<FMobileBasePassUniformParameters>::CreateUniformBufferImmediate(MobileBasePassUniformParameters, UniformBuffer_MultiFrame, EUniformBufferValidation::None);
 	MobileCSMOpaqueBasePassUniformBuffer = TUniformBufferRef<FMobileBasePassUniformParameters>::CreateUniformBufferImmediate(MobileBasePassUniformParameters, UniformBuffer_MultiFrame, EUniformBufferValidation::None);
diff --git a/Engine/Source/Runtime/Renderer/Private/ScenePrivate.h b/Engine/Source/Runtime/Renderer/Private/ScenePrivate.h
--- a/Engine/Source/Runtime/Renderer/Private/ScenePrivate.h
+++ b/Engine/Source/Runtime/Renderer/Private/ScenePrivate.h
@@ -2407,7 +2407,6 @@ public:
 	TUniformBufferRef<FNaniteUniformParameters> NaniteUniformBuffer;
 	TUniformBufferRef<FReflectionCaptureShaderData> ReflectionCaptureUniformBuffer;
 	TUniformBufferRef<FViewUniformShaderParameters> CSMShadowDepthViewUniformBuffer;
-	TUniformBufferRef<FShadowDepthPassUniformParameters> CSMShadowDepthPassUniformBuffer;
 	TUniformBufferRef<FViewUniformShaderParameters> VoxelizeVolumeViewUniformBuffer;
 	TUniformBufferRef<FViewUniformShaderParameters> CustomDepthViewUniformBuffer;
 	TUniformBufferRef<FInstancedViewUniformShaderParameters> InstancedCustomDepthViewUniformBuffer;
@@ -2417,7 +2416,6 @@ public:
 	TUniformBufferRef<FMobileBasePassUniformParameters> MobileOpaqueBasePassUniformBuffer;
 	TUniformBufferRef<FMobileBasePassUniformParameters> MobileCSMOpaqueBasePassUniformBuffer;
 	TUniformBufferRef<FMobileBasePassUniformParameters> MobileTranslucentBasePassUniformBuffer;
-	TUniformBufferRef<FMobileShadowDepthPassUniformParameters> MobileCSMShadowDepthPassUniformBuffer;
 	/** Mobile Directional Lighting uniform buffers, one for each lighting channel 
 	  * The first is used for primitives with no lighting channels set.
 	  */
diff --git a/Engine/Source/Runtime/Renderer/Private/SceneRendering.h b/Engine/Source/Runtime/Renderer/Private/SceneRendering.h
--- a/Engine/Source/Runtime/Renderer/Private/SceneRendering.h
+++ b/Engine/Source/Runtime/Renderer/Private/SceneRendering.h
@@ -46,6 +46,7 @@ struct FILCUpdatePrimTaskData;
 class FRaytracingLightDataPacked;
 class FRayTracingLocalShaderBindingWriter;
 class FVirtualShadowMapClipmap;
+class FShadowProjectionPassParameters;
 
 struct FCloudRenderContext;
 struct FSceneWithoutWaterTextures;
@@ -1845,7 +1846,8 @@ protected:
 	void SetupMeshPass(FViewInfo& View, FExclusiveDepthStencil::Type BasePassDepthStencilAccess, FViewCommands& ViewCommands);
 
 	void RenderShadowProjections(
-		FRHICommandListImmediate& RHICmdList,
+		FRDGBuilder& GraphBuilder,
+		const FShadowProjectionPassParameters& CommonPassParameters,
 		const FLightSceneProxy* LightSceneProxy,
 		const FHairStrandsVisibilityViews* HairVisibilityViews,
 		TArrayView<const FProjectedShadowInfo* const> Shadows,
diff --git a/Engine/Source/Runtime/Renderer/Private/ShadowDepthRendering.cpp b/Engine/Source/Runtime/Renderer/Private/ShadowDepthRendering.cpp
--- a/Engine/Source/Runtime/Renderer/Private/ShadowDepthRendering.cpp
+++ b/Engine/Source/Runtime/Renderer/Private/ShadowDepthRendering.cpp
@@ -46,26 +46,20 @@
 
 DECLARE_GPU_DRAWCALL_STAT_NAMED(ShadowDepths, TEXT("Shadow Depths"));
 
-IMPLEMENT_GLOBAL_SHADER_PARAMETER_STRUCT(FShadowDepthPassUniformParameters, "ShadowDepthPass");
-IMPLEMENT_GLOBAL_SHADER_PARAMETER_STRUCT(FMobileShadowDepthPassUniformParameters, "MobileShadowDepthPass");
+IMPLEMENT_STATIC_UNIFORM_BUFFER_STRUCT(FShadowDepthPassUniformParameters, "ShadowDepthPass", SceneTextures);
+IMPLEMENT_STATIC_UNIFORM_BUFFER_STRUCT(FMobileShadowDepthPassUniformParameters, "MobileShadowDepthPass", SceneTextures);
 
 template<bool bUsingVertexLayers = false>
 class TScreenVSForGS : public FScreenVS
 {
-	DECLARE_SHADER_TYPE(TScreenVSForGS, Global);
 public:
+	DECLARE_SHADER_TYPE(TScreenVSForGS, Global);
 
 	static bool ShouldCompilePermutation(const FGlobalShaderPermutationParameters& Parameters)
 	{
 		return IsFeatureLevelSupported(Parameters.Platform, ERHIFeatureLevel::SM5) && (!bUsingVertexLayers || RHISupportsVertexShaderLayer(Parameters.Platform));
 	}
 
-	TScreenVSForGS(const ShaderMetaType::CompiledShaderInitializerType& Initializer) :
-		FScreenVS(Initializer)
-	{
-	}
-	TScreenVSForGS() {}
-
 	static void ModifyCompilationEnvironment(const FGlobalShaderPermutationParameters& Parameters, FShaderCompilerEnvironment& OutEnvironment)
 	{
 		FScreenVS::ModifyCompilationEnvironment(Parameters, OutEnvironment);
@@ -75,6 +69,11 @@ public:
 			OutEnvironment.CompilerFlags.Add(CFLAG_VertexToGeometryShader);
 		}
 	}
+
+	TScreenVSForGS() = default;
+	TScreenVSForGS(const ShaderMetaType::CompiledShaderInitializerType& Initializer) :
+		FScreenVS(Initializer)
+	{}
 };
 
 IMPLEMENT_SHADER_TYPE(template<>, TScreenVSForGS<false>, TEXT("/Engine/Private/ScreenVertexShader.usf"), TEXT("MainForGS"), SF_Vertex);
@@ -137,12 +136,11 @@ static float ComputeNaniteShadowsLODScaleFactor()
 
 void SetupShadowDepthPassUniformBuffer(
 	const FProjectedShadowInfo* ShadowInfo,
-	FRHICommandListImmediate& RHICmdList,
+	FRDGBuilder& GraphBuilder,
 	const FViewInfo& View,
 	FShadowDepthPassUniformParameters& ShadowDepthPassParameters)
 {
-	FSceneRenderTargets& SceneRenderTargets = FSceneRenderTargets::Get(RHICmdList);
-	SetupSceneTextureUniformParameters(SceneRenderTargets, View.FeatureLevel, ESceneTextureSetupMode::None, ShadowDepthPassParameters.SceneTextures);
+	SetupSceneTextureUniformParameters(GraphBuilder, View.FeatureLevel, ESceneTextureSetupMode::None, ShadowDepthPassParameters.SceneTextures);
 
 	ShadowDepthPassParameters.ProjectionMatrix = FTranslationMatrix(ShadowInfo->PreShadowTranslation - View.ViewMatrices.GetPreViewTranslation()) * ShadowInfo->SubjectAndReceiverMatrix;	
 	ShadowDepthPassParameters.ViewMatrix = ShadowInfo->ShadowViewMatrix;
@@ -165,12 +163,11 @@ void SetupShadowDepthPassUniformBuffer(
 
 void SetupShadowDepthPassUniformBuffer(
 	const FProjectedShadowInfo* ShadowInfo,
-	FRHICommandListImmediate& RHICmdList,
+	FRDGBuilder& GraphBuilder,
 	const FViewInfo& View,
 	FMobileShadowDepthPassUniformParameters& ShadowDepthPassParameters)
 {
-	FSceneRenderTargets& SceneRenderTargets = FSceneRenderTargets::Get(RHICmdList);
-	SetupMobileSceneTextureUniformParameters(SceneRenderTargets, EMobileSceneTextureSetupMode::None, ShadowDepthPassParameters.SceneTextures);
+	SetupMobileSceneTextureUniformParameters(GraphBuilder, EMobileSceneTextureSetupMode::None, ShadowDepthPassParameters.SceneTextures);
 
 	ShadowDepthPassParameters.ProjectionMatrix = FTranslationMatrix(ShadowInfo->PreShadowTranslation - View.ViewMatrices.GetPreViewTranslation()) * ShadowInfo->SubjectAndReceiverMatrix;
 	ShadowDepthPassParameters.ViewMatrix = ShadowInfo->ShadowViewMatrix;
@@ -179,10 +176,28 @@ void SetupShadowDepthPassUniformBuffer(
 	ShadowDepthPassParameters.bClampToNearPlane = ShadowInfo->ShouldClampToNearPlane() ? 1.0f : 0.0f;
 }
 
+void AddClearShadowDepthPass(FRDGBuilder& GraphBuilder, FRDGTextureRef Texture)
+{
+	// Clear atlas depth, but ignore stencil.
+	auto* PassParameters = GraphBuilder.AllocParameters<FRenderTargetParameters>();
+	PassParameters->RenderTargets.DepthStencil = FDepthStencilBinding(Texture, ERenderTargetLoadAction::EClear, ERenderTargetLoadAction::ENoAction, FExclusiveDepthStencil::DepthWrite_StencilNop);
+	GraphBuilder.AddPass(RDG_EVENT_NAME("ClearShadowDepth"), PassParameters, ERDGPassFlags::Raster, [](FRHICommandList&) {});
+}
+
+void AddClearShadowDepthPass(FRDGBuilder& GraphBuilder, FRDGTextureRef Texture, const FProjectedShadowInfo* ProjectedShadowInfo)
+{
+	// Clear atlas depth, but ignore stencil.
+	auto* PassParameters = GraphBuilder.AllocParameters<FRenderTargetParameters>();
+	PassParameters->RenderTargets.DepthStencil = FDepthStencilBinding(Texture, ERenderTargetLoadAction::ELoad, ERenderTargetLoadAction::ENoAction, FExclusiveDepthStencil::DepthWrite_StencilNop);
+	GraphBuilder.AddPass(RDG_EVENT_NAME("ClearShadowDepthTile"), PassParameters, ERDGPassFlags::Raster, [ProjectedShadowInfo](FRHICommandList& RHICmdList)
+	{
+		ProjectedShadowInfo->ClearDepth(RHICmdList);
+	});
+}
+
 class FShadowDepthShaderElementData : public FMeshMaterialShaderElementData
 {
 public:
-
 	int32 LayerId;
 };
 
@@ -191,34 +206,14 @@ public:
 */
 class FShadowDepthVS : public FMeshMaterialShader
 {
-	DECLARE_INLINE_TYPE_LAYOUT(FShadowDepthVS, NonVirtual);
 public:
+	DECLARE_INLINE_TYPE_LAYOUT(FShadowDepthVS, NonVirtual);
 
 	static bool ShouldCompilePermutation(const FMeshMaterialShaderPermutationParameters& Parameters)
 	{
 		return false;
 	}
 
-	FShadowDepthVS(const ShaderMetaType::CompiledShaderInitializerType& Initializer) :
-		FMeshMaterialShader(Initializer)
-	{
-		const ERHIFeatureLevel::Type FeatureLevel = GetMaxSupportedFeatureLevel((EShaderPlatform)Initializer.Target.Platform);
-
-		if (FSceneInterface::GetShadingPath(FeatureLevel) == EShadingPath::Deferred)
-		{
-			PassUniformBuffer.Bind(Initializer.ParameterMap, FShadowDepthPassUniformParameters::StaticStructMetadata.GetShaderVariableName());
-		}
-
-		if (FSceneInterface::GetShadingPath(FeatureLevel) == EShadingPath::Mobile)
-		{
-			PassUniformBuffer.Bind(Initializer.ParameterMap, FMobileShadowDepthPassUniformParameters::StaticStructMetadata.GetShaderVariableName());
-		}
-
-		LayerId.Bind(Initializer.ParameterMap, TEXT("LayerId"));
-	}
-
-	FShadowDepthVS() {}
-
 	void GetShaderBindings(
 		const FScene* Scene,
 		ERHIFeatureLevel::Type FeatureLevel,
@@ -234,8 +229,14 @@ public:
 		ShaderBindings.Add(LayerId, ShaderElementData.LayerId);
 	}
 
+	FShadowDepthVS() = default;
+	FShadowDepthVS(const ShaderMetaType::CompiledShaderInitializerType & Initializer) :
+		FMeshMaterialShader(Initializer)
+	{
+		LayerId.Bind(Initializer.ParameterMap, TEXT("LayerId"));
+	}
+
 private:
-	
 	LAYOUT_FIELD(FShaderParameter, LayerId);
 };
 
@@ -259,15 +260,8 @@ static TAutoConsoleVariable<int32> CVarSupportPointLightWholeSceneShadows(
 template <EShadowDepthVertexShaderMode ShaderMode, bool bUsePositionOnlyStream, bool bIsForGeometryShader = false>
 class TShadowDepthVS : public FShadowDepthVS
 {
-	DECLARE_SHADER_TYPE(TShadowDepthVS, MeshMaterial);
 public:
-
-	TShadowDepthVS(const ShaderMetaType::CompiledShaderInitializerType& Initializer) :
-		FShadowDepthVS(Initializer)
-	{
-	}
-
-	TShadowDepthVS() {}
+	DECLARE_SHADER_TYPE(TShadowDepthVS, MeshMaterial);
 
 	static bool ShouldCompilePermutation(const FMeshMaterialShaderPermutationParameters& Parameters)
 	{
@@ -318,8 +312,12 @@ public:
 			OutEnvironment.CompilerFlags.Add(CFLAG_VertexUseAutoCulling);
 		}
 	}
-};
 
+	TShadowDepthVS() = default;
+	TShadowDepthVS(const ShaderMetaType::CompiledShaderInitializerType & Initializer)
+		: FShadowDepthVS(Initializer)
+	{}
+};
 
 /**
 * A Hull shader for rendering the depth of a mesh.
@@ -327,15 +325,8 @@ public:
 template <EShadowDepthVertexShaderMode ShaderMode>
 class TShadowDepthHS : public FBaseHS
 {
-	DECLARE_SHADER_TYPE(TShadowDepthHS, MeshMaterial);
 public:
-
-
-	TShadowDepthHS(const ShaderMetaType::CompiledShaderInitializerType& Initializer) :
-		FBaseHS(Initializer)
-	{}
-
-	TShadowDepthHS() {}
+	DECLARE_SHADER_TYPE(TShadowDepthHS, MeshMaterial);
 
 	static bool ShouldCompilePermutation(const FMeshMaterialShaderPermutationParameters& Parameters)
 	{
@@ -347,9 +338,13 @@ public:
 	static void ModifyCompilationEnvironment(const FMaterialShaderPermutationParameters& Parameters, FShaderCompilerEnvironment& OutEnvironment)
 	{
 		// Re-use compilation env from vertex shader
-
 		TShadowDepthVS<ShaderMode, false>::ModifyCompilationEnvironment(Parameters, OutEnvironment);
 	}
+
+	TShadowDepthHS() = default;
+	TShadowDepthHS(const ShaderMetaType::CompiledShaderInitializerType& Initializer)
+		: FBaseHS(Initializer)
+	{}
 };
 
 /**
@@ -739,42 +734,35 @@ static void CheckShadowDepthMaterials(const FMaterialRenderProxy* InRenderProxy,
 	check(Material == InMaterial);
 }
 
-void FProjectedShadowInfo::ClearDepth(FRHICommandList& RHICmdList, class FSceneRenderer* SceneRenderer, int32 NumColorTextures, FRHITexture** ColorTextures, FRHITexture* DepthTexture, bool bPerformClear)
+void FProjectedShadowInfo::ClearDepth(FRHICommandList& RHICmdList) const
 {
 	check(RHICmdList.IsInsideRenderPass());
 
-	uint32 ViewportMinX = X;
-	uint32 ViewportMinY = Y;
-	float ViewportMinZ = 0.0f;
-	uint32 ViewportMaxX = X + BorderSize * 2 + ResolutionX;
-	uint32 ViewportMaxY = Y + BorderSize * 2 + ResolutionY;
-	float ViewportMaxZ = 1.0f;
+	const uint32 ViewportMinX = X;
+	const uint32 ViewportMinY = Y;
+	const float ViewportMinZ = 0.0f;
+	const uint32 ViewportMaxX = X + BorderSize * 2 + ResolutionX;
+	const uint32 ViewportMaxY = Y + BorderSize * 2 + ResolutionY;
+	const float ViewportMaxZ = 1.0f;
 
-	int32 NumClearColors;
-	bool bClearColor;
-	FLinearColor Colors[2];
+	// Clear depth only.
+	const int32 NumClearColors = 1;
+	const bool bClearColor = false;
+	const FLinearColor Colors[] = { FLinearColor::White };
 
 	// Translucent shadows use draw call clear
 	check(!bTranslucentShadow);
 
-	// Clear depth only.
-	bClearColor = false;
-	Colors[0] = FLinearColor::White;
-	NumClearColors = FMath::Min(1, NumColorTextures);
-
-	if (bPerformClear)
-	{
-		RHICmdList.SetViewport(
-			ViewportMinX,
-			ViewportMinY,
-			ViewportMinZ,
-			ViewportMaxX,
-			ViewportMaxY,
-			ViewportMaxZ
-		);
+	RHICmdList.SetViewport(
+		ViewportMinX,
+		ViewportMinY,
+		ViewportMinZ,
+		ViewportMaxX,
+		ViewportMaxY,
+		ViewportMaxZ
+	);
 
-		DrawClearQuadMRT(RHICmdList, bClearColor, NumClearColors, Colors, true, 1.0f, false, 0);
-	}
+	DrawClearQuadMRT(RHICmdList, bClearColor, NumClearColors, Colors, true, 1.0f, false, 0);
 }
 
 void FProjectedShadowInfo::SetStateForView(FRHICommandList& RHICmdList) const
@@ -827,165 +815,157 @@ static TAutoConsoleVariable<int32> CVarRHICmdFlushRenderThreadTasksShadowPass(
 
 DECLARE_CYCLE_STAT(TEXT("Shadow"), STAT_CLP_Shadow, STATGROUP_ParallelCommandListMarkers);
 
-class FShadowParallelCommandListSet : public FParallelCommandListSet
+class FShadowParallelCommandListSet final : public FParallelCommandListSet
 {
 public:
 	FShadowParallelCommandListSet(
 		FRHICommandListImmediate& InParentCmdList,
 		const FViewInfo& InView,
-		bool bInCreateSceneContext,
-		FProjectedShadowInfo& InProjectedShadowInfo,
-		FBeginShadowRenderPassFunction InBeginShadowRenderPass)
-		: FParallelCommandListSet(GET_STATID(STAT_CLP_Shadow), InView, InParentCmdList, bInCreateSceneContext)
+		const FProjectedShadowInfo& InProjectedShadowInfo,
+		const FParallelCommandListBindings& InBindings)
+		: FParallelCommandListSet(GET_STATID(STAT_CLP_Shadow), InView, InParentCmdList, false)
 		, ProjectedShadowInfo(InProjectedShadowInfo)
-		, BeginShadowRenderPass(InBeginShadowRenderPass)
-	{
-		bBalanceCommands = false;
-	}
+		, Bindings(InBindings)
+	{}
 
-	virtual ~FShadowParallelCommandListSet()
+	~FShadowParallelCommandListSet() override
 	{
 		Dispatch();
 	}
 
-	virtual void SetStateOnCommandList(FRHICommandList& CmdList) override
+	void SetStateOnCommandList(FRHICommandList& RHICmdList) override
 	{
-		FParallelCommandListSet::SetStateOnCommandList(CmdList);
-		BeginShadowRenderPass(CmdList, false);
-		ProjectedShadowInfo.SetStateForView(CmdList);
+		FParallelCommandListSet::SetStateOnCommandList(RHICmdList);
+		Bindings.SetOnCommandList(RHICmdList);
+		ProjectedShadowInfo.SetStateForView(RHICmdList);
 	}
 
 private:
-	FProjectedShadowInfo& ProjectedShadowInfo;
-	FBeginShadowRenderPassFunction BeginShadowRenderPass;
-	EShadowDepthRenderMode RenderMode;
+	const FProjectedShadowInfo& ProjectedShadowInfo;
+	FParallelCommandListBindings Bindings;
 };
 
 class FCopyShadowMapsCubeGS : public FGlobalShader
 {
-	DECLARE_SHADER_TYPE(FCopyShadowMapsCubeGS, Global);
 public:
+	DECLARE_GLOBAL_SHADER(FCopyShadowMapsCubeGS);
 
 	static bool ShouldCompilePermutation(const FGlobalShaderPermutationParameters& Parameters)
 	{
 		return RHISupportsGeometryShaders(Parameters.Platform) && IsFeatureLevelSupported(Parameters.Platform, ERHIFeatureLevel::SM5);
 	}
 
-	FCopyShadowMapsCubeGS(const ShaderMetaType::CompiledShaderInitializerType& Initializer) :
-		FGlobalShader(Initializer)
-	{
-	}
-	FCopyShadowMapsCubeGS() {}
+	FCopyShadowMapsCubeGS() = default;
+	FCopyShadowMapsCubeGS(const ShaderMetaType::CompiledShaderInitializerType& Initializer)
+		: FGlobalShader(Initializer)
+	{}
 };
 
-IMPLEMENT_SHADER_TYPE(, FCopyShadowMapsCubeGS, TEXT("/Engine/Private/CopyShadowMaps.usf"), TEXT("CopyCubeDepthGS"), SF_Geometry);
+IMPLEMENT_GLOBAL_SHADER(FCopyShadowMapsCubeGS, "/Engine/Private/CopyShadowMaps.usf", "CopyCubeDepthGS", SF_Geometry);
 
 class FCopyShadowMapsCubePS : public FGlobalShader
 {
-	DECLARE_SHADER_TYPE(FCopyShadowMapsCubePS, Global);
 public:
+	DECLARE_GLOBAL_SHADER(FCopyShadowMapsCubePS);
+	SHADER_USE_PARAMETER_STRUCT(FCopyShadowMapsCubePS, FGlobalShader);
+
+	BEGIN_SHADER_PARAMETER_STRUCT(FParameters, )
+		SHADER_PARAMETER_STRUCT_REF(FViewUniformShaderParameters, View)
+		SHADER_PARAMETER_RDG_TEXTURE(TextureCube, ShadowDepthCubeTexture)
+		SHADER_PARAMETER_SAMPLER(SamplerState, ShadowDepthSampler)
+		RENDER_TARGET_BINDING_SLOTS()
+	END_SHADER_PARAMETER_STRUCT()
 
 	static bool ShouldCompilePermutation(const FGlobalShaderPermutationParameters& Parameters)
 	{
 		return IsFeatureLevelSupported(Parameters.Platform, ERHIFeatureLevel::SM5);
 	}
-
-	FCopyShadowMapsCubePS(const ShaderMetaType::CompiledShaderInitializerType& Initializer) :
-		FGlobalShader(Initializer)
-	{
-		ShadowDepthTexture.Bind(Initializer.ParameterMap, TEXT("ShadowDepthCubeTexture"));
-		ShadowDepthSampler.Bind(Initializer.ParameterMap, TEXT("ShadowDepthSampler"));
-	}
-	FCopyShadowMapsCubePS() {}
-
-	void SetParameters(FRHICommandList& RHICmdList, const FSceneView& View, IPooledRenderTarget* SourceShadowMap)
-	{
-		FGlobalShader::SetParameters<FViewUniformShaderParameters>(RHICmdList, RHICmdList.GetBoundPixelShader(), View.ViewUniformBuffer);
-
-		SetTextureParameter(RHICmdList, RHICmdList.GetBoundPixelShader(), ShadowDepthTexture, ShadowDepthSampler, TStaticSamplerState<SF_Point, AM_Clamp, AM_Clamp, AM_Clamp>::GetRHI(), SourceShadowMap->GetRenderTargetItem().ShaderResourceTexture);
-	}
-
-	LAYOUT_FIELD(FShaderResourceParameter, ShadowDepthTexture);
-	LAYOUT_FIELD(FShaderResourceParameter, ShadowDepthSampler);
 };
 
-IMPLEMENT_SHADER_TYPE(, FCopyShadowMapsCubePS, TEXT("/Engine/Private/CopyShadowMaps.usf"), TEXT("CopyCubeDepthPS"), SF_Pixel);
+IMPLEMENT_GLOBAL_SHADER(FCopyShadowMapsCubePS, "/Engine/Private/CopyShadowMaps.usf", "CopyCubeDepthPS", SF_Pixel);
 
-/** */
 class FCopyShadowMaps2DPS : public FGlobalShader
 {
-	DECLARE_SHADER_TYPE(FCopyShadowMaps2DPS, Global);
 public:
+	DECLARE_GLOBAL_SHADER(FCopyShadowMaps2DPS);
+	SHADER_USE_PARAMETER_STRUCT(FCopyShadowMaps2DPS, FGlobalShader);
 
-	static bool ShouldCompilePermutation(const FGlobalShaderPermutationParameters& Parameters)
-	{
-		return true;
-	}
-
-	FCopyShadowMaps2DPS(const ShaderMetaType::CompiledShaderInitializerType& Initializer) :
-		FGlobalShader(Initializer)
-	{
-		ShadowDepthTexture.Bind(Initializer.ParameterMap, TEXT("ShadowDepthTexture"));
-		ShadowDepthSampler.Bind(Initializer.ParameterMap, TEXT("ShadowDepthSampler"));
-	}
-	FCopyShadowMaps2DPS() {}
-
-	void SetParameters(FRHICommandList& RHICmdList, const FSceneView& View, IPooledRenderTarget* SourceShadowMap)
-	{
-		FGlobalShader::SetParameters<FViewUniformShaderParameters>(RHICmdList, RHICmdList.GetBoundPixelShader(), View.ViewUniformBuffer);
-
-		SetTextureParameter(RHICmdList, RHICmdList.GetBoundPixelShader(), ShadowDepthTexture, ShadowDepthSampler, TStaticSamplerState<SF_Point, AM_Clamp, AM_Clamp, AM_Clamp>::GetRHI(), SourceShadowMap->GetRenderTargetItem().ShaderResourceTexture);
-	}
-
-	LAYOUT_FIELD(FShaderResourceParameter, ShadowDepthTexture);
-	LAYOUT_FIELD(FShaderResourceParameter, ShadowDepthSampler);
+	BEGIN_SHADER_PARAMETER_STRUCT(FParameters, )
+		SHADER_PARAMETER_STRUCT_REF(FViewUniformShaderParameters, View)
+		SHADER_PARAMETER_RDG_TEXTURE(Texture2D, ShadowDepthTexture)
+		SHADER_PARAMETER_SAMPLER(SamplerState, ShadowDepthSampler)
+		RENDER_TARGET_BINDING_SLOTS()
+	END_SHADER_PARAMETER_STRUCT()
 };
 
-IMPLEMENT_SHADER_TYPE(, FCopyShadowMaps2DPS, TEXT("/Engine/Private/CopyShadowMaps.usf"), TEXT("Copy2DDepthPS"), SF_Pixel);
+IMPLEMENT_GLOBAL_SHADER(FCopyShadowMaps2DPS, "/Engine/Private/CopyShadowMaps.usf", "Copy2DDepthPS", SF_Pixel);
 
-void FProjectedShadowInfo::CopyCachedShadowMap(FRHICommandList& RHICmdList, const FMeshPassProcessorRenderState& DrawRenderState, FSceneRenderer* SceneRenderer, const FViewInfo& View)
+void FProjectedShadowInfo::CopyCachedShadowMap(
+	FRDGBuilder& GraphBuilder,
+	const FViewInfo& View,
+	const FSceneRenderer* SceneRenderer,
+	const FRenderTargetBindingSlots& RenderTargetBindingSlots,
+	const FMeshPassProcessorRenderState& DrawRenderState)
 {
 	check(CacheMode == SDCM_MovablePrimitivesOnly);
 	const FCachedShadowMapData& CachedShadowMapData = SceneRenderer->Scene->CachedShadowMaps.FindChecked(GetLightSceneInfo().Id);
 
-	FGraphicsPipelineStateInitializer GraphicsPSOInit;
-	DrawRenderState.ApplyToPSO(GraphicsPSOInit);
-	uint32 StencilRef = DrawRenderState.GetStencilRef();
-
 	if (CachedShadowMapData.bCachedShadowMapHasPrimitives && CachedShadowMapData.ShadowMap.IsValid())
 	{
-		SCOPED_DRAW_EVENT(RHICmdList, CopyCachedShadowMap);
+		FRDGTextureRef ShadowDepthTexture = GraphBuilder.RegisterExternalTexture(CachedShadowMapData.ShadowMap.DepthTarget);
+		const FIntPoint ShadowDepthExtent = ShadowDepthTexture->Desc.Extent;
 
-		RHICmdList.ApplyCachedRenderTargets(GraphicsPSOInit);
+		FGraphicsPipelineStateInitializer GraphicsPSOInit;
+		DrawRenderState.ApplyToPSO(GraphicsPSOInit);
+		const uint32 StencilRef = DrawRenderState.GetStencilRef();
 
 		GraphicsPSOInit.RasterizerState = TStaticRasterizerState<FM_Solid, CM_None>::GetRHI();
 		// No depth tests, so we can replace the clear
 		GraphicsPSOInit.DepthStencilState = TStaticDepthStencilState<true, CF_Always>::GetRHI();
+		GraphicsPSOInit.PrimitiveType = PT_TriangleList;
 
 		extern TGlobalResource<FFilterVertexDeclaration> GFilterVertexDeclaration;
+		GraphicsPSOInit.BoundShaderState.VertexDeclarationRHI = GFilterVertexDeclaration.VertexDeclarationRHI;
 
 		if (bOnePassPointLightShadow)
 		{
+			TShaderMapRef<TScreenVSForGS<false>> ScreenVertexShader(View.ShaderMap);
+			TShaderMapRef<FCopyShadowMapsCubePS> PixelShader(View.ShaderMap);
+			GraphicsPSOInit.BoundShaderState.VertexShaderRHI = ScreenVertexShader.GetVertexShader();
+			GraphicsPSOInit.BoundShaderState.PixelShaderRHI = PixelShader.GetPixelShader();
+
+			int32 InstanceCount = 1;
+
+#if PLATFORM_SUPPORTS_GEOMETRY_SHADERS
 			if (RHISupportsGeometryShaders(GShaderPlatformForFeatureLevel[SceneRenderer->FeatureLevel]))
 			{
-				// Set shaders and texture
-				TShaderMapRef<TScreenVSForGS<false>> ScreenVertexShader(View.ShaderMap);
 				TShaderMapRef<FCopyShadowMapsCubeGS> GeometryShader(View.ShaderMap);
-				TShaderMapRef<FCopyShadowMapsCubePS> PixelShader(View.ShaderMap);
-
-				GraphicsPSOInit.BoundShaderState.VertexDeclarationRHI = GFilterVertexDeclaration.VertexDeclarationRHI;
-				GraphicsPSOInit.BoundShaderState.VertexShaderRHI = ScreenVertexShader.GetVertexShader();
-#if PLATFORM_SUPPORTS_GEOMETRY_SHADERS
 				GraphicsPSOInit.BoundShaderState.GeometryShaderRHI = GeometryShader.GetGeometryShader();
+				InstanceCount = 6;
+			}
+			else
 #endif
-				GraphicsPSOInit.BoundShaderState.PixelShaderRHI = PixelShader.GetPixelShader();
-				GraphicsPSOInit.PrimitiveType = PT_TriangleList;
+			{
+				check(RHISupportsVertexShaderLayer(GShaderPlatformForFeatureLevel[SceneRenderer->FeatureLevel]));
+			}
 
+			auto* PassParameters = GraphBuilder.AllocParameters<FCopyShadowMapsCubePS::FParameters>();
+			PassParameters->RenderTargets = RenderTargetBindingSlots;
+			PassParameters->ShadowDepthCubeTexture = ShadowDepthTexture;
+			PassParameters->ShadowDepthSampler = TStaticSamplerState<SF_Point, AM_Clamp, AM_Clamp, AM_Clamp>::GetRHI();
+
+			GraphBuilder.AddPass(
+				RDG_EVENT_NAME("CopyCachedShadowMap"),
+				PassParameters,
+				ERDGPassFlags::Raster,
+				[this, ScreenVertexShader, PixelShader, GraphicsPSOInit, PassParameters, ShadowDepthExtent, InstanceCount, StencilRef](FRHICommandList& RHICmdList) mutable
+			{
+				SetStateForView(RHICmdList);
+				RHICmdList.ApplyCachedRenderTargets(GraphicsPSOInit);
 				SetGraphicsPipelineState(RHICmdList, GraphicsPSOInit);
+				SetShaderParameters(RHICmdList, PixelShader, PixelShader.GetPixelShader(), *PassParameters);
 				RHICmdList.SetStencilRef(StencilRef);
 
-				PixelShader->SetParameters(RHICmdList, View, CachedShadowMapData.ShadowMap.DepthTarget.GetReference());
-
 				DrawRectangle(
 					RHICmdList,
 					0, 0,
@@ -993,28 +973,36 @@ void FProjectedShadowInfo::CopyCachedShadowMap(FRHICommandList& RHICmdList, cons
 					BorderSize, BorderSize,
 					ResolutionX, ResolutionY,
 					FIntPoint(ResolutionX, ResolutionY),
-					CachedShadowMapData.ShadowMap.GetSize(),
+					ShadowDepthExtent,
 					ScreenVertexShader,
-					EDRF_Default);
-			}
-			else
-			{
-				check(RHISupportsVertexShaderLayer(GShaderPlatformForFeatureLevel[SceneRenderer->FeatureLevel]));
-
-				// Set shaders and texture
-				TShaderMapRef<TScreenVSForGS<true>> ScreenVertexShader(View.ShaderMap);
-				TShaderMapRef<FCopyShadowMapsCubePS> PixelShader(View.ShaderMap);
+					EDRF_Default,
+					InstanceCount);
+			});
+		}
+		else
+		{
+			TShaderMapRef<FScreenVS> ScreenVertexShader(View.ShaderMap);
+			TShaderMapRef<FCopyShadowMaps2DPS> PixelShader(View.ShaderMap);
+			GraphicsPSOInit.BoundShaderState.VertexShaderRHI = ScreenVertexShader.GetVertexShader();
+			GraphicsPSOInit.BoundShaderState.PixelShaderRHI = PixelShader.GetPixelShader();
 
-				GraphicsPSOInit.BoundShaderState.VertexDeclarationRHI = GFilterVertexDeclaration.VertexDeclarationRHI;
-				GraphicsPSOInit.BoundShaderState.VertexShaderRHI = ScreenVertexShader.GetVertexShader();
-				GraphicsPSOInit.BoundShaderState.PixelShaderRHI = PixelShader.GetPixelShader();
-				GraphicsPSOInit.PrimitiveType = PT_TriangleList;
+			auto* PassParameters = GraphBuilder.AllocParameters<FCopyShadowMaps2DPS::FParameters>();
+			PassParameters->RenderTargets = RenderTargetBindingSlots;
+			PassParameters->ShadowDepthTexture = ShadowDepthTexture;
+			PassParameters->ShadowDepthSampler = TStaticSamplerState<SF_Point, AM_Clamp, AM_Clamp, AM_Clamp>::GetRHI();
 
+			GraphBuilder.AddPass(
+				RDG_EVENT_NAME("CopyCachedShadowMap"),
+				PassParameters,
+				ERDGPassFlags::Raster,
+				[this, ScreenVertexShader, PixelShader, GraphicsPSOInit, PassParameters, ShadowDepthExtent, StencilRef](FRHICommandList& RHICmdList) mutable
+			{
+				SetStateForView(RHICmdList);
+				RHICmdList.ApplyCachedRenderTargets(GraphicsPSOInit);
 				SetGraphicsPipelineState(RHICmdList, GraphicsPSOInit);
+				SetShaderParameters(RHICmdList, PixelShader, PixelShader.GetPixelShader(), *PassParameters);
 				RHICmdList.SetStencilRef(StencilRef);
 
-				PixelShader->SetParameters(RHICmdList, View, CachedShadowMapData.ShadowMap.DepthTarget.GetReference());
-
 				DrawRectangle(
 					RHICmdList,
 					0, 0,
@@ -1022,60 +1010,20 @@ void FProjectedShadowInfo::CopyCachedShadowMap(FRHICommandList& RHICmdList, cons
 					BorderSize, BorderSize,
 					ResolutionX, ResolutionY,
 					FIntPoint(ResolutionX, ResolutionY),
-					CachedShadowMapData.ShadowMap.GetSize(),
+					ShadowDepthExtent,
 					ScreenVertexShader,
-					EDRF_Default,
-					6);
-			}
-		}
-		else
-		{
-			// Set shaders and texture
-			TShaderMapRef<FScreenVS> ScreenVertexShader(View.ShaderMap);
-			TShaderMapRef<FCopyShadowMaps2DPS> PixelShader(View.ShaderMap);
-
-			GraphicsPSOInit.BoundShaderState.VertexDeclarationRHI = GFilterVertexDeclaration.VertexDeclarationRHI;
-			GraphicsPSOInit.BoundShaderState.VertexShaderRHI = ScreenVertexShader.GetVertexShader();
-			GraphicsPSOInit.BoundShaderState.PixelShaderRHI = PixelShader.GetPixelShader();
-			GraphicsPSOInit.PrimitiveType = PT_TriangleList;
-
-			SetGraphicsPipelineState(RHICmdList, GraphicsPSOInit);
-			RHICmdList.SetStencilRef(StencilRef);
-
-			PixelShader->SetParameters(RHICmdList, View, CachedShadowMapData.ShadowMap.DepthTarget.GetReference());
-
-			DrawRectangle(
-				RHICmdList,
-				0, 0,
-				ResolutionX, ResolutionY,
-				BorderSize, BorderSize,
-				ResolutionX, ResolutionY,
-				FIntPoint(ResolutionX, ResolutionY),
-				CachedShadowMapData.ShadowMap.GetSize(),
-				ScreenVertexShader,
-				EDRF_Default);
+					EDRF_Default);
+			});
 		}
 	}
 }
 
-
-void FProjectedShadowInfo::SetupShadowUniformBuffers(FRHICommandListImmediate& RHICmdList, FScene* Scene)
+void FProjectedShadowInfo::BeginRenderView(FRDGBuilder& GraphBuilder, FScene* Scene)
 {
-	const ERHIFeatureLevel::Type FeatureLevel = ShadowDepthView->FeatureLevel;
-	if (FSceneInterface::GetShadingPath(FeatureLevel) == EShadingPath::Deferred)
+	if (DependentView)
 	{
-		FShadowDepthPassUniformParameters ShadowDepthPassParameters;
-		SetupShadowDepthPassUniformBuffer(this, RHICmdList, *ShadowDepthView, ShadowDepthPassParameters);
-
-		if (IsWholeSceneDirectionalShadow())
-		{
-			check(GetShadowDepthType() == CSMShadowDepthType);
-			Scene->UniformBuffers.CSMShadowDepthPassUniformBuffer.UpdateUniformBufferImmediate(ShadowDepthPassParameters);
-		}
-
-		ShadowDepthPassUniformBuffer.UpdateUniformBufferImmediate(ShadowDepthPassParameters);
-
-		if (DependentView)
+		const ERHIFeatureLevel::Type FeatureLevel = ShadowDepthView->FeatureLevel;
+		if (FSceneInterface::GetShadingPath(FeatureLevel) == EShadingPath::Deferred)
 		{
 			extern TSet<IPersistentViewUniformBufferExtension*> PersistentViewUniformBufferExtensions;
 
@@ -1085,107 +1033,119 @@ void FProjectedShadowInfo::SetupShadowUniformBuffers(FRHICommandListImmediate& R
 			}
 		}
 	}
-	
+
 	// This needs to be done for both mobile and deferred
-	UploadDynamicPrimitiveShaderDataForView(RHICmdList, *Scene, *ShadowDepthView);
+	UploadDynamicPrimitiveShaderDataForView(GraphBuilder.RHICmdList, *Scene, *ShadowDepthView);
 }
 
-void FProjectedShadowInfo::TransitionCachedShadowmap(FRHICommandListImmediate& RHICmdList, FScene* Scene)
+static bool IsShadowDepthPassWaitForTasksEnabled()
 {
-	if (CacheMode == SDCM_MovablePrimitivesOnly)
-	{
-		const FCachedShadowMapData& CachedShadowMapData = Scene->CachedShadowMaps.FindChecked(GetLightSceneInfo().Id);
-		if (CachedShadowMapData.bCachedShadowMapHasPrimitives && CachedShadowMapData.ShadowMap.IsValid())
-		{
-			RHICmdList.Transition(FRHITransitionInfo(CachedShadowMapData.ShadowMap.DepthTarget->GetRenderTargetItem().ShaderResourceTexture, ERHIAccess::Unknown, ERHIAccess::SRVGraphics));
-		}
-	}
+	return CVarRHICmdFlushRenderThreadTasksShadowPass.GetValueOnRenderThread() > 0 || CVarRHICmdFlushRenderThreadTasks.GetValueOnRenderThread() > 0;
 }
 
+BEGIN_SHADER_PARAMETER_STRUCT(FShadowDepthPassParameters, )
+	SHADER_PARAMETER_RDG_UNIFORM_BUFFER(FMobileShadowDepthPassUniformParameters, MobilePassUniformBuffer)
+	SHADER_PARAMETER_RDG_UNIFORM_BUFFER(FShadowDepthPassUniformParameters, DeferredPassUniformBuffer)
+	RENDER_TARGET_BINDING_SLOTS()
+END_SHADER_PARAMETER_STRUCT()
 
-void FProjectedShadowInfo::RenderDepthInner(FRHICommandListImmediate& RHICmdList, FSceneRenderer* SceneRenderer, FBeginShadowRenderPassFunction BeginShadowRenderPass, bool bDoParallelDispatch)
+void FProjectedShadowInfo::RenderDepth(
+	FRDGBuilder& GraphBuilder,
+	const FSceneRenderer* SceneRenderer,
+	FRDGTextureRef ShadowDepthTexture,
+	bool bDoParallelDispatch)
 {
-	const ERHIFeatureLevel::Type FeatureLevel = ShadowDepthView->FeatureLevel;
-	FRHIUniformBuffer* PassUniformBuffer = ShadowDepthPassUniformBuffer;
-
-	const bool bIsWholeSceneDirectionalShadow = IsWholeSceneDirectionalShadow();
+#if WANTS_DRAW_MESH_EVENTS
+	FString EventName;
 
-	if (bIsWholeSceneDirectionalShadow)
+	if (GetEmitDrawEvents())
 	{
-		// CSM shadow depth cached mesh draw commands are all referencing the same view uniform buffer.  We need to update it before rendering each cascade.
-		ShadowDepthView->ViewUniformBuffer.UpdateUniformBufferImmediate(*ShadowDepthView->CachedViewUniformShaderParameters);
-		
-		if (DependentView)
-		{
-			extern TSet<IPersistentViewUniformBufferExtension*> PersistentViewUniformBufferExtensions;
-
-			for (IPersistentViewUniformBufferExtension* Extension : PersistentViewUniformBufferExtensions)
-			{
-				Extension->BeginRenderView(DependentView);
-			}
-		}
+		GetShadowTypeNameForDrawEvent(EventName);
+		EventName += FString(TEXT(" ")) + FString::FromInt(ResolutionX) + TEXT("x") + FString::FromInt(ResolutionY);
 	}
 
-	if (FSceneInterface::GetShadingPath(FeatureLevel) == EShadingPath::Mobile)
-	{
-		FMobileShadowDepthPassUniformParameters ShadowDepthPassParameters;
-		SetupShadowDepthPassUniformBuffer(this, RHICmdList, *ShadowDepthView, ShadowDepthPassParameters);
-		SceneRenderer->Scene->UniformBuffers.MobileCSMShadowDepthPassUniformBuffer.UpdateUniformBufferImmediate(ShadowDepthPassParameters);
-		MobileShadowDepthPassUniformBuffer.UpdateUniformBufferImmediate(ShadowDepthPassParameters);
-		PassUniformBuffer = SceneRenderer->Scene->UniformBuffers.MobileCSMShadowDepthPassUniformBuffer;
-	}
+	RDG_EVENT_SCOPE(GraphBuilder, "%s", *EventName);
+#endif
+
+	CONDITIONAL_SCOPE_CYCLE_COUNTER(STAT_RenderWholeSceneShadowDepthsTime, bWholeSceneShadow);
+	CONDITIONAL_SCOPE_CYCLE_COUNTER(STAT_RenderPerObjectShadowDepthsTime, !bWholeSceneShadow);
+	QUICK_SCOPE_CYCLE_COUNTER(STAT_RenderShadowDepth);
+
+	FScene* Scene = SceneRenderer->Scene;
+	const ERHIFeatureLevel::Type FeatureLevel = ShadowDepthView->FeatureLevel;
+	BeginRenderView(GraphBuilder, Scene);
 
-	FMeshPassProcessorRenderState DrawRenderState(*ShadowDepthView, PassUniformBuffer);
-	SetStateForShadowDepth(bOnePassPointLightShadow, DrawRenderState);
-	SetStateForView(RHICmdList);
+	FShadowDepthPassParameters* PassParameters = GraphBuilder.AllocParameters<FShadowDepthPassParameters>();
+	PassParameters->RenderTargets.DepthStencil = FDepthStencilBinding(
+		ShadowDepthTexture,
+		ERenderTargetLoadAction::ELoad,
+		ERenderTargetLoadAction::ENoAction,
+		FExclusiveDepthStencil::DepthWrite_StencilNop);
 
 	if (CacheMode == SDCM_MovablePrimitivesOnly)
 	{
-		// In parallel mode we will not have a renderpass active at this point.
-		if (bDoParallelDispatch)
-		{
-			BeginShadowRenderPass(RHICmdList, false);
-		}
-
-		// Copy in depths of static primitives before we render movable primitives
-		CopyCachedShadowMap(RHICmdList, DrawRenderState, SceneRenderer, *ShadowDepthView);
+		// Copy in depths of static primitives before we render movable primitives.
+		FMeshPassProcessorRenderState DrawRenderState(*ShadowDepthView);
+		SetStateForShadowDepth(bOnePassPointLightShadow, DrawRenderState);
+		CopyCachedShadowMap(GraphBuilder, *ShadowDepthView, SceneRenderer, PassParameters->RenderTargets, DrawRenderState);
+	}
 
-		if (bDoParallelDispatch)
+	const auto TryUpdateCachedView = [this]()
+	{
+		if (IsWholeSceneDirectionalShadow())
 		{
-			RHICmdList.EndRenderPass();
+			// CSM shadow depth cached mesh draw commands are all referencing the same view uniform buffer.  We need to update it before rendering each cascade.
+			ShadowDepthView->ViewUniformBuffer.UpdateUniformBufferImmediate(*ShadowDepthView->CachedViewUniformShaderParameters);
 		}
+	};
+
+	switch (FSceneInterface::GetShadingPath(FeatureLevel))
+	{
+	case EShadingPath::Deferred:
+	{
+		auto* ShadowDepthPassParameters = GraphBuilder.AllocParameters<FShadowDepthPassUniformParameters>();
+		SetupShadowDepthPassUniformBuffer(this, GraphBuilder, *ShadowDepthView, *ShadowDepthPassParameters);
+		PassParameters->DeferredPassUniformBuffer = GraphBuilder.CreateUniformBuffer(ShadowDepthPassParameters);
+	}
+	break;
+	case EShadingPath::Mobile:
+	{
+		auto* ShadowDepthPassParameters = GraphBuilder.AllocParameters<FMobileShadowDepthPassUniformParameters>();
+		SetupShadowDepthPassUniformBuffer(this, GraphBuilder, *ShadowDepthView, *ShadowDepthPassParameters);
+		PassParameters->MobilePassUniformBuffer = GraphBuilder.CreateUniformBuffer(ShadowDepthPassParameters);
+	}
+	break;
+	default:
+		checkNoEntry();
 	}
 
 	if (bDoParallelDispatch)
 	{
-		check(IsInRenderingThread());
-		// Parallel encoding requires its own renderpass.
-		check(RHICmdList.IsOutsideRenderPass());
+		RDG_WAIT_FOR_TASKS_CONDITIONAL(GraphBuilder, IsShadowDepthPassWaitForTasksEnabled());
 
-		// parallel version
-		bool bFlush = CVarRHICmdFlushRenderThreadTasksShadowPass.GetValueOnRenderThread() > 0
-			|| CVarRHICmdFlushRenderThreadTasks.GetValueOnRenderThread() > 0;
-		FScopedCommandListWaitForTasks Flusher(bFlush);
-
-		// Dispatch commands
+		GraphBuilder.AddPass(
+			RDG_EVENT_NAME("ShadowDepthPassParallel"),
+			PassParameters,
+			ERDGPassFlags::Raster | ERDGPassFlags::SkipRenderPass,
+			[this, PassParameters, TryUpdateCachedView](FRHICommandListImmediate& RHICmdList)
 		{
-			FShadowParallelCommandListSet ParallelCommandListSet(RHICmdList, *ShadowDepthView, !bFlush, *this, BeginShadowRenderPass);
-
+			TryUpdateCachedView();
+			FShadowParallelCommandListSet ParallelCommandListSet(RHICmdList, *ShadowDepthView, *this, FParallelCommandListBindings(PassParameters));
 			ShadowDepthPass.DispatchDraw(&ParallelCommandListSet, RHICmdList);
-		}
-
-		// Renderpass must be closed once we get here.
-		check(RHICmdList.IsOutsideRenderPass());
+		});
 	}
 	else
 	{
-		// We must have already opened the renderpass by the time we get here.
-		check(RHICmdList.IsInsideRenderPass());
-
-		ShadowDepthPass.DispatchDraw(nullptr, RHICmdList);
-
-		// Renderpass must still be open when we reach here
-		check(RHICmdList.IsInsideRenderPass());
+		GraphBuilder.AddPass(
+			RDG_EVENT_NAME("ShadowDepthPass"),
+			PassParameters,
+			ERDGPassFlags::Raster,
+			[this, TryUpdateCachedView](FRHICommandList& RHICmdList)
+		{
+			TryUpdateCachedView();
+			SetStateForView(RHICmdList);
+			ShadowDepthPass.DispatchDraw(nullptr, RHICmdList);
+		});
 	}
 }
 
@@ -1258,27 +1218,6 @@ FViewInfo* FProjectedShadowInfo::FindViewForShadow(FSceneRenderer* SceneRenderer
 	return FoundView;
 }
 
-void FProjectedShadowInfo::RenderDepth(FRHICommandListImmediate& RHICmdList, FSceneRenderer* SceneRenderer, FBeginShadowRenderPassFunction BeginShadowRenderPass, bool bDoParallelDispatch)
-{
-#if WANTS_DRAW_MESH_EVENTS
-	FString EventName;
-
-	if (GetEmitDrawEvents())
-	{
-		GetShadowTypeNameForDrawEvent(EventName);
-		EventName += FString(TEXT(" ")) + FString::FromInt(ResolutionX) + TEXT("x") + FString::FromInt(ResolutionY);
-	}
-
-	SCOPED_DRAW_EVENTF(RHICmdList, EventShadowDepthActor, *EventName);
-#endif
-
-	CONDITIONAL_SCOPE_CYCLE_COUNTER(STAT_RenderWholeSceneShadowDepthsTime, bWholeSceneShadow);
-	CONDITIONAL_SCOPE_CYCLE_COUNTER(STAT_RenderPerObjectShadowDepthsTime, !bWholeSceneShadow);
-	QUICK_SCOPE_CYCLE_COUNTER(STAT_RenderShadowDepth);
-
-	RenderDepthInner(RHICmdList, SceneRenderer, BeginShadowRenderPass, bDoParallelDispatch);
-}
-
 void FProjectedShadowInfo::SetupShadowDepthView(FRHICommandListImmediate& RHICmdList, FSceneRenderer* SceneRenderer)
 {
 	FViewInfo* FoundView = FindViewForShadow(SceneRenderer);
@@ -1459,7 +1398,7 @@ static void RenderShadowDepthAtlasNanite(
 			);
 		}
 
-		FRDGTextureRef ShadowMap = GraphBuilder.RegisterExternalTexture(ShadowMapAtlas.RenderTargets.DepthTarget, TEXT("DepthBuffer"));
+		FRDGTextureRef ShadowMap = GraphBuilder.RegisterExternalTexture(ShadowMapAtlas.RenderTargets.DepthTarget);
 		for (FProjectedShadowInfo* ProjectedShadowInfo : ShadowsToEmit)
 		{
 			const FIntRect AtlasViewRect = ProjectedShadowInfo->GetViewRectForView();
@@ -1532,17 +1471,18 @@ static void CopyToCompleteShadowMap(
 	);
 }
 
+bool IsParallelDispatchEnabled(const FProjectedShadowInfo* ProjectedShadowInfo)
+{
+	return GRHICommandList.UseParallelAlgorithms() && CVarParallelShadows.GetValueOnRenderThread()
+		&& (ProjectedShadowInfo->IsWholeSceneDirectionalShadow() || CVarParallelShadowsNonWholeScene.GetValueOnRenderThread());
+}
+
 void FSceneRenderer::RenderShadowDepthMapAtlases(FRDGBuilder& GraphBuilder)
 {
 	// Perform setup work on all GPUs in case any cached shadows are being updated this
 	// frame. We revert to the AllViewsGPUMask for uncached shadows.
 	RDG_GPU_MASK_SCOPE(GraphBuilder, FRHIGPUMask::All());
 
-	FSceneRenderTargets& SceneContext = FSceneRenderTargets::Get(GraphBuilder.RHICmdList);
-
-	bool bCanUseParallelDispatch = GraphBuilder.RHICmdList.IsImmediate() &&  // translucent shadows are draw on the render thread, using a recursive cmdlist (which is not immediate)
-		GRHICommandList.UseParallelAlgorithms() && CVarParallelShadows.GetValueOnRenderThread();
-
 	const bool bNaniteEnabled = 
 		UseNanite(ShaderPlatform) &&
 		ViewFamily.EngineShowFlags.NaniteMeshes &&
@@ -1551,197 +1491,95 @@ void FSceneRenderer::RenderShadowDepthMapAtlases(FRDGBuilder& GraphBuilder)
 	for (int32 AtlasIndex = 0; AtlasIndex < SortedShadowsForShadowDepthPass.ShadowMapAtlases.Num(); AtlasIndex++)
 	{
 		const FSortedShadowMapAtlas& ShadowMapAtlas = SortedShadowsForShadowDepthPass.ShadowMapAtlases[AtlasIndex];
-		FSceneRenderTargetItem& RenderTarget = ShadowMapAtlas.RenderTargets.DepthTarget->GetRenderTargetItem();
-		FIntPoint AtlasSize = ShadowMapAtlas.RenderTargets.DepthTarget->GetDesc().Extent;
+		FRDGTextureRef AtlasDepthTexture = GraphBuilder.RegisterExternalTexture(ShadowMapAtlas.RenderTargets.DepthTarget);
+		const FIntPoint AtlasSize = AtlasDepthTexture->Desc.Extent;
 
-		AddUntrackedAccessPass(GraphBuilder, [this, &ShadowMapAtlas , &RenderTarget, &SceneContext, AtlasIndex, AtlasSize, bCanUseParallelDispatch](FRHICommandListImmediate& RHICmdList)
-		{
-			GVisualizeTexture.SetCheckPoint(RHICmdList, ShadowMapAtlas.RenderTargets.DepthTarget.GetReference());
+		RDG_EVENT_SCOPE(GraphBuilder, "Atlas%u %ux%u", AtlasIndex, AtlasSize.X, AtlasSize.Y);
 
-			SCOPED_DRAW_EVENTF(RHICmdList, EventShadowDepths, TEXT("Atlas%u %ux%u"), AtlasIndex, AtlasSize.X, AtlasSize.Y);
+		TArray<FProjectedShadowInfo*, SceneRenderingAllocator> ParallelShadowPasses;
+		TArray<FProjectedShadowInfo*, SceneRenderingAllocator> SerialShadowPasses;
 
-			auto BeginShadowRenderPass = [this, &RenderTarget, &SceneContext](FRHICommandList& InRHICmdList, bool bPerformClear)
+		// Gather our passes here to minimize switching render passes
+		for (FProjectedShadowInfo* ProjectedShadowInfo : ShadowMapAtlas.Shadows)
+		{
+			if (IsParallelDispatchEnabled(ProjectedShadowInfo))
 			{
-				check(RenderTarget.TargetableTexture->GetDepthClearValue() == 1.0f);
-
-				ERenderTargetLoadAction DepthLoadAction = bPerformClear ? ERenderTargetLoadAction::EClear : ERenderTargetLoadAction::ELoad;
-
-				FRHIRenderPassInfo RPInfo(RenderTarget.TargetableTexture, MakeDepthStencilTargetActions(MakeRenderTargetActions(DepthLoadAction, ERenderTargetStoreAction::EStore), ERenderTargetActions::DontLoad_DontStore), nullptr, FExclusiveDepthStencil::DepthWrite_StencilNop);
-
-				if (!GSupportsDepthRenderTargetWithoutColorRenderTarget)
-				{
-					RPInfo.ColorRenderTargets[0].Action = ERenderTargetActions::DontLoad_DontStore;
-					RPInfo.ColorRenderTargets[0].RenderTarget = SceneContext.GetOptionalShadowDepthColorSurface(InRHICmdList, RPInfo.DepthStencilRenderTarget.DepthStencilTarget->GetTexture2D()->GetSizeX(), RPInfo.DepthStencilRenderTarget.DepthStencilTarget->GetTexture2D()->GetSizeY());
-					InRHICmdList.Transition(FRHITransitionInfo(RPInfo.ColorRenderTargets[0].RenderTarget, ERHIAccess::Unknown, ERHIAccess::RTV));
-				}
-				InRHICmdList.Transition(FRHITransitionInfo(RPInfo.DepthStencilRenderTarget.DepthStencilTarget, ERHIAccess::Unknown, ERHIAccess::DSVWrite));
-				InRHICmdList.BeginRenderPass(RPInfo, TEXT("ShadowMapAtlases"));
-			};
-
-			TArray<FProjectedShadowInfo*, SceneRenderingAllocator> ParallelShadowPasses;
-			TArray<FProjectedShadowInfo*, SceneRenderingAllocator> SerialShadowPasses;
-
-			// Gather our passes here to minimize switching render passes
-			for (int32 ShadowIndex = 0; ShadowIndex < ShadowMapAtlas.Shadows.Num(); ShadowIndex++)
+				ParallelShadowPasses.Add(ProjectedShadowInfo);
+			}
+			else
 			{
-				FProjectedShadowInfo* ProjectedShadowInfo = ShadowMapAtlas.Shadows[ShadowIndex];
-
-				const bool bDoParallelDispatch = bCanUseParallelDispatch &&
-					(ProjectedShadowInfo->IsWholeSceneDirectionalShadow() || CVarParallelShadowsNonWholeScene.GetValueOnRenderThread());
-
-				if (bDoParallelDispatch)
-				{
-					ParallelShadowPasses.Add(ProjectedShadowInfo);
-				}
-				else
-				{
-					SerialShadowPasses.Add(ProjectedShadowInfo);
-				}
+				SerialShadowPasses.Add(ProjectedShadowInfo);
 			}
+		}
 
-			FLightSceneProxy* CurrentLightForDrawEvent = NULL;
+	#if WANTS_DRAW_MESH_EVENTS
+		FLightSceneProxy* CurrentLightForDrawEvent = nullptr;
+		FDrawEvent LightEvent;
+	#endif
 
+		const auto SetLightEventForShadow = [&](FProjectedShadowInfo* ProjectedShadowInfo)
+		{
 		#if WANTS_DRAW_MESH_EVENTS
-			FDrawEvent LightEvent;
-		#endif
-
-			if (ParallelShadowPasses.Num() > 0)
+			if (!CurrentLightForDrawEvent || ProjectedShadowInfo->GetLightSceneInfo().Proxy != CurrentLightForDrawEvent)
 			{
+				if (CurrentLightForDrawEvent)
 				{
-					// Clear before going wide.
-					SCOPED_DRAW_EVENT(RHICmdList, SetShadowRTsAndClear);
-					BeginShadowRenderPass(RHICmdList, true);
-					RHICmdList.EndRenderPass();
+					GraphBuilder.EndEventScope();
 				}
 
-				for (int32 ShadowIndex = 0; ShadowIndex < ParallelShadowPasses.Num(); ShadowIndex++)
-				{
-					FProjectedShadowInfo* ProjectedShadowInfo = ParallelShadowPasses[ShadowIndex];
-					SCOPED_GPU_MASK(RHICmdList, GetGPUMaskForShadow(ProjectedShadowInfo));
-
-				#if WANTS_DRAW_MESH_EVENTS
-					if (!CurrentLightForDrawEvent || ProjectedShadowInfo->GetLightSceneInfo().Proxy != CurrentLightForDrawEvent)
-					{
-						if (CurrentLightForDrawEvent)
-						{
-							SCOPED_GPU_MASK(RHICmdList, FRHIGPUMask::All());
-							STOP_DRAW_EVENT(LightEvent);
-						}
-
-						CurrentLightForDrawEvent = ProjectedShadowInfo->GetLightSceneInfo().Proxy;
-						FString LightNameWithLevel;
-						GetLightNameForDrawEvent(CurrentLightForDrawEvent, LightNameWithLevel);
-
-						SCOPED_GPU_MASK(RHICmdList, FRHIGPUMask::All());
-						BEGIN_DRAW_EVENTF(
-							RHICmdList,
-							LightNameEvent,
-							LightEvent,
-							*LightNameWithLevel);
-					}
-				#endif
-
-					ProjectedShadowInfo->SetupShadowUniformBuffers(RHICmdList, Scene);
-					ProjectedShadowInfo->TransitionCachedShadowmap(RHICmdList, Scene);
-					ProjectedShadowInfo->RenderDepth(RHICmdList, this, BeginShadowRenderPass, true);
-				}
+				CurrentLightForDrawEvent = ProjectedShadowInfo->GetLightSceneInfo().Proxy;
+				FString LightNameWithLevel;
+				GetLightNameForDrawEvent(CurrentLightForDrawEvent, LightNameWithLevel);
+				GraphBuilder.BeginEventScope(RDG_EVENT_NAME("%s", *LightNameWithLevel));
 			}
+		#endif
+		};
 
+		const auto EndLightEvent = [&]()
+		{
 		#if WANTS_DRAW_MESH_EVENTS
 			if (CurrentLightForDrawEvent)
 			{
-				SCOPED_GPU_MASK(RHICmdList, FRHIGPUMask::All());
-				STOP_DRAW_EVENT(LightEvent);
+				GraphBuilder.EndEventScope();
+				CurrentLightForDrawEvent = nullptr;
 			}
 		#endif
+		};
 
-			CurrentLightForDrawEvent = nullptr;
+		AddClearShadowDepthPass(GraphBuilder, AtlasDepthTexture);
 
-			if (SerialShadowPasses.Num() > 0)
+		if (ParallelShadowPasses.Num() > 0)
+		{
+			for (FProjectedShadowInfo* ProjectedShadowInfo : ParallelShadowPasses)
 			{
-				bool bForceSingleRenderPass = CVarShadowForceSerialSingleRenderPass.GetValueOnAnyThread() != 0;
-				if (bForceSingleRenderPass)
-				{
-					SCOPED_GPU_MASK(RHICmdList, AllViewsGPUMask);
-					BeginShadowRenderPass(RHICmdList, true);
-				}
-
-				for (int32 ShadowIndex = 0; ShadowIndex < SerialShadowPasses.Num(); ShadowIndex++)
-				{
-					FProjectedShadowInfo* ProjectedShadowInfo = SerialShadowPasses[ShadowIndex];
-					SCOPED_GPU_MASK(RHICmdList, GetGPUMaskForShadow(ProjectedShadowInfo));
+				RDG_GPU_MASK_SCOPE(GraphBuilder, GetGPUMaskForShadow(ProjectedShadowInfo));
+				SetLightEventForShadow(ProjectedShadowInfo);
 
-				#if WANTS_DRAW_MESH_EVENTS
-					if (!CurrentLightForDrawEvent || ProjectedShadowInfo->GetLightSceneInfo().Proxy != CurrentLightForDrawEvent)
-					{
-						if (CurrentLightForDrawEvent)
-						{
-							SCOPED_GPU_MASK(RHICmdList, FRHIGPUMask::All());
-							STOP_DRAW_EVENT(LightEvent);
-						}
+				const bool bParallelDispatch = true;
+				ProjectedShadowInfo->RenderDepth(GraphBuilder, this, AtlasDepthTexture, bParallelDispatch);
+			}
+		}
 
-						CurrentLightForDrawEvent = ProjectedShadowInfo->GetLightSceneInfo().Proxy;
-						FString LightNameWithLevel;
-						GetLightNameForDrawEvent(CurrentLightForDrawEvent, LightNameWithLevel);
+		EndLightEvent();
 
-						SCOPED_GPU_MASK(RHICmdList, FRHIGPUMask::All());
-						BEGIN_DRAW_EVENTF(
-							RHICmdList,
-							LightNameEvent,
-							LightEvent,
-							*LightNameWithLevel);
-					}
-				#endif
-
-					ProjectedShadowInfo->SetupShadowUniformBuffers(RHICmdList, Scene);
-					ProjectedShadowInfo->TransitionCachedShadowmap(RHICmdList, Scene);
-	#if WITH_MGPU
-					// In case the first shadow is view-dependent, ensure we do the clear on all GPUs.
-					FRHIGPUMask GPUMaskForRenderPass = RHICmdList.GetGPUMask();
-					if (ShadowIndex == 0)
-					{
-						// This ensures that we don't downgrade the GPU mask if the first shadow is a
-						// cached whole scene shadow.
-						GPUMaskForRenderPass |= AllViewsGPUMask;
-					}
-	#endif
-					if (!bForceSingleRenderPass)
-					{
-						SCOPED_GPU_MASK(RHICmdList, GPUMaskForRenderPass);
-						BeginShadowRenderPass(RHICmdList, ShadowIndex == 0);
-					}
-
-					ProjectedShadowInfo->RenderDepth(RHICmdList, this, BeginShadowRenderPass, false);
+		if (SerialShadowPasses.Num() > 0)
+		{
+			for (FProjectedShadowInfo* ProjectedShadowInfo : SerialShadowPasses)
+			{
+				RDG_GPU_MASK_SCOPE(GraphBuilder, GetGPUMaskForShadow(ProjectedShadowInfo));
+				SetLightEventForShadow(ProjectedShadowInfo);
 
-					if (!bForceSingleRenderPass)
-					{
-						RHICmdList.EndRenderPass();
-					}
-				}
-				if (bForceSingleRenderPass)
-				{
-					SCOPED_GPU_MASK(RHICmdList, AllViewsGPUMask);
-					RHICmdList.EndRenderPass();
-				}
+				const bool bParallelDispatch = false;
+				ProjectedShadowInfo->RenderDepth(GraphBuilder, this, AtlasDepthTexture, bParallelDispatch);
 			}
+		}
 
-			if (CurrentLightForDrawEvent)
-			{
-				SCOPED_GPU_MASK(RHICmdList, FRHIGPUMask::All());
-				STOP_DRAW_EVENT(LightEvent);
-				CurrentLightForDrawEvent = NULL;
-			}
-		});
+		EndLightEvent();
 
 		if (bNaniteEnabled)
 		{
 			RenderShadowDepthAtlasNanite(GraphBuilder, *Scene, ShadowMapAtlas);
 		}
-
-		AddUntrackedAccessPass(GraphBuilder, [&RenderTarget](FRHICommandList& RHICmdList)
-		{
-			RHICmdList.Transition(FRHITransitionInfo(RenderTarget.TargetableTexture, ERHIAccess::Unknown, ERHIAccess::SRVMask));
-		});
 	}
 
 	// Copy/resample shadow maps into "complete" shadow maps and add Nanite geometry
@@ -1757,9 +1595,8 @@ void FSceneRenderer::RenderShadowDepthMapAtlases(FRDGBuilder& GraphBuilder)
 			RDG_EVENT_SCOPE(GraphBuilder, "Complete Atlas %ux%u", AtlasSize.X, AtlasSize.Y);
 
 			bool bCleared = false;
-			for (int32 ShadowIndex = 0; ShadowIndex < ShadowMapAtlas.Shadows.Num(); ShadowIndex++)
+			for (FProjectedShadowInfo* ProjectedShadowInfo : ShadowMapAtlas.Shadows)
 			{
-				FProjectedShadowInfo* ProjectedShadowInfo = ShadowMapAtlas.Shadows[ShadowIndex];
 				FProjectedShadowInfo* SourceShadowInfo = ProjectedShadowInfo->CompleteShadowMapCopySource;
 
 				check(SourceShadowInfo);
@@ -2018,86 +1855,37 @@ void FSceneRenderer::RenderShadowDepthMaps(FRDGBuilder& GraphBuilder)
 
 		Scene->VirtualShadowMapArrayCacheManager->HZBPhysical  = VirtualShadowMapArray.HZBPhysical;
 		Scene->VirtualShadowMapArrayCacheManager->HZBPageTable = VirtualShadowMapArray.PageTable;
-
-		//GVisualizeTexture.SetCheckPoint(RHICmdList, VirtualShadowMapArray.PhysicalPagePool);
 	}
 
-	FSceneRenderer::RenderShadowDepthMapAtlases( GraphBuilder );	// Render non-VSM shadows. Must be after VSM so we can use TopMip optimization.
+	// Render non-VSM shadows. Must be after VSM so we can use TopMip optimization.
+	FSceneRenderer::RenderShadowDepthMapAtlases(GraphBuilder);
 
 	const bool bUseGeometryShader = !GRHISupportsArrayIndexFromAnyShader;
 
 	for (int32 CubemapIndex = 0; CubemapIndex < SortedShadowsForShadowDepthPass.ShadowMapCubemaps.Num(); CubemapIndex++)
 	{
 		const FSortedShadowMapAtlas& ShadowMap = SortedShadowsForShadowDepthPass.ShadowMapCubemaps[CubemapIndex];
-		FSceneRenderTargetItem& RenderTarget = ShadowMap.RenderTargets.DepthTarget->GetRenderTargetItem();
-		FIntPoint TargetSize = ShadowMap.RenderTargets.DepthTarget->GetDesc().Extent;
+		FRDGTextureRef ShadowDepthTexture = GraphBuilder.RegisterExternalTexture(ShadowMap.RenderTargets.DepthTarget);
+		const FIntPoint TargetSize = ShadowDepthTexture->Desc.Extent;
 
 		check(ShadowMap.Shadows.Num() == 1);
 		FProjectedShadowInfo* ProjectedShadowInfo = ShadowMap.Shadows[0];
 		RDG_GPU_MASK_SCOPE(GraphBuilder, GetGPUMaskForShadow(ProjectedShadowInfo));
 
-		AddUntrackedAccessPass(GraphBuilder, [this, ProjectedShadowInfo, &SceneContext, &ShadowMap, &RenderTarget, TargetSize](FRHICommandListImmediate& RHICmdList)
-		{
-			const bool bDoParallelDispatch = RHICmdList.IsImmediate() &&  // translucent shadows are draw on the render thread, using a recursive cmdlist (which is not immediate)
-				GRHICommandList.UseParallelAlgorithms() && CVarParallelShadows.GetValueOnRenderThread() &&
-				(ProjectedShadowInfo->IsWholeSceneDirectionalShadow() || CVarParallelShadowsNonWholeScene.GetValueOnRenderThread());
-
-			GVisualizeTexture.SetCheckPoint(RHICmdList, ShadowMap.RenderTargets.DepthTarget.GetReference());
-
-			FString LightNameWithLevel;
-			GetLightNameForDrawEvent(ProjectedShadowInfo->GetLightSceneInfo().Proxy, LightNameWithLevel);
-			SCOPED_DRAW_EVENTF(RHICmdList, EventShadowDepths, TEXT("Cubemap %s %u^2"), *LightNameWithLevel, TargetSize.X, TargetSize.Y);
+		FString LightNameWithLevel;
+		GetLightNameForDrawEvent(ProjectedShadowInfo->GetLightSceneInfo().Proxy, LightNameWithLevel);
+		RDG_EVENT_SCOPE(GraphBuilder, "Cubemap %s %u^2", *LightNameWithLevel, TargetSize.X, TargetSize.Y);
 
-			ProjectedShadowInfo->SetupShadowUniformBuffers(RHICmdList, Scene);
-
-			auto BeginShadowRenderPass = [this, &RenderTarget, &SceneContext](FRHICommandList& InRHICmdList, bool bPerformClear)
-			{
-				FRHITexture* DepthTarget = RenderTarget.TargetableTexture;
-				ERenderTargetLoadAction DepthLoadAction = bPerformClear ? ERenderTargetLoadAction::EClear : ERenderTargetLoadAction::ELoad;
-
-				check(DepthTarget->GetDepthClearValue() == FClearValueBinding::DepthFar.Value.DSValue.Depth);
-				FRHIRenderPassInfo RPInfo(DepthTarget, MakeDepthStencilTargetActions(MakeRenderTargetActions(DepthLoadAction, ERenderTargetStoreAction::EStore), ERenderTargetActions::DontLoad_DontStore), nullptr, FExclusiveDepthStencil::DepthWrite_StencilNop);
-
-				if (!GSupportsDepthRenderTargetWithoutColorRenderTarget)
-				{
-					RPInfo.ColorRenderTargets[0].Action = ERenderTargetActions::DontLoad_DontStore;
-					RPInfo.ColorRenderTargets[0].ArraySlice = -1;
-					RPInfo.ColorRenderTargets[0].MipIndex = 0;
-					RPInfo.ColorRenderTargets[0].RenderTarget = SceneContext.GetOptionalShadowDepthColorSurface(InRHICmdList, DepthTarget->GetTexture2D()->GetSizeX(), DepthTarget->GetTexture2D()->GetSizeY());
-
-					InRHICmdList.Transition(FRHITransitionInfo(RPInfo.ColorRenderTargets[0].RenderTarget, ERHIAccess::Unknown, ERHIAccess::RTV));
-				}
-				InRHICmdList.Transition(FRHITransitionInfo(DepthTarget, ERHIAccess::Unknown, ERHIAccess::DSVWrite));
-				InRHICmdList.BeginRenderPass(RPInfo, TEXT("ShadowDepthCubeMaps"));
-			};
-
-			{
-				bool bDoClear = true;
-
-				if (ProjectedShadowInfo->CacheMode == SDCM_MovablePrimitivesOnly
-					&& Scene->CachedShadowMaps.FindChecked(ProjectedShadowInfo->GetLightSceneInfo().Id).bCachedShadowMapHasPrimitives)
-				{
-					// Skip the clear when we'll copy from a cached shadowmap
-					bDoClear = false;
-				}
-
-				SCOPED_CONDITIONAL_DRAW_EVENT(RHICmdList, Clear, bDoClear);
-				BeginShadowRenderPass(RHICmdList, bDoClear);
-			}
-
-			if (bDoParallelDispatch)
-			{
-				// In parallel mode this first pass will just be the clear.
-				RHICmdList.EndRenderPass();
-			}
-
-			ProjectedShadowInfo->RenderDepth(RHICmdList, this, BeginShadowRenderPass, bDoParallelDispatch);
+		// Only clear when we're not copying from a cached shadow map.
+		if (ProjectedShadowInfo->CacheMode != SDCM_MovablePrimitivesOnly || !Scene->CachedShadowMaps.FindChecked(ProjectedShadowInfo->GetLightSceneInfo().Id).bCachedShadowMapHasPrimitives)
+		{
+			AddClearShadowDepthPass(GraphBuilder, ShadowDepthTexture);
+		}
 
-			if (!bDoParallelDispatch)
-			{
-				RHICmdList.EndRenderPass();
-			}
-		});
+		{
+			const bool bDoParallelDispatch = IsParallelDispatchEnabled(ProjectedShadowInfo);
+			ProjectedShadowInfo->RenderDepth(GraphBuilder, this, ShadowDepthTexture, bDoParallelDispatch);
+		}
 
 		if (bNaniteEnabled &&
 			CVarNaniteShadows.GetValueOnRenderThread() &&
@@ -2191,109 +1979,50 @@ void FSceneRenderer::RenderShadowDepthMaps(FRDGBuilder& GraphBuilder)
 				}
 			}
 		}
-
-		AddUntrackedAccessPass(GraphBuilder, [&RenderTarget](FRHICommandList& RHICmdList)
-		{
-			RHICmdList.Transition(FRHITransitionInfo(RenderTarget.TargetableTexture, ERHIAccess::Unknown, ERHIAccess::SRVMask));
-		});
 	}
 
-	AddUntrackedAccessPass(GraphBuilder, [this](FRHICommandListImmediate& RHICmdList)
+	if (SortedShadowsForShadowDepthPass.PreshadowCache.Shadows.Num() > 0)
 	{
-		if (SortedShadowsForShadowDepthPass.PreshadowCache.Shadows.Num() > 0)
-		{
-			FSceneRenderTargetItem& RenderTarget = SortedShadowsForShadowDepthPass.PreshadowCache.RenderTargets.DepthTarget->GetRenderTargetItem();
-
-			GVisualizeTexture.SetCheckPoint(RHICmdList, SortedShadowsForShadowDepthPass.PreshadowCache.RenderTargets.DepthTarget.GetReference());
+		RDG_EVENT_SCOPE(GraphBuilder, "PreshadowCache");
 
-			SCOPED_DRAW_EVENT(RHICmdList, PreshadowCache);
+		FRDGTextureRef PreshadowCacheTexture = GraphBuilder.RegisterExternalTexture(SortedShadowsForShadowDepthPass.PreshadowCache.RenderTargets.DepthTarget);
 
-			for (int32 ShadowIndex = 0; ShadowIndex < SortedShadowsForShadowDepthPass.PreshadowCache.Shadows.Num(); ShadowIndex++)
+		for (FProjectedShadowInfo* ProjectedShadowInfo : SortedShadowsForShadowDepthPass.PreshadowCache.Shadows)
+		{
+			if (!ProjectedShadowInfo->bDepthsCached)
 			{
-				FProjectedShadowInfo* ProjectedShadowInfo = SortedShadowsForShadowDepthPass.PreshadowCache.Shadows[ShadowIndex];
-
-				if (!ProjectedShadowInfo->bDepthsCached)
-				{
-					SCOPED_GPU_MASK(RHICmdList, GetGPUMaskForShadow(ProjectedShadowInfo));
-
-					const bool bDoParallelDispatch = RHICmdList.IsImmediate() &&  // translucent shadows are draw on the render thread, using a recursive cmdlist (which is not immediate)
-						GRHICommandList.UseParallelAlgorithms() && CVarParallelShadows.GetValueOnRenderThread() &&
-						(ProjectedShadowInfo->IsWholeSceneDirectionalShadow() || CVarParallelShadowsNonWholeScene.GetValueOnRenderThread());
+				RDG_GPU_MASK_SCOPE(GraphBuilder, GetGPUMaskForShadow(ProjectedShadowInfo));
+				AddClearShadowDepthPass(GraphBuilder, PreshadowCacheTexture, ProjectedShadowInfo);
 
-					ProjectedShadowInfo->SetupShadowUniformBuffers(RHICmdList, Scene);
-
-					auto BeginShadowRenderPass = [this, ProjectedShadowInfo](FRHICommandList& InRHICmdList, bool bPerformClear)
-					{
-						FRHITexture* PreShadowCacheDepthZ = Scene->PreShadowCacheDepthZ->GetRenderTargetItem().TargetableTexture.GetReference();
-						InRHICmdList.TransitionResources(ERHIAccess::DSVWrite, &PreShadowCacheDepthZ, 1);
-
-						FRHIRenderPassInfo RPInfo(PreShadowCacheDepthZ, EDepthStencilTargetActions::LoadDepthNotStencil_StoreDepthNotStencil, nullptr, FExclusiveDepthStencil::DepthWrite_StencilNop);
-
-						// Must preserve existing contents as the clear will be scissored
-						InRHICmdList.BeginRenderPass(RPInfo, TEXT("ShadowDepthMaps"));
-						ProjectedShadowInfo->ClearDepth(InRHICmdList, this, 0, nullptr, PreShadowCacheDepthZ, bPerformClear);
-					};
-
-					BeginShadowRenderPass(RHICmdList, true);
-
-					if (bDoParallelDispatch)
-					{
-						// In parallel mode the first pass is just the clear.
-						RHICmdList.EndRenderPass();
-					}
-
-					ProjectedShadowInfo->RenderDepth(RHICmdList, this, BeginShadowRenderPass, bDoParallelDispatch);
-
-					if (!bDoParallelDispatch)
-					{
-						RHICmdList.EndRenderPass();
-					}
-
-					ProjectedShadowInfo->bDepthsCached = true;
-				}
+				const bool bParallelDispatch = IsParallelDispatchEnabled(ProjectedShadowInfo);
+				ProjectedShadowInfo->RenderDepth(GraphBuilder, this, PreshadowCacheTexture, bParallelDispatch);
+				ProjectedShadowInfo->bDepthsCached = true;
 			}
-
-			RHICmdList.Transition(FRHITransitionInfo(RenderTarget.TargetableTexture, ERHIAccess::Unknown, ERHIAccess::SRVMask));
 		}
+	}
 
-		for (int32 AtlasIndex = 0; AtlasIndex < SortedShadowsForShadowDepthPass.TranslucencyShadowMapAtlases.Num(); AtlasIndex++)
-		{
-			const FSortedShadowMapAtlas& ShadowMapAtlas = SortedShadowsForShadowDepthPass.TranslucencyShadowMapAtlases[AtlasIndex];
-			FIntPoint TargetSize = ShadowMapAtlas.RenderTargets.ColorTargets[0]->GetDesc().Extent;
+	for (int32 AtlasIndex = 0; AtlasIndex < SortedShadowsForShadowDepthPass.TranslucencyShadowMapAtlases.Num(); AtlasIndex++)
+	{
+		const FSortedShadowMapAtlas& ShadowMapAtlas = SortedShadowsForShadowDepthPass.TranslucencyShadowMapAtlases[AtlasIndex];
 
-			SCOPED_DRAW_EVENTF(RHICmdList, EventShadowDepths, TEXT("TranslucencyAtlas%u %u^2"), AtlasIndex, TargetSize.X, TargetSize.Y);
+		FRDGTextureRef ColorTarget0 = GraphBuilder.RegisterExternalTexture(ShadowMapAtlas.RenderTargets.ColorTargets[0]);
+		FRDGTextureRef ColorTarget1 = GraphBuilder.RegisterExternalTexture(ShadowMapAtlas.RenderTargets.ColorTargets[1]);
+		const FIntPoint TargetSize  = ColorTarget0->Desc.Extent;
 
-			FSceneRenderTargetItem ColorTarget0 = ShadowMapAtlas.RenderTargets.ColorTargets[0]->GetRenderTargetItem();
-			FSceneRenderTargetItem ColorTarget1 = ShadowMapAtlas.RenderTargets.ColorTargets[1]->GetRenderTargetItem();
+		FRenderTargetBindingSlots RenderTargets;
+		RenderTargets[0] = FRenderTargetBinding(ColorTarget0, ERenderTargetLoadAction::ELoad);
+		RenderTargets[1] = FRenderTargetBinding(ColorTarget1, ERenderTargetLoadAction::ELoad);
 
-			FRHITexture* RenderTargetArray[2] =
-			{
-				ColorTarget0.TargetableTexture,
-				ColorTarget1.TargetableTexture
-			};
+		RDG_EVENT_SCOPE(GraphBuilder, "TranslucencyAtlas%u %u^2", AtlasIndex, TargetSize.X, TargetSize.Y);
 
-			FRHIRenderPassInfo RPInfo(UE_ARRAY_COUNT(RenderTargetArray), RenderTargetArray, ERenderTargetActions::Load_Store);
-			TransitionRenderPassTargets(RHICmdList, RPInfo);
-			RHICmdList.BeginRenderPass(RPInfo, TEXT("RenderTranslucencyDepths"));
-			{
-				for (int32 ShadowIndex = 0; ShadowIndex < ShadowMapAtlas.Shadows.Num(); ShadowIndex++)
-				{
-					FProjectedShadowInfo* ProjectedShadowInfo = ShadowMapAtlas.Shadows[ShadowIndex];
-					SCOPED_GPU_MASK(RHICmdList, GetGPUMaskForShadow(ProjectedShadowInfo));
-					ProjectedShadowInfo->SetupShadowUniformBuffers(RHICmdList, Scene);
-					ProjectedShadowInfo->RenderTranslucencyDepths(RHICmdList, this);
-				}
-			}
-			RHICmdList.EndRenderPass();
-
-			RHICmdList.Transition({
-				FRHITransitionInfo(ColorTarget0.TargetableTexture, ERHIAccess::Unknown, ERHIAccess::SRVMask),
-				FRHITransitionInfo(ColorTarget1.TargetableTexture, ERHIAccess::Unknown, ERHIAccess::SRVMask)
-			});
+		for (FProjectedShadowInfo* ProjectedShadowInfo : ShadowMapAtlas.Shadows)
+		{
+			RDG_GPU_MASK_SCOPE(GraphBuilder, GetGPUMaskForShadow(ProjectedShadowInfo));
+			ProjectedShadowInfo->RenderTranslucencyDepths(GraphBuilder, this, RenderTargets);
 		}
+	}
 
-		bShadowDepthRenderCompleted = true;
-	});
+	bShadowDepthRenderCompleted = true;
 }
 
 bool FShadowDepthPassMeshProcessor::Process(
@@ -2435,11 +2164,10 @@ FShadowDepthPassMeshProcessor::FShadowDepthPassMeshProcessor(
 	const FScene* Scene,
 	const FSceneView* InViewIfDynamicMeshCommand,
 	const TUniformBufferRef<FViewUniformShaderParameters>& InViewUniformBuffer,
-	FRHIUniformBuffer* InPassUniformBuffer,
 	FShadowDepthType InShadowDepthType,
 	FMeshPassDrawListContext* InDrawListContext)
 	: FMeshPassProcessor(Scene, Scene->GetFeatureLevel(), InViewIfDynamicMeshCommand, InDrawListContext)
-	, PassDrawRenderState(FMeshPassProcessorRenderState(InViewUniformBuffer, InPassUniformBuffer))
+	, PassDrawRenderState(FMeshPassProcessorRenderState(InViewUniformBuffer))
 	, ShadowDepthType(InShadowDepthType)
 {
 	SetStateForShadowDepth(ShadowDepthType.bOnePassPointLightShadow, PassDrawRenderState);
@@ -2449,23 +2177,10 @@ FShadowDepthType CSMShadowDepthType(true, false);
 
 FMeshPassProcessor* CreateCSMShadowDepthPassProcessor(const FScene* Scene, const FSceneView* InViewIfDynamicMeshCommand, FMeshPassDrawListContext* InDrawListContext)
 {
-	FRHIUniformBuffer* PassUniformBuffer = nullptr;
-
-	EShadingPath ShadingPath = Scene->GetShadingPath();
-	if (ShadingPath == EShadingPath::Mobile)
-	{
-		PassUniformBuffer = Scene->UniformBuffers.MobileCSMShadowDepthPassUniformBuffer;
-	}
-	else //deferred
-	{
-		PassUniformBuffer = Scene->UniformBuffers.CSMShadowDepthPassUniformBuffer;
-	}
-
 	return new(FMemStack::Get()) FShadowDepthPassMeshProcessor(
 		Scene,
 		InViewIfDynamicMeshCommand,
 		Scene->UniformBuffers.CSMShadowDepthViewUniformBuffer,
-		PassUniformBuffer,
 		CSMShadowDepthType,
 		InDrawListContext);
 }
diff --git a/Engine/Source/Runtime/Renderer/Private/ShadowRendering.cpp b/Engine/Source/Runtime/Renderer/Private/ShadowRendering.cpp
--- a/Engine/Source/Runtime/Renderer/Private/ShadowRendering.cpp
+++ b/Engine/Source/Runtime/Renderer/Private/ShadowRendering.cpp
@@ -950,23 +950,24 @@ void FProjectedShadowInfo::SetupProjectionStencilMask(
 	}
 }
 
-void FProjectedShadowInfo::RenderProjection(FRHICommandListImmediate& RHICmdList, int32 ViewIndex, const FViewInfo* View, const FSceneRenderer* SceneRender, bool bProjectingForForwardShading, bool bMobileModulatedProjections, const FHairStrandsVisibilityData* HairVisibilityData) const
-{
-#if WANTS_DRAW_MESH_EVENTS
-	FString EventName;
-
-	if (GetEmitDrawEvents())
-	{
-		GetShadowTypeNameForDrawEvent(EventName);
-	}
-	SCOPED_DRAW_EVENTF(RHICmdList, EventShadowProjectionActor, *EventName);
-#endif
-
-	FScopeCycleCounter Scope(bWholeSceneShadow ? GET_STATID(STAT_RenderWholeSceneShadowProjectionsTime) : GET_STATID(STAT_RenderPerObjectShadowProjectionsTime));
-
-	FGraphicsPipelineStateInitializer GraphicsPSOInit;
-	RHICmdList.ApplyCachedRenderTargets(GraphicsPSOInit);
+BEGIN_SHADER_PARAMETER_STRUCT(FShadowProjectionPassParameters, )
+	SHADER_PARAMETER_STRUCT_INCLUDE(FSceneTextureShaderParameters, SceneTextures)
+	RDG_TEXTURE_ACCESS(HairCategorizationTexture, ERHIAccess::SRVGraphics)
+	RDG_TEXTURE_ACCESS(ShadowTexture0, ERHIAccess::SRVGraphics)
+	RDG_TEXTURE_ACCESS(ShadowTexture1, ERHIAccess::SRVGraphics)
+	RENDER_TARGET_BINDING_SLOTS()
+END_SHADER_PARAMETER_STRUCT()
 
+void FProjectedShadowInfo::RenderProjection(
+	FRDGBuilder& GraphBuilder,
+	const FShadowProjectionPassParameters& CommonPassParameters,
+	int32 ViewIndex,
+	const FViewInfo* View,
+	const FSceneRenderer* SceneRender,
+	bool bProjectingForForwardShading,
+	bool bMobileModulatedProjections,
+	const FHairStrandsVisibilityData* HairVisibilityData) const
+{
 	// Find the shadow's view relevance.
 	const FVisibleLightViewInfo& VisibleLightViewInfo = View->VisibleLightInfos[LightSceneInfo->Id];
 	{
@@ -979,137 +980,173 @@ void FProjectedShadowInfo::RenderProjection(FRHICommandListImmediate& RHICmdList
 		}
 	}
 
-	bool bCameraInsideShadowFrustum;
-	TArray<FVector4, TInlineAllocator<8>> FrustumVertices;
-	SetupFrustumForProjection(View, FrustumVertices, bCameraInsideShadowFrustum);
+	FString EventName;
 
-	const bool bSubPixelSupport = HairVisibilityData != nullptr;
-	const bool bStencilTestEnabled = !bSubPixelSupport && GShadowStencilCulling;
-	const bool bDepthBoundsTestEnabled = IsWholeSceneDirectionalShadow() && GSupportsDepthBoundsTest && CVarCSMDepthBoundsTest.GetValueOnRenderThread() != 0 && !bSubPixelSupport;
-	
-	if (!bDepthBoundsTestEnabled && bStencilTestEnabled)
+#if WANTS_DRAW_MESH_EVENTS
+	if (GetEmitDrawEvents())
 	{
-		SetupProjectionStencilMask(RHICmdList, View, ViewIndex, SceneRender, FrustumVertices, bMobileModulatedProjections, bCameraInsideShadowFrustum);
+		GetShadowTypeNameForDrawEvent(EventName);
 	}
+#endif
 
-	// solid rasterization w/ back-face culling.
-	GraphicsPSOInit.RasterizerState = (View->bReverseCulling || IsWholeSceneDirectionalShadow()) ? TStaticRasterizerState<FM_Solid,CM_CCW>::GetRHI() : TStaticRasterizerState<FM_Solid,CM_CW>::GetRHI();
+	// Interpret null render targets as skipping the render pass; this is currently used by mobile.
+	const ERDGPassFlags PassFlags = CommonPassParameters.RenderTargets[0].GetTexture() == nullptr ? ERDGPassFlags::SkipRenderPass : ERDGPassFlags::None;
 
-	GraphicsPSOInit.bDepthBounds = bDepthBoundsTestEnabled;
-	if (bDepthBoundsTestEnabled)
-	{
-		// no depth test or writes
-		GraphicsPSOInit.DepthStencilState = TStaticDepthStencilState<false, CF_Always>::GetRHI();
-	}
-	else if (bStencilTestEnabled)
+	auto* PassParameters = GraphBuilder.AllocParameters<FShadowProjectionPassParameters>();
+	*PassParameters = CommonPassParameters;
+
+	if (RenderTargets.DepthTarget)
 	{
-		if (GStencilOptimization)
-		{
-			// No depth test or writes, zero the stencil
-			// Note: this will disable hi-stencil on many GPUs, but still seems 
-			// to be faster. However, early stencil still works 
-			GraphicsPSOInit.DepthStencilState =
-				TStaticDepthStencilState<
-				false, CF_Always,
-				true, CF_NotEqual, SO_Zero, SO_Zero, SO_Zero,
-				false, CF_Always, SO_Zero, SO_Zero, SO_Zero,
-				0xff, 0xff
-				>::GetRHI();
-		}
-		else
-		{
-			// no depth test or writes, Test stencil for non-zero.
-			GraphicsPSOInit.DepthStencilState = 
-				TStaticDepthStencilState<
-				false, CF_Always,
-				true, CF_NotEqual, SO_Keep, SO_Keep, SO_Keep,
-				false, CF_Always, SO_Keep, SO_Keep, SO_Keep,
-				0xff, 0xff
-				>::GetRHI();
-		}
+		PassParameters->ShadowTexture0 = GraphBuilder.RegisterExternalTexture(RenderTargets.DepthTarget);
 	}
 	else
 	{
-		GraphicsPSOInit.DepthStencilState = TStaticDepthStencilState<false, CF_DepthFartherOrEqual>::GetRHI();
+		PassParameters->ShadowTexture0 = GraphBuilder.RegisterExternalTexture(RenderTargets.ColorTargets[0]);
+		PassParameters->ShadowTexture1 = GraphBuilder.RegisterExternalTexture(RenderTargets.ColorTargets[1]);
 	}
 
-	GraphicsPSOInit.BlendState = GetBlendStateForProjection(bProjectingForForwardShading, bMobileModulatedProjections);
+	GraphBuilder.AddPass(
+		RDG_EVENT_NAME("%s", *EventName),
+		PassParameters,
+		ERDGPassFlags::Raster | PassFlags,
+		[this, SceneRender, View, ViewIndex, HairVisibilityData, bProjectingForForwardShading, bMobileModulatedProjections](FRHICommandListImmediate& RHICmdList)
+	{
+		FScopeCycleCounter Scope(bWholeSceneShadow ? GET_STATID(STAT_RenderWholeSceneShadowProjectionsTime) : GET_STATID(STAT_RenderPerObjectShadowProjectionsTime));
 
-	GraphicsPSOInit.PrimitiveType = IsWholeSceneDirectionalShadow() ? PT_TriangleStrip : PT_TriangleList;
+		FGraphicsPipelineStateInitializer GraphicsPSOInit;
+		RHICmdList.ApplyCachedRenderTargets(GraphicsPSOInit);
 
-	{
-		uint32 LocalQuality = GetShadowQuality();
+		bool bCameraInsideShadowFrustum;
+		TArray<FVector4, TInlineAllocator<8>> FrustumVertices;
+		SetupFrustumForProjection(View, FrustumVertices, bCameraInsideShadowFrustum);
 
-		if (LocalQuality > 1)
+		const bool bSubPixelSupport = HairVisibilityData != nullptr;
+		const bool bStencilTestEnabled = !bSubPixelSupport && GShadowStencilCulling;
+		const bool bDepthBoundsTestEnabled = IsWholeSceneDirectionalShadow() && GSupportsDepthBoundsTest && CVarCSMDepthBoundsTest.GetValueOnRenderThread() != 0 && !bSubPixelSupport;
+
+		if (!bDepthBoundsTestEnabled && bStencilTestEnabled)
+		{
+			SetupProjectionStencilMask(RHICmdList, View, ViewIndex, SceneRender, FrustumVertices, bMobileModulatedProjections, bCameraInsideShadowFrustum);
+		}
+
+		// solid rasterization w/ back-face culling.
+		GraphicsPSOInit.RasterizerState = (View->bReverseCulling || IsWholeSceneDirectionalShadow()) ? TStaticRasterizerState<FM_Solid, CM_CCW>::GetRHI() : TStaticRasterizerState<FM_Solid, CM_CW>::GetRHI();
+
+		GraphicsPSOInit.bDepthBounds = bDepthBoundsTestEnabled;
+		if (bDepthBoundsTestEnabled)
 		{
-			if (IsWholeSceneDirectionalShadow() && CascadeSettings.ShadowSplitIndex > 0)
+			// no depth test or writes
+			GraphicsPSOInit.DepthStencilState = TStaticDepthStencilState<false, CF_Always>::GetRHI();
+		}
+		else if (bStencilTestEnabled)
+		{
+			if (GStencilOptimization)
 			{
-				// adjust kernel size so that the penumbra size of distant splits will better match up with the closer ones
-				const float SizeScale = CascadeSettings.ShadowSplitIndex / FMath::Max(0.001f, CVarCSMSplitPenumbraScale.GetValueOnRenderThread());
+				// No depth test or writes, zero the stencil
+				// Note: this will disable hi-stencil on many GPUs, but still seems 
+				// to be faster. However, early stencil still works 
+				GraphicsPSOInit.DepthStencilState =
+					TStaticDepthStencilState<
+					false, CF_Always,
+					true, CF_NotEqual, SO_Zero, SO_Zero, SO_Zero,
+					false, CF_Always, SO_Zero, SO_Zero, SO_Zero,
+					0xff, 0xff
+					>::GetRHI();
 			}
-			else if (LocalQuality > 2 && !bWholeSceneShadow)
+			else
 			{
-				static auto CVarPreShadowResolutionFactor = IConsoleManager::Get().FindTConsoleVariableDataFloat(TEXT("r.Shadow.PreShadowResolutionFactor"));
-				const int32 TargetResolution = bPreShadow ? FMath::TruncToInt(512 * CVarPreShadowResolutionFactor->GetValueOnRenderThread()) : 512;
+				// no depth test or writes, Test stencil for non-zero.
+				GraphicsPSOInit.DepthStencilState =
+					TStaticDepthStencilState<
+					false, CF_Always,
+					true, CF_NotEqual, SO_Keep, SO_Keep, SO_Keep,
+					false, CF_Always, SO_Keep, SO_Keep, SO_Keep,
+					0xff, 0xff
+					>::GetRHI();
+			}
+		}
+		else
+		{
+			GraphicsPSOInit.DepthStencilState = TStaticDepthStencilState<false, CF_DepthFartherOrEqual>::GetRHI();
+		}
 
-				int32 Reduce = 0;
+		GraphicsPSOInit.BlendState = GetBlendStateForProjection(bProjectingForForwardShading, bMobileModulatedProjections);
 
+		GraphicsPSOInit.PrimitiveType = IsWholeSceneDirectionalShadow() ? PT_TriangleStrip : PT_TriangleList;
+
+		{
+			uint32 LocalQuality = GetShadowQuality();
+
+			if (LocalQuality > 1)
+			{
+				if (IsWholeSceneDirectionalShadow() && CascadeSettings.ShadowSplitIndex > 0)
+				{
+					// adjust kernel size so that the penumbra size of distant splits will better match up with the closer ones
+					const float SizeScale = CascadeSettings.ShadowSplitIndex / FMath::Max(0.001f, CVarCSMSplitPenumbraScale.GetValueOnRenderThread());
+				}
+				else if (LocalQuality > 2 && !bWholeSceneShadow)
 				{
-					int32 Res = ResolutionX;
+					static auto CVarPreShadowResolutionFactor = IConsoleManager::Get().FindTConsoleVariableDataFloat(TEXT("r.Shadow.PreShadowResolutionFactor"));
+					const int32 TargetResolution = bPreShadow ? FMath::TruncToInt(512 * CVarPreShadowResolutionFactor->GetValueOnRenderThread()) : 512;
+
+					int32 Reduce = 0;
 
-					while (Res < TargetResolution)
 					{
-						Res *= 2;
-						++Reduce;
+						int32 Res = ResolutionX;
+
+						while (Res < TargetResolution)
+						{
+							Res *= 2;
+							++Reduce;
+						}
 					}
+
+					// Never drop to quality 1 due to low resolution, aliasing is too bad
+					LocalQuality = FMath::Clamp((int32)LocalQuality - Reduce, 3, 5);
 				}
+			}
+
+			GraphicsPSOInit.BoundShaderState.VertexDeclarationRHI = GetVertexDeclarationFVector4();
+			BindShadowProjectionShaders(LocalQuality, RHICmdList, GraphicsPSOInit, ViewIndex, *View, HairVisibilityData, this, bMobileModulatedProjections);
 
-				// Never drop to quality 1 due to low resolution, aliasing is too bad
-				LocalQuality = FMath::Clamp((int32)LocalQuality - Reduce, 3, 5);
+			if (bDepthBoundsTestEnabled)
+			{
+				SetDepthBoundsTest(RHICmdList, CascadeSettings.SplitNear, CascadeSettings.SplitFar, View->ViewMatrices.GetProjectionMatrix());
 			}
-		}
 
-		GraphicsPSOInit.BoundShaderState.VertexDeclarationRHI = GetVertexDeclarationFVector4();
-		BindShadowProjectionShaders(LocalQuality, RHICmdList, GraphicsPSOInit, ViewIndex, *View, HairVisibilityData, this, bMobileModulatedProjections);
+			RHICmdList.SetStencilRef(0);
+		}
 
-		if (bDepthBoundsTestEnabled)
+		if (IsWholeSceneDirectionalShadow())
 		{
-			SetDepthBoundsTest(RHICmdList, CascadeSettings.SplitNear, CascadeSettings.SplitFar, View->ViewMatrices.GetProjectionMatrix());
+			RHICmdList.SetStreamSource(0, GClearVertexBuffer.VertexBufferRHI, 0);
+			RHICmdList.DrawPrimitive(0, 2, 1);
+		}
+		else
+		{
+			FRHIResourceCreateInfo CreateInfo(TEXT("FProjectedShadowInfoFrustum"));
+			FVertexBufferRHIRef VertexBufferRHI = RHICreateVertexBuffer(sizeof(FVector4) * FrustumVertices.Num(), BUF_Volatile, CreateInfo);
+			void* VoidPtr = RHILockVertexBuffer(VertexBufferRHI, 0, sizeof(FVector4) * FrustumVertices.Num(), RLM_WriteOnly);
+			FPlatformMemory::Memcpy(VoidPtr, FrustumVertices.GetData(), sizeof(FVector4) * FrustumVertices.Num());
+			RHIUnlockVertexBuffer(VertexBufferRHI);
+
+			RHICmdList.SetStreamSource(0, VertexBufferRHI, 0);
+			// Draw the frustum using the projection shader..
+			RHICmdList.DrawIndexedPrimitive(GCubeIndexBuffer.IndexBufferRHI, 0, 0, 8, 0, 12, 1);
+			VertexBufferRHI.SafeRelease();
 		}
 
-		RHICmdList.SetStencilRef(0);
-	}
-
-	if (IsWholeSceneDirectionalShadow())
-	{
-		RHICmdList.SetStreamSource(0, GClearVertexBuffer.VertexBufferRHI, 0);
-		RHICmdList.DrawPrimitive(0, 2, 1);
-	}
-	else
-	{
-		FRHIResourceCreateInfo CreateInfo(TEXT("FProjectedShadowInfoFrustum"));
-		FVertexBufferRHIRef VertexBufferRHI = RHICreateVertexBuffer(sizeof(FVector4) * FrustumVertices.Num(), BUF_Volatile, CreateInfo);
-		void* VoidPtr = RHILockVertexBuffer(VertexBufferRHI, 0, sizeof(FVector4) * FrustumVertices.Num(), RLM_WriteOnly);
-		FPlatformMemory::Memcpy(VoidPtr, FrustumVertices.GetData(), sizeof(FVector4) * FrustumVertices.Num());
-		RHIUnlockVertexBuffer(VertexBufferRHI);
-
-		RHICmdList.SetStreamSource(0, VertexBufferRHI, 0);
-		// Draw the frustum using the projection shader..
-		RHICmdList.DrawIndexedPrimitive(GCubeIndexBuffer.IndexBufferRHI, 0, 0, 8, 0, 12, 1);
-		VertexBufferRHI.SafeRelease();
-	}
-
-	if (!bDepthBoundsTestEnabled && bStencilTestEnabled)
-	{
-		// Clear the stencil buffer to 0.
-		if (!GStencilOptimization)
+		if (!bDepthBoundsTestEnabled && bStencilTestEnabled)
 		{
-			DrawClearQuad(RHICmdList, false, FLinearColor::Transparent, false, 0, true, 0);
+			// Clear the stencil buffer to 0.
+			if (!GStencilOptimization)
+			{
+				DrawClearQuad(RHICmdList, false, FLinearColor::Transparent, false, 0, true, 0);
+			}
 		}
-	}
+	});
 }
 
-
 template <uint32 Quality, bool bUseTransmission, bool bUseSubPixel>
 static void SetPointLightShaderTempl(FRHICommandList& RHICmdList, FGraphicsPipelineStateInitializer& GraphicsPSOInit, int32 ViewIndex, const FViewInfo& View, const FProjectedShadowInfo* ShadowInfo, const FHairStrandsVisibilityData* HairVisibilityData)
 {
@@ -1126,41 +1163,65 @@ static void SetPointLightShaderTempl(FRHICommandList& RHICmdList, FGraphicsPipel
 	PixelShader->SetParameters(RHICmdList, ViewIndex, View, HairVisibilityData, ShadowInfo);
 }
 
-/** Render one pass point light shadow projections. */
-void FProjectedShadowInfo::RenderOnePassPointLightProjection(FRHICommandListImmediate& RHICmdList, int32 ViewIndex, const FViewInfo& View, bool bProjectingForForwardShading, const FHairStrandsVisibilityData* HairVisibilityData) const
+void FProjectedShadowInfo::RenderOnePassPointLightProjection(
+	FRDGBuilder& GraphBuilder,
+	const FShadowProjectionPassParameters& CommonPassParameters,
+	int32 ViewIndex,
+	const FViewInfo& View,
+	bool bProjectingForForwardShading,
+	const FHairStrandsVisibilityData* HairVisibilityData) const
 {
 	SCOPE_CYCLE_COUNTER(STAT_RenderWholeSceneShadowProjectionsTime);
 
 	checkSlow(bOnePassPointLightShadow);
 	
 	const FSphere LightBounds = LightSceneInfo->Proxy->GetBoundingSphere();
+	const bool bUseTransmission = LightSceneInfo->Proxy->Transmission();
+	const bool bCameraInsideLightGeometry = ((FVector)View.ViewMatrices.GetViewOrigin() - LightBounds.Center).SizeSquared() < FMath::Square(LightBounds.W * 1.05f + View.NearClippingDistance * 2.0f);
 
-	bool bUseTransmission = LightSceneInfo->Proxy->Transmission();
-
-	FGraphicsPipelineStateInitializer GraphicsPSOInit;
-	RHICmdList.ApplyCachedRenderTargets(GraphicsPSOInit);
-	GraphicsPSOInit.BlendState = GetBlendStateForProjection(bProjectingForForwardShading, false);
-	GraphicsPSOInit.PrimitiveType = PT_TriangleList;
+	// Interpret null render targets as skipping the render pass; this is currently used by mobile.
+	const ERDGPassFlags PassFlags = CommonPassParameters.RenderTargets[0].GetTexture() == nullptr ? ERDGPassFlags::SkipRenderPass : ERDGPassFlags::None;
 
-	const bool bCameraInsideLightGeometry = ((FVector)View.ViewMatrices.GetViewOrigin() - LightBounds.Center).SizeSquared() < FMath::Square(LightBounds.W * 1.05f + View.NearClippingDistance * 2.0f);
+	auto* PassParameters = GraphBuilder.AllocParameters<FShadowProjectionPassParameters>();
+	*PassParameters = CommonPassParameters;
 
-	if (bCameraInsideLightGeometry)
+	if (RenderTargets.DepthTarget)
 	{
-		GraphicsPSOInit.DepthStencilState = TStaticDepthStencilState<false, CF_Always>::GetRHI();
-		// Render backfaces with depth tests disabled since the camera is inside (or close to inside) the light geometry
-		GraphicsPSOInit.RasterizerState = View.bReverseCulling ? TStaticRasterizerState<FM_Solid, CM_CW>::GetRHI() : TStaticRasterizerState<FM_Solid, CM_CCW>::GetRHI();
+		PassParameters->ShadowTexture0 = GraphBuilder.RegisterExternalTexture(RenderTargets.DepthTarget);
 	}
 	else
 	{
-		// Render frontfaces with depth tests on to get the speedup from HiZ since the camera is outside the light geometry
-		GraphicsPSOInit.DepthStencilState = TStaticDepthStencilState<false, CF_DepthNearOrEqual>::GetRHI();
-		GraphicsPSOInit.RasterizerState = View.bReverseCulling ? TStaticRasterizerState<FM_Solid, CM_CCW>::GetRHI() : TStaticRasterizerState<FM_Solid, CM_CW>::GetRHI();
-	}	
+		PassParameters->ShadowTexture0 = GraphBuilder.RegisterExternalTexture(RenderTargets.ColorTargets[0]);
+		PassParameters->ShadowTexture1 = GraphBuilder.RegisterExternalTexture(RenderTargets.ColorTargets[1]);
+	}
 
+	GraphBuilder.AddPass(
+		{},
+		PassParameters,
+		ERDGPassFlags::Raster | PassFlags,
+		[this, &View, HairVisibilityData, bProjectingForForwardShading, bCameraInsideLightGeometry, bUseTransmission, ViewIndex](FRHICommandList& RHICmdList)
 	{
-		uint32 LocalQuality = GetShadowQuality();
+		FGraphicsPipelineStateInitializer GraphicsPSOInit;
+		RHICmdList.ApplyCachedRenderTargets(GraphicsPSOInit);
+		GraphicsPSOInit.BlendState = GetBlendStateForProjection(bProjectingForForwardShading, false);
+		GraphicsPSOInit.PrimitiveType = PT_TriangleList;
+
+		if (bCameraInsideLightGeometry)
+		{
+			GraphicsPSOInit.DepthStencilState = TStaticDepthStencilState<false, CF_Always>::GetRHI();
+			// Render backfaces with depth tests disabled since the camera is inside (or close to inside) the light geometry
+			GraphicsPSOInit.RasterizerState = View.bReverseCulling ? TStaticRasterizerState<FM_Solid, CM_CW>::GetRHI() : TStaticRasterizerState<FM_Solid, CM_CCW>::GetRHI();
+		}
+		else
+		{
+			// Render frontfaces with depth tests on to get the speedup from HiZ since the camera is outside the light geometry
+			GraphicsPSOInit.DepthStencilState = TStaticDepthStencilState<false, CF_DepthNearOrEqual>::GetRHI();
+			GraphicsPSOInit.RasterizerState = View.bReverseCulling ? TStaticRasterizerState<FM_Solid, CM_CCW>::GetRHI() : TStaticRasterizerState<FM_Solid, CM_CW>::GetRHI();
+		}
+
+		const uint32 LocalQuality = GetShadowQuality();
 
-		if(LocalQuality > 1)
+		if (LocalQuality > 1)
 		{
 			// adjust kernel size so that the penumbra size of distant splits will better match up with the closer ones
 			//const float SizeScale = ShadowInfo->ResolutionX;
@@ -1169,7 +1230,7 @@ void FProjectedShadowInfo::RenderOnePassPointLightProjection(FRHICommandListImme
 			{
 				int32 Res = ResolutionX;
 
-				while(Res < 512)
+				while (Res < 512)
 				{
 					Res *= 2;
 					++Reduce;
@@ -1182,46 +1243,46 @@ void FProjectedShadowInfo::RenderOnePassPointLightProjection(FRHICommandListImme
 		{
 			switch (LocalQuality)
 			{
-				case 1: SetPointLightShaderTempl<1, false, true>(RHICmdList, GraphicsPSOInit, ViewIndex, View, this, HairVisibilityData); break;
-				case 2: SetPointLightShaderTempl<2, false, true>(RHICmdList, GraphicsPSOInit, ViewIndex, View, this, HairVisibilityData); break;
-				case 3: SetPointLightShaderTempl<3, false, true>(RHICmdList, GraphicsPSOInit, ViewIndex, View, this, HairVisibilityData); break;
-				case 4: SetPointLightShaderTempl<4, false, true>(RHICmdList, GraphicsPSOInit, ViewIndex, View, this, HairVisibilityData); break;
-				case 5: SetPointLightShaderTempl<5, false, true>(RHICmdList, GraphicsPSOInit, ViewIndex, View, this, HairVisibilityData); break;
-				default:
-					check(0);
+			case 1: SetPointLightShaderTempl<1, false, true>(RHICmdList, GraphicsPSOInit, ViewIndex, View, this, HairVisibilityData); break;
+			case 2: SetPointLightShaderTempl<2, false, true>(RHICmdList, GraphicsPSOInit, ViewIndex, View, this, HairVisibilityData); break;
+			case 3: SetPointLightShaderTempl<3, false, true>(RHICmdList, GraphicsPSOInit, ViewIndex, View, this, HairVisibilityData); break;
+			case 4: SetPointLightShaderTempl<4, false, true>(RHICmdList, GraphicsPSOInit, ViewIndex, View, this, HairVisibilityData); break;
+			case 5: SetPointLightShaderTempl<5, false, true>(RHICmdList, GraphicsPSOInit, ViewIndex, View, this, HairVisibilityData); break;
+			default:
+				check(0);
 			}
 		}
 		else if (bUseTransmission)
 		{
 			switch (LocalQuality)
 			{
-				case 1: SetPointLightShaderTempl<1, true, false>(RHICmdList, GraphicsPSOInit, ViewIndex, View, this, nullptr); break;
-				case 2: SetPointLightShaderTempl<2, true, false>(RHICmdList, GraphicsPSOInit, ViewIndex, View, this, nullptr); break;
-				case 3: SetPointLightShaderTempl<3, true, false>(RHICmdList, GraphicsPSOInit, ViewIndex, View, this, nullptr); break;
-				case 4: SetPointLightShaderTempl<4, true, false>(RHICmdList, GraphicsPSOInit, ViewIndex, View, this, nullptr); break;
-				case 5: SetPointLightShaderTempl<5, true, false>(RHICmdList, GraphicsPSOInit, ViewIndex, View, this, nullptr); break;
-				default:
-					check(0);
+			case 1: SetPointLightShaderTempl<1, true, false>(RHICmdList, GraphicsPSOInit, ViewIndex, View, this, nullptr); break;
+			case 2: SetPointLightShaderTempl<2, true, false>(RHICmdList, GraphicsPSOInit, ViewIndex, View, this, nullptr); break;
+			case 3: SetPointLightShaderTempl<3, true, false>(RHICmdList, GraphicsPSOInit, ViewIndex, View, this, nullptr); break;
+			case 4: SetPointLightShaderTempl<4, true, false>(RHICmdList, GraphicsPSOInit, ViewIndex, View, this, nullptr); break;
+			case 5: SetPointLightShaderTempl<5, true, false>(RHICmdList, GraphicsPSOInit, ViewIndex, View, this, nullptr); break;
+			default:
+				check(0);
 			}
 		}
 		else
 		{
 			switch (LocalQuality)
 			{
-				case 1: SetPointLightShaderTempl<1, false, false>(RHICmdList, GraphicsPSOInit, ViewIndex, View, this, nullptr); break;
-				case 2: SetPointLightShaderTempl<2, false, false>(RHICmdList, GraphicsPSOInit, ViewIndex, View, this, nullptr); break;
-				case 3: SetPointLightShaderTempl<3, false, false>(RHICmdList, GraphicsPSOInit, ViewIndex, View, this, nullptr); break;
-				case 4: SetPointLightShaderTempl<4, false, false>(RHICmdList, GraphicsPSOInit, ViewIndex, View, this, nullptr); break;
-				case 5: SetPointLightShaderTempl<5, false, false>(RHICmdList, GraphicsPSOInit, ViewIndex, View, this, nullptr); break;
-				default:
-					check(0);
+			case 1: SetPointLightShaderTempl<1, false, false>(RHICmdList, GraphicsPSOInit, ViewIndex, View, this, nullptr); break;
+			case 2: SetPointLightShaderTempl<2, false, false>(RHICmdList, GraphicsPSOInit, ViewIndex, View, this, nullptr); break;
+			case 3: SetPointLightShaderTempl<3, false, false>(RHICmdList, GraphicsPSOInit, ViewIndex, View, this, nullptr); break;
+			case 4: SetPointLightShaderTempl<4, false, false>(RHICmdList, GraphicsPSOInit, ViewIndex, View, this, nullptr); break;
+			case 5: SetPointLightShaderTempl<5, false, false>(RHICmdList, GraphicsPSOInit, ViewIndex, View, this, nullptr); break;
+			default:
+				check(0);
 			}
 		}
-	}
 
-	// Project the point light shadow with some approximately bounding geometry, 
-	// So we can get speedups from depth testing and not processing pixels outside of the light's influence.
-	StencilingGeometry::DrawSphere(RHICmdList);
+		// Project the point light shadow with some approximately bounding geometry, 
+		// So we can get speedups from depth testing and not processing pixels outside of the light's influence.
+		StencilingGeometry::DrawSphere(RHICmdList);
+	});
 }
 
 void FProjectedShadowInfo::RenderFrustumWireframe(FPrimitiveDrawInterface* PDI) const
@@ -1583,7 +1644,8 @@ bool FSceneRenderer::CheckForProjectedShadows( const FLightSceneInfo* LightScene
 }
 
 void FSceneRenderer::RenderShadowProjections(
-	FRHICommandListImmediate& RHICmdList,
+	FRDGBuilder& GraphBuilder,
+	const FShadowProjectionPassParameters& CommonPassParameters,
 	const FLightSceneProxy* LightSceneProxy,
 	const FHairStrandsVisibilityViews* HairVisibilityViews,
 	TArrayView<const FProjectedShadowInfo* const> Shadows,
@@ -1598,9 +1660,15 @@ void FSceneRenderer::RenderShadowProjections(
 	{
 		const FViewInfo& View = Views[ViewIndex];
 
-		SCOPED_GPU_MASK(RHICmdList, View.GPUMask);
-		SCOPED_CONDITIONAL_DRAW_EVENTF(RHICmdList, EventView, Views.Num() > 1, TEXT("View%d"), ViewIndex);
-		UniformBuffers.UpdateViewUniformBuffer(View);
+		RDG_GPU_MASK_SCOPE(GraphBuilder, View.GPUMask);
+		RDG_EVENT_SCOPE_CONDITIONAL(GraphBuilder, Views.Num() > 1, "View%d", ViewIndex);
+
+		AddPass(GraphBuilder, [&UniformBuffers, &View, LightSceneProxy](FRHICommandList& RHICmdList)
+		{
+			UniformBuffers.UpdateViewUniformBuffer(View);
+			RHICmdList.SetViewport(View.ViewRect.Min.X, View.ViewRect.Min.Y, 0.0f, View.ViewRect.Max.X, View.ViewRect.Max.Y, 1.0f);
+			LightSceneProxy->SetScissorRect(RHICmdList, View, View.ViewRect);
+		});
 
 		const FHairStrandsVisibilityData* HairVisibilityData = nullptr;
 		if (HairVisibilityViews)
@@ -1608,9 +1676,6 @@ void FSceneRenderer::RenderShadowProjections(
 			HairVisibilityData = &(HairVisibilityViews->HairDatas[ViewIndex]);
 		}
 
-		RHICmdList.SetViewport(View.ViewRect.Min.X, View.ViewRect.Min.Y, 0.0f, View.ViewRect.Max.X, View.ViewRect.Max.Y, 1.0f);
-		LightSceneProxy->SetScissorRect(RHICmdList, View, View.ViewRect);
-
 		// Project the shadow depth buffers onto the scene.
 		for (const FProjectedShadowInfo* ProjectedShadowInfo : Shadows)
 		{
@@ -1621,53 +1686,23 @@ void FSceneRenderer::RenderShadowProjections(
 				{
 					if (ProjectedShadowInfo->bOnePassPointLightShadow)
 					{
-						ProjectedShadowInfo->RenderOnePassPointLightProjection(RHICmdList, ViewIndex, View, bProjectingForForwardShading, HairVisibilityData);
+						ProjectedShadowInfo->RenderOnePassPointLightProjection(GraphBuilder, CommonPassParameters, ViewIndex, View, bProjectingForForwardShading, HairVisibilityData);
 					}
 					else
 					{
-						ProjectedShadowInfo->RenderProjection(RHICmdList, ViewIndex, &View, this, bProjectingForForwardShading, bMobileModulatedProjections, HairVisibilityData);
+						ProjectedShadowInfo->RenderProjection(GraphBuilder, CommonPassParameters, ViewIndex, &View, this, bProjectingForForwardShading, bMobileModulatedProjections, HairVisibilityData);
 					}
 				}
 			}
 		}
 	}
 
-	RHICmdList.SetScissorRect(false, 0, 0, 0, 0);
-}
-
-// TODO(RDG): This is a temporary solution while the shadow depth rendering code still isn't ported to RDG.
-static void TransitionShadowsToReadable(FRHICommandList& RHICmdList, TArrayView<const FProjectedShadowInfo* const> Shadows)
-{
-	TSet<IPooledRenderTarget*, DefaultKeyFuncs<IPooledRenderTarget*>, SceneRenderingSetAllocator> FoundTextures;
-	FoundTextures.Reserve(Shadows.Num());
-
-	TArray<FRHITransitionInfo, SceneRenderingAllocator> TexturesToTransition;
-	TexturesToTransition.Reserve(Shadows.Num());
-
-	for (const FProjectedShadowInfo* ProjectedShadowInfo : Shadows)
+	AddPass(GraphBuilder, [](FRHICommandList& RHICmdList)
 	{
-		IPooledRenderTarget* DepthTarget = ProjectedShadowInfo->RenderTargets.DepthTarget;
-
-		if (ProjectedShadowInfo->bAllocated && DepthTarget)
-		{
-			bool bAlreadyInSet = false;
-			FoundTextures.Emplace(DepthTarget, &bAlreadyInSet);
-			if (!bAlreadyInSet)
-			{
-				TexturesToTransition.Add(FRHITransitionInfo(DepthTarget->GetShaderResourceRHI(), ERHIAccess::Unknown, ERHIAccess::SRVMask));
-			}
-		}
-	}
-
-	RHICmdList.Transition(TexturesToTransition);
+		RHICmdList.SetScissorRect(false, 0, 0, 0, 0);
+	});
 }
 
-BEGIN_SHADER_PARAMETER_STRUCT(FRenderShadowProjectionsParameters, )
-	SHADER_PARAMETER_RDG_UNIFORM_BUFFER(FSceneTextureUniformParameters, SceneTextures)
-	SHADER_PARAMETER_RDG_TEXTURE(Texture2D, HairCategorizationTexture)
-	RENDER_TARGET_BINDING_SLOTS()
-END_SHADER_PARAMETER_STRUCT()
-
 void FDeferredShadingSceneRenderer::RenderShadowProjections(
 	FRDGBuilder& GraphBuilder,
 	TRDGUniformBufferRef<FSceneTextureUniformParameters> SceneTexturesUniformBuffer,
@@ -1711,34 +1746,27 @@ void FDeferredShadingSceneRenderer::RenderShadowProjections(
 
 	if (NormalShadows.Num() > 0)
 	{
-		AddPass(GraphBuilder, [&NormalShadows](FRHICommandList& RHICmdList)
-		{
-			TransitionShadowsToReadable(RHICmdList, NormalShadows);
-		});
-
 		const auto RenderNormalShadows = [&](FRDGTextureRef OutputTexture, FExclusiveDepthStencil ExclusiveDepthStencil, bool bSubPixel)
 		{
-			auto* PassParameters = GraphBuilder.AllocParameters<FRenderShadowProjectionsParameters>();
-			PassParameters->SceneTextures = SceneTexturesUniformBuffer;
-			PassParameters->HairCategorizationTexture = bSubPixel && HairVisibilityViews->HairDatas.Num() > 0 ? HairVisibilityViews->HairDatas[0].CategorizationTexture : nullptr;
-			PassParameters->RenderTargets[0] = FRenderTargetBinding(OutputTexture, ERenderTargetLoadAction::ELoad);
-			PassParameters->RenderTargets.DepthStencil = FDepthStencilBinding(SceneDepthTexture, ERenderTargetLoadAction::ELoad, ERenderTargetLoadAction::ELoad, ExclusiveDepthStencil);
-
 			FString LightNameWithLevel;
 			GetLightNameForDrawEvent(LightSceneProxy, LightNameWithLevel);
-
-			// All shadows projections are rendered in one RDG pass for efficiency purposes. Technically, RDG is able to merge all these
-			// render passes together if we used a separate one per shadow, but we are paying a cost for it which just seems unnecessary
-			// here. We are also able to bulk-transition all of the shadows in one go, which RDG is currently not able to do.
-			GraphBuilder.AddPass(
-				RDG_EVENT_NAME("%s", *LightNameWithLevel),
-				PassParameters,
-				ERDGPassFlags::Raster,
-				[this, &NormalShadows, LightSceneProxy, HairVisibilityViews, bProjectingForForwardShading, bSubPixel](FRHICommandListImmediate& RHICmdList)
-			{
-				const bool bMobileModulatedProjections = false;
-				FSceneRenderer::RenderShadowProjections(RHICmdList, LightSceneProxy, bSubPixel ? HairVisibilityViews : nullptr, NormalShadows, bProjectingForForwardShading, bMobileModulatedProjections);
-			});
+			RDG_EVENT_SCOPE(GraphBuilder, "%s", *LightNameWithLevel);
+
+			FShadowProjectionPassParameters CommonPassParameters;
+			CommonPassParameters.SceneTextures = GetSceneTextureShaderParameters(SceneTexturesUniformBuffer);
+			CommonPassParameters.HairCategorizationTexture = bSubPixel && HairVisibilityViews->HairDatas.Num() > 0 ? HairVisibilityViews->HairDatas[0].CategorizationTexture : nullptr;
+			CommonPassParameters.RenderTargets[0] = FRenderTargetBinding(OutputTexture, ERenderTargetLoadAction::ELoad);
+			CommonPassParameters.RenderTargets.DepthStencil = FDepthStencilBinding(SceneDepthTexture, ERenderTargetLoadAction::ELoad, ERenderTargetLoadAction::ELoad, ExclusiveDepthStencil);
+
+			const bool bMobileModulatedProjections = false;
+			FSceneRenderer::RenderShadowProjections(
+				GraphBuilder,
+				CommonPassParameters,
+				LightSceneProxy,
+				bSubPixel ? HairVisibilityViews : nullptr,
+				NormalShadows,
+				bProjectingForForwardShading,
+				bMobileModulatedProjections);
 		};
 
 		{
@@ -1932,7 +1960,6 @@ void FDeferredShadingSceneRenderer::RenderDeferredShadowProjections(
 	const bool bProjectingForForwardShading = false;
 	RenderShadowProjections(GraphBuilder, SceneTexturesUniformBuffer, ScreenShadowMaskTexture, ScreenShadowMaskSubPixelTexture, SceneDepthTexture, LightSceneInfo, HairVisibilityViews, bProjectingForForwardShading);
 
-
 	// Perform injection on tranlucent lighting volume
 	{
 		const TArray<FProjectedShadowInfo*, SceneRenderingAllocator>& ShadowMaps = VisibleLightInfo.CompleteProjectedShadows.Num() > 0
@@ -2021,9 +2048,10 @@ void FMobileSceneRenderer::RenderModulatedShadowProjections(FRHICommandListImmed
 	SCOPED_DRAW_EVENT(RHICmdList, ShadowProjectionOnOpaque);
 	SCOPED_GPU_STAT(RHICmdList, ShadowProjection);
 
-	FSceneRenderTargets& SceneContext = FSceneRenderTargets::Get(RHICmdList);
-	TRefCountPtr<FRHIUniformBuffer> SceneTexturesUniformBuffer = CreateMobileSceneTextureUniformBuffer(RHICmdList, EMobileSceneTextureSetupMode::SceneColor);
-	SCOPED_UNIFORM_BUFFER_GLOBAL_BINDINGS(RHICmdList, SceneTexturesUniformBuffer);
+	FRDGBuilder GraphBuilder(RHICmdList);
+
+	FShadowProjectionPassParameters CommonPassParameters;
+	CommonPassParameters.SceneTextures.MobileSceneTextures = CreateMobileSceneTextureUniformBuffer(GraphBuilder, EMobileSceneTextureSetupMode::SceneColor);
 
 	// render shadowmaps for relevant lights.
 	for (TSparseArray<FLightSceneInfoCompact>::TConstIterator LightIt(Scene->Lights); LightIt; ++LightIt)
@@ -2034,19 +2062,25 @@ void FMobileSceneRenderer::RenderModulatedShadowProjections(FRHICommandListImmed
 
 		if(LightSceneInfo->ShouldRenderLightViewIndependent() && LightSceneProxy && LightSceneProxy->CastsModulatedShadows())
 		{
-
 			const FVisibleLightInfo& VisibleLightInfo = VisibleLightInfos[LightSceneInfo->Id];
 
 			if (VisibleLightInfo.ShadowsToProject.Num() > 0)
 			{
-				TransitionShadowsToReadable(RHICmdList, VisibleLightInfo.ShadowsToProject);
-
 				const bool bProjectingForForwardShading = false;
 				const bool bMobileModulatedProjections = true;
-				RenderShadowProjections(RHICmdList, LightSceneProxy, nullptr, VisibleLightInfo.ShadowsToProject, bProjectingForForwardShading, bMobileModulatedProjections);
+				RenderShadowProjections(
+					GraphBuilder,
+					CommonPassParameters,
+					LightSceneProxy,
+					nullptr,
+					VisibleLightInfo.ShadowsToProject,
+					bProjectingForForwardShading,
+					bMobileModulatedProjections);
 			}
 		}
 	}
+
+	GraphBuilder.Execute();
 }
 
 IMPLEMENT_GLOBAL_SHADER_PARAMETER_STRUCT(FTranslucentSelfShadowUniformParameters, "TranslucentSelfShadow");
diff --git a/Engine/Source/Runtime/Renderer/Private/ShadowRendering.h b/Engine/Source/Runtime/Renderer/Private/ShadowRendering.h
--- a/Engine/Source/Runtime/Renderer/Private/ShadowRendering.h
+++ b/Engine/Source/Runtime/Renderer/Private/ShadowRendering.h
@@ -101,7 +101,6 @@ public:
 		const FScene* Scene, 
 		const FSceneView* InViewIfDynamicMeshCommand, 
 		const TUniformBufferRef<FViewUniformShaderParameters>& InViewUniformBuffer,
-		FRHIUniformBuffer* InPassUniformBuffer,
 		FShadowDepthType InShadowDepthType,
 		FMeshPassDrawListContext* InDrawListContext);
 
@@ -236,9 +235,6 @@ public:
 	/** The view to be used when rendering this shadow's depths. */
 	FViewInfo* ShadowDepthView;
 
-	TUniformBufferRef<FShadowDepthPassUniformParameters> ShadowDepthPassUniformBuffer;
-	TUniformBufferRef<FMobileShadowDepthPassUniformParameters> MobileShadowDepthPassUniformBuffer;
-
 	/** The depth or color targets this shadow was rendered to. */
 	FShadowMapRenderTargets RenderTargets;
 
@@ -423,11 +419,6 @@ public:
 	float GetShaderMaxSlopeDepthBias() const { return ShaderMaxSlopeDepthBias; }
 	float GetShaderReceiverDepthBias() const;
 
-	/**
-	 * Renders the shadow subject depth.
-	 */
-	void RenderDepth(FRHICommandListImmediate& RHICmdList, class FSceneRenderer* SceneRenderer, FBeginShadowRenderPassFunction BeginShadowRenderPass, bool bDoParallelDispatch);
-
 	void SetStateForView(FRHICommandList& RHICmdList) const;
 
 	FIntRect GetViewRectForView() const
@@ -443,10 +434,10 @@ public:
 	/** Set state for depth rendering */
 	void SetStateForDepth(FMeshPassProcessorRenderState& DrawRenderState) const;
 
-	void ClearDepth(FRHICommandList& RHICmdList, class FSceneRenderer* SceneRenderer, int32 NumColorTextures, FRHITexture** ColorTextures, FRHITexture* DepthTexture, bool bPerformClear);
+	void ClearDepth(FRHICommandList& RHICmdList) const;
 
 	/** Renders shadow maps for translucent primitives. */
-	void RenderTranslucencyDepths(FRHICommandList& RHICmdList, class FSceneRenderer* SceneRenderer);
+	void RenderTranslucencyDepths(FRDGBuilder& GraphBuilder, class FSceneRenderer* SceneRenderer, const FRenderTargetBindingSlots& RenderTargets);
 
 	static FRHIBlendState* GetBlendStateForProjection(
 		int32 ShadowMapChannel,
@@ -460,7 +451,24 @@ public:
 	/**
 	 * Projects the shadow onto the scene for a particular view.
 	 */
-	void RenderProjection(FRHICommandListImmediate& RHICmdList, int32 ViewIndex, const class FViewInfo* View, const class FSceneRenderer* SceneRender, bool bProjectingForForwardShading, bool bMobile, const struct FHairStrandsVisibilityData* HairVisibilityData) const;
+	void RenderProjection(
+		FRDGBuilder& GraphBuilder,
+		const FShadowProjectionPassParameters& CommonPassParameters,
+		int32 ViewIndex,
+		const FViewInfo* View,
+		const FSceneRenderer* SceneRender,
+		bool bProjectingForForwardShading,
+		bool bMobileModulatedProjections,
+		const FHairStrandsVisibilityData* HairVisibilityData) const;
+
+	/**
+	* Renders the shadow subject depth, to a particular hacked view
+	*/
+	void RenderDepth(
+		FRDGBuilder& GraphBuilder,
+		const FSceneRenderer* SceneRenderer,
+		FRDGTextureRef ShadowDepthTexture,
+		bool bDoParallelDispatch);
 
 	FRDGTextureRef BeginRenderRayTracedDistanceFieldProjection(
 		FRDGBuilder& GraphBuilder,
@@ -478,7 +486,13 @@ public:
 		bool bProjectingForForwardShading) const;
 
 	/** Render one pass point light shadow projections. */
-	void RenderOnePassPointLightProjection(FRHICommandListImmediate& RHICmdList, int32 ViewIndex, const FViewInfo& View, bool bProjectingForForwardShading, const FHairStrandsVisibilityData* HairVisibilityData) const;
+	void RenderOnePassPointLightProjection(
+		FRDGBuilder& GraphBuilder,
+		const FShadowProjectionPassParameters& CommonPassParameters,
+		int32 ViewIndex,
+		const FViewInfo& View,
+		bool bProjectingForForwardShading,
+		const FHairStrandsVisibilityData* HairVisibilityData) const;
 
 	/**
 	 * Renders the projected shadow's frustum wireframe with the given FPrimitiveDrawInterface.
@@ -519,7 +533,7 @@ public:
 	void GatherDynamicMeshElements(FSceneRenderer& Renderer, class FVisibleLightInfo& VisibleLightInfo, TArray<const FSceneView*>& ReusedViewsArray, 
 		FGlobalDynamicIndexBuffer& DynamicIndexBuffer, FGlobalDynamicVertexBuffer& DynamicVertexBuffer, FGlobalDynamicReadBuffer& DynamicReadBuffer);
 
-	void SetupMeshDrawCommandsForShadowDepth(FSceneRenderer& Renderer, FRHIUniformBuffer* PassUniformBuffer);
+	void SetupMeshDrawCommandsForShadowDepth(FSceneRenderer& Renderer);
 
 	void SetupMeshDrawCommandsForProjectionStenciling(FSceneRenderer& Renderer);
 
@@ -603,16 +617,6 @@ public:
 		return FShadowDepthType(bDirectionalLight, bOnePassPointLightShadow);
 	}
 
-	/**
-	* Setup uniformbuffers and update Primitive Shader Data
-	*/
-	void SetupShadowUniformBuffers(FRHICommandListImmediate& RHICmdList, FScene* Scene);
-
-	/**
-	* Ensure Cached Shadowmap is in readable state
-	*/
-	void TransitionCachedShadowmap(FRHICommandListImmediate& RHICmdList, FScene* Scene);
-
 	bool HasVirtualShadowMap() const { return VirtualShadowMap != nullptr; }
 
 private:
@@ -663,18 +667,20 @@ private:
 	float ShaderSlopeDepthBias;
 	float ShaderMaxSlopeDepthBias;
 
-	void CopyCachedShadowMap(FRHICommandList& RHICmdList, const FMeshPassProcessorRenderState& DrawRenderState, FSceneRenderer* SceneRenderer, const FViewInfo& View);
-
-	/**
-	* Renders the shadow subject depth, to a particular hacked view
-	*/
-	void RenderDepthInner(FRHICommandListImmediate& RHICmdList, class FSceneRenderer* SceneRenderer, FBeginShadowRenderPassFunction BeginShadowRenderPass, bool bDoParallelDispatch);
+	void CopyCachedShadowMap(
+		FRDGBuilder& GraphBuilder,
+		const FViewInfo& View,
+		const FSceneRenderer* SceneRenderer,
+		const FRenderTargetBindingSlots& RenderTargets,
+		const FMeshPassProcessorRenderState& DrawRenderState);
 
 	/**
 	* Modifies the passed in view for this shadow
 	*/
 	void ModifyViewForShadow(FRHICommandList& RHICmdList, FViewInfo* FoundView) const;
 
+	void BeginRenderView(FRDGBuilder& GraphBuilder, FScene* Scene);
+
 	/**
 	* Finds a relevant view for a shadow
 	*/
diff --git a/Engine/Source/Runtime/Renderer/Private/ShadowSetup.cpp b/Engine/Source/Runtime/Renderer/Private/ShadowSetup.cpp
--- a/Engine/Source/Runtime/Renderer/Private/ShadowSetup.cpp
+++ b/Engine/Source/Runtime/Renderer/Private/ShadowSetup.cpp
@@ -1845,7 +1845,7 @@ void FProjectedShadowInfo::AddReceiverPrimitive(FPrimitiveSceneInfo* PrimitiveSc
 	ReceiverPrimitives.Add(PrimitiveSceneInfo);
 }
 
-void FProjectedShadowInfo::SetupMeshDrawCommandsForShadowDepth(FSceneRenderer& Renderer, FRHIUniformBuffer* PassUniformBuffer)
+void FProjectedShadowInfo::SetupMeshDrawCommandsForShadowDepth(FSceneRenderer& Renderer)
 {
 	QUICK_SCOPE_CYCLE_COUNTER(STAT_SetupMeshDrawCommandsForShadowDepth);
 
@@ -1853,7 +1853,6 @@ void FProjectedShadowInfo::SetupMeshDrawCommandsForShadowDepth(FSceneRenderer& R
 		Renderer.Scene,
 		ShadowDepthView,
 		ShadowDepthView->ViewUniformBuffer,
-		PassUniformBuffer,
 		GetShadowDepthType(),
 		nullptr);
 
@@ -2056,24 +2055,7 @@ void FProjectedShadowInfo::GatherDynamicMeshElements(FSceneRenderer& Renderer, F
 		Renderer.MeshCollector.ProcessTasks();
     }
 
-	
-	// Create a pass uniform buffer so we can build mesh commands now in InitDynamicShadows.  This will be updated with the correct contents just before the actual pass.
-	const EShadingPath ShadingPath = FSceneInterface::GetShadingPath(Renderer.FeatureLevel);
-	FRHIUniformBuffer* PassUniformBuffer = nullptr;
-	if (ShadingPath == EShadingPath::Deferred)
-	{
-		FShadowDepthPassUniformParameters ShadowDepthParameters;
-		ShadowDepthPassUniformBuffer = TUniformBufferRef<FShadowDepthPassUniformParameters>::CreateUniformBufferImmediate(ShadowDepthParameters, UniformBuffer_MultiFrame, EUniformBufferValidation::None);
-		PassUniformBuffer = ShadowDepthPassUniformBuffer;
-	}
-	else if (ShadingPath == EShadingPath::Mobile)
-	{
-		FMobileShadowDepthPassUniformParameters ShadowDepthParameters;
-		MobileShadowDepthPassUniformBuffer = TUniformBufferRef<FMobileShadowDepthPassUniformParameters>::CreateUniformBufferImmediate(ShadowDepthParameters, UniformBuffer_MultiFrame, EUniformBufferValidation::None);
-		PassUniformBuffer = MobileShadowDepthPassUniformBuffer;
-	}
-
-	SetupMeshDrawCommandsForShadowDepth(Renderer, PassUniformBuffer);
+	SetupMeshDrawCommandsForShadowDepth(Renderer);
 	SetupMeshDrawCommandsForProjectionStenciling(Renderer);
 }
 
diff --git a/Engine/Source/Runtime/Renderer/Private/TranslucentLighting.cpp b/Engine/Source/Runtime/Renderer/Private/TranslucentLighting.cpp
--- a/Engine/Source/Runtime/Renderer/Private/TranslucentLighting.cpp
+++ b/Engine/Source/Runtime/Renderer/Private/TranslucentLighting.cpp
@@ -186,68 +186,9 @@ void FViewInfo::CalcTranslucencyLightingVolumeBounds(FBox* InOutCascadeBoundsArr
 	}
 }
 
-/** Shader parameters for rendering the depth of a mesh for shadowing. */
-class FShadowDepthShaderParameters
-{
-	DECLARE_INLINE_TYPE_LAYOUT(FShadowDepthShaderParameters, NonVirtual);
-public:
-
-	void Bind(const FShaderParameterMap& ParameterMap)
-	{
-		ProjectionMatrix.Bind(ParameterMap,TEXT("ProjectionMatrix"));
-		ShadowParams.Bind(ParameterMap,TEXT("ShadowParams"));
-		ClampToNearPlane.Bind(ParameterMap,TEXT("bClampToNearPlane"));
-	}
-
-	template<typename ShaderRHIParamRef>
-	void Set(FRHICommandList& RHICmdList, ShaderRHIParamRef ShaderRHI, const FSceneView& View, const FProjectedShadowInfo* ShadowInfo, const FMaterialRenderProxy* MaterialRenderProxy)
-	{
-		SetShaderValue(
-			RHICmdList, 
-			ShaderRHI,
-			ProjectionMatrix,
-			FTranslationMatrix(ShadowInfo->PreShadowTranslation - View.ViewMatrices.GetPreViewTranslation()) * ShadowInfo->SubjectAndReceiverMatrix
-			);
-
-		SetShaderValue(RHICmdList, ShaderRHI, ShadowParams, FVector4(ShadowInfo->GetShaderDepthBias(), ShadowInfo->GetShaderSlopeDepthBias(), ShadowInfo->GetShaderMaxSlopeDepthBias(), ShadowInfo->InvMaxSubjectDepth));
-		// Only clamp vertices to the near plane when rendering whole scene directional light shadow depths or preshadows from directional lights
-		const bool bClampToNearPlaneValue = ShadowInfo->IsWholeSceneDirectionalShadow() || (ShadowInfo->bPreShadow && ShadowInfo->bDirectionalLight);
-		SetShaderValue(RHICmdList, ShaderRHI,ClampToNearPlane,bClampToNearPlaneValue ? 1.0f : 0.0f);
-	}
-
-	/** Set the vertex shader parameter values. */
-	void SetVertexShader(FRHICommandList& RHICmdList, FRHIVertexShader* ShaderRHI, const FSceneView& View, const FProjectedShadowInfo* ShadowInfo, const FMaterialRenderProxy* MaterialRenderProxy)
-	{
-		Set(RHICmdList, ShaderRHI, View, ShadowInfo, MaterialRenderProxy);
-	}
-
-	/** Set the domain shader parameter values. */
-	void SetDomainShader(FRHICommandList& RHICmdList, FRHIDomainShader* ShaderRHI, const FSceneView& View, const FProjectedShadowInfo* ShadowInfo, const FMaterialRenderProxy* MaterialRenderProxy)
-	{
-		Set(RHICmdList, ShaderRHI, View, ShadowInfo, MaterialRenderProxy);
-	}
-
-	/** Serializer. */
-	friend FArchive& operator<<(FArchive& Ar,FShadowDepthShaderParameters& P)
-	{
-		Ar << P.ProjectionMatrix;
-		Ar << P.ShadowParams;
-		Ar << P.ClampToNearPlane;
-		return Ar;
-	}
-
-private:
-	
-		LAYOUT_FIELD(FShaderParameter, ProjectionMatrix)
-		LAYOUT_FIELD(FShaderParameter, ShadowParams)
-		LAYOUT_FIELD(FShaderParameter, ClampToNearPlane)
-	
-};
-
 class FTranslucencyDepthShaderElementData : public FMeshMaterialShaderElementData
 {
 public:
-
 	float TranslucentShadowStartOffset;
 };
 
@@ -259,19 +200,18 @@ BEGIN_GLOBAL_SHADER_PARAMETER_STRUCT(FTranslucencyDepthPassUniformParameters,)
 	SHADER_PARAMETER_STRUCT(FTranslucentSelfShadowUniformParameters, TranslucentSelfShadow)
 END_GLOBAL_SHADER_PARAMETER_STRUCT()
 
-IMPLEMENT_GLOBAL_SHADER_PARAMETER_STRUCT(FTranslucencyDepthPassUniformParameters, "TranslucentDepthPass");
+IMPLEMENT_STATIC_UNIFORM_BUFFER_STRUCT(FTranslucencyDepthPassUniformParameters, "TranslucentDepthPass", SceneTextures);
 
 void SetupTranslucencyDepthPassUniformBuffer(
 	const FProjectedShadowInfo* ShadowInfo,
-	FRHICommandList& RHICmdList,
+	FRDGBuilder& GraphBuilder,
 	const FViewInfo& View,
 	FTranslucencyDepthPassUniformParameters& TranslucencyDepthPassParameters)
 {
 	// Note - scene depth can be bound by the material for use in depth fades
 	// This is incorrect when rendering a shadowmap as it's not from the camera's POV
 	// Set the scene depth texture to something safe when rendering shadow depths
-	FSceneRenderTargets& SceneRenderTargets = FSceneRenderTargets::Get(RHICmdList);
-	SetupSceneTextureUniformParameters(SceneRenderTargets, View.FeatureLevel, ESceneTextureSetupMode::None, TranslucencyDepthPassParameters.SceneTextures);
+	SetupSceneTextureUniformParameters(GraphBuilder, View.FeatureLevel, ESceneTextureSetupMode::None, TranslucencyDepthPassParameters.SceneTextures);
 
 	TranslucencyDepthPassParameters.ProjectionMatrix = FTranslationMatrix(ShadowInfo->PreShadowTranslation - View.ViewMatrices.GetPreViewTranslation()) * ShadowInfo->SubjectAndReceiverMatrix;
 
@@ -300,12 +240,7 @@ public:
 	FTranslucencyShadowDepthVS() {}
 	FTranslucencyShadowDepthVS(const FMeshMaterialShaderType::CompiledShaderInitializerType& Initializer) :
 		FMeshMaterialShader(Initializer)
-	{
-		PassUniformBuffer.Bind(Initializer.ParameterMap, FTranslucencyDepthPassUniformParameters::StaticStructMetadata.GetShaderVariableName());
-	}
-
-	
-	
+	{}
 };
 
 enum ETranslucencyShadowDepthShaderMode
@@ -322,8 +257,7 @@ public:
 
 	TTranslucencyShadowDepthVS(const ShaderMetaType::CompiledShaderInitializerType& Initializer) :
 		FTranslucencyShadowDepthVS(Initializer)
-	{
-	}
+	{}
 
 	TTranslucencyShadowDepthVS() {}
 
@@ -344,22 +278,18 @@ class FTranslucencyShadowDepthPS : public FMeshMaterialShader
 {
 	DECLARE_INLINE_TYPE_LAYOUT(FTranslucencyShadowDepthPS, NonVirtual);
 public:
-
 	static bool ShouldCompilePermutation(const FMeshMaterialShaderPermutationParameters& Parameters)
 	{
 		return IsTranslucentBlendMode(Parameters.MaterialParameters.BlendMode) && IsFeatureLevelSupported(Parameters.Platform, ERHIFeatureLevel::SM5);
 	}
 
-	FTranslucencyShadowDepthPS(const ShaderMetaType::CompiledShaderInitializerType& Initializer):
+	FTranslucencyShadowDepthPS() = default;
+	FTranslucencyShadowDepthPS(const ShaderMetaType::CompiledShaderInitializerType& Initializer) :
 		FMeshMaterialShader(Initializer)
 	{
-		PassUniformBuffer.Bind(Initializer.ParameterMap, FTranslucencyDepthPassUniformParameters::StaticStructMetadata.GetShaderVariableName());
-
 		TranslucentShadowStartOffset.Bind(Initializer.ParameterMap, TEXT("TranslucentShadowStartOffset"));
 	}
 
-	FTranslucencyShadowDepthPS() {}
-
 	void GetShaderBindings(
 		const FScene* Scene,
 		ERHIFeatureLevel::Type FeatureLevel,
@@ -376,28 +306,25 @@ public:
 	}
 
 private:
-	
 	LAYOUT_FIELD(FShaderParameter, TranslucentShadowStartOffset);
 };
 
 template <ETranslucencyShadowDepthShaderMode ShaderMode>
 class TTranslucencyShadowDepthPS : public FTranslucencyShadowDepthPS
 {
-	DECLARE_SHADER_TYPE(TTranslucencyShadowDepthPS,MeshMaterial);
 public:
+	DECLARE_SHADER_TYPE(TTranslucencyShadowDepthPS, MeshMaterial);
 
-	TTranslucencyShadowDepthPS(const ShaderMetaType::CompiledShaderInitializerType& Initializer):
-		FTranslucencyShadowDepthPS(Initializer)
-	{
-	}
-
-	TTranslucencyShadowDepthPS() {}
-
-	static void ModifyCompilationEnvironment(const FMaterialShaderPermutationParameters& Parameters, FShaderCompilerEnvironment& OutEnvironment )
+	static void ModifyCompilationEnvironment(const FMaterialShaderPermutationParameters& Parameters, FShaderCompilerEnvironment& OutEnvironment)
 	{
 		FTranslucencyShadowDepthPS::ModifyCompilationEnvironment(Parameters, OutEnvironment);
 		OutEnvironment.SetDefine(TEXT("PERSPECTIVE_CORRECT_DEPTH"), (uint32)(ShaderMode == TranslucencyShadowDepth_PerspectiveCorrect ? 1 : 0));
 	}
+
+	TTranslucencyShadowDepthPS() = default;
+	TTranslucencyShadowDepthPS(const ShaderMetaType::CompiledShaderInitializerType & Initializer) :
+		FTranslucencyShadowDepthPS(Initializer)
+	{}
 };
 
 IMPLEMENT_MATERIAL_SHADER_TYPE(template<>,TTranslucencyShadowDepthPS<TranslucencyShadowDepth_PerspectiveCorrect>,TEXT("/Engine/Private/TranslucentShadowDepthShaders.usf"),TEXT("MainOpacityPS"),SF_Pixel);
@@ -527,27 +454,43 @@ void FTranslucencyDepthPassMeshProcessor::AddMeshBatch(const FMeshBatch& RESTRIC
 	}
 }
 
-void FProjectedShadowInfo::RenderTranslucencyDepths(FRHICommandList& RHICmdList, FSceneRenderer* SceneRenderer)
+BEGIN_SHADER_PARAMETER_STRUCT(FTranslucencyDepthPassParameters, )
+	SHADER_PARAMETER_RDG_UNIFORM_BUFFER(FTranslucencyDepthPassUniformParameters, PassUniformBuffer)
+	RENDER_TARGET_BINDING_SLOTS()
+END_SHADER_PARAMETER_STRUCT()
+
+void FProjectedShadowInfo::RenderTranslucencyDepths(FRDGBuilder& GraphBuilder, FSceneRenderer* SceneRenderer, const FRenderTargetBindingSlots& InRenderTargets)
 {
-	check(RHICmdList.IsInsideRenderPass());
 	check(IsInRenderingThread());
 	checkSlow(!bWholeSceneShadow);
 	SCOPE_CYCLE_COUNTER(STAT_RenderPerObjectShadowDepthsTime);
 
-	FTranslucencyDepthPassUniformParameters TranslucencyDepthPassParameters;
-	SetupTranslucencyDepthPassUniformBuffer(this, RHICmdList, *ShadowDepthView, TranslucencyDepthPassParameters);
-	TUniformBufferRef<FTranslucencyDepthPassUniformParameters> PassUniformBuffer = TUniformBufferRef<FTranslucencyDepthPassUniformParameters>::CreateUniformBufferImmediate(TranslucencyDepthPassParameters, UniformBuffer_SingleFrame, EUniformBufferValidation::None);
+	BeginRenderView(GraphBuilder, SceneRenderer->Scene);
 
-	FMeshPassProcessorRenderState DrawRenderState(*ShadowDepthView, PassUniformBuffer);
-	{
+	auto* TranslucencyDepthPassParameters = GraphBuilder.AllocParameters<FTranslucencyDepthPassUniformParameters>();
+	SetupTranslucencyDepthPassUniformBuffer(this, GraphBuilder, *ShadowDepthView, *TranslucencyDepthPassParameters);
+	TRDGUniformBufferRef<FTranslucencyDepthPassUniformParameters> PassUniformBuffer = GraphBuilder.CreateUniformBuffer(TranslucencyDepthPassParameters);
+
+	auto* PassParameters = GraphBuilder.AllocParameters<FTranslucencyDepthPassParameters>();
+	PassParameters->PassUniformBuffer = PassUniformBuffer;
+	PassParameters->RenderTargets = InRenderTargets;
+
+	FString EventName;
 #if WANTS_DRAW_MESH_EVENTS
-		FString EventName;
-		if (GetEmitDrawEvents())
-		{
-			GetShadowTypeNameForDrawEvent(EventName);
-		}
-		SCOPED_DRAW_EVENTF(RHICmdList, EventShadowDepthActor, *EventName);
+	if (GetEmitDrawEvents())
+	{
+		GetShadowTypeNameForDrawEvent(EventName);
+	}
 #endif
+
+	GraphBuilder.AddPass(
+		RDG_EVENT_NAME("%s", *EventName),
+		PassParameters,
+		ERDGPassFlags::Raster,
+		[this, SceneRenderer](FRHICommandListImmediate& RHICmdList)
+	{
+		FMeshPassProcessorRenderState DrawRenderState(*ShadowDepthView);
+
 		// Clear the shadow and its border
 		RHICmdList.SetViewport(
 			X,
@@ -556,9 +499,9 @@ void FProjectedShadowInfo::RenderTranslucencyDepths(FRHICommandList& RHICmdList,
 			(X + BorderSize * 2 + ResolutionX),
 			(Y + BorderSize * 2 + ResolutionY),
 			1.0f
-			);
+		);
 
-		FLinearColor ClearColors[2] = {FLinearColor(0,0,0,0), FLinearColor(0,0,0,0)};
+		FLinearColor ClearColors[2] = { FLinearColor(0,0,0,0), FLinearColor(0,0,0,0) };
 		DrawClearQuadMRT(RHICmdList, true, UE_ARRAY_COUNT(ClearColors), ClearColors, false, 1.0f, false, 0);
 
 		// Set the viewport for the shadow.
@@ -569,7 +512,7 @@ void FProjectedShadowInfo::RenderTranslucencyDepths(FRHICommandList& RHICmdList,
 			(X + BorderSize + ResolutionX),
 			(Y + BorderSize + ResolutionY),
 			1.0f
-			);
+		);
 
 		DrawRenderState.SetDepthStencilState(TStaticDepthStencilState<false, CF_Always>::GetRHI());
 		DrawRenderState.SetBlendState(TStaticBlendState<
@@ -626,7 +569,7 @@ void FProjectedShadowInfo::RenderTranslucencyDepths(FRHICommandList& RHICmdList,
 			SortAndMergeDynamicPassMeshDrawCommands(SceneRenderer->FeatureLevel, VisibleMeshDrawCommands, DynamicMeshDrawCommandStorage, PrimitiveIdVertexBuffer, 1);
 			SubmitMeshDrawCommands(VisibleMeshDrawCommands, GraphicsMinimalPipelineStateSet, PrimitiveIdVertexBuffer, 0, bDynamicInstancing, 1, RHICmdList);
 		}
-	}
+	});
 }
 
 /** Pixel shader used to filter a single volume lighting cascade. */
@@ -726,63 +669,12 @@ public:
 		SetShaderValue(RHICmdList, ShaderRHI, VolumeCascadeIndex, VolumeCascadeIndexValue);
 	}
 
-	/** Serializer. */ 
-	/*friend FArchive& operator<<(FArchive& Ar,FTranslucentInjectParameters& P)
-	{
-		Ar << P.WorldToShadowMatrix;
-		Ar << P.ShadowmapMinMax;
-		Ar << P.VolumeCascadeIndex;
-		return Ar;
-	}*/
-
-private:
-	
-		LAYOUT_FIELD(FShaderParameter, WorldToShadowMatrix)
-		LAYOUT_FIELD(FShaderParameter, ShadowmapMinMax)
-		LAYOUT_FIELD(FShaderParameter, VolumeCascadeIndex)
-	
-};
-
-/** Pixel shader used to accumulate per-object translucent shadows into a volume texture. */
-class FTranslucentObjectShadowingPS : public FGlobalShader
-{
-	DECLARE_SHADER_TYPE(FTranslucentObjectShadowingPS,Global);
-public:
-
-	static void ModifyCompilationEnvironment(const FGlobalShaderPermutationParameters& Parameters, FShaderCompilerEnvironment& OutEnvironment)
-	{
-		FGlobalShader::ModifyCompilationEnvironment(Parameters, OutEnvironment);
-		OutEnvironment.SetDefine(TEXT("INJECTION_PIXEL_SHADER"), 1);
-	}
-
-	static bool ShouldCompilePermutation(const FGlobalShaderPermutationParameters& Parameters)
-	{
-		return IsFeatureLevelSupported(Parameters.Platform, ERHIFeatureLevel::SM5) && (RHISupportsGeometryShaders(Parameters.Platform) || RHISupportsVertexShaderLayer(Parameters.Platform));
-	}
-
-	FTranslucentObjectShadowingPS(const ShaderMetaType::CompiledShaderInitializerType& Initializer):
-		FGlobalShader(Initializer)
-	{
-		TranslucentInjectParameters.Bind(Initializer.ParameterMap);
-	}
-	FTranslucentObjectShadowingPS() {}
-
-	void SetParameters(FRHICommandList& RHICmdList, const FViewInfo& View, const FLightSceneInfo* LightSceneInfo, const FProjectedShadowInfo* ShadowMap, uint32 VolumeCascadeIndex)
-	{
-		FGlobalShader::SetParameters<FViewUniformShaderParameters>(RHICmdList, RHICmdList.GetBoundPixelShader(), View.ViewUniformBuffer);
-		TranslucentInjectParameters.Set(RHICmdList, RHICmdList.GetBoundPixelShader(), this, View, LightSceneInfo, ShadowMap, VolumeCascadeIndex, true);
-
-		FTranslucentSelfShadowUniformParameters TranslucentSelfShadowUniformParameters;
-		SetupTranslucentSelfShadowUniformParameters(ShadowMap, TranslucentSelfShadowUniformParameters);
-		SetUniformBufferParameterImmediate(RHICmdList, RHICmdList.GetBoundPixelShader(), GetUniformBufferParameter<FTranslucentSelfShadowUniformParameters>(), TranslucentSelfShadowUniformParameters);
-	}
-
 private:
-	LAYOUT_FIELD(FTranslucentInjectParameters, TranslucentInjectParameters);
+	LAYOUT_FIELD(FShaderParameter, WorldToShadowMatrix)
+	LAYOUT_FIELD(FShaderParameter, ShadowmapMinMax)
+	LAYOUT_FIELD(FShaderParameter, VolumeCascadeIndex)
 };
 
-IMPLEMENT_SHADER_TYPE(,FTranslucentObjectShadowingPS,TEXT("/Engine/Private/TranslucentLightingShaders.usf"),TEXT("PerObjectShadowingMainPS"),SF_Pixel);
-
 /** Shader that adds direct lighting contribution from the given light to the current volume lighting cascade. */
 template<ELightComponentType InjectionType, bool bDynamicallyShadowed, bool bApplyLightFunction, bool bInverseSquared>
 class TTranslucentLightingInjectPS : public FMaterialShader
@@ -1264,88 +1156,6 @@ FVolumeBounds CalculateLightVolumeBounds(const FSphere& LightBounds, const FView
 	return VolumeBounds;
 }
 
-void FDeferredShadingSceneRenderer::AccumulateTranslucentVolumeObjectShadowing(FRHICommandList& RHICmdList, const FProjectedShadowInfo* InProjectedShadowInfo, bool bClearVolume, const FViewInfo& View, const int32 ViewIndex)
-{
-	const FLightSceneInfo* LightSceneInfo = &InProjectedShadowInfo->GetLightSceneInfo();
-
-	if (bClearVolume)
-	{
-		ClearTranslucentVolumePerObjectShadowing(RHICmdList, ViewIndex);
-	}
-
-	if (GUseTranslucentLightingVolumes && GSupportsVolumeTextureRendering)
-	{
-		SCOPED_DRAW_EVENT(RHICmdList, AccumulateTranslucentVolumeShadowing);
-		SCOPED_GPU_STAT(RHICmdList, TranslucentLighting);
-
-		auto ShaderMap = GetGlobalShaderMap(FeatureLevel);
-
-		FSceneRenderTargets& SceneContext = FSceneRenderTargets::Get(RHICmdList);
-		// Inject into each volume cascade
-		for (uint32 VolumeCascadeIndex = 0; VolumeCascadeIndex < TVC_MAX; VolumeCascadeIndex++)
-		{
-			const bool bDirectionalLight = LightSceneInfo->Proxy->GetLightType() == LightType_Directional;
-			const FVolumeBounds VolumeBounds = CalculateLightVolumeBounds(LightSceneInfo->Proxy->GetBoundingSphere(), View, VolumeCascadeIndex, bDirectionalLight);
-
-			if (VolumeBounds.IsValid())
-			{
-				FRHITexture* RenderTarget;
-
-				if (VolumeCascadeIndex == 0)
-				{
-					RenderTarget = SceneContext.GetTranslucencyVolumeAmbient(TVC_Inner, ViewIndex)->GetRenderTargetItem().TargetableTexture;
-				}
-				else
-				{
-					RenderTarget = SceneContext.GetTranslucencyVolumeDirectional(TVC_Inner, ViewIndex)->GetRenderTargetItem().TargetableTexture;
-				}
-
-				FRHIRenderPassInfo RPInfo(RenderTarget, ERenderTargetActions::Load_Store);
-				RHICmdList.BeginRenderPass(RPInfo, TEXT("AccumulateVolumeObjectShadowing"));
-				{
-					FGraphicsPipelineStateInitializer GraphicsPSOInit;
-					RHICmdList.ApplyCachedRenderTargets(GraphicsPSOInit);
-
-					GraphicsPSOInit.RasterizerState = TStaticRasterizerState<FM_Solid, CM_None>::GetRHI();
-					GraphicsPSOInit.DepthStencilState = TStaticDepthStencilState<false, CF_Always>::GetRHI();
-
-					// Modulate the contribution of multiple object shadows in rgb
-					GraphicsPSOInit.BlendState = TStaticBlendState<CW_RGB, BO_Add, BF_DestColor, BF_Zero>::GetRHI();
-
-					TShaderMapRef<FWriteToSliceVS> VertexShader(ShaderMap);
-					TOptionalShaderMapRef<FWriteToSliceGS> GeometryShader(ShaderMap);
-					TShaderMapRef<FTranslucentObjectShadowingPS> PixelShader(ShaderMap);
-
-					GraphicsPSOInit.BoundShaderState.VertexDeclarationRHI = GScreenVertexDeclaration.VertexDeclarationRHI;
-					GraphicsPSOInit.BoundShaderState.VertexShaderRHI = VertexShader.GetVertexShader();
-#if PLATFORM_SUPPORTS_GEOMETRY_SHADERS
-					GraphicsPSOInit.BoundShaderState.GeometryShaderRHI = GeometryShader.GetGeometryShader();
-#endif
-					GraphicsPSOInit.BoundShaderState.PixelShaderRHI = PixelShader.GetPixelShader();
-					GraphicsPSOInit.PrimitiveType = PT_TriangleStrip;
-
-					SetGraphicsPipelineState(RHICmdList, GraphicsPSOInit);
-
-					const int32 TranslucencyLightingVolumeDim = GetTranslucencyLightingVolumeDim();
-
-					VertexShader->SetParameters(RHICmdList, VolumeBounds, FIntVector(TranslucencyLightingVolumeDim));
-					if(GeometryShader.IsValid())
-					{
-						GeometryShader->SetParameters(RHICmdList, VolumeBounds.MinZ);
-					}
-					PixelShader->SetParameters(RHICmdList, View, LightSceneInfo, InProjectedShadowInfo, VolumeCascadeIndex);
-				
-					RasterizeToVolumeTexture(RHICmdList, VolumeBounds);
-				}
-				RHICmdList.EndRenderPass();
-
-				RHICmdList.CopyToResolveTarget(SceneContext.GetTranslucencyVolumeAmbient((ETranslucencyVolumeCascade)VolumeCascadeIndex, ViewIndex)->GetRenderTargetItem().TargetableTexture,
-					SceneContext.GetTranslucencyVolumeAmbient((ETranslucencyVolumeCascade)VolumeCascadeIndex, ViewIndex)->GetRenderTargetItem().ShaderResourceTexture, FResolveParams());
-			}
-		}
-	}
-}
-
 /**
  * Helper function for finding and setting the right version of TTranslucentLightingInjectPS given template parameters.
  * @param MaterialProxy must not be 0
