Fixed shadow projections looking incorrect due to missing viewport call.

#rnx

diff --git a/Engine/Source/Runtime/Renderer/Private/ShadowRendering.cpp b/Engine/Source/Runtime/Renderer/Private/ShadowRendering.cpp
--- a/Engine/Source/Runtime/Renderer/Private/ShadowRendering.cpp
+++ b/Engine/Source/Runtime/Renderer/Private/ShadowRendering.cpp
@@ -1011,6 +1011,8 @@ void FProjectedShadowInfo::RenderProjection(
 		ERDGPassFlags::Raster | PassFlags,
 		[this, SceneRender, View, ViewIndex, HairVisibilityData, bProjectingForForwardShading, bMobileModulatedProjections](FRHICommandListImmediate& RHICmdList)
 	{
+		RHICmdList.SetViewport(View->ViewRect.Min.X, View->ViewRect.Min.Y, 0.0f, View->ViewRect.Max.X, View->ViewRect.Max.Y, 1.0f);
+
 		FScopeCycleCounter Scope(bWholeSceneShadow ? GET_STATID(STAT_RenderWholeSceneShadowProjectionsTime) : GET_STATID(STAT_RenderPerObjectShadowProjectionsTime));
 
 		FGraphicsPipelineStateInitializer GraphicsPSOInit;
@@ -1201,6 +1203,8 @@ void FProjectedShadowInfo::RenderOnePassPointLightProjection(
 		ERDGPassFlags::Raster | PassFlags,
 		[this, &View, HairVisibilityData, bProjectingForForwardShading, bCameraInsideLightGeometry, bUseTransmission, ViewIndex](FRHICommandList& RHICmdList)
 	{
+		RHICmdList.SetViewport(View.ViewRect.Min.X, View.ViewRect.Min.Y, 0.0f, View.ViewRect.Max.X, View.ViewRect.Max.Y, 1.0f);
+
 		FGraphicsPipelineStateInitializer GraphicsPSOInit;
 		RHICmdList.ApplyCachedRenderTargets(GraphicsPSOInit);
 		GraphicsPSOInit.BlendState = GetBlendStateForProjection(bProjectingForForwardShading, false);
@@ -1666,7 +1670,6 @@ void FSceneRenderer::RenderShadowProjections(
 		AddPass(GraphBuilder, [&UniformBuffers, &View, LightSceneProxy](FRHICommandList& RHICmdList)
 		{
 			UniformBuffers.UpdateViewUniformBuffer(View);
-			RHICmdList.SetViewport(View.ViewRect.Min.X, View.ViewRect.Min.Y, 0.0f, View.ViewRect.Max.X, View.ViewRect.Max.Y, 1.0f);
 			LightSceneProxy->SetScissorRect(RHICmdList, View, View.ViewRect);
 		});
 
