Ported translucent lighting volume to RDG.

diff --git a/Engine/Source/Runtime/Renderer/Private/AmbientCubemapParameters.h b/Engine/Source/Runtime/Renderer/Private/AmbientCubemapParameters.h
--- a/Engine/Source/Runtime/Renderer/Private/AmbientCubemapParameters.h
+++ b/Engine/Source/Runtime/Renderer/Private/AmbientCubemapParameters.h
@@ -8,34 +8,7 @@
 
 #include "CoreMinimal.h"
 #include "FinalPostProcessSettings.h"
-#include "ShaderParameters.h"
-
-class FShaderParameterMap;
-
-// DEPRECATED: use FAmbientCubemapParameters instead.
-class FCubemapShaderParameters
-{
-	DECLARE_TYPE_LAYOUT(FCubemapShaderParameters, NonVirtual);
-public:
-
-	void Bind(const FShaderParameterMap& ParameterMap);
-
-	void SetParameters(FRHICommandList& RHICmdList, FRHIPixelShader* ShaderRHI, const FFinalPostProcessSettings::FCubemapEntry& Entry) const;
-	void SetParameters(FRHICommandList& RHICmdList, FRHIComputeShader* ShaderRHI, const FFinalPostProcessSettings::FCubemapEntry& Entry) const;
-
-	friend FArchive& operator<<(FArchive& Ar, FCubemapShaderParameters& P);
-
-private:
-	template<typename TRHIShader>
-	void SetParametersTemplate(FRHICommandList& RHICmdList, TRHIShader* ShaderRHI, const FFinalPostProcessSettings::FCubemapEntry& Entry) const;
-
-	
-		LAYOUT_FIELD(FShaderParameter, AmbientCubemapColor)
-		LAYOUT_FIELD(FShaderParameter, AmbientCubemapMipAdjust)
-		LAYOUT_FIELD(FShaderResourceParameter, AmbientCubemap)
-		LAYOUT_FIELD(FShaderResourceParameter, AmbientCubemapSampler)
-	
-};
+#include "ShaderParameterMacros.h"
 
 /** Shader parameters needed for deferred passes sampling the ambient cube map. */
 BEGIN_SHADER_PARAMETER_STRUCT(FAmbientCubemapParameters, )
diff --git a/Engine/Source/Runtime/Renderer/Private/CompositionLighting/PostProcessAmbient.cpp b/Engine/Source/Runtime/Renderer/Private/CompositionLighting/PostProcessAmbient.cpp
--- a/Engine/Source/Runtime/Renderer/Private/CompositionLighting/PostProcessAmbient.cpp
+++ b/Engine/Source/Runtime/Renderer/Private/CompositionLighting/PostProcessAmbient.cpp
@@ -1,68 +1,8 @@
 // Copyright Epic Games, Inc. All Rights Reserved.
 
-/*=============================================================================
-	PostProcessAmbient.cpp: Post processing ambient implementation.
-=============================================================================*/
-
-#include "StaticBoundShaderState.h"
-#include "SceneUtils.h"
-#include "PostProcess/SceneRenderTargets.h"
-#include "SceneRenderTargetParameters.h"
-#include "ScenePrivate.h"
-#include "PostProcess/SceneFilterRendering.h"
-#include "PostProcess/PostProcessing.h"
 #include "AmbientCubemapParameters.h"
-#include "PipelineStateCache.h"
-
-/*-----------------------------------------------------------------------------
-FCubemapShaderParameters
------------------------------------------------------------------------------*/
-IMPLEMENT_TYPE_LAYOUT(FCubemapShaderParameters);
-
-void FCubemapShaderParameters::Bind(const FShaderParameterMap& ParameterMap)
-{
-	AmbientCubemapColor.Bind(ParameterMap, TEXT("AmbientCubemapColor"));
-	AmbientCubemapMipAdjust.Bind(ParameterMap, TEXT("AmbientCubemapMipAdjust"));
-	AmbientCubemap.Bind(ParameterMap, TEXT("AmbientCubemap"));
-	AmbientCubemapSampler.Bind(ParameterMap, TEXT("AmbientCubemapSampler"));
-}
-
-void FCubemapShaderParameters::SetParameters(FRHICommandList& RHICmdList, FRHIPixelShader* ShaderRHI, const FFinalPostProcessSettings::FCubemapEntry& Entry) const
-{
-	SetParametersTemplate(RHICmdList, ShaderRHI, Entry);
-}
-
-void FCubemapShaderParameters::SetParameters(FRHICommandList& RHICmdList, FRHIComputeShader* ShaderRHI, const FFinalPostProcessSettings::FCubemapEntry& Entry) const
-{
-	SetParametersTemplate(RHICmdList, ShaderRHI, Entry);
-}
-
-template<typename TRHIShader>
-void FCubemapShaderParameters::SetParametersTemplate(FRHICommandList& RHICmdList, TRHIShader* ShaderRHI, const FFinalPostProcessSettings::FCubemapEntry& Entry) const
-{
-	FAmbientCubemapParameters ShaderParameters;
-	SetupAmbientCubemapParameters(Entry, &ShaderParameters);
-
-	// floats to render the cubemap
-	{
-		SetShaderValue(RHICmdList, ShaderRHI, AmbientCubemapColor, ShaderParameters.AmbientCubemapColor);
-		SetShaderValue(RHICmdList, ShaderRHI, AmbientCubemapMipAdjust, ShaderParameters.AmbientCubemapMipAdjust);
-	}
-
-	// cubemap texture
-	{
-		FTexture* InternalSpec = Entry.AmbientCubemap ? Entry.AmbientCubemap->Resource : GBlackTextureCube;
-
-		SetTextureParameter(RHICmdList, ShaderRHI, AmbientCubemap, AmbientCubemapSampler, InternalSpec);
-	}
-}
-
-FArchive& operator<<(FArchive& Ar, FCubemapShaderParameters& P)
-{
-	Ar << P.AmbientCubemapColor << P.AmbientCubemap << P.AmbientCubemapSampler << P.AmbientCubemapMipAdjust;
-
-	return Ar;
-}
+#include "Engine/Texture.h"
+#include "RenderUtils.h"
 
 void SetupAmbientCubemapParameters(const FFinalPostProcessSettings::FCubemapEntry& Entry, FAmbientCubemapParameters* OutParameters)
 {
@@ -90,4 +30,4 @@ void SetupAmbientCubemapParameters(const FFinalPostProcessSettings::FCubemapEntr
 		OutParameters->AmbientCubemap = AmbientCubemapTexture->TextureRHI;
 		OutParameters->AmbientCubemapSampler = AmbientCubemapTexture->SamplerStateRHI;
 	}
-}
+}
\ No newline at end of file
diff --git a/Engine/Source/Runtime/Renderer/Private/DeferredShadingRenderer.cpp b/Engine/Source/Runtime/Renderer/Private/DeferredShadingRenderer.cpp
--- a/Engine/Source/Runtime/Renderer/Private/DeferredShadingRenderer.cpp
+++ b/Engine/Source/Runtime/Renderer/Private/DeferredShadingRenderer.cpp
@@ -2132,10 +2132,7 @@ void FDeferredShadingSceneRenderer::Render(FRHICommandListImmediate& RHICmdList)
 		
 		if (GbEnableAsyncComputeTranslucencyLightingVolumeClear && GSupportsEfficientAsyncCompute)
 		{
-			AddPass(GraphBuilder, [this](FRHICommandListImmediate& InRHICmdList)
-			{
-				ClearTranslucentVolumeLightingAsyncCompute(InRHICmdList);
-			});
+			ClearTranslucentVolumeLighting(GraphBuilder, ERDGPassFlags::AsyncCompute);
 		}
 	}
 
@@ -2395,12 +2392,7 @@ void FDeferredShadingSceneRenderer::Render(FRHICommandListImmediate& RHICmdList)
 		// Clear the translucent lighting volumes before we accumulate
 		if ((GbEnableAsyncComputeTranslucencyLightingVolumeClear && GSupportsEfficientAsyncCompute) == false)
 		{
-			for (int32 ViewIndex = 0; ViewIndex < Views.Num(); ++ViewIndex)
-			{
-				const FViewInfo& View = Views[ViewIndex];
-				RDG_GPU_MASK_SCOPE(GraphBuilder, View.GPUMask);
-				ClearTranslucentVolumeLighting(GraphBuilder, ViewIndex);
-			}
+			ClearTranslucentVolumeLighting(GraphBuilder, ERDGPassFlags::Compute);
 		}
 
 #if RHI_RAYTRACING
diff --git a/Engine/Source/Runtime/Renderer/Private/DeferredShadingRenderer.h b/Engine/Source/Runtime/Renderer/Private/DeferredShadingRenderer.h
--- a/Engine/Source/Runtime/Renderer/Private/DeferredShadingRenderer.h
+++ b/Engine/Source/Runtime/Renderer/Private/DeferredShadingRenderer.h
@@ -802,18 +802,12 @@ private:
 		TRDGUniformBufferRef<FSceneTextureUniformParameters> SceneTexturesUniformBuffer,
 		const FSimpleLightArray& SimpleLights);
 
-	/** Clears the translucency lighting volumes before light accumulation. */
-	void ClearTranslucentVolumeLighting(FRDGBuilder& GraphBuilder, int32 ViewIndex);
-
-	/** Clears the translucency lighting volume via an async compute shader overlapped with the basepass. */
-	void ClearTranslucentVolumeLightingAsyncCompute(FRHICommandListImmediate& RHICmdList);
+	/** Clears the translucency lighting volumes. Can be compute or async compute. */
+	void ClearTranslucentVolumeLighting(FRDGBuilder& GraphBuilder, ERDGPassFlags PassFlags);
 
 	/** Add AmbientCubemap to the lighting volumes. */
 	void InjectAmbientCubemapTranslucentVolumeLighting(FRDGBuilder& GraphBuilder, const FViewInfo& View, int32 ViewIndex);
 
-	/** Clears the volume texture used to accumulate per object shadows for translucency. */
-	void ClearTranslucentVolumePerObjectShadowing(FRHICommandList& RHICmdList, const int32 ViewIndex);
-
 	/** Accumulates direct lighting for the given light.  InProjectedShadowInfo can be NULL in which case the light will be unshadowed. */
 	void InjectTranslucentVolumeLighting(FRDGBuilder& GraphBuilder, const FLightSceneInfo& LightSceneInfo, const FProjectedShadowInfo* InProjectedShadowInfo, const FViewInfo& View, int32 ViewIndex);
 
diff --git a/Engine/Source/Runtime/Renderer/Private/PostProcess/SceneRenderTargets.cpp b/Engine/Source/Runtime/Renderer/Private/PostProcess/SceneRenderTargets.cpp
--- a/Engine/Source/Runtime/Renderer/Private/PostProcess/SceneRenderTargets.cpp
+++ b/Engine/Source/Runtime/Renderer/Private/PostProcess/SceneRenderTargets.cpp
@@ -1898,41 +1898,6 @@ void FSceneRenderTargets::AllocateFoveationTexture(FRHICommandList& RHICmdList)
 	}
 }
 
-const FTexture2DRHIRef& FSceneRenderTargets::GetOptionalShadowDepthColorSurface(FRHICommandList& RHICmdList, int32 Width, int32 Height) const
-{
-	// Look for matching resolution
-	int32 EmptySlot = -1;
-	for (int32 Index = 0; Index < UE_ARRAY_COUNT(OptionalShadowDepthColor); Index++)
-	{
-		if (OptionalShadowDepthColor[Index])
-		{
-			const FTexture2DRHIRef& TargetTexture = (const FTexture2DRHIRef&)OptionalShadowDepthColor[Index]->GetRenderTargetItem().TargetableTexture;
-			if (TargetTexture->GetSizeX() == Width && TargetTexture->GetSizeY() == Height)
-			{
-				return TargetTexture;
-			}
-		}
-		else
-		{
-			// Remember this as a free slot for allocation attempt
-			EmptySlot = Index;
-		}
-	}
-
-	if (EmptySlot == -1)
-	{
-		UE_LOG(LogRenderer, Fatal, TEXT("Exceeded storage space for OptionalShadowDepthColorSurface. Increase array size."));
-	}
-
-	// Allocate new shadow color buffer (it must be the same resolution as the depth target!)
-	const FIntPoint ShadowColorBufferResolution = FIntPoint(Width, Height);
-	FPooledRenderTargetDesc Desc(FPooledRenderTargetDesc::Create2DDesc(ShadowColorBufferResolution, PF_B8G8R8A8, FClearValueBinding::None, TexCreate_None, TexCreate_RenderTargetable, false));
-	GRenderTargetPool.FindFreeElement(RHICmdList, Desc, (TRefCountPtr<IPooledRenderTarget>&)OptionalShadowDepthColor[EmptySlot], TEXT("OptionalShadowDepthColor"));
-	UE_LOG(LogRenderer, Log, TEXT("Allocated OptionalShadowDepthColorSurface %d x %d"), Width, Height);
-
-	return (const FTexture2DRHIRef&)OptionalShadowDepthColor[EmptySlot]->GetRenderTargetItem().TargetableTexture;
-}
-
 void FSceneRenderTargets::AllocateDeferredShadingPathRenderTargets(FRHICommandListImmediate& RHICmdList, const int32 NumViews)
 {
 	const EShaderPlatform ShaderPlatform = GetFeatureLevelShaderPlatform(CurrentFeatureLevel);
@@ -2028,12 +1993,6 @@ void FSceneRenderTargets::AllocateDeferredShadingPathRenderTargets(FRHICommandLi
 					TranslucencyLightingVolumeDim,
 					TranslucencyTargetFlags);
 			}
-
-			//these get bound even with the CVAR off, make sure they aren't full of garbage.
-			if (!GUseTranslucentLightingVolumes)
-			{
-				ClearTranslucentVolumeLighting(RHICmdList, 0);
-			}
 		}
 	}
 
@@ -2107,104 +2066,6 @@ EPixelFormat FSceneRenderTargets::GetMobileSceneColorFormat() const
 	return CurrentMobileSceneColorFormat;
 }
 
-void FSceneRenderTargets::ClearTranslucentVolumeLighting(FRHICommandListImmediate& RHICmdList, int32 ViewIndex)
-{
-	if (GSupportsVolumeTextureRendering)
-	{
-		// Clear all volume textures in the same draw with MRT, which is faster than individually
-		static_assert(TVC_MAX == 2, "Only expecting two translucency lighting cascades.");
-		static IConsoleVariable* CVarTranslucencyVolumeBlur =
-			IConsoleManager::Get().FindConsoleVariable(TEXT("r.TranslucencyVolumeBlur"));
-		static constexpr int32 Num3DTextures = NumTranslucentVolumeRenderTargetSets << 1;
-
-		FRHITexture* RenderTargets[Num3DTextures];
-		bool bUseTransLightingVolBlur = CVarTranslucencyVolumeBlur->GetInt() > 0;
-		const int32 NumIterations = bUseTransLightingVolBlur ?
-			NumTranslucentVolumeRenderTargetSets : NumTranslucentVolumeRenderTargetSets - 1;
-
-		for (int32 Idx = 0; Idx < NumIterations; ++Idx)
-		{
-			RenderTargets[Idx << 1] = TranslucencyLightingVolumeAmbient[Idx + NumTranslucentVolumeRenderTargetSets * ViewIndex]->GetRenderTargetItem().TargetableTexture;
-			RenderTargets[(Idx << 1) + 1] = TranslucencyLightingVolumeDirectional[Idx + NumTranslucentVolumeRenderTargetSets * ViewIndex]->GetRenderTargetItem().TargetableTexture;
-		}
-
-		static const FLinearColor ClearColors[Num3DTextures] = { FLinearColor::Transparent };
-
-		if (bUseTransLightingVolBlur)
-		{
-			ClearVolumeTextures<Num3DTextures>(RHICmdList, CurrentFeatureLevel, RenderTargets, ClearColors);
-		}
-		else
-		{
-			ClearVolumeTextures<Num3DTextures - 2>(RHICmdList, CurrentFeatureLevel, RenderTargets, ClearColors);
-		}
-	}
-}
-
-/** Helper function that clears the given volume texture render targets. */
-template<int32 NumRenderTargets>
-void FSceneRenderTargets::ClearVolumeTextures(FRHICommandList& RHICmdList, ERHIFeatureLevel::Type FeatureLevel, FRHITexture** RenderTargets, const FLinearColor* ClearColors)
-{
-	check(!RHICmdList.IsInsideRenderPass());
-
-	FRHIRenderPassInfo RPInfo(NumRenderTargets, RenderTargets, ERenderTargetActions::DontLoad_Store);
-	TransitionRenderPassTargets(RHICmdList, RPInfo);
-
-	RHICmdList.BeginRenderPass(RPInfo, TEXT("ClearVolumeTextures"));
-	{
-	FGraphicsPipelineStateInitializer GraphicsPSOInit;
-	RHICmdList.ApplyCachedRenderTargets(GraphicsPSOInit);
-	GraphicsPSOInit.RasterizerState = TStaticRasterizerState<FM_Solid, CM_None>::GetRHI();
-	GraphicsPSOInit.DepthStencilState = TStaticDepthStencilState<false, CF_Always>::GetRHI();
-	GraphicsPSOInit.BlendState = TStaticBlendState<>::GetRHI();
-
-	const int32 TranslucencyLightingVolumeDim = GetTranslucencyLightingVolumeDim();
-
-	const FVolumeBounds VolumeBounds(TranslucencyLightingVolumeDim);
-	auto ShaderMap = GetGlobalShaderMap(FeatureLevel);
-	TShaderMapRef<FWriteToSliceVS> VertexShader(ShaderMap);
-	TOptionalShaderMapRef<FWriteToSliceGS> GeometryShader(ShaderMap);
-	TOneColorPixelShaderMRT::FPermutationDomain PermutationVector;
-	PermutationVector.Set<TOneColorPixelShaderMRT::TOneColorPixelShaderNumOutputs>(NumRenderTargets);
-	TShaderMapRef<TOneColorPixelShaderMRT>PixelShader(ShaderMap, PermutationVector);
-
-	GraphicsPSOInit.BoundShaderState.VertexDeclarationRHI = GScreenVertexDeclaration.VertexDeclarationRHI;
-	GraphicsPSOInit.BoundShaderState.VertexShaderRHI = VertexShader.GetVertexShader();
-#if PLATFORM_SUPPORTS_GEOMETRY_SHADERS
-	GraphicsPSOInit.BoundShaderState.GeometryShaderRHI = GeometryShader.GetGeometryShader();
-#endif
-	GraphicsPSOInit.BoundShaderState.PixelShaderRHI = PixelShader.GetPixelShader();
-	GraphicsPSOInit.PrimitiveType = PT_TriangleStrip;
-
-	SetGraphicsPipelineState(RHICmdList, GraphicsPSOInit);
-
-	VertexShader->SetParameters(RHICmdList, VolumeBounds, FIntVector(TranslucencyLightingVolumeDim));
-	if (GeometryShader.IsValid())
-	{
-		GeometryShader->SetParameters(RHICmdList, VolumeBounds.MinZ);
-	}
-	PixelShader->SetColors(RHICmdList, ClearColors, NumRenderTargets);
-
-	RasterizeToVolumeTexture(RHICmdList, VolumeBounds);
-	}
-	RHICmdList.EndRenderPass();
-
-	FRHITransitionInfo SRVTransitions[NumRenderTargets];
-	for (int RT = 0; RT < NumRenderTargets; ++RT)
-	{
-		SRVTransitions[RT] = FRHITransitionInfo(RenderTargets[RT], ERHIAccess::Unknown, ERHIAccess::SRVMask);
-	}
-	RHICmdList.Transition(MakeArrayView(SRVTransitions, NumRenderTargets));
-}
-template void FSceneRenderTargets::ClearVolumeTextures<1>(FRHICommandList& RHICmdList, ERHIFeatureLevel::Type FeatureLevel, FRHITexture** RenderTargets, const FLinearColor* ClearColors);
-template void FSceneRenderTargets::ClearVolumeTextures<2>(FRHICommandList& RHICmdList, ERHIFeatureLevel::Type FeatureLevel, FRHITexture** RenderTargets, const FLinearColor* ClearColors);
-template void FSceneRenderTargets::ClearVolumeTextures<3>(FRHICommandList& RHICmdList, ERHIFeatureLevel::Type FeatureLevel, FRHITexture** RenderTargets, const FLinearColor* ClearColors);
-template void FSceneRenderTargets::ClearVolumeTextures<4>(FRHICommandList& RHICmdList, ERHIFeatureLevel::Type FeatureLevel, FRHITexture** RenderTargets, const FLinearColor* ClearColors);
-template void FSceneRenderTargets::ClearVolumeTextures<5>(FRHICommandList& RHICmdList, ERHIFeatureLevel::Type FeatureLevel, FRHITexture** RenderTargets, const FLinearColor* ClearColors);
-template void FSceneRenderTargets::ClearVolumeTextures<6>(FRHICommandList& RHICmdList, ERHIFeatureLevel::Type FeatureLevel, FRHITexture** RenderTargets, const FLinearColor* ClearColors);
-template void FSceneRenderTargets::ClearVolumeTextures<7>(FRHICommandList& RHICmdList, ERHIFeatureLevel::Type FeatureLevel, FRHITexture** RenderTargets, const FLinearColor* ClearColors);
-template void FSceneRenderTargets::ClearVolumeTextures<8>(FRHICommandList& RHICmdList, ERHIFeatureLevel::Type FeatureLevel, FRHITexture** RenderTargets, const FLinearColor* ClearColors);
-
 EPixelFormat FSceneRenderTargets::GetSceneColorFormat() const
 {
 	return GetSceneColorFormat(CurrentFeatureLevel);
diff --git a/Engine/Source/Runtime/Renderer/Private/PostProcess/SceneRenderTargets.h b/Engine/Source/Runtime/Renderer/Private/PostProcess/SceneRenderTargets.h
--- a/Engine/Source/Runtime/Renderer/Private/PostProcess/SceneRenderTargets.h
+++ b/Engine/Source/Runtime/Renderer/Private/PostProcess/SceneRenderTargets.h
@@ -330,7 +330,6 @@ public:
 	const FTextureRHIRef& GetSceneColorSurface() const;
 	const FTexture2DRHIRef& GetSceneDepthSurface() const							{ return (const FTexture2DRHIRef&)SceneDepthZ->GetRenderTargetItem().TargetableTexture; }
 	const FTexture2DRHIRef& GetSmallDepthSurface() const							{ return (const FTexture2DRHIRef&)SmallDepthZ->GetRenderTargetItem().TargetableTexture; }
-	const FTexture2DRHIRef& GetOptionalShadowDepthColorSurface(FRHICommandList& RHICmdList, int32 Width, int32 Height) const;
 
 	// @return can be empty if the feature is disabled
 	FCustomDepthTextures RequestCustomDepth(FRDGBuilder& GraphBuilder, bool bPrimitives);
@@ -339,13 +338,6 @@ public:
 
 	bool UseDownsizedOcclusionQueries() const { return bUseDownsizedOcclusionQueries; }
 
-	// ---
-
-	template<int32 NumRenderTargets>
-	static void ClearVolumeTextures(FRHICommandList& RHICmdList, ERHIFeatureLevel::Type FeatureLevel, FRHITexture** RenderTargets, const FLinearColor* ClearColors);
-
-	void ClearTranslucentVolumeLighting(FRHICommandListImmediate& RHICmdList, int32 ViewIndex);
-
 	/** Get the current translucent ambient lighting volume texture. Can vary depending on whether volume filtering is enabled */
 	IPooledRenderTarget* GetTranslucencyVolumeAmbient(ETranslucencyVolumeCascade Cascade, int32 ViewIndex = 0) const
 	{
@@ -700,15 +692,6 @@ private:
 
 };
 
-extern template void FSceneRenderTargets::ClearVolumeTextures<0>(FRHICommandList& RHICmdList, ERHIFeatureLevel::Type FeatureLevel, FRHITexture** RenderTargets, const FLinearColor* ClearColors);
-extern template void FSceneRenderTargets::ClearVolumeTextures<1>(FRHICommandList& RHICmdList, ERHIFeatureLevel::Type FeatureLevel, FRHITexture** RenderTargets, const FLinearColor* ClearColors);
-extern template void FSceneRenderTargets::ClearVolumeTextures<2>(FRHICommandList& RHICmdList, ERHIFeatureLevel::Type FeatureLevel, FRHITexture** RenderTargets, const FLinearColor* ClearColors);
-extern template void FSceneRenderTargets::ClearVolumeTextures<3>(FRHICommandList& RHICmdList, ERHIFeatureLevel::Type FeatureLevel, FRHITexture** RenderTargets, const FLinearColor* ClearColors);
-extern template void FSceneRenderTargets::ClearVolumeTextures<4>(FRHICommandList& RHICmdList, ERHIFeatureLevel::Type FeatureLevel, FRHITexture** RenderTargets, const FLinearColor* ClearColors);
-extern template void FSceneRenderTargets::ClearVolumeTextures<5>(FRHICommandList& RHICmdList, ERHIFeatureLevel::Type FeatureLevel, FRHITexture** RenderTargets, const FLinearColor* ClearColors);
-extern template void FSceneRenderTargets::ClearVolumeTextures<6>(FRHICommandList& RHICmdList, ERHIFeatureLevel::Type FeatureLevel, FRHITexture** RenderTargets, const FLinearColor* ClearColors);
-extern template void FSceneRenderTargets::ClearVolumeTextures<7>(FRHICommandList& RHICmdList, ERHIFeatureLevel::Type FeatureLevel, FRHITexture** RenderTargets, const FLinearColor* ClearColors);
-
 /** Sets up scene texture parameters for RDG (builder is valid) or passthrough RHI access (builder is null). Intended for temporary use during RDG refactor. */
 extern void SetupSceneTextureUniformParameters(
 	FRDGBuilder* GraphBuilder,
diff --git a/Engine/Source/Runtime/Renderer/Private/TranslucentLighting.cpp b/Engine/Source/Runtime/Renderer/Private/TranslucentLighting.cpp
--- a/Engine/Source/Runtime/Renderer/Private/TranslucentLighting.cpp
+++ b/Engine/Source/Runtime/Renderer/Private/TranslucentLighting.cpp
@@ -575,61 +575,27 @@ void FProjectedShadowInfo::RenderTranslucencyDepths(FRDGBuilder& GraphBuilder, F
 /** Pixel shader used to filter a single volume lighting cascade. */
 class FFilterTranslucentVolumePS : public FGlobalShader
 {
-	DECLARE_SHADER_TYPE(FFilterTranslucentVolumePS,Global);
 public:
+	DECLARE_GLOBAL_SHADER(FFilterTranslucentVolumePS);
+	SHADER_USE_PARAMETER_STRUCT(FFilterTranslucentVolumePS, FGlobalShader);
+
+	BEGIN_SHADER_PARAMETER_STRUCT(FParameters, )
+		SHADER_PARAMETER_STRUCT_REF(FViewUniformShaderParameters, View)
+		SHADER_PARAMETER_RDG_TEXTURE(Texture3D, TranslucencyLightingVolumeAmbient)
+		SHADER_PARAMETER_RDG_TEXTURE(Texture3D, TranslucencyLightingVolumeDirectional)
+		SHADER_PARAMETER_SAMPLER(SamplerState, TranslucencyLightingVolumeAmbientSampler)
+		SHADER_PARAMETER_SAMPLER(SamplerState, TranslucencyLightingVolumeDirectionalSampler)
+		SHADER_PARAMETER(float, TexelSize)
+		RENDER_TARGET_BINDING_SLOTS()
+	END_SHADER_PARAMETER_STRUCT()
 
 	static bool ShouldCompilePermutation(const FGlobalShaderPermutationParameters& Parameters) 
 	{ 
 		return IsFeatureLevelSupported(Parameters.Platform, ERHIFeatureLevel::SM5) && (RHISupportsGeometryShaders(Parameters.Platform) || RHISupportsVertexShaderLayer(Parameters.Platform));
 	}
-
-	FFilterTranslucentVolumePS(const ShaderMetaType::CompiledShaderInitializerType& Initializer):
-		FGlobalShader(Initializer)
-	{
-		TexelSize.Bind(Initializer.ParameterMap, TEXT("TexelSize"));
-		TranslucencyLightingVolumeAmbient.Bind(Initializer.ParameterMap, TEXT("TranslucencyLightingVolumeAmbient"));
-		TranslucencyLightingVolumeAmbientSampler.Bind(Initializer.ParameterMap, TEXT("TranslucencyLightingVolumeAmbientSampler"));
-		TranslucencyLightingVolumeDirectional.Bind(Initializer.ParameterMap, TEXT("TranslucencyLightingVolumeDirectional"));
-		TranslucencyLightingVolumeDirectionalSampler.Bind(Initializer.ParameterMap, TEXT("TranslucencyLightingVolumeDirectionalSampler"));
-	}
-	FFilterTranslucentVolumePS() {}
-
-	void SetParameters(FRHICommandList& RHICmdList, const FViewInfo& View, int32 VolumeCascadeIndex, const int32 ViewIndex)
-	{
-		FSceneRenderTargets& SceneContext = FSceneRenderTargets::Get(RHICmdList);
-		FRHIPixelShader* ShaderRHI = RHICmdList.GetBoundPixelShader();
-
-		FGlobalShader::SetParameters<FViewUniformShaderParameters>(RHICmdList, ShaderRHI, View.ViewUniformBuffer);
-
-		const int32 TranslucencyLightingVolumeDim = GetTranslucencyLightingVolumeDim();
-		SetShaderValue(RHICmdList, ShaderRHI, TexelSize, 1.0f / TranslucencyLightingVolumeDim);
-
-		SetTextureParameter(
-			RHICmdList,
-			ShaderRHI, 
-			TranslucencyLightingVolumeAmbient, 
-			TranslucencyLightingVolumeAmbientSampler, 
-			TStaticSamplerState<SF_Bilinear, AM_Clamp, AM_Clamp, AM_Clamp>::GetRHI(), 
-			SceneContext.TranslucencyLightingVolumeAmbient[VolumeCascadeIndex + NumTranslucentVolumeRenderTargetSets * ViewIndex]->GetRenderTargetItem().ShaderResourceTexture);
-
-		SetTextureParameter(
-			RHICmdList,
-			ShaderRHI, 
-			TranslucencyLightingVolumeDirectional, 
-			TranslucencyLightingVolumeDirectionalSampler, 
-			TStaticSamplerState<SF_Bilinear, AM_Clamp, AM_Clamp, AM_Clamp>::GetRHI(), 
-			SceneContext.TranslucencyLightingVolumeDirectional[VolumeCascadeIndex + NumTranslucentVolumeRenderTargetSets * ViewIndex]->GetRenderTargetItem().ShaderResourceTexture);
-	}
-
-private:
-	LAYOUT_FIELD(FShaderParameter, TexelSize);
-	LAYOUT_FIELD(FShaderResourceParameter, TranslucencyLightingVolumeAmbient);
-	LAYOUT_FIELD(FShaderResourceParameter, TranslucencyLightingVolumeAmbientSampler);
-	LAYOUT_FIELD(FShaderResourceParameter, TranslucencyLightingVolumeDirectional);
-	LAYOUT_FIELD(FShaderResourceParameter, TranslucencyLightingVolumeDirectionalSampler);
 };
 
-IMPLEMENT_SHADER_TYPE(,FFilterTranslucentVolumePS,TEXT("/Engine/Private/TranslucentLightingShaders.usf"),TEXT("FilterMainPS"),SF_Pixel);
+IMPLEMENT_GLOBAL_SHADER(FFilterTranslucentVolumePS, "/Engine/Private/TranslucentLightingShaders.usf", "FilterMainPS", SF_Pixel);
 
 /** Shader parameters needed to inject direct lighting into a volume. */
 class FTranslucentInjectParameters
@@ -852,29 +818,21 @@ IMPLEMENT_INJECTION_PIXELSHADER_TYPE(LightType_Point,false,false,true);
 IMPLEMENT_INJECTION_PIXELSHADER_TYPE(LightType_Point,true,false,false); 
 IMPLEMENT_INJECTION_PIXELSHADER_TYPE(LightType_Point,false,false,false); 
 
-void FDeferredShadingSceneRenderer::ClearTranslucentVolumeLighting(FRDGBuilder& GraphBuilder, int32 ViewIndex)
-{
-	if (!GUseTranslucentLightingVolumes || !GSupportsVolumeTextureRendering)
-	{
-		return;
-	}
-
-	RDG_GPU_STAT_SCOPE(GraphBuilder, TranslucentLighting);
-
-	AddPass(GraphBuilder, RDG_EVENT_NAME("ClearTranslucentVolumeLighting"), [this, ViewIndex](FRHICommandListImmediate& RHICmdList)
-	{
-		FSceneRenderTargets& SceneContext = FSceneRenderTargets::Get(RHICmdList);
-		SceneContext.ClearTranslucentVolumeLighting(RHICmdList, ViewIndex);
-	});
-}
-
 class FClearTranslucentLightingVolumeCS : public FGlobalShader
 {
-	DECLARE_SHADER_TYPE(FClearTranslucentLightingVolumeCS, Global)
 public:
+	DECLARE_GLOBAL_SHADER(FClearTranslucentLightingVolumeCS);
+	SHADER_USE_PARAMETER_STRUCT(FClearTranslucentLightingVolumeCS, FGlobalShader)
 
 	static const int32 CLEAR_BLOCK_SIZE = 4;
 
+	BEGIN_SHADER_PARAMETER_STRUCT(FParameters, )
+		SHADER_PARAMETER_RDG_TEXTURE_UAV(RWTexture3D<float4>, RWAmbient0)
+		SHADER_PARAMETER_RDG_TEXTURE_UAV(RWTexture3D<float4>, RWDirectional0)
+		SHADER_PARAMETER_RDG_TEXTURE_UAV(RWTexture3D<float4>, RWAmbient1)
+		SHADER_PARAMETER_RDG_TEXTURE_UAV(RWTexture3D<float4>, RWDirectional1)
+	END_SHADER_PARAMETER_STRUCT()
+
 	static bool ShouldCompilePermutation(const FGlobalShaderPermutationParameters& Parameters)
 	{
 		return IsFeatureLevelSupported(Parameters.Platform, ERHIFeatureLevel::SM5);
@@ -886,152 +844,58 @@ public:
 		OutEnvironment.SetDefine(TEXT("CLEAR_COMPUTE_SHADER"), 1);
 		OutEnvironment.SetDefine(TEXT("CLEAR_BLOCK_SIZE"), CLEAR_BLOCK_SIZE);
 	}
-
-	FClearTranslucentLightingVolumeCS(const ShaderMetaType::CompiledShaderInitializerType& Initializer)
-		: FGlobalShader(Initializer)
-	{
-		Ambient0.Bind(Initializer.ParameterMap, TEXT("Ambient0"));
-		Directional0.Bind(Initializer.ParameterMap, TEXT("Directional0"));
-		Ambient1.Bind(Initializer.ParameterMap, TEXT("Ambient1"));
-		Directional1.Bind(Initializer.ParameterMap, TEXT("Directional1"));
-	}
-
-	FClearTranslucentLightingVolumeCS()
-	{
-	}
-
-	void SetParameters(
-		FRHIAsyncComputeCommandListImmediate& RHICmdList,
-		FRHIUnorderedAccessView** VolumeUAVs,
-		int32 NumUAVs
-	)
-	{
-		check(NumUAVs == 4);
-		FRHIComputeShader* ShaderRHI = RHICmdList.GetBoundComputeShader();
-		Ambient0.SetTexture(RHICmdList, ShaderRHI, NULL, VolumeUAVs[0]);
-		Directional0.SetTexture(RHICmdList, ShaderRHI, NULL, VolumeUAVs[1]);
-		Ambient1.SetTexture(RHICmdList, ShaderRHI, NULL, VolumeUAVs[2]);
-		Directional1.SetTexture(RHICmdList, ShaderRHI, NULL, VolumeUAVs[3]);
-	}
-
-	void UnsetParameters(FRHIAsyncComputeCommandListImmediate& RHICmdList)
-	{
-		FRHIComputeShader* ShaderRHI = RHICmdList.GetBoundComputeShader();
-		Ambient0.UnsetUAV(RHICmdList, ShaderRHI);
-		Directional0.UnsetUAV(RHICmdList, ShaderRHI);
-		Ambient1.UnsetUAV(RHICmdList, ShaderRHI);
-		Directional1.UnsetUAV(RHICmdList, ShaderRHI);
-	}
-
-private:
-	LAYOUT_FIELD(FRWShaderParameter, Ambient0);
-	LAYOUT_FIELD(FRWShaderParameter, Directional0);
-	LAYOUT_FIELD(FRWShaderParameter, Ambient1);
-	LAYOUT_FIELD(FRWShaderParameter, Directional1);
 };
 
-IMPLEMENT_SHADER_TYPE(, FClearTranslucentLightingVolumeCS, TEXT("/Engine/Private/TranslucentLightInjectionShaders.usf"), TEXT("ClearTranslucentLightingVolumeCS"), SF_Compute)
+IMPLEMENT_GLOBAL_SHADER(FClearTranslucentLightingVolumeCS, "/Engine/Private/TranslucentLightInjectionShaders.usf", "ClearTranslucentLightingVolumeCS", SF_Compute);
 
-void FDeferredShadingSceneRenderer::ClearTranslucentVolumeLightingAsyncCompute(FRHICommandListImmediate& RHICmdList)
+void FDeferredShadingSceneRenderer::ClearTranslucentVolumeLighting(FRDGBuilder& GraphBuilder, ERDGPassFlags PassFlags)
 {
-	FSceneRenderTargets& SceneContext = FSceneRenderTargets::Get(RHICmdList);
-
-	const int32 NumUAVs = 4;
-
-	TShaderMapRef<FClearTranslucentLightingVolumeCS> ComputeShader(GetGlobalShaderMap(FeatureLevel));
-
-	TArray<FRHIUnorderedAccessView*, SceneRenderingAllocator> VolumeUAVs;
-	VolumeUAVs.Reserve(Views.Num() * NumUAVs);
-	for (int i = 0; i < Views.Num(); ++i)
-	{
-		VolumeUAVs.Add(SceneContext.TranslucencyLightingVolumeAmbient[(i * NumTranslucentVolumeRenderTargetSets)]->GetRenderTargetItem().UAV);
-		VolumeUAVs.Add(SceneContext.TranslucencyLightingVolumeDirectional[(i * NumTranslucentVolumeRenderTargetSets)]->GetRenderTargetItem().UAV);
-		VolumeUAVs.Add(SceneContext.TranslucencyLightingVolumeAmbient[(i * NumTranslucentVolumeRenderTargetSets) + 1]->GetRenderTargetItem().UAV);
-		VolumeUAVs.Add(SceneContext.TranslucencyLightingVolumeDirectional[(i * NumTranslucentVolumeRenderTargetSets) + 1]->GetRenderTargetItem().UAV);
-	}
-
-	TArray<FRHITransitionInfo, SceneRenderingAllocator> UAVTransitions;
-	TArray<FRHITransitionInfo, SceneRenderingAllocator> SRVTransitions;
-	for (FRHIUnorderedAccessView* UAV : VolumeUAVs)
-	{
-		UAVTransitions.Add(FRHITransitionInfo(UAV, ERHIAccess::Unknown, ERHIAccess::UAVCompute));
-		SRVTransitions.Add(FRHITransitionInfo(UAV, ERHIAccess::UAVCompute, ERHIAccess::SRVMask));
-	}
+	RDG_GPU_STAT_SCOPE(GraphBuilder, TranslucentLighting);
 
-	// Transition should have already been ended before
-	check(TranslucencyLightingVolumeClearEndTransition == nullptr);
-	const FRHITransition* UAVTransition = RHICreateTransition(ERHIPipeline::Graphics, ERHIPipeline::AsyncCompute, ERHICreateTransitionFlags::None, MakeArrayView(UAVTransitions.GetData(), UAVTransitions.Num()));
-	TranslucencyLightingVolumeClearEndTransition = RHICreateTransition(ERHIPipeline::AsyncCompute, ERHIPipeline::Graphics, ERHICreateTransitionFlags::None, MakeArrayView(SRVTransitions.GetData(), SRVTransitions.Num()));
+	FSceneRenderTargets& SceneContext = FSceneRenderTargets::Get(GraphBuilder.RHICmdList);
 
-	// Begin the UAV transition on the Gfx pipe so the async clear compute shader won't clear until the Gfx pipe is caught up.
-	RHICmdList.BeginTransition(UAVTransition);
+	const FIntVector TranslucencyLightingVolumeDim(GetTranslucencyLightingVolumeDim());
+	const FIntVector GroupCount = FComputeShaderUtils::GetGroupCount(TranslucencyLightingVolumeDim, FClearTranslucentLightingVolumeCS::CLEAR_BLOCK_SIZE);
 
-	const int32 TranslucencyLightingVolumeDim = GetTranslucencyLightingVolumeDim();
+	TShaderMapRef<FClearTranslucentLightingVolumeCS> ComputeShader(GetGlobalShaderMap(FeatureLevel));
 
-	//Grab the async compute commandlist.
-	FRHIAsyncComputeCommandListImmediate& RHICmdListComputeImmediate = FRHICommandListExecutor::GetImmediateAsyncComputeCommandList();
+	for (int32 ViewIndex = 0; ViewIndex < Views.Num(); ++ViewIndex)
 	{
-		SCOPED_COMPUTE_EVENTF(RHICmdListComputeImmediate, ClearTranslucencyLightingVolume, TEXT("ClearTranslucencyLightingVolumeCompute %d"), TranslucencyLightingVolumeDim);
-
-		// we must wait on the transition from the Gfx pipe to let us know all our dependencies are ready.
-		RHICmdListComputeImmediate.EndTransition(UAVTransition);
-
-		for (int i = 0; i < Views.Num(); ++i)
-		{
-			//standard compute setup, but on the async commandlist.
-			RHICmdListComputeImmediate.SetComputeShader(ComputeShader.GetComputeShader());
-
-			ComputeShader->SetParameters(RHICmdListComputeImmediate, VolumeUAVs.GetData() + i * NumUAVs, NumUAVs);
-		
-			int32 GroupsPerDim = TranslucencyLightingVolumeDim / FClearTranslucentLightingVolumeCS::CLEAR_BLOCK_SIZE;
-			DispatchComputeShader(RHICmdListComputeImmediate, ComputeShader.GetShader(), GroupsPerDim, GroupsPerDim, GroupsPerDim);
-
-			ComputeShader->UnsetParameters(RHICmdListComputeImmediate);
-		}
-
-		//transition the output to readable on the async compute queue - gfx will call end transition when needed
-		RHICmdListComputeImmediate.BeginTransition(TranslucencyLightingVolumeClearEndTransition);
+		auto* PassParameters = GraphBuilder.AllocParameters<FClearTranslucentLightingVolumeCS::FParameters>();
+		PassParameters->RWAmbient0 = GraphBuilder.CreateUAV(GraphBuilder.RegisterExternalTexture(SceneContext.TranslucencyLightingVolumeAmbient[(ViewIndex * NumTranslucentVolumeRenderTargetSets) + 0]));
+		PassParameters->RWAmbient1 = GraphBuilder.CreateUAV(GraphBuilder.RegisterExternalTexture(SceneContext.TranslucencyLightingVolumeAmbient[(ViewIndex * NumTranslucentVolumeRenderTargetSets) + 1]));
+		PassParameters->RWDirectional0 = GraphBuilder.CreateUAV(GraphBuilder.RegisterExternalTexture(SceneContext.TranslucencyLightingVolumeDirectional[(ViewIndex * NumTranslucentVolumeRenderTargetSets) + 0]));
+		PassParameters->RWDirectional1 = GraphBuilder.CreateUAV(GraphBuilder.RegisterExternalTexture(SceneContext.TranslucencyLightingVolumeDirectional[(ViewIndex * NumTranslucentVolumeRenderTargetSets) + 1]));
+
+		FComputeShaderUtils::AddPass(
+			GraphBuilder,
+			RDG_EVENT_NAME("ClearTranslucencyLightingVolumeCompute %d", TranslucencyLightingVolumeDim),
+			PassFlags,
+			ComputeShader,
+			PassParameters,
+			GroupCount);
 	}
-
-	//immediately dispatch our async compute commands to the RHI thread to be submitted to the GPU as soon as possible.
-	//dispatch after the scope so the drawevent pop is inside the dispatch
-	FRHIAsyncComputeCommandListImmediate::ImmediateDispatch(RHICmdListComputeImmediate);
 }
 
-/** Encapsulates a pixel shader that is adding ambient cubemap to the volume. */
 class FInjectAmbientCubemapPS : public FGlobalShader
 {
-	DECLARE_SHADER_TYPE(FInjectAmbientCubemapPS, Global);
-
-	static bool ShouldCompilePermutation(const FGlobalShaderPermutationParameters& Parameters)
-	{
-		return IsFeatureLevelSupported(Parameters.Platform, ERHIFeatureLevel::SM5);
-	}
-
-	/** Default constructor. */
-	FInjectAmbientCubemapPS() {}
-
 public:
-	LAYOUT_FIELD(FCubemapShaderParameters, CubemapShaderParameters)
+	DECLARE_GLOBAL_SHADER(FInjectAmbientCubemapPS);
+	SHADER_USE_PARAMETER_STRUCT(FInjectAmbientCubemapPS, FGlobalShader);
 
-	/** Initialization constructor. */
-	FInjectAmbientCubemapPS(const ShaderMetaType::CompiledShaderInitializerType& Initializer)
-		: FGlobalShader(Initializer)
-	{
-		CubemapShaderParameters.Bind(Initializer.ParameterMap);
-	}
+	BEGIN_SHADER_PARAMETER_STRUCT(FParameters, )
+		SHADER_PARAMETER_STRUCT_REF(FViewUniformShaderParameters, View)
+		SHADER_PARAMETER_STRUCT_INCLUDE(FAmbientCubemapParameters, AmbientCubemap)
+		RENDER_TARGET_BINDING_SLOTS()
+	END_SHADER_PARAMETER_STRUCT()
 
-	void SetParameters(FRHICommandList& RHICmdList, const FViewInfo& View, const FFinalPostProcessSettings::FCubemapEntry& CubemapEntry)
+	static bool ShouldCompilePermutation(const FGlobalShaderPermutationParameters& Parameters)
 	{
-		FRHIPixelShader* ShaderRHI = RHICmdList.GetBoundPixelShader();
-
-		FGlobalShader::SetParameters<FViewUniformShaderParameters>(RHICmdList, ShaderRHI, View.ViewUniformBuffer);
-
-		CubemapShaderParameters.SetParameters(RHICmdList, ShaderRHI, CubemapEntry);
+		return IsFeatureLevelSupported(Parameters.Platform, ERHIFeatureLevel::SM5);
 	}
 };
 
-IMPLEMENT_SHADER_TYPE(,FInjectAmbientCubemapPS,TEXT("/Engine/Private/TranslucentLightingShaders.usf"),TEXT("InjectAmbientCubemapMainPS"),SF_Pixel);
+IMPLEMENT_GLOBAL_SHADER(FInjectAmbientCubemapPS, "/Engine/Private/TranslucentLightingShaders.usf", "InjectAmbientCubemapMainPS", SF_Pixel);
 
 void FDeferredShadingSceneRenderer::InjectAmbientCubemapTranslucentVolumeLighting(FRDGBuilder& GraphBuilder, const FViewInfo& View, int32 ViewIndex)
 {
@@ -1040,40 +904,43 @@ void FDeferredShadingSceneRenderer::InjectAmbientCubemapTranslucentVolumeLightin
 		return;
 	}
 
-	AddUntrackedAccessPass(GraphBuilder, [this, &View, ViewIndex](FRHICommandList& RHICmdList)
-	{
-		FSceneRenderTargets& SceneContext = FSceneRenderTargets::Get(RHICmdList);
-
-		SCOPED_DRAW_EVENT(RHICmdList, InjectAmbientCubemapTranslucentVolumeLighting);
-		SCOPED_GPU_STAT(RHICmdList, TranslucentLighting);
+	RDG_EVENT_SCOPE(GraphBuilder, "InjectAmbientCubemapTranslucentVolumeLighting");
+	RDG_GPU_STAT_SCOPE(GraphBuilder, TranslucentLighting);
 
-		const int32 TranslucencyLightingVolumeDim = GetTranslucencyLightingVolumeDim();
+	FSceneRenderTargets& SceneContext = FSceneRenderTargets::Get(GraphBuilder.RHICmdList);
 
-		const FVolumeBounds VolumeBounds(TranslucencyLightingVolumeDim);
+	const int32 TranslucencyLightingVolumeDim = GetTranslucencyLightingVolumeDim();
+	const FVolumeBounds VolumeBounds(TranslucencyLightingVolumeDim);
 
-		auto ShaderMap = GetGlobalShaderMap(FeatureLevel);
+	auto ShaderMap = GetGlobalShaderMap(FeatureLevel);
 
-		FGraphicsPipelineStateInitializer GraphicsPSOInit;
-		GraphicsPSOInit.RasterizerState = TStaticRasterizerState<FM_Solid, CM_None>::GetRHI();
-		GraphicsPSOInit.DepthStencilState = TStaticDepthStencilState<false, CF_Always>::GetRHI();
-		GraphicsPSOInit.BlendState = TStaticBlendState<CW_RGBA, BO_Add, BF_One, BF_One, BO_Add, BF_One, BF_One>::GetRHI();
+	for (int32 VolumeCascadeIndex = 0; VolumeCascadeIndex < TVC_MAX; ++VolumeCascadeIndex)
+	{
+		FRDGTextureRef VolumeAmbientTexture = GraphBuilder.RegisterExternalTexture(SceneContext.TranslucencyLightingVolumeAmbient[VolumeCascadeIndex + NumTranslucentVolumeRenderTargetSets * ViewIndex]);
 
-		for (int32 VolumeCascadeIndex = 0; VolumeCascadeIndex < TVC_MAX; VolumeCascadeIndex++)
+		for (const FFinalPostProcessSettings::FCubemapEntry& CubemapEntry : View.FinalPostProcessSettings.ContributingCubemaps)
 		{
-			//Checks to detect/prevent UE-31578
-			const IPooledRenderTarget* RT0 = SceneContext.TranslucencyLightingVolumeAmbient[VolumeCascadeIndex + NumTranslucentVolumeRenderTargetSets * ViewIndex];
-
-			// we don't update the directional volume (could be a HQ option)
-			FRHIRenderPassInfo RPInfo(RT0->GetRenderTargetItem().TargetableTexture, ERenderTargetActions::Load_Store);
-			TransitionRenderPassTargets(RHICmdList, RPInfo);
-			RHICmdList.BeginRenderPass(RPInfo, TEXT("InjectAmbientCubemapTranslucentVolumeLighting"));
+			auto* PassParameters = GraphBuilder.AllocParameters<FInjectAmbientCubemapPS::FParameters>();
+			SetupAmbientCubemapParameters(CubemapEntry, &PassParameters->AmbientCubemap);
+			PassParameters->RenderTargets[0] = FRenderTargetBinding(VolumeAmbientTexture, ERenderTargetLoadAction::ELoad);
+			PassParameters->View = View.ViewUniformBuffer;
+
+			GraphBuilder.AddPass(
+				RDG_EVENT_NAME("Cascade %d", VolumeCascadeIndex),
+				PassParameters,
+				ERDGPassFlags::Raster,
+				[&View, PassParameters, ShaderMap, VolumeBounds, TranslucencyLightingVolumeDim](FRHICommandList& RHICmdList)
 			{
-				RHICmdList.ApplyCachedRenderTargets(GraphicsPSOInit);
-
 				TShaderMapRef<FWriteToSliceVS> VertexShader(ShaderMap);
 				TOptionalShaderMapRef<FWriteToSliceGS> GeometryShader(ShaderMap);
 				TShaderMapRef<FInjectAmbientCubemapPS> PixelShader(ShaderMap);
 
+				FGraphicsPipelineStateInitializer GraphicsPSOInit;
+				RHICmdList.ApplyCachedRenderTargets(GraphicsPSOInit);
+
+				GraphicsPSOInit.RasterizerState = TStaticRasterizerState<FM_Solid, CM_None>::GetRHI();
+				GraphicsPSOInit.DepthStencilState = TStaticDepthStencilState<false, CF_Always>::GetRHI();
+				GraphicsPSOInit.BlendState = TStaticBlendState<CW_RGBA, BO_Add, BF_One, BF_One, BO_Add, BF_One, BF_One>::GetRHI();
 				GraphicsPSOInit.BoundShaderState.VertexDeclarationRHI = GScreenVertexDeclaration.VertexDeclarationRHI;
 				GraphicsPSOInit.BoundShaderState.VertexShaderRHI = VertexShader.GetVertexShader();
 #if PLATFORM_SUPPORTS_GEOMETRY_SHADERS
@@ -1081,7 +948,6 @@ void FDeferredShadingSceneRenderer::InjectAmbientCubemapTranslucentVolumeLightin
 #endif
 				GraphicsPSOInit.BoundShaderState.PixelShaderRHI = PixelShader.GetPixelShader();
 				GraphicsPSOInit.PrimitiveType = PT_TriangleStrip;
-
 				SetGraphicsPipelineState(RHICmdList, GraphicsPSOInit);
 
 				VertexShader->SetParameters(RHICmdList, VolumeBounds, FIntVector(TranslucencyLightingVolumeDim));
@@ -1089,41 +955,9 @@ void FDeferredShadingSceneRenderer::InjectAmbientCubemapTranslucentVolumeLightin
 				{
 					GeometryShader->SetParameters(RHICmdList, VolumeBounds.MinZ);
 				}
-
-				uint32 Count = View.FinalPostProcessSettings.ContributingCubemaps.Num();
-				for (uint32 i = 0; i < Count; ++i)
-				{
-					const FFinalPostProcessSettings::FCubemapEntry& CubemapEntry = View.FinalPostProcessSettings.ContributingCubemaps[i];
-
-					PixelShader->SetParameters(RHICmdList, View, CubemapEntry);
-
-					RasterizeToVolumeTexture(RHICmdList, VolumeBounds);
-				}
-			}
-			RHICmdList.EndRenderPass();
-			RHICmdList.CopyToResolveTarget(RT0->GetRenderTargetItem().TargetableTexture, RT0->GetRenderTargetItem().ShaderResourceTexture, FResolveParams());
+				SetShaderParameters(RHICmdList, PixelShader, PixelShader.GetPixelShader(), *PassParameters);
+			});
 		}
-	});
-}
-
-void FDeferredShadingSceneRenderer::ClearTranslucentVolumePerObjectShadowing(FRHICommandList& RHICmdList, const int32 ViewIndex)
-{
-	if (GUseTranslucentLightingVolumes && GSupportsVolumeTextureRendering)
-	{
-		FSceneRenderTargets& SceneContext = FSceneRenderTargets::Get(RHICmdList);
-		SCOPED_DRAW_EVENT(RHICmdList, ClearTranslucentVolumePerLightShadowing);
-		SCOPED_GPU_STAT(RHICmdList, TranslucentLighting);
-
-		static_assert(TVC_MAX == 2, "Only expecting two translucency lighting cascades.");
-		FRHITexture* RenderTargets[2];
-		RenderTargets[0] = SceneContext.GetTranslucencyVolumeAmbient(TVC_Inner, ViewIndex)->GetRenderTargetItem().TargetableTexture;
-		RenderTargets[1] = SceneContext.GetTranslucencyVolumeDirectional(TVC_Inner, ViewIndex)->GetRenderTargetItem().TargetableTexture;
-
-		FLinearColor ClearColors[2];
-		ClearColors[0] = FLinearColor(1, 1, 1, 1);
-		ClearColors[1] = FLinearColor(1, 1, 1, 1);
-
-		FSceneRenderTargets::ClearVolumeTextures<UE_ARRAY_COUNT(RenderTargets)>(RHICmdList, FeatureLevel, RenderTargets, ClearColors);
 	}
 }
 
@@ -1316,6 +1150,7 @@ static FRDGTextureRef GetSkyTransmittanceLutTexture(FRDGBuilder& GraphBuilder, F
 
 BEGIN_SHADER_PARAMETER_STRUCT(FInjectTranslucentLightArrayParameters, )
 	RDG_TEXTURE_ACCESS(TransmittanceLutTexture, ERHIAccess::SRVGraphics)
+	RENDER_TARGET_BINDING_SLOTS()
 END_SHADER_PARAMETER_STRUCT()
 
 /** Injects all the lights in LightInjectionData into the translucent lighting volume textures. */
@@ -1323,120 +1158,106 @@ static void InjectTranslucentLightArray(
 	FRDGBuilder& GraphBuilder,
 	FScene* Scene,
 	const FViewInfo& View,
-	const TArray<FTranslucentLightInjectionData, SceneRenderingAllocator>& LightInjectionData,
+	TArrayView<const FTranslucentLightInjectionData> LightInjectionData,
 	int32 ViewIndex)
 {
- 	FInjectTranslucentLightArrayParameters* PassParameters = GraphBuilder.AllocParameters<FInjectTranslucentLightArrayParameters>();
- 	PassParameters->TransmittanceLutTexture = GetSkyTransmittanceLutTexture(GraphBuilder, Scene, View);
+	FSceneRenderTargets& SceneContext = FSceneRenderTargets::Get(GraphBuilder.RHICmdList);
+	INC_DWORD_STAT_BY(STAT_NumLightsInjectedIntoTranslucency, LightInjectionData.Num());
 
-	GraphBuilder.AddPass(
-		RDG_EVENT_NAME("InjectTranslucentLightArray"),
-		PassParameters,
-		ERDGPassFlags::Raster | ERDGPassFlags::SkipRenderPass | ERDGPassFlags::UntrackedAccess | ERDGPassFlags::NeverCull,
-		[&View, &LightInjectionData, ViewIndex](FRHICommandListImmediate& RHICmdList)
+	// Inject into each volume cascade
+	// Operate on one cascade at a time to reduce render target switches
+	for (uint32 VolumeCascadeIndex = 0; VolumeCascadeIndex < TVC_MAX; VolumeCascadeIndex++)
 	{
-		FSceneRenderTargets& SceneContext = FSceneRenderTargets::Get(RHICmdList);
-		INC_DWORD_STAT_BY(STAT_NumLightsInjectedIntoTranslucency, LightInjectionData.Num());
-
-		// Inject into each volume cascade
-		// Operate on one cascade at a time to reduce render target switches
-		for (uint32 VolumeCascadeIndex = 0; VolumeCascadeIndex < TVC_MAX; VolumeCascadeIndex++)
+		FRDGTextureRef VolumeAmbientTexture = GraphBuilder.RegisterExternalTexture(SceneContext.TranslucencyLightingVolumeAmbient[VolumeCascadeIndex + NumTranslucentVolumeRenderTargetSets * ViewIndex]);
+		FRDGTextureRef VolumeDirectionalTexture = GraphBuilder.RegisterExternalTexture(SceneContext.TranslucencyLightingVolumeDirectional[VolumeCascadeIndex + NumTranslucentVolumeRenderTargetSets * ViewIndex]);
+
+		auto* PassParameters = GraphBuilder.AllocParameters<FInjectTranslucentLightArrayParameters>();
+		PassParameters->TransmittanceLutTexture = GetSkyTransmittanceLutTexture(GraphBuilder, Scene, View);
+		PassParameters->RenderTargets[0] = FRenderTargetBinding(VolumeAmbientTexture, ERenderTargetLoadAction::ELoad);
+		PassParameters->RenderTargets[1] = FRenderTargetBinding(VolumeDirectionalTexture, ERenderTargetLoadAction::ELoad);
+
+		GraphBuilder.AddPass(
+			RDG_EVENT_NAME("InjectTranslucentLightArray"),
+			PassParameters,
+			ERDGPassFlags::Raster,
+			[&View, LightInjectionData, VolumeCascadeIndex](FRHICommandList& RHICmdList)
 		{
-			IPooledRenderTarget* RT0 = SceneContext.TranslucencyLightingVolumeAmbient[VolumeCascadeIndex + NumTranslucentVolumeRenderTargetSets * ViewIndex];
-			IPooledRenderTarget* RT1 = SceneContext.TranslucencyLightingVolumeDirectional[VolumeCascadeIndex + NumTranslucentVolumeRenderTargetSets * ViewIndex];
+			FGraphicsPipelineStateInitializer GraphicsPSOInit;
+			RHICmdList.ApplyCachedRenderTargets(GraphicsPSOInit);
 
-			GVisualizeTexture.SetCheckPoint(RHICmdList, RT0);
-			GVisualizeTexture.SetCheckPoint(RHICmdList, RT1);
+			GraphicsPSOInit.RasterizerState = TStaticRasterizerState<FM_Solid, CM_None>::GetRHI();
+			GraphicsPSOInit.DepthStencilState = TStaticDepthStencilState<false, CF_Always>::GetRHI();
+			GraphicsPSOInit.PrimitiveType = PT_TriangleStrip;
 
-			FRHITexture* RenderTargets[2];
-			RenderTargets[0] = RT0->GetRenderTargetItem().TargetableTexture;
-			RenderTargets[1] = RT1->GetRenderTargetItem().TargetableTexture;
-
-			FRHIRenderPassInfo RPInfo(UE_ARRAY_COUNT(RenderTargets), RenderTargets, ERenderTargetActions::Load_Store);
-			TransitionRenderPassTargets(RHICmdList, RPInfo);
-
-			RHICmdList.BeginRenderPass(RPInfo, TEXT("InjectTranslucentLightArray"));
+			for (int32 LightIndex = 0; LightIndex < LightInjectionData.Num(); LightIndex++)
 			{
-				FGraphicsPipelineStateInitializer GraphicsPSOInit;
-				RHICmdList.ApplyCachedRenderTargets(GraphicsPSOInit);
+				const FTranslucentLightInjectionData& InjectionData = LightInjectionData[LightIndex];
+				const FLightSceneInfo* const LightSceneInfo = InjectionData.LightSceneInfo;
+				const bool bInverseSquared = LightSceneInfo->Proxy->IsInverseSquared();
+				const bool bDirectionalLight = LightSceneInfo->Proxy->GetLightType() == LightType_Directional;
+				const FVolumeBounds VolumeBounds = CalculateLightVolumeBounds(LightSceneInfo->Proxy->GetBoundingSphere(), View, VolumeCascadeIndex, bDirectionalLight);
 
-				GraphicsPSOInit.RasterizerState = TStaticRasterizerState<FM_Solid, CM_None>::GetRHI();
-				GraphicsPSOInit.DepthStencilState = TStaticDepthStencilState<false, CF_Always>::GetRHI();
-				GraphicsPSOInit.PrimitiveType = PT_TriangleStrip;
-
-				for (int32 LightIndex = 0; LightIndex < LightInjectionData.Num(); LightIndex++)
+				if (VolumeBounds.IsValid())
 				{
-					const FTranslucentLightInjectionData& InjectionData = LightInjectionData[LightIndex];
-					const FLightSceneInfo* const LightSceneInfo = InjectionData.LightSceneInfo;
-					const bool bInverseSquared = LightSceneInfo->Proxy->IsInverseSquared();
-					const bool bDirectionalLight = LightSceneInfo->Proxy->GetLightType() == LightType_Directional;
-					const FVolumeBounds VolumeBounds = CalculateLightVolumeBounds(LightSceneInfo->Proxy->GetBoundingSphere(), View, VolumeCascadeIndex, bDirectionalLight);
+					TShaderMapRef<FWriteToSliceVS> VertexShader(View.ShaderMap);
+					TOptionalShaderMapRef<FWriteToSliceGS> GeometryShader(View.ShaderMap);
 
-					if (VolumeBounds.IsValid())
+					if (bDirectionalLight)
 					{
-						TShaderMapRef<FWriteToSliceVS> VertexShader(View.ShaderMap);
-						TOptionalShaderMapRef<FWriteToSliceGS> GeometryShader(View.ShaderMap);
+						// Accumulate the contribution of multiple lights
+						// Directional lights write their shadowing into alpha of the ambient texture
+						GraphicsPSOInit.BlendState = TStaticBlendState<
+							CW_RGBA, BO_Add, BF_One, BF_One, BO_Add, BF_One, BF_One,
+							CW_RGB, BO_Add, BF_One, BF_One, BO_Add, BF_One, BF_One>::GetRHI();
 
-						if (bDirectionalLight)
+						if (InjectionData.ProjectedShadowInfo)
 						{
-							// Accumulate the contribution of multiple lights
-							// Directional lights write their shadowing into alpha of the ambient texture
-							GraphicsPSOInit.BlendState = TStaticBlendState<
-								CW_RGBA, BO_Add, BF_One, BF_One, BO_Add, BF_One, BF_One,
-								CW_RGB, BO_Add, BF_One, BF_One, BO_Add, BF_One, BF_One>::GetRHI();
-
-							if (InjectionData.ProjectedShadowInfo)
-							{
-								// shadows, restricting light contribution to the cascade bounds (except last cascade far to get light functions and no shadows there)
-								SetInjectionShader<LightType_Directional, true>(RHICmdList, GraphicsPSOInit, View, InjectionData.LightFunctionMaterialProxy, LightSceneInfo,
-									InjectionData.ProjectedShadowInfo, InjectionData.ProjectedShadowInfo->CascadeSettings.ShadowSplitIndex, VolumeCascadeIndex,
-									VertexShader, GeometryShader, InjectionData.bApplyLightFunction, false);
-							}
-							else
-							{
-								// no shadows
-								SetInjectionShader<LightType_Directional, false>(RHICmdList, GraphicsPSOInit, View, InjectionData.LightFunctionMaterialProxy, LightSceneInfo,
-									InjectionData.ProjectedShadowInfo, INDEX_NONE, VolumeCascadeIndex,
-									VertexShader, GeometryShader, InjectionData.bApplyLightFunction, false);
-							}
+							// shadows, restricting light contribution to the cascade bounds (except last cascade far to get light functions and no shadows there)
+							SetInjectionShader<LightType_Directional, true>(RHICmdList, GraphicsPSOInit, View, InjectionData.LightFunctionMaterialProxy, LightSceneInfo,
+								InjectionData.ProjectedShadowInfo, InjectionData.ProjectedShadowInfo->CascadeSettings.ShadowSplitIndex, VolumeCascadeIndex,
+								VertexShader, GeometryShader, InjectionData.bApplyLightFunction, false);
 						}
 						else
 						{
-							// Accumulate the contribution of multiple lights
-							GraphicsPSOInit.BlendState = TStaticBlendState<
-								CW_RGB, BO_Add, BF_One, BF_One, BO_Add, BF_Zero, BF_One,
-								CW_RGB, BO_Add, BF_One, BF_One, BO_Add, BF_Zero, BF_One>::GetRHI();
-
-							if (InjectionData.ProjectedShadowInfo)
-							{
-								SetInjectionShader<LightType_Point, true>(RHICmdList, GraphicsPSOInit, View, InjectionData.LightFunctionMaterialProxy, LightSceneInfo,
-									InjectionData.ProjectedShadowInfo, INDEX_NONE, VolumeCascadeIndex,
-									VertexShader, GeometryShader, InjectionData.bApplyLightFunction, bInverseSquared);
-							}
-							else
-							{
-								SetInjectionShader<LightType_Point, false>(RHICmdList, GraphicsPSOInit, View, InjectionData.LightFunctionMaterialProxy, LightSceneInfo,
-									InjectionData.ProjectedShadowInfo, INDEX_NONE, VolumeCascadeIndex,
-									VertexShader, GeometryShader, InjectionData.bApplyLightFunction, bInverseSquared);
-							}
+							// no shadows
+							SetInjectionShader<LightType_Directional, false>(RHICmdList, GraphicsPSOInit, View, InjectionData.LightFunctionMaterialProxy, LightSceneInfo,
+								InjectionData.ProjectedShadowInfo, INDEX_NONE, VolumeCascadeIndex,
+								VertexShader, GeometryShader, InjectionData.bApplyLightFunction, false);
 						}
+					}
+					else
+					{
+						// Accumulate the contribution of multiple lights
+						GraphicsPSOInit.BlendState = TStaticBlendState<
+							CW_RGB, BO_Add, BF_One, BF_One, BO_Add, BF_Zero, BF_One,
+							CW_RGB, BO_Add, BF_One, BF_One, BO_Add, BF_Zero, BF_One>::GetRHI();
 
-						const int32 TranslucencyLightingVolumeDim = GetTranslucencyLightingVolumeDim();
-
-						VertexShader->SetParameters(RHICmdList, VolumeBounds, FIntVector(TranslucencyLightingVolumeDim));
-						if (GeometryShader.IsValid())
+						if (InjectionData.ProjectedShadowInfo)
 						{
-							GeometryShader->SetParameters(RHICmdList, VolumeBounds.MinZ);
+							SetInjectionShader<LightType_Point, true>(RHICmdList, GraphicsPSOInit, View, InjectionData.LightFunctionMaterialProxy, LightSceneInfo,
+								InjectionData.ProjectedShadowInfo, INDEX_NONE, VolumeCascadeIndex,
+								VertexShader, GeometryShader, InjectionData.bApplyLightFunction, bInverseSquared);
 						}
-						RasterizeToVolumeTexture(RHICmdList, VolumeBounds);
+						else
+						{
+							SetInjectionShader<LightType_Point, false>(RHICmdList, GraphicsPSOInit, View, InjectionData.LightFunctionMaterialProxy, LightSceneInfo,
+								InjectionData.ProjectedShadowInfo, INDEX_NONE, VolumeCascadeIndex,
+								VertexShader, GeometryShader, InjectionData.bApplyLightFunction, bInverseSquared);
+						}
+					}
+
+					const int32 TranslucencyLightingVolumeDim = GetTranslucencyLightingVolumeDim();
+
+					VertexShader->SetParameters(RHICmdList, VolumeBounds, FIntVector(TranslucencyLightingVolumeDim));
+					if (GeometryShader.IsValid())
+					{
+						GeometryShader->SetParameters(RHICmdList, VolumeBounds.MinZ);
 					}
+					RasterizeToVolumeTexture(RHICmdList, VolumeBounds);
 				}
 			}
-			RHICmdList.EndRenderPass();
-			RHICmdList.CopyToResolveTarget(RT0->GetRenderTargetItem().TargetableTexture, RT0->GetRenderTargetItem().ShaderResourceTexture, FResolveParams());
-			RHICmdList.CopyToResolveTarget(RT1->GetRenderTargetItem().TargetableTexture, RT1->GetRenderTargetItem().ShaderResourceTexture, FResolveParams());
-		}
-	});
+		});
+	}
 }
 
 void FDeferredShadingSceneRenderer::InjectTranslucentVolumeLighting(FRDGBuilder& GraphBuilder, const FLightSceneInfo& LightSceneInfo, const FProjectedShadowInfo* InProjectedShadowInfo, const FViewInfo& View, int32 ViewIndex)
@@ -1487,46 +1308,27 @@ void FDeferredShadingSceneRenderer::InjectTranslucentVolumeLightingArray(FRDGBui
 	}
 }
 
-/** Pixel shader used to inject simple lights into the translucent lighting volume */
 class FSimpleLightTranslucentLightingInjectPS : public FGlobalShader
 {
-	DECLARE_SHADER_TYPE(FSimpleLightTranslucentLightingInjectPS,Global);
 public:
+	DECLARE_GLOBAL_SHADER(FSimpleLightTranslucentLightingInjectPS);
+	SHADER_USE_PARAMETER_STRUCT(FSimpleLightTranslucentLightingInjectPS, FGlobalShader);
+
+	BEGIN_SHADER_PARAMETER_STRUCT(FParameters, )
+		SHADER_PARAMETER_STRUCT_REF(FViewUniformShaderParameters, View)
+		SHADER_PARAMETER(FVector4, SimpleLightPositionAndRadius)
+		SHADER_PARAMETER(FVector4, SimpleLightColorAndExponent)
+		SHADER_PARAMETER(uint32, VolumeCascadeIndex)
+		RENDER_TARGET_BINDING_SLOTS()
+	END_SHADER_PARAMETER_STRUCT()
 
 	static bool ShouldCompilePermutation(const FGlobalShaderPermutationParameters& Parameters) 
 	{ 
 		return IsFeatureLevelSupported(Parameters.Platform, ERHIFeatureLevel::SM5) && (RHISupportsGeometryShaders(Parameters.Platform) || RHISupportsVertexShaderLayer(Parameters.Platform));
 	}
-
-	FSimpleLightTranslucentLightingInjectPS(const ShaderMetaType::CompiledShaderInitializerType& Initializer):
-		FGlobalShader(Initializer)
-	{
-		VolumeCascadeIndex.Bind(Initializer.ParameterMap, TEXT("VolumeCascadeIndex"));
-		SimpleLightPositionAndRadius.Bind(Initializer.ParameterMap, TEXT("SimpleLightPositionAndRadius"));
-		SimpleLightColorAndExponent.Bind(Initializer.ParameterMap, TEXT("SimpleLightColorAndExponent"));
-	}
-	FSimpleLightTranslucentLightingInjectPS() {}
-
-	void SetParameters(FRHICommandList& RHICmdList, const FViewInfo& View, const FSimpleLightEntry& SimpleLight, const FSimpleLightPerViewEntry& SimpleLightPerViewData, int32 VolumeCascadeIndexValue)
-	{
-		FGlobalShader::SetParameters<FViewUniformShaderParameters>(RHICmdList, RHICmdList.GetBoundPixelShader(), View.ViewUniformBuffer);
-
-		FVector4 PositionAndRadius(SimpleLightPerViewData.Position, SimpleLight.Radius);
-		SetShaderValue(RHICmdList, RHICmdList.GetBoundPixelShader(), VolumeCascadeIndex, VolumeCascadeIndexValue);
-		SetShaderValue(RHICmdList, RHICmdList.GetBoundPixelShader(), SimpleLightPositionAndRadius, PositionAndRadius);
-
-		FVector4 LightColorAndExponent(SimpleLight.Color, SimpleLight.Exponent);
-
-		SetShaderValue(RHICmdList, RHICmdList.GetBoundPixelShader(), SimpleLightColorAndExponent, LightColorAndExponent);
-	}
-
-private:
-	LAYOUT_FIELD(FShaderParameter, VolumeCascadeIndex);
-	LAYOUT_FIELD(FShaderParameter, SimpleLightPositionAndRadius);
-	LAYOUT_FIELD(FShaderParameter, SimpleLightColorAndExponent);
 };
 
-IMPLEMENT_SHADER_TYPE(,FSimpleLightTranslucentLightingInjectPS,TEXT("/Engine/Private/TranslucentLightInjectionShaders.usf"),TEXT("SimpleLightInjectMainPS"),SF_Pixel);
+IMPLEMENT_GLOBAL_SHADER(FSimpleLightTranslucentLightingInjectPS, "/Engine/Private/TranslucentLightInjectionShaders.usf", "SimpleLightInjectMainPS", SF_Pixel);
 
 void FDeferredShadingSceneRenderer::InjectSimpleTranslucentVolumeLightingArray(FRDGBuilder& GraphBuilder, const FSimpleLightArray& SimpleLights, const FViewInfo& View, const int32 ViewIndex)
 {
@@ -1544,86 +1346,82 @@ void FDeferredShadingSceneRenderer::InjectSimpleTranslucentVolumeLightingArray(F
 
 	if (NumLightsToInject > 0)
 	{
-		AddUntrackedAccessPass(GraphBuilder, [this, &SimpleLights, &View, ViewIndex, NumLightsToInject](FRHICommandListImmediate& RHICmdList)
-		{
-			INC_DWORD_STAT_BY(STAT_NumLightsInjectedIntoTranslucency, NumLightsToInject);
-			FSceneRenderTargets& SceneContext = FSceneRenderTargets::Get(RHICmdList);
+		RDG_EVENT_SCOPE(GraphBuilder, "InjectSimpleTranslucentLightArray");
 
-			// Inject into each volume cascade
-			// Operate on one cascade at a time to reduce render target switches
-			for (int32 VolumeCascadeIndex = 0; VolumeCascadeIndex < TVC_MAX; VolumeCascadeIndex++)
-			{
-				IPooledRenderTarget* RT0 = SceneContext.TranslucencyLightingVolumeAmbient[VolumeCascadeIndex + NumTranslucentVolumeRenderTargetSets * ViewIndex];
-				IPooledRenderTarget* RT1 = SceneContext.TranslucencyLightingVolumeDirectional[VolumeCascadeIndex + NumTranslucentVolumeRenderTargetSets * ViewIndex];
+		INC_DWORD_STAT_BY(STAT_NumLightsInjectedIntoTranslucency, NumLightsToInject);
+		FSceneRenderTargets& SceneContext = FSceneRenderTargets::Get(GraphBuilder.RHICmdList);
 
-				GVisualizeTexture.SetCheckPoint(RHICmdList, RT0);
-				GVisualizeTexture.SetCheckPoint(RHICmdList, RT1);
+		const int32 TranslucencyLightingVolumeDim = GetTranslucencyLightingVolumeDim();
 
-				FRHITexture* RenderTargets[2];
-				RenderTargets[0] = RT0->GetRenderTargetItem().TargetableTexture;
-				RenderTargets[1] = RT1->GetRenderTargetItem().TargetableTexture;
+		// Inject into each volume cascade
+		// Operate on one cascade at a time to reduce render target switches
+		for (int32 VolumeCascadeIndex = 0; VolumeCascadeIndex < TVC_MAX; VolumeCascadeIndex++)
+		{
+			RDG_EVENT_SCOPE(GraphBuilder, "Cascade%d", VolumeCascadeIndex);
+			FRDGTextureRef VolumeAmbientTexture = GraphBuilder.RegisterExternalTexture(SceneContext.TranslucencyLightingVolumeAmbient[VolumeCascadeIndex + NumTranslucentVolumeRenderTargetSets * ViewIndex]);
+			FRDGTextureRef VolumeDirectionalTexture = GraphBuilder.RegisterExternalTexture(SceneContext.TranslucencyLightingVolumeDirectional[VolumeCascadeIndex + NumTranslucentVolumeRenderTargetSets * ViewIndex]);
+
+			for (int32 LightIndex = 0; LightIndex < SimpleLights.InstanceData.Num(); LightIndex++)
+			{
+				const FSimpleLightEntry& SimpleLight = SimpleLights.InstanceData[LightIndex];
+				const FSimpleLightPerViewEntry& SimpleLightPerViewData = SimpleLights.GetViewDependentData(LightIndex, ViewIndex, Views.Num());
 
-				FRHIRenderPassInfo RPInfo(UE_ARRAY_COUNT(RenderTargets), RenderTargets, ERenderTargetActions::Load_Store);
-				TransitionRenderPassTargets(RHICmdList, RPInfo);
-				RHICmdList.BeginRenderPass(RPInfo, TEXT("InjectSimpleTranslucentVolumeLightingArray"));
+				if (SimpleLight.bAffectTranslucency)
 				{
-					FGraphicsPipelineStateInitializer GraphicsPSOInit;
-					RHICmdList.ApplyCachedRenderTargets(GraphicsPSOInit);
-
-					GraphicsPSOInit.RasterizerState = TStaticRasterizerState<FM_Solid, CM_None>::GetRHI();
-					GraphicsPSOInit.DepthStencilState = TStaticDepthStencilState<false, CF_Always>::GetRHI();
-					// Accumulate the contribution of multiple lights
-					GraphicsPSOInit.BlendState = TStaticBlendState<
-						CW_RGB, BO_Add, BF_One, BF_One, BO_Add, BF_Zero, BF_One,
-						CW_RGB, BO_Add, BF_One, BF_One, BO_Add, BF_Zero, BF_One>::GetRHI();
-					GraphicsPSOInit.PrimitiveType = PT_TriangleStrip;
-
-					for (int32 LightIndex = 0; LightIndex < SimpleLights.InstanceData.Num(); LightIndex++)
+					const FSphere LightBounds(SimpleLightPerViewData.Position, SimpleLight.Radius);
+					const FVolumeBounds VolumeBounds = CalculateLightVolumeBounds(LightBounds, View, VolumeCascadeIndex, false);
+
+					if (VolumeBounds.IsValid())
 					{
-						const FSimpleLightEntry& SimpleLight = SimpleLights.InstanceData[LightIndex];
-						const FSimpleLightPerViewEntry& SimpleLightPerViewData = SimpleLights.GetViewDependentData(LightIndex, ViewIndex, Views.Num());
+						auto* PassParameters = GraphBuilder.AllocParameters<FSimpleLightTranslucentLightingInjectPS::FParameters>();
+						PassParameters->View = View.ViewUniformBuffer;
+						PassParameters->VolumeCascadeIndex = VolumeCascadeIndex;
+						PassParameters->SimpleLightPositionAndRadius = FVector4(SimpleLightPerViewData.Position, SimpleLight.Radius);
+						PassParameters->SimpleLightColorAndExponent = FVector4(SimpleLight.Color, SimpleLight.Exponent);
+						PassParameters->RenderTargets[0] = FRenderTargetBinding(VolumeAmbientTexture, ERenderTargetLoadAction::ELoad);
+						PassParameters->RenderTargets[1] = FRenderTargetBinding(VolumeDirectionalTexture, ERenderTargetLoadAction::ELoad);
 
-						if (SimpleLight.bAffectTranslucency)
+						TShaderMapRef<FWriteToSliceVS> VertexShader(View.ShaderMap);
+						TOptionalShaderMapRef<FWriteToSliceGS> GeometryShader(View.ShaderMap);
+						TShaderMapRef<FSimpleLightTranslucentLightingInjectPS> PixelShader(View.ShaderMap);
+
+						GraphBuilder.AddPass(
+							{},
+							PassParameters,
+							ERDGPassFlags::Raster,
+							[VertexShader, GeometryShader, PixelShader, PassParameters, VolumeBounds, TranslucencyLightingVolumeDim](FRHICommandList& RHICmdList)
 						{
-							const FSphere LightBounds(SimpleLightPerViewData.Position, SimpleLight.Radius);
-							const FVolumeBounds VolumeBounds = CalculateLightVolumeBounds(LightBounds, View, VolumeCascadeIndex, false);
+							FGraphicsPipelineStateInitializer GraphicsPSOInit;
+							RHICmdList.ApplyCachedRenderTargets(GraphicsPSOInit);
 
-							if (VolumeBounds.IsValid())
-							{
-								TShaderMapRef<FWriteToSliceVS> VertexShader(View.ShaderMap);
-								TOptionalShaderMapRef<FWriteToSliceGS> GeometryShader(View.ShaderMap);
-								TShaderMapRef<FSimpleLightTranslucentLightingInjectPS> PixelShader(View.ShaderMap);
+							GraphicsPSOInit.RasterizerState = TStaticRasterizerState<FM_Solid, CM_None>::GetRHI();
+							GraphicsPSOInit.DepthStencilState = TStaticDepthStencilState<false, CF_Always>::GetRHI();
+							// Accumulate the contribution of multiple lights
+							GraphicsPSOInit.BlendState = TStaticBlendState<
+								CW_RGB, BO_Add, BF_One, BF_One, BO_Add, BF_Zero, BF_One,
+								CW_RGB, BO_Add, BF_One, BF_One, BO_Add, BF_Zero, BF_One>::GetRHI();
+							GraphicsPSOInit.PrimitiveType = PT_TriangleStrip;
 
-								GraphicsPSOInit.BoundShaderState.VertexDeclarationRHI = GScreenVertexDeclaration.VertexDeclarationRHI;
-								GraphicsPSOInit.BoundShaderState.VertexShaderRHI = VertexShader.GetVertexShader();
+							GraphicsPSOInit.BoundShaderState.VertexDeclarationRHI = GScreenVertexDeclaration.VertexDeclarationRHI;
+							GraphicsPSOInit.BoundShaderState.VertexShaderRHI = VertexShader.GetVertexShader();
 #if PLATFORM_SUPPORTS_GEOMETRY_SHADERS
-								GraphicsPSOInit.BoundShaderState.GeometryShaderRHI = GeometryShader.GetGeometryShader();
+							GraphicsPSOInit.BoundShaderState.GeometryShaderRHI = GeometryShader.GetGeometryShader();
 #endif
-								GraphicsPSOInit.BoundShaderState.PixelShaderRHI = PixelShader.GetPixelShader();
-
-								SetGraphicsPipelineState(RHICmdList, GraphicsPSOInit);
-
-								const int32 TranslucencyLightingVolumeDim = GetTranslucencyLightingVolumeDim();
+							GraphicsPSOInit.BoundShaderState.PixelShaderRHI = PixelShader.GetPixelShader();
+							SetGraphicsPipelineState(RHICmdList, GraphicsPSOInit);
 
-								VertexShader->SetParameters(RHICmdList, VolumeBounds, FIntVector(TranslucencyLightingVolumeDim));
-								if (GeometryShader.IsValid())
-								{
-									GeometryShader->SetParameters(RHICmdList, VolumeBounds.MinZ);
-								}
-								PixelShader->SetParameters(RHICmdList, View, SimpleLight, SimpleLightPerViewData, VolumeCascadeIndex);
-
-
-
-								RasterizeToVolumeTexture(RHICmdList, VolumeBounds);
+							VertexShader->SetParameters(RHICmdList, VolumeBounds, FIntVector(TranslucencyLightingVolumeDim));
+							if (GeometryShader.IsValid())
+							{
+								GeometryShader->SetParameters(RHICmdList, VolumeBounds.MinZ);
 							}
-						}
+							SetShaderParameters(RHICmdList, PixelShader, PixelShader.GetPixelShader(), *PassParameters);
+							RasterizeToVolumeTexture(RHICmdList, VolumeBounds);
+						});
 					}
 				}
-				RHICmdList.EndRenderPass();
-				RHICmdList.CopyToResolveTarget(RT0->GetRenderTargetItem().TargetableTexture, RT0->GetRenderTargetItem().ShaderResourceTexture, FResolveParams());
-				RHICmdList.CopyToResolveTarget(RT1->GetRenderTargetItem().TargetableTexture, RT1->GetRenderTargetItem().ShaderResourceTexture, FResolveParams());
 			}
-		});
+		}
 	}
 }
 
@@ -1634,99 +1432,65 @@ void FDeferredShadingSceneRenderer::FilterTranslucentVolumeLighting(FRDGBuilder&
 		return;
 	}
 
-	AddUntrackedAccessPass(GraphBuilder, [this, &View, ViewIndex](FRHICommandListImmediate& RHICmdList)
-	{
-		FSceneRenderTargets& SceneContext = FSceneRenderTargets::Get(RHICmdList);
-
-		const int32 TranslucencyLightingVolumeDim = GetTranslucencyLightingVolumeDim();
-		SCOPED_DRAW_EVENTF(RHICmdList, FilterTranslucentVolume, TEXT("FilterTranslucentVolume %dx%dx%d Cascades:%d"),
-			TranslucencyLightingVolumeDim, TranslucencyLightingVolumeDim, TranslucencyLightingVolumeDim, TVC_MAX);
+	FSceneRenderTargets& SceneContext = FSceneRenderTargets::Get(GraphBuilder.RHICmdList);
 
-		SCOPED_GPU_STAT(RHICmdList, TranslucentLighting);
+	FRHISamplerState* SamplerStateRHI = TStaticSamplerState<SF_Bilinear, AM_Clamp, AM_Clamp, AM_Clamp>::GetRHI();
 
-		FGraphicsPipelineStateInitializer GraphicsPSOInit;
-		GraphicsPSOInit.RasterizerState = TStaticRasterizerState<FM_Solid, CM_None>::GetRHI();
-		GraphicsPSOInit.DepthStencilState = TStaticDepthStencilState<false, CF_Always>::GetRHI();
-		GraphicsPSOInit.BlendState = TStaticBlendState<>::GetRHI();
+	const int32 TranslucencyLightingVolumeDim = GetTranslucencyLightingVolumeDim();
+	RDG_EVENT_SCOPE(GraphBuilder, "FilterTranslucentVolume %dx%dx%d Cascades:%d", TranslucencyLightingVolumeDim, TranslucencyLightingVolumeDim, TranslucencyLightingVolumeDim, TVC_MAX);
+	RDG_GPU_STAT_SCOPE(GraphBuilder, TranslucentLighting);
 
-		bool bTransitionedToWriteable = (GUseTranslucentLightingVolumes && GSupportsVolumeTextureRendering && View.FinalPostProcessSettings.ContributingCubemaps.Num());
+	for (int32 VolumeCascadeIndex = 0; VolumeCascadeIndex < TVC_MAX; VolumeCascadeIndex++)
+	{
+		FRDGTextureRef OutputVolumeAmbientTexture = GraphBuilder.RegisterExternalTexture(SceneContext.GetTranslucencyVolumeAmbient((ETranslucencyVolumeCascade)VolumeCascadeIndex, ViewIndex));
+		FRDGTextureRef OutputVolumeDirectionalTexture = GraphBuilder.RegisterExternalTexture(SceneContext.GetTranslucencyVolumeDirectional((ETranslucencyVolumeCascade)VolumeCascadeIndex, ViewIndex));
+
+		FRDGTextureRef InputVolumeAmbientTexture = GraphBuilder.RegisterExternalTexture(SceneContext.TranslucencyLightingVolumeAmbient[VolumeCascadeIndex + NumTranslucentVolumeRenderTargetSets * ViewIndex]);
+		FRDGTextureRef InputVolumeDirectionalTexture = GraphBuilder.RegisterExternalTexture(SceneContext.TranslucencyLightingVolumeDirectional[VolumeCascadeIndex + NumTranslucentVolumeRenderTargetSets * ViewIndex]);
+
+		auto* PassParameters = GraphBuilder.AllocParameters<FFilterTranslucentVolumePS::FParameters>();
+		PassParameters->View = View.ViewUniformBuffer;
+		PassParameters->TexelSize = 1.0f / TranslucencyLightingVolumeDim;
+		PassParameters->TranslucencyLightingVolumeAmbient = InputVolumeAmbientTexture;
+		PassParameters->TranslucencyLightingVolumeDirectional = InputVolumeDirectionalTexture;
+		PassParameters->TranslucencyLightingVolumeAmbientSampler = SamplerStateRHI;
+		PassParameters->TranslucencyLightingVolumeDirectionalSampler = SamplerStateRHI;
+		PassParameters->RenderTargets[0] = FRenderTargetBinding(OutputVolumeAmbientTexture, ERenderTargetLoadAction::ELoad);
+		PassParameters->RenderTargets[1] = FRenderTargetBinding(OutputVolumeDirectionalTexture, ERenderTargetLoadAction::ELoad);
 
-		// Filter each cascade
-		for (int32 VolumeCascadeIndex = 0; VolumeCascadeIndex < TVC_MAX; VolumeCascadeIndex++)
+		const FVolumeBounds VolumeBounds(TranslucencyLightingVolumeDim);
+		TShaderMapRef<FWriteToSliceVS> VertexShader(View.ShaderMap);
+		TOptionalShaderMapRef<FWriteToSliceGS> GeometryShader(View.ShaderMap);
+		TShaderMapRef<FFilterTranslucentVolumePS> PixelShader(View.ShaderMap);
+
+		GraphBuilder.AddPass(
+			RDG_EVENT_NAME("Cascade%d", VolumeCascadeIndex),
+			PassParameters,
+			ERDGPassFlags::Raster,
+			[VertexShader, GeometryShader, PixelShader, PassParameters, VolumeBounds, TranslucencyLightingVolumeDim](FRHICommandList& RHICmdList)
 		{
-			IPooledRenderTarget* RT0 = SceneContext.GetTranslucencyVolumeAmbient((ETranslucencyVolumeCascade)VolumeCascadeIndex, ViewIndex);
-			IPooledRenderTarget* RT1 = SceneContext.GetTranslucencyVolumeDirectional((ETranslucencyVolumeCascade)VolumeCascadeIndex, ViewIndex);
-
-			const IPooledRenderTarget* Input0 = SceneContext.TranslucencyLightingVolumeAmbient[VolumeCascadeIndex + NumTranslucentVolumeRenderTargetSets * ViewIndex];
-			const IPooledRenderTarget* Input1 = SceneContext.TranslucencyLightingVolumeDirectional[VolumeCascadeIndex + NumTranslucentVolumeRenderTargetSets * ViewIndex];
-
-			GVisualizeTexture.SetCheckPoint(RHICmdList, RT0);
-			GVisualizeTexture.SetCheckPoint(RHICmdList, RT1);
-
-			FRHITexture* RenderTargets[2];
-			RenderTargets[0] = RT0->GetRenderTargetItem().TargetableTexture;
-			RenderTargets[1] = RT1->GetRenderTargetItem().TargetableTexture;
-
-			FRHITransitionInfo Inputs[2];
-			Inputs[0] = FRHITransitionInfo(Input0->GetRenderTargetItem().TargetableTexture, ERHIAccess::Unknown, ERHIAccess::SRVGraphics);
-			Inputs[1] = FRHITransitionInfo(Input1->GetRenderTargetItem().TargetableTexture, ERHIAccess::Unknown, ERHIAccess::SRVGraphics);
-
-			static_assert(TVC_MAX == 2, "Final transition logic should change");
-
-			//the volume textures should still be writable from the injection phase on the first loop.
-			if (!bTransitionedToWriteable || VolumeCascadeIndex > 0)
-			{
-				FRHITransitionInfo RTVTransitions[2];
-				for (int RT = 0; RT < 2; ++RT)
-				{
-					RTVTransitions[RT] = FRHITransitionInfo(RenderTargets[RT], ERHIAccess::Unknown, ERHIAccess::RTV);
-				}
-				RHICmdList.Transition(MakeArrayView(RTVTransitions, 2));
-			}
-			RHICmdList.Transition(MakeArrayView(Inputs, 2));
-
-			FRHIRenderPassInfo RPInfo(UE_ARRAY_COUNT(RenderTargets), RenderTargets, ERenderTargetActions::Load_Store);
-			TransitionRenderPassTargets(RHICmdList, RPInfo);
-			RHICmdList.BeginRenderPass(RPInfo, TEXT("FilterTranslucentVolumeLighting"));
-			{
-				RHICmdList.ApplyCachedRenderTargets(GraphicsPSOInit);
-
-				const FVolumeBounds VolumeBounds(TranslucencyLightingVolumeDim);
-				TShaderMapRef<FWriteToSliceVS> VertexShader(View.ShaderMap);
-				TOptionalShaderMapRef<FWriteToSliceGS> GeometryShader(View.ShaderMap);
-				TShaderMapRef<FFilterTranslucentVolumePS> PixelShader(View.ShaderMap);
-
-				GraphicsPSOInit.BoundShaderState.VertexDeclarationRHI = GScreenVertexDeclaration.VertexDeclarationRHI;
-				GraphicsPSOInit.BoundShaderState.VertexShaderRHI = VertexShader.GetVertexShader();
+			FGraphicsPipelineStateInitializer GraphicsPSOInit;
+			RHICmdList.ApplyCachedRenderTargets(GraphicsPSOInit);
+
+			GraphicsPSOInit.RasterizerState = TStaticRasterizerState<FM_Solid, CM_None>::GetRHI();
+			GraphicsPSOInit.DepthStencilState = TStaticDepthStencilState<false, CF_Always>::GetRHI();
+			GraphicsPSOInit.BlendState = TStaticBlendState<>::GetRHI();
+			GraphicsPSOInit.BoundShaderState.VertexDeclarationRHI = GScreenVertexDeclaration.VertexDeclarationRHI;
+			GraphicsPSOInit.BoundShaderState.VertexShaderRHI = VertexShader.GetVertexShader();
 #if PLATFORM_SUPPORTS_GEOMETRY_SHADERS
-				GraphicsPSOInit.BoundShaderState.GeometryShaderRHI = GeometryShader.GetGeometryShader();
+			GraphicsPSOInit.BoundShaderState.GeometryShaderRHI = GeometryShader.GetGeometryShader();
 #endif
-				GraphicsPSOInit.BoundShaderState.PixelShaderRHI = PixelShader.GetPixelShader();
-				GraphicsPSOInit.PrimitiveType = PT_TriangleStrip;
-
-				SetGraphicsPipelineState(RHICmdList, GraphicsPSOInit);
+			GraphicsPSOInit.BoundShaderState.PixelShaderRHI = PixelShader.GetPixelShader();
+			GraphicsPSOInit.PrimitiveType = PT_TriangleStrip;
+			SetGraphicsPipelineState(RHICmdList, GraphicsPSOInit);
 
-				VertexShader->SetParameters(RHICmdList, VolumeBounds, FIntVector(TranslucencyLightingVolumeDim));
-				if (GeometryShader.IsValid())
-				{
-					GeometryShader->SetParameters(RHICmdList, VolumeBounds.MinZ);
-				}
-				PixelShader->SetParameters(RHICmdList, View, VolumeCascadeIndex, ViewIndex);
-
-				RasterizeToVolumeTexture(RHICmdList, VolumeBounds);
-			}
-			RHICmdList.EndRenderPass();
-
-			//only do readable transition on the final loop since the other ones will do this up front.
-			//if (VolumeCascadeIndex == TVC_MAX - 1)
+			VertexShader->SetParameters(RHICmdList, VolumeBounds, FIntVector(TranslucencyLightingVolumeDim));
+			if (GeometryShader.IsValid())
 			{
-				FRHITransitionInfo SRVTransitions[2];
-				for (int RT = 0; RT < 2; ++RT)
-				{
-					SRVTransitions[RT] = FRHITransitionInfo(RenderTargets[RT], ERHIAccess::Unknown, ERHIAccess::SRVMask);
-				}
-				RHICmdList.Transition(MakeArrayView(SRVTransitions, 2));
+				GeometryShader->SetParameters(RHICmdList, VolumeBounds.MinZ);
 			}
-		}
-	});
+			SetShaderParameters(RHICmdList, PixelShader, PixelShader.GetPixelShader(), *PassParameters);
+			RasterizeToVolumeTexture(RHICmdList, VolumeBounds);
+		});
+	}
 }
