Fix wrapped VT address computation

diff --git a/Engine/Source/Runtime/Renderer/Private/VT/VirtualTextureSystem.cpp b/Engine/Source/Runtime/Renderer/Private/VT/VirtualTextureSystem.cpp
--- a/Engine/Source/Runtime/Renderer/Private/VT/VirtualTextureSystem.cpp
+++ b/Engine/Source/Runtime/Renderer/Private/VT/VirtualTextureSystem.cpp
@@ -1352,17 +1352,16 @@ void FVirtualTextureSystem::GatherRequestsTask(const FGatherRequestsParameters&
 			continue;
 		}
 
-		const uint32 vPageX = PageEncoded & 0xfff;
-		const uint32 vPageY = (PageEncoded >> 12) & 0xfff;
 		const uint32 vLevelPlusOne = ((PageEncoded >> 24) & 0x0f);
 		const uint32 vLevel = FMath::Max(vLevelPlusOne, 1u) - 1;
-		const uint32 vPosition = FMath::MortonCode2(vPageX) | (FMath::MortonCode2(vPageY) << 1);
-
-		// vPosition holds morton interleaved tileX/Y position, shifted down relative to current mip
-		// vAddress is the same quantity, but shifted to be relative to mip0
-		const uint32 vDimensions = Space->GetDimensions();
-		uint32 vAddress = vPosition << (vLevel * vDimensions);
-
+		
+		// vPageX/Y passed from shader are relative to the given vLevel, we shift them up so be relative to level0
+		// TODO - should we just do this in the shader?
+		const uint32 vPageX = (PageEncoded & 0xfff) << vLevel;
+		const uint32 vPageY = ((PageEncoded >> 12) & 0xfff) << vLevel;
+		
+		const uint32 vAddress = FMath::MortonCode2(vPageX) | (FMath::MortonCode2(vPageY) << 1);
+	
 		const FAdaptiveVirtualTexture* RESTRICT AdaptiveVT = AdaptiveVTs[ID];
 		if (AdaptiveVT != nullptr && vLevelPlusOne <= 1)
 		{
@@ -1474,6 +1473,7 @@ void FVirtualTextureSystem::GatherRequestsTask(const FGatherRequestsParameters&
 			}
 		}
 
+		const uint32 vDimensions = Space->GetDimensions();
 		check(vAddress >= AllocatedVT->GetVirtualAddress());
 
 		for (uint32 ProducerIndex = 0u; ProducerIndex < NumUniqueProducers; ++ProducerIndex)
@@ -1506,6 +1506,7 @@ void FVirtualTextureSystem::GatherRequestsTask(const FGatherRequestsParameters&
 			// Wrap Local_vAddress for the given producer
 			// For square textures, this could simply be (Local_vAddress % NumTilesInMip), but that doesn't work for non-square
 			// Possible there is a more clever approach to take here
+			uint32 Wrapped_vAddress = vAddress;
 			uint32 Local_vAddress;
 			{
 				const uint32 MipScaleFactor = (1u << Local_vLevel);
@@ -1513,11 +1514,21 @@ void FVirtualTextureSystem::GatherRequestsTask(const FGatherRequestsParameters&
 				const uint32 ProducerMipHeightInPages = FMath::DivideAndRoundUp(Producer->GetHeightInTiles(), MipScaleFactor);
 				uint32 Local_vPageX = (vPageX - AllocatedVT->GetVirtualPageX()) >> Mapping_vLevel;
 				uint32 Local_vTileY = (vPageY - AllocatedVT->GetVirtualPageY()) >> Mapping_vLevel;
-				Local_vPageX %= ProducerMipWidthInPages;
-				Local_vTileY %= ProducerMipHeightInPages;
-				Local_vAddress = FMath::MortonCode2(Local_vPageX) | (FMath::MortonCode2(Local_vTileY) << 1);
+				if (Local_vPageX >= ProducerMipWidthInPages || Local_vTileY >= ProducerMipHeightInPages)
+				{
+					Local_vPageX %= ProducerMipWidthInPages;
+					Local_vTileY %= ProducerMipHeightInPages;
+					Local_vAddress = FMath::MortonCode2(Local_vPageX) | (FMath::MortonCode2(Local_vTileY) << 1);
+					Wrapped_vAddress = (Local_vAddress << (Mapping_vLevel * vDimensions)) + AllocatedVT->GetVirtualAddress();
+				}
+				else
+				{
+					Local_vAddress = (vAddress - AllocatedVT->GetVirtualAddress()) >> (Mapping_vLevel * vDimensions);
+					// validate the wrapped calculations are nops when not wrapping
+					check(Local_vAddress == (FMath::MortonCode2(Local_vPageX) | (FMath::MortonCode2(Local_vTileY) << 1)));
+					check(vAddress == (Local_vAddress << (Mapping_vLevel * vDimensions)) + AllocatedVT->GetVirtualAddress());
+				}
 			}
-			vAddress = (Local_vAddress << (Mapping_vLevel * vDimensions)) + AllocatedVT->GetVirtualAddress();
 
 			const uint32 LocalMipBias = Producer->GetVirtualTexture()->GetLocalMipBias(Local_vLevel, Local_vAddress);
 			if (LocalMipBias > 0u)
@@ -1553,7 +1564,7 @@ void FVirtualTextureSystem::GatherRequestsTask(const FGatherRequestsParameters&
 					ensure(Allocated_vLevel <= AllocatedVT->GetMaxLevel());
 
 					const uint32 AllocatedMapping_vLevel = FMath::Max(Allocated_vLevel, ProducerMipBias);
-					const uint32 Allocated_vAddress = vAddress & (0xffffffff << (Allocated_vLevel * vDimensions));
+					const uint32 Allocated_vAddress = Wrapped_vAddress & (0xffffffff << (Allocated_vLevel * vDimensions));
 
 					AddPageUpdate(PageUpdateBuffers, PageUpdateFlushCount, PhysicalSpace->GetID(), pAddress);
 
@@ -1583,6 +1594,7 @@ void FVirtualTextureSystem::GatherRequestsTask(const FGatherRequestsParameters&
 								// map the page now if it wasn't already mapped
 								RequestList->AddDirectMappingRequest(Space->GetID(), PhysicalSpace->GetID(), PageTableLayerIndex, Allocated_vLevel, Allocated_vAddress, AllocatedMapping_vLevel, pAddress);
 							}
+
 							++NumMappedPages;
 						}
 					}
@@ -1648,7 +1660,7 @@ void FVirtualTextureSystem::GatherRequestsTask(const FGatherRequestsParameters&
 						const uint32 Prefetch_vLevel = PrefetchLocal_vLevel + ProducerMipBias;
 						ensure(Prefetch_vLevel <= AllocatedVT->GetMaxLevel());
 						const uint32 PrefetchMapping_vLevel = FMath::Max(Prefetch_vLevel, ProducerMipBias);
-						const uint32 Prefetch_vAddress = vAddress & (0xffffffff << (Prefetch_vLevel * vDimensions));
+						const uint32 Prefetch_vAddress = Wrapped_vAddress & (0xffffffff << (Prefetch_vLevel * vDimensions));
 						for (uint32 LoadLayerIndex = 0u; LoadLayerIndex < NumPageTableLayersToLoad; ++LoadLayerIndex)
 						{
 							const uint32 LayerIndex = PageTableLayersToLoad[LoadLayerIndex];
@@ -1678,7 +1690,7 @@ void FVirtualTextureSystem::GatherRequestsTask(const FGatherRequestsParameters&
 							const uint32 ProducerPhysicalGroupIndex = AllocatedVT->GetProducerPhysicalGroupIndexForPageTableLayer(LayerIndex);
 							if (GroupMaskToLoad & (1u << ProducerPhysicalGroupIndex))
 							{
-								RequestList->AddMappingRequest(LoadRequestIndex, ProducerPhysicalGroupIndex, ID, LayerIndex, vAddress, vLevel, Mapping_vLevel);
+								RequestList->AddMappingRequest(LoadRequestIndex, ProducerPhysicalGroupIndex, ID, LayerIndex, Wrapped_vAddress, vLevel, Mapping_vLevel);
 							}
 						}
 					}
