Restore a VT check that was previously removed...check is actually valid, was removed due to other bugs

diff --git a/Engine/Source/Runtime/Renderer/Private/VT/VirtualTextureSystem.cpp b/Engine/Source/Runtime/Renderer/Private/VT/VirtualTextureSystem.cpp
--- a/Engine/Source/Runtime/Renderer/Private/VT/VirtualTextureSystem.cpp
+++ b/Engine/Source/Runtime/Renderer/Private/VT/VirtualTextureSystem.cpp
@@ -1576,19 +1576,31 @@ void FVirtualTextureSystem::GatherRequestsTask(const FGatherRequestsParameters&
 							AllocatedVT->GetProducerIndexForPageTableLayer(PageTableLayerIndex) == ProducerIndex)
 						{
 							bool bPageWasMapped = false;
-							// if we found a lower resolution tile than was requested, it may have already been mapped, check for that first
-							const FTexturePageMap& PageMap = Space->GetPageMapForPageTableLayer(PageTableLayerIndex);
-							const FPhysicalSpaceIDAndAddress PrevPhysicalSpaceIDAndAddress = PageMap.FindPagePhysicalSpaceIDAndAddress(Allocated_vLevel, Allocated_vAddress);
-							if (PrevPhysicalSpaceIDAndAddress.Packed != ~0u)
+							if (Allocated_vLevel != vLevel)
 							{
-								// if this address was previously mapped, ensure that it was mapped by the same physical space
-								ensure(PrevPhysicalSpaceIDAndAddress.PhysicalSpaceID == PhysicalSpace->GetID());
-								// either it wasn't mapped, or it's mapped to the current physical address...
-								// otherwise that means that the same local tile is mapped to two separate physical addresses, which is an error
-								ensure(PrevPhysicalSpaceIDAndAddress.pAddress == pAddress);
-								bPageWasMapped = true;
+								// if we found a lower resolution tile than was requested, it may have already been mapped, check for that first
+								// don't need to check this if the allocated page is at the level that was requested...if that was already mapped we wouldn't have gotten this far
+								const FTexturePageMap& PageMap = Space->GetPageMapForPageTableLayer(PageTableLayerIndex);
+								const FPhysicalSpaceIDAndAddress PrevPhysicalSpaceIDAndAddress = PageMap.FindPagePhysicalSpaceIDAndAddress(Allocated_vLevel, Allocated_vAddress);
+								if (PrevPhysicalSpaceIDAndAddress.Packed != ~0u)
+								{
+									// if this address was previously mapped, ensure that it was mapped by the same physical space
+									ensure(PrevPhysicalSpaceIDAndAddress.PhysicalSpaceID == PhysicalSpace->GetID());
+									// either it wasn't mapped, or it's mapped to the current physical address...
+									// otherwise that means that the same local tile is mapped to two separate physical addresses, which is an error
+									ensure(PrevPhysicalSpaceIDAndAddress.pAddress == pAddress);
+									bPageWasMapped = true;
+								}
 							}
-
+#if DO_GUARD_SLOW
+							else
+							{
+								// verify our assumption that the page shouldn't be mapped yet
+								const FTexturePageMap& PageMap = Space->GetPageMapForPageTableLayer(PageTableLayerIndex);
+								const FPhysicalSpaceIDAndAddress PrevPhysicalSpaceIDAndAddress = PageMap.FindPagePhysicalSpaceIDAndAddress(Allocated_vLevel, Allocated_vAddress);
+								checkSlow(PrevPhysicalSpaceIDAndAddress.Packed == ~0u);
+							}
+#endif // DO_GUARD_SLOW
 							if (!bPageWasMapped)
 							{
 								// map the page now if it wasn't already mapped
