Allocate starta scene global uniform buffer only once

diff --git a/Engine/Source/Runtime/Renderer/Private/BasePassRendering.h b/Engine/Source/Runtime/Renderer/Private/BasePassRendering.h
--- a/Engine/Source/Runtime/Renderer/Private/BasePassRendering.h
+++ b/Engine/Source/Runtime/Renderer/Private/BasePassRendering.h
@@ -36,8 +36,6 @@ struct FSceneWithoutWaterTextures;
 class FViewInfo;
 class UMaterialExpressionSingleLayerWaterMaterialOutput;
 
-struct FStrataData;
-
 /** Whether to allow the indirect lighting cache to be applied to dynamic objects. */
 extern int32 GIndirectLightingCache;
 
diff --git a/Engine/Source/Runtime/Renderer/Private/DeferredShadingRenderer.h b/Engine/Source/Runtime/Renderer/Private/DeferredShadingRenderer.h
--- a/Engine/Source/Runtime/Renderer/Private/DeferredShadingRenderer.h
+++ b/Engine/Source/Runtime/Renderer/Private/DeferredShadingRenderer.h
@@ -41,7 +41,6 @@ struct FHeightFogRenderingParameters;
 struct FRayTracingReflectionOptions;
 struct FHairStrandsTransmittanceMaskData;
 struct FHairStrandsRenderingData;
-struct FStrataData;
 
 /**   
  * Data for rendering meshes into Lumen Lighting Cards.
diff --git a/Engine/Source/Runtime/Renderer/Private/ScenePrivate.h b/Engine/Source/Runtime/Renderer/Private/ScenePrivate.h
--- a/Engine/Source/Runtime/Renderer/Private/ScenePrivate.h
+++ b/Engine/Source/Runtime/Renderer/Private/ScenePrivate.h
@@ -2706,7 +2706,7 @@ public:
 	TSparseArray<FRuntimeVirtualTextureSceneProxy*> RuntimeVirtualTextures;
 
 	/** Strata data shared between all views. */
-	FStrataData StrataData;
+	FStrataSceneData StrataSceneData;
 	/** Mask used to determine whether primitives that draw to a runtime virtual texture should also be drawn in the main pass. */
 	uint8 RuntimeVirtualTexturePrimitiveHideMaskEditor;
 	uint8 RuntimeVirtualTexturePrimitiveHideMaskGame;
diff --git a/Engine/Source/Runtime/Renderer/Private/SceneRendering.cpp b/Engine/Source/Runtime/Renderer/Private/SceneRendering.cpp
--- a/Engine/Source/Runtime/Renderer/Private/SceneRendering.cpp
+++ b/Engine/Source/Runtime/Renderer/Private/SceneRendering.cpp
@@ -930,7 +930,7 @@ void FViewInfo::Init()
 	NumVisibleDynamicPrimitives = 0;
 	NumVisibleDynamicEditorPrimitives = 0;
 
-	StrataData = nullptr;
+	StrataSceneData = nullptr;
 }
 
 FViewInfo::~FViewInfo()
diff --git a/Engine/Source/Runtime/Renderer/Private/SceneRendering.h b/Engine/Source/Runtime/Renderer/Private/SceneRendering.h
--- a/Engine/Source/Runtime/Renderer/Private/SceneRendering.h
+++ b/Engine/Source/Runtime/Renderer/Private/SceneRendering.h
@@ -52,7 +52,7 @@ struct FCloudRenderContext;
 struct FSceneWithoutWaterTextures;
 struct FHairStrandsVisibilityViews;
 struct FSortedLightSetSceneInfo;
-struct FStrataData;
+struct FStrataSceneData;
 
 DECLARE_STATS_GROUP(TEXT("Command List Markers"), STATGROUP_CommandListMarkers, STATCAT_Advanced);
 
@@ -1180,7 +1180,8 @@ public:
 	// Used by mobile renderer to determine whether static meshes will be rendered with CSM shaders or not.
 	FMobileCSMVisibilityInfo MobileCSMVisibilityInfo;
 
-	FStrataData* StrataData;
+	FStrataSceneData* StrataSceneData;
+
 	//Spotlight shadow info for mobile.
 	FMobileMovableSpotLightsShadowInfo MobileMovableSpotLightsShadowInfo;
 
diff --git a/Engine/Source/Runtime/Renderer/Private/Strata/Strata.cpp b/Engine/Source/Runtime/Renderer/Private/Strata/Strata.cpp
--- a/Engine/Source/Runtime/Renderer/Private/Strata/Strata.cpp
+++ b/Engine/Source/Runtime/Renderer/Private/Strata/Strata.cpp
@@ -34,7 +34,7 @@ bool IsStrataEnabled()
 
 void InitialiseStrataFrameSceneData(FSceneRenderer& SceneRenderer, FRDGBuilder& GraphBuilder)
 {
-	FStrataData& StrataData = SceneRenderer.Scene->StrataData;
+	FStrataSceneData& StrataSceneData = SceneRenderer.Scene->StrataSceneData;
 
 	uint32 ResolutionX = 1;
 	uint32 ResolutionY = 1;
@@ -50,44 +50,47 @@ void InitialiseStrataFrameSceneData(FSceneRenderer& SceneRenderer, FRDGBuilder&
 		// Previous GBuffer when complete was 28bytes
 		// check out Strata.ush to see how this is computed
 		const uint32 MaterialConservativeByteCountPerPixel = 129u;
-		StrataData.MaxBytesPerPixel = FMath::DivideAndRoundUp(MaterialConservativeByteCountPerPixel, 4u) * 4u;
+		StrataSceneData.MaxBytesPerPixel = FMath::DivideAndRoundUp(MaterialConservativeByteCountPerPixel, 4u) * 4u;
 	}
 	else
 	{
-		StrataData.MaxBytesPerPixel = 4u;
+		StrataSceneData.MaxBytesPerPixel = 4u;
 	}
 
 	FRDGTextureRef MaterialLobesTexture = GraphBuilder.CreateTexture(
 		FRDGTextureDesc::Create2D(FIntPoint(ResolutionX, ResolutionY), PF_R16F, FClearValueBinding::None,
 		TexCreate_RenderTargetable | TexCreate_ShaderResource | TexCreate_UAV), TEXT("MaterialLobesTexture"));
 	AddClearRenderTargetPass(GraphBuilder, MaterialLobesTexture, FLinearColor::Black);
-	ConvertToExternalTexture(GraphBuilder, MaterialLobesTexture, StrataData.MaterialLobesTexture);
+	ConvertToExternalTexture(GraphBuilder, MaterialLobesTexture, StrataSceneData.MaterialLobesTexture);
 
-	const uint32 DesiredBufferSize = FMath::Max(4u, ResolutionX * ResolutionY * StrataData.MaxBytesPerPixel);
-	if (StrataData.MaterialLobesBuffer.NumBytes < DesiredBufferSize)
+	const uint32 DesiredBufferSize = FMath::Max(4u, ResolutionX * ResolutionY * StrataSceneData.MaxBytesPerPixel);
+	if (StrataSceneData.MaterialLobesBuffer.NumBytes < DesiredBufferSize)
 	{
-		if (StrataData.MaterialLobesBuffer.NumBytes > 0)
+		if (StrataSceneData.MaterialLobesBuffer.NumBytes > 0)
 		{
-			StrataData.MaterialLobesBuffer.Release();
+			StrataSceneData.MaterialLobesBuffer.Release();
 		}
-		StrataData.MaterialLobesBuffer.Initialize(DesiredBufferSize, BUF_Static, TEXT("MaterialLobesBuffer"));
+		StrataSceneData.MaterialLobesBuffer.Initialize(DesiredBufferSize, BUF_Static, TEXT("MaterialLobesBuffer"));
 	}
 
 	// Set reference to the Strata data from each view
 	for (int32 ViewIndex = 0; ViewIndex < SceneRenderer.Views.Num(); ViewIndex++)
 	{
 		FViewInfo& View = SceneRenderer.Views[ViewIndex];
-		View.StrataData = &SceneRenderer.Scene->StrataData;
+		View.StrataSceneData = &SceneRenderer.Scene->StrataSceneData;
 	}
+
+	// Always reset the strata uniform buffer
+	SceneRenderer.Scene->StrataSceneData.StrataGlobalUniformParameters.SafeRelease();
 }
 
 void BindStrataBasePassUniformParameters(const FViewInfo& View, FStrataOpaquePassUniformParameters& OutStrataUniformParameters)
 {
-	if (View.StrataData)
+	if (View.StrataSceneData)
 	{
-		OutStrataUniformParameters.MaxBytesPerPixel = View.StrataData->MaxBytesPerPixel;
-		OutStrataUniformParameters.MaterialLobesTextureUAV = View.StrataData->MaterialLobesTexture->GetRenderTargetItem().UAV;
-		OutStrataUniformParameters.MaterialLobesBufferUAV = View.StrataData->MaterialLobesBuffer.UAV;
+		OutStrataUniformParameters.MaxBytesPerPixel = View.StrataSceneData->MaxBytesPerPixel;
+		OutStrataUniformParameters.MaterialLobesTextureUAV = View.StrataSceneData->MaterialLobesTexture->GetRenderTargetItem().UAV;
+		OutStrataUniformParameters.MaterialLobesBufferUAV = View.StrataSceneData->MaterialLobesBuffer.UAV;
 	}
 	else
 	{
@@ -99,12 +102,19 @@ void BindStrataBasePassUniformParameters(const FViewInfo& View, FStrataOpaquePas
 
 TUniformBufferRef<FStrataGlobalUniformParameters> BindStrataGlobalUniformParameters(const FViewInfo& View)
 {
+	check(View.StrataSceneData);
+	if (View.StrataSceneData->StrataGlobalUniformParameters.IsValid())
+	{
+		return View.StrataSceneData->StrataGlobalUniformParameters;
+	}
+
+	// If the strata scene data has not been created this frame yet, create it.
 	FStrataGlobalUniformParameters StrataUniformParameters;
-	if (View.StrataData)
+	if (View.StrataSceneData)
 	{
-		StrataUniformParameters.MaxBytesPerPixel = View.StrataData->MaxBytesPerPixel;
-		StrataUniformParameters.MaterialLobesTexture = View.StrataData->MaterialLobesTexture->GetRenderTargetItem().ShaderResourceTexture;
-		StrataUniformParameters.MaterialLobesBuffer = View.StrataData->MaterialLobesBuffer.SRV;
+		StrataUniformParameters.MaxBytesPerPixel = View.StrataSceneData->MaxBytesPerPixel;
+		StrataUniformParameters.MaterialLobesTexture = View.StrataSceneData->MaterialLobesTexture->GetRenderTargetItem().ShaderResourceTexture;
+		StrataUniformParameters.MaterialLobesBuffer = View.StrataSceneData->MaterialLobesBuffer.SRV;
 	}
 	else
 	{
@@ -112,9 +122,9 @@ TUniformBufferRef<FStrataGlobalUniformParameters> BindStrataGlobalUniformParamet
 		StrataUniformParameters.MaterialLobesTexture = GSystemTextures.BlackDummy->GetRenderTargetItem().ShaderResourceTexture;
 		StrataUniformParameters.MaterialLobesBuffer = GEmptyVertexBufferWithUAV->ShaderResourceViewRHI;
 	}
+	View.StrataSceneData->StrataGlobalUniformParameters = CreateUniformBufferImmediate(StrataUniformParameters, UniformBuffer_SingleFrame);
 
-	// STRATA_TODO cache this on the view and use UniformBuffer_SingleFrame
-	return CreateUniformBufferImmediate(StrataUniformParameters, UniformBuffer_SingleDraw);
+	return View.StrataSceneData->StrataGlobalUniformParameters;
 }
 
 
diff --git a/Engine/Source/Runtime/Renderer/Private/Strata/Strata.h b/Engine/Source/Runtime/Renderer/Private/Strata/Strata.h
--- a/Engine/Source/Runtime/Renderer/Private/Strata/Strata.h
+++ b/Engine/Source/Runtime/Renderer/Private/Strata/Strata.h
@@ -29,13 +29,15 @@ END_GLOBAL_SHADER_PARAMETER_STRUCT()
 
 
 
-struct FStrataData
+struct FStrataSceneData
 {
 	uint32 MaxBytesPerPixel;
 	TRefCountPtr<IPooledRenderTarget> MaterialLobesTexture; // This should be a RDG resource when the refactoring gets in
 	FRWByteAddressBuffer MaterialLobesBuffer;				// This should be a RDG resource	"		"		"
 
-	FStrataData()
+	TUniformBufferRef<FStrataGlobalUniformParameters> StrataGlobalUniformParameters;
+
+	FStrataSceneData()
 	{
 	}
 };
