Fixed crash due to strata when strata is not enabled.

diff --git a/Engine/Source/Runtime/Renderer/Private/Strata/Strata.cpp b/Engine/Source/Runtime/Renderer/Private/Strata/Strata.cpp
--- a/Engine/Source/Runtime/Renderer/Private/Strata/Strata.cpp
+++ b/Engine/Source/Runtime/Renderer/Private/Strata/Strata.cpp
@@ -102,29 +102,30 @@ void BindStrataBasePassUniformParameters(const FViewInfo& View, FStrataOpaquePas
 
 TUniformBufferRef<FStrataGlobalUniformParameters> BindStrataGlobalUniformParameters(const FViewInfo& View)
 {
-	check(View.StrataSceneData);
-	if (View.StrataSceneData->StrataGlobalUniformParameters.IsValid())
-	{
-		return View.StrataSceneData->StrataGlobalUniformParameters;
-	}
-
 	// If the strata scene data has not been created this frame yet, create it.
 	FStrataGlobalUniformParameters StrataUniformParameters;
 	if (View.StrataSceneData)
 	{
+		if (View.StrataSceneData->StrataGlobalUniformParameters.IsValid())
+		{
+			return View.StrataSceneData->StrataGlobalUniformParameters;
+		}
+
 		StrataUniformParameters.MaxBytesPerPixel = View.StrataSceneData->MaxBytesPerPixel;
 		StrataUniformParameters.MaterialLobesTexture = View.StrataSceneData->MaterialLobesTexture->GetRenderTargetItem().ShaderResourceTexture;
 		StrataUniformParameters.MaterialLobesBuffer = View.StrataSceneData->MaterialLobesBuffer.SRV;
+
+		View.StrataSceneData->StrataGlobalUniformParameters = CreateUniformBufferImmediate(StrataUniformParameters, UniformBuffer_SingleFrame);
+		return View.StrataSceneData->StrataGlobalUniformParameters;
 	}
 	else
 	{
+		// Create each time. This path will go away when Strata is always enabled anyway.
 		StrataUniformParameters.MaxBytesPerPixel = 0;
 		StrataUniformParameters.MaterialLobesTexture = GSystemTextures.BlackDummy->GetRenderTargetItem().ShaderResourceTexture;
 		StrataUniformParameters.MaterialLobesBuffer = GEmptyVertexBufferWithUAV->ShaderResourceViewRHI;
+		return CreateUniformBufferImmediate(StrataUniformParameters, UniformBuffer_SingleDraw);
 	}
-	View.StrataSceneData->StrataGlobalUniformParameters = CreateUniformBufferImmediate(StrataUniformParameters, UniformBuffer_SingleFrame);
-
-	return View.StrataSceneData->StrataGlobalUniformParameters;
 }
 
 
