Refactored translucent lighting textures out of the scene textures singleton. Moved to RDG resources.

diff --git a/Engine/Source/Runtime/Renderer/Private/BasePassRendering.h b/Engine/Source/Runtime/Renderer/Private/BasePassRendering.h
--- a/Engine/Source/Runtime/Renderer/Private/BasePassRendering.h
+++ b/Engine/Source/Runtime/Renderer/Private/BasePassRendering.h
@@ -32,6 +32,7 @@ class FScene;
 template<typename TBufferStruct> class TUniformBufferRef;
 
 struct FSceneWithoutWaterTextures;
+struct FTranslucentVolumeLightingTextures;
 
 class FViewInfo;
 class UMaterialExpressionSingleLayerWaterMaterialOutput;
@@ -148,6 +149,7 @@ extern TRDGUniformBufferRef<FOpaqueBasePassUniformParameters> CreateOpaqueBasePa
 extern TRDGUniformBufferRef<FTranslucentBasePassUniformParameters> CreateTranslucentBasePassUniformBuffer(
 	FRDGBuilder& GraphBuilder,
 	const FViewInfo& View,
+	const FTranslucentVolumeLightingTextures* TranslucentVolumeLightingTextures,
 	FRDGTextureRef SceneColorCopy,
 	ESceneTextureSetupMode SceneTextureSetupMode,
 	const int32 ViewIndex);
diff --git a/Engine/Source/Runtime/Renderer/Private/DeferredShadingRenderer.cpp b/Engine/Source/Runtime/Renderer/Private/DeferredShadingRenderer.cpp
--- a/Engine/Source/Runtime/Renderer/Private/DeferredShadingRenderer.cpp
+++ b/Engine/Source/Runtime/Renderer/Private/DeferredShadingRenderer.cpp
@@ -1474,7 +1474,7 @@ void FDeferredShadingSceneRenderer::Render(FRHICommandListImmediate& RHICmdList)
 		GSystemTextures.InitializeTextures(RHICmdList, FeatureLevel);
 
 		// Allocate the maximum scene render target space for the current view family.
-		SceneContext.Allocate(RHICmdList, this);
+		SceneContext.Allocate(GraphBuilder, this);
 	}
 
 	const bool bUseVirtualTexturing = UseVirtualTexturing(FeatureLevel);
@@ -2122,30 +2122,11 @@ void FDeferredShadingSceneRenderer::Render(FRHICommandListImmediate& RHICmdList)
 			SceneContext.bScreenSpaceAOIsValid);
 	}
 
-	if (bRenderDeferredLighting)
-	{
-		bool bShouldAllocateDeferredShadingPathRenderTargets = false;
-		const TCHAR* Str = SceneContext.ScreenSpaceAO ? TEXT("Allocated") : TEXT("Unallocated"); // ScreenSpaceAO is determining factor of detecting render target allocation
-		for (int32 Index = 0; Index < (NumTranslucentVolumeRenderTargetSets * Views.Num()); ++Index)
-		{
-			if (!SceneContext.TranslucencyLightingVolumeAmbient[Index] || !SceneContext.TranslucencyLightingVolumeDirectional[Index])
-			{
-				ensureMsgf(SceneContext.TranslucencyLightingVolumeAmbient[Index], TEXT("%s%d is unallocated, Deferred Render Targets would be detected as: %s"), "TranslucencyLightingVolumeAmbient", Index, Str);
-				ensureMsgf(SceneContext.TranslucencyLightingVolumeDirectional[Index], TEXT("%s%d is unallocated, Deferred Render Targets would be detected as: %s"), "TranslucencyLightingVolumeDirectional", Index, Str);
-				bShouldAllocateDeferredShadingPathRenderTargets = true;
-				break;
-			}
-		}
+	FTranslucentVolumeLightingTextures TranslucentLightingVolumeTextures;
 
-		if(bShouldAllocateDeferredShadingPathRenderTargets)
-		{
-			SceneContext.AllocateDeferredShadingPathRenderTargets(RHICmdList);
-		}
-		
-		if (GbEnableAsyncComputeTranslucencyLightingVolumeClear && GSupportsEfficientAsyncCompute)
-		{
-			ClearTranslucentVolumeLighting(GraphBuilder, ERDGPassFlags::AsyncCompute);
-		}
+	if (bRenderDeferredLighting && GbEnableAsyncComputeTranslucencyLightingVolumeClear && GSupportsEfficientAsyncCompute)
+	{
+		InitTranslucentVolumeLighting(GraphBuilder, ERDGPassFlags::AsyncCompute, TranslucentLightingVolumeTextures);
 	}
 
 	SceneContext.AllocSceneColor(RHICmdList);
@@ -2425,7 +2406,7 @@ void FDeferredShadingSceneRenderer::Render(FRHICommandListImmediate& RHICmdList)
 		// Clear the translucent lighting volumes before we accumulate
 		if ((GbEnableAsyncComputeTranslucencyLightingVolumeClear && GSupportsEfficientAsyncCompute) == false)
 		{
-			ClearTranslucentVolumeLighting(GraphBuilder, ERDGPassFlags::Compute);
+			InitTranslucentVolumeLighting(GraphBuilder, ERDGPassFlags::Compute, TranslucentLightingVolumeTextures);
 		}
 
 #if RHI_RAYTRACING
@@ -2436,7 +2417,15 @@ void FDeferredShadingSceneRenderer::Render(FRHICommandListImmediate& RHICmdList)
 #endif
 
 		AddSetCurrentStatPass(GraphBuilder, GET_STATID(STAT_CLM_Lighting));
-		SceneColorTexture = FRDGTextureMSAA(RenderLights(GraphBuilder, SceneTextures, SceneColorTexture.Target, SceneDepthTexture.Target, LightingChannelsTexture, SortedLightSet, HairDatas));
+		SceneColorTexture = FRDGTextureMSAA(RenderLights(
+			GraphBuilder,
+			SceneTextures,
+			TranslucentLightingVolumeTextures,
+			SceneColorTexture.Target,
+			SceneDepthTexture.Target,
+			LightingChannelsTexture,
+			SortedLightSet,
+			HairDatas));
 		AddSetCurrentStatPass(GraphBuilder, GET_STATID(STAT_CLM_AfterLighting));
 		AddServiceLocalQueuePass(GraphBuilder);
 
@@ -2444,7 +2433,7 @@ void FDeferredShadingSceneRenderer::Render(FRHICommandListImmediate& RHICmdList)
 		{
 			const FViewInfo& View = Views[ViewIndex];
 			RDG_GPU_MASK_SCOPE(GraphBuilder, View.GPUMask);
-			InjectAmbientCubemapTranslucentVolumeLighting(GraphBuilder, Views[ViewIndex], ViewIndex);
+			InjectAmbientCubemapTranslucentVolumeLighting(GraphBuilder, TranslucentLightingVolumeTextures, Views[ViewIndex], ViewIndex);
 		}
 		AddServiceLocalQueuePass(GraphBuilder);
 
@@ -2454,7 +2443,7 @@ void FDeferredShadingSceneRenderer::Render(FRHICommandListImmediate& RHICmdList)
 			RDG_GPU_MASK_SCOPE(GraphBuilder, View.GPUMask);
 
 			// Filter the translucency lighting volume now that it is complete
-			FilterTranslucentVolumeLighting(GraphBuilder, View, ViewIndex);
+			FilterTranslucentVolumeLighting(GraphBuilder, TranslucentLightingVolumeTextures, View, ViewIndex);
 		}
 		AddServiceLocalQueuePass(GraphBuilder);
 
@@ -2513,7 +2502,7 @@ void FDeferredShadingSceneRenderer::Render(FRHICommandListImmediate& RHICmdList)
 			RDG_CSV_STAT_EXCLUSIVE_SCOPE(GraphBuilder, RenderTranslucency);
 			SCOPE_CYCLE_COUNTER(STAT_TranslucencyDrawTime);
 			AddSetCurrentStatPass(GraphBuilder, GET_STATID(STAT_CLM_Translucency));
-			RenderTranslucency(GraphBuilder, SceneColorTexture, SceneDepthTexture, nullptr, ETranslucencyView::UnderWater);
+			RenderTranslucency(GraphBuilder, TranslucentLightingVolumeTextures, SceneColorTexture, SceneDepthTexture, nullptr, ETranslucencyView::UnderWater);
 			EnumRemoveFlags(TranslucencyViewsToRender, ETranslucencyView::UnderWater);
 		}
 
@@ -2632,7 +2621,7 @@ void FDeferredShadingSceneRenderer::Render(FRHICommandListImmediate& RHICmdList)
 
 		// Render all remaining translucency views.
 		AddSetCurrentStatPass(GraphBuilder, GET_STATID(STAT_CLM_Translucency));
-		RenderTranslucency(GraphBuilder, SceneColorTexture, SceneDepthTexture, &SeparateTranslucencyTextures, TranslucencyViewsToRender);
+		RenderTranslucency(GraphBuilder, TranslucentLightingVolumeTextures, SceneColorTexture, SceneDepthTexture, &SeparateTranslucencyTextures, TranslucencyViewsToRender);
 		AddServiceLocalQueuePass(GraphBuilder);
 		TranslucencyViewsToRender = ETranslucencyView::None;
 
diff --git a/Engine/Source/Runtime/Renderer/Private/DeferredShadingRenderer.h b/Engine/Source/Runtime/Renderer/Private/DeferredShadingRenderer.h
--- a/Engine/Source/Runtime/Renderer/Private/DeferredShadingRenderer.h
+++ b/Engine/Source/Runtime/Renderer/Private/DeferredShadingRenderer.h
@@ -42,6 +42,8 @@ struct FRayTracingReflectionOptions;
 struct FHairStrandsTransmittanceMaskData;
 struct FHairStrandsRenderingData;
 
+struct FTranslucentVolumeLightingTextures;
+
 /**   
  * Data for rendering meshes into Lumen Lighting Cards.
  */
@@ -162,8 +164,6 @@ class FDeferredShadingSceneRenderer : public FSceneRenderer
 public:
 	/** Defines which objects we want to render in the EarlyZPass. */
 	FDepthPassInfo DepthPass;
-	
-	const FRHITransition* TranslucencyLightingVolumeClearEndTransition = nullptr;
 
 	FLumenCardRenderer LumenCardRenderer;
 
@@ -588,6 +588,7 @@ private:
 	FRDGTextureRef RenderLights(
 		FRDGBuilder& GraphBuilder,
 		TRDGUniformBufferRef<FSceneTextureUniformParameters> SceneTexturesUniformBuffer,
+		const FTranslucentVolumeLightingTextures& TranslucentVolumeLightingTextures,
 		FRDGTextureRef SceneColorTexture,
 		FRDGTextureRef SceneDepthTexture,
 		FRDGTextureRef LightingChannelsTexture,
@@ -611,6 +612,7 @@ private:
 	/** Renders the scene's translucency passes. */
 	void RenderTranslucency(
 		FRDGBuilder& GraphBuilder,
+		const FTranslucentVolumeLightingTextures& TranslucentVolumeLightingTextures,
 		FRDGTextureMSAA SceneColorTexture,
 		FRDGTextureMSAA SceneDepthTexture,
 		FSeparateTranslucencyTextures* OutSeparateTranslucencyTextures,
@@ -619,6 +621,7 @@ private:
 	/** Renders the scene's translucency given a specific pass. */
 	void RenderTranslucencyInner(
 		FRDGBuilder& GraphBuilder,
+		const FTranslucentVolumeLightingTextures& TranslucentVolumeLightingTextures,
 		FRDGTextureMSAA SceneColorTexture,
 		FRDGTextureMSAA SceneDepthTexture,
 		FSeparateTranslucencyTextures* OutSeparateTranslucencyTextures,
@@ -708,6 +711,7 @@ private:
 	void RenderDeferredShadowProjections(
 		FRDGBuilder& GraphBuilder,
 		TRDGUniformBufferRef<FSceneTextureUniformParameters> SceneTexturesUniformBuffer,
+		const FTranslucentVolumeLightingTextures& TranslucentVolumeLightingTextures,
 		const FLightSceneInfo* LightSceneInfo,
 		FRDGTextureRef ScreenShadowMaskTexture,
 		FRDGTextureRef ScreenShadowMaskSubPixelTexture,
@@ -801,23 +805,50 @@ private:
 		TRDGUniformBufferRef<FSceneTextureUniformParameters> SceneTexturesUniformBuffer,
 		const FSimpleLightArray& SimpleLights);
 
-	/** Clears the translucency lighting volumes. Can be compute or async compute. */
-	void ClearTranslucentVolumeLighting(FRDGBuilder& GraphBuilder, ERDGPassFlags PassFlags);
+	/** Initializes the translucency lighting volumes. Can be compute or async compute. */
+	void InitTranslucentVolumeLighting(
+		FRDGBuilder& GraphBuilder,
+		ERDGPassFlags PassFlags,
+		FTranslucentVolumeLightingTextures& Textures);
 
 	/** Add AmbientCubemap to the lighting volumes. */
-	void InjectAmbientCubemapTranslucentVolumeLighting(FRDGBuilder& GraphBuilder, const FViewInfo& View, int32 ViewIndex);
+	void InjectAmbientCubemapTranslucentVolumeLighting(
+		FRDGBuilder& GraphBuilder,
+		const FTranslucentVolumeLightingTextures& Textures,
+		const FViewInfo& View,
+		int32 ViewIndex);
 
 	/** Accumulates direct lighting for the given light.  InProjectedShadowInfo can be NULL in which case the light will be unshadowed. */
-	void InjectTranslucentVolumeLighting(FRDGBuilder& GraphBuilder, const FLightSceneInfo& LightSceneInfo, const FProjectedShadowInfo* InProjectedShadowInfo, const FViewInfo& View, int32 ViewIndex);
+	void InjectTranslucentVolumeLighting(
+		FRDGBuilder& GraphBuilder,
+		const FTranslucentVolumeLightingTextures& Textures,
+		const FLightSceneInfo& LightSceneInfo,
+		const FProjectedShadowInfo* InProjectedShadowInfo,
+		const FViewInfo& View,
+		int32 ViewIndex);
 
 	/** Accumulates direct lighting for an array of unshadowed lights. */
-	void InjectTranslucentVolumeLightingArray(FRDGBuilder& GraphBuilder, const TArray<FSortedLightSceneInfo, SceneRenderingAllocator>& SortedLights, int32 FirstLightIndex, int32 LightsEndIndex);
+	void InjectTranslucentVolumeLightingArray(
+		FRDGBuilder& GraphBuilder,
+		const FTranslucentVolumeLightingTextures& Textures,
+		const TArray<FSortedLightSceneInfo, SceneRenderingAllocator>& SortedLights,
+		int32 FirstLightIndex,
+		int32 LightsEndIndex);
 
 	/** Accumulates direct lighting for simple lights. */
-	void InjectSimpleTranslucentVolumeLightingArray(FRDGBuilder& GraphBuilder, const FSimpleLightArray& SimpleLights, const FViewInfo& View, const int32 ViewIndex);
+	void InjectSimpleTranslucentVolumeLightingArray(
+		FRDGBuilder& GraphBuilder,
+		const FTranslucentVolumeLightingTextures& Textures,
+		const FSimpleLightArray& SimpleLights,
+		const FViewInfo& View,
+		const int32 ViewIndex);
 
 	/** Filters the translucency lighting volumes to reduce aliasing. */
-	void FilterTranslucentVolumeLighting(FRDGBuilder& GraphBuilder, const FViewInfo& View, const int32 ViewIndex);
+	void FilterTranslucentVolumeLighting(
+		FRDGBuilder& GraphBuilder,
+		FTranslucentVolumeLightingTextures& Textures,
+		const FViewInfo& View,
+		const int32 ViewIndex);
 
 	bool ShouldRenderVolumetricFog() const;
 
diff --git a/Engine/Source/Runtime/Renderer/Private/LightRendering.cpp b/Engine/Source/Runtime/Renderer/Private/LightRendering.cpp
--- a/Engine/Source/Runtime/Renderer/Private/LightRendering.cpp
+++ b/Engine/Source/Runtime/Renderer/Private/LightRendering.cpp
@@ -1192,6 +1192,7 @@ void GetRenderLightParameters(
 FRDGTextureRef FDeferredShadingSceneRenderer::RenderLights(
 	FRDGBuilder& GraphBuilder,
 	TRDGUniformBufferRef<FSceneTextureUniformParameters> SceneTexturesUniformBuffer,
+	const FTranslucentVolumeLightingTextures& TranslucentVolumeLightingTextures,
 	FRDGTextureRef SceneColorTexture,
 	FRDGTextureRef SceneDepthTexture,
 	FRDGTextureRef LightingChannelsTexture,
@@ -1218,17 +1219,6 @@ FRDGTextureRef FDeferredShadingSceneRenderer::RenderLights(
 	{
 		RDG_EVENT_SCOPE(GraphBuilder, "DirectLighting");
 
-		if (GbEnableAsyncComputeTranslucencyLightingVolumeClear && GSupportsEfficientAsyncCompute)
-		{
-			AddPass(GraphBuilder, [this](FRHICommandList& RHICmdList)
-			{
-				//Gfx pipe must wait for the async compute clear of the translucency volume clear.
-				check(TranslucencyLightingVolumeClearEndTransition);
-				RHICmdList.EndTransition(TranslucencyLightingVolumeClearEndTransition);
-				TranslucencyLightingVolumeClearEndTransition = nullptr;
-			});
-		}
-
 		if(ViewFamily.EngineShowFlags.DirectLighting)
 		{
 			RDG_EVENT_SCOPE(GraphBuilder, "NonShadowedLights");
@@ -1327,7 +1317,7 @@ FRDGTextureRef FDeferredShadingSceneRenderer::RenderLights(
 				if (AttenuationLightStart)
 				{
 					// Inject non-shadowed, non-simple, non-light function lights in to the volume.
-					InjectTranslucentVolumeLightingArray(GraphBuilder, SortedLights, SimpleLightsEnd, AttenuationLightStart);
+					InjectTranslucentVolumeLightingArray(GraphBuilder, TranslucentVolumeLightingTextures, SortedLights, SimpleLightsEnd, AttenuationLightStart);
 				}
 
 				if (SimpleLights.InstanceData.Num() > 0)
@@ -1346,7 +1336,7 @@ FRDGTextureRef FDeferredShadingSceneRenderer::RenderLights(
 							FViewInfo& View = Views[ViewIndex];
 							RDG_GPU_MASK_SCOPE(GraphBuilder, View.GPUMask);
 							RDG_EVENT_SCOPE(GraphBuilder, "InjectSimpleLightsTranslucentLighting");
-							InjectSimpleTranslucentVolumeLightingArray(GraphBuilder, SimpleLightArray, View, ViewIndex);
+							InjectSimpleTranslucentVolumeLightingArray(GraphBuilder, TranslucentVolumeLightingTextures, SimpleLightArray, View, ViewIndex);
 						}
 					}
 				}
@@ -1878,7 +1868,7 @@ FRDGTextureRef FDeferredShadingSceneRenderer::RenderLights(
 							ClearShadowMask(ScreenShadowMaskSubPixelTexture);
 						}
 
-						RenderDeferredShadowProjections(GraphBuilder, SceneTexturesUniformBuffer,  &LightSceneInfo, ScreenShadowMaskTexture, ScreenShadowMaskSubPixelTexture, SceneDepthTexture, HairDatas, bInjectedTranslucentVolume);
+						RenderDeferredShadowProjections(GraphBuilder, SceneTexturesUniformBuffer, TranslucentVolumeLightingTextures, &LightSceneInfo, ScreenShadowMaskTexture, ScreenShadowMaskSubPixelTexture, SceneDepthTexture, HairDatas, bInjectedTranslucentVolume);
 					}
 
 					bUsedShadowMaskTexture = true;
@@ -1919,7 +1909,7 @@ FRDGTextureRef FDeferredShadingSceneRenderer::RenderLights(
 						RDG_GPU_MASK_SCOPE(GraphBuilder, View.GPUMask);
 
 						// Accumulate this light's unshadowed contribution to the translucency lighting volume
-						InjectTranslucentVolumeLighting(GraphBuilder, LightSceneInfo, nullptr, View, ViewIndex);
+						InjectTranslucentVolumeLighting(GraphBuilder, TranslucentVolumeLightingTextures, LightSceneInfo, nullptr, View, ViewIndex);
 					}
 				}
 
diff --git a/Engine/Source/Runtime/Renderer/Private/MobileShadingRenderer.cpp b/Engine/Source/Runtime/Renderer/Private/MobileShadingRenderer.cpp
--- a/Engine/Source/Runtime/Renderer/Private/MobileShadingRenderer.cpp
+++ b/Engine/Source/Runtime/Renderer/Private/MobileShadingRenderer.cpp
@@ -532,7 +532,11 @@ void FMobileSceneRenderer::Render(FRHICommandListImmediate& RHICmdList)
 
 	// Allocate the maximum scene render target space for the current view family.
 	SceneContext.SetKeepDepthContent(bKeepDepthContent);
-	SceneContext.Allocate(RHICmdList, this);
+	{
+		FRDGBuilder GraphBuilder(RHICmdList);
+		SceneContext.Allocate(GraphBuilder, this);
+		GraphBuilder.Execute();
+	}
 	if (bDeferredShading)
 	{
 		ETextureCreateFlags AddFlags = bRequiresMultiPass ? TexCreate_InputAttachmentRead : (TexCreate_InputAttachmentRead | TexCreate_Memoryless);
diff --git a/Engine/Source/Runtime/Renderer/Private/PostProcess/SceneRenderTargets.cpp b/Engine/Source/Runtime/Renderer/Private/PostProcess/SceneRenderTargets.cpp
--- a/Engine/Source/Runtime/Renderer/Private/PostProcess/SceneRenderTargets.cpp
+++ b/Engine/Source/Runtime/Renderer/Private/PostProcess/SceneRenderTargets.cpp
@@ -443,7 +443,7 @@ uint16 FSceneRenderTargets::GetNumSceneColorMSAASamples(ERHIFeatureLevel::Type I
 	return NumSamples;
 }
 
-void FSceneRenderTargets::Allocate(FRHICommandListImmediate& RHICmdList, const FSceneRenderer* SceneRenderer)
+void FSceneRenderTargets::Allocate(FRDGBuilder& GraphBuilder, const FSceneRenderer* SceneRenderer)
 {
 	check(IsInRenderingThread());
 	// ViewFamily setup wasn't complete
@@ -518,8 +518,6 @@ void FSceneRenderTargets::Allocate(FRHICommandListImmediate& RHICmdList, const F
 		}
 	}
 
-	const int32 TranslucencyLightingVolumeDim = GetTranslucencyLightingVolumeDim();
-
 	int32 MSAACount = GetNumSceneColorMSAASamples(NewFeatureLevel);
 
 	uint32 MinShadowResolution;
@@ -537,7 +535,6 @@ void FSceneRenderTargets::Allocate(FRHICommandListImmediate& RHICmdList, const F
 		(bAllowStaticLighting != bNewAllowStaticLighting) ||
 		(bUseDownsizedOcclusionQueries != bDownsampledOcclusionQueries) ||
 		(CurrentMaxShadowResolution != MaxShadowResolution) ||
-		(CurrentTranslucencyLightingVolumeDim != TranslucencyLightingVolumeDim) ||
 		(CurrentMSAACount != MSAACount) ||
 		(CurrentMinShadowResolution != MinShadowResolution))
 	{
@@ -547,7 +544,6 @@ void FSceneRenderTargets::Allocate(FRHICommandListImmediate& RHICmdList, const F
 		bAllowStaticLighting = bNewAllowStaticLighting;
 		bUseDownsizedOcclusionQueries = bDownsampledOcclusionQueries;
 		CurrentMaxShadowResolution = MaxShadowResolution;
-		CurrentTranslucencyLightingVolumeDim = TranslucencyLightingVolumeDim;
 		CurrentMSAACount = MSAACount;
 		CurrentMinShadowResolution = MinShadowResolution;
 
@@ -561,7 +557,7 @@ void FSceneRenderTargets::Allocate(FRHICommandListImmediate& RHICmdList, const F
 
 	// Do allocation of render targets if they aren't available for the current shading path
 	CurrentFeatureLevel = NewFeatureLevel;
-	AllocateRenderTargets(RHICmdList, ViewFamily.Views.Num());
+	AllocateRenderTargets(GraphBuilder, ViewFamily.Views.Num());
 }
 
 static int32 FindGBufferTargetIndex(const FGBufferInfo& GBufferInfo, const FString& Name)
@@ -1528,45 +1524,6 @@ void FSceneRenderTargets::AllocateMobileRenderTargets(FRHICommandListImmediate&
 	}
 }
 
-// This is a helper class. It generates and provides N names with
-// sequentially incremented postfix starting from 0.
-// Example: SomeName0, SomeName1, ..., SomeName117
-class FIncrementalNamesHolder
-{
-public:
-	FIncrementalNamesHolder(const TCHAR* const Name, uint32 Size)
-	{
-		check(Size > 0);
-		Names.Empty(Size);
-		for (uint32 i = 0; i < Size; ++i)
-		{
-			Names.Add(FString::Printf(TEXT("%s%d"), Name, i));
-		}
-	}
-
-	const TCHAR* const operator[] (uint32 Idx) const
-	{
-		CA_SUPPRESS(6385);	// Doesn't like COM
-		return *Names[Idx];
-	}
-
-private:
-	TArray<FString> Names;
-};
-
-// for easier use of "VisualizeTexture"
-static const TCHAR* const GetVolumeName(uint32 Id, bool bDirectional)
-{
-	constexpr uint32 MaxNames = 128;
-	static FIncrementalNamesHolder Names(TEXT("TranslucentVolume"), MaxNames);
-	static FIncrementalNamesHolder NamesDir(TEXT("TranslucentVolumeDir"), MaxNames);
-
-	check(Id < MaxNames);
-
-	CA_SUPPRESS(6385);	// Doesn't like COM
-	return bDirectional ? NamesDir[Id] : Names[Id];
-}
-
 void FSceneRenderTargets::AllocateReflectionTargets(FRHICommandList& RHICmdList, int32 TargetSize)
 {
 	if (GSupportsRenderTargetFormat_PF_FloatRGBA)
@@ -1753,18 +1710,18 @@ void FSceneRenderTargets::AllocateFoveationTexture(FRHICommandList& RHICmdList)
 	}
 }
 
-void FSceneRenderTargets::AllocateDeferredShadingPathRenderTargets(FRHICommandListImmediate& RHICmdList, const int32 NumViews)
+void FSceneRenderTargets::AllocateDeferredShadingPathRenderTargets(FRDGBuilder& GraphBuilder, const int32 NumViews)
 {
 	const EShaderPlatform ShaderPlatform = GetFeatureLevelShaderPlatform(CurrentFeatureLevel);
 
-	AllocateCommonDepthTargets(RHICmdList);
+	AllocateCommonDepthTargets(GraphBuilder.RHICmdList);
 
 	// Create a quarter-sized version of the scene depth.
 	{
 		FIntPoint SmallDepthZSize(FMath::Max<uint32>(BufferSize.X / SmallColorDepthDownsampleFactor, 1), FMath::Max<uint32>(BufferSize.Y / SmallColorDepthDownsampleFactor, 1));
 
 		FPooledRenderTargetDesc Desc(FPooledRenderTargetDesc::Create2DDesc(SmallDepthZSize, PF_DepthStencil, FClearValueBinding::None, TexCreate_None, TexCreate_DepthStencilTargetable | TexCreate_ShaderResource, true));
-		GRenderTargetPool.FindFreeElement(RHICmdList, Desc, SmallDepthZ, TEXT("SmallDepthZ"), ERenderTargetTransience::NonTransient);
+		GRenderTargetPool.FindFreeElement(GraphBuilder.RHICmdList, Desc, SmallDepthZ, TEXT("SmallDepthZ"), ERenderTargetTransience::NonTransient);
 	}
 
 	// Create the required render targets if running Highend.
@@ -1781,73 +1738,7 @@ void FSceneRenderTargets::AllocateDeferredShadingPathRenderTargets(FRHICommandLi
 				// todo: ideally this should be only UAV or RT, not both
 				Desc.TargetableFlags |= TexCreate_UAV;
 			}
-			GRenderTargetPool.FindFreeElement(RHICmdList, Desc, ScreenSpaceAO, TEXT("ScreenSpaceAO"), ERenderTargetTransience::NonTransient);
-		}
-		
-		{
-			const int32 TranslucencyLightingVolumeDim = GetTranslucencyLightingVolumeDim();
-
-			// TODO: We can skip the and TLV allocations when rendering in forward shading mode
-			ETextureCreateFlags TranslucencyTargetFlags = TexCreate_ShaderResource | TexCreate_RenderTargetable | TexCreate_ReduceMemoryWithTilingMode;
-
-			if (CurrentFeatureLevel >= ERHIFeatureLevel::SM5)
-			{
-				TranslucencyTargetFlags |= TexCreate_UAV;
-			}
-
-			TranslucencyLightingVolumeAmbient.SetNum(NumViews * NumTranslucentVolumeRenderTargetSets);
-			TranslucencyLightingVolumeDirectional.SetNum(NumViews * NumTranslucentVolumeRenderTargetSets);
-
-			for (int32 RTSetIndex = 0; RTSetIndex < NumTranslucentVolumeRenderTargetSets * NumViews; RTSetIndex++)
-			{
-				GRenderTargetPool.FindFreeElement(
-					RHICmdList,
-					FPooledRenderTargetDesc(FPooledRenderTargetDesc::CreateVolumeDesc(
-						TranslucencyLightingVolumeDim,
-						TranslucencyLightingVolumeDim,
-						TranslucencyLightingVolumeDim,
-						PF_FloatRGBA,
-						FClearValueBinding::Transparent,
-						TexCreate_None,
-						TranslucencyTargetFlags,
-						false,
-						1,
-						false)),
-					TranslucencyLightingVolumeAmbient[RTSetIndex],
-					GetVolumeName(RTSetIndex, false),
-					ERenderTargetTransience::NonTransient
-					);
-
-				//Tests to catch UE-31578, UE-32536 and UE-22073 crash (Defferred Render Targets not being allocated)
-				ensureMsgf(TranslucencyLightingVolumeAmbient[RTSetIndex], TEXT("Failed to allocate render target %s with dimension %i and flags %i"),
-					GetVolumeName(RTSetIndex, false),
-					TranslucencyLightingVolumeDim,
-					TranslucencyTargetFlags);
-
-				GRenderTargetPool.FindFreeElement(
-					RHICmdList,
-					FPooledRenderTargetDesc(FPooledRenderTargetDesc::CreateVolumeDesc(
-						TranslucencyLightingVolumeDim,
-						TranslucencyLightingVolumeDim,
-						TranslucencyLightingVolumeDim,
-						PF_FloatRGBA,
-						FClearValueBinding::Transparent,
-						TexCreate_None,
-						TranslucencyTargetFlags,
-						false,
-						1,
-						false)),
-					TranslucencyLightingVolumeDirectional[RTSetIndex],
-					GetVolumeName(RTSetIndex, true),
-					ERenderTargetTransience::NonTransient
-					);
-
-				//Tests to catch UE-31578, UE-32536 and UE-22073 crash
-				ensureMsgf(TranslucencyLightingVolumeDirectional[RTSetIndex], TEXT("Failed to allocate render target %s with dimension %i and flags %i"),
-					GetVolumeName(RTSetIndex, true),
-					TranslucencyLightingVolumeDim,
-					TranslucencyTargetFlags);
-			}
+			GRenderTargetPool.FindFreeElement(GraphBuilder.RHICmdList, Desc, ScreenSpaceAO, TEXT("ScreenSpaceAO"), ERenderTargetTransience::NonTransient);
 		}
 	}
 
@@ -1859,19 +1750,19 @@ void FSceneRenderTargets::AllocateDeferredShadingPathRenderTargets(FRHICommandLi
 			Desc.TargetableFlags |= TexCreate_UAV;
 		}
 		Desc.Flags |= GFastVRamConfig.LightAccumulation;
-		GRenderTargetPool.FindFreeElement(RHICmdList, Desc, LightAccumulation, TEXT("LightAccumulation"), ERenderTargetTransience::NonTransient);
+		GRenderTargetPool.FindFreeElement(GraphBuilder.RHICmdList, Desc, LightAccumulation, TEXT("LightAccumulation"), ERenderTargetTransience::NonTransient);
 	}
 
 	if (bAllocateVelocityGBuffer)
 	{
 		FPooledRenderTargetDesc VelocityRTDesc = Translate(FVelocityRendering::GetRenderTargetDesc(ShaderPlatform));
 		VelocityRTDesc.Flags |= GFastVRamConfig.GBufferVelocity;
-		GRenderTargetPool.FindFreeElement(RHICmdList, VelocityRTDesc, SceneVelocity, TEXT("GBufferVelocity"));
+		GRenderTargetPool.FindFreeElement(GraphBuilder.RHICmdList, VelocityRTDesc, SceneVelocity, TEXT("GBufferVelocity"));
 	}
 
-	AllocateVirtualTextureFeedbackBuffer(RHICmdList);
+	AllocateVirtualTextureFeedbackBuffer(GraphBuilder.RHICmdList);
 
-	AllocateDebugViewModeTargets(RHICmdList);
+	AllocateDebugViewModeTargets(GraphBuilder.RHICmdList);
 }
 
 void FSceneRenderTargets::AllocateAnisotropyTarget(FRHICommandListImmediate& RHICmdList)
@@ -1968,17 +1859,17 @@ EPixelFormat FSceneRenderTargets::GetSceneColorFormat(ERHIFeatureLevel::Type InF
 	return SceneColorBufferFormat;
 }
 
-void FSceneRenderTargets::AllocateRenderTargets(FRHICommandListImmediate& RHICmdList, const int32 NumViews)
+void FSceneRenderTargets::AllocateRenderTargets(FRDGBuilder& GraphBuilder, const int32 NumViews)
 {
 	if (BufferSize.X > 0 && BufferSize.Y > 0 && IsAllocateRenderTargetsRequired())
 	{
 		if ((EShadingPath)CurrentShadingPath == EShadingPath::Mobile)
 		{
-			AllocateMobileRenderTargets(RHICmdList);
+			AllocateMobileRenderTargets(GraphBuilder.RHICmdList);
 		}
 		else
 		{
-			AllocateDeferredShadingPathRenderTargets(RHICmdList, NumViews);
+			AllocateDeferredShadingPathRenderTargets(GraphBuilder, NumViews);
 		}
 	}
 	else if ((EShadingPath)CurrentShadingPath == EShadingPath::Mobile && SceneDepthZ)
@@ -1992,7 +1883,7 @@ void FSceneRenderTargets::AllocateRenderTargets(FRHICommandListImmediate& RHICmd
 			SceneDepthZ.SafeRelease();
 			// Make sure the old depth buffer is freed by flushing the target pool.
 			GRenderTargetPool.FreeUnusedResources();
-			AllocateCommonDepthTargets(RHICmdList);
+			AllocateCommonDepthTargets(GraphBuilder.RHICmdList);
 		}
 	}
 }
@@ -2050,13 +1941,6 @@ void FSceneRenderTargets::ReleaseAllTargets()
 
 	SkySHIrradianceMap.SafeRelease();
 
-	ensure(TranslucencyLightingVolumeAmbient.Num() == TranslucencyLightingVolumeDirectional.Num());
-	for (int32 RTSetIndex = 0; RTSetIndex < TranslucencyLightingVolumeAmbient.Num(); RTSetIndex++)
-	{
-		TranslucencyLightingVolumeAmbient[RTSetIndex].SafeRelease();
-		TranslucencyLightingVolumeDirectional[RTSetIndex].SafeRelease();
-	}
-
 	MobileMultiViewSceneColor.SafeRelease();
 	MobileMultiViewSceneDepthZ.SafeRelease();
 
diff --git a/Engine/Source/Runtime/Renderer/Private/PostProcess/SceneRenderTargets.h b/Engine/Source/Runtime/Renderer/Private/PostProcess/SceneRenderTargets.h
--- a/Engine/Source/Runtime/Renderer/Private/PostProcess/SceneRenderTargets.h
+++ b/Engine/Source/Runtime/Renderer/Private/PostProcess/SceneRenderTargets.h
@@ -26,67 +26,9 @@ class FRDGBuilder;
 /** Number of cube map shadow depth surfaces that will be created and used for rendering one pass point light shadows. */
 static const int32 NumCubeShadowDepthSurfaces = 5;
 
-/** 
- * Allocate enough sets of translucent volume textures to cover all the cascades, 
- * And then one more which will be used as a scratch target when doing ping-pong operations like filtering. 
- */
-static const int32 NumTranslucentVolumeRenderTargetSets = (TVC_MAX + 1);
-
 /** Forward declaration of console variable controlling translucent volume blur */
 extern int32 GUseTranslucencyVolumeBlur;
 
-/** Function returning current translucency lighting volume dimensions. */
-inline int32 GetTranslucencyLightingVolumeDim()
-{
-	extern int32 GTranslucencyLightingVolumeDim;
-	return FMath::Clamp(GTranslucencyLightingVolumeDim, 4, 2048);
-}
-
-/** Function to select the index of the volume texture that we will hold the final translucency lighting volume texture */
-inline int SelectTranslucencyVolumeTarget(ETranslucencyVolumeCascade InCascade)
-{
-	if (GUseTranslucencyVolumeBlur)
-	{
-		switch (InCascade)
-		{
-		case TVC_Inner:
-			{
-				return 2;
-			}
-		case TVC_Outer:
-			{
-				return 0;
-			}
-		default:
-			{
-				// error
-				check(false);
-				return 0;
-			}
-		}
-	}
-	else
-	{
-		switch (InCascade)
-		{
-		case TVC_Inner:
-			{
-				return 0;
-			}
-		case TVC_Outer:
-			{
-				return 1;
-			}
-		default:
-			{
-				// error
-				check(false);
-				return 0;
-			}
-		}
-	}
-}
-
 /** Number of surfaces used for translucent shadows. */
 static const int32 NumTranslucencyShadowSurfaces = 2;
 
@@ -165,7 +107,6 @@ protected:
 		CurrentMobileSceneColorFormat(EPixelFormat::PF_Unknown),
 		bAllowStaticLighting(true),
 		CurrentMaxShadowResolution(0),
-		CurrentTranslucencyLightingVolumeDim(64),
 		CurrentFeatureLevel(ERHIFeatureLevel::Num),
 		CurrentShadingPath(EShadingPath::Num),
 		bRequireSceneColorAlpha(false),
@@ -191,7 +132,7 @@ public:
 	 * Checks that scene render targets are ready for rendering a view family of the given dimensions.
 	 * If the allocated render targets are too small, they are reallocated.
 	 */
-	void Allocate(FRHICommandListImmediate& RHICmdList, const FSceneRenderer* SceneRenderer);
+	void Allocate(FRDGBuilder& GraphBuilder, const FSceneRenderer* SceneRenderer);
 
 	void AllocSceneColor(FRHICommandList& RHICmdList);
 
@@ -318,26 +259,6 @@ public:
 
 	bool UseDownsizedOcclusionQueries() const { return bUseDownsizedOcclusionQueries; }
 
-	/** Get the current translucent ambient lighting volume texture. Can vary depending on whether volume filtering is enabled */
-	IPooledRenderTarget* GetTranslucencyVolumeAmbient(ETranslucencyVolumeCascade Cascade, int32 ViewIndex = 0) const
-	{
-		if (TranslucencyLightingVolumeAmbient.Num())
-		{
-			return TranslucencyLightingVolumeAmbient[SelectTranslucencyVolumeTarget(Cascade) + ViewIndex * NumTranslucentVolumeRenderTargetSets].GetReference();
-		}
-		return nullptr;
-	}
-
-	/** Get the current translucent directional lighting volume texture. Can vary depending on whether volume filtering is enabled */
-	IPooledRenderTarget* GetTranslucencyVolumeDirectional(ETranslucencyVolumeCascade Cascade, int32 ViewIndex = 0) const
-	{
-		if (TranslucencyLightingVolumeDirectional.Num())
-		{
-			return TranslucencyLightingVolumeDirectional[SelectTranslucencyVolumeTarget(Cascade) + ViewIndex * NumTranslucentVolumeRenderTargetSets].GetReference();
-		}
-		return nullptr;
-	}
-
 	/** Returns the size of most screen space render targets e.g. SceneColor, SceneDepth, GBuffer, ... might be different from final RT or output Size because of ScreenPercentage use. */
 	FIntPoint GetBufferSizeXY() const { return BufferSize; }
 	/** */
@@ -405,7 +326,7 @@ public:
 	// Can be called when the Scene Color content is no longer needed. As we create SceneColor on demand we can make sure it is created with the right format.
 	// (as a call to SetSceneColor() can override it with a different format)
 	void ReleaseSceneColor();
-	
+
 	ERHIFeatureLevel::Type GetCurrentFeatureLevel() const { return CurrentFeatureLevel; }
 	bool IsFoveationTextureAllocated() const { return bAllocatedFoveationTexture;  }
 
@@ -467,10 +388,6 @@ public:
 	/** Temporary storage during SH irradiance map generation. */
 	TRefCountPtr<IPooledRenderTarget> SkySHIrradianceMap;
 
-	/** Volume textures used for lighting translucency. */
-	TArray<TRefCountPtr<IPooledRenderTarget>, TInlineAllocator<NumTranslucentVolumeRenderTargetSets>> TranslucencyLightingVolumeAmbient;
-	TArray<TRefCountPtr<IPooledRenderTarget>, TInlineAllocator<NumTranslucentVolumeRenderTargetSets>> TranslucencyLightingVolumeDirectional;
-
 	/** Color and depth texture arrays for mobile multi-view */
 	TRefCountPtr<IPooledRenderTarget> MobileMultiViewSceneColor;
 	TRefCountPtr<IPooledRenderTarget> MobileMultiViewSceneDepthZ;
@@ -528,7 +445,7 @@ private:
 public:
 	/** Allocates render targets for use with the deferred shading path. */
 	// Temporarily Public to call from DefferedShaderRenderer to attempt recovery from a crash until cause is found.
-	void AllocateDeferredShadingPathRenderTargets(FRHICommandListImmediate& RHICmdList, const int32 NumViews = 1);
+	void AllocateDeferredShadingPathRenderTargets(FRDGBuilder& GraphBuilder, const int32 NumViews = 1);
 
 	void AllocateAnisotropyTarget(FRHICommandListImmediate& RHICmdList);
 
@@ -545,7 +462,7 @@ public:
 private:
 
 	/** Allocates render targets for use with the current shading path. */
-	void AllocateRenderTargets(FRHICommandListImmediate& RHICmdList, const int32 NumViews);
+	void AllocateRenderTargets(FRDGBuilder& GraphBuilder, const int32 NumViews);
 
 
 	/** Allocates common depth render targets that are used by both mobile and deferred rendering paths */
@@ -629,8 +546,6 @@ private:
 	bool bAllowStaticLighting;
 	/** To detect a change of the CVar r.Shadow.MaxResolution */
 	int32 CurrentMaxShadowResolution;
-	/** To detect a change of the CVar r.TranslucencyLightingVolumeDim */
-	int32 CurrentTranslucencyLightingVolumeDim;
 	/** To detect a change of the CVar r.MobileMSAA or r.MSAA */
 	int32 CurrentMSAACount;
 	/** To detect a change of the CVar r.Shadow.MinResolution */
diff --git a/Engine/Source/Runtime/Renderer/Private/SceneHitProxyRendering.cpp b/Engine/Source/Runtime/Renderer/Private/SceneHitProxyRendering.cpp
--- a/Engine/Source/Runtime/Renderer/Private/SceneHitProxyRendering.cpp
+++ b/Engine/Source/Runtime/Renderer/Private/SceneHitProxyRendering.cpp
@@ -229,7 +229,7 @@ void InitHitProxyRender(FRDGBuilder& GraphBuilder, const FSceneRenderer* SceneRe
 
 	FSceneRenderTargets& SceneContext = FSceneRenderTargets::Get();
 	// Allocate the maximum scene render target space for the current view family.
-	SceneContext.Allocate(GraphBuilder.RHICmdList, SceneRenderer);
+	SceneContext.Allocate(GraphBuilder, SceneRenderer);
 
 	// Create a texture to store the resolved light attenuation values, and a render-targetable surface to hold the unresolved light attenuation values.
 	{
diff --git a/Engine/Source/Runtime/Renderer/Private/SceneRendering.cpp b/Engine/Source/Runtime/Renderer/Private/SceneRendering.cpp
--- a/Engine/Source/Runtime/Renderer/Private/SceneRendering.cpp
+++ b/Engine/Source/Runtime/Renderer/Private/SceneRendering.cpp
@@ -376,6 +376,7 @@ static TAutoConsoleVariable<int32> CVarTestSecondaryUpscaleOverride(
 static FParallelCommandListSet* GOutstandingParallelCommandListSet = nullptr;
 FGraphEventRef FSceneRenderer::OcclusionSubmittedFence[FOcclusionQueryHelpers::MaxBufferedOcclusionFrames];
 
+extern int32 GetTranslucencyLightingVolumeDim();
 
 DECLARE_CYCLE_STAT(TEXT("DeferredShadingSceneRenderer ViewExtensionPostRenderView"), STAT_FDeferredShadingSceneRenderer_ViewExtensionPostRenderView, STATGROUP_SceneRendering);
 DECLARE_CYCLE_STAT(TEXT("DeferredShadingSceneRenderer ViewExtensionPreRenderView"), STAT_FDeferredShadingSceneRenderer_ViewExtensionPreRenderView, STATGROUP_SceneRendering);
diff --git a/Engine/Source/Runtime/Renderer/Private/ShadowRendering.cpp b/Engine/Source/Runtime/Renderer/Private/ShadowRendering.cpp
--- a/Engine/Source/Runtime/Renderer/Private/ShadowRendering.cpp
+++ b/Engine/Source/Runtime/Renderer/Private/ShadowRendering.cpp
@@ -1945,6 +1945,7 @@ void FDeferredShadingSceneRenderer::RenderShadowProjections(
 void FDeferredShadingSceneRenderer::RenderDeferredShadowProjections(
 	FRDGBuilder& GraphBuilder,
 	TRDGUniformBufferRef<FSceneTextureUniformParameters> SceneTexturesUniformBuffer,
+	const FTranslucentVolumeLightingTextures& TranslucentVolumeLightingTextures,
 	const FLightSceneInfo* LightSceneInfo,
 	FRDGTextureRef ScreenShadowMaskTexture,
 	FRDGTextureRef ScreenShadowMaskSubPixelTexture,
@@ -1999,7 +2000,7 @@ void FDeferredShadingSceneRenderer::RenderDeferredShadowProjections(
 					}
 
 					RDG_GPU_MASK_SCOPE(GraphBuilder, ProjectedShadowInfo->DependentView->GPUMask);
-					InjectTranslucentVolumeLighting(GraphBuilder, *LightSceneInfo, ProjectedShadowInfo, *ProjectedShadowInfo->DependentView, ViewIndex);
+					InjectTranslucentVolumeLighting(GraphBuilder, TranslucentVolumeLightingTextures, *LightSceneInfo, ProjectedShadowInfo, *ProjectedShadowInfo->DependentView, ViewIndex);
 				}
 				else
 				{
@@ -2007,7 +2008,7 @@ void FDeferredShadingSceneRenderer::RenderDeferredShadowProjections(
 					{
 						FViewInfo& View = Views[ViewIndex];
 						RDG_GPU_MASK_SCOPE(GraphBuilder, View.GPUMask);
-						InjectTranslucentVolumeLighting(GraphBuilder, *LightSceneInfo, ProjectedShadowInfo, View, ViewIndex);
+						InjectTranslucentVolumeLighting(GraphBuilder, TranslucentVolumeLightingTextures, *LightSceneInfo, ProjectedShadowInfo, View, ViewIndex);
 					}
 				}
 			}
diff --git a/Engine/Source/Runtime/Renderer/Private/TranslucentLighting.cpp b/Engine/Source/Runtime/Renderer/Private/TranslucentLighting.cpp
--- a/Engine/Source/Runtime/Renderer/Private/TranslucentLighting.cpp
+++ b/Engine/Source/Runtime/Renderer/Private/TranslucentLighting.cpp
@@ -109,6 +109,13 @@ static TAutoConsoleVariable<float> CVarTranslucencyLightingVolumeOuterDistance(
 	TEXT("Distance from the camera that the second volume cascade should end"),
 	ECVF_RenderThreadSafe);
 
+/** Function returning current translucency lighting volume dimensions. */
+int32 GetTranslucencyLightingVolumeDim()
+{
+	extern int32 GTranslucencyLightingVolumeDim;
+	return FMath::Clamp(GTranslucencyLightingVolumeDim, 4, 2048);
+}
+
 void FViewInfo::CalcTranslucencyLightingVolumeBounds(FBox* InOutCascadeBoundsArray, int32 NumCascades) const
 {
 	for (int32 CascadeIndex = 0; CascadeIndex < NumCascades; CascadeIndex++)
@@ -848,13 +855,53 @@ public:
 
 IMPLEMENT_GLOBAL_SHADER(FClearTranslucentLightingVolumeCS, "/Engine/Private/TranslucentLightInjectionShaders.usf", "ClearTranslucentLightingVolumeCS", SF_Compute);
 
-void FDeferredShadingSceneRenderer::ClearTranslucentVolumeLighting(FRDGBuilder& GraphBuilder, ERDGPassFlags PassFlags)
+void FDeferredShadingSceneRenderer::InitTranslucentVolumeLighting(FRDGBuilder& GraphBuilder, ERDGPassFlags PassFlags, FTranslucentVolumeLightingTextures& Textures)
 {
 	RDG_GPU_STAT_SCOPE(GraphBuilder, TranslucentLighting);
 
-	FSceneRenderTargets& SceneContext = FSceneRenderTargets::Get();
+	Textures.VolumeDim = GetTranslucencyLightingVolumeDim();
+	const FIntVector TranslucencyLightingVolumeDim(Textures.VolumeDim);
+
+	{
+		// TODO: We can skip the and TLV allocations when rendering in forward shading mode
+		const ETextureCreateFlags TranslucencyTargetFlags = TexCreate_ShaderResource | TexCreate_RenderTargetable | TexCreate_ReduceMemoryWithTilingMode | TexCreate_UAV;
+
+		const int32 ViewCount = Views.Num();
+		Textures.Ambient.SetNum(ViewCount * TVC_MAX);
+		Textures.Directional.SetNum(ViewCount * TVC_MAX);
+
+		int32 TextureIndex = 0;
+		for (int32 ViewIndex = 0; ViewIndex < ViewCount; ++ViewIndex)
+		{
+			for (int32 CascadeIndex = 0; CascadeIndex < TVC_MAX; ++CascadeIndex)
+			{
+				FRDGEventName& AmbientName = *GraphBuilder.AllocObject<FRDGEventName>(RDG_EVENT_NAME("TranslucentVolumeAmbient%d", TextureIndex));
+				FRDGEventName& DirectionalName = *GraphBuilder.AllocObject<FRDGEventName>(RDG_EVENT_NAME("TranslucentVolumeDirectional%d", TextureIndex));
+
+				FRDGTextureRef AmbientTexture = GraphBuilder.CreateTexture(
+					FRDGTextureDesc::Create3D(
+						TranslucencyLightingVolumeDim,
+						PF_FloatRGBA,
+						FClearValueBinding::Transparent,
+						TranslucencyTargetFlags),
+					AmbientName.GetTCHAR());
+
+				Textures.SetAmbient(ViewIndex, CascadeIndex, AmbientTexture);
+
+				FRDGTextureRef DirectionalTexture = GraphBuilder.CreateTexture(
+					FRDGTextureDesc::Create3D(
+						TranslucencyLightingVolumeDim,
+						PF_FloatRGBA,
+						FClearValueBinding::Transparent,
+						TranslucencyTargetFlags),
+					DirectionalName.GetTCHAR());
+
+				Textures.SetDirectional(ViewIndex, CascadeIndex, DirectionalTexture);
+				TextureIndex++;
+			}
+		}
+	}
 
-	const FIntVector TranslucencyLightingVolumeDim(GetTranslucencyLightingVolumeDim());
 	const FIntVector GroupCount = FComputeShaderUtils::GetGroupCount(TranslucencyLightingVolumeDim, FClearTranslucentLightingVolumeCS::CLEAR_BLOCK_SIZE);
 
 	TShaderMapRef<FClearTranslucentLightingVolumeCS> ComputeShader(GetGlobalShaderMap(FeatureLevel));
@@ -862,14 +909,14 @@ void FDeferredShadingSceneRenderer::ClearTranslucentVolumeLighting(FRDGBuilder&
 	for (int32 ViewIndex = 0; ViewIndex < Views.Num(); ++ViewIndex)
 	{
 		auto* PassParameters = GraphBuilder.AllocParameters<FClearTranslucentLightingVolumeCS::FParameters>();
-		PassParameters->RWAmbient0 = GraphBuilder.CreateUAV(GraphBuilder.RegisterExternalTexture(SceneContext.TranslucencyLightingVolumeAmbient[(ViewIndex * NumTranslucentVolumeRenderTargetSets) + 0]));
-		PassParameters->RWAmbient1 = GraphBuilder.CreateUAV(GraphBuilder.RegisterExternalTexture(SceneContext.TranslucencyLightingVolumeAmbient[(ViewIndex * NumTranslucentVolumeRenderTargetSets) + 1]));
-		PassParameters->RWDirectional0 = GraphBuilder.CreateUAV(GraphBuilder.RegisterExternalTexture(SceneContext.TranslucencyLightingVolumeDirectional[(ViewIndex * NumTranslucentVolumeRenderTargetSets) + 0]));
-		PassParameters->RWDirectional1 = GraphBuilder.CreateUAV(GraphBuilder.RegisterExternalTexture(SceneContext.TranslucencyLightingVolumeDirectional[(ViewIndex * NumTranslucentVolumeRenderTargetSets) + 1]));
+		PassParameters->RWAmbient0 = GraphBuilder.CreateUAV(Textures.GetAmbient(ViewIndex, 0));
+		PassParameters->RWAmbient1 = GraphBuilder.CreateUAV(Textures.GetAmbient(ViewIndex, 1));
+		PassParameters->RWDirectional0 = GraphBuilder.CreateUAV(Textures.GetDirectional(ViewIndex, 0));
+		PassParameters->RWDirectional1 = GraphBuilder.CreateUAV(Textures.GetDirectional(ViewIndex, 1));
 
 		FComputeShaderUtils::AddPass(
 			GraphBuilder,
-			RDG_EVENT_NAME("ClearTranslucencyLightingVolumeCompute %d", TranslucencyLightingVolumeDim),
+			RDG_EVENT_NAME("ClearTranslucencyLightingVolumeCompute %d", Textures.VolumeDim),
 			PassFlags,
 			ComputeShader,
 			PassParameters,
@@ -897,7 +944,11 @@ public:
 
 IMPLEMENT_GLOBAL_SHADER(FInjectAmbientCubemapPS, "/Engine/Private/TranslucentLightingShaders.usf", "InjectAmbientCubemapMainPS", SF_Pixel);
 
-void FDeferredShadingSceneRenderer::InjectAmbientCubemapTranslucentVolumeLighting(FRDGBuilder& GraphBuilder, const FViewInfo& View, int32 ViewIndex)
+void FDeferredShadingSceneRenderer::InjectAmbientCubemapTranslucentVolumeLighting(
+	FRDGBuilder& GraphBuilder,
+	const FTranslucentVolumeLightingTextures& Textures,
+	const FViewInfo& View,
+	int32 ViewIndex)
 {
 	if (!GUseTranslucentLightingVolumes || !GSupportsVolumeTextureRendering || !View.FinalPostProcessSettings.ContributingCubemaps.Num())
 	{
@@ -907,16 +958,14 @@ void FDeferredShadingSceneRenderer::InjectAmbientCubemapTranslucentVolumeLightin
 	RDG_EVENT_SCOPE(GraphBuilder, "InjectAmbientCubemapTranslucentVolumeLighting");
 	RDG_GPU_STAT_SCOPE(GraphBuilder, TranslucentLighting);
 
-	FSceneRenderTargets& SceneContext = FSceneRenderTargets::Get();
-
-	const int32 TranslucencyLightingVolumeDim = GetTranslucencyLightingVolumeDim();
+	const int32 TranslucencyLightingVolumeDim = Textures.VolumeDim;
 	const FVolumeBounds VolumeBounds(TranslucencyLightingVolumeDim);
 
 	auto ShaderMap = GetGlobalShaderMap(FeatureLevel);
 
 	for (int32 VolumeCascadeIndex = 0; VolumeCascadeIndex < TVC_MAX; ++VolumeCascadeIndex)
 	{
-		FRDGTextureRef VolumeAmbientTexture = GraphBuilder.RegisterExternalTexture(SceneContext.TranslucencyLightingVolumeAmbient[VolumeCascadeIndex + NumTranslucentVolumeRenderTargetSets * ViewIndex]);
+		FRDGTextureRef VolumeAmbientTexture = Textures.GetAmbient(ViewIndex, VolumeCascadeIndex);
 
 		for (const FFinalPostProcessSettings::FCubemapEntry& CubemapEntry : View.FinalPostProcessSettings.ContributingCubemaps)
 		{
@@ -1156,6 +1205,7 @@ END_SHADER_PARAMETER_STRUCT()
 /** Injects all the lights in LightInjectionData into the translucent lighting volume textures. */
 static void InjectTranslucentLightArray(
 	FRDGBuilder& GraphBuilder,
+	const FTranslucentVolumeLightingTextures& Textures,
 	FScene* Scene,
 	const FViewInfo& View,
 	TArrayView<const FTranslucentLightInjectionData> LightInjectionData,
@@ -1168,8 +1218,8 @@ static void InjectTranslucentLightArray(
 	// Operate on one cascade at a time to reduce render target switches
 	for (uint32 VolumeCascadeIndex = 0; VolumeCascadeIndex < TVC_MAX; VolumeCascadeIndex++)
 	{
-		FRDGTextureRef VolumeAmbientTexture = GraphBuilder.RegisterExternalTexture(SceneContext.TranslucencyLightingVolumeAmbient[VolumeCascadeIndex + NumTranslucentVolumeRenderTargetSets * ViewIndex]);
-		FRDGTextureRef VolumeDirectionalTexture = GraphBuilder.RegisterExternalTexture(SceneContext.TranslucencyLightingVolumeDirectional[VolumeCascadeIndex + NumTranslucentVolumeRenderTargetSets * ViewIndex]);
+		FRDGTextureRef VolumeAmbientTexture = Textures.GetAmbient(ViewIndex, VolumeCascadeIndex);
+		FRDGTextureRef VolumeDirectionalTexture = Textures.GetDirectional(ViewIndex, VolumeCascadeIndex);
 
 		auto* PassParameters = GraphBuilder.AllocParameters<FInjectTranslucentLightArrayParameters>();
 		PassParameters->TransmittanceLutTexture = GetSkyTransmittanceLutTexture(GraphBuilder, Scene, View);
@@ -1260,7 +1310,13 @@ static void InjectTranslucentLightArray(
 	}
 }
 
-void FDeferredShadingSceneRenderer::InjectTranslucentVolumeLighting(FRDGBuilder& GraphBuilder, const FLightSceneInfo& LightSceneInfo, const FProjectedShadowInfo* InProjectedShadowInfo, const FViewInfo& View, int32 ViewIndex)
+void FDeferredShadingSceneRenderer::InjectTranslucentVolumeLighting(
+	FRDGBuilder& GraphBuilder,
+	const FTranslucentVolumeLightingTextures& Textures,
+	const FLightSceneInfo& LightSceneInfo,
+	const FProjectedShadowInfo* InProjectedShadowInfo,
+	const FViewInfo& View,
+	int32 ViewIndex)
 {
 	if (GUseTranslucentLightingVolumes && GSupportsVolumeTextureRendering)
 	{
@@ -1271,11 +1327,16 @@ void FDeferredShadingSceneRenderer::InjectTranslucentVolumeLighting(FRDGBuilder&
 		AddLightForInjection(*this, LightSceneInfo, InProjectedShadowInfo, LightInjectionData);
 
 		// shadowed or unshadowed (InProjectedShadowInfo==0)
-		InjectTranslucentLightArray(GraphBuilder, Scene, View, LightInjectionData, ViewIndex);
+		InjectTranslucentLightArray(GraphBuilder, Textures, Scene, View, LightInjectionData, ViewIndex);
 	}
 }
 
-void FDeferredShadingSceneRenderer::InjectTranslucentVolumeLightingArray(FRDGBuilder& GraphBuilder, const TArray<FSortedLightSceneInfo, SceneRenderingAllocator>& SortedLights, int32 FirstLightIndex, int32 LightsEndIndex)
+void FDeferredShadingSceneRenderer::InjectTranslucentVolumeLightingArray(
+	FRDGBuilder& GraphBuilder,
+	const FTranslucentVolumeLightingTextures& Textures,
+	const TArray<FSortedLightSceneInfo, SceneRenderingAllocator>& SortedLights,
+	int32 FirstLightIndex,
+	int32 LightsEndIndex)
 {
 	SCOPE_CYCLE_COUNTER(STAT_TranslucentInjectTime);
 
@@ -1304,7 +1365,7 @@ void FDeferredShadingSceneRenderer::InjectTranslucentVolumeLightingArray(FRDGBui
 	for (int32 ViewIndex = 0; ViewIndex < Views.Num(); ViewIndex++)
 	{
 		// non-shadowed, non-light function lights
-		InjectTranslucentLightArray(GraphBuilder, Scene, Views[ViewIndex], LightInjectionData[ViewIndex], ViewIndex);
+		InjectTranslucentLightArray(GraphBuilder, Textures, Scene, Views[ViewIndex], LightInjectionData[ViewIndex], ViewIndex);
 	}
 }
 
@@ -1330,7 +1391,12 @@ public:
 
 IMPLEMENT_GLOBAL_SHADER(FSimpleLightTranslucentLightingInjectPS, "/Engine/Private/TranslucentLightInjectionShaders.usf", "SimpleLightInjectMainPS", SF_Pixel);
 
-void FDeferredShadingSceneRenderer::InjectSimpleTranslucentVolumeLightingArray(FRDGBuilder& GraphBuilder, const FSimpleLightArray& SimpleLights, const FViewInfo& View, const int32 ViewIndex)
+void FDeferredShadingSceneRenderer::InjectSimpleTranslucentVolumeLightingArray(
+	FRDGBuilder& GraphBuilder,
+	const FTranslucentVolumeLightingTextures& Textures,
+	const FSimpleLightArray& SimpleLights,
+	const FViewInfo& View,
+	const int32 ViewIndex)
 {
 	SCOPE_CYCLE_COUNTER(STAT_TranslucentInjectTime);
 
@@ -1349,7 +1415,6 @@ void FDeferredShadingSceneRenderer::InjectSimpleTranslucentVolumeLightingArray(F
 		RDG_EVENT_SCOPE(GraphBuilder, "InjectSimpleTranslucentLightArray");
 
 		INC_DWORD_STAT_BY(STAT_NumLightsInjectedIntoTranslucency, NumLightsToInject);
-		FSceneRenderTargets& SceneContext = FSceneRenderTargets::Get();
 
 		const int32 TranslucencyLightingVolumeDim = GetTranslucencyLightingVolumeDim();
 
@@ -1358,8 +1423,8 @@ void FDeferredShadingSceneRenderer::InjectSimpleTranslucentVolumeLightingArray(F
 		for (int32 VolumeCascadeIndex = 0; VolumeCascadeIndex < TVC_MAX; VolumeCascadeIndex++)
 		{
 			RDG_EVENT_SCOPE(GraphBuilder, "Cascade%d", VolumeCascadeIndex);
-			FRDGTextureRef VolumeAmbientTexture = GraphBuilder.RegisterExternalTexture(SceneContext.TranslucencyLightingVolumeAmbient[VolumeCascadeIndex + NumTranslucentVolumeRenderTargetSets * ViewIndex]);
-			FRDGTextureRef VolumeDirectionalTexture = GraphBuilder.RegisterExternalTexture(SceneContext.TranslucencyLightingVolumeDirectional[VolumeCascadeIndex + NumTranslucentVolumeRenderTargetSets * ViewIndex]);
+			FRDGTextureRef VolumeAmbientTexture = Textures.GetAmbient(ViewIndex, VolumeCascadeIndex);
+			FRDGTextureRef VolumeDirectionalTexture = Textures.GetDirectional(ViewIndex, VolumeCascadeIndex);
 
 			for (int32 LightIndex = 0; LightIndex < SimpleLights.InstanceData.Num(); LightIndex++)
 			{
@@ -1425,7 +1490,11 @@ void FDeferredShadingSceneRenderer::InjectSimpleTranslucentVolumeLightingArray(F
 	}
 }
 
-void FDeferredShadingSceneRenderer::FilterTranslucentVolumeLighting(FRDGBuilder& GraphBuilder, const FViewInfo& View, const int32 ViewIndex)
+void FDeferredShadingSceneRenderer::FilterTranslucentVolumeLighting(
+	FRDGBuilder& GraphBuilder,
+	FTranslucentVolumeLightingTextures& Textures,
+	const FViewInfo& View,
+	const int32 ViewIndex)
 {
 	if (!GUseTranslucentLightingVolumes || !GSupportsVolumeTextureRendering || !GUseTranslucencyVolumeBlur)
 	{
@@ -1442,11 +1511,14 @@ void FDeferredShadingSceneRenderer::FilterTranslucentVolumeLighting(FRDGBuilder&
 
 	for (int32 VolumeCascadeIndex = 0; VolumeCascadeIndex < TVC_MAX; VolumeCascadeIndex++)
 	{
-		FRDGTextureRef OutputVolumeAmbientTexture = GraphBuilder.RegisterExternalTexture(SceneContext.GetTranslucencyVolumeAmbient((ETranslucencyVolumeCascade)VolumeCascadeIndex, ViewIndex));
-		FRDGTextureRef OutputVolumeDirectionalTexture = GraphBuilder.RegisterExternalTexture(SceneContext.GetTranslucencyVolumeDirectional((ETranslucencyVolumeCascade)VolumeCascadeIndex, ViewIndex));
+		FRDGTextureRef InputVolumeAmbientTexture = Textures.GetAmbient(ViewIndex, VolumeCascadeIndex);
+		FRDGTextureRef InputVolumeDirectionalTexture = Textures.GetDirectional(ViewIndex, VolumeCascadeIndex);
+
+		FRDGTextureRef OutputVolumeAmbientTexture = GraphBuilder.CreateTexture(InputVolumeAmbientTexture->Desc, InputVolumeAmbientTexture->Name);
+		FRDGTextureRef OutputVolumeDirectionalTexture = GraphBuilder.CreateTexture(InputVolumeDirectionalTexture->Desc, InputVolumeDirectionalTexture->Name);
 
-		FRDGTextureRef InputVolumeAmbientTexture = GraphBuilder.RegisterExternalTexture(SceneContext.TranslucencyLightingVolumeAmbient[VolumeCascadeIndex + NumTranslucentVolumeRenderTargetSets * ViewIndex]);
-		FRDGTextureRef InputVolumeDirectionalTexture = GraphBuilder.RegisterExternalTexture(SceneContext.TranslucencyLightingVolumeDirectional[VolumeCascadeIndex + NumTranslucentVolumeRenderTargetSets * ViewIndex]);
+		Textures.SetAmbient(ViewIndex, VolumeCascadeIndex, OutputVolumeAmbientTexture);
+		Textures.SetDirectional(ViewIndex, VolumeCascadeIndex, OutputVolumeDirectionalTexture);
 
 		auto* PassParameters = GraphBuilder.AllocParameters<FFilterTranslucentVolumePS::FParameters>();
 		PassParameters->View = View.ViewUniformBuffer;
@@ -1455,8 +1527,8 @@ void FDeferredShadingSceneRenderer::FilterTranslucentVolumeLighting(FRDGBuilder&
 		PassParameters->TranslucencyLightingVolumeDirectional = InputVolumeDirectionalTexture;
 		PassParameters->TranslucencyLightingVolumeAmbientSampler = SamplerStateRHI;
 		PassParameters->TranslucencyLightingVolumeDirectionalSampler = SamplerStateRHI;
-		PassParameters->RenderTargets[0] = FRenderTargetBinding(OutputVolumeAmbientTexture, ERenderTargetLoadAction::ELoad);
-		PassParameters->RenderTargets[1] = FRenderTargetBinding(OutputVolumeDirectionalTexture, ERenderTargetLoadAction::ELoad);
+		PassParameters->RenderTargets[0] = FRenderTargetBinding(OutputVolumeAmbientTexture, ERenderTargetLoadAction::ENoAction);
+		PassParameters->RenderTargets[1] = FRenderTargetBinding(OutputVolumeDirectionalTexture, ERenderTargetLoadAction::ENoAction);
 
 		const FVolumeBounds VolumeBounds(TranslucencyLightingVolumeDim);
 		TShaderMapRef<FWriteToSliceVS> VertexShader(View.ShaderMap);
diff --git a/Engine/Source/Runtime/Renderer/Private/TranslucentRendering.cpp b/Engine/Source/Runtime/Renderer/Private/TranslucentRendering.cpp
--- a/Engine/Source/Runtime/Renderer/Private/TranslucentRendering.cpp
+++ b/Engine/Source/Runtime/Renderer/Private/TranslucentRendering.cpp
@@ -611,6 +611,7 @@ void SetupTranslucentBasePassUniformParameters(
 	FRHICommandListImmediate& RHICmdList,
 	const FSceneRenderTargets& SceneRenderTargets,
 	const FViewInfo& View,
+	const FTranslucentVolumeLightingTextures* TranslucentVolumeLightingTextures,
 	FRDGTextureRef SceneColorCopyTexture,
 	ESceneTextureSetupMode SceneTextureSetupMode,
 	const int32 ViewIndex,
@@ -724,12 +725,12 @@ void SetupTranslucentBasePassUniformParameters(
 
 	// Translucency Lighting Volume
 	{
-		if (SceneRenderTargets.GetTranslucencyVolumeAmbient(TVC_Inner) != nullptr)
+		if (TranslucentVolumeLightingTextures)
 		{
-			BasePassParameters.TranslucencyLightingVolumeAmbientInner = GetRDG(SceneRenderTargets.GetTranslucencyVolumeAmbient(TVC_Inner, ViewIndex));
-			BasePassParameters.TranslucencyLightingVolumeAmbientOuter = GetRDG(SceneRenderTargets.GetTranslucencyVolumeAmbient(TVC_Outer, ViewIndex));
-			BasePassParameters.TranslucencyLightingVolumeDirectionalInner = GetRDG(SceneRenderTargets.GetTranslucencyVolumeDirectional(TVC_Inner, ViewIndex));
-			BasePassParameters.TranslucencyLightingVolumeDirectionalOuter = GetRDG(SceneRenderTargets.GetTranslucencyVolumeDirectional(TVC_Outer, ViewIndex));
+			BasePassParameters.TranslucencyLightingVolumeAmbientInner = TranslucentVolumeLightingTextures->GetAmbient(ViewIndex, TVC_Inner);
+			BasePassParameters.TranslucencyLightingVolumeAmbientOuter = TranslucentVolumeLightingTextures->GetAmbient(ViewIndex, TVC_Outer);
+			BasePassParameters.TranslucencyLightingVolumeDirectionalInner = TranslucentVolumeLightingTextures->GetDirectional(ViewIndex, TVC_Inner);
+			BasePassParameters.TranslucencyLightingVolumeDirectionalOuter = TranslucentVolumeLightingTextures->GetDirectional(ViewIndex, TVC_Outer);
 		}
 		else
 		{
@@ -770,13 +771,14 @@ void SetupTranslucentBasePassUniformParameters(
 TRDGUniformBufferRef<FTranslucentBasePassUniformParameters> CreateTranslucentBasePassUniformBuffer(
 	FRDGBuilder& GraphBuilder,
 	const FViewInfo& View,
+	const FTranslucentVolumeLightingTextures& TranslucentVolumeLightingTextures,
 	FRDGTextureRef SceneColorCopyTexture,
 	ESceneTextureSetupMode SceneTextureSetupMode,
 	const int32 ViewIndex)
 {
 	FSceneRenderTargets& SceneRenderTargets = FSceneRenderTargets::Get();
 	FTranslucentBasePassUniformParameters* BasePassParameters = GraphBuilder.AllocParameters<FTranslucentBasePassUniformParameters>();
-	SetupTranslucentBasePassUniformParameters(&GraphBuilder, GraphBuilder.RHICmdList, SceneRenderTargets, View, SceneColorCopyTexture, SceneTextureSetupMode, ViewIndex, *BasePassParameters);
+	SetupTranslucentBasePassUniformParameters(&GraphBuilder, GraphBuilder.RHICmdList, SceneRenderTargets, View, &TranslucentVolumeLightingTextures, SceneColorCopyTexture, SceneTextureSetupMode, ViewIndex, *BasePassParameters);
 	return GraphBuilder.CreateUniformBuffer(BasePassParameters);
 }
 
@@ -788,7 +790,7 @@ TUniformBufferRef<FTranslucentBasePassUniformParameters> CreateTranslucentBasePa
 {
 	FSceneRenderTargets& SceneRenderTargets = FSceneRenderTargets::Get();
 	FTranslucentBasePassUniformParameters BasePassParameters;
-	SetupTranslucentBasePassUniformParameters(nullptr, RHICmdList, SceneRenderTargets, View, nullptr, SceneTextureSetupMode, ViewIndex, BasePassParameters);
+	SetupTranslucentBasePassUniformParameters(nullptr, RHICmdList, SceneRenderTargets, View, nullptr, nullptr, SceneTextureSetupMode, ViewIndex, BasePassParameters);
 	return TUniformBufferRef<FTranslucentBasePassUniformParameters>::CreateUniformBufferImmediate(BasePassParameters, EUniformBufferUsage::UniformBuffer_SingleFrame);
 }
 
@@ -971,6 +973,7 @@ static void RenderTranslucencyViewInner(
 
 void FDeferredShadingSceneRenderer::RenderTranslucencyInner(
 	FRDGBuilder& GraphBuilder,
+	const FTranslucentVolumeLightingTextures& TranslucentVolumeLightingTextures,
 	FRDGTextureMSAA SceneColorTexture,
 	FRDGTextureMSAA SceneDepthTexture,
 	FSeparateTranslucencyTextures* OutSeparateTranslucencyTextures,
@@ -1066,7 +1069,7 @@ void FDeferredShadingSceneRenderer::RenderTranslucencyInner(
 				SeparateTranslucencyColorTexture,
 				SeparateTranslucencyColorLoadAction,
 				SeparateTranslucencyDepthTexture.Target,
-				CreateTranslucentBasePassUniformBuffer(GraphBuilder, View, SceneColorCopyTexture, SceneTextureSetupMode, ViewIndex),
+				CreateTranslucentBasePassUniformBuffer(GraphBuilder, View, TranslucentVolumeLightingTextures, SceneColorCopyTexture, SceneTextureSetupMode, ViewIndex),
 				TranslucencyPass,
 				!bCompositeBackToSceneColor,
 				bRenderInParallel);
@@ -1118,7 +1121,7 @@ void FDeferredShadingSceneRenderer::RenderTranslucencyInner(
 				SceneColorTexture,
 				SceneColorLoadAction,
 				SceneDepthTexture.Target,
-				CreateTranslucentBasePassUniformBuffer(GraphBuilder, View, SceneColorCopyTexture, SceneTextureSetupMode, ViewIndex),
+				CreateTranslucentBasePassUniformBuffer(GraphBuilder, View, TranslucentVolumeLightingTextures, SceneColorCopyTexture, SceneTextureSetupMode, ViewIndex),
 				TranslucencyPass,
 				bResolveColorTexture,
 				bRenderInParallel);
@@ -1130,6 +1133,7 @@ void FDeferredShadingSceneRenderer::RenderTranslucencyInner(
 
 void FDeferredShadingSceneRenderer::RenderTranslucency(
 	FRDGBuilder& GraphBuilder,
+	const FTranslucentVolumeLightingTextures& TranslucentVolumeLightingTextures,
 	FRDGTextureMSAA SceneColorTexture,
 	FRDGTextureMSAA SceneDepthTexture,
 	FSeparateTranslucencyTextures* OutSeparateTranslucencyTextures,
@@ -1151,12 +1155,12 @@ void FDeferredShadingSceneRenderer::RenderTranslucency(
 
 	if (ViewFamily.AllowTranslucencyAfterDOF())
 	{
-		RenderTranslucencyInner(GraphBuilder, SceneColorTexture, SceneDepthTexture, OutSeparateTranslucencyTextures, ViewsToRender, SceneColorCopyTexture, ETranslucencyPass::TPT_StandardTranslucency);
-		RenderTranslucencyInner(GraphBuilder, SceneColorTexture, SceneDepthTexture, OutSeparateTranslucencyTextures, ViewsToRender, SceneColorCopyTexture, ETranslucencyPass::TPT_TranslucencyAfterDOF);
-		RenderTranslucencyInner(GraphBuilder, SceneColorTexture, SceneDepthTexture, OutSeparateTranslucencyTextures, ViewsToRender, SceneColorCopyTexture, ETranslucencyPass::TPT_TranslucencyAfterDOFModulate);
+		RenderTranslucencyInner(GraphBuilder, TranslucentVolumeLightingTextures, SceneColorTexture, SceneDepthTexture, OutSeparateTranslucencyTextures, ViewsToRender, SceneColorCopyTexture, ETranslucencyPass::TPT_StandardTranslucency);
+		RenderTranslucencyInner(GraphBuilder, TranslucentVolumeLightingTextures, SceneColorTexture, SceneDepthTexture, OutSeparateTranslucencyTextures, ViewsToRender, SceneColorCopyTexture, ETranslucencyPass::TPT_TranslucencyAfterDOF);
+		RenderTranslucencyInner(GraphBuilder, TranslucentVolumeLightingTextures, SceneColorTexture, SceneDepthTexture, OutSeparateTranslucencyTextures, ViewsToRender, SceneColorCopyTexture, ETranslucencyPass::TPT_TranslucencyAfterDOFModulate);
 	}
 	else // Otherwise render translucent primitives in a single bucket.
 	{
-		RenderTranslucencyInner(GraphBuilder, SceneColorTexture, SceneDepthTexture, OutSeparateTranslucencyTextures, ViewsToRender, SceneColorCopyTexture, ETranslucencyPass::TPT_AllTranslucency);
+		RenderTranslucencyInner(GraphBuilder, TranslucentVolumeLightingTextures, SceneColorTexture, SceneDepthTexture, OutSeparateTranslucencyTextures, ViewsToRender, SceneColorCopyTexture, ETranslucencyPass::TPT_AllTranslucency);
 	}
 }
diff --git a/Engine/Source/Runtime/Renderer/Private/TranslucentRendering.h b/Engine/Source/Runtime/Renderer/Private/TranslucentRendering.h
--- a/Engine/Source/Runtime/Renderer/Private/TranslucentRendering.h
+++ b/Engine/Source/Runtime/Renderer/Private/TranslucentRendering.h
@@ -62,6 +62,33 @@ private:
 	FRDGTextureMSAA ColorModulateTexture;
 };
 
+struct FTranslucentVolumeLightingTextures
+{
+	FRDGTextureRef GetAmbient(int32 ViewIndex, int32 CascadeIndex) const
+	{
+		return Ambient[(ViewIndex * TVC_MAX) + CascadeIndex];
+	}
+
+	void SetAmbient(int32 ViewIndex, int32 CascadeIndex, FRDGTextureRef Texture)
+	{
+		Ambient[(ViewIndex * TVC_MAX) + CascadeIndex] = Texture;
+	}
+
+	FRDGTextureRef GetDirectional(int32 ViewIndex, int32 CascadeIndex) const
+	{
+		return Directional[(ViewIndex * TVC_MAX) + CascadeIndex];
+	}
+
+	void SetDirectional(int32 ViewIndex, int32 CascadeIndex, FRDGTextureRef Texture)
+	{
+		Directional[(ViewIndex * TVC_MAX) + CascadeIndex] = Texture;
+	}
+
+	TArray<FRDGTextureRef, TInlineAllocator<TVC_MAX>> Ambient;
+	TArray<FRDGTextureRef, TInlineAllocator<TVC_MAX>> Directional;
+	int32 VolumeDim = 0;
+};
+
 DECLARE_GPU_DRAWCALL_STAT_EXTERN(Translucency);
 
 /** Creates separate translucency textures. */
