Fix adaptive sampling in path tracer

There were a number of issues that resulted in NaNs creeping up to the top level of the hierarchy. This in turn led to the adaptive codepath being consistently disabled. I also fixed the incremental calculation of variance (we want to measure variance of the mean) although it doesn't seem to be possible to compute an exact variance when we have more than one sample per pixel from a single pass (in other words, when adaptive has started focusing on some pixels).

I switched the gaussian blur to a simple max dilation to be a bit more robust to outlier pixels.

Despite the adaptive sampler working fairly well on high error regions, I found it generally slows down rendering significantly, so I opted to disable it by default.

diff --git a/Engine/Shaders/Private/PathTracing/BuildVarianceMipTreeComputeShader.usf b/Engine/Shaders/Private/PathTracing/BuildVarianceMipTreeComputeShader.usf
--- a/Engine/Shaders/Private/PathTracing/BuildVarianceMipTreeComputeShader.usf
+++ b/Engine/Shaders/Private/PathTracing/BuildVarianceMipTreeComputeShader.usf
@@ -111,36 +111,45 @@ void BuildVarianceMipTreeCS(
 	{
 		float2 UV = float2(GridCoordinate.xy + 0.5f) / VarianceMapDimensions;
 
-		// Resample blurred result as the variance map
-		Gaussian5x5 Filter;
-		Filter.Initialize();
-		uint Samples = 0;
-		int2 PixelCenter = DispatchThreadId.xy;
-		int2 PixelMin = max(PixelCenter - Filter.Offset, 0);
-		int2 PixelMax = min(PixelCenter + Filter.Offset, int2(ViewSize) - 1);
+		// Measure some error function over a NxN neighborhood so we explore the area around noisy pixels
+		int2 PixelCenter = UV * ViewSize;
+		int Radius = 2;
+		int2 PixelMin = max(PixelCenter - Radius, 0);
+		int2 PixelMax = min(PixelCenter + Radius, int2(ViewSize) - 1);
 		for (int y = PixelMin.y; y <= PixelMax.y; ++y)
 		{
 			for (int x = PixelMin.x; x <= PixelMax.x; ++x)
 			{
 				int2 Pixel = int2(x, y);
-				int2 Kernel = Pixel - PixelCenter + Filter.Offset;
-				float Weight = Filter.Weight[Kernel.y * Filter.Width + Kernel.x];
 
 				float4 RadianceAndVariance = RadianceTexture.Load(uint3(Pixel, 0));
 				float3 Radiance = RadianceAndVariance.rgb;
 				float Variance = RadianceAndVariance.a;
-				float Mean = Tonemap(Luminance(Radiance));
+				float Mean = Luminance(Radiance);
+#if 0
+				// measure 3 std-deviations either way and take the overall change in tonemapped intensity as a proxy for error
+				float StdError = sqrt(Variance);
+				float High = Tonemap(Mean + 3 * StdError);
+				float Low = Tonemap(max(Mean - 3 * StdError, 0.0));
+				float VisibleError = High - Low;
+#elif 0
+				// Original heuristic
 				float Epsilon = 0.0001;
-				float Sample = Variance / (Mean + Epsilon) * Weight;
-				Sum += Sample;
-				Samples++;
+				float VisibleError = Variance / (Tonemap(Mean) + Epsilon);
+#else
+				// Relative error (spends too much effort on really dark pixels?)
+				float VisibleError = sqrt(Variance) / (Mean + 0.0001);
+#endif
+
+				Sum = max(Sum, VisibleError);
 			}
 		}
-		Sum /= Samples;
+		// TODO: turn errors below some threshold to 0 so we stop working on those pixels entirely?
 	}
 	else
 	{
-		// Report the mean as the aggregate.
+		// Average errors for hierarchical importance sampling. Taking the average keeps the error values
+		// comparable across level which makes importance sampling more numerically stable.
 		uint2 PixelOffset = uint2(GridCoordinate.x << 1, GridCoordinate.y << 1);
 		for (int i = 0; i < 4; ++i)
 		{
diff --git a/Engine/Shaders/Private/PathTracing/PathTracingCompositingPixelShader.usf b/Engine/Shaders/Private/PathTracing/PathTracingCompositingPixelShader.usf
--- a/Engine/Shaders/Private/PathTracing/PathTracingCompositingPixelShader.usf
+++ b/Engine/Shaders/Private/PathTracing/PathTracingCompositingPixelShader.usf
@@ -16,16 +16,15 @@ Texture2D<uint> SampleCountTexture;
 Texture2D<float4> CumulativeIrradianceTexture;
 Texture2D<uint> CumulativeSampleCountTexture;
 
-// TODO: Make C1 continuous
 float TonemapGammaAndHDR(float Value)
 {
 	if (Value > 1.0)
 	{
-		Value = log(Value);
+		Value = 1.0 + 2.0 * log(Value);
 	}
 	else
 	{
-		Value = pow(Value, 2.2);
+		Value = pow(Value, 2.0);
 	}
 	return Value;
 }
@@ -55,39 +54,53 @@ void CompositeMain(
 	float RadianceBlue = FixedPointToFloat(RadianceBlueTexture.Load(TexCoord), 24);
 	float RadianceAlpha = FixedPointToFloat(RadianceAlphaTexture.Load(TexCoord), 24);
 	float4 Radiance = float4(RadianceRed, RadianceGreen, RadianceBlue, RadianceAlpha);
-	float SampleCount = SampleCountTexture.Load(TexCoord);
+	uint SampleCount = SampleCountTexture.Load(TexCoord);
+
+	if (SampleCount > 1)
+		Radiance /= float(SampleCount);
 
 	float4 CumulativeIrradiance = CumulativeIrradianceTexture.Load(TexCoord);
-	float CumulativeSampleCount = CumulativeSampleCountTexture.Load(TexCoord);
+	uint CumulativeSampleCount = CumulativeSampleCountTexture.Load(TexCoord);
 
-	float CumulativeVariance = CumulativeIrradiance.a;
+	// extract variance from variance of the mean
+	float CumulativeVariance = CumulativeIrradiance.a * CumulativeSampleCount;
 
 	OutCumulativeSampleCount = CumulativeSampleCount + SampleCount;
-	OutCumulativeIrradiance.rgb = (CumulativeIrradiance.rgb * CumulativeSampleCount + Radiance.rgb) / OutCumulativeSampleCount;
+    // Average radiance using the smaller blend factor for extra precision
+	float BlendFactor = float(SampleCount) / float(OutCumulativeSampleCount);
+	OutCumulativeIrradiance.rgb = lerp(CumulativeIrradiance.rgb, Radiance.rgb, BlendFactor);
 
 	float OutVariance = 0;
-	if (OutCumulativeSampleCount > 1 && OutCumulativeSampleCount < 4096)
-	{
-#if 0
-		// Per-sample variance
-		float IrradianceY = Tonemap(Luminance(Radiance));
-		float DeviationSquared = IrradianceY - Tonemap(Luminance(OutCumulativeIrradiance));
-#else
-		// Variance of the mean
-		float IrradianceY = Luminance(CumulativeIrradiance);
-		float DeviationSquared = IrradianceY - Luminance(OutCumulativeIrradiance);
-#endif
-		DeviationSquared *= DeviationSquared;
-		OutVariance = DeviationSquared;
-		if (CumulativeSampleCount > 0)
-		{
-			OutVariance += CumulativeVariance * (CumulativeSampleCount - 1);
-		}
-		OutVariance /= OutCumulativeSampleCount - 1;
-	}
-	OutCumulativeIrradiance.a = OutVariance;
+	float RadianceY = Luminance(Radiance.rgb);
+	float DeviationSquared = (RadianceY - Luminance(OutCumulativeIrradiance.rgb)) * (RadianceY - Luminance(CumulativeIrradiance.rgb));
+	// This formula is correct when SampleCount == 1 - but for larger values we would
+    // theoretically need to know the variance of our new batch of samples ...
+	OutVariance = lerp(CumulativeVariance, DeviationSquared, BlendFactor);
+
+	// store variance of the mean
+	OutCumulativeIrradiance.a = OutCumulativeSampleCount > 0 ? OutVariance / OutCumulativeSampleCount : 0;
 	OutColor = float4(OutCumulativeIrradiance.rgb * View.PreExposure, 1.0);
 
+#if 0
+    // split-screen debugging aids
+	if (UV.x < 0.5)
+	{
+		float Lum = Luminance(OutCumulativeIrradiance.rgb); // mean
+		float StdErr = OutCumulativeSampleCount > 0 ? sqrt(OutVariance / OutCumulativeSampleCount) : 0; // stderr
+		//OutColor = Tonemap(Lum + 3 * StdErr) - Tonemap(max(Lum - 3 * StdErr, 0.0)); // error estimator
+		//OutColor = StdErr * View.PreExposure;
+		OutColor = StdErr / (Lum + 0.0001) * View.PreExposure;
+		//OutColor = OutVariance / (Tonemap(Lum) + 0.0001);
+		//OutColor = OutVariance * View.PreExposure;
+		//OutColor = OutCumulativeSampleCount / 1024.0 * View.PreExposure;
+		//OutColor.r = AdaptiveSamplingData.VarianceDimensions.x + AdaptiveSamplingData.VarianceDimensions.x;
+		//OutColor.g = AdaptiveSamplingData.VarianceDimensions.x;
+		//OutColor.b = AdaptiveSamplingData.VarianceDimensions.y;
+
+		//OutColor = AdaptiveSamplingData.VarianceMipTree[BufferOffsetAtPixel(TexCoord.xy, 0)];
+
+	}
+#endif
 	// To visualize the current frame:
-	// OutColor = float4(Irradiance.rgb / SampleCount, 1.0);
+	//OutColor = float4(Radiance.rgb / SampleCount, 1.0);
 }
\ No newline at end of file
diff --git a/Engine/Shaders/Private/PathTracing/Sensor/PathTracingAdaptivePinholeCamera.ush b/Engine/Shaders/Private/PathTracing/Sensor/PathTracingAdaptivePinholeCamera.ush
--- a/Engine/Shaders/Private/PathTracing/Sensor/PathTracingAdaptivePinholeCamera.ush
+++ b/Engine/Shaders/Private/PathTracing/Sensor/PathTracingAdaptivePinholeCamera.ush
@@ -71,50 +71,67 @@ void AdaptivePinholeCamera_SampleEmitter(
 	uint2 MipPixel = 0;
 	uint MipCount = Log2(AdaptiveSamplingData.VarianceDimensions.x);
 
+	float2 RandSample = RandomSequence_GenerateSample2D(RandSequence, 0);
 	uint StopLevel = 0;
+
+	float OneMinusEps = 0.999999940395; // float just before 1.0
 	for (uint Index = 0; Index < MipCount - StopLevel; ++Index)
 	{
 		uint MipIndex = MipCount - Index - 1;
-		float4 VarianceCdf;
-		VarianceCdf.x = AdaptiveSamplingData.VarianceMipTree[BufferOffsetAtPixel(MipPixel, MipIndex)];
-		VarianceCdf.y = VarianceCdf.x + AdaptiveSamplingData.VarianceMipTree[BufferOffsetAtPixel(MipPixel + uint2(1, 0), MipIndex)];
-		VarianceCdf.z = VarianceCdf.y + AdaptiveSamplingData.VarianceMipTree[BufferOffsetAtPixel(MipPixel + uint2(0, 1), MipIndex)];
-		VarianceCdf.w = VarianceCdf.z + AdaptiveSamplingData.VarianceMipTree[BufferOffsetAtPixel(MipPixel + uint2(1, 1), MipIndex)];
-		VarianceCdf /= VarianceCdf.w;
-
-		float RandSample = RandomSequence_GenerateSample1D(RandSequence, 0);
-		if (RandSample < VarianceCdf.x)
+		float Pdf00 = AdaptiveSamplingData.VarianceMipTree[BufferOffsetAtPixel(MipPixel, MipIndex)];
+		float Pdf01 = AdaptiveSamplingData.VarianceMipTree[BufferOffsetAtPixel(MipPixel + uint2(0, 1), MipIndex)];
+		float Pdf10 = AdaptiveSamplingData.VarianceMipTree[BufferOffsetAtPixel(MipPixel + uint2(1, 0), MipIndex)];
+		float Pdf11 = AdaptiveSamplingData.VarianceMipTree[BufferOffsetAtPixel(MipPixel + uint2(1, 1), MipIndex)];
+
+		// first decide between left / right:
+		float PdfSum = (Pdf00 + Pdf01) + (Pdf10 + Pdf11);
+		float ProbLeft = (Pdf00 + Pdf01) / PdfSum;
+
+		if (RandSample.x < ProbLeft)
 		{
-			OutPdf *= VarianceCdf.x;
+			RandSample.x = min(RandSample.x / ProbLeft, OneMinusEps);
+			// chose left - warp x and choose between the top / bottom pixels
+
+			float ProbTop = Pdf00 / (Pdf00 + Pdf01);
+			if (RandSample.y < ProbTop)
+			{
+				RandSample.y = min(RandSample.y / ProbTop, OneMinusEps);
+				OutPdf *= 4 * Pdf00 / PdfSum;
+			}
+			else
+			{
+				RandSample.y = min((1.0 - RandSample.y) / (1.0 - ProbTop), OneMinusEps);
+				OutPdf *= 4 * Pdf01 / PdfSum;
+				MipPixel.y += 1;
+			}
 		}
-		else if (RandSample < VarianceCdf.y)
+		else
 		{
-			OutPdf *= VarianceCdf.y - VarianceCdf.x;
+			// chose right - warp x and choose between the top / bottom pixels
+			RandSample.x = min((RandSample.x - ProbLeft) / (1.0 - ProbLeft), OneMinusEps);
+
 			MipPixel.x += 1;
-		}
-		else if (RandSample < VarianceCdf.z)
-		{
-			OutPdf *= VarianceCdf.z - VarianceCdf.y;
-			MipPixel.y += 1;
-		}
-		else // if (RandSample < VarianceCdf.w)
-		{
-			OutPdf *= 1.0 - VarianceCdf.z;
-			MipPixel += 1;
+
+			float ProbTop = Pdf10 / (Pdf10 + Pdf11);
+			if (RandSample.y < ProbTop)
+			{
+				RandSample.y = min(RandSample.y / ProbTop, OneMinusEps);
+				OutPdf *= 4 * Pdf10 / PdfSum;
+			}
+			else
+			{
+				RandSample.y = min((1.0 - RandSample.y) / (1.0 - ProbTop), OneMinusEps);
+				OutPdf *= 4 * Pdf11 / PdfSum;
+				MipPixel.y += 1;
+			}
 		}
 
-		MipPixel = uint2(MipPixel.x << 1, MipPixel.y << 1);
+		MipPixel <<= 1;
 	}
-#if 0
-	// Applying correct normalization causes catastrophic cancelation. Instead, OutPdf is only proportional to the correct weight
-	float NormalizationFactor = View.BufferSizeAndInvSize.x * View.BufferSizeAndInvSize.y;
-	float AreaPdf = (DimensionsAtLevel(StopLevel).x * DimensionsAtLevel(StopLevel).y) / (DimensionsAtLevel(0).x * DimensionsAtLevel(0).y);
-	OutPdf *= AreaPdf * NormalizationFactor;
-#endif
-	MipPixel = uint2(MipPixel.x >> 1, MipPixel.y >> 1);
-
-	float2 PixelOffset = RandomSequence_GenerateSample2D(RandSequence, 0);
-	float2 MipPixelUV = (MipPixel + PixelOffset) / float2(DimensionsAtLevel(StopLevel));
+
+	MipPixel >>= 1;
+
+	float2 MipPixelUV = (MipPixel + RandSample) / float2(DimensionsAtLevel(StopLevel));
 	// Correct for ViewSize/BufferSize discrepancy
 	MipPixelUV *= View.ViewSizeAndInvSize.xy * View.BufferSizeAndInvSize.zw;
 
diff --git a/Engine/Source/Runtime/Renderer/Private/PathTracing.cpp b/Engine/Source/Runtime/Renderer/Private/PathTracing.cpp
--- a/Engine/Source/Runtime/Renderer/Private/PathTracing.cpp
+++ b/Engine/Source/Runtime/Renderer/Private/PathTracing.cpp
@@ -66,10 +66,10 @@ TAutoConsoleVariable<int32> CVarPathTracingFrameIndependentTemporalSeed(
 
 TAutoConsoleVariable<int32> CVarPathTracingAdaptiveSampling(
 	TEXT("r.PathTracing.AdaptiveSampling"),
-	1,
+	0,
 	TEXT("Toggles the use of adaptive sampling\n")
-	TEXT("0: off\n")
-	TEXT("1: on (default)\n"),
+	TEXT("0: off (default)\n")
+	TEXT("1: on\n"),
 	ECVF_RenderThreadSafe
 );
 
@@ -348,20 +348,21 @@ public:
 
 		// Adaptive sampling
 		{
-			uint32 TemporalSeed;
 			if (CVarPathTracingFrameIndependentTemporalSeed.GetValueOnRenderThread() == 0)
 			{
-				TemporalSeed = Iteration;
+				// Count samples from 0 for deterministic results
+				AdaptiveSamplingData.TemporalSeed = Iteration;
 			}
 			else
 			{
-				TemporalSeed = FrameIndependentTemporalSeed;
+				// Count samples from an ever-increasing counter to avoid screen-door effect
+				AdaptiveSamplingData.TemporalSeed = FrameIndependentTemporalSeed;
 			}
 
+			AdaptiveSamplingData.Iteration = Iteration;
+
 			if (VarianceMipTree.NumBytes > 0)
 			{
-				AdaptiveSamplingData.Iteration = Iteration;
-				AdaptiveSamplingData.TemporalSeed = TemporalSeed;
 				AdaptiveSamplingData.VarianceDimensions = VarianceDimensions;
 				AdaptiveSamplingData.VarianceMipTree = VarianceMipTree.SRV;
 				AdaptiveSamplingData.MinimumSamplesPerPixel = CVarPathTracingAdaptiveSamplingMinimumSamplesPerPixel.GetValueOnRenderThread();
@@ -369,8 +370,6 @@ public:
 			else
 			{
 				AdaptiveSamplingData.UseAdaptiveSampling = 0;
-				AdaptiveSamplingData.Iteration = Iteration;
-				AdaptiveSamplingData.TemporalSeed = TemporalSeed;
 				AdaptiveSamplingData.VarianceDimensions = FIntVector(1, 1, 1);
 				AdaptiveSamplingData.VarianceMipTree = RHICreateShaderResourceView(GBlackTexture->TextureRHI->GetTexture2D(), 0);
 				AdaptiveSamplingData.MinimumSamplesPerPixel = CVarPathTracingAdaptiveSamplingMinimumSamplesPerPixel.GetValueOnRenderThread();
@@ -567,8 +566,9 @@ void FDeferredShadingSceneRenderer::RenderPathTracing(
 			ViewState->PathTracingInvalidate();
 			ViewState->PathTracingRect = View.ViewRect;
 		}
-		
 
+		bool NeedsMoreRays = ViewState->PathTracingSPP < MaxSPP;
+		
 		// Construct render targets for compositing
 		TRefCountPtr<IPooledRenderTarget> RadianceRT;
 		TRefCountPtr<IPooledRenderTarget> SampleCountRT;
@@ -612,7 +612,7 @@ void FDeferredShadingSceneRenderer::RenderPathTracing(
 
 		bool bDoMGPUPathTracing = GNumExplicitGPUsForRendering > 1 && GPUCount > 1;
 
-		if (bDoMGPUPathTracing && ViewState->PathTracingSPP < MaxSPP)
+		if (bDoMGPUPathTracing && NeedsMoreRays)
 		{
 			//#dxr-todo: Set minimum tile size for mGPU
 			int32 TileSizeX = FIntPoint::DivideAndRoundUp(View.ViewRect.Size(), GPUCount).X;
@@ -720,7 +720,7 @@ void FDeferredShadingSceneRenderer::RenderPathTracing(
 			}
 #endif
 		}
-		else if (ViewState->PathTracingSPP < MaxSPP)
+		else if (NeedsMoreRays)
 		{
 			FIntVector TileOffset;
 			TileOffset.X = bWiperMode > 0 ? WipeOffsetX : 0;
@@ -760,7 +760,6 @@ void FDeferredShadingSceneRenderer::RenderPathTracing(
 			);
 		}
 
-
 		// Save RayTracingIndirect for compositing
 		RHICmdList.CopyToResolveTarget(RadianceRT->GetRenderTargetItem().TargetableTexture, RadianceRT->GetRenderTargetItem().ShaderResourceTexture, FResolveParams());
 		RHICmdList.CopyToResolveTarget(SampleCountRT->GetRenderTargetItem().TargetableTexture, SampleCountRT->GetRenderTargetItem().ShaderResourceTexture, FResolveParams());
@@ -795,25 +794,28 @@ void FDeferredShadingSceneRenderer::RenderPathTracing(
 		RHICmdList.ClearUAVUint(RadianceSortedAlphaRT->GetRenderTargetItem().UAV, BlackColor);
 		RHICmdList.ClearUAVUint(SampleCountSortedRT->GetRenderTargetItem().UAV, BlackColor);
 
-		ComputePathCompaction(
-			RHICmdList,
-			View,
-			RadianceRT->GetRenderTargetItem().ShaderResourceTexture,
-			SampleCountRT->GetRenderTargetItem().ShaderResourceTexture,
-			PixelPositionRT->GetRenderTargetItem().ShaderResourceTexture,
-			RadianceSortedRedRT->GetRenderTargetItem().UAV,
-			RadianceSortedGreenRT->GetRenderTargetItem().UAV,
-			RadianceSortedBlueRT->GetRenderTargetItem().UAV,
-			RadianceSortedAlphaRT->GetRenderTargetItem().UAV,
-			SampleCountSortedRT->GetRenderTargetItem().UAV
-		);
-
-		RHICmdList.CopyToResolveTarget(RadianceSortedRedRT->GetRenderTargetItem().TargetableTexture, RadianceSortedRedRT->GetRenderTargetItem().ShaderResourceTexture, FResolveParams());
-		RHICmdList.CopyToResolveTarget(RadianceSortedGreenRT->GetRenderTargetItem().TargetableTexture, RadianceSortedGreenRT->GetRenderTargetItem().ShaderResourceTexture, FResolveParams());
-		RHICmdList.CopyToResolveTarget(RadianceSortedBlueRT->GetRenderTargetItem().TargetableTexture, RadianceSortedBlueRT->GetRenderTargetItem().ShaderResourceTexture, FResolveParams());
-		RHICmdList.CopyToResolveTarget(RadianceSortedAlphaRT->GetRenderTargetItem().TargetableTexture, RadianceSortedAlphaRT->GetRenderTargetItem().ShaderResourceTexture, FResolveParams());
-		RHICmdList.CopyToResolveTarget(SampleCountSortedRT->GetRenderTargetItem().TargetableTexture, SampleCountSortedRT->GetRenderTargetItem().ShaderResourceTexture, FResolveParams());
+		if (NeedsMoreRays)
+		{
+			// only need to compact paths if we've traced rays in the last pass
+			ComputePathCompaction(
+				RHICmdList,
+				View,
+				RadianceRT->GetRenderTargetItem().ShaderResourceTexture,
+				SampleCountRT->GetRenderTargetItem().ShaderResourceTexture,
+				PixelPositionRT->GetRenderTargetItem().ShaderResourceTexture,
+				RadianceSortedRedRT->GetRenderTargetItem().UAV,
+				RadianceSortedGreenRT->GetRenderTargetItem().UAV,
+				RadianceSortedBlueRT->GetRenderTargetItem().UAV,
+				RadianceSortedAlphaRT->GetRenderTargetItem().UAV,
+				SampleCountSortedRT->GetRenderTargetItem().UAV
+			);
 
+			RHICmdList.CopyToResolveTarget(RadianceSortedRedRT->GetRenderTargetItem().TargetableTexture, RadianceSortedRedRT->GetRenderTargetItem().ShaderResourceTexture, FResolveParams());
+			RHICmdList.CopyToResolveTarget(RadianceSortedGreenRT->GetRenderTargetItem().TargetableTexture, RadianceSortedGreenRT->GetRenderTargetItem().ShaderResourceTexture, FResolveParams());
+			RHICmdList.CopyToResolveTarget(RadianceSortedBlueRT->GetRenderTargetItem().TargetableTexture, RadianceSortedBlueRT->GetRenderTargetItem().ShaderResourceTexture, FResolveParams());
+			RHICmdList.CopyToResolveTarget(RadianceSortedAlphaRT->GetRenderTargetItem().TargetableTexture, RadianceSortedAlphaRT->GetRenderTargetItem().ShaderResourceTexture, FResolveParams());
+			RHICmdList.CopyToResolveTarget(SampleCountSortedRT->GetRenderTargetItem().TargetableTexture, SampleCountSortedRT->GetRenderTargetItem().ShaderResourceTexture, FResolveParams());
+		}
 
 		// Construct render targets for compositing
 		TRefCountPtr<IPooledRenderTarget> OutputRadianceRT;
@@ -905,12 +907,15 @@ void FDeferredShadingSceneRenderer::RenderPathTracing(
 		ViewState->PathTracingIrradianceRT = OutputRadianceRT;
 		ViewState->PathTracingSampleCountRT = OutputSampleCountRT;
 
-		// Process variance mip for adaptive sampling
-		if (ViewState->PathTracingSPP % CVarPathTracingVarianceMapRebuildFrequency.GetValueOnRenderThread() == 0)
+		if (NeedsMoreRays)
 		{
-			SCOPED_GPU_STAT(RHICmdList, Stat_GPU_PathTracingBuildVarianceMipTree);
+			// Process variance mip for adaptive sampling
+			if (ViewState->PathTracingSPP % CVarPathTracingVarianceMapRebuildFrequency.GetValueOnRenderThread() == 0)
+			{
+				SCOPED_GPU_STAT(RHICmdList, Stat_GPU_PathTracingBuildVarianceMipTree);
 
-			BuildVarianceMipTree(RHICmdList, View, OutputRadianceRT->GetRenderTargetItem().ShaderResourceTexture, *ViewState->VarianceMipTree, ViewState->VarianceMipTreeDimensions);
+				BuildVarianceMipTree(RHICmdList, View, OutputRadianceRT->GetRenderTargetItem().ShaderResourceTexture, *ViewState->VarianceMipTree, ViewState->VarianceMipTreeDimensions);
+			}
 		}
 
 		VisualizeVarianceMipTree(RHICmdList, View, *ViewState->VarianceMipTree, ViewState->VarianceMipTreeDimensions);
