Clean up of base pass uniform buffers to only include RDG versions.

#rnx

diff --git a/Engine/Source/Runtime/Renderer/Private/BasePassRendering.cpp b/Engine/Source/Runtime/Renderer/Private/BasePassRendering.cpp
--- a/Engine/Source/Runtime/Renderer/Private/BasePassRendering.cpp
+++ b/Engine/Source/Runtime/Renderer/Private/BasePassRendering.cpp
@@ -539,20 +539,19 @@ bool GetBasePassShaders<FUniformLightMapPolicy>(
 extern void SetupFogUniformParameters(FRDGBuilder* GraphBuilder, const FViewInfo& View, FFogUniformParameters& OutParameters);
 
 void SetupSharedBasePassParameters(
-	FRDGBuilder* GraphBuilder,
-	FRHICommandListImmediate& RHICmdList,
+	FRDGBuilder& GraphBuilder,
 	const FViewInfo& View,
 	FSharedBasePassUniformParameters& SharedParameters)
 {
 	SharedParameters.Forward = View.ForwardLightingResources->ForwardLightData;
 
-	SetupFogUniformParameters(GraphBuilder, View, SharedParameters.Fog);
+	SetupFogUniformParameters(&GraphBuilder, View, SharedParameters.Fog);
 
 	if (View.IsInstancedStereoPass())
 	{
 		const FSceneView& RightEye = *View.Family->Views[1];
 		SharedParameters.ForwardISR = RightEye.ForwardLightingResources->ForwardLightData;
-		SetupFogUniformParameters(GraphBuilder, (FViewInfo&)RightEye, SharedParameters.FogISR);
+		SetupFogUniformParameters(&GraphBuilder, (FViewInfo&)RightEye, SharedParameters.FogISR);
 	}
 	else
 	{
@@ -568,15 +567,14 @@ void SetupSharedBasePassParameters(
 
 	SharedParameters.SSProfilesTexture = GBlackTexture->TextureRHI;
 
-	if (const IPooledRenderTarget* PooledRT = GetSubsufaceProfileTexture_RT(RHICmdList))
+	if (const IPooledRenderTarget* PooledRT = GetSubsufaceProfileTexture_RT(GraphBuilder.RHICmdList))
 	{
 		SharedParameters.SSProfilesTexture = PooledRT->GetShaderResourceRHI();
 	}
 }
 
 void SetupSharedOpaqueBasePassParameters(
-	FRDGBuilder* GraphBuilder,
-	FRHICommandListImmediate& RHICmdList,
+	FRDGBuilder& GraphBuilder,
 	const FSceneRenderTargets& SceneRenderTargets,
 	const FViewInfo& View,
 	FRDGTextureRef ForwardScreenSpaceShadowMask,
@@ -586,10 +584,10 @@ void SetupSharedOpaqueBasePassParameters(
 {
 	const auto GetRDG = [&](const TRefCountPtr<IPooledRenderTarget>& PooledRenderTarget, ERDGTextureFlags Flags = ERDGTextureFlags::None)
 	{
-		return RegisterExternalOrPassthroughTexture(GraphBuilder, PooledRenderTarget, Flags);
+		return GraphBuilder.RegisterExternalTexture(PooledRenderTarget, ERenderTargetTexture::ShaderResource, Flags);
 	};
 
-	SetupSharedBasePassParameters(GraphBuilder, RHICmdList, View, BasePassParameters.Shared);
+	SetupSharedBasePassParameters(GraphBuilder, View, BasePassParameters.Shared);
 
 	FRDGTextureRef WhiteDummy = GetRDG(GSystemTextures.WhiteDummy);
 	FRDGTextureRef BlackDummy = GetRDG(GSystemTextures.BlackDummy);
@@ -693,22 +691,10 @@ TRDGUniformBufferRef<FOpaqueBasePassUniformParameters> CreateOpaqueBasePassUnifo
 {
 	FSceneRenderTargets& SceneRenderTargets = FSceneRenderTargets::Get();
 	FOpaqueBasePassUniformParameters* BasePassParameters = GraphBuilder.AllocParameters<FOpaqueBasePassUniformParameters>();
-	SetupSharedOpaqueBasePassParameters(&GraphBuilder, GraphBuilder.RHICmdList, SceneRenderTargets, View, ForwardScreenSpaceShadowMask, SceneWithoutWaterTextures, ViewIndex, *BasePassParameters);
+	SetupSharedOpaqueBasePassParameters(GraphBuilder, SceneRenderTargets, View, ForwardScreenSpaceShadowMask, SceneWithoutWaterTextures, ViewIndex, *BasePassParameters);
 	return GraphBuilder.CreateUniformBuffer(BasePassParameters);
 }
 
-TUniformBufferRef<FOpaqueBasePassUniformParameters> CreateOpaqueBasePassUniformBuffer(
-	FRHICommandListImmediate& RHICmdList,
-	const FViewInfo& View,
-	IPooledRenderTarget* ForwardScreenSpaceShadowMask)
-{
-	FRDGTextureRef ForwardScreenSpaceShadowMaskRDG = FRDGTexture::GetPassthrough(ForwardScreenSpaceShadowMask);
-
-	FOpaqueBasePassUniformParameters BasePassParameters;
-	SetupSharedOpaqueBasePassParameters(nullptr, RHICmdList, FSceneRenderTargets::Get(), View, ForwardScreenSpaceShadowMaskRDG, nullptr, 0, BasePassParameters);
-	return TUniformBufferRef<FOpaqueBasePassUniformParameters>::CreateUniformBufferImmediate(BasePassParameters, UniformBuffer_SingleFrame);
-}
-
 static void ClearGBufferAtMaxZ(
 	FRDGBuilder& GraphBuilder,
 	TArrayView<const FViewInfo> Views,
diff --git a/Engine/Source/Runtime/Renderer/Private/BasePassRendering.h b/Engine/Source/Runtime/Renderer/Private/BasePassRendering.h
--- a/Engine/Source/Runtime/Renderer/Private/BasePassRendering.h
+++ b/Engine/Source/Runtime/Renderer/Private/BasePassRendering.h
@@ -129,16 +129,10 @@ END_GLOBAL_SHADER_PARAMETER_STRUCT()
 DECLARE_GPU_DRAWCALL_STAT_EXTERN(Basepass);
 
 extern void SetupSharedBasePassParameters(
-	FRDGBuilder* GraphBuilder,
-	FRHICommandListImmediate& RHICmdList,
+	FRDGBuilder& GraphBuilder,
 	const FViewInfo& View,
 	FSharedBasePassUniformParameters& BasePassParameters);
 
-extern TUniformBufferRef<FOpaqueBasePassUniformParameters> CreateOpaqueBasePassUniformBuffer(
-	FRHICommandListImmediate& RHICmdList,
-	const FViewInfo& View,
-	IPooledRenderTarget* ForwardScreenSpaceShadowMask);
-
 extern TRDGUniformBufferRef<FOpaqueBasePassUniformParameters> CreateOpaqueBasePassUniformBuffer(
 	FRDGBuilder& GraphBuilder,
 	const FViewInfo& View,
@@ -154,12 +148,6 @@ extern TRDGUniformBufferRef<FTranslucentBasePassUniformParameters> CreateTranslu
 	ESceneTextureSetupMode SceneTextureSetupMode,
 	const int32 ViewIndex);
 
-extern TUniformBufferRef<FTranslucentBasePassUniformParameters> CreateTranslucentBasePassUniformBuffer(
-	FRHICommandListImmediate& RHICmdList,
-	const FViewInfo& View,
-	ESceneTextureSetupMode SceneTextureSetupMode,
-	const int32 ViewIndex);
-
 /** Parameters for computing forward lighting. */
 class FForwardLightingParameters
 {
diff --git a/Engine/Source/Runtime/Renderer/Private/DebugViewModeRendering.cpp b/Engine/Source/Runtime/Renderer/Private/DebugViewModeRendering.cpp
--- a/Engine/Source/Runtime/Renderer/Private/DebugViewModeRendering.cpp
+++ b/Engine/Source/Runtime/Renderer/Private/DebugViewModeRendering.cpp
@@ -69,15 +69,6 @@ void SetupDebugViewModePassUniformBufferConstants(const FViewInfo& ViewInfo, FDe
 	}
 }
 
-TUniformBufferRef<FDebugViewModePassUniformParameters> CreateDebugViewModePassUniformBuffer(FRHICommandList& RHICmdList, const FViewInfo& View)
-{
-	FDebugViewModePassUniformParameters Parameters;
-	FSceneRenderTargets& SceneContext = FSceneRenderTargets::Get();
-	SetupSceneTextureUniformParameters(SceneContext, View.FeatureLevel, ESceneTextureSetupMode::None, Parameters.SceneTextures);
-	SetupDebugViewModePassUniformBufferConstants(View, Parameters);
-	return TUniformBufferRef<FDebugViewModePassUniformParameters>::CreateUniformBufferImmediate(Parameters, UniformBuffer_SingleFrame);
-}
-
 TRDGUniformBufferRef<FDebugViewModePassUniformParameters> CreateDebugViewModePassUniformBuffer(FRDGBuilder& GraphBuilder, const FViewInfo& View)
 {
 	auto* UniformBufferParameters = GraphBuilder.AllocParameters<FDebugViewModePassUniformParameters>();
@@ -127,11 +118,6 @@ void FDeferredShadingSceneRenderer::RenderDebugViewMode(FRDGBuilder& GraphBuilde
 	}
 }
 
-FDebugViewModePS::FDebugViewModePS(const FMeshMaterialShaderType::CompiledShaderInitializerType& Initializer) : FMeshMaterialShader(Initializer)
-{
-	PassUniformBuffer.Bind(Initializer.ParameterMap, FDebugViewModePassUniformParameters::StaticStructMetadata.GetShaderVariableName());
-}
-
 void FDebugViewModePS::GetElementShaderBindings(
 	const FShaderMapPointerTable& PointerTable,
 	const FScene* Scene,
@@ -178,12 +164,10 @@ FDebugViewModeMeshProcessor::FDebugViewModeMeshProcessor(
 	const FScene* InScene, 
 	ERHIFeatureLevel::Type InFeatureLevel,
 	const FSceneView* InViewIfDynamicMeshCommand, 
-	FRHIUniformBuffer* InPassUniformBuffer,
 	bool bTranslucentBasePass,
 	FMeshPassDrawListContext* InDrawListContext
 )
 	: FMeshPassProcessor(InScene, InFeatureLevel, InViewIfDynamicMeshCommand, InDrawListContext)
-	, PassUniformBuffer(InPassUniformBuffer)
 	, DebugViewMode(DVSM_None)
 	, ViewModeParam(INDEX_NONE)
 	, DebugViewModeInterface(nullptr)
@@ -270,7 +254,6 @@ void FDebugViewModeMeshProcessor::AddMeshBatch(const FMeshBatch& RESTRICT MeshBa
 
 	FMeshPassProcessorRenderState DrawRenderState;
 	DrawRenderState.SetViewUniformBuffer(ViewUniformBuffer);
-	DrawRenderState.SetPassUniformBuffer(PassUniformBuffer);
 
 	FDebugViewModeInterface::FRenderState InterfaceRenderState;
 	DebugViewModeInterface->SetDrawRenderState(BatchMaterial->GetBlendMode(), InterfaceRenderState, Scene ? (Scene->GetShadingPath() == EShadingPath::Deferred && Scene->EarlyZPassMode != DDM_NonMaskedOnly) : false);
@@ -370,7 +353,7 @@ void FDebugViewModeMeshProcessor::UpdateInstructionCount(FDebugViewModeShaderEle
 FMeshPassProcessor* CreateDebugViewModePassProcessor(const FScene* Scene, const FSceneView* InViewIfDynamicMeshCommand, FMeshPassDrawListContext* InDrawListContext)
 {
 	const ERHIFeatureLevel::Type FeatureLevel = Scene ? Scene->GetFeatureLevel() : (InViewIfDynamicMeshCommand ? InViewIfDynamicMeshCommand->GetFeatureLevel() : GMaxRHIFeatureLevel);
-	return new(FMemStack::Get()) FDebugViewModeMeshProcessor(Scene, FeatureLevel, InViewIfDynamicMeshCommand, nullptr, false, InDrawListContext);
+	return new(FMemStack::Get()) FDebugViewModeMeshProcessor(Scene, FeatureLevel, InViewIfDynamicMeshCommand, false, InDrawListContext);
 }
 
 FRegisterPassProcessorCreateFunction RegisterDebugViewModeMobilePass(&CreateDebugViewModePassProcessor, EShadingPath::Mobile, EMeshPass::DebugViewMode, EMeshPassFlags::MainView);
diff --git a/Engine/Source/Runtime/Renderer/Private/DebugViewModeRendering.h b/Engine/Source/Runtime/Renderer/Private/DebugViewModeRendering.h
--- a/Engine/Source/Runtime/Renderer/Private/DebugViewModeRendering.h
+++ b/Engine/Source/Runtime/Renderer/Private/DebugViewModeRendering.h
@@ -33,7 +33,6 @@ END_GLOBAL_SHADER_PARAMETER_STRUCT()
 
 #if !(UE_BUILD_SHIPPING || UE_BUILD_TEST)
 
-TUniformBufferRef<FDebugViewModePassUniformParameters> CreateDebugViewModePassUniformBuffer(FRHICommandList& RHICmdList, const FViewInfo& View);
 TRDGUniformBufferRef<FDebugViewModePassUniformParameters> CreateDebugViewModePassUniformBuffer(FRDGBuilder& GraphBuilder, const FViewInfo& View);
 
 class FDebugViewModeShaderElementData : public FMeshMaterialShaderElementData
@@ -80,12 +79,10 @@ class FDebugViewModeVS : public FMeshMaterialShader
 	DECLARE_SHADER_TYPE(FDebugViewModeVS,MeshMaterial);
 protected:
 
-	FDebugViewModeVS(const FMeshMaterialShaderType::CompiledShaderInitializerType& Initializer) : FMeshMaterialShader(Initializer)
-	{
-		PassUniformBuffer.Bind(Initializer.ParameterMap, FSceneTextureUniformParameters::StaticStructMetadata.GetShaderVariableName());
-	}
-
-	FDebugViewModeVS() {}
+	FDebugViewModeVS() = default;
+	FDebugViewModeVS(const FMeshMaterialShaderType::CompiledShaderInitializerType& Initializer) 
+		: FMeshMaterialShader(Initializer)
+	{}
 
 public:
 
@@ -145,9 +142,10 @@ public:
 		FBaseHS::ModifyCompilationEnvironment(Parameters, OutEnvironment);
 	}
 
-
-	FDebugViewModeHS(const ShaderMetaType::CompiledShaderInitializerType& Initializer): FBaseHS(Initializer) {}
-	FDebugViewModeHS() {}
+	FDebugViewModeHS() = default;
+	FDebugViewModeHS(const ShaderMetaType::CompiledShaderInitializerType& Initializer)
+		: FBaseHS(Initializer)
+	{}
 };
 
 /**
@@ -169,15 +167,19 @@ public:
 		FBaseDS::ModifyCompilationEnvironment(Parameters, OutEnvironment);
 	}
 
-	FDebugViewModeDS(const ShaderMetaType::CompiledShaderInitializerType& Initializer): FBaseDS(Initializer) {}
-	FDebugViewModeDS() {}
+	FDebugViewModeDS() = default;
+	FDebugViewModeDS(const ShaderMetaType::CompiledShaderInitializerType& Initializer)
+		: FBaseDS(Initializer)
+	{}
 };
 
 class FDebugViewModePS : public FMeshMaterialShader
 {
 public:
-	FDebugViewModePS(const FMeshMaterialShaderType::CompiledShaderInitializerType& Initializer);
-	FDebugViewModePS() {}
+	FDebugViewModePS() = default;
+	FDebugViewModePS(const FMeshMaterialShaderType::CompiledShaderInitializerType& Initializer)
+		: FMeshMaterialShader(Initializer)
+	{}
 
 	void GetElementShaderBindings(
 		const FShaderMapPointerTable& PointerTable,
@@ -197,7 +199,7 @@ public:
 class FDebugViewModeMeshProcessor : public FMeshPassProcessor
 {
 public:
-	FDebugViewModeMeshProcessor(const FScene* InScene, ERHIFeatureLevel::Type InFeatureLevel, const FSceneView* InViewIfDynamicMeshCommand, FRHIUniformBuffer* InPassUniformBuffer, bool bTranslucentBasePass, FMeshPassDrawListContext* InDrawListContext);
+	FDebugViewModeMeshProcessor(const FScene* InScene, ERHIFeatureLevel::Type InFeatureLevel, const FSceneView* InViewIfDynamicMeshCommand, bool bTranslucentBasePass, FMeshPassDrawListContext* InDrawListContext);
 	virtual void AddMeshBatch(const FMeshBatch& RESTRICT MeshBatch, uint64 BatchElementMask, const FPrimitiveSceneProxy* RESTRICT PrimitiveSceneProxy, int32 StaticMeshId = -1) override final;
 
 private:
@@ -205,7 +207,6 @@ private:
 	void UpdateInstructionCount(FDebugViewModeShaderElementData& OutShaderElementData, const FMaterial* InBatchMaterial, FVertexFactoryType* InVertexFactoryType);
 
 	TUniformBufferRef<FViewUniformShaderParameters> ViewUniformBuffer;
-	FUniformBufferRHIRef PassUniformBuffer;
 	EDebugViewShaderMode DebugViewMode;
 	int32 ViewModeParam;
 	FName ViewModeParamName;
diff --git a/Engine/Source/Runtime/Renderer/Private/MobileShadingRenderer.cpp b/Engine/Source/Runtime/Renderer/Private/MobileShadingRenderer.cpp
--- a/Engine/Source/Runtime/Renderer/Private/MobileShadingRenderer.cpp
+++ b/Engine/Source/Runtime/Renderer/Private/MobileShadingRenderer.cpp
@@ -928,7 +928,9 @@ FRHITexture* FMobileSceneRenderer::RenderForward(FRHICommandListImmediate& RHICm
 		// Here we use the base pass depth result to get z culling for opaque and masque.
 		// The color needs to be cleared at this point since shader complexity renders in additive.
 		DrawClearQuad(RHICmdList, FLinearColor::Black);
-		RenderMobileDebugView(RHICmdList, ViewList);
+		FRDGBuilder GraphBuilder(RHICmdList);
+		RenderMobileDebugView(GraphBuilder, ViewList);
+		GraphBuilder.Execute();
 		RHICmdList.ImmediateFlush(EImmediateFlushType::DispatchToRHIThread);
 	}
 #endif // !(UE_BUILD_SHIPPING || UE_BUILD_TEST)
@@ -1216,28 +1218,38 @@ FRHITexture* FMobileSceneRenderer::RenderDeferred(FRHICommandListImmediate& RHIC
 	return ColorTargets[0];
 }
 
-void FMobileSceneRenderer::RenderMobileDebugView(FRHICommandListImmediate& RHICmdList, const TArrayView<const FViewInfo*> PassViews)
+BEGIN_SHADER_PARAMETER_STRUCT(FMobileDebugViewPassParameters, )
+	SHADER_PARAMETER_RDG_UNIFORM_BUFFER(FDebugViewModePassUniformParameters, DebugViewMode)
+END_SHADER_PARAMETER_STRUCT()
+
+void FMobileSceneRenderer::RenderMobileDebugView(FRDGBuilder& GraphBuilder, const TArrayView<const FViewInfo*> PassViews)
 {
 #if !(UE_BUILD_SHIPPING || UE_BUILD_TEST)
 	CSV_SCOPED_TIMING_STAT_EXCLUSIVE(RenderDebugView);
-	SCOPED_DRAW_EVENT(RHICmdList, MobileDebugView);
+	RDG_EVENT_SCOPE(GraphBuilder, "MobileDebugView");
 	SCOPE_CYCLE_COUNTER(STAT_BasePassDrawTime);
 
 	for (int32 ViewIndex = 0; ViewIndex < PassViews.Num(); ViewIndex++)
 	{
-		SCOPED_CONDITIONAL_DRAW_EVENTF(RHICmdList, EventView, Views.Num() > 1, TEXT("View%d"), ViewIndex);
+		RDG_EVENT_SCOPE_CONDITIONAL(GraphBuilder, Views.Num() > 1, "View%d", ViewIndex);
 		const FViewInfo& View = *PassViews[ViewIndex];
 		if (!View.ShouldRenderView())
 		{
 			continue;
 		}
 
-		TUniformBufferRef<FDebugViewModePassUniformParameters> DebugViewModePassUniformBuffer = CreateDebugViewModePassUniformBuffer(RHICmdList, View);
-		FUniformBufferStaticBindings GlobalUniformBuffers(DebugViewModePassUniformBuffer);
-		SCOPED_UNIFORM_BUFFER_GLOBAL_BINDINGS(RHICmdList, GlobalUniformBuffers);
+		auto* PassParameters = GraphBuilder.AllocParameters<FMobileDebugViewPassParameters>();
+		PassParameters->DebugViewMode = CreateDebugViewModePassUniformBuffer(GraphBuilder, View);
 
-		RHICmdList.SetViewport(View.ViewRect.Min.X, View.ViewRect.Min.Y, 0, View.ViewRect.Max.X, View.ViewRect.Max.Y, 1);
-		View.ParallelMeshDrawCommandPasses[EMeshPass::DebugViewMode].DispatchDraw(nullptr, RHICmdList);
+		GraphBuilder.AddPass(
+			{},
+			PassParameters,
+			ERDGPassFlags::Raster | ERDGPassFlags::SkipRenderPass | ERDGPassFlags::NeverCull,
+			[&View](FRHICommandList& RHICmdList)
+		{
+			RHICmdList.SetViewport(View.ViewRect.Min.X, View.ViewRect.Min.Y, 0, View.ViewRect.Max.X, View.ViewRect.Max.Y, 1);
+			View.ParallelMeshDrawCommandPasses[EMeshPass::DebugViewMode].DispatchDraw(nullptr, RHICmdList);
+		});
 	}
 #endif // !(UE_BUILD_SHIPPING || UE_BUILD_TEST)
 }
diff --git a/Engine/Source/Runtime/Renderer/Private/Renderer.cpp b/Engine/Source/Runtime/Renderer/Private/Renderer.cpp
--- a/Engine/Source/Runtime/Renderer/Private/Renderer.cpp
+++ b/Engine/Source/Runtime/Renderer/Private/Renderer.cpp
@@ -101,6 +101,13 @@ void FRendererModule::InitializeSystemTextures(FRHICommandListImmediate& RHICmdL
 	GSystemTextures.InitializeTextures(RHICmdList, GMaxRHIFeatureLevel);
 }
 
+BEGIN_SHADER_PARAMETER_STRUCT(FDrawTileMeshPassParameters, )
+	SHADER_PARAMETER_RDG_UNIFORM_BUFFER(FDebugViewModePassUniformParameters, DebugViewMode)
+	SHADER_PARAMETER_RDG_UNIFORM_BUFFER(FTranslucentBasePassUniformParameters, TranslucentBasePass)
+	SHADER_PARAMETER_RDG_UNIFORM_BUFFER(FOpaqueBasePassUniformParameters, OpaqueBasePass)
+	SHADER_PARAMETER_RDG_UNIFORM_BUFFER(FMobileBasePassUniformParameters, MobileBasePass)
+END_SHADER_PARAMETER_STRUCT()
+
 void FRendererModule::DrawTileMesh(FRHICommandListImmediate& RHICmdList, FMeshPassProcessorRenderState& DrawRenderState, const FSceneView& SceneView, FMeshBatch& Mesh, bool bIsHitTesting, const FHitProxyId& HitProxyId, bool bUse128bitRT)
 {
 	if (!GUsingNullRHI)
@@ -192,6 +199,11 @@ void FRendererModule::DrawTileMesh(FRHICommandListImmediate& RHICmdList, FMeshPa
 
 		GSystemTextures.InitializeTextures(RHICmdList, FeatureLevel);
 		FMemMark Mark(FMemStack::Get());
+		FRDGBuilder GraphBuilder(RHICmdList);
+
+		const ERDGPassFlags PassFlags = ERDGPassFlags::Raster | ERDGPassFlags::SkipRenderPass | ERDGPassFlags::NeverCull;
+
+		auto* PassParameters = GraphBuilder.AllocParameters<FDrawTileMeshPassParameters>();
 
 		// handle translucent material blend modes, not relevant in MaterialTexCoordScalesAnalysis since it outputs the scales.
 		if (ViewFamily->GetDebugViewShaderMode() == DVSM_OutputMaterialTextureScales)
@@ -203,22 +215,22 @@ void FRendererModule::DrawTileMesh(FRHICommandListImmediate& RHICmdList, FMeshPa
 			// is this path used on mobile?
 			if (ShadingPath == EShadingPath::Deferred)
 			{
-				TUniformBufferRef<FDebugViewModePassUniformParameters> DebugViewModePassUniformBuffer = CreateDebugViewModePassUniformBuffer(RHICmdList, View);
-				FUniformBufferStaticBindings GlobalUniformBuffers(DebugViewModePassUniformBuffer);
-				SCOPED_UNIFORM_BUFFER_GLOBAL_BINDINGS(RHICmdList, GlobalUniformBuffers);
+				PassParameters->DebugViewMode = CreateDebugViewModePassUniformBuffer(GraphBuilder, View);
 
-				DrawDynamicMeshPass(View, RHICmdList,
-					[Scene, &View, &DrawRenderState, &DebugViewModePassUniformBuffer, &Mesh](FMeshPassDrawListContext* InDrawListContext)
+				GraphBuilder.AddPass({}, PassParameters, PassFlags,
+					[Scene, &View, &Mesh](FRHICommandListImmediate& InRHICmdList)
 				{
-					FDebugViewModeMeshProcessor PassMeshProcessor(
-						Scene,
-						View.GetFeatureLevel(),
-						&View,
-						DebugViewModePassUniformBuffer,
-						false,
-						InDrawListContext);
-					const uint64 DefaultBatchElementMask = ~0ull;
-					PassMeshProcessor.AddMeshBatch(Mesh, DefaultBatchElementMask, nullptr);
+					DrawDynamicMeshPass(View, InRHICmdList, [&](FMeshPassDrawListContext* InDrawListContext)
+					{
+						FDebugViewModeMeshProcessor PassMeshProcessor(
+							Scene,
+							View.GetFeatureLevel(),
+							&View,
+							false,
+							InDrawListContext);
+						const uint64 DefaultBatchElementMask = ~0ull;
+						PassMeshProcessor.AddMeshBatch(Mesh, DefaultBatchElementMask, nullptr);
+					});
 				});
 			}
 #endif // !(UE_BUILD_SHIPPING || UE_BUILD_TEST)
@@ -227,23 +239,25 @@ void FRendererModule::DrawTileMesh(FRHICommandListImmediate& RHICmdList, FMeshPa
 		{
 			if (ShadingPath == EShadingPath::Deferred)
 			{
-				TUniformBufferRef<FTranslucentBasePassUniformParameters> TranslucentBasePassUniformBuffer = CreateTranslucentBasePassUniformBuffer(RHICmdList, View, ESceneTextureSetupMode::None, 0);
-				FUniformBufferStaticBindings GlobalUniformBuffers(TranslucentBasePassUniformBuffer);
-				SCOPED_UNIFORM_BUFFER_GLOBAL_BINDINGS(RHICmdList, GlobalUniformBuffers);
+				PassParameters->TranslucentBasePass = CreateTranslucentBasePassUniformBuffer(GraphBuilder, View, nullptr, nullptr, ESceneTextureSetupMode::None, 0);
 
-				DrawDynamicMeshPass(View, RHICmdList, [&](FDynamicPassMeshDrawListContext* DynamicMeshPassContext)
+				GraphBuilder.AddPass({}, PassParameters, PassFlags,
+					[Scene, &View, &Mesh, &DrawRenderState, bUse128bitRT](FRHICommandListImmediate& InRHICmdList)
 				{
-					FBasePassMeshProcessor PassMeshProcessor(
-						Scene,
-						View.GetFeatureLevel(),
-						&View,
-						DrawRenderState,
-						DynamicMeshPassContext,
-						bUse128bitRT ? FBasePassMeshProcessor::EFlags::bRequires128bitRT : FBasePassMeshProcessor::EFlags::None,
-						ETranslucencyPass::TPT_AllTranslucency);
+					DrawDynamicMeshPass(View, InRHICmdList, [&](FDynamicPassMeshDrawListContext* DynamicMeshPassContext)
+					{
+						FBasePassMeshProcessor PassMeshProcessor(
+							Scene,
+							View.GetFeatureLevel(),
+							&View,
+							DrawRenderState,
+							DynamicMeshPassContext,
+							bUse128bitRT ? FBasePassMeshProcessor::EFlags::bRequires128bitRT : FBasePassMeshProcessor::EFlags::None,
+							ETranslucencyPass::TPT_AllTranslucency);
 
-					const uint64 DefaultBatchElementMask = ~0ull;
-					PassMeshProcessor.AddMeshBatch(Mesh, DefaultBatchElementMask, nullptr);
+						const uint64 DefaultBatchElementMask = ~0ull;
+						PassMeshProcessor.AddMeshBatch(Mesh, DefaultBatchElementMask, nullptr);
+					});
 				});
 			}
 			else // Mobile
@@ -300,23 +314,25 @@ void FRendererModule::DrawTileMesh(FRHICommandListImmediate& RHICmdList, FMeshPa
 			{
 				if (ShadingPath == EShadingPath::Deferred)
 				{
-					TUniformBufferRef<FOpaqueBasePassUniformParameters> OpaqueBasePassUniformBuffer = CreateOpaqueBasePassUniformBuffer(RHICmdList, View, nullptr);
-					FUniformBufferStaticBindings GlobalUniformBuffers(OpaqueBasePassUniformBuffer);
-					SCOPED_UNIFORM_BUFFER_GLOBAL_BINDINGS(RHICmdList, GlobalUniformBuffers);
-					
-					DrawDynamicMeshPass(View, RHICmdList,
-						[Scene, &View, &DrawRenderState, &Mesh, bUse128bitRT](FDynamicPassMeshDrawListContext* DynamicMeshPassContext)
+					PassParameters->OpaqueBasePass = CreateOpaqueBasePassUniformBuffer(GraphBuilder, View, nullptr, nullptr, 0);
+
+					GraphBuilder.AddPass({}, PassParameters, PassFlags,
+						[Scene, &View, &DrawRenderState, &Mesh, bUse128bitRT](FRHICommandListImmediate& InRHICmdList)
 					{
-						FBasePassMeshProcessor PassMeshProcessor(
-							Scene,
-							View.GetFeatureLevel(),
-							&View,
-							DrawRenderState,
-							DynamicMeshPassContext,
-							bUse128bitRT ? FBasePassMeshProcessor::EFlags::bRequires128bitRT : FBasePassMeshProcessor::EFlags::None);
-						
-						const uint64 DefaultBatchElementMask = ~0ull;
-						PassMeshProcessor.AddMeshBatch(Mesh, DefaultBatchElementMask, nullptr);
+						DrawDynamicMeshPass(View, InRHICmdList,
+							[&](FDynamicPassMeshDrawListContext* DynamicMeshPassContext)
+						{
+							FBasePassMeshProcessor PassMeshProcessor(
+								Scene,
+								View.GetFeatureLevel(),
+								&View,
+								DrawRenderState,
+								DynamicMeshPassContext,
+								bUse128bitRT ? FBasePassMeshProcessor::EFlags::bRequires128bitRT : FBasePassMeshProcessor::EFlags::None);
+
+							const uint64 DefaultBatchElementMask = ~0ull;
+							PassMeshProcessor.AddMeshBatch(Mesh, DefaultBatchElementMask, nullptr);
+						});
 					});
 				}
 				else // Mobile
@@ -342,6 +358,8 @@ void FRendererModule::DrawTileMesh(FRHICommandListImmediate& RHICmdList, FMeshPa
 				}
 			}
 		}
+
+		GraphBuilder.Execute();
 	}
 }
 
diff --git a/Engine/Source/Runtime/Renderer/Private/SceneRendering.h b/Engine/Source/Runtime/Renderer/Private/SceneRendering.h
--- a/Engine/Source/Runtime/Renderer/Private/SceneRendering.h
+++ b/Engine/Source/Runtime/Renderer/Private/SceneRendering.h
@@ -2095,7 +2095,7 @@ protected:
 	void RenderMobileEditorPrimitives(FRHICommandList& RHICmdList, const FViewInfo& View, const FMeshPassProcessorRenderState& DrawRenderState);
 
 	/** Renders the debug view pass for mobile. */
-	void RenderMobileDebugView(FRHICommandListImmediate& RHICmdList, const TArrayView<const FViewInfo*> PassViews);
+	void RenderMobileDebugView(FRDGBuilder& GraphBuilder, const TArrayView<const FViewInfo*> PassViews);
 
 	/** Render modulated shadow projections in to the scene, loops over any unrendered shadows until all are processed.*/
 	void RenderModulatedShadowProjections(FRHICommandListImmediate& RHICmdList);
diff --git a/Engine/Source/Runtime/Renderer/Private/TranslucentRendering.cpp b/Engine/Source/Runtime/Renderer/Private/TranslucentRendering.cpp
--- a/Engine/Source/Runtime/Renderer/Private/TranslucentRendering.cpp
+++ b/Engine/Source/Runtime/Renderer/Private/TranslucentRendering.cpp
@@ -607,8 +607,7 @@ void SetupDownsampledTranslucencyViewParameters(
 }
 
 void SetupTranslucentBasePassUniformParameters(
-	FRDGBuilder* GraphBuilder,
-	FRHICommandListImmediate& RHICmdList,
+	FRDGBuilder& GraphBuilder,
 	const FSceneRenderTargets& SceneRenderTargets,
 	const FViewInfo& View,
 	const FTranslucentVolumeLightingTextures* TranslucentVolumeLightingTextures,
@@ -619,11 +618,11 @@ void SetupTranslucentBasePassUniformParameters(
 {
 	const auto GetRDG = [&](const TRefCountPtr<IPooledRenderTarget>& PooledRenderTarget, ERDGTextureFlags Flags = ERDGTextureFlags::None)
 	{
-		return RegisterExternalOrPassthroughTexture(GraphBuilder, PooledRenderTarget, Flags);
+		return GraphBuilder.RegisterExternalTexture(PooledRenderTarget, ERenderTargetTexture::ShaderResource, Flags);
 	};
 
-	SetupSharedBasePassParameters(GraphBuilder, RHICmdList, View, BasePassParameters.Shared);
-	SetupSceneTextureUniformParameters(GraphBuilder, View.FeatureLevel, SceneRenderTargets, SceneTextureSetupMode, BasePassParameters.SceneTextures);
+	SetupSharedBasePassParameters(GraphBuilder, View, BasePassParameters.Shared);
+	SetupSceneTextureUniformParameters(&GraphBuilder, View.FeatureLevel, SceneRenderTargets, SceneTextureSetupMode, BasePassParameters.SceneTextures);
 
 	FRDGTextureRef BlackDummyTexture = GetRDG(GSystemTextures.BlackDummy);
 
@@ -771,29 +770,17 @@ void SetupTranslucentBasePassUniformParameters(
 TRDGUniformBufferRef<FTranslucentBasePassUniformParameters> CreateTranslucentBasePassUniformBuffer(
 	FRDGBuilder& GraphBuilder,
 	const FViewInfo& View,
-	const FTranslucentVolumeLightingTextures& TranslucentVolumeLightingTextures,
+	const FTranslucentVolumeLightingTextures* TranslucentVolumeLightingTextures,
 	FRDGTextureRef SceneColorCopyTexture,
 	ESceneTextureSetupMode SceneTextureSetupMode,
 	const int32 ViewIndex)
 {
 	FSceneRenderTargets& SceneRenderTargets = FSceneRenderTargets::Get();
 	FTranslucentBasePassUniformParameters* BasePassParameters = GraphBuilder.AllocParameters<FTranslucentBasePassUniformParameters>();
-	SetupTranslucentBasePassUniformParameters(&GraphBuilder, GraphBuilder.RHICmdList, SceneRenderTargets, View, &TranslucentVolumeLightingTextures, SceneColorCopyTexture, SceneTextureSetupMode, ViewIndex, *BasePassParameters);
+	SetupTranslucentBasePassUniformParameters(GraphBuilder, SceneRenderTargets, View, TranslucentVolumeLightingTextures, SceneColorCopyTexture, SceneTextureSetupMode, ViewIndex, *BasePassParameters);
 	return GraphBuilder.CreateUniformBuffer(BasePassParameters);
 }
 
-TUniformBufferRef<FTranslucentBasePassUniformParameters> CreateTranslucentBasePassUniformBuffer(
-	FRHICommandListImmediate& RHICmdList,
-	const FViewInfo& View,
-	ESceneTextureSetupMode SceneTextureSetupMode,
-	const int32 ViewIndex)
-{
-	FSceneRenderTargets& SceneRenderTargets = FSceneRenderTargets::Get();
-	FTranslucentBasePassUniformParameters BasePassParameters;
-	SetupTranslucentBasePassUniformParameters(nullptr, RHICmdList, SceneRenderTargets, View, nullptr, nullptr, SceneTextureSetupMode, ViewIndex, BasePassParameters);
-	return TUniformBufferRef<FTranslucentBasePassUniformParameters>::CreateUniformBufferImmediate(BasePassParameters, EUniformBufferUsage::UniformBuffer_SingleFrame);
-}
-
 static void UpdateSeparateTranslucencyViewState(FScene* Scene, const FViewInfo& View, const FScreenPassTextureViewport& SeparateTranslucencyViewport, FMeshPassProcessorRenderState& DrawRenderState)
 {
 	Scene->UniformBuffers.UpdateViewUniformBuffer(View);
@@ -1069,7 +1056,7 @@ void FDeferredShadingSceneRenderer::RenderTranslucencyInner(
 				SeparateTranslucencyColorTexture,
 				SeparateTranslucencyColorLoadAction,
 				SeparateTranslucencyDepthTexture.Target,
-				CreateTranslucentBasePassUniformBuffer(GraphBuilder, View, TranslucentVolumeLightingTextures, SceneColorCopyTexture, SceneTextureSetupMode, ViewIndex),
+				CreateTranslucentBasePassUniformBuffer(GraphBuilder, View, &TranslucentVolumeLightingTextures, SceneColorCopyTexture, SceneTextureSetupMode, ViewIndex),
 				TranslucencyPass,
 				!bCompositeBackToSceneColor,
 				bRenderInParallel);
@@ -1121,7 +1108,7 @@ void FDeferredShadingSceneRenderer::RenderTranslucencyInner(
 				SceneColorTexture,
 				SceneColorLoadAction,
 				SceneDepthTexture.Target,
-				CreateTranslucentBasePassUniformBuffer(GraphBuilder, View, TranslucentVolumeLightingTextures, SceneColorCopyTexture, SceneTextureSetupMode, ViewIndex),
+				CreateTranslucentBasePassUniformBuffer(GraphBuilder, View, &TranslucentVolumeLightingTextures, SceneColorCopyTexture, SceneTextureSetupMode, ViewIndex),
 				TranslucencyPass,
 				bResolveColorTexture,
 				bRenderInParallel);
