Fix for crash in translucent rendering when in wireframe.

#rnx

diff --git a/Engine/Source/Runtime/Renderer/Private/TranslucentRendering.cpp b/Engine/Source/Runtime/Renderer/Private/TranslucentRendering.cpp
--- a/Engine/Source/Runtime/Renderer/Private/TranslucentRendering.cpp
+++ b/Engine/Source/Runtime/Renderer/Private/TranslucentRendering.cpp
@@ -724,7 +724,7 @@ void SetupTranslucentBasePassUniformParameters(
 
 	// Translucency Lighting Volume
 	{
-		if (TranslucentVolumeLightingTextures)
+		if (TranslucentVolumeLightingTextures && TranslucentVolumeLightingTextures->IsValid())
 		{
 			BasePassParameters.TranslucencyLightingVolumeAmbientInner = TranslucentVolumeLightingTextures->GetAmbient(ViewIndex, TVC_Inner);
 			BasePassParameters.TranslucencyLightingVolumeAmbientOuter = TranslucentVolumeLightingTextures->GetAmbient(ViewIndex, TVC_Outer);
diff --git a/Engine/Source/Runtime/Renderer/Private/TranslucentRendering.h b/Engine/Source/Runtime/Renderer/Private/TranslucentRendering.h
--- a/Engine/Source/Runtime/Renderer/Private/TranslucentRendering.h
+++ b/Engine/Source/Runtime/Renderer/Private/TranslucentRendering.h
@@ -64,6 +64,11 @@ private:
 
 struct FTranslucentVolumeLightingTextures
 {
+	bool IsValid() const
+	{
+		return VolumeDim != 0;
+	}
+
 	FRDGTextureRef GetAmbient(int32 ViewIndex, int32 CascadeIndex) const
 	{
 		return Ambient[(ViewIndex * TVC_MAX) + CascadeIndex];
