Converted DistanceFieldShadows to RDG resources and pass parameters

diff --git a/Engine/Shaders/Private/DistanceFieldLightingShared.ush b/Engine/Shaders/Private/DistanceFieldLightingShared.ush
--- a/Engine/Shaders/Private/DistanceFieldLightingShared.ush
+++ b/Engine/Shaders/Private/DistanceFieldLightingShared.ush
@@ -62,7 +62,7 @@ float SampleHFVisibilityTexture(float2 UV)
 	return 0.0;
 }
 #else
-Texture2D DistanceFieldTexture;
+Texture2D HeightFieldTexture;
 
 float SampleMeshDistanceField(float3 VolumeUV, float2 ObjectDistanceFieldMAD)
 {
@@ -71,7 +71,7 @@ float SampleMeshDistanceField(float3 VolumeUV, float2 ObjectDistanceFieldMAD)
 
 float SampleHeightFieldAtlas(float2 UV)
 {
-	float4 SampleValue = Texture2DSampleLevel(DistanceFieldTexture, GlobalBilinearWrappedSampler, UV, 0);
+	float4 SampleValue = Texture2DSampleLevel(HeightFieldTexture, GlobalBilinearWrappedSampler, UV, 0);
 	return DecodePackedHeight(SampleValue.xy);
 }
 
@@ -82,6 +82,7 @@ float SampleHFVisibilityTexture(float2 UV)
 #endif
 
 float3 DistanceFieldAtlasTexelSize;
+float2 HeightFieldAtlasTexelSize;
 
 RWBuffer<uint> RWObjectIndirectArguments;
 Buffer<uint> ObjectIndirectArguments;
diff --git a/Engine/Shaders/Private/DistanceFieldShadowing.usf b/Engine/Shaders/Private/DistanceFieldShadowing.usf
--- a/Engine/Shaders/Private/DistanceFieldShadowing.usf
+++ b/Engine/Shaders/Private/DistanceFieldShadowing.usf
@@ -9,6 +9,9 @@
 #include "DistanceFieldLightingShared.ush"
 #include "DistanceFieldShadowingShared.ush"
 
+#define SCATTER_TILE_CULLING (CULLING_TYPE == 0)
+#define POINT_LIGHT (CULLING_TYPE == 2)
+
 #if (DISTANCEFIELD_PRIMITIVE_TYPE == DFPT_SignedDistanceField)
 #define MAX_TRACE_SPHERE_RADIUS 100
 #define SELF_SHADOW_VERTICAL_BIAS 0
@@ -304,7 +307,7 @@ void HandleShadowTileObjectIntersection(uint TileIndex, uint ObjectIndex)
 
 		uint StartOffset = ShadowTileStartOffsets[TileIndex];
 		uint DataIndex = (ArrayIndex + StartOffset);
-		RWShadowTileArrayData[DataIndex * SHADOW_TILE_ARRAY_DATA_STRIDE] = ObjectIndex;
+		RWShadowTileArrayData[DataIndex] = ObjectIndex;
 	#endif
 }
 
@@ -702,7 +705,7 @@ void DistanceFieldShadowingCS(
 			float SubsurfaceDensity = 0;
 			bool bUseSubsurfaceTransmission = false;
 
-#if !FORWARD_SHADING && DF_SHADOW_QUALITY > 2
+#if !FORWARD_SHADING && DF_SHADOW_QUALITY > 1
 			FGBufferData GBufferData = GetGBufferData(ScreenUV);
 
 			BRANCH
diff --git a/Engine/Shaders/Private/DistanceFieldShadowingShared.ush b/Engine/Shaders/Private/DistanceFieldShadowingShared.ush
--- a/Engine/Shaders/Private/DistanceFieldShadowingShared.ush
+++ b/Engine/Shaders/Private/DistanceFieldShadowingShared.ush
@@ -11,10 +11,6 @@ Buffer<uint> ShadowTileArrayData;
 uint2 ShadowTileListGroupSize;
 uint ShadowAverageObjectsPerTile;
 
-#ifndef SHADOW_TILE_ARRAY_DATA_STRIDE
-#define SHADOW_TILE_ARRAY_DATA_STRIDE 0
-#endif
-
 uint2 GetShadowTileHead(uint2 TileCoordinate)
 {
 	uint TileIndex = TileCoordinate.y * ShadowTileListGroupSize.x + TileCoordinate.x;
@@ -72,7 +68,7 @@ float ShadowRayTraceThroughCulledObjects(
 			if (bUseScatterTileCulling)
 			{
 				uint CulledDataStart = CulledDataParameter;
-				ObjectIndex = ShadowTileArrayData.Load((ListObjectIndex + CulledDataStart) * SHADOW_TILE_ARRAY_DATA_STRIDE);
+				ObjectIndex = ShadowTileArrayData.Load(ListObjectIndex + CulledDataStart);
 			}
 			else
 			{
@@ -110,9 +106,9 @@ float ShadowRayTraceThroughCulledObjects(
 			if (IntersectionTimes.x < IntersectionTimes.y)
 			{
 				float SampleRayTime = IntersectionTimes.x * VolumeRayLength;
-#if DF_SHADOW_QUALITY == 3
+#if DF_SHADOW_QUALITY == 2
 				uint MaxSteps = 64;
-#elif DF_SHADOW_QUALITY == 2
+#elif DF_SHADOW_QUALITY == 1
 				uint MaxSteps = 32;
 #else
 				uint MaxSteps = 20;
@@ -197,9 +193,9 @@ float RayTraceHeightfieldLocal(
 	float LocalSelfShadowScale,
 	float LocalMaxSphereRadius)
 {
-#if DF_SHADOW_QUALITY == 3
+#if DF_SHADOW_QUALITY == 2
 	float MaxSteps = 32;
-#elif DF_SHADOW_QUALITY == 2
+#elif DF_SHADOW_QUALITY == 1
 	float MaxSteps = 16;
 #else
 	float MaxSteps = 8;
@@ -226,7 +222,7 @@ float RayTraceHeightfieldLocal(
 
 		MinConeVisibility = min(MinConeVisibility, StepVisibility);
 
-#if DF_SHADOW_QUALITY == 3
+#if DF_SHADOW_QUALITY == 2
 		// Compute ray visibility assuming it starts below heightfield. Only needed when the heightfield has hole
 		if (VisUVScaleBias.x > 0.0)
 		{
@@ -288,7 +284,7 @@ float ShadowRayTraceThroughCulledHeightFieldObjects(
 			if (bUseScatterTileCulling)
 			{
 				uint CulledDataStart = CulledDataParameter;
-				ObjectIndex = ShadowTileArrayData.Load((ListObjectIndex + CulledDataStart) * SHADOW_TILE_ARRAY_DATA_STRIDE);
+				ObjectIndex = ShadowTileArrayData.Load(ListObjectIndex + CulledDataStart);
 			}
 			else
 			{
@@ -318,7 +314,7 @@ float ShadowRayTraceThroughCulledHeightFieldObjects(
 		if (bValidIntersection)
 		{
 			float4 AtlasUVScaleBias = LoadHeightFieldObjectAtlasUVScaleBias(ObjectIndex);
-#if DF_SHADOW_QUALITY == 3
+#if DF_SHADOW_QUALITY == 2
 			float4 VisUVScaleBias = LoadHeightFieldObjectVisibilityAtlasUVScaleBias(ObjectIndex);
 #else
 			float4 VisUVScaleBias = float4(0, 0, 0, 0);
diff --git a/Engine/Shaders/Private/HeightfieldLighting.usf b/Engine/Shaders/Private/HeightfieldLighting.usf
--- a/Engine/Shaders/Private/HeightfieldLighting.usf
+++ b/Engine/Shaders/Private/HeightfieldLighting.usf
@@ -236,65 +236,6 @@ void ShadowHeightfieldsPS(
 	OutColor = float4(Visibility, 0, 0, 0);
 }
 
-float3 LightDirection;
-float TanLightAngle;
-
-void RayTracedShadowHeightfieldsPS(
-	in FHeightfieldComponentQuadVertexOutput Input, 
-	out float4 OutColor : SV_Target0)
-{
-	float3 OpaqueWorldPosition = GetGlobalHeightfieldWorldPosition(Input.LocalPosition, Input.GlobalUV, GetLocalToWorld(Input.HeightfieldIndex));
-	float ViewDistance = mul(float4(OpaqueWorldPosition + View.PreViewTranslation.xyz, 1), View.TranslatedWorldToView).z;
-	float Shadow = 1;
-
-	if (ViewDistance > CascadeDepthMinMax.x && ViewDistance < CascadeDepthMinMax.y)
-	{
-		float4x4 WorldToLocal = GetWorldToLocal(Input.HeightfieldIndex);
-		float WorldToLocalScale = length(WorldToLocal[0].xyz);
-		float HeightfieldTexelSize = 1.0f / WorldToLocalScale;
-
-		float3 WorldHeightfieldNormal = GetGlobalHeightfieldNormal(Input.GlobalUV);
-		// Bias away from the surface, mostly to avoid grass shadowing
-		OpaqueWorldPosition += .3f * HeightfieldTexelSize * WorldHeightfieldNormal;
-
-		uint NumIntersectingObjects = GetCulledNumObjects();
-		uint CulledDataStart = 0;
-
-		GetShadowTileCulledData(OpaqueWorldPosition, CulledDataStart, NumIntersectingObjects);
-
-		// Keeps result from going all the way sharp
-		float MinSphereRadius = .4f;
-		// Maintain reasonable culling bounds
-		float MaxSphereRadius = 100;
-		// Distance for directional lights to trace
-		float TraceDistance = 10000;
-
-		float3 WorldRayStart = OpaqueWorldPosition;
-		float3 WorldRayEnd = OpaqueWorldPosition + -LightDirection * TraceDistance;
-		float MaxRayTime = TraceDistance;
-
-		Shadow = ShadowRayTraceThroughCulledObjects(
-			WorldRayStart, 
-			WorldRayEnd, 
-			MaxRayTime, 
-			TanLightAngle, 
-			MinSphereRadius, 
-			MaxSphereRadius, 
-			0,
-			CulledDataStart, 
-			NumIntersectingObjects,
-			true,
-			true,
-			false);
-	}
-	else
-	{
-		clip(-1);
-	}
-
-	OutColor = float4(Shadow, 0, 0, 0);
-}
-
 float3 LightColor;
 float SkyLightIndirectScale;
 
diff --git a/Engine/Source/Runtime/Renderer/Private/DistanceFieldAmbientOcclusion.cpp b/Engine/Source/Runtime/Renderer/Private/DistanceFieldAmbientOcclusion.cpp
--- a/Engine/Source/Runtime/Renderer/Private/DistanceFieldAmbientOcclusion.cpp
+++ b/Engine/Source/Runtime/Renderer/Private/DistanceFieldAmbientOcclusion.cpp
@@ -634,27 +634,6 @@ void ListDistanceFieldLightingMemory(const FViewInfo& View, FSceneRenderer& Scen
 	{
 		UE_LOG(LogRenderer, Log, TEXT("   Screen grid temporaries %.3fMb"), ScreenGridResources->GetSizeBytesForAO() / 1024.0f / 1024.0f);
 	}
-	
-	UE_LOG(LogRenderer, Log, TEXT(""));
-	UE_LOG(LogRenderer, Log, TEXT("Ray Traced Distance Field Shadows"));
-
-	for (TSparseArray<FLightSceneInfoCompact>::TConstIterator LightIt(Scene->Lights); LightIt; ++LightIt)
-	{
-		const FLightSceneInfoCompact& LightSceneInfoCompact = *LightIt;
-		FLightSceneInfo* LightSceneInfo = LightSceneInfoCompact.LightSceneInfo;
-
-		FVisibleLightInfo& VisibleLightInfo = SceneRenderer.VisibleLightInfos[LightSceneInfo->Id];
-
-		for (int32 ShadowIndex = 0; ShadowIndex < VisibleLightInfo.ShadowsToProject.Num(); ShadowIndex++)
-		{
-			FProjectedShadowInfo* ProjectedShadowInfo = VisibleLightInfo.ShadowsToProject[ShadowIndex];
-
-			if (ProjectedShadowInfo->bRayTracedDistanceField && LightSceneInfo->TileIntersectionResources)
-			{
-				UE_LOG(LogRenderer, Log, TEXT("   Light Tile Culled objects %.3fMb"), LightSceneInfo->TileIntersectionResources->GetSizeBytes() / 1024.0f / 1024.0f);
-			}
-		}
-	}
 
 #endif // !NO_LOGGING
 }
diff --git a/Engine/Source/Runtime/Renderer/Private/DistanceFieldAmbientOcclusion.h b/Engine/Source/Runtime/Renderer/Private/DistanceFieldAmbientOcclusion.h
--- a/Engine/Source/Runtime/Renderer/Private/DistanceFieldAmbientOcclusion.h
+++ b/Engine/Source/Runtime/Renderer/Private/DistanceFieldAmbientOcclusion.h
@@ -494,6 +494,6 @@ extern void TrackGPUProgress(FRHICommandListImmediate& RHICmdList, uint32 DebugI
 extern bool ShouldRenderDynamicSkyLight(const FScene* Scene, const FSceneViewFamily& ViewFamily);
 extern bool ShouldRenderDeferredDynamicSkyLight(const FScene* Scene, const FSceneViewFamily& ViewFamily);
 
-extern void CullObjectsToView(FRDGBuilder& GraphBuilder, FScene* Scene, const FViewInfo& View, const FDistanceFieldAOParameters& Parameters, FDistanceFieldObjectBufferResource& CulledObjectBuffers);
+extern void CullObjectsToView(FRDGBuilder& GraphBuilder, FScene* Scene, const FViewInfo& View, const FDistanceFieldAOParameters& Parameters, class FDistanceFieldObjectBufferResource& CulledObjectBuffers);
 extern void BuildTileObjectLists(FRDGBuilder& GraphBuilder, FScene* Scene, TArray<FViewInfo>& Views, FRDGTextureRef DistanceFieldNormal, const FDistanceFieldAOParameters& Parameters);
 extern FIntPoint GetTileListGroupSizeForView(const FViewInfo& View);
diff --git a/Engine/Source/Runtime/Renderer/Private/DistanceFieldLightingPost.cpp b/Engine/Source/Runtime/Renderer/Private/DistanceFieldLightingPost.cpp
--- a/Engine/Source/Runtime/Renderer/Private/DistanceFieldLightingPost.cpp
+++ b/Engine/Source/Runtime/Renderer/Private/DistanceFieldLightingPost.cpp
@@ -12,8 +12,6 @@
 #include "CompositionLighting/PostProcessAmbientOcclusion.h"
 #include "PipelineStateCache.h"
 
-IMPLEMENT_TYPE_LAYOUT(FLightTileIntersectionParameters);
-
 #if WITH_MGPU
 DECLARE_GPU_STAT(AFRWaitForDistanceFieldAOHistory);
 #endif
diff --git a/Engine/Source/Runtime/Renderer/Private/DistanceFieldLightingShared.h b/Engine/Source/Runtime/Renderer/Private/DistanceFieldLightingShared.h
--- a/Engine/Source/Runtime/Renderer/Private/DistanceFieldLightingShared.h
+++ b/Engine/Source/Runtime/Renderer/Private/DistanceFieldLightingShared.h
@@ -73,6 +73,24 @@ public:
 class FDistanceFieldObjectBuffers : public TDistanceFieldObjectBuffers<DFPT_SignedDistanceField> {};
 class FHeightFieldObjectBuffers : public TDistanceFieldObjectBuffers<DFPT_HeightField> {};
 
+BEGIN_SHADER_PARAMETER_STRUCT(FDistanceFieldAtlasParameters, )
+	SHADER_PARAMETER_TEXTURE(Texture3D, DistanceFieldTexture)
+	SHADER_PARAMETER_SAMPLER(SamplerState, DistanceFieldSampler)
+	SHADER_PARAMETER(FVector, DistanceFieldAtlasTexelSize)
+END_SHADER_PARAMETER_STRUCT()
+
+BEGIN_SHADER_PARAMETER_STRUCT(FHeightFieldAtlasParameters, )
+	SHADER_PARAMETER_TEXTURE(Texture2D, HeightFieldTexture)
+	SHADER_PARAMETER_TEXTURE(Texture2D, HFVisibilityTexture)
+	SHADER_PARAMETER(FVector2D, HeightFieldAtlasTexelSize)
+END_SHADER_PARAMETER_STRUCT()
+
+BEGIN_SHADER_PARAMETER_STRUCT(FDistanceFieldObjectBufferParameters, )
+	SHADER_PARAMETER_SRV(StructuredBuffer<float4>, SceneObjectBounds)
+	SHADER_PARAMETER_SRV(StructuredBuffer<float4>, SceneObjectData)
+	SHADER_PARAMETER(uint32, NumSceneObjects)
+END_SHADER_PARAMETER_STRUCT()
+
 template <EDistanceFieldPrimitiveType PrimitiveType>
 class TDistanceFieldObjectBufferParameters
 {
@@ -289,6 +307,25 @@ public:
 class FDistanceFieldObjectBufferResource : public TDistanceFieldObjectBufferResource<DFPT_SignedDistanceField> {};
 class FHeightFieldObjectBufferResource : public TDistanceFieldObjectBufferResource<DFPT_HeightField> {};
 
+BEGIN_SHADER_PARAMETER_STRUCT(FDistanceFieldCulledObjectBufferParameters, )
+	SHADER_PARAMETER_RDG_BUFFER_UAV(RWBuffer<uint>, RWObjectIndirectArguments)
+	SHADER_PARAMETER_RDG_BUFFER_UAV(RWBuffer<uint>, RWCulledObjectBounds)
+	SHADER_PARAMETER_RDG_BUFFER_UAV(RWBuffer<uint>, RWCulledObjectData)
+	SHADER_PARAMETER_RDG_BUFFER_UAV(RWBuffer<uint>, RWCulledObjectBoxBounds)
+	SHADER_PARAMETER_RDG_BUFFER_SRV(Buffer<uint>, ObjectIndirectArguments)
+	SHADER_PARAMETER_RDG_BUFFER_SRV(Buffer<uint>, CulledObjectBounds)
+	SHADER_PARAMETER_RDG_BUFFER_SRV(Buffer<uint>, CulledObjectData)
+	SHADER_PARAMETER_RDG_BUFFER_SRV(Buffer<uint>, CulledObjectBoxBounds)
+END_SHADER_PARAMETER_STRUCT()
+
+extern void AllocateDistanceFieldCulledObjectBuffers(
+	FRDGBuilder& GraphBuilder, 
+	bool bWantBoxBounds, 
+	uint32 MaxObjects, 
+	uint32 NumBoundsElementsScale,
+	FRDGBufferRef& OutObjectIndirectArguments,
+	FDistanceFieldCulledObjectBufferParameters& OutParameters);
+
 template <EDistanceFieldPrimitiveType PrimitiveType>
 class TDistanceFieldCulledObjectBufferParameters
 {
@@ -458,214 +495,33 @@ public:
 	}
 };
 
-static int32 LightTileDataStride = 1;
-
-/**  */
-class FLightTileIntersectionResources
-{
-public:
-
-	FLightTileIntersectionResources() 
-		: TileDimensions(FIntPoint::ZeroValue)
-		, TileAlignedDimensions(FIntPoint::ZeroValue)
-		, b16BitIndices(false)
-	{}
-
-	void Initialize(int32 MaxNumObjectsPerTile)
-	{
-		TileAlignedDimensions = FIntPoint(Align(TileDimensions.X, 64), Align(TileDimensions.Y, 64));
-			
-		TileNumCulledObjects.Initialize(sizeof(uint32), TileAlignedDimensions.X * TileAlignedDimensions.Y, PF_R32_UINT, BUF_Static, TEXT("FLightTileIntersectionResources::TileNumCulledObjects"));
-		TileStartOffsets.Initialize(sizeof(uint32), TileAlignedDimensions.X * TileAlignedDimensions.Y, PF_R32_UINT, BUF_Static, TEXT("FLightTileIntersectionResources::TileStartOffsets"));
-		NextStartOffset.Initialize(sizeof(uint32), 1, PF_R32_UINT, BUF_Static, TEXT("FLightTileIntersectionResources::NextStartOffset"));
-
-		//@todo - handle max exceeded
-		TileArrayData.Initialize(b16BitIndices ? sizeof(uint16) : sizeof(uint32), MaxNumObjectsPerTile * TileAlignedDimensions.X * TileAlignedDimensions.Y * LightTileDataStride, b16BitIndices ? PF_R16_UINT : PF_R32_UINT, BUF_Static, TEXT("FLightTileIntersectionResources::TileArrayData"));
-	}
-
-	void Release()
-	{
-		TileNumCulledObjects.Release();
-		NextStartOffset.Release();
-		TileStartOffsets.Release();
-		TileArrayData.Release();
-	}
+BEGIN_SHADER_PARAMETER_STRUCT(FLightTileIntersectionParameters, )
+	SHADER_PARAMETER_RDG_BUFFER_UAV(RWBuffer<uint>, RWShadowTileNumCulledObjects)
+	SHADER_PARAMETER_RDG_BUFFER_UAV(RWBuffer<uint>, RWShadowTileStartOffsets)
+	SHADER_PARAMETER_RDG_BUFFER_UAV(RWBuffer<uint>, RWNextStartOffset)
+	SHADER_PARAMETER_RDG_BUFFER_UAV(RWBuffer<uint>, RWShadowTileArrayData)
 
-	size_t GetSizeBytes() const
-	{
-		return TileNumCulledObjects.NumBytes + NextStartOffset.NumBytes + TileStartOffsets.NumBytes + TileArrayData.NumBytes;
-	}
-
-	FIntPoint GetTileAlignedDimensions() const 
-	{
-		return TileAlignedDimensions;
-	}
-	
-	static FIntPoint GetAlignedDimensions(FIntPoint InTileDimensions)
-	{
-		return FIntPoint(Align(InTileDimensions.X, 64), Align(InTileDimensions.Y, 64));
-	}
-
-	FIntPoint TileDimensions;
-	FIntPoint TileAlignedDimensions;
-
-	FRWBuffer TileNumCulledObjects;
-	FRWBuffer NextStartOffset;
-	FRWBuffer TileStartOffsets;
-	FRWBuffer TileArrayData;
-	bool b16BitIndices;
-};
+	SHADER_PARAMETER_RDG_BUFFER_SRV(Buffer<uint>, ShadowTileNumCulledObjects)
+	SHADER_PARAMETER_RDG_BUFFER_SRV(Buffer<uint>, ShadowTileStartOffsets)
+	SHADER_PARAMETER_RDG_BUFFER_SRV(Buffer<uint>, NextStartOffset)
+	SHADER_PARAMETER_RDG_BUFFER_SRV(Buffer<uint>, ShadowTileArrayData)
 
-class FLightTileIntersectionParameters
-{
-	DECLARE_TYPE_LAYOUT(FLightTileIntersectionParameters, NonVirtual);
-public:
-
-	static void ModifyCompilationEnvironment(EShaderPlatform Platform, FShaderCompilerEnvironment& OutEnvironment)
-	{
-		OutEnvironment.SetDefine(TEXT("SHADOW_TILE_ARRAY_DATA_STRIDE"), LightTileDataStride);
-	}
-
-	void Bind(const FShaderParameterMap& ParameterMap)
-	{
-		ShadowTileNumCulledObjects.Bind(ParameterMap, TEXT("ShadowTileNumCulledObjects"));
-		ShadowTileStartOffsets.Bind(ParameterMap, TEXT("ShadowTileStartOffsets"));
-		NextStartOffset.Bind(ParameterMap, TEXT("NextStartOffset"));
-		ShadowTileArrayData.Bind(ParameterMap, TEXT("ShadowTileArrayData"));
-		ShadowTileListGroupSize.Bind(ParameterMap, TEXT("ShadowTileListGroupSize"));
-		ShadowAverageObjectsPerTile.Bind(ParameterMap, TEXT("ShadowAverageObjectsPerTile"));
-	}
-
-	bool IsBound() const
-	{
-		return ShadowTileNumCulledObjects.IsBound() || ShadowTileStartOffsets.IsBound() || NextStartOffset.IsBound() || ShadowTileArrayData.IsBound() || ShadowTileListGroupSize.IsBound() || ShadowAverageObjectsPerTile.IsBound();
-	}
-
-	template<typename TParamRef, typename TRHICommandList>
-	void Set(TRHICommandList& RHICmdList, const TParamRef& ShaderRHI, const FLightTileIntersectionResources& LightTileIntersectionResources)
-	{
-		ShadowTileNumCulledObjects.SetBuffer(RHICmdList, ShaderRHI, LightTileIntersectionResources.TileNumCulledObjects);
-		ShadowTileStartOffsets.SetBuffer(RHICmdList, ShaderRHI, LightTileIntersectionResources.TileStartOffsets);
-
-		NextStartOffset.SetBuffer(RHICmdList, ShaderRHI, LightTileIntersectionResources.NextStartOffset);
-
-		// Bind sorted array data if we are after the sort pass
-		ShadowTileArrayData.SetBuffer(RHICmdList, ShaderRHI, LightTileIntersectionResources.TileArrayData);
-
-		SetShaderValue(RHICmdList, ShaderRHI, ShadowTileListGroupSize, LightTileIntersectionResources.TileDimensions);
-		SetShaderValue(RHICmdList, ShaderRHI, ShadowAverageObjectsPerTile, GAverageObjectsPerShadowCullTile);
-	}
-
-	void GetUAVs(FLightTileIntersectionResources& TileIntersectionResources, TArray<FRHIUnorderedAccessView*>& UAVs)
-	{
-		int32 MaxIndex = FMath::Max(
-			FMath::Max(ShadowTileNumCulledObjects.GetUAVIndex(), ShadowTileStartOffsets.GetUAVIndex()), 
-			FMath::Max(NextStartOffset.GetUAVIndex(), ShadowTileArrayData.GetUAVIndex()));
-		UAVs.AddZeroed(MaxIndex + 1);
-
-		if (ShadowTileNumCulledObjects.IsUAVBound())
-		{
-			UAVs[ShadowTileNumCulledObjects.GetUAVIndex()] = TileIntersectionResources.TileNumCulledObjects.UAV;
-		}
-
-		if (ShadowTileStartOffsets.IsUAVBound())
-		{
-			UAVs[ShadowTileStartOffsets.GetUAVIndex()] = TileIntersectionResources.TileStartOffsets.UAV;
-		}
-
-		if (NextStartOffset.IsUAVBound())
-		{
-			UAVs[NextStartOffset.GetUAVIndex()] = TileIntersectionResources.NextStartOffset.UAV;
-		}
-
-		if (ShadowTileArrayData.IsUAVBound())
-		{
-			UAVs[ShadowTileArrayData.GetUAVIndex()] = TileIntersectionResources.TileArrayData.UAV;
-		}
-
-		check(UAVs.Num() > 0);
-	}
-
-	template<typename TParamRef>
-	void UnsetParameters(FRHIComputeCommandList& RHICmdList, const TParamRef& ShaderRHI)
-	{
-		ShadowTileNumCulledObjects.UnsetUAV(RHICmdList, ShaderRHI);
-		ShadowTileStartOffsets.UnsetUAV(RHICmdList, ShaderRHI);
-		NextStartOffset.UnsetUAV(RHICmdList, ShaderRHI);
-		ShadowTileArrayData.UnsetUAV(RHICmdList, ShaderRHI);
-	}
-
-	friend FArchive& operator<<(FArchive& Ar, FLightTileIntersectionParameters& P)
-	{
-		Ar << P.ShadowTileNumCulledObjects;
-		Ar << P.ShadowTileStartOffsets;
-		Ar << P.NextStartOffset;
-		Ar << P.ShadowTileArrayData;
-		Ar << P.ShadowTileListGroupSize;
-		Ar << P.ShadowAverageObjectsPerTile;
-		return Ar;
-	}
-
-private:
-	
-		LAYOUT_FIELD(FRWShaderParameter, ShadowTileNumCulledObjects)
-		LAYOUT_FIELD(FRWShaderParameter, ShadowTileStartOffsets)
-		LAYOUT_FIELD(FRWShaderParameter, NextStartOffset)
-		LAYOUT_FIELD(FRWShaderParameter, ShadowTileArrayData)
-		LAYOUT_FIELD(FShaderParameter, ShadowTileListGroupSize)
-		LAYOUT_FIELD(FShaderParameter, ShadowAverageObjectsPerTile)
-	
-};
+	SHADER_PARAMETER(FIntPoint, ShadowTileListGroupSize)
+END_SHADER_PARAMETER_STRUCT()
 
 extern void CullDistanceFieldObjectsForLight(
 	FRDGBuilder& GraphBuilder,
 	const FViewInfo& View,
 	const FLightSceneProxy* LightSceneProxy, 
+	EDistanceFieldPrimitiveType PrimitiveType,
 	const FMatrix& WorldToShadowValue, 
 	int32 NumPlanes, 
 	const FPlane* PlaneData, 
 	const FVector4& ShadowBoundingSphereValue,
 	float ShadowBoundingRadius,
-	TUniquePtr<class FLightTileIntersectionResources>& TileIntersectionResources);
-
-extern void CullHeightFieldObjectsForLight(
-	FRDGBuilder& GraphBuilder,
-	const FViewInfo& View,
-	const FLightSceneProxy* LightSceneProxy,
-	const FMatrix& WorldToShadowValue,
-	int32 NumPlanes,
-	const FPlane* PlaneData,
-	const FVector4& ShadowBoundingSphereValue,
-	float ShadowBoundingRadius,
-	TUniquePtr<class FLightTileIntersectionResources>& TileIntersectionResources);
-
-class FUniformMeshBuffers
-{
-public:
-
-	int32 MaxElements;
-
-	FVertexBufferRHIRef TriangleData;
-	FShaderResourceViewRHIRef TriangleDataSRV;
-
-	FRWBuffer TriangleAreas;
-	FRWBuffer TriangleCDFs;
-
-	FUniformMeshBuffers()
-	{
-		MaxElements = 0;
-	}
-
-	void Initialize();
-
-	void Release()
-	{
-		TriangleData.SafeRelease();
-		TriangleDataSRV.SafeRelease(); 
-		TriangleAreas.Release();
-		TriangleCDFs.Release();
-	}
-};
+	const FDistanceFieldObjectBufferParameters& ObjectBufferParameters,
+	FDistanceFieldCulledObjectBufferParameters& CulledObjectBufferParameters,
+	FLightTileIntersectionParameters& LightTileIntersectionParameters);
 
 extern TGlobalResource<FDistanceFieldObjectBufferResource> GAOCulledObjectBuffers;
 
diff --git a/Engine/Source/Runtime/Renderer/Private/DistanceFieldShadowing.cpp b/Engine/Source/Runtime/Renderer/Private/DistanceFieldShadowing.cpp
--- a/Engine/Source/Runtime/Renderer/Private/DistanceFieldShadowing.cpp
+++ b/Engine/Source/Runtime/Renderer/Private/DistanceFieldShadowing.cpp
@@ -130,11 +130,24 @@ FIntPoint GetBufferSizeForDFShadows()
 TGlobalResource<FDistanceFieldObjectBufferResource> GShadowCulledObjectBuffers;
 TGlobalResource<FHeightFieldObjectBufferResource> GShadowCulledHeightFieldObjectBuffers;
 
-template <EDistanceFieldPrimitiveType PrimitiveType>
-class TCullObjectsForShadowCS : public FGlobalShader
+class FCullObjectsForShadowCS : public FGlobalShader
 {
-	DECLARE_SHADER_TYPE(TCullObjectsForShadowCS, Global);
-public:
+	DECLARE_GLOBAL_SHADER(FCullObjectsForShadowCS);
+	SHADER_USE_PARAMETER_STRUCT(FCullObjectsForShadowCS, FGlobalShader);
+
+	BEGIN_SHADER_PARAMETER_STRUCT(FParameters, )
+		SHADER_PARAMETER_STRUCT_REF(FViewUniformShaderParameters, View)
+		SHADER_PARAMETER_STRUCT_INCLUDE(FDistanceFieldObjectBufferParameters, ObjectBufferParameters)
+		SHADER_PARAMETER_STRUCT_INCLUDE(FDistanceFieldCulledObjectBufferParameters, CulledObjectBufferParameters)
+		SHADER_PARAMETER(uint32, ObjectBoundingGeometryIndexCount)
+		SHADER_PARAMETER(FMatrix, WorldToShadow)
+		SHADER_PARAMETER(uint32, NumShadowHullPlanes)
+		SHADER_PARAMETER(FVector4, ShadowBoundingSphere)
+		SHADER_PARAMETER_ARRAY(FVector4,ShadowConvexHull,[12])
+	END_SHADER_PARAMETER_STRUCT()
+
+	class FPrimitiveType : SHADER_PERMUTATION_INT("DISTANCEFIELD_PRIMITIVE_TYPE", 2);
+	using FPermutationDomain = TShaderPermutationDomain<FPrimitiveType>;
 
 	static bool ShouldCompilePermutation(const FGlobalShaderPermutationParameters& Parameters)
 	{
@@ -145,209 +158,49 @@ public:
 	{
 		FGlobalShader::ModifyCompilationEnvironment(Parameters,OutEnvironment);
 		OutEnvironment.SetDefine(TEXT("UPDATEOBJECTS_THREADGROUP_SIZE"), UpdateObjectsGroupSize);
-		OutEnvironment.SetDefine(TEXT("DISTANCEFIELD_PRIMITIVE_TYPE"), (int32)PrimitiveType);
 	}
-
-	TCullObjectsForShadowCS(const ShaderMetaType::CompiledShaderInitializerType& Initializer)
-		: FGlobalShader(Initializer)
-	{
-		ObjectBufferParameters.Bind(Initializer.ParameterMap);
-		CulledObjectParameters.Bind(Initializer.ParameterMap);
-		ObjectBoundingGeometryIndexCount.Bind(Initializer.ParameterMap, TEXT("ObjectBoundingGeometryIndexCount"));
-		WorldToShadow.Bind(Initializer.ParameterMap, TEXT("WorldToShadow"));
-		NumShadowHullPlanes.Bind(Initializer.ParameterMap, TEXT("NumShadowHullPlanes"));
-		ShadowBoundingSphere.Bind(Initializer.ParameterMap, TEXT("ShadowBoundingSphere"));
-		ShadowConvexHull.Bind(Initializer.ParameterMap, TEXT("ShadowConvexHull"));
-	}
-
-	TCullObjectsForShadowCS()
-	{
-	}
-
-	void SetParameters(
-		FRHICommandList& RHICmdList,
-		const FScene* Scene,
-		const FSceneView& View,
-		const FMatrix& WorldToShadowValue,
-		int32 NumPlanes,
-		const FPlane* PlaneData,
-		const FVector4& ShadowBoundingSphereValue)
-	{
-		constexpr bool bIsHeightField = PrimitiveType == DFPT_HeightField;
-		const FDistanceFieldSceneData& SceneData = Scene->DistanceFieldSceneData;
-		const auto& SceneObjectBuffers = TSelector<bIsHeightField>()(*SceneData.GetHeightFieldObjectBuffers(), *SceneData.GetCurrentObjectBuffers());
-		auto& CulledObjectBuffers = TSelector<bIsHeightField>()(GShadowCulledHeightFieldObjectBuffers, GShadowCulledObjectBuffers);
-		
-		int32 NumObjectsInBuffer;
-		FRHITexture* TextureAtlas;
-		FRHITexture* VisibilityAtlas;
-		int32 AtlasSizeX;
-		int32 AtlasSizeY;
-		int32 AtlasSizeZ;
-
-		if (bIsHeightField)
-		{
-			NumObjectsInBuffer = SceneData.NumHeightFieldObjectsInBuffer;
-			TextureAtlas = GHeightFieldTextureAtlas.GetAtlasTexture();
-			VisibilityAtlas = GHFVisibilityTextureAtlas.GetAtlasTexture();
-			AtlasSizeX = GHeightFieldTextureAtlas.GetSizeX();
-			AtlasSizeY = GHeightFieldTextureAtlas.GetSizeY();
-			AtlasSizeZ = 1;
-		}
-		else
-		{
-			NumObjectsInBuffer = SceneData.NumObjectsInBuffer;
-			TextureAtlas = GDistanceFieldVolumeTextureAtlas.VolumeTextureRHI;
-			VisibilityAtlas = nullptr;
-			AtlasSizeX = GDistanceFieldVolumeTextureAtlas.GetSizeX();
-			AtlasSizeY = GDistanceFieldVolumeTextureAtlas.GetSizeY();
-			AtlasSizeZ = GDistanceFieldVolumeTextureAtlas.GetSizeZ();
-		}
-
-		FRHIComputeShader* ShaderRHI = RHICmdList.GetBoundComputeShader();
-		FGlobalShader::SetParameters<FViewUniformShaderParameters>(RHICmdList, ShaderRHI, View.ViewUniformBuffer);
-		ObjectBufferParameters.Set(RHICmdList, ShaderRHI, SceneObjectBuffers, NumObjectsInBuffer, TextureAtlas, AtlasSizeX, AtlasSizeY, AtlasSizeZ);
-
-		FRHITransitionInfo TransitionInfos[4];
-		TransitionInfos[0] = FRHITransitionInfo(CulledObjectBuffers.Buffers.ObjectIndirectArguments.UAV, ERHIAccess::Unknown, ERHIAccess::UAVCompute);
-		TransitionInfos[1] = FRHITransitionInfo(CulledObjectBuffers.Buffers.Bounds.UAV, ERHIAccess::Unknown, ERHIAccess::UAVCompute);
-		TransitionInfos[2] = FRHITransitionInfo(CulledObjectBuffers.Buffers.Data.UAV, ERHIAccess::Unknown, ERHIAccess::UAVCompute);
-		TransitionInfos[3] = FRHITransitionInfo(CulledObjectBuffers.Buffers.BoxBounds.UAV, ERHIAccess::Unknown, ERHIAccess::UAVCompute);
-		RHICmdList.Transition(MakeArrayView(TransitionInfos, UE_ARRAY_COUNT(TransitionInfos)));
-
-		CulledObjectParameters.Set(RHICmdList, ShaderRHI, CulledObjectBuffers.Buffers, TextureAtlas, FIntVector(AtlasSizeX, AtlasSizeY, AtlasSizeZ), VisibilityAtlas);
-
-		SetShaderValue(RHICmdList, ShaderRHI, ObjectBoundingGeometryIndexCount, UE_ARRAY_COUNT(GCubeIndices));
-		SetShaderValue(RHICmdList, ShaderRHI, WorldToShadow, WorldToShadowValue);
-		SetShaderValue(RHICmdList, ShaderRHI, ShadowBoundingSphere, ShadowBoundingSphereValue);
-
-		if (NumPlanes <= 12)
-		{
-			SetShaderValue(RHICmdList, ShaderRHI, NumShadowHullPlanes, NumPlanes);
-			SetShaderValueArray(RHICmdList, ShaderRHI, ShadowConvexHull, PlaneData, NumPlanes);
-		}
-		else
-		{
-			SetShaderValue(RHICmdList, ShaderRHI, NumShadowHullPlanes, 0);
-		}
-	}
-
-	void UnsetParameters(FRHICommandList& RHICmdList, const FScene* Scene)
-	{
-		constexpr bool bIsHeightField = PrimitiveType == DFPT_HeightField;
-		const FDistanceFieldSceneData& SceneData = Scene->DistanceFieldSceneData;
-		const auto& SceneObjectBuffers = TSelector<bIsHeightField>()(*SceneData.GetHeightFieldObjectBuffers(), *SceneData.GetCurrentObjectBuffers());
-		auto& CulledObjectBuffers = TSelector<bIsHeightField>()(GShadowCulledHeightFieldObjectBuffers, GShadowCulledObjectBuffers);
-
-		ObjectBufferParameters.UnsetParameters(RHICmdList, RHICmdList.GetBoundComputeShader(), SceneObjectBuffers);
-		CulledObjectParameters.UnsetParameters(RHICmdList, RHICmdList.GetBoundComputeShader());
-
-		FRHITransitionInfo TransitionInfos[4];
-		TransitionInfos[0] = FRHITransitionInfo(CulledObjectBuffers.Buffers.ObjectIndirectArguments.UAV, ERHIAccess::Unknown, ERHIAccess::IndirectArgs | ERHIAccess::SRVMask);
-		TransitionInfos[1] = FRHITransitionInfo(CulledObjectBuffers.Buffers.Bounds.UAV, ERHIAccess::Unknown, ERHIAccess::SRVMask);
-		TransitionInfos[2] = FRHITransitionInfo(CulledObjectBuffers.Buffers.Data.UAV, ERHIAccess::Unknown, ERHIAccess::SRVMask);
-		TransitionInfos[3] = FRHITransitionInfo(CulledObjectBuffers.Buffers.BoxBounds.UAV, ERHIAccess::Unknown, ERHIAccess::SRVMask);
-		RHICmdList.Transition(MakeArrayView(TransitionInfos, UE_ARRAY_COUNT(TransitionInfos)));
-	}
-
-private:
-
-	using FDistanceFieldObjectBufferParametersType = TDistanceFieldObjectBufferParameters<PrimitiveType>;
-	LAYOUT_FIELD(FDistanceFieldObjectBufferParametersType, ObjectBufferParameters)
-	using FDistanceFieldCulledObjectParametersType = TDistanceFieldCulledObjectBufferParameters<PrimitiveType>;
-	LAYOUT_FIELD(FDistanceFieldCulledObjectParametersType, CulledObjectParameters)
-	LAYOUT_FIELD(FShaderParameter, ObjectBoundingGeometryIndexCount)
-	LAYOUT_FIELD(FShaderParameter, WorldToShadow)
-	LAYOUT_FIELD(FShaderParameter, NumShadowHullPlanes)
-	LAYOUT_FIELD(FShaderParameter, ShadowBoundingSphere)
-	LAYOUT_FIELD(FShaderParameter, ShadowConvexHull)
 };
 
-IMPLEMENT_SHADER_TYPE(template<>, TCullObjectsForShadowCS<DFPT_SignedDistanceField>, TEXT("/Engine/Private/DistanceFieldShadowing.usf"), TEXT("CullObjectsForShadowCS"), SF_Compute);
-IMPLEMENT_SHADER_TYPE(template<>, TCullObjectsForShadowCS<DFPT_HeightField>, TEXT("/Engine/Private/DistanceFieldShadowing.usf"), TEXT("CullObjectsForShadowCS"), SF_Compute);
+IMPLEMENT_GLOBAL_SHADER(FCullObjectsForShadowCS, "/Engine/Private/DistanceFieldShadowing.usf", "CullObjectsForShadowCS", SF_Compute);
 
 /**  */
-template <EDistanceFieldPrimitiveType PrimitiveType>
-class TShadowObjectCullVS : public FGlobalShader
+class FShadowObjectCullVS : public FGlobalShader
 {
-	DECLARE_SHADER_TYPE(TShadowObjectCullVS,Global);
-public:
+	DECLARE_GLOBAL_SHADER(FShadowObjectCullVS);
+	SHADER_USE_PARAMETER_STRUCT(FShadowObjectCullVS, FGlobalShader);
+
+	BEGIN_SHADER_PARAMETER_STRUCT(FParameters, )
+		SHADER_PARAMETER_STRUCT_INCLUDE(FDistanceFieldObjectBufferParameters, ObjectBufferParameters)
+		SHADER_PARAMETER_STRUCT_INCLUDE(FDistanceFieldCulledObjectBufferParameters, CulledObjectBufferParameters)
+		SHADER_PARAMETER(FMatrix, WorldToShadow)
+		SHADER_PARAMETER(float, MinExpandRadius)
+	END_SHADER_PARAMETER_STRUCT()
+
+	class FPrimitiveType : SHADER_PERMUTATION_INT("DISTANCEFIELD_PRIMITIVE_TYPE", 2);
+	using FPermutationDomain = TShaderPermutationDomain<FPrimitiveType>;
 
 	static bool ShouldCompilePermutation(const FGlobalShaderPermutationParameters& Parameters) 
 	{ 
 		return IsFeatureLevelSupported(Parameters.Platform, ERHIFeatureLevel::SM5) && DoesPlatformSupportDistanceFieldShadowing(Parameters.Platform); 
 	}
-
-	static void ModifyCompilationEnvironment(const FGlobalShaderPermutationParameters& Parameters, FShaderCompilerEnvironment& OutEnvironment)
-	{
-		FGlobalShader::ModifyCompilationEnvironment(Parameters, OutEnvironment);
-		OutEnvironment.SetDefine(TEXT("DISTANCEFIELD_PRIMITIVE_TYPE"), (int32)PrimitiveType);
-	}
-
-	TShadowObjectCullVS(const ShaderMetaType::CompiledShaderInitializerType& Initializer):
-		FGlobalShader(Initializer)
-	{
-		ObjectParameters.Bind(Initializer.ParameterMap);
-		WorldToShadow.Bind(Initializer.ParameterMap, TEXT("WorldToShadow"));
-		MinExpandRadius.Bind(Initializer.ParameterMap, TEXT("MinExpandRadius"));
-	}
-
-	TShadowObjectCullVS() {}
-
-	void SetParameters(FRHICommandList& RHICmdList, const FSceneView& View, FVector2D NumGroupsValue, const FMatrix& WorldToShadowMatrixValue, float ShadowRadius)
-	{
-		FRHIVertexShader* ShaderRHI = RHICmdList.GetBoundVertexShader();
-		FGlobalShader::SetParameters<FViewUniformShaderParameters>(RHICmdList, ShaderRHI, View.ViewUniformBuffer);
-		
-		constexpr bool bIsHeightField = PrimitiveType == DFPT_HeightField;
-		auto& CulledObjectBuffers = TSelector<bIsHeightField>()(GShadowCulledHeightFieldObjectBuffers, GShadowCulledObjectBuffers);
-
-		FRHITexture* TextureAtlas;
-		FRHITexture* VisibilityAtlas;
-		int32 AtlasSizeX;
-		int32 AtlasSizeY;
-		int32 AtlasSizeZ;
-
-		if (bIsHeightField)
-		{
-			TextureAtlas = GHeightFieldTextureAtlas.GetAtlasTexture();
-			VisibilityAtlas = GHFVisibilityTextureAtlas.GetAtlasTexture();
-			AtlasSizeX = GHeightFieldTextureAtlas.GetSizeX();
-			AtlasSizeY = GHeightFieldTextureAtlas.GetSizeY();
-			AtlasSizeZ = 1;
-		}
-		else
-		{
-			TextureAtlas = GDistanceFieldVolumeTextureAtlas.VolumeTextureRHI;
-			VisibilityAtlas = nullptr;
-			AtlasSizeX = GDistanceFieldVolumeTextureAtlas.GetSizeX();
-			AtlasSizeY = GDistanceFieldVolumeTextureAtlas.GetSizeY();
-			AtlasSizeZ = GDistanceFieldVolumeTextureAtlas.GetSizeZ();
-		}
-
-		ObjectParameters.Set(RHICmdList, ShaderRHI, CulledObjectBuffers.Buffers, TextureAtlas, FIntVector(AtlasSizeX, AtlasSizeY, AtlasSizeZ), VisibilityAtlas);
-
-		SetShaderValue(RHICmdList, ShaderRHI, WorldToShadow, WorldToShadowMatrixValue);
-
-		const float MinExpandRadiusValue = (bIsHeightField ? 0.87f : 1.414f) * ShadowRadius / FMath::Min(NumGroupsValue.X, NumGroupsValue.Y);
-		SetShaderValue(RHICmdList, ShaderRHI, MinExpandRadius, MinExpandRadiusValue);
-	}
-
-private:
-	using FDistanceFieldObjectBufferParametersType = TDistanceFieldCulledObjectBufferParameters<PrimitiveType>;
-	LAYOUT_FIELD(FDistanceFieldObjectBufferParametersType, ObjectParameters)
-	LAYOUT_FIELD(FShaderParameter, WorldToShadow)
-	LAYOUT_FIELD(FShaderParameter, MinExpandRadius)
 };
 
-IMPLEMENT_SHADER_TYPE(template<>, TShadowObjectCullVS<DFPT_SignedDistanceField>, TEXT("/Engine/Private/DistanceFieldShadowing.usf"), TEXT("ShadowObjectCullVS"), SF_Vertex);
-IMPLEMENT_SHADER_TYPE(template<>, TShadowObjectCullVS<DFPT_HeightField>, TEXT("/Engine/Private/DistanceFieldShadowing.usf"), TEXT("ShadowObjectCullVS"), SF_Vertex);
+IMPLEMENT_GLOBAL_SHADER(FShadowObjectCullVS, "/Engine/Private/DistanceFieldShadowing.usf", "ShadowObjectCullVS", SF_Vertex);
 
-template <bool bCountingPass, EDistanceFieldPrimitiveType PrimitiveType>
-class TShadowObjectCullPS : public FGlobalShader
+class FShadowObjectCullPS : public FGlobalShader
 {
-	DECLARE_SHADER_TYPE(TShadowObjectCullPS, Global);
-public:
+	DECLARE_GLOBAL_SHADER(FShadowObjectCullPS);
+	SHADER_USE_PARAMETER_STRUCT(FShadowObjectCullPS, FGlobalShader);
+
+	BEGIN_SHADER_PARAMETER_STRUCT(FParameters, )
+		SHADER_PARAMETER_STRUCT_INCLUDE(FDistanceFieldObjectBufferParameters, ObjectBufferParameters)
+		SHADER_PARAMETER_STRUCT_INCLUDE(FDistanceFieldCulledObjectBufferParameters, CulledObjectBufferParameters)
+		SHADER_PARAMETER_STRUCT_INCLUDE(FLightTileIntersectionParameters, LightTileIntersectionParameters)
+	END_SHADER_PARAMETER_STRUCT()
+
+	class FPrimitiveType : SHADER_PERMUTATION_INT("DISTANCEFIELD_PRIMITIVE_TYPE", 2);
+	class FCountingPass : SHADER_PERMUTATION_BOOL("SCATTER_CULLING_COUNT_PASS");
+	using FPermutationDomain = TShaderPermutationDomain<FPrimitiveType, FCountingPass>;
 
 	static bool ShouldCompilePermutation(const FGlobalShaderPermutationParameters& Parameters)
 	{
@@ -357,82 +210,17 @@ public:
 	static void ModifyCompilationEnvironment(const FGlobalShaderPermutationParameters& Parameters, FShaderCompilerEnvironment& OutEnvironment)
 	{
 		FGlobalShader::ModifyCompilationEnvironment(Parameters,OutEnvironment);
-		FLightTileIntersectionParameters::ModifyCompilationEnvironment(Parameters.Platform, OutEnvironment);
-		OutEnvironment.SetDefine(TEXT("SCATTER_CULLING_COUNT_PASS"), bCountingPass ? 1 : 0);
-		OutEnvironment.SetDefine(TEXT("DISTANCEFIELD_PRIMITIVE_TYPE"), (int32)PrimitiveType);
 	}
-
-	/** Default constructor. */
-	TShadowObjectCullPS() {}
-
-	/** Initialization constructor. */
-	TShadowObjectCullPS(const ShaderMetaType::CompiledShaderInitializerType& Initializer)
-		: FGlobalShader(Initializer)
-	{
-		ObjectParameters.Bind(Initializer.ParameterMap);
-		LightTileIntersectionParameters.Bind(Initializer.ParameterMap);
-	}
-
-	void SetParameters(
-		FRHICommandList& RHICmdList, 
-		const FSceneView& View, 
-		FLightTileIntersectionResources* TileIntersectionResources)
-	{
-		FRHIPixelShader* ShaderRHI = RHICmdList.GetBoundPixelShader();
-		FGlobalShader::SetParameters<FViewUniformShaderParameters>(RHICmdList, ShaderRHI, View.ViewUniformBuffer);
-
-		constexpr bool bIsHeightField = PrimitiveType == DFPT_HeightField;
-		auto& CulledObjectBuffers = TSelector<bIsHeightField>()(GShadowCulledHeightFieldObjectBuffers, GShadowCulledObjectBuffers);
-
-		FRHITexture* TextureAtlas;
-		FRHITexture* VisibilityAtlas;
-		int32 AtlasSizeX;
-		int32 AtlasSizeY;
-		int32 AtlasSizeZ;
-
-		if (bIsHeightField)
-		{
-			TextureAtlas = GHeightFieldTextureAtlas.GetAtlasTexture();
-			VisibilityAtlas = GHFVisibilityTextureAtlas.GetAtlasTexture();
-			AtlasSizeX = GHeightFieldTextureAtlas.GetSizeX();
-			AtlasSizeY = GHeightFieldTextureAtlas.GetSizeY();
-			AtlasSizeZ = 1;
-		}
-		else
-		{
-			TextureAtlas = GDistanceFieldVolumeTextureAtlas.VolumeTextureRHI;
-			VisibilityAtlas = nullptr;
-			AtlasSizeX = GDistanceFieldVolumeTextureAtlas.GetSizeX();
-			AtlasSizeY = GDistanceFieldVolumeTextureAtlas.GetSizeY();
-			AtlasSizeZ = GDistanceFieldVolumeTextureAtlas.GetSizeZ();
-		}
-
-		ObjectParameters.Set(RHICmdList, ShaderRHI, CulledObjectBuffers.Buffers, TextureAtlas, FIntVector(AtlasSizeX, AtlasSizeY, AtlasSizeZ), VisibilityAtlas);
-
-		LightTileIntersectionParameters.Set(RHICmdList, ShaderRHI, *TileIntersectionResources);
-	}
-
-	void GetUAVs(const FSceneView& View, FLightTileIntersectionResources* TileIntersectionResources, TArray<FRHIUnorderedAccessView*>& UAVs)
-	{
-		LightTileIntersectionParameters.GetUAVs(*TileIntersectionResources, UAVs);
-	}
-
-private:
-	using FDistanceFieldObjectBufferParametersType = TDistanceFieldCulledObjectBufferParameters<PrimitiveType>;
-	LAYOUT_FIELD(FDistanceFieldObjectBufferParametersType, ObjectParameters)
-	LAYOUT_FIELD(FLightTileIntersectionParameters, LightTileIntersectionParameters)
 };
 
-#define VARIATION(PrimitiveType) \
-	typedef TShadowObjectCullPS<true, PrimitiveType> FShadowObjectCullPS##1##PrimitiveType; \
-	typedef TShadowObjectCullPS<false, PrimitiveType> FShadowObjectCullPS##0##PrimitiveType; \
-	IMPLEMENT_SHADER_TYPE(template<>, FShadowObjectCullPS##1##PrimitiveType, TEXT("/Engine/Private/DistanceFieldShadowing.usf"), TEXT("ShadowObjectCullPS"), SF_Pixel); \
-	IMPLEMENT_SHADER_TYPE(template<>, FShadowObjectCullPS##0##PrimitiveType, TEXT("/Engine/Private/DistanceFieldShadowing.usf"), TEXT("ShadowObjectCullPS"), SF_Pixel)
-
-VARIATION(DFPT_SignedDistanceField);
-VARIATION(DFPT_HeightField);
+BEGIN_SHADER_PARAMETER_STRUCT(FShadowMeshSDFObjectCull, )
+	SHADER_PARAMETER_STRUCT_INCLUDE(FShadowObjectCullVS::FParameters, VS)
+	SHADER_PARAMETER_STRUCT_INCLUDE(FShadowObjectCullPS::FParameters, PS)
+	SHADER_PARAMETER_RDG_BUFFER(Buffer<uint>, MeshSDFIndirectArgs)
+	RENDER_TARGET_BINDING_SLOTS()
+END_SHADER_PARAMETER_STRUCT()
 
-#undef VARIATION
+IMPLEMENT_GLOBAL_SHADER(FShadowObjectCullPS, "/Engine/Private/DistanceFieldShadowing.usf", "ShadowObjectCullPS", SF_Pixel);
 
 enum EDistanceFieldShadowingType
 {
@@ -441,11 +229,43 @@ enum EDistanceFieldShadowingType
 	DFS_PointLightTiledCulling
 };
 
-template<EDistanceFieldShadowingType ShadowingType, uint32 DFShadowQuality, EDistanceFieldPrimitiveType PrimitiveType, bool bHasPrevOutput>
-class TDistanceFieldShadowingCS : public FGlobalShader
+//template<EDistanceFieldShadowingType ShadowingType, uint32 DFShadowQuality, EDistanceFieldPrimitiveType PrimitiveType, bool bHasPrevOutput>
+class FDistanceFieldShadowingCS : public FGlobalShader
 {
-	DECLARE_SHADER_TYPE(TDistanceFieldShadowingCS, Global);
-public:
+	DECLARE_GLOBAL_SHADER(FDistanceFieldShadowingCS);
+	SHADER_USE_PARAMETER_STRUCT(FDistanceFieldShadowingCS, FGlobalShader);
+
+	BEGIN_SHADER_PARAMETER_STRUCT(FParameters, )
+		SHADER_PARAMETER_RDG_TEXTURE_UAV(RWTexture2D<float2>, RWShadowFactors)
+		SHADER_PARAMETER_STRUCT_REF(FViewUniformShaderParameters, View)
+		SHADER_PARAMETER_RDG_UNIFORM_BUFFER(FSceneTextureUniformParameters, SceneTextures)
+		SHADER_PARAMETER(FVector2D, NumGroups)
+		SHADER_PARAMETER(FVector, LightDirection)
+		SHADER_PARAMETER(FVector4, LightPositionAndInvRadius)
+		SHADER_PARAMETER(float, LightSourceRadius)
+		SHADER_PARAMETER(float, RayStartOffsetDepthScale)
+		SHADER_PARAMETER(FVector, TanLightAngleAndNormalThreshold)
+		SHADER_PARAMETER(FIntRect, ScissorRectMinAndSize)
+		SHADER_PARAMETER_STRUCT_INCLUDE(FDistanceFieldObjectBufferParameters, ObjectBufferParameters)
+		SHADER_PARAMETER_STRUCT_INCLUDE(FDistanceFieldCulledObjectBufferParameters, CulledObjectBufferParameters)
+		SHADER_PARAMETER_STRUCT_INCLUDE(FLightTileIntersectionParameters, LightTileIntersectionParameters)
+		SHADER_PARAMETER_STRUCT_INCLUDE(FDistanceFieldAtlasParameters, DistanceFieldAtlasParameters)
+		SHADER_PARAMETER_STRUCT_INCLUDE(FHeightFieldAtlasParameters, HeightFieldAtlasParameters)
+		SHADER_PARAMETER(FMatrix, WorldToShadow)
+		SHADER_PARAMETER(float, TwoSidedMeshDistanceBias)
+		SHADER_PARAMETER(float, MinDepth)
+		SHADER_PARAMETER(float, MaxDepth)
+		SHADER_PARAMETER(uint32, DownsampleFactor)
+		SHADER_PARAMETER(FVector2D, InvOutputBufferSize)
+		SHADER_PARAMETER_RDG_TEXTURE(Texture2D, ShadowFactorsTexture)
+		SHADER_PARAMETER_SAMPLER(SamplerState, ShadowFactorsSampler)
+	END_SHADER_PARAMETER_STRUCT()
+
+	class FCullingType : SHADER_PERMUTATION_INT("CULLING_TYPE", 3);
+	class FShadowQuality : SHADER_PERMUTATION_INT("DF_SHADOW_QUALITY", 3);
+	class FPrimitiveType : SHADER_PERMUTATION_INT("DISTANCEFIELD_PRIMITIVE_TYPE", 2);
+	class FHasPreviousOutput : SHADER_PERMUTATION_BOOL("HAS_PREVIOUS_OUTPUT");
+	using FPermutationDomain = TShaderPermutationDomain<FCullingType, FShadowQuality, FPrimitiveType, FHasPreviousOutput>;
 
 	static bool ShouldCompilePermutation(const FGlobalShaderPermutationParameters& Parameters)
 	{
@@ -454,209 +274,33 @@ public:
 
 	static void ModifyCompilationEnvironment(const FGlobalShaderPermutationParameters& Parameters, FShaderCompilerEnvironment& OutEnvironment)
 	{
-		FLightTileIntersectionParameters::ModifyCompilationEnvironment(Parameters.Platform, OutEnvironment);
 		OutEnvironment.SetDefine(TEXT("THREADGROUP_SIZEX"), GDistanceFieldShadowTileSizeX);
 		OutEnvironment.SetDefine(TEXT("THREADGROUP_SIZEY"), GDistanceFieldShadowTileSizeY);
-		OutEnvironment.SetDefine(TEXT("SCATTER_TILE_CULLING"), ShadowingType == DFS_DirectionalLightScatterTileCulling);
-		OutEnvironment.SetDefine(TEXT("POINT_LIGHT"), ShadowingType == DFS_PointLightTiledCulling);
-		OutEnvironment.SetDefine(TEXT("DF_SHADOW_QUALITY"), DFShadowQuality);
-		OutEnvironment.SetDefine(TEXT("DISTANCEFIELD_PRIMITIVE_TYPE"), (int32)PrimitiveType);
-		OutEnvironment.SetDefine(TEXT("HAS_PREVIOUS_OUTPUT"), (int32)bHasPrevOutput);
 		OutEnvironment.SetDefine(TEXT("PLATFORM_SUPPORTS_TYPED_UAV_LOAD"), (int32)RHISupports4ComponentUAVReadWrite(Parameters.Platform));
 	}
-
-	/** Default constructor. */
-	TDistanceFieldShadowingCS() {}
-
-	/** Initialization constructor. */
-	TDistanceFieldShadowingCS(const ShaderMetaType::CompiledShaderInitializerType& Initializer)
-		: FGlobalShader(Initializer)
-	{
-		ShadowFactors.Bind(Initializer.ParameterMap, TEXT("ShadowFactors"));
-		NumGroups.Bind(Initializer.ParameterMap, TEXT("NumGroups"));
-		LightDirection.Bind(Initializer.ParameterMap, TEXT("LightDirection"));
-		LightSourceRadius.Bind(Initializer.ParameterMap, TEXT("LightSourceRadius"));
-		RayStartOffsetDepthScale.Bind(Initializer.ParameterMap, TEXT("RayStartOffsetDepthScale"));
-		LightPositionAndInvRadius.Bind(Initializer.ParameterMap, TEXT("LightPositionAndInvRadius"));
-		TanLightAngleAndNormalThreshold.Bind(Initializer.ParameterMap, TEXT("TanLightAngleAndNormalThreshold"));
-		ScissorRectMinAndSize.Bind(Initializer.ParameterMap, TEXT("ScissorRectMinAndSize"));
-		ObjectParameters.Bind(Initializer.ParameterMap);
-		LightTileIntersectionParameters.Bind(Initializer.ParameterMap);
-		WorldToShadow.Bind(Initializer.ParameterMap, TEXT("WorldToShadow"));
-		TwoSidedMeshDistanceBias.Bind(Initializer.ParameterMap, TEXT("TwoSidedMeshDistanceBias"));
-		MinDepth.Bind(Initializer.ParameterMap, TEXT("MinDepth"));
-		MaxDepth.Bind(Initializer.ParameterMap, TEXT("MaxDepth"));
-		DownsampleFactor.Bind(Initializer.ParameterMap, TEXT("DownsampleFactor"));
-		InvOutputBufferSize.Bind(Initializer.ParameterMap, TEXT("InvOutputBufferSize"));
-		ShadowFactorsTexture.Bind(Initializer.ParameterMap, TEXT("ShadowFactorsTexture"));
-		ShadowFactorsSampler.Bind(Initializer.ParameterMap, TEXT("ShadowFactorsSampler"));
-	}
-
-	void SetParameters(
-		FRHIComputeCommandList& RHICmdList, 
-		const FSceneView& View, 
-		const FProjectedShadowInfo* ProjectedShadowInfo,
-		const FSceneRenderTargetItem& ShadowFactorsValue, 
-		FVector2D NumGroupsValue,
-		const FIntRect& ScissorRect,
-		FLightTileIntersectionResources* TileIntersectionResources,
-		FRHITexture* PrevOutput)
-	{
-		FRHIComputeShader* ShaderRHI = RHICmdList.GetBoundComputeShader();
-
-		FGlobalShader::SetParameters<FViewUniformShaderParameters>(RHICmdList, ShaderRHI, View.ViewUniformBuffer);
-
-		ShadowFactors.SetTexture(RHICmdList, ShaderRHI, ShadowFactorsValue.ShaderResourceTexture, ShadowFactorsValue.UAV);
-
-		constexpr bool bIsHeightField = PrimitiveType == DFPT_HeightField;
-		auto& CulledObjectBuffers = TSelector<bIsHeightField>()(GShadowCulledHeightFieldObjectBuffers, GShadowCulledObjectBuffers);
-
-		FRHITexture* TextureAtlas;
-		FRHITexture* VisibilityAtlas;
-		int32 AtlasSizeX;
-		int32 AtlasSizeY;
-		int32 AtlasSizeZ;
-
-		if (bIsHeightField)
-		{
-			TextureAtlas = GHeightFieldTextureAtlas.GetAtlasTexture();
-			VisibilityAtlas = GHFVisibilityTextureAtlas.GetAtlasTexture();
-			AtlasSizeX = GHeightFieldTextureAtlas.GetSizeX();
-			AtlasSizeY = GHeightFieldTextureAtlas.GetSizeY();
-			AtlasSizeZ = 1;
-		}
-		else
-		{
-			TextureAtlas = GDistanceFieldVolumeTextureAtlas.VolumeTextureRHI;
-			VisibilityAtlas = nullptr;
-			AtlasSizeX = GDistanceFieldVolumeTextureAtlas.GetSizeX();
-			AtlasSizeY = GDistanceFieldVolumeTextureAtlas.GetSizeY();
-			AtlasSizeZ = GDistanceFieldVolumeTextureAtlas.GetSizeZ();
-		}
-
-		ObjectParameters.Set(RHICmdList, ShaderRHI, CulledObjectBuffers.Buffers, TextureAtlas, FIntVector(AtlasSizeX, AtlasSizeY, AtlasSizeZ), VisibilityAtlas);
-
-		SetShaderValue(RHICmdList, ShaderRHI, NumGroups, NumGroupsValue);
-
-		const FLightSceneProxy& LightProxy = *(ProjectedShadowInfo->GetLightSceneInfo().Proxy);
-
-		FLightShaderParameters LightParameters;
-		LightProxy.GetLightShaderParameters(LightParameters);
-
-		SetShaderValue(RHICmdList, ShaderRHI, LightDirection, LightParameters.Direction);
-		FVector4 LightPositionAndInvRadiusValue(LightParameters.Position, LightParameters.InvRadius);
-		SetShaderValue(RHICmdList, ShaderRHI, LightPositionAndInvRadius, LightPositionAndInvRadiusValue);
-		// Default light source radius of 0 gives poor results
-		SetShaderValue(RHICmdList, ShaderRHI, LightSourceRadius, LightParameters.SourceRadius == 0 ? 20 : FMath::Clamp(LightParameters.SourceRadius, .001f, 1.0f / (4 * LightParameters.InvRadius)));
-
-		SetShaderValue(RHICmdList, ShaderRHI, RayStartOffsetDepthScale, LightProxy.GetRayStartOffsetDepthScale());
-
-		constexpr float MaxLightAngle = bIsHeightField ? 45.0f : 5.0f;
-		const float MinLightAngle = bIsHeightField ? FMath::Min(GMinDirectionalLightAngleForRTHF, MaxLightAngle) : 0.001f;
-		const float LightSourceAngle = FMath::Clamp(LightProxy.GetLightSourceAngle(), MinLightAngle, MaxLightAngle) * PI / 180.0f;
-		const FVector TanLightAngleAndNormalThresholdValue(FMath::Tan(LightSourceAngle), FMath::Cos(PI / 2 + LightSourceAngle), LightProxy.GetTraceDistance());
-		SetShaderValue(RHICmdList, ShaderRHI, TanLightAngleAndNormalThreshold, TanLightAngleAndNormalThresholdValue);
-
-		SetShaderValue(RHICmdList, ShaderRHI, ScissorRectMinAndSize, FIntRect(ScissorRect.Min, ScissorRect.Size()));
-
-		check(TileIntersectionResources || !LightTileIntersectionParameters.IsBound());
-
-		if (TileIntersectionResources)
-		{
-			LightTileIntersectionParameters.Set(RHICmdList, ShaderRHI, *TileIntersectionResources);
-		}
-
-		FMatrix WorldToShadowMatrixValue = FTranslationMatrix(ProjectedShadowInfo->PreShadowTranslation) * ProjectedShadowInfo->SubjectAndReceiverMatrix;
-		SetShaderValue(RHICmdList, ShaderRHI, WorldToShadow, WorldToShadowMatrixValue);
-
-		SetShaderValue(RHICmdList, ShaderRHI, TwoSidedMeshDistanceBias, GTwoSidedMeshDistanceBias);
-
-		if (ProjectedShadowInfo->bDirectionalLight)
-		{
-			SetShaderValue(RHICmdList, ShaderRHI, MinDepth, ProjectedShadowInfo->CascadeSettings.SplitNear - ProjectedShadowInfo->CascadeSettings.SplitNearFadeRegion);
-			SetShaderValue(RHICmdList, ShaderRHI, MaxDepth, ProjectedShadowInfo->CascadeSettings.SplitFar);
-		}
-		else
-		{
-			check(!bIsHeightField);
-			//@todo - set these up for point lights as well
-			SetShaderValue(RHICmdList, ShaderRHI, MinDepth, 0.0f);
-			SetShaderValue(RHICmdList, ShaderRHI, MaxDepth, HALF_WORLD_MAX);
-		}
-
-		SetShaderValue(RHICmdList, ShaderRHI, DownsampleFactor, GetDFShadowDownsampleFactor());
-
-		if (InvOutputBufferSize.IsBound())
-		{
-			check(ShadowFactorsTexture.IsBound() && ShadowFactorsSampler.IsBound() && PrevOutput);
-
-			const FIntPoint OutputBufferSize = GetBufferSizeForDFShadows();
-			SetShaderValue(RHICmdList, ShaderRHI, InvOutputBufferSize, FVector2D(1.f / OutputBufferSize.X, 1.f / OutputBufferSize.Y));
-			SetTextureParameter(RHICmdList, ShaderRHI, ShadowFactorsTexture, ShadowFactorsSampler, TStaticSamplerState<>::GetRHI(), PrevOutput);
-		}
-	}
-
-	void UnsetParameters(FRHIComputeCommandList& RHICmdList)
-	{
-		ShadowFactors.UnsetUAV(RHICmdList, RHICmdList.GetBoundComputeShader());
-	}
-
-	static const TCHAR* GetSourceFilename()
-	{
-		return TEXT("/Engine/Private/DistanceFieldShadowing.usf");
-	}
-
-	static const TCHAR* GetFunctionName()
-	{
-		return TEXT("DistanceFieldShadowingCS");
-	}
-
-private:
-	LAYOUT_FIELD(FRWShaderParameter, ShadowFactors);
-	LAYOUT_FIELD(FShaderParameter, NumGroups);
-	LAYOUT_FIELD(FShaderParameter, LightDirection);
-	LAYOUT_FIELD(FShaderParameter, LightPositionAndInvRadius);
-	LAYOUT_FIELD(FShaderParameter, LightSourceRadius);
-	LAYOUT_FIELD(FShaderParameter, RayStartOffsetDepthScale);
-	LAYOUT_FIELD(FShaderParameter, TanLightAngleAndNormalThreshold);
-	LAYOUT_FIELD(FShaderParameter, ScissorRectMinAndSize);
-	using FDistanceFieldObjectBufferParametersType = TDistanceFieldCulledObjectBufferParameters<PrimitiveType>;
-	LAYOUT_FIELD(FDistanceFieldObjectBufferParametersType, ObjectParameters)
-	LAYOUT_FIELD(FLightTileIntersectionParameters, LightTileIntersectionParameters);
-	LAYOUT_FIELD(FShaderParameter, WorldToShadow);
-	LAYOUT_FIELD(FShaderParameter, TwoSidedMeshDistanceBias);
-	LAYOUT_FIELD(FShaderParameter, MinDepth);
-	LAYOUT_FIELD(FShaderParameter, MaxDepth);
-	LAYOUT_FIELD(FShaderParameter, DownsampleFactor);
-	LAYOUT_FIELD(FShaderParameter, InvOutputBufferSize);
-	LAYOUT_FIELD(FShaderResourceParameter, ShadowFactorsTexture);
-	LAYOUT_FIELD(FShaderResourceParameter, ShadowFactorsSampler);
 };
 
-// #define avoids a lot of code duplication
-#define VARIATION(A, B, C) \
-	typedef TDistanceFieldShadowingCS<A, 1, B, C> FDistanceFieldShadowingCS##A##1##B##C; \
-	typedef TDistanceFieldShadowingCS<A, 2, B, C> FDistanceFieldShadowingCS##A##2##B##C; \
-    typedef TDistanceFieldShadowingCS<A, 3, B, C> FDistanceFieldShadowingCS##A##3##B##C; \
-	IMPLEMENT_SHADER_TYPE2(FDistanceFieldShadowingCS##A##1##B##C, SF_Compute); \
-	IMPLEMENT_SHADER_TYPE2(FDistanceFieldShadowingCS##A##2##B##C, SF_Compute); \
-    IMPLEMENT_SHADER_TYPE2(FDistanceFieldShadowingCS##A##3##B##C, SF_Compute)
-
-VARIATION(DFS_DirectionalLightScatterTileCulling, DFPT_SignedDistanceField, false);
-VARIATION(DFS_DirectionalLightScatterTileCulling, DFPT_HeightField, false);
-VARIATION(DFS_DirectionalLightScatterTileCulling, DFPT_HeightField, true);
-VARIATION(DFS_DirectionalLightTiledCulling, DFPT_SignedDistanceField, false);
-VARIATION(DFS_DirectionalLightTiledCulling, DFPT_HeightField, false);
-VARIATION(DFS_DirectionalLightTiledCulling, DFPT_HeightField, true);
-VARIATION(DFS_PointLightTiledCulling, DFPT_SignedDistanceField, false);
-
-#undef VARIATION
-
-template<bool bUpsampleRequired>
-class TDistanceFieldShadowingUpsamplePS : public FGlobalShader
+IMPLEMENT_GLOBAL_SHADER(FDistanceFieldShadowingCS, "/Engine/Private/DistanceFieldShadowing.usf", "DistanceFieldShadowingCS", SF_Compute);
+
+class FDistanceFieldShadowingUpsamplePS : public FGlobalShader
 {
-	DECLARE_SHADER_TYPE(TDistanceFieldShadowingUpsamplePS, Global);
-public:
+	DECLARE_GLOBAL_SHADER(FDistanceFieldShadowingUpsamplePS);
+	SHADER_USE_PARAMETER_STRUCT(FDistanceFieldShadowingUpsamplePS, FGlobalShader);
+
+	BEGIN_SHADER_PARAMETER_STRUCT(FParameters, )
+		SHADER_PARAMETER_RDG_UNIFORM_BUFFER(FSceneTextureUniformParameters, SceneTextures)
+		SHADER_PARAMETER_STRUCT_REF(FViewUniformShaderParameters, View)
+		SHADER_PARAMETER_RDG_TEXTURE(Texture2D, ShadowFactorsTexture)
+		SHADER_PARAMETER_SAMPLER(SamplerState, ShadowFactorsSampler)
+		SHADER_PARAMETER(FIntRect, ScissorRectMinAndSize)
+		SHADER_PARAMETER(float, FadePlaneOffset)
+		SHADER_PARAMETER(float, InvFadePlaneLength)
+		SHADER_PARAMETER(float, NearFadePlaneOffset)
+		SHADER_PARAMETER(float, InvNearFadePlaneLength)
+	END_SHADER_PARAMETER_STRUCT()
+
+	class FUpsample : SHADER_PERMUTATION_BOOL("UPSAMPLE_REQUIRED");
+	using FPermutationDomain = TShaderPermutationDomain<FUpsample>;
 
 	static bool ShouldCompilePermutation(const FGlobalShaderPermutationParameters& Parameters)
 	{
@@ -666,78 +310,22 @@ public:
 	static void ModifyCompilationEnvironment(const FGlobalShaderPermutationParameters& Parameters, FShaderCompilerEnvironment& OutEnvironment)
 	{
 		OutEnvironment.SetDefine(TEXT("DOWNSAMPLE_FACTOR"), GAODownsampleFactor);
-		OutEnvironment.SetDefine(TEXT("UPSAMPLE_REQUIRED"), bUpsampleRequired);
-	}
-
-	/** Default constructor. */
-	TDistanceFieldShadowingUpsamplePS() {}
-
-	/** Initialization constructor. */
-	TDistanceFieldShadowingUpsamplePS(const ShaderMetaType::CompiledShaderInitializerType& Initializer)
-		: FGlobalShader(Initializer)
-	{
-		ShadowFactorsTexture.Bind(Initializer.ParameterMap,TEXT("ShadowFactorsTexture"));
-		ShadowFactorsSampler.Bind(Initializer.ParameterMap,TEXT("ShadowFactorsSampler"));
-		ScissorRectMinAndSize.Bind(Initializer.ParameterMap,TEXT("ScissorRectMinAndSize"));
-		FadePlaneOffset.Bind(Initializer.ParameterMap,TEXT("FadePlaneOffset"));
-		InvFadePlaneLength.Bind(Initializer.ParameterMap,TEXT("InvFadePlaneLength"));
-		NearFadePlaneOffset.Bind(Initializer.ParameterMap,TEXT("NearFadePlaneOffset"));
-		InvNearFadePlaneLength.Bind(Initializer.ParameterMap,TEXT("InvNearFadePlaneLength"));
 	}
-
-	void SetParameters(FRHICommandList& RHICmdList, const FSceneView& View, const FProjectedShadowInfo* ShadowInfo, const FIntRect& ScissorRect, IPooledRenderTarget* ShadowFactorsTextureValue)
-	{
-		FRHIPixelShader* ShaderRHI = RHICmdList.GetBoundPixelShader();
-
-		FGlobalShader::SetParameters<FViewUniformShaderParameters>(RHICmdList, ShaderRHI, View.ViewUniformBuffer);
-
-		SetTextureParameter(RHICmdList, ShaderRHI, ShadowFactorsTexture, ShadowFactorsSampler, TStaticSamplerState<SF_Bilinear>::GetRHI(), ShadowFactorsTextureValue->GetRenderTargetItem().ShaderResourceTexture);
-	
-		SetShaderValue(RHICmdList, ShaderRHI, ScissorRectMinAndSize, FIntRect(ScissorRect.Min, ScissorRect.Size()));
-
-		if (ShadowInfo->bDirectionalLight && ShadowInfo->CascadeSettings.FadePlaneLength > 0)
-		{
-			SetShaderValue(RHICmdList, ShaderRHI, FadePlaneOffset, ShadowInfo->CascadeSettings.FadePlaneOffset);
-			SetShaderValue(RHICmdList, ShaderRHI, InvFadePlaneLength, 1.0f / FMath::Max(ShadowInfo->CascadeSettings.FadePlaneLength, .00001f));
-		}
-		else
-		{
-			SetShaderValue(RHICmdList, ShaderRHI, FadePlaneOffset, 0.0f);
-			SetShaderValue(RHICmdList, ShaderRHI, InvFadePlaneLength, 0.0f);
-		}
-
-		if (ShadowInfo->bDirectionalLight && ShadowInfo->CascadeSettings.SplitNearFadeRegion > 0)
-		{
-			SetShaderValue(RHICmdList, ShaderRHI, NearFadePlaneOffset, ShadowInfo->CascadeSettings.SplitNear - ShadowInfo->CascadeSettings.SplitNearFadeRegion);
-			SetShaderValue(RHICmdList, ShaderRHI, InvNearFadePlaneLength, 1.0f / FMath::Max(ShadowInfo->CascadeSettings.SplitNearFadeRegion, .00001f));
-		}
-		else
-		{
-			SetShaderValue(RHICmdList, ShaderRHI, NearFadePlaneOffset, -1.0f);
-			SetShaderValue(RHICmdList, ShaderRHI, InvNearFadePlaneLength, 1.0f);
-		}
-	}
-
-private:
-	LAYOUT_FIELD(FShaderResourceParameter, ShadowFactorsTexture);
-	LAYOUT_FIELD(FShaderResourceParameter, ShadowFactorsSampler);
-	LAYOUT_FIELD(FShaderParameter, ScissorRectMinAndSize);
-	LAYOUT_FIELD(FShaderParameter, FadePlaneOffset);
-	LAYOUT_FIELD(FShaderParameter, InvFadePlaneLength);
-	LAYOUT_FIELD(FShaderParameter, NearFadePlaneOffset);
-	LAYOUT_FIELD(FShaderParameter, InvNearFadePlaneLength);
 };
 
-IMPLEMENT_SHADER_TYPE(template<>,TDistanceFieldShadowingUpsamplePS<true>,TEXT("/Engine/Private/DistanceFieldShadowing.usf"),TEXT("DistanceFieldShadowingUpsamplePS"),SF_Pixel);
-IMPLEMENT_SHADER_TYPE(template<>,TDistanceFieldShadowingUpsamplePS<false>,TEXT("/Engine/Private/DistanceFieldShadowing.usf"),TEXT("DistanceFieldShadowingUpsamplePS"),SF_Pixel);
+IMPLEMENT_GLOBAL_SHADER(FDistanceFieldShadowingUpsamplePS, "/Engine/Private/DistanceFieldShadowing.usf", "DistanceFieldShadowingUpsamplePS", SF_Pixel);
 
 const uint32 ComputeCulledObjectStartOffsetGroupSize = 8;
 
 /**  */
 class FComputeCulledObjectStartOffsetCS : public FGlobalShader
 {
-	DECLARE_SHADER_TYPE(FComputeCulledObjectStartOffsetCS,Global)
-public:
+	DECLARE_GLOBAL_SHADER(FComputeCulledObjectStartOffsetCS);
+	SHADER_USE_PARAMETER_STRUCT(FComputeCulledObjectStartOffsetCS, FGlobalShader);
+
+	BEGIN_SHADER_PARAMETER_STRUCT(FParameters, )
+		SHADER_PARAMETER_STRUCT_INCLUDE(FLightTileIntersectionParameters, LightTileIntersectionParameters)
+	END_SHADER_PARAMETER_STRUCT()
 
 	static bool ShouldCompilePermutation(const FGlobalShaderPermutationParameters& Parameters)
 	{
@@ -749,306 +337,264 @@ public:
 		FGlobalShader::ModifyCompilationEnvironment(Parameters,OutEnvironment);
 		OutEnvironment.SetDefine(TEXT("COMPUTE_START_OFFSET_GROUP_SIZE"), ComputeCulledObjectStartOffsetGroupSize);
 	}
+};
 
-	FComputeCulledObjectStartOffsetCS(const ShaderMetaType::CompiledShaderInitializerType& Initializer)
-		: FGlobalShader(Initializer)
-	{
-		TileIntersectionParameters.Bind(Initializer.ParameterMap);
-	}
+IMPLEMENT_SHADER_TYPE(,FComputeCulledObjectStartOffsetCS,TEXT("/Engine/Private/DistanceFieldShadowing.usf"),TEXT("ComputeCulledTilesStartOffsetCS"),SF_Compute);
 
-	FComputeCulledObjectStartOffsetCS()
-	{
-	}
-	void SetParameters(FRHIComputeCommandList& RHICmdList, const FSceneView& View, FLightTileIntersectionResources* TileIntersectionResources)
+void ScatterObjectsToShadowTiles(
+	FRDGBuilder& GraphBuilder,
+	const FViewInfo& View, 
+	const FMatrix& WorldToShadowValue, 
+	float ShadowBoundingRadius,
+	bool bCountingPass, 
+	EDistanceFieldPrimitiveType PrimitiveType,
+	FIntPoint LightTileDimensions, 
+	FRDGBufferRef ObjectIndirectArguments,
+	const FDistanceFieldObjectBufferParameters& ObjectBufferParameters,
+	const FDistanceFieldCulledObjectBufferParameters& CulledObjectBufferParameters,
+	const FLightTileIntersectionParameters& LightTileIntersectionParameters)
+{
 	{
-		FRHIComputeShader* ShaderRHI = RHICmdList.GetBoundComputeShader();
+		FShadowMeshSDFObjectCull* PassParameters = GraphBuilder.AllocParameters<FShadowMeshSDFObjectCull>();
 
-		FGlobalShader::SetParameters<FViewUniformShaderParameters>(RHICmdList, ShaderRHI, View.ViewUniformBuffer);
-
-		TArray<FRHIUnorderedAccessView*> UAVs;
-		TileIntersectionParameters.GetUAVs(*TileIntersectionResources, UAVs);
-
-		TArray<FRHITransitionInfo> TransitionInfos;
-		for (FRHIUnorderedAccessView* UAV : UAVs)
+		if (GRHIRequiresRenderTargetForPixelShaderUAVs)
 		{
-			TransitionInfos.Add(FRHITransitionInfo(UAV, ERHIAccess::Unknown, ERHIAccess::UAVCompute));
+			FRDGTextureDesc DummyDesc = FRDGTextureDesc::Create2D(LightTileDimensions, PF_B8G8R8A8, FClearValueBinding::Black, TexCreate_RenderTargetable);
+			PassParameters->RenderTargets[0] = FRenderTargetBinding(GraphBuilder.CreateTexture(DummyDesc, TEXT("Dummy")), ERenderTargetLoadAction::ENoAction);
 		}
-		RHICmdList.Transition(MakeArrayView(TransitionInfos.GetData(), TransitionInfos.Num()));
 
-		TileIntersectionParameters.Set(RHICmdList, ShaderRHI, *TileIntersectionResources);
-	}
+		const float MinExpandRadiusValue = (PrimitiveType == DFPT_HeightField ? 0.87f : 1.414f) * ShadowBoundingRadius / FMath::Min(LightTileDimensions.X, LightTileDimensions.Y);
 
-	void UnsetParameters(FRHIComputeCommandList& RHICmdList, const FSceneView& View, FLightTileIntersectionResources* TileIntersectionResources)
-	{
-		TileIntersectionParameters.UnsetParameters(RHICmdList, RHICmdList.GetBoundComputeShader());
+		PassParameters->VS.ObjectBufferParameters = ObjectBufferParameters;
+		PassParameters->VS.CulledObjectBufferParameters = CulledObjectBufferParameters;
+		PassParameters->VS.WorldToShadow = WorldToShadowValue;
+		PassParameters->VS.MinExpandRadius = MinExpandRadiusValue;
+		PassParameters->PS.ObjectBufferParameters = ObjectBufferParameters;
+		PassParameters->PS.CulledObjectBufferParameters = CulledObjectBufferParameters;
+		PassParameters->PS.LightTileIntersectionParameters = LightTileIntersectionParameters;
+
+		PassParameters->MeshSDFIndirectArgs = ObjectIndirectArguments;
 
-		TArray<FRHIUnorderedAccessView*> UAVs;
-		TileIntersectionParameters.GetUAVs(*TileIntersectionResources, UAVs);
+		FShadowObjectCullVS::FPermutationDomain VSPermutationVector;
+		VSPermutationVector.Set< FShadowObjectCullVS::FPrimitiveType >(PrimitiveType);
+		auto VertexShader = View.ShaderMap->GetShader< FShadowObjectCullVS >(VSPermutationVector);
 
-		TArray<FRHITransitionInfo> TransitionInfos;
-		for (FRHIUnorderedAccessView* UAV : UAVs)
+		FShadowObjectCullPS::FPermutationDomain PermutationVector;
+		PermutationVector.Set< FShadowObjectCullPS::FPrimitiveType >(PrimitiveType);
+		PermutationVector.Set< FShadowObjectCullPS::FCountingPass >(bCountingPass);
+		auto PixelShader = View.ShaderMap->GetShader< FShadowObjectCullPS >(PermutationVector);
+
+		const bool bReverseCulling = View.bReverseCulling;
+
+		ClearUnusedGraphResources(VertexShader, &PassParameters->VS);
+		ClearUnusedGraphResources(PixelShader, &PassParameters->PS);
+
+		GraphBuilder.AddPass(
+			RDG_EVENT_NAME("ScatterMeshSDFsToLightGrid %ux%u", LightTileDimensions.X, LightTileDimensions.Y),
+			PassParameters,
+			ERDGPassFlags::Raster,
+			[LightTileDimensions, bReverseCulling, VertexShader, PixelShader, PassParameters](FRHICommandListImmediate& RHICmdList)
 		{
-			TransitionInfos.Add(FRHITransitionInfo(UAV, ERHIAccess::Unknown, ERHIAccess::SRVMask));
-		}
-		RHICmdList.Transition(MakeArrayView(TransitionInfos.GetData(), TransitionInfos.Num()));
-	}
+			FGraphicsPipelineStateInitializer GraphicsPSOInit;
+			RHICmdList.ApplyCachedRenderTargets(GraphicsPSOInit);
 
-private:
-	LAYOUT_FIELD(FLightTileIntersectionParameters, TileIntersectionParameters);
-};
+			RHICmdList.SetViewport(0, 0, 0.0f, LightTileDimensions.X, LightTileDimensions.Y, 1.0f);
 
-IMPLEMENT_SHADER_TYPE(,FComputeCulledObjectStartOffsetCS,TEXT("/Engine/Private/DistanceFieldShadowing.usf"),TEXT("ComputeCulledTilesStartOffsetCS"),SF_Compute);
+			// Render backfaces since camera may intersect
+			GraphicsPSOInit.RasterizerState = bReverseCulling ? TStaticRasterizerState<FM_Solid, CM_CW>::GetRHI() : TStaticRasterizerState<FM_Solid, CM_CCW>::GetRHI();
+			GraphicsPSOInit.DepthStencilState = TStaticDepthStencilState<false, CF_Always>::GetRHI();
+			GraphicsPSOInit.BlendState = TStaticBlendState<>::GetRHI();
+			GraphicsPSOInit.PrimitiveType = PT_TriangleList;
 
-template<bool bCountingPass, EDistanceFieldPrimitiveType PrimitiveType>
-void ScatterObjectsToShadowTiles(
-	FRHICommandListImmediate& RHICmdList, 
-	const FViewInfo& View, 
-	const FMatrix& WorldToShadowValue, 
-	float ShadowBoundingRadius,
-	FIntPoint LightTileDimensions, 
-	FLightTileIntersectionResources* TileIntersectionResources)
-{
-	TShaderMapRef<TShadowObjectCullVS<PrimitiveType>> VertexShader(View.ShaderMap);
-	TShaderMapRef<TShadowObjectCullPS<bCountingPass, PrimitiveType>> PixelShader(View.ShaderMap);
+			GraphicsPSOInit.BoundShaderState.VertexDeclarationRHI = GetVertexDeclarationFVector4();
+			GraphicsPSOInit.BoundShaderState.VertexShaderRHI = VertexShader.GetVertexShader();
+			GraphicsPSOInit.BoundShaderState.PixelShaderRHI = PixelShader.GetPixelShader();
 
-	TArray<FRHIUnorderedAccessView*> UAVs;
-	PixelShader->GetUAVs(View, TileIntersectionResources, UAVs);
+			SetGraphicsPipelineState(RHICmdList, GraphicsPSOInit);
 
-	TArray<FRHITransitionInfo> UAVTransitionInfos;
-	for (FRHIUnorderedAccessView* UAV : UAVs)
-	{
-		UAVTransitionInfos.Add(FRHITransitionInfo(UAV, ERHIAccess::Unknown, ERHIAccess::UAVMask));
-	}
-	RHICmdList.Transition(MakeArrayView(UAVTransitionInfos.GetData(), UAVTransitionInfos.Num()));
+			SetShaderParameters(RHICmdList, VertexShader, VertexShader.GetVertexShader(), PassParameters->VS);
+			SetShaderParameters(RHICmdList, PixelShader, PixelShader.GetPixelShader(), PassParameters->PS);
 
-	FRHIRenderPassInfo RPInfo(FRHIRenderPassInfo::NoRenderTargets);
-	if (GRHIRequiresRenderTargetForPixelShaderUAVs)
-	{
-		TRefCountPtr<IPooledRenderTarget> Dummy;
-		FPooledRenderTargetDesc Desc(FPooledRenderTargetDesc::Create2DDesc(LightTileDimensions, PF_B8G8R8A8, FClearValueBinding::None, TexCreate_None, TexCreate_RenderTargetable, false));
-		if (!GRenderTargetPool.FindFreeElement(RHICmdList, Desc, Dummy, TEXT("Dummy")))
-		{
-			RHICmdList.Transition(FRHITransitionInfo(Dummy->GetRenderTargetItem().TargetableTexture, ERHIAccess::Unknown, ERHIAccess::RTV));
-		}
+			RHICmdList.SetStreamSource(0, GetUnitCubeVertexBuffer(), 0);
 
-		RPInfo.ColorRenderTargets[0].Action = ERenderTargetActions::DontLoad_DontStore;
-		RPInfo.ColorRenderTargets[0].ArraySlice = 0;
-		RPInfo.ColorRenderTargets[0].MipIndex = 0;
-		RPInfo.ColorRenderTargets[0].RenderTarget = Dummy->GetRenderTargetItem().TargetableTexture;
+			RHICmdList.DrawIndexedPrimitiveIndirect(
+				GetUnitCubeIndexBuffer(),
+				PassParameters->MeshSDFIndirectArgs->GetIndirectRHICallBuffer(),
+				0);
+		});
 	}
-	
-	RHICmdList.BeginRenderPass(RPInfo, TEXT("ScatterObjectsToShadowTiles"));
-	{
-		FGraphicsPipelineStateInitializer GraphicsPSOInit;
-		RHICmdList.ApplyCachedRenderTargets(GraphicsPSOInit);
+}
 
-		RHICmdList.SetViewport(0, 0, 0.0f, LightTileDimensions.X, LightTileDimensions.Y, 1.0f);
+void AllocateDistanceFieldCulledObjectBuffers(
+	FRDGBuilder& GraphBuilder, 
+	bool bWantBoxBounds, 
+	uint32 MaxObjects, 
+	EDistanceFieldPrimitiveType PrimitiveType,
+	FRDGBufferRef& OutObjectIndirectArguments,
+	FDistanceFieldCulledObjectBufferParameters& OutParameters)
+{
+	check(MaxObjects > 0);
+	OutObjectIndirectArguments = GraphBuilder.CreateBuffer(FRDGBufferDesc::CreateIndirectDesc<FRHIDrawIndexedIndirectParameters>(), TEXT("FDistanceFieldCulledObjectBuffers_ObjectIndirectArguments"));
 
-		// Render backfaces since camera may intersect
-		GraphicsPSOInit.RasterizerState = View.bReverseCulling ? TStaticRasterizerState<FM_Solid, CM_CW>::GetRHI() : TStaticRasterizerState<FM_Solid, CM_CCW>::GetRHI();
-		GraphicsPSOInit.DepthStencilState = TStaticDepthStencilState<false, CF_Always>::GetRHI();
-		GraphicsPSOInit.BlendState = TStaticBlendState<>::GetRHI();
-		GraphicsPSOInit.BoundShaderState.VertexDeclarationRHI = GetVertexDeclarationFVector4();
-		GraphicsPSOInit.BoundShaderState.VertexShaderRHI = VertexShader.GetVertexShader();
-		GraphicsPSOInit.BoundShaderState.PixelShaderRHI = PixelShader.GetPixelShader();
-		GraphicsPSOInit.PrimitiveType = PT_TriangleList;
-		SetGraphicsPipelineState(RHICmdList, GraphicsPSOInit);
+	uint32 NumBoundsElementsScale;
+	uint32 ObjectDataStride;
+	uint32 ObjectBoxBoundsStride;
 
-		VertexShader->SetParameters(RHICmdList, View, FVector2D(LightTileDimensions.X, LightTileDimensions.Y), WorldToShadowValue, ShadowBoundingRadius);
-		PixelShader->SetParameters(RHICmdList, View, TileIntersectionResources);
+	if (PrimitiveType == DFPT_SignedDistanceField)
+	{
+		NumBoundsElementsScale = 1;
+		ObjectDataStride = TDistanceFieldCulledObjectBuffers<DFPT_SignedDistanceField>::ObjectDataStride;
+		ObjectBoxBoundsStride = TDistanceFieldCulledObjectBuffers<DFPT_SignedDistanceField>::ObjectBoxBoundsStride;
+	}
+	else
+	{
+		NumBoundsElementsScale = 2;
+		ObjectDataStride = TDistanceFieldCulledObjectBuffers<DFPT_HeightField>::ObjectDataStride;
+		ObjectBoxBoundsStride = TDistanceFieldCulledObjectBuffers<DFPT_HeightField>::ObjectBoxBoundsStride;
+	}
 
-		RHICmdList.SetStreamSource(0, GetUnitCubeVertexBuffer(), 0);
+	FRDGBufferRef Bounds = GraphBuilder.CreateBuffer(FRDGBufferDesc::CreateStructuredDesc(sizeof(FVector4), MaxObjects * NumBoundsElementsScale), TEXT("FDistanceFieldCulledObjectBuffers_Bounds"));
+	FRDGBufferRef Data = GraphBuilder.CreateBuffer(FRDGBufferDesc::CreateStructuredDesc(sizeof(FVector4), MaxObjects * ObjectDataStride), TEXT("FDistanceFieldCulledObjectBuffers_Data"));
 
-		FRHIVertexBuffer* IndirectArgsBuffer = PrimitiveType == DFPT_HeightField ?
-			GShadowCulledHeightFieldObjectBuffers.Buffers.ObjectIndirectArguments.Buffer :
-			GShadowCulledObjectBuffers.Buffers.ObjectIndirectArguments.Buffer;
+	OutParameters.RWObjectIndirectArguments = GraphBuilder.CreateUAV(OutObjectIndirectArguments, PF_R32_UINT);
+	OutParameters.RWCulledObjectBounds = GraphBuilder.CreateUAV(Bounds);
+	OutParameters.RWCulledObjectData = GraphBuilder.CreateUAV(Data);
 
-		RHICmdList.DrawIndexedPrimitiveIndirect(GetUnitCubeIndexBuffer(), IndirectArgsBuffer, 0);
-	}
-	RHICmdList.EndRenderPass();
+	OutParameters.ObjectIndirectArguments = GraphBuilder.CreateSRV(OutObjectIndirectArguments, PF_R32_UINT);
+	OutParameters.CulledObjectBounds = GraphBuilder.CreateSRV(Bounds);
+	OutParameters.CulledObjectData = GraphBuilder.CreateSRV(Data);
 
-	TArray<FRHITransitionInfo> SRVTransitionInfos;
-	for (FRHIUnorderedAccessView* UAV : UAVs)
+	if (bWantBoxBounds)
 	{
-		SRVTransitionInfos.Add(FRHITransitionInfo(UAV, ERHIAccess::Unknown, ERHIAccess::SRVMask));
+		FRDGBufferRef BoxBounds = GraphBuilder.CreateBuffer(FRDGBufferDesc::CreateStructuredDesc(sizeof(FVector4), MaxObjects * ObjectBoxBoundsStride), TEXT("FDistanceFieldCulledObjectBuffers_BoxBounds"));
+		OutParameters.RWCulledObjectBoxBounds = GraphBuilder.CreateUAV(BoxBounds);
+		OutParameters.CulledObjectBoxBounds = GraphBuilder.CreateSRV(BoxBounds);
 	}
-	RHICmdList.Transition(MakeArrayView(SRVTransitionInfos.GetData(), SRVTransitionInfos.Num()));
 }
 
-template <EDistanceFieldPrimitiveType PrimitiveType>
-void CullDistanceFieldObjectsForLight_Internal(
+void CullDistanceFieldObjectsForLight(
 	FRDGBuilder& GraphBuilder,
 	const FViewInfo& View,
 	const FLightSceneProxy* LightSceneProxy, 
+	EDistanceFieldPrimitiveType PrimitiveType,
 	const FMatrix& WorldToShadowValue, 
 	int32 NumPlanes, 
 	const FPlane* PlaneData, 
 	const FVector4& ShadowBoundingSphereValue,
 	float ShadowBoundingRadius,
-	TUniquePtr<FLightTileIntersectionResources>& TileIntersectionResources)
+	const FDistanceFieldObjectBufferParameters& ObjectBufferParameters,
+	FDistanceFieldCulledObjectBufferParameters& CulledObjectBufferParameters,
+	FLightTileIntersectionParameters& LightTileIntersectionParameters)
 {
-	constexpr bool bIsHeightfield = PrimitiveType == DFPT_HeightField;
+	const bool bIsHeightfield = PrimitiveType == DFPT_HeightField;
 	const FScene* Scene = (const FScene*)(View.Family->Scene);
+	FRDGBufferRef ObjectIndirectArguments = nullptr;
+
+	RDG_EVENT_SCOPE(GraphBuilder, "CullMeshSDFsForLight");
 
-	RDG_EVENT_SCOPE(GraphBuilder, "CullObjectsForLight");
+	const FDistanceFieldSceneData& SceneData = Scene->DistanceFieldSceneData;
+	const int32 NumObjectsInBuffer = bIsHeightfield ? SceneData.NumHeightFieldObjectsInBuffer : SceneData.NumObjectsInBuffer;
+
+	AllocateDistanceFieldCulledObjectBuffers(
+		GraphBuilder, 
+		true, 
+		FMath::DivideAndRoundUp(NumObjectsInBuffer, 256) * 256, 
+		PrimitiveType,
+		ObjectIndirectArguments,
+		CulledObjectBufferParameters);
+
+	AddClearUAVPass(GraphBuilder, CulledObjectBufferParameters.RWObjectIndirectArguments, 0);
 
 	{
-		const FDistanceFieldSceneData& SceneData = Scene->DistanceFieldSceneData;
-		const int32 NumObjectsInBuffer = bIsHeightfield ? SceneData.NumHeightFieldObjectsInBuffer : SceneData.NumObjectsInBuffer;
-		auto& CulledObjectBuffers = TSelector<bIsHeightfield>()(GShadowCulledHeightFieldObjectBuffers, GShadowCulledObjectBuffers);
-
-		if (!CulledObjectBuffers.IsInitialized()
-			|| CulledObjectBuffers.Buffers.MaxObjects < NumObjectsInBuffer
-			|| CulledObjectBuffers.Buffers.MaxObjects > 3 * NumObjectsInBuffer
-			|| GFastVRamConfig.bDirty)
+		FCullObjectsForShadowCS::FParameters* PassParameters = GraphBuilder.AllocParameters<FCullObjectsForShadowCS::FParameters>();
+			
+		PassParameters->View = View.ViewUniformBuffer;
+		PassParameters->ObjectBufferParameters = ObjectBufferParameters;
+		PassParameters->CulledObjectBufferParameters = CulledObjectBufferParameters;
+		PassParameters->ObjectBoundingGeometryIndexCount = UE_ARRAY_COUNT(GCubeIndices);
+		PassParameters->WorldToShadow = WorldToShadowValue;
+		PassParameters->NumShadowHullPlanes = NumPlanes;
+		PassParameters->ShadowBoundingSphere = ShadowBoundingSphereValue;
+
+		check(NumPlanes <= 12);
+
+		for (int32 i = 0; i < NumPlanes; i++)
 		{
-			CulledObjectBuffers.Buffers.bWantBoxBounds = true;
-			CulledObjectBuffers.Buffers.MaxObjects = NumObjectsInBuffer * 5 / 4;
-			CulledObjectBuffers.ReleaseResource();
-			CulledObjectBuffers.InitResource();
+			PassParameters->ShadowConvexHull[i] = FVector4(PlaneData[i], PlaneData[i].W);
 		}
 
-		AddPass(
-			GraphBuilder,
-			RDG_EVENT_NAME("CullObjectsToFrustum %sObjects %d", bIsHeightfield ? TEXT("HeightField") : TEXT("DistanceField")),
-			[Scene, &View, &CulledObjectBuffers, WorldToShadowValue, NumPlanes, PlaneData, ShadowBoundingSphereValue, NumObjectsInBuffer](FRHICommandList& RHICmdList)
-		{
-			CulledObjectBuffers.Buffers.AcquireTransientResource();
-
-			RHICmdList.Transition(FRHITransitionInfo(CulledObjectBuffers.Buffers.ObjectIndirectArguments.UAV, ERHIAccess::Unknown, ERHIAccess::UAVCompute));
-			RHICmdList.ClearUAVUint(CulledObjectBuffers.Buffers.ObjectIndirectArguments.UAV, FUintVector4(0, 0, 0, 0));
-
-			TShaderMapRef<TCullObjectsForShadowCS<PrimitiveType>> ComputeShader(GetGlobalShaderMap(Scene->GetFeatureLevel()));
-			RHICmdList.SetComputeShader(ComputeShader.GetComputeShader());
-			ComputeShader->SetParameters(RHICmdList, Scene, View, WorldToShadowValue, NumPlanes, PlaneData, ShadowBoundingSphereValue);
+		FCullObjectsForShadowCS::FPermutationDomain PermutationVector;
+		PermutationVector.Set< FCullObjectsForShadowCS::FPrimitiveType >(PrimitiveType);
+		auto ComputeShader = View.ShaderMap->GetShader<FCullObjectsForShadowCS>(PermutationVector);
+		const int32 GroupSize = FMath::DivideAndRoundUp<uint32>(NumObjectsInBuffer, UpdateObjectsGroupSize);
 
-			DispatchComputeShader(RHICmdList, ComputeShader.GetShader(), FMath::DivideAndRoundUp<uint32>(NumObjectsInBuffer, UpdateObjectsGroupSize), 1, 1);
-			ComputeShader->UnsetParameters(RHICmdList, Scene);
-		});
+		FComputeShaderUtils::AddPass(
+			GraphBuilder,
+			RDG_EVENT_NAME("CullMeshSDFObjectsToFrustum"),
+			ComputeShader,
+			PassParameters,
+			FIntVector(GroupSize, 1, 1));
 	}
 
 	// Allocate tile resolution based on world space size
 	const float LightTiles = FMath::Min(ShadowBoundingRadius / GShadowCullTileWorldSize + 1.0f, 256.0f);
-	FIntPoint LightTileDimensions(LightTiles, LightTiles);
+	FIntPoint LightTileDimensions(Align(FMath::TruncToInt(LightTiles), 64), Align(FMath::TruncToInt(LightTiles), 64));
 
 	if (LightSceneProxy->GetLightType() == LightType_Directional && GShadowScatterTileCulling)
 	{
 		const bool b16BitObjectIndices = Scene->DistanceFieldSceneData.CanUse16BitObjectIndices();
-		bool bLightDimensionsDirty = false;
 
-		if (TileIntersectionResources)
-		{
-			FIntPoint AlignedLightDimensions = FLightTileIntersectionResources::GetAlignedDimensions(LightTileDimensions);
-			bLightDimensionsDirty = AlignedLightDimensions.X > TileIntersectionResources->GetTileAlignedDimensions().X;
-		}
-		
-		if (!TileIntersectionResources || bLightDimensionsDirty || TileIntersectionResources->b16BitIndices != b16BitObjectIndices)
-		{
-			if (TileIntersectionResources)
-			{
-				TileIntersectionResources->Release();
-			}
-			else
-			{
-				TileIntersectionResources = MakeUnique<FLightTileIntersectionResources>();
-			}
+		FRDGBufferRef ShadowTileNumCulledObjects = GraphBuilder.CreateBuffer(FRDGBufferDesc::CreateBufferDesc(sizeof(uint32), LightTileDimensions.X * LightTileDimensions.Y), TEXT("ShadowTileNumCulledObjects"));
+		LightTileIntersectionParameters.RWShadowTileNumCulledObjects = GraphBuilder.CreateUAV(ShadowTileNumCulledObjects, PF_R32_UINT);
+		LightTileIntersectionParameters.ShadowTileNumCulledObjects = GraphBuilder.CreateSRV(ShadowTileNumCulledObjects, PF_R32_UINT);
 
-			TileIntersectionResources->b16BitIndices = b16BitObjectIndices;
-			TileIntersectionResources->TileDimensions = LightTileDimensions;
-			TileIntersectionResources->Initialize(bIsHeightfield ? GAverageHeightFieldObjectsPerShadowCullTile : GAverageObjectsPerShadowCullTile);
-		}
+		FRDGBufferRef ShadowTileStartOffsets = GraphBuilder.CreateBuffer(FRDGBufferDesc::CreateBufferDesc(sizeof(uint32), LightTileDimensions.X * LightTileDimensions.Y), TEXT("ShadowTileStartOffsets"));
+		LightTileIntersectionParameters.RWShadowTileStartOffsets = GraphBuilder.CreateUAV(ShadowTileStartOffsets, PF_R32_UINT);
+		LightTileIntersectionParameters.ShadowTileStartOffsets = GraphBuilder.CreateSRV(ShadowTileStartOffsets, PF_R32_UINT);
 
-		check(TileIntersectionResources);
-		TileIntersectionResources->TileDimensions = LightTileDimensions;
+		FRDGBufferRef NextStartOffset = GraphBuilder.CreateBuffer(FRDGBufferDesc::CreateBufferDesc(sizeof(uint32), 1), TEXT("ShadowNextStartOffset"));
+		LightTileIntersectionParameters.RWNextStartOffset = GraphBuilder.CreateUAV(NextStartOffset, PF_R32_UINT);
+		LightTileIntersectionParameters.NextStartOffset = GraphBuilder.CreateSRV(NextStartOffset, PF_R32_UINT);
 
-		FLightTileIntersectionResources* TileIntersectionResourcesPtr = TileIntersectionResources.Get();
+		const uint32 MaxNumObjectsPerTile = bIsHeightfield ? GAverageHeightFieldObjectsPerShadowCullTile : GAverageObjectsPerShadowCullTile;
+		FRDGBufferRef ShadowTileArrayData = GraphBuilder.CreateBuffer(FRDGBufferDesc::CreateBufferDesc(b16BitObjectIndices ? sizeof(uint16) : sizeof(uint32), MaxNumObjectsPerTile * LightTileDimensions.X * LightTileDimensions.Y), TEXT("ShadowTileArrayData"));
+		LightTileIntersectionParameters.RWShadowTileArrayData = GraphBuilder.CreateUAV(ShadowTileArrayData, b16BitObjectIndices ? PF_R16_UINT : PF_R32_UINT);
+		LightTileIntersectionParameters.ShadowTileArrayData = GraphBuilder.CreateSRV(ShadowTileArrayData, b16BitObjectIndices ? PF_R16_UINT : PF_R32_UINT);
+		LightTileIntersectionParameters.ShadowTileListGroupSize = LightTileDimensions;
 
-		AddPass(GraphBuilder, [&View, TileIntersectionResourcesPtr, WorldToShadowValue, ShadowBoundingRadius, LightTileDimensions](FRHICommandListImmediate& RHICmdList)
-		{
-			{
-				SCOPED_DRAW_EVENT(RHICmdList, ComputeTileStartOffsets);
-
-				// Start at 0 tiles per object
-				RHICmdList.Transition(FRHITransitionInfo(TileIntersectionResourcesPtr->TileNumCulledObjects.UAV, ERHIAccess::Unknown, ERHIAccess::UAVCompute));
-				RHICmdList.ClearUAVUint(TileIntersectionResourcesPtr->TileNumCulledObjects.UAV, FUintVector4(0, 0, 0, 0));
-
-				// Rasterize object bounding shapes and intersect with shadow tiles to compute how many objects intersect each tile
-				ScatterObjectsToShadowTiles<true, PrimitiveType>(RHICmdList, View, WorldToShadowValue, ShadowBoundingRadius, LightTileDimensions, TileIntersectionResourcesPtr);
+		// Start at 0 tiles per object
+		AddClearUAVPass(GraphBuilder, LightTileIntersectionParameters.RWShadowTileNumCulledObjects, 0);
 
-				RHICmdList.Transition(FRHITransitionInfo(TileIntersectionResourcesPtr->NextStartOffset.UAV, ERHIAccess::Unknown, ERHIAccess::UAVCompute));
-				RHICmdList.ClearUAVUint(TileIntersectionResourcesPtr->NextStartOffset.UAV, FUintVector4(0, 0, 0, 0));
+		// Rasterize object bounding shapes and intersect with shadow tiles to compute how many objects intersect each tile
+		ScatterObjectsToShadowTiles(GraphBuilder, View, WorldToShadowValue, ShadowBoundingRadius, true, PrimitiveType, LightTileDimensions, ObjectIndirectArguments, ObjectBufferParameters, CulledObjectBufferParameters, LightTileIntersectionParameters);
 
-				uint32 GroupSizeX = FMath::DivideAndRoundUp<int32>(LightTileDimensions.X, ComputeCulledObjectStartOffsetGroupSize);
-				uint32 GroupSizeY = FMath::DivideAndRoundUp<int32>(LightTileDimensions.Y, ComputeCulledObjectStartOffsetGroupSize);
+		AddClearUAVPass(GraphBuilder, LightTileIntersectionParameters.RWNextStartOffset, 0);
 
-				// Compute the start offset for each tile's culled object data
-				TShaderMapRef<FComputeCulledObjectStartOffsetCS> ComputeShader(View.ShaderMap);
-				RHICmdList.SetComputeShader(ComputeShader.GetComputeShader());
-				ComputeShader->SetParameters(RHICmdList, View, TileIntersectionResourcesPtr);
-				DispatchComputeShader(RHICmdList, ComputeShader.GetShader(), GroupSizeX, GroupSizeY, 1);
-				ComputeShader->UnsetParameters(RHICmdList, View, TileIntersectionResourcesPtr);
-			}
-
-			{
-				SCOPED_DRAW_EVENTF(RHICmdList, CullObjectsToTiles, TEXT("CullObjectsToTiles %ux%u"), LightTileDimensions.X, LightTileDimensions.Y);
+		// Compute the start offset for each tile's culled object data
+		{
+			FComputeCulledObjectStartOffsetCS::FParameters* PassParameters = GraphBuilder.AllocParameters<FComputeCulledObjectStartOffsetCS::FParameters>();
+			
+			PassParameters->LightTileIntersectionParameters = LightTileIntersectionParameters;
+			auto ComputeShader = View.ShaderMap->GetShader<FComputeCulledObjectStartOffsetCS>();
+			uint32 GroupSizeX = FMath::DivideAndRoundUp<int32>(LightTileDimensions.X, ComputeCulledObjectStartOffsetGroupSize);
+			uint32 GroupSizeY = FMath::DivideAndRoundUp<int32>(LightTileDimensions.Y, ComputeCulledObjectStartOffsetGroupSize);
+
+			FComputeShaderUtils::AddPass(
+				GraphBuilder,
+				RDG_EVENT_NAME("ComputeCulledObjectStartOffset"),
+				ComputeShader,
+				PassParameters,
+				FIntVector(GroupSizeX, GroupSizeY, 1));
+		}
 
-				// Start at 0 tiles per object
-				RHICmdList.Transition(FRHITransitionInfo(TileIntersectionResourcesPtr->TileNumCulledObjects.UAV, ERHIAccess::Unknown, ERHIAccess::UAVCompute));
-				RHICmdList.ClearUAVUint(TileIntersectionResourcesPtr->TileNumCulledObjects.UAV, FUintVector4(0, 0, 0, 0));
+		// Start at 0 tiles per object
+		AddClearUAVPass(GraphBuilder, LightTileIntersectionParameters.RWShadowTileNumCulledObjects, 0);
 
-				// Rasterize object bounding shapes and intersect with shadow tiles, and write out intersecting tile indices for the cone tracing pass
-				ScatterObjectsToShadowTiles<false, PrimitiveType>(RHICmdList, View, WorldToShadowValue, ShadowBoundingRadius, LightTileDimensions, TileIntersectionResourcesPtr);
-			}
-		});
+		// Rasterize object bounding shapes and intersect with shadow tiles, and write out intersecting tile indices for the cone tracing pass
+		ScatterObjectsToShadowTiles(GraphBuilder, View, WorldToShadowValue, ShadowBoundingRadius, false, PrimitiveType, LightTileDimensions, ObjectIndirectArguments, ObjectBufferParameters, CulledObjectBufferParameters, LightTileIntersectionParameters);
 	}
 }
 
-void CullDistanceFieldObjectsForLight(
-	FRDGBuilder& GraphBuilder,
-	const FViewInfo& View,
-	const FLightSceneProxy* LightSceneProxy,
-	const FMatrix& WorldToShadowValue,
-	int32 NumPlanes,
-	const FPlane* PlaneData,
-	const FVector4& ShadowBoundingSphereValue,
-	float ShadowBoundingRadius,
-	TUniquePtr<FLightTileIntersectionResources>& TileIntersectionResources)
-{
-	CullDistanceFieldObjectsForLight_Internal<DFPT_SignedDistanceField>(
-		GraphBuilder,
-		View,
-		LightSceneProxy,
-		WorldToShadowValue,
-		NumPlanes,
-		PlaneData,
-		ShadowBoundingSphereValue,
-		ShadowBoundingRadius,
-		TileIntersectionResources);
-}
-
-void CullHeightFieldObjectsForLight(
-	FRDGBuilder& GraphBuilder,
-	const FViewInfo& View,
-	const FLightSceneProxy* LightSceneProxy,
-	const FMatrix& WorldToShadowValue,
-	int32 NumPlanes,
-	const FPlane* PlaneData,
-	const FVector4& ShadowBoundingSphereValue,
-	float ShadowBoundingRadius,
-	TUniquePtr<FLightTileIntersectionResources>& TileIntersectionResources)
-{
-	CullDistanceFieldObjectsForLight_Internal<DFPT_HeightField>(
-		GraphBuilder,
-		View,
-		LightSceneProxy,
-		WorldToShadowValue,
-		NumPlanes,
-		PlaneData,
-		ShadowBoundingSphereValue,
-		ShadowBoundingRadius,
-		TileIntersectionResources);
-}
-
 int32 GetDFShadowQuality()
 {
 	return FMath::Clamp(GDFShadowQuality, 0, 3);
@@ -1113,138 +659,127 @@ bool FDeferredShadingSceneRenderer::ShouldPrepareHeightFieldScene() const
 		&& SupportsHeightFieldShadows(Scene->GetFeatureLevel(), Scene->GetShaderPlatform());
 }
 
-BEGIN_SHADER_PARAMETER_STRUCT(FRayTraceShadowsParameters, )
-	SHADER_PARAMETER_RDG_UNIFORM_BUFFER(FSceneTextureUniformParameters, SceneTextures)
-	RDG_TEXTURE_ACCESS(RayTracedShadows, ERHIAccess::UAVCompute)
-	RDG_TEXTURE_ACCESS(PrevOutputTexture, ERHIAccess::SRVCompute)
-END_SHADER_PARAMETER_STRUCT()
-
-template <EDistanceFieldShadowingType DFSType, EDistanceFieldPrimitiveType PrimitiveType, bool bHasPrevOutput>
-void RayTraceShadowsDispatch(
-	FRHIComputeCommandList& RHICmdList,
-	IPooledRenderTarget* RayTracedShadowTexture,
+void RayTraceShadows(
+	FRDGBuilder& GraphBuilder,
+	TRDGUniformBufferRef<FSceneTextureUniformParameters> SceneTexturesUniformBuffer,
+	FRDGTextureRef RayTracedShadowsTexture,
 	const FViewInfo& View,
 	const FProjectedShadowInfo* ProjectedShadowInfo,
-	FLightTileIntersectionResources* TileIntersectionResources,
-	FRHITexture* PrevOutput = nullptr)
+	EDistanceFieldPrimitiveType PrimitiveType,
+	bool bHasPrevOutput,
+	FRDGTextureRef PrevOutputTexture,
+	const FDistanceFieldObjectBufferParameters& ObjectBufferParameters,
+	const FDistanceFieldCulledObjectBufferParameters& CulledObjectBufferParameters,
+	const FLightTileIntersectionParameters& LightTileIntersectionParameters)
 {
-	check(DFSType != DFS_PointLightTiledCulling || PrimitiveType != DFPT_HeightField);
-
 	FIntRect ScissorRect;
 	if (!ProjectedShadowInfo->GetLightSceneInfo().Proxy->GetScissorRect(ScissorRect, View, View.ViewRect))
 	{
 		ScissorRect = View.ViewRect;
 	}
 
-	uint32 GroupSizeX = FMath::DivideAndRoundUp(ScissorRect.Size().X / GetDFShadowDownsampleFactor(), GDistanceFieldShadowTileSizeX);
-	uint32 GroupSizeY = FMath::DivideAndRoundUp(ScissorRect.Size().Y / GetDFShadowDownsampleFactor(), GDistanceFieldShadowTileSizeY);
-
-	auto DispatchTemplatedCS = [&](auto CS)
-	{
-		RHICmdList.SetComputeShader(CS.GetComputeShader());
-		FSceneRenderTargetItem& RayTracedShadowsRTI = RayTracedShadowTexture->GetRenderTargetItem();
-		CS->SetParameters(RHICmdList, View, ProjectedShadowInfo, RayTracedShadowsRTI, FVector2D(GroupSizeX, GroupSizeY), ScissorRect, TileIntersectionResources, PrevOutput);
-		DispatchComputeShader(RHICmdList, CS.GetShader(), GroupSizeX, GroupSizeY, 1);
-		CS->UnsetParameters(RHICmdList);
-	};
+	const int32 DFShadowQuality = (PrimitiveType == DFPT_HeightField ? GetHFShadowQuality() : GetDFShadowQuality()) - 1;
+	check(DFShadowQuality >= 0);
 
-	int32 const DFShadowQuality = PrimitiveType == DFPT_HeightField ? GetHFShadowQuality() : GetDFShadowQuality();
+	EDistanceFieldShadowingType DistanceFieldShadowingType;
 
-	if (DFShadowQuality == 1)
+	if (ProjectedShadowInfo->bDirectionalLight && GShadowScatterTileCulling)
 	{
-		TShaderMapRef<TDistanceFieldShadowingCS<DFSType, 1, PrimitiveType, bHasPrevOutput>> ComputeShader(View.ShaderMap);
-		DispatchTemplatedCS(ComputeShader);
+		DistanceFieldShadowingType = DFS_DirectionalLightScatterTileCulling;
 	}
-	else if (DFShadowQuality == 2)
+	else if (ProjectedShadowInfo->bDirectionalLight)
 	{
-		TShaderMapRef<TDistanceFieldShadowingCS<DFSType, 2, PrimitiveType, bHasPrevOutput>> ComputeShader(View.ShaderMap);
-		DispatchTemplatedCS(ComputeShader);
+		DistanceFieldShadowingType = DFS_DirectionalLightTiledCulling;
 	}
-	else if (DFShadowQuality == 3)
+	else
 	{
-		TShaderMapRef<TDistanceFieldShadowingCS<DFSType, 3, PrimitiveType, bHasPrevOutput>> ComputeShader(View.ShaderMap);
-		DispatchTemplatedCS(ComputeShader);
+		DistanceFieldShadowingType = DFS_PointLightTiledCulling;
 	}
-}
 
-void RayTraceShadows(
-	FRDGBuilder& GraphBuilder,
-	TRDGUniformBufferRef<FSceneTextureUniformParameters> SceneTexturesUniformBuffer,
-	FRDGTextureRef RayTracedShadowsTexture,
-	const FViewInfo& View,
-	const FProjectedShadowInfo* ProjectedShadowInfo,
-	FLightTileIntersectionResources* TileIntersectionResources)
-{
-	FRayTraceShadowsParameters* PassParameters = GraphBuilder.AllocParameters<FRayTraceShadowsParameters>();
-	PassParameters->SceneTextures = SceneTexturesUniformBuffer;
-	PassParameters->RayTracedShadows = RayTracedShadowsTexture;
-
-	GraphBuilder.AddPass(
-		RDG_EVENT_NAME("RayTracedShadows"),
-		PassParameters,
-		ERDGPassFlags::Compute,
-		[&View, ProjectedShadowInfo, TileIntersectionResources, RayTracedShadowsTexture](FRHIComputeCommandList& RHICmdList)
-	{
-		if (ProjectedShadowInfo->bDirectionalLight && GShadowScatterTileCulling)
-		{
-			RayTraceShadowsDispatch<DFS_DirectionalLightScatterTileCulling, DFPT_SignedDistanceField, false>(RHICmdList, RayTracedShadowsTexture->GetPooledRenderTarget(), View, ProjectedShadowInfo, TileIntersectionResources);
-		}
-		else if (ProjectedShadowInfo->bDirectionalLight)
-		{
-			RayTraceShadowsDispatch<DFS_DirectionalLightTiledCulling, DFPT_SignedDistanceField, false>(RHICmdList, RayTracedShadowsTexture->GetPooledRenderTarget(), View, ProjectedShadowInfo, TileIntersectionResources);
-		}
-		else
-		{
-			RayTraceShadowsDispatch<DFS_PointLightTiledCulling, DFPT_SignedDistanceField, false>(RHICmdList, RayTracedShadowsTexture->GetPooledRenderTarget(), View, ProjectedShadowInfo, TileIntersectionResources);
-		}
-	});
-}
+	check(DistanceFieldShadowingType != DFS_PointLightTiledCulling || PrimitiveType != DFPT_HeightField);
+
+	FDistanceFieldAtlasParameters DistanceFieldAtlasParameters;
+	DistanceFieldAtlasParameters.DistanceFieldTexture = GDistanceFieldVolumeTextureAtlas.VolumeTextureRHI;
+	DistanceFieldAtlasParameters.DistanceFieldSampler = TStaticSamplerState<SF_Bilinear,AM_Clamp,AM_Clamp,AM_Clamp>::GetRHI();
+	const int32 NumTexelsOneDimX = GDistanceFieldVolumeTextureAtlas.GetSizeX();
+	const int32 NumTexelsOneDimY = GDistanceFieldVolumeTextureAtlas.GetSizeY();
+	const int32 NumTexelsOneDimZ = GDistanceFieldVolumeTextureAtlas.GetSizeZ();
+	DistanceFieldAtlasParameters.DistanceFieldAtlasTexelSize = FVector(1.0f / NumTexelsOneDimX, 1.0f / NumTexelsOneDimY, 1.0f / NumTexelsOneDimZ);
+
+	FHeightFieldAtlasParameters HeightFieldAtlasParameters;
+	HeightFieldAtlasParameters.HeightFieldTexture = GHeightFieldTextureAtlas.GetAtlasTexture();
+	HeightFieldAtlasParameters.HFVisibilityTexture = GHFVisibilityTextureAtlas.GetAtlasTexture();
+	HeightFieldAtlasParameters.HeightFieldAtlasTexelSize = FVector2D(1.0f / GHeightFieldTextureAtlas.GetSizeX(), 1.0f / GHeightFieldTextureAtlas.GetSizeY());
 
-void RayTraceHeightFieldShadows(
-	FRDGBuilder& GraphBuilder,
-	TRDGUniformBufferRef<FSceneTextureUniformParameters> SceneTexturesUniformBuffer,
-	FRDGTextureRef RayTracedShadowsTexture,
-	const FViewInfo& View,
-	const FProjectedShadowInfo* ProjectedShadowInfo,
-	FLightTileIntersectionResources* TileIntersectionResources,
-	bool bHasPrevOutput,
-	FRDGTextureRef PrevOutputTexture)
-{
-	check(ProjectedShadowInfo->bDirectionalLight);
-
-	FRayTraceShadowsParameters* PassParameters = GraphBuilder.AllocParameters<FRayTraceShadowsParameters>();
-	PassParameters->SceneTextures = SceneTexturesUniformBuffer;
-	PassParameters->RayTracedShadows = RayTracedShadowsTexture;
-	PassParameters->PrevOutputTexture = bHasPrevOutput ? PrevOutputTexture : nullptr;
-
-	GraphBuilder.AddPass(
-		RDG_EVENT_NAME("RayTracedHeightFieldShadows"),
-		PassParameters,
-		ERDGPassFlags::Compute,
-		[&View, ProjectedShadowInfo, TileIntersectionResources, RayTracedShadowsTexture, bHasPrevOutput, PrevOutputTexture](FRHIComputeCommandList& RHICmdList)
 	{
-		if (GShadowScatterTileCulling && bHasPrevOutput)
-		{
-			RayTraceShadowsDispatch<DFS_DirectionalLightScatterTileCulling, DFPT_HeightField, true>(RHICmdList, RayTracedShadowsTexture->GetPooledRenderTarget(), View, ProjectedShadowInfo, TileIntersectionResources, TryGetRHI(PrevOutputTexture));
-		}
-		else if (GShadowScatterTileCulling)
-		{
-			RayTraceShadowsDispatch<DFS_DirectionalLightScatterTileCulling, DFPT_HeightField, false>(RHICmdList, RayTracedShadowsTexture->GetPooledRenderTarget(), View, ProjectedShadowInfo, TileIntersectionResources);
-		}
-		else if (bHasPrevOutput)
+		FDistanceFieldShadowingCS::FParameters* PassParameters = GraphBuilder.AllocParameters<FDistanceFieldShadowingCS::FParameters>();
+			
+		PassParameters->RWShadowFactors = GraphBuilder.CreateUAV(RayTracedShadowsTexture);
+		PassParameters->View = View.ViewUniformBuffer;
+		PassParameters->SceneTextures = SceneTexturesUniformBuffer;
+
+		const FLightSceneProxy& LightProxy = *(ProjectedShadowInfo->GetLightSceneInfo().Proxy);
+		FLightShaderParameters LightParameters;
+		LightProxy.GetLightShaderParameters(LightParameters);
+
+		PassParameters->LightDirection = LightParameters.Direction;
+		PassParameters->LightPositionAndInvRadius = FVector4(LightParameters.Position, LightParameters.InvRadius);
+		// Default light source radius of 0 gives poor results
+		PassParameters->LightSourceRadius = LightParameters.SourceRadius == 0 ? 20 : FMath::Clamp(LightParameters.SourceRadius, .001f, 1.0f / (4 * LightParameters.InvRadius));
+		PassParameters->RayStartOffsetDepthScale = LightProxy.GetRayStartOffsetDepthScale();
+
+		const bool bHeightfield = PrimitiveType == DFPT_HeightField;
+		const float MaxLightAngle = bHeightfield ? 45.0f : 5.0f;
+		const float MinLightAngle = bHeightfield ? FMath::Min(GMinDirectionalLightAngleForRTHF, MaxLightAngle) : 0.001f;
+		const float LightSourceAngle = FMath::Clamp(LightProxy.GetLightSourceAngle(), MinLightAngle, MaxLightAngle) * PI / 180.0f;
+		PassParameters->TanLightAngleAndNormalThreshold = FVector(FMath::Tan(LightSourceAngle), FMath::Cos(PI / 2 + LightSourceAngle), LightProxy.GetTraceDistance());
+		PassParameters->ScissorRectMinAndSize = FIntRect(ScissorRect.Min, ScissorRect.Size());
+		PassParameters->ObjectBufferParameters = ObjectBufferParameters;
+		PassParameters->CulledObjectBufferParameters = CulledObjectBufferParameters;
+		PassParameters->LightTileIntersectionParameters = LightTileIntersectionParameters;
+		PassParameters->DistanceFieldAtlasParameters = DistanceFieldAtlasParameters;
+		PassParameters->HeightFieldAtlasParameters = HeightFieldAtlasParameters;
+		PassParameters->WorldToShadow = FTranslationMatrix(ProjectedShadowInfo->PreShadowTranslation) * ProjectedShadowInfo->SubjectAndReceiverMatrix;
+		PassParameters->TwoSidedMeshDistanceBias = GTwoSidedMeshDistanceBias;
+
+		if (ProjectedShadowInfo->bDirectionalLight)
 		{
-			RayTraceShadowsDispatch<DFS_DirectionalLightTiledCulling, DFPT_HeightField, true>(RHICmdList, RayTracedShadowsTexture->GetPooledRenderTarget(), View, ProjectedShadowInfo, TileIntersectionResources, TryGetRHI(PrevOutputTexture));
+			PassParameters->MinDepth = ProjectedShadowInfo->CascadeSettings.SplitNear - ProjectedShadowInfo->CascadeSettings.SplitNearFadeRegion;
+			PassParameters->MaxDepth = ProjectedShadowInfo->CascadeSettings.SplitFar;
 		}
 		else
 		{
-			RayTraceShadowsDispatch<DFS_DirectionalLightTiledCulling, DFPT_HeightField, false>(RHICmdList, RayTracedShadowsTexture->GetPooledRenderTarget(), View, ProjectedShadowInfo, TileIntersectionResources);
+			check(!bHeightfield);
+			//@todo - set these up for point lights as well
+			PassParameters->MinDepth = 0.0f;
+			PassParameters->MaxDepth = HALF_WORLD_MAX;
 		}
-	});
-}
 
-BEGIN_SHADER_PARAMETER_STRUCT(FBeginRenderRayTracedDistanceFieldProjectionParameters, )
-	SHADER_PARAMETER_RDG_UNIFORM_BUFFER(FSceneTextureUniformParameters, SceneTextures)
-END_SHADER_PARAMETER_STRUCT()
+		PassParameters->DownsampleFactor = GetDFShadowDownsampleFactor();
+		const FIntPoint OutputBufferSize = GetBufferSizeForDFShadows();
+		PassParameters->InvOutputBufferSize = FVector2D(1.f / OutputBufferSize.X, 1.f / OutputBufferSize.Y);
+		PassParameters->ShadowFactorsTexture = PrevOutputTexture;
+		PassParameters->ShadowFactorsSampler = TStaticSamplerState<>::GetRHI();
+		
+		FDistanceFieldShadowingCS::FPermutationDomain PermutationVector;
+		PermutationVector.Set< FDistanceFieldShadowingCS::FCullingType >((uint32)DistanceFieldShadowingType);
+		PermutationVector.Set< FDistanceFieldShadowingCS::FShadowQuality >(DFShadowQuality);
+		PermutationVector.Set< FDistanceFieldShadowingCS::FPrimitiveType >(PrimitiveType);
+		PermutationVector.Set< FDistanceFieldShadowingCS::FHasPreviousOutput >(bHasPrevOutput);
+		auto ComputeShader = View.ShaderMap->GetShader< FDistanceFieldShadowingCS >(PermutationVector);
+
+		uint32 GroupSizeX = FMath::DivideAndRoundUp(ScissorRect.Size().X / GetDFShadowDownsampleFactor(), GDistanceFieldShadowTileSizeX);
+		uint32 GroupSizeY = FMath::DivideAndRoundUp(ScissorRect.Size().Y / GetDFShadowDownsampleFactor(), GDistanceFieldShadowTileSizeY);
+		PassParameters->NumGroups = FVector2D(GroupSizeX, GroupSizeY);
+
+		FComputeShaderUtils::AddPass(
+			GraphBuilder,
+			RDG_EVENT_NAME("DistanceFieldShadowing %ux%u", GroupSizeX * GDistanceFieldShadowTileSizeX, GroupSizeY * GDistanceFieldShadowTileSizeY),
+			ComputeShader,
+			PassParameters,
+			FIntVector(GroupSizeX, GroupSizeY, 1));
+	}
+}
 
 FRDGTextureRef FProjectedShadowInfo::BeginRenderRayTracedDistanceFieldProjection(
 	FRDGBuilder& GraphBuilder,
@@ -1253,7 +788,6 @@ FRDGTextureRef FProjectedShadowInfo::BeginRenderRayTracedDistanceFieldProjection
 {
 	const bool bDFShadowSupported = SupportsDistanceFieldShadows(View.GetFeatureLevel(), View.GetShaderPlatform());
 	const bool bHFShadowSupported = SupportsHeightFieldShadows(View.GetFeatureLevel(), View.GetShaderPlatform());
-	const bool bBufferAliasingEnabled = IsTransientResourceBufferAliasingEnabled();
 	const FScene* Scene = (const FScene*)(View.Family->Scene);
 
 	FRDGTextureRef RayTracedShadowsTexture = nullptr;
@@ -1289,39 +823,37 @@ FRDGTextureRef FProjectedShadowInfo::BeginRenderRayTracedDistanceFieldProjection
 
 			const FMatrix WorldToShadowValue = FTranslationMatrix(PreShadowTranslation) * SubjectAndReceiverMatrix;
 
+			FDistanceFieldObjectBufferParameters ObjectBufferParameters;
+			ObjectBufferParameters.SceneObjectBounds = Scene->DistanceFieldSceneData.GetCurrentObjectBuffers()->Bounds.SRV;
+			ObjectBufferParameters.SceneObjectData = Scene->DistanceFieldSceneData.GetCurrentObjectBuffers()->Data.SRV;
+			ObjectBufferParameters.NumSceneObjects = Scene->DistanceFieldSceneData.NumObjectsInBuffer;
+
+			FLightTileIntersectionParameters LightTileIntersectionParameters;
+			FDistanceFieldCulledObjectBufferParameters CulledObjectBufferParameters;
+
 			CullDistanceFieldObjectsForLight(
 				GraphBuilder,
 				View,
 				LightSceneInfo->Proxy,
+				DFPT_SignedDistanceField,
 				WorldToShadowValue,
 				NumPlanes,
 				PlaneData,
 				ShadowBoundingSphereValue,
 				ShadowBounds.W,
-				LightSceneInfo->TileIntersectionResources
+				ObjectBufferParameters,
+				CulledObjectBufferParameters,
+				LightTileIntersectionParameters
 				);
 
-			// Note: using the same TileIntersectionResources for multiple views, breaks splitscreen / stereo
-			FLightTileIntersectionResources* TileIntersectionResources = LightSceneInfo->TileIntersectionResources.Get();
-
-			View.HeightfieldLightingViewInfo.ComputeRayTracedShadowing(GraphBuilder, View, this, TileIntersectionResources, GShadowCulledObjectBuffers);
-
 			{
 				const FIntPoint BufferSize = GetBufferSizeForDFShadows();
-				FRDGTextureDesc Desc(FRDGTextureDesc::Create2D(BufferSize, PF_G16R16F, FClearValueBinding::None, TexCreate_RenderTargetable | TexCreate_UAV));
+				FRDGTextureDesc Desc(FRDGTextureDesc::Create2D(BufferSize, PF_G16R16F, FClearValueBinding::None, TexCreate_UAV));
 				Desc.Flags |= GFastVRamConfig.DistanceFieldShadows;
 				RayTracedShadowsTexture = GraphBuilder.CreateTexture(Desc, TEXT("RayTracedShadows"));
 			}
 
-			RayTraceShadows(GraphBuilder, SceneTexturesUniformBuffer, RayTracedShadowsTexture, View, this, TileIntersectionResources);
-
-			if (bBufferAliasingEnabled)
-			{
-				AddPass(GraphBuilder, [](FRHICommandList&)
-				{
-					GShadowCulledObjectBuffers.Buffers.DiscardTransientResource();
-				});
-			}
+			RayTraceShadows(GraphBuilder, SceneTexturesUniformBuffer, RayTracedShadowsTexture, View, this, DFPT_SignedDistanceField, false, nullptr, ObjectBufferParameters, CulledObjectBufferParameters, LightTileIntersectionParameters);
 		}
 	}
 
@@ -1341,19 +873,31 @@ FRDGTextureRef FProjectedShadowInfo::BeginRenderRayTracedDistanceFieldProjection
 		const FVector4 ShadowBoundingSphereValue(0.f, 0.f, 0.f, 0.f);
 		const FMatrix WorldToShadowValue = FTranslationMatrix(PreShadowTranslation) * SubjectAndReceiverMatrix;
 
-		CullHeightFieldObjectsForLight(
+		FDistanceFieldObjectBufferParameters ObjectBufferParameters;
+		ObjectBufferParameters.SceneObjectBounds = Scene->DistanceFieldSceneData.GetHeightFieldObjectBuffers()->Bounds.SRV;
+		ObjectBufferParameters.SceneObjectData = Scene->DistanceFieldSceneData.GetHeightFieldObjectBuffers()->Data.SRV;
+		ObjectBufferParameters.NumSceneObjects = Scene->DistanceFieldSceneData.NumHeightFieldObjectsInBuffer;
+
+		FLightTileIntersectionParameters LightTileIntersectionParameters;
+		FDistanceFieldCulledObjectBufferParameters CulledObjectBufferParameters;
+
+		CullDistanceFieldObjectsForLight(
 			GraphBuilder,
 			View,
 			LightSceneInfo->Proxy,
+			DFPT_HeightField,
 			WorldToShadowValue,
 			NumPlanes,
 			PlaneData,
 			ShadowBoundingSphereValue,
 			ShadowBounds.W,
-			LightSceneInfo->HeightFieldTileIntersectionResources);
+			ObjectBufferParameters,
+			CulledObjectBufferParameters,
+			LightTileIntersectionParameters
+			);
 
-		FLightTileIntersectionResources* TileIntersectionResources = LightSceneInfo->HeightFieldTileIntersectionResources.Get();
 		const bool bHasPrevOutput = !!RayTracedShadowsTexture;
+
 		FRDGTextureRef PrevOutputTexture = nullptr;
 
 		if (!RHISupports4ComponentUAVReadWrite(View.GetShaderPlatform()))
@@ -1365,28 +909,19 @@ FRDGTextureRef FProjectedShadowInfo::BeginRenderRayTracedDistanceFieldProjection
 		if (!RayTracedShadowsTexture)
 		{
 			const FIntPoint BufferSize = GetBufferSizeForDFShadows();
-			FRDGTextureDesc Desc(FRDGTextureDesc::Create2D(BufferSize, PF_G16R16F, FClearValueBinding::None, TexCreate_RenderTargetable | TexCreate_UAV));
+			FRDGTextureDesc Desc(FRDGTextureDesc::Create2D(BufferSize, PF_G16R16F, FClearValueBinding::None, TexCreate_UAV));
 			Desc.Flags |= GFastVRamConfig.DistanceFieldShadows;
 			RayTracedShadowsTexture = GraphBuilder.CreateTexture(Desc, TEXT("RayTracedShadows"));
 		}
 
-		RayTraceHeightFieldShadows(GraphBuilder, SceneTexturesUniformBuffer, RayTracedShadowsTexture, View, this, TileIntersectionResources, bHasPrevOutput, PrevOutputTexture);
-
-		if (bBufferAliasingEnabled)
-		{
-			AddPass(GraphBuilder, [](FRHICommandList&)
-			{
-				GShadowCulledHeightFieldObjectBuffers.Buffers.DiscardTransientResource();
-			});
-		}
+		RayTraceShadows(GraphBuilder, SceneTexturesUniformBuffer, RayTracedShadowsTexture, View, this, DFPT_HeightField, bHasPrevOutput, PrevOutputTexture, ObjectBufferParameters, CulledObjectBufferParameters, LightTileIntersectionParameters);
 	}
 
 	return RayTracedShadowsTexture;
 }
 
-BEGIN_SHADER_PARAMETER_STRUCT(FRenderRayTracedDistanceFieldProjectionParameters, )
-	SHADER_PARAMETER_RDG_UNIFORM_BUFFER(FSceneTextureUniformParameters, SceneTextures)
-	RDG_TEXTURE_ACCESS(RayTracedShadows, ERHIAccess::SRVGraphics)
+BEGIN_SHADER_PARAMETER_STRUCT(FDistanceFieldShadowingUpsample, )
+	SHADER_PARAMETER_STRUCT_INCLUDE(FDistanceFieldShadowingUpsamplePS::FParameters, PS)
 	RENDER_TARGET_BINDING_SLOTS()
 END_SHADER_PARAMETER_STRUCT()
 
@@ -1405,20 +940,52 @@ void FProjectedShadowInfo::RenderRayTracedDistanceFieldProjection(
 
 	if (RayTracedShadowsTexture)
 	{
-		auto* PassParameters = GraphBuilder.AllocParameters<FRenderRayTracedDistanceFieldProjectionParameters>();
-		PassParameters->SceneTextures = SceneTexturesUniformBuffer;
-		PassParameters->RayTracedShadows = RayTracedShadowsTexture;
+		FDistanceFieldShadowingUpsample* PassParameters = GraphBuilder.AllocParameters<FDistanceFieldShadowingUpsample>();
 		PassParameters->RenderTargets[0] = FRenderTargetBinding(ScreenShadowMaskTexture, ERenderTargetLoadAction::ELoad);
-		PassParameters->RenderTargets.DepthStencil = FDepthStencilBinding(SceneDepthTexture, ERenderTargetLoadAction::ELoad, ERenderTargetLoadAction::ELoad, FExclusiveDepthStencil::DepthRead_StencilWrite);
+		PassParameters->RenderTargets.DepthStencil = FDepthStencilBinding(SceneDepthTexture, ERenderTargetLoadAction::ELoad, ERenderTargetLoadAction::ELoad, FExclusiveDepthStencil::DepthRead_StencilRead);
+		
+		PassParameters->PS.View = View.ViewUniformBuffer;
+		PassParameters->PS.SceneTextures = SceneTexturesUniformBuffer;
+		PassParameters->PS.ShadowFactorsTexture = RayTracedShadowsTexture;
+		PassParameters->PS.ShadowFactorsSampler = TStaticSamplerState<SF_Bilinear>::GetRHI();
+		PassParameters->PS.ScissorRectMinAndSize = FIntRect(ScissorRect.Min, ScissorRect.Size());
+
+		if (bDirectionalLight && CascadeSettings.FadePlaneLength > 0)
+		{
+			PassParameters->PS.FadePlaneOffset = CascadeSettings.FadePlaneOffset;
+			PassParameters->PS.InvFadePlaneLength = 1.0f / FMath::Max(CascadeSettings.FadePlaneLength, .00001f);
+		}
+		else
+		{
+			PassParameters->PS.FadePlaneOffset = 0.0f;
+			PassParameters->PS.InvFadePlaneLength = 0.0f;
+		}
+
+		if (bDirectionalLight && CascadeSettings.SplitNearFadeRegion > 0)
+		{
+			PassParameters->PS.NearFadePlaneOffset = CascadeSettings.SplitNear - CascadeSettings.SplitNearFadeRegion;
+			PassParameters->PS.InvNearFadePlaneLength = 1.0f / FMath::Max(CascadeSettings.SplitNearFadeRegion, .00001f);
+		}
+		else
+		{
+			PassParameters->PS.NearFadePlaneOffset = -1.0f;
+			PassParameters->PS.InvNearFadePlaneLength = 1.0f;
+		}
+
+		FDistanceFieldShadowingUpsamplePS::FPermutationDomain PermutationVector;
+		PermutationVector.Set< FDistanceFieldShadowingUpsamplePS::FUpsample >(GFullResolutionDFShadowing == 0);
+		auto PixelShader = View.ShaderMap->GetShader< FDistanceFieldShadowingUpsamplePS >(PermutationVector);
+
+		const bool bReverseCulling = View.bReverseCulling;
+
+		ClearUnusedGraphResources(PixelShader, &PassParameters->PS);
 
 		GraphBuilder.AddPass(
 			RDG_EVENT_NAME("Upsample"),
 			PassParameters,
 			ERDGPassFlags::Raster,
-			[this, &View, ScissorRect, bProjectingForForwardShading, RayTracedShadowsTexture](FRHICommandListImmediate& RHICmdList)
+			[this, &View, PixelShader, ScissorRect, bProjectingForForwardShading, PassParameters](FRHICommandListImmediate& RHICmdList)
 		{
-			QUICK_SCOPE_CYCLE_COUNTER(STAT_RenderRayTracedDistanceFieldShadows);
-
 			FGraphicsPipelineStateInitializer GraphicsPSOInit;
 			RHICmdList.ApplyCachedRenderTargets(GraphicsPSOInit);
 
@@ -1435,21 +1002,10 @@ void FProjectedShadowInfo::RenderRayTracedDistanceFieldProjection(
 			GraphicsPSOInit.BoundShaderState.VertexShaderRHI = VertexShader.GetVertexShader();
 			GraphicsPSOInit.PrimitiveType = PT_TriangleList;
 			GraphicsPSOInit.bDepthBounds = bDirectionalLight;
+			GraphicsPSOInit.BoundShaderState.PixelShaderRHI = PixelShader.GetPixelShader();
+			SetGraphicsPipelineState(RHICmdList, GraphicsPSOInit);
 
-			if (GFullResolutionDFShadowing)
-			{
-				TShaderMapRef<TDistanceFieldShadowingUpsamplePS<false> > PixelShader(View.ShaderMap);
-				GraphicsPSOInit.BoundShaderState.PixelShaderRHI = PixelShader.GetPixelShader();
-				SetGraphicsPipelineState(RHICmdList, GraphicsPSOInit);
-				PixelShader->SetParameters(RHICmdList, View, this, ScissorRect, RayTracedShadowsTexture->GetPooledRenderTarget());
-			}
-			else
-			{
-				TShaderMapRef<TDistanceFieldShadowingUpsamplePS<true> > PixelShader(View.ShaderMap);
-				GraphicsPSOInit.BoundShaderState.PixelShaderRHI = PixelShader.GetPixelShader();
-				SetGraphicsPipelineState(RHICmdList, GraphicsPSOInit);
-				PixelShader->SetParameters(RHICmdList, View, this, ScissorRect, RayTracedShadowsTexture->GetPooledRenderTarget());
-			}
+			SetShaderParameters(RHICmdList, PixelShader, PixelShader.GetPixelShader(), PassParameters->PS);
 
 			//@todo - depth bounds test for local lights
 			if (bDirectionalLight)
diff --git a/Engine/Source/Runtime/Renderer/Private/HeightfieldLighting.cpp b/Engine/Source/Runtime/Renderer/Private/HeightfieldLighting.cpp
--- a/Engine/Source/Runtime/Renderer/Private/HeightfieldLighting.cpp
+++ b/Engine/Source/Runtime/Renderer/Private/HeightfieldLighting.cpp
@@ -941,171 +941,6 @@ void FHeightfieldLightingViewInfo::ComputeShadowMapShadowing(FRDGBuilder& GraphB
 	}
 }
 
-
-class FRayTracedShadowHeightfieldsPS : public FGlobalShader
-{
-	DECLARE_SHADER_TYPE(FRayTracedShadowHeightfieldsPS, Global);
-public:
-
-	static bool ShouldCompilePermutation(const FGlobalShaderPermutationParameters& Parameters)
-	{
-		return IsFeatureLevelSupported(Parameters.Platform, ERHIFeatureLevel::SM5) && DoesPlatformSupportDistanceFieldGI(Parameters.Platform);
-	}
-
-	static void ModifyCompilationEnvironment(const FGlobalShaderPermutationParameters& Parameters, FShaderCompilerEnvironment& OutEnvironment)
-	{
-		FLightTileIntersectionParameters::ModifyCompilationEnvironment(Parameters.Platform, OutEnvironment);
-	}
-
-	/** Default constructor. */
-	FRayTracedShadowHeightfieldsPS() {}
-
-	/** Initialization constructor. */
-	FRayTracedShadowHeightfieldsPS(const ShaderMetaType::CompiledShaderInitializerType& Initializer)
-		: FGlobalShader(Initializer)
-	{
-		HeightfieldDescriptionParameters.Bind(Initializer.ParameterMap);
-		GlobalHeightfieldParameters.Bind(Initializer.ParameterMap);
-		ObjectParameters.Bind(Initializer.ParameterMap);
-		LightDirection.Bind(Initializer.ParameterMap, TEXT("LightDirection"));
-		TanLightAngle.Bind(Initializer.ParameterMap, TEXT("TanLightAngle"));
-		CascadeDepthMinMax.Bind(Initializer.ParameterMap, TEXT("CascadeDepthMinMax"));
-		LightTileIntersectionParameters.Bind(Initializer.ParameterMap);
-		WorldToShadow.Bind(Initializer.ParameterMap, TEXT("WorldToShadow"));
-		TwoSidedMeshDistanceBias.Bind(Initializer.ParameterMap, TEXT("TwoSidedMeshDistanceBias"));
-	}
-
-	void SetParameters(
-		FRHICommandList& RHICmdList,
-		const FSceneView& View,
-		const FProjectedShadowInfo* ProjectedShadowInfo,
-		int32 NumHeightfieldsValue,
-		const FHeightfieldLightingAtlas& Atlas,
-		FLightTileIntersectionResources* TileIntersectionResources,
-		FDistanceFieldObjectBufferResource& CulledObjectBuffers)
-	{
-		FRHIPixelShader* ShaderRHI = RHICmdList.GetBoundPixelShader();
-
-		FGlobalShader::SetParameters<FViewUniformShaderParameters>(RHICmdList, ShaderRHI, View.ViewUniformBuffer);
-		HeightfieldDescriptionParameters.Set(RHICmdList, ShaderRHI, GetHeightfieldDescriptionsSRV(), NumHeightfieldsValue);
-		GlobalHeightfieldParameters.Set(RHICmdList, ShaderRHI, Atlas);
-
-		FRHITexture* TextureAtlas;
-		int32 AtlasSizeX;
-		int32 AtlasSizeY;
-		int32 AtlasSizeZ;
-
-		TextureAtlas = GDistanceFieldVolumeTextureAtlas.VolumeTextureRHI;
-		AtlasSizeX = GDistanceFieldVolumeTextureAtlas.GetSizeX();
-		AtlasSizeY = GDistanceFieldVolumeTextureAtlas.GetSizeY();
-		AtlasSizeZ = GDistanceFieldVolumeTextureAtlas.GetSizeZ();
-
-		ObjectParameters.Set(RHICmdList, ShaderRHI, CulledObjectBuffers.Buffers, TextureAtlas, FIntVector(AtlasSizeX, AtlasSizeY, AtlasSizeZ));
-
-		SetShaderValue(RHICmdList, ShaderRHI, LightDirection, ProjectedShadowInfo->GetLightSceneInfo().Proxy->GetDirection());
-
-		const float LightSourceAngle = FMath::Clamp(ProjectedShadowInfo->GetLightSceneInfo().Proxy->GetLightSourceAngle(), 0.001f, 5.0f) * PI / 180.0f;
-		SetShaderValue(RHICmdList, ShaderRHI, TanLightAngle, FMath::Tan(LightSourceAngle));
-
-		SetShaderValue(RHICmdList, ShaderRHI, CascadeDepthMinMax, FVector2D(ProjectedShadowInfo->CascadeSettings.SplitNear, ProjectedShadowInfo->CascadeSettings.SplitFar));
-
-		check(TileIntersectionResources || !LightTileIntersectionParameters.IsBound());
-
-		if (TileIntersectionResources)
-		{
-			LightTileIntersectionParameters.Set(RHICmdList, ShaderRHI, *TileIntersectionResources);
-		}
-
-		FMatrix WorldToShadowMatrixValue = FTranslationMatrix(ProjectedShadowInfo->PreShadowTranslation) * ProjectedShadowInfo->SubjectAndReceiverMatrix;
-		SetShaderValue(RHICmdList, ShaderRHI, WorldToShadow, WorldToShadowMatrixValue);
-
-		extern float GTwoSidedMeshDistanceBias;
-		SetShaderValue(RHICmdList, ShaderRHI, TwoSidedMeshDistanceBias, GTwoSidedMeshDistanceBias);
-	}
-
-private:
-
-	LAYOUT_FIELD(FHeightfieldDescriptionParameters, HeightfieldDescriptionParameters);
-	LAYOUT_FIELD(FGlobalHeightfieldParameters, GlobalHeightfieldParameters);
-	LAYOUT_FIELD((TDistanceFieldCulledObjectBufferParameters<DFPT_SignedDistanceField>), ObjectParameters);
-	LAYOUT_FIELD(FShaderParameter, LightDirection);
-	LAYOUT_FIELD(FShaderParameter, TanLightAngle);
-	LAYOUT_FIELD(FShaderParameter, CascadeDepthMinMax);
-	LAYOUT_FIELD(FLightTileIntersectionParameters, LightTileIntersectionParameters);
-	LAYOUT_FIELD(FShaderParameter, WorldToShadow);
-	LAYOUT_FIELD(FShaderParameter, TwoSidedMeshDistanceBias);
-};
-
-IMPLEMENT_SHADER_TYPE(, FRayTracedShadowHeightfieldsPS, TEXT("/Engine/Private/HeightfieldLighting.usf"), TEXT("RayTracedShadowHeightfieldsPS"), SF_Pixel);
-
-
-void FHeightfieldLightingViewInfo::ComputeRayTracedShadowing(
-	FRDGBuilder& GraphBuilder,
-	const FViewInfo& View,
-	const FProjectedShadowInfo* ProjectedShadowInfo,
-	FLightTileIntersectionResources* TileIntersectionResources,
-	FDistanceFieldObjectBufferResource& CulledObjectBuffers) const
-{
-	if (AllowHeightfieldGI(View)
-		&& SupportsHeightfieldLighting(View.GetFeatureLevel(), View.GetShaderPlatform())
-		&& Heightfield.ComponentDescriptions.Num() > 0
-		&& ProjectedShadowInfo->IsWholeSceneDirectionalShadow()
-		&& ProjectedShadowInfo->DependentView == &View
-		&& ProjectedShadowInfo->bRayTracedDistanceField)
-	{
-		// TODO(RDG): Port this to RDG.
-		AddPass(GraphBuilder, RDG_EVENT_NAME("HeightfieldRayTracedShadowingForGI"), [this, &View, ProjectedShadowInfo, TileIntersectionResources, &CulledObjectBuffers](FRHICommandListImmediate& RHICmdList)
-		{
-			FSceneViewState* ViewState = (FSceneViewState*)View.State;
-
-			const FHeightfieldLightingAtlas& Atlas = *ViewState->HeightfieldLightingAtlas;
-
-			const FIntPoint LightingAtlasSize = Atlas.GetAtlasSize();
-			const FVector2D InvLightingAtlasSize(1.0f / LightingAtlasSize.X, 1.0f / LightingAtlasSize.Y);
-
-			FRHIRenderPassInfo RPInfo(Atlas.DirectionalLightShadowing->GetRenderTargetItem().TargetableTexture, ERenderTargetActions::Load_Store);
-			RHICmdList.BeginRenderPass(RPInfo, TEXT("ComputeRayTracedShadowing"));
-			{
-				FGraphicsPipelineStateInitializer GraphicsPSOInit;
-				RHICmdList.ApplyCachedRenderTargets(GraphicsPSOInit);
-
-				RHICmdList.SetViewport(0, 0, 0.0f, LightingAtlasSize.X, LightingAtlasSize.Y, 1.0f);
-				GraphicsPSOInit.RasterizerState = TStaticRasterizerState<FM_Solid, CM_None>::GetRHI();
-				GraphicsPSOInit.DepthStencilState = TStaticDepthStencilState<false, CF_Always>::GetRHI();
-				// Combine with other shadow types with min (CSM)
-				GraphicsPSOInit.BlendState = TStaticBlendState<CW_RED, BO_Min, BF_One, BF_One>::GetRHI();
-				RHICmdList.SetStreamSource(0, GQuadVertexBuffer.VertexBufferRHI, 0);
-
-				TShaderMapRef<FHeightfieldComponentQuadVS> VertexShader(View.ShaderMap);
-				TShaderMapRef<FRayTracedShadowHeightfieldsPS> PixelShader(View.ShaderMap);
-
-				GraphicsPSOInit.BoundShaderState.VertexDeclarationRHI = GScreenVertexDeclaration.VertexDeclarationRHI;
-				GraphicsPSOInit.BoundShaderState.VertexShaderRHI = VertexShader.GetVertexShader();
-				GraphicsPSOInit.BoundShaderState.PixelShaderRHI = PixelShader.GetPixelShader();
-				GraphicsPSOInit.PrimitiveType = PT_TriangleList;
-
-				SetGraphicsPipelineState(RHICmdList, GraphicsPSOInit);
-
-				for (TMap<FHeightfieldComponentTextures, TArray<FHeightfieldComponentDescription>>::TConstIterator It(Heightfield.ComponentDescriptions); It; ++It)
-				{
-					const TArray<FHeightfieldComponentDescription>& HeightfieldDescriptions = It.Value();
-
-					if (HeightfieldDescriptions.Num() > 0)
-					{
-						UploadHeightfieldDescriptions(HeightfieldDescriptions, InvLightingAtlasSize, 1.0f / Heightfield.DownsampleFactor);
-
-						VertexShader->SetParameters(RHICmdList, View, HeightfieldDescriptions.Num());
-						PixelShader->SetParameters(RHICmdList, View, ProjectedShadowInfo, HeightfieldDescriptions.Num(), Atlas, TileIntersectionResources, CulledObjectBuffers);
-
-						RHICmdList.DrawPrimitive(0, 2, HeightfieldDescriptions.Num());
-					}
-				}
-			}
-			RHICmdList.EndRenderPass();
-		});
-	}
-}
-
 class FLightHeightfieldsPS : public FMaterialShader
 {
 	DECLARE_SHADER_TYPE(FLightHeightfieldsPS, Material);
diff --git a/Engine/Source/Runtime/Renderer/Private/HeightfieldLighting.h b/Engine/Source/Runtime/Renderer/Private/HeightfieldLighting.h
--- a/Engine/Source/Runtime/Renderer/Private/HeightfieldLighting.h
+++ b/Engine/Source/Runtime/Renderer/Private/HeightfieldLighting.h
@@ -18,7 +18,6 @@
 class FAOScreenGridResources;
 class FDistanceFieldAOParameters;
 class FLightSceneInfo;
-class FLightTileIntersectionResources;
 class FProjectedShadowInfo;
 class FScene;
 class FViewInfo;
@@ -107,13 +106,6 @@ public:
 
 	void ComputeShadowMapShadowing(FRDGBuilder& GraphBuilder, const FViewInfo& View, const FProjectedShadowInfo* ProjectedShadowInfo) const;
 
-	void ComputeRayTracedShadowing(
-		FRDGBuilder& GraphBuilder,
-		const FViewInfo& View, 
-		const FProjectedShadowInfo* ProjectedShadowInfo, 
-		FLightTileIntersectionResources* TileIntersectionResources,
-		class FDistanceFieldObjectBufferResource& CulledObjectBuffers) const;
-
 	void ComputeLighting(FRDGBuilder& GraphBuilder, const FViewInfo& View, const FLightSceneInfo& LightSceneInfo) const;
 
 	void ComputeOcclusionForScreenGrid(
diff --git a/Engine/Source/Runtime/Renderer/Private/LightSceneInfo.cpp b/Engine/Source/Runtime/Renderer/Private/LightSceneInfo.cpp
--- a/Engine/Source/Runtime/Renderer/Private/LightSceneInfo.cpp
+++ b/Engine/Source/Runtime/Renderer/Private/LightSceneInfo.cpp
@@ -48,8 +48,6 @@ FLightSceneInfo::FLightSceneInfo(FLightSceneProxy* InProxy, bool InbVisible)
 	, DynamicInteractionStaticPrimitiveList(NULL)
 	, Proxy(InProxy)
 	, Id(INDEX_NONE)
-	, TileIntersectionResources(nullptr)
-	, HeightFieldTileIntersectionResources(nullptr)
 	, DynamicShadowMapChannel(-1)
 	, bPrecomputedLightingIsValid(InProxy->GetLightComponent()->IsPrecomputedLightingValid())
 	, bVisible(InbVisible)
@@ -263,11 +261,6 @@ FLightPrimitiveInteraction* FLightSceneInfo::GetDynamicInteractionStaticPrimitiv
 
 void FLightSceneInfo::ReleaseRHI()
 {
-	if (TileIntersectionResources)
-	{
-		TileIntersectionResources->Release();
-	}
-
 	ShadowCapsuleShapesVertexBuffer.SafeRelease();
 	ShadowCapsuleShapesSRV.SafeRelease();
 }
diff --git a/Engine/Source/Runtime/Renderer/Private/LightSceneInfo.h b/Engine/Source/Runtime/Renderer/Private/LightSceneInfo.h
--- a/Engine/Source/Runtime/Renderer/Private/LightSceneInfo.h
+++ b/Engine/Source/Runtime/Renderer/Private/LightSceneInfo.h
@@ -171,10 +171,6 @@ public:
 	/** The identifier for the primitive in Scene->PrimitiveOctree. */
 	FOctreeElementId2 OctreeId;
 
-	/** Tile intersection buffer for distance field shadowing, stored on the light to avoid reallocating each frame. */
-	mutable TUniquePtr<class FLightTileIntersectionResources> TileIntersectionResources;
-	mutable TUniquePtr<class FLightTileIntersectionResources> HeightFieldTileIntersectionResources;
-
 	mutable FVertexBufferRHIRef ShadowCapsuleShapesVertexBuffer;
 	mutable FShaderResourceViewRHIRef ShadowCapsuleShapesSRV;
 
