Implemented auto instancing for ray tracing. Pull request #7471

#7471

diff --git a/Engine/Source/Runtime/Renderer/Private/DeferredShadingRenderer.cpp b/Engine/Source/Runtime/Renderer/Private/DeferredShadingRenderer.cpp
--- a/Engine/Source/Runtime/Renderer/Private/DeferredShadingRenderer.cpp
+++ b/Engine/Source/Runtime/Renderer/Private/DeferredShadingRenderer.cpp
@@ -56,6 +56,7 @@
 #include "VirtualShadowMaps/VirtualShadowMapCacheManager.h"
 #include "Strata/Strata.h"
 #include "Lumen/Lumen.h"
+#include "Experimental/Containers/SherwoodHashTable.h"
 
 extern int32 GNaniteDebugFlags;
 extern int32 GNaniteShowStats;
@@ -193,6 +194,13 @@ static TAutoConsoleVariable<float> CVarRayTracingCullingAngle(
 	TEXT("Do camera culling for objects behind the camera with a projected angle smaller than this threshold in ray tracing effects (default = 5 degrees )"),
 	ECVF_RenderThreadSafe);
 
+static TAutoConsoleVariable<int32> CVarRayTracingAutoInstance(
+	TEXT("r.RayTracing.AutoInstance"),
+	1,
+	TEXT("Whether to auto instance static meshes\n"),
+	ECVF_RenderThreadSafe
+);
+
 #if !UE_BUILD_SHIPPING
 static TAutoConsoleVariable<int32> CVarForceBlackVelocityBuffer(
 	TEXT("r.Test.ForceBlackVelocityBuffer"), 0,
@@ -580,6 +588,7 @@ bool FDeferredShadingSceneRenderer::GatherRayTracingWorldInstances(FRHICommandLi
 	{
 		FRHIRayTracingGeometry* RayTracingGeometryRHI = nullptr;
 		TArrayView<const int32> CachedRayTracingMeshCommandIndices;
+		uint64 StateHash = 0;
 		int32 PrimitiveIndex = -1;
 		int8 ViewIndex = -1;
 		int8 LODIndex = -1;
@@ -588,6 +597,16 @@ bool FDeferredShadingSceneRenderer::GatherRayTracingWorldInstances(FRHICommandLi
 		bool bAllSegmentsOpaque = true;
 		bool bAnySegmentsCastShadow = false;
 		bool bAnySegmentsDecal = false;
+
+		uint64 InstancingKey() const
+		{
+			uint64 Key = StateHash;
+			Key ^= uint64(InstanceMask) << 32;
+			Key ^= bAllSegmentsOpaque ? 0x1ull << 40 : 0x0;
+			Key ^= bAnySegmentsCastShadow ? 0x1ull << 41 : 0x0;
+			Key ^= bAnySegmentsDecal ? 0x1ull << 42 : 0x0;
+			return Key ^ reinterpret_cast<uint64>(RayTracingGeometryRHI);
+		}
 	};
 
 	// Unified array is used for static and dynamic primitives because we don't know ahead of time how many we'll have of each.
@@ -795,6 +814,7 @@ bool FDeferredShadingSceneRenderer::GatherRayTracingWorldInstances(FRHICommandLi
 						RelevantPrimitive.RayTracingGeometryRHI = SceneInfo->GetStaticRayTracingGeometryInstance(LODIndex);
 
 						RelevantPrimitive.CachedRayTracingMeshCommandIndices = SceneInfo->CachedRayTracingMeshCommandIndicesPerLOD[LODIndex];
+						RelevantPrimitive.StateHash = SceneInfo->CachedRayTracingMeshCommandsHashPerLOD[LODIndex];
 
 						for (int32 CommandIndex : RelevantPrimitive.CachedRayTracingMeshCommandIndices)
 						{
@@ -983,11 +1003,17 @@ bool FDeferredShadingSceneRenderer::GatherRayTracingWorldInstances(FRHICommandLi
 	{
 		TRACE_CPUPROFILER_EVENT_SCOPE(GatherRayTracingWorldInstances_AddInstances);
 
+		const bool bAutoInstance = CVarRayTracingAutoInstance.GetValueOnRenderThread() != 0;
+
 		{
 			TRACE_CPUPROFILER_EVENT_SCOPE(WaitForLODTasks);
 			FTaskGraphInterface::Get().WaitUntilTasksComplete(LODTaskList, ENamedThreads::GetRenderThread_Local());
 		}
 
+		Experimental::TSherwoodMap<uint64, int32> InstanceSet;
+		InstanceSet.Reserve(RelevantPrimitives.Num());
+
+		// scan relevant primitives computing hash data to look for duplicate instances
 		for (const FRelevantPrimitive& RelevantPrimitive : RelevantPrimitives)
 		{
 			const int32 PrimitiveIndex = RelevantPrimitive.PrimitiveIndex;
@@ -1006,37 +1032,52 @@ bool FDeferredShadingSceneRenderer::GatherRayTracingWorldInstances(FRHICommandLi
 				continue;
 			}
 
-			const int NewInstanceIndex = View.RayTracingGeometryInstances.Num();
+			// location if this is a new entry
+			const int32 NewInstanceIndex = View.RayTracingGeometryInstances.Num();
+			const uint64 InstanceKey = RelevantPrimitive.InstancingKey();
 
+			const int32 Index = bAutoInstance ? InstanceSet.FindOrAdd(InstanceKey, NewInstanceIndex) : NewInstanceIndex;
 
-			for (int32 CommandIndex : RelevantPrimitive.CachedRayTracingMeshCommandIndices)
+			if (Index != NewInstanceIndex)
 			{
-				if (CommandIndex >= 0)
-				{
-					FVisibleRayTracingMeshCommand NewVisibleMeshCommand;
+				// reusing a previous entry, just append to the instance list
+				FRayTracingGeometryInstance& RayTracingInstance = View.RayTracingGeometryInstances[Index];
+				RayTracingInstance.NumTransforms++;
+				RayTracingInstance.Transforms.Add(Scene->PrimitiveTransforms[PrimitiveIndex]);
+				RayTracingInstance.UserData.Add((uint32)PrimitiveIndex);
 
-					NewVisibleMeshCommand.RayTracingMeshCommand = &Scene->CachedRayTracingMeshCommands.RayTracingMeshCommands[CommandIndex];
-					NewVisibleMeshCommand.InstanceIndex = NewInstanceIndex;
-					View.VisibleRayTracingMeshCommands.Add(NewVisibleMeshCommand);
-					VisibleDrawCommandStartOffset[ViewIndex]++;
-				}
-				else
+			}
+			else
+			{
+				for (int32 CommandIndex : RelevantPrimitive.CachedRayTracingMeshCommandIndices)
 				{
-					// CommandIndex == -1 indicates that the mesh batch has been filtered by FRayTracingMeshProcessor (like the shadow depth pass batch)
-					// Do nothing in this case
+					if (CommandIndex >= 0)
+					{
+						FVisibleRayTracingMeshCommand NewVisibleMeshCommand;
+
+						NewVisibleMeshCommand.RayTracingMeshCommand = &Scene->CachedRayTracingMeshCommands.RayTracingMeshCommands[CommandIndex];
+						NewVisibleMeshCommand.InstanceIndex = NewInstanceIndex;
+						View.VisibleRayTracingMeshCommands.Add(NewVisibleMeshCommand);
+						VisibleDrawCommandStartOffset[ViewIndex]++;
+					}
+					else
+					{
+						// CommandIndex == -1 indicates that the mesh batch has been filtered by FRayTracingMeshProcessor (like the shadow depth pass batch)
+						// Do nothing in this case
+					}
 				}
-			}
 
-			FRayTracingGeometryInstance& RayTracingInstance = View.RayTracingGeometryInstances.Emplace_GetRef();
-			RayTracingInstance.NumTransforms = 1;
-			RayTracingInstance.Transforms.SetNumUninitialized(1);
-			RayTracingInstance.UserData.SetNumUninitialized(1);
+				FRayTracingGeometryInstance& RayTracingInstance = View.RayTracingGeometryInstances.Emplace_GetRef();
+				RayTracingInstance.NumTransforms = 1;
+				RayTracingInstance.Transforms.SetNumUninitialized(1);
+				RayTracingInstance.UserData.SetNumUninitialized(1);
 
-			RayTracingInstance.GeometryRHI = RelevantPrimitive.RayTracingGeometryRHI;
-			RayTracingInstance.Transforms[0] = Scene->PrimitiveTransforms[PrimitiveIndex];
-			RayTracingInstance.UserData[0] = (uint32)PrimitiveIndex;
-			RayTracingInstance.Mask = RelevantPrimitive.InstanceMask; // When no cached command is found, InstanceMask == 0 and the instance is effectively filtered out
-			RayTracingInstance.bForceOpaque = RelevantPrimitive.bAllSegmentsOpaque;
+				RayTracingInstance.GeometryRHI = RelevantPrimitive.RayTracingGeometryRHI;
+				RayTracingInstance.Transforms[0] = Scene->PrimitiveTransforms[PrimitiveIndex];
+				RayTracingInstance.UserData[0] = (uint32)PrimitiveIndex;
+				RayTracingInstance.Mask = RelevantPrimitive.InstanceMask; // When no cached command is found, InstanceMask == 0 and the instance is effectively filtered out
+				RayTracingInstance.bForceOpaque = RelevantPrimitive.bAllSegmentsOpaque;
+			}
 		}
 	}
 
diff --git a/Engine/Source/Runtime/Renderer/Private/PrimitiveSceneInfo.cpp b/Engine/Source/Runtime/Renderer/Private/PrimitiveSceneInfo.cpp
--- a/Engine/Source/Runtime/Renderer/Private/PrimitiveSceneInfo.cpp
+++ b/Engine/Source/Runtime/Renderer/Private/PrimitiveSceneInfo.cpp
@@ -425,6 +425,9 @@ void FPrimitiveSceneInfo::CacheMeshDrawCommands(FRHICommandListImmediate& RHICmd
 				SceneInfo->CachedRayTracingMeshCommandIndicesPerLOD.Empty(MaxLOD + 1);
 				SceneInfo->CachedRayTracingMeshCommandIndicesPerLOD.AddDefaulted(MaxLOD + 1);
 
+				SceneInfo->CachedRayTracingMeshCommandsHashPerLOD.Empty(MaxLOD + 1);
+				SceneInfo->CachedRayTracingMeshCommandsHashPerLOD.AddZeroed(MaxLOD + 1);
+
 				for (int32 MeshIndex = 0; MeshIndex < SceneInfo->StaticMeshes.Num(); MeshIndex++)
 				{
 					FStaticMeshBatch& Mesh = SceneInfo->StaticMeshes[MeshIndex];
@@ -433,6 +436,10 @@ void FPrimitiveSceneInfo::CacheMeshDrawCommands(FRHICommandListImmediate& RHICmd
 
 					if (CommandContext.CommandIndex >= 0)
 					{
+						uint64& Hash = SceneInfo->CachedRayTracingMeshCommandsHashPerLOD[Mesh.LODIndex];
+						Hash <<= 1;
+						Hash ^= Scene->CachedRayTracingMeshCommands.RayTracingMeshCommands[CommandContext.CommandIndex].ShaderBindings.GetDynamicInstancingHash();
+
 						SceneInfo->CachedRayTracingMeshCommandIndicesPerLOD[Mesh.LODIndex].Add(CommandContext.CommandIndex);
 						CommandContext.CommandIndex = -1;
 					}
@@ -531,6 +538,8 @@ void FPrimitiveSceneInfo::RemoveCachedMeshDrawCommands()
 		}
 
 		CachedRayTracingMeshCommandIndicesPerLOD.Empty();
+
+		CachedRayTracingMeshCommandsHashPerLOD.Empty();
 	}
 #endif
 
diff --git a/Engine/Source/Runtime/Renderer/Public/PrimitiveSceneInfo.h b/Engine/Source/Runtime/Renderer/Public/PrimitiveSceneInfo.h
--- a/Engine/Source/Runtime/Renderer/Public/PrimitiveSceneInfo.h
+++ b/Engine/Source/Runtime/Renderer/Public/PrimitiveSceneInfo.h
@@ -348,6 +348,8 @@ public:
 
 	TArray<TArray<int32, TInlineAllocator<2>>> CachedRayTracingMeshCommandIndicesPerLOD;
 
+	TArray<uint64> CachedRayTracingMeshCommandsHashPerLOD;
+
 	struct FStaticMeshOrCommandIndex
 	{
 		int32 StaticMeshIndex;
