Fix light / camera origin overlap test for point-light occlusion query to match conservative bounds calculated for drawn occlusion geometry.

diff --git a/Engine/Source/Runtime/Renderer/Private/SceneOcclusion.cpp b/Engine/Source/Runtime/Renderer/Private/SceneOcclusion.cpp
--- a/Engine/Source/Runtime/Renderer/Private/SceneOcclusion.cpp
+++ b/Engine/Source/Runtime/Renderer/Private/SceneOcclusion.cpp
@@ -448,7 +448,13 @@ static bool AllocateProjectedShadowOcclusionQuery(
 	{
 		FLightSceneProxy& LightProxy = *(ProjectedShadowInfo.GetLightSceneInfo().Proxy);
 
-		const bool bCameraInsideLightGeometry = LightProxy.AffectsBounds( FSphere( View.ViewMatrices.GetViewOrigin(), View.NearClippingDistance * 2.0f ) )
+		// Make sure to perform the overlap test using the same geometry as will be used to render the sphere
+		FVector4 StencilingSpherePosAndScale(ForceInit);
+		StencilingGeometry::GStencilSphereVertexBuffer.CalcTransform(StencilingSpherePosAndScale, LightProxy.GetBoundingSphere(), View.ViewMatrices.GetPreViewTranslation());
+
+		const bool bCameraInsideLightGeometry = 
+			// Calculate overlap of the conservative sphere that will be rendered vs the camera origin, adding 2xNearClippingDistance of slack to avoid clipping issues
+			FMath::Square(StencilingSpherePosAndScale.W + View.NearClippingDistance * 2.0f) >= StencilingSpherePosAndScale.SizeSquared3()
 			// Always draw backfaces in ortho
 			//@todo - accurate ortho camera / light intersection
 			|| !View.IsPerspectiveProjection();
