Compute VT page wrapping relative to the highest resolution mip, rather than the current mip...this allows us to preserve the correct wrapped vAddress, which also needs to be relative to the top mip

diff --git a/Engine/Source/Runtime/Renderer/Private/VT/VirtualTextureSystem.cpp b/Engine/Source/Runtime/Renderer/Private/VT/VirtualTextureSystem.cpp
--- a/Engine/Source/Runtime/Renderer/Private/VT/VirtualTextureSystem.cpp
+++ b/Engine/Source/Runtime/Renderer/Private/VT/VirtualTextureSystem.cpp
@@ -1503,33 +1503,28 @@ void FVirtualTextureSystem::GatherRequestsTask(const FGatherRequestsParameters&
 			// which effectively matches more detailed mips of lower resolution producers with less detailed mips of higher resolution producers
 			uint32 Local_vLevel = vLevel - FMath::Min(vLevel, ProducerMipBias);
 
-			// Wrap Local_vAddress for the given producer
-			// For square textures, this could simply be (Local_vAddress % NumTilesInMip), but that doesn't work for non-square
-			// Possible there is a more clever approach to take here
+			// Wrap vAddress for the given producer
 			uint32 Wrapped_vAddress = vAddress;
-			uint32 Local_vAddress;
 			{
-				const uint32 MipScaleFactor = (1u << Local_vLevel);
-				const uint32 ProducerMipWidthInPages = FMath::DivideAndRoundUp(Producer->GetWidthInTiles(), MipScaleFactor);
-				const uint32 ProducerMipHeightInPages = FMath::DivideAndRoundUp(Producer->GetHeightInTiles(), MipScaleFactor);
-				uint32 Local_vPageX = (vPageX - AllocatedVT->GetVirtualPageX()) >> Mapping_vLevel;
-				uint32 Local_vTileY = (vPageY - AllocatedVT->GetVirtualPageY()) >> Mapping_vLevel;
-				if (Local_vPageX >= ProducerMipWidthInPages || Local_vTileY >= ProducerMipHeightInPages)
+				// Scale size of producer up to be relative to size of entire allocated VT
+				const uint32 ProducerScaleFactor = (1u << ProducerMipBias);
+				const uint32 ProducerWidthInPages = Producer->GetWidthInTiles() * ProducerScaleFactor;
+				const uint32 ProducerHeightInPages = Producer->GetHeightInTiles() * ProducerScaleFactor;
+				const uint32 AllocatedPageX = AllocatedVT->GetVirtualPageX();
+				const uint32 AllocatedPageY = AllocatedVT->GetVirtualPageY();
+
+				uint32 Local_vPageX = vPageX - AllocatedPageX;
+				uint32 Local_vTileY = vPageY - AllocatedPageY;
+				if (Local_vPageX >= ProducerWidthInPages || Local_vTileY >= ProducerHeightInPages)
 				{
-					Local_vPageX %= ProducerMipWidthInPages;
-					Local_vTileY %= ProducerMipHeightInPages;
-					Local_vAddress = FMath::MortonCode2(Local_vPageX) | (FMath::MortonCode2(Local_vTileY) << 1);
-					Wrapped_vAddress = (Local_vAddress << (Mapping_vLevel * vDimensions)) + AllocatedVT->GetVirtualAddress();
-				}
-				else
-				{
-					Local_vAddress = (vAddress - AllocatedVT->GetVirtualAddress()) >> (Mapping_vLevel * vDimensions);
-					// validate the wrapped calculations are nops when not wrapping
-					check(Local_vAddress == (FMath::MortonCode2(Local_vPageX) | (FMath::MortonCode2(Local_vTileY) << 1)));
-					check(vAddress == (Local_vAddress << (Mapping_vLevel * vDimensions)) + AllocatedVT->GetVirtualAddress());
+					Local_vPageX %= ProducerWidthInPages;
+					Local_vTileY %= ProducerHeightInPages;
+					Wrapped_vAddress = FMath::MortonCode2(Local_vPageX + AllocatedPageX) | (FMath::MortonCode2(Local_vTileY + AllocatedPageY) << 1);
 				}
 			}
 
+			uint32 Local_vAddress = (Wrapped_vAddress - AllocatedVT->GetVirtualAddress()) >> (Mapping_vLevel * vDimensions);
+
 			const uint32 LocalMipBias = Producer->GetVirtualTexture()->GetLocalMipBias(Local_vLevel, Local_vAddress);
 			if (LocalMipBias > 0u)
 			{
