Removed unecessary strata texture, freeing a uav. We can use some gbuffer channel if needed later.

diff --git a/Engine/Shaders/Private/BasePassPixelShader.usf b/Engine/Shaders/Private/BasePassPixelShader.usf
--- a/Engine/Shaders/Private/BasePassPixelShader.usf
+++ b/Engine/Shaders/Private/BasePassPixelShader.usf
@@ -1475,8 +1475,6 @@ void FPixelShaderInOut_MainPS(
 	GBuffer.ShadingModelID = SHADINGMODELID_STRATA;
 	GBuffer.WorldNormal = MaterialParameters.WorldNormal;
 
-	OpaqueBasePass.Strata.MaterialLobesTextureUAV[uint2(In.SvPosition.xy)] = Strata.LayerCount; // Could store a bit pack of layer count, BSDF type masks for tiling, etc.
-
 	uint2 PixelPos = In.SvPosition.xy;
 	uint PixelStrataDataByteOffset = GetStrataPixelDataByteOffset(PixelPos, ResolvedView.BufferSizeAndInvSize.x, OpaqueBasePass.Strata.MaxBytesPerPixel);
 
diff --git a/Engine/Shaders/Private/DeferredLightPixelShaders.usf b/Engine/Shaders/Private/DeferredLightPixelShaders.usf
--- a/Engine/Shaders/Private/DeferredLightPixelShaders.usf
+++ b/Engine/Shaders/Private/DeferredLightPixelShaders.usf
@@ -195,8 +195,7 @@ void DeferredLightPixelMain(
 		float4 Radiance = 0;
 
 #if STRATA_ENABLED
-		uint2 PixelPos = InputParams.PixelPos;
-		if (Strata.MaterialLobesTexture[PixelPos].x > 0.0f)
+		if (ScreenSpaceData.GBuffer.ShadingModelID == SHADINGMODELID_STRATA)
 		{
 			float3 V =-DerivedParams.CameraVector;
 			float3 L = LightData.Direction;	// Already normalized
@@ -209,6 +208,7 @@ void DeferredLightPixelMain(
 
 			if (LightMask > 0)
 			{
+				uint2 PixelPos = InputParams.PixelPos;
 				uint PixelStrataDataByteOffset = GetStrataPixelDataByteOffset(PixelPos, View.BufferSizeAndInvSize.x, Strata.MaxBytesPerPixel);
 				FStrataPixelBSDFs PixelBSDFs = UnpackStrataIn(Strata.MaterialLobesBuffer, PixelStrataDataByteOffset);
 
diff --git a/Engine/Shaders/Private/ReflectionEnvironmentPixelShader.usf b/Engine/Shaders/Private/ReflectionEnvironmentPixelShader.usf
--- a/Engine/Shaders/Private/ReflectionEnvironmentPixelShader.usf
+++ b/Engine/Shaders/Private/ReflectionEnvironmentPixelShader.usf
@@ -385,7 +385,7 @@ void ReflectionEnvironmentSkyLighting(
 
 	const bool bUnlitMaterial = ShadingModelID == SHADINGMODELID_UNLIT;
 #if STRATA_ENABLED
-	const bool bStrataMaterial = Strata.MaterialLobesTexture[PixelPos].x > 0.0f;
+	const bool bStrataMaterial = ShadingModelID == SHADINGMODELID_STRATA;
 #else
 	const bool bStrataMaterial = false;
 #endif
diff --git a/Engine/Shaders/Private/Strata/StrataVisualize.usf b/Engine/Shaders/Private/Strata/StrataVisualize.usf
--- a/Engine/Shaders/Private/Strata/StrataVisualize.usf
+++ b/Engine/Shaders/Private/Strata/StrataVisualize.usf
@@ -1,6 +1,7 @@
 // Copyright Epic Games, Inc. All Rights Reserved.
 
 #include "/Engine/Private/Common.ush"
+#include "/Engine/Private/DeferredShadingCommon.ush"
 
 #include "/Engine/Private/Strata/Strata.ush"
 
@@ -21,10 +22,14 @@ void VisualizeMaterialPS(
 	const uint2 MaterialPos = View.CursorPosition;
 	const int  InspectSquare = 10;
 
+	// Sample scene textures.
+	float2 BufferUV = SvPositionToBufferUV(SVPos);
+	FGBufferData GBuffer = GetGBufferDataFromSceneTextures(BufferUV);
+	bool bStrataMaterial = GBuffer.ShadingModelID == SHADINGMODELID_STRATA;
 
 #if 0
 	OutColor = 0;
-	if (Strata.MaterialLobesTexture[PixelPos].x > 0.0f)
+	if (bStrataMaterial)
 	{
 		uint PixelStrataDataByteOffset = GetStrataPixelDataByteOffset(PixelPos, View.BufferSizeAndInvSize.x, Strata.MaxBytesPerPixel);
 		FStrataPixelBSDFs PixelBSDFs = UnpackStrataIn(Strata.MaterialLobesBuffer, PixelStrataDataByteOffset);
@@ -60,7 +65,7 @@ void VisualizeMaterialPS(
 		OutColor.w = 0.0f;				// Set background visiblity to 0 in the bottom right part of the screen to better reads the debug values
 	}
 
-	if (Strata.MaterialLobesTexture[MaterialPos].x > 0.0f)
+	if (bStrataMaterial)
 	{
 		uint PixelStrataDataByteOffset = GetStrataPixelDataByteOffset(MaterialPos, View.BufferSizeAndInvSize.x, Strata.MaxBytesPerPixel);
 		FStrataPixelBSDFs PixelBSDFs = UnpackStrataIn(Strata.MaterialLobesBuffer, PixelStrataDataByteOffset);
diff --git a/Engine/Source/Runtime/Renderer/Private/Strata/Strata.cpp b/Engine/Source/Runtime/Renderer/Private/Strata/Strata.cpp
--- a/Engine/Source/Runtime/Renderer/Private/Strata/Strata.cpp
+++ b/Engine/Source/Runtime/Renderer/Private/Strata/Strata.cpp
@@ -8,6 +8,7 @@
 #include "SceneRendering.h"
 #include "RendererInterface.h"
 #include "UniformBuffer.h"
+#include "SceneTextureParameters.h"
 
 
 
@@ -57,12 +58,6 @@ void InitialiseStrataFrameSceneData(FSceneRenderer& SceneRenderer, FRDGBuilder&
 		StrataSceneData.MaxBytesPerPixel = 4u;
 	}
 
-	FRDGTextureRef MaterialLobesTexture = GraphBuilder.CreateTexture(
-		FRDGTextureDesc::Create2D(FIntPoint(ResolutionX, ResolutionY), PF_R16F, FClearValueBinding::None,
-		TexCreate_RenderTargetable | TexCreate_ShaderResource | TexCreate_UAV), TEXT("MaterialLobesTexture"));
-	AddClearRenderTargetPass(GraphBuilder, MaterialLobesTexture, FLinearColor::Black);
-	ConvertToExternalTexture(GraphBuilder, MaterialLobesTexture, StrataSceneData.MaterialLobesTexture);
-
 	const uint32 DesiredBufferSize = FMath::Max(4u, ResolutionX * ResolutionY * StrataSceneData.MaxBytesPerPixel);
 	if (StrataSceneData.MaterialLobesBuffer.NumBytes < DesiredBufferSize)
 	{
@@ -89,13 +84,11 @@ void BindStrataBasePassUniformParameters(const FViewInfo& View, FStrataOpaquePas
 	if (View.StrataSceneData)
 	{
 		OutStrataUniformParameters.MaxBytesPerPixel = View.StrataSceneData->MaxBytesPerPixel;
-		OutStrataUniformParameters.MaterialLobesTextureUAV = View.StrataSceneData->MaterialLobesTexture->GetRenderTargetItem().UAV;
 		OutStrataUniformParameters.MaterialLobesBufferUAV = View.StrataSceneData->MaterialLobesBuffer.UAV;
 	}
 	else
 	{
 		OutStrataUniformParameters.MaxBytesPerPixel = 0;
-		OutStrataUniformParameters.MaterialLobesTextureUAV = GEmptyVertexBufferWithUAV->UnorderedAccessViewRHI;
 		OutStrataUniformParameters.MaterialLobesBufferUAV = GEmptyVertexBufferWithUAV->UnorderedAccessViewRHI;
 	}
 }
@@ -112,7 +105,6 @@ TUniformBufferRef<FStrataGlobalUniformParameters> BindStrataGlobalUniformParamet
 		}
 
 		StrataUniformParameters.MaxBytesPerPixel = View.StrataSceneData->MaxBytesPerPixel;
-		StrataUniformParameters.MaterialLobesTexture = View.StrataSceneData->MaterialLobesTexture->GetRenderTargetItem().ShaderResourceTexture;
 		StrataUniformParameters.MaterialLobesBuffer = View.StrataSceneData->MaterialLobesBuffer.SRV;
 
 		View.StrataSceneData->StrataGlobalUniformParameters = CreateUniformBufferImmediate(StrataUniformParameters, UniformBuffer_SingleFrame);
@@ -122,7 +114,6 @@ TUniformBufferRef<FStrataGlobalUniformParameters> BindStrataGlobalUniformParamet
 	{
 		// Create each time. This path will go away when Strata is always enabled anyway.
 		StrataUniformParameters.MaxBytesPerPixel = 0;
-		StrataUniformParameters.MaterialLobesTexture = GSystemTextures.BlackDummy->GetRenderTargetItem().ShaderResourceTexture;
 		StrataUniformParameters.MaterialLobesBuffer = GEmptyVertexBufferWithUAV->ShaderResourceViewRHI;
 		return CreateUniformBufferImmediate(StrataUniformParameters, UniformBuffer_SingleDraw);
 	}
@@ -143,6 +134,7 @@ class FVisualizeMaterialPS : public FGlobalShader
 		SHADER_PARAMETER_STRUCT_REF(FViewUniformShaderParameters, ViewUniformBuffer)
 		SHADER_PARAMETER_STRUCT_REF(FStrataGlobalUniformParameters, Strata)
 		SHADER_PARAMETER_TEXTURE(Texture2D, MiniFontTexture)
+		SHADER_PARAMETER_STRUCT_INCLUDE(FSceneTextureParameters, SceneTextures)
 		RENDER_TARGET_BINDING_SLOTS()
 	END_SHADER_PARAMETER_STRUCT()
 
@@ -186,6 +178,7 @@ void AddVisualizeMaterialPasses(FRDGBuilder& GraphBuilder, const TArray<FViewInf
 			PassParameters->ViewUniformBuffer = View.ViewUniformBuffer;
 			PassParameters->Strata = Strata::BindStrataGlobalUniformParameters(View);
 			PassParameters->MiniFontTexture = GetMiniFontTexture();
+			PassParameters->SceneTextures = GetSceneTextureParameters(GraphBuilder);
 			PassParameters->RenderTargets[0] = FRenderTargetBinding(SceneColorTexture, ERenderTargetLoadAction::ELoad);
 
 			FPixelShaderUtils::AddFullscreenPass<FVisualizeMaterialPS>(GraphBuilder, View.ShaderMap, RDG_EVENT_NAME("StrataVisualizeMaterial"),
diff --git a/Engine/Source/Runtime/Renderer/Private/Strata/Strata.h b/Engine/Source/Runtime/Renderer/Private/Strata/Strata.h
--- a/Engine/Source/Runtime/Renderer/Private/Strata/Strata.h
+++ b/Engine/Source/Runtime/Renderer/Private/Strata/Strata.h
@@ -18,13 +18,11 @@ class FRDGBuilder;
 BEGIN_SHADER_PARAMETER_STRUCT(FStrataOpaquePassUniformParameters, )
 	SHADER_PARAMETER(uint32, MaxBytesPerPixel)
 	SHADER_PARAMETER_UAV(RWByteAddressBuffer, MaterialLobesBufferUAV)
-	SHADER_PARAMETER_UAV(RWTexture2D<float>, MaterialLobesTextureUAV)
 END_SHADER_PARAMETER_STRUCT()
 
 BEGIN_GLOBAL_SHADER_PARAMETER_STRUCT(FStrataGlobalUniformParameters, )
 	SHADER_PARAMETER(uint32, MaxBytesPerPixel)
 	SHADER_PARAMETER_SRV(ByteAddressBuffer, MaterialLobesBuffer)
-	SHADER_PARAMETER_TEXTURE(Texture2D<float>, MaterialLobesTexture)
 END_GLOBAL_SHADER_PARAMETER_STRUCT()
 
 
