Removed deprecated scene texture uniform buffer APIs.

diff --git a/Engine/Source/Runtime/Renderer/Public/SceneRenderTargetParameters.h b/Engine/Source/Runtime/Renderer/Public/SceneRenderTargetParameters.h
--- a/Engine/Source/Runtime/Renderer/Public/SceneRenderTargetParameters.h
+++ b/Engine/Source/Runtime/Renderer/Public/SceneRenderTargetParameters.h
@@ -161,25 +161,4 @@ extern RENDERER_API TRefCountPtr<FRHIUniformBuffer> CreateSceneTextureUniformBuf
 	ESceneTextureSetupMode SetupMode = ESceneTextureSetupMode::All);
 
 
-extern RENDERER_API bool IsSceneTexturesValid(FRHICommandListImmediate& RHICmdList);
-
-
-/** Deprecated APIs. */
-
-/** This was renamed to remove the extraneous 's' for consistency. */
-UE_DEPRECATED(4.26, "FSceneTexturesUniformParameters has been renamed to FSceneTextureUniformParameters.")
-typedef FSceneTextureUniformParameters FSceneTexturesUniformParameters;
-
-UE_DEPRECATED(4.26, "BindSceneTextureUniformBufferDependentOnShadingPath has been removed.")
-inline void BindSceneTextureUniformBufferDependentOnShadingPath(
-	const FShader::CompiledShaderInitializerType& Initializer,
-	FShaderUniformBufferParameter& SceneTexturesUniformBuffer)
-{}
-
-UE_DEPRECATED(4.26, "SetupMobileSceneTextureUniformParameters has been refactored to EMobileSceneTextureSetupMode.")
-extern void SetupMobileSceneTextureUniformParameters(
-	FSceneRenderTargets& SceneContext,
-	ERHIFeatureLevel::Type FeatureLevel,
-	bool bSceneTexturesValid,
-	bool bCustomDepthIsValid,
-	FMobileSceneTextureUniformParameters& SceneTextureParameters);
\ No newline at end of file
+extern RENDERER_API bool IsSceneTexturesValid(FRHICommandListImmediate& RHICmdList);
\ No newline at end of file
