Removed scene texture uniform buffer binding from dead code.

#rnx

diff --git a/Engine/Source/Runtime/Renderer/Private/HeightfieldLighting.cpp b/Engine/Source/Runtime/Renderer/Private/HeightfieldLighting.cpp
--- a/Engine/Source/Runtime/Renderer/Private/HeightfieldLighting.cpp
+++ b/Engine/Source/Runtime/Renderer/Private/HeightfieldLighting.cpp
@@ -1041,10 +1041,6 @@ void FHeightfieldLightingViewInfo::ComputeLighting(FRDGBuilder& GraphBuilder, co
 
 	AddPass(GraphBuilder, RDG_EVENT_NAME("HeightfieldLightingForGI"), [this, &View, &LightSceneInfo, FeatureLevel](FRHICommandListImmediate& RHICmdList)
 	{
-		FUniformBufferRHIRef PassUniformBuffer = CreateSceneTextureUniformBuffer(RHICmdList, View.FeatureLevel);
-		FUniformBufferStaticBindings GlobalUniformBuffers(PassUniformBuffer);
-		SCOPED_UNIFORM_BUFFER_GLOBAL_BINDINGS(RHICmdList, GlobalUniformBuffers);
-
 		FSceneViewState* ViewState = (FSceneViewState*)View.State;
 		const FHeightfieldLightingAtlas& Atlas = *ViewState->HeightfieldLightingAtlas;
 
