Fixed RDG uniform buffer usage in shadows.

#rnx

diff --git a/Engine/Source/Runtime/Renderer/Private/VirtualShadowMaps/VirtualShadowMapProjection.cpp b/Engine/Source/Runtime/Renderer/Private/VirtualShadowMaps/VirtualShadowMapProjection.cpp
--- a/Engine/Source/Runtime/Renderer/Private/VirtualShadowMaps/VirtualShadowMapProjection.cpp
+++ b/Engine/Source/Runtime/Renderer/Private/VirtualShadowMaps/VirtualShadowMapProjection.cpp
@@ -41,7 +41,7 @@ static TAutoConsoleVariable<int32> CVarVirtualShadowMapDebugProjection(
 BEGIN_SHADER_PARAMETER_STRUCT(FProjectionParameters, )
 	SHADER_PARAMETER_STRUCT_INCLUDE(FVirtualShadowMapSamplingParameters, ProjectionParameters)
 	SHADER_PARAMETER_STRUCT(FLightShaderParameters, Light)
-	SHADER_PARAMETER_STRUCT_REF(FSceneTextureUniformParameters, SceneTexturesStruct)
+	SHADER_PARAMETER_RDG_UNIFORM_BUFFER(FSceneTextureUniformParameters, SceneTexturesStruct)
 	SHADER_PARAMETER_STRUCT_REF(FViewUniformShaderParameters, View)
 	SHADER_PARAMETER(int32, VirtualShadowMapId)
 	SHADER_PARAMETER(float, ContactShadowLength)
@@ -116,7 +116,7 @@ static void AddPass_RenderVirtualShadowMapProjection(
 
 	FProjectionParameters* PassParameters = GraphBuilder.AllocParameters<FProjectionParameters>();
 	VirtualShadowMapArray.SetProjectionParameters(GraphBuilder, PassParameters->ProjectionParameters);
-	PassParameters->SceneTexturesStruct = CreateSceneTextureUniformBuffer(GraphBuilder.RHICmdList, View.FeatureLevel, ESceneTextureSetupMode::GBuffers | ESceneTextureSetupMode::SceneDepth);
+	PassParameters->SceneTexturesStruct = CreateSceneTextureUniformBuffer(GraphBuilder, View.FeatureLevel, ESceneTextureSetupMode::GBuffers | ESceneTextureSetupMode::SceneDepth);
 	PassParameters->View = View.ViewUniformBuffer;
 		
 	FLightShaderParameters LightParameters;
