Remove Adaptive sampler from path tracer in preparation for porting to RDG

This reduces the number of stages involved to just two: one for the actual path tracing and one for display to SceneColor which should help port the code to RDG as well as make it easier to add new buffers required for denoising.

As a side effect of this change, the path tracing samples are accumulated in full floating point which fixes the "washed out" appearrance on scenes with very bright lights that came from the fixed point accumulation required by the adaptive sampling path compaction pass.

diff --git a/Engine/Shaders/Private/PathTracing/PathTracing.usf b/Engine/Shaders/Private/PathTracing/PathTracing.usf
--- a/Engine/Shaders/Private/PathTracing/PathTracing.usf
+++ b/Engine/Shaders/Private/PathTracing/PathTracing.usf
@@ -22,9 +22,6 @@ PathTracingRayGenShader.usf: Reference path tracing
 
 
 RWTexture2D<float4> RadianceRT;
-RWTexture2D<uint>	SampleCountRT;
-RWTexture2D<uint>   PixelPositionRT;
-RWTexture2D<uint>	RayCountPerPixelRT;
 RaytracingAccelerationStructure TLAS;
 
 void AccumulateRadiance(inout float3 TotalRadiance, float3 PathRadiance)
@@ -46,11 +43,6 @@ RAY_TRACING_ENTRY_RAYGEN(PathTracingMainRG)
 {
 	uint2 LaunchIndex = DispatchRaysIndex().xy + View.ViewRectMin.xy + PathTracingData.TileOffset.xy;
 	
-	RadianceRT[LaunchIndex] = 0;
-	SampleCountRT[LaunchIndex] = 0;
-	PixelPositionRT[LaunchIndex] = 0;
-	RayCountPerPixelRT[LaunchIndex] = 0;
-	
 	// Initialize random sequence 
 	float2 BufferSize = View.BufferSizeAndInvSize.xy; 
 	uint LinearIndex = LaunchIndex.x + (LaunchIndex.y * BufferSize.x); 
@@ -299,11 +291,25 @@ RAY_TRACING_ENTRY_RAYGEN(PathTracingMainRG)
 		// Update ray and loop back around to trace again
 		Ray.Direction = Direction;
 		Ray.TMin = 0;
-		ApplyPositionBias(Ray, SignedPositionBias* Payload.WorldNormal, PathTracingData.MaxNormalBias);
+		ApplyPositionBias(Ray, SignedPositionBias * Payload.WorldNormal, PathTracingData.MaxNormalBias);
 	}
 
-	RadianceRT[LaunchIndex].rgb += Radiance;
-	SampleCountRT[LaunchIndex] += 1;
-	PixelPositionRT[LaunchIndex] = Pixel.y * BufferSize.x + Pixel.x;
-	RayCountPerPixelRT[LaunchIndex] += RayCounter;
+
+	// Accumulate radiance and update pixel variance
+	uint NumSamples = AdaptiveSamplingData.Iteration + 1;
+	float BlendFactor = 1.0 / float(NumSamples);
+	// Avoid reading the old pixel on the first sample on the off-chance there is a NaN/Inf pixel ...
+	float4 OldPixel = NumSamples > 1 ? RadianceRT[LaunchIndex] : 0;
+	float3 OldRadiance = OldPixel.rgb;
+	float OldVariance = OldPixel.a;
+
+	float3 NewRadiance = lerp(OldRadiance, Radiance, BlendFactor);
+
+	float RadianceY = Luminance(Radiance);
+
+	// https://en.wikipedia.org/wiki/Algorithms_for_calculating_variance#Welford's_online_algorithm
+	float DeviationSquared = (RadianceY - Luminance(NewRadiance)) * (RadianceY - Luminance(OldRadiance));
+	float NewVariance = lerp(OldVariance, DeviationSquared, BlendFactor);
+
+	RadianceRT[LaunchIndex] = float4(NewRadiance, NewVariance);
 }
diff --git a/Engine/Shaders/Private/PathTracing/PathTracingCompositingPixelShader.usf b/Engine/Shaders/Private/PathTracing/PathTracingCompositingPixelShader.usf
--- a/Engine/Shaders/Private/PathTracing/PathTracingCompositingPixelShader.usf
+++ b/Engine/Shaders/Private/PathTracing/PathTracingCompositingPixelShader.usf
@@ -7,14 +7,9 @@ ReferencePathTracinPixelShader.ush: Reference path tracing  compositing pixel sh
 #include "../Common.ush"
 #include "Utilities/PathTracingFixedPointArithmetic.ush"
 
-Texture2D<uint> RadianceRedTexture;
-Texture2D<uint> RadianceGreenTexture;
-Texture2D<uint> RadianceBlueTexture;
-Texture2D<uint> RadianceAlphaTexture;
-Texture2D<uint> SampleCountTexture;
+uint Iteration;
+Texture2D<float4> RadianceTexture;
 
-Texture2D<float4> CumulativeIrradianceTexture;
-Texture2D<uint> CumulativeSampleCountTexture;
 
 float TonemapGammaAndHDR(float Value)
 {
@@ -41,66 +36,27 @@ float Tonemap(float Value)
 
 void CompositeMain(
 	in noperspective float2 UV : TEXCOORD0,
-	out float4 OutColor : SV_Target0,
-	out float4 OutCumulativeIrradiance : SV_Target1,
-	out uint OutCumulativeSampleCount : SV_Target2
+	out float4 OutColor : SV_Target0
 )
 {
 	float2 BufferSize = View.BufferSizeAndInvSize.xy;
 	int3 TexCoord = int3(UV * BufferSize, 0);
 
-	float RadianceRed = FixedPointToFloat(RadianceRedTexture.Load(TexCoord), 24);
-	float RadianceGreen = FixedPointToFloat(RadianceGreenTexture.Load(TexCoord), 24);
-	float RadianceBlue = FixedPointToFloat(RadianceBlueTexture.Load(TexCoord), 24);
-	float RadianceAlpha = FixedPointToFloat(RadianceAlphaTexture.Load(TexCoord), 24);
-	float4 Radiance = float4(RadianceRed, RadianceGreen, RadianceBlue, RadianceAlpha);
-	uint SampleCount = SampleCountTexture.Load(TexCoord);
+	float4 Radiance = RadianceTexture.Load(TexCoord);
 
-	if (SampleCount > 1)
-		Radiance /= float(SampleCount);
-
-	float4 CumulativeIrradiance = CumulativeIrradianceTexture.Load(TexCoord);
-	uint CumulativeSampleCount = CumulativeSampleCountTexture.Load(TexCoord);
-
-	// extract variance from variance of the mean
-	float CumulativeVariance = CumulativeIrradiance.a * CumulativeSampleCount;
-
-	OutCumulativeSampleCount = CumulativeSampleCount + SampleCount;
-    // Average radiance using the smaller blend factor for extra precision
-	float BlendFactor = float(SampleCount) / float(OutCumulativeSampleCount);
-	OutCumulativeIrradiance.rgb = lerp(CumulativeIrradiance.rgb, Radiance.rgb, BlendFactor);
-
-	float OutVariance = 0;
-	float RadianceY = Luminance(Radiance.rgb);
-	float DeviationSquared = (RadianceY - Luminance(OutCumulativeIrradiance.rgb)) * (RadianceY - Luminance(CumulativeIrradiance.rgb));
-	// This formula is correct when SampleCount == 1 - but for larger values we would
-    // theoretically need to know the variance of our new batch of samples ...
-	OutVariance = lerp(CumulativeVariance, DeviationSquared, BlendFactor);
-
-	// store variance of the mean
-	OutCumulativeIrradiance.a = OutCumulativeSampleCount > 0 ? OutVariance / OutCumulativeSampleCount : 0;
-	OutColor = float4(OutCumulativeIrradiance.rgb * View.PreExposure, 1.0);
+	OutColor = float4(Radiance.rgb * View.PreExposure, 1.0);
 
 #if 0
     // split-screen debugging aids
 	if (UV.x < 0.5)
 	{
-		float Lum = Luminance(OutCumulativeIrradiance.rgb); // mean
-		float StdErr = OutCumulativeSampleCount > 0 ? sqrt(OutVariance / OutCumulativeSampleCount) : 0; // stderr
+		uint NumSamples = Iteration + 1;
+		float OutVariance = Radiance.a;
+		float Lum = Luminance(Radiance.rgb); // mean
+		float StdErr = NumSamples > 0 ? sqrt(OutVariance / NumSamples) : 0; // stderr
 		//OutColor = Tonemap(Lum + 3 * StdErr) - Tonemap(max(Lum - 3 * StdErr, 0.0)); // error estimator
-		//OutColor = StdErr * View.PreExposure;
-		OutColor = StdErr / (Lum + 0.0001) * View.PreExposure;
-		//OutColor = OutVariance / (Tonemap(Lum) + 0.0001);
-		//OutColor = OutVariance * View.PreExposure;
-		//OutColor = OutCumulativeSampleCount / 1024.0 * View.PreExposure;
-		//OutColor.r = AdaptiveSamplingData.VarianceDimensions.x + AdaptiveSamplingData.VarianceDimensions.x;
-		//OutColor.g = AdaptiveSamplingData.VarianceDimensions.x;
-		//OutColor.b = AdaptiveSamplingData.VarianceDimensions.y;
-
-		//OutColor = AdaptiveSamplingData.VarianceMipTree[BufferOffsetAtPixel(TexCoord.xy, 0)];
-
+		OutColor = StdErr * View.PreExposure;
+		//OutColor = StdErr / (Lum + 0.0001) * View.PreExposure;
 	}
 #endif
-	// To visualize the current frame:
-	//OutColor = float4(Radiance.rgb / SampleCount, 1.0);
-}
\ No newline at end of file
+}
diff --git a/Engine/Shaders/Private/PathTracing/Sensor/PathTracingAdaptivePinholeCamera.ush b/Engine/Shaders/Private/PathTracing/Sensor/PathTracingAdaptivePinholeCamera.ush
--- a/Engine/Shaders/Private/PathTracing/Sensor/PathTracingAdaptivePinholeCamera.ush
+++ b/Engine/Shaders/Private/PathTracing/Sensor/PathTracingAdaptivePinholeCamera.ush
@@ -38,6 +38,10 @@ uint Log2(uint Value)
 
 bool DoAdaptiveSampling()
 {
+#if 1
+	return false;
+#else
+	// TODO: remove this
 	bool bVarianceMapExists = AdaptiveSamplingData.VarianceDimensions.x + AdaptiveSamplingData.VarianceDimensions.y > 0;
 	bool bVarianceMapIsPrimed = AdaptiveSamplingData.Iteration >= AdaptiveSamplingData.MinimumSamplesPerPixel;
 	bool bVarianceSumIsValid = false;
@@ -49,6 +53,7 @@ bool DoAdaptiveSampling()
 		bVarianceSumIsValid = VarianceSum > Epsilon;
 	}
 	return bVarianceMapExists && bVarianceMapIsPrimed && bVarianceSumIsValid;
+#endif
 }
 
 void AdaptivePinholeCamera_SampleEmitter(
diff --git a/Engine/Source/Runtime/Renderer/Private/PathTracing.cpp b/Engine/Source/Runtime/Renderer/Private/PathTracing.cpp
--- a/Engine/Source/Runtime/Renderer/Private/PathTracing.cpp
+++ b/Engine/Source/Runtime/Renderer/Private/PathTracing.cpp
@@ -48,7 +48,7 @@ TAutoConsoleVariable<int32> CVarPathTracingVisibleLights(
 	ECVF_RenderThreadSafe
 );
 
-TAutoConsoleVariable<int32> CVarPathTracingMaxPathIntensity(
+TAutoConsoleVariable<float> CVarPathTracingMaxPathIntensity(
 	TEXT("r.PathTracing.MaxPathIntensity"),
 	-1,
 	TEXT("When positive, light paths greater that this amount are clamped to prevent fireflies (default = -1 (off))"),
@@ -205,9 +205,6 @@ public:
 
 		// Output
 		RadianceRT.Bind(Initializer.ParameterMap, TEXT("RadianceRT"));
-		SampleCountRT.Bind(Initializer.ParameterMap, TEXT("SampleCountRT"));
-		PixelPositionRT.Bind(Initializer.ParameterMap, TEXT("PixelPositionRT"));
-		RayCountPerPixelRT.Bind(Initializer.ParameterMap, TEXT("RayCountPerPixelRT"));
 	}
 
 	void SetParameters(
@@ -229,10 +226,7 @@ public:
 		const FRWBuffer& VarianceMipTree,
 		const FIntVector& TileOffset,
 		// Output
-		FRHIUnorderedAccessView* RadianceUAV,
-		FRHIUnorderedAccessView* SampleCountUAV,
-		FRHIUnorderedAccessView* PixelPositionUAV,
-		FRHIUnorderedAccessView* RayCountPerPixelUAV)
+		FRHIUnorderedAccessView* RadianceUAV)
 	{
 
 		GlobalResources.Set(TLASParameter, RayTracingScene.RayTracingSceneRHI->GetShaderResourceView());
@@ -379,12 +373,9 @@ public:
 			GlobalResources.Set(AdaptiveSamplingParameters, AdaptiveSamplingDataUniformBuffer);
 		}
 
-		// Output
+		// Accumulated output
 		{
 			GlobalResources.Set(RadianceRT, RadianceUAV);
-			GlobalResources.Set(SampleCountRT, SampleCountUAV);
-			GlobalResources.Set(PixelPositionRT, PixelPositionUAV);
-			GlobalResources.Set(RayCountPerPixelRT, RayCountPerPixelUAV);
 		}
 	}
 
@@ -415,13 +406,10 @@ public:
 
 	// Output parameters
 	LAYOUT_FIELD(FShaderResourceParameter, RadianceRT);
-	LAYOUT_FIELD(FShaderResourceParameter, SampleCountRT);
-	LAYOUT_FIELD(FShaderResourceParameter, PixelPositionRT);
-	LAYOUT_FIELD(FShaderResourceParameter, RayCountPerPixelRT);
 };
 IMPLEMENT_SHADER_TYPE(, FPathTracingRG, TEXT("/Engine/Private/PathTracing/PathTracing.usf"), TEXT("PathTracingMainRG"), SF_RayGen);
 
-DECLARE_GPU_STAT_NAMED(Stat_GPU_PathTracing, TEXT("Reference Path Tracing"));
+DECLARE_GPU_STAT_NAMED(Stat_GPU_PathTracing, TEXT("Path Tracing"));
 DECLARE_GPU_STAT_NAMED(Stat_GPU_PathTracingBuildSkyLightCDF, TEXT("Path Tracing: Build Sky Light CDF"));
 DECLARE_GPU_STAT_NAMED(Stat_GPU_PathTracingBuildVarianceMipTree, TEXT("Path Tracing: Build Variance Map Tree"));
 
@@ -447,14 +435,8 @@ class FPathTracingCompositorPS : public FGlobalShader
 	FPathTracingCompositorPS(const ShaderMetaType::CompiledShaderInitializerType& Initializer)
 		: FGlobalShader(Initializer)
 	{
-		RadianceRedTexture.Bind(Initializer.ParameterMap, TEXT("RadianceRedTexture"));
-		RadianceGreenTexture.Bind(Initializer.ParameterMap, TEXT("RadianceGreenTexture"));
-		RadianceBlueTexture.Bind(Initializer.ParameterMap, TEXT("RadianceBlueTexture"));
-		RadianceAlphaTexture.Bind(Initializer.ParameterMap, TEXT("RadianceAlphaTexture"));
-		SampleCountTexture.Bind(Initializer.ParameterMap, TEXT("SampleCountTexture"));
-
-		CumulativeIrradianceTexture.Bind(Initializer.ParameterMap, TEXT("CumulativeIrradianceTexture"));
-		CumulativeSampleCountTexture.Bind(Initializer.ParameterMap, TEXT("CumulativeSampleCountTexture"));
+		Iteration.Bind(Initializer.ParameterMap, TEXT("Iteration"));
+		RadianceTexture.Bind(Initializer.ParameterMap, TEXT("RadianceTexture"));
 	}
 
 	FPathTracingCompositorPS()
@@ -465,34 +447,18 @@ class FPathTracingCompositorPS : public FGlobalShader
 	void SetParameters(
 		TRHICommandList& RHICmdList,
 		const FViewInfo& View,
-		FRHITexture* RadianceRedRT,
-		FRHITexture* RadianceGreenRT,
-		FRHITexture* RadianceBlueRT,
-		FRHITexture* RadianceAlphaRT,
-		FRHITexture* SampleCountRT,
-		FRHITexture* CumulativeIrradianceRT,
-		FRHITexture* CumulativeSampleCountRT)
+		uint32 IterationValue,
+		FRHITexture* RadianceRT)
 	{
 		FRHIPixelShader* ShaderRHI = RHICmdList.GetBoundPixelShader();
 		FGlobalShader::SetParameters<FViewUniformShaderParameters>(RHICmdList, ShaderRHI, View.ViewUniformBuffer);
-		SetTextureParameter(RHICmdList, ShaderRHI, RadianceRedTexture, RadianceRedRT);
-		SetTextureParameter(RHICmdList, ShaderRHI, RadianceGreenTexture, RadianceGreenRT);
-		SetTextureParameter(RHICmdList, ShaderRHI, RadianceBlueTexture, RadianceBlueRT);
-		SetTextureParameter(RHICmdList, ShaderRHI, RadianceAlphaTexture, RadianceAlphaRT);
-		SetTextureParameter(RHICmdList, ShaderRHI, SampleCountTexture, SampleCountRT);
-		SetTextureParameter(RHICmdList, ShaderRHI, CumulativeIrradianceTexture, CumulativeIrradianceRT);
-		SetTextureParameter(RHICmdList, ShaderRHI, CumulativeSampleCountTexture, CumulativeSampleCountRT);
+		SetShaderValue(RHICmdList, ShaderRHI, Iteration, IterationValue);
+		SetTextureParameter(RHICmdList, ShaderRHI, RadianceTexture, RadianceRT);
 	}
 
 public:
-	LAYOUT_FIELD(FShaderResourceParameter, RadianceRedTexture);
-	LAYOUT_FIELD(FShaderResourceParameter, RadianceGreenTexture);
-	LAYOUT_FIELD(FShaderResourceParameter, RadianceBlueTexture);
-	LAYOUT_FIELD(FShaderResourceParameter, RadianceAlphaTexture);
-	LAYOUT_FIELD(FShaderResourceParameter, SampleCountTexture);
-
-	LAYOUT_FIELD(FShaderResourceParameter, CumulativeIrradianceTexture);
-	LAYOUT_FIELD(FShaderResourceParameter, CumulativeSampleCountTexture);
+	LAYOUT_FIELD(FShaderParameter, Iteration);
+	LAYOUT_FIELD(FShaderResourceParameter, RadianceTexture);
 };
 
 IMPLEMENT_SHADER_TYPE(, FPathTracingCompositorPS, TEXT("/Engine/Private/PathTracing/PathTracingCompositingPixelShader.usf"), TEXT("CompositeMain"), SF_Pixel);
@@ -514,8 +480,7 @@ END_SHADER_PARAMETER_STRUCT()
 
 void FSceneViewState::PathTracingInvalidate()
 {
-	PathTracingIrradianceRT.SafeRelease();
-	PathTracingSampleCountRT.SafeRelease();
+	PathTracingRadianceRT.SafeRelease();
 	VarianceMipTreeDimensions = FIntVector(0);
 	TotalRayCount = 0;
 	PathTracingSPP = 0;
@@ -571,29 +536,20 @@ void FDeferredShadingSceneRenderer::RenderPathTracing(
 		
 		// Construct render targets for compositing
 		TRefCountPtr<IPooledRenderTarget> RadianceRT;
-		TRefCountPtr<IPooledRenderTarget> SampleCountRT;
-		TRefCountPtr<IPooledRenderTarget> PixelPositionRT;
-		TRefCountPtr<IPooledRenderTarget> RayCountPerPixelRT;
-
-		FPooledRenderTargetDesc Desc = Translate(SceneColorOutputTexture->Desc);
-		Desc.Flags &= ~(TexCreate_FastVRAM | TexCreate_Transient);
-		Desc.Format = PF_FloatRGBA;
-		GRenderTargetPool.FindFreeElement(RHICmdList, Desc, RadianceRT, TEXT("RadianceRT"));
-		// TODO: InterlockedCompareExchange() doesn't appear to work with 16-bit uint render target
-		//Desc.Format = PF_R16_UINT;
-		Desc.Format = PF_R32_UINT;
-		GRenderTargetPool.FindFreeElement(RHICmdList, Desc, SampleCountRT, TEXT("SampleCountRT"));
-		GRenderTargetPool.FindFreeElement(RHICmdList, Desc, PixelPositionRT, TEXT("PixelPositionRT"));
-		GRenderTargetPool.FindFreeElement(RHICmdList, Desc, RayCountPerPixelRT, TEXT("RayCountPerPixelRT"));
-
-		// Clear render targets
-		FUintVector4 BlackColor(0, 0, 0, 1);
-		FVector4 BlackColorF(0, 0, 0, 1);
-
-		RHICmdList.ClearUAVFloat(RadianceRT->GetRenderTargetItem().UAV, BlackColorF);
-		RHICmdList.ClearUAVUint(SampleCountRT->GetRenderTargetItem().UAV, BlackColor);
-		RHICmdList.ClearUAVUint(PixelPositionRT->GetRenderTargetItem().UAV, BlackColor);
-		RHICmdList.ClearUAVUint(RayCountPerPixelRT->GetRenderTargetItem().UAV, BlackColor);
+
+		if (ViewState->PathTracingRadianceRT)
+		{
+			// already got a buffer
+			RadianceRT = ViewState->PathTracingRadianceRT;
+		}
+		else
+		{
+			FPooledRenderTargetDesc Desc = Translate(SceneColorOutputTexture->Desc);
+			Desc.Flags &= ~(TexCreate_FastVRAM | TexCreate_Transient);
+			Desc.Format = PF_FloatRGBA;
+			GRenderTargetPool.FindFreeElement(RHICmdList, Desc, RadianceRT, TEXT("RadianceRT"));
+			ViewState->PathTracingRadianceRT = RadianceRT;
+		}
 
 		auto RayGenShader = GetGlobalShaderMap(FeatureLevel)->GetShader<FPathTracingRG>();
 
@@ -617,16 +573,13 @@ void FDeferredShadingSceneRenderer::RenderPathTracing(
 			//#dxr-todo: Set minimum tile size for mGPU
 			int32 TileSizeX = FIntPoint::DivideAndRoundUp(View.ViewRect.Size(), GPUCount).X;
 
-			FRHIUnorderedAccessView* UAVs[] = {
-					RadianceRT->GetRenderTargetItem().UAV,
-					SampleCountRT->GetRenderTargetItem().UAV,
-					PixelPositionRT->GetRenderTargetItem().UAV,
-					RayCountPerPixelRT->GetRenderTargetItem().UAV
-			};
-
 			{
 				// FIXME: cross-GPU fences aren't currently supported in the transition API, we need to figure it out.
 #if 0
+				FRHIUnorderedAccessView* UAVs[] = {
+						RadianceRT->GetRenderTargetItem().UAV
+				};
+
 				// Begin mGPU fence
 				FRHIGPUMask GPUMask = FRHIGPUMask::All();
 				FComputeFenceRHIRef Fence = RHICmdList.CreateComputeFence(TEXT("PathTracingRayGen_Fence_Begin"));
@@ -668,10 +621,7 @@ void FDeferredShadingSceneRenderer::RenderPathTracing(
 					ViewState->VarianceMipTreeDimensions,
 					*ViewState->VarianceMipTree,
 					TileOffset,
-					RadianceRT->GetRenderTargetItem().UAV,
-					SampleCountRT->GetRenderTargetItem().UAV,
-					PixelPositionRT->GetRenderTargetItem().UAV,
-					RayCountPerPixelRT->GetRenderTargetItem().UAV
+					RadianceRT->GetRenderTargetItem().UAV
 				);
 
 				int32 DispatchSizeX = FMath::Min<int32>(TileSizeX, View.ViewRect.Size().X - TileOffset.X);
@@ -692,10 +642,7 @@ void FDeferredShadingSceneRenderer::RenderPathTracing(
 
 				if (GPUIndex > 0)
 				{
-					RHICmdList.TransferTexture(RadianceRT->GetRenderTargetItem().TargetableTexture->GetTexture2D(), GPURect, 1, 0, true);
-					RHICmdList.TransferTexture(SampleCountRT->GetRenderTargetItem().TargetableTexture->GetTexture2D(), GPURect, 1, 0, true);
-					RHICmdList.TransferTexture(PixelPositionRT->GetRenderTargetItem().TargetableTexture->GetTexture2D(), GPURect, 1, 0, true);
-					RHICmdList.TransferTexture(RayCountPerPixelRT->GetRenderTargetItem().TargetableTexture->GetTexture2D(), GPURect, 1, 0, true);
+					RHICmdList.TransferTexture(RadianceRT->GetRenderTargetItem().TargetableTexture->GetTexture2D(), GPURect, GPUIndex, 0, true);
 				}
 			}
 
@@ -741,10 +688,7 @@ void FDeferredShadingSceneRenderer::RenderPathTracing(
 				ViewState->VarianceMipTreeDimensions,
 				*ViewState->VarianceMipTree,
 				TileOffset,
-				RadianceRT->GetRenderTargetItem().UAV,
-				SampleCountRT->GetRenderTargetItem().UAV,
-				PixelPositionRT->GetRenderTargetItem().UAV,
-				RayCountPerPixelRT->GetRenderTargetItem().UAV
+				RadianceRT->GetRenderTargetItem().UAV
 			);
 
 			int32 ViewWidth = View.ViewRect.Size().X;
@@ -762,87 +706,26 @@ void FDeferredShadingSceneRenderer::RenderPathTracing(
 
 		// Save RayTracingIndirect for compositing
 		RHICmdList.CopyToResolveTarget(RadianceRT->GetRenderTargetItem().TargetableTexture, RadianceRT->GetRenderTargetItem().ShaderResourceTexture, FResolveParams());
-		RHICmdList.CopyToResolveTarget(SampleCountRT->GetRenderTargetItem().TargetableTexture, SampleCountRT->GetRenderTargetItem().ShaderResourceTexture, FResolveParams());
-		RHICmdList.CopyToResolveTarget(PixelPositionRT->GetRenderTargetItem().TargetableTexture, PixelPositionRT->GetRenderTargetItem().ShaderResourceTexture, FResolveParams());
-		RHICmdList.CopyToResolveTarget(RayCountPerPixelRT->GetRenderTargetItem().TargetableTexture, RayCountPerPixelRT->GetRenderTargetItem().ShaderResourceTexture, FResolveParams());
 
 		// Single GPU for launching compute shaders
 		SCOPED_GPU_MASK(RHICmdList, FRHIGPUMask::FromIndex(MainGPUIndex));
 
+#if 0
 		if ((ViewState->PathTracingSPP % CVarPathTracingRayCountFrequency.GetValueOnRenderThread() == 0))
 		{
 			ComputeRayCount(RHICmdList, View, RayCountPerPixelRT->GetRenderTargetItem().ShaderResourceTexture);
 		}
-
-		// Run ray continuation compute shader
-		TRefCountPtr<IPooledRenderTarget> RadianceSortedRedRT;
-		TRefCountPtr<IPooledRenderTarget> RadianceSortedGreenRT;
-		TRefCountPtr<IPooledRenderTarget> RadianceSortedBlueRT;
-		TRefCountPtr<IPooledRenderTarget> RadianceSortedAlphaRT;
-		TRefCountPtr<IPooledRenderTarget> SampleCountSortedRT;
-
-		Desc.Format = PF_R32_UINT;
-		GRenderTargetPool.FindFreeElement(RHICmdList, Desc, RadianceSortedRedRT, TEXT("RadianceSortedRedRT"));
-		GRenderTargetPool.FindFreeElement(RHICmdList, Desc, RadianceSortedGreenRT, TEXT("RadianceSortedGreenRT"));
-		GRenderTargetPool.FindFreeElement(RHICmdList, Desc, RadianceSortedBlueRT, TEXT("RadianceSortedBlueRT"));
-		GRenderTargetPool.FindFreeElement(RHICmdList, Desc, RadianceSortedAlphaRT, TEXT("RadianceSortedAlphaRT"));
-		GRenderTargetPool.FindFreeElement(RHICmdList, Desc, SampleCountSortedRT, TEXT("SampleCountSortedRT"));
-
-		RHICmdList.ClearUAVUint(RadianceSortedRedRT->GetRenderTargetItem().UAV, BlackColor);
-		RHICmdList.ClearUAVUint(RadianceSortedGreenRT->GetRenderTargetItem().UAV, BlackColor);
-		RHICmdList.ClearUAVUint(RadianceSortedBlueRT->GetRenderTargetItem().UAV, BlackColor);
-		RHICmdList.ClearUAVUint(RadianceSortedAlphaRT->GetRenderTargetItem().UAV, BlackColor);
-		RHICmdList.ClearUAVUint(SampleCountSortedRT->GetRenderTargetItem().UAV, BlackColor);
-
-		if (NeedsMoreRays)
-		{
-			// only need to compact paths if we've traced rays in the last pass
-			ComputePathCompaction(
-				RHICmdList,
-				View,
-				RadianceRT->GetRenderTargetItem().ShaderResourceTexture,
-				SampleCountRT->GetRenderTargetItem().ShaderResourceTexture,
-				PixelPositionRT->GetRenderTargetItem().ShaderResourceTexture,
-				RadianceSortedRedRT->GetRenderTargetItem().UAV,
-				RadianceSortedGreenRT->GetRenderTargetItem().UAV,
-				RadianceSortedBlueRT->GetRenderTargetItem().UAV,
-				RadianceSortedAlphaRT->GetRenderTargetItem().UAV,
-				SampleCountSortedRT->GetRenderTargetItem().UAV
-			);
-
-			RHICmdList.CopyToResolveTarget(RadianceSortedRedRT->GetRenderTargetItem().TargetableTexture, RadianceSortedRedRT->GetRenderTargetItem().ShaderResourceTexture, FResolveParams());
-			RHICmdList.CopyToResolveTarget(RadianceSortedGreenRT->GetRenderTargetItem().TargetableTexture, RadianceSortedGreenRT->GetRenderTargetItem().ShaderResourceTexture, FResolveParams());
-			RHICmdList.CopyToResolveTarget(RadianceSortedBlueRT->GetRenderTargetItem().TargetableTexture, RadianceSortedBlueRT->GetRenderTargetItem().ShaderResourceTexture, FResolveParams());
-			RHICmdList.CopyToResolveTarget(RadianceSortedAlphaRT->GetRenderTargetItem().TargetableTexture, RadianceSortedAlphaRT->GetRenderTargetItem().ShaderResourceTexture, FResolveParams());
-			RHICmdList.CopyToResolveTarget(SampleCountSortedRT->GetRenderTargetItem().TargetableTexture, SampleCountSortedRT->GetRenderTargetItem().ShaderResourceTexture, FResolveParams());
-		}
-
-		// Construct render targets for compositing
-		TRefCountPtr<IPooledRenderTarget> OutputRadianceRT;
-		TRefCountPtr<IPooledRenderTarget> OutputSampleCountRT;
-		Desc.Flags &= ~(TexCreate_FastVRAM | TexCreate_Transient);
-		Desc.Format = PF_A32B32G32R32F;
-		//Desc.Format = PF_A16B16G16R16;
-		GRenderTargetPool.FindFreeElement(RHICmdList, Desc, OutputRadianceRT, TEXT("OutputRadianceRT"));
-		Desc.Format = PF_R16_UINT;
-		//Desc.Format = PF_R32_UINT;
-		GRenderTargetPool.FindFreeElement(RHICmdList, Desc, OutputSampleCountRT, TEXT("OutputSampleCountRT"));
-
-		RHICmdList.ClearUAVFloat(OutputRadianceRT->GetRenderTargetItem().UAV, BlackColorF);
-		RHICmdList.ClearUAVUint(OutputSampleCountRT->GetRenderTargetItem().UAV, BlackColor);
-
+#endif
 		// Run compositing engine
 		const auto ShaderMap = GetGlobalShaderMap(FeatureLevel);
 
 		TShaderMapRef<FPostProcessVS> VertexShader(ShaderMap);
 		TShaderMapRef<FPathTracingCompositorPS> PixelShader(ShaderMap);
-		FRHITexture* RenderTargets[3] =
+		FRHITexture* RenderTargets[1] =
 		{
-			SceneColorOutputTexture->GetRHI(),
-			OutputRadianceRT->GetRenderTargetItem().TargetableTexture,
-			OutputSampleCountRT->GetRenderTargetItem().TargetableTexture
+			SceneColorOutputTexture->GetRHI()
 		};
-		FRHIRenderPassInfo RenderPassInfo(3, RenderTargets, ERenderTargetActions::Load_Store);
+		FRHIRenderPassInfo RenderPassInfo(UE_ARRAY_COUNT(RenderTargets), RenderTargets, ERenderTargetActions::Load_Store);
 		RHICmdList.BeginRenderPass(RenderPassInfo, TEXT("PathTracing"));
 
 		// DEBUG: Inspect render target in isolation
@@ -859,26 +742,8 @@ void FDeferredShadingSceneRenderer::RenderPathTracing(
 
 		{
 			RHICmdList.SetViewport(View.ViewRect.Min.X, View.ViewRect.Min.Y, 0.0f, View.ViewRect.Max.X, View.ViewRect.Max.Y, 1.0f);
-			FTextureRHIRef RadianceRedTexture = RadianceSortedRedRT->GetRenderTargetItem().ShaderResourceTexture;
-			FTextureRHIRef RadianceGreenTexture = RadianceSortedGreenRT->GetRenderTargetItem().ShaderResourceTexture;
-			FTextureRHIRef RadianceBlueTexture = RadianceSortedBlueRT->GetRenderTargetItem().ShaderResourceTexture;
-			FTextureRHIRef RadianceAlphaTexture = RadianceSortedAlphaRT->GetRenderTargetItem().ShaderResourceTexture;
-			FTextureRHIRef SampleCountTexture = SampleCountSortedRT->GetRenderTargetItem().ShaderResourceTexture;
 
-			FTextureRHIRef CumulativeRadianceTexture = GBlackTexture->TextureRHI;
-			FTextureRHIRef CumulativeSampleCount = GBlackTexture->TextureRHI;
-
-			if (ViewState->PathTracingIrradianceRT)
-			{
-				// If we still have a valid accumulated image - re-use it
-				CumulativeRadianceTexture = ViewState->PathTracingIrradianceRT->GetRenderTargetItem().ShaderResourceTexture;
-				CumulativeSampleCount = ViewState->PathTracingSampleCountRT->GetRenderTargetItem().ShaderResourceTexture;
-			}
-			else
-			{
-				// Accumulated image got flushed by invalidating, start over
-			}
-			PixelShader->SetParameters(RHICmdList, View, RadianceRedTexture, RadianceGreenTexture, RadianceBlueTexture, RadianceAlphaTexture, SampleCountTexture, CumulativeRadianceTexture, CumulativeSampleCount);
+			PixelShader->SetParameters(RHICmdList, View, AdaptiveSamplingData.Iteration, RadianceRT->GetRenderTargetItem().ShaderResourceTexture);
 
 			int32 DispatchSizeX = View.ViewRect.Size().X;
 
@@ -897,28 +762,6 @@ void FDeferredShadingSceneRenderer::RenderPathTracing(
 		// Bump counters for next frame
 		++ViewState->PathTracingSPP;
 		++ViewState->PathTracingFrameIndependentTemporalSeed;
-
-		RHICmdList.CopyToResolveTarget(OutputRadianceRT->GetRenderTargetItem().TargetableTexture, OutputRadianceRT->GetRenderTargetItem().ShaderResourceTexture, FResolveParams());
-		RHICmdList.CopyToResolveTarget(OutputSampleCountRT->GetRenderTargetItem().TargetableTexture, OutputSampleCountRT->GetRenderTargetItem().ShaderResourceTexture, FResolveParams());
-		GVisualizeTexture.SetCheckPoint(RHICmdList, OutputRadianceRT);
-		GVisualizeTexture.SetCheckPoint(RHICmdList, OutputSampleCountRT);
-
-		// Cache values for reuse
-		ViewState->PathTracingIrradianceRT = OutputRadianceRT;
-		ViewState->PathTracingSampleCountRT = OutputSampleCountRT;
-
-		if (NeedsMoreRays)
-		{
-			// Process variance mip for adaptive sampling
-			if (ViewState->PathTracingSPP % CVarPathTracingVarianceMapRebuildFrequency.GetValueOnRenderThread() == 0)
-			{
-				SCOPED_GPU_STAT(RHICmdList, Stat_GPU_PathTracingBuildVarianceMipTree);
-
-				BuildVarianceMipTree(RHICmdList, View, OutputRadianceRT->GetRenderTargetItem().ShaderResourceTexture, *ViewState->VarianceMipTree, ViewState->VarianceMipTreeDimensions);
-			}
-		}
-
-		VisualizeVarianceMipTree(RHICmdList, View, *ViewState->VarianceMipTree, ViewState->VarianceMipTreeDimensions);
 	});
 }
 #endif
\ No newline at end of file
diff --git a/Engine/Source/Runtime/Renderer/Private/ScenePrivate.h b/Engine/Source/Runtime/Renderer/Private/ScenePrivate.h
--- a/Engine/Source/Runtime/Renderer/Private/ScenePrivate.h
+++ b/Engine/Source/Runtime/Renderer/Private/ScenePrivate.h
@@ -1008,8 +1008,7 @@ public:
 	void PathTracingInvalidate();
 
 	// Reference path tracing cached results
-	TRefCountPtr<IPooledRenderTarget> PathTracingIrradianceRT;
-	TRefCountPtr<IPooledRenderTarget> PathTracingSampleCountRT;
+	TRefCountPtr<IPooledRenderTarget> PathTracingRadianceRT;
 	FIntRect PathTracingRect;
 	FRWBuffer* VarianceMipTree;
 	FIntVector VarianceMipTreeDimensions;
