Refactored some functions out of scene renderer and into free functions.

#rnx

diff --git a/Engine/Source/Runtime/Renderer/Private/BasePassRendering.cpp b/Engine/Source/Runtime/Renderer/Private/BasePassRendering.cpp
--- a/Engine/Source/Runtime/Renderer/Private/BasePassRendering.cpp
+++ b/Engine/Source/Runtime/Renderer/Private/BasePassRendering.cpp
@@ -1110,12 +1110,12 @@ void FDeferredShadingSceneRenderer::RenderBasePassInternal(
 	if (bRenderLightmapDensity)
 	{
 		// Override the base pass with the lightmap density pass if the viewmode is enabled.
-		RenderLightMapDensities(GraphBuilder, BasePassRenderTargets);
+		RenderLightMapDensities(GraphBuilder, Views, Scene, BasePassRenderTargets);
 	}
 	else if (ViewFamily.UseDebugViewPS())
 	{
 		// Override the base pass with one of the debug view shader mode (see EDebugViewShaderMode) if required.
-		RenderDebugViewMode(GraphBuilder, BasePassRenderTargets);
+		RenderDebugViewMode(GraphBuilder, Views, Scene, BasePassRenderTargets);
 	}
 	else
 	{
diff --git a/Engine/Source/Runtime/Renderer/Private/DebugViewModeRendering.cpp b/Engine/Source/Runtime/Renderer/Private/DebugViewModeRendering.cpp
--- a/Engine/Source/Runtime/Renderer/Private/DebugViewModeRendering.cpp
+++ b/Engine/Source/Runtime/Renderer/Private/DebugViewModeRendering.cpp
@@ -91,13 +91,13 @@ BEGIN_SHADER_PARAMETER_STRUCT(FDebugViewModePassParameters, )
 	RENDER_TARGET_BINDING_SLOTS()
 END_SHADER_PARAMETER_STRUCT()
 
-void FDeferredShadingSceneRenderer::RenderDebugViewMode(FRDGBuilder& GraphBuilder, const FRenderTargetBindingSlots& RenderTargets)
+void RenderDebugViewMode(FRDGBuilder& GraphBuilder, TArrayView<const FViewInfo> Views, FScene* Scene, const FRenderTargetBindingSlots& RenderTargets)
 {
 	RDG_EVENT_SCOPE(GraphBuilder, "DebugViewMode");
 
 	for (int32 ViewIndex = 0; ViewIndex < Views.Num(); ++ViewIndex)
 	{
-		FViewInfo& View = Views[ViewIndex];
+		const FViewInfo& View = Views[ViewIndex];
 		RDG_GPU_MASK_SCOPE(GraphBuilder, View.GPUMask);
 		RDG_EVENT_SCOPE_CONDITIONAL(GraphBuilder, Views.Num() > 1, "View%d", ViewIndex);
 
@@ -109,7 +109,7 @@ void FDeferredShadingSceneRenderer::RenderDebugViewMode(FRDGBuilder& GraphBuilde
 			{},
 			PassParameters,
 			ERDGPassFlags::Raster,
-			[this, &View](FRHICommandList& RHICmdList)
+			[Scene, &View](FRHICommandList& RHICmdList)
 		{
 			Scene->UniformBuffers.UpdateViewUniformBuffer(View);
 			RHICmdList.SetViewport(View.ViewRect.Min.X, View.ViewRect.Min.Y, 0, View.ViewRect.Max.X, View.ViewRect.Max.Y, 1);
@@ -377,7 +377,7 @@ void InitDebugViewModeInterfaces()
 
 #else // !(UE_BUILD_SHIPPING || UE_BUILD_TEST)
 
-void FDeferredShadingSceneRenderer::RenderDebugViewMode(FRDGBuilder& GraphBuilder, const FRenderTargetBindingSlots& RenderTargets) {}
+void RenderDebugViewMode(FRDGBuilder& GraphBuilder, TArrayView<const FViewInfo> Views, const FRenderTargetBindingSlots& RenderTargets) {}
 
 #endif // !(UE_BUILD_SHIPPING || UE_BUILD_TEST)
 
diff --git a/Engine/Source/Runtime/Renderer/Private/DebugViewModeRendering.h b/Engine/Source/Runtime/Renderer/Private/DebugViewModeRendering.h
--- a/Engine/Source/Runtime/Renderer/Private/DebugViewModeRendering.h
+++ b/Engine/Source/Runtime/Renderer/Private/DebugViewModeRendering.h
@@ -219,4 +219,10 @@ void AddDebugViewModeShaderTypes(ERHIFeatureLevel::Type FeatureLevel,
 	const FVertexFactoryType* VertexFactoryType,
 	FMaterialShaderTypes& OutShaderTypes);
 
+void RenderDebugViewMode(
+	FRDGBuilder& GraphBuilder,
+	TArrayView<const FViewInfo> Views,
+	FScene* Scene,
+	const FRenderTargetBindingSlots& RenderTargets);
+
 #endif // !(UE_BUILD_SHIPPING || UE_BUILD_TEST)
\ No newline at end of file
diff --git a/Engine/Source/Runtime/Renderer/Private/DeferredShadingRenderer.cpp b/Engine/Source/Runtime/Renderer/Private/DeferredShadingRenderer.cpp
--- a/Engine/Source/Runtime/Renderer/Private/DeferredShadingRenderer.cpp
+++ b/Engine/Source/Runtime/Renderer/Private/DeferredShadingRenderer.cpp
@@ -2053,7 +2053,7 @@ void FDeferredShadingSceneRenderer::Render(FRHICommandListImmediate& RHICmdList)
 
 	if (bVolumetricRenderTargetRequired)
 	{
-		InitVolumetricRenderTargetForViews(GraphBuilder);
+		InitVolumetricRenderTargetForViews(GraphBuilder, Views);
 	}
 
 	if (bShouldRenderVolumetricCloud)
@@ -2346,7 +2346,7 @@ void FDeferredShadingSceneRenderer::Render(FRHICommandListImmediate& RHICmdList)
 #endif // RHI_RAYTRACING
 
 	// Copy lighting channels out of stencil before deferred decals which overwrite those values
-	FRDGTextureRef LightingChannelsTexture = CopyStencilToLightingChannelTexture(GraphBuilder, SceneStencilTexture);
+	FRDGTextureRef LightingChannelsTexture = CopyStencilToLightingChannelTexture(GraphBuilder, Views, SceneStencilTexture);
 
 	if (IsForwardShadingEnabled(ShaderPlatform))
 	{
@@ -2509,10 +2509,11 @@ void FDeferredShadingSceneRenderer::Render(FRHICommandListImmediate& RHICmdList)
 		RenderDeferredReflectionsAndSkyLightingHair(GraphBuilder, HairDatas);
 	}
 
+	FRDGTextureRef HalfResolutionDepthCheckerboardMinMaxTexture = nullptr;
+
 	if (bShouldRenderVolumetricCloud && IsVolumetricRenderTargetEnabled())
 	{
-		// The checkerboarded half resolution depth texture will be needed.
-		UpdateHalfResDepthSurfaceCheckerboardMinMax(GraphBuilder, SceneDepthTexture.Resolve);
+		HalfResolutionDepthCheckerboardMinMaxTexture = CreateHalfResolutionDepthCheckerboardMinMax(GraphBuilder, Views, SceneDepthTexture.Resolve);
 	}
 
 	if (bShouldRenderVolumetricCloud)
@@ -2520,9 +2521,10 @@ void FDeferredShadingSceneRenderer::Render(FRHICommandListImmediate& RHICmdList)
 		// Generate the volumetric cloud render target
 		bool bSkipVolumetricRenderTarget = false;
 		bool bSkipPerPixelTracing = true;
-		RenderVolumetricCloud(GraphBuilder, GetSceneTextureShaderParameters(SceneTextures), bSkipVolumetricRenderTarget, bSkipPerPixelTracing, SceneColorTexture, SceneDepthTexture);
+		RenderVolumetricCloud(GraphBuilder, GetSceneTextureShaderParameters(SceneTextures), bSkipVolumetricRenderTarget, bSkipPerPixelTracing, SceneColorTexture, SceneDepthTexture, HalfResolutionDepthCheckerboardMinMaxTexture);
+
 		// Reconstruct the volumetric cloud render target to be ready to compose it over the scene
-		ReconstructVolumetricRenderTarget(GraphBuilder);
+		ReconstructVolumetricRenderTarget(GraphBuilder, Views, SceneDepthTexture.Resolve, HalfResolutionDepthCheckerboardMinMaxTexture);
 	}
 
 	const bool bShouldRenderTranslucency = ShouldRenderTranslucency();
@@ -2587,12 +2589,13 @@ void FDeferredShadingSceneRenderer::Render(FRHICommandListImmediate& RHICmdList)
 	{
 		bool bSkipVolumetricRenderTarget = true;
 		bool bSkipPerPixelTracing = false;
-		RenderVolumetricCloud(GraphBuilder, GetSceneTextureShaderParameters(SceneTextures), bSkipVolumetricRenderTarget, bSkipPerPixelTracing, SceneColorTexture, SceneDepthTexture);
+		RenderVolumetricCloud(GraphBuilder, GetSceneTextureShaderParameters(SceneTextures), bSkipVolumetricRenderTarget, bSkipPerPixelTracing, SceneColorTexture, SceneDepthTexture, HalfResolutionDepthCheckerboardMinMaxTexture);
 	}
+
 	// or composite the off screen buffer over the scene.
 	if (bVolumetricRenderTargetRequired)
 	{
-		ComposeVolumetricRenderTargetOverScene(GraphBuilder, SceneColorTexture.Target, SceneDepthTexture.Target, bShouldRenderSingleLayerWater, SceneWithoutWaterTextures);
+		ComposeVolumetricRenderTargetOverScene(GraphBuilder, Views, SceneColorTexture.Target, SceneDepthTexture.Target, bShouldRenderSingleLayerWater, SceneWithoutWaterTextures);
 	}
 
 	FRendererModule& RendererModule = static_cast<FRendererModule&>(GetRendererModule());
@@ -2924,103 +2927,6 @@ void FDeferredShadingSceneRenderer::Render(FRHICommandListImmediate& RHICmdList)
 	ServiceLocalQueue();
 }
 
-/** Updates the downsized depth buffer with the current full resolution depth buffer. */
-void FDeferredShadingSceneRenderer::UpdateHalfResDepthSurfaceCheckerboardMinMax(FRDGBuilder& GraphBuilder, FRDGTextureRef SceneDepthTexture)
-{
-	const uint32 DownscaleFactor = 2;
-	const FIntPoint SmallDepthExtent = GetDownscaledExtent(SceneDepthTexture->Desc.Extent, DownscaleFactor);
-	const FRDGTextureDesc SmallDepthDesc = FRDGTextureDesc::Create2D(SmallDepthExtent, PF_DepthStencil, FClearValueBinding::None, TexCreate_DepthStencilTargetable | TexCreate_ShaderResource);
-	FRDGTextureRef SmallDepthTexture = GraphBuilder.CreateTexture(SmallDepthDesc, TEXT("HalfResDepthSurfaceCheckerboardMinMax"));
-
-	for (FViewInfo& View : Views)
-	{
-		RDG_GPU_MASK_SCOPE(GraphBuilder, View.GPUMask);
-
-		const FScreenPassTexture SceneDepth(SceneDepthTexture, View.ViewRect);
-		const FScreenPassRenderTarget SmallDepth(SmallDepthTexture, GetDownscaledRect(View.ViewRect, DownscaleFactor), View.DecayLoadAction(ERenderTargetLoadAction::ENoAction));
-		AddDownsampleDepthPass(GraphBuilder, View, SceneDepth, SmallDepth, EDownsampleDepthFilter::Checkerboard);
-	}
-
-	TRefCountPtr<IPooledRenderTarget> SmallDepthTarget;
-	ConvertToUntrackedExternalTexture(GraphBuilder, SmallDepthTexture, SmallDepthTarget, ERHIAccess::SRVMask);
-	for (FViewInfo& View : Views)
-	{
-		View.HalfResDepthSurfaceCheckerboardMinMax = SmallDepthTarget;
-	}
-}
-
-class FCopyStencilToLightingChannelsPS : public FGlobalShader
-{
-public:
-	DECLARE_GLOBAL_SHADER(FCopyStencilToLightingChannelsPS);
-	SHADER_USE_PARAMETER_STRUCT(FCopyStencilToLightingChannelsPS, FGlobalShader);
-
-	BEGIN_SHADER_PARAMETER_STRUCT(FParameters, )
-		SHADER_PARAMETER_STRUCT_REF(FViewUniformShaderParameters, View)
-		SHADER_PARAMETER_RDG_TEXTURE_SRV(Texture2D<float>, SceneStencilTexture)
-		RENDER_TARGET_BINDING_SLOTS()
-	END_SHADER_PARAMETER_STRUCT()
-
-	static bool ShouldCompilePermutation(const FGlobalShaderPermutationParameters& Parameters)
-	{ 
-		return IsFeatureLevelSupported(Parameters.Platform, ERHIFeatureLevel::SM5);
-	}
-
-	static void ModifyCompilationEnvironment(const FGlobalShaderPermutationParameters& Parameters, FShaderCompilerEnvironment& OutEnvironment)
-	{
-		FGlobalShader::ModifyCompilationEnvironment(Parameters, OutEnvironment);
-		OutEnvironment.SetDefine(TEXT("STENCIL_LIGHTING_CHANNELS_SHIFT"), STENCIL_LIGHTING_CHANNELS_BIT_ID);
-		OutEnvironment.SetRenderTargetOutputFormat(0, PF_R16_UINT);
-	}
-};
-
-IMPLEMENT_GLOBAL_SHADER(FCopyStencilToLightingChannelsPS, "/Engine/Private/DownsampleDepthPixelShader.usf", "CopyStencilToLightingChannelsPS", SF_Pixel);
-
-FRDGTextureRef FDeferredShadingSceneRenderer::CopyStencilToLightingChannelTexture(FRDGBuilder& GraphBuilder, FRDGTextureSRVRef SceneStencilTexture)
-{
-	bool bAnyViewUsesLightingChannels = false;
-
-	for (int32 ViewIndex = 0, ViewCount = Views.Num(); ViewIndex < ViewCount; ++ViewIndex)
-	{
-		bAnyViewUsesLightingChannels = bAnyViewUsesLightingChannels || Views[ViewIndex].bUsesLightingChannels;
-	}
-
-	FRDGTextureRef LightingChannelsTexture = nullptr;
-
-	if (bAnyViewUsesLightingChannels)
-	{
-		RDG_EVENT_SCOPE(GraphBuilder, "CopyStencilToLightingChannels");
-
-		{
-			FSceneRenderTargets& SceneContext = FSceneRenderTargets::Get();
-
-			const FRDGTextureDesc Desc = FRDGTextureDesc::Create2D(SceneContext.GetBufferSizeXY(), PF_R16_UINT, FClearValueBinding::None, TexCreate_RenderTargetable | TexCreate_ShaderResource);
-			LightingChannelsTexture = GraphBuilder.CreateTexture(Desc, TEXT("LightingChannels"));
-		}
-
-		const ERenderTargetLoadAction LoadAction = ERenderTargetLoadAction::ENoAction;
-
-		for (int32 ViewIndex = 0, ViewCount = Views.Num(); ViewIndex < ViewCount; ++ViewIndex)
-		{
-			const FViewInfo& View = Views[ViewIndex];
-			RDG_EVENT_SCOPE_CONDITIONAL(GraphBuilder, Views.Num() > 1, "View%d", ViewIndex);
-			RDG_GPU_MASK_SCOPE(GraphBuilder, View.GPUMask);
-
-			auto* PassParameters = GraphBuilder.AllocParameters<FCopyStencilToLightingChannelsPS::FParameters>();
-			PassParameters->RenderTargets[0] = FRenderTargetBinding(LightingChannelsTexture, View.DecayLoadAction(LoadAction));
-			PassParameters->SceneStencilTexture = SceneStencilTexture;
-			PassParameters->View = View.ViewUniformBuffer;
-
-			const FScreenPassTextureViewport Viewport(LightingChannelsTexture, View.ViewRect);
-
-			TShaderMapRef<FCopyStencilToLightingChannelsPS> PixelShader(View.ShaderMap);
-			AddDrawScreenPass(GraphBuilder, {}, View, Viewport, Viewport, PixelShader, PassParameters);
-		}
-	}
-
-	return LightingChannelsTexture;
-}
-
 #if RHI_RAYTRACING
 
 bool AnyRayTracingPassEnabled(const FScene* Scene, const FViewInfo& View)
diff --git a/Engine/Source/Runtime/Renderer/Private/DeferredShadingRenderer.h b/Engine/Source/Runtime/Renderer/Private/DeferredShadingRenderer.h
--- a/Engine/Source/Runtime/Renderer/Private/DeferredShadingRenderer.h
+++ b/Engine/Source/Runtime/Renderer/Private/DeferredShadingRenderer.h
@@ -655,17 +655,6 @@ private:
 	bool ShouldRenderDistortion() const;
 	void RenderDistortion(FRDGBuilder& GraphBuilder, FRDGTextureRef SceneColorTexture, FRDGTextureRef SceneDepthTexture);
 
-	/** Renders world-space lightmap density instead of the normal color. */
-	void RenderLightMapDensities(FRDGBuilder& GraphBuilder, const FRenderTargetBindingSlots& RenderTargets);
-
-	/** Renders one of the EDebugViewShaderMode instead of the normal color. */
-	void RenderDebugViewMode(FRDGBuilder& GraphBuilder, const FRenderTargetBindingSlots& RenderTargets);
-
-	/** Updates the downsized depth buffer with the current full resolution depth buffer using a min/max checkerboard pattern. */
-	void UpdateHalfResDepthSurfaceCheckerboardMinMax(FRDGBuilder& GraphBuilder, FRDGTextureRef SceneDepthTexture);
-
-	FRDGTextureRef CopyStencilToLightingChannelTexture(FRDGBuilder& GraphBuilder, FRDGTextureSRVRef SceneStencilTexture);
-
 	/** Renders capsule shadows for all per-object shadows using it for the given light. */
 	bool RenderCapsuleDirectShadows(
 		FRDGBuilder& GraphBuilder,
diff --git a/Engine/Source/Runtime/Renderer/Private/LightMapDensityRendering.cpp b/Engine/Source/Runtime/Renderer/Private/LightMapDensityRendering.cpp
--- a/Engine/Source/Runtime/Renderer/Private/LightMapDensityRendering.cpp
+++ b/Engine/Source/Runtime/Renderer/Private/LightMapDensityRendering.cpp
@@ -68,18 +68,22 @@ BEGIN_SHADER_PARAMETER_STRUCT(FLightMapDensitiesPassParameters, )
 	RENDER_TARGET_BINDING_SLOTS()
 END_SHADER_PARAMETER_STRUCT()
 
-void FDeferredShadingSceneRenderer::RenderLightMapDensities(FRDGBuilder& GraphBuilder, const FRenderTargetBindingSlots& RenderTargets)
+void RenderLightMapDensities(
+	FRDGBuilder& GraphBuilder,
+	TArrayView<const FViewInfo> Views,
+	FScene* Scene,
+	const FRenderTargetBindingSlots& RenderTargets)
 {
 	RDG_EVENT_SCOPE(GraphBuilder, "LightMapDensity");
 
 	auto* PassParameters = GraphBuilder.AllocParameters<FLightMapDensitiesPassParameters>();
-	PassParameters->Pass = CreateLightmapDensityPassUniformBuffer(GraphBuilder, FeatureLevel);
+	PassParameters->Pass = CreateLightmapDensityPassUniformBuffer(GraphBuilder, Scene->GetFeatureLevel());
 	PassParameters->RenderTargets = RenderTargets;
 
 	// Draw the scene's emissive and light-map color.
 	for (int32 ViewIndex = 0; ViewIndex < Views.Num(); ++ViewIndex)
 	{
-		FViewInfo& View = Views[ViewIndex];
+		const FViewInfo& View = Views[ViewIndex];
 		RDG_GPU_MASK_SCOPE(GraphBuilder, View.GPUMask);
 		RDG_EVENT_SCOPE_CONDITIONAL(GraphBuilder, Views.Num() > 1, "View%d", ViewIndex);
 
@@ -87,7 +91,7 @@ void FDeferredShadingSceneRenderer::RenderLightMapDensities(FRDGBuilder& GraphBu
 			{},
 			PassParameters,
 			ERDGPassFlags::Raster,
-			[this, &View](FRHICommandList& RHICmdList)
+			[Scene, &View](FRHICommandList& RHICmdList)
 		{
 			Scene->UniformBuffers.UpdateViewUniformBuffer(View);
 			RHICmdList.SetViewport(View.ViewRect.Min.X, View.ViewRect.Min.Y, 0, View.ViewRect.Max.X, View.ViewRect.Max.Y, 1);
diff --git a/Engine/Source/Runtime/Renderer/Private/LightMapDensityRendering.h b/Engine/Source/Runtime/Renderer/Private/LightMapDensityRendering.h
--- a/Engine/Source/Runtime/Renderer/Private/LightMapDensityRendering.h
+++ b/Engine/Source/Runtime/Renderer/Private/LightMapDensityRendering.h
@@ -270,4 +270,10 @@ private:
 		ERasterizerCullMode MeshCullMode);
 
 	FMeshPassProcessorRenderState PassDrawRenderState;
-};
\ No newline at end of file
+};
+
+void RenderLightMapDensities(
+	FRDGBuilder& GraphBuilder,
+	TArrayView<const FViewInfo> Views,
+	FScene* Scene,
+	const FRenderTargetBindingSlots& RenderTargets);
\ No newline at end of file
diff --git a/Engine/Source/Runtime/Renderer/Private/LightRendering.cpp b/Engine/Source/Runtime/Renderer/Private/LightRendering.cpp
--- a/Engine/Source/Runtime/Renderer/Private/LightRendering.cpp
+++ b/Engine/Source/Runtime/Renderer/Private/LightRendering.cpp
@@ -2616,3 +2616,75 @@ void FDeferredShadingSceneRenderer::RenderSimpleLightsStandardDeferred(
 		}
 	});
 }
+
+class FCopyStencilToLightingChannelsPS : public FGlobalShader
+{
+public:
+	DECLARE_GLOBAL_SHADER(FCopyStencilToLightingChannelsPS);
+	SHADER_USE_PARAMETER_STRUCT(FCopyStencilToLightingChannelsPS, FGlobalShader);
+
+	BEGIN_SHADER_PARAMETER_STRUCT(FParameters, )
+		SHADER_PARAMETER_STRUCT_REF(FViewUniformShaderParameters, View)
+		SHADER_PARAMETER_RDG_TEXTURE_SRV(Texture2D<float>, SceneStencilTexture)
+		RENDER_TARGET_BINDING_SLOTS()
+	END_SHADER_PARAMETER_STRUCT()
+
+	static bool ShouldCompilePermutation(const FGlobalShaderPermutationParameters& Parameters)
+	{
+		return IsFeatureLevelSupported(Parameters.Platform, ERHIFeatureLevel::SM5);
+	}
+
+	static void ModifyCompilationEnvironment(const FGlobalShaderPermutationParameters& Parameters, FShaderCompilerEnvironment& OutEnvironment)
+	{
+		FGlobalShader::ModifyCompilationEnvironment(Parameters, OutEnvironment);
+		OutEnvironment.SetDefine(TEXT("STENCIL_LIGHTING_CHANNELS_SHIFT"), STENCIL_LIGHTING_CHANNELS_BIT_ID);
+		OutEnvironment.SetRenderTargetOutputFormat(0, PF_R16_UINT);
+	}
+};
+
+IMPLEMENT_GLOBAL_SHADER(FCopyStencilToLightingChannelsPS, "/Engine/Private/DownsampleDepthPixelShader.usf", "CopyStencilToLightingChannelsPS", SF_Pixel);
+
+FRDGTextureRef CopyStencilToLightingChannelTexture(FRDGBuilder& GraphBuilder, TArrayView<const FViewInfo> Views, FRDGTextureSRVRef SceneStencilTexture)
+{
+	bool bAnyViewUsesLightingChannels = false;
+
+	for (int32 ViewIndex = 0, ViewCount = Views.Num(); ViewIndex < ViewCount; ++ViewIndex)
+	{
+		bAnyViewUsesLightingChannels = bAnyViewUsesLightingChannels || Views[ViewIndex].bUsesLightingChannels;
+	}
+
+	FRDGTextureRef LightingChannelsTexture = nullptr;
+
+	if (bAnyViewUsesLightingChannels)
+	{
+		RDG_EVENT_SCOPE(GraphBuilder, "CopyStencilToLightingChannels");
+
+		{
+			check(SceneStencilTexture && SceneStencilTexture->Desc.Texture);
+			const FIntPoint TextureExtent = SceneStencilTexture->Desc.Texture->Desc.Extent;
+			const FRDGTextureDesc Desc = FRDGTextureDesc::Create2D(TextureExtent, PF_R16_UINT, FClearValueBinding::None, TexCreate_RenderTargetable | TexCreate_ShaderResource);
+			LightingChannelsTexture = GraphBuilder.CreateTexture(Desc, TEXT("LightingChannels"));
+		}
+
+		const ERenderTargetLoadAction LoadAction = ERenderTargetLoadAction::ENoAction;
+
+		for (int32 ViewIndex = 0, ViewCount = Views.Num(); ViewIndex < ViewCount; ++ViewIndex)
+		{
+			const FViewInfo& View = Views[ViewIndex];
+			RDG_EVENT_SCOPE_CONDITIONAL(GraphBuilder, Views.Num() > 1, "View%d", ViewIndex);
+			RDG_GPU_MASK_SCOPE(GraphBuilder, View.GPUMask);
+
+			auto* PassParameters = GraphBuilder.AllocParameters<FCopyStencilToLightingChannelsPS::FParameters>();
+			PassParameters->RenderTargets[0] = FRenderTargetBinding(LightingChannelsTexture, View.DecayLoadAction(LoadAction));
+			PassParameters->SceneStencilTexture = SceneStencilTexture;
+			PassParameters->View = View.ViewUniformBuffer;
+
+			const FScreenPassTextureViewport Viewport(LightingChannelsTexture, View.ViewRect);
+
+			TShaderMapRef<FCopyStencilToLightingChannelsPS> PixelShader(View.ShaderMap);
+			AddDrawScreenPass(GraphBuilder, {}, View, Viewport, Viewport, PixelShader, PassParameters);
+		}
+	}
+
+	return LightingChannelsTexture;
+}
\ No newline at end of file
diff --git a/Engine/Source/Runtime/Renderer/Private/LightRendering.h b/Engine/Source/Runtime/Renderer/Private/LightRendering.h
--- a/Engine/Source/Runtime/Renderer/Private/LightRendering.h
+++ b/Engine/Source/Runtime/Renderer/Private/LightRendering.h
@@ -505,3 +505,4 @@ enum class FLightOcclusionType : uint8
 FLightOcclusionType GetLightOcclusionType(const FLightSceneProxy& Proxy);
 FLightOcclusionType GetLightOcclusionType(const FLightSceneInfoCompact& LightInfo);
 
+FRDGTextureRef CopyStencilToLightingChannelTexture(FRDGBuilder& GraphBuilder, TArrayView<const FViewInfo> Views, FRDGTextureSRVRef SceneStencilTexture);
\ No newline at end of file
diff --git a/Engine/Source/Runtime/Renderer/Private/ReflectionEnvironmentRealTimeCapture.cpp b/Engine/Source/Runtime/Renderer/Private/ReflectionEnvironmentRealTimeCapture.cpp
--- a/Engine/Source/Runtime/Renderer/Private/ReflectionEnvironmentRealTimeCapture.cpp
+++ b/Engine/Source/Runtime/Renderer/Private/ReflectionEnvironmentRealTimeCapture.cpp
@@ -396,7 +396,7 @@ void FScene::AllocateAndCaptureFrameSkyEnvMap(
 					FMaterialRenderProxy* CloudVolumeMaterialProxy = CloudSceneProxy.GetCloudVolumeMaterial()->GetRenderProxy();
 					CloudRC.CloudInfo = &CloudInfo;
 					CloudRC.CloudVolumeMaterialProxy = CloudVolumeMaterialProxy;
-					CloudRC.SceneDepthZ = GSystemTextures.MaxFP16Depth;
+					CloudRC.SceneDepthZ = GSystemTextures.GetMaxFP16Depth(GraphBuilder);
 
 					CloudRC.MainView = &CubeView; /// This is only accessing data that is not changing between view oerientation. Such data are accessed from the ViewUniformBuffer. See CubeView comment above.
 
diff --git a/Engine/Source/Runtime/Renderer/Private/SceneRendering.cpp b/Engine/Source/Runtime/Renderer/Private/SceneRendering.cpp
--- a/Engine/Source/Runtime/Renderer/Private/SceneRendering.cpp
+++ b/Engine/Source/Runtime/Renderer/Private/SceneRendering.cpp
@@ -4652,6 +4652,25 @@ void AddResolveSceneDepthPass(FRDGBuilder& GraphBuilder, TArrayView<const FViewI
 	}
 }
 
+FRDGTextureRef CreateHalfResolutionDepthCheckerboardMinMax(FRDGBuilder& GraphBuilder, TArrayView<const FViewInfo> Views, FRDGTextureRef SceneDepthTexture)
+{
+	const uint32 DownscaleFactor = 2;
+	const FIntPoint SmallDepthExtent = GetDownscaledExtent(SceneDepthTexture->Desc.Extent, DownscaleFactor);
+	const FRDGTextureDesc SmallDepthDesc = FRDGTextureDesc::Create2D(SmallDepthExtent, PF_DepthStencil, FClearValueBinding::None, TexCreate_DepthStencilTargetable | TexCreate_ShaderResource);
+	FRDGTextureRef SmallDepthTexture = GraphBuilder.CreateTexture(SmallDepthDesc, TEXT("HalfResolutionDepthCheckerboardMinMax"));
+
+	for (const FViewInfo& View : Views)
+	{
+		RDG_GPU_MASK_SCOPE(GraphBuilder, View.GPUMask);
+
+		const FScreenPassTexture SceneDepth(SceneDepthTexture, View.ViewRect);
+		const FScreenPassRenderTarget SmallDepth(SmallDepthTexture, GetDownscaledRect(View.ViewRect, DownscaleFactor), View.DecayLoadAction(ERenderTargetLoadAction::ENoAction));
+		AddDownsampleDepthPass(GraphBuilder, View, SceneDepth, SmallDepth, EDownsampleDepthFilter::Checkerboard);
+	}
+
+	return SmallDepthTexture;
+}
+
 void FSceneRenderer::ResolveSceneDepth(FRHICommandListImmediate& RHICmdList)
 {
 	FSceneRenderTargets& SceneContext = FSceneRenderTargets::Get();
diff --git a/Engine/Source/Runtime/Renderer/Private/SceneRendering.h b/Engine/Source/Runtime/Renderer/Private/SceneRendering.h
--- a/Engine/Source/Runtime/Renderer/Private/SceneRendering.h
+++ b/Engine/Source/Runtime/Renderer/Private/SceneRendering.h
@@ -1292,8 +1292,6 @@ public:
 	float FurthestReflectionCaptureDistance;
 	TUniformBufferRef<FReflectionCaptureShaderData> ReflectionCaptureUniformBuffer;
 
-	TRefCountPtr<IPooledRenderTarget> HalfResDepthSurfaceCheckerboardMinMax;
-
 	// Sky / Atmosphere textures (transient owned by this view info) and pointer to constants owned by SkyAtmosphere proxy.
 	TRefCountPtr<IPooledRenderTarget> SkyAtmosphereCameraAerialPerspectiveVolume;
 	TRefCountPtr<IPooledRenderTarget> SkyAtmosphereViewLutTexture;
@@ -2007,7 +2005,8 @@ protected:
 		bool bSkipVolumetricRenderTarget,
 		bool bSkipPerPixelTracing,
 		FRDGTextureMSAA SceneColorTexture,
-		FRDGTextureMSAA SceneDepthTexture);
+		FRDGTextureMSAA SceneDepthTexture,
+		FRDGTextureRef HalfResolutionDepthCheckerboardMinMaxTexture);
 
 	/** Render notification to artist when a sky material is used but it might comtains the camera (and then the sky/background would look black).*/
 	void RenderSkyAtmosphereEditorNotifications(FRDGBuilder& GraphBuilder, FRDGTextureRef SceneColorTexture);
@@ -2015,15 +2014,6 @@ protected:
 	/** We should render on screen notification only if any of the scene contains a mesh using a sky material.*/
 	bool ShouldRenderSkyAtmosphereEditorNotifications();
 
-	/** Initialise volumetric render target.*/
-	void InitVolumetricRenderTargetForViews(FRDGBuilder& GraphBuilder);
-	/** Process the volumetric render target, generating the high resolution version.*/
-	void ReconstructVolumetricRenderTarget(FRDGBuilder& GraphBuilder);
-	/** Compose the volumetric render target over the scene.*/
-	void ComposeVolumetricRenderTargetOverScene(FRDGBuilder& GraphBuilder, FRDGTextureRef SceneColorTexture, FRDGTextureRef SceneDepthResolveTexture, bool bShouldRenderSingleLayerWater, const FSceneWithoutWaterTextures& WaterPassData);
-	/** Compose the volumetric render target over the scene from a view under water, in the water render target.*/
-	void ComposeVolumetricRenderTargetOverSceneUnderWater(FRDGBuilder& GraphBuilder, const FSceneWithoutWaterTextures& WaterPassData);
-
 	void ResolveSceneColor(FRHICommandListImmediate& RHICmdList);
 	void ResolveSceneDepth(FRHICommandListImmediate& RHICmdList);
 
@@ -2328,4 +2318,7 @@ void AddResolveSceneDepthPass(FRDGBuilder& GraphBuilder, const FViewInfo& View,
 
 /** Resolves all views for scene color / depth. */
 void AddResolveSceneColorPass(FRDGBuilder& GraphBuilder, TArrayView<const FViewInfo> Views, FRDGTextureMSAA SceneColor);
-void AddResolveSceneDepthPass(FRDGBuilder& GraphBuilder, TArrayView<const FViewInfo> Views, FRDGTextureMSAA SceneDepth);
\ No newline at end of file
+void AddResolveSceneDepthPass(FRDGBuilder& GraphBuilder, TArrayView<const FViewInfo> Views, FRDGTextureMSAA SceneDepth);
+
+/** Creates a half resolution checkerboard min / max depth buffer from the input full resolution depth buffer. */
+FRDGTextureRef CreateHalfResolutionDepthCheckerboardMinMax(FRDGBuilder& GraphBuilder, TArrayView<const FViewInfo> Views, FRDGTextureRef SceneDepth);
\ No newline at end of file
diff --git a/Engine/Source/Runtime/Renderer/Private/SingleLayerWaterRendering.cpp b/Engine/Source/Runtime/Renderer/Private/SingleLayerWaterRendering.cpp
--- a/Engine/Source/Runtime/Renderer/Private/SingleLayerWaterRendering.cpp
+++ b/Engine/Source/Runtime/Renderer/Private/SingleLayerWaterRendering.cpp
@@ -10,6 +10,7 @@
 #include "TemporalAA.h"
 #include "RayTracing/RaytracingOptions.h"
 #include "RayTracing/RayTracingReflections.h"
+#include "VolumetricRenderTarget.h"
 #include "RenderGraph.h"
 #include "ScenePrivate.h"
 #include "SceneRendering.h"
@@ -662,7 +663,7 @@ void FDeferredShadingSceneRenderer::RenderSingleLayerWater(
 	if (SceneWithoutWaterTextures.ColorTexture && bShouldRenderVolumetricCloud)
 	{
 		// This path is only taken when rendering the clouds in a render target that can be composited
-		ComposeVolumetricRenderTargetOverSceneUnderWater(GraphBuilder, SceneWithoutWaterTextures);
+		ComposeVolumetricRenderTargetOverSceneUnderWater(GraphBuilder, Views, SceneWithoutWaterTextures);
 	}
 
 	RenderSingleLayerWaterInner(GraphBuilder, SceneColorTexture, SceneDepthTexture, SceneWithoutWaterTextures);
diff --git a/Engine/Source/Runtime/Renderer/Private/VolumetricCloudRendering.cpp b/Engine/Source/Runtime/Renderer/Private/VolumetricCloudRendering.cpp
--- a/Engine/Source/Runtime/Renderer/Private/VolumetricCloudRendering.cpp
+++ b/Engine/Source/Runtime/Renderer/Private/VolumetricCloudRendering.cpp
@@ -1642,13 +1642,12 @@ void FSceneRenderer::RenderVolumetricCloudsInternal(FRDGBuilder& GraphBuilder, F
 	const bool bSecondAtmosphereLightEnabled = CloudRC.bSecondAtmosphereLightEnabled;
 
 	uint32 NoiseFrameIndexModPattern = CloudRC.NoiseFrameIndexModPattern;
-	TRefCountPtr<IPooledRenderTarget> SceneDepthZ = CloudRC.SceneDepthZ;
 	FVolumeShadowingShaderParametersGlobal0 LightShadowShaderParams0 = CloudRC.LightShadowShaderParams0;
 	bool bSkipAerialPerspective = CloudRC.bSkipAerialPerspective;
 
 	FRenderVolumetricCloudGlobalParameters& VolumetricCloudParams = *GraphBuilder.AllocParameters<FRenderVolumetricCloudGlobalParameters>();
 	SetupDefaultRenderVolumetricCloudGlobalParameters(GraphBuilder, VolumetricCloudParams, CloudInfo, MainView);
-	VolumetricCloudParams.SceneDepthTexture = GraphBuilder.RegisterExternalTexture(SceneDepthZ);
+	VolumetricCloudParams.SceneDepthTexture = CloudRC.SceneDepthZ;
 	VolumetricCloudParams.Light0Shadow = LightShadowShaderParams0;
 	VolumetricCloudParams.CloudShadowTexture0 = RenderViewPassParameters->CloudShadowTexture0;
 	VolumetricCloudParams.CloudShadowTexture1 = RenderViewPassParameters->CloudShadowTexture1;
@@ -1721,7 +1720,8 @@ void FSceneRenderer::RenderVolumetricCloud(
 	bool bSkipVolumetricRenderTarget,
 	bool bSkipPerPixelTracing,
 	FRDGTextureMSAA SceneColorTexture,
-	FRDGTextureMSAA SceneDepthTexture)
+	FRDGTextureMSAA SceneDepthTexture,
+	FRDGTextureRef HalfResolutionDepthCheckerboardMinMaxTexture)
 {
 	check(ShouldRenderVolumetricCloud(Scene, ViewFamily.EngineShowFlags)); // This should not be called if we should not render SkyAtmosphere
 
@@ -1742,9 +1742,7 @@ void FSceneRenderer::RenderVolumetricCloud(
 
 			FSceneRenderTargets& SceneContext = FSceneRenderTargets::Get();
 
-			TRefCountPtr<IPooledRenderTarget> SceneDepthZ = SceneContext.SceneDepthZ;
-			TRefCountPtr<IPooledRenderTarget> BlackDummy = GSystemTextures.BlackDummy;
-			FRDGTextureRef BlackDummyRDG = GraphBuilder.RegisterExternalTexture(BlackDummy);
+			FRDGTextureRef SceneDepthZ = SceneDepthTexture.Resolve;
 
 			FCloudRenderContext CloudRC;
 			CloudRC.CloudInfo = &CloudInfo;
@@ -1797,8 +1795,8 @@ void FSceneRenderer::RenderVolumetricCloud(
 					CloudRC.TracingCoordToZbufferCoordScaleBias = VRT.GetTracingCoordToZbufferCoordScaleBias();
 					CloudRC.NoiseFrameIndexModPattern = VRT.GetNoiseFrameIndexModPattern();
 
-					check(VRT.GetMode() != 0 || ViewInfo.HalfResDepthSurfaceCheckerboardMinMax.IsValid());
-					CloudRC.SceneDepthZ = VRT.GetMode() == 0 ? ViewInfo.HalfResDepthSurfaceCheckerboardMinMax : SceneDepthZ;
+					check(VRT.GetMode() != 0 || HalfResolutionDepthCheckerboardMinMaxTexture);
+					CloudRC.SceneDepthZ = VRT.GetMode() == 0 ? HalfResolutionDepthCheckerboardMinMaxTexture : SceneDepthZ;
 				}
 				else
 				{
diff --git a/Engine/Source/Runtime/Renderer/Private/VolumetricCloudRendering.h b/Engine/Source/Runtime/Renderer/Private/VolumetricCloudRendering.h
--- a/Engine/Source/Runtime/Renderer/Private/VolumetricCloudRendering.h
+++ b/Engine/Source/Runtime/Renderer/Private/VolumetricCloudRendering.h
@@ -101,7 +101,7 @@ struct FCloudRenderContext
 	FVolumetricCloudRenderSceneInfo* CloudInfo;
 	FMaterialRenderProxy* CloudVolumeMaterialProxy;
 
-	TRefCountPtr<IPooledRenderTarget> SceneDepthZ;
+	FRDGTextureRef SceneDepthZ = nullptr;
 
 	///////////////////////////////////
 	// Per view parameters
diff --git a/Engine/Source/Runtime/Renderer/Private/VolumetricRenderTarget.cpp b/Engine/Source/Runtime/Renderer/Private/VolumetricRenderTarget.cpp
--- a/Engine/Source/Runtime/Renderer/Private/VolumetricRenderTarget.cpp
+++ b/Engine/Source/Runtime/Renderer/Private/VolumetricRenderTarget.cpp
@@ -103,7 +103,7 @@ static void GetTextureSafeUvCoordBound(FRDGTextureRef Texture, FUintVector4& Tex
 	TextureValidUvRect.W = (float(TexSize.Y) - 0.51f) / float(TexSize.Y);
 };
 
-static bool AnyViewRequiresProcessing(TArray<FViewInfo>& Views)
+static bool AnyViewRequiresProcessing(TArrayView<FViewInfo> Views)
 {
 	bool bAnyViewRequiresProcessing = false;
 	for (int32 ViewIndex = 0; ViewIndex < Views.Num(); ViewIndex++)
@@ -328,7 +328,7 @@ FUintVector4 FVolumetricRenderTargetViewStateData::GetTracingCoordToZbufferCoord
 	FSceneRenderer implementation.
 =============================================================================*/
 
-void FSceneRenderer::InitVolumetricRenderTargetForViews(FRDGBuilder& GraphBuilder)
+void InitVolumetricRenderTargetForViews(FRDGBuilder& GraphBuilder, TArrayView<FViewInfo> Views)
 {
 	for (int32 ViewIndex = 0; ViewIndex < Views.Num(); ViewIndex++)
 	{
@@ -433,7 +433,11 @@ IMPLEMENT_GLOBAL_SHADER(FReconstructVolumetricRenderTargetPS, "/Engine/Private/V
 
 //////////////////////////////////////////////////////////////////////////
 
-void FSceneRenderer::ReconstructVolumetricRenderTarget(FRDGBuilder& GraphBuilder)
+void ReconstructVolumetricRenderTarget(
+	FRDGBuilder& GraphBuilder,
+	TArrayView<FViewInfo> Views,
+	FRDGTextureRef SceneDepthTexture,
+	FRDGTextureRef HalfResolutionDepthCheckerboardMinMaxTexture)
 {
 	if (!AnyViewRequiresProcessing(Views))
 	{
@@ -480,7 +484,7 @@ void FSceneRenderer::ReconstructVolumetricRenderTarget(FRDGBuilder& GraphBuilder
 		PassParameters->CurrentTracingPixelOffset = VolumetricCloudRT.GetCurrentTracingPixelOffset();
 		PassParameters->DownSampleFactor = TracingVolumetricCloudRTDownSample;
 		PassParameters->VolumetricRenderTargetMode = VolumetricCloudRT.GetMode();
-		PassParameters->HalfResDepthTexture = VolumetricCloudRT.GetMode() == 0 ? GraphBuilder.RegisterExternalTexture(ViewInfo.HalfResDepthSurfaceCheckerboardMinMax) : GraphBuilder.RegisterExternalTexture(SceneDepthZ);
+		PassParameters->HalfResDepthTexture = VolumetricCloudRT.GetMode() == 0 ? HalfResolutionDepthCheckerboardMinMaxTexture : SceneDepthTexture;
 
 		GetTextureSafeUvCoordBound(SrcTracingVolumetric, PassParameters->TracingVolumetricTextureValidCoordRect, PassParameters->TracingVolumetricTextureValidUvRect);
 		GetTextureSafeUvCoordBound(PreviousFrameVolumetricTexture, PassParameters->PreviousFrameVolumetricTextureValidCoordRect, PassParameters->PreviousFrameVolumetricTextureValidUvRect);
@@ -548,7 +552,13 @@ IMPLEMENT_GLOBAL_SHADER(FComposeVolumetricRTOverScenePS, "/Engine/Private/Volume
 
 //////////////////////////////////////////////////////////////////////////
 
-void FSceneRenderer::ComposeVolumetricRenderTargetOverScene(FRDGBuilder& GraphBuilder, FRDGTextureRef SceneColorTexture, FRDGTextureRef SceneDepthResolveTexture, bool bShouldRenderSingleLayerWater, const FSceneWithoutWaterTextures& WaterPassData)
+void ComposeVolumetricRenderTargetOverScene(
+	FRDGBuilder& GraphBuilder,
+	TArrayView<FViewInfo> Views,
+	FRDGTextureRef SceneColorTexture,
+	FRDGTextureRef SceneDepthResolveTexture,
+	bool bShouldRenderSingleLayerWater,
+	const FSceneWithoutWaterTextures& WaterPassData)
 {
 	if (!AnyViewRequiresProcessing(Views))
 	{
@@ -610,7 +620,10 @@ void FSceneRenderer::ComposeVolumetricRenderTargetOverScene(FRDGBuilder& GraphBu
 
 //////////////////////////////////////////////////////////////////////////
 
-void FSceneRenderer::ComposeVolumetricRenderTargetOverSceneUnderWater(FRDGBuilder& GraphBuilder, const FSceneWithoutWaterTextures& WaterPassData)
+void ComposeVolumetricRenderTargetOverSceneUnderWater(
+	FRDGBuilder& GraphBuilder,
+	TArrayView<FViewInfo> Views,
+	const FSceneWithoutWaterTextures& WaterPassData)
 {
 	if (!AnyViewRequiresProcessing(Views))
 	{
diff --git a/Engine/Source/Runtime/Renderer/Private/VolumetricRenderTarget.h b/Engine/Source/Runtime/Renderer/Private/VolumetricRenderTarget.h
--- a/Engine/Source/Runtime/Renderer/Private/VolumetricRenderTarget.h
+++ b/Engine/Source/Runtime/Renderer/Private/VolumetricRenderTarget.h
@@ -12,14 +12,30 @@
 #include "RenderResource.h"
 #include "VolumetricRenderTargetViewStateData.h"
 
-
-
 class FScene;
 class FViewInfo;
+struct FSceneWithoutWaterTextures;
 
+void InitVolumetricRenderTargetForViews(FRDGBuilder& GraphBuilder, TArrayView<FViewInfo> Views);
 
+void ReconstructVolumetricRenderTarget(
+	FRDGBuilder& GraphBuilder,
+	TArrayView<FViewInfo> Views,
+	FRDGTextureRef SceneDepthTexture,
+	FRDGTextureRef HalfResolutionDepthCheckerboardMinMaxTexture);
 
-bool ShouldViewRenderVolumetricCloudRenderTarget(const FViewInfo& ViewInfo);
-bool IsVolumetricRenderTargetEnabled();
+void ComposeVolumetricRenderTargetOverScene(
+	FRDGBuilder& GraphBuilder,
+	TArrayView<FViewInfo> Views,
+	FRDGTextureRef SceneColorTexture,
+	FRDGTextureRef SceneDepthResolveTexture,
+	bool bShouldRenderSingleLayerWater,
+	const FSceneWithoutWaterTextures& WaterPassData);
 
+void ComposeVolumetricRenderTargetOverSceneUnderWater(
+	FRDGBuilder& GraphBuilder,
+	TArrayView<FViewInfo> Views,
+	const FSceneWithoutWaterTextures& WaterPassData);
 
+bool ShouldViewRenderVolumetricCloudRenderTarget(const FViewInfo& ViewInfo);
+bool IsVolumetricRenderTargetEnabled();
\ No newline at end of file
