Fix for debug view mode build break.

#rnx

diff --git a/Engine/Source/Runtime/Renderer/Private/DebugViewModeRendering.cpp b/Engine/Source/Runtime/Renderer/Private/DebugViewModeRendering.cpp
--- a/Engine/Source/Runtime/Renderer/Private/DebugViewModeRendering.cpp
+++ b/Engine/Source/Runtime/Renderer/Private/DebugViewModeRendering.cpp
@@ -377,7 +377,12 @@ void InitDebugViewModeInterfaces()
 
 #else // !(UE_BUILD_SHIPPING || UE_BUILD_TEST)
 
-void RenderDebugViewMode(FRDGBuilder& GraphBuilder, TArrayView<const FViewInfo> Views, const FRenderTargetBindingSlots& RenderTargets) {}
+void RenderDebugViewMode(
+	FRDGBuilder& GraphBuilder,
+	TArrayView<const FViewInfo> Views,
+	FScene* Scene,
+	const FRenderTargetBindingSlots& RenderTargets)
+{}
 
 #endif // !(UE_BUILD_SHIPPING || UE_BUILD_TEST)
 
diff --git a/Engine/Source/Runtime/Renderer/Private/DebugViewModeRendering.h b/Engine/Source/Runtime/Renderer/Private/DebugViewModeRendering.h
--- a/Engine/Source/Runtime/Renderer/Private/DebugViewModeRendering.h
+++ b/Engine/Source/Runtime/Renderer/Private/DebugViewModeRendering.h
@@ -219,10 +219,10 @@ void AddDebugViewModeShaderTypes(ERHIFeatureLevel::Type FeatureLevel,
 	const FVertexFactoryType* VertexFactoryType,
 	FMaterialShaderTypes& OutShaderTypes);
 
+#endif // !(UE_BUILD_SHIPPING || UE_BUILD_TEST)
+
 void RenderDebugViewMode(
 	FRDGBuilder& GraphBuilder,
 	TArrayView<const FViewInfo> Views,
 	FScene* Scene,
-	const FRenderTargetBindingSlots& RenderTargets);
-
-#endif // !(UE_BUILD_SHIPPING || UE_BUILD_TEST)
\ No newline at end of file
+	const FRenderTargetBindingSlots& RenderTargets);
\ No newline at end of file
