Port Path Tracer to RDG

diff --git a/Engine/Shaders/Private/PathTracing/PathTracing.usf b/Engine/Shaders/Private/PathTracing/PathTracing.usf
--- a/Engine/Shaders/Private/PathTracing/PathTracing.usf
+++ b/Engine/Shaders/Private/PathTracing/PathTracing.usf
@@ -21,8 +21,9 @@ PathTracingRayGenShader.usf: Reference path tracing
 #include "PathTracingCamera.ush"
 
 
-RWTexture2D<float4> RadianceRT;
+RWTexture2D<float4> RadianceTexture;
 RaytracingAccelerationStructure TLAS;
+int3 TileOffset;
 
 void AccumulateRadiance(inout float3 TotalRadiance, float3 PathRadiance)
 {
@@ -41,13 +42,13 @@ void AccumulateRadiance(inout float3 TotalRadiance, float3 PathRadiance)
 
 RAY_TRACING_ENTRY_RAYGEN(PathTracingMainRG)
 {
-	uint2 LaunchIndex = DispatchRaysIndex().xy + View.ViewRectMin.xy + PathTracingData.TileOffset.xy;
+	uint2 LaunchIndex = DispatchRaysIndex().xy + View.ViewRectMin.xy + TileOffset.xy;
 	
 	// Initialize random sequence 
 	float2 BufferSize = View.BufferSizeAndInvSize.xy; 
 	uint LinearIndex = LaunchIndex.x + (LaunchIndex.y * BufferSize.x); 
 	RandomSequence RandSequence;
-	RandomSequence_Initialize(RandSequence, LinearIndex, AdaptiveSamplingData.TemporalSeed);
+	RandomSequence_Initialize(RandSequence, LinearIndex, PathTracingData.TemporalSeed);
 
 	float3 Radiance = 0;
 	uint RayCounter = 0;
@@ -57,10 +58,6 @@ RAY_TRACING_ENTRY_RAYGEN(PathTracingMainRG)
 	uint2 Pixel;
 	float EmitterPdf;
 	SampleEmitter(LaunchIndex, RandSequence, 0, 0, Ray, Pixel, EmitterPdf);
-	if (EmitterPdf <= 0.0)
-	{
-		return;
-	}
 
 	// This array will hold a CDF for light picking
 	// Seed the array with a uniform CDF at first so that we always have a valid CDF
@@ -296,10 +293,10 @@ RAY_TRACING_ENTRY_RAYGEN(PathTracingMainRG)
 
 
 	// Accumulate radiance and update pixel variance
-	uint NumSamples = AdaptiveSamplingData.Iteration + 1;
+	uint NumSamples = PathTracingData.Iteration + 1;
 	float BlendFactor = 1.0 / float(NumSamples);
 	// Avoid reading the old pixel on the first sample on the off-chance there is a NaN/Inf pixel ...
-	float4 OldPixel = NumSamples > 1 ? RadianceRT[LaunchIndex] : 0;
+	float4 OldPixel = NumSamples > 1 ? RadianceTexture[LaunchIndex] : 0;
 	float3 OldRadiance = OldPixel.rgb;
 	float OldVariance = OldPixel.a;
 
@@ -311,5 +308,5 @@ RAY_TRACING_ENTRY_RAYGEN(PathTracingMainRG)
 	float DeviationSquared = (RadianceY - Luminance(NewRadiance)) * (RadianceY - Luminance(OldRadiance));
 	float NewVariance = lerp(OldVariance, DeviationSquared, BlendFactor);
 
-	RadianceRT[LaunchIndex] = float4(NewRadiance, NewVariance);
+	RadianceTexture[LaunchIndex] = float4(NewRadiance, NewVariance);
 }
diff --git a/Engine/Shaders/Private/PathTracing/PathTracingCamera.ush b/Engine/Shaders/Private/PathTracing/PathTracingCamera.ush
--- a/Engine/Shaders/Private/PathTracing/PathTracingCamera.ush
+++ b/Engine/Shaders/Private/PathTracing/PathTracingCamera.ush
@@ -8,7 +8,7 @@
 
 #include "PathTracingCommon.ush"
 #include "Sensor/PathTracingPinholeCamera.ush"
-#include "Sensor/PathTracingAdaptivePinholeCamera.ush"
+//#include "Sensor/PathTracingAdaptivePinholeCamera.ush"
 
 void SampleEmitter(
 	uint2 LaunchIndex,
@@ -21,14 +21,7 @@ void SampleEmitter(
 )
 {
 	// Pinhole camera is preferred until adaptive variance reduction is a clear win
-	if (AdaptiveSamplingData.UseAdaptiveSampling == 1)
-	{
-		AdaptivePinholeCamera_SampleEmitter(LaunchIndex, RandSequence, SourceEmitterId, SensorId, OutRay, OutPixel, OutPdf);
-	}
-	else
-	{
-		PinholeCamera_SampleEmitter(LaunchIndex, RandSequence, SourceEmitterId, SensorId, OutRay, OutPixel, OutPdf);
-	}
+	PinholeCamera_SampleEmitter(LaunchIndex, RandSequence, SourceEmitterId, SensorId, OutRay, OutPixel, OutPdf);
 }
 
 void PdfEmitter(
@@ -36,12 +29,5 @@ void PdfEmitter(
 	out float OutPdf
 )
 {
-	if (AdaptiveSamplingData.UseAdaptiveSampling)
-	{
-		AdaptivePinholeCamera_PdfEmitter(Ray, OutPdf);
-	}
-	else
-	{
-		PinholeCamera_PdfEmitter(Ray, OutPdf);
-	}
+	PinholeCamera_PdfEmitter(Ray, OutPdf);
 }
diff --git a/Engine/Shaders/Private/PathTracing/PathTracingCompositingPixelShader.usf b/Engine/Shaders/Private/PathTracing/PathTracingCompositingPixelShader.usf
--- a/Engine/Shaders/Private/PathTracing/PathTracingCompositingPixelShader.usf
+++ b/Engine/Shaders/Private/PathTracing/PathTracingCompositingPixelShader.usf
@@ -8,6 +8,7 @@ ReferencePathTracinPixelShader.ush: Reference path tracing  compositing pixel sh
 #include "Utilities/PathTracingFixedPointArithmetic.ush"
 
 uint Iteration;
+uint MaxSamples;
 Texture2D<float4> RadianceTexture;
 
 
@@ -46,6 +47,24 @@ void CompositeMain(
 
 	OutColor = float4(Radiance.rgb * View.PreExposure, 1.0);
 
+#if 0
+	// TODO: make a configureable option for this
+	if (Iteration < MaxSamples)
+	{
+		// Draw a rectangle of how far along we are in accumulating samples
+		float Border = 0.05;
+		if (UV.y > 1 - 3 * Border && UV.y < 1 - Border)
+		{
+			float Fraction = saturate(float(Iteration) / float(MaxSamples));
+			if (UV.x > Border && (UV.x - Border) / (1 - 2 * Border) < Fraction)
+			{
+				float4 ProgressBarColor = float4(0.9, 0.1, 0.1, 1);
+				OutColor = lerp(OutColor, ProgressBarColor, 0.5);
+			}
+		}
+	}
+#endif
+
 #if 0
     // split-screen debugging aids
 	if (UV.x < 0.5)
diff --git a/Engine/Source/Runtime/Renderer/Private/PathTracing.cpp b/Engine/Source/Runtime/Renderer/Private/PathTracing.cpp
--- a/Engine/Source/Runtime/Renderer/Private/PathTracing.cpp
+++ b/Engine/Source/Runtime/Renderer/Private/PathTracing.cpp
@@ -64,36 +64,6 @@ TAutoConsoleVariable<int32> CVarPathTracingFrameIndependentTemporalSeed(
 	ECVF_RenderThreadSafe
 );
 
-TAutoConsoleVariable<int32> CVarPathTracingAdaptiveSampling(
-	TEXT("r.PathTracing.AdaptiveSampling"),
-	0,
-	TEXT("Toggles the use of adaptive sampling\n")
-	TEXT("0: off (default)\n")
-	TEXT("1: on\n"),
-	ECVF_RenderThreadSafe
-);
-
-TAutoConsoleVariable<int32> CVarPathTracingAdaptiveSamplingMinimumSamplesPerPixel(
-	TEXT("r.PathTracing.AdaptiveSampling.MinimumSamplesPerPixel"),
-	16,
-	TEXT("Changes the minimum samples-per-pixel before applying adaptive sampling (default=16)\n"),
-	ECVF_RenderThreadSafe
-);
-
-TAutoConsoleVariable<int32> CVarPathTracingVarianceMapRebuildFrequency(
-	TEXT("r.PathTracing.VarianceMapRebuildFrequency"),
-	16,
-	TEXT("Sets the variance map rebuild frequency (default = every 16 iterations)"),
-	ECVF_RenderThreadSafe
-);
-
-TAutoConsoleVariable<int32> CVarPathTracingRayCountFrequency(
-	TEXT("r.PathTracing.RayCountFrequency"),
-	128,
-	TEXT("Sets the ray count computation frequency (default = every 128 iterations)"),
-	ECVF_RenderThreadSafe
-);
-
 // r.PathTracing.GPUCount is read only because ComputeViewGPUMasks results cannot change after UE has been launched
 TAutoConsoleVariable<int32> CVarPathTracingGPUCount(
 	TEXT("r.PathTracing.GPUCount"),
@@ -111,10 +81,9 @@ TAutoConsoleVariable<int32> CVarPathTracingWiperMode(
 
 IMPLEMENT_GLOBAL_SHADER_PARAMETER_STRUCT(FPathTracingData, "PathTracingData");
 IMPLEMENT_GLOBAL_SHADER_PARAMETER_STRUCT(FPathTracingLightData, "SceneLightsData");
-IMPLEMENT_GLOBAL_SHADER_PARAMETER_STRUCT(FPathTracingAdaptiveSamplingData, "AdaptiveSamplingData");
 
 // This function prepares the portion of shader arguments that may involve invalidating the path traced state
-static bool PrepareShaderArgs(const FViewInfo& View, FPathTracingData& PathTracingData, FPathTracingAdaptiveSamplingData& AdaptiveSamplingData) {
+static bool PrepareShaderArgs(const FViewInfo& View, FPathTracingData& PathTracingData) {
 	int32 PathTracingMaxBounces = CVarPathTracingMaxBounces.GetValueOnRenderThread();
 	if (PathTracingMaxBounces < 0)
 		PathTracingMaxBounces = View.FinalPostProcessSettings.PathTracingMaxBounces;
@@ -160,307 +129,140 @@ static bool PrepareShaderArgs(const FViewInfo& View, FPathTracingData& PathTraci
 		PreviousMaxPathIntensity = PathTracingData.MaxPathIntensity;
 	}
 
-	AdaptiveSamplingData.UseAdaptiveSampling = CVarPathTracingAdaptiveSampling.GetValueOnRenderThread();
-
-	// Changing Adaptive sampling mode requires starting over
-	static uint32 PreviousUseAdaptiveSampling = AdaptiveSamplingData.UseAdaptiveSampling;
-	if (PreviousUseAdaptiveSampling != AdaptiveSamplingData.UseAdaptiveSampling)
-	{
-		NeedInvalidation = true;
-		PreviousUseAdaptiveSampling = AdaptiveSamplingData.UseAdaptiveSampling;
-	}
-
 	// the rest of PathTracingData and AdaptiveSamplingData is filled in by SetParameters below
 	return NeedInvalidation;
 }
 
 class FPathTracingRG : public FGlobalShader
 {
-	DECLARE_SHADER_TYPE(FPathTracingRG, Global);
-
-	static void ModifyCompilationEnvironment(const FGlobalShaderPermutationParameters& Parameters, FShaderCompilerEnvironment& OutEnvironment)
-	{
-	}
-
-public:
+	DECLARE_GLOBAL_SHADER(FPathTracingRG)
+	SHADER_USE_ROOT_PARAMETER_STRUCT(FPathTracingRG, FGlobalShader)
 
 	static bool ShouldCompilePermutation(const FGlobalShaderPermutationParameters& Parameters)
 	{
 		return ShouldCompileRayTracingShadersForProject(Parameters.Platform);
 	}
 
-	FPathTracingRG() {}
-	~FPathTracingRG() {}
-
-	FPathTracingRG(const ShaderMetaType::CompiledShaderInitializerType& Initializer)
-		: FGlobalShader(Initializer)
-	{
-		TLASParameter.Bind(Initializer.ParameterMap, TEXT("TLAS"));
-		ViewParameter.Bind(Initializer.ParameterMap, TEXT("View"));
-		SceneLightsParameters.Bind(Initializer.ParameterMap, TEXT("SceneLightsData"));
-		PathTracingParameters.Bind(Initializer.ParameterMap, TEXT("PathTracingData"));
-		SkyLightParameters.Bind(Initializer.ParameterMap, TEXT("SkyLight"));
-		check(SkyLightParameters.IsBound());
-		AdaptiveSamplingParameters.Bind(Initializer.ParameterMap, TEXT("AdaptiveSamplingData"));
-
-		// Output
-		RadianceRT.Bind(Initializer.ParameterMap, TEXT("RadianceRT"));
-	}
+	BEGIN_SHADER_PARAMETER_STRUCT(FParameters, )
+		SHADER_PARAMETER_RDG_TEXTURE_UAV(RWTexture<float4>, RadianceTexture)
+		SHADER_PARAMETER_SRV(RaytracingAccelerationStructure, TLAS)
 
-	void SetParameters(
-		FScene* Scene,
-		const FViewInfo& View,
-		FRayTracingShaderBindingsWriter& GlobalResources,
-		const FRayTracingScene& RayTracingScene,
-		FRHIUniformBuffer* ViewUniformBuffer,
-		FRHIUniformBuffer* SceneTexturesUniformBuffer,
-		// Shader arguments (expected to be filled in by PrepareShaderArgs ahead of time)
-		FPathTracingData& PathTracingData,
-		FPathTracingAdaptiveSamplingData& AdaptiveSamplingData,
-		// Light buffer
-		const TSparseArray<FLightSceneInfoCompact>& Lights,
-		// Adaptive sampling
-		uint32 Iteration,
-		uint32 FrameIndependentTemporalSeed,
-		FIntVector VarianceDimensions,
-		const FRWBuffer& VarianceMipTree,
-		const FIntVector& TileOffset,
-		// Output
-		FRHIUnorderedAccessView* RadianceUAV)
-	{
+		SHADER_PARAMETER_STRUCT_REF(FViewUniformShaderParameters, ViewUniformBuffer)
+		SHADER_PARAMETER_STRUCT_REF(FSkyLightData, SkyLightData)
+		SHADER_PARAMETER_STRUCT_REF(FPathTracingLightData, SceneLightsData)
+		SHADER_PARAMETER_STRUCT_REF(FPathTracingData, PathTracingData)
+		SHADER_PARAMETER(FIntVector, TileOffset)	// Used by multi-GPU rendering
+	END_SHADER_PARAMETER_STRUCT()
+};
+IMPLEMENT_GLOBAL_SHADER(FPathTracingRG, "/Engine/Private/PathTracing/PathTracing.usf", "PathTracingMainRG", SF_RayGen);
 
-		GlobalResources.Set(TLASParameter, RayTracingScene.RayTracingSceneRHI->GetShaderResourceView());
-		GlobalResources.Set(ViewParameter, ViewUniformBuffer);
+void SetLightParameters(FPathTracingLightData& LightData, FSkyLightData& SkyLightData, FScene* Scene)
+{
+	// Sky light
+	bool IsSkyLightValid = SetupSkyLightParameters(*Scene, &SkyLightData);
 
-		// Path tracing data
-		{
-			PathTracingData.TileOffset = TileOffset;
+	// Lights
+	LightData.Count = 0;
 
-			FUniformBufferRHIRef PathTracingDataUniformBuffer = RHICreateUniformBuffer(&PathTracingData, FPathTracingData::StaticStructMetadata.GetLayout(), EUniformBufferUsage::UniformBuffer_SingleDraw);
-			GlobalResources.Set(PathTracingParameters, PathTracingDataUniformBuffer);
-		}
+	// Prepend SkyLight to light buffer since it is not part of the regular light list
+	if (IsSkyLightValid)
+	{
+		uint32 SkyLightIndex = 0;
+		uint8 SkyLightLightingChannelMask = 0xFF;
+		LightData.Type[SkyLightIndex] = 0;
+		LightData.Color[SkyLightIndex] = FVector(SkyLightData.Color);
+		LightData.Flags[SkyLightIndex] = SkyLightData.bTransmission & 0x01;
+		LightData.Flags[SkyLightIndex] |= (SkyLightLightingChannelMask & 0x7) << 1;
+		LightData.Count++;
+	}
 
-		// Sky light
-		FSkyLightData SkyLightData;
-		bool IsSkyLightValid = false;
-		{
-			IsSkyLightValid = SetupSkyLightParameters(*Scene, &SkyLightData);
+	for (auto Light : Scene->Lights)
+	{
+		if (LightData.Count >= RAY_TRACING_LIGHT_COUNT_MAXIMUM) break;
 
-			FUniformBufferRHIRef SkyLightUniformBuffer = RHICreateUniformBuffer(&SkyLightData, FSkyLightData::StaticStructMetadata.GetLayout(), EUniformBufferUsage::UniformBuffer_SingleDraw);
-			GlobalResources.Set(SkyLightParameters, SkyLightUniformBuffer);
-		}
+		FLightShaderParameters LightParameters;
+		Light.LightSceneInfo->Proxy->GetLightShaderParameters(LightParameters);
+		uint32 Transmission = Light.LightSceneInfo->Proxy->Transmission();
+		uint8 LightingChannelMask = Light.LightSceneInfo->Proxy->GetLightingChannelMask();
+		LightData.Flags[LightData.Count] = Transmission & 0x01;
+		LightData.Flags[LightData.Count] |= (LightingChannelMask & 0x7) << 1;
 
-		// Lights
+		ELightComponentType LightComponentType = (ELightComponentType)Light.LightSceneInfo->Proxy->GetLightType();
+		switch (LightComponentType)
 		{
-			FPathTracingLightData LightData;
-			LightData.Count = 0;
-
-			// Prepend SkyLight to light buffer since it is not part of the regular light list
-			if (IsSkyLightValid)
+			case LightType_Directional:
 			{
-				uint32 SkyLightIndex = 0;
-				uint8 SkyLightLightingChannelMask = 0xFF;
-				LightData.Type[SkyLightIndex] = 0;
-				LightData.Color[SkyLightIndex] = FVector(SkyLightData.Color);
-				LightData.Flags[SkyLightIndex] = SkyLightData.bTransmission & 0x01;
-				LightData.Flags[SkyLightIndex] |= (SkyLightLightingChannelMask & 0x7) << 1;
-				LightData.Count++;
+				LightData.Type[LightData.Count] = 2;
+				LightData.Normal[LightData.Count] = LightParameters.Direction;
+				LightData.Color[LightData.Count] = LightParameters.Color;
+				LightData.Dimensions[LightData.Count] = FVector(0.0f, 0.0f, LightParameters.SourceRadius);
+				LightData.Attenuation[LightData.Count] = 1.0 / LightParameters.InvRadius;
+				break;
 			}
-
-			for (auto Light : Lights)
+			case LightType_Rect:
 			{
-				if (LightData.Count >= RAY_TRACING_LIGHT_COUNT_MAXIMUM) break;
-
-				if (Light.LightSceneInfo->Proxy->HasStaticLighting() && Light.LightSceneInfo->IsPrecomputedLightingValid()) continue;
-
-				FLightShaderParameters LightParameters;
-				Light.LightSceneInfo->Proxy->GetLightShaderParameters(LightParameters);
-				uint32 Transmission = Light.LightSceneInfo->Proxy->Transmission();
-				uint8 LightingChannelMask = Light.LightSceneInfo->Proxy->GetLightingChannelMask();
-				LightData.Flags[LightData.Count] = Transmission & 0x01;
-				LightData.Flags[LightData.Count] |= (LightingChannelMask & 0x7) << 1;
-
-				ELightComponentType LightComponentType = (ELightComponentType)Light.LightSceneInfo->Proxy->GetLightType();
-				switch (LightComponentType)
-				{
-					case LightType_Directional:
-					{
-						LightData.Type[LightData.Count] = 2;
-						LightData.Normal[LightData.Count] = LightParameters.Direction;
-						LightData.Color[LightData.Count] = LightParameters.Color;
-						LightData.Dimensions[LightData.Count] = FVector(0.0f, 0.0f, LightParameters.SourceRadius);
-						LightData.Attenuation[LightData.Count] = 1.0 / LightParameters.InvRadius;
-						break;
-					}
-					case LightType_Rect:
-					{
-						LightData.Type[LightData.Count] = 3;
-						LightData.Position[LightData.Count] = LightParameters.Position;
-						LightData.Normal[LightData.Count] = -LightParameters.Direction;
-						LightData.dPdu[LightData.Count] = FVector::CrossProduct(LightParameters.Tangent, LightParameters.Direction);
-						LightData.dPdv[LightData.Count] = LightParameters.Tangent;
-						LightData.Color[LightData.Count] = LightParameters.Color;
-						LightData.Dimensions[LightData.Count] = FVector(2.0f * LightParameters.SourceRadius, 2.0f * LightParameters.SourceLength, 0.0f);
-						LightData.Attenuation[LightData.Count] = 1.0 / LightParameters.InvRadius;
-						LightData.RectLightBarnCosAngle[LightData.Count] = LightParameters.RectLightBarnCosAngle;
-						LightData.RectLightBarnLength[LightData.Count] = LightParameters.RectLightBarnLength;
-						break;
-					}
-					case LightType_Spot:
-					{
-						LightData.Type[LightData.Count] = 4;
-						LightData.Position[LightData.Count] = LightParameters.Position;
-						LightData.Normal[LightData.Count] = -LightParameters.Direction;
-						LightData.Color[LightData.Count] = LightParameters.Color;
-						LightData.Dimensions[LightData.Count] = FVector(LightParameters.SpotAngles, LightParameters.SourceRadius);
-						LightData.Attenuation[LightData.Count] = 1.0 / LightParameters.InvRadius;
-						break;
-					}
-					case LightType_Point:
-					{
-						LightData.Type[LightData.Count] = 1;
-						LightData.Position[LightData.Count] = LightParameters.Position;
-						LightData.Color[LightData.Count] = LightParameters.Color;
-						LightData.Dimensions[LightData.Count] = FVector(0.0, 0.0, LightParameters.SourceRadius);
-						LightData.Attenuation[LightData.Count] = 1.0 / LightParameters.InvRadius;
-						break;
-					}
-					default:
-					{
-						// Just in case someone adds a new light type one day ...
-						checkNoEntry();
-						break;
-					}
-				};
-
-				LightData.Count++;
+				LightData.Type[LightData.Count] = 3;
+				LightData.Position[LightData.Count] = LightParameters.Position;
+				LightData.Normal[LightData.Count] = -LightParameters.Direction;
+				LightData.dPdu[LightData.Count] = FVector::CrossProduct(LightParameters.Tangent, LightParameters.Direction);
+				LightData.dPdv[LightData.Count] = LightParameters.Tangent;
+				LightData.Color[LightData.Count] = LightParameters.Color;
+				LightData.Dimensions[LightData.Count] = FVector(2.0f * LightParameters.SourceRadius, 2.0f * LightParameters.SourceLength, 0.0f);
+				LightData.Attenuation[LightData.Count] = 1.0 / LightParameters.InvRadius;
+				LightData.RectLightBarnCosAngle[LightData.Count] = LightParameters.RectLightBarnCosAngle;
+				LightData.RectLightBarnLength[LightData.Count] = LightParameters.RectLightBarnLength;
+				break;
 			}
-
-			FUniformBufferRHIRef SceneLightsUniformBuffer = RHICreateUniformBuffer(&LightData, FPathTracingLightData::StaticStructMetadata.GetLayout(), EUniformBufferUsage::UniformBuffer_SingleDraw);
-			GlobalResources.Set(SceneLightsParameters, SceneLightsUniformBuffer);
-		}
-
-		// Adaptive sampling
-		{
-			if (CVarPathTracingFrameIndependentTemporalSeed.GetValueOnRenderThread() == 0)
-			{
-				// Count samples from 0 for deterministic results
-				AdaptiveSamplingData.TemporalSeed = Iteration;
-			}
-			else
+			case LightType_Spot:
 			{
-				// Count samples from an ever-increasing counter to avoid screen-door effect
-				AdaptiveSamplingData.TemporalSeed = FrameIndependentTemporalSeed;
+				LightData.Type[LightData.Count] = 4;
+				LightData.Position[LightData.Count] = LightParameters.Position;
+				LightData.Normal[LightData.Count] = -LightParameters.Direction;
+				LightData.Color[LightData.Count] = LightParameters.Color;
+				LightData.Dimensions[LightData.Count] = FVector(LightParameters.SpotAngles, LightParameters.SourceRadius);
+				LightData.Attenuation[LightData.Count] = 1.0 / LightParameters.InvRadius;
+				break;
 			}
-
-			AdaptiveSamplingData.Iteration = Iteration;
-
-			if (VarianceMipTree.NumBytes > 0)
+			case LightType_Point:
 			{
-				AdaptiveSamplingData.VarianceDimensions = VarianceDimensions;
-				AdaptiveSamplingData.VarianceMipTree = VarianceMipTree.SRV;
-				AdaptiveSamplingData.MinimumSamplesPerPixel = CVarPathTracingAdaptiveSamplingMinimumSamplesPerPixel.GetValueOnRenderThread();
+				LightData.Type[LightData.Count] = 1;
+				LightData.Position[LightData.Count] = LightParameters.Position;
+				LightData.Color[LightData.Count] = LightParameters.Color;
+				LightData.Dimensions[LightData.Count] = FVector(0.0, 0.0, LightParameters.SourceRadius);
+				LightData.Attenuation[LightData.Count] = 1.0 / LightParameters.InvRadius;
+				break;
 			}
-			else
+			default:
 			{
-				AdaptiveSamplingData.UseAdaptiveSampling = 0;
-				AdaptiveSamplingData.VarianceDimensions = FIntVector(1, 1, 1);
-				AdaptiveSamplingData.VarianceMipTree = RHICreateShaderResourceView(GBlackTexture->TextureRHI->GetTexture2D(), 0);
-				AdaptiveSamplingData.MinimumSamplesPerPixel = CVarPathTracingAdaptiveSamplingMinimumSamplesPerPixel.GetValueOnRenderThread();
+				// Just in case someone adds a new light type one day ...
+				checkNoEntry();
+				break;
 			}
+		};
 
-			FUniformBufferRHIRef AdaptiveSamplingDataUniformBuffer = RHICreateUniformBuffer(&AdaptiveSamplingData, FPathTracingAdaptiveSamplingData::StaticStructMetadata.GetLayout(), EUniformBufferUsage::UniformBuffer_SingleDraw);
-			GlobalResources.Set(AdaptiveSamplingParameters, AdaptiveSamplingDataUniformBuffer);
-		}
-
-		// Accumulated output
-		{
-			GlobalResources.Set(RadianceRT, RadianceUAV);
-		}
+		LightData.Count++;
 	}
-
-	/*bool Serialize(FArchive& Ar)
-	{
-		bool bShaderHasOutdatedParameters = FGlobalShader::Serialize(Ar);
-		Ar << TLASParameter;
-		Ar << ViewParameter;
-		Ar << PathTracingParameters;
-		Ar << SceneLightsParameters;
-		Ar << SkyLightParameters;
-		Ar << AdaptiveSamplingParameters;
-		// Output
-		Ar << RadianceRT;
-		Ar << SampleCountRT;
-		Ar << PixelPositionRT;
-		Ar << RayCountPerPixelRT;
-
-		return bShaderHasOutdatedParameters;
-	}*/
-	
-	LAYOUT_FIELD(FShaderResourceParameter, TLASParameter);   // RaytracingAccelerationStructure
-	LAYOUT_FIELD(FShaderUniformBufferParameter, ViewParameter);
-	LAYOUT_FIELD(FShaderUniformBufferParameter, PathTracingParameters);
-	LAYOUT_FIELD(FShaderUniformBufferParameter, SceneLightsParameters);
-	LAYOUT_FIELD(FShaderUniformBufferParameter, SkyLightParameters);
-	LAYOUT_FIELD(FShaderUniformBufferParameter, AdaptiveSamplingParameters);
-
-	// Output parameters
-	LAYOUT_FIELD(FShaderResourceParameter, RadianceRT);
-};
-IMPLEMENT_SHADER_TYPE(, FPathTracingRG, TEXT("/Engine/Private/PathTracing/PathTracing.usf"), TEXT("PathTracingMainRG"), SF_RayGen);
-
-DECLARE_GPU_STAT_NAMED(Stat_GPU_PathTracing, TEXT("Path Tracing"));
-DECLARE_GPU_STAT_NAMED(Stat_GPU_PathTracingBuildSkyLightCDF, TEXT("Path Tracing: Build Sky Light CDF"));
-DECLARE_GPU_STAT_NAMED(Stat_GPU_PathTracingBuildVarianceMipTree, TEXT("Path Tracing: Build Variance Map Tree"));
+}
 
 class FPathTracingCompositorPS : public FGlobalShader
 {
-	DECLARE_SHADER_TYPE(FPathTracingCompositorPS, Global);
-
-	static bool ShouldCache(EShaderPlatform Platform)
-	{
-		return IsFeatureLevelSupported(Platform, ERHIFeatureLevel::SM5);
-	}
+	DECLARE_GLOBAL_SHADER(FPathTracingCompositorPS)
+	SHADER_USE_PARAMETER_STRUCT(FPathTracingCompositorPS, FGlobalShader)
 
 	static bool ShouldCompilePermutation(const FGlobalShaderPermutationParameters& Parameters)
 	{
 		return ShouldCompileRayTracingShadersForProject(Parameters.Platform);
 	}
 
-	static void ModifyCompilationEnvironment(const FGlobalShaderPermutationParameters& Parameters, FShaderCompilerEnvironment& OutEnvironment)
-	{
-		FGlobalShader::ModifyCompilationEnvironment(Parameters, OutEnvironment);
-	}
-
-	FPathTracingCompositorPS(const ShaderMetaType::CompiledShaderInitializerType& Initializer)
-		: FGlobalShader(Initializer)
-	{
-		Iteration.Bind(Initializer.ParameterMap, TEXT("Iteration"));
-		RadianceTexture.Bind(Initializer.ParameterMap, TEXT("RadianceTexture"));
-	}
-
-	FPathTracingCompositorPS()
-	{
-	}
-
-	template<typename TRHICommandList>
-	void SetParameters(
-		TRHICommandList& RHICmdList,
-		const FViewInfo& View,
-		uint32 IterationValue,
-		FRHITexture* RadianceRT)
-	{
-		FRHIPixelShader* ShaderRHI = RHICmdList.GetBoundPixelShader();
-		FGlobalShader::SetParameters<FViewUniformShaderParameters>(RHICmdList, ShaderRHI, View.ViewUniformBuffer);
-		SetShaderValue(RHICmdList, ShaderRHI, Iteration, IterationValue);
-		SetTextureParameter(RHICmdList, ShaderRHI, RadianceTexture, RadianceRT);
-	}
+	BEGIN_SHADER_PARAMETER_STRUCT(FParameters, )
+		SHADER_PARAMETER_RDG_TEXTURE_SRV(Texture2D<float4>, RadianceTexture)
+		SHADER_PARAMETER_STRUCT_REF(FViewUniformShaderParameters, ViewUniformBuffer)
+		SHADER_PARAMETER(unsigned, Iteration)
+		SHADER_PARAMETER(unsigned, MaxSamples)
 
-public:
-	LAYOUT_FIELD(FShaderParameter, Iteration);
-	LAYOUT_FIELD(FShaderResourceParameter, RadianceTexture);
+		RENDER_TARGET_BINDING_SLOTS()
+	END_SHADER_PARAMETER_STRUCT()
 };
-
 IMPLEMENT_SHADER_TYPE(, FPathTracingCompositorPS, TEXT("/Engine/Private/PathTracing/PathTracingCompositingPixelShader.usf"), TEXT("CompositeMain"), SF_Pixel);
 
 void FDeferredShadingSceneRenderer::PreparePathTracing(const FViewInfo& View, TArray<FRHIRayTracingShader*>& OutRayGenShaders)
@@ -473,11 +275,6 @@ void FDeferredShadingSceneRenderer::PreparePathTracing(const FViewInfo& View, TA
 	}
 }
 
-BEGIN_SHADER_PARAMETER_STRUCT(FPathTracingPassParameters, )
-	SHADER_PARAMETER_RDG_UNIFORM_BUFFER(FSceneTextureUniformParameters, SceneTextures)
-	RENDER_TARGET_BINDING_SLOTS()
-END_SHADER_PARAMETER_STRUCT()
-
 void FSceneViewState::PathTracingInvalidate()
 {
 	PathTracingRadianceRT.SafeRelease();
@@ -486,6 +283,7 @@ void FSceneViewState::PathTracingInvalidate()
 	PathTracingSPP = 0;
 }
 
+DECLARE_GPU_STAT_NAMED(Stat_GPU_PathTracing, TEXT("Path Tracing"));
 void FDeferredShadingSceneRenderer::RenderPathTracing(
 	FRDGBuilder& GraphBuilder,
 	const FViewInfo& View,
@@ -493,275 +291,154 @@ void FDeferredShadingSceneRenderer::RenderPathTracing(
 	FRDGTextureRef SceneColorOutputTexture)
 {
 	RDG_GPU_STAT_SCOPE(GraphBuilder, Stat_GPU_PathTracing);
+	RDG_EVENT_SCOPE(GraphBuilder, "Path Tracing");
 
-	FPathTracingPassParameters* PassParameters = GraphBuilder.AllocParameters<FPathTracingPassParameters>();
-	PassParameters->SceneTextures = SceneTexturesUniformBuffer;
 
-	// NOTE: The SkipRenderPass flag means this doesn't get bound. It just ensures that it's put in the RTV state.
-	PassParameters->RenderTargets[0] = FRenderTargetBinding(SceneColorOutputTexture, ERenderTargetLoadAction::ELoad);
+	bool bArgsChanged = false;
 
-	GraphBuilder.AddPass(
-		RDG_EVENT_NAME("PathTracing"),
-		PassParameters,
-		ERDGPassFlags::Compute | ERDGPassFlags::Raster | ERDGPassFlags::SkipRenderPass | ERDGPassFlags::UntrackedAccess,
-		[this, &View, SceneTexturesUniformBuffer, SceneColorOutputTexture](FRHICommandListImmediate& RHICmdList)
+	// Get current value of MaxSPP and reset render if it has changed
+	int32 SamplesPerPixelCVar = CVarPathTracingSamplesPerPixel.GetValueOnRenderThread();
+	uint32 MaxSPP = SamplesPerPixelCVar > -1 ? SamplesPerPixelCVar : View.FinalPostProcessSettings.PathTracingSamplesPerPixel;
+	static uint32 PreviousMaxSPP = MaxSPP;
+	if (PreviousMaxSPP != MaxSPP)
 	{
-		auto ViewSize = View.ViewRect.Size();
-		FSceneViewState* ViewState = (FSceneViewState*)View.State;
-
-		FPathTracingData PathTracingData;
-		FPathTracingAdaptiveSamplingData AdaptiveSamplingData;
-
-		bool bArgsChanged = PrepareShaderArgs(View, PathTracingData, AdaptiveSamplingData);
-
-        // Get current value of MaxSPP and reset render if it has changed
-		int32 SamplesPerPixelCVar = CVarPathTracingSamplesPerPixel.GetValueOnRenderThread();
-		uint32 MaxSPP = SamplesPerPixelCVar > -1 ? SamplesPerPixelCVar : View.FinalPostProcessSettings.PathTracingSamplesPerPixel;
-		static uint32 PreviousMaxSPP = MaxSPP;
-		if (PreviousMaxSPP != MaxSPP)
-		{
-			PreviousMaxSPP = MaxSPP;
-			bArgsChanged = true;
-		}
-
-		// If the scene has changed in some way (camera move, object movement, etc ...)
-		// we must invalidate the viewstate to start over from scratch
-		if (bArgsChanged || ViewState->PathTracingRect != View.ViewRect)
-		{
-			ViewState->PathTracingInvalidate();
-			ViewState->PathTracingRect = View.ViewRect;
-		}
-
-		bool NeedsMoreRays = ViewState->PathTracingSPP < MaxSPP;
-		
-		// Construct render targets for compositing
-		TRefCountPtr<IPooledRenderTarget> RadianceRT;
-
-		if (ViewState->PathTracingRadianceRT)
-		{
-			// already got a buffer
-			RadianceRT = ViewState->PathTracingRadianceRT;
-		}
-		else
-		{
-			FPooledRenderTargetDesc Desc = Translate(SceneColorOutputTexture->Desc);
-			Desc.Flags &= ~(TexCreate_FastVRAM | TexCreate_Transient);
-			Desc.Format = PF_FloatRGBA;
-			GRenderTargetPool.FindFreeElement(RHICmdList, Desc, RadianceRT, TEXT("RadianceRT"));
-			ViewState->PathTracingRadianceRT = RadianceRT;
-		}
-
-		auto RayGenShader = GetGlobalShaderMap(FeatureLevel)->GetShader<FPathTracingRG>();
-
-		FRayTracingShaderBindingsWriter GlobalResources;
-
-		FRHIUniformBuffer* SceneTexturesUniformBufferRHI = SceneTexturesUniformBuffer->GetRHI();
+		PreviousMaxSPP = MaxSPP;
+		bArgsChanged = true;
+	}
+	// Changing FrameIndependentTemporalSeed requires starting over
+	bool LockedSamplingPattern = CVarPathTracingFrameIndependentTemporalSeed.GetValueOnRenderThread() == 0;
+	static bool PreviousLockedSamplingPattern = LockedSamplingPattern;
+	if (PreviousLockedSamplingPattern != LockedSamplingPattern)
+	{
+		PreviousLockedSamplingPattern = LockedSamplingPattern;
+		bArgsChanged = true;
+	}
 
-		FRHIRayTracingScene* RayTracingSceneRHI = View.RayTracingScene.RayTracingSceneRHI;
+	// Get other basic path tracing settings and see if we need to invalidate the current state
+	FPathTracingData PathTracingData;
+	bArgsChanged |= PrepareShaderArgs(View, PathTracingData);
 
-		int32 GPUCount = CVarPathTracingGPUCount.GetValueOnRenderThread();
-		uint32 MainGPUIndex = 0; // Default GPU for rendering
+	// If the scene has changed in some way (camera move, object movement, etc ...)
+	// we must invalidate the ViewState to start over from scratch
+	if (bArgsChanged || View.ViewState->PathTracingRect != View.ViewRect)
+	{
+		View.ViewState->PathTracingInvalidate();
+		View.ViewState->PathTracingRect = View.ViewRect;
+	}
 
-		float DPIScale = FPlatformApplicationMisc::GetDPIScaleFactorAtPoint(View.CursorPos.X, View.CursorPos.Y);
-		const int32 bWiperMode = CVarPathTracingWiperMode.GetValueOnRenderThread();
-		const int32 WipeOffsetX = bWiperMode ? View.CursorPos.X / DPIScale : 0;
+	// Setup temporal seed _after_ invalidation in case we got reset
+	if (LockedSamplingPattern)
+	{
+		// Count samples from 0 for deterministic results
+		PathTracingData.TemporalSeed = View.ViewState->PathTracingSPP;
+	}
+	else
+	{
+		// Count samples from an ever-increasing counter to avoid screen-door effect
+		PathTracingData.TemporalSeed = View.ViewState->PathTracingFrameIndependentTemporalSeed;
+	}
+	PathTracingData.Iteration = View.ViewState->PathTracingSPP;
 
-		bool bDoMGPUPathTracing = GNumExplicitGPUsForRendering > 1 && GPUCount > 1;
+	// Prepare radiance buffer (will be shared with display pass)
+	FRDGTexture* RadianceTexture = nullptr;
+	if (View.ViewState->PathTracingRadianceRT)
+	{
+		// we already have a valid radiance texture, re-use it
+		RadianceTexture = GraphBuilder.RegisterExternalTexture(View.ViewState->PathTracingRadianceRT, TEXT("PathTracerRadiance"));
+	}
+	else
+	{
+		// First time through, need to make a new texture
+		FRDGTextureDesc RadianceTextureDesc = FRDGTextureDesc::Create2D(
+			View.ViewRect.Size(),
+			PF_FloatRGBA,
+			FClearValueBinding::None,
+			TexCreate_ShaderResource | TexCreate_UAV);
+		RadianceTexture = GraphBuilder.CreateTexture(RadianceTextureDesc, TEXT("PathTracerRadiance"), ERDGTextureFlags::MultiFrame);
+	}
+	bool NeedsMoreRays = PathTracingData.Iteration < MaxSPP;
 
-		if (bDoMGPUPathTracing && NeedsMoreRays)
+	if (NeedsMoreRays)
+	{
+		FPathTracingRG::FParameters* PassParameters = GraphBuilder.AllocParameters<FPathTracingRG::FParameters>();
+		PassParameters->TLAS = View.RayTracingScene.RayTracingSceneRHI->GetShaderResourceView();
+		PassParameters->ViewUniformBuffer = View.ViewUniformBuffer;
+		PassParameters->PathTracingData = CreateUniformBufferImmediate(PathTracingData, EUniformBufferUsage::UniformBuffer_SingleFrame);
 		{
-			//#dxr-todo: Set minimum tile size for mGPU
-			int32 TileSizeX = FIntPoint::DivideAndRoundUp(View.ViewRect.Size(), GPUCount).X;
-
-			{
-				// FIXME: cross-GPU fences aren't currently supported in the transition API, we need to figure it out.
-#if 0
-				FRHIUnorderedAccessView* UAVs[] = {
-						RadianceRT->GetRenderTargetItem().UAV
-				};
-
-				// Begin mGPU fence
-				FRHIGPUMask GPUMask = FRHIGPUMask::All();
-				FComputeFenceRHIRef Fence = RHICmdList.CreateComputeFence(TEXT("PathTracingRayGen_Fence_Begin"));
-
-				RHICmdList.TransitionResources(EResourceTransitionAccess::ERWBarrier, EResourceTransitionPipeline::EGfxToGfx, UAVs, UE_ARRAY_COUNT(UAVs), Fence);
-
-				for (uint32 GPUIndex : GPUMask)
-				{
-					if (GPUIndex == MainGPUIndex)
-						continue;
-
-					SCOPED_GPU_MASK(RHICmdList, FRHIGPUMask::FromIndex(GPUIndex));
-
-					RHICmdList.WaitComputeFence(Fence);
-				}
-#endif
-			}
-
-			for (int32 GPUIndex = 0; GPUIndex < GPUCount; ++GPUIndex)
-			{
-				SCOPED_GPU_MASK(RHICmdList, FRHIGPUMask::FromIndex(GPUIndex));
-
-				FIntVector TileOffset;
-				TileOffset.X = GPUIndex * TileSizeX;
-				TileOffset.Y = 0; //vertical tiles only 
-
-				RayGenShader->SetParameters(
-					Scene,
-					View,
-					GlobalResources,
-					View.RayTracingScene,
-					View.ViewUniformBuffer,
-					SceneTexturesUniformBufferRHI,
-					PathTracingData,
-					AdaptiveSamplingData,
-					Scene->Lights,
-					ViewState->PathTracingSPP,
-					ViewState->PathTracingFrameIndependentTemporalSeed,
-					ViewState->VarianceMipTreeDimensions,
-					*ViewState->VarianceMipTree,
-					TileOffset,
-					RadianceRT->GetRenderTargetItem().UAV
-				);
-
-				int32 DispatchSizeX = FMath::Min<int32>(TileSizeX, View.ViewRect.Size().X - TileOffset.X);
-				int32 DispatchSizeY = View.ViewRect.Size().Y;
-
-				RHICmdList.RayTraceDispatch(
-					View.RayTracingMaterialPipeline,
-					RayGenShader.GetRayTracingShader(),
-					RayTracingSceneRHI, GlobalResources,
-					DispatchSizeX, DispatchSizeY
-				);
-
-				FIntRect GPURect;
-				GPURect.Min.X = TileOffset.X;
-				GPURect.Min.Y = TileOffset.Y;
-				GPURect.Max.X = TileOffset.X + DispatchSizeX;
-				GPURect.Max.Y = TileOffset.Y + DispatchSizeY;
-
-				if (GPUIndex > 0)
-				{
-					RHICmdList.TransferTexture(RadianceRT->GetRenderTargetItem().TargetableTexture->GetTexture2D(), GPURect, GPUIndex, 0, true);
-				}
-			}
-
-			// FIXME: cross-GPU fences aren't currently supported in the transition API, we need to figure it out.
-#if 0
-			{
-				// End mGPU fence
-				FRHIGPUMask GPUMask = FRHIGPUMask::All();
-				FComputeFenceRHIRef Fence = RHICmdList.CreateComputeFence(TEXT("PathTracingRayGen_Fence_End"));
-
-				RHICmdList.TransitionResources(EResourceTransitionAccess::ERWBarrier, EResourceTransitionPipeline::EGfxToGfx, UAVs, UE_ARRAY_COUNT(UAVs), Fence);
-
-				for (uint32 GPUIndex : GPUMask)
-				{
-					if (GPUIndex == MainGPUIndex)
-						continue;
-
-					SCOPED_GPU_MASK(RHICmdList, FRHIGPUMask::FromIndex(MainGPUIndex));
-
-					RHICmdList.WaitComputeFence(Fence);
-				}
-			}
-#endif
+			// upload sky/lights data
+			FSkyLightData SkyLightData;
+			FPathTracingLightData LightData;
+			SetLightParameters(LightData, SkyLightData, Scene);
+			PassParameters->SkyLightData = CreateUniformBufferImmediate(SkyLightData, EUniformBufferUsage::UniformBuffer_SingleFrame);
+			PassParameters->SceneLightsData = CreateUniformBufferImmediate(LightData, EUniformBufferUsage::UniformBuffer_SingleFrame);
 		}
-		else if (NeedsMoreRays)
+		PassParameters->RadianceTexture = GraphBuilder.CreateUAV(RadianceTexture);
+
+		// TODO: in multi-gpu case, split image into tiles
+		PassParameters->TileOffset.X = 0;
+		PassParameters->TileOffset.Y = 0;
+
+		TShaderMapRef<FPathTracingRG> RayGenShader(View.ShaderMap);
+		ClearUnusedGraphResources(RayGenShader, PassParameters);
+		GraphBuilder.AddPass(
+			RDG_EVENT_NAME("Path Tracer Compute (%d x %d) Sample=%d/%d", View.ViewRect.Size().X, View.ViewRect.Size().Y, View.ViewState->PathTracingSPP, MaxSPP),
+			PassParameters,
+			ERDGPassFlags::Compute,
+			[PassParameters, RayGenShader, &View](FRHICommandListImmediate& RHICmdList)
 		{
-			FIntVector TileOffset;
-			TileOffset.X = bWiperMode > 0 ? WipeOffsetX : 0;
-			TileOffset.Y = 0;
-
-			RayGenShader->SetParameters(
-				Scene,
-				View,
-				GlobalResources,
-				View.RayTracingScene,
-				View.ViewUniformBuffer,
-				SceneTexturesUniformBufferRHI,
-				PathTracingData,
-				AdaptiveSamplingData,
-				Scene->Lights,
-				ViewState->PathTracingSPP,
-				ViewState->PathTracingFrameIndependentTemporalSeed,
-				ViewState->VarianceMipTreeDimensions,
-				*ViewState->VarianceMipTree,
-				TileOffset,
-				RadianceRT->GetRenderTargetItem().UAV
-			);
-
-			int32 ViewWidth = View.ViewRect.Size().X;
-
-			int32 DispatchSizeX = bWiperMode > 0 ? ViewWidth - WipeOffsetX : ViewWidth;
+			FRHIRayTracingScene* RayTracingSceneRHI = View.RayTracingScene.RayTracingSceneRHI;
+			
+			int32 DispatchSizeX = View.ViewRect.Size().X;
 			int32 DispatchSizeY = View.ViewRect.Size().Y;
 
+			FRayTracingShaderBindingsWriter GlobalResources;
+			SetShaderParameters(GlobalResources, RayGenShader, *PassParameters);
+
 			RHICmdList.RayTraceDispatch(
 				View.RayTracingMaterialPipeline,
 				RayGenShader.GetRayTracingShader(),
 				RayTracingSceneRHI, GlobalResources,
 				DispatchSizeX, DispatchSizeY
 			);
-		}
-
-		// Save RayTracingIndirect for compositing
-		RHICmdList.CopyToResolveTarget(RadianceRT->GetRenderTargetItem().TargetableTexture, RadianceRT->GetRenderTargetItem().ShaderResourceTexture, FResolveParams());
-
-		// Single GPU for launching compute shaders
-		SCOPED_GPU_MASK(RHICmdList, FRHIGPUMask::FromIndex(MainGPUIndex));
+		});
 
-#if 0
-		if ((ViewState->PathTracingSPP % CVarPathTracingRayCountFrequency.GetValueOnRenderThread() == 0))
-		{
-			ComputeRayCount(RHICmdList, View, RayCountPerPixelRT->GetRenderTargetItem().ShaderResourceTexture);
-		}
-#endif
-		// Run compositing engine
-		const auto ShaderMap = GetGlobalShaderMap(FeatureLevel);
-
-		TShaderMapRef<FPostProcessVS> VertexShader(ShaderMap);
-		TShaderMapRef<FPathTracingCompositorPS> PixelShader(ShaderMap);
-		FRHITexture* RenderTargets[1] =
-		{
-			SceneColorOutputTexture->GetRHI()
-		};
-		FRHIRenderPassInfo RenderPassInfo(UE_ARRAY_COUNT(RenderTargets), RenderTargets, ERenderTargetActions::Load_Store);
-		RHICmdList.BeginRenderPass(RenderPassInfo, TEXT("PathTracing"));
-
-		// DEBUG: Inspect render target in isolation
-		FGraphicsPipelineStateInitializer GraphicsPSOInit;
-		RHICmdList.ApplyCachedRenderTargets(GraphicsPSOInit);
-		GraphicsPSOInit.BlendState = TStaticBlendState<>::GetRHI();
-		GraphicsPSOInit.RasterizerState = TStaticRasterizerState<FM_Solid, CM_None>::GetRHI();
-		GraphicsPSOInit.DepthStencilState = TStaticDepthStencilState<false, CF_Always>::GetRHI();
-		GraphicsPSOInit.BoundShaderState.VertexDeclarationRHI = GFilterVertexDeclaration.VertexDeclarationRHI;
-		GraphicsPSOInit.BoundShaderState.VertexShaderRHI = VertexShader.GetVertexShader();
-		GraphicsPSOInit.BoundShaderState.PixelShaderRHI = PixelShader.GetPixelShader();
-		GraphicsPSOInit.PrimitiveType = PT_TriangleList;
-		SetGraphicsPipelineState(RHICmdList, GraphicsPSOInit);
+		// After we are done, make sure we remember our texture for next time so that we can accumulate samples across frames
+		GraphBuilder.QueueTextureExtraction(RadianceTexture, &View.ViewState->PathTracingRadianceRT);
+	}
 
-		{
-			RHICmdList.SetViewport(View.ViewRect.Min.X, View.ViewRect.Min.Y, 0.0f, View.ViewRect.Max.X, View.ViewRect.Max.Y, 1.0f);
+	// now add a pixel shader pass to display our Radiance buffer
 
-			PixelShader->SetParameters(RHICmdList, View, AdaptiveSamplingData.Iteration, RadianceRT->GetRenderTargetItem().ShaderResourceTexture);
+	FPathTracingCompositorPS::FParameters* DisplayParameters = GraphBuilder.AllocParameters<FPathTracingCompositorPS::FParameters>();
+	DisplayParameters->Iteration = PathTracingData.Iteration;
+	DisplayParameters->MaxSamples = MaxSPP;
+	DisplayParameters->ViewUniformBuffer = View.ViewUniformBuffer;
+	DisplayParameters->RadianceTexture = GraphBuilder.CreateSRV(FRDGTextureSRVDesc::Create(RadianceTexture));
+	DisplayParameters->RenderTargets[0] = FRenderTargetBinding(SceneColorOutputTexture, ERenderTargetLoadAction::ENoAction);
 
-			int32 DispatchSizeX = View.ViewRect.Size().X;
+	FScreenPassTextureViewport Viewport(SceneColorOutputTexture, View.ViewRect);
 
-			DrawRectangle(
-				RHICmdList,
-				WipeOffsetX, 0,
-				DispatchSizeX, View.ViewRect.Height(),
-				WipeOffsetX, View.ViewRect.Min.Y,
-				DispatchSizeX, View.ViewRect.Height(),
-				FIntPoint(View.ViewRect.Width(), View.ViewRect.Height()),
-				SceneColorOutputTexture->Desc.Extent,
-				VertexShader);
-		}
-		RHICmdList.EndRenderPass();
+	// wiper mode - reveals the render below the path tracing display
+	// NOTE: we still path trace the full resolution even while wiping the cursor so that rendering does not get out of sync
+	if (CVarPathTracingWiperMode.GetValueOnRenderThread() != 0)
+	{
+		float DPIScale = FPlatformApplicationMisc::GetDPIScaleFactorAtPoint(View.CursorPos.X, View.CursorPos.Y);
+		Viewport.Rect.Min.X = View.CursorPos.X / DPIScale;
+	}
 
-		// Bump counters for next frame
-		++ViewState->PathTracingSPP;
-		++ViewState->PathTracingFrameIndependentTemporalSeed;
-	});
+	TShaderMapRef<FPathTracingCompositorPS> PixelShader(View.ShaderMap);
+	AddDrawScreenPass(
+		GraphBuilder,
+		RDG_EVENT_NAME("Path Tracer Display (%d x %d)", View.ViewRect.Size().X, View.ViewRect.Size().Y),
+		View,
+		Viewport,
+		Viewport,
+		PixelShader,
+		DisplayParameters
+	);
+
+	// Bump counters for next frame
+	++View.ViewState->PathTracingSPP;
+	++View.ViewState->PathTracingFrameIndependentTemporalSeed;
 }
-#endif
\ No newline at end of file
+
+#endif
diff --git a/Engine/Source/Runtime/Renderer/Private/PathTracingUniformBuffers.h b/Engine/Source/Runtime/Renderer/Private/PathTracingUniformBuffers.h
--- a/Engine/Source/Runtime/Renderer/Private/PathTracingUniformBuffers.h
+++ b/Engine/Source/Runtime/Renderer/Private/PathTracingUniformBuffers.h
@@ -10,12 +10,13 @@
 #include "RayTracingDefinitions.h"
 
 BEGIN_GLOBAL_SHADER_PARAMETER_STRUCT(FPathTracingData, )
+	SHADER_PARAMETER(uint32, Iteration)
+	SHADER_PARAMETER(uint32, TemporalSeed)
 	SHADER_PARAMETER(uint32, MaxBounces)
 	SHADER_PARAMETER(uint32, MISMode)
 	SHADER_PARAMETER(uint32, VisibleLights)
 	SHADER_PARAMETER(float, MaxPathIntensity)
 	SHADER_PARAMETER(float, MaxNormalBias)
-	SHADER_PARAMETER(FIntVector, TileOffset)
 END_GLOBAL_SHADER_PARAMETER_STRUCT()
 
 
@@ -41,13 +42,3 @@ BEGIN_GLOBAL_SHADER_PARAMETER_STRUCT(FPathTracingLightData, RENDERER_API)
 	// Only used by GPULightmass currently, not filled in realtime paths
 	SHADER_PARAMETER_ARRAY(uint32, Mobility, [RAY_TRACING_LIGHT_COUNT_MAXIMUM])
 END_GLOBAL_SHADER_PARAMETER_STRUCT()
-
-
-BEGIN_GLOBAL_SHADER_PARAMETER_STRUCT(FPathTracingAdaptiveSamplingData, )
-	SHADER_PARAMETER(uint32, UseAdaptiveSampling)
-	SHADER_PARAMETER(uint32, MinimumSamplesPerPixel)
-	SHADER_PARAMETER(uint32, Iteration)
-	SHADER_PARAMETER(uint32, TemporalSeed)
-	SHADER_PARAMETER(FIntVector, VarianceDimensions)
-	SHADER_PARAMETER_SRV(Buffer<float>, VarianceMipTree)
-END_GLOBAL_SHADER_PARAMETER_STRUCT()
