Distance field rhi validation fixes for missing/wrong transitions

diff --git a/Engine/Source/Runtime/Renderer/Private/DistanceFieldObjectCulling.cpp b/Engine/Source/Runtime/Renderer/Private/DistanceFieldObjectCulling.cpp
--- a/Engine/Source/Runtime/Renderer/Private/DistanceFieldObjectCulling.cpp
+++ b/Engine/Source/Runtime/Renderer/Private/DistanceFieldObjectCulling.cpp
@@ -142,8 +142,8 @@ public:
 		UAVTransitions[1] = FRHITransitionInfo(GAOCulledObjectBuffers.Buffers.Bounds.UAV, ERHIAccess::Unknown, ERHIAccess::UAVCompute);
 		UAVTransitions[2] = FRHITransitionInfo(GAOCulledObjectBuffers.Buffers.Data.UAV, ERHIAccess::Unknown, ERHIAccess::UAVCompute);
 		UAVTransitions[3] = FRHITransitionInfo(GAOCulledObjectBuffers.Buffers.BoxBounds.UAV, ERHIAccess::Unknown, ERHIAccess::UAVCompute);
-		UAVTransitions[4] = FRHITransitionInfo(Scene->DistanceFieldSceneData.GetCurrentObjectBuffers()->Data.UAV, ERHIAccess::Unknown, ERHIAccess::UAVCompute);
-		UAVTransitions[5] = FRHITransitionInfo(Scene->DistanceFieldSceneData.GetCurrentObjectBuffers()->Bounds.UAV, ERHIAccess::Unknown, ERHIAccess::UAVCompute);
+		UAVTransitions[4] = FRHITransitionInfo(Scene->DistanceFieldSceneData.GetCurrentObjectBuffers()->Data.UAV, ERHIAccess::Unknown, ERHIAccess::SRVCompute);
+		UAVTransitions[5] = FRHITransitionInfo(Scene->DistanceFieldSceneData.GetCurrentObjectBuffers()->Bounds.UAV, ERHIAccess::Unknown, ERHIAccess::SRVCompute);
 		RHICmdList.Transition(MakeArrayView(UAVTransitions, UE_ARRAY_COUNT(UAVTransitions)));
 
 		FRHIComputeShader* ShaderRHI = RHICmdList.GetBoundComputeShader();
@@ -514,6 +514,7 @@ public:
 		{
 			TransitionInfos.Add(FRHITransitionInfo(UAV, ERHIAccess::Unknown, ERHIAccess::UAVCompute));
 		}
+		TransitionInfos.Add(FRHITransitionInfo(TileIntersectionResources->NumCulledTilesArray.UAV, ERHIAccess::Unknown, ERHIAccess::SRVMask));
 		RHICmdList.Transition(MakeArrayView(TransitionInfos.GetData(), TransitionInfos.Num()));
 
 		TileIntersectionParameters.Set(RHICmdList, ShaderRHI, *TileIntersectionResources);
@@ -549,7 +550,7 @@ void ScatterTilesToObjects(FRHICommandListImmediate& RHICmdList, const FViewInfo
 	TArray<FRHITransitionInfo> TransitionInfos;
 	for (FRHIUnorderedAccessView* UAV : UAVs)
 	{
-		TransitionInfos.Add(FRHITransitionInfo(UAV, ERHIAccess::Unknown, ERHIAccess::UAVCompute));
+		TransitionInfos.Add(FRHITransitionInfo(UAV, ERHIAccess::Unknown, ERHIAccess::UAVMask));
 	}
 	RHICmdList.Transition(MakeArrayView(TransitionInfos.GetData(), TransitionInfos.Num()));
 
diff --git a/Engine/Source/Runtime/Renderer/Private/DistanceFieldScreenGridLighting.cpp b/Engine/Source/Runtime/Renderer/Private/DistanceFieldScreenGridLighting.cpp
--- a/Engine/Source/Runtime/Renderer/Private/DistanceFieldScreenGridLighting.cpp
+++ b/Engine/Source/Runtime/Renderer/Private/DistanceFieldScreenGridLighting.cpp
@@ -155,6 +155,8 @@ public:
 		FTileIntersectionResources* TileIntersectionResources = View.ViewState->AOTileIntersectionResources;
 		SetSRVParameter(RHICmdList, ShaderRHI, TileConeDepthRanges, TileIntersectionResources->TileConeDepthRanges.SRV);
 
+		RHICmdList.Transition(FRHITransitionInfo(TileIntersectionResources->CulledTileDataArray.UAV, ERHIAccess::Unknown, ERHIAccess::SRVMask));
+
 		TileIntersectionParameters.Set(RHICmdList, ShaderRHI, *TileIntersectionResources);
 
 		extern float GAOConeHalfAngle;
@@ -306,7 +308,9 @@ public:
 		SetShaderValue(RHICmdList, ShaderRHI, BentNormalNormalizeFactor, BentNormalNormalizeFactorValue);
 
 		FAOScreenGridResources* ScreenGridResources = View.ViewState->AOScreenGridResources;
-				
+			
+		RHICmdList.Transition(FRHITransitionInfo(ScreenGridResources->ScreenGridConeVisibility.UAV, ERHIAccess::Unknown, ERHIAccess::UAVMask));
+
 		// Note: no transition, want to overlap object cone tracing and global DF cone tracing since both shaders use atomics to ScreenGridConeVisibility
 		RHICmdList.BeginUAVOverlap(ScreenGridResources->ScreenGridConeVisibility.UAV);
 
@@ -317,6 +321,8 @@ public:
 	{
 		ScreenGridConeVisibility.UnsetUAV(RHICmdList, RHICmdList.GetBoundComputeShader());
 		RHICmdList.EndUAVOverlap(View.ViewState->AOScreenGridResources->ScreenGridConeVisibility.UAV);
+
+		RHICmdList.Transition(FRHITransitionInfo(View.ViewState->AOScreenGridResources->ScreenGridConeVisibility.UAV, ERHIAccess::UAVMask, ERHIAccess::SRVMask));
 	}
 
 private:
@@ -566,7 +572,7 @@ void FDeferredShadingSceneRenderer::RenderDistanceFieldAOScreenGrid(
 				SCOPED_DRAW_EVENT(RHICmdList, ConeTraceObjects);
 				FTileIntersectionResources* TileIntersectionResources = ((FSceneViewState*)View.State)->AOTileIntersectionResources;
 
-				RHICmdList.Transition(FRHITransitionInfo(ScreenGridResources->ScreenGridConeVisibility.UAV, ERHIAccess::UAVCompute, ERHIAccess::UAVCompute));
+				RHICmdList.Transition(FRHITransitionInfo(ScreenGridResources->ScreenGridConeVisibility.UAV, ERHIAccess::Unknown, ERHIAccess::UAVCompute));
 
 				if (bUseGlobalDistanceField)
 				{
