Fix crash when destroying a VT.
Happens when SortedIndices haven't been updated after mapping and before destroy.
This quick fix detects if we have mappings that haven't been fully updated, and handles them gracefully.
Probably not most performant though, so next I will look for the best way to flush the mappings before any destroy happens.

diff --git a/Engine/Source/Runtime/Renderer/Private/VT/TexturePageMap.cpp b/Engine/Source/Runtime/Renderer/Private/VT/TexturePageMap.cpp
--- a/Engine/Source/Runtime/Renderer/Private/VT/TexturePageMap.cpp
+++ b/Engine/Source/Runtime/Renderer/Private/VT/TexturePageMap.cpp
@@ -138,13 +138,22 @@ void FTexturePageMap::UnmapPage(FVirtualTextureSystem* System, FVirtualTextureSp
 		}
 	}
 
-	const uint32 OldKey = EncodeSortKey(vLogSize, vAddress);
-	const uint32 OldIndex = LowerBound(0, SortedKeys.Num(), OldKey, ~0u);
-	check(SortedKeys[OldIndex] == OldKey); // make sure we actually found the key (should always exist, since we're removing it)
-	checkSlow(UpperBound(0, SortedKeys.Num(), OldKey, ~0u) == OldIndex + 1u); // make sure key only exists once
-	checkSlow(!SortedSubIndexes.Contains(OldIndex)); // make sure we're not somehow removing the same key twice
+	// Deal with case where SortedAddIndexes has not been processed yet.
+	int32 ToAddIndex = SortedAddIndexes.Find(PageIndex);
+	if (ToAddIndex != INDEX_NONE)
+	{
+		SortedAddIndexes.RemoveAtSwap(ToAddIndex, 1, false);
+	}
+	else
+	{
+		const uint32 OldKey = EncodeSortKey(vLogSize, vAddress);
+		const uint32 OldIndex = LowerBound(0, SortedKeys.Num(), OldKey, ~0u);
+		check(SortedKeys[OldIndex] == OldKey); // make sure we actually found the key (should always exist, since we're removing it)
+		checkSlow(UpperBound(0, SortedKeys.Num(), OldKey, ~0u) == OldIndex + 1u); // make sure key only exists once
+		checkSlow(!SortedSubIndexes.Contains(OldIndex)); // make sure we're not somehow removing the same key twice
 
-	SortedSubIndexes.Add(OldIndex);
+		SortedSubIndexes.Add(OldIndex);
+	}
 
 	RemovePageFromList(PageIndex);
 	AddPageToList(PageListHead_Unmapped, PageIndex);
