Fixed Distance Field Object Management state getting out of sync when r.LumenScene is changed during loading

diff --git a/Engine/Source/Runtime/Renderer/Private/RendererScene.cpp b/Engine/Source/Runtime/Renderer/Private/RendererScene.cpp
--- a/Engine/Source/Runtime/Renderer/Private/RendererScene.cpp
+++ b/Engine/Source/Runtime/Renderer/Private/RendererScene.cpp
@@ -422,7 +422,7 @@ bool IncludePrimitiveInDistanceFieldSceneData(bool bTrackAllPrimitives, const FP
 	return (bTrackAllPrimitives || Proxy->CastsDynamicIndirectShadow())
 		&& Proxy->AffectsDistanceFieldLighting()
 		&& (Proxy->IsDrawnInGame() || Proxy->CastsHiddenShadow())
-		&& (Proxy->CastsDynamicShadow() || (Proxy->AffectsDynamicIndirectLighting() && GAllowLumenScene));
+		&& (Proxy->CastsDynamicShadow() || Proxy->AffectsDynamicIndirectLighting());
 }
 
 void FDistanceFieldSceneData::AddPrimitive(FPrimitiveSceneInfo* InPrimitive)
@@ -521,6 +521,8 @@ void FDistanceFieldSceneData::RemovePrimitive(FPrimitiveSceneInfo* InPrimitive)
 			InPrimitive->DistanceFieldInstanceIndices.Empty();
 		}
 	}
+
+	checkf(!PendingAddOperations.Contains(InPrimitive), TEXT("Primitive is being removed from the scene, but didn't remove from Distance Field Scene properly - a crash will occur when processing PendingAddOperations.  This can happen if the proxy's properties have changed without recreating its render state."));
 }
 
 void FDistanceFieldSceneData::Release()
