Always test if page is mapped instead of assuming it isn't in one case.
This is a placeholder solution for a VT mapping crash while I investigate how the logic is failing here.

diff --git a/Engine/Source/Runtime/Renderer/Private/VT/VirtualTextureSystem.cpp b/Engine/Source/Runtime/Renderer/Private/VT/VirtualTextureSystem.cpp
--- a/Engine/Source/Runtime/Renderer/Private/VT/VirtualTextureSystem.cpp
+++ b/Engine/Source/Runtime/Renderer/Private/VT/VirtualTextureSystem.cpp
@@ -1571,7 +1571,7 @@ void FVirtualTextureSystem::GatherRequestsTask(const FGatherRequestsParameters&
 							AllocatedVT->GetProducerIndexForPageTableLayer(PageTableLayerIndex) == ProducerIndex)
 						{
 							bool bPageWasMapped = false;
-							if (Allocated_vLevel != vLevel)
+							if (1)//Allocated_vLevel != vLevel)
 							{
 								// if we found a lower resolution tile than was requested, it may have already been mapped, check for that first
 								// don't need to check this if the allocated page is at the level that was requested...if that was already mapped we wouldn't have gotten this far
