Fix crash when destroying a VT.
SortedAddIndexes only holds PageIndex in lower 32 bits so mask correctly.

diff --git a/Engine/Source/Runtime/Renderer/Private/VT/TexturePageMap.cpp b/Engine/Source/Runtime/Renderer/Private/VT/TexturePageMap.cpp
--- a/Engine/Source/Runtime/Renderer/Private/VT/TexturePageMap.cpp
+++ b/Engine/Source/Runtime/Renderer/Private/VT/TexturePageMap.cpp
@@ -139,12 +139,18 @@ void FTexturePageMap::UnmapPage(FVirtualTextureSystem* System, FVirtualTextureSp
 	}
 
 	// Deal with case where SortedAddIndexes has not been processed yet.
-	int32 ToAddIndex = SortedAddIndexes.Find(PageIndex);
-	if (ToAddIndex != INDEX_NONE)
+	bool bFoundInAddIndexes = false;
+	for (int32 AddIndex = 0; AddIndex < SortedAddIndexes.Num(); ++AddIndex)
 	{
-		SortedAddIndexes.RemoveAtSwap(ToAddIndex, 1, false);
+		if ((SortedAddIndexes[AddIndex] & 0xffffffff) == PageIndex)
+		{
+			bFoundInAddIndexes = true;
+			SortedAddIndexes.RemoveAtSwap(AddIndex, 1, false);
+			break;
+		}
 	}
-	else
+
+	if (!bFoundInAddIndexes)
 	{
 		const uint32 OldKey = EncodeSortKey(vLogSize, vAddress);
 		const uint32 OldIndex = LowerBound(0, SortedKeys.Num(), OldKey, ~0u);
