Fix RT transition issues

diff --git a/Engine/Source/Runtime/Engine/Private/Rendering/SkyLightImportanceSampling.cpp b/Engine/Source/Runtime/Engine/Private/Rendering/SkyLightImportanceSampling.cpp
--- a/Engine/Source/Runtime/Engine/Private/Rendering/SkyLightImportanceSampling.cpp
+++ b/Engine/Source/Runtime/Engine/Private/Rendering/SkyLightImportanceSampling.cpp
@@ -189,6 +189,8 @@ public:
 		FRWBuffer& MipTreePdf)
 	{
 		FRHIComputeShader* ShaderRHI = RHICmdList.GetBoundComputeShader();
+		
+		RHICmdList.Transition(FRHITransitionInfo(MipTree.UAV, ERHIAccess::Unknown, ERHIAccess::SRVMask));
 
 		SetSRVParameter(RHICmdList, ShaderRHI, MipTreeParameter, MipTree.SRV);
 		SetShaderValue(RHICmdList, ShaderRHI, DimensionsParameter, Dimensions);
diff --git a/Engine/Source/Runtime/Renderer/Private/RayTracing/RayTracingLighting.cpp b/Engine/Source/Runtime/Renderer/Private/RayTracing/RayTracingLighting.cpp
--- a/Engine/Source/Runtime/Renderer/Private/RayTracing/RayTracingLighting.cpp
+++ b/Engine/Source/Runtime/Renderer/Private/RayTracing/RayTracingLighting.cpp
@@ -180,7 +180,7 @@ static void CreateRaytracingLightCullingStructure(
 	{
 		FRHITransitionInfo Transitions[] =
 		{
-			FRHITransitionInfo(LightCullVolumeUAV.GetReference(), ERHIAccess::UAVCompute, ERHIAccess::SRVMask),
+			FRHITransitionInfo(LightCullVolumeUAV.GetReference(), ERHIAccess::UAVMask, ERHIAccess::SRVMask),
 			FRHITransitionInfo(OutLightingData.LightIndices.UAV.GetReference(), ERHIAccess::UAVCompute, ERHIAccess::SRVMask)
 		};
 		RHICmdList.Transition(MakeArrayView(Transitions, UE_ARRAY_COUNT(Transitions)));
