Refactored plugin hooks to use RDG. Deprecated RHICmdList versions of desktop scene textures uniform buffer.

diff --git a/Engine/Plugins/Experimental/GPULightmass/Source/GPULightmass/Private/IrradianceCaching.h b/Engine/Plugins/Experimental/GPULightmass/Source/GPULightmass/Private/IrradianceCaching.h
--- a/Engine/Plugins/Experimental/GPULightmass/Source/GPULightmass/Private/IrradianceCaching.h
+++ b/Engine/Plugins/Experimental/GPULightmass/Source/GPULightmass/Private/IrradianceCaching.h
@@ -55,7 +55,7 @@ public:
 
 	BEGIN_SHADER_PARAMETER_STRUCT(FParameters, )
 		SHADER_PARAMETER_STRUCT_REF(FViewUniformShaderParameters, View)
-		SHADER_PARAMETER_STRUCT_REF(FSceneTextureUniformParameters, SceneTextures)
+		SHADER_PARAMETER_RDG_UNIFORM_BUFFER(FSceneTextureUniformParameters, SceneTextures)
 		SHADER_PARAMETER_STRUCT_REF(FIrradianceCachingParameters, IrradianceCachingParameters)
 		RENDER_TARGET_BINDING_SLOTS()
 	END_SHADER_PARAMETER_STRUCT()
diff --git a/Engine/Plugins/Experimental/GPULightmass/Source/GPULightmass/Private/LightmapRenderer.cpp b/Engine/Plugins/Experimental/GPULightmass/Source/GPULightmass/Private/LightmapRenderer.cpp
--- a/Engine/Plugins/Experimental/GPULightmass/Source/GPULightmass/Private/LightmapRenderer.cpp
+++ b/Engine/Plugins/Experimental/GPULightmass/Source/GPULightmass/Private/LightmapRenderer.cpp
@@ -22,6 +22,7 @@
 #include "LightmapDenoising.h"
 #include "EngineModule.h"
 #include "PostProcess/PostProcessing.h"
+#include "ScreenPass.h"
 
 class FCopyConvergedLightmapTilesCS : public FGlobalShader
 {
@@ -2911,51 +2912,54 @@ void FLightmapRenderer::BumpRevision()
 
 void FLightmapRenderer::RenderIrradianceCacheVisualization(FPostOpaqueRenderParameters& Parameters)
 {
-	FRHICommandListImmediate& RHICmdList = *Parameters.RHICmdList;
+	FRDGBuilder& GraphBuilder = *Parameters.GraphBuilder;
 
-	FVisualizeIrradianceCachePS::FParameters PassParameters;
+	auto* PassParameters = GraphBuilder.AllocParameters<FVisualizeIrradianceCachePS::FParameters>();
 	TUniformBufferRef<FViewUniformShaderParameters> Ref;
 	*Ref.GetInitReference() = Parameters.ViewUniformBuffer;
-	PassParameters.View = Ref;
-	PassParameters.SceneTextures = CreateSceneTextureUniformBuffer(RHICmdList, GMaxRHIFeatureLevel, ESceneTextureSetupMode::All);
-	PassParameters.IrradianceCachingParameters = Scene->IrradianceCache->IrradianceCachingParametersUniformBuffer;
-
-	FUniformBufferStaticBindings GlobalUniformBuffers(PassParameters.SceneTextures);
-	SCOPED_UNIFORM_BUFFER_GLOBAL_BINDINGS(RHICmdList, GlobalUniformBuffers);
-
-	FRHIRenderPassInfo RPInfo(FSceneRenderTargets::Get().GetSceneColor()->GetTargetableRHI(), ERenderTargetActions::Load_Store);
-	RHICmdList.BeginRenderPass(RPInfo, TEXT("AP_ClearIrradiance"));
-
-	RHICmdList.SetViewport(0, 0, 0.0f, Parameters.ViewportRect.Width(), Parameters.ViewportRect.Height(), 1.0f);
-
-	TShaderMapRef<FPostProcessVS> VertexShader(GetGlobalShaderMap(GMaxRHIFeatureLevel));
-	TShaderMapRef<FVisualizeIrradianceCachePS> PixelShader(GetGlobalShaderMap(GMaxRHIFeatureLevel));
-
-	FGraphicsPipelineStateInitializer GraphicsPSOInit;
-	RHICmdList.ApplyCachedRenderTargets(GraphicsPSOInit);
-	GraphicsPSOInit.RasterizerState = TStaticRasterizerState<FM_Solid, CM_None>::GetRHI();
-	GraphicsPSOInit.DepthStencilState = TStaticDepthStencilState<false, CF_Always>::GetRHI();
-	GraphicsPSOInit.BlendState = TStaticBlendState<>::GetRHI();
-	GraphicsPSOInit.PrimitiveType = PT_TriangleList;
-	GraphicsPSOInit.BoundShaderState.VertexDeclarationRHI = GFilterVertexDeclaration.VertexDeclarationRHI;
-	GraphicsPSOInit.BoundShaderState.VertexShaderRHI = VertexShader.GetVertexShader();
-	GraphicsPSOInit.BoundShaderState.PixelShaderRHI = PixelShader.GetPixelShader();
-
-	SetGraphicsPipelineState(RHICmdList, GraphicsPSOInit);
-
-	SetShaderParameters(RHICmdList, PixelShader, PixelShader.GetPixelShader(), PassParameters);
-
-	DrawRectangle(
-		RHICmdList,
-		0, 0,
-		Parameters.ViewportRect.Width(), Parameters.ViewportRect.Height(),
-		0, 0,
-		Parameters.ViewportRect.Width(), Parameters.ViewportRect.Height(),
-		FIntPoint(Parameters.ViewportRect.Width(), Parameters.ViewportRect.Height()),
-		FSceneRenderTargets::Get().GetBufferSizeXY(),
-		VertexShader);
-
-	RHICmdList.EndRenderPass();
+	PassParameters->View = Ref;
+	PassParameters->SceneTextures = Parameters.SceneTexturesUniformParams;
+	PassParameters->IrradianceCachingParameters = Scene->IrradianceCache->IrradianceCachingParametersUniformBuffer;
+	PassParameters->RenderTargets[0] = FRenderTargetBinding(Parameters.ColorTexture, ERenderTargetLoadAction::ELoad);
+
+	const FIntRect ViewportRect = Parameters.ViewportRect;
+	const FIntPoint TextureExtent = Parameters.ColorTexture->Desc.Extent;
+
+	GraphBuilder.AddPass(
+		RDG_EVENT_NAME("ClearIrradiance"),
+		PassParameters,
+		ERDGPassFlags::Raster,
+		[ViewportRect, PassParameters, TextureExtent] (FRHICommandList& RHICmdList)
+	{
+		RHICmdList.SetViewport(0, 0, 0.0f, ViewportRect.Width(), ViewportRect.Height(), 1.0f);
+
+		TShaderMapRef<FPostProcessVS> VertexShader(GetGlobalShaderMap(GMaxRHIFeatureLevel));
+		TShaderMapRef<FVisualizeIrradianceCachePS> PixelShader(GetGlobalShaderMap(GMaxRHIFeatureLevel));
+
+		FGraphicsPipelineStateInitializer GraphicsPSOInit;
+		RHICmdList.ApplyCachedRenderTargets(GraphicsPSOInit);
+		GraphicsPSOInit.RasterizerState = TStaticRasterizerState<FM_Solid, CM_None>::GetRHI();
+		GraphicsPSOInit.DepthStencilState = TStaticDepthStencilState<false, CF_Always>::GetRHI();
+		GraphicsPSOInit.BlendState = TStaticBlendState<>::GetRHI();
+		GraphicsPSOInit.PrimitiveType = PT_TriangleList;
+		GraphicsPSOInit.BoundShaderState.VertexDeclarationRHI = GFilterVertexDeclaration.VertexDeclarationRHI;
+		GraphicsPSOInit.BoundShaderState.VertexShaderRHI = VertexShader.GetVertexShader();
+		GraphicsPSOInit.BoundShaderState.PixelShaderRHI = PixelShader.GetPixelShader();
+
+		SetGraphicsPipelineState(RHICmdList, GraphicsPSOInit);
+
+		SetShaderParameters(RHICmdList, PixelShader, PixelShader.GetPixelShader(), *PassParameters);
+
+		DrawRectangle(
+			RHICmdList,
+			0, 0,
+			ViewportRect.Width(), ViewportRect.Height(),
+			0, 0,
+			ViewportRect.Width(), ViewportRect.Height(),
+			FIntPoint(ViewportRect.Width(), ViewportRect.Height()),
+			TextureExtent,
+			VertexShader);
+	});
 }
 
 }
diff --git a/Engine/Plugins/Experimental/RemoteSession/Source/RemoteSession/Private/Channels/RemoteSessionARCameraChannel.cpp b/Engine/Plugins/Experimental/RemoteSession/Source/RemoteSession/Private/Channels/RemoteSessionARCameraChannel.cpp
--- a/Engine/Plugins/Experimental/RemoteSession/Source/RemoteSession/Private/Channels/RemoteSessionARCameraChannel.cpp
+++ b/Engine/Plugins/Experimental/RemoteSession/Source/RemoteSession/Private/Channels/RemoteSessionARCameraChannel.cpp
@@ -27,6 +27,7 @@
 #include "PostProcess/SceneFilterRendering.h"
 #include "PostProcess/PostProcessMaterial.h"
 #include "EngineModule.h"
+#include "ScreenPass.h"
 
 #include "GeneralProjectSettings.h"
 #include "ARTextures.h"
@@ -147,11 +148,11 @@ private:
 	virtual void BeginRenderViewFamily(FSceneViewFamily& InViewFamily) override;
 	virtual void PreRenderView_RenderThread(FRHICommandListImmediate& RHICmdList, FSceneView& InView) override;
 	virtual void PreRenderViewFamily_RenderThread(FRHICommandListImmediate& RHICmdList, FSceneViewFamily& InViewFamily) override;
-	virtual void PostRenderViewFamily_RenderThread(FRHICommandListImmediate& RHICmdList, FSceneViewFamily& InViewFamily) override;
+	virtual void PostRenderViewFamily_RenderThread(FRDGBuilder& GraphBuilder, FSceneViewFamily& InViewFamily) override;
 	virtual bool IsActiveThisFrame(FViewport* InViewport) const override;
 	//~ISceneViewExtension interface
 
-	void RenderARCamera_RenderThread(FRHICommandListImmediate& RHICmdList, const FSceneView& InView);
+	void RenderARCamera_RenderThread(FRDGBuilder& GraphBuilder, const FSceneView& InView, FRDGTextureRef ViewFamilyTexture);
 
 	/** The channel this is rendering for */
 	FRemoteSessionARCameraChannel& Channel;
@@ -231,7 +232,7 @@ void FARCameraSceneViewExtension::PreRenderViewFamily_RenderThread(FRHICommandLi
 
 }
 
-void FARCameraSceneViewExtension::PostRenderViewFamily_RenderThread(FRHICommandListImmediate& RHICmdList, FSceneViewFamily& InViewFamily)
+void FARCameraSceneViewExtension::PostRenderViewFamily_RenderThread(FRDGBuilder& GraphBuilder, FSceneViewFamily& InViewFamily)
 {
 	if (PPMaterial == nullptr || !PPMaterial->IsValidLowLevel() ||
 		VertexBufferRHI == nullptr || !VertexBufferRHI.IsValid() ||
@@ -240,70 +241,86 @@ void FARCameraSceneViewExtension::PostRenderViewFamily_RenderThread(FRHICommandL
 		return;
 	}
 
+	FRDGTextureRef ViewFamilyTexture = TryCreateViewFamilyTexture(GraphBuilder, InViewFamily);
+	if (!ViewFamilyTexture)
+	{
+		return;
+	}
+
 	for (int32 ViewIndex = 0; ViewIndex < InViewFamily.Views.Num(); ++ViewIndex)
 	{
-		RenderARCamera_RenderThread(RHICmdList, *InViewFamily.Views[ViewIndex]);
+		RenderARCamera_RenderThread(GraphBuilder, *InViewFamily.Views[ViewIndex], ViewFamilyTexture);
 	}
 }
 
-void FARCameraSceneViewExtension::RenderARCamera_RenderThread(FRHICommandListImmediate& RHICmdList, const FSceneView& InView)
+BEGIN_SHADER_PARAMETER_STRUCT(FRenderARCameraPassParameters, )
+	SHADER_PARAMETER_RDG_UNIFORM_BUFFER(FSceneTextureUniformParameters, SceneTextures)
+	RENDER_TARGET_BINDING_SLOTS()
+END_SHADER_PARAMETER_STRUCT()
+
+void FARCameraSceneViewExtension::RenderARCamera_RenderThread(FRDGBuilder& GraphBuilder, const FSceneView& InView, FRDGTextureRef ViewFamilyTexture)
 {
 #if PLATFORM_DESKTOP
-	const auto FeatureLevel = InView.GetFeatureLevel();
-
-	IRendererModule& RendererModule = GetRendererModule();
-
-	FUniformBufferRHIRef PassUniformBuffer = CreateSceneTextureUniformBufferDependentOnShadingPath(
-		RHICmdList,
-		InView.GetFeatureLevel(),
-		ESceneTextureSetupMode::None);
-	FUniformBufferStaticBindings GlobalUniformBuffers(PassUniformBuffer);
-	SCOPED_UNIFORM_BUFFER_GLOBAL_BINDINGS(RHICmdList, GlobalUniformBuffers);
-	const FMaterialRenderProxy* MaterialProxy = PPMaterial->GetRenderProxy();
-	const FMaterial& CameraMaterial = MaterialProxy->GetMaterialWithFallback(FeatureLevel, MaterialProxy);
-	const FMaterialShaderMap* const MaterialShaderMap = CameraMaterial.GetRenderingThreadShaderMap();
-
-	FGraphicsPipelineStateInitializer GraphicsPSOInit;
-	RHICmdList.ApplyCachedRenderTargets(GraphicsPSOInit);
-
-	GraphicsPSOInit.BlendState = TStaticBlendState<>::GetRHI();
-	GraphicsPSOInit.RasterizerState = TStaticRasterizerState<>::GetRHI();
-	GraphicsPSOInit.DepthStencilState = TStaticDepthStencilState<false, CF_DepthNearOrEqual>::GetRHI();
-	GraphicsPSOInit.PrimitiveType = PT_TriangleList;
-	GraphicsPSOInit.BoundShaderState.VertexDeclarationRHI = GFilterVertexDeclaration.VertexDeclarationRHI;
-
-	TShaderRef<FRemoteSessionARCameraVS> VertexShader = MaterialShaderMap->GetShader<FRemoteSessionARCameraVS>();
-	TShaderRef<FRemoteSessionARCameraPS> PixelShader = MaterialShaderMap->GetShader<FRemoteSessionARCameraPS>();
-	GraphicsPSOInit.BoundShaderState.VertexShaderRHI = VertexShader.GetVertexShader();
-	GraphicsPSOInit.BoundShaderState.PixelShaderRHI = PixelShader.GetPixelShader();
-
-	SetGraphicsPipelineState(RHICmdList, GraphicsPSOInit);
-
-	const FIntPoint ViewSize = InView.UnconstrainedViewRect.Size();
-	FDrawRectangleParameters Parameters;
-	Parameters.PosScaleBias = FVector4(ViewSize.X, ViewSize.Y, 0, 0);
-	Parameters.UVScaleBias = FVector4(1.0f, 1.0f, 0.0f, 0.0f);
-	Parameters.InvTargetSizeAndTextureSize = FVector4(
+	if (!VertexBufferRHI || !IndexBufferRHI.IsValid())
+	{
+		return;
+	}
+
+	auto* PassParameters = GraphBuilder.AllocParameters<FRenderARCameraPassParameters>();
+	PassParameters->RenderTargets[0] = FRenderTargetBinding(ViewFamilyTexture, ERenderTargetLoadAction::ELoad);
+	PassParameters->SceneTextures = CreateSceneTextureUniformBuffer(GraphBuilder, InView.FeatureLevel, ESceneTextureSetupMode::None);
+
+	GraphBuilder.AddPass(
+		RDG_EVENT_NAME("ARCameraOverlay"),
+		PassParameters,
+		ERDGPassFlags::Raster,
+		[this, &InView](FRHICommandList& RHICmdList)
+	{
+		const auto FeatureLevel = InView.GetFeatureLevel();
+
+		const FMaterialRenderProxy* MaterialProxy = PPMaterial->GetRenderProxy();
+		const FMaterial& CameraMaterial = MaterialProxy->GetMaterialWithFallback(FeatureLevel, MaterialProxy);
+		const FMaterialShaderMap* const MaterialShaderMap = CameraMaterial.GetRenderingThreadShaderMap();
+
+		FGraphicsPipelineStateInitializer GraphicsPSOInit;
+		RHICmdList.ApplyCachedRenderTargets(GraphicsPSOInit);
+
+		GraphicsPSOInit.BlendState = TStaticBlendState<>::GetRHI();
+		GraphicsPSOInit.RasterizerState = TStaticRasterizerState<>::GetRHI();
+		GraphicsPSOInit.DepthStencilState = TStaticDepthStencilState<false, CF_DepthNearOrEqual>::GetRHI();
+		GraphicsPSOInit.PrimitiveType = PT_TriangleList;
+		GraphicsPSOInit.BoundShaderState.VertexDeclarationRHI = GFilterVertexDeclaration.VertexDeclarationRHI;
+
+		TShaderRef<FRemoteSessionARCameraVS> VertexShader = MaterialShaderMap->GetShader<FRemoteSessionARCameraVS>();
+		TShaderRef<FRemoteSessionARCameraPS> PixelShader = MaterialShaderMap->GetShader<FRemoteSessionARCameraPS>();
+		GraphicsPSOInit.BoundShaderState.VertexShaderRHI = VertexShader.GetVertexShader();
+		GraphicsPSOInit.BoundShaderState.PixelShaderRHI = PixelShader.GetPixelShader();
+
+		SetGraphicsPipelineState(RHICmdList, GraphicsPSOInit);
+
+		const FIntPoint ViewSize = InView.UnconstrainedViewRect.Size();
+		FDrawRectangleParameters Parameters;
+		Parameters.PosScaleBias = FVector4(ViewSize.X, ViewSize.Y, 0, 0);
+		Parameters.UVScaleBias = FVector4(1.0f, 1.0f, 0.0f, 0.0f);
+		Parameters.InvTargetSizeAndTextureSize = FVector4(
 			1.0f / ViewSize.X, 1.0f / ViewSize.Y,
 			1.0f, 1.0f);
 
-	SetUniformBufferParameterImmediate(RHICmdList, VertexShader.GetVertexShader(), VertexShader->GetUniformBufferParameter<FDrawRectangleParameters>(), Parameters);
-	VertexShader->SetParameters(RHICmdList, InView);
-	PixelShader->SetParameters(RHICmdList, InView, MaterialProxy);
+		SetUniformBufferParameterImmediate(RHICmdList, VertexShader.GetVertexShader(), VertexShader->GetUniformBufferParameter<FDrawRectangleParameters>(), Parameters);
+		VertexShader->SetParameters(RHICmdList, InView);
+		PixelShader->SetParameters(RHICmdList, InView, MaterialProxy);
 
-	if (VertexBufferRHI && IndexBufferRHI.IsValid())
-	{
 		RHICmdList.SetStreamSource(0, VertexBufferRHI, 0);
 		RHICmdList.DrawIndexedPrimitive(
-				IndexBufferRHI,
-				/*BaseVertexIndex=*/ 0,
-				/*MinIndex=*/ 0,
-				/*NumVertices=*/ 4,
-				/*StartIndex=*/ 0,
-				/*NumPrimitives=*/ 2,
-				/*NumInstances=*/ 1
-			);
-	}
+			IndexBufferRHI,
+			/*BaseVertexIndex=*/ 0,
+			/*MinIndex=*/ 0,
+			/*NumVertices=*/ 4,
+			/*StartIndex=*/ 0,
+			/*NumPrimitives=*/ 2,
+			/*NumInstances=*/ 1
+		);
+	});
 #endif
 }
 
diff --git a/Engine/Plugins/FX/Niagara/Source/Niagara/Private/NiagaraWorldManager.cpp b/Engine/Plugins/FX/Niagara/Source/Niagara/Private/NiagaraWorldManager.cpp
--- a/Engine/Plugins/FX/Niagara/Source/Niagara/Private/NiagaraWorldManager.cpp
+++ b/Engine/Plugins/FX/Niagara/Source/Niagara/Private/NiagaraWorldManager.cpp
@@ -136,6 +136,38 @@ namespace FNiagaraUtilities
 	}
 }
 
+BEGIN_SHADER_PARAMETER_STRUCT(FExtractUniformBufferParameters, )
+	SHADER_PARAMETER_RDG_UNIFORM_BUFFER(FSceneTextureUniformParameters, SceneTextures)
+	RDG_TEXTURE_ACCESS(Depth, ERHIAccess::SRVGraphics)
+	RDG_TEXTURE_ACCESS(Normal, ERHIAccess::SRVGraphics)
+	RDG_TEXTURE_ACCESS(Velocity, ERHIAccess::SRVGraphics)
+END_SHADER_PARAMETER_STRUCT()
+
+void FNiagaraViewDataMgr::PostOpaqueRender(FPostOpaqueRenderParameters& Params)
+{
+	FRDGBuilder& GraphBuilder = *Params.GraphBuilder;
+
+	ViewUniformBuffer = Params.ViewUniformBuffer;
+
+	FExtractUniformBufferParameters* PassParameters = GraphBuilder.AllocParameters<FExtractUniformBufferParameters>();
+	PassParameters->SceneTextures = Params.SceneTexturesUniformParams;
+	PassParameters->Depth = Params.DepthTexture;
+	PassParameters->Normal = Params.NormalTexture;
+	PassParameters->Velocity = Params.VelocityTexture;
+
+	// Niagara is not ported to RDG and is manually accessing RHI resources. This extraction is technically not safe,
+	// because the lifetime of the textures are not guaranteed outside of the pass. However, we happen to know that
+	// they will outlast Niagara.
+
+	GraphBuilder.AddPass({}, PassParameters, ERDGPassFlags::Raster | ERDGPassFlags::SkipRenderPass | ERDGPassFlags::NeverCull, [this, PassParameters](FRHICommandList&)
+	{
+		SceneTexturesUniformParams = PassParameters->SceneTextures->GetRHIRef();
+		SceneDepthTexture = static_cast<FRHITexture2D*>(PassParameters->Depth ? PassParameters->Depth->GetRHI() : nullptr);
+		SceneNormalTexture = static_cast<FRHITexture2D*>(PassParameters->Normal ? PassParameters->Normal->GetRHI() : nullptr);
+		SceneVelocityTexture = static_cast<FRHITexture2D*>(PassParameters->Velocity ? PassParameters->Velocity->GetRHI() : nullptr);
+	});
+}
+
 FNiagaraViewDataMgr::FNiagaraViewDataMgr()
 	: FRenderResource()
 	, SceneDepthTexture(nullptr)
diff --git a/Engine/Plugins/FX/Niagara/Source/Niagara/Public/NiagaraWorldManager.h b/Engine/Plugins/FX/Niagara/Source/Niagara/Public/NiagaraWorldManager.h
--- a/Engine/Plugins/FX/Niagara/Source/Niagara/Public/NiagaraWorldManager.h
+++ b/Engine/Plugins/FX/Niagara/Source/Niagara/Public/NiagaraWorldManager.h
@@ -36,14 +36,7 @@ public:
 	static void Init();
 	static void Shutdown();
 
-	void PostOpaqueRender(FPostOpaqueRenderParameters& Params)
-	{
-		SceneDepthTexture = Params.DepthTexture;
-		ViewUniformBuffer = Params.ViewUniformBuffer;
-		SceneNormalTexture = Params.NormalTexture;
-		SceneVelocityTexture = Params.VelocityTexture;
-		SceneTexturesUniformParams = Params.SceneTexturesUniformParams;
-	}
+	NIAGARA_API void PostOpaqueRender(FPostOpaqueRenderParameters& Params);
 
 	FRHITexture2D* GetSceneDepthTexture() { return SceneDepthTexture; }
 	FRHITexture2D* GetSceneNormalTexture() { return SceneNormalTexture; }
diff --git a/Engine/Plugins/Runtime/AR/AppleAR/AppleARKit/Source/AppleARKit/Private/AppleARKitVideoOverlay.cpp b/Engine/Plugins/Runtime/AR/AppleAR/AppleARKit/Source/AppleARKit/Private/AppleARKitVideoOverlay.cpp
--- a/Engine/Plugins/Runtime/AR/AppleAR/AppleARKit/Source/AppleARKit/Private/AppleARKitVideoOverlay.cpp
+++ b/Engine/Plugins/Runtime/AR/AppleAR/AppleARKit/Source/AppleARKit/Private/AppleARKitVideoOverlay.cpp
@@ -342,10 +342,6 @@ void FAppleARKitVideoOverlay::RenderVideoOverlayWithMaterial(FRHICommandListImme
 	const auto FeatureLevel = InView.GetFeatureLevel();
 	IRendererModule& RendererModule = GetRendererModule();
 
-	FUniformBufferRHIRef PassUniformBuffer = CreateSceneTextureUniformBufferDependentOnShadingPath(RHICmdList, FeatureLevel, ESceneTextureSetupMode::None);
-	FUniformBufferStaticBindings GlobalUniformBuffers(PassUniformBuffer);
-	SCOPED_UNIFORM_BUFFER_GLOBAL_BINDINGS(RHICmdList, GlobalUniformBuffers);
-
 	const FMaterialRenderProxy* MaterialProxy = RenderingOverlayMaterial->GetRenderProxy();
 	const FMaterial& CameraMaterial = MaterialProxy->GetMaterialWithFallback(FeatureLevel, MaterialProxy);
 	const FMaterialShaderMap* const MaterialShaderMap = CameraMaterial.GetRenderingThreadShaderMap();
diff --git a/Engine/Plugins/Runtime/AR/Google/GoogleARCore/Source/GoogleARCoreRendering/Private/GoogleARCorePassthroughCameraRenderer.cpp b/Engine/Plugins/Runtime/AR/Google/GoogleARCore/Source/GoogleARCoreRendering/Private/GoogleARCorePassthroughCameraRenderer.cpp
--- a/Engine/Plugins/Runtime/AR/Google/GoogleARCore/Source/GoogleARCoreRendering/Private/GoogleARCorePassthroughCameraRenderer.cpp
+++ b/Engine/Plugins/Runtime/AR/Google/GoogleARCore/Source/GoogleARCoreRendering/Private/GoogleARCorePassthroughCameraRenderer.cpp
@@ -191,10 +191,6 @@ void FGoogleARCorePassthroughCameraRenderer::RenderVideoOverlayWithMaterial(FRHI
 
 	if (FeatureLevel <= ERHIFeatureLevel::ES3_1)
 	{
-		FUniformBufferRHIRef PassUniformBuffer = CreateSceneTextureUniformBufferDependentOnShadingPath(RHICmdList, FeatureLevel, ESceneTextureSetupMode::None);
-		FUniformBufferStaticBindings GlobalUniformBuffers(PassUniformBuffer);
-		SCOPED_UNIFORM_BUFFER_GLOBAL_BINDINGS(RHICmdList, GlobalUniformBuffers);
-
 		const FMaterialRenderProxy* MaterialProxy = OverlayMaterialToUse->GetRenderProxy();
 		const FMaterial& CameraMaterial = MaterialProxy->GetMaterialWithFallback(FeatureLevel, MaterialProxy);
 		const FMaterialShaderMap* MaterialShaderMap = CameraMaterial.GetRenderingThreadShaderMap();
diff --git a/Engine/Plugins/Runtime/TextureShare/Source/TextureShare/Private/TextureShareDisplayManager.cpp b/Engine/Plugins/Runtime/TextureShare/Source/TextureShare/Private/TextureShareDisplayManager.cpp
--- a/Engine/Plugins/Runtime/TextureShare/Source/TextureShare/Private/TextureShareDisplayManager.cpp
+++ b/Engine/Plugins/Runtime/TextureShare/Source/TextureShare/Private/TextureShareDisplayManager.cpp
@@ -10,6 +10,7 @@
 #include "Engine/GameEngine.h"
 #include "Engine/GameInstance.h"
 #include "Engine/GameViewportClient.h"
+#include "RenderGraphUtils.h"
 
 #include "Slate\SceneViewport.h"
 
@@ -119,12 +120,15 @@ void FTextureShareDisplayManager::PreRenderViewFamily_RenderThread(FRHICommandLi
 	CurrentSceneViewFamily = &InViewFamily;
 }
 
-void FTextureShareDisplayManager::OnResolvedSceneColor_RenderThread(FRHICommandListImmediate& RHICmdList, class FSceneRenderTargets& SceneContext)
+void FTextureShareDisplayManager::OnResolvedSceneColor_RenderThread(FRDGBuilder& GraphBuilder, class FSceneRenderTargets& SceneContext)
 {
 	// Forward renderer cb with saved viewfamily (Call from RendererModule->GetResolvedSceneColorCallbacks)
 	if (IsSceneSharingValid() && CurrentSceneViewFamily)
 	{
-		TextureShareModule.OnResolvedSceneColor_RenderThread(RHICmdList, SceneContext, *CurrentSceneViewFamily);
+		AddPass(GraphBuilder, [this, &SceneContext](FRHICommandListImmediate& RHICmdList)
+		{
+			TextureShareModule.OnResolvedSceneColor_RenderThread(RHICmdList, SceneContext, *CurrentSceneViewFamily);
+		});
 	}
 }
 
diff --git a/Engine/Plugins/Runtime/TextureShare/Source/TextureShare/Private/TextureShareDisplayManager.h b/Engine/Plugins/Runtime/TextureShare/Source/TextureShare/Private/TextureShareDisplayManager.h
--- a/Engine/Plugins/Runtime/TextureShare/Source/TextureShare/Private/TextureShareDisplayManager.h
+++ b/Engine/Plugins/Runtime/TextureShare/Source/TextureShare/Private/TextureShareDisplayManager.h
@@ -13,6 +13,7 @@ class FTextureShareDisplayExtension;
 class FViewport;
 class FTextureShareModule;
 class FSceneViewFamily;
+class FRDGBuilder;
 
 class FTextureShareDisplayManager
 {
@@ -40,7 +41,7 @@ public:
 
 private:
 	/** Rendered callback (get scene textures to share) */
-	void OnResolvedSceneColor_RenderThread(FRHICommandListImmediate& RHICmdList, class FSceneRenderTargets& SceneContext);
+	void OnResolvedSceneColor_RenderThread(FRDGBuilder& GraphBuilder, class FSceneRenderTargets& SceneContext);
 
 	bool IsSceneSharingValid() const
 	{ return bIsRenderedCallbackAssigned && DisplayExtensions.Num()>0; }
diff --git a/Engine/Source/Runtime/Engine/Private/SceneViewExtension.cpp b/Engine/Source/Runtime/Engine/Private/SceneViewExtension.cpp
--- a/Engine/Source/Runtime/Engine/Private/SceneViewExtension.cpp
+++ b/Engine/Source/Runtime/Engine/Private/SceneViewExtension.cpp
@@ -2,6 +2,7 @@
 
 #include "SceneViewExtension.h"
 #include "Engine/Engine.h"
+#include "RenderGraphUtils.h"
 
 //
 // FSceneViewExtensionBase
@@ -31,7 +32,37 @@ bool FSceneViewExtensionBase::IsActiveThisFrameInContext(FSceneViewExtensionCont
 	return IsActiveThisFrame(Context.Viewport);
 }
 
+void ISceneViewExtension::PreRenderViewFamily_RenderThread(FRDGBuilder& GraphBuilder, FSceneViewFamily& InViewFamily)
+{
+	AddPass(GraphBuilder, [this, &InViewFamily](FRHICommandListImmediate& RHICmdList)
+	{
+		PreRenderViewFamily_RenderThread(RHICmdList, InViewFamily);
+	});
+}
+
+void ISceneViewExtension::PreRenderView_RenderThread(FRDGBuilder& GraphBuilder, FSceneView& InView)
+{
+	AddPass(GraphBuilder, [this, &InView](FRHICommandListImmediate& RHICmdList)
+	{
+		PreRenderView_RenderThread(RHICmdList, InView);
+	});
+}
 
+void ISceneViewExtension::PostRenderViewFamily_RenderThread(FRDGBuilder& GraphBuilder, FSceneViewFamily& InViewFamily)
+{
+	AddPass(GraphBuilder, [this, &InViewFamily](FRHICommandListImmediate& RHICmdList)
+	{
+		PostRenderViewFamily_RenderThread(RHICmdList, InViewFamily);
+	});
+}
+
+void ISceneViewExtension::PostRenderView_RenderThread(FRDGBuilder& GraphBuilder, FSceneView& InView)
+{
+	AddPass(GraphBuilder, [this, &InView](FRHICommandListImmediate& RHICmdList)
+	{
+		PostRenderView_RenderThread(RHICmdList, InView);
+	});
+}
 
 //
 // FSceneViewExtensions
diff --git a/Engine/Source/Runtime/Engine/Public/SceneViewExtension.h b/Engine/Source/Runtime/Engine/Public/SceneViewExtension.h
--- a/Engine/Source/Runtime/Engine/Public/SceneViewExtension.h
+++ b/Engine/Source/Runtime/Engine/Public/SceneViewExtension.h
@@ -138,12 +138,14 @@ public:
     /**
      * Called on render thread at the start of rendering.
      */
-    virtual void PreRenderViewFamily_RenderThread(FRHICommandListImmediate& RHICmdList, FSceneViewFamily& InViewFamily) = 0;
+	virtual void PreRenderViewFamily_RenderThread(FRHICommandListImmediate& RHICmdList, FSceneViewFamily& InViewFamily) {};
+	virtual ENGINE_API void PreRenderViewFamily_RenderThread(FRDGBuilder& GraphBuilder, FSceneViewFamily& InViewFamily);
 
 	/**
      * Called on render thread at the start of rendering, for each view, after PreRenderViewFamily_RenderThread call.
      */
-    virtual void PreRenderView_RenderThread(FRHICommandListImmediate& RHICmdList, FSceneView& InView) = 0;
+	virtual void PreRenderView_RenderThread(FRHICommandListImmediate& RHICmdList, FSceneView& InView) {};
+	virtual ENGINE_API void PreRenderView_RenderThread(FRDGBuilder& GraphBuilder, FSceneView& InView);
 
 	/**
 	 * Called right after Base Pass rendering finished
@@ -164,11 +166,13 @@ public:
 	 * Allows to render content after the 3D content scene, useful for debugging
 	 */
 	virtual void PostRenderViewFamily_RenderThread(FRHICommandListImmediate& RHICmdList, FSceneViewFamily& InViewFamily) {}
+	virtual ENGINE_API void PostRenderViewFamily_RenderThread(FRDGBuilder& GraphBuilder, FSceneViewFamily& InViewFamily);
 
 	/**
 	 * Allows to render content after the 3D content scene, useful for debugging
 	 */
 	virtual void PostRenderView_RenderThread(FRHICommandListImmediate& RHICmdList, FSceneView& InView) {}
+	virtual ENGINE_API void PostRenderView_RenderThread(FRDGBuilder& GraphBuilder, FSceneView& InView);
 
 	/**
      * Called to determine view extensions priority in relation to other view extensions, higher comes first
diff --git a/Engine/Source/Runtime/RenderCore/Public/RenderGraphDefinitions.h b/Engine/Source/Runtime/RenderCore/Public/RenderGraphDefinitions.h
--- a/Engine/Source/Runtime/RenderCore/Public/RenderGraphDefinitions.h
+++ b/Engine/Source/Runtime/RenderCore/Public/RenderGraphDefinitions.h
@@ -3,8 +3,6 @@
 #pragma once
 
 #include "CoreMinimal.h"
-#include "RHI.h"
-#include "RendererInterface.h"
 #include "ProfilingDebugging/RealtimeGPUProfiler.h"
 
 /** DEFINES */
@@ -288,7 +286,7 @@ private:
 	};
 
 	FMemStackBase MemStack;
-	TArray<FTrackedAlloc*, SceneRenderingAllocator> TrackedAllocs;
+	TArray<FTrackedAlloc*, TMemStackAllocator<>> TrackedAllocs;
 };
 
 /** HANDLE UTILITIES */
@@ -416,14 +414,14 @@ public:
 	}
 
 private:
-	TArray<ObjectType*, SceneRenderingAllocator> Array;
+	TArray<ObjectType*, TMemStackAllocator<>> Array;
 };
 
 /** Specialization of bit array with compile-time type checking for handles and a pre-configured allocator. */
 template <typename HandleType>
-class TRDGHandleBitArray : public TBitArray<SceneRenderingBitArrayAllocator>
+class TRDGHandleBitArray : public TBitArray<TInlineAllocator<4, TMemStackAllocator<>>>
 {
-	using Base = TBitArray<SceneRenderingBitArrayAllocator>;
+	using Base = TBitArray<TInlineAllocator<4, TMemStackAllocator<>>>;
 public:
 	TRDGHandleBitArray() = default;
 	TRDGHandleBitArray(const TRDGHandleBitArray&) = default;
@@ -536,7 +534,7 @@ class FRDGPass;
 using FRDGPassRef = const FRDGPass*;
 using FRDGPassHandle = TRDGHandle<FRDGPass, uint16>;
 using FRDGPassRegistry = TRDGHandleRegistry<FRDGPassHandle>;
-using FRDGPassHandleArray = TArray<FRDGPassHandle, TInlineAllocator<4, SceneRenderingAllocator>>;
+using FRDGPassHandleArray = TArray<FRDGPassHandle, TInlineAllocator<4, TMemStackAllocator<>>>;
 using FRDGPassBitArray = TRDGHandleBitArray<FRDGPassHandle>;
 
 class FRDGUniformBuffer;
diff --git a/Engine/Source/Runtime/RenderCore/Public/RenderGraphEvent.h b/Engine/Source/Runtime/RenderCore/Public/RenderGraphEvent.h
--- a/Engine/Source/Runtime/RenderCore/Public/RenderGraphEvent.h
+++ b/Engine/Source/Runtime/RenderCore/Public/RenderGraphEvent.h
@@ -3,6 +3,7 @@
 #pragma once
 
 #include "RenderGraphDefinitions.h"
+#include "RendererInterface.h"
 #include "ProfilingDebugging/RealtimeGPUProfiler.h"
 
 /** Macros for create render graph event names and scopes.
diff --git a/Engine/Source/Runtime/RenderCore/Public/RenderGraphParameter.h b/Engine/Source/Runtime/RenderCore/Public/RenderGraphParameter.h
--- a/Engine/Source/Runtime/RenderCore/Public/RenderGraphParameter.h
+++ b/Engine/Source/Runtime/RenderCore/Public/RenderGraphParameter.h
@@ -3,6 +3,7 @@
 #pragma once
 
 #include "RenderGraphDefinitions.h"
+#include "ShaderParameterMacros.h"
 
 /** A helper class for identifying and accessing a render graph pass parameter. */
 class FRDGParameter final
diff --git a/Engine/Source/Runtime/RenderCore/Public/RenderGraphResources.h b/Engine/Source/Runtime/RenderCore/Public/RenderGraphResources.h
--- a/Engine/Source/Runtime/RenderCore/Public/RenderGraphResources.h
+++ b/Engine/Source/Runtime/RenderCore/Public/RenderGraphResources.h
@@ -4,6 +4,7 @@
 
 #include "RenderGraphParameter.h"
 #include "RenderGraphTextureSubresource.h"
+#include "RendererInterface.h"
 
 struct FPooledRenderTarget;
 class FRenderTargetPool;
diff --git a/Engine/Source/Runtime/RenderCore/Public/RendererInterface.h b/Engine/Source/Runtime/RenderCore/Public/RendererInterface.h
--- a/Engine/Source/Runtime/RenderCore/Public/RendererInterface.h
+++ b/Engine/Source/Runtime/RenderCore/Public/RendererInterface.h
@@ -16,6 +16,7 @@
 #include "Misc/EnumClassFlags.h"
 #include "UniformBuffer.h"
 #include "VirtualTexturing.h"
+#include "RenderGraphDefinitions.h"
 
 class FCanvas;
 class FMaterial;
@@ -558,13 +559,14 @@ public:
 	FIntRect ViewportRect;
 	FMatrix ViewMatrix;
 	FMatrix ProjMatrix;
-	FRHITexture2D* DepthTexture = nullptr;
-	FRHITexture2D* NormalTexture = nullptr;
-	FRHITexture2D* VelocityTexture = nullptr;
-	FRHITexture2D* SmallDepthTexture = nullptr;
-	FRHICommandListImmediate* RHICmdList = nullptr;
+	FRDGTexture* ColorTexture = nullptr;
+	FRDGTexture* DepthTexture = nullptr;
+	FRDGTexture* NormalTexture = nullptr;
+	FRDGTexture* VelocityTexture = nullptr;
+	FRDGTexture* SmallDepthTexture = nullptr;
+	FRDGBuilder* GraphBuilder = nullptr;
 	FRHIUniformBuffer* ViewUniformBuffer = nullptr;
-	TUniformBufferRef<FSceneTextureUniformParameters> SceneTexturesUniformParams;
+	TRDGUniformBufferRef<FSceneTextureUniformParameters> SceneTexturesUniformParams = nullptr;
 	const FGlobalDistanceFieldParameterData* GlobalDistanceFieldParams = nullptr;
 	void* Uid = nullptr; // A unique identifier for the view.
 };
@@ -760,7 +762,7 @@ public:
 	virtual void RemoveOverlayRenderDelegate(FDelegateHandle OverlayRenderDelegate) = 0;
 
 	/** Delegate that is called upon resolving scene color. */
-	DECLARE_MULTICAST_DELEGATE_TwoParams(FOnResolvedSceneColor, FRHICommandListImmediate& /*RHICmdList*/, class FSceneRenderTargets& /*SceneContext*/);
+	DECLARE_MULTICAST_DELEGATE_TwoParams(FOnResolvedSceneColor, FRDGBuilder& /*GraphBuilder*/, FSceneRenderTargets& /*SceneContext*/);
 
 	/** Accessor for post scene color resolve delegates */
 	virtual FOnResolvedSceneColor& GetResolvedSceneColorCallbacks() = 0;
diff --git a/Engine/Source/Runtime/Renderer/Private/BasePassRendering.cpp b/Engine/Source/Runtime/Renderer/Private/BasePassRendering.cpp
--- a/Engine/Source/Runtime/Renderer/Private/BasePassRendering.cpp
+++ b/Engine/Source/Runtime/Renderer/Private/BasePassRendering.cpp
@@ -766,6 +766,11 @@ static void ClearGBufferAtMaxZ(
 	}
 }
 
+BEGIN_SHADER_PARAMETER_STRUCT(FPostBasePassViewExtensionParameters, )
+	SHADER_PARAMETER_RDG_UNIFORM_BUFFER(FSceneTextureUniformParameters, SceneTextures)
+	RENDER_TARGET_BINDING_SLOTS()
+END_SHADER_PARAMETER_STRUCT()
+
 void FDeferredShadingSceneRenderer::RenderBasePass(
 	FRDGBuilder& GraphBuilder,
 	FExclusiveDepthStencil::Type BasePassDepthStencilAccess,
@@ -943,8 +948,9 @@ void FDeferredShadingSceneRenderer::RenderBasePass(
 	{
 		SCOPE_CYCLE_COUNTER(STAT_FDeferredShadingSceneRenderer_ViewExtensionPostRenderBasePass);
 		RDG_EVENT_SCOPE(GraphBuilder, "BasePass_ViewExtensions");
-		auto* PassParameters = GraphBuilder.AllocParameters<FRenderTargetParameters>();
+		auto* PassParameters = GraphBuilder.AllocParameters<FPostBasePassViewExtensionParameters>();
 		PassParameters->RenderTargets = BasePassRenderTargets;
+		PassParameters->SceneTextures = CreateSceneTextureUniformBuffer(GraphBuilder, FeatureLevel, ESceneTextureSetupMode::None);
 		for (auto& ViewExtension : ViewFamily.ViewExtensions)
 		{
 			for (FViewInfo& View : Views)
diff --git a/Engine/Source/Runtime/Renderer/Private/CapsuleShadowRendering.cpp b/Engine/Source/Runtime/Renderer/Private/CapsuleShadowRendering.cpp
--- a/Engine/Source/Runtime/Renderer/Private/CapsuleShadowRendering.cpp
+++ b/Engine/Source/Runtime/Renderer/Private/CapsuleShadowRendering.cpp
@@ -1479,10 +1479,6 @@ void FDeferredShadingSceneRenderer::RenderCapsuleShadowsForMovableSkylight(
 					RDG_EVENT_SCOPE(GraphBuilder, "IndirectCapsuleShadows");
 					RDG_GPU_STAT_SCOPE(GraphBuilder, CapsuleShadows);
 
-					FUniformBufferRHIRef PassUniformBuffer = CreateSceneTextureUniformBufferDependentOnShadingPath(GraphBuilder.RHICmdList, View.GetFeatureLevel(), ESceneTextureSetupMode::All);
-					FUniformBufferStaticBindings GlobalUniformBuffers(PassUniformBuffer);
-					SCOPED_UNIFORM_BUFFER_GLOBAL_BINDINGS(GraphBuilder.RHICmdList, GlobalUniformBuffers);
-
 					int32 NumCapsuleShapes = 0;
 					int32 NumMeshesWithCapsules = 0;
 					int32 NumMeshDistanceFieldCasters = 0;
diff --git a/Engine/Source/Runtime/Renderer/Private/DeferredShadingRenderer.cpp b/Engine/Source/Runtime/Renderer/Private/DeferredShadingRenderer.cpp
--- a/Engine/Source/Runtime/Renderer/Private/DeferredShadingRenderer.cpp
+++ b/Engine/Source/Runtime/Renderer/Private/DeferredShadingRenderer.cpp
@@ -2591,7 +2591,7 @@ void FDeferredShadingSceneRenderer::Render(FRHICommandListImmediate& RHICmdList)
 	}
 
 	FRendererModule& RendererModule = static_cast<FRendererModule&>(GetRendererModule());
-	RendererModule.RenderPostOpaqueExtensions(GraphBuilder, Views, SceneContext);
+	RendererModule.RenderPostOpaqueExtensions(GraphBuilder, Views, SceneColorTexture.Target, SceneDepthTexture.Target, SceneTextures, SceneContext);
 
 	// Notify the FX system that opaque primitives have been rendered and we now have a valid depth buffer.
 	if (FXSystem && Views.IsValidIndex(0))
@@ -2599,15 +2599,15 @@ void FDeferredShadingSceneRenderer::Render(FRHICommandListImmediate& RHICmdList)
 		RDG_GPU_STAT_SCOPE(GraphBuilder, PostRenderOpsFX);
 		RDG_CSV_STAT_EXCLUSIVE_SCOPE(GraphBuilder, RenderOpaqueFX);
 
-		AddUntrackedAccessPass(GraphBuilder, [this](FRHICommandListImmediate& InRHICmdList)
+		AddPass(GraphBuilder, [this](FRHICommandListImmediate& InRHICmdList)
 		{
 			SCOPE_CYCLE_COUNTER(STAT_FDeferredShadingSceneRenderer_FXSystem_PostRenderOpaque);
 
 			FXSystem->PostRenderOpaque(
 				InRHICmdList,
 				Views[0].ViewUniformBuffer,
-				&FSceneTextureUniformParameters::StaticStructMetadata,
-				CreateSceneTextureUniformBuffer(InRHICmdList, FeatureLevel),
+				nullptr,
+				nullptr,
 				Views[0].AllowGPUParticleUpdate()
 			);
 
@@ -2758,7 +2758,7 @@ void FDeferredShadingSceneRenderer::Render(FRHICommandListImmediate& RHICmdList)
 	}
 #endif
 
-	RendererModule.RenderOverlayExtensions(GraphBuilder, Views, SceneContext);
+	RendererModule.RenderOverlayExtensions(GraphBuilder, Views, SceneColorTexture.Target, SceneDepthTexture.Target, SceneContext);
 
 	if (ViewFamily.EngineShowFlags.VisualizeDistanceFieldAO && ShouldRenderDistanceFieldLighting())
 	{
diff --git a/Engine/Source/Runtime/Renderer/Private/MobileShadingRenderer.cpp b/Engine/Source/Runtime/Renderer/Private/MobileShadingRenderer.cpp
--- a/Engine/Source/Runtime/Renderer/Private/MobileShadingRenderer.cpp
+++ b/Engine/Source/Runtime/Renderer/Private/MobileShadingRenderer.cpp
@@ -819,6 +819,10 @@ void FMobileSceneRenderer::Render(FRHICommandListImmediate& RHICmdList)
 	FRHICommandListExecutor::GetImmediateCommandList().ImmediateFlush(EImmediateFlushType::DispatchToRHIThread);
 }
 
+BEGIN_SHADER_PARAMETER_STRUCT(FMobilePostBasePassViewExtensionParameters, )
+	SHADER_PARAMETER_RDG_UNIFORM_BUFFER(FMobileSceneTextureUniformParameters, SceneTextures)
+END_SHADER_PARAMETER_STRUCT()
+
 FRHITexture* FMobileSceneRenderer::RenderForward(FRHICommandListImmediate& RHICmdList, const TArrayView<const FViewInfo*> ViewList)
 {
 	const FViewInfo& View = *ViewList[0];
@@ -944,18 +948,31 @@ FRHITexture* FMobileSceneRenderer::RenderForward(FRHICommandListImmediate& RHICm
 		RHICmdList.ImmediateFlush(EImmediateFlushType::DispatchToRHIThread);
 	}
 
+	if (ViewFamily.ViewExtensions.Num() > 1)
 	{
-		CSV_SCOPED_TIMING_STAT_EXCLUSIVE(ViewExtensionPostRenderBasePass);
-		QUICK_SCOPE_CYCLE_COUNTER(STAT_FMobileSceneRenderer_ViewExtensionPostRenderBasePass);
-		for (int32 ViewExt = 0; ViewExt < ViewFamily.ViewExtensions.Num(); ++ViewExt)
+		FRDGBuilder GraphBuilder(RHICmdList);
+		auto* PassParameters = GraphBuilder.AllocParameters<FMobilePostBasePassViewExtensionParameters>();
+		PassParameters->SceneTextures = CreateMobileSceneTextureUniformBuffer(GraphBuilder, EMobileSceneTextureSetupMode::None);
+
+		GraphBuilder.AddPass(
+			{},
+			PassParameters,
+			ERDGPassFlags::Raster | ERDGPassFlags::SkipRenderPass | ERDGPassFlags::NeverCull,
+			[this](FRHICommandListImmediate& RHICmdList)
 		{
-			for (int32 ViewIndex = 0; ViewIndex < ViewFamily.Views.Num(); ++ViewIndex)
+			CSV_SCOPED_TIMING_STAT_EXCLUSIVE(ViewExtensionPostRenderBasePass);
+			QUICK_SCOPE_CYCLE_COUNTER(STAT_FMobileSceneRenderer_ViewExtensionPostRenderBasePass);
+			for (int32 ViewExt = 0; ViewExt < ViewFamily.ViewExtensions.Num(); ++ViewExt)
 			{
-				ViewFamily.ViewExtensions[ViewExt]->PostRenderBasePass_RenderThread(RHICmdList, Views[ViewIndex]);
+				for (int32 ViewIndex = 0; ViewIndex < ViewFamily.Views.Num(); ++ViewIndex)
+				{
+					ViewFamily.ViewExtensions[ViewExt]->PostRenderBasePass_RenderThread(RHICmdList, Views[ViewIndex]);
+				}
 			}
-		}
+		});
+		GraphBuilder.Execute();
 	}
-		
+
 	// Split if we need to render translucency in a separate render pass
 	// Split if we need to render pixel projected reflection
 	if (bRequiresMultiPass || bRequiresPixelProjectedPlanarRelfectionPass)
diff --git a/Engine/Source/Runtime/Renderer/Private/PlanarReflectionRendering.cpp b/Engine/Source/Runtime/Renderer/Private/PlanarReflectionRendering.cpp
--- a/Engine/Source/Runtime/Renderer/Private/PlanarReflectionRendering.cpp
+++ b/Engine/Source/Runtime/Renderer/Private/PlanarReflectionRendering.cpp
@@ -147,12 +147,23 @@ void SetupPlanarReflectionUniformParameters(const class FSceneView& View, const
 
 IMPLEMENT_GLOBAL_SHADER_PARAMETER_STRUCT(FPlanarReflectionUniformParameters, "PlanarReflectionStruct");
 
-
 template< bool bEnablePlanarReflectionPrefilter >
-class FPrefilterPlanarReflectionPS : public FGlobalShader
+class TPrefilterPlanarReflectionPS : public FGlobalShader
 {
-	DECLARE_SHADER_TYPE(FPrefilterPlanarReflectionPS, Global);
 public:
+	DECLARE_SHADER_TYPE(TPrefilterPlanarReflectionPS, Global);
+	SHADER_USE_PARAMETER_STRUCT(TPrefilterPlanarReflectionPS, FGlobalShader);
+
+	BEGIN_SHADER_PARAMETER_STRUCT(FParameters, )
+		SHADER_PARAMETER_STRUCT_REF(FViewUniformShaderParameters, View)
+		SHADER_PARAMETER_STRUCT_REF(FPlanarReflectionUniformParameters, PlanarReflection)
+		SHADER_PARAMETER_STRUCT_INCLUDE(FSceneTextureShaderParameters, SceneTextures)
+		SHADER_PARAMETER(float, KernelRadiusY)
+		SHADER_PARAMETER(float, InvPrefilterRoughnessDistance)
+		SHADER_PARAMETER_RDG_TEXTURE(Texture2D, SceneColorInputTexture)
+		SHADER_PARAMETER_SAMPLER(SamplerState, SceneColorInputSampler)
+		RENDER_TARGET_BINDING_SLOTS()
+	END_SHADER_PARAMETER_STRUCT()
 
 	static bool ShouldCompilePermutation(const FGlobalShaderPermutationParameters& Parameters)
 	{
@@ -164,77 +175,62 @@ public:
 		OutEnvironment.SetDefine(TEXT("ENABLE_PLANAR_REFLECTIONS_PREFILTER"), bEnablePlanarReflectionPrefilter);
 		FGlobalShader::ModifyCompilationEnvironment(Parameters, OutEnvironment);
 	}
+};
 
-	/** Default constructor. */
-	FPrefilterPlanarReflectionPS() {}
+IMPLEMENT_SHADER_TYPE(template<>, TPrefilterPlanarReflectionPS<false>, TEXT("/Engine/Private/PlanarReflectionShaders.usf"), TEXT("PrefilterPlanarReflectionPS"), SF_Pixel);
+IMPLEMENT_SHADER_TYPE(template<>, TPrefilterPlanarReflectionPS<true>, TEXT("/Engine/Private/PlanarReflectionShaders.usf"), TEXT("PrefilterPlanarReflectionPS"), SF_Pixel);
 
-	/** Initialization constructor. */
-	FPrefilterPlanarReflectionPS(const ShaderMetaType::CompiledShaderInitializerType& Initializer)
-		: FGlobalShader(Initializer)
-	{
-		KernelRadiusY.Bind(Initializer.ParameterMap, TEXT("KernelRadiusY"));
-		InvPrefilterRoughnessDistance.Bind(Initializer.ParameterMap, TEXT("InvPrefilterRoughnessDistance"));
-		SceneColorInputTexture.Bind(Initializer.ParameterMap, TEXT("SceneColorInputTexture"));
-		SceneColorInputSampler.Bind(Initializer.ParameterMap, TEXT("SceneColorInputSampler"));
-	}
+template<bool bEnablePlanarReflectionPrefilter>
+void PrefilterPlanarReflection(
+	FRDGBuilder& GraphBuilder,
+	const FViewInfo& View,
+	FSceneTextureShaderParameters SceneTextures,
+	const FPlanarReflectionSceneProxy* ReflectionSceneProxy,
+	FRDGTextureRef SceneColorTexture,
+	FRDGTextureRef ViewFamilyTexture)
+{
+	using FPrefilterPlanarReflectionPS = TPrefilterPlanarReflectionPS<bEnablePlanarReflectionPrefilter>;
 
-	void SetParameters(FRHICommandList& RHICmdList, const FSceneView& View, const FPlanarReflectionSceneProxy* ReflectionSceneProxy, FRHITexture* SceneColorInput, int32 FilterWidth)
+	if(View.FeatureLevel >= ERHIFeatureLevel::SM5)
 	{
-		FRHIPixelShader* ShaderRHI = RHICmdList.GetBoundPixelShader();
-		FGlobalShader::SetParameters<FViewUniformShaderParameters>(RHICmdList, ShaderRHI, View.ViewUniformBuffer);
-
-		const float KernelRadiusYValue = FMath::Clamp(ReflectionSceneProxy->PrefilterRoughness, 0.0f, 0.04f) * 0.5f * FilterWidth;
-		SetShaderValue(RHICmdList, ShaderRHI, KernelRadiusY, KernelRadiusYValue);
-
-		SetShaderValue(RHICmdList, ShaderRHI, InvPrefilterRoughnessDistance, 1.0f / FMath::Max(ReflectionSceneProxy->PrefilterRoughnessDistance, DELTA));
-
-		SetTextureParameter(RHICmdList, ShaderRHI, SceneColorInputTexture, SceneColorInputSampler, TStaticSamplerState<SF_Bilinear, AM_Clamp, AM_Clamp, AM_Clamp>::GetRHI(), SceneColorInput);
-
-		FPlanarReflectionUniformParameters PlanarReflectionUniformParameters;
-		SetupPlanarReflectionUniformParameters(View, ReflectionSceneProxy, PlanarReflectionUniformParameters);
-		SetUniformBufferParameterImmediate(RHICmdList, ShaderRHI, GetUniformBufferParameter<FPlanarReflectionUniformParameters>(), PlanarReflectionUniformParameters);
+		SceneColorTexture = AddProcessPlanarReflectionPass(GraphBuilder, View, SceneColorTexture);
 	}
 
-private:
-
-	LAYOUT_FIELD(FShaderParameter, KernelRadiusY)
-	LAYOUT_FIELD(FShaderParameter, InvPrefilterRoughnessDistance)
-	LAYOUT_FIELD(FShaderResourceParameter, SceneColorInputTexture)
-	LAYOUT_FIELD(FShaderResourceParameter, SceneColorInputSampler)
-};
+	{
+		RDG_EVENT_SCOPE(GraphBuilder, "PrefilterPlanarReflection");
 
-IMPLEMENT_SHADER_TYPE(template<>, FPrefilterPlanarReflectionPS<false>, TEXT("/Engine/Private/PlanarReflectionShaders.usf"), TEXT("PrefilterPlanarReflectionPS"), SF_Pixel);
-IMPLEMENT_SHADER_TYPE(template<>, FPrefilterPlanarReflectionPS<true>, TEXT("/Engine/Private/PlanarReflectionShaders.usf"), TEXT("PrefilterPlanarReflectionPS"), SF_Pixel);
+		// Workaround for a possible driver bug on S7 Adreno, missing planar reflections
+		const ERenderTargetLoadAction RTLoadAction = IsVulkanMobilePlatform(View.GetShaderPlatform()) ? ERenderTargetLoadAction::EClear : ERenderTargetLoadAction::ENoAction;
 
-template<bool bEnablePlanarReflectionPrefilter>
-void PrefilterPlanarReflection(FRHICommandListImmediate& RHICmdList, FViewInfo& View, const FUniformBufferRHIRef& PassUniformBuffer, const FPlanarReflectionSceneProxy* ReflectionSceneProxy, const FRenderTarget* Target)
-{
-	FRHITexture* SceneColorInput = FSceneRenderTargets::Get().GetSceneColorTexture();
+		const float FilterWidth = View.ViewRect.Width();
 
-	if(View.FeatureLevel >= ERHIFeatureLevel::SM5)
-	{
-		TRefCountPtr<IPooledRenderTarget> FilteredSceneColor;
-		GPostProcessing.ProcessPlanarReflection(RHICmdList, View, FilteredSceneColor);
+		auto* PassParameters = GraphBuilder.AllocParameters<typename FPrefilterPlanarReflectionPS::FParameters>();
+		PassParameters->View = View.ViewUniformBuffer;
 
-		if (FilteredSceneColor)
 		{
-			SceneColorInput = FilteredSceneColor->GetRenderTargetItem().ShaderResourceTexture;
+			FPlanarReflectionUniformParameters PlanarReflectionUniformParameters;
+			SetupPlanarReflectionUniformParameters(View, ReflectionSceneProxy, PlanarReflectionUniformParameters);
+			PassParameters->PlanarReflection = TUniformBufferRef<FPlanarReflectionUniformParameters>::CreateUniformBufferImmediate(PlanarReflectionUniformParameters, UniformBuffer_SingleFrame);
 		}
-	}
 
-	{
-		SCOPED_DRAW_EVENT(RHICmdList, PrefilterPlanarReflection);
+		PassParameters->SceneTextures = SceneTextures;
 
-		FUniformBufferStaticBindings GlobalUniformBuffers(PassUniformBuffer);
-		SCOPED_UNIFORM_BUFFER_GLOBAL_BINDINGS(RHICmdList, GlobalUniformBuffers);
+		PassParameters->KernelRadiusY = FMath::Clamp(ReflectionSceneProxy->PrefilterRoughness, 0.0f, 0.04f) * 0.5f * FilterWidth;
+		PassParameters->InvPrefilterRoughnessDistance = 1.0f / FMath::Max(ReflectionSceneProxy->PrefilterRoughnessDistance, DELTA);
+		PassParameters->SceneColorInputTexture = SceneColorTexture;
+		PassParameters->SceneColorInputSampler = TStaticSamplerState<SF_Bilinear, AM_Clamp, AM_Clamp, AM_Clamp>::GetRHI();
+		PassParameters->RenderTargets[0] = FRenderTargetBinding(ViewFamilyTexture, RTLoadAction);
 
-		// Workaround for a possible driver bug on S7 Adreno, missing planar reflections
-		ERenderTargetLoadAction RTLoadAction = IsVulkanMobilePlatform(View.GetShaderPlatform()) ?  ERenderTargetLoadAction::EClear : ERenderTargetLoadAction::ENoAction;
+		TShaderMapRef<TDeferredLightVS<false> > VertexShader(View.ShaderMap);
+		TShaderMapRef<FPrefilterPlanarReflectionPS> PixelShader(View.ShaderMap);
 
-		RHICmdList.Transition(FRHITransitionInfo(Target->GetRenderTargetTexture(), ERHIAccess::Unknown, ERHIAccess::RTV));
+		const FIntPoint SceneColorExtent = SceneColorTexture->Desc.Extent;
 
-		FRHIRenderPassInfo RPInfo(Target->GetRenderTargetTexture(), MakeRenderTargetActions(RTLoadAction, ERenderTargetStoreAction::EStore));
-		RHICmdList.BeginRenderPass(RPInfo, TEXT("PrefilterPlanarReflections"));
+		GraphBuilder.AddPass(
+			RDG_EVENT_NAME("PrefilterPlanarReflections"),
+			PassParameters,
+			ERDGPassFlags::Raster,
+			[&View, VertexShader, PixelShader, PassParameters, SceneColorExtent](FRHICommandList& RHICmdList)
 		{
 			RHICmdList.SetViewport(View.ViewRect.Min.X, View.ViewRect.Min.Y, 0.0f, View.ViewRect.Max.X, View.ViewRect.Max.Y, 1.0f);
 
@@ -244,9 +240,6 @@ void PrefilterPlanarReflection(FRHICommandListImmediate& RHICmdList, FViewInfo&
 			GraphicsPSOInit.RasterizerState = TStaticRasterizerState<FM_Solid, CM_None>::GetRHI();
 			GraphicsPSOInit.DepthStencilState = TStaticDepthStencilState<false, CF_Always>::GetRHI();
 
-			TShaderMapRef<TDeferredLightVS<false> > VertexShader(View.ShaderMap);
-			TShaderMapRef<FPrefilterPlanarReflectionPS<bEnablePlanarReflectionPrefilter> > PixelShader(View.ShaderMap);
-
 			GraphicsPSOInit.BoundShaderState.VertexDeclarationRHI = GFilterVertexDeclaration.VertexDeclarationRHI;
 			GraphicsPSOInit.BoundShaderState.VertexShaderRHI = VertexShader.GetVertexShader();
 			GraphicsPSOInit.BoundShaderState.PixelShaderRHI = PixelShader.GetPixelShader();
@@ -254,7 +247,7 @@ void PrefilterPlanarReflection(FRHICommandListImmediate& RHICmdList, FViewInfo&
 
 			SetGraphicsPipelineState(RHICmdList, GraphicsPSOInit);
 
-			PixelShader->SetParameters(RHICmdList, View, ReflectionSceneProxy, SceneColorInput, View.ViewRect.Width());
+			SetShaderParameters(RHICmdList, PixelShader, PixelShader.GetPixelShader(), *PassParameters);
 			VertexShader->SetSimpleLightParameters(RHICmdList, View, FSphere(0));
 
 			FIntPoint UV = View.ViewRect.Min;
@@ -273,11 +266,10 @@ void PrefilterPlanarReflection(FRHICommandListImmediate& RHICmdList, FViewInfo&
 				UV.X, UV.Y,
 				UVSize.X, UVSize.Y,
 				View.ViewRect.Size(),
-				FSceneRenderTargets::Get().GetBufferSizeXY(),
+				SceneColorExtent,
 				VertexShader,
 				EDRF_UseTriangleOptimization);
-		}
-		RHICmdList.EndRenderPass();
+		});
 	}
 }
 
@@ -286,11 +278,9 @@ static void UpdatePlanarReflectionContents_RenderThread(
 	FSceneRenderer* MainSceneRenderer, 
 	FSceneRenderer* SceneRenderer, 
 	FPlanarReflectionSceneProxy* SceneProxy,
-	FPlanarReflectionRenderTarget* RenderTarget, 
-	FTexture* RenderTargetTexture, 
+	FPlanarReflectionRenderTarget* RenderTarget,
 	const FPlane& MirrorPlane,
 	const FName OwnerName, 
-	const FResolveParams& ResolveParams, 
 	bool bUseSceneColorTexture)
 {
 	QUICK_SCOPE_CYCLE_COUNTER(STAT_RenderPlanarReflection);
@@ -372,8 +362,6 @@ static void UpdatePlanarReflectionContents_RenderThread(
 				SCOPED_DRAW_EVENT(RHICmdList, UpdatePlanarReflectionContent_RenderThread);
 #endif
 
-				const FRenderTarget* Target = SceneRenderer->ViewFamily.RenderTarget;
-
 				// Reflection view late update
 				if (SceneRenderer->Views.Num() > 1)
 				{
@@ -403,22 +391,26 @@ static void UpdatePlanarReflectionContents_RenderThread(
 					SceneProxy->ViewRect[ViewIndex] = SceneRenderer->Views[ViewIndex].ViewRect;
 				}
 
-				FUniformBufferRHIRef PassUniformBuffer = CreateSceneTextureUniformBufferDependentOnShadingPath(RHICmdList, SceneRenderer->FeatureLevel);
+				FRDGBuilder GraphBuilder(RHICmdList);
+				FRDGTextureRef ViewFamilyTexture = GraphBuilder.RegisterExternalTexture(CreateRenderTarget(RenderTarget->TextureRHI, TEXT("ViewFamilyTexture")));
+				FRDGTextureRef SceneColorTexture = GraphBuilder.RegisterExternalTexture(FSceneRenderTargets::Get().GetSceneColor());
+
+				FSceneTextureShaderParameters SceneTextures = CreateSceneTextureShaderParameters(GraphBuilder, SceneRenderer->FeatureLevel, ESceneTextureSetupMode::SceneDepth);
 
 				for (int32 ViewIndex = 0; ViewIndex < SceneRenderer->Views.Num(); ++ViewIndex)
 				{
 					FViewInfo& View = SceneRenderer->Views[ViewIndex];
-					SCOPED_GPU_MASK(RHICmdList, View.GPUMask);
+					RDG_GPU_MASK_SCOPE(GraphBuilder, View.GPUMask);
 					if (MainSceneRenderer->Scene->GetShadingPath() == EShadingPath::Deferred)
 					{
-						PrefilterPlanarReflection<true>(RHICmdList, View, PassUniformBuffer, SceneProxy, Target);
+						PrefilterPlanarReflection<true>(GraphBuilder, View, SceneTextures, SceneProxy, SceneColorTexture, ViewFamilyTexture);
 					}
 					else
 					{
-						PrefilterPlanarReflection<false>(RHICmdList, View, PassUniformBuffer, SceneProxy, Target);
+						PrefilterPlanarReflection<false>(GraphBuilder, View, SceneTextures, SceneProxy, SceneColorTexture, ViewFamilyTexture);
 					}
 				}
-				RHICmdList.CopyToResolveTarget(RenderTarget->GetRenderTargetTexture(), RenderTargetTexture->TextureRHI, ResolveParams);
+				GraphBuilder.Execute();
 			}
 		}
 	}
@@ -688,7 +680,7 @@ void FScene::UpdatePlanarReflectionContents(UPlanarReflectionComponent* CaptureC
 				ENQUEUE_RENDER_COMMAND(CaptureCommand)(
 					[SceneRenderer, MirrorPlane, OwnerName, MainSceneRendererPtr, SceneProxyPtr, RenderTargetPtr](FRHICommandListImmediate& RHICmdList)
 				{
-					UpdatePlanarReflectionContents_RenderThread(RHICmdList, MainSceneRendererPtr, SceneRenderer, SceneProxyPtr, RenderTargetPtr, RenderTargetPtr, MirrorPlane, OwnerName, FResolveParams(), true);
+					UpdatePlanarReflectionContents_RenderThread(RHICmdList, MainSceneRendererPtr, SceneRenderer, SceneProxyPtr, RenderTargetPtr, MirrorPlane, OwnerName, true);
 				});
 			}
 		}
diff --git a/Engine/Source/Runtime/Renderer/Private/PostProcess/PostProcessing.cpp b/Engine/Source/Runtime/Renderer/Private/PostProcess/PostProcessing.cpp
--- a/Engine/Source/Runtime/Renderer/Private/PostProcess/PostProcessing.cpp
+++ b/Engine/Source/Runtime/Renderer/Private/PostProcess/PostProcessing.cpp
@@ -50,9 +50,6 @@
 #include "ScreenSpaceRayTracing.h"
 #include "SceneViewExtension.h"
 
-/** The global center for all post processing activities. */
-FPostProcessing GPostProcessing;
-
 bool IsMobileEyeAdaptationEnabled(const FViewInfo& View);
 
 bool IsValidBloomSetupVariation(bool bUseBloom, bool bUseSun, bool bUseDof, bool bUseEyeAdaptation);
@@ -2097,19 +2094,18 @@ void AddMobilePostProcessingPasses(FRDGBuilder& GraphBuilder, const FViewInfo& V
 	}
 }
 
-void FPostProcessing::ProcessPlanarReflection(FRHICommandListImmediate& RHICmdList, const FViewInfo& View, TRefCountPtr<IPooledRenderTarget>& OutFilteredSceneColor)
+FRDGTextureRef AddProcessPlanarReflectionPass(
+	FRDGBuilder& GraphBuilder,
+	const FViewInfo& View,
+	FRDGTextureRef SceneColorTexture)
 {
 	FSceneViewState* ViewState = View.ViewState;
 	const EAntiAliasingMethod AntiAliasingMethod = View.AntiAliasingMethod;
 
-	const FSceneRenderTargets& SceneContext = FSceneRenderTargets::Get();
-
 	if (AntiAliasingMethod == AAM_TemporalAA)
 	{
 		check(ViewState);
 
-		FRDGBuilder GraphBuilder(RHICmdList);
-
 		FSceneTextureParameters SceneTextures = GetSceneTextureParameters(GraphBuilder);
 
 		const FTemporalAAHistory& InputHistory = View.PrevViewInfo.TemporalAAHistory;
@@ -2121,7 +2117,7 @@ void FPostProcessing::ProcessPlanarReflection(FRHICommandListImmediate& RHICmdLi
 		// Planar reflections don't support velocity.
 		Parameters.SceneVelocityTexture = nullptr;
 
-		Parameters.SceneColorInput = GraphBuilder.RegisterExternalTexture(SceneContext.GetSceneColor(), TEXT("SceneColor"));
+		Parameters.SceneColorInput = SceneColorTexture;
 
 		FTAAOutputs PassOutputs = AddTemporalAAPass(
 			GraphBuilder,
@@ -2130,12 +2126,10 @@ void FPostProcessing::ProcessPlanarReflection(FRHICommandListImmediate& RHICmdLi
 			InputHistory,
 			OutputHistory);
 
-		GraphBuilder.QueueTextureExtraction(PassOutputs.SceneColor, &OutFilteredSceneColor);
-
-		GraphBuilder.Execute();
+		return PassOutputs.SceneColor;
 	}
 	else
 	{
-		OutFilteredSceneColor = SceneContext.GetSceneColor();
+		return SceneColorTexture;
 	}
 }
diff --git a/Engine/Source/Runtime/Renderer/Private/PostProcess/PostProcessing.h b/Engine/Source/Runtime/Renderer/Private/PostProcess/PostProcessing.h
--- a/Engine/Source/Runtime/Renderer/Private/PostProcess/PostProcessing.h
+++ b/Engine/Source/Runtime/Renderer/Private/PostProcess/PostProcessing.h
@@ -84,14 +84,4 @@ void AddMobilePostProcessingPasses(FRDGBuilder& GraphBuilder, const FViewInfo& V
 
 void AddBasicPostProcessPasses(FRDGBuilder& GraphBuilder, const FViewInfo& View);
 
-/**
- * The center for all post processing activities.
- */
-class FPostProcessing
-{
-public:
-	void ProcessPlanarReflection(FRHICommandListImmediate& RHICmdList, const FViewInfo& View, TRefCountPtr<IPooledRenderTarget>& OutFilteredSceneColor);
-};
-
-/** The global used for post processing. */
-extern FPostProcessing GPostProcessing;
+FRDGTextureRef AddProcessPlanarReflectionPass(FRDGBuilder& GraphBuilder, const FViewInfo& View, FRDGTextureRef SceneColorTexture);
diff --git a/Engine/Source/Runtime/Renderer/Private/PostProcess/SceneRenderTargets.cpp b/Engine/Source/Runtime/Renderer/Private/PostProcess/SceneRenderTargets.cpp
--- a/Engine/Source/Runtime/Renderer/Private/PostProcess/SceneRenderTargets.cpp
+++ b/Engine/Source/Runtime/Renderer/Private/PostProcess/SceneRenderTargets.cpp
@@ -2505,14 +2505,7 @@ TRefCountPtr<FRHIUniformBuffer> CreateSceneTextureUniformBufferDependentOnShadin
 	ERHIFeatureLevel::Type FeatureLevel,
 	ESceneTextureSetupMode SetupMode)
 {
-	if (FSceneInterface::GetShadingPath(FeatureLevel) == EShadingPath::Deferred)
-	{
-		return CreateSceneTextureUniformBuffer(RHICmdList, FeatureLevel, SetupMode);
-	}
-	else if (FSceneInterface::GetShadingPath(FeatureLevel) == EShadingPath::Mobile)
-	{
-		return CreateMobileSceneTextureUniformBuffer(RHICmdList, Translate(SetupMode));
-	}
+	// Deprecated
 	checkNoEntry();
 	return nullptr;
 }
diff --git a/Engine/Source/Runtime/Renderer/Private/RendererModule.h b/Engine/Source/Runtime/Renderer/Private/RendererModule.h
--- a/Engine/Source/Runtime/Renderer/Private/RendererModule.h
+++ b/Engine/Source/Runtime/Renderer/Private/RendererModule.h
@@ -107,9 +107,22 @@ public:
 	
 	virtual void RegisterPersistentViewUniformBufferExtension(IPersistentViewUniformBufferExtension* Extension) override;
 
-	void RenderPostOpaqueExtensions(FRDGBuilder& GraphBuilder, TArrayView<const FViewInfo> Views, FSceneRenderTargets& SceneContext);
-	void RenderOverlayExtensions(FRDGBuilder& GraphBuilder, TArrayView<const FViewInfo> Views, FSceneRenderTargets& SceneContext);
-	void RenderPostResolvedSceneColorExtension(FRDGBuilder& GraphBuilder, class FSceneRenderTargets& SceneContext);
+	void RenderPostOpaqueExtensions(
+		FRDGBuilder& GraphBuilder,
+		TArrayView<const FViewInfo> Views,
+		FRDGTextureRef SceneColorTexture,
+		FRDGTextureRef SceneDepthTexture,
+		TRDGUniformBufferRef<FSceneTextureUniformParameters> SceneTexturesUniformBuffer,
+		const FSceneRenderTargets& SceneContext);
+
+	void RenderOverlayExtensions(
+		FRDGBuilder& GraphBuilder,
+		TArrayView<const FViewInfo> Views,
+		FRDGTextureRef SceneColorTexture,
+		FRDGTextureRef SceneDepthTexture,
+		const FSceneRenderTargets& SceneContext);
+
+	void RenderPostResolvedSceneColorExtension(FRDGBuilder& GraphBuilder, FSceneRenderTargets& SceneContext);
 
 private:
 	TSet<FSceneInterface*> AllocatedScenes;
diff --git a/Engine/Source/Runtime/Renderer/Private/SceneRendering.cpp b/Engine/Source/Runtime/Renderer/Private/SceneRendering.cpp
--- a/Engine/Source/Runtime/Renderer/Private/SceneRendering.cpp
+++ b/Engine/Source/Runtime/Renderer/Private/SceneRendering.cpp
@@ -3116,19 +3116,12 @@ void FSceneRenderer::RenderFinish(FRDGBuilder& GraphBuilder, FRDGTextureRef View
 		{
 			RDG_EVENT_SCOPE(GraphBuilder, "ViewFamilyExtension(%d)", ViewExt);
 			ISceneViewExtension& ViewExtension = *ViewFamily.ViewExtensions[ViewExt];
-
-			AddUntrackedAccessPass(GraphBuilder, [this, &ViewExtension](FRHICommandListImmediate& RHICmdList)
-			{
-				ViewExtension.PostRenderViewFamily_RenderThread(RHICmdList, ViewFamily);
-			});
+			ViewExtension.PostRenderViewFamily_RenderThread(GraphBuilder, ViewFamily);
 
 			for(int32 ViewIndex = 0; ViewIndex < ViewFamily.Views.Num(); ++ViewIndex)
 			{
-				AddUntrackedAccessPass(GraphBuilder, RDG_EVENT_NAME("ViewExtension(%d)", ViewIndex),
-					[this, &ViewExtension, ViewIndex](FRHICommandListImmediate& RHICmdList)
-				{
-					ViewExtension.PostRenderView_RenderThread(RHICmdList, Views[ViewIndex]);
-				});
+				RDG_EVENT_SCOPE(GraphBuilder, "ViewExtension(%d)", ViewIndex);
+				ViewExtension.PostRenderView_RenderThread(GraphBuilder, Views[ViewIndex]);
 			}
 		}
 	}
@@ -3581,17 +3574,20 @@ static void ViewExtensionPreRender_RenderThread(FRHICommandListImmediate& RHICmd
 	FMemMark MemStackMark(FMemStack::Get());
 
 	{
+		FRDGBuilder GraphBuilder(RHICmdList);
 		CSV_SCOPED_TIMING_STAT_EXCLUSIVE(PreRender);
 		SCOPE_CYCLE_COUNTER(STAT_FDeferredShadingSceneRenderer_ViewExtensionPreRenderView);
 
 		for (int ViewExt = 0; ViewExt < SceneRenderer->ViewFamily.ViewExtensions.Num(); ViewExt++)
 		{
-			SceneRenderer->ViewFamily.ViewExtensions[ViewExt]->PreRenderViewFamily_RenderThread(RHICmdList, SceneRenderer->ViewFamily);
+			SceneRenderer->ViewFamily.ViewExtensions[ViewExt]->PreRenderViewFamily_RenderThread(GraphBuilder, SceneRenderer->ViewFamily);
 			for (int ViewIndex = 0; ViewIndex < SceneRenderer->ViewFamily.Views.Num(); ViewIndex++)
 			{
-				SceneRenderer->ViewFamily.ViewExtensions[ViewExt]->PreRenderView_RenderThread(RHICmdList, SceneRenderer->Views[ViewIndex]);
+				SceneRenderer->ViewFamily.ViewExtensions[ViewExt]->PreRenderView_RenderThread(GraphBuilder, SceneRenderer->Views[ViewIndex]);
 			}
 		}
+
+		GraphBuilder.Execute();
 	}
 	
 	// update any resources that needed a deferred update
@@ -4004,7 +4000,13 @@ void FRendererModule::RemoveOverlayRenderDelegate(FDelegateHandle InOverlayRende
 	OverlayRenderDelegate.Remove(InOverlayRenderDelegate);
 }
 
-void FRendererModule::RenderPostOpaqueExtensions(FRDGBuilder& GraphBuilder, TArrayView<const FViewInfo> Views, FSceneRenderTargets& SceneContext)
+void FRendererModule::RenderPostOpaqueExtensions(
+	FRDGBuilder& GraphBuilder,
+	TArrayView<const FViewInfo> Views,
+	FRDGTextureRef SceneColorTexture,
+	FRDGTextureRef SceneDepthTexture,
+	TRDGUniformBufferRef<FSceneTextureUniformParameters> SceneTexturesUniformBuffer,
+	const FSceneRenderTargets& SceneContext)
 {
 	if (PostOpaqueRenderDelegate.IsBound())
 	{
@@ -4016,33 +4018,34 @@ void FRendererModule::RenderPostOpaqueExtensions(FRDGBuilder& GraphBuilder, TArr
 			RDG_EVENT_SCOPE_CONDITIONAL(GraphBuilder, Views.Num() > 1, "View%d", ViewIndex);
 			RDG_GPU_MASK_SCOPE(GraphBuilder, View.GPUMask);
 
-			AddUntrackedAccessPass(GraphBuilder, [this, &View, &SceneContext](FRHICommandListImmediate& RHICmdList)
-			{
-				RHICmdList.SetViewport(View.ViewRect.Min.X, View.ViewRect.Min.Y, 0.0f, View.ViewRect.Max.X, View.ViewRect.Max.Y, 1.0f);
-
-				check(IsInRenderingThread());
-				FPostOpaqueRenderParameters RenderParameters;
-				RenderParameters.ViewMatrix = View.ViewMatrices.GetViewMatrix();
-				RenderParameters.ProjMatrix = View.ViewMatrices.GetProjectionMatrix();
-				RenderParameters.DepthTexture = SceneContext.GetSceneDepthSurface()->GetTexture2D();
-				RenderParameters.NormalTexture = SceneContext.GBufferA.IsValid() ? SceneContext.GetGBufferATexture() : nullptr;
-				RenderParameters.VelocityTexture = SceneContext.SceneVelocity.IsValid() ? SceneContext.SceneVelocity->GetRenderTargetItem().ShaderResourceTexture->GetTexture2D() : nullptr;
-				RenderParameters.SmallDepthTexture = SceneContext.GetSmallDepthSurface()->GetTexture2D();
-				RenderParameters.ViewUniformBuffer = View.ViewUniformBuffer;
-				RenderParameters.SceneTexturesUniformParams = CreateSceneTextureUniformBuffer(RHICmdList, View.FeatureLevel, ESceneTextureSetupMode::SceneColor | ESceneTextureSetupMode::SceneDepth | ESceneTextureSetupMode::SceneVelocity | ESceneTextureSetupMode::GBuffers);
-				RenderParameters.GlobalDistanceFieldParams = &View.GlobalDistanceFieldInfo.ParameterData;
-
-				RenderParameters.ViewportRect = View.ViewRect;
-				RenderParameters.RHICmdList = &RHICmdList;
-
-				RenderParameters.Uid = (void*)(&View);
-				PostOpaqueRenderDelegate.Broadcast(RenderParameters);
-			});
+			check(IsInRenderingThread());
+			FPostOpaqueRenderParameters RenderParameters;
+			RenderParameters.ViewMatrix = View.ViewMatrices.GetViewMatrix();
+			RenderParameters.ProjMatrix = View.ViewMatrices.GetProjectionMatrix();
+			RenderParameters.ColorTexture = SceneColorTexture;
+			RenderParameters.DepthTexture = SceneDepthTexture;
+			RenderParameters.NormalTexture = TryRegisterExternalTexture(GraphBuilder, SceneContext.GBufferA);
+			RenderParameters.VelocityTexture = TryRegisterExternalTexture(GraphBuilder, SceneContext.SceneVelocity);
+			RenderParameters.SmallDepthTexture = TryRegisterExternalTexture(GraphBuilder, SceneContext.SmallDepthZ);
+			RenderParameters.ViewUniformBuffer = View.ViewUniformBuffer;
+			RenderParameters.SceneTexturesUniformParams = SceneTexturesUniformBuffer;
+			RenderParameters.GlobalDistanceFieldParams = &View.GlobalDistanceFieldInfo.ParameterData;
+
+			RenderParameters.ViewportRect = View.ViewRect;
+			RenderParameters.GraphBuilder = &GraphBuilder;
+
+			RenderParameters.Uid = (void*)(&View);
+			PostOpaqueRenderDelegate.Broadcast(RenderParameters);
 		}
 	}
 }
 
-void FRendererModule::RenderOverlayExtensions(FRDGBuilder& GraphBuilder, TArrayView<const FViewInfo> Views, FSceneRenderTargets& SceneContext)
+void FRendererModule::RenderOverlayExtensions(
+	FRDGBuilder& GraphBuilder,
+	TArrayView<const FViewInfo> Views,
+	FRDGTextureRef SceneColorTexture,
+	FRDGTextureRef SceneDepthTexture,
+	const FSceneRenderTargets& SceneContext)
 {
 	if (OverlayRenderDelegate.IsBound())
 	{
@@ -4054,22 +4057,18 @@ void FRendererModule::RenderOverlayExtensions(FRDGBuilder& GraphBuilder, TArrayV
 			RDG_EVENT_SCOPE_CONDITIONAL(GraphBuilder, Views.Num() > 1, "View%d", ViewIndex);
 			RDG_GPU_MASK_SCOPE(GraphBuilder, View.GPUMask);
 
-			AddUntrackedAccessPass(GraphBuilder, [this, &View, &SceneContext](FRHICommandListImmediate& RHICmdList)
-			{
-				RHICmdList.SetViewport(View.ViewRect.Min.X, View.ViewRect.Min.Y, 0.0f, View.ViewRect.Max.X, View.ViewRect.Max.Y, 1.0f);
-
-				FPostOpaqueRenderParameters RenderParameters;
-				RenderParameters.ViewMatrix = View.ViewMatrices.GetViewMatrix();
-				RenderParameters.ProjMatrix = View.ViewMatrices.GetProjectionMatrix();
-				RenderParameters.DepthTexture = SceneContext.GetSceneDepthSurface()->GetTexture2D();
-				RenderParameters.SmallDepthTexture = SceneContext.GetSmallDepthSurface()->GetTexture2D();
+			FPostOpaqueRenderParameters RenderParameters;
+			RenderParameters.ViewMatrix = View.ViewMatrices.GetViewMatrix();
+			RenderParameters.ProjMatrix = View.ViewMatrices.GetProjectionMatrix();
+			RenderParameters.ColorTexture = SceneColorTexture;
+			RenderParameters.DepthTexture = SceneDepthTexture;
+			RenderParameters.SmallDepthTexture = TryRegisterExternalTexture(GraphBuilder, SceneContext.SmallDepthZ);
 
-				RenderParameters.ViewportRect = View.ViewRect;
-				RenderParameters.RHICmdList = &RHICmdList;
+			RenderParameters.ViewportRect = View.ViewRect;
+			RenderParameters.GraphBuilder = &GraphBuilder;
 
-				RenderParameters.Uid = (void*)(&View);
-				OverlayRenderDelegate.Broadcast(RenderParameters);
-			});
+			RenderParameters.Uid = (void*)(&View);
+			OverlayRenderDelegate.Broadcast(RenderParameters);
 		}
 	}
 }
@@ -4078,13 +4077,7 @@ void FRendererModule::RenderPostResolvedSceneColorExtension(FRDGBuilder& GraphBu
 {
 	if (PostResolvedSceneColorCallbacks.IsBound())
 	{
-		GraphBuilder.AddPass(
-			RDG_EVENT_NAME("PostResolvedSceneColorExtentions"),
-			ERDGPassFlags::UntrackedAccess,
-			[this, &SceneContext](FRHICommandListImmediate& RHICmdList)
-		{
-			PostResolvedSceneColorCallbacks.Broadcast(RHICmdList, SceneContext);
-		});
+		PostResolvedSceneColorCallbacks.Broadcast(GraphBuilder, SceneContext);
 	}
 }
 
diff --git a/Engine/Source/Runtime/Renderer/Private/ScreenPass.h b/Engine/Source/Runtime/Renderer/Private/ScreenPass.h
--- a/Engine/Source/Runtime/Renderer/Private/ScreenPass.h
+++ b/Engine/Source/Runtime/Renderer/Private/ScreenPass.h
@@ -18,7 +18,7 @@ bool IsHMDHiddenAreaMaskActive();
 const FTextureRHIRef& GetMiniFontTexture();
 
 // Creates and returns an RDG texture for the view family output. Returns null if no RHI texture exists.
-FRDGTextureRef TryCreateViewFamilyTexture(FRDGBuilder& GraphBuilder, const FSceneViewFamily& ViewFamily);
+FRDGTextureRef RENDERER_API TryCreateViewFamilyTexture(FRDGBuilder& GraphBuilder, const FSceneViewFamily& ViewFamily);
 
 // The vertex shader used by DrawScreenPass to draw a rectangle.
 class RENDERER_API FScreenPassVS : public FGlobalShader
diff --git a/Engine/Source/Runtime/Renderer/Public/SceneRenderTargetParameters.h b/Engine/Source/Runtime/Renderer/Public/SceneRenderTargetParameters.h
--- a/Engine/Source/Runtime/Renderer/Public/SceneRenderTargetParameters.h
+++ b/Engine/Source/Runtime/Renderer/Public/SceneRenderTargetParameters.h
@@ -52,24 +52,12 @@ enum class ESceneTextureSetupMode : uint32
 ENUM_CLASS_FLAGS(ESceneTextureSetupMode);
 
 /** Fills the shader parameter struct. */
-extern RENDERER_API void SetupSceneTextureUniformParameters(
-	const FSceneRenderTargets& SceneContext,
-	ERHIFeatureLevel::Type FeatureLevel,
-	ESceneTextureSetupMode SetupMode,
-	FSceneTextureUniformParameters& OutParameters);
-
 extern RENDERER_API void SetupSceneTextureUniformParameters(
 	FRDGBuilder& GraphBuilder,
 	ERHIFeatureLevel::Type FeatureLevel,
 	ESceneTextureSetupMode SetupMode,
 	FSceneTextureUniformParameters& OutParameters);
 
-/** Returns RHI scene texture uniform buffer with passthrough RDG resources. */
-extern RENDERER_API TUniformBufferRef<FSceneTextureUniformParameters> CreateSceneTextureUniformBuffer(
-	FRHIComputeCommandList& RHICmdList,
-	ERHIFeatureLevel::Type FeatureLevel,
-	ESceneTextureSetupMode SetupMode = ESceneTextureSetupMode::All);
-
 /** Returns RDG scene texture uniform buffer. */
 extern RENDERER_API TRDGUniformBufferRef<FSceneTextureUniformParameters> CreateSceneTextureUniformBuffer(
 	FRDGBuilder& GraphBuilder,
@@ -154,11 +142,23 @@ extern RENDERER_API FSceneTextureShaderParameters CreateSceneTextureShaderParame
 	ERHIFeatureLevel::Type FeatureLevel,
 	ESceneTextureSetupMode SetupMode = ESceneTextureSetupMode::All);
 
-/** Returns appropriate scene textures RHI uniform buffer for deferred or mobile path. */
-extern RENDERER_API TRefCountPtr<FRHIUniformBuffer> CreateSceneTextureUniformBufferDependentOnShadingPath(
+UE_DEPRECATED(5.0, "SetupSceneTextureUniforParameters now requires an FRDGBuilder.")
+extern RENDERER_API void SetupSceneTextureUniformParameters(
+	const FSceneRenderTargets& SceneContext,
+	ERHIFeatureLevel::Type FeatureLevel,
+	ESceneTextureSetupMode SetupMode,
+	FSceneTextureUniformParameters& OutParameters);
+
+UE_DEPRECATED(5.0, "CreateSceneTextureUniformBuffer now requires an FRDGBuilder.")
+extern RENDERER_API TUniformBufferRef<FSceneTextureUniformParameters> CreateSceneTextureUniformBuffer(
 	FRHIComputeCommandList& RHICmdList,
 	ERHIFeatureLevel::Type FeatureLevel,
 	ESceneTextureSetupMode SetupMode = ESceneTextureSetupMode::All);
 
+UE_DEPRECATED(5.0, "Use CreateSceneTextureShaderParameters instead.")
+extern RENDERER_API TRefCountPtr<FRHIUniformBuffer> CreateSceneTextureUniformBufferDependentOnShadingPath(
+	FRHIComputeCommandList& RHICmdList,
+	ERHIFeatureLevel::Type FeatureLevel,
+	ESceneTextureSetupMode SetupMode = ESceneTextureSetupMode::All);
 
 extern RENDERER_API bool IsSceneTexturesValid(FRHICommandListImmediate& RHICmdList);
\ No newline at end of file
