Removed dead code from FSceneRenderTargets.

#rnx

diff --git a/Engine/Plugins/Runtime/TextureShare/Source/TextureShare/Private/TextureShareModule.cpp b/Engine/Plugins/Runtime/TextureShare/Source/TextureShare/Private/TextureShareModule.cpp
--- a/Engine/Plugins/Runtime/TextureShare/Source/TextureShare/Private/TextureShareModule.cpp
+++ b/Engine/Plugins/Runtime/TextureShare/Source/TextureShare/Private/TextureShareModule.cpp
@@ -245,16 +245,6 @@ bool FTextureShareModule::SendSceneContext_RenderThread(FRHICommandListImmediate
 			SendTexture_RenderThread(RHICmdList, ShareItem, TextureShareStrings::texture_name::GBufferD, SceneContext.GBufferD);
 			SendTexture_RenderThread(RHICmdList, ShareItem, TextureShareStrings::texture_name::GBufferE, SceneContext.GBufferE);
 			SendTexture_RenderThread(RHICmdList, ShareItem, TextureShareStrings::texture_name::GBufferF, SceneContext.GBufferF);
-
-			// TODO: These textures aren't pooled anymore
-// 			if (SceneContext.HasLightAttenuation())
-// 			{
-// 				SendTexture_RenderThread(RHICmdList, ShareItem, TextureShareStrings::texture_name::LightAttenuation, SceneContext.GetLightAttenuationTexture());
-// 			}
-
-			SendTexture_RenderThread(RHICmdList, ShareItem, TextureShareStrings::texture_name::LightAccumulation, SceneContext.LightAccumulation);
-			//SendTexture_RenderThread(RHICmdList, ShareItem, TextureShareStrings::texture_name::LightingChannels, SceneContext.LightingChannels);
-
 			SendTexture_RenderThread(RHICmdList, ShareItem, TextureShareStrings::texture_name::GBufferVelocity, SceneContext.SceneVelocity);
 			SendTexture_RenderThread(RHICmdList, ShareItem, TextureShareStrings::texture_name::Foveation, SceneContext.FoveationTexture);
 			//@todo: Add more textures
diff --git a/Engine/Plugins/Runtime/TextureShare/Source/TextureShare/Private/TextureShareStrings.h b/Engine/Plugins/Runtime/TextureShare/Source/TextureShare/Private/TextureShareStrings.h
--- a/Engine/Plugins/Runtime/TextureShare/Source/TextureShare/Private/TextureShareStrings.h
+++ b/Engine/Plugins/Runtime/TextureShare/Source/TextureShare/Private/TextureShareStrings.h
@@ -21,10 +21,6 @@ namespace TextureShareStrings
 		static constexpr auto GBufferE = TEXT("GBufferE");
 		static constexpr auto GBufferF = TEXT("GBufferF");
 
-		static constexpr auto LightAttenuation = TEXT("LightAttenuation");
-		static constexpr auto LightAccumulation = TEXT("LightAccumulation");
-		static constexpr auto LightingChannels = TEXT("LightingChannels");
-
 		static constexpr auto GBufferVelocity = TEXT("GBufferVelocity");
 		static constexpr auto Foveation = TEXT("Foveation");
 
diff --git a/Engine/Source/Runtime/Renderer/Private/PostProcess/SceneRenderTargets.cpp b/Engine/Source/Runtime/Renderer/Private/PostProcess/SceneRenderTargets.cpp
--- a/Engine/Source/Runtime/Renderer/Private/PostProcess/SceneRenderTargets.cpp
+++ b/Engine/Source/Runtime/Renderer/Private/PostProcess/SceneRenderTargets.cpp
@@ -1742,17 +1742,6 @@ void FSceneRenderTargets::AllocateDeferredShadingPathRenderTargets(FRDGBuilder&
 		}
 	}
 
-	if (CurrentFeatureLevel >= ERHIFeatureLevel::SM5) 
-	{
-		FPooledRenderTargetDesc Desc(FPooledRenderTargetDesc::Create2DDesc(BufferSize, PF_FloatRGBA, FClearValueBinding::Black, TexCreate_None, TexCreate_RenderTargetable | TexCreate_ShaderResource, false));
-		if (CurrentFeatureLevel >= ERHIFeatureLevel::SM5)
-		{
-			Desc.TargetableFlags |= TexCreate_UAV;
-		}
-		Desc.Flags |= GFastVRamConfig.LightAccumulation;
-		GRenderTargetPool.FindFreeElement(GraphBuilder.RHICmdList, Desc, LightAccumulation, TEXT("LightAccumulation"), ERenderTargetTransience::NonTransient);
-	}
-
 	if (bAllocateVelocityGBuffer)
 	{
 		FPooledRenderTargetDesc VelocityRTDesc = Translate(FVelocityRendering::GetRenderTargetDesc(ShaderPlatform));
@@ -1916,7 +1905,6 @@ void FSceneRenderTargets::ReleaseAllTargets()
 	ScreenSpaceAO.SafeRelease();
 	ScreenSpaceGTAOHorizons.SafeRelease();
 	QuadOverdrawBuffer.SafeRelease();
-	LightAccumulation.SafeRelease();
 	CustomDepth.SafeRelease();
 	MobileCustomDepth.SafeRelease();
 	MobileCustomStencil.SafeRelease();
diff --git a/Engine/Source/Runtime/Renderer/Private/PostProcess/SceneRenderTargets.h b/Engine/Source/Runtime/Renderer/Private/PostProcess/SceneRenderTargets.h
--- a/Engine/Source/Runtime/Renderer/Private/PostProcess/SceneRenderTargets.h
+++ b/Engine/Source/Runtime/Renderer/Private/PostProcess/SceneRenderTargets.h
@@ -335,9 +335,6 @@ private: // Get...() methods instead of direct access
 	// Permanently allocated for forward shading
 	TRefCountPtr<IPooledRenderTarget> SceneColor[(int32)ESceneColorFormatType::Num];
 public:
-	// Light Accumulation is a high precision scratch pad matching the size of the scene color buffer used by many passes.
-	TRefCountPtr<IPooledRenderTarget> LightAccumulation;
-
 	// Scene depth and stencil.
 	TRefCountPtr<IPooledRenderTarget> SceneDepthZ;
 	TRefCountPtr<FRHIShaderResourceView> SceneStencilSRV;
@@ -427,8 +424,6 @@ private:
 	uint32 ThisFrameNumber;
 	uint32 CurrentDesiredSizeIndex;
 
-	/** CAUTION: When adding new data, make sure you copy it in the snapshot constructor! **/
-
 	/**
 	 * Initializes the editor primitive color render target
 	 */
@@ -550,8 +545,6 @@ private:
 	int32 CurrentMSAACount;
 	/** To detect a change of the CVar r.Shadow.MinResolution */
 	int32 CurrentMinShadowResolution;
-	/** To detect a change of the CVar r.LightPropagationVolume */
-	bool bCurrentLightPropagationVolume;
 	/** Feature level we were initialized for */
 	ERHIFeatureLevel::Type CurrentFeatureLevel;
 	/** Shading path that we are currently drawing through. Set when calling Allocate at the start of a scene render. */
@@ -568,9 +561,6 @@ private:
 	/** Clear depth value, defaults to FClearValueBinding::DepthFar */
 	FClearValueBinding DefaultDepthClear;
 
-	/** All outstanding snapshots */
-	TArray<FSceneRenderTargets*> Snapshots;
-
 	/** True if the depth target is allocated by an HMD plugin. This is a temporary fix to deal with HMD depth target swap chains not tracking the stencil SRV. */
 	bool bHMDAllocatedDepthTarget;
 
