Ensure that locked pages are still flushed after a virtual texture flush command.
Previously for streaming pages, the request would return pending, but because the page is mapped, no subsequent request would occur.

diff --git a/Engine/Source/Runtime/Renderer/Private/VT/VirtualTextureSystem.cpp b/Engine/Source/Runtime/Renderer/Private/VT/VirtualTextureSystem.cpp
--- a/Engine/Source/Runtime/Renderer/Private/VT/VirtualTextureSystem.cpp
+++ b/Engine/Source/Runtime/Renderer/Private/VT/VirtualTextureSystem.cpp
@@ -282,8 +282,7 @@ void FVirtualTextureSystem::FlushCache(FVirtualTextureProducerHandle const& Prod
 			PhysicalSpacesForProducer.AddUnique(Producer->GetPhysicalSpaceForPhysicalGroup(i));
 		}
 
-		// Don't resize to allow this container to grow as needed (avoid allocations when collecting)
-		TransientCollectedPages.Reset();
+		check(TransientCollectedPages.Num() == 0);
 
 		for (int32 i = 0; i < PhysicalSpacesForProducer.Num(); ++i)
 		{
@@ -295,6 +294,9 @@ void FVirtualTextureSystem::FlushCache(FVirtualTextureProducerHandle const& Prod
 		{
 			MappedTilesToProduce.Add(Page);
 		}
+
+		// Don't resize to allow this container to grow as needed (avoid allocations when collecting)
+		TransientCollectedPages.Reset();
 	}
 }
 
@@ -1782,7 +1784,7 @@ void FVirtualTextureSystem::UpdateCSVStats() const
 #endif
 }
 
-void FVirtualTextureSystem::SubmitRequestsFromLocalTileList(const TSet<FVirtualTextureLocalTile>& LocalTileList, EVTProducePageFlags Flags, FRDGBuilder& GraphBuilder, ERHIFeatureLevel::Type FeatureLevel)
+void FVirtualTextureSystem::SubmitRequestsFromLocalTileList(TArray<FVirtualTextureLocalTile>& OutDeferredTiles, const TSet<FVirtualTextureLocalTile>& LocalTileList, EVTProducePageFlags Flags, FRDGBuilder& GraphBuilder, ERHIFeatureLevel::Type FeatureLevel)
 {
 	LLM_SCOPE(ELLMTag::VirtualTextureSystem);
 
@@ -1830,7 +1832,8 @@ void FVirtualTextureSystem::SubmitRequestsFromLocalTileList(const TSet<FVirtualT
 
 		if (RequestPageResult.Status != EVTRequestPageStatus::Available)
 		{
-			//todo[vt]: Should we unmap? Or maybe keep the request for the next frame?
+			// Keep the request for the next frame?
+			OutDeferredTiles.Add(Tile);
 			continue;
 		}
 
@@ -1851,16 +1854,21 @@ void FVirtualTextureSystem::SubmitRequestsFromLocalTileList(const TSet<FVirtualT
 
 void FVirtualTextureSystem::SubmitPreMappedRequests(FRDGBuilder& GraphBuilder, ERHIFeatureLevel::Type FeatureLevel)
 {
+	check(TransientCollectedPages.Num() == 0);
+
 	{
 		INC_DWORD_STAT_BY(STAT_NumMappedPageUpdate, MappedTilesToProduce.Num());
-		SubmitRequestsFromLocalTileList(MappedTilesToProduce, EVTProducePageFlags::None, GraphBuilder, FeatureLevel);
+		SubmitRequestsFromLocalTileList(TransientCollectedPages, MappedTilesToProduce, EVTProducePageFlags::None, GraphBuilder, FeatureLevel);
 		MappedTilesToProduce.Reset();
+		MappedTilesToProduce.Append(TransientCollectedPages);
+		TransientCollectedPages.Reset();
 	}
 
 	{
 		INC_DWORD_STAT_BY(STAT_NumContinuousPageUpdate, ContinuousUpdateTilesToProduce.Num());
-		SubmitRequestsFromLocalTileList(ContinuousUpdateTilesToProduce, EVTProducePageFlags::ContinuousUpdate, GraphBuilder, FeatureLevel);
+		SubmitRequestsFromLocalTileList(TransientCollectedPages, ContinuousUpdateTilesToProduce, EVTProducePageFlags::ContinuousUpdate, GraphBuilder, FeatureLevel);
 		ContinuousUpdateTilesToProduce.Reset();
+		TransientCollectedPages.Reset();
 	}
 }
 
diff --git a/Engine/Source/Runtime/Renderer/Private/VT/VirtualTextureSystem.h b/Engine/Source/Runtime/Renderer/Private/VT/VirtualTextureSystem.h
--- a/Engine/Source/Runtime/Renderer/Private/VT/VirtualTextureSystem.h
+++ b/Engine/Source/Runtime/Renderer/Private/VT/VirtualTextureSystem.h
@@ -84,7 +84,7 @@ private:
 
 	void RequestTilesForRegionInternal(const IAllocatedVirtualTexture* AllocatedVT, const FIntRect& InTextureRegion, uint32 vLevel);
 	
-	void SubmitRequestsFromLocalTileList(const TSet<FVirtualTextureLocalTile>& LocalTileList, EVTProducePageFlags Flags, FRDGBuilder& GraphBuilder, ERHIFeatureLevel::Type FeatureLevel);
+	void SubmitRequestsFromLocalTileList(TArray<FVirtualTextureLocalTile>& OutDeferredTiles, const TSet<FVirtualTextureLocalTile>& LocalTileList, EVTProducePageFlags Flags, FRDGBuilder& GraphBuilder, ERHIFeatureLevel::Type FeatureLevel);
 
 	void SubmitPreMappedRequests(FRDGBuilder& GraphBuilder, ERHIFeatureLevel::Type FeatureLevel);
 
