Removed unused scene texture uniform buffer instance.
#rnx

diff --git a/Engine/Source/Runtime/Renderer/Private/SceneCaptureRendering.cpp b/Engine/Source/Runtime/Renderer/Private/SceneCaptureRendering.cpp
--- a/Engine/Source/Runtime/Renderer/Private/SceneCaptureRendering.cpp
+++ b/Engine/Source/Runtime/Renderer/Private/SceneCaptureRendering.cpp
@@ -195,8 +195,6 @@ void FDeferredShadingSceneRenderer::CopySceneCaptureComponentToTarget(
 
 		const FSceneCapturePS::FPermutationDomain PixelPermutationVector = FSceneCapturePS::GetPermutationVector(SceneCaptureSource);
 
-		TRDGUniformBufferRef<FSceneTextureUniformParameters> SceneTextureUniformBuffer = CreateSceneTextureUniformBuffer(GraphBuilder, FeatureLevel, ESceneTextureSetupMode::GBuffers | ESceneTextureSetupMode::SceneColor | ESceneTextureSetupMode::SceneDepth);
-
 		for (int32 ViewIndex = 0; ViewIndex < Views.Num(); ViewIndex++)
 		{
 			FViewInfo& View = Views[ViewIndex];
