Fixed black scene when clouds are compiling.

diff --git a/Engine/Source/Runtime/Renderer/Private/DeferredShadingRenderer.cpp b/Engine/Source/Runtime/Renderer/Private/DeferredShadingRenderer.cpp
--- a/Engine/Source/Runtime/Renderer/Private/DeferredShadingRenderer.cpp
+++ b/Engine/Source/Runtime/Renderer/Private/DeferredShadingRenderer.cpp
@@ -2048,10 +2048,7 @@ void FDeferredShadingSceneRenderer::Render(FRHICommandListImmediate& RHICmdList)
 		InitVolumetricRenderTargetForViews(GraphBuilder, Views);
 	}
 
-	if (bShouldRenderVolumetricCloud)
-	{
-		InitVolumetricCloudsForViews(GraphBuilder);
-	}
+	InitVolumetricCloudsForViews(GraphBuilder, bShouldRenderVolumetricCloud);
 
 	// Generate sky LUTs once all shadow map has been evaluated (for volumetric light shafts). Requires bOcclusionBeforeBasePass.
 	// This also must happen before the BasePass for Sky material to be able to sample valid LUTs.
diff --git a/Engine/Source/Runtime/Renderer/Private/SceneRendering.h b/Engine/Source/Runtime/Renderer/Private/SceneRendering.h
--- a/Engine/Source/Runtime/Renderer/Private/SceneRendering.h
+++ b/Engine/Source/Runtime/Renderer/Private/SceneRendering.h
@@ -1996,7 +1996,7 @@ protected:
 		FRDGTextureRef SceneDepthTexture);
 
 	/** Initialise volumetric cloud resources.*/
-	void InitVolumetricCloudsForViews(FRDGBuilder& GraphBuilder);
+	void InitVolumetricCloudsForViews(FRDGBuilder& GraphBuilder, bool bShouldRenderVolumetricCloud);
 
 	/** Render volumetric cloud. */
 	void RenderVolumetricCloud(
diff --git a/Engine/Source/Runtime/Renderer/Private/VolumetricCloudRendering.cpp b/Engine/Source/Runtime/Renderer/Private/VolumetricCloudRendering.cpp
--- a/Engine/Source/Runtime/Renderer/Private/VolumetricCloudRendering.cpp
+++ b/Engine/Source/Runtime/Renderer/Private/VolumetricCloudRendering.cpp
@@ -197,8 +197,17 @@ bool ShouldRenderVolumetricCloud(const FScene* Scene, const FEngineShowFlags& En
 
 		const bool bShadersCompiled = ShouldPipelineCompileVolumetricCloudShader(Scene->GetShaderPlatform());
 
+		bool bCloudMaterialValid = false;
+		if (VolumetricCloud->GetVolumetricCloudSceneProxy().GetCloudVolumeMaterial())
+		{
+			FMaterialRenderProxy* CloudVolumeMaterialProxy = VolumetricCloud->GetVolumetricCloudSceneProxy().GetCloudVolumeMaterial()->GetRenderProxy();
+			const FMaterialRenderProxy* FallbackMaterialRenderProxyPtr = nullptr;
+			const FMaterial& MaterialTest = CloudVolumeMaterialProxy->GetMaterialWithFallback(Scene->GetFeatureLevel(), FallbackMaterialRenderProxyPtr);
+			bCloudMaterialValid = MaterialTest.GetMaterialDomain() == MD_Volume;
+		}
+
 		FLightSceneInfo* AtmosphericLight0 = Scene->AtmosphereLights[0];
-		return bShadersCompiled && CVarVolumetricCloud.GetValueOnRenderThread() > 0 && AtmosphericLight0!=nullptr && VolumetricCloud->GetVolumetricCloudSceneProxy().GetCloudVolumeMaterial()!=nullptr;
+		return bShadersCompiled && CVarVolumetricCloud.GetValueOnRenderThread() > 0 && AtmosphericLight0!=nullptr && bCloudMaterialValid;
 	}
 	return false;
 }
@@ -783,12 +792,7 @@ public:
 		// Determine the mesh's material and blend mode.
 		const FMaterialRenderProxy* FallbackMaterialRenderProxyPtr = nullptr;
 		const FMaterial& Material = MeshBatch.MaterialRenderProxy->GetMaterialWithFallback(FeatureLevel, FallbackMaterialRenderProxyPtr);
-
-		if (Material.GetMaterialDomain() != MD_Volume)
-		{
-			// Skip in this case. This can happens when the material is compiled and a fallback is provided.
-			return;
-		}
+		check(Material.GetMaterialDomain() == MD_Volume);
 
 		const ERasterizerFillMode MeshFillMode = FM_Solid;
 		const ERasterizerCullMode MeshCullMode = CM_None;
@@ -957,7 +961,6 @@ public:
 		// Determine the mesh's material and blend mode.
 		const FMaterialRenderProxy* FallbackMaterialRenderProxyPtr = nullptr;
 		const FMaterial& Material = MeshBatch.MaterialRenderProxy->GetMaterialWithFallback(FeatureLevel, FallbackMaterialRenderProxyPtr);
-
 		check(Material.GetMaterialDomain() == MD_Volume);
 
 		const ERasterizerFillMode MeshFillMode = FM_Solid;
@@ -1117,7 +1120,7 @@ IMPLEMENT_GLOBAL_SHADER(FCloudShadowTemporalProcessCS, "/Engine/Private/Volumetr
 
 
 
-void FSceneRenderer::InitVolumetricCloudsForViews(FRDGBuilder& GraphBuilder)
+void FSceneRenderer::InitVolumetricCloudsForViews(FRDGBuilder& GraphBuilder, bool bShouldRenderVolumetricCloud)
 {
 	auto CleanUpCloudData = [&Views = Views](FRDGBuilder& GraphBuilder)
 	{
@@ -1138,6 +1141,13 @@ void FSceneRenderer::InitVolumetricCloudsForViews(FRDGBuilder& GraphBuilder)
 		});
 	};
 
+	if (!bShouldRenderVolumetricCloud)
+	{
+		// Clean everything up
+		CleanUpCloudData(GraphBuilder);
+		return;
+	}
+
 	if (Scene)
 	{
 		check(ShouldRenderVolumetricCloud(Scene, ViewFamily.EngineShowFlags)); // This should not be called if we should not render SkyAtmosphere
