Fix for crash in cascade FX system due to null scene textures.

diff --git a/Engine/Source/Runtime/Engine/Private/Particles/ParticleGpuSimulation.cpp b/Engine/Source/Runtime/Engine/Private/Particles/ParticleGpuSimulation.cpp
--- a/Engine/Source/Runtime/Engine/Private/Particles/ParticleGpuSimulation.cpp
+++ b/Engine/Source/Runtime/Engine/Private/Particles/ParticleGpuSimulation.cpp
@@ -1435,13 +1435,6 @@ void ExecuteSimulationCommands(
 	SCOPED_DRAW_EVENT(RHICmdList, ParticleSimulation);
 	SCOPED_GPU_STAT(RHICmdList, ParticleSimulation);
 
-	FUniformBufferStaticBindings GlobalUniformBuffers;
-	if (SceneTexturesUniformBuffer)
-	{
-		GlobalUniformBuffers.AddUniformBuffer(SceneTexturesUniformBuffer);
-	}
-	SCOPED_UNIFORM_BUFFER_GLOBAL_BINDINGS(RHICmdList, GlobalUniformBuffers);
-
 	const float FixDeltaSeconds = CVarGPUParticleFixDeltaSeconds.GetValueOnRenderThread();
 	const FParticleStateTextures& TextureResources = (FixDeltaSeconds <= 0 || bUseFixDT) ? ParticleSimulationResources->GetPreviousStateTextures() : ParticleSimulationResources->GetCurrentStateTextures();
 	const FParticleAttributesTexture& AttributeTexture = ParticleSimulationResources->SimulationAttributesTexture;
diff --git a/Engine/Source/Runtime/Renderer/Private/DeferredShadingRenderer.cpp b/Engine/Source/Runtime/Renderer/Private/DeferredShadingRenderer.cpp
--- a/Engine/Source/Runtime/Renderer/Private/DeferredShadingRenderer.cpp
+++ b/Engine/Source/Runtime/Renderer/Private/DeferredShadingRenderer.cpp
@@ -433,6 +433,51 @@ static int32 GetCustomDepthPassLocation()
 	return FMath::Clamp(CVarCustomDepthOrder.GetValueOnRenderThread(), 0, 1);
 }
 
+BEGIN_SHADER_PARAMETER_STRUCT(FRenderOpaqueFXPassParameters, )
+	SHADER_PARAMETER_RDG_UNIFORM_BUFFER(FSceneTextureUniformParameters, SceneTextures)
+END_SHADER_PARAMETER_STRUCT()
+
+static void RenderOpaqueFX(
+	FRDGBuilder& GraphBuilder,
+	TArrayView<const FViewInfo> Views,
+	FFXSystemInterface* FXSystem,
+	TRDGUniformBufferRef<FSceneTextureUniformParameters> SceneTexturesUniformBuffer)
+{
+	// Notify the FX system that opaque primitives have been rendered and we now have a valid depth buffer.
+	if (FXSystem && Views.Num() > 0)
+	{
+		RDG_GPU_STAT_SCOPE(GraphBuilder, PostRenderOpsFX);
+		RDG_CSV_STAT_EXCLUSIVE_SCOPE(GraphBuilder, RenderOpaqueFX);
+
+		auto* PassParameters = GraphBuilder.AllocParameters<FRenderOpaqueFXPassParameters>();
+		PassParameters->SceneTextures = SceneTexturesUniformBuffer;
+
+		GraphBuilder.AddPass(
+			RDG_EVENT_NAME("OpaqueFX"),
+			PassParameters,
+			ERDGPassFlags::Compute | ERDGPassFlags::NeverCull,
+			[FXSystem, Views](FRHICommandListImmediate& RHICmdList)
+		{
+			SCOPE_CYCLE_COUNTER(STAT_FDeferredShadingSceneRenderer_FXSystem_PostRenderOpaque);
+
+			FXSystem->PostRenderOpaque(
+				RHICmdList,
+				Views[0].ViewUniformBuffer,
+				&FSceneTextureUniformParameters::StaticStructMetadata,
+				nullptr,
+				Views[0].AllowGPUParticleUpdate()
+			);
+
+			if (FGPUSortManager* GPUSortManager = FXSystem->GetGPUSortManager())
+			{
+				GPUSortManager->OnPostRenderOpaque(RHICmdList);
+			}
+
+			ServiceLocalQueue();
+		});
+	}
+}
+
 void FDeferredShadingSceneRenderer::PrepareDistanceFieldScene(FRDGBuilder& GraphBuilder, bool bSplitDispatch)
 {
 	CSV_SCOPED_TIMING_STAT_EXCLUSIVE(RenderDFAO);
@@ -2590,32 +2635,7 @@ void FDeferredShadingSceneRenderer::Render(FRHICommandListImmediate& RHICmdList)
 	FRendererModule& RendererModule = static_cast<FRendererModule&>(GetRendererModule());
 	RendererModule.RenderPostOpaqueExtensions(GraphBuilder, Views, SceneColorTexture.Target, SceneDepthTexture.Target, SceneTextures, SceneContext);
 
-	// Notify the FX system that opaque primitives have been rendered and we now have a valid depth buffer.
-	if (FXSystem && Views.IsValidIndex(0))
-	{
-		RDG_GPU_STAT_SCOPE(GraphBuilder, PostRenderOpsFX);
-		RDG_CSV_STAT_EXCLUSIVE_SCOPE(GraphBuilder, RenderOpaqueFX);
-
-		AddPass(GraphBuilder, [this](FRHICommandListImmediate& InRHICmdList)
-		{
-			SCOPE_CYCLE_COUNTER(STAT_FDeferredShadingSceneRenderer_FXSystem_PostRenderOpaque);
-
-			FXSystem->PostRenderOpaque(
-				InRHICmdList,
-				Views[0].ViewUniformBuffer,
-				nullptr,
-				nullptr,
-				Views[0].AllowGPUParticleUpdate()
-			);
-
-			if (FGPUSortManager* GPUSortManager = FXSystem->GetGPUSortManager())
-			{
-				GPUSortManager->OnPostRenderOpaque(InRHICmdList);
-			}
-
-			ServiceLocalQueue();
-		});
-	}
+	RenderOpaqueFX(GraphBuilder, Views, FXSystem, SceneTextures);
 
 	if (bShouldRenderSkyAtmosphere)
 	{
