add ERDGPassFlags::Raster | ERDGPassFlags::SkipRenderPass to Cascade FX pass because it's using rasterization internally as well

diff --git a/Engine/Source/Runtime/Renderer/Private/DeferredShadingRenderer.cpp b/Engine/Source/Runtime/Renderer/Private/DeferredShadingRenderer.cpp
--- a/Engine/Source/Runtime/Renderer/Private/DeferredShadingRenderer.cpp
+++ b/Engine/Source/Runtime/Renderer/Private/DeferredShadingRenderer.cpp
@@ -452,10 +452,11 @@ static void RenderOpaqueFX(
 		auto* PassParameters = GraphBuilder.AllocParameters<FRenderOpaqueFXPassParameters>();
 		PassParameters->SceneTextures = SceneTexturesUniformBuffer;
 
+		// Cascade uses pixel shaders for compute stuff in PostRenderOpaque so ERDGPassFlags::Raster is needed
 		GraphBuilder.AddPass(
 			RDG_EVENT_NAME("OpaqueFX"),
 			PassParameters,
-			ERDGPassFlags::Compute | ERDGPassFlags::NeverCull,
+			ERDGPassFlags::Raster | ERDGPassFlags::SkipRenderPass | ERDGPassFlags::Compute | ERDGPassFlags::NeverCull,
 			[FXSystem, Views](FRHICommandListImmediate& RHICmdList)
 		{
 			SCOPE_CYCLE_COUNTER(STAT_FDeferredShadingSceneRenderer_FXSystem_PostRenderOpaque);
