Fix for shadow depth rendering not using the correct vertex shader when geometry shaders are not supported.

diff --git a/Engine/Source/Runtime/Renderer/Private/ShadowDepthRendering.cpp b/Engine/Source/Runtime/Renderer/Private/ShadowDepthRendering.cpp
--- a/Engine/Source/Runtime/Renderer/Private/ShadowDepthRendering.cpp
+++ b/Engine/Source/Runtime/Renderer/Private/ShadowDepthRendering.cpp
@@ -931,9 +931,8 @@ void FProjectedShadowInfo::CopyCachedShadowMap(
 
 		if (bOnePassPointLightShadow)
 		{
-			TShaderMapRef<TScreenVSForGS<false>> ScreenVertexShader(View.ShaderMap);
+			TShaderRef<FScreenVS> ScreenVertexShader;
 			TShaderMapRef<FCopyShadowMapsCubePS> PixelShader(View.ShaderMap);
-			GraphicsPSOInit.BoundShaderState.VertexShaderRHI = ScreenVertexShader.GetVertexShader();
 			GraphicsPSOInit.BoundShaderState.PixelShaderRHI = PixelShader.GetPixelShader();
 
 			int32 InstanceCount = 1;
@@ -941,14 +940,21 @@ void FProjectedShadowInfo::CopyCachedShadowMap(
 #if PLATFORM_SUPPORTS_GEOMETRY_SHADERS
 			if (RHISupportsGeometryShaders(GShaderPlatformForFeatureLevel[SceneRenderer->FeatureLevel]))
 			{
+				TShaderMapRef<TScreenVSForGS<false>> VertexShader(View.ShaderMap);
 				TShaderMapRef<FCopyShadowMapsCubeGS> GeometryShader(View.ShaderMap);
+				GraphicsPSOInit.BoundShaderState.VertexShaderRHI = VertexShader.GetVertexShader();
 				GraphicsPSOInit.BoundShaderState.GeometryShaderRHI = GeometryShader.GetGeometryShader();
 				InstanceCount = 6;
+
+				ScreenVertexShader = VertexShader;
 			}
 			else
 #endif
 			{
 				check(RHISupportsVertexShaderLayer(GShaderPlatformForFeatureLevel[SceneRenderer->FeatureLevel]));
+				TShaderMapRef<TScreenVSForGS<true>> VertexShader(View.ShaderMap);
+				GraphicsPSOInit.BoundShaderState.VertexShaderRHI = VertexShader.GetVertexShader();
+				ScreenVertexShader = VertexShader;
 			}
 
 			auto* PassParameters = GraphBuilder.AllocParameters<FCopyShadowMapsCubePS::FParameters>();
