Added RDG system textures struct tied to the RDG blackboard. Pre-registers all resources with a ReadOnly flag which will disallow writes. Also allows us to avoid re-registering dummy textures dozens of times. Removed unused GTAO randomization texture.

diff --git a/Engine/Source/Runtime/RenderCore/Private/RenderGraphBuilder.cpp b/Engine/Source/Runtime/RenderCore/Private/RenderGraphBuilder.cpp
--- a/Engine/Source/Runtime/RenderCore/Private/RenderGraphBuilder.cpp
+++ b/Engine/Source/Runtime/RenderCore/Private/RenderGraphBuilder.cpp
@@ -454,9 +454,6 @@ FRDGTextureRef FRDGBuilder::RegisterExternalTexture(
 
 	if (FRDGTextureRef FoundTexture = FindExternalTexture(ExternalTextureRHI))
 	{
-#if RDG_ENABLE_DEBUG
-		checkf(FoundTexture->Flags == Flags, TEXT("External texture %s is already registered, but with different resource flags."), Name);
-#endif
 		return FoundTexture;
 	}
 
@@ -517,11 +514,7 @@ FRDGBufferRef FRDGBuilder::RegisterExternalBuffer(
 
 	if (FRDGBufferRef* FoundBufferPtr = ExternalBuffers.Find(ExternalPooledBuffer.GetReference()))
 	{
-		FRDGBufferRef FoundBuffer = *FoundBufferPtr;
-#if RDG_ENABLE_DEBUG
-		checkf(FoundBuffer->Flags == Flags, TEXT("External buffer %s is already registered, but with different resource flags."), Name);
-#endif
-		return FoundBuffer;
+		return *FoundBufferPtr;
 	}
 
 	FRDGBufferRef Buffer = Buffers.Allocate(Allocator, Name, ExternalPooledBuffer->Desc, Flags);
diff --git a/Engine/Source/Runtime/RenderCore/Private/RenderGraphValidation.cpp b/Engine/Source/Runtime/RenderCore/Private/RenderGraphValidation.cpp
--- a/Engine/Source/Runtime/RenderCore/Private/RenderGraphValidation.cpp
+++ b/Engine/Source/Runtime/RenderCore/Private/RenderGraphValidation.cpp
@@ -295,6 +295,8 @@ void FRDGUserValidation::ValidateCreateTexture(const FRDGTextureDesc& Desc, cons
 	const bool bIsUAVForMSAATexture = bIsMSAA && bCanHaveUAV;
 	checkf(!bIsUAVForMSAATexture, TEXT("TexCreate_UAV is not allowed on MSAA texture %s."), Name);
 
+	checkf(!EnumHasAnyFlags(Flags, ERDGTextureFlags::ReadOnly), TEXT("Cannot create texture %s with the ReadOnly flag. Only registered textures can use this flag."), Name);
+
 	checkf(!EnumHasAnyFlags(BuilderFlags, ERDGBuilderFlags::SkipBarriers), TEXT("Cannot create texture '%s' because RDG builder has 'SkipBarriers' flag."), Name);
 }
 
@@ -312,6 +314,8 @@ void FRDGUserValidation::ValidateCreateBuffer(const FRDGBufferDesc& Desc, const
 		checkf(Desc.BytesPerElement == 4, TEXT("Creating buffer '%s' as a structured buffer that is also byte addressable, BytesPerElement must be 4! Instead it is %d"), Name, Desc.BytesPerElement);
 	}
 
+	checkf(!EnumHasAnyFlags(Flags, ERDGBufferFlags::ReadOnly), TEXT("Cannot create buffer %s with the ReadOnly flag. Only registered buffers can use this flag."), Name);
+
 	checkf(!EnumHasAnyFlags(BuilderFlags, ERDGBuilderFlags::SkipBarriers), TEXT("Cannot create buffer '%s' because RDG builder has 'SkipBarriers' flag."), Name);
 }
 
@@ -517,7 +521,7 @@ void FRDGUserValidation::ValidateAddPass(const FRDGPass* Pass, bool bSkipPassAcc
 	const bool bIsAnyCompute = EnumHasAnyFlags(PassFlags, ERDGPassFlags::Compute | ERDGPassFlags::AsyncCompute);
 	const bool bSkipRenderPass = EnumHasAnyFlags(PassFlags, ERDGPassFlags::SkipRenderPass);
 
-	const auto MarkAsProduced = [&] (FRDGParentResourceRef Resource)
+	const auto MarkAsProduced = [&](FRDGParentResourceRef Resource)
 	{
 		if (!bSkipPassAccessMarking)
 		{
@@ -530,6 +534,18 @@ void FRDGUserValidation::ValidateAddPass(const FRDGPass* Pass, bool bSkipPassAcc
 		}
 	};
 
+	const auto MarkTextureAsProduced = [&](FRDGTextureRef Texture)
+	{
+		checkf(!EnumHasAnyFlags(Texture->Flags, ERDGTextureFlags::ReadOnly), TEXT("Pass %s is attempting to write to texture %s which is marked as ReadOnly."), Pass->GetName(), Texture->Name);
+		MarkAsProduced(Texture);
+	};
+
+	const auto MarkBufferAsProduced = [&](FRDGBufferRef Buffer)
+	{
+		checkf(!EnumHasAnyFlags(Buffer->Flags, ERDGBufferFlags::ReadOnly), TEXT("Pass %s is attempting to write to buffer %s which is marked as ReadOnly."), Pass->GetName(), Buffer->Name);
+		MarkAsProduced(Buffer);
+	};
+
 	const auto MarkAsConsumed = [&] (FRDGParentResourceRef Resource)
 	{
 		ensureMsgf(Resource->bProduced || Resource->bExternal,
@@ -598,7 +614,7 @@ void FRDGUserValidation::ValidateAddPass(const FRDGPass* Pass, bool bSkipPassAcc
 			{
 				FRDGTextureRef Texture = UAV->GetParent();
 				CheckNotCopy(Texture);
-				MarkAsProduced(Texture);
+				MarkTextureAsProduced(Texture);
 			}
 		}
 		break;
@@ -627,7 +643,7 @@ void FRDGUserValidation::ValidateAddPass(const FRDGPass* Pass, bool bSkipPassAcc
 			{
 				FRDGBufferRef Buffer = UAV->GetParent();
 				CheckNotCopy(Buffer);
-				MarkAsProduced(Buffer);
+				MarkBufferAsProduced(Buffer);
 			}
 		}
 		break;
@@ -647,7 +663,7 @@ void FRDGUserValidation::ValidateAddPass(const FRDGPass* Pass, bool bSkipPassAcc
 
 				if (bIsWriteAccess)
 				{
-					MarkAsProduced(Texture);
+					MarkTextureAsProduced(Texture);
 				}
 			}
 		}
@@ -668,7 +684,7 @@ void FRDGUserValidation::ValidateAddPass(const FRDGPass* Pass, bool bSkipPassAcc
 
 				if (IsWritableAccess(BufferAccess.GetAccess()))
 				{
-					MarkAsProduced(Buffer);
+					MarkBufferAsProduced(Buffer);
 				}
 			}
 		}
@@ -710,7 +726,7 @@ void FRDGUserValidation::ValidateAddPass(const FRDGPass* Pass, bool bSkipPassAcc
 				CheckNotPassthrough(Texture);
 				if (DepthStencil.GetDepthStencilAccess().IsAnyWrite())
 				{
-					MarkAsProduced(Texture);
+					MarkTextureAsProduced(Texture);
 				}
 				else
 				{
@@ -754,7 +770,7 @@ void FRDGUserValidation::ValidateAddPass(const FRDGPass* Pass, bool bSkipPassAcc
 						PassName, ResolveTexture->Name);
 
 					CheckNotPassthrough(ResolveTexture);
-					MarkAsProduced(ResolveTexture);
+					MarkTextureAsProduced(ResolveTexture);
 				}
 
 				if (Texture)
@@ -793,7 +809,7 @@ void FRDGUserValidation::ValidateAddPass(const FRDGPass* Pass, bool bSkipPassAcc
 					}
 
 					/** Mark the pass as a producer for render targets with a store action. */
-					MarkAsProduced(Texture);
+					MarkTextureAsProduced(Texture);
 				}
 				else
 				{
diff --git a/Engine/Source/Runtime/RenderCore/Public/RenderGraphDefinitions.h b/Engine/Source/Runtime/RenderCore/Public/RenderGraphDefinitions.h
--- a/Engine/Source/Runtime/RenderCore/Public/RenderGraphDefinitions.h
+++ b/Engine/Source/Runtime/RenderCore/Public/RenderGraphDefinitions.h
@@ -111,7 +111,10 @@ enum class ERDGBufferFlags : uint8
 	None = 0,
 
 	/** Tag the buffer to survive through frame, that is important for multi GPU alternate frame rendering. */
-	MultiFrame = 1 << 0
+	MultiFrame = 1 << 0,
+
+	/** The buffer may only be used for read-only access within the graph. This flag is only allowed for registered buffers. */
+	ReadOnly = 1 << 1
 };
 ENUM_CLASS_FLAGS(ERDGBufferFlags);
 
@@ -123,8 +126,11 @@ enum class ERDGTextureFlags : uint8
 	/** Tag the texture to survive through frame, that is important for multi GPU alternate frame rendering. */
 	MultiFrame = 1 << 0,
 
+	/** The texture may only be used for read-only access within the graph. This flag is only allowed for registered textures. */
+	ReadOnly = 1 << 1,
+
 	/** Prevents metadata decompression on this texture. */
-	MaintainCompression = 1 << 1,
+	MaintainCompression = 1 << 2,
 };
 ENUM_CLASS_FLAGS(ERDGTextureFlags);
 
diff --git a/Engine/Source/Runtime/Renderer/Private/CompositionLighting/PostProcessAmbientOcclusion.cpp b/Engine/Source/Runtime/Renderer/Private/CompositionLighting/PostProcessAmbientOcclusion.cpp
--- a/Engine/Source/Runtime/Renderer/Private/CompositionLighting/PostProcessAmbientOcclusion.cpp
+++ b/Engine/Source/Runtime/Renderer/Private/CompositionLighting/PostProcessAmbientOcclusion.cpp
@@ -331,15 +331,7 @@ static FSSAOShaderParameters GetSSAOShaderParameters(
 {
 	const FFinalPostProcessSettings& Settings = View.FinalPostProcessSettings;
 
-	FIntPoint RandomizationSize;
-	if (AOTechnique == EAOTechnique::GTAO)
-	{
-		RandomizationSize = GSystemTextures.GTAORandomization->GetDesc().Extent;
-	}
-	else
-	{
-		RandomizationSize = GSystemTextures.SSAORandomization->GetDesc().Extent;
-	}
+	const FIntPoint RandomizationSize = GSystemTextures.SSAORandomization->GetDesc().Extent;
 	FVector2D ViewportUVToRandomUV(InputViewport.Extent.X / (float)RandomizationSize.X, InputViewport.Extent.Y / (float)RandomizationSize.Y);
 
 	// e.g. 4 means the input texture is 4x smaller than the buffer size
@@ -681,7 +673,7 @@ BEGIN_SHADER_PARAMETER_STRUCT(FAmbientOcclusionParameters, )
 
 	SHADER_PARAMETER_SAMPLER(SamplerState, SSAO_Sampler)
 
-	SHADER_PARAMETER_RDG_TEXTURE(Texture2D, RandomNormalTexture)
+	SHADER_PARAMETER_TEXTURE(Texture2D, RandomNormalTexture)
 	SHADER_PARAMETER_SAMPLER(SamplerState, RandomNormalTextureSampler)
 END_SHADER_PARAMETER_STRUCT();
 
@@ -828,7 +820,7 @@ FScreenPassTexture AddAmbientOcclusionPass(
 
 	SharedParameters.SSAO_Sampler = TStaticSamplerState<SF_Point, AM_Clamp, AM_Clamp, AM_Clamp>::GetRHI();
 
-	SharedParameters.RandomNormalTexture = GraphBuilder.RegisterExternalTexture(GSystemTextures.SSAORandomization, TEXT("SSAORandomization"));
+	SharedParameters.RandomNormalTexture = GSystemTextures.SSAORandomization->GetShaderResourceRHI();
 	SharedParameters.RandomNormalTextureSampler = TStaticSamplerState<SF_Point, AM_Wrap, AM_Wrap, AM_Wrap>::GetRHI();
 
 	FRDGEventName EventName(TEXT("AmbientOcclusion%s %dx%d SetupAsInput=%d Upsample=%d ShaderQuality=%d"),
diff --git a/Engine/Source/Runtime/Renderer/Private/DeferredShadingRenderer.cpp b/Engine/Source/Runtime/Renderer/Private/DeferredShadingRenderer.cpp
--- a/Engine/Source/Runtime/Renderer/Private/DeferredShadingRenderer.cpp
+++ b/Engine/Source/Runtime/Renderer/Private/DeferredShadingRenderer.cpp
@@ -1562,6 +1562,9 @@ void FDeferredShadingSceneRenderer::Render(FRHICommandListImmediate& RHICmdList)
 
 		// Allocate the maximum scene render target space for the current view family.
 		SceneContext.Allocate(GraphBuilder, this);
+
+		// Initialize the RDG read-only system textures.
+		FRDGSystemTextures::Create(GraphBuilder);
 	}
 
 	const bool bUseVirtualTexturing = UseVirtualTexturing(FeatureLevel);
diff --git a/Engine/Source/Runtime/Renderer/Private/Renderer.cpp b/Engine/Source/Runtime/Renderer/Private/Renderer.cpp
--- a/Engine/Source/Runtime/Renderer/Private/Renderer.cpp
+++ b/Engine/Source/Runtime/Renderer/Private/Renderer.cpp
@@ -201,6 +201,9 @@ void FRendererModule::DrawTileMesh(FRHICommandListImmediate& RHICmdList, FMeshPa
 		FMemMark Mark(FMemStack::Get());
 		FRDGBuilder GraphBuilder(RHICmdList, {}, ERDGBuilderFlags::SkipBarriers);
 
+		// Initialize the RDG read-only system textures.
+		FRDGSystemTextures::Create(GraphBuilder);
+
 		const ERDGPassFlags PassFlags = ERDGPassFlags::Raster | ERDGPassFlags::SkipRenderPass | ERDGPassFlags::NeverCull;
 
 		auto* PassParameters = GraphBuilder.AllocParameters<FDrawTileMeshPassParameters>();
diff --git a/Engine/Source/Runtime/Renderer/Private/SystemTextures.cpp b/Engine/Source/Runtime/Renderer/Private/SystemTextures.cpp
--- a/Engine/Source/Runtime/Renderer/Private/SystemTextures.cpp
+++ b/Engine/Source/Runtime/Renderer/Private/SystemTextures.cpp
@@ -15,9 +15,59 @@
 SystemTextures
 -----------------------------------------------------------------------------*/
 
+RDG_REGISTER_BLACKBOARD_STRUCT(FRDGSystemTextures);
+
+const FRDGSystemTextures& FRDGSystemTextures::Create(FRDGBuilder& GraphBuilder)
+{
+	const auto Register = [&](const TRefCountPtr<IPooledRenderTarget>& RenderTarget)
+	{
+		return GraphBuilder.RegisterExternalTexture(RenderTarget, ERenderTargetTexture::ShaderResource, ERDGTextureFlags::ReadOnly);
+	};
+
+	auto& SystemTextures = GraphBuilder.Blackboard.Create<FRDGSystemTextures>();
+	SystemTextures.White = Register(GSystemTextures.WhiteDummy);
+	SystemTextures.Black = Register(GSystemTextures.BlackDummy);
+	SystemTextures.BlackAlphaOne = Register(GSystemTextures.BlackAlphaOneDummy);
+	SystemTextures.MaxFP16Depth = Register(GSystemTextures.MaxFP16Depth);
+	SystemTextures.DepthDummy = Register(GSystemTextures.DepthDummy);
+	SystemTextures.StencilDummy = Register(GSystemTextures.StencilDummy);
+	SystemTextures.Green = Register(GSystemTextures.GreenDummy);
+	SystemTextures.DefaultNormal8Bit = Register(GSystemTextures.DefaultNormal8Bit);
+	SystemTextures.MidGrey = Register(GSystemTextures.MidGreyDummy);
+	SystemTextures.VolumetricBlack = Register(GSystemTextures.VolumetricBlackDummy);
+	return SystemTextures;
+}
+
+const FRDGSystemTextures& FRDGSystemTextures::Get(FRDGBuilder& GraphBuilder)
+{
+	const FRDGSystemTextures* SystemTextures = GraphBuilder.Blackboard.Get<FRDGSystemTextures>();
+	checkf(SystemTextures, TEXT("FRDGSystemTextures were not initialized. Call FRDGSystemTextures::Create() first."));
+	return *SystemTextures;
+}
+
 /** The global render targets used for scene rendering. */
 TGlobalResource<FSystemTextures> GSystemTextures;
 
+void FSystemTextures::InitializeTextures(FRHICommandListImmediate& RHICmdList, const ERHIFeatureLevel::Type InFeatureLevel)
+{
+	// When we render to system textures it should occur on all GPUs since this only
+	// happens once on startup (or when the feature level changes).
+	SCOPED_GPU_MASK(RHICmdList, FRHIGPUMask::All());
+
+	// if this is the first call initialize everything
+	if (FeatureLevelInitializedTo == ERHIFeatureLevel::Num)
+	{
+		InitializeCommonTextures(RHICmdList);
+		InitializeFeatureLevelDependentTextures(RHICmdList, InFeatureLevel);
+	}
+	// otherwise, if we request a higher feature level, we might need to initialize those textures that depend on the feature level
+	else if (InFeatureLevel > FeatureLevelInitializedTo)
+	{
+		InitializeFeatureLevelDependentTextures(RHICmdList, InFeatureLevel);
+	}
+	// there's no needed setup for those feature levels lower or identical to the current one
+}
+
 void FSystemTextures::InitializeCommonTextures(FRHICommandListImmediate& RHICmdList)
 {
 	// First initialize textures that are common to all feature levels. This is always done the first time we come into this function, as doesn't care about the
@@ -571,36 +621,7 @@ void FSystemTextures::InitializeFeatureLevelDependentTextures(FRHICommandListImm
 		        }
 		        RHICmdList.UnlockTexture2D((FTexture2DRHIRef&)SSAORandomization->GetRenderTargetItem().ShaderResourceTexture, 0, false);
 			}
-	
-			{
-				FPooledRenderTargetDesc Desc(FPooledRenderTargetDesc::Create2DDesc(FIntPoint(4, 4), PF_R8G8B8A8, FClearValueBinding::None, TexCreate_HideInVisualizeTexture, TexCreate_None | TexCreate_NoFastClear, false));
-				Desc.AutoWritable = false;
-				GRenderTargetPool.FindFreeElement(RHICmdList, Desc,GTAORandomization, TEXT("GTAORandomization"), ERenderTargetTransience::NonTransient);
-				// Write the contents of the texture.
-				uint32 DestStride;
-				uint8* DestBuffer = (uint8*)RHICmdList.LockTexture2D((FTexture2DRHIRef&)GTAORandomization->GetRenderTargetItem().ShaderResourceTexture, 0, RLM_WriteOnly, DestStride, false);
-
-				for(int32 y = 0; y < Desc.Extent.Y; ++y)
-				{
-					for(int32 x = 0; x < Desc.Extent.X; ++x)
-					{
-						uint8* Dest = (uint8*)(DestBuffer + x * sizeof(uint32) + y * DestStride);
-			
-						float Angle  = (PI/16.0f)  * ((((x + y) & 0x3) << 2) + (x & 0x3));
-						float Step   = (1.0f / 4.0f) * ((y - x) & 0x3);
-
-						float ScaleCos = FMath::Cos(Angle) ;
-						float ScaleSin = FMath::Sin(Angle) ;
-			
-						Dest[0] = (uint8_t)(ScaleCos*127.5f + 127.5f);
-						Dest[1] = (uint8_t)(ScaleSin*127.5f + 127.5f);
-						Dest[2] = (uint8_t)(Step*255.0f);
-						Dest[3] = 0;
-					}
-				}
-			}
-			RHICmdList.UnlockTexture2D((FTexture2DRHIRef&)GTAORandomization->GetRenderTargetItem().ShaderResourceTexture, 0, false);
-			
+
 		    {
 			    FPooledRenderTargetDesc Desc(FPooledRenderTargetDesc::Create2DDesc(FIntPoint(LTC_Size, LTC_Size), PF_FloatRGBA, FClearValueBinding::None, TexCreate_FastVRAM, TexCreate_ShaderResource, false));
 			    Desc.AutoWritable = false;
@@ -759,7 +780,6 @@ void FSystemTextures::ReleaseDynamicRHI()
 	PerlinNoise3D.SafeRelease();
 	SobolSampling.SafeRelease();
 	SSAORandomization.SafeRelease();
-	GTAORandomization.SafeRelease();
 	GTAOPreIntegrated.SafeRelease();
 	PreintegratedGF.SafeRelease();
 	HairLUT0.SafeRelease();
@@ -799,41 +819,6 @@ FRDGTextureRef FSystemTextures::GetWhiteDummy(FRDGBuilder& GraphBuilder) const
 	return GraphBuilder.RegisterExternalTexture(WhiteDummy, TEXT("WhiteDummy"));
 }
 
-FRDGTextureRef FSystemTextures::GetPerlinNoiseGradient(FRDGBuilder& GraphBuilder) const
-{
-	return GraphBuilder.RegisterExternalTexture(PerlinNoiseGradient, TEXT("PerlinNoiseGradient"));
-}
-
-FRDGTextureRef FSystemTextures::GetPerlinNoise3D(FRDGBuilder& GraphBuilder) const
-{
-	return GraphBuilder.RegisterExternalTexture(PerlinNoise3D, TEXT("PerlinNoise3D"));
-}
-
-FRDGTextureRef FSystemTextures::GetSobolSampling(FRDGBuilder& GraphBuilder) const
-{
-	return GraphBuilder.RegisterExternalTexture(SobolSampling, TEXT("SobolSampling"));
-}
-
-FRDGTextureRef FSystemTextures::GetSSAORandomization(FRDGBuilder& GraphBuilder) const
-{
-	return GraphBuilder.RegisterExternalTexture(SSAORandomization, TEXT("SSAORandomization"));
-}
-
-FRDGTextureRef FSystemTextures::GetPreintegratedGF(FRDGBuilder& GraphBuilder) const
-{
-	return GraphBuilder.RegisterExternalTexture(PreintegratedGF, TEXT("PreintegratedGF"));
-}
-
-FRDGTextureRef FSystemTextures::GetLTCMat(FRDGBuilder& GraphBuilder) const
-{
-	return GraphBuilder.RegisterExternalTexture(LTCMat, TEXT("LTCMat"));
-}
-
-FRDGTextureRef FSystemTextures::GetLTCAmp(FRDGBuilder& GraphBuilder) const
-{
-	return GraphBuilder.RegisterExternalTexture(LTCAmp, TEXT("LTCAmp"));
-}
-
 FRDGTextureRef FSystemTextures::GetMaxFP16Depth(FRDGBuilder& GraphBuilder) const
 {
 	return GraphBuilder.RegisterExternalTexture(MaxFP16Depth, TEXT("MaxFP16Depth"));
diff --git a/Engine/Source/Runtime/Renderer/Private/SystemTextures.h b/Engine/Source/Runtime/Renderer/Private/SystemTextures.h
--- a/Engine/Source/Runtime/Renderer/Private/SystemTextures.h
+++ b/Engine/Source/Runtime/Renderer/Private/SystemTextures.h
@@ -8,6 +8,27 @@
 
 #include "RenderGraph.h"
 
+/** Contains system textures which can be registered for read-only access on an RDG pass. */
+struct FRDGSystemTextures
+{
+	/** Call to initialize for the requested builder instance. */
+	RENDERER_API static const FRDGSystemTextures& Create(FRDGBuilder& GraphBuilder);
+
+	/** Returns the instance for the builder. Must be called after Create. */
+	RENDERER_API static const FRDGSystemTextures& Get(FRDGBuilder& GraphBuilder);
+
+	FRDGTextureRef White{};
+	FRDGTextureRef Black{};
+	FRDGTextureRef BlackAlphaOne{};
+	FRDGTextureRef MaxFP16Depth{};
+	FRDGTextureRef DepthDummy{};
+	FRDGTextureRef StencilDummy{};
+	FRDGTextureRef Green{};
+	FRDGTextureRef DefaultNormal8Bit{};
+	FRDGTextureRef MidGrey{};
+	FRDGTextureRef VolumetricBlack{};
+};
+
 /**
  * Encapsulates the system textures used for scene rendering.
  */
@@ -22,25 +43,7 @@ public:
 	/**
 	 * Initialize/allocate textures if not already.
 	 */
-	inline void InitializeTextures(FRHICommandListImmediate& RHICmdList, const ERHIFeatureLevel::Type InFeatureLevel)
-	{
-		// When we render to system textures it should occur on all GPUs since this only
-		// happens once on startup (or when the feature level changes).
-		SCOPED_GPU_MASK(RHICmdList, FRHIGPUMask::All());
-
-		// if this is the first call initialize everything
-		if (FeatureLevelInitializedTo == ERHIFeatureLevel::Num)
-		{
-			InitializeCommonTextures(RHICmdList);
-			InitializeFeatureLevelDependentTextures(RHICmdList, InFeatureLevel);
-		}
-		// otherwise, if we request a higher feature level, we might need to initialize those textures that depend on the feature level
-		else if (InFeatureLevel > FeatureLevelInitializedTo)
-		{
-			InitializeFeatureLevelDependentTextures(RHICmdList, InFeatureLevel);
-		}
-		// there's no needed setup for those feature levels lower or identical to the current one
-	}
+	void InitializeTextures(FRHICommandListImmediate& RHICmdList, const ERHIFeatureLevel::Type InFeatureLevel);
 
 	// FRenderResource interface.
 	/**
@@ -69,8 +72,6 @@ public:
 	TRefCountPtr<IPooledRenderTarget> SobolSampling;
 	/** SSAO randomization */
 	TRefCountPtr<IPooledRenderTarget> SSAORandomization;
-	/** GTAO randomization */
-	TRefCountPtr<IPooledRenderTarget> GTAORandomization;
 	/** GTAO PreIntegrated */
 	TRefCountPtr<IPooledRenderTarget> GTAOPreIntegrated;
 
@@ -112,13 +113,6 @@ public:
 	FRDGTextureRef GetBlackDummy(FRDGBuilder& GraphBuilder) const;
 	FRDGTextureRef GetZeroUIntDummy(FRDGBuilder& GraphBuilder) const;
 	FRDGTextureRef GetBlackAlphaOneDummy(FRDGBuilder& GraphBuilder) const;
-	FRDGTextureRef GetPerlinNoiseGradient(FRDGBuilder& GraphBuilder) const;
-	FRDGTextureRef GetPerlinNoise3D(FRDGBuilder& GraphBuilder) const;
-	FRDGTextureRef GetSobolSampling(FRDGBuilder& GraphBuilder) const;
-	FRDGTextureRef GetSSAORandomization(FRDGBuilder& GraphBuilder) const;
-	FRDGTextureRef GetPreintegratedGF(FRDGBuilder& GraphBuilder) const;
-	FRDGTextureRef GetLTCMat(FRDGBuilder& GraphBuilder) const;
-	FRDGTextureRef GetLTCAmp(FRDGBuilder& GraphBuilder) const;
 	FRDGTextureRef GetMaxFP16Depth(FRDGBuilder& GraphBuilder) const;
 	FRDGTextureRef GetDepthDummy(FRDGBuilder& GraphBuilder) const;
 	FRDGTextureRef GetStencilDummy(FRDGBuilder& GraphBuilder) const;
