Added RDG system texture creation to hit proxy builder path.

#rnx

diff --git a/Engine/Source/Runtime/Renderer/Private/SceneHitProxyRendering.cpp b/Engine/Source/Runtime/Renderer/Private/SceneHitProxyRendering.cpp
--- a/Engine/Source/Runtime/Renderer/Private/SceneHitProxyRendering.cpp
+++ b/Engine/Source/Runtime/Renderer/Private/SceneHitProxyRendering.cpp
@@ -226,6 +226,7 @@ void InitHitProxyRender(FRDGBuilder& GraphBuilder, const FSceneRenderer* SceneRe
 
 	// Initialize global system textures (pass-through if already initialized).
 	GSystemTextures.InitializeTextures(GraphBuilder.RHICmdList, FeatureLevel);
+	FRDGSystemTextures::Create(GraphBuilder);
 
 	FSceneRenderTargets& SceneContext = FSceneRenderTargets::Get();
 	// Allocate the maximum scene render target space for the current view family.
