Virtual Texture feedback should be in UAVGraphics during rendering because it can be access by raster

diff --git a/Engine/Source/Runtime/Renderer/Private/DeferredShadingRenderer.cpp b/Engine/Source/Runtime/Renderer/Private/DeferredShadingRenderer.cpp
--- a/Engine/Source/Runtime/Renderer/Private/DeferredShadingRenderer.cpp
+++ b/Engine/Source/Runtime/Renderer/Private/DeferredShadingRenderer.cpp
@@ -1662,7 +1662,7 @@ void FDeferredShadingSceneRenderer::Render(FRHICommandListImmediate& RHICmdList)
 				FUnorderedAccessViewRHIRef FeedbackUAV = SceneContext.GetVirtualTextureFeedbackUAV();
 				InRHICmdList.Transition(FRHITransitionInfo(FeedbackUAV, ERHIAccess::Unknown, ERHIAccess::UAVCompute));
 				InRHICmdList.ClearUAVUint(FeedbackUAV, FUintVector4(~0u, ~0u, ~0u, ~0u));
-				InRHICmdList.Transition(FRHITransitionInfo(FeedbackUAV, ERHIAccess::UAVCompute, ERHIAccess::UAVCompute));
+				InRHICmdList.Transition(FRHITransitionInfo(FeedbackUAV, ERHIAccess::UAVCompute, ERHIAccess::UAVGraphics));
 				InRHICmdList.BeginUAVOverlap(FeedbackUAV);
 			});
 		}
@@ -2746,7 +2746,7 @@ void FDeferredShadingSceneRenderer::Render(FRHICommandListImmediate& RHICmdList)
 			InRHICmdList.EndUAVOverlap(SceneContext.VirtualTextureFeedbackUAV);
 
 			// No pass after this should make VT page requests
-			InRHICmdList.Transition(FRHITransitionInfo(SceneContext.VirtualTextureFeedbackUAV, ERHIAccess::UAVCompute, ERHIAccess::CopySrc));
+			InRHICmdList.Transition(FRHITransitionInfo(SceneContext.VirtualTextureFeedbackUAV, ERHIAccess::UAVGraphics, ERHIAccess::CopySrc));
 
 			TArray<FIntRect, TInlineAllocator<4>> ViewRects;
 			ViewRects.AddUninitialized(Views.Num());
