Reprojected normal for DBuffer decals.

diff --git a/Engine/Shaders/Private/DBufferNormalReprojection.ush b/Engine/Shaders/Private/DBufferNormalReprojection.ush
new file mode 100644
--- /dev/null
+++ b/Engine/Shaders/Private/DBufferNormalReprojection.ush
@@ -0,0 +1,126 @@
+// Copyright Epic Games, Inc. All Rights Reserved.
+
+/**
+ * DBufferNormalReprojection.usf: Utilities for reprojecting the previous frame's normal when read by the scene texture node in a DBuffer decal.
+ */
+
+#pragma once
+
+#if IS_DBUFFER_DECAL && FEATURE_LEVEL >= FEATURE_LEVEL_SM5
+
+/** Samples the screen-space velocity for the specified UV coordinates. 
+ *  this one is slightly different from PostProcessVelocityLookup because it
+ *  samples from SceneTexturesStruct_GBufferVelocityTextureSampler, and ignores
+ *  the GBUFFER_HAS_VELOCITY #ifdef
+ **/
+float2 DeferredDecalPostProcessVelocityLookup(float Depth, float2 UV, float4 EncodedVelocity)
+{
+	float2 Velocity;
+	if( EncodedVelocity.x > 0.0 )
+	{
+		Velocity = DecodeVelocityFromTexture(EncodedVelocity).xy;
+	}
+	else
+	{
+		float4 ThisClip = float4( UV, Depth, 1 );
+		float4 PrevClip = mul( ThisClip, View.ClipToPrevClip );
+		float2 PrevScreen = PrevClip.xy / PrevClip.w;
+		Velocity = UV - PrevScreen;
+	}
+
+	return Velocity;
+}
+
+// The functions TakeSmallerAbsDelta and ReconstructNormalFromDepthBuffer were copy/modified from PostProcessAmbientOcclusion.usf.
+float TakeSmallerAbsDelta(float left, float mid, float right)
+{
+	float a = mid - left;
+	float b = right - mid;
+
+	return (abs(a) < abs(b)) ? a : b;
+}
+
+// could use ddx,ddy but that would have less quality and would nto work fo ComputeShaders
+// @return not normalized normal in world space
+float3 ReconstructNormalFromDepthBuffer(int2 BaseCoord)
+{
+	// could use a modified version of GatherSceneDepth later on
+	float DeviceZ = SceneTexturesStruct.SceneDepthTexture.Load(int3(BaseCoord + int2(0, 0),0)).x;
+	float DeviceZLeft = SceneTexturesStruct.SceneDepthTexture.Load(int3(BaseCoord + int2(-1, 0), 0)).x;
+	float DeviceZTop = SceneTexturesStruct.SceneDepthTexture.Load(int3(BaseCoord + int2(0, -1), 0)).x;
+	float DeviceZRight = SceneTexturesStruct.SceneDepthTexture.Load(int3(BaseCoord + int2(1, 0), 0)).x;
+	float DeviceZBottom = SceneTexturesStruct.SceneDepthTexture.Load(int3(BaseCoord + int2(0, 1), 0)).x;
+
+	// Favor the surfae we are looking at. Simiar to: http://www.humus.name/index.php?page=3D&ID=84
+	float DeviceZDdx = TakeSmallerAbsDelta(DeviceZLeft, DeviceZ, DeviceZRight);
+	float DeviceZDdy = TakeSmallerAbsDelta(DeviceZTop, DeviceZ, DeviceZBottom);
+
+	float2 SvPosition = float2(BaseCoord) + float2(.5f, .5f);
+
+	// can be optimized, is not fully centered but that should not matter much
+	float3 Mid =	SvPositionToTranslatedWorld(float4(SvPosition.xy + float2(0, 0), DeviceZ, 1));
+	float3 Right =	SvPositionToTranslatedWorld(float4(SvPosition.xy + float2(1, 0), DeviceZ + DeviceZDdx, 1)) - Mid;
+	float3 Down =	SvPositionToTranslatedWorld(float4(SvPosition.xy + float2(0, 1), DeviceZ + DeviceZDdy, 1)) - Mid;
+
+	return normalize(cross(Right, Down));
+}
+
+float4 GetDBufferReprojectedWorldNormal(float2 UV, float3 FallbackNormal)
+{
+	int2 BaseCoord = (int2)trunc(UV * View.BufferSizeAndInvSize.xy);
+
+	float DeviceZ = SceneTexturesStruct.SceneDepthTexture.Load(int3(BaseCoord,0)).x;
+
+	// Velocity has to be the same size as the rest of the targets
+	float4 EncodedVelocity = SceneTexturesStruct.GBufferVelocityTexture.Load(int3(BaseCoord,0));
+	float3 FacetedNorm = ReconstructNormalFromDepthBuffer(BaseCoord);
+
+	float3 N = FacetedNorm;
+
+	if (DeferredDecal.NormalReprojectionEnabled)
+	{
+		// convert from [0,1] to [-1,1], and flip Y.
+		float2 ViewportUV = (UV*2.0f - 1.0f) * float2(1.0f,-1.0f);
+
+		float2 Velocity = DeferredDecalPostProcessVelocityLookup(DeviceZ, ViewportUV, EncodedVelocity);
+
+		float4 ThisClip = float4(ViewportUV, DeviceZ, 1);
+		float4 PrevClip = mul( ThisClip, View.ClipToPrevClip );
+		float2 PrevScreen = PrevClip.xy / PrevClip.w;
+
+		// Velocity is in clip space [-1,1] with flipped Y.
+		// convert from [-1,1] to [0,1], and flip Y back.
+		float2 ReprojUV = ((UV - Velocity*float2(1,-1)*.5f));
+
+		bool bIsEdgeOcclusion = ReprojUV.x < 0 || ReprojUV.y < 0 || ReprojUV.x >= 1.0f || ReprojUV.y >= 1.0f;
+
+		// Note: Doing a linear sample here, as it seems more correct since the normal buffer can blend between texels. Also,
+		// we have the previous frame's jitter (DeferredDecal.NormalReprojectionJitter) available but it seems like we get the
+		// best results by ignoring it. Also, it is possible that we have normal from a bad edge, but in that case it will
+		// deviate too much from the faceted normal and get fixed a few lines later.
+		float3 ReprojNormTex = Texture2DSampleLevel(DeferredDecal.PreviousFrameNormal, GlobalBilinearClampedSampler, ReprojUV, 0.0f).rgb;
+
+		float3 ReprojNorm = normalize(ReprojNormTex * 2.0f - 1.0f);
+
+		float DotFacet = saturate(dot(ReprojNorm,FacetedNorm));
+
+		float Influence = saturate((DotFacet - DeferredDecal.NormalReprojectionThresholdLow) * DeferredDecal.NormalReprojectionThresholdScaleHelper);
+		if (bIsEdgeOcclusion)
+		{
+			Influence = 0.0f;
+		}
+		N = normalize(lerp(FacetedNorm,ReprojNorm,Influence));
+	}
+
+	return float4(N,1.0f);
+}
+
+#else
+
+// for mobile, just return the geometry normal of the surface on top as a fallback
+float4 GetDBufferReprojectedWorldNormal(float2 UV, float3 FallbackNormal)
+{
+	return float4(FallbackNormal.xyz,1.0);
+}
+
+#endif
\ No newline at end of file
diff --git a/Engine/Shaders/Private/MaterialTemplate.ush b/Engine/Shaders/Private/MaterialTemplate.ush
--- a/Engine/Shaders/Private/MaterialTemplate.ush
+++ b/Engine/Shaders/Private/MaterialTemplate.ush
@@ -90,6 +90,7 @@
 #endif
 
 #include "/Engine/Private/PaniniProjection.ush"
+#include "/Engine/Private/DBufferNormalReprojection.ush"
 
 #ifndef USE_DITHERED_LOD_TRANSITION
     #if USE_INSTANCING
diff --git a/Engine/Source/Runtime/Engine/Private/Materials/HLSLMaterialTranslator.cpp b/Engine/Source/Runtime/Engine/Private/Materials/HLSLMaterialTranslator.cpp
--- a/Engine/Source/Runtime/Engine/Private/Materials/HLSLMaterialTranslator.cpp
+++ b/Engine/Source/Runtime/Engine/Private/Materials/HLSLMaterialTranslator.cpp
@@ -4871,6 +4871,54 @@ int32 FHLSLMaterialTranslator::SceneTextureLookup(int32 ViewportUV, uint32 InSce
 									SceneTextureId == PPI_SceneDepth ||
 									SceneTextureId == PPI_CustomStencil;
 
+	// special case for DBuffer normals
+	bool bIsUsingReprojectedNormal = false;
+
+	if(Material->GetMaterialDomain() == MD_DeferredDecal)
+	{
+		EDecalBlendMode DecalBlendMode = (EDecalBlendMode)Material->GetDecalBlendMode();
+		bool bDBuffer = IsDBufferDecalBlendMode(DecalBlendMode);
+
+		if(bDBuffer)
+		{
+			if(SceneTextureId == PPI_WorldNormal)
+			{
+				bIsUsingReprojectedNormal = true;
+			}
+		}
+	}
+
+	// The normal is valid, sample from previous frame with reprojection. Otherwise, return the original geometry normal, both for SM5 when the normal is missing,
+	// but also on Mobile as a fallback.
+	if (bIsUsingReprojectedNormal)
+	{
+		int32 BufferUV;
+		if (ViewportUV != INDEX_NONE)
+		{
+			BufferUV = AddCodeChunk(MCT_Float2,
+				TEXT("ClampSceneTextureUV(ViewportUVToSceneTextureUV(%s, %d), %d)"),
+				*CoerceParameter(ViewportUV, MCT_Float2), (int)SceneTextureId, (int)SceneTextureId);
+		}
+		else
+		{
+			BufferUV = AddInlinedCodeChunk(MCT_Float2, TEXT("GetDefaultSceneTextureUV(Parameters, %d)"), (int)SceneTextureId);
+		}
+
+		if (FeatureLevel >= ERHIFeatureLevel::SM5)
+		{
+			return AddCodeChunk(MCT_Float4,
+				TEXT("GetDBufferReprojectedWorldNormal(%s,%s)"),
+				*CoerceParameter(BufferUV, MCT_Float2),
+				*CoerceParameter(PixelNormalWS(),MCT_Float3));
+		}
+		else
+		{
+			return AddCodeChunk(MCT_Float4,
+				TEXT("float4(%s,1.0f)"),
+				*CoerceParameter(PixelNormalWS(),MCT_Float3));
+		}
+	}
+
 	if (!bSupportedOnMobile	&& ErrorUnlessFeatureLevelSupported(ERHIFeatureLevel::SM5) == INDEX_NONE)
 	{
 		return INDEX_NONE;
@@ -4952,10 +5000,10 @@ void FHLSLMaterialTranslator::UseSceneTextureId(ESceneTextureId SceneTextureId,
 
 		if(bDBuffer)
 		{
-			if(!(SceneTextureId == PPI_SceneDepth || SceneTextureId == PPI_CustomDepth || SceneTextureId == PPI_CustomStencil))
+			if(!(SceneTextureId == PPI_SceneDepth || SceneTextureId == PPI_CustomDepth || SceneTextureId == PPI_CustomStencil || SceneTextureId == PPI_WorldNormal))
 			{
 				// Note: For DBuffer decals: CustomDepth and CustomStencil are only available if r.CustomDepth.Order = 0
-				Errorf(TEXT("DBuffer decals (MaterialDomain=DeferredDecal and DecalBlendMode is using DBuffer) can only access SceneDepth, CustomDepth, CustomStencil"));
+				Errorf(TEXT("DBuffer decals (MaterialDomain=DeferredDecal and DecalBlendMode is using DBuffer) can only access SceneDepth, CustomDepth, CustomStencil, and WorldNormal"));
 			}
 		}
 		else
diff --git a/Engine/Source/Runtime/Renderer/Private/CompositionLighting/PostProcessDeferredDecals.cpp b/Engine/Source/Runtime/Renderer/Private/CompositionLighting/PostProcessDeferredDecals.cpp
--- a/Engine/Source/Runtime/Renderer/Private/CompositionLighting/PostProcessDeferredDecals.cpp
+++ b/Engine/Source/Runtime/Renderer/Private/CompositionLighting/PostProcessDeferredDecals.cpp
@@ -18,6 +18,33 @@ static TAutoConsoleVariable<float> CVarStencilSizeThreshold(
 	TEXT("0..1: optimization is enabled, value defines the minimum size (screen space) to trigger the optimization (default 0.1)")
 );
 
+static TAutoConsoleVariable<float> CVarDBufferDecalNormalReprojectionThresholdLow(
+	TEXT("r.Decal.NormalReprojectionThresholdLow"),
+	0.990f, 
+	TEXT("When reading the normal from a SceneTexture node in a DBuffer decal shader, ")
+	TEXT("the normal is a mix of the geometry normal (extracted from the depth buffer) and the normal from the reprojected ")
+	TEXT("previous frame. When the dot product of the geometry and reprojected normal is below the r.Decal.NormalReprojectionThresholdLow, ")
+	TEXT("the geometry normal is used. When that value is above r.Decal.NormalReprojectionThresholdHigh, the reprojected ")
+	TEXT("normal is used. Otherwise it uses a lerp between them."),
+	ECVF_RenderThreadSafe);
+
+static TAutoConsoleVariable<float> CVarDBufferDecalNormalReprojectionThresholdHigh(
+	TEXT("r.Decal.NormalReprojectionThresholdHigh"),
+	0.995f, 
+	TEXT("When reading the normal from a SceneTexture node in a DBuffer decal shader, ")
+	TEXT("the normal is a mix of the geometry normal (extracted from the depth buffer) and the normal from the reprojected ")
+	TEXT("previous frame. When the dot product of the geometry and reprojected normal is below the r.Decal.NormalReprojectionThresholdLow, ")
+	TEXT("the geometry normal is used. When that value is above r.Decal.NormalReprojectionThresholdHigh, the reprojected ")
+	TEXT("normal is used. Otherwise it uses a lerp between them."),
+	ECVF_RenderThreadSafe);
+
+static TAutoConsoleVariable<bool> CVarDBufferDecalNormalReprojectionEnabled(
+	TEXT("r.Decal.NormalReprojectionEnabled"),
+	true, 
+	TEXT("If true, normal reprojection from the previous frame is allowed in SceneTexture nodes on DBuffer decals, provided that motion ")
+	TEXT("in depth prepass is enabled as well (r.DepthPassMergedWithVelocity). Otherwise the fallback is the normal extracted from the depth buffer."),
+	ECVF_RenderThreadSafe);
+
 FDeferredDecalPassTextures GetDeferredDecalPassTextures(
 	FRDGBuilder& GraphBuilder,
 	const FViewInfo& View,
@@ -134,6 +161,36 @@ void GetDeferredDecalPassParameters(
 		bWritingToDepth ? FExclusiveDepthStencil::DepthWrite_StencilWrite : FExclusiveDepthStencil::DepthRead_StencilWrite);
 }
 
+void GetDeferredDecalUniformParameters(
+	const FViewInfo& View,
+	FDeferredDecalUniformParameters& UniformParameters)
+{
+	// This extern is not pretty, but it's only temporary. The plan is to make velocity in DepthPass the only option, but requires some testing to 
+	// ensure quality and perf is what we expect.
+	extern bool IsVelocityMergedWithDepthPass();
+
+	bool bIsMotionInDepth = IsVelocityMergedWithDepthPass();
+	// if we have early motion vectors (bIsMotionInDepth) and the cvar is enabled and we actually have a buffer from the previous frame (View.PrevViewInfo.GBufferA.IsValid())
+	bool bIsNormalReprojectionEnabled = (bIsMotionInDepth && CVarDBufferDecalNormalReprojectionEnabled.GetValueOnRenderThread() && View.PrevViewInfo.GBufferA.IsValid());
+
+	UniformParameters.NormalReprojectionThresholdLow  = CVarDBufferDecalNormalReprojectionThresholdLow .GetValueOnRenderThread();
+	UniformParameters.NormalReprojectionThresholdHigh = CVarDBufferDecalNormalReprojectionThresholdHigh.GetValueOnRenderThread();
+	UniformParameters.NormalReprojectionEnabled = bIsNormalReprojectionEnabled ? 1 : 0;
+	
+	// the algorithm is:
+	//    value = (dot - low)/(high - low)
+	// so calculate the divide in the helper to turn the math into:
+	//    helper = 1.0f/(high - low)
+	//    value = (dot - low)*helper;
+	// also check for the case where high <= low.
+	float Denom = FMath::Max(UniformParameters.NormalReprojectionThresholdHigh - UniformParameters.NormalReprojectionThresholdLow,1e-4f);
+	UniformParameters.NormalReprojectionThresholdScaleHelper = 1.0f / Denom;
+
+	UniformParameters.PreviousFrameNormal = (bIsNormalReprojectionEnabled) ? View.PrevViewInfo.GBufferA->GetShaderResourceRHI() : GSystemTextures.BlackDummy->GetShaderResourceRHI();
+
+	UniformParameters.NormalReprojectionJitter = View.PrevViewInfo.ViewMatrices.GetTemporalAAJitter();
+}
+
 enum EDecalDepthInputState
 {
 	DDS_Undefined,
@@ -481,14 +538,23 @@ void AddDeferredDecalPass(
 	{
 		auto* PassParameters = GraphBuilder.AllocParameters<FDeferredDecalPassParameters>();
 		GetDeferredDecalPassParameters(View, PassTextures, RenderTargetMode, *PassParameters);
+
+		FDeferredDecalUniformParameters DecalParams;
+		GetDeferredDecalUniformParameters(
+			View,
+			DecalParams);
+
 		GraphBuilder.AddPass(
 			RDG_EVENT_NAME("Batch [%d, %d]", DecalIndexBegin, DecalIndexEnd - 1),
 			PassParameters,
 			ERDGPassFlags::Raster,
-			[&View, FeatureLevel, ShaderPlatform, DecalIndexBegin, DecalIndexEnd, SortedDecals, DecalRenderStage, bStencilSizeThreshold, bShaderComplexity](FRHICommandList& RHICmdList)
+			[&View, FeatureLevel, ShaderPlatform, DecalIndexBegin, DecalIndexEnd, SortedDecals, DecalRenderStage, bStencilSizeThreshold, bShaderComplexity, DecalParams](FRHICommandList& RHICmdList)
 		{
 			RHICmdList.SetViewport(View.ViewRect.Min.X, View.ViewRect.Min.Y, 0.0f, View.ViewRect.Max.X, View.ViewRect.Max.Y, 1.0f);
 
+			const FScene & Scene = *View.Family->Scene->GetRenderScene();
+			const_cast<FScene&>(Scene).UniformBuffers.DeferredDecalUniformBuffer.UpdateUniformBufferImmediate(DecalParams);
+
 			for (uint32 DecalIndex = DecalIndexBegin; DecalIndex < DecalIndexEnd; ++DecalIndex)
 			{
 				const FTransientDecalRenderData& DecalData = (*SortedDecals)[DecalIndex];
diff --git a/Engine/Source/Runtime/Renderer/Private/CompositionLighting/PostProcessDeferredDecals.h b/Engine/Source/Runtime/Renderer/Private/CompositionLighting/PostProcessDeferredDecals.h
--- a/Engine/Source/Runtime/Renderer/Private/CompositionLighting/PostProcessDeferredDecals.h
+++ b/Engine/Source/Runtime/Renderer/Private/CompositionLighting/PostProcessDeferredDecals.h
@@ -54,16 +54,30 @@ void AddDeferredDecalPass(
 	EDecalRenderStage RenderStage);
 
 BEGIN_SHADER_PARAMETER_STRUCT(FDeferredDecalPassParameters, )
+	SHADER_PARAMETER_STRUCT_REF(FViewUniformShaderParameters, View)
 	SHADER_PARAMETER_RDG_UNIFORM_BUFFER(FSceneTextureUniformParameters, SceneTextures)
 	RENDER_TARGET_BINDING_SLOTS()
 END_SHADER_PARAMETER_STRUCT()
 
+BEGIN_GLOBAL_SHADER_PARAMETER_STRUCT(FDeferredDecalUniformParameters, )
+	SHADER_PARAMETER_TEXTURE(Texture2D, PreviousFrameNormal)
+	SHADER_PARAMETER(int32, NormalReprojectionEnabled)
+	SHADER_PARAMETER(float, NormalReprojectionThresholdLow)
+	SHADER_PARAMETER(float, NormalReprojectionThresholdHigh)
+	SHADER_PARAMETER(float, NormalReprojectionThresholdScaleHelper)
+	SHADER_PARAMETER(FVector2D, NormalReprojectionJitter)
+END_SHADER_PARAMETER_STRUCT()
+
 void GetDeferredDecalPassParameters(
 	const FViewInfo& View,
 	FDeferredDecalPassTextures& DecalPassTextures,
 	FDecalRenderingCommon::ERenderTargetMode RenderTargetMode,
 	FDeferredDecalPassParameters& PassParameters);
 
+void GetDeferredDecalUniformParameters(
+	const FViewInfo& View,
+	FDeferredDecalUniformParameters& UniformParameters);
+
 void RenderMeshDecals(
 	FRDGBuilder& GraphBuilder,
 	const FViewInfo& View,
diff --git a/Engine/Source/Runtime/Renderer/Private/CompositionLighting/PostProcessMeshDecals.cpp b/Engine/Source/Runtime/Renderer/Private/CompositionLighting/PostProcessMeshDecals.cpp
--- a/Engine/Source/Runtime/Renderer/Private/CompositionLighting/PostProcessMeshDecals.cpp
+++ b/Engine/Source/Runtime/Renderer/Private/CompositionLighting/PostProcessMeshDecals.cpp
@@ -100,6 +100,16 @@ public:
 		FDecalRendering::SetDecalCompilationEnvironment(Parameters, OutEnvironment);
 	}
 
+	void GetShaderBindings(
+		const FScene* Scene,
+		ERHIFeatureLevel::Type FeatureLevel,
+		const FPrimitiveSceneProxy* PrimitiveSceneProxy,
+		const FMaterialRenderProxy& MaterialRenderProxy,
+		const FMaterial& Material,
+		const FMeshPassProcessorRenderState& DrawRenderState,
+		const FMeshMaterialShaderElementData& ShaderElementData,
+		FMeshDrawSingleShaderBindings& ShaderBindings) const;
+
 	FMeshDecalsPS() = default;
 	FMeshDecalsPS(const ShaderMetaType::CompiledShaderInitializerType& Initializer)
 		: FMeshMaterialShader(Initializer)
@@ -108,6 +118,21 @@ public:
 
 IMPLEMENT_MATERIAL_SHADER_TYPE(,FMeshDecalsPS,TEXT("/Engine/Private/MeshDecals.usf"),TEXT("MainPS"),SF_Pixel);
 
+void FMeshDecalsPS::GetShaderBindings(
+	const FScene* Scene,
+	ERHIFeatureLevel::Type FeatureLevel,
+	const FPrimitiveSceneProxy* PrimitiveSceneProxy,
+	const FMaterialRenderProxy& MaterialRenderProxy,
+	const FMaterial& Material,
+	const FMeshPassProcessorRenderState& DrawRenderState,
+	const FMeshMaterialShaderElementData& ShaderElementData,
+	FMeshDrawSingleShaderBindings& ShaderBindings) const
+{
+	FMeshMaterialShader::GetShaderBindings(Scene, FeatureLevel, PrimitiveSceneProxy, MaterialRenderProxy, Material, DrawRenderState, ShaderElementData, ShaderBindings);
+
+	ShaderBindings.Add(GetUniformBufferParameter<FDeferredDecalUniformParameters>(), DrawRenderState.GetDeferredDecalUniformBuffer());
+}
+
 class FMeshDecalsEmissivePS : public FMeshDecalsPS
 {
 public:
@@ -162,6 +187,8 @@ private:
 	const FDecalRenderingCommon::ERenderTargetMode RenderTargetMode;
 };
 
+IMPLEMENT_GLOBAL_SHADER_PARAMETER_STRUCT(FDeferredDecalUniformParameters, "DeferredDecal");
+
 
 FMeshDecalMeshProcessor::FMeshDecalMeshProcessor(const FScene* Scene, 
 	const FSceneView* InViewIfDynamicMeshCommand, 
@@ -175,6 +202,7 @@ FMeshDecalMeshProcessor::FMeshDecalMeshProcessor(const FScene* Scene,
 	PassDrawRenderState.SetDepthStencilState(TStaticDepthStencilState<false, CF_DepthNearOrEqual>::GetRHI());
 	PassDrawRenderState.SetViewUniformBuffer(Scene->UniformBuffers.ViewUniformBuffer);
 	PassDrawRenderState.SetInstancedViewUniformBuffer(Scene->UniformBuffers.InstancedViewUniformBuffer);
+	PassDrawRenderState.SetDeferredDecalUniformBuffer(Scene->UniformBuffers.DeferredDecalUniformBuffer);
 	if (FeatureLevel == ERHIFeatureLevel::ES3_1)
 	{
 		PassDrawRenderState.SetPassUniformBuffer(Scene->UniformBuffers.MobileTranslucentBasePassUniformBuffer);
@@ -349,16 +377,25 @@ void DrawDecalMeshCommands(
 {
 	auto* PassParameters = GraphBuilder.AllocParameters<FDeferredDecalPassParameters>();
 	GetDeferredDecalPassParameters(View, DecalPassTextures, RenderTargetMode, *PassParameters);
+
+	FDeferredDecalUniformParameters DecalParams;
+	GetDeferredDecalUniformParameters(
+		View,
+		DecalParams);
+
 	GraphBuilder.AddPass(
 		RDG_EVENT_NAME("MeshDecals"),
 		PassParameters,
 		ERDGPassFlags::Raster,
-		[&View, DecalRenderStage, RenderTargetMode](FRHICommandList& RHICmdList)
+		[&View, DecalRenderStage, RenderTargetMode, DecalParams](FRHICommandList& RHICmdList)
 	{
 		RHICmdList.SetViewport(View.ViewRect.Min.X, View.ViewRect.Min.Y, 0.0f, View.ViewRect.Max.X, View.ViewRect.Max.Y, 1.0f);
 
+		const FScene& Scene = *View.Family->Scene->GetRenderScene();
+		const_cast<FScene&>(Scene).UniformBuffers.DeferredDecalUniformBuffer.UpdateUniformBufferImmediate(DecalParams);
+
 		DrawDynamicMeshPass(View, RHICmdList,
-			[&View, DecalRenderStage, RenderTargetMode](FDynamicPassMeshDrawListContext* DynamicMeshPassContext)
+			[&View, DecalRenderStage, RenderTargetMode, DecalParams](FDynamicPassMeshDrawListContext* DynamicMeshPassContext)
 		{
 			FMeshDecalMeshProcessor PassMeshProcessor(
 				View.Family->Scene->GetRenderScene(),
diff --git a/Engine/Source/Runtime/Renderer/Private/DecalRenderingShared.cpp b/Engine/Source/Runtime/Renderer/Private/DecalRenderingShared.cpp
--- a/Engine/Source/Runtime/Renderer/Private/DecalRenderingShared.cpp
+++ b/Engine/Source/Runtime/Renderer/Private/DecalRenderingShared.cpp
@@ -459,6 +459,11 @@ void FDecalRendering::SetDecalCompilationEnvironment(const FMaterialShaderPermut
 	OutEnvironment.SetDefine(TEXT("MATERIAL_DBUFFERC"), (bDBufferMask & 0x4) != 0);
 	// define so that we know to output the proper MRTs in the derived defines.
 	OutEnvironment.SetDefine(TEXT("IS_DECAL"), 1);
+
+	if (IsDBufferDecalBlendMode(FinalDecalBlendMode))
+	{
+		OutEnvironment.SetDefine(TEXT("IS_DBUFFER_DECAL"), 1);
+	}
 }
 
 void FDecalRendering::SetEmissiveDBufferDecalCompilationEnvironment(const FMaterialShaderPermutationParameters& Parameters, FShaderCompilerEnvironment& OutEnvironment)
diff --git a/Engine/Source/Runtime/Renderer/Private/DeferredShadingRenderer.cpp b/Engine/Source/Runtime/Renderer/Private/DeferredShadingRenderer.cpp
--- a/Engine/Source/Runtime/Renderer/Private/DeferredShadingRenderer.cpp
+++ b/Engine/Source/Runtime/Renderer/Private/DeferredShadingRenderer.cpp
@@ -333,6 +333,9 @@ static FORCEINLINE bool NeedsPrePass(const FDeferredShadingSceneRenderer* Render
 		(Renderer->DepthPass.EarlyZPassMode != DDM_None || Renderer->DepthPass.bEarlyZPassMovable != 0);
 }
 
+// This is a temporary extern while testing the Velocity changes. Will be removed once a decision is made.
+extern bool IsVelocityMergedWithDepthPass();
+
 bool FDeferredShadingSceneRenderer::RenderHzb(
 	FRDGBuilder& GraphBuilder,
 	TRDGUniformBufferRef<FSceneTextureUniformParameters> SceneTexturesUniformBuffer)
@@ -2253,6 +2256,27 @@ void FDeferredShadingSceneRenderer::Render(FRHICommandListImmediate& RHICmdList)
 	SceneTexturesSetupMode |= ESceneTextureSetupMode::GBuffers;
 	SceneTextures = CreateSceneTextureUniformBuffer(GraphBuilder, FeatureLevel, SceneTexturesSetupMode);
 
+	// mark GBufferA for saving for next frame
+	{
+		static auto CVarNormalReprojectionEnabled = IConsoleManager::Get().FindConsoleVariable(TEXT("r.Decal.NormalReprojectionEnabled"));
+
+		// save the previous frame if early motion vectors are enabled and normal reprojection is enabled, so there should be no cost if these options are off
+		bool bApplyReproject = IsVelocityMergedWithDepthPass() && CVarNormalReprojectionEnabled && CVarNormalReprojectionEnabled->GetInt() > 0;
+
+		if (bApplyReproject)
+		{
+			bool bAnyViewWriteable = false;
+			for (int32 Index = 0; Index < Views.Num(); Index++)
+			{
+				if (Views[Index].bStatePrevViewInfoIsReadOnly == false)
+				{
+					GraphBuilder.QueueTextureExtraction(GraphBuilder.RegisterExternalTexture(SceneContext.GBufferA),&Views[Index].ViewState->PrevFrameViewInfo.GBufferA);
+				}
+			}
+		
+		}
+	}
+
 	if (ViewFamily.EngineShowFlags.VisualizeLightCulling)
 	{
 		AddClearRenderTargetPass(GraphBuilder, SceneColorTexture.Target, FLinearColor::Transparent);
diff --git a/Engine/Source/Runtime/Renderer/Private/RendererScene.cpp b/Engine/Source/Runtime/Renderer/Private/RendererScene.cpp
--- a/Engine/Source/Runtime/Renderer/Private/RendererScene.cpp
+++ b/Engine/Source/Runtime/Renderer/Private/RendererScene.cpp
@@ -926,6 +926,9 @@ void FPersistentUniformBuffers::Initialize()
 	FNaniteUniformParameters NaniteUniformBufferParameters;
 	NaniteUniformBuffer = TUniformBufferRef<FNaniteUniformParameters>::CreateUniformBufferImmediate(NaniteUniformBufferParameters, UniformBuffer_MultiFrame, EUniformBufferValidation::None);
 
+	FDeferredDecalUniformParameters DeferredDecalUniformBufferParameters;
+	DeferredDecalUniformBuffer = TUniformBufferRef<FDeferredDecalUniformParameters>::CreateUniformBufferImmediate(DeferredDecalUniformBufferParameters, UniformBuffer_MultiFrame, EUniformBufferValidation::None);
+
 	FReflectionCaptureShaderData ReflectionCaptureParameters;
 	ReflectionCaptureUniformBuffer = TUniformBufferRef<FReflectionCaptureShaderData>::CreateUniformBufferImmediate(ReflectionCaptureParameters, UniformBuffer_MultiFrame);
 
diff --git a/Engine/Source/Runtime/Renderer/Private/ScenePrivate.h b/Engine/Source/Runtime/Renderer/Private/ScenePrivate.h
--- a/Engine/Source/Runtime/Renderer/Private/ScenePrivate.h
+++ b/Engine/Source/Runtime/Renderer/Private/ScenePrivate.h
@@ -56,6 +56,7 @@
 #include "Lumen/LumenViewState.h"
 #include "VolumetricRenderTargetViewStateData.h"
 #include "GPUScene.h"
+#include "CompositionLighting/PostProcessDeferredDecals.h"
 
 /** Factor by which to grow occlusion tests **/
 #define OCCLUSION_SLOP (1.0f)
@@ -2404,6 +2405,7 @@ public:
 	TUniformBufferRef<FViewUniformShaderParameters> ViewUniformBuffer;
 	TUniformBufferRef<FInstancedViewUniformShaderParameters> InstancedViewUniformBuffer;
 	TUniformBufferRef<FNaniteUniformParameters> NaniteUniformBuffer;
+	TUniformBufferRef<FDeferredDecalUniformParameters> DeferredDecalUniformBuffer;
 	TUniformBufferRef<FReflectionCaptureShaderData> ReflectionCaptureUniformBuffer;
 	TUniformBufferRef<FViewUniformShaderParameters> CSMShadowDepthViewUniformBuffer;
 	TUniformBufferRef<FViewUniformShaderParameters> VoxelizeVolumeViewUniformBuffer;
diff --git a/Engine/Source/Runtime/Renderer/Public/MeshPassProcessor.h b/Engine/Source/Runtime/Renderer/Public/MeshPassProcessor.h
--- a/Engine/Source/Runtime/Renderer/Public/MeshPassProcessor.h
+++ b/Engine/Source/Runtime/Renderer/Public/MeshPassProcessor.h
@@ -1591,6 +1591,16 @@ public:
 		return NaniteUniformBuffer;
 	}
 
+	FORCEINLINE_DEBUGGABLE void SetDeferredDecalUniformBuffer(FRHIUniformBuffer* InDeferredDecalUniformBuffer)
+	{
+		DeferredDecalUniformBuffer = InDeferredDecalUniformBuffer;
+	}
+
+	FORCEINLINE_DEBUGGABLE FRHIUniformBuffer* GetDeferredDecalUniformBuffer() const
+	{
+		return DeferredDecalUniformBuffer;
+	}
+
 	FORCEINLINE_DEBUGGABLE uint32 GetStencilRef() const
 	{
 		return StencilRef;
@@ -1615,6 +1625,7 @@ private:
 
 	FRHIUniformBuffer*				PassUniformBuffer = nullptr;
 	FRHIUniformBuffer*				NaniteUniformBuffer = nullptr;
+	FRHIUniformBuffer*				DeferredDecalUniformBuffer = nullptr;
 
 	uint32							StencilRef = 0;
 };
