Added FMinimalSceneTextures and FSceneTextures to deferred shading renderer and marshaled it through the render pipeline.

#rnx

diff --git a/Engine/Source/Runtime/RenderCore/Public/RenderGraphUtils.h b/Engine/Source/Runtime/RenderCore/Public/RenderGraphUtils.h
--- a/Engine/Source/Runtime/RenderCore/Public/RenderGraphUtils.h
+++ b/Engine/Source/Runtime/RenderCore/Public/RenderGraphUtils.h
@@ -159,7 +159,7 @@ struct FRDGTextureMSAA
 		, Resolve(InResolve)
 	{}
 
-	explicit FRDGTextureMSAA(FRDGTextureRef InTexture)
+	FRDGTextureMSAA(FRDGTextureRef InTexture)
 		: Target(InTexture)
 		, Resolve(InTexture)
 	{}
diff --git a/Engine/Source/Runtime/Renderer/Private/AtmosphereRendering.cpp b/Engine/Source/Runtime/Renderer/Private/AtmosphereRendering.cpp
--- a/Engine/Source/Runtime/Renderer/Private/AtmosphereRendering.cpp
+++ b/Engine/Source/Runtime/Renderer/Private/AtmosphereRendering.cpp
@@ -472,10 +472,8 @@ END_SHADER_PARAMETER_STRUCT()
 
 void FDeferredShadingSceneRenderer::RenderAtmosphere(
 	FRDGBuilder& GraphBuilder,
-	FRDGTextureRef SceneColorTexture,
-	FRDGTextureRef SceneDepthTexture,
-	FRDGTextureRef LightShaftOcclusionTexture,
-	TRDGUniformBufferRef<FSceneTextureUniformParameters> SceneTextures)
+	const FMinimalSceneTextures& SceneTextures,
+	FRDGTextureRef LightShaftOcclusionTexture)
 {
 	FAtmosphericFogSceneInfo* AtmosphericFog = Scene->AtmosphericFog;
 	if (!AtmosphericFog)
@@ -505,10 +503,10 @@ void FDeferredShadingSceneRenderer::RenderAtmosphere(
 	RDG_EVENT_SCOPE(GraphBuilder, "AtmosphericFog");
 
 	FAtmospherePassParameters* PassParameters = GraphBuilder.AllocParameters<FAtmospherePassParameters>();
-	PassParameters->SceneTextures = SceneTextures;
+	PassParameters->SceneTextures = SceneTextures.UniformBuffer;
 	PassParameters->LightShaftOcclusionTexture = LightShaftOcclusionTexture;
-	PassParameters->RenderTargets[0] = FRenderTargetBinding(SceneColorTexture, ERenderTargetLoadAction::ELoad);
-	PassParameters->RenderTargets.DepthStencil = FDepthStencilBinding(SceneDepthTexture, ERenderTargetLoadAction::ELoad, ERenderTargetLoadAction::ELoad, FExclusiveDepthStencil::DepthRead_StencilWrite);
+	PassParameters->RenderTargets[0] = FRenderTargetBinding(SceneTextures.Color.Target, ERenderTargetLoadAction::ELoad);
+	PassParameters->RenderTargets.DepthStencil = FDepthStencilBinding(SceneTextures.Depth.Target, ERenderTargetLoadAction::ELoad, ERenderTargetLoadAction::ELoad, FExclusiveDepthStencil::DepthRead_StencilWrite);
 
 	for(int32 ViewIndex = 0; ViewIndex < Views.Num(); ++ViewIndex)
 	{
diff --git a/Engine/Source/Runtime/Renderer/Private/BasePassRendering.cpp b/Engine/Source/Runtime/Renderer/Private/BasePassRendering.cpp
--- a/Engine/Source/Runtime/Renderer/Private/BasePassRendering.cpp
+++ b/Engine/Source/Runtime/Renderer/Private/BasePassRendering.cpp
@@ -773,10 +773,8 @@ END_SHADER_PARAMETER_STRUCT()
 
 void FDeferredShadingSceneRenderer::RenderBasePass(
 	FRDGBuilder& GraphBuilder,
+	const FSceneTextures& SceneTextures,
 	FExclusiveDepthStencil::Type BasePassDepthStencilAccess,
-	FRDGTextureRef SceneColorTexture,
-	FRDGTextureRef SceneDepthTexture,
-	ERenderTargetLoadAction SceneDepthLoadAction,
 	FRDGTextureRef ForwardShadowMaskTexture)
 {
 	const bool bEnableParallelBasePasses = GRHICommandList.UseParallelAlgorithms() && CVarParallelBasePass.GetValueOnRenderThread();
@@ -842,7 +840,7 @@ void FDeferredShadingSceneRenderer::RenderBasePass(
 	TStaticArray<FRDGTextureRef, MaxSimultaneousRenderTargets> BasePassTextures;
 	uint32 BasePassTextureCount = SceneContext.GetGBufferRenderTargets(GraphBuilder, BasePassTextures);
 	TArrayView<FRDGTextureRef> BasePassTexturesView = MakeArrayView(BasePassTextures.GetData(), BasePassTextureCount);
-	FRDGTextureRef BasePassDepthTexture = SceneDepthTexture;
+	FRDGTextureRef BasePassDepthTexture = SceneTextures.Depth.Target;
 	FLinearColor SceneColorClearValue;
 
 	if (bRequiresRHIClear)
@@ -863,7 +861,7 @@ void FDeferredShadingSceneRenderer::RenderBasePass(
 
 		ERenderTargetLoadAction ColorLoadAction = ERenderTargetLoadAction::ELoad;
 
-		if (SceneColorTexture->Desc.ClearValue.GetClearColor() == SceneColorClearValue)
+		if (SceneTextures.Color.Target->Desc.ClearValue.GetClearColor() == SceneColorClearValue)
 		{
 			ColorLoadAction = ERenderTargetLoadAction::EClear;
 		}
@@ -875,11 +873,6 @@ void FDeferredShadingSceneRenderer::RenderBasePass(
 		auto* PassParameters = GraphBuilder.AllocParameters<FRenderTargetParameters>();
 		PassParameters->RenderTargets = GetRenderTargetBindings(ColorLoadAction, BasePassTexturesView);
 
-		if (SceneDepthLoadAction == ERenderTargetLoadAction::EClear)
-		{
-			PassParameters->RenderTargets.DepthStencil = FDepthStencilBinding(BasePassDepthTexture, SceneDepthLoadAction, SceneDepthLoadAction, ExclusiveDepthStencil);
-		}
-
 		GraphBuilder.AddPass(RDG_EVENT_NAME("GBufferClear"), PassParameters, ERDGPassFlags::Raster,
 			[PassParameters, ColorLoadAction, SceneColorClearValue](FRHICommandList& RHICmdList)
 		{
@@ -912,7 +905,7 @@ void FDeferredShadingSceneRenderer::RenderBasePass(
 		{
 			// We only render this warning text when bRequiresRHIClear==true to make sure the scene color buffer is allocated at this stage.
 			// When false, the option specifies that all pixels must be written to by a sky dome anyway.
-			RenderSkyAtmosphereEditorNotifications(GraphBuilder, SceneColorTexture);
+			RenderSkyAtmosphereEditorNotifications(GraphBuilder, SceneTextures.Color.Target);
 		}
 	}
 
@@ -976,7 +969,7 @@ void FDeferredShadingSceneRenderer::RenderBasePass(
 	if (ShouldRenderAnisotropyPass())
 	{
 		AddSetCurrentStatPass(GraphBuilder, GET_STATID(STAT_CLM_AnisotropyPass));
-		RenderAnisotropyPass(GraphBuilder, SceneDepthTexture, bEnableParallelBasePasses);
+		RenderAnisotropyPass(GraphBuilder, SceneTextures.Depth.Target, bEnableParallelBasePasses);
 		AddSetCurrentStatPass(GraphBuilder, GET_STATID(STAT_CLM_AfterAnisotropyPass));
 	}
 }
diff --git a/Engine/Source/Runtime/Renderer/Private/ClusteredDeferredShadingPass.cpp b/Engine/Source/Runtime/Renderer/Private/ClusteredDeferredShadingPass.cpp
--- a/Engine/Source/Runtime/Renderer/Private/ClusteredDeferredShadingPass.cpp
+++ b/Engine/Source/Runtime/Renderer/Private/ClusteredDeferredShadingPass.cpp
@@ -71,10 +71,10 @@ class FClusteredShadingPS : public FGlobalShader
 		SHADER_PARAMETER_STRUCT_REF(FViewUniformShaderParameters, View)
 		SHADER_PARAMETER_RDG_UNIFORM_BUFFER(FSceneTextureUniformParameters, SceneTextures)
 
-		SHADER_PARAMETER_RDG_TEXTURE(Texture2D, LTCMatTexture)
+		SHADER_PARAMETER_TEXTURE(Texture2D, LTCMatTexture)
 		SHADER_PARAMETER_SAMPLER(SamplerState, LTCMatSampler)
 
-		SHADER_PARAMETER_RDG_TEXTURE(Texture2D, LTCAmpTexture)
+		SHADER_PARAMETER_TEXTURE(Texture2D, LTCAmpTexture)
 		SHADER_PARAMETER_SAMPLER(SamplerState, LTCAmpSampler)
 
 		SHADER_PARAMETER_TEXTURE(Texture2D, SSProfilesTexture)
@@ -100,8 +100,7 @@ IMPLEMENT_GLOBAL_SHADER(FClusteredShadingPS, "/Engine/Private/ClusteredDeferredS
 
 void FDeferredShadingSceneRenderer::AddClusteredDeferredShadingPass(
 	FRDGBuilder& GraphBuilder,
-	FRDGTextureRef SceneColorTexture,
-	TRDGUniformBufferRef<FSceneTextureUniformParameters> SceneTexturesUniformBuffer,
+	const FMinimalSceneTextures& SceneTextures,
 	const FSortedLightSetSceneInfo &SortedLightsSet)
 {
 	check(GUseClusteredDeferredShading);
@@ -113,7 +112,7 @@ void FDeferredShadingSceneRenderer::AddClusteredDeferredShadingPass(
 		RDG_GPU_STAT_SCOPE(GraphBuilder, ClusteredShading);
 		RDG_EVENT_SCOPE(GraphBuilder, "ClusteredShading");
 
-		const FIntPoint SceneTextureExtent = SceneColorTexture->Desc.Extent;
+		const FIntPoint SceneTextureExtent = SceneTextures.Extent;
 
 		for (int32 ViewIndex = 0, Num = Views.Num(); ViewIndex < Num; ViewIndex++)
 		{
@@ -123,16 +122,16 @@ void FDeferredShadingSceneRenderer::AddClusteredDeferredShadingPass(
 
 			PassParameters->View = View.ViewUniformBuffer;
 			PassParameters->Forward = View.ForwardLightingResources->ForwardLightDataUniformBuffer;
-			PassParameters->SceneTextures = SceneTexturesUniformBuffer;
+			PassParameters->SceneTextures = SceneTextures.UniformBuffer;
 
-			PassParameters->LTCMatTexture = GraphBuilder.RegisterExternalTexture(GSystemTextures.LTCMat);
+			PassParameters->LTCMatTexture = GSystemTextures.LTCMat->GetShaderResourceRHI();
 			PassParameters->LTCMatSampler = TStaticSamplerState<SF_Bilinear, AM_Clamp, AM_Clamp, AM_Clamp>::GetRHI();
-			PassParameters->LTCAmpTexture = GraphBuilder.RegisterExternalTexture(GSystemTextures.LTCAmp);
+			PassParameters->LTCAmpTexture = GSystemTextures.LTCAmp->GetShaderResourceRHI();
 			PassParameters->LTCAmpSampler = TStaticSamplerState<SF_Bilinear, AM_Clamp, AM_Clamp, AM_Clamp>::GetRHI();
 			PassParameters->SSProfilesTexture = GetSubsufaceProfileTexture_RT(GraphBuilder.RHICmdList)->GetShaderResourceRHI();
 			PassParameters->TransmissionProfilesLinearSampler = TStaticSamplerState<SF_Bilinear, AM_Clamp, AM_Clamp, AM_Clamp>::GetRHI();
 
-			PassParameters->RenderTargets[0] = FRenderTargetBinding(SceneColorTexture, ERenderTargetLoadAction::ELoad);
+			PassParameters->RenderTargets[0] = FRenderTargetBinding(SceneTextures.Color.Target, ERenderTargetLoadAction::ELoad);
 
 			GraphBuilder.AddPass(
 				RDG_EVENT_NAME("ClusteredDeferredShading, #Lights: %d", NumLightsToRender),
diff --git a/Engine/Source/Runtime/Renderer/Private/DeferredShadingRenderer.cpp b/Engine/Source/Runtime/Renderer/Private/DeferredShadingRenderer.cpp
--- a/Engine/Source/Runtime/Renderer/Private/DeferredShadingRenderer.cpp
+++ b/Engine/Source/Runtime/Renderer/Private/DeferredShadingRenderer.cpp
@@ -336,14 +336,10 @@ static FORCEINLINE bool NeedsPrePass(const FDeferredShadingSceneRenderer* Render
 // This is a temporary extern while testing the Velocity changes. Will be removed once a decision is made.
 extern bool IsVelocityMergedWithDepthPass();
 
-bool FDeferredShadingSceneRenderer::RenderHzb(
-	FRDGBuilder& GraphBuilder,
-	TRDGUniformBufferRef<FSceneTextureUniformParameters> SceneTexturesUniformBuffer)
+bool FDeferredShadingSceneRenderer::RenderHzb(FRDGBuilder& GraphBuilder, const FSceneTextures& SceneTextures)
 {
 	RDG_GPU_STAT_SCOPE(GraphBuilder, HZB);
 
-	const FSceneTextureParameters SceneTextures = GetSceneTextureParameters(GraphBuilder, SceneTexturesUniformBuffer);
-
 	for (int32 ViewIndex = 0; ViewIndex < Views.Num(); ViewIndex++)
 	{
 		FViewInfo& View = Views[ViewIndex];
@@ -363,7 +359,7 @@ bool FDeferredShadingSceneRenderer::RenderHzb(
 
 			BuildHZB(
 				GraphBuilder,
-				SceneTextures.SceneDepthTexture,
+				SceneTextures.Depth.Resolve,
 				/* VisBufferTexture = */ nullptr,
 				View,
 				/* OutClosestHZBTexture = */ ViewPipelineState.bClosestHZB ? &ClosestHZBTexture : nullptr,
@@ -398,7 +394,7 @@ bool FDeferredShadingSceneRenderer::RenderHzb(
 	// Async ssao only requires HZB and depth as inputs so get started ASAP
 	if (CanOverlayRayTracingOutput(Views[0]) && GCompositionLighting.CanProcessAsyncSSAO(Views))
 	{
-		GCompositionLighting.ProcessAsyncSSAO(GraphBuilder, Views, SceneTexturesUniformBuffer);
+		GCompositionLighting.ProcessAsyncSSAO(GraphBuilder, Views, SceneTextures.UniformBuffer);
 	}
 
 	return FamilyPipelineState->bHZBOcclusion;
@@ -1648,8 +1644,6 @@ void FDeferredShadingSceneRenderer::Render(FRHICommandListImmediate& RHICmdList)
 	}
 #endif // RHI_RAYTRACING
 
-	const FIntPoint SceneTextureExtent = SceneContext.GetBufferSizeXY();
-
 	{
 		RDG_GPU_STAT_SCOPE(GraphBuilder, GPUSceneUpdate);
 
@@ -1720,9 +1714,7 @@ void FDeferredShadingSceneRenderer::Render(FRHICommandListImmediate& RHICmdList)
 		}
 	}
 
-	FRDGTextureMSAA SceneDepthTexture = RegisterExternalTextureMSAA(GraphBuilder, SceneContext.SceneDepthZ);
-	FRDGTextureRef SmallDepthTexture = GraphBuilder.RegisterExternalTexture(SceneContext.SmallDepthZ, ERenderTargetTexture::Targetable);
-	FRDGTextureSRVRef SceneStencilTexture = GraphBuilder.CreateSRV(FRDGTextureSRVDesc::CreateWithPixelFormat(SceneDepthTexture.Target, PF_X24_G8));
+	FSceneTextures& SceneTextures = FSceneTextures::Create(GraphBuilder);
 
 	const bool bUseGBuffer = IsUsingGBuffers(ShaderPlatform);
 	bool bCanOverlayRayTracingOutput = CanOverlayRayTracingOutput(Views[0]);// #dxr_todo: UE-72557 multi-view case
@@ -1791,7 +1783,7 @@ void FDeferredShadingSceneRenderer::Render(FRHICommandListImmediate& RHICmdList)
 
 	if (DepthPass.IsComputeStencilDitherEnabled())
 	{
-		AddDitheredStencilFillPass(GraphBuilder, Views, SceneDepthTexture.Target, DepthPass);
+		AddDitheredStencilFillPass(GraphBuilder, Views, SceneTextures.Depth.Target, DepthPass);
 	}
 
 	// Notify the FX system that the scene is about to be rendered.
@@ -1850,7 +1842,7 @@ void FDeferredShadingSceneRenderer::Render(FRHICommandListImmediate& RHICmdList)
 		Nanite::GStreamingManager.EndAsyncUpdate(GraphBuilder);
 	}
 
-	FRDGTextureRef VelocityTexture = TryRegisterExternalTexture(GraphBuilder, SceneContext.SceneVelocity);
+	SceneTextures.Velocity = TryRegisterExternalTexture(GraphBuilder, SceneContext.SceneVelocity);
 
 	const bool bShouldRenderVelocities = ShouldRenderVelocities();
 	const bool bBasePassCanOutputVelocity = FVelocityRendering::BasePassCanOutputVelocity(FeatureLevel);
@@ -1874,17 +1866,17 @@ void FDeferredShadingSceneRenderer::Render(FRHICommandListImmediate& RHICmdList)
 			: ERenderTargetLoadAction::EClear;
 
 		const ERenderTargetLoadAction DepthLoadAction = ERenderTargetLoadAction::EClear;
-		AddClearDepthStencilPass(GraphBuilder, SceneDepthTexture.Target, DepthLoadAction, StencilLoadAction);
+		AddClearDepthStencilPass(GraphBuilder, SceneTextures.Depth.Target, DepthLoadAction, StencilLoadAction);
 
 		// Draw the scene pre-pass / early z pass, populating the scene depth buffer and HiZ
 		if (bNeedsPrePass)
 		{
-			RenderPrePass(GraphBuilder, SceneDepthTexture.Target);
+			RenderPrePass(GraphBuilder, SceneTextures.Depth.Target);
 		}
 		else
 		{
 			// We didn't do the prepass, but we still want the HMD mask if there is one
-			RenderPrePassHMD(GraphBuilder, SceneDepthTexture.Target);
+			RenderPrePassHMD(GraphBuilder, SceneTextures.Depth.Target);
 		}
 
 		AddSetCurrentStatPass(GraphBuilder, GET_STATID(STAT_CLM_AfterPrePass));
@@ -1895,7 +1887,7 @@ void FDeferredShadingSceneRenderer::Render(FRHICommandListImmediate& RHICmdList)
 		{
 			// Render the velocities of movable objects
 			AddSetCurrentStatPass(GraphBuilder, GET_STATID(STAT_CLM_Velocity));
-			RenderVelocities(GraphBuilder, SceneDepthTexture.Resolve, VelocityTexture, nullptr, EVelocityPass::Opaque, bHairEnable);
+			RenderVelocities(GraphBuilder, SceneTextures, EVelocityPass::Opaque, bHairEnable);
 			AddSetCurrentStatPass(GraphBuilder, GET_STATID(STAT_CLM_AfterVelocity));
 			AddServiceLocalQueuePass(GraphBuilder);
 		}
@@ -1936,7 +1928,7 @@ void FDeferredShadingSceneRenderer::Render(FRHICommandListImmediate& RHICmdList)
 		NaniteRasterResults.AddDefaulted(Views.Num());
 
 		RDG_GPU_STAT_SCOPE(GraphBuilder, NaniteRaster);
-		const FIntPoint RasterTextureSize = SceneDepthTexture.Target->Desc.Extent;
+		const FIntPoint RasterTextureSize = SceneTextures.Depth.Target->Desc.Extent;
 
 		const FViewInfo& PrimaryViewRef = Views[0];
 		const FIntRect PrimaryViewRect = PrimaryViewRef.ViewRect;
@@ -2025,10 +2017,10 @@ void FDeferredShadingSceneRenderer::Render(FRHICommandListImmediate& RHICmdList)
 		}
 	}
 
-	ESceneTextureSetupMode SceneTexturesSetupMode = ESceneTextureSetupMode::SceneDepth;
-	TRDGUniformBufferRef<FSceneTextureUniformParameters> SceneTextures = CreateSceneTextureUniformBuffer(GraphBuilder, FeatureLevel, SceneTexturesSetupMode);
+	SceneTextures.SetupMode = ESceneTextureSetupMode::SceneDepth;
+	SceneTextures.UniformBuffer = CreateSceneTextureUniformBuffer(GraphBuilder, FeatureLevel, SceneTextures.SetupMode);
 
-	AddResolveSceneDepthPass(GraphBuilder, Views, SceneDepthTexture);
+	AddResolveSceneDepthPass(GraphBuilder, Views, SceneTextures.Depth);
 
 	// NOTE: The ordering of the lights is used to select sub-sets for different purposes, e.g., those that support clustered deferred.
 	FSortedLightSetSceneInfo SortedLightSet;
@@ -2072,7 +2064,7 @@ void FDeferredShadingSceneRenderer::Render(FRHICommandListImmediate& RHICmdList)
 
 	if (bOcclusionBeforeBasePass)
 	{
-		RenderOcclusion(GraphBuilder, SceneDepthTexture.Target, SmallDepthTexture, SceneTextures, bIsOcclusionTesting);
+		RenderOcclusion(GraphBuilder, SceneTextures, bIsOcclusionTesting);
 	}
 
 	AddServiceLocalQueuePass(GraphBuilder);
@@ -2130,8 +2122,8 @@ void FDeferredShadingSceneRenderer::Render(FRHICommandListImmediate& RHICmdList)
 		QUICK_SCOPE_CYCLE_COUNTER(STAT_FDeferredShadingSceneRenderer_CustomDepthPass0);
 		RenderCustomDepthPassAtLocation(GraphBuilder, 0);
 
-		SceneTexturesSetupMode |= ESceneTextureSetupMode::CustomDepth;
-		SceneTextures = CreateSceneTextureUniformBuffer(GraphBuilder, FeatureLevel, SceneTexturesSetupMode);
+		SceneTextures.SetupMode |= ESceneTextureSetupMode::CustomDepth;
+		SceneTextures.UniformBuffer = CreateSceneTextureUniformBuffer(GraphBuilder, FeatureLevel, SceneTextures.SetupMode);
 	}
 
 	UpdateLumenScene(GraphBuilder);
@@ -2160,7 +2152,7 @@ void FDeferredShadingSceneRenderer::Render(FRHICommandListImmediate& RHICmdList)
 			HairDatas = &HairDatasStorage;
 		}
 
-		RenderForwardShadowProjections(GraphBuilder, SceneTextures, SceneDepthTexture.Target, ForwardScreenSpaceShadowMaskTexture, ForwardScreenSpaceShadowMaskHairTexture, HairDatas);
+		RenderForwardShadowProjections(GraphBuilder, SceneTextures, ForwardScreenSpaceShadowMaskTexture, ForwardScreenSpaceShadowMaskHairTexture, HairDatas);
 	}
 
 	// only temporarily available after early z pass and until base pass
@@ -2189,7 +2181,7 @@ void FDeferredShadingSceneRenderer::Render(FRHICommandListImmediate& RHICmdList)
 				SSAOLevels = 0;
 			}
 
-			GCompositionLighting.ProcessBeforeBasePass(GraphBuilder, Scene->UniformBuffers, View, SceneTextures, bDBuffer, SSAOLevels);
+			GCompositionLighting.ProcessBeforeBasePass(GraphBuilder, Scene->UniformBuffers, View, SceneTextures.UniformBuffer, bDBuffer, SSAOLevels);
 		}
 
 		AddServiceLocalQueuePass(GraphBuilder);
@@ -2199,7 +2191,7 @@ void FDeferredShadingSceneRenderer::Render(FRHICommandListImmediate& RHICmdList)
 	{
 		RenderIndirectCapsuleShadows(
 			GraphBuilder,
-			SceneTextures,
+			SceneTextures.UniformBuffer,
 			nullptr,
 			GraphBuilder.RegisterExternalTexture(SceneContext.ScreenSpaceAO),
 			SceneContext.bScreenSpaceAOIsValid);
@@ -2213,10 +2205,10 @@ void FDeferredShadingSceneRenderer::Render(FRHICommandListImmediate& RHICmdList)
 	}
 
 	SceneContext.AllocSceneColor(RHICmdList);
-	FRDGTextureMSAA SceneColorTexture = RegisterExternalTextureMSAA(GraphBuilder, SceneContext.GetSceneColor());
+	SceneTextures.Color = RegisterExternalTextureMSAA(GraphBuilder, SceneContext.GetSceneColor());
 
 	{
-		RenderBasePass(GraphBuilder, BasePassDepthStencilAccess, SceneColorTexture.Target, SceneDepthTexture.Target, ERenderTargetLoadAction::ELoad, ForwardScreenSpaceShadowMaskTexture);
+		RenderBasePass(GraphBuilder, SceneTextures, BasePassDepthStencilAccess, ForwardScreenSpaceShadowMaskTexture);
 		AddServiceLocalQueuePass(GraphBuilder);
 		
 		if (bNaniteEnabled && bShouldApplyNaniteMaterials)
@@ -2226,7 +2218,7 @@ void FDeferredShadingSceneRenderer::Render(FRHICommandListImmediate& RHICmdList)
 				const FViewInfo& View = Views[ViewIndex];
 				const Nanite::FRasterResults& RasterResults = NaniteRasterResults[ViewIndex];
 
-				FRDGTextureRef SceneDepth = SceneDepthTexture.Target;
+				FRDGTextureRef SceneDepth = SceneTextures.Depth.Target;
 
 				Nanite::DrawBasePass(
 					GraphBuilder,
@@ -2248,13 +2240,13 @@ void FDeferredShadingSceneRenderer::Render(FRHICommandListImmediate& RHICmdList)
 
 		if (!bAllowReadOnlyDepthBasePass)
 		{
-			AddResolveSceneDepthPass(GraphBuilder, Views, SceneDepthTexture);
+			AddResolveSceneDepthPass(GraphBuilder, Views, SceneTextures.Depth);
 		}
 	}
 
 	// Rebuild scene textures to include GBuffers.
-	SceneTexturesSetupMode |= ESceneTextureSetupMode::GBuffers;
-	SceneTextures = CreateSceneTextureUniformBuffer(GraphBuilder, FeatureLevel, SceneTexturesSetupMode);
+	SceneTextures.SetupMode |= ESceneTextureSetupMode::GBuffers;
+	SceneTextures.UniformBuffer = CreateSceneTextureUniformBuffer(GraphBuilder, FeatureLevel, SceneTextures.SetupMode);
 
 	// mark GBufferA for saving for next frame
 	{
@@ -2279,12 +2271,12 @@ void FDeferredShadingSceneRenderer::Render(FRHICommandListImmediate& RHICmdList)
 
 	if (ViewFamily.EngineShowFlags.VisualizeLightCulling)
 	{
-		AddClearRenderTargetPass(GraphBuilder, SceneColorTexture.Target, FLinearColor::Transparent);
+		AddClearRenderTargetPass(GraphBuilder, SceneTextures.Color.Target, FLinearColor::Transparent);
 	}
 
 	if (bRealTimeSkyCaptureEnabled)
 	{
-		Scene->ValidateSkyLightRealTimeCapture(GraphBuilder, Views[0], SceneColorTexture.Target);
+		Scene->ValidateSkyLightRealTimeCapture(GraphBuilder, Views[0], SceneTextures.Color.Target);
 	}
 
 	AddPass(GraphBuilder, [&SceneContext](FRHICommandList&)
@@ -2295,12 +2287,12 @@ void FDeferredShadingSceneRenderer::Render(FRHICommandListImmediate& RHICmdList)
 		SceneContext.DBufferMask = nullptr;
 	});
 
-	VisualizeVolumetricLightmap(GraphBuilder, SceneColorTexture.Target, SceneDepthTexture.Target);
+	VisualizeVolumetricLightmap(GraphBuilder, SceneTextures);
 
 	// Occlusion after base pass
 	if (!bOcclusionBeforeBasePass)
 	{
-		RenderOcclusion(GraphBuilder, SceneDepthTexture.Target, SmallDepthTexture, SceneTextures, bIsOcclusionTesting);
+		RenderOcclusion(GraphBuilder, SceneTextures, bIsOcclusionTesting);
 	}
 
 	AddServiceLocalQueuePass(GraphBuilder);
@@ -2309,7 +2301,7 @@ void FDeferredShadingSceneRenderer::Render(FRHICommandListImmediate& RHICmdList)
 
 	if (!bUseGBuffer)
 	{
-		AddResolveSceneColorPass(GraphBuilder, Views, SceneColorTexture);
+		AddResolveSceneColorPass(GraphBuilder, Views, SceneTextures.Color);
 	}
 
 	// Shadow and fog after base pass
@@ -2363,8 +2355,8 @@ void FDeferredShadingSceneRenderer::Render(FRHICommandListImmediate& RHICmdList)
 		QUICK_SCOPE_CYCLE_COUNTER(STAT_FDeferredShadingSceneRenderer_CustomDepthPass1);
 		RenderCustomDepthPassAtLocation(GraphBuilder, 1);
 
-		SceneTexturesSetupMode |= ESceneTextureSetupMode::CustomDepth;
-		SceneTextures = CreateSceneTextureUniformBuffer(GraphBuilder, FeatureLevel, SceneTexturesSetupMode);
+		SceneTextures.SetupMode |= ESceneTextureSetupMode::CustomDepth;
+		SceneTextures.UniformBuffer = CreateSceneTextureUniformBuffer(GraphBuilder, FeatureLevel, SceneTextures.SetupMode);
 	}
 
 #if RHI_RAYTRACING
@@ -2381,7 +2373,7 @@ void FDeferredShadingSceneRenderer::Render(FRHICommandListImmediate& RHICmdList)
 	{
 		// Render the velocities of movable objects
 		AddSetCurrentStatPass(GraphBuilder, GET_STATID(STAT_CLM_Velocity));
-		RenderVelocities(GraphBuilder, SceneDepthTexture.Resolve, VelocityTexture, nullptr, EVelocityPass::Opaque, bHairEnable);
+		RenderVelocities(GraphBuilder, SceneTextures, EVelocityPass::Opaque, bHairEnable);
 		AddSetCurrentStatPass(GraphBuilder, GET_STATID(STAT_CLM_AfterVelocity));
 		AddServiceLocalQueuePass(GraphBuilder);
 
@@ -2402,13 +2394,13 @@ void FDeferredShadingSceneRenderer::Render(FRHICommandListImmediate& RHICmdList)
 
 		if (bCanOverlayRayTracingOutput && !IsForwardShadingEnabled(ShaderPlatform))
 		{
-			RenderRayTracingSkyLight(GraphBuilder, SceneColorTexture.Target, SkyLightTexture, SkyLightHitDistanceTexture, HairDatas);
+			RenderRayTracingSkyLight(GraphBuilder, SceneTextures.Color.Target, SkyLightTexture, SkyLightHitDistanceTexture, HairDatas);
 		}
 	}
 #endif // RHI_RAYTRACING
 
 	// Copy lighting channels out of stencil before deferred decals which overwrite those values
-	FRDGTextureRef LightingChannelsTexture = CopyStencilToLightingChannelTexture(GraphBuilder, Views, SceneStencilTexture);
+	FRDGTextureRef LightingChannelsTexture = CopyStencilToLightingChannelTexture(GraphBuilder, Views, SceneTextures.Stencil);
 
 	if (IsForwardShadingEnabled(ShaderPlatform))
 	{
@@ -2429,7 +2421,7 @@ void FDeferredShadingSceneRenderer::Render(FRHICommandListImmediate& RHICmdList)
 
 		if (!IsForwardShadingEnabled(ShaderPlatform))
 		{
-			AddResolveSceneDepthPass(GraphBuilder, Views, SceneDepthTexture);
+			AddResolveSceneDepthPass(GraphBuilder, Views, SceneTextures.Depth);
 		}
 
 		for (int32 ViewIndex = 0; ViewIndex < Views.Num(); ViewIndex++)
@@ -2439,7 +2431,7 @@ void FDeferredShadingSceneRenderer::Render(FRHICommandListImmediate& RHICmdList)
 			RDG_EVENT_SCOPE_CONDITIONAL(GraphBuilder, Views.Num() > 1, "View%d", ViewIndex);
 
 			bool bEnableSSAO = ViewPipelineState->AmbientOcclusionMethod == EAmbientOcclusionMethod::SSAO;
-			GCompositionLighting.ProcessAfterBasePass(GraphBuilder, Scene->UniformBuffers, View, SceneTextures, bEnableSSAO);
+			GCompositionLighting.ProcessAfterBasePass(GraphBuilder, Scene->UniformBuffers, View, SceneTextures.UniformBuffer, bEnableSSAO);
 		}
 		SceneContext.ScreenSpaceGTAOHorizons.SafeRelease();
 		AddServiceLocalQueuePass(GraphBuilder);
@@ -2452,14 +2444,14 @@ void FDeferredShadingSceneRenderer::Render(FRHICommandListImmediate& RHICmdList)
 
 		// Disable temporarly HZB update with hair data. This is not mandatory.
 		#if 0
-		const FSceneTextureParameters SceneTextureParameters = GetSceneTextureParameters(GraphBuilder, SceneTextures);
+		const FSceneTextureParameters SceneTextureParameters = GetSceneTextureParameters(GraphBuilder, SceneTextures.UniformBuffer);
 		for (int32 ViewIndex = 0; ViewIndex < Views.Num(); ViewIndex++)
 		{
 			FViewInfo& View = Views[ViewIndex];
 			RDG_EVENT_SCOPE(GraphBuilder, "BuildHZB_HairUpdate(ViewId=%d)", ViewIndex);
 			BuildHZB(
 				GraphBuilder,
-				SceneTextureParameters.SceneDepthTexture, 
+				SceneTextureParameters.SceneTextures.Depth, 
 				/* VisBufferTexture = */ nullptr,
 				View,
 				/* OutClosestHZBTexture = */ nullptr,
@@ -2469,13 +2461,13 @@ void FDeferredShadingSceneRenderer::Render(FRHICommandListImmediate& RHICmdList)
 	}
 
 	// Rebuild scene textures to include velocity, custom depth, and SSAO.
-	SceneTexturesSetupMode |= ESceneTextureSetupMode::All;
-	SceneTextures = CreateSceneTextureUniformBuffer(GraphBuilder, FeatureLevel, SceneTexturesSetupMode);
+	SceneTextures.SetupMode |= ESceneTextureSetupMode::All;
+	SceneTextures.UniformBuffer = CreateSceneTextureUniformBuffer(GraphBuilder, FeatureLevel, SceneTextures.SetupMode);
 
 	if (!IsForwardShadingEnabled(ShaderPlatform))
 	{
 		// Clear stencil to 0 now that deferred decals are done using what was setup in the base pass.
-		AddClearStencilPass(GraphBuilder, SceneDepthTexture.Target);
+		AddClearStencilPass(GraphBuilder, SceneTextures.Depth.Target);
 	}
 
 
@@ -2486,21 +2478,21 @@ void FDeferredShadingSceneRenderer::Render(FRHICommandListImmediate& RHICmdList)
 		SCOPE_CYCLE_COUNTER(STAT_FDeferredShadingSceneRenderer_Lighting);
 
 		FRDGTextureRef DynamicBentNormalAOTexture = nullptr;
-		RenderDiffuseIndirectAndAmbientOcclusion(GraphBuilder, SceneTextures, SceneColorTexture.Target, HairDatas, /* bIsVisualizePass = */ false);
+		RenderDiffuseIndirectAndAmbientOcclusion(GraphBuilder, SceneTextures, HairDatas, /* bIsVisualizePass = */ false);
 
 		// These modulate the scenecolor output from the basepass, which is assumed to be indirect lighting
 		if (SceneContext.IsStaticLightingAllowed())
 		{
 			RenderIndirectCapsuleShadows(
 				GraphBuilder,
-				SceneTextures,
-				SceneColorTexture.Target,
+				SceneTextures.UniformBuffer,
+				SceneTextures.Color.Target,
 				GraphBuilder.RegisterExternalTexture(SceneContext.ScreenSpaceAO),
 				SceneContext.bScreenSpaceAOIsValid);
 		}
 
 		// These modulate the scene color output from the base pass, which is assumed to be indirect lighting
-		RenderDFAOAsIndirectShadowing(GraphBuilder, SceneTextures, SceneColorTexture.Target, VelocityTexture, DynamicBentNormalAOTexture);
+		RenderDFAOAsIndirectShadowing(GraphBuilder, SceneTextures, DynamicBentNormalAOTexture);
 
 		// Clear the translucent lighting volumes before we accumulate
 		if ((GbEnableAsyncComputeTranslucencyLightingVolumeClear && GSupportsEfficientAsyncCompute) == false)
@@ -2511,20 +2503,12 @@ void FDeferredShadingSceneRenderer::Render(FRHICommandListImmediate& RHICmdList)
 #if RHI_RAYTRACING
 		if (bRayTracingEnabled)
 		{
-			RenderDitheredLODFadingOutMask(GraphBuilder, Views[0], SceneDepthTexture.Target);
+			RenderDitheredLODFadingOutMask(GraphBuilder, Views[0], SceneTextures.Depth.Target);
 		}
 #endif
 
 		AddSetCurrentStatPass(GraphBuilder, GET_STATID(STAT_CLM_Lighting));
-		SceneColorTexture = FRDGTextureMSAA(RenderLights(
-			GraphBuilder,
-			SceneTextures,
-			TranslucentLightingVolumeTextures,
-			SceneColorTexture.Target,
-			SceneDepthTexture.Target,
-			LightingChannelsTexture,
-			SortedLightSet,
-			HairDatas));
+		RenderLights(GraphBuilder, SceneTextures, TranslucentLightingVolumeTextures, LightingChannelsTexture, SortedLightSet, HairDatas);
 		AddSetCurrentStatPass(GraphBuilder, GET_STATID(STAT_CLM_AfterLighting));
 		AddServiceLocalQueuePass(GraphBuilder);
 
@@ -2547,19 +2531,19 @@ void FDeferredShadingSceneRenderer::Render(FRHICommandListImmediate& RHICmdList)
 		AddServiceLocalQueuePass(GraphBuilder);
 
 		// Render diffuse sky lighting and reflections that only operate on opaque pixels
-		RenderDeferredReflectionsAndSkyLighting(GraphBuilder, SceneTextures, SceneColorTexture, DynamicBentNormalAOTexture, VelocityTexture, HairDatas);
+		RenderDeferredReflectionsAndSkyLighting(GraphBuilder, SceneTextures, DynamicBentNormalAOTexture, HairDatas);
 
-		SceneColorTexture = FRDGTextureMSAA(AddSubsurfacePass(GraphBuilder, SceneTextures, Views, SceneColorTexture.Target));
+		AddSubsurfacePass(GraphBuilder, SceneTextures, Views);
 
 		if (HairDatas)
 		{
-			RenderHairStrandsSceneColorScattering(GraphBuilder, SceneColorTexture.Target, Scene, Views, HairDatas);
+			RenderHairStrandsSceneColorScattering(GraphBuilder, SceneTextures.Color.Target, Scene, Views, HairDatas);
 		}
 
 	#if RHI_RAYTRACING
 		if (SkyLightTexture)
 		{
-			CompositeRayTracingSkyLight(GraphBuilder, SceneColorTexture.Target, SceneTextureExtent, SkyLightTexture, SkyLightHitDistanceTexture);
+			CompositeRayTracingSkyLight(GraphBuilder, SceneTextures, SkyLightTexture, SkyLightHitDistanceTexture);
 		}
 	#endif
 
@@ -2567,7 +2551,7 @@ void FDeferredShadingSceneRenderer::Render(FRHICommandListImmediate& RHICmdList)
 	}
 	else if (HairDatas)
 	{
-		RenderLightsForHair(GraphBuilder, SceneTextures, SortedLightSet, HairDatas, ForwardScreenSpaceShadowMaskHairTexture, LightingChannelsTexture);
+		RenderLightsForHair(GraphBuilder, SceneTextures.UniformBuffer, SortedLightSet, HairDatas, ForwardScreenSpaceShadowMaskHairTexture, LightingChannelsTexture);
 		RenderDeferredReflectionsAndSkyLightingHair(GraphBuilder, HairDatas);
 	}
 
@@ -2575,7 +2559,7 @@ void FDeferredShadingSceneRenderer::Render(FRHICommandListImmediate& RHICmdList)
 
 	if (bShouldRenderVolumetricCloud && IsVolumetricRenderTargetEnabled())
 	{
-		HalfResolutionDepthCheckerboardMinMaxTexture = CreateHalfResolutionDepthCheckerboardMinMax(GraphBuilder, Views, SceneDepthTexture.Resolve);
+		HalfResolutionDepthCheckerboardMinMaxTexture = CreateHalfResolutionDepthCheckerboardMinMax(GraphBuilder, Views, SceneTextures.Depth.Resolve);
 	}
 
 	if (bShouldRenderVolumetricCloud)
@@ -2583,10 +2567,10 @@ void FDeferredShadingSceneRenderer::Render(FRHICommandListImmediate& RHICmdList)
 		// Generate the volumetric cloud render target
 		bool bSkipVolumetricRenderTarget = false;
 		bool bSkipPerPixelTracing = true;
-		RenderVolumetricCloud(GraphBuilder, GetSceneTextureShaderParameters(SceneTextures), bSkipVolumetricRenderTarget, bSkipPerPixelTracing, SceneColorTexture, SceneDepthTexture, HalfResolutionDepthCheckerboardMinMaxTexture);
+		RenderVolumetricCloud(GraphBuilder, SceneTextures, bSkipVolumetricRenderTarget, bSkipPerPixelTracing, HalfResolutionDepthCheckerboardMinMaxTexture);
 
 		// Reconstruct the volumetric cloud render target to be ready to compose it over the scene
-		ReconstructVolumetricRenderTarget(GraphBuilder, Views, SceneDepthTexture.Resolve, HalfResolutionDepthCheckerboardMinMaxTexture);
+		ReconstructVolumetricRenderTarget(GraphBuilder, Views, SceneTextures.Depth.Resolve, HalfResolutionDepthCheckerboardMinMaxTexture);
 	}
 
 	const bool bShouldRenderTranslucency = ShouldRenderTranslucency();
@@ -2603,17 +2587,17 @@ void FDeferredShadingSceneRenderer::Render(FRHICommandListImmediate& RHICmdList)
 			RDG_CSV_STAT_EXCLUSIVE_SCOPE(GraphBuilder, RenderTranslucency);
 			SCOPE_CYCLE_COUNTER(STAT_TranslucencyDrawTime);
 			AddSetCurrentStatPass(GraphBuilder, GET_STATID(STAT_CLM_Translucency));
-			RenderTranslucency(GraphBuilder, TranslucentLightingVolumeTextures, SceneColorTexture, SceneDepthTexture, nullptr, ETranslucencyView::UnderWater);
+			RenderTranslucency(GraphBuilder, SceneTextures, TranslucentLightingVolumeTextures, nullptr, ETranslucencyView::UnderWater);
 			EnumRemoveFlags(TranslucencyViewsToRender, ETranslucencyView::UnderWater);
 		}
 
 		AddSetCurrentStatPass(GraphBuilder, GET_STATID(STAT_CLM_WaterPass));
-		RenderSingleLayerWater(GraphBuilder, SceneColorTexture, SceneDepthTexture, SceneTextures, bShouldRenderVolumetricCloud, SceneWithoutWaterTextures);
+		RenderSingleLayerWater(GraphBuilder, SceneTextures, bShouldRenderVolumetricCloud, SceneWithoutWaterTextures);
 		AddServiceLocalQueuePass(GraphBuilder);
 	}
 
 	// Rebuild scene textures to include scene color.
-	SceneTextures = CreateSceneTextureUniformBuffer(GraphBuilder, FeatureLevel, SceneTexturesSetupMode);
+	SceneTextures.UniformBuffer = CreateSceneTextureUniformBuffer(GraphBuilder, FeatureLevel, SceneTextures.SetupMode);
 
 	FRDGTextureRef LightShaftOcclusionTexture = nullptr;
 
@@ -2621,21 +2605,21 @@ void FDeferredShadingSceneRenderer::Render(FRHICommandListImmediate& RHICmdList)
 	if (ViewFamily.EngineShowFlags.LightShafts)
 	{
 		SCOPE_CYCLE_COUNTER(STAT_FDeferredShadingSceneRenderer_RenderLightShaftOcclusion);
-		LightShaftOcclusionTexture = RenderLightShaftOcclusion(GraphBuilder, SceneTextures, SceneTextureExtent);
+		LightShaftOcclusionTexture = RenderLightShaftOcclusion(GraphBuilder, SceneTextures);
 	}
 
 	// Draw atmosphere
 	if (bCanOverlayRayTracingOutput && ShouldRenderAtmosphere(ViewFamily))
 	{
 		SCOPE_CYCLE_COUNTER(STAT_FDeferredShadingSceneRenderer_RenderAtmosphere);
-		RenderAtmosphere(GraphBuilder, SceneColorTexture.Target, SceneDepthTexture.Target, LightShaftOcclusionTexture, SceneTextures);
+		RenderAtmosphere(GraphBuilder, SceneTextures, LightShaftOcclusionTexture);
 	}
 
 	// Draw the sky atmosphere
 	if (bShouldRenderSkyAtmosphere)
 	{
 		SCOPE_CYCLE_COUNTER(STAT_FDeferredShadingSceneRenderer_RenderSkyAtmosphere);
-		RenderSkyAtmosphere(GraphBuilder, SceneTextures, SceneColorTexture.Target, SceneDepthTexture.Target);
+		RenderSkyAtmosphere(GraphBuilder, SceneTextures);
 	}
 
 	// Draw fog.
@@ -2643,7 +2627,7 @@ void FDeferredShadingSceneRenderer::Render(FRHICommandListImmediate& RHICmdList)
 	{
 		RDG_CSV_STAT_EXCLUSIVE_SCOPE(GraphBuilder, RenderFog);
 		SCOPE_CYCLE_COUNTER(STAT_FDeferredShadingSceneRenderer_RenderFog);
-		RenderFog(GraphBuilder, SceneColorTexture.Target, SceneDepthTexture.Target, LightShaftOcclusionTexture, SceneTextures);
+		RenderFog(GraphBuilder, SceneTextures, LightShaftOcclusionTexture);
 	}
 
 	// After the height fog, Draw volumetric clouds (having fog applied on them already) when using per pixel tracing,
@@ -2651,30 +2635,30 @@ void FDeferredShadingSceneRenderer::Render(FRHICommandListImmediate& RHICmdList)
 	{
 		bool bSkipVolumetricRenderTarget = true;
 		bool bSkipPerPixelTracing = false;
-		RenderVolumetricCloud(GraphBuilder, GetSceneTextureShaderParameters(SceneTextures), bSkipVolumetricRenderTarget, bSkipPerPixelTracing, SceneColorTexture, SceneDepthTexture, HalfResolutionDepthCheckerboardMinMaxTexture);
+		RenderVolumetricCloud(GraphBuilder, SceneTextures, bSkipVolumetricRenderTarget, bSkipPerPixelTracing, HalfResolutionDepthCheckerboardMinMaxTexture);
 	}
 
 	// or composite the off screen buffer over the scene.
 	if (bVolumetricRenderTargetRequired)
 	{
-		ComposeVolumetricRenderTargetOverScene(GraphBuilder, Views, SceneColorTexture.Target, SceneDepthTexture.Target, bShouldRenderSingleLayerWater, SceneWithoutWaterTextures);
+		ComposeVolumetricRenderTargetOverScene(GraphBuilder, Views, SceneTextures.Color.Target, SceneTextures.Depth.Target, bShouldRenderSingleLayerWater, SceneWithoutWaterTextures);
 	}
 
 	FRendererModule& RendererModule = static_cast<FRendererModule&>(GetRendererModule());
-	RendererModule.RenderPostOpaqueExtensions(GraphBuilder, Views, SceneColorTexture.Target, SceneDepthTexture.Target, SceneTextures, SceneContext);
+	RendererModule.RenderPostOpaqueExtensions(GraphBuilder, Views, SceneTextures, SceneContext);
 
-	RenderOpaqueFX(GraphBuilder, Views, FXSystem, SceneTextures);
+	RenderOpaqueFX(GraphBuilder, Views, FXSystem, SceneTextures.UniformBuffer);
 
 	if (bShouldRenderSkyAtmosphere)
 	{
 		// Debug the sky atmosphere. Critically rendered before translucency to avoid emissive leaking over visualization by writing depth. 
 		// Alternative: render in post process chain as VisualizeHDR.
-		RenderDebugSkyAtmosphere(GraphBuilder, SceneColorTexture.Target, SceneDepthTexture.Target);
+		RenderDebugSkyAtmosphere(GraphBuilder, SceneTextures.Color.Target, SceneTextures.Depth.Target);
 	}
 
 	if (HairDatas && !IsHairStrandsComposeAfterTranslucency())
 	{
-		RenderHairComposition(GraphBuilder, Views, HairDatas, SceneColorTexture.Target, SceneDepthTexture.Target);
+		RenderHairComposition(GraphBuilder, Views, HairDatas, SceneTextures.Color.Target, SceneTextures.Depth.Target);
 	}
 
 	FSeparateTranslucencyTextures SeparateTranslucencyTextures(SeparateTranslucencyDimensions);
@@ -2691,36 +2675,36 @@ void FDeferredShadingSceneRenderer::Render(FRHICommandListImmediate& RHICmdList)
 #if RHI_RAYTRACING
 		if (EnumHasAnyFlags(TranslucencyViewsToRender, ETranslucencyView::RayTracing))
 		{
-			RenderRayTracingTranslucency(GraphBuilder, SceneColorTexture);
+			RenderRayTracingTranslucency(GraphBuilder, SceneTextures.Color);
 			EnumRemoveFlags(TranslucencyViewsToRender, ETranslucencyView::RayTracing);
 		}
 #endif
 
 		// Render all remaining translucency views.
 		AddSetCurrentStatPass(GraphBuilder, GET_STATID(STAT_CLM_Translucency));
-		RenderTranslucency(GraphBuilder, TranslucentLightingVolumeTextures, SceneColorTexture, SceneDepthTexture, &SeparateTranslucencyTextures, TranslucencyViewsToRender);
+		RenderTranslucency(GraphBuilder, SceneTextures, TranslucentLightingVolumeTextures, &SeparateTranslucencyTextures, TranslucencyViewsToRender);
 		AddServiceLocalQueuePass(GraphBuilder);
 		TranslucencyViewsToRender = ETranslucencyView::None;
 
 		if (bShouldRenderDistortion)
 		{
 			AddSetCurrentStatPass(GraphBuilder, GET_STATID(STAT_CLM_Distortion));
-			RenderDistortion(GraphBuilder, SceneColorTexture.Target, SceneDepthTexture.Target);
+			RenderDistortion(GraphBuilder, SceneTextures.Color.Target, SceneTextures.Depth.Target);
 			AddServiceLocalQueuePass(GraphBuilder);
 		}
 
 		if (bShouldRenderVelocities)
 		{
-			const bool bRecreateSceneTextures = !VelocityTexture;
+			const bool bRecreateSceneTextures = !SceneTextures.Velocity;
 
 			AddSetCurrentStatPass(GraphBuilder, GET_STATID(STAT_CLM_TranslucentVelocity));
-			RenderVelocities(GraphBuilder, SceneDepthTexture.Resolve, VelocityTexture, SceneTextures, EVelocityPass::Translucent, false);
+			RenderVelocities(GraphBuilder, SceneTextures, EVelocityPass::Translucent, false);
 			AddServiceLocalQueuePass(GraphBuilder);
 
 			if (bRecreateSceneTextures)
 			{
 				// Rebuild scene textures to include newly allocated velocity.
-				SceneTextures = CreateSceneTextureUniformBuffer(GraphBuilder, FeatureLevel, SceneTexturesSetupMode);
+				SceneTextures.UniformBuffer = CreateSceneTextureUniformBuffer(GraphBuilder, FeatureLevel, SceneTextures.SetupMode);
 			}
 		}
 
@@ -2731,7 +2715,8 @@ void FDeferredShadingSceneRenderer::Render(FRHICommandListImmediate& RHICmdList)
 	if (CVarForceBlackVelocityBuffer.GetValueOnRenderThread())
 	{
 		SceneContext.SceneVelocity = GSystemTextures.BlackDummy;
-		VelocityTexture = GraphBuilder.RegisterExternalTexture(SceneContext.SceneVelocity);
+		SceneTextures.Velocity = GraphBuilder.RegisterExternalTexture(SceneContext.SceneVelocity);
+		SceneTextures.UniformBuffer = CreateSceneTextureUniformBuffer(GraphBuilder, FeatureLevel, SceneTextures.SetupMode);
 	}
 #endif
 
@@ -2739,25 +2724,25 @@ void FDeferredShadingSceneRenderer::Render(FRHICommandListImmediate& RHICmdList)
 		RDG_GPU_STAT_SCOPE(GraphBuilder, HairRendering);
 		if (HairDatas && IsHairStrandsComposeAfterTranslucency())
 		{
-			RenderHairComposition(GraphBuilder, Views, HairDatas, SceneColorTexture.Target, SceneDepthTexture.Target);
+			RenderHairComposition(GraphBuilder, Views, HairDatas, SceneTextures.Color.Target, SceneTextures.Depth.Target);
 		}
 
 		if (HairStrandsBookmarkParameters.bHasElements)
 		{
-			RenderHairStrandsDebugInfo(GraphBuilder, Views, HairDatas, HairStrandsBookmarkParameters.HairClusterData, SceneColorTexture.Target);
+			RenderHairStrandsDebugInfo(GraphBuilder, Views, HairDatas, HairStrandsBookmarkParameters.HairClusterData, SceneTextures.Color.Target);
 		}
 	}
 
 	if (bStrataEnabled)
 	{
-		Strata::AddVisualizeMaterialPasses(GraphBuilder, Views, SceneColorTexture.Target);
+		Strata::AddVisualizeMaterialPasses(GraphBuilder, Views, SceneTextures.Color.Target);
 	}
 
 	if (bCanOverlayRayTracingOutput && ViewFamily.EngineShowFlags.LightShafts)
 	{
 		SCOPE_CYCLE_COUNTER(STAT_FDeferredShadingSceneRenderer_RenderLightShaftBloom);
 		AddSetCurrentStatPass(GraphBuilder, GET_STATID(STAT_CLM_LightShaftBloom));
-		RenderLightShaftBloom(GraphBuilder, SceneTextures, SceneTextureExtent, SceneColorTexture.Target, SeparateTranslucencyTextures);
+		RenderLightShaftBloom(GraphBuilder, SceneTextures, SeparateTranslucencyTextures);
 		AddServiceLocalQueuePass(GraphBuilder);
 	}
 
@@ -2793,17 +2778,17 @@ void FDeferredShadingSceneRenderer::Render(FRHICommandListImmediate& RHICmdList)
 		{
 			if (Views[ViewIndex].RayTracingRenderMode == ERayTracingRenderMode::PathTracing)
 			{
-				RenderPathTracing(GraphBuilder, Views[ViewIndex], SceneTextures, SceneColorTexture.Target);
+				RenderPathTracing(GraphBuilder, Views[ViewIndex], SceneTextures.UniformBuffer, SceneTextures.Color.Target);
 			}
 			else if (Views[ViewIndex].RayTracingRenderMode == ERayTracingRenderMode::RayTracingDebug)
 			{
-				RenderRayTracingDebug(GraphBuilder, Views[ViewIndex], SceneColorTexture.Target);
+				RenderRayTracingDebug(GraphBuilder, Views[ViewIndex], SceneTextures.Color.Target);
 			}
 		}
 	}
 #endif
 
-	RendererModule.RenderOverlayExtensions(GraphBuilder, Views, SceneColorTexture.Target, SceneDepthTexture.Target, SceneContext);
+	RendererModule.RenderOverlayExtensions(GraphBuilder, Views, SceneTextures, SceneContext);
 
 	if (ViewFamily.EngineShowFlags.VisualizeDistanceFieldAO && ShouldRenderDistanceFieldLighting())
 	{
@@ -2812,34 +2797,34 @@ void FDeferredShadingSceneRenderer::Render(FRHICommandListImmediate& RHICmdList)
 		// Use the skylight's max distance if there is one, to be consistent with DFAO shadowing on the skylight
 		const float OcclusionMaxDistance = Scene->SkyLight && !Scene->SkyLight->bWantsStaticShadowing ? Scene->SkyLight->OcclusionMaxDistance : Scene->DefaultMaxDistanceFieldOcclusionDistance;
 		FRDGTextureRef DummyOutput = nullptr;
-		RenderDistanceFieldLighting(GraphBuilder, SceneTextures, FDistanceFieldAOParameters(OcclusionMaxDistance), SceneColorTexture.Target, VelocityTexture, DummyOutput, false, ViewFamily.EngineShowFlags.VisualizeDistanceFieldAO);
+		RenderDistanceFieldLighting(GraphBuilder, SceneTextures, FDistanceFieldAOParameters(OcclusionMaxDistance), DummyOutput, false, ViewFamily.EngineShowFlags.VisualizeDistanceFieldAO);
 		AddServiceLocalQueuePass(GraphBuilder);
 	}
 
 	// Draw visualizations just before use to avoid target contamination
 	if (ViewFamily.EngineShowFlags.VisualizeMeshDistanceFields || ViewFamily.EngineShowFlags.VisualizeGlobalDistanceField)
 	{
-		RenderMeshDistanceFieldVisualization(GraphBuilder, SceneColorTexture.Target, SceneDepthTexture.Target, SceneTextures, FDistanceFieldAOParameters(Scene->DefaultMaxDistanceFieldOcclusionDistance));
+		RenderMeshDistanceFieldVisualization(GraphBuilder, SceneTextures, FDistanceFieldAOParameters(Scene->DefaultMaxDistanceFieldOcclusionDistance));
 		AddServiceLocalQueuePass(GraphBuilder);
 	}
 
 	RenderLumenSceneVisualization(GraphBuilder);
-	RenderDiffuseIndirectAndAmbientOcclusion(GraphBuilder, SceneTextures, SceneColorTexture.Target, HairDatas, /* bIsVisualizePass = */ true);
+	RenderDiffuseIndirectAndAmbientOcclusion(GraphBuilder, SceneTextures, HairDatas, /* bIsVisualizePass = */ true);
 
 	if (ViewFamily.EngineShowFlags.StationaryLightOverlap)
 	{
-		RenderStationaryLightOverlap(GraphBuilder, SceneColorTexture.Target, SceneDepthTexture.Target, LightingChannelsTexture, SceneTextures);
+		RenderStationaryLightOverlap(GraphBuilder, SceneTextures, LightingChannelsTexture);
 		AddServiceLocalQueuePass(GraphBuilder);
 	}
 
 	// Resolve the scene color for post processing.
-	AddResolveSceneColorPass(GraphBuilder, Views, SceneColorTexture);
+	AddResolveSceneColorPass(GraphBuilder, Views, SceneTextures.Color);
 
 	RendererModule.RenderPostResolvedSceneColorExtension(GraphBuilder, SceneContext);
 
 	FRDGTextureRef ViewFamilyTexture = TryCreateViewFamilyTexture(GraphBuilder, ViewFamily);
 
-	CopySceneCaptureComponentToTarget(GraphBuilder, SceneTextures, ViewFamilyTexture);
+	CopySceneCaptureComponentToTarget(GraphBuilder, SceneTextures.UniformBuffer, ViewFamilyTexture);
 
 	// Finish rendering for each view.
 	if (ViewFamily.bResolveScene && ViewFamilyTexture)
@@ -2853,7 +2838,7 @@ void FDeferredShadingSceneRenderer::Render(FRHICommandListImmediate& RHICmdList)
 		FPostProcessingInputs PostProcessingInputs;
 		PostProcessingInputs.ViewFamilyTexture = ViewFamilyTexture;
 		PostProcessingInputs.SeparateTranslucencyTextures = &SeparateTranslucencyTextures;
-		PostProcessingInputs.SceneTextures = SceneTextures;
+		PostProcessingInputs.SceneTextures = SceneTextures.UniformBuffer;
 
 		if (ViewFamily.UseDebugViewPS())
 		{
@@ -2943,8 +2928,8 @@ void FDeferredShadingSceneRenderer::Render(FRHICommandListImmediate& RHICmdList)
 		{
 			// Keep scene color and depth for next frame screen space ray tracing.
 			FSceneViewState* ViewState = View.ViewState;
-			GraphBuilder.QueueTextureExtraction((*SceneTextures)->SceneDepthTexture, &ViewState->PrevFrameViewInfo.DepthBuffer);
-			GraphBuilder.QueueTextureExtraction((*SceneTextures)->SceneColorTexture, &ViewState->PrevFrameViewInfo.ScreenSpaceRayTracingInput);
+			GraphBuilder.QueueTextureExtraction(SceneTextures.Depth.Resolve, &ViewState->PrevFrameViewInfo.DepthBuffer);
+			GraphBuilder.QueueTextureExtraction(SceneTextures.Color.Resolve, &ViewState->PrevFrameViewInfo.ScreenSpaceRayTracingInput);
 		}
 	}
 
diff --git a/Engine/Source/Runtime/Renderer/Private/DeferredShadingRenderer.h b/Engine/Source/Runtime/Renderer/Private/DeferredShadingRenderer.h
--- a/Engine/Source/Runtime/Renderer/Private/DeferredShadingRenderer.h
+++ b/Engine/Source/Runtime/Renderer/Private/DeferredShadingRenderer.h
@@ -198,10 +198,8 @@ public:
 
 	void RenderBasePass(
 		FRDGBuilder& GraphBuilder,
+		const FSceneTextures& SceneTextures,
 		FExclusiveDepthStencil::Type BasePassDepthStencilAccess,
-		FRDGTextureRef SceneColorTexture,
-		FRDGTextureRef SceneDepthTexture,
-		ERenderTargetLoadAction SceneDepthLoadAction,
 		FRDGTextureRef ForwardShadowMaskTexture);
 
 	void RenderBasePassInternal(
@@ -221,16 +219,13 @@ public:
 
 	void RenderSingleLayerWater(
 		FRDGBuilder& GraphBuilder,
-		FRDGTextureMSAA SceneColorTexture,
-		FRDGTextureMSAA SceneDepthTexture,
-		TRDGUniformBufferRef<FSceneTextureUniformParameters> SceneTexturesUniformBuffer,
+		const FSceneTextures& SceneTextures,
 		bool bShouldRenderVolumetricCloud,
 		FSceneWithoutWaterTextures& SceneWithoutWaterTextures);
 	
 	void RenderSingleLayerWaterInner(
 		FRDGBuilder& GraphBuilder,
-		FRDGTextureMSAA SceneColorTexture,
-		FRDGTextureMSAA SceneDepthTexture,
+		const FSceneTextures& SceneTextures,
 		const FSceneWithoutWaterTextures& SceneWithoutWaterTextures);
 
 	void RenderSingleLayerWaterReflections(
@@ -240,14 +235,12 @@ public:
 
 	void RenderOcclusion(
 		FRDGBuilder& GraphBuilder,
-		FRDGTextureRef SceneDepthTexture,
-		FRDGTextureRef SmallDepthTexture,
-		TRDGUniformBufferRef<FSceneTextureUniformParameters> SceneTexturesUniformBuffer,
+		const FSceneTextures& SceneTextures,
 		bool bIsOcclusionTesting);
 
 	bool RenderHzb(
 		FRDGBuilder& GraphBuilder,
-		TRDGUniformBufferRef<FSceneTextureUniformParameters> SceneTexturesUniformBuffer);
+		const FSceneTextures& SceneTextures);
 
 	/** Renders the view family. */
 	virtual void Render(FRHICommandListImmediate& RHICmdList) override;
@@ -400,10 +393,8 @@ private:
 
 	void RenderFog(
 		FRDGBuilder& GraphBuilder,
-		FRDGTextureRef SceneColorTexture,
-		FRDGTextureRef SceneDepthTexture,
-		FRDGTextureRef LightShaftOcclusionTexture,
-		TRDGUniformBufferRef<FSceneTextureUniformParameters> SceneTexturesWithDepth);
+		const FMinimalSceneTextures& SceneTextures,
+		FRDGTextureRef LightShaftOcclusionTexture);
 
 	void RenderUnderWaterFog(
 		FRDGBuilder& GraphBuilder,
@@ -412,10 +403,8 @@ private:
 
 	void RenderAtmosphere(
 		FRDGBuilder& GraphBuilder,
-		FRDGTextureRef SceneColorTexture,
-		FRDGTextureRef SceneDepthTexture,
-		FRDGTextureRef LightShaftOcclusionTexture,
-		TRDGUniformBufferRef<FSceneTextureUniformParameters> SceneTexturesWithDepth);
+		const FMinimalSceneTextures& SceneTextures,
+		FRDGTextureRef LightShaftOcclusionTexture);
 
 	/** Render debug visualizations about the sky atmosphere into the scene render target.*/
 	void RenderDebugSkyAtmosphere(FRDGBuilder& GraphBuilder, FRDGTextureRef SceneColorTexture, FRDGTextureRef SceneDepthTexture);
@@ -438,18 +427,15 @@ private:
 	/** Render diffuse indirect (regardless of the method) of the views into the scene color. */
 	void RenderDiffuseIndirectAndAmbientOcclusion(
 		FRDGBuilder& GraphBuilder,
-		TRDGUniformBufferRef<FSceneTextureUniformParameters> SceneTexturesUniformBuffer,
-		FRDGTextureRef SceneColorTexture,
+		const FMinimalSceneTextures& SceneTextures,
 		FHairStrandsRenderingData* HairDatas,
 		bool bIsVisualizePass);
 
 	/** Renders sky lighting and reflections that can be done in a deferred pass. */
 	void RenderDeferredReflectionsAndSkyLighting(
 		FRDGBuilder& GraphBuilder,
-		TRDGUniformBufferRef<FSceneTextureUniformParameters> SceneTexturesUniformBuffer,
-		FRDGTextureMSAA SceneColorTexture,
+		const FSceneTextures& SceneTextures,
 		FRDGTextureRef DynamicBentNormalAOTexture,
-		FRDGTextureRef VelocityTexture,
 		struct FHairStrandsRenderingData* HairDatas);
 
 	void RenderDeferredReflectionsAndSkyLightingHair(FRDGBuilder& GraphBuilder, struct FHairStrandsRenderingData* HairDatas);
@@ -457,9 +443,7 @@ private:
 	/** Computes DFAO, modulates it to scene color (which is assumed to contain diffuse indirect lighting), and stores the output bent normal for use occluding specular. */
 	void RenderDFAOAsIndirectShadowing(
 		FRDGBuilder& GraphBuilder,
-		TRDGUniformBufferRef<FSceneTextureUniformParameters> SceneTexturesUniformBuffer,
-		FRDGTextureRef SceneColorTexture,
-		FRDGTextureRef VelocityTexture,
+		const FSceneTextures& SceneTextures,
 		FRDGTextureRef& DynamicBentNormalAO);
 
 	bool ShouldRenderDistanceFieldLighting() const;
@@ -467,10 +451,8 @@ private:
 	/** Render Ambient Occlusion using mesh distance fields and the surface cache, which supports dynamic rigid meshes. */
 	void RenderDistanceFieldLighting(
 		FRDGBuilder& GraphBuilder,
-		TRDGUniformBufferRef<FSceneTextureUniformParameters> SceneTexturesUniformBuffer,
+		const FSceneTextures& SceneTextures,
 		const class FDistanceFieldAOParameters& Parameters,
-		FRDGTextureRef SceneColorTexture,
-		FRDGTextureRef VelocityTexture,
 		FRDGTextureRef& OutDynamicBentNormalAO,
 		bool bModulateToSceneColor,
 		bool bVisualizeAmbientOcclusion);
@@ -478,18 +460,15 @@ private:
 	/** Render Ambient Occlusion using mesh distance fields on a screen based grid. */
 	void RenderDistanceFieldAOScreenGrid(
 		FRDGBuilder& GraphBuilder,
+		const FSceneTextures& SceneTextures,
 		const FViewInfo& View,
 		const FDistanceFieldAOParameters& Parameters,
-		TRDGUniformBufferRef<FSceneTextureUniformParameters> SceneTexturesUniformBuffer,
-		FRDGTextureRef VelocityTexture,
 		FRDGTextureRef DistanceFieldNormal,
 		FRDGTextureRef& OutDynamicBentNormalAO);
 
 	void RenderMeshDistanceFieldVisualization(
 		FRDGBuilder& GraphBuilder,
-		FRDGTextureRef SceneColorTexture,
-		FRDGTextureRef SceneDepthTexture,
-		TRDGUniformBufferRef<FSceneTextureUniformParameters> SceneTexturesUniformBuffer,
+		const FMinimalSceneTextures& SceneTextures,
 		const FDistanceFieldAOParameters& Parameters);
 
 	bool ShouldRenderLumenDiffuseGI(const FViewInfo& View) const;
@@ -570,27 +549,23 @@ private:
 	/** Add a clustered deferred shading lighting render pass.	Note: in the future it should take the RenderGraph builder as argument */
 	void AddClusteredDeferredShadingPass(
 		FRDGBuilder& GraphBuilder,
-		FRDGTextureRef SceneColorTexture,
-		TRDGUniformBufferRef<FSceneTextureUniformParameters> SceneTexturesUniformBuffer,
+		const FMinimalSceneTextures& SceneTextures,
 		const FSortedLightSetSceneInfo& SortedLightsSet);
 
 	/** Renders the lights in SortedLights in the range [TiledDeferredLightsStart, TiledDeferredLightsEnd) using tiled deferred shading. */
-	FRDGTextureRef RenderTiledDeferredLighting(
+	void RenderTiledDeferredLighting(
 		FRDGBuilder& GraphBuilder,
-		FRDGTextureRef InSceneColorTexture,
-		TRDGUniformBufferRef<FSceneTextureUniformParameters> SceneTexturesUniformBuffer,
+		FMinimalSceneTextures& SceneTextures,
 		const TArray<FSortedLightSceneInfo, SceneRenderingAllocator>& SortedLights,
 		int32 TiledDeferredLightsStart,
 		int32 TiledDeferredLightsEnd,
 		const FSimpleLightArray& SimpleLights);
 
 	/** Renders the scene's lighting. */
-	FRDGTextureRef RenderLights(
+	void RenderLights(
 		FRDGBuilder& GraphBuilder,
-		TRDGUniformBufferRef<FSceneTextureUniformParameters> SceneTexturesUniformBuffer,
+		FMinimalSceneTextures& SceneTextures,
 		const FTranslucentVolumeLightingTextures& TranslucentVolumeLightingTextures,
-		FRDGTextureRef SceneColorTexture,
-		FRDGTextureRef SceneDepthTexture,
 		FRDGTextureRef LightingChannelsTexture,
 		FSortedLightSetSceneInfo& SortedLightSet,
 		const FHairStrandsRenderingData* HairDatas);
@@ -604,26 +579,22 @@ private:
 	/** Render stationary light overlap as complexity to scene color. */
 	void RenderStationaryLightOverlap(
 		FRDGBuilder& GraphBuilder,
-		FRDGTextureRef SceneColorTexture,
-		FRDGTextureRef SceneDepthTexture,
-		FRDGTextureRef LightingChannelsTexture,
-		TRDGUniformBufferRef<FSceneTextureUniformParameters> SceneTexturesUniformBuffer);
+		const FMinimalSceneTextures& SceneTextures,
+		FRDGTextureRef LightingChannelsTexture);
 	
 	/** Renders the scene's translucency passes. */
 	void RenderTranslucency(
 		FRDGBuilder& GraphBuilder,
+		const FMinimalSceneTextures& SceneTextures,
 		const FTranslucentVolumeLightingTextures& TranslucentVolumeLightingTextures,
-		FRDGTextureMSAA SceneColorTexture,
-		FRDGTextureMSAA SceneDepthTexture,
 		FSeparateTranslucencyTextures* OutSeparateTranslucencyTextures,
 		ETranslucencyView ViewsToRender);
 
 	/** Renders the scene's translucency given a specific pass. */
 	void RenderTranslucencyInner(
 		FRDGBuilder& GraphBuilder,
+		const FMinimalSceneTextures& SceneTextures,
 		const FTranslucentVolumeLightingTextures& TranslucentVolumeLightingTextures,
-		FRDGTextureMSAA SceneColorTexture,
-		FRDGTextureMSAA SceneDepthTexture,
 		FSeparateTranslucencyTextures* OutSeparateTranslucencyTextures,
 		ETranslucencyView ViewsToRender,
 		FRDGTextureRef SceneColorCopyTexture,
@@ -632,23 +603,18 @@ private:
 	/** Renders the scene's light shafts */
 	FRDGTextureRef RenderLightShaftOcclusion(
 		FRDGBuilder& GraphBuilder,
-		TRDGUniformBufferRef<FSceneTextureUniformParameters> SceneTexturesUniformBuffer,
-		FIntPoint SceneTextureExtent);
+		const FMinimalSceneTextures& SceneTextures);
 
 	void RenderLightShaftBloom(
 		FRDGBuilder& GraphBuilder,
-		TRDGUniformBufferRef<FSceneTextureUniformParameters> SceneTexturesUniformBuffer,
-		FIntPoint SceneTextureExtent,
-		FRDGTextureRef SceneColorTexture,
+		const FMinimalSceneTextures& SceneTextures,
 		FSeparateTranslucencyTextures& OutSeparateTranslucencyTextures);
 
 	bool ShouldRenderVelocities() const;
 
 	void RenderVelocities(
 		FRDGBuilder& GraphBuilder,
-		FRDGTextureRef SceneDepthTexture,
-		FRDGTextureRef& VelocityTexture,
-		TRDGUniformBufferRef<FSceneTextureUniformParameters> SceneTexturesUniformBuffer,
+		FSceneTextures& SceneTextures,
 		EVelocityPass VelocityPass,
 		bool bForceVelocity);
 
@@ -689,29 +655,26 @@ private:
 
 	void RenderShadowProjections(
 		FRDGBuilder& GraphBuilder,
-		TRDGUniformBufferRef<FSceneTextureUniformParameters> SceneTexturesUniformBuffer,
+		const FMinimalSceneTextures& SceneTextures,
 		FRDGTextureRef ScreenShadowMaskTexture,
 		FRDGTextureRef ScreenShadowMaskSubPixelTexture,
-		FRDGTextureRef SceneDepthTexture,
 		const FLightSceneInfo* LightSceneInfo,
 		const FHairStrandsVisibilityViews* HairVisibilityViews,
 		bool bProjectingForForwardShading);
 
 	void RenderDeferredShadowProjections(
 		FRDGBuilder& GraphBuilder,
-		TRDGUniformBufferRef<FSceneTextureUniformParameters> SceneTexturesUniformBuffer,
+		const FMinimalSceneTextures& SceneTextures,
 		const FTranslucentVolumeLightingTextures& TranslucentVolumeLightingTextures,
 		const FLightSceneInfo* LightSceneInfo,
 		FRDGTextureRef ScreenShadowMaskTexture,
 		FRDGTextureRef ScreenShadowMaskSubPixelTexture,
-		FRDGTextureRef SceneDepthTexture,
 		const FHairStrandsRenderingData* HairDatas,
 		bool& bInjectedTranslucentVolume);
 
 	void RenderForwardShadowProjections(
 		FRDGBuilder& GraphBuilder,
-		TRDGUniformBufferRef<FSceneTextureUniformParameters> SceneTexturesUniformBuffer,
-		FRDGTextureRef SceneDepthTexture,
+		const FMinimalSceneTextures& SceneTextures,
 		FRDGTextureRef& ForwardScreenSpaceShadowMask,
 		FRDGTextureRef& ForwardScreenSpaceShadowMaskSubPixel,
 		const FHairStrandsRenderingData* InHairDatas);
@@ -719,8 +682,7 @@ private:
 	/** Used by RenderLights to render a light function to the attenuation buffer. */
 	bool RenderLightFunction(
 		FRDGBuilder& GraphBuilder,
-		FRDGTextureRef SceneDepthTexture,
-		TRDGUniformBufferRef<FSceneTextureUniformParameters> SceneTexturesUniformBuffer,
+		const FMinimalSceneTextures& SceneTextures,
 		const FLightSceneInfo* LightSceneInfo,
 		FRDGTextureRef ScreenShadowMaskTexture,
 		bool bLightAttenuationCleared,
@@ -729,8 +691,7 @@ private:
 	/** Renders a light function indicating that whole scene shadowing being displayed is for previewing only, and will go away in game. */
 	bool RenderPreviewShadowsIndicator(
 		FRDGBuilder& GraphBuilder,
-		FRDGTextureRef SceneDepthTexture,
-		TRDGUniformBufferRef<FSceneTextureUniformParameters> SceneTexturesUniformBuffer,
+		const FMinimalSceneTextures& SceneTextures,
 		const FLightSceneInfo* LightSceneInfo,
 		FRDGTextureRef ScreenShadowMaskTexture,
 		bool bLightAttenuationCleared);
@@ -738,8 +699,7 @@ private:
 	/** Renders a light function with the given material. */
 	bool RenderLightFunctionForMaterial(
 		FRDGBuilder& GraphBuilder,
-		FRDGTextureRef SceneDepthTexture,
-		TRDGUniformBufferRef<FSceneTextureUniformParameters> SceneTexturesUniformBuffer,
+		const FMinimalSceneTextures& SceneTextures,
 		const FLightSceneInfo* LightSceneInfo,
 		FRDGTextureRef ScreenShadowMaskTexture,
 		const FMaterialRenderProxy* MaterialProxy,
@@ -759,9 +719,7 @@ private:
 
 	void RenderLight(
 		FRDGBuilder& GraphBuilder,
-		FRDGTextureRef SceneColorTexture,
-		FRDGTextureRef SceneDepthTexture,
-		TRDGUniformBufferRef<FSceneTextureUniformParameters> SceneTexturesUniformBuffer,
+		const FMinimalSceneTextures& SceneTextures,
 		const FLightSceneInfo* LightSceneInfo,
 		FRDGTextureRef ScreenShadowMaskTexture,
 		FRDGTextureRef LightingChannelsTexture,
@@ -789,9 +747,7 @@ private:
 	/** Renders an array of simple lights using standard deferred shading. */
 	void RenderSimpleLightsStandardDeferred(
 		FRDGBuilder& GraphBuilder,
-		FRDGTextureRef SceneColorTexture,
-		FRDGTextureRef SceneDepthTexture,
-		TRDGUniformBufferRef<FSceneTextureUniformParameters> SceneTexturesUniformBuffer,
+		const FMinimalSceneTextures& SceneTextures,
 		const FSimpleLightArray& SimpleLights);
 
 	/** Initializes the translucency lighting volumes. Can be compute or async compute. */
@@ -875,10 +831,7 @@ private:
 
 	void ComputeVolumetricFog(FRDGBuilder& GraphBuilder);
 
-	void VisualizeVolumetricLightmap(
-		FRDGBuilder& GraphBuilder,
-		FRDGTextureRef SceneColorTexture,
-		FRDGTextureRef SceneDepthTexture);
+	void VisualizeVolumetricLightmap(FRDGBuilder& GraphBuilder, const FMinimalSceneTextures& SceneTextures);
 
 	/** Render image based reflections (SSR, Env, SkyLight) without compute shaders */
 	void RenderStandardDeferredImageBasedReflections(FRHICommandListImmediate& RHICmdList, FGraphicsPipelineStateInitializer& GraphicsPSOInit, bool bReflectionEnv, const TRefCountPtr<IPooledRenderTarget>& DynamicBentNormalAO, TRefCountPtr<IPooledRenderTarget>& VelocityRT);
@@ -945,8 +898,7 @@ private:
 
 	void CompositeRayTracingSkyLight(
 		FRDGBuilder& GraphBuilder,
-		FRDGTextureRef SceneColorTexture,
-		FIntPoint SceneTextureExtent,
+		const FMinimalSceneTextures& SceneTextures,
 		FRDGTextureRef SkyLightRT,
 		FRDGTextureRef HitDistanceRT);
 	
diff --git a/Engine/Source/Runtime/Renderer/Private/DistanceFieldAmbientOcclusion.cpp b/Engine/Source/Runtime/Renderer/Private/DistanceFieldAmbientOcclusion.cpp
--- a/Engine/Source/Runtime/Renderer/Private/DistanceFieldAmbientOcclusion.cpp
+++ b/Engine/Source/Runtime/Renderer/Private/DistanceFieldAmbientOcclusion.cpp
@@ -720,16 +720,14 @@ bool FDeferredShadingSceneRenderer::ShouldPrepareGlobalDistanceField() const
 
 void FDeferredShadingSceneRenderer::RenderDFAOAsIndirectShadowing(
 	FRDGBuilder& GraphBuilder,
-	TRDGUniformBufferRef<FSceneTextureUniformParameters> SceneTexturesUniformBuffer,
-	FRDGTextureRef SceneColorTexture,
-	FRDGTextureRef VelocityTexture,
+	const FSceneTextures& SceneTextures,
 	FRDGTextureRef& DynamicBentNormalAO)
 {
 	if (GDistanceFieldAOApplyToStaticIndirect && ShouldRenderDistanceFieldAO() && ShouldRenderDistanceFieldLighting())
 	{
 		// Use the skylight's max distance if there is one, to be consistent with DFAO shadowing on the skylight
 		const float OcclusionMaxDistance = Scene->SkyLight && !Scene->SkyLight->bWantsStaticShadowing ? Scene->SkyLight->OcclusionMaxDistance : Scene->DefaultMaxDistanceFieldOcclusionDistance;
-		RenderDistanceFieldLighting(GraphBuilder, SceneTexturesUniformBuffer, FDistanceFieldAOParameters(OcclusionMaxDistance), SceneColorTexture, VelocityTexture, DynamicBentNormalAO, true, false);
+		RenderDistanceFieldLighting(GraphBuilder, SceneTextures, FDistanceFieldAOParameters(OcclusionMaxDistance), DynamicBentNormalAO, true, false);
 	}
 }
 
@@ -749,10 +747,8 @@ bool FDeferredShadingSceneRenderer::ShouldRenderDistanceFieldLighting() const
 
 void FDeferredShadingSceneRenderer::RenderDistanceFieldLighting(
 	FRDGBuilder& GraphBuilder,
-	TRDGUniformBufferRef<FSceneTextureUniformParameters> SceneTexturesUniformBuffer,
+	const FSceneTextures& SceneTextures,
 	const FDistanceFieldAOParameters& Parameters,
-	FRDGTextureRef SceneColorTexture,
-	FRDGTextureRef VelocityTexture,
 	FRDGTextureRef& OutDynamicBentNormalAO,
 	bool bModulateToSceneColor,
 	bool bVisualizeAmbientOcclusion)
@@ -795,7 +791,7 @@ void FDeferredShadingSceneRenderer::RenderDistanceFieldLighting(
 		CullObjectsToView(GraphBuilder, Scene, View, Parameters, GAOCulledObjectBuffers);
 	}
 
-	ComputeDistanceFieldNormal(GraphBuilder, Views, SceneTexturesUniformBuffer, DistanceFieldNormal, Parameters);
+	ComputeDistanceFieldNormal(GraphBuilder, Views, SceneTextures.UniformBuffer, DistanceFieldNormal, Parameters);
 
 	// Intersect objects with screen tiles, build lists
 	if (UseAOObjectDistanceField())
@@ -807,10 +803,9 @@ void FDeferredShadingSceneRenderer::RenderDistanceFieldLighting(
 
 	RenderDistanceFieldAOScreenGrid(
 		GraphBuilder,
+		SceneTextures,
 		View,
 		Parameters,
-		SceneTexturesUniformBuffer,
-		VelocityTexture,
 		DistanceFieldNormal,
 		BentNormalOutput);
 
@@ -825,12 +820,12 @@ void FDeferredShadingSceneRenderer::RenderDistanceFieldLighting(
 		});
 	}
 
-	RenderCapsuleShadowsForMovableSkylight(GraphBuilder, SceneTexturesUniformBuffer, BentNormalOutput);
+	RenderCapsuleShadowsForMovableSkylight(GraphBuilder, SceneTextures.UniformBuffer, BentNormalOutput);
 
 	// Upsample to full resolution, write to output in case of debug AO visualization or scene color modulation (standard upsampling is done later together with sky lighting and reflection environment)
 	if (bModulateToSceneColor || bVisualizeAmbientOcclusion)
 	{
-		UpsampleBentNormalAO(GraphBuilder, Views, SceneTexturesUniformBuffer, SceneColorTexture, BentNormalOutput, bModulateToSceneColor && !bVisualizeAmbientOcclusion);
+		UpsampleBentNormalAO(GraphBuilder, Views, SceneTextures.UniformBuffer, SceneTextures.Color.Target, BentNormalOutput, bModulateToSceneColor && !bVisualizeAmbientOcclusion);
 	}
 
 	OutDynamicBentNormalAO = BentNormalOutput;
diff --git a/Engine/Source/Runtime/Renderer/Private/DistanceFieldScreenGridLighting.cpp b/Engine/Source/Runtime/Renderer/Private/DistanceFieldScreenGridLighting.cpp
--- a/Engine/Source/Runtime/Renderer/Private/DistanceFieldScreenGridLighting.cpp
+++ b/Engine/Source/Runtime/Renderer/Private/DistanceFieldScreenGridLighting.cpp
@@ -484,10 +484,9 @@ END_SHADER_PARAMETER_STRUCT()
 
 void FDeferredShadingSceneRenderer::RenderDistanceFieldAOScreenGrid(
 	FRDGBuilder& GraphBuilder,
+	const FSceneTextures& SceneTextures,
 	const FViewInfo& View,
 	const FDistanceFieldAOParameters& Parameters,
-	TRDGUniformBufferRef<FSceneTextureUniformParameters> SceneTexturesUniformBuffer,
-	FRDGTextureRef VelocityTexture,
 	FRDGTextureRef DistanceFieldNormal,
 	FRDGTextureRef& OutDynamicBentNormalAO)
 {
@@ -520,7 +519,7 @@ void FDeferredShadingSceneRenderer::RenderDistanceFieldAOScreenGrid(
 
 	{
 		FDistanceFieldAOScreenGridParameters* PassParameters = GraphBuilder.AllocParameters<FDistanceFieldAOScreenGridParameters>();
-		PassParameters->SceneTextures = SceneTexturesUniformBuffer;
+		PassParameters->SceneTextures = SceneTextures.UniformBuffer;
 		PassParameters->DistanceFieldNormal = DistanceFieldNormal;
 
 		GraphBuilder.AddPass(
@@ -648,8 +647,8 @@ void FDeferredShadingSceneRenderer::RenderDistanceFieldAOScreenGrid(
 		GraphBuilder,
 		Parameters,
 		View,
-		SceneTexturesUniformBuffer,
-		VelocityTexture,
+		SceneTextures.UniformBuffer,
+		SceneTextures.Velocity,
 		DownsampledBentNormal,
 		DistanceFieldNormal,
 		OutDynamicBentNormalAO);
diff --git a/Engine/Source/Runtime/Renderer/Private/DistanceFieldShadowing.cpp b/Engine/Source/Runtime/Renderer/Private/DistanceFieldShadowing.cpp
--- a/Engine/Source/Runtime/Renderer/Private/DistanceFieldShadowing.cpp
+++ b/Engine/Source/Runtime/Renderer/Private/DistanceFieldShadowing.cpp
@@ -927,25 +927,24 @@ END_SHADER_PARAMETER_STRUCT()
 
 void FProjectedShadowInfo::RenderRayTracedDistanceFieldProjection(
 	FRDGBuilder& GraphBuilder,
-	TRDGUniformBufferRef<FSceneTextureUniformParameters> SceneTexturesUniformBuffer,
+	const FMinimalSceneTextures& SceneTextures,
 	FRDGTextureRef ScreenShadowMaskTexture,
-	FRDGTextureRef SceneDepthTexture,
 	const FViewInfo& View,
 	FIntRect ScissorRect,
 	bool bProjectingForForwardShading) const
 {
 	check(ScissorRect.Area() > 0);
 
-	FRDGTextureRef RayTracedShadowsTexture = BeginRenderRayTracedDistanceFieldProjection(GraphBuilder, SceneTexturesUniformBuffer, View);
+	FRDGTextureRef RayTracedShadowsTexture = BeginRenderRayTracedDistanceFieldProjection(GraphBuilder, SceneTextures.UniformBuffer, View);
 
 	if (RayTracedShadowsTexture)
 	{
 		FDistanceFieldShadowingUpsample* PassParameters = GraphBuilder.AllocParameters<FDistanceFieldShadowingUpsample>();
 		PassParameters->RenderTargets[0] = FRenderTargetBinding(ScreenShadowMaskTexture, ERenderTargetLoadAction::ELoad);
-		PassParameters->RenderTargets.DepthStencil = FDepthStencilBinding(SceneDepthTexture, ERenderTargetLoadAction::ELoad, ERenderTargetLoadAction::ELoad, FExclusiveDepthStencil::DepthRead_StencilRead);
+		PassParameters->RenderTargets.DepthStencil = FDepthStencilBinding(SceneTextures.Depth.Target, ERenderTargetLoadAction::ELoad, ERenderTargetLoadAction::ELoad, FExclusiveDepthStencil::DepthRead_StencilRead);
 		
 		PassParameters->PS.View = View.ViewUniformBuffer;
-		PassParameters->PS.SceneTextures = SceneTexturesUniformBuffer;
+		PassParameters->PS.SceneTextures = SceneTextures.UniformBuffer;
 		PassParameters->PS.ShadowFactorsTexture = RayTracedShadowsTexture;
 		PassParameters->PS.ShadowFactorsSampler = TStaticSamplerState<SF_Bilinear>::GetRHI();
 		PassParameters->PS.ScissorRectMinAndSize = FIntRect(ScissorRect.Min, ScissorRect.Size());
diff --git a/Engine/Source/Runtime/Renderer/Private/DistanceFieldVisualization.cpp b/Engine/Source/Runtime/Renderer/Private/DistanceFieldVisualization.cpp
--- a/Engine/Source/Runtime/Renderer/Private/DistanceFieldVisualization.cpp
+++ b/Engine/Source/Runtime/Renderer/Private/DistanceFieldVisualization.cpp
@@ -117,9 +117,7 @@ IMPLEMENT_GLOBAL_SHADER(FVisualizeDistanceFieldUpsamplePS, "/Engine/Private/Dist
 
 void FDeferredShadingSceneRenderer::RenderMeshDistanceFieldVisualization(
 	FRDGBuilder& GraphBuilder,
-	FRDGTextureRef SceneColorTexture,
-	FRDGTextureRef SceneDepthTexture,
-	TRDGUniformBufferRef<FSceneTextureUniformParameters> SceneTexturesUniformBuffer,
+	const FMinimalSceneTextures& SceneTextures,
 	const FDistanceFieldAOParameters& Parameters)
 {
 	const FViewInfo& FirstView = Views[0];
@@ -166,7 +164,7 @@ void FDeferredShadingSceneRenderer::RenderMeshDistanceFieldVisualization(
 		auto* PassParameters = GraphBuilder.AllocParameters<FVisualizeMeshDistanceFieldCS::FParameters>();
 		PassParameters->View = View.ViewUniformBuffer;
 		PassParameters->NumGroups = FVector2D(GroupSizeX, GroupSizeY);
-		PassParameters->SceneTextures = SceneTexturesUniformBuffer;
+		PassParameters->SceneTextures = SceneTextures.UniformBuffer;
 		PassParameters->RWVisualizeMeshDistanceFields = GraphBuilder.CreateUAV(VisualizeResultTexture);
 
 		check(!bUseGlobalDistanceField || View.GlobalDistanceFieldInfo.Clipmaps.Num() > 0);
@@ -202,11 +200,11 @@ void FDeferredShadingSceneRenderer::RenderMeshDistanceFieldVisualization(
 		PassParameters->View = View.ViewUniformBuffer;
 		PassParameters->VisualizeDistanceFieldTexture = VisualizeResultTexture;
 		PassParameters->VisualizeDistanceFieldSampler = TStaticSamplerState<SF_Bilinear>::GetRHI();
-		PassParameters->RenderTargets[0] = FRenderTargetBinding(SceneColorTexture, ERenderTargetLoadAction::ELoad);
-		PassParameters->RenderTargets.DepthStencil = FDepthStencilBinding(SceneDepthTexture, ERenderTargetLoadAction::ELoad, ERenderTargetLoadAction::ELoad, FExclusiveDepthStencil::DepthRead_StencilWrite);
+		PassParameters->RenderTargets[0] = FRenderTargetBinding(SceneTextures.Color.Target, ERenderTargetLoadAction::ELoad);
+		PassParameters->RenderTargets.DepthStencil = FDepthStencilBinding(SceneTextures.Depth.Target, ERenderTargetLoadAction::ELoad, ERenderTargetLoadAction::ELoad, FExclusiveDepthStencil::DepthRead_StencilWrite);
 
 		const FScreenPassTextureViewport InputViewport(VisualizeResultTexture, GetDownscaledRect(View.ViewRect, GAODownsampleFactor));
-		const FScreenPassTextureViewport OutputViewport(SceneColorTexture, View.ViewRect);
+		const FScreenPassTextureViewport OutputViewport(SceneTextures.Color.Target, View.ViewRect);
 
 		AddDrawScreenPass(GraphBuilder, {}, View, OutputViewport, InputViewport, PixelShader, PassParameters);
 	}
diff --git a/Engine/Source/Runtime/Renderer/Private/FogRendering.cpp b/Engine/Source/Runtime/Renderer/Private/FogRendering.cpp
--- a/Engine/Source/Runtime/Renderer/Private/FogRendering.cpp
+++ b/Engine/Source/Runtime/Renderer/Private/FogRendering.cpp
@@ -494,10 +494,8 @@ END_SHADER_PARAMETER_STRUCT()
 
 void FDeferredShadingSceneRenderer::RenderFog(
 	FRDGBuilder& GraphBuilder,
-	FRDGTextureRef SceneColorTexture,
-	FRDGTextureRef SceneDepthTexture,
-	FRDGTextureRef LightShaftOcclusionTexture,
-	TRDGUniformBufferRef<FSceneTextureUniformParameters> SceneTexturesWithDepth)
+	const FMinimalSceneTextures& SceneTextures,
+	FRDGTextureRef LightShaftOcclusionTexture)
 {
 	if (Scene->ExponentialFogs.Num() > 0 
 		// Fog must be done in the base pass for MSAA to work
@@ -518,11 +516,11 @@ void FDeferredShadingSceneRenderer::RenderFog(
 				TRDGUniformBufferRef<FFogUniformParameters> FogUniformBuffer = CreateFogUniformBuffer(GraphBuilder, View);
 
 				FFogPassParameters* PassParameters = GraphBuilder.AllocParameters<FFogPassParameters>();
-				PassParameters->SceneTextures = SceneTexturesWithDepth;
+				PassParameters->SceneTextures = SceneTextures.UniformBuffer;
 				PassParameters->Fog = FogUniformBuffer;
 				PassParameters->LightShaftOcclusionTexture = LightShaftOcclusionTexture;
-				PassParameters->RenderTargets[0] = FRenderTargetBinding(SceneColorTexture, ERenderTargetLoadAction::ELoad);
-				PassParameters->RenderTargets.DepthStencil = FDepthStencilBinding(SceneDepthTexture, ERenderTargetLoadAction::ELoad, ERenderTargetLoadAction::ELoad, FExclusiveDepthStencil::DepthRead_StencilWrite);
+				PassParameters->RenderTargets[0] = FRenderTargetBinding(SceneTextures.Color.Target, ERenderTargetLoadAction::ELoad);
+				PassParameters->RenderTargets.DepthStencil = FDepthStencilBinding(SceneTextures.Depth.Target, ERenderTargetLoadAction::ELoad, ERenderTargetLoadAction::ELoad, FExclusiveDepthStencil::DepthRead_StencilWrite);
 
 				GraphBuilder.AddPass({}, PassParameters, ERDGPassFlags::Raster, [this, &View, bShouldRenderVolumetricFog, LightShaftOcclusionTexture, FogUniformBuffer](FRHICommandList& RHICmdList)
 				{
diff --git a/Engine/Source/Runtime/Renderer/Private/IndirectLightRendering.cpp b/Engine/Source/Runtime/Renderer/Private/IndirectLightRendering.cpp
--- a/Engine/Source/Runtime/Renderer/Private/IndirectLightRendering.cpp
+++ b/Engine/Source/Runtime/Renderer/Private/IndirectLightRendering.cpp
@@ -687,8 +687,7 @@ void FDeferredShadingSceneRenderer::SetupCommonDiffuseIndirectParameters(
 
 void FDeferredShadingSceneRenderer::RenderDiffuseIndirectAndAmbientOcclusion(
 	FRDGBuilder& GraphBuilder,
-	TRDGUniformBufferRef<FSceneTextureUniformParameters> SceneTexturesUniformBuffer,
-	FRDGTextureRef SceneColorTexture,
+	const FMinimalSceneTextures& SceneTextures,
 	FHairStrandsRenderingData* InHairDatas,
 	bool bIsVisualizePass)
 {
@@ -703,8 +702,9 @@ void FDeferredShadingSceneRenderer::RenderDiffuseIndirectAndAmbientOcclusion(
 
 	FSceneRenderTargets& SceneContext = FSceneRenderTargets::Get();
 
-	FSceneTextureParameters SceneTextures = GetSceneTextureParameters(GraphBuilder, SceneTexturesUniformBuffer);
-	
+	FSceneTextureParameters SceneTextureParameters = GetSceneTextureParameters(GraphBuilder, SceneTextures.UniformBuffer);
+	FRDGTextureRef SceneColorTexture = SceneTextures.Color.Target;
+
 	for (FViewInfo& View : Views)
 	{
 		RDG_GPU_MASK_SCOPE(GraphBuilder, View.GPUMask);
@@ -715,7 +715,7 @@ void FDeferredShadingSceneRenderer::RenderDiffuseIndirectAndAmbientOcclusion(
 
 		// Setup the common diffuse parameter for this view.
 		FCommonParameters CommonDiffuseParameters;
-		SetupCommonDiffuseIndirectParameters(GraphBuilder, SceneTextures, View, /* out */ CommonDiffuseParameters);
+		SetupCommonDiffuseIndirectParameters(GraphBuilder, SceneTextureParameters, View, /* out */ CommonDiffuseParameters);
 
 		// Update old ray tracing config for the denoiser.
 		IScreenSpaceDenoiser::FAmbientOcclusionRayTracingConfig RayTracingConfig;
@@ -727,7 +727,7 @@ void FDeferredShadingSceneRenderer::RenderDiffuseIndirectAndAmbientOcclusion(
 		ScreenSpaceRayTracing::FPrevSceneColorMip PrevSceneColorMip;
 		if (ViewPipelineState.bEnableSSGI)
 		{
-			PrevSceneColorMip = ScreenSpaceRayTracing::ReducePrevSceneColorMip(GraphBuilder, SceneTextures, View);
+			PrevSceneColorMip = ScreenSpaceRayTracing::ReducePrevSceneColorMip(GraphBuilder, SceneTextureParameters, View);
 		}
 
 		FSSDSignalTextures DenoiserOutputs;
@@ -757,7 +757,7 @@ void FDeferredShadingSceneRenderer::RenderDiffuseIndirectAndAmbientOcclusion(
 		{
 			// TODO: Refactor under the HybridIndirectLighting standard API.
 			// TODO: hybrid SSGI / RTGI
-			RenderRayTracingGlobalIllumination(GraphBuilder, SceneTextures, View, /* out */ &RayTracingConfig, /* out */ &DenoiserInputs);
+			RenderRayTracingGlobalIllumination(GraphBuilder, SceneTextureParameters, View, /* out */ &RayTracingConfig, /* out */ &DenoiserInputs);
 		}
 		else if (ViewPipelineState.DiffuseIndirectMethod == EDiffuseIndirectMethod::Lumen)
 		{
@@ -767,7 +767,7 @@ void FDeferredShadingSceneRenderer::RenderDiffuseIndirectAndAmbientOcclusion(
 
 			DenoiserOutputs = RenderLumenScreenProbeGather(
 				GraphBuilder, 
-				SceneTexturesUniformBuffer,
+				SceneTextures.UniformBuffer,
 				PrevSceneColorMip, 
 				View,
 				&View.PrevViewInfo,
@@ -821,7 +821,7 @@ void FDeferredShadingSceneRenderer::RenderDiffuseIndirectAndAmbientOcclusion(
 					GraphBuilder,
 					View,
 					&View.PrevViewInfo,
-					SceneTextures,
+					SceneTextureParameters,
 					DenoiserInputs,
 					RayTracingConfig);
 
@@ -833,7 +833,7 @@ void FDeferredShadingSceneRenderer::RenderDiffuseIndirectAndAmbientOcclusion(
 					GraphBuilder,
 					View,
 					&View.PrevViewInfo,
-					SceneTextures,
+					SceneTextureParameters,
 					DenoiserInputs,
 					RayTracingConfig);
 
@@ -855,7 +855,7 @@ void FDeferredShadingSceneRenderer::RenderDiffuseIndirectAndAmbientOcclusion(
 			RenderRayTracingAmbientOcclusion(
 				GraphBuilder,
 				View,
-				SceneTextures,
+				SceneTextureParameters,
 				&AmbientOcclusionMask);
 		}
 		else if (ViewPipelineState.AmbientOcclusionMethod == EAmbientOcclusionMethod::SSGI)
@@ -917,7 +917,7 @@ void FDeferredShadingSceneRenderer::RenderDiffuseIndirectAndAmbientOcclusion(
 				PassParameters->ApplyAOToDynamicDiffuseIndirect = 1.0f;
 			}
 
-			const FIntPoint BufferExtent = SceneTextures.SceneDepthTexture->Desc.Extent;
+			const FIntPoint BufferExtent = SceneTextureParameters.SceneDepthTexture->Desc.Extent;
 
 			{
 				// Placeholder texture for textures pulled in from SSDCommon.ush
@@ -953,11 +953,11 @@ void FDeferredShadingSceneRenderer::RenderDiffuseIndirectAndAmbientOcclusion(
 			}
 
 			Denoiser::SetupCommonShaderParameters(
-				View, SceneTextures,
+				View, SceneTextureParameters,
 				View.ViewRect,
 				1.0f / CommonDiffuseParameters.DownscaleFactor,
 				/* out */ &PassParameters->DenoiserCommonParameters);
-			PassParameters->SceneTextures = SceneTextures;
+			PassParameters->SceneTextures = SceneTextureParameters;
 			PassParameters->ViewUniformBuffer = View.ViewUniformBuffer;
 
 			PassParameters->RenderTargets[0] = FRenderTargetBinding(
@@ -1047,7 +1047,7 @@ void FDeferredShadingSceneRenderer::RenderDiffuseIndirectAndAmbientOcclusion(
 				PassParameters->AmbientOcclusionTexture = GraphBuilder.RegisterExternalTexture(GSystemTextures.WhiteDummy);
 			}
 
-			PassParameters->SceneTextures = SceneTextures;
+			PassParameters->SceneTextures = SceneTextureParameters;
 			PassParameters->ViewUniformBuffer = View.ViewUniformBuffer;
 
 			PassParameters->RenderTargets[0] = FRenderTargetBinding(
@@ -1106,10 +1106,8 @@ void FDeferredShadingSceneRenderer::RenderDiffuseIndirectAndAmbientOcclusion(
 
 void FDeferredShadingSceneRenderer::RenderDeferredReflectionsAndSkyLighting(
 	FRDGBuilder& GraphBuilder,
-	TRDGUniformBufferRef<FSceneTextureUniformParameters> SceneTexturesUniformBuffer,
-	FRDGTextureMSAA SceneColorTexture,
+	const FSceneTextures& SceneTextures,
 	FRDGTextureRef DynamicBentNormalAOTexture,
-	FRDGTextureRef VelocityTexture,
 	FHairStrandsRenderingData* HairDatas)
 {
 	if (ViewFamily.EngineShowFlags.VisualizeLightCulling 
@@ -1157,7 +1155,7 @@ void FDeferredShadingSceneRenderer::RenderDeferredReflectionsAndSkyLighting(
 		{
 			bApplySkyShadowing = true;
 			FDistanceFieldAOParameters Parameters(Scene->SkyLight->OcclusionMaxDistance, Scene->SkyLight->Contrast);
-			RenderDistanceFieldLighting(GraphBuilder, SceneTexturesUniformBuffer, Parameters, SceneColorTexture.Target, VelocityTexture, DynamicBentNormalAOTexture, false, false);
+			RenderDistanceFieldLighting(GraphBuilder, SceneTextures, Parameters, DynamicBentNormalAOTexture, false, false);
 		}
 	}
 
@@ -1176,7 +1174,8 @@ void FDeferredShadingSceneRenderer::RenderDeferredReflectionsAndSkyLighting(
 		DynamicBentNormalAO = 0.0f;
 	}
 
-	FSceneTextureParameters SceneTextures = GetSceneTextureParameters(GraphBuilder);
+	FSceneTextureParameters SceneTextureParameters = GetSceneTextureParameters(GraphBuilder);
+	const auto& SceneColorTexture = SceneTextures.Color;
 
 	IScreenSpaceDenoiser::FReflectionsInputs DenoiserInputs;
 	IScreenSpaceDenoiser::FReflectionsRayTracingConfig RayTracingConfig;
@@ -1187,7 +1186,7 @@ void FDeferredShadingSceneRenderer::RenderDeferredReflectionsAndSkyLighting(
 
 	{
 		FRDGTextureDesc Desc = FRDGTextureDesc::Create2D(
-			SceneTextures.SceneDepthTexture->Desc.Extent / UpscaleFactor,
+			SceneTextureParameters.SceneDepthTexture->Desc.Extent / UpscaleFactor,
 			PF_FloatRGBA,
 			FClearValueBinding::None,
 			TexCreate_ShaderResource | TexCreate_RenderTargetable | TexCreate_UAV);
@@ -1251,7 +1250,7 @@ void FDeferredShadingSceneRenderer::RenderDeferredReflectionsAndSkyLighting(
 
 				RenderRayTracingReflections(
 					GraphBuilder,
-					SceneTextures,
+					SceneTextureParameters,
 					View,
 					RayTracingReflectionOptions,
 					&DenoiserInputs);
@@ -1268,7 +1267,7 @@ void FDeferredShadingSceneRenderer::RenderDeferredReflectionsAndSkyLighting(
 				RDG_EVENT_SCOPE(GraphBuilder, "ScreenSpaceReflections(Quality=%d)", int32(SSRQuality));
 
 				ScreenSpaceRayTracing::RenderScreenSpaceReflections(
-					GraphBuilder, SceneTextures, SceneColorTexture.Resolve, View, SSRQuality, bDenoise, &DenoiserInputs);
+					GraphBuilder, SceneTextureParameters, SceneColorTexture.Resolve, View, SSRQuality, bDenoise, &DenoiserInputs);
 			}
 			else
 			{
@@ -1290,7 +1289,7 @@ void FDeferredShadingSceneRenderer::RenderDeferredReflectionsAndSkyLighting(
 					GraphBuilder,
 					View,
 					&View.PrevViewInfo,
-					SceneTextures,
+					SceneTextureParameters,
 					DenoiserInputs,
 					DenoiserConfig);
 
@@ -1301,8 +1300,8 @@ void FDeferredShadingSceneRenderer::RenderDeferredReflectionsAndSkyLighting(
 				check(View.ViewState);
 				FTAAPassParameters TAASettings(View);
 				TAASettings.Pass = ETAAPassConfig::ScreenSpaceReflections;
-				TAASettings.SceneDepthTexture = SceneTextures.SceneDepthTexture;
-				TAASettings.SceneVelocityTexture = SceneTextures.GBufferVelocityTexture;
+				TAASettings.SceneDepthTexture = SceneTextureParameters.SceneDepthTexture;
+				TAASettings.SceneVelocityTexture = SceneTextureParameters.GBufferVelocityTexture;
 				TAASettings.SceneColorInput = DenoiserInputs.Color;
 				TAASettings.bOutputRenderTargetable = (
 					ViewPipelineState.bComposePlanarReflections ||
@@ -1335,7 +1334,7 @@ void FDeferredShadingSceneRenderer::RenderDeferredReflectionsAndSkyLighting(
 		if (ViewPipelineState.bComposePlanarReflections)
 		{
 			check(!RayTracingReflectionOptions.bEnabled);
-			RenderDeferredPlanarReflections(GraphBuilder, SceneTextures, View, /* inout */ ReflectionsColor);
+			RenderDeferredPlanarReflections(GraphBuilder, SceneTextureParameters, View, /* inout */ ReflectionsColor);
 		}
 
 		bool bRequiresApply = ReflectionsColor != nullptr || bSkyLight || bDynamicSkyLight || bReflectionEnv;
@@ -1424,7 +1423,7 @@ void FDeferredShadingSceneRenderer::RenderDeferredReflectionsAndSkyLighting(
 				PassParameters->PreIntegratedGF = GSystemTextures.PreintegratedGF->GetRenderTargetItem().ShaderResourceTexture;
 				PassParameters->PreIntegratedGFSampler = TStaticSamplerState<SF_Bilinear, AM_Clamp, AM_Clamp, AM_Clamp>::GetRHI();
 
-				PassParameters->SceneTextures = SceneTextures;
+				PassParameters->SceneTextures = SceneTextureParameters;
 
 				PassParameters->ViewUniformBuffer = View.ViewUniformBuffer;
 				PassParameters->ReflectionCaptureData = View.ReflectionCaptureUniformBuffer;
diff --git a/Engine/Source/Runtime/Renderer/Private/LightFunctionRendering.cpp b/Engine/Source/Runtime/Renderer/Private/LightFunctionRendering.cpp
--- a/Engine/Source/Runtime/Renderer/Private/LightFunctionRendering.cpp
+++ b/Engine/Source/Runtime/Renderer/Private/LightFunctionRendering.cpp
@@ -230,8 +230,7 @@ bool FSceneRenderer::CheckForLightFunction( const FLightSceneInfo* LightSceneInf
  */
 bool FDeferredShadingSceneRenderer::RenderLightFunction(
 	FRDGBuilder& GraphBuilder,
-	FRDGTextureRef SceneDepthTexture,
-	TRDGUniformBufferRef<FSceneTextureUniformParameters> SceneTexturesUniformBuffer,
+	const FMinimalSceneTextures& SceneTextures,
 	const FLightSceneInfo* LightSceneInfo,
 	FRDGTextureRef ScreenShadowMaskTexture,
 	bool bLightAttenuationCleared,
@@ -239,7 +238,7 @@ bool FDeferredShadingSceneRenderer::RenderLightFunction(
 {
 	if (ViewFamily.EngineShowFlags.LightFunctions)
 	{
-		return RenderLightFunctionForMaterial(GraphBuilder, SceneDepthTexture, SceneTexturesUniformBuffer, LightSceneInfo, ScreenShadowMaskTexture, LightSceneInfo->Proxy->GetLightFunctionMaterial(), bLightAttenuationCleared, bProjectingForForwardShading, false);
+		return RenderLightFunctionForMaterial(GraphBuilder, SceneTextures, LightSceneInfo, ScreenShadowMaskTexture, LightSceneInfo->Proxy->GetLightFunctionMaterial(), bLightAttenuationCleared, bProjectingForForwardShading, false);
 	}
 	
 	return false;
@@ -247,15 +246,14 @@ bool FDeferredShadingSceneRenderer::RenderLightFunction(
 
 bool FDeferredShadingSceneRenderer::RenderPreviewShadowsIndicator(
 	FRDGBuilder& GraphBuilder,
-	FRDGTextureRef SceneDepthTexture,
-	TRDGUniformBufferRef<FSceneTextureUniformParameters> SceneTexturesUniformBuffer,
+	const FMinimalSceneTextures& SceneTextures,
 	const FLightSceneInfo* LightSceneInfo,
 	FRDGTextureRef ScreenShadowMaskTexture,
 	bool bLightAttenuationCleared)
 {
 	if (GEngine->PreviewShadowsIndicatorMaterial)
 	{
-		return RenderLightFunctionForMaterial(GraphBuilder, SceneDepthTexture, SceneTexturesUniformBuffer, LightSceneInfo, ScreenShadowMaskTexture, GEngine->PreviewShadowsIndicatorMaterial->GetRenderProxy(), bLightAttenuationCleared, false, true);
+		return RenderLightFunctionForMaterial(GraphBuilder, SceneTextures, LightSceneInfo, ScreenShadowMaskTexture, GEngine->PreviewShadowsIndicatorMaterial->GetRenderProxy(), bLightAttenuationCleared, false, true);
 	}
 
 	return false;
@@ -288,8 +286,7 @@ static bool TryGetLightFunctionShaders(ERHIFeatureLevel::Type InFeatureLevel, FM
 
 bool FDeferredShadingSceneRenderer::RenderLightFunctionForMaterial(
 	FRDGBuilder& GraphBuilder,
-	FRDGTextureRef SceneDepthTexture,
-	TRDGUniformBufferRef<FSceneTextureUniformParameters> SceneTexturesUniformBuffer,
+	const FMinimalSceneTextures& SceneTextures,
 	const FLightSceneInfo* LightSceneInfo,
 	FRDGTextureRef ScreenShadowMaskTexture,
 	const FMaterialRenderProxy* MaterialProxy,
@@ -321,9 +318,9 @@ bool FDeferredShadingSceneRenderer::RenderLightFunctionForMaterial(
 		if (View.VisibleLightInfos[LightSceneInfo->Id].bInViewFrustum)
 		{
 			FRenderLightFunctionForMaterialParameters* PassParameters = GraphBuilder.AllocParameters<FRenderLightFunctionForMaterialParameters>();
-			PassParameters->SceneTextures = SceneTexturesUniformBuffer;
+			PassParameters->SceneTextures = SceneTextures.UniformBuffer;
 			PassParameters->RenderTargets[0] = FRenderTargetBinding(ScreenShadowMaskTexture, bLightAttenuationCleared ? ERenderTargetLoadAction::ELoad : ERenderTargetLoadAction::ENoAction);
-			PassParameters->RenderTargets.DepthStencil = FDepthStencilBinding(SceneDepthTexture, ERenderTargetLoadAction::ELoad, ERenderTargetLoadAction::ELoad, FExclusiveDepthStencil::DepthRead_StencilWrite);
+			PassParameters->RenderTargets.DepthStencil = FDepthStencilBinding(SceneTextures.Depth.Target, ERenderTargetLoadAction::ELoad, ERenderTargetLoadAction::ELoad, FExclusiveDepthStencil::DepthRead_StencilWrite);
 
 			GraphBuilder.AddPass(
 				RDG_EVENT_NAME("LightFunction Material=%s", *MaterialForRendering->GetFriendlyName()),
diff --git a/Engine/Source/Runtime/Renderer/Private/LightGridInjection.cpp b/Engine/Source/Runtime/Renderer/Private/LightGridInjection.cpp
--- a/Engine/Source/Runtime/Renderer/Private/LightGridInjection.cpp
+++ b/Engine/Source/Runtime/Renderer/Private/LightGridInjection.cpp
@@ -785,8 +785,7 @@ void FDeferredShadingSceneRenderer::GatherLightsAndComputeLightGrid(FRDGBuilder&
 
 void FDeferredShadingSceneRenderer::RenderForwardShadowProjections(
 	FRDGBuilder& GraphBuilder,
-	TRDGUniformBufferRef<FSceneTextureUniformParameters> SceneTexturesUniformBuffer,
-	FRDGTextureRef SceneDepthTexture,
+	const FMinimalSceneTextures& SceneTextures,
 	FRDGTextureRef& OutForwardScreenSpaceShadowMask,
 	FRDGTextureRef& OutForwardScreenSpaceShadowMaskSubPixel,
 	const FHairStrandsRenderingData* InHairDatas)
@@ -794,6 +793,8 @@ void FDeferredShadingSceneRenderer::RenderForwardShadowProjections(
 	const bool bIsHairEnable = InHairDatas != nullptr;
 	bool bScreenShadowMaskNeeded = false;
 
+	FRDGTextureRef SceneDepthTexture = SceneTextures.Depth.Target;
+
 	for (TSparseArray<FLightSceneInfoCompact>::TConstIterator LightIt(Scene->Lights); LightIt; ++LightIt)
 	{
 		const FLightSceneInfoCompact& LightSceneInfoCompact = *LightIt;
@@ -811,8 +812,7 @@ void FDeferredShadingSceneRenderer::RenderForwardShadowProjections(
 		FRDGTextureMSAA ForwardScreenSpaceShadowMaskSubPixel;
 
 		{
-			const FIntPoint SceneTextureExtent = SceneDepthTexture->Desc.Extent;
-			FRDGTextureDesc Desc(FRDGTextureDesc::Create2D(SceneTextureExtent, PF_B8G8R8A8, FClearValueBinding::White, TexCreate_RenderTargetable | TexCreate_ShaderResource));
+			FRDGTextureDesc Desc(FRDGTextureDesc::Create2D(SceneTextures.Extent, PF_B8G8R8A8, FClearValueBinding::White, TexCreate_RenderTargetable | TexCreate_ShaderResource));
 			Desc.NumSamples = SceneDepthTexture->Desc.NumSamples;
 			ForwardScreenSpaceShadowMask = CreateTextureMSAA(GraphBuilder, Desc, TEXT("ShadowMaskTexture"), GFastVRamConfig.ScreenSpaceShadowMask);
 			if (bIsHairEnable)
@@ -850,10 +850,9 @@ void FDeferredShadingSceneRenderer::RenderForwardShadowProjections(
 			{
 				RenderShadowProjections(
 					GraphBuilder,
-					SceneTexturesUniformBuffer,
+					SceneTextures,
 					ForwardScreenSpaceShadowMask.Target,
 					ForwardScreenSpaceShadowMaskSubPixel.Target,
-					SceneDepthTexture,
 					LightSceneInfo,
 					HairVisibilityViews,
 					bProjectingForForwardShading);
@@ -864,14 +863,13 @@ void FDeferredShadingSceneRenderer::RenderForwardShadowProjections(
 				}
 			}
 
-			RenderCapsuleDirectShadows(GraphBuilder, SceneTexturesUniformBuffer, *LightSceneInfo, ForwardScreenSpaceShadowMask.Target, VisibleLightInfo.CapsuleShadowsToProject, bProjectingForForwardShading);
+			RenderCapsuleDirectShadows(GraphBuilder, SceneTextures.UniformBuffer, *LightSceneInfo, ForwardScreenSpaceShadowMask.Target, VisibleLightInfo.CapsuleShadowsToProject, bProjectingForForwardShading);
 
 			if (LightSceneInfo->GetDynamicShadowMapChannel() >= 0 && LightSceneInfo->GetDynamicShadowMapChannel() < 4)
 			{
 				RenderLightFunction(
 					GraphBuilder,
-					SceneDepthTexture,
-					SceneTexturesUniformBuffer,
+					SceneTextures,
 					LightSceneInfo,
 					ForwardScreenSpaceShadowMask.Target,
 					true, true);
diff --git a/Engine/Source/Runtime/Renderer/Private/LightRendering.cpp b/Engine/Source/Runtime/Renderer/Private/LightRendering.cpp
--- a/Engine/Source/Runtime/Renderer/Private/LightRendering.cpp
+++ b/Engine/Source/Runtime/Renderer/Private/LightRendering.cpp
@@ -1169,32 +1169,28 @@ BEGIN_SHADER_PARAMETER_STRUCT(FRenderLightParameters, )
 END_SHADER_PARAMETER_STRUCT()
 
 void GetRenderLightParameters(
-	FRDGTextureRef SceneColorTexture,
-	FRDGTextureRef SceneDepthTexture,
+	const FMinimalSceneTextures& SceneTextures,
 	FRDGTextureRef ShadowMaskTexture,
 	FRDGTextureRef LightingChannelsTexture,
-	TRDGUniformBufferRef<FSceneTextureUniformParameters> SceneTexturesUniformBuffer,
 	const FHairStrandsVisibilityViews* InHairVisibilityViews,
 	FRenderLightParameters& Parameters)
 {
-	Parameters.SceneTextures = SceneTexturesUniformBuffer;
+	Parameters.SceneTextures = SceneTextures.UniformBuffer;
 	Parameters.ShadowMaskTexture = ShadowMaskTexture;
 	Parameters.LightingChannelsTexture = LightingChannelsTexture;
 	Parameters.HairCategorizationTexture = InHairVisibilityViews && InHairVisibilityViews->HairDatas.Num() > 0 ? InHairVisibilityViews->HairDatas[0].CategorizationTexture : nullptr;
-	Parameters.RenderTargets[0] = FRenderTargetBinding(SceneColorTexture, ERenderTargetLoadAction::ELoad);
+	Parameters.RenderTargets[0] = FRenderTargetBinding(SceneTextures.Color.Target, ERenderTargetLoadAction::ELoad);
 
-	if (SceneDepthTexture)
+	if (SceneTextures.Depth.IsValid())
 	{
-		Parameters.RenderTargets.DepthStencil = FDepthStencilBinding(SceneDepthTexture, ERenderTargetLoadAction::ELoad, ERenderTargetLoadAction::ELoad, FExclusiveDepthStencil::DepthRead_StencilWrite);
+		Parameters.RenderTargets.DepthStencil = FDepthStencilBinding(SceneTextures.Depth.Target, ERenderTargetLoadAction::ELoad, ERenderTargetLoadAction::ELoad, FExclusiveDepthStencil::DepthRead_StencilWrite);
 	}
 }
 
-FRDGTextureRef FDeferredShadingSceneRenderer::RenderLights(
+void FDeferredShadingSceneRenderer::RenderLights(
 	FRDGBuilder& GraphBuilder,
-	TRDGUniformBufferRef<FSceneTextureUniformParameters> SceneTexturesUniformBuffer,
+	FMinimalSceneTextures& SceneTextures,
 	const FTranslucentVolumeLightingTextures& TranslucentVolumeLightingTextures,
-	FRDGTextureRef SceneColorTexture,
-	FRDGTextureRef SceneDepthTexture,
 	FRDGTextureRef LightingChannelsTexture,
 	FSortedLightSetSceneInfo &SortedLightSet,
 	const FHairStrandsRenderingData* HairDatas)
@@ -1241,7 +1237,7 @@ FRDGTextureRef FDeferredShadingSceneRenderer::RenderLights(
 				// Tell the trad. deferred that the simple lights are spoken for.
 				bRenderSimpleLightsStandardDeferred = false;
 
-				AddClusteredDeferredShadingPass(GraphBuilder, SceneColorTexture, SceneTexturesUniformBuffer, SortedLightSet);
+				AddClusteredDeferredShadingPass(GraphBuilder, SceneTextures, SortedLightSet);
 			}
 			else if (CanUseTiledDeferred())
 			{
@@ -1262,19 +1258,19 @@ FRDGTextureRef FDeferredShadingSceneRenderer::RenderLights(
 					StandardDeferredStart = SortedLightSet.TiledSupportedEnd;
 					bRenderSimpleLightsStandardDeferred = false;
 
-					SceneColorTexture = RenderTiledDeferredLighting(GraphBuilder, SceneColorTexture, SceneTexturesUniformBuffer, SortedLights, SortedLightSet.SimpleLightsEnd, SortedLightSet.TiledSupportedEnd, SimpleLights);
+					RenderTiledDeferredLighting(GraphBuilder, SceneTextures, SortedLights, SortedLightSet.SimpleLightsEnd, SortedLightSet.TiledSupportedEnd, SimpleLights);
 				}
 			}
 
 			if (bRenderSimpleLightsStandardDeferred)
 			{
-				RenderSimpleLightsStandardDeferred(GraphBuilder, SceneColorTexture, SceneDepthTexture, SceneTexturesUniformBuffer, SortedLightSet.SimpleLights);
+				RenderSimpleLightsStandardDeferred(GraphBuilder, SceneTextures, SortedLightSet.SimpleLights);
 			}
 
 			if (!bUseHairLighting)
 			{
 				FRenderLightParameters* PassParameters = GraphBuilder.AllocParameters<FRenderLightParameters>();
-				GetRenderLightParameters(SceneColorTexture, SceneDepthTexture, nullptr, LightingChannelsTexture, SceneTexturesUniformBuffer, HairDatas ? &HairDatas->HairVisibilityViews : nullptr, *PassParameters);
+				GetRenderLightParameters(SceneTextures, nullptr, LightingChannelsTexture, HairDatas ? &HairDatas->HairVisibilityViews : nullptr, *PassParameters);
 
 				GraphBuilder.AddPass(
 					RDG_EVENT_NAME("StandardDeferredLighting"),
@@ -1308,7 +1304,7 @@ FRDGTextureRef FDeferredShadingSceneRenderer::RenderLights(
 					TransmittanceMaskData = RenderHairStrandsTransmittanceMask(GraphBuilder, Views, LightSceneInfo, HairDatas, NullScreenShadowMaskSubPixelTexture);
 
 					// Render the light to the scene color buffer, using a 1x1 white texture as input
-					RenderLight(GraphBuilder, SceneColorTexture, SceneDepthTexture, SceneTexturesUniformBuffer, LightSceneInfo, nullptr, LightingChannelsTexture, InHairVisibilityViews, false);
+					RenderLight(GraphBuilder, SceneTextures, LightSceneInfo, nullptr, LightingChannelsTexture, InHairVisibilityViews, false);
 				}
 			}
 
@@ -1414,8 +1410,7 @@ FRDGTextureRef FDeferredShadingSceneRenderer::RenderLights(
 
 				if (bDrawShadows || bDrawLightFunction || bDrawPreviewIndicator)
 				{
-					const FIntPoint SceneTextureExtent = SceneDepthTexture->Desc.Extent;
-					const FRDGTextureDesc Desc(FRDGTextureDesc::Create2D(SceneTextureExtent, PF_B8G8R8A8, FClearValueBinding::White, TexCreate_RenderTargetable | TexCreate_ShaderResource | GFastVRamConfig.ScreenSpaceShadowMask));
+					const FRDGTextureDesc Desc(FRDGTextureDesc::Create2D(SceneTextures.Extent, PF_B8G8R8A8, FClearValueBinding::White, TexCreate_RenderTargetable | TexCreate_ShaderResource | GFastVRamConfig.ScreenSpaceShadowMask));
 					ScreenShadowMaskTexture = GraphBuilder.CreateTexture(Desc, TEXT("ShadowMaskTexture"));
 					if (bUseHairLighting)
 					{
@@ -1459,7 +1454,7 @@ FRDGTextureRef FDeferredShadingSceneRenderer::RenderLights(
 							TStaticArray<IScreenSpaceDenoiser::FShadowVisibilityParameters, IScreenSpaceDenoiser::kMaxBatchSize> DenoisingQueue;
 							TStaticArray<int32, IScreenSpaceDenoiser::kMaxBatchSize> LightIndices;
 
-							FSceneTextureParameters SceneTextures = GetSceneTextureParameters(GraphBuilder, SceneTexturesUniformBuffer);
+							FSceneTextureParameters SceneTextureParameters = GetSceneTextureParameters(GraphBuilder, SceneTextures.UniformBuffer);
 
 							int32 ProcessShadows = 0;
 
@@ -1485,7 +1480,7 @@ FRDGTextureRef FDeferredShadingSceneRenderer::RenderLights(
 									GraphBuilder,
 									View,
 									&View.PrevViewInfo,
-									SceneTextures,
+									SceneTextureParameters,
 									DenoisingQueue,
 									InputParameterCount,
 									Outputs);
@@ -1539,7 +1534,7 @@ FRDGTextureRef FDeferredShadingSceneRenderer::RenderLights(
 								FRDGTextureRef RayTracingShadowMaskTexture;
 								{
 									FRDGTextureDesc Desc = FRDGTextureDesc::Create2D(
-										SceneTextures.SceneDepthTexture->Desc.Extent,
+										SceneTextures.Extent,
 										PF_FloatRGBA,
 										FClearValueBinding::Black,
 										TexCreate_ShaderResource | TexCreate_RenderTargetable | TexCreate_UAV);
@@ -1549,7 +1544,7 @@ FRDGTextureRef FDeferredShadingSceneRenderer::RenderLights(
 								FRDGTextureRef RayDistanceTexture;
 								{
 									FRDGTextureDesc Desc = FRDGTextureDesc::Create2D(
-										SceneTextures.SceneDepthTexture->Desc.Extent,
+										SceneTextures.Extent,
 										PF_R16F,
 										FClearValueBinding::Black,
 										TexCreate_ShaderResource | TexCreate_RenderTargetable | TexCreate_UAV);
@@ -1561,7 +1556,7 @@ FRDGTextureRef FDeferredShadingSceneRenderer::RenderLights(
 								if (bUseHairLighting)
 								{
 									FRDGTextureDesc Desc = FRDGTextureDesc::Create2D(
-										SceneTextures.SceneDepthTexture->Desc.Extent,
+										SceneTextures.Extent,
 										PF_FloatRGBA,
 										FClearValueBinding::Black,
 										TexCreate_ShaderResource | TexCreate_RenderTargetable | TexCreate_UAV);
@@ -1583,7 +1578,7 @@ FRDGTextureRef FDeferredShadingSceneRenderer::RenderLights(
 									// Note: No denoiser is required on this output, as the hair strands are geometrically noisy, which make it hard to denoise
 									RenderRayTracingShadows(
 										GraphBuilder,
-										SceneTextures,
+										SceneTextureParameters,
 										View,
 										BatchLightSceneInfo,
 										BatchRayTracingConfig,
@@ -1673,12 +1668,12 @@ FRDGTextureRef FDeferredShadingSceneRenderer::RenderLights(
 					}
 					else if (OcclusionType == FLightOcclusionType::Raytraced)
 					{
-						FSceneTextureParameters SceneTextures = GetSceneTextureParameters(GraphBuilder, SceneTexturesUniformBuffer);
+						FSceneTextureParameters SceneTextureParameters = GetSceneTextureParameters(GraphBuilder, SceneTextures.UniformBuffer);
 
 						FRDGTextureRef RayTracingShadowMaskTexture;
 						{
 							FRDGTextureDesc Desc = FRDGTextureDesc::Create2D(
-								SceneTextures.SceneDepthTexture->Desc.Extent,
+								SceneTextures.Extent,
 								PF_FloatRGBA,
 								FClearValueBinding::Black,
 								TexCreate_ShaderResource | TexCreate_RenderTargetable | TexCreate_UAV);
@@ -1688,7 +1683,7 @@ FRDGTextureRef FDeferredShadingSceneRenderer::RenderLights(
 						FRDGTextureRef RayDistanceTexture;
 						{
 							FRDGTextureDesc Desc = FRDGTextureDesc::Create2D(
-								SceneTextures.SceneDepthTexture->Desc.Extent,
+								SceneTextures.Extent,
 								PF_R16F,
 								FClearValueBinding::Black,
 								TexCreate_ShaderResource | TexCreate_RenderTargetable | TexCreate_UAV);
@@ -1703,7 +1698,7 @@ FRDGTextureRef FDeferredShadingSceneRenderer::RenderLights(
 						if (bUseHairLighting)
 						{
 							FRDGTextureDesc Desc = FRDGTextureDesc::Create2D(
-								SceneTextures.SceneDepthTexture->Desc.Extent,
+								SceneTextures.Extent,
 								PF_FloatRGBA,
 								FClearValueBinding::Black,
 								TexCreate_ShaderResource | TexCreate_RenderTargetable | TexCreate_UAV);
@@ -1715,7 +1710,7 @@ FRDGTextureRef FDeferredShadingSceneRenderer::RenderLights(
 						FRDGTextureRef RayTracingShadowMaskTileTexture;
 						{
 							FRDGTextureDesc Desc = FRDGTextureDesc::Create2D(
-								SceneTextures.SceneDepthTexture->Desc.Extent,
+								SceneTextures.Extent,
 								PF_FloatRGBA,
 								FClearValueBinding::Black,
 								TexCreate_ShaderResource | TexCreate_RenderTargetable | TexCreate_UAV);
@@ -1743,7 +1738,7 @@ FRDGTextureRef FDeferredShadingSceneRenderer::RenderLights(
 
 							RenderRayTracingShadows(
 								GraphBuilder,
-								SceneTextures,
+								SceneTextureParameters,
 								View,
 								LightSceneInfo,
 								RayTracingConfig,
@@ -1776,7 +1771,7 @@ FRDGTextureRef FDeferredShadingSceneRenderer::RenderLights(
 									GraphBuilder,
 									View,
 									&View.PrevViewInfo,
-									SceneTextures,
+									SceneTextureParameters,
 									InputParameters,
 									InputParameterCount,
 									Outputs);
@@ -1868,7 +1863,7 @@ FRDGTextureRef FDeferredShadingSceneRenderer::RenderLights(
 							ClearShadowMask(ScreenShadowMaskSubPixelTexture);
 						}
 
-						RenderDeferredShadowProjections(GraphBuilder, SceneTexturesUniformBuffer, TranslucentVolumeLightingTextures, &LightSceneInfo, ScreenShadowMaskTexture, ScreenShadowMaskSubPixelTexture, SceneDepthTexture, HairDatas, bInjectedTranslucentVolume);
+						RenderDeferredShadowProjections(GraphBuilder, SceneTextures, TranslucentVolumeLightingTextures, &LightSceneInfo, ScreenShadowMaskTexture, ScreenShadowMaskSubPixelTexture, HairDatas, bInjectedTranslucentVolume);
 					}
 
 					bUsedShadowMaskTexture = true;
@@ -1886,13 +1881,13 @@ FRDGTextureRef FDeferredShadingSceneRenderer::RenderLights(
 				{
 					if (bDrawLightFunction)
 					{
-						const bool bLightFunctionRendered = RenderLightFunction(GraphBuilder, SceneDepthTexture, SceneTexturesUniformBuffer, &LightSceneInfo, ScreenShadowMaskTexture, bDrawShadows, false);
+						const bool bLightFunctionRendered = RenderLightFunction(GraphBuilder, SceneTextures, &LightSceneInfo, ScreenShadowMaskTexture, bDrawShadows, false);
 						bUsedShadowMaskTexture |= bLightFunctionRendered;
 					}
 
 					if (bDrawPreviewIndicator)
 					{
-						RenderPreviewShadowsIndicator(GraphBuilder, SceneDepthTexture, SceneTexturesUniformBuffer, &LightSceneInfo, ScreenShadowMaskTexture, bUsedShadowMaskTexture);
+						RenderPreviewShadowsIndicator(GraphBuilder, SceneTextures, &LightSceneInfo, ScreenShadowMaskTexture, bUsedShadowMaskTexture);
 					}
 
 					if (!bDrawShadows)
@@ -1924,7 +1919,7 @@ FRDGTextureRef FDeferredShadingSceneRenderer::RenderLights(
 				if (bDirectLighting)
 				{
 					const bool bRenderOverlap = false;
-					RenderLight(GraphBuilder, SceneColorTexture, SceneDepthTexture, SceneTexturesUniformBuffer, &LightSceneInfo, ScreenShadowMaskTexture, LightingChannelsTexture, InHairVisibilityViews, bRenderOverlap);
+					RenderLight(GraphBuilder, SceneTextures, &LightSceneInfo, ScreenShadowMaskTexture, LightingChannelsTexture, InHairVisibilityViews, bRenderOverlap);
 				}
 
 				if (bUseHairLighting)
@@ -1935,13 +1930,11 @@ FRDGTextureRef FDeferredShadingSceneRenderer::RenderLights(
 						TransmittanceMaskData = RenderHairStrandsTransmittanceMask(GraphBuilder, Views, &LightSceneInfo, HairDatas, ScreenShadowMaskSubPixelTexture);
 					}
 
-					RenderLightForHair(GraphBuilder, SceneTexturesUniformBuffer, &LightSceneInfo, ScreenShadowMaskSubPixelTexture, LightingChannelsTexture, TransmittanceMaskData, InHairVisibilityViews);
+					RenderLightForHair(GraphBuilder, SceneTextures.UniformBuffer, &LightSceneInfo, ScreenShadowMaskSubPixelTexture, LightingChannelsTexture, TransmittanceMaskData, InHairVisibilityViews);
 				}
 			}
 		}
 	}
-
-	return SceneColorTexture;
 }
 
 void FDeferredShadingSceneRenderer::RenderLightArrayForOverlapViewmode(
@@ -1981,15 +1974,13 @@ void FDeferredShadingSceneRenderer::RenderLightArrayForOverlapViewmode(
 
 void FDeferredShadingSceneRenderer::RenderStationaryLightOverlap(
 	FRDGBuilder& GraphBuilder,
-	FRDGTextureRef SceneColorTexture,
-	FRDGTextureRef SceneDepthTexture,
-	FRDGTextureRef LightingChannelsTexture,
-	TRDGUniformBufferRef<FSceneTextureUniformParameters> SceneTexturesUniformBuffer)
+	const FMinimalSceneTextures& SceneTextures,
+	FRDGTextureRef LightingChannelsTexture)
 {
 	if (Scene->bIsEditorScene)
 	{
 		FRenderLightParameters* PassParameters = GraphBuilder.AllocParameters<FRenderLightParameters>();
-		GetRenderLightParameters(SceneColorTexture, SceneDepthTexture, nullptr, LightingChannelsTexture, SceneTexturesUniformBuffer, nullptr, *PassParameters);
+		GetRenderLightParameters(SceneTextures, nullptr, LightingChannelsTexture, nullptr, *PassParameters);
 
 		GraphBuilder.AddPass(
 			RDG_EVENT_NAME("StationaryLightOverlap"),
@@ -2322,9 +2313,7 @@ void FDeferredShadingSceneRenderer::RenderLight(
 
 void FDeferredShadingSceneRenderer::RenderLight(
 	FRDGBuilder& GraphBuilder,
-	FRDGTextureRef SceneColorTexture,
-	FRDGTextureRef SceneDepthTexture,
-	TRDGUniformBufferRef<FSceneTextureUniformParameters> SceneTexturesUniformBuffer,
+	const FMinimalSceneTextures& SceneTextures,
 	const FLightSceneInfo* LightSceneInfo,
 	FRDGTextureRef ScreenShadowMaskTexture,
 	FRDGTextureRef LightingChannelsTexture,
@@ -2332,7 +2321,7 @@ void FDeferredShadingSceneRenderer::RenderLight(
 	bool bRenderOverlap)
 {
 	FRenderLightParameters* PassParameters = GraphBuilder.AllocParameters<FRenderLightParameters>();
-	GetRenderLightParameters(SceneColorTexture, SceneDepthTexture, ScreenShadowMaskTexture, LightingChannelsTexture, SceneTexturesUniformBuffer, InHairVisibilityViews, *PassParameters);
+	GetRenderLightParameters(SceneTextures, ScreenShadowMaskTexture, LightingChannelsTexture, InHairVisibilityViews, *PassParameters);
 
 	ERDGPassFlags PassFlags = ERDGPassFlags::Raster;
 
@@ -2409,8 +2398,12 @@ void FDeferredShadingSceneRenderer::RenderLightForHair(
 			continue;
 		}
 
+		FMinimalSceneTextures SceneTextures{};
+		SceneTextures.Color = HairVisibilityData.SampleLightingBuffer;
+		SceneTextures.UniformBuffer = SceneTexturesUniformBuffer;
+
 		FRenderLightForHairParameters* PassParameters = GraphBuilder.AllocParameters<FRenderLightForHairParameters>();
-		GetRenderLightParameters(HairVisibilityData.SampleLightingBuffer, nullptr, HairShadowMaskTexture, LightingChannelsTexture, SceneTexturesUniformBuffer, InHairVisibilityViews, PassParameters->Light);
+		GetRenderLightParameters(SceneTextures, HairShadowMaskTexture, LightingChannelsTexture, InHairVisibilityViews, PassParameters->Light);
 		PassParameters->HairIndexAndCountTexture = HairVisibilityData.NodeIndex;
 
 		PassParameters->HairTransmittanceMask = InTransmittanceMaskData.TransmittanceMask;
@@ -2546,18 +2539,16 @@ END_SHADER_PARAMETER_STRUCT()
 
 void FDeferredShadingSceneRenderer::RenderSimpleLightsStandardDeferred(
 	FRDGBuilder& GraphBuilder,
-	FRDGTextureRef SceneColorTexture,
-	FRDGTextureRef SceneDepthTexture,
-	TRDGUniformBufferRef<FSceneTextureUniformParameters> SceneTexturesUniformBuffer,
+	const FMinimalSceneTextures& SceneTextures,
 	const FSimpleLightArray& SimpleLights)
 {
 	SCOPE_CYCLE_COUNTER(STAT_DirectLightRenderingTime);
 	INC_DWORD_STAT_BY(STAT_NumLightsUsingStandardDeferred, SimpleLights.InstanceData.Num());
 
 	FSimpleLightsStandardDeferredParameters* PassParameters = GraphBuilder.AllocParameters<FSimpleLightsStandardDeferredParameters>();
-	PassParameters->SceneTextures = SceneTexturesUniformBuffer;
-	PassParameters->RenderTargets[0] = FRenderTargetBinding(SceneColorTexture, ERenderTargetLoadAction::ELoad);
-	PassParameters->RenderTargets.DepthStencil = FDepthStencilBinding(SceneDepthTexture, ERenderTargetLoadAction::ELoad, ERenderTargetLoadAction::ELoad, FExclusiveDepthStencil::DepthRead_StencilWrite);
+	PassParameters->SceneTextures = SceneTextures.UniformBuffer;
+	PassParameters->RenderTargets[0] = FRenderTargetBinding(SceneTextures.Color.Target, ERenderTargetLoadAction::ELoad);
+	PassParameters->RenderTargets.DepthStencil = FDepthStencilBinding(SceneTextures.Depth.Target, ERenderTargetLoadAction::ELoad, ERenderTargetLoadAction::ELoad, FExclusiveDepthStencil::DepthRead_StencilWrite);
 
 	GraphBuilder.AddPass(
 		RDG_EVENT_NAME("StandardDeferredSimpleLights"),
diff --git a/Engine/Source/Runtime/Renderer/Private/LightShaftRendering.cpp b/Engine/Source/Runtime/Renderer/Private/LightShaftRendering.cpp
--- a/Engine/Source/Runtime/Renderer/Private/LightShaftRendering.cpp
+++ b/Engine/Source/Runtime/Renderer/Private/LightShaftRendering.cpp
@@ -423,8 +423,7 @@ void AddOcclusionTermPass(
 
 FRDGTextureRef FDeferredShadingSceneRenderer::RenderLightShaftOcclusion(
 	FRDGBuilder& GraphBuilder,
-	TRDGUniformBufferRef<FSceneTextureUniformParameters> SceneTexturesUniformBuffer,
-	FIntPoint SceneTextureExtent)
+	const FMinimalSceneTextures& SceneTextures)
 {
 	FScreenPassRenderTarget Output;
 
@@ -459,7 +458,7 @@ FRDGTextureRef FDeferredShadingSceneRenderer::RenderLightShaftOcclusion(
 						RDG_GPU_MASK_SCOPE(GraphBuilder, View.GPUMask);
 						INC_DWORD_STAT(STAT_LightShaftsLights);
 
-						const FScreenPassTextureViewport SceneViewport(SceneTextureExtent, View.ViewRect);
+						const FScreenPassTextureViewport SceneViewport(SceneTextures.Extent, View.ViewRect);
 						const FScreenPassTextureViewport OutputViewport(GetDownscaledViewport(SceneViewport, GetLightShaftDownsampleFactor()));
 						const FLightShaftPixelShaderParameters LightShaftParameters = GetLightShaftParameters(View, LightSceneInfo, SceneViewport, OutputViewport);
 						FTemporalAAHistory* TemporalHistory = nullptr;
@@ -482,7 +481,7 @@ FRDGTextureRef FDeferredShadingSceneRenderer::RenderLightShaftOcclusion(
 							GraphBuilder,
 							View,
 							LightShaftParameters,
-							SceneTexturesUniformBuffer,
+							SceneTextures.UniformBuffer,
 							SceneViewport,
 							OutputViewport,
 							TemporalHistory,
@@ -526,9 +525,7 @@ void AddLightShaftBloomPass(
 
 void FDeferredShadingSceneRenderer::RenderLightShaftBloom(
 	FRDGBuilder& GraphBuilder,
-	TRDGUniformBufferRef<FSceneTextureUniformParameters> SceneTexturesUniformBuffer,
-	FIntPoint SceneTextureExtent,
-	FRDGTextureRef SceneColorTexture,
+	const FMinimalSceneTextures& SceneTextures,
 	FSeparateTranslucencyTextures& SeparateTranslucencyTextures)
 {
 	if (ShouldRenderLightShafts(ViewFamily))
@@ -558,7 +555,7 @@ void FDeferredShadingSceneRenderer::RenderLightShaftBloom(
 				if (bWillRenderLightShafts)
 				{
 					// Default to scene color output.
-					FRDGTextureMSAA OutputTexture(SceneColorTexture);
+					FRDGTextureMSAA OutputTexture(SceneTextures.Color.Target);
 					ERenderTargetLoadAction OutputLoadAction = ERenderTargetLoadAction::ELoad;
 					float OutputViewRectScale = 1.0f;
 
@@ -586,7 +583,7 @@ void FDeferredShadingSceneRenderer::RenderLightShaftBloom(
 							RDG_GPU_MASK_SCOPE(GraphBuilder, View.GPUMask);
 							INC_DWORD_STAT(STAT_LightShaftsLights);
 
-							const FScreenPassTextureViewport SceneViewport(SceneTextureExtent, View.ViewRect);
+							const FScreenPassTextureViewport SceneViewport(SceneTextures.Extent, View.ViewRect);
 							const FScreenPassTextureViewport LightShaftViewport(GetDownscaledViewport(SceneViewport, GetLightShaftDownsampleFactor()));
 							const FScreenPassTextureViewport OutputViewport(OutputTexture.Target, GetScaledRect(View.ViewRect, OutputViewRectScale));
 							const FLightShaftPixelShaderParameters LightShaftParameters = GetLightShaftParameters(View, LightSceneInfo, SceneViewport, LightShaftViewport);
@@ -601,7 +598,7 @@ void FDeferredShadingSceneRenderer::RenderLightShaftBloom(
 								GraphBuilder,
 								View,
 								LightShaftParameters,
-								SceneTexturesUniformBuffer,
+								SceneTextures.UniformBuffer,
 								SceneViewport,
 								LightShaftViewport,
 								TemporalHistory,
diff --git a/Engine/Source/Runtime/Renderer/Private/PostProcess/PostProcessSubsurface.cpp b/Engine/Source/Runtime/Renderer/Private/PostProcess/PostProcessSubsurface.cpp
--- a/Engine/Source/Runtime/Renderer/Private/PostProcess/PostProcessSubsurface.cpp
+++ b/Engine/Source/Runtime/Renderer/Private/PostProcess/PostProcessSubsurface.cpp
@@ -1303,28 +1303,22 @@ FScreenPassTexture AddVisualizeSubsurfacePass(FRDGBuilder& GraphBuilder, const F
 	return MoveTemp(Output);
 }
 
-FRDGTextureRef AddSubsurfacePass(
+void AddSubsurfacePass(
 	FRDGBuilder& GraphBuilder,
-	TRDGUniformBufferRef<FSceneTextureUniformParameters> SceneTexturesUniformBuffer,
-	TArrayView<const FViewInfo> Views,
-	FRDGTextureRef SceneColorTexture)
+	FMinimalSceneTextures& SceneTextures,
+	TArrayView<const FViewInfo> Views)
 {
 	const uint32 ViewMask = GetSubsurfaceRequiredViewMask(Views);
 
 	if (!ViewMask)
 	{
-		return SceneColorTexture;
+		return;
 	}
 
-	checkf(SceneColorTexture->Desc.NumSamples == 1, TEXT("Subsurface rendering requires the deferred renderer."));
-
-	FRDGTextureRef SceneColorOutputTexture = GraphBuilder.CreateTexture(SceneColorTexture->Desc, TEXT("SceneColorSubsurface"));
-	AddSubsurfacePass(GraphBuilder, Views, ViewMask, SceneTexturesUniformBuffer, SceneColorTexture, SceneColorOutputTexture);
-
-	{
-		FSceneRenderTargets& SceneContext = FSceneRenderTargets::Get();
-		ConvertToExternalTexture(GraphBuilder, SceneColorOutputTexture, SceneContext.GetSceneColor());
-	}
+	checkf(!SceneTextures.Color.IsSeparate(), TEXT("Subsurface rendering requires the deferred renderer."));
 
-	return SceneColorOutputTexture;
+	FRDGTextureRef SceneColorOutputTexture = GraphBuilder.CreateTexture(SceneTextures.Color.Target->Desc, TEXT("SceneColorSubsurface"));
+	AddSubsurfacePass(GraphBuilder, Views, ViewMask, SceneTextures.UniformBuffer, SceneTextures.Color.Target, SceneColorOutputTexture);
+	ConvertToExternalTexture(GraphBuilder, SceneColorOutputTexture, FSceneRenderTargets::Get().GetSceneColor());
+	SceneTextures.Color = SceneColorOutputTexture;
 }
\ No newline at end of file
diff --git a/Engine/Source/Runtime/Renderer/Private/PostProcess/PostProcessSubsurface.h b/Engine/Source/Runtime/Renderer/Private/PostProcess/PostProcessSubsurface.h
--- a/Engine/Source/Runtime/Renderer/Private/PostProcess/PostProcessSubsurface.h
+++ b/Engine/Source/Runtime/Renderer/Private/PostProcess/PostProcessSubsurface.h
@@ -30,8 +30,7 @@ struct FVisualizeSubsurfaceInputs
 
 FScreenPassTexture AddVisualizeSubsurfacePass(FRDGBuilder& GraphBuilder, const FViewInfo& View, const FVisualizeSubsurfaceInputs& Inputs);
 
-FRDGTextureRef AddSubsurfacePass(
+void AddSubsurfacePass(
 	FRDGBuilder& GraphBuilder,
-	TRDGUniformBufferRef<FSceneTextureUniformParameters> SceneTexturesUniformBuffer,
-	TArrayView<const FViewInfo> Views,
-	FRDGTextureRef SceneColorTexture);
\ No newline at end of file
+	FMinimalSceneTextures& SceneTextures,
+	TArrayView<const FViewInfo> Views);
\ No newline at end of file
diff --git a/Engine/Source/Runtime/Renderer/Private/PostProcess/SceneRenderTargets.cpp b/Engine/Source/Runtime/Renderer/Private/PostProcess/SceneRenderTargets.cpp
--- a/Engine/Source/Runtime/Renderer/Private/PostProcess/SceneRenderTargets.cpp
+++ b/Engine/Source/Runtime/Renderer/Private/PostProcess/SceneRenderTargets.cpp
@@ -154,6 +154,22 @@ static TGlobalResource<FSceneRenderTargets> SceneRenderTargetsSingleton;
 
 extern int32 GUseTranslucentLightingVolumes;
 
+RDG_REGISTER_BLACKBOARD_STRUCT(FSceneTextures);
+
+FSceneTextures& FSceneTextures::Create(FRDGBuilder& GraphBuilder)
+{
+	auto& SceneTextures = GraphBuilder.Blackboard.Create<FSceneTextures>();
+
+	const FSceneRenderTargets& SceneContext = FSceneRenderTargets::Get();
+	SceneTextures.Extent = SceneContext.GetBufferSizeXY();
+
+	SceneTextures.Depth = RegisterExternalTextureMSAA(GraphBuilder, SceneContext.SceneDepthZ);
+	SceneTextures.Stencil = GraphBuilder.CreateSRV(FRDGTextureSRVDesc::CreateWithPixelFormat(SceneTextures.Depth.Target, PF_X24_G8));
+	SceneTextures.SmallDepth = GraphBuilder.RegisterExternalTexture(SceneContext.SmallDepthZ, ERenderTargetTexture::Targetable);
+
+	return SceneTextures;
+}
+
 FSceneRenderTargets& FSceneRenderTargets::Get()
 {
 	check(IsInRenderingThread());
@@ -2230,35 +2246,33 @@ IMPLEMENT_STATIC_UNIFORM_BUFFER_SLOT(SceneTextures);
 IMPLEMENT_STATIC_UNIFORM_BUFFER_STRUCT(FSceneTextureUniformParameters, "SceneTexturesStruct", SceneTextures);
 
 void SetupSceneTextureUniformParameters(
-	FRDGBuilder* GraphBuilder,
+	FRDGBuilder& GraphBuilder,
 	ERHIFeatureLevel::Type FeatureLevel,
-	const FSceneRenderTargets& SceneContext,
 	ESceneTextureSetupMode SetupMode,
 	FSceneTextureUniformParameters& SceneTextureParameters)
 {
+	const FSceneRenderTargets& SceneContext = FSceneRenderTargets::Get();
+	const FRDGSystemTextures& SystemTextures = FRDGSystemTextures::Get(GraphBuilder);
+
 	const auto GetRDG = [&](const TRefCountPtr<IPooledRenderTarget>& PooledRenderTarget, ERDGTextureFlags Flags = ERDGTextureFlags::None)
 	{
-		return RegisterExternalOrPassthroughTexture(GraphBuilder, PooledRenderTarget, Flags);
+		return GraphBuilder.RegisterExternalTexture(PooledRenderTarget, ERenderTargetTexture::ShaderResource, Flags);
 	};
 
-	FRDGTextureRef WhiteDefault2D = GetRDG(GSystemTextures.WhiteDummy);
-	FRDGTextureRef BlackDefault2D = GetRDG(GSystemTextures.BlackDummy);
-	FRDGTextureRef DepthDefault = GetRDG(GSystemTextures.DepthDummy);
-	check(WhiteDefault2D && BlackDefault2D && DepthDefault);
+	SceneTextureParameters.SceneColorTexture = SystemTextures.Black;
+	SceneTextureParameters.SceneDepthTexture = SystemTextures.DepthDummy;
 
 	// Scene Color / Depth
+	if (const FSceneTextures* SceneTextures = GraphBuilder.Blackboard.Get<FSceneTextures>())
 	{
-		SceneTextureParameters.SceneColorTexture = BlackDefault2D;
-		SceneTextureParameters.SceneDepthTexture = DepthDefault;
-
 		if (EnumHasAnyFlags(SetupMode, ESceneTextureSetupMode::SceneColor))
 		{
-			SceneTextureParameters.SceneColorTexture = GetRDG(SceneContext.GetSceneColor());
+			SceneTextureParameters.SceneColorTexture = SceneTextures->Color.Resolve;
 		}
 
 		if (EnumHasAnyFlags(SetupMode, ESceneTextureSetupMode::SceneDepth))
 		{
-			SceneTextureParameters.SceneDepthTexture = GetRDG(SceneContext.SceneDepthZ);
+			SceneTextureParameters.SceneDepthTexture = SceneTextures->Depth.Resolve;
 		}
 	}
 
@@ -2268,30 +2282,30 @@ void SetupSceneTextureUniformParameters(
 		const bool bCanReadGBufferUniforms = IsUsingGBuffers(ShaderPlatform) || IsSimpleForwardShadingEnabled(ShaderPlatform);
 
 		// Allocate the Gbuffer resource uniform buffer.
-		SceneTextureParameters.GBufferATexture = bCanReadGBufferUniforms && EnumHasAnyFlags(SetupMode, ESceneTextureSetupMode::GBufferA) && SceneContext.GBufferA ? GetRDG(SceneContext.GBufferA) : BlackDefault2D;
-		SceneTextureParameters.GBufferBTexture = bCanReadGBufferUniforms && EnumHasAnyFlags(SetupMode, ESceneTextureSetupMode::GBufferB) && SceneContext.GBufferB ? GetRDG(SceneContext.GBufferB) : BlackDefault2D;
-		SceneTextureParameters.GBufferCTexture = bCanReadGBufferUniforms && EnumHasAnyFlags(SetupMode, ESceneTextureSetupMode::GBufferC) && SceneContext.GBufferC ? GetRDG(SceneContext.GBufferC) : BlackDefault2D;
-		SceneTextureParameters.GBufferDTexture = bCanReadGBufferUniforms && EnumHasAnyFlags(SetupMode, ESceneTextureSetupMode::GBufferD) && SceneContext.GBufferD ? GetRDG(SceneContext.GBufferD) : BlackDefault2D;
-		SceneTextureParameters.GBufferETexture = bCanReadGBufferUniforms && EnumHasAnyFlags(SetupMode, ESceneTextureSetupMode::GBufferE) && SceneContext.GBufferE ? GetRDG(SceneContext.GBufferE) : BlackDefault2D;
-		SceneTextureParameters.GBufferFTexture = bCanReadGBufferUniforms && EnumHasAnyFlags(SetupMode, ESceneTextureSetupMode::GBufferF) && SceneContext.GBufferF ? GetRDG(SceneContext.GBufferF) : BlackDefault2D;
+		SceneTextureParameters.GBufferATexture = bCanReadGBufferUniforms && EnumHasAnyFlags(SetupMode, ESceneTextureSetupMode::GBufferA) && SceneContext.GBufferA ? GetRDG(SceneContext.GBufferA) : SystemTextures.Black;
+		SceneTextureParameters.GBufferBTexture = bCanReadGBufferUniforms && EnumHasAnyFlags(SetupMode, ESceneTextureSetupMode::GBufferB) && SceneContext.GBufferB ? GetRDG(SceneContext.GBufferB) : SystemTextures.Black;
+		SceneTextureParameters.GBufferCTexture = bCanReadGBufferUniforms && EnumHasAnyFlags(SetupMode, ESceneTextureSetupMode::GBufferC) && SceneContext.GBufferC ? GetRDG(SceneContext.GBufferC) : SystemTextures.Black;
+		SceneTextureParameters.GBufferDTexture = bCanReadGBufferUniforms && EnumHasAnyFlags(SetupMode, ESceneTextureSetupMode::GBufferD) && SceneContext.GBufferD ? GetRDG(SceneContext.GBufferD) : SystemTextures.Black;
+		SceneTextureParameters.GBufferETexture = bCanReadGBufferUniforms && EnumHasAnyFlags(SetupMode, ESceneTextureSetupMode::GBufferE) && SceneContext.GBufferE ? GetRDG(SceneContext.GBufferE) : SystemTextures.Black;
+		SceneTextureParameters.GBufferFTexture = bCanReadGBufferUniforms && EnumHasAnyFlags(SetupMode, ESceneTextureSetupMode::GBufferF) && SceneContext.GBufferF ? GetRDG(SceneContext.GBufferF) : SystemTextures.Black;
 	}
 
 	// Velocity
 	{
-		SceneTextureParameters.GBufferVelocityTexture = EnumHasAnyFlags(SetupMode, ESceneTextureSetupMode::SceneVelocity) && SceneContext.SceneVelocity ? GetRDG(SceneContext.SceneVelocity) : BlackDefault2D;
+		SceneTextureParameters.GBufferVelocityTexture = EnumHasAnyFlags(SetupMode, ESceneTextureSetupMode::SceneVelocity) && SceneContext.SceneVelocity ? GetRDG(SceneContext.SceneVelocity) : SystemTextures.Black;
 	}
 
 	// SSAO
 	{
 		const bool bSetupSSAO = EnumHasAnyFlags(SetupMode, ESceneTextureSetupMode::SSAO);
-		SceneTextureParameters.ScreenSpaceAOTexture = bSetupSSAO && SceneContext.bScreenSpaceAOIsValid && SceneContext.ScreenSpaceAO ? GetRDG(SceneContext.ScreenSpaceAO) : WhiteDefault2D;
+		SceneTextureParameters.ScreenSpaceAOTexture = bSetupSSAO && SceneContext.bScreenSpaceAOIsValid && SceneContext.ScreenSpaceAO ? GetRDG(SceneContext.ScreenSpaceAO) : SystemTextures.White;
 	}
 
 	// Custom Depth / Stencil
 	{
 		const bool bSetupCustomDepth = EnumHasAnyFlags(SetupMode, ESceneTextureSetupMode::CustomDepth);
 
-		FRDGTextureRef CustomDepth = DepthDefault;
+		FRDGTextureRef CustomDepth = SystemTextures.DepthDummy;
 		FRHIShaderResourceView* CustomStencilSRV = GSystemTextures.StencilDummySRV;
 
 		if (SceneContext.bCustomDepthIsValid)
@@ -2314,17 +2328,8 @@ void SetupSceneTextureUniformParameters(
 	ESceneTextureSetupMode SetupMode,
 	FSceneTextureUniformParameters& SceneTextureParameters)
 {
-	SetupSceneTextureUniformParameters(nullptr, FeatureLevel, SceneContext, SetupMode, SceneTextureParameters);
-}
-
-void SetupSceneTextureUniformParameters(
-	FRDGBuilder& GraphBuilder,
-	ERHIFeatureLevel::Type FeatureLevel,
-	ESceneTextureSetupMode SetupMode,
-	FSceneTextureUniformParameters& SceneTextureParameters)
-{
-	FSceneRenderTargets& SceneContext = FSceneRenderTargets::Get();
-	SetupSceneTextureUniformParameters(&GraphBuilder, FeatureLevel, SceneContext, SetupMode, SceneTextureParameters);
+	// Deprecated.
+	checkNoEntry();
 }
 
 TUniformBufferRef<FSceneTextureUniformParameters> CreateSceneTextureUniformBuffer(
@@ -2332,10 +2337,9 @@ TUniformBufferRef<FSceneTextureUniformParameters> CreateSceneTextureUniformBuffe
 	ERHIFeatureLevel::Type FeatureLevel,
 	ESceneTextureSetupMode SetupMode)
 {
-	FSceneTextureUniformParameters SceneTextures;
-	FSceneRenderTargets& SceneContext = FSceneRenderTargets::Get();
-	SetupSceneTextureUniformParameters(nullptr, FeatureLevel, SceneContext, SetupMode, SceneTextures);
-	return TUniformBufferRef<FSceneTextureUniformParameters>::CreateUniformBufferImmediate(SceneTextures, EUniformBufferUsage::UniformBuffer_SingleFrame);
+	// Deprecated.
+	checkNoEntry();
+	return {};
 }
 
 TRDGUniformBufferRef<FSceneTextureUniformParameters> CreateSceneTextureUniformBuffer(
@@ -2344,8 +2348,7 @@ TRDGUniformBufferRef<FSceneTextureUniformParameters> CreateSceneTextureUniformBu
 	ESceneTextureSetupMode SetupMode)
 {
 	FSceneTextureUniformParameters* SceneTextures = GraphBuilder.AllocParameters<FSceneTextureUniformParameters>();
-	FSceneRenderTargets& SceneContext = FSceneRenderTargets::Get();
-	SetupSceneTextureUniformParameters(&GraphBuilder, FeatureLevel, SceneContext, SetupMode, *SceneTextures);
+	SetupSceneTextureUniformParameters(GraphBuilder, FeatureLevel, SetupMode, *SceneTextures);
 	return GraphBuilder.CreateUniformBuffer(SceneTextures);
 }
 
diff --git a/Engine/Source/Runtime/Renderer/Private/PostProcess/SceneRenderTargets.h b/Engine/Source/Runtime/Renderer/Private/PostProcess/SceneRenderTargets.h
--- a/Engine/Source/Runtime/Renderer/Private/PostProcess/SceneRenderTargets.h
+++ b/Engine/Source/Runtime/Renderer/Private/PostProcess/SceneRenderTargets.h
@@ -82,6 +82,26 @@ struct FCustomDepthTextures
 	FRDGTextureRef MobileCustomStencil{};
 };
 
+struct FMinimalSceneTextures
+{
+	FRDGTextureMSAA Color{};
+	FRDGTextureMSAA Depth{};
+	FRDGTextureSRVRef Stencil{};
+
+	FIntPoint Extent = FIntPoint::ZeroValue;
+
+	ESceneTextureSetupMode SetupMode = ESceneTextureSetupMode::None;
+	TRDGUniformBufferRef<FSceneTextureUniformParameters> UniformBuffer{};
+};
+
+struct FSceneTextures : public FMinimalSceneTextures
+{
+	RENDERER_API static FSceneTextures& Create(FRDGBuilder& GraphBuilder);
+
+	FRDGTextureRef SmallDepth{};
+	FRDGTextureRef Velocity{};
+};
+
 /**
  * Encapsulates the render targets used for scene rendering.
  */
@@ -569,12 +589,4 @@ private:
 
 	/** True if the a variable resolution texture is allocated to control sampling or shading rate */
 	bool bAllocatedFoveationTexture;
-};
-
-/** Sets up scene texture parameters for RDG (builder is valid) or passthrough RHI access (builder is null). Intended for temporary use during RDG refactor. */
-extern void SetupSceneTextureUniformParameters(
-	FRDGBuilder* GraphBuilder,
-	ERHIFeatureLevel::Type FeatureLevel,
-	const FSceneRenderTargets& SceneContext,
-	ESceneTextureSetupMode SetupMode,
-	FSceneTextureUniformParameters& SceneTextureParameters);
\ No newline at end of file
+};
\ No newline at end of file
diff --git a/Engine/Source/Runtime/Renderer/Private/RayTracing/RaytracingSkylight.cpp b/Engine/Source/Runtime/Renderer/Private/RayTracing/RaytracingSkylight.cpp
--- a/Engine/Source/Runtime/Renderer/Private/RayTracing/RaytracingSkylight.cpp
+++ b/Engine/Source/Runtime/Renderer/Private/RayTracing/RaytracingSkylight.cpp
@@ -694,8 +694,7 @@ IMPLEMENT_GLOBAL_SHADER(FCompositeSkyLightPS, "/Engine/Private/RayTracing/Compos
 
 void FDeferredShadingSceneRenderer::CompositeRayTracingSkyLight(
 	FRDGBuilder& GraphBuilder,
-	FRDGTextureRef SceneColorTexture,
-	FIntPoint SceneTextureExtent,
+	const FMinimalSceneTextures& SceneTextures,
 	FRDGTextureRef SkyLightRT,
 	FRDGTextureRef HitDistanceRT)
 #if RHI_RAYTRACING
@@ -706,21 +705,21 @@ void FDeferredShadingSceneRenderer::CompositeRayTracingSkyLight(
 	{
 		const FViewInfo& View = Views[ViewIndex];
 
-		FSceneTextureParameters SceneTextures = GetSceneTextureParameters(GraphBuilder);
+		FSceneTextureParameters SceneTextureParameters = GetSceneTextureParameters(GraphBuilder, SceneTextures.UniformBuffer);
 
 		FCompositeSkyLightPS::FParameters *PassParameters = GraphBuilder.AllocParameters<FCompositeSkyLightPS::FParameters>();
 		PassParameters->SkyLightTexture = SkyLightRT;
 		PassParameters->SkyLightTextureSampler = TStaticSamplerState<SF_Point, AM_Clamp, AM_Clamp, AM_Clamp>::GetRHI();
 		PassParameters->ViewUniformBuffer = Views[ViewIndex].ViewUniformBuffer;
-		PassParameters->RenderTargets[0] = FRenderTargetBinding(SceneColorTexture, ERenderTargetLoadAction::ELoad);
-		PassParameters->SceneTextures = SceneTextures;
+		PassParameters->RenderTargets[0] = FRenderTargetBinding(SceneTextures.Color.Target, ERenderTargetLoadAction::ELoad);
+		PassParameters->SceneTextures = SceneTextureParameters;
 
 		// dxr_todo: Unify with RTGI compositing workflow
 		GraphBuilder.AddPass(
 			RDG_EVENT_NAME("GlobalIlluminationComposite"),
 			PassParameters,
 			ERDGPassFlags::Raster,
-			[this, &View, PassParameters, SceneTextureExtent](FRHICommandListImmediate& RHICmdList)
+			[this, &View, PassParameters, SceneTextureExtent = SceneTextures.Extent](FRHICommandListImmediate& RHICmdList)
 		{
 			TShaderMapRef<FPostProcessVS> VertexShader(View.ShaderMap);
 			TShaderMapRef<FCompositeSkyLightPS> PixelShader(View.ShaderMap);
diff --git a/Engine/Source/Runtime/Renderer/Private/RendererModule.h b/Engine/Source/Runtime/Renderer/Private/RendererModule.h
--- a/Engine/Source/Runtime/Renderer/Private/RendererModule.h
+++ b/Engine/Source/Runtime/Renderer/Private/RendererModule.h
@@ -110,16 +110,13 @@ public:
 	void RenderPostOpaqueExtensions(
 		FRDGBuilder& GraphBuilder,
 		TArrayView<const FViewInfo> Views,
-		FRDGTextureRef SceneColorTexture,
-		FRDGTextureRef SceneDepthTexture,
-		TRDGUniformBufferRef<FSceneTextureUniformParameters> SceneTexturesUniformBuffer,
+		const FSceneTextures& SceneTextures,
 		const FSceneRenderTargets& SceneContext);
 
 	void RenderOverlayExtensions(
 		FRDGBuilder& GraphBuilder,
 		TArrayView<const FViewInfo> Views,
-		FRDGTextureRef SceneColorTexture,
-		FRDGTextureRef SceneDepthTexture,
+		const FSceneTextures& SceneTextures,
 		const FSceneRenderTargets& SceneContext);
 
 	void RenderPostResolvedSceneColorExtension(FRDGBuilder& GraphBuilder, FSceneRenderTargets& SceneContext);
diff --git a/Engine/Source/Runtime/Renderer/Private/SceneOcclusion.cpp b/Engine/Source/Runtime/Renderer/Private/SceneOcclusion.cpp
--- a/Engine/Source/Runtime/Renderer/Private/SceneOcclusion.cpp
+++ b/Engine/Source/Runtime/Renderer/Private/SceneOcclusion.cpp
@@ -1328,21 +1328,17 @@ static void BeginOcclusionTests(
 
 void FDeferredShadingSceneRenderer::RenderOcclusion(
 	FRDGBuilder& GraphBuilder,
-	FRDGTextureRef SceneDepthTexture,
-	FRDGTextureRef SmallDepthTexture,
-	TRDGUniformBufferRef<FSceneTextureUniformParameters> SceneTexturesUniformBuffer,
+	const FSceneTextures& SceneTextures,
 	bool bIsOcclusionTesting)
 {
 	if (bIsOcclusionTesting)
 	{
-		check(SceneDepthTexture);
-		check(SmallDepthTexture);
-
 		RDG_GPU_STAT_SCOPE(GraphBuilder, HZB);
 
 		FSceneRenderTargets& SceneContext = FSceneRenderTargets::Get();
 
 		uint32 DownsampleFactor = 1;
+		FRDGTextureRef SceneDepthTexture = SceneTextures.Depth.Target;
 		FRDGTextureRef OcclusionDepthTexture = SceneDepthTexture;
 
 		// Update the quarter-sized depth buffer with the current contents of the scene depth texture.
@@ -1350,7 +1346,7 @@ void FDeferredShadingSceneRenderer::RenderOcclusion(
 		if (SceneContext.UseDownsizedOcclusionQueries())
 		{
 			DownsampleFactor = SceneContext.GetSmallColorDepthDownsampleFactor();
-			OcclusionDepthTexture = SmallDepthTexture;
+			OcclusionDepthTexture = SceneTextures.SmallDepth;
 
 			for (int32 ViewIndex = 0; ViewIndex < Views.Num(); ViewIndex++)
 			{
@@ -1358,7 +1354,7 @@ void FDeferredShadingSceneRenderer::RenderOcclusion(
 				RDG_GPU_MASK_SCOPE(GraphBuilder, View.GPUMask);
 
 				const FScreenPassTexture SceneDepth(SceneDepthTexture, View.ViewRect);
-				const FScreenPassRenderTarget SmallDepth(SmallDepthTexture, GetDownscaledRect(View.ViewRect, DownsampleFactor), ERenderTargetLoadAction::ELoad);
+				const FScreenPassRenderTarget SmallDepth(SceneTextures.SmallDepth, GetDownscaledRect(View.ViewRect, DownsampleFactor), ERenderTargetLoadAction::ELoad);
 				AddDownsampleDepthPass(GraphBuilder, View, SceneDepth, SmallDepth, EDownsampleDepthFilter::Max);
 			}
 		}
@@ -1406,7 +1402,7 @@ void FDeferredShadingSceneRenderer::RenderOcclusion(
 		}
 	}
 
-	const bool bUseHzbOcclusion = RenderHzb(GraphBuilder, SceneTexturesUniformBuffer);
+	const bool bUseHzbOcclusion = RenderHzb(GraphBuilder, SceneTextures);
 
 	if (bUseHzbOcclusion || bIsOcclusionTesting)
 	{
diff --git a/Engine/Source/Runtime/Renderer/Private/SceneRendering.cpp b/Engine/Source/Runtime/Renderer/Private/SceneRendering.cpp
--- a/Engine/Source/Runtime/Renderer/Private/SceneRendering.cpp
+++ b/Engine/Source/Runtime/Renderer/Private/SceneRendering.cpp
@@ -4004,9 +4004,7 @@ void FRendererModule::RemoveOverlayRenderDelegate(FDelegateHandle InOverlayRende
 void FRendererModule::RenderPostOpaqueExtensions(
 	FRDGBuilder& GraphBuilder,
 	TArrayView<const FViewInfo> Views,
-	FRDGTextureRef SceneColorTexture,
-	FRDGTextureRef SceneDepthTexture,
-	TRDGUniformBufferRef<FSceneTextureUniformParameters> SceneTexturesUniformBuffer,
+	const FSceneTextures& SceneTextures,
 	const FSceneRenderTargets& SceneContext)
 {
 	if (PostOpaqueRenderDelegate.IsBound())
@@ -4023,13 +4021,13 @@ void FRendererModule::RenderPostOpaqueExtensions(
 			FPostOpaqueRenderParameters RenderParameters;
 			RenderParameters.ViewMatrix = View.ViewMatrices.GetViewMatrix();
 			RenderParameters.ProjMatrix = View.ViewMatrices.GetProjectionMatrix();
-			RenderParameters.ColorTexture = SceneColorTexture;
-			RenderParameters.DepthTexture = SceneDepthTexture;
+			RenderParameters.ColorTexture = SceneTextures.Color.Target;
+			RenderParameters.DepthTexture = SceneTextures.Depth.Target;
 			RenderParameters.NormalTexture = TryRegisterExternalTexture(GraphBuilder, SceneContext.GBufferA);
-			RenderParameters.VelocityTexture = TryRegisterExternalTexture(GraphBuilder, SceneContext.SceneVelocity);
-			RenderParameters.SmallDepthTexture = TryRegisterExternalTexture(GraphBuilder, SceneContext.SmallDepthZ);
+			RenderParameters.VelocityTexture = SceneTextures.Velocity;
+			RenderParameters.SmallDepthTexture = SceneTextures.SmallDepth;
 			RenderParameters.ViewUniformBuffer = View.ViewUniformBuffer;
-			RenderParameters.SceneTexturesUniformParams = SceneTexturesUniformBuffer;
+			RenderParameters.SceneTexturesUniformParams = SceneTextures.UniformBuffer;
 			RenderParameters.GlobalDistanceFieldParams = &View.GlobalDistanceFieldInfo.ParameterData;
 
 			RenderParameters.ViewportRect = View.ViewRect;
@@ -4044,8 +4042,7 @@ void FRendererModule::RenderPostOpaqueExtensions(
 void FRendererModule::RenderOverlayExtensions(
 	FRDGBuilder& GraphBuilder,
 	TArrayView<const FViewInfo> Views,
-	FRDGTextureRef SceneColorTexture,
-	FRDGTextureRef SceneDepthTexture,
+	const FSceneTextures& SceneTextures,
 	const FSceneRenderTargets& SceneContext)
 {
 	if (OverlayRenderDelegate.IsBound())
@@ -4061,9 +4058,9 @@ void FRendererModule::RenderOverlayExtensions(
 			FPostOpaqueRenderParameters RenderParameters;
 			RenderParameters.ViewMatrix = View.ViewMatrices.GetViewMatrix();
 			RenderParameters.ProjMatrix = View.ViewMatrices.GetProjectionMatrix();
-			RenderParameters.ColorTexture = SceneColorTexture;
-			RenderParameters.DepthTexture = SceneDepthTexture;
-			RenderParameters.SmallDepthTexture = TryRegisterExternalTexture(GraphBuilder, SceneContext.SmallDepthZ);
+			RenderParameters.ColorTexture = SceneTextures.Color.Target;
+			RenderParameters.DepthTexture = SceneTextures.Depth.Target;
+			RenderParameters.SmallDepthTexture = SceneTextures.SmallDepth;
 
 			RenderParameters.ViewportRect = View.ViewRect;
 			RenderParameters.GraphBuilder = &GraphBuilder;
diff --git a/Engine/Source/Runtime/Renderer/Private/SceneRendering.h b/Engine/Source/Runtime/Renderer/Private/SceneRendering.h
--- a/Engine/Source/Runtime/Renderer/Private/SceneRendering.h
+++ b/Engine/Source/Runtime/Renderer/Private/SceneRendering.h
@@ -54,6 +54,9 @@ struct FHairStrandsVisibilityViews;
 struct FSortedLightSetSceneInfo;
 struct FStrataSceneData;
 
+struct FMinimalSceneTextures;
+struct FSceneTextures;
+
 DECLARE_STATS_GROUP(TEXT("Command List Markers"), STATGROUP_CommandListMarkers, STATCAT_Advanced);
 
 DECLARE_GPU_DRAWCALL_STAT_EXTERN(VirtualTextureUpdate);
@@ -1989,11 +1992,7 @@ protected:
 	void RenderSkyAtmosphereLookUpTables(FRDGBuilder& GraphBuilder);
 
 	/** Render the sky atmosphere over the scene.*/
-	void RenderSkyAtmosphere(
-		FRDGBuilder& GraphBuilder,
-		TRDGUniformBufferRef<FSceneTextureUniformParameters> SceneTexturesUniformBuffer,
-		FRDGTextureRef SceneColorTexture,
-		FRDGTextureRef SceneDepthTexture);
+	void RenderSkyAtmosphere(FRDGBuilder& GraphBuilder, const FMinimalSceneTextures& SceneTextures);
 
 	/** Initialise volumetric cloud resources.*/
 	void InitVolumetricCloudsForViews(FRDGBuilder& GraphBuilder, bool bShouldRenderVolumetricCloud);
@@ -2001,11 +2000,9 @@ protected:
 	/** Render volumetric cloud. */
 	void RenderVolumetricCloud(
 		FRDGBuilder& GraphBuilder,
-		const FSceneTextureShaderParameters& SceneTextures,
+		const FMinimalSceneTextures& SceneTextures,
 		bool bSkipVolumetricRenderTarget,
 		bool bSkipPerPixelTracing,
-		FRDGTextureMSAA SceneColorTexture,
-		FRDGTextureMSAA SceneDepthTexture,
 		FRDGTextureRef HalfResolutionDepthCheckerboardMinMaxTexture);
 
 	/** Render notification to artist when a sky material is used but it might comtains the camera (and then the sky/background would look black).*/
diff --git a/Engine/Source/Runtime/Renderer/Private/ShadowRendering.cpp b/Engine/Source/Runtime/Renderer/Private/ShadowRendering.cpp
--- a/Engine/Source/Runtime/Renderer/Private/ShadowRendering.cpp
+++ b/Engine/Source/Runtime/Renderer/Private/ShadowRendering.cpp
@@ -1708,10 +1708,9 @@ void FSceneRenderer::RenderShadowProjections(
 
 void FDeferredShadingSceneRenderer::RenderShadowProjections(
 	FRDGBuilder& GraphBuilder,
-	TRDGUniformBufferRef<FSceneTextureUniformParameters> SceneTexturesUniformBuffer,
+	const FMinimalSceneTextures& SceneTextures,
 	FRDGTextureRef ScreenShadowMaskTexture,
 	FRDGTextureRef ScreenShadowMaskSubPixelTexture,
-	FRDGTextureRef SceneDepthTexture,
 	const FLightSceneInfo* LightSceneInfo,
 	const FHairStrandsVisibilityViews* HairVisibilityViews,
 	bool bProjectingForForwardShading)
@@ -1756,10 +1755,10 @@ void FDeferredShadingSceneRenderer::RenderShadowProjections(
 			RDG_EVENT_SCOPE(GraphBuilder, "%s", *LightNameWithLevel);
 
 			FShadowProjectionPassParameters CommonPassParameters;
-			CommonPassParameters.SceneTextures = GetSceneTextureShaderParameters(SceneTexturesUniformBuffer);
+			CommonPassParameters.SceneTextures = GetSceneTextureShaderParameters(SceneTextures.UniformBuffer);
 			CommonPassParameters.HairCategorizationTexture = bSubPixel && HairVisibilityViews->HairDatas.Num() > 0 ? HairVisibilityViews->HairDatas[0].CategorizationTexture : nullptr;
 			CommonPassParameters.RenderTargets[0] = FRenderTargetBinding(OutputTexture, ERenderTargetLoadAction::ELoad);
-			CommonPassParameters.RenderTargets.DepthStencil = FDepthStencilBinding(SceneDepthTexture, ERenderTargetLoadAction::ELoad, ERenderTargetLoadAction::ELoad, ExclusiveDepthStencil);
+			CommonPassParameters.RenderTargets.DepthStencil = FDepthStencilBinding(SceneTextures.Depth.Target, ERenderTargetLoadAction::ELoad, ERenderTargetLoadAction::ELoad, ExclusiveDepthStencil);
 
 			const bool bMobileModulatedProjections = false;
 			FSceneRenderer::RenderShadowProjections(
@@ -1796,7 +1795,7 @@ void FDeferredShadingSceneRenderer::RenderShadowProjections(
 		{
 			RDG_EVENT_SCOPE(GraphBuilder, "Virtual Shadow Maps" );
 
-			FSceneTextureParameters SceneTextures = GetSceneTextureParameters(GraphBuilder, SceneTexturesUniformBuffer);
+			FSceneTextureParameters SceneTextureParameters = GetSceneTextureParameters(GraphBuilder, SceneTextures.UniformBuffer);
 
 			for (int32 ViewIndex = 0; ViewIndex < Views.Num(); ViewIndex++)
 			{
@@ -1819,7 +1818,7 @@ void FDeferredShadingSceneRenderer::RenderShadowProjections(
 						const FLinearColor ClearColor(0.0f, 0.0f, 0.0f, 0.0f);
 						{
 							FRDGTextureDesc Desc = FRDGTextureDesc::Create2D(
-								SceneTextures.SceneDepthTexture->Desc.Extent,
+								SceneTextures.Extent,
 								PF_FloatRGBA,
 								FClearValueBinding(ClearColor),
 								TexCreate_ShaderResource | TexCreate_RenderTargetable);
@@ -1865,7 +1864,7 @@ void FDeferredShadingSceneRenderer::RenderShadowProjections(
 								GraphBuilder,
 								View,
 								&View.PrevViewInfo,
-								SceneTextures,
+								SceneTextureParameters,
 								LightSceneInfo,
 								ScissorRect,
 								Inputs);
@@ -1930,9 +1929,8 @@ void FDeferredShadingSceneRenderer::RenderShadowProjections(
 					const FProjectedShadowInfo* ProjectedShadowInfo = DistanceFieldShadows[ShadowIndex];
 					ProjectedShadowInfo->RenderRayTracedDistanceFieldProjection(
 						GraphBuilder,
-						SceneTexturesUniformBuffer,
+						SceneTextures,
 						ScreenShadowMaskTexture,
-						SceneDepthTexture,
 						View,
 						ScissorRect,
 						bProjectingForForwardShading);
@@ -1944,12 +1942,11 @@ void FDeferredShadingSceneRenderer::RenderShadowProjections(
 
 void FDeferredShadingSceneRenderer::RenderDeferredShadowProjections(
 	FRDGBuilder& GraphBuilder,
-	TRDGUniformBufferRef<FSceneTextureUniformParameters> SceneTexturesUniformBuffer,
+	const FMinimalSceneTextures& SceneTextures,
 	const FTranslucentVolumeLightingTextures& TranslucentVolumeLightingTextures,
 	const FLightSceneInfo* LightSceneInfo,
 	FRDGTextureRef ScreenShadowMaskTexture,
 	FRDGTextureRef ScreenShadowMaskSubPixelTexture,
-	FRDGTextureRef SceneDepthTexture,
 	const FHairStrandsRenderingData* HairDatas,
 	bool& bInjectedTranslucentVolume)
 {
@@ -1962,7 +1959,7 @@ void FDeferredShadingSceneRenderer::RenderDeferredShadowProjections(
 	const FHairStrandsVisibilityViews* HairVisibilityViews = HairDatas ? &HairDatas->HairVisibilityViews : nullptr;
 
 	const bool bProjectingForForwardShading = false;
-	RenderShadowProjections(GraphBuilder, SceneTexturesUniformBuffer, ScreenShadowMaskTexture, ScreenShadowMaskSubPixelTexture, SceneDepthTexture, LightSceneInfo, HairVisibilityViews, bProjectingForForwardShading);
+	RenderShadowProjections(GraphBuilder, SceneTextures, ScreenShadowMaskTexture, ScreenShadowMaskSubPixelTexture, LightSceneInfo, HairVisibilityViews, bProjectingForForwardShading);
 
 	// Perform injection on tranlucent lighting volume
 	{
@@ -2015,7 +2012,7 @@ void FDeferredShadingSceneRenderer::RenderDeferredShadowProjections(
 		}
 	}
 
-	RenderCapsuleDirectShadows(GraphBuilder, SceneTexturesUniformBuffer, *LightSceneInfo, ScreenShadowMaskTexture, VisibleLightInfo.CapsuleShadowsToProject, bProjectingForForwardShading);
+	RenderCapsuleDirectShadows(GraphBuilder, SceneTextures.UniformBuffer, *LightSceneInfo, ScreenShadowMaskTexture, VisibleLightInfo.CapsuleShadowsToProject, bProjectingForForwardShading);
 
 	for (int32 ViewIndex = 0; ViewIndex < Views.Num(); ViewIndex++)
 	{
diff --git a/Engine/Source/Runtime/Renderer/Private/ShadowRendering.h b/Engine/Source/Runtime/Renderer/Private/ShadowRendering.h
--- a/Engine/Source/Runtime/Renderer/Private/ShadowRendering.h
+++ b/Engine/Source/Runtime/Renderer/Private/ShadowRendering.h
@@ -505,9 +505,8 @@ public:
 	/** Renders ray traced distance field shadows. */
 	void RenderRayTracedDistanceFieldProjection(
 		FRDGBuilder& GraphBuilder,
-		TRDGUniformBufferRef<FSceneTextureUniformParameters> SceneTexturesUniformBuffer,
+		const FMinimalSceneTextures& SceneTextures,
 		FRDGTextureRef ScreenShadowMaskTexture,
-		FRDGTextureRef SceneDepthTexture,
 		const FViewInfo& View,
 		FIntRect ScissorRect,
 		bool bProjectingForForwardShading) const;
diff --git a/Engine/Source/Runtime/Renderer/Private/SingleLayerWaterRendering.cpp b/Engine/Source/Runtime/Renderer/Private/SingleLayerWaterRendering.cpp
--- a/Engine/Source/Runtime/Renderer/Private/SingleLayerWaterRendering.cpp
+++ b/Engine/Source/Runtime/Renderer/Private/SingleLayerWaterRendering.cpp
@@ -643,9 +643,7 @@ void FDeferredShadingSceneRenderer::RenderSingleLayerWaterReflections(
 
 void FDeferredShadingSceneRenderer::RenderSingleLayerWater(
 	FRDGBuilder& GraphBuilder,
-	FRDGTextureMSAA SceneColorTexture,
-	FRDGTextureMSAA SceneDepthTexture,
-	TRDGUniformBufferRef<FSceneTextureUniformParameters> DepthOnlySceneTextures,
+	const FSceneTextures& SceneTextures,
 	bool bShouldRenderVolumetricCloud,
 	FSceneWithoutWaterTextures& SceneWithoutWaterTextures)
 {
@@ -653,12 +651,12 @@ void FDeferredShadingSceneRenderer::RenderSingleLayerWater(
 	RDG_GPU_STAT_SCOPE(GraphBuilder, SingleLayerWater);
 
 	// Copy the texture to be available for the water surface to refract
-	SceneWithoutWaterTextures = AddCopySceneWithoutWaterPass(GraphBuilder, Views, SceneColorTexture.Resolve, SceneDepthTexture.Resolve);
+	SceneWithoutWaterTextures = AddCopySceneWithoutWaterPass(GraphBuilder, Views, SceneTextures.Color.Resolve, SceneTextures.Depth.Resolve);
 
 	// Render height fog over the color buffer if it is allocated, e.g. SingleLayerWaterUsesSimpleShading is true which is not the case on Switch.
 	if (SceneWithoutWaterTextures.ColorTexture && ShouldRenderFog(ViewFamily))
 	{
-		RenderUnderWaterFog(GraphBuilder, SceneWithoutWaterTextures, DepthOnlySceneTextures);
+		RenderUnderWaterFog(GraphBuilder, SceneWithoutWaterTextures, SceneTextures.UniformBuffer);
 	}
 	if (SceneWithoutWaterTextures.ColorTexture && bShouldRenderVolumetricCloud)
 	{
@@ -666,13 +664,13 @@ void FDeferredShadingSceneRenderer::RenderSingleLayerWater(
 		ComposeVolumetricRenderTargetOverSceneUnderWater(GraphBuilder, Views, SceneWithoutWaterTextures);
 	}
 
-	RenderSingleLayerWaterInner(GraphBuilder, SceneColorTexture, SceneDepthTexture, SceneWithoutWaterTextures);
+	RenderSingleLayerWaterInner(GraphBuilder, SceneTextures, SceneWithoutWaterTextures);
 
 	// No SSR or composite needed in Forward. Reflections are applied in the WaterGBuffer pass.
 	if (!IsAnyForwardShadingEnabled(ShaderPlatform))
 	{
 		// If supported render SSR, the composite pass in non deferred and/or under water effect.
-		RenderSingleLayerWaterReflections(GraphBuilder, SceneColorTexture.Resolve, SceneWithoutWaterTextures);
+		RenderSingleLayerWaterReflections(GraphBuilder, SceneTextures.Color.Resolve, SceneWithoutWaterTextures);
 	}
 }
 
@@ -683,8 +681,7 @@ END_SHADER_PARAMETER_STRUCT()
 
 void FDeferredShadingSceneRenderer::RenderSingleLayerWaterInner(
 	FRDGBuilder& GraphBuilder,
-	FRDGTextureMSAA SceneColorTexture,
-	FRDGTextureMSAA SceneDepthTexture,
+	const FSceneTextures& SceneTextures,
 	const FSceneWithoutWaterTextures& SceneWithoutWaterTextures)
 {
 	RDG_CSV_STAT_EXCLUSIVE_SCOPE(GraphBuilder, Water);
@@ -696,12 +693,13 @@ void FDeferredShadingSceneRenderer::RenderSingleLayerWaterInner(
 	const bool bRenderInParallel = GRHICommandList.UseParallelAlgorithms() && CVarParallelSingleLayerWaterPass.GetValueOnRenderThread() == 1;
 
 	FSceneRenderTargets& SceneContext = FSceneRenderTargets::Get();
+	const FRDGSystemTextures& SystemTextures = FRDGSystemTextures::Get(GraphBuilder);
 
 	FRenderTargetBindingSlots RenderTargets;
 	SceneContext.GetGBufferRenderTargets(GraphBuilder, ERenderTargetLoadAction::ELoad, RenderTargets);
-	RenderTargets.DepthStencil = FDepthStencilBinding(SceneDepthTexture.Target, ERenderTargetLoadAction::ELoad, ERenderTargetLoadAction::ELoad, FExclusiveDepthStencil::DepthWrite_StencilWrite);
+	RenderTargets.DepthStencil = FDepthStencilBinding(SceneTextures.Depth.Target, ERenderTargetLoadAction::ELoad, ERenderTargetLoadAction::ELoad, FExclusiveDepthStencil::DepthWrite_StencilWrite);
 
-	FRDGTextureRef WhiteForwardScreenSpaceShadowMask = GraphBuilder.RegisterExternalTexture(GSystemTextures.WhiteDummy);
+	FRDGTextureRef WhiteForwardScreenSpaceShadowMask = SystemTextures.White;
 
 	for (int32 ViewIndex = 0; ViewIndex < Views.Num(); ++ViewIndex)
 	{
@@ -747,7 +745,7 @@ void FDeferredShadingSceneRenderer::RenderSingleLayerWaterInner(
 		}
 	}
 
-	AddResolveSceneDepthPass(GraphBuilder, Views, SceneDepthTexture);
+	AddResolveSceneDepthPass(GraphBuilder, Views, SceneTextures.Depth);
 }
 
 class FSingleLayerWaterPassMeshProcessor : public FMeshPassProcessor
diff --git a/Engine/Source/Runtime/Renderer/Private/SkyAtmosphereRendering.cpp b/Engine/Source/Runtime/Renderer/Private/SkyAtmosphereRendering.cpp
--- a/Engine/Source/Runtime/Renderer/Private/SkyAtmosphereRendering.cpp
+++ b/Engine/Source/Runtime/Renderer/Private/SkyAtmosphereRendering.cpp
@@ -1755,11 +1755,7 @@ void FSceneRenderer::RenderSkyAtmosphereInternal(
 
 
 
-void FSceneRenderer::RenderSkyAtmosphere(
-	FRDGBuilder& GraphBuilder,
-	TRDGUniformBufferRef<FSceneTextureUniformParameters> SceneTexturesUniformBuffer,
-	FRDGTextureRef SceneColor,
-	FRDGTextureRef SceneDepth)
+void FSceneRenderer::RenderSkyAtmosphere(FRDGBuilder& GraphBuilder, const FMinimalSceneTextures& SceneTextures)
 {
 	check(!IsMobilePlatform(Scene->GetShaderPlatform()));
 
@@ -1791,8 +1787,8 @@ void FSceneRenderer::RenderSkyAtmosphere(
 	SkyRC.TransmittanceLut = GraphBuilder.RegisterExternalTexture(SkyInfo.GetTransmittanceLutTexture());
 	SkyRC.MultiScatteredLuminanceLut = GraphBuilder.RegisterExternalTexture(SkyInfo.GetMultiScatteredLuminanceLutTexture());
 
-	SkyRC.RenderTargets[0] = FRenderTargetBinding(SceneColor, ERenderTargetLoadAction::ELoad);
-	SkyRC.RenderTargets.DepthStencil = FDepthStencilBinding(SceneDepth, ERenderTargetLoadAction::ELoad, ERenderTargetLoadAction::ENoAction, FExclusiveDepthStencil::DepthRead_StencilNop);
+	SkyRC.RenderTargets[0] = FRenderTargetBinding(SceneTextures.Color.Target, ERenderTargetLoadAction::ELoad);
+	SkyRC.RenderTargets.DepthStencil = FDepthStencilBinding(SceneTextures.Depth.Target, ERenderTargetLoadAction::ELoad, ERenderTargetLoadAction::ENoAction, FExclusiveDepthStencil::DepthRead_StencilNop);
 
 	for (int32 ViewIndex = 0; ViewIndex < Views.Num(); ViewIndex++)
 	{
@@ -1827,7 +1823,7 @@ void FSceneRenderer::RenderSkyAtmosphere(
 		SkyRC.bShouldSampleCloudSkyAO = CloudShadowAOData.bShouldSampleCloudSkyAO;
 		SkyRC.VolumetricCloudSkyAO = CloudShadowAOData.VolumetricCloudSkyAO;
 
-		RenderSkyAtmosphereInternal(GraphBuilder, GetSceneTextureShaderParameters(SceneTexturesUniformBuffer), SkyRC);
+		RenderSkyAtmosphereInternal(GraphBuilder, GetSceneTextureShaderParameters(SceneTextures.UniformBuffer), SkyRC);
 	}
 
 #if WITH_EDITOR
@@ -1837,7 +1833,7 @@ void FSceneRenderer::RenderSkyAtmosphere(
 		{
 			const FViewInfo& View = Views[ViewIndex];
 
-			AddDrawCanvasPass(GraphBuilder, {}, View, FScreenPassRenderTarget(SceneColor, View.ViewRect, ERenderTargetLoadAction::ELoad),
+			AddDrawCanvasPass(GraphBuilder, {}, View, FScreenPassRenderTarget(SceneTextures.Color.Target, View.ViewRect, ERenderTargetLoadAction::ELoad),
 				[&View](FCanvas& Canvas)
 			{
 				const float ViewPortWidth = float(View.ViewRect.Width());
diff --git a/Engine/Source/Runtime/Renderer/Private/TiledDeferredLightRendering.cpp b/Engine/Source/Runtime/Renderer/Private/TiledDeferredLightRendering.cpp
--- a/Engine/Source/Runtime/Renderer/Private/TiledDeferredLightRendering.cpp
+++ b/Engine/Source/Runtime/Renderer/Private/TiledDeferredLightRendering.cpp
@@ -335,10 +335,9 @@ BEGIN_SHADER_PARAMETER_STRUCT(FTiledDeferredLightingParameters, )
 	SHADER_PARAMETER_RDG_UNIFORM_BUFFER(FSceneTextureUniformParameters, SceneTextures)
 END_SHADER_PARAMETER_STRUCT()
 
-FRDGTextureRef FDeferredShadingSceneRenderer::RenderTiledDeferredLighting(
+void FDeferredShadingSceneRenderer::RenderTiledDeferredLighting(
 	FRDGBuilder& GraphBuilder,
-	FRDGTextureRef SceneColorTexture,
-	TRDGUniformBufferRef<FSceneTextureUniformParameters> SceneTexturesUniformBuffer,
+	FMinimalSceneTextures& SceneTextures,
 	const TArray<FSortedLightSceneInfo, SceneRenderingAllocator>& SortedLights,
 	int32 TiledDeferredLightsStart,
 	int32 TiledDeferredLightsEnd,
@@ -358,6 +357,8 @@ FRDGTextureRef FDeferredShadingSceneRenderer::RenderTiledDeferredLighting(
 		INC_DWORD_STAT_BY(STAT_NumLightsUsingSimpleTiledDeferred, SimpleLights.InstanceData.Num());
 		SCOPE_CYCLE_COUNTER(STAT_DirectLightRenderingTime);
 
+		FRDGTextureRef& SceneColorTexture = SceneTextures.Color.Target;
+
 		// Determine how many compute shader passes will be needed to process all the lights
 		const int32 NumPassesNeeded = FMath::DivideAndRoundUp(NumLightsToRender, GMaxNumTiledDeferredLights);
 		for (int32 PassIndex = 0; PassIndex < NumPassesNeeded; PassIndex++)
@@ -370,7 +371,7 @@ FRDGTextureRef FDeferredShadingSceneRenderer::RenderTiledDeferredLighting(
 			FRDGTextureRef SceneColorOutputTexture = nullptr;
 
 			{
-				FRDGTextureDesc Desc = SceneColorTexture->Desc;
+				FRDGTextureDesc Desc = SceneTextures.Color.Target->Desc;
 				Desc.Flags |= TexCreate_UAV;
 				SceneColorOutputTexture = GraphBuilder.CreateTexture(Desc, TEXT("SceneColorTiled"));
 			}
@@ -384,7 +385,7 @@ FRDGTextureRef FDeferredShadingSceneRenderer::RenderTiledDeferredLighting(
 				FTiledDeferredLightingParameters* PassParameters = GraphBuilder.AllocParameters<FTiledDeferredLightingParameters>();
 				PassParameters->SceneColorInput = SceneColorTexture;
 				PassParameters->SceneColorOutput = SceneColorOutputTexture;
-				PassParameters->SceneTextures = SceneTexturesUniformBuffer;
+				PassParameters->SceneTextures = SceneTextures.UniformBuffer;
 
 				GraphBuilder.AddPass(
 					RDG_EVENT_NAME("TiledDeferredLighting"),
@@ -409,6 +410,4 @@ FRDGTextureRef FDeferredShadingSceneRenderer::RenderTiledDeferredLighting(
 		FSceneRenderTargets& SceneContext = FSceneRenderTargets::Get();
 		ConvertToExternalTexture(GraphBuilder, SceneColorTexture, SceneContext.GetSceneColor());
 	}
-
-	return SceneColorTexture;
 }
\ No newline at end of file
diff --git a/Engine/Source/Runtime/Renderer/Private/TranslucentRendering.cpp b/Engine/Source/Runtime/Renderer/Private/TranslucentRendering.cpp
--- a/Engine/Source/Runtime/Renderer/Private/TranslucentRendering.cpp
+++ b/Engine/Source/Runtime/Renderer/Private/TranslucentRendering.cpp
@@ -622,7 +622,7 @@ void SetupTranslucentBasePassUniformParameters(
 	};
 
 	SetupSharedBasePassParameters(GraphBuilder, View, BasePassParameters.Shared);
-	SetupSceneTextureUniformParameters(&GraphBuilder, View.FeatureLevel, SceneRenderTargets, SceneTextureSetupMode, BasePassParameters.SceneTextures);
+	SetupSceneTextureUniformParameters(GraphBuilder, View.FeatureLevel, SceneTextureSetupMode, BasePassParameters.SceneTextures);
 
 	FRDGTextureRef BlackDummyTexture = GetRDG(GSystemTextures.BlackDummy);
 
@@ -960,9 +960,8 @@ static void RenderTranslucencyViewInner(
 
 void FDeferredShadingSceneRenderer::RenderTranslucencyInner(
 	FRDGBuilder& GraphBuilder,
+	const FMinimalSceneTextures& SceneTextures,
 	const FTranslucentVolumeLightingTextures& TranslucentVolumeLightingTextures,
-	FRDGTextureMSAA SceneColorTexture,
-	FRDGTextureMSAA SceneDepthTexture,
 	FSeparateTranslucencyTextures* OutSeparateTranslucencyTextures,
 	ETranslucencyView ViewsToRender,
 	FRDGTextureRef SceneColorCopyTexture,
@@ -1024,7 +1023,7 @@ void FDeferredShadingSceneRenderer::RenderTranslucencyInner(
 			}
 
 			// When rendering to a 1-to-1 separate translucency target, use the existing scene depth.
-			FRDGTextureMSAA SeparateTranslucencyDepthTexture = SceneDepthTexture;
+			FRDGTextureMSAA SeparateTranslucencyDepthTexture = SceneTextures.Depth;
 
 			// Rendering to a downscaled target; allocate a new depth texture and downsample depth.
 			if (SeparateTranslucencyDimensions.Scale < 1.0f)
@@ -1036,7 +1035,7 @@ void FDeferredShadingSceneRenderer::RenderTranslucencyInner(
 
 				AddDownsampleDepthPass(
 					GraphBuilder, View,
-					FScreenPassTexture(SceneDepthTexture.Resolve, View.ViewRect),
+					FScreenPassTexture(SceneTextures.Depth.Resolve, View.ViewRect),
 					FScreenPassRenderTarget(SeparateTranslucencyDepthTexture.Target, SeparateTranslucencyViewport.Rect, ERenderTargetLoadAction::ENoAction),
 					EDownsampleDepthFilter::Point);
 			}
@@ -1067,10 +1066,10 @@ void FDeferredShadingSceneRenderer::RenderTranslucencyInner(
 
 				AddTranslucencyUpsamplePass(
 					GraphBuilder, View,
-					FScreenPassRenderTarget(SceneColorTexture.Target, View.ViewRect, ERenderTargetLoadAction::ELoad),
+					FScreenPassRenderTarget(SceneTextures.Color.Target, View.ViewRect, ERenderTargetLoadAction::ELoad),
 					FScreenPassTexture(SeparateTranslucencyColorTexture.Resolve, SeparateTranslucencyViewport.Rect),
 					SeparateTranslucencyDepthTexture.Resolve,
-					SceneDepthTexture.Resolve,
+					SceneTextures.Depth.Resolve,
 					SeparateTranslucencyDimensions.Scale);
 			}
 
@@ -1095,7 +1094,7 @@ void FDeferredShadingSceneRenderer::RenderTranslucencyInner(
 			AddBeginTranslucencyTimerPass(GraphBuilder, View);
 
 			const ERenderTargetLoadAction SceneColorLoadAction = ERenderTargetLoadAction::ELoad;
-			const FScreenPassTextureViewport Viewport(SceneColorTexture.Target, View.ViewRect);
+			const FScreenPassTextureViewport Viewport(SceneTextures.Color.Target, View.ViewRect);
 			const float ViewportScale = 1.0f;
 			const bool bResolveColorTexture = false;
 
@@ -1105,9 +1104,9 @@ void FDeferredShadingSceneRenderer::RenderTranslucencyInner(
 				View,
 				Viewport,
 				ViewportScale,
-				SceneColorTexture,
+				SceneTextures.Color,
 				SceneColorLoadAction,
-				SceneDepthTexture.Target,
+				SceneTextures.Depth.Target,
 				CreateTranslucentBasePassUniformBuffer(GraphBuilder, View, &TranslucentVolumeLightingTextures, SceneColorCopyTexture, SceneTextureSetupMode, ViewIndex),
 				TranslucencyPass,
 				bResolveColorTexture,
@@ -1120,9 +1119,8 @@ void FDeferredShadingSceneRenderer::RenderTranslucencyInner(
 
 void FDeferredShadingSceneRenderer::RenderTranslucency(
 	FRDGBuilder& GraphBuilder,
+	const FMinimalSceneTextures& SceneTextures,
 	const FTranslucentVolumeLightingTextures& TranslucentVolumeLightingTextures,
-	FRDGTextureMSAA SceneColorTexture,
-	FRDGTextureMSAA SceneDepthTexture,
 	FSeparateTranslucencyTextures* OutSeparateTranslucencyTextures,
 	ETranslucencyView ViewsToRender)
 {
@@ -1137,17 +1135,17 @@ void FDeferredShadingSceneRenderer::RenderTranslucency(
 
 	if (EnumHasAnyFlags(ViewsToRender, ETranslucencyView::AboveWater))
 	{
-		SceneColorCopyTexture = AddCopySceneColorPass(GraphBuilder, Views, SceneColorTexture);
+		SceneColorCopyTexture = AddCopySceneColorPass(GraphBuilder, Views, SceneTextures.Color);
 	}
 
 	if (ViewFamily.AllowTranslucencyAfterDOF())
 	{
-		RenderTranslucencyInner(GraphBuilder, TranslucentVolumeLightingTextures, SceneColorTexture, SceneDepthTexture, OutSeparateTranslucencyTextures, ViewsToRender, SceneColorCopyTexture, ETranslucencyPass::TPT_StandardTranslucency);
-		RenderTranslucencyInner(GraphBuilder, TranslucentVolumeLightingTextures, SceneColorTexture, SceneDepthTexture, OutSeparateTranslucencyTextures, ViewsToRender, SceneColorCopyTexture, ETranslucencyPass::TPT_TranslucencyAfterDOF);
-		RenderTranslucencyInner(GraphBuilder, TranslucentVolumeLightingTextures, SceneColorTexture, SceneDepthTexture, OutSeparateTranslucencyTextures, ViewsToRender, SceneColorCopyTexture, ETranslucencyPass::TPT_TranslucencyAfterDOFModulate);
+		RenderTranslucencyInner(GraphBuilder, SceneTextures, TranslucentVolumeLightingTextures, OutSeparateTranslucencyTextures, ViewsToRender, SceneColorCopyTexture, ETranslucencyPass::TPT_StandardTranslucency);
+		RenderTranslucencyInner(GraphBuilder, SceneTextures, TranslucentVolumeLightingTextures, OutSeparateTranslucencyTextures, ViewsToRender, SceneColorCopyTexture, ETranslucencyPass::TPT_TranslucencyAfterDOF);
+		RenderTranslucencyInner(GraphBuilder, SceneTextures, TranslucentVolumeLightingTextures, OutSeparateTranslucencyTextures, ViewsToRender, SceneColorCopyTexture, ETranslucencyPass::TPT_TranslucencyAfterDOFModulate);
 	}
 	else // Otherwise render translucent primitives in a single bucket.
 	{
-		RenderTranslucencyInner(GraphBuilder, TranslucentVolumeLightingTextures, SceneColorTexture, SceneDepthTexture, OutSeparateTranslucencyTextures, ViewsToRender, SceneColorCopyTexture, ETranslucencyPass::TPT_AllTranslucency);
+		RenderTranslucencyInner(GraphBuilder, SceneTextures, TranslucentVolumeLightingTextures, OutSeparateTranslucencyTextures, ViewsToRender, SceneColorCopyTexture, ETranslucencyPass::TPT_AllTranslucency);
 	}
 }
diff --git a/Engine/Source/Runtime/Renderer/Private/VelocityRendering.cpp b/Engine/Source/Runtime/Renderer/Private/VelocityRendering.cpp
--- a/Engine/Source/Runtime/Renderer/Private/VelocityRendering.cpp
+++ b/Engine/Source/Runtime/Renderer/Private/VelocityRendering.cpp
@@ -224,45 +224,25 @@ END_SHADER_PARAMETER_STRUCT()
 
 void FDeferredShadingSceneRenderer::RenderVelocities(
 	FRDGBuilder& GraphBuilder,
-	FRDGTextureRef DepthTexture,
-	FRDGTextureRef& InOutVelocityTexture,
-	TRDGUniformBufferRef<FSceneTextureUniformParameters> SceneTexturesUniformBuffer,
+	FSceneTextures& SceneTextures,
 	EVelocityPass VelocityPass,
 	bool bForceVelocity)
 {
-	check(FeatureLevel >= ERHIFeatureLevel::SM5);
-
-	
 	if (!ShouldRenderVelocities())
 	{
 		return;
 	}
 
-	{
-		const EMeshPass::Type MeshPass = GetMeshPassFromVelocityPass(VelocityPass);
-		bool bHasAnyDraw = false;
-		for (const FViewInfo& View : Views)
-		{
-			bHasAnyDraw |= View.ParallelMeshDrawCommandPasses[MeshPass].HasAnyDraw();
-		}
-
-		if( !bHasAnyDraw )
-		{
-			return;
-		}
-	}
-
 	RDG_CSV_STAT_EXCLUSIVE_SCOPE(GraphBuilder, RenderVelocities);
 	SCOPED_NAMED_EVENT(FDeferredShadingSceneRenderer_RenderVelocities, FColor::Emerald);
 	SCOPE_CYCLE_COUNTER(STAT_RenderVelocities);
 
 	ERenderTargetLoadAction VelocityLoadAction = ERenderTargetLoadAction::ELoad;
-	FRDGTextureRef VelocityTexture = InOutVelocityTexture;
 	bool bVelocityRendered = false;
 
-	if (!VelocityTexture)
+	if (!SceneTextures.Velocity)
 	{
-		VelocityTexture = GraphBuilder.CreateTexture(FVelocityRendering::GetRenderTargetDesc(ShaderPlatform), TEXT("Velocity"));
+		SceneTextures.Velocity = GraphBuilder.CreateTexture(FVelocityRendering::GetRenderTargetDesc(ShaderPlatform), TEXT("Velocity"));
 		VelocityLoadAction = ERenderTargetLoadAction::EClear;
 	}
 
@@ -289,7 +269,7 @@ void FDeferredShadingSceneRenderer::RenderVelocities(
 
 			if (VelocityLoadAction == ERenderTargetLoadAction::EClear)
 			{
-				AddClearRenderTargetPass(GraphBuilder, VelocityTexture);
+				AddClearRenderTargetPass(GraphBuilder, SceneTextures.Velocity);
 
 				if (!View.Family->bMultiGPUForkAndJoin)
 				{
@@ -304,9 +284,9 @@ void FDeferredShadingSceneRenderer::RenderVelocities(
 			}
 
 			FVelocityPassParameters* PassParameters = GraphBuilder.AllocParameters<FVelocityPassParameters>();
-			PassParameters->SceneTextures = SceneTexturesUniformBuffer;
+			PassParameters->SceneTextures = SceneTextures.UniformBuffer;
 			PassParameters->RenderTargets.DepthStencil = FDepthStencilBinding(
-				DepthTexture,
+				SceneTextures.Depth.Resolve,
 				ERenderTargetLoadAction::ELoad,
 				ERenderTargetLoadAction::ELoad,
 				(VelocityPass == EVelocityPass::Opaque && !(Scene->EarlyZPassMode == DDM_AllOpaqueNoVelocity))
@@ -315,7 +295,7 @@ void FDeferredShadingSceneRenderer::RenderVelocities(
 
 			if (IsParallelVelocity())
 			{
-				PassParameters->RenderTargets[0] = FRenderTargetBinding(VelocityTexture, ERenderTargetLoadAction::ELoad);
+				PassParameters->RenderTargets[0] = FRenderTargetBinding(SceneTextures.Velocity, ERenderTargetLoadAction::ELoad);
 
 				GraphBuilder.AddPass(
 					RDG_EVENT_NAME("VelocityParallel"),
@@ -330,7 +310,7 @@ void FDeferredShadingSceneRenderer::RenderVelocities(
 			}
 			else
 			{
-				PassParameters->RenderTargets[0] = FRenderTargetBinding(VelocityTexture, ERenderTargetLoadAction::ELoad);
+				PassParameters->RenderTargets[0] = FRenderTargetBinding(SceneTextures.Velocity, ERenderTargetLoadAction::ELoad);
 
 				GraphBuilder.AddPass(
 					RDG_EVENT_NAME("Velocity"),
@@ -348,11 +328,9 @@ void FDeferredShadingSceneRenderer::RenderVelocities(
 		}
 	}
 
-	if (!InOutVelocityTexture && bVelocityRendered)
+	if (bVelocityRendered)
 	{
-		FSceneRenderTargets& SceneContext = FSceneRenderTargets::Get();
-		ConvertToExternalTexture(GraphBuilder, VelocityTexture, SceneContext.SceneVelocity);
-		InOutVelocityTexture = VelocityTexture;
+		ConvertToExternalTexture(GraphBuilder, SceneTextures.Velocity, FSceneRenderTargets::Get().SceneVelocity);
 	}
 }
 
diff --git a/Engine/Source/Runtime/Renderer/Private/VisualizeVolumetricLightmap.cpp b/Engine/Source/Runtime/Renderer/Private/VisualizeVolumetricLightmap.cpp
--- a/Engine/Source/Runtime/Renderer/Private/VisualizeVolumetricLightmap.cpp
+++ b/Engine/Source/Runtime/Renderer/Private/VisualizeVolumetricLightmap.cpp
@@ -94,8 +94,7 @@ IMPLEMENT_GLOBAL_SHADER(FVisualizeVolumetricLightmapPS, "/Engine/Private/Visuali
 
 void FDeferredShadingSceneRenderer::VisualizeVolumetricLightmap(
 	FRDGBuilder& GraphBuilder,
-	FRDGTextureRef SceneColorTexture,
-	FRDGTextureRef SceneDepthTexture)
+	const FMinimalSceneTextures& SceneTextures)
 {
 	if (!ViewFamily.EngineShowFlags.VisualizeVolumetricLightmap)
 	{
@@ -124,8 +123,8 @@ void FDeferredShadingSceneRenderer::VisualizeVolumetricLightmap(
 	for (const FViewInfo& View : Views)
 	{
 		auto* PassParameters = GraphBuilder.AllocParameters<FVisualizeVolumetricLightmapPS::FParameters>();
-		PassParameters->RenderTargets.DepthStencil = FDepthStencilBinding(SceneDepthTexture, ERenderTargetLoadAction::ELoad, ERenderTargetLoadAction::ELoad, FExclusiveDepthStencil::DepthWrite_StencilWrite);
-		PassParameters->RenderTargets[0] = FRenderTargetBinding(SceneColorTexture, ERenderTargetLoadAction::ELoad);
+		PassParameters->RenderTargets.DepthStencil = FDepthStencilBinding(SceneTextures.Depth.Target, ERenderTargetLoadAction::ELoad, ERenderTargetLoadAction::ELoad, FExclusiveDepthStencil::DepthWrite_StencilWrite);
+		PassParameters->RenderTargets[0] = FRenderTargetBinding(SceneTextures.Color.Target, ERenderTargetLoadAction::ELoad);
 
 		if (SceneContext.GBufferB)
 		{
diff --git a/Engine/Source/Runtime/Renderer/Private/VolumetricCloudRendering.cpp b/Engine/Source/Runtime/Renderer/Private/VolumetricCloudRendering.cpp
--- a/Engine/Source/Runtime/Renderer/Private/VolumetricCloudRendering.cpp
+++ b/Engine/Source/Runtime/Renderer/Private/VolumetricCloudRendering.cpp
@@ -1726,11 +1726,9 @@ void FSceneRenderer::RenderVolumetricCloudsInternal(FRDGBuilder& GraphBuilder, F
 
 void FSceneRenderer::RenderVolumetricCloud(
 	FRDGBuilder& GraphBuilder,
-	const FSceneTextureShaderParameters& SceneTextures,
+	const FMinimalSceneTextures& SceneTextures,
 	bool bSkipVolumetricRenderTarget,
 	bool bSkipPerPixelTracing,
-	FRDGTextureMSAA SceneColorTexture,
-	FRDGTextureMSAA SceneDepthTexture,
 	FRDGTextureRef HalfResolutionDepthCheckerboardMinMaxTexture)
 {
 	check(ShouldRenderVolumetricCloud(Scene, ViewFamily.EngineShowFlags)); // This should not be called if we should not render SkyAtmosphere
@@ -1752,7 +1750,7 @@ void FSceneRenderer::RenderVolumetricCloud(
 
 			FSceneRenderTargets& SceneContext = FSceneRenderTargets::Get();
 
-			FRDGTextureRef SceneDepthZ = SceneDepthTexture.Resolve;
+			FRDGTextureRef SceneDepthZ = SceneTextures.Depth.Resolve;
 
 			FCloudRenderContext CloudRC;
 			CloudRC.CloudInfo = &CloudInfo;
@@ -1914,7 +1912,7 @@ void FSceneRenderer::RenderVolumetricCloud(
 					SkyRC.bShouldSampleCloudSkyAO = CloudShadowAOData.bShouldSampleCloudSkyAO;
 					SkyRC.VolumetricCloudSkyAO = CloudShadowAOData.VolumetricCloudSkyAO;
 
-					RenderSkyAtmosphereInternal(GraphBuilder, SceneTextures, SkyRC);
+					RenderSkyAtmosphereInternal(GraphBuilder, GetSceneTextureShaderParameters(SceneTextures.UniformBuffer), SkyRC);
 				}
 
 
@@ -1948,7 +1946,7 @@ void FSceneRenderer::RenderVolumetricCloud(
 						const int DebugLightIndex = 0;	// only debug atmospheric light 0 for now
 						const int32 CloudShadowmapSampleCount = VolumetricCloudParams.VolumetricCloud.CloudShadowmapSampleCount[DebugLightIndex];
 
-						AddDrawCanvasPass(GraphBuilder, {}, ViewInfo, FScreenPassRenderTarget(SceneColorTexture.Target, ViewInfo.ViewRect, ERenderTargetLoadAction::ELoad), [CloudShadowmapSampleCount, &ViewInfo](FCanvas& Canvas)
+						AddDrawCanvasPass(GraphBuilder, {}, ViewInfo, FScreenPassRenderTarget(SceneTextures.Color.Target, ViewInfo.ViewRect, ERenderTargetLoadAction::ELoad), [CloudShadowmapSampleCount, &ViewInfo](FCanvas& Canvas)
 						{
 							const float ViewPortWidth = float(ViewInfo.ViewRect.Width());
 							const float ViewPortHeight = float(ViewInfo.ViewRect.Height());
