Fixed CIS build.

diff --git a/Engine/Source/Runtime/Renderer/Private/RendererModule.h b/Engine/Source/Runtime/Renderer/Private/RendererModule.h
--- a/Engine/Source/Runtime/Renderer/Private/RendererModule.h
+++ b/Engine/Source/Runtime/Renderer/Private/RendererModule.h
@@ -21,6 +21,7 @@ class FSceneViewFamily;
 class FSceneViewStateInterface;
 struct FMeshBatch;
 struct FSynthBenchmarkResults;
+struct FSceneTextures;
 
 template<class ResourceType> class TGlobalResource;
 
