Fixed missing transition in translucent lighting.

diff --git a/Engine/Source/Runtime/Renderer/Private/TranslucentLighting.cpp b/Engine/Source/Runtime/Renderer/Private/TranslucentLighting.cpp
--- a/Engine/Source/Runtime/Renderer/Private/TranslucentLighting.cpp
+++ b/Engine/Source/Runtime/Renderer/Private/TranslucentLighting.cpp
@@ -1199,6 +1199,7 @@ static FRDGTextureRef GetSkyTransmittanceLutTexture(FRDGBuilder& GraphBuilder, F
 
 BEGIN_SHADER_PARAMETER_STRUCT(FInjectTranslucentLightArrayParameters, )
 	RDG_TEXTURE_ACCESS(TransmittanceLutTexture, ERHIAccess::SRVGraphics)
+	RDG_TEXTURE_ACCESS(ShadowDepthTexture, ERHIAccess::SRVGraphics)
 	RENDER_TARGET_BINDING_SLOTS()
 END_SHADER_PARAMETER_STRUCT()
 
@@ -1214,6 +1215,8 @@ static void InjectTranslucentLightArray(
 	FSceneRenderTargets& SceneContext = FSceneRenderTargets::Get();
 	INC_DWORD_STAT_BY(STAT_NumLightsInjectedIntoTranslucency, LightInjectionData.Num());
 
+	FRDGTextureRef TransmittanceLutTexture = GetSkyTransmittanceLutTexture(GraphBuilder, Scene, View);
+
 	// Inject into each volume cascade
 	// Operate on one cascade at a time to reduce render target switches
 	for (uint32 VolumeCascadeIndex = 0; VolumeCascadeIndex < TVC_MAX; VolumeCascadeIndex++)
@@ -1221,36 +1224,44 @@ static void InjectTranslucentLightArray(
 		FRDGTextureRef VolumeAmbientTexture = Textures.GetAmbient(ViewIndex, VolumeCascadeIndex);
 		FRDGTextureRef VolumeDirectionalTexture = Textures.GetDirectional(ViewIndex, VolumeCascadeIndex);
 
-		auto* PassParameters = GraphBuilder.AllocParameters<FInjectTranslucentLightArrayParameters>();
-		PassParameters->TransmittanceLutTexture = GetSkyTransmittanceLutTexture(GraphBuilder, Scene, View);
-		PassParameters->RenderTargets[0] = FRenderTargetBinding(VolumeAmbientTexture, ERenderTargetLoadAction::ELoad);
-		PassParameters->RenderTargets[1] = FRenderTargetBinding(VolumeDirectionalTexture, ERenderTargetLoadAction::ELoad);
-
-		GraphBuilder.AddPass(
-			RDG_EVENT_NAME("InjectTranslucentLightArray"),
-			PassParameters,
-			ERDGPassFlags::Raster,
-			[&View, LightInjectionData, VolumeCascadeIndex](FRHICommandList& RHICmdList)
+		for (int32 LightIndex = 0; LightIndex < LightInjectionData.Num(); LightIndex++)
 		{
-			FGraphicsPipelineStateInitializer GraphicsPSOInit;
-			RHICmdList.ApplyCachedRenderTargets(GraphicsPSOInit);
+			const FTranslucentLightInjectionData& InjectionData = LightInjectionData[LightIndex];
+			const FLightSceneInfo* const LightSceneInfo = InjectionData.LightSceneInfo;
+			const bool bInverseSquared = LightSceneInfo->Proxy->IsInverseSquared();
+			const bool bDirectionalLight = LightSceneInfo->Proxy->GetLightType() == LightType_Directional;
+			const FVolumeBounds VolumeBounds = CalculateLightVolumeBounds(LightSceneInfo->Proxy->GetBoundingSphere(), View, VolumeCascadeIndex, bDirectionalLight);
 
-			GraphicsPSOInit.RasterizerState = TStaticRasterizerState<FM_Solid, CM_None>::GetRHI();
-			GraphicsPSOInit.DepthStencilState = TStaticDepthStencilState<false, CF_Always>::GetRHI();
-			GraphicsPSOInit.PrimitiveType = PT_TriangleStrip;
-
-			for (int32 LightIndex = 0; LightIndex < LightInjectionData.Num(); LightIndex++)
+			if (VolumeBounds.IsValid())
 			{
-				const FTranslucentLightInjectionData& InjectionData = LightInjectionData[LightIndex];
-				const FLightSceneInfo* const LightSceneInfo = InjectionData.LightSceneInfo;
-				const bool bInverseSquared = LightSceneInfo->Proxy->IsInverseSquared();
-				const bool bDirectionalLight = LightSceneInfo->Proxy->GetLightType() == LightType_Directional;
-				const FVolumeBounds VolumeBounds = CalculateLightVolumeBounds(LightSceneInfo->Proxy->GetBoundingSphere(), View, VolumeCascadeIndex, bDirectionalLight);
+				TShaderMapRef<FWriteToSliceVS> VertexShader(View.ShaderMap);
+				TOptionalShaderMapRef<FWriteToSliceGS> GeometryShader(View.ShaderMap);
+
+				FRDGTextureRef ShadowDepthTexture = nullptr;
+
+				if (InjectionData.ProjectedShadowInfo)
+				{
+					ShadowDepthTexture = TryRegisterExternalTexture(GraphBuilder, InjectionData.ProjectedShadowInfo->RenderTargets.DepthTarget);
+				}
 
-				if (VolumeBounds.IsValid())
+				auto* PassParameters = GraphBuilder.AllocParameters<FInjectTranslucentLightArrayParameters>();
+				PassParameters->TransmittanceLutTexture = TransmittanceLutTexture;
+				PassParameters->ShadowDepthTexture = ShadowDepthTexture;
+				PassParameters->RenderTargets[0] = FRenderTargetBinding(VolumeAmbientTexture, ERenderTargetLoadAction::ELoad);
+				PassParameters->RenderTargets[1] = FRenderTargetBinding(VolumeDirectionalTexture, ERenderTargetLoadAction::ELoad);
+
+				GraphBuilder.AddPass(
+					RDG_EVENT_NAME("InjectTranslucentLightArray"),
+					PassParameters,
+					ERDGPassFlags::Raster,
+					[VertexShader, GeometryShader, &View, &InjectionData, LightSceneInfo, bInverseSquared, bDirectionalLight, VolumeBounds, VolumeCascadeIndex](FRHICommandList& RHICmdList)
 				{
-					TShaderMapRef<FWriteToSliceVS> VertexShader(View.ShaderMap);
-					TOptionalShaderMapRef<FWriteToSliceGS> GeometryShader(View.ShaderMap);
+					FGraphicsPipelineStateInitializer GraphicsPSOInit;
+					RHICmdList.ApplyCachedRenderTargets(GraphicsPSOInit);
+
+					GraphicsPSOInit.RasterizerState = TStaticRasterizerState<FM_Solid, CM_None>::GetRHI();
+					GraphicsPSOInit.DepthStencilState = TStaticDepthStencilState<false, CF_Always>::GetRHI();
+					GraphicsPSOInit.PrimitiveType = PT_TriangleStrip;
 
 					if (bDirectionalLight)
 					{
@@ -1304,9 +1315,9 @@ static void InjectTranslucentLightArray(
 						GeometryShader->SetParameters(RHICmdList, VolumeBounds.MinZ);
 					}
 					RasterizeToVolumeTexture(RHICmdList, VolumeBounds);
-				}
+				});
 			}
-		});
+		}
 	}
 }
 
