Cleaning up normal reprojection code.

diff --git a/Engine/Source/Runtime/Renderer/Private/CompositionLighting/PostProcessDeferredDecals.cpp b/Engine/Source/Runtime/Renderer/Private/CompositionLighting/PostProcessDeferredDecals.cpp
--- a/Engine/Source/Runtime/Renderer/Private/CompositionLighting/PostProcessDeferredDecals.cpp
+++ b/Engine/Source/Runtime/Renderer/Private/CompositionLighting/PostProcessDeferredDecals.cpp
@@ -45,6 +45,9 @@ static TAutoConsoleVariable<bool> CVarDBufferDecalNormalReprojectionEnabled(
 	TEXT("in depth prepass is enabled as well (r.DepthPassMergedWithVelocity). Otherwise the fallback is the normal extracted from the depth buffer."),
 	ECVF_RenderThreadSafe);
 
+// This is a temporary extern while testing the Velocity changes. Will be removed once a decision is made.
+extern bool IsVelocityMergedWithDepthPass();
+
 FDeferredDecalPassTextures GetDeferredDecalPassTextures(
 	FRDGBuilder& GraphBuilder,
 	const FViewInfo& View,
@@ -698,3 +701,27 @@ void AddDeferredDecalPass(
 		ConvertToExternalTexture(GraphBuilder, OutputTexture, SceneContext.DBufferMask);
 	}
 }
+
+void ExtractNormalsForNextFrameReprojection(FRDGBuilder& GraphBuilder, FSceneRenderTargets& SceneContext, const TArray<FViewInfo>& Views)
+{
+	static auto CVarNormalReprojectionEnabled = IConsoleManager::Get().FindConsoleVariable(TEXT("r.Decal.NormalReprojectionEnabled"));
+
+	// save the previous frame if early motion vectors are enabled and normal reprojection is enabled, so there should be no cost if these options are off
+	bool bApplyReproject = IsVelocityMergedWithDepthPass() && CVarNormalReprojectionEnabled && CVarNormalReprojectionEnabled->GetInt() > 0;
+
+	if (bApplyReproject)
+	{
+		bool bAnyViewWriteable = false;
+		for (int32 Index = 0; Index < Views.Num(); Index++)
+		{
+			if (Views[Index].bStatePrevViewInfoIsReadOnly == false)
+			{
+				GraphBuilder.QueueTextureExtraction(GraphBuilder.RegisterExternalTexture(SceneContext.GBufferA),&Views[Index].ViewState->PrevFrameViewInfo.GBufferA);
+			}
+		}
+
+	}
+}
+
+
+
diff --git a/Engine/Source/Runtime/Renderer/Private/CompositionLighting/PostProcessDeferredDecals.h b/Engine/Source/Runtime/Renderer/Private/CompositionLighting/PostProcessDeferredDecals.h
--- a/Engine/Source/Runtime/Renderer/Private/CompositionLighting/PostProcessDeferredDecals.h
+++ b/Engine/Source/Runtime/Renderer/Private/CompositionLighting/PostProcessDeferredDecals.h
@@ -82,4 +82,9 @@ void RenderMeshDecals(
 	FRDGBuilder& GraphBuilder,
 	const FViewInfo& View,
 	FDeferredDecalPassTextures& DecalPassTextures,
-	EDecalRenderStage DecalRenderStage);
\ No newline at end of file
+	EDecalRenderStage DecalRenderStage);
+
+void ExtractNormalsForNextFrameReprojection(
+	FRDGBuilder& GraphBuilder,
+	FSceneRenderTargets& SceneContext,
+	const TArray<FViewInfo>& Views);
diff --git a/Engine/Source/Runtime/Renderer/Private/DeferredShadingRenderer.cpp b/Engine/Source/Runtime/Renderer/Private/DeferredShadingRenderer.cpp
--- a/Engine/Source/Runtime/Renderer/Private/DeferredShadingRenderer.cpp
+++ b/Engine/Source/Runtime/Renderer/Private/DeferredShadingRenderer.cpp
@@ -333,9 +333,6 @@ static FORCEINLINE bool NeedsPrePass(const FDeferredShadingSceneRenderer* Render
 		(Renderer->DepthPass.EarlyZPassMode != DDM_None || Renderer->DepthPass.bEarlyZPassMovable != 0);
 }
 
-// This is a temporary extern while testing the Velocity changes. Will be removed once a decision is made.
-extern bool IsVelocityMergedWithDepthPass();
-
 bool FDeferredShadingSceneRenderer::RenderHzb(FRDGBuilder& GraphBuilder, const FSceneTextures& SceneTextures)
 {
 	RDG_GPU_STAT_SCOPE(GraphBuilder, HZB);
@@ -2248,26 +2245,8 @@ void FDeferredShadingSceneRenderer::Render(FRHICommandListImmediate& RHICmdList)
 	SceneTextures.SetupMode |= ESceneTextureSetupMode::GBuffers;
 	SceneTextures.UniformBuffer = CreateSceneTextureUniformBuffer(GraphBuilder, FeatureLevel, SceneTextures.SetupMode);
 
-	// mark GBufferA for saving for next frame
-	{
-		static auto CVarNormalReprojectionEnabled = IConsoleManager::Get().FindConsoleVariable(TEXT("r.Decal.NormalReprojectionEnabled"));
-
-		// save the previous frame if early motion vectors are enabled and normal reprojection is enabled, so there should be no cost if these options are off
-		bool bApplyReproject = IsVelocityMergedWithDepthPass() && CVarNormalReprojectionEnabled && CVarNormalReprojectionEnabled->GetInt() > 0;
-
-		if (bApplyReproject)
-		{
-			bool bAnyViewWriteable = false;
-			for (int32 Index = 0; Index < Views.Num(); Index++)
-			{
-				if (Views[Index].bStatePrevViewInfoIsReadOnly == false)
-				{
-					GraphBuilder.QueueTextureExtraction(GraphBuilder.RegisterExternalTexture(SceneContext.GBufferA),&Views[Index].ViewState->PrevFrameViewInfo.GBufferA);
-				}
-			}
-		
-		}
-	}
+	// mark GBufferA for saving for next frame if it's needed
+	ExtractNormalsForNextFrameReprojection(GraphBuilder, SceneContext, Views);
 
 	if (ViewFamily.EngineShowFlags.VisualizeLightCulling)
 	{
