Refactored translucent lighting volume methods out of deferred shading renderer.

diff --git a/Engine/Shaders/Private/BasePassCommon.ush b/Engine/Shaders/Private/BasePassCommon.ush
--- a/Engine/Shaders/Private/BasePassCommon.ush
+++ b/Engine/Shaders/Private/BasePassCommon.ush
@@ -135,18 +135,10 @@ struct FBasePassInterpolantsVSToDS : FSharedBasePassInterpolants
 
 #endif
 
-
-#if SUPPORTS_INDEPENDENT_SAMPLERS
-	#define SharedAmbientInnerSampler View.SharedBilinearClampedSampler
-	#define SharedAmbientOuterSampler View.SharedBilinearClampedSampler
-	#define SharedDirectionalInnerSampler View.SharedBilinearClampedSampler
-	#define SharedDirectionalOuterSampler View.SharedBilinearClampedSampler
-#else
-	#define SharedAmbientInnerSampler TranslucentBasePass.TranslucencyLightingVolumeAmbientInnerSampler
-	#define SharedAmbientOuterSampler TranslucentBasePass.TranslucencyLightingVolumeAmbientOuterSampler
-	#define SharedDirectionalInnerSampler TranslucentBasePass.TranslucencyLightingVolumeDirectionalInnerSampler
-	#define SharedDirectionalOuterSampler TranslucentBasePass.TranslucencyLightingVolumeDirectionalOuterSampler
-#endif
+#define SharedAmbientInnerSampler View.SharedBilinearClampedSampler
+#define SharedAmbientOuterSampler View.SharedBilinearClampedSampler
+#define SharedDirectionalInnerSampler View.SharedBilinearClampedSampler
+#define SharedDirectionalOuterSampler View.SharedBilinearClampedSampler
 
 void ComputeVolumeUVs(float3 WorldPosition, float3 LightingPositionOffset, out float3 InnerVolumeUVs, out float3 OuterVolumeUVs, out float FinalLerpFactor)
 {
diff --git a/Engine/Source/Runtime/Renderer/Private/BasePassRendering.h b/Engine/Source/Runtime/Renderer/Private/BasePassRendering.h
--- a/Engine/Source/Runtime/Renderer/Private/BasePassRendering.h
+++ b/Engine/Source/Runtime/Renderer/Private/BasePassRendering.h
@@ -22,6 +22,7 @@
 #include "MeshMaterialShader.h"
 #include "ShaderBaseClasses.h"
 #include "FogRendering.h"
+#include "TranslucentLighting.h"
 #include "PlanarReflectionRendering.h"
 #include "UnrealEngine.h"
 #include "ReflectionEnvironment.h"
@@ -32,7 +33,6 @@ class FScene;
 template<typename TBufferStruct> class TUniformBufferRef;
 
 struct FSceneWithoutWaterTextures;
-struct FTranslucentVolumeLightingTextures;
 
 class FViewInfo;
 class UMaterialExpressionSingleLayerWaterMaterialOutput;
@@ -108,14 +108,7 @@ BEGIN_GLOBAL_SHADER_PARAMETER_STRUCT(FTranslucentBasePassUniformParameters,)
 	SHADER_PARAMETER_SAMPLER(SamplerState, VolumetricCloudDepthSampler)
 	SHADER_PARAMETER(float, ApplyVolumetricCloudOnTransparent)
 	// Translucency Lighting Volume
-	SHADER_PARAMETER_RDG_TEXTURE(Texture3D, TranslucencyLightingVolumeAmbientInner)
-	SHADER_PARAMETER_SAMPLER(SamplerState, TranslucencyLightingVolumeAmbientInnerSampler)
-	SHADER_PARAMETER_RDG_TEXTURE(Texture3D, TranslucencyLightingVolumeAmbientOuter)
-	SHADER_PARAMETER_SAMPLER(SamplerState, TranslucencyLightingVolumeAmbientOuterSampler)
-	SHADER_PARAMETER_RDG_TEXTURE(Texture3D, TranslucencyLightingVolumeDirectionalInner)
-	SHADER_PARAMETER_SAMPLER(SamplerState, TranslucencyLightingVolumeDirectionalInnerSampler)
-	SHADER_PARAMETER_RDG_TEXTURE(Texture3D, TranslucencyLightingVolumeDirectionalOuter)
-	SHADER_PARAMETER_SAMPLER(SamplerState, TranslucencyLightingVolumeDirectionalOuterSampler)
+	SHADER_PARAMETER_STRUCT_INCLUDE(FTranslucencyLightingVolumeParameters, TranslucencyLightingVolume)
 	SHADER_PARAMETER_STRUCT_INCLUDE(FLumenTranslucencyLightingParameters, LumenParameters)
 	SHADER_PARAMETER_TEXTURE(Texture2D, PreIntegratedGFTexture)
 	SHADER_PARAMETER_SAMPLER(SamplerState, PreIntegratedGFSampler)
@@ -136,17 +129,17 @@ extern void SetupSharedBasePassParameters(
 extern TRDGUniformBufferRef<FOpaqueBasePassUniformParameters> CreateOpaqueBasePassUniformBuffer(
 	FRDGBuilder& GraphBuilder,
 	const FViewInfo& View,
-	FRDGTextureRef ForwardScreenSpaceShadowMask,
-	const FSceneWithoutWaterTextures* SceneWithoutWaterTextures,
-	const int32 ViewIndex);
+	FRDGTextureRef ForwardScreenSpaceShadowMask = nullptr,
+	const FSceneWithoutWaterTextures* SceneWithoutWaterTextures = nullptr,
+	const int32 ViewIndex = 0);
 
 extern TRDGUniformBufferRef<FTranslucentBasePassUniformParameters> CreateTranslucentBasePassUniformBuffer(
 	FRDGBuilder& GraphBuilder,
 	const FViewInfo& View,
-	const FTranslucentVolumeLightingTextures* TranslucentVolumeLightingTextures,
-	FRDGTextureRef SceneColorCopy,
-	ESceneTextureSetupMode SceneTextureSetupMode,
-	const int32 ViewIndex);
+	const int32 ViewIndex = 0,
+	const FTranslucencyLightingVolumeTextures& TranslucencyLightingVolumeTextures = {},
+	FRDGTextureRef SceneColorCopyTexture = nullptr,
+	const ESceneTextureSetupMode SceneTextureSetupMode = ESceneTextureSetupMode::None);
 
 /** Parameters for computing forward lighting. */
 class FForwardLightingParameters
diff --git a/Engine/Source/Runtime/Renderer/Private/DeferredShadingRenderer.cpp b/Engine/Source/Runtime/Renderer/Private/DeferredShadingRenderer.cpp
--- a/Engine/Source/Runtime/Renderer/Private/DeferredShadingRenderer.cpp
+++ b/Engine/Source/Runtime/Renderer/Private/DeferredShadingRenderer.cpp
@@ -2194,11 +2194,11 @@ void FDeferredShadingSceneRenderer::Render(FRHICommandListImmediate& RHICmdList)
 			SceneContext.bScreenSpaceAOIsValid);
 	}
 
-	FTranslucentVolumeLightingTextures TranslucentLightingVolumeTextures;
+	FTranslucencyLightingVolumeTextures TranslucencyLightingVolumeTextures;
 
 	if (bRenderDeferredLighting && GbEnableAsyncComputeTranslucencyLightingVolumeClear && GSupportsEfficientAsyncCompute)
 	{
-		InitTranslucentVolumeLighting(GraphBuilder, ERDGPassFlags::AsyncCompute, TranslucentLightingVolumeTextures);
+		InitTranslucencyLightingVolumeTextures(GraphBuilder, Views, ERDGPassFlags::AsyncCompute, TranslucencyLightingVolumeTextures);
 	}
 
 	SceneContext.AllocSceneColor(RHICmdList);
@@ -2476,7 +2476,7 @@ void FDeferredShadingSceneRenderer::Render(FRHICommandListImmediate& RHICmdList)
 		// Clear the translucent lighting volumes before we accumulate
 		if ((GbEnableAsyncComputeTranslucencyLightingVolumeClear && GSupportsEfficientAsyncCompute) == false)
 		{
-			InitTranslucentVolumeLighting(GraphBuilder, ERDGPassFlags::Compute, TranslucentLightingVolumeTextures);
+			InitTranslucencyLightingVolumeTextures(GraphBuilder, Views, ERDGPassFlags::Compute, TranslucencyLightingVolumeTextures);
 		}
 
 #if RHI_RAYTRACING
@@ -2487,26 +2487,12 @@ void FDeferredShadingSceneRenderer::Render(FRHICommandListImmediate& RHICmdList)
 #endif
 
 		AddSetCurrentStatPass(GraphBuilder, GET_STATID(STAT_CLM_Lighting));
-		RenderLights(GraphBuilder, SceneTextures, TranslucentLightingVolumeTextures, LightingChannelsTexture, SortedLightSet, HairDatas);
+		RenderLights(GraphBuilder, SceneTextures, TranslucencyLightingVolumeTextures, LightingChannelsTexture, SortedLightSet, HairDatas);
 		AddSetCurrentStatPass(GraphBuilder, GET_STATID(STAT_CLM_AfterLighting));
 		AddServiceLocalQueuePass(GraphBuilder);
 
-		for (int32 ViewIndex = 0; ViewIndex < Views.Num(); ++ViewIndex)
-		{
-			const FViewInfo& View = Views[ViewIndex];
-			RDG_GPU_MASK_SCOPE(GraphBuilder, View.GPUMask);
-			InjectAmbientCubemapTranslucentVolumeLighting(GraphBuilder, TranslucentLightingVolumeTextures, Views[ViewIndex], ViewIndex);
-		}
-		AddServiceLocalQueuePass(GraphBuilder);
-
-		for (int32 ViewIndex = 0; ViewIndex < Views.Num(); ++ViewIndex)
-		{
-			const FViewInfo& View = Views[ViewIndex];
-			RDG_GPU_MASK_SCOPE(GraphBuilder, View.GPUMask);
-
-			// Filter the translucency lighting volume now that it is complete
-			FilterTranslucentVolumeLighting(GraphBuilder, TranslucentLightingVolumeTextures, View, ViewIndex);
-		}
+		InjectTranslucencyLightingVolumeAmbientCubemap(GraphBuilder, Views, TranslucencyLightingVolumeTextures);
+		FilterTranslucencyLightingVolume(GraphBuilder, Views, TranslucencyLightingVolumeTextures);
 		AddServiceLocalQueuePass(GraphBuilder);
 
 		// Render diffuse sky lighting and reflections that only operate on opaque pixels
@@ -2566,7 +2552,7 @@ void FDeferredShadingSceneRenderer::Render(FRHICommandListImmediate& RHICmdList)
 			RDG_CSV_STAT_EXCLUSIVE_SCOPE(GraphBuilder, RenderTranslucency);
 			SCOPE_CYCLE_COUNTER(STAT_TranslucencyDrawTime);
 			AddSetCurrentStatPass(GraphBuilder, GET_STATID(STAT_CLM_Translucency));
-			RenderTranslucency(GraphBuilder, SceneTextures, TranslucentLightingVolumeTextures, nullptr, ETranslucencyView::UnderWater);
+			RenderTranslucency(GraphBuilder, SceneTextures, TranslucencyLightingVolumeTextures, nullptr, ETranslucencyView::UnderWater);
 			EnumRemoveFlags(TranslucencyViewsToRender, ETranslucencyView::UnderWater);
 		}
 
@@ -2661,7 +2647,7 @@ void FDeferredShadingSceneRenderer::Render(FRHICommandListImmediate& RHICmdList)
 
 		// Render all remaining translucency views.
 		AddSetCurrentStatPass(GraphBuilder, GET_STATID(STAT_CLM_Translucency));
-		RenderTranslucency(GraphBuilder, SceneTextures, TranslucentLightingVolumeTextures, &SeparateTranslucencyTextures, TranslucencyViewsToRender);
+		RenderTranslucency(GraphBuilder, SceneTextures, TranslucencyLightingVolumeTextures, &SeparateTranslucencyTextures, TranslucencyViewsToRender);
 		AddServiceLocalQueuePass(GraphBuilder);
 		TranslucencyViewsToRender = ETranslucencyView::None;
 
@@ -2851,7 +2837,7 @@ void FDeferredShadingSceneRenderer::Render(FRHICommandListImmediate& RHICmdList)
 			}
 		}
 
-		AddPass(GraphBuilder, [this, &SceneContext](FRHICommandListImmediate&)
+		AddPass(GraphBuilder, [&SceneContext](FRHICommandListImmediate&)
 		{
 			SceneContext.SetSceneColor(nullptr);
 		});
diff --git a/Engine/Source/Runtime/Renderer/Private/DeferredShadingRenderer.h b/Engine/Source/Runtime/Renderer/Private/DeferredShadingRenderer.h
--- a/Engine/Source/Runtime/Renderer/Private/DeferredShadingRenderer.h
+++ b/Engine/Source/Runtime/Renderer/Private/DeferredShadingRenderer.h
@@ -42,7 +42,7 @@ struct FRayTracingReflectionOptions;
 struct FHairStrandsTransmittanceMaskData;
 struct FHairStrandsRenderingData;
 
-struct FTranslucentVolumeLightingTextures;
+struct FTranslucencyLightingVolumeTextures;
 
 /**   
  * Data for rendering meshes into Lumen Lighting Cards.
@@ -565,7 +565,7 @@ private:
 	void RenderLights(
 		FRDGBuilder& GraphBuilder,
 		FMinimalSceneTextures& SceneTextures,
-		const FTranslucentVolumeLightingTextures& TranslucentVolumeLightingTextures,
+		const FTranslucencyLightingVolumeTextures& TranslucencyLightingVolumeTextures,
 		FRDGTextureRef LightingChannelsTexture,
 		FSortedLightSetSceneInfo& SortedLightSet,
 		const FHairStrandsRenderingData* HairDatas);
@@ -586,7 +586,7 @@ private:
 	void RenderTranslucency(
 		FRDGBuilder& GraphBuilder,
 		const FMinimalSceneTextures& SceneTextures,
-		const FTranslucentVolumeLightingTextures& TranslucentVolumeLightingTextures,
+		const FTranslucencyLightingVolumeTextures& TranslucencyLightingVolumeTextures,
 		FSeparateTranslucencyTextures* OutSeparateTranslucencyTextures,
 		ETranslucencyView ViewsToRender);
 
@@ -594,7 +594,7 @@ private:
 	void RenderTranslucencyInner(
 		FRDGBuilder& GraphBuilder,
 		const FMinimalSceneTextures& SceneTextures,
-		const FTranslucentVolumeLightingTextures& TranslucentVolumeLightingTextures,
+		const FTranslucencyLightingVolumeTextures& TranslucencyLightingVolumeTextures,
 		FSeparateTranslucencyTextures* OutSeparateTranslucencyTextures,
 		ETranslucencyView ViewsToRender,
 		FRDGTextureRef SceneColorCopyTexture,
@@ -665,7 +665,7 @@ private:
 	void RenderDeferredShadowProjections(
 		FRDGBuilder& GraphBuilder,
 		const FMinimalSceneTextures& SceneTextures,
-		const FTranslucentVolumeLightingTextures& TranslucentVolumeLightingTextures,
+		const FTranslucencyLightingVolumeTextures& TranslucencyLightingVolumeTextures,
 		const FLightSceneInfo* LightSceneInfo,
 		FRDGTextureRef ScreenShadowMaskTexture,
 		FRDGTextureRef ScreenShadowMaskSubPixelTexture,
@@ -750,51 +750,6 @@ private:
 		const FMinimalSceneTextures& SceneTextures,
 		const FSimpleLightArray& SimpleLights);
 
-	/** Initializes the translucency lighting volumes. Can be compute or async compute. */
-	void InitTranslucentVolumeLighting(
-		FRDGBuilder& GraphBuilder,
-		ERDGPassFlags PassFlags,
-		FTranslucentVolumeLightingTextures& Textures);
-
-	/** Add AmbientCubemap to the lighting volumes. */
-	void InjectAmbientCubemapTranslucentVolumeLighting(
-		FRDGBuilder& GraphBuilder,
-		const FTranslucentVolumeLightingTextures& Textures,
-		const FViewInfo& View,
-		int32 ViewIndex);
-
-	/** Accumulates direct lighting for the given light.  InProjectedShadowInfo can be NULL in which case the light will be unshadowed. */
-	void InjectTranslucentVolumeLighting(
-		FRDGBuilder& GraphBuilder,
-		const FTranslucentVolumeLightingTextures& Textures,
-		const FLightSceneInfo& LightSceneInfo,
-		const FProjectedShadowInfo* InProjectedShadowInfo,
-		const FViewInfo& View,
-		int32 ViewIndex);
-
-	/** Accumulates direct lighting for an array of unshadowed lights. */
-	void InjectTranslucentVolumeLightingArray(
-		FRDGBuilder& GraphBuilder,
-		const FTranslucentVolumeLightingTextures& Textures,
-		const TArray<FSortedLightSceneInfo, SceneRenderingAllocator>& SortedLights,
-		int32 FirstLightIndex,
-		int32 LightsEndIndex);
-
-	/** Accumulates direct lighting for simple lights. */
-	void InjectSimpleTranslucentVolumeLightingArray(
-		FRDGBuilder& GraphBuilder,
-		const FTranslucentVolumeLightingTextures& Textures,
-		const FSimpleLightArray& SimpleLights,
-		const FViewInfo& View,
-		const int32 ViewIndex);
-
-	/** Filters the translucency lighting volumes to reduce aliasing. */
-	void FilterTranslucentVolumeLighting(
-		FRDGBuilder& GraphBuilder,
-		FTranslucentVolumeLightingTextures& Textures,
-		const FViewInfo& View,
-		const int32 ViewIndex);
-
 	bool ShouldRenderVolumetricFog() const;
 
 	void SetupVolumetricFog();
diff --git a/Engine/Source/Runtime/Renderer/Private/LightRendering.cpp b/Engine/Source/Runtime/Renderer/Private/LightRendering.cpp
--- a/Engine/Source/Runtime/Renderer/Private/LightRendering.cpp
+++ b/Engine/Source/Runtime/Renderer/Private/LightRendering.cpp
@@ -1190,7 +1190,7 @@ void GetRenderLightParameters(
 void FDeferredShadingSceneRenderer::RenderLights(
 	FRDGBuilder& GraphBuilder,
 	FMinimalSceneTextures& SceneTextures,
-	const FTranslucentVolumeLightingTextures& TranslucentVolumeLightingTextures,
+	const FTranslucencyLightingVolumeTextures& TranslucencyLightingVolumeTextures,
 	FRDGTextureRef LightingChannelsTexture,
 	FSortedLightSetSceneInfo &SortedLightSet,
 	const FHairStrandsRenderingData* HairDatas)
@@ -1313,7 +1313,7 @@ void FDeferredShadingSceneRenderer::RenderLights(
 				if (AttenuationLightStart)
 				{
 					// Inject non-shadowed, non-simple, non-light function lights in to the volume.
-					InjectTranslucentVolumeLightingArray(GraphBuilder, TranslucentVolumeLightingTextures, SortedLights, SimpleLightsEnd, AttenuationLightStart);
+					InjectTranslucencyLightingVolumeArray(GraphBuilder, Views, Scene, TranslucencyLightingVolumeTextures, VisibleLightInfos, SortedLights, TInterval<int32>(SimpleLightsEnd, AttenuationLightStart));
 				}
 
 				if (SimpleLights.InstanceData.Num() > 0)
@@ -1332,7 +1332,7 @@ void FDeferredShadingSceneRenderer::RenderLights(
 							FViewInfo& View = Views[ViewIndex];
 							RDG_GPU_MASK_SCOPE(GraphBuilder, View.GPUMask);
 							RDG_EVENT_SCOPE(GraphBuilder, "InjectSimpleLightsTranslucentLighting");
-							InjectSimpleTranslucentVolumeLightingArray(GraphBuilder, TranslucentVolumeLightingTextures, SimpleLightArray, View, ViewIndex);
+							InjectSimpleTranslucencyLightingVolumeArray(GraphBuilder, View, ViewIndex, Views.Num(), TranslucencyLightingVolumeTextures, SimpleLightArray);
 						}
 					}
 				}
@@ -1863,7 +1863,7 @@ void FDeferredShadingSceneRenderer::RenderLights(
 							ClearShadowMask(ScreenShadowMaskSubPixelTexture);
 						}
 
-						RenderDeferredShadowProjections(GraphBuilder, SceneTextures, TranslucentVolumeLightingTextures, &LightSceneInfo, ScreenShadowMaskTexture, ScreenShadowMaskSubPixelTexture, HairDatas, bInjectedTranslucentVolume);
+						RenderDeferredShadowProjections(GraphBuilder, SceneTextures, TranslucencyLightingVolumeTextures, &LightSceneInfo, ScreenShadowMaskTexture, ScreenShadowMaskSubPixelTexture, HairDatas, bInjectedTranslucentVolume);
 					}
 
 					bUsedShadowMaskTexture = true;
@@ -1904,7 +1904,7 @@ void FDeferredShadingSceneRenderer::RenderLights(
 						RDG_GPU_MASK_SCOPE(GraphBuilder, View.GPUMask);
 
 						// Accumulate this light's unshadowed contribution to the translucency lighting volume
-						InjectTranslucentVolumeLighting(GraphBuilder, TranslucentVolumeLightingTextures, LightSceneInfo, nullptr, View, ViewIndex);
+						InjectTranslucencyLightingVolume(GraphBuilder, View, ViewIndex, Scene, TranslucencyLightingVolumeTextures, VisibleLightInfos, LightSceneInfo, nullptr);
 					}
 				}
 
diff --git a/Engine/Source/Runtime/Renderer/Private/PostProcess/PostProcessCompositeEditorPrimitives.cpp b/Engine/Source/Runtime/Renderer/Private/PostProcess/PostProcessCompositeEditorPrimitives.cpp
--- a/Engine/Source/Runtime/Renderer/Private/PostProcess/PostProcessCompositeEditorPrimitives.cpp
+++ b/Engine/Source/Runtime/Renderer/Private/PostProcess/PostProcessCompositeEditorPrimitives.cpp
@@ -316,7 +316,7 @@ FScreenPassTexture AddEditorPrimitivePass(
 
 			if (BasePassType == FEditorPrimitiveInputs::EBasePassType::Deferred)
 			{
-				PassParameters->BasePass = CreateOpaqueBasePassUniformBuffer(GraphBuilder, *EditorView, nullptr, nullptr, 0);
+				PassParameters->BasePass = CreateOpaqueBasePassUniformBuffer(GraphBuilder, *EditorView, 0);
 			}
 
 			GraphBuilder.AddPass(
diff --git a/Engine/Source/Runtime/Renderer/Private/ReflectionEnvironmentRealTimeCapture.cpp b/Engine/Source/Runtime/Renderer/Private/ReflectionEnvironmentRealTimeCapture.cpp
--- a/Engine/Source/Runtime/Renderer/Private/ReflectionEnvironmentRealTimeCapture.cpp
+++ b/Engine/Source/Runtime/Renderer/Private/ReflectionEnvironmentRealTimeCapture.cpp
@@ -533,7 +533,7 @@ void FScene::AllocateAndCaptureFrameSkyEnvMap(
 					{
 						auto* PassParameters = GraphBuilder.AllocParameters<FCaptureSkyMeshReflectionPassParameters>();
 						PassParameters->RenderTargets = SkyRC.RenderTargets;
-						PassParameters->BasePass = CreateOpaqueBasePassUniformBuffer(GraphBuilder, MainView, nullptr, nullptr, 0);
+						PassParameters->BasePass = CreateOpaqueBasePassUniformBuffer(GraphBuilder, MainView, 0);
 
 						// Setup the depth buffer
 						if (bUseDepthBuffer)
diff --git a/Engine/Source/Runtime/Renderer/Private/Renderer.cpp b/Engine/Source/Runtime/Renderer/Private/Renderer.cpp
--- a/Engine/Source/Runtime/Renderer/Private/Renderer.cpp
+++ b/Engine/Source/Runtime/Renderer/Private/Renderer.cpp
@@ -242,7 +242,7 @@ void FRendererModule::DrawTileMesh(FRHICommandListImmediate& RHICmdList, FMeshPa
 		{
 			if (ShadingPath == EShadingPath::Deferred)
 			{
-				PassParameters->TranslucentBasePass = CreateTranslucentBasePassUniformBuffer(GraphBuilder, View, nullptr, nullptr, ESceneTextureSetupMode::None, 0);
+				PassParameters->TranslucentBasePass = CreateTranslucentBasePassUniformBuffer(GraphBuilder, View);
 
 				GraphBuilder.AddPass({}, PassParameters, PassFlags,
 					[Scene, &View, &Mesh, &DrawRenderState, bUse128bitRT](FRHICommandListImmediate& InRHICmdList)
@@ -317,7 +317,7 @@ void FRendererModule::DrawTileMesh(FRHICommandListImmediate& RHICmdList, FMeshPa
 			{
 				if (ShadingPath == EShadingPath::Deferred)
 				{
-					PassParameters->OpaqueBasePass = CreateOpaqueBasePassUniformBuffer(GraphBuilder, View, nullptr, nullptr, 0);
+					PassParameters->OpaqueBasePass = CreateOpaqueBasePassUniformBuffer(GraphBuilder, View);
 
 					GraphBuilder.AddPass({}, PassParameters, PassFlags,
 						[Scene, &View, &DrawRenderState, &Mesh, bUse128bitRT](FRHICommandListImmediate& InRHICmdList)
diff --git a/Engine/Source/Runtime/Renderer/Private/SceneRendering.cpp b/Engine/Source/Runtime/Renderer/Private/SceneRendering.cpp
--- a/Engine/Source/Runtime/Renderer/Private/SceneRendering.cpp
+++ b/Engine/Source/Runtime/Renderer/Private/SceneRendering.cpp
@@ -416,9 +416,7 @@ FASTVRAM_CVAR(Distortion, 1);
 FASTVRAM_CVAR(ScreenSpaceShadowMask, 1);
 FASTVRAM_CVAR(VolumetricFog, 1);
 FASTVRAM_CVAR(SeparateTranslucency, 0); 
-FASTVRAM_CVAR(SeparateTranslucencyModulate, 0); 
-FASTVRAM_CVAR(LightAccumulation, 0); 
-FASTVRAM_CVAR(LightAttenuation, 0); 
+FASTVRAM_CVAR(SeparateTranslucencyModulate, 0);
 FASTVRAM_CVAR(ScreenSpaceAO,0);
 FASTVRAM_CVAR(SSR, 0);
 FASTVRAM_CVAR(DBufferA, 0);
@@ -569,8 +567,6 @@ void FFastVramConfig::Update()
 	bDirty |= UpdateTextureFlagFromCVar(CVarFastVRam_VolumetricFog, VolumetricFog);
 	bDirty |= UpdateTextureFlagFromCVar(CVarFastVRam_SeparateTranslucency, SeparateTranslucency);
 	bDirty |= UpdateTextureFlagFromCVar(CVarFastVRam_SeparateTranslucencyModulate, SeparateTranslucencyModulate);
-	bDirty |= UpdateTextureFlagFromCVar(CVarFastVRam_LightAccumulation, LightAccumulation);
-	bDirty |= UpdateTextureFlagFromCVar(CVarFastVRam_LightAttenuation, LightAttenuation);
 	bDirty |= UpdateTextureFlagFromCVar(CVarFastVRam_ScreenSpaceAO, ScreenSpaceAO);
 	bDirty |= UpdateTextureFlagFromCVar(CVarFastVRam_SSR, SSR);
 	bDirty |= UpdateTextureFlagFromCVar(CVarFastVRam_DBufferA, DBufferA);
diff --git a/Engine/Source/Runtime/Renderer/Private/SceneRendering.h b/Engine/Source/Runtime/Renderer/Private/SceneRendering.h
--- a/Engine/Source/Runtime/Renderer/Private/SceneRendering.h
+++ b/Engine/Source/Runtime/Renderer/Private/SceneRendering.h
@@ -2006,10 +2006,10 @@ protected:
 		FRDGTextureRef HalfResolutionDepthCheckerboardMinMaxTexture);
 
 	/** Render notification to artist when a sky material is used but it might comtains the camera (and then the sky/background would look black).*/
-	void RenderSkyAtmosphereEditorNotifications(FRDGBuilder& GraphBuilder, FRDGTextureRef SceneColorTexture);
+	void RenderSkyAtmosphereEditorNotifications(FRDGBuilder& GraphBuilder, FRDGTextureRef SceneColorTexture) const;
 
 	/** We should render on screen notification only if any of the scene contains a mesh using a sky material.*/
-	bool ShouldRenderSkyAtmosphereEditorNotifications();
+	bool ShouldRenderSkyAtmosphereEditorNotifications() const;
 
 	void ResolveSceneColor(FRHICommandListImmediate& RHICmdList);
 	void ResolveSceneDepth(FRHICommandListImmediate& RHICmdList);
@@ -2265,8 +2265,6 @@ struct FFastVramConfig
 	ETextureCreateFlags VolumetricFog;
 	ETextureCreateFlags SeparateTranslucency;
 	ETextureCreateFlags SeparateTranslucencyModulate;
-	ETextureCreateFlags LightAccumulation;
-	ETextureCreateFlags LightAttenuation;
 	ETextureCreateFlags ScreenSpaceAO;
 	ETextureCreateFlags SSR;
 	ETextureCreateFlags DBufferA;
diff --git a/Engine/Source/Runtime/Renderer/Private/ShadowRendering.cpp b/Engine/Source/Runtime/Renderer/Private/ShadowRendering.cpp
--- a/Engine/Source/Runtime/Renderer/Private/ShadowRendering.cpp
+++ b/Engine/Source/Runtime/Renderer/Private/ShadowRendering.cpp
@@ -1943,7 +1943,7 @@ void FDeferredShadingSceneRenderer::RenderShadowProjections(
 void FDeferredShadingSceneRenderer::RenderDeferredShadowProjections(
 	FRDGBuilder& GraphBuilder,
 	const FMinimalSceneTextures& SceneTextures,
-	const FTranslucentVolumeLightingTextures& TranslucentVolumeLightingTextures,
+	const FTranslucencyLightingVolumeTextures& TranslucencyLightingVolumeTextures,
 	const FLightSceneInfo* LightSceneInfo,
 	FRDGTextureRef ScreenShadowMaskTexture,
 	FRDGTextureRef ScreenShadowMaskSubPixelTexture,
@@ -1961,7 +1961,7 @@ void FDeferredShadingSceneRenderer::RenderDeferredShadowProjections(
 	const bool bProjectingForForwardShading = false;
 	RenderShadowProjections(GraphBuilder, SceneTextures, ScreenShadowMaskTexture, ScreenShadowMaskSubPixelTexture, LightSceneInfo, HairVisibilityViews, bProjectingForForwardShading);
 
-	// Perform injection on tranlucent lighting volume
+	// Perform injection on translucent lighting volume
 	{
 		const TArray<FProjectedShadowInfo*, SceneRenderingAllocator>& ShadowMaps = VisibleLightInfo.CompleteProjectedShadows.Num() > 0
 			? VisibleLightInfo.CompleteProjectedShadows
@@ -1997,7 +1997,7 @@ void FDeferredShadingSceneRenderer::RenderDeferredShadowProjections(
 					}
 
 					RDG_GPU_MASK_SCOPE(GraphBuilder, ProjectedShadowInfo->DependentView->GPUMask);
-					InjectTranslucentVolumeLighting(GraphBuilder, TranslucentVolumeLightingTextures, *LightSceneInfo, ProjectedShadowInfo, *ProjectedShadowInfo->DependentView, ViewIndex);
+					InjectTranslucencyLightingVolume(GraphBuilder, *ProjectedShadowInfo->DependentView, ViewIndex, Scene, TranslucencyLightingVolumeTextures, VisibleLightInfos, *LightSceneInfo, ProjectedShadowInfo);
 				}
 				else
 				{
@@ -2005,7 +2005,7 @@ void FDeferredShadingSceneRenderer::RenderDeferredShadowProjections(
 					{
 						FViewInfo& View = Views[ViewIndex];
 						RDG_GPU_MASK_SCOPE(GraphBuilder, View.GPUMask);
-						InjectTranslucentVolumeLighting(GraphBuilder, TranslucentVolumeLightingTextures, *LightSceneInfo, ProjectedShadowInfo, View, ViewIndex);
+						InjectTranslucencyLightingVolume(GraphBuilder, View, ViewIndex, Scene, TranslucencyLightingVolumeTextures, VisibleLightInfos, *LightSceneInfo, ProjectedShadowInfo);
 					}
 				}
 			}
diff --git a/Engine/Source/Runtime/Renderer/Private/SkyAtmosphereRendering.cpp b/Engine/Source/Runtime/Renderer/Private/SkyAtmosphereRendering.cpp
--- a/Engine/Source/Runtime/Renderer/Private/SkyAtmosphereRendering.cpp
+++ b/Engine/Source/Runtime/Renderer/Private/SkyAtmosphereRendering.cpp
@@ -1847,7 +1847,7 @@ void FSceneRenderer::RenderSkyAtmosphere(FRDGBuilder& GraphBuilder, const FMinim
 #endif
 }
 
-bool FSceneRenderer::ShouldRenderSkyAtmosphereEditorNotifications()
+bool FSceneRenderer::ShouldRenderSkyAtmosphereEditorNotifications() const
 {
 #if WITH_EDITOR
 	bool bAnyViewHasSkyMaterial = false;
@@ -1860,7 +1860,7 @@ bool FSceneRenderer::ShouldRenderSkyAtmosphereEditorNotifications()
 	return false;
 }
 
-void FSceneRenderer::RenderSkyAtmosphereEditorNotifications(FRDGBuilder& GraphBuilder, FRDGTextureRef SceneColorTexture)
+void FSceneRenderer::RenderSkyAtmosphereEditorNotifications(FRDGBuilder& GraphBuilder, FRDGTextureRef SceneColorTexture) const
 {
 #if WITH_EDITOR
 	RDG_EVENT_SCOPE(GraphBuilder, "SkyAtmosphereEditor");
@@ -1868,7 +1868,7 @@ void FSceneRenderer::RenderSkyAtmosphereEditorNotifications(FRDGBuilder& GraphBu
 
 	for (int32 ViewIndex = 0; ViewIndex < Views.Num(); ViewIndex++)
 	{
-		FViewInfo& View = Views[ViewIndex];
+		const FViewInfo& View = Views[ViewIndex];
 		if (View.bSceneHasSkyMaterial && View.Family->EngineShowFlags.Atmosphere)
 		{
 			RenderSkyAtmosphereEditorHudPS::FPermutationDomain PermutationVector;
diff --git a/Engine/Source/Runtime/Renderer/Private/SkyAtmosphereRendering.h b/Engine/Source/Runtime/Renderer/Private/SkyAtmosphereRendering.h
--- a/Engine/Source/Runtime/Renderer/Private/SkyAtmosphereRendering.h
+++ b/Engine/Source/Runtime/Renderer/Private/SkyAtmosphereRendering.h
@@ -6,9 +6,8 @@
 
 #pragma once
 
-#include "CoreMinimal.h"
 #include "EngineDefines.h"
-#include "RendererInterface.h"
+#include "RenderGraph.h"
 #include "RenderResource.h"
 #include "Rendering/SkyAtmosphereCommonData.h"
 #include "RenderGraphResources.h"
@@ -134,6 +133,8 @@ public:
 	TRefCountPtr<IPooledRenderTarget>& GetMultiScatteredLuminanceLutTexture() { return MultiScatteredLuminanceLutTexture; }
 	TRefCountPtr<IPooledRenderTarget>& GetDistantSkyLightLutTexture();
 
+	FRDGTextureRef GetTransmittanceLutTexture(FRDGBuilder& GraphBuilder) const { return GraphBuilder.RegisterExternalTexture(TransmittanceLutTexture); }
+
 	const FAtmosphereUniformShaderParameters* GetAtmosphereShaderParameters() const { return &AtmosphereUniformShaderParameters; }
 	const FSkyAtmosphereSceneProxy& GetSkyAtmosphereSceneProxy() const { return SkyAtmosphereSceneProxy; }
 
diff --git a/Engine/Source/Runtime/Renderer/Private/TranslucentLighting.cpp b/Engine/Source/Runtime/Renderer/Private/TranslucentLighting.cpp
--- a/Engine/Source/Runtime/Renderer/Private/TranslucentLighting.cpp
+++ b/Engine/Source/Runtime/Renderer/Private/TranslucentLighting.cpp
@@ -4,6 +4,7 @@
 	TranslucentLighting.cpp: Translucent lighting implementation.
 =============================================================================*/
 
+#include "TranslucentLighting.h"
 #include "CoreMinimal.h"
 #include "Stats/Stats.h"
 #include "HAL/IConsoleManager.h"
@@ -855,28 +856,33 @@ public:
 
 IMPLEMENT_GLOBAL_SHADER(FClearTranslucentLightingVolumeCS, "/Engine/Private/TranslucentLightInjectionShaders.usf", "ClearTranslucentLightingVolumeCS", SF_Compute);
 
-void FDeferredShadingSceneRenderer::InitTranslucentVolumeLighting(FRDGBuilder& GraphBuilder, ERDGPassFlags PassFlags, FTranslucentVolumeLightingTextures& Textures)
+void InitTranslucencyLightingVolumeTextures(FRDGBuilder& GraphBuilder, TArrayView<const FViewInfo> Views, ERDGPassFlags PassFlags, FTranslucencyLightingVolumeTextures& OutTextures)
 {
+	check(PassFlags == ERDGPassFlags::Compute || PassFlags == ERDGPassFlags::AsyncCompute);
+
 	RDG_GPU_STAT_SCOPE(GraphBuilder, TranslucentLighting);
 
-	Textures.VolumeDim = GetTranslucencyLightingVolumeDim();
-	const FIntVector TranslucencyLightingVolumeDim(Textures.VolumeDim);
+	OutTextures.VolumeDim = GetTranslucencyLightingVolumeDim();
+	const FIntVector TranslucencyLightingVolumeDim(OutTextures.VolumeDim);
 
 	{
 		// TODO: We can skip the and TLV allocations when rendering in forward shading mode
 		const ETextureCreateFlags TranslucencyTargetFlags = TexCreate_ShaderResource | TexCreate_RenderTargetable | TexCreate_ReduceMemoryWithTilingMode | TexCreate_UAV;
 
 		const int32 ViewCount = Views.Num();
-		Textures.Ambient.SetNum(ViewCount * TVC_MAX);
-		Textures.Directional.SetNum(ViewCount * TVC_MAX);
+		check(ViewCount > 0);
+
+		OutTextures.Ambient.SetNum(ViewCount * TVC_MAX);
+		OutTextures.Directional.SetNum(ViewCount * TVC_MAX);
 
-		int32 TextureIndex = 0;
 		for (int32 ViewIndex = 0; ViewIndex < ViewCount; ++ViewIndex)
 		{
 			for (int32 CascadeIndex = 0; CascadeIndex < TVC_MAX; ++CascadeIndex)
 			{
-				FRDGEventName& AmbientName = *GraphBuilder.AllocObject<FRDGEventName>(RDG_EVENT_NAME("TranslucentVolumeAmbient%d", TextureIndex));
-				FRDGEventName& DirectionalName = *GraphBuilder.AllocObject<FRDGEventName>(RDG_EVENT_NAME("TranslucentVolumeDirectional%d", TextureIndex));
+				const uint32 TextureIndex = FTranslucencyLightingVolumeTextures::GetIndex(ViewIndex, CascadeIndex);
+
+				const FRDGEventName& AmbientName = *GraphBuilder.AllocObject<FRDGEventName>(RDG_EVENT_NAME("TranslucentVolumeAmbient%d", TextureIndex));
+				const FRDGEventName& DirectionalName = *GraphBuilder.AllocObject<FRDGEventName>(RDG_EVENT_NAME("TranslucentVolumeDirectional%d", TextureIndex));
 
 				FRDGTextureRef AmbientTexture = GraphBuilder.CreateTexture(
 					FRDGTextureDesc::Create3D(
@@ -886,8 +892,6 @@ void FDeferredShadingSceneRenderer::InitTranslucentVolumeLighting(FRDGBuilder& G
 						TranslucencyTargetFlags),
 					AmbientName.GetTCHAR());
 
-				Textures.SetAmbient(ViewIndex, CascadeIndex, AmbientTexture);
-
 				FRDGTextureRef DirectionalTexture = GraphBuilder.CreateTexture(
 					FRDGTextureDesc::Create3D(
 						TranslucencyLightingVolumeDim,
@@ -896,27 +900,27 @@ void FDeferredShadingSceneRenderer::InitTranslucentVolumeLighting(FRDGBuilder& G
 						TranslucencyTargetFlags),
 					DirectionalName.GetTCHAR());
 
-				Textures.SetDirectional(ViewIndex, CascadeIndex, DirectionalTexture);
-				TextureIndex++;
+				OutTextures.Ambient[TextureIndex] = AmbientTexture;
+				OutTextures.Directional[TextureIndex] = DirectionalTexture;
 			}
 		}
 	}
 
 	const FIntVector GroupCount = FComputeShaderUtils::GetGroupCount(TranslucencyLightingVolumeDim, FClearTranslucentLightingVolumeCS::CLEAR_BLOCK_SIZE);
 
-	TShaderMapRef<FClearTranslucentLightingVolumeCS> ComputeShader(GetGlobalShaderMap(FeatureLevel));
+	TShaderMapRef<FClearTranslucentLightingVolumeCS> ComputeShader(Views[0].ShaderMap);
 
 	for (int32 ViewIndex = 0; ViewIndex < Views.Num(); ++ViewIndex)
 	{
 		auto* PassParameters = GraphBuilder.AllocParameters<FClearTranslucentLightingVolumeCS::FParameters>();
-		PassParameters->RWAmbient0 = GraphBuilder.CreateUAV(Textures.GetAmbient(ViewIndex, 0));
-		PassParameters->RWAmbient1 = GraphBuilder.CreateUAV(Textures.GetAmbient(ViewIndex, 1));
-		PassParameters->RWDirectional0 = GraphBuilder.CreateUAV(Textures.GetDirectional(ViewIndex, 0));
-		PassParameters->RWDirectional1 = GraphBuilder.CreateUAV(Textures.GetDirectional(ViewIndex, 1));
+		PassParameters->RWAmbient0 = GraphBuilder.CreateUAV(OutTextures.Ambient[FTranslucencyLightingVolumeTextures::GetIndex(ViewIndex, 0)]);
+		PassParameters->RWAmbient1 = GraphBuilder.CreateUAV(OutTextures.Ambient[FTranslucencyLightingVolumeTextures::GetIndex(ViewIndex, 1)]);
+		PassParameters->RWDirectional0 = GraphBuilder.CreateUAV(OutTextures.Directional[FTranslucencyLightingVolumeTextures::GetIndex(ViewIndex, 0)]);
+		PassParameters->RWDirectional1 = GraphBuilder.CreateUAV(OutTextures.Directional[FTranslucencyLightingVolumeTextures::GetIndex(ViewIndex, 1)]);
 
 		FComputeShaderUtils::AddPass(
 			GraphBuilder,
-			RDG_EVENT_NAME("ClearTranslucencyLightingVolumeCompute %d", Textures.VolumeDim),
+			RDG_EVENT_NAME("ClearTranslucencyLightingVolumeCompute %d", OutTextures.VolumeDim),
 			PassFlags,
 			ComputeShader,
 			PassParameters,
@@ -924,6 +928,30 @@ void FDeferredShadingSceneRenderer::InitTranslucentVolumeLighting(FRDGBuilder& G
 	}
 }
 
+FTranslucencyLightingVolumeParameters GetTranslucencyLightingVolumeParameters(FRDGBuilder& GraphBuilder, const FTranslucencyLightingVolumeTextures& Textures, uint32 ViewIndex)
+{
+	FTranslucencyLightingVolumeParameters Parameters;
+	if (Textures.IsValid())
+	{
+		const uint32 InnerIndex = FTranslucencyLightingVolumeTextures::GetIndex(ViewIndex, TVC_Inner);
+		const uint32 OuterIndex = FTranslucencyLightingVolumeTextures::GetIndex(ViewIndex, TVC_Outer);
+
+		Parameters.TranslucencyLightingVolumeAmbientInner = Textures.Ambient[InnerIndex];
+		Parameters.TranslucencyLightingVolumeAmbientOuter = Textures.Ambient[OuterIndex];
+		Parameters.TranslucencyLightingVolumeDirectionalInner = Textures.Directional[InnerIndex];
+		Parameters.TranslucencyLightingVolumeDirectionalOuter = Textures.Directional[OuterIndex];
+	}
+	else
+	{
+		const FRDGSystemTextures& SystemTextures = FRDGSystemTextures::Get(GraphBuilder);
+		Parameters.TranslucencyLightingVolumeAmbientInner = SystemTextures.Black;
+		Parameters.TranslucencyLightingVolumeAmbientOuter = SystemTextures.Black;
+		Parameters.TranslucencyLightingVolumeDirectionalInner = SystemTextures.Black;
+		Parameters.TranslucencyLightingVolumeDirectionalOuter = SystemTextures.Black;
+	}
+	return Parameters;
+}
+
 class FInjectAmbientCubemapPS : public FGlobalShader
 {
 public:
@@ -944,13 +972,12 @@ public:
 
 IMPLEMENT_GLOBAL_SHADER(FInjectAmbientCubemapPS, "/Engine/Private/TranslucentLightingShaders.usf", "InjectAmbientCubemapMainPS", SF_Pixel);
 
-void FDeferredShadingSceneRenderer::InjectAmbientCubemapTranslucentVolumeLighting(
+void InjectTranslucencyLightingVolumeAmbientCubemap(
 	FRDGBuilder& GraphBuilder,
-	const FTranslucentVolumeLightingTextures& Textures,
-	const FViewInfo& View,
-	int32 ViewIndex)
+	const TArrayView<const FViewInfo> Views,
+	const FTranslucencyLightingVolumeTextures& Textures)
 {
-	if (!GUseTranslucentLightingVolumes || !GSupportsVolumeTextureRendering || !View.FinalPostProcessSettings.ContributingCubemaps.Num())
+	if (!GUseTranslucentLightingVolumes || !GSupportsVolumeTextureRendering)
 	{
 		return;
 	}
@@ -961,51 +988,56 @@ void FDeferredShadingSceneRenderer::InjectAmbientCubemapTranslucentVolumeLightin
 	const int32 TranslucencyLightingVolumeDim = Textures.VolumeDim;
 	const FVolumeBounds VolumeBounds(TranslucencyLightingVolumeDim);
 
-	auto ShaderMap = GetGlobalShaderMap(FeatureLevel);
-
-	for (int32 VolumeCascadeIndex = 0; VolumeCascadeIndex < TVC_MAX; ++VolumeCascadeIndex)
+	for (int32 ViewIndex = 0; ViewIndex < Views.Num(); ++ViewIndex)
 	{
-		FRDGTextureRef VolumeAmbientTexture = Textures.GetAmbient(ViewIndex, VolumeCascadeIndex);
+		const FViewInfo& View = Views[ViewIndex];
+		RDG_GPU_MASK_SCOPE(GraphBuilder, View.GPUMask);
+		RDG_EVENT_SCOPE_CONDITIONAL(GraphBuilder, Views.Num() > 1, "View%d", ViewIndex);
 
-		for (const FFinalPostProcessSettings::FCubemapEntry& CubemapEntry : View.FinalPostProcessSettings.ContributingCubemaps)
+		for (int32 VolumeCascadeIndex = 0; VolumeCascadeIndex < TVC_MAX; ++VolumeCascadeIndex)
 		{
-			auto* PassParameters = GraphBuilder.AllocParameters<FInjectAmbientCubemapPS::FParameters>();
-			SetupAmbientCubemapParameters(CubemapEntry, &PassParameters->AmbientCubemap);
-			PassParameters->RenderTargets[0] = FRenderTargetBinding(VolumeAmbientTexture, ERenderTargetLoadAction::ELoad);
-			PassParameters->View = View.ViewUniformBuffer;
+			FRDGTextureRef VolumeAmbientTexture = Textures.Ambient[FTranslucencyLightingVolumeTextures::GetIndex(ViewIndex, VolumeCascadeIndex)];
 
-			GraphBuilder.AddPass(
-				RDG_EVENT_NAME("Cascade %d", VolumeCascadeIndex),
-				PassParameters,
-				ERDGPassFlags::Raster,
-				[&View, PassParameters, ShaderMap, VolumeBounds, TranslucencyLightingVolumeDim](FRHICommandList& RHICmdList)
+			for (const FFinalPostProcessSettings::FCubemapEntry& CubemapEntry : View.FinalPostProcessSettings.ContributingCubemaps)
 			{
-				TShaderMapRef<FWriteToSliceVS> VertexShader(ShaderMap);
-				TOptionalShaderMapRef<FWriteToSliceGS> GeometryShader(ShaderMap);
-				TShaderMapRef<FInjectAmbientCubemapPS> PixelShader(ShaderMap);
+				auto* PassParameters = GraphBuilder.AllocParameters<FInjectAmbientCubemapPS::FParameters>();
+				SetupAmbientCubemapParameters(CubemapEntry, &PassParameters->AmbientCubemap);
+				PassParameters->RenderTargets[0] = FRenderTargetBinding(VolumeAmbientTexture, ERenderTargetLoadAction::ELoad);
+				PassParameters->View = View.ViewUniformBuffer;
 
-				FGraphicsPipelineStateInitializer GraphicsPSOInit;
-				RHICmdList.ApplyCachedRenderTargets(GraphicsPSOInit);
+				GraphBuilder.AddPass(
+					RDG_EVENT_NAME("Cascade %d", VolumeCascadeIndex),
+					PassParameters,
+					ERDGPassFlags::Raster,
+					[&View, PassParameters, VolumeBounds, TranslucencyLightingVolumeDim](FRHICommandList& RHICmdList)
+				{
+					TShaderMapRef<FWriteToSliceVS> VertexShader(View.ShaderMap);
+					TOptionalShaderMapRef<FWriteToSliceGS> GeometryShader(View.ShaderMap);
+					TShaderMapRef<FInjectAmbientCubemapPS> PixelShader(View.ShaderMap);
 
-				GraphicsPSOInit.RasterizerState = TStaticRasterizerState<FM_Solid, CM_None>::GetRHI();
-				GraphicsPSOInit.DepthStencilState = TStaticDepthStencilState<false, CF_Always>::GetRHI();
-				GraphicsPSOInit.BlendState = TStaticBlendState<CW_RGBA, BO_Add, BF_One, BF_One, BO_Add, BF_One, BF_One>::GetRHI();
-				GraphicsPSOInit.BoundShaderState.VertexDeclarationRHI = GScreenVertexDeclaration.VertexDeclarationRHI;
-				GraphicsPSOInit.BoundShaderState.VertexShaderRHI = VertexShader.GetVertexShader();
+					FGraphicsPipelineStateInitializer GraphicsPSOInit;
+					RHICmdList.ApplyCachedRenderTargets(GraphicsPSOInit);
+
+					GraphicsPSOInit.RasterizerState = TStaticRasterizerState<FM_Solid, CM_None>::GetRHI();
+					GraphicsPSOInit.DepthStencilState = TStaticDepthStencilState<false, CF_Always>::GetRHI();
+					GraphicsPSOInit.BlendState = TStaticBlendState<CW_RGBA, BO_Add, BF_One, BF_One, BO_Add, BF_One, BF_One>::GetRHI();
+					GraphicsPSOInit.BoundShaderState.VertexDeclarationRHI = GScreenVertexDeclaration.VertexDeclarationRHI;
+					GraphicsPSOInit.BoundShaderState.VertexShaderRHI = VertexShader.GetVertexShader();
 #if PLATFORM_SUPPORTS_GEOMETRY_SHADERS
-				GraphicsPSOInit.BoundShaderState.GeometryShaderRHI = GeometryShader.GetGeometryShader();
+					GraphicsPSOInit.BoundShaderState.GeometryShaderRHI = GeometryShader.GetGeometryShader();
 #endif
-				GraphicsPSOInit.BoundShaderState.PixelShaderRHI = PixelShader.GetPixelShader();
-				GraphicsPSOInit.PrimitiveType = PT_TriangleStrip;
-				SetGraphicsPipelineState(RHICmdList, GraphicsPSOInit);
+					GraphicsPSOInit.BoundShaderState.PixelShaderRHI = PixelShader.GetPixelShader();
+					GraphicsPSOInit.PrimitiveType = PT_TriangleStrip;
+					SetGraphicsPipelineState(RHICmdList, GraphicsPSOInit);
 
-				VertexShader->SetParameters(RHICmdList, VolumeBounds, FIntVector(TranslucencyLightingVolumeDim));
-				if (GeometryShader.IsValid())
-				{
-					GeometryShader->SetParameters(RHICmdList, VolumeBounds.MinZ);
-				}
-				SetShaderParameters(RHICmdList, PixelShader, PixelShader.GetPixelShader(), *PassParameters);
-			});
+					VertexShader->SetParameters(RHICmdList, VolumeBounds, FIntVector(TranslucencyLightingVolumeDim));
+					if (GeometryShader.IsValid())
+					{
+						GeometryShader->SetParameters(RHICmdList, VolumeBounds.MinZ);
+					}
+					SetShaderParameters(RHICmdList, PixelShader, PixelShader.GetPixelShader(), *PassParameters);
+				});
+			}
 		}
 	}
 }
@@ -1151,18 +1183,18 @@ struct FTranslucentLightInjectionData
  * @param InProjectedShadowInfo is 0 for unshadowed lights
  */
 static void AddLightForInjection(
-	FDeferredShadingSceneRenderer& SceneRenderer,
-	const FLightSceneInfo& LightSceneInfo, 
+	const FViewInfo& View,
+	TArrayView<const FVisibleLightInfo> VisibleLightInfos,
+	const FLightSceneInfo& LightSceneInfo,
 	const FProjectedShadowInfo* InProjectedShadowInfo,
 	TArray<FTranslucentLightInjectionData, SceneRenderingAllocator>& LightInjectionData)
 {
 	if (LightSceneInfo.Proxy->AffectsTranslucentLighting())
 	{
-		const FVisibleLightInfo& VisibleLightInfo = SceneRenderer.VisibleLightInfos[LightSceneInfo.Id];
-
-		const ERHIFeatureLevel::Type FeatureLevel = SceneRenderer.Scene->GetFeatureLevel();
+		const FVisibleLightInfo& VisibleLightInfo = VisibleLightInfos[LightSceneInfo.Id];
+		const ERHIFeatureLevel::Type FeatureLevel = View.FeatureLevel;
 
-		const bool bApplyLightFunction = (SceneRenderer.ViewFamily.EngineShowFlags.LightFunctions &&
+		const bool bApplyLightFunction = (View.Family->EngineShowFlags.LightFunctions &&
 			LightSceneInfo.Proxy->GetLightFunctionMaterial() && 
 			LightSceneInfo.Proxy->GetLightFunctionMaterial()->GetIncompleteMaterialWithFallback(FeatureLevel).IsLightFunction());
 
@@ -1183,15 +1215,14 @@ static void AddLightForInjection(
 	}
 }
 
-static FRDGTextureRef GetSkyTransmittanceLutTexture(FRDGBuilder& GraphBuilder, FScene* Scene, const FViewInfo& View)
+static FRDGTextureRef GetSkyTransmittanceLutTexture(FRDGBuilder& GraphBuilder, const FScene* Scene, const FViewInfo& View)
 {
 	FRDGTextureRef TransmittanceLutTexture = nullptr;
 	if (ShouldRenderSkyAtmosphere(Scene, View.Family->EngineShowFlags))
 	{
-		if (FSkyAtmosphereRenderSceneInfo* SkyInfo = Scene->GetSkyAtmosphereSceneInfo())
+		if (const FSkyAtmosphereRenderSceneInfo* SkyInfo = Scene->GetSkyAtmosphereSceneInfo())
 		{
-			TRefCountPtr<IPooledRenderTarget>& TransmittanceLutTarget = SkyInfo->GetTransmittanceLutTexture();
-			TransmittanceLutTexture = GraphBuilder.RegisterExternalTexture(TransmittanceLutTarget, TEXT("TransmittanceLutTexture"));
+			TransmittanceLutTexture = SkyInfo->GetTransmittanceLutTexture(GraphBuilder);
 		}
 	}
 	return TransmittanceLutTexture;
@@ -1206,23 +1237,23 @@ END_SHADER_PARAMETER_STRUCT()
 /** Injects all the lights in LightInjectionData into the translucent lighting volume textures. */
 static void InjectTranslucentLightArray(
 	FRDGBuilder& GraphBuilder,
-	const FTranslucentVolumeLightingTextures& Textures,
-	FScene* Scene,
 	const FViewInfo& View,
-	TArrayView<const FTranslucentLightInjectionData> LightInjectionData,
-	int32 ViewIndex)
+	const uint32 ViewIndex,
+	const FScene* Scene,
+	const FTranslucencyLightingVolumeTextures& Textures,
+	TArrayView<const FTranslucentLightInjectionData> LightInjectionData)
 {
-	FSceneRenderTargets& SceneContext = FSceneRenderTargets::Get();
 	INC_DWORD_STAT_BY(STAT_NumLightsInjectedIntoTranslucency, LightInjectionData.Num());
 
 	FRDGTextureRef TransmittanceLutTexture = GetSkyTransmittanceLutTexture(GraphBuilder, Scene, View);
 
-	// Inject into each volume cascade
-	// Operate on one cascade at a time to reduce render target switches
+
+	// Inject into each volume cascade. Operate on one cascade at a time to reduce render target switches.
 	for (uint32 VolumeCascadeIndex = 0; VolumeCascadeIndex < TVC_MAX; VolumeCascadeIndex++)
 	{
-		FRDGTextureRef VolumeAmbientTexture = Textures.GetAmbient(ViewIndex, VolumeCascadeIndex);
-		FRDGTextureRef VolumeDirectionalTexture = Textures.GetDirectional(ViewIndex, VolumeCascadeIndex);
+		const uint32 TextureIndex = FTranslucencyLightingVolumeTextures::GetIndex(ViewIndex, VolumeCascadeIndex);
+		FRDGTextureRef VolumeAmbientTexture = Textures.Ambient[TextureIndex];
+		FRDGTextureRef VolumeDirectionalTexture = Textures.Directional[TextureIndex];
 
 		for (int32 LightIndex = 0; LightIndex < LightInjectionData.Num(); LightIndex++)
 		{
@@ -1321,33 +1352,34 @@ static void InjectTranslucentLightArray(
 	}
 }
 
-void FDeferredShadingSceneRenderer::InjectTranslucentVolumeLighting(
+void InjectTranslucencyLightingVolume(
 	FRDGBuilder& GraphBuilder,
-	const FTranslucentVolumeLightingTextures& Textures,
-	const FLightSceneInfo& LightSceneInfo,
-	const FProjectedShadowInfo* InProjectedShadowInfo,
 	const FViewInfo& View,
-	int32 ViewIndex)
+	const uint32 ViewIndex,
+	const FScene* Scene,
+	const FTranslucencyLightingVolumeTextures& Textures,
+	TArrayView<const FVisibleLightInfo> VisibleLightInfos,
+	const FLightSceneInfo& LightSceneInfo,
+	const FProjectedShadowInfo* ProjectedShadowInfo)
 {
 	if (GUseTranslucentLightingVolumes && GSupportsVolumeTextureRendering)
 	{
 		SCOPE_CYCLE_COUNTER(STAT_TranslucentInjectTime);
 
 		auto& LightInjectionData = *GraphBuilder.AllocObject<TArray<FTranslucentLightInjectionData, SceneRenderingAllocator>>();
-
-		AddLightForInjection(*this, LightSceneInfo, InProjectedShadowInfo, LightInjectionData);
-
-		// shadowed or unshadowed (InProjectedShadowInfo==0)
-		InjectTranslucentLightArray(GraphBuilder, Textures, Scene, View, LightInjectionData, ViewIndex);
+		AddLightForInjection(View, VisibleLightInfos, LightSceneInfo, ProjectedShadowInfo, LightInjectionData);
+		InjectTranslucentLightArray(GraphBuilder, View, ViewIndex, Scene, Textures, LightInjectionData);
 	}
 }
 
-void FDeferredShadingSceneRenderer::InjectTranslucentVolumeLightingArray(
+void InjectTranslucencyLightingVolumeArray(
 	FRDGBuilder& GraphBuilder,
-	const FTranslucentVolumeLightingTextures& Textures,
-	const TArray<FSortedLightSceneInfo, SceneRenderingAllocator>& SortedLights,
-	int32 FirstLightIndex,
-	int32 LightsEndIndex)
+	const TArrayView<const FViewInfo> Views,
+	const FScene* Scene,
+	const FTranslucencyLightingVolumeTextures& Textures,
+	const TArrayView<const FVisibleLightInfo> VisibleLightInfos,
+	TArrayView<const FSortedLightSceneInfo> SortedLights,
+	TInterval<int32> SortedLightInterval)
 {
 	SCOPE_CYCLE_COUNTER(STAT_TranslucentInjectTime);
 
@@ -1357,26 +1389,29 @@ void FDeferredShadingSceneRenderer::InjectTranslucentVolumeLightingArray(
 
 	for (int32 ViewIndex = 0; ViewIndex < LightInjectionData.Num(); ++ViewIndex)
 	{
-		LightInjectionData[ViewIndex].Reserve(LightsEndIndex - FirstLightIndex);
+		LightInjectionData[ViewIndex].Reserve(SortedLightInterval.Size());
 	}
 
-	for (int32 LightIndex = FirstLightIndex; LightIndex < LightsEndIndex; LightIndex++)
+	for (int32 LightIndex = SortedLightInterval.Min; LightIndex < SortedLightInterval.Max; LightIndex++)
 	{
 		const FSortedLightSceneInfo& SortedLightInfo = SortedLights[LightIndex];
 		const FLightSceneInfo* const LightSceneInfo = SortedLightInfo.LightSceneInfo;
 		for (int32 ViewIndex = 0; ViewIndex < Views.Num(); ViewIndex++)
 		{
-			if (LightSceneInfo->ShouldRenderLight(Views[ViewIndex]))
+			const FViewInfo& View = Views[ViewIndex];
+			if (LightSceneInfo->ShouldRenderLight(View))
 			{
-				AddLightForInjection(*this, *LightSceneInfo, NULL, LightInjectionData[ViewIndex]);
+				AddLightForInjection(View, VisibleLightInfos, *LightSceneInfo, nullptr, LightInjectionData[ViewIndex]);
 			}
 		}
 	}
 
 	for (int32 ViewIndex = 0; ViewIndex < Views.Num(); ViewIndex++)
 	{
+		const FViewInfo& View = Views[ViewIndex];
+
 		// non-shadowed, non-light function lights
-		InjectTranslucentLightArray(GraphBuilder, Textures, Scene, Views[ViewIndex], LightInjectionData[ViewIndex], ViewIndex);
+		InjectTranslucentLightArray(GraphBuilder, View, ViewIndex, Scene, Textures, LightInjectionData[ViewIndex]);
 	}
 }
 
@@ -1402,12 +1437,13 @@ public:
 
 IMPLEMENT_GLOBAL_SHADER(FSimpleLightTranslucentLightingInjectPS, "/Engine/Private/TranslucentLightInjectionShaders.usf", "SimpleLightInjectMainPS", SF_Pixel);
 
-void FDeferredShadingSceneRenderer::InjectSimpleTranslucentVolumeLightingArray(
+void InjectSimpleTranslucencyLightingVolumeArray(
 	FRDGBuilder& GraphBuilder,
-	const FTranslucentVolumeLightingTextures& Textures,
-	const FSimpleLightArray& SimpleLights,
 	const FViewInfo& View,
-	const int32 ViewIndex)
+	const uint32 ViewIndex,
+	const uint32 ViewCount,
+	const FTranslucencyLightingVolumeTextures& Textures,
+	const FSimpleLightArray& SimpleLights)
 {
 	SCOPE_CYCLE_COUNTER(STAT_TranslucentInjectTime);
 
@@ -1433,14 +1469,16 @@ void FDeferredShadingSceneRenderer::InjectSimpleTranslucentVolumeLightingArray(
 		// Operate on one cascade at a time to reduce render target switches
 		for (int32 VolumeCascadeIndex = 0; VolumeCascadeIndex < TVC_MAX; VolumeCascadeIndex++)
 		{
+			const uint32 TextureIndex = FTranslucencyLightingVolumeTextures::GetIndex(ViewIndex, VolumeCascadeIndex);
+
 			RDG_EVENT_SCOPE(GraphBuilder, "Cascade%d", VolumeCascadeIndex);
-			FRDGTextureRef VolumeAmbientTexture = Textures.GetAmbient(ViewIndex, VolumeCascadeIndex);
-			FRDGTextureRef VolumeDirectionalTexture = Textures.GetDirectional(ViewIndex, VolumeCascadeIndex);
+			FRDGTextureRef VolumeAmbientTexture = Textures.Ambient[TextureIndex];
+			FRDGTextureRef VolumeDirectionalTexture = Textures.Directional[TextureIndex];
 
 			for (int32 LightIndex = 0; LightIndex < SimpleLights.InstanceData.Num(); LightIndex++)
 			{
 				const FSimpleLightEntry& SimpleLight = SimpleLights.InstanceData[LightIndex];
-				const FSimpleLightPerViewEntry& SimpleLightPerViewData = SimpleLights.GetViewDependentData(LightIndex, ViewIndex, Views.Num());
+				const FSimpleLightPerViewEntry& SimpleLightPerViewData = SimpleLights.GetViewDependentData(LightIndex, ViewIndex, ViewCount);
 
 				if (SimpleLight.bAffectTranslucency)
 				{
@@ -1501,79 +1539,85 @@ void FDeferredShadingSceneRenderer::InjectSimpleTranslucentVolumeLightingArray(
 	}
 }
 
-void FDeferredShadingSceneRenderer::FilterTranslucentVolumeLighting(
+void FilterTranslucencyLightingVolume(
 	FRDGBuilder& GraphBuilder,
-	FTranslucentVolumeLightingTextures& Textures,
-	const FViewInfo& View,
-	const int32 ViewIndex)
+	const TArrayView<const FViewInfo> Views,
+	FTranslucencyLightingVolumeTextures& Textures)
 {
 	if (!GUseTranslucentLightingVolumes || !GSupportsVolumeTextureRendering || !GUseTranslucencyVolumeBlur)
 	{
 		return;
 	}
 
-	FSceneRenderTargets& SceneContext = FSceneRenderTargets::Get();
-
 	FRHISamplerState* SamplerStateRHI = TStaticSamplerState<SF_Bilinear, AM_Clamp, AM_Clamp, AM_Clamp>::GetRHI();
 
 	const int32 TranslucencyLightingVolumeDim = GetTranslucencyLightingVolumeDim();
 	RDG_EVENT_SCOPE(GraphBuilder, "FilterTranslucentVolume %dx%dx%d Cascades:%d", TranslucencyLightingVolumeDim, TranslucencyLightingVolumeDim, TranslucencyLightingVolumeDim, TVC_MAX);
 	RDG_GPU_STAT_SCOPE(GraphBuilder, TranslucentLighting);
 
-	for (int32 VolumeCascadeIndex = 0; VolumeCascadeIndex < TVC_MAX; VolumeCascadeIndex++)
+	for (int32 ViewIndex = 0; ViewIndex < Views.Num(); ++ViewIndex)
 	{
-		FRDGTextureRef InputVolumeAmbientTexture = Textures.GetAmbient(ViewIndex, VolumeCascadeIndex);
-		FRDGTextureRef InputVolumeDirectionalTexture = Textures.GetDirectional(ViewIndex, VolumeCascadeIndex);
-
-		FRDGTextureRef OutputVolumeAmbientTexture = GraphBuilder.CreateTexture(InputVolumeAmbientTexture->Desc, InputVolumeAmbientTexture->Name);
-		FRDGTextureRef OutputVolumeDirectionalTexture = GraphBuilder.CreateTexture(InputVolumeDirectionalTexture->Desc, InputVolumeDirectionalTexture->Name);
-
-		Textures.SetAmbient(ViewIndex, VolumeCascadeIndex, OutputVolumeAmbientTexture);
-		Textures.SetDirectional(ViewIndex, VolumeCascadeIndex, OutputVolumeDirectionalTexture);
-
-		auto* PassParameters = GraphBuilder.AllocParameters<FFilterTranslucentVolumePS::FParameters>();
-		PassParameters->View = View.ViewUniformBuffer;
-		PassParameters->TexelSize = 1.0f / TranslucencyLightingVolumeDim;
-		PassParameters->TranslucencyLightingVolumeAmbient = InputVolumeAmbientTexture;
-		PassParameters->TranslucencyLightingVolumeDirectional = InputVolumeDirectionalTexture;
-		PassParameters->TranslucencyLightingVolumeAmbientSampler = SamplerStateRHI;
-		PassParameters->TranslucencyLightingVolumeDirectionalSampler = SamplerStateRHI;
-		PassParameters->RenderTargets[0] = FRenderTargetBinding(OutputVolumeAmbientTexture, ERenderTargetLoadAction::ENoAction);
-		PassParameters->RenderTargets[1] = FRenderTargetBinding(OutputVolumeDirectionalTexture, ERenderTargetLoadAction::ENoAction);
-
-		const FVolumeBounds VolumeBounds(TranslucencyLightingVolumeDim);
-		TShaderMapRef<FWriteToSliceVS> VertexShader(View.ShaderMap);
-		TOptionalShaderMapRef<FWriteToSliceGS> GeometryShader(View.ShaderMap);
-		TShaderMapRef<FFilterTranslucentVolumePS> PixelShader(View.ShaderMap);
-
-		GraphBuilder.AddPass(
-			RDG_EVENT_NAME("Cascade%d", VolumeCascadeIndex),
-			PassParameters,
-			ERDGPassFlags::Raster,
-			[VertexShader, GeometryShader, PixelShader, PassParameters, VolumeBounds, TranslucencyLightingVolumeDim](FRHICommandList& RHICmdList)
+		const FViewInfo& View = Views[ViewIndex];
+		RDG_GPU_MASK_SCOPE(GraphBuilder, View.GPUMask);
+		RDG_EVENT_SCOPE_CONDITIONAL(GraphBuilder, Views.Num() > 1, "View%d", ViewIndex);
+
+		for (int32 VolumeCascadeIndex = 0; VolumeCascadeIndex < TVC_MAX; VolumeCascadeIndex++)
 		{
-			FGraphicsPipelineStateInitializer GraphicsPSOInit;
-			RHICmdList.ApplyCachedRenderTargets(GraphicsPSOInit);
-
-			GraphicsPSOInit.RasterizerState = TStaticRasterizerState<FM_Solid, CM_None>::GetRHI();
-			GraphicsPSOInit.DepthStencilState = TStaticDepthStencilState<false, CF_Always>::GetRHI();
-			GraphicsPSOInit.BlendState = TStaticBlendState<>::GetRHI();
-			GraphicsPSOInit.BoundShaderState.VertexDeclarationRHI = GScreenVertexDeclaration.VertexDeclarationRHI;
-			GraphicsPSOInit.BoundShaderState.VertexShaderRHI = VertexShader.GetVertexShader();
+			const uint32 TextureIndex = FTranslucencyLightingVolumeTextures::GetIndex(ViewIndex, VolumeCascadeIndex);
+
+			FRDGTextureRef InputVolumeAmbientTexture = Textures.Ambient[TextureIndex];
+			FRDGTextureRef InputVolumeDirectionalTexture = Textures.Directional[TextureIndex];
+
+			FRDGTextureRef OutputVolumeAmbientTexture = GraphBuilder.CreateTexture(InputVolumeAmbientTexture->Desc, InputVolumeAmbientTexture->Name);
+			FRDGTextureRef OutputVolumeDirectionalTexture = GraphBuilder.CreateTexture(InputVolumeDirectionalTexture->Desc, InputVolumeDirectionalTexture->Name);
+
+			Textures.Ambient[TextureIndex] = OutputVolumeAmbientTexture;
+			Textures.Directional[TextureIndex] = OutputVolumeDirectionalTexture;
+
+			auto* PassParameters = GraphBuilder.AllocParameters<FFilterTranslucentVolumePS::FParameters>();
+			PassParameters->View = View.ViewUniformBuffer;
+			PassParameters->TexelSize = 1.0f / TranslucencyLightingVolumeDim;
+			PassParameters->TranslucencyLightingVolumeAmbient = InputVolumeAmbientTexture;
+			PassParameters->TranslucencyLightingVolumeDirectional = InputVolumeDirectionalTexture;
+			PassParameters->TranslucencyLightingVolumeAmbientSampler = SamplerStateRHI;
+			PassParameters->TranslucencyLightingVolumeDirectionalSampler = SamplerStateRHI;
+			PassParameters->RenderTargets[0] = FRenderTargetBinding(OutputVolumeAmbientTexture, ERenderTargetLoadAction::ENoAction);
+			PassParameters->RenderTargets[1] = FRenderTargetBinding(OutputVolumeDirectionalTexture, ERenderTargetLoadAction::ENoAction);
+
+			const FVolumeBounds VolumeBounds(TranslucencyLightingVolumeDim);
+			TShaderMapRef<FWriteToSliceVS> VertexShader(View.ShaderMap);
+			TOptionalShaderMapRef<FWriteToSliceGS> GeometryShader(View.ShaderMap);
+			TShaderMapRef<FFilterTranslucentVolumePS> PixelShader(View.ShaderMap);
+
+			GraphBuilder.AddPass(
+				RDG_EVENT_NAME("Cascade%d", VolumeCascadeIndex),
+				PassParameters,
+				ERDGPassFlags::Raster,
+				[VertexShader, GeometryShader, PixelShader, PassParameters, VolumeBounds, TranslucencyLightingVolumeDim](FRHICommandList& RHICmdList)
+			{
+				FGraphicsPipelineStateInitializer GraphicsPSOInit;
+				RHICmdList.ApplyCachedRenderTargets(GraphicsPSOInit);
+
+				GraphicsPSOInit.RasterizerState = TStaticRasterizerState<FM_Solid, CM_None>::GetRHI();
+				GraphicsPSOInit.DepthStencilState = TStaticDepthStencilState<false, CF_Always>::GetRHI();
+				GraphicsPSOInit.BlendState = TStaticBlendState<>::GetRHI();
+				GraphicsPSOInit.BoundShaderState.VertexDeclarationRHI = GScreenVertexDeclaration.VertexDeclarationRHI;
+				GraphicsPSOInit.BoundShaderState.VertexShaderRHI = VertexShader.GetVertexShader();
 #if PLATFORM_SUPPORTS_GEOMETRY_SHADERS
-			GraphicsPSOInit.BoundShaderState.GeometryShaderRHI = GeometryShader.GetGeometryShader();
+				GraphicsPSOInit.BoundShaderState.GeometryShaderRHI = GeometryShader.GetGeometryShader();
 #endif
-			GraphicsPSOInit.BoundShaderState.PixelShaderRHI = PixelShader.GetPixelShader();
-			GraphicsPSOInit.PrimitiveType = PT_TriangleStrip;
-			SetGraphicsPipelineState(RHICmdList, GraphicsPSOInit);
+				GraphicsPSOInit.BoundShaderState.PixelShaderRHI = PixelShader.GetPixelShader();
+				GraphicsPSOInit.PrimitiveType = PT_TriangleStrip;
+				SetGraphicsPipelineState(RHICmdList, GraphicsPSOInit);
 
-			VertexShader->SetParameters(RHICmdList, VolumeBounds, FIntVector(TranslucencyLightingVolumeDim));
-			if (GeometryShader.IsValid())
-			{
-				GeometryShader->SetParameters(RHICmdList, VolumeBounds.MinZ);
-			}
-			SetShaderParameters(RHICmdList, PixelShader, PixelShader.GetPixelShader(), *PassParameters);
-			RasterizeToVolumeTexture(RHICmdList, VolumeBounds);
-		});
+				VertexShader->SetParameters(RHICmdList, VolumeBounds, FIntVector(TranslucencyLightingVolumeDim));
+				if (GeometryShader.IsValid())
+				{
+					GeometryShader->SetParameters(RHICmdList, VolumeBounds.MinZ);
+				}
+				SetShaderParameters(RHICmdList, PixelShader, PixelShader.GetPixelShader(), *PassParameters);
+				RasterizeToVolumeTexture(RHICmdList, VolumeBounds);
+			});
+		}
 	}
 }
diff --git a/Engine/Source/Runtime/Renderer/Private/TranslucentLighting.h b/Engine/Source/Runtime/Renderer/Private/TranslucentLighting.h
new file mode 100644
--- /dev/null
+++ b/Engine/Source/Runtime/Renderer/Private/TranslucentLighting.h
@@ -0,0 +1,84 @@
+// Copyright Epic Games, Inc. All Rights Reserved.
+
+#pragma once
+
+#include "SceneRendering.h"
+#include "RenderGraph.h"
+
+struct FSortedLightSceneInfo;
+
+class FLightSceneInfo;
+class FProjectedShadowInfo;
+
+struct FTranslucencyLightingVolumeTextures
+{
+	static int32 GetIndex(int32 ViewIndex, int32 CascadeIndex)
+	{
+		return (ViewIndex * TVC_MAX) + CascadeIndex;
+	}
+
+	bool IsValid() const
+	{
+		check(!VolumeDim || (Ambient.Num() == Directional.Num() && Ambient.Num() > 0));
+		return VolumeDim != 0;
+	}
+
+	TArray<FRDGTextureRef, TInlineAllocator<TVC_MAX>> Ambient;
+	TArray<FRDGTextureRef, TInlineAllocator<TVC_MAX>> Directional;
+	int32 VolumeDim = 0;
+};
+
+/** Initializes the translucent volume textures and clears them using a compute pass. PassFlags should be ERDGPassFlags::Compute or ERDGPassFlags::AsyncCompute. */
+void InitTranslucencyLightingVolumeTextures(FRDGBuilder& GraphBuilder, TArrayView<const FViewInfo> Views, ERDGPassFlags PassFlags, FTranslucencyLightingVolumeTextures& OutTextures);
+
+BEGIN_SHADER_PARAMETER_STRUCT(FTranslucencyLightingVolumeParameters, )
+	SHADER_PARAMETER_RDG_TEXTURE(Texture3D, TranslucencyLightingVolumeAmbientInner)
+	SHADER_PARAMETER_RDG_TEXTURE(Texture3D, TranslucencyLightingVolumeAmbientOuter)
+	SHADER_PARAMETER_RDG_TEXTURE(Texture3D, TranslucencyLightingVolumeDirectionalInner)
+	SHADER_PARAMETER_RDG_TEXTURE(Texture3D, TranslucencyLightingVolumeDirectionalOuter)
+END_SHADER_PARAMETER_STRUCT()
+
+/** Initializes translucency volume lighting shader parameters from an optional textures struct. If null or uninitialized, fallback textures are used. */
+FTranslucencyLightingVolumeParameters GetTranslucencyLightingVolumeParameters(FRDGBuilder& GraphBuilder, const FTranslucencyLightingVolumeTextures& Textures, uint32 ViewIndex);
+
+void InjectTranslucencyLightingVolume(
+	FRDGBuilder& GraphBuilder,
+	const FViewInfo& View,
+	const uint32 ViewIndex,
+	const FScene* Scene,
+	const FTranslucencyLightingVolumeTextures& Textures,
+	const TArrayView<const FVisibleLightInfo> VisibleLightInfos,
+	const FLightSceneInfo& LightSceneInfo,
+	const FProjectedShadowInfo* ProjectedShadowInfo);
+
+void InjectTranslucencyLightingVolumeArray(
+	FRDGBuilder& GraphBuilder,
+	const TArrayView<const FViewInfo> Views,
+	const FScene* Scene,
+	const FTranslucencyLightingVolumeTextures& Textures,
+	const TArrayView<const FVisibleLightInfo> VisibleLightInfos,
+	const TArrayView<const FSortedLightSceneInfo> SortedLights,
+	const TInterval<int32> SortedLightInterval);
+
+void InjectSimpleTranslucencyLightingVolumeArray(
+	FRDGBuilder& GraphBuilder,
+	const FViewInfo& View,
+	const uint32 ViewIndex,
+	const uint32 ViewCount,
+	const FTranslucencyLightingVolumeTextures& Textures,
+	const FSimpleLightArray& SimpleLights);
+
+void FinalizeTranslucencyLightingVolume(
+	FRDGBuilder& GraphBuilder,
+	const TArrayView<const FViewInfo> Views,
+	FTranslucencyLightingVolumeTextures& Textures);
+
+void InjectTranslucencyLightingVolumeAmbientCubemap(
+	FRDGBuilder& GraphBuilder,
+	const TArrayView<const FViewInfo> Views,
+	const FTranslucencyLightingVolumeTextures& Textures);
+
+void FilterTranslucencyLightingVolume(
+	FRDGBuilder& GraphBuilder,
+	const TArrayView<const FViewInfo> Views,
+	FTranslucencyLightingVolumeTextures& Textures);
\ No newline at end of file
diff --git a/Engine/Source/Runtime/Renderer/Private/TranslucentRendering.cpp b/Engine/Source/Runtime/Renderer/Private/TranslucentRendering.cpp
--- a/Engine/Source/Runtime/Renderer/Private/TranslucentRendering.cpp
+++ b/Engine/Source/Runtime/Renderer/Private/TranslucentRendering.cpp
@@ -606,16 +606,16 @@ void SetupDownsampledTranslucencyViewParameters(
 		View.PrevViewInfo.ViewMatrices);
 }
 
-void SetupTranslucentBasePassUniformParameters(
+TRDGUniformBufferRef<FTranslucentBasePassUniformParameters> CreateTranslucentBasePassUniformBuffer(
 	FRDGBuilder& GraphBuilder,
-	const FSceneRenderTargets& SceneRenderTargets,
 	const FViewInfo& View,
-	const FTranslucentVolumeLightingTextures* TranslucentVolumeLightingTextures,
-	FRDGTextureRef SceneColorCopyTexture,
-	ESceneTextureSetupMode SceneTextureSetupMode,
 	const int32 ViewIndex,
-	FTranslucentBasePassUniformParameters& BasePassParameters)
+	const FTranslucencyLightingVolumeTextures& TranslucencyLightingVolumeTextures,
+	FRDGTextureRef SceneColorCopyTexture,
+	const ESceneTextureSetupMode SceneTextureSetupMode)
 {
+	FTranslucentBasePassUniformParameters& BasePassParameters = *GraphBuilder.AllocParameters<FTranslucentBasePassUniformParameters>();
+
 	const auto GetRDG = [&](const TRefCountPtr<IPooledRenderTarget>& PooledRenderTarget, ERDGTextureFlags Flags = ERDGTextureFlags::None)
 	{
 		return GraphBuilder.RegisterExternalTexture(PooledRenderTarget, ERenderTargetTexture::ShaderResource, Flags);
@@ -624,7 +624,7 @@ void SetupTranslucentBasePassUniformParameters(
 	SetupSharedBasePassParameters(GraphBuilder, View, BasePassParameters.Shared);
 	SetupSceneTextureUniformParameters(GraphBuilder, View.FeatureLevel, SceneTextureSetupMode, BasePassParameters.SceneTextures);
 
-	FRDGTextureRef BlackDummyTexture = GetRDG(GSystemTextures.BlackDummy);
+	const FRDGSystemTextures& SystemTextures = FRDGSystemTextures::Get(GraphBuilder);
 
 	// Material SSR
 	{
@@ -635,7 +635,7 @@ void SetupTranslucentBasePassUniformParameters(
 			BasePassParameters.HZBTexture = View.HZB;
 			BasePassParameters.HZBSampler = TStaticSamplerState<SF_Point, AM_Clamp, AM_Clamp, AM_Clamp>::GetRHI();
 
-			FRDGTextureRef PrevSceneColorTexture = BlackDummyTexture;
+			FRDGTextureRef PrevSceneColorTexture = SystemTextures.Black;
 
 			if (View.PrevViewInfo.CustomSSRInput.IsValid())
 			{
@@ -666,9 +666,9 @@ void SetupTranslucentBasePassUniformParameters(
 		}
 		else
 		{
-			BasePassParameters.HZBTexture = BlackDummyTexture;
+			BasePassParameters.HZBTexture = SystemTextures.Black;
 			BasePassParameters.HZBSampler = TStaticSamplerState<SF_Point, AM_Clamp, AM_Clamp, AM_Clamp>::GetRHI();
-			BasePassParameters.PrevSceneColor = BlackDummyTexture;
+			BasePassParameters.PrevSceneColor = SystemTextures.Black;
 			BasePassParameters.PrevSceneColorSampler = TStaticSamplerState<SF_Bilinear, AM_Clamp, AM_Clamp, AM_Clamp>::GetRHI();
 		}
 
@@ -699,7 +699,7 @@ void SetupTranslucentBasePassUniformParameters(
 		FIntPoint ViewportExtent = View.ViewRect.Size();
 
 		// Scene render targets might not exist yet; avoids NaNs.
-		FIntPoint EffectiveBufferSize = SceneRenderTargets.GetBufferSizeXY();
+		FIntPoint EffectiveBufferSize = FSceneRenderTargets::Get().GetBufferSizeXY();
 		EffectiveBufferSize.X = FMath::Max(EffectiveBufferSize.X, 1);
 		EffectiveBufferSize.Y = FMath::Max(EffectiveBufferSize.Y, 1);
 
@@ -723,31 +723,10 @@ void SetupTranslucentBasePassUniformParameters(
 	}
 
 	// Translucency Lighting Volume
-	{
-		if (TranslucentVolumeLightingTextures && TranslucentVolumeLightingTextures->IsValid())
-		{
-			BasePassParameters.TranslucencyLightingVolumeAmbientInner = TranslucentVolumeLightingTextures->GetAmbient(ViewIndex, TVC_Inner);
-			BasePassParameters.TranslucencyLightingVolumeAmbientOuter = TranslucentVolumeLightingTextures->GetAmbient(ViewIndex, TVC_Outer);
-			BasePassParameters.TranslucencyLightingVolumeDirectionalInner = TranslucentVolumeLightingTextures->GetDirectional(ViewIndex, TVC_Inner);
-			BasePassParameters.TranslucencyLightingVolumeDirectionalOuter = TranslucentVolumeLightingTextures->GetDirectional(ViewIndex, TVC_Outer);
-		}
-		else
-		{
-			BasePassParameters.TranslucencyLightingVolumeAmbientInner = BlackDummyTexture;
-			BasePassParameters.TranslucencyLightingVolumeAmbientOuter = BlackDummyTexture;
-			BasePassParameters.TranslucencyLightingVolumeDirectionalInner = BlackDummyTexture;
-			BasePassParameters.TranslucencyLightingVolumeDirectionalOuter = BlackDummyTexture;
-		}
-
-		BasePassParameters.TranslucencyLightingVolumeAmbientInnerSampler = TStaticSamplerState<SF_Bilinear, AM_Clamp, AM_Clamp, AM_Clamp>::GetRHI();
-		BasePassParameters.TranslucencyLightingVolumeAmbientOuterSampler = TStaticSamplerState<SF_Bilinear, AM_Clamp, AM_Clamp, AM_Clamp>::GetRHI();
-		BasePassParameters.TranslucencyLightingVolumeDirectionalInnerSampler = TStaticSamplerState<SF_Bilinear, AM_Clamp, AM_Clamp, AM_Clamp>::GetRHI();
-		BasePassParameters.TranslucencyLightingVolumeDirectionalOuterSampler = TStaticSamplerState<SF_Bilinear, AM_Clamp, AM_Clamp, AM_Clamp>::GetRHI();
-	}
-
+	BasePassParameters.TranslucencyLightingVolume = GetTranslucencyLightingVolumeParameters(GraphBuilder, TranslucencyLightingVolumeTextures, ViewIndex);
 	BasePassParameters.LumenParameters = GetLumenTranslucencyLightingParameters(View.LumenTranslucencyGIVolume);
 
-	BasePassParameters.SceneColorCopyTexture = BlackDummyTexture;
+	BasePassParameters.SceneColorCopyTexture = SystemTextures.Black;
 	BasePassParameters.SceneColorCopySampler = TStaticSamplerState<SF_Bilinear, AM_Clamp, AM_Clamp, AM_Clamp>::GetRHI();
 
 	if (SceneColorCopyTexture)
@@ -755,7 +734,7 @@ void SetupTranslucentBasePassUniformParameters(
 		BasePassParameters.SceneColorCopyTexture = SceneColorCopyTexture;
 	}
 
-	BasePassParameters.EyeAdaptationTexture = BlackDummyTexture;
+	BasePassParameters.EyeAdaptationTexture = SystemTextures.Black;
 
 	// Setup by passes that support it
 	if (View.HasValidEyeAdaptationTexture())
@@ -765,20 +744,8 @@ void SetupTranslucentBasePassUniformParameters(
 
 	BasePassParameters.PreIntegratedGFTexture = GSystemTextures.PreintegratedGF->GetRenderTargetItem().ShaderResourceTexture;
 	BasePassParameters.PreIntegratedGFSampler = TStaticSamplerState<SF_Bilinear, AM_Clamp, AM_Clamp, AM_Clamp>::GetRHI();
-}
 
-TRDGUniformBufferRef<FTranslucentBasePassUniformParameters> CreateTranslucentBasePassUniformBuffer(
-	FRDGBuilder& GraphBuilder,
-	const FViewInfo& View,
-	const FTranslucentVolumeLightingTextures* TranslucentVolumeLightingTextures,
-	FRDGTextureRef SceneColorCopyTexture,
-	ESceneTextureSetupMode SceneTextureSetupMode,
-	const int32 ViewIndex)
-{
-	FSceneRenderTargets& SceneRenderTargets = FSceneRenderTargets::Get();
-	FTranslucentBasePassUniformParameters* BasePassParameters = GraphBuilder.AllocParameters<FTranslucentBasePassUniformParameters>();
-	SetupTranslucentBasePassUniformParameters(GraphBuilder, SceneRenderTargets, View, TranslucentVolumeLightingTextures, SceneColorCopyTexture, SceneTextureSetupMode, ViewIndex, *BasePassParameters);
-	return GraphBuilder.CreateUniformBuffer(BasePassParameters);
+	return GraphBuilder.CreateUniformBuffer(&BasePassParameters);
 }
 
 static void UpdateSeparateTranslucencyViewState(FScene* Scene, const FViewInfo& View, const FScreenPassTextureViewport& SeparateTranslucencyViewport, FMeshPassProcessorRenderState& DrawRenderState)
@@ -961,7 +928,7 @@ static void RenderTranslucencyViewInner(
 void FDeferredShadingSceneRenderer::RenderTranslucencyInner(
 	FRDGBuilder& GraphBuilder,
 	const FMinimalSceneTextures& SceneTextures,
-	const FTranslucentVolumeLightingTextures& TranslucentVolumeLightingTextures,
+	const FTranslucencyLightingVolumeTextures& TranslucentLightingVolumeTextures,
 	FSeparateTranslucencyTextures* OutSeparateTranslucencyTextures,
 	ETranslucencyView ViewsToRender,
 	FRDGTextureRef SceneColorCopyTexture,
@@ -1055,7 +1022,7 @@ void FDeferredShadingSceneRenderer::RenderTranslucencyInner(
 				SeparateTranslucencyColorTexture,
 				SeparateTranslucencyColorLoadAction,
 				SeparateTranslucencyDepthTexture.Target,
-				CreateTranslucentBasePassUniformBuffer(GraphBuilder, View, &TranslucentVolumeLightingTextures, SceneColorCopyTexture, SceneTextureSetupMode, ViewIndex),
+				CreateTranslucentBasePassUniformBuffer(GraphBuilder, View, ViewIndex, TranslucentLightingVolumeTextures, SceneColorCopyTexture, SceneTextureSetupMode),
 				TranslucencyPass,
 				!bCompositeBackToSceneColor,
 				bRenderInParallel);
@@ -1107,7 +1074,7 @@ void FDeferredShadingSceneRenderer::RenderTranslucencyInner(
 				SceneTextures.Color,
 				SceneColorLoadAction,
 				SceneTextures.Depth.Target,
-				CreateTranslucentBasePassUniformBuffer(GraphBuilder, View, &TranslucentVolumeLightingTextures, SceneColorCopyTexture, SceneTextureSetupMode, ViewIndex),
+				CreateTranslucentBasePassUniformBuffer(GraphBuilder, View, ViewIndex, TranslucentLightingVolumeTextures, SceneColorCopyTexture, SceneTextureSetupMode),
 				TranslucencyPass,
 				bResolveColorTexture,
 				bRenderInParallel);
@@ -1120,7 +1087,7 @@ void FDeferredShadingSceneRenderer::RenderTranslucencyInner(
 void FDeferredShadingSceneRenderer::RenderTranslucency(
 	FRDGBuilder& GraphBuilder,
 	const FMinimalSceneTextures& SceneTextures,
-	const FTranslucentVolumeLightingTextures& TranslucentVolumeLightingTextures,
+	const FTranslucencyLightingVolumeTextures& TranslucentLightingVolumeTextures,
 	FSeparateTranslucencyTextures* OutSeparateTranslucencyTextures,
 	ETranslucencyView ViewsToRender)
 {
@@ -1140,12 +1107,12 @@ void FDeferredShadingSceneRenderer::RenderTranslucency(
 
 	if (ViewFamily.AllowTranslucencyAfterDOF())
 	{
-		RenderTranslucencyInner(GraphBuilder, SceneTextures, TranslucentVolumeLightingTextures, OutSeparateTranslucencyTextures, ViewsToRender, SceneColorCopyTexture, ETranslucencyPass::TPT_StandardTranslucency);
-		RenderTranslucencyInner(GraphBuilder, SceneTextures, TranslucentVolumeLightingTextures, OutSeparateTranslucencyTextures, ViewsToRender, SceneColorCopyTexture, ETranslucencyPass::TPT_TranslucencyAfterDOF);
-		RenderTranslucencyInner(GraphBuilder, SceneTextures, TranslucentVolumeLightingTextures, OutSeparateTranslucencyTextures, ViewsToRender, SceneColorCopyTexture, ETranslucencyPass::TPT_TranslucencyAfterDOFModulate);
+		RenderTranslucencyInner(GraphBuilder, SceneTextures, TranslucentLightingVolumeTextures, OutSeparateTranslucencyTextures, ViewsToRender, SceneColorCopyTexture, ETranslucencyPass::TPT_StandardTranslucency);
+		RenderTranslucencyInner(GraphBuilder, SceneTextures, TranslucentLightingVolumeTextures, OutSeparateTranslucencyTextures, ViewsToRender, SceneColorCopyTexture, ETranslucencyPass::TPT_TranslucencyAfterDOF);
+		RenderTranslucencyInner(GraphBuilder, SceneTextures, TranslucentLightingVolumeTextures, OutSeparateTranslucencyTextures, ViewsToRender, SceneColorCopyTexture, ETranslucencyPass::TPT_TranslucencyAfterDOFModulate);
 	}
 	else // Otherwise render translucent primitives in a single bucket.
 	{
-		RenderTranslucencyInner(GraphBuilder, SceneTextures, TranslucentVolumeLightingTextures, OutSeparateTranslucencyTextures, ViewsToRender, SceneColorCopyTexture, ETranslucencyPass::TPT_AllTranslucency);
+		RenderTranslucencyInner(GraphBuilder, SceneTextures, TranslucentLightingVolumeTextures, OutSeparateTranslucencyTextures, ViewsToRender, SceneColorCopyTexture, ETranslucencyPass::TPT_AllTranslucency);
 	}
 }
diff --git a/Engine/Source/Runtime/Renderer/Private/TranslucentRendering.h b/Engine/Source/Runtime/Renderer/Private/TranslucentRendering.h
--- a/Engine/Source/Runtime/Renderer/Private/TranslucentRendering.h
+++ b/Engine/Source/Runtime/Renderer/Private/TranslucentRendering.h
@@ -62,38 +62,6 @@ private:
 	FRDGTextureMSAA ColorModulateTexture;
 };
 
-struct FTranslucentVolumeLightingTextures
-{
-	bool IsValid() const
-	{
-		return VolumeDim != 0;
-	}
-
-	FRDGTextureRef GetAmbient(int32 ViewIndex, int32 CascadeIndex) const
-	{
-		return Ambient[(ViewIndex * TVC_MAX) + CascadeIndex];
-	}
-
-	void SetAmbient(int32 ViewIndex, int32 CascadeIndex, FRDGTextureRef Texture)
-	{
-		Ambient[(ViewIndex * TVC_MAX) + CascadeIndex] = Texture;
-	}
-
-	FRDGTextureRef GetDirectional(int32 ViewIndex, int32 CascadeIndex) const
-	{
-		return Directional[(ViewIndex * TVC_MAX) + CascadeIndex];
-	}
-
-	void SetDirectional(int32 ViewIndex, int32 CascadeIndex, FRDGTextureRef Texture)
-	{
-		Directional[(ViewIndex * TVC_MAX) + CascadeIndex] = Texture;
-	}
-
-	TArray<FRDGTextureRef, TInlineAllocator<TVC_MAX>> Ambient;
-	TArray<FRDGTextureRef, TInlineAllocator<TVC_MAX>> Directional;
-	int32 VolumeDim = 0;
-};
-
 DECLARE_GPU_DRAWCALL_STAT_EXTERN(Translucency);
 
 /** Creates separate translucency textures. */
