Transition untracked shadow texture to SRVMask after shadow rendering because these textures are not tracked via RDG yet - temporary fix until it's moved over to RDG

diff --git a/Engine/Source/Runtime/Renderer/Private/ShadowDepthRendering.cpp b/Engine/Source/Runtime/Renderer/Private/ShadowDepthRendering.cpp
--- a/Engine/Source/Runtime/Renderer/Private/ShadowDepthRendering.cpp
+++ b/Engine/Source/Runtime/Renderer/Private/ShadowDepthRendering.cpp
@@ -1561,7 +1561,7 @@ void FSceneRenderer::RenderShadowDepthMapAtlases(FRDGBuilder& GraphBuilder)
 
 	for (int32 AtlasIndex = 0; AtlasIndex < SortedShadowsForShadowDepthPass.ShadowMapAtlases.Num(); AtlasIndex++)
 	{
-		const FSortedShadowMapAtlas& ShadowMapAtlas = SortedShadowsForShadowDepthPass.ShadowMapAtlases[AtlasIndex];
+		FSortedShadowMapAtlas& ShadowMapAtlas = SortedShadowsForShadowDepthPass.ShadowMapAtlases[AtlasIndex];
 		FRDGTextureRef AtlasDepthTexture = GraphBuilder.RegisterExternalTexture(ShadowMapAtlas.RenderTargets.DepthTarget);
 		const FIntPoint AtlasSize = AtlasDepthTexture->Desc.Extent;
 
@@ -1651,6 +1651,9 @@ void FSceneRenderer::RenderShadowDepthMapAtlases(FRDGBuilder& GraphBuilder)
 		{
 			RenderShadowDepthAtlasNanite(GraphBuilder, *Scene, ShadowMapAtlas, AtlasIndex, false);
 		}
+
+		// Make readable because AtlasDepthTexture is not tracked via RDG yet
+		ConvertToUntrackedExternalTexture(GraphBuilder, AtlasDepthTexture, ShadowMapAtlas.RenderTargets.DepthTarget, ERHIAccess::SRVMask);
 	}
 
 	// Copy/resample shadow maps into "complete" shadow maps and add Nanite geometry
@@ -1944,7 +1947,7 @@ void FSceneRenderer::RenderShadowDepthMaps(FRDGBuilder& GraphBuilder)
 
 	for (int32 CubemapIndex = 0; CubemapIndex < SortedShadowsForShadowDepthPass.ShadowMapCubemaps.Num(); CubemapIndex++)
 	{
-		const FSortedShadowMapAtlas& ShadowMap = SortedShadowsForShadowDepthPass.ShadowMapCubemaps[CubemapIndex];
+		FSortedShadowMapAtlas& ShadowMap = SortedShadowsForShadowDepthPass.ShadowMapCubemaps[CubemapIndex];
 		FRDGTextureRef ShadowDepthTexture = GraphBuilder.RegisterExternalTexture(ShadowMap.RenderTargets.DepthTarget);
 		const FIntPoint TargetSize = ShadowDepthTexture->Desc.Extent;
 
@@ -2087,6 +2090,9 @@ void FSceneRenderer::RenderShadowDepthMaps(FRDGBuilder& GraphBuilder)
 				}
 			}
 		}
+
+		// Make readable because ShadowDepthTexture is not tracked via RDG yet
+		ConvertToUntrackedExternalTexture(GraphBuilder, ShadowDepthTexture, ShadowMap.RenderTargets.DepthTarget, ERHIAccess::SRVMask);
 	}
 
 	if (SortedShadowsForShadowDepthPass.PreshadowCache.Shadows.Num() > 0)
