Fix for recursive RDG builder in ray tracing code.

#rnx

diff --git a/Engine/Source/Runtime/Renderer/Private/RayTracing/RaytracingSkylight.cpp b/Engine/Source/Runtime/Renderer/Private/RayTracing/RaytracingSkylight.cpp
--- a/Engine/Source/Runtime/Renderer/Private/RayTracing/RaytracingSkylight.cpp
+++ b/Engine/Source/Runtime/Renderer/Private/RayTracing/RaytracingSkylight.cpp
@@ -283,7 +283,7 @@ void SetupSkyLightVisibilityRaysParameters(
 	// Get the Scene View State
 	FSceneViewState* SceneViewState = (FSceneViewState*)View.State;
 
-	TRefCountPtr<FRDGPooledBuffer> PooledSkyLightVisibilityRaysBuffer;
+	FRDGBufferRef SkyLightVisibilityRaysBuffer = nullptr;
 	FIntVector SkyLightVisibilityRaysDimensions;
 
 	// Check if the Sky Light Visibility Ray Data should be set based on if decoupled sample generation is being used
@@ -293,7 +293,7 @@ void SetupSkyLightVisibilityRaysParameters(
 		(CVarRayTracingSkyLightDecoupleSampleGeneration.GetValueOnRenderThread() == 1))
 	{
 		// Set the Sky Light Visibility Ray pooled buffer to the stored pooled buffer
-		PooledSkyLightVisibilityRaysBuffer = SceneViewState->SkyLightVisibilityRaysBuffer;
+		SkyLightVisibilityRaysBuffer = GraphBuilder.RegisterExternalBuffer(SceneViewState->SkyLightVisibilityRaysBuffer);
 
 		// Set the Sky Light Visibility Ray Dimensions from the stored dimensions
 		SkyLightVisibilityRaysDimensions = SceneViewState->SkyLightVisibilityRaysDimensions;
@@ -301,25 +301,17 @@ void SetupSkyLightVisibilityRaysParameters(
 	else
 	{
 		// Create a dummy Sky Light Visibility Ray buffer in a dummy RDG
-		FRDGBuilder DummyGraphBuilder{ GraphBuilder.RHICmdList };
 		FRDGBufferDesc DummyBufferDesc = FRDGBufferDesc::CreateStructuredDesc(sizeof(FSkyLightVisibilityRays), 1);
-		FRDGBufferRef DummyRDGBuffer = DummyGraphBuilder.CreateBuffer(DummyBufferDesc, TEXT("DummySkyLightVisibilityRays"));
-		FRDGBufferUAVRef DummyRDGBufferUAV = DummyGraphBuilder.CreateUAV(DummyRDGBuffer, EPixelFormat::PF_R32_UINT);
+		SkyLightVisibilityRaysBuffer = GraphBuilder.CreateBuffer(DummyBufferDesc, TEXT("DummySkyLightVisibilityRays"));
+		FRDGBufferUAVRef DummyRDGBufferUAV = GraphBuilder.CreateUAV(SkyLightVisibilityRaysBuffer, EPixelFormat::PF_R32_UINT);
 
 		// Clear the dummy Sky Light Visibility Ray buffer
-		AddClearUAVPass(DummyGraphBuilder, DummyRDGBufferUAV, 0);
-
-		// Set the Sky Light Visibility Ray pooled buffer to the extracted pooled dummy buffer
-		DummyGraphBuilder.QueueBufferExtraction(DummyRDGBuffer, &PooledSkyLightVisibilityRaysBuffer);
-		DummyGraphBuilder.Execute();
+		AddClearUAVPass(GraphBuilder, DummyRDGBufferUAV, 0);
 
 		// Set the Sky Light Visibility Ray Dimensions to a dummy value
 		SkyLightVisibilityRaysDimensions = FIntVector(1);
 	}
 
-	// Set the Sky Light Visibility Ray Buffer to the pooled RDG buffer
-	FRDGBufferRef SkyLightVisibilityRaysBuffer = GraphBuilder.RegisterExternalBuffer(PooledSkyLightVisibilityRaysBuffer);
-
 	// Set Sky Light Visibility Ray Data information
 	OutSkyLightVisibilityRaysData->SkyLightVisibilityRays = GraphBuilder.CreateSRV(SkyLightVisibilityRaysBuffer, EPixelFormat::PF_R32_UINT);
 	OutSkyLightVisibilityRaysData->SkyLightVisibilityRaysDimensions = SkyLightVisibilityRaysDimensions;
