Fix static analysis error

#rnx

diff --git a/Engine/Source/Runtime/Renderer/Private/GPUScene.cpp b/Engine/Source/Runtime/Renderer/Private/GPUScene.cpp
--- a/Engine/Source/Runtime/Renderer/Private/GPUScene.cpp
+++ b/Engine/Source/Runtime/Renderer/Private/GPUScene.cpp
@@ -1032,7 +1032,7 @@ TRange<int32> FGPUScene::CommitPrimitiveCollector(FGPUScenePrimitiveCollector& P
 	ensure(CurrentDynamicContext != nullptr);
 
 	// Make sure we are not trying to commit data that lives in a different context.
-	ensure(CurrentDynamicContext->DymamicPrimitiveUploadData.Find(PrimitiveCollector.UploadData) != INDEX_NONE);
+	ensure(CurrentDynamicContext == nullptr || CurrentDynamicContext->DymamicPrimitiveUploadData.Find(PrimitiveCollector.UploadData) != INDEX_NONE);
 
 	int32 StartOffset = DynamicPrimitivesOffset;
 	DynamicPrimitivesOffset += PrimitiveCollector.UploadData->PrimitiveShaderData.Num();
