Fix GPUScene static analysis errors

diff --git a/Engine/Source/Runtime/Renderer/Private/SceneRendering.cpp b/Engine/Source/Runtime/Renderer/Private/SceneRendering.cpp
--- a/Engine/Source/Runtime/Renderer/Private/SceneRendering.cpp
+++ b/Engine/Source/Runtime/Renderer/Private/SceneRendering.cpp
@@ -2136,13 +2136,19 @@ void FSortedShadowMaps::Release()
 	PreshadowCache.RenderTargets.Release();
 }
 
+template <typename T>
+inline T* CheckPointer(T* Ptr)
+{
+	check(Ptr != nullptr);
+	return Ptr;
+}
 /*-----------------------------------------------------------------------------
 	FSceneRenderer
 -----------------------------------------------------------------------------*/
 
 FSceneRenderer::FSceneRenderer(const FSceneViewFamily* InViewFamily,FHitProxyConsumer* HitProxyConsumer)
-:	Scene(InViewFamily->Scene ? InViewFamily->Scene->GetRenderScene() : NULL)
-,	ViewFamily(*InViewFamily)
+:	ViewFamily(*CheckPointer(InViewFamily))
+,	Scene(CheckPointer(InViewFamily->Scene)->GetRenderScene())
 ,	MeshCollector(InViewFamily->GetFeatureLevel())
 ,	RayTracingCollector(InViewFamily->GetFeatureLevel())
 ,	bHasRequestedToggleFreeze(false)
@@ -2150,13 +2156,13 @@ FSceneRenderer::FSceneRenderer(const FSceneViewFamily* InViewFamily,FHitProxyCon
 ,	InstancedStereoWidth(0)
 ,	RootMark(nullptr)
 ,	FamilySize(0, 0)
-,	GPUSceneDynamicContext(Scene->GPUScene)
+,	GPUSceneDynamicContext(CheckPointer(Scene)->GPUScene)
 ,	bShadowDepthRenderCompleted(false)
 {
 	check(Scene != NULL);
 
 	check(IsInGameThread());
-	ViewFamily.FrameNumber = Scene ? Scene->GetFrameNumber() : GFrameNumber;
+	ViewFamily.FrameNumber = Scene->GetFrameNumber();
 
 	// Copy the individual views.
 	bool bAnyViewIsLocked = false;
