Fix SceneRendering initialization order issue.

diff --git a/Engine/Source/Runtime/Renderer/Private/SceneRendering.cpp b/Engine/Source/Runtime/Renderer/Private/SceneRendering.cpp
--- a/Engine/Source/Runtime/Renderer/Private/SceneRendering.cpp
+++ b/Engine/Source/Runtime/Renderer/Private/SceneRendering.cpp
@@ -2147,8 +2147,8 @@ inline T* CheckPointer(T* Ptr)
 -----------------------------------------------------------------------------*/
 
 FSceneRenderer::FSceneRenderer(const FSceneViewFamily* InViewFamily,FHitProxyConsumer* HitProxyConsumer)
-:	ViewFamily(*CheckPointer(InViewFamily))
-,	Scene(CheckPointer(InViewFamily->Scene)->GetRenderScene())
+:	Scene(CheckPointer(InViewFamily->Scene)->GetRenderScene())
+,	ViewFamily(*CheckPointer(InViewFamily))
 ,	MeshCollector(InViewFamily->GetFeatureLevel())
 ,	RayTracingCollector(InViewFamily->GetFeatureLevel())
 ,	bHasRequestedToggleFreeze(false)
