Moving RT acceleration structure build to occur before Lumen scene lighting is computed.

Moving RT SkyLight raygen pipeline to occur at the composition pipeline.

diff --git a/Engine/Source/Runtime/Renderer/Private/DeferredShadingRenderer.cpp b/Engine/Source/Runtime/Renderer/Private/DeferredShadingRenderer.cpp
--- a/Engine/Source/Runtime/Renderer/Private/DeferredShadingRenderer.cpp
+++ b/Engine/Source/Runtime/Renderer/Private/DeferredShadingRenderer.cpp
@@ -2244,6 +2244,14 @@ void FDeferredShadingSceneRenderer::Render(FRHICommandListImmediate& RHICmdList)
 		AddResolveSceneColorPass(GraphBuilder, Views, SceneTextures.Color);
 	}
 
+#if RHI_RAYTRACING
+	const bool bRayTracingEnabled = IsRayTracingEnabled();
+	if (bRayTracingEnabled)
+	{
+		WaitForRayTracingScene(GraphBuilder);
+	}
+#endif // RHI_RAYTRACING
+
 	// Shadow and fog after base pass
 	if (!bOcclusionBeforeBasePass)
 	{
@@ -2299,12 +2307,6 @@ void FDeferredShadingSceneRenderer::Render(FRHICommandListImmediate& RHICmdList)
 		SceneTextures.UniformBuffer = CreateSceneTextureUniformBuffer(GraphBuilder, FeatureLevel, SceneTextures.SetupMode);
 	}
 
-#if RHI_RAYTRACING
-	const bool bRayTracingEnabled = IsRayTracingEnabled();
-	FRDGTextureRef SkyLightTexture = nullptr;
-	FRDGTextureRef SkyLightHitDistanceTexture = nullptr;
-#endif
-
 	// TODO: Keeping the velocities here for testing, but if that works, this pass will be remove and DDM_AllOpaqueNoVelocity will be the only option with
 	// DBuffer decals enabled.
 
@@ -2327,18 +2329,6 @@ void FDeferredShadingSceneRenderer::Render(FRHICommandListImmediate& RHICmdList)
 		HairDatas = &HairDatasStorage;
 	}
 
-#if RHI_RAYTRACING
-	if (bRayTracingEnabled)
-	{
-		WaitForRayTracingScene(GraphBuilder);
-
-		if (bCanOverlayRayTracingOutput && !IsForwardShadingEnabled(ShaderPlatform))
-		{
-			RenderRayTracingSkyLight(GraphBuilder, SceneTextures.Color.Target, SkyLightTexture, SkyLightHitDistanceTexture, HairDatas);
-		}
-	}
-#endif // RHI_RAYTRACING
-
 	// Copy lighting channels out of stencil before deferred decals which overwrite those values
 	FRDGTextureRef LightingChannelsTexture = CopyStencilToLightingChannelTexture(GraphBuilder, Views, SceneTextures.Stencil);
 
@@ -2449,8 +2439,11 @@ void FDeferredShadingSceneRenderer::Render(FRHICommandListImmediate& RHICmdList)
 		}
 
 	#if RHI_RAYTRACING
-		if (SkyLightTexture)
+		if (ShouldRenderRayTracingSkyLight(Scene->SkyLight))
 		{
+			FRDGTextureRef SkyLightTexture = nullptr;
+			FRDGTextureRef SkyLightHitDistanceTexture = nullptr;
+			RenderRayTracingSkyLight(GraphBuilder, SceneTextures.Color.Target, SkyLightTexture, SkyLightHitDistanceTexture, HairDatas);
 			CompositeRayTracingSkyLight(GraphBuilder, SceneTextures, SkyLightTexture, SkyLightHitDistanceTexture);
 		}
 	#endif
