Moved reflection scratch cubemaps out of the public scene context and into their own private global variable. Removed some dead code. Moved some GPU mask logic out of the deferred shading render function.

diff --git a/Engine/Source/Runtime/Renderer/Private/DeferredShadingRenderer.cpp b/Engine/Source/Runtime/Renderer/Private/DeferredShadingRenderer.cpp
--- a/Engine/Source/Runtime/Renderer/Private/DeferredShadingRenderer.cpp
+++ b/Engine/Source/Runtime/Renderer/Private/DeferredShadingRenderer.cpp
@@ -1503,17 +1503,7 @@ void FDeferredShadingSceneRenderer::Render(FRHICommandListImmediate& RHICmdList)
 	SCOPED_NAMED_EVENT(FDeferredShadingSceneRenderer_Render, FColor::Emerald);
 
 #if WITH_MGPU
-	FRHIGPUMask RenderTargetGPUMask = (GNumExplicitGPUsForRendering > 1 && ViewFamily.RenderTarget) ? ViewFamily.RenderTarget->GetGPUMask(RHICmdList) : FRHIGPUMask::GPU0();
-	
-	{
-		static auto CVar = IConsoleManager::Get().FindConsoleVariable(TEXT("r.PathTracing.GPUCount"));
-		if (CVar && CVar->GetInt() > 1)
-		{
-			RenderTargetGPUMask = FRHIGPUMask::All(); // Broadcast to all GPUs 
-		}
-	}
-
-	ComputeViewGPUMasks(RenderTargetGPUMask);
+	const FRHIGPUMask RenderTargetGPUMask = ComputeGPUMasks(RHICmdList);
 #endif // WITH_MGPU
 
 	// By default, limit our GPU usage to only GPUs specified in the view masks.
@@ -1555,6 +1545,8 @@ void FDeferredShadingSceneRenderer::Render(FRHICommandListImmediate& RHICmdList)
 
 	const FRDGSystemTextures& SystemTextures = FRDGSystemTextures::Create(GraphBuilder);
 
+	const bool bAllowStaticLighting = IsStaticLightingAllowed();
+
 	const bool bUseVirtualTexturing = UseVirtualTexturing(FeatureLevel);
 	if (bUseVirtualTexturing)
 	{
@@ -2157,7 +2149,7 @@ void FDeferredShadingSceneRenderer::Render(FRHICommandListImmediate& RHICmdList)
 		CompositionLighting::ProcessBeforeBasePass(GraphBuilder, Views, Scene->UniformBuffers, SceneTextures, DBufferTextures);
 	}
 	
-	if (IsForwardShadingEnabled(ShaderPlatform) && SceneContext.IsStaticLightingAllowed())
+	if (IsForwardShadingEnabled(ShaderPlatform) && bAllowStaticLighting)
 	{
 		RenderIndirectCapsuleShadows(GraphBuilder, SceneTextures);
 	}
@@ -2398,7 +2390,7 @@ void FDeferredShadingSceneRenderer::Render(FRHICommandListImmediate& RHICmdList)
 		RenderDiffuseIndirectAndAmbientOcclusion(GraphBuilder, SceneTextures, HairDatas, /* bIsVisualizePass = */ false);
 
 		// These modulate the scenecolor output from the basepass, which is assumed to be indirect lighting
-		if (SceneContext.IsStaticLightingAllowed())
+		if (bAllowStaticLighting)
 		{
 			RenderIndirectCapsuleShadows(GraphBuilder, SceneTextures);
 		}
diff --git a/Engine/Source/Runtime/Renderer/Private/MobileReflectionEnvironmentCapture.cpp b/Engine/Source/Runtime/Renderer/Private/MobileReflectionEnvironmentCapture.cpp
--- a/Engine/Source/Runtime/Renderer/Private/MobileReflectionEnvironmentCapture.cpp
+++ b/Engine/Source/Runtime/Renderer/Private/MobileReflectionEnvironmentCapture.cpp
@@ -25,6 +25,8 @@ static TAutoConsoleVariable<int32> CVarMobileUseHighQualitySkyCaptureFiltering(
 	TEXT("0: use simple bilinear filtering when generating mobile sky captures."),
 	ECVF_RenderThreadSafe);
 
+extern TGlobalResource<FReflectionScratchCubemaps> GReflectionScratchCubemaps;
+
 class FMobileDownsamplePS : public FGlobalShader
 {
 	DECLARE_SHADER_TYPE(FMobileDownsamplePS, Global);
@@ -70,7 +72,7 @@ IMPLEMENT_SHADER_TYPE(, FMobileDownsamplePS, TEXT("/Engine/Private/ReflectionEnv
 namespace MobileReflectionEnvironmentCapture
 {
 	/** Encapsulates render target picking logic for cubemap mip generation. */
-	FSceneRenderTargetItem& GetEffectiveRenderTarget(FSceneRenderTargets& SceneContext, bool bDownsamplePass, int32 TargetMipIndex)
+	FSceneRenderTargetItem& GetEffectiveRenderTarget(bool bDownsamplePass, int32 TargetMipIndex)
 	{
 		int32 ScratchTextureIndex = TargetMipIndex % 2;
 
@@ -79,11 +81,11 @@ namespace MobileReflectionEnvironmentCapture
 			ScratchTextureIndex = 1 - ScratchTextureIndex;
 		}
 
-		return SceneContext.ReflectionColorScratchCubemap[ScratchTextureIndex]->GetRenderTargetItem();
+		return GReflectionScratchCubemaps.Color[ScratchTextureIndex]->GetRenderTargetItem();
 	}
 
 	/** Encapsulates source texture picking logic for cubemap mip generation. */
-	FSceneRenderTargetItem& GetEffectiveSourceTexture(FSceneRenderTargets& SceneContext, bool bDownsamplePass, int32 TargetMipIndex)
+	FSceneRenderTargetItem& GetEffectiveSourceTexture(bool bDownsamplePass, int32 TargetMipIndex)
 	{
 		int32 ScratchTextureIndex = TargetMipIndex % 2;
 
@@ -92,7 +94,7 @@ namespace MobileReflectionEnvironmentCapture
 			ScratchTextureIndex = 1 - ScratchTextureIndex;
 		}
 
-		return SceneContext.ReflectionColorScratchCubemap[ScratchTextureIndex]->GetRenderTargetItem();
+		return GReflectionScratchCubemaps.Color[ScratchTextureIndex]->GetRenderTargetItem();
 	}
 
 	void ComputeAverageBrightness(FRHICommandListImmediate& RHICmdList, ERHIFeatureLevel::Type FeatureLevel, int32 CubmapSize, float& OutAverageBrightness)
@@ -107,8 +109,6 @@ namespace MobileReflectionEnvironmentCapture
 
 		auto ShaderMap = GetGlobalShaderMap(FeatureLevel);
 
-		FSceneRenderTargets& SceneContext = FSceneRenderTargets::Get();
-
 		{
 			SCOPED_DRAW_EVENT(RHICmdList, DownsampleCubeMips);
 
@@ -118,8 +118,8 @@ namespace MobileReflectionEnvironmentCapture
 				const int32 SourceMipIndex = FMath::Max(MipIndex - 1, 0);
 				const int32 MipSize = 1 << (NumMips - MipIndex - 1);
 
-				FSceneRenderTargetItem& EffectiveRT = GetEffectiveRenderTarget(SceneContext, true, MipIndex);
-				FSceneRenderTargetItem& EffectiveSource = GetEffectiveSourceTexture(SceneContext, true, MipIndex);
+				FSceneRenderTargetItem& EffectiveRT = GetEffectiveRenderTarget(true, MipIndex);
+				FSceneRenderTargetItem& EffectiveSource = GetEffectiveSourceTexture(true, MipIndex);
 				check(EffectiveRT.TargetableTexture != EffectiveSource.ShaderResourceTexture);
 
 				for (int32 CubeFace = 0; CubeFace < CubeFace_MAX; CubeFace++)
@@ -166,7 +166,7 @@ namespace MobileReflectionEnvironmentCapture
 			}
 		}
 
-		OutAverageBrightness = ComputeSingleAverageBrightnessFromCubemap(RHICmdList, FeatureLevel, CubmapSize, GetEffectiveRenderTarget(FSceneRenderTargets::Get(), true, NumMips - 1));
+		OutAverageBrightness = ComputeSingleAverageBrightnessFromCubemap(RHICmdList, FeatureLevel, CubmapSize, GetEffectiveRenderTarget(true, NumMips - 1));
 	}
 
 	void CopyToSkyTexture(FRHICommandList& RHICmdList, FScene* Scene, FTexture* ProcessedTexture)
@@ -177,7 +177,6 @@ namespace MobileReflectionEnvironmentCapture
 			const bool bUseHQFiltering = CVarMobileUseHighQualitySkyCaptureFiltering.GetValueOnRenderThread() == 1;
 			const int32 EffectiveTopMipSize = ProcessedTexture->GetSizeX();
 			const int32 NumMips = FMath::CeilLogTwo(EffectiveTopMipSize) + 1;
-			FSceneRenderTargets& SceneContext = FSceneRenderTargets::Get();
 
 			RHICmdList.Transition(FRHITransitionInfo(ProcessedTexture->TextureRHI, ERHIAccess::Unknown, ERHIAccess::CopyDest));
 
@@ -190,7 +189,7 @@ namespace MobileReflectionEnvironmentCapture
 			{
 				// For simple mobile bilin filtering the source for this copy is the dest from the filtering pass.
 				// In the HQ case, the full image is contained in GetEffectiveRenderTarget(.., false,0).
-				FSceneRenderTargetItem& EffectiveSource = GetEffectiveRenderTarget(SceneContext, false, bUseHQFiltering ? 0 : MipIndex);
+				FSceneRenderTargetItem& EffectiveSource = GetEffectiveRenderTarget(false, bUseHQFiltering ? 0 : MipIndex);
 				RHICmdList.Transition(FRHITransitionInfo(EffectiveSource.ShaderResourceTexture, ERHIAccess::Unknown, ERHIAccess::CopySrc));
 				RHICmdList.CopyTexture(EffectiveSource.ShaderResourceTexture, ProcessedTexture->TextureRHI, CopyInfo);
 
@@ -213,7 +212,7 @@ namespace MobileReflectionEnvironmentCapture
 		const int32 NumMips = FMath::CeilLogTwo(EffectiveTopMipSize) + 1;
 		const bool bUseHQFiltering = CVarMobileUseHighQualitySkyCaptureFiltering.GetValueOnRenderThread() == 1;
 		{
-			FSceneRenderTargetItem& EffectiveColorRT = FSceneRenderTargets::Get().ReflectionColorScratchCubemap[0]->GetRenderTargetItem();
+			FSceneRenderTargetItem& EffectiveColorRT = GReflectionScratchCubemaps.Color[0]->GetRenderTargetItem();
 			// Premultiply alpha in-place using alpha blending
 			for (uint32 CubeFace = 0; CubeFace < CubeFace_MAX; CubeFace++)
 			{
@@ -266,7 +265,6 @@ namespace MobileReflectionEnvironmentCapture
 		FSceneRenderTargetItem* DiffuseConvolutionSource = NULL;
 
 		auto ShaderMap = GetGlobalShaderMap(FeatureLevel);
-		FSceneRenderTargets& SceneContext = FSceneRenderTargets::Get();
 
 		{
 			SCOPED_DRAW_EVENT(RHICmdList, DownsampleCubeMips);
@@ -277,8 +275,8 @@ namespace MobileReflectionEnvironmentCapture
 				const int32 SourceMipIndex = FMath::Max(MipIndex - 1, 0);
 				const int32 MipSize = 1 << (NumMips - MipIndex - 1);
 
-				FSceneRenderTargetItem& EffectiveRT = GetEffectiveRenderTarget(SceneContext, true, MipIndex);
-				FSceneRenderTargetItem& EffectiveSource = GetEffectiveSourceTexture(SceneContext, true, MipIndex);
+				FSceneRenderTargetItem& EffectiveRT = GetEffectiveRenderTarget(true, MipIndex);
+				FSceneRenderTargetItem& EffectiveSource = GetEffectiveSourceTexture(true, MipIndex);
 				check(EffectiveRT.TargetableTexture != EffectiveSource.ShaderResourceTexture);
 
 				for (int32 CubeFace = 0; CubeFace < CubeFace_MAX; CubeFace++)
@@ -348,8 +346,8 @@ namespace MobileReflectionEnvironmentCapture
 			SCOPED_DRAW_EVENT(RHICmdList, PrepareSourceCubemapMipsForHQFiltering);
 			for (int32 MipIndex = 1; MipIndex < NumMips; MipIndex += 2)
 			{
-				FSceneRenderTargetItem& SourceTarget = GetEffectiveRenderTarget(SceneContext, true, MipIndex);
-				FSceneRenderTargetItem& DestTarget = GetEffectiveSourceTexture(SceneContext, true, MipIndex);
+				FSceneRenderTargetItem& SourceTarget = GetEffectiveRenderTarget(true, MipIndex);
+				FSceneRenderTargetItem& DestTarget = GetEffectiveSourceTexture(true, MipIndex);
 				check(DestTarget.TargetableTexture != SourceTarget.ShaderResourceTexture);
 
 				// Transition the textures once, so CopyToResolveTarget doesn't ping-pong uselessly between the copy and SRV states.
@@ -387,8 +385,8 @@ namespace MobileReflectionEnvironmentCapture
 			for (int32 MipIndex = 0; MipIndex < NumMips; MipIndex++)
 			{
 				SCOPED_DRAW_EVENT(RHICmdList, FilterCubeMip);
-				FSceneRenderTargetItem& EffectiveRT = GetEffectiveRenderTarget(SceneContext, false, bUseHQFiltering ? 0 : MipIndex);
-				FSceneRenderTargetItem& EffectiveSource = GetEffectiveSourceTexture(SceneContext, false, bUseHQFiltering ? 0 : MipIndex);
+				FSceneRenderTargetItem& EffectiveRT = GetEffectiveRenderTarget(false, bUseHQFiltering ? 0 : MipIndex);
+				FSceneRenderTargetItem& EffectiveSource = GetEffectiveSourceTexture(false, bUseHQFiltering ? 0 : MipIndex);
 				check(EffectiveRT.TargetableTexture != EffectiveSource.ShaderResourceTexture);
 				const int32 MipSize = 1 << (NumMips - MipIndex - 1);
 
diff --git a/Engine/Source/Runtime/Renderer/Private/PostProcess/SceneRenderTargets.cpp b/Engine/Source/Runtime/Renderer/Private/PostProcess/SceneRenderTargets.cpp
--- a/Engine/Source/Runtime/Renderer/Private/PostProcess/SceneRenderTargets.cpp
+++ b/Engine/Source/Runtime/Renderer/Private/PostProcess/SceneRenderTargets.cpp
@@ -618,12 +618,7 @@ void FSceneRenderTargets::Allocate(FRDGBuilder& GraphBuilder, const FSceneRender
 
 	EPixelFormat MobileSceneColorFormat = GetDesiredMobileSceneColorFormat();
 		
-	bool bNewAllowStaticLighting;
-	{
-		static const auto CVar = IConsoleManager::Get().FindTConsoleVariableDataInt(TEXT("r.AllowStaticLighting"));
-
-		bNewAllowStaticLighting = CVar->GetValueOnRenderThread() != 0;
-	}
+	const bool bNewAllowStaticLighting = IsStaticLightingAllowed();
 
 	bool bDownsampledOcclusionQueries = GDownsampledOcclusionQueries != 0;
 
@@ -1300,51 +1295,6 @@ void FSceneRenderTargets::AllocateMobileRenderTargets(FRHICommandListImmediate&
 	}
 }
 
-void FSceneRenderTargets::AllocateReflectionTargets(FRHICommandList& RHICmdList, int32 TargetSize)
-{
-	if (GSupportsRenderTargetFormat_PF_FloatRGBA)
-	{
-		const int32 NumReflectionCaptureMips = FMath::CeilLogTwo(TargetSize) + 1;
-
-		if (ReflectionColorScratchCubemap[0] && ReflectionColorScratchCubemap[0]->GetRenderTargetItem().TargetableTexture->GetNumMips() != NumReflectionCaptureMips)
-		{
-			ReflectionColorScratchCubemap[0].SafeRelease();
-			ReflectionColorScratchCubemap[1].SafeRelease();
-		}
-
-		// Reflection targets are shared between both mobile and deferred shading paths. If we have already allocated for one and are now allocating for the other,
-		// we can skip these targets.
-		bool bSharedReflectionTargetsAllocated = ReflectionColorScratchCubemap[0] != nullptr;
-
-		if (!bSharedReflectionTargetsAllocated)
-		{
-			// We write to these cubemap faces individually during filtering
-			ETextureCreateFlags CubeTexFlags = TexCreate_TargetArraySlicesIndependently;
-
-			{
-				// Create scratch cubemaps for filtering passes
-				FPooledRenderTargetDesc Desc2(FPooledRenderTargetDesc::CreateCubemapDesc(TargetSize, PF_FloatRGBA, FClearValueBinding(FLinearColor(0, 10000, 0, 0)), CubeTexFlags, TexCreate_RenderTargetable, false, 1, NumReflectionCaptureMips));
-				GRenderTargetPool.FindFreeElement(RHICmdList, Desc2, ReflectionColorScratchCubemap[0], TEXT("ReflectionColorScratchCubemap0"), ERenderTargetTransience::NonTransient );
-				GRenderTargetPool.FindFreeElement(RHICmdList, Desc2, ReflectionColorScratchCubemap[1], TEXT("ReflectionColorScratchCubemap1"), ERenderTargetTransience::NonTransient );
-			}
-
-			extern int32 GDiffuseIrradianceCubemapSize;
-			const int32 NumDiffuseIrradianceMips = FMath::CeilLogTwo(GDiffuseIrradianceCubemapSize) + 1;
-
-			{
-				FPooledRenderTargetDesc Desc2(FPooledRenderTargetDesc::CreateCubemapDesc(GDiffuseIrradianceCubemapSize, PF_FloatRGBA, FClearValueBinding(FLinearColor(0, 10000, 0, 0)), CubeTexFlags, TexCreate_RenderTargetable, false, 1, NumDiffuseIrradianceMips));
-				GRenderTargetPool.FindFreeElement(RHICmdList, Desc2, DiffuseIrradianceScratchCubemap[0], TEXT("DiffuseIrradianceScratchCubemap0"), ERenderTargetTransience::NonTransient);
-				GRenderTargetPool.FindFreeElement(RHICmdList, Desc2, DiffuseIrradianceScratchCubemap[1], TEXT("DiffuseIrradianceScratchCubemap1"), ERenderTargetTransience::NonTransient);
-			}
-
-			{
-				FPooledRenderTargetDesc Desc(FPooledRenderTargetDesc::Create2DDesc(FIntPoint(FSHVector3::MaxSHBasis, 1), PF_FloatRGBA, FClearValueBinding(FLinearColor(0, 10000, 0, 0)), TexCreate_None, TexCreate_RenderTargetable, false));
-				GRenderTargetPool.FindFreeElement(RHICmdList, Desc, SkySHIrradianceMap, TEXT("SkySHIrradianceMap"), ERenderTargetTransience::NonTransient);
-			}
-		}
-	}
-}
-
 void FSceneRenderTargets::AllocateVirtualTextureFeedbackBuffer(FRHICommandList& RHICmdList)
 {
 	if (UseVirtualTexturing(CurrentFeatureLevel))
@@ -1658,23 +1608,6 @@ void FSceneRenderTargets::ReleaseAllTargets()
 	VirtualTextureFeedback.SafeRelease();
 	VirtualTextureFeedbackUAV.SafeRelease();
 
-	for (int32 i = 0; i < UE_ARRAY_COUNT(OptionalShadowDepthColor); i++)
-	{
-		OptionalShadowDepthColor[i].SafeRelease();
-	}
-
-	for (int32 i = 0; i < UE_ARRAY_COUNT(ReflectionColorScratchCubemap); i++)
-	{
-		ReflectionColorScratchCubemap[i].SafeRelease();
-	}
-
-	for (int32 i = 0; i < UE_ARRAY_COUNT(DiffuseIrradianceScratchCubemap); i++)
-	{
-		DiffuseIrradianceScratchCubemap[i].SafeRelease();
-	}
-
-	SkySHIrradianceMap.SafeRelease();
-
 	MobileMultiViewSceneColor.SafeRelease();
 	MobileMultiViewSceneDepthZ.SafeRelease();
 
@@ -1753,13 +1686,6 @@ const FTextureRHIRef& FSceneRenderTargets::GetSceneColorTexture() const
 	return (const FTextureRHIRef&)GetSceneColor()->GetRenderTargetItem().ShaderResourceTexture; 
 }
 
-const FUnorderedAccessViewRHIRef& FSceneRenderTargets::GetSceneColorTextureUAV() const
-{
-	check(GetSceneColorForCurrentShadingPath());
-
-	return (const FUnorderedAccessViewRHIRef&)GetSceneColor()->GetRenderTargetItem().UAV;
-}
-
 FCustomDepthTextures FSceneRenderTargets::RequestCustomDepth(FRDGBuilder& GraphBuilder, bool bPrimitives)
 {
 	FCustomDepthTextures CustomDepthTextures{};
diff --git a/Engine/Source/Runtime/Renderer/Private/PostProcess/SceneRenderTargets.h b/Engine/Source/Runtime/Renderer/Private/PostProcess/SceneRenderTargets.h
--- a/Engine/Source/Runtime/Renderer/Private/PostProcess/SceneRenderTargets.h
+++ b/Engine/Source/Runtime/Renderer/Private/PostProcess/SceneRenderTargets.h
@@ -144,7 +144,6 @@ protected:
 		CurrentFeatureLevel(ERHIFeatureLevel::Num),
 		CurrentShadingPath(EShadingPath::Num),
 		bRequireSceneColorAlpha(false),
-		bAllocateVelocityGBuffer(false),
 		DefaultColorClear(FClearValueBinding::Black),
 		DefaultDepthClear(FClearValueBinding::DepthFar),
 		bHMDAllocatedDepthTarget(false),
@@ -192,21 +191,6 @@ public:
 	int32 GetVirtualTextureFeedbackScale() const;
 	FIntPoint GetVirtualTextureFeedbackBufferSize() const;
 
-	/** Binds the appropriate shadow depth cube map for rendering. */
-	void BeginRenderingCubeShadowDepth(FRHICommandList& RHICmdList, int32 ShadowResolution);
-
-	void FreeReflectionScratchRenderTargets()
-	{
-		for (int32 Idx = 0; Idx < UE_ARRAY_COUNT(ReflectionColorScratchCubemap); ++Idx)
-		{
-			ReflectionColorScratchCubemap[Idx].SafeRelease();
-		}
-		for (int32 Idx = 0; Idx < UE_ARRAY_COUNT(DiffuseIrradianceScratchCubemap); ++Idx)
-		{
-			DiffuseIrradianceScratchCubemap[Idx].SafeRelease();
-		}
-	}
-
 	void SetDefaultColorClear(const FClearValueBinding ColorClear)
 	{
 		DefaultColorClear = ColorClear;
@@ -244,8 +228,6 @@ public:
 	*/
 	static uint16 GetNumSceneColorMSAASamples(ERHIFeatureLevel::Type InFeatureLevel);
 
-	bool IsStaticLightingAllowed() const { return bAllowStaticLighting; }
-
 	/**
 	 * Gets the editor primitives color target/shader resource.  This may recreate the target
 	 * if the msaa settings dont match
@@ -265,7 +247,6 @@ public:
 	// Texture Accessors -----------
 
 	const FTextureRHIRef& GetSceneColorTexture() const;
-	const FUnorderedAccessViewRHIRef& GetSceneColorTextureUAV() const;
 
 	const FTexture2DRHIRef& GetSceneDepthTexture() const
 	{
@@ -349,8 +330,6 @@ public:
 	void AllocateVirtualTextureFeedbackBuffer(FRHICommandList& RHICmdList);
 
 	void AllocateDebugViewModeTargets(FRHICommandList& RHICmdList);
-
-	TRefCountPtr<IPooledRenderTarget>& GetReflectionBrightnessTarget();
 	
 	// Can be called when the Scene Color content is no longer needed. As we create SceneColor on demand we can make sure it is created with the right format.
 	// (as a call to SetSceneColor() can override it with a different format)
@@ -391,17 +370,6 @@ public:
 	TRefCountPtr<IPooledRenderTarget> MobileCustomStencil;
 	// used by the CustomDepth material feature for stencil
 	TRefCountPtr<FRHIShaderResourceView> CustomStencilSRV;
-	// optional in case this RHI requires a color render target (adjust up if necessary)
-	TRefCountPtr<IPooledRenderTarget> OptionalShadowDepthColor[4];
-
-	/** 2 scratch cubemaps used for filtering reflections. */
-	TRefCountPtr<IPooledRenderTarget> ReflectionColorScratchCubemap[2];
-
-	/** Temporary storage during SH irradiance map generation. */
-	TRefCountPtr<IPooledRenderTarget> DiffuseIrradianceScratchCubemap[2];
-
-	/** Temporary storage during SH irradiance map generation. */
-	TRefCountPtr<IPooledRenderTarget> SkySHIrradianceMap;
 
 	/** Color and depth texture arrays for mobile multi-view */
 	TRefCountPtr<IPooledRenderTarget> MobileMultiViewSceneColor;
@@ -556,9 +524,6 @@ private:
 
 	bool bRequireSceneColorAlpha;
 
-	// Set this per frame since there might be cases where we don't need an extra GBuffer
-	bool bAllocateVelocityGBuffer;
-
 	/** Clear color value, defaults to FClearValueBinding::Black */
 	FClearValueBinding DefaultColorClear;
 
diff --git a/Engine/Source/Runtime/Renderer/Private/ReflectionEnvironmentCapture.cpp b/Engine/Source/Runtime/Renderer/Private/ReflectionEnvironmentCapture.cpp
--- a/Engine/Source/Runtime/Renderer/Private/ReflectionEnvironmentCapture.cpp
+++ b/Engine/Source/Runtime/Renderer/Private/ReflectionEnvironmentCapture.cpp
@@ -99,17 +99,78 @@ static FAutoConsoleVariableRef CVarFreeReflectionScratchAfterUse(
 	GFreeReflectionScratchAfterUse,
 	TEXT("Free reflection scratch render targets after use."));
 
+TGlobalResource<FReflectionScratchCubemaps> GReflectionScratchCubemaps;
+
 bool DoGPUArrayCopy()
 {
 	return GRHISupportsResolveCubemapFaces && CVarReflectionCaptureGPUArrayCopy.GetValueOnAnyThread();
 }
 
+void FReflectionScratchCubemaps::Allocate(FRHICommandList& RHICmdList, uint32 TargetSize)
+{
+	if (GSupportsRenderTargetFormat_PF_FloatRGBA)
+	{
+		const int32 NumReflectionCaptureMips = FMath::CeilLogTwo(TargetSize) + 1;
+
+		if (Color[0] && Color[0]->GetTargetableRHI()->GetNumMips() != NumReflectionCaptureMips)
+		{
+			Color[0].SafeRelease();
+			Color[1].SafeRelease();
+		}
+
+		// Reflection targets are shared between both mobile and deferred shading paths. If we have already allocated for one and are now allocating for the other,
+		// we can skip these targets.
+		bool bSharedReflectionTargetsAllocated = Color[0] != nullptr;
+
+		if (!bSharedReflectionTargetsAllocated)
+		{
+			// We write to these cubemap faces individually during filtering
+			ETextureCreateFlags CubeTexFlags = TexCreate_TargetArraySlicesIndependently;
+
+			{
+				// Create scratch cubemaps for filtering passes
+				FPooledRenderTargetDesc Desc2(FPooledRenderTargetDesc::CreateCubemapDesc(TargetSize, PF_FloatRGBA, FClearValueBinding(FLinearColor(0, 10000, 0, 0)), CubeTexFlags, TexCreate_RenderTargetable, false, 1, NumReflectionCaptureMips));
+				GRenderTargetPool.FindFreeElement(RHICmdList, Desc2, Color[0], TEXT("ReflectionColorScratchCubemap0"), ERenderTargetTransience::NonTransient);
+				GRenderTargetPool.FindFreeElement(RHICmdList, Desc2, Color[1], TEXT("ReflectionColorScratchCubemap1"), ERenderTargetTransience::NonTransient);
+			}
+
+			extern int32 GDiffuseIrradianceCubemapSize;
+			const int32 NumDiffuseIrradianceMips = FMath::CeilLogTwo(GDiffuseIrradianceCubemapSize) + 1;
+
+			{
+				FPooledRenderTargetDesc Desc2(FPooledRenderTargetDesc::CreateCubemapDesc(GDiffuseIrradianceCubemapSize, PF_FloatRGBA, FClearValueBinding(FLinearColor(0, 10000, 0, 0)), CubeTexFlags, TexCreate_RenderTargetable, false, 1, NumDiffuseIrradianceMips));
+				GRenderTargetPool.FindFreeElement(RHICmdList, Desc2, Irradiance[0], TEXT("DiffuseIrradianceScratchCubemap0"), ERenderTargetTransience::NonTransient);
+				GRenderTargetPool.FindFreeElement(RHICmdList, Desc2, Irradiance[1], TEXT("DiffuseIrradianceScratchCubemap1"), ERenderTargetTransience::NonTransient);
+			}
+
+			{
+				FPooledRenderTargetDesc Desc(FPooledRenderTargetDesc::Create2DDesc(FIntPoint(FSHVector3::MaxSHBasis, 1), PF_FloatRGBA, FClearValueBinding(FLinearColor(0, 10000, 0, 0)), TexCreate_None, TexCreate_RenderTargetable, false));
+				GRenderTargetPool.FindFreeElement(RHICmdList, Desc, SkySHIrradiance, TEXT("SkySHIrradiance"), ERenderTargetTransience::NonTransient);
+			}
+		}
+	}
+}
+
+void FReflectionScratchCubemaps::Release()
+{
+	for (int32 Index = 0; Index < UE_ARRAY_COUNT(Color); Index++)
+	{
+		Color[Index].SafeRelease();
+	}
+
+	for (int32 Index = 0; Index < UE_ARRAY_COUNT(Irradiance); Index++)
+	{
+		Irradiance[Index].SafeRelease();
+	}
+
+	SkySHIrradiance.SafeRelease();
+}
+
 void FullyResolveReflectionScratchCubes(FRHICommandListImmediate& RHICmdList)
 {
 	SCOPED_DRAW_EVENT(RHICmdList, FullyResolveReflectionScratchCubes);
-	FSceneRenderTargets& SceneContext = FSceneRenderTargets::Get();
-	FTextureRHIRef& Scratch0 = SceneContext.ReflectionColorScratchCubemap[0]->GetRenderTargetItem().TargetableTexture;
-	FTextureRHIRef& Scratch1 = SceneContext.ReflectionColorScratchCubemap[1]->GetRenderTargetItem().TargetableTexture;
+	FRHITexture* Scratch0 = GReflectionScratchCubemaps.Color[0]->GetTargetableRHI();
+	FRHITexture* Scratch1 = GReflectionScratchCubemaps.Color[1]->GetTargetableRHI();
 	FResolveParams ResolveParams(FResolveRect(), CubeFace_PosX, -1, -1, -1);
 	RHICmdList.CopyToResolveTarget(Scratch0, Scratch0, ResolveParams);
 	RHICmdList.CopyToResolveTarget(Scratch1, Scratch1, ResolveParams);  
@@ -283,7 +344,7 @@ float ComputeSingleAverageBrightnessFromCubemap(FRHICommandListImmediate& RHICmd
 	FPooledRenderTargetDesc Desc(FPooledRenderTargetDesc::Create2DDesc(FIntPoint(1, 1), PF_FloatRGBA, FClearValueBinding::None, TexCreate_None, TexCreate_RenderTargetable, false));
 	GRenderTargetPool.FindFreeElement(RHICmdList, Desc, ReflectionBrightnessTarget, TEXT("ReflectionBrightness"));
 
-	FTextureRHIRef& BrightnessTarget = ReflectionBrightnessTarget->GetRenderTargetItem().TargetableTexture;
+	FRHITexture* BrightnessTarget = ReflectionBrightnessTarget->GetTargetableRHI();
 
 	FRHIRenderPassInfo RPInfo(BrightnessTarget, ERenderTargetActions::Load_Store);
 	TransitionRenderPassTargets(RHICmdList, RPInfo);
@@ -342,7 +403,7 @@ void ComputeAverageBrightness(FRHICommandListImmediate& RHICmdList, ERHIFeatureL
 	// necessary to resolve the clears which touched all the mips.  scene rendering only resolves mip 0.
 	FullyResolveReflectionScratchCubes(RHICmdList);	
 
-	FSceneRenderTargetItem& DownSampledCube = FSceneRenderTargets::Get().ReflectionColorScratchCubemap[0]->GetRenderTargetItem();
+	FSceneRenderTargetItem& DownSampledCube = GReflectionScratchCubemaps.Color[0]->GetRenderTargetItem();
 	CreateCubeMips( RHICmdList, FeatureLevel, NumMips, DownSampledCube );
 
 	OutAverageBrightness = ComputeSingleAverageBrightnessFromCubemap(RHICmdList, FeatureLevel, CubmapSize, DownSampledCube);
@@ -429,7 +490,7 @@ void FilterReflectionEnvironment(FRHICommandListImmediate& RHICmdList, ERHIFeatu
 	const int32 EffectiveTopMipSize = CubmapSize;
 	const int32 NumMips = FMath::CeilLogTwo(EffectiveTopMipSize) + 1;
 
-	FSceneRenderTargetItem& EffectiveColorRT = FSceneRenderTargets::Get().ReflectionColorScratchCubemap[0]->GetRenderTargetItem();
+	FSceneRenderTargetItem& EffectiveColorRT = GReflectionScratchCubemaps.Color[0]->GetRenderTargetItem();
 
 	FGraphicsPipelineStateInitializer GraphicsPSOInit;
 	GraphicsPSOInit.RasterizerState = TStaticRasterizerState<FM_Solid, CM_None>::GetRHI();
@@ -478,10 +539,9 @@ void FilterReflectionEnvironment(FRHICommandListImmediate& RHICmdList, ERHIFeatu
 	RHICmdList.Transition(FRHITransitionInfo(EffectiveColorRT.TargetableTexture, ERHIAccess::Unknown, ERHIAccess::SRVMask));
 
 	auto ShaderMap = GetGlobalShaderMap(FeatureLevel);
-	FSceneRenderTargets& SceneContext = FSceneRenderTargets::Get();
 
-	FSceneRenderTargetItem& DownSampledCube = FSceneRenderTargets::Get().ReflectionColorScratchCubemap[0]->GetRenderTargetItem();
-	FSceneRenderTargetItem& FilteredCube = FSceneRenderTargets::Get().ReflectionColorScratchCubemap[1]->GetRenderTargetItem();
+	FSceneRenderTargetItem& DownSampledCube = GReflectionScratchCubemaps.Color[0]->GetRenderTargetItem();
+	FSceneRenderTargetItem& FilteredCube = GReflectionScratchCubemaps.Color[1]->GetRenderTargetItem();
 
 	CreateCubeMips( RHICmdList, FeatureLevel, NumMips, DownSampledCube );
 
@@ -615,15 +675,14 @@ void ClearScratchCubemaps(FRHICommandListImmediate& RHICmdList, int32 TargetSize
 {
 	SCOPED_DRAW_EVENT(RHICmdList, ClearScratchCubemaps);
 
-	FSceneRenderTargets& SceneContext = FSceneRenderTargets::Get();
-	SceneContext.AllocateReflectionTargets(RHICmdList, TargetSize);
+	GReflectionScratchCubemaps.Allocate(RHICmdList, TargetSize);
 
 	FMemMark Mark(FMemStack::Get());
 	FRDGBuilder GraphBuilder(RHICmdList);
 
 	for (int32 RenderTargetIndex = 0; RenderTargetIndex < 2; ++RenderTargetIndex)
 	{
-		FRDGTextureRef OutputTexture = GraphBuilder.RegisterExternalTexture(SceneContext.ReflectionColorScratchCubemap[RenderTargetIndex], TEXT("OutputCubemap"));
+		FRDGTextureRef OutputTexture = GraphBuilder.RegisterExternalTexture(GReflectionScratchCubemaps.Color[RenderTargetIndex], TEXT("OutputCubemap"));
 
 		RDG_EVENT_SCOPE(GraphBuilder, "ClearScratchCubemapsRT%d", RenderTargetIndex);
 
@@ -675,14 +734,13 @@ void CaptureSceneToScratchCubemap(FRHICommandListImmediate& RHICmdList, FSceneRe
 		// allow them to restart their command buffers
 		RHICmdList.SubmitCommandsAndFlushGPU();
 
-		FSceneRenderTargets& SceneContext = FSceneRenderTargets::Get();
-		SceneContext.AllocateReflectionTargets(RHICmdList, CubemapSize);
+		GReflectionScratchCubemaps.Allocate(RHICmdList, CubemapSize);
 
 		const FViewInfo& View = SceneRenderer->Views[0];
 
 		FRDGBuilder GraphBuilder(RHICmdList);
 
-		FRDGTextureRef OutputTexture = GraphBuilder.RegisterExternalTexture(SceneContext.ReflectionColorScratchCubemap[0], TEXT("ReflectionColorScratchCubemap"));
+		FRDGTextureRef OutputTexture = GraphBuilder.RegisterExternalTexture(GReflectionScratchCubemaps.Color[0], TEXT("Color"));
 
 		auto* PassParameters = GraphBuilder.AllocParameters<FCopySceneColorToCubeFacePS::FParameters>();
 		PassParameters->RenderTargets[0] = FRenderTargetBinding(OutputTexture, ERenderTargetLoadAction::ENoAction, 0, CubeFace);
@@ -714,6 +772,8 @@ void CaptureSceneToScratchCubemap(FRHICommandListImmediate& RHICmdList, FSceneRe
 			PassParameters->SkyLightCaptureParameters = SkyLightParametersValue;
 		}
 
+		FSceneRenderTargets& SceneContext = FSceneRenderTargets::Get();
+
 		PassParameters->View = View.ViewUniformBuffer;
 		PassParameters->SceneColorSampler = TStaticSamplerState<SF_Bilinear, AM_Clamp, AM_Clamp, AM_Clamp>::GetRHI();
 		PassParameters->SceneColorTexture = GraphBuilder.RegisterExternalTexture(SceneContext.GetSceneColor(), TEXT("ColorTexture"));
@@ -783,12 +843,10 @@ void CopyCubemapToScratchCubemap(FRHICommandListImmediate& RHICmdList, ERHIFeatu
 		return;
 	}
 
-	FSceneRenderTargets& SceneContext = FSceneRenderTargets::Get();
-
 	FMemMark Mark(FMemStack::Get());
 	FRDGBuilder GraphBuilder(RHICmdList);
 
-	FRDGTextureRef OutputTexture = GraphBuilder.RegisterExternalTexture(SceneContext.ReflectionColorScratchCubemap[0], TEXT("ReflectionColorScratchCubemap"));
+	FRDGTextureRef OutputTexture = GraphBuilder.RegisterExternalTexture(GReflectionScratchCubemaps.Color[0], TEXT("Color"));
 
 	for (uint32 CubeFace = 0; CubeFace < CubeFace_MAX; CubeFace++)
 	{
@@ -807,7 +865,7 @@ void CopyCubemapToScratchCubemap(FRHICommandListImmediate& RHICmdList, ERHIFeatu
 			RDG_EVENT_NAME("CopyCubemapToCubeFace"),
 			PassParameters,
 			ERDGPassFlags::Raster,
-			[EffectiveSize, &SceneContext, SourceCubemapResource, PassParameters, FeatureLevel](FRHICommandList& InRHICmdList)
+			[EffectiveSize, SourceCubemapResource, PassParameters, FeatureLevel](FRHICommandList& InRHICmdList)
 		{
 			const FIntPoint SourceDimensions(SourceCubemapResource->GetSizeX(), SourceCubemapResource->GetSizeY());
 			const FIntRect ViewRect(0, 0, EffectiveSize, EffectiveSize);
@@ -1376,9 +1434,8 @@ void CopyToSceneArray(FRHICommandListImmediate& RHICmdList, FScene* Scene, FRefl
 	const int32 NumMips = FMath::CeilLogTwo(EffectiveTopMipSize) + 1;
 
 	const int32 CaptureIndex = FindOrAllocateCubemapIndex(Scene, ReflectionProxy->Component);
-	FSceneRenderTargets& SceneContext = FSceneRenderTargets::Get();
-	
-	FSceneRenderTargetItem& FilteredCube = SceneContext.ReflectionColorScratchCubemap[1]->GetRenderTargetItem();
+
+	FSceneRenderTargetItem& FilteredCube = GReflectionScratchCubemaps.Color[1]->GetRenderTargetItem();
 	FSceneRenderTargetItem& DestCube = Scene->ReflectionSceneData.CubemapArray.GetRenderTarget();
 
 	// GPU copy back to the scene's texture array, which is not a render target
@@ -1543,7 +1600,7 @@ void ReadbackRadianceMap(FRHICommandListImmediate& RHICmdList, int32 CubmapSize,
 
 	const int32 MipIndex = 0;
 
-	FSceneRenderTargetItem& SourceCube = FSceneRenderTargets::Get().ReflectionColorScratchCubemap[0]->GetRenderTargetItem();
+	FSceneRenderTargetItem& SourceCube = GReflectionScratchCubemaps.Color[0]->GetRenderTargetItem();
 	check(SourceCube.ShaderResourceTexture->GetFormat() == PF_FloatRGBA);
 	const int32 CubeFaceBytes = CubmapSize * CubmapSize * OutRadianceMap.GetTypeSize();
 
@@ -1567,9 +1624,8 @@ void CopyToSkyTexture(FRHICommandList& RHICmdList, FScene* Scene, FTexture* Proc
 	{
 		const int32 EffectiveTopMipSize = ProcessedTexture->GetSizeX();
 		const int32 NumMips = FMath::CeilLogTwo(EffectiveTopMipSize) + 1;
-		FSceneRenderTargets& SceneContext = FSceneRenderTargets::Get();
 
-		FSceneRenderTargetItem& FilteredCube = FSceneRenderTargets::Get().ReflectionColorScratchCubemap[1]->GetRenderTargetItem();
+		FSceneRenderTargetItem& FilteredCube = GReflectionScratchCubemaps.Color[1]->GetRenderTargetItem();
 
 		// GPU copy back to the skylight's texture, which is not a render target
 		FRHICopyTextureInfo CopyInfo;
@@ -1709,8 +1765,7 @@ void FScene::UpdateSkyCaptureContents(
 			ENQUEUE_RENDER_COMMAND(FreeReflectionScratch)(
 				[](FRHICommandListImmediate& RHICmdList)
 			{
-				FSceneRenderTargets& SceneContext = FSceneRenderTargets::Get();
-				SceneContext.FreeReflectionScratchRenderTargets();
+				GReflectionScratchCubemaps.Release();
 				GRenderTargetPool.FreeUnusedResources();
 			});
 		}
diff --git a/Engine/Source/Runtime/Renderer/Private/ReflectionEnvironmentCapture.h b/Engine/Source/Runtime/Renderer/Private/ReflectionEnvironmentCapture.h
--- a/Engine/Source/Runtime/Renderer/Private/ReflectionEnvironmentCapture.h
+++ b/Engine/Source/Runtime/Renderer/Private/ReflectionEnvironmentCapture.h
@@ -10,6 +10,7 @@
 #include "Math/SHMath.h"
 #include "RHI.h"
 #include "GlobalShader.h"
+#include "RendererInterface.h"
 
 extern void ComputeDiffuseIrradiance(FRHICommandListImmediate& RHICmdList, ERHIFeatureLevel::Type FeatureLevel, FTextureRHIRef LightingSource, int32 LightingSourceMipIndex, FSHVectorRGB3* OutIrradianceEnvironmentMap);
 
@@ -61,4 +62,26 @@ public:
 	{}
 
 	TCubeFilterPS() {}
+};
+
+class FReflectionScratchCubemaps : public FRenderResource
+{
+public:
+	void Allocate(FRHICommandList& RHICmdList, uint32 TargetSize);
+	void Release();
+
+	/** 2 scratch cubemaps used for filtering reflections. */
+	TRefCountPtr<IPooledRenderTarget> Color[2];
+
+	/** Temporary storage during SH irradiance map generation. */
+	TRefCountPtr<IPooledRenderTarget> Irradiance[2];
+
+	/** Temporary storage during SH irradiance map generation. */
+	TRefCountPtr<IPooledRenderTarget> SkySHIrradiance;
+
+private:
+	void ReleaseDynamicRHI() override
+	{
+		Release();
+	}
 };
\ No newline at end of file
diff --git a/Engine/Source/Runtime/Renderer/Private/ReflectionEnvironmentDiffuseIrradiance.cpp b/Engine/Source/Runtime/Renderer/Private/ReflectionEnvironmentDiffuseIrradiance.cpp
--- a/Engine/Source/Runtime/Renderer/Private/ReflectionEnvironmentDiffuseIrradiance.cpp
+++ b/Engine/Source/Runtime/Renderer/Private/ReflectionEnvironmentDiffuseIrradiance.cpp
@@ -11,7 +11,7 @@
 #include "Shader.h"
 #include "StaticBoundShaderState.h"
 #include "RHIStaticStates.h"
-#include "PostProcess/SceneRenderTargets.h"
+#include "ReflectionEnvironmentCapture.h"
 #include "GlobalShader.h"
 #include "PostProcess/SceneFilterRendering.h"
 #include "ScreenRendering.h"
@@ -21,18 +21,20 @@
 
 extern int32 GDiffuseIrradianceCubemapSize;
 
-FSceneRenderTargetItem& GetEffectiveDiffuseIrradianceRenderTarget(FSceneRenderTargets& SceneContext, int32 TargetMipIndex)
+extern TGlobalResource<FReflectionScratchCubemaps> GReflectionScratchCubemaps;
+
+FSceneRenderTargetItem& GetEffectiveDiffuseIrradianceRenderTarget(int32 TargetMipIndex)
 {
 	const int32 ScratchTextureIndex = TargetMipIndex % 2;
 
-	return SceneContext.DiffuseIrradianceScratchCubemap[ScratchTextureIndex]->GetRenderTargetItem();
+	return GReflectionScratchCubemaps.Irradiance[ScratchTextureIndex]->GetRenderTargetItem();
 }
 
-FSceneRenderTargetItem& GetEffectiveDiffuseIrradianceSourceTexture(FSceneRenderTargets& SceneContext, int32 TargetMipIndex)
+FSceneRenderTargetItem& GetEffectiveDiffuseIrradianceSourceTexture(int32 TargetMipIndex)
 {
 	const int32 ScratchTextureIndex = 1 - TargetMipIndex % 2;
 
-	return SceneContext.DiffuseIrradianceScratchCubemap[ScratchTextureIndex]->GetRenderTargetItem();
+	return GReflectionScratchCubemaps.Irradiance[ScratchTextureIndex]->GetRenderTargetItem();
 }
 
 /** Pixel shader used for copying to diffuse irradiance texture. */
@@ -195,7 +197,6 @@ IMPLEMENT_SHADER_TYPE(,FAccumulateCubeFacesPS,TEXT("/Engine/Private/ReflectionEn
 void ComputeDiffuseIrradiance(FRHICommandListImmediate& RHICmdList, ERHIFeatureLevel::Type FeatureLevel, FTextureRHIRef LightingSource, int32 LightingSourceMipIndex, FSHVectorRGB3* OutIrradianceEnvironmentMap)
 {
 	auto ShaderMap = GetGlobalShaderMap(FeatureLevel);
-	FSceneRenderTargets& SceneContext = FSceneRenderTargets::Get();
 
 	FGraphicsPipelineStateInitializer GraphicsPSOInit;
 	GraphicsPSOInit.RasterizerState = TStaticRasterizerState<FM_Solid, CM_None>::GetRHI();
@@ -208,7 +209,7 @@ void ComputeDiffuseIrradiance(FRHICommandListImmediate& RHICmdList, ERHIFeatureL
 		{
 			const int32 MipIndex = 0;
 			const int32 MipSize = GDiffuseIrradianceCubemapSize;
-			FSceneRenderTargetItem& EffectiveRT = GetEffectiveDiffuseIrradianceRenderTarget(SceneContext, MipIndex);			
+			FSceneRenderTargetItem& EffectiveRT = GetEffectiveDiffuseIrradianceRenderTarget(MipIndex);
 
 			RHICmdList.Transition(FRHITransitionInfo(EffectiveRT.TargetableTexture, ERHIAccess::Unknown, ERHIAccess::RTV));
 
@@ -258,8 +259,8 @@ void ComputeDiffuseIrradiance(FRHICommandListImmediate& RHICmdList, ERHIFeatureL
 				const int32 SourceMipIndex = FMath::Max(MipIndex - 1, 0);
 				const int32 MipSize = 1 << (NumMips - MipIndex - 1);
 
-				FSceneRenderTargetItem& EffectiveRT = GetEffectiveDiffuseIrradianceRenderTarget(SceneContext, MipIndex);
-				FSceneRenderTargetItem& EffectiveSource = GetEffectiveDiffuseIrradianceSourceTexture(SceneContext, MipIndex);
+				FSceneRenderTargetItem& EffectiveRT = GetEffectiveDiffuseIrradianceRenderTarget(MipIndex);
+				FSceneRenderTargetItem& EffectiveSource = GetEffectiveDiffuseIrradianceSourceTexture(MipIndex);
 				check(EffectiveRT.TargetableTexture != EffectiveSource.ShaderResourceTexture);
 
 				RHICmdList.Transition(FRHITransitionInfo(EffectiveRT.TargetableTexture, ERHIAccess::Unknown, ERHIAccess::RTV));
@@ -304,7 +305,7 @@ void ComputeDiffuseIrradiance(FRHICommandListImmediate& RHICmdList, ERHIFeatureL
 
 		{
 			// Gather the cubemap face results and normalize, copy this coefficient to FSceneRenderTargets::Get().SkySHIrradianceMap
-			FSceneRenderTargetItem& EffectiveRT = FSceneRenderTargets::Get().SkySHIrradianceMap->GetRenderTargetItem();
+			FSceneRenderTargetItem& EffectiveRT = GReflectionScratchCubemaps.SkySHIrradiance->GetRenderTargetItem();
 
 			//load/store actions so we don't lose results as we render one pixel at a time on tile renderers.
 			RHICmdList.Transition(FRHITransitionInfo(EffectiveRT.TargetableTexture, ERHIAccess::Unknown, ERHIAccess::RTV));
@@ -327,7 +328,7 @@ void ComputeDiffuseIrradiance(FRHICommandListImmediate& RHICmdList, ERHIFeatureL
 
 			const int32 SourceMipIndex = NumMips - 1;
 			const int32 MipSize = 1;
-			FSceneRenderTargetItem& EffectiveSource = GetEffectiveDiffuseIrradianceRenderTarget(SceneContext, SourceMipIndex);
+			FSceneRenderTargetItem& EffectiveSource = GetEffectiveDiffuseIrradianceRenderTarget(SourceMipIndex);
 			PixelShader->SetParameters(RHICmdList, SourceMipIndex, EffectiveSource.ShaderResourceTexture);
 
 			DrawRectangle( 
@@ -347,7 +348,7 @@ void ComputeDiffuseIrradiance(FRHICommandListImmediate& RHICmdList, ERHIFeatureL
 
 	{
 		// Read back the completed SH environment map
-		FSceneRenderTargetItem& EffectiveRT = FSceneRenderTargets::Get().SkySHIrradianceMap->GetRenderTargetItem();
+		FSceneRenderTargetItem& EffectiveRT = GReflectionScratchCubemaps.SkySHIrradiance->GetRenderTargetItem();
 		check(EffectiveRT.ShaderResourceTexture->GetFormat() == PF_FloatRGBA);
 
 		TArray<FFloat16Color> SurfaceData;
diff --git a/Engine/Source/Runtime/Renderer/Private/SceneRendering.cpp b/Engine/Source/Runtime/Renderer/Private/SceneRendering.cpp
--- a/Engine/Source/Runtime/Renderer/Private/SceneRendering.cpp
+++ b/Engine/Source/Runtime/Renderer/Private/SceneRendering.cpp
@@ -439,6 +439,11 @@ FASTVRAM_CVAR(DistanceFieldAOScreenGridResources, 1);
 FASTVRAM_CVAR(ForwardLightingCullingResources, 1);
 FASTVRAM_CVAR(GlobalDistanceFieldCullGridBuffers, 1);
 
+bool IsStaticLightingAllowed()
+{
+	static const auto CVar = IConsoleManager::Get().FindTConsoleVariableDataInt(TEXT("r.AllowStaticLighting"));
+	return CVar->GetValueOnRenderThread() != 0;
+}
 
 #if !UE_BUILD_SHIPPING
 namespace
@@ -2550,8 +2555,23 @@ void FSceneRenderer::PrepareViewRectsForRendering()
 }
 
 #if WITH_MGPU
-void FSceneRenderer::ComputeViewGPUMasks(FRHIGPUMask RenderTargetGPUMask)
+FRHIGPUMask FSceneRenderer::ComputeGPUMasks(FRHICommandListImmediate& RHICmdList)
 {
+	FRHIGPUMask RenderTargetGPUMask = FRHIGPUMask::GPU0();
+	
+	if (GNumExplicitGPUsForRendering > 1 && ViewFamily.RenderTarget)
+	{
+		RenderTargetGPUMask = ViewFamily.RenderTarget->GetGPUMask(RHICmdList);
+	}
+
+	{
+		static auto CVar = IConsoleManager::Get().FindConsoleVariable(TEXT("r.PathTracing.GPUCount"));
+		if (CVar && CVar->GetInt() > 1)
+		{
+			RenderTargetGPUMask = FRHIGPUMask::All(); // Broadcast to all GPUs 
+		}
+	}
+
 	// First check whether we are in multi-GPU and if fork and join cross-gpu transfers are enabled.
 	// Otherwise fallback on rendering the whole view family on each relevant GPU using broadcast logic.
 	if (GNumExplicitGPUsForRendering > 1 && CVarEnableMultiGPUForkAndJoin.GetValueOnAnyThread() != 0)
@@ -2612,6 +2632,8 @@ void FSceneRenderer::ComputeViewGPUMasks(FRHIGPUMask RenderTargetGPUMask)
 	{
 		AllViewsGPUMask |= Views[ViewIndex].GPUMask;
 	}
+
+	return RenderTargetGPUMask;
 }
 #endif // WITH_MGPU
 
diff --git a/Engine/Source/Runtime/Renderer/Private/SceneRendering.h b/Engine/Source/Runtime/Renderer/Private/SceneRendering.h
--- a/Engine/Source/Runtime/Renderer/Private/SceneRendering.h
+++ b/Engine/Source/Runtime/Renderer/Private/SceneRendering.h
@@ -1751,8 +1751,8 @@ public:
 	void PrepareViewRectsForRendering();
 
 #if WITH_MGPU
-	/** Setups each FViewInfo::GPUMask. */
-	void ComputeViewGPUMasks(FRHIGPUMask RenderTargetGPUMask);
+	/** Assigns the view GPU masks and returns the render target GPU mask. */
+	FRHIGPUMask ComputeGPUMasks(FRHICommandListImmediate& RHICmdList);
 #endif
 
 	/** Update the rendertarget with each view results.*/
@@ -2307,6 +2307,8 @@ enum class EGPUSkinCacheTransition
 	Renderer,
 };
 
+extern bool IsStaticLightingAllowed();
+
 /* Run GPU skin cache resource transitions */
 void RunGPUSkinCacheTransition(class FRHICommandList& RHICmdList, class FScene* Scene, EGPUSkinCacheTransition Type);
 /** Resolves the view rect of scene color or depth using either a custom resolve or hardware resolve. */
