Add FTaskTagScopes to VT tasks

diff --git a/Engine/Source/Runtime/Renderer/Private/VT/VirtualTextureSystem.cpp b/Engine/Source/Runtime/Renderer/Private/VT/VirtualTextureSystem.cpp
--- a/Engine/Source/Runtime/Renderer/Private/VT/VirtualTextureSystem.cpp
+++ b/Engine/Source/Runtime/Renderer/Private/VT/VirtualTextureSystem.cpp
@@ -167,6 +167,7 @@ public:
 
 	void DoTask(ENamedThreads::Type CurrentThread, const FGraphEventRef& MyCompletionGraphEvent)
 	{
+		FTaskTagScope TaskTagScope(ETaskTag::EParallelRenderingThread);
 		DoTask(Parameters);
 	}
 
@@ -184,6 +185,7 @@ public:
 
 	void DoTask(ENamedThreads::Type CurrentThread, const FGraphEventRef& MyCompletionGraphEvent)
 	{
+		FTaskTagScope TaskTagScope(ETaskTag::EParallelRenderingThread);
 		Parameters.RequestList->Initialize();
 		Parameters.System->GatherRequestsTask(Parameters);
 	}
