Fix incorrect modulo operation in FSceneViewState::GetFrameIndex

Code calling this function to obtain a value in [0,N-1] was actually getting back a value in [0,N-2].

diff --git a/Engine/Source/Runtime/Renderer/Private/ScenePrivate.h b/Engine/Source/Runtime/Renderer/Private/ScenePrivate.h
--- a/Engine/Source/Runtime/Renderer/Private/ScenePrivate.h
+++ b/Engine/Source/Runtime/Renderer/Private/ScenePrivate.h
@@ -1135,7 +1135,7 @@ public:
 	inline uint32 GetFrameIndex(uint32 Pow2Modulus) const
 	{
 		check(FMath::IsPowerOfTwo(Pow2Modulus));
-		return FrameIndex % (Pow2Modulus - 1);
+		return FrameIndex & (Pow2Modulus - 1);
 	}
 
 	// Returns 32bits frame index.
