BuildHZB: SrcSize code for visibility buffer path was stomped by main merge. Reverted it back.

diff --git a/Engine/Source/Runtime/Renderer/Private/SceneTextureReductions.cpp b/Engine/Source/Runtime/Renderer/Private/SceneTextureReductions.cpp
--- a/Engine/Source/Runtime/Renderer/Private/SceneTextureReductions.cpp
+++ b/Engine/Source/Runtime/Renderer/Private/SceneTextureReductions.cpp
@@ -213,8 +213,8 @@ void BuildHZB(
 
 	// Reduce first mips Closesy and furtherest are done at same time.
 	{
-		FIntPoint SrcSize = SceneDepth->Desc.Extent;
 		FRDGTextureSRVRef ParentTextureMip = GraphBuilder.CreateSRV(FRDGTextureSRVDesc::Create(SceneDepth));
+		FIntPoint SrcSize = VisBufferTexture ? VisBufferTexture->Desc.Extent : ParentTextureMip->Desc.Texture->Desc.Extent;
 		
 		FVector4 DispatchThreadIdToBufferUV;
 		DispatchThreadIdToBufferUV.X = 2.0f / float(SrcSize.X);
