Fixed: Assertion failed: !GRDGBuilderActive after enabling mobile preview in UE5 QAGame editor

diff --git a/Engine/Source/Runtime/Renderer/Private/MobileShadingRenderer.cpp b/Engine/Source/Runtime/Renderer/Private/MobileShadingRenderer.cpp
--- a/Engine/Source/Runtime/Renderer/Private/MobileShadingRenderer.cpp
+++ b/Engine/Source/Runtime/Renderer/Private/MobileShadingRenderer.cpp
@@ -291,6 +291,14 @@ void FMobileSceneRenderer::InitViews(FRDGBuilder& GraphBuilder)
 
 	check(Scene);
 
+	if (bUseVirtualTexturing)
+	{
+		RDG_GPU_STAT_SCOPE(GraphBuilder, VirtualTextureUpdate);
+		// AllocateResources needs to be called before RHIBeginScene
+		FVirtualTextureSystem::Get().AllocateResources(GraphBuilder, FeatureLevel);
+		FVirtualTextureSystem::Get().CallPendingCallbacks();
+	}
+
 	FILCUpdatePrimTaskData ILCTaskData;
 	FViewVisibleCommandsPerView ViewCommandsPerView;
 	ViewCommandsPerView.SetNum(Views.Num());
@@ -355,11 +363,8 @@ void FMobileSceneRenderer::InitViews(FRDGBuilder& GraphBuilder)
 
 	// Allocate the maximum scene render target space for the current view family.
 	SceneContext.SetKeepDepthContent(bKeepDepthContent);
-	{
-		FRDGBuilder LocalGraphBuilder(RHICmdList);
-		SceneContext.Allocate(LocalGraphBuilder, this);
-		LocalGraphBuilder.Execute();
-	}
+	SceneContext.Allocate(GraphBuilder, this);
+
 	if (bDeferredShading)
 	{
 		ETextureCreateFlags AddFlags = bRequiresMultiPass ? TexCreate_InputAttachmentRead : (TexCreate_InputAttachmentRead | TexCreate_Memoryless);
@@ -371,18 +376,6 @@ void FMobileSceneRenderer::InitViews(FRDGBuilder& GraphBuilder)
 	{
 		InitSkyAtmosphereForViews(RHICmdList);
 	}
-		
-	if (bUseVirtualTexturing)
-	{
-		FRDGBuilder LocalGraphBuilder(RHICmdList);
-		{
-			RDG_GPU_STAT_SCOPE(LocalGraphBuilder, VirtualTextureUpdate);
-			// AllocateResources needs to be called before RHIBeginScene
-			FVirtualTextureSystem::Get().AllocateResources(LocalGraphBuilder, FeatureLevel);
-			FVirtualTextureSystem::Get().CallPendingCallbacks();
-		}
-		LocalGraphBuilder.Execute();
-	}
 
 	if (bRequiresPixelProjectedPlanarRelfectionPass)
 	{
