Disabled FFTBloom on platforms that don't support it instead of crashing.

diff --git a/Engine/Source/Runtime/Renderer/Private/PostProcess/PostProcessFFTBloom.cpp b/Engine/Source/Runtime/Renderer/Private/PostProcess/PostProcessFFTBloom.cpp
--- a/Engine/Source/Runtime/Renderer/Private/PostProcess/PostProcessFFTBloom.cpp
+++ b/Engine/Source/Runtime/Renderer/Private/PostProcess/PostProcessFFTBloom.cpp
@@ -14,6 +14,12 @@ TAutoConsoleVariable<int32> CVarHalfResFFTBloom(
 	TEXT(" 1: Half-resolution convolution that excludes the center of the kernel.\n"),
 	ECVF_Scalability | ECVF_RenderThreadSafe);
 
+static bool DoesPlatformSupportFFTBloom(EShaderPlatform Platform)
+{
+	// @todo MetalMRT: Metal MRT can't cope with the threadgroup storage requirements for these shaders right now
+	return IsFeatureLevelSupported(Platform, ERHIFeatureLevel::SM5) && !IsMetalMRTPlatform(Platform) && IsPCPlatform(Platform);
+}
+
 class FFFTShader : public FGlobalShader
 {
 public:
@@ -22,8 +28,7 @@ public:
 
 	static bool ShouldCompilePermutation(const FGlobalShaderPermutationParameters& Parameters)
 	{
-		// @todo MetalMRT: Metal MRT can't cope with the threadgroup storage requirements for these shaders right now
-		return IsFeatureLevelSupported(Parameters.Platform, ERHIFeatureLevel::SM5) && !IsMetalMRTPlatform(Parameters.Platform) && IsPCPlatform(Parameters.Platform);
+		return DoesPlatformSupportFFTBloom(Parameters.Platform);
 	}
 
 	FFFTShader() = default;
@@ -186,7 +191,7 @@ bool IsFFTBloomPhysicalKernelReady(const FViewInfo& View)
 bool IsFFTBloomEnabled(const FViewInfo& View)
 {
 	const bool bOldMetalNoFFT = IsMetalPlatform(View.GetShaderPlatform()) && (RHIGetShaderLanguageVersion(View.GetShaderPlatform()) < 4) && IsPCPlatform(View.GetShaderPlatform());
-	const bool bUseFFTBloom = View.FinalPostProcessSettings.BloomMethod == EBloomMethod::BM_FFT && View.ViewState != nullptr;
+	const bool bUseFFTBloom = View.FinalPostProcessSettings.BloomMethod == EBloomMethod::BM_FFT && View.ViewState != nullptr && DoesPlatformSupportFFTBloom(View.GetShaderPlatform());
 
 	static bool bWarnAboutOldMetalFFTOnce = false;
 
