Update names of render targets used for VT to actually include the word 'Virtual'

diff --git a/Engine/Source/Runtime/Renderer/Private/VT/VirtualTexturePhysicalSpace.cpp b/Engine/Source/Runtime/Renderer/Private/VT/VirtualTexturePhysicalSpace.cpp
--- a/Engine/Source/Runtime/Renderer/Private/VT/VirtualTexturePhysicalSpace.cpp
+++ b/Engine/Source/Runtime/Renderer/Private/VT/VirtualTexturePhysicalSpace.cpp
@@ -112,7 +112,7 @@ void FVirtualTexturePhysicalSpace::InitRHI()
 			(bCreateAliasedUAV || FormatSRV == PF_A32B32G32R32F) ? TexCreate_ShaderResource | TexCreate_UAV : TexCreate_ShaderResource,
 			false);
 
-		GRenderTargetPool.FindFreeElement(RHICmdList, Desc, PooledRenderTarget[Layer], TEXT("PhysicalTexture"));
+		GRenderTargetPool.FindFreeElement(RHICmdList, Desc, PooledRenderTarget[Layer], TEXT("VirtualPhysicalTexture"));
 		FRHITexture* TextureRHI = PooledRenderTarget[Layer]->GetRenderTargetItem().ShaderResourceTexture;
 
 		// Create sRGB and non-sRGB shader resource views into the physical texture
diff --git a/Engine/Source/Runtime/Renderer/Private/VT/VirtualTextureSpace.cpp b/Engine/Source/Runtime/Renderer/Private/VT/VirtualTextureSpace.cpp
--- a/Engine/Source/Runtime/Renderer/Private/VT/VirtualTextureSpace.cpp
+++ b/Engine/Source/Runtime/Renderer/Private/VT/VirtualTextureSpace.cpp
@@ -266,7 +266,7 @@ void FVirtualTextureSpace::AllocateTextures(FRDGBuilder& GraphBuilder)
 	{
 		RDG_GPU_MASK_SCOPE(GraphBuilder, FRHIGPUMask::All());
 
-		const TCHAR* TextureNames[] = { TEXT("PageTable_0"), TEXT("PageTable_1") };
+		const TCHAR* TextureNames[] = { TEXT("VirtualPageTable_0"), TEXT("VirtualPageTable_1") };
 		static_assert(UE_ARRAY_COUNT(TextureNames) == TextureCapacity, "");
 
 		for (uint32 TextureIndex = 0u; TextureIndex < GetNumPageTableTextures(); ++TextureIndex)
