Relax restriction that requires all VTs within a stack to match aspect ratio.  Instead, crop mismatched textures to match.  This should support common case of dummy/constant textures at least.

diff --git a/Engine/Source/Runtime/Renderer/Private/VT/AllocatedVirtualTexture.cpp b/Engine/Source/Runtime/Renderer/Private/VT/AllocatedVirtualTexture.cpp
--- a/Engine/Source/Runtime/Renderer/Private/VT/AllocatedVirtualTexture.cpp
+++ b/Engine/Source/Runtime/Renderer/Private/VT/AllocatedVirtualTexture.cpp
@@ -227,13 +227,17 @@ uint32 FAllocatedVirtualTexture::AddUniqueProducer(FVirtualTextureProducerHandle
 		check(ProducerDesc.TileSize == Description.TileSize);
 		check(ProducerDesc.TileBorderSize == Description.TileBorderSize);
 
-		const uint32 BlockSizeInTiles = FMath::Max(BlockWidthInTiles, BlockHeightInTiles);
-		const uint32 ProducerBlockSizeInTiles = FMath::Max(ProducerDesc.BlockWidthInTiles, ProducerDesc.BlockHeightInTiles);
-		MipBias = FMath::CeilLogTwo(BlockSizeInTiles / ProducerBlockSizeInTiles);
-
-		check((BlockSizeInTiles / ProducerBlockSizeInTiles) * ProducerBlockSizeInTiles == BlockSizeInTiles);
-		check(ProducerDesc.BlockWidthInTiles << MipBias == BlockWidthInTiles);
-		check(ProducerDesc.BlockHeightInTiles << MipBias == BlockHeightInTiles);
+		const uint32 MipBiasX = FMath::CeilLogTwo(BlockWidthInTiles / ProducerDesc.BlockWidthInTiles);
+		const uint32 MipBiasY = FMath::CeilLogTwo(BlockHeightInTiles / ProducerDesc.BlockHeightInTiles);
+		check(ProducerDesc.BlockWidthInTiles << MipBiasX == BlockWidthInTiles);
+		check(ProducerDesc.BlockHeightInTiles << MipBiasY == BlockHeightInTiles);
+
+		// If the producer aspect ratio doesn't match the aspect ratio for the AllocatedVT, there's no way to choose a 100% mip bias
+		// By chossing the minimum of X/Y bias, we'll effectively crop this producer to match the aspect ratio of the AllocatedVT
+		// This case can happen as base materials will choose to group VTs together into a stack as long as all the textures assigned in the base material share the same aspect ratio
+		// But it's possible for a MI to overide some of thse textures such that the aspect ratios no longer match
+		// This will be fine for some cases, especially if the common case where the mismatched texture is a small dummy texture with a constant color
+		MipBias = FMath::Min(MipBiasX, MipBiasY);
 
 		MaxLevel = FMath::Max<uint32>(MaxLevel, ProducerDesc.MaxLevel + MipBias);
 	}
