Fix complete shadow map snapping/stability and alignment when copying from engine shadow maps
- Stick to integer downsample factors. (Power-of-two is theoretically necessary but too harsh here.)
- Scale the stretch blit based on the inner rectangle (w/o border); clamp out of range border pixels when downsampling
  - Since the border is a static size in pixels, lower-resolution shadow maps technically cover more world area when their border pixels are included

diff --git a/Engine/Shaders/Private/VirtualShadowMaps/CopyCompleteShadows.usf b/Engine/Shaders/Private/VirtualShadowMaps/CopyCompleteShadows.usf
--- a/Engine/Shaders/Private/VirtualShadowMaps/CopyCompleteShadows.usf
+++ b/Engine/Shaders/Private/VirtualShadowMaps/CopyCompleteShadows.usf
@@ -4,7 +4,7 @@
 
 Texture2D<float> SourceBuffer;
 float4 SourceScaleOffset;
-float4 DestInvTargetExtent;
+int4 SourceMinMax;
 
 void CopyToCompleteShadowMapPS(
 	in float4 SvPosition	: SV_Position,
@@ -12,7 +12,8 @@ void CopyToCompleteShadowMapPS(
 )
 {
 	float2 DestPos = SvPosition.xy;
-	float2 SourcePos = (DestPos * SourceScaleOffset.xy + SourceScaleOffset.zw);
+	int2 SourcePos = int2(floor(DestPos * SourceScaleOffset.xy + SourceScaleOffset.zw));
+	SourcePos = clamp(SourcePos, SourceMinMax.xy, SourceMinMax.zw - 1);
 	float Depth = SourceBuffer.Load(int3(SourcePos, 0));
 	OutDepth = Depth;
 }
diff --git a/Engine/Source/Runtime/Renderer/Private/ShadowDepthRendering.cpp b/Engine/Source/Runtime/Renderer/Private/ShadowDepthRendering.cpp
--- a/Engine/Source/Runtime/Renderer/Private/ShadowDepthRendering.cpp
+++ b/Engine/Source/Runtime/Renderer/Private/ShadowDepthRendering.cpp
@@ -1492,6 +1492,7 @@ class FCopyToCompleteShadowMapPS : public FGlobalShader
 
 	BEGIN_SHADER_PARAMETER_STRUCT(FParameters, )
 		SHADER_PARAMETER(FVector4, SourceScaleOffset)
+		SHADER_PARAMETER(FIntVector4, SourceMinMax)
 		SHADER_PARAMETER_RDG_TEXTURE(Texture2D<uint>, SourceBuffer)
 		RENDER_TARGET_BINDING_SLOTS()
 	END_SHADER_PARAMETER_STRUCT()
@@ -1507,20 +1508,23 @@ static void CopyToCompleteShadowMap(
 	FRDGBuilder& GraphBuilder,
 	const FRDGTextureRef SourceBuffer,
 	const FRDGTextureRef DestBuffer,
-	const FIntRect& SourceRect,
-	const FIntRect& DestRect,
+	const FIntRect& SourceInnerRect,
+	const FIntRect& SourceOuterRect,
+	const FIntRect& DestInnerRect,
+	const FIntRect& DestOuterRect,
 	ERenderTargetLoadAction LoadAction
 	)
 {
 	FVector4 SourceScaleOffset;
-	SourceScaleOffset.X = static_cast<float>(SourceRect.Width() ) / DestRect.Width();
-	SourceScaleOffset.Y = static_cast<float>(SourceRect.Height()) / DestRect.Height();
-	SourceScaleOffset.Z = SourceRect.Min.X - (SourceScaleOffset.X * DestRect.Min.X);
-	SourceScaleOffset.W = SourceRect.Min.Y - (SourceScaleOffset.Y * DestRect.Min.Y);
+	SourceScaleOffset.X = static_cast<float>(SourceInnerRect.Width() ) / DestInnerRect.Width();
+	SourceScaleOffset.Y = static_cast<float>(SourceInnerRect.Height()) / DestInnerRect.Height();
+	SourceScaleOffset.Z = SourceInnerRect.Min.X - (SourceScaleOffset.X * DestInnerRect.Min.X);
+	SourceScaleOffset.W = SourceInnerRect.Min.Y - (SourceScaleOffset.Y * DestInnerRect.Min.Y);
 
 	auto* PassParameters = GraphBuilder.AllocParameters<FCopyToCompleteShadowMapPS::FParameters>();
 	PassParameters->SourceBuffer = SourceBuffer;
 	PassParameters->SourceScaleOffset = SourceScaleOffset;
+	PassParameters->SourceMinMax = FIntVector4(SourceOuterRect.Min.X, SourceOuterRect.Min.Y, SourceOuterRect.Max.X, SourceOuterRect.Max.Y);
 	PassParameters->RenderTargets.DepthStencil = FDepthStencilBinding(DestBuffer, LoadAction, FExclusiveDepthStencil::DepthWrite_StencilNop);
 
 	auto ShaderMap = GetGlobalShaderMap(GMaxRHIFeatureLevel);
@@ -1532,7 +1536,7 @@ static void CopyToCompleteShadowMap(
 		RDG_EVENT_NAME("CopyToCompleteShadowMap"),
 		PixelShader,
 		PassParameters,
-		DestRect,
+		DestOuterRect,
 		nullptr,
 		nullptr,
 		TStaticDepthStencilState<true, CF_Always>::GetRHI()
@@ -1679,9 +1683,14 @@ void FSceneRenderer::RenderShadowDepthMapAtlases(FRDGBuilder& GraphBuilder)
 				if (SourceShadowInfo && SourceShadowInfo->RenderTargets.DepthTarget != nullptr)
 				{
 					FRDGTextureRef SourceShadowMap = GraphBuilder.RegisterExternalTexture(SourceShadowInfo->RenderTargets.DepthTarget, TEXT("SourceDepthBuffer"));
-					const FIntRect SourceRect = SourceShadowInfo->GetOuterViewRect();
-					const FIntRect DestRect = ProjectedShadowInfo->GetOuterViewRect();
-					CopyToCompleteShadowMap(GraphBuilder, SourceShadowMap, DestShadowMap, SourceRect, DestRect, bCleared ? ERenderTargetLoadAction::ELoad : ERenderTargetLoadAction::EClear);
+					CopyToCompleteShadowMap(GraphBuilder,
+						SourceShadowMap,
+						DestShadowMap,
+						SourceShadowInfo->GetInnerViewRect(),
+						SourceShadowInfo->GetOuterViewRect(),
+						ProjectedShadowInfo->GetInnerViewRect(),
+						ProjectedShadowInfo->GetOuterViewRect(),
+						bCleared ? ERenderTargetLoadAction::ELoad : ERenderTargetLoadAction::EClear);
 					bCleared = true;
 				}
 			}
diff --git a/Engine/Source/Runtime/Renderer/Private/ShadowSetup.cpp b/Engine/Source/Runtime/Renderer/Private/ShadowSetup.cpp
--- a/Engine/Source/Runtime/Renderer/Private/ShadowSetup.cpp
+++ b/Engine/Source/Runtime/Renderer/Private/ShadowSetup.cpp
@@ -248,7 +248,8 @@ FAutoConsoleVariableRef CVarUseOctreeForShadowCulling(
 static TAutoConsoleVariable<int32> CVarCompleteShadowMapResolution(
 	TEXT("r.Shadow.CompleteShadowMapResolution"),
 	256,
-	TEXT("(Max)Resolution in texels of the complete shadows maps. These are generated for each whole scene shadows as well as CSMs when there is a Virtual Shadow Map that captures the Nanite geometry. 0 = disabled, in which case Nanite geometry may not show up in various indirect lighting."),
+	TEXT("(Max)Resolution in texels of the complete shadows maps. These are generated for each whole scene shadows as well as CSMs when there is a Virtual Shadow Map that captures the Nanite geometry.")
+	TEXT("0 = disabled, in which case Nanite geometry may not show up in various indirect lighting."),
 	ECVF_RenderThreadSafe
 );
 
@@ -4120,6 +4121,7 @@ void FSceneRenderer::AddViewDependentWholeSceneShadowsForView(
 	const bool bDirectionalLight = LightSceneInfo.Proxy->GetLightType() == LightType_Directional;
 	// Unmodified UE does not use this, so the virtual SM ignores this path for the time being (more likely to be removed).
 	const bool bNeedsVirtualShadowMap = bNaniteEnabled && EnableVirtualSMCVar->GetValueOnRenderThread() != 0 && bDirectionalLight;
+	const bool bNeedsCompleteShadowMap = bNeedsVirtualShadowMap && CVarCompleteShadowMapResolution.GetValueOnRenderThread() > 0;
 
 	// Allow each view to create a whole scene view dependent shadow
 	for (int32 ViewIndex = 0; ViewIndex < Views.Num(); ViewIndex++)
@@ -4179,11 +4181,27 @@ void FSceneRenderer::AddViewDependentWholeSceneShadowsForView(
 					MaxShadowCascadeDistance = FMath::Max(MaxShadowCascadeDistance, ProjectedShadowInitializer.CascadeSettings.SplitFar);
 
 					uint32 ShadowBorder = NeedsUnatlasedCSMDepthsWorkaround( FeatureLevel ) ? 0 : SHADOW_BORDER;
-
+					
 					const int32 MaxCSMResolution = GetCachedScalabilityCVars().MaxCSMShadowResolution;
-					FIntPoint ShadowBufferResolution = FIntPoint( FMath::Clamp( MaxCSMResolution, 1, (int32)GMaxShadowDepthBufferSizeX ) - ShadowBorder * 2,
-														          FMath::Clamp( MaxCSMResolution, 1, (int32)GMaxShadowDepthBufferSizeY ) - ShadowBorder * 2 );
-						
+					const int32 MinCSMResolution = 32;
+					FIntPoint ShadowBufferResolution = FIntPoint( FMath::Clamp( MaxCSMResolution, MinCSMResolution, (int32)GMaxShadowDepthBufferSizeX ) - ShadowBorder * 2,
+														          FMath::Clamp( MaxCSMResolution, MinCSMResolution, (int32)GMaxShadowDepthBufferSizeY ) - ShadowBorder * 2 );
+					
+					// "Complete" shadow map (one with lower resolution and containing both nanite and non-nanite)					
+					int32 CompleteShadowMapResolution = CVarCompleteShadowMapResolution.GetValueOnRenderThread() - ShadowBorder * 2;
+					// NOTE: RHI requires GMaxShadowDepthBufferSizeX >= GMaxShadowDepthBufferSizeY (in practice, they are equal)
+					CompleteShadowMapResolution = FMath::Clamp(CompleteShadowMapResolution, MinCSMResolution, ShadowBufferResolution.X);
+
+					int32 CompleteResolutionFactor = 1;
+					if (bNeedsCompleteShadowMap)
+					{
+						// Note that while the static borders guarantee we're not going to have a perfect pow2 reduction,
+						// sticking to integer factors keeps things pretty stable in practice.
+						CompleteResolutionFactor = FMath::CeilToInt(static_cast<float>(ShadowBufferResolution.X) / CompleteShadowMapResolution);
+					}
+					FIntPoint CompleteShadowBufferResolution(ShadowBufferResolution.X / CompleteResolutionFactor,
+														     ShadowBufferResolution.Y / CompleteResolutionFactor);
+
 					// Create the projected shadow info.
 					FProjectedShadowInfo* ProjectedShadowInfo = new(FMemStack::Get(), 1, 16) FProjectedShadowInfo;
 					ProjectedShadowInfo->SetupWholeSceneProjection(
@@ -4192,8 +4210,9 @@ void FSceneRenderer::AddViewDependentWholeSceneShadowsForView(
 						ProjectedShadowInitializer,
 						ShadowBufferResolution.X,
 						ShadowBufferResolution.Y,
-						ShadowBufferResolution.X,
-						ShadowBufferResolution.Y,
+						// If we are using complete shadow maps, We want both the complete and regular shadow maps to stay snapped in world space
+						CompleteShadowBufferResolution.X,
+						CompleteShadowBufferResolution.Y,
 						ShadowBorder
 					);
 					ProjectedShadowInfo->FadeAlphas = FadeAlphas;
@@ -4203,36 +4222,27 @@ void FSceneRenderer::AddViewDependentWholeSceneShadowsForView(
 					VisibleLightInfo.AllProjectedShadows.Add(ProjectedShadowInfo);
 					ShadowInfos.Add(ProjectedShadowInfo);
 
-					// If we have a virtual shadow map, disable nanite rendering into the regular shadow map or else we'd get double-shadowing
 					if (bNeedsVirtualShadowMap)
 					{
+						// If we have a virtual shadow map, disable nanite rendering into the regular shadow map or else we'd get double-shadowing
 						ProjectedShadowInfo->bNaniteGeometry = false;
-
-						// Also set up "complete" shadow map (one with lower resolution and containing both nanite and non-nanite)
-						const int32 CompleteShadowMapSize = CVarCompleteShadowMapResolution.GetValueOnRenderThread() - ShadowBorder * 2;
-						if (CompleteShadowMapSize > 0)
+						
+						if (bNeedsCompleteShadowMap)
 						{
 							// Create the projected shadow info.
 							FProjectedShadowInfo* CompleteProjectedShadowInfo = new(FMemStack::Get(), 1, 16) FProjectedShadowInfo;
 							// Add to remember-to-call-dtor list
 							VisibleLightInfo.MemStackProjectedShadows.Add(CompleteProjectedShadowInfo);
-
-							// Rescale size to fit whole virtual SM but keeping aspect ratio
-							int32 CompleteShadowMapSizeX = ShadowBufferResolution.X >= ShadowBufferResolution.Y ? CompleteShadowMapSize : (CompleteShadowMapSize * ShadowBufferResolution.X) / ShadowBufferResolution.Y;
-							int32 CompleteShadowMapSizeY = ShadowBufferResolution.Y >= ShadowBufferResolution.X ? CompleteShadowMapSize : (CompleteShadowMapSize * ShadowBufferResolution.Y) / ShadowBufferResolution.X;
-
 							// Set up projection as per normal
 							CompleteProjectedShadowInfo->SetupWholeSceneProjection(
 								&LightSceneInfo,
 								&View,
 								ProjectedShadowInitializer,
-								CompleteShadowMapSizeX,
-								CompleteShadowMapSizeY,
-								// Snap consistently with the original resolution shadow map
-								ProjectedShadowInfo->ResolutionX,
-								ProjectedShadowInfo->ResolutionY,
-								ShadowBorder
-								);
+								CompleteShadowBufferResolution.X,
+								CompleteShadowBufferResolution.Y,
+								CompleteShadowBufferResolution.X,
+								CompleteShadowBufferResolution.Y,
+								ShadowBorder);
 
 							CompleteProjectedShadowInfo->bCompleteShadowMap = true;
 							CompleteProjectedShadowInfo->bIncludeInScreenSpaceShadowMask = false;
