Fixed artifact in tiled deferred lighting.

diff --git a/Engine/Source/Runtime/Renderer/Private/TiledDeferredLightRendering.cpp b/Engine/Source/Runtime/Renderer/Private/TiledDeferredLightRendering.cpp
--- a/Engine/Source/Runtime/Renderer/Private/TiledDeferredLightRendering.cpp
+++ b/Engine/Source/Runtime/Renderer/Private/TiledDeferredLightRendering.cpp
@@ -357,8 +357,6 @@ void FDeferredShadingSceneRenderer::RenderTiledDeferredLighting(
 		INC_DWORD_STAT_BY(STAT_NumLightsUsingSimpleTiledDeferred, SimpleLights.InstanceData.Num());
 		SCOPE_CYCLE_COUNTER(STAT_DirectLightRenderingTime);
 
-		FRDGTextureRef& SceneColorTexture = SceneTextures.Color.Target;
-
 		// Determine how many compute shader passes will be needed to process all the lights
 		const int32 NumPassesNeeded = FMath::DivideAndRoundUp(NumLightsToRender, GMaxNumTiledDeferredLights);
 		for (int32 PassIndex = 0; PassIndex < NumPassesNeeded; PassIndex++)
@@ -383,7 +381,7 @@ void FDeferredShadingSceneRenderer::RenderTiledDeferredLighting(
 				const FViewInfo& View = Views[ViewIndex];
 
 				FTiledDeferredLightingParameters* PassParameters = GraphBuilder.AllocParameters<FTiledDeferredLightingParameters>();
-				PassParameters->SceneColorInput = SceneColorTexture;
+				PassParameters->SceneColorInput = SceneTextures.Color.Target;
 				PassParameters->SceneColorOutput = SceneColorOutputTexture;
 				PassParameters->SceneTextures = SceneTextures.UniformBuffer;
 
@@ -391,7 +389,7 @@ void FDeferredShadingSceneRenderer::RenderTiledDeferredLighting(
 					RDG_EVENT_NAME("TiledDeferredLighting"),
 					PassParameters,
 					ERDGPassFlags::Compute,
-					[&View, ViewIndex, ViewCount, &SortedLights, TiledDeferredLightsStart, TiledDeferredLightsEnd, &SimpleLights, StartIndex, NumThisPass, SceneColorTexture, SceneColorOutputTexture](FRHIComputeCommandList& RHICmdList)
+					[&View, ViewIndex, ViewCount, &SortedLights, TiledDeferredLightsStart, TiledDeferredLightsEnd, &SimpleLights, StartIndex, NumThisPass, SceneColorTexture = SceneTextures.Color.Target, SceneColorOutputTexture](FRHIComputeCommandList& RHICmdList)
 				{
 					if (View.Family->EngineShowFlags.VisualizeLightCulling)
 					{
@@ -404,10 +402,10 @@ void FDeferredShadingSceneRenderer::RenderTiledDeferredLighting(
 				});
 			}
 
-			SceneColorTexture = SceneColorOutputTexture;
+			SceneTextures.Color = SceneColorOutputTexture;
 		}
 
 		FSceneRenderTargets& SceneContext = FSceneRenderTargets::Get();
-		ConvertToExternalTexture(GraphBuilder, SceneColorTexture, SceneContext.GetSceneColor());
+		ConvertToExternalTexture(GraphBuilder, SceneTextures.Color.Target, SceneContext.GetSceneColor());
 	}
 }
\ No newline at end of file
