Port mobile base pass to RDG.

diff --git a/Engine/Source/Runtime/Engine/Classes/Engine/EngineTypes.h b/Engine/Source/Runtime/Engine/Classes/Engine/EngineTypes.h
--- a/Engine/Source/Runtime/Engine/Classes/Engine/EngineTypes.h
+++ b/Engine/Source/Runtime/Engine/Classes/Engine/EngineTypes.h
@@ -359,7 +359,9 @@ enum ESceneCaptureSource
 	SCS_Normal UMETA(DisplayName="Normal in RGB (Deferred Renderer only)"),
 	SCS_BaseColor UMETA(DisplayName = "BaseColor in RGB (Deferred Renderer only)"),
 	SCS_FinalColorHDR UMETA(DisplayName = "Final Color (HDR) in Linear sRGB gamut"),
-	SCS_FinalToneCurveHDR UMETA(DisplayName = "Final Color (with tone curve) in Linear sRGB gamut")
+	SCS_FinalToneCurveHDR UMETA(DisplayName = "Final Color (with tone curve) in Linear sRGB gamut"),
+
+	SCS_MAX
 };
 
 /** Specifies how scene captures are composited into render buffers */
diff --git a/Engine/Source/Runtime/Engine/Public/SceneView.h b/Engine/Source/Runtime/Engine/Public/SceneView.h
--- a/Engine/Source/Runtime/Engine/Public/SceneView.h
+++ b/Engine/Source/Runtime/Engine/Public/SceneView.h
@@ -943,12 +943,6 @@ public:
 	/** The uniform buffer for the view's parameters. This is only initialized in the rendering thread's copies of the FSceneView. */
 	TUniformBufferRef<FViewUniformShaderParameters> ViewUniformBuffer;
 
-	/** Mobile Directional Lighting uniform buffers, one for each lighting channel 
-	  * The first is used for primitives with no lighting channels set.
-	  * Only initialized in the rendering thread's copies of the FSceneView.
-	  */
-	TUniformBufferRef<FMobileDirectionalLightShaderParameters> MobileDirectionalLightUniformBuffers[NUM_LIGHTING_CHANNELS+1];
-
 private:
 	/** During GetDynamicMeshElements this will be the correct cull volume for shadow stuff */
 	const FConvexVolume* DynamicMeshElementsShadowCullFrustum;
diff --git a/Engine/Source/Runtime/RenderCore/Private/RenderGraphPass.cpp b/Engine/Source/Runtime/RenderCore/Private/RenderGraphPass.cpp
--- a/Engine/Source/Runtime/RenderCore/Private/RenderGraphPass.cpp
+++ b/Engine/Source/Runtime/RenderCore/Private/RenderGraphPass.cpp
@@ -101,6 +101,9 @@ FRHIRenderPassInfo FRDGParameterStruct::GetRenderPassInfo() const
 	RenderPassInfo.ResolveParameters.DestAccessFinal = ERHIAccess::ResolveDst;
 	RenderPassInfo.NumOcclusionQueries = RenderTargets.NumOcclusionQueries;
 	RenderPassInfo.bOcclusionQueries = RenderTargets.NumOcclusionQueries > 0;
+	RenderPassInfo.SubpassHint = RenderTargets.SubpassHint;
+	RenderPassInfo.MultiViewCount = RenderTargets.MultiViewCount;
+	RenderPassInfo.FoveationTexture = RenderTargets.FoveationTexture ? RenderTargets.FoveationTexture->GetRHI() : nullptr;
 
 	return RenderPassInfo;
 }
diff --git a/Engine/Source/Runtime/RenderCore/Public/ShaderParameterMacros.h b/Engine/Source/Runtime/RenderCore/Public/ShaderParameterMacros.h
--- a/Engine/Source/Runtime/RenderCore/Public/ShaderParameterMacros.h
+++ b/Engine/Source/Runtime/RenderCore/Public/ShaderParameterMacros.h
@@ -440,6 +440,9 @@ struct alignas(SHADER_PARAMETER_STRUCT_ALIGNMENT) FRenderTargetBindingSlots
 	FDepthStencilBinding DepthStencil;
 	FResolveRect ResolveRect;
 	uint32 NumOcclusionQueries = 0;
+	ESubpassHint SubpassHint = ESubpassHint::None;
+	uint8 MultiViewCount = 0;
+	FRDGTexture* FoveationTexture = nullptr;
 
 	/** Accessors for regular output to simplify the syntax to:
 	 *
diff --git a/Engine/Source/Runtime/Renderer/Private/CompositionLighting/PostProcessMeshDecals.cpp b/Engine/Source/Runtime/Renderer/Private/CompositionLighting/PostProcessMeshDecals.cpp
--- a/Engine/Source/Runtime/Renderer/Private/CompositionLighting/PostProcessMeshDecals.cpp
+++ b/Engine/Source/Runtime/Renderer/Private/CompositionLighting/PostProcessMeshDecals.cpp
@@ -449,34 +449,39 @@ void RenderMeshDecals(
 	}
 }
 
-void RenderMeshDecalsMobile(FRHICommandList& RHICmdList, const FViewInfo& View)
+void RenderMeshDecalsMobile(FRDGBuilder& GraphBuilder, FRenderTargetBindingSlots& BasePassRenderTargets, const FViewInfo& View)
 {
-	SCOPED_DRAW_EVENT(RHICmdList, MeshDecals);
-
 	FDecalRenderingCommon::ERenderTargetMode RenderTargetMode = IsMobileDeferredShadingEnabled(View.GetShaderPlatform()) ? 
 		FDecalRenderingCommon::RTM_SceneColorAndGBufferWithNormal : 
 		FDecalRenderingCommon::RTM_SceneColor;
 
-	FGraphicsPipelineStateInitializer GraphicsPSOInit;
-	RHICmdList.SetViewport(View.ViewRect.Min.X, View.ViewRect.Min.Y, 0, View.ViewRect.Max.X, View.ViewRect.Max.Y, 1);
-	RHICmdList.ApplyCachedRenderTargets(GraphicsPSOInit);
+	auto* PassParameters = GraphBuilder.AllocParameters<FRenderTargetParameters>();
+	PassParameters->RenderTargets = BasePassRenderTargets;
 
-	DrawDynamicMeshPass(View, RHICmdList, [&View, RenderTargetMode](FDynamicPassMeshDrawListContext* DynamicMeshPassContext)
+	GraphBuilder.AddPass(RDG_EVENT_NAME("MeshDecalsPass"), PassParameters, ERDGPassFlags::Raster,
+		[RenderTargetMode, &View](FRHICommandList& RHICmdList)
 	{
-		FMeshDecalMeshProcessor PassMeshProcessor(
-			View.Family->Scene->GetRenderScene(),
-			&View,
-			DRS_Mobile,
-			RenderTargetMode,
-			DynamicMeshPassContext);
-
-		for (int32 MeshBatchIndex = 0; MeshBatchIndex < View.MeshDecalBatches.Num(); ++MeshBatchIndex)
+		FGraphicsPipelineStateInitializer GraphicsPSOInit;
+		RHICmdList.SetViewport(View.ViewRect.Min.X, View.ViewRect.Min.Y, 0, View.ViewRect.Max.X, View.ViewRect.Max.Y, 1);
+		RHICmdList.ApplyCachedRenderTargets(GraphicsPSOInit);
+
+		DrawDynamicMeshPass(View, RHICmdList, [&View, RenderTargetMode](FDynamicPassMeshDrawListContext* DynamicMeshPassContext)
 		{
-			const FMeshBatch* Mesh = View.MeshDecalBatches[MeshBatchIndex].Mesh;
-			const FPrimitiveSceneProxy* PrimitiveSceneProxy = View.MeshDecalBatches[MeshBatchIndex].Proxy;
-			const uint64 DefaultBatchElementMask = ~0ull;
+			FMeshDecalMeshProcessor PassMeshProcessor(
+				View.Family->Scene->GetRenderScene(),
+				&View,
+				DRS_Mobile,
+				RenderTargetMode,
+				DynamicMeshPassContext);
 
-			PassMeshProcessor.AddMeshBatch(*Mesh, DefaultBatchElementMask, PrimitiveSceneProxy);
-		}
-	}, true);
+			for (int32 MeshBatchIndex = 0; MeshBatchIndex < View.MeshDecalBatches.Num(); ++MeshBatchIndex)
+			{
+				const FMeshBatch* Mesh = View.MeshDecalBatches[MeshBatchIndex].Mesh;
+				const FPrimitiveSceneProxy* PrimitiveSceneProxy = View.MeshDecalBatches[MeshBatchIndex].Proxy;
+				const uint64 DefaultBatchElementMask = ~0ull;
+
+				PassMeshProcessor.AddMeshBatch(*Mesh, DefaultBatchElementMask, PrimitiveSceneProxy);
+			}
+		}, true);
+	});
 }
\ No newline at end of file
diff --git a/Engine/Source/Runtime/Renderer/Private/DepthRendering.cpp b/Engine/Source/Runtime/Renderer/Private/DepthRendering.cpp
--- a/Engine/Source/Runtime/Renderer/Private/DepthRendering.cpp
+++ b/Engine/Source/Runtime/Renderer/Private/DepthRendering.cpp
@@ -592,37 +592,52 @@ void FDeferredShadingSceneRenderer::RenderPrePass(FRDGBuilder& GraphBuilder, FRD
 	}
 }
 
-void FMobileSceneRenderer::RenderPrePass(FRHICommandListImmediate& RHICmdList)
+void FMobileSceneRenderer::RenderPrePass(FRDGBuilder& GraphBuilder, FRenderTargetBindingSlots& BasePassRenderTargets, TFunction<void(FRenderTargetBindingSlots&)> UpdateRenderTargetsLoadAction)
 {
-	check(!RHICmdList.IsOutsideRenderPass());
-
 	SCOPED_NAMED_EVENT(FMobileSceneRenderer_RenderPrePass, FColor::Emerald);
-	SCOPED_DRAW_EVENT(RHICmdList, MobileRenderPrePass);
+	RDG_EVENT_SCOPE(GraphBuilder, "MobileRenderPrePass");
 
 	SCOPE_CYCLE_COUNTER(STAT_DepthDrawTime);
 	CSV_SCOPED_TIMING_STAT_EXCLUSIVE(RenderPrePass);
-	SCOPED_GPU_STAT(RHICmdList, Prepass);
+	RDG_GPU_STAT_SCOPE(GraphBuilder, Prepass);
 
 	// Draw a depth pass to avoid overdraw in the other passes.
 	// Mobile only does MaskedOnly DepthPass for the moment
 	if (Scene->EarlyZPassMode == DDM_MaskedOnly)
 	{
+		bool bAnyPassesAdded = false;
+
 		for (int32 ViewIndex = 0; ViewIndex < Views.Num(); ViewIndex++)
 		{
 			const FViewInfo& View = Views[ViewIndex];
 
-			SCOPED_GPU_MASK(RHICmdList, !View.IsInstancedStereoPass() ? View.GPUMask : (Views[0].GPUMask | Views[1].GPUMask));
-			SCOPED_CONDITIONAL_DRAW_EVENTF(RHICmdList, EventView, Views.Num() > 1, TEXT("View%d"), ViewIndex);
+			RDG_GPU_MASK_SCOPE(GraphBuilder, !View.IsInstancedStereoPass() ? View.GPUMask : (Views[0].GPUMask | Views[1].GPUMask));
+			RDG_EVENT_SCOPE_CONDITIONAL(GraphBuilder, Views.Num() > 1, "View%d", ViewIndex);
+
 			if (!View.ShouldRenderView())
 			{
 				continue;
 			}
+			
+			auto* PassParameters = GraphBuilder.AllocParameters<FRenderTargetParameters>();
+			PassParameters->RenderTargets = BasePassRenderTargets;
 
-			Scene->UniformBuffers.UpdateViewUniformBuffer(View);
+			GraphBuilder.AddPass(RDG_EVENT_NAME("RenderPrePass"), PassParameters, ERDGPassFlags::Raster,
+				[this, &View, PassParameters](FRHICommandListImmediate& RHICmdList)
+			{
+				Scene->UniformBuffers.UpdateViewUniformBuffer(View);
 
-			SetStereoViewport(RHICmdList, View);
+				SetStereoViewport(RHICmdList, View);
 
-			View.ParallelMeshDrawCommandPasses[EMeshPass::DepthPass].DispatchDraw(nullptr, RHICmdList);
+				View.ParallelMeshDrawCommandPasses[EMeshPass::DepthPass].DispatchDraw(nullptr, RHICmdList);
+			});
+
+			bAnyPassesAdded = true;
+		}
+
+		if (bAnyPassesAdded)
+		{
+			UpdateRenderTargetsLoadAction(BasePassRenderTargets);
 		}
 	}
 }
diff --git a/Engine/Source/Runtime/Renderer/Private/MobileBasePassRendering.cpp b/Engine/Source/Runtime/Renderer/Private/MobileBasePassRendering.cpp
--- a/Engine/Source/Runtime/Renderer/Private/MobileBasePassRendering.cpp
+++ b/Engine/Source/Runtime/Renderer/Private/MobileBasePassRendering.cpp
@@ -302,39 +302,57 @@ void SetupMobileSkyReflectionUniformParameters(FSkyLightSceneProxy* SkyLight, FM
 	Parameters.TextureSampler = CaptureTexture->SamplerStateRHI;
 }
 
-void FMobileSceneRenderer::RenderMobileBasePass(FRHICommandListImmediate& RHICmdList, const TArrayView<const FViewInfo*> PassViews)
+void FMobileSceneRenderer::RenderMobileBasePass(FRDGBuilder& GraphBuilder, FRenderTargetBindingSlots& BasePassRenderTargets, TFunction<void(FRenderTargetBindingSlots &)> UpdateRenderTargetsLoadAction, const TArrayView<const FViewInfo*> PassViews)
 {
 	CSV_SCOPED_TIMING_STAT_EXCLUSIVE(RenderBasePass);
-	SCOPED_DRAW_EVENT(RHICmdList, MobileBasePass);
+	RDG_EVENT_SCOPE(GraphBuilder, "MobileBasePass");
 	SCOPE_CYCLE_COUNTER(STAT_BasePassDrawTime);
-	SCOPED_GPU_STAT(RHICmdList, Basepass);
+	RDG_GPU_STAT_SCOPE(GraphBuilder, Basepass);
+
+	bool bAnyPassesAdded = false;
 
 	for (int32 ViewIndex = 0; ViewIndex < PassViews.Num(); ViewIndex++)
 	{
-		SCOPED_CONDITIONAL_DRAW_EVENTF(RHICmdList, EventView, Views.Num() > 1, TEXT("View%d"), ViewIndex);
+		RDG_EVENT_SCOPE_CONDITIONAL(GraphBuilder, Views.Num() > 1, "View%d", ViewIndex);
 		const FViewInfo& View = *PassViews[ViewIndex];
 		if (!View.ShouldRenderView())
 		{
 			continue;
 		}
 
-		if (Scene->UniformBuffers.UpdateViewUniformBuffer(View))
+		auto* OpaqueBasePassParameters = GraphBuilder.AllocParameters<FMobileBasePassParameters>();
+		OpaqueBasePassParameters->RenderTargets = BasePassRenderTargets;
+		OpaqueBasePassParameters->MobileBasePass = Scene->UniformBuffers.MobileOpaqueBasePassUniformBuffer;
+
+		GraphBuilder.AddPass(RDG_EVENT_NAME("RenderOpaqueBasePass"), OpaqueBasePassParameters, ERDGPassFlags::Raster,
+			[this, &View, OpaqueBasePassParameters](FRHICommandListImmediate& RHICmdList)
 		{
-			UpdateOpaqueBasePassUniformBuffer(RHICmdList, View);
-			UpdateDirectionalLightUniformBuffers(RHICmdList, View);
-		}
+			if (Scene->UniformBuffers.UpdateViewUniformBuffer(View))
+			{
+				UpdateOpaqueBasePassUniformBuffer(RHICmdList, View);
+				UpdateDirectionalLightUniformBuffers(RHICmdList, View);
+			}
+
+			RHICmdList.SetViewport(View.ViewRect.Min.X, View.ViewRect.Min.Y, 0, View.ViewRect.Max.X, View.ViewRect.Max.Y, 1);
+			View.ParallelMeshDrawCommandPasses[EMeshPass::BasePass].DispatchDraw(nullptr, RHICmdList);
 		
-		RHICmdList.SetViewport(View.ViewRect.Min.X, View.ViewRect.Min.Y, 0, View.ViewRect.Max.X, View.ViewRect.Max.Y, 1);
-		View.ParallelMeshDrawCommandPasses[EMeshPass::BasePass].DispatchDraw(nullptr, RHICmdList);
+			// editor primitives
+			{
+				FMeshPassProcessorRenderState DrawRenderState(View, Scene->UniformBuffers.MobileOpaqueBasePassUniformBuffer);
+				DrawRenderState.SetBlendState(TStaticBlendStateWriteMask<CW_RGBA>::GetRHI());
+				DrawRenderState.SetDepthStencilAccess(Scene->DefaultBasePassDepthStencilAccess);
+				DrawRenderState.SetDepthStencilState(TStaticDepthStencilState<true, CF_DepthNearOrEqual>::GetRHI());
+				RenderMobileEditorPrimitives(RHICmdList, View, DrawRenderState);
+			}
 
-		// editor primitives
-		{
-			FMeshPassProcessorRenderState DrawRenderState(View, Scene->UniformBuffers.MobileOpaqueBasePassUniformBuffer);
-			DrawRenderState.SetBlendState(TStaticBlendStateWriteMask<CW_RGBA>::GetRHI());
-			DrawRenderState.SetDepthStencilAccess(Scene->DefaultBasePassDepthStencilAccess);
-			DrawRenderState.SetDepthStencilState(TStaticDepthStencilState<true, CF_DepthNearOrEqual>::GetRHI());
-			RenderMobileEditorPrimitives(RHICmdList, View, DrawRenderState);
-		}
+		});
+
+		bAnyPassesAdded = true;
+	}
+
+	if (bAnyPassesAdded)
+	{
+		UpdateRenderTargetsLoadAction(BasePassRenderTargets);
 	}
 }
 
diff --git a/Engine/Source/Runtime/Renderer/Private/MobileBasePassRendering.h b/Engine/Source/Runtime/Renderer/Private/MobileBasePassRendering.h
--- a/Engine/Source/Runtime/Renderer/Private/MobileBasePassRendering.h
+++ b/Engine/Source/Runtime/Renderer/Private/MobileBasePassRendering.h
@@ -37,6 +37,11 @@ BEGIN_GLOBAL_SHADER_PARAMETER_STRUCT(FMobileBasePassUniformParameters, )
 	SHADER_PARAMETER_SAMPLER(SamplerState, AmbientOcclusionSampler)
 END_GLOBAL_SHADER_PARAMETER_STRUCT()
 
+BEGIN_SHADER_PARAMETER_STRUCT(FMobileBasePassParameters, )
+	SHADER_PARAMETER_STRUCT_REF(FMobileBasePassUniformParameters, MobileBasePass)
+	RENDER_TARGET_BINDING_SLOTS()
+END_SHADER_PARAMETER_STRUCT()
+
 extern void SetupMobileBasePassUniformParameters(
 	FRHICommandListImmediate& RHICmdList,
 	const FViewInfo& View,
diff --git a/Engine/Source/Runtime/Renderer/Private/MobileDecalRendering.cpp b/Engine/Source/Runtime/Renderer/Private/MobileDecalRendering.cpp
--- a/Engine/Source/Runtime/Renderer/Private/MobileDecalRendering.cpp
+++ b/Engine/Source/Runtime/Renderer/Private/MobileDecalRendering.cpp
@@ -15,9 +15,14 @@
 #include "ScenePrivate.h"
 #include "DecalRenderingShared.h"
 
+BEGIN_SHADER_PARAMETER_STRUCT(FMobileDecalPassParameters, )
+	SHADER_PARAMETER_RDG_UNIFORM_BUFFER(FMobileSceneTextureUniformParameters, MobileSceneTextures)
+	RENDER_TARGET_BINDING_SLOTS()
+END_SHADER_PARAMETER_STRUCT()
+
 extern FRHIRasterizerState* GetDecalRasterizerState(EDecalRasterizerState DecalRasterizerState);
-extern void RenderMeshDecalsMobile(FRHICommandList& RHICmdList, const FViewInfo& View);
-void RenderDeferredDecalsMobile(FRHICommandList& RHICmdList, const FScene& Scene, const FViewInfo& View);
+extern void RenderMeshDecalsMobile(FRDGBuilder& GraphBuilder, FRenderTargetBindingSlots& BasePassRenderTargets, const FViewInfo& View);
+void RenderDeferredDecalsMobile(FRDGBuilder& GraphBuilder, FRenderTargetBindingSlots& BasePassRenderTargets, TRDGUniformBufferRef<FMobileSceneTextureUniformParameters> MobileSceneTextures, const FScene& Scene, const FViewInfo& View);
 
 FRHIBlendState* MobileForward_GetDecalBlendState(EDecalBlendMode DecalBlendMode)
 {
@@ -182,7 +187,7 @@ FRHIBlendState* MobileDeferred_GetDecalBlendState(EDecalBlendMode DecalBlendMode
 	return TStaticBlendState<>::GetRHI(); 
 }
 
-void FMobileSceneRenderer::RenderDecals(FRHICommandListImmediate& RHICmdList)
+void FMobileSceneRenderer::RenderDecals(FRDGBuilder& GraphBuilder, FRenderTargetBindingSlots& BasePassRenderTargets, TRDGUniformBufferRef<FMobileSceneTextureUniformParameters> MobileSceneTextures)
 {
 	if (!IsMobileHDR())
 	{
@@ -197,7 +202,7 @@ void FMobileSceneRenderer::RenderDecals(FRHICommandListImmediate& RHICmdList)
 		for (int32 ViewIndex = 0; ViewIndex < Views.Num(); ViewIndex++)
 		{
 			const FViewInfo& View = Views[ViewIndex];
-			RenderDeferredDecalsMobile(RHICmdList, *Scene, View);
+			RenderDeferredDecalsMobile(GraphBuilder, BasePassRenderTargets, MobileSceneTextures, *Scene, View);
 		}
 	}
 
@@ -207,81 +212,93 @@ void FMobileSceneRenderer::RenderDecals(FRHICommandListImmediate& RHICmdList)
 		const FViewInfo& View = Views[ViewIndex];
 		if (View.MeshDecalBatches.Num() > 0)
 		{
-			RenderMeshDecalsMobile(RHICmdList, View);
+			RenderMeshDecalsMobile(GraphBuilder, BasePassRenderTargets, View);
 		}
 	}
 }
 
-void RenderDeferredDecalsMobile(FRHICommandList& RHICmdList, const FScene& Scene, const FViewInfo& View)
+void RenderDeferredDecalsMobile(FRDGBuilder& GraphBuilder, FRenderTargetBindingSlots& BasePassRenderTargets, TRDGUniformBufferRef<FMobileSceneTextureUniformParameters> MobileSceneTextures, const FScene& Scene, const FViewInfo& View)
 {
-	const bool bDeferredShading = IsMobileDeferredShadingEnabled(View.GetShaderPlatform());
-
-	FGraphicsPipelineStateInitializer GraphicsPSOInit;
-	RHICmdList.ApplyCachedRenderTargets(GraphicsPSOInit);
-
-	FUniformBufferRHIRef PassUniformBuffer = CreateMobileSceneTextureUniformBuffer(RHICmdList);
-	FUniformBufferStaticBindings GlobalUniformBuffers(PassUniformBuffer);
-	SCOPED_UNIFORM_BUFFER_GLOBAL_BINDINGS(RHICmdList, GlobalUniformBuffers);
+	const uint32 DecalCount = Scene.Decals.Num();
+	int32 SortedDecalCount = 0;
+	FTransientDecalRenderDataList* SortedDecals = nullptr;
 
-	// Build a list of decals that need to be rendered for this view
-	FTransientDecalRenderDataList SortedDecals;
-	FDecalRendering::BuildVisibleDecalList(Scene, View, DRS_Mobile, &SortedDecals);
-	if (SortedDecals.Num())
+	if (DecalCount > 0)
 	{
-		SCOPED_DRAW_EVENT(RHICmdList, DeferredDecals);
-		INC_DWORD_STAT_BY(STAT_Decals, SortedDecals.Num());
+		// Build a list of decals that need to be rendered for this view
+		SortedDecals = GraphBuilder.AllocObject<FTransientDecalRenderDataList>();
+		FDecalRendering::BuildVisibleDecalList(Scene, View, DRS_Mobile, SortedDecals);
+		SortedDecalCount = SortedDecals->Num();
+		INC_DWORD_STAT_BY(STAT_Decals, SortedDecalCount);
+	}
 
-		RHICmdList.SetViewport(View.ViewRect.Min.X, View.ViewRect.Min.Y, 0, View.ViewRect.Max.X, View.ViewRect.Max.Y, 1);
-		RHICmdList.SetStreamSource(0, GetUnitCubeVertexBuffer(), 0);
+	if (SortedDecalCount > 0)
+	{
+		auto* PassParameters = GraphBuilder.AllocParameters<FMobileDecalPassParameters>();
+		PassParameters->RenderTargets = BasePassRenderTargets;
+		PassParameters->MobileSceneTextures = MobileSceneTextures;
 
-		for (int32 DecalIndex = 0, DecalCount = SortedDecals.Num(); DecalIndex < DecalCount; DecalIndex++)
+		GraphBuilder.AddPass(
+			RDG_EVENT_NAME("MobileDeferredDecalsPass"), PassParameters, ERDGPassFlags::Raster,
+			[&View, SortedDecalCount, SortedDecals](FRHICommandListImmediate& RHICmdList)
 		{
-			const FTransientDecalRenderData& DecalData = SortedDecals[DecalIndex];
-			const FDeferredDecalProxy& DecalProxy = *DecalData.DecalProxy;
-			const FMatrix ComponentToWorldMatrix = DecalProxy.ComponentTrans.ToMatrixWithScale();
-			const FMatrix FrustumComponentToClip = FDecalRendering::ComputeComponentToClipMatrix(View, ComponentToWorldMatrix);
-						
-			const float ConservativeRadius = DecalData.ConservativeRadius;
-			const bool bInsideDecal = ((FVector)View.ViewMatrices.GetViewOrigin() - ComponentToWorldMatrix.GetOrigin()).SizeSquared() < FMath::Square(ConservativeRadius * 1.05f + View.NearClippingDistance * 2.0f);
-			bool bReverseHanded = false;
-			{
-				// Account for the reversal of handedness caused by negative scale on the decal
-				const auto& Scale3d = DecalProxy.ComponentTrans.GetScale3D();
-				bReverseHanded = Scale3d[0] * Scale3d[1] * Scale3d[2] < 0.f;
-			}
-			EDecalRasterizerState DecalRasterizerState = FDecalRenderingCommon::ComputeDecalRasterizerState(bInsideDecal, bReverseHanded, View.bReverseCulling);
-			GraphicsPSOInit.RasterizerState = GetDecalRasterizerState(DecalRasterizerState);
+			const bool bDeferredShading = IsMobileDeferredShadingEnabled(View.GetShaderPlatform());
 
-			if (bInsideDecal)
-			{
-				GraphicsPSOInit.DepthStencilState = TStaticDepthStencilState<
-					false, CF_Always,
-					true, CF_Equal, SO_Keep, SO_Keep, SO_Keep,
-					false, CF_Always, SO_Keep, SO_Keep, SO_Keep,
-					GET_STENCIL_BIT_MASK(RECEIVE_DECAL, 1), 0x00>::GetRHI();
-			}
-			else
-			{
-				GraphicsPSOInit.DepthStencilState = TStaticDepthStencilState<
-					false, CF_DepthNearOrEqual,
-					true, CF_Equal, SO_Keep, SO_Keep, SO_Keep,
-					false, CF_Always, SO_Keep, SO_Keep, SO_Keep,
-					GET_STENCIL_BIT_MASK(RECEIVE_DECAL, 1), 0x00>::GetRHI();
-			}
-			
-			if (bDeferredShading)
-			{
-				GraphicsPSOInit.BlendState = MobileDeferred_GetDecalBlendState(DecalData.FinalDecalBlendMode, DecalData.bHasNormal);
-			}
-			else
+			FGraphicsPipelineStateInitializer GraphicsPSOInit;
+			RHICmdList.ApplyCachedRenderTargets(GraphicsPSOInit);
+
+			RHICmdList.SetViewport(View.ViewRect.Min.X, View.ViewRect.Min.Y, 0, View.ViewRect.Max.X, View.ViewRect.Max.Y, 1);
+			RHICmdList.SetStreamSource(0, GetUnitCubeVertexBuffer(), 0);
+
+			for (int32 DecalIndex = 0; DecalIndex < SortedDecalCount; DecalIndex++)
 			{
-				GraphicsPSOInit.BlendState = MobileForward_GetDecalBlendState(DecalData.FinalDecalBlendMode);
-			}
+				const FTransientDecalRenderData& DecalData = (*SortedDecals)[DecalIndex];
+				const FDeferredDecalProxy& DecalProxy = *DecalData.DecalProxy;
+				const FMatrix ComponentToWorldMatrix = DecalProxy.ComponentTrans.ToMatrixWithScale();
+				const FMatrix FrustumComponentToClip = FDecalRendering::ComputeComponentToClipMatrix(View, ComponentToWorldMatrix);
 
-			// Set shader params
-			FDecalRendering::SetShader(RHICmdList, GraphicsPSOInit, View, DecalData, DRS_Mobile, FrustumComponentToClip);
-			
-			RHICmdList.DrawIndexedPrimitive(GetUnitCubeIndexBuffer(), 0, 0, 8, 0, UE_ARRAY_COUNT(GCubeIndices) / 3, 1);
-		}
+				const float ConservativeRadius = DecalData.ConservativeRadius;
+				const bool bInsideDecal = ((FVector)View.ViewMatrices.GetViewOrigin() - ComponentToWorldMatrix.GetOrigin()).SizeSquared() < FMath::Square(ConservativeRadius * 1.05f + View.NearClippingDistance * 2.0f);
+				bool bReverseHanded = false;
+				{
+					// Account for the reversal of handedness caused by negative scale on the decal
+					const auto& Scale3d = DecalProxy.ComponentTrans.GetScale3D();
+					bReverseHanded = Scale3d[0] * Scale3d[1] * Scale3d[2] < 0.f;
+				}
+				EDecalRasterizerState DecalRasterizerState = FDecalRenderingCommon::ComputeDecalRasterizerState(bInsideDecal, bReverseHanded, View.bReverseCulling);
+				GraphicsPSOInit.RasterizerState = GetDecalRasterizerState(DecalRasterizerState);
+
+				if (bInsideDecal)
+				{
+					GraphicsPSOInit.DepthStencilState = TStaticDepthStencilState<
+						false, CF_Always,
+						true, CF_Equal, SO_Keep, SO_Keep, SO_Keep,
+						false, CF_Always, SO_Keep, SO_Keep, SO_Keep,
+						GET_STENCIL_BIT_MASK(RECEIVE_DECAL, 1), 0x00>::GetRHI();
+				}
+				else
+				{
+					GraphicsPSOInit.DepthStencilState = TStaticDepthStencilState<
+						false, CF_DepthNearOrEqual,
+						true, CF_Equal, SO_Keep, SO_Keep, SO_Keep,
+						false, CF_Always, SO_Keep, SO_Keep, SO_Keep,
+						GET_STENCIL_BIT_MASK(RECEIVE_DECAL, 1), 0x00>::GetRHI();
+				}
+
+				if (bDeferredShading)
+				{
+					GraphicsPSOInit.BlendState = MobileDeferred_GetDecalBlendState(DecalData.FinalDecalBlendMode, DecalData.bHasNormal);
+				}
+				else
+				{
+					GraphicsPSOInit.BlendState = MobileForward_GetDecalBlendState(DecalData.FinalDecalBlendMode);
+				}
+
+				// Set shader params
+				FDecalRendering::SetShader(RHICmdList, GraphicsPSOInit, View, DecalData, DRS_Mobile, FrustumComponentToClip);
+
+				RHICmdList.DrawIndexedPrimitive(GetUnitCubeIndexBuffer(), 0, 0, 8, 0, UE_ARRAY_COUNT(GCubeIndices) / 3, 1);
+			}
+		});
 	}
 }
diff --git a/Engine/Source/Runtime/Renderer/Private/MobileDeferredShadingPass.cpp b/Engine/Source/Runtime/Renderer/Private/MobileDeferredShadingPass.cpp
--- a/Engine/Source/Runtime/Renderer/Private/MobileDeferredShadingPass.cpp
+++ b/Engine/Source/Runtime/Renderer/Private/MobileDeferredShadingPass.cpp
@@ -23,6 +23,11 @@ static FAutoConsoleVariableRef CVarMobileUseLightStencilCulling(
 	ECVF_RenderThreadSafe
 );
 
+BEGIN_SHADER_PARAMETER_STRUCT(FMobileDeferredPassParameters, )
+	SHADER_PARAMETER_RDG_UNIFORM_BUFFER(FMobileSceneTextureUniformParameters, MobileSceneTextures)
+	RENDER_TARGET_BINDING_SLOTS()
+END_SHADER_PARAMETER_STRUCT()
+
 class FMobileDeferredShadingPS : public FGlobalShader
 {
 	DECLARE_SHADER_TYPE(FMobileDeferredShadingPS, Global);
@@ -321,30 +326,34 @@ static void RenderLocalLight(FRHICommandListImmediate& RHICmdList, const FScene&
 	}
 }
 
-void MobileDeferredShadingPass(FRHICommandListImmediate& RHICmdList, const FScene& Scene, const FViewInfo& View, const FSortedLightSetSceneInfo &SortedLightSet)
+void MobileDeferredShadingPass(FRDGBuilder& GraphBuilder, FRenderTargetBindingSlots& BasePassRenderTargets, TRDGUniformBufferRef<FMobileSceneTextureUniformParameters> MobileSceneTextures, const FScene& Scene, const FViewInfo& View, const FSortedLightSetSceneInfo &SortedLightSet)
 {
-	SCOPED_DRAW_EVENT(RHICmdList, MobileDeferredShading);
+	auto* PassParameters = GraphBuilder.AllocParameters<FMobileDeferredPassParameters>();
+	PassParameters->RenderTargets = BasePassRenderTargets;
+	PassParameters->MobileSceneTextures = MobileSceneTextures;
 
-	FSceneRenderTargets& SceneContext = FSceneRenderTargets::Get();
-	FUniformBufferRHIRef PassUniformBuffer = CreateMobileSceneTextureUniformBuffer(RHICmdList);
-	FUniformBufferStaticBindings GlobalUniformBuffers(PassUniformBuffer);
-	SCOPED_UNIFORM_BUFFER_GLOBAL_BINDINGS(RHICmdList, GlobalUniformBuffers);
-	RHICmdList.SetViewport(View.ViewRect.Min.X, View.ViewRect.Min.Y, 0.0f, View.ViewRect.Max.X, View.ViewRect.Max.Y, 1.0f);
+	GraphBuilder.AddPass(
+		RDG_EVENT_NAME("MobileDeferredShadingPass"), PassParameters, ERDGPassFlags::Raster,
+		[&Scene, &View, &SortedLightSet](FRHICommandListImmediate& RHICmdList)
+	{
 
-	RenderDirectLight(RHICmdList, Scene, View);
+		RHICmdList.SetViewport(View.ViewRect.Min.X, View.ViewRect.Min.Y, 0.0f, View.ViewRect.Max.X, View.ViewRect.Max.Y, 1.0f);
 
-	// Render non-clustered local lights
-	int32 StandardDeferredStart = SortedLightSet.SimpleLightsEnd;
-	int32 AttenuationLightStart = SortedLightSet.AttenuationLightStart;
-	if (GMobileUseClusteredDeferredShading != 0)
-	{
-		StandardDeferredStart = SortedLightSet.ClusteredSupportedEnd;
-	}
+		RenderDirectLight(RHICmdList, Scene, View);
 
-	for (int32 LightIdx = StandardDeferredStart; LightIdx < AttenuationLightStart; ++LightIdx)
-	{
-		const FSortedLightSceneInfo& SortedLight = SortedLightSet.SortedLights[LightIdx];
-		const FLightSceneInfo& LightSceneInfo = *SortedLight.LightSceneInfo;
-		RenderLocalLight(RHICmdList, Scene, View, LightSceneInfo);
-	}
+		// Render non-clustered local lights
+		int32 StandardDeferredStart = SortedLightSet.SimpleLightsEnd;
+		int32 AttenuationLightStart = SortedLightSet.AttenuationLightStart;
+		if (GMobileUseClusteredDeferredShading != 0)
+		{
+			StandardDeferredStart = SortedLightSet.ClusteredSupportedEnd;
+		}
+
+		for (int32 LightIdx = StandardDeferredStart; LightIdx < AttenuationLightStart; ++LightIdx)
+		{
+			const FSortedLightSceneInfo& SortedLight = SortedLightSet.SortedLights[LightIdx];
+			const FLightSceneInfo& LightSceneInfo = *SortedLight.LightSceneInfo;
+			RenderLocalLight(RHICmdList, Scene, View, LightSceneInfo);
+		}
+	});
 }
\ No newline at end of file
diff --git a/Engine/Source/Runtime/Renderer/Private/MobileDeferredShadingPass.h b/Engine/Source/Runtime/Renderer/Private/MobileDeferredShadingPass.h
--- a/Engine/Source/Runtime/Renderer/Private/MobileDeferredShadingPass.h
+++ b/Engine/Source/Runtime/Renderer/Private/MobileDeferredShadingPass.h
@@ -8,7 +8,8 @@ class FRHICommandListImmediate;
 class FScene;
 class FViewInfo;
 struct FSortedLightSetSceneInfo;
+class FMobileSceneTextureUniformParameters;
 
 extern int32 GMobileUseClusteredDeferredShading;
 
-void MobileDeferredShadingPass(FRHICommandListImmediate& RHICmdList, const FScene& Scene, const FViewInfo& View, const FSortedLightSetSceneInfo &SortedLightSet);
\ No newline at end of file
+void MobileDeferredShadingPass(FRDGBuilder& GraphBuilder, FRenderTargetBindingSlots& BasePassRenderTargets, TRDGUniformBufferRef<FMobileSceneTextureUniformParameters> MobileSceneTextures, const FScene& Scene, const FViewInfo& View, const FSortedLightSetSceneInfo &SortedLightSet);
\ No newline at end of file
diff --git a/Engine/Source/Runtime/Renderer/Private/MobileSceneCaptureRendering.cpp b/Engine/Source/Runtime/Renderer/Private/MobileSceneCaptureRendering.cpp
--- a/Engine/Source/Runtime/Renderer/Private/MobileSceneCaptureRendering.cpp
+++ b/Engine/Source/Runtime/Renderer/Private/MobileSceneCaptureRendering.cpp
@@ -28,6 +28,7 @@ MobileSceneCaptureRendering.cpp - Mobile specific scene capture code.
 #include "PipelineStateCache.h"
 #include "CommonRenderResources.h"
 #include "GenerateMips.h"
+#include "ScreenPass.h"
 
 /**
 * Shader set for the copy of scene color to capture target, alpha channel will contain opacity information. (Determined from depth buffer content)
@@ -46,149 +47,93 @@ static const TCHAR* GShaderSourceModeDefineName[] =
 	nullptr
 };
 
-template<ESceneCaptureSource CaptureSource>
 class FMobileSceneCaptureCopyPS : public FGlobalShader
 {
-	DECLARE_SHADER_TYPE(FMobileSceneCaptureCopyPS, Global)
 public:
-
-	static bool ShouldCompilePermutation(const FGlobalShaderPermutationParameters& Parameters) { return IsMobilePlatform(Parameters.Platform); }
-
-	FMobileSceneCaptureCopyPS(const ShaderMetaType::CompiledShaderInitializerType& Initializer) :
-		FGlobalShader(Initializer)
-	{
-		InTexture.Bind(Initializer.ParameterMap, TEXT("InTexture"), SPF_Mandatory);
-		InTextureSampler.Bind(Initializer.ParameterMap, TEXT("InTextureSampler"));
+	DECLARE_GLOBAL_SHADER(FMobileSceneCaptureCopyPS);
+	SHADER_USE_PARAMETER_STRUCT(FMobileSceneCaptureCopyPS, FGlobalShader);
+
+	class FCaptureSourceDim : SHADER_PERMUTATION_INT("CAPTURE_SOURCE", ESceneCaptureSource::SCS_MAX);
+
+	using FPermutationDomain = TShaderPermutationDomain<FCaptureSourceDim>;
+
+	BEGIN_SHADER_PARAMETER_STRUCT(FParameters, )
+		SHADER_PARAMETER_STRUCT_REF(FViewUniformShaderParameters, View)
+		SHADER_PARAMETER_RDG_TEXTURE(Texture2D, InTexture)
+		SHADER_PARAMETER_SAMPLER(SamplerState, InTextureSampler)
+		RENDER_TARGET_BINDING_SLOTS()
+	END_SHADER_PARAMETER_STRUCT()
+
+	static bool ShouldCompilePermutation(const FGlobalShaderPermutationParameters& Parameters) 
+	{ 
+		FPermutationDomain PermutationVector(Parameters.PermutationId);
+
+		auto CaptureSourceDim = PermutationVector.Get<FCaptureSourceDim>();
+		return IsMobilePlatform(Parameters.Platform) && 
+			(CaptureSourceDim == SCS_SceneColorHDR
+			|| CaptureSourceDim == SCS_FinalColorLDR
+			|| CaptureSourceDim == SCS_FinalColorHDR
+			|| CaptureSourceDim == SCS_SceneColorHDRNoAlpha
+			|| CaptureSourceDim == SCS_SceneColorSceneDepth
+			|| CaptureSourceDim == SCS_SceneDepth
+			|| CaptureSourceDim == SCS_DeviceDepth);
 	}
-	FMobileSceneCaptureCopyPS() {}
 
 	static void ModifyCompilationEnvironment(const FGlobalShaderPermutationParameters& Parameters, FShaderCompilerEnvironment& OutEnvironment)
 	{
+		FPermutationDomain PermutationVector(Parameters.PermutationId);
+
+		auto CaptureSourceDim = PermutationVector.Get<FCaptureSourceDim>();
+
 		OutEnvironment.SetDefine(TEXT("POST_PROCESS_MATERIAL"), 1u); // this will force reading depth from a SceneColor.A
-		const TCHAR* DefineName = GShaderSourceModeDefineName[CaptureSource];
+		const TCHAR* DefineName = GShaderSourceModeDefineName[CaptureSourceDim];
 		if (DefineName)
 		{
 			OutEnvironment.SetDefine(DefineName, 1);
 		}
 	}
 
-	void SetParameters(FRHICommandList& RHICmdList, const FViewInfo& View, FRHISamplerState* SamplerStateRHI, FRHITexture* TextureRHI)
+	static FPermutationDomain BuildPermutationVector(ESceneCaptureSource CaptureSource)
 	{
-		FGlobalShader::SetParameters<FViewUniformShaderParameters>(RHICmdList, RHICmdList.GetBoundPixelShader(), View.ViewUniformBuffer);
-		SetTextureParameter(RHICmdList, RHICmdList.GetBoundPixelShader(), InTexture, InTextureSampler, SamplerStateRHI, TextureRHI);
+		FPermutationDomain PermutationVector;
+		return PermutationVector;
 	}
-
-private:
-	LAYOUT_FIELD(FShaderResourceParameter, InTexture)
-	LAYOUT_FIELD(FShaderResourceParameter, InTextureSampler)
 };
 
+IMPLEMENT_GLOBAL_SHADER(FMobileSceneCaptureCopyPS, "/Engine/Private/MobileSceneCapture.usf", "MainCopyPS", SF_Pixel);
+
 /**
 * A vertex shader for rendering a textured screen element.
 */
 class FMobileSceneCaptureCopyVS : public FGlobalShader
 {
-	DECLARE_SHADER_TYPE(FMobileSceneCaptureCopyVS, Global)
-public:
+	DECLARE_GLOBAL_SHADER(FMobileSceneCaptureCopyVS);
+	SHADER_USE_PARAMETER_STRUCT_WITH_LEGACY_BASE(FMobileSceneCaptureCopyVS, FGlobalShader);
 
-	static bool ShouldCompilePermutation(const FGlobalShaderPermutationParameters& Parameters) { return IsMobilePlatform(Parameters.Platform); }
+	BEGIN_SHADER_PARAMETER_STRUCT(FParameters, )
+		SHADER_PARAMETER_STRUCT_REF(FViewUniformShaderParameters, View)
+		SHADER_PARAMETER(FVector2D, InvTexSize)
+	END_SHADER_PARAMETER_STRUCT()
 
-	FMobileSceneCaptureCopyVS(const ShaderMetaType::CompiledShaderInitializerType& Initializer) :
-		FGlobalShader(Initializer)
-	{
-		InvTexSizeParameter.Bind(Initializer.ParameterMap, TEXT("InvTexSize"));
-	}
-	FMobileSceneCaptureCopyVS() {}
-
-	static void ModifyCompilationEnvironment(const FGlobalShaderPermutationParameters& Parameters, FShaderCompilerEnvironment& OutEnvironment)
-	{
-	}
-
-	void SetParameters(FRHICommandList& RHICmdList, const FViewInfo& View, const FIntPoint& SourceTexSize)
-	{
-		FGlobalShader::SetParameters<FViewUniformShaderParameters>(RHICmdList, RHICmdList.GetBoundVertexShader(), View.ViewUniformBuffer);
-		if (InvTexSizeParameter.IsBound())
-		{
-			FVector2D InvTexSize(1.0f / SourceTexSize.X, 1.0f / SourceTexSize.Y);
-			SetShaderValue(RHICmdList, RHICmdList.GetBoundVertexShader(), InvTexSizeParameter, InvTexSize);
-		}
-	}
-
-	LAYOUT_FIELD(FShaderParameter, InvTexSizeParameter)
+	static bool ShouldCompilePermutation(const FGlobalShaderPermutationParameters& Parameters) { return IsMobilePlatform(Parameters.Platform); }
 };
 
-#define IMPLEMENT_MOBILE_SCENE_CAPTURECOPY(SCENETYPE) \
-typedef FMobileSceneCaptureCopyPS<SCENETYPE> FMobileSceneCaptureCopyPS##SCENETYPE;\
-IMPLEMENT_SHADER_TYPE(template<>, FMobileSceneCaptureCopyPS##SCENETYPE, TEXT("/Engine/Private/MobileSceneCapture.usf"), TEXT("MainCopyPS"), SF_Pixel); \
-
-IMPLEMENT_MOBILE_SCENE_CAPTURECOPY(SCS_SceneColorHDR);
-IMPLEMENT_MOBILE_SCENE_CAPTURECOPY(SCS_FinalColorLDR);
-IMPLEMENT_MOBILE_SCENE_CAPTURECOPY(SCS_FinalColorHDR);
-IMPLEMENT_MOBILE_SCENE_CAPTURECOPY(SCS_SceneColorHDRNoAlpha);
-IMPLEMENT_MOBILE_SCENE_CAPTURECOPY(SCS_SceneColorSceneDepth);
-IMPLEMENT_MOBILE_SCENE_CAPTURECOPY(SCS_SceneDepth);
-IMPLEMENT_MOBILE_SCENE_CAPTURECOPY(SCS_DeviceDepth);
-IMPLEMENT_SHADER_TYPE(, FMobileSceneCaptureCopyVS, TEXT("/Engine/Private/MobileSceneCapture.usf"), TEXT("MainCopyVS"), SF_Vertex);
+IMPLEMENT_GLOBAL_SHADER(FMobileSceneCaptureCopyVS, "/Engine/Private/MobileSceneCapture.usf", "MainCopyVS", SF_Vertex);
  
-
-template <ESceneCaptureSource CaptureSource>
-static TShaderRef<FShader> SetCaptureToTargetShaders(FRHICommandListImmediate& RHICmdList, FGraphicsPipelineStateInitializer& GraphicsPSOInit, FViewInfo& View, const FIntPoint& SourceTexSize, FRHITexture* SourceTextureRHI)
-{
-	TShaderMapRef<FMobileSceneCaptureCopyVS> VertexShader(View.ShaderMap);
-	TShaderMapRef<FMobileSceneCaptureCopyPS<CaptureSource>> PixelShader(View.ShaderMap);
-
-	GraphicsPSOInit.BoundShaderState.VertexDeclarationRHI = GFilterVertexDeclaration.VertexDeclarationRHI;
-	GraphicsPSOInit.BoundShaderState.VertexShaderRHI = VertexShader.GetVertexShader();
-	GraphicsPSOInit.BoundShaderState.PixelShaderRHI = PixelShader.GetPixelShader();
-	SetGraphicsPipelineState(RHICmdList, GraphicsPSOInit);
-
-	VertexShader->SetParameters(RHICmdList, View, SourceTexSize);
-	PixelShader->SetParameters(RHICmdList, View, TStaticSamplerState<SF_Point>::GetRHI(), SourceTextureRHI);
-
-	return VertexShader;
-}
-
-static TShaderRef<FShader> SetCaptureToTargetShaders(FRHICommandListImmediate& RHICmdList, FGraphicsPipelineStateInitializer& GraphicsPSOInit, ESceneCaptureSource CaptureSource, FViewInfo& View, const FIntPoint& SourceTexSize, FRHITexture* SourceTextureRHI)
-{
-	switch (CaptureSource)
-	{
-		case SCS_SceneColorHDR:
-			return SetCaptureToTargetShaders<SCS_SceneColorHDR>(RHICmdList, GraphicsPSOInit, View, SourceTexSize, SourceTextureRHI);
-		case SCS_FinalColorLDR:
-			return SetCaptureToTargetShaders<SCS_FinalColorLDR>(RHICmdList, GraphicsPSOInit, View, SourceTexSize, SourceTextureRHI);
-		case SCS_FinalColorHDR:
-			return SetCaptureToTargetShaders<SCS_FinalColorHDR>(RHICmdList, GraphicsPSOInit, View, SourceTexSize, SourceTextureRHI);
-		case SCS_SceneColorHDRNoAlpha:
-			return SetCaptureToTargetShaders<SCS_SceneColorHDRNoAlpha>(RHICmdList, GraphicsPSOInit, View, SourceTexSize, SourceTextureRHI);
-		case SCS_SceneColorSceneDepth:
-			return SetCaptureToTargetShaders<SCS_SceneColorSceneDepth>(RHICmdList, GraphicsPSOInit, View, SourceTexSize, SourceTextureRHI);
-		case SCS_SceneDepth:
-			return SetCaptureToTargetShaders<SCS_SceneDepth>(RHICmdList, GraphicsPSOInit, View, SourceTexSize, SourceTextureRHI);
-		case SCS_DeviceDepth:
-			return SetCaptureToTargetShaders<SCS_DeviceDepth>(RHICmdList, GraphicsPSOInit, View, SourceTexSize, SourceTextureRHI);
-		default:
-			checkNoEntry();
-			return TShaderRef<FShader>();
-	}
-}
-
 // Copies into render target, optionally flipping it in the Y-axis
 static void CopyCaptureToTarget(
-	FRHICommandListImmediate& RHICmdList, 
-	const FRenderTarget* Target, 
+	FRDGBuilder& GraphBuilder,
+	const FRDGTextureRef Target, 
 	const FIntPoint& TargetSize, 
 	FViewInfo& View, 
 	const FIntRect& ViewRect, 
-	FRHITexture2D* SourceTextureRHI,
+	FRDGTextureRef SourceTexture,
 	bool bNeedsFlippedRenderTarget,
 	FSceneRenderer* SceneRenderer)
 {
-	check(SourceTextureRHI);
-	check(RHICmdList.IsOutsideRenderPass());
+	FIntPoint SourceTexSize = SourceTexture->Desc.Extent;
 
-	FGraphicsPipelineStateInitializer GraphicsPSOInit;
 	ESceneCaptureSource CaptureSource = View.Family->SceneCaptureSource;
-	ESceneCaptureCompositeMode CaptureCompositeMode = View.Family->SceneCaptureCompositeMode;
 
 	// Normal and BaseColor not supported on mobile, fall back to scene colour.
 	if (CaptureSource == SCS_Normal || CaptureSource == SCS_BaseColor)
@@ -196,68 +141,97 @@ static void CopyCaptureToTarget(
 		CaptureSource = SCS_SceneColorHDR;
 	}
 
-	ERenderTargetLoadAction RTLoadAction;
-	if (CaptureSource == SCS_SceneColorHDR && CaptureCompositeMode == SCCM_Composite)
-	{
-		// Blend with existing render target color. Scene capture color is already pre-multiplied by alpha.
-		GraphicsPSOInit.BlendState = TStaticBlendState<CW_RGBA, BO_Add, BF_One, BF_SourceAlpha, BO_Add, BF_Zero, BF_SourceAlpha>::GetRHI();
-		RTLoadAction = ERenderTargetLoadAction::ELoad;
-	}
-	else if (CaptureSource == SCS_SceneColorHDR && CaptureCompositeMode == SCCM_Additive)
-	{
-		// Add to existing render target color. Scene capture color is already pre-multiplied by alpha.
-		GraphicsPSOInit.BlendState = TStaticBlendState<CW_RGBA, BO_Add, BF_One, BF_One, BO_Add, BF_Zero, BF_SourceAlpha>::GetRHI();
-		RTLoadAction = ERenderTargetLoadAction::ELoad;
-	}
-	else
-	{
-		RTLoadAction = ERenderTargetLoadAction::ENoAction;
-		GraphicsPSOInit.BlendState = TStaticBlendState<>::GetRHI();
-	}
-
-	GraphicsPSOInit.DepthStencilState = TStaticDepthStencilState<false, CF_Always>::GetRHI();
-	FIntPoint SourceTexSize = SourceTextureRHI->GetSizeXY();
-	
+	ESceneCaptureCompositeMode CaptureCompositeMode = View.Family->SceneCaptureCompositeMode;
 	{
-		RHICmdList.Transition(FRHITransitionInfo(Target->GetRenderTargetTexture(), ERHIAccess::Unknown,ERHIAccess::RTV));
-		FRHIRenderPassInfo RPInfo(Target->GetRenderTargetTexture(), MakeRenderTargetActions(RTLoadAction, ERenderTargetStoreAction::EStore));
-		RHICmdList.BeginRenderPass(RPInfo, TEXT("CaptureToTarget"));
-		RHICmdList.ApplyCachedRenderTargets(GraphicsPSOInit);
-
-		GraphicsPSOInit.PrimitiveType = PT_TriangleList;
-		GraphicsPSOInit.RasterizerState = TStaticRasterizerState<FM_Solid, CM_None>::GetRHI();
-
-		TShaderRef<FShader> VertexShader = SetCaptureToTargetShaders(RHICmdList, GraphicsPSOInit, CaptureSource, View, SourceTexSize, SourceTextureRHI);
-
-		if (bNeedsFlippedRenderTarget)
+		ERenderTargetLoadAction RTLoadAction;
+		FRHIBlendState* BlendState;
+		if (CaptureSource == SCS_SceneColorHDR && CaptureCompositeMode == SCCM_Composite)
 		{
-			DrawRectangle(
-				RHICmdList,
-				ViewRect.Min.X, ViewRect.Min.Y,
-				ViewRect.Width(), ViewRect.Height(),
-				ViewRect.Min.X, ViewRect.Height() - ViewRect.Min.Y,
-				ViewRect.Width(), -ViewRect.Height(),
-				TargetSize,
-				SourceTexSize,
-				VertexShader,
-				EDRF_UseTriangleOptimization);
+			// Blend with existing render target color. Scene capture color is already pre-multiplied by alpha.
+			BlendState = TStaticBlendState<CW_RGBA, BO_Add, BF_One, BF_SourceAlpha, BO_Add, BF_Zero, BF_SourceAlpha>::GetRHI();
+			RTLoadAction = ERenderTargetLoadAction::ELoad;
+		}
+		else if (CaptureSource == SCS_SceneColorHDR && CaptureCompositeMode == SCCM_Additive)
+		{
+			// Add to existing render target color. Scene capture color is already pre-multiplied by alpha.
+			BlendState = TStaticBlendState<CW_RGBA, BO_Add, BF_One, BF_One, BO_Add, BF_Zero, BF_SourceAlpha>::GetRHI();
+			RTLoadAction = ERenderTargetLoadAction::ELoad;
 		}
 		else
 		{
-			DrawRectangle(
-				RHICmdList,
-				ViewRect.Min.X, ViewRect.Min.Y,
-				ViewRect.Width(), ViewRect.Height(),
-				ViewRect.Min.X, ViewRect.Min.Y,
-				ViewRect.Width(), ViewRect.Height(),
-				TargetSize,
-				SourceTexSize,
-				VertexShader,
-				EDRF_UseTriangleOptimization);
+			RTLoadAction = ERenderTargetLoadAction::ENoAction;
+			BlendState = TStaticBlendState<>::GetRHI();
 		}
-		RHICmdList.EndRenderPass();
+
+		TShaderMapRef<FMobileSceneCaptureCopyVS> VertexShader(View.ShaderMap);
+
+		FMobileSceneCaptureCopyVS::FParameters VSShaderParameters;
+
+		VSShaderParameters.View = View.ViewUniformBuffer;
+		VSShaderParameters.InvTexSize = FVector2D(1.0f / SourceTexSize.X, 1.0f / SourceTexSize.Y);
+
+		auto ShaderPermutationVector = FMobileSceneCaptureCopyPS::BuildPermutationVector(CaptureSource);
+
+		TShaderMapRef<FMobileSceneCaptureCopyPS> PixelShader(View.ShaderMap, ShaderPermutationVector);
+
+		FMobileSceneCaptureCopyPS::FParameters* PSShaderParameters = GraphBuilder.AllocParameters<FMobileSceneCaptureCopyPS::FParameters>();
+		PSShaderParameters->View = View.ViewUniformBuffer;
+		PSShaderParameters->InTexture = SourceTexture;
+		PSShaderParameters->InTextureSampler = TStaticSamplerState<SF_Point>::GetRHI();
+		PSShaderParameters->RenderTargets[0] = FRenderTargetBinding(Target, RTLoadAction);
+
+		GraphBuilder.AddPass(
+			RDG_EVENT_NAME("CaptureToTarget"),
+			PSShaderParameters,
+			ERDGPassFlags::Raster,
+			[VertexShader, VSShaderParameters, PixelShader, PSShaderParameters, BlendState, bNeedsFlippedRenderTarget, &ViewRect, &TargetSize, SourceTexSize](FRHICommandListImmediate& RHICmdList)
+		{
+			RHICmdList.SetViewport(ViewRect.Min.X, ViewRect.Min.Y, 0.0f, ViewRect.Max.X, ViewRect.Max.Y, 1.0f);
+
+			FGraphicsPipelineStateInitializer GraphicsPSOInit;
+			RHICmdList.ApplyCachedRenderTargets(GraphicsPSOInit);
+
+			GraphicsPSOInit.BoundShaderState.VertexDeclarationRHI = GFilterVertexDeclaration.VertexDeclarationRHI;
+			GraphicsPSOInit.BoundShaderState.VertexShaderRHI = VertexShader.GetVertexShader();
+			GraphicsPSOInit.BoundShaderState.PixelShaderRHI = PixelShader.GetPixelShader();
+			GraphicsPSOInit.DepthStencilState = TStaticDepthStencilState<false, CF_Always>::GetRHI();
+			GraphicsPSOInit.BlendState = BlendState;
+			GraphicsPSOInit.PrimitiveType = PT_TriangleList;
+			GraphicsPSOInit.RasterizerState = TStaticRasterizerState<FM_Solid, CM_None>::GetRHI();
+
+			SetGraphicsPipelineState(RHICmdList, GraphicsPSOInit);
+
+			SetShaderParameters(RHICmdList, VertexShader, VertexShader.GetVertexShader(), VSShaderParameters);
+			SetShaderParameters(RHICmdList, PixelShader, PixelShader.GetPixelShader(), *PSShaderParameters);
+
+			if (bNeedsFlippedRenderTarget)
+			{
+				DrawRectangle(
+					RHICmdList,
+					ViewRect.Min.X, ViewRect.Min.Y,
+					ViewRect.Width(), ViewRect.Height(),
+					ViewRect.Min.X, ViewRect.Height() - ViewRect.Min.Y,
+					ViewRect.Width(), -ViewRect.Height(),
+					TargetSize,
+					SourceTexSize,
+					VertexShader,
+					EDRF_UseTriangleOptimization);
+			}
+			else
+			{
+				DrawRectangle(
+					RHICmdList,
+					ViewRect.Min.X, ViewRect.Min.Y,
+					ViewRect.Width(), ViewRect.Height(),
+					ViewRect.Min.X, ViewRect.Min.Y,
+					ViewRect.Width(), ViewRect.Height(),
+					TargetSize,
+					SourceTexSize,
+					VertexShader,
+					EDRF_UseTriangleOptimization);
+			}
+		});
 	}
-	
 
 	// if opacity is needed.
 	if (CaptureSource == SCS_SceneColorHDR)
@@ -267,13 +241,22 @@ static void CopyCaptureToTarget(
 		FMobileSceneRenderer* MobileSceneRenderer = (FMobileSceneRenderer*)SceneRenderer;
 
 		{
-			MobileSceneRenderer->RenderInverseOpacity(RHICmdList, View);
-		}		
+			MobileSceneRenderer->RenderInverseOpacity(GraphBuilder, View);
+		}
 
 		// Set capture target.
-		FRHIRenderPassInfo RPInfo(Target->GetRenderTargetTexture(), ERenderTargetActions::Load_Store);
-		RHICmdList.BeginRenderPass(RPInfo, TEXT("OpacitySceneCapture"));
+		FCopyRectPS::FParameters* PSShaderParameters = GraphBuilder.AllocParameters<FCopyRectPS::FParameters>();
+		PSShaderParameters->InputTexture = SourceTexture;
+		PSShaderParameters->InputSampler = TStaticSamplerState<SF_Point>::GetRHI();
+		PSShaderParameters->RenderTargets[0] = FRenderTargetBinding(Target, ERenderTargetLoadAction::ELoad);
+
+		TShaderMapRef<FScreenPassVS> ScreenVertexShader(View.ShaderMap);
+		TShaderMapRef<FCopyRectPS> PixelShader(View.ShaderMap);
+
+		GraphBuilder.AddPass(RDG_EVENT_NAME("OpacitySceneCapturePass"), PSShaderParameters, ERDGPassFlags::Raster,
+			[&View, ScreenVertexShader, PixelShader, bNeedsFlippedRenderTarget, PSShaderParameters, &ViewRect, &TargetSize, SourceTexSize](FRHICommandListImmediate& RHICmdList)
 		{
+			FGraphicsPipelineStateInitializer GraphicsPSOInit;
 			RHICmdList.ApplyCachedRenderTargets(GraphicsPSOInit);
 
 			GraphicsPSOInit.DepthStencilState = TStaticDepthStencilState<false, CF_Always>::GetRHI();
@@ -283,9 +266,6 @@ static void CopyCaptureToTarget(
 			GraphicsPSOInit.PrimitiveType = PT_TriangleList;
 
 			// Combine translucent opacity pass to earlier opaque pass to build final inverse opacity.
-			TShaderMapRef<FScreenVS> ScreenVertexShader(View.ShaderMap);
-			TShaderMapRef<FScreenPS> PixelShader(View.ShaderMap);
-
 			GraphicsPSOInit.BoundShaderState.VertexDeclarationRHI = GFilterVertexDeclaration.VertexDeclarationRHI;
 			GraphicsPSOInit.BoundShaderState.VertexShaderRHI = ScreenVertexShader.GetVertexShader();
 			GraphicsPSOInit.BoundShaderState.PixelShaderRHI = PixelShader.GetPixelShader();
@@ -293,8 +273,7 @@ static void CopyCaptureToTarget(
 
 			SetGraphicsPipelineState(RHICmdList, GraphicsPSOInit);
 
-			ScreenVertexShader->SetParameters(RHICmdList, View.ViewUniformBuffer);
-			PixelShader->SetParameters(RHICmdList, TStaticSamplerState<SF_Point>::GetRHI(), SourceTextureRHI);
+			SetShaderParameters(RHICmdList, PixelShader, PixelShader.GetPixelShader(), *PSShaderParameters);
 
 			int32 TargetPosY = ViewRect.Min.Y;
 			int32 TargetHeight = ViewRect.Height();
@@ -315,8 +294,7 @@ static void CopyCaptureToTarget(
 				SourceTexSize,
 				ScreenVertexShader,
 				EDRF_UseTriangleOptimization);
-		}
-		RHICmdList.EndRenderPass();
+		});
 	}
 }
 
@@ -429,29 +407,24 @@ void UpdateSceneCaptureContentMobile_RenderThread(
 			}
 		}
 
+		FRDGBuilder GraphBuilder(RHICmdList);
+
+		FRDGTextureRef ViewFamilyTexture = TryCreateViewFamilyTexture(GraphBuilder, SceneRenderer->ViewFamily);
+
 		const FIntPoint TargetSize(UnconstrainedViewRect.Width(), UnconstrainedViewRect.Height());
 		if (bNeedsFlippedFinalColor)
 		{
 			// We need to flip this texture upside down (since we depended on tonemapping to fix this on the hdr path)
-			SCOPED_DRAW_EVENT(RHICmdList, FlipCapture);
-			CopyCaptureToTarget(RHICmdList, Target, TargetSize, View, ViewRect, FlippedRenderTarget.GetTextureParamRef(), bNeedsFlippedCopy, SceneRenderer);
+			RDG_EVENT_SCOPE(GraphBuilder, "FlipCapture");
+			CopyCaptureToTarget(GraphBuilder, ViewFamilyTexture, TargetSize, View, ViewRect, GraphBuilder.RegisterExternalTexture(FlippedPooledRenderTarget), bNeedsFlippedCopy, SceneRenderer);
 		}
 		else if(bUseSceneTextures)
 		{
 			// Copy the captured scene into the destination texture
-			SCOPED_DRAW_EVENT(RHICmdList, CaptureSceneColor);
-			// If multiview is enabled, the SceneColorTexture will be an array texture.
-			if (View.bIsMobileMultiViewEnabled)
-			{
-				CopyCaptureToTarget(RHICmdList, Target, TargetSize, View, ViewRect, FSceneRenderTargets::Get().GetSceneColorTexture()->GetTexture2DArray(), bNeedsFlippedCopy, SceneRenderer);
-			}
-			else
-			{
-				CopyCaptureToTarget(RHICmdList, Target, TargetSize, View, ViewRect, FSceneRenderTargets::Get().GetSceneColorTexture()->GetTexture2D(), bNeedsFlippedCopy, SceneRenderer);
-			}
+			RDG_EVENT_SCOPE(GraphBuilder, "CaptureSceneColor");
+			CopyCaptureToTarget(GraphBuilder, ViewFamilyTexture, TargetSize, View, ViewRect, GraphBuilder.RegisterExternalTexture(FSceneRenderTargets::Get().GetSceneColor()), bNeedsFlippedCopy, SceneRenderer);
 		}
-
-		FRDGBuilder GraphBuilder(RHICmdList);
+		
 		{
 			FRDGTextureRef MipTexture = GraphBuilder.RegisterExternalTexture(CreateRenderTarget(RenderTarget->GetRenderTargetTexture(), TEXT("MipGenerationInput")));
 			FRDGTextureRef OutputTexture = GraphBuilder.RegisterExternalTexture(CreateRenderTarget(RenderTargetTexture->TextureRHI, TEXT("MipGenerationOutput")));
diff --git a/Engine/Source/Runtime/Renderer/Private/MobileShadingRenderer.cpp b/Engine/Source/Runtime/Renderer/Private/MobileShadingRenderer.cpp
--- a/Engine/Source/Runtime/Renderer/Private/MobileShadingRenderer.cpp
+++ b/Engine/Source/Runtime/Renderer/Private/MobileShadingRenderer.cpp
@@ -159,6 +159,9 @@ static bool UsesCustomDepthStencilLookup(const FViewInfo& View)
 	return bUsesCustomDepthStencil;
 }
 
+BEGIN_SHADER_PARAMETER_STRUCT(FRenderOpaqueFXPassParameters, )
+	SHADER_PARAMETER_RDG_UNIFORM_BUFFER(FMobileSceneTextureUniformParameters, SceneTextures)
+END_SHADER_PARAMETER_STRUCT()
 
 FMobileSceneRenderer::FMobileSceneRenderer(const FSceneViewFamily* InViewFamily,FHitProxyConsumer* HitProxyConsumer)
 	: FSceneRenderer(InViewFamily, HitProxyConsumer)
@@ -462,10 +465,6 @@ void FMobileSceneRenderer::InitViews(FRDGBuilder& GraphBuilder)
 		// Initialize the view's RHI resources.
 		View.InitRHIResources();
 
-		// TODO: remove when old path is removed
-		// Create the directional light uniform buffers
-		CreateDirectionalLightUniformBuffers(View);
-
 		// Get the custom 1x1 target used to store exposure value and Toggle the two render targets used to store new and old.
 		if (IsMobileEyeAdaptationEnabled(View))
 		{
@@ -523,12 +522,14 @@ void FMobileSceneRenderer::InitViews(FRDGBuilder& GraphBuilder)
 */
 void FMobileSceneRenderer::Render(FRHICommandListImmediate& RHICmdList)
 {
+	FRDGBuilder GraphBuilder(RHICmdList, RDG_EVENT_NAME("MobileShadingRenderer_Render(ViewFamily=%s)", ViewFamily.bResolveScene ? TEXT("Primary") : TEXT("Auxiliary")));
+
 	RHICmdList.SetCurrentStat(GET_STATID(STAT_CLMM_SceneStart));
 
 	SCOPED_DRAW_EVENT(RHICmdList, MobileSceneRender);
 	SCOPED_GPU_STAT(RHICmdList, MobileSceneRender);
 
-	Scene->UpdateAllPrimitiveSceneInfos(RHICmdList);
+	Scene->UpdateAllPrimitiveSceneInfos(GraphBuilder);
 
 	// Establish scene primitive count (must be done after UpdateAllPrimitiveSceneInfos)
 	FGPUSceneScopeBeginEndHelper GPUSceneScopeBeginEndHelper(Scene->GPUScene, GPUSceneDynamicContext, *Scene);
@@ -556,10 +557,10 @@ void FMobileSceneRenderer::Render(FRHICommandListImmediate& RHICmdList)
 
 	if (!ViewFamily.EngineShowFlags.Rendering)
 	{
+		GraphBuilder.Execute();
 		return;
 	}
 
-	SCOPED_DRAW_EVENT(RHICmdList, MobileSceneRender);
 	SCOPED_GPU_STAT(RHICmdList, MobileSceneRender);
 
 	WaitOcclusionTests(RHICmdList);
@@ -567,11 +568,8 @@ void FMobileSceneRenderer::Render(FRHICommandListImmediate& RHICmdList)
 	RHICmdList.ImmediateFlush(EImmediateFlushType::DispatchToRHIThread);
 
 	// Find the visible primitives and prepare targets and buffers for rendering
-	{
-		FRDGBuilder GraphBuilder(RHICmdList);
-		InitViews(GraphBuilder);
-		GraphBuilder.Execute();
-	}
+	InitViews(GraphBuilder);
+
 	
 	if (GRHINeedsExtraDeletionLatency || !GRHICommandList.Bypass())
 	{
@@ -590,27 +588,23 @@ void FMobileSceneRenderer::Render(FRHICommandListImmediate& RHICmdList)
 	DynamicReadBuffer.Commit();
 	RHICmdList.ImmediateFlush(EImmediateFlushType::DispatchToRHIThread);
 
-	RHICmdList.SetCurrentStat(GET_STATID(STAT_CLMM_SceneSim));
+	GraphBuilder.SetCommandListStat(GET_STATID(STAT_CLMM_SceneSim));
 
 	FSceneRenderTargets& SceneContext = FSceneRenderTargets::Get();
 
 	if (bUseVirtualTexturing)
 	{
-		FRDGBuilder GraphBuilder(RHICmdList);
-		{
-			RDG_GPU_STAT_SCOPE(GraphBuilder, VirtualTextureUpdate);
-			FVirtualTextureSystem::Get().Update(GraphBuilder, FeatureLevel, Scene);
+		RDG_GPU_STAT_SCOPE(GraphBuilder, VirtualTextureUpdate);
+		FVirtualTextureSystem::Get().Update(GraphBuilder, FeatureLevel, Scene);
 
-			AddPass(GraphBuilder, [&SceneContext](FRHICommandList& RHICmdList)
-			{
-				// Clear virtual texture feedback to default value
-				FUnorderedAccessViewRHIRef FeedbackUAV = SceneContext.GetVirtualTextureFeedbackUAV();
-				RHICmdList.Transition(FRHITransitionInfo(FeedbackUAV, ERHIAccess::Unknown, ERHIAccess::UAVCompute));
-				RHICmdList.ClearUAVUint(FeedbackUAV, FUintVector4(~0u, ~0u, ~0u, ~0u));
-				RHICmdList.Transition(FRHITransitionInfo(FeedbackUAV, ERHIAccess::UAVCompute, ERHIAccess::UAVCompute));
-			});
-		}
-		GraphBuilder.Execute();
+		AddPass(GraphBuilder, [&SceneContext](FRHICommandList& InRHICmdList)
+		{
+			// Clear virtual texture feedback to default value
+			FUnorderedAccessViewRHIRef FeedbackUAV = SceneContext.GetVirtualTextureFeedbackUAV();
+			InRHICmdList.Transition(FRHITransitionInfo(FeedbackUAV, ERHIAccess::Unknown, ERHIAccess::UAVCompute));
+			InRHICmdList.ClearUAVUint(FeedbackUAV, FUintVector4(~0u, ~0u, ~0u, ~0u));
+			InRHICmdList.Transition(FRHITransitionInfo(FeedbackUAV, ERHIAccess::UAVCompute, ERHIAccess::UAVCompute));
+		});
 	}
 
 	FSortedLightSetSceneInfo SortedLightSet;
@@ -619,41 +613,42 @@ void FMobileSceneRenderer::Render(FRHICommandListImmediate& RHICmdList)
 		GatherAndSortLights(SortedLightSet);
 		int32 NumReflectionCaptures = Views[0].NumBoxReflectionCaptures + Views[0].NumSphereReflectionCaptures;
 		bool bCullLightsToGrid = (NumReflectionCaptures > 0 || GMobileUseClusteredDeferredShading != 0);
-		FRDGBuilder GraphBuilder(RHICmdList);
 		ComputeLightGrid(GraphBuilder, bCullLightsToGrid, SortedLightSet);
-		GraphBuilder.Execute();
 	}
 
 	// Generate the Sky/Atmosphere look up tables
 	const bool bShouldRenderSkyAtmosphere = ShouldRenderSkyAtmosphere(Scene, ViewFamily.EngineShowFlags);
 	if (bShouldRenderSkyAtmosphere)
 	{
-		FRDGBuilder GraphBuilder(RHICmdList);
 		RenderSkyAtmosphereLookUpTables(GraphBuilder);
-		GraphBuilder.Execute();
 	}
 
 	// Notify the FX system that the scene is about to be rendered.
 	if (FXSystem && ViewFamily.EngineShowFlags.Particles)
 	{
-		FXSystem->PreRender(RHICmdList, Views[0].ViewUniformBuffer, NULL, !Views[0].bIsPlanarReflection);
-		if (FGPUSortManager* GPUSortManager = FXSystem->GetGPUSortManager())
+		AddPass(GraphBuilder, [this](FRHICommandListImmediate& InRHICmdList)
 		{
-			GPUSortManager->OnPreRender(RHICmdList);
-		}
+			FXSystem->PreRender(InRHICmdList, Views[0].ViewUniformBuffer, NULL, !Views[0].bIsPlanarReflection);
+			if (FGPUSortManager* GPUSortManager = FXSystem->GetGPUSortManager())
+			{
+				GPUSortManager->OnPreRender(InRHICmdList);
+			}
+		});
 	}
-	FRHICommandListExecutor::GetImmediateCommandList().PollOcclusionQueries();
-	RHICmdList.ImmediateFlush(EImmediateFlushType::DispatchToRHIThread);
-
-	RHICmdList.SetCurrentStat(GET_STATID(STAT_CLMM_Shadows));
 
+	auto PollOcclusionQueriesAndDispatchToRHIThreadPass = [](FRHICommandListImmediate& InRHICmdList)
 	{
-		FRDGBuilder GraphBuilder(RHICmdList);
-		RenderShadowDepthMaps(GraphBuilder);
-		GraphBuilder.Execute();
-	}
-	FRHICommandListExecutor::GetImmediateCommandList().PollOcclusionQueries();
-	RHICmdList.ImmediateFlush(EImmediateFlushType::DispatchToRHIThread);
+		FRHICommandListExecutor::GetImmediateCommandList().PollOcclusionQueries();
+		InRHICmdList.ImmediateFlush(EImmediateFlushType::DispatchToRHIThread);
+	};
+
+	AddPass(GraphBuilder, PollOcclusionQueriesAndDispatchToRHIThreadPass);
+
+	GraphBuilder.SetCommandListStat(GET_STATID(STAT_CLMM_Shadows));
+
+	RenderShadowDepthMaps(GraphBuilder);
+
+	AddPass(GraphBuilder, PollOcclusionQueriesAndDispatchToRHIThreadPass);
 
 	// Default view list
 	TArray<const FViewInfo*> ViewList;
@@ -666,88 +661,106 @@ void FMobileSceneRenderer::Render(FRHICommandListImmediate& RHICmdList)
 	// bShouldRenderCustomDepth has been initialized in InitViews on mobile platform
 	if (bShouldRenderCustomDepth)
 	{
-		FRDGBuilder GraphBuilder(RHICmdList);
 		RenderCustomDepthPass(GraphBuilder);
-		GraphBuilder.Execute();
 	}
+
+	// Include the CustomDepth flag in any case, so we don't have to create the MobileSceneTextures uniform buffer for each view.
+	EMobileSceneTextureSetupMode SetupMode = EMobileSceneTextureSetupMode::SceneColor | EMobileSceneTextureSetupMode::CustomDepth;
+
+	TRDGUniformBufferRef<FMobileSceneTextureUniformParameters> MobileSceneTextures = CreateMobileSceneTextureUniformBuffer(GraphBuilder, SetupMode);
+
+	FRDGTextureRef ViewFamilyTexture = TryCreateViewFamilyTexture(GraphBuilder, ViewFamily);
+
+	FRDGTextureMSAA SceneColorMSAA = RegisterExternalTextureMSAA(GraphBuilder, SceneContext.GetSceneColor());
+	FRDGTextureMSAA SceneDepthMSAA = RegisterExternalTextureMSAA(GraphBuilder, SceneContext.SceneDepthZ);
 	
-	FRHITexture* SceneColor = nullptr;
 	if (bDeferredShading)
 	{
-		SceneColor = RenderDeferred(RHICmdList, ViewList, SortedLightSet);
+		RenderDeferred(GraphBuilder, ViewList, SortedLightSet, ViewFamilyTexture, SceneColorMSAA, SceneDepthMSAA, MobileSceneTextures);
 	}
 	else
 	{
-		SceneColor = RenderForward(RHICmdList, ViewList);
+		RenderForward(GraphBuilder, ViewList, ViewFamilyTexture, SceneColorMSAA, SceneDepthMSAA, MobileSceneTextures);
 	}
 		
 	if (FXSystem && Views.IsValidIndex(0))
 	{
-		check(RHICmdList.IsOutsideRenderPass());
+		auto* PassParameters = GraphBuilder.AllocParameters<FRenderOpaqueFXPassParameters>();
+		PassParameters->SceneTextures = MobileSceneTextures;
 
-		FXSystem->PostRenderOpaque(
-			RHICmdList,
-			Views[0].ViewUniformBuffer,
-			nullptr,
-			nullptr,
-			Views[0].AllowGPUParticleUpdate()
-		);
-		if (FGPUSortManager* GPUSortManager = FXSystem->GetGPUSortManager())
+		// Cascade uses pixel shaders for compute stuff in PostRenderOpaque so ERDGPassFlags::Raster is needed
+		GraphBuilder.AddPass(
+			RDG_EVENT_NAME("OpaqueFX"),
+			PassParameters,
+			ERDGPassFlags::Raster | ERDGPassFlags::SkipRenderPass | ERDGPassFlags::Compute | ERDGPassFlags::NeverCull,
+			[this](FRHICommandListImmediate& InRHICmdList)
 		{
-			GPUSortManager->OnPostRenderOpaque(RHICmdList);
-		}
-		RHICmdList.ImmediateFlush(EImmediateFlushType::DispatchToRHIThread);
+			check(InRHICmdList.IsOutsideRenderPass());
+
+			FXSystem->PostRenderOpaque(
+				InRHICmdList,
+				Views[0].ViewUniformBuffer,
+				nullptr,
+				nullptr,
+				Views[0].AllowGPUParticleUpdate()
+			);
+			if (FGPUSortManager* GPUSortManager = FXSystem->GetGPUSortManager())
+			{
+				GPUSortManager->OnPostRenderOpaque(InRHICmdList);
+			}
+			InRHICmdList.ImmediateFlush(EImmediateFlushType::DispatchToRHIThread);
+		});
 	}
 
 	// Flush / submit cmdbuffer
 	if (bSubmitOffscreenRendering)
 	{
-		RHICmdList.SubmitCommandsHint();
-		RHICmdList.ImmediateFlush(EImmediateFlushType::DispatchToRHIThread);
-	}
-	
-	if (!bGammaSpace || bRenderToSceneColor)
-	{
-		RHICmdList.Transition(FRHITransitionInfo(SceneColor, ERHIAccess::Unknown, ERHIAccess::SRVMask));
+		AddPass(GraphBuilder, [this](FRHICommandListImmediate& InRHICmdList)
+		{
+			InRHICmdList.SubmitCommandsHint();
+			InRHICmdList.ImmediateFlush(EImmediateFlushType::DispatchToRHIThread);
+		});
 	}
 
 	if (bRequriesAmbientOcclusionPass)
 	{
-		RenderAmbientOcclusion(RHICmdList, SceneContext.SceneDepthZ);
+		FRDGTextureRef AmbientOcclusionTexture = GraphBuilder.RegisterExternalTexture(GAmbientOcclusionMobileOutputs.AmbientOcclusionTexture, ERenderTargetTexture::Targetable);
+		RenderAmbientOcclusion(GraphBuilder, (*MobileSceneTextures)->SceneDepthTexture, AmbientOcclusionTexture);
 	}
 
 	if (bDeferredShading)
 	{
-		// Release the original reference on the scene render targets
-		SceneContext.AdjustGBufferRefCount(RHICmdList, -1);
+		AddPass(GraphBuilder, [this, &SceneContext](FRHICommandListImmediate& InRHICmdList)
+		{
+			// Release the original reference on the scene render targets
+			SceneContext.AdjustGBufferRefCount(InRHICmdList, -1);
+		});
 	}
 
-	RHICmdList.SetCurrentStat(GET_STATID(STAT_CLMM_Post));
+	GraphBuilder.SetCommandListStat(GET_STATID(STAT_CLMM_Post));
 
 	if (bUseVirtualTexturing)
 	{	
-		SCOPED_GPU_STAT(RHICmdList, VirtualTextureUpdate);
-
-		// No pass after this should make VT page requests
-		RHICmdList.Transition(FRHITransitionInfo(SceneContext.VirtualTextureFeedbackUAV, ERHIAccess::UAVMask, ERHIAccess::SRVMask));
+		RDG_GPU_STAT_SCOPE(GraphBuilder, VirtualTextureUpdate);
 
-		TArray<FIntRect, TInlineAllocator<4>> ViewRects;
-		ViewRects.AddUninitialized(Views.Num());
-		for (int32 ViewIndex = 0; ViewIndex < Views.Num(); ++ViewIndex)
+		AddPass(GraphBuilder, RDG_EVENT_NAME("VirtualTextureUpdate"), [this, &SceneContext](FRHICommandListImmediate& InRHICmdList)
 		{
-			ViewRects[ViewIndex] = Views[ViewIndex].ViewRect;
-		}
-		
-		FVirtualTextureFeedbackBufferDesc Desc;
-		Desc.Init2D(SceneContext.GetBufferSizeXY(), ViewRects, SceneContext.GetVirtualTextureFeedbackScale());
+			// No pass after this should make VT page requests
+			InRHICmdList.Transition(FRHITransitionInfo(SceneContext.VirtualTextureFeedbackUAV, ERHIAccess::UAVMask, ERHIAccess::SRVMask));
 
-		SubmitVirtualTextureFeedbackBuffer(RHICmdList, SceneContext.VirtualTextureFeedback, Desc);
-	}
+			TArray<FIntRect, TInlineAllocator<4>> ViewRects;
+			ViewRects.AddUninitialized(Views.Num());
+			for (int32 ViewIndex = 0; ViewIndex < Views.Num(); ++ViewIndex)
+			{
+				ViewRects[ViewIndex] = Views[ViewIndex].ViewRect;
+			}
 
-	FMemMark Mark(FMemStack::Get());
-	FRDGBuilder GraphBuilder(RHICmdList);
+			FVirtualTextureFeedbackBufferDesc Desc;
+			Desc.Init2D(SceneContext.GetBufferSizeXY(), ViewRects, SceneContext.GetVirtualTextureFeedbackScale());
 
-	FRDGTextureRef ViewFamilyTexture = TryCreateViewFamilyTexture(GraphBuilder, ViewFamily);
+			SubmitVirtualTextureFeedbackBuffer(InRHICmdList, SceneContext.VirtualTextureFeedback, Desc);
+		});
+	}
 	
 	if (ViewFamily.bResolveScene)
 	{
@@ -758,65 +771,45 @@ void FMobileSceneRenderer::Render(FRHICommandListImmediate& RHICmdList)
 				RDG_EVENT_SCOPE(GraphBuilder, "PostProcessing");
 				SCOPE_CYCLE_COUNTER(STAT_FinishRenderViewTargetTime);
 
-				// Note that we should move this uniform buffer set up process right after the InitView to avoid any uniform buffer creation during the rendering after we porting all the passes to the RDG.
-				// We couldn't do it right now because the ResolveSceneDepth has another GraphicBuilder and it will re-register SceneDepthZ and that will cause crash.
-				TArray<TRDGUniformBufferRef<FMobileSceneTextureUniformParameters>, TInlineAllocator<1, SceneRenderingAllocator>> MobileSceneTexturesPerView;
-				MobileSceneTexturesPerView.SetNumZeroed(Views.Num());
-
-				const auto SetupMobileSceneTexturesPerView = [&]()
-				{
-					for (int32 ViewIndex = 0; ViewIndex < Views.Num(); ++ViewIndex)
-					{
-						EMobileSceneTextureSetupMode SetupMode = EMobileSceneTextureSetupMode::SceneColor;
-						if (Views[ViewIndex].bCustomDepthStencilValid)
-						{
-							SetupMode |= EMobileSceneTextureSetupMode::CustomDepth;
-						}
-
-						MobileSceneTexturesPerView[ViewIndex] = CreateMobileSceneTextureUniformBuffer(GraphBuilder, SetupMode);
-					}
-				};
-
-				SetupMobileSceneTexturesPerView();
-
 				FMobilePostProcessingInputs PostProcessingInputs;
 				PostProcessingInputs.ViewFamilyTexture = ViewFamilyTexture;
 
 				for (int32 ViewIndex = 0; ViewIndex < Views.Num(); ViewIndex++)
 				{
 					RDG_EVENT_SCOPE_CONDITIONAL(GraphBuilder, Views.Num() > 1, "View%d", ViewIndex);
-					PostProcessingInputs.SceneTextures = MobileSceneTexturesPerView[ViewIndex];
+					PostProcessingInputs.SceneTextures = MobileSceneTextures;
 					AddMobilePostProcessingPasses(GraphBuilder, Views[ViewIndex], PostProcessingInputs);
 				}
 			}
 		}
 	}
 
-	GEngine->GetPostRenderDelegate().Broadcast();
+	AddPass(GraphBuilder, [this](FRHICommandListImmediate& InRHICmdList)
+	{
+		GEngine->GetPostRenderDelegate().Broadcast();
 
-	RHICmdList.SetCurrentStat(GET_STATID(STAT_CLMM_SceneEnd));
+		InRHICmdList.SetCurrentStat(GET_STATID(STAT_CLMM_SceneEnd));
+	});
 	
 	RenderFinish(GraphBuilder, ViewFamilyTexture);
-	GraphBuilder.Execute();
+	
+	AddPass(GraphBuilder, PollOcclusionQueriesAndDispatchToRHIThreadPass);
 
-	FRHICommandListExecutor::GetImmediateCommandList().PollOcclusionQueries();
-	FRHICommandListExecutor::GetImmediateCommandList().ImmediateFlush(EImmediateFlushType::DispatchToRHIThread);
+	GraphBuilder.Execute();
 }
 
 BEGIN_SHADER_PARAMETER_STRUCT(FMobilePostBasePassViewExtensionParameters, )
 	SHADER_PARAMETER_RDG_UNIFORM_BUFFER(FMobileSceneTextureUniformParameters, SceneTextures)
 END_SHADER_PARAMETER_STRUCT()
 
-FRHITexture* FMobileSceneRenderer::RenderForward(FRHICommandListImmediate& RHICmdList, const TArrayView<const FViewInfo*> ViewList)
+void FMobileSceneRenderer::RenderForward(FRDGBuilder& GraphBuilder, const TArrayView<const FViewInfo*> ViewList, FRDGTextureRef ViewFamilyTexture, FRDGTextureMSAA& SceneColorMSAA, FRDGTextureMSAA& SceneDepthMSAA, TRDGUniformBufferRef<FMobileSceneTextureUniformParameters> MobileSceneTextures)
 {
 	const FViewInfo& View = *ViewList[0];
 	FSceneRenderTargets& SceneContext = FSceneRenderTargets::Get();
 				
-	FRHITexture* SceneColor = nullptr;
-	FRHITexture* SceneColorResolve = nullptr;
-	FRHITexture* SceneDepth = nullptr;
-	ERenderTargetActions ColorTargetAction = ERenderTargetActions::Clear_Store;
-	EDepthStencilTargetActions DepthTargetAction = EDepthStencilTargetActions::ClearDepthStencil_DontStoreDepthStencil;
+	FRDGTextureRef SceneColor = nullptr;
+	FRDGTextureRef SceneColorResolve = nullptr;
+	FRDGTextureRef SceneDepth = nullptr;
 
 	// Verify using both MSAA sample count AND the scene color surface sample count, since on GLES you can't have MSAA color targets,
 	// so the color target would be created without MSAA, and MSAA is achieved through magical means (the framebuffer, being MSAA,
@@ -830,99 +823,91 @@ FRHITexture* FMobileSceneRenderer::RenderForward(FRHICommandListImmediate& RHICm
 	{
 		if (bMobileMSAA)
 		{
-			SceneColor = SceneContext.GetSceneColorSurface();
-			SceneColorResolve = ViewFamily.RenderTarget->GetRenderTargetTexture();
-			ColorTargetAction = ERenderTargetActions::Clear_Resolve;
-			RHICmdList.Transition(FRHITransitionInfo(SceneColorResolve, ERHIAccess::Unknown, ERHIAccess::RTV | ERHIAccess::ResolveDst));
+			SceneColor = SceneColorMSAA.Target;
+			SceneColorResolve = ViewFamilyTexture;
 		}
 		else
 		{
-			SceneColor = ViewFamily.RenderTarget->GetRenderTargetTexture();
-			RHICmdList.Transition(FRHITransitionInfo(SceneColor, ERHIAccess::Unknown, ERHIAccess::RTV));
+			SceneColor = ViewFamilyTexture;
 		}
-		SceneDepth = SceneContext.GetSceneDepthSurface();
+		SceneDepth = SceneDepthMSAA.Target;
 	}
 	else
 	{
-		SceneColor = SceneContext.GetSceneColorSurface();
-		if (bMobileMSAA)
-		{
-			SceneColorResolve = SceneContext.GetSceneColorTexture();
-			ColorTargetAction = ERenderTargetActions::Clear_Resolve;
-			RHICmdList.Transition(FRHITransitionInfo(SceneColorResolve, ERHIAccess::Unknown, ERHIAccess::RTV | ERHIAccess::ResolveDst));
-		}
-		else
-		{
-			SceneColorResolve = nullptr;
-			ColorTargetAction = ERenderTargetActions::Clear_Store;
-		}
-
-		SceneDepth = SceneContext.GetSceneDepthSurface();
-				
-		if (bRequiresMultiPass)
-		{	
-			// store targets after opaque so translucency render pass can be restarted
-			ColorTargetAction = ERenderTargetActions::Clear_Store;
-			DepthTargetAction = EDepthStencilTargetActions::ClearDepthStencil_StoreDepthStencil;
-		}
-						
-		if (bKeepDepthContent)
-		{
-			// store depth if post-processing/capture needs it
-			DepthTargetAction = EDepthStencilTargetActions::ClearDepthStencil_StoreDepthStencil;
-		}
+		SceneColor = SceneColorMSAA.Target;
+		SceneColorResolve = bMobileMSAA ? SceneColorMSAA.Resolve : nullptr;
+		SceneDepth = SceneDepthMSAA.Target;
 	}
 
-	FRHITexture* FoveationTexture = nullptr;
+	FRDGTextureRef FoveationTexture = nullptr;
 	
 	if (SceneContext.IsFoveationTextureAllocated()	&& !View.bIsSceneCapture && !View.bIsReflectionCapture)
 	{
-		FoveationTexture = SceneContext.GetFoveationTexture();
+		FoveationTexture = GraphBuilder.RegisterExternalTexture(SceneContext.FoveationTexture, ERenderTargetTexture::Targetable);
 	}
 
-	FRHIRenderPassInfo SceneColorRenderPassInfo(
-		SceneColor,
-		ColorTargetAction,
-		SceneColorResolve,
-		SceneDepth,
-		DepthTargetAction,
-		nullptr, // we never resolve scene depth on mobile
-		FoveationTexture,
-		FExclusiveDepthStencil::DepthWrite_StencilWrite
-	);
-	SceneColorRenderPassInfo.SubpassHint = ESubpassHint::DepthReadSubpass;
-	SceneColorRenderPassInfo.NumOcclusionQueries = ComputeNumOcclusionQueriesToBatch();
-	SceneColorRenderPassInfo.bOcclusionQueries = SceneColorRenderPassInfo.NumOcclusionQueries != 0;
+	FRenderTargetBindingSlots BasePassRenderTargets;
+	BasePassRenderTargets[0] = FRenderTargetBinding(SceneColor, SceneColorResolve, ERenderTargetLoadAction::EClear);
+	BasePassRenderTargets.DepthStencil = FDepthStencilBinding(SceneDepth, ERenderTargetLoadAction::EClear, FExclusiveDepthStencil::DepthWrite_StencilWrite);
+	BasePassRenderTargets.FoveationTexture = FoveationTexture;
+	BasePassRenderTargets.SubpassHint = ESubpassHint::DepthReadSubpass;
+	BasePassRenderTargets.NumOcclusionQueries = ComputeNumOcclusionQueriesToBatch();
 
 	//if the scenecolor isn't multiview but the app is, need to render as a single-view multiview due to shaders
-	SceneColorRenderPassInfo.MultiViewCount = View.bIsMobileMultiViewEnabled ? 2 : (bIsMultiViewApplication ? 1 : 0);
+	BasePassRenderTargets.MultiViewCount = View.bIsMobileMultiViewEnabled ? 2 : (bIsMultiViewApplication ? 1 : 0);
+
+	// Update the render targets load action to Eload after executing the first pass that is EClear flag to let the RDG merges as many passes as possible.
+	auto UpdateRenderTargetsLoadAction = [](FRenderTargetBindingSlots& BasePassRenderTargets)
+	{
+		if (BasePassRenderTargets[0].GetLoadAction() == ERenderTargetLoadAction::EClear)
+		{
+			BasePassRenderTargets[0].SetLoadAction(ERenderTargetLoadAction::ELoad);
+			BasePassRenderTargets.DepthStencil.SetDepthLoadAction(ERenderTargetLoadAction::ELoad);
+			BasePassRenderTargets.DepthStencil.SetStencilLoadAction(ERenderTargetLoadAction::ELoad);
+		}
+	};
+
+	RDG_EVENT_SCOPE(GraphBuilder, "SceneColorRendering");
 
-	RHICmdList.BeginRenderPass(SceneColorRenderPassInfo, TEXT("SceneColorRendering"));
-	
 	if (GIsEditor && !View.bIsSceneCapture)
 	{
-		DrawClearQuad(RHICmdList, Views[0].BackgroundColor);
+		auto* PassParameters = GraphBuilder.AllocParameters<FRenderTargetParameters>();
+		PassParameters->RenderTargets = BasePassRenderTargets;
+
+		GraphBuilder.AddPass(RDG_EVENT_NAME("ClearSceneColorPass"), PassParameters, ERDGPassFlags::Raster,
+			[this](FRHICommandListImmediate& RHICmdList)
+		{
+			DrawClearQuad(RHICmdList, Views[0].BackgroundColor);
+		});
+
+		UpdateRenderTargetsLoadAction(BasePassRenderTargets);
 	}
 
-	RHICmdList.SetCurrentStat(GET_STATID(STAT_CLM_MobilePrePass));
 	// Depth pre-pass
-	RenderPrePass(RHICmdList);
+	GraphBuilder.SetCommandListStat(GET_STATID(STAT_CLM_MobilePrePass));
+	RenderPrePass(GraphBuilder, BasePassRenderTargets, UpdateRenderTargetsLoadAction);
 	
 	// Opaque and masked
-	RHICmdList.SetCurrentStat(GET_STATID(STAT_CLMM_Opaque));
-	RenderMobileBasePass(RHICmdList, ViewList);
-	RHICmdList.ImmediateFlush(EImmediateFlushType::DispatchToRHIThread);
+	GraphBuilder.SetCommandListStat(GET_STATID(STAT_CLMM_Opaque));
+	RenderMobileBasePass(GraphBuilder, BasePassRenderTargets, UpdateRenderTargetsLoadAction, ViewList);
+	AddDispatchToRHIThreadPass(GraphBuilder);
 
 #if !(UE_BUILD_SHIPPING || UE_BUILD_TEST)
 	if (ViewFamily.UseDebugViewPS())
 	{
+		auto* PassParameters = GraphBuilder.AllocParameters<FRenderTargetParameters>();
+		PassParameters->RenderTargets = BasePassRenderTargets;
+
 		// Here we use the base pass depth result to get z culling for opaque and masque.
 		// The color needs to be cleared at this point since shader complexity renders in additive.
-		DrawClearQuad(RHICmdList, FLinearColor::Black);
-		FRDGBuilder GraphBuilder(RHICmdList);
+		GraphBuilder.AddPass(RDG_EVENT_NAME("ClearDebugViewPass"), PassParameters, ERDGPassFlags::Raster,
+			[this](FRHICommandListImmediate& RHICmdList)
+		{
+			DrawClearQuad(RHICmdList, FLinearColor::Black);
+		});
+
 		RenderMobileDebugView(GraphBuilder, ViewList);
-		GraphBuilder.Execute();
-		RHICmdList.ImmediateFlush(EImmediateFlushType::DispatchToRHIThread);
+		AddDispatchToRHIThreadPass(GraphBuilder);
 	}
 #endif // !(UE_BUILD_SHIPPING || UE_BUILD_TEST)
 
@@ -930,21 +915,20 @@ FRHITexture* FMobileSceneRenderer::RenderForward(FRHICommandListImmediate& RHICm
 	if (!bAdrenoOcclusionMode)
 	{
 		// Issue occlusion queries
-		RHICmdList.SetCurrentStat(GET_STATID(STAT_CLMM_Occlusion));
-		RenderOcclusion(RHICmdList);
-		RHICmdList.ImmediateFlush(EImmediateFlushType::DispatchToRHIThread);
+		GraphBuilder.SetCommandListStat(GET_STATID(STAT_CLMM_Occlusion));
+		RenderOcclusion(GraphBuilder, BasePassRenderTargets);
+		AddDispatchToRHIThreadPass(GraphBuilder);
 	}
 
 	if (ViewFamily.ViewExtensions.Num() > 1)
 	{
-		FRDGBuilder GraphBuilder(RHICmdList);
 		auto* PassParameters = GraphBuilder.AllocParameters<FMobilePostBasePassViewExtensionParameters>();
-		PassParameters->SceneTextures = CreateMobileSceneTextureUniformBuffer(GraphBuilder, EMobileSceneTextureSetupMode::None);
+		PassParameters->SceneTextures = MobileSceneTextures;
 
 		GraphBuilder.AddPass(
 			{},
 			PassParameters,
-			ERDGPassFlags::Raster | ERDGPassFlags::SkipRenderPass | ERDGPassFlags::NeverCull,
+			ERDGPassFlags::Raster | ERDGPassFlags::NeverCull,
 			[this](FRHICommandListImmediate& RHICmdList)
 		{
 			CSV_SCOPED_TIMING_STAT_EXCLUSIVE(ViewExtensionPostRenderBasePass);
@@ -957,40 +941,29 @@ FRHITexture* FMobileSceneRenderer::RenderForward(FRHICommandListImmediate& RHICm
 				}
 			}
 		});
-		GraphBuilder.Execute();
 	}
+	   
+	GraphBuilder.SetCommandListStat(GET_STATID(STAT_CLMM_Translucency));
 
-	// Split if we need to render translucency in a separate render pass
-	// Split if we need to render pixel projected reflection
+	// Restart translucency render pass if needed
 	if (bRequiresMultiPass || bRequiresPixelProjectedPlanarRelfectionPass)
 	{
-		RHICmdList.EndRenderPass();
+		// Make a copy of the scene depth if the current hardware doesn't support reading and writing to the same depth buffer
+		ConditionalResolveSceneDepth(GraphBuilder, View, SceneDepthMSAA);
 	}
-	   
-	RHICmdList.SetCurrentStat(GET_STATID(STAT_CLMM_Translucency));
 
-		
-	// Restart translucency render pass if needed
-	if (bRequiresMultiPass || bRequiresPixelProjectedPlanarRelfectionPass)
+	if (bRequiresPixelProjectedPlanarRelfectionPass)
 	{
-		check(RHICmdList.IsOutsideRenderPass());
+		const FPlanarReflectionSceneProxy* PlanarReflectionSceneProxy = Scene ? Scene->GetForwardPassGlobalPlanarReflection() : nullptr;
 
-		// Make a copy of the scene depth if the current hardware doesn't support reading and writing to the same depth buffer
-		ConditionalResolveSceneDepth(RHICmdList, View);
+		FRDGTextureRef PixelProjectedReflectionTexture = GraphBuilder.RegisterExternalTexture(GPixelProjectedReflectionMobileOutputs.PixelProjectedReflectionTexture);
 
-		if (bRequiresPixelProjectedPlanarRelfectionPass)
-		{
-			const FPlanarReflectionSceneProxy* PlanarReflectionSceneProxy = Scene ? Scene->GetForwardPassGlobalPlanarReflection() : nullptr;
-			RenderPixelProjectedReflection(RHICmdList, SceneContext, PlanarReflectionSceneProxy);
-
-			FRHITransitionInfo TranslucentRenderPassTransitions[] = {
-			FRHITransitionInfo(SceneColor, ERHIAccess::SRVMask, ERHIAccess::RTV),
-			FRHITransitionInfo(SceneDepth, ERHIAccess::SRVMask, ERHIAccess::DSVWrite)
-			};
-			RHICmdList.Transition(MakeArrayView(TranslucentRenderPassTransitions, UE_ARRAY_COUNT(TranslucentRenderPassTransitions)));
-		}
+		RenderPixelProjectedReflection(GraphBuilder, SceneColorMSAA.Resolve, SceneDepthMSAA.Resolve, PixelProjectedReflectionTexture, PlanarReflectionSceneProxy);
+	}
 
-		DepthTargetAction = EDepthStencilTargetActions::LoadDepthStencil_DontStoreDepthStencil;
+
+	if (bRequiresMultiPass || bRequiresPixelProjectedPlanarRelfectionPass)
+	{
 		FExclusiveDepthStencil::Type ExclusiveDepthStencil = FExclusiveDepthStencil::DepthRead_StencilRead;
 		if (bModulatedShadowsInUse)
 		{
@@ -1005,42 +978,36 @@ FRHITexture* FMobileSceneRenderer::RenderForward(FRHICommandListImmediate& RHICm
 			ExclusiveDepthStencil = FExclusiveDepthStencil::DepthWrite_StencilWrite;
 		}
 
-		if (bKeepDepthContent && !bMobileMSAA)
-		{
-			DepthTargetAction = EDepthStencilTargetActions::LoadDepthStencil_StoreDepthStencil;
-		}
+		BasePassRenderTargets.DepthStencil.SetDepthStencilAccess(ExclusiveDepthStencil);
+
+		BasePassRenderTargets.NumOcclusionQueries = 0;
+		BasePassRenderTargets.SubpassHint = ESubpassHint::DepthReadSubpass;
+	}
 
-		FRHIRenderPassInfo TranslucentRenderPassInfo(
-			SceneColor,
-			SceneColorResolve ? ERenderTargetActions::Load_Resolve : ERenderTargetActions::Load_Store,
-			SceneColorResolve,
-			SceneDepth,
-			DepthTargetAction, 
-			nullptr,
-			FoveationTexture,
-			ExclusiveDepthStencil
-		);
-		TranslucentRenderPassInfo.NumOcclusionQueries = 0;
-		TranslucentRenderPassInfo.bOcclusionQueries = false;
-		TranslucentRenderPassInfo.SubpassHint = ESubpassHint::DepthReadSubpass;
-		RHICmdList.BeginRenderPass(TranslucentRenderPassInfo, TEXT("SceneColorTranslucencyRendering"));
-	}
-
-	// scene depth is read only and can be fetched
-	RHICmdList.NextSubpass();
+	{
+		auto* PassParameters = GraphBuilder.AllocParameters<FRenderTargetParameters>();
+		PassParameters->RenderTargets = BasePassRenderTargets;
+
+		GraphBuilder.AddPass(RDG_EVENT_NAME("Nextsubpass"), PassParameters, ERDGPassFlags::Raster,
+			[this](FRHICommandListImmediate& RHICmdList)
+		{
+			// scene depth is read only and can be fetched
+			RHICmdList.NextSubpass();
+		});
+	}
 
 	if (!View.bIsPlanarReflection)
 	{
 		if (ViewFamily.EngineShowFlags.Decals)
 		{
 			CSV_SCOPED_TIMING_STAT_EXCLUSIVE(RenderDecals);
-			RenderDecals(RHICmdList);
+			RenderDecals(GraphBuilder, BasePassRenderTargets, MobileSceneTextures);
 		}
 
 		if (ViewFamily.EngineShowFlags.DynamicShadows)
 		{
 			CSV_SCOPED_TIMING_STAT_EXCLUSIVE(RenderShadowProjections);
-			RenderModulatedShadowProjections(RHICmdList);
+			RenderModulatedShadowProjections(GraphBuilder, BasePassRenderTargets, MobileSceneTextures);
 		}
 	}
 	
@@ -1049,183 +1016,199 @@ FRHITexture* FMobileSceneRenderer::RenderForward(FRHICommandListImmediate& RHICm
 	{
 		CSV_SCOPED_TIMING_STAT_EXCLUSIVE(RenderTranslucency);
 		SCOPE_CYCLE_COUNTER(STAT_TranslucencyDrawTime);
-		RenderTranslucency(RHICmdList, ViewList);
-		FRHICommandListExecutor::GetImmediateCommandList().PollOcclusionQueries();
-		RHICmdList.ImmediateFlush(EImmediateFlushType::DispatchToRHIThread);
+		RenderTranslucency(GraphBuilder, BasePassRenderTargets, ViewList);
+		AddPass(GraphBuilder, [](FRHICommandListImmediate& RHICmdList)
+		{
+			FRHICommandListExecutor::GetImmediateCommandList().PollOcclusionQueries();
+			RHICmdList.ImmediateFlush(EImmediateFlushType::DispatchToRHIThread);
+		});
 	}
 
 	if (bAdrenoOcclusionMode)
 	{
-		RHICmdList.SetCurrentStat(GET_STATID(STAT_CLMM_Occlusion));
-		// flush
-		RHICmdList.SubmitCommandsHint();
+		GraphBuilder.SetCommandListStat(GET_STATID(STAT_CLMM_Occlusion));
+		AddPass(GraphBuilder, [this](FRHICommandListImmediate& RHICmdList)
+		{
+			// flush
+			RHICmdList.SubmitCommandsHint();
+		});
 		bSubmitOffscreenRendering = false; // submit once
 		// Issue occlusion queries
-		RenderOcclusion(RHICmdList);
-		RHICmdList.ImmediateFlush(EImmediateFlushType::DispatchToRHIThread);
+		RenderOcclusion(GraphBuilder, BasePassRenderTargets);
+		AddDispatchToRHIThreadPass(GraphBuilder);
 	}
 
 	// Pre-tonemap before MSAA resolve (iOS only)
 	if (!bGammaSpace)
 	{
-		PreTonemapMSAA(RHICmdList);
+		PreTonemapMSAA(GraphBuilder, BasePassRenderTargets);
 	}
-
-	// End of scene color rendering
-	RHICmdList.EndRenderPass();
-
-	return SceneColorResolve ? SceneColorResolve : SceneColor;
 }
 
-FRHITexture* FMobileSceneRenderer::RenderDeferred(FRHICommandListImmediate& RHICmdList, const TArrayView<const FViewInfo*> ViewList, const FSortedLightSetSceneInfo& SortedLightSet)
+void FMobileSceneRenderer::RenderDeferred(FRDGBuilder& GraphBuilder, const TArrayView<const FViewInfo*> ViewList, const FSortedLightSetSceneInfo& SortedLightSet, FRDGTextureRef ViewFamilyTexture, FRDGTextureMSAA& SceneColorMSAA, FRDGTextureMSAA& SceneDepthMSAA, TRDGUniformBufferRef<FMobileSceneTextureUniformParameters> MobileSceneTextures)
 {
 	FSceneRenderTargets& SceneContext = FSceneRenderTargets::Get();
+
+	FRDGTextureRef ColorTargets[5] = {
+		SceneColorMSAA.Target,
+		GraphBuilder.RegisterExternalTexture(SceneContext.GBufferA),
+		GraphBuilder.RegisterExternalTexture(SceneContext.GBufferB),
+		GraphBuilder.RegisterExternalTexture(SceneContext.GBufferC),
+		GraphBuilder.RegisterExternalTexture(SceneContext.SceneDepthAux)
+	};
+
+	TArrayView<FRDGTextureRef> BasePassTexturesView = MakeArrayView(ColorTargets, 5);
+
+	FRenderTargetBindingSlots BasePassRenderTargets = GetRenderTargetBindings(ERenderTargetLoadAction::EClear, BasePassTexturesView);
+	BasePassRenderTargets.DepthStencil = FDepthStencilBinding(SceneDepthMSAA.Target, ERenderTargetLoadAction::EClear, FExclusiveDepthStencil::DepthWrite_StencilWrite);
+	BasePassRenderTargets.SubpassHint = ESubpassHint::DeferredShadingSubpass;
+	BasePassRenderTargets.NumOcclusionQueries = ComputeNumOcclusionQueriesToBatch();
+	BasePassRenderTargets.FoveationTexture = nullptr;
+	BasePassRenderTargets.MultiViewCount = 0;
+
+	// Update the render targets load action to Eload after adding the first pass with EClear flag to let the RDG merge all passes together.
+	auto UpdateRenderTargetsLoadAction = [](FRenderTargetBindingSlots& BasePassRenderTargets)
+	{
+		if (BasePassRenderTargets[0].GetLoadAction() == ERenderTargetLoadAction::EClear)
+		{
+			for (uint32 i = 0; i < 5; ++i)
+			{
+				BasePassRenderTargets[i].SetLoadAction(ERenderTargetLoadAction::ELoad);
+			}
 			
-	FRHITexture* ColorTargets[5] = {
-		SceneContext.GetSceneColorSurface(),
-		SceneContext.GetGBufferATexture().GetReference(),
-		SceneContext.GetGBufferBTexture().GetReference(),
-		SceneContext.GetGBufferCTexture().GetReference(),
-		SceneContext.SceneDepthAux->GetRenderTargetItem().ShaderResourceTexture.GetReference()
+			BasePassRenderTargets.DepthStencil.SetDepthLoadAction(ERenderTargetLoadAction::ELoad);
+			BasePassRenderTargets.DepthStencil.SetStencilLoadAction(ERenderTargetLoadAction::ELoad);
+		}
 	};
 
-	// Whether RHI needs to store GBuffer to system memory and do shading in separate render-pass
-	ERenderTargetActions GBufferAction = bRequiresMultiPass ? ERenderTargetActions::Clear_Store : ERenderTargetActions::Clear_DontStore;
-	EDepthStencilTargetActions DepthAction = bKeepDepthContent ? EDepthStencilTargetActions::ClearDepthStencil_StoreDepthStencil : EDepthStencilTargetActions::ClearDepthStencil_DontStoreDepthStencil;
-		
-	ERenderTargetActions ColorTargetsAction[5] = {ERenderTargetActions::Clear_Store, GBufferAction, GBufferAction, GBufferAction, ERenderTargetActions::Clear_Store};
-	
-	FRHIRenderPassInfo BasePassInfo = FRHIRenderPassInfo();
-	for (int32 Index = 0; Index < UE_ARRAY_COUNT(ColorTargets); ++Index)
-	{
-		BasePassInfo.ColorRenderTargets[Index].RenderTarget = ColorTargets[Index];
-		BasePassInfo.ColorRenderTargets[Index].ResolveTarget = nullptr;
-		BasePassInfo.ColorRenderTargets[Index].ArraySlice = -1;
-		BasePassInfo.ColorRenderTargets[Index].MipIndex = 0;
-		BasePassInfo.ColorRenderTargets[Index].Action = ColorTargetsAction[Index];
-	}
-	BasePassInfo.DepthStencilRenderTarget.DepthStencilTarget = SceneContext.GetSceneDepthSurface();
-	BasePassInfo.DepthStencilRenderTarget.ResolveTarget = nullptr;
-	BasePassInfo.DepthStencilRenderTarget.Action = DepthAction;
-	BasePassInfo.DepthStencilRenderTarget.ExclusiveDepthStencil = FExclusiveDepthStencil::DepthWrite_StencilWrite;
-		
-	BasePassInfo.SubpassHint = ESubpassHint::DeferredShadingSubpass;
-	BasePassInfo.NumOcclusionQueries = ComputeNumOcclusionQueriesToBatch();
-	BasePassInfo.bOcclusionQueries = BasePassInfo.NumOcclusionQueries != 0;
-	BasePassInfo.FoveationTexture = nullptr;
-	BasePassInfo.bIsMSAA = false;
-	BasePassInfo.MultiViewCount = 0;
-
-	RHICmdList.BeginRenderPass(BasePassInfo, TEXT("BasePassRendering"));
+	RDG_EVENT_SCOPE(GraphBuilder, "SceneColorRendering");
 	
 	if (GIsEditor && !Views[0].bIsSceneCapture)
 	{
-		DrawClearQuad(RHICmdList, Views[0].BackgroundColor);
+		auto* PassParameters = GraphBuilder.AllocParameters<FRenderTargetParameters>();
+		PassParameters->RenderTargets = BasePassRenderTargets;
+
+		GraphBuilder.AddPass(RDG_EVENT_NAME("ClearSceneColorPass"), PassParameters, ERDGPassFlags::Raster,
+			[this](FRHICommandListImmediate& RHICmdList)
+		{
+			DrawClearQuad(RHICmdList, Views[0].BackgroundColor);
+		});
+
+		UpdateRenderTargetsLoadAction(BasePassRenderTargets);
 	}
 
-	RHICmdList.SetCurrentStat(GET_STATID(STAT_CLM_MobilePrePass));
+	GraphBuilder.SetCommandListStat(GET_STATID(STAT_CLM_MobilePrePass));
 	// Depth pre-pass
-	RenderPrePass(RHICmdList);
+	RenderPrePass(GraphBuilder, BasePassRenderTargets, UpdateRenderTargetsLoadAction);
 	
 	// Opaque and masked
-	RHICmdList.SetCurrentStat(GET_STATID(STAT_CLMM_Opaque));
-	RenderMobileBasePass(RHICmdList, ViewList);
-	RHICmdList.ImmediateFlush(EImmediateFlushType::DispatchToRHIThread);
+	GraphBuilder.SetCommandListStat(GET_STATID(STAT_CLMM_Opaque));
+	RenderMobileBasePass(GraphBuilder, BasePassRenderTargets, UpdateRenderTargetsLoadAction, ViewList);
+	AddDispatchToRHIThreadPass(GraphBuilder);
 
 	// Issue occlusion queries
-	RHICmdList.SetCurrentStat(GET_STATID(STAT_CLMM_Occlusion));
-	RenderOcclusion(RHICmdList);
-	RHICmdList.ImmediateFlush(EImmediateFlushType::DispatchToRHIThread);
+	GraphBuilder.SetCommandListStat(GET_STATID(STAT_CLMM_Occlusion));
+	RenderOcclusion(GraphBuilder, BasePassRenderTargets);
+	AddDispatchToRHIThreadPass(GraphBuilder);
+
+	auto PollOcclusionQueriesAndDispatchToRHIThreadPass = [](FRHICommandListImmediate& InRHICmdList)
+	{
+		FRHICommandListExecutor::GetImmediateCommandList().PollOcclusionQueries();
+		InRHICmdList.ImmediateFlush(EImmediateFlushType::DispatchToRHIThread);
+	};
 
 	if (!bRequiresMultiPass)
 	{
-		// SceneColor + GBuffer write, SceneDepth is read only
-		RHICmdList.NextSubpass();
+		{
+			auto* PassParameters = GraphBuilder.AllocParameters<FRenderTargetParameters>();
+			PassParameters->RenderTargets = BasePassRenderTargets;
+
+			GraphBuilder.AddPass(RDG_EVENT_NAME("Nextsubpass"), PassParameters, ERDGPassFlags::Raster,
+				[this](FRHICommandListImmediate& RHICmdList)
+			{
+				// SceneColor + GBuffer write, SceneDepth is read only
+				RHICmdList.NextSubpass();
+			});
+		}
 		
 		if (ViewFamily.EngineShowFlags.Decals)
 		{
 			CSV_SCOPED_TIMING_STAT_EXCLUSIVE(RenderDecals);
-			RenderDecals(RHICmdList);
+			RenderDecals(GraphBuilder, BasePassRenderTargets, MobileSceneTextures);
 		}
 
-		// SceneColor write, SceneDepth is read only
-		RHICmdList.NextSubpass();
+		{
+			auto* PassParameters = GraphBuilder.AllocParameters<FRenderTargetParameters>();
+			PassParameters->RenderTargets = BasePassRenderTargets;
+
+			GraphBuilder.AddPass(RDG_EVENT_NAME("Nextsubpass"), PassParameters, ERDGPassFlags::Raster,
+				[this](FRHICommandListImmediate& RHICmdList)
+			{
+				// SceneColor write, SceneDepth is read only
+				RHICmdList.NextSubpass();
+			});
+		}
 		
-		MobileDeferredShadingPass(RHICmdList, *Scene, *ViewList[0], SortedLightSet);
+		MobileDeferredShadingPass(GraphBuilder, BasePassRenderTargets, MobileSceneTextures, *Scene, *ViewList[0], SortedLightSet);
 		// Draw translucency.
 		if (ViewFamily.EngineShowFlags.Translucency)
 		{
 			CSV_SCOPED_TIMING_STAT_EXCLUSIVE(RenderTranslucency);
 			SCOPE_CYCLE_COUNTER(STAT_TranslucencyDrawTime);
-			RenderTranslucency(RHICmdList, ViewList);
-			FRHICommandListExecutor::GetImmediateCommandList().PollOcclusionQueries();
-			RHICmdList.ImmediateFlush(EImmediateFlushType::DispatchToRHIThread);
+			RenderTranslucency(GraphBuilder, BasePassRenderTargets, ViewList);
+			AddPass(GraphBuilder, PollOcclusionQueriesAndDispatchToRHIThreadPass);
 		}
-		
-		RHICmdList.EndRenderPass();
 	}
 	else
 	{
-		RHICmdList.NextSubpass();
-		RHICmdList.NextSubpass();
-		RHICmdList.EndRenderPass();
+		{
+			auto* PassParameters = GraphBuilder.AllocParameters<FRenderTargetParameters>();
+			PassParameters->RenderTargets = BasePassRenderTargets;
+
+			GraphBuilder.AddPass(RDG_EVENT_NAME("Nextsubpass"), PassParameters, ERDGPassFlags::Raster,
+				[this](FRHICommandListImmediate& RHICmdList)
+			{
+				RHICmdList.NextSubpass();
+				RHICmdList.NextSubpass();
+			});
+		}
 		
 		// SceneColor + GBuffer write, SceneDepth is read only
 		{
-			for (int32 Index = 0; Index < UE_ARRAY_COUNT(ColorTargets); ++Index)
-			{
-				BasePassInfo.ColorRenderTargets[Index].Action = ERenderTargetActions::Load_Store;
-			}
-			BasePassInfo.DepthStencilRenderTarget.Action = EDepthStencilTargetActions::LoadDepthStencil_StoreDepthStencil;
-			BasePassInfo.DepthStencilRenderTarget.ExclusiveDepthStencil = FExclusiveDepthStencil::DepthRead_StencilRead;
-			BasePassInfo.SubpassHint = ESubpassHint::None;
-			BasePassInfo.NumOcclusionQueries = 0;
-			BasePassInfo.bOcclusionQueries = false;
-			
-			RHICmdList.BeginRenderPass(BasePassInfo, TEXT("AfterBasePass"));
+			BasePassRenderTargets.DepthStencil.SetDepthStencilAccess(FExclusiveDepthStencil::DepthRead_StencilRead);
+			BasePassRenderTargets.SubpassHint = ESubpassHint::None;
+			BasePassRenderTargets.NumOcclusionQueries = 0;
 			
+			RDG_EVENT_SCOPE(GraphBuilder, "AfterBasePass");
+
 			if (ViewFamily.EngineShowFlags.Decals)
 			{
 				CSV_SCOPED_TIMING_STAT_EXCLUSIVE(RenderDecals);
-				RenderDecals(RHICmdList);
+				RenderDecals(GraphBuilder, BasePassRenderTargets, MobileSceneTextures);
 			}
-			
-			RHICmdList.EndRenderPass();
 		}
 
 		// SceneColor write, SceneDepth is read only
 		{
-			FRHIRenderPassInfo ShadingPassInfo(
-				SceneContext.GetSceneColorSurface(),
-				ERenderTargetActions::Load_Store,
-				nullptr,
-				SceneContext.GetSceneDepthSurface(),
-				EDepthStencilTargetActions::LoadDepthStencil_StoreDepthStencil, 
-				nullptr,
-				nullptr,
-				FExclusiveDepthStencil::DepthRead_StencilWrite
-			);
-			ShadingPassInfo.NumOcclusionQueries = 0;
-			ShadingPassInfo.bOcclusionQueries = false;
-			
-			RHICmdList.BeginRenderPass(ShadingPassInfo, TEXT("MobileShadingPass"));
+			for (uint32 i = 1; i < UE_ARRAY_COUNT(ColorTargets); ++i)
+			{
+				BasePassRenderTargets[i] = FRenderTargetBinding();
+			}
+			BasePassRenderTargets.DepthStencil.SetDepthStencilAccess(FExclusiveDepthStencil::DepthRead_StencilWrite);
+
+			RDG_EVENT_SCOPE(GraphBuilder, "MobileShadingPass");
 			
-			MobileDeferredShadingPass(RHICmdList, *Scene, *ViewList[0], SortedLightSet);
+			MobileDeferredShadingPass(GraphBuilder, BasePassRenderTargets, MobileSceneTextures, *Scene, *ViewList[0], SortedLightSet);
 			// Draw translucency.
 			if (ViewFamily.EngineShowFlags.Translucency)
 			{
 				CSV_SCOPED_TIMING_STAT_EXCLUSIVE(RenderTranslucency);
 				SCOPE_CYCLE_COUNTER(STAT_TranslucencyDrawTime);
-				RenderTranslucency(RHICmdList, ViewList);
-				FRHICommandListExecutor::GetImmediateCommandList().PollOcclusionQueries();
-				RHICmdList.ImmediateFlush(EImmediateFlushType::DispatchToRHIThread);
+				RenderTranslucency(GraphBuilder, BasePassRenderTargets, ViewList);
+				AddPass(GraphBuilder, PollOcclusionQueriesAndDispatchToRHIThreadPass);
 			}
-
-			RHICmdList.EndRenderPass();
 		}
 	}
-
-	return ColorTargets[0];
 }
 
 BEGIN_SHADER_PARAMETER_STRUCT(FMobileDebugViewPassParameters, )
@@ -1332,14 +1315,14 @@ bool FMobileSceneRenderer::RequiresMultiPass(FRHICommandListImmediate& RHICmdLis
 }
 
 
-void FMobileSceneRenderer::ConditionalResolveSceneDepth(FRHICommandListImmediate& RHICmdList, const FViewInfo& View)
+void FMobileSceneRenderer::ConditionalResolveSceneDepth(FRDGBuilder& GraphBuilder, const FViewInfo& View, FRDGTextureMSAA& SceneDepth)
 {
 	FSceneRenderTargets& SceneContext = FSceneRenderTargets::Get();
 	
 	if (IsSimulatedPlatform(ShaderPlatform)) // mobile emulation on PC
 	{
 		// resolve MSAA depth for translucency
-		ResolveSceneDepth(RHICmdList);
+		AddResolveSceneDepthPass(GraphBuilder, View, SceneDepth);
 	}
 	else if (IsAndroidOpenGLESPlatform(ShaderPlatform))
 	{
@@ -1351,18 +1334,19 @@ void FMobileSceneRenderer::ConditionalResolveSceneDepth(FRHICommandListImmediate
 
 		if (bDecals || bModulatedShadows || bAlwaysResolveDepth || View.bUsesSceneDepth || bRequiresPixelProjectedPlanarRelfectionPass)
 		{
-			SCOPED_DRAW_EVENT(RHICmdList, ConditionalResolveSceneDepth);
-
 			// WEBGL copies depth from SceneColor alpha to a separate texture
 			// Switch target to force hardware flush current depth to texture
-			FTextureRHIRef DummySceneColor = GSystemTextures.BlackDummy->GetRenderTargetItem().TargetableTexture;
-			FTextureRHIRef DummyDepthTarget = GSystemTextures.DepthDummy->GetRenderTargetItem().TargetableTexture;
-					
-			FRHIRenderPassInfo RPInfo(DummySceneColor, ERenderTargetActions::DontLoad_DontStore);
-			RPInfo.DepthStencilRenderTarget.Action = EDepthStencilTargetActions::ClearDepthStencil_StoreDepthStencil;
-			RPInfo.DepthStencilRenderTarget.DepthStencilTarget = DummyDepthTarget;
-			RPInfo.DepthStencilRenderTarget.ExclusiveDepthStencil = FExclusiveDepthStencil::DepthWrite_StencilWrite;
-			RHICmdList.BeginRenderPass(RPInfo, TEXT("ResolveDepth"));
+			FRDGTextureRef DummySceneColor = GraphBuilder.RegisterExternalTexture(GSystemTextures.BlackDummy, ERenderTargetTexture::Targetable);
+			FRDGTextureRef DummyDepthTarget = GraphBuilder.RegisterExternalTexture(GSystemTextures.DepthDummy, ERenderTargetTexture::Targetable);
+
+			auto* PassParameters = GraphBuilder.AllocParameters<FRenderTargetParameters>();
+
+			PassParameters->RenderTargets[0] = FRenderTargetBinding(DummySceneColor, ERenderTargetLoadAction::ENoAction);
+			PassParameters->RenderTargets.DepthStencil = FDepthStencilBinding(DummyDepthTarget, ERenderTargetLoadAction::EClear, FExclusiveDepthStencil::DepthWrite_StencilWrite);
+
+			GraphBuilder.AddPass(
+				RDG_EVENT_NAME("ResolveDepthPass"), PassParameters, ERDGPassFlags::Raster,
+				[&View, &SceneContext](FRHICommandListImmediate& RHICmdList)
 			{
 				FGraphicsPipelineStateInitializer GraphicsPSOInit;
 				RHICmdList.ApplyCachedRenderTargets(GraphicsPSOInit);
@@ -1397,8 +1381,7 @@ void FMobileSceneRenderer::ConditionalResolveSceneDepth(FRHICommandListImmediate
 					FIntPoint(1, 1),
 					ScreenVertexShader,
 					EDRF_UseTriangleOptimization);
-			} // force depth resolve
-			RHICmdList.EndRenderPass();	
+			}); // force depth resolve
 		}
 	}
 }
@@ -1449,20 +1432,6 @@ void FMobileSceneRenderer::UpdateSkyReflectionUniformBuffer()
 	Scene->UniformBuffers.MobileSkyReflectionUniformBuffer.UpdateUniformBufferImmediate(Parameters);
 }
 
-void FMobileSceneRenderer::CreateDirectionalLightUniformBuffers(FViewInfo& View)
-{
-	bool bDynamicShadows = ViewFamily.EngineShowFlags.DynamicShadows;
-	// First array entry is used for primitives with no lighting channel set
-	View.MobileDirectionalLightUniformBuffers[0] = TUniformBufferRef<FMobileDirectionalLightShaderParameters>::CreateUniformBufferImmediate(FMobileDirectionalLightShaderParameters(), UniformBuffer_SingleFrame);
-	// Fill in the other entries based on the lights
-	for (int32 ChannelIdx = 0; ChannelIdx < UE_ARRAY_COUNT(Scene->MobileDirectionalLights); ChannelIdx++)
-	{
-		FMobileDirectionalLightShaderParameters Params;
-		SetupMobileDirectionalLightUniformParameters(*Scene, View, VisibleLightInfos, ChannelIdx, bDynamicShadows, Params);
-		View.MobileDirectionalLightUniformBuffers[ChannelIdx + 1] = TUniformBufferRef<FMobileDirectionalLightShaderParameters>::CreateUniformBufferImmediate(Params, UniformBuffer_SingleFrame);
-	}
-}
-
 class FPreTonemapMSAA_Mobile : public FGlobalShader
 {
 	DECLARE_SHADER_TYPE(FPreTonemapMSAA_Mobile, Global);
@@ -1483,7 +1452,7 @@ public:
 
 IMPLEMENT_SHADER_TYPE(, FPreTonemapMSAA_Mobile,TEXT("/Engine/Private/PostProcessMobile.usf"),TEXT("PreTonemapMSAA_Mobile"),SF_Pixel);
 
-void FMobileSceneRenderer::PreTonemapMSAA(FRHICommandListImmediate& RHICmdList)
+void FMobileSceneRenderer::PreTonemapMSAA(FRDGBuilder& GraphBuilder, FRenderTargetBindingSlots& BasePassRenderTargets)
 {
 	// iOS only
 	bool bOnChipPP = GSupportsRenderTargetFormat_PF_FloatRGBA && GSupportsShaderFramebufferFetch &&	ViewFamily.EngineShowFlags.PostProcessing;
@@ -1492,43 +1461,50 @@ void FMobileSceneRenderer::PreTonemapMSAA(FRHICommandListImmediate& RHICmdList)
 	{
 		return;
 	}
-	
-	// Part of scene rendering pass
-	check (RHICmdList.IsInsideRenderPass());
 
 	FSceneRenderTargets& SceneContext = FSceneRenderTargets::Get();
+	const FIntPoint TargetSize = SceneContext.GetBufferSizeXY();
 
-	FGraphicsPipelineStateInitializer GraphicsPSOInit;
-	RHICmdList.ApplyCachedRenderTargets(GraphicsPSOInit);
-	GraphicsPSOInit.BlendState = TStaticBlendState<>::GetRHI();
-	GraphicsPSOInit.RasterizerState = TStaticRasterizerState<>::GetRHI();
-	GraphicsPSOInit.DepthStencilState = TStaticDepthStencilState<false, CF_Always>::GetRHI();
+	auto* PassParameters = GraphBuilder.AllocParameters<FRenderTargetParameters>();
+	PassParameters->RenderTargets = BasePassRenderTargets;
 
 	const auto ShaderMap = GetGlobalShaderMap(FeatureLevel);
 	TShaderMapRef<FScreenVS> VertexShader(ShaderMap);
 	TShaderMapRef<FPreTonemapMSAA_Mobile> PixelShader(ShaderMap);
 
-	extern TGlobalResource<FFilterVertexDeclaration> GFilterVertexDeclaration;
-	GraphicsPSOInit.BoundShaderState.VertexDeclarationRHI = GFilterVertexDeclaration.VertexDeclarationRHI;
-	GraphicsPSOInit.BoundShaderState.VertexShaderRHI = VertexShader.GetVertexShader();
-	GraphicsPSOInit.BoundShaderState.PixelShaderRHI = PixelShader.GetPixelShader();
-	GraphicsPSOInit.PrimitiveType = PT_TriangleList;
+	GraphBuilder.AddPass(
+		RDG_EVENT_NAME("PreTonemapMSAAPass"),
+		PassParameters,
+		ERDGPassFlags::Raster,
+		[this, TargetSize, VertexShader, PixelShader](FRHICommandListImmediate& RHICmdList)
+	{
+
+		FGraphicsPipelineStateInitializer GraphicsPSOInit;
+		RHICmdList.ApplyCachedRenderTargets(GraphicsPSOInit);
+		GraphicsPSOInit.BlendState = TStaticBlendState<>::GetRHI();
+		GraphicsPSOInit.RasterizerState = TStaticRasterizerState<>::GetRHI();
+		GraphicsPSOInit.DepthStencilState = TStaticDepthStencilState<false, CF_Always>::GetRHI();
 
-	SetGraphicsPipelineState(RHICmdList, GraphicsPSOInit);
-	
-	const FIntPoint TargetSize = SceneContext.GetBufferSizeXY();
-	RHICmdList.SetViewport(0, 0, 0.0f, TargetSize.X, TargetSize.Y, 1.0f);
-
-	DrawRectangle(
-		RHICmdList,
-		0, 0,
-		TargetSize.X, TargetSize.Y,
-		0, 0,
-		TargetSize.X, TargetSize.Y,
-		TargetSize,
-		TargetSize,
-		VertexShader,
-		EDRF_UseTriangleOptimization);
+		GraphicsPSOInit.BoundShaderState.VertexDeclarationRHI = GFilterVertexDeclaration.VertexDeclarationRHI;
+		GraphicsPSOInit.BoundShaderState.VertexShaderRHI = VertexShader.GetVertexShader();
+		GraphicsPSOInit.BoundShaderState.PixelShaderRHI = PixelShader.GetPixelShader();
+		GraphicsPSOInit.PrimitiveType = PT_TriangleList;
+
+		SetGraphicsPipelineState(RHICmdList, GraphicsPSOInit);
+
+		RHICmdList.SetViewport(0, 0, 0.0f, TargetSize.X, TargetSize.Y, 1.0f);
+
+		DrawRectangle(
+			RHICmdList,
+			0, 0,
+			TargetSize.X, TargetSize.Y,
+			0, 0,
+			TargetSize.X, TargetSize.Y,
+			TargetSize,
+			TargetSize,
+			VertexShader,
+			EDRF_UseTriangleOptimization);
+	});
 }
 
 /** Before SetupMobileBasePassAfterShadowInit, we need to update the uniform buffer and shadow info for all movable point lights.*/
diff --git a/Engine/Source/Runtime/Renderer/Private/MobileTranslucentRendering.cpp b/Engine/Source/Runtime/Renderer/Private/MobileTranslucentRendering.cpp
--- a/Engine/Source/Runtime/Renderer/Private/MobileTranslucentRendering.cpp
+++ b/Engine/Source/Runtime/Renderer/Private/MobileTranslucentRendering.cpp
@@ -31,19 +31,19 @@
 #include "MeshPassProcessor.inl"
 #include "ClearQuad.h"
 
-void FMobileSceneRenderer::RenderTranslucency(FRHICommandListImmediate& RHICmdList, const TArrayView<const FViewInfo*> PassViews)
+void FMobileSceneRenderer::RenderTranslucency(FRDGBuilder& GraphBuilder, FRenderTargetBindingSlots& BasePassRenderTargets, const TArrayView<const FViewInfo*> PassViews)
 {
 	ETranslucencyPass::Type TranslucencyPass = 
 		ViewFamily.AllowTranslucencyAfterDOF() ? ETranslucencyPass::TPT_StandardTranslucency : ETranslucencyPass::TPT_AllTranslucency;
 		
 	if (ShouldRenderTranslucency(TranslucencyPass))
 	{
-		SCOPED_DRAW_EVENT(RHICmdList, Translucency);
-		SCOPED_GPU_STAT(RHICmdList, Translucency);
+		RDG_EVENT_SCOPE(GraphBuilder, "Translucency");
+		RDG_GPU_STAT_SCOPE(GraphBuilder, Translucency);
 
 		for (int32 ViewIndex = 0; ViewIndex < PassViews.Num(); ViewIndex++)
 		{
-			SCOPED_CONDITIONAL_DRAW_EVENTF(RHICmdList, EventView, Views.Num() > 1, TEXT("View%d"), ViewIndex);
+			RDG_EVENT_SCOPE_CONDITIONAL(GraphBuilder, Views.Num() > 1, "View%d", ViewIndex);
 
 			const FViewInfo& View = *PassViews[ViewIndex];
 			if (!View.ShouldRenderView())
@@ -51,98 +51,71 @@ void FMobileSceneRenderer::RenderTranslucency(FRHICommandListImmediate& RHICmdLi
 				continue;
 			}
 
-			RHICmdList.SetViewport(View.ViewRect.Min.X, View.ViewRect.Min.Y, 0.0f, View.ViewRect.Max.X, View.ViewRect.Max.Y, 1.0f);
+			auto* TranslucencyBasePassParameters = GraphBuilder.AllocParameters<FMobileBasePassParameters>();
+			TranslucencyBasePassParameters->RenderTargets = BasePassRenderTargets;
+			TranslucencyBasePassParameters->MobileBasePass = Scene->UniformBuffers.MobileTranslucentBasePassUniformBuffer;
 
-			if (!View.Family->UseDebugViewPS())
+			GraphBuilder.AddPass(RDG_EVENT_NAME("RenderTranslucencyBasePass"), TranslucencyBasePassParameters, ERDGPassFlags::Raster,
+				[this, &View, TranslucencyPass](FRHICommandListImmediate& RHICmdList)
 			{
-				if (Scene->UniformBuffers.UpdateViewUniformBuffer(View))
+				RHICmdList.SetViewport(View.ViewRect.Min.X, View.ViewRect.Min.Y, 0.0f, View.ViewRect.Max.X, View.ViewRect.Max.Y, 1.0f);
+
+				if (!View.Family->UseDebugViewPS())
 				{
-					UpdateTranslucentBasePassUniformBuffer(RHICmdList, View);
-					UpdateDirectionalLightUniformBuffers(RHICmdList, View);
+					if (Scene->UniformBuffers.UpdateViewUniformBuffer(View))
+					{
+						UpdateTranslucentBasePassUniformBuffer(RHICmdList, View);
+						UpdateDirectionalLightUniformBuffers(RHICmdList, View);
+					}
+
+					const EMeshPass::Type MeshPass = TranslucencyPassToMeshPass(TranslucencyPass);
+					View.ParallelMeshDrawCommandPasses[MeshPass].DispatchDraw(nullptr, RHICmdList);
 				}
-		
-				const EMeshPass::Type MeshPass = TranslucencyPassToMeshPass(TranslucencyPass);
-				View.ParallelMeshDrawCommandPasses[MeshPass].DispatchDraw(nullptr, RHICmdList);
-			}
+			});
 		}
 	}
 }
 
-bool FMobileSceneRenderer::RenderInverseOpacity(FRHICommandListImmediate& RHICmdList, const FViewInfo& View)
+void FMobileSceneRenderer::RenderInverseOpacity(FRDGBuilder& GraphBuilder, const FViewInfo& View)
 {
-	// Function MUST be self-contained wrt RenderPasses
-	check(RHICmdList.IsOutsideRenderPass());
-
-	bool bDirty = false;
 	FSceneRenderTargets& SceneContext = FSceneRenderTargets::Get();
 
-	SceneContext.AllocSceneColor(RHICmdList);
+	SceneContext.AllocSceneColor(GraphBuilder.RHICmdList);
 
-	if (Scene->UniformBuffers.UpdateViewUniformBuffer(View))
-	{
-		UpdateTranslucentBasePassUniformBuffer(RHICmdList, View);
-		UpdateDirectionalLightUniformBuffers(RHICmdList, View);
-	}
+	FRDGTextureMSAA SceneColorMSAA = RegisterExternalTextureMSAA(GraphBuilder, SceneContext.GetSceneColor());
+	FRDGTextureMSAA SceneDepthMSAA = RegisterExternalTextureMSAA(GraphBuilder, SceneContext.SceneDepthZ);
 
-	const bool bMobileMSAA = SceneContext.GetSceneColorSurface()->GetNumSamples() > 1;
-	
-	FRHITexture* SceneColorResolve = bMobileMSAA ? SceneContext.GetSceneColorTexture() : nullptr;
-	ERenderTargetActions ColorTargetAction = bMobileMSAA ? ERenderTargetActions::Clear_Resolve : ERenderTargetActions::Clear_Store;
-	FRHIRenderPassInfo RPInfo(
-		SceneContext.GetSceneColorSurface(), 
-		ColorTargetAction,
-		SceneColorResolve,
-		SceneContext.GetSceneDepthSurface(),
-		EDepthStencilTargetActions::ClearDepthStencil_DontStoreDepthStencil,
-		nullptr,
-		FExclusiveDepthStencil::DepthWrite_StencilWrite
-	);
+	auto* InverseOpacityParameters = GraphBuilder.AllocParameters<FMobileBasePassParameters>();
+	InverseOpacityParameters->MobileBasePass = Scene->UniformBuffers.MobileTranslucentBasePassUniformBuffer;
+	InverseOpacityParameters->RenderTargets[0] = FRenderTargetBinding(SceneColorMSAA.Target, SceneColorMSAA.Resolve, ERenderTargetLoadAction::EClear);
+	InverseOpacityParameters->RenderTargets.DepthStencil = FDepthStencilBinding(SceneDepthMSAA.Target, ERenderTargetLoadAction::EClear, FExclusiveDepthStencil::DepthWrite_StencilWrite);
 	// Opacity could fetch depth as we use exactly the same shaders as in base pass
-	RPInfo.SubpassHint = ESubpassHint::DepthReadSubpass;
-
-	// make sure targets are writable
-	FRHITransitionInfo TransitionsBefore[3];
-	int32 NumTransitionsBefore = 0;
-	TransitionsBefore[NumTransitionsBefore] = FRHITransitionInfo(SceneContext.GetSceneColorSurface(), ERHIAccess::Unknown, ERHIAccess::RTV);
-	++NumTransitionsBefore;
-	TransitionsBefore[NumTransitionsBefore] = FRHITransitionInfo(SceneContext.GetSceneDepthSurface(), ERHIAccess::Unknown, ERHIAccess::DSVWrite);
-	++NumTransitionsBefore;
-	if (SceneColorResolve)
-	{
-		TransitionsBefore[NumTransitionsBefore] = FRHITransitionInfo(SceneColorResolve, ERHIAccess::Unknown, ERHIAccess::RTV | ERHIAccess::ResolveDst);
-		++NumTransitionsBefore;
-	}
-	RHICmdList.Transition(MakeArrayView(TransitionsBefore, NumTransitionsBefore));
+	InverseOpacityParameters->RenderTargets.SubpassHint = ESubpassHint::DepthReadSubpass;
 
-	if (Scene->UniformBuffers.UpdateViewUniformBuffer(View))
+	GraphBuilder.AddPass(RDG_EVENT_NAME("InverseOpacityPass"), InverseOpacityParameters, ERDGPassFlags::Raster,
+		[this, &View](FRHICommandListImmediate& RHICmdList)
 	{
-		UpdateTranslucentBasePassUniformBuffer(RHICmdList, View);
-		UpdateDirectionalLightUniformBuffers(RHICmdList, View);
-	}
-	
-	RHICmdList.BeginRenderPass(RPInfo, TEXT("InverseOpacity"));
+		if (Scene->UniformBuffers.UpdateViewUniformBuffer(View))
+		{
+			UpdateTranslucentBasePassUniformBuffer(RHICmdList, View);
+			UpdateDirectionalLightUniformBuffers(RHICmdList, View);
+		}
 
-	// Mobile multi-view is not side by side stereo
-	const FViewInfo& TranslucentViewport = (View.bIsMobileMultiViewEnabled) ? Views[0] : View;
-	RHICmdList.SetViewport(TranslucentViewport.ViewRect.Min.X, TranslucentViewport.ViewRect.Min.Y, 0.0f, TranslucentViewport.ViewRect.Max.X, TranslucentViewport.ViewRect.Max.Y, 1.0f);
-		
-	// Default clear value for a SceneColor is (0,0,0,0), after this passs will blend inverse opacity into final render target with an 1-SrcAlpha op
-	// to make this blending work untouched pixels must have alpha = 1
-	DrawClearQuad(RHICmdList, FLinearColor(0,0,0,1));
-
-	RHICmdList.NextSubpass();
-	if (ShouldRenderTranslucency(ETranslucencyPass::TPT_AllTranslucency) && View.ShouldRenderView())
-	{		
-		View.ParallelMeshDrawCommandPasses[EMeshPass::MobileInverseOpacity].DispatchDraw(nullptr, RHICmdList);
-		bDirty |= View.ParallelMeshDrawCommandPasses[EMeshPass::MobileInverseOpacity].HasAnyDraw();
-	}
-	
-	RHICmdList.EndRenderPass();
-	
-	ERHIAccess AccessBefore = bMobileMSAA ? ERHIAccess::RTV | ERHIAccess::ResolveDst : ERHIAccess::RTV;
-	RHICmdList.Transition(FRHITransitionInfo(SceneContext.GetSceneColorTexture(), AccessBefore, ERHIAccess::SRVMask));
-	
-	return bDirty;
+		// Mobile multi-view is not side by side stereo
+		const FViewInfo& TranslucentViewport = (View.bIsMobileMultiViewEnabled) ? Views[0] : View;
+		RHICmdList.SetViewport(TranslucentViewport.ViewRect.Min.X, TranslucentViewport.ViewRect.Min.Y, 0.0f, TranslucentViewport.ViewRect.Max.X, TranslucentViewport.ViewRect.Max.Y, 1.0f);
+
+		// Default clear value for a SceneColor is (0,0,0,0), after this passs will blend inverse opacity into final render target with an 1-SrcAlpha op
+		// to make this blending work untouched pixels must have alpha = 1
+		DrawClearQuad(RHICmdList, FLinearColor(0, 0, 0, 1));
+
+		RHICmdList.NextSubpass();
+		if (ShouldRenderTranslucency(ETranslucencyPass::TPT_AllTranslucency) && View.ShouldRenderView())
+		{
+			View.ParallelMeshDrawCommandPasses[EMeshPass::MobileInverseOpacity].DispatchDraw(nullptr, RHICmdList);
+			View.ParallelMeshDrawCommandPasses[EMeshPass::MobileInverseOpacity].HasAnyDraw();
+		}
+	});
 }
 
 // This pass is registered only when we render to scene capture, see UpdateSceneCaptureContentMobile_RenderThread()
diff --git a/Engine/Source/Runtime/Renderer/Private/PostProcess/PostProcessAmbientOcclusionMobile.cpp b/Engine/Source/Runtime/Renderer/Private/PostProcess/PostProcessAmbientOcclusionMobile.cpp
--- a/Engine/Source/Runtime/Renderer/Private/PostProcess/PostProcessAmbientOcclusionMobile.cpp
+++ b/Engine/Source/Runtime/Renderer/Private/PostProcess/PostProcessAmbientOcclusionMobile.cpp
@@ -416,24 +416,6 @@ void FMobileSceneRenderer::ReleaseAmbientOcclusionOutputs()
 	GAmbientOcclusionMobileOutputs.Release();
 }
 
-void FMobileSceneRenderer::RenderAmbientOcclusion(FRHICommandListImmediate& RHICmdList, const TRefCountPtr<IPooledRenderTarget>& SceneDepthZ)
-{
-	checkSlow(GAmbientOcclusionMobileOutputs.IsValid() && SceneDepthZ.IsValid());
-
-	SCOPED_DRAW_EVENT(RHICmdList, AmbientOcclusion);
-
-	FMemMark Mark(FMemStack::Get());
-	FRDGBuilder GraphBuilder(RHICmdList);
-
-	FRDGTextureRef SceneDepthTexture = GraphBuilder.RegisterExternalTexture(SceneDepthZ, TEXT("SceneDepthTexture"));
-
-	FRDGTextureRef AmbientOcclusionTexture = GraphBuilder.RegisterExternalTexture(GAmbientOcclusionMobileOutputs.AmbientOcclusionTexture, TEXT("AmbientOcclusionTexture"));
-
-	RenderAmbientOcclusion(GraphBuilder, SceneDepthTexture, AmbientOcclusionTexture);
-
-	GraphBuilder.Execute();
-}
-
 void FMobileSceneRenderer::RenderAmbientOcclusion(FRDGBuilder& GraphBuilder, FRDGTextureRef SceneDepthTexture, FRDGTextureRef AmbientOcclusionTexture)
 {
 	static const auto GTAOThicknessBlendCVar = IConsoleManager::Get().FindTConsoleVariableDataFloat(TEXT("r.GTAO.ThicknessBlend"));
diff --git a/Engine/Source/Runtime/Renderer/Private/PostProcess/PostProcessPixelProjectedReflectionMobile.cpp b/Engine/Source/Runtime/Renderer/Private/PostProcess/PostProcessPixelProjectedReflectionMobile.cpp
--- a/Engine/Source/Runtime/Renderer/Private/PostProcess/PostProcessPixelProjectedReflectionMobile.cpp
+++ b/Engine/Source/Runtime/Renderer/Private/PostProcess/PostProcessPixelProjectedReflectionMobile.cpp
@@ -201,26 +201,6 @@ void FMobileSceneRenderer::ReleasePixelProjectedReflectionOutputs()
 	GPixelProjectedReflectionMobileOutputs.Release();
 }
 
-void FMobileSceneRenderer::RenderPixelProjectedReflection(FRHICommandListImmediate& RHICmdList, const FSceneRenderTargets& SceneContext, const FPlanarReflectionSceneProxy* PlanarReflectionSceneProxy)
-{
-	checkSlow(GPixelProjectedReflectionMobileOutputs.IsValid());
-
-	SCOPED_DRAW_EVENT(RHICmdList, PixelProjectedReflection);
-
-	FMemMark Mark(FMemStack::Get());
-	FRDGBuilder GraphBuilder(RHICmdList);
-
-	FRDGTextureRef SceneColorTexture = GraphBuilder.RegisterExternalTexture(SceneContext.GetSceneColor(), TEXT("SceneColorTexture"));
-
-	FRDGTextureRef SceneDepthTexture = GraphBuilder.RegisterExternalTexture(SceneContext.SceneDepthZ, TEXT("SceneDepthTexture"));
-
-	FRDGTextureRef PixelProjectedReflectionTexture = GraphBuilder.RegisterExternalTexture(GPixelProjectedReflectionMobileOutputs.PixelProjectedReflectionTexture, TEXT("PixelProjectedReflectionTexture"));
-
-	RenderPixelProjectedReflection(GraphBuilder, SceneColorTexture, SceneDepthTexture, PixelProjectedReflectionTexture, PlanarReflectionSceneProxy);
-
-	GraphBuilder.Execute();
-}
-
 void FMobileSceneRenderer::RenderPixelProjectedReflection(FRDGBuilder& GraphBuilder, FRDGTextureRef SceneColorTexture, FRDGTextureRef SceneDepthTexture, FRDGTextureRef PixelProjectedReflectionTexture, const FPlanarReflectionSceneProxy* PlanarReflectionSceneProxy)
 {
 	const FIntPoint BufferSize = PlanarReflectionSceneProxy->RenderTarget->GetSizeXY();
diff --git a/Engine/Source/Runtime/Renderer/Private/Renderer.cpp b/Engine/Source/Runtime/Renderer/Private/Renderer.cpp
--- a/Engine/Source/Runtime/Renderer/Private/Renderer.cpp
+++ b/Engine/Source/Runtime/Renderer/Private/Renderer.cpp
@@ -189,11 +189,6 @@ void FRendererModule::DrawTileMesh(FRHICommandListImmediate& RHICmdList, FMeshPa
 			DrawRenderState.SetReflectionCaptureUniformBuffer(EmptyReflectionCaptureUniformBuffer);
 		}
 
-		if (ShadingPath == EShadingPath::Mobile)
-		{
-			View.MobileDirectionalLightUniformBuffers[0] = TUniformBufferRef<FMobileDirectionalLightShaderParameters>::CreateUniformBufferImmediate(FMobileDirectionalLightShaderParameters(), UniformBuffer_SingleFrame);
-		}
-
 		//get the blend mode of the material
 		const EBlendMode MaterialBlendMode = Mesh.MaterialRenderProxy->GetIncompleteMaterialWithFallback(FeatureLevel).GetBlendMode();
 
diff --git a/Engine/Source/Runtime/Renderer/Private/SceneOcclusion.cpp b/Engine/Source/Runtime/Renderer/Private/SceneOcclusion.cpp
--- a/Engine/Source/Runtime/Renderer/Private/SceneOcclusion.cpp
+++ b/Engine/Source/Runtime/Renderer/Private/SceneOcclusion.cpp
@@ -1420,7 +1420,7 @@ void FDeferredShadingSceneRenderer::RenderOcclusion(
 	}
 }
 
-void FMobileSceneRenderer::RenderOcclusion(FRHICommandListImmediate& RHICmdList)
+void FMobileSceneRenderer::RenderOcclusion(FRDGBuilder& GraphBuilder, FRenderTargetBindingSlots& BasePassRenderTargets)
 {
 	if (!DoOcclusionQueries(FeatureLevel))
 	{
@@ -1429,17 +1429,22 @@ void FMobileSceneRenderer::RenderOcclusion(FRHICommandListImmediate& RHICmdList)
 
 	{
 		SCOPED_NAMED_EVENT(FMobileSceneRenderer_BeginOcclusionTests, FColor::Emerald);
-		const FViewOcclusionQueriesPerView QueriesPerView = AllocateOcclusionTests(Scene, VisibleLightInfos, Views);
+		FViewOcclusionQueriesPerView QueriesPerView = AllocateOcclusionTests(Scene, VisibleLightInfos, Views);
 
 		if (QueriesPerView.Num())
 		{
-			BeginOcclusionTests(RHICmdList, Views, FeatureLevel, QueriesPerView, 1.0f);
+			auto* PassParameters = GraphBuilder.AllocParameters<FRenderTargetParameters>();
+			PassParameters->RenderTargets = BasePassRenderTargets;
+
+			GraphBuilder.AddPass(RDG_EVENT_NAME("BeginOcclusionTestsPass"), PassParameters, ERDGPassFlags::Raster,
+				[this, LocalQueriesPerView = MoveTemp(QueriesPerView)](FRHICommandListImmediate& RHICmdList)
+			{
+				BeginOcclusionTests(RHICmdList, Views, FeatureLevel, LocalQueriesPerView, 1.0f);
+			});
 		}
 	}
 
-	FRDGBuilder GraphBuilder(RHICmdList);
 	FenceOcclusionTests(GraphBuilder);
-	GraphBuilder.Execute();
 }
 
 DECLARE_CYCLE_STAT(TEXT("OcclusionSubmittedFence Dispatch"), STAT_OcclusionSubmittedFence_Dispatch, STATGROUP_SceneRendering);
diff --git a/Engine/Source/Runtime/Renderer/Private/SceneRendering.cpp b/Engine/Source/Runtime/Renderer/Private/SceneRendering.cpp
--- a/Engine/Source/Runtime/Renderer/Private/SceneRendering.cpp
+++ b/Engine/Source/Runtime/Renderer/Private/SceneRendering.cpp
@@ -1868,16 +1868,6 @@ FViewInfo* FViewInfo::CreateSnapshot() const
 	// we want these to start null without a reference count, since we clear a ref later
 	TUniformBufferRef<FViewUniformShaderParameters> NullViewUniformBuffer;
 	FMemory::Memcpy(Result->ViewUniformBuffer, NullViewUniformBuffer); 
-	TUniformBufferRef<FMobileDirectionalLightShaderParameters> NullMobileDirectionalLightUniformBuffer;
-	for (size_t i = 0; i < UE_ARRAY_COUNT(Result->MobileDirectionalLightUniformBuffers); i++)
-	{
-		// This memcpy is necessary to clear the reference from the memcpy of the whole of this -> Result without releasing the pointer
-		// But what we really want is the null buffer.
-		FMemory::Memcpy(
-			&Result->MobileDirectionalLightUniformBuffers[i],
-			&GetNullMobileDirectionalLightShaderParameters(),
-			sizeof(TUniformBufferRef<FMobileDirectionalLightShaderParameters>));
-	}
 
 	TUniquePtr<FViewUniformShaderParameters> NullViewParameters;
 	FMemory::Memcpy(Result->CachedViewUniformShaderParameters, NullViewParameters); 
@@ -4652,15 +4642,6 @@ FRDGTextureRef CreateHalfResolutionDepthCheckerboardMinMax(FRDGBuilder& GraphBui
 	return SmallDepthTexture;
 }
 
-void FSceneRenderer::ResolveSceneDepth(FRHICommandListImmediate& RHICmdList)
-{
-	FSceneRenderTargets& SceneContext = FSceneRenderTargets::Get();
-
-	FRDGBuilder GraphBuilder(RHICmdList);
-	AddResolveSceneDepthPass(GraphBuilder, Views, RegisterExternalTextureMSAA(GraphBuilder, SceneContext.SceneDepthZ));
-	GraphBuilder.Execute();
-}
-
 void FSceneRenderer::UpdateSkyIrradianceGpuBuffer(FRHICommandListImmediate& RHICmdList)
 {
 	if (Scene == nullptr)
diff --git a/Engine/Source/Runtime/Renderer/Private/SceneRendering.h b/Engine/Source/Runtime/Renderer/Private/SceneRendering.h
--- a/Engine/Source/Runtime/Renderer/Private/SceneRendering.h
+++ b/Engine/Source/Runtime/Renderer/Private/SceneRendering.h
@@ -48,6 +48,7 @@ class FRaytracingLightDataPacked;
 class FRayTracingLocalShaderBindingWriter;
 class FVirtualShadowMapClipmap;
 class FShadowProjectionPassParameters;
+class FMobileBasePassParameters;
 
 struct FCloudRenderContext;
 struct FSceneWithoutWaterTextures;
@@ -2004,7 +2005,6 @@ protected:
 	bool ShouldRenderSkyAtmosphereEditorNotifications() const;
 
 	void ResolveSceneColor(FRHICommandListImmediate& RHICmdList);
-	void ResolveSceneDepth(FRHICommandListImmediate& RHICmdList);
 
 	/**
 	 * Rounds up lights and sorts them according to what type of renderer supports them. The result is stored in OutSortedLights 
@@ -2055,7 +2055,7 @@ public:
 
 	virtual void RenderHitProxies(FRHICommandListImmediate& RHICmdList) override;
 
-	bool RenderInverseOpacity(FRHICommandListImmediate& RHICmdList, const FViewInfo& View);
+	void RenderInverseOpacity(FRDGBuilder& GraphBuilder, const FViewInfo& View);
 
 protected:
 	/** Finds the visible dynamic shadows for each view. */
@@ -2068,10 +2068,10 @@ protected:
 
 	void InitViews(FRDGBuilder& GraphBuilder);
 
-	void RenderPrePass(FRHICommandListImmediate& RHICmdList);
+	void RenderPrePass(FRDGBuilder& GraphBuilder, FRenderTargetBindingSlots& BasePassRenderTargets, TFunction<void(FRenderTargetBindingSlots&)> UpdateRenderTargetsLoadAction);
 
 	/** Renders the opaque base pass for mobile. */
-	void RenderMobileBasePass(FRHICommandListImmediate& RHICmdList, const TArrayView<const FViewInfo*> PassViews);
+	void RenderMobileBasePass(FRDGBuilder& GraphBuilder, FRenderTargetBindingSlots& BasePassRenderTargets, TFunction<void(FRenderTargetBindingSlots &)> UpdateRenderTargetsLoadAction, const TArrayView<const FViewInfo*> PassViews);
 
 	void RenderMobileEditorPrimitives(FRHICommandList& RHICmdList, const FViewInfo& View, const FMeshPassProcessorRenderState& DrawRenderState);
 
@@ -2079,13 +2079,13 @@ protected:
 	void RenderMobileDebugView(FRDGBuilder& GraphBuilder, const TArrayView<const FViewInfo*> PassViews);
 
 	/** Render modulated shadow projections in to the scene, loops over any unrendered shadows until all are processed.*/
-	void RenderModulatedShadowProjections(FRHICommandListImmediate& RHICmdList);
+	void RenderModulatedShadowProjections(FRDGBuilder& GraphBuilder, FRenderTargetBindingSlots& BasePassRenderTargets, TRDGUniformBufferRef<FMobileSceneTextureUniformParameters> MobileSceneTextures);
 
 	/** Resolves scene depth in case hardware does not support reading depth in the shader */
-	void ConditionalResolveSceneDepth(FRHICommandListImmediate& RHICmdList, const FViewInfo& View);
+	void ConditionalResolveSceneDepth(FRDGBuilder& GraphBuilder, const FViewInfo& View, FRDGTextureMSAA& SceneDepth);
 
 	/** Issues occlusion queries */
-	void RenderOcclusion(FRHICommandListImmediate& RHICmdList);
+	void RenderOcclusion(FRDGBuilder& GraphBuilder, FRenderTargetBindingSlots& BasePassRenderTargets);
 	
 	/** Computes how many queries will be issued this frame */
 	int32 ComputeNumOcclusionQueriesToBatch() const;
@@ -2094,16 +2094,13 @@ protected:
 	bool RequiresMultiPass(FRHICommandListImmediate& RHICmdList, const FViewInfo& View) const;
 
 	/** Renders decals. */
-	void RenderDecals(FRHICommandListImmediate& RHICmdList);
+	void RenderDecals(FRDGBuilder& GraphBuilder, FRenderTargetBindingSlots& BasePassRenderTargets, TRDGUniformBufferRef<FMobileSceneTextureUniformParameters> MobileSceneTextures);
 
 	/** Renders the base pass for translucency. */
-	void RenderTranslucency(FRHICommandListImmediate& RHICmdList, const TArrayView<const FViewInfo*> PassViews);
-
-	/** Creates uniform buffers with the mobile directional light parameters, for each lighting channel. Called by InitViews */
-	void CreateDirectionalLightUniformBuffers(FViewInfo& View);
+	void RenderTranslucency(FRDGBuilder& GraphBuilder, FRenderTargetBindingSlots& BasePassRenderTargets, const TArrayView<const FViewInfo*> PassViews);
 
 	/** On chip pre-tonemap before scene color MSAA resolve (iOS only) */
-	void PreTonemapMSAA(FRHICommandListImmediate& RHICmdList);
+	void PreTonemapMSAA(FRDGBuilder& GraphBuilder, FRenderTargetBindingSlots& BasePassRenderTargets);
 
 	void SortMobileBasePassAfterShadowInit(FExclusiveDepthStencil::Type BasePassDepthStencilAccess, FViewVisibleCommandsPerView& ViewCommandsPerView);
 	void SetupMobileBasePassAfterShadowInit(FExclusiveDepthStencil::Type BasePassDepthStencilAccess, FViewVisibleCommandsPerView& ViewCommandsPerView);
@@ -2113,16 +2110,14 @@ protected:
 	void UpdateDirectionalLightUniformBuffers(FRHICommandListImmediate& RHICmdList, const FViewInfo& View);
 	void UpdateSkyReflectionUniformBuffer();
 
-	FRHITexture* RenderForward(FRHICommandListImmediate& RHICmdList, const TArrayView<const FViewInfo*> ViewList);
-	FRHITexture* RenderDeferred(FRHICommandListImmediate& RHICmdList, const TArrayView<const FViewInfo*> ViewList, const FSortedLightSetSceneInfo& SortedLightSet);
+	void RenderForward(FRDGBuilder& GraphBuilder, const TArrayView<const FViewInfo*> ViewList, FRDGTextureRef ViewFamilyTexture, FRDGTextureMSAA& SceneColorMSAA, FRDGTextureMSAA& SceneDepthMSAA, TRDGUniformBufferRef<FMobileSceneTextureUniformParameters> MobileSceneTextures);
+	void RenderDeferred(FRDGBuilder& GraphBuilder, const TArrayView<const FViewInfo*> ViewList, const FSortedLightSetSceneInfo& SortedLightSet, FRDGTextureRef ViewFamilyTexture, FRDGTextureMSAA& SceneColorMSAA, FRDGTextureMSAA& SceneDepthMSAA, TRDGUniformBufferRef<FMobileSceneTextureUniformParameters> MobileSceneTextures);
 	
 	void InitAmbientOcclusionOutputs(FRHICommandListImmediate& RHICmdList, const TRefCountPtr<IPooledRenderTarget>& SceneDepthZ);
-	void RenderAmbientOcclusion(FRHICommandListImmediate& RHICmdList, const TRefCountPtr<IPooledRenderTarget>& SceneDepthZ);
 	void RenderAmbientOcclusion(FRDGBuilder& GraphBuilder, FRDGTextureRef SceneDepthTexture, FRDGTextureRef AmbientOcclusionTexture);
 	void ReleaseAmbientOcclusionOutputs();
 
 	void InitPixelProjectedReflectionOutputs(FRHICommandListImmediate& RHICmdList, const FIntPoint& BufferSize);
-	void RenderPixelProjectedReflection(FRHICommandListImmediate& RHICmdList, const FSceneRenderTargets& SceneContext, const FPlanarReflectionSceneProxy* PlanarReflectionSceneProxy);
 	void RenderPixelProjectedReflection(FRDGBuilder& GraphBuilder, FRDGTextureRef SceneColorTexture, FRDGTextureRef SceneDepthTexture, FRDGTextureRef PixelProjectedReflectionTexture, const FPlanarReflectionSceneProxy* PlanarReflectionSceneProxy);
 	void ReleasePixelProjectedReflectionOutputs();
 
diff --git a/Engine/Source/Runtime/Renderer/Private/ShadowRendering.cpp b/Engine/Source/Runtime/Renderer/Private/ShadowRendering.cpp
--- a/Engine/Source/Runtime/Renderer/Private/ShadowRendering.cpp
+++ b/Engine/Source/Runtime/Renderer/Private/ShadowRendering.cpp
@@ -2037,7 +2037,7 @@ void FDeferredShadingSceneRenderer::RenderDeferredShadowProjections(
 	}
 }
 
-void FMobileSceneRenderer::RenderModulatedShadowProjections(FRHICommandListImmediate& RHICmdList)
+void FMobileSceneRenderer::RenderModulatedShadowProjections(FRDGBuilder& GraphBuilder, FRenderTargetBindingSlots& BasePassRenderTargets, TRDGUniformBufferRef<FMobileSceneTextureUniformParameters> MobileSceneTextures)
 {
 	if (IsSimpleForwardShadingEnabled(ShaderPlatform) || !ViewFamily.EngineShowFlags.DynamicShadows)
 	{
@@ -2046,13 +2046,12 @@ void FMobileSceneRenderer::RenderModulatedShadowProjections(FRHICommandListImmed
 
 	SCOPED_NAMED_EVENT(FMobileSceneRenderer_RenderModulatedShadowProjections, FColor::Emerald);
 	SCOPE_CYCLE_COUNTER(STAT_ProjectedShadowDrawTime);
-	SCOPED_DRAW_EVENT(RHICmdList, ShadowProjectionOnOpaque);
-	SCOPED_GPU_STAT(RHICmdList, ShadowProjection);
-
-	FRDGBuilder GraphBuilder(RHICmdList);
+	RDG_EVENT_SCOPE(GraphBuilder, "ShadowProjectionOnOpaque");
+	RDG_GPU_STAT_SCOPE(GraphBuilder, ShadowProjection);
 
 	FShadowProjectionPassParameters CommonPassParameters;
-	CommonPassParameters.SceneTextures.MobileSceneTextures = CreateMobileSceneTextureUniformBuffer(GraphBuilder, EMobileSceneTextureSetupMode::SceneColor);
+	CommonPassParameters.SceneTextures.MobileSceneTextures = MobileSceneTextures;
+	CommonPassParameters.RenderTargets = BasePassRenderTargets;
 
 	// render shadowmaps for relevant lights.
 	for (TSparseArray<FLightSceneInfoCompact>::TConstIterator LightIt(Scene->Lights); LightIt; ++LightIt)
@@ -2080,8 +2079,6 @@ void FMobileSceneRenderer::RenderModulatedShadowProjections(FRHICommandListImmed
 			}
 		}
 	}
-
-	GraphBuilder.Execute();
 }
 
 IMPLEMENT_GLOBAL_SHADER_PARAMETER_STRUCT(FTranslucentSelfShadowUniformParameters, "TranslucentSelfShadow");
