Extend UAV -> RTV DCC workaround to include a flush after the tiled deferred rendering draw step

diff --git a/Engine/Source/Runtime/Renderer/Private/TiledDeferredLightRendering.cpp b/Engine/Source/Runtime/Renderer/Private/TiledDeferredLightRendering.cpp
--- a/Engine/Source/Runtime/Renderer/Private/TiledDeferredLightRendering.cpp
+++ b/Engine/Source/Runtime/Renderer/Private/TiledDeferredLightRendering.cpp
@@ -399,6 +399,10 @@ void FDeferredShadingSceneRenderer::RenderTiledDeferredLighting(
 					{
 						SetShaderTemplTiledLighting<0>(RHICmdList, View, ViewIndex, ViewCount, SortedLights, TiledDeferredLightsStart, TiledDeferredLightsEnd, SimpleLights, StartIndex, NumThisPass, *SceneColorTexture->GetPooledRenderTarget(), *SceneColorOutputTexture->GetPooledRenderTarget());
 					}
+
+#if PLATFORM_REQUIRES_UAV_TO_RTV_TEXTURE_CACHE_FLUSH_WORKAROUND
+					RHICmdList.RHIFlushTextureCacheBOP(SceneColorOutputTexture->GetRHI());
+#endif // #if PLATFORM_REQUIRES_UAV_TO_RTV_TEXTURE_CACHE_FLUSH_WORKAROUND
 				});
 			}
 
