Fixed cloud rendering pixel shader nainge being too long for dump shader path.

diff --git a/Engine/Source/Runtime/Renderer/Private/VolumetricCloudRendering.cpp b/Engine/Source/Runtime/Renderer/Private/VolumetricCloudRendering.cpp
--- a/Engine/Source/Runtime/Renderer/Private/VolumetricCloudRendering.cpp
+++ b/Engine/Source/Runtime/Renderer/Private/VolumetricCloudRendering.cpp
@@ -575,15 +575,15 @@ IMPLEMENT_MATERIAL_SHADER_TYPE(, FRenderVolumetricCloudVS, TEXT("/Engine/Private
 
 enum EVolumetricCloudRenderViewPsPermutations
 {
-	VolumetricCloudRenderViewPs_PerSampleAtmosphereTransmittance0_SampleShadow0_SecondLight0,
-	VolumetricCloudRenderViewPs_PerSampleAtmosphereTransmittance1_SampleShadow0_SecondLight0,
-	VolumetricCloudRenderViewPs_PerSampleAtmosphereTransmittance0_SampleShadow1_SecondLight0,
-	VolumetricCloudRenderViewPs_PerSampleAtmosphereTransmittance1_SampleShadow1_SecondLight0,
-	VolumetricCloudRenderViewPs_PerSampleAtmosphereTransmittance0_SampleShadow0_SecondLight1,
-	VolumetricCloudRenderViewPs_PerSampleAtmosphereTransmittance1_SampleShadow0_SecondLight1,
-	VolumetricCloudRenderViewPs_PerSampleAtmosphereTransmittance0_SampleShadow1_SecondLight1,
-	VolumetricCloudRenderViewPs_PerSampleAtmosphereTransmittance1_SampleShadow1_SecondLight1,
-	VolumetricCloudRenderViewPsCount
+	CloudPs_AtmoTrans0_SampleShad0_2ndLight0,
+	CloudPs_AtmoTrans1_SampleShad0_2ndLight0,
+	CloudPs_AtmoTrans0_SampleShad1_2ndLight0,
+	CloudPs_AtmoTrans1_SampleShad1_2ndLight0,
+	CloudPs_AtmoTrans0_SampleShad0_2ndLight1,
+	CloudPs_AtmoTrans1_SampleShad0_2ndLight1,
+	CloudPs_AtmoTrans0_SampleShad1_2ndLight1,
+	CloudPs_AtmoTrans1_SampleShad1_2ndLight1,
+	CloudPsCount
 };
 
 BEGIN_SHADER_PARAMETER_STRUCT(FRenderVolumetricCloudRenderViewParametersPS, )
@@ -621,30 +621,30 @@ public:
 		// Force texture fetches to not use automatic mip generation because the pixel shader is using a dynamic loop to evaluate the material multiple times.
 		OutEnvironment.SetDefine(TEXT("USE_FORCE_TEXTURE_MIP"), TEXT("1"));
 		
-		const bool bUseAtmosphereTransmittance =Permutation == VolumetricCloudRenderViewPs_PerSampleAtmosphereTransmittance1_SampleShadow0_SecondLight0 || Permutation == VolumetricCloudRenderViewPs_PerSampleAtmosphereTransmittance1_SampleShadow1_SecondLight0 ||
-												Permutation == VolumetricCloudRenderViewPs_PerSampleAtmosphereTransmittance1_SampleShadow0_SecondLight1 || Permutation == VolumetricCloudRenderViewPs_PerSampleAtmosphereTransmittance1_SampleShadow1_SecondLight1;
+		const bool bUseAtmosphereTransmittance =Permutation == CloudPs_AtmoTrans1_SampleShad0_2ndLight0 || Permutation == CloudPs_AtmoTrans1_SampleShad0_2ndLight1 ||
+												Permutation == CloudPs_AtmoTrans1_SampleShad1_2ndLight1 || Permutation == CloudPs_AtmoTrans1_SampleShad1_2ndLight0;
 		OutEnvironment.SetDefine(TEXT("CLOUD_PER_SAMPLE_ATMOSPHERE_TRANSMITTANCE"), bUseAtmosphereTransmittance ? TEXT("1") : TEXT("0"));
 
-		const bool bSampleLightShadowmap =	Permutation == VolumetricCloudRenderViewPs_PerSampleAtmosphereTransmittance0_SampleShadow1_SecondLight0 || Permutation == VolumetricCloudRenderViewPs_PerSampleAtmosphereTransmittance1_SampleShadow1_SecondLight0 ||
-											Permutation == VolumetricCloudRenderViewPs_PerSampleAtmosphereTransmittance0_SampleShadow1_SecondLight1 || Permutation == VolumetricCloudRenderViewPs_PerSampleAtmosphereTransmittance1_SampleShadow1_SecondLight1;
+		const bool bSampleLightShadowmap =	Permutation == CloudPs_AtmoTrans0_SampleShad1_2ndLight0 || Permutation == CloudPs_AtmoTrans0_SampleShad1_2ndLight1 ||
+											Permutation == CloudPs_AtmoTrans1_SampleShad1_2ndLight1 || Permutation == CloudPs_AtmoTrans1_SampleShad1_2ndLight0;
 		OutEnvironment.SetDefine(TEXT("CLOUD_SAMPLE_ATMOSPHERIC_LIGHT_SHADOWMAP"), bSampleLightShadowmap ? TEXT("1") : TEXT("0"));
 
-		const bool bSampleSecondLight =	Permutation == VolumetricCloudRenderViewPs_PerSampleAtmosphereTransmittance0_SampleShadow0_SecondLight1 || Permutation == VolumetricCloudRenderViewPs_PerSampleAtmosphereTransmittance1_SampleShadow0_SecondLight1 ||
-										Permutation == VolumetricCloudRenderViewPs_PerSampleAtmosphereTransmittance0_SampleShadow1_SecondLight1 || Permutation == VolumetricCloudRenderViewPs_PerSampleAtmosphereTransmittance1_SampleShadow1_SecondLight1;
+		const bool bSampleSecondLight =	Permutation == CloudPs_AtmoTrans0_SampleShad0_2ndLight1 || Permutation == CloudPs_AtmoTrans0_SampleShad1_2ndLight1 ||
+										Permutation == CloudPs_AtmoTrans1_SampleShad1_2ndLight1 || Permutation == CloudPs_AtmoTrans1_SampleShad0_2ndLight1;
 		OutEnvironment.SetDefine(TEXT("CLOUD_SAMPLE_SECOND_LIGHT"), bSampleSecondLight ? TEXT("1") : TEXT("0"));
 	}
 
 private:
 };
 
-IMPLEMENT_MATERIAL_SHADER_TYPE(template<>, FRenderVolumetricCloudRenderViewPs<VolumetricCloudRenderViewPs_PerSampleAtmosphereTransmittance0_SampleShadow0_SecondLight0>, TEXT("/Engine/Private/VolumetricCloud.usf"), TEXT("MainPS"), SF_Pixel);
-IMPLEMENT_MATERIAL_SHADER_TYPE(template<>, FRenderVolumetricCloudRenderViewPs<VolumetricCloudRenderViewPs_PerSampleAtmosphereTransmittance1_SampleShadow0_SecondLight0>, TEXT("/Engine/Private/VolumetricCloud.usf"), TEXT("MainPS"), SF_Pixel);
-IMPLEMENT_MATERIAL_SHADER_TYPE(template<>, FRenderVolumetricCloudRenderViewPs<VolumetricCloudRenderViewPs_PerSampleAtmosphereTransmittance0_SampleShadow1_SecondLight0>, TEXT("/Engine/Private/VolumetricCloud.usf"), TEXT("MainPS"), SF_Pixel);
-IMPLEMENT_MATERIAL_SHADER_TYPE(template<>, FRenderVolumetricCloudRenderViewPs<VolumetricCloudRenderViewPs_PerSampleAtmosphereTransmittance1_SampleShadow1_SecondLight0>, TEXT("/Engine/Private/VolumetricCloud.usf"), TEXT("MainPS"), SF_Pixel);
-IMPLEMENT_MATERIAL_SHADER_TYPE(template<>, FRenderVolumetricCloudRenderViewPs<VolumetricCloudRenderViewPs_PerSampleAtmosphereTransmittance0_SampleShadow0_SecondLight1>, TEXT("/Engine/Private/VolumetricCloud.usf"), TEXT("MainPS"), SF_Pixel);
-IMPLEMENT_MATERIAL_SHADER_TYPE(template<>, FRenderVolumetricCloudRenderViewPs<VolumetricCloudRenderViewPs_PerSampleAtmosphereTransmittance1_SampleShadow0_SecondLight1>, TEXT("/Engine/Private/VolumetricCloud.usf"), TEXT("MainPS"), SF_Pixel);
-IMPLEMENT_MATERIAL_SHADER_TYPE(template<>, FRenderVolumetricCloudRenderViewPs<VolumetricCloudRenderViewPs_PerSampleAtmosphereTransmittance0_SampleShadow1_SecondLight1>, TEXT("/Engine/Private/VolumetricCloud.usf"), TEXT("MainPS"), SF_Pixel);
-IMPLEMENT_MATERIAL_SHADER_TYPE(template<>, FRenderVolumetricCloudRenderViewPs<VolumetricCloudRenderViewPs_PerSampleAtmosphereTransmittance1_SampleShadow1_SecondLight1>, TEXT("/Engine/Private/VolumetricCloud.usf"), TEXT("MainPS"), SF_Pixel);
+IMPLEMENT_MATERIAL_SHADER_TYPE(template<>, FRenderVolumetricCloudRenderViewPs<CloudPs_AtmoTrans0_SampleShad0_2ndLight0>, TEXT("/Engine/Private/VolumetricCloud.usf"), TEXT("MainPS"), SF_Pixel);
+IMPLEMENT_MATERIAL_SHADER_TYPE(template<>, FRenderVolumetricCloudRenderViewPs<CloudPs_AtmoTrans1_SampleShad0_2ndLight0>, TEXT("/Engine/Private/VolumetricCloud.usf"), TEXT("MainPS"), SF_Pixel);
+IMPLEMENT_MATERIAL_SHADER_TYPE(template<>, FRenderVolumetricCloudRenderViewPs<CloudPs_AtmoTrans0_SampleShad1_2ndLight0>, TEXT("/Engine/Private/VolumetricCloud.usf"), TEXT("MainPS"), SF_Pixel);
+IMPLEMENT_MATERIAL_SHADER_TYPE(template<>, FRenderVolumetricCloudRenderViewPs<CloudPs_AtmoTrans1_SampleShad1_2ndLight0>, TEXT("/Engine/Private/VolumetricCloud.usf"), TEXT("MainPS"), SF_Pixel);
+IMPLEMENT_MATERIAL_SHADER_TYPE(template<>, FRenderVolumetricCloudRenderViewPs<CloudPs_AtmoTrans0_SampleShad0_2ndLight1>, TEXT("/Engine/Private/VolumetricCloud.usf"), TEXT("MainPS"), SF_Pixel);
+IMPLEMENT_MATERIAL_SHADER_TYPE(template<>, FRenderVolumetricCloudRenderViewPs<CloudPs_AtmoTrans1_SampleShad0_2ndLight1>, TEXT("/Engine/Private/VolumetricCloud.usf"), TEXT("MainPS"), SF_Pixel);
+IMPLEMENT_MATERIAL_SHADER_TYPE(template<>, FRenderVolumetricCloudRenderViewPs<CloudPs_AtmoTrans0_SampleShad1_2ndLight1>, TEXT("/Engine/Private/VolumetricCloud.usf"), TEXT("MainPS"), SF_Pixel);
+IMPLEMENT_MATERIAL_SHADER_TYPE(template<>, FRenderVolumetricCloudRenderViewPs<CloudPs_AtmoTrans1_SampleShad1_2ndLight1>, TEXT("/Engine/Private/VolumetricCloud.usf"), TEXT("MainPS"), SF_Pixel);
 
 
 
@@ -803,12 +803,12 @@ public:
 			{
 				if (bVolumetricCloudPerSampleAtmosphereTransmittance)
 				{
-					TemplatedProcess<FRenderVolumetricCloudRenderViewPs<VolumetricCloudRenderViewPs_PerSampleAtmosphereTransmittance1_SampleShadow1_SecondLight1>>(
+					TemplatedProcess<FRenderVolumetricCloudRenderViewPs<CloudPs_AtmoTrans1_SampleShad1_2ndLight1>>(
 						MeshBatch, BatchElementMask, PrimitiveSceneProxy, MaterialRenderProxy, Material, StaticMeshId, MeshFillMode, MeshCullMode);
 				}
 				else
 				{
-					TemplatedProcess<FRenderVolumetricCloudRenderViewPs<VolumetricCloudRenderViewPs_PerSampleAtmosphereTransmittance0_SampleShadow1_SecondLight1>>(
+					TemplatedProcess<FRenderVolumetricCloudRenderViewPs<CloudPs_AtmoTrans0_SampleShad1_2ndLight1>>(
 						MeshBatch, BatchElementMask, PrimitiveSceneProxy, MaterialRenderProxy, Material, StaticMeshId, MeshFillMode, MeshCullMode);
 				}
 			}
@@ -816,12 +816,12 @@ public:
 			{
 				if (bVolumetricCloudPerSampleAtmosphereTransmittance)
 				{
-					TemplatedProcess<FRenderVolumetricCloudRenderViewPs<VolumetricCloudRenderViewPs_PerSampleAtmosphereTransmittance1_SampleShadow0_SecondLight1>>(
+					TemplatedProcess<FRenderVolumetricCloudRenderViewPs<CloudPs_AtmoTrans1_SampleShad0_2ndLight1>>(
 						MeshBatch, BatchElementMask, PrimitiveSceneProxy, MaterialRenderProxy, Material, StaticMeshId, MeshFillMode, MeshCullMode);
 				}
 				else
 				{
-					TemplatedProcess<FRenderVolumetricCloudRenderViewPs<VolumetricCloudRenderViewPs_PerSampleAtmosphereTransmittance0_SampleShadow0_SecondLight1>>(
+					TemplatedProcess<FRenderVolumetricCloudRenderViewPs<CloudPs_AtmoTrans0_SampleShad0_2ndLight1>>(
 						MeshBatch, BatchElementMask, PrimitiveSceneProxy, MaterialRenderProxy, Material, StaticMeshId, MeshFillMode, MeshCullMode);
 				}
 			}
@@ -832,12 +832,12 @@ public:
 			{
 				if (bVolumetricCloudPerSampleAtmosphereTransmittance)
 				{
-					TemplatedProcess<FRenderVolumetricCloudRenderViewPs<VolumetricCloudRenderViewPs_PerSampleAtmosphereTransmittance1_SampleShadow1_SecondLight0>>(
+					TemplatedProcess<FRenderVolumetricCloudRenderViewPs<CloudPs_AtmoTrans1_SampleShad1_2ndLight0>>(
 						MeshBatch, BatchElementMask, PrimitiveSceneProxy, MaterialRenderProxy, Material, StaticMeshId, MeshFillMode, MeshCullMode);
 				}
 				else
 				{
-					TemplatedProcess<FRenderVolumetricCloudRenderViewPs<VolumetricCloudRenderViewPs_PerSampleAtmosphereTransmittance0_SampleShadow1_SecondLight0>>(
+					TemplatedProcess<FRenderVolumetricCloudRenderViewPs<CloudPs_AtmoTrans0_SampleShad1_2ndLight0>>(
 						MeshBatch, BatchElementMask, PrimitiveSceneProxy, MaterialRenderProxy, Material, StaticMeshId, MeshFillMode, MeshCullMode);
 				}
 			}
@@ -845,12 +845,12 @@ public:
 			{
 				if (bVolumetricCloudPerSampleAtmosphereTransmittance)
 				{
-					TemplatedProcess<FRenderVolumetricCloudRenderViewPs<VolumetricCloudRenderViewPs_PerSampleAtmosphereTransmittance1_SampleShadow0_SecondLight0>>(
+					TemplatedProcess<FRenderVolumetricCloudRenderViewPs<CloudPs_AtmoTrans1_SampleShad0_2ndLight0>>(
 						MeshBatch, BatchElementMask, PrimitiveSceneProxy, MaterialRenderProxy, Material, StaticMeshId, MeshFillMode, MeshCullMode);
 				}
 				else
 				{
-					TemplatedProcess<FRenderVolumetricCloudRenderViewPs<VolumetricCloudRenderViewPs_PerSampleAtmosphereTransmittance0_SampleShadow0_SecondLight0>>(
+					TemplatedProcess<FRenderVolumetricCloudRenderViewPs<CloudPs_AtmoTrans0_SampleShad0_2ndLight0>>(
 						MeshBatch, BatchElementMask, PrimitiveSceneProxy, MaterialRenderProxy, Material, StaticMeshId, MeshFillMode, MeshCullMode);
 				}
 			}
