Fixed mobile base pass view extensions to use render targets.

diff --git a/Engine/Source/Runtime/Renderer/Private/MobileShadingRenderer.cpp b/Engine/Source/Runtime/Renderer/Private/MobileShadingRenderer.cpp
--- a/Engine/Source/Runtime/Renderer/Private/MobileShadingRenderer.cpp
+++ b/Engine/Source/Runtime/Renderer/Private/MobileShadingRenderer.cpp
@@ -800,6 +800,7 @@ void FMobileSceneRenderer::Render(FRHICommandListImmediate& RHICmdList)
 
 BEGIN_SHADER_PARAMETER_STRUCT(FMobilePostBasePassViewExtensionParameters, )
 	SHADER_PARAMETER_RDG_UNIFORM_BUFFER(FMobileSceneTextureUniformParameters, SceneTextures)
+	RENDER_TARGET_BINDING_SLOTS()
 END_SHADER_PARAMETER_STRUCT()
 
 void FMobileSceneRenderer::RenderForward(FRDGBuilder& GraphBuilder, const TArrayView<const FViewInfo*> ViewList, FRDGTextureRef ViewFamilyTexture, FRDGTextureMSAA& SceneColorMSAA, FRDGTextureMSAA& SceneDepthMSAA, TRDGUniformBufferRef<FMobileSceneTextureUniformParameters> MobileSceneTextures)
@@ -923,12 +924,13 @@ void FMobileSceneRenderer::RenderForward(FRDGBuilder& GraphBuilder, const TArray
 	if (ViewFamily.ViewExtensions.Num() > 1)
 	{
 		auto* PassParameters = GraphBuilder.AllocParameters<FMobilePostBasePassViewExtensionParameters>();
+		PassParameters->RenderTargets = BasePassRenderTargets;
 		PassParameters->SceneTextures = MobileSceneTextures;
 
 		GraphBuilder.AddPass(
 			{},
 			PassParameters,
-			ERDGPassFlags::Raster | ERDGPassFlags::NeverCull,
+			ERDGPassFlags::Raster | ERDGPassFlags::SkipRenderPass | ERDGPassFlags::NeverCull,
 			[this](FRHICommandListImmediate& RHICmdList)
 		{
 			CSV_SCOPED_TIMING_STAT_EXCLUSIVE(ViewExtensionPostRenderBasePass);
