Fix SSR tracing due to incorrect default anisotropy value.

diff --git a/Engine/Source/Runtime/Renderer/Private/PostProcess/SceneRenderTargets.cpp b/Engine/Source/Runtime/Renderer/Private/PostProcess/SceneRenderTargets.cpp
--- a/Engine/Source/Runtime/Renderer/Private/PostProcess/SceneRenderTargets.cpp
+++ b/Engine/Source/Runtime/Renderer/Private/PostProcess/SceneRenderTargets.cpp
@@ -1853,7 +1853,7 @@ void SetupSceneTextureUniformParameters(
 	SceneTextureParameters.GBufferCTexture = SystemTextures.Black;
 	SceneTextureParameters.GBufferDTexture = SystemTextures.Black;
 	SceneTextureParameters.GBufferETexture = SystemTextures.Black;
-	SceneTextureParameters.GBufferFTexture = SystemTextures.Black;
+	SceneTextureParameters.GBufferFTexture = SystemTextures.MidGrey;
 	SceneTextureParameters.GBufferVelocityTexture = SystemTextures.Black;
 	SceneTextureParameters.ScreenSpaceAOTexture = GetScreenSpaceAOFallback(SystemTextures);
 
diff --git a/Engine/Source/Runtime/Renderer/Private/SceneTextureParameters.cpp b/Engine/Source/Runtime/Renderer/Private/SceneTextureParameters.cpp
--- a/Engine/Source/Runtime/Renderer/Private/SceneTextureParameters.cpp
+++ b/Engine/Source/Runtime/Renderer/Private/SceneTextureParameters.cpp
@@ -25,7 +25,7 @@ FSceneTextureParameters GetSceneTextureParameters(FRDGBuilder& GraphBuilder)
 	Parameters.GBufferCTexture = GetIfProduced(SceneTextures.GBufferC);
 	Parameters.GBufferDTexture = GetIfProduced(SceneTextures.GBufferD);
 	Parameters.GBufferETexture = GetIfProduced(SceneTextures.GBufferE);
-	Parameters.GBufferFTexture = GetIfProduced(SceneTextures.GBufferF, SystemTextures.Black);
+	Parameters.GBufferFTexture = GetIfProduced(SceneTextures.GBufferF, SystemTextures.MidGrey);
 
 	return Parameters;
 }
