Add initial material classification for strata material.

This pass stored shading models and material bits (subsurface, transmission, ...) for shadow and postprocess passes.

diff --git a/Engine/Shaders/Private/ShadowProjectionPixelShader.usf b/Engine/Shaders/Private/ShadowProjectionPixelShader.usf
--- a/Engine/Shaders/Private/ShadowProjectionPixelShader.usf
+++ b/Engine/Shaders/Private/ShadowProjectionPixelShader.usf
@@ -94,6 +94,11 @@ float2 ShadowNearAndFarDepth;
 Texture2D<uint4> HairCategorizationTexture;
 #endif
 
+#if STRATA_ENABLED
+#define STRATA_INLINE_SHADING 0
+#include "Strata/StrataEvaluation.ush"
+#endif
+
 /**
 * Entry point for uniform manual PCF that supports lights using normal shadows.
 */
@@ -118,9 +123,9 @@ void Main(
 	float SceneW = CalcSceneDepth(ScreenUV);
 
 	bool bIsSubsurfaceCompatible = true;
+	const uint3 PixelCoord = uint3(floor(SVPos.xy), 0);
 #if SUBPIXEL_SHADOW
 	const bool bUseFadePlane = bCascadeUseFadePlane > 0;
-	const uint3 PixelCoord = uint3(floor(SVPos.xy), 0);
 	const FCategorizationData CategorizationData = DecodeCategorizationData(HairCategorizationTexture.Load(PixelCoord));
 	if (CategorizationData.PixelCoverage > 0)
 	{
@@ -259,17 +264,29 @@ void Main(
 
 	#if FEATURE_LEVEL >= FEATURE_LEVEL_SM4 && !FORWARD_SHADING && !APPLY_TRANSLUCENCY_SHADOWS
 
-		FGBufferData GBufferData = GetGBufferData(ScreenUV);
+		FGBufferData GBufferData	= GetGBufferData(ScreenUV); 
+		bool bHasSubsurface			= IsSubsurfaceModel(GBufferData.ShadingModelID);
+		bool bHasSubsurfaceProfile	= GBufferData.ShadingModelID == SHADINGMODELID_SUBSURFACE_PROFILE;
+		bool bIsEyeOrHair 			= GBufferData.ShadingModelID == SHADINGMODELID_HAIR || GBufferData.ShadingModelID == SHADINGMODELID_EYE;
+		#if STRATA_ENABLED
+		{
+			const FStrataClassification StrataData = StrataUnpackClassificationData(Strata.ClassificationTexture.Load(PixelCoord));
+			if (StrataData.bIsStrataMaterial)
+			{
+				bIsEyeOrHair			= bIsEyeOrHair || StrataHasShadingModel(StrataData.ShadingModels, STRATA_BSDF_TYPE_HAIR);
+				bHasSubsurface			= bHasSubsurface || StrataData.bHasSubsurface || bIsEyeOrHair;
+				bHasSubsurfaceProfile	= bHasSubsurfaceProfile || StrataData.bHasSubsurfaceProfile;
+			}
+		}
+		#endif
 
 		BRANCH
-		if (bIsSubsurfaceCompatible && IsSubsurfaceModel(GBufferData.ShadingModelID))
+		if (bIsSubsurfaceCompatible && bHasSubsurface)
 		{
 			float Opacity = GBufferData.CustomData.a;
 			float Density = SubsurfaceDensityFromOpacity(Opacity);
-			if( GBufferData.ShadingModelID == SHADINGMODELID_HAIR || GBufferData.ShadingModelID == SHADINGMODELID_EYE )
+			if(bIsEyeOrHair)
 			{
-				// STRATA_TODO strata hair should also go there. Right now we would need to add SHADINGMODEL_STRATA there 
-				// and also enable HasCustomGBufferData just because we go through the GBuffer. We are going to fixcthat next the right way.
 				Opacity = 1;
 				Density = 1;
 			}
@@ -299,7 +316,7 @@ void Main(
 				Settings.ProjectionDepthBiasParameters = ProjectionDepthBiasParameters.xw;
 
 			  #if USE_TRANSMISSION
-				if (GBufferData.ShadingModelID == SHADINGMODELID_SUBSURFACE_PROFILE || GBufferData.ShadingModelID == SHADINGMODELID_STRATA)
+				if (bHasSubsurfaceProfile)
 				{
 					SSSTransmission = CalcTransmissionThickness(ScreenPosition.xyz, LightPositionAndInvRadius.xyz, GBufferData, Settings);
 				}
diff --git a/Engine/Shaders/Private/Strata/Strata.ush b/Engine/Shaders/Private/Strata/Strata.ush
--- a/Engine/Shaders/Private/Strata/Strata.ush
+++ b/Engine/Shaders/Private/Strata/Strata.ush
@@ -143,6 +143,46 @@ struct FStrataBSDF
 
 void StrataRequestSharedNormalTangentBasis(inout uint NormalTypes, uint Index);
 
+struct FStrataClassification
+{
+	uint ShadingModels;
+	bool bIsStrataMaterial;
+	bool bHasSubsurface;
+	bool bHasSubsurfaceProfile;
+	bool bHasTransmission;
+};
+
+uint StrataPackClassificationData(FStrataClassification In)
+{
+	return 
+		(In.ShadingModels & 0xFF) | 
+		((In.bIsStrataMaterial		? 1 : 0) << 8) |
+		((In.bHasSubsurface			? 1 : 0) << 9) | 
+		((In.bHasSubsurfaceProfile	? 1 : 0) << 10)|
+		((In.bHasTransmission		? 1 : 0) << 11);
+}
+
+FStrataClassification StrataUnpackClassificationData(uint In)
+{
+	FStrataClassification Out;
+	Out.ShadingModels			= (In & 0xFF);
+	Out.bIsStrataMaterial		= (In >>  8) & 0x1;
+	Out.bHasSubsurface			= (In >>  9) & 0x1;
+	Out.bHasSubsurfaceProfile	= (In >> 10) & 0x1;
+	Out.bHasTransmission		= (In >> 11) & 0x1;
+	return Out;
+}
+
+uint StrataShadingModelBit(uint ShadingModel)
+{
+	return 1 << ShadingModel;
+}
+
+bool StrataHasShadingModel(uint ShadingModels, uint ShadingModel)
+{
+	return (ShadingModels & (1 << ShadingModel)) > 0;
+}
+
 bool StrataHasAnisotropy(in float Roughness0, float Roughness1)
 {
 	return Roughness0 != Roughness1;
diff --git a/Engine/Shaders/Private/Strata/StrataMaterialClassification.usf b/Engine/Shaders/Private/Strata/StrataMaterialClassification.usf
new file mode 100644
--- /dev/null
+++ b/Engine/Shaders/Private/Strata/StrataMaterialClassification.usf
@@ -0,0 +1,61 @@
+// Copyright Epic Games, Inc. All Rights Reserved.
+
+#include "/Engine/Private/Common.ush"
+#include "/Engine/Private/DeferredShadingCommon.ush"
+
+#define STRATA_INLINE_SHADING 0
+#include "/Engine/Private/Strata/Strata.ush"
+
+void MainPS(
+	float4 SVPos : SV_POSITION, 
+	out uint OutColor0 : SV_Target0)
+{
+	const uint2 PixelPos = uint2(SVPos.xy);
+
+	// Sample scene textures.
+	float2 BufferUV = SvPositionToBufferUV(float4(PixelPos, SVPos.zw));
+	FGBufferData GBuffer = GetGBufferDataFromSceneTextures(BufferUV);
+	bool bStrataMaterial = GBuffer.ShadingModelID == SHADINGMODELID_STRATA;
+	
+	FStrataClassification Out = (FStrataClassification)0;
+	if (bStrataMaterial)
+	{
+		Out.bIsStrataMaterial = true;
+		uint PixelStrataDataByteOffset = GetStrataPixelDataByteOffset(PixelPos, View.BufferSizeAndInvSize.x, Strata.MaxBytesPerPixel);
+		FStrataPixelHeader Header = UnpackStrataHeaderIn(Strata.MaterialLobesBuffer, PixelStrataDataByteOffset);
+		for (uint BSDFIndex = 0; BSDFIndex < Header.BSDFCount; ++BSDFIndex)
+		{
+			const FStrataBSDF BSDF = UnpackStrataBSDFIn(Strata.MaterialLobesBuffer, PixelStrataDataByteOffset);
+			Out.ShadingModels |= StrataShadingModelBit(BSDF_GETTYPE(BSDF));
+
+			switch (BSDF_GETTYPE(BSDF))
+			{
+			case STRATA_BSDF_TYPE_DIFFUSE:
+				if (BSDF_GETHASSCATTERING(BSDF))
+				{
+					Out.bHasSubsurface = true;
+					if (BSDF_GETHASSSSPROFILE(BSDF))
+					{
+						Out.bHasSubsurfaceProfile = true;
+					}
+				}
+				break;
+			case STRATA_BSDF_TYPE_HAIR:
+				Out.bHasTransmission = true;
+				break;
+			case STRATA_BSDF_TYPE_DIELECTRIC:
+				break;
+			case STRATA_BSDF_TYPE_CONDUCTOR:
+				break;
+			case STRATA_BSDF_TYPE_VOLUME:	
+				break;
+			case STRATA_BSDF_TYPE_SHEEN:	
+				break;
+			}
+		}
+	}
+
+	OutColor0 = StrataPackClassificationData(Out);
+}
+
+
diff --git a/Engine/Shaders/Private/TransmissionCommon.ush b/Engine/Shaders/Private/TransmissionCommon.ush
--- a/Engine/Shaders/Private/TransmissionCommon.ush
+++ b/Engine/Shaders/Private/TransmissionCommon.ush
@@ -1,5 +1,7 @@
 // Copyright Epic Games, Inc. All Rights Reserved.
 
+#pragma once
+
 /*=========================================================================
 	TransmissionCommon.ush: Subsurface scattering transmission utilities.
 =========================================================================*/
diff --git a/Engine/Shaders/Private/TransmissionThickness.ush b/Engine/Shaders/Private/TransmissionThickness.ush
--- a/Engine/Shaders/Private/TransmissionThickness.ush
+++ b/Engine/Shaders/Private/TransmissionThickness.ush
@@ -1,5 +1,7 @@
 // Copyright Epic Games, Inc. All Rights Reserved.
 
+#pragma once 
+
 /*=============================================================================
 	TransmissionThickness.ush: Calc Transmission Thickness shaders.
 =============================================================================*/
@@ -95,18 +97,6 @@ float2(0.865413, 0.763726),
 float2(0.872005, -0.927)
 };
 
-// Normalized(1.0f / (4 * PI))
-#define ISO_PHASE_FUNC_Normalized ( 0.5f )
-#define Light_Intensity_Blit (2.0f) 
-float ApproximateHG(float cosJ, float g)		// HG phase function approximated
-{
-	float g2 = g * g;
-	float gcos2 = 1.0f - (g * cosJ);
-	gcos2 *= gcos2;
-
-	return (Light_Intensity_Blit * ISO_PHASE_FUNC_Normalized * (1.0f - g2) / gcos2);
-}
-
 float CalcShadowLinearDepth(float shadowDepth, float2 param)
 {
 	return 1.0f / (shadowDepth * param.x - param.y);
diff --git a/Engine/Source/Runtime/Renderer/Private/DeferredShadingRenderer.cpp b/Engine/Source/Runtime/Renderer/Private/DeferredShadingRenderer.cpp
--- a/Engine/Source/Runtime/Renderer/Private/DeferredShadingRenderer.cpp
+++ b/Engine/Source/Runtime/Renderer/Private/DeferredShadingRenderer.cpp
@@ -2359,6 +2359,12 @@ void FDeferredShadingSceneRenderer::Render(FRHICommandListImmediate& RHICmdList)
 		}
 	}
 
+	// Post base pass for material classification
+	if (Strata::IsStrataEnabled())
+	{
+		Strata::AddStrataMaterialClassificationPass(GraphBuilder, Views);
+	}
+
 	// Hair base pass for deferred shading
 	if (bHairEnable && !IsForwardShadingEnabled(ShaderPlatform))
 	{
diff --git a/Engine/Source/Runtime/Renderer/Private/ShadowRendering.cpp b/Engine/Source/Runtime/Renderer/Private/ShadowRendering.cpp
--- a/Engine/Source/Runtime/Renderer/Private/ShadowRendering.cpp
+++ b/Engine/Source/Runtime/Renderer/Private/ShadowRendering.cpp
@@ -400,6 +400,13 @@ static void BindShaderShaders(FRHICommandList& RHICmdList, FGraphicsPipelineStat
 
 	VertexShader->SetParameters(RHICmdList, View, ShadowInfo);
 	PixelShader->SetParameters(RHICmdList, ViewIndex, View, HairVisibilityData, ShadowInfo);
+
+	if (Strata::IsStrataEnabled())
+	{
+		FRHIPixelShader* ShaderRHI = PixelShader.GetPixelShader();
+		FRHIUniformBuffer* StrataUniformBuffer = Strata::BindStrataGlobalUniformParameters(View);
+		PixelShader->FGlobalShader::SetParameters<FStrataGlobalUniformParameters>(RHICmdList, ShaderRHI, StrataUniformBuffer);
+	}
 }
 
 
diff --git a/Engine/Source/Runtime/Renderer/Private/ShadowRendering.h b/Engine/Source/Runtime/Renderer/Private/ShadowRendering.h
--- a/Engine/Source/Runtime/Renderer/Private/ShadowRendering.h
+++ b/Engine/Source/Runtime/Renderer/Private/ShadowRendering.h
@@ -28,6 +28,7 @@
 #include "ShaderParameterUtils.h"
 #include "LightRendering.h"
 #include "HairStrands/HairStrandsRendering.h"
+#include "Strata/Strata.h"
 
 ENGINE_API IPooledRenderTarget* GetSubsufaceProfileTexture_RT(FRHICommandListImmediate& RHICmdList);
 
@@ -1083,6 +1084,7 @@ public:
 		OutEnvironment.SetDefine(TEXT("SUBPIXEL_SHADOW"), (uint32)(SubPixelShadow ? 1 : 0));
 		OutEnvironment.SetDefine(TEXT("USE_FADE_PLANE"), (uint32)(bUseFadePlane ? 1 : 0));
 		OutEnvironment.SetDefine(TEXT("USE_TRANSMISSION"), (uint32)(bUseTransmission ? 1 : 0));
+		OutEnvironment.SetDefine(TEXT("STRATA_ENABLED"), (uint32)(Strata::IsStrataEnabled() ? 1 : 0));
 	}
 
 	/**
@@ -1406,6 +1408,7 @@ public:
 		OutEnvironment.SetDefine(TEXT("SHADOW_QUALITY"), Quality);
 		OutEnvironment.SetDefine(TEXT("USE_TRANSMISSION"), (uint32)(bUseTransmission ? 1 : 0));
 		OutEnvironment.SetDefine(TEXT("SUBPIXEL_SHADOW"), (uint32)(bUseSubPixel ? 1 : 0));
+		OutEnvironment.SetDefine(TEXT("STRATA_ENABLED"), (uint32)(Strata::IsStrataEnabled() ? 1 : 0));
 	}
 
 	static bool ShouldCompilePermutation(const FGlobalShaderPermutationParameters& Parameters)
diff --git a/Engine/Source/Runtime/Renderer/Private/Strata/Strata.cpp b/Engine/Source/Runtime/Renderer/Private/Strata/Strata.cpp
--- a/Engine/Source/Runtime/Renderer/Private/Strata/Strata.cpp
+++ b/Engine/Source/Runtime/Renderer/Private/Strata/Strata.cpp
@@ -28,8 +28,6 @@ static TAutoConsoleVariable<int32> CVarStrataBytePerPixel(
 
 IMPLEMENT_GLOBAL_SHADER_PARAMETER_STRUCT(FStrataGlobalUniformParameters, "Strata");
 
-
-
 namespace Strata
 {
 
@@ -58,6 +56,11 @@ void InitialiseStrataFrameSceneData(FSceneRenderer& SceneRenderer, FRDGBuilder&
 		const uint32 MaterialConservativeByteCountPerPixel = 100u;
 		const uint32 RoundToValue = 4u;
 		StrataSceneData.MaxBytesPerPixel = FMath::DivideAndRoundUp(MaterialConservativeByteCountPerPixel, RoundToValue) * RoundToValue;
+
+		FRDGTextureDesc Desc = FRDGTextureDesc::Create2D(BufferSizeXY, PF_R32_UINT, FClearValueBinding::Black, TexCreate_ShaderResource | TexCreate_RenderTargetable);
+		FRDGTextureRef Texture = GraphBuilder.CreateTexture(Desc, TEXT("StrataClassificationTexture"));
+		GraphBuilder.PreallocateTexture(Texture);
+		StrataSceneData.ClassificationTexture = GraphBuilder.GetPooledTexture(Texture);
 	}
 	else
 	{
@@ -112,6 +115,7 @@ TUniformBufferRef<FStrataGlobalUniformParameters> BindStrataGlobalUniformParamet
 
 		StrataUniformParameters.MaxBytesPerPixel = View.StrataSceneData->MaxBytesPerPixel;
 		StrataUniformParameters.MaterialLobesBuffer = View.StrataSceneData->MaterialLobesBuffer.SRV;
+		StrataUniformParameters.ClassificationTexture = View.StrataSceneData->ClassificationTexture->GetRenderTargetItem().ShaderResourceTexture;
 
 		View.StrataSceneData->StrataGlobalUniformParameters = CreateUniformBufferImmediate(StrataUniformParameters, UniformBuffer_SingleFrame);
 		return View.StrataSceneData->StrataGlobalUniformParameters;
@@ -121,13 +125,13 @@ TUniformBufferRef<FStrataGlobalUniformParameters> BindStrataGlobalUniformParamet
 		// Create each time. This path will go away when Strata is always enabled anyway.
 		StrataUniformParameters.MaxBytesPerPixel = 0;
 		StrataUniformParameters.MaterialLobesBuffer = GEmptyVertexBufferWithUAV->ShaderResourceViewRHI;
+		StrataUniformParameters.ClassificationTexture = GSystemTextures.BlackDummy->GetRenderTargetItem().ShaderResourceTexture;
 		return CreateUniformBufferImmediate(StrataUniformParameters, UniformBuffer_SingleDraw);
 	}
 }
 
-
-////////////////////////////////////////////////////////////////////////// Debug
-
+////////////////////////////////////////////////////////////////////////// 
+// Debug
 
 class FVisualizeMaterialPS : public FGlobalShader
 {
@@ -199,6 +203,81 @@ void AddVisualizeMaterialPasses(FRDGBuilder& GraphBuilder, const TArray<FViewInf
 	}
 }
 
+//////////////////////////////////////////////////////////////////////////
+// Material classification pass
+// * Classification texture (shading models, BSDF bits, ...)
+// For future:
+// * SSS: hasSSS, Normal, ProfilID, BaseColor, Opacity,  MFPAlbedo/MFPRadius, Shadingmodel | 64bit?
+// * SSR: depth, roughness, normal, (clear coat amount/roughness), tangent, aniso
+
+// SSS/SSR/Auxilary data (AO/ShadowMask/...)
+class FStrataMaterialClassificationPassPS : public FGlobalShader
+{
+	DECLARE_GLOBAL_SHADER(FStrataMaterialClassificationPassPS);
+	SHADER_USE_PARAMETER_STRUCT(FStrataMaterialClassificationPassPS, FGlobalShader);
+
+	using FPermutationDomain = TShaderPermutationDomain<>;
+
+	BEGIN_SHADER_PARAMETER_STRUCT(FParameters, )
+		SHADER_PARAMETER_STRUCT_REF(FViewUniformShaderParameters, ViewUniformBuffer)
+		SHADER_PARAMETER_STRUCT_REF(FStrataGlobalUniformParameters, Strata)
+		SHADER_PARAMETER_STRUCT_INCLUDE(FSceneTextureParameters, SceneTextures)
+		SHADER_PARAMETER_STRUCT_INCLUDE(ShaderDrawDebug::FShaderDrawDebugParameters, ShaderDrawParameters)
+		RENDER_TARGET_BINDING_SLOTS()
+	END_SHADER_PARAMETER_STRUCT()
+
+	static FPermutationDomain RemapPermutation(FPermutationDomain PermutationVector)
+	{
+		return PermutationVector;
+	}
+
+	static bool ShouldCompilePermutation(const FGlobalShaderPermutationParameters& Parameters)
+	{
+		return GetMaxSupportedFeatureLevel(Parameters.Platform) >= ERHIFeatureLevel::SM5 && Strata::IsStrataEnabled();
+	}
+
+	static void ModifyCompilationEnvironment(const FGlobalShaderPermutationParameters& Parameters, FShaderCompilerEnvironment& OutEnvironment)
+	{
+		FGlobalShader::ModifyCompilationEnvironment(Parameters, OutEnvironment);
+	}
+};
+IMPLEMENT_GLOBAL_SHADER(FStrataMaterialClassificationPassPS, "/Engine/Private/Strata/StrataMaterialClassification.usf", "MainPS", SF_Pixel);
+
+void AddStrataMaterialClassificationPass(FRDGBuilder& GraphBuilder, const TArray<FViewInfo>& Views)
+{
+	RDG_EVENT_SCOPE_CONDITIONAL(GraphBuilder, IsStrataEnabled() && Views.Num() > 0, "StrataMaterialClassification");
+	if (!IsStrataEnabled())
+	{
+		return;
+	}
+
+	for (int32 i = 0; i < Views.Num(); ++i)
+	{
+		const FViewInfo& View = Views[i];
+		
+		FStrataMaterialClassificationPassPS::FPermutationDomain PermutationVector;
+		TShaderMapRef<FStrataMaterialClassificationPassPS> PixelShader(View.ShaderMap, PermutationVector);
+		FStrataMaterialClassificationPassPS::FParameters* PassParameters = GraphBuilder.AllocParameters<FStrataMaterialClassificationPassPS::FParameters>();
+		PassParameters->ViewUniformBuffer = View.ViewUniformBuffer;
+		PassParameters->Strata = Strata::BindStrataGlobalUniformParameters(View);
+		PassParameters->SceneTextures = GetSceneTextureParameters(GraphBuilder);		
+		PassParameters->RenderTargets[0] = FRenderTargetBinding(GraphBuilder.RegisterExternalTexture(View.StrataSceneData->ClassificationTexture), ERenderTargetLoadAction::EClear);
+
+		if (ShaderDrawDebug::IsShaderDrawDebugEnabled())
+		{
+			ShaderDrawDebug::SetParameters(GraphBuilder, View.ShaderDrawData, PassParameters->ShaderDrawParameters);
+		}
+
+		FPixelShaderUtils::AddFullscreenPass<FStrataMaterialClassificationPassPS>(
+			GraphBuilder, 
+			View.ShaderMap, 
+			RDG_EVENT_NAME("StrataMaterialClassification"),
+			PixelShader, 
+			PassParameters, 
+			View.ViewRect);
+	}
+}
+
 
 } // namespace Strata
 
diff --git a/Engine/Source/Runtime/Renderer/Private/Strata/Strata.h b/Engine/Source/Runtime/Renderer/Private/Strata/Strata.h
--- a/Engine/Source/Runtime/Renderer/Private/Strata/Strata.h
+++ b/Engine/Source/Runtime/Renderer/Private/Strata/Strata.h
@@ -23,15 +23,14 @@ END_SHADER_PARAMETER_STRUCT()
 BEGIN_GLOBAL_SHADER_PARAMETER_STRUCT(FStrataGlobalUniformParameters, )
 	SHADER_PARAMETER(uint32, MaxBytesPerPixel)
 	SHADER_PARAMETER_SRV(ByteAddressBuffer, MaterialLobesBuffer)
+	SHADER_PARAMETER_TEXTURE(Texture2D<uint>, ClassificationTexture)
 END_GLOBAL_SHADER_PARAMETER_STRUCT()
 
-
-
 struct FStrataSceneData
 {
 	uint32 MaxBytesPerPixel;
-	TRefCountPtr<IPooledRenderTarget> MaterialLobesTexture; // This should be a RDG resource when the refactoring gets in
-	FRWByteAddressBuffer MaterialLobesBuffer;				// This should be a RDG resource	"		"		"
+	FRWByteAddressBuffer MaterialLobesBuffer;				// This should be a RDG resource
+	TRefCountPtr<IPooledRenderTarget> ClassificationTexture;// This should be a RDG resource
 
 	TUniformBufferRef<FStrataGlobalUniformParameters> StrataGlobalUniformParameters;
 
@@ -40,9 +39,6 @@ struct FStrataSceneData
 	}
 };
 
-
-
-
 namespace Strata
 {
 
@@ -56,6 +52,8 @@ TUniformBufferRef<FStrataGlobalUniformParameters> BindStrataGlobalUniformParamet
 
 void AddVisualizeMaterialPasses(FRDGBuilder& GraphBuilder, const TArray<FViewInfo>& Views, FRDGTextureRef SceneColorTexture);
 
+void AddStrataMaterialClassificationPass(FRDGBuilder& GraphBuilder, const TArray<FViewInfo>& Views);
+
 };
 
 
