Fix compile issue from missing template specifier.

diff --git a/Engine/Source/Runtime/Renderer/Private/ShadowRendering.cpp b/Engine/Source/Runtime/Renderer/Private/ShadowRendering.cpp
--- a/Engine/Source/Runtime/Renderer/Private/ShadowRendering.cpp
+++ b/Engine/Source/Runtime/Renderer/Private/ShadowRendering.cpp
@@ -405,7 +405,7 @@ static void BindShaderShaders(FRHICommandList& RHICmdList, FGraphicsPipelineStat
 	{
 		FRHIPixelShader* ShaderRHI = PixelShader.GetPixelShader();
 		FRHIUniformBuffer* StrataUniformBuffer = Strata::BindStrataGlobalUniformParameters(View);
-		PixelShader->FGlobalShader::SetParameters<FStrataGlobalUniformParameters>(RHICmdList, ShaderRHI, StrataUniformBuffer);
+		PixelShader->FGlobalShader::template SetParameters<FStrataGlobalUniformParameters>(RHICmdList, ShaderRHI, StrataUniformBuffer);
 	}
 }
 
