Fix for RDG system textures to fail gracefully if not created.

diff --git a/Engine/Source/Runtime/Renderer/Private/SystemTextures.cpp b/Engine/Source/Runtime/Renderer/Private/SystemTextures.cpp
--- a/Engine/Source/Runtime/Renderer/Private/SystemTextures.cpp
+++ b/Engine/Source/Runtime/Renderer/Private/SystemTextures.cpp
@@ -21,7 +21,7 @@ const FRDGSystemTextures& FRDGSystemTextures::Create(FRDGBuilder& GraphBuilder)
 {
 	const auto Register = [&](const TRefCountPtr<IPooledRenderTarget>& RenderTarget)
 	{
-		return GraphBuilder.RegisterExternalTexture(RenderTarget, ERenderTargetTexture::ShaderResource, ERDGTextureFlags::ReadOnly);
+		return TryRegisterExternalTexture(GraphBuilder, RenderTarget, ERenderTargetTexture::ShaderResource, ERDGTextureFlags::ReadOnly);
 	};
 
 	auto& SystemTextures = GraphBuilder.Blackboard.Create<FRDGSystemTextures>();
