Fix for black raytracing shaders.

diff --git a/Engine/Source/Runtime/Renderer/Private/SceneVisibility.cpp b/Engine/Source/Runtime/Renderer/Private/SceneVisibility.cpp
--- a/Engine/Source/Runtime/Renderer/Private/SceneVisibility.cpp
+++ b/Engine/Source/Runtime/Renderer/Private/SceneVisibility.cpp
@@ -4564,6 +4564,9 @@ bool FDeferredShadingSceneRenderer::InitViews(FRDGBuilder& GraphBuilder, FExclus
 			// Initialize the view's RHI resources.
 			View.InitRHIResources();
 		}
+
+		// This is done for raytracing hit groups.
+		Scene->UniformBuffers.ViewUniformBuffer.UpdateUniformBufferImmediate(*Views[0].CachedViewUniformShaderParameters);
 	}
 
 	SetupVolumetricFog();
