Use explict begin/end pass for RDG mobile base pass to handle sub-pass properly.

diff --git a/Engine/Source/Runtime/Renderer/Private/CompositionLighting/PostProcessMeshDecals.cpp b/Engine/Source/Runtime/Renderer/Private/CompositionLighting/PostProcessMeshDecals.cpp
--- a/Engine/Source/Runtime/Renderer/Private/CompositionLighting/PostProcessMeshDecals.cpp
+++ b/Engine/Source/Runtime/Renderer/Private/CompositionLighting/PostProcessMeshDecals.cpp
@@ -424,7 +424,7 @@ void RenderMeshDecalsMobile(FRDGBuilder& GraphBuilder, FRenderTargetBindingSlots
 	auto* PassParameters = GraphBuilder.AllocParameters<FRenderTargetParameters>();
 	PassParameters->RenderTargets = BasePassRenderTargets;
 
-	GraphBuilder.AddPass(RDG_EVENT_NAME("MeshDecalsPass"), PassParameters, ERDGPassFlags::Raster,
+	GraphBuilder.AddPass(RDG_EVENT_NAME("MeshDecalsPass"), PassParameters, ERDGPassFlags::Raster | ERDGPassFlags::SkipRenderPass,
 		[RenderTargetMode, &View](FRHICommandList& RHICmdList)
 	{
 		FGraphicsPipelineStateInitializer GraphicsPSOInit;
diff --git a/Engine/Source/Runtime/Renderer/Private/DepthRendering.cpp b/Engine/Source/Runtime/Renderer/Private/DepthRendering.cpp
--- a/Engine/Source/Runtime/Renderer/Private/DepthRendering.cpp
+++ b/Engine/Source/Runtime/Renderer/Private/DepthRendering.cpp
@@ -602,7 +602,7 @@ void FDeferredShadingSceneRenderer::RenderPrePass(FRDGBuilder& GraphBuilder, FRD
 	}
 }
 
-void FMobileSceneRenderer::RenderPrePass(FRDGBuilder& GraphBuilder, FRenderTargetBindingSlots& BasePassRenderTargets, TFunction<void(FRenderTargetBindingSlots&)> UpdateRenderTargetsLoadAction)
+void FMobileSceneRenderer::RenderPrePass(FRDGBuilder& GraphBuilder, FRenderTargetBindingSlots& BasePassRenderTargets)
 {
 	SCOPED_NAMED_EVENT(FMobileSceneRenderer_RenderPrePass, FColor::Emerald);
 	RDG_EVENT_SCOPE(GraphBuilder, "MobileRenderPrePass");
@@ -615,8 +615,6 @@ void FMobileSceneRenderer::RenderPrePass(FRDGBuilder& GraphBuilder, FRenderTarge
 	// Mobile only does MaskedOnly DepthPass for the moment
 	if (Scene->EarlyZPassMode == DDM_MaskedOnly)
 	{
-		bool bAnyPassesAdded = false;
-
 		for (int32 ViewIndex = 0; ViewIndex < Views.Num(); ViewIndex++)
 		{
 			const FViewInfo& View = Views[ViewIndex];
@@ -632,7 +630,7 @@ void FMobileSceneRenderer::RenderPrePass(FRDGBuilder& GraphBuilder, FRenderTarge
 			auto* PassParameters = GraphBuilder.AllocParameters<FRenderTargetParameters>();
 			PassParameters->RenderTargets = BasePassRenderTargets;
 
-			GraphBuilder.AddPass(RDG_EVENT_NAME("RenderPrePass"), PassParameters, ERDGPassFlags::Raster,
+			GraphBuilder.AddPass(RDG_EVENT_NAME("RenderPrePass"), PassParameters, ERDGPassFlags::Raster | ERDGPassFlags::SkipRenderPass,
 				[this, &View, PassParameters](FRHICommandListImmediate& RHICmdList)
 			{
 				Scene->UniformBuffers.UpdateViewUniformBuffer(View);
@@ -641,13 +639,6 @@ void FMobileSceneRenderer::RenderPrePass(FRDGBuilder& GraphBuilder, FRenderTarge
 
 				View.ParallelMeshDrawCommandPasses[EMeshPass::DepthPass].DispatchDraw(nullptr, RHICmdList);
 			});
-
-			bAnyPassesAdded = true;
-		}
-
-		if (bAnyPassesAdded)
-		{
-			UpdateRenderTargetsLoadAction(BasePassRenderTargets);
 		}
 	}
 }
diff --git a/Engine/Source/Runtime/Renderer/Private/MobileBasePassRendering.cpp b/Engine/Source/Runtime/Renderer/Private/MobileBasePassRendering.cpp
--- a/Engine/Source/Runtime/Renderer/Private/MobileBasePassRendering.cpp
+++ b/Engine/Source/Runtime/Renderer/Private/MobileBasePassRendering.cpp
@@ -302,15 +302,13 @@ void SetupMobileSkyReflectionUniformParameters(FSkyLightSceneProxy* SkyLight, FM
 	Parameters.TextureSampler = CaptureTexture->SamplerStateRHI;
 }
 
-void FMobileSceneRenderer::RenderMobileBasePass(FRDGBuilder& GraphBuilder, FRenderTargetBindingSlots& BasePassRenderTargets, TFunction<void(FRenderTargetBindingSlots &)> UpdateRenderTargetsLoadAction, const TArrayView<const FViewInfo*> PassViews)
+void FMobileSceneRenderer::RenderMobileBasePass(FRDGBuilder& GraphBuilder, FRenderTargetBindingSlots& BasePassRenderTargets, const TArrayView<const FViewInfo*> PassViews)
 {
 	CSV_SCOPED_TIMING_STAT_EXCLUSIVE(RenderBasePass);
 	RDG_EVENT_SCOPE(GraphBuilder, "MobileBasePass");
 	SCOPE_CYCLE_COUNTER(STAT_BasePassDrawTime);
 	RDG_GPU_STAT_SCOPE(GraphBuilder, Basepass);
 
-	bool bAnyPassesAdded = false;
-
 	for (int32 ViewIndex = 0; ViewIndex < PassViews.Num(); ViewIndex++)
 	{
 		RDG_EVENT_SCOPE_CONDITIONAL(GraphBuilder, Views.Num() > 1, "View%d", ViewIndex);
@@ -324,7 +322,7 @@ void FMobileSceneRenderer::RenderMobileBasePass(FRDGBuilder& GraphBuilder, FRend
 		OpaqueBasePassParameters->RenderTargets = BasePassRenderTargets;
 		OpaqueBasePassParameters->MobileBasePass = Scene->UniformBuffers.MobileOpaqueBasePassUniformBuffer;
 
-		GraphBuilder.AddPass(RDG_EVENT_NAME("RenderOpaqueBasePass"), OpaqueBasePassParameters, ERDGPassFlags::Raster,
+		GraphBuilder.AddPass(RDG_EVENT_NAME("RenderOpaqueBasePass"), OpaqueBasePassParameters, ERDGPassFlags::Raster | ERDGPassFlags::SkipRenderPass,
 			[this, &View, OpaqueBasePassParameters](FRHICommandListImmediate& RHICmdList)
 		{
 			if (Scene->UniformBuffers.UpdateViewUniformBuffer(View))
@@ -346,13 +344,6 @@ void FMobileSceneRenderer::RenderMobileBasePass(FRDGBuilder& GraphBuilder, FRend
 			}
 
 		});
-
-		bAnyPassesAdded = true;
-	}
-
-	if (bAnyPassesAdded)
-	{
-		UpdateRenderTargetsLoadAction(BasePassRenderTargets);
 	}
 }
 
diff --git a/Engine/Source/Runtime/Renderer/Private/MobileDecalRendering.cpp b/Engine/Source/Runtime/Renderer/Private/MobileDecalRendering.cpp
--- a/Engine/Source/Runtime/Renderer/Private/MobileDecalRendering.cpp
+++ b/Engine/Source/Runtime/Renderer/Private/MobileDecalRendering.cpp
@@ -241,7 +241,7 @@ void RenderDeferredDecalsMobile(FRDGBuilder& GraphBuilder, FRenderTargetBindingS
 		PassParameters->MobileSceneTextures = MobileSceneTextures;
 
 		GraphBuilder.AddPass(
-			RDG_EVENT_NAME("MobileDeferredDecalsPass"), PassParameters, ERDGPassFlags::Raster,
+			RDG_EVENT_NAME("MobileDeferredDecalsPass"), PassParameters, ERDGPassFlags::Raster | ERDGPassFlags::SkipRenderPass,
 			[&View, SortedDecalCount, SortedDecals](FRHICommandListImmediate& RHICmdList)
 		{
 			const bool bDeferredShading = IsMobileDeferredShadingEnabled(View.GetShaderPlatform());
diff --git a/Engine/Source/Runtime/Renderer/Private/MobileDeferredShadingPass.cpp b/Engine/Source/Runtime/Renderer/Private/MobileDeferredShadingPass.cpp
--- a/Engine/Source/Runtime/Renderer/Private/MobileDeferredShadingPass.cpp
+++ b/Engine/Source/Runtime/Renderer/Private/MobileDeferredShadingPass.cpp
@@ -333,7 +333,7 @@ void MobileDeferredShadingPass(FRDGBuilder& GraphBuilder, FRenderTargetBindingSl
 	PassParameters->MobileSceneTextures = MobileSceneTextures;
 
 	GraphBuilder.AddPass(
-		RDG_EVENT_NAME("MobileDeferredShadingPass"), PassParameters, ERDGPassFlags::Raster,
+		RDG_EVENT_NAME("MobileDeferredShadingPass"), PassParameters, ERDGPassFlags::Raster | ERDGPassFlags::SkipRenderPass,
 		[&Scene, &View, &SortedLightSet](FRHICommandListImmediate& RHICmdList)
 	{
 
diff --git a/Engine/Source/Runtime/Renderer/Private/MobileDeferredShadingPass.h b/Engine/Source/Runtime/Renderer/Private/MobileDeferredShadingPass.h
--- a/Engine/Source/Runtime/Renderer/Private/MobileDeferredShadingPass.h
+++ b/Engine/Source/Runtime/Renderer/Private/MobileDeferredShadingPass.h
@@ -3,12 +3,15 @@
 #pragma once
 
 #include "CoreMinimal.h"
+#include "RenderGraphDefinitions.h"
 
+class FMobileSceneTextureUniformParameters;
+class FRDGBuilder;
 class FRHICommandListImmediate;
 class FScene;
 class FViewInfo;
+struct FRenderTargetBindingSlots;
 struct FSortedLightSetSceneInfo;
-class FMobileSceneTextureUniformParameters;
 
 extern int32 GMobileUseClusteredDeferredShading;
 
diff --git a/Engine/Source/Runtime/Renderer/Private/MobileShadingRenderer.cpp b/Engine/Source/Runtime/Renderer/Private/MobileShadingRenderer.cpp
--- a/Engine/Source/Runtime/Renderer/Private/MobileShadingRenderer.cpp
+++ b/Engine/Source/Runtime/Renderer/Private/MobileShadingRenderer.cpp
@@ -159,7 +159,7 @@ static bool UsesCustomDepthStencilLookup(const FViewInfo& View)
 	return bUsesCustomDepthStencil;
 }
 
-BEGIN_SHADER_PARAMETER_STRUCT(FRenderOpaqueFXPassParameters, )
+BEGIN_SHADER_PARAMETER_STRUCT(FMobileRenderOpaqueFXPassParameters, )
 	SHADER_PARAMETER_RDG_UNIFORM_BUFFER(FMobileSceneTextureUniformParameters, SceneTextures)
 END_SHADER_PARAMETER_STRUCT()
 
@@ -685,7 +685,7 @@ void FMobileSceneRenderer::Render(FRHICommandListImmediate& RHICmdList)
 		
 	if (FXSystem && Views.IsValidIndex(0))
 	{
-		auto* PassParameters = GraphBuilder.AllocParameters<FRenderOpaqueFXPassParameters>();
+		auto* PassParameters = GraphBuilder.AllocParameters<FMobileRenderOpaqueFXPassParameters>();
 		PassParameters->SceneTextures = MobileSceneTextures;
 
 		// Cascade uses pixel shaders for compute stuff in PostRenderOpaque so ERDGPassFlags::Raster is needed
@@ -857,56 +857,44 @@ void FMobileSceneRenderer::RenderForward(FRDGBuilder& GraphBuilder, const TArray
 	//if the scenecolor isn't multiview but the app is, need to render as a single-view multiview due to shaders
 	BasePassRenderTargets.MultiViewCount = View.bIsMobileMultiViewEnabled ? 2 : (bIsMultiViewApplication ? 1 : 0);
 
-	// Update the render targets load action to Eload after executing the first pass that is EClear flag to let the RDG merges as many passes as possible.
-	auto UpdateRenderTargetsLoadAction = [](FRenderTargetBindingSlots& BasePassRenderTargets)
-	{
-		if (BasePassRenderTargets[0].GetLoadAction() == ERenderTargetLoadAction::EClear)
-		{
-			BasePassRenderTargets[0].SetLoadAction(ERenderTargetLoadAction::ELoad);
-			BasePassRenderTargets.DepthStencil.SetDepthLoadAction(ERenderTargetLoadAction::ELoad);
-			BasePassRenderTargets.DepthStencil.SetStencilLoadAction(ERenderTargetLoadAction::ELoad);
-		}
-	};
+	auto* OpaqueBasePassParameters = GraphBuilder.AllocParameters<FRenderTargetParameters>();
+	OpaqueBasePassParameters->RenderTargets = BasePassRenderTargets;
 
-	RDG_EVENT_SCOPE(GraphBuilder, "SceneColorRendering");
+	GraphBuilder.AddPass(
+		{},
+		OpaqueBasePassParameters,
+		ERDGPassFlags::Raster | ERDGPassFlags::SkipRenderPass,
+		[OpaqueBasePassParameters](FRHICommandListImmediate& RHICmdList)
+	{
+		RHICmdList.BeginRenderPass(GetRenderPassInfo(OpaqueBasePassParameters), TEXT("SceneColorRendering"));
+	});
 
 	if (GIsEditor && !View.bIsSceneCapture)
 	{
-		auto* PassParameters = GraphBuilder.AllocParameters<FRenderTargetParameters>();
-		PassParameters->RenderTargets = BasePassRenderTargets;
-
-		GraphBuilder.AddPass(RDG_EVENT_NAME("ClearSceneColorPass"), PassParameters, ERDGPassFlags::Raster,
-			[this](FRHICommandListImmediate& RHICmdList)
+		AddPass(GraphBuilder, [this](FRHICommandListImmediate& RHICmdList)
 		{
 			DrawClearQuad(RHICmdList, Views[0].BackgroundColor);
 		});
-
-		UpdateRenderTargetsLoadAction(BasePassRenderTargets);
 	}
 
 	// Depth pre-pass
 	GraphBuilder.SetCommandListStat(GET_STATID(STAT_CLM_MobilePrePass));
-	RenderPrePass(GraphBuilder, BasePassRenderTargets, UpdateRenderTargetsLoadAction);
+	RenderPrePass(GraphBuilder, BasePassRenderTargets);
 	
 	// Opaque and masked
 	GraphBuilder.SetCommandListStat(GET_STATID(STAT_CLMM_Opaque));
-	RenderMobileBasePass(GraphBuilder, BasePassRenderTargets, UpdateRenderTargetsLoadAction, ViewList);
+	RenderMobileBasePass(GraphBuilder, BasePassRenderTargets, ViewList);
 	AddDispatchToRHIThreadPass(GraphBuilder);
 
 #if !(UE_BUILD_SHIPPING || UE_BUILD_TEST)
 	if (ViewFamily.UseDebugViewPS())
 	{
-		auto* PassParameters = GraphBuilder.AllocParameters<FRenderTargetParameters>();
-		PassParameters->RenderTargets = BasePassRenderTargets;
-
 		// Here we use the base pass depth result to get z culling for opaque and masque.
 		// The color needs to be cleared at this point since shader complexity renders in additive.
-		GraphBuilder.AddPass(RDG_EVENT_NAME("ClearDebugViewPass"), PassParameters, ERDGPassFlags::Raster,
-			[this](FRHICommandListImmediate& RHICmdList)
+		AddPass(GraphBuilder, [](FRHICommandListImmediate& RHICmdList)
 		{
 			DrawClearQuad(RHICmdList, FLinearColor::Black);
 		});
-
 		RenderMobileDebugView(GraphBuilder, ViewList);
 		AddDispatchToRHIThreadPass(GraphBuilder);
 	}
@@ -923,13 +911,13 @@ void FMobileSceneRenderer::RenderForward(FRDGBuilder& GraphBuilder, const TArray
 
 	if (ViewFamily.ViewExtensions.Num() > 1)
 	{
-		auto* PassParameters = GraphBuilder.AllocParameters<FMobilePostBasePassViewExtensionParameters>();
-		PassParameters->RenderTargets = BasePassRenderTargets;
-		PassParameters->SceneTextures = MobileSceneTextures;
+		auto* ViewExtensionPassParameters = GraphBuilder.AllocParameters<FMobilePostBasePassViewExtensionParameters>();
+		ViewExtensionPassParameters->RenderTargets = BasePassRenderTargets;
+		ViewExtensionPassParameters->SceneTextures = MobileSceneTextures;
 
 		GraphBuilder.AddPass(
 			{},
-			PassParameters,
+			ViewExtensionPassParameters,
 			ERDGPassFlags::Raster | ERDGPassFlags::SkipRenderPass | ERDGPassFlags::NeverCull,
 			[this](FRHICommandListImmediate& RHICmdList)
 		{
@@ -944,6 +932,16 @@ void FMobileSceneRenderer::RenderForward(FRDGBuilder& GraphBuilder, const TArray
 			}
 		});
 	}
+
+	// Split if we need to render translucency in a separate render pass
+	// Split if we need to render pixel projected reflection
+	if (bRequiresMultiPass || bRequiresPixelProjectedPlanarRelfectionPass)
+	{
+		AddPass(GraphBuilder, [](FRHICommandListImmediate& RHICmdList)
+		{
+			RHICmdList.EndRenderPass();
+		});
+	}
 	   
 	GraphBuilder.SetCommandListStat(GET_STATID(STAT_CLMM_Translucency));
 
@@ -963,7 +961,6 @@ void FMobileSceneRenderer::RenderForward(FRDGBuilder& GraphBuilder, const TArray
 		RenderPixelProjectedReflection(GraphBuilder, SceneColorMSAA.Resolve, SceneDepthMSAA.Resolve, PixelProjectedReflectionTexture, PlanarReflectionSceneProxy);
 	}
 
-
 	if (bRequiresMultiPass || bRequiresPixelProjectedPlanarRelfectionPass)
 	{
 		FExclusiveDepthStencil::Type ExclusiveDepthStencil = FExclusiveDepthStencil::DepthRead_StencilRead;
@@ -980,24 +977,32 @@ void FMobileSceneRenderer::RenderForward(FRDGBuilder& GraphBuilder, const TArray
 			ExclusiveDepthStencil = FExclusiveDepthStencil::DepthWrite_StencilWrite;
 		}
 
+		BasePassRenderTargets[0].SetLoadAction(ERenderTargetLoadAction::ELoad);
+		BasePassRenderTargets.DepthStencil.SetDepthLoadAction(ERenderTargetLoadAction::ELoad);
+		BasePassRenderTargets.DepthStencil.SetStencilLoadAction(ERenderTargetLoadAction::ELoad);
 		BasePassRenderTargets.DepthStencil.SetDepthStencilAccess(ExclusiveDepthStencil);
-
 		BasePassRenderTargets.NumOcclusionQueries = 0;
 		BasePassRenderTargets.SubpassHint = ESubpassHint::DepthReadSubpass;
-	}
 
-	{
-		auto* PassParameters = GraphBuilder.AllocParameters<FRenderTargetParameters>();
-		PassParameters->RenderTargets = BasePassRenderTargets;
+		auto* TranslucencyBasePassParameters = GraphBuilder.AllocParameters<FRenderTargetParameters>();
+		TranslucencyBasePassParameters->RenderTargets = BasePassRenderTargets;
 
-		GraphBuilder.AddPass(RDG_EVENT_NAME("Nextsubpass"), PassParameters, ERDGPassFlags::Raster,
-			[this](FRHICommandListImmediate& RHICmdList)
+		GraphBuilder.AddPass(
+			{},
+			TranslucencyBasePassParameters,
+			ERDGPassFlags::Raster | ERDGPassFlags::SkipRenderPass,
+			[TranslucencyBasePassParameters](FRHICommandListImmediate& RHICmdList)
 		{
-			// scene depth is read only and can be fetched
-			RHICmdList.NextSubpass();
+			RHICmdList.BeginRenderPass(GetRenderPassInfo(TranslucencyBasePassParameters), TEXT("SceneColorTranslucencyRendering"));
 		});
 	}
 
+	AddPass(GraphBuilder, [](FRHICommandListImmediate& RHICmdList)
+	{
+		// scene depth is read only and can be fetched
+		RHICmdList.NextSubpass();
+	});
+
 	if (!View.bIsPlanarReflection)
 	{
 		if (ViewFamily.EngineShowFlags.Decals)
@@ -1009,7 +1014,7 @@ void FMobileSceneRenderer::RenderForward(FRDGBuilder& GraphBuilder, const TArray
 		if (ViewFamily.EngineShowFlags.DynamicShadows)
 		{
 			CSV_SCOPED_TIMING_STAT_EXCLUSIVE(RenderShadowProjections);
-			RenderModulatedShadowProjections(GraphBuilder, BasePassRenderTargets, MobileSceneTextures);
+			RenderModulatedShadowProjections(GraphBuilder, MobileSceneTextures);
 		}
 	}
 	
@@ -1029,7 +1034,7 @@ void FMobileSceneRenderer::RenderForward(FRDGBuilder& GraphBuilder, const TArray
 	if (bAdrenoOcclusionMode)
 	{
 		GraphBuilder.SetCommandListStat(GET_STATID(STAT_CLMM_Occlusion));
-		AddPass(GraphBuilder, [this](FRHICommandListImmediate& RHICmdList)
+		AddPass(GraphBuilder, [](FRHICommandListImmediate& RHICmdList)
 		{
 			// flush
 			RHICmdList.SubmitCommandsHint();
@@ -1045,6 +1050,11 @@ void FMobileSceneRenderer::RenderForward(FRDGBuilder& GraphBuilder, const TArray
 	{
 		PreTonemapMSAA(GraphBuilder, BasePassRenderTargets);
 	}
+
+	AddPass(GraphBuilder, [](FRHICommandListImmediate& RHICmdList)
+	{
+		RHICmdList.EndRenderPass();
+	});
 }
 
 void FMobileSceneRenderer::RenderDeferred(FRDGBuilder& GraphBuilder, const TArrayView<const FViewInfo*> ViewList, const FSortedLightSetSceneInfo& SortedLightSet, FRDGTextureRef ViewFamilyTexture, FRDGTextureMSAA& SceneColorMSAA, FRDGTextureMSAA& SceneDepthMSAA, TRDGUniformBufferRef<FMobileSceneTextureUniformParameters> MobileSceneTextures)
@@ -1068,44 +1078,33 @@ void FMobileSceneRenderer::RenderDeferred(FRDGBuilder& GraphBuilder, const TArra
 	BasePassRenderTargets.FoveationTexture = nullptr;
 	BasePassRenderTargets.MultiViewCount = 0;
 
-	// Update the render targets load action to Eload after adding the first pass with EClear flag to let the RDG merge all passes together.
-	auto UpdateRenderTargetsLoadAction = [](FRenderTargetBindingSlots& BasePassRenderTargets)
-	{
-		if (BasePassRenderTargets[0].GetLoadAction() == ERenderTargetLoadAction::EClear)
-		{
-			for (uint32 i = 0; i < 5; ++i)
-			{
-				BasePassRenderTargets[i].SetLoadAction(ERenderTargetLoadAction::ELoad);
-			}
-			
-			BasePassRenderTargets.DepthStencil.SetDepthLoadAction(ERenderTargetLoadAction::ELoad);
-			BasePassRenderTargets.DepthStencil.SetStencilLoadAction(ERenderTargetLoadAction::ELoad);
-		}
-	};
+	auto* OpaqueBasePassParameters = GraphBuilder.AllocParameters<FRenderTargetParameters>();
+	OpaqueBasePassParameters->RenderTargets = BasePassRenderTargets;
 
-	RDG_EVENT_SCOPE(GraphBuilder, "SceneColorRendering");
+	GraphBuilder.AddPass(
+		{},
+		OpaqueBasePassParameters,
+		ERDGPassFlags::Raster | ERDGPassFlags::SkipRenderPass,
+		[OpaqueBasePassParameters](FRHICommandListImmediate& RHICmdList)
+	{
+		RHICmdList.BeginRenderPass(GetRenderPassInfo(OpaqueBasePassParameters), TEXT("BasePassRendering"));
+	});
 	
 	if (GIsEditor && !Views[0].bIsSceneCapture)
 	{
-		auto* PassParameters = GraphBuilder.AllocParameters<FRenderTargetParameters>();
-		PassParameters->RenderTargets = BasePassRenderTargets;
-
-		GraphBuilder.AddPass(RDG_EVENT_NAME("ClearSceneColorPass"), PassParameters, ERDGPassFlags::Raster,
-			[this](FRHICommandListImmediate& RHICmdList)
+		AddPass(GraphBuilder, [this](FRHICommandListImmediate& RHICmdList)
 		{
 			DrawClearQuad(RHICmdList, Views[0].BackgroundColor);
 		});
-
-		UpdateRenderTargetsLoadAction(BasePassRenderTargets);
 	}
 
 	GraphBuilder.SetCommandListStat(GET_STATID(STAT_CLM_MobilePrePass));
 	// Depth pre-pass
-	RenderPrePass(GraphBuilder, BasePassRenderTargets, UpdateRenderTargetsLoadAction);
+	RenderPrePass(GraphBuilder, BasePassRenderTargets);
 	
 	// Opaque and masked
 	GraphBuilder.SetCommandListStat(GET_STATID(STAT_CLMM_Opaque));
-	RenderMobileBasePass(GraphBuilder, BasePassRenderTargets, UpdateRenderTargetsLoadAction, ViewList);
+	RenderMobileBasePass(GraphBuilder, BasePassRenderTargets, ViewList);
 	AddDispatchToRHIThreadPass(GraphBuilder);
 
 	// Issue occlusion queries
@@ -1121,17 +1120,11 @@ void FMobileSceneRenderer::RenderDeferred(FRDGBuilder& GraphBuilder, const TArra
 
 	if (!bRequiresMultiPass)
 	{
+		AddPass(GraphBuilder, [](FRHICommandListImmediate& RHICmdList)
 		{
-			auto* PassParameters = GraphBuilder.AllocParameters<FRenderTargetParameters>();
-			PassParameters->RenderTargets = BasePassRenderTargets;
-
-			GraphBuilder.AddPass(RDG_EVENT_NAME("Nextsubpass"), PassParameters, ERDGPassFlags::Raster,
-				[this](FRHICommandListImmediate& RHICmdList)
-			{
-				// SceneColor + GBuffer write, SceneDepth is read only
-				RHICmdList.NextSubpass();
-			});
-		}
+			// SceneColor + GBuffer write, SceneDepth is read only
+			RHICmdList.NextSubpass();
+		});
 		
 		if (ViewFamily.EngineShowFlags.Decals)
 		{
@@ -1139,17 +1132,11 @@ void FMobileSceneRenderer::RenderDeferred(FRDGBuilder& GraphBuilder, const TArra
 			RenderDecals(GraphBuilder, BasePassRenderTargets, MobileSceneTextures);
 		}
 
+		AddPass(GraphBuilder, [](FRHICommandListImmediate& RHICmdList)
 		{
-			auto* PassParameters = GraphBuilder.AllocParameters<FRenderTargetParameters>();
-			PassParameters->RenderTargets = BasePassRenderTargets;
-
-			GraphBuilder.AddPass(RDG_EVENT_NAME("Nextsubpass"), PassParameters, ERDGPassFlags::Raster,
-				[this](FRHICommandListImmediate& RHICmdList)
-			{
-				// SceneColor write, SceneDepth is read only
-				RHICmdList.NextSubpass();
-			});
-		}
+			// SceneColor write, SceneDepth is read only
+			RHICmdList.NextSubpass();
+		});
 		
 		MobileDeferredShadingPass(GraphBuilder, BasePassRenderTargets, MobileSceneTextures, *Scene, *ViewList[0], SortedLightSet);
 		// Draw translucency.
@@ -1160,34 +1147,56 @@ void FMobileSceneRenderer::RenderDeferred(FRDGBuilder& GraphBuilder, const TArra
 			RenderTranslucency(GraphBuilder, BasePassRenderTargets, ViewList);
 			AddPass(GraphBuilder, PollOcclusionQueriesAndDispatchToRHIThreadPass);
 		}
+
+		AddPass(GraphBuilder, [](FRHICommandListImmediate& RHICmdList)
+		{
+			RHICmdList.EndRenderPass();
+		});
 	}
 	else
 	{
+		AddPass(GraphBuilder, [](FRHICommandListImmediate& RHICmdList)
 		{
-			auto* PassParameters = GraphBuilder.AllocParameters<FRenderTargetParameters>();
-			PassParameters->RenderTargets = BasePassRenderTargets;
-
-			GraphBuilder.AddPass(RDG_EVENT_NAME("Nextsubpass"), PassParameters, ERDGPassFlags::Raster,
-				[this](FRHICommandListImmediate& RHICmdList)
-			{
-				RHICmdList.NextSubpass();
-				RHICmdList.NextSubpass();
-			});
-		}
+			RHICmdList.NextSubpass();
+			RHICmdList.NextSubpass();
+			RHICmdList.EndRenderPass();
+		});
 		
 		// SceneColor + GBuffer write, SceneDepth is read only
 		{
+			for (uint32 i = 0; i < UE_ARRAY_COUNT(ColorTargets); ++i)
+			{
+				BasePassRenderTargets[i].SetLoadAction(ERenderTargetLoadAction::ELoad);
+			}
+
+			BasePassRenderTargets.DepthStencil.SetDepthLoadAction(ERenderTargetLoadAction::ELoad);
+			BasePassRenderTargets.DepthStencil.SetStencilLoadAction(ERenderTargetLoadAction::ELoad);
 			BasePassRenderTargets.DepthStencil.SetDepthStencilAccess(FExclusiveDepthStencil::DepthRead_StencilRead);
 			BasePassRenderTargets.SubpassHint = ESubpassHint::None;
 			BasePassRenderTargets.NumOcclusionQueries = 0;
-			
-			RDG_EVENT_SCOPE(GraphBuilder, "AfterBasePass");
+
+			auto* DecalPassParameters = GraphBuilder.AllocParameters<FRenderTargetParameters>();
+			DecalPassParameters->RenderTargets = BasePassRenderTargets;
+
+			GraphBuilder.AddPass(
+				{},
+				DecalPassParameters,
+				ERDGPassFlags::Raster | ERDGPassFlags::SkipRenderPass,
+				[DecalPassParameters](FRHICommandListImmediate& RHICmdList)
+			{
+				RHICmdList.BeginRenderPass(GetRenderPassInfo(DecalPassParameters), TEXT("AfterBasePass"));
+			});
 
 			if (ViewFamily.EngineShowFlags.Decals)
 			{
 				CSV_SCOPED_TIMING_STAT_EXCLUSIVE(RenderDecals);
 				RenderDecals(GraphBuilder, BasePassRenderTargets, MobileSceneTextures);
 			}
+
+			AddPass(GraphBuilder, [](FRHICommandListImmediate& RHICmdList)
+			{
+				RHICmdList.EndRenderPass();
+			});
 		}
 
 		// SceneColor write, SceneDepth is read only
@@ -1198,7 +1207,17 @@ void FMobileSceneRenderer::RenderDeferred(FRDGBuilder& GraphBuilder, const TArra
 			}
 			BasePassRenderTargets.DepthStencil.SetDepthStencilAccess(FExclusiveDepthStencil::DepthRead_StencilWrite);
 
-			RDG_EVENT_SCOPE(GraphBuilder, "MobileShadingPass");
+			auto* DeferredShadingPassParameters = GraphBuilder.AllocParameters<FRenderTargetParameters>();
+			DeferredShadingPassParameters->RenderTargets = BasePassRenderTargets;
+
+			GraphBuilder.AddPass(
+				{},
+				DeferredShadingPassParameters,
+				ERDGPassFlags::Raster | ERDGPassFlags::SkipRenderPass,
+				[DeferredShadingPassParameters](FRHICommandListImmediate& RHICmdList)
+			{
+				RHICmdList.BeginRenderPass(GetRenderPassInfo(DeferredShadingPassParameters), TEXT("MobileShadingPass"));
+			});
 			
 			MobileDeferredShadingPass(GraphBuilder, BasePassRenderTargets, MobileSceneTextures, *Scene, *ViewList[0], SortedLightSet);
 			// Draw translucency.
@@ -1209,6 +1228,11 @@ void FMobileSceneRenderer::RenderDeferred(FRDGBuilder& GraphBuilder, const TArra
 				RenderTranslucency(GraphBuilder, BasePassRenderTargets, ViewList);
 				AddPass(GraphBuilder, PollOcclusionQueriesAndDispatchToRHIThreadPass);
 			}
+
+			AddPass(GraphBuilder, [](FRHICommandListImmediate& RHICmdList)
+			{
+				RHICmdList.EndRenderPass();
+			});
 		}
 	}
 }
@@ -1477,7 +1501,7 @@ void FMobileSceneRenderer::PreTonemapMSAA(FRDGBuilder& GraphBuilder, FRenderTarg
 	GraphBuilder.AddPass(
 		RDG_EVENT_NAME("PreTonemapMSAAPass"),
 		PassParameters,
-		ERDGPassFlags::Raster,
+		ERDGPassFlags::Raster | ERDGPassFlags::SkipRenderPass,
 		[this, TargetSize, VertexShader, PixelShader](FRHICommandListImmediate& RHICmdList)
 	{
 
diff --git a/Engine/Source/Runtime/Renderer/Private/MobileTranslucentRendering.cpp b/Engine/Source/Runtime/Renderer/Private/MobileTranslucentRendering.cpp
--- a/Engine/Source/Runtime/Renderer/Private/MobileTranslucentRendering.cpp
+++ b/Engine/Source/Runtime/Renderer/Private/MobileTranslucentRendering.cpp
@@ -55,7 +55,7 @@ void FMobileSceneRenderer::RenderTranslucency(FRDGBuilder& GraphBuilder, FRender
 			TranslucencyBasePassParameters->RenderTargets = BasePassRenderTargets;
 			TranslucencyBasePassParameters->MobileBasePass = Scene->UniformBuffers.MobileTranslucentBasePassUniformBuffer;
 
-			GraphBuilder.AddPass(RDG_EVENT_NAME("RenderTranslucencyBasePass"), TranslucencyBasePassParameters, ERDGPassFlags::Raster,
+			GraphBuilder.AddPass(RDG_EVENT_NAME("RenderTranslucencyBasePass"), TranslucencyBasePassParameters, ERDGPassFlags::Raster | ERDGPassFlags::SkipRenderPass,
 				[this, &View, TranslucencyPass](FRHICommandListImmediate& RHICmdList)
 			{
 				RHICmdList.SetViewport(View.ViewRect.Min.X, View.ViewRect.Min.Y, 0.0f, View.ViewRect.Max.X, View.ViewRect.Max.Y, 1.0f);
diff --git a/Engine/Source/Runtime/Renderer/Private/SceneOcclusion.cpp b/Engine/Source/Runtime/Renderer/Private/SceneOcclusion.cpp
--- a/Engine/Source/Runtime/Renderer/Private/SceneOcclusion.cpp
+++ b/Engine/Source/Runtime/Renderer/Private/SceneOcclusion.cpp
@@ -1436,7 +1436,7 @@ void FMobileSceneRenderer::RenderOcclusion(FRDGBuilder& GraphBuilder, FRenderTar
 			auto* PassParameters = GraphBuilder.AllocParameters<FRenderTargetParameters>();
 			PassParameters->RenderTargets = BasePassRenderTargets;
 
-			GraphBuilder.AddPass(RDG_EVENT_NAME("BeginOcclusionTestsPass"), PassParameters, ERDGPassFlags::Raster,
+			GraphBuilder.AddPass(RDG_EVENT_NAME("BeginOcclusionTestsPass"), PassParameters, ERDGPassFlags::Raster | ERDGPassFlags::SkipRenderPass,
 				[this, LocalQueriesPerView = MoveTemp(QueriesPerView)](FRHICommandListImmediate& RHICmdList)
 			{
 				BeginOcclusionTests(RHICmdList, Views, FeatureLevel, LocalQueriesPerView, 1.0f);
diff --git a/Engine/Source/Runtime/Renderer/Private/SceneRendering.h b/Engine/Source/Runtime/Renderer/Private/SceneRendering.h
--- a/Engine/Source/Runtime/Renderer/Private/SceneRendering.h
+++ b/Engine/Source/Runtime/Renderer/Private/SceneRendering.h
@@ -2080,10 +2080,10 @@ protected:
 
 	void InitViews(FRDGBuilder& GraphBuilder);
 
-	void RenderPrePass(FRDGBuilder& GraphBuilder, FRenderTargetBindingSlots& BasePassRenderTargets, TFunction<void(FRenderTargetBindingSlots&)> UpdateRenderTargetsLoadAction);
+	void RenderPrePass(FRDGBuilder& GraphBuilder, FRenderTargetBindingSlots& BasePassRenderTargets);
 
 	/** Renders the opaque base pass for mobile. */
-	void RenderMobileBasePass(FRDGBuilder& GraphBuilder, FRenderTargetBindingSlots& BasePassRenderTargets, TFunction<void(FRenderTargetBindingSlots &)> UpdateRenderTargetsLoadAction, const TArrayView<const FViewInfo*> PassViews);
+	void RenderMobileBasePass(FRDGBuilder& GraphBuilder, FRenderTargetBindingSlots& BasePassRenderTargets, const TArrayView<const FViewInfo*> PassViews);
 
 	void RenderMobileEditorPrimitives(FRHICommandList& RHICmdList, const FViewInfo& View, const FMeshPassProcessorRenderState& DrawRenderState);
 
@@ -2091,7 +2091,7 @@ protected:
 	void RenderMobileDebugView(FRDGBuilder& GraphBuilder, const TArrayView<const FViewInfo*> PassViews);
 
 	/** Render modulated shadow projections in to the scene, loops over any unrendered shadows until all are processed.*/
-	void RenderModulatedShadowProjections(FRDGBuilder& GraphBuilder, FRenderTargetBindingSlots& BasePassRenderTargets, TRDGUniformBufferRef<FMobileSceneTextureUniformParameters> MobileSceneTextures);
+	void RenderModulatedShadowProjections(FRDGBuilder& GraphBuilder, TRDGUniformBufferRef<FMobileSceneTextureUniformParameters> MobileSceneTextures);
 
 	/** Resolves scene depth in case hardware does not support reading depth in the shader */
 	void ConditionalResolveSceneDepth(FRDGBuilder& GraphBuilder, const FViewInfo& View, FRDGTextureMSAA& SceneDepth);
diff --git a/Engine/Source/Runtime/Renderer/Private/ShadowRendering.cpp b/Engine/Source/Runtime/Renderer/Private/ShadowRendering.cpp
--- a/Engine/Source/Runtime/Renderer/Private/ShadowRendering.cpp
+++ b/Engine/Source/Runtime/Renderer/Private/ShadowRendering.cpp
@@ -2044,7 +2044,7 @@ void FDeferredShadingSceneRenderer::RenderDeferredShadowProjections(
 	}
 }
 
-void FMobileSceneRenderer::RenderModulatedShadowProjections(FRDGBuilder& GraphBuilder, FRenderTargetBindingSlots& BasePassRenderTargets, TRDGUniformBufferRef<FMobileSceneTextureUniformParameters> MobileSceneTextures)
+void FMobileSceneRenderer::RenderModulatedShadowProjections(FRDGBuilder& GraphBuilder, TRDGUniformBufferRef<FMobileSceneTextureUniformParameters> MobileSceneTextures)
 {
 	if (IsSimpleForwardShadingEnabled(ShaderPlatform) || !ViewFamily.EngineShowFlags.DynamicShadows)
 	{
@@ -2058,7 +2058,6 @@ void FMobileSceneRenderer::RenderModulatedShadowProjections(FRDGBuilder& GraphBu
 
 	FShadowProjectionPassParameters CommonPassParameters;
 	CommonPassParameters.SceneTextures.MobileSceneTextures = MobileSceneTextures;
-	CommonPassParameters.RenderTargets = BasePassRenderTargets;
 
 	// render shadowmaps for relevant lights.
 	for (TSparseArray<FLightSceneInfoCompact>::TConstIterator LightIt(Scene->Lights); LightIt; ++LightIt)
