Sky light real time capture now does not require a sky atmosphere to work. A mesh with a isSky material or clouds is also enough.

diff --git a/Engine/Source/Runtime/Renderer/Private/ReflectionEnvironmentRealTimeCapture.cpp b/Engine/Source/Runtime/Renderer/Private/ReflectionEnvironmentRealTimeCapture.cpp
--- a/Engine/Source/Runtime/Renderer/Private/ReflectionEnvironmentRealTimeCapture.cpp
+++ b/Engine/Source/Runtime/Renderer/Private/ReflectionEnvironmentRealTimeCapture.cpp
@@ -337,6 +337,21 @@ void FScene::AllocateAndCaptureFrameSkyEnvMap(
 		GRenderTargetPool.FindFreeElement(GraphBuilder.RHICmdList, SkyCubeTexDesc, ConvolvedSkyRenderTarget[1], TEXT("ConvolvedSkyRenderTarget"), ERenderTargetTransience::NonTransient);
 	}
 
+	auto ClearCubeFace = [&](FRDGTextureRef SkyCubeTexture, int32 CubeFace)
+	{
+		FRenderTargetParameters* Parameters = GraphBuilder.AllocParameters<FRenderTargetParameters>();
+		Parameters->RenderTargets[0] = FRenderTargetBinding(SkyCubeTexture, ERenderTargetLoadAction::ENoAction, 0, CubeFace);
+
+		FLinearColor ClearColor = FLinearColor::Black;
+		GraphBuilder.AddPass(
+			RDG_EVENT_NAME("ClearSkyRenderTarget"),
+			Parameters,
+			ERDGPassFlags::Raster,
+			[Parameters, ClearColor](FRHICommandList& RHICmdList)
+			{
+				DrawClearQuad(RHICmdList, ClearColor);
+			});
+	};
 
 	auto RenderCubeFaces_SkyCloud = [&](bool bExecuteSky, bool bExecuteCloud, TRefCountPtr<IPooledRenderTarget>& SkyRenderTarget)
 	{
@@ -344,42 +359,45 @@ void FScene::AllocateAndCaptureFrameSkyEnvMap(
 
 		FRDGTextureRef SkyCubeTexture = GraphBuilder.RegisterExternalTexture(SkyRenderTarget, TEXT("SkyRenderTarget"));
 
-		if (bShouldRenderSkyAtmosphere)
+		if (bExecuteSky || bExecuteCloud)
 		{
-			FSkyAtmosphereRenderSceneInfo& SkyInfo = *GetSkyAtmosphereSceneInfo();
-			const FSkyAtmosphereSceneProxy& SkyAtmosphereSceneProxy = SkyInfo.GetSkyAtmosphereSceneProxy();
-
 			FRDGTextureRef BlackDummy2dTex = GraphBuilder.RegisterExternalTexture(GSystemTextures.BlackDummy);
 			FRDGTextureRef BlackDummy3dTex = GraphBuilder.RegisterExternalTexture(GSystemTextures.VolumetricBlackDummy);
+			const bool CaptureShadowFromOpaque = CVarRealTimeReflectionCaptureShadowFromOpaque.GetValueOnRenderThread() > 0;
 
 			FSkyAtmosphereRenderContext SkyRC;
+			const FAtmosphereSetup* AtmosphereSetup = nullptr;
+			if (bShouldRenderSkyAtmosphere)
+			{
+				FSkyAtmosphereRenderSceneInfo& SkyInfo = *GetSkyAtmosphereSceneInfo();
+				const FSkyAtmosphereSceneProxy& SkyAtmosphereSceneProxy = SkyInfo.GetSkyAtmosphereSceneProxy();
 
-			// Global data constant between faces
-			const FAtmosphereSetup& AtmosphereSetup = SkyAtmosphereSceneProxy.GetAtmosphereSetup();
-			SkyRC.bFastSky = false;
-			SkyRC.bFastAerialPerspective = false;
-			SkyRC.bFastAerialPerspectiveDepthTest = false;
-			SkyRC.bSecondAtmosphereLightEnabled = IsSecondAtmosphereLightEnabled();
+				// Global data constant between faces
+				AtmosphereSetup = &SkyAtmosphereSceneProxy.GetAtmosphereSetup();
 
-			const bool CaptureShadowFromOpaque = CVarRealTimeReflectionCaptureShadowFromOpaque.GetValueOnRenderThread() > 0;
+				SkyRC.bFastSky = false;
+				SkyRC.bFastAerialPerspective = false;
+				SkyRC.bFastAerialPerspectiveDepthTest = false;
+				SkyRC.bSecondAtmosphereLightEnabled = IsSecondAtmosphereLightEnabled();
 
-			// Enable opaque shadow on sky if needed
-			SkyRC.bShouldSampleOpaqueShadow = false;
-			if (CaptureShadowFromOpaque)
-			{
-				SkyAtmosphereLightShadowData LightShadowData;
-				SkyRC.bShouldSampleOpaqueShadow = ShouldSkySampleAtmosphereLightsOpaqueShadow(*Scene, SceneRenderer.VisibleLightInfos, LightShadowData);
-				GetSkyAtmosphereLightsUniformBuffers(SkyRC.LightShadowShaderParams0UniformBuffer, SkyRC.LightShadowShaderParams1UniformBuffer,
-					LightShadowData, CubeView, SkyRC.bShouldSampleOpaqueShadow, UniformBuffer_SingleDraw);
-			}
+				// Enable opaque shadow on sky if needed
+				SkyRC.bShouldSampleOpaqueShadow = false;
+				if (CaptureShadowFromOpaque)
+				{
+					SkyAtmosphereLightShadowData LightShadowData;
+					SkyRC.bShouldSampleOpaqueShadow = ShouldSkySampleAtmosphereLightsOpaqueShadow(*Scene, SceneRenderer.VisibleLightInfos, LightShadowData);
+					GetSkyAtmosphereLightsUniformBuffers(SkyRC.LightShadowShaderParams0UniformBuffer, SkyRC.LightShadowShaderParams1UniformBuffer,
+						LightShadowData, CubeView, SkyRC.bShouldSampleOpaqueShadow, UniformBuffer_SingleDraw);
+				}
 
-			SkyRC.bUseDepthBoundTestIfPossible = false;
-			SkyRC.bForceRayMarching = true;				// We do not have any valid view LUT
-			SkyRC.bDepthReadDisabled = true;
-			SkyRC.bDisableBlending = true;
+				SkyRC.bUseDepthBoundTestIfPossible = false;
+				SkyRC.bForceRayMarching = true;				// We do not have any valid view LUT
+				SkyRC.bDepthReadDisabled = true;
+				SkyRC.bDisableBlending = true;
 
-			SkyRC.TransmittanceLut = GraphBuilder.RegisterExternalTexture(SkyInfo.GetTransmittanceLutTexture());
-			SkyRC.MultiScatteredLuminanceLut = GraphBuilder.RegisterExternalTexture(SkyInfo.GetMultiScatteredLuminanceLutTexture());
+				SkyRC.TransmittanceLut = GraphBuilder.RegisterExternalTexture(SkyInfo.GetTransmittanceLutTexture());
+				SkyRC.MultiScatteredLuminanceLut = GraphBuilder.RegisterExternalTexture(SkyInfo.GetMultiScatteredLuminanceLutTexture());
+			}
 
 			FCloudRenderContext CloudRC;
 			if (bShouldRenderVolumetricCloud)
@@ -458,30 +476,33 @@ void FScene::AllocateAndCaptureFrameSkyEnvMap(
 				CubeView.CachedViewUniformShaderParameters->RealTimeReflectionCapture = 1.0f;
 
 				// We have rendered a sky dome with identity rotation at the SkyLight position for the capture.
-				if (MainView.bSceneHasSkyMaterial)
-				{
-					// Setup a constant referential for each of the faces of the dynamic reflection capture.
-					// This is to have the FastSkyViewLUT match the one generated specifically for the capture point of view.
-					const FVector SkyViewLutReferentialForward = FVector(1.0f, 0.0f, 0.0f);
-					const FVector SkyViewLutReferentialRight = FVector(0.0f, 0.0f, -1.0f);
-					AtmosphereSetup.ComputeViewData(SkyLight->CapturePosition, SkyViewLutReferentialForward, SkyViewLutReferentialRight,
-						CubeView.CachedViewUniformShaderParameters->SkyWorldCameraOrigin, CubeView.CachedViewUniformShaderParameters->SkyPlanetCenterAndViewHeight, 
-						CubeView.CachedViewUniformShaderParameters->SkyViewLutReferential);
-
-					CubeView.CachedViewUniformShaderParameters->SkyViewLutTexture = RealTimeReflectionCaptureSkyAtmosphereViewLutTexture->GetRenderTargetItem().ShaderResourceTexture;
-				}
-				else
+				if (AtmosphereSetup)
 				{
-					// Else if there is no sky material, we assume that no material is sampling the FastSkyViewLUT texture in the sky light reflection (bFastSky=bFastAerialPerspective=false).
-					// But, we still need to udpate the sky parameters on the view according to the sky light capture position
-					const FVector SkyViewLutReferentialForward = FVector(1.0f, 0.0f, 0.0f);
-					const FVector SkyViewLutReferentialRight = FVector(0.0f, 0.0f, -1.0f);
-					AtmosphereSetup.ComputeViewData(SkyLight->CapturePosition, SkyViewLutReferentialForward, SkyViewLutReferentialRight,
-						CubeView.CachedViewUniformShaderParameters->SkyWorldCameraOrigin, CubeView.CachedViewUniformShaderParameters->SkyPlanetCenterAndViewHeight,
-						CubeView.CachedViewUniformShaderParameters->SkyViewLutReferential);
+					if (MainView.bSceneHasSkyMaterial)
+					{
+						// Setup a constant referential for each of the faces of the dynamic reflection capture.
+						// This is to have the FastSkyViewLUT match the one generated specifically for the capture point of view.
+						const FVector SkyViewLutReferentialForward = FVector(1.0f, 0.0f, 0.0f);
+						const FVector SkyViewLutReferentialRight = FVector(0.0f, 0.0f, -1.0f);
+						AtmosphereSetup->ComputeViewData(SkyLight->CapturePosition, SkyViewLutReferentialForward, SkyViewLutReferentialRight,
+							CubeView.CachedViewUniformShaderParameters->SkyWorldCameraOrigin, CubeView.CachedViewUniformShaderParameters->SkyPlanetCenterAndViewHeight,
+							CubeView.CachedViewUniformShaderParameters->SkyViewLutReferential);
+
+						CubeView.CachedViewUniformShaderParameters->SkyViewLutTexture = RealTimeReflectionCaptureSkyAtmosphereViewLutTexture->GetRenderTargetItem().ShaderResourceTexture;
+					}
+					else
+					{
+						// Else if there is no sky material, we assume that no material is sampling the FastSkyViewLUT texture in the sky light reflection (bFastSky=bFastAerialPerspective=false).
+						// But, we still need to udpate the sky parameters on the view according to the sky light capture position
+						const FVector SkyViewLutReferentialForward = FVector(1.0f, 0.0f, 0.0f);
+						const FVector SkyViewLutReferentialRight = FVector(0.0f, 0.0f, -1.0f);
+						AtmosphereSetup->ComputeViewData(SkyLight->CapturePosition, SkyViewLutReferentialForward, SkyViewLutReferentialRight,
+							CubeView.CachedViewUniformShaderParameters->SkyWorldCameraOrigin, CubeView.CachedViewUniformShaderParameters->SkyPlanetCenterAndViewHeight,
+							CubeView.CachedViewUniformShaderParameters->SkyViewLutReferential);
+					}
 				}
 
-				if (MainView.bSceneHasSkyMaterial || HasVolumetricCloud())
+				if (HasSkyAtmosphere() && (MainView.bSceneHasSkyMaterial || HasVolumetricCloud()))
 				{
 					CubeView.CachedViewUniformShaderParameters->CameraAerialPerspectiveVolume = RealTimeReflectionCaptureCamera360APLutTexture->GetRenderTargetItem().ShaderResourceTexture;
 				}
@@ -520,112 +541,120 @@ void FScene::AllocateAndCaptureFrameSkyEnvMap(
 
 				if (bExecuteSky)
 				{
-					if (MainView.bSceneHasSkyMaterial)
+					if(MainView.bSceneHasSkyMaterial || bShouldRenderSkyAtmosphere)
 					{
-						auto* PassParameters = GraphBuilder.AllocParameters<FCaptureSkyMeshReflectionPassParameters>();
-						PassParameters->RenderTargets = SkyRC.RenderTargets;
-						PassParameters->BasePass = CreateOpaqueBasePassUniformBuffer(GraphBuilder, MainView, 0);
-
-						// Setup the depth buffer
-						if (bUseDepthBuffer)
+						// If there are any mesh tagged as IsSky then we render them only, otherwise we simply render the sky atmosphere itself.
+						if (MainView.bSceneHasSkyMaterial)
 						{
-							FRDGTextureDesc CubeDepthTextureDesc = FRDGTextureDesc::Create2D(FIntPoint(CubeWidth, CubeWidth), PF_DepthStencil, SceneContext.GetDefaultDepthClear(), 
-								TexCreate_DepthStencilTargetable | TexCreate_ShaderResource);
-							CubeDepthTexture = GraphBuilder.CreateTexture(CubeDepthTextureDesc, TEXT("CubeDepthTexture"));
-							PassParameters->RenderTargets.DepthStencil = FDepthStencilBinding(CubeDepthTexture, ERenderTargetLoadAction::EClear, FExclusiveDepthStencil::DepthWrite_StencilNop);
-						}
+							auto* PassParameters = GraphBuilder.AllocParameters<FCaptureSkyMeshReflectionPassParameters>();
+							PassParameters->RenderTargets = SkyRC.RenderTargets;
+							PassParameters->BasePass = CreateOpaqueBasePassUniformBuffer(GraphBuilder, MainView, 0);
 
-						GraphBuilder.AddPass(
-							RDG_EVENT_NAME("CaptureSkyMeshReflection"),
-							PassParameters,
-							ERDGPassFlags::Raster,
-							[&MainView, CubeViewUniformBuffer, bUseDepthBuffer](FRHICommandListImmediate& RHICmdList)
+							// Setup the depth buffer
+							if (bUseDepthBuffer)
 							{
-								DrawDynamicMeshPass(MainView, RHICmdList,
-									[&MainView, &CubeViewUniformBuffer, bUseDepthBuffer](FDynamicPassMeshDrawListContext* DynamicMeshPassContext)
-									{
-										FScene* Scene = MainView.Family->Scene->GetRenderScene();
+								FRDGTextureDesc CubeDepthTextureDesc = FRDGTextureDesc::Create2D(FIntPoint(CubeWidth, CubeWidth), PF_DepthStencil, SceneContext.GetDefaultDepthClear(),
+									TexCreate_DepthStencilTargetable | TexCreate_ShaderResource);
+								CubeDepthTexture = GraphBuilder.CreateTexture(CubeDepthTextureDesc, TEXT("CubeDepthTexture"));
+								PassParameters->RenderTargets.DepthStencil = FDepthStencilBinding(CubeDepthTexture, ERenderTargetLoadAction::EClear, FExclusiveDepthStencil::DepthWrite_StencilNop);
+							}
+
+							GraphBuilder.AddPass(
+								RDG_EVENT_NAME("CaptureSkyMeshReflection"),
+								PassParameters,
+								ERDGPassFlags::Raster,
+								[&MainView, CubeViewUniformBuffer, bUseDepthBuffer](FRHICommandListImmediate& RHICmdList)
+								{
+									DrawDynamicMeshPass(MainView, RHICmdList,
+										[&MainView, &CubeViewUniformBuffer, bUseDepthBuffer](FDynamicPassMeshDrawListContext* DynamicMeshPassContext)
+										{
+											FScene* Scene = MainView.Family->Scene->GetRenderScene();
 
-										FMeshPassProcessorRenderState DrawRenderState(CubeViewUniformBuffer);
-										DrawRenderState.SetInstancedViewUniformBuffer(Scene->UniformBuffers.InstancedViewUniformBuffer);
+											FMeshPassProcessorRenderState DrawRenderState(CubeViewUniformBuffer);
+											DrawRenderState.SetInstancedViewUniformBuffer(Scene->UniformBuffers.InstancedViewUniformBuffer);
 
-										FExclusiveDepthStencil::Type BasePassDepthStencilAccess_Sky = bUseDepthBuffer ? FExclusiveDepthStencil::Type(Scene->DefaultBasePassDepthStencilAccess | FExclusiveDepthStencil::DepthWrite)
-											: FExclusiveDepthStencil::Type(Scene->DefaultBasePassDepthStencilAccess & ~FExclusiveDepthStencil::DepthWrite);
-										SetupBasePassState(BasePassDepthStencilAccess_Sky, false, DrawRenderState);
+											FExclusiveDepthStencil::Type BasePassDepthStencilAccess_Sky = bUseDepthBuffer ? FExclusiveDepthStencil::Type(Scene->DefaultBasePassDepthStencilAccess | FExclusiveDepthStencil::DepthWrite)
+												: FExclusiveDepthStencil::Type(Scene->DefaultBasePassDepthStencilAccess & ~FExclusiveDepthStencil::DepthWrite);
+											SetupBasePassState(BasePassDepthStencilAccess_Sky, false, DrawRenderState);
 
-										FSkyPassMeshProcessor PassMeshProcessor(Scene, nullptr, DrawRenderState, DynamicMeshPassContext);
-										const int32 SkyRealTimeReflectionOnlyMeshBatcheCount = MainView.SkyMeshBatches.Num();
-										for (int32 MeshBatchIndex = 0; MeshBatchIndex < SkyRealTimeReflectionOnlyMeshBatcheCount; ++MeshBatchIndex)
-										{
-											FSkyMeshBatch& SkyMeshBatch = MainView.SkyMeshBatches[MeshBatchIndex];
-											if (!SkyMeshBatch.bVisibleInRealTimeSkyCapture)
+											FSkyPassMeshProcessor PassMeshProcessor(Scene, nullptr, DrawRenderState, DynamicMeshPassContext);
+											const int32 SkyRealTimeReflectionOnlyMeshBatcheCount = MainView.SkyMeshBatches.Num();
+											for (int32 MeshBatchIndex = 0; MeshBatchIndex < SkyRealTimeReflectionOnlyMeshBatcheCount; ++MeshBatchIndex)
 											{
-												continue;
+												FSkyMeshBatch& SkyMeshBatch = MainView.SkyMeshBatches[MeshBatchIndex];
+												if (!SkyMeshBatch.bVisibleInRealTimeSkyCapture)
+												{
+													continue;
+												}
+
+												const FMeshBatch* MeshBatch = SkyMeshBatch.Mesh;
+												const FPrimitiveSceneProxy* PrimitiveSceneProxy = SkyMeshBatch.Proxy;
+												const FPrimitiveSceneInfo* PrimitiveSceneInfo = PrimitiveSceneProxy->GetPrimitiveSceneInfo();
+
+												const uint64 DefaultBatchElementMask = ~0ull;
+												PassMeshProcessor.AddMeshBatch(*MeshBatch, DefaultBatchElementMask, PrimitiveSceneProxy);
 											}
+										});
+								});
+						}
+						else
+						{
+							FSceneTextureShaderParameters SceneTextures = CreateSceneTextureShaderParameters(GraphBuilder, SceneRenderer.FeatureLevel, ESceneTextureSetupMode::SceneDepth);
+							SceneRenderer.RenderSkyAtmosphereInternal(GraphBuilder, SceneTextures, SkyRC);
+						}
 
-											const FMeshBatch* MeshBatch = SkyMeshBatch.Mesh;
-											const FPrimitiveSceneProxy* PrimitiveSceneProxy = SkyMeshBatch.Proxy;
-											const FPrimitiveSceneInfo* PrimitiveSceneInfo = PrimitiveSceneProxy->GetPrimitiveSceneInfo();
-
-											const uint64 DefaultBatchElementMask = ~0ull;
-											PassMeshProcessor.AddMeshBatch(*MeshBatch, DefaultBatchElementMask, PrimitiveSceneProxy);
-										}
-									});
-							});
+						// Also render the height fog as part of the sky render pass when time slicing is enabled.
+						if (Scene && Scene->ExponentialFogs.Num() > 0)
+						{
+							FRenderRealTimeReflectionHeightFogVS::FPermutationDomain VsPermutationVector;
+							TShaderMapRef<FRenderRealTimeReflectionHeightFogVS> VertexShader(GetGlobalShaderMap(SkyRC.FeatureLevel), VsPermutationVector);
+
+							FRenderRealTimeReflectionHeightFogPS::FPermutationDomain PsPermutationVector;
+							PsPermutationVector.Set<FRenderRealTimeReflectionHeightFogPS::FDepthTexture>(CubeDepthTexture != nullptr);
+							TShaderMapRef<FRenderRealTimeReflectionHeightFogPS> PixelShader(GetGlobalShaderMap(SkyRC.FeatureLevel), PsPermutationVector);
+
+							FRenderRealTimeReflectionHeightFogPS::FParameters* PsPassParameters = GraphBuilder.AllocParameters<FRenderRealTimeReflectionHeightFogPS::FParameters>();
+							PsPassParameters->ViewUniformBuffer = CubeViewUniformBuffer;
+							PsPassParameters->RenderTargets = SkyRC.RenderTargets;
+							PsPassParameters->DepthTexture = CubeDepthTexture != nullptr ? CubeDepthTexture : BlackDummy2dTex;
+							PsPassParameters->FogStruct = CreateFogUniformBuffer(GraphBuilder, CubeView);
+
+							ClearUnusedGraphResources(PixelShader, PsPassParameters);
+
+							// Render height fog at an infinite distance since real time reflections does not have a depth buffer for now.
+							// Volumetric fog is not supported in such reflections.
+							GraphBuilder.AddPass(
+								RDG_EVENT_NAME("DistantHeightFog"),
+								PsPassParameters,
+								ERDGPassFlags::Raster,
+								[PsPassParameters, VertexShader, PixelShader, CubeWidth](FRHICommandList& RHICmdListLambda)
+								{
+									RHICmdListLambda.SetViewport(0.0f, 0.0f, 0.0f, CubeWidth, CubeWidth, 1.0f);
+
+									FGraphicsPipelineStateInitializer GraphicsPSOInit;
+									RHICmdListLambda.ApplyCachedRenderTargets(GraphicsPSOInit);
+
+									GraphicsPSOInit.BlendState = TStaticBlendState<CW_RGB, BO_Add, BF_One, BF_SourceAlpha, BO_Add, BF_Zero, BF_One>::GetRHI();
+									GraphicsPSOInit.DepthStencilState = TStaticDepthStencilState<false, CF_Always>::GetRHI();
+									GraphicsPSOInit.RasterizerState = TStaticRasterizerState<FM_Solid, CM_None>::GetRHI();
+									GraphicsPSOInit.BoundShaderState.VertexDeclarationRHI = GEmptyVertexDeclaration.VertexDeclarationRHI;
+									GraphicsPSOInit.BoundShaderState.VertexShaderRHI = VertexShader.GetVertexShader();
+									GraphicsPSOInit.BoundShaderState.PixelShaderRHI = PixelShader.GetPixelShader();
+									GraphicsPSOInit.PrimitiveType = PT_TriangleList;
+									SetGraphicsPipelineState(RHICmdListLambda, GraphicsPSOInit);
+
+									FRenderRealTimeReflectionHeightFogVS::FParameters VsPassParameters;
+									VsPassParameters.ViewUniformBuffer = PsPassParameters->ViewUniformBuffer;
+									SetShaderParameters(RHICmdListLambda, VertexShader, VertexShader.GetVertexShader(), VsPassParameters);
+									SetShaderParameters(RHICmdListLambda, PixelShader, PixelShader.GetPixelShader(), *PsPassParameters);
+
+									RHICmdListLambda.DrawPrimitive(0, 1, 1);
+								});
+						}
 					}
 					else
 					{
-						FSceneTextureShaderParameters SceneTextures = CreateSceneTextureShaderParameters(GraphBuilder, SceneRenderer.FeatureLevel, ESceneTextureSetupMode::SceneDepth);
-						SceneRenderer.RenderSkyAtmosphereInternal(GraphBuilder, SceneTextures, SkyRC);
-					}
-
-					// Also render the height fog as part of the sky render pass when time slicing is enabled.
-					if (Scene && Scene->ExponentialFogs.Num() > 0)
-					{
-						FRenderRealTimeReflectionHeightFogVS::FPermutationDomain VsPermutationVector;
-						TShaderMapRef<FRenderRealTimeReflectionHeightFogVS> VertexShader(GetGlobalShaderMap(SkyRC.FeatureLevel), VsPermutationVector);
-
-						FRenderRealTimeReflectionHeightFogPS::FPermutationDomain PsPermutationVector;
-						PsPermutationVector.Set<FRenderRealTimeReflectionHeightFogPS::FDepthTexture>(CubeDepthTexture != nullptr);
-						TShaderMapRef<FRenderRealTimeReflectionHeightFogPS> PixelShader(GetGlobalShaderMap(SkyRC.FeatureLevel), PsPermutationVector);
-
-						FRenderRealTimeReflectionHeightFogPS::FParameters* PsPassParameters = GraphBuilder.AllocParameters<FRenderRealTimeReflectionHeightFogPS::FParameters>();
-						PsPassParameters->ViewUniformBuffer = CubeViewUniformBuffer;
-						PsPassParameters->RenderTargets = SkyRC.RenderTargets;
-						PsPassParameters->DepthTexture = CubeDepthTexture != nullptr ? CubeDepthTexture : BlackDummy2dTex;
-						PsPassParameters->FogStruct = CreateFogUniformBuffer(GraphBuilder, CubeView);
-
-						ClearUnusedGraphResources(PixelShader, PsPassParameters);
-
-						// Render height fog at an infinite distance since real time reflections does not have a depth buffer for now.
-						// Volumetric fog is not supported in such reflections.
-						GraphBuilder.AddPass(
-							RDG_EVENT_NAME("DistantHeightFog"),
-							PsPassParameters,
-							ERDGPassFlags::Raster,
-							[PsPassParameters, VertexShader, PixelShader, CubeWidth](FRHICommandList& RHICmdListLambda)
-						{
-							RHICmdListLambda.SetViewport(0.0f, 0.0f, 0.0f, CubeWidth, CubeWidth, 1.0f);
-
-							FGraphicsPipelineStateInitializer GraphicsPSOInit;
-							RHICmdListLambda.ApplyCachedRenderTargets(GraphicsPSOInit);
-
-							GraphicsPSOInit.BlendState = TStaticBlendState<CW_RGB, BO_Add, BF_One, BF_SourceAlpha, BO_Add, BF_Zero, BF_One>::GetRHI();
-							GraphicsPSOInit.DepthStencilState = TStaticDepthStencilState<false, CF_Always>::GetRHI();
-							GraphicsPSOInit.RasterizerState = TStaticRasterizerState<FM_Solid, CM_None>::GetRHI();
-							GraphicsPSOInit.BoundShaderState.VertexDeclarationRHI = GEmptyVertexDeclaration.VertexDeclarationRHI;
-							GraphicsPSOInit.BoundShaderState.VertexShaderRHI = VertexShader.GetVertexShader();
-							GraphicsPSOInit.BoundShaderState.PixelShaderRHI = PixelShader.GetPixelShader();
-							GraphicsPSOInit.PrimitiveType = PT_TriangleList;
-							SetGraphicsPipelineState(RHICmdListLambda, GraphicsPSOInit);
-
-							FRenderRealTimeReflectionHeightFogVS::FParameters VsPassParameters;
-							VsPassParameters.ViewUniformBuffer = PsPassParameters->ViewUniformBuffer;
-							SetShaderParameters(RHICmdListLambda, VertexShader, VertexShader.GetVertexShader(), VsPassParameters);
-							SetShaderParameters(RHICmdListLambda, PixelShader, PixelShader.GetPixelShader(), *PsPassParameters);
-
-							RHICmdListLambda.DrawPrimitive(0, 1, 1);
-						});
+						ClearCubeFace(SkyCubeTexture, CubeFace);
 					}
 				}
 
@@ -669,21 +698,9 @@ void FScene::AllocateAndCaptureFrameSkyEnvMap(
 		}
 		else
 		{
-
 			for (int32 CubeFace = 0; CubeFace < CubeFace_MAX; CubeFace++)
 			{
-				FRenderTargetParameters* Parameters = GraphBuilder.AllocParameters<FRenderTargetParameters>();
-				Parameters->RenderTargets[0] = FRenderTargetBinding(SkyCubeTexture, ERenderTargetLoadAction::ENoAction, 0, CubeFace);
-
-				FLinearColor ClearColor = FLinearColor::Black;
-				GraphBuilder.AddPass(
-					RDG_EVENT_NAME("ClearSkyRenderTarget"),
-					Parameters,
-					ERDGPassFlags::Raster,
-					[Parameters, ClearColor](FRHICommandList& RHICmdList)
-					{
-						DrawClearQuad(RHICmdList, ClearColor);
-					});
+				ClearCubeFace(SkyCubeTexture, CubeFace);
 			}
 		}
 	};
