Added a HUD warning in editor when a real time sky light capture is in the scene but nothing can be captured.

diff --git a/Engine/Source/Runtime/Renderer/Private/SceneRendering.cpp b/Engine/Source/Runtime/Renderer/Private/SceneRendering.cpp
--- a/Engine/Source/Runtime/Renderer/Private/SceneRendering.cpp
+++ b/Engine/Source/Runtime/Renderer/Private/SceneRendering.cpp
@@ -2874,6 +2874,7 @@ void FSceneRenderer::RenderFinish(FRDGBuilder& GraphBuilder, FRDGTextureRef View
 
 		const bool bStationarySkylight = Scene->SkyLight && Scene->SkyLight->bWantsStaticShadowing;
 		const bool bShowSkylightWarning = bStationarySkylight && !ReadOnlyCVARCache.bEnableStationarySkylight;
+		const bool bRealTimeSkyCaptureButNothingToCapture = Scene->SkyLight && Scene->SkyLight->bRealTimeCaptureEnabled && (!Scene->HasSkyAtmosphere() && !Scene->HasVolumetricCloud() && (Views.Num() > 0 && !Views[0].bSceneHasSkyMaterial));
 
 		const bool bShowPointLightWarning = UsedWholeScenePointLightNames.Num() > 0 && !ReadOnlyCVARCache.bEnablePointLightShadows;
 		const bool bShowShadowedLightOverflowWarning = Scene->OverflowingDynamicShadowedLights.Num() > 0;
@@ -2904,7 +2905,7 @@ void FSceneRenderer::RenderFinish(FRDGBuilder& GraphBuilder, FRDGTextureRef View
 		const bool bAnyWarning = bShowPrecomputedVisibilityWarning || bShowGlobalClipPlaneWarning || bShowAtmosphericFogWarning || bShowSkylightWarning || bShowPointLightWarning 
 			|| bShowDFAODisabledWarning || bShowShadowedLightOverflowWarning || bShowMobileDynamicCSMWarning || bShowMobileLowQualityLightmapWarning || bShowMobileMovableDirectionalLightWarning
 			|| bMobileShowVertexFogWarning || bShowSkinCacheOOM || bSingleLayerWaterWarning || bShowDFDisabledWarning || bShowNoSkyAtmosphereComponentWarning || bFxDebugDraw 
-			|| bShowSkyAtmosphereFogComponentsConflicts || bLumenEnabledButNoDistanceFieldsWarning;
+			|| bShowSkyAtmosphereFogComponentsConflicts || bLumenEnabledButNoDistanceFieldsWarning || bRealTimeSkyCaptureButNothingToCapture;
 
 		for(int32 ViewIndex = 0;ViewIndex < Views.Num();ViewIndex++)
 		{	
@@ -2926,7 +2927,8 @@ void FSceneRenderer::RenderFinish(FRDGBuilder& GraphBuilder, FRDGTextureRef View
 						bLocked, bShowPrecomputedVisibilityWarning, bShowGlobalClipPlaneWarning, bShowDFAODisabledWarning, bShowDFDisabledWarning,
 						bShowAtmosphericFogWarning, bViewParentOrFrozen, bShowSkylightWarning, bShowPointLightWarning, bShowShadowedLightOverflowWarning,
 						bShowMobileLowQualityLightmapWarning, bShowMobileMovableDirectionalLightWarning, bShowMobileDynamicCSMWarning, bMobileShowVertexFogWarning,
-						bShowSkinCacheOOM, bSingleLayerWaterWarning, bShowNoSkyAtmosphereComponentWarning, bFxDebugDraw, FXInterface, bShowSkyAtmosphereFogComponentsConflicts, bLumenEnabledButNoDistanceFieldsWarning]
+						bShowSkinCacheOOM, bSingleLayerWaterWarning, bShowNoSkyAtmosphereComponentWarning, bFxDebugDraw, FXInterface, bShowSkyAtmosphereFogComponentsConflicts, 
+						bLumenEnabledButNoDistanceFieldsWarning, bRealTimeSkyCaptureButNothingToCapture]
 						(FCanvas& Canvas)
 					{
 						// so it can get the screen size
@@ -2992,6 +2994,12 @@ void FSceneRenderer::RenderFinish(FRDGBuilder& GraphBuilder, FRDGTextureRef View
 							Canvas.DrawShadowedText(10, Y, Message, GetStatsFont(), FLinearColor(1.0, 0.05, 0.05, 1.0));
 							Y += 14;
 						}
+						if (bRealTimeSkyCaptureButNothingToCapture)
+						{
+							static const FText Message = NSLOCTEXT("Renderer", "Skylight", "A sky light with real-time capture enable is in the scene. It requires at least a SkyAtmosphere component, A volumetricCloud component or a mesh with a material tagged as IsSky. Otherwise it will be black.");
+							Canvas.DrawShadowedText(10, Y, Message, GetStatsFont(), FLinearColor(1.0, 0.05, 0.05, 1.0));
+							Y += 14;
+						}
 						if (bShowPointLightWarning)
 						{
 							static const FText Message = NSLOCTEXT("Renderer", "PointLight", "PROJECT DOES NOT SUPPORT WHOLE SCENE POINT LIGHT SHADOWS: ");
