More removal of view uniform buffer from mesh draw commands. Refactored pass uniform buffers in hair to use RDG.

diff --git a/Engine/Shaders/Public/ShaderVersion.ush b/Engine/Shaders/Public/ShaderVersion.ush
--- a/Engine/Shaders/Public/ShaderVersion.ush
+++ b/Engine/Shaders/Public/ShaderVersion.ush
@@ -3,4 +3,4 @@
 // in Platform.ush (which should be included in any shader) it allows to invalidate the shader DDC.
 // 
 // If you are merging streams and there is a conflict with this GUID you should make a new GUID rather than taking one or the other.
-// GUID = 5952C866-E7D6-4116-ABCC-48ACA6A8D216
+// GUID = 5952C866-E7D6-4116-ABDC-48ACA6A8D216
diff --git a/Engine/Source/Developer/MaterialUtilities/Private/MeshRendering.cpp b/Engine/Source/Developer/MaterialUtilities/Private/MeshRendering.cpp
--- a/Engine/Source/Developer/MaterialUtilities/Private/MeshRendering.cpp
+++ b/Engine/Source/Developer/MaterialUtilities/Private/MeshRendering.cpp
@@ -597,7 +597,7 @@ public:
 		RenderScope.EnqueueRenderCommand(
 			[Parameters](FRHICommandListImmediate& RHICmdList)
 		{
-			FMeshPassProcessorRenderState DrawRenderState(*Parameters.View);
+			FMeshPassProcessorRenderState DrawRenderState;
 
 			// disable depth test & writes
 			DrawRenderState.SetDepthStencilState(TStaticDepthStencilState<false, CF_Always>::GetRHI());
diff --git a/Engine/Source/Editor/UnrealEd/Private/Lightmass/LightmassLandscapeRender.cpp b/Engine/Source/Editor/UnrealEd/Private/Lightmass/LightmassLandscapeRender.cpp
--- a/Engine/Source/Editor/UnrealEd/Private/Lightmass/LightmassLandscapeRender.cpp
+++ b/Engine/Source/Editor/UnrealEd/Private/Lightmass/LightmassLandscapeRender.cpp
@@ -134,7 +134,7 @@ void RenderLandscapeMaterialForLightmass(const FLandscapeStaticLightingMesh* Lan
 
 					FSceneView View(ViewInitOptions);
 
-					FMeshPassProcessorRenderState DrawRenderState(View);
+					FMeshPassProcessorRenderState DrawRenderState;
 
 					// disable depth test & writes
 					DrawRenderState.SetDepthStencilState(TStaticDepthStencilState<false, CF_Always>::GetRHI());
diff --git a/Engine/Source/Runtime/Engine/Private/TileRendering.cpp b/Engine/Source/Runtime/Engine/Private/TileRendering.cpp
--- a/Engine/Source/Runtime/Engine/Private/TileRendering.cpp
+++ b/Engine/Source/Runtime/Engine/Private/TileRendering.cpp
@@ -300,7 +300,7 @@ bool FCanvasTileRendererItem::Render_GameThread(const FCanvas* Canvas, FRenderTh
 			[LocalData = Data, View, bIsHitTesting, bNeedsToSwitchVerticalAxis, bRequiresExplicit128bitRT]
 		(FRHICommandListImmediate& RHICmdList)
 		{
-			FMeshPassProcessorRenderState DrawRenderState(*View);
+			FMeshPassProcessorRenderState DrawRenderState;
 
 			// disable depth test & writes
 			DrawRenderState.SetDepthStencilState(TStaticDepthStencilState<false, CF_Always>::GetRHI());
diff --git a/Engine/Source/Runtime/Engine/Private/TriangleRendering.cpp b/Engine/Source/Runtime/Engine/Private/TriangleRendering.cpp
--- a/Engine/Source/Runtime/Engine/Private/TriangleRendering.cpp
+++ b/Engine/Source/Runtime/Engine/Private/TriangleRendering.cpp
@@ -254,7 +254,7 @@ bool FCanvasTriangleRendererItem::Render_GameThread(const FCanvas* Canvas, FRend
 	RenderScope.EnqueueRenderCommand(
 		[LocalData = Data, View, bIsHitTesting, bNeedsToSwitchVerticalAxis](FRHICommandListImmediate& RHICmdList)
 	{
-		FMeshPassProcessorRenderState DrawRenderState(*View);
+		FMeshPassProcessorRenderState DrawRenderState;
 
 		// disable depth test & writes
 		DrawRenderState.SetDepthStencilState(TStaticDepthStencilState<false, CF_Always>::GetRHI());
diff --git a/Engine/Source/Runtime/Engine/Private/UserInterface/Canvas.cpp b/Engine/Source/Runtime/Engine/Private/UserInterface/Canvas.cpp
--- a/Engine/Source/Runtime/Engine/Private/UserInterface/Canvas.cpp
+++ b/Engine/Source/Runtime/Engine/Private/UserInterface/Canvas.cpp
@@ -479,7 +479,7 @@ bool FCanvasBatchedElementRenderItem::Render_GameThread(const FCanvas* Canvas, F
 		{
 			FSceneView SceneView = FBatchedElements::CreateProxySceneView(DrawParameters.RenderData->Transform.GetMatrix(), FIntRect(0, 0, DrawParameters.ViewportSizeX, DrawParameters.ViewportSizeY));
 
-			FMeshPassProcessorRenderState DrawRenderState(SceneView);
+			FMeshPassProcessorRenderState DrawRenderState;
 
 			// disable depth test & writes
 			DrawRenderState.SetDepthStencilState(TStaticDepthStencilState<false, CF_Always>::GetRHI());
diff --git a/Engine/Source/Runtime/Renderer/Private/HairStrands/HairStrandsDeepShadow.cpp b/Engine/Source/Runtime/Renderer/Private/HairStrands/HairStrandsDeepShadow.cpp
--- a/Engine/Source/Runtime/Renderer/Private/HairStrands/HairStrandsDeepShadow.cpp
+++ b/Engine/Source/Runtime/Renderer/Private/HairStrands/HairStrandsDeepShadow.cpp
@@ -434,13 +434,21 @@ void RenderHairStrandsDeepShadows(
 					RDG_GPU_STAT_SCOPE(GraphBuilder, HairStrandsDeepShadowFrontDepth);
 
 					FHairDeepShadowRasterPassParameters* PassParameters = GraphBuilder.AllocParameters<FHairDeepShadowRasterPassParameters>();
-					PassParameters->CPU_WorldToClipMatrix = DomData.CPU_WorldToLightTransform;;
-					PassParameters->SliceValue = FVector4(1, 1, 1, 1);
-					PassParameters->AtlasRect = DomData.AtlasRect;
-					PassParameters->AtlasSlotIndex = DomData.AtlasSlotIndex;
-					PassParameters->LayerDepths = LayerDepths;
-					PassParameters->ViewportResolution = AtlasSlotResolution;
-					PassParameters->DeepShadowViewInfoBuffer = DeepShadowViewInfoBufferSRV;
+
+					{
+						FHairDeepShadowRasterUniformParameters* UniformParameters = GraphBuilder.AllocParameters<FHairDeepShadowRasterUniformParameters>();
+
+						UniformParameters->CPU_WorldToClipMatrix = DomData.CPU_WorldToLightTransform;;
+						UniformParameters->SliceValue = FVector4(1, 1, 1, 1);
+						UniformParameters->AtlasRect = DomData.AtlasRect;
+						UniformParameters->AtlasSlotIndex = DomData.AtlasSlotIndex;
+						UniformParameters->LayerDepths = LayerDepths;
+						UniformParameters->ViewportResolution = AtlasSlotResolution;
+						UniformParameters->DeepShadowViewInfoBuffer = DeepShadowViewInfoBufferSRV;
+
+						PassParameters->UniformBuffer = GraphBuilder.CreateUniformBuffer(UniformParameters);
+					}
+
 					PassParameters->RenderTargets.DepthStencil = FDepthStencilBinding(FrontDepthAtlasTexture, bClear ? ERenderTargetLoadAction::EClear : ERenderTargetLoadAction::ELoad, ERenderTargetLoadAction::ENoAction, FExclusiveDepthStencil::DepthWrite_StencilNop);
 
 					AddHairDeepShadowRasterPass(
@@ -464,14 +472,22 @@ void RenderHairStrandsDeepShadows(
 					RDG_GPU_STAT_SCOPE(GraphBuilder, HairStrandsDeepShadowLayers);
 
 					FHairDeepShadowRasterPassParameters* PassParameters = GraphBuilder.AllocParameters<FHairDeepShadowRasterPassParameters>();
-					PassParameters->CPU_WorldToClipMatrix = DomData.CPU_WorldToLightTransform;;
-					PassParameters->SliceValue = FVector4(1, 1, 1, 1);
-					PassParameters->AtlasRect = DomData.AtlasRect;
-					PassParameters->AtlasSlotIndex = DomData.AtlasSlotIndex;
-					PassParameters->LayerDepths = LayerDepths;
-					PassParameters->ViewportResolution = AtlasSlotResolution;
-					PassParameters->FrontDepthTexture = FrontDepthAtlasTexture;
-					PassParameters->DeepShadowViewInfoBuffer = DeepShadowViewInfoBufferSRV;
+
+					{
+						FHairDeepShadowRasterUniformParameters* UniformParameters = GraphBuilder.AllocParameters<FHairDeepShadowRasterUniformParameters>();
+
+						UniformParameters->CPU_WorldToClipMatrix = DomData.CPU_WorldToLightTransform;;
+						UniformParameters->SliceValue = FVector4(1, 1, 1, 1);
+						UniformParameters->AtlasRect = DomData.AtlasRect;
+						UniformParameters->AtlasSlotIndex = DomData.AtlasSlotIndex;
+						UniformParameters->LayerDepths = LayerDepths;
+						UniformParameters->ViewportResolution = AtlasSlotResolution;
+						UniformParameters->FrontDepthTexture = FrontDepthAtlasTexture;
+						UniformParameters->DeepShadowViewInfoBuffer = DeepShadowViewInfoBufferSRV;
+
+						PassParameters->UniformBuffer = GraphBuilder.CreateUniformBuffer(UniformParameters);
+					}
+
 					PassParameters->RenderTargets[0] = FRenderTargetBinding(DeepShadowLayersAtlasTexture, bClear ? ERenderTargetLoadAction::EClear : ERenderTargetLoadAction::ELoad, 0);
 
 					AddHairDeepShadowRasterPass(
diff --git a/Engine/Source/Runtime/Renderer/Private/HairStrands/HairStrandsRasterCommon.cpp b/Engine/Source/Runtime/Renderer/Private/HairStrands/HairStrandsRasterCommon.cpp
--- a/Engine/Source/Runtime/Renderer/Private/HairStrands/HairStrandsRasterCommon.cpp
+++ b/Engine/Source/Runtime/Renderer/Private/HairStrands/HairStrandsRasterCommon.cpp
@@ -11,63 +11,8 @@
 #include "MeshPassProcessor.inl"
 #include "ScenePrivate.h"
 
-/////////////////////////////////////////////////////////////////////////////////////////
-// Deep shadow global parameters
-BEGIN_GLOBAL_SHADER_PARAMETER_STRUCT(FHairDeepShadowRasterGlobalParameters, )
-	SHADER_PARAMETER(FMatrix, CPU_WorldToClipMatrix)
-	SHADER_PARAMETER(FVector4, SliceValue)
-	SHADER_PARAMETER(FIntRect, AtlasRect)
-	SHADER_PARAMETER(FIntPoint, ViewportResolution)
-	SHADER_PARAMETER(uint32, AtlasSlotIndex)
-	SHADER_PARAMETER(FVector4, LayerDepths)
-	SHADER_PARAMETER_TEXTURE(Texture2D<float>, FrontDepthTexture)
-	SHADER_PARAMETER_SRV(StructuredBuffer<FDeepShadowViewInfo>, DeepShadowViewInfoBuffer)
-END_GLOBAL_SHADER_PARAMETER_STRUCT()
-IMPLEMENT_GLOBAL_SHADER_PARAMETER_STRUCT(FHairDeepShadowRasterGlobalParameters, "DeepRasterPass");
-
-static FHairDeepShadowRasterGlobalParameters ConvertToGlobalPassParameter(const FHairDeepShadowRasterPassParameters* In)
-{
-	FHairDeepShadowRasterGlobalParameters Out;
-	Out.CPU_WorldToClipMatrix	= In->CPU_WorldToClipMatrix;
-	Out.SliceValue				= In->SliceValue;
-	Out.AtlasRect				= In->AtlasRect;
-	Out.ViewportResolution		= In->ViewportResolution;
-	Out.AtlasSlotIndex			= In->AtlasSlotIndex;
-	Out.LayerDepths				= In->LayerDepths;
-	Out.FrontDepthTexture		= In->FrontDepthTexture ? In->FrontDepthTexture->GetRHI() : (FRHITexture*)(GSystemTextures.DepthDummy->GetRenderTargetItem().ShaderResourceTexture);
-	Out.DeepShadowViewInfoBuffer= In->DeepShadowViewInfoBuffer->GetRHI();
-	return Out;
-}
-
-/////////////////////////////////////////////////////////////////////////////////////////
-// Voxelization global parameters
-BEGIN_GLOBAL_SHADER_PARAMETER_STRUCT(FHairVoxelizationRasterGlobalParameters, )
-	SHADER_PARAMETER_STRUCT(FVirtualVoxelCommonParameters, VirtualVoxel)
-	SHADER_PARAMETER(FMatrix, WorldToClipMatrix)
-	SHADER_PARAMETER(FVector, VoxelMinAABB)
-	SHADER_PARAMETER(FVector, VoxelMaxAABB)
-	SHADER_PARAMETER(FIntVector, VoxelResolution)
-	SHADER_PARAMETER(uint32, MacroGroupId)
-	SHADER_PARAMETER(FIntPoint, ViewportResolution)
-	SHADER_PARAMETER_SRV(StructuredBuffer<FVoxelizationViewInfo>, VoxelizationViewInfoBuffer)
-	SHADER_PARAMETER_UAV(RWTexture3D<uint>, DensityTexture)
-END_GLOBAL_SHADER_PARAMETER_STRUCT()
-IMPLEMENT_GLOBAL_SHADER_PARAMETER_STRUCT(FHairVoxelizationRasterGlobalParameters, "VoxelRasterPass");
-
-static FHairVoxelizationRasterGlobalParameters ConvertToGlobalPassParameter(const FHairVoxelizationRasterPassParameters* In)
-{
-	FHairVoxelizationRasterGlobalParameters Out;
-	Out.VirtualVoxel	   = In->VirtualVoxel;
-	Out.WorldToClipMatrix  = In->WorldToClipMatrix;
-	Out.VoxelMinAABB	   = In->VoxelMinAABB;
-	Out.VoxelMaxAABB	   = In->VoxelMaxAABB;
-	Out.VoxelResolution	   = In->VoxelResolution;
-	Out.MacroGroupId	   = In->MacroGroupId;
-	Out.ViewportResolution = In->ViewportResolution;
-	Out.VoxelizationViewInfoBuffer = In->VoxelizationViewInfoBuffer->GetRHI();
-	Out.DensityTexture  = In->DensityTexture->GetRHI();
-	return Out;
-}
+IMPLEMENT_STATIC_UNIFORM_BUFFER_STRUCT(FHairDeepShadowRasterUniformParameters, "DeepRasterPass", SceneTextures);
+IMPLEMENT_STATIC_UNIFORM_BUFFER_STRUCT(FHairVoxelizationRasterUniformParameters, "VoxelRasterPass", SceneTextures);
 
 /////////////////////////////////////////////////////////////////////////////////////////
 
@@ -82,8 +27,6 @@ protected:
 	{
 		ERHIFeatureLevel::Type FeatureLevel = GetMaxSupportedFeatureLevel((EShaderPlatform)Initializer.Target.Platform);
 		check(FSceneInterface::GetShadingPath(FeatureLevel) != EShadingPath::Mobile);
-		// deferred
-		PassUniformBuffer.Bind(Initializer.ParameterMap, FHairDeepShadowRasterGlobalParameters::StaticStructMetadata.GetShaderVariableName());
 	}
 
 	FDeepShadowDepthMeshVS() {}
@@ -116,7 +59,6 @@ protected:
 	{
 		ERHIFeatureLevel::Type FeatureLevel = GetMaxSupportedFeatureLevel((EShaderPlatform)Initializer.Target.Platform);
 		check(FSceneInterface::GetShadingPath(FeatureLevel) != EShadingPath::Mobile);
-		PassUniformBuffer.Bind(Initializer.ParameterMap, FHairDeepShadowRasterGlobalParameters::StaticStructMetadata.GetShaderVariableName());
 	}
 
 	FDeepShadowDomMeshVS() {}
@@ -150,7 +92,6 @@ protected:
 	{
 		ERHIFeatureLevel::Type FeatureLevel = GetMaxSupportedFeatureLevel((EShaderPlatform)Initializer.Target.Platform);
 		check(FSceneInterface::GetShadingPath(FeatureLevel) != EShadingPath::Mobile);
-		PassUniformBuffer.Bind(Initializer.ParameterMap, FHairVoxelizationRasterGlobalParameters::StaticStructMetadata.GetShaderVariableName());
 	}
 
 	FVoxelMeshVS() {}
@@ -190,7 +131,6 @@ public:
 	{
 		ERHIFeatureLevel::Type FeatureLevel = GetMaxSupportedFeatureLevel((EShaderPlatform)Initializer.Target.Platform);
 		check(FSceneInterface::GetShadingPath(FeatureLevel) != EShadingPath::Mobile);
-		PassUniformBuffer.Bind(Initializer.ParameterMap, FHairDeepShadowRasterGlobalParameters::StaticStructMetadata.GetShaderVariableName());
 	}
 
 	FDeepShadowDepthMeshPS() {}
@@ -215,7 +155,6 @@ public:
 	{
 		ERHIFeatureLevel::Type FeatureLevel = GetMaxSupportedFeatureLevel((EShaderPlatform)Initializer.Target.Platform);
 		check(FSceneInterface::GetShadingPath(FeatureLevel) != EShadingPath::Mobile);
-		PassUniformBuffer.Bind(Initializer.ParameterMap, FHairDeepShadowRasterGlobalParameters::StaticStructMetadata.GetShaderVariableName());
 	}
 
 	FDeepShadowDomMeshPS() {}
@@ -246,7 +185,6 @@ public:
 	{
 		ERHIFeatureLevel::Type FeatureLevel = GetMaxSupportedFeatureLevel((EShaderPlatform)Initializer.Target.Platform);
 		check(FSceneInterface::GetShadingPath(FeatureLevel) != EShadingPath::Mobile);
-		PassUniformBuffer.Bind(Initializer.ParameterMap, FHairVoxelizationRasterGlobalParameters::StaticStructMetadata.GetShaderVariableName());
 	}
 
 	FVoxelMeshPS() {}
@@ -397,7 +335,7 @@ FHairRasterMeshProcessor::FHairRasterMeshProcessor(
 
 /////////////////////////////////////////////////////////////////////////////////////////
 
-template<typename TPassParameter, typename TGlobalParameter>
+template<typename TPassParameter>
 void AddHairStrandsRasterPass(
 	FRDGBuilder& GraphBuilder,
 	const FScene* Scene,
@@ -421,34 +359,26 @@ void AddHairStrandsRasterPass(
 		}
 	};
 
-	GraphBuilder.AddPass(
-		GetPassName(RasterPassType),
-		PassParameters,
-		ERDGPassFlags::Raster,
-		[PassParameters, Scene = Scene, ViewInfo, RasterPassType, &PrimitiveSceneInfos, ViewportRect, HairRenderInfo, HairRenderInfoBits, RasterDirection](FRHICommandListImmediate& RHICmdList)
 	{
-		check(RHICmdList.IsInsideRenderPass());
-		check(IsInRenderingThread());
-
-		check(RHICmdList.IsInsideRenderPass());
-		check(IsInRenderingThread());
-
-		SCOPE_CYCLE_COUNTER(STAT_RenderPerObjectShadowDepthsTime);
-
 		TUniformBufferRef<FViewUniformShaderParameters> ViewUniformShaderParameters;
 		ViewInfo->CachedViewUniformShaderParameters->HairRenderInfo = HairRenderInfo;
 		ViewInfo->CachedViewUniformShaderParameters->HairRenderInfoBits = HairRenderInfoBits;
 
-		const FVector SavedViewForward = ViewInfo->CachedViewUniformShaderParameters->ViewForward;	
+		const FVector SavedViewForward = ViewInfo->CachedViewUniformShaderParameters->ViewForward;
 		ViewInfo->CachedViewUniformShaderParameters->ViewForward = RasterDirection;
-		ViewUniformShaderParameters = TUniformBufferRef<FViewUniformShaderParameters>::CreateUniformBufferImmediate(*ViewInfo->CachedViewUniformShaderParameters, UniformBuffer_SingleFrame);
+		PassParameters->View = TUniformBufferRef<FViewUniformShaderParameters>::CreateUniformBufferImmediate(*ViewInfo->CachedViewUniformShaderParameters, UniformBuffer_SingleFrame);
 		ViewInfo->CachedViewUniformShaderParameters->ViewForward = SavedViewForward;
+	}
 
-		TGlobalParameter GlobalPassParameters = ConvertToGlobalPassParameter(PassParameters);
-		TUniformBufferRef<TGlobalParameter> GlobalPassParametersBuffer = TUniformBufferRef<TGlobalParameter>::CreateUniformBufferImmediate(GlobalPassParameters, UniformBuffer_SingleFrame);
+	GraphBuilder.AddPass(
+		GetPassName(RasterPassType),
+		PassParameters,
+		ERDGPassFlags::Raster,
+		[PassParameters, Scene = Scene, ViewInfo, RasterPassType, &PrimitiveSceneInfos, ViewportRect, HairRenderInfo, HairRenderInfoBits, RasterDirection](FRHICommandListImmediate& RHICmdList)
+	{
+		SCOPE_CYCLE_COUNTER(STAT_RenderPerObjectShadowDepthsTime);
 
-		FMeshPassProcessorRenderState DrawRenderState(*ViewInfo, GlobalPassParametersBuffer);
-		DrawRenderState.SetViewUniformBuffer(ViewUniformShaderParameters);
+		FMeshPassProcessorRenderState DrawRenderState;
 
 		RHICmdList.SetViewport(ViewportRect.Min.X, ViewportRect.Min.Y, 0.0f, ViewportRect.Max.X, ViewportRect.Max.Y, 1.0f);
 
@@ -512,7 +442,7 @@ void AddHairDeepShadowRasterPass(
 {
 	check(PassType == EHairStrandsRasterPassType::FrontDepth || PassType == EHairStrandsRasterPassType::DeepOpacityMap);
 
-	AddHairStrandsRasterPass<FHairDeepShadowRasterPassParameters, FHairDeepShadowRasterGlobalParameters>(
+	AddHairStrandsRasterPass<FHairDeepShadowRasterPassParameters>(
 		GraphBuilder, 
 		Scene, 
 		ViewInfo, 
@@ -536,7 +466,7 @@ void AddHairVoxelizationRasterPass(
 	const FVector& RasterDirection,
 	FHairVoxelizationRasterPassParameters* PassParameters)
 {
-	AddHairStrandsRasterPass<FHairVoxelizationRasterPassParameters, FHairVoxelizationRasterGlobalParameters>(
+	AddHairStrandsRasterPass<FHairVoxelizationRasterPassParameters>(
 		GraphBuilder, 
 		Scene, 
 		ViewInfo, 
diff --git a/Engine/Source/Runtime/Renderer/Private/HairStrands/HairStrandsRasterCommon.h b/Engine/Source/Runtime/Renderer/Private/HairStrands/HairStrandsRasterCommon.h
--- a/Engine/Source/Runtime/Renderer/Private/HairStrands/HairStrandsRasterCommon.h
+++ b/Engine/Source/Runtime/Renderer/Private/HairStrands/HairStrandsRasterCommon.h
@@ -23,7 +23,7 @@ enum class EHairStrandsRasterPassType : uint8
 // ////////////////////////////////////////////////////////////////
 // Deep shadow raster pass
 
-BEGIN_SHADER_PARAMETER_STRUCT(FHairDeepShadowRasterPassParameters, )
+BEGIN_GLOBAL_SHADER_PARAMETER_STRUCT(FHairDeepShadowRasterUniformParameters, )
 	SHADER_PARAMETER(FMatrix, CPU_WorldToClipMatrix)
 	SHADER_PARAMETER(FVector4, SliceValue)
 	SHADER_PARAMETER(FIntRect, AtlasRect)
@@ -32,7 +32,11 @@ BEGIN_SHADER_PARAMETER_STRUCT(FHairDeepShadowRasterPassParameters, )
 	SHADER_PARAMETER(FVector4, LayerDepths)
 	SHADER_PARAMETER_RDG_TEXTURE(Texture2D<float>, FrontDepthTexture)
 	SHADER_PARAMETER_RDG_BUFFER_SRV(StructuredBuffer<FDeepShadowViewInfo>, DeepShadowViewInfoBuffer)
+END_GLOBAL_SHADER_PARAMETER_STRUCT()
 
+BEGIN_SHADER_PARAMETER_STRUCT(FHairDeepShadowRasterPassParameters, )
+	SHADER_PARAMETER_STRUCT_REF(FViewUniformShaderParameters, View)
+	SHADER_PARAMETER_RDG_UNIFORM_BUFFER(FHairDeepShadowRasterUniformParameters, UniformBuffer)
 	RENDER_TARGET_BINDING_SLOTS()
 END_SHADER_PARAMETER_STRUCT()
 
@@ -50,7 +54,7 @@ void AddHairDeepShadowRasterPass(
 
 // ////////////////////////////////////////////////////////////////
 // Voxelization raster pass
-BEGIN_SHADER_PARAMETER_STRUCT(FHairVoxelizationRasterPassParameters, )
+BEGIN_GLOBAL_SHADER_PARAMETER_STRUCT(FHairVoxelizationRasterUniformParameters, )
 	SHADER_PARAMETER_STRUCT(FVirtualVoxelCommonParameters, VirtualVoxel)
 	SHADER_PARAMETER(FMatrix, WorldToClipMatrix)
 	SHADER_PARAMETER(FVector, VoxelMinAABB)
@@ -60,7 +64,11 @@ BEGIN_SHADER_PARAMETER_STRUCT(FHairVoxelizationRasterPassParameters, )
 	SHADER_PARAMETER(FIntPoint, ViewportResolution)
 	SHADER_PARAMETER_RDG_BUFFER_SRV(StructuredBuffer<FVoxelizationViewInfo>, VoxelizationViewInfoBuffer)
 	SHADER_PARAMETER_RDG_TEXTURE_UAV(RWTexture3D<uint>, DensityTexture)
+END_GLOBAL_SHADER_PARAMETER_STRUCT()
 
+BEGIN_SHADER_PARAMETER_STRUCT(FHairVoxelizationRasterPassParameters, )
+	SHADER_PARAMETER_STRUCT_REF(FViewUniformShaderParameters, View)
+	SHADER_PARAMETER_RDG_UNIFORM_BUFFER(FHairVoxelizationRasterUniformParameters, UniformBuffer)
 	RENDER_TARGET_BINDING_SLOTS()
 END_SHADER_PARAMETER_STRUCT()
 
diff --git a/Engine/Source/Runtime/Renderer/Private/HairStrands/HairStrandsVisibility.cpp b/Engine/Source/Runtime/Renderer/Private/HairStrands/HairStrandsVisibility.cpp
--- a/Engine/Source/Runtime/Renderer/Private/HairStrands/HairStrandsVisibility.cpp
+++ b/Engine/Source/Runtime/Renderer/Private/HairStrands/HairStrandsVisibility.cpp
@@ -379,17 +379,16 @@ static FRDGTextureRef AddClearLightSamplePass(
 /////////////////////////////////////////////////////////////////////////////////////////
 // Material pass which write value directly into the GBuffer. This is a fast (low-quality) 
 // pass, which dither raster result into GBuffer, and works only with 1spp
-BEGIN_GLOBAL_SHADER_PARAMETER_STRUCT(FMaterialGBufferPassParameters, )
+BEGIN_GLOBAL_SHADER_PARAMETER_STRUCT(FMaterialGBufferPassUniformParameters, )
 	SHADER_PARAMETER(FIntPoint, MaxResolution)
 	SHADER_PARAMETER(uint32, InputType)
 	SHADER_PARAMETER(float, CoverageThreshold)
-	SHADER_PARAMETER_TEXTURE(Texture2D<float>, InTransmittanceTexture)
-	SHADER_PARAMETER_TEXTURE(Texture2D<uint>,  InRasterOutputVisibilityTexture)
-	SHADER_PARAMETER_TEXTURE(Texture2D<uint>,  InMSAAIDTexture)
-	SHADER_PARAMETER_TEXTURE(Texture2D<float>, InMSAADepthTexture)
+	SHADER_PARAMETER_RDG_TEXTURE(Texture2D<float>, InTransmittanceTexture)
+	SHADER_PARAMETER_RDG_TEXTURE(Texture2D<uint>,  InRasterOutputVisibilityTexture)
+	SHADER_PARAMETER_RDG_TEXTURE(Texture2D<uint>,  InMSAAIDTexture)
+	SHADER_PARAMETER_RDG_TEXTURE(Texture2D<float>, InMSAADepthTexture)
 END_GLOBAL_SHADER_PARAMETER_STRUCT()
-
-IMPLEMENT_GLOBAL_SHADER_PARAMETER_STRUCT(FMaterialGBufferPassParameters, "MaterialGBufferPassParameters");
+IMPLEMENT_STATIC_UNIFORM_BUFFER_STRUCT(FMaterialGBufferPassUniformParameters, "MaterialGBufferPassParameters", SceneTextures);
 
 class FHairMaterialGBufferVS : public FMeshMaterialShader
 {
@@ -401,7 +400,6 @@ protected:
 	{
 		ERHIFeatureLevel::Type FeatureLevel = GetMaxSupportedFeatureLevel((EShaderPlatform)Initializer.Target.Platform);
 		check(FSceneInterface::GetShadingPath(FeatureLevel) != EShadingPath::Mobile);
-		PassUniformBuffer.Bind(Initializer.ParameterMap, FMaterialGBufferPassParameters::StaticStructMetadata.GetShaderVariableName());
 	}
 
 	FHairMaterialGBufferVS() {}
@@ -443,7 +441,6 @@ public:
 	{
 		ERHIFeatureLevel::Type FeatureLevel = GetMaxSupportedFeatureLevel((EShaderPlatform)Initializer.Target.Platform);
 		check(FSceneInterface::GetShadingPath(FeatureLevel) != EShadingPath::Mobile);
-		PassUniformBuffer.Bind(Initializer.ParameterMap, FMaterialGBufferPassParameters::StaticStructMetadata.GetShaderVariableName());
 		MaterialPass_MacroGroupId.Bind(Initializer.ParameterMap, TEXT("MaterialPass_MacroGroupId"));
 		MaterialPass_MaterialId.Bind(Initializer.ParameterMap, TEXT("MaterialPass_MaterialId"));
 		MaterialPass_PrimitiveId.Bind(Initializer.ParameterMap, TEXT("MaterialPass_PrimitiveId"));
@@ -491,10 +488,8 @@ IMPLEMENT_MATERIAL_SHADER_TYPE(, FHairMaterialGBufferPS, TEXT("/Engine/Private/H
 /////////////////////////////////////////////////////////////////////////////////////////
 
 BEGIN_SHADER_PARAMETER_STRUCT(FVisibilityMaterialGBufferPassParameters, )
-	SHADER_PARAMETER_RDG_TEXTURE(Texture2D<uint2>, InRasterOutputVisibilityTexture)
-	SHADER_PARAMETER_RDG_TEXTURE(Texture2D<uint>,  InMSAAIDTexture)
-	SHADER_PARAMETER_RDG_TEXTURE(Texture2D<float>, InMSAADepthTexture)
-	SHADER_PARAMETER_RDG_TEXTURE(Texture2D<float>, InTransmittanceTexture)
+	SHADER_PARAMETER_STRUCT_REF(FViewUniformShaderParameters, View)
+	SHADER_PARAMETER_RDG_UNIFORM_BUFFER(FMaterialGBufferPassUniformParameters, UniformBuffer)
 	RENDER_TARGET_BINDING_SLOTS()
 END_SHADER_PARAMETER_STRUCT()
 
@@ -648,15 +643,44 @@ static void AddHairMaterialGBufferPass(
 	FRDGTextureRef OutDepthTexture,
 	FRDGTextureRef OutVelocityTexture)
 {
+	const FRDGSystemTextures& SystemTextures = FRDGSystemTextures::Get(GraphBuilder);
 
 	// Add resources reference to the pass parameters, in order to get the resource lifetime extended to this pass
 	FVisibilityMaterialGBufferPassParameters* PassParameters = GraphBuilder.AllocParameters<FVisibilityMaterialGBufferPassParameters>();
-	PassParameters->InTransmittanceTexture = InTransmittanceTexture;
-	PassParameters->InRasterOutputVisibilityTexture = InRasterOutputVisibilityTexture;
-	PassParameters->InMSAAIDTexture = InMSAAIDTexture;
-	PassParameters->InMSAADepthTexture = InMSAADepthTexture;
+
 	const FIntPoint Resolution = OutBufferATexture->Desc.Extent;
 
+	{
+		FMaterialGBufferPassUniformParameters* UniformParameters = GraphBuilder.AllocParameters<FMaterialGBufferPassUniformParameters>();
+		UniformParameters->InputType = InRasterOutputVisibilityTexture != nullptr ? 1 : 0;
+		UniformParameters->MaxResolution = Resolution;
+		UniformParameters->CoverageThreshold = FMath::Clamp(GHairStrandsFullCoverageThreshold, 0.1f, 1.f);
+		UniformParameters->InTransmittanceTexture = InTransmittanceTexture;
+		UniformParameters->InMSAAIDTexture = SystemTextures.Black;
+		UniformParameters->InMSAADepthTexture = SystemTextures.Black;
+		UniformParameters->InRasterOutputVisibilityTexture = SystemTextures.Black;
+		if (UniformParameters->InputType == 0)
+		{
+			check(InMSAAIDTexture);
+			check(InMSAADepthTexture);
+			UniformParameters->InMSAAIDTexture = InMSAAIDTexture;
+			UniformParameters->InMSAADepthTexture = InMSAADepthTexture;
+		}
+		else if (UniformParameters->InputType == 1)
+		{
+			check(InRasterOutputVisibilityTexture);
+			UniformParameters->InRasterOutputVisibilityTexture = InRasterOutputVisibilityTexture;
+		}
+
+		PassParameters->UniformBuffer = GraphBuilder.CreateUniformBuffer(UniformParameters);
+	}
+
+	{
+		const bool bEnableMSAA = false;
+		SetUpViewHairRenderInfo(*ViewInfo, bEnableMSAA, ViewInfo->CachedViewUniformShaderParameters->HairRenderInfo, ViewInfo->CachedViewUniformShaderParameters->HairRenderInfoBits, ViewInfo->CachedViewUniformShaderParameters->HairComponents);
+		PassParameters->View = TUniformBufferRef<FViewUniformShaderParameters>::CreateUniformBufferImmediate(*ViewInfo->CachedViewUniformShaderParameters, UniformBuffer_SingleFrame);
+	}
+
 	// If there is velocity texture, we recreate a dummy one
 	bool bIsVelocityDummy = false;
 	if (OutVelocityTexture == nullptr)
@@ -684,44 +708,7 @@ static void AddHairMaterialGBufferPass(
 		ERDGPassFlags::Raster,
 		[PassParameters, Scene = Scene, ViewInfo, &MacroGroupDatas, Resolution](FRHICommandListImmediate& RHICmdList)
 		{
-			check(RHICmdList.IsInsideRenderPass());
-			check(IsInRenderingThread());
-
-			FMaterialGBufferPassParameters MaterialPassParameters;
-			MaterialPassParameters.InputType = PassParameters->InRasterOutputVisibilityTexture != nullptr ? 1 : 0;
-			MaterialPassParameters.MaxResolution = Resolution;
-			MaterialPassParameters.CoverageThreshold = FMath::Clamp(GHairStrandsFullCoverageThreshold, 0.1f, 1.f);
-			MaterialPassParameters.InTransmittanceTexture = PassParameters->InTransmittanceTexture->GetPooledRenderTarget()->GetRenderTargetItem().ShaderResourceTexture;
-			
-			FTextureRHIRef DefaultTexture = GSystemTextures.BlackDummy->GetRenderTargetItem().ShaderResourceTexture;
-			MaterialPassParameters.InMSAAIDTexture = DefaultTexture;
-			MaterialPassParameters.InMSAADepthTexture = DefaultTexture;
-			MaterialPassParameters.InRasterOutputVisibilityTexture = DefaultTexture;
-			if (MaterialPassParameters.InputType == 0)
-			{
-				check(PassParameters->InMSAAIDTexture);
-				check(PassParameters->InMSAADepthTexture);
-				MaterialPassParameters.InMSAAIDTexture = PassParameters->InMSAAIDTexture->GetPooledRenderTarget()->GetRenderTargetItem().ShaderResourceTexture;
-				MaterialPassParameters.InMSAADepthTexture = PassParameters->InMSAADepthTexture->GetPooledRenderTarget()->GetRenderTargetItem().ShaderResourceTexture;;
-			}
-			else if (MaterialPassParameters.InputType == 1)
-			{
-				check(PassParameters->InRasterOutputVisibilityTexture);
-				MaterialPassParameters.InRasterOutputVisibilityTexture = PassParameters->InRasterOutputVisibilityTexture->GetPooledRenderTarget()->GetRenderTargetItem().ShaderResourceTexture;
-			}
-
-			TUniformBufferRef<FMaterialGBufferPassParameters> MaterialPassParametersBuffer = TUniformBufferRef<FMaterialGBufferPassParameters>::CreateUniformBufferImmediate(MaterialPassParameters, UniformBuffer_SingleFrame);
-
-			FMeshPassProcessorRenderState DrawRenderState(*ViewInfo, MaterialPassParametersBuffer);
-			// Note: this reference needs to persistent until SubmitMeshDrawCommands() is called, as DrawRenderState does not ref count 
-			// the view uniform buffer (raw pointer). It is only within the MeshProcessor that the uniform buffer get reference
-			TUniformBufferRef<FViewUniformShaderParameters> ViewUniformShaderParameters;
-			{
-				const bool bEnableMSAA = false;
-				SetUpViewHairRenderInfo(*ViewInfo, bEnableMSAA, ViewInfo->CachedViewUniformShaderParameters->HairRenderInfo, ViewInfo->CachedViewUniformShaderParameters->HairRenderInfoBits, ViewInfo->CachedViewUniformShaderParameters->HairComponents);
-				ViewUniformShaderParameters = TUniformBufferRef<FViewUniformShaderParameters>::CreateUniformBufferImmediate(*ViewInfo->CachedViewUniformShaderParameters, UniformBuffer_SingleFrame);
-				DrawRenderState.SetViewUniformBuffer(ViewUniformShaderParameters);
-			}
+			FMeshPassProcessorRenderState DrawRenderState;
 
 			{
 				RHICmdList.SetViewport(0, 0, 0.0f, Resolution.X, Resolution.Y, 1.0f);
@@ -758,22 +745,6 @@ static void AddHairMaterialGBufferPass(
 
 /////////////////////////////////////////////////////////////////////////////////////////
 
-BEGIN_GLOBAL_SHADER_PARAMETER_STRUCT(FMaterialPassParameters, )
-	SHADER_PARAMETER(FIntPoint, MaxResolution)
-	SHADER_PARAMETER(uint32, MaxSampleCount)
-	SHADER_PARAMETER(uint32, NodeGroupSize)
-	SHADER_PARAMETER(uint32, bUpdateSampleCoverage)
-	SHADER_PARAMETER(uint32, bOutputEmissive)
-	SHADER_PARAMETER_TEXTURE(Texture2D<uint>, NodeIndex)
-	SHADER_PARAMETER_TEXTURE(Texture2D<uint>, TotalNodeCounter)
-	SHADER_PARAMETER_SRV(StructuredBuffer<uint>, NodeCoord)
-	SHADER_PARAMETER_SRV(StructuredBuffer<FNodeVis>, NodeVis)
-	SHADER_PARAMETER_SRV(Buffer<uint>, IndirectArgs)
-	SHADER_PARAMETER_UAV(RWStructuredBuffer<FPackedHairSample>, OutNodeData)
-	SHADER_PARAMETER_UAV(RWBuffer<float2>, OutNodeVelocity)
-END_GLOBAL_SHADER_PARAMETER_STRUCT()
-IMPLEMENT_GLOBAL_SHADER_PARAMETER_STRUCT(FMaterialPassParameters, "MaterialPassParameters");
-
 class FHairMaterialVS : public FMeshMaterialShader
 {
 	DECLARE_SHADER_TYPE(FHairMaterialVS, MeshMaterial);
@@ -784,7 +755,6 @@ protected:
 	{
 		ERHIFeatureLevel::Type FeatureLevel = GetMaxSupportedFeatureLevel((EShaderPlatform)Initializer.Target.Platform);
 		check(FSceneInterface::GetShadingPath(FeatureLevel) != EShadingPath::Mobile);
-		PassUniformBuffer.Bind(Initializer.ParameterMap, FMaterialPassParameters::StaticStructMetadata.GetShaderVariableName());
 	}
 
 	FHairMaterialVS() {}
@@ -836,7 +806,6 @@ public:
 	{
 		ERHIFeatureLevel::Type FeatureLevel = GetMaxSupportedFeatureLevel((EShaderPlatform)Initializer.Target.Platform);
 		check(FSceneInterface::GetShadingPath(FeatureLevel) != EShadingPath::Mobile);
-		PassUniformBuffer.Bind(Initializer.ParameterMap, FMaterialPassParameters::StaticStructMetadata.GetShaderVariableName());
 		MaterialPass_MacroGroupId.Bind(Initializer.ParameterMap, TEXT("MaterialPass_MacroGroupId"));
 		MaterialPass_MaterialId.Bind(Initializer.ParameterMap, TEXT("MaterialPass_MaterialId"));
 		MaterialPass_PrimitiveId.Bind(Initializer.ParameterMap, TEXT("MaterialPass_PrimitiveId"));
@@ -1013,7 +982,12 @@ FHairMaterialProcessor::FHairMaterialProcessor(
 
 /////////////////////////////////////////////////////////////////////////////////////////
 
-BEGIN_SHADER_PARAMETER_STRUCT(FVisibilityMaterialPassParameters, )
+BEGIN_GLOBAL_SHADER_PARAMETER_STRUCT(FVisibilityMaterialPassUniformParameters, )
+	SHADER_PARAMETER(FIntPoint, MaxResolution)
+	SHADER_PARAMETER(uint32, MaxSampleCount)
+	SHADER_PARAMETER(uint32, NodeGroupSize)
+	SHADER_PARAMETER(uint32, bUpdateSampleCoverage)
+	SHADER_PARAMETER(uint32, bOutputEmissive)
 	SHADER_PARAMETER_RDG_TEXTURE(Texture2D<uint>, NodeIndex)
 	SHADER_PARAMETER_RDG_TEXTURE(Texture2D<uint>, TotalNodeCounter)
 	SHADER_PARAMETER_RDG_BUFFER_SRV(StructuredBuffer<uint>, NodeCoord)
@@ -1021,6 +995,12 @@ BEGIN_SHADER_PARAMETER_STRUCT(FVisibilityMaterialPassParameters, )
 	SHADER_PARAMETER_RDG_BUFFER_SRV(Buffer<uint>, IndirectArgs)
 	SHADER_PARAMETER_RDG_BUFFER_UAV(RWStructuredBuffer<FPackedHairSample>, OutNodeData)
 	SHADER_PARAMETER_RDG_BUFFER_UAV(RWBuffer<float2>, OutNodeVelocity)
+END_GLOBAL_SHADER_PARAMETER_STRUCT()
+IMPLEMENT_STATIC_UNIFORM_BUFFER_STRUCT(FVisibilityMaterialPassUniformParameters, "MaterialPassParameters", SceneTextures);
+
+BEGIN_SHADER_PARAMETER_STRUCT(FVisibilityMaterialPassParameters, )
+	SHADER_PARAMETER_STRUCT_REF(FViewUniformShaderParameters, View)
+	SHADER_PARAMETER_RDG_UNIFORM_BUFFER(FVisibilityMaterialPassUniformParameters, UniformBuffer)
 	RENDER_TARGET_BINDING_SLOTS()
 END_SHADER_PARAMETER_STRUCT()
 
@@ -1109,20 +1089,38 @@ static FMaterialPassOutput AddHairMaterialPass(
 	const uint32 ResolutionDim = FMath::CeilToInt(FMath::Sqrt(MaxNodeCount));
 	const FIntPoint Resolution(ResolutionDim, ResolutionDim);
 
+	const bool bOutputEmissive = IsHairStrandsEmissiveEnable();
+	const bool bIsPlatformRequireRenderTarget = IsPlatformRequiringRenderTargetForMaterialPass(Scene->GetShaderPlatform()) || GRHIRequiresRenderTargetForPixelShaderUAVs;
 
 	// Add resources reference to the pass parameters, in order to get the resource lifetime extended to this pass
 	FVisibilityMaterialPassParameters* PassParameters = GraphBuilder.AllocParameters<FVisibilityMaterialPassParameters>();
-	PassParameters->TotalNodeCounter= CompactNodeCounter;
-	PassParameters->NodeIndex		= CompactNodeIndex;
-	PassParameters->NodeVis			= GraphBuilder.CreateSRV(CompactNodeVis);
-	PassParameters->NodeCoord		= GraphBuilder.CreateSRV(CompactNodeCoord);
-	PassParameters->IndirectArgs	= GraphBuilder.CreateSRV(IndirectArgBuffer);
-	PassParameters->OutNodeData		= GraphBuilder.CreateUAV(FRDGBufferUAVDesc(Output.NodeData));
-	PassParameters->OutNodeVelocity	= GraphBuilder.CreateUAV(FRDGBufferUAVDesc(Output.NodeVelocity, FMaterialPassOutput::VelocityFormat));
+
+	{
+		FVisibilityMaterialPassUniformParameters* UniformParameters = GraphBuilder.AllocParameters<FVisibilityMaterialPassUniformParameters>();
+
+		UniformParameters->bUpdateSampleCoverage = bUpdateSampleCoverage ? 1 : 0;
+		UniformParameters->bOutputEmissive = bOutputEmissive ? 1 : 0;
+		UniformParameters->MaxResolution = Resolution;
+		UniformParameters->NodeGroupSize = NodeGroupSize;
+		UniformParameters->MaxSampleCount = MaxNodeCount;
+		UniformParameters->TotalNodeCounter = CompactNodeCounter;
+		UniformParameters->NodeIndex = CompactNodeIndex;
+		UniformParameters->NodeVis = GraphBuilder.CreateSRV(CompactNodeVis);
+		UniformParameters->NodeCoord = GraphBuilder.CreateSRV(CompactNodeCoord);
+		UniformParameters->IndirectArgs = GraphBuilder.CreateSRV(IndirectArgBuffer);
+		UniformParameters->OutNodeData = GraphBuilder.CreateUAV(FRDGBufferUAVDesc(Output.NodeData));
+		UniformParameters->OutNodeVelocity = GraphBuilder.CreateUAV(FRDGBufferUAVDesc(Output.NodeVelocity, FMaterialPassOutput::VelocityFormat));
+
+		PassParameters->UniformBuffer = GraphBuilder.CreateUniformBuffer(UniformParameters);
+	}
+
+	{
+		const bool bEnableMSAA = false;
+		SetUpViewHairRenderInfo(*ViewInfo, bEnableMSAA, ViewInfo->CachedViewUniformShaderParameters->HairRenderInfo, ViewInfo->CachedViewUniformShaderParameters->HairRenderInfoBits, ViewInfo->CachedViewUniformShaderParameters->HairComponents);
+		PassParameters->View = TUniformBufferRef<FViewUniformShaderParameters>::CreateUniformBufferImmediate(*ViewInfo->CachedViewUniformShaderParameters, UniformBuffer_SingleFrame);
+	}
 
 	// For debug purpose only
-	const bool bOutputEmissive = IsHairStrandsEmissiveEnable();
-	const bool bIsPlatformRequireRenderTarget = IsPlatformRequiringRenderTargetForMaterialPass(Scene->GetShaderPlatform()) || GRHIRequiresRenderTargetForPixelShaderUAVs;
 	if (bOutputEmissive)
 	{
 		FRDGTextureDesc OutputDesc = FRDGTextureDesc::Create2D(Resolution, PF_FloatRGBA, FClearValueBinding::Black, TexCreate_RenderTargetable);
@@ -1144,35 +1142,7 @@ static FMaterialPassOutput AddHairMaterialPass(
 		ERDGPassFlags::Raster,
 		[PassParameters, Scene = Scene, ViewInfo, &MacroGroupDatas, MaxNodeCount, Resolution, NodeGroupSize, bUpdateSampleCoverage, bOutputEmissive](FRHICommandListImmediate& RHICmdList)
 	{
-		check(RHICmdList.IsInsideRenderPass());
-		check(IsInRenderingThread());
-
-		FMaterialPassParameters MaterialPassParameters;
-		MaterialPassParameters.bUpdateSampleCoverage = bUpdateSampleCoverage ? 1 : 0;
-		MaterialPassParameters.bOutputEmissive	= bOutputEmissive ? 1 : 0;
-		MaterialPassParameters.MaxResolution	= Resolution;
-		MaterialPassParameters.NodeGroupSize	= NodeGroupSize;
-		MaterialPassParameters.MaxSampleCount	= MaxNodeCount;
-		MaterialPassParameters.TotalNodeCounter	= PassParameters->TotalNodeCounter->GetPooledRenderTarget()->GetRenderTargetItem().ShaderResourceTexture;
-		MaterialPassParameters.NodeIndex		= PassParameters->NodeIndex->GetPooledRenderTarget()->GetRenderTargetItem().ShaderResourceTexture;
-		MaterialPassParameters.NodeCoord		= PassParameters->NodeCoord->GetRHI();
-		MaterialPassParameters.NodeVis			= PassParameters->NodeVis->GetRHI();
-		MaterialPassParameters.IndirectArgs		= PassParameters->IndirectArgs->GetRHI();
-		MaterialPassParameters.OutNodeData 		= PassParameters->OutNodeData->GetRHI();
-		MaterialPassParameters.OutNodeVelocity 	= PassParameters->OutNodeVelocity->GetRHI();
-
-		TUniformBufferRef<FMaterialPassParameters> MaterialPassParametersBuffer = TUniformBufferRef<FMaterialPassParameters>::CreateUniformBufferImmediate(MaterialPassParameters, UniformBuffer_SingleFrame);
-
-		FMeshPassProcessorRenderState DrawRenderState(*ViewInfo, MaterialPassParametersBuffer);
-		// Note: this reference needs to persistent until SubmitMeshDrawCommands() is called, as DrawRenderState does not ref count 
-		// the view uniform buffer (raw pointer). It is only within the MeshProcessor that the uniform buffer get reference
-		TUniformBufferRef<FViewUniformShaderParameters> ViewUniformShaderParameters;
-		{
-			const bool bEnableMSAA = false;
-			SetUpViewHairRenderInfo(*ViewInfo, bEnableMSAA, ViewInfo->CachedViewUniformShaderParameters->HairRenderInfo, ViewInfo->CachedViewUniformShaderParameters->HairRenderInfoBits, ViewInfo->CachedViewUniformShaderParameters->HairComponents);
-			ViewUniformShaderParameters = TUniformBufferRef<FViewUniformShaderParameters>::CreateUniformBufferImmediate(*ViewInfo->CachedViewUniformShaderParameters, UniformBuffer_SingleFrame);
-			DrawRenderState.SetViewUniformBuffer(ViewUniformShaderParameters);
-		}
+		FMeshPassProcessorRenderState DrawRenderState;
 
 		{
 			RHICmdList.SetViewport(0, 0, 0.0f, Resolution.X, Resolution.Y, 1.0f);
@@ -1361,32 +1331,20 @@ static FRDGTextureRef AddHairLightChannelMaskPass(
 }
 
 /////////////////////////////////////////////////////////////////////////////////////////
-BEGIN_GLOBAL_SHADER_PARAMETER_STRUCT(FVisibilityPassGlobalParameters, )
+BEGIN_GLOBAL_SHADER_PARAMETER_STRUCT(FVisibilityPassUniformParameters, )
 	SHADER_PARAMETER(uint32, MaxPPLLNodeCount)
-	SHADER_PARAMETER_UAV(RWTexture2D<uint>, PPLLCounter)
-	SHADER_PARAMETER_UAV(RWTexture2D<uint>, PPLLNodeIndex)
-	SHADER_PARAMETER_UAV(RWStructuredBuffer<FPPLLNodeData>, PPLLNodeData)
+	SHADER_PARAMETER_RDG_TEXTURE_UAV(RWTexture2D<uint>, PPLLCounter)
+	SHADER_PARAMETER_RDG_TEXTURE_UAV(RWTexture2D<uint>, PPLLNodeIndex)
+	SHADER_PARAMETER_RDG_BUFFER_UAV(RWStructuredBuffer<FPPLLNodeData>, PPLLNodeData)
 END_GLOBAL_SHADER_PARAMETER_STRUCT()
-IMPLEMENT_GLOBAL_SHADER_PARAMETER_STRUCT(FVisibilityPassGlobalParameters, "HairVisibilityPass");
+IMPLEMENT_STATIC_UNIFORM_BUFFER_STRUCT(FVisibilityPassUniformParameters, "HairVisibilityPass", SceneTextures);
 
 BEGIN_SHADER_PARAMETER_STRUCT(FVisibilityPassParameters, )
-	SHADER_PARAMETER(uint32, MaxPPLLNodeCount)
-	SHADER_PARAMETER_RDG_TEXTURE_UAV(RWTexture2D<uint>, PPLLCounter)
-	SHADER_PARAMETER_RDG_TEXTURE_UAV(RWTexture2D<uint>, PPLLNodeIndex)
-	SHADER_PARAMETER_RDG_BUFFER_UAV(RWStructuredBuffer<FPPLLNodeData, PPLLNodeData)
+	SHADER_PARAMETER_STRUCT_REF(FViewUniformShaderParameters, View)
+	SHADER_PARAMETER_RDG_UNIFORM_BUFFER(FVisibilityPassUniformParameters, UniformBuffer)
 	RENDER_TARGET_BINDING_SLOTS()
 END_SHADER_PARAMETER_STRUCT()
 
-static FVisibilityPassGlobalParameters ConvertToGlobalPassParameter(const FVisibilityPassParameters* In)
-{
-	FVisibilityPassGlobalParameters Out;
-	Out.MaxPPLLNodeCount = In->MaxPPLLNodeCount;
-	Out.PPLLCounter = In->PPLLCounter->GetRHI();
-	Out.PPLLNodeIndex = In->PPLLNodeIndex->GetRHI();
-	Out.PPLLNodeData = In->PPLLNodeData->GetRHI();
-	return Out;
-}
-
 // Example: 28bytes * 8spp = 224bytes per pixel = 442Mb @ 1080p
 struct PPLLNodeData
 {
@@ -1399,12 +1357,16 @@ struct PPLLNodeData
 	uint32 PackedVelocity;
 };
 
-void CreatePassDummyTextures(FRDGBuilder& GraphBuilder, FVisibilityPassParameters* PassParameters)
+TRDGUniformBufferRef<FVisibilityPassUniformParameters> CreatePassDummyTextures(FRDGBuilder& GraphBuilder)
 {
+	FVisibilityPassUniformParameters* UniformParameters = GraphBuilder.AllocParameters<FVisibilityPassUniformParameters>();
+
 	FRDGTextureDesc Desc = FRDGTextureDesc::Create2D(FIntPoint(1,1), PF_R32_UINT, FClearValueBinding::None, TexCreate_UAV | TexCreate_ShaderResource);
-	PassParameters->PPLLCounter		= GraphBuilder.CreateUAV(GraphBuilder.CreateTexture(Desc, TEXT("HairVisibilityPPLLNodeIndex")));
-	PassParameters->PPLLNodeIndex	= GraphBuilder.CreateUAV(GraphBuilder.CreateTexture(Desc, TEXT("HairVisibilityPPLLNodeIndex")));
-	PassParameters->PPLLNodeData	= GraphBuilder.CreateUAV(GraphBuilder.CreateBuffer(FRDGBufferDesc::CreateStructuredDesc(sizeof(PPLLNodeData), 1), TEXT("DummyPPLLNodeData")));
+	UniformParameters->PPLLCounter		= GraphBuilder.CreateUAV(GraphBuilder.CreateTexture(Desc, TEXT("HairVisibilityPPLLNodeIndex")));
+	UniformParameters->PPLLNodeIndex	= GraphBuilder.CreateUAV(GraphBuilder.CreateTexture(Desc, TEXT("HairVisibilityPPLLNodeIndex")));
+	UniformParameters->PPLLNodeData		= GraphBuilder.CreateUAV(GraphBuilder.CreateBuffer(FRDGBufferDesc::CreateStructuredDesc(sizeof(PPLLNodeData), 1), TEXT("DummyPPLLNodeData")));
+
+	return GraphBuilder.CreateUniformBuffer(UniformParameters);
 }
 
 template<EHairVisibilityRenderMode RenderMode, bool bCullingEnable>
@@ -1419,7 +1381,6 @@ protected:
 	{
 		ERHIFeatureLevel::Type FeatureLevel = GetMaxSupportedFeatureLevel((EShaderPlatform)Initializer.Target.Platform);
 		check(FSceneInterface::GetShadingPath(FeatureLevel) != EShadingPath::Mobile);
-		PassUniformBuffer.Bind(Initializer.ParameterMap, FVisibilityPassGlobalParameters::StaticStructMetadata.GetShaderVariableName());
 	}
 
 	FHairVisibilityVS() {}
@@ -1475,7 +1436,6 @@ public:
 	{
 		ERHIFeatureLevel::Type FeatureLevel = GetMaxSupportedFeatureLevel((EShaderPlatform)Initializer.Target.Platform);
 		check(FSceneInterface::GetShadingPath(FeatureLevel) != EShadingPath::Mobile);
-		PassUniformBuffer.Bind(Initializer.ParameterMap, FVisibilityPassGlobalParameters::StaticStructMetadata.GetShaderVariableName());
 		HairVisibilityPass_HairMacroGroupIndex.Bind(Initializer.ParameterMap, TEXT("HairVisibilityPass_HairMacroGroupIndex"));
 		HairVisibilityPass_HairMaterialId.Bind(Initializer.ParameterMap, TEXT("HairVisibilityPass_HairMaterialId"));
 		HairVisibilityPass_LightChannelMask.Bind(Initializer.ParameterMap, TEXT("HairVisibilityPass_LightChannelMask"));
@@ -2414,6 +2374,22 @@ static void AddHairVisibilityCommonPass(
 
 	AddHairCulledVertexResourcesTransitionPass(GraphBuilder, MacroGroupDatas);
 
+	// Note: this reference needs to persistent until SubmitMeshDrawCommands() is called, as DrawRenderState does not ref count 
+	// the view uniform buffer (raw pointer). It is only within the MeshProcessor that the uniform buffer get reference
+	TUniformBufferRef<FViewUniformShaderParameters> ViewUniformShaderParameters;
+	if (RenderMode == HairVisibilityRenderMode_Transmittance || RenderMode == HairVisibilityRenderMode_TransmittanceAndHairCount || RenderMode == HairVisibilityRenderMode_PPLL)
+	{
+		const bool bEnableMSAA = false;
+		SetUpViewHairRenderInfo(*ViewInfo, bEnableMSAA, ViewInfo->CachedViewUniformShaderParameters->HairRenderInfo, ViewInfo->CachedViewUniformShaderParameters->HairRenderInfoBits, ViewInfo->CachedViewUniformShaderParameters->HairComponents);
+
+		// Create and set the uniform buffer
+		PassParameters->View = TUniformBufferRef<FViewUniformShaderParameters>::CreateUniformBufferImmediate(*ViewInfo->CachedViewUniformShaderParameters, UniformBuffer_SingleFrame);
+	}
+	else
+	{
+		PassParameters->View = ViewInfo->ViewUniformBuffer;
+	}
+
 	GraphBuilder.AddPass(
 		GetPassName(),
 		PassParameters,
@@ -2423,22 +2399,7 @@ static void AddHairVisibilityCommonPass(
 		check(RHICmdList.IsInsideRenderPass());
 		check(IsInRenderingThread());
 
-		FVisibilityPassGlobalParameters GlobalPassParameters = ConvertToGlobalPassParameter(PassParameters);
-		TUniformBufferRef<FVisibilityPassGlobalParameters> GlobalPassParametersBuffer = TUniformBufferRef<FVisibilityPassGlobalParameters>::CreateUniformBufferImmediate(GlobalPassParameters, UniformBuffer_SingleFrame);
-
-		FMeshPassProcessorRenderState DrawRenderState(*ViewInfo, GlobalPassParametersBuffer);
-
-		// Note: this reference needs to persistent until SubmitMeshDrawCommands() is called, as DrawRenderState does not ref count 
-		// the view uniform buffer (raw pointer). It is only within the MeshProcessor that the uniform buffer get reference
-		TUniformBufferRef<FViewUniformShaderParameters> ViewUniformShaderParameters;
-		if(RenderMode == HairVisibilityRenderMode_Transmittance || RenderMode == HairVisibilityRenderMode_TransmittanceAndHairCount || RenderMode == HairVisibilityRenderMode_PPLL)
-		{
-			const bool bEnableMSAA = false;
-			SetUpViewHairRenderInfo(*ViewInfo, bEnableMSAA, ViewInfo->CachedViewUniformShaderParameters->HairRenderInfo, ViewInfo->CachedViewUniformShaderParameters->HairRenderInfoBits, ViewInfo->CachedViewUniformShaderParameters->HairComponents);
-			// Create and set the uniform buffer
-			ViewUniformShaderParameters = TUniformBufferRef<FViewUniformShaderParameters>::CreateUniformBufferImmediate(*ViewInfo->CachedViewUniformShaderParameters, UniformBuffer_SingleFrame);
-			DrawRenderState.SetViewUniformBuffer(ViewUniformShaderParameters);
-		}
+		FMeshPassProcessorRenderState DrawRenderState;
 
 		{
 			RHICmdList.SetViewport(0, 0, 0.0f, ViewInfo->ViewRect.Width(), ViewInfo->ViewRect.Height(), 1.0f);
@@ -2527,7 +2488,7 @@ static void AddHairVisibilityMSAAPass(
 		AddClearGraphicPass(GraphBuilder, RDG_EVENT_NAME("HairStrandsClearVisibilityMSAAIdTexture"), ViewInfo, 0xFFFFFFFF, OutVisibilityIdTexture);
 
 		FVisibilityPassParameters* PassParameters = GraphBuilder.AllocParameters<FVisibilityPassParameters>();
-		CreatePassDummyTextures(GraphBuilder, PassParameters);
+		PassParameters->UniformBuffer = CreatePassDummyTextures(GraphBuilder);
 		PassParameters->RenderTargets[0] = FRenderTargetBinding(OutVisibilityIdTexture, ERenderTargetLoadAction::ELoad, 0);
 		PassParameters->RenderTargets.DepthStencil = FDepthStencilBinding(
 			OutVisibilityDepthTexture,
@@ -2572,7 +2533,7 @@ static void AddHairVisibilityMSAAPass(
 		}
 
 		FVisibilityPassParameters* PassParameters = GraphBuilder.AllocParameters<FVisibilityPassParameters>();
-		CreatePassDummyTextures(GraphBuilder, PassParameters);
+		PassParameters->UniformBuffer = CreatePassDummyTextures(GraphBuilder);
 		PassParameters->RenderTargets[0] = FRenderTargetBinding(OutVisibilityIdTexture, ERenderTargetLoadAction::ELoad, 0);
 		PassParameters->RenderTargets[1] = FRenderTargetBinding(OutVisibilityMaterialTexture, LoadAction, 0);
 		PassParameters->RenderTargets[2] = FRenderTargetBinding(OutVisibilityAttributeTexture, LoadAction, 0);
@@ -2618,10 +2579,18 @@ static void AddHairVisibilityPPLLPass(
 	AddClearUAVPass(GraphBuilder, GraphBuilder.CreateUAV(OutVisibilityPPLLNodeIndex), ClearValueInvalid);
 
 	FVisibilityPassParameters* PassParameters = GraphBuilder.AllocParameters<FVisibilityPassParameters>();
-	PassParameters->PPLLCounter = GraphBuilder.CreateUAV(FRDGTextureUAVDesc(OutVisibilityPPLLNodeCounter, 0));
-	PassParameters->PPLLNodeIndex = GraphBuilder.CreateUAV(FRDGTextureUAVDesc(OutVisibilityPPLLNodeIndex, 0));
-	PassParameters->PPLLNodeData = GraphBuilder.CreateUAV(FRDGBufferUAVDesc(OutVisibilityPPLLNodeData));
-	PassParameters->MaxPPLLNodeCount = PPLLMaxTotalListElementCount;
+
+	{
+		FVisibilityPassUniformParameters* UniformParameters = GraphBuilder.AllocParameters<FVisibilityPassUniformParameters>();
+
+		UniformParameters->PPLLCounter = GraphBuilder.CreateUAV(FRDGTextureUAVDesc(OutVisibilityPPLLNodeCounter, 0));
+		UniformParameters->PPLLNodeIndex = GraphBuilder.CreateUAV(FRDGTextureUAVDesc(OutVisibilityPPLLNodeIndex, 0));
+		UniformParameters->PPLLNodeData = GraphBuilder.CreateUAV(FRDGBufferUAVDesc(OutVisibilityPPLLNodeData));
+		UniformParameters->MaxPPLLNodeCount = PPLLMaxTotalListElementCount;
+
+		PassParameters->UniformBuffer = GraphBuilder.CreateUniformBuffer(UniformParameters);
+	}
+
 	PassParameters->RenderTargets.DepthStencil = FDepthStencilBinding(InViewZDepthTexture, ERenderTargetLoadAction::ELoad, ERenderTargetLoadAction::ENoAction, FExclusiveDepthStencil::DepthRead_StencilNop);
 	AddHairVisibilityCommonPass(GraphBuilder, Scene, ViewInfo, MacroGroupDatas, HairVisibilityRenderMode_PPLL, PassParameters);
 }
@@ -2652,7 +2621,7 @@ static FHairPrimaryTransmittance AddHairViewTransmittancePass(
 	// Clear to transmittance 1
 	FRDGTextureDesc Desc = FRDGTextureDesc::Create2D(Resolution, PF_R32_FLOAT, FClearValueBinding(FLinearColor(1.0f, 1.0f, 1.0f, 1.0f)), TexCreate_RenderTargetable | TexCreate_ShaderResource);
 	FVisibilityPassParameters* PassParameters = GraphBuilder.AllocParameters<FVisibilityPassParameters>();
-	CreatePassDummyTextures(GraphBuilder, PassParameters);
+	PassParameters->UniformBuffer = CreatePassDummyTextures(GraphBuilder);
 	FHairPrimaryTransmittance Out;
 
 	Out.TransmittanceTexture = GraphBuilder.CreateTexture(Desc, TEXT("HairViewTransmittanceTexture"));
diff --git a/Engine/Source/Runtime/Renderer/Private/HairStrands/HairStrandsVoxelization.cpp b/Engine/Source/Runtime/Renderer/Private/HairStrands/HairStrandsVoxelization.cpp
--- a/Engine/Source/Runtime/Renderer/Private/HairStrands/HairStrandsVoxelization.cpp
+++ b/Engine/Source/Runtime/Renderer/Private/HairStrands/HairStrandsVoxelization.cpp
@@ -1183,15 +1183,22 @@ static void AddVirtualVoxelizationRasterPass(
 	}
 
 	FHairVoxelizationRasterPassParameters* PassParameters = GraphBuilder.AllocParameters<FHairVoxelizationRasterPassParameters>();
-	PassParameters->VirtualVoxel = VoxelResources.Parameters.Common;
-	PassParameters->WorldToClipMatrix = WorldToClip;
-	PassParameters->VoxelMinAABB = MacroGroup.VirtualVoxelNodeDesc.WorldMinAABB;
-	PassParameters->VoxelMaxAABB = MacroGroup.VirtualVoxelNodeDesc.WorldMaxAABB;
-	PassParameters->VoxelResolution = TotalVoxelResolution; // i.e., the virtual resolution
-	PassParameters->MacroGroupId = MacroGroup.MacroGroupId;
-	PassParameters->ViewportResolution = RasterResolution;
-	PassParameters->VoxelizationViewInfoBuffer = GraphBuilder.CreateSRV(VoxelResources.VoxelizationViewInfoBuffer);
-	PassParameters->DensityTexture = GraphBuilder.CreateUAV(VoxelResources.PageTexture);
+
+	{
+		FHairVoxelizationRasterUniformParameters* UniformParameters = GraphBuilder.AllocParameters<FHairVoxelizationRasterUniformParameters>();
+
+		UniformParameters->VirtualVoxel = VoxelResources.Parameters.Common;
+		UniformParameters->WorldToClipMatrix = WorldToClip;
+		UniformParameters->VoxelMinAABB = MacroGroup.VirtualVoxelNodeDesc.WorldMinAABB;
+		UniformParameters->VoxelMaxAABB = MacroGroup.VirtualVoxelNodeDesc.WorldMaxAABB;
+		UniformParameters->VoxelResolution = TotalVoxelResolution; // i.e., the virtual resolution
+		UniformParameters->MacroGroupId = MacroGroup.MacroGroupId;
+		UniformParameters->ViewportResolution = RasterResolution;
+		UniformParameters->VoxelizationViewInfoBuffer = GraphBuilder.CreateSRV(VoxelResources.VoxelizationViewInfoBuffer);
+		UniformParameters->DensityTexture = GraphBuilder.CreateUAV(VoxelResources.PageTexture);
+
+		PassParameters->UniformBuffer = GraphBuilder.CreateUniformBuffer(UniformParameters);
+	}
 
 	// For debug purpose
 	#if 0
diff --git a/Engine/Source/Runtime/Renderer/Private/ReflectionEnvironmentRealTimeCapture.cpp b/Engine/Source/Runtime/Renderer/Private/ReflectionEnvironmentRealTimeCapture.cpp
--- a/Engine/Source/Runtime/Renderer/Private/ReflectionEnvironmentRealTimeCapture.cpp
+++ b/Engine/Source/Runtime/Renderer/Private/ReflectionEnvironmentRealTimeCapture.cpp
@@ -275,6 +275,7 @@ void FScene::ValidateSkyLightRealTimeCapture(
 }
 
 BEGIN_SHADER_PARAMETER_STRUCT(FCaptureSkyMeshReflectionPassParameters, )
+	SHADER_PARAMETER_STRUCT_INCLUDE(FViewShaderParameters, View)
 	SHADER_PARAMETER_RDG_UNIFORM_BUFFER(FOpaqueBasePassUniformParameters, BasePass)
 	RENDER_TARGET_BINDING_SLOTS()
 END_SHADER_PARAMETER_STRUCT()
@@ -547,6 +548,7 @@ void FScene::AllocateAndCaptureFrameSkyEnvMap(
 						if (MainView.bSceneHasSkyMaterial)
 						{
 							auto* PassParameters = GraphBuilder.AllocParameters<FCaptureSkyMeshReflectionPassParameters>();
+							PassParameters->View = CubeView.GetShaderParameters();
 							PassParameters->RenderTargets = SkyRC.RenderTargets;
 							PassParameters->BasePass = CreateOpaqueBasePassUniformBuffer(GraphBuilder, MainView, 0);
 
@@ -570,8 +572,7 @@ void FScene::AllocateAndCaptureFrameSkyEnvMap(
 										{
 											FScene* Scene = MainView.Family->Scene->GetRenderScene();
 
-											FMeshPassProcessorRenderState DrawRenderState(CubeViewUniformBuffer);
-											DrawRenderState.SetInstancedViewUniformBuffer(Scene->UniformBuffers.InstancedViewUniformBuffer);
+											FMeshPassProcessorRenderState DrawRenderState;
 
 											FExclusiveDepthStencil::Type BasePassDepthStencilAccess_Sky = bUseDepthBuffer ? FExclusiveDepthStencil::Type(Scene->DefaultBasePassDepthStencilAccess | FExclusiveDepthStencil::DepthWrite)
 												: FExclusiveDepthStencil::Type(Scene->DefaultBasePassDepthStencilAccess & ~FExclusiveDepthStencil::DepthWrite);
diff --git a/Engine/Source/Runtime/Renderer/Private/ShadowRendering.cpp b/Engine/Source/Runtime/Renderer/Private/ShadowRendering.cpp
--- a/Engine/Source/Runtime/Renderer/Private/ShadowRendering.cpp
+++ b/Engine/Source/Runtime/Renderer/Private/ShadowRendering.cpp
@@ -783,7 +783,7 @@ void FProjectedShadowInfo::SetupProjectionStencilMask(
 	bool bMobileModulatedProjections, 
 	bool bCameraInsideShadowFrustum) const
 {
-	FMeshPassProcessorRenderState DrawRenderState(*View);
+	FMeshPassProcessorRenderState DrawRenderState;
 
 	// Depth test wo/ writes, no color writing.
 	DrawRenderState.SetDepthStencilState(TStaticDepthStencilState<false, CF_DepthNearOrEqual>::GetRHI());
diff --git a/Engine/Source/Runtime/Renderer/Private/TranslucentLighting.cpp b/Engine/Source/Runtime/Renderer/Private/TranslucentLighting.cpp
--- a/Engine/Source/Runtime/Renderer/Private/TranslucentLighting.cpp
+++ b/Engine/Source/Runtime/Renderer/Private/TranslucentLighting.cpp
@@ -497,7 +497,7 @@ void FProjectedShadowInfo::RenderTranslucencyDepths(FRDGBuilder& GraphBuilder, F
 		ERDGPassFlags::Raster,
 		[this, SceneRenderer](FRHICommandListImmediate& RHICmdList)
 	{
-		FMeshPassProcessorRenderState DrawRenderState(*ShadowDepthView);
+		FMeshPassProcessorRenderState DrawRenderState;
 
 		// Clear the shadow and its border
 		RHICmdList.SetViewport(
diff --git a/Engine/Source/Runtime/Renderer/Private/VolumetricCloudRendering.cpp b/Engine/Source/Runtime/Renderer/Private/VolumetricCloudRendering.cpp
--- a/Engine/Source/Runtime/Renderer/Private/VolumetricCloudRendering.cpp
+++ b/Engine/Source/Runtime/Renderer/Private/VolumetricCloudRendering.cpp
@@ -905,6 +905,7 @@ private:
 //////////////////////////////////////////////////////////////////////////
 
 BEGIN_SHADER_PARAMETER_STRUCT(FVolumetricCloudShadowParametersPS, )
+	SHADER_PARAMETER_STRUCT_INCLUDE(FViewShaderParameters, View)
 	SHADER_PARAMETER_RDG_UNIFORM_BUFFER(FRenderVolumetricCloudGlobalParameters, TraceVolumetricCloudParamsUB)
 	RENDER_TARGET_BINDING_SLOTS()
 END_SHADER_PARAMETER_STRUCT()
@@ -952,7 +953,6 @@ public:
 	{
 		PassDrawRenderState.SetBlendState(TStaticBlendState<>::GetRHI());
 		PassDrawRenderState.SetDepthStencilState(TStaticDepthStencilState<false, CF_Always>::GetRHI());
-		PassDrawRenderState.SetViewUniformBuffer(Scene->UniformBuffers.ViewUniformBuffer);
 	}
 
 	virtual void AddMeshBatch(const FMeshBatch& RESTRICT MeshBatch, uint64 BatchElementMask, const FPrimitiveSceneProxy* RESTRICT PrimitiveSceneProxy, int32 StaticMeshId = -1) override final
@@ -1379,6 +1379,7 @@ void FSceneRenderer::InitVolumetricCloudsForViews(FRDGBuilder& GraphBuilder, boo
 						TRDGUniformBufferRef<FRenderVolumetricCloudGlobalParameters> TraceVolumetricCloudParamsUB)
 					{
 						FVolumetricCloudShadowParametersPS* CloudShadowParameters = GraphBuilder.AllocParameters<FVolumetricCloudShadowParametersPS>();
+						CloudShadowParameters->View = ViewInfo.GetShaderParameters();
 						CloudShadowParameters->TraceVolumetricCloudParamsUB = TraceVolumetricCloudParamsUB;
 						CloudShadowParameters->RenderTargets[0] = FRenderTargetBinding(CloudTextureTracedOutput, ERenderTargetLoadAction::ENoAction);
 
