Fix for crash in virtual landscape due to erroneous BeginRenderView call in VT.

diff --git a/Engine/Source/Runtime/Renderer/Private/VT/RuntimeVirtualTextureRender.cpp b/Engine/Source/Runtime/Renderer/Private/VT/RuntimeVirtualTextureRender.cpp
--- a/Engine/Source/Runtime/Renderer/Private/VT/RuntimeVirtualTextureRender.cpp
+++ b/Engine/Source/Runtime/Renderer/Private/VT/RuntimeVirtualTextureRender.cpp
@@ -1046,8 +1046,6 @@ namespace RuntimeVirtualTexture
 		View->ViewUniformBuffer = TUniformBufferRef<FViewUniformShaderParameters>::CreateUniformBufferImmediate(*View->CachedViewUniformShaderParameters, UniformBuffer_SingleFrame);
 		Scene->GPUScene.UploadDynamicPrimitiveShaderDataForView(GraphBuilder.RHICmdList, *(const_cast<FScene*>(Scene)), *View);
 
-		View->BeginRenderView();
-
 		// Build graph
 		FRenderGraphSetup GraphSetup(GraphBuilder, Scene->GetFeatureLevel(), MaterialType, OutputTexture0, TextureSize, bIsThumbnails);
 
