Add IES profile support to the path tracer

The cpu side looks for the set of unique profiles and packs them into a 3d atlas so they can be read more conveniently from the HLSL side.

diff --git a/Engine/Shaders/Private/PathTracing/Light/PathTracingLightCommon.ush b/Engine/Shaders/Private/PathTracing/Light/PathTracingLightCommon.ush
--- a/Engine/Shaders/Private/PathTracing/Light/PathTracingLightCommon.ush
+++ b/Engine/Shaders/Private/PathTracing/Light/PathTracingLightCommon.ush
@@ -144,6 +144,31 @@ float ComputeAttenuationFalloff(float Distance, float Attenuation)
 #endif
 }
 
+
+Texture3D    IESTexture;
+SamplerState IESTextureSampler;
+
+float ComputeIESAttenuation(int LightId, float3 WorldPosition)
+{
+	float Result = 1.f;
+	if (SceneLightsData.IESTextureSlice[LightId] >= 0)
+	{
+		float3 LightDirection = normalize(GetPosition(LightId) - WorldPosition);
+
+		float DotProd = dot(LightDirection, GetNormal(LightId));
+		float Angle = asin(DotProd);
+		const float NormAngle = Angle / PI + 0.5f;
+
+		float du = dot(LightDirection, GetdPdu(LightId));
+		float dv = dot(LightDirection, GetdPdv(LightId));
+		float NormTangentAngle = atan2(dv, du) / (PI * 2.f) + 0.5f;
+
+		const float3 UVW = float3(NormAngle, NormTangentAngle, SceneLightsData.IESTextureSlice[LightId]);
+		Result = Texture3DSampleLevel(IESTexture, IESTextureSampler, UVW, 0).r;
+	}
+	return Result;
+}
+
 bool HasTransmission(int LightId)
 {
 	return SceneLightsData.Flags[LightId] & 0x01;
diff --git a/Engine/Shaders/Private/PathTracing/Light/PathTracingPointLight.ush b/Engine/Shaders/Private/PathTracing/Light/PathTracingPointLight.ush
--- a/Engine/Shaders/Private/PathTracing/Light/PathTracingPointLight.ush
+++ b/Engine/Shaders/Private/PathTracing/Light/PathTracingPointLight.ush
@@ -30,7 +30,7 @@ float4 PointLight_TraceLight(RayDesc Ray, int LightId, float HitT)
 		if (t > Ray.TMin && t < HitT)
 		{
 			// #dxr_todo: sphere area is 4*pi*r^2 -- but the factor of 4 is missing for some reason?
-			float3 LightPower = GetColor(LightId);
+			float3 LightPower = GetColor(LightId) * ComputeIESAttenuation(LightId, Ray.Origin);
 			float3 LightRadiance = LightPower / (PI * LightRadius2);
 
 			float SinThetaMax2 = saturate(LightRadius2 / LightDistanceSquared);
@@ -76,7 +76,7 @@ void PointLight_SampleLight(
 
 	OutPdf = DirAndPdf.w;
 
-	float3 LightPower = GetColor(LightId);
+	float3 LightPower = GetColor(LightId) * ComputeIESAttenuation(LightId, WorldPos);
 	float3 LightRadiance = LightPower / (PI * Radius2);
 
 	// When the angle is very small, Radiance over pdf simplifies even more since SolidAngle ~= PI * SinThetaMax2
diff --git a/Engine/Shaders/Private/PathTracing/Light/PathTracingRectLight.ush b/Engine/Shaders/Private/PathTracing/Light/PathTracingRectLight.ush
--- a/Engine/Shaders/Private/PathTracing/Light/PathTracingRectLight.ush
+++ b/Engine/Shaders/Private/PathTracing/Light/PathTracingRectLight.ush
@@ -43,7 +43,7 @@ float4 RectLight_TraceLight(RayDesc Ray, int LightId, float HitT)
 			if (all(abs(UV) <= Rect.Extent))
 			{
 				// stored color is radiance
-				float3 Radiance = GetColor(LightId);
+				float3 Radiance = GetColor(LightId) * ComputeIESAttenuation(LightId, Ray.Origin);
 				float SolidAngle = BuildSphericalRect(Rect).SolidAngle;
 				return float4(Radiance, 1.0 / SolidAngle);
 			}
@@ -84,7 +84,7 @@ void RectLight_SampleLight(
 
 	OutDirection = normalize(Direction);
 	OutDistance = length(Direction);
-	OutRadianceOverPdf = GetColor(LightId) * SphericalRect.SolidAngle;
+	OutRadianceOverPdf = GetColor(LightId) * SphericalRect.SolidAngle * ComputeIESAttenuation(LightId, WorldPos);
 	OutPdf = 1.0 / SphericalRect.SolidAngle;
 }
 
diff --git a/Engine/Shaders/Private/PathTracing/Light/PathTracingSpotLight.ush b/Engine/Shaders/Private/PathTracing/Light/PathTracingSpotLight.ush
--- a/Engine/Shaders/Private/PathTracing/Light/PathTracingSpotLight.ush
+++ b/Engine/Shaders/Private/PathTracing/Light/PathTracingSpotLight.ush
@@ -48,7 +48,7 @@ float4 SpotLight_TraceLight(RayDesc Ray, int LightId, float HitT) {
 		if (t > Ray.TMin && t < HitT)
 		{
 			// #dxr_todo: sphere area is 4*pi*r^2 -- but the factor of 4 is missing for some reason?
-			float3 LightPower = GetColor(LightId);
+			float3 LightPower = GetColor(LightId) * ComputeIESAttenuation(LightId, Ray.Origin);
 			float3 LightRadiance = LightPower / (PI * LightRadius2);
 
 			// adjust lighting according to spot light cone falloff
@@ -99,7 +99,7 @@ void SpotLight_SampleLight(
 
 	OutPdf = DirAndPdf.w;
 
-	float3 LightPower = GetColor(LightId);
+	float3 LightPower = GetColor(LightId) * ComputeIESAttenuation(LightId, WorldPos);
 	float3 LightRadiance = LightPower / (PI * Radius2);
 
 	// When the angle is very small, Radiance over pdf simplifies even more since SolidAngle ~= PI * SinThetaMax2
diff --git a/Engine/Shaders/Private/PathTracing/PathTracingIESAtlas.usf b/Engine/Shaders/Private/PathTracing/PathTracingIESAtlas.usf
new file mode 100644
--- /dev/null
+++ b/Engine/Shaders/Private/PathTracing/PathTracingIESAtlas.usf
@@ -0,0 +1,22 @@
+// Copyright Epic Games, Inc. All Rights Reserved.
+
+#include "../Common.ush"
+
+Texture2D IESTexture;
+SamplerState IESSampler;
+RWTexture3D<float> IESAtlas;
+int IESAtlasSlice;
+
+[numthreads(THREADGROUPSIZE_X, THREADGROUPSIZE_Y, 1)]
+void PathTracingIESAtlasCS(uint2 DispatchThreadId : SV_DispatchThreadID)
+{
+	uint3 AtlasSize;
+	IESAtlas.GetDimensions(AtlasSize.x, AtlasSize.y, AtlasSize.z);
+	if (any(DispatchThreadId >= AtlasSize.xy))
+	{
+		return;
+	}
+
+	float2 UV = (DispatchThreadId + 0.5) / float2(AtlasSize.xy);
+	IESAtlas[uint3(DispatchThreadId, IESAtlasSlice)] = Texture2DSampleLevel(IESTexture, IESSampler, UV, 0).r;
+}
diff --git a/Engine/Source/Runtime/Renderer/Private/PathTracing.cpp b/Engine/Source/Runtime/Renderer/Private/PathTracing.cpp
--- a/Engine/Source/Runtime/Renderer/Private/PathTracing.cpp
+++ b/Engine/Source/Runtime/Renderer/Private/PathTracing.cpp
@@ -176,12 +176,42 @@ class FPathTracingRG : public FGlobalShader
 		SHADER_PARAMETER_STRUCT_REF(FSkyLightData, SkyLightData)
 		SHADER_PARAMETER_STRUCT_REF(FPathTracingLightData, SceneLightsData)
 		SHADER_PARAMETER_STRUCT_REF(FPathTracingData, PathTracingData)
+		// IES Profiles
+		SHADER_PARAMETER_RDG_TEXTURE(Texture3D, IESTexture)
+		SHADER_PARAMETER_SAMPLER(SamplerState, IESTextureSampler)
+		// Shared sampler for all IES profiles
 		SHADER_PARAMETER(FIntVector, TileOffset)	// Used by multi-GPU rendering
 	END_SHADER_PARAMETER_STRUCT()
 };
 IMPLEMENT_GLOBAL_SHADER(FPathTracingRG, "/Engine/Private/PathTracing/PathTracing.usf", "PathTracingMainRG", SF_RayGen);
 
-void SetLightParameters(FPathTracingLightData& LightData, FSkyLightData& SkyLightData, FScene* Scene)
+class FPathTracingIESAtlasCS : public FGlobalShader
+{
+	DECLARE_GLOBAL_SHADER(FPathTracingIESAtlasCS)
+	SHADER_USE_PARAMETER_STRUCT(FPathTracingIESAtlasCS, FGlobalShader)
+
+	static bool ShouldCompilePermutation(const FGlobalShaderPermutationParameters& Parameters)
+	{
+		return ShouldCompileRayTracingShadersForProject(Parameters.Platform);
+	}
+
+	static void ModifyCompilationEnvironment(const FGlobalShaderPermutationParameters& Parameters, FShaderCompilerEnvironment& OutEnvironment)
+	{
+		FGlobalShader::ModifyCompilationEnvironment(Parameters, OutEnvironment);
+		OutEnvironment.SetDefine(TEXT("THREADGROUPSIZE_X"), FComputeShaderUtils::kGolden2DGroupSize);
+		OutEnvironment.SetDefine(TEXT("THREADGROUPSIZE_Y"), FComputeShaderUtils::kGolden2DGroupSize);
+	}
+
+	BEGIN_SHADER_PARAMETER_STRUCT(FParameters, )
+		SHADER_PARAMETER_TEXTURE(Texture2D, IESTexture)
+		SHADER_PARAMETER_SAMPLER(SamplerState, IESSampler)
+		SHADER_PARAMETER_RDG_TEXTURE_UAV(RWTexture3D<float>, IESAtlas)
+		SHADER_PARAMETER(int32, IESAtlasSlice)
+	END_SHADER_PARAMETER_STRUCT()
+};
+IMPLEMENT_SHADER_TYPE(, FPathTracingIESAtlasCS, TEXT("/Engine/Private/PathTracing/PathTracingIESAtlas.usf"), TEXT("PathTracingIESAtlasCS"), SF_Compute);
+
+void SetLightParameters(FRDGBuilder& GraphBuilder, FPathTracingRG::FParameters* PassParameters, FPathTracingLightData& LightData, FSkyLightData& SkyLightData, FScene* Scene, const FViewInfo& View, bool UseLightProfiles)
 {
 	// Sky light
 	bool IsSkyLightValid = SetupSkyLightParameters(*Scene, &SkyLightData);
@@ -198,9 +228,11 @@ void SetLightParameters(FPathTracingLightData& LightData, FSkyLightData& SkyLigh
 		LightData.Color[SkyLightIndex] = FVector(SkyLightData.Color);
 		LightData.Flags[SkyLightIndex] = SkyLightData.bTransmission & 0x01;
 		LightData.Flags[SkyLightIndex] |= (SkyLightLightingChannelMask & 0x7) << 1;
+		LightData.IESTextureSlice[SkyLightIndex] = -1;
 		LightData.Count++;
 	}
 
+	TMap<FTexture*, int> IESLightProfilesMap;
 	for (auto Light : Scene->Lights)
 	{
 		if (LightData.Count >= RAY_TRACING_LIGHT_COUNT_MAXIMUM) break;
@@ -211,6 +243,24 @@ void SetLightParameters(FPathTracingLightData& LightData, FSkyLightData& SkyLigh
 		uint8 LightingChannelMask = Light.LightSceneInfo->Proxy->GetLightingChannelMask();
 		LightData.Flags[LightData.Count] = Transmission & 0x01;
 		LightData.Flags[LightData.Count] |= (LightingChannelMask & 0x7) << 1;
+		LightData.IESTextureSlice[LightData.Count] = -1;
+
+		if (UseLightProfiles)
+		{
+			FTexture* IESTexture = Light.LightSceneInfo->Proxy->GetIESTextureResource();
+			if (IESTexture != nullptr)
+			{
+				// Only add a given texture once
+				LightData.IESTextureSlice[LightData.Count] = IESLightProfilesMap.FindOrAdd(IESTexture, IESLightProfilesMap.Num());
+			}
+		}
+
+		// these mean roughly the same thing across all light types
+		LightData.Color[LightData.Count] = LightParameters.Color;
+		LightData.Position[LightData.Count] = LightParameters.Position;
+		LightData.Normal[LightData.Count] = -LightParameters.Direction;
+		LightData.dPdu[LightData.Count] = FVector::CrossProduct(LightParameters.Tangent, LightParameters.Direction);
+		LightData.dPdv[LightData.Count] = LightParameters.Tangent;
 
 		ELightComponentType LightComponentType = (ELightComponentType)Light.LightSceneInfo->Proxy->GetLightType();
 		switch (LightComponentType)
@@ -218,8 +268,7 @@ void SetLightParameters(FPathTracingLightData& LightData, FSkyLightData& SkyLigh
 			case LightType_Directional:
 			{
 				LightData.Type[LightData.Count] = 2;
-				LightData.Normal[LightData.Count] = LightParameters.Direction;
-				LightData.Color[LightData.Count] = LightParameters.Color;
+                LightData.Normal[LightData.Count] = LightParameters.Direction;
 				LightData.Dimensions[LightData.Count] = FVector(0.0f, 0.0f, LightParameters.SourceRadius);
 				LightData.Attenuation[LightData.Count] = 1.0 / LightParameters.InvRadius;
 				break;
@@ -227,11 +276,6 @@ void SetLightParameters(FPathTracingLightData& LightData, FSkyLightData& SkyLigh
 			case LightType_Rect:
 			{
 				LightData.Type[LightData.Count] = 3;
-				LightData.Position[LightData.Count] = LightParameters.Position;
-				LightData.Normal[LightData.Count] = -LightParameters.Direction;
-				LightData.dPdu[LightData.Count] = FVector::CrossProduct(LightParameters.Tangent, LightParameters.Direction);
-				LightData.dPdv[LightData.Count] = LightParameters.Tangent;
-				LightData.Color[LightData.Count] = LightParameters.Color;
 				LightData.Dimensions[LightData.Count] = FVector(2.0f * LightParameters.SourceRadius, 2.0f * LightParameters.SourceLength, 0.0f);
 				LightData.Attenuation[LightData.Count] = 1.0 / LightParameters.InvRadius;
 				LightData.RectLightBarnCosAngle[LightData.Count] = LightParameters.RectLightBarnCosAngle;
@@ -241,9 +285,6 @@ void SetLightParameters(FPathTracingLightData& LightData, FSkyLightData& SkyLigh
 			case LightType_Spot:
 			{
 				LightData.Type[LightData.Count] = 4;
-				LightData.Position[LightData.Count] = LightParameters.Position;
-				LightData.Normal[LightData.Count] = -LightParameters.Direction;
-				LightData.Color[LightData.Count] = LightParameters.Color;
 				LightData.Dimensions[LightData.Count] = FVector(LightParameters.SpotAngles, LightParameters.SourceRadius);
 				LightData.Attenuation[LightData.Count] = 1.0 / LightParameters.InvRadius;
 				break;
@@ -251,8 +292,6 @@ void SetLightParameters(FPathTracingLightData& LightData, FSkyLightData& SkyLigh
 			case LightType_Point:
 			{
 				LightData.Type[LightData.Count] = 1;
-				LightData.Position[LightData.Count] = LightParameters.Position;
-				LightData.Color[LightData.Count] = LightParameters.Color;
 				LightData.Dimensions[LightData.Count] = FVector(0.0, 0.0, LightParameters.SourceRadius);
 				LightData.Attenuation[LightData.Count] = 1.0 / LightParameters.InvRadius;
 				break;
@@ -267,6 +306,65 @@ void SetLightParameters(FPathTracingLightData& LightData, FSkyLightData& SkyLigh
 
 		LightData.Count++;
 	}
+
+	if (!IESLightProfilesMap.IsEmpty())
+	{
+		// We found some IES profiles to use -- upload them into a single atlas so we can access them easily in HLSL
+
+		// FIXME: This is redundant because all the IES textures are already on the GPU, we just don't have the ability to use
+		// an array of texture handles on the HLSL side.
+
+		// FIXME: This is also redundant with the logic in RayTracingLighting.cpp, but the latter is limitted to 1D profiles and 
+		// does not consider the same set of lights as the path tracer. Longer term we should aim to unify the representation of lights
+		// across both passes
+
+		// FIXME: This process is repeated every frame! would be nicer to cache the data somehow. Perhaps just do this step for
+		// Iteration == 0 since we can assume that any changes in IES profiles will invalidate the path tracer anyway?
+
+		// This size matches the import resolution of light profiles (see FIESLoader::GetWidth)
+		const int kIESAtlasSize = 256;
+		const int NumSlices = IESLightProfilesMap.Num();
+		// TODO: would a Texture2DArray object be more efficient?
+		FRDGTextureDesc IESTextureDesc = FRDGTextureDesc::Create3D(
+			FIntVector(kIESAtlasSize, kIESAtlasSize, NumSlices),
+			PF_R32_FLOAT,
+			FClearValueBinding::None,
+			TexCreate_ShaderResource | TexCreate_UAV);
+		FRDGTexture* IESTexture = GraphBuilder.CreateTexture(IESTextureDesc, TEXT("PathTracerIESAtlas"), ERDGTextureFlags::None);
+
+		for (auto&& Entry : IESLightProfilesMap)
+		{
+			FPathTracingIESAtlasCS::FParameters* AtlasPassParameters = GraphBuilder.AllocParameters<FPathTracingIESAtlasCS::FParameters>();
+			const int Slice = Entry.Value;
+			AtlasPassParameters->IESTexture = Entry.Key->TextureRHI;
+			AtlasPassParameters->IESSampler = TStaticSamplerState<SF_Bilinear, AM_Clamp, AM_Clamp, AM_Clamp>::GetRHI();
+			AtlasPassParameters->IESAtlas = GraphBuilder.CreateUAV(IESTexture);
+			AtlasPassParameters->IESAtlasSlice = Slice;
+			TShaderMapRef<FPathTracingIESAtlasCS> ComputeShader(View.ShaderMap);
+			FComputeShaderUtils::AddPass(
+				GraphBuilder,
+				RDG_EVENT_NAME("Path Tracing IES Atlas (Slice=%d)", Slice),
+				ComputeShader,
+				AtlasPassParameters,
+				FComputeShaderUtils::GetGroupCount(FIntPoint(kIESAtlasSize, kIESAtlasSize), FComputeShaderUtils::kGolden2DGroupSize));
+		}
+
+		// Adjust row per light so it can be used directly as a texture coordinate
+		for (uint32 Index = 0; Index < LightData.Count; Index++)
+		{
+			if (LightData.IESTextureSlice[Index] >= 0)
+			{
+				// nudge data to be easier to use on GPU
+				LightData.IESTextureSlice[Index] = (LightData.IESTextureSlice[Index] + 0.5f) / NumSlices;
+			}
+		}
+
+		PassParameters->IESTexture = IESTexture;
+	}
+	else
+	{
+		PassParameters->IESTexture = GraphBuilder.RegisterExternalTexture(GSystemTextures.WhiteDummy);
+	}
 }
 
 class FPathTracingCompositorPS : public FGlobalShader
@@ -325,6 +423,7 @@ void FDeferredShadingSceneRenderer::RenderPathTracing(
 	// Get current value of MaxSPP and reset render if it has changed
 	int32 SamplesPerPixelCVar = CVarPathTracingSamplesPerPixel.GetValueOnRenderThread();
 	uint32 MaxSPP = SamplesPerPixelCVar > -1 ? SamplesPerPixelCVar : View.FinalPostProcessSettings.PathTracingSamplesPerPixel;
+	MaxSPP = FMath::Max(MaxSPP, 1u);
 	static uint32 PreviousMaxSPP = MaxSPP;
 	if (PreviousMaxSPP != MaxSPP)
 	{
@@ -340,6 +439,14 @@ void FDeferredShadingSceneRenderer::RenderPathTracing(
 		bArgsChanged = true;
 	}
 
+	bool UseLightProfiles = View.Family->EngineShowFlags.TexturedLightProfiles;
+	static bool PreviousUseLightProfiles = UseLightProfiles;
+	if (PreviousUseLightProfiles != UseLightProfiles)
+	{
+		PreviousUseLightProfiles = UseLightProfiles;
+		bArgsChanged = true;
+	}
+
 	// Get other basic path tracing settings and see if we need to invalidate the current state
 	FPathTracingData PathTracingData;
 	bArgsChanged |= PrepareShaderArgs(View, PathTracingData);
@@ -395,10 +502,11 @@ void FDeferredShadingSceneRenderer::RenderPathTracing(
 			// upload sky/lights data
 			FSkyLightData SkyLightData;
 			FPathTracingLightData LightData;
-			SetLightParameters(LightData, SkyLightData, Scene);
+			SetLightParameters(GraphBuilder, PassParameters, LightData, SkyLightData, Scene, View, UseLightProfiles);
 			PassParameters->SkyLightData = CreateUniformBufferImmediate(SkyLightData, EUniformBufferUsage::UniformBuffer_SingleFrame);
 			PassParameters->SceneLightsData = CreateUniformBufferImmediate(LightData, EUniformBufferUsage::UniformBuffer_SingleFrame);
 		}
+		PassParameters->IESTextureSampler = TStaticSamplerState<SF_Bilinear, AM_Clamp, AM_Clamp, AM_Clamp>::GetRHI();
 		PassParameters->RadianceTexture = GraphBuilder.CreateUAV(RadianceTexture);
 
 		// TODO: in multi-gpu case, split image into tiles
diff --git a/Engine/Source/Runtime/Renderer/Private/PathTracingUniformBuffers.h b/Engine/Source/Runtime/Renderer/Private/PathTracingUniformBuffers.h
--- a/Engine/Source/Runtime/Renderer/Private/PathTracingUniformBuffers.h
+++ b/Engine/Source/Runtime/Renderer/Private/PathTracingUniformBuffers.h
@@ -39,6 +39,8 @@ BEGIN_GLOBAL_SHADER_PARAMETER_STRUCT(FPathTracingLightData, RENDERER_API)
 	SHADER_PARAMETER_ARRAY(float, Attenuation, [RAY_TRACING_LIGHT_COUNT_MAXIMUM])
 	SHADER_PARAMETER_ARRAY(float, RectLightBarnCosAngle, [RAY_TRACING_LIGHT_COUNT_MAXIMUM])
 	SHADER_PARAMETER_ARRAY(float, RectLightBarnLength, [RAY_TRACING_LIGHT_COUNT_MAXIMUM])
+	// IES index (valid if >= 0)
+	SHADER_PARAMETER_ARRAY(float, IESTextureSlice , [RAY_TRACING_LIGHT_COUNT_MAXIMUM])
 	// Flags
 	SHADER_PARAMETER_ARRAY(uint32, Flags, [RAY_TRACING_LIGHT_COUNT_MAXIMUM])
 	// Only used by GPULightmass currently, not filled in realtime paths
