Fixed null texture in uniform buffer crash in hair strands deep shadow pass.

diff --git a/Engine/Source/Runtime/Renderer/Private/HairStrands/HairStrandsDeepShadow.cpp b/Engine/Source/Runtime/Renderer/Private/HairStrands/HairStrandsDeepShadow.cpp
--- a/Engine/Source/Runtime/Renderer/Private/HairStrands/HairStrandsDeepShadow.cpp
+++ b/Engine/Source/Runtime/Renderer/Private/HairStrands/HairStrandsDeepShadow.cpp
@@ -429,6 +429,8 @@ void RenderHairStrandsDeepShadows(
 				// Front depth
 				if (bDeepShadow)
 				{
+					const FRDGSystemTextures& SystemTextures = FRDGSystemTextures::Get(GraphBuilder);
+
 					DECLARE_GPU_STAT(HairStrandsDeepShadowFrontDepth);
 					RDG_EVENT_SCOPE(GraphBuilder, "HairStrandsDeepShadowFrontDepth");
 					RDG_GPU_STAT_SCOPE(GraphBuilder, HairStrandsDeepShadowFrontDepth);
@@ -444,6 +446,7 @@ void RenderHairStrandsDeepShadows(
 						UniformParameters->AtlasSlotIndex = DomData.AtlasSlotIndex;
 						UniformParameters->LayerDepths = LayerDepths;
 						UniformParameters->ViewportResolution = AtlasSlotResolution;
+						UniformParameters->FrontDepthTexture = SystemTextures.DepthDummy;
 						UniformParameters->DeepShadowViewInfoBuffer = DeepShadowViewInfoBufferSRV;
 
 						PassParameters->UniformBuffer = GraphBuilder.CreateUniformBuffer(UniformParameters);
