Add a path to sample conventional directional shadow map that does not contain nanite geometry in forward shading passes.

diff --git a/Engine/Shaders/Private/ForwardShadowingCommon.ush b/Engine/Shaders/Private/ForwardShadowingCommon.ush
--- a/Engine/Shaders/Private/ForwardShadowingCommon.ush
+++ b/Engine/Shaders/Private/ForwardShadowingCommon.ush
@@ -50,25 +50,22 @@ float ComputeDirectionalLightStaticShadowing(float3 WorldPosition)
 #define FILTER_DIRECTIONAL_LIGHT_SHADOWING 0
 #endif
 
+uint GetForwardLightingCascadeIndex(float4 CascadeEndDepths, float SceneDepth)
+{
+	float4 Count = float4(SceneDepth.xxxx >= CascadeEndDepths);
+	return uint(Count.x + Count.y + Count.z + Count.w);
+}
+
 float ComputeDirectionalLightDynamicShadowing(float3 WorldPosition, float SceneDepth, inout bool bShadowingFromValidUVArea)
 {
 	float ShadowFactor = 1;
 	bShadowingFromValidUVArea = false;
 
-	if (ForwardLightData.NumDirectionalLightCascades > 0)
+	const uint NumCascades = ForwardLightData.NumDirectionalLightCascades;
+	if (NumCascades > 0)
 	{
-		uint CascadeIndex = ForwardLightData.NumDirectionalLightCascades;
-
-		for (uint TestCascadeIndex = 0; TestCascadeIndex < ForwardLightData.NumDirectionalLightCascades; TestCascadeIndex++)
-		{
-			if (SceneDepth < ForwardLightData.CascadeEndDepths[TestCascadeIndex])
-			{
-				CascadeIndex = TestCascadeIndex;
-				break;
-			}
-		}
-
-		if (CascadeIndex < ForwardLightData.NumDirectionalLightCascades)
+		uint CascadeIndex = GetForwardLightingCascadeIndex(ForwardLightData.CascadeEndDepths, SceneDepth);
+		if (CascadeIndex < NumCascades)
 		{
 			// Transform the world position into shadowmap space
 			float4 HomogeneousShadowPosition = mul(float4(WorldPosition, 1), ForwardLightData.DirectionalLightWorldToShadowMatrix[CascadeIndex]);
@@ -103,5 +100,49 @@ float ComputeDirectionalLightDynamicShadowing(float3 WorldPosition, float SceneD
 		}
 	}
 
+	return ShadowFactor;
+}
+
+// As above, but samples the high(er) resolution conventional shadow map that does not contain Nanite geometry
+// rather than the combined, lower-resolution "complete" shadow map that contains both types.
+float ComputeDirectionalLightDynamicShadowingNonNanite(float3 WorldPosition, float SceneDepth, inout bool bShadowingFromValidUVArea)
+{
+	float ShadowFactor = 1;
+	bShadowingFromValidUVArea = false;
+
+	const uint NumCascades = ForwardLightData.NumDirectionalLightCascades;
+	if (NumCascades > 0)
+	{
+		uint CascadeIndex = GetForwardLightingCascadeIndex(ForwardLightData.NonNaniteCascadeEndDepths, SceneDepth);
+		if (CascadeIndex < NumCascades)
+		{
+			float4 HomogeneousShadowPosition = mul(float4(WorldPosition, 1), ForwardLightData.NonNaniteDirectionalLightWorldToShadowMatrix[CascadeIndex]);
+			float2 ShadowUVs = HomogeneousShadowPosition.xy / HomogeneousShadowPosition.w;
+			float4 ShadowmapMinMax = ForwardLightData.NonNaniteDirectionalLightShadowmapMinMax[CascadeIndex];
+
+			if (all(ShadowUVs >= ShadowmapMinMax.xy && ShadowUVs <= ShadowmapMinMax.zw))
+			{
+	#if FILTER_DIRECTIONAL_LIGHT_SHADOWING
+				FPCFSamplerSettings Settings;
+				Settings.ShadowDepthTexture = ForwardLightData.NonNaniteDirectionalLightShadowmapAtlas;
+				Settings.ShadowDepthTextureSampler = ForwardLightData.ShadowmapSampler;
+				Settings.ShadowBufferSize = ForwardLightData.NonNaniteDirectionalLightShadowmapAtlasBufferSize;
+				Settings.SceneDepth = 1 - HomogeneousShadowPosition.z;
+				Settings.TransitionScale = 4000;
+				Settings.bSubsurface = false;
+				Settings.bTreatMaxDepthUnshadowed = false;
+				Settings.DensityMulConstant = 0;
+				Settings.ProjectionDepthBiasParameters = float2(0, 0);
+
+				ShadowFactor = Manual1x1PCF(ShadowUVs, Settings);
+	#else
+				float ShadowDepth = Texture2DSampleLevel(ForwardLightData.NonNaniteDirectionalLightShadowmapAtlas, ForwardLightData.ShadowmapSampler, ShadowUVs, 0).x;
+				ShadowFactor = 1 - HomogeneousShadowPosition.z < ShadowDepth - ForwardLightData.NonNaniteDirectionalLightDepthBias.x;
+	#endif
+				bShadowingFromValidUVArea = true;
+			}
+		}
+	}
+
 	return ShadowFactor;
 }
\ No newline at end of file
diff --git a/Engine/Source/Runtime/Renderer/Private/LightGridInjection.cpp b/Engine/Source/Runtime/Renderer/Private/LightGridInjection.cpp
--- a/Engine/Source/Runtime/Renderer/Private/LightGridInjection.cpp
+++ b/Engine/Source/Runtime/Renderer/Private/LightGridInjection.cpp
@@ -133,6 +133,7 @@ FForwardLightData::FForwardLightData()
 {
 	FMemory::Memzero(*this);
 	DirectionalLightShadowmapAtlas = GBlackTexture->TextureRHI;
+	NonNaniteDirectionalLightShadowmapAtlas = GBlackTexture->TextureRHI;
 	ShadowmapSampler = TStaticSamplerState<SF_Point,AM_Clamp,AM_Clamp,AM_Clamp>::GetRHI();
 	DirectionalLightStaticShadowmap = GBlackTexture->TextureRHI;
 	StaticShadowmapSampler = TStaticSamplerState<SF_Bilinear,AM_Clamp,AM_Clamp,AM_Clamp>::GetRHI();
@@ -486,33 +487,50 @@ void FSceneRenderer::ComputeLightGrid(FRDGBuilder& GraphBuilder, bool bCullLight
 						{
 							// If complete/fallback shadow maps are available, prefer those as they include all geometry types
 							bool bUseCompleteShadowMaps = VisibleLightInfos[LightSceneInfo->Id].CompleteProjectedShadows.Num() > 0;
-							const TArray<FProjectedShadowInfo*, SceneRenderingAllocator>& DirectionalLightShadowInfos = bUseCompleteShadowMaps ?
-								VisibleLightInfos[LightSceneInfo->Id].CompleteProjectedShadows :
-								VisibleLightInfos[LightSceneInfo->Id].AllProjectedShadows;
+							const TArray<FProjectedShadowInfo*, SceneRenderingAllocator>& DirectionalLightShadowInfos = VisibleLightInfos[LightSceneInfo->Id].AllProjectedShadows;
 
 							ForwardLightData.NumDirectionalLightCascades = 0;
+							ForwardLightData.NonNaniteNumDirectionalLightCascades = 0;
+							// Unused cascades should compare > all scene depths
+							ForwardLightData.CascadeEndDepths = FVector4(MAX_FLT, MAX_FLT, MAX_FLT, MAX_FLT);
+							ForwardLightData.NonNaniteCascadeEndDepths = ForwardLightData.CascadeEndDepths;
 
 							for (const FProjectedShadowInfo* ShadowInfo : DirectionalLightShadowInfos)
 							{
-								if (!bUseCompleteShadowMaps && ShadowInfo->bCompleteShadowMap)
-								{
-									continue;
-								}
-
 								const int32 CascadeIndex = ShadowInfo->CascadeSettings.ShadowSplitIndex;
 
 								if (ShadowInfo->IsWholeSceneDirectionalShadow() && ShadowInfo->bAllocated && CascadeIndex < GMaxForwardShadowCascades)
 								{
-									ForwardLightData.NumDirectionalLightCascades++;
-									ForwardLightData.DirectionalLightWorldToShadowMatrix[CascadeIndex] = ShadowInfo->GetWorldToShadowMatrix(ForwardLightData.DirectionalLightShadowmapMinMax[CascadeIndex]);
-									ForwardLightData.CascadeEndDepths[CascadeIndex] = ShadowInfo->CascadeSettings.SplitFar;
-
-									if (CascadeIndex == 0)
+									// Always put the conventional shadow maps in the "non-nanite" slots
+									if (!ShadowInfo->bCompleteShadowMap)
+									{
+										ForwardLightData.NonNaniteNumDirectionalLightCascades++;
+										ForwardLightData.NonNaniteDirectionalLightWorldToShadowMatrix[CascadeIndex] = ShadowInfo->GetWorldToShadowMatrix(ForwardLightData.NonNaniteDirectionalLightShadowmapMinMax[CascadeIndex]);
+										ForwardLightData.NonNaniteCascadeEndDepths[CascadeIndex] = ShadowInfo->CascadeSettings.SplitFar;
+
+										if (CascadeIndex == 0)
+										{
+											ForwardLightData.NonNaniteDirectionalLightShadowmapAtlas = ShadowInfo->RenderTargets.DepthTarget->GetRenderTargetItem().ShaderResourceTexture.GetReference();
+											ForwardLightData.NonNaniteDirectionalLightDepthBias = ShadowInfo->GetShaderDepthBias();
+											FVector2D AtlasSize = ShadowInfo->RenderTargets.DepthTarget->GetDesc().Extent;
+											ForwardLightData.NonNaniteDirectionalLightShadowmapAtlasBufferSize = FVector4(AtlasSize.X, AtlasSize.Y, 1.0f / AtlasSize.X, 1.0f / AtlasSize.Y);
+										}
+									}
+									
+									// Put either the complete or the conventional shadow map in the regular sampling slot as configured
+									if (bUseCompleteShadowMaps == ShadowInfo->bCompleteShadowMap)
 									{
-										ForwardLightData.DirectionalLightShadowmapAtlas = ShadowInfo->RenderTargets.DepthTarget->GetRenderTargetItem().ShaderResourceTexture.GetReference();
-										ForwardLightData.DirectionalLightDepthBias = ShadowInfo->GetShaderDepthBias();
-										FVector2D AtlasSize = ShadowInfo->RenderTargets.DepthTarget->GetDesc().Extent;
-										ForwardLightData.DirectionalLightShadowmapAtlasBufferSize = FVector4(AtlasSize.X, AtlasSize.Y, 1.0f / AtlasSize.X, 1.0f / AtlasSize.Y);
+										ForwardLightData.NumDirectionalLightCascades++;
+										ForwardLightData.DirectionalLightWorldToShadowMatrix[CascadeIndex] = ShadowInfo->GetWorldToShadowMatrix(ForwardLightData.DirectionalLightShadowmapMinMax[CascadeIndex]);
+										ForwardLightData.CascadeEndDepths[CascadeIndex] = ShadowInfo->CascadeSettings.SplitFar;
+
+										if (CascadeIndex == 0)
+										{
+											ForwardLightData.DirectionalLightShadowmapAtlas = ShadowInfo->RenderTargets.DepthTarget->GetRenderTargetItem().ShaderResourceTexture.GetReference();
+											ForwardLightData.DirectionalLightDepthBias = ShadowInfo->GetShaderDepthBias();
+											FVector2D AtlasSize = ShadowInfo->RenderTargets.DepthTarget->GetDesc().Extent;
+											ForwardLightData.DirectionalLightShadowmapAtlasBufferSize = FVector4(AtlasSize.X, AtlasSize.Y, 1.0f / AtlasSize.X, 1.0f / AtlasSize.Y);
+										}
 									}
 								}
 							}
diff --git a/Engine/Source/Runtime/Renderer/Private/SceneRendering.h b/Engine/Source/Runtime/Renderer/Private/SceneRendering.h
--- a/Engine/Source/Runtime/Renderer/Private/SceneRendering.h
+++ b/Engine/Source/Runtime/Renderer/Private/SceneRendering.h
@@ -663,12 +663,19 @@ const int32 GMaxForwardShadowCascades = 4;
 	SHADER_PARAMETER_ARRAY(FVector4, DirectionalLightShadowmapMinMax, [GMaxForwardShadowCascades]) \
 	SHADER_PARAMETER(FVector4, DirectionalLightShadowmapAtlasBufferSize) \
 	SHADER_PARAMETER(float, DirectionalLightDepthBias) \
+	SHADER_PARAMETER(uint32, NonNaniteNumDirectionalLightCascades) \
+	SHADER_PARAMETER(FVector4, NonNaniteCascadeEndDepths) \
+	SHADER_PARAMETER_ARRAY(FMatrix, NonNaniteDirectionalLightWorldToShadowMatrix, [GMaxForwardShadowCascades]) \
+	SHADER_PARAMETER_ARRAY(FVector4, NonNaniteDirectionalLightShadowmapMinMax, [GMaxForwardShadowCascades]) \
+	SHADER_PARAMETER(FVector4, NonNaniteDirectionalLightShadowmapAtlasBufferSize) \
+	SHADER_PARAMETER(float, NonNaniteDirectionalLightDepthBias) \
 	SHADER_PARAMETER(uint32, DirectionalLightUseStaticShadowing) \
 	SHADER_PARAMETER(uint32, SimpleLightsEndIndex) \
 	SHADER_PARAMETER(uint32, ClusteredDeferredSupportedEndIndex) \
 	SHADER_PARAMETER(FVector4, DirectionalLightStaticShadowBufferSize) \
 	SHADER_PARAMETER(FMatrix, DirectionalLightWorldToStaticShadow) \
 	SHADER_PARAMETER_TEXTURE(Texture2D, DirectionalLightShadowmapAtlas) \
+	SHADER_PARAMETER_TEXTURE(Texture2D, NonNaniteDirectionalLightShadowmapAtlas) \
 	SHADER_PARAMETER_SAMPLER(SamplerState, ShadowmapSampler) \
 	SHADER_PARAMETER_TEXTURE(Texture2D, DirectionalLightStaticShadowmap) \
 	SHADER_PARAMETER_SAMPLER(SamplerState, StaticShadowmapSampler) \
