Fix a bug of missing view uniform buffer in mobile deferred rendering after porting to RDG.

diff --git a/Engine/Source/Runtime/Renderer/Private/MobileDeferredShadingPass.cpp b/Engine/Source/Runtime/Renderer/Private/MobileDeferredShadingPass.cpp
--- a/Engine/Source/Runtime/Renderer/Private/MobileDeferredShadingPass.cpp
+++ b/Engine/Source/Runtime/Renderer/Private/MobileDeferredShadingPass.cpp
@@ -137,7 +137,7 @@ static void RenderDirectLight(FRHICommandListImmediate& RHICmdList, const FScene
 	FSceneRenderTargets& SceneContext = FSceneRenderTargets::Get();
 	
 	FMobileDeferredShadingPS::FParameters PassParameters;
-	PassParameters.View = View.ViewUniformBuffer;
+	PassParameters.View = GetShaderBinding(View.ViewUniformBuffer);
 	PassParameters.Forward = View.ForwardLightingResources->ForwardLightDataUniformBuffer;
 	//PassParameters.SceneTextures = SceneTexturesBuffer;
 	PassParameters.MobileDirectionalLight = Scene.UniformBuffers.MobileDirectionalLightUniformBuffers[1];
@@ -305,7 +305,7 @@ static void RenderLocalLight(FRHICommandListImmediate& RHICmdList, const FScene&
 	SetGraphicsPipelineState(RHICmdList, GraphicsPSOInit);
 
 	FMobileRadialLightPS::FParameters PassParameters;
-	PassParameters.View = View.ViewUniformBuffer;
+	PassParameters.View = GetShaderBinding(View.ViewUniformBuffer);
 	PassParameters.DeferredLightUniforms = TUniformBufferRef<FDeferredLightUniformStruct>::CreateUniformBufferImmediate(GetDeferredLightParameters(View, LightSceneInfo), EUniformBufferUsage::UniformBuffer_SingleFrame);
 	PassParameters.IESTexture = IESTextureResource->TextureRHI;
 	PassParameters.IESTextureSampler = IESTextureResource->SamplerStateRHI;
