Fix ReflectionCapture in editor for cases where the resource in the UTextureCube is updated after the ReflectionProxy is created

diff --git a/Engine/Source/Runtime/Renderer/Private/RendererScene.cpp b/Engine/Source/Runtime/Renderer/Private/RendererScene.cpp
--- a/Engine/Source/Runtime/Renderer/Private/RendererScene.cpp
+++ b/Engine/Source/Runtime/Renderer/Private/RendererScene.cpp
@@ -54,6 +54,7 @@
 #include "Lumen/LumenSceneRendering.h"
 #include "PlanarReflectionSceneProxy.h"
 #include "Engine/StaticMesh.h"
+#include "Engine/TextureCube.h"
 #include "GPUSkinCache.h"
 #include "ComputeFramework/ComputeFramework.h"
 #include "DynamicShadowMapChannelBindingHelper.h"
@@ -2120,11 +2121,16 @@ void FScene::UpdateReflectionCaptureTransform(UReflectionCaptureComponent* Compo
 
 		FScene* Scene = this;
 		FReflectionCaptureProxy* Proxy = Component->SceneProxy;
+
+		FTexture* EncodedHDRCubemapResource = Component->CachedEncodedHDRCubemap != nullptr ? Component->CachedEncodedHDRCubemap->Resource : nullptr;;
 		FMatrix Transform = Component->GetComponentTransform().ToMatrixWithScale();
 
 		ENQUEUE_RENDER_COMMAND(FUpdateTransformCommand)
-			([Scene, Proxy, Transform, bUsingPreviewCaptureData](FRHICommandListImmediate& RHICmdList)
+			([Scene, Proxy, Transform, bUsingPreviewCaptureData, EncodedHDRCubemapResource](FRHICommandListImmediate& RHICmdList)
 		{
+#if WITH_EDITOR
+			Proxy->EncodedHDRCubemap = EncodedHDRCubemapResource;
+#endif
 			if (Proxy->bUsingPreviewCaptureData)
 			{
 				FPlatformAtomics::InterlockedDecrement(&Scene->NumUnbuiltReflectionCaptures);
