Fixed an additional issue with EmitDepthTargets() and prepass

diff --git a/Engine/Source/Runtime/Renderer/Private/DeferredShadingRenderer.cpp b/Engine/Source/Runtime/Renderer/Private/DeferredShadingRenderer.cpp
--- a/Engine/Source/Runtime/Renderer/Private/DeferredShadingRenderer.cpp
+++ b/Engine/Source/Runtime/Renderer/Private/DeferredShadingRenderer.cpp
@@ -58,6 +58,8 @@
 #include "Lumen/Lumen.h"
 #include "Experimental/Containers/SherwoodHashTable.h"
 
+PRAGMA_DISABLE_OPTIMIZATION
+
 extern int32 GNaniteDebugFlags;
 extern int32 GNaniteShowStats;
 
@@ -1940,7 +1942,6 @@ void FDeferredShadingSceneRenderer::Render(FRHICommandListImmediate& RHICmdList)
 			for (int32 ViewIndex = 0; ViewIndex < Views.Num(); ViewIndex++)
 			{
 				const FViewInfo& View = Views[ViewIndex];
-				Nanite::FRasterResults& RasterResults = NaniteRasterResults[ViewIndex];
 
 				Nanite::FCullingContext CullingContext = Nanite::InitCullingContext(
 					GraphBuilder,
@@ -1969,16 +1970,18 @@ void FDeferredShadingSceneRenderer::Render(FRHICommandListImmediate& RHICmdList)
 					bExtractStats
 				);
 
+				Nanite::FRasterResults& RasterResults = NaniteRasterResults[ViewIndex];
+
 				if (bNeedsPrePass)
 				{
 					Nanite::EmitDepthTargets(
 						GraphBuilder,
 						*Scene,
 						Views[ViewIndex],
-						RasterResults.SOAStrides,
-						GraphBuilder.RegisterExternalBuffer(RasterResults.VisibleClustersSWHW),
-						GraphBuilder.RegisterExternalBuffer(RasterResults.ViewsBuffer),
-						GraphBuilder.RegisterExternalTexture(RasterResults.VisBuffer64),
+						CullingContext.SOAStrides,
+						CullingContext.VisibleClustersSWHW,
+						CullingContext.ViewsBuffer,
+						RasterContext.VisBuffer64,
 						RasterResults.MaterialDepth,
 						RasterResults.NaniteMask,
 						RasterResults.VelocityBuffer,
