Removed pragma that snuck in

diff --git a/Engine/Source/Runtime/Renderer/Private/DeferredShadingRenderer.cpp b/Engine/Source/Runtime/Renderer/Private/DeferredShadingRenderer.cpp
--- a/Engine/Source/Runtime/Renderer/Private/DeferredShadingRenderer.cpp
+++ b/Engine/Source/Runtime/Renderer/Private/DeferredShadingRenderer.cpp
@@ -58,8 +58,6 @@
 #include "Lumen/Lumen.h"
 #include "Experimental/Containers/SherwoodHashTable.h"
 
-PRAGMA_DISABLE_OPTIMIZATION
-
 extern int32 GNaniteDebugFlags;
 extern int32 GNaniteShowStats;
 
