Refactor of mobile pass / view uniform buffers to use static bindings.

Pass / view uniform buffers are no longer bound through mesh draw commands. Instead, they are bound using static bindings on the RDG pass. This is necessary because pass uniform buffers contain RDG resources, which need to be declared to RDG.

diff --git a/Engine/Shaders/Private/MobileBasePassPixelShader.usf b/Engine/Shaders/Private/MobileBasePassPixelShader.usf
--- a/Engine/Shaders/Private/MobileBasePassPixelShader.usf
+++ b/Engine/Shaders/Private/MobileBasePassPixelShader.usf
@@ -115,6 +115,10 @@ int NumDynamicPointLights;
 	#endif
 #endif
 
+#if DIRECTIONAL_LIGHT_CSM && !MATERIAL_SHADINGMODEL_SINGLELAYERWATER
+	uint UseCSM;
+#endif
+
 #if !FULLY_ROUGH
 half4 GetPlanarReflection(float3 WorldPosition, half3 WorldNormal, half Roughness)
 {
@@ -688,7 +692,7 @@ void Main(
 
 #if DIRECTIONAL_LIGHT_CSM && !MATERIAL_SHADINGMODEL_SINGLELAYERWATER
 	// Cascaded Shadow Map
-	if(MobileBasePass.UseCSM)
+	if (UseCSM)
 	{
 		half ShadowMap = MobileDirectionalLightCSM(MaterialParameters.ScreenPosition.xy, MaterialParameters.ScreenPosition.w);
 	#if MOVABLE_DIRECTIONAL_LIGHT
diff --git a/Engine/Shaders/Public/ShaderVersion.ush b/Engine/Shaders/Public/ShaderVersion.ush
--- a/Engine/Shaders/Public/ShaderVersion.ush
+++ b/Engine/Shaders/Public/ShaderVersion.ush
@@ -3,4 +3,4 @@
 // in Platform.ush (which should be included in any shader) it allows to invalidate the shader DDC.
 // 
 // If you are merging streams and there is a conflict with this GUID you should make a new GUID rather than taking one or the other.
-// GUID = 5952C866-E7D6-4116-ABDC-48ACA6A8D216
+// GUID = 5952C866-E7D6-4116-AGCC-48ACA6A8D218
\ No newline at end of file
diff --git a/Engine/Source/Runtime/Renderer/Private/BasePassRendering.cpp b/Engine/Source/Runtime/Renderer/Private/BasePassRendering.cpp
--- a/Engine/Source/Runtime/Renderer/Private/BasePassRendering.cpp
+++ b/Engine/Source/Runtime/Renderer/Private/BasePassRendering.cpp
@@ -539,8 +539,6 @@ bool GetBasePassShaders<FUniformLightMapPolicy>(
 	}
 }
 
-extern void SetupFogUniformParameters(FRDGBuilder* GraphBuilder, const FViewInfo& View, FFogUniformParameters& OutParameters);
-
 void SetupSharedBasePassParameters(
 	FRDGBuilder& GraphBuilder,
 	const FViewInfo& View,
@@ -548,13 +546,13 @@ void SetupSharedBasePassParameters(
 {
 	SharedParameters.Forward = View.ForwardLightingResources->ForwardLightData;
 
-	SetupFogUniformParameters(&GraphBuilder, View, SharedParameters.Fog);
+	SetupFogUniformParameters(GraphBuilder, View, SharedParameters.Fog);
 
 	if (View.IsInstancedStereoPass())
 	{
 		const FSceneView& RightEye = *View.Family->Views[1];
 		SharedParameters.ForwardISR = RightEye.ForwardLightingResources->ForwardLightData;
-		SetupFogUniformParameters(&GraphBuilder, (FViewInfo&)RightEye, SharedParameters.FogISR);
+		SetupFogUniformParameters(GraphBuilder, (FViewInfo&)RightEye, SharedParameters.FogISR);
 	}
 	else
 	{
diff --git a/Engine/Source/Runtime/Renderer/Private/CompositionLighting/PostProcessMeshDecals.cpp b/Engine/Source/Runtime/Renderer/Private/CompositionLighting/PostProcessMeshDecals.cpp
--- a/Engine/Source/Runtime/Renderer/Private/CompositionLighting/PostProcessMeshDecals.cpp
+++ b/Engine/Source/Runtime/Renderer/Private/CompositionLighting/PostProcessMeshDecals.cpp
@@ -175,10 +175,6 @@ FMeshDecalMeshProcessor::FMeshDecalMeshProcessor(const FScene* Scene,
 	, RenderTargetMode(InRenderTargetMode)
 {
 	PassDrawRenderState.SetDepthStencilState(TStaticDepthStencilState<false, CF_DepthNearOrEqual>::GetRHI());
-	if (FeatureLevel == ERHIFeatureLevel::ES3_1)
-	{
-		PassDrawRenderState.SetPassUniformBuffer(Scene->UniformBuffers.MobileTranslucentBasePassUniformBuffer);
-	}
 }
 
 void FMeshDecalMeshProcessor::AddMeshBatch(const FMeshBatch& RESTRICT MeshBatch, uint64 BatchElementMask, const FPrimitiveSceneProxy* RESTRICT PrimitiveSceneProxy, int32 StaticMeshId)
@@ -421,7 +417,8 @@ void RenderMeshDecalsMobile(FRDGBuilder& GraphBuilder, FRenderTargetBindingSlots
 		FDecalRenderingCommon::RTM_SceneColorAndGBufferWithNormal : 
 		FDecalRenderingCommon::RTM_SceneColor;
 
-	auto* PassParameters = GraphBuilder.AllocParameters<FRenderTargetParameters>();
+	auto* PassParameters = GraphBuilder.AllocParameters<FMobileBasePassParameters>();
+	PassParameters->View = View.GetShaderParameters();
 	PassParameters->RenderTargets = BasePassRenderTargets;
 
 	GraphBuilder.AddPass(RDG_EVENT_NAME("MeshDecalsPass"), PassParameters, ERDGPassFlags::Raster | ERDGPassFlags::SkipRenderPass,
diff --git a/Engine/Source/Runtime/Renderer/Private/DepthRendering.cpp b/Engine/Source/Runtime/Renderer/Private/DepthRendering.cpp
--- a/Engine/Source/Runtime/Renderer/Private/DepthRendering.cpp
+++ b/Engine/Source/Runtime/Renderer/Private/DepthRendering.cpp
@@ -602,6 +602,11 @@ void FDeferredShadingSceneRenderer::RenderPrePass(FRDGBuilder& GraphBuilder, FRD
 	}
 }
 
+BEGIN_SHADER_PARAMETER_STRUCT(FMobileDepthPassParameters, )
+	SHADER_PARAMETER_STRUCT_INCLUDE(FViewShaderParameters, View)
+	RENDER_TARGET_BINDING_SLOTS()
+END_SHADER_PARAMETER_STRUCT()
+
 void FMobileSceneRenderer::RenderPrePass(FRDGBuilder& GraphBuilder, FRenderTargetBindingSlots& BasePassRenderTargets)
 {
 	SCOPED_NAMED_EVENT(FMobileSceneRenderer_RenderPrePass, FColor::Emerald);
@@ -626,17 +631,17 @@ void FMobileSceneRenderer::RenderPrePass(FRDGBuilder& GraphBuilder, FRenderTarge
 			{
 				continue;
 			}
-			
-			auto* PassParameters = GraphBuilder.AllocParameters<FRenderTargetParameters>();
+
+			View.BeginRenderView();
+
+			auto* PassParameters = GraphBuilder.AllocParameters<FMobileDepthPassParameters>();
+			PassParameters->View = View.GetShaderParameters();
 			PassParameters->RenderTargets = BasePassRenderTargets;
 
 			GraphBuilder.AddPass(RDG_EVENT_NAME("RenderPrePass"), PassParameters, ERDGPassFlags::Raster | ERDGPassFlags::SkipRenderPass,
 				[this, &View, PassParameters](FRHICommandListImmediate& RHICmdList)
 			{
-				Scene->UniformBuffers.UpdateViewUniformBuffer(View);
-
 				SetStereoViewport(RHICmdList, View);
-
 				View.ParallelMeshDrawCommandPasses[EMeshPass::DepthPass].DispatchDraw(nullptr, RHICmdList);
 			});
 		}
diff --git a/Engine/Source/Runtime/Renderer/Private/EditorPrimitivesRendering.cpp b/Engine/Source/Runtime/Renderer/Private/EditorPrimitivesRendering.cpp
--- a/Engine/Source/Runtime/Renderer/Private/EditorPrimitivesRendering.cpp
+++ b/Engine/Source/Runtime/Renderer/Private/EditorPrimitivesRendering.cpp
@@ -162,7 +162,7 @@ bool FEditorPrimitivesBasePassMeshProcessor::ProcessMobileShadingPath(const FMes
 	ERasterizerFillMode MeshFillMode = ComputeMeshFillMode(MeshBatch, Material, OverrideSettings);
 	ERasterizerCullMode MeshCullMode = ComputeMeshCullMode(MeshBatch, Material, OverrideSettings);
 	
-	TMobileBasePassShaderElementData<LightMapPolicyType> ShaderElementData(nullptr);
+	TMobileBasePassShaderElementData<LightMapPolicyType> ShaderElementData(nullptr, false);
 	ShaderElementData.InitializeMeshMaterialData(ViewIfDynamicMeshCommand, PrimitiveSceneProxy, MeshBatch, StaticMeshId, false);
 
 	const FMeshDrawCommandSortKey SortKey = CalculateMeshStaticSortKey(BasePassShaders.VertexShader, BasePassShaders.PixelShader);
diff --git a/Engine/Source/Runtime/Renderer/Private/FogRendering.cpp b/Engine/Source/Runtime/Renderer/Private/FogRendering.cpp
--- a/Engine/Source/Runtime/Renderer/Private/FogRendering.cpp
+++ b/Engine/Source/Runtime/Renderer/Private/FogRendering.cpp
@@ -56,7 +56,7 @@ struct FHeightFogRenderingParameters
 	FVector4 LinearDepthMinMaxUV;
 };
 
-void SetupFogUniformParameters(FRDGBuilder* GraphBuilder, const FViewInfo& View, FFogUniformParameters& OutParameters)
+void SetupFogUniformParameters(FRDGBuilder& GraphBuilder, const FViewInfo& View, FFogUniformParameters& OutParameters)
 {
 	// Exponential Height Fog
 	{
@@ -89,25 +89,13 @@ void SetupFogUniformParameters(FRDGBuilder* GraphBuilder, const FViewInfo& View,
 		}
 		else
 		{
-			if (GraphBuilder)
-			{
-				OutParameters.IntegratedLightScattering = GBlackAlpha1VolumeTexture->GetRDG(*GraphBuilder);
-			}
-			else
-			{
-				OutParameters.IntegratedLightScattering = GBlackAlpha1VolumeTexture->GetPassthroughRDG();
-			}
+			OutParameters.IntegratedLightScattering = GBlackAlpha1VolumeTexture->GetRDG(GraphBuilder);
 			OutParameters.ApplyVolumetricFog = 0.0f;
 		}
 		OutParameters.IntegratedLightScatteringSampler = TStaticSamplerState<SF_Bilinear, AM_Clamp, AM_Clamp, AM_Clamp>::GetRHI();
 	}
 }
 
-void SetupFogUniformParameters(FRDGBuilder& GraphBuilder, const FViewInfo& View, FFogUniformParameters& OutParameters)
-{
-	SetupFogUniformParameters(&GraphBuilder, View, OutParameters);
-}
-
 TRDGUniformBufferRef<FFogUniformParameters> CreateFogUniformBuffer(FRDGBuilder& GraphBuilder, const FViewInfo& View)
 {
 	auto* FogStruct = GraphBuilder.AllocParameters<FFogUniformParameters>();
diff --git a/Engine/Source/Runtime/Renderer/Private/MobileBasePass.cpp b/Engine/Source/Runtime/Renderer/Private/MobileBasePass.cpp
--- a/Engine/Source/Runtime/Renderer/Private/MobileBasePass.cpp
+++ b/Engine/Source/Runtime/Renderer/Private/MobileBasePass.cpp
@@ -575,6 +575,11 @@ void TMobileBasePassPSPolicyParamType<FUniformLightMapPolicy>::GetShaderBindings
 		float CSMDebugValue = CVarsCSMDebugHint->GetValueOnRenderThread();
 		ShaderBindings.Add(CSMDebugHintParams, CSMDebugValue);
 	}
+
+	if (UseCSMParameter.IsBound())
+	{
+		ShaderBindings.Add(UseCSMParameter, ShaderElementData.bCanReceiveCSM ? 1 : 0);
+	}
 }
 
 FMobileBasePassMeshProcessor::FMobileBasePassMeshProcessor(
@@ -619,7 +624,7 @@ bool FMobileBasePassMeshProcessor::TryAddMeshBatch(const FMeshBatch& RESTRICT Me
 			const FLightSceneInfo* MobileDirectionalLight = MobileBasePass::GetDirectionalLightInfo(Scene, PrimitiveSceneProxy);
 			// Opaque meshes used for mobile pixel projected reflection could receive CSM in translucent pass.
 			ELightMapPolicyType LightmapPolicyType = MobileBasePass::SelectMeshLightmapPolicy(Scene, MeshBatch, PrimitiveSceneProxy, MobileDirectionalLight, ShadingModels, bCanReceiveCSM || (!bIsTranslucent && bIsUsingMobilePixelProjectedReflection), false, FeatureLevel, BlendMode);
-			bResult = Process(MeshBatch, BatchElementMask, StaticMeshId, PrimitiveSceneProxy, MaterialRenderProxy, Material, BlendMode, ShadingModels, LightmapPolicyType, MeshBatch.LCI);
+			bResult = Process(MeshBatch, BatchElementMask, StaticMeshId, PrimitiveSceneProxy, MaterialRenderProxy, Material, BlendMode, ShadingModels, LightmapPolicyType, bCanReceiveCSM, MeshBatch.LCI);
 		}
 	}
 	else
@@ -630,7 +635,7 @@ bool FMobileBasePassMeshProcessor::TryAddMeshBatch(const FMeshBatch& RESTRICT Me
 		{
 			const FLightSceneInfo* MobileDirectionalLight = MobileBasePass::GetDirectionalLightInfo(Scene, PrimitiveSceneProxy);
 			ELightMapPolicyType LightmapPolicyType = MobileBasePass::SelectMeshLightmapPolicy(Scene, MeshBatch, PrimitiveSceneProxy, MobileDirectionalLight, ShadingModels, bCanReceiveCSM, bUsesDeferredShading, FeatureLevel, BlendMode);
-			bResult = Process(MeshBatch, BatchElementMask, StaticMeshId, PrimitiveSceneProxy, MaterialRenderProxy, Material, BlendMode, ShadingModels, LightmapPolicyType, MeshBatch.LCI);
+			bResult = Process(MeshBatch, BatchElementMask, StaticMeshId, PrimitiveSceneProxy, MaterialRenderProxy, Material, BlendMode, ShadingModels, LightmapPolicyType, bCanReceiveCSM, MeshBatch.LCI);
 		}
 	}
 
@@ -670,6 +675,7 @@ bool FMobileBasePassMeshProcessor::Process(
 		EBlendMode BlendMode,
 		FMaterialShadingModelField ShadingModels,
 		const ELightMapPolicyType LightMapPolicyType,
+		const bool bCanReceiveCSM,
 		const FUniformLightMapPolicy::ElementDataType& RESTRICT LightMapElementData)
 {
 	TMeshProcessorShaders<
@@ -756,7 +762,7 @@ bool FMobileBasePassMeshProcessor::Process(
 	ERasterizerFillMode MeshFillMode = ComputeMeshFillMode(MeshBatch, MaterialResource, OverrideSettings);
 	ERasterizerCullMode MeshCullMode = ComputeMeshCullMode(MeshBatch, MaterialResource, OverrideSettings);
 
-	TMobileBasePassShaderElementData<FUniformLightMapPolicy> ShaderElementData(LightMapElementData);
+	TMobileBasePassShaderElementData<FUniformLightMapPolicy> ShaderElementData(LightMapElementData, bCanReceiveCSM);
 	ShaderElementData.InitializeMeshMaterialData(ViewIfDynamicMeshCommand, PrimitiveSceneProxy, MeshBatch, StaticMeshId, false);
 
 	BuildMeshDrawCommands(
@@ -777,8 +783,7 @@ bool FMobileBasePassMeshProcessor::Process(
 
 FMeshPassProcessor* CreateMobileBasePassProcessor(const FScene* Scene, const FSceneView* InViewIfDynamicMeshCommand, FMeshPassDrawListContext* InDrawListContext)
 {
-	FMeshPassProcessorRenderState PassDrawRenderState(Scene->UniformBuffers.ViewUniformBuffer, Scene->UniformBuffers.MobileOpaqueBasePassUniformBuffer);
-	PassDrawRenderState.SetInstancedViewUniformBuffer(Scene->UniformBuffers.InstancedViewUniformBuffer);
+	FMeshPassProcessorRenderState PassDrawRenderState;
 	PassDrawRenderState.SetBlendState(TStaticBlendStateWriteMask<CW_RGBA>::GetRHI());
 	PassDrawRenderState.SetDepthStencilAccess(Scene->DefaultBasePassDepthStencilAccess);
 	PassDrawRenderState.SetDepthStencilState(TStaticDepthStencilState<true, CF_DepthNearOrEqual>::GetRHI());
@@ -790,8 +795,7 @@ FMeshPassProcessor* CreateMobileBasePassProcessor(const FScene* Scene, const FSc
 
 FMeshPassProcessor* CreateMobileBasePassCSMProcessor(const FScene* Scene, const FSceneView* InViewIfDynamicMeshCommand, FMeshPassDrawListContext* InDrawListContext)
 {
-	FMeshPassProcessorRenderState PassDrawRenderState(Scene->UniformBuffers.ViewUniformBuffer, Scene->UniformBuffers.MobileCSMOpaqueBasePassUniformBuffer);
-	PassDrawRenderState.SetInstancedViewUniformBuffer(Scene->UniformBuffers.InstancedViewUniformBuffer);
+	FMeshPassProcessorRenderState PassDrawRenderState;
 	PassDrawRenderState.SetBlendState(TStaticBlendStateWriteMask<CW_RGBA>::GetRHI());
 	PassDrawRenderState.SetDepthStencilAccess(Scene->DefaultBasePassDepthStencilAccess);
 	PassDrawRenderState.SetDepthStencilState(TStaticDepthStencilState<true, CF_DepthNearOrEqual>::GetRHI());
@@ -803,8 +807,7 @@ FMeshPassProcessor* CreateMobileBasePassCSMProcessor(const FScene* Scene, const
 
 FMeshPassProcessor* CreateMobileTranslucencyStandardPassProcessor(const FScene* Scene, const FSceneView* InViewIfDynamicMeshCommand, FMeshPassDrawListContext* InDrawListContext)
 {
-	FMeshPassProcessorRenderState PassDrawRenderState(Scene->UniformBuffers.ViewUniformBuffer, Scene->UniformBuffers.MobileTranslucentBasePassUniformBuffer);
-	PassDrawRenderState.SetInstancedViewUniformBuffer(Scene->UniformBuffers.InstancedViewUniformBuffer);
+	FMeshPassProcessorRenderState PassDrawRenderState;
 	PassDrawRenderState.SetDepthStencilState(TStaticDepthStencilState<false, CF_DepthNearOrEqual>::GetRHI());
 	PassDrawRenderState.SetDepthStencilAccess(FExclusiveDepthStencil::DepthRead_StencilRead);
 	
@@ -816,8 +819,7 @@ FMeshPassProcessor* CreateMobileTranslucencyStandardPassProcessor(const FScene*
 
 FMeshPassProcessor* CreateMobileTranslucencyAfterDOFProcessor(const FScene* Scene, const FSceneView* InViewIfDynamicMeshCommand, FMeshPassDrawListContext* InDrawListContext)
 {
-	FMeshPassProcessorRenderState PassDrawRenderState(Scene->UniformBuffers.ViewUniformBuffer, Scene->UniformBuffers.MobileTranslucentBasePassUniformBuffer);
-	PassDrawRenderState.SetInstancedViewUniformBuffer(Scene->UniformBuffers.InstancedViewUniformBuffer);
+	FMeshPassProcessorRenderState PassDrawRenderState;
 	PassDrawRenderState.SetDepthStencilState(TStaticDepthStencilState<false, CF_DepthNearOrEqual>::GetRHI());
 	PassDrawRenderState.SetDepthStencilAccess(FExclusiveDepthStencil::DepthRead_StencilRead);
 
@@ -828,8 +830,7 @@ FMeshPassProcessor* CreateMobileTranslucencyAfterDOFProcessor(const FScene* Scen
 
 FMeshPassProcessor* CreateMobileTranslucencyAllPassProcessor(const FScene* Scene, const FSceneView* InViewIfDynamicMeshCommand, FMeshPassDrawListContext* InDrawListContext)
 {
-	FMeshPassProcessorRenderState PassDrawRenderState(Scene->UniformBuffers.ViewUniformBuffer, Scene->UniformBuffers.MobileTranslucentBasePassUniformBuffer);
-	PassDrawRenderState.SetInstancedViewUniformBuffer(Scene->UniformBuffers.InstancedViewUniformBuffer);
+	FMeshPassProcessorRenderState PassDrawRenderState;
 	PassDrawRenderState.SetDepthStencilState(TStaticDepthStencilState<false, CF_DepthNearOrEqual>::GetRHI());
 	PassDrawRenderState.SetDepthStencilAccess(FExclusiveDepthStencil::DepthRead_StencilRead);
 
diff --git a/Engine/Source/Runtime/Renderer/Private/MobileBasePassRendering.cpp b/Engine/Source/Runtime/Renderer/Private/MobileBasePassRendering.cpp
--- a/Engine/Source/Runtime/Renderer/Private/MobileBasePassRendering.cpp
+++ b/Engine/Source/Runtime/Renderer/Private/MobileBasePassRendering.cpp
@@ -43,7 +43,7 @@ static TAutoConsoleVariable<int32> CVarMobileMaxVisibleMovableSpotLightsShadow(
 	TEXT("The max number of visible spotlighs can cast shadow sorted by screen size, should be as less as possible for performance reason"),
 	ECVF_RenderThreadSafe);
 
-IMPLEMENT_GLOBAL_SHADER_PARAMETER_STRUCT(FMobileBasePassUniformParameters, "MobileBasePass");
+IMPLEMENT_STATIC_UNIFORM_BUFFER_STRUCT(FMobileBasePassUniformParameters, "MobileBasePass", SceneTextures);
 
 static TAutoConsoleVariable<int32> CVarMobileUseHWsRGBEncoding(
 	TEXT("r.Mobile.UseHWsRGBEncoding"),
@@ -137,24 +137,20 @@ FMobileBasePassMovableLightInfo::FMobileBasePassMovableLightInfo(const FPrimitiv
 	}
 }
 
-extern void SetupFogUniformParameters(FRDGBuilder* GraphBuilder, const FViewInfo& View, FFogUniformParameters& OutParameters);
-
 void SetupMobileBasePassUniformParameters(
-	FRHICommandListImmediate& RHICmdList, 
-	const FViewInfo& View, 
-	bool bTranslucentPass, 
-	bool bCanUseCSM,
+	FRDGBuilder& GraphBuilder,
+	const FViewInfo& View,
+	EMobileBasePass BasePass,
 	FMobileBasePassUniformParameters& BasePassParameters)
 {
-	SetupFogUniformParameters(nullptr, View, BasePassParameters.Fog);
+	SetupFogUniformParameters(GraphBuilder, View, BasePassParameters.Fog);
 
 	const FScene* Scene = View.Family->Scene ? View.Family->Scene->GetRenderScene() : nullptr;
 	const FPlanarReflectionSceneProxy* ReflectionSceneProxy = Scene ? Scene->GetForwardPassGlobalPlanarReflection() : nullptr;
 	SetupPlanarReflectionUniformParameters(View, ReflectionSceneProxy, BasePassParameters.PlanarReflection);
-	BasePassParameters.UseCSM = bCanUseCSM ? 1 : 0;
 
 	EMobileSceneTextureSetupMode SetupMode = EMobileSceneTextureSetupMode::None;
-	if (bTranslucentPass)
+	if (BasePass == EMobileBasePass::Translucent)
 	{
 		SetupMode |= EMobileSceneTextureSetupMode::SceneColor;
 	}
@@ -163,15 +159,16 @@ void SetupMobileBasePassUniformParameters(
 		SetupMode |= EMobileSceneTextureSetupMode::CustomDepth;
 	}
 
-	FSceneRenderTargets& SceneContext = FSceneRenderTargets::Get();
-	SetupMobileSceneTextureUniformParameters(SceneContext, SetupMode, BasePassParameters.SceneTextures);
+	const FRDGSystemTextures& SystemTextures = FRDGSystemTextures::Get(GraphBuilder);
+
+	SetupMobileSceneTextureUniformParameters(GraphBuilder, SetupMode, BasePassParameters.SceneTextures);
 
 	BasePassParameters.PreIntegratedGFTexture = GSystemTextures.PreintegratedGF->GetRenderTargetItem().ShaderResourceTexture;
 	BasePassParameters.PreIntegratedGFSampler = TStaticSamplerState<SF_Bilinear, AM_Clamp, AM_Clamp, AM_Clamp>::GetRHI();
 
 	if (GPixelProjectedReflectionMobileOutputs.IsValid())
 	{
-		if (bTranslucentPass)
+		if (BasePass == EMobileBasePass::Translucent)
 		{
 			BasePassParameters.PlanarReflection.PlanarReflectionTexture = GPixelProjectedReflectionMobileOutputs.PixelProjectedReflectionTexture->GetRenderTargetItem().ShaderResourceTexture;
 			if (GetMobilePixelProjectedReflectionQuality() <= EMobilePixelProjectedReflectionQuality::BestPerformance)
@@ -194,28 +191,26 @@ void SetupMobileBasePassUniformParameters(
 
 	BasePassParameters.EyeAdaptationBuffer = GetEyeAdaptationBuffer(View);
 
-	if (!bTranslucentPass && GAmbientOcclusionMobileOutputs.IsValid())
+	if (BasePass == EMobileBasePass::Opaque && GAmbientOcclusionMobileOutputs.IsValid())
 	{
-		BasePassParameters.AmbientOcclusionTexture = GAmbientOcclusionMobileOutputs.AmbientOcclusionTexture->GetRenderTargetItem().ShaderResourceTexture;
+		BasePassParameters.AmbientOcclusionTexture = GraphBuilder.RegisterExternalTexture(GAmbientOcclusionMobileOutputs.AmbientOcclusionTexture);
 	}
 	else
 	{
-		BasePassParameters.AmbientOcclusionTexture = GSystemTextures.WhiteDummy->GetRenderTargetItem().ShaderResourceTexture;
+		BasePassParameters.AmbientOcclusionTexture = SystemTextures.White;
 	}
 	BasePassParameters.AmbientOcclusionSampler = TStaticSamplerState<SF_Bilinear, AM_Clamp, AM_Clamp, AM_Clamp>::GetRHI();
 	BasePassParameters.AmbientOcclusionStaticFraction = FMath::Clamp(View.FinalPostProcessSettings.AmbientOcclusionStaticFraction, 0.0f, 1.0f);
 }
 
-void CreateMobileBasePassUniformBuffer(
-	FRHICommandListImmediate& RHICmdList, 
+TRDGUniformBufferRef<FMobileBasePassUniformParameters> CreateMobileBasePassUniformBuffer(
+	FRDGBuilder& GraphBuilder,
 	const FViewInfo& View,
-	bool bTranslucentPass,
-	bool bCanUseCSM,
-	TUniformBufferRef<FMobileBasePassUniformParameters>& BasePassUniformBuffer)
+	EMobileBasePass BasePass)
 {
-	FMobileBasePassUniformParameters BasePassParameters;
-	SetupMobileBasePassUniformParameters(RHICmdList, View, bTranslucentPass, bCanUseCSM, BasePassParameters);
-	BasePassUniformBuffer = TUniformBufferRef<FMobileBasePassUniformParameters>::CreateUniformBufferImmediate(BasePassParameters, UniformBuffer_SingleFrame);
+	FMobileBasePassUniformParameters* BasePassParameters = GraphBuilder.AllocParameters<FMobileBasePassUniformParameters>();
+	SetupMobileBasePassUniformParameters(GraphBuilder, View, BasePass, *BasePassParameters);
+	return GraphBuilder.CreateUniformBuffer(BasePassParameters);
 }
 
 void SetupMobileDirectionalLightUniformParameters(
@@ -318,31 +313,27 @@ void FMobileSceneRenderer::RenderMobileBasePass(FRDGBuilder& GraphBuilder, FRend
 			continue;
 		}
 
+		View.BeginRenderView();
+
+		UpdateDirectionalLightUniformBuffers(GraphBuilder, View);
+
 		auto* OpaqueBasePassParameters = GraphBuilder.AllocParameters<FMobileBasePassParameters>();
+		OpaqueBasePassParameters->View = View.GetShaderParameters();
 		OpaqueBasePassParameters->RenderTargets = BasePassRenderTargets;
-		OpaqueBasePassParameters->MobileBasePass = Scene->UniformBuffers.MobileOpaqueBasePassUniformBuffer;
+		OpaqueBasePassParameters->MobileBasePass = CreateMobileBasePassUniformBuffer(GraphBuilder, View, EMobileBasePass::Opaque);
 
 		GraphBuilder.AddPass(RDG_EVENT_NAME("RenderOpaqueBasePass"), OpaqueBasePassParameters, ERDGPassFlags::Raster | ERDGPassFlags::SkipRenderPass,
 			[this, &View, OpaqueBasePassParameters](FRHICommandListImmediate& RHICmdList)
 		{
-			if (Scene->UniformBuffers.UpdateViewUniformBuffer(View))
-			{
-				UpdateOpaqueBasePassUniformBuffer(RHICmdList, View);
-				UpdateDirectionalLightUniformBuffers(RHICmdList, View);
-			}
-
 			RHICmdList.SetViewport(View.ViewRect.Min.X, View.ViewRect.Min.Y, 0, View.ViewRect.Max.X, View.ViewRect.Max.Y, 1);
 			View.ParallelMeshDrawCommandPasses[EMeshPass::BasePass].DispatchDraw(nullptr, RHICmdList);
 		
 			// editor primitives
-			{
-				FMeshPassProcessorRenderState DrawRenderState(View, Scene->UniformBuffers.MobileOpaqueBasePassUniformBuffer);
-				DrawRenderState.SetBlendState(TStaticBlendStateWriteMask<CW_RGBA>::GetRHI());
-				DrawRenderState.SetDepthStencilAccess(Scene->DefaultBasePassDepthStencilAccess);
-				DrawRenderState.SetDepthStencilState(TStaticDepthStencilState<true, CF_DepthNearOrEqual>::GetRHI());
-				RenderMobileEditorPrimitives(RHICmdList, View, DrawRenderState);
-			}
-
+			FMeshPassProcessorRenderState DrawRenderState;
+			DrawRenderState.SetBlendState(TStaticBlendStateWriteMask<CW_RGBA>::GetRHI());
+			DrawRenderState.SetDepthStencilAccess(Scene->DefaultBasePassDepthStencilAccess);
+			DrawRenderState.SetDepthStencilState(TStaticDepthStencilState<true, CF_DepthNearOrEqual>::GetRHI());
+			RenderMobileEditorPrimitives(RHICmdList, View, DrawRenderState);
 		});
 	}
 }
diff --git a/Engine/Source/Runtime/Renderer/Private/MobileBasePassRendering.h b/Engine/Source/Runtime/Renderer/Private/MobileBasePassRendering.h
--- a/Engine/Source/Runtime/Renderer/Private/MobileBasePassRendering.h
+++ b/Engine/Source/Runtime/Renderer/Private/MobileBasePassRendering.h
@@ -25,7 +25,6 @@
 #include "SkyAtmosphereRendering.h"
 
 BEGIN_GLOBAL_SHADER_PARAMETER_STRUCT(FMobileBasePassUniformParameters, )
-	SHADER_PARAMETER(int32, UseCSM)
 	SHADER_PARAMETER(float, AmbientOcclusionStaticFraction)
 	SHADER_PARAMETER_STRUCT(FFogUniformParameters, Fog)
 	SHADER_PARAMETER_STRUCT(FPlanarReflectionUniformParameters, PlanarReflection) // Single global planar reflection for the forward pass.
@@ -33,28 +32,32 @@ BEGIN_GLOBAL_SHADER_PARAMETER_STRUCT(FMobileBasePassUniformParameters, )
 	SHADER_PARAMETER_TEXTURE(Texture2D, PreIntegratedGFTexture)
 	SHADER_PARAMETER_SAMPLER(SamplerState, PreIntegratedGFSampler)
 	SHADER_PARAMETER_SRV(Buffer<float4>, EyeAdaptationBuffer)
-	SHADER_PARAMETER_TEXTURE(Texture2D, AmbientOcclusionTexture)
+	SHADER_PARAMETER_RDG_TEXTURE(Texture2D, AmbientOcclusionTexture)
 	SHADER_PARAMETER_SAMPLER(SamplerState, AmbientOcclusionSampler)
 END_GLOBAL_SHADER_PARAMETER_STRUCT()
 
 BEGIN_SHADER_PARAMETER_STRUCT(FMobileBasePassParameters, )
-	SHADER_PARAMETER_STRUCT_REF(FMobileBasePassUniformParameters, MobileBasePass)
+	SHADER_PARAMETER_STRUCT_INCLUDE(FViewShaderParameters, View)
+	SHADER_PARAMETER_RDG_UNIFORM_BUFFER(FMobileBasePassUniformParameters, MobileBasePass)
 	RENDER_TARGET_BINDING_SLOTS()
 END_SHADER_PARAMETER_STRUCT()
 
+enum class EMobileBasePass
+{
+	Opaque,
+	Translucent
+};
+
 extern void SetupMobileBasePassUniformParameters(
-	FRHICommandListImmediate& RHICmdList,
+	FRDGBuilder& GraphBuilder,
 	const FViewInfo& View,
-	bool bTranslucentPass,
-	bool bCanUseCSM,
+	EMobileBasePass BasePass,
 	FMobileBasePassUniformParameters& BasePassParameters);
 
-extern void CreateMobileBasePassUniformBuffer(
-	FRHICommandListImmediate& RHICmdList,
+extern TRDGUniformBufferRef<FMobileBasePassUniformParameters> CreateMobileBasePassUniformBuffer(
+	FRDGBuilder& GraphBuilder,
 	const FViewInfo& View,
-	bool bTranslucentPass,
-	bool bCanUseCSM,
-	TUniformBufferRef<FMobileBasePassUniformParameters>& BasePassUniformBuffer);
+	EMobileBasePass BasePass);
 
 extern void SetupMobileDirectionalLightUniformParameters(
 	const FScene& Scene,
@@ -105,11 +108,14 @@ template<typename LightMapPolicyType>
 class TMobileBasePassShaderElementData : public FMeshMaterialShaderElementData
 {
 public:
-	TMobileBasePassShaderElementData(const typename LightMapPolicyType::ElementDataType& InLightMapPolicyElementData) :
-		LightMapPolicyElementData(InLightMapPolicyElementData)
+	TMobileBasePassShaderElementData(const typename LightMapPolicyType::ElementDataType& InLightMapPolicyElementData, bool bInCanReceiveCSM)
+		: LightMapPolicyElementData(InLightMapPolicyElementData)
+		, bCanReceiveCSM(bInCanReceiveCSM)
 	{}
 
 	typename LightMapPolicyType::ElementDataType LightMapPolicyElementData;
+
+	const bool bCanReceiveCSM;
 };
 
 /**
@@ -261,8 +267,9 @@ public:
 		HQReflectionCaptureBoxScalesArray.Bind(Initializer.ParameterMap, TEXT("CaptureBoxScalesArray"));
 
 		NumDynamicPointLightsParameter.Bind(Initializer.ParameterMap, TEXT("NumDynamicPointLights"));
-						
+
 		CSMDebugHintParams.Bind(Initializer.ParameterMap, TEXT("CSMDebugHint"));
+		UseCSMParameter.Bind(Initializer.ParameterMap, TEXT("UseCSM"));
 	}
 
 	TMobileBasePassPSPolicyParamType() {}
@@ -283,6 +290,7 @@ private:
 	LAYOUT_FIELD(FShaderParameter, NumDynamicPointLightsParameter);
 
 	LAYOUT_FIELD(FShaderParameter, CSMDebugHintParams);
+	LAYOUT_FIELD(FShaderParameter, UseCSMParameter);
 	
 public:
 	void GetShaderBindings(
@@ -480,6 +488,7 @@ private:
 		EBlendMode BlendMode,
 		FMaterialShadingModelField ShadingModels,
 		const ELightMapPolicyType LightMapPolicyType,
+		const bool bCanReceiveCSM,
 		const FUniformLightMapPolicy::ElementDataType& RESTRICT LightMapElementData);
 
 	const ETranslucencyPass::Type TranslucencyPassType;
diff --git a/Engine/Source/Runtime/Renderer/Private/MobileShadingRenderer.cpp b/Engine/Source/Runtime/Renderer/Private/MobileShadingRenderer.cpp
--- a/Engine/Source/Runtime/Renderer/Private/MobileShadingRenderer.cpp
+++ b/Engine/Source/Runtime/Renderer/Private/MobileShadingRenderer.cpp
@@ -360,8 +360,6 @@ void FMobileSceneRenderer::InitViews(FRDGBuilder& GraphBuilder)
 	// never keep MSAA depth
 	bKeepDepthContent = (NumMSAASamples > 1 ? false : bKeepDepthContent);
 
-	// Initialize global system textures (pass-through if already initialized).
-	GSystemTextures.InitializeTextures(RHICmdList, FeatureLevel);
 	FSceneRenderTargets& SceneContext = FSceneRenderTargets::Get();
 
 	// Allocate the maximum scene render target space for the current view family.
@@ -491,17 +489,6 @@ void FMobileSceneRenderer::InitViews(FRDGBuilder& GraphBuilder)
 		}
 	}
 
-	// update buffers used in cached mesh path
-	// in case there are multiple views, these buffers will be updated before rendering each view
-	if (Views.Num() > 0)
-	{
-		const FViewInfo& View = Views[0];
-		// We want to wait for the extension jobs only when the view is being actually rendered for the first time
-		Scene->UniformBuffers.UpdateViewUniformBuffer(View, false);
-		UpdateOpaqueBasePassUniformBuffer(RHICmdList, View);
-		UpdateTranslucentBasePassUniformBuffer(RHICmdList, View);
-		UpdateDirectionalLightUniformBuffers(RHICmdList, View);
-	}
 	if (bDeferredShading)
 	{
 		SetupSceneReflectionCaptureBuffer(RHICmdList);
@@ -567,10 +554,13 @@ void FMobileSceneRenderer::Render(FRHICommandListImmediate& RHICmdList)
 	FRHICommandListExecutor::GetImmediateCommandList().PollOcclusionQueries();
 	RHICmdList.ImmediateFlush(EImmediateFlushType::DispatchToRHIThread);
 
+	// Initialize global system textures (pass-through if already initialized).
+	GSystemTextures.InitializeTextures(RHICmdList, FeatureLevel);
+	FRDGSystemTextures::Create(GraphBuilder);
+
 	// Find the visible primitives and prepare targets and buffers for rendering
 	InitViews(GraphBuilder);
 
-	
 	if (GRHINeedsExtraDeletionLatency || !GRHICommandList.Bypass())
 	{
 		QUICK_SCOPE_CYCLE_COUNTER(STAT_FMobileSceneRenderer_PostInitViewsFlushDel);
@@ -1412,32 +1402,25 @@ void FMobileSceneRenderer::ConditionalResolveSceneDepth(FRDGBuilder& GraphBuilde
 	}
 }
 
-void FMobileSceneRenderer::UpdateOpaqueBasePassUniformBuffer(FRHICommandListImmediate& RHICmdList, const FViewInfo& View)
-{
-	FMobileBasePassUniformParameters Parameters;
-	SetupMobileBasePassUniformParameters(RHICmdList, View, false, false, Parameters);
-	Scene->UniformBuffers.MobileOpaqueBasePassUniformBuffer.UpdateUniformBufferImmediate(Parameters);
-	SetupMobileBasePassUniformParameters(RHICmdList, View, false, true, Parameters);
-	Scene->UniformBuffers.MobileCSMOpaqueBasePassUniformBuffer.UpdateUniformBufferImmediate(Parameters);	
-}
-
-void FMobileSceneRenderer::UpdateTranslucentBasePassUniformBuffer(FRHICommandListImmediate& RHICmdList, const FViewInfo& View)
+void FMobileSceneRenderer::UpdateDirectionalLightUniformBuffers(FRDGBuilder& GraphBuilder, const FViewInfo& View)
 {
-	FMobileBasePassUniformParameters Parameters;
-	SetupMobileBasePassUniformParameters(RHICmdList, View, true, false, Parameters);
-	Scene->UniformBuffers.MobileTranslucentBasePassUniformBuffer.UpdateUniformBufferImmediate(Parameters);
-}
-
-void FMobileSceneRenderer::UpdateDirectionalLightUniformBuffers(FRHICommandListImmediate& RHICmdList, const FViewInfo& View)
-{
-	bool bDynamicShadows = ViewFamily.EngineShowFlags.DynamicShadows;
-	// Fill in the other entries based on the lights
-	for (int32 ChannelIdx = 0; ChannelIdx < UE_ARRAY_COUNT(Scene->MobileDirectionalLights); ChannelIdx++)
+	if (CachedView == &View)
 	{
-		FMobileDirectionalLightShaderParameters Params;
-		SetupMobileDirectionalLightUniformParameters(*Scene, View, VisibleLightInfos, ChannelIdx, bDynamicShadows, Params);
-		Scene->UniformBuffers.MobileDirectionalLightUniformBuffers[ChannelIdx + 1].UpdateUniformBufferImmediate(Params);
+		return;
 	}
+	CachedView = &View;
+
+	AddPass(GraphBuilder, [this, &View](FRHICommandList&)
+	{
+		const bool bDynamicShadows = ViewFamily.EngineShowFlags.DynamicShadows;
+		// Fill in the other entries based on the lights
+		for (int32 ChannelIdx = 0; ChannelIdx < UE_ARRAY_COUNT(Scene->MobileDirectionalLights); ChannelIdx++)
+		{
+			FMobileDirectionalLightShaderParameters Params;
+			SetupMobileDirectionalLightUniformParameters(*Scene, View, VisibleLightInfos, ChannelIdx, bDynamicShadows, Params);
+			Scene->UniformBuffers.MobileDirectionalLightUniformBuffers[ChannelIdx + 1].UpdateUniformBufferImmediate(Params);
+		}
+	});
 }
 
 void FMobileSceneRenderer::UpdateSkyReflectionUniformBuffer()
diff --git a/Engine/Source/Runtime/Renderer/Private/MobileTranslucentRendering.cpp b/Engine/Source/Runtime/Renderer/Private/MobileTranslucentRendering.cpp
--- a/Engine/Source/Runtime/Renderer/Private/MobileTranslucentRendering.cpp
+++ b/Engine/Source/Runtime/Renderer/Private/MobileTranslucentRendering.cpp
@@ -51,9 +51,13 @@ void FMobileSceneRenderer::RenderTranslucency(FRDGBuilder& GraphBuilder, FRender
 				continue;
 			}
 
+			View.BeginRenderView();
+			UpdateDirectionalLightUniformBuffers(GraphBuilder, View);
+
 			auto* TranslucencyBasePassParameters = GraphBuilder.AllocParameters<FMobileBasePassParameters>();
+			TranslucencyBasePassParameters->View = View.GetShaderParameters();
+			TranslucencyBasePassParameters->MobileBasePass = CreateMobileBasePassUniformBuffer(GraphBuilder, View, EMobileBasePass::Translucent);
 			TranslucencyBasePassParameters->RenderTargets = BasePassRenderTargets;
-			TranslucencyBasePassParameters->MobileBasePass = Scene->UniformBuffers.MobileTranslucentBasePassUniformBuffer;
 
 			GraphBuilder.AddPass(RDG_EVENT_NAME("RenderTranslucencyBasePass"), TranslucencyBasePassParameters, ERDGPassFlags::Raster | ERDGPassFlags::SkipRenderPass,
 				[this, &View, TranslucencyPass](FRHICommandListImmediate& RHICmdList)
@@ -62,12 +66,6 @@ void FMobileSceneRenderer::RenderTranslucency(FRDGBuilder& GraphBuilder, FRender
 
 				if (!View.Family->UseDebugViewPS())
 				{
-					if (Scene->UniformBuffers.UpdateViewUniformBuffer(View))
-					{
-						UpdateTranslucentBasePassUniformBuffer(RHICmdList, View);
-						UpdateDirectionalLightUniformBuffers(RHICmdList, View);
-					}
-
 					const EMeshPass::Type MeshPass = TranslucencyPassToMeshPass(TranslucencyPass);
 					View.ParallelMeshDrawCommandPasses[MeshPass].DispatchDraw(nullptr, RHICmdList);
 				}
@@ -85,8 +83,12 @@ void FMobileSceneRenderer::RenderInverseOpacity(FRDGBuilder& GraphBuilder, const
 	FRDGTextureMSAA SceneColorMSAA = RegisterExternalTextureMSAA(GraphBuilder, SceneContext.GetSceneColor());
 	FRDGTextureMSAA SceneDepthMSAA = RegisterExternalTextureMSAA(GraphBuilder, SceneContext.SceneDepthZ);
 
+	View.BeginRenderView();
+	UpdateDirectionalLightUniformBuffers(GraphBuilder, View);
+
 	auto* InverseOpacityParameters = GraphBuilder.AllocParameters<FMobileBasePassParameters>();
-	InverseOpacityParameters->MobileBasePass = Scene->UniformBuffers.MobileTranslucentBasePassUniformBuffer;
+	InverseOpacityParameters->View = View.GetShaderParameters();
+	InverseOpacityParameters->MobileBasePass = CreateMobileBasePassUniformBuffer(GraphBuilder, View, EMobileBasePass::Translucent);
 	InverseOpacityParameters->RenderTargets[0] = FRenderTargetBinding(SceneColorMSAA.Target, SceneColorMSAA.Resolve, ERenderTargetLoadAction::EClear);
 	InverseOpacityParameters->RenderTargets.DepthStencil = FDepthStencilBinding(SceneDepthMSAA.Target, ERenderTargetLoadAction::EClear, FExclusiveDepthStencil::DepthWrite_StencilWrite);
 	// Opacity could fetch depth as we use exactly the same shaders as in base pass
@@ -95,12 +97,6 @@ void FMobileSceneRenderer::RenderInverseOpacity(FRDGBuilder& GraphBuilder, const
 	GraphBuilder.AddPass(RDG_EVENT_NAME("InverseOpacityPass"), InverseOpacityParameters, ERDGPassFlags::Raster,
 		[this, &View](FRHICommandListImmediate& RHICmdList)
 	{
-		if (Scene->UniformBuffers.UpdateViewUniformBuffer(View))
-		{
-			UpdateTranslucentBasePassUniformBuffer(RHICmdList, View);
-			UpdateDirectionalLightUniformBuffers(RHICmdList, View);
-		}
-
 		// Mobile multi-view is not side by side stereo
 		const FViewInfo& TranslucentViewport = (View.bIsMobileMultiViewEnabled) ? Views[0] : View;
 		RHICmdList.SetViewport(TranslucentViewport.ViewRect.Min.X, TranslucentViewport.ViewRect.Min.Y, 0.0f, TranslucentViewport.ViewRect.Max.X, TranslucentViewport.ViewRect.Max.Y, 1.0f);
@@ -121,7 +117,7 @@ void FMobileSceneRenderer::RenderInverseOpacity(FRDGBuilder& GraphBuilder, const
 // This pass is registered only when we render to scene capture, see UpdateSceneCaptureContentMobile_RenderThread()
 FMeshPassProcessor* CreateMobileInverseOpacityPassProcessor(const FScene* Scene, const FSceneView* InViewIfDynamicMeshCommand, FMeshPassDrawListContext* InDrawListContext)
 {
-	FMeshPassProcessorRenderState PassDrawRenderState(Scene->UniformBuffers.ViewUniformBuffer, Scene->UniformBuffers.MobileTranslucentBasePassUniformBuffer);
+	FMeshPassProcessorRenderState PassDrawRenderState;
 	PassDrawRenderState.SetDepthStencilState(TStaticDepthStencilState<false, CF_Always>::GetRHI());
 	PassDrawRenderState.SetBlendState(TStaticBlendState<CW_ALPHA, BO_Add, BF_Zero, BF_One, BO_Add, BF_Zero, BF_InverseSourceAlpha>::GetRHI());
 	PassDrawRenderState.SetDepthStencilAccess(FExclusiveDepthStencil::DepthRead_StencilRead);
diff --git a/Engine/Source/Runtime/Renderer/Private/PostProcess/PostProcessCompositeEditorPrimitives.cpp b/Engine/Source/Runtime/Renderer/Private/PostProcess/PostProcessCompositeEditorPrimitives.cpp
--- a/Engine/Source/Runtime/Renderer/Private/PostProcess/PostProcessCompositeEditorPrimitives.cpp
+++ b/Engine/Source/Runtime/Renderer/Private/PostProcess/PostProcessCompositeEditorPrimitives.cpp
@@ -233,6 +233,7 @@ const FViewInfo* UpdateEditorPrimitiveView(
 BEGIN_SHADER_PARAMETER_STRUCT(FEditorPrimitivesPassParameters, )
 	SHADER_PARAMETER_STRUCT_REF(FViewUniformShaderParameters, View)
 	SHADER_PARAMETER_RDG_UNIFORM_BUFFER(FOpaqueBasePassUniformParameters, BasePass)
+	SHADER_PARAMETER_RDG_UNIFORM_BUFFER(FMobileBasePassUniformParameters, MobileBasePass)
 	RENDER_TARGET_BINDING_SLOTS()
 END_SHADER_PARAMETER_STRUCT()
 
@@ -319,6 +320,10 @@ FScreenPassTexture AddEditorPrimitivePass(
 			{
 				PassParameters->BasePass = CreateOpaqueBasePassUniformBuffer(GraphBuilder, *EditorView, 0);
 			}
+			else
+			{
+				PassParameters->MobileBasePass = CreateMobileBasePassUniformBuffer(GraphBuilder, *EditorView, EMobileBasePass::Translucent);
+			}
 
 			GraphBuilder.AddPass(
 				RDG_EVENT_NAME("EditorPrimitives"),
@@ -328,17 +333,7 @@ FScreenPassTexture AddEditorPrimitivePass(
 			{
 				RHICmdList.SetViewport(EditorPrimitivesViewport.Rect.Min.X, EditorPrimitivesViewport.Rect.Min.Y, 0.0f, EditorPrimitivesViewport.Rect.Max.X, EditorPrimitivesViewport.Rect.Max.Y, 1.0f);
 
-				TUniformBufferRef<FMobileBasePassUniformParameters> MobileBasePassUniformBuffer;
-				FRHIUniformBuffer* BasePassUniformBuffer = nullptr;
-
-				if (BasePassType == FEditorPrimitiveInputs::EBasePassType::Mobile)
-				{
-					CreateMobileBasePassUniformBuffer(RHICmdList, *EditorView, true, false, MobileBasePassUniformBuffer);
-					BasePassUniformBuffer = MobileBasePassUniformBuffer;
-				}
-
 				FMeshPassProcessorRenderState DrawRenderState;
-				DrawRenderState.SetPassUniformBuffer(BasePassUniformBuffer);
 				DrawRenderState.SetDepthStencilAccess(FExclusiveDepthStencil::DepthWrite_StencilWrite);
 				DrawRenderState.SetBlendState(TStaticBlendStateWriteMask<CW_RGBA>::GetRHI());
 
diff --git a/Engine/Source/Runtime/Renderer/Private/PostProcess/SceneRenderTargets.cpp b/Engine/Source/Runtime/Renderer/Private/PostProcess/SceneRenderTargets.cpp
--- a/Engine/Source/Runtime/Renderer/Private/PostProcess/SceneRenderTargets.cpp
+++ b/Engine/Source/Runtime/Renderer/Private/PostProcess/SceneRenderTargets.cpp
@@ -1981,26 +1981,25 @@ EMobileSceneTextureSetupMode Translate(ESceneTextureSetupMode InSetupMode)
 
 IMPLEMENT_STATIC_UNIFORM_BUFFER_STRUCT(FMobileSceneTextureUniformParameters, "MobileSceneTextures", SceneTextures);
 
-static void SetupMobileSceneTextureUniformParameters(
-	FRDGBuilder* GraphBuilder,
-	const FSceneRenderTargets& SceneContext,
+void SetupMobileSceneTextureUniformParameters(
+	FRDGBuilder& GraphBuilder,
 	EMobileSceneTextureSetupMode SetupMode,
 	FMobileSceneTextureUniformParameters& SceneTextureParameters)
 {
+	const FSceneRenderTargets& SceneContext = FSceneRenderTargets::Get();
+
 	const auto GetRDG = [&](const TRefCountPtr<IPooledRenderTarget>& PooledRenderTarget)
 	{
-		return RegisterExternalOrPassthroughTexture(GraphBuilder, PooledRenderTarget);
+		return GraphBuilder.RegisterExternalTexture(PooledRenderTarget);
 	};
 
-	FRDGTextureRef BlackDefault2D = GetRDG(GSystemTextures.BlackDummy);
-	FRDGTextureRef MaxFP16Depth2D = GetRDG(GSystemTextures.MaxFP16Depth);
-	FRDGTextureRef DepthDefault = GetRDG(GSystemTextures.DepthDummy);
+	const FRDGSystemTextures& SystemTextures = FRDGSystemTextures::Get(GraphBuilder);
 
 	const bool bUseSceneTextures = EnumHasAnyFlags(SetupMode, EMobileSceneTextureSetupMode::SceneColor);
 
-	SceneTextureParameters.SceneColorTexture = BlackDefault2D;
+	SceneTextureParameters.SceneColorTexture = SystemTextures.Black;
 	SceneTextureParameters.SceneColorTextureSampler = TStaticSamplerState<SF_Point, AM_Clamp, AM_Clamp, AM_Clamp>::GetRHI();
-	SceneTextureParameters.SceneDepthTexture = DepthDefault;
+	SceneTextureParameters.SceneDepthTexture = SystemTextures.DepthDummy;
 	SceneTextureParameters.SceneDepthTextureSampler = TStaticSamplerState<SF_Point, AM_Clamp, AM_Clamp, AM_Clamp>::GetRHI();
 
 	if (bUseSceneTextures)
@@ -2013,8 +2012,8 @@ static void SetupMobileSceneTextureUniformParameters(
 		}
 	}
 
-	FRDGTextureRef CustomDepth = MaxFP16Depth2D;
-	FRDGTextureRef CustomStencil = BlackDefault2D;
+	FRDGTextureRef CustomDepth = SystemTextures.MaxFP16Depth;
+	FRDGTextureRef CustomStencil = SystemTextures.Black;
 
 	const bool bUseCustomDepth = EnumHasAnyFlags(SetupMode, EMobileSceneTextureSetupMode::CustomDepth) && SceneContext.bCustomDepthIsValid;
 
@@ -2067,17 +2066,17 @@ void SetupMobileSceneTextureUniformParameters(
 	EMobileSceneTextureSetupMode SetupMode,
 	FMobileSceneTextureUniformParameters& SceneTextureParameters)
 {
-	SetupMobileSceneTextureUniformParameters(nullptr, SceneContext, SetupMode, SceneTextureParameters);
+	// Deprecated.
+	checkNoEntry();
 }
 
 TUniformBufferRef<FMobileSceneTextureUniformParameters> CreateMobileSceneTextureUniformBuffer(
 	FRHIComputeCommandList& RHICmdList,
 	EMobileSceneTextureSetupMode SetupMode)
 {
-	FMobileSceneTextureUniformParameters SceneTextures;
-	FSceneRenderTargets& SceneContext = FSceneRenderTargets::Get();
-	SetupMobileSceneTextureUniformParameters(nullptr, SceneContext, SetupMode, SceneTextures);
-	return TUniformBufferRef<FMobileSceneTextureUniformParameters>::CreateUniformBufferImmediate(SceneTextures, EUniformBufferUsage::UniformBuffer_SingleFrame);
+	// Deprecated
+	checkNoEntry();
+	return {};
 }
 
 TRDGUniformBufferRef<FMobileSceneTextureUniformParameters> CreateMobileSceneTextureUniformBuffer(
@@ -2085,8 +2084,7 @@ TRDGUniformBufferRef<FMobileSceneTextureUniformParameters> CreateMobileSceneText
 	EMobileSceneTextureSetupMode SetupMode)
 {
 	FMobileSceneTextureUniformParameters* SceneTextures = GraphBuilder.AllocParameters<FMobileSceneTextureUniformParameters>();
-	FSceneRenderTargets& SceneContext = FSceneRenderTargets::Get();
-	SetupMobileSceneTextureUniformParameters(&GraphBuilder, SceneContext, SetupMode, *SceneTextures);
+	SetupMobileSceneTextureUniformParameters(GraphBuilder, SetupMode, *SceneTextures);
 	return GraphBuilder.CreateUniformBuffer(SceneTextures);
 }
 
@@ -2132,23 +2130,6 @@ void SetupMobileSceneTextureUniformParameters(
 	bool bCustomDepthIsValid,
 	FMobileSceneTextureUniformParameters& SceneTextureParameters)
 {
-	EMobileSceneTextureSetupMode SetupMode = EMobileSceneTextureSetupMode::None;
-	if (bSceneTexturesValid)
-	{
-		SetupMode |= EMobileSceneTextureSetupMode::SceneColor;
-	}
-	if (bCustomDepthIsValid)
-	{
-		SetupMode |= EMobileSceneTextureSetupMode::CustomDepth;
-	}
-	SetupMobileSceneTextureUniformParameters(nullptr, SceneContext, SetupMode, SceneTextureParameters);
-}
-
-void SetupMobileSceneTextureUniformParameters(
-	FRDGBuilder& GraphBuilder,
-	EMobileSceneTextureSetupMode SetupMode,
-	FMobileSceneTextureUniformParameters& SceneTextureParameters)
-{
-	FSceneRenderTargets& SceneContext = FSceneRenderTargets::Get();
-	SetupMobileSceneTextureUniformParameters(&GraphBuilder, SceneContext, SetupMode, SceneTextureParameters);
+	// Deprecated
+	checkNoEntry();
 }
\ No newline at end of file
diff --git a/Engine/Source/Runtime/Renderer/Private/Renderer.cpp b/Engine/Source/Runtime/Renderer/Private/Renderer.cpp
--- a/Engine/Source/Runtime/Renderer/Private/Renderer.cpp
+++ b/Engine/Source/Runtime/Renderer/Private/Renderer.cpp
@@ -262,10 +262,8 @@ void FRendererModule::DrawTileMesh(FRHICommandListImmediate& RHICmdList, FMeshPa
 			}
 			else // Mobile
 			{
-				TUniformBufferRef<FMobileBasePassUniformParameters> MobileBasePassUniformBuffer;
-				CreateMobileBasePassUniformBuffer(RHICmdList, View, true, false, MobileBasePassUniformBuffer);
-				DrawRenderState.SetPassUniformBuffer(MobileBasePassUniformBuffer);
-				
+				PassParameters->MobileBasePass = CreateMobileBasePassUniformBuffer(GraphBuilder, View, EMobileBasePass::Translucent);
+
 				DrawDynamicMeshPass(View, RHICmdList,
 					[Scene, &View, &DrawRenderState, &Mesh](FDynamicPassMeshDrawListContext* DynamicMeshPassContext)
 				{
@@ -337,9 +335,7 @@ void FRendererModule::DrawTileMesh(FRHICommandListImmediate& RHICmdList, FMeshPa
 				}
 				else // Mobile
 				{
-					TUniformBufferRef<FMobileBasePassUniformParameters> MobileBasePassUniformBuffer;
-					CreateMobileBasePassUniformBuffer(RHICmdList, View, false, true, MobileBasePassUniformBuffer);
-					DrawRenderState.SetPassUniformBuffer(MobileBasePassUniformBuffer);
+					PassParameters->MobileBasePass = CreateMobileBasePassUniformBuffer(GraphBuilder, View, EMobileBasePass::Opaque);
 					
 					DrawDynamicMeshPass(View, RHICmdList,
 						[Scene, &View, &DrawRenderState, &Mesh](FDynamicPassMeshDrawListContext* DynamicMeshPassContext)
diff --git a/Engine/Source/Runtime/Renderer/Private/RendererScene.cpp b/Engine/Source/Runtime/Renderer/Private/RendererScene.cpp
--- a/Engine/Source/Runtime/Renderer/Private/RendererScene.cpp
+++ b/Engine/Source/Runtime/Renderer/Private/RendererScene.cpp
@@ -897,11 +897,6 @@ FORCEINLINE static void VerifyProperPIEScene(UPrimitiveComponent* Component, UWo
 void FPersistentUniformBuffers::Clear()
 {
 	ViewUniformBuffer.SafeRelease();
-	InstancedViewUniformBuffer.SafeRelease();
-	MobileOpaqueBasePassUniformBuffer.SafeRelease();
-	MobileTranslucentBasePassUniformBuffer.SafeRelease();
-
-
 	for (auto& UniformBuffer : MobileDirectionalLightUniformBuffers)
 	{
 		UniformBuffer.SafeRelease();
@@ -910,24 +905,17 @@ void FPersistentUniformBuffers::Clear()
 
 	Initialize();
 }
+
 void FPersistentUniformBuffers::Initialize()
 {
 	FViewUniformShaderParameters ViewUniformBufferParameters;
 	ViewUniformBuffer = TUniformBufferRef<FViewUniformShaderParameters>::CreateUniformBufferImmediate(ViewUniformBufferParameters, UniformBuffer_MultiFrame, EUniformBufferValidation::None);
 
-	FInstancedViewUniformShaderParameters InstancedViewUniformBufferParameters;
-	InstancedViewUniformBuffer = TUniformBufferRef<FInstancedViewUniformShaderParameters>::CreateUniformBufferImmediate(InstancedViewUniformBufferParameters, UniformBuffer_MultiFrame, EUniformBufferValidation::None);
-
 	FNaniteUniformParameters NaniteUniformBufferParameters;
 	NaniteUniformBuffer = TUniformBufferRef<FNaniteUniformParameters>::CreateUniformBufferImmediate(NaniteUniformBufferParameters, UniformBuffer_MultiFrame, EUniformBufferValidation::None);
 
 	LumenCardCaptureViewUniformBuffer = TUniformBufferRef<FViewUniformShaderParameters>::CreateUniformBufferImmediate(ViewUniformBufferParameters, UniformBuffer_MultiFrame, EUniformBufferValidation::None);
 
-	FMobileBasePassUniformParameters MobileBasePassUniformParameters;
-	MobileOpaqueBasePassUniformBuffer = TUniformBufferRef<FMobileBasePassUniformParameters>::CreateUniformBufferImmediate(MobileBasePassUniformParameters, UniformBuffer_MultiFrame, EUniformBufferValidation::None);
-	MobileCSMOpaqueBasePassUniformBuffer = TUniformBufferRef<FMobileBasePassUniformParameters>::CreateUniformBufferImmediate(MobileBasePassUniformParameters, UniformBuffer_MultiFrame, EUniformBufferValidation::None);
-	MobileTranslucentBasePassUniformBuffer = TUniformBufferRef<FMobileBasePassUniformParameters>::CreateUniformBufferImmediate(MobileBasePassUniformParameters, UniformBuffer_MultiFrame, EUniformBufferValidation::None);
-		
 	FMobileDirectionalLightShaderParameters MobileDirectionalLightShaderParameters = {};
 	for (int32 Index = 0; Index < UE_ARRAY_COUNT(MobileDirectionalLightUniformBuffers); ++Index)
 	{
@@ -945,48 +933,6 @@ void FRendererModule::RegisterPersistentViewUniformBufferExtension(IPersistentVi
 	PersistentViewUniformBufferExtensions.Add(Extension);
 }
 
-bool FPersistentUniformBuffers::UpdateViewUniformBuffer(const FViewInfo& View, bool bShouldWaitForPersistentViewUniformBufferExtensionsJobs)
-{
-	checkf(FSceneInterface::GetShadingPath(View.GetFeatureLevel()) != EShadingPath::Deferred,
-		TEXT("UpdateViewUniformBuffer use is deprecated for the deferred shading renderer. The view uniform buffer must be bound ")
-		TEXT("globally instead through an RDG pass. This happens automatically if the view uniform buffer is included in the pass ")
-		TEXT("parameter struct."));
-
-	// Let the implementation of each extension decide whether it can cache the result for CachedView
-	for (IPersistentViewUniformBufferExtension* Extension : PersistentViewUniformBufferExtensions)
-	{
-		Extension->BeginRenderView(&View, bShouldWaitForPersistentViewUniformBufferExtensionsJobs);
-	}
-
-	// ViewUniformBuffer can be cached by mesh commands, so we need to update it every time we change current view.
-	if (CachedView != &View)
-	{
-		ViewUniformBuffer.UpdateUniformBufferImmediate(*View.CachedViewUniformShaderParameters);
-
-		if (const FViewInfo* InstancedView = View.GetInstancedView())
-		{
-			InstancedViewUniformBuffer.UpdateUniformBufferImmediate(reinterpret_cast<FInstancedViewUniformShaderParameters&>(*InstancedView->CachedViewUniformShaderParameters));
-		}
-		else
-		{
-			// If we don't render this pass in stereo we simply update the buffer with the same view uniform parameters.
-			// The shader will detect this and it will not attempt to apply ISR while this view is being rendered.
-			// TODO: It's more efficient to change the shader binding to point to ViewUniformBuffer instead of updating InstancedViewUniformBuffer.
-			InstancedViewUniformBuffer.UpdateUniformBufferImmediate(reinterpret_cast<FInstancedViewUniformShaderParameters&>(*View.CachedViewUniformShaderParameters));
-		}
-
-		CachedView = &View;
-
-		return true;
-	}
-	return false;
-}
-
-void FPersistentUniformBuffers::InvalidateCachedView()
-{
-	CachedView = nullptr;
-}
-
 class FAsyncCreateLightPrimitiveInteractionsTask : public FNonAbandonableTask
 {
 	TSet<FPrimitiveSceneInfo*> PendingPrimitives;
diff --git a/Engine/Source/Runtime/Renderer/Private/SceneHitProxyRendering.cpp b/Engine/Source/Runtime/Renderer/Private/SceneHitProxyRendering.cpp
--- a/Engine/Source/Runtime/Renderer/Private/SceneHitProxyRendering.cpp
+++ b/Engine/Source/Runtime/Renderer/Private/SceneHitProxyRendering.cpp
@@ -1165,8 +1165,6 @@ FEditorLevelInstanceMeshProcessor::FEditorLevelInstanceMeshProcessor(const FScen
 
 	PassDrawRenderState.SetDepthStencilState(TStaticDepthStencilState<true, CF_DepthNearOrEqual, true, CF_Always, SO_Keep, SO_Keep, SO_Replace>::GetRHI());
 	PassDrawRenderState.SetBlendState(TStaticBlendStateWriteMask<CW_NONE, CW_NONE, CW_NONE, CW_NONE>::GetRHI());
-	PassDrawRenderState.SetViewUniformBuffer(Scene->UniformBuffers.ViewUniformBuffer);
-	PassDrawRenderState.SetInstancedViewUniformBuffer(Scene->UniformBuffers.InstancedViewUniformBuffer);
 }
 
 FMeshPassProcessor* CreateEditorLevelInstancePassProcessor(const FScene* Scene, const FSceneView* InViewIfDynamicMeshCommand, FMeshPassDrawListContext* InDrawListContext)
diff --git a/Engine/Source/Runtime/Renderer/Private/ScenePrivate.h b/Engine/Source/Runtime/Renderer/Private/ScenePrivate.h
--- a/Engine/Source/Runtime/Renderer/Private/ScenePrivate.h
+++ b/Engine/Source/Runtime/Renderer/Private/ScenePrivate.h
@@ -2360,39 +2360,20 @@ public:
 class FPersistentUniformBuffers
 {
 public:
-	FPersistentUniformBuffers()
-		: CachedView(nullptr)
-	{
-	}
+	FPersistentUniformBuffers() = default;
 
 	void Initialize();
 	void Clear();
 
-	/** Compares the provided view against the cached view and updates the view uniform buffer
-	 *  if the views differ. Returns whether uniform buffer was updated.
-	 *  If bShouldWaitForPersistentViewUniformBufferExtensionsJobs == true, it calls Extension->BeginRenderView() which
-	 *  waits on the potential jobs dispatched in Extension->PrepareView(). Currently it is false only in FMobileSceneRenderer::InitViews()
-	 */
-	bool UpdateViewUniformBuffer(const FViewInfo& View, bool bShouldWaitForPersistentViewUniformBufferExtensionsJobs = true);
-
-	void InvalidateCachedView();
-
 	TUniformBufferRef<FViewUniformShaderParameters> ViewUniformBuffer;
-	TUniformBufferRef<FInstancedViewUniformShaderParameters> InstancedViewUniformBuffer;
 	TUniformBufferRef<FNaniteUniformParameters> NaniteUniformBuffer;
 	TUniformBufferRef<FViewUniformShaderParameters> LumenCardCaptureViewUniformBuffer;
 
-	TUniformBufferRef<FMobileBasePassUniformParameters> MobileOpaqueBasePassUniformBuffer;
-	TUniformBufferRef<FMobileBasePassUniformParameters> MobileCSMOpaqueBasePassUniformBuffer;
-	TUniformBufferRef<FMobileBasePassUniformParameters> MobileTranslucentBasePassUniformBuffer;
 	/** Mobile Directional Lighting uniform buffers, one for each lighting channel 
 	  * The first is used for primitives with no lighting channels set.
 	  */
 	TUniformBufferRef<FMobileDirectionalLightShaderParameters> MobileDirectionalLightUniformBuffers[NUM_LIGHTING_CHANNELS+1];
 	TUniformBufferRef<FMobileReflectionCaptureShaderParameters> MobileSkyReflectionUniformBuffer;
-
-	// View from which ViewUniformBuffer was last updated.
-	const FViewInfo* CachedView;
 };
 
 #if RHI_RAYTRACING
diff --git a/Engine/Source/Runtime/Renderer/Private/SceneRendering.cpp b/Engine/Source/Runtime/Renderer/Private/SceneRendering.cpp
--- a/Engine/Source/Runtime/Renderer/Private/SceneRendering.cpp
+++ b/Engine/Source/Runtime/Renderer/Private/SceneRendering.cpp
@@ -1827,13 +1827,6 @@ void FViewInfo::InitRHIResources()
 
 	const int32 TranslucencyLightingVolumeDim = GetTranslucencyLightingVolumeDim();
 
-	// Reset CachedView when CachedViewUniformShaderParameters change.
-	FScene* Scene = Family->Scene ? Family->Scene->GetRenderScene() : nullptr;
-	if (Scene)
-	{
-		Scene->UniformBuffers.InvalidateCachedView();
-	}
-
 	for (int32 CascadeIndex = 0; CascadeIndex < TVC_MAX; CascadeIndex++)
 	{
 		TranslucencyLightingVolumeMin[CascadeIndex] = VolumeBounds[CascadeIndex].Min;
diff --git a/Engine/Source/Runtime/Renderer/Private/SceneRendering.h b/Engine/Source/Runtime/Renderer/Private/SceneRendering.h
--- a/Engine/Source/Runtime/Renderer/Private/SceneRendering.h
+++ b/Engine/Source/Runtime/Renderer/Private/SceneRendering.h
@@ -2124,9 +2124,7 @@ protected:
 	void SortMobileBasePassAfterShadowInit(FExclusiveDepthStencil::Type BasePassDepthStencilAccess, FViewVisibleCommandsPerView& ViewCommandsPerView);
 	void SetupMobileBasePassAfterShadowInit(FExclusiveDepthStencil::Type BasePassDepthStencilAccess, FViewVisibleCommandsPerView& ViewCommandsPerView);
 
-	void UpdateOpaqueBasePassUniformBuffer(FRHICommandListImmediate& RHICmdList, const FViewInfo& View);
-	void UpdateTranslucentBasePassUniformBuffer(FRHICommandListImmediate& RHICmdList, const FViewInfo& View);
-	void UpdateDirectionalLightUniformBuffers(FRHICommandListImmediate& RHICmdList, const FViewInfo& View);
+	void UpdateDirectionalLightUniformBuffers(FRDGBuilder& GraphBuilder, const FViewInfo& View);
 	void UpdateSkyReflectionUniformBuffer();
 
 	void RenderForward(FRDGBuilder& GraphBuilder, const TArrayView<const FViewInfo*> ViewList, FRDGTextureRef ViewFamilyTexture, FRDGTextureMSAA& SceneColorMSAA, FRDGTextureMSAA& SceneDepthMSAA, TRDGUniformBufferRef<FMobileSceneTextureUniformParameters> MobileSceneTextures);
@@ -2158,6 +2156,8 @@ private:
 	static FGlobalDynamicIndexBuffer DynamicIndexBuffer;
 	static FGlobalDynamicVertexBuffer DynamicVertexBuffer;
 	static TGlobalResource<FGlobalDynamicReadBuffer> DynamicReadBuffer;
+
+	const FViewInfo* CachedView = nullptr;
 };
 
 // The noise textures need to be set in Slate too.
diff --git a/Engine/Source/Runtime/Renderer/Private/ShaderBaseClasses.cpp b/Engine/Source/Runtime/Renderer/Private/ShaderBaseClasses.cpp
--- a/Engine/Source/Runtime/Renderer/Private/ShaderBaseClasses.cpp
+++ b/Engine/Source/Runtime/Renderer/Private/ShaderBaseClasses.cpp
@@ -493,24 +493,12 @@ void FMeshMaterialShader::GetShaderBindings(
 	FMeshDrawSingleShaderBindings& ShaderBindings) const
 {
 	FMaterialShader::GetShaderBindings(Scene, FeatureLevel, MaterialRenderProxy, Material, ShaderBindings);
-	if (PassUniformBuffer.IsInitialized())
-	{
-		ShaderBindings.Add(PassUniformBuffer, DrawRenderState.GetPassUniformBuffer());
-	}
-	else
-	{
-		checkf(!DrawRenderState.GetPassUniformBuffer(), TEXT("Shader is missing the pass uniform buffer binding in the parameter map and the mesh pass draw render state requires it."));
-	}
 	if (DrawRenderState.GetViewUniformBuffer())
 	{
 		ShaderBindings.Add(GetUniformBufferParameter<FViewUniformShaderParameters>(), DrawRenderState.GetViewUniformBuffer());
 	}
 	ShaderBindings.Add(GetUniformBufferParameter<FDistanceCullFadeUniformShaderParameters>(), ShaderElementData.FadeUniformBuffer);
 	ShaderBindings.Add(GetUniformBufferParameter<FDitherUniformShaderParameters>(), ShaderElementData.DitherUniformBuffer);
-	if (DrawRenderState.GetInstancedViewUniformBuffer())
-	{
-		ShaderBindings.Add(GetUniformBufferParameter<FInstancedViewUniformShaderParameters>(), DrawRenderState.GetInstancedViewUniformBuffer());
-	}
 	ShaderBindings.Add(GetUniformBufferParameter<FNaniteUniformParameters>(), DrawRenderState.GetNaniteUniformBuffer());
 }
 
diff --git a/Engine/Source/Runtime/Renderer/Public/SceneRenderTargetParameters.h b/Engine/Source/Runtime/Renderer/Public/SceneRenderTargetParameters.h
--- a/Engine/Source/Runtime/Renderer/Public/SceneRenderTargetParameters.h
+++ b/Engine/Source/Runtime/Renderer/Public/SceneRenderTargetParameters.h
@@ -97,21 +97,11 @@ enum class EMobileSceneTextureSetupMode : uint32
 ENUM_CLASS_FLAGS(EMobileSceneTextureSetupMode);
 
 /** Fills the scene texture uniform buffer struct. */
-extern RENDERER_API void SetupMobileSceneTextureUniformParameters(
-	const FSceneRenderTargets& SceneContext,
-	EMobileSceneTextureSetupMode SetupMode,
-	FMobileSceneTextureUniformParameters& SceneTextureParameters);
-
 extern RENDERER_API void SetupMobileSceneTextureUniformParameters(
 	FRDGBuilder& GraphBuilder,
 	EMobileSceneTextureSetupMode SetupMode,
 	FMobileSceneTextureUniformParameters& SceneTextureParameters);
 
-/** Creates the RHI mobile scene texture uniform buffer with passthrough RDG resources. */
-extern RENDERER_API TUniformBufferRef<FMobileSceneTextureUniformParameters> CreateMobileSceneTextureUniformBuffer(
-	FRHIComputeCommandList& RHICmdList,
-	EMobileSceneTextureSetupMode SetupMode = EMobileSceneTextureSetupMode::All);
-
 /** Creates the RDG mobile scene texture uniform buffer. */
 extern RENDERER_API TRDGUniformBufferRef<FMobileSceneTextureUniformParameters> CreateMobileSceneTextureUniformBuffer(
 	FRDGBuilder& GraphBuilder,
@@ -155,6 +145,18 @@ extern RENDERER_API TUniformBufferRef<FSceneTextureUniformParameters> CreateScen
 	ERHIFeatureLevel::Type FeatureLevel,
 	ESceneTextureSetupMode SetupMode = ESceneTextureSetupMode::All);
 
+UE_DEPRECATED(5.0, "SetupMobileSceneTextureUniformParameters now requires an FRDGBuilder.")
+extern RENDERER_API void SetupMobileSceneTextureUniformParameters(
+	const FSceneRenderTargets& SceneContext,
+	EMobileSceneTextureSetupMode SetupMode,
+	FMobileSceneTextureUniformParameters& SceneTextureParameters);
+
+/** Creates the RHI mobile scene texture uniform buffer with passthrough RDG resources. */
+UE_DEPRECATED(5.0, "CreateMobileSceneTextureUniformBuffer now requires an FRDGBuilder.")
+extern RENDERER_API TUniformBufferRef<FMobileSceneTextureUniformParameters> CreateMobileSceneTextureUniformBuffer(
+	FRHIComputeCommandList& RHICmdList,
+	EMobileSceneTextureSetupMode SetupMode = EMobileSceneTextureSetupMode::All);
+
 UE_DEPRECATED(5.0, "Use CreateSceneTextureShaderParameters instead.")
 extern RENDERER_API TRefCountPtr<FRHIUniformBuffer> CreateSceneTextureUniformBufferDependentOnShadingPath(
 	FRHIComputeCommandList& RHICmdList,
