Expanded RDG builder outside of scene renderer.

diff --git a/Engine/Source/Runtime/RenderCore/Public/RenderGraphEvent.h b/Engine/Source/Runtime/RenderCore/Public/RenderGraphEvent.h
--- a/Engine/Source/Runtime/RenderCore/Public/RenderGraphEvent.h
+++ b/Engine/Source/Runtime/RenderCore/Public/RenderGraphEvent.h
@@ -26,9 +26,9 @@
 
 #if HAS_GPU_STATS
 	#if STATS
-		#define RDG_GPU_STAT_SCOPE(GraphBuilder, StatName) FRDGGPUStatScopeGuard PREPROCESSOR_JOIN(GPUStatEvent_##StatName,__LINE__) ((GraphBuilder), CSV_STAT_FNAME(StatName), GET_STATID(Stat_GPU_##StatName).GetName(), &DrawcallCountCategory_##StatName.Counter);
+		#define RDG_GPU_STAT_SCOPE(GraphBuilder, StatName) FRDGGPUStatScopeGuard PREPROCESSOR_JOIN(__RDG_GPUStatEvent_##StatName,__LINE__) ((GraphBuilder), CSV_STAT_FNAME(StatName), GET_STATID(Stat_GPU_##StatName).GetName(), &DrawcallCountCategory_##StatName.Counter);
 	#else
-		#define RDG_GPU_STAT_SCOPE(GraphBuilder, StatName) FRDGGPUStatScopeGuard PREPROCESSOR_JOIN(GPUStatEvent_##StatName,__LINE__) ((GraphBuilder), CSV_STAT_FNAME(StatName), FName(), &DrawcallCountCategory_##StatName.Counter);
+		#define RDG_GPU_STAT_SCOPE(GraphBuilder, StatName) FRDGGPUStatScopeGuard PREPROCESSOR_JOIN(__RDG_GPUStatEvent_##StatName,__LINE__) ((GraphBuilder), CSV_STAT_FNAME(StatName), FName(), &DrawcallCountCategory_##StatName.Counter);
 	#endif
 #else
 	#define RDG_GPU_STAT_SCOPE(GraphBuilder, StatName)
@@ -42,6 +42,10 @@
 	#define RDG_CSV_STAT_EXCLUSIVE_SCOPE_CONDITIONAL(GraphBuilder, StatName, bCondition)
 #endif
 
+ /** Injects a scope onto both the RDG and RHI timeline. */
+#define RDG_RHI_EVENT_SCOPE(GraphBuilder, Name) RDG_EVENT_SCOPE(GraphBuilder, #Name); SCOPED_DRAW_EVENT(GraphBuilder.RHICmdList, Name);
+#define RDG_RHI_GPU_STAT_SCOPE(GraphBuilder, StatName) RDG_GPU_STAT_SCOPE(GraphBuilder, StatName); SCOPED_GPU_STAT(GraphBuilder.RHICmdList, StatName);
+
 /** Returns whether the current frame is emitting render graph events. */
 RENDERCORE_API bool GetEmitRDGEvents();
 
diff --git a/Engine/Source/Runtime/RenderCore/Public/RenderGraphUtils.h b/Engine/Source/Runtime/RenderCore/Public/RenderGraphUtils.h
--- a/Engine/Source/Runtime/RenderCore/Public/RenderGraphUtils.h
+++ b/Engine/Source/Runtime/RenderCore/Public/RenderGraphUtils.h
@@ -159,6 +159,16 @@ inline FRDGBufferRef TryRegisterExternalBuffer(
 	return ExternalPooledBuffer ? GraphBuilder.RegisterExternalBuffer(ExternalPooledBuffer, Flags) : nullptr;
 }
 
+inline FRDGTextureRef RegisterExternalTexture(FRDGBuilder& GraphBuilder, FRHITexture* Texture, const TCHAR* NameIfUnregistered)
+{
+	if (FRDGTextureRef FoundTexture = GraphBuilder.FindExternalTexture(Texture))
+	{
+		return FoundTexture;
+	}
+
+	return GraphBuilder.RegisterExternalTexture(CreateRenderTarget(Texture, NameIfUnregistered));
+}
+
 /** Simple pair of RDG textures used for MSAA. */
 struct FRDGTextureMSAA
 {
diff --git a/Engine/Source/Runtime/Renderer/Private/DeferredShadingRenderer.cpp b/Engine/Source/Runtime/Renderer/Private/DeferredShadingRenderer.cpp
--- a/Engine/Source/Runtime/Renderer/Private/DeferredShadingRenderer.cpp
+++ b/Engine/Source/Runtime/Renderer/Private/DeferredShadingRenderer.cpp
@@ -1487,12 +1487,10 @@ void FDeferredShadingSceneRenderer::CommitFinalPipelineState()
 	} 
 }
 
-void FDeferredShadingSceneRenderer::Render(FRHICommandListImmediate& RHICmdList)
+void FDeferredShadingSceneRenderer::Render(FRDGBuilder& GraphBuilder)
 {
 	const bool bNaniteEnabled = UseNanite(ShaderPlatform) && ViewFamily.EngineShowFlags.NaniteMeshes;
 
-	FRDGBuilder GraphBuilder(RHICmdList, RDG_EVENT_NAME("DeferredShadingRenderer_Render(ViewFamily=%s)", ViewFamily.bResolveScene ? TEXT("Primary") : TEXT("Auxiliary")));
-
 	Scene->UpdateAllPrimitiveSceneInfos(GraphBuilder, true);
 
 	FGPUSceneScopeBeginEndHelper GPUSceneScopeBeginEndHelper(Scene->GPUScene, GPUSceneDynamicContext, *Scene);
@@ -1532,41 +1530,36 @@ void FDeferredShadingSceneRenderer::Render(FRHICommandListImmediate& RHICmdList)
 	SCOPED_NAMED_EVENT(FDeferredShadingSceneRenderer_Render, FColor::Emerald);
 
 #if WITH_MGPU
-	const FRHIGPUMask RenderTargetGPUMask = ComputeGPUMasks(RHICmdList);
+	const FRHIGPUMask RenderTargetGPUMask = ComputeGPUMasks(GraphBuilder.RHICmdList);
 #endif // WITH_MGPU
 
 	// By default, limit our GPU usage to only GPUs specified in the view masks.
-	SCOPED_GPU_MASK(RHICmdList, AllViewsGPUMask);
-	SCOPED_GPU_MASK(FRHICommandListExecutor::GetImmediateAsyncComputeCommandList(), AllViewsGPUMask);
+	RDG_GPU_MASK_SCOPE(GraphBuilder, AllViewsGPUMask);
 
 	FSceneRenderTargets& SceneContext = FSceneRenderTargets::Get();
 	
 	//make sure all the targets we're going to use will be safely writable.
-	GRenderTargetPool.TransitionTargetsWritable(RHICmdList);
+	GRenderTargetPool.TransitionTargetsWritable(GraphBuilder.RHICmdList);
 
 	// this way we make sure the SceneColor format is the correct one and not the one from the end of frame before
 	SceneContext.ReleaseSceneColor();
 
-	WaitOcclusionTests(RHICmdList);
+	WaitOcclusionTests(GraphBuilder.RHICmdList);
 
 	if (!ViewFamily.EngineShowFlags.Rendering)
 	{
-		GraphBuilder.Execute();
 		return;
 	}
 
-	SCOPED_DRAW_EVENT(RHICmdList, Scene);
-
-	// Anything rendered inside Render() which isn't accounted for will fall into this stat
-	// This works because child stat events do not contribute to their parents' times (see GPU_STATS_CHILD_TIMES_INCLUDED)
-	SCOPED_GPU_STAT(RHICmdList, Unaccounted);
+	RDG_RHI_EVENT_SCOPE(GraphBuilder, Scene);
+	RDG_RHI_GPU_STAT_SCOPE(GraphBuilder, Unaccounted);
 	
 	{
 		SCOPE_CYCLE_COUNTER(STAT_FDeferredShadingSceneRenderer_Render_Init);
-		SCOPED_GPU_STAT(RHICmdList, AllocateRendertargets);
+		RDG_RHI_GPU_STAT_SCOPE(GraphBuilder, AllocateRendertargets);
 
 		// Initialize global system textures (pass-through if already initialized).
-		GSystemTextures.InitializeTextures(RHICmdList, FeatureLevel);
+		GSystemTextures.InitializeTextures(GraphBuilder.RHICmdList, FeatureLevel);
 
 		// Allocate the maximum scene render target space for the current view family.
 		SceneContext.Allocate(GraphBuilder, this);
@@ -1610,7 +1603,7 @@ void FDeferredShadingSceneRenderer::Render(FRHICommandListImmediate& RHICmdList)
 	FILCUpdatePrimTaskData ILCTaskData;
 
 	// Find the visible primitives.
-	RHICmdList.ImmediateFlush(EImmediateFlushType::DispatchToRHIThread);
+	GraphBuilder.RHICmdList.ImmediateFlush(EImmediateFlushType::DispatchToRHIThread);
 
 
 	bool bDoInitViewAftersPrepass = false;
@@ -1645,7 +1638,7 @@ void FDeferredShadingSceneRenderer::Render(FRHICommandListImmediate& RHICmdList)
 
 #if RHI_RAYTRACING
 	// Gather mesh instances, shaders, resources, parameters, etc. and build ray tracing acceleration structure
-	GatherRayTracingWorldInstances(RHICmdList);
+	GatherRayTracingWorldInstances(GraphBuilder.RHICmdList);
 
 	if (Views[0].RayTracingRenderMode != ERayTracingRenderMode::PathTracing)
 	{
@@ -1710,7 +1703,7 @@ void FDeferredShadingSceneRenderer::Render(FRHICommandListImmediate& RHICmdList)
 	// GBuffers are temporarily being allocated here until they are converted to RDG.
 	{
 		SCOPE_CYCLE_COUNTER(STAT_FDeferredShadingSceneRenderer_AllocGBufferTargets);
-		SceneContext.AllocGBufferTargets(RHICmdList);
+		SceneContext.AllocGBufferTargets(GraphBuilder.RHICmdList);
 	}
 
 	FSceneTextures& SceneTextures = FSceneTextures::Create(GraphBuilder);
@@ -1823,7 +1816,7 @@ void FDeferredShadingSceneRenderer::Render(FRHICommandListImmediate& RHICmdList)
 		RunGPUSkinCacheTransition(InRHICmdList, Scene, EGPUSkinCacheTransition::Renderer);
 	});
 
-	FHairStrandsBookmarkParameters HairStrandsBookmarkParameters;
+	FHairStrandsBookmarkParameters& HairStrandsBookmarkParameters = *GraphBuilder.AllocObject<FHairStrandsBookmarkParameters>();
 
 	{
 		if (IsHairStrandsEnabled(EHairStrandsShaderType::All, Scene->GetShaderPlatform()))
@@ -1903,12 +1896,12 @@ void FDeferredShadingSceneRenderer::Render(FRHICommandListImmediate& RHICmdList)
 		if (bDoInitViewAftersPrepass)
 		{
 			{
-				SCOPED_GPU_STAT(GraphBuilder.RHICmdList, VisibilityCommands);
+				RDG_RHI_GPU_STAT_SCOPE(GraphBuilder, VisibilityCommands);
 				InitViewsPossiblyAfterPrepass(GraphBuilder, ILCTaskData);
 			}
 
 			{
-				SCOPED_GPU_STAT(GraphBuilder.RHICmdList, GPUSceneUpdate);
+				RDG_RHI_GPU_STAT_SCOPE(GraphBuilder, GPUSceneUpdate);
 				PrepareDistanceFieldScene(GraphBuilder, false);
 			}
 
@@ -2038,7 +2031,7 @@ void FDeferredShadingSceneRenderer::Render(FRHICommandListImmediate& RHICmdList)
 	AddResolveSceneDepthPass(GraphBuilder, Views, SceneTextures.Depth);
 
 	// NOTE: The ordering of the lights is used to select sub-sets for different purposes, e.g., those that support clustered deferred.
-	FSortedLightSetSceneInfo SortedLightSet;
+	FSortedLightSetSceneInfo& SortedLightSet = *GraphBuilder.AllocObject<FSortedLightSetSceneInfo>();
 	{
 		RDG_GPU_STAT_SCOPE(GraphBuilder, SortLights);
 		GatherLightsAndComputeLightGrid(GraphBuilder, bComputeLightGrid, SortedLightSet);
@@ -2178,7 +2171,7 @@ void FDeferredShadingSceneRenderer::Render(FRHICommandListImmediate& RHICmdList)
 	}
 
 	FHairStrandsRenderingData* HairDatas = nullptr;
-	FHairStrandsRenderingData HairDatasStorage;
+	FHairStrandsRenderingData* HairDatasStorage = GraphBuilder.AllocObject<FHairStrandsRenderingData>();
 
 	FRDGTextureRef ForwardScreenSpaceShadowMaskTexture = nullptr;
 	FRDGTextureRef ForwardScreenSpaceShadowMaskHairTexture = nullptr;
@@ -2186,9 +2179,9 @@ void FDeferredShadingSceneRenderer::Render(FRHICommandListImmediate& RHICmdList)
 	{
 		if (bHairEnable)
 		{
-			RenderHairPrePass(GraphBuilder, Scene, Views, HairDatasStorage);
-			RenderHairBasePass(GraphBuilder, Scene, SceneTextures, Views, HairDatasStorage);
-			HairDatas = &HairDatasStorage;
+			RenderHairPrePass(GraphBuilder, Scene, Views, *HairDatasStorage);
+			RenderHairBasePass(GraphBuilder, Scene, SceneTextures, Views, *HairDatasStorage);
+			HairDatas = HairDatasStorage;
 		}
 
 		RenderForwardShadowProjections(GraphBuilder, SceneTextures, ForwardScreenSpaceShadowMaskTexture, ForwardScreenSpaceShadowMaskHairTexture, HairDatas);
@@ -2214,7 +2207,7 @@ void FDeferredShadingSceneRenderer::Render(FRHICommandListImmediate& RHICmdList)
 		InitTranslucencyLightingVolumeTextures(GraphBuilder, Views, ERDGPassFlags::AsyncCompute, TranslucencyLightingVolumeTextures);
 	}
 
-	SceneContext.AllocSceneColor(RHICmdList);
+	SceneContext.AllocSceneColor(GraphBuilder.RHICmdList);
 	SceneTextures.Color = RegisterExternalTextureMSAA(GraphBuilder, SceneContext.GetSceneColor());
 
 	{
@@ -2396,8 +2389,8 @@ void FDeferredShadingSceneRenderer::Render(FRHICommandListImmediate& RHICmdList)
 	// Hair base pass for deferred shading
 	if (bHairEnable && !IsForwardShadingEnabled(ShaderPlatform))
 	{
-		RenderHairPrePass(GraphBuilder, Scene, Views, HairDatasStorage);
-		HairDatas = &HairDatasStorage;
+		RenderHairPrePass(GraphBuilder, Scene, Views, *HairDatasStorage);
+		HairDatas = HairDatasStorage;
 	}
 
 	// Copy lighting channels out of stencil before deferred decals which overwrite those values
@@ -2435,24 +2428,7 @@ void FDeferredShadingSceneRenderer::Render(FRHICommandListImmediate& RHICmdList)
 	if (bHairEnable && !IsForwardShadingEnabled(ShaderPlatform))
 	{
 		check(HairDatas);
-		RenderHairBasePass(GraphBuilder, Scene,  SceneTextures, Views, HairDatasStorage);
-
-		// Disable temporarly HZB update with hair data. This is not mandatory.
-		#if 0
-		const FSceneTextureParameters SceneTextureParameters = GetSceneTextureParameters(GraphBuilder, SceneTextures.UniformBuffer);
-		for (int32 ViewIndex = 0; ViewIndex < Views.Num(); ViewIndex++)
-		{
-			FViewInfo& View = Views[ViewIndex];
-			RDG_EVENT_SCOPE(GraphBuilder, "BuildHZB_HairUpdate(ViewId=%d)", ViewIndex);
-			BuildHZB(
-				GraphBuilder,
-				SceneTextureParameters.SceneTextures.Depth, 
-				/* VisBufferTexture = */ nullptr,
-				View,
-				/* OutClosestHZBTexture = */ nullptr,
-				/* OutFurthestHZBTexture = */ nullptr);
-		}
-		#endif
+		RenderHairBasePass(GraphBuilder, Scene,  SceneTextures, Views, *HairDatasStorage);
 	}
 
 	// Rebuild scene textures to include velocity, custom depth, and SSAO.
@@ -2931,9 +2907,7 @@ void FDeferredShadingSceneRenderer::Render(FRHICommandListImmediate& RHICmdList)
 		GraphBuilder.SetCommandListStat(GET_STATID(STAT_CLM_AfterFrame));
 	}
 
-	GraphBuilder.Execute();
-
-	ServiceLocalQueue();
+	AddServiceLocalQueuePass(GraphBuilder);
 }
 
 #if RHI_RAYTRACING
diff --git a/Engine/Source/Runtime/Renderer/Private/DeferredShadingRenderer.h b/Engine/Source/Runtime/Renderer/Private/DeferredShadingRenderer.h
--- a/Engine/Source/Runtime/Renderer/Private/DeferredShadingRenderer.h
+++ b/Engine/Source/Runtime/Renderer/Private/DeferredShadingRenderer.h
@@ -243,10 +243,10 @@ public:
 	bool RenderHzb(FRDGBuilder& GraphBuilder, FRDGTextureRef SceneDepthTexture);
 
 	/** Renders the view family. */
-	virtual void Render(FRHICommandListImmediate& RHICmdList) override;
+	virtual void Render(FRDGBuilder& GraphBuilder) override;
 
 	/** Render the view family's hit proxies. */
-	virtual void RenderHitProxies(FRHICommandListImmediate& RHICmdList) override;
+	virtual void RenderHitProxies(FRDGBuilder& GraphBuilder) override;
 
 #if !(UE_BUILD_SHIPPING || UE_BUILD_TEST)
 	void RenderVisualizeTexturePool(FRHICommandListImmediate& RHICmdList);
diff --git a/Engine/Source/Runtime/Renderer/Private/MobileSceneCaptureRendering.cpp b/Engine/Source/Runtime/Renderer/Private/MobileSceneCaptureRendering.cpp
--- a/Engine/Source/Runtime/Renderer/Private/MobileSceneCaptureRendering.cpp
+++ b/Engine/Source/Runtime/Renderer/Private/MobileSceneCaptureRendering.cpp
@@ -319,9 +319,12 @@ void UpdateSceneCaptureContentMobile_RenderThread(
 	{
 #if WANTS_DRAW_MESH_EVENTS
 		SCOPED_DRAW_EVENTF(RHICmdList, SceneCaptureMobile, TEXT("SceneCaptureMobile %s"), *EventName);
+		FRDGBuilder GraphBuilder(RHICmdList, RDG_EVENT_NAME("SceneCaptureMobile %s", *EventName));
 #else
 		SCOPED_DRAW_EVENT(RHICmdList, UpdateSceneCaptureContentMobile_RenderThread);
+		FRDGBuilder GraphBuilder(RHICmdList, RDG_EVENT_NAME("SceneCaptureMobile"));
 #endif
+
 		FViewInfo& View = SceneRenderer->Views[0];
 
 		const bool bIsMobileHDR = IsMobileHDR();
@@ -334,54 +337,28 @@ void UpdateSceneCaptureContentMobile_RenderThread(
 		const bool bNeedsFlippedFinalColor = bNeedsFlippedCopy && !bUseSceneTextures;
 
 		// Intermediate render target that will need to be flipped (needed on !IsMobileHDR())
-		TRefCountPtr<IPooledRenderTarget> FlippedPooledRenderTarget;
+		FRDGTextureRef FlippedOutputTexture{};
 
 		const FRenderTarget* Target = SceneRenderer->ViewFamily.RenderTarget;
 		if (bNeedsFlippedFinalColor)
 		{
 			// We need to use an intermediate render target since the result will be flipped
 			auto& RenderTargetRHI = Target->GetRenderTargetTexture();
-			FPooledRenderTargetDesc Desc(FPooledRenderTargetDesc::Create2DDesc(Target->GetSizeXY(),
+			FRDGTextureDesc Desc(FRDGTextureDesc::Create2D(
+				Target->GetSizeXY(),
 				RenderTargetRHI.GetReference()->GetFormat(),
 				RenderTargetRHI.GetReference()->GetClearBinding(),
-				TexCreate_None,
-				TexCreate_RenderTargetable,
-				false));
-			GRenderTargetPool.FindFreeElement(RHICmdList, Desc, FlippedPooledRenderTarget, TEXT("SceneCaptureFlipped"));
+				TexCreate_RenderTargetable));
+			FlippedOutputTexture = GraphBuilder.CreateTexture(Desc, TEXT("SceneCaptureFlipped"));
 		}
 
-		// Helper class to allow setting render target
-		struct FRenderTargetOverride : public FRenderTarget
-		{
-			FRenderTargetOverride(const FRenderTarget* TargetIn, FRHITexture2D* In)
-			{
-				RenderTargetTextureRHI = In;
-				OriginalTarget = TargetIn;
-			}
-
-			virtual FIntPoint GetSizeXY() const override { return FIntPoint(RenderTargetTextureRHI->GetSizeX(), RenderTargetTextureRHI->GetSizeY()); }
-			virtual float GetDisplayGamma() const override { return OriginalTarget->GetDisplayGamma(); }
-
-			FTexture2DRHIRef GetTextureParamRef() { return RenderTargetTextureRHI; }
-			const FRenderTarget* OriginalTarget;
-		} FlippedRenderTarget(Target, 
-			FlippedPooledRenderTarget.GetReference()
-			? FlippedPooledRenderTarget.GetReference()->GetRenderTargetItem().TargetableTexture->GetTexture2D()
-			: nullptr);
-
 		// We don't support screen percentage in scene capture.
 		FIntRect ViewRect = View.UnscaledViewRect;
 		FIntRect UnconstrainedViewRect = View.UnconstrainedViewRect;
 
 		if(bNeedsFlippedFinalColor)
 		{
-			auto& RenderTargetRHI = Target->GetRenderTargetTexture();
-			
-			FRHIRenderPassInfo RPInfo(RenderTargetRHI, ERenderTargetActions::DontLoad_Store);
-			TransitionRenderPassTargets(RHICmdList, RPInfo);
-			RHICmdList.BeginRenderPass(RPInfo, TEXT("Clear"));
-			DrawClearQuad(RHICmdList, true, FLinearColor::Black, false, 0, false, 0, Target->GetSizeXY(), ViewRect);
-			RHICmdList.EndRenderPass();
+			AddClearRenderTargetPass(GraphBuilder, FlippedOutputTexture, FLinearColor::Black, ViewRect);
 		}
 
 		// Register pass for InverseOpacity for this scope
@@ -390,15 +367,33 @@ void UpdateSceneCaptureContentMobile_RenderThread(
 		
 		// Render the scene normally
 		{
-			SCOPED_DRAW_EVENT(RHICmdList, RenderScene);
+			RDG_RHI_EVENT_SCOPE(GraphBuilder, RenderScene);
 
 			if (bNeedsFlippedFinalColor)
 			{
+				// Helper class to allow setting render target
+				struct FRenderTargetOverride : public FRenderTarget
+				{
+					FRenderTargetOverride(const FRenderTarget* TargetIn, FRHITexture2D* OverrideTexture)
+					{
+						RenderTargetTextureRHI = OverrideTexture;
+						OriginalTarget = TargetIn;
+					}
+
+					virtual FIntPoint GetSizeXY() const override { return FIntPoint(RenderTargetTextureRHI->GetSizeX(), RenderTargetTextureRHI->GetSizeY()); }
+					virtual float GetDisplayGamma() const override { return OriginalTarget->GetDisplayGamma(); }
+
+					FTexture2DRHIRef GetTextureParamRef() { return RenderTargetTextureRHI; }
+					const FRenderTarget* OriginalTarget;
+				};
+
 				// Hijack the render target
-				SceneRenderer->ViewFamily.RenderTarget = &FlippedRenderTarget; //-V506
+				GraphBuilder.PreallocateTexture(FlippedOutputTexture);
+				FRHITexture2D* FlippedOutputTextureRHI = GraphBuilder.GetPooledTexture(FlippedOutputTexture)->GetTargetableRHI()->GetTexture2D();
+				SceneRenderer->ViewFamily.RenderTarget = GraphBuilder.AllocObject<FRenderTargetOverride>(Target, FlippedOutputTextureRHI); //-V506
 			}
 
-			SceneRenderer->Render(RHICmdList);
+			SceneRenderer->Render(GraphBuilder);
 
 			if (bNeedsFlippedFinalColor)
 			{
@@ -407,35 +402,29 @@ void UpdateSceneCaptureContentMobile_RenderThread(
 			}
 		}
 
-		FRDGBuilder GraphBuilder(RHICmdList);
-
-		FRDGTextureRef ViewFamilyTexture = TryCreateViewFamilyTexture(GraphBuilder, SceneRenderer->ViewFamily);
+		FRDGTextureRef OutputTexture = RegisterExternalTexture(GraphBuilder, Target->GetRenderTargetTexture(), TEXT("OutputTexture"));
 
 		const FIntPoint TargetSize(UnconstrainedViewRect.Width(), UnconstrainedViewRect.Height());
 		if (bNeedsFlippedFinalColor)
 		{
 			// We need to flip this texture upside down (since we depended on tonemapping to fix this on the hdr path)
 			RDG_EVENT_SCOPE(GraphBuilder, "FlipCapture");
-			CopyCaptureToTarget(GraphBuilder, ViewFamilyTexture, TargetSize, View, ViewRect, GraphBuilder.RegisterExternalTexture(FlippedPooledRenderTarget), bNeedsFlippedCopy, SceneRenderer);
+			CopyCaptureToTarget(GraphBuilder, OutputTexture, TargetSize, View, ViewRect, FlippedOutputTexture, bNeedsFlippedCopy, SceneRenderer);
 		}
 		else if(bUseSceneTextures)
 		{
+			const FMinimalSceneTextures& SceneTextures = FSceneTextures::Get(GraphBuilder);
+
 			// Copy the captured scene into the destination texture
 			RDG_EVENT_SCOPE(GraphBuilder, "CaptureSceneColor");
-			CopyCaptureToTarget(GraphBuilder, ViewFamilyTexture, TargetSize, View, ViewRect, GraphBuilder.RegisterExternalTexture(FSceneRenderTargets::Get().GetSceneColor()), bNeedsFlippedCopy, SceneRenderer);
+			CopyCaptureToTarget(GraphBuilder, OutputTexture, TargetSize, View, ViewRect, SceneTextures.Color.Target, bNeedsFlippedCopy, SceneRenderer);
 		}
-		
-		{
-			FRDGTextureRef MipTexture = GraphBuilder.RegisterExternalTexture(CreateRenderTarget(RenderTarget->GetRenderTargetTexture(), TEXT("MipGenerationInput")));
-			FRDGTextureRef OutputTexture = GraphBuilder.RegisterExternalTexture(CreateRenderTarget(RenderTargetTexture->TextureRHI, TEXT("MipGenerationOutput")));
 
-			if (bGenerateMips)
-			{
-				FGenerateMips::Execute(GraphBuilder, MipTexture, GenerateMipsParams);
-			}
-
-			AddCopyToResolveTargetPass(GraphBuilder, MipTexture, OutputTexture, ResolveParams);
+		if (bGenerateMips)
+		{
+			FGenerateMips::Execute(GraphBuilder, OutputTexture, GenerateMipsParams);
 		}
+
 		GraphBuilder.Execute();
 	}
 	FSceneRenderer::WaitForTasksClearSnapshotsAndDeleteSceneRenderer(RHICmdList, SceneRenderer);
diff --git a/Engine/Source/Runtime/Renderer/Private/MobileShadingRenderer.cpp b/Engine/Source/Runtime/Renderer/Private/MobileShadingRenderer.cpp
--- a/Engine/Source/Runtime/Renderer/Private/MobileShadingRenderer.cpp
+++ b/Engine/Source/Runtime/Renderer/Private/MobileShadingRenderer.cpp
@@ -533,14 +533,12 @@ void FMobileSceneRenderer::InitViews(FRDGBuilder& GraphBuilder)
 /** 
 * Renders the view family. 
 */
-void FMobileSceneRenderer::Render(FRHICommandListImmediate& RHICmdList)
+void FMobileSceneRenderer::Render(FRDGBuilder& GraphBuilder)
 {
-	FRDGBuilder GraphBuilder(RHICmdList, RDG_EVENT_NAME("MobileShadingRenderer_Render(ViewFamily=%s)", ViewFamily.bResolveScene ? TEXT("Primary") : TEXT("Auxiliary")));
+	GraphBuilder.SetCommandListStat(GET_STATID(STAT_CLMM_SceneStart));
 
-	RHICmdList.SetCurrentStat(GET_STATID(STAT_CLMM_SceneStart));
-
-	SCOPED_DRAW_EVENT(RHICmdList, MobileSceneRender);
-	SCOPED_GPU_STAT(RHICmdList, MobileSceneRender);
+	RDG_RHI_EVENT_SCOPE(GraphBuilder, MobileSceneRender);
+	RDG_RHI_GPU_STAT_SCOPE(GraphBuilder, MobileSceneRender);
 
 	Scene->UpdateAllPrimitiveSceneInfos(GraphBuilder);
 
@@ -566,22 +564,18 @@ void FMobileSceneRenderer::Render(FRHICommandListImmediate& RHICmdList)
 
 	CSV_SCOPED_TIMING_STAT_EXCLUSIVE(RenderOther);
 	QUICK_SCOPE_CYCLE_COUNTER(STAT_FMobileSceneRenderer_Render);
-	//FRHICommandListExecutor::GetImmediateCommandList().ImmediateFlush(EImmediateFlushType::DispatchToRHIThread);
 
 	if (!ViewFamily.EngineShowFlags.Rendering)
 	{
-		GraphBuilder.Execute();
 		return;
 	}
 
-	SCOPED_GPU_STAT(RHICmdList, MobileSceneRender);
-
-	WaitOcclusionTests(RHICmdList);
+	WaitOcclusionTests(GraphBuilder.RHICmdList);
 	FRHICommandListExecutor::GetImmediateCommandList().PollOcclusionQueries();
-	RHICmdList.ImmediateFlush(EImmediateFlushType::DispatchToRHIThread);
+	GraphBuilder.RHICmdList.ImmediateFlush(EImmediateFlushType::DispatchToRHIThread);
 
 	// Initialize global system textures (pass-through if already initialized).
-	GSystemTextures.InitializeTextures(RHICmdList, FeatureLevel);
+	GSystemTextures.InitializeTextures(GraphBuilder.RHICmdList, FeatureLevel);
 	FRDGSystemTextures::Create(GraphBuilder);
 
 	// Find the visible primitives and prepare targets and buffers for rendering
@@ -602,11 +596,12 @@ void FMobileSceneRenderer::Render(FRHICommandListImmediate& RHICmdList)
 	DynamicIndexBuffer.Commit();
 	DynamicVertexBuffer.Commit();
 	DynamicReadBuffer.Commit();
-	RHICmdList.ImmediateFlush(EImmediateFlushType::DispatchToRHIThread);
+	GraphBuilder.RHICmdList.ImmediateFlush(EImmediateFlushType::DispatchToRHIThread);
 
 	GraphBuilder.SetCommandListStat(GET_STATID(STAT_CLMM_SceneSim));
 
 	FSceneRenderTargets& SceneContext = FSceneRenderTargets::Get();
+	FSceneTextures& SceneTextures = FSceneTextures::Create(GraphBuilder);
 
 	if (bUseVirtualTexturing)
 	{
@@ -623,7 +618,7 @@ void FMobileSceneRenderer::Render(FRHICommandListImmediate& RHICmdList)
 		});
 	}
 
-	FSortedLightSetSceneInfo SortedLightSet;
+	FSortedLightSetSceneInfo& SortedLightSet = *GraphBuilder.AllocObject<FSortedLightSetSceneInfo>();
 	if (bDeferredShading)
 	{
 		GatherAndSortLights(SortedLightSet);
@@ -756,7 +751,7 @@ void FMobileSceneRenderer::Render(FRHICommandListImmediate& RHICmdList)
 	GraphBuilder.SetCommandListStat(GET_STATID(STAT_CLMM_Post));
 
 	if (bUseVirtualTexturing)
-	{	
+	{
 		RDG_GPU_STAT_SCOPE(GraphBuilder, VirtualTextureUpdate);
 
 		AddPass(GraphBuilder, RDG_EVENT_NAME("VirtualTextureUpdate"), [this, &SceneContext](FRHICommandListImmediate& InRHICmdList)
@@ -800,18 +795,16 @@ void FMobileSceneRenderer::Render(FRHICommandListImmediate& RHICmdList)
 		}
 	}
 
-	AddPass(GraphBuilder, [this](FRHICommandListImmediate& InRHICmdList)
+	AddPass(GraphBuilder, [this](FRHICommandListImmediate&)
 	{
 		GEngine->GetPostRenderDelegate().Broadcast();
-
-		InRHICmdList.SetCurrentStat(GET_STATID(STAT_CLMM_SceneEnd));
 	});
-	
+
+	GraphBuilder.SetCommandListStat(GET_STATID(STAT_CLMM_SceneEnd));
+
 	RenderFinish(GraphBuilder, ViewFamilyTexture);
-	
-	AddPass(GraphBuilder, PollOcclusionQueriesAndDispatchToRHIThreadPass);
 
-	GraphBuilder.Execute();
+	AddPass(GraphBuilder, PollOcclusionQueriesAndDispatchToRHIThreadPass);
 }
 
 BEGIN_SHADER_PARAMETER_STRUCT(FMobilePostBasePassViewExtensionParameters, )
@@ -831,7 +824,7 @@ void FMobileSceneRenderer::RenderForward(FRDGBuilder& GraphBuilder, const TArray
 	// Verify using both MSAA sample count AND the scene color surface sample count, since on GLES you can't have MSAA color targets,
 	// so the color target would be created without MSAA, and MSAA is achieved through magical means (the framebuffer, being MSAA,
 	// tells the GPU "execute this renderpass as MSAA, and when you're done, automatically resolve and copy into this non-MSAA texture").
-	bool bMobileMSAA = NumMSAASamples > 1 && SceneContext.GetSceneColorSurface()->GetNumSamples() > 1;
+	bool bMobileMSAA = NumMSAASamples > 1 && SceneColorMSAA.Target->Desc.NumSamples > 1;
 
 	static const auto CVarMobileMultiView = IConsoleManager::Get().FindTConsoleVariableDataInt(TEXT("vr.MobileMultiView"));
 	const bool bIsMultiViewApplication = (CVarMobileMultiView && CVarMobileMultiView->GetValueOnAnyThread() != 0);
diff --git a/Engine/Source/Runtime/Renderer/Private/PlanarReflectionRendering.cpp b/Engine/Source/Runtime/Renderer/Private/PlanarReflectionRendering.cpp
--- a/Engine/Source/Runtime/Renderer/Private/PlanarReflectionRendering.cpp
+++ b/Engine/Source/Runtime/Renderer/Private/PlanarReflectionRendering.cpp
@@ -355,6 +355,8 @@ static void UpdatePlanarReflectionContents_RenderThread(
 				SCOPED_DRAW_EVENT(RHICmdList, UpdatePlanarReflectionContent_RenderThread);
 #endif
 
+				FRDGBuilder GraphBuilder(RHICmdList, RDG_EVENT_NAME("PlanarReflection"));
+
 				// Reflection view late update
 				if (SceneRenderer->Views.Num() > 1)
 				{
@@ -370,8 +372,8 @@ static void UpdatePlanarReflectionContents_RenderThread(
 
 				// Render the scene normally
 				{
-					SCOPED_DRAW_EVENT(RHICmdList, RenderScene);
-					SceneRenderer->Render(RHICmdList);
+					RDG_RHI_EVENT_SCOPE(GraphBuilder, RenderScene);
+					SceneRenderer->Render(GraphBuilder);
 				}
 
 				SceneProxy->RenderTarget = RenderTarget;
@@ -384,14 +386,9 @@ static void UpdatePlanarReflectionContents_RenderThread(
 					SceneProxy->ViewRect[ViewIndex] = SceneRenderer->Views[ViewIndex].ViewRect;
 				}
 
-				FRDGBuilder GraphBuilder(RHICmdList);
-				FRDGSystemTextures::Create(GraphBuilder);
-
-				FRDGTextureRef ViewFamilyTexture = GraphBuilder.RegisterExternalTexture(CreateRenderTarget(RenderTarget->TextureRHI, TEXT("ViewFamilyTexture")));
-
-				const FSceneTextures& SceneTextures = FSceneTextures::CreateMinimal(GraphBuilder);
-
+				FRDGTextureRef ReflectionOutputTexture = GraphBuilder.RegisterExternalTexture(CreateRenderTarget(RenderTarget->TextureRHI, TEXT("ReflectionOutputTexture")));
 				FSceneTextureShaderParameters SceneTextureParameters = CreateSceneTextureShaderParameters(GraphBuilder, SceneRenderer->FeatureLevel, ESceneTextureSetupMode::SceneDepth);
+				const FMinimalSceneTextures& SceneTextures = FSceneTextures::Get(GraphBuilder);
 
 				for (int32 ViewIndex = 0; ViewIndex < SceneRenderer->Views.Num(); ++ViewIndex)
 				{
@@ -399,13 +396,14 @@ static void UpdatePlanarReflectionContents_RenderThread(
 					RDG_GPU_MASK_SCOPE(GraphBuilder, View.GPUMask);
 					if (MainSceneRenderer->Scene->GetShadingPath() == EShadingPath::Deferred)
 					{
-						PrefilterPlanarReflection<true>(GraphBuilder, View, SceneTextureParameters, SceneProxy, SceneTextures.Color.Resolve, ViewFamilyTexture);
+						PrefilterPlanarReflection<true>(GraphBuilder, View, SceneTextureParameters, SceneProxy, SceneTextures.Color.Resolve, ReflectionOutputTexture);
 					}
 					else
 					{
-						PrefilterPlanarReflection<false>(GraphBuilder, View, SceneTextureParameters, SceneProxy, SceneTextures.Color.Resolve, ViewFamilyTexture);
+						PrefilterPlanarReflection<false>(GraphBuilder, View, SceneTextureParameters, SceneProxy, SceneTextures.Color.Resolve, ReflectionOutputTexture);
 					}
 				}
+
 				GraphBuilder.Execute();
 			}
 		}
diff --git a/Engine/Source/Runtime/Renderer/Private/PostProcess/SceneRenderTargets.cpp b/Engine/Source/Runtime/Renderer/Private/PostProcess/SceneRenderTargets.cpp
--- a/Engine/Source/Runtime/Renderer/Private/PostProcess/SceneRenderTargets.cpp
+++ b/Engine/Source/Runtime/Renderer/Private/PostProcess/SceneRenderTargets.cpp
@@ -158,7 +158,7 @@ extern int32 GUseTranslucentLightingVolumes;
 
 RDG_REGISTER_BLACKBOARD_STRUCT(FSceneTextures);
 
-FSceneTextures& FSceneTextures::CreateMinimal(FRDGBuilder& GraphBuilder)
+FSceneTextures& FSceneTextures::Create(FRDGBuilder& GraphBuilder)
 {
 	auto& SceneTextures = GraphBuilder.Blackboard.Create<FSceneTextures>();
 
@@ -174,26 +174,19 @@ FSceneTextures& FSceneTextures::CreateMinimal(FRDGBuilder& GraphBuilder)
 		SceneTextures.Stencil = GraphBuilder.CreateSRV(FRDGTextureSRVDesc::CreateWithPixelFormat(SceneTextures.Depth.Target, PF_X24_G8));
 	}
 
-	return SceneTextures;
-}
-
-FSceneTextures& FSceneTextures::Create(FRDGBuilder& GraphBuilder)
-{
-	auto& SceneTextures = CreateMinimal(GraphBuilder);
-
-	const FSceneRenderTargets& SceneContext = FSceneRenderTargets::Get();
-	SceneTextures.SmallDepth = GraphBuilder.RegisterExternalTexture(SceneContext.SmallDepthZ, ERenderTargetTexture::Targetable);
-	SceneTextures.ScreenSpaceAO = CreateScreenSpaceAOTexture(GraphBuilder, SceneTextures.Extent);
-
+	if (FSceneInterface::GetShadingPath(SceneTextures.FeatureLevel) == EShadingPath::Deferred)
 	{
-		FRDGTextureDesc Desc = FVelocityRendering::GetRenderTargetDesc(SceneTextures.ShaderPlatform, SceneTextures.Extent);
+		SceneTextures.SmallDepth = GraphBuilder.RegisterExternalTexture(SceneContext.SmallDepthZ, ERenderTargetTexture::Targetable);
+		SceneTextures.ScreenSpaceAO = CreateScreenSpaceAOTexture(GraphBuilder, SceneTextures.Extent);
+
+		FRDGTextureDesc VelocityDesc = FVelocityRendering::GetRenderTargetDesc(SceneTextures.ShaderPlatform, SceneTextures.Extent);
 
 		if (FVelocityRendering::BasePassCanOutputVelocity(SceneTextures.FeatureLevel))
 		{
-			Desc.Flags |= GFastVRamConfig.GBufferVelocity;
+			VelocityDesc.Flags |= GFastVRamConfig.GBufferVelocity;
 		}
 
-		SceneTextures.Velocity = GraphBuilder.CreateTexture(Desc, TEXT("SceneVelocity"));
+		SceneTextures.Velocity = GraphBuilder.CreateTexture(VelocityDesc, TEXT("SceneVelocity"));
 	}
 
 	SceneTextures.GBufferA = TryRegisterExternalTexture(GraphBuilder, SceneContext.GBufferA);
@@ -987,7 +980,6 @@ void FSceneRenderTargets::AllocGBufferTargets(FRHICommandList& RHICmdList, EText
 		const int32 IndexC = FindGBufferTargetByName(GBufferInfo, "GBufferC");
 		const int32 IndexD = FindGBufferTargetByName(GBufferInfo, "GBufferD");
 		const int32 IndexE = FindGBufferTargetByName(GBufferInfo, "GBufferE");
-		const int32 IndexVelocity = FindGBufferTargetByName(GBufferInfo, "Velocity");
 
 		if (IndexA >= 0)
 		{
diff --git a/Engine/Source/Runtime/Renderer/Private/PostProcess/SceneRenderTargets.h b/Engine/Source/Runtime/Renderer/Private/PostProcess/SceneRenderTargets.h
--- a/Engine/Source/Runtime/Renderer/Private/PostProcess/SceneRenderTargets.h
+++ b/Engine/Source/Runtime/Renderer/Private/PostProcess/SceneRenderTargets.h
@@ -84,6 +84,8 @@ struct FCustomDepthTextures
 
 struct FMinimalSceneTextures
 {
+	static const FMinimalSceneTextures& Get(FRDGBuilder& GraphBuilder);
+
 	FRDGTextureMSAA Color{};
 	FRDGTextureMSAA Depth{};
 	FRDGTextureSRVRef Stencil{};
@@ -99,7 +101,6 @@ struct FMinimalSceneTextures
 struct RENDERER_API FSceneTextures : public FMinimalSceneTextures
 {
 	static FSceneTextures& Create(FRDGBuilder& GraphBuilder);
-	static FSceneTextures& CreateMinimal(FRDGBuilder& GraphBuilder);
 
 	static const FSceneTextures& Get(FRDGBuilder& GraphBuilder);
 
diff --git a/Engine/Source/Runtime/Renderer/Private/ReflectionEnvironmentCapture.cpp b/Engine/Source/Runtime/Renderer/Private/ReflectionEnvironmentCapture.cpp
--- a/Engine/Source/Runtime/Renderer/Private/ReflectionEnvironmentCapture.cpp
+++ b/Engine/Source/Runtime/Renderer/Private/ReflectionEnvironmentCapture.cpp
@@ -578,11 +578,11 @@ class FCopyToCubeFaceVS : public FCopyToCubeFaceShader
 {
 public:
 	DECLARE_GLOBAL_SHADER(FCopyToCubeFaceVS);
-	SHADER_USE_PARAMETER_STRUCT_WITH_LEGACY_BASE(FCopyToCubeFaceVS, FCopyToCubeFaceShader);
 
-	BEGIN_SHADER_PARAMETER_STRUCT(FParameters, )
-		SHADER_PARAMETER_STRUCT_REF(FViewUniformShaderParameters, View)
-	END_SHADER_PARAMETER_STRUCT()
+	FCopyToCubeFaceVS() = default;
+	FCopyToCubeFaceVS(const CompiledShaderInitializerType & Initializer)
+		: FCopyToCubeFaceShader(Initializer)
+	{}
 };
 
 IMPLEMENT_GLOBAL_SHADER(FCopyToCubeFaceVS, "/Engine/Private/ReflectionEnvironmentShaders.usf", "CopyToCubeFaceVS", SF_Vertex);
@@ -721,26 +721,26 @@ void CaptureSceneToScratchCubemap(FRHICommandListImmediate& RHICmdList, FSceneRe
 	{
 		SCOPED_DRAW_EVENT(RHICmdList, CubeMapCapture);
 
+		FRDGBuilder GraphBuilder(RHICmdList, RDG_EVENT_NAME("CubeMapCapture"));
+
 		// Render the scene normally for one face of the cubemap
-		SceneRenderer->Render(RHICmdList);
-		check(&RHICmdList == &FRHICommandListExecutor::GetImmediateCommandList());
-		check(IsInRenderingThread());
+		SceneRenderer->Render(GraphBuilder);
+
+		AddPass(GraphBuilder, [](FRHICommandListImmediate& InRHICmdList)
 		{
 			QUICK_SCOPE_CYCLE_COUNTER(STAT_CaptureSceneToScratchCubemap_Flush);
 			FRHICommandListExecutor::GetImmediateCommandList().ImmediateFlush(EImmediateFlushType::FlushRHIThread);
-		}
 
-		// some platforms may not be able to keep enqueueing commands like crazy, this will
-		// allow them to restart their command buffers
-		RHICmdList.SubmitCommandsAndFlushGPU();
+			// some platforms may not be able to keep enqueueing commands like crazy, this will
+			// allow them to restart their command buffers
+			InRHICmdList.SubmitCommandsAndFlushGPU();
+		});
 
 		GReflectionScratchCubemaps.Allocate(RHICmdList, CubemapSize);
 
 		const FViewInfo& View = SceneRenderer->Views[0];
 
-		FRDGBuilder GraphBuilder(RHICmdList);
-
-		FRDGTextureRef OutputTexture = GraphBuilder.RegisterExternalTexture(GReflectionScratchCubemaps.Color[0], TEXT("Color"));
+		FRDGTextureRef OutputTexture = GraphBuilder.RegisterExternalTexture(GReflectionScratchCubemaps.Color[0]);
 
 		auto* PassParameters = GraphBuilder.AllocParameters<FCopySceneColorToCubeFacePS::FParameters>();
 		PassParameters->RenderTargets[0] = FRenderTargetBinding(OutputTexture, ERenderTargetLoadAction::ENoAction, 0, CubeFace);
@@ -772,21 +772,22 @@ void CaptureSceneToScratchCubemap(FRHICommandListImmediate& RHICmdList, FSceneRe
 			PassParameters->SkyLightCaptureParameters = SkyLightParametersValue;
 		}
 
-		FSceneRenderTargets& SceneContext = FSceneRenderTargets::Get();
+		const FMinimalSceneTextures& SceneTextures = FSceneTextures::Get(GraphBuilder);
 
 		PassParameters->View = View.ViewUniformBuffer;
 		PassParameters->SceneColorSampler = TStaticSamplerState<SF_Bilinear, AM_Clamp, AM_Clamp, AM_Clamp>::GetRHI();
-		PassParameters->SceneColorTexture = GraphBuilder.RegisterExternalTexture(SceneContext.GetSceneColor(), TEXT("ColorTexture"));
+		PassParameters->SceneColorTexture = SceneTextures.Color.Target;
 		PassParameters->SceneDepthSampler = TStaticSamplerState<SF_Bilinear, AM_Clamp, AM_Clamp, AM_Clamp>::GetRHI();
-		PassParameters->SceneDepthTexture = GraphBuilder.RegisterExternalTexture(SceneContext.SceneDepthZ, TEXT("DepthTexture"));
+		PassParameters->SceneDepthTexture = SceneTextures.Depth.Target;
 
 		const int32 EffectiveSize = CubemapSize;
+		const FIntPoint SceneTextureExtent = SceneTextures.Extent;
 
 		GraphBuilder.AddPass(
 			RDG_EVENT_NAME("CopySceneToCubeFace"),
 			PassParameters,
 			ERDGPassFlags::Raster,
-			[EffectiveSize, &SceneContext, FeatureLevel, PassParameters](FRHICommandList& InRHICmdList)
+			[EffectiveSize, SceneTextureExtent, FeatureLevel, PassParameters](FRHICommandList& InRHICmdList)
 		{
 			const FIntRect ViewRect(0, 0, EffectiveSize, EffectiveSize);
 			InRHICmdList.SetViewport(0.0f, 0.0f, 0.0f, (float)EffectiveSize, (float)EffectiveSize, 1.0f);
@@ -806,10 +807,6 @@ void CaptureSceneToScratchCubemap(FRHICommandListImmediate& RHICmdList, FSceneRe
 			GraphicsPSOInit.PrimitiveType = PT_TriangleList;
 
 			SetGraphicsPipelineState(InRHICmdList, GraphicsPSOInit);
-
-			FCopyToCubeFaceVS::FParameters VertexParameters;
-			VertexParameters.View = PassParameters->View;
-			SetShaderParameters(InRHICmdList, VertexShader, VertexShader.GetVertexShader(), VertexParameters);
 			SetShaderParameters(InRHICmdList, PixelShader, PixelShader.GetPixelShader(), *PassParameters);
 
 			const int32 SupersampleCaptureFactor = FMath::Clamp(GSupersampleCaptureFactor, MinSupersampleCaptureFactor, MaxSupersampleCaptureFactor);
@@ -821,7 +818,7 @@ void CaptureSceneToScratchCubemap(FRHICommandListImmediate& RHICmdList, FSceneRe
 				ViewRect.Min.X, ViewRect.Min.Y, 
 				ViewRect.Width() * SupersampleCaptureFactor, ViewRect.Height() * SupersampleCaptureFactor,
 				FIntPoint(ViewRect.Width(), ViewRect.Height()),
-				SceneContext.GetBufferSizeXY(),
+				SceneTextureExtent,
 				VertexShader);
 		});
 
diff --git a/Engine/Source/Runtime/Renderer/Private/SceneCaptureRendering.cpp b/Engine/Source/Runtime/Renderer/Private/SceneCaptureRendering.cpp
--- a/Engine/Source/Runtime/Renderer/Private/SceneCaptureRendering.cpp
+++ b/Engine/Source/Runtime/Renderer/Private/SceneCaptureRendering.cpp
@@ -256,43 +256,26 @@ static void UpdateSceneCaptureContentDeferred_RenderThread(
 	{
 #if WANTS_DRAW_MESH_EVENTS
 		SCOPED_DRAW_EVENTF(RHICmdList, SceneCapture, TEXT("SceneCapture %s"), *EventName);
+		FRDGBuilder GraphBuilder(RHICmdList, RDG_EVENT_NAME("SceneCapture %s", *EventName));
 #else
 		SCOPED_DRAW_EVENT(RHICmdList, UpdateSceneCaptureContent_RenderThread);
+		FRDGBuilder GraphBuilder(RHICmdList, RDG_EVENT_NAME("SceneCapture"));
 #endif
 
-		const FRenderTarget* Target = SceneRenderer->ViewFamily.RenderTarget;
-
-		// TODO: Could avoid the clear by replacing with dummy black system texture.
-		FViewInfo& View = SceneRenderer->Views[0];
-
-		FRHIRenderPassInfo RPInfo(Target->GetRenderTargetTexture(), ERenderTargetActions::DontLoad_Store);
-		RPInfo.ResolveParameters = ResolveParams;
-		TransitionRenderPassTargets(RHICmdList, RPInfo);
-
-		RHICmdList.BeginRenderPass(RPInfo, TEXT("ClearSceneCaptureContent"));
-		DrawClearQuad(RHICmdList, true, FLinearColor::Black, false, 0, false, 0, Target->GetSizeXY(), View.UnscaledViewRect);
-		RHICmdList.EndRenderPass();
+		FRDGTextureRef OutputTexture = RegisterExternalTexture(GraphBuilder, RenderTarget->GetRenderTargetTexture(), TEXT("SceneCapture"));
+		AddClearRenderTargetPass(GraphBuilder, OutputTexture, FLinearColor::Black, SceneRenderer->Views[0].UnscaledViewRect);
 
 		// Render the scene normally
 		{
-			SCOPED_DRAW_EVENT(RHICmdList, RenderScene);
-			SceneRenderer->Render(RHICmdList);
+			RDG_RHI_EVENT_SCOPE(GraphBuilder, RenderScene);
+			SceneRenderer->Render(GraphBuilder);
 		}
 
-		FRDGBuilder GraphBuilder(RHICmdList);
+		if (bGenerateMips)
 		{
-			FRDGTextureRef MipTexture = GraphBuilder.RegisterExternalTexture(CreateRenderTarget(RenderTarget->GetRenderTargetTexture(), TEXT("MipGenerationInput")));
-			FRDGTextureRef OutputTexture = GraphBuilder.RegisterExternalTexture(CreateRenderTarget(RenderTargetTexture->TextureRHI, TEXT("MipGenerationOutput")));
-
-			if (bGenerateMips)
-			{
-				FGenerateMips::Execute(GraphBuilder, MipTexture, GenerateMipsParams);
-			}
-
-			// Note: When the ViewFamily.SceneCaptureSource requires scene textures (i.e. SceneCaptureSource != SCS_FinalColorLDR), the copy to RenderTarget 
-			// will be done in CopySceneCaptureComponentToTarget while the GBuffers are still alive for the frame.
-			AddCopyToResolveTargetPass(GraphBuilder, MipTexture, OutputTexture, ResolveParams);
+			FGenerateMips::Execute(GraphBuilder, OutputTexture, GenerateMipsParams);
 		}
+
 		GraphBuilder.Execute();
 	}
 
diff --git a/Engine/Source/Runtime/Renderer/Private/SceneHitProxyRendering.cpp b/Engine/Source/Runtime/Renderer/Private/SceneHitProxyRendering.cpp
--- a/Engine/Source/Runtime/Renderer/Private/SceneHitProxyRendering.cpp
+++ b/Engine/Source/Runtime/Renderer/Private/SceneHitProxyRendering.cpp
@@ -547,12 +547,10 @@ static void DoRenderHitProxies(
 }
 #endif
 
-void FMobileSceneRenderer::RenderHitProxies(FRHICommandListImmediate& RHICmdList)
+void FMobileSceneRenderer::RenderHitProxies(FRDGBuilder& GraphBuilder)
 {
-	FRDGBuilder GraphBuilder(RHICmdList);
-
 	Scene->UpdateAllPrimitiveSceneInfos(GraphBuilder);
-	
+
 	PrepareViewRectsForRendering();
 
 #if WITH_EDITOR
@@ -578,22 +576,18 @@ void FMobileSceneRenderer::RenderHitProxies(FRHICommandListImmediate& RHICmdList
 	GEngine->GetPostRenderDelegate().Broadcast();
 #endif
 
-	GraphBuilder.Execute();
-
 	if (bDeferredShading)
 	{
 		// Release the original reference on the scene render targets
 		FSceneRenderTargets& SceneContext = FSceneRenderTargets::Get();
-		SceneContext.AdjustGBufferRefCount(RHICmdList, -1);
+		SceneContext.AdjustGBufferRefCount(GraphBuilder.RHICmdList, -1);
 	}
 }
 
-void FDeferredShadingSceneRenderer::RenderHitProxies(FRHICommandListImmediate& RHICmdList)
+void FDeferredShadingSceneRenderer::RenderHitProxies(FRDGBuilder& GraphBuilder)
 {
 	static const bool bNaniteEnabled = DoesPlatformSupportNanite(GMaxRHIShaderPlatform);
 
-	FRDGBuilder GraphBuilder(RHICmdList);
-
 	Scene->UpdateAllPrimitiveSceneInfos(GraphBuilder);
 
 	FGPUSceneScopeBeginEndHelper GPUSceneScopeBeginEndHelper(Scene->GPUScene, GPUSceneDynamicContext, *Scene);
@@ -719,8 +713,6 @@ void FDeferredShadingSceneRenderer::RenderHitProxies(FRHICommandListImmediate& R
 	GEngine->GetPostRenderDelegate().Broadcast();
 
 #endif
-
-	GraphBuilder.Execute();
 }
 
 #if WITH_EDITOR
diff --git a/Engine/Source/Runtime/Renderer/Private/SceneRendering.cpp b/Engine/Source/Runtime/Renderer/Private/SceneRendering.cpp
--- a/Engine/Source/Runtime/Renderer/Private/SceneRendering.cpp
+++ b/Engine/Source/Runtime/Renderer/Private/SceneRendering.cpp
@@ -3648,54 +3648,66 @@ static void RenderViewFamily_RenderThread(FRHICommandListImmediate& RHICmdList,
 	// update any resources that needed a deferred update
 	FDeferredUpdateResource::UpdateResources(RHICmdList);
 
-	if(SceneRenderer->ViewFamily.EngineShowFlags.OnScreenDebug)
+	const FSceneViewFamily& ViewFamily = SceneRenderer->ViewFamily;
+
+	if(ViewFamily.EngineShowFlags.OnScreenDebug)
 	{
 		GRenderTargetPool.SetEventRecordingActive(true);
 	}
 
 	{
 		SCOPE_CYCLE_COUNTER(STAT_TotalSceneRenderingTime);
-	
-		if(SceneRenderer->ViewFamily.EngineShowFlags.HitProxies)
-		{
-			// Render the scene's hit proxies.
-			SceneRenderer->RenderHitProxies(RHICmdList);
-		}
-		else
+
 		{
-			// Render the scene.
-			SceneRenderer->Render(RHICmdList);
-		}
+			FRDGBuilder GraphBuilder(
+				RHICmdList,
+				RDG_EVENT_NAME("SceneRenderer_%s(ViewFamily=%s)",
+					ViewFamily.EngineShowFlags.HitProxies ? TEXT("RenderHitProxies") : TEXT("Render"),
+					ViewFamily.bResolveScene ? TEXT("Primary") : TEXT("Auxiliary")
+				)
+			);
+
+			if (ViewFamily.EngineShowFlags.HitProxies)
+			{
+				// Render the scene's hit proxies.
+				SceneRenderer->RenderHitProxies(GraphBuilder);
+			}
+			else
+			{
+				// Render the scene.
+				SceneRenderer->Render(GraphBuilder);
+			}
+			
+			if (SceneRenderer->Views.Num() > 0 && !ViewFamily.EngineShowFlags.HitProxies)
+			{
+				FHairStrandsBookmarkParameters& Parameters = *GraphBuilder.AllocObject<FHairStrandsBookmarkParameters>(CreateHairStrandsBookmarkParameters(SceneRenderer->Views));
+				if (Parameters.bHasElements)
+				{
+					RunHairStrandsBookmark(GraphBuilder, EHairStrandsBookmark::ProcessEndOfFrame, Parameters);
+				}
+			}
 
-		CSV_SCOPED_TIMING_STAT_EXCLUSIVE(PostRenderCleanUp);
+			GraphBuilder.Execute();
 
-		// Only reset per-frame scene state once all views have processed their frame, including those in planar reflections
-		for (int32 CacheType = 0; CacheType < UE_ARRAY_COUNT(SceneRenderer->Scene->DistanceFieldSceneData.PrimitiveModifiedBounds); CacheType++)
-		{
-			SceneRenderer->Scene->DistanceFieldSceneData.PrimitiveModifiedBounds[CacheType].Reset();
-		}
-		if (SceneRenderer->Scene->LumenSceneData)
-		{
-			SceneRenderer->Scene->LumenSceneData->PrimitiveModifiedBounds.Reset();
-		}
+			CSV_SCOPED_TIMING_STAT_EXCLUSIVE(PostRenderCleanUp);
 
-		if (SceneRenderer->Views.Num() > 0 && !SceneRenderer->ViewFamily.EngineShowFlags.HitProxies)
-		{
-			FHairStrandsBookmarkParameters Parameters = CreateHairStrandsBookmarkParameters(SceneRenderer->Views);
-			if (Parameters.bHasElements)
+			// Only reset per-frame scene state once all views have processed their frame, including those in planar reflections
+			for (int32 CacheType = 0; CacheType < UE_ARRAY_COUNT(SceneRenderer->Scene->DistanceFieldSceneData.PrimitiveModifiedBounds); CacheType++)
 			{
-				FRDGBuilder GraphBuilder(RHICmdList);
-				RunHairStrandsBookmark(GraphBuilder, EHairStrandsBookmark::ProcessEndOfFrame, Parameters);
-				GraphBuilder.Execute();
+				SceneRenderer->Scene->DistanceFieldSceneData.PrimitiveModifiedBounds[CacheType].Reset();
+			}
+			if (SceneRenderer->Scene->LumenSceneData)
+			{
+				SceneRenderer->Scene->LumenSceneData->PrimitiveModifiedBounds.Reset();
 			}
-		}
 
-		// Immediately issue EndFrame() for all extensions in case any of the outstanding tasks they issued getting out of this frame
-		extern TSet<IPersistentViewUniformBufferExtension*> PersistentViewUniformBufferExtensions;
+			// Immediately issue EndFrame() for all extensions in case any of the outstanding tasks they issued getting out of this frame
+			extern TSet<IPersistentViewUniformBufferExtension*> PersistentViewUniformBufferExtensions;
 
-		for (IPersistentViewUniformBufferExtension* Extension : PersistentViewUniformBufferExtensions)
-		{
-			Extension->EndFrame();
+			for (IPersistentViewUniformBufferExtension* Extension : PersistentViewUniformBufferExtensions)
+			{
+				Extension->EndFrame();
+			}
 		}
 
 #if STATS
@@ -3718,7 +3730,6 @@ static void RenderViewFamily_RenderThread(FRHICommandListImmediate& RHICmdList,
 			SET_MEMORY_STAT(STAT_LightInteractionMemory, FLightPrimitiveInteraction::GetMemoryPoolSize());
 		}
 #endif
-		
 
 		GRenderTargetPool.SetEventRecordingActive(false);
 
diff --git a/Engine/Source/Runtime/Renderer/Private/SceneRendering.h b/Engine/Source/Runtime/Renderer/Private/SceneRendering.h
--- a/Engine/Source/Runtime/Renderer/Private/SceneRendering.h
+++ b/Engine/Source/Runtime/Renderer/Private/SceneRendering.h
@@ -1759,8 +1759,8 @@ public:
 
 	// FSceneRenderer interface
 
-	virtual void Render(FRHICommandListImmediate& RHICmdList) = 0;
-	virtual void RenderHitProxies(FRHICommandListImmediate& RHICmdList) {}
+	virtual void Render(FRDGBuilder& GraphBuilder) = 0;
+	virtual void RenderHitProxies(FRDGBuilder& GraphBuilder) {}
 
 	/** Creates a scene renderer based on the current feature level. */
 	static FSceneRenderer* CreateSceneRenderer(const FSceneViewFamily* InViewFamily, FHitProxyConsumer* HitProxyConsumer);
@@ -2074,9 +2074,9 @@ public:
 
 	// FSceneRenderer interface
 
-	virtual void Render(FRHICommandListImmediate& RHICmdList) override;
+	virtual void Render(FRDGBuilder& GraphBuilder) override;
 
-	virtual void RenderHitProxies(FRHICommandListImmediate& RHICmdList) override;
+	virtual void RenderHitProxies(FRDGBuilder& GraphBuilder) override;
 
 	void RenderInverseOpacity(FRDGBuilder& GraphBuilder, const FViewInfo& View);
 
diff --git a/Engine/Source/Runtime/Renderer/Private/ScreenPass.cpp b/Engine/Source/Runtime/Renderer/Private/ScreenPass.cpp
--- a/Engine/Source/Runtime/Renderer/Private/ScreenPass.cpp
+++ b/Engine/Source/Runtime/Renderer/Private/ScreenPass.cpp
@@ -26,7 +26,7 @@ FRDGTextureRef TryCreateViewFamilyTexture(FRDGBuilder& GraphBuilder, const FScen
 	FRDGTextureRef Texture = nullptr;
 	if (TextureRHI)
 	{
-		Texture = GraphBuilder.RegisterExternalTexture(CreateRenderTarget(TextureRHI, TEXT("ViewFamilyTexture")));
+		Texture = RegisterExternalTexture(GraphBuilder, TextureRHI, TEXT("ViewFamilyTexture"));
 		GraphBuilder.SetTextureAccessFinal(Texture, ERHIAccess::RTV);
 	}
 	return Texture;
diff --git a/Engine/Source/Runtime/Renderer/Public/MeshPassProcessor.h b/Engine/Source/Runtime/Renderer/Public/MeshPassProcessor.h
--- a/Engine/Source/Runtime/Renderer/Public/MeshPassProcessor.h
+++ b/Engine/Source/Runtime/Renderer/Public/MeshPassProcessor.h
@@ -1564,21 +1564,25 @@ public:
 		return ViewUniformBuffer;
 	}
 
+	UE_DEPRECATED(5.0, "SetInstancedViewUniformBuffer is deprecated. Use View.GetShaderParameters() and bind on an RDG pass instead.")
 	FORCEINLINE_DEBUGGABLE void SetInstancedViewUniformBuffer(const TUniformBufferRef<FInstancedViewUniformShaderParameters>& InViewUniformBuffer)
 	{
 		InstancedViewUniformBuffer = InViewUniformBuffer;
 	}
 
+	UE_DEPRECATED(5.0, "GetInstancedViewUniformBuffer is deprecated. Use View.GetShaderParameters() and bind on an RDG pass instead.")
 	FORCEINLINE_DEBUGGABLE const FRHIUniformBuffer* GetInstancedViewUniformBuffer() const
 	{
 		return InstancedViewUniformBuffer != nullptr ? InstancedViewUniformBuffer : ViewUniformBuffer;
 	}
 
+	UE_DEPRECATED(5.0, "SetPassUniformBuffer is deprecated. Use a static uniform buffer and bind on an RDG pass instead.")
 	FORCEINLINE_DEBUGGABLE void SetPassUniformBuffer(const FUniformBufferRHIRef& InPassUniformBuffer)
 	{
 		PassUniformBuffer = InPassUniformBuffer;
 	}
 
+	UE_DEPRECATED(5.0, "GetPassUniformBuffer is deprecated. Use a static uniform buffer and bind on an RDG pass instead.")
 	FORCEINLINE_DEBUGGABLE FRHIUniformBuffer* GetPassUniformBuffer() const
 	{
 		return PassUniformBuffer;
