Add support for light falloff radius and exponent in the path tracer and Lightmass2

These controls are non-physical but still heavily used, so will help with comparing between the path tracer and lit modes

Also fixed Lightmass2 plugin from recently added IES support (although this path does not support IES just yet)

Merged the mobiility and type fields into flags to reduce memory usage a bit and moved the definitions that need to be shared between c++ and HLSL to a shared header file

diff --git a/Engine/Plugins/Experimental/GPULightmass/Shaders/Private/LightmapPathTracing.usf b/Engine/Plugins/Experimental/GPULightmass/Shaders/Private/LightmapPathTracing.usf
--- a/Engine/Plugins/Experimental/GPULightmass/Shaders/Private/LightmapPathTracing.usf
+++ b/Engine/Plugins/Experimental/GPULightmass/Shaders/Private/LightmapPathTracing.usf
@@ -264,7 +264,7 @@ void PathTracingKernel(
 				PreviousCdfValue = LightPickingCdf[LightId];
 
 				// GPULightmass records direct contribution for only static lights
-				if (Bounce > 1 || (Bounce == 1 && SceneLightsData.Mobility[LightId] == 0))
+				if (Bounce > 1 || (Bounce == 1 && !IsStationary(LightId)))
 				{
 					AccumulateRadiance(Radiance, PathThroughput * LightRadiance * MISWeightRobust(PrevMaterialPdf, LightPdf * LightPickPdf));
 				}
@@ -439,7 +439,7 @@ void PathTracingKernel(
 						else
 						{
 							// GPULightmass records contribution of direct lighting separately, and only for static lights
-							if (SceneLightsData.Mobility[LightId] == 0)
+							if (!IsStationary(LightId))
 							{
 								DirectLightingNEERadianceValue += MISWeight * PathThroughput * RadianceOverPdf * MaterialEval;
 								DirectLightingNEERadianceDirection = LightRay.Direction;
diff --git a/Engine/Plugins/Experimental/GPULightmass/Source/GPULightmass/Private/LightmapRayTracing.h b/Engine/Plugins/Experimental/GPULightmass/Source/GPULightmass/Private/LightmapRayTracing.h
--- a/Engine/Plugins/Experimental/GPULightmass/Source/GPULightmass/Private/LightmapRayTracing.h
+++ b/Engine/Plugins/Experimental/GPULightmass/Source/GPULightmass/Private/LightmapRayTracing.h
@@ -72,6 +72,8 @@ class FLightmapPathTracingRGS : public FGlobalShader
 		SHADER_PARAMETER_STRUCT_REF(FSkyLightData, SkyLight)
 		SHADER_PARAMETER_STRUCT_REF(FViewUniformShaderParameters, ViewUniformBuffer)
 		SHADER_PARAMETER_STRUCT_REF(FIrradianceCachingParameters, IrradianceCachingParameters)
+		SHADER_PARAMETER_TEXTURE(Texture3D, IESTexture)
+		SHADER_PARAMETER_SAMPLER(SamplerState, IESTextureSampler)
 	END_SHADER_PARAMETER_STRUCT()
 };
 
@@ -122,6 +124,8 @@ class FVolumetricLightmapPathTracingRGS : public FGlobalShader
 		SHADER_PARAMETER_STRUCT_REF(FSkyLightData, SkyLight)
 		SHADER_PARAMETER_STRUCT_REF(FViewUniformShaderParameters, ViewUniformBuffer)
 		SHADER_PARAMETER_STRUCT_REF(FIrradianceCachingParameters, IrradianceCachingParameters)
+		SHADER_PARAMETER_TEXTURE(Texture3D, IESTexture)
+		SHADER_PARAMETER_SAMPLER(SamplerState, IESTextureSampler)
 	END_SHADER_PARAMETER_STRUCT()
 };
 
diff --git a/Engine/Plugins/Experimental/GPULightmass/Source/GPULightmass/Private/LightmapRenderer.cpp b/Engine/Plugins/Experimental/GPULightmass/Source/GPULightmass/Private/LightmapRenderer.cpp
--- a/Engine/Plugins/Experimental/GPULightmass/Source/GPULightmass/Private/LightmapRenderer.cpp
+++ b/Engine/Plugins/Experimental/GPULightmass/Source/GPULightmass/Private/LightmapRenderer.cpp
@@ -176,13 +176,12 @@ FPathTracingLightData SetupPathTracingLightParameters(const GPULightmass::FLight
 	// Prepend SkyLight to light buffer
 	// WARNING: Until ray payload encodes Light data buffer, the execution depends on this ordering!
 	uint32 SkyLightIndex = 0;
-	LightParameters.Type[SkyLightIndex] = 0;
 	LightParameters.Color[SkyLightIndex] = LightScene.SkyLight.IsSet() ? FVector(LightScene.SkyLight->Color) : FVector(0);
-	LightParameters.Mobility[SkyLightIndex] = (LightScene.SkyLight.IsSet() && LightScene.SkyLight->bStationary) ? 1 : 0;
-	uint32 Transmission = 1;
-	uint8 LightingChannelMask = 0b111;
-	LightParameters.Flags[SkyLightIndex] = Transmission & 0x01;
-	LightParameters.Flags[SkyLightIndex] |= (LightingChannelMask & 0x7) << 1;
+	LightParameters.Flags[SkyLightIndex] = PATHTRACER_FLAG_TRANSMISSION_MASK;
+	LightParameters.Flags[SkyLightIndex] |= PATHTRACER_FLAG_LIGHTING_CHANNEL_MASK;
+	bool SkyLightIsStationary = LightScene.SkyLight.IsSet() && LightScene.SkyLight->bStationary;
+	LightParameters.Flags[SkyLightIndex] |= SkyLightIsStationary ? PATHTRACER_FLAG_STATIONARY_MASK : 0;
+	LightParameters.Flags[LightParameters.Count] |= PATHTRACING_LIGHT_SKY;
 	LightParameters.Count++;
 
 	uint32 MaxLightCount = RAY_TRACING_LIGHT_COUNT_MAXIMUM;
@@ -191,15 +190,16 @@ FPathTracingLightData SetupPathTracingLightParameters(const GPULightmass::FLight
 	{
 		if (LightParameters.Count < MaxLightCount)
 		{
-			LightParameters.Type[LightParameters.Count] = 2;
 			LightParameters.Normal[LightParameters.Count] = -Light.Direction;
 			LightParameters.Color[LightParameters.Count] = FVector(Light.Color);
 			LightParameters.Dimensions[LightParameters.Count] = FVector(0.0f, 0.0f, FMath::Sin(0.5f * FMath::DegreesToRadians(Light.LightSourceAngle)));
 			LightParameters.Attenuation[LightParameters.Count] = 1.0;
-			LightParameters.Mobility[LightParameters.Count] = Light.bStationary ? 1 : 0;
+			LightParameters.IESTextureSlice[LightParameters.Count] = -1;
 
-			LightParameters.Flags[LightParameters.Count] = Transmission & 0x01;
-			LightParameters.Flags[LightParameters.Count] |= (LightingChannelMask & 0x7) << 1;
+			LightParameters.Flags[LightParameters.Count] = PATHTRACER_FLAG_TRANSMISSION_MASK;
+			LightParameters.Flags[LightParameters.Count] |= PATHTRACER_FLAG_LIGHTING_CHANNEL_MASK;
+			LightParameters.Flags[LightParameters.Count] |= Light.bStationary ? PATHTRACER_FLAG_STATIONARY_MASK : 0;
+			LightParameters.Flags[LightParameters.Count] |= PATHTRACING_LIGHT_DIRECTIONAL;
 
 			LightParameters.Count++;
 		}
@@ -209,15 +209,18 @@ FPathTracingLightData SetupPathTracingLightParameters(const GPULightmass::FLight
 	{
 		if (LightParameters.Count < MaxLightCount)
 		{
-			LightParameters.Type[LightParameters.Count] = 1;
 			LightParameters.Position[LightParameters.Count] = Light.Position;
 			LightParameters.Color[LightParameters.Count] = FVector(Light.Color);
 			LightParameters.Dimensions[LightParameters.Count] = FVector(0.0, 0.0, Light.SourceRadius);
-			LightParameters.Attenuation[LightParameters.Count] = Light.AttenuationRadius;
-			LightParameters.Mobility[LightParameters.Count] = Light.bStationary ? 1 : 0;
+			LightParameters.Attenuation[LightParameters.Count] = 1.0f / Light.AttenuationRadius;
+			LightParameters.FalloffExponent[LightParameters.Count] = Light.FalloffExponent;
+			LightParameters.IESTextureSlice[LightParameters.Count] = -1;
 
-			LightParameters.Flags[LightParameters.Count] = Transmission & 0x01;
-			LightParameters.Flags[LightParameters.Count] |= (LightingChannelMask & 0x7) << 1;
+			LightParameters.Flags[LightParameters.Count] = PATHTRACER_FLAG_TRANSMISSION_MASK;
+			LightParameters.Flags[LightParameters.Count] |= PATHTRACER_FLAG_LIGHTING_CHANNEL_MASK;
+			LightParameters.Flags[LightParameters.Count] |= Light.IsInverseSquared ? 0 : PATHTRACER_FLAG_NON_INVERSE_SQUARE_FALLOFF_MASK;
+			LightParameters.Flags[LightParameters.Count] |= Light.bStationary ? PATHTRACER_FLAG_STATIONARY_MASK : 0;
+			LightParameters.Flags[LightParameters.Count] |= PATHTRACING_LIGHT_POINT;
 
 			LightParameters.Count++;
 		}
@@ -227,16 +230,19 @@ FPathTracingLightData SetupPathTracingLightParameters(const GPULightmass::FLight
 	{
 		if (LightParameters.Count < MaxLightCount)
 		{
-			LightParameters.Type[LightParameters.Count] = 4;
 			LightParameters.Position[LightParameters.Count] = Light.Position;
 			LightParameters.Normal[LightParameters.Count] = Light.Direction;
 			LightParameters.Color[LightParameters.Count] = FVector(Light.Color);
 			LightParameters.Dimensions[LightParameters.Count] = FVector(Light.SpotAngles, Light.SourceRadius);
-			LightParameters.Attenuation[LightParameters.Count] = Light.AttenuationRadius;
-			LightParameters.Mobility[LightParameters.Count] = Light.bStationary ? 1 : 0;
+			LightParameters.Attenuation[LightParameters.Count] = 1.0f / Light.AttenuationRadius;
+			LightParameters.FalloffExponent[LightParameters.Count] = Light.FalloffExponent;
+			LightParameters.IESTextureSlice[LightParameters.Count] = -1;
 
-			LightParameters.Flags[LightParameters.Count] = Transmission & 0x01;
-			LightParameters.Flags[LightParameters.Count] |= (LightingChannelMask & 0x7) << 1;
+			LightParameters.Flags[LightParameters.Count] = PATHTRACER_FLAG_TRANSMISSION_MASK;
+			LightParameters.Flags[LightParameters.Count] |= PATHTRACER_FLAG_LIGHTING_CHANNEL_MASK;
+			LightParameters.Flags[LightParameters.Count] |= Light.IsInverseSquared ? 0 : PATHTRACER_FLAG_NON_INVERSE_SQUARE_FALLOFF_MASK;
+			LightParameters.Flags[LightParameters.Count] |= Light.bStationary ? PATHTRACER_FLAG_STATIONARY_MASK : 0;
+			LightParameters.Flags[LightParameters.Count] |= PATHTRACING_LIGHT_SPOT;
 
 			LightParameters.Count++;
 		}
@@ -246,7 +252,6 @@ FPathTracingLightData SetupPathTracingLightParameters(const GPULightmass::FLight
 	{
 		if (LightParameters.Count < MaxLightCount)
 		{
-			LightParameters.Type[LightParameters.Count] = 3;
 			LightParameters.Position[LightParameters.Count] = Light.Position;
 			LightParameters.Normal[LightParameters.Count] = Light.Direction;
 			LightParameters.dPdu[LightParameters.Count] = FVector::CrossProduct(Light.Tangent, -Light.Direction);
@@ -257,14 +262,16 @@ FPathTracingLightData SetupPathTracingLightParameters(const GPULightmass::FLight
 			LightParameters.Color[LightParameters.Count] = FVector(LightColor);
 
 			LightParameters.Dimensions[LightParameters.Count] = FVector(Light.SourceWidth, Light.SourceHeight, 0.0f);
-			LightParameters.Attenuation[LightParameters.Count] = Light.AttenuationRadius;
+			LightParameters.Attenuation[LightParameters.Count] = 1.0f / Light.AttenuationRadius;
 			LightParameters.RectLightBarnCosAngle[LightParameters.Count] = FMath::Cos(FMath::DegreesToRadians(Light.BarnDoorAngle));
 			LightParameters.RectLightBarnLength[LightParameters.Count] = Light.BarnDoorLength;
 
-			LightParameters.Mobility[LightParameters.Count] = Light.bStationary ? 1 : 0;
+			LightParameters.IESTextureSlice[LightParameters.Count] = -1;
 
-			LightParameters.Flags[LightParameters.Count] = Transmission & 0x01;
-			LightParameters.Flags[LightParameters.Count] |= (LightingChannelMask & 0x7) << 1;
+			LightParameters.Flags[LightParameters.Count] = PATHTRACER_FLAG_TRANSMISSION_MASK;
+			LightParameters.Flags[LightParameters.Count] |= PATHTRACER_FLAG_LIGHTING_CHANNEL_MASK;
+			LightParameters.Flags[LightParameters.Count] |= Light.bStationary ? PATHTRACER_FLAG_STATIONARY_MASK : 0;
+			LightParameters.Flags[LightParameters.Count] |= PATHTRACING_LIGHT_RECT;
 
 			LightParameters.Count++;
 		}
@@ -1735,6 +1742,10 @@ void FLightmapRenderer::Finalize(FRDGBuilder& GraphBuilder)
 											PassParameters->SkyLight = SkyLightDataUniformBuffer;
 										}
 
+										// TODO: find a way to share IES atlas with path tracer ...
+										PassParameters->IESTexture = GWhiteTexture->TextureRHI;
+										PassParameters->IESTextureSampler = TStaticSamplerState<SF_Bilinear, AM_Clamp, AM_Clamp, AM_Clamp>::GetRHI();
+
 										FLightmapPathTracingRGS::FPermutationDomain PermutationVector;
 										PermutationVector.Set<FLightmapPathTracingRGS::FUseFirstBounceRayGuiding>(bUseFirstBounceRayGuiding);
 										PermutationVector.Set<FLightmapPathTracingRGS::FUseIrradianceCaching>(Scene->Settings->bUseIrradianceCaching);
diff --git a/Engine/Plugins/Experimental/GPULightmass/Source/GPULightmass/Private/Scene/Lights.cpp b/Engine/Plugins/Experimental/GPULightmass/Source/GPULightmass/Private/Scene/Lights.cpp
--- a/Engine/Plugins/Experimental/GPULightmass/Source/GPULightmass/Private/Scene/Lights.cpp
+++ b/Engine/Plugins/Experimental/GPULightmass/Source/GPULightmass/Private/Scene/Lights.cpp
@@ -72,6 +72,8 @@ FPointLightRenderState::FPointLightRenderState(UPointLightComponent* PointLightC
 	AttenuationRadius = PointLightComponent->AttenuationRadius;
 	SourceRadius = PointLightComponent->SourceRadius;
 	ShadowMapChannel = PointLightComponent->PreviewShadowMapChannel;
+	FalloffExponent = PointLightComponent->LightFalloffExponent;
+	IsInverseSquared = PointLightComponent->bUseInverseSquaredFalloff;
 }
 
 FSpotLightBuildInfo::FSpotLightBuildInfo(USpotLightComponent* SpotLightComponent)
@@ -142,6 +144,8 @@ FSpotLightRenderState::FSpotLightRenderState(USpotLightComponent* SpotLightCompo
 	AttenuationRadius = SpotLightComponent->AttenuationRadius;
 	SourceRadius = SpotLightComponent->SourceRadius;
 	ShadowMapChannel = SpotLightComponent->PreviewShadowMapChannel;
+	FalloffExponent = SpotLightComponent->LightFalloffExponent;
+	IsInverseSquared = SpotLightComponent->bUseInverseSquaredFalloff;
 }
 
 FRectLightBuildInfo::FRectLightBuildInfo(URectLightComponent* RectLightComponent)
diff --git a/Engine/Plugins/Experimental/GPULightmass/Source/GPULightmass/Private/Scene/Lights.h b/Engine/Plugins/Experimental/GPULightmass/Source/GPULightmass/Private/Scene/Lights.h
--- a/Engine/Plugins/Experimental/GPULightmass/Source/GPULightmass/Private/Scene/Lights.h
+++ b/Engine/Plugins/Experimental/GPULightmass/Source/GPULightmass/Private/Scene/Lights.h
@@ -217,6 +217,8 @@ struct FPointLightRenderState : public FLocalLightRenderState
 	FLinearColor Color;
 	float AttenuationRadius;
 	float SourceRadius;
+	float FalloffExponent;
+	bool IsInverseSquared;
 
 	virtual FLightShaderParameters GetLightShaderParameters() const override;
 };
@@ -234,6 +236,9 @@ struct FSpotLightRenderState : public FLocalLightRenderState
 	FLinearColor Color;
 	float AttenuationRadius;
 	float SourceRadius;
+	float FalloffExponent;
+	bool IsInverseSquared;
+
 
 	virtual FLightShaderParameters GetLightShaderParameters() const override;
 };
diff --git a/Engine/Plugins/Experimental/GPULightmass/Source/GPULightmass/Private/VolumetricLightmap.cpp b/Engine/Plugins/Experimental/GPULightmass/Source/GPULightmass/Private/VolumetricLightmap.cpp
--- a/Engine/Plugins/Experimental/GPULightmass/Source/GPULightmass/Private/VolumetricLightmap.cpp
+++ b/Engine/Plugins/Experimental/GPULightmass/Source/GPULightmass/Private/VolumetricLightmap.cpp
@@ -558,6 +558,10 @@ void FVolumetricLightmapRenderer::BackgroundTick()
 					Parameters.SkyLight = SkyLightDataUniformBuffer;
 				}
 
+				// TODO: find a way to share IES atlas with path tracer ...
+				Parameters.IESTexture = GWhiteTexture->TextureRHI;
+				Parameters.IESTextureSampler = TStaticSamplerState<SF_Bilinear, AM_Clamp, AM_Clamp, AM_Clamp>::GetRHI();
+
 				FRayTracingShaderBindingsWriter GlobalResources;
 				SetShaderParameters(GlobalResources, RayGenShader, Parameters);
 
diff --git a/Engine/Shaders/Private/PathTracing/Light/PathTracingLightCommon.ush b/Engine/Shaders/Private/PathTracing/Light/PathTracingLightCommon.ush
--- a/Engine/Shaders/Private/PathTracing/Light/PathTracingLightCommon.ush
+++ b/Engine/Shaders/Private/PathTracing/Light/PathTracingLightCommon.ush
@@ -6,13 +6,6 @@ LightCommon.usf: Common utility functions for light sampling
 
 #pragma once
 
-// These ids must be kept in sync with the logic in PathTracer.cpp
-#define PATHTRACING_LIGHT_SKY				0
-#define PATHTRACING_LIGHT_POINT				1
-#define PATHTRACING_LIGHT_DIRECTIONAL		2
-#define PATHTRACING_LIGHT_RECT				3
-#define PATHTRACING_LIGHT_SPOT				4
-
 int GetLightId(FMaterialClosestHitPayload RayHitInfo)
 {
 	return -1;
@@ -20,27 +13,27 @@ int GetLightId(FMaterialClosestHitPayload RayHitInfo)
 
 bool IsEnvironmentLight(int LightId)
 {
-	return SceneLightsData.Type[LightId] == PATHTRACING_LIGHT_SKY;
+	return (SceneLightsData.Flags[LightId] & PATHTRACER_FLAG_TYPE_MASK) == PATHTRACING_LIGHT_SKY;
 }
 
 bool IsPointLight(int LightId)
 {
-	return SceneLightsData.Type[LightId] == PATHTRACING_LIGHT_POINT;
+	return (SceneLightsData.Flags[LightId] & PATHTRACER_FLAG_TYPE_MASK) == PATHTRACING_LIGHT_POINT;
 }
 
 bool IsDirectionalLight(int LightId)
 {
-	return SceneLightsData.Type[LightId] == PATHTRACING_LIGHT_DIRECTIONAL;
+	return (SceneLightsData.Flags[LightId] & PATHTRACER_FLAG_TYPE_MASK) == PATHTRACING_LIGHT_DIRECTIONAL;
 }
 
 bool IsRectLight(int LightId)
 {
-	return SceneLightsData.Type[LightId] == PATHTRACING_LIGHT_RECT;
+	return (SceneLightsData.Flags[LightId] & PATHTRACER_FLAG_TYPE_MASK) == PATHTRACING_LIGHT_RECT;
 }
 
 bool IsSpotLight(int LightId)
 {
-	return SceneLightsData.Type[LightId] == PATHTRACING_LIGHT_SPOT;
+	return (SceneLightsData.Flags[LightId] & PATHTRACER_FLAG_TYPE_MASK) == PATHTRACING_LIGHT_SPOT;
 }
 
 // A light is a physical light if it can be intersected by a ray.
@@ -130,15 +123,27 @@ float GetRectLightBarnLength(int LightId)
 }
 
 #ifndef USE_ATTENUATION_TERM
-	#define USE_ATTENUATION_TERM 0
+	#define USE_ATTENUATION_TERM 1
 #endif
 
-float ComputeAttenuationFalloff(float Distance, float Attenuation)
+float ComputeAttenuationFalloff(float DistanceSquared, int LightId)
 {
 #if USE_ATTENUATION_TERM
 	// Mirrors GetLocalLightAttenuation() custom attenuation controls
 	// #dxr_todo: UE-72508: encapsulate this function in a shared space
-	return Square(saturate(1.0 - Square(Square(Distance) / Square(Attenuation))));
+	float InvAttenuationRadius = GetAttenuation(LightId);
+	float NormalizeDistanceSquared = DistanceSquared * Square(InvAttenuationRadius);
+	if ((SceneLightsData.Flags[LightId] & PATHTRACER_FLAG_NON_INVERSE_SQUARE_FALLOFF_MASK) == 0)
+	{
+		return Square(saturate(1.0 - Square(NormalizeDistanceSquared)));
+	}
+	else
+	{
+		// roughly cancel out "native" square distance falloff before applying exponent based falloff function
+        // this appears to match the behavior of the deferred lighting passes
+		float FalloffExponent = SceneLightsData.FalloffExponent[LightId];
+		return DistanceSquared * pow(1.0 - saturate(NormalizeDistanceSquared), FalloffExponent);
+	}
 #else
 	return 1.0;
 #endif
@@ -171,10 +176,15 @@ float ComputeIESAttenuation(int LightId, float3 WorldPosition)
 
 bool HasTransmission(int LightId)
 {
-	return SceneLightsData.Flags[LightId] & 0x01;
+	return (SceneLightsData.Flags[LightId] & PATHTRACER_FLAG_TRANSMISSION_MASK) != 0;
 }
 
 uint GetLightingChannelMask(int LightId)
 {
-	return (SceneLightsData.Flags[LightId] >> 1) & 0x7;
-}
\ No newline at end of file
+	return SceneLightsData.Flags[LightId] & PATHTRACER_FLAG_LIGHTING_CHANNEL_MASK;
+}
+
+bool IsStationary(int LightId)
+{
+	return (SceneLightsData.Flags[LightId] & PATHTRACER_FLAG_STATIONARY_MASK) != 0;
+}
diff --git a/Engine/Shaders/Private/PathTracing/Light/PathTracingLightSampling.ush b/Engine/Shaders/Private/PathTracing/Light/PathTracingLightSampling.ush
--- a/Engine/Shaders/Private/PathTracing/Light/PathTracingLightSampling.ush
+++ b/Engine/Shaders/Private/PathTracing/Light/PathTracingLightSampling.ush
@@ -22,7 +22,7 @@ LightSampling.usf: Light sampling for next-event estimation
 // Returns float4 packed with Radiance and Pdf
 float4 TraceLight(RayDesc Ray, int LightId, float HitT)
 {
-	switch (SceneLightsData.Type[LightId])
+	switch (SceneLightsData.Flags[LightId] & PATHTRACER_FLAG_TYPE_MASK)
 	{
 		case PATHTRACING_LIGHT_SKY:			return SkyLight_TraceLight(Ray, LightId, HitT);
 		case PATHTRACING_LIGHT_POINT:		return PointLight_TraceLight(Ray, LightId, HitT);
@@ -44,7 +44,7 @@ void SampleLight(
 	out float3 OutRadianceOverPdf,
 	out float OutPdf)
 {
-	switch (SceneLightsData.Type[LightId])
+	switch (SceneLightsData.Flags[LightId] & PATHTRACER_FLAG_TYPE_MASK)
 	{
 		case PATHTRACING_LIGHT_SKY:
 			SkyLight_SampleLight(LightId, RandSample, WorldPos, WorldNormal, OutDirection, OutDistance, OutRadianceOverPdf, OutPdf);
@@ -75,7 +75,7 @@ void GenerateLightRay(RayDesc Ray, int LightId, float3 LightUV, inout RayDesc Ou
 {
 	OutLightRay = (RayDesc)0; // #dxr_todo: Return false in the case of no matching light type and handle case from caller rather than blindly using a degenerate ray
 
-	switch (SceneLightsData.Type[LightId])
+	switch (SceneLightsData.Flags[LightId] & PATHTRACER_FLAG_TYPE_MASK)
 	{
 		case PATHTRACING_LIGHT_SKY:
 			SkyLight_GenerateLightRay(Ray, LightId, LightUV, OutLightRay);
@@ -110,7 +110,7 @@ void EstimateLight(
 	OutIrradiance = 0.0f;
 	if ((GetLightingChannelMask(LightId) & PrimitiveLightingChannelMask) == 0) return;
 
-	switch (SceneLightsData.Type[LightId])
+	switch (SceneLightsData.Flags[LightId] & PATHTRACER_FLAG_TYPE_MASK)
 	{
 		case PATHTRACING_LIGHT_SKY:
 			SkyLight_EstimateLight(LightId, WorldPos, WorldNormal, IsTransmissiveMaterial, OutIrradiance);
@@ -264,7 +264,7 @@ void PdfLight(
 	float LightPdf = 0.0;
 
 	// Light sampling
-	switch (SceneLightsData.Type[LightId])
+	switch (SceneLightsData.Flags[LightId] & PATHTRACER_FLAG_TYPE_MASK)
 	{
 		case PATHTRACING_LIGHT_SKY:
 			SkyLight_PdfLight(Ray, LightId, LightUV, LightPdf);
@@ -297,7 +297,7 @@ void EvalLight(
 {
 	OutRadiance = float3(0.0f, 0.0f, 0.0f);
 
-	switch (SceneLightsData.Type[LightId])
+	switch (SceneLightsData.Flags[LightId] & PATHTRACER_FLAG_TYPE_MASK)
 	{
 		case PATHTRACING_LIGHT_SKY:
 			SkyLight_EvalLight(LightId, LightUV, Ray, OutRadiance);
@@ -339,7 +339,7 @@ void SampleLight(
 	// Light sampling
 	float LightPdf = 0.0;
 
-	switch (SceneLightsData.Type[OutLightId])
+	switch (SceneLightsData.Flags[OutLightId] & PATHTRACER_FLAG_TYPE_MASK)
 	{
 		case PATHTRACING_LIGHT_SKY:
 			SkyLight_SampleLight(Ray, Payload, RandSample, OutLightId, OutLightUV, LightPdf);
diff --git a/Engine/Shaders/Private/PathTracing/Light/PathTracingPointLight.ush b/Engine/Shaders/Private/PathTracing/Light/PathTracingPointLight.ush
--- a/Engine/Shaders/Private/PathTracing/Light/PathTracingPointLight.ush
+++ b/Engine/Shaders/Private/PathTracing/Light/PathTracingPointLight.ush
@@ -32,6 +32,7 @@ float4 PointLight_TraceLight(RayDesc Ray, int LightId, float HitT)
 			// #dxr_todo: sphere area is 4*pi*r^2 -- but the factor of 4 is missing for some reason?
 			float3 LightPower = GetColor(LightId) * ComputeIESAttenuation(LightId, Ray.Origin);
 			float3 LightRadiance = LightPower / (PI * LightRadius2);
+			LightRadiance *= ComputeAttenuationFalloff(LightDistanceSquared, LightId);
 
 			float SinThetaMax2 = saturate(LightRadius2 / LightDistanceSquared);
 			float OneMinusCosThetaMax = SinThetaMax2 < 0.01 ? SinThetaMax2 * (0.5 + 0.125 * SinThetaMax2) : 1 - sqrt(1 - SinThetaMax2);
@@ -82,6 +83,7 @@ void PointLight_SampleLight(
 	// When the angle is very small, Radiance over pdf simplifies even more since SolidAngle ~= PI * SinThetaMax2
 	// Canceling out common factors further leads to the classic Power / D^2 formula
 	OutRadianceOverPdf = SinThetaMax2 < 0.001 ? LightPower / LightDistanceSquared : LightRadiance / OutPdf;
+	OutRadianceOverPdf *= ComputeAttenuationFalloff(LightDistanceSquared, LightId);
 }
 
 void PointLight_GenerateLightRay(
@@ -133,7 +135,7 @@ void PointLight_EvalLight(
 
 	float3 Color = GetColor(LightId);
 	float Radius = GetRadius(LightId);
-	float Falloff = ComputeAttenuationFalloff(Ray.TMax, GetAttenuation(LightId));
+	float Falloff = ComputeAttenuationFalloff(LightDistanceSquared, LightId);
 	OutRadiance = Color * Falloff;
 
 	if (Radius == 0)
@@ -158,7 +160,7 @@ void PointLight_EstimateLight(
 	float NoL = 1.0; // trivial upper bound -- trying to be more accurate appears to reduce performance
 
 	float LightPower = Luminance(GetColor(LightId));
-	float Falloff = ComputeAttenuationFalloff(sqrt(LightDistanceSquared), GetAttenuation(LightId));
+	float Falloff = ComputeAttenuationFalloff(LightDistanceSquared, LightId);
 	OutIrradiance = LightPower * Falloff * NoL / LightDistanceSquared;
 }
 
diff --git a/Engine/Shaders/Private/PathTracing/Light/PathTracingRectLight.ush b/Engine/Shaders/Private/PathTracing/Light/PathTracingRectLight.ush
--- a/Engine/Shaders/Private/PathTracing/Light/PathTracingRectLight.ush
+++ b/Engine/Shaders/Private/PathTracing/Light/PathTracingRectLight.ush
@@ -44,6 +44,7 @@ float4 RectLight_TraceLight(RayDesc Ray, int LightId, float HitT)
 			{
 				// stored color is radiance
 				float3 Radiance = GetColor(LightId) * ComputeIESAttenuation(LightId, Ray.Origin);
+				Radiance *= ComputeAttenuationFalloff(dot(LightDirection, LightDirection), LightId);
 				float SolidAngle = BuildSphericalRect(Rect).SolidAngle;
 				return float4(Radiance, 1.0 / SolidAngle);
 			}
@@ -70,7 +71,8 @@ void RectLight_SampleLight(
 	float LightHeight = GetHeight(LightId);
 
 	// Define rectangle and compute solid angle
-	FRect Rect = GetRect(LightPosition - WorldPos,
+	float3 LightDirection = LightPosition - WorldPos;
+	FRect Rect = GetRect(LightDirection,
 		-LightNormal,
 		LightdPdv,
 		0.5 * LightWidth,
@@ -85,6 +87,7 @@ void RectLight_SampleLight(
 	OutDirection = normalize(Direction);
 	OutDistance = length(Direction);
 	OutRadianceOverPdf = GetColor(LightId) * SphericalRect.SolidAngle * ComputeIESAttenuation(LightId, WorldPos);
+	OutRadianceOverPdf *= ComputeAttenuationFalloff(dot(LightDirection, LightDirection), LightId);
 	OutPdf = 1.0 / SphericalRect.SolidAngle;
 }
 
@@ -120,7 +123,7 @@ void RectLight_EvalLight(
 )
 {
 	float3 Color = GetColor(LightId);
-	float Falloff = ComputeAttenuationFalloff(Ray.TMax, GetAttenuation(LightId));
+	float Falloff = ComputeAttenuationFalloff(Square(Ray.TMax), LightId);
 	OutRadiance = Color * Falloff;
 }
 
@@ -153,7 +156,7 @@ void RectLight_EstimateLight(
 	float Area = Width * Height;
 
 	float LightPower = Luminance(GetColor(LightId));
-	float Falloff = ComputeAttenuationFalloff(sqrt(LightDistanceSquared), GetAttenuation(LightId));
+	float Falloff = ComputeAttenuationFalloff(LightDistanceSquared, LightId);
 	OutIrradiance = LightPower * Falloff * Area * NoL / LightDistanceSquared;
 }
 
diff --git a/Engine/Shaders/Private/PathTracing/Light/PathTracingSpotLight.ush b/Engine/Shaders/Private/PathTracing/Light/PathTracingSpotLight.ush
--- a/Engine/Shaders/Private/PathTracing/Light/PathTracingSpotLight.ush
+++ b/Engine/Shaders/Private/PathTracing/Light/PathTracingSpotLight.ush
@@ -55,6 +55,7 @@ float4 SpotLight_TraceLight(RayDesc Ray, int LightId, float HitT) {
 			float3 LightNormal = GetNormal(LightId);
 			float2 CosConeAngles = GetCosConeAngles(LightId);
 			LightRadiance *= SpotAttenuation(-normalize(oc), LightNormal, CosConeAngles);
+			LightRadiance *= ComputeAttenuationFalloff(LightDistanceSquared, LightId);
 
 			float SinThetaMax2 = saturate(LightRadius2 / LightDistanceSquared);
 			float OneMinusCosThetaMax = SinThetaMax2 < 0.01 ? SinThetaMax2 * (0.5 + 0.125 * SinThetaMax2) : 1 - sqrt(1 - SinThetaMax2);
@@ -110,6 +111,7 @@ void SpotLight_SampleLight(
 	float3 LightNormal = GetNormal(LightId);
 	float2 CosConeAngles = GetCosConeAngles(LightId);
 	OutRadianceOverPdf *= SpotAttenuation(normalize(LightDirection), LightNormal, CosConeAngles);
+	OutRadianceOverPdf *= ComputeAttenuationFalloff(LightDistanceSquared, LightId);
 }
 
 void SpotLight_GenerateLightRay(
@@ -161,7 +163,7 @@ void SpotLight_EvalLight(
 
 	float3 Color = GetColor(LightId);
 	float Radius = GetRadius(LightId);
-	float Falloff = ComputeAttenuationFalloff(Ray.TMax, GetAttenuation(LightId));
+	float Falloff = ComputeAttenuationFalloff(LightDistanceSquared, LightId);
 	OutRadiance = Color * Falloff;
 
 	// Same math as point lights
@@ -196,7 +198,7 @@ void SpotLight_EstimateLight(
 	// Geometric term
 	float NoL = 1.0; // trivial upper bound -- trying to be more accurate appears to reduce performance
 
-	float Falloff = ComputeAttenuationFalloff(sqrt(LightDistanceSquared), GetAttenuation(LightId));
+	float Falloff = ComputeAttenuationFalloff(LightDistanceSquared, LightId);
 	OutIrradiance = LightPower * Falloff * NoL / LightDistanceSquared;
 	
 	float2 CosConeAngles = GetCosConeAngles(LightId);
diff --git a/Engine/Shaders/Shared/RayTracingDefinitions.h b/Engine/Shaders/Shared/RayTracingDefinitions.h
--- a/Engine/Shaders/Shared/RayTracingDefinitions.h
+++ b/Engine/Shaders/Shared/RayTracingDefinitions.h
@@ -35,3 +35,14 @@
 #define RAY_TRACING_MAX_ALLOWED_ATTRIBUTE_SIZE  8   // Sizeof 2 floats (barycentrics)
 #define RAY_TRACING_MAX_ALLOWED_PAYLOAD_SIZE    64  // Our maximum allowed payload size (sizeof FPackedMaterialClosestHitPayload)
 
+// Constants for the 'Flags' field of FPathTracingLightData
+#define PATHTRACER_FLAG_LIGHTING_CHANNEL_MASK				(7 << 0)	// Which lighting channel is this light assigned to?
+#define PATHTRACER_FLAG_TRANSMISSION_MASK					(1 << 3)	// Does the light affect the transmission side?
+#define PATHTRACER_FLAG_NON_INVERSE_SQUARE_FALLOFF_MASK		(1 << 4)	// Does the light have a non-inverse square decay?
+#define PATHTRACER_FLAG_STATIONARY_MASK						(1 << 5)	// Only used by GPULightmass
+#define PATHTRACER_FLAG_TYPE_MASK							(7 << 6)
+#define PATHTRACING_LIGHT_SKY								(0 << 6)
+#define PATHTRACING_LIGHT_DIRECTIONAL						(1 << 6)
+#define PATHTRACING_LIGHT_POINT								(2 << 6)
+#define PATHTRACING_LIGHT_SPOT								(3 << 6)
+#define PATHTRACING_LIGHT_RECT								(4 << 6)
diff --git a/Engine/Source/Runtime/Renderer/Private/PathTracing.cpp b/Engine/Source/Runtime/Renderer/Private/PathTracing.cpp
--- a/Engine/Source/Runtime/Renderer/Private/PathTracing.cpp
+++ b/Engine/Source/Runtime/Renderer/Private/PathTracing.cpp
@@ -224,11 +224,12 @@ void SetLightParameters(FRDGBuilder& GraphBuilder, FPathTracingRG::FParameters*
 	{
 		uint32 SkyLightIndex = 0;
 		uint8 SkyLightLightingChannelMask = 0xFF;
-		LightData.Type[SkyLightIndex] = 0;
 		LightData.Color[SkyLightIndex] = FVector(SkyLightData.Color);
-		LightData.Flags[SkyLightIndex] = SkyLightData.bTransmission & 0x01;
-		LightData.Flags[SkyLightIndex] |= (SkyLightLightingChannelMask & 0x7) << 1;
+		LightData.Flags[SkyLightIndex] = SkyLightData.bTransmission ? PATHTRACER_FLAG_TRANSMISSION_MASK : 0;
+		LightData.Flags[SkyLightIndex] |= PATHTRACER_FLAG_LIGHTING_CHANNEL_MASK;
+		LightData.Flags[SkyLightIndex] |= PATHTRACING_LIGHT_SKY;
 		LightData.IESTextureSlice[SkyLightIndex] = -1;
+		LightData.FalloffExponent[LightData.Count] = 0;
 		LightData.Count++;
 	}
 
@@ -241,8 +242,8 @@ void SetLightParameters(FRDGBuilder& GraphBuilder, FPathTracingRG::FParameters*
 		Light.LightSceneInfo->Proxy->GetLightShaderParameters(LightParameters);
 		uint32 Transmission = Light.LightSceneInfo->Proxy->Transmission();
 		uint8 LightingChannelMask = Light.LightSceneInfo->Proxy->GetLightingChannelMask();
-		LightData.Flags[LightData.Count] = Transmission & 0x01;
-		LightData.Flags[LightData.Count] |= (LightingChannelMask & 0x7) << 1;
+		LightData.Flags[LightData.Count] = Transmission ? PATHTRACER_FLAG_TRANSMISSION_MASK : 0;
+		LightData.Flags[LightData.Count] |= LightingChannelMask & PATHTRACER_FLAG_LIGHTING_CHANNEL_MASK;
 		LightData.IESTextureSlice[LightData.Count] = -1;
 
 		if (UseLightProfiles)
@@ -261,39 +262,43 @@ void SetLightParameters(FRDGBuilder& GraphBuilder, FPathTracingRG::FParameters*
 		LightData.Normal[LightData.Count] = -LightParameters.Direction;
 		LightData.dPdu[LightData.Count] = FVector::CrossProduct(LightParameters.Tangent, LightParameters.Direction);
 		LightData.dPdv[LightData.Count] = LightParameters.Tangent;
+		LightData.Attenuation[LightData.Count] = LightParameters.InvRadius;
+		LightData.FalloffExponent[LightData.Count] = 0;
 
 		ELightComponentType LightComponentType = (ELightComponentType)Light.LightSceneInfo->Proxy->GetLightType();
 		switch (LightComponentType)
 		{
 			case LightType_Directional:
 			{
-				LightData.Type[LightData.Count] = 2;
                 LightData.Normal[LightData.Count] = LightParameters.Direction;
 				LightData.Dimensions[LightData.Count] = FVector(0.0f, 0.0f, LightParameters.SourceRadius);
-				LightData.Attenuation[LightData.Count] = 1.0 / LightParameters.InvRadius;
+				LightData.Flags[LightData.Count] |= PATHTRACING_LIGHT_DIRECTIONAL;
 				break;
 			}
 			case LightType_Rect:
 			{
-				LightData.Type[LightData.Count] = 3;
 				LightData.Dimensions[LightData.Count] = FVector(2.0f * LightParameters.SourceRadius, 2.0f * LightParameters.SourceLength, 0.0f);
-				LightData.Attenuation[LightData.Count] = 1.0 / LightParameters.InvRadius;
 				LightData.RectLightBarnCosAngle[LightData.Count] = LightParameters.RectLightBarnCosAngle;
 				LightData.RectLightBarnLength[LightData.Count] = LightParameters.RectLightBarnLength;
+				LightData.FalloffExponent[LightData.Count] = LightParameters.FalloffExponent;
+				LightData.Flags[LightData.Count] |= Light.LightSceneInfo->Proxy->IsInverseSquared() ? 0 : PATHTRACER_FLAG_NON_INVERSE_SQUARE_FALLOFF_MASK;
+				LightData.Flags[LightData.Count] |= PATHTRACING_LIGHT_RECT;
 				break;
 			}
 			case LightType_Spot:
 			{
-				LightData.Type[LightData.Count] = 4;
 				LightData.Dimensions[LightData.Count] = FVector(LightParameters.SpotAngles, LightParameters.SourceRadius);
-				LightData.Attenuation[LightData.Count] = 1.0 / LightParameters.InvRadius;
+				LightData.FalloffExponent[LightData.Count] = LightParameters.FalloffExponent;
+				LightData.Flags[LightData.Count] |= Light.LightSceneInfo->Proxy->IsInverseSquared() ? 0 : PATHTRACER_FLAG_NON_INVERSE_SQUARE_FALLOFF_MASK;
+				LightData.Flags[LightData.Count] |= PATHTRACING_LIGHT_SPOT;
 				break;
 			}
 			case LightType_Point:
 			{
-				LightData.Type[LightData.Count] = 1;
 				LightData.Dimensions[LightData.Count] = FVector(0.0, 0.0, LightParameters.SourceRadius);
-				LightData.Attenuation[LightData.Count] = 1.0 / LightParameters.InvRadius;
+				LightData.FalloffExponent[LightData.Count] = LightParameters.FalloffExponent;
+				LightData.Flags[LightData.Count] |= Light.LightSceneInfo->Proxy->IsInverseSquared() ? 0 : PATHTRACER_FLAG_NON_INVERSE_SQUARE_FALLOFF_MASK;
+				LightData.Flags[LightData.Count] |= PATHTRACING_LIGHT_POINT;
 				break;
 			}
 			default:
diff --git a/Engine/Source/Runtime/Renderer/Private/PathTracingUniformBuffers.h b/Engine/Source/Runtime/Renderer/Private/PathTracingUniformBuffers.h
--- a/Engine/Source/Runtime/Renderer/Private/PathTracingUniformBuffers.h
+++ b/Engine/Source/Runtime/Renderer/Private/PathTracingUniformBuffers.h
@@ -26,7 +26,6 @@ END_GLOBAL_SHADER_PARAMETER_STRUCT()
 
 BEGIN_GLOBAL_SHADER_PARAMETER_STRUCT(FPathTracingLightData, RENDERER_API)
 	SHADER_PARAMETER(uint32, Count)
-	SHADER_PARAMETER_ARRAY(uint32, Type, [RAY_TRACING_LIGHT_COUNT_MAXIMUM])
 	// Geometry
 	SHADER_PARAMETER_ARRAY(FVector, Position, [RAY_TRACING_LIGHT_COUNT_MAXIMUM])
 	SHADER_PARAMETER_ARRAY(FVector, Normal, [RAY_TRACING_LIGHT_COUNT_MAXIMUM])
@@ -37,12 +36,11 @@ BEGIN_GLOBAL_SHADER_PARAMETER_STRUCT(FPathTracingLightData, RENDERER_API)
 	// Light-specific
 	SHADER_PARAMETER_ARRAY(FVector, Dimensions, [RAY_TRACING_LIGHT_COUNT_MAXIMUM])
 	SHADER_PARAMETER_ARRAY(float, Attenuation, [RAY_TRACING_LIGHT_COUNT_MAXIMUM])
+	SHADER_PARAMETER_ARRAY(float, FalloffExponent, [RAY_TRACING_LIGHT_COUNT_MAXIMUM])
 	SHADER_PARAMETER_ARRAY(float, RectLightBarnCosAngle, [RAY_TRACING_LIGHT_COUNT_MAXIMUM])
 	SHADER_PARAMETER_ARRAY(float, RectLightBarnLength, [RAY_TRACING_LIGHT_COUNT_MAXIMUM])
 	// IES index (valid if >= 0)
 	SHADER_PARAMETER_ARRAY(float, IESTextureSlice , [RAY_TRACING_LIGHT_COUNT_MAXIMUM])
 	// Flags
 	SHADER_PARAMETER_ARRAY(uint32, Flags, [RAY_TRACING_LIGHT_COUNT_MAXIMUM])
-	// Only used by GPULightmass currently, not filled in realtime paths
-	SHADER_PARAMETER_ARRAY(uint32, Mobility, [RAY_TRACING_LIGHT_COUNT_MAXIMUM])
 END_GLOBAL_SHADER_PARAMETER_STRUCT()
diff --git a/Engine/Source/Runtime/Renderer/Private/RayTracing/RayTracingGlobalIllumination.cpp b/Engine/Source/Runtime/Renderer/Private/RayTracing/RayTracingGlobalIllumination.cpp
--- a/Engine/Source/Runtime/Renderer/Private/RayTracing/RayTracingGlobalIllumination.cpp
+++ b/Engine/Source/Runtime/Renderer/Private/RayTracing/RayTracingGlobalIllumination.cpp
@@ -109,13 +109,6 @@ static FAutoConsoleVariableRef CVarRayTracingGlobalIlluminationScreenPercentage(
 	TEXT("Screen percentage for ray tracing global illumination (default = 50)")
 );
 
-static TAutoConsoleVariable<int32> CVarRayTracingGlobalIlluminationEnableLightAttenuation(
-	TEXT("r.RayTracing.GlobalIllumination.EnableLightAttenuation"),
-	1,
-	TEXT("Enables light attenuation when calculating irradiance during next-event estimation (default = 1)"),
-	ECVF_RenderThreadSafe
-);
-
 static TAutoConsoleVariable<int32> CVarRayTracingGlobalIlluminationEnableTwoSidedGeometry(
 	TEXT("r.RayTracing.GlobalIllumination.EnableTwoSidedGeometry"),
 	1,
@@ -262,8 +255,6 @@ RENDERER_API void SetupLightParameters(
 	FSkyLightSceneProxy* SkyLight = Scene.SkyLight;
 	FVector SkyLightColor = FVector(0.0f, 0.0f, 0.0f);
 	uint32 SkyLightTransmission = 0;
-	// SkyLight does not have a LightingChannelMask
-	uint8 SkyLightLightingChannelMask = 0xFF;
 	if (SkyLight && SkyLight->bAffectGlobalIllumination && CVarRayTracingGlobalIlluminationSkyLight.GetValueOnRenderThread())
 	{
 		SkyLightColor = FVector(SkyLight->GetEffectiveLightColor());
@@ -273,10 +264,11 @@ RENDERER_API void SetupLightParameters(
 	// Prepend SkyLight to light buffer
 	// WARNING: Until ray payload encodes Light data buffer, the execution depends on this ordering!
 	uint32 SkyLightIndex = 0;
-	LightParameters->Type[SkyLightIndex] = 0;
 	LightParameters->Color[SkyLightIndex] = SkyLightColor;
-	LightParameters->Flags[SkyLightIndex]  = SkyLightTransmission & 0x01;
-	LightParameters->Flags[SkyLightIndex] |= (SkyLightLightingChannelMask & 0x7) << 1;
+	LightParameters->Flags[SkyLightIndex] = SkyLightTransmission ? PATHTRACER_FLAG_TRANSMISSION_MASK : 0;
+	// SkyLight does not have a LightingChannelMask
+	LightParameters->Flags[SkyLightIndex] |= PATHTRACER_FLAG_LIGHTING_CHANNEL_MASK;
+	LightParameters->Flags[SkyLightIndex] |= PATHTRACING_LIGHT_SKY;
 	
 	LightParameters->Count++;
 
@@ -293,8 +285,13 @@ RENDERER_API void SetupLightParameters(
 
 		uint32 Transmission = Light.LightSceneInfo->Proxy->Transmission();
 		uint8 LightingChannelMask = Light.LightSceneInfo->Proxy->GetLightingChannelMask();
-		LightParameters->Flags[LightParameters->Count]  = Transmission & 0x01;
-		LightParameters->Flags[LightParameters->Count] |= (LightingChannelMask & 0x7) << 1;
+		LightParameters->Flags[LightParameters->Count]  = Transmission ? PATHTRACER_FLAG_TRANSMISSION_MASK : 0;
+		LightParameters->Flags[LightParameters->Count] |= LightingChannelMask & PATHTRACER_FLAG_LIGHTING_CHANNEL_MASK;
+		LightParameters->Flags[LightParameters->Count] |= Light.LightSceneInfo->Proxy->IsInverseSquared() ? 0 : PATHTRACER_FLAG_NON_INVERSE_SQUARE_FALLOFF_MASK;
+
+		LightParameters->FalloffExponent[LightParameters->Count] = LightShaderParameters.FalloffExponent;
+		LightParameters->Attenuation[LightParameters->Count] = LightShaderParameters.InvRadius;
+		LightParameters->IESTextureSlice[LightParameters->Count] = -1; // not used by this path at the moment
 
 		ELightComponentType LightComponentType = (ELightComponentType)Light.LightSceneInfo->Proxy->GetLightType();
 		switch (LightComponentType)
@@ -303,54 +300,50 @@ RENDERER_API void SetupLightParameters(
 		{
 			if (CVarRayTracingGlobalIlluminationDirectionalLight.GetValueOnRenderThread() == 0) continue;
 
-			LightParameters->Type[LightParameters->Count] = 2;
 			LightParameters->Normal[LightParameters->Count] = LightShaderParameters.Direction;
 			LightParameters->Color[LightParameters->Count] = LightShaderParameters.Color;
-			LightParameters->Attenuation[LightParameters->Count] = 1.0 / LightShaderParameters.InvRadius;
+			LightParameters->Flags[LightParameters->Count] |= PATHTRACING_LIGHT_DIRECTIONAL;
 			break;
 		}
 		case LightType_Rect:
 		{
 			if (CVarRayTracingGlobalIlluminationRectLight.GetValueOnRenderThread() == 0) continue;
 
-			LightParameters->Type[LightParameters->Count] = 3;
 			LightParameters->Position[LightParameters->Count] = LightShaderParameters.Position;
 			LightParameters->Normal[LightParameters->Count] = -LightShaderParameters.Direction;
 			LightParameters->dPdu[LightParameters->Count] = FVector::CrossProduct(LightShaderParameters.Direction, LightShaderParameters.Tangent);
 			LightParameters->dPdv[LightParameters->Count] = LightShaderParameters.Tangent;
 			LightParameters->Color[LightParameters->Count] = LightShaderParameters.Color;
 			LightParameters->Dimensions[LightParameters->Count] = FVector(2.0f * LightShaderParameters.SourceRadius, 2.0f * LightShaderParameters.SourceLength, 0.0f);
-			LightParameters->Attenuation[LightParameters->Count] = 1.0 / LightShaderParameters.InvRadius;
 			LightParameters->RectLightBarnCosAngle[LightParameters->Count] = LightShaderParameters.RectLightBarnCosAngle;
 			LightParameters->RectLightBarnLength[LightParameters->Count] = LightShaderParameters.RectLightBarnLength;
+			LightParameters->Flags[LightParameters->Count] |= PATHTRACING_LIGHT_RECT;
 			break;
 		}
 		case LightType_Point:
 		default:
 		{
 			if (CVarRayTracingGlobalIlluminationPointLight.GetValueOnRenderThread() == 0) continue;
-			
-			LightParameters->Type[LightParameters->Count] = 1;
+
 			LightParameters->Position[LightParameters->Count] = LightShaderParameters.Position;
 			// #dxr_todo: UE-72556 define these differences from Lit..
 			LightParameters->Color[LightParameters->Count] = LightShaderParameters.Color;
 			float SourceRadius = 0.0; // LightShaderParameters.SourceRadius causes too much noise for little pay off at this time
 			LightParameters->Dimensions[LightParameters->Count] = FVector(0.0, 0.0, SourceRadius);
-			LightParameters->Attenuation[LightParameters->Count] = 1.0 / LightShaderParameters.InvRadius;
+			LightParameters->Flags[LightParameters->Count] |= PATHTRACING_LIGHT_POINT;
 			break;
 		}
 		case LightType_Spot:
 		{
 			if (CVarRayTracingGlobalIlluminationSpotLight.GetValueOnRenderThread() == 0) continue;
 
-			LightParameters->Type[LightParameters->Count] = 4;
 			LightParameters->Position[LightParameters->Count] = LightShaderParameters.Position;
 			LightParameters->Normal[LightParameters->Count] = -LightShaderParameters.Direction;
 			// #dxr_todo: UE-72556 define these differences from Lit..
 			LightParameters->Color[LightParameters->Count] = LightShaderParameters.Color;
 			float SourceRadius = 0.0; // LightShaderParameters.SourceRadius causes too much noise for little pay off at this time
 			LightParameters->Dimensions[LightParameters->Count] = FVector(LightShaderParameters.SpotAngles, SourceRadius);
-			LightParameters->Attenuation[LightParameters->Count] = 1.0 / LightShaderParameters.InvRadius;
+			LightParameters->Flags[LightParameters->Count] |= PATHTRACING_LIGHT_SPOT;
 			break;
 		}
 		};
@@ -437,11 +430,10 @@ class FGlobalIlluminationRGS : public FGlobalShader
 	DECLARE_GLOBAL_SHADER(FGlobalIlluminationRGS)
 	SHADER_USE_ROOT_PARAMETER_STRUCT(FGlobalIlluminationRGS, FGlobalShader)
 
-	class FUseAttenuationTermDim : SHADER_PERMUTATION_BOOL("USE_ATTENUATION_TERM");
 	class FEnableTwoSidedGeometryDim : SHADER_PERMUTATION_BOOL("ENABLE_TWO_SIDED_GEOMETRY");
 	class FEnableTransmissionDim : SHADER_PERMUTATION_INT("ENABLE_TRANSMISSION", 2);
 
-	using FPermutationDomain = TShaderPermutationDomain<FUseAttenuationTermDim, FEnableTwoSidedGeometryDim, FEnableTransmissionDim>;
+	using FPermutationDomain = TShaderPermutationDomain<FEnableTwoSidedGeometryDim, FEnableTransmissionDim>;
 
 	static bool ShouldCompilePermutation(const FGlobalShaderPermutationParameters& Parameters)
 	{
@@ -496,12 +488,11 @@ class FRayTracingGlobalIlluminationCreateGatherPointsRGS : public FGlobalShader
 	DECLARE_GLOBAL_SHADER(FRayTracingGlobalIlluminationCreateGatherPointsRGS)
 	SHADER_USE_ROOT_PARAMETER_STRUCT(FRayTracingGlobalIlluminationCreateGatherPointsRGS, FGlobalShader)
 
-	class FUseAttenuationTermDim : SHADER_PERMUTATION_BOOL("USE_ATTENUATION_TERM");
 	class FEnableTwoSidedGeometryDim : SHADER_PERMUTATION_BOOL("ENABLE_TWO_SIDED_GEOMETRY");
 	class FDeferredMaterialMode : SHADER_PERMUTATION_ENUM_CLASS("DIM_DEFERRED_MATERIAL_MODE", EDeferredMaterialMode);
 	class FEnableTransmissionDim : SHADER_PERMUTATION_INT("ENABLE_TRANSMISSION", 2);
 
-	using FPermutationDomain = TShaderPermutationDomain<FUseAttenuationTermDim, FEnableTwoSidedGeometryDim, FDeferredMaterialMode, FEnableTransmissionDim>;
+	using FPermutationDomain = TShaderPermutationDomain<FEnableTwoSidedGeometryDim, FDeferredMaterialMode, FEnableTransmissionDim>;
 
 	static bool ShouldCompilePermutation(const FGlobalShaderPermutationParameters& Parameters)
 	{
@@ -569,11 +560,10 @@ class FRayTracingGlobalIlluminationCreateGatherPointsTraceRGS : public FGlobalSh
 	DECLARE_GLOBAL_SHADER(FRayTracingGlobalIlluminationCreateGatherPointsTraceRGS)
 	SHADER_USE_ROOT_PARAMETER_STRUCT(FRayTracingGlobalIlluminationCreateGatherPointsTraceRGS, FGlobalShader)
 
-	class FUseAttenuationTermDim : SHADER_PERMUTATION_BOOL("USE_ATTENUATION_TERM");
 	class FEnableTwoSidedGeometryDim : SHADER_PERMUTATION_BOOL("ENABLE_TWO_SIDED_GEOMETRY");
 	class FDeferredMaterialMode : SHADER_PERMUTATION_ENUM_CLASS("DIM_DEFERRED_MATERIAL_MODE", EDeferredMaterialMode);
 
-	using FPermutationDomain = TShaderPermutationDomain<FUseAttenuationTermDim, FEnableTwoSidedGeometryDim, FDeferredMaterialMode>;
+	using FPermutationDomain = TShaderPermutationDomain<FEnableTwoSidedGeometryDim, FDeferredMaterialMode>;
 
 	static bool ShouldCompilePermutation(const FGlobalShaderPermutationParameters& Parameters)
 	{
@@ -634,11 +624,10 @@ class FRayTracingGlobalIlluminationFinalGatherRGS : public FGlobalShader
 	DECLARE_GLOBAL_SHADER(FRayTracingGlobalIlluminationFinalGatherRGS)
 	SHADER_USE_ROOT_PARAMETER_STRUCT(FRayTracingGlobalIlluminationFinalGatherRGS, FGlobalShader)
 
-	class FUseAttenuationTermDim : SHADER_PERMUTATION_BOOL("USE_ATTENUATION_TERM");
 	class FEnableTwoSidedGeometryDim : SHADER_PERMUTATION_BOOL("ENABLE_TWO_SIDED_GEOMETRY");
 	class FEnableNeighborVisibilityTestDim : SHADER_PERMUTATION_BOOL("USE_NEIGHBOR_VISIBILITY_TEST");
 
-	using FPermutationDomain = TShaderPermutationDomain<FUseAttenuationTermDim, FEnableTwoSidedGeometryDim, FEnableNeighborVisibilityTestDim>;
+	using FPermutationDomain = TShaderPermutationDomain<FEnableTwoSidedGeometryDim, FEnableNeighborVisibilityTestDim>;
 
 	static bool ShouldCompilePermutation(const FGlobalShaderPermutationParameters& Parameters)
 	{
@@ -695,63 +684,54 @@ void FDeferredShadingSceneRenderer::PrepareRayTracingGlobalIllumination(const FV
 	int EnableTransmission = CVarRayTracingGlobalIlluminationEnableTransmission.GetValueOnRenderThread();
 
 	// Declare all RayGen shaders that require material closest hit shaders to be bound
-	for (int UseAttenuationTerm = 0; UseAttenuationTerm < 2; ++UseAttenuationTerm)
+	for (int EnableTwoSidedGeometry = 0; EnableTwoSidedGeometry < 2; ++EnableTwoSidedGeometry)
 	{
-		for (int EnableTwoSidedGeometry = 0; EnableTwoSidedGeometry < 2; ++EnableTwoSidedGeometry)
+		FGlobalIlluminationRGS::FPermutationDomain PermutationVector;
+		PermutationVector.Set<FGlobalIlluminationRGS::FEnableTwoSidedGeometryDim>(EnableTwoSidedGeometry == 1);
+		PermutationVector.Set<FGlobalIlluminationRGS::FEnableTransmissionDim>(EnableTransmission);
+		TShaderMapRef<FGlobalIlluminationRGS> RayGenerationShader(View.ShaderMap, PermutationVector);
+		OutRayGenShaders.Add(RayGenerationShader.GetRayTracingShader());
+
+		if (bSortMaterials)
 		{
-			FGlobalIlluminationRGS::FPermutationDomain PermutationVector;
-			PermutationVector.Set<FGlobalIlluminationRGS::FUseAttenuationTermDim>(CVarRayTracingGlobalIlluminationEnableLightAttenuation.GetValueOnRenderThread() != 0);
-			PermutationVector.Set<FGlobalIlluminationRGS::FEnableTwoSidedGeometryDim>(EnableTwoSidedGeometry == 1);
-			PermutationVector.Set<FGlobalIlluminationRGS::FEnableTransmissionDim>(EnableTransmission);
-			TShaderMapRef<FGlobalIlluminationRGS> RayGenerationShader(View.ShaderMap, PermutationVector);
-			OutRayGenShaders.Add(RayGenerationShader.GetRayTracingShader());
-
-			if (bSortMaterials)
+			// Gather
 			{
-				// Gather
-				{
-					FRayTracingGlobalIlluminationCreateGatherPointsTraceRGS::FPermutationDomain CreateGatherPointsPermutationVector;
-					CreateGatherPointsPermutationVector.Set<FRayTracingGlobalIlluminationCreateGatherPointsTraceRGS::FUseAttenuationTermDim>(UseAttenuationTerm == 1);
-					CreateGatherPointsPermutationVector.Set<FRayTracingGlobalIlluminationCreateGatherPointsTraceRGS::FEnableTwoSidedGeometryDim>(EnableTwoSidedGeometry == 1);
-					CreateGatherPointsPermutationVector.Set<FRayTracingGlobalIlluminationCreateGatherPointsTraceRGS::FDeferredMaterialMode>(EDeferredMaterialMode::Gather);
-					TShaderMapRef<FRayTracingGlobalIlluminationCreateGatherPointsTraceRGS> CreateGatherPointsRayGenerationShader(View.ShaderMap, CreateGatherPointsPermutationVector);
-					OutRayGenShaders.Add(CreateGatherPointsRayGenerationShader.GetRayTracingShader());
-				}
-
-				// Shade
-				{
-					FRayTracingGlobalIlluminationCreateGatherPointsRGS::FPermutationDomain CreateGatherPointsPermutationVector;
-					CreateGatherPointsPermutationVector.Set<FRayTracingGlobalIlluminationCreateGatherPointsRGS::FUseAttenuationTermDim>(UseAttenuationTerm == 1);
-					CreateGatherPointsPermutationVector.Set<FRayTracingGlobalIlluminationCreateGatherPointsRGS::FEnableTwoSidedGeometryDim>(EnableTwoSidedGeometry == 1);
-					CreateGatherPointsPermutationVector.Set<FRayTracingGlobalIlluminationCreateGatherPointsRGS::FDeferredMaterialMode>(EDeferredMaterialMode::Shade);
-					CreateGatherPointsPermutationVector.Set<FRayTracingGlobalIlluminationCreateGatherPointsRGS::FEnableTransmissionDim>(EnableTransmission);
-					TShaderMapRef<FRayTracingGlobalIlluminationCreateGatherPointsRGS> CreateGatherPointsRayGenerationShader(View.ShaderMap, CreateGatherPointsPermutationVector);
-					OutRayGenShaders.Add(CreateGatherPointsRayGenerationShader.GetRayTracingShader());
-				}
+				FRayTracingGlobalIlluminationCreateGatherPointsTraceRGS::FPermutationDomain CreateGatherPointsPermutationVector;
+				CreateGatherPointsPermutationVector.Set<FRayTracingGlobalIlluminationCreateGatherPointsTraceRGS::FEnableTwoSidedGeometryDim>(EnableTwoSidedGeometry == 1);
+				CreateGatherPointsPermutationVector.Set<FRayTracingGlobalIlluminationCreateGatherPointsTraceRGS::FDeferredMaterialMode>(EDeferredMaterialMode::Gather);
+				TShaderMapRef<FRayTracingGlobalIlluminationCreateGatherPointsTraceRGS> CreateGatherPointsRayGenerationShader(View.ShaderMap, CreateGatherPointsPermutationVector);
+				OutRayGenShaders.Add(CreateGatherPointsRayGenerationShader.GetRayTracingShader());
 			}
-			else
+
+			// Shade
 			{
 				FRayTracingGlobalIlluminationCreateGatherPointsRGS::FPermutationDomain CreateGatherPointsPermutationVector;
-				CreateGatherPointsPermutationVector.Set<FRayTracingGlobalIlluminationCreateGatherPointsRGS::FUseAttenuationTermDim>(UseAttenuationTerm == 1);
 				CreateGatherPointsPermutationVector.Set<FRayTracingGlobalIlluminationCreateGatherPointsRGS::FEnableTwoSidedGeometryDim>(EnableTwoSidedGeometry == 1);
-				CreateGatherPointsPermutationVector.Set<FRayTracingGlobalIlluminationCreateGatherPointsRGS::FDeferredMaterialMode>(EDeferredMaterialMode::None);
-				CreateGatherPointsPermutationVector.Set < FRayTracingGlobalIlluminationCreateGatherPointsRGS::FEnableTransmissionDim>(EnableTransmission);
+				CreateGatherPointsPermutationVector.Set<FRayTracingGlobalIlluminationCreateGatherPointsRGS::FDeferredMaterialMode>(EDeferredMaterialMode::Shade);
+				CreateGatherPointsPermutationVector.Set<FRayTracingGlobalIlluminationCreateGatherPointsRGS::FEnableTransmissionDim>(EnableTransmission);
 				TShaderMapRef<FRayTracingGlobalIlluminationCreateGatherPointsRGS> CreateGatherPointsRayGenerationShader(View.ShaderMap, CreateGatherPointsPermutationVector);
 				OutRayGenShaders.Add(CreateGatherPointsRayGenerationShader.GetRayTracingShader());
 			}
+		}
+		else
+		{
+			FRayTracingGlobalIlluminationCreateGatherPointsRGS::FPermutationDomain CreateGatherPointsPermutationVector;
+			CreateGatherPointsPermutationVector.Set<FRayTracingGlobalIlluminationCreateGatherPointsRGS::FEnableTwoSidedGeometryDim>(EnableTwoSidedGeometry == 1);
+			CreateGatherPointsPermutationVector.Set<FRayTracingGlobalIlluminationCreateGatherPointsRGS::FDeferredMaterialMode>(EDeferredMaterialMode::None);
+			CreateGatherPointsPermutationVector.Set < FRayTracingGlobalIlluminationCreateGatherPointsRGS::FEnableTransmissionDim>(EnableTransmission);
+			TShaderMapRef<FRayTracingGlobalIlluminationCreateGatherPointsRGS> CreateGatherPointsRayGenerationShader(View.ShaderMap, CreateGatherPointsPermutationVector);
+			OutRayGenShaders.Add(CreateGatherPointsRayGenerationShader.GetRayTracingShader());
+		}
 
-			for (int EnableNeighborVisibilityTest = 0; EnableNeighborVisibilityTest < 2; ++EnableNeighborVisibilityTest)
-			{
-				FRayTracingGlobalIlluminationFinalGatherRGS::FPermutationDomain GatherPassPermutationVector;
-				GatherPassPermutationVector.Set<FRayTracingGlobalIlluminationFinalGatherRGS::FUseAttenuationTermDim>(UseAttenuationTerm == 1);
-				GatherPassPermutationVector.Set<FRayTracingGlobalIlluminationFinalGatherRGS::FEnableTwoSidedGeometryDim>(EnableTwoSidedGeometry == 1);
-				GatherPassPermutationVector.Set<FRayTracingGlobalIlluminationFinalGatherRGS::FEnableNeighborVisibilityTestDim>(EnableNeighborVisibilityTest == 1);
-				TShaderMapRef<FRayTracingGlobalIlluminationFinalGatherRGS> GatherPassRayGenerationShader(View.ShaderMap, GatherPassPermutationVector);
-				OutRayGenShaders.Add(GatherPassRayGenerationShader.GetRayTracingShader());
-			}
+		for (int EnableNeighborVisibilityTest = 0; EnableNeighborVisibilityTest < 2; ++EnableNeighborVisibilityTest)
+		{
+			FRayTracingGlobalIlluminationFinalGatherRGS::FPermutationDomain GatherPassPermutationVector;
+			GatherPassPermutationVector.Set<FRayTracingGlobalIlluminationFinalGatherRGS::FEnableTwoSidedGeometryDim>(EnableTwoSidedGeometry == 1);
+			GatherPassPermutationVector.Set<FRayTracingGlobalIlluminationFinalGatherRGS::FEnableNeighborVisibilityTestDim>(EnableNeighborVisibilityTest == 1);
+			TShaderMapRef<FRayTracingGlobalIlluminationFinalGatherRGS> GatherPassRayGenerationShader(View.ShaderMap, GatherPassPermutationVector);
+			OutRayGenShaders.Add(GatherPassRayGenerationShader.GetRayTracingShader());
 		}
 	}
-
 }
 
 void FDeferredShadingSceneRenderer::PrepareRayTracingGlobalIlluminationDeferredMaterial(const FViewInfo& View, TArray<FRHIRayTracingShader*>& OutRayGenShaders)
@@ -765,28 +745,23 @@ void FDeferredShadingSceneRenderer::PrepareRayTracingGlobalIlluminationDeferredM
 	}
 
 	// Declare all RayGen shaders that require material closest hit shaders to be bound
-	for (int UseAttenuationTerm = 0; UseAttenuationTerm < 2; ++UseAttenuationTerm)
+	for (int EnableTwoSidedGeometry = 0; EnableTwoSidedGeometry < 2; ++EnableTwoSidedGeometry)
 	{
-		for (int EnableTwoSidedGeometry = 0; EnableTwoSidedGeometry < 2; ++EnableTwoSidedGeometry)
-		{
-			FGlobalIlluminationRGS::FPermutationDomain PermutationVector;
-			PermutationVector.Set<FGlobalIlluminationRGS::FUseAttenuationTermDim>(CVarRayTracingGlobalIlluminationEnableLightAttenuation.GetValueOnRenderThread() != 0);
-			PermutationVector.Set<FGlobalIlluminationRGS::FEnableTwoSidedGeometryDim>(EnableTwoSidedGeometry == 1);
-			PermutationVector.Set<FGlobalIlluminationRGS::FEnableTransmissionDim>(EnableTransmission);
-			TShaderMapRef<FGlobalIlluminationRGS> RayGenerationShader(View.ShaderMap, PermutationVector);
-			OutRayGenShaders.Add(RayGenerationShader.GetRayTracingShader());
-
-			// Gather
-			{
-				FRayTracingGlobalIlluminationCreateGatherPointsTraceRGS::FPermutationDomain CreateGatherPointsPermutationVector;
-				CreateGatherPointsPermutationVector.Set<FRayTracingGlobalIlluminationCreateGatherPointsTraceRGS::FUseAttenuationTermDim>(UseAttenuationTerm == 1);
-				CreateGatherPointsPermutationVector.Set<FRayTracingGlobalIlluminationCreateGatherPointsTraceRGS::FEnableTwoSidedGeometryDim>(EnableTwoSidedGeometry == 1);
-				CreateGatherPointsPermutationVector.Set<FRayTracingGlobalIlluminationCreateGatherPointsTraceRGS::FDeferredMaterialMode>(EDeferredMaterialMode::Gather);
-				TShaderMapRef<FRayTracingGlobalIlluminationCreateGatherPointsTraceRGS> CreateGatherPointsRayGenerationShader(View.ShaderMap, CreateGatherPointsPermutationVector);
-				OutRayGenShaders.Add(CreateGatherPointsRayGenerationShader.GetRayTracingShader());
-			}
+		FGlobalIlluminationRGS::FPermutationDomain PermutationVector;
+		PermutationVector.Set<FGlobalIlluminationRGS::FEnableTwoSidedGeometryDim>(EnableTwoSidedGeometry == 1);
+		PermutationVector.Set<FGlobalIlluminationRGS::FEnableTransmissionDim>(EnableTransmission);
+		TShaderMapRef<FGlobalIlluminationRGS> RayGenerationShader(View.ShaderMap, PermutationVector);
+		OutRayGenShaders.Add(RayGenerationShader.GetRayTracingShader());
 
+		// Gather
+		{
+			FRayTracingGlobalIlluminationCreateGatherPointsTraceRGS::FPermutationDomain CreateGatherPointsPermutationVector;
+			CreateGatherPointsPermutationVector.Set<FRayTracingGlobalIlluminationCreateGatherPointsTraceRGS::FEnableTwoSidedGeometryDim>(EnableTwoSidedGeometry == 1);
+			CreateGatherPointsPermutationVector.Set<FRayTracingGlobalIlluminationCreateGatherPointsTraceRGS::FDeferredMaterialMode>(EDeferredMaterialMode::Gather);
+			TShaderMapRef<FRayTracingGlobalIlluminationCreateGatherPointsTraceRGS> CreateGatherPointsRayGenerationShader(View.ShaderMap, CreateGatherPointsPermutationVector);
+			OutRayGenShaders.Add(CreateGatherPointsRayGenerationShader.GetRayTracingShader());
 		}
+
 	}
 }
 
@@ -1062,7 +1037,6 @@ void FDeferredShadingSceneRenderer::RayTracingGlobalIlluminationCreateGatherPoin
 		}
 
 		FRayTracingGlobalIlluminationCreateGatherPointsRGS::FPermutationDomain PermutationVector;
-		PermutationVector.Set<FRayTracingGlobalIlluminationCreateGatherPointsRGS::FUseAttenuationTermDim>(true);
 		PermutationVector.Set<FRayTracingGlobalIlluminationCreateGatherPointsRGS::FEnableTwoSidedGeometryDim>(CVarRayTracingGlobalIlluminationEnableTwoSidedGeometry.GetValueOnRenderThread() != 0);
 		PermutationVector.Set<FRayTracingGlobalIlluminationCreateGatherPointsRGS::FEnableTransmissionDim>(CVarRayTracingGlobalIlluminationEnableTransmission.GetValueOnRenderThread());
 		TShaderMapRef<FRayTracingGlobalIlluminationCreateGatherPointsRGS> RayGenerationShader(GetGlobalShaderMap(FeatureLevel), PermutationVector);
@@ -1105,7 +1079,6 @@ void FDeferredShadingSceneRenderer::RayTracingGlobalIlluminationCreateGatherPoin
 			GatherPassParameters->MaterialBuffer = GraphBuilder.CreateUAV(DeferredMaterialBuffer);
 
 			FRayTracingGlobalIlluminationCreateGatherPointsTraceRGS::FPermutationDomain PermutationVector;
-			PermutationVector.Set<FRayTracingGlobalIlluminationCreateGatherPointsTraceRGS::FUseAttenuationTermDim>(true);
 			PermutationVector.Set<FRayTracingGlobalIlluminationCreateGatherPointsTraceRGS::FEnableTwoSidedGeometryDim>(CVarRayTracingGlobalIlluminationEnableTwoSidedGeometry.GetValueOnRenderThread() != 0);
 			PermutationVector.Set<FRayTracingGlobalIlluminationCreateGatherPointsTraceRGS::FDeferredMaterialMode>(EDeferredMaterialMode::Gather);
 			TShaderMapRef<FRayTracingGlobalIlluminationCreateGatherPointsTraceRGS> RayGenerationShader(GetGlobalShaderMap(FeatureLevel), PermutationVector);
@@ -1145,7 +1118,6 @@ void FDeferredShadingSceneRenderer::RayTracingGlobalIlluminationCreateGatherPoin
 			GatherPassParameters->MaterialBuffer = GraphBuilder.CreateUAV(DeferredMaterialBuffer);
 
 			FRayTracingGlobalIlluminationCreateGatherPointsRGS::FPermutationDomain PermutationVector;
-			PermutationVector.Set<FRayTracingGlobalIlluminationCreateGatherPointsRGS::FUseAttenuationTermDim>(true);
 			PermutationVector.Set<FRayTracingGlobalIlluminationCreateGatherPointsRGS::FEnableTwoSidedGeometryDim>(CVarRayTracingGlobalIlluminationEnableTwoSidedGeometry.GetValueOnRenderThread() != 0);
 			PermutationVector.Set<FRayTracingGlobalIlluminationCreateGatherPointsRGS::FDeferredMaterialMode>(EDeferredMaterialMode::Shade);
 			PermutationVector.Set<FRayTracingGlobalIlluminationCreateGatherPointsRGS::FEnableTransmissionDim>(CVarRayTracingGlobalIlluminationEnableTransmission.GetValueOnRenderThread());
@@ -1267,7 +1239,6 @@ void FDeferredShadingSceneRenderer::RenderRayTracingGlobalIlluminationFinalGathe
 	PassParameters->RWGlobalIlluminationRayDistanceUAV = GraphBuilder.CreateUAV(OutDenoiserInputs->RayHitDistance);
 
 	FRayTracingGlobalIlluminationFinalGatherRGS::FPermutationDomain PermutationVector;
-	PermutationVector.Set<FRayTracingGlobalIlluminationFinalGatherRGS::FUseAttenuationTermDim>(true);
 	PermutationVector.Set<FRayTracingGlobalIlluminationFinalGatherRGS::FEnableTwoSidedGeometryDim>(CVarRayTracingGlobalIlluminationEnableTwoSidedGeometry.GetValueOnRenderThread() != 0);
 	PermutationVector.Set<FRayTracingGlobalIlluminationFinalGatherRGS::FEnableNeighborVisibilityTestDim>(CVarRayTracingGlobalIlluminationFinalGatherEnableNeighborVisbilityTest.GetValueOnRenderThread() != 0);
 	TShaderMapRef<FRayTracingGlobalIlluminationFinalGatherRGS> RayGenerationShader(GetGlobalShaderMap(FeatureLevel), PermutationVector);
@@ -1382,7 +1353,6 @@ void FDeferredShadingSceneRenderer::RenderRayTracingGlobalIlluminationBruteForce
 	PassParameters->RenderTileOffsetY = 0;
 
 	FGlobalIlluminationRGS::FPermutationDomain PermutationVector;
-	PermutationVector.Set<FGlobalIlluminationRGS::FUseAttenuationTermDim>(true);
 	PermutationVector.Set<FGlobalIlluminationRGS::FEnableTwoSidedGeometryDim>(CVarRayTracingGlobalIlluminationEnableTwoSidedGeometry.GetValueOnRenderThread() != 0);
 	PermutationVector.Set<FGlobalIlluminationRGS::FEnableTransmissionDim>(CVarRayTracingGlobalIlluminationEnableTransmission.GetValueOnRenderThread());
 	TShaderMapRef<FGlobalIlluminationRGS> RayGenerationShader(GetGlobalShaderMap(FeatureLevel), PermutationVector);
