Cleanup unused halton/bluenoise parameters from ray tracing passes

These are no longer needed now that these passes are using the RandomSequence API.

diff --git a/Engine/Source/Runtime/Renderer/Private/RayTracing/RayTracingGlobalIllumination.cpp b/Engine/Source/Runtime/Renderer/Private/RayTracing/RayTracingGlobalIllumination.cpp
--- a/Engine/Source/Runtime/Renderer/Private/RayTracing/RayTracingGlobalIllumination.cpp
+++ b/Engine/Source/Runtime/Renderer/Private/RayTracing/RayTracingGlobalIllumination.cpp
@@ -460,9 +460,6 @@ class FGlobalIlluminationRGS : public FGlobalShader
 		SHADER_PARAMETER_RDG_TEXTURE_UAV(RWTexture2D<float4>, RWGlobalIlluminationUAV)
 		SHADER_PARAMETER_RDG_TEXTURE_UAV(RWTexture2D<float2>, RWGlobalIlluminationRayDistanceUAV)
 		SHADER_PARAMETER_STRUCT_REF(FViewUniformShaderParameters, ViewUniformBuffer)
-		SHADER_PARAMETER_STRUCT_REF(FHaltonIteration, HaltonIteration)
-		SHADER_PARAMETER_STRUCT_REF(FHaltonPrimes, HaltonPrimes)
-		SHADER_PARAMETER_STRUCT_REF(FBlueNoise, BlueNoise)
 		SHADER_PARAMETER_STRUCT_REF(FPathTracingLightData, LightParameters)
 		SHADER_PARAMETER_STRUCT_REF(FSkyLightData, SkyLight)
 		SHADER_PARAMETER_STRUCT_INCLUDE(FSceneTextureParameters, SceneTextures)
@@ -521,11 +518,6 @@ class FRayTracingGlobalIlluminationCreateGatherPointsRGS : public FGlobalShader
 		SHADER_PARAMETER_SRV(RaytracingAccelerationStructure, TLAS)
 		SHADER_PARAMETER_STRUCT_REF(FViewUniformShaderParameters, ViewUniformBuffer)
 
-		// Sampling sequence
-		SHADER_PARAMETER_STRUCT_REF(FHaltonIteration, HaltonIteration)
-		SHADER_PARAMETER_STRUCT_REF(FHaltonPrimes, HaltonPrimes)
-		SHADER_PARAMETER_STRUCT_REF(FBlueNoise, BlueNoise)
-
 		// Light data
 		SHADER_PARAMETER_STRUCT_REF(FPathTracingLightData, LightParameters)
 		SHADER_PARAMETER_STRUCT_REF(FSkyLightData, SkyLight)
@@ -592,11 +584,6 @@ class FRayTracingGlobalIlluminationCreateGatherPointsTraceRGS : public FGlobalSh
 		SHADER_PARAMETER_SRV(RaytracingAccelerationStructure, TLAS)
 		SHADER_PARAMETER_STRUCT_REF(FViewUniformShaderParameters, ViewUniformBuffer)
 
-		// Sampling sequence
-		SHADER_PARAMETER_STRUCT_REF(FHaltonIteration, HaltonIteration)
-		SHADER_PARAMETER_STRUCT_REF(FHaltonPrimes, HaltonPrimes)
-		SHADER_PARAMETER_STRUCT_REF(FBlueNoise, BlueNoise)
-
 		// Light data
 		SHADER_PARAMETER_STRUCT_REF(FPathTracingLightData, LightParameters)
 		SHADER_PARAMETER_STRUCT_REF(FSkyLightData, SkyLight)
@@ -848,10 +835,6 @@ void CopyGatherPassParameters(
 	NewParameters->TLAS = PassParameters.TLAS;
 	NewParameters->ViewUniformBuffer = PassParameters.ViewUniformBuffer;
 
-	NewParameters->HaltonIteration = PassParameters.HaltonIteration;
-	NewParameters->HaltonPrimes = PassParameters.HaltonPrimes;
-	NewParameters->BlueNoise = PassParameters.BlueNoise;
-
 	NewParameters->LightParameters = PassParameters.LightParameters;
 	NewParameters->SkyLight = PassParameters.SkyLight;
 
@@ -892,10 +875,6 @@ void CopyGatherPassParameters(
 	NewParameters->TLAS = PassParameters.TLAS;
 	NewParameters->ViewUniformBuffer = PassParameters.ViewUniformBuffer;
 
-	NewParameters->HaltonIteration = PassParameters.HaltonIteration;
-	NewParameters->HaltonPrimes = PassParameters.HaltonPrimes;
-	NewParameters->BlueNoise = PassParameters.BlueNoise;
-
 	NewParameters->LightParameters = PassParameters.LightParameters;
 	NewParameters->SkyLight = PassParameters.SkyLight;
 
@@ -933,21 +912,6 @@ void FDeferredShadingSceneRenderer::RayTracingGlobalIlluminationCreateGatherPoin
 	int32 GatherFilterWidth = FMath::Max(CVarRayTracingGlobalIlluminationFinalGatherFilterWidth.GetValueOnRenderThread(), 0);
 	GatherFilterWidth = GatherFilterWidth * 2 + 1;
 
-	uint32 IterationCount = GatherFilterWidth * GatherFilterWidth;
-	uint32 SequenceCount = 1;
-	uint32 DimensionCount = 24;
-	int32 FrameIndex = View.ViewState->FrameIndex % 1024;
-	FHaltonSequenceIteration HaltonSequenceIteration(Scene->HaltonSequence, IterationCount, SequenceCount, DimensionCount, FrameIndex);
-
-	FHaltonIteration HaltonIteration;
-	InitializeHaltonSequenceIteration(HaltonSequenceIteration, HaltonIteration);
-
-	FHaltonPrimes HaltonPrimes;
-	InitializeHaltonPrimes(Scene->HaltonPrimesResource, HaltonPrimes);
-
-	FBlueNoise BlueNoise;
-	InitializeBlueNoise(BlueNoise);
-
 	FPathTracingLightData LightParameters;
 	SetupLightParameters(*Scene, View, &LightParameters);
 
@@ -987,11 +951,6 @@ void FDeferredShadingSceneRenderer::RayTracingGlobalIlluminationCreateGatherPoin
 	PassParameters->TLAS = View.RayTracingScene.RayTracingSceneRHI->GetShaderResourceView();
 	PassParameters->ViewUniformBuffer = View.ViewUniformBuffer;
 
-	// Sampling sequence
-	PassParameters->HaltonIteration = CreateUniformBufferImmediate(HaltonIteration, EUniformBufferUsage::UniformBuffer_SingleFrame);
-	PassParameters->HaltonPrimes = CreateUniformBufferImmediate(HaltonPrimes, EUniformBufferUsage::UniformBuffer_SingleFrame);
-	PassParameters->BlueNoise = CreateUniformBufferImmediate(BlueNoise, EUniformBufferUsage::UniformBuffer_SingleFrame);
-
 	// Light data
 	PassParameters->LightParameters = CreateUniformBufferImmediate(LightParameters, EUniformBufferUsage::UniformBuffer_SingleFrame);
 	PassParameters->SceneTextures = SceneTextures;
@@ -1280,19 +1239,6 @@ void FDeferredShadingSceneRenderer::RenderRayTracingGlobalIlluminationBruteForce
 	RDG_EVENT_SCOPE(GraphBuilder, "Ray Tracing GI: Brute Force");
 
 	int32 RayTracingGISamplesPerPixel = GetRayTracingGlobalIlluminationSamplesPerPixel(View);
-	uint32 IterationCount = RayTracingGISamplesPerPixel;
-	uint32 SequenceCount = 1;
-	uint32 DimensionCount = 24;
-	FHaltonSequenceIteration HaltonSequenceIteration(Scene->HaltonSequence, IterationCount, SequenceCount, DimensionCount, View.ViewState->FrameIndex % 1024);
-
-	FHaltonIteration HaltonIteration;
-	InitializeHaltonSequenceIteration(HaltonSequenceIteration, HaltonIteration);
-
-	FHaltonPrimes HaltonPrimes;
-	InitializeHaltonPrimes(Scene->HaltonPrimesResource, HaltonPrimes);
-
-	FBlueNoise BlueNoise;
-	InitializeBlueNoise(BlueNoise);
 
 	FPathTracingLightData LightParameters;
 	SetupLightParameters(*Scene, View, &LightParameters);
@@ -1332,9 +1278,6 @@ void FDeferredShadingSceneRenderer::RenderRayTracingGlobalIlluminationBruteForce
 	PassParameters->NextEventEstimationSamples = GRayTracingGlobalIlluminationNextEventEstimationSamples;
 	PassParameters->TLAS = View.RayTracingScene.RayTracingSceneRHI->GetShaderResourceView();
 	PassParameters->ViewUniformBuffer = View.ViewUniformBuffer;
-	PassParameters->HaltonIteration = CreateUniformBufferImmediate(HaltonIteration, EUniformBufferUsage::UniformBuffer_SingleDraw);
-	PassParameters->HaltonPrimes = CreateUniformBufferImmediate(HaltonPrimes, EUniformBufferUsage::UniformBuffer_SingleDraw);
-	PassParameters->BlueNoise = CreateUniformBufferImmediate(BlueNoise, EUniformBufferUsage::UniformBuffer_SingleDraw);
 	PassParameters->LightParameters = CreateUniformBufferImmediate(LightParameters, EUniformBufferUsage::UniformBuffer_SingleDraw);
 	PassParameters->SceneTextures = SceneTextures;
 	PassParameters->SkyLight = CreateUniformBufferImmediate(SkyLightParameters, EUniformBufferUsage::UniformBuffer_SingleDraw);
diff --git a/Engine/Source/Runtime/Renderer/Private/RayTracing/RayTracingReflections.cpp b/Engine/Source/Runtime/Renderer/Private/RayTracing/RayTracingReflections.cpp
--- a/Engine/Source/Runtime/Renderer/Private/RayTracing/RayTracingReflections.cpp
+++ b/Engine/Source/Runtime/Renderer/Private/RayTracing/RayTracingReflections.cpp
@@ -296,7 +296,6 @@ class FRayTracingReflectionsRGS : public FGlobalShader
 		SHADER_PARAMETER_SRV(StructuredBuffer<FRTLightingData>, LightDataBuffer)
 
 		SHADER_PARAMETER_STRUCT_INCLUDE(FSceneTextureParameters, SceneTextures)
-		SHADER_PARAMETER_STRUCT_INCLUDE(FSkyLightQuasiRandomData, SkyLightQuasiRandomData)
 		SHADER_PARAMETER_STRUCT_INCLUDE(FSkyLightVisibilityRaysData, SkyLightVisibilityRaysData)
 
 		SHADER_PARAMETER_STRUCT_REF(FViewUniformShaderParameters, ViewUniformBuffer)
@@ -698,10 +697,6 @@ void FDeferredShadingSceneRenderer::RenderRayTracingReflections(
 		SetupSkyLightParameters(*Scene, &SkyLightData);
 
 		CommonParameters.SkyLightData = CreateUniformBufferImmediate(SkyLightData, EUniformBufferUsage::UniformBuffer_SingleDraw);
-
-		// Setup Sky Light quasi random parameters
-		FIntVector BlueNoiseDimensions;
-		SetupSkyLightQuasiRandomParameters(*Scene, View, BlueNoiseDimensions, &CommonParameters.SkyLightQuasiRandomData);
 	}
 
 	for (int32 SamplePassIndex = 0; SamplePassIndex < Options.SamplesPerPixel; SamplePassIndex++)
diff --git a/Engine/Source/Runtime/Renderer/Private/RayTracing/RayTracingSkyLight.h b/Engine/Source/Runtime/Renderer/Private/RayTracing/RayTracingSkyLight.h
--- a/Engine/Source/Runtime/Renderer/Private/RayTracing/RayTracingSkyLight.h
+++ b/Engine/Source/Runtime/Renderer/Private/RayTracing/RayTracingSkyLight.h
@@ -8,8 +8,6 @@ RaytracingOptions.h declares ray tracing options for use in rendering
 
 #include "UniformBuffer.h"
 #include "RenderGraph.h"
-#include "Halton.h"
-#include "BlueNoise.h"
 #include "Renderer/Private/SceneRendering.h"
 
 BEGIN_GLOBAL_SHADER_PARAMETER_STRUCT(FSkyLightData, RENDERER_API)
@@ -39,12 +37,6 @@ BEGIN_GLOBAL_SHADER_PARAMETER_STRUCT(FSkyLightData, RENDERER_API)
 	SHADER_PARAMETER_SRV(Buffer<float>, SolidAnglePdf)
 END_GLOBAL_SHADER_PARAMETER_STRUCT()
 
-BEGIN_SHADER_PARAMETER_STRUCT(FSkyLightQuasiRandomData, )
-	SHADER_PARAMETER_STRUCT_REF(FHaltonIteration, HaltonIteration)
-	SHADER_PARAMETER_STRUCT_REF(FHaltonPrimes, HaltonPrimes)
-	SHADER_PARAMETER_STRUCT_REF(FBlueNoise, BlueNoise)
-END_SHADER_PARAMETER_STRUCT()
-
 BEGIN_SHADER_PARAMETER_STRUCT(FWritableSkyLightVisibilityRaysData, )
 	SHADER_PARAMETER_RDG_BUFFER_UAV(RWStructuredBuffer<SkyLightVisibilityRays>, OutSkyLightVisibilityRays)
 	SHADER_PARAMETER(FIntVector, SkyLightVisibilityRaysDimensions)
@@ -60,7 +52,6 @@ END_SHADER_PARAMETER_STRUCT()
 int32 GetRayTracingSkyLightDecoupleSampleGenerationCVarValue();
 
 extern RENDERER_API bool SetupSkyLightParameters(const FScene& Scene, FSkyLightData* SkyLight);
-void SetupSkyLightQuasiRandomParameters(const FScene& Scene, const FViewInfo& View, FIntVector& OutBlueNoiseDimensions, FSkyLightQuasiRandomData* OutSkyLightQuasiRandomData);
 void SetupSkyLightVisibilityRaysParameters(FRDGBuilder& GraphBuilder, const FViewInfo& View, FSkyLightVisibilityRaysData* OutSkyLightVisibilityRaysData);
 
 extern RENDERER_API void BuildSkyLightCdfs(FRHICommandListImmediate& RHICmdList, class FSkyLightSceneProxy* SkyLight);
diff --git a/Engine/Source/Runtime/Renderer/Private/RayTracing/RaytracingSkylight.cpp b/Engine/Source/Runtime/Renderer/Private/RayTracing/RaytracingSkylight.cpp
--- a/Engine/Source/Runtime/Renderer/Private/RayTracing/RaytracingSkylight.cpp
+++ b/Engine/Source/Runtime/Renderer/Private/RayTracing/RaytracingSkylight.cpp
@@ -232,49 +232,6 @@ bool SetupSkyLightParameters(
 	}
 }
 
-void SetupSkyLightQuasiRandomParameters(
-	const FScene& Scene,
-	const FViewInfo& View,
-	FIntVector& OutBlueNoiseDimensions,
-	FSkyLightQuasiRandomData* OutSkyLightQuasiRandomData)
-{
-	uint32 IterationCount;
-
-	// Choose to set the iteration count to the sky light samples per pixel or a dummy value depending on its presence
-	if (Scene.SkyLight != nullptr)
-	{
-		IterationCount = FMath::Max(GetSkyLightSamplesPerPixel(Scene.SkyLight), 1);
-	}
-	else
-	{
-		IterationCount = 1;
-	}
-
-	// Halton iteration setup
-	uint32 SequenceCount = 1;
-	uint32 DimensionCount = 3;
-	FHaltonSequenceIteration HaltonSequenceIteration(Scene.HaltonSequence, IterationCount, SequenceCount, DimensionCount, View.ViewState ? (View.ViewState->FrameIndex % 1024) : 0);
-
-	FHaltonIteration HaltonIteration;
-	InitializeHaltonSequenceIteration(HaltonSequenceIteration, HaltonIteration);
-	
-	// Halton primes setup
-	FHaltonPrimes HaltonPrimes;
-	InitializeHaltonPrimes(Scene.HaltonPrimesResource, HaltonPrimes);
-
-	// Blue noise setup
-	FBlueNoise BlueNoise;
-	InitializeBlueNoise(BlueNoise);
-
-	// Set the Blue Noise dimensions
-	OutBlueNoiseDimensions = FIntVector(BlueNoise.Dimensions.X, BlueNoise.Dimensions.Y, 0);
-
-	// Set Sky Light Quasi Random Data information
-	OutSkyLightQuasiRandomData->HaltonIteration = CreateUniformBufferImmediate(HaltonIteration, EUniformBufferUsage::UniformBuffer_SingleDraw);
-	OutSkyLightQuasiRandomData->HaltonPrimes = CreateUniformBufferImmediate(HaltonPrimes, EUniformBufferUsage::UniformBuffer_SingleDraw);
-	OutSkyLightQuasiRandomData->BlueNoise = CreateUniformBufferImmediate(BlueNoise, EUniformBufferUsage::UniformBuffer_SingleDraw);
-}
-
 void SetupSkyLightVisibilityRaysParameters(
 	FRDGBuilder& GraphBuilder,
 	const FViewInfo& View,
@@ -344,7 +301,6 @@ class FRayTracingSkyLightRGS : public FGlobalShader
 		SHADER_PARAMETER_STRUCT_REF(FSkyLightData, SkyLightData)
 		SHADER_PARAMETER_RDG_UNIFORM_BUFFER(FVirtualVoxelParameters, VirtualVoxel)
 		
-		SHADER_PARAMETER_STRUCT_INCLUDE(FSkyLightQuasiRandomData, SkyLightQuasiRandomData)
 		SHADER_PARAMETER_STRUCT_INCLUDE(FSkyLightVisibilityRaysData, SkyLightVisibilityRaysData)
 		SHADER_PARAMETER_STRUCT_INCLUDE(FSceneTextureParameters, SceneTextures)
 		SHADER_PARAMETER_TEXTURE(Texture2D, SSProfilesTexture)
@@ -408,7 +364,6 @@ class FGenerateSkyLightVisibilityRaysCS : public FGlobalShader
 		SHADER_PARAMETER_STRUCT_REF(FSkyLightData, SkyLightData)
 		SHADER_PARAMETER_STRUCT_REF(FViewUniformShaderParameters, ViewUniformBuffer)
 
-		SHADER_PARAMETER_STRUCT_INCLUDE(FSkyLightQuasiRandomData, SkyLightQuasiRandomData)
 		// Writable variant to allow for Sky Light Visibility Ray output
 		SHADER_PARAMETER_STRUCT_INCLUDE(FWritableSkyLightVisibilityRaysData, WritableSkyLightVisibilityRaysData)
 	END_SHADER_PARAMETER_STRUCT()
@@ -429,11 +384,6 @@ void FDeferredShadingSceneRenderer::GenerateSkyLightVisibilityRays(
 	FSkyLightData SkyLightData;
 	SetupSkyLightParameters(*Scene, &SkyLightData);
 
-	// Sky Light Quasi Random data setup
-	// Note: Sets Dimensions based on blue noise dimensions
-	FSkyLightQuasiRandomData SkyLightQuasiRandomData;
-	SetupSkyLightQuasiRandomParameters(*Scene, Views[0], Dimensions, &SkyLightQuasiRandomData);
-
 	// Output structured buffer creation
 	FRDGBufferDesc BufferDesc = FRDGBufferDesc::CreateStructuredDesc(sizeof(FSkyLightVisibilityRays), Dimensions.X * Dimensions.Y * SkyLightData.SamplesPerPixel);
 	SkyLightVisibilityRaysBuffer = GraphBuilder.CreateBuffer(BufferDesc, TEXT("SkyLightVisibilityRays"));
@@ -442,7 +392,6 @@ void FDeferredShadingSceneRenderer::GenerateSkyLightVisibilityRays(
 	FGenerateSkyLightVisibilityRaysCS::FParameters* PassParameters = GraphBuilder.AllocParameters<FGenerateSkyLightVisibilityRaysCS::FParameters>();
 	PassParameters->ViewUniformBuffer = Views[0].ViewUniformBuffer;
 	PassParameters->SkyLightData = CreateUniformBufferImmediate(SkyLightData, EUniformBufferUsage::UniformBuffer_SingleDraw);
-	PassParameters->SkyLightQuasiRandomData = SkyLightQuasiRandomData;
 	PassParameters->WritableSkyLightVisibilityRaysData.SkyLightVisibilityRaysDimensions = FIntVector(Dimensions.X, Dimensions.Y, 0);
 	PassParameters->WritableSkyLightVisibilityRaysData.OutSkyLightVisibilityRays = GraphBuilder.CreateUAV(SkyLightVisibilityRaysBuffer, EPixelFormat::PF_R32_UINT);
 
@@ -537,16 +486,10 @@ void FDeferredShadingSceneRenderer::RenderRayTracingSkyLight(
 
 		FSceneViewState* SceneViewState = (FSceneViewState*)View.State;
 
-		// Sky Light Quasi Random data setup
-		FIntVector BlueNoiseDimensions;
-		FSkyLightQuasiRandomData SkyLightQuasiRandomData;
-		SetupSkyLightQuasiRandomParameters(*Scene, View, BlueNoiseDimensions, &SkyLightQuasiRandomData);
-
 		FRayTracingSkyLightRGS::FParameters *PassParameters = GraphBuilder.AllocParameters<FRayTracingSkyLightRGS::FParameters>();
 		PassParameters->RWSkyOcclusionMaskUAV = SkyLightkUAV;
 		PassParameters->RWSkyOcclusionRayDistanceUAV = RayDistanceUAV;
 		PassParameters->SkyLightData = CreateUniformBufferImmediate(SkyLightData, EUniformBufferUsage::UniformBuffer_SingleDraw);
-		PassParameters->SkyLightQuasiRandomData = SkyLightQuasiRandomData;
 		PassParameters->SkyLightVisibilityRaysData.SkyLightVisibilityRaysDimensions = SkyLightVisibilityRaysDimensions;
 		if (CVarRayTracingSkyLightDecoupleSampleGeneration.GetValueOnRenderThread() == 1)
 		{
