Change IES lookup to use a Texture2DArray instead of a Texture3D as this more cleanly reflects how the atlas is built (no interpolation between slices is desired).

diff --git a/Engine/Shaders/Private/PathTracing/Light/PathTracingLightCommon.ush b/Engine/Shaders/Private/PathTracing/Light/PathTracingLightCommon.ush
--- a/Engine/Shaders/Private/PathTracing/Light/PathTracingLightCommon.ush
+++ b/Engine/Shaders/Private/PathTracing/Light/PathTracingLightCommon.ush
@@ -150,8 +150,8 @@ float ComputeAttenuationFalloff(float DistanceSquared, int LightId)
 }
 
 
-Texture3D    IESTexture;
-SamplerState IESTextureSampler;
+Texture2DArray<float> IESTexture;
+SamplerState          IESTextureSampler;
 
 float ComputeIESAttenuation(int LightId, float3 WorldPosition)
 {
@@ -169,7 +169,7 @@ float ComputeIESAttenuation(int LightId, float3 WorldPosition)
 		float NormTangentAngle = atan2(dv, du) / (PI * 2.f) + 0.5f;
 
 		const float3 UVW = float3(NormAngle, NormTangentAngle, SceneLightsData.IESTextureSlice[LightId]);
-		Result = Texture3DSampleLevel(IESTexture, IESTextureSampler, UVW, 0).r;
+		Result = IESTexture.SampleLevel(IESTextureSampler, UVW, 0);
 	}
 	return Result;
 }
diff --git a/Engine/Shaders/Private/PathTracing/PathTracingIESAtlas.usf b/Engine/Shaders/Private/PathTracing/PathTracingIESAtlas.usf
--- a/Engine/Shaders/Private/PathTracing/PathTracingIESAtlas.usf
+++ b/Engine/Shaders/Private/PathTracing/PathTracingIESAtlas.usf
@@ -4,7 +4,7 @@
 
 Texture2D IESTexture;
 SamplerState IESSampler;
-RWTexture3D<float> IESAtlas;
+RWTexture2DArray<float> IESAtlas;
 int IESAtlasSlice;
 
 [numthreads(THREADGROUPSIZE_X, THREADGROUPSIZE_Y, 1)]
@@ -18,5 +18,5 @@ void PathTracingIESAtlasCS(uint2 DispatchThreadId : SV_DispatchThreadID)
 	}
 
 	float2 UV = (DispatchThreadId + 0.5) / float2(AtlasSize.xy);
-	IESAtlas[uint3(DispatchThreadId, IESAtlasSlice)] = Texture2DSampleLevel(IESTexture, IESSampler, UV, 0).r;
+	IESAtlas[uint3(DispatchThreadId, IESAtlasSlice)] = IESTexture.SampleLevel(IESSampler, UV, 0).r;
 }
diff --git a/Engine/Source/Runtime/Renderer/Private/PathTracing.cpp b/Engine/Source/Runtime/Renderer/Private/PathTracing.cpp
--- a/Engine/Source/Runtime/Renderer/Private/PathTracing.cpp
+++ b/Engine/Source/Runtime/Renderer/Private/PathTracing.cpp
@@ -177,7 +177,7 @@ class FPathTracingRG : public FGlobalShader
 		SHADER_PARAMETER_STRUCT_REF(FPathTracingLightData, SceneLightsData)
 		SHADER_PARAMETER_STRUCT_REF(FPathTracingData, PathTracingData)
 		// IES Profiles
-		SHADER_PARAMETER_RDG_TEXTURE(Texture3D, IESTexture)
+		SHADER_PARAMETER_RDG_TEXTURE(Texture2DArray, IESTexture)
 		SHADER_PARAMETER_SAMPLER(SamplerState, IESTextureSampler)
 		// Shared sampler for all IES profiles
 		SHADER_PARAMETER(FIntVector, TileOffset)	// Used by multi-GPU rendering
@@ -205,7 +205,7 @@ class FPathTracingIESAtlasCS : public FGlobalShader
 	BEGIN_SHADER_PARAMETER_STRUCT(FParameters, )
 		SHADER_PARAMETER_TEXTURE(Texture2D, IESTexture)
 		SHADER_PARAMETER_SAMPLER(SamplerState, IESSampler)
-		SHADER_PARAMETER_RDG_TEXTURE_UAV(RWTexture3D<float>, IESAtlas)
+		SHADER_PARAMETER_RDG_TEXTURE_UAV(RWTexture2DArray, IESAtlas)
 		SHADER_PARAMETER(int32, IESAtlasSlice)
 	END_SHADER_PARAMETER_STRUCT()
 };
@@ -329,12 +329,12 @@ void SetLightParameters(FRDGBuilder& GraphBuilder, FPathTracingRG::FParameters*
 		// This size matches the import resolution of light profiles (see FIESLoader::GetWidth)
 		const int kIESAtlasSize = 256;
 		const int NumSlices = IESLightProfilesMap.Num();
-		// TODO: would a Texture2DArray object be more efficient?
-		FRDGTextureDesc IESTextureDesc = FRDGTextureDesc::Create3D(
-			FIntVector(kIESAtlasSize, kIESAtlasSize, NumSlices),
+		FRDGTextureDesc IESTextureDesc = FRDGTextureDesc::Create2DArray(
+			FIntPoint(kIESAtlasSize, kIESAtlasSize),
 			PF_R32_FLOAT,
 			FClearValueBinding::None,
-			TexCreate_ShaderResource | TexCreate_UAV);
+			TexCreate_ShaderResource | TexCreate_UAV,
+			NumSlices);
 		FRDGTexture* IESTexture = GraphBuilder.CreateTexture(IESTextureDesc, TEXT("PathTracerIESAtlas"), ERDGTextureFlags::None);
 
 		for (auto&& Entry : IESLightProfilesMap)
@@ -354,16 +354,6 @@ void SetLightParameters(FRDGBuilder& GraphBuilder, FPathTracingRG::FParameters*
 				FComputeShaderUtils::GetGroupCount(FIntPoint(kIESAtlasSize, kIESAtlasSize), FComputeShaderUtils::kGolden2DGroupSize));
 		}
 
-		// Adjust row per light so it can be used directly as a texture coordinate
-		for (uint32 Index = 0; Index < LightData.Count; Index++)
-		{
-			if (LightData.IESTextureSlice[Index] >= 0)
-			{
-				// nudge data to be easier to use on GPU
-				LightData.IESTextureSlice[Index] = (LightData.IESTextureSlice[Index] + 0.5f) / NumSlices;
-			}
-		}
-
 		PassParameters->IESTexture = IESTexture;
 	}
 	else
