Add more info to a check() that's reported to sometimes fail

diff --git a/Engine/Source/Runtime/Renderer/Private/VT/AllocatedVirtualTexture.cpp b/Engine/Source/Runtime/Renderer/Private/VT/AllocatedVirtualTexture.cpp
--- a/Engine/Source/Runtime/Renderer/Private/VT/AllocatedVirtualTexture.cpp
+++ b/Engine/Source/Runtime/Renderer/Private/VT/AllocatedVirtualTexture.cpp
@@ -69,7 +69,12 @@ FAllocatedVirtualTexture::FAllocatedVirtualTexture(FVirtualTextureSystem* InSyst
 			const uint32 MipBias = UniqueProducers[ProducerIndex].MipBias;
 			check(MipBias <= MaxLevel);
 			const uint32 Local_vLevel = MaxLevel - MipBias;
-			check(Local_vLevel <= Producer->GetMaxLevel());
+			checkf(Local_vLevel <= Producer->GetMaxLevel(), TEXT("Invalid Local_vLevel %d for VT producer %s, Producer MaxLevel %d, MipBias %d, AllocatedVT MaxLevel %d"),
+				Local_vLevel,
+				*Producer->GetName().ToString(),
+				Producer->GetMaxLevel(),
+				MipBias,
+				MaxLevel);
 
 			const uint32 MipScaleFactor = (1u << Local_vLevel);
 			const uint32 RootWidthInTiles = FMath::DivideAndRoundUp(Producer->GetWidthInTiles(), MipScaleFactor);
