Fixing DBuffer mask ensure on render target.

diff --git a/Engine/Source/Runtime/Renderer/Private/CompositionLighting/PostProcessDeferredDecals.cpp b/Engine/Source/Runtime/Renderer/Private/CompositionLighting/PostProcessDeferredDecals.cpp
--- a/Engine/Source/Runtime/Renderer/Private/CompositionLighting/PostProcessDeferredDecals.cpp
+++ b/Engine/Source/Runtime/Renderer/Private/CompositionLighting/PostProcessDeferredDecals.cpp
@@ -541,7 +541,7 @@ FDBufferTextures CreateDBufferTextures(FRDGBuilder& GraphBuilder, FIntPoint Exte
 			// This significantly reduces bandwidth for clearing, writing and reading on some GPUs.
 			// While a smaller format, such as R8_UINT, will use less video memory, it will result in slower clears and higher bandwidth requirements.
 			check(Desc.Format == PF_B8G8R8A8);
-			Desc.Flags = TexCreate_ShaderResource;
+			Desc.Flags = TexCreate_ShaderResource | TexCreate_RenderTargetable;
 			Desc.ClearValue = FClearValueBinding::Transparent;
 			DBufferTextures.DBufferMask = GraphBuilder.CreateTexture(Desc, TEXT("DBufferMask"));
 		}
