Add interface to support case when the VT feedback buffer is a full FRDGBuffer.
New interface will be used by VHM feedback.
Old interface can be deprecated whenever main pass feedback buffer is converted to RDG.
#rnx

diff --git a/Engine/Source/Runtime/Renderer/Private/VT/VirtualTextureFeedback.cpp b/Engine/Source/Runtime/Renderer/Private/VT/VirtualTextureFeedback.cpp
--- a/Engine/Source/Runtime/Renderer/Private/VT/VirtualTextureFeedback.cpp
+++ b/Engine/Source/Runtime/Renderer/Private/VT/VirtualTextureFeedback.cpp
@@ -2,6 +2,9 @@
 
 #include "VT/VirtualTextureFeedback.h"
 
+#include "RenderGraphBuilder.h"
+#include "RenderGraphUtils.h"
+
 #if PLATFORM_WINDOWS
 // Use Query objects until RHI has a good fence on D3D11
 #define USE_RHI_FENCES 0
@@ -179,8 +182,6 @@ void FVirtualTextureFeedback::ReleaseRHI()
 
 void FVirtualTextureFeedback::TransferGPUToCPU(FRHICommandListImmediate& RHICmdList, FVertexBufferRHIRef const& Buffer, FVirtualTextureFeedbackBufferDesc const& Desc)
 {
-	QUICK_SCOPE_CYCLE_COUNTER(STAT_VirtualTextureFeedback_TransferGPUToCPU);
-
 	if (NumPending >= MaxTransfers)
 	{
 		// If we have too many pending transfers, start throwing away the oldest in the ring buffer.
@@ -207,6 +208,25 @@ void FVirtualTextureFeedback::TransferGPUToCPU(FRHICommandListImmediate& RHICmdL
 	++NumPending;
 }
 
+BEGIN_SHADER_PARAMETER_STRUCT(FVirtualTextureFeedbackCopyParameters, )
+	RDG_BUFFER_ACCESS(Input, ERHIAccess::CopySrc)
+END_SHADER_PARAMETER_STRUCT()
+
+void FVirtualTextureFeedback::TransferGPUToCPU(FRDGBuilder& GraphBuilder, FRDGBuffer* Buffer, FVirtualTextureFeedbackBufferDesc const& Desc)
+{
+	FVirtualTextureFeedbackCopyParameters* Parameters = GraphBuilder.AllocParameters<FVirtualTextureFeedbackCopyParameters>();
+	Parameters->Input = Buffer;
+
+	GraphBuilder.AddPass(
+		RDG_EVENT_NAME("VirtualTextureFeedbackCopy"),
+		Parameters,
+		ERDGPassFlags::Readback,
+		[this, Buffer, Desc](FRHICommandListImmediate& InRHICmdList)
+	{
+		TransferGPUToCPU(InRHICmdList, Buffer->GetRHI(), Desc);
+	});
+}
+
 bool FVirtualTextureFeedback::CanMap(FRHICommandListImmediate& RHICmdList)
 {
 	if (NumPending > 0u)
diff --git a/Engine/Source/Runtime/Renderer/Private/VT/VirtualTextureFeedback.h b/Engine/Source/Runtime/Renderer/Private/VT/VirtualTextureFeedback.h
--- a/Engine/Source/Runtime/Renderer/Private/VT/VirtualTextureFeedback.h
+++ b/Engine/Source/Runtime/Renderer/Private/VT/VirtualTextureFeedback.h
@@ -20,6 +20,7 @@ public:
 
 	/** Commit a RHIBuffer to be transferred for later CPU analysis. */
 	void TransferGPUToCPU(FRHICommandListImmediate& RHICmdList, FVertexBufferRHIRef const& Buffer, FVirtualTextureFeedbackBufferDesc const& Desc);
+	void TransferGPUToCPU(FRDGBuilder& GraphBuilder, class FRDGBuffer* Buffer, FVirtualTextureFeedbackBufferDesc const& Desc);
 
 	/** Returns true if there are any pending transfer results that are ready so that we can call Map(). */
 	bool CanMap(FRHICommandListImmediate& RHICmdList);
diff --git a/Engine/Source/Runtime/Renderer/Private/VT/VirtualTextureFeedbackBuffer.cpp b/Engine/Source/Runtime/Renderer/Private/VT/VirtualTextureFeedbackBuffer.cpp
--- a/Engine/Source/Runtime/Renderer/Private/VT/VirtualTextureFeedbackBuffer.cpp
+++ b/Engine/Source/Runtime/Renderer/Private/VT/VirtualTextureFeedbackBuffer.cpp
@@ -44,3 +44,9 @@ void SubmitVirtualTextureFeedbackBuffer(FRHICommandListImmediate& RHICmdList, FV
 {
 	GVirtualTextureFeedback.TransferGPUToCPU(RHICmdList, Buffer, Desc);
 }
+
+void SubmitVirtualTextureFeedbackBuffer(class FRDGBuilder& GraphBuilder, FRDGBuffer* Buffer, FVirtualTextureFeedbackBufferDesc const& Desc)
+{
+	GVirtualTextureFeedback.TransferGPUToCPU(GraphBuilder, Buffer, Desc);
+}
+
diff --git a/Engine/Source/Runtime/Renderer/Public/VT/VirtualTextureFeedbackBuffer.h b/Engine/Source/Runtime/Renderer/Public/VT/VirtualTextureFeedbackBuffer.h
--- a/Engine/Source/Runtime/Renderer/Public/VT/VirtualTextureFeedbackBuffer.h
+++ b/Engine/Source/Runtime/Renderer/Public/VT/VirtualTextureFeedbackBuffer.h
@@ -41,3 +41,4 @@ struct FVirtualTextureFeedbackBufferDesc
  * The function can be called from the render thread only.
 */
 RENDERER_API void SubmitVirtualTextureFeedbackBuffer(class FRHICommandListImmediate& RHICmdList, FVertexBufferRHIRef const& Buffer, FVirtualTextureFeedbackBufferDesc const& Desc);
+RENDERER_API void SubmitVirtualTextureFeedbackBuffer(class FRDGBuilder& GraphBuilder, class FRDGBuffer* Buffer, FVirtualTextureFeedbackBufferDesc const& Desc);
