Tidy virtual texture update event markers.
#rnx

diff --git a/Engine/Source/Runtime/Renderer/Private/DeferredShadingRenderer.cpp b/Engine/Source/Runtime/Renderer/Private/DeferredShadingRenderer.cpp
--- a/Engine/Source/Runtime/Renderer/Private/DeferredShadingRenderer.cpp
+++ b/Engine/Source/Runtime/Renderer/Private/DeferredShadingRenderer.cpp
@@ -1711,11 +1711,12 @@ void FDeferredShadingSceneRenderer::Render(FRDGBuilder& GraphBuilder)
 	// Note, should happen after the GPU-Scene update to ensure rendering to runtime virtual textures is using the correctly updated scene
 	if (bUseVirtualTexturing)
 	{
+		RDG_EVENT_SCOPE(GraphBuilder, "VirtualTextureUpdate");
 		RDG_GPU_STAT_SCOPE(GraphBuilder, VirtualTextureUpdate);
 
 		FVirtualTextureSystem::Get().Update(GraphBuilder, FeatureLevel, Scene);
 
-		AddPass(GraphBuilder, [this, &SceneContext](FRHICommandListImmediate& InRHICmdList)
+		AddPass(GraphBuilder, RDG_EVENT_NAME("VirtualTextureFeedbackClear"), [this, &SceneContext](FRHICommandListImmediate& InRHICmdList)
 		{
 			// Clear virtual texture feedback to default value
 			FUnorderedAccessViewRHIRef FeedbackUAV = SceneContext.GetVirtualTextureFeedbackUAV();
@@ -2707,7 +2708,7 @@ void FDeferredShadingSceneRenderer::Render(FRDGBuilder& GraphBuilder)
 	{
 		RDG_GPU_STAT_SCOPE(GraphBuilder, VirtualTextureUpdate);
 
-		AddPass(GraphBuilder, RDG_EVENT_NAME("VirtualTextureUpdate"), [this, &SceneContext](FRHICommandListImmediate& InRHICmdList)
+		AddPass(GraphBuilder, [this, &SceneContext](FRHICommandListImmediate& InRHICmdList)
 		{
 			InRHICmdList.EndUAVOverlap(SceneContext.VirtualTextureFeedbackUAV);
 
diff --git a/Engine/Source/Runtime/Renderer/Private/MobileShadingRenderer.cpp b/Engine/Source/Runtime/Renderer/Private/MobileShadingRenderer.cpp
--- a/Engine/Source/Runtime/Renderer/Private/MobileShadingRenderer.cpp
+++ b/Engine/Source/Runtime/Renderer/Private/MobileShadingRenderer.cpp
@@ -605,10 +605,12 @@ void FMobileSceneRenderer::Render(FRDGBuilder& GraphBuilder)
 
 	if (bUseVirtualTexturing)
 	{
+		RDG_EVENT_SCOPE(GraphBuilder, "VirtualTextureUpdate");
 		RDG_GPU_STAT_SCOPE(GraphBuilder, VirtualTextureUpdate);
+
 		FVirtualTextureSystem::Get().Update(GraphBuilder, FeatureLevel, Scene);
 
-		AddPass(GraphBuilder, [&SceneContext](FRHICommandList& InRHICmdList)
+		AddPass(GraphBuilder, RDG_EVENT_NAME("VirtualTextureFeedbackClear"), [&SceneContext](FRHICommandList& InRHICmdList)
 		{
 			// Clear virtual texture feedback to default value
 			FUnorderedAccessViewRHIRef FeedbackUAV = SceneContext.GetVirtualTextureFeedbackUAV();
@@ -754,7 +756,7 @@ void FMobileSceneRenderer::Render(FRDGBuilder& GraphBuilder)
 	{
 		RDG_GPU_STAT_SCOPE(GraphBuilder, VirtualTextureUpdate);
 
-		AddPass(GraphBuilder, RDG_EVENT_NAME("VirtualTextureUpdate"), [this, &SceneContext](FRHICommandListImmediate& InRHICmdList)
+		AddPass(GraphBuilder, [this, &SceneContext](FRHICommandListImmediate& InRHICmdList)
 		{
 			// No pass after this should make VT page requests
 			InRHICmdList.Transition(FRHITransitionInfo(SceneContext.VirtualTextureFeedbackUAV, ERHIAccess::UAVMask, ERHIAccess::SRVMask));
diff --git a/Engine/Source/Runtime/Renderer/Private/VT/RuntimeVirtualTextureProducer.cpp b/Engine/Source/Runtime/Renderer/Private/VT/RuntimeVirtualTextureProducer.cpp
--- a/Engine/Source/Runtime/Renderer/Private/VT/RuntimeVirtualTextureProducer.cpp
+++ b/Engine/Source/Runtime/Renderer/Private/VT/RuntimeVirtualTextureProducer.cpp
@@ -64,6 +64,8 @@ void FRuntimeVirtualTextureFinalizer::AddTile(FTileEntry& Tile)
 
 void FRuntimeVirtualTextureFinalizer::Finalize(FRDGBuilder& GraphBuilder)
 {
+	RDG_EVENT_SCOPE(GraphBuilder, "RuntimeVirtualTextureFinalize");
+
 	RuntimeVirtualTexture::FRenderPageBatchDesc RenderPageBatchDesc;
 	RenderPageBatchDesc.Scene = Scene->GetRenderScene();
 	RenderPageBatchDesc.RuntimeVirtualTextureMask = RuntimeVirtualTextureMask;
diff --git a/Engine/Source/Runtime/Renderer/Private/VT/RuntimeVirtualTextureRender.cpp b/Engine/Source/Runtime/Renderer/Private/VT/RuntimeVirtualTextureRender.cpp
--- a/Engine/Source/Runtime/Renderer/Private/VT/RuntimeVirtualTextureRender.cpp
+++ b/Engine/Source/Runtime/Renderer/Private/VT/RuntimeVirtualTextureRender.cpp
@@ -1132,7 +1132,6 @@ namespace RuntimeVirtualTexture
 
 	void RenderPages(FRDGBuilder& GraphBuilder, FRenderPageBatchDesc const& InDesc)
 	{
-		RDG_EVENT_SCOPE(GraphBuilder, "RuntimeVirtualTextureRenderPages");
 		check(InDesc.NumPageDescs <= EMaxRenderPageBatch);
 
 		for (int32 PageIndex = 0; PageIndex < InDesc.NumPageDescs; ++PageIndex)
