Starta visualize/debug shader now do not request optimisation for shorter compilation time.

diff --git a/Engine/Source/Runtime/Renderer/Private/Strata/Strata.cpp b/Engine/Source/Runtime/Renderer/Private/Strata/Strata.cpp
--- a/Engine/Source/Runtime/Renderer/Private/Strata/Strata.cpp
+++ b/Engine/Source/Runtime/Renderer/Private/Strata/Strata.cpp
@@ -173,6 +173,9 @@ class FVisualizeMaterialPS : public FGlobalShader
 	static void ModifyCompilationEnvironment(const FGlobalShaderPermutationParameters& Parameters, FShaderCompilerEnvironment& OutEnvironment)
 	{
 		FGlobalShader::ModifyCompilationEnvironment(Parameters, OutEnvironment);
+
+		// Stay debug and skip optimizations to reduce compilation time on this long shader.
+		OutEnvironment.CompilerFlags.Add(CFLAG_Debug);
 	}
 };
 IMPLEMENT_GLOBAL_SHADER(FVisualizeMaterialPS, "/Engine/Private/Strata/StrataVisualize.usf", "VisualizeMaterialPS", SF_Pixel);
