Added an RDG scope marker removed by mistake.
#rnx

diff --git a/Engine/Source/Runtime/Renderer/Private/DeferredShadingRenderer.cpp b/Engine/Source/Runtime/Renderer/Private/DeferredShadingRenderer.cpp
--- a/Engine/Source/Runtime/Renderer/Private/DeferredShadingRenderer.cpp
+++ b/Engine/Source/Runtime/Renderer/Private/DeferredShadingRenderer.cpp
@@ -2708,7 +2708,7 @@ void FDeferredShadingSceneRenderer::Render(FRDGBuilder& GraphBuilder)
 	{
 		RDG_GPU_STAT_SCOPE(GraphBuilder, VirtualTextureUpdate);
 
-		AddPass(GraphBuilder, [this, &SceneContext](FRHICommandListImmediate& InRHICmdList)
+		AddPass(GraphBuilder, RDG_EVENT_NAME("VirtualTextureFeedbackCopy"), [this, &SceneContext](FRHICommandListImmediate& InRHICmdList)
 		{
 			InRHICmdList.EndUAVOverlap(SceneContext.VirtualTextureFeedbackUAV);
 
diff --git a/Engine/Source/Runtime/Renderer/Private/MobileShadingRenderer.cpp b/Engine/Source/Runtime/Renderer/Private/MobileShadingRenderer.cpp
--- a/Engine/Source/Runtime/Renderer/Private/MobileShadingRenderer.cpp
+++ b/Engine/Source/Runtime/Renderer/Private/MobileShadingRenderer.cpp
@@ -756,7 +756,7 @@ void FMobileSceneRenderer::Render(FRDGBuilder& GraphBuilder)
 	{
 		RDG_GPU_STAT_SCOPE(GraphBuilder, VirtualTextureUpdate);
 
-		AddPass(GraphBuilder, [this, &SceneContext](FRHICommandListImmediate& InRHICmdList)
+		AddPass(GraphBuilder, RDG_EVENT_NAME("VirtualTextureFeedbackCopy"), [this, &SceneContext](FRHICommandListImmediate& InRHICmdList)
 		{
 			// No pass after this should make VT page requests
 			InRHICmdList.Transition(FRHITransitionInfo(SceneContext.VirtualTextureFeedbackUAV, ERHIAccess::UAVMask, ERHIAccess::SRVMask));
