Added GetSceneTextureParameters variant which takes an FSceneTextures instance.

#rnx

diff --git a/Engine/Source/Runtime/Renderer/Private/SceneTextureParameters.cpp b/Engine/Source/Runtime/Renderer/Private/SceneTextureParameters.cpp
--- a/Engine/Source/Runtime/Renderer/Private/SceneTextureParameters.cpp
+++ b/Engine/Source/Runtime/Renderer/Private/SceneTextureParameters.cpp
@@ -8,7 +8,11 @@
 
 FSceneTextureParameters GetSceneTextureParameters(FRDGBuilder& GraphBuilder)
 {
-	const auto& SceneTextures = FSceneTextures::Get(GraphBuilder);
+	return GetSceneTextureParameters(GraphBuilder, FSceneTextures::Get(GraphBuilder));
+}
+
+FSceneTextureParameters GetSceneTextureParameters(FRDGBuilder& GraphBuilder, const FSceneTextures& SceneTextures)
+{
 	const auto& SystemTextures = FRDGSystemTextures::Get(GraphBuilder);
 
 	FSceneTextureParameters Parameters;
diff --git a/Engine/Source/Runtime/Renderer/Private/SceneTextureParameters.h b/Engine/Source/Runtime/Renderer/Private/SceneTextureParameters.h
--- a/Engine/Source/Runtime/Renderer/Private/SceneTextureParameters.h
+++ b/Engine/Source/Runtime/Renderer/Private/SceneTextureParameters.h
@@ -8,6 +8,8 @@
 
 #include "SceneRenderTargetParameters.h"
 
+struct FSceneTextures;
+
 /** Contains reference to scene GBuffer textures. */
 BEGIN_SHADER_PARAMETER_STRUCT(FSceneTextureParameters, )
 	SHADER_PARAMETER_RDG_TEXTURE(Texture2D, SceneDepthTexture)
@@ -24,6 +26,9 @@ END_SHADER_PARAMETER_STRUCT()
 /** Constructs scene texture parameters from the scene context. */
 FSceneTextureParameters GetSceneTextureParameters(FRDGBuilder& GraphBuilder);
 
+/** Constructs scene texture parameters from the scene texture blackboard struct. */
+FSceneTextureParameters GetSceneTextureParameters(FRDGBuilder& GraphBuilder, const FSceneTextures& SceneTextures);
+
 /** Constructs scene texture parameters from the scene texture uniform buffer. Useful if you prefer to use loose parameters. */
 FSceneTextureParameters GetSceneTextureParameters(FRDGBuilder& GraphBuilder, TRDGUniformBufferRef<FSceneTextureUniformParameters> SceneTextureUniformBuffer);
 
