Added stencil dummy SRV to system textures.

#rnx

diff --git a/Engine/Source/Runtime/Renderer/Private/SystemTextures.cpp b/Engine/Source/Runtime/Renderer/Private/SystemTextures.cpp
--- a/Engine/Source/Runtime/Renderer/Private/SystemTextures.cpp
+++ b/Engine/Source/Runtime/Renderer/Private/SystemTextures.cpp
@@ -35,6 +35,7 @@ const FRDGSystemTextures& FRDGSystemTextures::Create(FRDGBuilder& GraphBuilder)
 	SystemTextures.DefaultNormal8Bit = Register(GSystemTextures.DefaultNormal8Bit);
 	SystemTextures.MidGrey = Register(GSystemTextures.MidGreyDummy);
 	SystemTextures.VolumetricBlack = Register(GSystemTextures.VolumetricBlackDummy);
+	SystemTextures.StencilDummySRV = GraphBuilder.CreateSRV(FRDGTextureSRVDesc::CreateWithPixelFormat(SystemTextures.DepthDummy, PF_X24_G8));
 	return SystemTextures;
 }
 
diff --git a/Engine/Source/Runtime/Renderer/Private/SystemTextures.h b/Engine/Source/Runtime/Renderer/Private/SystemTextures.h
--- a/Engine/Source/Runtime/Renderer/Private/SystemTextures.h
+++ b/Engine/Source/Runtime/Renderer/Private/SystemTextures.h
@@ -27,6 +27,8 @@ struct FRDGSystemTextures
 	FRDGTextureRef DefaultNormal8Bit{};
 	FRDGTextureRef MidGrey{};
 	FRDGTextureRef VolumetricBlack{};
+
+	FRDGTextureSRVRef StencilDummySRV{};
 };
 
 /**
