Removed unused scene render target singleton accesses.

#rnx

diff --git a/Engine/Source/Runtime/Renderer/Private/DistanceFieldScreenGridLighting.cpp b/Engine/Source/Runtime/Renderer/Private/DistanceFieldScreenGridLighting.cpp
--- a/Engine/Source/Runtime/Renderer/Private/DistanceFieldScreenGridLighting.cpp
+++ b/Engine/Source/Runtime/Renderer/Private/DistanceFieldScreenGridLighting.cpp
@@ -531,8 +531,6 @@ void FDeferredShadingSceneRenderer::RenderDistanceFieldAOScreenGrid(
 		{
 			ScreenGridResources->AcquireTransientResource();
 
-			FSceneRenderTargets& SceneContext = FSceneRenderTargets::Get();
-
 			float ConeVisibilityClearValue = 1.0f;
 			RHICmdList.Transition(FRHITransitionInfo(ScreenGridResources->ScreenGridConeVisibility.UAV, ERHIAccess::Unknown, ERHIAccess::UAVCompute));
 			RHICmdList.ClearUAVUint(ScreenGridResources->ScreenGridConeVisibility.UAV, FUintVector4(*(uint32*)&ConeVisibilityClearValue, 0, 0, 0)); // @todo - ScreenGridConeVisibility should probably be R32_FLOAT format.
diff --git a/Engine/Source/Runtime/Renderer/Private/MobileDistortionPass.cpp b/Engine/Source/Runtime/Renderer/Private/MobileDistortionPass.cpp
--- a/Engine/Source/Runtime/Renderer/Private/MobileDistortionPass.cpp
+++ b/Engine/Source/Runtime/Renderer/Private/MobileDistortionPass.cpp
@@ -40,8 +40,6 @@ TRDGUniformBufferRef<FMobileDistortionPassUniformParameters> CreateMobileDistort
 {
 	auto* Parameters = GraphBuilder.AllocParameters<FMobileDistortionPassUniformParameters>();
 
-	FSceneRenderTargets& SceneContext = FSceneRenderTargets::Get();
-
 	EMobileSceneTextureSetupMode SetupMode = EMobileSceneTextureSetupMode::SceneColor;
 	if (View.bCustomDepthStencilValid)
 	{
diff --git a/Engine/Source/Runtime/Renderer/Private/ShadowRendering.h b/Engine/Source/Runtime/Renderer/Private/ShadowRendering.h
--- a/Engine/Source/Runtime/Renderer/Private/ShadowRendering.h
+++ b/Engine/Source/Runtime/Renderer/Private/ShadowRendering.h
@@ -1446,7 +1446,6 @@ public:
 			SetTextureParameter(RHICmdList, ShaderRHI, HairCategorizationTexture, HairVisibilityData->CategorizationTexture->GetPooledRenderTarget()->GetRenderTargetItem().ShaderResourceTexture);
 		}
 
-		FSceneRenderTargets& SceneContext = FSceneRenderTargets::Get();
 		{
 			const IPooledRenderTarget* PooledRT = GetSubsufaceProfileTexture_RT((FRHICommandListImmediate&)RHICmdList);
 
