Scale clipmap resolution by main view resolution.

diff --git a/Engine/Source/Runtime/Renderer/Private/ShadowSetup.cpp b/Engine/Source/Runtime/Renderer/Private/ShadowSetup.cpp
--- a/Engine/Source/Runtime/Renderer/Private/ShadowSetup.cpp
+++ b/Engine/Source/Runtime/Renderer/Private/ShadowSetup.cpp
@@ -4255,6 +4255,7 @@ void FSceneRenderer::AddViewDependentWholeSceneShadowsForView(
 					LightSceneInfo,
 					WorldToLight,
 					View.ViewMatrices,
+					View.ViewRect.Size(),
 					MaxShadowCascadeDistance
 				));
 
diff --git a/Engine/Source/Runtime/Renderer/Private/VirtualShadowMaps/VirtualShadowMapClipmap.cpp b/Engine/Source/Runtime/Renderer/Private/VirtualShadowMaps/VirtualShadowMapClipmap.cpp
--- a/Engine/Source/Runtime/Renderer/Private/VirtualShadowMaps/VirtualShadowMapClipmap.cpp
+++ b/Engine/Source/Runtime/Renderer/Private/VirtualShadowMaps/VirtualShadowMapClipmap.cpp
@@ -13,7 +13,7 @@ VirtualShadowMapClipmap.cpp
 
 static TAutoConsoleVariable<float> CVarVirtualShadowMapClipmapResolutionLodBias(
 	TEXT( "r.Shadow.v.Clipmap.ResolutionLodBias" ),
-	2.0f,
+	-0.5f,
 	TEXT( "" ),
 	ECVF_RenderThreadSafe
 );
@@ -27,8 +27,8 @@ static TAutoConsoleVariable<float> CVarVirtualShadowMapClipmapFirstLevel(
 // "Virtual" clipmap level to clipmap radius
 static float GetLevelRadius(int32 Level)
 {
-	// Virtual clipmap indices can be negative (although not commonly), so need a real Pow()
-	// Remember clipmap level rounds *down*, so radius needs to cover out to 2^(Level+1), where it flips
+	// NOTE: Virtual clipmap indices can be negative (although not commonly)
+	// Clipmap level rounds *down*, so radius needs to cover out to 2^(Level+1), where it flips
 	return FMath::Pow(2.0f, static_cast<float>(Level + 1));
 }
 
@@ -38,6 +38,7 @@ FVirtualShadowMapClipmap::FVirtualShadowMapClipmap(
 	const FLightSceneInfo& InLightSceneInfo,
 	const FMatrix& WorldToLightRotationMatrix,
 	const FViewMatrices& CameraViewMatrices,
+	FIntPoint CameraViewRectSize,
 	float MaxRadius)
 	: LightSceneInfo(InLightSceneInfo)
 {
@@ -55,9 +56,10 @@ FVirtualShadowMapClipmap::FVirtualShadowMapClipmap(
 	
 	// NOTE: Rotational (roll) invariance of the directional light depends on square pixels so we just base everything on the camera X scales/resolution
 	float LodScale = 1.0f / CameraViewMatrices.GetProjectionScale().X;
+	LodScale *= float(FVirtualShadowMap::VirtualMaxResolutionXY) / float(CameraViewRectSize.X);
 	
-	// Clamp negative resolution biases as they would exceed the maximum resolution/ranges allocated
 	ResolutionLodBias = CVarVirtualShadowMapClipmapResolutionLodBias.GetValueOnRenderThread() + FMath::Log2(LodScale);
+	// Clamp negative absolute resolution biases as they would exceed the maximum resolution/ranges allocated
 	ResolutionLodBias = FMath::Max(0.0f, ResolutionLodBias);
 
 	// For now we adjust resolution by just biasing the page we look up in. This is wasteful in terms of page table vs.
diff --git a/Engine/Source/Runtime/Renderer/Private/VirtualShadowMaps/VirtualShadowMapClipmap.h b/Engine/Source/Runtime/Renderer/Private/VirtualShadowMaps/VirtualShadowMapClipmap.h
--- a/Engine/Source/Runtime/Renderer/Private/VirtualShadowMaps/VirtualShadowMapClipmap.h
+++ b/Engine/Source/Runtime/Renderer/Private/VirtualShadowMaps/VirtualShadowMapClipmap.h
@@ -22,6 +22,7 @@ public:
 		const FLightSceneInfo& InLightSceneInfo,
 		const FMatrix& WorldToLightRotationMatrix,
 		const FViewMatrices& CameraViewMatrices,
+		FIntPoint CameraViewRectSize,
 		float MaxRadius		// Maximum radius the clipmap must cover from the center point; used to compute level count
 	);
 
