Respect r.Shadow.NaniteUseHZB in nanite atlas passes.

diff --git a/Engine/Source/Runtime/Renderer/Private/ShadowDepthRendering.cpp b/Engine/Source/Runtime/Renderer/Private/ShadowDepthRendering.cpp
--- a/Engine/Source/Runtime/Renderer/Private/ShadowDepthRendering.cpp
+++ b/Engine/Source/Runtime/Renderer/Private/ShadowDepthRendering.cpp
@@ -1324,6 +1324,7 @@ static void RenderShadowDepthAtlasNanite(
 {
 	const FIntPoint AtlasSize = ShadowMapAtlas.RenderTargets.DepthTarget->GetDesc().Extent;
 
+	const bool bUseHZB = (CVarNaniteShadowsUseHZB.GetValueOnRenderThread() != 0);
 	TArray<TRefCountPtr<IPooledRenderTarget>>&	PrevAtlasHZBs = bIsCompletePass ? Scene.PrevAtlasCompleteHZBs : Scene.PrevAtlasHZBs;
 
 	bool bWantsNearClip = false;
@@ -1392,7 +1393,8 @@ static void RenderShadowDepthAtlasNanite(
 		// It might be worthwhile going through the virtual SM rendering path even for "dense" cases even just for proper handling of all the details.
 		FIntRect FullAtlasViewRect(FIntPoint(0, 0), AtlasSize);
 		const bool bUpdateStreaming = CVarNaniteShadowsUpdateStreaming.GetValueOnRenderThread() != 0;
-		Nanite::FCullingContext CullingContext = Nanite::InitCullingContext(GraphBuilder, Scene, PrevAtlasHZBs[AtlasIndex], FullAtlasViewRect, true, bUpdateStreaming, bSupportsMultiplePasses, false, bPrimaryContext);
+		TRefCountPtr<IPooledRenderTarget> PrevAtlasHZB = bUseHZB ? PrevAtlasHZBs[AtlasIndex] : nullptr;
+		Nanite::FCullingContext CullingContext = Nanite::InitCullingContext(GraphBuilder, Scene, PrevAtlasHZB, FullAtlasViewRect, true, bUpdateStreaming, bSupportsMultiplePasses, false, bPrimaryContext);
 		Nanite::FRasterContext RasterContext = Nanite::InitRasterContext(GraphBuilder, AtlasSize, Nanite::EOutputBufferMode::DepthOnly);
 
 		bool bExtractStats = false;		
@@ -1436,6 +1438,7 @@ static void RenderShadowDepthAtlasNanite(
 			);
 		}
 
+		if (bUseHZB)
 		{
 			FRDGTextureRef FurthestHZBTexture;
 			BuildHZB(
@@ -1447,6 +1450,10 @@ static void RenderShadowDepthAtlasNanite(
 				/* OutFurthestHZBTexture = */ &FurthestHZBTexture);
 			ConvertToExternalTexture(GraphBuilder, FurthestHZBTexture, PrevAtlasHZBs[AtlasIndex]);
 		}
+		else
+		{
+			PrevAtlasHZBs[AtlasIndex] = nullptr;
+		}
 
 		FRDGTextureRef ShadowMap = GraphBuilder.RegisterExternalTexture(ShadowMapAtlas.RenderTargets.DepthTarget);
 
@@ -1712,11 +1719,10 @@ void FSceneRenderer::RenderShadowDepthMaps(FRDGBuilder& GraphBuilder)
 	const bool bHasVSMShadows = SortedShadowsForShadowDepthPass.VirtualShadowMapShadows.Num() > 0;
 	const bool bHasVSMClipMaps = SortedShadowsForShadowDepthPass.VirtualShadowMapClipmaps.Num() > 0;
 	const bool bNaniteEnabled = UseNanite(ShaderPlatform) && ViewFamily.EngineShowFlags.NaniteMeshes;
+	const bool bUseHZB = (CVarNaniteShadowsUseHZB.GetValueOnRenderThread() != 0);
 
 	if (bNaniteEnabled && (bHasVSMShadows || bHasVSMClipMaps))
 	{
-		bool bUseHZB = (CVarNaniteShadowsUseHZB.GetValueOnRenderThread() != 0);
-
 		if (bUseHZB)
 		{
 			VirtualShadowMapArray.HZBPhysical	= Scene->VirtualShadowMapArrayCacheManager->HZBPhysical;
@@ -1932,7 +1938,7 @@ void FSceneRenderer::RenderShadowDepthMaps(FRDGBuilder& GraphBuilder)
 		GraphBuilder.QueueBufferExtraction(VirtualShadowMapArray.PageTableRDG, &Scene->VirtualShadowMapArrayCacheManager->HZBPageTable);
 	}
 
-	// Render non-VSM shadows. Must be after VSM so we can use TopMip optimization.
+	// Render non-VSM shadows
 	FSceneRenderer::RenderShadowDepthMapAtlases(GraphBuilder);
 
 	const bool bUseGeometryShader = !GRHISupportsArrayIndexFromAnyShader;
@@ -2011,7 +2017,9 @@ void FSceneRenderer::RenderShadowDepthMaps(FRDGBuilder& GraphBuilder)
 					FPersistentShadowState* PrevShadowState = LightSceneInfo.PrevPersistentShadows.Find(ShadowKey);
 
 					const bool bPrimaryContext = false;
-					Nanite::FCullingContext CullingContext = Nanite::InitCullingContext(GraphBuilder, *Scene, PrevShadowState ? PrevShadowState->HZB : nullptr, ShadowViewRect, true, bUpdateStreaming, false, false, bPrimaryContext);
+					
+					TRefCountPtr<IPooledRenderTarget> PrevHZB = (PrevShadowState && bUseHZB) ? PrevShadowState->HZB : nullptr;
+					Nanite::FCullingContext CullingContext = Nanite::InitCullingContext(GraphBuilder, *Scene, PrevHZB, ShadowViewRect, true, bUpdateStreaming, false, false, bPrimaryContext);
 					Nanite::FRasterContext RasterContext = Nanite::InitRasterContext(GraphBuilder, TargetSize, Nanite::EOutputBufferMode::DepthOnly);
 
 					// Setup packed view
@@ -2063,9 +2071,10 @@ void FSceneRenderer::RenderShadowDepthMaps(FRDGBuilder& GraphBuilder)
 						ShadowViewRect,
 						CubemapFaceIndex,
 						bUseGeometryShader);
-
+										
+					TRefCountPtr<IPooledRenderTarget> HZB;
+					if (bUseHZB)
 					{
-						TRefCountPtr<IPooledRenderTarget> HZB;
 						FRDGTextureRef FurthestHZBTexture;
 						BuildHZB(
 							GraphBuilder,
@@ -2075,10 +2084,9 @@ void FSceneRenderer::RenderShadowDepthMaps(FRDGBuilder& GraphBuilder)
 							/* OutClosestHZBTexture = */ nullptr,
 							/* OutFurthestHZBTexture = */ &FurthestHZBTexture);
 
-
 						ConvertToExternalTexture(GraphBuilder, FurthestHZBTexture, HZB);
-						UpdateCurrentFrameHZB(LightSceneInfo, ShadowKey, ProjectedShadowInfo, HZB);
 					}
+					UpdateCurrentFrameHZB(LightSceneInfo, ShadowKey, ProjectedShadowInfo, HZB);
 				}
 			}
 		}
