Fix scale application order for outer shadow projection

#rnx

diff --git a/Engine/Source/Runtime/Renderer/Private/ShadowSetup.cpp b/Engine/Source/Runtime/Renderer/Private/ShadowSetup.cpp
--- a/Engine/Source/Runtime/Renderer/Private/ShadowSetup.cpp
+++ b/Engine/Source/Runtime/Renderer/Private/ShadowSetup.cpp
@@ -631,7 +631,7 @@ bool FProjectedShadowInfo::SetupPerObjectProjection(
 		FMatrix BorderScaleMatrix = FScaleMatrix(FVector(BorderScale, 1.0f));
 
 		const FMatrix ProjectionInner = FShadowProjectionMatrix( MinProjected, MaxProjected, Initializer.WAxis );
-		const FMatrix ProjectionOuter = BorderScaleMatrix * ProjectionInner;
+		const FMatrix ProjectionOuter = ProjectionInner * BorderScaleMatrix;
 
 		FMatrix ReceiverInnerMatrix;
 		float MaxSubjectDepth;
