Make sure the view uniform buffer is bound during RenderDitheredLODFadingOutMask pass

diff --git a/Engine/Source/Runtime/Renderer/Private/RayTracing/RayTracingShadows.cpp b/Engine/Source/Runtime/Renderer/Private/RayTracing/RayTracingShadows.cpp
--- a/Engine/Source/Runtime/Renderer/Private/RayTracing/RayTracingShadows.cpp
+++ b/Engine/Source/Runtime/Renderer/Private/RayTracing/RayTracingShadows.cpp
@@ -380,9 +380,15 @@ void FDeferredShadingSceneRenderer::RenderRayTracingShadows(
 }
 #endif
 
+BEGIN_SHADER_PARAMETER_STRUCT(FDitheredLODFadingOutMaskParameters, )
+	RENDER_TARGET_BINDING_SLOTS()
+	SHADER_PARAMETER_STRUCT_REF(FViewUniformShaderParameters, View)
+END_SHADER_PARAMETER_STRUCT()
+
 void FDeferredShadingSceneRenderer::RenderDitheredLODFadingOutMask(FRDGBuilder& GraphBuilder, const FViewInfo& View, FRDGTextureRef SceneDepthTexture)
 {
-	auto* PassParameters = GraphBuilder.AllocParameters<FRenderTargetParameters>();
+	auto* PassParameters = GraphBuilder.AllocParameters<FDitheredLODFadingOutMaskParameters>();
+	PassParameters->View = View.ViewUniformBuffer;
 	PassParameters->RenderTargets.DepthStencil = FDepthStencilBinding(SceneDepthTexture, ERenderTargetLoadAction::ELoad, ERenderTargetLoadAction::ELoad, FExclusiveDepthStencil::DepthWrite_StencilWrite);
 
 	GraphBuilder.AddPass(
