Fixed failed type coersion on latest VS compiler.

diff --git a/Engine/Source/Runtime/Renderer/Public/SceneRenderTargetParameters.h b/Engine/Source/Runtime/Renderer/Public/SceneRenderTargetParameters.h
--- a/Engine/Source/Runtime/Renderer/Public/SceneRenderTargetParameters.h
+++ b/Engine/Source/Runtime/Renderer/Public/SceneRenderTargetParameters.h
@@ -145,7 +145,7 @@ class RENDERER_API FSceneTextureExtracts : public FRenderResource
 public:
 	FRHIUniformBuffer* GetUniformBuffer() const
 	{
-		return UniformBuffer.IsValid() ? UniformBuffer : MobileUniformBuffer;
+		return UniformBuffer.IsValid() ? UniformBuffer.GetReference() : MobileUniformBuffer.GetReference();
 	}
 
 	TUniformBufferRef<FSceneTextureUniformParameters> GetUniformBufferRef() const
