Fix translucent shadows by adding missing view uniform buffer binding to the pass parameters
- Also removed now-redundant and broken ShadowViewMatrix

#rnx

diff --git a/Engine/Source/Runtime/Renderer/Private/ShadowDepthRendering.cpp b/Engine/Source/Runtime/Renderer/Private/ShadowDepthRendering.cpp
--- a/Engine/Source/Runtime/Renderer/Private/ShadowDepthRendering.cpp
+++ b/Engine/Source/Runtime/Renderer/Private/ShadowDepthRendering.cpp
@@ -143,7 +143,7 @@ void SetupShadowDepthPassUniformBuffer(
 	SetupSceneTextureUniformParameters(GraphBuilder, View.FeatureLevel, ESceneTextureSetupMode::None, ShadowDepthPassParameters.SceneTextures);
 
 	ShadowDepthPassParameters.ProjectionMatrix = FTranslationMatrix(ShadowInfo->PreShadowTranslation - View.ViewMatrices.GetPreViewTranslation()) * ShadowInfo->TranslatedWorldToClipOuterMatrix;	
-	ShadowDepthPassParameters.ViewMatrix = ShadowInfo->ShadowViewMatrix;
+	ShadowDepthPassParameters.ViewMatrix = ShadowInfo->TranslatedWorldToView;
 
 	ShadowDepthPassParameters.ShadowParams = FVector4(ShadowInfo->GetShaderDepthBias(), ShadowInfo->GetShaderSlopeDepthBias(), ShadowInfo->GetShaderMaxSlopeDepthBias(), ShadowInfo->bOnePassPointLightShadow ? 1 : ShadowInfo->InvMaxSubjectDepth);
 	ShadowDepthPassParameters.bClampToNearPlane = ShadowInfo->ShouldClampToNearPlane() ? 1.0f : 0.0f;
@@ -172,7 +172,7 @@ void SetupShadowDepthPassUniformBuffer(
 	SetupMobileSceneTextureUniformParameters(GraphBuilder, EMobileSceneTextureSetupMode::None, ShadowDepthPassParameters.SceneTextures);
 
 	ShadowDepthPassParameters.ProjectionMatrix = FTranslationMatrix(ShadowInfo->PreShadowTranslation - View.ViewMatrices.GetPreViewTranslation()) * ShadowInfo->TranslatedWorldToClipOuterMatrix;
-	ShadowDepthPassParameters.ViewMatrix = ShadowInfo->ShadowViewMatrix;
+	ShadowDepthPassParameters.ViewMatrix = ShadowInfo->TranslatedWorldToView;
 
 	ShadowDepthPassParameters.ShadowParams = FVector4(ShadowInfo->GetShaderDepthBias(), ShadowInfo->GetShaderSlopeDepthBias(), ShadowInfo->GetShaderMaxSlopeDepthBias(), ShadowInfo->InvMaxSubjectDepth);
 	ShadowDepthPassParameters.bClampToNearPlane = ShadowInfo->ShouldClampToNearPlane() ? 1.0f : 0.0f;
diff --git a/Engine/Source/Runtime/Renderer/Private/ShadowRendering.h b/Engine/Source/Runtime/Renderer/Private/ShadowRendering.h
--- a/Engine/Source/Runtime/Renderer/Private/ShadowRendering.h
+++ b/Engine/Source/Runtime/Renderer/Private/ShadowRendering.h
@@ -249,9 +249,9 @@ public:
 	/** A translation that is applied to world-space before transforming by one of the shadow matrices. */
 	FVector PreShadowTranslation;
 
-	/** The effective view matrix of the shadow, used as an override to the main view's view matrix when rendering the shadow depth pass. */
-	FMatrix ShadowViewMatrix;
-
+	/** 
+	 * The view matrix of the shadow, ALSO used as an override to the main view's view matrix when rendering the shadow depth pass. 
+	 */
 	FMatrix TranslatedWorldToView;
 
 	/** View space to clip space. Excluding border area. */
diff --git a/Engine/Source/Runtime/Renderer/Private/ShadowSetup.cpp b/Engine/Source/Runtime/Renderer/Private/ShadowSetup.cpp
--- a/Engine/Source/Runtime/Renderer/Private/ShadowSetup.cpp
+++ b/Engine/Source/Runtime/Renderer/Private/ShadowSetup.cpp
@@ -789,15 +789,6 @@ bool FProjectedShadowInfo::SetupPerObjectProjection(
 
 		MinPreSubjectZ = Initializer.MinLightW;
 
-		// Store the view matrix
-		// Reorder the vectors to match the main view, since ShadowViewMatrix will be used to override the main view's view matrix during shadow depth rendering
-		ShadowViewMatrix = Initializer.WorldToLight *
-			FMatrix(
-			FPlane(0, 0, 1, 0),
-			FPlane(1, 0, 0, 0),
-			FPlane(0, 1, 0, 0),
-			FPlane(0, 0, 0, 1));
-
 		GetViewFrustumBounds(CasterOuterFrustum, TranslatedWorldToClipOuterMatrix, true);
 		GetViewFrustumBounds(ReceiverInnerFrustum, ReceiverInnerMatrix, true);
 		
@@ -928,15 +919,6 @@ void FProjectedShadowInfo::SetupWholeSceneProjection(
 
 	InvMaxSubjectDepth = 1.0f / MaxSubjectDepth;
 
-	// Store the view matrix
-	// Reorder the vectors to match the main view, since ShadowViewMatrix will be used to override the main view's view matrix during shadow depth rendering
-	ShadowViewMatrix = Initializer.WorldToLight * 
-		FMatrix(
-		FPlane(0,	0,	1,	0),
-		FPlane(1,	0,	0,	0),
-		FPlane(0,	1,	0,	0),
-		FPlane(0,	0,	0,	1));
-
 	// Any meshes between the light and the subject can cast shadows, also any meshes inside the subject region
 	const FMatrix CasterOuterMatrix = WorldToViewScaledInner * BorderScaleMatrix * FShadowProjectionMatrix(Initializer.MinLightW, MaxSubjectZ, Initializer.WAxis);
 	const FMatrix ReceiverInnerMatrix = WorldToViewScaledInner * FShadowProjectionMatrix(MinSubjectZ, ClampedMaxLightW, Initializer.WAxis);
diff --git a/Engine/Source/Runtime/Renderer/Private/TranslucentLighting.cpp b/Engine/Source/Runtime/Renderer/Private/TranslucentLighting.cpp
--- a/Engine/Source/Runtime/Renderer/Private/TranslucentLighting.cpp
+++ b/Engine/Source/Runtime/Renderer/Private/TranslucentLighting.cpp
@@ -463,6 +463,7 @@ void FTranslucencyDepthPassMeshProcessor::AddMeshBatch(const FMeshBatch& RESTRIC
 }
 
 BEGIN_SHADER_PARAMETER_STRUCT(FTranslucencyDepthPassParameters, )
+	SHADER_PARAMETER_STRUCT_REF(FViewUniformShaderParameters, View)
 	SHADER_PARAMETER_RDG_UNIFORM_BUFFER(FTranslucencyDepthPassUniformParameters, PassUniformBuffer)
 	RENDER_TARGET_BINDING_SLOTS()
 END_SHADER_PARAMETER_STRUCT()
@@ -482,6 +483,7 @@ void FProjectedShadowInfo::RenderTranslucencyDepths(FRDGBuilder& GraphBuilder, F
 	auto* PassParameters = GraphBuilder.AllocParameters<FTranslucencyDepthPassParameters>();
 	PassParameters->PassUniformBuffer = PassUniformBuffer;
 	PassParameters->RenderTargets = InRenderTargets;
+	PassParameters->View = ShadowDepthView->ViewUniformBuffer;
 
 	FString EventName;
 #if WANTS_DRAW_MESH_EVENTS
