Fixed: A missing View UB in mobile distortion rendering

diff --git a/Engine/Source/Runtime/Renderer/Private/MobileDistortionPass.cpp b/Engine/Source/Runtime/Renderer/Private/MobileDistortionPass.cpp
--- a/Engine/Source/Runtime/Renderer/Private/MobileDistortionPass.cpp
+++ b/Engine/Source/Runtime/Renderer/Private/MobileDistortionPass.cpp
@@ -32,6 +32,7 @@ bool IsMobileDistortionActive(const FViewInfo& View)
 }
 
 BEGIN_SHADER_PARAMETER_STRUCT(FMobileDistortionPassParameters, )
+	SHADER_PARAMETER_STRUCT_INCLUDE(FViewShaderParameters, View)
 	SHADER_PARAMETER_RDG_UNIFORM_BUFFER(FMobileDistortionPassUniformParameters, Pass)
 	RENDER_TARGET_BINDING_SLOTS()
 END_SHADER_PARAMETER_STRUCT()
@@ -60,6 +61,7 @@ FMobileDistortionAccumulateOutputs AddMobileDistortionAccumulatePass(FRDGBuilder
 	FScreenPassRenderTarget DistortionAccumulateOutput = FScreenPassRenderTarget(GraphBuilder.CreateTexture(DistortionAccumulateDesc, TEXT("DistortionAccumulatePass")), Inputs.SceneColor.ViewRect, ERenderTargetLoadAction::EClear);
 
 	FMobileDistortionPassParameters* PassParameters = GraphBuilder.AllocParameters<FMobileDistortionPassParameters>();
+	PassParameters->View = View.GetShaderParameters();
 	PassParameters->Pass = CreateMobileDistortionPassUniformBuffer(GraphBuilder, View);
 	PassParameters->RenderTargets[0] = DistortionAccumulateOutput.GetRenderTargetBinding();
 
