Disable rendering nanite meshes being marked to render to runtimve virtual texture.
Fixes check() firing when that happens.

diff --git a/Engine/Source/Runtime/Renderer/Private/PrimitiveSceneInfo.cpp b/Engine/Source/Runtime/Renderer/Private/PrimitiveSceneInfo.cpp
--- a/Engine/Source/Runtime/Renderer/Private/PrimitiveSceneInfo.cpp
+++ b/Engine/Source/Runtime/Renderer/Private/PrimitiveSceneInfo.cpp
@@ -1123,7 +1123,11 @@ void FPrimitiveSceneInfo::UpdateRuntimeVirtualTextureFlags()
 
 	if (Proxy->WritesVirtualTexture())
 	{
-		if (StaticMeshes.Num() == 0)
+		if (Proxy->IsNaniteMesh())
+		{
+			UE_LOG(LogRenderer, Warning, TEXT("Rendering a nanite mesh to a runtime virtual texture isn't yet supported. Please disable this option on primitive component : %s"), *Proxy->GetOwnerName().ToString());
+		}
+		else if (StaticMeshes.Num() == 0)
 		{
 			UE_LOG(LogRenderer, Warning, TEXT("Rendering a primitive in a runtime virtual texture implies that there is a mesh to render. Please disable this option on primitive component : %s"), *Proxy->GetOwnerName().ToString());
 		}
