Fixed compilation error with unity build.

diff --git a/Engine/Source/Runtime/Renderer/Private/MobileShadingRenderer.cpp b/Engine/Source/Runtime/Renderer/Private/MobileShadingRenderer.cpp
--- a/Engine/Source/Runtime/Renderer/Private/MobileShadingRenderer.cpp
+++ b/Engine/Source/Runtime/Renderer/Private/MobileShadingRenderer.cpp
@@ -151,11 +151,9 @@ static bool UsesCustomDepthStencilLookup(const FViewInfo& View)
 		}
 	}
 
-	// Find out whether there are primitives will render in custom depth pass or just always render custom depth 
-	static const auto CVarCustomDepth = IConsoleManager::Get().FindTConsoleVariableDataInt(TEXT("r.CustomDepth"));
+	// Find out whether there are primitives will render in custom depth pass or just always render custom depth
+	bUsesCustomDepthStencil &= (View.bHasCustomDepthPrimitives || GetCustomDepthMode() == ECustomDepthMode::EnabledWithStencil);
 
-	bUsesCustomDepthStencil &= (View.bHasCustomDepthPrimitives || (CVarCustomDepth && CVarCustomDepth->GetValueOnRenderThread() > 1));
-	
 	return bUsesCustomDepthStencil;
 }
 
