Fixed missing view uniform buffer in translucent depth rendering.

diff --git a/Engine/Source/Runtime/Renderer/Private/TranslucentLighting.cpp b/Engine/Source/Runtime/Renderer/Private/TranslucentLighting.cpp
--- a/Engine/Source/Runtime/Renderer/Private/TranslucentLighting.cpp
+++ b/Engine/Source/Runtime/Renderer/Private/TranslucentLighting.cpp
@@ -481,6 +481,7 @@ void FProjectedShadowInfo::RenderTranslucencyDepths(FRDGBuilder& GraphBuilder, F
 	TRDGUniformBufferRef<FTranslucencyDepthPassUniformParameters> PassUniformBuffer = GraphBuilder.CreateUniformBuffer(TranslucencyDepthPassParameters);
 
 	auto* PassParameters = GraphBuilder.AllocParameters<FTranslucencyDepthPassParameters>();
+	PassParameters->View = ShadowDepthView->ViewUniformBuffer;
 	PassParameters->PassUniformBuffer = PassUniformBuffer;
 	PassParameters->RenderTargets = InRenderTargets;
 	PassParameters->View = ShadowDepthView->ViewUniformBuffer;
