Fix toggling nanite off/on and nanite viewflags off by centralizing the location from which we decide on whether we're using Virtual Shadow Maps or not (top of Render()).
Centralize VSM enable CVar.

diff --git a/Engine/Source/Runtime/Renderer/Private/DeferredShadingRenderer.cpp b/Engine/Source/Runtime/Renderer/Private/DeferredShadingRenderer.cpp
--- a/Engine/Source/Runtime/Renderer/Private/DeferredShadingRenderer.cpp
+++ b/Engine/Source/Runtime/Renderer/Private/DeferredShadingRenderer.cpp
@@ -1505,6 +1505,9 @@ void FDeferredShadingSceneRenderer::Render(FRDGBuilder& GraphBuilder)
 {
 	const bool bNaniteEnabled = UseNanite(ShaderPlatform) && ViewFamily.EngineShowFlags.NaniteMeshes;
 
+	// Important that this uses consistent logic for whether or not nanite is enabled, so pass in the flag from here
+	VirtualShadowMapArray.Initialize(bNaniteEnabled);
+
 	Scene->UpdateAllPrimitiveSceneInfos(GraphBuilder, true);
 
 	FGPUSceneScopeBeginEndHelper GPUSceneScopeBeginEndHelper(Scene->GPUScene, GPUSceneDynamicContext, *Scene);
@@ -1741,12 +1744,8 @@ void FDeferredShadingSceneRenderer::Render(FRDGBuilder& GraphBuilder)
 		bComputeLightGrid |= (
 			ShouldRenderVolumetricFog() ||
 			ViewFamily.ViewMode != VMI_Lit ||
-			GAllowLumenScene);
-
-		// TODO: Move decision logic to separate function for readability
-		static const auto EnableVirtualSMCVar = IConsoleManager::Get().FindTConsoleVariableDataInt(TEXT("r.Shadow.v.Enable"));
-
-		bComputeLightGrid = bComputeLightGrid || EnableVirtualSMCVar->GetValueOnRenderThread() != 0;
+			GAllowLumenScene ||
+			VirtualShadowMapArray.IsEnabled());
 	}
 
 	// force using occ queries for wireframe if rendering is parented or frozen in the first view
@@ -2034,9 +2033,7 @@ void FDeferredShadingSceneRenderer::Render(FRDGBuilder& GraphBuilder)
 	// Local helper function to perform virtual shadow map allocation, which can occur early, or late.
 	const auto AllocateVirtualShadowMaps = [&](bool bPostBasePass)
 	{
-		static const auto CVarEnableVirtualSM = IConsoleManager::Get().FindTConsoleVariableDataInt(TEXT("r.Shadow.v.Enable"));
-
-		if (CVarEnableVirtualSM->GetValueOnRenderThread() != 0)
+		if (VirtualShadowMapArray.IsEnabled())
 		{
 			ensureMsgf(AreLightsInLightGrid(), TEXT("Virtual shadow map setup requires local lights to be injected into the light grid (this may be caused by 'r.LightCulling.Quality=0')."));
 			// ensure(ShadowMapSetupDone)
@@ -2323,6 +2320,15 @@ void FDeferredShadingSceneRenderer::Render(FRDGBuilder& GraphBuilder)
 			ComputeVolumetricFog(GraphBuilder);
 		}
 
+		Nanite::GStreamingManager.SubmitFrameStreamingRequests(GraphBuilder);
+	}
+	else if (!bOcclusionBeforeBasePass)
+	{
+		ComputeVolumetricFog(GraphBuilder);
+	}
+
+	if (VirtualShadowMapArray.IsEnabled())
+	{
 		VirtualShadowMapArray.RenderDebugInfo(GraphBuilder, Scene->VirtualShadowMapArrayCacheManager);
 
 		if (Views.Num() > 0)
@@ -2330,17 +2336,11 @@ void FDeferredShadingSceneRenderer::Render(FRDGBuilder& GraphBuilder)
 			VirtualShadowMapArray.PrintStats(GraphBuilder, Views[0]);
 		}
 
-		Nanite::GStreamingManager.SubmitFrameStreamingRequests(GraphBuilder);
-
 		if (Scene->VirtualShadowMapArrayCacheManager)
 		{
 			Scene->VirtualShadowMapArrayCacheManager->ExtractFrameData(VirtualShadowMapArray, GraphBuilder);
 		}
 	}
-	else if (!bOcclusionBeforeBasePass)
-	{
-		ComputeVolumetricFog(GraphBuilder);
-	}
 
 	// If not all depth is written during the prepass, kick off async compute cloud after basepass
 	if (bShouldRenderVolumetricCloud && bAsyncComputeVolumetricCloud && DepthPass.EarlyZPassMode != DDM_AllOpaque)
diff --git a/Engine/Source/Runtime/Renderer/Private/LightGridInjection.cpp b/Engine/Source/Runtime/Renderer/Private/LightGridInjection.cpp
--- a/Engine/Source/Runtime/Renderer/Private/LightGridInjection.cpp
+++ b/Engine/Source/Runtime/Renderer/Private/LightGridInjection.cpp
@@ -788,12 +788,9 @@ void FDeferredShadingSceneRenderer::GatherLightsAndComputeLightGrid(FRDGBuilder&
 		bAnyViewUsesForwardLighting |= View.bTranslucentSurfaceLighting || ShouldRenderVolumetricFog() || View.bHasSingleLayerWaterMaterial;
 	}
 	
-	// TODO: Move decision logic to separate function for readability
-	static const auto EnableVirtualSMCVar = IConsoleManager::Get().FindTConsoleVariableDataInt(TEXT("r.Shadow.v.Enable"));
-
 	const bool bCullLightsToGrid = GLightCullingQuality
 		&& (ViewFamily.EngineShowFlags.DirectLighting
-			&& (IsForwardShadingEnabled(ShaderPlatform) || bAnyViewUsesForwardLighting || IsRayTracingEnabled() || ShouldUseClusteredDeferredShading() || GAllowLumenScene || EnableVirtualSMCVar->GetValueOnRenderThread() != 0));
+			&& (IsForwardShadingEnabled(ShaderPlatform) || bAnyViewUsesForwardLighting || IsRayTracingEnabled() || ShouldUseClusteredDeferredShading() || GAllowLumenScene || VirtualShadowMapArray.IsEnabled()));
 
 	// Store this flag if lights are injected in the grids, check with 'AreLightsInLightGrid()'
 	bAreLightsInLightGrid = bCullLightsToGrid;
diff --git a/Engine/Source/Runtime/Renderer/Private/ShadowSetup.cpp b/Engine/Source/Runtime/Renderer/Private/ShadowSetup.cpp
--- a/Engine/Source/Runtime/Renderer/Private/ShadowSetup.cpp
+++ b/Engine/Source/Runtime/Renderer/Private/ShadowSetup.cpp
@@ -3154,11 +3154,6 @@ void FSceneRenderer::CreateWholeSceneProjectedShadow(
 		return;
 	}
 
-	// Determine if we want a virtual shadow map for this light
-	// TODO: Base this off of a light/editor parameter; for now just all spot lights get virtual shadow maps when the CVar is enabled
-	static const auto EnableVirtualSMCVar = IConsoleManager::Get().FindTConsoleVariableDataInt(TEXT("r.Shadow.v.Enable"));
-	const bool bNaniteEnabled = UseNanite(ShaderPlatform);
-
 	// Try to create a whole-scene projected shadow initializer for the light.
 	TArray<FWholeSceneProjectedShadowInitializer, TInlineAllocator<6> > ProjectedShadowInitializers;
 	if (LightSceneInfo->Proxy->GetWholeSceneProjectedShadowInitializer(ViewFamily, ProjectedShadowInitializers))
@@ -3268,9 +3263,8 @@ void FSceneRenderer::CreateWholeSceneProjectedShadow(
 					SizeX = SizeY = GetCubeShadowDepthZResolution(FeatureLevel, GetCubeShadowDepthZIndex(FeatureLevel, MaxDesiredResolution));
 				}
 
-				const bool bNeedsVirtualShadowMap = bNaniteEnabled
-					&& (EnableVirtualSMCVar->GetValueOnRenderThread() != 0
-					&& LightSceneInfo->Proxy->GetLightType() == LightType_Spot)
+				const bool bNeedsVirtualShadowMap = VirtualShadowMapArray.IsEnabled()
+					&& (LightSceneInfo->Proxy->GetLightType() == LightType_Spot)
 					&& !ProjectedShadowInitializer.bRayTracedDistanceField;
 
 				int32 NumShadowMaps = 1;
@@ -4197,12 +4191,10 @@ void FSceneRenderer::AddViewDependentWholeSceneShadowsForView(
 {
 	SCOPE_CYCLE_COUNTER(STAT_AddViewDependentWholeSceneShadowsForView);
 
-	static const auto EnableVirtualSMCVar = IConsoleManager::Get().FindTConsoleVariableDataInt(TEXT("r.Shadow.v.Enable"));
-	const bool bNaniteEnabled = UseNanite(ShaderPlatform);
-	// Note: it is possible for a non-directional light to set up a proxy that requests a view-dependent SM (or several),
-	const bool bDirectionalLight = LightSceneInfo.Proxy->GetLightType() == LightType_Directional;
+	// Note: it is possible for a non-directional light to set up a proxy that requests a view-dependent SM (or several)
 	// Unmodified UE does not use this, so the virtual SM ignores this path for the time being (more likely to be removed).
-	const bool bNeedsVirtualShadowMap = bNaniteEnabled && EnableVirtualSMCVar->GetValueOnRenderThread() != 0 && bDirectionalLight;
+	const bool bDirectionalLight = LightSceneInfo.Proxy->GetLightType() == LightType_Directional;
+	const bool bNeedsVirtualShadowMap = VirtualShadowMapArray.IsEnabled() && bDirectionalLight;
 	const bool bNeedsCompleteShadowMap = bNeedsVirtualShadowMap && CVarCompleteShadowMapResolution.GetValueOnRenderThread() > 0;
 
 	// Allow each view to create a whole scene view dependent shadow
diff --git a/Engine/Source/Runtime/Renderer/Private/VirtualShadowMaps/VirtualShadowMapArray.cpp b/Engine/Source/Runtime/Renderer/Private/VirtualShadowMaps/VirtualShadowMapArray.cpp
--- a/Engine/Source/Runtime/Renderer/Private/VirtualShadowMaps/VirtualShadowMapArray.cpp
+++ b/Engine/Source/Runtime/Renderer/Private/VirtualShadowMaps/VirtualShadowMapArray.cpp
@@ -196,6 +196,11 @@ FVirtualShadowMapArray::~FVirtualShadowMapArray()
 	}
 }
 
+void FVirtualShadowMapArray::Initialize(bool bInEnabled)
+{
+	bEnabled = bInEnabled && CVarEnableVirtualShadowMaps.GetValueOnRenderThread() > 0;
+}
+
 void FVirtualShadowMapArray::SetShaderDefines(FShaderCompilerEnvironment& OutEnvironment)
 {
 	OutEnvironment.SetDefine(TEXT("VSM_PAGE_SIZE"), FVirtualShadowMap::PageSize);
@@ -369,8 +374,11 @@ IMPLEMENT_GLOBAL_SHADER(FClearPhysicalPagesCS, "/Engine/Private/VirtualShadowMap
 
 void FVirtualShadowMapArray::ClearPhysicalMemory(FRDGBuilder& GraphBuilder, FRDGTextureRef& PhysicalTexture, FVirtualShadowMapArrayCacheManager *VirtualShadowMapArrayCacheManager)
 {
+	check(IsEnabled());
 	if (ShadowMaps.Num() == 0)
+	{
 		return;
+	}
 
 	RDG_EVENT_SCOPE( GraphBuilder, "FVirtualShadowMapArray::ClearPhysicalMemory" );
 
@@ -419,7 +427,6 @@ void FVirtualShadowMapArray::ClearPhysicalMemory(FRDGBuilder& GraphBuilder, FRDG
 			);
 		}
 	}
-
 }
 
 class FMarkRenderedPhysicalPagesCS : public FVirtualPageManagementShader
@@ -438,6 +445,7 @@ IMPLEMENT_GLOBAL_SHADER(FMarkRenderedPhysicalPagesCS, "/Engine/Private/VirtualSh
 
 void FVirtualShadowMapArray::MarkPhysicalPagesRendered(FRDGBuilder& GraphBuilder, const TArray<uint32, SceneRenderingAllocator> &VirtualShadowMapFlags)
 {
+	check(IsEnabled());
 	if (VirtualShadowMapFlags.Num() == 0)
 	{
 		return;
@@ -530,6 +538,7 @@ void FVirtualShadowMapArray::BuildPageAllocations(
 	bool bPostBasePass,
 	FVirtualShadowMapArrayCacheManager* VirtualShadowMapArrayCacheManager)
 {
+	check(IsEnabled());
 	RDG_EVENT_SCOPE(GraphBuilder, "FVirtualShadowMapArray::GeneratePageFlagsFromLightGrid");
 
 	ensure(NaniteRasterResults.Num() == Views.Num());
@@ -814,6 +823,7 @@ IMPLEMENT_GLOBAL_SHADER(FDebugVisualizeVirtualSmCS, "/Engine/Private/VirtualShad
 
 void FVirtualShadowMapArray::RenderDebugInfo(FRDGBuilder& GraphBuilder, FVirtualShadowMapArrayCacheManager *VirtualShadowMapArrayCacheManager)
 {
+	check(IsEnabled());
 	int32 DebugMethod = CVarDebugVisualizeVirtualSms.GetValueOnRenderThread();
 	if (ShadowMaps.Num() && DebugMethod > 0)
 	{
@@ -897,6 +907,7 @@ IMPLEMENT_GLOBAL_SHADER(FVirtualSmPrintStatsCS, "/Engine/Private/VirtualShadowMa
 
 void FVirtualShadowMapArray::PrintStats(FRDGBuilder& GraphBuilder, const FViewInfo& View)
 {
+	check(IsEnabled());
 	LLM_SCOPE_BYTAG(Nanite);
 
 	// Print stats
diff --git a/Engine/Source/Runtime/Renderer/Private/VirtualShadowMaps/VirtualShadowMapArray.h b/Engine/Source/Runtime/Renderer/Private/VirtualShadowMaps/VirtualShadowMapArray.h
--- a/Engine/Source/Runtime/Renderer/Private/VirtualShadowMaps/VirtualShadowMapArray.h
+++ b/Engine/Source/Runtime/Renderer/Private/VirtualShadowMaps/VirtualShadowMapArray.h
@@ -132,11 +132,19 @@ public:
 	FVirtualShadowMapArray();
 	~FVirtualShadowMapArray();
 
+	void Initialize(bool bInEnabled);
+
+	// Returns true if virtual shadow maps are enabled
+	bool IsEnabled() const
+	{
+		return bEnabled;
+	}
+
 	FVirtualShadowMap *Allocate()
 	{
+		check(IsEnabled());
 		FVirtualShadowMap *SM = new(FMemStack::Get(), 1, 16) FVirtualShadowMap(ShadowMaps.Num());
 		ShadowMaps.Add(SM);
-
 		return SM;
 	}
 
@@ -174,14 +182,16 @@ public:
 	// 
 	void PrintStats(FRDGBuilder& GraphBuilder, const FViewInfo& View);
 
+	// Are virtual shadow maps enabled. We store this at the start of the frame to centralize the logic.
+	bool bEnabled = false;
+
 	TArray<FVirtualShadowMap*, SceneRenderingAllocator> ShadowMaps;
 
 	// Large physical texture of depth format, say 4096^2 or whatever we think is enough texels to go around
 	FRDGTextureRef PhysicalPagePoolRDG = nullptr;
 	// Buffer that serves as the page table for all virtual shadow maps
 	FRDGBufferRef PageTableRDG = nullptr;
-
-	
+		
 	// Buffer that stores flags (uints) marking each page that needs to be rendered and cache status, for all virtual shadow maps.
 	// Flag values defined in PageAccessCommon.ush: VSM_ALLOCATED_FLAG | VSM_INVALID_FLAG
 	FRDGBufferRef PageFlagsRDG = nullptr;
