Fixing PVS 7.7 Issues:

warning V1063: The modulo by 1 operation is meaningless. The result will always be zero.

diff --git a/Engine/Source/Runtime/Renderer/Private/GPUFastFourierTransform.cpp b/Engine/Source/Runtime/Renderer/Private/GPUFastFourierTransform.cpp
--- a/Engine/Source/Runtime/Renderer/Private/GPUFastFourierTransform.cpp
+++ b/Engine/Source/Runtime/Renderer/Private/GPUFastFourierTransform.cpp
@@ -1170,7 +1170,7 @@ namespace
 		const uint32 YThreadCount = FCopyWindowCS::YThreadCount();
 
 		// Number of thread groups
-		const uint32 XGroups = ( DstExtent.X / XThreadCount ) + ( (DstExtent.X % XThreadCount == 0) ? 0 : 1 ); //-V547
+		const uint32 XGroups = ( DstExtent.X / XThreadCount ) + ( (DstExtent.X % XThreadCount == 0) ? 0 : 1 ); //-V::547, 1063
 		const uint32 YGroups = ( DstExtent.Y / YThreadCount ) + ( (DstExtent.Y % YThreadCount == 0) ? 0 : 1 ); //-V547
 
 		const FGPUFFTShaderContext::ShaderMapType& ShaderMap = Context.GetShaderMap();
