Fixed unity build compilation error.

#rnx

diff --git a/Engine/Source/Runtime/Renderer/Private/PostProcess/PostProcessMaterial.cpp b/Engine/Source/Runtime/Renderer/Private/PostProcess/PostProcessMaterial.cpp
--- a/Engine/Source/Runtime/Renderer/Private/PostProcess/PostProcessMaterial.cpp
+++ b/Engine/Source/Runtime/Renderer/Private/PostProcess/PostProcessMaterial.cpp
@@ -52,16 +52,7 @@ bool IsPostProcessStencilTestAllowed()
 	return CVarPostProcessAllowStencilTest.GetValueOnRenderThread() != 0;
 }
 
-bool IsCustomDepthEnabled()
-{
-	static const IConsoleVariable* CVarCustomDepth = IConsoleManager::Get().FindConsoleVariable(TEXT("r.CustomDepth"));
-
-	check(CVarCustomDepth);
-
-	return CVarCustomDepth->GetInt() == 3;
-}
-
-enum class ECustomDepthPolicy : uint32
+enum class EMaterialCustomDepthPolicy : uint32
 {
 	// Custom depth is disabled.
 	Disabled,
@@ -70,7 +61,7 @@ enum class ECustomDepthPolicy : uint32
 	Enabled
 };
 
-ECustomDepthPolicy GetMaterialCustomDepthPolicy(const FMaterial* Material, ERHIFeatureLevel::Type FeatureLevel)
+EMaterialCustomDepthPolicy GetMaterialCustomDepthPolicy(const FMaterial* Material, ERHIFeatureLevel::Type FeatureLevel)
 {
 	check(Material);
 
@@ -78,9 +69,9 @@ ECustomDepthPolicy GetMaterialCustomDepthPolicy(const FMaterial* Material, ERHIF
 	if (Material->IsStencilTestEnabled() && IsPostProcessStencilTestAllowed())
 	{
 		// Custom stencil texture allocated and available.
-		if (IsCustomDepthEnabled())
+		if (GetCustomDepthMode() == ECustomDepthMode::EnabledWithStencil)
 		{
-			return ECustomDepthPolicy::Enabled;
+			return EMaterialCustomDepthPolicy::Enabled;
 		}
 		else
 		{
@@ -88,7 +79,7 @@ ECustomDepthPolicy GetMaterialCustomDepthPolicy(const FMaterial* Material, ERHIF
 		}
 	}
 
-	return ECustomDepthPolicy::Disabled;
+	return EMaterialCustomDepthPolicy::Disabled;
 }
 
 FRHIDepthStencilState* GetMaterialStencilState(const FMaterial* Material)
@@ -392,9 +383,9 @@ FScreenPassTexture AddPostProcessMaterialPass(
 	FRDGTextureRef DepthStencilTexture = nullptr;
 
 	// Allocate custom depth stencil texture(s) and depth stencil state.
-	const ECustomDepthPolicy CustomStencilPolicy = GetMaterialCustomDepthPolicy(Material, FeatureLevel);
+	const EMaterialCustomDepthPolicy CustomStencilPolicy = GetMaterialCustomDepthPolicy(Material, FeatureLevel);
 
-	if (CustomStencilPolicy == ECustomDepthPolicy::Enabled)
+	if (CustomStencilPolicy == EMaterialCustomDepthPolicy::Enabled)
 	{
 		check(Inputs.CustomDepthTexture);
 		DepthStencilTexture = Inputs.CustomDepthTexture;
