Rename for new FVirtualTextureFeedbackGPU class

diff --git a/Engine/Source/Runtime/Renderer/Private/SceneRendering.cpp b/Engine/Source/Runtime/Renderer/Private/SceneRendering.cpp
--- a/Engine/Source/Runtime/Renderer/Private/SceneRendering.cpp
+++ b/Engine/Source/Runtime/Renderer/Private/SceneRendering.cpp
@@ -71,7 +71,7 @@
 	Globals
 -----------------------------------------------------------------------------*/
 
-static TGlobalResource<FVirtualTextureFeedbackGPU> GVirtualTextureFeedbackGPU;
+static TGlobalResource<FVirtualTextureFeedbackBuffer> GVirtualTextureFeedbackBuffer;
 
 static TAutoConsoleVariable<int32> CVarCachedMeshDrawCommands(
 	TEXT("r.MeshDrawCommands.UseCachedCommands"),
@@ -1786,7 +1786,7 @@ void FViewInfo::SetupUniformBufferParameters(
 		ViewUniformShaderParameters.WaterData = GIdentityPrimitiveBuffer.PrimitiveSceneDataBufferSRV;
 	}
 
-	ViewUniformShaderParameters.VTFeedbackBuffer = GVirtualTextureFeedbackGPU.GetUAV();
+	ViewUniformShaderParameters.VTFeedbackBuffer = GVirtualTextureFeedbackBuffer.GetUAV();
 
 #if WITH_EDITOR
 	if (EditorVisualizeLevelInstanceBuffer.SRV)
@@ -4504,12 +4504,12 @@ void VirtualTextureFeedbackBegin(FRDGBuilder& GraphBuilder, TArrayView<const FVi
 
 	FVirtualTextureFeedbackBufferDesc Desc;
 	Desc.Init2D(SceneTextureExtent, ViewRects, GetVirtualTextureFeedbackScale());
-	GVirtualTextureFeedbackGPU.Begin(GraphBuilder, Desc);
+	GVirtualTextureFeedbackBuffer.Begin(GraphBuilder, Desc);
 }
 
 void VirtualTextureFeedbackEnd(FRDGBuilder& GraphBuilder)
 {
-	GVirtualTextureFeedbackGPU.End(GraphBuilder);
+	GVirtualTextureFeedbackBuffer.End(GraphBuilder);
 }
 
 FRDGTextureRef CreateHalfResolutionDepthCheckerboardMinMax(FRDGBuilder& GraphBuilder, TArrayView<const FViewInfo> Views, FRDGTextureRef SceneDepthTexture)
diff --git a/Engine/Source/Runtime/Renderer/Private/VT/VirtualTextureFeedbackBuffer.cpp b/Engine/Source/Runtime/Renderer/Private/VT/VirtualTextureFeedbackBuffer.cpp
--- a/Engine/Source/Runtime/Renderer/Private/VT/VirtualTextureFeedbackBuffer.cpp
+++ b/Engine/Source/Runtime/Renderer/Private/VT/VirtualTextureFeedbackBuffer.cpp
@@ -1,8 +1,9 @@
 // Copyright Epic Games, Inc. All Rights Reserved.
 
 #include "VT/VirtualTextureFeedbackBuffer.h"
-#include "VT/VirtualTextureFeedback.h"
+
 #include "RenderGraphUtils.h"
+#include "VT/VirtualTextureFeedback.h"
 
 int32 GVirtualTextureFeedbackFactor = 16;
 static FAutoConsoleVariableRef CVarVirtualTextureFeedbackFactor(
@@ -13,12 +14,23 @@ static FAutoConsoleVariableRef CVarVirtualTextureFeedbackFactor(
 	ECVF_RenderThreadSafe | ECVF_ReadOnly /*Read-only as shaders are compiled with this value*/
 );
 
-uint32 SampleVirtualTextureFeedbackSequence(uint32 FrameIndex)
+int32 GetVirtualTextureFeedbackScale()
+{
+	// Round to nearest power of two to ensure that shader maths is efficient and sampling sequence logic is simple.
+	return FMath::RoundUpToPowerOfTwo(FMath::Max(GVirtualTextureFeedbackFactor, 1));
+}
+
+FIntPoint GetVirtualTextureFeedbackBufferSize(FIntPoint InSceneTextureExtent)
+{
+	return FIntPoint::DivideAndRoundUp(InSceneTextureExtent, FMath::Max(GetVirtualTextureFeedbackScale(), 1));
+}
+
+uint32 SampleVirtualTextureFeedbackSequence(uint32 InFrameIndex)
 {
 	const uint32 TileSize = GetVirtualTextureFeedbackScale();
 	const uint32 TileSizeLog2 = FMath::CeilLogTwo(TileSize);
 	const uint32 SequenceSize = FMath::Square(TileSize);
-	const uint32 PixelIndex = FrameIndex % SequenceSize;
+	const uint32 PixelIndex = InFrameIndex % SequenceSize;
 	const uint32 PixelAddress = ReverseBits(PixelIndex) >> (32U - 2 * TileSizeLog2);
 	const uint32 X = FMath::ReverseMortonCode2(PixelAddress);
 	const uint32 Y = FMath::ReverseMortonCode2(PixelAddress >> 1);
@@ -26,17 +38,6 @@ uint32 SampleVirtualTextureFeedbackSequence(uint32 FrameIndex)
 	return PixelSequenceIndex;
 }
 
-FIntPoint GetVirtualTextureFeedbackBufferSize(FIntPoint SceneTextureExtent)
-{
-	return FIntPoint::DivideAndRoundUp(SceneTextureExtent, FMath::Max(GetVirtualTextureFeedbackScale(), 1));
-}
-
-int32 GetVirtualTextureFeedbackScale()
-{
-	// Round to nearest power of two to ensure that shader maths is efficient and sampling sequence logic is simple.
-	return FMath::RoundUpToPowerOfTwo(FMath::Max(GVirtualTextureFeedbackFactor, 1));
-}
-
 void FVirtualTextureFeedbackBufferDesc::Init(int32 InBufferSize)
 {
 	BufferSize = FIntPoint(InBufferSize, 1);
@@ -74,17 +75,7 @@ void FVirtualTextureFeedbackBufferDesc::Init2D(FIntPoint InUnscaledBufferSize, T
 	}
 }
 
-void SubmitVirtualTextureFeedbackBuffer(FRHICommandListImmediate& RHICmdList, FVertexBufferRHIRef const& Buffer, FVirtualTextureFeedbackBufferDesc const& Desc)
-{
-	GVirtualTextureFeedback.TransferGPUToCPU(RHICmdList, Buffer, Desc);
-}
-
-void SubmitVirtualTextureFeedbackBuffer(class FRDGBuilder& GraphBuilder, FRDGBuffer* Buffer, FVirtualTextureFeedbackBufferDesc const& Desc)
-{
-	GVirtualTextureFeedback.TransferGPUToCPU(GraphBuilder, Buffer, Desc);
-}
-
-void FVirtualTextureFeedbackGPU::Begin(FRDGBuilder& GraphBuilder, const FVirtualTextureFeedbackBufferDesc& InDesc)
+void FVirtualTextureFeedbackBuffer::Begin(FRDGBuilder& GraphBuilder, const FVirtualTextureFeedbackBufferDesc& InDesc)
 {
 	// NOTE: Transitions and allocations are handled manually right now, because the VT feedback UAV is used by
 	// the view uniform buffer, which is not an RDG uniform buffer. If it can be factored out into its own RDG
@@ -112,7 +103,7 @@ void FVirtualTextureFeedbackGPU::Begin(FRDGBuilder& GraphBuilder, const FVirtual
 	});
 }
 
-void FVirtualTextureFeedbackGPU::End(FRDGBuilder& GraphBuilder)
+void FVirtualTextureFeedbackBuffer::End(FRDGBuilder& GraphBuilder)
 {
 	AddPass(GraphBuilder, RDG_EVENT_NAME("VirtualTextureFeedbackCopy"), [this](FRHICommandListImmediate& RHICmdList)
 	{
@@ -122,8 +113,18 @@ void FVirtualTextureFeedbackGPU::End(FRDGBuilder& GraphBuilder)
 	});
 }
 
-void FVirtualTextureFeedbackGPU::ReleaseRHI()
+void FVirtualTextureFeedbackBuffer::ReleaseRHI()
 {
 	PooledBuffer = nullptr;
 	UAV = nullptr;
 }
+
+void SubmitVirtualTextureFeedbackBuffer(FRHICommandListImmediate& RHICmdList, FVertexBufferRHIRef const& InBuffer, FVirtualTextureFeedbackBufferDesc const& InDesc)
+{
+	GVirtualTextureFeedback.TransferGPUToCPU(RHICmdList, InBuffer, InDesc);
+}
+
+void SubmitVirtualTextureFeedbackBuffer(class FRDGBuilder& GraphBuilder, FRDGBuffer* InBuffer, FVirtualTextureFeedbackBufferDesc const& InDesc)
+{
+	GVirtualTextureFeedback.TransferGPUToCPU(GraphBuilder, InBuffer, InDesc);
+}
diff --git a/Engine/Source/Runtime/Renderer/Public/VT/VirtualTextureFeedbackBuffer.h b/Engine/Source/Runtime/Renderer/Public/VT/VirtualTextureFeedbackBuffer.h
--- a/Engine/Source/Runtime/Renderer/Public/VT/VirtualTextureFeedbackBuffer.h
+++ b/Engine/Source/Runtime/Renderer/Public/VT/VirtualTextureFeedbackBuffer.h
@@ -4,6 +4,15 @@
 
 #include "RenderGraphResources.h"
 
+/** Returns the scale used to divide the scene texture resolution to get the virtual texture feedback resolution. */
+extern int32 GetVirtualTextureFeedbackScale();
+
+/** Get size for a virtual texture feedback buffer. This is calculated internally by applying GetVirtualTextureFeedbackScale() to the SceneTextureExtent. */
+extern FIntPoint GetVirtualTextureFeedbackBufferSize(FIntPoint InSceneTextureExtent);
+
+/** Returns a value from the sampling sequence used to traverse each feedback tile. */
+extern uint32 SampleVirtualTextureFeedbackSequence(uint32 InFrameIndex);
+
 /** 
  * Description of how to interpret an RHIBuffer that is being fed to the virtual texture feedback system.
  * For example a buffer may be a simple flat buffer, or a 2D screen-space buffer with rectangles representing multiple player viewports.
@@ -30,17 +39,8 @@ struct FVirtualTextureFeedbackBufferDesc
 	int32 TotalReadSize = 0;
 };
 
-/** Returns a value from the sampling sequence used to traverse the feedback tile. */
-extern uint32 SampleVirtualTextureFeedbackSequence(uint32 FrameIndex);
-
-/** Get size of screen tiles used for virtual texture feedback. We evaluate one feedback item from each tile per frame. */
-extern FIntPoint GetVirtualTextureFeedbackBufferSize(FIntPoint SceneTextureExtent);
-
-/** Returns the scale used to divide the against the scene texture resolution to get the virtual texture feedback resolution. */
-extern int32 GetVirtualTextureFeedbackScale();
-
 /** GPU feedback buffer for tracking virtual texture requests from the GPU. */
-class FVirtualTextureFeedbackGPU : public FRenderResource
+class FVirtualTextureFeedbackBuffer : public FRenderResource
 {
 public:
 	// Allocates and prepares the feedback buffer for submission.
@@ -72,5 +72,5 @@ private:
  * Multiple buffers can be transferred per frame using this function.
  * The function can be called from the render thread only.
 */
-RENDERER_API void SubmitVirtualTextureFeedbackBuffer(class FRHICommandListImmediate& RHICmdList, FVertexBufferRHIRef const& Buffer, FVirtualTextureFeedbackBufferDesc const& Desc);
-RENDERER_API void SubmitVirtualTextureFeedbackBuffer(class FRDGBuilder& GraphBuilder, class FRDGBuffer* Buffer, FVirtualTextureFeedbackBufferDesc const& Desc);
\ No newline at end of file
+RENDERER_API void SubmitVirtualTextureFeedbackBuffer(class FRHICommandListImmediate& RHICmdList, FVertexBufferRHIRef const& InBuffer, FVirtualTextureFeedbackBufferDesc const& InDesc);
+RENDERER_API void SubmitVirtualTextureFeedbackBuffer(class FRDGBuilder& GraphBuilder, class FRDGBuffer* InBuffer, FVirtualTextureFeedbackBufferDesc const& InDesc);
\ No newline at end of file
