Fixed cubemap scene capture rendering.

#rnx

diff --git a/Engine/Source/Runtime/Renderer/Private/SceneCaptureRendering.cpp b/Engine/Source/Runtime/Renderer/Private/SceneCaptureRendering.cpp
--- a/Engine/Source/Runtime/Renderer/Private/SceneCaptureRendering.cpp
+++ b/Engine/Source/Runtime/Renderer/Private/SceneCaptureRendering.cpp
@@ -262,8 +262,8 @@ static void UpdateSceneCaptureContentDeferred_RenderThread(
 		FRDGBuilder GraphBuilder(RHICmdList, RDG_EVENT_NAME("SceneCapture"));
 #endif
 
-		FRDGTextureRef OutputTexture = RegisterExternalTexture(GraphBuilder, RenderTarget->GetRenderTargetTexture(), TEXT("SceneCapture"));
-		AddClearRenderTargetPass(GraphBuilder, OutputTexture, FLinearColor::Black, SceneRenderer->Views[0].UnscaledViewRect);
+		FRDGTextureRef TargetTexture = RegisterExternalTexture(GraphBuilder, RenderTarget->GetRenderTargetTexture(), TEXT("SceneCaptureTarget"));
+		AddClearRenderTargetPass(GraphBuilder, TargetTexture, FLinearColor::Black, SceneRenderer->Views[0].UnscaledViewRect);
 
 		// Render the scene normally
 		{
@@ -273,9 +273,12 @@ static void UpdateSceneCaptureContentDeferred_RenderThread(
 
 		if (bGenerateMips)
 		{
-			FGenerateMips::Execute(GraphBuilder, OutputTexture, GenerateMipsParams);
+			FGenerateMips::Execute(GraphBuilder, TargetTexture, GenerateMipsParams);
 		}
 
+		FRDGTextureRef ResolveTexture = RegisterExternalTexture(GraphBuilder, RenderTargetTexture->TextureRHI, TEXT("SceneCaptureResolve"));
+		AddCopyToResolveTargetPass(GraphBuilder, TargetTexture, ResolveTexture, ResolveParams);
+
 		GraphBuilder.Execute();
 	}
 
