Fix/reverb some debug code in VSM projection.
Reduce default normal offset as it is removing too many fine grained shadows in Reverb.

diff --git a/Engine/Shaders/Private/VirtualShadowMaps/VirtualShadowMapProjection.usf b/Engine/Shaders/Private/VirtualShadowMaps/VirtualShadowMapProjection.usf
--- a/Engine/Shaders/Private/VirtualShadowMaps/VirtualShadowMapProjection.usf
+++ b/Engine/Shaders/Private/VirtualShadowMaps/VirtualShadowMapProjection.usf
@@ -299,7 +299,7 @@ void VirtualShadowMapProjectionDirectionalPS(
 				Light.Direction,
 				LIGHT_TYPE_DIRECTIONAL,
 				TranslatedWorldPosition,
-				float3(0, 0, 0),//EstimatedGeoWorldNormal,
+				EstimatedGeoWorldNormal,
 				ContactShadowLengthWorld,
 				/* out */ DebugOutput,
 				HitDistance);
diff --git a/Engine/Source/Runtime/Renderer/Private/VirtualShadowMaps/VirtualShadowMapProjection.cpp b/Engine/Source/Runtime/Renderer/Private/VirtualShadowMaps/VirtualShadowMapProjection.cpp
--- a/Engine/Source/Runtime/Renderer/Private/VirtualShadowMaps/VirtualShadowMapProjection.cpp
+++ b/Engine/Source/Runtime/Renderer/Private/VirtualShadowMaps/VirtualShadowMapProjection.cpp
@@ -33,7 +33,7 @@ static TAutoConsoleVariable<float> CVarContactShadowLength(
 
 static TAutoConsoleVariable<float> CVarNormalOffsetWorld(
 	TEXT( "r.Shadow.v.NormalOffsetWorld" ),
-	0.5f,
+	0.1f,
 	TEXT( "World space offset along surface normal for shadow lookup." )
 	TEXT( "Higher values avoid artifacts on surfaces nearly parallel to the light, but also visibility offset shadows and increase the chance of hitting unmapped pages." ),
 	ECVF_RenderThreadSafe
