Added destructor to FProjectedShadowInfo

diff --git a/Engine/Source/Runtime/Renderer/Private/ShadowRendering.h b/Engine/Source/Runtime/Renderer/Private/ShadowRendering.h
--- a/Engine/Source/Runtime/Renderer/Private/ShadowRendering.h
+++ b/Engine/Source/Runtime/Renderer/Private/ShadowRendering.h
@@ -399,6 +399,8 @@ public:
 	// default constructor
 	FProjectedShadowInfo();
 
+	~FProjectedShadowInfo();
+
 	/**
 	 * for a per-object shadow. e.g. translucent particle system or a dynamic object in a precomputed shadow situation
 	 * @param InParentSceneInfo must not be 0
diff --git a/Engine/Source/Runtime/Renderer/Private/ShadowSetup.cpp b/Engine/Source/Runtime/Renderer/Private/ShadowSetup.cpp
--- a/Engine/Source/Runtime/Renderer/Private/ShadowSetup.cpp
+++ b/Engine/Source/Runtime/Renderer/Private/ShadowSetup.cpp
@@ -657,6 +657,10 @@ FProjectedShadowInfo::FProjectedShadowInfo()
 {
 }
 
+FProjectedShadowInfo::~FProjectedShadowInfo()
+{
+}
+
 /** Shadow border needs to be wide enough to prevent the shadow filtering from picking up content in other shadowmaps in the atlas. */
 const static uint32 SHADOW_BORDER = 4; 
 
