Changed RenderHitProxies() to check UseNanite() instead of DoesPlatformSupportNanite(), so the runtime cvar toggle is respected.

diff --git a/Engine/Source/Runtime/Renderer/Private/SceneHitProxyRendering.cpp b/Engine/Source/Runtime/Renderer/Private/SceneHitProxyRendering.cpp
--- a/Engine/Source/Runtime/Renderer/Private/SceneHitProxyRendering.cpp
+++ b/Engine/Source/Runtime/Renderer/Private/SceneHitProxyRendering.cpp
@@ -582,7 +582,7 @@ void FMobileSceneRenderer::RenderHitProxies(FRDGBuilder& GraphBuilder)
 
 void FDeferredShadingSceneRenderer::RenderHitProxies(FRDGBuilder& GraphBuilder)
 {
-	static const bool bNaniteEnabled = DoesPlatformSupportNanite(GMaxRHIShaderPlatform);
+	static const bool bNaniteEnabled = UseNanite(ShaderPlatform);
 
 	Scene->UpdateAllPrimitiveSceneInfos(GraphBuilder);
 
