Implemented SceneRenderingUtils.h with various mesh draw pass helper functions related to GPU instance culling.

diff --git a/Engine/Source/Runtime/Renderer/Private/DepthRendering.cpp b/Engine/Source/Runtime/Renderer/Private/DepthRendering.cpp
--- a/Engine/Source/Runtime/Renderer/Private/DepthRendering.cpp
+++ b/Engine/Source/Runtime/Renderer/Private/DepthRendering.cpp
@@ -30,6 +30,7 @@
 #include "MeshPassProcessor.inl"
 #include "PixelShaderUtils.h"
 #include "RenderGraphUtils.h"
+#include "SceneRenderingUtils.h"
 
 static TAutoConsoleVariable<int32> CVarParallelPrePass(
 	TEXT("r.ParallelPrePass"),
diff --git a/Engine/Source/Runtime/Renderer/Private/PostProcess/PostProcessCompositeEditorPrimitives.cpp b/Engine/Source/Runtime/Renderer/Private/PostProcess/PostProcessCompositeEditorPrimitives.cpp
--- a/Engine/Source/Runtime/Renderer/Private/PostProcess/PostProcessCompositeEditorPrimitives.cpp
+++ b/Engine/Source/Runtime/Renderer/Private/PostProcess/PostProcessCompositeEditorPrimitives.cpp
@@ -7,6 +7,7 @@
 #include "MeshPassProcessor.inl"
 #include "BasePassRendering.h"
 #include "MobileBasePassRendering.h"
+#include "SceneRenderingUtils.h"
 
 namespace
 {
diff --git a/Engine/Source/Runtime/Renderer/Private/SceneHitProxyRendering.cpp b/Engine/Source/Runtime/Renderer/Private/SceneHitProxyRendering.cpp
--- a/Engine/Source/Runtime/Renderer/Private/SceneHitProxyRendering.cpp
+++ b/Engine/Source/Runtime/Renderer/Private/SceneHitProxyRendering.cpp
@@ -27,6 +27,7 @@
 #include "FXSystem.h"
 #include "GPUSortManager.h"
 #include "VT/VirtualTextureSystem.h"
+#include "SceneRenderingUtils.h"
 #include "InstanceCulling/InstanceCullingManager.h"
 
 class FHitProxyShaderElementData : public FMeshMaterialShaderElementData
diff --git a/Engine/Source/Runtime/Renderer/Private/SceneRenderingUtils.h b/Engine/Source/Runtime/Renderer/Private/SceneRenderingUtils.h
new file mode 100644
--- /dev/null
+++ b/Engine/Source/Runtime/Renderer/Private/SceneRenderingUtils.h
@@ -0,0 +1,57 @@
+// Copyright Epic Games, Inc. All Rights Reserved.
+
+#pragma once
+
+#include "CoreMinimal.h"
+//#include "SceneManagement.h"
+#include "SceneRendering.h"
+
+template <typename PassParametersType>
+inline void AddSimpleElementCollectorPass(const FSimpleElementCollector& SimpleElementCollector, FRDGBuilder& GraphBuilder, PassParametersType* PassParameters, const FViewInfo& View, const FMeshPassProcessorRenderState& DrawRenderState)
+{
+	GraphBuilder.AddPass(
+		RDG_EVENT_NAME("EditorPrimitives"),
+		PassParameters,
+		ERDGPassFlags::Raster,
+		[&View, &SimpleElementCollector, DrawRenderState](FRHICommandList& RHICmdList)
+		{
+			RHICmdList.SetViewport(View.ViewRect.Min.X, View.ViewRect.Min.Y, 0.0f, View.ViewRect.Max.X, View.ViewRect.Max.Y, 1.0f);
+
+			SimpleElementCollector.DrawBatchedElements(RHICmdList, DrawRenderState, View, EBlendModeFilter::OpaqueAndMasked, SDPG_World);
+			SimpleElementCollector.DrawBatchedElements(RHICmdList, DrawRenderState, View, EBlendModeFilter::OpaqueAndMasked, SDPG_Foreground);
+		}
+	);
+}
+
+template <typename PassParametersType>
+inline void AddSimpleElementCollectorPass(const FSimpleElementCollector& SimpleElementCollector, FRDGBuilder& GraphBuilder, PassParametersType* PassParameters, const FViewInfo& View, const FMeshPassProcessorRenderState& DrawRenderState, ESceneDepthPriorityGroup SceneDepthPriorityGroup)
+{
+	GraphBuilder.AddPass(
+		RDG_EVENT_NAME("EditorPrimitives"),
+		PassParameters,
+		ERDGPassFlags::Raster,
+		[&View, &SimpleElementCollector, DrawRenderState, SceneDepthPriorityGroup](FRHICommandList& RHICmdList)
+		{
+			RHICmdList.SetViewport(View.ViewRect.Min.X, View.ViewRect.Min.Y, 0.0f, View.ViewRect.Max.X, View.ViewRect.Max.Y, 1.0f);
+
+			SimpleElementCollector.DrawBatchedElements(RHICmdList, DrawRenderState, View, EBlendModeFilter::OpaqueAndMasked, SceneDepthPriorityGroup);
+		}
+	);
+}
+
+template <typename PassParametersType>
+inline void AddBatchedElementsPass(const FBatchedElements& BatchedElements, FRDGBuilder& GraphBuilder, PassParametersType* PassParameters, const FViewInfo& View, const FMeshPassProcessorRenderState& DrawRenderState, bool bNeedToSwitchVerticalAxis)
+{
+	GraphBuilder.AddPass(
+		RDG_EVENT_NAME("BatchedElements"),
+		PassParameters,
+		ERDGPassFlags::Raster,
+		[&View, &BatchedElements, DrawRenderState, bNeedToSwitchVerticalAxis](FRHICommandList& RHICmdList)
+		{
+			RHICmdList.SetViewport(View.ViewRect.Min.X, View.ViewRect.Min.Y, 0.0f, View.ViewRect.Max.X, View.ViewRect.Max.Y, 1.0f);
+
+			// Draw the view's batched simple elements(lines, sprites, etc).
+			BatchedElements.Draw(RHICmdList, DrawRenderState, View.FeatureLevel, bNeedToSwitchVerticalAxis, View, false);
+		}
+	);
+}
