Store PrimitiveType value on FMeshDrawCommand for GPU instance culling, and move FCompareFMeshDrawCommands into the header, and stubbed out part of the new mesh draw command API.

diff --git a/Engine/Source/Runtime/Renderer/Private/MeshDrawCommands.cpp b/Engine/Source/Runtime/Renderer/Private/MeshDrawCommands.cpp
--- a/Engine/Source/Runtime/Renderer/Private/MeshDrawCommands.cpp
+++ b/Engine/Source/Runtime/Renderer/Private/MeshDrawCommands.cpp
@@ -127,27 +127,6 @@ void FPrimitiveIdVertexBufferPool::ReleaseDynamicRHI()
 	Entries.Empty();
 }
 
-
-struct FCompareFMeshDrawCommands
-{
-	FORCEINLINE bool operator() (FVisibleMeshDrawCommand A, FVisibleMeshDrawCommand B) const
-	{
-		// First order by a sort key.
-		if (A.SortKey != B.SortKey)
-		{
-			return A.SortKey < B.SortKey;
-		}
-
-		// Next order by instancing bucket.
-		if (A.StateBucketId != B.StateBucketId)
-		{
-			return A.StateBucketId < B.StateBucketId;
-		}
-
-		return false;
-	}
-};
-
 uint32 BitInvertIfNegativeFloat(uint32 f)
 {
 	unsigned mask = -int32(f >> 31) | 0x80000000;
@@ -520,6 +499,93 @@ void BuildMeshDrawCommandPrimitiveIdBuffer(
 	}
 }
 
+/**
+ * Allocate indirect arg slots for all meshes to use instancing, 
+ * add commands that populate the indirect calls and index & id buffers, and
+ * Collapse all commands that share the same state bucket ID 
+ * NOTE: VisibleMeshDrawCommandsInOut can only become shorter.
+ */
+void SetupGPUInstancedDraws(
+	FInstanceCullingContext &InstanceCullingContext,
+	FMeshCommandOneFrameArray& VisibleMeshDrawCommandsInOut,
+	// Stats
+	int32& MaxInstances,
+	int32& VisibleMeshDrawCommandsNum,
+	int32& NewPassVisibleMeshDrawCommandsNum)
+{
+#if defined(GPUCULL_TODO)
+	QUICK_SCOPE_CYCLE_COUNTER(STAT_BuildMeshDrawCommandPrimitiveIdBuffer);
+
+	FVisibleMeshDrawCommand* RESTRICT PassVisibleMeshDrawCommands = VisibleMeshDrawCommandsInOut.GetData();
+
+
+	QUICK_SCOPE_CYCLE_COUNTER(STAT_DynamicInstancingOfVisibleMeshDrawCommands);
+
+	int32 CurrentStateBucketId = -1;
+	MaxInstances = 1;
+
+	// Scan through and compact away all with consecutive statebucked ID, and record primitive IDs in GPU-scene culling command
+	const int32 NumDrawCommandsIn = VisibleMeshDrawCommandsInOut.Num();
+	int32 NumDrawCommandsOut = 0;
+	// Allocate conservatively for all commands, may not use all.
+	int32 BaseArgsSlotIndex = InstanceCullingContext.AllocateArgsSlotRange(NumDrawCommandsIn);
+	for (int32 DrawCommandIndex = 0; DrawCommandIndex < NumDrawCommandsIn; ++DrawCommandIndex)
+	{
+		const FVisibleMeshDrawCommand& RESTRICT VisibleMeshDrawCommand = PassVisibleMeshDrawCommands[DrawCommandIndex];
+
+
+		// GPUCULL_TODO: Hoist this flag to the VisibleMeshDrawCommand to avoid going through this indirection
+		const bool bSupporsGPUSceneInstancing = VisibleMeshDrawCommand.MeshDrawCommand->PrimitiveIdStreamIndex != INDEX_NONE;
+
+
+		if (CurrentStateBucketId != -1 && VisibleMeshDrawCommand.StateBucketId == CurrentStateBucketId)
+		{
+			// Drop since previous covers for this
+		}
+		else
+		{
+			const FMeshDrawCommand* RESTRICT MeshDrawCommand = VisibleMeshDrawCommand.MeshDrawCommand;
+			
+			// GPUCULL_TODO: Always allocate command as otherwise the 1:1 mapping between mesh draw command index and culling command index is broken.
+			// if (bSupporsGPUSceneInstancing)
+			{
+				InstanceCullingContext.BeginCullingCommand(MeshDrawCommand->PrimitiveType, MeshDrawCommand->VertexParams.BaseVertexIndex, MeshDrawCommand->FirstIndex, MeshDrawCommand->NumPrimitives);
+			}
+			// Record the last bucket ID (may be -1)
+			CurrentStateBucketId = VisibleMeshDrawCommand.StateBucketId;
+
+			// If we have dropped any we need to move up
+			if (DrawCommandIndex > NumDrawCommandsOut)
+			{
+				PassVisibleMeshDrawCommands[NumDrawCommandsOut] = PassVisibleMeshDrawCommands[DrawCommandIndex];
+			}
+			NumDrawCommandsOut++;
+		}
+
+		if (bSupporsGPUSceneInstancing)
+		{
+			// append 'culling command' targeting the current slot
+			// This will cause all instances belonging to the Primitive to be added to the command, if they are visible etc (GPU-Scene knows all - sees all)
+			if (VisibleMeshDrawCommand.RunArray)
+			{
+				InstanceCullingContext.AddInstanceRunToCullingCommand(VisibleMeshDrawCommand.DrawPrimitiveId, VisibleMeshDrawCommand.RunArray, VisibleMeshDrawCommand.NumRuns);
+			}
+			else
+			{
+				InstanceCullingContext.AddPrimitiveToCullingCommand(VisibleMeshDrawCommand.DrawPrimitiveId);
+			}
+		}
+	}
+	ensureMsgf(NumDrawCommandsOut == InstanceCullingContext.CullingCommands.Num(), TEXT("There must be a 1:1 mapping between culling commands and mesh draw commands, as this assumption is made in SubmitGPUInstancedMeshDrawCommandsRange."));
+	// Setup instancing stats for logging.
+	VisibleMeshDrawCommandsNum = VisibleMeshDrawCommandsInOut.Num();
+	NewPassVisibleMeshDrawCommandsNum = NumDrawCommandsOut;
+
+	// Resize array post-compaction of dynamic instances
+	VisibleMeshDrawCommandsInOut.SetNum(NumDrawCommandsOut, false);
+#endif // GPUCULL_TODO
+}
+
 /**
  * Converts each FMeshBatch into a set of FMeshDrawCommands for a specific mesh pass type.
  */
@@ -1017,6 +1083,37 @@ void SortAndMergeDynamicPassMeshDrawCommands(
 }
 
 
+
+/*
+ * Used by various dynamic passes to sort/merge mesh draw commands immediately on a rendering thread.
+ */
+void SortAndMergeDynamicPassMeshDrawCommands(
+	ERHIFeatureLevel::Type FeatureLevel,
+	FMeshCommandOneFrameArray& VisibleMeshDrawCommands,
+	FDynamicMeshDrawCommandStorage& MeshDrawCommandStorage,
+	FInstanceCullingContext& InstanceCullingContext)
+{
+#if defined(GPUCULL_TODO)
+	const bool bUseGPUScene = UseGPUScene(GMaxRHIShaderPlatform, FeatureLevel);
+
+	const int32 NumDrawCommands = VisibleMeshDrawCommands.Num();
+	if (NumDrawCommands > 0)
+	{
+		FMeshCommandOneFrameArray NewPassVisibleMeshDrawCommands;
+		int32 MaxInstances = 1;
+		int32 VisibleMeshDrawCommandsNum = 0;
+		int32 NewPassVisibleMeshDrawCommandsNum = 0;
+
+		VisibleMeshDrawCommands.Sort(FCompareFMeshDrawCommands());
+
+		if (bUseGPUScene)
+		{
+			SetupGPUInstancedDraws(InstanceCullingContext, VisibleMeshDrawCommands, MaxInstances, VisibleMeshDrawCommandsNum, NewPassVisibleMeshDrawCommandsNum);
+		}
+	}
+#endif // GPUCULL_TODO
+}
+
 void FParallelMeshDrawCommandPass::DispatchPassSetup(
 	FScene* Scene,
 	const FViewInfo& View,
@@ -1216,6 +1313,31 @@ FParallelMeshDrawCommandPass::~FParallelMeshDrawCommandPass()
 	check(TaskEventRef == nullptr);
 }
 
+void SubmitGPUInstancedMeshDrawCommandsRange(
+	const FMeshCommandOneFrameArray& VisibleMeshDrawCommands,
+	const FGraphicsMinimalPipelineStateSet& GraphicsMinimalPipelineStateSet,
+	int32 StartIndex,
+	int32 NumMeshDrawCommands,
+	FRHIVertexBuffer* InstanceIdsOffsetBuffer, // Bound to a vertex stream to fetch a start offset for all instances, need to be 0-stepping
+	FRHIVertexBuffer* IndirectArgsBuffer, // Overrides the args for the draw call
+	FRHICommandList& RHICmdList)
+{
+#if defined(GPUCULL_TODO)
+	FMeshDrawCommandStateCache StateCache;
+	INC_DWORD_STAT_BY(STAT_MeshDrawCalls, NumMeshDrawCommands);
+
+	for (int32 DrawCommandIndex = StartIndex; DrawCommandIndex < StartIndex + NumMeshDrawCommands; DrawCommandIndex++)
+	{
+		//SCOPED_CONDITIONAL_DRAW_EVENTF(RHICmdList, MeshEvent, GEmitMeshDrawEvent != 0, TEXT("Mesh Draw"));
+
+		const FVisibleMeshDrawCommand& VisibleMeshDrawCommand = VisibleMeshDrawCommands[DrawCommandIndex];
+		const uint32 IndirectArgsByteOffset = uint32(DrawCommandIndex) * FInstanceCullingContext::IndirectArgsNumWords * sizeof(uint32);
+		const int32 InstanceIdsOffsetBufferByteOffset = DrawCommandIndex * sizeof(int32);
+		FMeshDrawCommand::SubmitDraw(*VisibleMeshDrawCommand.MeshDrawCommand, GraphicsMinimalPipelineStateSet, InstanceIdsOffsetBuffer, InstanceIdsOffsetBufferByteOffset, 1, RHICmdList, StateCache, IndirectArgsBuffer, IndirectArgsByteOffset);
+	}
+#endif // GPUCULL_TODO
+}
+
 class FDrawVisibleMeshCommandsAnyThreadTask : public FRenderTask
 {
 	FRHICommandList& RHICmdList;
@@ -1281,6 +1403,38 @@ public:
 	}
 };
 
+#if 0 // GPUCULL_TODO
+void FParallelMeshDrawCommandPass::DispatchInstanceCullingFinalize(FRHICommandList& RHICmdList)
+{
+#if 0
+
+	// 1. Add pre-req on the setup pass
+	if (TaskEventRef.IsValid())
+	{
+		RHICmdList.AddDispatchPrerequisite(TaskEventRef);
+	}
+	// 2. In the same way that the DispatchDraw did, queue up commands to upload the input data to the culling unpack - this will happen after the AddDispatchPrerequisite has been reached, so would be fine.
+	//    At this point we must have stable pointers such that none of the data moves, it must also be passed over to be owned by the RHI thread so it can delete the data.
+
+	RHICmdList.EnqueueLambda([
+			VertexBuffer = PrimitiveIdsBuffer,
+			VertexBufferData = TaskContext.PrimitiveIdBufferData,
+			VertexBufferDataSize = TaskContext.PrimitiveIdBufferDataSize,
+			PrimitiveIdVertexBufferPoolEntry = PrimitiveIdVertexBufferPoolEntry](FRHICommandListImmediate& CmdList)
+		{
+			// Upload vertex buffer data.
+			void* RESTRICT Data = (void* RESTRICT)CmdList.LockVertexBuffer(VertexBuffer, 0, VertexBufferDataSize, RLM_WriteOnly);
+			FMemory::Memcpy(Data, VertexBufferData, VertexBufferDataSize);
+			CmdList.UnlockVertexBuffer(VertexBuffer);
+
+			FMemory::Free(VertexBufferData);
+		});
+
+	RHICommandList.RHIThreadFence(true);
+#endif
+}
+#endif // GPUCULL_TODO
+
 void FParallelMeshDrawCommandPass::BuildRenderingCommands(FRDGBuilder& GraphBuilder, FGPUScene& GPUScene, FInstanceCullingDrawParams& OutInstanceCullingDrawParams)
 {
 #if defined(GPUCULL_TODO)
@@ -1403,6 +1557,22 @@ void FParallelMeshDrawCommandPass::DispatchDraw(FParallelCommandListSet* Paralle
 	}
 }
 
+// GPUCULL_TODO
+//void FParallelMeshDrawCommandPass::AddDrawPasses(FRDGBuilder& GraphBuilder, const FRDGEventName &EventName, PassParameters) const
+//{
+//	GraphBuilder.AddPass(EventName
+//		RDG_EVENT_NAME("BasePassParallel"),
+//		PassParameters,
+//		ERDGPassFlags::Raster | ERDGPassFlags::SkipRenderPass,
+//		[this, &View, PassParameters](FRHICommandListImmediate& RHICmdList)
+//		{
+//			Scene->UniformBuffers.UpdateViewUniformBuffer(View);
+//			FRDGParallelCommandListSet ParallelCommandListSet(RHICmdList, GET_STATID(STAT_CLP_BasePass), *this, View, FParallelCommandListBindings(PassParameters));
+//			View.ParallelMeshDrawCommandPasses[EMeshPass::BasePass].DispatchDraw(&ParallelCommandListSet, RHICmdList);
+//		});
+//
+//}
+
 void FParallelMeshDrawCommandPass::DumpInstancingStats() const
 {
 	if (!PassNameForStats.IsEmpty() && TaskContext.VisibleMeshDrawCommandsNum > 0)
@@ -1417,4 +1587,4 @@ void FParallelMeshDrawCommandPass::DumpInstancingStats() const
 void FParallelMeshDrawCommandPass::SetDumpInstancingStats(const FString& InPassNameForStats)
 {
 	PassNameForStats = InPassNameForStats;
-}
\ No newline at end of file
+}
diff --git a/Engine/Source/Runtime/Renderer/Private/MeshDrawCommands.h b/Engine/Source/Runtime/Renderer/Private/MeshDrawCommands.h
--- a/Engine/Source/Runtime/Renderer/Private/MeshDrawCommands.h
+++ b/Engine/Source/Runtime/Renderer/Private/MeshDrawCommands.h
@@ -229,3 +229,26 @@ extern void SortAndMergeDynamicPassMeshDrawCommands(
 	FDynamicMeshDrawCommandStorage& MeshDrawCommandStorage,
 	FRHIVertexBuffer*& OutPrimitiveIdVertexBuffer,
 	uint32 InstanceFactor);
+
+extern void SortAndMergeDynamicPassMeshDrawCommands(
+	ERHIFeatureLevel::Type FeatureLevel,
+	FMeshCommandOneFrameArray& VisibleMeshDrawCommands,
+	FDynamicMeshDrawCommandStorage& MeshDrawCommandStorage,
+	FInstanceCullingContext* InstanceCullingContext);
+
+void SubmitGPUInstancedMeshDrawCommandsRange(
+	const FMeshCommandOneFrameArray& VisibleMeshDrawCommands,
+	const FGraphicsMinimalPipelineStateSet& GraphicsMinimalPipelineStateSet,
+	int32 StartIndex,
+	int32 NumMeshDrawCommands,
+	FRHIVertexBuffer* InstanceIdsOffsetBuffer, // Bound to a vertex stream to fetch a start offset for all instances, need to be 0-stepping
+	FRHIVertexBuffer* IndirectArgsBuffer, // Overrides the args for the draw call
+	FRHICommandList& RHICmdList);
+
+void SetupGPUInstancedDraws(
+	FInstanceCullingContext& InstanceCullingContext,
+	FMeshCommandOneFrameArray& VisibleMeshDrawCommandsInOut,
+	// Stats
+	int32& MaxInstances,
+	int32& VisibleMeshDrawCommandsNum,
+	int32& NewPassVisibleMeshDrawCommandsNum);
diff --git a/Engine/Source/Runtime/Renderer/Private/MeshPassProcessor.cpp b/Engine/Source/Runtime/Renderer/Private/MeshPassProcessor.cpp
--- a/Engine/Source/Runtime/Renderer/Private/MeshPassProcessor.cpp
+++ b/Engine/Source/Runtime/Renderer/Private/MeshPassProcessor.cpp
@@ -1277,6 +1277,36 @@ void DrawDynamicMeshPassPrivate(
 	}
 }
 
+void DrawDynamicMeshPassPrivate(
+	const FSceneView& View,
+	const FScene &Scene,
+	FRHICommandListImmediate& RHICmdList,
+	FMeshCommandOneFrameArray& VisibleMeshDrawCommands,
+	FDynamicMeshDrawCommandStorage& DynamicMeshDrawCommandStorage,
+	FGraphicsMinimalPipelineStateSet& GraphicsMinimalPipelineStateSet,
+	bool& InNeedsShaderInitialization)
+{
+#if defined(GPUCULL_TODO)
+	check(View.bIsViewInfo);
+
+	if (VisibleMeshDrawCommands.Num() > 0)
+	{
+		const bool bDynamicInstancing = IsDynamicInstancingEnabled(View.GetFeatureLevel());
+
+		FRHIVertexBuffer* PrimitiveIdVertexBuffer = nullptr;
+
+		ApplyViewOverridesToMeshDrawCommands(View, VisibleMeshDrawCommands, DynamicMeshDrawCommandStorage, GraphicsMinimalPipelineStateSet, InNeedsShaderInitialisation);
+		
+		// GPUCULL_TODO: determine if unsupportable OR just not supporting scene primitives (as this requires a multi-pass approach)
+
+		//SortAndSubmitDynamicMeshPassDrawCommands(View, Scene, VisibleMeshDrawCommands, DynamicMeshDrawCommandStorage, GraphicsMinimalPipelineStateSet, RHICmdList);
+
+		//SortAndMergeDynamicPassMeshDrawCommands(View.GetFeatureLevel(), VisibleMeshDrawCommands, DynamicMeshDrawCommandStorage, PrimitiveIdVertexBuffer, InstanceFactor);
+		//SubmitMeshDrawCommandsRange(VisibleMeshDrawCommands, GraphicsMinimalPipelineStateSet, PrimitiveIdVertexBuffer, 0, bDynamicInstancing, 0, VisibleMeshDrawCommands.Num(), InstanceFactor, RHICmdList);
+	}
+#endif
+}
+
 FMeshDrawCommandSortKey CalculateMeshStaticSortKey(const FMeshMaterialShader* VertexShader, const FMeshMaterialShader* PixelShader)
 {
 	FMeshDrawCommandSortKey SortKey;
diff --git a/Engine/Source/Runtime/Renderer/Public/MeshPassProcessor.h b/Engine/Source/Runtime/Renderer/Public/MeshPassProcessor.h
--- a/Engine/Source/Runtime/Renderer/Public/MeshPassProcessor.h
+++ b/Engine/Source/Runtime/Renderer/Public/MeshPassProcessor.h
@@ -986,6 +986,9 @@ public:
 	/** Non-pipeline state */
 	uint8 StencilRef;
 
+	/** Redundant as already present in CachedPipelineId, but need access for dynamic instancing on GPU. */
+	EPrimitiveType PrimitiveType : PT_NumBits;
+
 	FMeshDrawCommand() {};
 	FMeshDrawCommand(FMeshDrawCommand&& Other) = default;
 	FMeshDrawCommand(const FMeshDrawCommand& Other) = default;
@@ -1281,6 +1284,26 @@ public:
 	ERasterizerCullMode MeshCullMode : ERasterizerCullMode_NumBits + 1;
 };
 
+struct FCompareFMeshDrawCommands
+{
+	FORCEINLINE bool operator() (FVisibleMeshDrawCommand A, FVisibleMeshDrawCommand B) const
+	{
+		// First order by a sort key.
+		if (A.SortKey != B.SortKey)
+		{
+			return A.SortKey < B.SortKey;
+		}
+
+		// Next order by instancing bucket.
+		if (A.StateBucketId != B.StateBucketId)
+		{
+			return A.StateBucketId < B.StateBucketId;
+		}
+
+		return false;
+	}
+};
+
 template <>
 struct TUseBitwiseSwap<FVisibleMeshDrawCommand>
 {
diff --git a/Engine/Source/Runtime/Renderer/Public/MeshPassProcessor.inl b/Engine/Source/Runtime/Renderer/Public/MeshPassProcessor.inl
--- a/Engine/Source/Runtime/Renderer/Public/MeshPassProcessor.inl
+++ b/Engine/Source/Runtime/Renderer/Public/MeshPassProcessor.inl
@@ -50,6 +50,7 @@ void FMeshPassProcessor::BuildMeshDrawCommands(
 	FMeshDrawCommand SharedMeshDrawCommand;
 
 	SharedMeshDrawCommand.SetStencilRef(DrawRenderState.GetStencilRef());
+	SharedMeshDrawCommand.PrimitiveType = (EPrimitiveType)MeshBatch.Type;
 
 	FGraphicsMinimalPipelineStateInitializer PipelineState;
 	PipelineState.PrimitiveType = (EPrimitiveType)MeshBatch.Type;
