Implemented FProjectedShadowInfo::TestPrimitiveFarCascadeConditions() helper

diff --git a/Engine/Source/Runtime/Renderer/Private/ShadowRendering.h b/Engine/Source/Runtime/Renderer/Private/ShadowRendering.h
--- a/Engine/Source/Runtime/Renderer/Private/ShadowRendering.h
+++ b/Engine/Source/Runtime/Renderer/Private/ShadowRendering.h
@@ -776,6 +776,9 @@ private:
 		bool bCameraInsideShadowFrustum,
 		const FInstanceCullingDrawParams& InstanceCullingDrawParams) const;
 
+
+	FORCEINLINE bool TestPrimitiveFarCascadeConditions(bool bPrimitiveCastsFarShadow, const FBoxSphereBounds& Bounds) const;
+
 	friend class FShadowDepthVS;
 	friend class FShadowDepthBasePS;
 	friend class FShadowVolumeBoundProjectionVS;
diff --git a/Engine/Source/Runtime/Renderer/Private/ShadowSetup.cpp b/Engine/Source/Runtime/Renderer/Private/ShadowSetup.cpp
--- a/Engine/Source/Runtime/Renderer/Private/ShadowSetup.cpp
+++ b/Engine/Source/Runtime/Renderer/Private/ShadowSetup.cpp
@@ -1468,7 +1468,7 @@ void FProjectedShadowInfo::AddSubjectPrimitive(FPrimitiveSceneInfo* PrimitiveSce
 
 	if (!ReceiverPrimitives.Contains(PrimitiveSceneInfo)
 		// Far cascade only casts from primitives marked for it
-		&& (!CascadeSettings.bFarShadowCascade || PrimitiveSceneInfo->Proxy->CastsFarShadow()))
+		&& TestPrimitiveFarCascadeConditions(PrimitiveSceneInfo->Proxy->CastsFarShadow(), PrimitiveSceneInfo->Proxy->GetBounds()))
 	{
 		const FPrimitiveSceneProxy* Proxy = PrimitiveSceneInfo->Proxy;
 
@@ -1635,6 +1635,16 @@ void FProjectedShadowInfo::AddSubjectPrimitive(FPrimitiveSceneInfo* PrimitiveSce
 	}
 }
 
+bool FProjectedShadowInfo::TestPrimitiveFarCascadeConditions(bool bPrimitiveCastsFarShadow, const FBoxSphereBounds& Bounds) const
+{
+	if (CascadeSettings.bFarShadowCascade)
+	{
+		return bPrimitiveCastsFarShadow;
+	}
+
+	return true;
+}
+
 uint64 FProjectedShadowInfo::AddSubjectPrimitive_AnyThread(
 	const FPrimitiveSceneInfoCompact& PrimitiveSceneInfoCompact,
 	TArray<FViewInfo>* ViewArray,
@@ -1673,7 +1683,7 @@ uint64 FProjectedShadowInfo::AddSubjectPrimitive_AnyThread(
 
 	if (!ReceiverPrimitives.Contains(PrimitiveSceneInfo)
 		// Far cascade only casts from primitives marked for it
-		&& (!CascadeSettings.bFarShadowCascade || PrimitiveSceneInfoCompact.Proxy->CastsFarShadow()))
+		&& TestPrimitiveFarCascadeConditions(PrimitiveSceneInfoCompact.Proxy->CastsFarShadow(), PrimitiveSceneInfoCompact.Bounds))
 	{
 		FViewInfo* CurrentView;
 		const bool bWholeSceneDirectionalShadow = IsWholeSceneDirectionalShadow();
