Additional (disabled for now) helper routines and changes related to GPU instance culling (wrapped with GPUCULL_TODO)

diff --git a/Engine/Source/Runtime/Engine/Classes/Engine/InstancedStaticMesh.h b/Engine/Source/Runtime/Engine/Classes/Engine/InstancedStaticMesh.h
--- a/Engine/Source/Runtime/Engine/Classes/Engine/InstancedStaticMesh.h
+++ b/Engine/Source/Runtime/Engine/Classes/Engine/InstancedStaticMesh.h
@@ -467,7 +467,9 @@ public:
 	,	bHasSelectedInstances(InComponent->SelectedInstances.Num() > 0)
 #endif
 	{
+	#if !defined(GPUCULL_TODO)
 		bVFRequiresPrimitiveUniformBuffer = true;
+	#endif
 		SetupProxy(InComponent);
 
 #if RHI_RAYTRACING
diff --git a/Engine/Source/Runtime/Engine/Private/InstancedStaticMesh.cpp b/Engine/Source/Runtime/Engine/Private/InstancedStaticMesh.cpp
--- a/Engine/Source/Runtime/Engine/Private/InstancedStaticMesh.cpp
+++ b/Engine/Source/Runtime/Engine/Private/InstancedStaticMesh.cpp
@@ -626,7 +626,13 @@ void FInstancedStaticMeshVertexFactory::ModifyCompilationEnvironment(const FVert
 		OutEnvironment.SetDefine(TEXT("MANUAL_VERTEX_FETCH"), TEXT("1"));
 	}
 
+#if defined(GPUCULL_TODO)
+	// USE_INSTANCE_CULLING - set up additional instancing attributes (basic instancing is the default)
+	OutEnvironment.SetDefine(TEXT("USE_INSTANCE_CULLING"), TEXT("1"));
+#else
 	OutEnvironment.SetDefine(TEXT("USE_INSTANCING"), TEXT("1"));
+#endif
+
 	if (IsFeatureLevelSupported(Parameters.Platform, ERHIFeatureLevel::SM5))
 	{
 		OutEnvironment.SetDefine(TEXT("USE_DITHERED_LOD_TRANSITION_FOR_INSTANCED"), ALLOW_DITHERED_LOD_FOR_INSTANCED_STATIC_MESHES);
@@ -701,6 +707,20 @@ void FInstancedStaticMeshVertexFactory::InitRHI()
 		Elements.Add(AccessStreamComponent(Data.PositionComponent,0));
 	}
 
+#if defined(GPUCULL_TODO)
+	{
+		const uint8 Index = static_cast<uint8>(EVertexInputStreamType::Default);
+		PrimitiveIdStreamIndex[Index] = -1;
+		const bool bCanUseGPUScene = UseGPUScene(GMaxRHIShaderPlatform, GMaxRHIFeatureLevel);
+		if (GetType()->SupportsPrimitiveIdStream() && bCanUseGPUScene)
+		{
+			// When the VF is used for rendering in normal mesh passes, this vertex buffer and offset will be overridden
+			Elements.Add(AccessStreamComponent(FVertexStreamComponent(&GPrimitiveIdDummy, 0, 0, /*sizeof(uint32)*/0, VET_UInt, EVertexStreamUsage::Instancing), 13));
+			PrimitiveIdStreamIndex[Index] = Elements.Last().StreamIndex;
+		}
+	}
+#endif
+
 	// only tangent,normal are used by the stream. the binormal is derived in the shader
 	uint8 TangentBasisAttributes[2] = { 1, 2 };
 	for(int32 AxisIndex = 0;AxisIndex < 2;AxisIndex++)
@@ -796,7 +816,11 @@ IMPLEMENT_VERTEX_FACTORY_PARAMETER_TYPE(FInstancedStaticMeshVertexFactory, SF_Ve
 #if RHI_RAYTRACING
 IMPLEMENT_VERTEX_FACTORY_PARAMETER_TYPE(FInstancedStaticMeshVertexFactory, SF_RayHitGroup, FInstancedStaticMeshVertexFactoryShaderParameters);
 #endif
+#if defined(GPUCULL_TODO)
+IMPLEMENT_VERTEX_FACTORY_TYPE_EX(FInstancedStaticMeshVertexFactory,"/Engine/Private/LocalVertexFactory.ush",true,true,true,true,true,true,true,false);
+#else
 IMPLEMENT_VERTEX_FACTORY_TYPE_EX(FInstancedStaticMeshVertexFactory,"/Engine/Private/LocalVertexFactory.ush",true,true,true,true,true,true,false,false);
+#endif
 
 void FInstancedStaticMeshRenderData::InitVertexFactories()
 {
@@ -1024,6 +1048,32 @@ void FInstancedStaticMeshSceneProxy::SetupProxy(UInstancedStaticMeshComponent* I
 	// unselected only
 	UserData_DeselectedInstances = UserData_AllInstances;
 	UserData_DeselectedInstances.bRenderSelected = false;
+
+#if defined(GPUCULL_TODO)
+	const TArray<int32>& InstanceReorderTable = InComponent->InstanceReorderTable;
+	bSupportsInstanceDataBuffer = true;
+	Instances.SetNum(InComponent->GetInstanceCount());
+	for (int32 InInstanceIndex = 0; InInstanceIndex < Instances.Num(); ++InInstanceIndex)
+	{
+		int32 OutInstanceIndex = InInstanceIndex;
+		if (OutInstanceIndex < InstanceReorderTable.Num())
+		{
+			OutInstanceIndex = InstanceReorderTable[OutInstanceIndex];
+		}
+		FTransform InstanceTransform;
+		InComponent->GetInstanceTransform(InInstanceIndex, InstanceTransform);
+
+		FPrimitiveInstance& Instance = Instances[OutInstanceIndex];
+		Instance.PrimitiveId = ~uint32(0);
+		Instance.InstanceToLocal = InstanceTransform.ToMatrixWithScale();
+		// GPUCULL_TODO: not sure this is needed either - might be better to delegate to later anyway since inverse can then be threaded, plus some platforms might not need it at all.
+		Instance.LocalToInstance = Instance.InstanceToLocal.Inverse();
+		// Filled in during GPU Scene update...
+		Instance.LocalToWorld.SetIdentity();
+		Instance.RenderBounds = InComponent->GetStaticMesh()->GetBounds();
+		Instance.LocalBounds = Instance.RenderBounds.TransformBy(Instance.InstanceToLocal);
+	}
+#endif
 }
 
 void FInstancedStaticMeshSceneProxy::DestroyRenderThreadResources()
@@ -2913,7 +2963,11 @@ void UInstancedStaticMeshComponent::InitPerInstanceRenderData(bool InitializeFro
 	UWorld* World = GetWorld();
 	ERHIFeatureLevel::Type FeatureLevel = World != nullptr ? World->FeatureLevel.GetValue() : GMaxRHIFeatureLevel;
 
+#if defined(GPUCULL_TODO)
+	bool KeepInstanceBufferCPUAccess = true;//GIsEditor || InRequireCPUAccess || ComponentRequestsCPUAccess(this, FeatureLevel);
+#else
 	bool KeepInstanceBufferCPUAccess = GIsEditor || InRequireCPUAccess || ComponentRequestsCPUAccess(this, FeatureLevel);
+#endif
 
 	if (InSharedInstanceBufferData != nullptr)
 	{
diff --git a/Engine/Source/Runtime/Engine/Public/StaticMeshResources.h b/Engine/Source/Runtime/Engine/Public/StaticMeshResources.h
--- a/Engine/Source/Runtime/Engine/Public/StaticMeshResources.h
+++ b/Engine/Source/Runtime/Engine/Public/StaticMeshResources.h
@@ -934,6 +934,18 @@ public:
 
 	virtual int32 GetLightMapCoordinateIndex() const override;
 
+#if defined(GPUCULL_TODO)
+	virtual const TArray<FPrimitiveInstance>* GetPrimitiveInstances() const
+	{
+		return &Instances;
+	}
+
+	virtual TArray<FPrimitiveInstance>* GetPrimitiveInstances()
+	{
+		return &Instances;
+	}
+#endif
+
 protected:
 	/** Configures mesh batch vertex / index state. Returns the number of primitives used in the element. */
 	uint32 SetMeshElementGeometrySource(
@@ -1146,6 +1158,11 @@ protected:
 	uint8 HierarchicalLODIndex;
 #endif
 
+#if defined(GPUCULL_TODO)
+protected:
+	TArray<FPrimitiveInstance> Instances;
+#endif
+
 public:
 
 	/**
diff --git a/Engine/Source/Runtime/RenderCore/Public/RendererInterface.h b/Engine/Source/Runtime/RenderCore/Public/RendererInterface.h
--- a/Engine/Source/Runtime/RenderCore/Public/RendererInterface.h
+++ b/Engine/Source/Runtime/RenderCore/Public/RendererInterface.h
@@ -658,6 +658,36 @@ public:
 	virtual void EndFrame() {}
 };
 
+#if defined(GPUCULL_TODO)
+
+/**
+ */
+class IScenePrimitiveRenderingContext
+{
+public:
+	virtual ~IScenePrimitiveRenderingContext()
+	{
+	}
+};
+
+struct FScenePrimitiveRenderingContextScopeHelper
+{
+	FScenePrimitiveRenderingContextScopeHelper(IScenePrimitiveRenderingContext* InScenePrimitiveRenderingContext)
+	: ScenePrimitiveRenderingContext(InScenePrimitiveRenderingContext)
+	{
+	}
+
+	~FScenePrimitiveRenderingContextScopeHelper()
+	{
+		// GPUCULL_TODO: Is new/delete reasonable here?
+		delete ScenePrimitiveRenderingContext;
+	}
+
+	IScenePrimitiveRenderingContext* ScenePrimitiveRenderingContext;
+};
+
+#endif // GPUCULL_TODO
+
 /**
  * The public interface of the renderer module.
  */
@@ -795,5 +825,13 @@ public:
 	virtual void FlushVirtualTextureCache() = 0;
 
 	virtual void RegisterPersistentViewUniformBufferExtension(IPersistentViewUniformBufferExtension* Extension) = 0;
+
+#if defined(GPUCULL_TODO)
+	/**
+	 * Prepare Scene primitive rendering and return context. Ensures all primitives that are created are commited and GPU-Scene is updated and allocates a dynamic primitive context.
+	 * The intended use is for stand-alone rendering that involves Scene proxies (that then may need the machinery to render GPU-Scene aware primitives.
+	 */
+	virtual IScenePrimitiveRenderingContext* BeginScenePrimitiveRendering(FRDGBuilder& GraphBuilder, FSceneViewFamily* ViewFamily) = 0;
+#endif
 };
 
diff --git a/Engine/Source/Runtime/Renderer/Private/DepthRendering.cpp b/Engine/Source/Runtime/Renderer/Private/DepthRendering.cpp
--- a/Engine/Source/Runtime/Renderer/Private/DepthRendering.cpp
+++ b/Engine/Source/Runtime/Renderer/Private/DepthRendering.cpp
@@ -403,6 +403,36 @@ void AddDitheredStencilFillPass(FRDGBuilder& GraphBuilder, TConstArrayView<FView
 	}
 }
 
+#if defined(GPUCULL_TODO)
+
+// GPUCULL_TODO: Move to Utils file and make templated on params and mesh pass processor
+static void AddViewMeshElementsPass(const TIndirectArray<FMeshBatch>& MeshElements, FRDGBuilder& GraphBuilder, FDepthPassParameters* PassParameters, const FScene* Scene, const FViewInfo& View, const FMeshPassProcessorRenderState& DrawRenderState, bool bRespectUseAsOccluderFlag, EDepthDrawingMode DepthDrawingMode, FInstanceCullingManager& InstanceCullingManager)
+{
+	AddSimpleMeshPass(GraphBuilder, PassParameters, Scene, View, &InstanceCullingManager, RDG_EVENT_NAME("ViewMeshElementsPass"), View.ViewRect,
+		[&View, Scene, DrawRenderState, &MeshElements, bRespectUseAsOccluderFlag, DepthDrawingMode](FDynamicPassMeshDrawListContext* DynamicMeshPassContext)
+		{
+			FDepthPassMeshProcessor PassMeshProcessor(
+				View.Family->Scene->GetRenderScene(),
+				&View,
+				DrawRenderState,
+				bRespectUseAsOccluderFlag,
+				DepthDrawingMode,
+				false,
+				false,
+				DynamicMeshPassContext);
+
+			const uint64 DefaultBatchElementMask = ~0ull;
+
+			for (const FMeshBatch& MeshBatch : MeshElements)
+			{
+				PassMeshProcessor.AddMeshBatch(MeshBatch, DefaultBatchElementMask, nullptr);
+			}
+		}
+	);
+}
+
+#endif // GPUCULL_TODO
+
 static void RenderPrePassEditorPrimitives(
 	FRDGBuilder& GraphBuilder,
 	const FViewInfo& View,
diff --git a/Engine/Source/Runtime/Renderer/Private/MeshPassProcessor.cpp b/Engine/Source/Runtime/Renderer/Private/MeshPassProcessor.cpp
--- a/Engine/Source/Runtime/Renderer/Private/MeshPassProcessor.cpp
+++ b/Engine/Source/Runtime/Renderer/Private/MeshPassProcessor.cpp
@@ -1468,3 +1468,73 @@ void FCachedPassMeshDrawListContext::FinalizeCommand(
 
 PassProcessorCreateFunction FPassProcessorManager::JumpTable[(int32)EShadingPath::Num][EMeshPass::Num] = {};
 EMeshPassFlags FPassProcessorManager::Flags[(int32)EShadingPath::Num][EMeshPass::Num] = {};
+
+#if defined(GPUCULL_TODO)
+
+FSimpleMeshDrawCommandPass::FSimpleMeshDrawCommandPass(const FSceneView& View, FInstanceCullingManager* InstanceCullingManager) :
+	bNeedsInitialization(false),
+	DynamicPassMeshDrawListContext(DynamicMeshDrawCommandStorage, VisibleMeshDrawCommands, GraphicsMinimalPipelineStateSet, bNeedsInitialization)
+{
+	check(View.bIsViewInfo);
+	const FViewInfo* ViewInfo = static_cast<const FViewInfo*>(&View);
+	InstanceCullingContext.ViewIds.Add(ViewInfo->GPUSceneViewId);
+	InstanceCullingContext.InstanceCullingManager = InstanceCullingManager;
+}
+
+void FSimpleMeshDrawCommandPass::BuildRenderingCommands(FRDGBuilder& GraphBuilder, const FSceneView& View, const FGPUScene& GPUScene, FInstanceCullingDrawParams& OutInstanceCullingDrawParams)
+{
+	// NOTE: Everything up to InstanceCullingContext.BuildRenderingCommands could be peeled off into an async task.
+	ApplyViewOverridesToMeshDrawCommands(View, VisibleMeshDrawCommands, DynamicMeshDrawCommandStorage, GraphicsMinimalPipelineStateSet, bNeedsInitialization);
+
+	VisibleMeshDrawCommands.Sort(FCompareFMeshDrawCommands());
+
+	FInstanceCullingResult InstanceCullingResult;
+	if (GPUScene.IsEnabled())
+	{
+		int32 MaxInstances = 0;
+		int32 VisibleMeshDrawCommandsNum = 0;
+		int32 NewPassVisibleMeshDrawCommandsNum = 0;
+
+		// 1. do the first thing first
+		SetupGPUInstancedDraws(InstanceCullingContext, VisibleMeshDrawCommands, MaxInstances, VisibleMeshDrawCommandsNum, NewPassVisibleMeshDrawCommandsNum);
+
+		// 2. Run finalize culling commands pass
+		InstanceCullingContext.BuildRenderingCommands(GraphBuilder, GPUScene, InstanceCullingResult);
+		
+		// Signal that scene primitives are supported, used for validation, the existence of a valid InstanceCullingResult is the required signal
+		bSupportsScenePrimitives = true;
+	}
+
+	InstanceCullingResult.GetDrawParameters(OutInstanceCullingDrawParams);
+}
+
+void FSimpleMeshDrawCommandPass::BuildRenderingCommands(FRDGBuilder& GraphBuilder, const FSceneView& View, const FScene& Scene, FInstanceCullingDrawParams& OutInstanceCullingDrawParams)
+{
+	BuildRenderingCommands(GraphBuilder, View, Scene.GPUScene, OutInstanceCullingDrawParams);
+}
+
+void FSimpleMeshDrawCommandPass::SubmitDraw(FRHICommandListImmediate& RHICmdList, const FInstanceCullingDrawParams& InstanceCullingDrawParams) const
+{
+	if (VisibleMeshDrawCommands.Num() > 0)
+	{
+		FRHIVertexBuffer* DrawIndirectArgsBuffer = nullptr;
+		FRHIVertexBuffer* InstanceIdOffsetBuffer = nullptr;
+
+		if (InstanceCullingDrawParams.DrawIndirectArgsBuffer != nullptr && InstanceCullingDrawParams.InstanceIdOffsetBuffer != nullptr)
+		{
+			DrawIndirectArgsBuffer = InstanceCullingDrawParams.DrawIndirectArgsBuffer->GetRHI();
+			InstanceIdOffsetBuffer = InstanceCullingDrawParams.InstanceIdOffsetBuffer->GetRHI();
+		}
+
+		SubmitGPUInstancedMeshDrawCommandsRange(
+			VisibleMeshDrawCommands,
+			GraphicsMinimalPipelineStateSet,
+			0,
+			VisibleMeshDrawCommands.Num(),
+			InstanceIdOffsetBuffer,
+			DrawIndirectArgsBuffer,
+			RHICmdList);
+	}
+}
+
+#endif // GPUCULL_TODO
diff --git a/Engine/Source/Runtime/Renderer/Private/PrimitiveSceneInfo.cpp b/Engine/Source/Runtime/Renderer/Private/PrimitiveSceneInfo.cpp
--- a/Engine/Source/Runtime/Renderer/Private/PrimitiveSceneInfo.cpp
+++ b/Engine/Source/Runtime/Renderer/Private/PrimitiveSceneInfo.cpp
@@ -305,6 +305,7 @@ void FPrimitiveSceneInfo::CacheMeshDrawCommands(FRHICommandListImmediate& RHICmd
 						check(!MeshRelevance.CommandInfosMask.Get(PassType));
 
 						uint64 BatchElementMask = ~0ull;
+						// NOTE: Modifies CommandInfo (through a reference), AddMeshBatch calls FCachedPassMeshDrawListContext::FinalizeCommand
 						PassMeshProcessor->AddMeshBatch(Mesh, BatchElementMask, SceneInfo->Proxy);
 
 						if (CommandInfo.CommandIndex != -1 || CommandInfo.StateBucketId != -1)
diff --git a/Engine/Source/Runtime/Renderer/Private/RendererModule.h b/Engine/Source/Runtime/Renderer/Private/RendererModule.h
--- a/Engine/Source/Runtime/Renderer/Private/RendererModule.h
+++ b/Engine/Source/Runtime/Renderer/Private/RendererModule.h
@@ -118,6 +118,10 @@ public:
 
 	void RenderPostResolvedSceneColorExtension(FRDGBuilder& GraphBuilder, const FSceneTextures& SceneTextures);
 
+#if defined(GPUCULL_TODO)
+	virtual IScenePrimitiveRenderingContext* BeginScenePrimitiveRendering(FRDGBuilder& GraphBuilder, FSceneViewFamily* ViewFamily) override;
+#endif
+
 private:
 	TSet<FSceneInterface*> AllocatedScenes;
 	FOnPostOpaqueRender PostOpaqueRenderDelegate;
diff --git a/Engine/Source/Runtime/Renderer/Private/SceneHitProxyRendering.cpp b/Engine/Source/Runtime/Renderer/Private/SceneHitProxyRendering.cpp
--- a/Engine/Source/Runtime/Renderer/Private/SceneHitProxyRendering.cpp
+++ b/Engine/Source/Runtime/Renderer/Private/SceneHitProxyRendering.cpp
@@ -265,6 +265,32 @@ BEGIN_SHADER_PARAMETER_STRUCT(FHitProxyCopyToViewFamilyParameters, )
 	RENDER_TARGET_BINDING_SLOTS()
 END_SHADER_PARAMETER_STRUCT()
 
+#if defined(GPUCULL_TODO)
+
+static void AddViewMeshElementsPass(const TIndirectArray<FMeshBatch> &MeshElements, FRDGBuilder& GraphBuilder, FHitProxyPassParameters* PassParameters, const FScene* Scene, const FViewInfo& View, const FMeshPassProcessorRenderState& DrawRenderState, FInstanceCullingManager& InstanceCullingManager)
+{
+	AddSimpleMeshPass(GraphBuilder, PassParameters, Scene, View, &InstanceCullingManager, RDG_EVENT_NAME("HitProxyMeshElementsPass"), View.ViewRect,
+		[&View, Scene, DrawRenderState, &MeshElements](FDynamicPassMeshDrawListContext* DynamicMeshPassContext)
+		{
+			FHitProxyMeshProcessor PassMeshProcessor(
+				Scene,
+				&View,
+				View.bAllowTranslucentPrimitivesInHitProxy,
+				DrawRenderState,
+				DynamicMeshPassContext);
+
+			const uint64 DefaultBatchElementMask = ~0ull;
+
+			for (const FMeshBatch& MeshBatch : MeshElements)
+			{
+				PassMeshProcessor.AddMeshBatch(MeshBatch, DefaultBatchElementMask, nullptr);
+			}
+		}
+	);
+}
+
+#endif // GPUCULL_TODO
+
 static void DoRenderHitProxies(
 	FRDGBuilder& GraphBuilder, 
 	const FSceneRenderer* SceneRenderer, 
diff --git a/Engine/Source/Runtime/Renderer/Private/SceneRendering.cpp b/Engine/Source/Runtime/Renderer/Private/SceneRendering.cpp
--- a/Engine/Source/Runtime/Renderer/Private/SceneRendering.cpp
+++ b/Engine/Source/Runtime/Renderer/Private/SceneRendering.cpp
@@ -3942,6 +3942,43 @@ void FRendererModule::RenderPostResolvedSceneColorExtension(FRDGBuilder& GraphBu
 	}
 }
 
+#if defined(GPUCULL_TODO)
+
+class FScenePrimitiveRenderingContext : public IScenePrimitiveRenderingContext
+{
+public:
+	FScenePrimitiveRenderingContext(FRDGBuilder& GraphBuilder, FScene& Scene) :
+		GPUScene(Scene.GPUScene),
+		GPUSceneDynamicContext(GPUScene)
+	{
+		Scene.UpdateAllPrimitiveSceneInfos(GraphBuilder, false);
+		GPUScene.BeginRender(&Scene, GPUSceneDynamicContext);
+		Scene.GPUScene.Update(GraphBuilder, Scene);
+	}
+
+	virtual ~FScenePrimitiveRenderingContext()
+	{
+		GPUScene.EndRender();
+	}
+
+	FGPUScene& GPUScene;
+	FGPUSceneDynamicContext GPUSceneDynamicContext;
+};
+
+
+IScenePrimitiveRenderingContext* FRendererModule::BeginScenePrimitiveRendering(FRDGBuilder &GraphBuilder, FSceneViewFamily* ViewFamily)
+{
+	check(ViewFamily->Scene);
+	FScene* Scene = ViewFamily->Scene->GetRenderScene();
+	check(Scene);
+
+	FScenePrimitiveRenderingContext* ScenePrimitiveRenderingContext = new FScenePrimitiveRenderingContext(GraphBuilder, *Scene);
+
+	return ScenePrimitiveRenderingContext;
+}
+
+#endif // GPUCULL_TODO
+
 IAllocatedVirtualTexture* FRendererModule::AllocateVirtualTexture(const FAllocatedVTDescription& Desc)
 {
 	return FVirtualTextureSystem::Get().AllocateVirtualTexture(Desc);
diff --git a/Engine/Source/Runtime/Renderer/Private/SceneVisibility.cpp b/Engine/Source/Runtime/Renderer/Private/SceneVisibility.cpp
--- a/Engine/Source/Runtime/Renderer/Private/SceneVisibility.cpp
+++ b/Engine/Source/Runtime/Renderer/Private/SceneVisibility.cpp
@@ -1861,6 +1861,7 @@ struct FDrawCommandRelevancePacket
 			const int32 StaticMeshCommandInfoIndex = StaticMeshRelevance.GetStaticMeshCommandInfoIndex(PassType);
 			if (StaticMeshCommandInfoIndex >= 0)
 			{
+				// GPUCULL_TODO: Here translates from FCachedMeshDrawCommandInfo -> FVisibleMeshDrawCommand
 				const FCachedMeshDrawCommandInfo& CachedMeshDrawCommand = InPrimitiveSceneInfo->StaticMeshCommandInfos[StaticMeshCommandInfoIndex];
 				const FCachedPassMeshDrawList& SceneDrawList = Scene->CachedDrawLists[PassType];
 
@@ -1872,6 +1873,8 @@ struct FDrawCommandRelevancePacket
 					? &Scene->CachedMeshDrawCommandStateBuckets[PassType].GetByElementId(CachedMeshDrawCommand.StateBucketId).Key
 					: &SceneDrawList.MeshDrawCommands[CachedMeshDrawCommand.CommandIndex];
 
+				// GPUCULL_TODO: We propagate the primitive ID because for static draws the GPU-Scene contains the list of instance IDs for each primitive
+				//           Thus we can re-process the list of primitives into the culled list of instances later.
 				NewVisibleMeshDrawCommand.Setup(
 					MeshDrawCommand,
 					PrimitiveIndex,
diff --git a/Engine/Source/Runtime/Renderer/Public/MeshPassProcessor.h b/Engine/Source/Runtime/Renderer/Public/MeshPassProcessor.h
--- a/Engine/Source/Runtime/Renderer/Public/MeshPassProcessor.h
+++ b/Engine/Source/Runtime/Renderer/Public/MeshPassProcessor.h
@@ -1817,6 +1817,15 @@ RENDERER_API extern void DrawDynamicMeshPassPrivate(
 	bool& InNeedsShaderInitialisation,
 	uint32 InstanceFactor);
 
+RENDERER_API extern void DrawDynamicMeshPassPrivate(
+	const FSceneView& View,
+	const FScene& Scene,
+	FRHICommandListImmediate& RHICmdList,
+	FMeshCommandOneFrameArray& VisibleMeshDrawCommands,
+	FDynamicMeshDrawCommandStorage& DynamicMeshDrawCommandStorage,
+	FGraphicsMinimalPipelineStateSet& GraphicsMinimalPipelineStateSet,
+	bool& InNeedsShaderInitialization);
+
 RENDERER_API extern FMeshDrawCommandSortKey CalculateMeshStaticSortKey(const FMeshMaterialShader* VertexShader, const FMeshMaterialShader* PixelShader);
 
 inline FMeshDrawCommandSortKey CalculateMeshStaticSortKey(const TShaderRef<FMeshMaterialShader>& VertexShader, const TShaderRef<FMeshMaterialShader>& PixelShader)
@@ -1938,3 +1947,190 @@ private:
 	uint32 GeometrySegmentIndex;
 	uint32 RayTracingInstanceIndex;
 };
+
+#if defined(GPUCULL_TODO)
+
+#if 0
+/**
+ * Same as the parallel one, but intended for use with simpler tasks where the overhead and complexity of parallel is not justified.
+ */
+class SimpleMeshDrawCommandPass
+{
+public:
+
+private:
+	FMeshCommandOneFrameArray VisibleMeshDrawCommands;
+	FInstanceCullingContext InstanceCullingContext;
+};
+
+#endif
+
+/**
+ * Same as the parallel one, but intended for use with simpler tasks where the overhead and complexity of parallel is not justified.
+ */
+class FSimpleMeshDrawCommandPass
+{
+public:
+	RENDERER_API FSimpleMeshDrawCommandPass(const FSceneView& View, FInstanceCullingManager* InstanceCullingManager);
+
+	/**
+	 * Run post-instance culling job to create the render commands and instance ID lists and optionally vertex instance data.
+	 * Needs to happen after DispatchPassSetup and before DispatchDraw, but not before global instance culling has been done.
+	 */
+	RENDERER_API void BuildRenderingCommands(FRDGBuilder& GraphBuilder, const FSceneView& View, const FGPUScene& GPUScene, FInstanceCullingDrawParams& OutInstanceCullingDrawParams);
+	RENDERER_API void BuildRenderingCommands(FRDGBuilder& GraphBuilder, const FSceneView& View, const FScene& Scene, FInstanceCullingDrawParams& OutInstanceCullingDrawParams);
+	RENDERER_API void SubmitDraw(FRHICommandListImmediate& RHICmdList, const FInstanceCullingDrawParams& InstanceCullingDrawParams) const;
+
+	FDynamicPassMeshDrawListContext* GetDynamicPassMeshDrawListContext() { return &DynamicPassMeshDrawListContext; }
+
+private:
+	bool bNeedsInitialization = false;
+	FDynamicPassMeshDrawListContext DynamicPassMeshDrawListContext;
+	FMeshCommandOneFrameArray VisibleMeshDrawCommands;
+	FInstanceCullingContext InstanceCullingContext;
+	FDynamicMeshDrawCommandStorage DynamicMeshDrawCommandStorage;
+	FGraphicsMinimalPipelineStateSet GraphicsMinimalPipelineStateSet;
+
+	// Is set to true if and only if the BuildRenderingCommands has been called with an enabled GPU scene (which implies a valid Scene etc).
+	// Is used to check that we don't submit any draw commands that require a GPU scene without supplying one.
+	bool bSupportsScenePrimitives = false;
+};
+
+// GPUCULL_TODO: Write documentation
+template <typename PassParametersType, typename AddMeshBatchesCallbackLambdaType, typename PassPrologueLambdaType>
+void AddSimpleMeshPass(FRDGBuilder& GraphBuilder, PassParametersType* PassParametersIn, const FScene* Scene, const FSceneView &View, FInstanceCullingManager* InstanceCullingManager, FRDGEventName&& PassName,
+	const ERDGPassFlags& PassFlags,
+	AddMeshBatchesCallbackLambdaType AddMeshBatchesCallback,
+	PassPrologueLambdaType PassPrologueCallback)
+{
+	// TODO: don't do this when the parameters are exclusive (optimization)
+	PassParametersType* PassParameters = GraphBuilder.AllocParameters<PassParametersType>();
+
+	// Copy parameters to ensure we can overwrite the InstanceCullingDrawParams
+	*PassParameters = *PassParametersIn;
+
+	FSimpleMeshDrawCommandPass* SimpleMeshDrawCommandPass = GraphBuilder.AllocObject<FSimpleMeshDrawCommandPass>(View, InstanceCullingManager);
+
+	AddMeshBatchesCallback(SimpleMeshDrawCommandPass->GetDynamicPassMeshDrawListContext());
+
+	// It is legal to render a simple mesh pass without a FScene, but then mesh draw commands must not reference primitive ID streams.
+	if (Scene != nullptr)
+	{
+		SimpleMeshDrawCommandPass->BuildRenderingCommands(GraphBuilder, View, *Scene, PassParameters->InstanceCullingDrawParams);
+	}
+
+	GraphBuilder.AddPass(
+		MoveTemp(PassName),
+		PassParameters,
+		PassFlags,
+		[SimpleMeshDrawCommandPass, PassParameters, PassPrologueCallback](FRHICommandListImmediate& RHICmdList)
+		{
+			PassPrologueCallback(RHICmdList);
+
+			SimpleMeshDrawCommandPass->SubmitDraw(RHICmdList, PassParameters->InstanceCullingDrawParams);
+		}
+	);
+}
+
+// GPUCULL_TODO: Write documentation
+template <typename PassParametersType, typename AddMeshBatchesCallbackLambdaType>
+void AddSimpleMeshPass(FRDGBuilder& GraphBuilder, PassParametersType* PassParameters, const FScene* Scene, const FSceneView& View, FInstanceCullingManager *InstanceCullingManager, FRDGEventName&& PassName, const FIntRect& ViewPortRect,
+	AddMeshBatchesCallbackLambdaType AddMeshBatchesCallback)
+{
+	AddSimpleMeshPass(GraphBuilder, PassParameters, Scene, View, InstanceCullingManager, MoveTemp(PassName), ERDGPassFlags::Raster, AddMeshBatchesCallback,
+		[ViewPortRect](FRHICommandListImmediate& RHICmdList)
+		{
+			RHICmdList.SetViewport(ViewPortRect.Min.X, ViewPortRect.Min.Y, 0.0f, ViewPortRect.Max.X, ViewPortRect.Max.Y, 1.0f);
+		}
+	);
+}
+
+// GPUCULL_TODO: Write documentation
+template <typename PassParametersType, typename AddMeshBatchesCallbackLambdaType>
+void AddSimpleMeshPass(FRDGBuilder& GraphBuilder, PassParametersType* PassParameters, const FScene* Scene, const FSceneView& View, FInstanceCullingManager *InstanceCullingManager, FRDGEventName&& PassName, 
+	const FIntRect& ViewPortRect, 
+	const ERDGPassFlags &PassFlags,
+	AddMeshBatchesCallbackLambdaType AddMeshBatchesCallback)
+{
+	AddSimpleMeshPass(GraphBuilder, PassParameters, Scene, View, InstanceCullingManager, MoveTemp(PassName), PassFlags, AddMeshBatchesCallback,
+		[ViewPortRect](FRHICommandListImmediate& RHICmdList)
+		{
+			RHICmdList.SetViewport(ViewPortRect.Min.X, ViewPortRect.Min.Y, 0.0f, ViewPortRect.Max.X, ViewPortRect.Max.Y, 1.0f);
+		}
+	);
+}
+
+// GPUCULL_TODO: Write documentation
+template <typename PassParametersType, typename AddMeshBatchesCallbackLambdaType, typename PassPrologueLambdaType>
+void AddSimpleMeshPass(FRDGBuilder& GraphBuilder, PassParametersType* PassParametersIn, const FGPUScene& GPUScene, const FSceneView& View, FInstanceCullingManager* InstanceCullingManager, FRDGEventName&& PassName,
+	const ERDGPassFlags& PassFlags,
+	AddMeshBatchesCallbackLambdaType AddMeshBatchesCallback,
+	PassPrologueLambdaType PassPrologueCallback)
+{
+	// TODO: don't do this when the parameters are exclusive (optimization)
+	PassParametersType* PassParameters = GraphBuilder.AllocParameters<PassParametersType>();
+
+	// Copy parameters to ensure we can overwrite the InstanceCullingDrawParams
+	*PassParameters = *PassParametersIn;
+
+	FSimpleMeshDrawCommandPass* SimpleMeshDrawCommandPass = GraphBuilder.AllocObject<FSimpleMeshDrawCommandPass>(View, InstanceCullingManager);
+
+	AddMeshBatchesCallback(SimpleMeshDrawCommandPass->GetDynamicPassMeshDrawListContext());
+
+	SimpleMeshDrawCommandPass->BuildRenderingCommands(GraphBuilder, View, GPUScene, PassParameters->InstanceCullingDrawParams);
+
+	GraphBuilder.AddPass(
+		MoveTemp(PassName),
+		PassParameters,
+		PassFlags,
+		[SimpleMeshDrawCommandPass, PassParameters, PassPrologueCallback](FRHICommandListImmediate& RHICmdList)
+		{
+			PassPrologueCallback(RHICmdList);
+
+			SimpleMeshDrawCommandPass->SubmitDraw(RHICmdList, PassParameters->InstanceCullingDrawParams);
+		}
+	);
+}
+
+// GPUCULL_TODO: Write documentation
+template <typename PassParametersType, typename AddMeshBatchesCallbackLambdaType>
+void AddSimpleMeshPass(FRDGBuilder& GraphBuilder, PassParametersType* PassParameters, const FGPUScene& GPUScene, const FSceneView& View, FInstanceCullingManager* InstanceCullingManager, FRDGEventName&& PassName, const FIntRect& ViewPortRect,
+	AddMeshBatchesCallbackLambdaType AddMeshBatchesCallback)
+{
+	AddSimpleMeshPass(GraphBuilder, PassParameters, GPUScene, View, InstanceCullingManager, MoveTemp(PassName), ERDGPassFlags::Raster, AddMeshBatchesCallback,
+		[ViewPortRect](FRHICommandListImmediate& RHICmdList)
+		{
+			RHICmdList.SetViewport(ViewPortRect.Min.X, ViewPortRect.Min.Y, 0.0f, ViewPortRect.Max.X, ViewPortRect.Max.Y, 1.0f);
+		}
+	);
+}
+
+// GPUCULL_TODO: Write documentation
+template <typename PassParametersType, typename AddMeshBatchesCallbackLambdaType>
+void AddSimpleMeshPass(FRDGBuilder& GraphBuilder, PassParametersType* PassParameters, const FGPUScene& GPUScene, const FSceneView& View, FInstanceCullingManager* InstanceCullingManager, FRDGEventName&& PassName,
+	const FIntRect& ViewPortRect,
+	const ERDGPassFlags& PassFlags,
+	AddMeshBatchesCallbackLambdaType AddMeshBatchesCallback)
+{
+	AddSimpleMeshPass(GraphBuilder, PassParameters, GPUScene, View, InstanceCullingManager, MoveTemp(PassName), PassFlags, AddMeshBatchesCallback,
+		[ViewPortRect](FRHICommandListImmediate& RHICmdList)
+		{
+			RHICmdList.SetViewport(ViewPortRect.Min.X, ViewPortRect.Min.Y, 0.0f, ViewPortRect.Max.X, ViewPortRect.Max.Y, 1.0f);
+		}
+	);
+}
+
+#if 0
+
+// GPUCULL_TODO: Write documentation
+template <typename PassParametersType, typename AddMeshBatchesCallbackLambdaType, typename PassPrologueLambdaType>
+void AddSimpleMeshPass(FRDGBuilder& GraphBuilder, PassParametersType* PassParameters, const FScene* Scene, const FSceneView& View, FRDGEventName&& PassName,
+	AddMeshBatchesCallbackLambdaType AddMeshBatchesCallback)
+{
+	// TODO: ViewRect is only in FViewInfo
+	AddSimpleMeshPass(GraphBuilder, PassParameters, Scene, View, MoveTemp(PassName), View.ViewRect, AddMeshBatchesCallback);
+}
+
+#endif
+
+#endif // GPUCULL_TODO
diff --git a/Engine/Source/Runtime/Renderer/Public/MeshPassProcessor.inl b/Engine/Source/Runtime/Renderer/Public/MeshPassProcessor.inl
--- a/Engine/Source/Runtime/Renderer/Public/MeshPassProcessor.inl
+++ b/Engine/Source/Runtime/Renderer/Public/MeshPassProcessor.inl
@@ -177,3 +177,25 @@ void DrawDynamicMeshPass(const FSceneView& View, FRHICommandList& RHICmdList, co
 	const uint32 InstanceFactor = (!bForceStereoInstancingOff && View.IsInstancedStereoPass()) ? 2 : 1;
 	DrawDynamicMeshPassPrivate(View, RHICmdList, VisibleMeshDrawCommands, DynamicMeshDrawCommandStorage, GraphicsMinimalPipelineStateSet, NeedsShaderInitialisation, InstanceFactor);
 }
+
+#if defined(GPUCULL_TODO)
+
+/**
+* Overload that requires FViewInfo, need to be used when rendering draw commands that make use of Scene features (e.g., require instancing etc).
+*/
+template<typename LambdaType>
+void DrawDynamicMeshPass(const FSceneView& View, const FScene &Scene, FRHICommandListImmediate& RHICmdList, const LambdaType& BuildPassProcessorLambda, bool bForceStereoInstancingOff = false)
+{
+	FDynamicMeshDrawCommandStorage DynamicMeshDrawCommandStorage;
+	FMeshCommandOneFrameArray VisibleMeshDrawCommands;
+	FGraphicsMinimalPipelineStateSet GraphicsMinimalPipelineStateSet;
+	bool NeedsShaderInitialization;
+
+	FDynamicPassMeshDrawListContext DynamicMeshPassContext(DynamicMeshDrawCommandStorage, VisibleMeshDrawCommands, GraphicsMinimalPipelineStateSet, NeedsShaderInitialization);
+
+	BuildPassProcessorLambda(&DynamicMeshPassContext);
+
+	DrawDynamicMeshPassPrivate(View, Scene, RHICmdList, VisibleMeshDrawCommands, DynamicMeshDrawCommandStorage, GraphicsMinimalPipelineStateSet, NeedsShaderInitialization);
+}
+
+#endif // GPUCULL_TODO
