Fold GPU-Scene uploads into one funciton such that dynamic primitives can have instances also & more instance culling prep.
- Moved the primitive instance registration such that we don't end up with added primitives that have no instance index at GPU-Scene update time
- added calls to InstanceCullingManager and disabled implementation.
- GPU scene can now exists without a FScene  instance

#rnx

diff --git a/Engine/Source/Runtime/Renderer/Private/DeferredShadingRenderer.cpp b/Engine/Source/Runtime/Renderer/Private/DeferredShadingRenderer.cpp
--- a/Engine/Source/Runtime/Renderer/Private/DeferredShadingRenderer.cpp
+++ b/Engine/Source/Runtime/Renderer/Private/DeferredShadingRenderer.cpp
@@ -1513,7 +1513,7 @@ void FDeferredShadingSceneRenderer::Render(FRDGBuilder& GraphBuilder)
 
 	Scene->UpdateAllPrimitiveSceneInfos(GraphBuilder, true);
 
-	FGPUSceneScopeBeginEndHelper GPUSceneScopeBeginEndHelper(Scene->GPUScene, GPUSceneDynamicContext, *Scene);
+	FGPUSceneScopeBeginEndHelper GPUSceneScopeBeginEndHelper(Scene->GPUScene, GPUSceneDynamicContext, Scene);
 	if (bNaniteEnabled)
 	{
 		Nanite::GGlobalResources.Update(GraphBuilder); // Needed to managed scratch buffers for Nanite.
@@ -1693,7 +1693,12 @@ void FDeferredShadingSceneRenderer::Render(FRDGBuilder& GraphBuilder)
 			FViewInfo& View = Views[ViewIndex];
 			RDG_GPU_MASK_SCOPE(GraphBuilder, View.GPUMask);
 
-			Scene->GPUScene.UploadDynamicPrimitiveShaderDataForView(GraphBuilder.RHICmdList, *Scene, View);
+			Scene->GPUScene.UploadDynamicPrimitiveShaderDataForView(GraphBuilder.RHICmdList, Scene, View);
+		}
+
+		{
+			// GPUCULL_TODO: Possibly fold into unpack step
+			InstanceCullingManager.CullInstances(GraphBuilder, Scene->GPUScene);
 		}
 
 		if (!bDoInitViewAftersPrepass)
@@ -2072,11 +2077,8 @@ void FDeferredShadingSceneRenderer::Render(FRDGBuilder& GraphBuilder)
 	// Early Shadow depth rendering
 	if (bOcclusionBeforeBasePass)
 	{
-		if (bNaniteEnabled)
-		{
-			const bool bAfterBasePass = false;
-			AllocateVirtualShadowMaps(bAfterBasePass);
-		}
+		const bool bAfterBasePass = false;
+		AllocateVirtualShadowMaps(bAfterBasePass);
 
 		RenderShadowDepthMaps(GraphBuilder, InstanceCullingManager);
 		AddServiceLocalQueuePass(GraphBuilder);
@@ -2301,11 +2303,8 @@ void FDeferredShadingSceneRenderer::Render(FRDGBuilder& GraphBuilder)
 	// Shadow and fog after base pass
 	if (!bOcclusionBeforeBasePass)
 	{
-		if (bNaniteEnabled)
-		{
-			const bool bAfterBasePass = true;
-			AllocateVirtualShadowMaps(bAfterBasePass);
-		}
+		const bool bAfterBasePass = true;
+		AllocateVirtualShadowMaps(bAfterBasePass);
 
 		RenderShadowDepthMaps(GraphBuilder, InstanceCullingManager);
 		
diff --git a/Engine/Source/Runtime/Renderer/Private/GPUScene.cpp b/Engine/Source/Runtime/Renderer/Private/GPUScene.cpp
--- a/Engine/Source/Runtime/Renderer/Private/GPUScene.cpp
+++ b/Engine/Source/Runtime/Renderer/Private/GPUScene.cpp
@@ -50,17 +50,6 @@ FAutoConsoleVariableRef CVarGPUSceneParallelUpdate(
 	ECVF_RenderThreadSafe
 );
 
-// Always a benefit unless the game is doing tons of add/remove instance
-// calls, which isn't advisable anyways.
-int32 GGPUSceneInstanceClearList = 1;
-FAutoConsoleVariableRef CVarGPUSceneInstanceClearList(
-	TEXT("r.GPUScene.InstanceClearList"),
-	GGPUSceneInstanceClearList,
-	TEXT("Whether to use instance clear indirection list."),
-	ECVF_RenderThreadSafe
-);
-
-
 template<typename ResourceType>
 ResourceType* GetMirrorGPU(FGPUScene& GPUScene);
 
@@ -77,18 +66,23 @@ FTextureRWBuffer2D* GetMirrorGPU<FTextureRWBuffer2D>(FGPUScene& GPUScene)
 	return &GPUScene.PrimitiveTexture;
 }
 
+/**
+ * Returns a pointer to the underlying storage (either a texture or buffer), 
+ * Only implemented for Buffer at present, since the texture data layout is opaque and complex to iterate.
+ * Only used for debugging at present.
+ */
 template<typename ResourceType>
-void* LockResource(ResourceType& Resource, uint32& Stride);
+void* LockResource(ResourceType& Resource);
 template<>
-void* LockResource<FRWBufferStructured>(FRWBufferStructured& Resource, uint32& Stride)
+void* LockResource<FRWBufferStructured>(FRWBufferStructured& Resource)
 {
-	Stride = 0;
 	return RHILockBuffer(Resource.Buffer, 0, Resource.NumBytes, RLM_ReadOnly);
 }
 template<>
-void* LockResource<FTextureRWBuffer2D>(FTextureRWBuffer2D& Resource, uint32& Stride)
+void* LockResource<FTextureRWBuffer2D>(FTextureRWBuffer2D& Resource)
 {
-	return RHILockTexture2D(Resource.Buffer, 0, RLM_ReadOnly, Stride, false);
+	// RHILockTexture2D(Resource.Buffer, 0, RLM_ReadOnly, RowStrideInBytesOut, false);
+	return nullptr;
 }
 
 template<typename ResourceType>
@@ -101,7 +95,7 @@ void UnlockResourceGPUScene<FRWBufferStructured>(FRWBufferStructured& Resource)
 template<>
 void UnlockResourceGPUScene<FTextureRWBuffer2D>(FTextureRWBuffer2D& Resource)
 {
-	RHIUnlockTexture2D(Resource.Buffer, 0, false);
+	// RHIUnlockTexture2D(Resource.Buffer, 0, false);
 }
 
 template<typename ResourceType>
@@ -177,17 +171,228 @@ FGPUScenePrimitiveCollector::FUploadData* FGPUScenePrimitiveCollector::AllocateU
 }
 
 
+/**
+ * Info needed by the uploader to update a primitive.
+ */
+struct FPrimitiveUploadInfo
+{
+	/** Required */
+	FPrimitiveSceneShaderData PrimitiveSceneData;
+	int32 PrimitiveID = INDEX_NONE;
+
+	/** Optional */
+	int32 InstanceDataOffset = INDEX_NONE;
+	int32 InstanceUploadCount = 0;
+	int32 LightmapUploadCount = 0;
+
+	/** NaniteSceneProxy must be set if the proxy is a Nanite proxy */
+	const Nanite::FSceneProxyBase* NaniteSceneProxy = nullptr;
+	const FPrimitiveSceneInfo* PrimitiveSceneInfo = nullptr;
+};
+
+/**
+ * Info required by the uploader to update the instances that belong to a primitive.
+ */
+struct FInstanceUploadInfo
+{
+	/**
+	 * Transform of the primitive from local to world (not the instances).
+	 */
+	FMatrix PrimitiveLocalToWorld;
+	FMatrix PreviousPrimitiveLocalToWorld;
+	TArrayView<FPrimitiveInstance> PrimitiveInstances;
+	int32 InstanceDataOffset = INDEX_NONE;
+
+	int32 PrimitiveID = INDEX_NONE;
+
+	// Used for primitives that need to create a dummy instance (they do not have instance data in the proxy)
+	FPrimitiveInstance DummyInstance;
+};
+
+/**
+ * Info required by the uploader to update the lightmap data for a primitive.
+ */
+struct FLightMapUploadInfo
+{
+	FPrimitiveSceneProxy::FLCIArray LCIs;
+	int32 LightmapDataOffset = 0;
+};
+
+/**
+ * Implements a thin data abstraction such that the UploadGeneral function can upload primitive data from
+ * both scene primitives and dynamic primitives (which are not stored in the same way). 
+ * Note: handling of Nanite material table upload data is not abstracted (since at present it can only come via the scene primitives).
+ */
+struct FUploadDataSourceAdapterScenePrimitives
+{
+	static constexpr bool bUpdateNaniteMaterialTables = true;
+
+	FUploadDataSourceAdapterScenePrimitives(FGPUScene& InGPUScene, FScene& InScene) : GPUScene(InGPUScene), Scene(InScene) {}
+
+	/**
+	 * Return the number of primitives to upload N, GetPrimitiveInfo will be called with ItemIndex in [0,N).
+	 */
+	FORCEINLINE int32 NumPrimitivesToUpload() const 
+	{ 
+		return GPUScene.PrimitivesToUpdate.Num(); 
+	}
+
+	/**
+	 * Populate the primitive info for a given item index.
+	 * 
+	 */
+	FORCEINLINE bool GetPrimitiveInfo(int32 ItemIndex, FPrimitiveUploadInfo& PrimitiveUploadInfo) const
+	{
+		int32 PrimitiveID = GPUScene.PrimitivesToUpdate[ItemIndex];
+		if (PrimitiveID < Scene.PrimitiveSceneProxies.Num())
+		{
+			const FPrimitiveSceneProxy* RESTRICT PrimitiveSceneProxy = Scene.PrimitiveSceneProxies[PrimitiveID];
+			const FPrimitiveSceneInfo* PrimitiveSceneInfo = PrimitiveSceneProxy->GetPrimitiveSceneInfo();
+
+			PrimitiveUploadInfo.PrimitiveID = PrimitiveID;;
+			PrimitiveUploadInfo.InstanceDataOffset = INDEX_NONE;;
+			PrimitiveUploadInfo.InstanceUploadCount = 0;
+			PrimitiveUploadInfo.LightmapUploadCount = PrimitiveSceneInfo->GetNumLightmapDataEntries();
+			PrimitiveUploadInfo.NaniteSceneProxy = PrimitiveSceneProxy->IsNaniteMesh() ? static_cast<const Nanite::FSceneProxyBase*>(PrimitiveSceneProxy) : nullptr;
+			PrimitiveUploadInfo.PrimitiveSceneInfo = PrimitiveSceneInfo;
+			// Count all primitive instances represented in the instance data buffer.
+			if (PrimitiveSceneProxy->SupportsInstanceDataBuffer())
+			{
+				const TArray<FPrimitiveInstance>* PrimitiveInstances = PrimitiveSceneProxy->GetPrimitiveInstances();
+				PrimitiveUploadInfo.InstanceDataOffset = PrimitiveSceneProxy->GetPrimitiveSceneInfo()->GetInstanceDataOffset();
+				PrimitiveUploadInfo.InstanceUploadCount = PrimitiveInstances->Num();
+			}
+#if defined(GPUCULL_TODO)
+			else
+			{
+				PrimitiveUploadInfo.InstanceDataOffset = PrimitiveSceneProxy->GetPrimitiveSceneInfo()->GetInstanceDataOffset();
+				PrimitiveUploadInfo.InstanceUploadCount = 1;
+			}
+#endif //defined(GPUCULL_TODO)
+			PrimitiveUploadInfo.PrimitiveSceneData = FPrimitiveSceneShaderData(PrimitiveSceneProxy);
+
+			return true;
+		}
+
+		return false;
+	}
+
+	FORCEINLINE bool GetInstanceInfo(int32 ItemIndex, FInstanceUploadInfo& InstanceUploadInfo)
+	{
+		const int32 PrimitiveID = GPUScene.PrimitivesToUpdate[ItemIndex];
+		if (PrimitiveID < Scene.PrimitiveSceneProxies.Num())
+		{
+			InstanceUploadInfo.PrimitiveID = PrimitiveID;
+			FPrimitiveSceneProxy* PrimitiveSceneProxy = Scene.PrimitiveSceneProxies[PrimitiveID];
+			const FPrimitiveSceneInfo* PrimitiveSceneInfo = PrimitiveSceneProxy->GetPrimitiveSceneInfo();
+
+			if (PrimitiveSceneProxy->SupportsInstanceDataBuffer())
+			{
+				TArray<FPrimitiveInstance>* PrimitiveInstancesPtr = PrimitiveSceneProxy->GetPrimitiveInstances();
+				check(PrimitiveInstancesPtr);
+
+				InstanceUploadInfo.PrimitiveInstances = TArrayView<FPrimitiveInstance>(*PrimitiveInstancesPtr);
+			}
+#if defined(GPUCULL_TODO)
+			else
+			{
+				// We always create an instance to ensure that we can always use the same code paths in the shader
+				// In the future we should remove redundant data from the primitive, and then the instances should be
+				// provided by the proxy. However, this is a lot of work before we can just enable it in the base proxy class.
+				FPrimitiveInstance &DummyInstance = InstanceUploadInfo.DummyInstance;
+				if (!PrimitiveSceneProxy->SupportsInstanceDataBuffer())
+				{
+					DummyInstance.InstanceToLocal = FMatrix::Identity;
+					DummyInstance.LocalToInstance = FMatrix::Identity;
+					DummyInstance.LocalToWorld = FMatrix::Identity;
+					DummyInstance.PrevLocalToWorld = FMatrix::Identity;
+					DummyInstance.NonUniformScale = FVector4(1.0f, 1.0f, 1.0f, 1.0f);
+					DummyInstance.InvNonUniformScaleAndDeterminantSign = FVector4(1.0f, 1.0f, 1.0f, 1.0f);
+					DummyInstance.RenderBounds = PrimitiveSceneProxy->GetLocalBounds();
+					DummyInstance.LocalBounds = DummyInstance.RenderBounds;;
+					DummyInstance.PrimitiveId = 0xFFFFFFFF;
+					DummyInstance.LastUpdateSceneFrameNumber = 0xFFFFFFFF;
+				}
+
+				InstanceUploadInfo.PrimitiveInstances = TArrayView<FPrimitiveInstance>(&DummyInstance, 1);
+			}
+#endif // defined(GPUCULL_TODO)
+
+			check(InstanceUploadInfo.PrimitiveInstances.Num() == PrimitiveSceneInfo->GetNumInstanceDataEntries());
+			if (InstanceUploadInfo.PrimitiveInstances.Num() == 0)
+			{
+				return false;
+			}
+
+			InstanceUploadInfo.PrimitiveLocalToWorld = PrimitiveSceneProxy->GetLocalToWorld();
+			
+			InstanceUploadInfo.InstanceDataOffset = PrimitiveSceneProxy->GetPrimitiveSceneInfo()->GetInstanceDataOffset();
+			check(InstanceUploadInfo.InstanceDataOffset != INDEX_NONE);
+			check(PrimitiveSceneInfo->GetIndex() == PrimitiveID);
+
+			// Dummy data that we don't actually need
+			bool bHasPrecomputedVolumetricLightmap = false;
+			int32 SingleCaptureIndex = 0;
+			bool bOutputVelocity = false;
+
+			Scene.GetPrimitiveUniformShaderParameters_RenderThread(
+				PrimitiveSceneInfo,
+				bHasPrecomputedVolumetricLightmap,
+				InstanceUploadInfo.PreviousPrimitiveLocalToWorld,
+				SingleCaptureIndex,
+				bOutputVelocity
+			);
+
+			return true;
+		}
+
+		return false;
+	}
+
+	FORCEINLINE bool GetLightMapInfo(int32 ItemIndex, FLightMapUploadInfo &UploadInfo) const
+	{
+		const int32 PrimitiveID = GPUScene.PrimitivesToUpdate[ItemIndex];
+		if (PrimitiveID < Scene.PrimitiveSceneProxies.Num())
+		{
+			FPrimitiveSceneProxy* PrimitiveSceneProxy = Scene.PrimitiveSceneProxies[PrimitiveID];
+
+			PrimitiveSceneProxy->GetLCIs(UploadInfo.LCIs);
+			check(UploadInfo.LCIs.Num() == PrimitiveSceneProxy->GetPrimitiveSceneInfo()->GetNumLightmapDataEntries());
+			UploadInfo.LightmapDataOffset = PrimitiveSceneProxy->GetPrimitiveSceneInfo()->GetLightmapDataOffset();
+			return true;
+		}
+
+		return false;
+	}
+
+	FGPUScene& GPUScene;
+	FScene& Scene;
+};
+
+void FGPUScene::SetEnabled(ERHIFeatureLevel::Type InFeatureLevel)
+{
+	FeatureLevel = InFeatureLevel;
+	bIsEnabled = UseGPUScene(GMaxRHIShaderPlatform, FeatureLevel);
+}
+
 FGPUScene::~FGPUScene()
 {
 }
 
-void FGPUScene::BeginRender(FScene& Scene, FGPUSceneDynamicContext &GPUSceneDynamicContext)
+void FGPUScene::BeginRender(const FScene* Scene, FGPUSceneDynamicContext &GPUSceneDynamicContext)
 {
 	ensure(!bInBeginEndBlock);
 	ensure(CurrentDynamicContext == nullptr);
-	ensure(bIsEnabled == UseGPUScene(GMaxRHIShaderPlatform, Scene.GetFeatureLevel()));
+	if (Scene != nullptr)
+	{
+		ensure(bIsEnabled == UseGPUScene(GMaxRHIShaderPlatform, Scene->GetFeatureLevel()));
+		NumScenePrimitives = Scene->Primitives.Num();
+	}
+	else
+	{
+		NumScenePrimitives = 0;
+	}
 	CurrentDynamicContext = &GPUSceneDynamicContext;
-	NumScenePrimitives = Scene.Primitives.Num();
 	DynamicPrimitivesOffset = NumScenePrimitives;
 	bInBeginEndBlock = true;
 }
@@ -214,17 +419,6 @@ void FGPUScene::UpdateInternal(FRHICommandListImmediate& RHICmdList, FScene& Sce
 		QUICK_SCOPE_CYCLE_COUNTER(STAT_UpdateGPUScene);
 		SCOPE_CYCLE_COUNTER(STAT_UpdateGPUSceneTime);
 
-		// Multi-GPU support : Updating on all GPUs is inefficient for AFR. Work is wasted
-		// for any primitives that update on consecutive frames.
-		SCOPED_GPU_MASK(RHICmdList, FRHIGPUMask::All());
-
-		// Store in GPU-scene to enable validation that update has been carried out.
-		SceneFrameNumber = Scene.GetFrameNumber();
-
-		const bool bNaniteEnabled = DoesPlatformSupportNanite(GMaxRHIShaderPlatform);
-
-		const bool bExecuteInParallel = GGPUSceneParallelUpdate != 0 && FApp::ShouldUseThreadingForPerformance();
-
 		if (GGPUSceneUploadEveryFrame || bUpdateAllPrimitives)
 		{
 			PrimitivesMarkedToUpdate.Init(false, PrimitivesMarkedToUpdate.Num());
@@ -239,47 +433,107 @@ void FGPUScene::UpdateInternal(FRHICommandListImmediate& RHICmdList, FScene& Sce
 			InstanceDataToClear.Init(true, InstanceDataToClear.Num());
 
 			// Set entire instance range for possible clearing.
-			if (GGPUSceneInstanceClearList != 0)
+			for (int32 Index = 0; Index < InstanceDataToClear.Num(); ++Index)
 			{
-				for (int32 Index = 0; Index < InstanceDataToClear.Num(); ++Index)
-				{
-					InstanceClearList.Add(Index);
-				}
+				InstanceClearList.Add(Index);
 			}
 
 			bUpdateAllPrimitives = false;
 		}
 
-		bool bResizedPrimitiveData = false;
-		bool bResizedInstanceData = false;
-		bool bResizedLightmapData = false;
+		FUploadDataSourceAdapterScenePrimitives Adapter(*this, Scene);
+		UploadGeneral<ResourceType, FUploadDataSourceAdapterScenePrimitives>(RHICmdList, &Scene, Adapter);
 
-		ResourceType* MirrorResourceGPU = GetMirrorGPU<ResourceType>(*this);
+		PrimitivesMarkedToUpdate.Init(false, PrimitivesMarkedToUpdate.Num());
+
+
+#if DO_CHECK
+		// Validate the scene primitives are identical to the uploaded data (not the dynamic ones).
+		if (GGPUSceneValidatePrimitiveBuffer && (PrimitiveBuffer.NumBytes > 0 || PrimitiveTexture.NumBytes > 0))
 		{
-			const uint32 SizeReserve = FMath::RoundUpToPowerOfTwo( FMath::Max(DynamicPrimitivesOffset, 256 ) );
-			bResizedPrimitiveData = ResizeResourceIfNeeded(RHICmdList, *MirrorResourceGPU, SizeReserve * sizeof(FPrimitiveSceneShaderData::Data), TEXT("PrimitiveData"));
+			ResourceType* MirrorResourceGPU = GetMirrorGPU<ResourceType>(*this);
+			//UE_LOG(LogRenderer, Warning, TEXT("r.GPUSceneValidatePrimitiveBuffer enabled, doing slow readback from GPU"));
+			const FPrimitiveSceneShaderData* PrimitiveBufferPtr = reinterpret_cast<const FPrimitiveSceneShaderData*>(LockResource(*MirrorResourceGPU));
+			ensureMsgf(PrimitiveBufferPtr != nullptr, TEXT("Validation not implemented for Texture2D buffer type as layout is opaque."));
+
+			if (PrimitiveBufferPtr != nullptr)
+			{
+				const int32 TotalNumberPrimitives = Scene.PrimitiveSceneProxies.Num();
+				check(MirrorResourceGPU->NumBytes >= TotalNumberPrimitives * sizeof(FPrimitiveSceneShaderData));
+
+				int32 MaxPrimitivesUploads = GetMaxPrimitivesUpdate(TotalNumberPrimitives, FPrimitiveSceneShaderData::PrimitiveDataStrideInFloat4s);
+				for (int32 IndexOffset = 0; IndexOffset < TotalNumberPrimitives; IndexOffset += MaxPrimitivesUploads)
+				{
+					for (int32 Index = 0; (Index < MaxPrimitivesUploads) && ((Index + IndexOffset) < TotalNumberPrimitives); ++Index)
+					{
+						FPrimitiveSceneShaderData PrimitiveSceneData(Scene.PrimitiveSceneProxies[Index + IndexOffset]);
+						const FPrimitiveSceneShaderData& Item = PrimitiveBufferPtr[Index + IndexOffset];
+						for (int32 i = 0; i < FPrimitiveSceneShaderData::PrimitiveDataStrideInFloat4s; i++)
+						{
+							check(PrimitiveSceneData.Data[i] == Item.Data[i]);
+						}
+					}
+				}
+			}
+			UnlockResourceGPUScene(*MirrorResourceGPU);
 		}
+#endif // DO_CHECK
+
+		PrimitivesToUpdate.Reset();
+
+		// Clear the flags that mark newly added primitives.
+		AddedPrimitiveFlags.Init(false, AddedPrimitiveFlags.Num());
+
+		// Dynamic primitives are allocated after the regular ones, this needs to match the value passed in in BeginRender
+		ensure(DynamicPrimitivesOffset >= Scene.Primitives.Num());
+	}
+
+	checkSlow(PrimitivesToUpdate.Num() == 0);
+}
+
+template<typename ResourceType, typename FUploadDataSourceAdapter>
+void FGPUScene::UploadGeneral(FRHICommandListImmediate& RHICmdList, FScene *Scene, FUploadDataSourceAdapter &UploadDataSourceAdapter)
+{
+	ensure(bInBeginEndBlock);
+	if (Scene != nullptr)
+	{
+		ensure(bIsEnabled == UseGPUScene(GMaxRHIShaderPlatform, Scene->GetFeatureLevel()));
+		ensure(NumScenePrimitives == Scene->Primitives.Num());
+	}
+	{
+		// Multi-GPU support : Updating on all GPUs is inefficient for AFR. Work is wasted
+		// for any primitives that update on consecutive frames.
+		SCOPED_GPU_MASK(RHICmdList, FRHIGPUMask::All());
+
+		const bool bExecuteInParallel = GGPUSceneParallelUpdate != 0 && FApp::ShouldUseThreadingForPerformance();
+		const bool bNaniteEnabled = DoesPlatformSupportNanite(GMaxRHIShaderPlatform);
+
+		ResourceType* MirrorResourceGPU = GetMirrorGPU<ResourceType>(*this);
+
+		const uint32 SizeReserve = FMath::RoundUpToPowerOfTwo(FMath::Max(DynamicPrimitivesOffset, 256));
+		const bool bResizedPrimitiveData = ResizeResourceIfNeeded(RHICmdList, *MirrorResourceGPU, SizeReserve * sizeof(FPrimitiveSceneShaderData::Data), TEXT("PrimitiveData"));
 
 		const uint32 InstanceDataNumArrays = FInstanceSceneShaderData::InstanceDataStrideInFloat4s;
 		const uint32 InstanceDataSizeReserve = FMath::RoundUpToPowerOfTwo(FMath::Max(InstanceDataAllocator.GetMaxSize(), 256));
-		bResizedInstanceData = ResizeResourceSOAIfNeeded(RHICmdList, InstanceDataBuffer, InstanceDataSizeReserve * sizeof(FInstanceSceneShaderData::Data), InstanceDataNumArrays, TEXT("InstanceData"));
+		const bool bResizedInstanceData = ResizeResourceSOAIfNeeded(RHICmdList, InstanceDataBuffer, InstanceDataSizeReserve * sizeof(FInstanceSceneShaderData::Data), InstanceDataNumArrays, TEXT("InstanceData"));
 		InstanceDataSOAStride = InstanceDataSizeReserve;
-		
-		{
-			const uint32 SizeReserve = FMath::RoundUpToPowerOfTwo( FMath::Max( LightmapDataAllocator.GetMaxSize(), 256 ) );
-			bResizedLightmapData = ResizeResourceIfNeeded(RHICmdList, LightmapDataBuffer, SizeReserve * sizeof(FLightmapSceneShaderData::Data), TEXT("LightmapData"));
-		}
+
+		const uint32 LightMapDataBufferSize = FMath::RoundUpToPowerOfTwo(FMath::Max(LightmapDataAllocator.GetMaxSize(), 256));
+		const bool bResizedLightmapData = ResizeResourceIfNeeded(RHICmdList, LightmapDataBuffer, LightMapDataBufferSize * sizeof(FLightmapSceneShaderData::Data), TEXT("LightmapData"));
 
 		// These should always be in sync with each other.
 		check(InstanceDataToClear.Num() == InstanceDataAllocator.GetMaxSize());
 
-		const int32 NumPrimitiveDataUploads = PrimitivesToUpdate.Num();
+		const int32 NumPrimitiveDataUploads = UploadDataSourceAdapter.NumPrimitivesToUpload();
 
-		if (bNaniteEnabled)
+		if (Scene != nullptr)
 		{
-			for (int32 NaniteMeshPassIndex = 0; NaniteMeshPassIndex < ENaniteMeshPass::Num; ++NaniteMeshPassIndex)
+			if (UploadDataSourceAdapter.bUpdateNaniteMaterialTables && bNaniteEnabled)
 			{
-				Scene.MaterialTables[NaniteMeshPassIndex].Begin(RHICmdList, Scene.Primitives.Num(), NumPrimitiveDataUploads);
+				for (int32 NaniteMeshPassIndex = 0; NaniteMeshPassIndex < ENaniteMeshPass::Num; ++NaniteMeshPassIndex)
+				{
+					Scene->MaterialTables[NaniteMeshPassIndex].Begin(RHICmdList, Scene->Primitives.Num(), NumPrimitiveDataUploads);
+				}
 			}
 		}
 
@@ -293,26 +547,11 @@ void FGPUScene::UpdateInternal(FRHICommandListImmediate& RHICmdList, FScene& Sce
 
 		if (NumPrimitiveDataUploads > 0)
 		{
-			auto ProcessPrimitiveFn = [this, &Scene, &NumLightmapDataUploads, &NumInstanceDataUploads, bNaniteEnabled](int32 PrimitiveIndex, bool bThreaded)
+			auto ProcessPrimitiveFn = [this, &UploadDataSourceAdapter, Scene, &NumLightmapDataUploads, &NumInstanceDataUploads, bNaniteEnabled](int32 ItemIndex, bool bThreaded)
 			{
-				// PrimitivesToUpdate may contain a stale out of bounds index, as we don't remove update request on primitive removal from scene.
-				if (PrimitiveIndex < Scene.PrimitiveSceneProxies.Num())
+				FPrimitiveUploadInfo UploadInfo;
+				if (UploadDataSourceAdapter.GetPrimitiveInfo(ItemIndex, UploadInfo))
 				{
-					const FPrimitiveSceneProxy* RESTRICT PrimitiveSceneProxy = Scene.PrimitiveSceneProxies[PrimitiveIndex];
-					const FPrimitiveSceneInfo* PrimitiveSceneInfo = PrimitiveSceneProxy->GetPrimitiveSceneInfo();
-
-					int32 TaskInstanceDataOffset  = 0;
-					int32 TaskInstanceUploadCount = 0;
-					int32 TaskLightmapUploadCount = PrimitiveSceneInfo->GetNumLightmapDataEntries();
-
-					// Count all primitive instances represented in the instance data buffer.
-					if (PrimitiveSceneProxy->SupportsInstanceDataBuffer())
-					{
-						const TArray<FPrimitiveInstance>* PrimitiveInstances = PrimitiveSceneProxy->GetPrimitiveInstances();
-						TaskInstanceDataOffset  = PrimitiveSceneProxy->GetPrimitiveSceneInfo()->GetInstanceDataOffset();
-						TaskInstanceUploadCount = PrimitiveInstances->Num();
-					}
-
 					{
 						if (bThreaded)
 						{
@@ -320,38 +559,43 @@ void FGPUScene::UpdateInternal(FRHICommandListImmediate& RHICmdList, FScene& Sce
 						}
 
 						// Cancel out any pending clear bits for these instances.
-						if (TaskInstanceUploadCount > 0)
+						if (UploadInfo.InstanceUploadCount > 0)
 						{
-							InstanceDataToClear.SetRange(TaskInstanceDataOffset, TaskInstanceUploadCount, false);
+							check(UploadInfo.InstanceDataOffset != INDEX_NONE);
+							InstanceDataToClear.SetRange(UploadInfo.InstanceDataOffset, UploadInfo.InstanceUploadCount, false);
 						}
 
-						NumLightmapDataUploads += TaskLightmapUploadCount; // Not thread safe
-						NumInstanceDataUploads += TaskInstanceUploadCount; // Not thread safe
+						NumLightmapDataUploads += UploadInfo.LightmapUploadCount; // Not thread safe
+						NumInstanceDataUploads += UploadInfo.InstanceUploadCount; // Not thread safe
 
-						void* UploadDst = PrimitiveUploadBuffer.Add_GetRef(PrimitiveIndex); // Not thread safe
+						void* UploadDst = PrimitiveUploadBuffer.Add_GetRef(UploadInfo.PrimitiveID); // Not thread safe
 
 						if (bThreaded)
 						{
 							PrimitiveUploadBufferCS.Unlock();
 						}
 
-						FPrimitiveSceneShaderData PrimitiveSceneData(PrimitiveSceneProxy);
+
 						FVector4* DstData = static_cast<FVector4*>(UploadDst);
 						for (uint32 VectorIndex = 0; VectorIndex < FPrimitiveSceneShaderData::PrimitiveDataStrideInFloat4s; ++VectorIndex)
 						{
-							DstData[VectorIndex] = PrimitiveSceneData.Data[VectorIndex];
+							DstData[VectorIndex] = UploadInfo.PrimitiveSceneData.Data[VectorIndex];
 						}
 					}
 
 					// Update Nanite material tables associated with this primitive index.
-					if (bNaniteEnabled && PrimitiveSceneProxy->IsNaniteMesh())
+					// GPUCULL_TODO: Abstract this also if we ever need Nanite in dynamic primitives
+					if (Scene != nullptr && bNaniteEnabled && UploadInfo.NaniteSceneProxy != nullptr)
 					{
-						const Nanite::FSceneProxyBase* NaniteSceneProxy = static_cast<const Nanite::FSceneProxyBase*>(PrimitiveSceneProxy);
+						check(UploadDataSourceAdapter.bUpdateNaniteMaterialTables);
+						check(UploadInfo.PrimitiveSceneInfo != nullptr);
+						const FPrimitiveSceneInfo* PrimitiveSceneInfo = UploadInfo.PrimitiveSceneInfo;
+						const Nanite::FSceneProxyBase* NaniteSceneProxy = UploadInfo.NaniteSceneProxy;
 
 						// Update material depth and hit proxy ID remapping tables.
 						for (int32 NaniteMeshPass = 0; NaniteMeshPass < ENaniteMeshPass::Num; ++NaniteMeshPass)
 						{
-							FNaniteMaterialTables& PassMaterialTables = Scene.MaterialTables[NaniteMeshPass];
+							FNaniteMaterialTables& PassMaterialTables = Scene->MaterialTables[NaniteMeshPass];
 							const TArray<uint32>& PassMaterialIds = PrimitiveSceneInfo->NaniteMaterialIds[NaniteMeshPass];
 							check(NaniteSceneProxy->GetMaterialSections().Num() == PassMaterialIds.Num());
 
@@ -360,10 +604,10 @@ void FGPUScene::UpdateInternal(FRHICommandListImmediate& RHICmdList, FScene& Sce
 								MaterialTableUploadCS.Lock();
 							}
 
-							void* DepthTable = PassMaterialTables.GetDepthTablePtr(PrimitiveIndex, PassMaterialIds.Num());
+							void* DepthTable = PassMaterialTables.GetDepthTablePtr(UploadInfo.PrimitiveID, PassMaterialIds.Num());
 						#if WITH_EDITOR
 							const uint32 HitProxyEntryCount = (NaniteMeshPass == ENaniteMeshPass::BasePass) ? PrimitiveSceneInfo->NaniteHitProxyIds.Num() : NANITE_MAX_MATERIALS;
-							void* HitProxyTable = PassMaterialTables.GetHitProxyTablePtr(PrimitiveIndex, HitProxyEntryCount);
+							void* HitProxyTable = PassMaterialTables.GetHitProxyTablePtr(UploadInfo.PrimitiveID, HitProxyEntryCount);
 						#endif
 
 							if (bThreaded)
@@ -407,19 +651,18 @@ void FGPUScene::UpdateInternal(FRHICommandListImmediate& RHICmdList, FScene& Sce
 			if (MaxPrimitivesUploads == NumPrimitiveDataUploads)
 			{
 				// One large batch
-				SCOPED_DRAW_EVENTF(RHICmdList, UpdateGPUScene, TEXT("UpdateGPUScene PrimitivesToUpdate %u"), NumPrimitiveDataUploads);
+				SCOPED_DRAW_EVENTF(RHICmdList, UpdateGPUScene, TEXT("UpdateGPUScene NumPrimitiveDataUploads %u"), NumPrimitiveDataUploads);
 
 				PrimitiveUploadBuffer.Init(NumPrimitiveDataUploads, sizeof(FPrimitiveSceneShaderData::Data), true, TEXT("PrimitiveUploadBuffer"));
 
 				int32 RangeCount = PartitionUpdateRanges(ParallelRanges, NumPrimitiveDataUploads, bExecuteInParallel);
 
 				ParallelFor(RangeCount,
-					[this, &ProcessPrimitiveFn, &ParallelRanges, &Scene, RangeCount](int32 RangeIndex)
+					[&ProcessPrimitiveFn, &ParallelRanges, RangeCount](int32 RangeIndex)
 					{
 						for (int32 ItemIndex = ParallelRanges.Range[RangeIndex].ItemStart; ItemIndex < ParallelRanges.Range[RangeIndex].ItemStart + ParallelRanges.Range[RangeIndex].ItemCount; ++ItemIndex)
 						{
-							const int32 Index = PrimitivesToUpdate[ItemIndex];
-							ProcessPrimitiveFn(Index, RangeCount > 1);
+							ProcessPrimitiveFn(ItemIndex, RangeCount > 1);
 						}
 					},
 					RangeCount == 1
@@ -435,14 +678,14 @@ void FGPUScene::UpdateInternal(FRHICommandListImmediate& RHICmdList, FScene& Sce
 				// Break into multiple batches
 				for (int32 PrimitiveOffset = 0; PrimitiveOffset < NumPrimitiveDataUploads; PrimitiveOffset += MaxPrimitivesUploads)
 				{
-					SCOPED_DRAW_EVENTF(RHICmdList, UpdateGPUScene, TEXT("UpdateGPUScene PrimitivesToUpdate and Offset = %u %u"), NumPrimitiveDataUploads, PrimitiveOffset);
+					SCOPED_DRAW_EVENTF(RHICmdList, UpdateGPUScene, TEXT("UpdateGPUScene NumPrimitiveDataUploads and Offset = %u %u"), NumPrimitiveDataUploads, PrimitiveOffset);
 
 					PrimitiveUploadBuffer.Init(MaxPrimitivesUploads, sizeof(FPrimitiveSceneShaderData::Data), true, TEXT("PrimitiveUploadBuffer"));
 
 					for (int32 IndexUpdate = 0; (IndexUpdate < MaxPrimitivesUploads) && ((IndexUpdate + PrimitiveOffset) < NumPrimitiveDataUploads); ++IndexUpdate)
 					{
-						const int32 Index = PrimitivesToUpdate[IndexUpdate + PrimitiveOffset];
-						ProcessPrimitiveFn(Index, false /* threaded */);
+						const int32 ItemIndex = IndexUpdate + PrimitiveOffset;
+						ProcessPrimitiveFn(ItemIndex, false /* threaded */);
 					}
 
 					RHICmdList.Transition(FRHITransitionInfo(MirrorResourceGPU->UAV, CurrentAccess, ERHIAccess::UAVCompute));
@@ -458,109 +701,69 @@ void FGPUScene::UpdateInternal(FRHICommandListImmediate& RHICmdList, FScene& Sce
 			RHICmdList.Transition(FRHITransitionInfo(MirrorResourceGPU->UAV, CurrentAccess, ERHIAccess::SRVMask));
 		}
 
-		PrimitivesMarkedToUpdate.Init(false, PrimitivesMarkedToUpdate.Num());
-
-		if (bNaniteEnabled)
+		if (Scene != nullptr)
 		{
-			for (int32 NaniteMeshPassIndex = 0; NaniteMeshPassIndex < ENaniteMeshPass::Num; ++NaniteMeshPassIndex)
+			if (UploadDataSourceAdapter.bUpdateNaniteMaterialTables && bNaniteEnabled)
 			{
-				Scene.MaterialTables[NaniteMeshPassIndex].Finish(RHICmdList);
+				for (int32 NaniteMeshPassIndex = 0; NaniteMeshPassIndex < ENaniteMeshPass::Num; ++NaniteMeshPassIndex)
+				{
+					Scene->MaterialTables[NaniteMeshPassIndex].Finish(RHICmdList);
+				}
 			}
 		}
-
 		// Make sure instance buffer always has valid or properly reset entries.
 		TArray<uint32, TInlineAllocator<64, SceneRenderingAllocator>> InstancesToClear;
 
-		if (GGPUSceneInstanceClearList != 0)
+	#if !UE_BUILD_SHIPPING
+		static const bool bVerifyClearList = false;
+		if (bVerifyClearList)
 		{
-		#if !UE_BUILD_SHIPPING
-			static const bool bVerifyClearList = false;
-			if (bVerifyClearList)
-			{
-				// We need to make sure that every set clear bit in the total list is
-				// represented by an entry in the clear list. We can safely ignore unset
-				// bits - such as clear list contains an entry from an earlier removal,
-				// but clear bit was unset prior to clearing because a slot was reused
-				// by an instance added.
-				for (int32 InstanceIndex = 0; InstanceIndex < InstanceDataAllocator.GetMaxSize(); ++InstanceIndex)
-				{
-					if (InstanceDataToClear[InstanceIndex])
-					{
-						check(InstanceClearList.Contains(InstanceIndex));
-					}
-				}
-			}
-		#endif
-
-			for (uint32 InstanceIndex : InstanceClearList)
+			// We need to make sure that every set clear bit in the total list is
+			// represented by an entry in the clear list. We can safely ignore unset
+			// bits - such as clear list contains an entry from an earlier removal,
+			// but clear bit was unset prior to clearing because a slot was reused
+			// by an instance added.
+			for (int32 InstanceIndex = 0; InstanceIndex < InstanceDataAllocator.GetMaxSize(); ++InstanceIndex)
 			{
-				// Any clear bits set after enumerating the primitives being updated are
-				// stale sections of the instance data buffer, so they should be reset to
-				// an invalid state and skipped on the GPU.
 				if (InstanceDataToClear[InstanceIndex])
 				{
-					InstanceDataToClear[InstanceIndex] = false;
-					InstancesToClear.Add(InstanceIndex);
+					check(InstanceClearList.Contains(InstanceIndex));
 				}
 			}
-
-			InstanceClearList.Reset();
 		}
-		else
+	#endif
+
+		for (uint32 InstanceIndex : InstanceClearList)
 		{
-			for (int32 InstanceIndex = 0; InstanceIndex < InstanceDataAllocator.GetMaxSize(); ++InstanceIndex)
+			// Any clear bits set after enumerating the primitives being updated are
+			// stale sections of the instance data buffer, so they should be reset to
+			// an invalid state and skipped on the GPU.
+			if (InstanceDataToClear[InstanceIndex])
 			{
-				// Any clear bits set after enumerating the primitives being updated are
-				// stale sections of the instance data buffer, so they should be reset to
-				// an invalid state and skipped on the GPU.
-				if (InstanceDataToClear[InstanceIndex])
-				{
-					InstancesToClear.Add(InstanceIndex);
-					InstanceDataToClear[InstanceIndex] = false;
-				}
+				InstanceDataToClear[InstanceIndex] = false;
+				InstancesToClear.Add(InstanceIndex);
 			}
 		}
 
-		// Clears count toward the total instance data uploads - batched together for efficiency.
-		NumInstanceDataUploads += InstancesToClear.Num();
-
-		if (GGPUSceneValidatePrimitiveBuffer && (PrimitiveBuffer.NumBytes > 0 || PrimitiveTexture.NumBytes > 0))
-		{
-			//UE_LOG(LogRenderer, Warning, TEXT("r.GPUSceneValidatePrimitiveBuffer enabled, doing slow readback from GPU"));
-			uint32 Stride = 0;
-			FPrimitiveSceneShaderData* InstanceBufferCopy = (FPrimitiveSceneShaderData*)(FPrimitiveSceneShaderData*)LockResource(*MirrorResourceGPU, Stride);
-
-			const int32 TotalNumberPrimitives = Scene.PrimitiveSceneProxies.Num();
-			int32 MaxPrimitivesUploads = GetMaxPrimitivesUpdate(TotalNumberPrimitives, FPrimitiveSceneShaderData::PrimitiveDataStrideInFloat4s);
-			for (int32 IndexOffset = 0; IndexOffset < TotalNumberPrimitives; IndexOffset += MaxPrimitivesUploads)
-			{
-				for (int32 Index = 0; (Index < MaxPrimitivesUploads) && ((Index + IndexOffset) < TotalNumberPrimitives); ++Index)
-				{
-					FPrimitiveSceneShaderData PrimitiveSceneData(Scene.PrimitiveSceneProxies[Index + IndexOffset]);
+		InstanceClearList.Reset();
 
-					for (int32 i = 0; i < FPrimitiveSceneShaderData::PrimitiveDataStrideInFloat4s; i++)
-					{
-						check(PrimitiveSceneData.Data[i] == InstanceBufferCopy[Index].Data[i]);
-					}
-				}
-				InstanceBufferCopy += Stride / sizeof(FPrimitiveSceneShaderData);
-			}
 
-			UnlockResourceGPUScene(*MirrorResourceGPU);
-		}
+		// Clears count toward the total instance data uploads - batched together for efficiency.
+		NumInstanceDataUploads += InstancesToClear.Num();
 
-		if (NumPrimitiveDataUploads > 0)
+		// GPUCULL_TODO: May this not skip clears? E.g. if something is removed?
+		//ensure(NumPrimitiveDataUploads == 0 ? NumInstanceDataUploads == 0 : true);
 		{
 			// Upload instancing data for the scene.
 			if (NumInstanceDataUploads > 0)
 			{
-				InstanceUploadBuffer.Init(NumInstanceDataUploads * InstanceDataNumArrays, sizeof( FVector4 ), true, TEXT("InstanceUploadBuffer"));
+				InstanceUploadBuffer.Init(NumInstanceDataUploads * InstanceDataNumArrays, sizeof(FVector4), true, TEXT("InstanceUploadBuffer"));
 
 				int32 RangeCount = PartitionUpdateRanges(ParallelRanges, InstancesToClear.Num(), bExecuteInParallel);
 
 				// Reset any instance slots marked for clearing.
 				ParallelFor(RangeCount,
-					[this, &InstancesToClear, &ParallelRanges, &Scene, RangeCount, InstanceDataNumArrays, InstanceDataSizeReserve](int32 RangeIndex)
+					[this, &InstancesToClear, &ParallelRanges, RangeCount, InstanceDataNumArrays, InstanceDataSizeReserve](int32 RangeIndex)
 					{
 						for (int32 ItemIndex = ParallelRanges.Range[RangeIndex].ItemStart; ItemIndex < ParallelRanges.Range[RangeIndex].ItemStart + ParallelRanges.Range[RangeIndex].ItemCount; ++ItemIndex)
 						{
@@ -568,7 +771,7 @@ void FGPUScene::UpdateInternal(FRHICommandListImmediate& RHICmdList, FScene& Sce
 							FPrimitiveInstance PrimitiveInstance;
 							PrimitiveInstance.PrimitiveId = ~uint32(0);
 							FInstanceSceneShaderData InstanceSceneData(PrimitiveInstance);
-						
+
 							void* DstRefs[FInstanceSceneShaderData::InstanceDataStrideInFloat4s];
 							if (RangeCount > 1)
 							{
@@ -595,117 +798,87 @@ void FGPUScene::UpdateInternal(FRHICommandListImmediate& RHICmdList, FScene& Sce
 					RangeCount == 1
 				);
 
-				RangeCount = PartitionUpdateRanges(ParallelRanges, PrimitivesToUpdate.Num(), bExecuteInParallel);
+				if (NumPrimitiveDataUploads > 0)
+				{
+					// Note: we iterate over the primitives, whether they have instances or not (which is a bit wasteful) but this is the way we currently get to the instance data.
+					// GPUCULL_TODO: move instance data ownership to GPU-scene such that it can be put in a compact list or something, and be tracked independent of primitives?
+					RangeCount = PartitionUpdateRanges(ParallelRanges, NumPrimitiveDataUploads, bExecuteInParallel);
 
-				// Upload any out of data instance slots.
-				ParallelFor(RangeCount,
-					[this, &Scene, &ParallelRanges, RangeCount, InstanceDataNumArrays, InstanceDataSizeReserve](int32 RangeIndex)
-					{
-						for (int32 ItemIndex = ParallelRanges.Range[RangeIndex].ItemStart; ItemIndex < ParallelRanges.Range[RangeIndex].ItemStart + ParallelRanges.Range[RangeIndex].ItemCount; ++ItemIndex)
+					// Upload any out of date instance slots.
+					ParallelFor(RangeCount,
+						[&UploadDataSourceAdapter, this, &ParallelRanges, RangeCount, InstanceDataNumArrays, InstanceDataSizeReserve](int32 RangeIndex)
 						{
-							const int32 Index = PrimitivesToUpdate[ItemIndex];
-
-							// PrimitivesToUpdate may contain a stale out of bounds index, as we don't remove update request on primitive removal from scene.
-							if (Index < Scene.PrimitiveSceneProxies.Num())
+							for (int32 ItemIndex = ParallelRanges.Range[RangeIndex].ItemStart; ItemIndex < ParallelRanges.Range[RangeIndex].ItemStart + ParallelRanges.Range[RangeIndex].ItemCount; ++ItemIndex)
 							{
-								FPrimitiveSceneProxy* PrimitiveSceneProxy = Scene.PrimitiveSceneProxies[Index];
-								TArray<FPrimitiveInstance>* PrimitiveInstancesPtr = PrimitiveSceneProxy->GetPrimitiveInstances();
-								if (!PrimitiveSceneProxy->SupportsInstanceDataBuffer() || !PrimitiveInstancesPtr)
+								FInstanceUploadInfo UploadInfo;
+								if (UploadDataSourceAdapter.GetInstanceInfo(ItemIndex, UploadInfo))
 								{
-									continue;
-								}
-
-								TArray<FPrimitiveInstance>& PrimitiveInstances = *PrimitiveInstancesPtr;
-								if (PrimitiveInstances.Num() == 0)
-								{
-									continue;
-								}
-
-								const FPrimitiveSceneInfo* PrimitiveSceneInfo = PrimitiveSceneProxy->GetPrimitiveSceneInfo();
-
-								check(PrimitiveInstances.Num() == PrimitiveSceneInfo->GetNumInstanceDataEntries());
-								const int32 InstanceDataOffset = PrimitiveSceneInfo->GetInstanceDataOffset();
-								check(InstanceDataOffset != INDEX_NONE);
-
-								// These are the bounds covering all primitives (which we don't want to use for per-instance culling).
-								//const FBoxSphereBounds& PrimitiveBounds = PrimitiveSceneProxy->GetLocalBounds();
-
-								// Update each primitive instance with current data.
-								for (int32 InstanceIndex = 0; InstanceIndex < PrimitiveInstances.Num(); ++InstanceIndex)
-								{
-									bool bOutHasPrecomputedVolumetricLightmap;
-									FMatrix OutPreviousLocalToWorld;
-									int32 OutSingleCaptureIndex;
-									bool bOutOutputVelocity;
-
-									Scene.GetPrimitiveUniformShaderParameters_RenderThread(
-										PrimitiveSceneInfo,
-										bOutHasPrecomputedVolumetricLightmap,
-										OutPreviousLocalToWorld,
-										OutSingleCaptureIndex,
-										bOutOutputVelocity
-									);
-
-									FPrimitiveInstance& PrimitiveInstance = PrimitiveInstances[InstanceIndex];
-									PrimitiveInstance.PrimitiveId = uint32(PrimitiveSceneInfo->GetIndex());
-									PrimitiveInstance.LocalBounds = PrimitiveInstance.RenderBounds;// .TransformBy(PrimitiveInstance.LocalToWorld);
-									PrimitiveInstance.LocalToWorld = PrimitiveInstance.InstanceToLocal * PrimitiveSceneProxy->GetLocalToWorld();
-									PrimitiveInstance.PrevLocalToWorld = PrimitiveInstance.InstanceToLocal * OutPreviousLocalToWorld;
-									PrimitiveInstance.LastUpdateSceneFrameNumber = SceneFrameNumber;
+									// These are the bounds covering all primitives (which we don't want to use for per-instance culling).
+									//const FBoxSphereBounds& PrimitiveBounds = PrimitiveSceneProxy->GetLocalBounds();
 
+									// Update each primitive instance with current data.
+									for (int32 InstanceIndex = 0; InstanceIndex < UploadInfo.PrimitiveInstances.Num(); ++InstanceIndex)
 									{
-										// Extract per axis scales from InstanceToWorld transform
-										FVector4 WorldX = FVector4(PrimitiveInstance.LocalToWorld.M[0][0], PrimitiveInstance.LocalToWorld.M[0][1], PrimitiveInstance.LocalToWorld.M[0][2], 0);
-										FVector4 WorldY = FVector4(PrimitiveInstance.LocalToWorld.M[1][0], PrimitiveInstance.LocalToWorld.M[1][1], PrimitiveInstance.LocalToWorld.M[1][2], 0);
-										FVector4 WorldZ = FVector4(PrimitiveInstance.LocalToWorld.M[2][0], PrimitiveInstance.LocalToWorld.M[2][1], PrimitiveInstance.LocalToWorld.M[2][2], 0);
-
-										const float ScaleX = FVector(WorldX).Size();
-										const float ScaleY = FVector(WorldY).Size();
-										const float ScaleZ = FVector(WorldZ).Size();
-
-										PrimitiveInstance.NonUniformScale = FVector4(
-											ScaleX, ScaleY, ScaleZ,
-											FMath::Max3(FMath::Abs(ScaleX), FMath::Abs(ScaleY), FMath::Abs(ScaleZ))
-										);
-
-										PrimitiveInstance.InvNonUniformScaleAndDeterminantSign = FVector4(
-											ScaleX > KINDA_SMALL_NUMBER ? 1.0f / ScaleX : 0.0f,
-											ScaleY > KINDA_SMALL_NUMBER ? 1.0f / ScaleY : 0.0f,
-											ScaleZ > KINDA_SMALL_NUMBER ? 1.0f / ScaleZ : 0.0f,
-											FMath::FloatSelect(PrimitiveInstance.LocalToWorld.RotDeterminant(), 1.0f, -1.0f)
-										);
-									}
-
-									FInstanceSceneShaderData InstanceSceneData(PrimitiveInstance);
-
-									void* DstRefs[FInstanceSceneShaderData::InstanceDataStrideInFloat4s];
-									if (RangeCount > 1)
-									{
-										PrimitiveUploadBufferCS.Lock();
-									}
-
-									for (uint32 RefIndex = 0; RefIndex < InstanceDataNumArrays; ++RefIndex)
-									{
-										DstRefs[RefIndex] = InstanceUploadBuffer.Add_GetRef(RefIndex * InstanceDataSizeReserve + InstanceDataOffset + InstanceIndex);
-									}
-
-									if (RangeCount > 1)
-									{
-										PrimitiveUploadBufferCS.Unlock();
-									}
-
-									for (uint32 RefIndex = 0; RefIndex < InstanceDataNumArrays; ++RefIndex) //TODO: make a SOA version of InstanceUploadBuffer.Add
-									{
-										FVector4* DstVector = static_cast<FVector4*>(DstRefs[RefIndex]);
-										*DstVector = InstanceSceneData.Data[RefIndex];
+										FPrimitiveInstance& PrimitiveInstance = UploadInfo.PrimitiveInstances[InstanceIndex];
+										PrimitiveInstance.PrimitiveId = UploadInfo.PrimitiveID;
+										PrimitiveInstance.LocalBounds = PrimitiveInstance.RenderBounds;// .TransformBy(PrimitiveInstance.LocalToWorld);
+										PrimitiveInstance.LocalToWorld = PrimitiveInstance.InstanceToLocal * UploadInfo.PrimitiveLocalToWorld;
+										PrimitiveInstance.PrevLocalToWorld = PrimitiveInstance.InstanceToLocal * UploadInfo.PreviousPrimitiveLocalToWorld;
+										PrimitiveInstance.LastUpdateSceneFrameNumber = SceneFrameNumber;
+
+										{
+											// Extract per axis scales from InstanceToWorld transform
+											FVector4 WorldX = FVector4(PrimitiveInstance.LocalToWorld.M[0][0], PrimitiveInstance.LocalToWorld.M[0][1], PrimitiveInstance.LocalToWorld.M[0][2], 0);
+											FVector4 WorldY = FVector4(PrimitiveInstance.LocalToWorld.M[1][0], PrimitiveInstance.LocalToWorld.M[1][1], PrimitiveInstance.LocalToWorld.M[1][2], 0);
+											FVector4 WorldZ = FVector4(PrimitiveInstance.LocalToWorld.M[2][0], PrimitiveInstance.LocalToWorld.M[2][1], PrimitiveInstance.LocalToWorld.M[2][2], 0);
+
+											const float ScaleX = FVector(WorldX).Size();
+											const float ScaleY = FVector(WorldY).Size();
+											const float ScaleZ = FVector(WorldZ).Size();
+
+											PrimitiveInstance.NonUniformScale = FVector4(
+												ScaleX, ScaleY, ScaleZ,
+												FMath::Max3(FMath::Abs(ScaleX), FMath::Abs(ScaleY), FMath::Abs(ScaleZ))
+											);
+
+											PrimitiveInstance.InvNonUniformScaleAndDeterminantSign = FVector4(
+												ScaleX > KINDA_SMALL_NUMBER ? 1.0f / ScaleX : 0.0f,
+												ScaleY > KINDA_SMALL_NUMBER ? 1.0f / ScaleY : 0.0f,
+												ScaleZ > KINDA_SMALL_NUMBER ? 1.0f / ScaleZ : 0.0f,
+												FMath::FloatSelect(PrimitiveInstance.LocalToWorld.RotDeterminant(), 1.0f, -1.0f)
+											);
+										}
+
+										FInstanceSceneShaderData InstanceSceneData(PrimitiveInstance);
+
+										void* DstRefs[FInstanceSceneShaderData::InstanceDataStrideInFloat4s];
+										if (RangeCount > 1)
+										{
+											PrimitiveUploadBufferCS.Lock();
+										}
+
+										for (uint32 RefIndex = 0; RefIndex < InstanceDataNumArrays; ++RefIndex)
+										{
+											DstRefs[RefIndex] = InstanceUploadBuffer.Add_GetRef(RefIndex * InstanceDataSizeReserve + UploadInfo.InstanceDataOffset + InstanceIndex);
+										}
+
+										if (RangeCount > 1)
+										{
+											PrimitiveUploadBufferCS.Unlock();
+										}
+
+										for (uint32 RefIndex = 0; RefIndex < InstanceDataNumArrays; ++RefIndex) //TODO: make a SOA version of InstanceUploadBuffer.Add
+										{
+											FVector4* DstVector = static_cast<FVector4*>(DstRefs[RefIndex]);
+											*DstVector = InstanceSceneData.Data[RefIndex];
+										}
 									}
 								}
 							}
-						}
-					},
-					RangeCount == 1
-				);
-
+						},
+						RangeCount == 1
+					);
+				}
 				RHICmdList.Transition(FRHITransitionInfo(InstanceDataBuffer.UAV, ERHIAccess::Unknown, ERHIAccess::UAVCompute));
 				InstanceUploadBuffer.ResourceUploadTo(RHICmdList, InstanceDataBuffer, false);
 				RHICmdList.Transition(FRHITransitionInfo(InstanceDataBuffer.UAV, ERHIAccess::UAVCompute, ERHIAccess::SRVMask));
@@ -713,7 +886,7 @@ void FGPUScene::UpdateInternal(FRHICommandListImmediate& RHICmdList, FScene& Sce
 			else if (bResizedInstanceData)
 			{
 				RHICmdList.Transition(FRHITransitionInfo(InstanceDataBuffer.UAV, ERHIAccess::Unknown, ERHIAccess::SRVMask));
-			 }
+			}
 
 			if (NumLightmapDataUploads > 0)
 			{
@@ -726,22 +899,13 @@ void FGPUScene::UpdateInternal(FRHICommandListImmediate& RHICmdList, FScene& Sce
 
 					for (int32 IndexUpdate = 0; (IndexUpdate < MaxLightmapsUploads) && ((IndexUpdate + PrimitiveOffset) < NumPrimitiveDataUploads); ++IndexUpdate)
 					{
-						int32 Index = PrimitivesToUpdate[IndexUpdate + PrimitiveOffset];
-						// PrimitivesToUpdate may contain a stale out of bounds index, as we don't remove update request on primitive removal from scene.
-						if (Index < Scene.PrimitiveSceneProxies.Num())
+						FLightMapUploadInfo UploadInfo;
+						if (UploadDataSourceAdapter.GetLightMapInfo(IndexUpdate, UploadInfo))
 						{
-							FPrimitiveSceneProxy* PrimitiveSceneProxy = Scene.PrimitiveSceneProxies[Index];
-
-							FPrimitiveSceneProxy::FLCIArray LCIs;
-							PrimitiveSceneProxy->GetLCIs(LCIs);
-
-							check(LCIs.Num() == PrimitiveSceneProxy->GetPrimitiveSceneInfo()->GetNumLightmapDataEntries());
-							const int32 LightmapDataOffset = PrimitiveSceneProxy->GetPrimitiveSceneInfo()->GetLightmapDataOffset();
-
-							for (int32 i = 0; i < LCIs.Num(); i++)
+							for (int32 LCIIndex = 0; LCIIndex < UploadInfo.LCIs.Num(); LCIIndex++)
 							{
-								FLightmapSceneShaderData LightmapSceneData(LCIs[i], Scene.GetFeatureLevel());
-								LightmapUploadBuffer.Add(LightmapDataOffset + i, &LightmapSceneData.Data[0]);
+								FLightmapSceneShaderData LightmapSceneData(UploadInfo.LCIs[LCIIndex], FeatureLevel);
+								LightmapUploadBuffer.Add(UploadInfo.LightmapDataOffset + LCIIndex, &LightmapSceneData.Data[0]);
 							}
 						}
 					}
@@ -755,8 +919,6 @@ void FGPUScene::UpdateInternal(FRHICommandListImmediate& RHICmdList, FScene& Sce
 				RHICmdList.Transition(FRHITransitionInfo(LightmapDataBuffer.UAV, CurrentAccess, ERHIAccess::SRVMask));
 			}
 
-			PrimitivesToUpdate.Reset();
-			
 			if (PrimitiveUploadBuffer.GetNumBytes() > (uint32)GGPUSceneMaxPooledUploadBufferSize)
 			{
 				PrimitiveUploadBuffer.Release();
@@ -772,21 +934,99 @@ void FGPUScene::UpdateInternal(FRHICommandListImmediate& RHICmdList, FScene& Sce
 				LightmapUploadBuffer.Release();
 			}
 		}
-		// Clear the flags that mark newly added primitives.
-		AddedPrimitiveFlags.Init(false, AddedPrimitiveFlags.Num());
+	}
+}
 
-		// Dynamic primitives are allocated after the regular ones, this needs to match the value passed in in BeginRender
-		ensure(DynamicPrimitivesOffset >= Scene.Primitives.Num());
+struct FUploadDataSourceAdapterDynamicPrimitives
+{
+	static constexpr bool bUpdateNaniteMaterialTables = false;
+
+	FUploadDataSourceAdapterDynamicPrimitives(const TArray<FPrimitiveUniformShaderParameters, TInlineAllocator<8>>& InPrimitiveShaderData, int32 InPrimitiveIDStartOffset, int32 InInstanceIDStartOffset) :
+		PrimitiveShaderData(InPrimitiveShaderData), 
+		PrimitiveIDStartOffset(InPrimitiveIDStartOffset), 
+		InstanceIDStartOffset(InInstanceIDStartOffset) {}
+
+	FORCEINLINE int32 NumPrimitivesToUpload() const
+	{ 
+		return PrimitiveShaderData.Num(); 
 	}
 
-	checkSlow(PrimitivesToUpdate.Num() == 0);
-}
+	FORCEINLINE bool GetPrimitiveInfo(int32 ItemIndex, FPrimitiveUploadInfo& PrimitiveUploadInfo) const
+	{
+		PrimitiveUploadInfo.LightmapUploadCount = 0;
+		PrimitiveUploadInfo.NaniteSceneProxy = nullptr;
+		PrimitiveUploadInfo.PrimitiveSceneInfo = nullptr;
+		if (ItemIndex < PrimitiveShaderData.Num())
+		{
+			// Needed to ensure the link back to instace list is up to date
+			FPrimitiveUniformShaderParameters Tmp = PrimitiveShaderData[ItemIndex];
+#if defined(GPUCULL_TODO)
+			Tmp.InstanceDataOffset = InstanceIDStartOffset + ItemIndex;
+			Tmp.NumInstanceDataEntries = 1;
+
+			PrimitiveUploadInfo.InstanceDataOffset = InstanceIDStartOffset + ItemIndex;
+			PrimitiveUploadInfo.InstanceUploadCount = 1;
+#else // !defined(GPUCULL_TODO)
+			PrimitiveUploadInfo.InstanceDataOffset = INDEX_NONE;
+			PrimitiveUploadInfo.InstanceUploadCount = 0;
+#endif // defined(GPUCULL_TODO)
+			PrimitiveUploadInfo.PrimitiveID = PrimitiveIDStartOffset + ItemIndex;
+			PrimitiveUploadInfo.PrimitiveSceneData = FPrimitiveSceneShaderData(Tmp);
+
+			return true;
+		}
+		return false;
+	}
+
+	FORCEINLINE bool GetInstanceInfo(int32 ItemIndex, FInstanceUploadInfo& InstanceUploadInfo)
+	{	
+
+#if defined(GPUCULL_TODO)
+		if (ItemIndex < PrimitiveShaderData.Num())
+		{
+			InstanceUploadInfo.PrimitiveLocalToWorld = PrimitiveShaderData[ItemIndex].LocalToWorld;
+			InstanceUploadInfo.PreviousPrimitiveLocalToWorld = PrimitiveShaderData[ItemIndex].PreviousLocalToWorld;
+			InstanceUploadInfo.InstanceDataOffset = InstanceIDStartOffset + ItemIndex;
+			InstanceUploadInfo.PrimitiveID = PrimitiveIDStartOffset + ItemIndex;;
+
+			// We always create an instance to ensure that we can always use the same code paths in the shader
+			// In the future we should remove redundant data from the primitive, and then the instances should be
+			// provided by the proxy. However, this is a lot of work before we can just enable it in the base proxy class.
+			FPrimitiveInstance& DummyInstance = InstanceUploadInfo.DummyInstance;
+			DummyInstance.InstanceToLocal = FMatrix::Identity;
+			DummyInstance.LocalToInstance = FMatrix::Identity;
+			DummyInstance.LocalToWorld = FMatrix::Identity;
+			DummyInstance.PrevLocalToWorld = FMatrix::Identity;
+			DummyInstance.NonUniformScale = FVector4(1.0f, 1.0f, 1.0f, 1.0f);
+			DummyInstance.InvNonUniformScaleAndDeterminantSign = FVector4(1.0f, 1.0f, 1.0f, 1.0f);
+			DummyInstance.RenderBounds = FBoxSphereBounds(FBox(PrimitiveShaderData[ItemIndex].LocalObjectBoundsMin, PrimitiveShaderData[ItemIndex].LocalObjectBoundsMax));
+			DummyInstance.LocalBounds = DummyInstance.RenderBounds;;
+			DummyInstance.PrimitiveId = 0xFFFFFFFF;
+			DummyInstance.LastUpdateSceneFrameNumber = 0xFFFFFFFF;
+
+			InstanceUploadInfo.PrimitiveInstances = TArrayView<FPrimitiveInstance>(&DummyInstance, 1);
+
+			return true;
+		}
+#endif // defined(GPUCULL_TODO)
+		return false;
+	}
+
+	/*FORCEINLINE*/ bool GetLightMapInfo(int32 ItemIndex, FLightMapUploadInfo& UploadInfo) const
+	{
+		return false;
+	}
+
+	const TArray<FPrimitiveUniformShaderParameters, TInlineAllocator<8>> &PrimitiveShaderData;
+	int32 PrimitiveIDStartOffset;
+	int32 InstanceIDStartOffset;
+};
 
 template<typename ResourceType>
-void FGPUScene::UploadDynamicPrimitiveShaderDataForViewInternal(FRHICommandListImmediate& RHICmdList, FScene& Scene, FViewInfo& View)
+void FGPUScene::UploadDynamicPrimitiveShaderDataForViewInternal(FRHICommandListImmediate& RHICmdList, FScene *Scene, FViewInfo& View)
 {
 	ensure(bInBeginEndBlock);
-	ensure(DynamicPrimitivesOffset >= Scene.Primitives.Num());
+	ensure(Scene == nullptr || DynamicPrimitivesOffset >= Scene->Primitives.Num());
 	
 	CSV_SCOPED_TIMING_STAT_EXCLUSIVE(UploadDynamicPrimitiveShaderData);
 	QUICK_SCOPE_CYCLE_COUNTER(STAT_UploadDynamicPrimitiveShaderData);
@@ -808,76 +1048,12 @@ void FGPUScene::UploadDynamicPrimitiveShaderDataForViewInternal(FRHICommandListI
 
 		int32 UploadIdStart = Collector.GetPrimitiveIdRange().GetLowerBoundValue();
 		ensure(UploadIdStart < DynamicPrimitivesOffset);
+#if defined(GPUCULL_TODO)
+		ensure(Collector.UploadData->InstanceDataOffset != INDEX_NONE);
+#endif // defined(GPUCULL_TODO)
 
-		ResourceType* MirrorResourceGPU = GetMirrorGPU<ResourceType>(*this);
-		{
-			// Work out space requirements for all dynamic primitives
-			const int32 TotalNumberPrimitives = DynamicPrimitivesOffset;
-			const uint32 PrimitiveSceneNumFloat4s = TotalNumberPrimitives * FPrimitiveSceneShaderData::PrimitiveDataStrideInFloat4s;
-
-			const uint32 SizeReserve = FMath::RoundUpToPowerOfTwo(FMath::Max(TotalNumberPrimitives, 256));
-
-			ERHIAccess CurrentAccess = ERHIAccess::SRVMask;
-			// Reserve enough space
-			if (ResizeResourceIfNeeded(RHICmdList, *MirrorResourceGPU, SizeReserve * sizeof(FPrimitiveSceneShaderData::Data), TEXT("PrimitiveData")))
-			{
-				CurrentAccess = ERHIAccess::Unknown;
-			}
-			RHICmdList.Transition(FRHITransitionInfo(MirrorResourceGPU->UAV, CurrentAccess, ERHIAccess::UAVCompute));
-			CurrentAccess = ERHIAccess::UAVCompute;
-
-			// Upload DynamicPrimitiveShaderData to the allocated range in the primitive data resource
-			int32 MaxPrimitivesUploads = GetMaxPrimitivesUpdate(NumPrimitiveDataUploads, FPrimitiveSceneShaderData::PrimitiveDataStrideInFloat4s);
-			for (int32 BatchStartIndex = 0; BatchStartIndex < NumPrimitiveDataUploads; BatchStartIndex += MaxPrimitivesUploads)
-			{
-				PrimitiveUploadViewBuffer.Init(MaxPrimitivesUploads, sizeof(FPrimitiveSceneShaderData::Data), true, TEXT("PrimitiveUploadViewBuffer"));
-
-				for (int32 IndexUpdate = 0; (IndexUpdate < MaxPrimitivesUploads) && ((IndexUpdate + BatchStartIndex) < NumPrimitiveDataUploads); ++IndexUpdate)
-				{
-					int32 DynamicUploadIndex = BatchStartIndex + IndexUpdate;
-					FPrimitiveSceneShaderData PrimitiveSceneData(DynamicPrimitiveShaderData[DynamicUploadIndex]);
-					// Place dynamic primitive shader data after any previous batches of dynamic primitive data
-					PrimitiveUploadViewBuffer.Add(UploadIdStart + DynamicUploadIndex, &PrimitiveSceneData.Data[0]);
-				}
-
-				PrimitiveUploadViewBuffer.ResourceUploadTo(RHICmdList, *MirrorResourceGPU, false);
-			}
-			RHICmdList.Transition(FRHITransitionInfo(MirrorResourceGPU->UAV, CurrentAccess, ERHIAccess::SRVMask));
-		}
-		if (PrimitiveUploadViewBuffer.GetNumBytes() > (uint32)GGPUSceneMaxPooledUploadBufferSize)
-		{
-			PrimitiveUploadViewBuffer.Release();
-		}
-
-
-		if (GGPUSceneValidatePrimitiveBuffer && (PrimitiveBuffer.NumBytes > 0 || PrimitiveTexture.NumBytes > 0))
-		{
-			uint32 Stride = 0;
-			FPrimitiveSceneShaderData* InstanceBufferCopy = (FPrimitiveSceneShaderData*)LockResource(*MirrorResourceGPU, Stride);
-
-			const int32 TotalNumberPrimitives = DynamicPrimitivesOffset;
-			int32 MaxPrimitivesUploads = GetMaxPrimitivesUpdate(TotalNumberPrimitives, FPrimitiveSceneShaderData::PrimitiveDataStrideInFloat4s);
-			for (int32 IndexOffset = 0; IndexOffset < TotalNumberPrimitives; IndexOffset += MaxPrimitivesUploads)
-			{
-				for (int32 Index = 0; (Index < MaxPrimitivesUploads) && ((Index + IndexOffset) < TotalNumberPrimitives); ++Index)
-				{
-					if ((Index + IndexOffset) < Scene.PrimitiveSceneProxies.Num())
-					{
-						const FPrimitiveSceneShaderData &PrimitiveSceneData = Scene.PrimitiveSceneProxies[Index + IndexOffset];
-						check(FMemory::Memcmp(&InstanceBufferCopy[Index + IndexOffset], &PrimitiveSceneData, sizeof(FPrimitiveSceneShaderData)) == 0);
-					}
-					else if (Collector.GetPrimitiveIdRange().Contains(Index + IndexOffset))
-					{
-						const FPrimitiveSceneShaderData& PrimitiveSceneData = FPrimitiveSceneShaderData(DynamicPrimitiveShaderData[Index + IndexOffset - Scene.PrimitiveSceneProxies.Num()]);
-						check(FMemory::Memcmp(&InstanceBufferCopy[Index + IndexOffset], &PrimitiveSceneData, sizeof(FPrimitiveSceneShaderData)) == 0);
-					}
-
-				}
-				InstanceBufferCopy += Stride / sizeof(FPrimitiveSceneShaderData);
-			}
-
-			UnlockResourceGPUScene(*MirrorResourceGPU);
-		}
+		FUploadDataSourceAdapterDynamicPrimitives UploadAdapter(DynamicPrimitiveShaderData, UploadIdStart, Collector.UploadData->InstanceDataOffset);
+		UploadGeneral<ResourceType, FUploadDataSourceAdapterDynamicPrimitives>(RHICmdList, Scene, UploadAdapter);
 	}
 
 	UpdateUniformResource<ResourceType>(View, *this);
@@ -885,6 +1061,10 @@ void FGPUScene::UploadDynamicPrimitiveShaderDataForViewInternal(FRHICommandListI
 	// Update view uniform buffer
 	View.CachedViewUniformShaderParameters->InstanceSceneData = InstanceDataBuffer.SRV;
 	View.CachedViewUniformShaderParameters->LightmapSceneData = LightmapDataBuffer.SRV;
+#if defined(GPUCULL_TODO)
+	View.CachedViewUniformShaderParameters->InstanceDataSOAStride = InstanceDataSOAStride;
+#endif //defined(GPUCULL_TODO)
+
 	View.ViewUniformBuffer.UpdateUniformBufferImmediate(*View.CachedViewUniformShaderParameters);
 }
 
@@ -897,7 +1077,7 @@ void AddPrimitiveToUpdateGPU(FScene& Scene, int32 PrimitiveId)
 void FGPUScene::AddPrimitiveToUpdate(int32 PrimitiveId)
 {
 	if (bIsEnabled)
-	{ 
+	{
 		if (PrimitiveId + 1 > PrimitivesMarkedToUpdate.Num())
 		{
 			const int32 NewSize = Align(PrimitiveId + 1, 64);
@@ -923,7 +1103,7 @@ void FGPUScene::Update(FRDGBuilder& GraphBuilder, FScene& Scene)
 		{
 			Scene.VirtualShadowMapArrayCacheManager->ProcessPrimitivesToUpdate(GraphBuilder, Scene);
 		}
-		if (GPUSceneUseTexture2D(Scene.GetShaderPlatform()))
+		if (GPUSceneUseTexture2D(GetShaderPlatform()))
 		{
 			UpdateInternal<FTextureRWBuffer2D>(GraphBuilder.RHICmdList, Scene);
 		}
@@ -934,11 +1114,11 @@ void FGPUScene::Update(FRDGBuilder& GraphBuilder, FScene& Scene)
 	}
 }
 
-void FGPUScene::UploadDynamicPrimitiveShaderDataForView(FRHICommandListImmediate& RHICmdList, FScene& Scene, FViewInfo& View)
+void FGPUScene::UploadDynamicPrimitiveShaderDataForView(FRHICommandListImmediate& RHICmdList, FScene *Scene, FViewInfo& View)
 {
 	if (bIsEnabled)
 	{
-		if (GPUSceneUseTexture2D(Scene.GetShaderPlatform()))
+		if (GPUSceneUseTexture2D(GetShaderPlatform()))
 		{
 			UploadDynamicPrimitiveShaderDataForViewInternal<FTextureRWBuffer2D>(RHICmdList, Scene, View);
 		}
@@ -964,10 +1144,7 @@ int32 FGPUScene::AllocateInstanceSlots(int32 NumInstanceDataEntries)
 				InstanceDataToClear.Add(false, NewSize - InstanceDataToClear.Num());
 			}
 
-			if (GGPUSceneInstanceClearList != 0)
-			{
-				InstanceClearList.Reserve(InstanceDataToClear.Num());
-			}
+			InstanceClearList.Reserve(InstanceDataToClear.Num());
 
 			// Unset all bits associated with newly allocated instance data.
 			InstanceDataToClear.SetRange(InstanceDataOffset, NumInstanceDataEntries, false);
@@ -986,14 +1163,10 @@ void FGPUScene::FreeInstanceSlots(int InstanceDataOffset, int32 NumInstanceDataE
 	{
 		InstanceDataAllocator.Free(InstanceDataOffset, NumInstanceDataEntries);
 		InstanceDataToClear.SetRange(InstanceDataOffset, NumInstanceDataEntries, true);
-
-		if (GGPUSceneInstanceClearList != 0)
+		InstanceClearList.Reserve(InstanceDataToClear.Num());
+		for (int32 AddIndex = 0; AddIndex < NumInstanceDataEntries; ++AddIndex)
 		{
-			InstanceClearList.Reserve(InstanceDataToClear.Num());
-			for (int32 AddIndex = 0; AddIndex < NumInstanceDataEntries; ++AddIndex)
-			{
-				InstanceClearList.Add(InstanceDataOffset + AddIndex);
-			}
+			InstanceClearList.Add(InstanceDataOffset + AddIndex);
 		}
 
 		// Resize bit arrays to match new high watermark
@@ -1036,6 +1209,10 @@ TRange<int32> FGPUScene::CommitPrimitiveCollector(FGPUScenePrimitiveCollector& P
 	int32 StartOffset = DynamicPrimitivesOffset;
 	DynamicPrimitivesOffset += PrimitiveCollector.UploadData->PrimitiveShaderData.Num();
 
+#if defined(GPUCULL_TODO)
+	PrimitiveCollector.UploadData->InstanceDataOffset = AllocateInstanceSlots(PrimitiveCollector.UploadData->PrimitiveShaderData.Num());
+#endif // defined(GPUCULL_TODO)
+
 	return TRange<int32>(StartOffset, DynamicPrimitivesOffset);
 }
 
@@ -1044,6 +1221,10 @@ FGPUSceneDynamicContext::~FGPUSceneDynamicContext()
 {
 	for (auto UploadData : DymamicPrimitiveUploadData)
 	{
+#if defined(GPUCULL_TODO)
+		check(UploadData->InstanceDataOffset != INDEX_NONE);
+		GPUScene.FreeInstanceSlots(UploadData->InstanceDataOffset, UploadData->PrimitiveShaderData.Num());
+#endif // defined(GPUCULL_TODO)
 		delete UploadData;
 	}
 	DymamicPrimitiveUploadData.Empty();
diff --git a/Engine/Source/Runtime/Renderer/Private/GPUScene.h b/Engine/Source/Runtime/Renderer/Private/GPUScene.h
--- a/Engine/Source/Runtime/Renderer/Private/GPUScene.h
+++ b/Engine/Source/Runtime/Renderer/Private/GPUScene.h
@@ -76,6 +76,16 @@ public:
 	 */
 	const TRange<int32> &GetPrimitiveIdRange() const { return PrimitiveIdRange; }
 
+	FORCEINLINE int32 GetInstanceDataOffset(int32 PrimitiveId) const
+	{
+		ensure(bCommitted);
+		ensure(PrimitiveIdRange.Contains(PrimitiveId));
+		ensure(UploadData->bIsUploaded);
+
+		// Assume a 1:1 mapping between primitive ID and instance ID
+		return UploadData->InstanceDataOffset + (PrimitiveId - PrimitiveIdRange.GetLowerBoundValue());
+	}
+
 	int32 Num() const {	return UploadData != nullptr ? UploadData->PrimitiveShaderData.Num() : 0; }
 private:
 
@@ -85,6 +95,7 @@ private:
 	struct FUploadData
 	{
 		TArray<FPrimitiveUniformShaderParameters, TInlineAllocator<8>> PrimitiveShaderData;
+		int32 InstanceDataOffset = INDEX_NONE;
 		bool bIsUploaded = false;
 	};
 
@@ -130,16 +141,18 @@ public:
 	}
 	~FGPUScene();
 
-	void SetEnabled(bool bInIsEnabled) { bIsEnabled = bInIsEnabled; }
-
+	void SetEnabled(ERHIFeatureLevel::Type InFeatureLevel);
+	bool IsEnabled() const { return bIsEnabled; }
 	/**
 	 * Call at start of rendering (but after scene primitives are updated) to let GPU-Scene record scene primitive count 
-	 * and prepare for dynamic primitive allocations
+	 * and prepare for dynamic primitive allocations.
+	 * Scene may be NULL which means there are zero scene primitives (but there may be dynamic ones added later).
 	 */
-	void BeginRender(FScene& Scene, FGPUSceneDynamicContext &GPUSceneDynamicContext);
+	void BeginRender(const FScene* Scene, FGPUSceneDynamicContext &GPUSceneDynamicContext);
 	inline bool IsRendering() const { return bInBeginEndBlock; }
 	void EndRender();
 
+	EShaderPlatform GetShaderPlatform() const { return GShaderPlatformForFeatureLevel[FeatureLevel]; }
 
 	/**
 	 * Allocates a range of space in the instance data buffer for the required number of instances, 
@@ -161,7 +174,7 @@ public:
 	/**
 	 * Upload primitives from View.DynamicPrimitiveCollector.
 	 */
-	void UploadDynamicPrimitiveShaderDataForView(FRHICommandListImmediate& RHICmdList, FScene& Scene, FViewInfo& View);
+	void UploadDynamicPrimitiveShaderDataForView(FRHICommandListImmediate& RHICmdList, FScene* Scene, FViewInfo& View);
 
 	/**
 	 * Pull all pending updates from Scene and upload primitive & instance data.
@@ -227,8 +240,17 @@ private:
 	FGPUSceneDynamicContext* CurrentDynamicContext = nullptr;
 	int32 NumScenePrimitives = 0;
 
+	ERHIFeatureLevel::Type FeatureLevel;
+
+	/**
+	 * Generalized upload that uses an adapter to abstract the data souce. Enables uploading scene primitives & dynamic primitives using a single path.
+	 * @parameter Scene may be null, as it is only needed for the Nanite material table update (which is coupled to the Scene at the moment).
+	 */
+	template<typename ResourceType, typename FUploadDataSourceAdapter>
+	void UploadGeneral(FRHICommandListImmediate& RHICmdList, FScene* Scene, FUploadDataSourceAdapter& UploadDataSourceAdapter);
+
 	template<typename ResourceType>
-	void UploadDynamicPrimitiveShaderDataForViewInternal(FRHICommandListImmediate& RHICmdList, FScene& Scene, FViewInfo& View);
+	void UploadDynamicPrimitiveShaderDataForViewInternal(FRHICommandListImmediate& RHICmdList, FScene* Scene, FViewInfo& View);
 
 	template<typename ResourceType>
 	void UpdateInternal(FRHICommandListImmediate& RHICmdList, FScene& Scene);
@@ -237,7 +259,7 @@ private:
 class FGPUSceneScopeBeginEndHelper
 {
 public:
-	FGPUSceneScopeBeginEndHelper(FGPUScene& InGPUScene, FGPUSceneDynamicContext &GPUSceneDynamicContext, FScene& Scene) :
+	FGPUSceneScopeBeginEndHelper(FGPUScene& InGPUScene, FGPUSceneDynamicContext &GPUSceneDynamicContext, const FScene* Scene) :
 		GPUScene(InGPUScene)
 	{
 		GPUScene.BeginRender(Scene, GPUSceneDynamicContext);
@@ -247,6 +269,7 @@ public:
 	{
 		GPUScene.EndRender();
 	}
+
 private:
 	FGPUSceneScopeBeginEndHelper() = delete;
 	FGPUSceneScopeBeginEndHelper(const FGPUSceneScopeBeginEndHelper&) = delete;
diff --git a/Engine/Source/Runtime/Renderer/Private/InstanceCulling/InstanceCullingManager.cpp b/Engine/Source/Runtime/Renderer/Private/InstanceCulling/InstanceCullingManager.cpp
--- a/Engine/Source/Runtime/Renderer/Private/InstanceCulling/InstanceCullingManager.cpp
+++ b/Engine/Source/Runtime/Renderer/Private/InstanceCulling/InstanceCullingManager.cpp
@@ -20,6 +20,9 @@ static TAutoConsoleVariable<int32> CVarCullInstances(
 	ECVF_RenderThreadSafe);
 
 
+#if defined(GPUCULL_TODO)
+
+
 int32 FInstanceCullingManager::RegisterView(const FViewInfo& ViewInfo)
 {
 	Nanite::FPackedViewParams Params;
@@ -40,7 +43,6 @@ int32 FInstanceCullingManager::RegisterView(const Nanite::FPackedViewParams& Par
 }
 
 
-
 class FCullInstancesCs : public FGlobalShader
 {
 	DECLARE_GLOBAL_SHADER(FCullInstancesCs);
@@ -84,8 +86,6 @@ public:
 IMPLEMENT_GLOBAL_SHADER(FCullInstancesCs, "/Engine/Private/InstanceCulling/CullInstances.usf", "CullInstancesCs", SF_Compute);
 
 
-
-
 void FInstanceCullingManager::CullInstances(FRDGBuilder& GraphBuilder, FGPUScene& GPUScene)
 {
 	RDG_EVENT_SCOPE(GraphBuilder, "CullInstances");
@@ -144,38 +144,39 @@ void FInstanceCullingManager::CullInstances(FRDGBuilder& GraphBuilder, FGPUScene
 			// All are visible
 			AddClearUAVPass(GraphBuilder, VisibleInstanceFlagsUAV, 0xFFFFFFFF);
 		}
-		//GraphBuilder.QueueBufferExtraction(VisibleInstanceFlagsRDG, &CullingIntermediate.VisibleInstanceFlags);
-		//ConvertToExternalBuffer(GraphBuilder, VisibleInstanceFlagsRDG, CullingIntermediate.VisibleInstanceFlags);
 	}
-	//GraphBuilder.QueueBufferExtraction(InstanceIdOutOffsetBufferRDG, &CullingIntermediate.InstanceIdOutOffsetBuffer);
-	//ConvertToExternalBuffer(GraphBuilder, InstanceIdOutOffsetBufferRDG, CullingIntermediate.InstanceIdOutOffsetBuffer);
 }
 
-#if 0
-void FInstanceCullingManager::ProcessGPUPostCullingJobs(FRDGBuilder& GraphBuilder, FGPUScene &GPUScene)
+
+FInstanceCullingContext* FInstanceCullingManager::CreateContext(const int32* ViewIds, int32 NumViews)
 {
-	RDG_EVENT_SCOPE(GraphBuilder, "ProcessGPUPostCullingJobs");
+	FInstanceCullingContext* InstanceCullingContext = new(FMemStack::Get()) FInstanceCullingContext;
+	InstanceCullingContext->ViewIds.Insert(ViewIds, NumViews, 0);
+	InstanceCullingContext->InstanceCullingManager = this;
+	return InstanceCullingContext;
+}
 
-	for (auto& GPUPostCullingJob : GPUPostCullingJobs)
-	{
-		GPUPostCullingJob(this, CullingIntermediate.InstanceIdOutOffsetBuffer, GraphBuilder, CullingIntermediate, GPUScene);
-	}
+
+#else //!defined(GPUCULL_TODO)
+
+int32 FInstanceCullingManager::RegisterView(const FViewInfo& ViewInfo)
+{
+	return 0;
 }
 
-void FInstanceCullingManager::QueueCullingContextForInstanceCulling(FInstanceCullingContext*& InstanceCullingContext, const TSharedPtr<FInstanceCullingResult> &InstanceCullingResult)
+int32 FInstanceCullingManager::RegisterView(const Nanite::FPackedViewParams& Params)
 {
-	QueueGPUPostCullingJob([InstanceCullingContext , &InstanceCullingResult](FInstanceCullingManager* InstanceCullingManager, const FRDGBufferRef& InstanceIdOutOffsetBufferRDG, FRDGBuilder& RDGBuilder, const FInstanceCullingIntermediate& Intermediate, FGPUScene &GPUScene)
-		{
-			InstanceCullingContext->BuildRenderingCommands(Intermediate, GPUScene, InstanceIdOutOffsetBufferRDG, RDGBuilder, *InstanceCullingResult);
-		});
+	return 0;
+}
 
+void FInstanceCullingManager::CullInstances(FRDGBuilder& GraphBuilder, FGPUScene& GPUScene)
+{
 }
-#endif
 
-FInstanceCullingContext* FInstanceCullingManager::CreateContext(const int32* ViewIds, int32 NumViews)
+FInstanceCullingContext* FInstanceCullingManager::CreateContext(const int32 * ViewIds, int32 NumViews)
 {
-	FInstanceCullingContext* InstanceCullingContext = new(FMemStack::Get()) FInstanceCullingContext;
-	InstanceCullingContext->ViewIds.Insert(ViewIds, NumViews, 0);
-	InstanceCullingContext->InstanceCullingManager = this;
-	return InstanceCullingContext;
+	return nullptr;
 }
+
+
+#endif // defined(GPUCULL_TODO)
diff --git a/Engine/Source/Runtime/Renderer/Private/InstanceCulling/InstanceCullingManager.h b/Engine/Source/Runtime/Renderer/Private/InstanceCulling/InstanceCullingManager.h
--- a/Engine/Source/Runtime/Renderer/Private/InstanceCulling/InstanceCullingManager.h
+++ b/Engine/Source/Runtime/Renderer/Private/InstanceCulling/InstanceCullingManager.h
@@ -94,24 +94,6 @@ public:
 	{
 	}
 
-#if 0
-	// Queue up jobs that need to run after the instance cull has been performed, this may not be a viable option,
-	// instead queue them up on a per-context basis or just do it just before submitting rendering commands.
-	// BUT rendering commands are submitted on their separate threads and I don't know that the RDG builder can handle that?
-	// SO TBD
-	FORCEINLINE void QueueGPUPostCullingJob(TFunction<void(FInstanceCullingManager*, const FRDGBufferRef&, FRDGBuilder&, const FInstanceCullingIntermediate&, FGPUScene&)>&& Job)
-	{
-		GPUPostCullingJobs.Add(MoveTemp(Job));
-	}
-
-	void ProcessGPUPostCullingJobs(FRDGBuilder& GraphBuilder, FGPUScene &GPUScene);
-	/**
-	 * Queue up the context, which after this point should not be touched, note that the pointer is cleared when this function is called, to indicate changed of ownership.
-	 * The result is returned via the CullingResults structure, which can be queried (and blocked upon).
-	 */
-	void QueueCullingContextForInstanceCulling(FInstanceCullingContext* &InstanceCullingContext, const TSharedPtr<FInstanceCullingResult>& InstanceCullingResult);
-#endif // 0
-
 	/**
 	 * Create a context for the instance culling, used to collect draw commands for a render pass. The context lifetime is managed by the Manager, and in general is a single-frame construct
 	 * Once the render pass setup is done, the Culling Context should be queued up for final processing (where the culled instances are used to fill in the draw commands).
diff --git a/Engine/Source/Runtime/Renderer/Private/MobileShadingRenderer.cpp b/Engine/Source/Runtime/Renderer/Private/MobileShadingRenderer.cpp
--- a/Engine/Source/Runtime/Renderer/Private/MobileShadingRenderer.cpp
+++ b/Engine/Source/Runtime/Renderer/Private/MobileShadingRenderer.cpp
@@ -319,6 +319,14 @@ void FMobileSceneRenderer::InitViews(FRDGBuilder& GraphBuilder, FSceneTexturesCo
 
 	check(Scene);
 
+#if defined(GPUCULL_TODO)
+	// Create GPU-side reprensenation of the view for instance culling.
+	for (int32 ViewIndex = 0; ViewIndex < Views.Num(); ++ViewIndex)
+	{
+		Views[ViewIndex].GPUSceneViewId = InstanceCullingManager.RegisterView(Views[ViewIndex]);
+	}
+#endif //defined(GPUCULL_TODO)
+
 	if (bUseVirtualTexturing)
 	{
 		RDG_GPU_STAT_SCOPE(GraphBuilder, VirtualTextureUpdate);
@@ -481,9 +489,13 @@ void FMobileSceneRenderer::InitViews(FRDGBuilder& GraphBuilder, FSceneTexturesCo
 	Scene->GPUScene.Update(GraphBuilder, *Scene);
 	for (int32 ViewIndex = 0; ViewIndex < Views.Num(); ViewIndex++)
 	{
-		Scene->GPUScene.UploadDynamicPrimitiveShaderDataForView(RHICmdList, *Scene, Views[ViewIndex]);
+		Scene->GPUScene.UploadDynamicPrimitiveShaderDataForView(RHICmdList, Scene, Views[ViewIndex]);
 	}
 
+	{
+		// GPUCULL_TODO: Possibly fold into unpack step
+		InstanceCullingManager.CullInstances(GraphBuilder, Scene->GPUScene);
+	}
 	extern TSet<IPersistentViewUniformBufferExtension*> PersistentViewUniformBufferExtensions;
 
 	for (IPersistentViewUniformBufferExtension* Extension : PersistentViewUniformBufferExtensions)
@@ -525,7 +537,7 @@ void FMobileSceneRenderer::Render(FRDGBuilder& GraphBuilder)
 	Scene->UpdateAllPrimitiveSceneInfos(GraphBuilder);
 
 	// Establish scene primitive count (must be done after UpdateAllPrimitiveSceneInfos)
-	FGPUSceneScopeBeginEndHelper GPUSceneScopeBeginEndHelper(Scene->GPUScene, GPUSceneDynamicContext, *Scene);
+	FGPUSceneScopeBeginEndHelper GPUSceneScopeBeginEndHelper(Scene->GPUScene, GPUSceneDynamicContext, Scene);
 
 	PrepareViewRectsForRendering();
 
diff --git a/Engine/Source/Runtime/Renderer/Private/PrimitiveSceneInfo.cpp b/Engine/Source/Runtime/Renderer/Private/PrimitiveSceneInfo.cpp
--- a/Engine/Source/Runtime/Renderer/Private/PrimitiveSceneInfo.cpp
+++ b/Engine/Source/Runtime/Renderer/Private/PrimitiveSceneInfo.cpp
@@ -460,30 +460,6 @@ void FPrimitiveSceneInfo::CacheMeshDrawCommands(FRHICommandListImmediate& RHICmd
 	}
 #endif
 
-	if (UseGPUScene(GMaxRHIShaderPlatform, Scene->GetFeatureLevel()))
-	{
-		SCOPE_CYCLE_COUNTER(STAT_UpdateGPUSceneTime);
-		for (FPrimitiveSceneInfo* SceneInfo : SceneInfos)
-		{
-			check(SceneInfo->InstanceDataOffset == INDEX_NONE);
-			check(SceneInfo->NumInstanceDataEntries == 0);
-			if (SceneInfo->Proxy->SupportsInstanceDataBuffer())
-			{
-				if (const TArray<FPrimitiveInstance>* PrimitiveInstances = SceneInfo->Proxy->GetPrimitiveInstances())
-				{
-					SceneInfo->InstanceDataOffset = Scene->GPUScene.AllocateInstanceSlots(PrimitiveInstances->Num());
-					SceneInfo->NumInstanceDataEntries = PrimitiveInstances->Num();
-				}
-			}
-
-			// Force a primitive update in the GPU scene
-			if (!Scene->GPUScene.PrimitivesMarkedToUpdate[SceneInfo->PackedIndex])
-			{
-				Scene->GPUScene.PrimitivesToUpdate.Add(SceneInfo->PackedIndex);
-				Scene->GPUScene.PrimitivesMarkedToUpdate[SceneInfo->PackedIndex] = true;
-			}
-		}
-	}
 }
 
 void FPrimitiveSceneInfo::RemoveCachedMeshDrawCommands()
@@ -553,17 +529,6 @@ void FPrimitiveSceneInfo::RemoveCachedMeshDrawCommands()
 		CachedRayTracingMeshCommandsHashPerLOD.Empty();
 	}
 #endif
-
-	// Release all instance data slots associated with this primitive.
-	if (InstanceDataOffset != INDEX_NONE && UseGPUScene(GMaxRHIShaderPlatform, Scene->GetFeatureLevel()))
-	{
-		SCOPE_CYCLE_COUNTER(STAT_UpdateGPUSceneTime);
-
-		check(Proxy->SupportsInstanceDataBuffer());
-		Scene->GPUScene.FreeInstanceSlots(InstanceDataOffset, NumInstanceDataEntries);
-		InstanceDataOffset = INDEX_NONE;
-		NumInstanceDataEntries = 0;
-	}
 }
 
 void FPrimitiveSceneInfo::CacheNaniteDrawCommands(FRHICommandListImmediate& RHICmdList, FScene* Scene, const TArrayView<FPrimitiveSceneInfo*>& SceneInfos)
@@ -906,6 +871,39 @@ void FPrimitiveSceneInfo::AddToScene(FRHICommandListImmediate& RHICmdList, FScen
 		}
 	}
 
+	if (Scene->GPUScene.IsEnabled())
+	{
+		SCOPE_CYCLE_COUNTER(STAT_UpdateGPUSceneTime);
+		for (FPrimitiveSceneInfo* SceneInfo : SceneInfos)
+		{
+			check(SceneInfo->InstanceDataOffset == INDEX_NONE);
+			check(SceneInfo->NumInstanceDataEntries == 0);
+			if (SceneInfo->Proxy->SupportsInstanceDataBuffer())
+			{
+				if (const TArray<FPrimitiveInstance>* PrimitiveInstances = SceneInfo->Proxy->GetPrimitiveInstances())
+				{
+					SceneInfo->InstanceDataOffset = Scene->GPUScene.AllocateInstanceSlots(PrimitiveInstances->Num());
+					SceneInfo->NumInstanceDataEntries = PrimitiveInstances->Num();
+				}
+			}
+#if defined(GPUCULL_TODO)
+			else
+			{
+				// Allocate a single 'dummy/fallback' instance for the primitive that gets automatically populated with the data from the primitive
+				SceneInfo->InstanceDataOffset = Scene->GPUScene.AllocateInstanceSlots(1);
+				SceneInfo->NumInstanceDataEntries = 1;
+			}
+#endif //defined(GPUCULL_TODO)
+
+			// Force a primitive update in the GPU scene
+			if (!Scene->GPUScene.PrimitivesMarkedToUpdate[SceneInfo->PackedIndex])
+			{
+				Scene->GPUScene.PrimitivesToUpdate.Add(SceneInfo->PackedIndex);
+				Scene->GPUScene.PrimitivesMarkedToUpdate[SceneInfo->PackedIndex] = true;
+			}
+		}
+	}
+
 	{
 		SCOPED_NAMED_EVENT(FPrimitiveSceneInfo_AddToScene_ReflectionCaptures, FColor::Yellow);
 		for (FPrimitiveSceneInfo* SceneInfo : SceneInfos)
@@ -1075,6 +1073,17 @@ void FPrimitiveSceneInfo::RemoveFromScene(bool bUpdateStaticDrawLists)
 		Scene->GPUScene.LightmapDataAllocator.Free(LightmapDataOffset, NumLightmapDataEntries);
 	}
 
+	// Release all instance data slots associated with this primitive.
+	if (InstanceDataOffset != INDEX_NONE && Scene->GPUScene.IsEnabled())
+	{
+		SCOPE_CYCLE_COUNTER(STAT_UpdateGPUSceneTime);
+
+		check(Proxy->SupportsInstanceDataBuffer() || NumInstanceDataEntries == 1);
+		Scene->GPUScene.FreeInstanceSlots(InstanceDataOffset, NumInstanceDataEntries);
+		InstanceDataOffset = INDEX_NONE;
+		NumInstanceDataEntries = 0;
+	}
+
 	if (Proxy->CastsDynamicIndirectShadow())
 	{
 		Scene->DynamicIndirectCasterPrimitives.RemoveSingleSwap(this);
diff --git a/Engine/Source/Runtime/Renderer/Private/ReflectionEnvironmentRealTimeCapture.cpp b/Engine/Source/Runtime/Renderer/Private/ReflectionEnvironmentRealTimeCapture.cpp
--- a/Engine/Source/Runtime/Renderer/Private/ReflectionEnvironmentRealTimeCapture.cpp
+++ b/Engine/Source/Runtime/Renderer/Private/ReflectionEnvironmentRealTimeCapture.cpp
@@ -512,7 +512,7 @@ void FScene::AllocateAndCaptureFrameSkyEnvMap(
 				CubeView.ViewUniformBuffer = CubeViewUniformBuffer;
 				if (CubeView.bSceneHasSkyMaterial)
 				{
-					GPUScene.UploadDynamicPrimitiveShaderDataForView(GraphBuilder.RHICmdList, *this, CubeView);
+					GPUScene.UploadDynamicPrimitiveShaderDataForView(GraphBuilder.RHICmdList, this, CubeView);
 				}
 
 				SkyRC.ViewUniformBuffer = CubeViewUniformBuffer;
diff --git a/Engine/Source/Runtime/Renderer/Private/RendererScene.cpp b/Engine/Source/Runtime/Renderer/Private/RendererScene.cpp
--- a/Engine/Source/Runtime/Renderer/Private/RendererScene.cpp
+++ b/Engine/Source/Runtime/Renderer/Private/RendererScene.cpp
@@ -1048,7 +1048,7 @@ FScene::FScene(UWorld* InWorld, bool bInRequiresHitProxies, bool bInIsEditorScen
 
 	FeatureLevel = World->FeatureLevel;
 
-	GPUScene.SetEnabled(UseGPUScene(GMaxRHIShaderPlatform, FeatureLevel));
+	GPUScene.SetEnabled(FeatureLevel);
 
 	if (World->FXSystem)
 	{
@@ -2597,6 +2597,11 @@ void FSceneVelocityData::StartFrame(FScene* Scene)
 
 		if (bTrimOld && (InternalFrameIndex - VelocityData.LastFrameUsed) > 10)
 		{
+			if (VelocityData.PrimitiveSceneInfo)
+			{
+				Scene->GPUScene.AddPrimitiveToUpdate(VelocityData.PrimitiveSceneInfo->GetIndex());
+			}
+
 			It.RemoveCurrent();
 		}
 	}
diff --git a/Engine/Source/Runtime/Renderer/Private/SceneHitProxyRendering.cpp b/Engine/Source/Runtime/Renderer/Private/SceneHitProxyRendering.cpp
--- a/Engine/Source/Runtime/Renderer/Private/SceneHitProxyRendering.cpp
+++ b/Engine/Source/Runtime/Renderer/Private/SceneHitProxyRendering.cpp
@@ -339,11 +339,22 @@ static void DoRenderHitProxies(
 	for (int32 ViewIndex = 0; ViewIndex < Views.Num(); ViewIndex++)
 	{
 		FViewInfo& View = const_cast<FViewInfo&>(Views[ViewIndex]);
-		const FScene* LocalScene = SceneRenderer->Scene;
+		FScene* LocalScene = SceneRenderer->Scene;
 		View.BeginRenderView();
 
 		auto* PassParameters = GraphBuilder.AllocParameters<FHitProxyPassParameters>();
 		PassParameters->View = View.GetShaderParameters();
+
+		// Adjust the visibility map for this view
+		if (View.bAllowTranslucentPrimitivesInHitProxy)
+		{
+			View.ParallelMeshDrawCommandPasses[EMeshPass::HitProxy].BuildRenderingCommands(GraphBuilder, LocalScene->GPUScene, PassParameters->InstanceCullingDrawParams);
+		}
+		else
+		{
+			View.ParallelMeshDrawCommandPasses[EMeshPass::HitProxyOpaqueOnly].BuildRenderingCommands(GraphBuilder, LocalScene->GPUScene, PassParameters->InstanceCullingDrawParams);
+		}
+
 		PassParameters->RenderTargets[0] = FRenderTargetBinding(HitProxyTexture, ERenderTargetLoadAction::ELoad);
 		PassParameters->RenderTargets.DepthStencil = FDepthStencilBinding(HitProxyDepthTexture, ERenderTargetLoadAction::ELoad, ERenderTargetLoadAction::ELoad, FExclusiveDepthStencil::DepthWrite_StencilWrite);
 		PassParameters->SceneTextures = CreateSceneTextureUniformBuffer(GraphBuilder, SceneRenderer->FeatureLevel, ESceneTextureSetupMode::None);
@@ -352,7 +363,7 @@ static void DoRenderHitProxies(
 			RDG_EVENT_NAME("RenderHitProxies"),
 			PassParameters,
 			ERDGPassFlags::Raster,
-			[SceneRenderer, &View, LocalScene, FeatureLevel, bNeedToSwitchVerticalAxis](FRHICommandList& RHICmdList)
+			[SceneRenderer, &View, LocalScene, FeatureLevel, bNeedToSwitchVerticalAxis, PassParameters](FRHICommandList& RHICmdList)
 		{
 			FMeshPassProcessorRenderState DrawRenderState;
 
@@ -368,11 +379,11 @@ static void DoRenderHitProxies(
 			// Adjust the visibility map for this view
 			if (View.bAllowTranslucentPrimitivesInHitProxy)
 			{
-				View.ParallelMeshDrawCommandPasses[EMeshPass::HitProxy].DispatchDraw(nullptr, RHICmdList);
+				View.ParallelMeshDrawCommandPasses[EMeshPass::HitProxy].DispatchDraw(nullptr, RHICmdList, &PassParameters->InstanceCullingDrawParams);
 			}
 			else
 			{
-				View.ParallelMeshDrawCommandPasses[EMeshPass::HitProxyOpaqueOnly].DispatchDraw(nullptr, RHICmdList);
+				View.ParallelMeshDrawCommandPasses[EMeshPass::HitProxyOpaqueOnly].DispatchDraw(nullptr, RHICmdList, &PassParameters->InstanceCullingDrawParams);
 			}
 
 			DrawDynamicMeshPass(View, RHICmdList,
@@ -579,6 +590,8 @@ void FMobileSceneRenderer::RenderHitProxies(FRDGBuilder& GraphBuilder)
 {
 	Scene->UpdateAllPrimitiveSceneInfos(GraphBuilder);
 
+	FGPUSceneScopeBeginEndHelper GPUSceneScopeBeginEndHelper(Scene->GPUScene, GPUSceneDynamicContext, Scene);
+
 	PrepareViewRectsForRendering();
 
 #if WITH_EDITOR
@@ -613,7 +626,7 @@ void FDeferredShadingSceneRenderer::RenderHitProxies(FRDGBuilder& GraphBuilder)
 
 	Scene->UpdateAllPrimitiveSceneInfos(GraphBuilder);
 
-	FGPUSceneScopeBeginEndHelper GPUSceneScopeBeginEndHelper(Scene->GPUScene, GPUSceneDynamicContext, *Scene);
+	FGPUSceneScopeBeginEndHelper GPUSceneScopeBeginEndHelper(Scene->GPUScene, GPUSceneDynamicContext, Scene);
 
 	PrepareViewRectsForRendering();
 
@@ -661,9 +674,11 @@ void FDeferredShadingSceneRenderer::RenderHitProxies(FRDGBuilder& GraphBuilder)
 
 	for (int32 ViewIndex = 0; ViewIndex < Views.Num(); ViewIndex++)
 	{
-		Scene->GPUScene.UploadDynamicPrimitiveShaderDataForView(GraphBuilder.RHICmdList, *Scene, Views[ViewIndex]);
+		Scene->GPUScene.UploadDynamicPrimitiveShaderDataForView(GraphBuilder.RHICmdList, Scene, Views[ViewIndex]);
 	}
 
+	InstanceCullingManager.CullInstances(GraphBuilder, Scene->GPUScene);
+
 	if (bNaniteEnabled)
 	{
 		Nanite::GGlobalResources.Update(GraphBuilder);
diff --git a/Engine/Source/Runtime/Renderer/Private/ShadowDepthRendering.cpp b/Engine/Source/Runtime/Renderer/Private/ShadowDepthRendering.cpp
--- a/Engine/Source/Runtime/Renderer/Private/ShadowDepthRendering.cpp
+++ b/Engine/Source/Runtime/Renderer/Private/ShadowDepthRendering.cpp
@@ -1047,7 +1047,7 @@ void FProjectedShadowInfo::BeginRenderView(FRDGBuilder& GraphBuilder, FScene* Sc
 	}
 
 	// This needs to be done for both mobile and deferred
-	Scene->GPUScene.UploadDynamicPrimitiveShaderDataForView(GraphBuilder.RHICmdList, *Scene, *ShadowDepthView);
+	Scene->GPUScene.UploadDynamicPrimitiveShaderDataForView(GraphBuilder.RHICmdList, Scene, *ShadowDepthView);
 }
 
 static bool IsShadowDepthPassWaitForTasksEnabled()
diff --git a/Engine/Source/Runtime/Renderer/Private/VT/RuntimeVirtualTextureRender.cpp b/Engine/Source/Runtime/Renderer/Private/VT/RuntimeVirtualTextureRender.cpp
--- a/Engine/Source/Runtime/Renderer/Private/VT/RuntimeVirtualTextureRender.cpp
+++ b/Engine/Source/Runtime/Renderer/Private/VT/RuntimeVirtualTextureRender.cpp
@@ -1043,7 +1043,7 @@ namespace RuntimeVirtualTexture
 		View->CachedViewUniformShaderParameters->RuntimeVirtualTexturePackHeight = WorldHeightPackParameter;
 		View->CachedViewUniformShaderParameters->RuntimeVirtualTextureDebugParams = FVector4(DebugType == ERuntimeVirtualTextureDebugType::Debug ? 1.f : 0.f, 0.f, 0.f, 0.f);
 		View->ViewUniformBuffer = TUniformBufferRef<FViewUniformShaderParameters>::CreateUniformBufferImmediate(*View->CachedViewUniformShaderParameters, UniformBuffer_SingleFrame);
-		Scene->GPUScene.UploadDynamicPrimitiveShaderDataForView(GraphBuilder.RHICmdList, *(const_cast<FScene*>(Scene)), *View);
+		Scene->GPUScene.UploadDynamicPrimitiveShaderDataForView(GraphBuilder.RHICmdList, (const_cast<FScene*>(Scene)), *View);
 
 		// Build graph
 		FRenderGraphSetup GraphSetup(GraphBuilder, Scene->GetFeatureLevel(), MaterialType, OutputTexture0, TextureSize, bIsThumbnails);
@@ -1164,7 +1164,7 @@ namespace RuntimeVirtualTexture
 		InDesc.Scene->UpdateAllPrimitiveSceneInfos(GraphBuilder, true);
 
 		// Call to let GPU-Scene determine if it is active and record scene primitive count
-		FGPUSceneScopeBeginEndHelper GPUSceneScopeBeginEndHelper(InDesc.Scene->GPUScene, GPUSceneDynamicContext, *InDesc.Scene);
+		FGPUSceneScopeBeginEndHelper GPUSceneScopeBeginEndHelper(InDesc.Scene->GPUScene, GPUSceneDynamicContext, InDesc.Scene);
 		InDesc.Scene->GPUScene.Update(GraphBuilder, *InDesc.Scene);
 		RenderPages(GraphBuilder, InDesc);
 	}
