Dynamic BVH with LBVH build.

TestBVH_ZeroToZero add [94.57ms]
TestBVH_ZeroToZero cost 10945156.00
TestBVH_ZeroToZero remove [6.62ms]
TestBVH_AddRemove [55.36ms]
TestBVH_AddRemove tested 12.13
TestBVH_BatchAddRemove [53.50ms]
TestBVH_UpdateTeleport [54.91ms]
TestBVH_UpdateJitter [44.68ms]
TestBVH_Ordered cost 283292.00
TestBVH_Ordered cost after 128869.00
TestBVH_Ordered cost after 85122.00
TestBVH_Optimize cost before 2697313.00
TestBVH_Optimize cost after 2405411.25
TestBVH_Optimize cost after 1598318.63
TestBVH_Build [5.79ms]

diff --git a/Engine/Source/Runtime/Renderer/Private/DynamicBVH.cpp b/Engine/Source/Runtime/Renderer/Private/DynamicBVH.cpp
new file mode 100644
--- /dev/null
+++ b/Engine/Source/Runtime/Renderer/Private/DynamicBVH.cpp
@@ -0,0 +1,332 @@
+// Copyright Epic Games, Inc. All Rights Reserved.
+
+#pragma once
+
+#include "DynamicBVH.h"
+#include "Misc/AutomationTest.h"
+
+
+FMortonArray::FMortonArray( const TArray< FBounds >& InBounds )
+	: Bounds( InBounds )
+{
+	TArray< FSortPair > Unsorted;
+	Unsorted.AddUninitialized( Bounds.Num() );
+	Sorted.AddUninitialized( Bounds.Num() );
+	
+	FBounds TotalBounds;
+	for( int i = 0; i < Bounds.Num(); i++ )
+	{
+		TotalBounds += Bounds[i].Min + Bounds[i].Max;
+	}
+
+	FVector Scale = FVector( 1.0f ) / ( TotalBounds.Max - TotalBounds.Min );
+	FVector Bias = -TotalBounds.Min / ( TotalBounds.Max - TotalBounds.Min );
+
+	for( int i = 0; i < Bounds.Num(); i++ )
+	{
+		FVector CenterLocal = ( Bounds[i].Min + Bounds[i].Max ) * Scale + Bias;
+
+		uint32 Morton;
+		Morton  = FMath::MortonCode3( CenterLocal.X * 1023 );
+		Morton |= FMath::MortonCode3( CenterLocal.Y * 1023 ) << 1;
+		Morton |= FMath::MortonCode3( CenterLocal.Z * 1023 ) << 2;
+
+		Unsorted[i].Code = Morton;
+		Unsorted[i].Index = i;
+	}
+
+	RadixSort32( Sorted.GetData(), Unsorted.GetData(), Unsorted.Num(),
+		[&]( FSortPair Pair )
+		{
+			return Pair.Code;
+		} );
+}
+
+void FMortonArray::RegenerateCodes( const FRange& Range )
+{
+	FBounds TotalBounds;
+	for( int32 i = Range.Begin; i < Range.End; i++ )
+	{
+		uint32 Index = Sorted[i].Index;
+		TotalBounds += Bounds[ Index ].Min + Bounds[ Index ].Max;
+	}
+
+	FVector Scale = FVector( 1.0f ) / ( TotalBounds.Max - TotalBounds.Min );
+	FVector Bias = -TotalBounds.Min / ( TotalBounds.Max - TotalBounds.Min );
+
+	for( int32 i = Range.Begin; i < Range.End; i++ )
+	{
+		uint32 Index = Sorted[i].Index;
+
+		FVector CenterLocal = ( Bounds[ Index ].Min + Bounds[ Index ].Max ) * Scale + Bias;
+
+		uint32 Morton;
+		Morton  = FMath::MortonCode3( CenterLocal.X * 1023 );
+		Morton |= FMath::MortonCode3( CenterLocal.Y * 1023 ) << 1;
+		Morton |= FMath::MortonCode3( CenterLocal.Z * 1023 ) << 2;
+
+		Sorted[i].Code = Morton;
+	}
+
+	Sort( &Sorted[ Range.Begin ], Range.End - Range.Begin );
+}
+
+
+inline FVector RandomVector( float Min, float Max )
+{
+	FVector V;
+	V.X = FMath::RandRange( Min, Max );
+	V.Y = FMath::RandRange( Min, Max );
+	V.Z = FMath::RandRange( Min, Max );
+	return V;
+}
+
+static FBounds RandomBounds( float CenterMin, float CenterMax, float ExtentMin, float ExtentMax )
+{
+	FVector Center = RandomVector( CenterMin, CenterMax );
+	FVector Extent = RandomVector( ExtentMin, ExtentMax );
+	return FBounds( { Center - Extent, Center + Extent } );
+}
+
+template< uint32 MaxChildren >
+static void	RandomBVH( float CenterMin, float CenterMax, float ExtentMin, float ExtentMax, uint32 Num, FDynamicBVH< MaxChildren >& BVH )
+{
+	for( uint32 i = 0; i < Num; i++ )
+	{
+		BVH.Add( RandomBounds( CenterMin, CenterMax, ExtentMin,  ExtentMax ), i );
+	}
+}
+
+void TestBVH_ZeroToZero()
+{
+	FMath::RandInit( 17 );
+
+	FDynamicBVH<4> BVH;
+	
+	uint32 Time0 = FPlatformTime::Cycles();
+
+	RandomBVH( -64.0f, 64.0f, 1.0f, 4.0f, 65536, BVH );
+
+	uint32 Num = BVH.GetNumLeaves();
+	uint32 Time1 = FPlatformTime::Cycles();
+
+	float Cost = BVH.GetTotalCost();
+
+	for( uint32 i = 0; i < Num; i++ )
+	{
+		BVH.Remove(i);
+	}
+
+	uint32 Time2 = FPlatformTime::Cycles();
+	GLog->Logf( TEXT("TestBVH_ZeroToZero add [%.2fms]"), FPlatformTime::ToMilliseconds( Time1 - Time0 ) );
+	GLog->Logf( TEXT("TestBVH_ZeroToZero cost %.2f"), Cost );
+	GLog->Logf( TEXT("TestBVH_ZeroToZero remove [%.2fms]"), FPlatformTime::ToMilliseconds( Time2 - Time1 ) );
+}
+
+void TestBVH_AddRemove()
+{
+	FMath::RandInit( 17 );
+
+	FDynamicBVH<4> BVH;
+	RandomBVH( -64.0f, 64.0f, 1.0f, 4.0f, 8192, BVH );
+
+	uint32 Num = BVH.GetNumLeaves();
+	uint32 Time0 = FPlatformTime::Cycles();
+
+	BVH.NumTested = 0;
+	
+	for( int i = 0; i < 65536; i++ )
+	{
+		BVH.Add( RandomBounds( -64.0f, 64.0f, 1.0f, 4.0f ), Num );
+		BVH.Remove( Num );
+	}
+
+	uint32 Time1 = FPlatformTime::Cycles();
+	GLog->Logf( TEXT("TestBVH_AddRemove [%.2fms]"), FPlatformTime::ToMilliseconds( Time1 - Time0 ) );
+	GLog->Logf( TEXT("TestBVH_AddRemove tested %.2f"), (float)BVH.NumTested / 65536.0f );
+
+	BVH.Check();
+}
+
+void TestBVH_BatchAddRemove()
+{
+	FMath::RandInit( 17 );
+
+	FDynamicBVH<4> BVH;
+	RandomBVH( -64.0f, 64.0f, 1.0f, 4.0f, 8192, BVH );
+
+	uint32 Num = BVH.GetNumLeaves();
+	uint32 Time0 = FPlatformTime::Cycles();
+
+	for( int i = 0; i < 256; i++ )
+	{
+		for( int j = 0; j < 256; j++ )
+			BVH.Add( RandomBounds( -64.0f, 64.0f, 1.0f, 4.0f ), Num + j );
+
+		for( int j = 0; j < 256; j++ )
+			BVH.Remove( Num + j );
+	}
+
+	uint32 Time1 = FPlatformTime::Cycles();
+	GLog->Logf( TEXT("TestBVH_BatchAddRemove [%.2fms]"), FPlatformTime::ToMilliseconds( Time1 - Time0 ) );
+
+	BVH.Check();
+}
+
+void TestBVH_UpdateTeleport()
+{
+	FMath::RandInit( 17 );
+
+	FDynamicBVH<4> BVH;
+	RandomBVH( -64.0f, 64.0f, 1.0f, 4.0f, 8192, BVH );
+
+	uint32 Num = BVH.GetNumLeaves();
+	uint32 Time0 = FPlatformTime::Cycles();
+
+	for( int i = 0; i < 65536; i++ )
+	{
+		uint32 Index = FMath::Rand() % Num;
+		BVH.Update( RandomBounds( -64.0f, 64.0f, 1.0f, 4.0f ), Index );
+	}
+
+	uint32 Time1 = FPlatformTime::Cycles();
+	GLog->Logf( TEXT("TestBVH_UpdateTeleport [%.2fms]"), FPlatformTime::ToMilliseconds( Time1 - Time0 ) );
+
+	BVH.Check();
+}
+
+void TestBVH_UpdateJitter()
+{
+	FMath::RandInit( 17 );
+	
+	FDynamicBVH<4> BVH;
+	RandomBVH( -64.0f, 64.0f, 1.0f, 4.0f, 8192, BVH );
+
+	uint32 Num = BVH.GetNumLeaves();
+	uint32 Time0 = FPlatformTime::Cycles();
+
+	for( int i = 0; i < 65536; i++ )
+	{
+		uint32 Index = FMath::Rand() % Num;
+
+		FBounds Bounds = BVH.GetBounds( Index );
+		FVector Jitter = RandomVector( -0.01f, 0.01f );
+		Bounds.Min += Jitter;
+		Bounds.Max += Jitter;
+
+		BVH.Update( Bounds, Index );
+	}
+
+	uint32 Time1 = FPlatformTime::Cycles();
+	GLog->Logf( TEXT("TestBVH_UpdateJitter [%.2fms]"), FPlatformTime::ToMilliseconds( Time1 - Time0 ) );
+
+	BVH.Check();
+}
+
+void TestBVH_Ordered()
+{
+	FBounds UnitBounds;
+	UnitBounds.Min = FVector::ZeroVector;
+	UnitBounds.Max = FVector::OneVector;
+
+	FDynamicBVH<4> BVH;
+
+	for( uint32 i = 0; i < 8192; i++ )
+	{
+		FBounds Bounds = UnitBounds;
+		Bounds.Min.X += (float)i;
+		Bounds.Max.X += (float)i;
+		BVH.Add( Bounds, i );
+	}
+
+	GLog->Logf( TEXT("TestBVH_Ordered cost %.2f"), BVH.GetTotalCost() );
+
+	uint32 Num = BVH.GetNumLeaves();
+
+	for( int i = 0; i < 1024; i++ )
+	{
+		uint32 Index = FMath::Rand() % Num;
+		BVH.Update( BVH.GetBounds( Index ), Index );
+		//BVH.Optimize(1);
+	}
+
+	GLog->Logf( TEXT("TestBVH_Ordered cost after %.2f"), BVH.GetTotalCost() );
+
+	for( int i = 0; i < 65536; i++ )
+	{
+		uint32 Index = FMath::Rand() % Num;
+		BVH.Update( BVH.GetBounds( Index ), Index );
+		//BVH.Optimize(1);
+	}
+
+	GLog->Logf( TEXT("TestBVH_Ordered cost after %.2f"), BVH.GetTotalCost() );
+}
+
+void TestBVH_Optimize()
+{
+	FMath::RandInit( 17 );
+	
+	FDynamicBVH<4> BVH;
+	RandomBVH( -64.0f, 64.0f, 1.0f, 4.0f, 8192, BVH );
+
+	GLog->Logf( TEXT("TestBVH_Optimize cost before %.2f"), BVH.GetTotalCost() );
+
+	uint32 Num = BVH.GetNumLeaves();
+
+	for( int i = 0; i < 1024; i++ )
+	{
+		uint32 Index = FMath::Rand() % Num;
+		BVH.Update( BVH.GetBounds( Index ), Index );
+		//BVH.Optimize(1);
+	}
+
+	GLog->Logf( TEXT("TestBVH_Optimize cost after %.2f"), BVH.GetTotalCost() );
+
+	for( int i = 0; i < 65536; i++ )
+	{
+		uint32 Index = FMath::Rand() % Num;
+		BVH.Update( BVH.GetBounds( Index ), Index );
+		//BVH.Optimize(1);
+	}
+
+	GLog->Logf( TEXT("TestBVH_Optimize cost after %.2f"), BVH.GetTotalCost() );
+}
+
+void TestBVH_Build()
+{
+	FMath::RandInit( 17 );
+
+	FDynamicBVH<4> BVH;
+
+	uint32 Num = 65536;
+
+	TArray< FBounds > BoundsArray;
+	BoundsArray.AddUninitialized( Num );
+	
+	for( uint32 i = 0; i < Num; i++ )
+		BoundsArray[i] = RandomBounds( -64.0f, 64.0f, 1.0f, 4.0f );
+	
+	uint32 Time0 = FPlatformTime::Cycles();
+
+	BVH.Build( BoundsArray, 0 );
+
+	uint32 Time1 = FPlatformTime::Cycles();
+
+	GLog->Logf( TEXT("TestBVH_Build [%.2fms]"), FPlatformTime::ToMilliseconds( Time1 - Time0 ) );
+	GLog->Logf( TEXT("TestBVH_Build cost %.2f"), BVH.GetTotalCost() );
+}
+
+IMPLEMENT_SIMPLE_AUTOMATION_TEST( FTestBVH, "System.Renderer.DynamicBVH", EAutomationTestFlags::EditorContext | EAutomationTestFlags::EngineFilter );
+bool FTestBVH::RunTest( const FString& Parameters )
+{
+	TestBVH_ZeroToZero();
+	TestBVH_AddRemove();
+	TestBVH_BatchAddRemove();
+	TestBVH_UpdateTeleport();
+	TestBVH_UpdateJitter();
+	TestBVH_Ordered();
+	TestBVH_Optimize();
+	TestBVH_Build();
+
+	return true;
+}
\ No newline at end of file
diff --git a/Engine/Source/Runtime/Renderer/Private/DynamicBVH.h b/Engine/Source/Runtime/Renderer/Private/DynamicBVH.h
new file mode 100644
--- /dev/null
+++ b/Engine/Source/Runtime/Renderer/Private/DynamicBVH.h
@@ -0,0 +1,1087 @@
+// Copyright Epic Games, Inc. All Rights Reserved.
+
+#pragma once
+
+#include "CoreMinimal.h"
+#include "Misc/AutomationTest.h"
+
+struct FBounds
+{
+	FVector4	Min = FVector4(  MAX_flt,  MAX_flt,  MAX_flt );
+	FVector4	Max = FVector4( -MAX_flt, -MAX_flt, -MAX_flt );
+
+	FORCEINLINE FBounds& operator=( const FVector& Other )
+	{
+		Min = Other;
+		Max = Other;
+		return *this;
+	}
+
+	FORCEINLINE FBounds& operator+=( const FVector& Other )
+	{
+		VectorStoreAligned( VectorMin( VectorLoadAligned( &Min ), VectorLoadAligned( &Other ) ), &Min );
+		VectorStoreAligned( VectorMax( VectorLoadAligned( &Max ), VectorLoadAligned( &Other ) ), &Max );
+		return *this;
+	}
+
+	FORCEINLINE FBounds& operator+=( const FBounds& Other )
+	{
+		VectorStoreAligned( VectorMin( VectorLoadAligned( &Min ), VectorLoadAligned( &Other.Min ) ), &Min );
+		VectorStoreAligned( VectorMax( VectorLoadAligned( &Max ), VectorLoadAligned( &Other.Max ) ), &Max );
+		return *this;
+	}
+
+	FORCEINLINE FBounds operator+( const FBounds& Other ) const
+	{
+		return FBounds(*this) += Other;
+	}
+
+	FORCEINLINE float GetSurfaceArea() const
+	{
+		FVector Size = Max - Min;
+		return 0.5f * ( Size.X * Size.Y + Size.X * Size.Z + Size.Y * Size.Z );
+	}
+
+	FORCEINLINE friend FArchive& operator<<( FArchive& Ar, FBounds& Bounds )
+	{
+		Ar << Bounds.Min;
+		Ar << Bounds.Max;
+		return Ar;
+	}
+};
+
+struct FSurfaceAreaHeuristic
+{
+	float operator()( const FBounds& Bounds ) const
+	{
+		FVector Extent = Bounds.Max - Bounds.Min;
+		return Extent.X * Extent.Y + Extent.X * Extent.Z + Extent.Y * Extent.Z;
+	}
+};
+
+constexpr uint32 ConstLog2( uint32 x )
+{
+	return ( x < 2 ) ? 0 : 1 + ConstLog2( x / 2 );
+}
+
+template< uint32 MaxChildren, typename FCostMetric = FSurfaceAreaHeuristic >
+class FDynamicBVH
+{
+	FCostMetric		CostMetric;
+
+public:
+					FDynamicBVH();
+
+	const FBounds&	GetBounds() const { return RootBounds; }
+	int32			GetNumNodes() const { return Nodes.Num(); }
+	int32			GetNumLeaves() const { return Leaves.Num(); }
+	int32			GetNumDirty() const { return DirtyNodes.Num(); }
+
+	void			Add( const FBounds& Bounds, uint32 Index );
+	void			Update( const FBounds& Bounds, uint32 Index );
+	void			Remove( uint32 Index );
+	
+	void			AddDefaulted()		{ Leaves.Add( ~0u ); }
+	void			SwapIndexes( uint32 Index0, uint32 Index1 );
+
+	void			Build( const TArray< FBounds >& BoundsArray, uint32 FirstIndex );
+
+	const FBounds&	GetBounds( uint32 Index ) const
+	{
+		check( Leaves[ Index ] != ~0u );
+		uint32 NodeIndex = Leaves[ Index ];
+		return GetNode( NodeIndex ).GetBounds( NodeIndex );
+	}
+
+	template< typename FPredicate, typename FFuncType >
+	void			ForAll( const FPredicate& Predicate, const FFuncType& Func ) const;
+
+	template< typename FFuncType >
+	void			ForAllDirty( const FFuncType& Func );
+
+	float			GetTotalCost() const
+	{
+		float TotalCost = 0.0f;
+		for( auto& Node : Nodes )
+		{
+			for( uint32 i = 0; i < Node.NumChildren; i++ )
+			{
+				if( ( Node.ChildIndexes[i] & 1 ) == 0 )
+					TotalCost += CostMetric( Node.ChildBounds[i] );
+			}
+		}
+
+		return TotalCost;
+	}
+
+	bool			Check() const
+	{
+		for( int32 i = 0; i < Nodes.Num(); i++ )
+		{
+			for( uint32 j = 0; j < Nodes[i].NumChildren; j++ )
+			{
+				CheckNode( (i << IndexShift) | j );
+			}
+		}
+
+		return true;
+	}
+
+	uint32 NumTested = 0;
+
+protected:
+	static constexpr uint32	IndexShift = ConstLog2( MaxChildren );
+	static constexpr uint32	ChildMask = MaxChildren - 1;
+	static constexpr uint32 MaxChildren4 = (MaxChildren + 3) / 4;
+
+	struct FNode
+	{
+		// Index: low bits index child, high bits index FNode
+		uint32	ParentIndex;
+		uint32	NumChildren;		// Lots of bits for this!
+		uint32	ChildIndexes[ MaxChildren ];
+		FBounds	ChildBounds[ MaxChildren ];
+		
+		uint32			GetFirstChild( uint32 NodeIndex ) const	{ return ChildIndexes[ NodeIndex & ChildMask ]; }
+		const FBounds&	GetBounds( uint32 NodeIndex ) const		{ return ChildBounds[ NodeIndex & ChildMask ]; }
+		
+		bool	IsRoot() const						{ return ParentIndex == ~0u; }
+		bool	IsLeaf( uint32 NodeIndex ) const	{ return GetFirstChild( NodeIndex ) & 1; }
+		bool	IsFull() const						{ return NumChildren == MaxChildren; }
+
+		FBounds	UnionBounds() const
+		{
+			FBounds Bounds;
+			for( uint32 i = 0; i < NumChildren; i++ )
+			{
+				Bounds += ChildBounds[i];
+			}
+			return Bounds;
+		}
+	};
+
+	TArray< FNode >		Nodes;
+	TArray< uint32 >	Leaves;
+	TArray< uint32 >	FreeList;
+	FBounds				RootBounds;
+
+	TBitArray<>			NodeIsDirty;
+	TArray< uint32 >	DirtyNodes;
+
+	using FCandidate = TPair< float, uint32 >;
+	TArray< FCandidate > Candidates;
+
+protected:
+	FNode&			GetNode( uint32 NodeIndex )			{ return Nodes[ NodeIndex >> IndexShift ]; }
+	const FNode&	GetNode( uint32 NodeIndex ) const	{ return Nodes[ NodeIndex >> IndexShift ]; }
+
+	void	MarkDirty( uint32 NodeIndex )
+	{
+		if( !NodeIsDirty[ NodeIndex >> IndexShift ] )
+		{
+			NodeIsDirty[ NodeIndex >> IndexShift ] = true;
+			DirtyNodes.Add( NodeIndex >> IndexShift );
+		}
+	}
+
+	void	Set( uint32 NodeIndex, const FBounds& Bounds, uint32 FirstChild )
+	{
+		GetNode( NodeIndex ).ChildBounds[ NodeIndex & ChildMask ] = Bounds;
+		SetFirstChild( NodeIndex, FirstChild );
+	}
+	
+	void	SetBounds( uint32 NodeIndex, const FBounds& Bounds )
+	{
+		GetNode( NodeIndex ).ChildBounds[ NodeIndex & ChildMask ] = Bounds;
+		MarkDirty( NodeIndex );
+	}
+
+	void	SetFirstChild( uint32 NodeIndex, uint32 FirstChild )
+	{
+		GetNode( NodeIndex ).ChildIndexes[ NodeIndex & ChildMask ] = FirstChild;
+		MarkDirty( NodeIndex );
+		
+		if( FirstChild & 1 )
+		{
+			Leaves[ FirstChild >> 1 ] = NodeIndex;
+		}
+		else
+		{
+			GetNode( FirstChild ).ParentIndex = NodeIndex;
+			MarkDirty( FirstChild );
+		}
+	}
+
+	uint32	FindBestInsertion_BranchAndBound( const FBounds& RESTRICT Bounds );
+	uint32	FindBestInsertion_Greedy( const FBounds& RESTRICT Bounds );
+
+	uint32	Insert( const FBounds& RESTRICT Bounds, uint32 NodeIndex );
+	void	Extract( uint32 NodeIndex );
+	uint32	PromoteChild( uint32 NodeIndex );
+	void	Rotate( uint32 NodeIndex );
+	
+	uint32	AllocNode();
+	void	FreeNode( uint32 NodeIndex );
+	
+	void	CheckNode( uint32 NodeIndex ) const;
+};
+
+template< uint32 MaxChildren, typename FCostMetric >
+FDynamicBVH< MaxChildren, FCostMetric >::FDynamicBVH()
+{
+	static_assert( MaxChildren > 1, "Must at least be binary tree." );
+	static_assert( ( MaxChildren & (MaxChildren - 1) ) == 0, "MaxChildren must be power of 2" );
+
+	Nodes.AddDefaulted();
+	NodeIsDirty.Add( true );
+
+	Nodes[0].ParentIndex = ~0u;
+	Nodes[0].NumChildren = 0;
+}
+
+template< uint32 MaxChildren, typename FCostMetric >
+FORCEINLINE void FDynamicBVH< MaxChildren, FCostMetric >::Add( const FBounds& Bounds, uint32 Index )
+{
+	if( Index >= (uint32)Leaves.Num() )
+	{
+		int32 Count = Index + 1 - Leaves.Num();
+		int32 First = Leaves.AddUninitialized( Count );
+		FMemory::Memset( &Leaves[ First ], 0xff, Count * sizeof( uint32 ) );
+	}
+
+	check( Leaves[ Index ] == ~0u );
+	Leaves[ Index ] = Insert( Bounds, ( Index << 1 ) | 1 );
+
+	//CheckNode( Leaves[ Index ] );
+}
+
+template< uint32 MaxChildren, typename FCostMetric >
+FORCEINLINE void FDynamicBVH< MaxChildren, FCostMetric >::Update( const FBounds& Bounds, uint32 Index )
+{
+	Remove( Index );
+	Add( Bounds, Index );
+}
+
+template< uint32 MaxChildren, typename FCostMetric >
+FORCEINLINE void FDynamicBVH< MaxChildren, FCostMetric >::Remove( uint32 Index )
+{
+	check( Leaves[ Index ] != ~0u );
+	check( GetNode( Leaves[ Index ] ).GetFirstChild( Leaves[ Index ] ) == ( (Index << 1) | 1 ) );
+
+	Extract( Leaves[ Index ] );
+	Leaves[ Index ] = ~0u;
+}
+
+template< uint32 MaxChildren, typename FCostMetric >
+FORCEINLINE void FDynamicBVH< MaxChildren, FCostMetric >::SwapIndexes( uint32 Index0, uint32 Index1 )
+{
+	Swap( Leaves[ Index0 ], Leaves[ Index1 ] );
+		
+	uint32 NodeIndex0 = Leaves[ Index0 ];
+	uint32 NodeIndex1 = Leaves[ Index1 ];
+
+	if( NodeIndex0 != ~0u )
+	{
+		GetNode( NodeIndex0 ).ChildIndexes[ NodeIndex0 & ChildMask ] = (Index0 << 1) | 1;
+		MarkDirty( NodeIndex0 );
+	}
+
+	if( NodeIndex1 != ~0u )
+	{
+		GetNode( NodeIndex1 ).ChildIndexes[ NodeIndex1 & ChildMask ] = (Index1 << 1) | 1;
+		MarkDirty( NodeIndex1 );
+	}
+}
+
+
+template< uint32 MaxChildren, typename FCostMetric >
+uint32 FDynamicBVH< MaxChildren, FCostMetric >::FindBestInsertion_BranchAndBound( const FBounds& RESTRICT Bounds )
+{
+	// Uses branch and bound search algorithm outlined in:
+	// [ Bittner et al. 2012, "Fast Insertion-Based Optimization of Bounding Volume Hierarchies" ]
+	
+	// Binary tree nodes besides the root are always full meaning a new level will always be added.
+	float MinAddedCost = MaxChildren > 2 ? 0.0f : CostMetric( Bounds );
+
+	// Find best node to merge with.
+	float	BestCost = MAX_flt;
+	uint32	BestIndex = 0;
+
+	int32 CandidateHead = 0;
+	Candidates.Reset();
+
+	constexpr uint32 MaxZeros = 32;
+	uint32 NumZeros = 0;
+	uint32 ZeroCostNodes[ MaxZeros ];
+	
+	float InducedCost = 0.0f;
+	uint32 NodeIndex = 0;
+
+	while( true )
+	{
+		const FNode& RESTRICT Node = GetNode( NodeIndex );
+		NumTested++;
+
+		if( Node.IsFull() )
+		{
+			for( uint32 i = 0; i < Node.NumChildren; i += 4 )
+			{
+				FVector4 TotalCost;
+				FVector4 ChildCost;
+				constexpr uint32 Four = MaxChildren < 4 ? MaxChildren : 4;
+				for( uint32 j = 0; j < Four; j++ )
+				{
+					const FBounds& RESTRICT NodeBounds = Node.ChildBounds[ i + j ];
+
+					float DirectCost = CostMetric( Bounds + NodeBounds );
+					// Cost if we need to add a level
+					TotalCost[j] = InducedCost + DirectCost;
+					// Induced cost for children
+					ChildCost[j] = TotalCost[j] - CostMetric( NodeBounds );
+				}
+
+				for( uint32 j = 0; i + j < Node.NumChildren; j++ )
+				{
+					if( ChildCost[j] < BestCost )
+					{
+						if( TotalCost[j] < BestCost )
+						{
+							BestCost = TotalCost[j];
+							BestIndex = NodeIndex + i + j;
+						}
+
+						uint32 FirstChild = Node.ChildIndexes[ i + j ];
+						bool bIsLeaf = FirstChild & 1;
+						if( !bIsLeaf )
+						{
+							if( ChildCost[j] < SMALL_NUMBER && NumZeros < MaxZeros )
+							{
+								ZeroCostNodes[ NumZeros++ ] = FirstChild;
+							}
+							else
+							{
+								Candidates.Add( FCandidate( ChildCost[j], FirstChild ) );
+							}
+						}
+					}
+				}
+			}
+		}
+		else
+		{
+			// Don't need to add a level because we can add a child.
+			if( InducedCost < BestCost )
+			{
+				// Can't do better as this was already the smallest from the heap.
+				return Node.ParentIndex;
+			}
+		}
+		
+		if( NumZeros )
+		{
+			NodeIndex = ZeroCostNodes[ --NumZeros ];
+		}
+		else
+		{
+			// Move head up
+			int32 Num = Candidates.Num();
+			while( CandidateHead < Num && Candidates[ CandidateHead ].Key >= BestCost )
+				CandidateHead++;
+			if( CandidateHead == Num )
+				break;
+
+			// Find smallest linear search
+			float SmallestCost	= Candidates[ CandidateHead ].Key;
+			int32 SmallestIndex	= CandidateHead;
+			for( int32 i = CandidateHead + 1; i < Num; i++ )
+			{
+				float Cost = Candidates[i].Key;
+				if( Cost < SmallestCost )
+				{
+					SmallestCost = Cost;
+					SmallestIndex = i;
+				}
+			}
+
+			InducedCost = SmallestCost;
+			NodeIndex = Candidates[ SmallestIndex ].Value;
+
+			Candidates.RemoveAtSwap( SmallestIndex, 1, false );
+		}
+		
+		if( InducedCost + MinAddedCost >= BestCost )
+		{
+			// Not possible to reduce cost further.
+			break;
+		}
+	}
+
+	return BestIndex;
+}
+
+template< uint32 MaxChildren, typename FCostMetric >
+uint32 FDynamicBVH< MaxChildren, FCostMetric >::FindBestInsertion_Greedy( const FBounds& RESTRICT Bounds )
+{
+	// Binary tree nodes besides the root are always full meaning a new level will always be added.
+	float MinAddedCost = MaxChildren > 2 ? 0.0f : CostMetric( Bounds );
+
+	// Find best node to merge with.
+	float	BestCost = MAX_flt;
+	uint32	BestIndex = 0;
+
+	float InducedCost = 0.0f;
+	uint32 NodeIndex = 0;
+
+	do
+	{
+		FNode& RESTRICT Node = GetNode( NodeIndex );
+		NumTested++;
+
+		if( Node.IsFull() )
+		{
+			float	BestChildDist = MAX_flt;
+			float	BestChildCost = MAX_flt;
+			uint32	BestChildIndex = ~0u;
+			for( uint32 i = 0; i < Node.NumChildren; i += 4 )
+			{
+				FVector4 TotalCost;
+				FVector4 ChildCost;
+				FVector4 Dist;
+				constexpr uint32 Four = MaxChildren < 4 ? MaxChildren : 4;
+				for( uint32 j = 0; j < Four; j++ )
+				{
+					const FBounds& RESTRICT NodeBounds = Node.ChildBounds[ i + j ];
+
+					float DirectCost = CostMetric( Bounds + NodeBounds );
+					// Cost if we need to add a level
+					TotalCost[j] = InducedCost + DirectCost;
+					// Induced cost for children
+					ChildCost[j] = TotalCost[j] - CostMetric( NodeBounds );
+
+					FVector Delta = ( Bounds.Min + Bounds.Max ) - ( NodeBounds.Min + NodeBounds.Max );
+					Delta = Delta.GetAbs();
+					Dist[j] = Delta.X + Delta.Y + Delta.Z;
+					Dist[j] += ChildCost[j] * 4.0f;
+				}
+
+				for( uint32 j = 0; i + j < Node.NumChildren; j++ )
+				{
+					if( ChildCost[j] < BestCost )
+					{
+						if( TotalCost[j] < BestCost )
+						{
+							BestCost = TotalCost[j];
+							BestIndex = NodeIndex + i + j;
+						}
+
+						uint32 FirstChild = Node.ChildIndexes[ i + j ];
+						bool bIsLeaf = FirstChild & 1;
+						if( !bIsLeaf )
+						{
+							// Pick only 1 child to continue
+							//if( ChildCost < BestChildCost )
+							if( Dist[j] < BestChildDist )
+							{
+								BestChildDist = Dist[j];
+								BestChildCost = ChildCost[j];
+								BestChildIndex = FirstChild;
+							}
+						}
+					}
+				}
+			}
+
+			InducedCost	= BestChildCost;
+			NodeIndex	= BestChildIndex;
+		}
+		else
+		{
+			// Don't need to add a level because we can add a child.
+			// Can't do better. Cost is monotonic.
+			return Node.ParentIndex;
+		}
+		
+		if( InducedCost + MinAddedCost >= BestCost )
+		{
+			// Not possible to reduce cost further.
+			break;
+		}
+	}
+	while( NodeIndex != ~0u );
+
+	return BestIndex;
+}
+
+template< uint32 MaxChildren, typename FCostMetric >
+uint32 FDynamicBVH< MaxChildren, FCostMetric >::Insert( const FBounds& RESTRICT Bounds, uint32 Index )
+{
+	FNode& Root = Nodes[0];
+	if( !Root.IsFull() )
+	{
+		uint32 NodeIndex = Root.NumChildren++;
+		Set( NodeIndex, Bounds, Index );
+		RootBounds += Bounds;
+		return NodeIndex;
+	}
+
+	uint32 BestIndex = FindBestInsertion_BranchAndBound( Bounds );
+	//uint32 BestIndex = FindBestInsertion_Greedy( Bounds );
+	
+	// Add to BestIndex's children
+	uint32 NodeIndex = GetNode( BestIndex ).GetFirstChild( BestIndex );
+	bool bIsLeaf = NodeIndex & 1;
+	bool bAddLevel = bIsLeaf || GetNode( NodeIndex ).IsFull();
+	if( bAddLevel )
+	{
+		// Create a new node and add NodeIndex as a child.
+		uint32 NewNodeIndex = AllocNode();
+		FNode& NewNode = GetNode( NewNodeIndex );
+
+		NewNode.NumChildren = 1;
+		Set( NewNodeIndex,
+			GetNode( BestIndex ).GetBounds( BestIndex ),
+			GetNode( BestIndex ).GetFirstChild( BestIndex ) );
+
+		SetFirstChild( BestIndex, NewNodeIndex );
+
+		checkSlow( NewNode.ParentIndex == BestIndex );
+		checkSlow( NewNode.ChildIndexes[0] == NodeIndex );
+
+		NodeIndex = NewNodeIndex;
+	}
+	
+	FNode& Children = GetNode( NodeIndex );
+
+	// Add child
+	NodeIndex |= Children.NumChildren++;
+	Set( NodeIndex, Bounds, Index );
+
+	// Propagate bounds up tree
+	FBounds PathBounds	= Bounds;
+	uint32  PathIndex	= BestIndex;
+	while( PathIndex != ~0u )
+	{
+		FNode& PathNode = GetNode( PathIndex );
+		SetBounds( PathIndex, PathNode.GetBounds( PathIndex ) + PathBounds );
+
+		Rotate( PathIndex );
+
+		PathBounds	= PathNode.GetBounds( PathIndex );
+		PathIndex	= PathNode.ParentIndex;
+	}
+	RootBounds += PathBounds;
+
+	return NodeIndex;
+}
+
+template< uint32 MaxChildren, typename FCostMetric >
+void FDynamicBVH< MaxChildren, FCostMetric >::Extract( uint32 NodeIndex )
+{
+	FNode& Node = GetNode( NodeIndex );
+	check( Node.IsRoot() || Node.NumChildren > 1 );
+
+	uint32 LastChild = --Node.NumChildren;
+	if( ( NodeIndex & ChildMask ) < LastChild )
+	{
+		// Fill with last
+		Set( NodeIndex,
+			Node.GetBounds( LastChild ),
+			Node.GetFirstChild( LastChild ) );
+	}
+	
+	// Propagate bounds up tree
+	FBounds PathBounds	= Node.UnionBounds();
+	uint32  PathIndex	= Node.ParentIndex;
+	while( PathIndex != ~0u )
+	{
+		SetBounds( PathIndex, PathBounds );
+
+		FNode& PathNode = GetNode( PathIndex );
+		PathBounds	= PathNode.UnionBounds();
+		PathIndex	= PathNode.ParentIndex;
+	}
+	RootBounds = PathBounds;
+
+	if( !Node.IsRoot() && Node.NumChildren == 1 )
+	{
+		FNode& ParentNode = GetNode( Node.ParentIndex );
+
+		Set( Node.ParentIndex,
+			Node.ChildBounds[0],
+			Node.ChildIndexes[0] );
+			
+		FreeNode( NodeIndex );
+	}
+	else
+	{
+		// Recursively promotes children to fill the hole until it reaches a leaf.
+		// The result of doing this on every Extract is that all inner nodes are guarenteed to be full.
+		do
+		{
+			FNode& PathNode = GetNode( NodeIndex );
+
+			// Find best node to promote a child from.
+			float	BestCost = 0.0f;
+			uint32	BestIndex = ~0u;
+			for( uint32 i = 0; i < PathNode.NumChildren; i++ )
+			{
+				if( !PathNode.IsLeaf(i) )
+				{
+					float Cost = CostMetric( PathNode.ChildBounds[i] );
+					if( Cost > BestCost )
+					{
+						BestCost = Cost;
+						BestIndex = ( NodeIndex & ~ChildMask ) | i;
+					}
+				}
+			}
+
+			if( BestIndex == ~0u )
+				break;
+
+			NodeIndex = PromoteChild( BestIndex );
+		}
+		while( NodeIndex != ~0u );
+	}
+}
+
+template< uint32 MaxChildren, typename FCostMetric >
+uint32 FDynamicBVH< MaxChildren, FCostMetric >::PromoteChild( uint32 NodeIndex )
+{
+	FNode& RESTRICT Node = GetNode( NodeIndex );
+	check( !Node.IsLeaf( NodeIndex ) );
+	check( Node.NumChildren < MaxChildren );
+
+	uint32 FirstChild = Node.GetFirstChild( NodeIndex );
+	FNode& RESTRICT Children = GetNode( FirstChild );
+	
+	FBounds Excluded[ MaxChildren ];
+
+	// Sweep forward and backward with prefix and postfix sums. Prefix + postfix sums == sum excluding this.
+	FBounds Forward;
+	FBounds Back;
+	for( uint32 i = 0; i < Children.NumChildren; i++ )
+	{
+		uint32 j = Children.NumChildren - 1 - i;
+
+		Excluded[i] += Forward;
+		Excluded[j] += Back;
+
+		Forward	+= Children.ChildBounds[i];
+		Back	+= Children.ChildBounds[j];
+	}
+
+	float	BestCost = MAX_flt;
+	uint32	BestIndex = ~0u;
+	
+	for( uint32 i = 0; i < Children.NumChildren; i++ )
+	{
+		float Cost = CostMetric( Excluded[i] );
+		if( Cost < BestCost )
+		{
+			BestCost = Cost;
+			BestIndex = FirstChild | i;
+		}
+	}
+	
+	// Promote from child to sibling
+	
+	// Remove from bounds
+	SetBounds( NodeIndex, Excluded[ BestIndex & ChildMask ] );
+
+	// Add as sibling
+	uint32 SiblingIndex = ( NodeIndex & ~ChildMask ) | Node.NumChildren;
+	Set( SiblingIndex,
+		Children.GetBounds( BestIndex ),
+		Children.GetFirstChild( BestIndex ) );
+	Node.NumChildren++;
+
+	// Remove from children
+	uint32 LastChild = --Children.NumChildren;
+	if( Children.NumChildren == 1 )
+	{
+		uint32 OtherChild = ~BestIndex & 1;
+
+		Set( NodeIndex,
+			Children.ChildBounds[ OtherChild ],
+			Children.ChildIndexes[ OtherChild ] );
+		
+		// Delete Children
+		FreeNode( BestIndex );
+		BestIndex = ~0u;
+	}
+	else if( ( BestIndex & ChildMask ) != LastChild )
+	{
+		// Fill with last
+		Set( BestIndex,
+			Children.GetBounds( LastChild ),
+			Children.GetFirstChild( LastChild ) );
+	}
+
+	return BestIndex;
+}
+
+template< uint32 MaxChildren, typename FCostMetric >
+void FDynamicBVH< MaxChildren, FCostMetric >::Rotate( uint32 NodeIndex )
+{
+	FNode& RESTRICT Node = GetNode( NodeIndex );
+
+	if( Node.IsRoot() )
+		return;
+
+	FBounds ExcludedBounds;
+	for( uint32 i = 0; i < Node.NumChildren; i++ )
+	{
+		if( i != (NodeIndex & ChildMask) )
+			ExcludedBounds += Node.ChildBounds[i];
+	}
+	
+	FNode& RESTRICT ParentNode = GetNode( Node.ParentIndex );
+
+	float	BestCost = CostMetric( ParentNode.GetBounds( Node.ParentIndex ) );
+	uint32	BestIndex = ~0u;
+	for( uint32 i = 0; i < ParentNode.NumChildren; i++ )
+	{
+		if( i != (Node.ParentIndex & ChildMask) )
+		{
+			// Parent's sibling
+			float Cost = CostMetric( ExcludedBounds + ParentNode.ChildBounds[i] );
+			if( Cost < BestCost )
+			{
+				BestCost = Cost;
+				BestIndex = ( Node.ParentIndex & ~ChildMask ) | i;
+			}
+		}
+	}
+
+	if( BestIndex != ~0u )
+	{
+		// Swap
+		FBounds Bounds		= Node.GetBounds( NodeIndex );
+		uint32 FirstChild	= Node.GetFirstChild( NodeIndex );
+
+		Set( NodeIndex, ParentNode.GetBounds( BestIndex ), ParentNode.GetFirstChild( BestIndex ) );
+		Set( BestIndex, Bounds, FirstChild );
+	}
+}
+
+template< uint32 MaxChildren, typename FCostMetric >
+FORCEINLINE uint32 FDynamicBVH< MaxChildren, FCostMetric >::AllocNode()
+{
+	if( FreeList.Num() )
+	{
+		return FreeList.Pop( false ) << IndexShift;
+	}
+	else
+	{
+		NodeIsDirty.Add( true );
+		return Nodes.AddUninitialized() << IndexShift;
+	}
+}
+
+template< uint32 MaxChildren, typename FCostMetric >
+FORCEINLINE void FDynamicBVH< MaxChildren, FCostMetric >::FreeNode( uint32 NodeIndex )
+{
+	// Assumes nothing is still linking to it
+	GetNode( NodeIndex ).ParentIndex = ~0u;
+	GetNode( NodeIndex ).NumChildren = 0;
+	FreeList.Push( NodeIndex >> IndexShift );
+}
+
+template< uint32 MaxChildren, typename FCostMetric >
+void FDynamicBVH< MaxChildren, FCostMetric >::CheckNode( uint32 NodeIndex ) const
+{
+	const FNode& Node = GetNode( NodeIndex );
+
+	check( ( NodeIndex & ChildMask ) < Node.NumChildren );
+
+	if( Node.IsRoot() )
+	{
+		//check( ( NodeIndex & ~ChildMask ) == 0 );
+	}
+	else
+	{
+		check( Node.NumChildren > 1 );
+		check( GetNode( Node.ParentIndex ).GetFirstChild( Node.ParentIndex ) == ( NodeIndex & ~ChildMask ) );
+	}
+
+	for( uint32 i = 0; i < Node.NumChildren; i++ )
+	{
+		uint32 FirstChild = Node.GetFirstChild( NodeIndex );
+		if( FirstChild & 1 )
+		{
+			check( Leaves[ FirstChild >> 1 ] == NodeIndex );
+		}
+		else
+		{
+			check( ( FirstChild & ChildMask ) == 0 );
+			check( GetNode( FirstChild ).ParentIndex == NodeIndex );
+		}
+	}
+}
+
+template< uint32 MaxChildren, typename FCostMetric >
+template< typename FPredicate, typename FFuncType >
+void FDynamicBVH< MaxChildren, FCostMetric >::ForAll( const FPredicate& Predicate, const FFuncType& Func ) const
+{
+	if( !Predicate( RootBounds ) )
+		return;
+
+	TArray< uint32, TInlineAllocator<256> > Stack;
+
+	uint32 NodeIndex = 0;
+
+	while( true )
+	{
+		FNode& RESTRICT Node = GetNode( NodeIndex );
+
+		for( uint32 i = 0; i < Node.NumChildren; i++ )
+		{
+			if( Predicate( Node.ChildBounds[i] ) )
+			{
+				uint32 FirstChild = Node.GetFirstChild( NodeIndex );
+				if( FirstChild & 1 )
+				{
+					Func( FirstChild >> 1 );
+				}
+				else
+				{
+					Stack.Push( FirstChild );
+				}
+			}
+		}
+
+		if( Stack.Num() == 0 )
+			break;
+
+		NodeIndex = Stack.Pop( false );
+	}
+}
+
+template< uint32 MaxChildren, typename FCostMetric >
+template< typename FFuncType >
+void FDynamicBVH< MaxChildren, FCostMetric >::ForAllDirty( const FFuncType& Func )
+{
+	for( uint32 NodeIndex : DirtyNodes )
+	{
+		Func( NodeIndex, GetNode( NodeIndex ) );
+		NodeIsDirty[ NodeIndex ] = false;
+	}
+	DirtyNodes.Reset();
+}
+
+
+class FMortonArray
+{
+public:
+	struct FRange
+	{
+		int32	Begin;
+		int32	End;
+
+		int32	Num() const { return End - Begin; }
+	};
+	
+public:
+			FMortonArray( const TArray< FBounds >& InBounds );
+
+	uint32	GetIndex( int32 i ) const { return Sorted[i].Index; }
+	uint32	Split( const FRange& Range );
+	
+private:
+	void	RegenerateCodes( const FRange& Range );
+
+	struct FSortPair
+	{
+		uint32 Code;
+		uint32 Index;
+
+		bool operator<( const FSortPair& Other ) const { return Code < Other.Code; }
+	};
+	TArray< FSortPair >			Sorted;
+
+	const TArray< FBounds >&	Bounds;
+};
+
+FORCEINLINE uint32 FMortonArray::Split( const FRange& Range )
+{
+	uint32 Code0 = Sorted[ Range.Begin ].Code;
+	uint32 Code1 = Sorted[ Range.End - 1 ].Code;
+	uint32 Diff = Code0 ^ Code1;
+	if( Diff == 0 )
+	{
+		RegenerateCodes( Range );
+
+		Code0 = Sorted[ Range.Begin ].Code;
+		Code1 = Sorted[ Range.End - 1 ].Code;
+		Diff = Code0 ^ Code1;
+
+		if( Diff == 0 )
+			return ( Range.Begin + Range.End ) >> 1;
+	}
+
+	uint32 HighestBitDiff = FMath::FloorLog2( Diff );
+	uint32 Mask = 1 << HighestBitDiff;
+
+	int32 Min = Range.Begin;
+	int32 Max = Range.End;
+	while( Min + 1 != Max )
+	{
+		int32 Mid = ( Min + Max ) >> 1;
+		if( Sorted[ Mid ].Code & Mask )
+			Max = Mid;
+		else
+			Min = Mid;
+	}
+	
+	return Max;
+}
+
+template< uint32 MaxChildren, typename FCostMetric >
+void FDynamicBVH< MaxChildren, FCostMetric >::Build( const TArray< FBounds >& BoundsArray, uint32 FirstIndex )
+{
+	if( FirstIndex + BoundsArray.Num() > (uint32)Leaves.Num() )
+	{
+		int32 Count = FirstIndex + BoundsArray.Num() - Leaves.Num();
+		int32 First = Leaves.AddUninitialized( Count );
+		FMemory::Memset( &Leaves[ First ], 0xff, Count * sizeof( uint32 ) );
+	}
+	
+	FMortonArray MortonArray( BoundsArray );
+
+	using FRange = FMortonArray::FRange;
+
+	// TEMP Start empty
+	FreeNode(0);
+
+	struct FCreateNode
+	{
+		uint32	ParentIndex;
+		FRange	Range;
+	};
+	TArray< FCreateNode, TInlineAllocator<32> > Stack;
+
+	uint32 ParentIndex = ~0u;
+	FRange Range = { 0, BoundsArray.Num() };
+
+	while( true )
+	{
+		uint32 NodeIndex = AllocNode();
+		FNode& Node = GetNode( NodeIndex );
+
+		Node.ParentIndex = ParentIndex;
+		if( ParentIndex != ~0u )
+			SetFirstChild( ParentIndex, NodeIndex );
+
+		check( Range.Begin < Range.End );
+
+		int32 NumLeaves = Range.Num();
+		if( NumLeaves <= MaxChildren )
+		{
+			Node.NumChildren = NumLeaves;
+			for( int32 i = 0; i < NumLeaves; i++ )
+			{
+				uint32 Index = MortonArray.GetIndex( Range.Begin + i );
+				Set( NodeIndex + i, BoundsArray[ Index ], ( ( FirstIndex + Index ) << 1 ) | 1 );
+			}
+
+			// Propagate bounds up tree
+			FBounds PathBounds	= Node.UnionBounds();
+			uint32  PathIndex	= Node.ParentIndex;
+			while( PathIndex != ~0u )
+			{
+				SetBounds( PathIndex, PathBounds );
+
+				// Only continue if first child, which signifies the node is complete.
+				if( PathIndex & ChildMask )
+					break;
+
+				FNode& PathNode = GetNode( PathIndex );
+				PathBounds	= PathNode.UnionBounds();
+				PathIndex	= PathNode.ParentIndex;
+			}
+			// TODO: RootBounds
+
+			if( Stack.Num() == 0 )
+				break;
+
+			ParentIndex	= Stack.Last().ParentIndex;
+			Range		= Stack.Last().Range;
+
+			Stack.Pop( false );
+		}
+		else
+		{
+			FRange Children[ MaxChildren ];
+			Children[0] = Range;
+
+			int32 NumChildren = 1;
+			int32 SplitIndex = 0;
+			do
+			{
+				FRange Child = Children[ SplitIndex ];
+
+				uint32 Middle = MortonArray.Split( Child );
+				check( Middle != Child.Begin );
+				check( Middle != Child.End );
+
+				Children[ SplitIndex ].Begin	= Child.Begin;
+				Children[ SplitIndex ].End		= Middle;
+				Children[ NumChildren ].Begin	= Middle;
+				Children[ NumChildren ].End		= Child.End;
+				NumChildren++;
+
+				int32 LargestNum = 0;
+				int32 LargestIndex = -1;
+				for( int32 i = 0; i < NumChildren; i++ )
+				{
+					int32 Num = Children[i].Num();
+					if( Num <= MaxChildren )
+						continue;
+
+					if( Num > LargestNum )
+					{
+						LargestNum = Num;
+						LargestIndex = i;
+					}
+				}
+
+				SplitIndex = LargestIndex;
+			}
+			while( NumChildren < MaxChildren && SplitIndex >= 0 );
+			
+			Node.NumChildren = NumChildren;
+
+			// Move leaves to the back
+			for( int32 Front = 0, Back = NumChildren - 1; Front < Back; )
+			{
+				if( Children[ Front ].Num() == 1 )
+					Swap( Children[ Front ], Children[ Back-- ] );
+				else
+					Front++;
+			}
+
+			NumLeaves = 0;
+			for( int32 i = NumChildren - 1; i >= 0; i-- )
+			{
+				bool bIsLeaf = Children[i].Num() == 1;
+				if( !bIsLeaf )
+					break;
+				NumLeaves++;
+
+				uint32 Index = MortonArray.GetIndex( Children[i].Begin );
+				Set( NodeIndex + i, BoundsArray[ Index ], ( ( FirstIndex + Index ) << 1 ) | 1 );
+			}
+			check( NumLeaves < NumChildren );
+
+			int32 Last = NumChildren - NumLeaves - 1;
+			for( int32 i = 0; i < Last; i++ )
+			{
+				Stack.Push( { NodeIndex + i, Children[i] } );
+			}
+
+			ParentIndex = NodeIndex + Last;
+			Range = Children[ Last ];
+		}
+	}
+}
\ No newline at end of file
