Instance BVH uploaded to GPU. Disabled for now.

diff --git a/Engine/Source/Runtime/Renderer/Private/GPUScene.cpp b/Engine/Source/Runtime/Renderer/Private/GPUScene.cpp
--- a/Engine/Source/Runtime/Renderer/Private/GPUScene.cpp
+++ b/Engine/Source/Runtime/Renderer/Private/GPUScene.cpp
@@ -50,6 +50,15 @@ FAutoConsoleVariableRef CVarGPUSceneParallelUpdate(
 	ECVF_RenderThreadSafe
 );
 
+int32 GGPUSceneInstanceBVH = 0;
+FAutoConsoleVariableRef CVarGPUSceneInstanceBVH(
+	TEXT("r.GPUScene.InstanceBVH"),
+	GGPUSceneInstanceBVH,
+	TEXT("Add instances to BVH. (WIP)"),
+	ECVF_RenderThreadSafe | ECVF_ReadOnly
+);
+
+
 template<typename ResourceType>
 ResourceType* GetMirrorGPU(FGPUScene& GPUScene);
 
@@ -170,6 +179,13 @@ FGPUScenePrimitiveCollector::FUploadData* FGPUScenePrimitiveCollector::AllocateU
 	return GPUSceneDynamicContext->AllocateDynamicPrimitiveData();
 }
 
+struct FBVHNode
+{
+	uint32		ChildIndexes[4];
+	FVector4	ChildMin[3];
+	FVector4	ChildMax[3];
+};
+
 
 /**
  * Info needed by the uploader to update a primitive.
@@ -518,6 +534,11 @@ void FGPUScene::UploadGeneral(FRHICommandListImmediate& RHICmdList, FScene *Scen
 		const bool bResizedInstanceData = ResizeResourceSOAIfNeeded(RHICmdList, InstanceDataBuffer, InstanceDataSizeReserve * sizeof(FInstanceSceneShaderData::Data), InstanceDataNumArrays, TEXT("InstanceData"));
 		InstanceDataSOAStride = InstanceDataSizeReserve;
 
+		{
+			const uint32 NumNodes = FMath::RoundUpToPowerOfTwo( FMath::Max( Scene->InstanceBVH.GetNumNodes(), 256 ) );
+			ResizeResourceIfNeeded( RHICmdList, InstanceBVHBuffer, NumNodes * sizeof( FBVHNode ), TEXT("InstanceBVH") );
+		}
+
 		const uint32 LightMapDataBufferSize = FMath::RoundUpToPowerOfTwo(FMath::Max(LightmapDataAllocator.GetMaxSize(), 256));
 		const bool bResizedLightmapData = ResizeResourceIfNeeded(RHICmdList, LightmapDataBuffer, LightMapDataBufferSize * sizeof(FLightmapSceneShaderData::Data), TEXT("LightmapData"));
 
@@ -888,6 +909,35 @@ void FGPUScene::UploadGeneral(FRHICommandListImmediate& RHICmdList, FScene *Scen
 				RHICmdList.Transition(FRHITransitionInfo(InstanceDataBuffer.UAV, ERHIAccess::Unknown, ERHIAccess::SRVMask));
 			}
 
+			if( Scene->InstanceBVH.GetNumDirty() > 0 )
+			{
+				InstanceUploadBuffer.Init( Scene->InstanceBVH.GetNumDirty(), sizeof( FBVHNode ), true, TEXT("InstanceUploadBuffer") );
+
+				Scene->InstanceBVH.ForAllDirty(
+					[&]( uint32 NodeIndex, const auto& Node )
+					{
+						FBVHNode GPUNode;
+						for( int i = 0; i < 4; i++ )
+						{
+							GPUNode.ChildIndexes[i] = Node.ChildIndexes[i];
+
+							GPUNode.ChildMin[0][i] = Node.ChildBounds[i].Min.X;
+							GPUNode.ChildMin[1][i] = Node.ChildBounds[i].Min.Y;
+							GPUNode.ChildMin[2][i] = Node.ChildBounds[i].Min.Z;
+
+							GPUNode.ChildMax[0][i] = Node.ChildBounds[i].Max.X;
+							GPUNode.ChildMax[1][i] = Node.ChildBounds[i].Max.Y;
+							GPUNode.ChildMax[2][i] = Node.ChildBounds[i].Max.Z;
+						}
+
+						InstanceUploadBuffer.Add( NodeIndex, &GPUNode );
+					} );
+
+				RHICmdList.Transition( FRHITransitionInfo( InstanceBVHBuffer.UAV, ERHIAccess::Unknown, ERHIAccess::UAVCompute ) );
+				InstanceUploadBuffer.ResourceUploadTo( RHICmdList, InstanceBVHBuffer, false );
+				RHICmdList.Transition( FRHITransitionInfo( InstanceBVHBuffer.UAV, ERHIAccess::UAVCompute, ERHIAccess::SRVMask ) );
+			}
+
 			if (NumLightmapDataUploads > 0)
 			{
 				ERHIAccess CurrentAccess = ERHIAccess::Unknown;
diff --git a/Engine/Source/Runtime/Renderer/Private/GPUScene.h b/Engine/Source/Runtime/Renderer/Private/GPUScene.h
--- a/Engine/Source/Runtime/Renderer/Private/GPUScene.h
+++ b/Engine/Source/Runtime/Renderer/Private/GPUScene.h
@@ -211,6 +211,7 @@ public:
 	FScatterUploadBuffer	InstanceUploadBuffer;
 	uint32					InstanceDataSOAStride;	// Distance between arrays in float4s
 	TSet<uint32>			InstanceClearList;
+	FRWBufferStructured		InstanceBVHBuffer;
 
 	/** GPU light map data */
 	FGrowOnlySpanAllocator	LightmapDataAllocator;
diff --git a/Engine/Source/Runtime/Renderer/Private/PrimitiveSceneInfo.cpp b/Engine/Source/Runtime/Renderer/Private/PrimitiveSceneInfo.cpp
--- a/Engine/Source/Runtime/Renderer/Private/PrimitiveSceneInfo.cpp
+++ b/Engine/Source/Runtime/Renderer/Private/PrimitiveSceneInfo.cpp
@@ -24,6 +24,7 @@
 #include "Lumen/LumenSceneRendering.h"
 
 extern int32 GGPUSceneInstanceClearList;
+extern int32 GGPUSceneInstanceBVH;
 
 /** An implementation of FStaticPrimitiveDrawInterface that stores the drawn elements for the rendering thread to use. */
 class FBatchingSPDI : public FStaticPrimitiveDrawInterface
@@ -459,7 +460,6 @@ void FPrimitiveSceneInfo::CacheMeshDrawCommands(FRHICommandListImmediate& RHICmd
 		}
 	}
 #endif
-
 }
 
 void FPrimitiveSceneInfo::RemoveCachedMeshDrawCommands()
@@ -884,6 +884,17 @@ void FPrimitiveSceneInfo::AddToScene(FRHICommandListImmediate& RHICmdList, FScen
 				{
 					SceneInfo->InstanceDataOffset = Scene->GPUScene.AllocateInstanceSlots(PrimitiveInstances->Num());
 					SceneInfo->NumInstanceDataEntries = PrimitiveInstances->Num();
+
+					if( GGPUSceneInstanceBVH )
+					{
+						for( int32 InstanceIndex = 0; InstanceIndex < PrimitiveInstances->Num(); ++InstanceIndex )
+						{
+							const FPrimitiveInstance& PrimitiveInstance = (*PrimitiveInstances)[ InstanceIndex ];
+
+							FBox WorldBox = PrimitiveInstance.LocalBounds.TransformBy( SceneInfo->Proxy->GetLocalToWorld() ).GetBox();
+							Scene->InstanceBVH.Add( FBounds( { WorldBox.Min, WorldBox.Max } ), SceneInfo->InstanceDataOffset + InstanceIndex );
+						}
+					}
 				}
 			}
 #if defined(GPUCULL_TODO)
@@ -1079,6 +1090,14 @@ void FPrimitiveSceneInfo::RemoveFromScene(bool bUpdateStaticDrawLists)
 		SCOPE_CYCLE_COUNTER(STAT_UpdateGPUSceneTime);
 
 		check(Proxy->SupportsInstanceDataBuffer() || NumInstanceDataEntries == 1);
+		if( GGPUSceneInstanceBVH )
+		{
+			for( int32 InstanceIndex = 0; InstanceIndex < NumInstanceDataEntries; InstanceIndex++ )
+			{
+				Scene->InstanceBVH.Remove( InstanceDataOffset + InstanceIndex );
+			}
+		}
+
 		Scene->GPUScene.FreeInstanceSlots(InstanceDataOffset, NumInstanceDataEntries);
 		InstanceDataOffset = INDEX_NONE;
 		NumInstanceDataEntries = 0;
diff --git a/Engine/Source/Runtime/Renderer/Private/ScenePrivate.h b/Engine/Source/Runtime/Renderer/Private/ScenePrivate.h
--- a/Engine/Source/Runtime/Renderer/Private/ScenePrivate.h
+++ b/Engine/Source/Runtime/Renderer/Private/ScenePrivate.h
@@ -58,6 +58,7 @@
 #include "VolumetricRenderTargetViewStateData.h"
 #include "GPUScene.h"
 #include "CompositionLighting/PostProcessDeferredDecals.h"
+#include "DynamicBVH.h"
 
 /** Factor by which to grow occlusion tests **/
 #define OCCLUSION_SLOP (1.0f)
@@ -2623,6 +2624,8 @@ public:
 	/** An octree containing the primitives in the scene. */
 	FScenePrimitiveOctree PrimitiveOctree;
 
+	FDynamicBVH<4> InstanceBVH;
+
 	/** Indicates whether this scene requires hit proxy rendering. */
 	bool bRequiresHitProxies;
 
