Additional GPU instance culling improvements and cleanup

diff --git a/Engine/Shaders/Private/LocalVertexFactory.ush b/Engine/Shaders/Private/LocalVertexFactory.ush
--- a/Engine/Shaders/Private/LocalVertexFactory.ush
+++ b/Engine/Shaders/Private/LocalVertexFactory.ush
@@ -261,7 +261,7 @@ struct FPositionAndNormalOnlyVertexFactoryInput
 	float4	Position	: ATTRIBUTE0;
 	float4	Normal		: ATTRIBUTE2;
 
-#if USE_INSTANCING && !MANUAL_VERTEX_FETCH
+#if (USE_INSTANCING || USE_INSTANCE_CULLING) && !MANUAL_VERTEX_FETCH
 	float4 InstanceOrigin : ATTRIBUTE8;  // per-instance random in w 
 	half4 InstanceTransform1 : ATTRIBUTE9;  // hitproxy.r + 256 * selected in .w
 	half4 InstanceTransform2 : ATTRIBUTE10; // hitproxy.g in .w
@@ -303,23 +303,26 @@ struct FVertexFactoryIntermediates
 #endif
 
 #if USE_INSTANCING
-		float4 InstanceOrigin;
-		float4 InstanceTransform1;
-		float4 InstanceTransform2;
-		float4 InstanceTransform3;
-
-		#if USE_INSTANCING_BONEMAP
-			float4 InstancePrevOrigin;
-			float4 InstancePrevTransform1;
-			float4 InstancePrevTransform2;
-			float4 InstancePrevTransform3;
-		#endif
+	float IsInstanceSelected;
+
+	float4 InstanceOrigin;
+	float4 InstanceTransform1;
+	float4 InstanceTransform2;
+	float4 InstanceTransform3;
+
+	#if USE_INSTANCING_BONEMAP
+		float4 InstancePrevOrigin;
+		float4 InstancePrevTransform1;
+		float4 InstancePrevTransform2;
+		float4 InstancePrevTransform3;
+	#endif
 
-		float4 InstanceLightmapAndShadowMapUVBias;
+	float4 InstanceLightmapAndShadowMapUVBias;
 
 	// x = per-instance random, y = per-instance fade out amount, z = hide/show flag, w dither fade cutoff
 	float4 PerInstanceParams;
-#endif	// USE_INSTANCING
+#endif
+
 	uint PrimitiveId;
 
 	float3 PreSkinPosition;
diff --git a/Engine/Source/Runtime/Engine/Private/SceneManagement.cpp b/Engine/Source/Runtime/Engine/Private/SceneManagement.cpp
--- a/Engine/Source/Runtime/Engine/Private/SceneManagement.cpp
+++ b/Engine/Source/Runtime/Engine/Private/SceneManagement.cpp
@@ -817,6 +817,10 @@ FViewUniformShaderParameters::FViewUniformShaderParameters()
 	EditorVisualizeLevelInstanceIds = GIdentityPrimitiveBuffer.EditorVisualizeLevelInstanceDataBufferSRV;
 	EditorSelectedHitProxyIds = GIdentityPrimitiveBuffer.EditorSelectedDataBufferSRV;
 //#endif
+
+	// Automatically bind these for everyone to use!
+	InstanceIdsBuffer = GInstanceCullingManagerResources.GetInstancesIdBufferSrv();
+	PageInfoBuffer = GInstanceCullingManagerResources.GetPageInfoBufferSrv();
 }
 
 FInstancedViewUniformShaderParameters::FInstancedViewUniformShaderParameters()
diff --git a/Engine/Source/Runtime/Engine/Public/SceneView.h b/Engine/Source/Runtime/Engine/Public/SceneView.h
--- a/Engine/Source/Runtime/Engine/Public/SceneView.h
+++ b/Engine/Source/Runtime/Engine/Public/SceneView.h
@@ -819,6 +819,8 @@ enum ETranslucencyVolumeCascade
 	VIEW_UNIFORM_BUFFER_MEMBER_ARRAY(int, PhysicsFieldVectorTargets, [MAX_PHYSICS_FIELD_TARGETS]) \
 	VIEW_UNIFORM_BUFFER_MEMBER_ARRAY(int, PhysicsFieldScalarTargets, [MAX_PHYSICS_FIELD_TARGETS]) \
 	VIEW_UNIFORM_BUFFER_MEMBER_ARRAY(int, PhysicsFieldIntegerTargets, [MAX_PHYSICS_FIELD_TARGETS]) \
+	VIEW_UNIFORM_BUFFER_MEMBER(uint32, InstanceDataSOAStride) \
+	VIEW_UNIFORM_BUFFER_MEMBER(uint32, GPUSceneViewId) \
 
 #define VIEW_UNIFORM_BUFFER_MEMBER(type, identifier) \
 	SHADER_PARAMETER(type, identifier)
@@ -908,6 +910,9 @@ BEGIN_GLOBAL_SHADER_PARAMETER_STRUCT_WITH_CONSTRUCTOR(FViewUniformShaderParamete
 
 	SHADER_PARAMETER_SRV(Buffer<float>, PhysicsFieldClipmapBuffer)
 
+	SHADER_PARAMETER_SRV(Buffer<uint>, InstanceIdsBuffer)
+	SHADER_PARAMETER_SRV(Buffer<uint>, PageInfoBuffer)
+
 END_GLOBAL_SHADER_PARAMETER_STRUCT()
 
 /** Copy of the view uniform shader parameters associated with a view for instanced stereo. */
diff --git a/Engine/Source/Runtime/Renderer/Private/GPUScene.cpp b/Engine/Source/Runtime/Renderer/Private/GPUScene.cpp
--- a/Engine/Source/Runtime/Renderer/Private/GPUScene.cpp
+++ b/Engine/Source/Runtime/Renderer/Private/GPUScene.cpp
@@ -1111,14 +1111,11 @@ void FGPUScene::UploadDynamicPrimitiveShaderDataForViewInternal(FRHICommandListI
 	// Update view uniform buffer
 	View.CachedViewUniformShaderParameters->InstanceSceneData = InstanceDataBuffer.SRV;
 	View.CachedViewUniformShaderParameters->LightmapSceneData = LightmapDataBuffer.SRV;
-#if defined(GPUCULL_TODO)
 	View.CachedViewUniformShaderParameters->InstanceDataSOAStride = InstanceDataSOAStride;
-#endif //defined(GPUCULL_TODO)
 
 	View.ViewUniformBuffer.UpdateUniformBufferImmediate(*View.CachedViewUniformShaderParameters);
 }
 
-
 void AddPrimitiveToUpdateGPU(FScene& Scene, int32 PrimitiveId)
 {
 	Scene.GPUScene.AddPrimitiveToUpdate(PrimitiveId);
diff --git a/Engine/Source/Runtime/Renderer/Private/MeshPassProcessor.cpp b/Engine/Source/Runtime/Renderer/Private/MeshPassProcessor.cpp
--- a/Engine/Source/Runtime/Renderer/Private/MeshPassProcessor.cpp
+++ b/Engine/Source/Runtime/Renderer/Private/MeshPassProcessor.cpp
@@ -1295,7 +1295,7 @@ void DrawDynamicMeshPassPrivate(
 
 		FRHIVertexBuffer* PrimitiveIdVertexBuffer = nullptr;
 
-		ApplyViewOverridesToMeshDrawCommands(View, VisibleMeshDrawCommands, DynamicMeshDrawCommandStorage, GraphicsMinimalPipelineStateSet, InNeedsShaderInitialisation);
+		ApplyViewOverridesToMeshDrawCommands(View, VisibleMeshDrawCommands, DynamicMeshDrawCommandStorage, GraphicsMinimalPipelineStateSet, InNeedsShaderInitialization);
 		
 		// GPUCULL_TODO: determine if unsupportable OR just not supporting scene primitives (as this requires a multi-pass approach)
 
diff --git a/Engine/Source/Runtime/Renderer/Private/SceneRendering.cpp b/Engine/Source/Runtime/Renderer/Private/SceneRendering.cpp
--- a/Engine/Source/Runtime/Renderer/Private/SceneRendering.cpp
+++ b/Engine/Source/Runtime/Renderer/Private/SceneRendering.cpp
@@ -925,6 +925,8 @@ void FViewInfo::Init()
 	NumVisibleDynamicEditorPrimitives = 0;
 
 	StrataSceneData = nullptr;
+	
+	GPUSceneViewId = INDEX_NONE;
 }
 
 FViewInfo::~FViewInfo()
@@ -1788,6 +1790,8 @@ void FViewInfo::SetupUniformBufferParameters(
 
 	ViewUniformShaderParameters.VTFeedbackBuffer = GVirtualTextureFeedbackBuffer.GetUAV();
 
+	ViewUniformShaderParameters.GPUSceneViewId = GPUSceneViewId;
+
 #if WITH_EDITOR
 	if (EditorVisualizeLevelInstanceBuffer.SRV)
 	{
diff --git a/Engine/Source/Runtime/Renderer/Private/SceneRendering.h b/Engine/Source/Runtime/Renderer/Private/SceneRendering.h
--- a/Engine/Source/Runtime/Renderer/Private/SceneRendering.h
+++ b/Engine/Source/Runtime/Renderer/Private/SceneRendering.h
@@ -1386,6 +1386,11 @@ public:
 	FRayTracingLightData							RayTracingLightData;
 #endif // RHI_RAYTRACING
 
+	/**
+	 * Id of the view in the GPU instance culling manager. Used to fetch the culled draw commands.
+	 */
+	int32 GPUSceneViewId;
+
 #if WITH_EDITOR
 	TArray<uint32> EditorVisualizeLevelInstanceIds;
 	FDynamicReadBuffer EditorVisualizeLevelInstanceBuffer;
diff --git a/Engine/Source/Runtime/Renderer/Private/SceneVisibility.cpp b/Engine/Source/Runtime/Renderer/Private/SceneVisibility.cpp
--- a/Engine/Source/Runtime/Renderer/Private/SceneVisibility.cpp
+++ b/Engine/Source/Runtime/Renderer/Private/SceneVisibility.cpp
@@ -4487,6 +4487,12 @@ bool FDeferredShadingSceneRenderer::InitViews(FRDGBuilder& GraphBuilder, const F
 
 	DispatchToRHIThread();
 
+	// Create GPU-side representation of the view for instance culling.
+	for (int32 ViewIndex = 0; ViewIndex < Views.Num(); ++ViewIndex)
+	{
+		Views[ViewIndex].GPUSceneViewId = InstanceCullingManager.RegisterView(Views[ViewIndex]);
+	}
+
 	{
 		// This is to init the ViewUniformBuffer before rendering for the Niagara compute shader.
 		// This needs to run before ComputeViewVisibility() is called, but the views normally initialize the ViewUniformBuffer after that (at the end of this method).
diff --git a/Engine/Source/Runtime/Renderer/Private/ShadowRendering.h b/Engine/Source/Runtime/Renderer/Private/ShadowRendering.h
--- a/Engine/Source/Runtime/Renderer/Private/ShadowRendering.h
+++ b/Engine/Source/Runtime/Renderer/Private/ShadowRendering.h
@@ -394,6 +394,8 @@ public:
 	/** Index of the subject primitive for per object shadows. */
 	int32 SubjectPrimitiveComponentIndex = -1;
 
+	TArray<int32> ViewIds;
+
 public:
 
 	// default constructor
diff --git a/Engine/Source/Runtime/Renderer/Private/TranslucentRendering.cpp b/Engine/Source/Runtime/Renderer/Private/TranslucentRendering.cpp
--- a/Engine/Source/Runtime/Renderer/Private/TranslucentRendering.cpp
+++ b/Engine/Source/Runtime/Renderer/Private/TranslucentRendering.cpp
@@ -880,13 +880,14 @@ BEGIN_SHADER_PARAMETER_STRUCT(FTranslucentBasePassParameters, )
 	SHADER_PARAMETER_STRUCT_INCLUDE(FViewShaderParameters, View)
 	SHADER_PARAMETER_STRUCT_REF(FReflectionCaptureShaderData, ReflectionCapture)
 	SHADER_PARAMETER_RDG_UNIFORM_BUFFER(FTranslucentBasePassUniformParameters, BasePass)
+	SHADER_PARAMETER_STRUCT_INCLUDE(FInstanceCullingDrawParams, InstanceCullingDrawParams)
 	RENDER_TARGET_BINDING_SLOTS()
 END_SHADER_PARAMETER_STRUCT()
 
 static void RenderTranslucencyViewInner(
 	FRDGBuilder& GraphBuilder,
 	const FSceneRenderer& SceneRenderer,
-	const FViewInfo& View,
+	FViewInfo& View,
 	FScreenPassTextureViewport Viewport,
 	float ViewportScale,
 	FRDGTextureMSAA SceneColorTexture,
@@ -908,7 +909,7 @@ static void RenderTranslucencyViewInner(
 	PassParameters->RenderTargets.ResolveRect = FResolveRect(Viewport.Rect);
 
 	const EMeshPass::Type MeshPass = TranslucencyPassToMeshPass(TranslucencyPass);
-	// GPUCULL_TODO: View.ParallelMeshDrawCommandPasses[MeshPass].BuildRenderingCommands(GraphBuilder, SceneRenderer.Scene->GPUScene, PassParameters->InstanceCullingDrawParams);
+	View.ParallelMeshDrawCommandPasses[MeshPass].BuildRenderingCommands(GraphBuilder, SceneRenderer.Scene->GPUScene, PassParameters->InstanceCullingDrawParams);
 
 	if (bRenderInParallel)
 	{
@@ -926,9 +927,7 @@ static void RenderTranslucencyViewInner(
 			[&SceneRenderer, &View, PassParameters, ViewportScale, Viewport, TranslucencyPass](FRHICommandListImmediate& RHICmdList)
 		{
 			FRDGParallelCommandListSet ParallelCommandListSet(RHICmdList, GET_STATID(STAT_CLP_Translucency), SceneRenderer, View, FParallelCommandListBindings(PassParameters), ViewportScale);
-			FInstanceCullingDrawParams CullingParams; // GPUCULL_TODO
-			RenderViewTranslucencyInner(RHICmdList, SceneRenderer, View, Viewport, ViewportScale, TranslucencyPass, &ParallelCommandListSet, CullingParams); // GPUCULL_TODO
-			//RenderViewTranslucencyInner(RHICmdList, SceneRenderer, View, Viewport, ViewportScale, TranslucencyPass, &ParallelCommandListSet, PassParameters->InstanceCullingDrawParams); // GPUCULL_TODO
+			RenderViewTranslucencyInner(RHICmdList, SceneRenderer, View, Viewport, ViewportScale, TranslucencyPass, &ParallelCommandListSet, PassParameters->InstanceCullingDrawParams);
 		});
 
 		if (bResolveColorTexture)
@@ -946,9 +945,7 @@ static void RenderTranslucencyViewInner(
 			ERDGPassFlags::Raster,
 			[&SceneRenderer, &View, ViewportScale, Viewport, TranslucencyPass, PassParameters](FRHICommandListImmediate& RHICmdList)
 		{
-			FInstanceCullingDrawParams CullingParams; // GPUCULL_TODO
-			RenderViewTranslucencyInner(RHICmdList, SceneRenderer, View, Viewport, ViewportScale, TranslucencyPass, nullptr, CullingParams); // GPUCULL_TODO
-			//RenderViewTranslucencyInner(RHICmdList, SceneRenderer, View, Viewport, ViewportScale, TranslucencyPass, nullptr, PassParameters->InstanceCullingDrawParams); // GPUCULL_TODO
+			RenderViewTranslucencyInner(RHICmdList, SceneRenderer, View, Viewport, ViewportScale, TranslucencyPass, nullptr, PassParameters->InstanceCullingDrawParams);
 		});
 	}
 }
