Fix a bug of mobile deferred rendering in RDG.

diff --git a/Engine/Source/Runtime/Renderer/Private/MobileBasePassRendering.cpp b/Engine/Source/Runtime/Renderer/Private/MobileBasePassRendering.cpp
--- a/Engine/Source/Runtime/Renderer/Private/MobileBasePassRendering.cpp
+++ b/Engine/Source/Runtime/Renderer/Private/MobileBasePassRendering.cpp
@@ -302,7 +302,7 @@ void SetupMobileSkyReflectionUniformParameters(FSkyLightSceneProxy* SkyLight, FM
 void FMobileSceneRenderer::RenderMobileBasePass(
 	FRDGBuilder& GraphBuilder,
 	FRenderTargetBindingSlots& BasePassRenderTargets,
-	const TArrayView<const FViewInfo*> PassViews,
+	const TArrayView<const FViewInfo> PassViews,
 	FRDGTextureRef ScreenSpaceAO)
 {
 	CSV_SCOPED_TIMING_STAT_EXCLUSIVE(RenderBasePass);
@@ -313,7 +313,7 @@ void FMobileSceneRenderer::RenderMobileBasePass(
 	for (int32 ViewIndex = 0; ViewIndex < PassViews.Num(); ViewIndex++)
 	{
 		RDG_EVENT_SCOPE_CONDITIONAL(GraphBuilder, Views.Num() > 1, "View%d", ViewIndex);
-		const FViewInfo& View = *PassViews[ViewIndex];
+		const FViewInfo& View = PassViews[ViewIndex];
 		if (!View.ShouldRenderView())
 		{
 			continue;
diff --git a/Engine/Source/Runtime/Renderer/Private/MobileDeferredShadingPass.cpp b/Engine/Source/Runtime/Renderer/Private/MobileDeferredShadingPass.cpp
--- a/Engine/Source/Runtime/Renderer/Private/MobileDeferredShadingPass.cpp
+++ b/Engine/Source/Runtime/Renderer/Private/MobileDeferredShadingPass.cpp
@@ -470,13 +470,13 @@ static void SetupSimpleLightPSO(
 static void RenderSimpleLights(
 	FRHICommandListImmediate& RHICmdList, 
 	const FScene& Scene, 
-	const TArrayView<const FViewInfo*> PassViews, 
+	const TArrayView<const FViewInfo> PassViews, 
 	const FSortedLightSetSceneInfo &SortedLightSet, 
 	const FCachedLightMaterial& DefaultMaterial)
 {
 	const FSimpleLightArray& SimpleLights = SortedLightSet.SimpleLights;
 	const int32 NumViews = PassViews.Num();
-	const FViewInfo& View0 = *PassViews[0];
+	const FViewInfo& View0 = PassViews[0];
 
 	TShaderMapRef<TDeferredLightVS<true>> VertexShader(View0.ShaderMap);
 	TShaderRef<FMobileRadialLightFunctionPS> PixelShaders[2];
@@ -520,7 +520,7 @@ static void RenderSimpleLights(
 		const FSimpleLightEntry& SimpleLight = SimpleLights.InstanceData[LightIndex];
 		for (int32 ViewIndex = 0; ViewIndex < NumViews; ViewIndex++)
 		{
-			const FViewInfo& View = *PassViews[ViewIndex];
+			const FViewInfo& View = PassViews[ViewIndex];
 			const FSimpleLightPerViewEntry& SimpleLightPerViewData = SimpleLights.GetViewDependentData(LightIndex, ViewIndex, NumViews);
 			const FSphere LightBounds(SimpleLightPerViewData.Position, SimpleLight.Radius);
 			
@@ -572,7 +572,7 @@ void MobileDeferredShadingPass(
 	FRenderTargetBindingSlots& BasePassRenderTargets,
 	TRDGUniformBufferRef<FMobileSceneTextureUniformParameters> MobileSceneTextures,
 	const FScene& Scene, 
-	const TArrayView<const FViewInfo*> PassViews, 
+	const TArrayView<const FViewInfo> PassViews, 
 	const FSortedLightSetSceneInfo &SortedLightSet)
 {
 	auto* PassParameters = GraphBuilder.AllocParameters<FMobileDeferredPassParameters>();
@@ -581,10 +581,9 @@ void MobileDeferredShadingPass(
 
 	GraphBuilder.AddPass(
 		RDG_EVENT_NAME("MobileDeferredShadingPass"), PassParameters, ERDGPassFlags::Raster | ERDGPassFlags::SkipRenderPass,
-		[&Scene, &PassViews, &SortedLightSet](FRHICommandListImmediate& RHICmdList)
+		[&Scene, PassViews, &SortedLightSet](FRHICommandListImmediate& RHICmdList)
 	{
-
-		const FViewInfo& View0 = *PassViews[0];
+		const FViewInfo& View0 = PassViews[0];
 		RHICmdList.SetViewport(View0.ViewRect.Min.X, View0.ViewRect.Min.Y, 0.0f, View0.ViewRect.Max.X, View0.ViewRect.Max.Y, 1.0f);
 
 		// Default material for light rendering
diff --git a/Engine/Source/Runtime/Renderer/Private/MobileDeferredShadingPass.h b/Engine/Source/Runtime/Renderer/Private/MobileDeferredShadingPass.h
--- a/Engine/Source/Runtime/Renderer/Private/MobileDeferredShadingPass.h
+++ b/Engine/Source/Runtime/Renderer/Private/MobileDeferredShadingPass.h
@@ -20,5 +20,5 @@ void MobileDeferredShadingPass(
 	FRenderTargetBindingSlots& BasePassRenderTargets,
 	TRDGUniformBufferRef<FMobileSceneTextureUniformParameters> MobileSceneTextures,
 	const FScene& Scene, 
-	const TArrayView<const FViewInfo*> PassViews, 
+	const TArrayView<const FViewInfo> PassViews, 
 	const FSortedLightSetSceneInfo &SortedLightSet);
diff --git a/Engine/Source/Runtime/Renderer/Private/MobileShadingRenderer.cpp b/Engine/Source/Runtime/Renderer/Private/MobileShadingRenderer.cpp
--- a/Engine/Source/Runtime/Renderer/Private/MobileShadingRenderer.cpp
+++ b/Engine/Source/Runtime/Renderer/Private/MobileShadingRenderer.cpp
@@ -648,13 +648,6 @@ void FMobileSceneRenderer::Render(FRDGBuilder& GraphBuilder)
 
 	AddPass(GraphBuilder, PollOcclusionQueriesAndDispatchToRHIThreadPass);
 
-	// Default view list
-	TArray<const FViewInfo*> ViewList;
-	for (int32 ViewIndex = 0; ViewIndex < Views.Num(); ViewIndex++) 
-	{
-		ViewList.Add(&Views[ViewIndex]);
-	}
-
 	// Custom depth
 	// bShouldRenderCustomDepth has been initialized in InitViews on mobile platform
 	if (bShouldRenderCustomDepth)
@@ -669,11 +662,11 @@ void FMobileSceneRenderer::Render(FRDGBuilder& GraphBuilder)
 
 	if (bDeferredShading)
 	{
-		RenderDeferred(GraphBuilder, ViewList, SortedLightSet, ViewFamilyTexture, SceneTextures);
+		RenderDeferred(GraphBuilder, Views, SortedLightSet, ViewFamilyTexture, SceneTextures);
 	}
 	else
 	{
-		RenderForward(GraphBuilder, ViewList, ViewFamilyTexture, SceneTextures);
+		RenderForward(GraphBuilder, Views, ViewFamilyTexture, SceneTextures);
 	}
 
 	SceneTextures.MobileSetupMode = EMobileSceneTextureSetupMode::All;
@@ -767,9 +760,9 @@ BEGIN_SHADER_PARAMETER_STRUCT(FMobilePostBasePassViewExtensionParameters, )
 	RENDER_TARGET_BINDING_SLOTS()
 END_SHADER_PARAMETER_STRUCT()
 
-void FMobileSceneRenderer::RenderForward(FRDGBuilder& GraphBuilder, const TArrayView<const FViewInfo*> ViewList, FRDGTextureRef ViewFamilyTexture, const FSceneTextures& SceneTextures)
+void FMobileSceneRenderer::RenderForward(FRDGBuilder& GraphBuilder, const TArrayView<const FViewInfo> ViewList, FRDGTextureRef ViewFamilyTexture, const FSceneTextures& SceneTextures)
 {
-	const FViewInfo& View = *ViewList[0];
+	const FViewInfo& View = ViewList[0];
 
 	FRDGTextureRef SceneColor = nullptr;
 	FRDGTextureRef SceneColorResolve = nullptr;
@@ -1015,7 +1008,7 @@ void FMobileSceneRenderer::RenderForward(FRDGBuilder& GraphBuilder, const TArray
 	QueueSceneTextureExtractions(GraphBuilder, SceneTextures);
 }
 
-void FMobileSceneRenderer::RenderDeferred(FRDGBuilder& GraphBuilder, const TArrayView<const FViewInfo*> ViewList, const FSortedLightSetSceneInfo& SortedLightSet, FRDGTextureRef ViewFamilyTexture, const FSceneTextures& SceneTextures)
+void FMobileSceneRenderer::RenderDeferred(FRDGBuilder& GraphBuilder, const TArrayView<const FViewInfo> ViewList, const FSortedLightSetSceneInfo& SortedLightSet, FRDGTextureRef ViewFamilyTexture, const FSceneTextures& SceneTextures)
 {
 	FRDGTextureRef ColorTargets[5] = {
 		SceneTextures.Color.Target,
@@ -1197,7 +1190,7 @@ BEGIN_SHADER_PARAMETER_STRUCT(FMobileDebugViewPassParameters, )
 	SHADER_PARAMETER_RDG_UNIFORM_BUFFER(FDebugViewModePassUniformParameters, DebugViewMode)
 END_SHADER_PARAMETER_STRUCT()
 
-void FMobileSceneRenderer::RenderMobileDebugView(FRDGBuilder& GraphBuilder, const TArrayView<const FViewInfo*> PassViews, FRDGTextureRef QuadOverdrawTexture)
+void FMobileSceneRenderer::RenderMobileDebugView(FRDGBuilder& GraphBuilder, const TArrayView<const FViewInfo> PassViews, FRDGTextureRef QuadOverdrawTexture)
 {
 #if WITH_DEBUG_VIEW_MODES
 	CSV_SCOPED_TIMING_STAT_EXCLUSIVE(RenderDebugView);
@@ -1207,7 +1200,7 @@ void FMobileSceneRenderer::RenderMobileDebugView(FRDGBuilder& GraphBuilder, cons
 	for (int32 ViewIndex = 0; ViewIndex < PassViews.Num(); ViewIndex++)
 	{
 		RDG_EVENT_SCOPE_CONDITIONAL(GraphBuilder, Views.Num() > 1, "View%d", ViewIndex);
-		const FViewInfo& View = *PassViews[ViewIndex];
+		const FViewInfo& View = PassViews[ViewIndex];
 		if (!View.ShouldRenderView())
 		{
 			continue;
diff --git a/Engine/Source/Runtime/Renderer/Private/MobileTranslucentRendering.cpp b/Engine/Source/Runtime/Renderer/Private/MobileTranslucentRendering.cpp
--- a/Engine/Source/Runtime/Renderer/Private/MobileTranslucentRendering.cpp
+++ b/Engine/Source/Runtime/Renderer/Private/MobileTranslucentRendering.cpp
@@ -31,7 +31,7 @@
 #include "MeshPassProcessor.inl"
 #include "ClearQuad.h"
 
-void FMobileSceneRenderer::RenderTranslucency(FRDGBuilder& GraphBuilder, FRenderTargetBindingSlots& BasePassRenderTargets, const TArrayView<const FViewInfo*> PassViews, FRDGTextureRef ScreenSpaceAO)
+void FMobileSceneRenderer::RenderTranslucency(FRDGBuilder& GraphBuilder, FRenderTargetBindingSlots& BasePassRenderTargets, const TArrayView<const FViewInfo> PassViews, FRDGTextureRef ScreenSpaceAO)
 {
 	ETranslucencyPass::Type TranslucencyPass = 
 		ViewFamily.AllowTranslucencyAfterDOF() ? ETranslucencyPass::TPT_StandardTranslucency : ETranslucencyPass::TPT_AllTranslucency;
@@ -45,7 +45,7 @@ void FMobileSceneRenderer::RenderTranslucency(FRDGBuilder& GraphBuilder, FRender
 		{
 			RDG_EVENT_SCOPE_CONDITIONAL(GraphBuilder, Views.Num() > 1, "View%d", ViewIndex);
 
-			const FViewInfo& View = *PassViews[ViewIndex];
+			const FViewInfo& View = PassViews[ViewIndex];
 			if (!View.ShouldRenderView())
 			{
 				continue;
diff --git a/Engine/Source/Runtime/Renderer/Private/SceneRendering.h b/Engine/Source/Runtime/Renderer/Private/SceneRendering.h
--- a/Engine/Source/Runtime/Renderer/Private/SceneRendering.h
+++ b/Engine/Source/Runtime/Renderer/Private/SceneRendering.h
@@ -2099,12 +2099,12 @@ protected:
 	void RenderPrePass(FRDGBuilder& GraphBuilder, FRenderTargetBindingSlots& BasePassRenderTargets);
 
 	/** Renders the opaque base pass for mobile. */
-	void RenderMobileBasePass(FRDGBuilder& GraphBuilder, FRenderTargetBindingSlots& BasePassRenderTargets, const TArrayView<const FViewInfo*> PassViews, FRDGTextureRef ScreenSpaceAO);
+	void RenderMobileBasePass(FRDGBuilder& GraphBuilder, FRenderTargetBindingSlots& BasePassRenderTargets, const TArrayView<const FViewInfo> PassViews, FRDGTextureRef ScreenSpaceAO);
 
 	void RenderMobileEditorPrimitives(FRHICommandList& RHICmdList, const FViewInfo& View, const FMeshPassProcessorRenderState& DrawRenderState);
 
 	/** Renders the debug view pass for mobile. */
-	void RenderMobileDebugView(FRDGBuilder& GraphBuilder, const TArrayView<const FViewInfo*> PassViews, FRDGTextureRef QuadOverdrawTexture);
+	void RenderMobileDebugView(FRDGBuilder& GraphBuilder, const TArrayView<const FViewInfo> PassViews, FRDGTextureRef QuadOverdrawTexture);
 
 	/** Render modulated shadow projections in to the scene, loops over any unrendered shadows until all are processed.*/
 	void RenderModulatedShadowProjections(FRDGBuilder& GraphBuilder, TRDGUniformBufferRef<FMobileSceneTextureUniformParameters> MobileSceneTextures);
@@ -2125,7 +2125,7 @@ protected:
 	void RenderDecals(FRDGBuilder& GraphBuilder, FRenderTargetBindingSlots& BasePassRenderTargets, TRDGUniformBufferRef<FMobileSceneTextureUniformParameters> MobileSceneTextures);
 
 	/** Renders the base pass for translucency. */
-	void RenderTranslucency(FRDGBuilder& GraphBuilder, FRenderTargetBindingSlots& BasePassRenderTargets, const TArrayView<const FViewInfo*> PassViews, FRDGTextureRef ScreenSpaceAO);
+	void RenderTranslucency(FRDGBuilder& GraphBuilder, FRenderTargetBindingSlots& BasePassRenderTargets, const TArrayView<const FViewInfo> PassViews, FRDGTextureRef ScreenSpaceAO);
 
 	/** On chip pre-tonemap before scene color MSAA resolve (iOS only) */
 	void PreTonemapMSAA(FRDGBuilder& GraphBuilder, FRenderTargetBindingSlots& BasePassRenderTargets, const FMinimalSceneTextures& SceneTextures);
@@ -2136,8 +2136,8 @@ protected:
 	void UpdateDirectionalLightUniformBuffers(FRDGBuilder& GraphBuilder, const FViewInfo& View);
 	void UpdateSkyReflectionUniformBuffer();
 
-	void RenderForward(FRDGBuilder& GraphBuilder, const TArrayView<const FViewInfo*> ViewList, FRDGTextureRef ViewFamilyTexture, const FSceneTextures& SceneTextures);
-	void RenderDeferred(FRDGBuilder& GraphBuilder, const TArrayView<const FViewInfo*> ViewList, const FSortedLightSetSceneInfo& SortedLightSet, FRDGTextureRef ViewFamilyTexture, const FSceneTextures& SceneTextures);
+	void RenderForward(FRDGBuilder& GraphBuilder, const TArrayView<const FViewInfo> ViewList, FRDGTextureRef ViewFamilyTexture, const FSceneTextures& SceneTextures);
+	void RenderDeferred(FRDGBuilder& GraphBuilder, const TArrayView<const FViewInfo> ViewList, const FSortedLightSetSceneInfo& SortedLightSet, FRDGTextureRef ViewFamilyTexture, const FSceneTextures& SceneTextures);
 	
 	void RenderAmbientOcclusion(FRDGBuilder& GraphBuilder, FRDGTextureRef SceneDepthTexture, FRDGTextureRef AmbientOcclusionTexture);
 
