Check for null

diff --git a/Engine/Source/Runtime/Renderer/Private/GPUScene.cpp b/Engine/Source/Runtime/Renderer/Private/GPUScene.cpp
--- a/Engine/Source/Runtime/Renderer/Private/GPUScene.cpp
+++ b/Engine/Source/Runtime/Renderer/Private/GPUScene.cpp
@@ -534,6 +534,7 @@ void FGPUScene::UploadGeneral(FRHICommandListImmediate& RHICmdList, FScene *Scen
 		const bool bResizedInstanceData = ResizeResourceSOAIfNeeded(RHICmdList, InstanceDataBuffer, InstanceDataSizeReserve * sizeof(FInstanceSceneShaderData::Data), InstanceDataNumArrays, TEXT("InstanceData"));
 		InstanceDataSOAStride = InstanceDataSizeReserve;
 
+		if (Scene != nullptr)
 		{
 			const uint32 NumNodes = FMath::RoundUpToPowerOfTwo( FMath::Max( Scene->InstanceBVH.GetNumNodes(), 256 ) );
 			ResizeResourceIfNeeded( RHICmdList, InstanceBVHBuffer, NumNodes * sizeof( FBVHNode ), TEXT("InstanceBVH") );
