Fix index used to upload Lightmap info to GPU-Scene

#rnx

diff --git a/Engine/Source/Runtime/Renderer/Private/GPUScene.cpp b/Engine/Source/Runtime/Renderer/Private/GPUScene.cpp
--- a/Engine/Source/Runtime/Renderer/Private/GPUScene.cpp
+++ b/Engine/Source/Runtime/Renderer/Private/GPUScene.cpp
@@ -943,6 +943,10 @@ void FGPUScene::UploadGeneral(FRHICommandListImmediate& RHICmdList, FScene *Scen
 			{
 				ERHIAccess CurrentAccess = ERHIAccess::Unknown;
 
+				// GPUCULL_TODO: This code is wrong: the intention is to break it up into batches such that the uploaded data fits in the max buffer size
+				//               However, what it does do is break it up into batches of MaxLightmapsUploads (while iterating over primitives). This is bad
+				//               because it a) makes more batches than needed, b) does not AFAICT guarantee that we don't overflow (as each prim may have 
+				//				 multiple LCIs - so all may belong to the first 1/8th of primitives).
 				const int32 MaxLightmapsUploads = GetMaxPrimitivesUpdate(NumLightmapDataUploads, FLightmapSceneShaderData::LightmapDataStrideInFloat4s);
 				for (int32 PrimitiveOffset = 0; PrimitiveOffset < NumPrimitiveDataUploads; PrimitiveOffset += MaxLightmapsUploads)
 				{
@@ -950,8 +954,9 @@ void FGPUScene::UploadGeneral(FRHICommandListImmediate& RHICmdList, FScene *Scen
 
 					for (int32 IndexUpdate = 0; (IndexUpdate < MaxLightmapsUploads) && ((IndexUpdate + PrimitiveOffset) < NumPrimitiveDataUploads); ++IndexUpdate)
 					{
+						const int32 ItemIndex = IndexUpdate + PrimitiveOffset;
 						FLightMapUploadInfo UploadInfo;
-						if (UploadDataSourceAdapter.GetLightMapInfo(IndexUpdate, UploadInfo))
+						if (UploadDataSourceAdapter.GetLightMapInfo(ItemIndex, UploadInfo))
 						{
 							for (int32 LCIIndex = 0; LCIIndex < UploadInfo.LCIs.Num(); LCIIndex++)
 							{
