Fix for skeletal meshes being greyed out in an editing LevelInstance

diff --git a/Engine/Source/Runtime/Engine/Private/SkeletalMesh.cpp b/Engine/Source/Runtime/Engine/Private/SkeletalMesh.cpp
--- a/Engine/Source/Runtime/Engine/Private/SkeletalMesh.cpp
+++ b/Engine/Source/Runtime/Engine/Private/SkeletalMesh.cpp
@@ -5630,6 +5630,8 @@ FPrimitiveViewRelevance FSkeletalMeshSceneProxy::GetViewRelevance(const FSceneVi
 	if (Result.bStaticRelevance)
 	{
 		Result.bEditorStaticSelectionRelevance = (IsSelected() || IsHovered());
+
+		Result.bEditorVisualizeLevelInstanceRelevance = IsEditingLevelInstanceChild();
 	}
 #endif
 
diff --git a/Engine/Source/Runtime/Renderer/Private/SceneRendering.cpp b/Engine/Source/Runtime/Renderer/Private/SceneRendering.cpp
--- a/Engine/Source/Runtime/Renderer/Private/SceneRendering.cpp
+++ b/Engine/Source/Runtime/Renderer/Private/SceneRendering.cpp
@@ -3199,6 +3199,7 @@ void FSceneRenderer::SetupMeshPass(FViewInfo& View, FExclusiveDepthStencil::Type
 					case EMeshPass::HitProxy:
 					case EMeshPass::HitProxyOpaqueOnly:
 					case EMeshPass::EditorSelection:
+					case EMeshPass::EditorLevelInstance:
 #endif
 						break;
 					default:
diff --git a/Engine/Source/Runtime/Renderer/Private/SceneVisibility.cpp b/Engine/Source/Runtime/Renderer/Private/SceneVisibility.cpp
--- a/Engine/Source/Runtime/Renderer/Private/SceneVisibility.cpp
+++ b/Engine/Source/Runtime/Renderer/Private/SceneVisibility.cpp
@@ -2882,6 +2882,9 @@ void ComputeDynamicMeshRelevance(EShadingPath ShadingPath, bool bAddLightmapDens
 	{
 		PassMask.Set(EMeshPass::EditorSelection);
 		View.NumVisibleDynamicMeshElements[EMeshPass::EditorSelection] += NumElements;
+
+		PassMask.Set(EMeshPass::EditorLevelInstance);
+		View.NumVisibleDynamicMeshElements[EMeshPass::EditorLevelInstance] += NumElements;
 	}
 
 	// Hair strands are not rendered into the base pass (bRenderInMainPass=0) and so this 
