Non-Nanite directional light VSM Prototype (disabled, depends on instance culling)
- All non-nanite VSM code is wrapped in ENABLE_NON_NANITE_VSM define (on by default if GPUCULL_TODO is on)
- Added FProjectedShadowInfo with setup to gather primitives for clipmap VSM
- Added RenderVirtualShadowMapsHw to FVirtualShadowMapArray to draw into the VSM
- Shadow depth Vertex shader contains logic to route the instances to pages and uses clip distance to clip to page boundary
- Added culling shaders to replicate instances for each overlapped page

diff --git a/Engine/Shaders/Private/SceneData.ush b/Engine/Shaders/Private/SceneData.ush
--- a/Engine/Shaders/Private/SceneData.ush
+++ b/Engine/Shaders/Private/SceneData.ush
@@ -9,7 +9,7 @@
 // Whether to fetch primitive values (eg LocalToWorld) by dynamically indexing a scene-wide buffer, or to reference a single Primitive uniform buffer
 #define VF_USE_PRIMITIVE_SCENE_DATA ((FEATURE_LEVEL >= FEATURE_LEVEL_SM5 || FEATURE_LEVEL == FEATURE_LEVEL_ES3_1) && VF_SUPPORTS_PRIMITIVE_SCENE_DATA)
 
-#define INSTANCE_SCENE_DATA_SUPPORTED (VF_USE_PRIMITIVE_SCENE_DATA && (FEATURE_LEVEL >= FEATURE_LEVEL_SM5 || FEATURE_LEVEL == FEATURE_LEVEL_ES3_1))
+#define INSTANCE_SCENE_DATA_SUPPORTED (FEATURE_LEVEL >= FEATURE_LEVEL_SM5 || FEATURE_LEVEL == FEATURE_LEVEL_ES3_1)
 
 // Unpacked AoS layout - see FInstanceSceneShaderData::Setup() for SoA packed layout.
 #if INSTANCE_SCENE_DATA_SUPPORTED
@@ -96,7 +96,7 @@ FInstanceSceneData GetInstanceData(uint InstanceId, uint SOAStride)
 	InstanceData.WorldToLocal[3].xyz					= mul(float4(-InstanceData.LocalToWorld[3].xyz, 0.0f), InstanceData.WorldToLocal).xyz;
 
 	// TODO: Do something cleaner to indicate instanced LM/SM UV bias vs. uniform
-	InstanceData.UseInstanceData						= (InstanceData.LightMapAndShadowMapUVBias != float4(-11.0f, -22.0f, -33.0f, -44.0f));
+	InstanceData.UseInstanceData						= any(InstanceData.LightMapAndShadowMapUVBias != float4(-11.0f, -22.0f, -33.0f, -44.0f));
 
 	return InstanceData;
 }
@@ -267,7 +267,10 @@ FPrimitiveSceneData GetPrimitiveData(uint PrimitiveId)
 #define GetPrimitiveData(x) Primitive
 
 // Route to Instance uniform buffer
-#define GetInstanceData(x,y) PrimitiveInstance
+#if !INSTANCE_SCENE_DATA_SUPPORTED
+	// GPUCULL_TODO: Figure out if this one does anything?
+	#define GetInstanceData(x,y) PrimitiveInstance
+#endif // !INSTANCE_SCENE_DATA_SUPPORTED
 
 #endif // VF_USE_PRIMITIVE_SCENE_DATA
 
diff --git a/Engine/Shaders/Private/ShadowDepthVertexShader.usf b/Engine/Shaders/Private/ShadowDepthVertexShader.usf
--- a/Engine/Shaders/Private/ShadowDepthVertexShader.usf
+++ b/Engine/Shaders/Private/ShadowDepthVertexShader.usf
@@ -8,6 +8,14 @@
 #define SHADOW_DEPTH_SHADER 1
 #define USE_STENCIL_LOD_DITHER	0
 
+#if ENABLE_NON_NANITE_VSM
+	#define VIRTUAL_SM_ENABLED (!(USING_TESSELLATION || ONEPASS_POINTLIGHT_SHADOW))
+#else //!ENABLE_NON_NANITE_VSM
+	#define VIRTUAL_SM_ENABLED 0
+#endif // ENABLE_NON_NANITE_VSM
+
+// Need to get the definition of nanite packed view for generated unform buffer (which is pulled in by Common.ush)
+#include "Nanite/PackedNaniteView.ush"
 #include "Common.ush"
 
 // Reroute SceneTexturesStruct uniform buffer references to the shadow depth pass uniform buffer
@@ -23,6 +31,14 @@
 #include "/Engine/Generated/VertexFactory.ush"
 #include "ShadowDepthCommon.ush"
 
+
+#if VIRTUAL_SM_ENABLED
+#include "Nanite/NaniteDataDecode.ush"
+#include "SceneData.ush"
+#include "VirtualShadowMaps/PageAccessCommon.ush"
+#endif // VIRTUAL_SM_ENABLED
+
+
 void SetShadowDepthOutputs(
 	float4x4 WorldToClipMatrix, 
 	float4x4 WorldToShadowMatrix, 
@@ -78,6 +94,121 @@ void SetShadowDepthOutputs(
 	#endif
 }
 
+/** 
+ * Load physical page offset and perform the needed scale and bias to render the instance into the given physical page,
+ * with the correct clipping set up.
+ */
+#if VIRTUAL_SM_ENABLED
+
+FPageInfo UnPackPageInfo(uint PackedData)
+{
+	// TODO: Why not do address translation here? Physical pages need fewer bits.
+	uint4 PackedData4 = 0U;
+	PackedData4.x = PackedData;
+	uint BitPos = 0U;
+	FPageInfo PageInfo;
+	PageInfo.VirtualPage.x = ReadBits(PackedData, BitPos, VSM_LOG2_LEVEL0_DIM_PAGES_XY);
+	PageInfo.VirtualPage.y = ReadBits(PackedData, BitPos, VSM_LOG2_LEVEL0_DIM_PAGES_XY);
+	PageInfo.ViewId = ReadBits(PackedData, BitPos, VSM_PACKED_PAGE_INFO_VIEW_ID_BITS);
+	PageInfo.SuperPageLevel = ReadBits(PackedData, BitPos, VSM_PACKED_PAGE_INFO_SUPER_PAGE_LEVEL_BITS);
+
+	return PageInfo;
+}
+
+// GPUCULL_TODO: This function is awful, and only works for directional lights
+void ScaleBiasClipToPhysicalSmPage(FNaniteView NaniteView, inout float4 PosCsHInOut, inout float4 ClipPlanesInOut, FPageInfo PageInfo)
+{
+	uint2 VirtPage = PageInfo.VirtualPage;
+	uint2 PhysPage = PassStruct.VirtualSmPageTable[CalcPageOffset(NaniteView.TargetLayerIndex, NaniteView.TargetMipLevel, VirtPage)];
+	float2 VirtPageF = float2(VirtPage);
+	float2 PhysPageF = float2(PhysPage);
+	// TODO: do we ever need to support viewports that are not the full mip level here? If so this needs to be taken into account.
+	uint VirtualResXY = VSM_VIRTUAL_MAX_RESOLUTION_XY >> NaniteView.TargetMipLevel;
+	float2 RecPhysicalPoolSize = VirtualShadowMap.RecPhysicalPoolSize.xy;
+	float2 Scale = RecPhysicalPoolSize * float(VirtualResXY);
+
+	// Scale from physical page address space coord to clip space coord
+	float2 PPScale = RecPhysicalPoolSize * 2.0f * float(VSM_PAGE_SIZE);
+
+    //1. subtract virtual page origin
+    //2. Apply scale to physical space
+    //3. Offset to physical address space
+
+	float W = PosCsHInOut.w;
+	// X is different from Y as the page offsets are in pixel address spaces, that are inverted in X wrt clip space
+	//PosCsHInOut.x = PosCsHInOut.x * Scale.x + (PPScale.x * (VirtPageF.x - PhysPageF.x) - Scale.x + 1.0f) * W;
+	//PosCsHInOut.y = PosCsHInOut.y * Scale.y + (PPScale.y * (PhysPageF.y - VirtPageF.y) + Scale.y - 1.0f) * W;
+	// Swapped:
+	//PosCsHInOut.x = PosCsHInOut.x * Scale.x + (PPScale.x * (PhysPageF.x - VirtPageF.x) + Scale.x - 1.0f) * W;
+	//PosCsHInOut.y = PosCsHInOut.y * Scale.y + (PPScale.y * (VirtPageF.y - PhysPageF.y) - Scale.y + 1.0f) * W;
+
+	float LevelResPages = float(VSM_LEVEL0_DIM_PAGES_XY >> NaniteView.TargetMipLevel);
+	
+	float2 VirtPageClipSpaceOrigin = ((VirtPageF + 0.5f) / LevelResPages) * float2(2.0f, -2.0f) + float2(-1.0f, 1.0f);
+	// calculate page-local position (-1,1) range
+	float2 PtPg = (PosCsHInOut.xy - VirtPageClipSpaceOrigin) * LevelResPages;
+
+	float2 PhysPageScale = RecPhysicalPoolSize * VSM_PAGE_SIZE;
+	float2 PhysPageClipSpaceOrigin = ((PhysPageF + 0.5f) * PhysPageScale) * float2(2.0f, -2.0f) + float2(-1.0f, 1.0f);
+	// Scale and Translate to physical page location
+	float2 PhysPt = PtPg * PhysPageScale + PhysPageClipSpaceOrigin;
+
+	uint4 PageRect = PassStruct.PageRectBounds[NaniteView.TargetLayerIndex * VSM_MAX_MIP_LEVELS + NaniteView.TargetMipLevel];
+
+
+	if (PassStruct.bInstancePerPage)
+	{
+		// Functional single-page clip bounds
+		// TODO: Figure out clip boundaries based on super-page rather than page. The routing can stay the same since it doesn't matter
+		//       but the clipping must be adjusted since there will only be one instance for all the pages in the super-page
+		ClipPlanesInOut.x = PtPg.x + 1.0f;
+		ClipPlanesInOut.y = 1.0f - PtPg.x;
+		ClipPlanesInOut.z = PtPg.y + 1.0f;
+		ClipPlanesInOut.w = 1.0f - PtPg.y;
+	}
+	else
+	{
+		// Clip to allocated page rect
+		float2 VirtPageRectMinClipSpace = (float2(PageRect.xy) / LevelResPages) * float2(2.0f, -2.0f) + float2(-1.0f, 1.0f);
+		float2 VirtPageRectMaxClipSpace = (float2(PageRect.zw + 1U) / LevelResPages) * float2(2.0f, -2.0f) + float2(-1.0f, 1.0f);
+		ClipPlanesInOut.x = PosCsHInOut.x - VirtPageRectMinClipSpace.x;
+		ClipPlanesInOut.y = VirtPageRectMaxClipSpace.x - PosCsHInOut.x;
+		ClipPlanesInOut.z = PosCsHInOut.y - VirtPageRectMaxClipSpace.y;
+		ClipPlanesInOut.w = VirtPageRectMinClipSpace.y - PosCsHInOut.y;
+	}
+	PosCsHInOut.xy = PhysPt;
+
+#if 0
+    // positive X, negaive X, positive Y, negative Y
+	ClipPlanesInOut.x = ((PhysPageF.x + 1.0f) * PPScale.x - 1.0f) * W - PosCsHInOut.x;
+	ClipPlanesInOut.y = -(PhysPageF.x * PPScale.x - 1.0f) * W + PosCsHInOut.x;
+	ClipPlanesInOut.z = (1.0f - PhysPageF.y * PPScale.y) * W - PosCsHInOut.y;
+	ClipPlanesInOut.w = -(1.0f - (PhysPageF.y + 1.0f) * PPScale.y) * W + PosCsHInOut.y;
+#endif//0
+}
+
+void TransformToVirtualSmPage(inout float4 PosCsHInOut, inout float4 ClipPlanesInOut, FPageInfo PageInfo, float3 TranslatedWorldPos)
+{
+	// Load nanite view and perform transform to SM space.
+	FNaniteView NaniteView = UnpackNaniteView(PassStruct.PackedNaniteViews[PageInfo.ViewId]);
+	PosCsHInOut = mul(float4(TranslatedWorldPos - ResolvedView.PreViewTranslation.xyz + NaniteView.PreViewTranslation.xyz, 1.0f), NaniteView.TranslatedWorldToClip);
+
+
+	// Clamp the vertex to the near plane if it is in front of the near plane
+	// This has problems if some vertices of a triangle get clamped and others do not, also causes artifacts with non-ortho projections
+	if (PassStruct.bClampToNearPlane > 0 && PosCsHInOut.z > PosCsHInOut.w)
+	{
+		PosCsHInOut.z = 0.999999f;
+		PosCsHInOut.w = 1.0f;
+	}
+
+	// 
+	ScaleBiasClipToPhysicalSmPage(NaniteView, PosCsHInOut, ClipPlanesInOut, PageInfo);
+}
+
+#endif // VIRTUAL_SM_ENABLED
+
+
 #if FEATURE_LEVEL >= FEATURE_LEVEL_SM4 && ONEPASS_POINTLIGHT_SHADOW
 
 // Layer Id per draw call for Vertex-Shader-Layer implementation
@@ -293,6 +424,10 @@ void Main(
 #endif
 	OPTIONAL_VertexID
 	out float4 OutPosition : SV_POSITION
+#if VIRTUAL_SM_ENABLED	
+	// OLA-TODO: this collides with instanced stereo, which thankfully is not used with shadow maps, so should be fine, presumably.
+	, out float4 OutVirtualSmPageClip : SV_ClipDistance
+#endif // VIRTUAL_SM_ENABLED
 	)
 {
 	ResolvedView = ResolveView();
@@ -379,6 +514,19 @@ void Main(
 	#if !PERSPECTIVE_CORRECT_DEPTH && !COMPILER_SUPPORTS_EMPTY_STRUCTS
 		OutParameters.Dummy = 0;
 	#endif
+
+#if VIRTUAL_SM_ENABLED
+	OutVirtualSmPageClip = float4(1.0f, 1.0f, 1.0f, 1.0f);
+	if (PassStruct.bRenderToVirtualShadowMap != 0)
+	{
+		// Get the offset from which we loaded the instance ID
+		uint InstanceIdIndex = VertexFactoryGetInstanceIdLoadIndex(VFIntermediates);
+		FPageInfo PageInfo = UnPackPageInfo(View.PageInfoBuffer[InstanceIdIndex]);
+
+		TransformToVirtualSmPage(OutPosition, OutVirtualSmPageClip, PageInfo, WorldPos.xyz);
+	}
+#endif // VIRTUAL_SM_ENABLED
+
 #endif
 
 	OutputVertexID( OutParameters );
@@ -437,6 +585,10 @@ void PositionOnlyMain(
 	out FShadowDepthVSToPS OutParameters,
 #endif
 	out float4 OutPosition : SV_POSITION
+#if VIRTUAL_SM_ENABLED	
+	// OLA-TODO: this collides with instanced stereo, which thankfully is not used with shadow maps, so should be fine, presumably.
+	, out float4 OutVirtualSmPageClip : SV_ClipDistance
+#endif // VIRTUAL_SM_ENABLED
 	)
 {
 	ResolvedView = ResolveView();
@@ -492,6 +644,19 @@ void PositionOnlyMain(
 #if INTERPOLATE_POSITION
 	OutParameters.PixelPosition = WorldPos.xyz;
 #endif
+
+#if VIRTUAL_SM_ENABLED
+	OutVirtualSmPageClip = float4(1.0f, 1.0f, 1.0f, 1.0f);
+	if (PassStruct.bRenderToVirtualShadowMap != 0)
+	{
+		// Get the offset from which we loaded the instance ID
+		uint InstanceIdIndex = VertexFactoryGetInstanceIdLoadIndex(Input);
+		// TODO: Maybe offset index to local buffer for this? We may not want to use a global since most passes are not supporting this
+		//       Or perhaps both are different, as they are managed somewhat differently anyway, maybe.
+		FPageInfo PageInfo = UnPackPageInfo(View.PageInfoBuffer[InstanceIdIndex]);
+		TransformToVirtualSmPage(OutPosition, OutVirtualSmPageClip, PageInfo, WorldPos.xyz);
+	}
+#endif // VIRTUAL_SM_ENABLED
 }
 
 
diff --git a/Engine/Shaders/Private/VirtualShadowMaps/BuildPerPageDrawCommands.usf b/Engine/Shaders/Private/VirtualShadowMaps/BuildPerPageDrawCommands.usf
new file mode 100644
--- /dev/null
+++ b/Engine/Shaders/Private/VirtualShadowMaps/BuildPerPageDrawCommands.usf
@@ -0,0 +1,482 @@
+// Copyright Epic Games, Inc. All Rights Reserved.
+
+/*=============================================================================
+	BuildPerPageDrawCommands.usf.usf: 
+=============================================================================*/
+
+#include "../Common.ush"
+#include "PageAccessCommon.ush"
+#include "ProjectionCommon.ush"
+#include "../Nanite/NaniteDataDecode.ush"
+
+// Stored in an unordered fashion, needs to be reorganized later
+struct FVisibleInstanceCmd
+{
+	uint PackedPageInfo;
+	uint InstanceId;
+	uint DrawCommandId;
+};
+
+uint PackPageInfo(FPageInfo PageInfo)
+{
+	// TODO: Why not do address translation here? Physical pages need fewer bits.
+	uint4 PackedData = 0U;
+	uint BitPos = 0U;
+	WriteBits(PackedData, BitPos, PageInfo.VirtualPage.x, VSM_LOG2_LEVEL0_DIM_PAGES_XY);
+	WriteBits(PackedData, BitPos, PageInfo.VirtualPage.y, VSM_LOG2_LEVEL0_DIM_PAGES_XY);
+	WriteBits(PackedData, BitPos, PageInfo.ViewId, VSM_PACKED_PAGE_INFO_VIEW_ID_BITS);
+	WriteBits(PackedData, BitPos, PageInfo.SuperPageLevel, VSM_PACKED_PAGE_INFO_SUPER_PAGE_LEVEL_BITS);
+
+	return PackedData.x;
+}
+
+
+RWStructuredBuffer<FVisibleInstanceCmd> VisibleInstancesOut;
+Buffer<uint> InstanceIdsBuffer;
+Buffer<uint> DrawCommandIdsBuffer;
+StructuredBuffer<uint> NumInstanceIdsBuffer;
+uint FirstPrimaryView;
+uint NumPrimaryViews;
+RWStructuredBuffer<uint> DrawCommandInstanceCountBufferOut;
+RWStructuredBuffer<uint> VisibleInstanceCountBufferOut;
+uint InstanceDataSOAStride;
+uint bInstancePerPage;
+
+groupshared uint GroupTotalVisibleInstances;
+groupshared uint GroupVisibleInstanceOutputOffset;
+
+
+// Detitute persons enum...
+#define PML_COUNT (1)
+#define PML_EMIT (2)
+
+// Poor mans template...
+void ProcessMipLevel(uint PrimaryViewId, uint MipLevel, FFrustumCullData Cull, uint InstanceId, uint DrawCommandId, int ProcessingMode, inout uint MipOverlapCount, inout uint VisibleInstanceOutputOffset)
+{
+	uint MipViewId = MipLevel * NumPrimaryViews + PrimaryViewId;
+	FNaniteView MipView = GetNaniteView(MipViewId);
+	uint VirtualShadowMapId = uint(MipView.TargetLayerIndex);
+
+	FScreenRect Rect = GetScreenRect(MipView.ViewRect, Cull, 4);
+
+	// OLA-TODO: Caching is not respected
+	if (OverlapsAnyValidPage(VirtualShadowMapId, MipLevel, MipView.ViewRect.xy, Rect, true))
+	{
+		uint4 RectPages = uint4(Rect.Pixels) >> VSM_LOG2_PAGE_SIZE;
+
+		// Clip to actually allocated pages
+		// TODO: move this to be done as part of or before the overlap test?
+		uint4 AllocatedBounds = PageRectBounds[VirtualShadowMapId * VSM_MAX_MIP_LEVELS + MipLevel];
+		RectPages.xy = max(RectPages.xy, AllocatedBounds.xy);
+		RectPages.zw = min(RectPages.zw, AllocatedBounds.zw);
+		if (all(RectPages.zw >= RectPages.xy))
+		{
+			// 1. Loop over the super-pages, note the implied +1 for the regular level inherent in the >= 0
+			// TODO: May want to get the combined flag from OverlapsAnyValidPage and find the first bit set to not loop where there are no super pages
+			//       but then we can't unroll.
+			for (int SuperPageLevel = VSM_NUM_SUPER_PAGE_LEVELS; SuperPageLevel >= 0; --SuperPageLevel)
+			{
+				// The stride is equal to the size of the super pages
+				uint PageStride = 1U << uint(SuperPageLevel);
+				// Super page mask is stored in the higher bits of the page flags 
+				uint SuperLevelMask = 1U << (uint(SuperPageLevel) + VSM_NUM_SUPER_START_BIT);
+				// Mask to test against to know if there was a higher-level super page for this page
+				// Should be 1 for all bits higher than the current super page
+				uint HigherSuperLevelMask = ~(SuperLevelMask - 1U);
+
+				// Count commands for each overlapped page
+				uint PageTableLevelOffset = CalcPageTableLevelOffset(VirtualShadowMapId, MipLevel);
+				for (uint y = RectPages.y; y <= RectPages.w; y += PageStride)
+				{
+					for (uint x = RectPages.x; x <= RectPages.z; x += PageStride)
+					{
+						uint PageFlagOffset = PageTableLevelOffset + CalcPageOffsetInLevel(MipLevel, uint2(x, y));
+						uint PageFlag = PageFlags[PageFlagOffset];
+
+						// Mask off all bits above and including the current super page level, if this is the first one set, we should emit for this level
+						// otherwise the page is covered by a higher level.
+						if ((PageFlag & HigherSuperLevelMask) == SuperLevelMask)
+						{
+							BRANCH
+							if (ProcessingMode == PML_COUNT)
+							{
+								MipOverlapCount += 1U;
+							}
+							else
+							{
+								FPageInfo PageInfo;
+								PageInfo.VirtualPage = uint2(x, y);
+								PageInfo.ViewId = MipViewId;
+								PageInfo.SuperPageLevel = uint(SuperPageLevel);
+
+								FVisibleInstanceCmd VisibleInstanceCmd;
+								VisibleInstanceCmd.PackedPageInfo = PackPageInfo(PageInfo);
+								VisibleInstanceCmd.InstanceId = InstanceId;
+								VisibleInstanceCmd.DrawCommandId = DrawCommandId;
+
+								VisibleInstancesOut[VisibleInstanceOutputOffset] = VisibleInstanceCmd;
+								VisibleInstanceOutputOffset += 1U;
+							}
+						}
+					}
+				}
+			}
+		}
+	}
+
+}
+
+[numthreads(NUM_THREADS_PER_GROUP, 1, 1)]
+void CullPerPageDrawCommandsCs(uint InstanceIdIndex : SV_DispatchThreadID, uint GroupIndex : SV_GroupIndex)
+{
+	// Produced by previous passes
+	uint NumInstances = NumInstanceIdsBuffer[0];
+
+	if (NumInstances == 0)
+	{
+		return;
+	}
+
+	bool bIsInstanceValid = InstanceIdIndex < NumInstances;
+	// Clamp to a valid index
+	InstanceIdIndex = min(InstanceIdIndex, NumInstances - 1u);
+	if (GroupIndex == 0u)
+	{
+		GroupTotalVisibleInstances = 0u;
+		// Not strictly needed, but probably good form to initialize anyway
+		GroupVisibleInstanceOutputOffset = 0u;
+	}
+
+	// Load the offset (produced by the previous pass on GPU) or maybe use a fresh/recycled buffer for each and set to zero?
+	uint InstanceIdOffset = 0;//InstanceIdsOffsetBuffer[BaseOffset];
+	uint InstanceId = InstanceIdsBuffer[InstanceIdOffset + InstanceIdIndex] & ((1U << 28U) - 1);
+	// Do we need this? Probably for virtual cube maps if we get there.
+	uint ViewIdIndex = InstanceIdsBuffer[InstanceIdOffset + InstanceIdIndex] >> 28U;
+	uint DrawCommandId = DrawCommandIdsBuffer[InstanceIdOffset + InstanceIdIndex];
+
+	const bool bNearClip = false;//(NEAR_CLIP != 0);
+	// Load relevant instance data
+	FInstanceSceneData InstanceData = GetInstanceData(InstanceId, InstanceDataSOAStride);
+	const bool bHasMoved = GetGPUSceneFrameNumber() == InstanceData.LastUpdateSceneFrameNumber;
+
+	GroupMemoryBarrierWithGroupSync();
+
+	uint ThreadTotalForAllViews = 0;
+	// Loop over views and output visible instance (i.e., those that overlap a valid page)
+	for (uint PrimaryViewId = FirstPrimaryView; PrimaryViewId < FirstPrimaryView + NumPrimaryViews; ++PrimaryViewId)
+	{
+		uint VisibleMipsMask = 0;
+		uint ThreadTotalPageOverlaps = 0;
+
+		FNaniteView NaniteView = GetNaniteView(PrimaryViewId);
+		const uint2 TargetViewSize = uint2(NaniteView.ViewSizeAndInvSize.xy);
+
+		float4x4 LocalToTranslatedWorld = InstanceData.LocalToWorld;
+		LocalToTranslatedWorld[3].xyz += NaniteView.PreViewTranslation.xyz;
+		float4x4 LocalToClip = mul(LocalToTranslatedWorld, NaniteView.TranslatedWorldToClip);
+
+		FFrustumCullData Cull = BoxCullFrustum(InstanceData.LocalBoundsCenter, InstanceData.LocalBoundsExtent, LocalToClip, bNearClip, false);
+
+		if (bIsInstanceValid && Cull.bIsVisible)
+		{
+			// Loop over mip levels and count number of output visible instances, also retain a bit for each mip level (used to skip empty levels in output loop)
+			for (uint MipLevel = 0; MipLevel < uint(NaniteView.TargetNumMipLevels); ++MipLevel)
+			{
+				uint MipPageOverlapCount = 0U;
+				//uint DummyVisibleInstanceOutputOffset = 0U;
+				//ProcessMipLevel(PrimaryViewId, MipLevel, Cull, InstanceId, DrawCommandId, PML_COUNT, MipPageOverlapCount, DummyVisibleInstanceOutputOffset);
+				uint MipViewId = MipLevel * NumPrimaryViews + PrimaryViewId;
+				FNaniteView MipView = GetNaniteView(MipViewId);
+				uint VirtualShadowMapId = uint(MipView.TargetLayerIndex);
+
+				FScreenRect Rect = GetScreenRect(MipView.ViewRect, Cull, 4);
+
+				// OLA-TODO: Caching is not respected
+				if (OverlapsAnyValidPage(VirtualShadowMapId, MipLevel, MipView.ViewRect.xy, Rect, true))
+				{
+					uint4 RectPages = uint4(Rect.Pixels) >> VSM_LOG2_PAGE_SIZE;
+
+					// Clip to actually allocated pages
+					// TODO: move this to be done as part of or before the overlap test?
+					uint4 AllocatedBounds = PageRectBounds[VirtualShadowMapId * VSM_MAX_MIP_LEVELS + MipLevel];
+					RectPages.xy = max(RectPages.xy, AllocatedBounds.xy);
+					RectPages.zw = min(RectPages.zw, AllocatedBounds.zw);
+					if (all(RectPages.zw >= RectPages.xy))
+					{
+						if (bInstancePerPage)
+						{
+							// 1. Loop over the super-pages, note the implied +1 for the regular level inherent in the >= 0
+							// TODO: May want to get the combined flag from OverlapsAnyValidPage and find the first bit set to not loop where there are no super pages
+							//       but then we can't unroll.
+							for (int SuperPageLevel = VSM_NUM_SUPER_PAGE_LEVELS; SuperPageLevel >= 0; --SuperPageLevel)
+							{
+								// The stride is equal to the size of the super pages
+								uint PageStride = 1U << uint(SuperPageLevel);
+								// Super page mask is stored in the higher bits of the page flags 
+								uint SuperLevelMask = 1U << (uint(SuperPageLevel) + VSM_NUM_SUPER_START_BIT);
+								// Mask to test against to know if there was a higher-level super page for this page
+								// Should be 1 for all bits higher than the current super page
+								uint HigherSuperLevelMask = ~(SuperLevelMask - 1U);
+
+								// Count commands for each overlapped page
+								uint PageTableLevelOffset = CalcPageTableLevelOffset(VirtualShadowMapId, MipLevel);
+								for (uint y = RectPages.y; y <= RectPages.w; y += PageStride)
+								{
+									for (uint x = RectPages.x; x <= RectPages.z; x += PageStride)
+									{
+										uint PageFlagOffset = PageTableLevelOffset + CalcPageOffsetInLevel(MipLevel, uint2(x, y));
+										uint PageFlag = PageFlags[PageFlagOffset];
+
+										// Mask off all bits above and including the current super page level, if this is the first one set, we should emit for this level
+										// otherwise the page is covered by a higher level.
+										if ((PageFlag & HigherSuperLevelMask) == SuperLevelMask)
+										{
+											MipPageOverlapCount += 1U;
+										}
+									}
+								}
+							}
+						}
+						else
+						{
+							// bInstancePerPage just one per rect
+							MipPageOverlapCount += 1U;
+						}
+					}
+				}
+
+				if (MipPageOverlapCount > 0u)
+				{
+					VisibleMipsMask |= (1u << MipLevel);
+				}
+				ThreadTotalPageOverlaps += MipPageOverlapCount;
+			}
+		}
+
+		uint VisibleInstanceOutputOffset = 0;
+		// Allocate space within the group to get sub-range for output, and also add to total for whole group
+		BRANCH
+		if (ThreadTotalPageOverlaps > 0)
+		{
+			InterlockedAdd(GroupTotalVisibleInstances, ThreadTotalPageOverlaps, VisibleInstanceOutputOffset);
+		}
+		ThreadTotalForAllViews += ThreadTotalPageOverlaps;
+		GroupMemoryBarrierWithGroupSync();
+
+		// Let one thread per group allocate space in global output buffer, sufficient for all instance draw infos.
+		BRANCH
+		if (GroupIndex == 0)
+		{
+			if (GroupTotalVisibleInstances > 0)
+			{
+				uint NumNodes = GroupTotalVisibleInstances;
+				GroupTotalVisibleInstances = 0;
+				InterlockedAdd(VisibleInstanceCountBufferOut[0], NumNodes, GroupVisibleInstanceOutputOffset);
+			}
+		}
+		GroupMemoryBarrierWithGroupSync();
+		// Add offset in global node output buffer to the output offset for the thread
+		VisibleInstanceOutputOffset += GroupVisibleInstanceOutputOffset;
+
+		// Process mip levels with visible pages and output nodes
+		while (VisibleMipsMask)
+		{
+			uint MipLevel = firstbitlow(VisibleMipsMask);	// Jump to next set bit.
+			VisibleMipsMask &= VisibleMipsMask - 1u;				// Remove bottom bit.
+
+			//uint DummyMipPageOverlapCount = 0U;
+			//ProcessMipLevel(PrimaryViewId, MipLevel, Cull, InstanceId, DrawCommandId, PML_EMIT, DummyMipPageOverlapCount, VisibleInstanceOutputOffset);
+			uint MipViewId = MipLevel * NumPrimaryViews + PrimaryViewId;
+			FNaniteView MipView = GetNaniteView(MipViewId);
+			uint VirtualShadowMapId = uint(MipView.TargetLayerIndex);
+
+			FScreenRect Rect = GetScreenRect(MipView.ViewRect, Cull, 4);
+
+			// OLA-TODO: Caching is not respected
+			if (OverlapsAnyValidPage(VirtualShadowMapId, MipLevel, MipView.ViewRect.xy, Rect, true))
+			{
+				uint4 RectPages = uint4(Rect.Pixels) >> VSM_LOG2_PAGE_SIZE;
+
+				// Clip to actually allocated pages
+				// TODO: move this to be done as part of or before the overlap test?
+				uint4 AllocatedBounds = PageRectBounds[VirtualShadowMapId * VSM_MAX_MIP_LEVELS + MipLevel];
+				RectPages.xy = max(RectPages.xy, AllocatedBounds.xy);
+				RectPages.zw = min(RectPages.zw, AllocatedBounds.zw);
+				if (all(RectPages.zw >= RectPages.xy))
+				{
+					if (bInstancePerPage)
+					{
+						// 1. Loop over the super-pages, note the implied +1 for the regular level inherent in the >= 0
+						// TODO: May want to get the combined flag from OverlapsAnyValidPage and find the first bit set to not loop where there are no super pages
+						//       but then we can't unroll.
+						for (int SuperPageLevel = VSM_NUM_SUPER_PAGE_LEVELS; SuperPageLevel >= 0; --SuperPageLevel)
+						{
+							// The stride is equal to the size of the super pages
+							uint PageStride = 1U << uint(SuperPageLevel);
+							// Super page mask is stored in the higher bits of the page flags 
+							uint SuperLevelMask = 1U << (uint(SuperPageLevel) + VSM_NUM_SUPER_START_BIT);
+							// Mask to test against to know if there was a higher-level super page for this page
+							// Should be 1 for all bits higher than the current super page
+							uint HigherSuperLevelMask = ~(SuperLevelMask - 1U);
+
+							// Count commands for each overlapped page
+							uint PageTableLevelOffset = CalcPageTableLevelOffset(VirtualShadowMapId, MipLevel);
+							for (uint y = RectPages.y; y <= RectPages.w; y += PageStride)
+							{
+								for (uint x = RectPages.x; x <= RectPages.z; x += PageStride)
+								{
+									uint PageFlagOffset = PageTableLevelOffset + CalcPageOffsetInLevel(MipLevel, uint2(x, y));
+									uint PageFlag = PageFlags[PageFlagOffset];
+
+									// Mask off all bits above and including the current super page level, if this is the first one set, we should emit for this level
+									// otherwise the page is covered by a higher level.
+									if ((PageFlag & HigherSuperLevelMask) == SuperLevelMask)
+									{
+										FPageInfo PageInfo;
+										PageInfo.VirtualPage = uint2(x, y);
+										PageInfo.ViewId = MipViewId;
+										PageInfo.SuperPageLevel = uint(SuperPageLevel);
+
+										FVisibleInstanceCmd VisibleInstanceCmd;
+										VisibleInstanceCmd.PackedPageInfo = PackPageInfo(PageInfo);
+										VisibleInstanceCmd.InstanceId = InstanceId;
+										VisibleInstanceCmd.DrawCommandId = DrawCommandId;
+
+										VisibleInstancesOut[VisibleInstanceOutputOffset] = VisibleInstanceCmd;
+										VisibleInstanceOutputOffset += 1U;
+									}
+								}
+							}
+						}
+					}
+					else // !bInstancePerPage
+					{
+						FPageInfo PageInfo;
+						PageInfo.VirtualPage = RectPages.xy;
+						PageInfo.ViewId = MipViewId;
+						PageInfo.SuperPageLevel = 0U;
+
+						FVisibleInstanceCmd VisibleInstanceCmd;
+						VisibleInstanceCmd.PackedPageInfo = PackPageInfo(PageInfo);
+						VisibleInstanceCmd.InstanceId = InstanceId;
+						VisibleInstanceCmd.DrawCommandId = DrawCommandId;
+
+						VisibleInstancesOut[VisibleInstanceOutputOffset] = VisibleInstanceCmd;
+						VisibleInstanceOutputOffset += 1U;
+					}
+				}
+			}
+			// End ProcessMipLevel
+#if 0
+			uint MipViewId = MipLevel * NumPrimaryViews + PrimaryViewId;
+			FNaniteView MipView = GetNaniteView(MipViewId);
+
+			uint VirtualShadowMapId = uint(MipView.TargetLayerIndex);
+
+			// TODO: minor optimization possible, but need to duplicate setup from CullRasterize for virtual targets
+			//float2 ViewSize = float2( ( TargetViewSize + ( 1u << MipLevel ) - 1u ) >> MipLevel );
+			FScreenRect Rect = GetScreenRect(MipView.ViewRect, Cull, 4);
+
+			uint4 RectPages = uint4(Rect.Pixels) >> VSM_LOG2_PAGE_SIZE;
+
+			// Clip to actually allocated pages
+			uint4 AllocatedBounds = PageRectBounds[VirtualShadowMapId * VSM_MAX_MIP_LEVELS + MipLevel];
+			RectPages.xy = max(RectPages.xy, AllocatedBounds.xy);
+			RectPages.zw = min(RectPages.zw, AllocatedBounds.zw);
+			// Count commands for each overlapped page
+			uint PageTableLevelOffset = CalcPageTableLevelOffset(VirtualShadowMapId, MipLevel);
+			for (uint y = RectPages.y; y <= RectPages.w; y++)
+			{
+				for (uint x = RectPages.x; x <= RectPages.z; x++)
+				{
+					uint PageFlagOffset = PageTableLevelOffset + CalcPageOffsetInLevel(MipLevel, uint2(x, y));
+					uint PageFlag = PageFlags[PageFlagOffset];
+
+					if (PageFlag != 0)
+					{
+						FPageInfo PageInfo;
+						PageInfo.VirtualPage = uint2(x, y);
+						PageInfo.ViewId = MipViewId;
+
+						FVisibleInstanceCmd VisibleInstanceCmd;
+						VisibleInstanceCmd.PackedPageInfo = PackPageInfo(PageInfo);
+						VisibleInstanceCmd.InstanceId = InstanceId;
+						VisibleInstanceCmd.DrawCommandId = DrawCommandId;
+
+						VisibleInstancesOut[NodeOutputOffset] = VisibleInstanceCmd;
+						NodeOutputOffset += 1U;
+					}
+				}
+			}
+#endif
+		}
+	}
+	if (bIsInstanceValid)
+	{
+		// Accumulate total number of instances/page pairs for each draw command, is used to allocate space and output compact range of instances later
+		InterlockedAdd(DrawCommandInstanceCountBufferOut[DrawCommandId], ThreadTotalForAllViews);
+	}
+}
+
+RWBuffer<uint> DrawIndirectArgsBufferOut;
+RWBuffer<uint> InstanceIdOffsetBufferOut;
+RWStructuredBuffer<uint> OutputOffsetBufferOut;
+RWStructuredBuffer<uint> TmpInstanceIdOffsetBufferOut;
+StructuredBuffer<uint> DrawCommandInstanceCountBuffer;
+
+uint NumDrawCommands;
+/**
+ * Separate pass to allocate space, needs to run once the final space requirements are known. We buffer the page/instance-draw info and reshuffle later.
+ * TODO: Possibly just re-run the culling process in the output pass, saves storing stuff, but may cost more and runs the risk of the passes disagreeing e.g., due to rounding or whatever.
+ */
+[numthreads(NUM_THREADS_PER_GROUP, 1, 1)]
+void AllocateCommandInstanceOutputSpaceCs(uint DrawCommandId : SV_DispatchThreadID)
+{
+	if (DrawCommandId < NumDrawCommands)
+	{
+		uint CommandInstanceCount = DrawCommandInstanceCountBuffer[DrawCommandId];
+		uint CommandInstanceOffset = 0U;
+		if (CommandInstanceCount > 0U)
+		{
+			InterlockedAdd(OutputOffsetBufferOut[0], CommandInstanceCount, CommandInstanceOffset);
+		}
+		InstanceIdOffsetBufferOut[DrawCommandId] = CommandInstanceOffset;
+		// Store second copy for use during output pass (as we need the first offset buffer during the actual rendering)
+		TmpInstanceIdOffsetBufferOut[DrawCommandId] = CommandInstanceOffset;
+		
+		// fix up the instance count of the draw command (the rest is filled in earlier):
+		DrawIndirectArgsBufferOut[DrawCommandId * INDIRECT_ARGS_NUM_WORDS + 1] = CommandInstanceCount;
+	}
+
+	// Also set up indirect dispatch args for the output pass (OutputCommandInstanceLists)
+	//if (DrawCommandId == 0)
+	//{
+	//	uint NumVisibleInstances = VisibleInstanceCountBuffer[0];
+	//	// ...dispatch args to process all the visible instances
+	//}
+}
+
+StructuredBuffer<FVisibleInstanceCmd> VisibleInstances;
+StructuredBuffer <uint> VisibleInstanceCountBuffer;
+//RWStructuredBuffer<uint> TmpInstanceIdOffsetBufferOut;
+RWBuffer<uint> InstanceIdsBufferLegacyOut;
+RWBuffer<uint> PageInfoBufferLegacyOut;
+
+
+[numthreads(NUM_THREADS_PER_GROUP, 1, 1)]
+void OutputCommandInstanceListsCs(uint VisibleInstanceIndex : SV_DispatchThreadID)
+{
+	uint NumVisibleInstances = VisibleInstanceCountBuffer[0];
+
+	if (VisibleInstanceIndex < NumVisibleInstances)
+	{
+		FVisibleInstanceCmd VisibleInstanceCmd = VisibleInstances[VisibleInstanceIndex];
+
+		// Scatter the instance ID & other data.
+		uint InstanceIdOutputOffset = 0;
+		InterlockedAdd(TmpInstanceIdOffsetBufferOut[VisibleInstanceCmd.DrawCommandId], 1U, InstanceIdOutputOffset);
+		// TODO: maybe repack as uint2 since that might be better for these type of presumably scalar loads.
+		InstanceIdsBufferLegacyOut[InstanceIdOutputOffset] = VisibleInstanceCmd.InstanceId;
+		PageInfoBufferLegacyOut[InstanceIdOutputOffset] = VisibleInstanceCmd.PackedPageInfo;
+	}
+}
\ No newline at end of file
diff --git a/Engine/Shaders/Private/VirtualShadowMaps/PageAccessCommon.ush b/Engine/Shaders/Private/VirtualShadowMaps/PageAccessCommon.ush
--- a/Engine/Shaders/Private/VirtualShadowMaps/PageAccessCommon.ush
+++ b/Engine/Shaders/Private/VirtualShadowMaps/PageAccessCommon.ush
@@ -20,11 +20,26 @@
 // Mapped (in page tables) during a call to render (not tracking whether any geo was drawn into the page)
 #define VSM_PHYSICAL_PAGE_STATE_RENDERED 2
 
+
+
+#if ENABLE_NON_NANITE_VSM
+// GPUCULL_TODO: This is probably ok to just leave on.
+// Marks pages that are not cached (bit 0)
+#define VSM_INVALID_FLAG (1)
+// Marks pages that are allocated (bit 1)
+#define VSM_ALLOCATED_FLAG (2)
+// Super pages are of size 2^SP-Level up to VSM_NUM_SUPER_PAGE_LEVELS
+#define VSM_NUM_SUPER_PAGE_LEVELS (3U)
+#define VSM_NUM_SUPER_START_BIT (1U)
+// Mask inclusive of super pages and the allocated flag, Bits [1,1+VSM_NUM_SUPER_PAGE_LEVELS+1]
+#define VSM_SUPER_PAGE_MASK (((1U << (VSM_NUM_SUPER_PAGE_LEVELS + 1U)) - 1U) << 1U)
+#else //!ENABLE_NON_NANITE_VSM
 // Used in the PageFlags
 // Marks pages that are allocated
 #define VSM_ALLOCATED_FLAG (1)
 // Marks pages that are not cached
 #define VSM_INVALID_FLAG (2)
+#endif // ENABLE_NON_NANITE_VSM
 
 struct FPhysicalPageMetaData
 {
@@ -50,6 +65,17 @@ struct FCachedPageInfo
 	float Padding;
 };
 
+// GPUCULL_TODO: These 10 bits don't match Nanite and we may run out what with mip-levels and all!
+#define VSM_PACKED_PAGE_INFO_VIEW_ID_BITS (10U)
+#define VSM_PACKED_PAGE_INFO_SUPER_PAGE_LEVEL_BITS (2U)
+
+struct FPageInfo
+{
+	uint2 VirtualPage;
+	uint ViewId;
+	uint SuperPageLevel;
+};
+
 #if HAS_CACHE_DATA
 StructuredBuffer<FShadowMapCacheData> ShadowMapCacheData;
 StructuredBuffer<uint> PrevPageFlags;
diff --git a/Engine/Shaders/Private/VirtualShadowMaps/PageManagement.usf b/Engine/Shaders/Private/VirtualShadowMaps/PageManagement.usf
--- a/Engine/Shaders/Private/VirtualShadowMaps/PageManagement.usf
+++ b/Engine/Shaders/Private/VirtualShadowMaps/PageManagement.usf
@@ -722,3 +722,229 @@ void InitPhysicalPageMetaData(uint3 Index : SV_DispatchThreadID)
 	}
 }
 
+#if ENABLE_NON_NANITE_VSM
+
+
+StructuredBuffer<uint> InputCountBuffer;
+RWBuffer<uint> IndirectDispatchArgsOut;
+uint Multiplier;
+uint Divisor;
+uint InputCountOffset;
+
+[numthreads( 1, 1, 1 )]
+void InitIndirectArgs1D(
+	uint3 DTID : SV_DispatchThreadID)
+{
+	uint GroupCount = (InputCountBuffer[InputCountOffset] * Multiplier + Divisor - 1U) / Divisor;
+
+	IndirectDispatchArgsOut[ 0 ] = GroupCount;
+	IndirectDispatchArgsOut[ 1 ] = 1;
+	IndirectDispatchArgsOut[ 2 ] = 1;
+}
+
+
+
+
+// 1. Calc page rects. Dispacth size {VSM_LEVEL0_DIM_PAGES_XY,VSM_LEVEL0_DIM_PAGES_XY,NumSms}
+[numthreads(VSM_DEFAULT_CS_GROUP_XY, VSM_DEFAULT_CS_GROUP_XY, 1)]
+void CalculatePageRects(uint3 DispatchThreadID : SV_DispatchThreadID)
+{
+	// One group per each mip 0
+	uint ShadowMapID = DispatchThreadID.z;
+	uint2 vPage = DispatchThreadID.xy;
+
+	UNROLL
+	for (uint MipLevel = 0; MipLevel < VSM_MAX_MIP_LEVELS; ++MipLevel)
+	{
+		// Early out for out of bounds 
+		if (any(vPage >= uint2(VSM_LEVEL0_DIM_PAGES_XY >> MipLevel, VSM_LEVEL0_DIM_PAGES_XY >> MipLevel)))
+		{
+			return;
+		}
+
+		uint PageOffset = CalcPageOffset(ShadowMapID, MipLevel, vPage);
+		uint Flag = PageFlags[PageOffset];
+		if (Flag)
+		{
+			// Compute the min/max rect of requested pages
+			uint PageBoundIndex = ShadowMapID * VSM_MAX_MIP_LEVELS + MipLevel;
+			InterlockedMin(PageRectBoundsOut[PageBoundIndex].x, vPage.x);
+			InterlockedMin(PageRectBoundsOut[PageBoundIndex].y, vPage.y);
+			InterlockedMax(PageRectBoundsOut[PageBoundIndex].z, vPage.x);
+			InterlockedMax(PageRectBoundsOut[PageBoundIndex].w, vPage.y);
+		}
+		OutPageTable[PageOffset] = uint2(VSM_INVALID_PHYSICAL_PAGE_ADDRESS, VSM_INVALID_PHYSICAL_PAGE_ADDRESS);
+
+	}
+}
+
+
+bool TestRectOverlap(uint4 A, uint4 B)
+{
+	return all(A.xy <= B.zw) && all(A.zw >= B.xy);
+}
+
+#define MAX_PAGE_RECTS_BUFFER 1024
+
+// Store compacted range of index+area to skip all empty ones at least
+groupshared uint2 RectsToProcess[MAX_PAGE_RECTS_BUFFER];
+groupshared uint NumRectsToProcessGroupShared;
+
+groupshared uint MaxAreaAndIndex;
+groupshared uint NextFreeOffsetX;
+
+groupshared uint4 AllocatedRectsGroupShared[MAX_PAGE_RECTS_BUFFER];
+groupshared uint NumAllocatedRectsGroupShared;
+
+RWStructuredBuffer<uint4> AllocatedPageRectBoundsOut;
+
+// 2. Sort page rects (just find max)
+// 3. Allocate page rects in phys space. (serial)
+[numthreads(VSM_DEFAULT_CS_GROUP_X, 1, 1)]
+void AllocatePagesUsingRects(uint GroupThreadId : SV_GroupThreadID, uint2 DispatchThreadID : SV_DispatchThreadID)
+{
+	if (GroupThreadId == 0)
+	{
+		NumRectsToProcessGroupShared = 0;
+		MaxAreaAndIndex = 0;
+		NextFreeOffsetX = 0;
+		NumAllocatedRectsGroupShared = 0;
+	}
+	GroupMemoryBarrierWithGroupSync();
+
+	uint TotalMips = VirtualShadowMap.NumShadowMaps * VSM_MAX_MIP_LEVELS;
+	// Load all rects that are non-empty
+	for (uint Index = GroupThreadId; Index < TotalMips; Index += VSM_DEFAULT_CS_GROUP_X)
+	{
+		uint4 Rect = PageRectBounds[Index];
+		AllocatedPageRectBoundsOut[Index] = uint4(VSM_LEVEL0_DIM_PAGES_XY, VSM_LEVEL0_DIM_PAGES_XY, 0, 0);
+		if (all(Rect.zw >= Rect.xy))
+		{
+			uint WriteOffset = 0;
+			uint Area = (Rect.z - Rect.x + 1) * (Rect.w - Rect.y + 1);
+			InterlockedAdd(NumRectsToProcessGroupShared, 1, WriteOffset);
+			if (WriteOffset < MAX_PAGE_RECTS_BUFFER)
+			{
+				RectsToProcess[WriteOffset] = uint2(Index, Area);
+			}
+		}
+	}
+	GroupMemoryBarrierWithGroupSync();
+	// Protect against overflow (but should also detect and somehow fix)
+	uint NumRectsToProcess = min(MAX_PAGE_RECTS_BUFFER, NumRectsToProcessGroupShared);
+	// Loop until all the rects are processed, each iteration finds the max, and then processes that.
+	// Note: this is extremely stupid! If we had a simple to use sorting function for group shared then we could make it less extra stupid.
+	for (uint IterIndex = 0; IterIndex < NumRectsToProcess; ++IterIndex)
+	{
+		// Group stride loop to find max index using atomics
+		for (uint Index = GroupThreadId; Index < NumRectsToProcess; Index += VSM_DEFAULT_CS_GROUP_X)
+		{
+			uint2 Item = RectsToProcess[Index];
+			// Pack area and item index (not source rect index)
+			uint AreaAndIndex = (Item.y << 12U) | Index;
+			InterlockedMax(MaxAreaAndIndex, AreaAndIndex);
+		}
+		GroupMemoryBarrierWithGroupSync();
+
+		// Find item that won the bidding
+		uint ItemIndex = MaxAreaAndIndex & ((1U << 12U) - 1U);
+		// Get info
+		uint2 Item = RectsToProcess[ItemIndex];
+		// Reset stuff and clear the winning rect so it won't win again
+		GroupMemoryBarrierWithGroupSync();
+		if (GroupThreadId == 0)
+		{
+			RectsToProcess[ItemIndex] = uint2(0U, 0U);
+			MaxAreaAndIndex = 0U;
+		}
+		GroupMemoryBarrierWithGroupSync();
+		// Now we have the next largest item, must find allocation
+
+		// Allocation Algo
+		// Scan each possible position across each row, 
+		// test overlap against all existing allocations and skip horizontally
+		// when we run out of row, move to next and repeat
+		uint4 Rect = PageRectBounds[Item.x];
+		uint ShadowMapID = Item.x / VSM_MAX_MIP_LEVELS;
+		uint MipLevel = Item.x % VSM_MAX_MIP_LEVELS;
+		uint2 RectSize = Rect.zw - Rect.xy;
+		// Scan each row where it could fit
+		for (uint RowInd = 0; RowInd < VirtualShadowMap.PhysicalPoolSizePages.y - RectSize.y; ++RowInd)
+		{
+			// Current probe position, test against all rects using group stride loop
+			uint2 TestPos = uint2(0U, RowInd);
+			uint4 ProbeRect = uint4(TestPos.xy, TestPos.xy + RectSize.xy);
+			bool bFoundFree = false;
+			
+			// Until we run off the edge 
+			while (ProbeRect.z < VirtualShadowMap.PhysicalPoolSizePages.x)
+			{
+				NextFreeOffsetX = 0;
+				GroupMemoryBarrierWithGroupSync();
+				for (uint Index = GroupThreadId; Index < NumAllocatedRectsGroupShared; Index += VSM_DEFAULT_CS_GROUP_X)
+				{
+					uint4 AllocatedRect = AllocatedRectsGroupShared[Index];
+					if (TestRectOverlap(AllocatedRect, ProbeRect))
+					{
+						// Find free location to the right of all rects on this row
+						InterlockedMax(NextFreeOffsetX, AllocatedRect.z + 1);
+					}
+				}
+				GroupMemoryBarrierWithGroupSync();
+				// If we did not overlap anything that modified the rect, then we found a free space.
+				if (NextFreeOffsetX == 0U)
+				{
+					bFoundFree = true;
+					break;
+				}
+				else
+				{
+					// Bump probe rect for next iteration
+					ProbeRect.xz = uint2(NextFreeOffsetX, NextFreeOffsetX + RectSize.x);
+				}
+				// need to sync before next iteration also or else we may read the NextFreeOffsetX out of order
+				GroupMemoryBarrierWithGroupSync();
+			}
+
+			if (bFoundFree)
+			{
+				if (GroupThreadId == 0)
+				{
+					AllocatedRectsGroupShared[NumAllocatedRectsGroupShared] = ProbeRect;
+					NumAllocatedRectsGroupShared += 1;
+					AllocatedPageRectBoundsOut[Item.x] = ProbeRect;
+				}
+					// 4. Allocate pages to rects (1:1)
+				for (uint PY = 0U; PY <= RectSize.y; ++PY)
+				{
+					// Group stride loop in X
+					for (uint PX = GroupThreadId; PX <= RectSize.x; PX += VSM_DEFAULT_CS_GROUP_X)
+					{
+						// Virtual page
+						uint2 vPage = Rect.xy + uint2(PX, PY);
+						uint2 pPage = ProbeRect.xy + uint2(PX, PY);
+
+						FCachedPageInfo CachedPageInfo;
+
+						CachedPageInfo.PhysPageAddress = uint2(VSM_INVALID_PHYSICAL_PAGE_ADDRESS, VSM_INVALID_PHYSICAL_PAGE_ADDRESS);
+						CachedPageInfo.DepthOffset = 0.0f;
+						CachedPageInfo.Padding = 0.0f;
+						// Store data that is used later in the process (for each physical page).
+						OutCachedPageInfos[pPage.y * VirtualShadowMap.PhysicalPoolSizePages.x + pPage.x] = CachedPageInfo;
+
+						uint PageOffset = CalcPageOffset(ShadowMapID, MipLevel, vPage);
+						// Request rendering of pages that are not cached (AKA VSM_INVALID_PHYSICAL_PAGE_ADDRESS in cache).
+						// TODO: make sure caching machinery is use in this path also
+						OutPageFlags[PageOffset] = VSM_INVALID_FLAG | VSM_ALLOCATED_FLAG;
+						OutPageTable[PageOffset] = pPage;
+					}
+				}
+
+				break;
+			}
+		}
+		GroupMemoryBarrierWithGroupSync();
+	}
+}
+
+#endif // ENABLE_NON_NANITE_VSM
diff --git a/Engine/Shaders/Private/VirtualShadowMaps/PrintStats.usf b/Engine/Shaders/Private/VirtualShadowMaps/PrintStats.usf
--- a/Engine/Shaders/Private/VirtualShadowMaps/PrintStats.usf
+++ b/Engine/Shaders/Private/VirtualShadowMaps/PrintStats.usf
@@ -7,6 +7,10 @@
 // Just the page count for now
 StructuredBuffer< uint >	InStatsBuffer;
 
+#if ENABLE_NON_NANITE_VSM
+StructuredBuffer<uint4> AllocatedPageRectBounds;
+#endif // ENABLE_NON_NANITE_VSM
+
 float2 ShaderPrintLeftAlign(float2 ScreenPos, uint Value)
 {
 	int Len = 9;
@@ -127,3 +131,114 @@ void PrintStats()
 		ScreenPos = ShaderPrintLeftAlign(ScreenPos, (InStatsBuffer[2] * 100) / InStatsBuffer[0]);
 	}
 }
+
+#if ENABLE_NON_NANITE_VSM
+
+
+uint ShadowMapIdRangeStart;
+uint ShadowMapIdRangeEnd;
+
+[numthreads(1, 1, 1)]
+void PrintClipmapStats()
+{
+	float TopMargin = 0.35f;
+	float HeadlineX = 0.07f;
+	float ItemX = 0.08f;
+
+	// Main Pass
+	float2 ScreenPos = float2(HeadlineX, TopMargin);
+	ScreenPos.x = ItemX;
+	ScreenPos = ShaderPrintSymbol(ScreenPos, _P_);
+	ScreenPos = ShaderPrintSymbol(ScreenPos, _A_);
+	ScreenPos = ShaderPrintSymbol(ScreenPos, _G_);
+	ScreenPos = ShaderPrintSymbol(ScreenPos, _E_);
+	ScreenPos = ShaderPrintSymbol(ScreenPos, _SPC_);
+	ScreenPos = ShaderPrintSymbol(ScreenPos, _R_);
+	ScreenPos = ShaderPrintSymbol(ScreenPos, _E_);
+	ScreenPos = ShaderPrintSymbol(ScreenPos, _C_);
+	ScreenPos = ShaderPrintSymbol(ScreenPos, _T_);
+	ScreenPos = ShaderPrintSymbol(ScreenPos, _S_);
+	ScreenPos = ShaderPrintNewline(ScreenPos);
+	uint TotalArea = 0U;
+	for (uint SmId = ShadowMapIdRangeStart; SmId < ShadowMapIdRangeEnd; ++SmId)
+	{
+		uint4 Rect = PageRectBounds[SmId * VSM_MAX_MIP_LEVELS];
+		ScreenPos.x = ItemX;
+		ScreenPos = ShaderPrint(ScreenPos, Rect.x, SHADER_PRINT_TYPE_UINT);//ShaderPrintLeftAlign(ScreenPos, Rect.x);
+		ScreenPos = ShaderPrintSymbol(ScreenPos, _SPC_);
+		ScreenPos = ShaderPrint(ScreenPos, Rect.y, SHADER_PRINT_TYPE_UINT);//ShaderPrintLeftAlign(ScreenPos, Rect.y);
+		ScreenPos = ShaderPrintSymbol(ScreenPos, _SPC_);
+		ScreenPos = ShaderPrint(ScreenPos, Rect.z, SHADER_PRINT_TYPE_UINT);//ShaderPrintLeftAlign(ScreenPos, Rect.z);
+		ScreenPos = ShaderPrintSymbol(ScreenPos, _SPC_);
+		ScreenPos = ShaderPrint(ScreenPos, Rect.w, SHADER_PRINT_TYPE_UINT);//ShaderPrintLeftAlign(ScreenPos, Rect.w);
+		if (Rect.x <= Rect.z && Rect.y <= Rect.w)
+		{
+			ScreenPos = ShaderPrintSymbol(ScreenPos, _SPC_);
+			ScreenPos = ShaderPrint(ScreenPos, Rect.z - Rect.x + 1, SHADER_PRINT_TYPE_UINT);//ShaderPrintLeftAlign(ScreenPos, Rect.z);
+			ScreenPos = ShaderPrintSymbol(ScreenPos, _SPC_);
+			ScreenPos = ShaderPrint(ScreenPos, Rect.w - Rect.y + 1, SHADER_PRINT_TYPE_UINT);//ShaderPrintLeftAlign(ScreenPos, Rect.w);
+
+			uint Area = (Rect.z - Rect.x + 1) * (Rect.w - Rect.y + 1);
+			ScreenPos = ShaderPrintSymbol(ScreenPos, _SPC_);
+			ScreenPos = ShaderPrint(ScreenPos, Area, SHADER_PRINT_TYPE_UINT);//ShaderPrintLeftAlign(ScreenPos, Rect.w);
+
+			TotalArea += Area;
+		}
+		ScreenPos = ShaderPrintNewline(ScreenPos);
+	}
+	ScreenPos.x = ItemX;
+	ScreenPos = ShaderPrintSymbol(ScreenPos, _T_);
+	ScreenPos = ShaderPrintSymbol(ScreenPos, _O_);
+	ScreenPos = ShaderPrintSymbol(ScreenPos, _T_);
+	ScreenPos = ShaderPrintSymbol(ScreenPos, _A_);
+	ScreenPos = ShaderPrintSymbol(ScreenPos, _L_);
+	ScreenPos = ShaderPrintSymbol(ScreenPos, _SPC_);
+	ScreenPos = ShaderPrintSymbol(ScreenPos, _A_);
+	ScreenPos = ShaderPrintSymbol(ScreenPos, _R_);
+	ScreenPos = ShaderPrintSymbol(ScreenPos, _E_);
+	ScreenPos = ShaderPrintSymbol(ScreenPos, _A_);
+	ScreenPos = ShaderPrintSymbol(ScreenPos, _SPC_);
+	ScreenPos = ShaderPrint(ScreenPos, TotalArea, SHADER_PRINT_TYPE_UINT);//ShaderPrintLeftAlign(ScreenPos, Rect.x);
+	ScreenPos = ShaderPrintNewline(ScreenPos);
+
+	ScreenPos.x = ItemX;
+	ScreenPos = ShaderPrintSymbol(ScreenPos, _A_);
+	ScreenPos = ShaderPrintSymbol(ScreenPos, _L_);
+	ScreenPos = ShaderPrintSymbol(ScreenPos, _C_);
+	ScreenPos = ShaderPrintSymbol(ScreenPos, _D_);
+	ScreenPos = ShaderPrintSymbol(ScreenPos, _SPC_);
+	ScreenPos = ShaderPrintSymbol(ScreenPos, _R_);
+	ScreenPos = ShaderPrintSymbol(ScreenPos, _E_);
+	ScreenPos = ShaderPrintSymbol(ScreenPos, _C_);
+	ScreenPos = ShaderPrintSymbol(ScreenPos, _T_);
+	ScreenPos = ShaderPrintSymbol(ScreenPos, _S_);
+	ScreenPos = ShaderPrintNewline(ScreenPos);
+	for (uint SmId = ShadowMapIdRangeStart; SmId < ShadowMapIdRangeEnd; ++SmId)
+	{
+		uint4 Rect = AllocatedPageRectBounds[SmId * VSM_MAX_MIP_LEVELS];
+		ScreenPos.x = ItemX;
+		ScreenPos = ShaderPrint(ScreenPos, Rect.x, SHADER_PRINT_TYPE_UINT);//ShaderPrintLeftAlign(ScreenPos, Rect.x);
+		ScreenPos = ShaderPrintSymbol(ScreenPos, _SPC_);
+		ScreenPos = ShaderPrint(ScreenPos, Rect.y, SHADER_PRINT_TYPE_UINT);//ShaderPrintLeftAlign(ScreenPos, Rect.y);
+		ScreenPos = ShaderPrintSymbol(ScreenPos, _SPC_);
+		ScreenPos = ShaderPrint(ScreenPos, Rect.z, SHADER_PRINT_TYPE_UINT);//ShaderPrintLeftAlign(ScreenPos, Rect.z);
+		ScreenPos = ShaderPrintSymbol(ScreenPos, _SPC_);
+		ScreenPos = ShaderPrint(ScreenPos, Rect.w, SHADER_PRINT_TYPE_UINT);//ShaderPrintLeftAlign(ScreenPos, Rect.w);
+		if (Rect.x <= Rect.z && Rect.y <= Rect.w)
+		{
+			ScreenPos = ShaderPrintSymbol(ScreenPos, _SPC_);
+			ScreenPos = ShaderPrint(ScreenPos, Rect.z - Rect.x + 1, SHADER_PRINT_TYPE_UINT);//ShaderPrintLeftAlign(ScreenPos, Rect.z);
+			ScreenPos = ShaderPrintSymbol(ScreenPos, _SPC_);
+			ScreenPos = ShaderPrint(ScreenPos, Rect.w - Rect.y + 1, SHADER_PRINT_TYPE_UINT);//ShaderPrintLeftAlign(ScreenPos, Rect.w);
+
+			uint Area = (Rect.z - Rect.x + 1) * (Rect.w - Rect.y + 1);
+			ScreenPos = ShaderPrintSymbol(ScreenPos, _SPC_);
+			ScreenPos = ShaderPrint(ScreenPos, Area, SHADER_PRINT_TYPE_UINT);//ShaderPrintLeftAlign(ScreenPos, Rect.w);
+
+			TotalArea += Area;
+		}
+		ScreenPos = ShaderPrintNewline(ScreenPos);
+	}
+
+}
+#endif // ENABLE_NON_NANITE_VSM
diff --git a/Engine/Shaders/Private/VirtualShadowMaps/ProjectionCommon.ush b/Engine/Shaders/Private/VirtualShadowMaps/ProjectionCommon.ush
--- a/Engine/Shaders/Private/VirtualShadowMaps/ProjectionCommon.ush
+++ b/Engine/Shaders/Private/VirtualShadowMaps/ProjectionCommon.ush
@@ -14,6 +14,7 @@ ProjectionCommon.ush:
 // Must match C++ structure in VirtualShadowMapArray.h
 struct FVirtualShadowMapProjectionShaderData
 {
+	float4x4 TranslatedWorldToShadowViewMatrix;
 	float4x4 ShadowViewToClipMatrix;
 	float4x4 TranslatedWorldToShadowUVMatrix;
 	float4x4 TranslatedWorldToShadowUVNormalMatrix;
@@ -34,6 +35,9 @@ struct FVirtualShadowMapProjectionShaderData
 // FVirtualShadowMapProjectionParameters
 StructuredBuffer< FVirtualShadowMapProjectionShaderData > VirtualShadowMapProjectionData;
 Texture2D< uint > PhysicalPagePool;
+#if ENABLE_NON_NANITE_VSM
+Texture2D< float > PhysicalPagePoolHw;
+#endif // ENABLE_NON_NANITE_VSM
 
 FVirtualShadowMapProjectionShaderData GetVirtualShadowMapProjectionData(int VirtualShadowMapId)
 {
@@ -91,7 +95,14 @@ FVirtualSmSample SampleVirtualSmLevel(int VirtualShadowMapId, float2 ShadowMapUV
 	if (PagePresent)
 	{
 		// TODO: maybe store metadata about Max depth (and/or use sparse HZB) for each tile, and early out if completely in shadow (must test using kernel size).
+#if ENABLE_NON_NANITE_VSM
+		Result.Depth = max(
+			asfloat(PhysicalPagePool.Load(int3(Result.PhysicalTexelAddress, 0))),
+			PhysicalPagePoolHw.Load(int3(Result.PhysicalTexelAddress, 0))
+		);
+#else //!ENABLE_NON_NANITE_VSM
 		Result.Depth = asfloat(PhysicalPagePool.Load(int3(Result.PhysicalTexelAddress, 0)));
+#endif // ENABLE_NON_NANITE_VSM
 		Result.MipLevel = VSMLevel;
 		Result.bValid = true;
 	}
diff --git a/Engine/Source/Runtime/Renderer/Private/SceneRendering.h b/Engine/Source/Runtime/Renderer/Private/SceneRendering.h
--- a/Engine/Source/Runtime/Renderer/Private/SceneRendering.h
+++ b/Engine/Source/Runtime/Renderer/Private/SceneRendering.h
@@ -33,6 +33,7 @@
 #include "GpuDebugRendering.h"
 #include "PostProcess/PostProcessAmbientOcclusionMobile.h"
 #include "Nanite/NaniteRender.h"
+#include "VirtualShadowMaps/VirtualShadowMapConfig.h"
 #include "VirtualShadowMaps/VirtualShadowMapArray.h"
 #include "Lumen/LumenTranslucencyVolumeLighting.h"
 #include "GPUScene.h"
@@ -1677,6 +1678,9 @@ struct FSortedShadowMaps
 	TArray<FProjectedShadowInfo*, SceneRenderingAllocator> VirtualShadowMapShadows;
 
 	TArray<TSharedPtr<FVirtualShadowMapClipmap>, SceneRenderingAllocator> VirtualShadowMapClipmaps;
+#if ENABLE_NON_NANITE_VSM
+	TArray<FProjectedShadowInfo*, SceneRenderingAllocator> VirtualShadowClipmapsHw;
+#endif // ENABLE_NON_NANITE_VSM
 
 	TArray<FSortedShadowMapAtlas,SceneRenderingAllocator> CompleteShadowMapAtlases;
 
diff --git a/Engine/Source/Runtime/Renderer/Private/ShadowDepthRendering.cpp b/Engine/Source/Runtime/Renderer/Private/ShadowDepthRendering.cpp
--- a/Engine/Source/Runtime/Renderer/Private/ShadowDepthRendering.cpp
+++ b/Engine/Source/Runtime/Renderer/Private/ShadowDepthRendering.cpp
@@ -161,6 +161,13 @@ void SetupShadowDepthPassUniformBuffer(
 			ShadowDepthPassParameters.ShadowViewMatrices[FaceIndex] = Translation * ShadowInfo->OnePassShadowViewMatrices[FaceIndex];
 		}
 	}
+#if ENABLE_NON_NANITE_VSM
+	ShadowDepthPassParameters.bRenderToVirtualShadowMap = false;
+	ShadowDepthPassParameters.bInstancePerPage = false;
+	ShadowDepthPassParameters.VirtualSmPageTable = GraphBuilder.CreateSRV(CreateStructuredBuffer(GraphBuilder, TEXT("Dummy-VirtualSmPageTable"), TArray<uint32>()));
+	ShadowDepthPassParameters.PackedNaniteViews = GraphBuilder.CreateSRV(CreateStructuredBuffer(GraphBuilder, TEXT("Dummy-PackedNaniteViews"), TArray<Nanite::FPackedView>()));
+	ShadowDepthPassParameters.PageRectBounds = GraphBuilder.CreateSRV(CreateStructuredBuffer(GraphBuilder, TEXT("Dummy-PageRectBounds"), TArray<FIntVector4>()));
+#endif // ENABLE_NON_NANITE_VSM
 }
 
 void SetupShadowDepthPassUniformBuffer(
@@ -316,7 +323,10 @@ public:
 		OutEnvironment.SetDefine(TEXT("USING_VERTEX_SHADER_LAYER"), (uint32)(ShaderMode == VertexShadowDepth_VSLayer));
 		OutEnvironment.SetDefine(TEXT("POSITION_ONLY"), (uint32)bUsePositionOnlyStream);
 		OutEnvironment.SetDefine(TEXT("IS_FOR_GEOMETRY_SHADER"), (uint32)bIsForGeometryShader);
-
+		OutEnvironment.SetDefine(TEXT("ENABLE_NON_NANITE_VSM"), (uint32)ENABLE_NON_NANITE_VSM);
+#if ENABLE_NON_NANITE_VSM
+		FVirtualShadowMapArray::SetShaderDefines(OutEnvironment);
+#endif // ENABLE_NON_NANITE_VSM
 		if (bIsForGeometryShader)
 		{
 			OutEnvironment.CompilerFlags.Add(CFLAG_VertexToGeometryShader);
@@ -815,12 +825,22 @@ void FProjectedShadowInfo::SetStateForView(FRHICommandList& RHICmdList) const
 	);
 }
 
-void SetStateForShadowDepth(bool bOnePassPointLightShadow, FMeshPassProcessorRenderState& DrawRenderState)
+void SetStateForShadowDepth(bool bOnePassPointLightShadow, bool bDirectionalLight, FMeshPassProcessorRenderState& DrawRenderState)
 {
 	// Disable color writes
 	DrawRenderState.SetBlendState(TStaticBlendState<CW_NONE>::GetRHI());
 
+#if ENABLE_NON_NANITE_VSM
+
+	// GPUCULL_TODO: Unhack me please
+	static const auto EnableVirtualSMCVar = IConsoleManager::Get().FindTConsoleVariableDataInt(TEXT("r.Shadow.v.Enable"));
+
+	const bool bNeedsVirtualShadowMap = EnableVirtualSMCVar->GetValueOnRenderThread() != 0 && bDirectionalLight;
+
+	if (bOnePassPointLightShadow || bNeedsVirtualShadowMap)
+#else //!ENABLE_NON_NANITE_VSM
 	if (bOnePassPointLightShadow)
+#endif // ENABLE_NON_NANITE_VSM
 	{
 		// Point lights use reverse Z depth maps
 		DrawRenderState.SetDepthStencilState(TStaticDepthStencilState<true, CF_DepthNearOrEqual>::GetRHI());
@@ -1093,6 +1113,9 @@ BEGIN_SHADER_PARAMETER_STRUCT(FShadowDepthPassParameters, )
 	SHADER_PARAMETER_STRUCT_REF(FViewUniformShaderParameters, View)
 	SHADER_PARAMETER_RDG_UNIFORM_BUFFER(FMobileShadowDepthPassUniformParameters, MobilePassUniformBuffer)
 	SHADER_PARAMETER_RDG_UNIFORM_BUFFER(FShadowDepthPassUniformParameters, DeferredPassUniformBuffer)
+#if ENABLE_NON_NANITE_VSM
+	SHADER_PARAMETER_RDG_UNIFORM_BUFFER(FVirtualShadowMapCommonParameters, VirtualSmCommon)
+#endif // ENABLE_NON_NANITE_VSM
 	SHADER_PARAMETER_STRUCT_INCLUDE(FInstanceCullingDrawParams, InstanceCullingDrawParams)
 	RENDER_TARGET_BINDING_SLOTS()
 END_SHADER_PARAMETER_STRUCT()
@@ -1135,9 +1158,13 @@ void FProjectedShadowInfo::RenderDepth(
 	{
 		// Copy in depths of static primitives before we render movable primitives.
 		FMeshPassProcessorRenderState DrawRenderState;
-		SetStateForShadowDepth(bOnePassPointLightShadow, DrawRenderState);
+		SetStateForShadowDepth(bOnePassPointLightShadow, bDirectionalLight, DrawRenderState);
 		CopyCachedShadowMap(GraphBuilder, *ShadowDepthView, SceneRenderer, PassParameters->RenderTargets, DrawRenderState);
 	}
+#if ENABLE_NON_NANITE_VSM
+	FVirtualShadowMapCommonParameters DummyVsmParameters = FVirtualShadowMapCommonParameters();
+	PassParameters->VirtualSmCommon = GraphBuilder.CreateUniformBuffer(&DummyVsmParameters);
+#endif // ENABLE_NON_NANITE_VSM
 
 	switch (FSceneInterface::GetShadingPath(FeatureLevel))
 	{
@@ -1755,6 +1782,14 @@ void FSceneRenderer::RenderShadowDepthMaps(FRDGBuilder& GraphBuilder, FInstanceC
 	const bool bNaniteEnabled = UseNanite(ShaderPlatform) && ViewFamily.EngineShowFlags.NaniteMeshes;
 	const bool bUseHZB = (CVarNaniteShadowsUseHZB.GetValueOnRenderThread() != 0);
 
+#if ENABLE_NON_NANITE_VSM
+
+	if (SortedShadowsForShadowDepthPass.VirtualShadowClipmapsHw.Num() != 0)
+	{
+		VirtualShadowMapArray.RenderVirtualShadowMapsHw(GraphBuilder, SortedShadowsForShadowDepthPass.VirtualShadowClipmapsHw, *Scene);
+	}
+#endif // ENABLE_NON_NANITE_VSM
+
 	if (bNaniteEnabled && (bHasVSMShadows || bHasVSMClipMaps))
 	{
 		if (bUseHZB)
@@ -2355,7 +2390,7 @@ FShadowDepthPassMeshProcessor::FShadowDepthPassMeshProcessor(
 	: FMeshPassProcessor(Scene, Scene->GetFeatureLevel(), InViewIfDynamicMeshCommand, InDrawListContext)
 	, ShadowDepthType(InShadowDepthType)
 {
-	SetStateForShadowDepth(ShadowDepthType.bOnePassPointLightShadow, PassDrawRenderState);
+	SetStateForShadowDepth(ShadowDepthType.bOnePassPointLightShadow, ShadowDepthType.bDirectionalLight, PassDrawRenderState);
 }
 
 FShadowDepthType CSMShadowDepthType(true, false);
diff --git a/Engine/Source/Runtime/Renderer/Private/ShadowRendering.cpp b/Engine/Source/Runtime/Renderer/Private/ShadowRendering.cpp
--- a/Engine/Source/Runtime/Renderer/Private/ShadowRendering.cpp
+++ b/Engine/Source/Runtime/Renderer/Private/ShadowRendering.cpp
@@ -1813,7 +1813,7 @@ void FDeferredShadingSceneRenderer::RenderShadowProjections(
 		check(Views.Num() == 1);					// TODO: Test multiview
 		
 
-		if (VirtualShadowMapArray.IsAllocated())
+		if (VirtualShadowMapArray.HasAnyShadowData())
 		{
 			RDG_EVENT_SCOPE(GraphBuilder, "Virtual Shadow Maps" );
 
diff --git a/Engine/Source/Runtime/Renderer/Private/ShadowRendering.h b/Engine/Source/Runtime/Renderer/Private/ShadowRendering.h
--- a/Engine/Source/Runtime/Renderer/Private/ShadowRendering.h
+++ b/Engine/Source/Runtime/Renderer/Private/ShadowRendering.h
@@ -29,6 +29,7 @@
 #include "LightRendering.h"
 #include "HairStrands/HairStrandsRendering.h"
 #include "Strata/Strata.h"
+#include "VirtualShadowMaps/VirtualShadowMapConfig.h"
 
 ENGINE_API IPooledRenderTarget* GetSubsufaceProfileTexture_RT(FRHICommandListImmediate& RHICmdList);
 
@@ -82,7 +83,7 @@ public:
 
 	inline bool operator==(const FShadowDepthType& rhs) const
 	{
-		if (bDirectionalLight != rhs.bDirectionalLight ||
+		if (bDirectionalLight != rhs.bDirectionalLight || 
 			bOnePassPointLightShadow != rhs.bOnePassPointLightShadow)
 		{
 			return false;
@@ -197,6 +198,7 @@ public:
 };
 
 typedef TFunctionRef<void(FRHICommandList& RHICmdList, bool bFirst)> FBeginShadowRenderPassFunction;
+using FPackedNaniteView = Nanite::FPackedView;
 
 BEGIN_GLOBAL_SHADER_PARAMETER_STRUCT(FShadowDepthPassUniformParameters,)
 	SHADER_PARAMETER_STRUCT(FSceneTextureUniformParameters, SceneTextures)
@@ -206,6 +208,17 @@ BEGIN_GLOBAL_SHADER_PARAMETER_STRUCT(FShadowDepthPassUniformParameters,)
 	SHADER_PARAMETER(float, bClampToNearPlane)
 	SHADER_PARAMETER_ARRAY(FMatrix, ShadowViewProjectionMatrices, [6])
 	SHADER_PARAMETER_ARRAY(FMatrix, ShadowViewMatrices, [6])
+#if ENABLE_NON_NANITE_VSM
+	// GPUCULL_TODO: ?
+	SHADER_PARAMETER(int, bRenderToVirtualShadowMap)
+	SHADER_PARAMETER(int, bInstancePerPage)
+
+	SHADER_PARAMETER_RDG_BUFFER_SRV(StructuredBuffer<uint2>, VirtualSmPageTable)
+	SHADER_PARAMETER_RDG_BUFFER_SRV(StructuredBuffer< FPackedNaniteView >, PackedNaniteViews)
+	SHADER_PARAMETER_RDG_BUFFER_SRV(StructuredBuffer< uint4 >, PageRectBounds)
+
+	//SHADER_PARAMETER_RDG_BUFFER_SRV(Buffer< uint >, PackedPageInfoBuffer)
+#endif // ENABLE_NON_NANITE_VSM
 END_GLOBAL_SHADER_PARAMETER_STRUCT()
 
 BEGIN_GLOBAL_SHADER_PARAMETER_STRUCT(FMobileShadowDepthPassUniformParameters,)
@@ -395,7 +408,9 @@ public:
 	int32 SubjectPrimitiveComponentIndex = -1;
 
 	TArray<int32> ViewIds;
-
+#if ENABLE_NON_NANITE_VSM
+	TSharedPtr<FVirtualShadowMapClipmap> VirtualShadowMapClipmap;
+#endif // ENABLE_NON_NANITE_VSM
 public:
 
 	// default constructor
@@ -431,6 +446,15 @@ public:
 		uint32 InSnapResolutionY,
 		uint32 InBorderSize
 		);
+#if ENABLE_NON_NANITE_VSM
+	/** for a clipmap shadow. */
+	void SetupClipmapProjection(
+		FLightSceneInfo* InLightSceneInfo,
+		FViewInfo* InDependentView,
+		const TSharedPtr<FVirtualShadowMapClipmap> &VirtualShadowMapClipmap,
+		float InMaxNonFarCascadeDistance
+		);
+#endif // ENABLE_NON_NANITE_VSM
 
 	float GetShaderDepthBias() const { return ShaderDepthBias; }
 	float GetShaderSlopeDepthBias() const { return ShaderSlopeDepthBias; }
@@ -447,7 +471,7 @@ public:
 			int32(Y + BorderSize),
 			int32(X + BorderSize + ResolutionX),
 			int32(Y + BorderSize + ResolutionY),
-		};
+			};
 	}
 
 	/** Get view rect including border area */
@@ -650,6 +674,9 @@ public:
 
 	FParallelMeshDrawCommandPass& GetShadowDepthPass() { return ShadowDepthPass; }
 
+#if ENABLE_NON_NANITE_VSM
+	float GetMaxNonFarCascadeDistance() const { return MaxNonFarCascadeDistance; }
+#endif // ENABLE_NON_NANITE_VSM
 private:
 	// 0 if Setup...() wasn't called yet
 	FLightSceneInfo* LightSceneInfo;
@@ -689,6 +716,10 @@ private:
 	FDynamicMeshDrawCommandStorage DynamicMeshDrawCommandStorage;
 	FGraphicsMinimalPipelineStateSet GraphicsMinimalPipelineStateSet;
 	bool NeedsShaderInitialisation;
+#if ENABLE_NON_NANITE_VSM
+	// If >= 0.0f then this records the cascade distance to the farthest non-'far' shadow cascade. Only used for clipmaps at the moment.
+	float MaxNonFarCascadeDistance = -1.0f;
+#endif // ENABLE_NON_NANITE_VSM
 
 	/**
 	 * Bias during in shadowmap rendering, stored redundantly for better performance 
@@ -710,6 +741,7 @@ private:
 	*/
 	void ModifyViewForShadow(FRHICommandList& RHICmdList, FViewInfo* FoundView) const;
 
+	friend class FVirtualShadowMapArray;
 	void BeginRenderView(FRDGBuilder& GraphBuilder, FScene* Scene);
 
 	/**
diff --git a/Engine/Source/Runtime/Renderer/Private/ShadowSetup.cpp b/Engine/Source/Runtime/Renderer/Private/ShadowSetup.cpp
--- a/Engine/Source/Runtime/Renderer/Private/ShadowSetup.cpp
+++ b/Engine/Source/Runtime/Renderer/Private/ShadowSetup.cpp
@@ -275,6 +275,16 @@ static TAutoConsoleVariable<int32> CVarAlwaysAllocateMaxResolutionAtlases(
 );
 
 
+#if ENABLE_NON_NANITE_VSM
+
+static TAutoConsoleVariable<int32> CVarVsmUseFarShadowRules(
+	TEXT("r.Shadow.v.UseFarShadowCulling"),
+	1,
+	TEXT("Switch between implementing the far shadow culling logic for VSMs."),
+	ECVF_RenderThreadSafe
+);
+
+#endif // ENABLE_NON_NANITE_VSM
 
 CSV_DECLARE_CATEGORY_EXTERN(LightCount);
 
@@ -975,6 +985,44 @@ void FProjectedShadowInfo::SetupWholeSceneProjection(
 	}
 }
 
+
+#if ENABLE_NON_NANITE_VSM
+
+void FProjectedShadowInfo::SetupClipmapProjection(FLightSceneInfo* InLightSceneInfo, FViewInfo* InDependentView, const TSharedPtr<FVirtualShadowMapClipmap>& InVirtualShadowMapClipmap, float InMaxNonFarCascadeDistance)
+{
+	VirtualShadowMapClipmap = InVirtualShadowMapClipmap;
+	LightSceneInfo = InLightSceneInfo;
+	LightSceneInfoCompact = InLightSceneInfo;
+	DependentView = InDependentView;
+	ResolutionX = FVirtualShadowMap::VirtualMaxResolutionXY;
+	ResolutionY = FVirtualShadowMap::VirtualMaxResolutionXY;
+	bDirectionalLight = true;
+	bWholeSceneShadow = true;
+	BorderSize = 0;
+	bIncludeInScreenSpaceShadowMask = false;
+	MaxNonFarCascadeDistance = InMaxNonFarCascadeDistance;
+
+
+	FVirtualShadowMapProjectionShaderData Data = VirtualShadowMapClipmap->GetProjectionShaderData(VirtualShadowMapClipmap->GetLevelCount() - 1);
+	PreShadowTranslation = Data.ShadowPreViewTranslation;
+	const FMatrix CasterMatrix = Data.TranslatedWorldToShadowViewMatrix * Data.ShadowViewToClipMatrix;
+	GetViewFrustumBounds(CasterOuterFrustum, CasterMatrix, true);
+	ReceiverInnerFrustum = CasterOuterFrustum;
+
+	// Magically get the accurate shadow bounds set up - ka-ching!
+	ShadowBounds = InLightSceneInfo->Proxy->GetShadowSplitBoundsDepthRange(*DependentView, DependentView->ViewMatrices.GetViewOrigin(), 10.0f, VirtualShadowMapClipmap->GetMaxRadius(), &CascadeSettings);
+	// Um... it's checked in IsWholeSceneDirectionalShadow()
+	CascadeSettings.ShadowSplitIndex = 1000;
+
+	ViewToClipInner = Data.ShadowViewToClipMatrix;
+	ViewToClipOuter = Data.ShadowViewToClipMatrix;
+	TranslatedWorldToView = Data.TranslatedWorldToShadowViewMatrix;
+	TranslatedWorldToClipInnerMatrix = TranslatedWorldToView * ViewToClipInner;
+	TranslatedWorldToClipOuterMatrix = TranslatedWorldToView * ViewToClipOuter;
+}
+
+#endif // ENABLE_NON_NANITE_VSM
+
 void FProjectedShadowInfo::AddCachedMeshDrawCommandsForPass(
 	int32 PrimitiveIndex,
 	const FPrimitiveSceneInfo* InPrimitiveSceneInfo,
@@ -1682,6 +1730,22 @@ bool FProjectedShadowInfo::TestPrimitiveFarCascadeConditions(bool bPrimitiveCast
 		return bPrimitiveCastsFarShadow;
 	}
 
+#if ENABLE_NON_NANITE_VSM
+
+	const bool bWholeSceneDirectionalShadow = IsWholeSceneDirectionalShadow();
+
+	if (bWholeSceneDirectionalShadow && MaxNonFarCascadeDistance > 0.0f)
+	{
+		check(DependentView);
+
+		FVector ViewOrigin = DependentView->ShadowViewMatrices.GetViewOrigin();
+		float ViewDistance = (Bounds.Origin - ViewOrigin) | DependentView->GetViewDirection();
+		if (ViewDistance - Bounds.SphereRadius > MaxNonFarCascadeDistance)
+		{
+			return bPrimitiveCastsFarShadow;
+		}
+	}
+#endif // ENABLE_NON_NANITE_VSM
 	return true;
 }
 
@@ -2040,6 +2104,21 @@ void FProjectedShadowInfo::SetupMeshDrawCommandsForShadowDepth(FSceneRenderer& R
 			}
 		}
 	}
+
+#if ENABLE_NON_NANITE_VSM
+	else if (VirtualShadowMapClipmap.IsValid())
+	{
+		// TODO: Register view per clip level such that they are culled early (?)
+		Nanite::FPackedViewParams Params;
+		// Note: To ensure conservative culling, we get the coarsest clip-level view since it covers the finer ones.
+		Params.ViewMatrices = VirtualShadowMapClipmap->GetViewMatrices(VirtualShadowMapClipmap->GetLevelCount() - 1);
+		// TODO: Real prev frame matrices
+		Params.PrevViewMatrices = Params.ViewMatrices;
+		Params.ViewRect = GetInnerViewRect();
+		Params.RasterContextSize = FIntPoint(ResolutionX, ResolutionY);
+		ViewIds.Add(InstanceCullingManager.RegisterView(Params));
+	}
+#endif // ENABLE_NON_NANITE_VSM
 	else
 	{
 		Nanite::FPackedViewParams Params;
@@ -2204,7 +2283,7 @@ void FProjectedShadowInfo::SetupMeshDrawCommandsForProjectionStenciling(FSceneRe
 }
 
 void FProjectedShadowInfo::GatherDynamicMeshElements(FSceneRenderer& Renderer, FVisibleLightInfo& VisibleLightInfo, TArray<const FSceneView*>& ReusedViewsArray,
-	FGlobalDynamicIndexBuffer& DynamicIndexBuffer, FGlobalDynamicVertexBuffer& DynamicVertexBuffer, FGlobalDynamicReadBuffer& DynamicReadBuffer, FInstanceCullingManager& InstanceCullingManager)
+	FGlobalDynamicIndexBuffer& DynamicIndexBuffer, FGlobalDynamicVertexBuffer& DynamicVertexBuffer, FGlobalDynamicReadBuffer& DynamicReadBuffer, FInstanceCullingManager &InstanceCullingManager)
 {
 	QUICK_SCOPE_CYCLE_COUNTER(STAT_Shadow_GatherDynamicMeshElements);
 
@@ -2526,7 +2605,11 @@ void FSceneRenderer::SetupInteractionShadows(
 	if (!bShadowHandledByParent)
 	{
 		const bool bCreateTranslucentObjectShadow = GUseTranslucencyShadowDepths && Interaction->HasTranslucentObjectShadow();
+#if ENABLE_NON_NANITE_VSM
+		const bool bCreateInsetObjectShadow = !UseVirtualShadowMaps(ShaderPlatform, FeatureLevel) &&  Interaction->HasInsetObjectShadow();
+#else //!ENABLE_NON_NANITE_VSM
 		const bool bCreateInsetObjectShadow = Interaction->HasInsetObjectShadow();
+#endif // ENABLE_NON_NANITE_VSM
 		const bool bCreateObjectShadowForStationaryLight = ShouldCreateObjectShadowForStationaryLight(Interaction->GetLight(), PrimitiveSceneInfo->Proxy, Interaction->IsShadowMapped());
 
 		if (Interaction->HasShadow() 
@@ -3834,6 +3917,15 @@ void FSceneRenderer::GatherShadowDynamicMeshElements(FGlobalDynamicIndexBuffer&
 			ProjectedShadowInfo->GatherDynamicMeshElements(*this, VisibleLightInfo, ReusedViewsArray, DynamicIndexBuffer, DynamicVertexBuffer, DynamicReadBuffer, InstanceCullingManager);
 		}
 	}
+#if ENABLE_NON_NANITE_VSM
+
+	// GPUCULL_TODO: Replace with new shadow culling processor thingo
+	for (FProjectedShadowInfo* ProjectedShadowInfo : SortedShadowsForShadowDepthPass.VirtualShadowClipmapsHw)
+	{
+		FVisibleLightInfo& VisibleLightInfo = VisibleLightInfos[ProjectedShadowInfo->GetLightSceneInfo().Id];
+		ProjectedShadowInfo->GatherDynamicMeshElements(*this, VisibleLightInfo, ReusedViewsArray, DynamicIndexBuffer, DynamicVertexBuffer, DynamicReadBuffer, InstanceCullingManager);
+	}
+#endif // defined(ENABLE_NON_NANITE_VSM)
 }
 
 typedef TArray<FAddSubjectPrimitiveOp> FShadowSubjectPrimitives;
@@ -3949,6 +4041,13 @@ struct FGatherShadowPrimitivesPacket
 
 	bool DoesPrimitiveCastInsetShadow(const FPrimitiveSceneInfo* PrimitiveSceneInfo, const FPrimitiveSceneProxy* PrimitiveProxy) const
 	{
+#if ENABLE_NON_NANITE_VSM
+		if (UseVirtualShadowMaps(GShaderPlatformForFeatureLevel[FeatureLevel], FeatureLevel))
+		{
+			return false;
+		}
+#endif // ENABLE_NON_NANITE_VSM
+
 		// If light attachment root is valid, we're in a group and need to get the flag from the root.
 		if (PrimitiveSceneInfo->LightingAttachmentRoot.IsValid())
 		{
@@ -4295,6 +4394,9 @@ void FSceneRenderer::AddViewDependentWholeSceneShadowsForView(
 			static_assert(INDEX_NONE == -1, "INDEX_NONE != -1!");
 
 			float MaxShadowCascadeDistance = 0.0f;
+#if ENABLE_NON_NANITE_VSM
+			float MaxNonFarCascadeDistance = 0.0f;
+#endif // ENABLE_NON_NANITE_VSM
 
 			// todo: this code can be simplified by computing all the distances in one place - avoiding some redundant work and complexity
 			for (int32 Index = 0; Index < ProjectionCount; Index++)
@@ -4312,7 +4414,12 @@ void FSceneRenderer::AddViewDependentWholeSceneShadowsForView(
 				if (LightSceneInfo.Proxy->GetViewDependentWholeSceneProjectedShadowInitializer(View, LocalIndex, LightSceneInfo.IsPrecomputedLightingValid(), ProjectedShadowInitializer))
 				{
 					MaxShadowCascadeDistance = FMath::Max(MaxShadowCascadeDistance, ProjectedShadowInitializer.CascadeSettings.SplitFar);
-
+#if ENABLE_NON_NANITE_VSM
+					if (!ProjectedShadowInitializer.CascadeSettings.bFarShadowCascade)
+					{
+						MaxNonFarCascadeDistance = FMath::Max(MaxNonFarCascadeDistance, ProjectedShadowInitializer.CascadeSettings.SplitFar);
+					}
+#endif // ENABLE_NON_NANITE_VSM
 					uint32 ShadowBorder = NeedsUnatlasedCSMDepthsWorkaround( FeatureLevel ) ? 0 : SHADOW_BORDER;
 					
 					const int32 MaxCSMResolution = GetCachedScalabilityCVars().MaxCSMShadowResolution;
@@ -4335,9 +4442,15 @@ void FSceneRenderer::AddViewDependentWholeSceneShadowsForView(
 					FIntPoint CompleteShadowBufferResolution(ShadowBufferResolution.X / CompleteResolutionFactor,
 														     ShadowBufferResolution.Y / CompleteResolutionFactor);
 
+					FProjectedShadowInfo* RegularCascadePSInfo = nullptr;
+
+#if ENABLE_NON_NANITE_VSM
+					if (!bNeedsVirtualShadowMap)
+#endif // ENABLE_NON_NANITE_VSM
+					{
 					// Create the projected shadow info.
-					FProjectedShadowInfo* ProjectedShadowInfo = new(FMemStack::Get(), 1, 16) FProjectedShadowInfo;
-					ProjectedShadowInfo->SetupWholeSceneProjection(
+						RegularCascadePSInfo = new(FMemStack::Get(), 1, 16) FProjectedShadowInfo;
+						RegularCascadePSInfo->SetupWholeSceneProjection(
 						&LightSceneInfo,
 						&View,
 						ProjectedShadowInitializer,
@@ -4348,18 +4461,27 @@ void FSceneRenderer::AddViewDependentWholeSceneShadowsForView(
 						CompleteShadowBufferResolution.Y,
 						ShadowBorder
 					);
-					ProjectedShadowInfo->FadeAlphas = FadeAlphas;
-					ProjectedShadowInfo->ProjectionIndex = Index;
-
-					VisibleLightInfo.MemStackProjectedShadows.Add(ProjectedShadowInfo);
-					VisibleLightInfo.AllProjectedShadows.Add(ProjectedShadowInfo);
-					ShadowInfos.Add(ProjectedShadowInfo);
+						RegularCascadePSInfo->FadeAlphas = FadeAlphas;
+						RegularCascadePSInfo->ProjectionIndex = Index;
 
+						VisibleLightInfo.MemStackProjectedShadows.Add(RegularCascadePSInfo);
+						VisibleLightInfo.AllProjectedShadows.Add(RegularCascadePSInfo);
+						ShadowInfos.Add(RegularCascadePSInfo);
+					}
+#if ENABLE_NON_NANITE_VSM
+					// GPUCULL_TODO: THis logic is not fully functional!
+					// If we have a virtual shadow map, 
+					if (false)//bNeedsVirtualShadowMap)
+					{
+#else //!ENABLE_NON_NANITE_VSM
 					if (bNeedsVirtualShadowMap)
 					{
-						// If we have a virtual shadow map, disable nanite rendering into the regular shadow map or else we'd get double-shadowing
-						ProjectedShadowInfo->bNaniteGeometry = false;
-						
+#endif // ENABLE_NON_NANITE_VSM
+						if (RegularCascadePSInfo != nullptr)
+						{
+							// ...disable nanite rendering into the regular shadow map or else we'd get double-shadowing
+							RegularCascadePSInfo->bNaniteGeometry = false;
+						}
 						if (bNeedsCompleteShadowMap)
 						{
 							// Create the projected shadow info.
@@ -4379,7 +4501,7 @@ void FSceneRenderer::AddViewDependentWholeSceneShadowsForView(
 
 							CompleteProjectedShadowInfo->bCompleteShadowMap = true;
 							CompleteProjectedShadowInfo->bIncludeInScreenSpaceShadowMask = false;
-							CompleteProjectedShadowInfo->CompleteShadowMapCopySource = ProjectedShadowInfo;
+							CompleteProjectedShadowInfo->CompleteShadowMapCopySource = RegularCascadePSInfo;
 
 							VisibleLightInfo.AllProjectedShadows.Add(CompleteProjectedShadowInfo);
 							VisibleLightInfo.CompleteProjectedShadows.Add(CompleteProjectedShadowInfo);
@@ -4388,9 +4510,9 @@ void FSceneRenderer::AddViewDependentWholeSceneShadowsForView(
 					}
 
 					// Ray traced shadows use the GPU managed distance field object buffers, no CPU culling needed
-					if (!ProjectedShadowInfo->bRayTracedDistanceField)
+					if (RegularCascadePSInfo != nullptr && !RegularCascadePSInfo->bRayTracedDistanceField)
 					{
-						ShadowInfosThatNeedCulling.Add(ProjectedShadowInfo);
+						ShadowInfosThatNeedCulling.Add(RegularCascadePSInfo);
 					}
 				}
 			}
@@ -4412,6 +4534,16 @@ void FSceneRenderer::AddViewDependentWholeSceneShadowsForView(
 				// TODO: Not clear we need both of these in this path, but keep it consistent for now
 				VisibleLightInfo.VirtualShadowMapClipmaps.Add(VirtualShadowMapClipmap);
 				SortedShadowsForShadowDepthPass.VirtualShadowMapClipmaps.Add(VirtualShadowMapClipmap);
+
+#if ENABLE_NON_NANITE_VSM
+				// Create the projected shadow info to make sure that culling happens.
+				FProjectedShadowInfo* ProjectedShadowInfo = new(FMemStack::Get(), 1, 16) FProjectedShadowInfo;
+				// Add to remember-to-call-dtor list
+				VisibleLightInfo.MemStackProjectedShadows.Add(ProjectedShadowInfo);
+				ProjectedShadowInfo->SetupClipmapProjection(&LightSceneInfo, &View, VirtualShadowMapClipmap, CVarVsmUseFarShadowRules.GetValueOnRenderThread() != 0 ? MaxNonFarCascadeDistance : -1.0f);
+				VisibleLightInfo.AllProjectedShadows.Add(ProjectedShadowInfo);
+				ShadowInfosThatNeedCulling.Add(ProjectedShadowInfo);
+#endif // ENABLE_NON_NANITE_VSM
 			}
 		}
 	}
@@ -4530,6 +4662,14 @@ void FSceneRenderer::AllocateShadowDepthTargets(FRHICommandListImmediate& RHICmd
 
 				if (bNeedsShadowmapSetup)
 				{
+#if ENABLE_NON_NANITE_VSM
+					if (ProjectedShadowInfo->VirtualShadowMapClipmap.IsValid())
+					{
+						ProjectedShadowInfo->SetupShadowDepthView(RHICmdList, this);
+						SortedShadowsForShadowDepthPass.VirtualShadowClipmapsHw.Add(ProjectedShadowInfo);
+					}
+					else 
+#endif // ENABLE_NON_NANITE_VSM
 					if (ProjectedShadowInfo->HasVirtualShadowMap())
 					{
 						ProjectedShadowInfo->SetupShadowDepthView(RHICmdList, this);
diff --git a/Engine/Source/Runtime/Renderer/Private/VirtualShadowMaps/VirtualShadowMapArray.cpp b/Engine/Source/Runtime/Renderer/Private/VirtualShadowMaps/VirtualShadowMapArray.cpp
--- a/Engine/Source/Runtime/Renderer/Private/VirtualShadowMaps/VirtualShadowMapArray.cpp
+++ b/Engine/Source/Runtime/Renderer/Private/VirtualShadowMaps/VirtualShadowMapArray.cpp
@@ -10,6 +10,7 @@
 #include "ShaderPrintParameters.h"
 #include "VirtualShadowMapCacheManager.h"
 #include "VirtualShadowMapClipmap.h"
+#include "ComponentRecreateRenderStateContext.h"
 
 IMPLEMENT_STATIC_UNIFORM_BUFFER_SLOT(VirtualShadowMapUbSlot);
 
@@ -40,11 +41,17 @@ struct FCachedPageInfo
 	float Padding;
 };
 
+// GPUCULL_TODO: Hook this up again!
+static void RecreateGlobalRenderState(IConsoleVariable* Var)
+{
+	FGlobalComponentRecreateRenderStateContext Context;
+}
 
 static TAutoConsoleVariable<int32> CVarEnableVirtualShadowMaps(
 	TEXT("r.Shadow.v.Enable"),
 	0,
 	TEXT("Enable Virtual Shadow Maps."),
+	//FConsoleVariableDelegate::CreateStatic(&RecreateGlobalRenderState),
 	ECVF_RenderThreadSafe
 );
 
@@ -88,6 +95,24 @@ FMatrix CalcTranslatedWorldToShadowUVMatrix(
 	FMatrix TranslatedWorldToShadowUv = TranslatedWorldToShadowClip * ScaleAndBiasToSmUV;
 	return TranslatedWorldToShadowUv;
 }
+#if ENABLE_NON_NANITE_VSM
+
+static TAutoConsoleVariable<int32> CVarAllocatePagesUsingRects(
+	TEXT("r.Shadow.v.AllocatePagesUsingRects"),
+	0,
+	TEXT("If set to 1 then pages are allocated in a contigious block to support the bounding rectangle of allocated virtual pages leading to a trivial mapping but higher memory overhead."),
+	ECVF_RenderThreadSafe
+);
+
+
+static TAutoConsoleVariable<int32> CVarShowClipmapStats(
+	TEXT("r.Shadow.v.ShowClipmapStats"),
+	-1,
+	TEXT("Set to the number of clipmap you want to show stats for (-1 == off)\n"),
+	ECVF_RenderThreadSafe
+);
+
+#endif // ENABLE_NON_NANITE_VSM
 
 FMatrix CalcTranslatedWorldToShadowUVNormalMatrix(
 	const FMatrix& TranslatedWorldToShadowView,
@@ -122,6 +147,10 @@ FVirtualShadowMapArray::FVirtualShadowMapArray()
 	CommonParameters.MaxPhysicalPages = PhysSizePages.X * PhysSizePages.Y;
 	CommonParameters.PhysicalPageRowMask = (PhysSizePages.X - 1);
 	CommonParameters.PhysicalPageRowShift = FMath::FloorLog2( PhysSizePages.X );
+	// 1.0f / float(GetPhysicalPoolSize())
+	CommonParameters.RecPhysicalPoolSize = FVector4(1.0f / PhysSize.X, 1.0f / PhysSize.Y, 1.0f, 1.0f);
+	CommonParameters.PhysicalPoolSizePages = PhysSizePages;
+
 }
 
 BEGIN_SHADER_PARAMETER_STRUCT(FCacheDataParameters, )
@@ -174,6 +203,11 @@ void FVirtualShadowMapArray::Initialize(bool bInEnabled)
 
 void FVirtualShadowMapArray::SetShaderDefines(FShaderCompilerEnvironment& OutEnvironment)
 {
+#if ENABLE_NON_NANITE_VSM
+	static_assert(FVirtualShadowMap::Log2Level0DimPagesXY * 2U + MAX_VIEWS_PER_CULL_RASTERIZE_PASS_BITS <= 32U, "Page indirection plus view index must fit into 32-bits for page-routing storage!");
+#endif // ENABLE_NON_NANITE_VSM
+	OutEnvironment.SetDefine(TEXT("ENABLE_NON_NANITE_VSM"), int32(ENABLE_NON_NANITE_VSM));
+
 	OutEnvironment.SetDefine(TEXT("VSM_PAGE_SIZE"), FVirtualShadowMap::PageSize);
 	OutEnvironment.SetDefine(TEXT("VSM_PAGE_SIZE_MASK"), FVirtualShadowMap::PageSizeMask);
 	OutEnvironment.SetDefine(TEXT("VSM_LOG2_PAGE_SIZE"), FVirtualShadowMap::Log2PageSize);
@@ -208,7 +242,11 @@ public:
 
 	static bool ShouldCompilePermutation(const FGlobalShaderPermutationParameters& Parameters)
 	{
+#if ENABLE_NON_NANITE_VSM
+		return FDataDrivenShaderPlatformInfo::GetSupportsGPUScene(Parameters.Platform);
+#else //!ENABLE_NON_NANITE_VSM
 		return DoesPlatformSupportNanite(Parameters.Platform);
+#endif // ENABLE_NON_NANITE_VSM
 	}
 
 	/**
@@ -344,6 +382,52 @@ class FClearPhysicalPagesCS : public FVirtualPageManagementShader
 IMPLEMENT_GLOBAL_SHADER(FClearPhysicalPagesCS, "/Engine/Private/VirtualShadowMaps/PageManagement.usf", "ClearPhysicalPages", SF_Compute);
 
 
+#if ENABLE_NON_NANITE_VSM
+
+class FInitIndirectArgs1DCS : public FVirtualPageManagementShader
+{
+	DECLARE_GLOBAL_SHADER(FInitIndirectArgs1DCS);
+	SHADER_USE_PARAMETER_STRUCT(FInitIndirectArgs1DCS, FVirtualPageManagementShader)
+
+	BEGIN_SHADER_PARAMETER_STRUCT(FParameters, )
+		SHADER_PARAMETER_RDG_BUFFER_SRV(StructuredBuffer< uint >, InputCountBuffer)
+		SHADER_PARAMETER(uint32, Multiplier)
+		SHADER_PARAMETER(uint32, Divisor)
+		SHADER_PARAMETER(uint32, InputCountOffset)
+
+		SHADER_PARAMETER_RDG_BUFFER_UAV(RWBuffer< uint >, IndirectDispatchArgsOut)
+	END_SHADER_PARAMETER_STRUCT()
+};
+IMPLEMENT_GLOBAL_SHADER(FInitIndirectArgs1DCS, "/Engine/Private/VirtualShadowMaps/PageManagement.usf", "InitIndirectArgs1D", SF_Compute);
+
+
+FRDGBufferRef AddIndirectArgsSetupCsPass1D(FRDGBuilder& GraphBuilder, FRDGBufferRef &InputCountBuffer, uint32 Multiplier, uint32 Divisor = 1U, uint32 InputCountOffset = 0U)
+{
+	// 1. Add setup pass
+	FRDGBufferRef IndirectArgsBuffer = GraphBuilder.CreateBuffer(FRDGBufferDesc::CreateIndirectDesc(4), TEXT("IndirectArgs"));
+	{
+		FInitIndirectArgs1DCS::FParameters* PassParameters = GraphBuilder.AllocParameters<FInitIndirectArgs1DCS::FParameters>();
+		PassParameters->InputCountBuffer = GraphBuilder.CreateSRV(InputCountBuffer);
+		PassParameters->Multiplier = Multiplier;
+		PassParameters->Divisor = Divisor;
+		PassParameters->InputCountOffset = InputCountOffset;
+		PassParameters->IndirectDispatchArgsOut = GraphBuilder.CreateUAV(IndirectArgsBuffer, PF_R32_UINT);
+
+		auto ComputeShader = GetGlobalShaderMap(GMaxRHIFeatureLevel)->GetShader<FInitIndirectArgs1DCS>();
+		FComputeShaderUtils::AddPass(
+			GraphBuilder,
+			RDG_EVENT_NAME("InitIndirectArgs1D"),
+			ComputeShader,
+			PassParameters,
+			FIntVector(1, 1, 1)
+		);
+	}
+
+	return IndirectArgsBuffer;
+}
+
+#endif // ENABLE_NON_NANITE_VSM
+
 void FVirtualShadowMapArray::ClearPhysicalMemory(FRDGBuilder& GraphBuilder, FRDGTextureRef& PhysicalTexture, FVirtualShadowMapArrayCacheManager *VirtualShadowMapArrayCacheManager)
 {
 	check(IsEnabled());
@@ -481,6 +565,62 @@ class FInitPageRectBoundsCS : public FVirtualPageManagementShader
 IMPLEMENT_GLOBAL_SHADER(FInitPageRectBoundsCS, "/Engine/Private/VirtualShadowMaps/PageManagement.usf", "InitPageRectBounds", SF_Compute);
 
 
+#if ENABLE_NON_NANITE_VSM
+
+
+
+
+class FCalculatePageRectsCS : public FVirtualPageManagementShader
+{
+	DECLARE_GLOBAL_SHADER(FCalculatePageRectsCS);
+	SHADER_USE_PARAMETER_STRUCT(FCalculatePageRectsCS, FVirtualPageManagementShader)
+
+		BEGIN_SHADER_PARAMETER_STRUCT(FParameters, )
+		SHADER_PARAMETER_RDG_UNIFORM_BUFFER(FVirtualShadowMapCommonParameters, VirtualSmCommon)
+
+		SHADER_PARAMETER_RDG_BUFFER_SRV(StructuredBuffer< uint >, PageFlags)
+		SHADER_PARAMETER_RDG_BUFFER_UAV(RWStructuredBuffer<FIntVector4>, PageRectBoundsOut)
+		SHADER_PARAMETER_RDG_BUFFER_UAV(RWStructuredBuffer< uint2 >, OutPageTable)
+	END_SHADER_PARAMETER_STRUCT()
+};
+IMPLEMENT_GLOBAL_SHADER(FCalculatePageRectsCS, "/Engine/Private/VirtualShadowMaps/PageManagement.usf", "CalculatePageRects", SF_Compute);
+
+
+class FAllocatePagesUsingRectsCS : public FVirtualPageManagementShader
+{
+	DECLARE_GLOBAL_SHADER(FAllocatePagesUsingRectsCS);
+	SHADER_USE_PARAMETER_STRUCT(FAllocatePagesUsingRectsCS, FVirtualPageManagementShader)
+
+	//class FHasCacheDataDim : SHADER_PERMUTATION_BOOL("HAS_CACHE_DATA");
+	//class FGenerateStatsDim : SHADER_PERMUTATION_BOOL("VSM_GENERATE_STATS");
+
+	//using FPermutationDomain = TShaderPermutationDomain<FHasCacheDataDim, FGenerateStatsDim>;
+
+	static void ModifyCompilationEnvironment(const FGlobalShaderPermutationParameters& Parameters, FShaderCompilerEnvironment& OutEnvironment)
+	{
+		FVirtualPageManagementShader::ModifyCompilationEnvironment(Parameters, OutEnvironment);
+
+		// Ensure these are off for now.
+		OutEnvironment.SetDefine(TEXT("HAS_CACHE_DATA"), 0);
+		OutEnvironment.SetDefine(TEXT("VSM_GENERATE_STATS"), 0);
+	}
+
+	BEGIN_SHADER_PARAMETER_STRUCT(FParameters, )
+		SHADER_PARAMETER_RDG_UNIFORM_BUFFER(FVirtualShadowMapCommonParameters, VirtualSmCommon)
+		SHADER_PARAMETER_RDG_BUFFER_UAV(RWStructuredBuffer< uint >, AllocatedPagesOffset)
+		SHADER_PARAMETER_RDG_BUFFER_UAV(RWStructuredBuffer< uint2 >, OutPageTable)
+		SHADER_PARAMETER_RDG_BUFFER_UAV(RWStructuredBuffer< uint >, OutStatsBuffer)
+		SHADER_PARAMETER_RDG_BUFFER_UAV(RWStructuredBuffer< FCachedPageInfo >, OutCachedPageInfos)
+		SHADER_PARAMETER_RDG_BUFFER_UAV(RWStructuredBuffer< uint >, OutPageFlags)
+		SHADER_PARAMETER_RDG_BUFFER_SRV(StructuredBuffer< FIntVector4 >, PageRectBounds)
+		SHADER_PARAMETER_RDG_BUFFER_UAV(RWStructuredBuffer<FIntVector4>, AllocatedPageRectBoundsOut)
+
+		//SHADER_PARAMETER_STRUCT_INCLUDE(FCacheDataParameters, CacheDataParameters)
+	END_SHADER_PARAMETER_STRUCT()
+};
+IMPLEMENT_GLOBAL_SHADER(FAllocatePagesUsingRectsCS, "/Engine/Private/VirtualShadowMaps/PageManagement.usf", "AllocatePagesUsingRects", SF_Compute);
+#endif // ENABLE_NON_NANITE_VSM
+
 void FVirtualShadowMapArray::BuildPageAllocations(
 	FRDGBuilder& GraphBuilder,
 	const FMinimalSceneTextures& SceneTextures,
@@ -494,8 +634,6 @@ void FVirtualShadowMapArray::BuildPageAllocations(
 	check(IsEnabled());
 	RDG_EVENT_SCOPE(GraphBuilder, "FVirtualShadowMapArray::GeneratePageFlagsFromLightGrid");
 
-	ensure(NaniteRasterResults.Num() == Views.Num());
-
 	// Scale the projected footprint by the inverse scale factor such that 2x -> double the res.
 	const float LodFootprintScale = 1.0f / CVarResolutionLodScale.GetValueOnRenderThread();
 
@@ -559,10 +697,10 @@ void FVirtualShadowMapArray::BuildPageAllocations(
 
 					if( ClipmapLevel == 0 )
 					{
-						DirectionalLightSmInds.Add(ID);
-					}
+					DirectionalLightSmInds.Add(ID);
 				}
 			}
+			}
 
 			for (FProjectedShadowInfo* ProjectedShadowInfo : VisibleLightInfo.AllProjectedShadows)
 			{
@@ -592,19 +730,21 @@ void FVirtualShadowMapArray::BuildPageAllocations(
 							FMatrix TranslatedWorldToView = ProjectedShadowInfo->OnePassShadowViewMatrices[ CubeFace ] * FScaleMatrix(FVector(1, -1, 1));	// TODO add cull direction to FPackedView instead of this scale nonsense.
 	
 							FVirtualShadowMapProjectionShaderData& Data = ShadowMapProjectionData[ ProjectedShadowInfo->VirtualShadowMaps[ CubeFace ]->ID ];
+							Data.TranslatedWorldToShadowViewMatrix = TranslatedWorldToView;
 							Data.ShadowViewToClipMatrix = ViewToClip;
 							Data.TranslatedWorldToShadowUVMatrix = CalcTranslatedWorldToShadowUVMatrix(TranslatedWorldToView, ViewToClip);
 							Data.TranslatedWorldToShadowUVNormalMatrix = CalcTranslatedWorldToShadowUVNormalMatrix(TranslatedWorldToView, ViewToClip);
 							Data.ShadowPreViewTranslation = FVector(ProjectedShadowInfo->PreShadowTranslation);
 							Data.VirtualShadowMapId = ProjectedShadowInfo->VirtualShadowMaps[ CubeFace ]->ID;
 							Data.LightType = ProjectedShadowInfo->GetLightSceneInfo().Proxy->GetLightType();
-						}
-					}
+				}
+			}
 					else
 					{
 						FMatrix ViewToClip = ProjectedShadowInfo->ViewToClipInner;
 	
 						FVirtualShadowMapProjectionShaderData& Data = ShadowMapProjectionData[ ProjectedShadowInfo->VirtualShadowMaps[0]->ID ];
+						Data.TranslatedWorldToShadowViewMatrix = ProjectedShadowInfo->TranslatedWorldToView;
 						Data.ShadowViewToClipMatrix = ViewToClip;
 						Data.TranslatedWorldToShadowUVMatrix = CalcTranslatedWorldToShadowUVMatrix(ProjectedShadowInfo->TranslatedWorldToView, ViewToClip);
 						Data.TranslatedWorldToShadowUVNormalMatrix = CalcTranslatedWorldToShadowUVNormalMatrix(ProjectedShadowInfo->TranslatedWorldToView, ViewToClip);
@@ -613,7 +753,7 @@ void FVirtualShadowMapArray::BuildPageAllocations(
 						Data.LightType = ProjectedShadowInfo->GetLightSceneInfo().Proxy->GetLightType();
 
 						ensure(!ProjectedShadowInfo->bDirectionalLight);
-					}
+		}
 				}
 			}
 		}
@@ -622,7 +762,7 @@ void FVirtualShadowMapArray::BuildPageAllocations(
 		for (int32 ViewIndex = 0; ViewIndex < Views.Num(); ++ViewIndex)
 		{
 			const FViewInfo &View = Views[ViewIndex];
-			const Nanite::FRasterResults &NaniteRasterResult = NaniteRasterResults[ViewIndex];
+			FRDGTextureRef NaniteVisBuffer64 = ViewIndex < NaniteRasterResults.Num() ? NaniteRasterResults[ViewIndex].VisBuffer64 : nullptr;
 
 			// This view contained no local lights (that were stored in the light grid), and no directional lights, so nothing to do.
 			if (View.ForwardLightingResources->LocalLightVisibleLightInfosIndex.Num() + DirectionalLightSmInds.Num() == 0)
@@ -658,21 +798,21 @@ void FVirtualShadowMapArray::BuildPageAllocations(
 
 			const FRDGSystemTextures& SystemTextures = FRDGSystemTextures::Get(GraphBuilder);
 			FRDGBufferRef VirtualShadowMapIdRemapRDG = CreateStructuredBuffer(GraphBuilder, TEXT("VirtualShadowMapIdRemap"), VirtualShadowMapIdRemap);
-			FRDGTextureRef VisBuffer64 = NaniteRasterResult.VisBuffer64 ? NaniteRasterResult.VisBuffer64 : SystemTextures.Black;
+			FRDGTextureRef VisBuffer64 = NaniteVisBuffer64 != nullptr ? NaniteVisBuffer64 : SystemTextures.Black;
 
 			FRDGBufferRef ScreenSpaceGridBoundsRDG = nullptr;
 			
 			// Project Pixels onto SMs
 			{
 				FGeneratePageFlagsFromPixelsCS::FPermutationDomain PermutationVector;
-				PermutationVector.Set<FGeneratePageFlagsFromPixelsCS::FNaniteDepthBufferDim>(NaniteRasterResult.VisBuffer64 != nullptr && !bPostBasePass);			
+				PermutationVector.Set<FGeneratePageFlagsFromPixelsCS::FNaniteDepthBufferDim>(NaniteVisBuffer64 != nullptr && !bPostBasePass);
 
 				FGeneratePageFlagsFromPixelsCS::FParameters* PassParameters = GraphBuilder.AllocParameters< FGeneratePageFlagsFromPixelsCS::FParameters >();
 				PassParameters->VirtualSmCommon = GetCommonUniformBuffer(GraphBuilder);
 
 				PassParameters->SceneTexturesStruct = SceneTextures.UniformBuffer;
 				PassParameters->bPostBasePass = bPostBasePass;
-								
+				
 				PassParameters->VisBuffer64 = VisBuffer64;
 				PassParameters->View = View.ViewUniformBuffer;
 				PassParameters->OutPageRequestFlags = GraphBuilder.CreateUAV(PageRequestFlagsRDG);
@@ -714,7 +854,11 @@ void FVirtualShadowMapArray::BuildPageAllocations(
 		// Enough space for all physical pages that might be allocated
 		CachedPageInfosRDG = GraphBuilder.CreateBuffer(FRDGBufferDesc::CreateStructuredDesc(sizeof(FCachedPageInfo), CommonParameters.MaxPhysicalPages), TEXT("CachedPageInfos"));
 		PhysicalPageMetaDataRDG = GraphBuilder.CreateBuffer(FRDGBufferDesc::CreateStructuredDesc(sizeof(FPhysicalPageMetaData), CommonParameters.MaxPhysicalPages), TEXT("PhysicalPageMetaData"));
-		
+
+#if ENABLE_NON_NANITE_VSM
+		AllocatedPageRectBoundsRDG = GraphBuilder.CreateBuffer(FRDGBufferDesc::CreateStructuredDesc(sizeof(FIntVector4), NumPageRects), TEXT("AllocatedPageRectBounds"));
+#endif // ENABLE_NON_NANITE_VSM
+
 		{
 			FInitPhysicalPageMetaData::FParameters* PassParameters = GraphBuilder.AllocParameters< FInitPhysicalPageMetaData::FParameters >();
 			PassParameters->PhysicalPageMetaDataOut = GraphBuilder.CreateUAV(PhysicalPageMetaDataRDG);
@@ -730,6 +874,54 @@ void FVirtualShadowMapArray::BuildPageAllocations(
 			);
 		}
 		
+#if ENABLE_NON_NANITE_VSM
+
+		const bool bAllocatePageRectAtlas = CVarAllocatePagesUsingRects.GetValueOnRenderThread() != 0;
+		if (bAllocatePageRectAtlas)
+		{
+			// 1. Build page rects from requested pages
+			{
+				FCalculatePageRectsCS::FParameters* PassParameters = GraphBuilder.AllocParameters< FCalculatePageRectsCS::FParameters >();
+				PassParameters->VirtualSmCommon = GetCommonUniformBuffer(GraphBuilder);
+				PassParameters->PageFlags = GraphBuilder.CreateSRV(PageRequestFlagsRDG);
+				PassParameters->PageRectBoundsOut = GraphBuilder.CreateUAV(PageRectBoundsRDG);
+				PassParameters->OutPageTable = GraphBuilder.CreateUAV(PageTableRDG);
+				AddClearUAVPass(GraphBuilder, GraphBuilder.CreateUAV(PageFlagsRDG), 0);
+
+				auto ComputeShader = Views[0].ShaderMap->GetShader< FCalculatePageRectsCS >();
+				FComputeShaderUtils::AddPass(
+					GraphBuilder,
+					RDG_EVENT_NAME("CalculatePageRequestRects"),
+					ComputeShader,
+					PassParameters,
+					FComputeShaderUtils::GetGroupCount(FIntVector(FVirtualShadowMap::Level0DimPagesXY, FVirtualShadowMap::Level0DimPagesXY, ShadowMaps.Num()), FIntVector(FCalculatePageRectsCS::DefaultCSGroupXY, FCalculatePageRectsCS::DefaultCSGroupXY, 1U))
+				);
+			}
+
+
+			// 2. Allocate and all that
+			{
+				FAllocatePagesUsingRectsCS::FParameters* PassParameters = GraphBuilder.AllocParameters< FAllocatePagesUsingRectsCS::FParameters >();
+				PassParameters->VirtualSmCommon = GetCommonUniformBuffer(GraphBuilder);
+				PassParameters->PageRectBounds = GraphBuilder.CreateSRV(PageRectBoundsRDG);
+
+				PassParameters->OutPageTable = GraphBuilder.CreateUAV(PageTableRDG);
+				PassParameters->OutCachedPageInfos = GraphBuilder.CreateUAV(CachedPageInfosRDG);
+				PassParameters->OutPageFlags = GraphBuilder.CreateUAV(PageFlagsRDG);
+				PassParameters->AllocatedPageRectBoundsOut = GraphBuilder.CreateUAV(AllocatedPageRectBoundsRDG);
+
+				auto ComputeShader = Views[0].ShaderMap->GetShader< FAllocatePagesUsingRectsCS >();
+				FComputeShaderUtils::AddPass(
+					GraphBuilder,
+					RDG_EVENT_NAME("AllocatePagesUsingRects"),
+					ComputeShader,
+					PassParameters,
+					FIntVector(1U, 1U, 1U)
+				);
+			}
+		}
+		else
+#endif // ENABLE_NON_NANITE_VSM
 
 		{
 			// Note: does not actually need mip0 so can be trimmed down a bit
@@ -795,7 +987,13 @@ void FVirtualShadowMapArray::SetProjectionParameters(FRDGBuilder& GraphBuilder,
 {
 	OutParameters.VirtualSmCommon = GetCommonUniformBuffer(GraphBuilder);
 	OutParameters.PageTable = GraphBuilder.CreateSRV(PageTableRDG);
+#if ENABLE_NON_NANITE_VSM
+	// GPUCULL_TODO: Add flags to avoid sampling if not present?
+	OutParameters.PhysicalPagePool = PhysicalPagePoolRDG != nullptr ? PhysicalPagePoolRDG : GraphBuilder.RegisterExternalTexture(GSystemTextures.BlackDummy);
+	OutParameters.PhysicalPagePoolHw = PhysicalPagePoolHw != nullptr ? PhysicalPagePoolHw : GraphBuilder.RegisterExternalTexture(GSystemTextures.BlackDummy);
+#else //!ENABLE_NON_NANITE_VSM
 	OutParameters.PhysicalPagePool = PhysicalPagePoolRDG;
+#endif // ENABLE_NON_NANITE_VSM
 	OutParameters.VirtualShadowMapProjectionData = GraphBuilder.CreateSRV(ShadowMapProjectionDataRDG);
 }
 
@@ -809,6 +1007,7 @@ class FDebugVisualizeVirtualSmCS : public FVirtualPageManagementShader
 
 	BEGIN_SHADER_PARAMETER_STRUCT(FParameters, )
 		SHADER_PARAMETER_RDG_UNIFORM_BUFFER(FVirtualShadowMapCommonParameters, VirtualSmCommon)
+		SHADER_PARAMETER_STRUCT_INCLUDE(FVirtualShadowMapSamplingParameters, ProjectionParameters)
 		SHADER_PARAMETER(uint32, DebugTargetWidth)
 		SHADER_PARAMETER(uint32, DebugTargetHeight)
 		SHADER_PARAMETER(uint32, BorderWidth)
@@ -816,9 +1015,6 @@ class FDebugVisualizeVirtualSmCS : public FVirtualPageManagementShader
 		SHADER_PARAMETER(uint32, DebugMethod)
 		SHADER_PARAMETER_RDG_BUFFER_SRV(StructuredBuffer< uint >, PageFlags)
 		SHADER_PARAMETER_RDG_BUFFER_SRV(StructuredBuffer< uint >, HPageFlags)
-
-		SHADER_PARAMETER_RDG_TEXTURE(	Texture2D< uint >,			PhysicalPagePool )
-		SHADER_PARAMETER_RDG_BUFFER_SRV(StructuredBuffer< uint2 >,	PageTable )
 		
 		SHADER_PARAMETER_RDG_TEXTURE(	Texture2D< float >,			HZBPhysical )
 		SHADER_PARAMETER_RDG_BUFFER_SRV(StructuredBuffer< uint2 >,	HZBPageTable )
@@ -866,8 +1062,7 @@ void FVirtualShadowMapArray::RenderDebugInfo(FRDGBuilder& GraphBuilder, FVirtual
 		PassParameters->PageFlags			= GraphBuilder.CreateSRV( PageFlagsRDG );
 		PassParameters->HPageFlags			= GraphBuilder.CreateSRV( HPageFlagsRDG );
 
-		PassParameters->PhysicalPagePool	= PhysicalPagePoolRDG ? PhysicalPagePoolRDG : GraphBuilder.RegisterExternalTexture(GSystemTextures.BlackDummy);
-		PassParameters->PageTable			= GraphBuilder.CreateSRV( PageTableRDG );
+		SetProjectionParameters(GraphBuilder, PassParameters->ProjectionParameters);
 
 		PassParameters->HZBPhysical			= RegisterExternalTextureWithFallback( GraphBuilder, HZBPhysical, GSystemTextures.BlackDummy);
 		PassParameters->HZBPageTable		= GraphBuilder.CreateSRV( HZBPageTable ? GraphBuilder.RegisterExternalBuffer(HZBPageTable) : PageTableRDG);
@@ -911,6 +1106,9 @@ class FVirtualSmPrintStatsCS : public FVirtualPageManagementShader
 		SHADER_PARAMETER_RDG_UNIFORM_BUFFER(FVirtualShadowMapCommonParameters, VirtualSmCommon)
 		SHADER_PARAMETER_STRUCT_INCLUDE( ShaderPrint::FShaderParameters, ShaderPrintStruct )
 		SHADER_PARAMETER_RDG_BUFFER_SRV(StructuredBuffer< uint >, InStatsBuffer)
+#if ENABLE_NON_NANITE_VSM
+		SHADER_PARAMETER_RDG_BUFFER_SRV(StructuredBuffer<FIntVector4>, AllocatedPageRectBounds)
+#endif // ENABLE_NON_NANITE_VSM
 	END_SHADER_PARAMETER_STRUCT()
 };
 IMPLEMENT_GLOBAL_SHADER(FVirtualSmPrintStatsCS, "/Engine/Private/VirtualShadowMaps/PrintStats.usf", "PrintStats", SF_Compute);
@@ -942,3 +1140,466 @@ void FVirtualShadowMapArray::PrintStats(FRDGBuilder& GraphBuilder, const FViewIn
 		}
 	}
 }
+
+
+#if ENABLE_NON_NANITE_VSM
+
+
+class FVirtualSmPrintClipmapStatsCS : public FVirtualPageManagementShader
+{
+	DECLARE_GLOBAL_SHADER(FVirtualSmPrintClipmapStatsCS);
+	SHADER_USE_PARAMETER_STRUCT(FVirtualSmPrintClipmapStatsCS, FVirtualPageManagementShader)
+
+	BEGIN_SHADER_PARAMETER_STRUCT(FParameters, )
+		//SHADER_PARAMETER_RDG_UNIFORM_BUFFER(FVirtualShadowMapCommonParameters, VirtualSmCommon)
+		SHADER_PARAMETER_STRUCT_INCLUDE(ShaderPrint::FShaderParameters, ShaderPrintStruct)
+		SHADER_PARAMETER_RDG_BUFFER_SRV(StructuredBuffer< FIntVector4 >, PageRectBounds)
+		SHADER_PARAMETER_RDG_BUFFER_SRV(StructuredBuffer< FIntVector4 >, AllocatedPageRectBounds)
+		SHADER_PARAMETER(uint32, ShadowMapIdRangeStart)
+		SHADER_PARAMETER(uint32, ShadowMapIdRangeEnd)
+	END_SHADER_PARAMETER_STRUCT()
+};
+IMPLEMENT_GLOBAL_SHADER(FVirtualSmPrintClipmapStatsCS, "/Engine/Private/VirtualShadowMaps/PrintStats.usf", "PrintClipmapStats", SF_Compute);
+
+
+void FVirtualShadowMapArray::GetPageTableParameters(FRDGBuilder& GraphBuilder, FVirtualShadowMapPageTableParameters& OutParameters)
+{
+	OutParameters.VirtualSmCommon = GetCommonUniformBuffer(GraphBuilder);
+	OutParameters.PageFlags = GraphBuilder.CreateSRV(PageFlagsRDG);
+	OutParameters.HPageFlags = GraphBuilder.CreateSRV(HPageFlagsRDG);
+	OutParameters.PageRectBounds = GraphBuilder.CreateSRV(PageRectBoundsRDG);
+	OutParameters.PageTable = GraphBuilder.CreateSRV(PageTableRDG);
+}
+
+BEGIN_SHADER_PARAMETER_STRUCT(FVirtualShadowDepthPassParameters,)
+	SHADER_PARAMETER_STRUCT_REF(FViewUniformShaderParameters, View)
+	SHADER_PARAMETER_RDG_UNIFORM_BUFFER(FShadowDepthPassUniformParameters, ShadowDepthPass)
+	SHADER_PARAMETER_RDG_UNIFORM_BUFFER(FVirtualShadowMapCommonParameters, VirtualSmCommon)
+	SHADER_PARAMETER_STRUCT_INCLUDE(FInstanceCullingDrawParams, InstanceCullingDrawParams)
+
+	SHADER_PARAMETER_RDG_BUFFER_SRV(StructuredBuffer< FPackedView >, InViews)
+
+	RENDER_TARGET_BINDING_SLOTS()
+END_SHADER_PARAMETER_STRUCT()
+
+
+struct FVisibleInstanceCmd
+{
+	uint32 PackedPageInfo;
+	uint32 InstanceId;
+	uint32 DrawCommandId;
+};
+
+class FCullPerPageDrawCommandsCs : public FGlobalShader
+{
+	DECLARE_GLOBAL_SHADER(FCullPerPageDrawCommandsCs);
+	SHADER_USE_PARAMETER_STRUCT(FCullPerPageDrawCommandsCs, FGlobalShader)
+public:
+	static constexpr int32 NumThreadsPerGroup = 64;
+
+	static bool ShouldCompilePermutation(const FGlobalShaderPermutationParameters& Parameters)
+	{
+		return IsFeatureLevelSupported(Parameters.Platform, ERHIFeatureLevel::SM5);
+	}
+
+	/**
+	 */
+	static void ModifyCompilationEnvironment(const FGlobalShaderPermutationParameters& Parameters, FShaderCompilerEnvironment& OutEnvironment)
+	{
+		FGlobalShader::ModifyCompilationEnvironment(Parameters, OutEnvironment);
+		FVirtualShadowMapArray::SetShaderDefines(OutEnvironment);
+
+		OutEnvironment.SetDefine(TEXT("NUM_THREADS_PER_GROUP"), NumThreadsPerGroup);
+		OutEnvironment.SetDefine(TEXT("NANITE_MULTI_VIEW"), 1);
+		OutEnvironment.SetDefine(TEXT("INDIRECT_ARGS_NUM_WORDS"), FInstanceCullingContext::IndirectArgsNumWords);
+		OutEnvironment.SetDefine(TEXT("VF_SUPPORTS_PRIMITIVE_SCENE_DATA"), 1);
+		OutEnvironment.SetDefine(TEXT("USE_GLOBAL_GPU_SCENE_DATA"), 1);
+	}
+
+	BEGIN_SHADER_PARAMETER_STRUCT(FParameters, )
+		SHADER_PARAMETER_STRUCT_INCLUDE(FVirtualShadowMapPageTableParameters, PageTableParams)
+
+		SHADER_PARAMETER_SRV(StructuredBuffer<float4>, GPUSceneInstanceSceneData)
+		SHADER_PARAMETER_SRV(StructuredBuffer<float4>, GPUScenePrimitiveSceneData)
+		SHADER_PARAMETER(uint32, InstanceDataSOAStride)
+		SHADER_PARAMETER(uint32, GPUSceneFrameNumber)
+
+
+		SHADER_PARAMETER_RDG_BUFFER_SRV(Buffer< uint>, InstanceIdsBuffer)
+		SHADER_PARAMETER_RDG_BUFFER_SRV(Buffer< uint>, DrawCommandIdsBuffer)
+		SHADER_PARAMETER_RDG_BUFFER_SRV(StructuredBuffer< uint32 >, NumInstanceIdsBuffer)
+		SHADER_PARAMETER(int32, FirstPrimaryView)
+		SHADER_PARAMETER(int32, NumPrimaryViews)
+		SHADER_PARAMETER_RDG_BUFFER_SRV(StructuredBuffer< FPackedView >, InViews)
+
+		SHADER_PARAMETER_RDG_BUFFER_UAV(RWStructuredBuffer<FVisibleInstanceCmd>, VisibleInstancesOut)
+		SHADER_PARAMETER_RDG_BUFFER_UAV(RWStructuredBuffer<uint>, DrawCommandInstanceCountBufferOut)
+		SHADER_PARAMETER_RDG_BUFFER_UAV(RWStructuredBuffer<uint>, VisibleInstanceCountBufferOut)
+		SHADER_PARAMETER(int32, bInstancePerPage)
+
+		// Needed reference for make RDG happy somehow
+		SHADER_PARAMETER_RDG_BUFFER(Buffer<uint>, IndirectArgs)
+	END_SHADER_PARAMETER_STRUCT()
+};
+IMPLEMENT_GLOBAL_SHADER(FCullPerPageDrawCommandsCs, "/Engine/Private/VirtualShadowMaps/BuildPerPageDrawCommands.usf", "CullPerPageDrawCommandsCs", SF_Compute);
+
+
+
+class FAllocateCommandInstanceOutputSpaceCs : public FGlobalShader
+{
+	DECLARE_GLOBAL_SHADER(FAllocateCommandInstanceOutputSpaceCs);
+	SHADER_USE_PARAMETER_STRUCT(FAllocateCommandInstanceOutputSpaceCs, FGlobalShader)
+public:
+	static constexpr int32 NumThreadsPerGroup = 64;
+
+	static bool ShouldCompilePermutation(const FGlobalShaderPermutationParameters& Parameters)
+	{
+		return IsFeatureLevelSupported(Parameters.Platform, ERHIFeatureLevel::SM5);
+	}
+
+	/**
+	 */
+	static void ModifyCompilationEnvironment(const FGlobalShaderPermutationParameters& Parameters, FShaderCompilerEnvironment& OutEnvironment)
+	{
+		FGlobalShader::ModifyCompilationEnvironment(Parameters, OutEnvironment);
+		FVirtualShadowMapArray::SetShaderDefines(OutEnvironment);
+
+		OutEnvironment.SetDefine(TEXT("NUM_THREADS_PER_GROUP"), NumThreadsPerGroup);
+		OutEnvironment.SetDefine(TEXT("NANITE_MULTI_VIEW"), 1);
+		OutEnvironment.SetDefine(TEXT("INDIRECT_ARGS_NUM_WORDS"), FInstanceCullingContext::IndirectArgsNumWords);
+	}
+
+
+	BEGIN_SHADER_PARAMETER_STRUCT(FParameters, )
+		SHADER_PARAMETER(uint32, NumDrawCommands)
+		SHADER_PARAMETER_RDG_BUFFER_UAV(RWBuffer<uint>, DrawIndirectArgsBufferOut)
+		SHADER_PARAMETER_RDG_BUFFER_UAV(RWBuffer<uint>, InstanceIdOffsetBufferOut)
+		SHADER_PARAMETER_RDG_BUFFER_UAV(RWStructuredBuffer<uint>, OutputOffsetBufferOut)
+		SHADER_PARAMETER_RDG_BUFFER_UAV(RWStructuredBuffer<uint>, TmpInstanceIdOffsetBufferOut)
+		SHADER_PARAMETER_RDG_BUFFER_SRV(StructuredBuffer<uint>, DrawCommandInstanceCountBuffer)
+	END_SHADER_PARAMETER_STRUCT()
+};
+IMPLEMENT_GLOBAL_SHADER(FAllocateCommandInstanceOutputSpaceCs, "/Engine/Private/VirtualShadowMaps/BuildPerPageDrawCommands.usf", "AllocateCommandInstanceOutputSpaceCs", SF_Compute);
+
+
+class FOutputCommandInstanceListsCs : public FGlobalShader
+{
+	DECLARE_GLOBAL_SHADER(FOutputCommandInstanceListsCs);
+	SHADER_USE_PARAMETER_STRUCT(FOutputCommandInstanceListsCs, FGlobalShader)
+public:
+	static constexpr int32 NumThreadsPerGroup = 64;
+
+	static bool ShouldCompilePermutation(const FGlobalShaderPermutationParameters& Parameters)
+	{
+		return IsFeatureLevelSupported(Parameters.Platform, ERHIFeatureLevel::SM5);
+	}
+
+	/**
+	 */
+	static void ModifyCompilationEnvironment(const FGlobalShaderPermutationParameters& Parameters, FShaderCompilerEnvironment& OutEnvironment)
+	{
+		FGlobalShader::ModifyCompilationEnvironment(Parameters, OutEnvironment);
+		FVirtualShadowMapArray::SetShaderDefines(OutEnvironment);
+
+		OutEnvironment.SetDefine(TEXT("NUM_THREADS_PER_GROUP"), NumThreadsPerGroup);
+		OutEnvironment.SetDefine(TEXT("NANITE_MULTI_VIEW"), 1);
+		OutEnvironment.SetDefine(TEXT("INDIRECT_ARGS_NUM_WORDS"), FInstanceCullingContext::IndirectArgsNumWords);
+	}
+
+
+	BEGIN_SHADER_PARAMETER_STRUCT(FParameters, )
+		SHADER_PARAMETER_RDG_BUFFER_SRV(StructuredBuffer< FVisibleInstanceCmd >, VisibleInstances)
+		SHADER_PARAMETER_UAV(RWBuffer<uint>, InstanceIdsBufferLegacyOut)
+		SHADER_PARAMETER_UAV(RWBuffer<uint>, PageInfoBufferLegacyOut)
+		SHADER_PARAMETER_RDG_BUFFER_UAV(RWStructuredBuffer<uint>, TmpInstanceIdOffsetBufferOut)
+		SHADER_PARAMETER_RDG_BUFFER_SRV(StructuredBuffer<uint>, VisibleInstanceCountBuffer)
+
+		// Needed reference for make RDG happy somehow
+		SHADER_PARAMETER_RDG_BUFFER(Buffer<uint>, IndirectArgs)
+	END_SHADER_PARAMETER_STRUCT()
+};
+IMPLEMENT_GLOBAL_SHADER(FOutputCommandInstanceListsCs, "/Engine/Private/VirtualShadowMaps/BuildPerPageDrawCommands.usf", "OutputCommandInstanceListsCs", SF_Compute);
+
+void FVirtualShadowMapArray::RenderVirtualShadowMapsHw(FRDGBuilder& GraphBuilder, const TArray<FProjectedShadowInfo *, SceneRenderingAllocator>& VirtualSmMeshCommandPasses, FScene& Scene)
+{
+	RDG_EVENT_SCOPE(GraphBuilder, "RenderVirtualShadowMapsHw");
+	// Create physical page pool
+	PhysicalPagePoolHw = GraphBuilder.CreateTexture(FRDGTextureDesc::Create2D(GetPhysicalPoolSize(), PF_DepthStencil, FClearValueBinding::DepthFar, TexCreate_DepthStencilTargetable | TexCreate_ShaderResource), TEXT("PhysicalPagePoolHw"));
+
+	FGPUScene& GPUScene = Scene.GPUScene;
+	// Create that pass struct and whatever I guess.
+
+	// track whether the page pool was cleared
+	bool bWasCleared = false;
+
+	const bool bAllocatePageRectAtlas = CVarAllocatePagesUsingRects.GetValueOnRenderThread() != 0;
+
+	// TODO: create a VirtualClipMapMeshRenderInfo - or something - that contains the required info:
+	//       The View used to set up (to get all the usual crap), The MeshCommandPass (to perform setup & culling), Culling volume?, Ref to Clipmap itself.
+
+	// Loop over the mesh command passes needed (one for each SM, we need to do this on a per-SM basis mainly because different views may have different view infos)
+	// In practice there is no fundamental reason why we could not merge all the draws into one giant submission - should probably consider this at some point. 
+	// One question is whether the shaders refer to the view, and actually can get sensible data out (I'm not sure this is the case at present anyway, as it contains a bastardization of main view and shadow view info)
+	// The other one to beware of is the fact that we have these dynamic primitives, which until this is fixed will cause the creation of an entire copy of the GPU scene data for each shadow view that contains one...
+	//   Starting to become urgent!
+	for (int Index = 0; Index < VirtualSmMeshCommandPasses.Num(); ++Index)
+	{
+		FProjectedShadowInfo* ProjectedShadowInfo = VirtualSmMeshCommandPasses[Index];
+
+		FParallelMeshDrawCommandPass& MeshCommandPass = ProjectedShadowInfo->GetShadowDepthPass();
+		// TODO: Right now, this is for clipmaps only! Get hold of them somehow!
+		const TSharedPtr<FVirtualShadowMapClipmap> Clipmap = ProjectedShadowInfo->VirtualShadowMapClipmap;
+		// TODO: Also get the associated view such that all the parameters can be set.
+		const FViewInfo* ViewUsedToCreateShadow = ProjectedShadowInfo->DependentView;
+		FViewInfo* ShadowDepthView = ProjectedShadowInfo->ShadowDepthView;
+
+		const FViewInfo& View = *ViewUsedToCreateShadow;
+		TArray<Nanite::FPackedView, SceneRenderingAllocator> VirtualShadowViews;
+		{
+			Nanite::FPackedViewParams BaseParams;
+			BaseParams.ViewRect = FIntRect(0, 0, FVirtualShadowMap::VirtualMaxResolutionXY, FVirtualShadowMap::VirtualMaxResolutionXY);
+			BaseParams.RasterContextSize = GetPhysicalPoolSize();
+			BaseParams.PrevTargetLayerIndex = INDEX_NONE;
+			BaseParams.TargetMipLevel = 0;
+			BaseParams.TargetMipCount = 1;	// No mips for clipmaps
+
+			for (int32 ClipmapLevelIndex = 0; ClipmapLevelIndex < Clipmap->GetLevelCount(); ++ClipmapLevelIndex)
+			{
+				Nanite::FPackedViewParams Params = BaseParams;
+				Params.TargetLayerIndex = Clipmap->GetVirtualShadowMap(ClipmapLevelIndex)->ID;
+				Params.ViewMatrices = Clipmap->GetViewMatrices(ClipmapLevelIndex);
+				Params.PrevTargetLayerIndex = INDEX_NONE;
+				Params.PrevViewMatrices = Params.ViewMatrices;
+
+				VirtualShadowViews.Add(Nanite::CreatePackedView(Params));
+			}
+		}
+		// Created by BuildInstanceList
+		// 1. Run post-cull unpack command (?) - maybe skip straight to the second pass, or output to intermediate instance index buffer?
+		//MeshCommandPass.BuildRenderingCommands(GraphBuilder, GPUScene, InstanceListBuffer, InstanceCountBuffer);
+		// 1.1. Pass to unpack into flat drawcommand + instance ID pairs
+		FInstanceCullingRdgParams Params;
+		MeshCommandPass.BuildInstanceList(GraphBuilder, GPUScene, Params);
+
+		const int32 NumDrawCommands = MeshCommandPass.GetInstanceCullingContext()->CullingCommands.Num();
+		if (NumDrawCommands > 0)
+		{
+			TArray<int32, SceneRenderingAllocator> DrawCommandInstanceCountTmp;
+			DrawCommandInstanceCountTmp.AddZeroed(NumDrawCommands);
+			FRDGBufferRef DrawCommandInstanceCountRdg = CreateStructuredBuffer(GraphBuilder, TEXT("DrawCommandInstanceCount"), DrawCommandInstanceCountTmp);
+			FRDGBufferRef TmpInstanceIdOffsetBufferRdg = CreateStructuredBuffer(GraphBuilder, TEXT("TmpInstanceIdOffsetBufferRdg"), DrawCommandInstanceCountTmp);
+			//AddBuildInstanceListPass(GraphBuilder, GPUScene, InstanceIdBuffer, InstanceCountBuffer);
+			// 2. Run page culling & command building pass(es)
+			// This takes the per high-level view instance list and performs page overlap tests and amplification/compation for all target mip levels OR clip levels.
+			//if (IsClipMap())
+			// AddPageRenderCommandPass(GraphBuilder, GPUScene, InstanceIdBuffer, InstanceCountBuffer, ShadowMapViews);
+			// 2.1. Pass to perform the culling & replication to each page, store 
+			const uint32 MaxNumInstancesPerPass = NumDrawCommands * FInstanceCullingManager::MaxAverageInstanceFactor * 64u;
+			FRDGBufferRef VisibleInstancesRdg = CreateStructuredBuffer(GraphBuilder, TEXT("VisibleInstances"), sizeof(FVisibleInstanceCmd), MaxNumInstancesPerPass, nullptr, 0);
+
+			TArray<uint32> NullArray;
+			NullArray.AddZeroed(1);
+			FRDGBufferRef VisibleInstanceWriteOffsetRdg = CreateStructuredBuffer(GraphBuilder, TEXT("VisibleInstanceWriteOffset"), NullArray);
+			FRDGBufferRef OutputOffsetBufferRdg = CreateStructuredBuffer(GraphBuilder, TEXT("OutputOffsetBuffer"), NullArray);
+			FRDGBufferRef VirtualShadowViewsRdg = CreateStructuredBuffer(GraphBuilder, TEXT("VirtualShadowViews"), VirtualShadowViews);
+
+
+			// TODO: Remove this when everything is properly RDG'd
+			AddPass(GraphBuilder, [](FRHICommandList& RHICmdList)
+			{
+				FRHITransitionInfo 	Transitions[2] =
+				{
+					FRHITransitionInfo(GInstanceCullingManagerResources.GetInstancesIdBufferUav(), ERHIAccess::Unknown, ERHIAccess::UAVCompute),
+					FRHITransitionInfo(GInstanceCullingManagerResources.GetPageInfoBufferUav(), ERHIAccess::Unknown, ERHIAccess::UAVCompute)
+				};
+
+				RHICmdList.Transition(Transitions);
+			});
+
+			{
+				FRDGBufferRef IndirectArgs = AddIndirectArgsSetupCsPass1D(GraphBuilder, Params.InstanceIdWriteOffsetBuffer, 1, FCullPerPageDrawCommandsCs::NumThreadsPerGroup);
+
+				FCullPerPageDrawCommandsCs::FParameters* PassParameters = GraphBuilder.AllocParameters<FCullPerPageDrawCommandsCs::FParameters>();
+
+				GetPageTableParameters(GraphBuilder, PassParameters->PageTableParams);
+
+				PassParameters->GPUSceneInstanceSceneData = GPUScene.InstanceDataBuffer.SRV;
+				PassParameters->GPUScenePrimitiveSceneData = GPUScene.PrimitiveBuffer.SRV;
+				PassParameters->InstanceDataSOAStride = GPUScene.InstanceDataSOAStride;
+				PassParameters->GPUSceneFrameNumber = GPUScene.GetSceneFrameNumber();
+				PassParameters->InstanceIdsBuffer = GraphBuilder.CreateSRV(Params.InstanceIdsBuffer, PF_R32_UINT);;
+				PassParameters->DrawCommandIdsBuffer = GraphBuilder.CreateSRV(Params.DrawCommandIdsBuffer, PF_R32_UINT);
+				PassParameters->NumInstanceIdsBuffer = GraphBuilder.CreateSRV(Params.InstanceIdWriteOffsetBuffer);
+				PassParameters->FirstPrimaryView = 0;
+				PassParameters->NumPrimaryViews = VirtualShadowViews.Num();
+				PassParameters->InViews = GraphBuilder.CreateSRV(VirtualShadowViewsRdg);
+				PassParameters->VisibleInstancesOut = GraphBuilder.CreateUAV(VisibleInstancesRdg);
+				PassParameters->DrawCommandInstanceCountBufferOut = GraphBuilder.CreateUAV(DrawCommandInstanceCountRdg);
+				PassParameters->VisibleInstanceCountBufferOut = GraphBuilder.CreateUAV(VisibleInstanceWriteOffsetRdg);
+				PassParameters->IndirectArgs = IndirectArgs;
+				PassParameters->bInstancePerPage = !bAllocatePageRectAtlas;
+
+				auto ComputeShader = GetGlobalShaderMap(GMaxRHIFeatureLevel)->GetShader<FCullPerPageDrawCommandsCs>();
+
+				FComputeShaderUtils::AddPass(
+					GraphBuilder,
+					RDG_EVENT_NAME("CullPerPageDrawCommandsCs"),
+					ComputeShader,
+					PassParameters,
+					IndirectArgs,
+					0
+				);
+			}
+			// 2.2.Allocate space for the final instance ID output and so on.
+			{
+				FAllocateCommandInstanceOutputSpaceCs::FParameters* PassParameters = GraphBuilder.AllocParameters<FAllocateCommandInstanceOutputSpaceCs::FParameters>();
+
+				// get this in a neater way
+				FRDGBufferRef InstanceIdOutOffsetBufferRDG = MeshCommandPass.GetInstanceCullingContext()->InstanceCullingManager->CullingIntermediate.InstanceIdOutOffsetBuffer;
+
+				PassParameters->NumDrawCommands = NumDrawCommands;
+				PassParameters->DrawIndirectArgsBufferOut = GraphBuilder.CreateUAV(Params.DrawIndirectArgs, PF_R32_UINT);
+				PassParameters->InstanceIdOffsetBufferOut = GraphBuilder.CreateUAV(Params.InstanceIdStartOffsetBuffer, PF_R32_UINT);;
+				PassParameters->OutputOffsetBufferOut = GraphBuilder.CreateUAV(InstanceIdOutOffsetBufferRDG);
+				PassParameters->TmpInstanceIdOffsetBufferOut = GraphBuilder.CreateUAV(TmpInstanceIdOffsetBufferRdg);
+				PassParameters->DrawCommandInstanceCountBuffer = GraphBuilder.CreateSRV(DrawCommandInstanceCountRdg);
+
+				auto ComputeShader = GetGlobalShaderMap(GMaxRHIFeatureLevel)->GetShader<FAllocateCommandInstanceOutputSpaceCs>();
+
+				FComputeShaderUtils::AddPass(
+					GraphBuilder,
+					RDG_EVENT_NAME("AllocateCommandInstanceOutputSpaceCs"),
+					ComputeShader,
+					PassParameters,
+					FComputeShaderUtils::GetGroupCount(NumDrawCommands, FAllocateCommandInstanceOutputSpaceCs::NumThreadsPerGroup)
+				);
+
+			}
+			// 2.3. Perform final pass to re-shuffle the instance ID's to their final resting places
+			{
+				FRDGBufferRef IndirectArgs = AddIndirectArgsSetupCsPass1D(GraphBuilder, VisibleInstanceWriteOffsetRdg, 1, FOutputCommandInstanceListsCs::NumThreadsPerGroup);
+
+				FOutputCommandInstanceListsCs::FParameters* PassParameters = GraphBuilder.AllocParameters<FOutputCommandInstanceListsCs::FParameters>();
+
+				PassParameters->VisibleInstances = GraphBuilder.CreateSRV(VisibleInstancesRdg);
+				PassParameters->PageInfoBufferLegacyOut = GInstanceCullingManagerResources.GetPageInfoBufferUav();
+				PassParameters->InstanceIdsBufferLegacyOut = GInstanceCullingManagerResources.GetInstancesIdBufferUav();
+				PassParameters->TmpInstanceIdOffsetBufferOut = GraphBuilder.CreateUAV(TmpInstanceIdOffsetBufferRdg);
+				PassParameters->VisibleInstanceCountBuffer = GraphBuilder.CreateSRV(VisibleInstanceWriteOffsetRdg);
+				PassParameters->IndirectArgs = IndirectArgs;
+
+				auto ComputeShader = GetGlobalShaderMap(GMaxRHIFeatureLevel)->GetShader<FOutputCommandInstanceListsCs>();
+
+				FComputeShaderUtils::AddPass(
+					GraphBuilder,
+					RDG_EVENT_NAME("OutputCommandInstanceListsCs"),
+					ComputeShader,
+					PassParameters,
+					IndirectArgs,
+					0
+				);
+			}
+
+			ProjectedShadowInfo->BeginRenderView(GraphBuilder, &Scene);
+
+			FVirtualShadowDepthPassParameters* PassParameters = GraphBuilder.AllocParameters<FVirtualShadowDepthPassParameters>();
+			PassParameters->View = ShadowDepthView->ViewUniformBuffer;
+			PassParameters->RenderTargets.DepthStencil = FDepthStencilBinding(PhysicalPagePoolHw, bWasCleared ? ERenderTargetLoadAction::ELoad : ERenderTargetLoadAction::EClear, FExclusiveDepthStencil::DepthWrite_StencilNop);
+			bWasCleared = true;
+
+			FShadowDepthPassUniformParameters* ShadowDepthPassParameters = GraphBuilder.AllocParameters<FShadowDepthPassUniformParameters>();
+
+			SetupSceneTextureUniformParameters(GraphBuilder, GMaxRHIFeatureLevel, ESceneTextureSetupMode::None, ShadowDepthPassParameters->SceneTextures);
+
+			// TODO: Only valid for clipmap
+			ShadowDepthPassParameters->bClampToNearPlane = true;
+
+			// TODO: These are not used for this case anyway
+			ShadowDepthPassParameters->ProjectionMatrix = FMatrix::Identity;
+			ShadowDepthPassParameters->ViewMatrix = FMatrix::Identity;
+			ShadowDepthPassParameters->ShadowParams = FVector4(0.0f, 0.0f, 0.0f, 1.0f);
+			ShadowDepthPassParameters->bRenderToVirtualShadowMap = true;
+			ShadowDepthPassParameters->VirtualSmPageTable = GraphBuilder.CreateSRV(PageTableRDG);
+			ShadowDepthPassParameters->PackedNaniteViews = GraphBuilder.CreateSRV(VirtualShadowViewsRdg);
+			ShadowDepthPassParameters->PageRectBounds = GraphBuilder.CreateSRV(PageRectBoundsRDG);
+			ShadowDepthPassParameters->bInstancePerPage = !bAllocatePageRectAtlas;
+			
+			PassParameters->ShadowDepthPass = GraphBuilder.CreateUniformBuffer(ShadowDepthPassParameters);;
+			PassParameters->VirtualSmCommon = GetCommonUniformBuffer(GraphBuilder);
+			PassParameters->InViews = GraphBuilder.CreateSRV(VirtualShadowViewsRdg);
+
+			FInstanceCullingResult InstanceCullingResult;
+			InstanceCullingResult.DrawIndirectArgsBuffer = Params.DrawIndirectArgs;
+			InstanceCullingResult.InstanceIdOffsetBuffer = Params.InstanceIdStartOffsetBuffer;
+			InstanceCullingResult.GetDrawParameters(PassParameters->InstanceCullingDrawParams);
+
+			FMeshPassProcessorRenderState DrawRenderState;
+			DrawRenderState.SetBlendState(TStaticBlendState<CW_NONE>::GetRHI());
+			DrawRenderState.SetDepthStencilState(TStaticDepthStencilState<true, CF_DepthNearOrEqual>::GetRHI());
+
+
+			// TODO: Remove this when everything is properly RDG'd
+			AddPass(GraphBuilder, [](FRHICommandList& RHICmdList)
+			{
+				FRHITransitionInfo 	Transitions[2] =
+				{
+					FRHITransitionInfo(GInstanceCullingManagerResources.GetInstancesIdBufferUav(), ERHIAccess::UAVCompute, ERHIAccess::SRVGraphics),
+					FRHITransitionInfo(GInstanceCullingManagerResources.GetPageInfoBufferUav(), ERHIAccess::UAVCompute, ERHIAccess::SRVGraphics)
+				};
+
+				RHICmdList.Transition(Transitions);
+			});
+
+			GraphBuilder.AddPass(
+				RDG_EVENT_NAME("RenderVirtualShadowMapsHw"),
+				PassParameters,
+				ERDGPassFlags::Raster,
+				[this, ProjectedShadowInfo , &MeshCommandPass, PassParameters, ShadowDepthPassParameters](FRHICommandList& RHICmdList)
+				{
+					RHICmdList.SetViewport(
+						0.0f,
+						0.0f,
+						0.0f,
+						GetPhysicalPoolSize().X,
+						GetPhysicalPoolSize().Y,
+						1.0f
+					);
+					MeshCommandPass.DispatchDraw(nullptr, RHICmdList, &PassParameters->InstanceCullingDrawParams);
+				});
+
+
+		}
+
+
+		//
+		if (Index == CVarShowClipmapStats.GetValueOnRenderThread())
+		{
+			FVirtualSmPrintClipmapStatsCS::FParameters* PassParameters = GraphBuilder.AllocParameters<FVirtualSmPrintClipmapStatsCS::FParameters>();
+
+			ShaderPrint::SetParameters(GraphBuilder, View, PassParameters->ShaderPrintStruct);
+			//PassParameters->VirtualSmCommon = GetCommonUniformBuffer(GraphBuilder);
+			PassParameters->ShadowMapIdRangeStart = Clipmap->GetVirtualShadowMap(0)->ID;
+			// Note: assumes range!
+			PassParameters->ShadowMapIdRangeEnd = Clipmap->GetVirtualShadowMap(0)->ID + Clipmap->GetLevelCount();
+			PassParameters->PageRectBounds = GraphBuilder.CreateSRV(PageRectBoundsRDG);
+			PassParameters->AllocatedPageRectBounds = GraphBuilder.CreateSRV(AllocatedPageRectBoundsRDG);
+
+			auto ComputeShader = View.ShaderMap->GetShader<FVirtualSmPrintClipmapStatsCS>();
+
+			FComputeShaderUtils::AddPass(
+				GraphBuilder,
+				RDG_EVENT_NAME("PrintClipmapStats"),
+				ComputeShader,
+				PassParameters,
+				FIntVector(1, 1, 1)
+			);
+		}
+	}
+
+	if (!bWasCleared)
+	{
+		AddClearDepthStencilPass(GraphBuilder, PhysicalPagePoolHw);
+	}
+}
+
+#endif // ENABLE_NON_NANITE_VSM
diff --git a/Engine/Source/Runtime/Renderer/Private/VirtualShadowMaps/VirtualShadowMapArray.h b/Engine/Source/Runtime/Renderer/Private/VirtualShadowMaps/VirtualShadowMapArray.h
--- a/Engine/Source/Runtime/Renderer/Private/VirtualShadowMaps/VirtualShadowMapArray.h
+++ b/Engine/Source/Runtime/Renderer/Private/VirtualShadowMaps/VirtualShadowMapArray.h
@@ -4,7 +4,9 @@
 =============================================================================*/
 #pragma once
 
+#include "VirtualShadowMapConfig.h"
 #include "../Nanite/NaniteRender.h"
+#include "../MeshDrawCommands.h"
 #include "SceneTypes.h"
 
 struct FMinimalSceneTextures;
@@ -44,9 +46,19 @@ public:
 	static constexpr uint32 PageSize = 128U;
 	static constexpr uint32 Level0DimPagesXY = 128U;
 
+	// 512x32 = 16k
+	//static constexpr uint32 PageSize = 512U;
+	//static constexpr uint32 Level0DimPagesXY = 32U;
+
+
 	// With 128x128 pages, a 8k x 4k texture holds 2048 physical pages
 	static constexpr uint32 PhysicalPagePoolTexureSizeX = 8192U;
+#if ENABLE_NON_NANITE_VSM
+	// GPUCULL_TODO: Think we must have square RT for address translation to work, the way it is implemented now (badly)
+	static constexpr uint32 PhysicalPagePoolTexureSizeY = 8192U;
+#else //!ENABLE_NON_NANITE_VSM
 	static constexpr uint32 PhysicalPagePoolTexureSizeY = 4096U;
+#endif // ENABLE_NON_NANITE_VSM
 
 	static constexpr uint32 PageSizeMask = PageSize - 1U;
 	static constexpr uint32 Log2PageSize = ILog2Const(PageSize);
@@ -79,6 +91,11 @@ public:
 // as well as cached shadow maps
 struct FVirtualShadowMapProjectionShaderData
 {
+	/**
+	 * Transform from shadow-pre-translated world space to shadow view space, example use: (WorldSpacePos + ShadowPreViewTranslation) * TranslatedWorldToShadowViewMatrix
+	 * TODO: Why don't we call it a rotation and store in a 3x3? Does it ever have translation in?
+	 */
+	FMatrix TranslatedWorldToShadowViewMatrix;
 	FMatrix ShadowViewToClipMatrix;
 	FMatrix TranslatedWorldToShadowUVMatrix;
 	FMatrix TranslatedWorldToShadowUVNormalMatrix;
@@ -113,6 +130,8 @@ BEGIN_GLOBAL_SHADER_PARAMETER_STRUCT(FVirtualShadowMapCommonParameters, )
 	// use to map linear index to x,y page coord
 	SHADER_PARAMETER(uint32, PhysicalPageRowMask)
 	SHADER_PARAMETER(uint32, PhysicalPageRowShift)
+	SHADER_PARAMETER(FVector4, RecPhysicalPoolSize)
+	SHADER_PARAMETER(FIntPoint, PhysicalPoolSizePages)
 END_GLOBAL_SHADER_PARAMETER_STRUCT()
 
 BEGIN_SHADER_PARAMETER_STRUCT(FVirtualShadowMapSamplingParameters, )
@@ -120,8 +139,23 @@ BEGIN_SHADER_PARAMETER_STRUCT(FVirtualShadowMapSamplingParameters, )
 	SHADER_PARAMETER_RDG_BUFFER_SRV(StructuredBuffer<uint2>, PageTable)
 	SHADER_PARAMETER_RDG_TEXTURE(Texture2D<uint>, PhysicalPagePool)
 	SHADER_PARAMETER_RDG_BUFFER_SRV(StructuredBuffer< FVirtualShadowMapProjectionShaderData >, VirtualShadowMapProjectionData)
+#if ENABLE_NON_NANITE_VSM
+	SHADER_PARAMETER_RDG_TEXTURE(Texture2D<float>, PhysicalPagePoolHw)
+#endif // ENABLE_NON_NANITE_VSM
 END_SHADER_PARAMETER_STRUCT()
 
+/**
+ * Use after page allocation but before rendering phase to access page table & related data structures, but not the physical backing.
+ */
+BEGIN_SHADER_PARAMETER_STRUCT(FVirtualShadowMapPageTableParameters, )
+	SHADER_PARAMETER_RDG_UNIFORM_BUFFER(FVirtualShadowMapCommonParameters, VirtualSmCommon)
+	SHADER_PARAMETER_RDG_BUFFER_SRV(StructuredBuffer<uint2>, PageTable)
+	SHADER_PARAMETER_RDG_BUFFER_SRV(StructuredBuffer< uint >, PageFlags)
+	SHADER_PARAMETER_RDG_BUFFER_SRV(StructuredBuffer< uint >, HPageFlags)
+	SHADER_PARAMETER_RDG_BUFFER_SRV(StructuredBuffer< uint4 >, PageRectBounds)
+END_SHADER_PARAMETER_STRUCT()
+
+
 class FVirtualShadowMapArray
 {
 public:
@@ -175,6 +209,13 @@ public:
 		return PhysicalPagePoolRDG != nullptr && PageTableRDG != nullptr;
 	}
 
+#if ENABLE_NON_NANITE_VSM
+	/**
+	 * Draw old-school hardware based shadow map tiles into virtual SM.
+	 */
+	void RenderVirtualShadowMapsHw(FRDGBuilder& GraphBuilder, const TArray<FProjectedShadowInfo*, SceneRenderingAllocator>& VirtualSmMeshCommandPasses, FScene& Scene);
+#endif // ENABLE_NON_NANITE_VSM
+
 	// Draw debug info into render target 'VirtSmDebug' of screen-size, the mode is controlled by 'r.Shadow.v.DebugVisualize' (defaults to not doing aught). 
 	void RenderDebugInfo(FRDGBuilder& GraphBuilder, FVirtualShadowMapArrayCacheManager *VirtualShadowMapArrayCacheManager);
 	// 
@@ -188,13 +229,21 @@ public:
 		return GraphBuilder.CreateUniformBuffer(&CommonParameters);
 	}
 
+#if ENABLE_NON_NANITE_VSM
+	bool HasAnyShadowData() const { return PhysicalPagePoolRDG != nullptr || PhysicalPagePoolHw != nullptr;  }
+#else //!ENABLE_NON_NANITE_VSM
+	bool HasAnyShadowData() const { return PhysicalPagePoolRDG != nullptr;  }
+#endif // ENABLE_NON_NANITE_VSM
+
+	void GetPageTableParameters(FRDGBuilder& GraphBuilder, FVirtualShadowMapPageTableParameters& OutParameters);
+
 	TArray<FVirtualShadowMap*, SceneRenderingAllocator> ShadowMaps;
 
 	// Large physical texture of depth format, say 4096^2 or whatever we think is enough texels to go around
 	FRDGTextureRef PhysicalPagePoolRDG = nullptr;
 	// Buffer that serves as the page table for all virtual shadow maps
 	FRDGBufferRef PageTableRDG = nullptr;
-
+		
 	// Buffer that stores flags (uints) marking each page that needs to be rendered and cache status, for all virtual shadow maps.
 	// Flag values defined in PageAccessCommon.ush: VSM_ALLOCATED_FLAG | VSM_INVALID_FLAG
 	FRDGBufferRef PageFlagsRDG = nullptr;
@@ -218,11 +267,13 @@ public:
 
 	// Buffer that stores flags marking each page that received dynamic geo.
 	FRDGBufferRef DynamicCasterPageFlagsRDG = nullptr;
-
+	
 	// uint4 buffer with one rect for each mip level in all SMs, calculated to bound committed pages
 	// Used to clip the rect size of clusters during culling.
 	FRDGBufferRef PageRectBoundsRDG = nullptr;
-
+#if ENABLE_NON_NANITE_VSM
+	FRDGBufferRef AllocatedPageRectBoundsRDG = nullptr;
+#endif // ENABLE_NON_NANITE_VSM
 	FRDGBufferRef ShadowMapProjectionDataRDG = nullptr;
 
 	TRefCountPtr<IPooledRenderTarget>	HZBPhysical;
@@ -235,4 +286,9 @@ public:
 	//FRDGBufferRef RDG = nullptr;
 	TRefCountPtr<IPooledRenderTarget>	DebugVisualizationProjectionOutput;
 	//FRDGBufferRef RDG = nullptr;
+#if ENABLE_NON_NANITE_VSM
+	// page pool for HW rasterized shadow data. 
+	// Mirrors the regular one, but uses a shadow depth target format
+	FRDGTextureRef PhysicalPagePoolHw = nullptr;
+#endif // ENABLE_NON_NANITE_VSM
 };
diff --git a/Engine/Source/Runtime/Renderer/Private/VirtualShadowMaps/VirtualShadowMapCacheManager.h b/Engine/Source/Runtime/Renderer/Private/VirtualShadowMaps/VirtualShadowMapCacheManager.h
--- a/Engine/Source/Runtime/Renderer/Private/VirtualShadowMaps/VirtualShadowMapCacheManager.h
+++ b/Engine/Source/Runtime/Renderer/Private/VirtualShadowMaps/VirtualShadowMapCacheManager.h
@@ -5,6 +5,7 @@
 =============================================================================*/
 #pragma once
 
+#include "VirtualShadowMapConfig.h"
 #include "CoreMinimal.h"
 #include "VirtualShadowMapArray.h"
 #include "SceneManagement.h"
diff --git a/Engine/Source/Runtime/Renderer/Private/VirtualShadowMaps/VirtualShadowMapClipmap.cpp b/Engine/Source/Runtime/Renderer/Private/VirtualShadowMaps/VirtualShadowMapClipmap.cpp
--- a/Engine/Source/Runtime/Renderer/Private/VirtualShadowMaps/VirtualShadowMapClipmap.cpp
+++ b/Engine/Source/Runtime/Renderer/Private/VirtualShadowMaps/VirtualShadowMapClipmap.cpp
@@ -39,8 +39,10 @@ FVirtualShadowMapClipmap::FVirtualShadowMapClipmap(
 	const FMatrix& WorldToLightRotationMatrix,
 	const FViewMatrices& CameraViewMatrices,
 	FIntPoint CameraViewRectSize,
-	float MaxRadius)
-	: LightSceneInfo(InLightSceneInfo)
+
+	float InMaxRadius)
+	: LightSceneInfo(InLightSceneInfo), 
+	  MaxRadius(InMaxRadius)
 {
 	check(WorldToLightRotationMatrix.GetOrigin() == FVector(0, 0, 0));	// Should not contain translation or scaling
 
@@ -155,6 +157,7 @@ FVirtualShadowMapProjectionShaderData FVirtualShadowMapClipmap::GetProjectionSha
 	
 	// NOTE: Some shader logic (projection, etc) assumes some of these parameters are constant across all levels in a clipmap
 	FVirtualShadowMapProjectionShaderData Data;
+	Data.TranslatedWorldToShadowViewMatrix = WorldToViewRotationMatrix;
 	Data.ShadowViewToClipMatrix = Level.ViewToClip;
 	Data.TranslatedWorldToShadowUVMatrix = CalcTranslatedWorldToShadowUVMatrix(WorldToViewRotationMatrix, Level.ViewToClip);
 	Data.TranslatedWorldToShadowUVNormalMatrix = CalcTranslatedWorldToShadowUVNormalMatrix(WorldToViewRotationMatrix, Level.ViewToClip);
diff --git a/Engine/Source/Runtime/Renderer/Private/VirtualShadowMaps/VirtualShadowMapClipmap.h b/Engine/Source/Runtime/Renderer/Private/VirtualShadowMaps/VirtualShadowMapClipmap.h
--- a/Engine/Source/Runtime/Renderer/Private/VirtualShadowMaps/VirtualShadowMapClipmap.h
+++ b/Engine/Source/Runtime/Renderer/Private/VirtualShadowMaps/VirtualShadowMapClipmap.h
@@ -6,6 +6,7 @@ VirtualShadowMapClipmap.h
 
 #pragma once
 
+#include "VirtualShadowMapConfig.h"
 #include "CoreMinimal.h"
 #include "VirtualShadowMapProjection.h"
 
@@ -56,6 +57,7 @@ public:
 		return WorldOrigin;
 	}
 
+	float GetMaxRadius() const { return MaxRadius; }
 
 private:
 	const FLightSceneInfo& LightSceneInfo;
@@ -71,6 +73,7 @@ private:
 
 	int32 FirstLevel;
 	float ResolutionLodBias;
+	float MaxRadius;
 
 	struct FLevelData
 	{
diff --git a/Engine/Source/Runtime/Renderer/Private/VirtualShadowMaps/VirtualShadowMapConfig.h b/Engine/Source/Runtime/Renderer/Private/VirtualShadowMaps/VirtualShadowMapConfig.h
new file mode 100644
--- /dev/null
+++ b/Engine/Source/Runtime/Renderer/Private/VirtualShadowMaps/VirtualShadowMapConfig.h
@@ -0,0 +1,13 @@
+// Copyright Epic Games, Inc. All Rights Reserved.
+
+/*=============================================================================
+	VirtualShadowMapConfig.h
+=============================================================================*/
+
+#pragma once
+
+#if defined(GPUCULL_TODO)
+	#define ENABLE_NON_NANITE_VSM 1
+#else // !defined(GPUCULL_TODO)
+	#define ENABLE_NON_NANITE_VSM 0
+#endif // defined(GPUCULL_TODO)
\ No newline at end of file
diff --git a/Engine/Source/Runtime/Renderer/Private/VirtualShadowMaps/VirtualShadowMapProjection.h b/Engine/Source/Runtime/Renderer/Private/VirtualShadowMaps/VirtualShadowMapProjection.h
--- a/Engine/Source/Runtime/Renderer/Private/VirtualShadowMaps/VirtualShadowMapProjection.h
+++ b/Engine/Source/Runtime/Renderer/Private/VirtualShadowMaps/VirtualShadowMapProjection.h
@@ -6,6 +6,7 @@
 
 #pragma once
 
+#include "VirtualShadowMapConfig.h"
 #include "CoreMinimal.h"
 #include "RenderGraphResources.h"
 #include "ScreenSpaceDenoise.h"
