Fix missing/incorrect lighting on hair cards.

There was two separted issue:
* When hair strands was disabled, the hair LUT data was never updated, left unitilialized (black).
* When the hair base color was near black, a precision issue could arise due to a ratio between two small number causing NaN/Inf value.

diff --git a/Engine/Shaders/Private/DeferredLightPixelShaders.usf b/Engine/Shaders/Private/DeferredLightPixelShaders.usf
--- a/Engine/Shaders/Private/DeferredLightPixelShaders.usf
+++ b/Engine/Shaders/Private/DeferredLightPixelShaders.usf
@@ -184,6 +184,8 @@ void DeferredLightPixelMain(
 		FDeferredLightData LightData = SetupLightDataForStandardDeferred();
 		if (ScreenSpaceData.GBuffer.ShadingModelID == SHADINGMODELID_HAIR && ShouldUseHairComplexTransmittance(ScreenSpaceData.GBuffer))
 		{
+			// Prevent numerical issue causing NaN during hair lighting (i.e., during baseColor/luma(baseColor) computation)
+			ScreenSpaceData.GBuffer.BaseColor = max(0.0001f, ScreenSpaceData.GBuffer.BaseColor);
 			LightData.HairTransmittance = EvaluateDualScattering(ScreenSpaceData.GBuffer, DerivedParams.CameraVector, -DeferredLightUniforms.Direction);
 		}
 
diff --git a/Engine/Source/Runtime/Renderer/Private/HairStrands/HairStrandsLUT.cpp b/Engine/Source/Runtime/Renderer/Private/HairStrands/HairStrandsLUT.cpp
--- a/Engine/Source/Runtime/Renderer/Private/HairStrands/HairStrandsLUT.cpp
+++ b/Engine/Source/Runtime/Renderer/Private/HairStrands/HairStrandsLUT.cpp
@@ -19,6 +19,11 @@ static FAutoConsoleVariableRef CVarHairLUTRoughnessCount(TEXT("r.HairStrands.Hai
 static FAutoConsoleVariableRef CVarHairLUTAbsorptionCount(TEXT("r.HairStrands.HairLUT.AbsorptionCount"), GHairLUTAbsorptionCount, TEXT("Change the number of slices of the hair LUT for the absorption axis"));
 static FAutoConsoleVariableRef CVarHairLUTSampleCount(TEXT("r.HairStrands.HairLUT.SampleCountScale"), GHairLUTSampleCountScale, TEXT("Change the number of sample used for computing the hair LUT. This is a multiplier, default is 1."));
 
+static bool IsHairMobilePlatform(EShaderPlatform Platform)
+{
+	return IsMobilePlatform(Platform) || Platform == SP_PCD3D_ES3_1;
+}
+
 /////////////////////////////////////////////////////////////////////////////////////////
 
 class FHairLUTCS : public FGlobalShader
@@ -39,7 +44,7 @@ class FHairLUTCS : public FGlobalShader
 	END_SHADER_PARAMETER_STRUCT()
 
 public:
-	static bool ShouldCompilePermutation(const FGlobalShaderPermutationParameters& Parameters) { return IsHairStrandsSupported(EHairStrandsShaderType::Strands, Parameters.Platform); }
+	static bool ShouldCompilePermutation(const FGlobalShaderPermutationParameters& Parameters) { return !IsHairMobilePlatform(Parameters.Platform) && IsHairStrandsSupported(EHairStrandsShaderType::Cards, Parameters.Platform); }
 	static void ModifyCompilationEnvironment(const FGlobalShaderPermutationParameters& Parameters, FShaderCompilerEnvironment& OutEnvironment)
 	{
 		FGlobalShader::ModifyCompilationEnvironment(Parameters, OutEnvironment);
@@ -260,7 +265,7 @@ class FHairCoverageLUTCS : public FGlobalShader
 	END_SHADER_PARAMETER_STRUCT()
 
 public:
-	static bool ShouldCompilePermutation(const FGlobalShaderPermutationParameters& Parameters) { return IsHairStrandsSupported(EHairStrandsShaderType::Strands, Parameters.Platform); }
+	static bool ShouldCompilePermutation(const FGlobalShaderPermutationParameters& Parameters) { return !IsHairMobilePlatform(Parameters.Platform) && IsHairStrandsSupported(EHairStrandsShaderType::Cards, Parameters.Platform); }
 	static void ModifyCompilationEnvironment(const FGlobalShaderPermutationParameters& Parameters, FShaderCompilerEnvironment& OutEnvironment)
 	{
 		FGlobalShader::ModifyCompilationEnvironment(Parameters, OutEnvironment);
@@ -324,9 +329,12 @@ static FRDGTextureRef AddHairCoverageLUTPass(FRDGBuilder& GraphBuilder, const FV
 FHairLUT GetHairLUT(FRDGBuilder& GraphBuilder, const FViewInfo& View)
 {
 	// Lazy LUT generation
-	const bool bNeedGenerate =
-		GSystemTextures.HairLUT0.GetReference() == nullptr || GSystemTextures.HairLUT1.GetReference() == nullptr || GSystemTextures.HairLUT2.GetReference() == nullptr ||
-		GSystemTextures.HairLUT0.GetReference()->GetRenderTargetItem().ShaderResourceTexture->GetSizeXYZ() != FIntVector(GHairLUTIncidentAngleCount, GHairLUTRoughnessCount, GHairLUTAbsorptionCount);
+	const EShaderPlatform Platform = View.GetShaderPlatform();
+	const bool bNeedGenerate = (!IsHairMobilePlatform(Platform) && IsHairStrandsSupported(EHairStrandsShaderType::Cards, Platform)) &&
+		(GSystemTextures.HairLUT0.GetReference() == nullptr || 
+		 GSystemTextures.HairLUT1.GetReference() == nullptr || 
+		 GSystemTextures.HairLUT2.GetReference() == nullptr ||
+		 GSystemTextures.HairLUT0.GetReference()->GetRenderTargetItem().ShaderResourceTexture->GetSizeXYZ() != FIntVector(GHairLUTIncidentAngleCount, GHairLUTRoughnessCount, GHairLUTAbsorptionCount));
 	FHairLUT HairLUTData;
 	if (bNeedGenerate)
 	{
@@ -346,9 +354,14 @@ FHairLUT GetHairLUT(FRDGBuilder& GraphBuilder, const FViewInfo& View)
 	else
 	{
 		HairLUTData.Textures[HairLUTType_Coverage]		= GraphBuilder.RegisterExternalTexture(GSystemTextures.HairLUT2, TEXT("HairLUTType_Coverage"));
-		HairLUTData.Textures[HairLUTType_DualScattering] = GraphBuilder.RegisterExternalTexture(GSystemTextures.HairLUT0, TEXT("HairLUTType_DualScattering"));
-		HairLUTData.Textures[HairLUTType_MeanEnergy]		= GraphBuilder.RegisterExternalTexture(GSystemTextures.HairLUT1, TEXT("HairLUTType_MeanEnergy"));
+		HairLUTData.Textures[HairLUTType_DualScattering]= GraphBuilder.RegisterExternalTexture(GSystemTextures.HairLUT0, TEXT("HairLUTType_DualScattering"));
+		HairLUTData.Textures[HairLUTType_MeanEnergy]	= GraphBuilder.RegisterExternalTexture(GSystemTextures.HairLUT1, TEXT("HairLUTType_MeanEnergy"));
 	}
 
 	return HairLUTData;
+}
+
+void UpdateHairResources(FRDGBuilder& GraphBuilder, const FViewInfo& View)
+{
+	GetHairLUT(GraphBuilder, View);
 }
\ No newline at end of file
diff --git a/Engine/Source/Runtime/Renderer/Private/SceneVisibility.cpp b/Engine/Source/Runtime/Renderer/Private/SceneVisibility.cpp
--- a/Engine/Source/Runtime/Renderer/Private/SceneVisibility.cpp
+++ b/Engine/Source/Runtime/Renderer/Private/SceneVisibility.cpp
@@ -4453,6 +4453,7 @@ void FSceneRenderer::PostVisibilityFrameSetup(FILCUpdatePrimTaskData& OutILCTask
 }
 
 uint32 GetShadowQuality();
+void UpdateHairResources(FRDGBuilder& GraphBuilder, const FViewInfo& View);
 
 /** 
 * Performs once per frame setup prior to visibility determination.
@@ -4579,6 +4580,7 @@ bool FDeferredShadingSceneRenderer::InitViews(FRDGBuilder& GraphBuilder, const F
 			}
 
 			// Initialize the view's RHI resources.
+			UpdateHairResources(GraphBuilder, View);
 			View.InitRHIResources();
 		}
 
