Fix GPUScene update of SceneFrameNumber

#rnx

diff --git a/Engine/Source/Runtime/Renderer/Private/GPUScene.cpp b/Engine/Source/Runtime/Renderer/Private/GPUScene.cpp
--- a/Engine/Source/Runtime/Renderer/Private/GPUScene.cpp
+++ b/Engine/Source/Runtime/Renderer/Private/GPUScene.cpp
@@ -435,6 +435,9 @@ void FGPUScene::UpdateInternal(FRHICommandListImmediate& RHICmdList, FScene& Sce
 		QUICK_SCOPE_CYCLE_COUNTER(STAT_UpdateGPUScene);
 		SCOPE_CYCLE_COUNTER(STAT_UpdateGPUSceneTime);
 
+		// Store in GPU-scene to enable validation that update has been carried out.
+		SceneFrameNumber = Scene.GetFrameNumber();
+
 		if (GGPUSceneUploadEveryFrame || bUpdateAllPrimitives)
 		{
 			PrimitivesMarkedToUpdate.Init(false, PrimitivesMarkedToUpdate.Num());
