Fix crash when combining adaptive virtual texture with streaming low mips that cover more than the adaptive low mips.

diff --git a/Engine/Source/Runtime/Renderer/Private/VT/AdaptiveVirtualTexture.cpp b/Engine/Source/Runtime/Renderer/Private/VT/AdaptiveVirtualTexture.cpp
--- a/Engine/Source/Runtime/Renderer/Private/VT/AdaptiveVirtualTexture.cpp
+++ b/Engine/Source/Runtime/Renderer/Private/VT/AdaptiveVirtualTexture.cpp
@@ -52,9 +52,8 @@ static TAutoConsoleVariable<int32> CVarAVTLevelIncrement(
 class FVirtualTextureAddressRedirect : public IVirtualTexture
 {
 public:
-	FVirtualTextureAddressRedirect(FVirtualTextureProducerHandle InProducerHandle, IVirtualTexture* InVirtualTexture, FIntPoint InAddressOffset, int32 InLevelOffset)
-		: ProducerHandle(InProducerHandle)
-		, VirtualTexture(InVirtualTexture)
+	FVirtualTextureAddressRedirect(IVirtualTexture* InVirtualTexture, FIntPoint InAddressOffset, int32 InLevelOffset)
+		: VirtualTexture(InVirtualTexture)
 		, AddressOffset(InAddressOffset)
 		, LevelOffset(InLevelOffset)
 	{
@@ -65,7 +64,7 @@ public:
 	}
 
 	virtual FVTRequestPageResult RequestPageData(
-		const FVirtualTextureProducerHandle& InProducerHandle,
+		const FVirtualTextureProducerHandle& ProducerHandle,
 		uint8 LayerMask,
 		uint8 vLevel,
 		uint64 vAddress,
@@ -84,7 +83,7 @@ public:
 		FRHICommandListImmediate& RHICmdList,
 		ERHIFeatureLevel::Type FeatureLevel,
 		EVTProducePageFlags Flags,
-		const FVirtualTextureProducerHandle& InProducerHandle,
+		const FVirtualTextureProducerHandle& ProducerHandle,
 		uint8 LayerMask,
 		uint8 vLevel,
 		uint64 vAddress,
@@ -101,7 +100,6 @@ public:
 	}
 
 private:
-	FVirtualTextureProducerHandle ProducerHandle;
 	IVirtualTexture* VirtualTexture;
 	FIntPoint AddressOffset;
 	int32 LevelOffset;
@@ -160,7 +158,7 @@ namespace
 			NewProducerDesc.MaxLevel = FMath::CeilLogTwo(FMath::Max(InWidthInTiles, InHeightInTiles));
 
 			IVirtualTexture* VirtualTextureProducer = Producer->GetVirtualTexture();
-			IVirtualTexture* NewVirtualTextureProducer = new FVirtualTextureAddressRedirect(ProducerHandle, VirtualTextureProducer, InAddressOffset, InLevelOffset);
+			IVirtualTexture* NewVirtualTextureProducer = new FVirtualTextureAddressRedirect(VirtualTextureProducer, InAddressOffset, InLevelOffset);
 			FVirtualTextureProducerHandle NewProducerHandle = InSystem->RegisterProducer(NewProducerDesc, NewVirtualTextureProducer);
 
 			// Copy new producer to all subsequent layers.
