Fixed cloud RHI validation failures.

#rnx

diff --git a/Engine/Source/Runtime/RenderCore/Private/RenderGraphUtils.cpp b/Engine/Source/Runtime/RenderCore/Private/RenderGraphUtils.cpp
--- a/Engine/Source/Runtime/RenderCore/Private/RenderGraphUtils.cpp
+++ b/Engine/Source/Runtime/RenderCore/Private/RenderGraphUtils.cpp
@@ -831,13 +831,11 @@ BEGIN_SHADER_PARAMETER_STRUCT(FTextureAccessDynamicPassParameters, )
 	RDG_TEXTURE_ACCESS_DYNAMIC(Texture)
 END_SHADER_PARAMETER_STRUCT()
 
-void ConvertToUntrackedExternalTexture(
+void ConvertToUntrackedTexture(
 	FRDGBuilder& GraphBuilder,
 	FRDGTextureRef Texture,
-	TRefCountPtr<IPooledRenderTarget>& OutPooledRenderTarget,
 	ERHIAccess AccessFinal)
 {
-	ConvertToExternalTexture(GraphBuilder, Texture, OutPooledRenderTarget);
 	GraphBuilder.SetTextureAccessFinal(Texture, AccessFinal);
 
 	auto* PassParameters = GraphBuilder.AllocParameters<FTextureAccessDynamicPassParameters>();
@@ -857,13 +855,11 @@ BEGIN_SHADER_PARAMETER_STRUCT(FBufferAccessDynamicPassParameters, )
 	RDG_BUFFER_ACCESS_DYNAMIC(Buffer)
 END_SHADER_PARAMETER_STRUCT()
 
-void ConvertToUntrackedExternalBuffer(
+void ConvertToUntrackedBuffer(
 	FRDGBuilder& GraphBuilder,
 	FRDGBufferRef Buffer,
-	TRefCountPtr<FRDGPooledBuffer>& OutPooledBuffer,
 	ERHIAccess AccessFinal)
 {
-	ConvertToExternalBuffer(GraphBuilder, Buffer, OutPooledBuffer);
 	GraphBuilder.SetBufferAccessFinal(Buffer, AccessFinal);
 
 	auto* PassParameters = GraphBuilder.AllocParameters<FBufferAccessDynamicPassParameters>();
diff --git a/Engine/Source/Runtime/RenderCore/Public/RenderGraphUtils.h b/Engine/Source/Runtime/RenderCore/Public/RenderGraphUtils.h
--- a/Engine/Source/Runtime/RenderCore/Public/RenderGraphUtils.h
+++ b/Engine/Source/Runtime/RenderCore/Public/RenderGraphUtils.h
@@ -798,19 +798,31 @@ inline void ConvertToExternalTexture(FRDGBuilder& GraphBuilder, FRDGTextureRef T
 	OutPooledRenderTarget = GraphBuilder.GetPooledTexture(Texture);
 }
 
-// Performs the same steps as ConvertToExternalX but forces the resource into the final state. Assuming the resource isn't used again
-// after this pass, the graph will leave it in that state and it's technically safe to read from the raw RHI resource.
-RENDERCORE_API void ConvertToUntrackedExternalTexture(
+// Forces a resource into its final state early in the graph. Useful if an external texture is produced by the graph but consumed
+// later in the graph without being registered. For example, if the resource is read-only embedded in a non-RDG uniform buffer and
+// or registration with RDG is too expensive or difficult.
+RENDERCORE_API void ConvertToUntrackedTexture(FRDGBuilder& GraphBuilder, FRDGTextureRef Texture, ERHIAccess AccessFinal);
+RENDERCORE_API void ConvertToUntrackedBuffer(FRDGBuilder& GraphBuilder, FRDGBufferRef Buffer, ERHIAccess AccessFinal);
+
+inline void ConvertToUntrackedExternalTexture(
 	FRDGBuilder& GraphBuilder,
 	FRDGTextureRef Texture,
 	TRefCountPtr<IPooledRenderTarget>& OutPooledRenderTarget,
-	ERHIAccess AccessFinal);
+	ERHIAccess AccessFinal)
+{
+	ConvertToExternalTexture(GraphBuilder, Texture, OutPooledRenderTarget);
+	ConvertToUntrackedTexture(GraphBuilder, Texture, AccessFinal);
+}
 
-RENDERCORE_API void ConvertToUntrackedExternalBuffer(
+inline void ConvertToUntrackedExternalBuffer(
 	FRDGBuilder& GraphBuilder,
 	FRDGBufferRef Buffer,
 	TRefCountPtr<FRDGPooledBuffer>& OutPooledBuffer,
-	ERHIAccess AccessFinal);
+	ERHIAccess AccessFinal)
+{
+	ConvertToExternalBuffer(GraphBuilder, Buffer, OutPooledBuffer);
+	ConvertToUntrackedBuffer(GraphBuilder, Buffer, AccessFinal);
+}
 
 // Used to help port code over to RDG. If the graph builder is null, creates a passthrough RDG texture instead. Will be removed once port is complete.
 RENDERCORE_API FRDGTextureRef RegisterExternalOrPassthroughTexture(
diff --git a/Engine/Source/Runtime/Renderer/Private/DeferredShadingRenderer.h b/Engine/Source/Runtime/Renderer/Private/DeferredShadingRenderer.h
--- a/Engine/Source/Runtime/Renderer/Private/DeferredShadingRenderer.h
+++ b/Engine/Source/Runtime/Renderer/Private/DeferredShadingRenderer.h
@@ -705,6 +705,8 @@ private:
 	/** Used by RenderLights to render a light to the scene color buffer. */
 	void RenderLight(
 		FRHICommandList& RHICmdList,
+		const FViewInfo& View,
+		int32 ViewIndex,
 		const FLightSceneInfo* LightSceneInfo,
 		FRHITexture* ScreenShadowMaskTexture,
 		FRHITexture* LightingChannelTexture,
diff --git a/Engine/Source/Runtime/Renderer/Private/IndirectLightRendering.cpp b/Engine/Source/Runtime/Renderer/Private/IndirectLightRendering.cpp
--- a/Engine/Source/Runtime/Renderer/Private/IndirectLightRendering.cpp
+++ b/Engine/Source/Runtime/Renderer/Private/IndirectLightRendering.cpp
@@ -1406,7 +1406,7 @@ void FDeferredShadingSceneRenderer::RenderDeferredReflectionsAndSkyLighting(
 				{
 					FVolumetricCloudRenderSceneInfo* CloudInfo = Scene->GetVolumetricCloudSceneInfo();
 
-					PassParameters->CloudSkyAOTexture = View.VolumetricCloudSkyAO.IsValid() ? GraphBuilder.RegisterExternalTexture(View.VolumetricCloudSkyAO) : SystemTextures.Black;
+					PassParameters->CloudSkyAOTexture = View.VolumetricCloudSkyAO != nullptr ? View.VolumetricCloudSkyAO : SystemTextures.Black;
 					PassParameters->CloudSkyAOWorldToLightClipMatrix = CloudInfo->GetVolumetricCloudCommonShaderParameters().CloudSkyAOWorldToLightClipMatrix;
 					PassParameters->CloudSkyAOFarDepthKm = CloudInfo->GetVolumetricCloudCommonShaderParameters().CloudSkyAOFarDepthKm;
 					PassParameters->CloudSkyAOEnabled = 1;
diff --git a/Engine/Source/Runtime/Renderer/Private/LightRendering.cpp b/Engine/Source/Runtime/Renderer/Private/LightRendering.cpp
--- a/Engine/Source/Runtime/Renderer/Private/LightRendering.cpp
+++ b/Engine/Source/Runtime/Renderer/Private/LightRendering.cpp
@@ -1183,9 +1183,10 @@ static FHairStrandsOcclusionResources GetHairStrandsResources(int32 ViewIndex, F
 /** Shader parameters to use when creating a RenderLight(...) pass. */
 BEGIN_SHADER_PARAMETER_STRUCT(FRenderLightParameters, )
 	SHADER_PARAMETER_RDG_UNIFORM_BUFFER(FSceneTextureUniformParameters, SceneTextures)
+	SHADER_PARAMETER_STRUCT_INCLUDE(FVolumetricCloudShadowAOParameters, CloudShadowAO)
 	RDG_TEXTURE_ACCESS(HairCategorizationTexture, ERHIAccess::SRVGraphics)
 	RDG_TEXTURE_ACCESS(ShadowMaskTexture, ERHIAccess::SRVGraphics)
-	RDG_TEXTURE_ACCESS(LightingChannelsTexture, ERHIAccess::SRVGraphics) 
+	RDG_TEXTURE_ACCESS(LightingChannelsTexture, ERHIAccess::SRVGraphics)
 	RENDER_TARGET_BINDING_SLOTS()
 END_SHADER_PARAMETER_STRUCT()
 
@@ -1195,12 +1196,14 @@ void GetRenderLightParameters(
 	TRDGUniformBufferRef<FSceneTextureUniformParameters> SceneTexturesUniformBuffer,
 	FRDGTextureRef ShadowMaskTexture,
 	FRDGTextureRef LightingChannelsTexture,
+	const FVolumetricCloudShadowAOParameters& CloudShadowAOParameters,
 	const FHairStrandsVisibilityViews* InHairVisibilityViews,
 	FRenderLightParameters& Parameters)
 {
 	Parameters.SceneTextures = SceneTexturesUniformBuffer;
 	Parameters.ShadowMaskTexture = ShadowMaskTexture;
 	Parameters.LightingChannelsTexture = LightingChannelsTexture;
+	Parameters.CloudShadowAO = CloudShadowAOParameters;
 	Parameters.HairCategorizationTexture = InHairVisibilityViews && InHairVisibilityViews->HairDatas.Num() > 0 ? InHairVisibilityViews->HairDatas[0].CategorizationTexture : nullptr;
 	Parameters.RenderTargets[0] = FRenderTargetBinding(SceneColorTexture, ERenderTargetLoadAction::ELoad);
 
@@ -1214,10 +1217,11 @@ void GetRenderLightParameters(
 	const FMinimalSceneTextures& SceneTextures,
 	FRDGTextureRef ShadowMaskTexture,
 	FRDGTextureRef LightingChannelsTexture,
+	const FVolumetricCloudShadowAOParameters& CloudShadowAOParameters,
 	const FHairStrandsVisibilityViews* InHairVisibilityViews,
 	FRenderLightParameters& Parameters)
 {
-	GetRenderLightParameters(SceneTextures.Color.Target, SceneTextures.Depth.Target, SceneTextures.UniformBuffer, ShadowMaskTexture, LightingChannelsTexture, InHairVisibilityViews, Parameters);
+	GetRenderLightParameters(SceneTextures.Color.Target, SceneTextures.Depth.Target, SceneTextures.UniformBuffer, ShadowMaskTexture, LightingChannelsTexture, CloudShadowAOParameters, InHairVisibilityViews, Parameters);
 }
 
 void FDeferredShadingSceneRenderer::RenderLights(
@@ -1303,7 +1307,7 @@ void FDeferredShadingSceneRenderer::RenderLights(
 			if (!bUseHairLighting)
 			{
 				FRenderLightParameters* PassParameters = GraphBuilder.AllocParameters<FRenderLightParameters>();
-				GetRenderLightParameters(SceneTextures, nullptr, LightingChannelsTexture, HairDatas ? &HairDatas->HairVisibilityViews : nullptr, *PassParameters);
+				GetRenderLightParameters(SceneTextures, nullptr, LightingChannelsTexture, {}, HairDatas ? &HairDatas->HairVisibilityViews : nullptr, *PassParameters);
 
 				GraphBuilder.AddPass(
 					RDG_EVENT_NAME("StandardDeferredLighting"),
@@ -1317,8 +1321,14 @@ void FDeferredShadingSceneRenderer::RenderLights(
 						const FSortedLightSceneInfo& SortedLightInfo = SortedLights[LightIndex];
 						const FLightSceneInfo* const LightSceneInfo = SortedLightInfo.LightSceneInfo;
 
-						// Render the light to the scene color buffer, using a 1x1 white texture as input
-						RenderLight(RHICmdList, LightSceneInfo, nullptr, TryGetRHI(LightingChannelsTexture), nullptr, false, false);
+						for (int32 ViewIndex = 0, ViewCount = Views.Num(); ViewIndex < ViewCount; ++ViewIndex)
+						{
+							const FViewInfo& View = Views[ViewIndex];
+							SCOPED_GPU_MASK(RHICmdList, View.GPUMask);
+
+							// Render the light to the scene color buffer, using a 1x1 white texture as input
+							RenderLight(RHICmdList, View, ViewIndex, LightSceneInfo, nullptr, TryGetRHI(LightingChannelsTexture), nullptr, false, false);
+						}
 					}
 				});
 			}
@@ -1986,21 +1996,21 @@ void FDeferredShadingSceneRenderer::RenderLightArrayForOverlapViewmode(
 			continue;
 		}
 
-		bool bShouldRender = false;
+		// Only render shadow casting stationary lights
+		if (!LightSceneInfo->Proxy->HasStaticShadowing() ||
+			 LightSceneInfo->Proxy->HasStaticLighting()  ||
+			!LightSceneInfo->Proxy->CastsStaticShadow())
+		{
+			continue;
+		}
 
 		// Check if the light is visible in any of the views.
 		for (int32 ViewIndex = 0; ViewIndex < Views.Num(); ViewIndex++)
 		{
-			bShouldRender |= LightSceneInfo->ShouldRenderLight(Views[ViewIndex]);
-		}
+			const FViewInfo& View = Views[ViewIndex];
+			SCOPED_GPU_MASK(RHICmdList, View.GPUMask);
 
-		if (bShouldRender
-			// Only render shadow casting stationary lights
-			&& LightSceneInfo->Proxy->HasStaticShadowing()
-			&& !LightSceneInfo->Proxy->HasStaticLighting()
-			&& LightSceneInfo->Proxy->CastsStaticShadow())
-		{
-			RenderLight(RHICmdList, LightSceneInfo, nullptr, LightingChannelsTexture, nullptr, true, false);
+			RenderLight(RHICmdList, View, ViewIndex, LightSceneInfo, nullptr, LightingChannelsTexture, nullptr, true, false);
 		}
 	}
 }
@@ -2013,7 +2023,7 @@ void FDeferredShadingSceneRenderer::RenderStationaryLightOverlap(
 	if (Scene->bIsEditorScene)
 	{
 		FRenderLightParameters* PassParameters = GraphBuilder.AllocParameters<FRenderLightParameters>();
-		GetRenderLightParameters(SceneTextures, nullptr, LightingChannelsTexture, nullptr, *PassParameters);
+		GetRenderLightParameters(SceneTextures, nullptr, LightingChannelsTexture, {}, nullptr, *PassParameters);
 
 		GraphBuilder.AddPass(
 			RDG_EVENT_NAME("StationaryLightOverlap"),
@@ -2126,12 +2136,20 @@ void CalculateLightNearFarDepthFromBounds(const FViewInfo& View, const FSphere &
 
 void FDeferredShadingSceneRenderer::RenderLight(
 	FRHICommandList& RHICmdList,
+	const FViewInfo& View,
+	int32 ViewIndex,
 	const FLightSceneInfo* LightSceneInfo,
 	FRHITexture* ScreenShadowMaskTexture,
 	FRHITexture* LightingChannelsTexture,
 	const FHairStrandsVisibilityViews* InHairVisibilityViews,
 	bool bRenderOverlap, bool bIssueDrawEvent)
 {
+	// Ensure the light is valid for this view
+	if (!LightSceneInfo->ShouldRenderLight(View))
+	{
+		return;
+	}
+
 	SCOPE_CYCLE_COUNTER(STAT_DirectLightRenderingTime);
 	INC_DWORD_STAT(STAT_NumLightsUsingStandardDeferred);
 	SCOPED_CONDITIONAL_DRAW_EVENT(RHICmdList, StandardDeferredLighting, bIssueDrawEvent);
@@ -2147,199 +2165,187 @@ void FDeferredShadingSceneRenderer::RenderLight(
 	const FSphere LightBounds = LightProxy->GetBoundingSphere();
 	const bool bTransmission = LightProxy->Transmission();
 
-	for (int32 ViewIndex = 0; ViewIndex < Views.Num(); ViewIndex++)
+	bool bUseIESTexture = false;
+
+	if(View.Family->EngineShowFlags.TexturedLightProfiles)
 	{
-		FViewInfo& View = Views[ViewIndex];
-		SCOPED_GPU_MASK(RHICmdList, View.GPUMask);
+		bUseIESTexture = (LightSceneInfo->Proxy->GetIESTextureResource() != 0);
+	}
 
-		// Ensure the light is valid for this view
-		if (!LightSceneInfo->ShouldRenderLight(View))
-		{
-			continue;
-		}
+	// Set the device viewport for the view.
+	RHICmdList.SetViewport(View.ViewRect.Min.X, View.ViewRect.Min.Y, 0.0f, View.ViewRect.Max.X, View.ViewRect.Max.Y, 1.0f);
 
-		bool bUseIESTexture = false;
+	FRenderLightParams RenderLightParams;
+	const bool bHairLighting = InHairVisibilityViews && ViewIndex < InHairVisibilityViews->HairDatas.Num() && InHairVisibilityViews->HairDatas[ViewIndex].CategorizationTexture != nullptr;
+	if (bHairLighting)
+	{
+		RenderLightParams.HairCategorizationTexture = InHairVisibilityViews->HairDatas[ViewIndex].CategorizationTexture->GetPooledRenderTarget();
+	}
 
-		if(View.Family->EngineShowFlags.TexturedLightProfiles)
-		{
-			bUseIESTexture = (LightSceneInfo->Proxy->GetIESTextureResource() != 0);
-		}
+	if (LightProxy->GetLightType() == LightType_Directional)
+	{
+		// Turn DBT back off
+		GraphicsPSOInit.bDepthBounds = false;
+		TShaderMapRef<TDeferredLightVS<false> > VertexShader(View.ShaderMap);
 
-		// Set the device viewport for the view.
-		RHICmdList.SetViewport(View.ViewRect.Min.X, View.ViewRect.Min.Y, 0.0f, View.ViewRect.Max.X, View.ViewRect.Max.Y, 1.0f);
+		GraphicsPSOInit.RasterizerState = TStaticRasterizerState<FM_Solid, CM_None>::GetRHI();
+		GraphicsPSOInit.DepthStencilState = TStaticDepthStencilState<false, CF_Always>::GetRHI();
 
-		FRenderLightParams RenderLightParams;
-		const bool bHairLighting = InHairVisibilityViews && ViewIndex < InHairVisibilityViews->HairDatas.Num() && InHairVisibilityViews->HairDatas[ViewIndex].CategorizationTexture != nullptr;
-		if (bHairLighting)
+		if (bRenderOverlap)
 		{
-			RenderLightParams.HairCategorizationTexture = InHairVisibilityViews->HairDatas[ViewIndex].CategorizationTexture->GetPooledRenderTarget();
+			TShaderMapRef<TDeferredLightOverlapPS<false> > PixelShader(View.ShaderMap);
+			GraphicsPSOInit.BoundShaderState.VertexDeclarationRHI = GFilterVertexDeclaration.VertexDeclarationRHI;
+			GraphicsPSOInit.BoundShaderState.VertexShaderRHI = VertexShader.GetVertexShader();
+			GraphicsPSOInit.BoundShaderState.PixelShaderRHI = PixelShader.GetPixelShader();
+			SetGraphicsPipelineState(RHICmdList, GraphicsPSOInit);
+			PixelShader->SetParameters(RHICmdList, View, LightSceneInfo);
 		}
-
-		if (LightProxy->GetLightType() == LightType_Directional)
+		else
 		{
-			// Turn DBT back off
-			GraphicsPSOInit.bDepthBounds = false;
-			TShaderMapRef<TDeferredLightVS<false> > VertexShader(View.ShaderMap);
-
-			GraphicsPSOInit.RasterizerState = TStaticRasterizerState<FM_Solid, CM_None>::GetRHI();
-			GraphicsPSOInit.DepthStencilState = TStaticDepthStencilState<false, CF_Always>::GetRHI();
-
-			if (bRenderOverlap)
+			const bool bAtmospherePerPixelTransmittance = LightSceneInfo->Proxy->IsUsedAsAtmosphereSunLight() 
+				&& LightSceneInfo->Proxy->GetUsePerPixelAtmosphereTransmittance() && ShouldRenderSkyAtmosphere(Scene, View.Family->EngineShowFlags);
+
+			// Only atmospheric light 0 supports cloud shadow as of today.
+			FLightSceneProxy* AtmosphereLight0Proxy = Scene->AtmosphereLights[0] ? Scene->AtmosphereLights[0]->Proxy : nullptr;
+			FLightSceneProxy* AtmosphereLight1Proxy = Scene->AtmosphereLights[1] ? Scene->AtmosphereLights[1]->Proxy : nullptr;
+			FVolumetricCloudRenderSceneInfo* CloudInfo = Scene->GetVolumetricCloudSceneInfo();
+			const bool VolumetricCloudShadowMap0Valid = View.VolumetricCloudShadowRenderTarget[0] != nullptr;
+			const bool VolumetricCloudShadowMap1Valid = View.VolumetricCloudShadowRenderTarget[1] != nullptr;
+			const bool bLight0CloudPerPixelTransmittance = CloudInfo && VolumetricCloudShadowMap0Valid && AtmosphereLight0Proxy == LightSceneInfo->Proxy && AtmosphereLight0Proxy && AtmosphereLight0Proxy->GetCloudShadowOnSurfaceStrength() > 0.0f;
+			const bool bLight1CloudPerPixelTransmittance = CloudInfo && VolumetricCloudShadowMap1Valid && AtmosphereLight1Proxy == LightSceneInfo->Proxy && AtmosphereLight1Proxy && AtmosphereLight1Proxy->GetCloudShadowOnSurfaceStrength() > 0.0f;
+			if (bLight0CloudPerPixelTransmittance)
 			{
-				TShaderMapRef<TDeferredLightOverlapPS<false> > PixelShader(View.ShaderMap);
-				GraphicsPSOInit.BoundShaderState.VertexDeclarationRHI = GFilterVertexDeclaration.VertexDeclarationRHI;
-				GraphicsPSOInit.BoundShaderState.VertexShaderRHI = VertexShader.GetVertexShader();
-				GraphicsPSOInit.BoundShaderState.PixelShaderRHI = PixelShader.GetPixelShader();
-				SetGraphicsPipelineState(RHICmdList, GraphicsPSOInit);
-				PixelShader->SetParameters(RHICmdList, View, LightSceneInfo);
+				RenderLightParams.Cloud_ShadowmapTexture = TryGetPooledRenderTarget(View.VolumetricCloudShadowRenderTarget[0]);
+				RenderLightParams.Cloud_ShadowmapFarDepthKm = CloudInfo->GetVolumetricCloudCommonShaderParameters().CloudShadowmapFarDepthKm[0];
+				RenderLightParams.Cloud_WorldToLightClipShadowMatrix = CloudInfo->GetVolumetricCloudCommonShaderParameters().CloudShadowmapWorldToLightClipMatrix[0];
+				RenderLightParams.Cloud_ShadowmapStrength = AtmosphereLight0Proxy->GetCloudShadowOnSurfaceStrength();
 			}
-			else
+			else if(bLight1CloudPerPixelTransmittance)
 			{
-				const bool bAtmospherePerPixelTransmittance = LightSceneInfo->Proxy->IsUsedAsAtmosphereSunLight() 
-					&& LightSceneInfo->Proxy->GetUsePerPixelAtmosphereTransmittance() && ShouldRenderSkyAtmosphere(Scene, View.Family->EngineShowFlags);
-
-				// Only atmospheric light 0 supports cloud shadow as of today.
-				FLightSceneProxy* AtmosphereLight0Proxy = Scene->AtmosphereLights[0] ? Scene->AtmosphereLights[0]->Proxy : nullptr;
-				FLightSceneProxy* AtmosphereLight1Proxy = Scene->AtmosphereLights[1] ? Scene->AtmosphereLights[1]->Proxy : nullptr;
-				FVolumetricCloudRenderSceneInfo* CloudInfo = Scene->GetVolumetricCloudSceneInfo();
-				const bool VolumetricCloudShadowMap0Valid = View.VolumetricCloudShadowRenderTarget[0].IsValid();
-				const bool VolumetricCloudShadowMap1Valid = View.VolumetricCloudShadowRenderTarget[1].IsValid();
-				const bool bLight0CloudPerPixelTransmittance = CloudInfo && VolumetricCloudShadowMap0Valid && AtmosphereLight0Proxy == LightSceneInfo->Proxy && AtmosphereLight0Proxy && AtmosphereLight0Proxy->GetCloudShadowOnSurfaceStrength() > 0.0f;
-				const bool bLight1CloudPerPixelTransmittance = CloudInfo && VolumetricCloudShadowMap1Valid && AtmosphereLight1Proxy == LightSceneInfo->Proxy && AtmosphereLight1Proxy && AtmosphereLight1Proxy->GetCloudShadowOnSurfaceStrength() > 0.0f;
-				if (bLight0CloudPerPixelTransmittance)
-				{
-					RenderLightParams.Cloud_ShadowmapTexture = View.VolumetricCloudShadowRenderTarget[0];
-					RenderLightParams.Cloud_ShadowmapFarDepthKm = CloudInfo->GetVolumetricCloudCommonShaderParameters().CloudShadowmapFarDepthKm[0];
-					RenderLightParams.Cloud_WorldToLightClipShadowMatrix = CloudInfo->GetVolumetricCloudCommonShaderParameters().CloudShadowmapWorldToLightClipMatrix[0];
-					RenderLightParams.Cloud_ShadowmapStrength = AtmosphereLight0Proxy->GetCloudShadowOnSurfaceStrength();
-				}
-				else if(bLight1CloudPerPixelTransmittance)
-				{
-					RenderLightParams.Cloud_ShadowmapTexture = View.VolumetricCloudShadowRenderTarget[1];
-					RenderLightParams.Cloud_ShadowmapFarDepthKm = CloudInfo->GetVolumetricCloudCommonShaderParameters().CloudShadowmapFarDepthKm[1];
-					RenderLightParams.Cloud_WorldToLightClipShadowMatrix = CloudInfo->GetVolumetricCloudCommonShaderParameters().CloudShadowmapWorldToLightClipMatrix[1];
-					RenderLightParams.Cloud_ShadowmapStrength = AtmosphereLight1Proxy->GetCloudShadowOnSurfaceStrength();
-				}
-
-				FDeferredLightPS::FPermutationDomain PermutationVector;
-				PermutationVector.Set< FDeferredLightPS::FSourceShapeDim >( ELightSourceShape::Directional );
-				PermutationVector.Set< FDeferredLightPS::FIESProfileDim >( false );
-				PermutationVector.Set< FDeferredLightPS::FInverseSquaredDim >( false );
-				PermutationVector.Set< FDeferredLightPS::FVisualizeCullingDim >( View.Family->EngineShowFlags.VisualizeLightCulling );
-				PermutationVector.Set< FDeferredLightPS::FLightingChannelsDim >( View.bUsesLightingChannels );
-				PermutationVector.Set< FDeferredLightPS::FAnistropicMaterials >(ShouldRenderAnisotropyPass());
-				PermutationVector.Set< FDeferredLightPS::FTransmissionDim >( bTransmission );
-				PermutationVector.Set< FDeferredLightPS::FHairLighting>(0);
-				// Only directional lights are rendered in this path, so we only need to check if it is use to light the atmosphere
-				PermutationVector.Set< FDeferredLightPS::FAtmosphereTransmittance >(bAtmospherePerPixelTransmittance);
-				PermutationVector.Set< FDeferredLightPS::FCloudTransmittance >(bLight0CloudPerPixelTransmittance || bLight1CloudPerPixelTransmittance);
-				PermutationVector.Set< FDeferredLightPS::FStrata >(Strata::IsStrataEnabled());
-
-				TShaderMapRef< FDeferredLightPS > PixelShader( View.ShaderMap, PermutationVector );
-				GraphicsPSOInit.BoundShaderState.VertexDeclarationRHI = GFilterVertexDeclaration.VertexDeclarationRHI;
-				GraphicsPSOInit.BoundShaderState.VertexShaderRHI = VertexShader.GetVertexShader();
-				GraphicsPSOInit.BoundShaderState.PixelShaderRHI = PixelShader.GetPixelShader();
-
-				SetGraphicsPipelineState(RHICmdList, GraphicsPSOInit);
-				PixelShader->SetParameters(RHICmdList, View, LightSceneInfo, ScreenShadowMaskTexture, LightingChannelsTexture, &RenderLightParams);
+				RenderLightParams.Cloud_ShadowmapTexture = TryGetPooledRenderTarget(View.VolumetricCloudShadowRenderTarget[1]);
+				RenderLightParams.Cloud_ShadowmapFarDepthKm = CloudInfo->GetVolumetricCloudCommonShaderParameters().CloudShadowmapFarDepthKm[1];
+				RenderLightParams.Cloud_WorldToLightClipShadowMatrix = CloudInfo->GetVolumetricCloudCommonShaderParameters().CloudShadowmapWorldToLightClipMatrix[1];
+				RenderLightParams.Cloud_ShadowmapStrength = AtmosphereLight1Proxy->GetCloudShadowOnSurfaceStrength();
 			}
 
-			VertexShader->SetParameters(RHICmdList, View, LightSceneInfo);
+			FDeferredLightPS::FPermutationDomain PermutationVector;
+			PermutationVector.Set< FDeferredLightPS::FSourceShapeDim >( ELightSourceShape::Directional );
+			PermutationVector.Set< FDeferredLightPS::FIESProfileDim >( false );
+			PermutationVector.Set< FDeferredLightPS::FInverseSquaredDim >( false );
+			PermutationVector.Set< FDeferredLightPS::FVisualizeCullingDim >( View.Family->EngineShowFlags.VisualizeLightCulling );
+			PermutationVector.Set< FDeferredLightPS::FLightingChannelsDim >( View.bUsesLightingChannels );
+			PermutationVector.Set< FDeferredLightPS::FAnistropicMaterials >(ShouldRenderAnisotropyPass());
+			PermutationVector.Set< FDeferredLightPS::FTransmissionDim >( bTransmission );
+			PermutationVector.Set< FDeferredLightPS::FHairLighting>(0);
+			// Only directional lights are rendered in this path, so we only need to check if it is use to light the atmosphere
+			PermutationVector.Set< FDeferredLightPS::FAtmosphereTransmittance >(bAtmospherePerPixelTransmittance);
+			PermutationVector.Set< FDeferredLightPS::FCloudTransmittance >(bLight0CloudPerPixelTransmittance || bLight1CloudPerPixelTransmittance);
+			PermutationVector.Set< FDeferredLightPS::FStrata >(Strata::IsStrataEnabled());
 
-			// Apply the directional light as a full screen quad
-			DrawRectangle(
-				RHICmdList,
-				0, 0,
-				View.ViewRect.Width(), View.ViewRect.Height(),
-				View.ViewRect.Min.X, View.ViewRect.Min.Y,
-				View.ViewRect.Width(), View.ViewRect.Height(),
-				View.ViewRect.Size(),
-				GetSceneTextureExtent(),
-				VertexShader,
-				EDRF_UseTriangleOptimization);
+			TShaderMapRef< FDeferredLightPS > PixelShader( View.ShaderMap, PermutationVector );
+			GraphicsPSOInit.BoundShaderState.VertexDeclarationRHI = GFilterVertexDeclaration.VertexDeclarationRHI;
+			GraphicsPSOInit.BoundShaderState.VertexShaderRHI = VertexShader.GetVertexShader();
+			GraphicsPSOInit.BoundShaderState.PixelShaderRHI = PixelShader.GetPixelShader();
+
+			SetGraphicsPipelineState(RHICmdList, GraphicsPSOInit);
+			PixelShader->SetParameters(RHICmdList, View, LightSceneInfo, ScreenShadowMaskTexture, LightingChannelsTexture, &RenderLightParams);
 		}
-		else
-		{
-			// Use DBT to allow work culling on shadow lights
-			// Disable depth bound when hair rendering is enabled as this rejects partially covered pixel write (with opaque background)
-			GraphicsPSOInit.bDepthBounds = GSupportsDepthBoundsTest && GAllowDepthBoundsTest != 0;
 
-			TShaderMapRef<TDeferredLightVS<true> > VertexShader(View.ShaderMap);
+		VertexShader->SetParameters(RHICmdList, View, LightSceneInfo);
 
-			const bool bCameraInsideLightGeometry = ((FVector)View.ViewMatrices.GetViewOrigin() - LightBounds.Center).SizeSquared() < FMath::Square(LightBounds.W * 1.05f + View.NearClippingDistance * 2.0f)
-			//const bool bCameraInsideLightGeometry = LightProxy->AffectsBounds( FSphere( View.ViewMatrices.GetViewOrigin(), View.NearClippingDistance * 2.0f ) )
-				// Always draw backfaces in ortho
-				//@todo - accurate ortho camera / light intersection
-				|| !View.IsPerspectiveProjection();
+		// Apply the directional light as a full screen quad
+		DrawRectangle(
+			RHICmdList,
+			0, 0,
+			View.ViewRect.Width(), View.ViewRect.Height(),
+			View.ViewRect.Min.X, View.ViewRect.Min.Y,
+			View.ViewRect.Width(), View.ViewRect.Height(),
+			View.ViewRect.Size(),
+			GetSceneTextureExtent(),
+			VertexShader,
+			EDRF_UseTriangleOptimization);
+	}
+	else
+	{
+		// Use DBT to allow work culling on shadow lights
+		// Disable depth bound when hair rendering is enabled as this rejects partially covered pixel write (with opaque background)
+		GraphicsPSOInit.bDepthBounds = GSupportsDepthBoundsTest && GAllowDepthBoundsTest != 0;
 
-			SetBoundingGeometryRasterizerAndDepthState(GraphicsPSOInit, View, bCameraInsideLightGeometry);
+		TShaderMapRef<TDeferredLightVS<true> > VertexShader(View.ShaderMap);
 
-			if (bRenderOverlap)
-			{
-				TShaderMapRef<TDeferredLightOverlapPS<true> > PixelShader(View.ShaderMap);
-				GraphicsPSOInit.BoundShaderState.VertexDeclarationRHI = GetVertexDeclarationFVector4();
-				GraphicsPSOInit.BoundShaderState.VertexShaderRHI = VertexShader.GetVertexShader();
-				GraphicsPSOInit.BoundShaderState.PixelShaderRHI = PixelShader.GetPixelShader();
+		const bool bCameraInsideLightGeometry = ((FVector)View.ViewMatrices.GetViewOrigin() - LightBounds.Center).SizeSquared() < FMath::Square(LightBounds.W * 1.05f + View.NearClippingDistance * 2.0f)
+		//const bool bCameraInsideLightGeometry = LightProxy->AffectsBounds( FSphere( View.ViewMatrices.GetViewOrigin(), View.NearClippingDistance * 2.0f ) )
+			// Always draw backfaces in ortho
+			//@todo - accurate ortho camera / light intersection
+			|| !View.IsPerspectiveProjection();
 
-				SetGraphicsPipelineState(RHICmdList, GraphicsPSOInit);
-				PixelShader->SetParameters(RHICmdList, View, LightSceneInfo);
-			}
-			else
-			{
-				FDeferredLightPS::FPermutationDomain PermutationVector;
-				PermutationVector.Set< FDeferredLightPS::FSourceShapeDim >( LightProxy->IsRectLight() ? ELightSourceShape::Rect : ELightSourceShape::Capsule );
-				PermutationVector.Set< FDeferredLightPS::FSourceTextureDim >( LightProxy->IsRectLight() && LightProxy->HasSourceTexture() );
-				PermutationVector.Set< FDeferredLightPS::FIESProfileDim >( bUseIESTexture );
-				PermutationVector.Set< FDeferredLightPS::FInverseSquaredDim >( LightProxy->IsInverseSquared() );
-				PermutationVector.Set< FDeferredLightPS::FVisualizeCullingDim >( View.Family->EngineShowFlags.VisualizeLightCulling );
-				PermutationVector.Set< FDeferredLightPS::FLightingChannelsDim >( View.bUsesLightingChannels );
-				PermutationVector.Set< FDeferredLightPS::FAnistropicMaterials >(ShouldRenderAnisotropyPass() && !LightSceneInfo->Proxy->IsRectLight());
-				PermutationVector.Set< FDeferredLightPS::FTransmissionDim >( bTransmission );
-				PermutationVector.Set< FDeferredLightPS::FHairLighting>(0);
-				PermutationVector.Set < FDeferredLightPS::FAtmosphereTransmittance >(false);
-				PermutationVector.Set< FDeferredLightPS::FCloudTransmittance >(false);
-				PermutationVector.Set< FDeferredLightPS::FStrata >(Strata::IsStrataEnabled());
-
-				TShaderMapRef< FDeferredLightPS > PixelShader( View.ShaderMap, PermutationVector );
-				GraphicsPSOInit.BoundShaderState.VertexDeclarationRHI = GetVertexDeclarationFVector4();
-				GraphicsPSOInit.BoundShaderState.VertexShaderRHI = VertexShader.GetVertexShader();
-				GraphicsPSOInit.BoundShaderState.PixelShaderRHI = PixelShader.GetPixelShader();
-
-				SetGraphicsPipelineState(RHICmdList, GraphicsPSOInit);
-				PixelShader->SetParameters(RHICmdList, View, LightSceneInfo, ScreenShadowMaskTexture, LightingChannelsTexture, (bHairLighting) ? &RenderLightParams : nullptr);
-			}
+		SetBoundingGeometryRasterizerAndDepthState(GraphicsPSOInit, View, bCameraInsideLightGeometry);
 
-			VertexShader->SetParameters(RHICmdList, View, LightSceneInfo);
+		if (bRenderOverlap)
+		{
+			TShaderMapRef<TDeferredLightOverlapPS<true> > PixelShader(View.ShaderMap);
+			GraphicsPSOInit.BoundShaderState.VertexDeclarationRHI = GetVertexDeclarationFVector4();
+			GraphicsPSOInit.BoundShaderState.VertexShaderRHI = VertexShader.GetVertexShader();
+			GraphicsPSOInit.BoundShaderState.PixelShaderRHI = PixelShader.GetPixelShader();
 
-			// Use DBT to allow work culling on shadow lights
-			if (GraphicsPSOInit.bDepthBounds)
-			{
-				// Can use the depth bounds test to skip work for pixels which won't be touched by the light (i.e outside the depth range)
-				float NearDepth = 1.f;
-				float FarDepth = 0.f;
-				CalculateLightNearFarDepthFromBounds(View,LightBounds,NearDepth,FarDepth);
+			SetGraphicsPipelineState(RHICmdList, GraphicsPSOInit);
+			PixelShader->SetParameters(RHICmdList, View, LightSceneInfo);
+		}
+		else
+		{
+			FDeferredLightPS::FPermutationDomain PermutationVector;
+			PermutationVector.Set< FDeferredLightPS::FSourceShapeDim >( LightProxy->IsRectLight() ? ELightSourceShape::Rect : ELightSourceShape::Capsule );
+			PermutationVector.Set< FDeferredLightPS::FSourceTextureDim >( LightProxy->IsRectLight() && LightProxy->HasSourceTexture() );
+			PermutationVector.Set< FDeferredLightPS::FIESProfileDim >( bUseIESTexture );
+			PermutationVector.Set< FDeferredLightPS::FInverseSquaredDim >( LightProxy->IsInverseSquared() );
+			PermutationVector.Set< FDeferredLightPS::FVisualizeCullingDim >( View.Family->EngineShowFlags.VisualizeLightCulling );
+			PermutationVector.Set< FDeferredLightPS::FLightingChannelsDim >( View.bUsesLightingChannels );
+			PermutationVector.Set< FDeferredLightPS::FAnistropicMaterials >(ShouldRenderAnisotropyPass() && !LightSceneInfo->Proxy->IsRectLight());
+			PermutationVector.Set< FDeferredLightPS::FTransmissionDim >( bTransmission );
+			PermutationVector.Set< FDeferredLightPS::FHairLighting>(0);
+			PermutationVector.Set < FDeferredLightPS::FAtmosphereTransmittance >(false);
+			PermutationVector.Set< FDeferredLightPS::FCloudTransmittance >(false);
+			PermutationVector.Set< FDeferredLightPS::FStrata >(Strata::IsStrataEnabled());
 
-				if (NearDepth <= FarDepth)
-				{
-					NearDepth = 1.0f;
-					FarDepth = 0.0f;
-				}
+			TShaderMapRef< FDeferredLightPS > PixelShader( View.ShaderMap, PermutationVector );
+			GraphicsPSOInit.BoundShaderState.VertexDeclarationRHI = GetVertexDeclarationFVector4();
+			GraphicsPSOInit.BoundShaderState.VertexShaderRHI = VertexShader.GetVertexShader();
+			GraphicsPSOInit.BoundShaderState.PixelShaderRHI = PixelShader.GetPixelShader();
 
-				// UE4 uses reversed depth, so far < near
-				RHICmdList.SetDepthBounds(FarDepth, NearDepth);
-			}
+			SetGraphicsPipelineState(RHICmdList, GraphicsPSOInit);
+			PixelShader->SetParameters(RHICmdList, View, LightSceneInfo, ScreenShadowMaskTexture, LightingChannelsTexture, (bHairLighting) ? &RenderLightParams : nullptr);
+		}
 
-			if( LightProxy->GetLightType() == LightType_Point ||
-				LightProxy->GetLightType() == LightType_Rect )
-			{
-				// Apply the point or spot light with some approximate bounding geometry,
-				// So we can get speedups from depth testing and not processing pixels outside of the light's influence.
-				StencilingGeometry::DrawSphere(RHICmdList);
-			}
-			else if (LightProxy->GetLightType() == LightType_Spot)
+		VertexShader->SetParameters(RHICmdList, View, LightSceneInfo);
+
+		// Use DBT to allow work culling on shadow lights
+		if (GraphicsPSOInit.bDepthBounds)
+		{
+			// Can use the depth bounds test to skip work for pixels which won't be touched by the light (i.e outside the depth range)
+			float NearDepth = 1.f;
+			float FarDepth = 0.f;
+			CalculateLightNearFarDepthFromBounds(View,LightBounds,NearDepth,FarDepth);
+
+			if (NearDepth <= FarDepth)
 			{
-				StencilingGeometry::DrawCone(RHICmdList);
+				NearDepth = 1.0f;
+				FarDepth = 0.0f;
 			}
+
+			// UE4 uses reversed depth, so far < near
+			RHICmdList.SetDepthBounds(FarDepth, NearDepth);
+		}
+
+		if( LightProxy->GetLightType() == LightType_Point ||
+			LightProxy->GetLightType() == LightType_Rect )
+		{
+			// Apply the point or spot light with some approximate bounding geometry,
+			// So we can get speedups from depth testing and not processing pixels outside of the light's influence.
+			StencilingGeometry::DrawSphere(RHICmdList);
+		}
+		else if (LightProxy->GetLightType() == LightType_Spot)
+		{
+			StencilingGeometry::DrawCone(RHICmdList);
 		}
 	}
 }
@@ -2353,31 +2359,36 @@ void FDeferredShadingSceneRenderer::RenderLight(
 	const FHairStrandsVisibilityViews* InHairVisibilityViews,
 	bool bRenderOverlap)
 {
-	FRenderLightParameters* PassParameters = GraphBuilder.AllocParameters<FRenderLightParameters>();
-	GetRenderLightParameters(SceneTextures, ScreenShadowMaskTexture, LightingChannelsTexture, InHairVisibilityViews, *PassParameters);
-
 	ERDGPassFlags PassFlags = ERDGPassFlags::Raster;
 
-	if (InHairVisibilityViews)
-	{
-		PassFlags |= ERDGPassFlags::UntrackedAccess;
-	}
+	const FVolumetricCloudRenderSceneInfo* CloudInfo = Scene->GetVolumetricCloudSceneInfo();
 
-	GraphBuilder.AddPass(
-		RDG_EVENT_NAME("StandardDeferredLighting"),
-		PassParameters,
-		PassFlags,
-		[this, LightSceneInfo, ScreenShadowMaskTexture, LightingChannelsTexture, InHairVisibilityViews, bRenderOverlap](FRHICommandList& RHICmdList)
+	for (int32 ViewIndex = 0, ViewCount = Views.Num(); ViewIndex < ViewCount; ++ViewIndex)
 	{
-		RenderLight(
-			RHICmdList,
-			LightSceneInfo,
-			TryGetRHI(ScreenShadowMaskTexture),
-			TryGetRHI(LightingChannelsTexture),
-			InHairVisibilityViews,
-			bRenderOverlap,
-			false);
-	});
+		RDG_EVENT_SCOPE_CONDITIONAL(GraphBuilder, ViewCount > 1, "View%d", ViewIndex);
+		const FViewInfo& View = Views[ViewIndex];
+
+		FRenderLightParameters* PassParameters = GraphBuilder.AllocParameters<FRenderLightParameters>();
+		GetRenderLightParameters(SceneTextures, ScreenShadowMaskTexture, LightingChannelsTexture, GetCloudShadowAOParameters(GraphBuilder, View, CloudInfo), InHairVisibilityViews, *PassParameters);
+
+		GraphBuilder.AddPass(
+			RDG_EVENT_NAME("StandardDeferredLighting"),
+			PassParameters,
+			PassFlags,
+			[this, &View, ViewIndex, LightSceneInfo, ScreenShadowMaskTexture, LightingChannelsTexture, InHairVisibilityViews, bRenderOverlap](FRHICommandList& RHICmdList)
+		{
+			RenderLight(
+				RHICmdList,
+				View,
+				ViewIndex,
+				LightSceneInfo,
+				TryGetRHI(ScreenShadowMaskTexture),
+				TryGetRHI(LightingChannelsTexture),
+				InHairVisibilityViews,
+				bRenderOverlap,
+				false);
+		});
+	}
 }
 
 BEGIN_SHADER_PARAMETER_STRUCT(FRenderLightForHairParameters, )
@@ -2414,6 +2425,8 @@ void FDeferredShadingSceneRenderer::RenderLightForHair(
 	const FSphere LightBounds = LightSceneInfo->Proxy->GetBoundingSphere();
 	const bool bTransmission = LightSceneInfo->Proxy->Transmission();
 
+	const FVolumetricCloudRenderSceneInfo* CloudInfo = Scene->GetVolumetricCloudSceneInfo();
+
 	for (int32 ViewIndex = 0; ViewIndex < Views.Num(); ViewIndex++)
 	{
 		FViewInfo& View = Views[ViewIndex];
@@ -2432,7 +2445,7 @@ void FDeferredShadingSceneRenderer::RenderLightForHair(
 		}
 
 		FRenderLightForHairParameters* PassParameters = GraphBuilder.AllocParameters<FRenderLightForHairParameters>();
-		GetRenderLightParameters(HairVisibilityData.SampleLightingBuffer, nullptr, SceneTexturesUniformBuffer, HairShadowMaskTexture, LightingChannelsTexture, InHairVisibilityViews, PassParameters->Light);
+		GetRenderLightParameters(HairVisibilityData.SampleLightingBuffer, nullptr, SceneTexturesUniformBuffer, HairShadowMaskTexture, LightingChannelsTexture, GetCloudShadowAOParameters(GraphBuilder, View, CloudInfo), InHairVisibilityViews, PassParameters->Light);
 		PassParameters->HairIndexAndCountTexture = HairVisibilityData.NodeIndex;
 
 		PassParameters->HairTransmittanceMask = InTransmittanceMaskData.TransmittanceMask;
diff --git a/Engine/Source/Runtime/Renderer/Private/ReflectionEnvironmentRealTimeCapture.cpp b/Engine/Source/Runtime/Renderer/Private/ReflectionEnvironmentRealTimeCapture.cpp
--- a/Engine/Source/Runtime/Renderer/Private/ReflectionEnvironmentRealTimeCapture.cpp
+++ b/Engine/Source/Runtime/Renderer/Private/ReflectionEnvironmentRealTimeCapture.cpp
@@ -757,8 +757,8 @@ void FScene::AllocateAndCaptureFrameSkyEnvMap(
 	auto RenderCubeFaces_SpecularConvolution = [&](uint32 CubeMipStart, uint32 CubeMipEnd, uint32 FaceStart, uint32 FaceCount, TRefCountPtr<IPooledRenderTarget>& DstRenderTarget, TRefCountPtr<IPooledRenderTarget>& SrcRenderTarget)
 	{
 		check((FaceStart + FaceCount) <= 6);
-		FRDGTextureRef RDGSrcRenderTarget = GraphBuilder.RegisterExternalTexture(SrcRenderTarget, TEXT("CapturedSkyRenderTarget"));
-		FRDGTextureRef RDGDstRenderTarget = GraphBuilder.RegisterExternalTexture(DstRenderTarget, TEXT("CapturedSkyRenderTarget"));
+		FRDGTextureRef RDGSrcRenderTarget = GraphBuilder.RegisterExternalTexture(SrcRenderTarget);
+		FRDGTextureRef RDGDstRenderTarget = GraphBuilder.RegisterExternalTexture(DstRenderTarget);
 
 		FRDGTextureSRVRef RDGSrcRenderTargetSRV = GraphBuilder.CreateSRV(FRDGTextureSRVDesc::Create(RDGSrcRenderTarget));
 
@@ -796,30 +796,38 @@ void FScene::AllocateAndCaptureFrameSkyEnvMap(
 	auto RenderCubeFaces_DiffuseIrradiance = [&](TRefCountPtr<IPooledRenderTarget>& SourceCubemap)
 	{
 		// ComputeDiffuseIrradiance using N uniform samples
+		AddPass(GraphBuilder, [SkyIrradianceEnvironmentMap = SkyIrradianceEnvironmentMap.Buffer](FRHIComputeCommandList& RHICmdList)
 		{
-			FRDGTextureRef SourceCubemapTexture = GraphBuilder.RegisterExternalTexture(SourceCubemap, TEXT("SourceCubemap"));
-			FRDGTextureSRVRef SourceCubemapTextureSRV = GraphBuilder.CreateSRV(FRDGTextureSRVDesc::Create(SourceCubemapTexture));
+			RHICmdList.Transition({ SkyIrradianceEnvironmentMap, ERHIAccess::Unknown, ERHIAccess::UAVCompute });
+		});
 
-			TShaderMapRef<FComputeSkyEnvMapDiffuseIrradianceCS> ComputeShader(GetGlobalShaderMap(FeatureLevel));
+		FRDGTextureRef SourceCubemapTexture = GraphBuilder.RegisterExternalTexture(SourceCubemap);
+		FRDGTextureSRVRef SourceCubemapTextureSRV = GraphBuilder.CreateSRV(FRDGTextureSRVDesc::Create(SourceCubemapTexture));
 
-			const float SampleCount = FComputeSkyEnvMapDiffuseIrradianceCS::ThreadGroupSizeX * FComputeSkyEnvMapDiffuseIrradianceCS::ThreadGroupSizeY;
-			const float UniformSampleSolidAngle = 4.0f * PI / SampleCount; // uniform distribution
+		TShaderMapRef<FComputeSkyEnvMapDiffuseIrradianceCS> ComputeShader(GetGlobalShaderMap(FeatureLevel));
 
-			FComputeSkyEnvMapDiffuseIrradianceCS::FParameters* PassParameters = GraphBuilder.AllocParameters<FComputeSkyEnvMapDiffuseIrradianceCS::FParameters>();
-			PassParameters->SourceCubemapSampler = TStaticSamplerState<SF_Point>::GetRHI();
-			PassParameters->SourceCubemapTexture = SourceCubemapTextureSRV;
-			PassParameters->OutIrradianceEnvMapSH = SkyIrradianceEnvironmentMap.UAV;
-			PassParameters->UniformSampleSolidAngle = UniformSampleSolidAngle;
-
-			// For 64 uniform samples on the unit sphere, we roughly have 10 samples per face.
-			// Considering mip generation and bilinear sampling, we can assume 10 samples is enough to integrate 10*4=40 texels.
-			// With that, we target integration of 16*16 face.
-			const uint32 Log2_16 = 4; // FMath::Log2(16.0f)
-			PassParameters->MipIndex = uint32(FMath::Log2(float(CapturedSkyRenderTarget->GetDesc().GetSize().X))) - Log2_16;
-
-			const FIntVector NumGroups = FIntVector(1, 1, 1);
-			FComputeShaderUtils::AddPass(GraphBuilder, RDG_EVENT_NAME("ComputeSkyEnvMapDiffuseIrradianceCS"), ComputeShader, PassParameters, NumGroups);
-		}
+		const float SampleCount = FComputeSkyEnvMapDiffuseIrradianceCS::ThreadGroupSizeX * FComputeSkyEnvMapDiffuseIrradianceCS::ThreadGroupSizeY;
+		const float UniformSampleSolidAngle = 4.0f * PI / SampleCount; // uniform distribution
+
+		FComputeSkyEnvMapDiffuseIrradianceCS::FParameters* PassParameters = GraphBuilder.AllocParameters<FComputeSkyEnvMapDiffuseIrradianceCS::FParameters>();
+		PassParameters->SourceCubemapSampler = TStaticSamplerState<SF_Point>::GetRHI();
+		PassParameters->SourceCubemapTexture = SourceCubemapTextureSRV;
+		PassParameters->OutIrradianceEnvMapSH = SkyIrradianceEnvironmentMap.UAV;
+		PassParameters->UniformSampleSolidAngle = UniformSampleSolidAngle;
+
+		// For 64 uniform samples on the unit sphere, we roughly have 10 samples per face.
+		// Considering mip generation and bilinear sampling, we can assume 10 samples is enough to integrate 10*4=40 texels.
+		// With that, we target integration of 16*16 face.
+		const uint32 Log2_16 = 4; // FMath::Log2(16.0f)
+		PassParameters->MipIndex = uint32(FMath::Log2(float(CapturedSkyRenderTarget->GetDesc().GetSize().X))) - Log2_16;
+
+		const FIntVector NumGroups = FIntVector(1, 1, 1);
+		FComputeShaderUtils::AddPass(GraphBuilder, RDG_EVENT_NAME("ComputeSkyEnvMapDiffuseIrradianceCS"), ComputeShader, PassParameters, NumGroups);
+
+		AddPass(GraphBuilder, [SkyIrradianceEnvironmentMap = SkyIrradianceEnvironmentMap.Buffer](FRHICommandList& RHICmdList)
+		{
+			RHICmdList.Transition({ SkyIrradianceEnvironmentMap, ERHIAccess::UAVCompute, ERHIAccess::SRVMask });
+		});
 	};
 
 	const uint32 LastMipLevel = CubeMipCount - 1;
@@ -965,6 +973,11 @@ void FScene::AllocateAndCaptureFrameSkyEnvMap(
 		}
 #endif
 	}
+
+	if (ConvolvedSkyRenderTarget[ConvolvedSkyRenderTargetReadyIndex])
+	{
+		ConvertToUntrackedTexture(GraphBuilder, GraphBuilder.RegisterExternalTexture(ConvolvedSkyRenderTarget[ConvolvedSkyRenderTargetReadyIndex]), ERHIAccess::SRVMask);
+	}
 }
 
 
diff --git a/Engine/Source/Runtime/Renderer/Private/SceneRendering.h b/Engine/Source/Runtime/Renderer/Private/SceneRendering.h
--- a/Engine/Source/Runtime/Renderer/Private/SceneRendering.h
+++ b/Engine/Source/Runtime/Renderer/Private/SceneRendering.h
@@ -1304,10 +1304,10 @@ public:
 	TRefCountPtr<IPooledRenderTarget> SkyAtmosphereViewLutTexture;
 	const FAtmosphereUniformShaderParameters* SkyAtmosphereUniformShaderParameters;
 
-	TRefCountPtr<IPooledRenderTarget> VolumetricCloudSkyAO;
+	FRDGTextureRef VolumetricCloudSkyAO = nullptr;
 	TUniformBufferRef<FViewUniformShaderParameters> VolumetricRenderTargetViewUniformBuffer;
 	// The effective cloud shadow target this frame independently of the fact that a view can have a state (primary view) or not (sky light reflection capture)
-	TRefCountPtr<IPooledRenderTarget> VolumetricCloudShadowRenderTarget[NUM_ATMOSPHERE_LIGHTS];
+	FRDGTextureRef VolumetricCloudShadowRenderTarget[NUM_ATMOSPHERE_LIGHTS] = {};
 
 	/** Used when there is no view state, buffers reallocate every frame. */
 	TUniquePtr<FForwardLightingViewResources> ForwardLightingResourcesStorage;
diff --git a/Engine/Source/Runtime/Renderer/Private/SkyAtmosphereRendering.cpp b/Engine/Source/Runtime/Renderer/Private/SkyAtmosphereRendering.cpp
--- a/Engine/Source/Runtime/Renderer/Private/SkyAtmosphereRendering.cpp
+++ b/Engine/Source/Runtime/Renderer/Private/SkyAtmosphereRendering.cpp
@@ -1407,7 +1407,7 @@ void FSceneRenderer::RenderSkyAtmosphereLookUpTables(FRDGBuilder& GraphBuilder)
 			const FIntVector NumGroups = FIntVector::DivideAndRoundUp(TextureSize, FRenderSkyViewLutCS::GroupSize);
 			FComputeShaderUtils::AddPass(GraphBuilder, RDG_EVENT_NAME("RealTimeCaptureSkyViewLut"), ComputeShader, PassParameters, NumGroups);
 
-			ConvertToExternalTexture(GraphBuilder, RealTimeReflectionCaptureSkyAtmosphereViewLutTexture, Scene->RealTimeReflectionCaptureSkyAtmosphereViewLutTexture);
+			ConvertToUntrackedExternalTexture(GraphBuilder, RealTimeReflectionCaptureSkyAtmosphereViewLutTexture, Scene->RealTimeReflectionCaptureSkyAtmosphereViewLutTexture, ERHIAccess::SRVMask);
 		}
 		else
 		{
@@ -1460,7 +1460,7 @@ void FSceneRenderer::RenderSkyAtmosphereLookUpTables(FRDGBuilder& GraphBuilder)
 			const FIntVector NumGroups = FIntVector::DivideAndRoundUp(TextureSize, FRenderCameraAerialPerspectiveVolumeCS::GroupSize);
 			FComputeShaderUtils::AddPass(GraphBuilder, RDG_EVENT_NAME("RealTimeCaptureCamera360VolumeLut"), ComputeShader, PassParameters, NumGroups);
 
-			ConvertToExternalTexture(GraphBuilder, RealTimeReflectionCaptureCamera360APLutTexture, Scene->RealTimeReflectionCaptureCamera360APLutTexture);
+			ConvertToUntrackedExternalTexture(GraphBuilder, RealTimeReflectionCaptureCamera360APLutTexture, Scene->RealTimeReflectionCaptureCamera360APLutTexture, ERHIAccess::SRVMask);
 		}
 		else
 		{
diff --git a/Engine/Source/Runtime/Renderer/Private/TranslucentLighting.cpp b/Engine/Source/Runtime/Renderer/Private/TranslucentLighting.cpp
--- a/Engine/Source/Runtime/Renderer/Private/TranslucentLighting.cpp
+++ b/Engine/Source/Runtime/Renderer/Private/TranslucentLighting.cpp
@@ -746,8 +746,8 @@ public:
 		{
 			FVolumetricCloudRenderSceneInfo* CloudInfo = LightSceneInfo->Scene->GetVolumetricCloudSceneInfo();
 
-			const bool bLight0CloudPerPixelTransmittance = CloudInfo && View.VolumetricCloudShadowRenderTarget[0].IsValid() && AtmosphereLight0Proxy && AtmosphereLight0Proxy == LightSceneInfo->Proxy;
-			const bool bLight1CloudPerPixelTransmittance = CloudInfo && View.VolumetricCloudShadowRenderTarget[1].IsValid() && AtmosphereLight1Proxy && AtmosphereLight1Proxy == LightSceneInfo->Proxy;
+			const bool bLight0CloudPerPixelTransmittance = CloudInfo && View.VolumetricCloudShadowRenderTarget[0] != nullptr && AtmosphereLight0Proxy && AtmosphereLight0Proxy == LightSceneInfo->Proxy;
+			const bool bLight1CloudPerPixelTransmittance = CloudInfo && View.VolumetricCloudShadowRenderTarget[1] != nullptr && AtmosphereLight1Proxy && AtmosphereLight1Proxy == LightSceneInfo->Proxy;
 
 			if (bLight0CloudPerPixelTransmittance || bLight1CloudPerPixelTransmittance)
 			{
@@ -761,7 +761,7 @@ public:
 					VolumetricCloudShadowmapTexture,
 					VolumetricCloudShadowmapTextureSampler,
 					TStaticSamplerState<SF_Bilinear, AM_Clamp, AM_Clamp, AM_Clamp>::GetRHI(),
-					View.VolumetricCloudShadowRenderTarget[LightIndex]->GetRenderTargetItem().ShaderResourceTexture);
+					View.VolumetricCloudShadowRenderTarget[LightIndex]->GetRHI());
 
 				if (bLight0CloudPerPixelTransmittance)
 				{
@@ -1233,6 +1233,7 @@ static FRDGTextureRef GetSkyTransmittanceLutTexture(FRDGBuilder& GraphBuilder, c
 }
 
 BEGIN_SHADER_PARAMETER_STRUCT(FInjectTranslucentLightArrayParameters, )
+	SHADER_PARAMETER_STRUCT_INCLUDE(FVolumetricCloudShadowAOParameters, CloudShadowAO)
 	RDG_TEXTURE_ACCESS(TransmittanceLutTexture, ERHIAccess::SRVGraphics)
 	RDG_TEXTURE_ACCESS(ShadowDepthTexture, ERHIAccess::SRVGraphics)
 	RENDER_TARGET_BINDING_SLOTS()
@@ -1249,8 +1250,9 @@ static void InjectTranslucentLightArray(
 {
 	INC_DWORD_STAT_BY(STAT_NumLightsInjectedIntoTranslucency, LightInjectionData.Num());
 
-	FRDGTextureRef TransmittanceLutTexture = GetSkyTransmittanceLutTexture(GraphBuilder, Scene, View);
+	const FVolumetricCloudShadowAOParameters CloudShadowAOParameters = GetCloudShadowAOParameters(GraphBuilder, View, Scene->GetVolumetricCloudSceneInfo());
 
+	FRDGTextureRef TransmittanceLutTexture = GetSkyTransmittanceLutTexture(GraphBuilder, Scene, View);
 
 	// Inject into each volume cascade. Operate on one cascade at a time to reduce render target switches.
 	for (uint32 VolumeCascadeIndex = 0; VolumeCascadeIndex < TVC_MAX; VolumeCascadeIndex++)
@@ -1282,6 +1284,7 @@ static void InjectTranslucentLightArray(
 				auto* PassParameters = GraphBuilder.AllocParameters<FInjectTranslucentLightArrayParameters>();
 				PassParameters->TransmittanceLutTexture = TransmittanceLutTexture;
 				PassParameters->ShadowDepthTexture = ShadowDepthTexture;
+				PassParameters->CloudShadowAO = CloudShadowAOParameters;
 				PassParameters->RenderTargets[0] = FRenderTargetBinding(VolumeAmbientTexture, ERenderTargetLoadAction::ELoad);
 				PassParameters->RenderTargets[1] = FRenderTargetBinding(VolumeDirectionalTexture, ERenderTargetLoadAction::ELoad);
 
diff --git a/Engine/Source/Runtime/Renderer/Private/TranslucentLighting.h b/Engine/Source/Runtime/Renderer/Private/TranslucentLighting.h
--- a/Engine/Source/Runtime/Renderer/Private/TranslucentLighting.h
+++ b/Engine/Source/Runtime/Renderer/Private/TranslucentLighting.h
@@ -6,6 +6,7 @@
 #include "RenderGraph.h"
 
 struct FSortedLightSceneInfo;
+struct FCloudShadowAOData;
 
 class FLightSceneInfo;
 class FProjectedShadowInfo;
diff --git a/Engine/Source/Runtime/Renderer/Private/VolumetricCloudRendering.cpp b/Engine/Source/Runtime/Renderer/Private/VolumetricCloudRendering.cpp
--- a/Engine/Source/Runtime/Renderer/Private/VolumetricCloudRendering.cpp
+++ b/Engine/Source/Runtime/Renderer/Private/VolumetricCloudRendering.cpp
@@ -354,18 +354,31 @@ static bool ShouldUsePerSampleAtmosphereTransmittance(const FScene* Scene, const
 
 //////////////////////////////////////////////////////////////////////////
 
+FVolumetricCloudShadowAOParameters GetCloudShadowAOParameters(FRDGBuilder& GraphBuilder, const FViewInfo& View, const FVolumetricCloudRenderSceneInfo* CloudInfo)
+{
+	FVolumetricCloudShadowAOParameters Parameters;
+	if (CloudInfo)
+	{
+		Parameters.ShadowMap0 = View.VolumetricCloudShadowRenderTarget[0];
+		Parameters.ShadowMap1 = View.VolumetricCloudShadowRenderTarget[1];
+		Parameters.SkyAO = View.VolumetricCloudSkyAO;
+	}
+	return Parameters;
+}
 
-void GetCloudShadowAOData(FVolumetricCloudRenderSceneInfo* CloudInfo, FViewInfo& View, FRDGBuilder& GraphBuilder, FCloudShadowAOData& OutData)
+void GetCloudShadowAOData(const FVolumetricCloudRenderSceneInfo* CloudInfo, const FViewInfo& View, FRDGBuilder& GraphBuilder, FCloudShadowAOData& OutData)
 {
+	const FRDGSystemTextures& SystemTextures = FRDGSystemTextures::Get(GraphBuilder);
+
 	// We pick up the texture if they exists, the decision has been mande to render them before already.
-	const bool VolumetricCloudShadowMap0Valid = View.VolumetricCloudShadowRenderTarget[0].IsValid();
-	const bool VolumetricCloudShadowMap1Valid = View.VolumetricCloudShadowRenderTarget[1].IsValid();
+	const bool VolumetricCloudShadowMap0Valid = View.VolumetricCloudShadowRenderTarget[0] != nullptr;
+	const bool VolumetricCloudShadowMap1Valid = View.VolumetricCloudShadowRenderTarget[1] != nullptr;
 	OutData.bShouldSampleCloudShadow = CloudInfo && (VolumetricCloudShadowMap0Valid || VolumetricCloudShadowMap1Valid);
-	OutData.VolumetricCloudShadowMap[0] = OutData.bShouldSampleCloudShadow && VolumetricCloudShadowMap0Valid ? GraphBuilder.RegisterExternalTexture(View.VolumetricCloudShadowRenderTarget[0]) : GSystemTextures.GetBlackDummy(GraphBuilder);
-	OutData.VolumetricCloudShadowMap[1] = OutData.bShouldSampleCloudShadow && VolumetricCloudShadowMap1Valid ? GraphBuilder.RegisterExternalTexture(View.VolumetricCloudShadowRenderTarget[1]) : GSystemTextures.GetBlackDummy(GraphBuilder);
+	OutData.VolumetricCloudShadowMap[0] = OutData.bShouldSampleCloudShadow && VolumetricCloudShadowMap0Valid ? View.VolumetricCloudShadowRenderTarget[0] : SystemTextures.Black;
+	OutData.VolumetricCloudShadowMap[1] = OutData.bShouldSampleCloudShadow && VolumetricCloudShadowMap1Valid ? View.VolumetricCloudShadowRenderTarget[1] : SystemTextures.Black;
 
-	OutData.bShouldSampleCloudSkyAO = CloudInfo && View.VolumetricCloudSkyAO.IsValid();
-	OutData.VolumetricCloudSkyAO = OutData.bShouldSampleCloudSkyAO ? GraphBuilder.RegisterExternalTexture(View.VolumetricCloudSkyAO) : GSystemTextures.GetBlackDummy(GraphBuilder);
+	OutData.bShouldSampleCloudSkyAO = CloudInfo && View.VolumetricCloudSkyAO != nullptr;
+	OutData.VolumetricCloudSkyAO = OutData.bShouldSampleCloudSkyAO ? View.VolumetricCloudSkyAO : SystemTextures.Black;
 
 	// We also force disable shadows if the VolumetricCloud component itself has not initialised due to extra reasons in ShouldRenderVolumetricCloud().
 	// The below test is a simple way to check the fact that if ShouldRenderVolumetricCloud is false, then InitVolumetricCloudsForViews is not run
@@ -1264,7 +1277,7 @@ void FSceneRenderer::InitVolumetricCloudsForViews(FRDGBuilder& GraphBuilder, boo
 				{
 					for (int LightIndex = 0; LightIndex < NUM_ATMOSPHERE_LIGHTS; ++LightIndex)
 					{
-						ViewInfo.VolumetricCloudShadowRenderTarget[LightIndex].SafeRelease();
+						ViewInfo.VolumetricCloudShadowRenderTarget[LightIndex] = nullptr;
 						ViewInfo.ViewState->VolumetricCloudShadowRenderTarget[LightIndex].Reset();
 					}
 				}
@@ -1591,7 +1604,7 @@ void FSceneRenderer::InitVolumetricCloudsForViews(FRDGBuilder& GraphBuilder, boo
 							FilterTracedCloudTexture(&CloudSkyAOTexture, VolumetricCloudParamsAO.VolumetricCloud.CloudSkyAOSizeInvSize, CloudAOTextureTexelWorldSizeInvSize, true);
 						}
 
-						ConvertToExternalTexture(GraphBuilder, CloudSkyAOTexture, ViewInfo.VolumetricCloudSkyAO);
+						ViewInfo.VolumetricCloudSkyAO = CloudSkyAOTexture;
 					}
 
 
@@ -1682,11 +1695,11 @@ void FSceneRenderer::InitVolumetricCloudsForViews(FRDGBuilder& GraphBuilder, boo
 									{
 										FilterTracedCloudTexture(&SpatiallyFilteredShadowTexture, VolumetricCloudParams.VolumetricCloud.CloudShadowmapSizeInvSize[LightIndex], CloudShadowTextureTexelWorldSizeInvSize, false);
 									}
-									ConvertToExternalTexture(GraphBuilder, SpatiallyFilteredShadowTexture, ViewInfo.VolumetricCloudShadowRenderTarget[LightIndex]);
+									ViewInfo.VolumetricCloudShadowRenderTarget[LightIndex] = SpatiallyFilteredShadowTexture;
 								}
 								else
 								{
-									ViewInfo.VolumetricCloudShadowRenderTarget[LightIndex] = CloudShadowTemporalRT.CurrentRenderTarget();
+									ViewInfo.VolumetricCloudShadowRenderTarget[LightIndex] = GraphBuilder.RegisterExternalTexture(CloudShadowTemporalRT.CurrentRenderTarget());
 								}
 
 							}
@@ -1701,11 +1714,11 @@ void FSceneRenderer::InitVolumetricCloudsForViews(FRDGBuilder& GraphBuilder, boo
 									{
 										FilterTracedCloudTexture(&SpatiallyFilteredShadowTexture, VolumetricCloudParams.VolumetricCloud.CloudShadowmapSizeInvSize[LightIndex], CloudShadowTextureTexelWorldSizeInvSize, false);
 									}
-									ConvertToExternalTexture(GraphBuilder, SpatiallyFilteredShadowTexture, ViewInfo.VolumetricCloudShadowRenderTarget[LightIndex]);
+									ViewInfo.VolumetricCloudShadowRenderTarget[LightIndex] = SpatiallyFilteredShadowTexture;
 								}
 								else
 								{
-									ConvertToExternalTexture(GraphBuilder, NewCloudShadowTexture, ViewInfo.VolumetricCloudShadowRenderTarget[LightIndex]);
+									ViewInfo.VolumetricCloudShadowRenderTarget[LightIndex] = NewCloudShadowTexture;
 								}
 							}
 						}
@@ -1713,7 +1726,7 @@ void FSceneRenderer::InitVolumetricCloudsForViews(FRDGBuilder& GraphBuilder, boo
 						{
 							if (ViewInfo.ViewState)
 							{
-								ViewInfo.VolumetricCloudShadowRenderTarget[LightIndex].SafeRelease();
+								ViewInfo.VolumetricCloudShadowRenderTarget[LightIndex] = nullptr;
 								ViewInfo.ViewState->VolumetricCloudShadowRenderTarget[LightIndex].Reset();
 							}
 						}
@@ -2275,11 +2288,11 @@ bool FSceneRenderer::RenderVolumetricCloud(
 						DebugCloudTexture(Parameters);
 					}
 
-					if (DebugCloudSkyAO && ViewInfo.VolumetricCloudSkyAO.IsValid())
+					if (DebugCloudSkyAO && ViewInfo.VolumetricCloudSkyAO != nullptr)
 					{
 						DrawFrustumWireframe(&ShadowFrustumPDI, VolumetricCloudParams.VolumetricCloud.CloudSkyAOWorldToLightClipMatrixInv, FColor::Blue, 0); 
 						FDrawDebugCloudShadowCS::FParameters* Parameters = GraphBuilder.AllocParameters<FDrawDebugCloudShadowCS::FParameters>();
-						Parameters->CloudTracedTexture = GraphBuilder.RegisterExternalTexture(ViewInfo.VolumetricCloudSkyAO);
+						Parameters->CloudTracedTexture = ViewInfo.VolumetricCloudSkyAO;
 						Parameters->CloudTextureSizeInvSize = VolumetricCloudParams.VolumetricCloud.CloudSkyAOSizeInvSize;
 						Parameters->CloudTraceDirection = VolumetricCloudParams.VolumetricCloud.CloudSkyAOTraceDir;
 						Parameters->CloudWorldToLightClipMatrixInv = VolumetricCloudParams.VolumetricCloud.CloudSkyAOWorldToLightClipMatrixInv;
diff --git a/Engine/Source/Runtime/Renderer/Private/VolumetricCloudRendering.h b/Engine/Source/Runtime/Renderer/Private/VolumetricCloudRendering.h
--- a/Engine/Source/Runtime/Renderer/Private/VolumetricCloudRendering.h
+++ b/Engine/Source/Runtime/Renderer/Private/VolumetricCloudRendering.h
@@ -133,7 +133,13 @@ struct FCloudRenderContext
 private:
 };
 
+BEGIN_SHADER_PARAMETER_STRUCT(FVolumetricCloudShadowAOParameters, )
+	RDG_TEXTURE_ACCESS(ShadowMap0, ERHIAccess::SRVGraphics)
+	RDG_TEXTURE_ACCESS(ShadowMap1, ERHIAccess::SRVGraphics)
+	RDG_TEXTURE_ACCESS(SkyAO, ERHIAccess::SRVGraphics)
+END_SHADER_PARAMETER_STRUCT()
 
+FVolumetricCloudShadowAOParameters GetCloudShadowAOParameters(FRDGBuilder& GraphBuilder, const FViewInfo& View, const FVolumetricCloudRenderSceneInfo* CloudInfo);
 
 struct FCloudShadowAOData
 {
@@ -143,6 +149,6 @@ struct FCloudShadowAOData
 	FRDGTextureRef VolumetricCloudSkyAO;
 };
 
-void GetCloudShadowAOData(FVolumetricCloudRenderSceneInfo* CloudInfo, FViewInfo& View, FRDGBuilder& GraphBuilder, FCloudShadowAOData& OutData);
+void GetCloudShadowAOData(const FVolumetricCloudRenderSceneInfo* CloudInfo, const FViewInfo& View, FRDGBuilder& GraphBuilder, FCloudShadowAOData& OutData);
 
 
