Fix pathological perf when processing a few large instance ranges to invalidate VSM cache.

#rnx

diff --git a/Engine/Shaders/Private/VirtualShadowMaps/CacheManagement.usf b/Engine/Shaders/Private/VirtualShadowMaps/CacheManagement.usf
--- a/Engine/Shaders/Private/VirtualShadowMaps/CacheManagement.usf
+++ b/Engine/Shaders/Private/VirtualShadowMaps/CacheManagement.usf
@@ -29,17 +29,23 @@ uint InstanceDataSOAStride;
 [numthreads(CS_1D_GROUP_SIZE_X, 1, 1)]
 void VirtualSmInvalidateInstancePagesCS(
 	uint DispatchIndex : SV_DispatchThreadID,
-	uint GroupIndex : SV_GroupIndex)
+	uint GroupId : SV_GroupID,
+	uint GroupThreadIndex : SV_GroupIndex)
 {
 	bool bIsVisible = false;
 	bool bWasOccluded = false;
-
+#if PROCESS_LARGE_INSTANCE_COUNT_RANGES
+	// Large instance count ranges use one group per range instead, and use a group-stride loop
+	FInstanceDataRange InstanceDataRange = InstanceRanges[GroupId];
+	for (int InstanceId = InstanceDataRange.InstanceDataOffset + GroupThreadIndex; InstanceId < InstanceDataRange.NumInstanceDataEntries + InstanceDataRange.InstanceDataOffset; InstanceId += CS_1D_GROUP_SIZE_X)
+#else // PROCESS_LARGE_INSTANCE_COUNT_RANGES
 	if (DispatchIndex >= NumRemovedItems)
 	{
 		return;
 	}
 	FInstanceDataRange InstanceDataRange = InstanceRanges[DispatchIndex];
 	for (int InstanceId = InstanceDataRange.InstanceDataOffset; InstanceId < InstanceDataRange.NumInstanceDataEntries + InstanceDataRange.InstanceDataOffset; ++InstanceId)
+#endif // PROCESS_LARGE_INSTANCE_COUNT_RANGES
 	{
 		FInstanceSceneData InstanceData = GetInstanceData(InstanceId, InstanceDataSOAStride);
 
diff --git a/Engine/Source/Runtime/Renderer/Private/VirtualShadowMaps/VirtualShadowMapCacheManager.cpp b/Engine/Source/Runtime/Renderer/Private/VirtualShadowMaps/VirtualShadowMapCacheManager.cpp
--- a/Engine/Source/Runtime/Renderer/Private/VirtualShadowMaps/VirtualShadowMapCacheManager.cpp
+++ b/Engine/Source/Runtime/Renderer/Private/VirtualShadowMaps/VirtualShadowMapCacheManager.cpp
@@ -304,15 +304,24 @@ void FVirtualShadowMapArrayCacheManager::ProcessRemovedPrimives(FRDGBuilder& Gra
 	if (CVarCacheVirtualSMs.GetValueOnRenderThread() != 0 && RemovedPrimitiveSceneInfos.Num() > 0 && PrevDynamicCasterPageFlags.IsValid())
 	{
 		// Note: Could filter out primitives that have no nanite here (though later this might be bad anyway, when regular geo is also rendered into virtual SMs)
-		TArray<FInstanceDataRange> InstanceRanges;
+		TArray<FInstanceDataRange> InstanceRangesLarge;
+		TArray<FInstanceDataRange> InstanceRangesSmall;
 		for (const FPrimitiveSceneInfo* PrimitiveSceneInfo : RemovedPrimitiveSceneInfos)
 		{
 			if (PrimitiveSceneInfo->GetInstanceDataOffset() != INDEX_NONE)
 			{
-				InstanceRanges.Add(FInstanceDataRange{ PrimitiveSceneInfo->GetInstanceDataOffset(), PrimitiveSceneInfo->GetNumInstanceDataEntries() });
+				int32 NumInstanceDataEntries = PrimitiveSceneInfo->GetNumInstanceDataEntries();
+				if (NumInstanceDataEntries >= 8U)
+				{
+					InstanceRangesLarge.Add(FInstanceDataRange{ PrimitiveSceneInfo->GetInstanceDataOffset(), NumInstanceDataEntries });
+				}
+				else
+				{
+					InstanceRangesSmall.Add(FInstanceDataRange{ PrimitiveSceneInfo->GetInstanceDataOffset(), NumInstanceDataEntries });
+				}
 			}
 		}
-		ProcessInstanceRangeInvalidation(GraphBuilder, InstanceRanges, GPUScene);
+		ProcessInstanceRangeInvalidation(GraphBuilder, InstanceRangesLarge, InstanceRangesSmall, GPUScene);
 	}
 }
 
@@ -323,7 +332,8 @@ void FVirtualShadowMapArrayCacheManager::ProcessPrimitivesToUpdate(FRDGBuilder&
 	if (IsValid() && GPUScene.PrimitivesToUpdate.Num() > 0)
 	{
 		// TODO: As a slight CPU optimization just pass primitive ID list and use instance ranges stored in GPU scene
-		TArray<FInstanceDataRange> InstanceRanges;
+		TArray<FInstanceDataRange> InstanceRangesLarge;
+		TArray<FInstanceDataRange> InstanceRangesSmall;
 		for (const int32 PrimitiveId : GPUScene.PrimitivesToUpdate)
 		{
 			// Skip added ones (they dont need it, but must be marked as having moved).
@@ -332,11 +342,19 @@ void FVirtualShadowMapArrayCacheManager::ProcessPrimitivesToUpdate(FRDGBuilder&
 				const FPrimitiveSceneInfo* PrimitiveSceneInfo = Scene.Primitives[PrimitiveId];
 				if (PrimitiveSceneInfo->GetInstanceDataOffset() != INDEX_NONE)
 				{
-					InstanceRanges.Add(FInstanceDataRange{ PrimitiveSceneInfo->GetInstanceDataOffset(), PrimitiveSceneInfo->GetNumInstanceDataEntries() });
+					int32 NumInstanceDataEntries = PrimitiveSceneInfo->GetNumInstanceDataEntries();
+					if (NumInstanceDataEntries >= 8U)
+					{
+						InstanceRangesLarge.Add(FInstanceDataRange{ PrimitiveSceneInfo->GetInstanceDataOffset(), NumInstanceDataEntries });
+					}
+					else
+					{
+						InstanceRangesSmall.Add(FInstanceDataRange{ PrimitiveSceneInfo->GetInstanceDataOffset(), NumInstanceDataEntries });
+					}
 				}
 			}
 		}
-		ProcessInstanceRangeInvalidation(GraphBuilder, InstanceRanges, GPUScene);
+		ProcessInstanceRangeInvalidation(GraphBuilder, InstanceRangesLarge, InstanceRangesSmall, GPUScene);
 	}
 }
 
@@ -349,6 +367,10 @@ class FVirtualSmInvalidateInstancePagesCS : public FGlobalShader
 {
 	DECLARE_GLOBAL_SHADER(FVirtualSmInvalidateInstancePagesCS);
 	SHADER_USE_PARAMETER_STRUCT(FVirtualSmInvalidateInstancePagesCS, FGlobalShader)
+
+	class FLargeSmallDim : SHADER_PERMUTATION_BOOL("PROCESS_LARGE_INSTANCE_COUNT_RANGES");
+	using FPermutationDomain = TShaderPermutationDomain<FLargeSmallDim>;
+
 public:
 	BEGIN_SHADER_PARAMETER_STRUCT(FParameters, )
 		SHADER_PARAMETER_RDG_UNIFORM_BUFFER(FVirtualShadowMapCommonParameters, VirtualSmCommon)
@@ -388,11 +410,11 @@ public:
 IMPLEMENT_GLOBAL_SHADER(FVirtualSmInvalidateInstancePagesCS, "/Engine/Private/VirtualShadowMaps/CacheManagement.usf", "VirtualSmInvalidateInstancePagesCS", SF_Compute);
 
 
-void FVirtualShadowMapArrayCacheManager::ProcessInstanceRangeInvalidation(FRDGBuilder& GraphBuilder, const TArray<FInstanceDataRange>& InstanceRanges, const FGPUScene& GPUScene)
+void FVirtualShadowMapArrayCacheManager::ProcessInstanceRangeInvalidation(FRDGBuilder& GraphBuilder, const TArray<FInstanceDataRange>& InstanceRangesLarge, const TArray<FInstanceDataRange>& InstanceRangesSmall, const FGPUScene& GPUScene)
 {
-	if (InstanceRanges.Num())
+	if (InstanceRangesSmall.Num())
 	{
-		RDG_EVENT_SCOPE(GraphBuilder, "ProcessInstanceRangeInvalidation [%d ranges]", InstanceRanges.Num());
+		RDG_EVENT_SCOPE(GraphBuilder, "ProcessInstanceRangeInvalidation [%d small-ranges]", InstanceRangesSmall.Num());
 
 		auto RegExtCreateSrv = [&GraphBuilder](const TRefCountPtr<FRDGPooledBuffer>& Buffer, const TCHAR* Name) -> FRDGBufferSRVRef
 		{
@@ -402,9 +424,9 @@ void FVirtualShadowMapArrayCacheManager::ProcessInstanceRangeInvalidation(FRDGBu
 		FVirtualSmInvalidateInstancePagesCS::FParameters* PassParameters = GraphBuilder.AllocParameters<FVirtualSmInvalidateInstancePagesCS::FParameters>();
 
 		PassParameters->VirtualSmCommon = GraphBuilder.CreateUniformBuffer(&PrevCommonParameters);
-		FRDGBufferRef InstanceRangesRDG = CreateStructuredBuffer(GraphBuilder, TEXT("InstanceRanges"), InstanceRanges);
+		FRDGBufferRef InstanceRangesRDG = CreateStructuredBuffer(GraphBuilder, TEXT("InstanceRangesSmall"), InstanceRangesSmall);
 		PassParameters->InstanceRanges = GraphBuilder.CreateSRV(InstanceRangesRDG);
-		PassParameters->NumRemovedItems = InstanceRanges.Num();
+		PassParameters->NumRemovedItems = InstanceRangesSmall.Num();
 		PassParameters->ShadowMapProjectionData = RegExtCreateSrv(PrevShadowMapProjectionDataBuffer, TEXT("PrevShadowMapProjectionData"));
 
 		PassParameters->PageFlags = RegExtCreateSrv(PrevPageFlags, TEXT("PrevPageFlags"));
@@ -419,14 +441,59 @@ void FVirtualShadowMapArrayCacheManager::ProcessInstanceRangeInvalidation(FRDGBu
 		PassParameters->GPUSceneFrameNumber = GPUScene.GetSceneFrameNumber();
 		PassParameters->InstanceDataSOAStride = GPUScene.InstanceDataSOAStride;
 
-		auto ComputeShader = GetGlobalShaderMap(GMaxRHIFeatureLevel)->GetShader<FVirtualSmInvalidateInstancePagesCS>();
+		FVirtualSmInvalidateInstancePagesCS::FPermutationDomain PermutationVector;
+		PermutationVector.Set<FVirtualSmInvalidateInstancePagesCS::FLargeSmallDim>(0);
+
+		auto ComputeShader = GetGlobalShaderMap(GMaxRHIFeatureLevel)->GetShader<FVirtualSmInvalidateInstancePagesCS>(PermutationVector);
+
+		FComputeShaderUtils::AddPass(
+			GraphBuilder,
+			RDG_EVENT_NAME("VirtualSmInvalidateInstancePagesCS"),
+			ComputeShader,
+			PassParameters,
+			FIntVector(FMath::DivideAndRoundUp(InstanceRangesSmall.Num(), FVirtualSmInvalidateInstancePagesCS::Cs1dGroupSizeX), 1, 1)
+		);
+	}
+	if (InstanceRangesLarge.Num())
+	{
+		RDG_EVENT_SCOPE(GraphBuilder, "ProcessInstanceRangeInvalidation [%d large-ranges]", InstanceRangesLarge.Num());
+
+		auto RegExtCreateSrv = [&GraphBuilder](const TRefCountPtr<FRDGPooledBuffer>& Buffer, const TCHAR* Name) -> FRDGBufferSRVRef
+		{
+			return GraphBuilder.CreateSRV(GraphBuilder.RegisterExternalBuffer(Buffer, Name));
+		};
+
+		FVirtualSmInvalidateInstancePagesCS::FParameters* PassParameters = GraphBuilder.AllocParameters<FVirtualSmInvalidateInstancePagesCS::FParameters>();
+
+		PassParameters->VirtualSmCommon = GraphBuilder.CreateUniformBuffer(&PrevCommonParameters);
+		FRDGBufferRef InstanceRangesRDG = CreateStructuredBuffer(GraphBuilder, TEXT("InstanceRangesSmall"), InstanceRangesLarge);
+		PassParameters->InstanceRanges = GraphBuilder.CreateSRV(InstanceRangesRDG);
+		PassParameters->NumRemovedItems = InstanceRangesLarge.Num();
+		PassParameters->ShadowMapProjectionData = RegExtCreateSrv(PrevShadowMapProjectionDataBuffer, TEXT("PrevShadowMapProjectionData"));
+
+		PassParameters->PageFlags = RegExtCreateSrv(PrevPageFlags, TEXT("PrevPageFlags"));
+		PassParameters->HPageFlags = RegExtCreateSrv(PrevHPageFlags, TEXT("PrevHPageFlags"));
+		PassParameters->PageRectBounds = RegExtCreateSrv(PrevPageRectBounds, TEXT("PrevPageRectBounds"));
+
+		FRDGBufferRef DynamicCasterFlagsRDG = GraphBuilder.RegisterExternalBuffer(PrevDynamicCasterPageFlags, TEXT("DynamicCasterFlags"));
+		PassParameters->OutDynamicCasterFlags = GraphBuilder.CreateUAV(DynamicCasterFlagsRDG);
+
+		PassParameters->GPUSceneInstanceSceneData = GPUScene.InstanceDataBuffer.SRV;
+		PassParameters->GPUScenePrimitiveSceneData = GPUScene.PrimitiveBuffer.SRV;
+		PassParameters->GPUSceneFrameNumber = GPUScene.GetSceneFrameNumber();
+		PassParameters->InstanceDataSOAStride = GPUScene.InstanceDataSOAStride;
+
+		FVirtualSmInvalidateInstancePagesCS::FPermutationDomain PermutationVector;
+		PermutationVector.Set<FVirtualSmInvalidateInstancePagesCS::FLargeSmallDim>(1);
+
+		auto ComputeShader = GetGlobalShaderMap(GMaxRHIFeatureLevel)->GetShader<FVirtualSmInvalidateInstancePagesCS>(PermutationVector);
 
 		FComputeShaderUtils::AddPass(
 			GraphBuilder,
 			RDG_EVENT_NAME("VirtualSmInvalidateInstancePagesCS"),
 			ComputeShader,
 			PassParameters,
-			FIntVector(FMath::DivideAndRoundUp(InstanceRanges.Num(), FVirtualSmInvalidateInstancePagesCS::Cs1dGroupSizeX), 1, 1)
+			FIntVector(InstanceRangesLarge.Num(), 1, 1)
 		);
 	}
 }
diff --git a/Engine/Source/Runtime/Renderer/Private/VirtualShadowMaps/VirtualShadowMapCacheManager.h b/Engine/Source/Runtime/Renderer/Private/VirtualShadowMaps/VirtualShadowMapCacheManager.h
--- a/Engine/Source/Runtime/Renderer/Private/VirtualShadowMaps/VirtualShadowMapCacheManager.h
+++ b/Engine/Source/Runtime/Renderer/Private/VirtualShadowMaps/VirtualShadowMapCacheManager.h
@@ -163,5 +163,5 @@ protected:
 		int32 NumInstanceDataEntries;
 	};
 
-	void ProcessInstanceRangeInvalidation(FRDGBuilder& GraphBuilder, const TArray<FInstanceDataRange>& InstanceRanges, const FGPUScene& GPUScene);
+	void ProcessInstanceRangeInvalidation(FRDGBuilder& GraphBuilder, const TArray<FInstanceDataRange>& InstanceRangesLarge, const TArray<FInstanceDataRange>& InstanceRangesSmall, const FGPUScene& GPUScene);
 };
