Fix DrawTileMesh instance data setup & instance SOAStride for override buffer.

#rnx

diff --git a/Engine/Source/Runtime/Renderer/Private/Renderer.cpp b/Engine/Source/Runtime/Renderer/Private/Renderer.cpp
--- a/Engine/Source/Runtime/Renderer/Private/Renderer.cpp
+++ b/Engine/Source/Runtime/Renderer/Private/Renderer.cpp
@@ -145,10 +145,27 @@ void FRendererModule::DrawTileMesh(FRHICommandListImmediate& RHICmdList, FMeshPa
 				PrimitiveParams.LightmapDataIndex = 0;
 				PrimitiveParams.LightmapUVIndex = 0;
 
+				// Set up reference to the single-instance 
+				PrimitiveParams.InstanceDataOffset = 0;
+				PrimitiveParams.NumInstanceDataEntries = 1;
+
+
 				// Now we just need to fill out the first entry of primitive data in a buffer and bind it
 				SinglePrimitiveStructured.PrimitiveSceneData = FPrimitiveSceneShaderData(PrimitiveParams);
 				SinglePrimitiveStructured.ShaderPlatform = View.GetShaderPlatform();
-				SinglePrimitiveStructured.InstanceSceneData = FInstanceSceneShaderData();
+
+				// Also fill out correct single-primitive instance data, derived from the primitive.
+				SinglePrimitiveStructured.InstanceSceneData = FInstanceSceneShaderData(
+					GetInstanceUniformShaderParameters(PrimitiveParams.LocalToWorld, PrimitiveParams.PreviousLocalToWorld, 
+						(PrimitiveParams.LocalObjectBoundsMin + PrimitiveParams.LocalObjectBoundsMax) * 0.5f, 
+						PrimitiveParams.LocalObjectBoundsMax - PrimitiveParams.LocalObjectBoundsMin,
+						PrimitiveParams.NonUniformScale, PrimitiveParams.InvNonUniformScaleAndDeterminantSign, 
+						FVector4(-11.0f, -22.0f, -33.0f, -44.0f), // TODO: Do something cleaner to indicate instanced LM/SM UV bias vs. uniform
+						FNaniteInfo(),
+						0,
+						0xFFFFFFFFu,
+						0.0f
+					));
 
 				// Set up the parameters for the LightmapSceneData from the given LCI data 
 				FPrecomputedLightingUniformParameters LightmapParams;
diff --git a/Engine/Source/Runtime/Renderer/Private/SceneRendering.cpp b/Engine/Source/Runtime/Renderer/Private/SceneRendering.cpp
--- a/Engine/Source/Runtime/Renderer/Private/SceneRendering.cpp
+++ b/Engine/Source/Runtime/Renderer/Private/SceneRendering.cpp
@@ -1751,7 +1751,7 @@ void FViewInfo::SetupUniformBufferParameters(
 		if (InstanceSceneDataOverrideSRV)
 		{
 			ViewUniformShaderParameters.InstanceSceneData = InstanceSceneDataOverrideSRV;
-			ViewUniformShaderParameters.InstanceDataSOAStride = 0;
+			ViewUniformShaderParameters.InstanceDataSOAStride = 1;
 		}
 		else if (Scene && Scene->GPUScene.InstanceDataBuffer.SRV)
 		{
