More non-Nanite VSM fixes

diff --git a/Engine/Shaders/Private/ShadowDepthVertexShader.usf b/Engine/Shaders/Private/ShadowDepthVertexShader.usf
--- a/Engine/Shaders/Private/ShadowDepthVertexShader.usf
+++ b/Engine/Shaders/Private/ShadowDepthVertexShader.usf
@@ -119,7 +119,7 @@ FPageInfo UnPackPageInfo(uint PackedData)
 void ScaleBiasClipToPhysicalSmPage(FNaniteView NaniteView, inout float4 PosCsHInOut, inout float4 ClipPlanesInOut, FPageInfo PageInfo)
 {
 	uint2 VirtPage = PageInfo.VirtualPage;
-	uint2 PhysPage = PassStruct.VirtualSmPageTable[CalcPageOffset(NaniteView.TargetLayerIndex, NaniteView.TargetMipLevel, VirtPage)];
+	uint2 PhysPage = ShadowDecodePageTable(PassStruct.VirtualSmPageTable[CalcPageOffset(NaniteView.TargetLayerIndex, NaniteView.TargetMipLevel, VirtPage)]).PageIndex;
 	float2 VirtPageF = float2(VirtPage);
 	float2 PhysPageF = float2(PhysPage);
 	// TODO: do we ever need to support viewports that are not the full mip level here? If so this needs to be taken into account.
diff --git a/Engine/Source/Runtime/Renderer/Private/ShadowRendering.h b/Engine/Source/Runtime/Renderer/Private/ShadowRendering.h
--- a/Engine/Source/Runtime/Renderer/Private/ShadowRendering.h
+++ b/Engine/Source/Runtime/Renderer/Private/ShadowRendering.h
@@ -213,7 +213,7 @@ BEGIN_GLOBAL_SHADER_PARAMETER_STRUCT(FShadowDepthPassUniformParameters,)
 	SHADER_PARAMETER(int, bRenderToVirtualShadowMap)
 	SHADER_PARAMETER(int, bInstancePerPage)
 
-	SHADER_PARAMETER_RDG_BUFFER_SRV(StructuredBuffer<uint2>, VirtualSmPageTable)
+	SHADER_PARAMETER_RDG_BUFFER_SRV(StructuredBuffer< uint >, VirtualSmPageTable)
 	SHADER_PARAMETER_RDG_BUFFER_SRV(StructuredBuffer< FPackedNaniteView >, PackedNaniteViews)
 	SHADER_PARAMETER_RDG_BUFFER_SRV(StructuredBuffer< uint4 >, PageRectBounds)
 
diff --git a/Engine/Source/Runtime/Renderer/Private/VirtualShadowMaps/VirtualShadowMapArray.cpp b/Engine/Source/Runtime/Renderer/Private/VirtualShadowMaps/VirtualShadowMapArray.cpp
--- a/Engine/Source/Runtime/Renderer/Private/VirtualShadowMaps/VirtualShadowMapArray.cpp
+++ b/Engine/Source/Runtime/Renderer/Private/VirtualShadowMaps/VirtualShadowMapArray.cpp
@@ -579,7 +579,7 @@ class FCalculatePageRectsCS : public FVirtualPageManagementShader
 
 		SHADER_PARAMETER_RDG_BUFFER_SRV(StructuredBuffer< uint >, PageFlags)
 		SHADER_PARAMETER_RDG_BUFFER_UAV(RWStructuredBuffer<FIntVector4>, PageRectBoundsOut)
-		SHADER_PARAMETER_RDG_BUFFER_UAV(RWStructuredBuffer< uint2 >, OutPageTable)
+		SHADER_PARAMETER_RDG_BUFFER_UAV(RWStructuredBuffer< uint >, OutPageTable)
 	END_SHADER_PARAMETER_STRUCT()
 };
 IMPLEMENT_GLOBAL_SHADER(FCalculatePageRectsCS, "/Engine/Private/VirtualShadowMaps/PageManagement.usf", "CalculatePageRects", SF_Compute);
@@ -607,7 +607,7 @@ class FAllocatePagesUsingRectsCS : public FVirtualPageManagementShader
 	BEGIN_SHADER_PARAMETER_STRUCT(FParameters, )
 		SHADER_PARAMETER_RDG_UNIFORM_BUFFER(FVirtualShadowMapCommonParameters, VirtualSmCommon)
 		SHADER_PARAMETER_RDG_BUFFER_UAV(RWStructuredBuffer< uint >, AllocatedPagesOffset)
-		SHADER_PARAMETER_RDG_BUFFER_UAV(RWStructuredBuffer< uint2 >, OutPageTable)
+		SHADER_PARAMETER_RDG_BUFFER_UAV(RWStructuredBuffer< uint >, OutPageTable)
 		SHADER_PARAMETER_RDG_BUFFER_UAV(RWStructuredBuffer< uint >, OutStatsBuffer)
 		SHADER_PARAMETER_RDG_BUFFER_UAV(RWStructuredBuffer< FCachedPageInfo >, OutCachedPageInfos)
 		SHADER_PARAMETER_RDG_BUFFER_UAV(RWStructuredBuffer< uint >, OutPageFlags)
diff --git a/Engine/Source/Runtime/Renderer/Private/VirtualShadowMaps/VirtualShadowMapArray.h b/Engine/Source/Runtime/Renderer/Private/VirtualShadowMaps/VirtualShadowMapArray.h
--- a/Engine/Source/Runtime/Renderer/Private/VirtualShadowMaps/VirtualShadowMapArray.h
+++ b/Engine/Source/Runtime/Renderer/Private/VirtualShadowMaps/VirtualShadowMapArray.h
@@ -145,7 +145,7 @@ END_SHADER_PARAMETER_STRUCT()
  */
 BEGIN_SHADER_PARAMETER_STRUCT(FVirtualShadowMapPageTableParameters, )
 	SHADER_PARAMETER_RDG_UNIFORM_BUFFER(FVirtualShadowMapCommonParameters, VirtualSmCommon)
-	SHADER_PARAMETER_RDG_BUFFER_SRV(StructuredBuffer<uint2>, PageTable)
+	SHADER_PARAMETER_RDG_BUFFER_SRV(StructuredBuffer< uint >, PageTable)
 	SHADER_PARAMETER_RDG_BUFFER_SRV(StructuredBuffer< uint >, PageFlags)
 	SHADER_PARAMETER_RDG_BUFFER_SRV(StructuredBuffer< uint >, HPageFlags)
 	SHADER_PARAMETER_RDG_BUFFER_SRV(StructuredBuffer< uint4 >, PageRectBounds)
