Hook up render state recreate to VSM Cvar (r.Shadow.v.Enable)
- needed to handle z-order change in cached render commands.

diff --git a/Engine/Source/Runtime/RenderCore/Public/RenderUtils.h b/Engine/Source/Runtime/RenderCore/Public/RenderUtils.h
--- a/Engine/Source/Runtime/RenderCore/Public/RenderUtils.h
+++ b/Engine/Source/Runtime/RenderCore/Public/RenderUtils.h
@@ -619,11 +619,11 @@ inline bool UseNanite(EShaderPlatform ShaderPlatform)
 
 inline bool UseVirtualShadowMaps(EShaderPlatform ShaderPlatform, const FStaticFeatureLevel FeatureLevel)
 {
-	static const auto EnableVirtualSMCVar = IConsoleManager::Get().FindTConsoleVariableDataInt(TEXT("r.Shadow.v.Enable"));
+	static const auto EnableVirtualSMCVar = IConsoleManager::Get().FindConsoleVariable(TEXT("r.Shadow.v.Enable"));
 
 	// TODO: Add test for required RHI platform stuff like GPUScene
 	// bool bNaniteSupported = DoesPlatformSupportNanite(ShaderPlatform);
 	bool bForwardShadingEnabled = IsForwardShadingEnabled(ShaderPlatform);
 
-	return UseGPUScene(ShaderPlatform, FeatureLevel) && !bForwardShadingEnabled && EnableVirtualSMCVar->GetValueOnAnyThread() != 0;
+	return UseGPUScene(ShaderPlatform, FeatureLevel) && !bForwardShadingEnabled && EnableVirtualSMCVar->GetInt() != 0;
 }
diff --git a/Engine/Source/Runtime/Renderer/Private/DeferredShadingRenderer.cpp b/Engine/Source/Runtime/Renderer/Private/DeferredShadingRenderer.cpp
--- a/Engine/Source/Runtime/Renderer/Private/DeferredShadingRenderer.cpp
+++ b/Engine/Source/Runtime/Renderer/Private/DeferredShadingRenderer.cpp
@@ -1509,7 +1509,7 @@ void FDeferredShadingSceneRenderer::Render(FRDGBuilder& GraphBuilder)
 	const bool bNaniteEnabled = UseNanite(ShaderPlatform) && ViewFamily.EngineShowFlags.NaniteMeshes;
 
 	// Important that this uses consistent logic for whether or not nanite is enabled, so pass in the flag from here
-	VirtualShadowMapArray.Initialize(bNaniteEnabled);
+	VirtualShadowMapArray.Initialize(UseVirtualShadowMaps(ShaderPlatform, FeatureLevel));
 
 	Scene->UpdateAllPrimitiveSceneInfos(GraphBuilder, true);
 
diff --git a/Engine/Source/Runtime/Renderer/Private/ShadowDepthRendering.cpp b/Engine/Source/Runtime/Renderer/Private/ShadowDepthRendering.cpp
--- a/Engine/Source/Runtime/Renderer/Private/ShadowDepthRendering.cpp
+++ b/Engine/Source/Runtime/Renderer/Private/ShadowDepthRendering.cpp
@@ -833,9 +833,9 @@ void SetStateForShadowDepth(bool bOnePassPointLightShadow, bool bDirectionalLigh
 #if ENABLE_NON_NANITE_VSM
 
 	// GPUCULL_TODO: Unhack me please
-	static const auto EnableVirtualSMCVar = IConsoleManager::Get().FindTConsoleVariableDataInt(TEXT("r.Shadow.v.Enable"));
+	static const auto EnableVirtualSMCVar = IConsoleManager::Get().FindConsoleVariable(TEXT("r.Shadow.v.Enable"));
 
-	const bool bNeedsVirtualShadowMap = EnableVirtualSMCVar->GetValueOnRenderThread() != 0 && bDirectionalLight;
+	const bool bNeedsVirtualShadowMap = EnableVirtualSMCVar->GetInt() != 0 && bDirectionalLight;
 
 	if (bOnePassPointLightShadow || bNeedsVirtualShadowMap)
 #else //!ENABLE_NON_NANITE_VSM
diff --git a/Engine/Source/Runtime/Renderer/Private/VirtualShadowMaps/VirtualShadowMapArray.cpp b/Engine/Source/Runtime/Renderer/Private/VirtualShadowMaps/VirtualShadowMapArray.cpp
--- a/Engine/Source/Runtime/Renderer/Private/VirtualShadowMaps/VirtualShadowMapArray.cpp
+++ b/Engine/Source/Runtime/Renderer/Private/VirtualShadowMaps/VirtualShadowMapArray.cpp
@@ -41,18 +41,21 @@ struct FCachedPageInfo
 	float Padding;
 };
 
-// GPUCULL_TODO: Hook this up again!
 static void RecreateGlobalRenderState(IConsoleVariable* Var)
 {
+#if ENABLE_NON_NANITE_VSM
+	// Needed because the depth state changes with method (so cached draw commands must be re-created) see SetStateForShadowDepth
 	FGlobalComponentRecreateRenderStateContext Context;
+#endif // ENABLE_NON_NANITE_VSM
 }
 
-static TAutoConsoleVariable<int32> CVarEnableVirtualShadowMaps(
+int32 GEnableVirtualShadowMaps = 0;
+FAutoConsoleVariableRef CVarEnableVirtualShadowMaps(
 	TEXT("r.Shadow.v.Enable"),
-	0,
+	GEnableVirtualShadowMaps,
 	TEXT("Enable Virtual Shadow Maps."),
-	//FConsoleVariableDelegate::CreateStatic(&RecreateGlobalRenderState),
-	ECVF_RenderThreadSafe
+	FConsoleVariableDelegate::CreateStatic(&RecreateGlobalRenderState),
+	ECVF_Scalability | ECVF_RenderThreadSafe
 );
 
 
@@ -198,7 +201,7 @@ FVirtualShadowMapArray::~FVirtualShadowMapArray()
 
 void FVirtualShadowMapArray::Initialize(bool bInEnabled)
 {
-	bEnabled = bInEnabled && CVarEnableVirtualShadowMaps.GetValueOnRenderThread() > 0;
+	bEnabled = bInEnabled;
 }
 
 void FVirtualShadowMapArray::SetShaderDefines(FShaderCompilerEnvironment& OutEnvironment)
