Make sure BuildAccelerationStructures is not called on a parallel command list anymore because it can happen otherwise that a build operation is done on the same RTGeometry on multiple threads at the same time

diff --git a/Engine/Source/Runtime/Renderer/Private/RayTracing/RayTracingDynamicGeometry.cpp b/Engine/Source/Runtime/Renderer/Private/RayTracing/RayTracingDynamicGeometry.cpp
--- a/Engine/Source/Runtime/Renderer/Private/RayTracing/RayTracingDynamicGeometry.cpp
+++ b/Engine/Source/Runtime/Renderer/Private/RayTracing/RayTracingDynamicGeometry.cpp
@@ -451,13 +451,6 @@ void FRayTracingDynamicGeometryCollection::DispatchUpdates(FRHIComputeCommandLis
 				RHICmdList.Transition(TransitionsAfter);
 			}
 
-			{
-				FRHIComputeCommandList& RHICmdList = AllocateCommandList(1, GET_STATID(STAT_CLM_RTDynGeomBuild));
-
-				SCOPED_DRAW_OR_COMPUTE_EVENT(RHICmdList, Build);
-				RHICmdList.BuildAccelerationStructures(BuildParams);
-			}
-
 			// Need to kick parallel translate command lists?
 			if (CommandLists.Num() > 0)
 			{
@@ -471,6 +464,17 @@ void FRayTracingDynamicGeometryCollection::DispatchUpdates(FRHIComputeCommandLis
 					false // bSpewMerge
 				);
 			}
+
+			if (BuildParams.Num() > 0)
+			{
+				// Can't use parallel command list because we have to make sure we are not building BVH data
+				// on the same RTGeometry on multiple threads at the same time. Ideally move the build
+				// requests over to the RaytracingGeometry manager so they can be correctly scheduled
+				// with other build requests in the engine (see UE-106982)
+				SCOPED_DRAW_OR_COMPUTE_EVENT(ParentCmdList, Build);
+				ParentCmdList.BuildAccelerationStructures(BuildParams);
+			}
+
 			}
 		}
 		
