Fix static analysis errors for CIS

#rnx

diff --git a/Engine/Source/Runtime/Renderer/Private/GPUScene.cpp b/Engine/Source/Runtime/Renderer/Private/GPUScene.cpp
--- a/Engine/Source/Runtime/Renderer/Private/GPUScene.cpp
+++ b/Engine/Source/Runtime/Renderer/Private/GPUScene.cpp
@@ -913,7 +913,7 @@ void FGPUScene::UploadGeneral(FRHICommandListImmediate& RHICmdList, FScene *Scen
 				RHICmdList.Transition(FRHITransitionInfo(InstanceDataBuffer.UAV, ERHIAccess::Unknown, ERHIAccess::SRVMask));
 			}
 
-			if( Scene->InstanceBVH.GetNumDirty() > 0 )
+			if( Scene != nullptr && Scene->InstanceBVH.GetNumDirty() > 0 )
 			{
 				InstanceUploadBuffer.Init( Scene->InstanceBVH.GetNumDirty(), sizeof( FBVHNode ), true, TEXT("InstanceUploadBuffer") );
 
@@ -1099,7 +1099,7 @@ void FGPUScene::UploadDynamicPrimitiveShaderDataForViewInternal(FRHICommandListI
 	ensure(Collector.UploadData == nullptr || CurrentDynamicContext->DymamicPrimitiveUploadData.Find(Collector.UploadData) != INDEX_NONE);
 	
 	// Skip uploading empty & already uploaded data
-	if (NumPrimitiveDataUploads > 0 && !Collector.UploadData->bIsUploaded)
+	if (Collector.UploadData != nullptr && NumPrimitiveDataUploads > 0 && !Collector.UploadData->bIsUploaded)
 	{
 		Collector.UploadData->bIsUploaded = true;
 		/** Tracks dynamic primitive data for upload to GPU Scene, when enabled. */
