Add fallback path for primitives not supporting instance culling when VSM is enabled
- added separate MeshPass for VSMShadowDepth that collect primitives that support non-nanite VSM (instance culling)

diff --git a/Engine/Source/Runtime/Renderer/Private/PrimitiveSceneInfo.cpp b/Engine/Source/Runtime/Renderer/Private/PrimitiveSceneInfo.cpp
--- a/Engine/Source/Runtime/Renderer/Private/PrimitiveSceneInfo.cpp
+++ b/Engine/Source/Runtime/Renderer/Private/PrimitiveSceneInfo.cpp
@@ -328,7 +328,7 @@ void FPrimitiveSceneInfo::CacheMeshDrawCommands(FRHICommandListImmediate& RHICmd
 
 								ensureMsgf(MeshDrawCommand->VertexStreams.GetAllocatedSize() == 0, TEXT("Cached Mesh Draw command overflows VertexStreams.  VertexStream inline size should be tweaked."));
 
-								if (PassType == EMeshPass::BasePass || PassType == EMeshPass::DepthPass || PassType == EMeshPass::CSMShadowDepth)
+								if (PassType == EMeshPass::BasePass || PassType == EMeshPass::DepthPass || PassType == EMeshPass::CSMShadowDepth || PassType == EMeshPass::VSMShadowDepth)
 								{
 									TArray<EShaderFrequency, TInlineAllocator<SF_NumFrequencies>> ShaderFrequencies;
 									MeshDrawCommand->ShaderBindings.GetShaderFrequencies(ShaderFrequencies);
diff --git a/Engine/Source/Runtime/Renderer/Private/ShadowDepthRendering.cpp b/Engine/Source/Runtime/Renderer/Private/ShadowDepthRendering.cpp
--- a/Engine/Source/Runtime/Renderer/Private/ShadowDepthRendering.cpp
+++ b/Engine/Source/Runtime/Renderer/Private/ShadowDepthRendering.cpp
@@ -825,19 +825,13 @@ void FProjectedShadowInfo::SetStateForView(FRHICommandList& RHICmdList) const
 	);
 }
 
-void SetStateForShadowDepth(bool bOnePassPointLightShadow, bool bDirectionalLight, FMeshPassProcessorRenderState& DrawRenderState)
+void SetStateForShadowDepth(bool bOnePassPointLightShadow, bool bDirectionalLight, FMeshPassProcessorRenderState& DrawRenderState, EMeshPass::Type InMeshPassTargetType)
 {
 	// Disable color writes
 	DrawRenderState.SetBlendState(TStaticBlendState<CW_NONE>::GetRHI());
 
 #if ENABLE_NON_NANITE_VSM
-
-	// GPUCULL_TODO: Unhack me please
-	static const auto EnableVirtualSMCVar = IConsoleManager::Get().FindConsoleVariable(TEXT("r.Shadow.v.Enable"));
-
-	const bool bNeedsVirtualShadowMap = EnableVirtualSMCVar->GetInt() != 0 && bDirectionalLight;
-
-	if (bOnePassPointLightShadow || bNeedsVirtualShadowMap)
+	if (bOnePassPointLightShadow || InMeshPassTargetType == EMeshPass::VSMShadowDepth)
 #else //!ENABLE_NON_NANITE_VSM
 	if (bOnePassPointLightShadow)
 #endif // ENABLE_NON_NANITE_VSM
@@ -1158,7 +1152,7 @@ void FProjectedShadowInfo::RenderDepth(
 	{
 		// Copy in depths of static primitives before we render movable primitives.
 		FMeshPassProcessorRenderState DrawRenderState;
-		SetStateForShadowDepth(bOnePassPointLightShadow, bDirectionalLight, DrawRenderState);
+		SetStateForShadowDepth(bOnePassPointLightShadow, bDirectionalLight, DrawRenderState, MeshPassTargetType);
 		CopyCachedShadowMap(GraphBuilder, *ShadowDepthView, SceneRenderer, PassParameters->RenderTargets, DrawRenderState);
 	}
 #if ENABLE_NON_NANITE_VSM
@@ -2361,7 +2355,34 @@ bool FShadowDepthPassMeshProcessor::TryAddMeshBatch(const FMeshBatch& RESTRICT M
 
 		OverrideWithDefaultMaterialForShadowDepth(EffectiveMaterialRenderProxy, EffectiveMaterial, FeatureLevel);
 
-		bResult = Process(MeshBatch, BatchElementMask, StaticMeshId, PrimitiveSceneProxy, *EffectiveMaterialRenderProxy, *EffectiveMaterial, MeshFillMode, FinalCullMode);
+#if ENABLE_NON_NANITE_VSM
+		// TODO: This uses a lot of indirections and complex logic, optimize by precomputing as far as possible
+		//       E.g., maybe store  bSupportsGpuSceneInstancing as a flag in the MeshBatch?
+		const FVertexFactory* VertexFactory = MeshBatch.VertexFactory;
+		const bool bUsePositionOnlyVS =
+			VertexFactory->SupportsPositionAndNormalOnlyStream()
+			&& EffectiveMaterial->WritesEveryPixel(true)
+			&& !EffectiveMaterial->MaterialModifiesMeshPosition_RenderThread();
+
+		// TODO: Store in MeshBatch?
+		const bool bSupportsGpuSceneInstancing = UseGPUScene(GShaderPlatformForFeatureLevel[FeatureLevel], FeatureLevel)
+			&& VertexFactory->GetPrimitiveIdStreamIndex(bUsePositionOnlyVS ? EVertexInputStreamType::PositionAndNormalOnly : EVertexInputStreamType::Default) != INDEX_NONE;
+
+
+		// EMeshPass::CSMShadowDepth: If no VSM: include everything, else only !bSupportsGpuSceneInstancing
+		// EMeshPass::VSMShadowDepth: If VSM: only bSupportsGpuSceneInstancing, else nothing needs to go in.
+		// TODO: I'm sure I can reduce this logic
+		const bool bDraw = MeshPassTargetType == EMeshPass::CSMShadowDepth ?
+			(!UseVirtualShadowMaps(GShaderPlatformForFeatureLevel[FeatureLevel], FeatureLevel) || !bSupportsGpuSceneInstancing) :
+			(UseVirtualShadowMaps(GShaderPlatformForFeatureLevel[FeatureLevel], FeatureLevel)  && bSupportsGpuSceneInstancing);
+		
+		if (bDraw)
+		{
+#endif // ENABLE_NON_NANITE_VSM
+			bResult = Process(MeshBatch, BatchElementMask, StaticMeshId, PrimitiveSceneProxy, *EffectiveMaterialRenderProxy, *EffectiveMaterial, MeshFillMode, FinalCullMode);
+#if ENABLE_NON_NANITE_VSM
+		}
+#endif // ENABLE_NON_NANITE_VSM
 	}
 
 	return bResult;
@@ -2391,11 +2412,13 @@ FShadowDepthPassMeshProcessor::FShadowDepthPassMeshProcessor(
 	const FScene* Scene,
 	const FSceneView* InViewIfDynamicMeshCommand,
 	FShadowDepthType InShadowDepthType,
-	FMeshPassDrawListContext* InDrawListContext)
+	FMeshPassDrawListContext* InDrawListContext,
+	EMeshPass::Type InMeshPassTargetType)
 	: FMeshPassProcessor(Scene, Scene->GetFeatureLevel(), InViewIfDynamicMeshCommand, InDrawListContext)
 	, ShadowDepthType(InShadowDepthType)
+	, MeshPassTargetType(InMeshPassTargetType)
 {
-	SetStateForShadowDepth(ShadowDepthType.bOnePassPointLightShadow, ShadowDepthType.bDirectionalLight, PassDrawRenderState);
+	SetStateForShadowDepth(ShadowDepthType.bOnePassPointLightShadow, ShadowDepthType.bDirectionalLight, PassDrawRenderState, MeshPassTargetType);
 }
 
 FShadowDepthType CSMShadowDepthType(true, false);
@@ -2406,8 +2429,29 @@ FMeshPassProcessor* CreateCSMShadowDepthPassProcessor(const FScene* Scene, const
 		Scene,
 		InViewIfDynamicMeshCommand,
 		CSMShadowDepthType,
-		InDrawListContext);
+		InDrawListContext,
+		EMeshPass::CSMShadowDepth);
 }
 
 FRegisterPassProcessorCreateFunction RegisterCSMShadowDepthPass(&CreateCSMShadowDepthPassProcessor, EShadingPath::Deferred, EMeshPass::CSMShadowDepth, EMeshPassFlags::CachedMeshCommands);
+
+FMeshPassProcessor* CreateVSMShadowDepthPassProcessor(const FScene* Scene, const FSceneView* InViewIfDynamicMeshCommand, FMeshPassDrawListContext* InDrawListContext)
+{
+#if ENABLE_NON_NANITE_VSM
+	// Only create the mesh pass processor if VSMs are not enabled as this prevents wasting time caching the SM draw commands
+	if (UseVirtualShadowMaps(Scene->GetShaderPlatform(), Scene->GetFeatureLevel()))
+	{
+		return new(FMemStack::Get()) FShadowDepthPassMeshProcessor(
+			Scene,
+			InViewIfDynamicMeshCommand,
+			CSMShadowDepthType,
+			InDrawListContext,
+			EMeshPass::VSMShadowDepth);
+	}
+#endif // ENABLE_NON_NANITE_VSM
+	return nullptr;
+}
+FRegisterPassProcessorCreateFunction RegisterVSMShadowDepthPass(&CreateVSMShadowDepthPassProcessor, EShadingPath::Deferred, EMeshPass::VSMShadowDepth, EMeshPassFlags::CachedMeshCommands);
+
+
 FRegisterPassProcessorCreateFunction RegisterMobileCSMShadowDepthPass(&CreateCSMShadowDepthPassProcessor, EShadingPath::Mobile, EMeshPass::CSMShadowDepth, EMeshPassFlags::CachedMeshCommands);
diff --git a/Engine/Source/Runtime/Renderer/Private/ShadowRendering.h b/Engine/Source/Runtime/Renderer/Private/ShadowRendering.h
--- a/Engine/Source/Runtime/Renderer/Private/ShadowRendering.h
+++ b/Engine/Source/Runtime/Renderer/Private/ShadowRendering.h
@@ -103,7 +103,8 @@ public:
 		const FScene* Scene, 
 		const FSceneView* InViewIfDynamicMeshCommand, 
 		FShadowDepthType InShadowDepthType,
-		FMeshPassDrawListContext* InDrawListContext);
+		FMeshPassDrawListContext* InDrawListContext,
+		EMeshPass::Type InMeshPassTargetType);
 
 	virtual void AddMeshBatch(const FMeshBatch& RESTRICT MeshBatch, uint64 BatchElementMask, const FPrimitiveSceneProxy* RESTRICT PrimitiveSceneProxy, int32 StaticMeshId = -1) override final;
 
@@ -128,6 +129,7 @@ private:
 		ERasterizerCullMode MeshCullMode);
 
 	FShadowDepthType ShadowDepthType;
+	EMeshPass::Type MeshPassTargetType = EMeshPass::CSMShadowDepth;
 };
 
 enum EShadowDepthCacheMode
@@ -385,6 +387,9 @@ public:
 	/** Regular shadow map to resample for use as a complete shadow map */
 	FProjectedShadowInfo *CompleteShadowMapCopySource;
 
+	/** Used to fetch the correct cached static mesh draw commands */
+	EMeshPass::Type MeshPassTargetType = EMeshPass::CSMShadowDepth;
+
 	/** */
 	TArray< FVirtualShadowMap*, TInlineAllocator<6> > VirtualShadowMaps;
 
diff --git a/Engine/Source/Runtime/Renderer/Private/ShadowSetup.cpp b/Engine/Source/Runtime/Renderer/Private/ShadowSetup.cpp
--- a/Engine/Source/Runtime/Renderer/Private/ShadowSetup.cpp
+++ b/Engine/Source/Runtime/Renderer/Private/ShadowSetup.cpp
@@ -1001,6 +1001,7 @@ void FProjectedShadowInfo::SetupClipmapProjection(FLightSceneInfo* InLightSceneI
 	BorderSize = 0;
 	bIncludeInScreenSpaceShadowMask = false;
 	MaxNonFarCascadeDistance = InMaxNonFarCascadeDistance;
+	MeshPassTargetType = EMeshPass::VSMShadowDepth;
 
 
 	FVirtualShadowMapProjectionShaderData Data = VirtualShadowMapClipmap->GetProjectionShaderData(VirtualShadowMapClipmap->GetLevelCount() - 1);
@@ -1292,7 +1293,7 @@ void FProjectedShadowInfo::AddCachedMeshDrawCommands_AnyThread(
 	FAddSubjectPrimitiveStats& OutStats,
 	FAddSubjectPrimitiveOverflowedIndices& OverflowBuffer) const
 {
-	const EMeshPass::Type PassType = EMeshPass::CSMShadowDepth;
+	const EMeshPass::Type PassType = MeshPassTargetType;
 	const EShadingPath ShadingPath = Scene->GetShadingPath();
 	const bool bUseCachedMeshCommand = UseCachedMeshDrawCommands_AnyThread()
 		&& !!(FPassProcessorManager::GetPassFlags(ShadingPath, PassType) & EMeshPassFlags::CachedMeshCommands)
@@ -1395,7 +1396,7 @@ bool FProjectedShadowInfo::ShouldDrawStaticMeshes(FViewInfo& InCurrentView, FPri
 							StaticMeshRelevance,
 							StaticMesh,
 							InPrimitiveSceneInfo->Scene,
-							EMeshPass::CSMShadowDepth,
+							MeshPassTargetType,
 							ShadowDepthPassVisibleCommands,
 							SubjectMeshCommandBuildRequests,
 							NumSubjectMeshCommandBuildRequestElements);
@@ -1952,8 +1953,8 @@ void FProjectedShadowInfo::FinalizeAddSubjectPrimitive(
 			const FCachedMeshDrawCommandInfo& CmdInfo = PrimitiveSceneInfo->StaticMeshCommandInfos[CmdIdx];
 			const FScene* Scene = PrimitiveSceneInfo->Scene;
 			const FMeshDrawCommand* CachedCmd = CmdInfo.StateBucketId >= 0 ?
-				&Scene->CachedMeshDrawCommandStateBuckets[EMeshPass::CSMShadowDepth].GetByElementId(CmdInfo.StateBucketId).Key :
-				&Scene->CachedDrawLists[EMeshPass::CSMShadowDepth].MeshDrawCommands[CmdInfo.CommandIndex];
+				&Scene->CachedMeshDrawCommandStateBuckets[MeshPassTargetType].GetByElementId(CmdInfo.StateBucketId).Key :
+				&Scene->CachedDrawLists[MeshPassTargetType].MeshDrawCommands[CmdInfo.CommandIndex];
 			const int32 PrimIdx = PrimitiveSceneInfo->GetIndex();
 
 			FVisibleMeshDrawCommand& VisibleCmd = ShadowDepthPassVisibleCommands[ShadowDepthPassVisibleCommands.AddUninitialized()];
@@ -2021,7 +2022,8 @@ void FProjectedShadowInfo::SetupMeshDrawCommandsForShadowDepth(FSceneRenderer& R
 		Renderer.Scene,
 		ShadowDepthView,
 		GetShadowDepthType(),
-		nullptr);
+		nullptr,
+		MeshPassTargetType);
 
 	if (Renderer.ShouldDumpMeshDrawCommandInstancingStats())
 	{
@@ -4395,15 +4397,9 @@ void FSceneRenderer::AddViewDependentWholeSceneShadowsForView(
 					FIntPoint CompleteShadowBufferResolution(ShadowBufferResolution.X / CompleteResolutionFactor,
 														     ShadowBufferResolution.Y / CompleteResolutionFactor);
 
-					FProjectedShadowInfo* RegularCascadePSInfo = nullptr;
-
-#if ENABLE_NON_NANITE_VSM
-					if (!bNeedsVirtualShadowMap)
-#endif // ENABLE_NON_NANITE_VSM
-					{
 					// Create the projected shadow info.
-						RegularCascadePSInfo = new(FMemStack::Get(), 1, 16) FProjectedShadowInfo;
-						RegularCascadePSInfo->SetupWholeSceneProjection(
+					FProjectedShadowInfo* ProjectedShadowInfo = new(FMemStack::Get(), 1, 16) FProjectedShadowInfo;
+					ProjectedShadowInfo->SetupWholeSceneProjection(
 						&LightSceneInfo,
 						&View,
 						ProjectedShadowInitializer,
@@ -4414,27 +4410,18 @@ void FSceneRenderer::AddViewDependentWholeSceneShadowsForView(
 						CompleteShadowBufferResolution.Y,
 						ShadowBorder
 					);
-						RegularCascadePSInfo->FadeAlphas = FadeAlphas;
-						RegularCascadePSInfo->ProjectionIndex = Index;
+					ProjectedShadowInfo->FadeAlphas = FadeAlphas;
+					ProjectedShadowInfo->ProjectionIndex = Index;
+
+					VisibleLightInfo.MemStackProjectedShadows.Add(ProjectedShadowInfo);
+					VisibleLightInfo.AllProjectedShadows.Add(ProjectedShadowInfo);
+					ShadowInfos.Add(ProjectedShadowInfo);
 
-						VisibleLightInfo.MemStackProjectedShadows.Add(RegularCascadePSInfo);
-						VisibleLightInfo.AllProjectedShadows.Add(RegularCascadePSInfo);
-						ShadowInfos.Add(RegularCascadePSInfo);
-					}
-#if ENABLE_NON_NANITE_VSM
-					// GPUCULL_TODO: THis logic is not fully functional!
-					// If we have a virtual shadow map, 
-					if (false)//bNeedsVirtualShadowMap)
-					{
-#else //!ENABLE_NON_NANITE_VSM
 					if (bNeedsVirtualShadowMap)
 					{
-#endif // ENABLE_NON_NANITE_VSM
-						if (RegularCascadePSInfo != nullptr)
-						{
-							// ...disable nanite rendering into the regular shadow map or else we'd get double-shadowing
-							RegularCascadePSInfo->bNaniteGeometry = false;
-						}
+						// If we have a virtual shadow map, disable nanite rendering into the regular shadow map or else we'd get double-shadowing
+						ProjectedShadowInfo->bNaniteGeometry = false;
+						
 						if (bNeedsCompleteShadowMap)
 						{
 							// Create the projected shadow info.
@@ -4454,7 +4441,7 @@ void FSceneRenderer::AddViewDependentWholeSceneShadowsForView(
 
 							CompleteProjectedShadowInfo->bCompleteShadowMap = true;
 							CompleteProjectedShadowInfo->bIncludeInScreenSpaceShadowMask = false;
-							CompleteProjectedShadowInfo->CompleteShadowMapCopySource = RegularCascadePSInfo;
+							CompleteProjectedShadowInfo->CompleteShadowMapCopySource = ProjectedShadowInfo;
 
 							VisibleLightInfo.AllProjectedShadows.Add(CompleteProjectedShadowInfo);
 							VisibleLightInfo.CompleteProjectedShadows.Add(CompleteProjectedShadowInfo);
@@ -4463,9 +4450,9 @@ void FSceneRenderer::AddViewDependentWholeSceneShadowsForView(
 					}
 
 					// Ray traced shadows use the GPU managed distance field object buffers, no CPU culling needed
-					if (RegularCascadePSInfo != nullptr && !RegularCascadePSInfo->bRayTracedDistanceField)
+					if (!ProjectedShadowInfo->bRayTracedDistanceField)
 					{
-						ShadowInfosThatNeedCulling.Add(RegularCascadePSInfo);
+						ShadowInfosThatNeedCulling.Add(ProjectedShadowInfo);
 					}
 				}
 			}
diff --git a/Engine/Source/Runtime/Renderer/Public/MeshPassProcessor.h b/Engine/Source/Runtime/Renderer/Public/MeshPassProcessor.h
--- a/Engine/Source/Runtime/Renderer/Public/MeshPassProcessor.h
+++ b/Engine/Source/Runtime/Renderer/Public/MeshPassProcessor.h
@@ -31,6 +31,7 @@ namespace EMeshPass
 		SkyPass,
 		SingleLayerWaterPass,
 		CSMShadowDepth,
+		VSMShadowDepth,
 		Distortion,
 		Velocity,
 		TranslucentVelocity,
@@ -70,6 +71,7 @@ inline const TCHAR* GetMeshPassName(EMeshPass::Type MeshPass)
 	case EMeshPass::SkyPass: return TEXT("SkyPass");
 	case EMeshPass::SingleLayerWaterPass: return TEXT("SingleLayerWaterPass");
 	case EMeshPass::CSMShadowDepth: return TEXT("CSMShadowDepth");
+	case EMeshPass::VSMShadowDepth: return TEXT("VSMShadowDepth");
 	case EMeshPass::Distortion: return TEXT("Distortion");
 	case EMeshPass::Velocity: return TEXT("Velocity");
 	case EMeshPass::TranslucentVelocity: return TEXT("TranslucentVelocity");
