Prevent RT SkyLight crash by reinstating dimensions for visibility ray mask.

Before the sampling change, the dimension mask was aligned with the BlueNoise texture. It is now explicitly set to FIntPoint(256, 256, 0) to mimic previous behavior.

diff --git a/Engine/Source/Runtime/Renderer/Private/RayTracing/RaytracingSkylight.cpp b/Engine/Source/Runtime/Renderer/Private/RayTracing/RaytracingSkylight.cpp
--- a/Engine/Source/Runtime/Renderer/Private/RayTracing/RaytracingSkylight.cpp
+++ b/Engine/Source/Runtime/Renderer/Private/RayTracing/RaytracingSkylight.cpp
@@ -380,6 +380,9 @@ void FDeferredShadingSceneRenderer::GenerateSkyLightVisibilityRays(
 	FSkyLightData SkyLightData;
 	SetupSkyLightParameters(*Scene, &SkyLightData);
 
+	// Allocating mask of 256 x 256 rays
+	Dimensions = FIntVector(256, 256, 0);
+
 	// Output structured buffer creation
 	FRDGBufferDesc BufferDesc = FRDGBufferDesc::CreateStructuredDesc(sizeof(FSkyLightVisibilityRays), Dimensions.X * Dimensions.Y * SkyLightData.SamplesPerPixel);
 	SkyLightVisibilityRaysBuffer = GraphBuilder.CreateBuffer(BufferDesc, TEXT("SkyLightVisibilityRays"));
