Fix clamp of VT MaxLevel (MaxLevel is MaxNumMips-1)

diff --git a/Engine/Source/Runtime/Renderer/Private/VT/AllocatedVirtualTexture.cpp b/Engine/Source/Runtime/Renderer/Private/VT/AllocatedVirtualTexture.cpp
--- a/Engine/Source/Runtime/Renderer/Private/VT/AllocatedVirtualTexture.cpp
+++ b/Engine/Source/Runtime/Renderer/Private/VT/AllocatedVirtualTexture.cpp
@@ -57,7 +57,7 @@ FAllocatedVirtualTexture::FAllocatedVirtualTexture(FVirtualTextureSystem* InSyst
 	// With multiple layers of different sizes, some layers may have mips smaller than a single tile
 	MaxLevel = FMath::Min(MaxLevel, FMath::CeilLogTwo(FMath::Max(GetWidthInTiles(), GetHeightInTiles())));
 
-	MaxLevel = FMath::Min(MaxLevel, VIRTUALTEXTURE_LOG2_MAX_PAGETABLE_SIZE);
+	MaxLevel = FMath::Min(MaxLevel, VIRTUALTEXTURE_LOG2_MAX_PAGETABLE_SIZE - 1u);
 
 	// Lock lowest resolution mip from each producer
 	// Depending on the block dimensions of the producers that make up this allocated VT, different allocated VTs may need to lock different low resolution mips from the same producer
