Fix Shader Compile Mobile

diff --git a/Engine/Source/Runtime/Renderer/Private/VirtualShadowMaps/VirtualShadowMapProjection.cpp b/Engine/Source/Runtime/Renderer/Private/VirtualShadowMaps/VirtualShadowMapProjection.cpp
--- a/Engine/Source/Runtime/Renderer/Private/VirtualShadowMaps/VirtualShadowMapProjection.cpp
+++ b/Engine/Source/Runtime/Renderer/Private/VirtualShadowMaps/VirtualShadowMapProjection.cpp
@@ -468,6 +468,11 @@ class FVirtualShadowMapProjectionCS : public FGlobalShader
 
 		OutEnvironment.CompilerFlags.Add(CFLAG_WaveOperations);
 	}
+
+	static bool ShouldCompilePermutation(const FGlobalShaderPermutationParameters& Parameters)
+	{
+		return DoesPlatformSupportNanite(Parameters.Platform);
+	}
 };
 IMPLEMENT_GLOBAL_SHADER(FVirtualShadowMapProjectionCS, "/Engine/Private/VirtualShadowMaps/VirtualShadowMapProjection.usf", "VirtualShadowMapProjection", SF_Compute);
 
