Replace box filter kernel with a gaussian in the path tracer

Expose a filter width setting for users to adjust the softness/aliasing tradeoff. Default width of 3.0 maps to stddev=1/2 pixel.

diff --git a/Engine/Shaders/Private/PathTracing/Sensor/PathTracingPinholeCamera.ush b/Engine/Shaders/Private/PathTracing/Sensor/PathTracingPinholeCamera.ush
--- a/Engine/Shaders/Private/PathTracing/Sensor/PathTracingPinholeCamera.ush
+++ b/Engine/Shaders/Private/PathTracing/Sensor/PathTracingPinholeCamera.ush
@@ -16,9 +16,15 @@ void PinholeCamera_SampleEmitter(
 	out float OutPdf
 )
 {
-	float2 RandSample = RandomSequence_GenerateSample2D(RandSequence);
-	float2 ViewportUV = (LaunchIndex.xy + RandSample.xy) * View.BufferSizeAndInvSize.zw;
+	float2 AAJitter = RandomSequence_GenerateSample2D(RandSequence);
 
+	// importance sample a gaussian kernel with variable sigma
+	float3 Disk = ConcentricDiskSamplingHelper(AAJitter);
+	float Sigma = PathTracingData.FilterWidth / 6.0; // user-provided width covers +/-3*Sigma
+	float R = min(Disk.z, 0.99999994); // avoid log(0) when R=1
+	AAJitter = 0.5 + Sigma * Disk.xy * sqrt(-2.0 * log(1.0 - R * R));
+
+	float2 ViewportUV = (LaunchIndex.xy + AAJitter.xy) * View.BufferSizeAndInvSize.zw;
 	OutRay = CreatePrimaryRay(ViewportUV);
 	OutPixel = LaunchIndex.xy;
 	OutPdf = 1.0;
diff --git a/Engine/Source/Runtime/Engine/Classes/Engine/Scene.h b/Engine/Source/Runtime/Engine/Classes/Engine/Scene.h
--- a/Engine/Source/Runtime/Engine/Classes/Engine/Scene.h
+++ b/Engine/Source/Runtime/Engine/Classes/Engine/Scene.h
@@ -1202,6 +1202,9 @@ struct FPostProcessSettings
 	UPROPERTY(EditAnywhere, BlueprintReadWrite, Category = Overrides, meta = (PinHiddenByDefault, InlineEditConditionToggle))
 	uint32 bOverride_PathTracingSamplesPerPixel : 1;
 
+	UPROPERTY(EditAnywhere, BlueprintReadWrite, Category = Overrides, meta = (PinHiddenByDefault, InlineEditConditionToggle))
+	uint32 bOverride_PathTracingFilterWidth : 1;
+
 	// -----------------------------------------------------------------------
 
 	/** Enable HQ Gaussian on high end mobile platforms. (ES3_1) */
@@ -1917,6 +1920,11 @@ struct FPostProcessSettings
 	UPROPERTY(interp, EditAnywhere, BlueprintReadWrite, Category = "Rendering Features|PathTracing", meta = (ClampMin = "1", ClampMax = "65536", editcondition = "bOverride_PathTracingSamplesPerPixel", DisplayName = "Samples Per Pixel"))
 	int32 PathTracingSamplesPerPixel;
 
+	/** Sets anti-aliasing filter width for the path tracer. */
+	UPROPERTY(interp, EditAnywhere, BlueprintReadWrite, Category = "Rendering Features|PathTracing", meta = (ClampMin = "1.0", ClampMax = "6.0", editcondition = "bOverride_PathTracingFilterWidth", DisplayName = "Filter Width"))
+	float PathTracingFilterWidth;
+
+
 	UPROPERTY()
 	float LPVFadeRange_DEPRECATED;
 
diff --git a/Engine/Source/Runtime/Engine/Private/Scene.cpp b/Engine/Source/Runtime/Engine/Private/Scene.cpp
--- a/Engine/Source/Runtime/Engine/Private/Scene.cpp
+++ b/Engine/Source/Runtime/Engine/Private/Scene.cpp
@@ -571,6 +571,7 @@ FPostProcessSettings::FPostProcessSettings()
 
 	PathTracingMaxBounces = 32;
 	PathTracingSamplesPerPixel = 16384;
+	PathTracingFilterWidth = 3.0f;
 	
 	bMobileHQGaussian = false;
 
@@ -750,6 +751,7 @@ FPostProcessSettings::FPostProcessSettings(const FPostProcessSettings& Settings)
 	, bOverride_RayTracingGISamplesPerPixel(Settings.bOverride_RayTracingGISamplesPerPixel)
 	, bOverride_PathTracingMaxBounces(Settings.bOverride_PathTracingMaxBounces)
 	, bOverride_PathTracingSamplesPerPixel(Settings.bOverride_PathTracingSamplesPerPixel)
+	, bOverride_PathTracingFilterWidth(Settings.bOverride_PathTracingFilterWidth)
 
 	, bMobileHQGaussian(Settings.bMobileHQGaussian)
 	, BloomMethod(Settings.BloomMethod)
@@ -919,6 +921,7 @@ FPostProcessSettings::FPostProcessSettings(const FPostProcessSettings& Settings)
 
 	, PathTracingMaxBounces(Settings.PathTracingMaxBounces)
 	, PathTracingSamplesPerPixel(Settings.PathTracingSamplesPerPixel)
+	, PathTracingFilterWidth(Settings.PathTracingFilterWidth)
 
 	, ScreenPercentage(Settings.ScreenPercentage)
 
diff --git a/Engine/Source/Runtime/Engine/Private/SceneView.cpp b/Engine/Source/Runtime/Engine/Private/SceneView.cpp
--- a/Engine/Source/Runtime/Engine/Private/SceneView.cpp
+++ b/Engine/Source/Runtime/Engine/Private/SceneView.cpp
@@ -1552,6 +1552,11 @@ void FSceneView::OverridePostProcessSettings(const FPostProcessSettings& Src, fl
 			Dest.PathTracingSamplesPerPixel = Src.PathTracingSamplesPerPixel;
 		}
 
+		if (Src.bOverride_PathTracingFilterWidth)
+		{
+			Dest.PathTracingFilterWidth = Src.PathTracingFilterWidth;
+		}
+
 
 		if (Src.bOverride_DepthOfFieldBladeCount)
 		{
diff --git a/Engine/Source/Runtime/Renderer/Private/PathTracing.cpp b/Engine/Source/Runtime/Renderer/Private/PathTracing.cpp
--- a/Engine/Source/Runtime/Renderer/Private/PathTracing.cpp
+++ b/Engine/Source/Runtime/Renderer/Private/PathTracing.cpp
@@ -30,6 +30,14 @@ TAutoConsoleVariable<int32> CVarPathTracingSamplesPerPixel(
 	ECVF_RenderThreadSafe
 );
 
+TAutoConsoleVariable<float> CVarPathTracingFilterWidth(
+	TEXT("r.PathTracing.FilterWidth"),
+	-1,
+	TEXT("Define the anti-aliasing filter width (default = -1 (driven by postprocesing volume))"),
+	ECVF_RenderThreadSafe
+);
+
+
 TAutoConsoleVariable<int32> CVarPathTracingUseErrorDiffusion(
 	TEXT("r.PathTracing.UseErrorDiffusion"),
 	0,
@@ -110,7 +118,9 @@ IMPLEMENT_GLOBAL_SHADER_PARAMETER_STRUCT(FPathTracingData, "PathTracingData");
 static bool PrepareShaderArgs(const FViewInfo& View, FPathTracingData& PathTracingData) {
 	int32 PathTracingMaxBounces = CVarPathTracingMaxBounces.GetValueOnRenderThread();
 	if (PathTracingMaxBounces < 0)
+	{
 		PathTracingMaxBounces = View.FinalPostProcessSettings.PathTracingMaxBounces;
+	}
 	PathTracingData.MaxBounces = PathTracingMaxBounces;
 	PathTracingData.MaxNormalBias = GetRaytracingMaxNormalBias();
 	PathTracingData.MISMode = CVarPathTracingMISMode.GetValueOnRenderThread();
@@ -118,6 +128,12 @@ static bool PrepareShaderArgs(const FViewInfo& View, FPathTracingData& PathTraci
 	PathTracingData.MaxPathIntensity = CVarPathTracingMaxPathIntensity.GetValueOnRenderThread();
 	PathTracingData.UseErrorDiffusion = CVarPathTracingUseErrorDiffusion.GetValueOnRenderThread();
 	PathTracingData.ApproximateCaustics = CVarPathTracingApproximateCaustics.GetValueOnRenderThread();
+	float FilterWidth = CVarPathTracingFilterWidth.GetValueOnRenderThread();
+	if (FilterWidth < 0)
+	{
+		FilterWidth = View.FinalPostProcessSettings.PathTracingFilterWidth;
+	}
+	PathTracingData.FilterWidth = FilterWidth;
 
 	bool NeedInvalidation = false;
 
@@ -171,6 +187,14 @@ static bool PrepareShaderArgs(const FViewInfo& View, FPathTracingData& PathTraci
 		PreviousApproximateCaustics = PathTracingData.ApproximateCaustics;
 	}
 
+	// Changing filter width requires starting over
+	static float PreviousFilterWidth = PathTracingData.FilterWidth;
+	if (PreviousFilterWidth != PathTracingData.FilterWidth)
+	{
+		NeedInvalidation = true;
+		PreviousFilterWidth = PathTracingData.FilterWidth;
+	}
+
 	// the rest of PathTracingData and AdaptiveSamplingData is filled in by SetParameters below
 	return NeedInvalidation;
 }
diff --git a/Engine/Source/Runtime/Renderer/Private/PathTracingUniformBuffers.h b/Engine/Source/Runtime/Renderer/Private/PathTracingUniformBuffers.h
--- a/Engine/Source/Runtime/Renderer/Private/PathTracingUniformBuffers.h
+++ b/Engine/Source/Runtime/Renderer/Private/PathTracingUniformBuffers.h
@@ -20,4 +20,5 @@ BEGIN_GLOBAL_SHADER_PARAMETER_STRUCT(FPathTracingData, )
 	SHADER_PARAMETER(uint32, ApproximateCaustics)
 	SHADER_PARAMETER(float, MaxPathIntensity)
 	SHADER_PARAMETER(float, MaxNormalBias)
+	SHADER_PARAMETER(float, FilterWidth)
 END_GLOBAL_SHADER_PARAMETER_STRUCT()
