Disable engine-anti-aliasing method when path tracer is enabled

The path tracer can do its own anti-aliasing, so disable the engine feature to avoid doubling-up the effect.

diff --git a/Engine/Source/Runtime/Renderer/Private/SceneRendering.cpp b/Engine/Source/Runtime/Renderer/Private/SceneRendering.cpp
--- a/Engine/Source/Runtime/Renderer/Private/SceneRendering.cpp
+++ b/Engine/Source/Runtime/Renderer/Private/SceneRendering.cpp
@@ -2440,7 +2440,12 @@ void FSceneRenderer::PrepareViewRectsForRendering()
 
 		// Fallback to no anti aliasing.
 		{
-			const bool bWillApplyTemporalAA = IsPostProcessingEnabled(View) || (View.bIsPlanarReflection && FeatureLevel >= ERHIFeatureLevel::SM5);
+			const bool bWillApplyTemporalAA = (IsPostProcessingEnabled(View) || (View.bIsPlanarReflection && FeatureLevel >= ERHIFeatureLevel::SM5))
+#if RHI_RAYTRACING
+				// path tracer does its own anti-aliasing
+				&& (View.RayTracingRenderMode != ERayTracingRenderMode::PathTracing)
+#endif
+			;
 
 			if (!bWillApplyTemporalAA)
 			{
