Fix memory leaks in new instance culling machinery

diff --git a/Engine/Source/Runtime/Renderer/Private/InstanceCulling/InstanceCullingManager.cpp b/Engine/Source/Runtime/Renderer/Private/InstanceCulling/InstanceCullingManager.cpp
--- a/Engine/Source/Runtime/Renderer/Private/InstanceCulling/InstanceCullingManager.cpp
+++ b/Engine/Source/Runtime/Renderer/Private/InstanceCulling/InstanceCullingManager.cpp
@@ -20,6 +20,15 @@ static TAutoConsoleVariable<int32> CVarCullInstances(
 	ECVF_RenderThreadSafe);
 
 
+FInstanceCullingManager::~FInstanceCullingManager()
+{
+	for (FInstanceCullingContext* InstanceCullingContext : CullingContexts)
+	{
+		// Memstack allocated, so no destructor is called.
+		InstanceCullingContext->~FInstanceCullingContext();
+	}
+}
+
 int32 FInstanceCullingManager::RegisterView(const FViewInfo& ViewInfo)
 {
 	if (!bIsEnabled)
@@ -163,6 +172,7 @@ FInstanceCullingContext* FInstanceCullingManager::CreateContext(const int32* Vie
 	if (bIsEnabled)
 	{
 		FInstanceCullingContext* InstanceCullingContext = new(FMemStack::Get()) FInstanceCullingContext;
+		CullingContexts.Add(InstanceCullingContext);
 		InstanceCullingContext->ViewIds.Insert(ViewIds, NumViews, 0);
 		InstanceCullingContext->InstanceCullingManager = this;
 		return InstanceCullingContext;
diff --git a/Engine/Source/Runtime/Renderer/Private/InstanceCulling/InstanceCullingManager.h b/Engine/Source/Runtime/Renderer/Private/InstanceCulling/InstanceCullingManager.h
--- a/Engine/Source/Runtime/Renderer/Private/InstanceCulling/InstanceCullingManager.h
+++ b/Engine/Source/Runtime/Renderer/Private/InstanceCulling/InstanceCullingManager.h
@@ -71,6 +71,7 @@ struct FInstanceCullingResult
 class FInstanceCullingManager
 {
 public:
+	~FInstanceCullingManager();
 
 	// Max average number of instances that primitives are expanded to. GPUCULL_TODO: Not very robust
 	static constexpr uint32 MaxAverageInstanceFactor = 128;
@@ -110,7 +111,7 @@ private:
 
 	FInstanceCullingManagerResources& Resources;
 	TArray<Nanite::FPackedView> CullingViews;
-	
+	TArray<FInstanceCullingContext*> CullingContexts;
 	bool bIsEnabled;
 };
 
diff --git a/Engine/Source/Runtime/Renderer/Private/ShadowSetup.cpp b/Engine/Source/Runtime/Renderer/Private/ShadowSetup.cpp
--- a/Engine/Source/Runtime/Renderer/Private/ShadowSetup.cpp
+++ b/Engine/Source/Runtime/Renderer/Private/ShadowSetup.cpp
@@ -2032,7 +2032,7 @@ void FProjectedShadowInfo::SetupMeshDrawCommandsForShadowDepth(FSceneRenderer& R
 		ShadowDepthPass.SetDumpInstancingStats(TEXT("ShadowDepth ") + PassNameForStats);
 	}
 	
-
+	ViewIds.Reset();
 	if (GetShadowDepthType().bOnePassPointLightShadow)
 	{
 		ViewIds.AddDefaulted(6);
