Disabled dynamic instancing for shadow stencil rendering for the new instance culling path (dynamic draw pass)

#rnx

diff --git a/Engine/Source/Runtime/Renderer/Private/MeshPassProcessor.cpp b/Engine/Source/Runtime/Renderer/Private/MeshPassProcessor.cpp
--- a/Engine/Source/Runtime/Renderer/Private/MeshPassProcessor.cpp
+++ b/Engine/Source/Runtime/Renderer/Private/MeshPassProcessor.cpp
@@ -1208,6 +1208,14 @@ void SubmitMeshDrawCommandsRange(
 	uint32 InstanceFactor,
 	FRHICommandList& RHICmdList)
 {
+#if defined(GPUCULL_TODO)
+	// GPUCULL_TODO: workaround for the fact that DrawDynamicMeshPassPrivate et al. don't work with GPU-Scene instancing
+	//               we don't support dynamic instancing for this path since we require one primitive per draw command
+	//               This is because the stride on the instance data buffer is set to 0 so only the first will ever be fetched.
+	checkSlow(!bDynamicInstancing);
+	bDynamicInstancing = false;
+#endif
+
 	FMeshDrawCommandStateCache StateCache;
 	INC_DWORD_STAT_BY(STAT_MeshDrawCalls, NumMeshDrawCommands);
 
diff --git a/Engine/Source/Runtime/Renderer/Private/ShadowRendering.cpp b/Engine/Source/Runtime/Renderer/Private/ShadowRendering.cpp
--- a/Engine/Source/Runtime/Renderer/Private/ShadowRendering.cpp
+++ b/Engine/Source/Runtime/Renderer/Private/ShadowRendering.cpp
@@ -811,7 +811,14 @@ void FProjectedShadowInfo::SetupProjectionStencilMask(
 	DrawRenderState.SetDepthStencilState(TStaticDepthStencilState<false, CF_DepthNearOrEqual>::GetRHI());
 	DrawRenderState.SetBlendState(TStaticBlendState<CW_NONE>::GetRHI());
 
+#if defined(GPUCULL_TODO)
+	// GPUCULL_TODO: workaround for the fact that DrawDynamicMeshPassPrivate et al. don't work with GPU-Scene instancing
+	//               we don't support dynamic instancing for this path since we require one primitive per draw command
+	//               This is because the stride on the instance data buffer is set to 0 so only the first will ever be fetched.
+	const bool bDynamicInstancing = false;
+#else
 	const bool bDynamicInstancing = IsDynamicInstancingEnabled(View->FeatureLevel);
+#endif
 
 	// If this is a preshadow, mask the projection by the receiver primitives.
 	if (bPreShadow || bSelfShadowOnly)
