Fixed inconsistence in ray tracing Cvar names

diff --git a/Engine/Plugins/Runtime/Nvidia/Ansel/Source/Ansel/Private/Ansel.cpp b/Engine/Plugins/Runtime/Nvidia/Ansel/Source/Ansel/Private/Ansel.cpp
--- a/Engine/Plugins/Runtime/Nvidia/Ansel/Source/Ansel/Private/Ansel.cpp
+++ b/Engine/Plugins/Runtime/Nvidia/Ansel/Source/Ansel/Private/Ansel.cpp
@@ -924,8 +924,8 @@ void FNVAnselCameraPhotographyPrivate::ConfigureRenderingSettingsForPhotography(
 			QUALITY_CVAR_AT_LEAST("r.RayTracing.Translucency.Shadows", 1); // turn on at least basic quality
 
 			/*** HIGH-QUALITY MODE DOES *NOT* FORCE RT SHADOWS ON ***/
-			QUALITY_CVAR("r.RayTracing.Shadow.MaxLights", -1); // unlimited
-			QUALITY_CVAR("r.RayTracing.Shadow.MaxDenoisedLights", -1); // unlimited
+			QUALITY_CVAR("r.RayTracing.Shadows.MaxLights", -1); // unlimited
+			QUALITY_CVAR("r.RayTracing.Shadows.MaxDenoisedLights", -1); // unlimited
 
 			// these apply to various RT effects but mostly reflections+translucency
 			QUALITY_CVAR_AT_LEAST("r.raytracing.lighting.maxshadowlights", 256); // as seen in reflections/translucencies
diff --git a/Engine/Source/Runtime/Renderer/Private/LightRendering.cpp b/Engine/Source/Runtime/Renderer/Private/LightRendering.cpp
--- a/Engine/Source/Runtime/Renderer/Private/LightRendering.cpp
+++ b/Engine/Source/Runtime/Renderer/Private/LightRendering.cpp
@@ -68,7 +68,7 @@ static TAutoConsoleVariable<int32> CVarRayTracingOcclusion(
 
 static int32 GShadowRayTracingSamplesPerPixel = 1;
 static FAutoConsoleVariableRef CVarShadowRayTracingSamplesPerPixel(
-	TEXT("r.RayTracing.Shadow.SamplesPerPixel"),
+	TEXT("r.RayTracing.Shadows.SamplesPerPixel"),
 	GShadowRayTracingSamplesPerPixel,
 	TEXT("Sets the samples-per-pixel for directional light occlusion (default = 1)"),
 	ECVF_RenderThreadSafe);
@@ -88,7 +88,7 @@ static TAutoConsoleVariable<int32> CVarMaxShadowDenoisingBatchSize(
 	ECVF_RenderThreadSafe);
 
 static TAutoConsoleVariable<int32> CVarMaxShadowRayTracingBatchSize(
-	TEXT("r.RayTracing.Shadow.MaxBatchSize"), 8,
+	TEXT("r.RayTracing.Shadows.MaxBatchSize"), 8,
 	TEXT("Maximum number of shadows to trace at the same time."),
 	ECVF_RenderThreadSafe);
 
