Fix static analysis issue.

diff --git a/Engine/Source/Runtime/Renderer/Private/ShadowRendering.cpp b/Engine/Source/Runtime/Renderer/Private/ShadowRendering.cpp
--- a/Engine/Source/Runtime/Renderer/Private/ShadowRendering.cpp
+++ b/Engine/Source/Runtime/Renderer/Private/ShadowRendering.cpp
@@ -1891,7 +1891,7 @@ void FDeferredShadingSceneRenderer::RenderShadowProjections(
 			CommonPassParameters.HairCategorizationTexture = bSubPixel && HairDatas && HairDatas->HairVisibilityViews.HairDatas.Num() > 0 ? HairDatas->HairVisibilityViews.HairDatas[0].CategorizationTexture : nullptr;
 			CommonPassParameters.RenderTargets[0] = FRenderTargetBinding(OutputTexture, ERenderTargetLoadAction::ELoad);
 			CommonPassParameters.RenderTargets.DepthStencil = 
-				bSubPixel && HairDatas->HairVisibilityViews.HairDatas.Num() > 0 ?
+				bSubPixel && HairDatas && HairDatas->HairVisibilityViews.HairDatas.Num() > 0 ?
 				FDepthStencilBinding(HairDatas->HairVisibilityViews.HairDatas[0].HairOnlyDepthTexture, ERenderTargetLoadAction::ELoad, ERenderTargetLoadAction::ELoad, ExclusiveDepthStencil) :
 				FDepthStencilBinding(SceneTextures.Depth.Target, ERenderTargetLoadAction::ELoad, ERenderTargetLoadAction::ELoad, ExclusiveDepthStencil);
 
