* Add classification pass for strata for separating pixel/tile having a single slab with no aniso/sss/hazing feature from more complex pixels. The classification is done by rendering stencil mask (vs. compute tile/raster tile) as light are rendered by geometric volume at the moment. This makes the classification compatible with light rendering.
* Add lighting permutation for analytical lights permutation for strata fast path

This changeslist also changes read/write stencil mask used for shadow rendering. This changelist has been tested against engine test shadow/lightin cases.

diff --git a/Engine/Shaders/Private/Strata/Strata.ush b/Engine/Shaders/Private/Strata/Strata.ush
--- a/Engine/Shaders/Private/Strata/Strata.ush
+++ b/Engine/Shaders/Private/Strata/Strata.ush
@@ -176,6 +176,7 @@ struct FStrataClassification
 	bool bHasSubsurface;
 	bool bHasSubsurfaceProfile;
 	bool bHasTransmission;
+	bool bIsSimple;
 };
 
 uint StrataPackClassificationData(FStrataClassification In)
@@ -185,7 +186,8 @@ uint StrataPackClassificationData(FStrataClassification In)
 		((In.bIsStrataMaterial		? 1 : 0) << 8 ) |
 		((In.bHasSubsurface			? 1 : 0) << 9 ) | 
 		((In.bHasSubsurfaceProfile	? 1 : 0) << 10) |
-		((In.bHasTransmission		? 1 : 0) << 11) ;
+		((In.bHasTransmission		? 1 : 0) << 11) |
+		((In.bIsSimple				? 1 : 0) << 12);
 }
 
 FStrataClassification StrataUnpackClassificationData(uint In)
@@ -196,6 +198,7 @@ FStrataClassification StrataUnpackClassificationData(uint In)
 	Out.bHasSubsurface			= (In >> 9 ) & 0x1;
 	Out.bHasSubsurfaceProfile	= (In >> 10) & 0x1;
 	Out.bHasTransmission		= (In >> 11) & 0x1;
+	Out.bIsSimple				= (In >> 12) & 0x1;
 	return Out;
 }
 
diff --git a/Engine/Shaders/Private/Strata/StrataDeferredLighting.ush b/Engine/Shaders/Private/Strata/StrataDeferredLighting.ush
--- a/Engine/Shaders/Private/Strata/StrataDeferredLighting.ush
+++ b/Engine/Shaders/Private/Strata/StrataDeferredLighting.ush
@@ -139,6 +139,10 @@ void StrataPreprocessAreaLightIntegration(
 	}
 }
 
+#ifndef STRATA_FASTPATH 
+#define STRATA_FASTPATH 0
+#endif
+
 // Analytical lighting evaluation for strata material.
 // Unpack BSDF on-the-fly
 float4 StrataDeferredLighting(
@@ -169,11 +173,24 @@ float4 StrataDeferredLighting(
 
 	FLightAccumulator Out = (FLightAccumulator)0;
 
+#if STRATA_FASTPATH
+	uint BSDFIndex = 0;
+	if (StrataPixelHeader.BSDFCount == 1)
+#else
 	LOOP
 	for (uint BSDFIndex = 0; BSDFIndex < StrataPixelHeader.BSDFCount; ++BSDFIndex)
+#endif
 	{
 		// Unpack BSDF data
-		const FStrataBSDF BSDF = UnpackStrataBSDFIn(StrataDataBuffer, PixelStrataDataByteOffset);
+		FStrataBSDF BSDF = UnpackStrataBSDFIn(StrataDataBuffer, PixelStrataDataByteOffset);
+#if STRATA_FASTPATH
+		BSDF_SETTYPE(BSDF, STRATA_BSDF_TYPE_SLAB);
+		BSDF_SETANISOTROPY(BSDF, 0);
+		BSDF_SETISTOPLAYER(BSDF, 0);
+		BSDF_SETHASSCATTERING(BSDF, 0);
+		BSDF_SETHASSSSPROFILE(BSDF, 0);
+		BSDF_SETHASHAZINESS(BSDF, 0);
+#endif
 		FStrataBSDFContext BSDFContext = StrataCreateBSDFContext(StrataPixelHeader, BSDF, V, L);
 
 		// Evaluate incident lighting & (optionally) precompute area light integration
diff --git a/Engine/Shaders/Private/Strata/StrataMaterialClassification.usf b/Engine/Shaders/Private/Strata/StrataMaterialClassification.usf
--- a/Engine/Shaders/Private/Strata/StrataMaterialClassification.usf
+++ b/Engine/Shaders/Private/Strata/StrataMaterialClassification.usf
@@ -6,6 +6,8 @@
 #define STRATA_INLINE_SHADING 0
 #include "/Engine/Private/Strata/Strata.ush"
 
+#if SHADER_CATEGORIZATION
+
 // Note: 
 // This code is partial copy of PackStrataOut in Strata.ush. 
 // The command part will progressively be removed as we remove data from the gbuffer.
@@ -38,6 +40,7 @@ void ComputePostProcessBSDFData(
 	const bool bTopLayer = BSDF_GETISTOPLAYER(BSDF);
 
 	Classification.ShadingModels |= StrataShadingModelBit(BSDF_GETTYPE(BSDF));
+	Classification.bIsSimple = Classification.bIsSimple && BSDF_GETTYPE(BSDF) == STRATA_BSDF_TYPE_SLAB;
 
 	switch (BSDF_GETTYPE(BSDF))
 	{
@@ -55,6 +58,13 @@ void ComputePostProcessBSDFData(
 			TopLayer.BaseColor += BSDFCoverage * EnvBRDFApprox(F0, SafeRoughness, saturate(ContextNoLight.NoV));
 		}
 
+		Classification.bIsSimple =
+			Classification.bIsSimple &&
+			!BSDF_GETANISOTROPY(BSDF) &&
+			!BSDF_GETHASSCATTERING(BSDF) &&
+			!BSDF_GETHASSSSPROFILE(BSDF) &&
+			!BSDF_GETHASHAZINESS(BSDF);
+
 		// Select the first/toppest layer with subsurface scattering
 		if (BSDF_GETHASSCATTERING(BSDF) && BSDFCoverageAvg > 0 && !SSS.bHasSubsurface)
 		{
@@ -143,6 +153,7 @@ void MainPS(
 
 		uint PixelStrataDataByteOffset = GetStrataPixelDataByteOffset(PixelPos, View.BufferSizeAndInvSize.x, Strata.MaxBytesPerPixel);
 		FStrataPixelHeader Header = UnpackStrataHeaderIn(Strata.MaterialLobesBuffer, PixelStrataDataByteOffset);
+		ClassificationData.bIsSimple = Header.BSDFCount == 1;
 		for (uint BSDFIndex = 0; BSDFIndex < Header.BSDFCount; ++BSDFIndex)
 		{
 			const FStrataBSDF BSDF = UnpackStrataBSDFIn(Strata.MaterialLobesBuffer, PixelStrataDataByteOffset);
@@ -155,3 +166,96 @@ void MainPS(
 	OutColor0 = StrataPackClassificationData(ClassificationData);
 	OutColor1 = StrataPackTopLayerData(TopLayerData);
 }
+#endif // SHADER_CATEGORIZATION
+
+////////////////////////////////////////////////////////////////////////////////////////////////////////////
+
+#if SHADER_TILE_CATEGORIZATION
+int  TileSize;
+int2 ClassificationResolution;
+int2 TileClassificationResolution;
+Texture2D<uint> ClassificationTexture;
+
+#define GROUP_SIZE 8
+#define GROUP_THREAD_COUNT (GROUP_SIZE*GROUP_SIZE)
+groupshared uint s_TileType;
+RWTexture2D<uint> TileClassificationTexture;
+[numthreads(8, 8, 1)]
+void TileMainCS(uint2 DispatchThreadId : SV_DispatchThreadID, uint GroupIndex : SV_GroupIndex)
+{
+	if (GroupIndex == 0)
+	{
+		s_TileType = 0;
+	}
+
+	const uint2 PixelCoord = DispatchThreadId;
+	const uint PackClassification = ClassificationTexture.Load(uint3(PixelCoord, 0));
+	FStrataClassification Data = StrataUnpackClassificationData(PackClassification);
+
+	if (all(PixelCoord < ClassificationResolution))
+	{
+		InterlockedAdd(s_TileType, Data.bIsSimple ? 1 : 0);
+	}
+
+	GroupMemoryBarrierWithGroupSync();
+
+	if (GroupIndex == 0)
+	{
+		const uint2 TileCoord = PixelCoord / TileSize;
+		if (s_TileType == GROUP_THREAD_COUNT && all(TileCoord < TileClassificationResolution))
+		{
+			InterlockedAdd(TileClassificationTexture[TileCoord], s_TileType);
+		}
+	}
+}
+#endif // SHADER_TILE_CATEGORIZATION
+
+////////////////////////////////////////////////////////////////////////////////////////////////////////////
+
+#if SHADER_STENCIL_CATEGORIZATION
+
+int  TileSize;
+int2 TileCount;
+int2 OutputResolution;
+Texture2D<uint> TileClassificationTexture;
+
+void StencilMainVS(
+	in uint InVertexId : SV_VertexID,
+	in uint InInstanceId : SV_InstanceID,
+	out float4 OutPosition : SV_POSITION)
+{
+	uint2 TileCoord = uint2(InInstanceId % TileCount.x, InInstanceId / TileCount.y);
+	
+	// Only tile whose all the threads are flagged as simple will be set as valid
+	const uint bIsSimple = TileClassificationTexture.Load(uint3(TileCoord, 0)) == (TileSize*TileSize);
+	float3 P = float3(INFINITE_FLOAT, INFINITE_FLOAT, INFINITE_FLOAT);
+	if (bIsSimple)
+	{
+		float2 UV = 0;
+		switch (InVertexId)
+		{
+		case 0: UV = float2(0, 0); break;
+		case 1: UV = float2(1, 0); break;
+		case 2: UV = float2(1, 1); break;
+		case 3: UV = float2(0, 0); break;
+		case 4: UV = float2(1, 1); break;
+		case 5: UV = float2(0, 1); break;
+		}
+
+		int LocalTileSize = TileSize;
+		P.xy = TileCoord * LocalTileSize + UV * LocalTileSize;
+		P.xy/= float2(OutputResolution);
+		P.xy = float2(P.x, (1 - P.y)) * 2 - 1;
+		P.z  = 0.5f;
+	}
+	OutPosition = float4(P, 1);
+}
+
+void StencilMainPS(
+	float4 SVPos : SV_POSITION,
+	out uint OutColor0 : SV_Target0)
+{
+	OutColor0 = 255;
+}
+
+#endif //SHADER_STENCIL_CATEGORIZATION
\ No newline at end of file
diff --git a/Engine/Source/Runtime/Renderer/Private/DeferredShadingRenderer.cpp b/Engine/Source/Runtime/Renderer/Private/DeferredShadingRenderer.cpp
--- a/Engine/Source/Runtime/Renderer/Private/DeferredShadingRenderer.cpp
+++ b/Engine/Source/Runtime/Renderer/Private/DeferredShadingRenderer.cpp
@@ -2419,7 +2419,7 @@ void FDeferredShadingSceneRenderer::Render(FRDGBuilder& GraphBuilder)
 	// Post base pass for material classification
 	if (Strata::IsStrataEnabled())
 	{
-		Strata::AddStrataMaterialClassificationPass(GraphBuilder, Views);
+		Strata::AddStrataMaterialClassificationPass(GraphBuilder, SceneTextures, Views);
 	}
 
 	// Hair base pass for deferred shading
diff --git a/Engine/Source/Runtime/Renderer/Private/LightRendering.cpp b/Engine/Source/Runtime/Renderer/Private/LightRendering.cpp
--- a/Engine/Source/Runtime/Renderer/Private/LightRendering.cpp
+++ b/Engine/Source/Runtime/Renderer/Private/LightRendering.cpp
@@ -442,6 +442,7 @@ class FDeferredLightPS : public FGlobalShader
 	class FCloudTransmittance 	: SHADER_PERMUTATION_BOOL("USE_CLOUD_TRANSMITTANCE");
 	class FAnistropicMaterials 	: SHADER_PERMUTATION_BOOL("SUPPORTS_ANISOTROPIC_MATERIALS");
 	class FStrata				: SHADER_PERMUTATION_BOOL("STRATA_ENABLED");
+	class FStrataFastPath		: SHADER_PERMUTATION_BOOL("STRATA_FASTPATH");
 
 	using FPermutationDomain = TShaderPermutationDomain<
 		FSourceShapeDim,
@@ -455,7 +456,8 @@ class FDeferredLightPS : public FGlobalShader
 		FAtmosphereTransmittance,
 		FCloudTransmittance,
 		FAnistropicMaterials,
-		FStrata>;
+		FStrata,
+		FStrataFastPath>;
 
 	static bool ShouldCompilePermutation(const FGlobalShaderPermutationParameters& Parameters)
 	{
@@ -522,6 +524,10 @@ class FDeferredLightPS : public FGlobalShader
 			return false;
 		}
 
+		if (PermutationVector.Get<FStrataFastPath>() && (!Strata::IsStrataEnabled() || !PermutationVector.Get<FStrata>()))
+		{
+			return false;
+		}
 		return IsFeatureLevelSupported(Parameters.Platform, ERHIFeatureLevel::SM5);
 	}
 
@@ -1334,6 +1340,11 @@ void FDeferredShadingSceneRenderer::RenderLights(
 
 			if (!bUseHairLighting)
 			{
+				if (Strata::IsStrataEnabled() && Strata::IsClassificationEnabled())
+				{
+					Strata::AddStrataStencilPass(GraphBuilder, Views, SceneTextures);
+				}
+
 				FRenderLightParameters* PassParameters = GraphBuilder.AllocParameters<FRenderLightParameters>();
 				GetRenderLightParameters(SceneTextures, nullptr, LightingChannelsTexture, {}, HairDatas ? &HairDatas->HairVisibilityViews : nullptr, *PassParameters);
 
@@ -2087,10 +2098,28 @@ void SetBoundingGeometryRasterizerAndDepthState(FGraphicsPipelineStateInitialize
 		GraphicsPSOInit.RasterizerState = View.bReverseCulling ? TStaticRasterizerState<FM_Solid, CM_CCW>::GetRHI() : TStaticRasterizerState<FM_Solid, CM_CW>::GetRHI();
 	}
 
-	GraphicsPSOInit.DepthStencilState =
+	if (Strata::IsStrataEnabled() && Strata::IsClassificationEnabled())
+	{
+		GraphicsPSOInit.DepthStencilState =
 		bCameraInsideLightGeometry
-		? TStaticDepthStencilState<false, CF_Always>::GetRHI()
-		: TStaticDepthStencilState<false, CF_DepthNearOrEqual>::GetRHI();
+		? TStaticDepthStencilState<
+			false, CF_Always,
+			true, CF_Equal, SO_Keep, SO_Keep, SO_Keep,
+			true, CF_Equal, SO_Keep, SO_Keep, SO_Keep,
+			Strata::StencilBit, 0x0>::GetRHI()
+		: TStaticDepthStencilState<
+			false, CF_DepthNearOrEqual,
+			true, CF_Equal, SO_Keep, SO_Keep, SO_Keep,
+			true, CF_Equal, SO_Keep, SO_Keep, SO_Keep,
+			Strata::StencilBit, 0x0>::GetRHI();
+	}
+	else
+	{
+		GraphicsPSOInit.DepthStencilState =
+			bCameraInsideLightGeometry
+			? TStaticDepthStencilState<false, CF_Always>::GetRHI()
+			: TStaticDepthStencilState<false, CF_DepthNearOrEqual>::GetRHI();
+	}
 }
 
 template<bool bUseIESProfile, bool bRadialAttenuation, bool bInverseSquaredFalloff>
@@ -2114,6 +2143,7 @@ static void SetShaderTemplLightingSimple(
 	PermutationVector.Set< FDeferredLightPS::FAtmosphereTransmittance >( false );
 	PermutationVector.Set< FDeferredLightPS::FCloudTransmittance >(false);
 	PermutationVector.Set< FDeferredLightPS::FStrata >(Strata::IsStrataEnabled());
+	PermutationVector.Set< FDeferredLightPS::FStrataFastPath>(false);
 
 	TShaderMapRef< FDeferredLightPS > PixelShader( View.ShaderMap, PermutationVector );
 	GraphicsPSOInit.BoundShaderState.VertexDeclarationRHI = GetVertexDeclarationFVector4();
@@ -2182,6 +2212,9 @@ void FDeferredShadingSceneRenderer::RenderLight(
 	INC_DWORD_STAT(STAT_NumLightsUsingStandardDeferred);
 	SCOPED_CONDITIONAL_DRAW_EVENT(RHICmdList, StandardDeferredLighting, bIssueDrawEvent);
 
+	auto RenderInternalLight = [&](bool bStrataFastPath)
+	{
+
 	FGraphicsPipelineStateInitializer GraphicsPSOInit;
 	RHICmdList.ApplyCachedRenderTargets(GraphicsPSOInit);
 
@@ -2219,6 +2252,15 @@ void FDeferredShadingSceneRenderer::RenderLight(
 		GraphicsPSOInit.RasterizerState = TStaticRasterizerState<FM_Solid, CM_None>::GetRHI();
 		GraphicsPSOInit.DepthStencilState = TStaticDepthStencilState<false, CF_Always>::GetRHI();
 
+		if (Strata::IsStrataEnabled() && Strata::IsClassificationEnabled())
+		{
+			GraphicsPSOInit.DepthStencilState = TStaticDepthStencilState<
+				false, CF_Always,
+				true, CF_Equal, SO_Keep, SO_Keep, SO_Keep,
+				true, CF_Equal, SO_Keep, SO_Keep, SO_Keep,
+				Strata::StencilBit, 0x0>::GetRHI();
+		}
+
 		if (bRenderOverlap)
 		{
 			TShaderMapRef<TDeferredLightOverlapPS<false> > PixelShader(View.ShaderMap);
@@ -2269,6 +2311,7 @@ void FDeferredShadingSceneRenderer::RenderLight(
 			PermutationVector.Set< FDeferredLightPS::FAtmosphereTransmittance >(bAtmospherePerPixelTransmittance);
 			PermutationVector.Set< FDeferredLightPS::FCloudTransmittance >(bLight0CloudPerPixelTransmittance || bLight1CloudPerPixelTransmittance);
 			PermutationVector.Set< FDeferredLightPS::FStrata >(Strata::IsStrataEnabled());
+			PermutationVector.Set< FDeferredLightPS::FStrataFastPath >(Strata::IsStrataEnabled() && Strata::IsClassificationEnabled() && bStrataFastPath);
 
 			TShaderMapRef< FDeferredLightPS > PixelShader( View.ShaderMap, PermutationVector );
 			GraphicsPSOInit.BoundShaderState.VertexDeclarationRHI = GFilterVertexDeclaration.VertexDeclarationRHI;
@@ -2282,6 +2325,7 @@ void FDeferredShadingSceneRenderer::RenderLight(
 		VertexShader->SetParameters(RHICmdList, View, LightSceneInfo);
 
 		// Apply the directional light as a full screen quad
+		RHICmdList.SetStencilRef(bStrataFastPath ? Strata::StencilBit : 0u);
 		DrawRectangle(
 			RHICmdList,
 			0, 0,
@@ -2334,6 +2378,7 @@ void FDeferredShadingSceneRenderer::RenderLight(
 			PermutationVector.Set < FDeferredLightPS::FAtmosphereTransmittance >(false);
 			PermutationVector.Set< FDeferredLightPS::FCloudTransmittance >(false);
 			PermutationVector.Set< FDeferredLightPS::FStrata >(Strata::IsStrataEnabled());
+			PermutationVector.Set< FDeferredLightPS::FStrata >(Strata::IsStrataEnabled() && Strata::IsClassificationEnabled() && bStrataFastPath);
 
 			TShaderMapRef< FDeferredLightPS > PixelShader( View.ShaderMap, PermutationVector );
 			GraphicsPSOInit.BoundShaderState.VertexDeclarationRHI = GetVertexDeclarationFVector4();
@@ -2376,6 +2421,14 @@ void FDeferredShadingSceneRenderer::RenderLight(
 			StencilingGeometry::DrawCone(RHICmdList);
 		}
 	}
+	};
+
+
+	RenderInternalLight(false);
+	if (Strata::IsStrataEnabled() && Strata::IsClassificationEnabled())
+	{
+		RenderInternalLight(true);
+	}
 }
 
 void FDeferredShadingSceneRenderer::RenderLight(
@@ -2396,6 +2449,11 @@ void FDeferredShadingSceneRenderer::RenderLight(
 		RDG_EVENT_SCOPE_CONDITIONAL(GraphBuilder, ViewCount > 1, "View%d", ViewIndex);
 		const FViewInfo& View = Views[ViewIndex];
 
+		if (Strata::IsStrataEnabled() && Strata::IsClassificationEnabled())
+		{
+			Strata::AddStrataStencilPass(GraphBuilder, View, SceneTextures);
+		}
+
 		FRenderLightParameters* PassParameters = GraphBuilder.AllocParameters<FRenderLightParameters>();
 		GetRenderLightParameters(SceneTextures, ScreenShadowMaskTexture, LightingChannelsTexture, GetCloudShadowAOParameters(GraphBuilder, View, CloudInfo), InHairVisibilityViews, *PassParameters);
 
diff --git a/Engine/Source/Runtime/Renderer/Private/PostProcess/SceneRenderTargets.h b/Engine/Source/Runtime/Renderer/Private/PostProcess/SceneRenderTargets.h
--- a/Engine/Source/Runtime/Renderer/Private/PostProcess/SceneRenderTargets.h
+++ b/Engine/Source/Runtime/Renderer/Private/PostProcess/SceneRenderTargets.h
@@ -24,6 +24,8 @@
 #define STENCIL_TEMPORAL_RESPONSIVE_AA_BIT_ID 3
 #define STENCIL_LIGHTING_CHANNELS_BIT_ID	4
 #define STENCIL_RECEIVE_DECAL_BIT_ID		7
+// Used only during the lighting pass - alias/reuse decal bit
+#define STENCIL_STRATA_FASTPATH				7 
 
 // Outputs a compile-time constant stencil's bit mask ready to be used
 // in TStaticDepthStencilState<> template parameter. It also takes care
diff --git a/Engine/Source/Runtime/Renderer/Private/ShadowRendering.cpp b/Engine/Source/Runtime/Renderer/Private/ShadowRendering.cpp
--- a/Engine/Source/Runtime/Renderer/Private/ShadowRendering.cpp
+++ b/Engine/Source/Runtime/Renderer/Private/ShadowRendering.cpp
@@ -210,6 +210,9 @@ static TAutoConsoleVariable<int32> CVarVirtualSmDenoiserTemporal(
 
 DEFINE_GPU_DRAWCALL_STAT(ShadowProjection);
 
+// Use Shadow stencil mask is set to 0x7F instead of 0xFF so that that last bit can be used for strata classification without clearing the stencil bit for pre-shadow/per-object static shadow mask
+constexpr uint32 ShadowStencilMask = 0x7f;
+
 // 0:off, 1:low, 2:med, 3:high, 4:very high, 5:max
 uint32 GetShadowQuality()
 {
@@ -856,7 +859,7 @@ void FProjectedShadowInfo::SetupProjectionStencilMask(
 			false, CF_DepthNearOrEqual,
 			true, CF_Always, SO_Keep, SO_Increment, SO_Keep,
 			true, CF_Always, SO_Keep, SO_Decrement, SO_Keep,
-			0xff, 0xff
+			ShadowStencilMask, ShadowStencilMask
 			>::GetRHI());
 
 		FGraphicsPipelineStateInitializer GraphicsPSOInit;
@@ -927,7 +930,7 @@ void FProjectedShadowInfo::SetupProjectionStencilMask(
 				false, CF_DepthNearOrEqual,
 				true, CF_Always, SO_Keep, SO_Increment, SO_Keep,
 				true, CF_Always, SO_Keep, SO_Decrement, SO_Keep,
-				0xff, 0xff
+				ShadowStencilMask, ShadowStencilMask
 				>::GetRHI());
 		}
 		else
@@ -939,7 +942,7 @@ void FProjectedShadowInfo::SetupProjectionStencilMask(
 				false, CF_DepthNearOrEqual,
 				true, CF_Always, SO_Keep, SO_Keep, SO_Increment,
 				true, CF_Always, SO_Keep, SO_Keep, SO_Decrement,
-				0xff, 0xff
+				ShadowStencilMask, ShadowStencilMask
 				>::GetRHI());
 		}
 		
@@ -1154,7 +1157,7 @@ void FProjectedShadowInfo::RenderProjection(
 					false, CF_Always,
 					true, CF_NotEqual, SO_Zero, SO_Zero, SO_Zero,
 					false, CF_Always, SO_Zero, SO_Zero, SO_Zero,
-					0xff, 0xff
+					ShadowStencilMask, ShadowStencilMask
 					>::GetRHI();
 			}
 			else
@@ -1165,7 +1168,7 @@ void FProjectedShadowInfo::RenderProjection(
 					false, CF_Always,
 					true, CF_NotEqual, SO_Keep, SO_Keep, SO_Keep,
 					false, CF_Always, SO_Keep, SO_Keep, SO_Keep,
-					0xff, 0xff
+					ShadowStencilMask, ShadowStencilMask
 					>::GetRHI();
 			}
 		}
diff --git a/Engine/Source/Runtime/Renderer/Private/SingleLayerWaterRendering.cpp b/Engine/Source/Runtime/Renderer/Private/SingleLayerWaterRendering.cpp
--- a/Engine/Source/Runtime/Renderer/Private/SingleLayerWaterRendering.cpp
+++ b/Engine/Source/Runtime/Renderer/Private/SingleLayerWaterRendering.cpp
@@ -424,7 +424,7 @@ void FDeferredShadingSceneRenderer::RenderSingleLayerWaterReflections(
 		// We need to update the classification texture once again because the GBuffer has been patched with single layer water data.
 		if (Strata::IsStrataEnabled() && bRunTiled)
 		{
-			Strata::AddStrataMaterialClassificationPass(GraphBuilder, Views);
+			Strata::AddStrataMaterialClassificationPass(GraphBuilder, SceneTextures, Views);
 		}
 
 
diff --git a/Engine/Source/Runtime/Renderer/Private/Strata/Strata.cpp b/Engine/Source/Runtime/Renderer/Private/Strata/Strata.cpp
--- a/Engine/Source/Runtime/Renderer/Private/Strata/Strata.cpp
+++ b/Engine/Source/Runtime/Renderer/Private/Strata/Strata.cpp
@@ -25,6 +25,17 @@ static TAutoConsoleVariable<int32> CVarStrataBytePerPixel(
 	TEXT("Strata allocated byte per pixel to store materials data. Higher value means more complex material can be represented."),
 	ECVF_ReadOnly | ECVF_RenderThreadSafe);
 
+static TAutoConsoleVariable<int32> CVarStrataClassification(
+	TEXT("r.Strata.Classification"),
+	1,
+	TEXT("Enable strata classification to speed up lighting pass."),
+	ECVF_RenderThreadSafe); 
+
+static TAutoConsoleVariable<int32> CVarStrataTileInPixel(
+	TEXT("r.Strata.TileInPixel"),
+	64,
+	TEXT("Strata allocated byte per pixel to store materials data. Higher value means more complex material can be represented."),
+	ECVF_RenderThreadSafe);
 
 IMPLEMENT_GLOBAL_SHADER_PARAMETER_STRUCT(FStrataGlobalUniformParameters, "Strata");
 
@@ -36,6 +47,16 @@ bool IsStrataEnabled()
 	return CVarStrata.GetValueOnAnyThread() > 0;
 }
 
+bool IsClassificationEnabled()
+{
+	return CVarStrataClassification.GetValueOnAnyThread() > 0;
+}
+
+uint32 GetStrataTileSize()
+{
+	return FMath::RoundUpToPowerOfTwo(FMath::Clamp(CVarStrataTileInPixel.GetValueOnRenderThread(), 8, 128));
+}
+
 void InitialiseStrataFrameSceneData(FSceneRenderer& SceneRenderer, FRDGBuilder& GraphBuilder)
 {
 	FStrataSceneData& StrataSceneData = SceneRenderer.Scene->StrataSceneData;
@@ -63,6 +84,17 @@ void InitialiseStrataFrameSceneData(FSceneRenderer& SceneRenderer, FRDGBuilder&
 			StrataSceneData.ClassificationTexture = GraphBuilder.GetPooledTexture(Texture);
 		}
 
+		// Tile classification texture
+		{
+			const int32 TileInPixel = GetStrataTileSize();
+			const uint32 TileResolutionX = FMath::DivideAndRoundUp(BufferSizeXY.X, TileInPixel);
+			const uint32 TileResolutionY = FMath::DivideAndRoundUp(BufferSizeXY.Y, TileInPixel);
+
+			FRDGTextureRef Texture = GraphBuilder.CreateTexture(FRDGTextureDesc::Create2D(FIntPoint(TileResolutionX, TileResolutionY), PF_R32_UINT, FClearValueBinding::Black, TexCreate_ShaderResource | TexCreate_UAV), TEXT("StrataTileClassificationTexture"));
+			GraphBuilder.PreallocateTexture(Texture);
+			StrataSceneData.TileClassificationTexture = GraphBuilder.GetPooledTexture(Texture);
+		}
+
 		// Top layer texture
 		{
 			FRDGTextureRef Texture = GraphBuilder.CreateTexture(FRDGTextureDesc::Create2D(BufferSizeXY, PF_R32G32_UINT, FClearValueBinding::Black, TexCreate_ShaderResource | TexCreate_RenderTargetable), TEXT("StrataTopLayerTexture"));
@@ -257,11 +289,212 @@ class FStrataMaterialClassificationPassPS : public FGlobalShader
 	static void ModifyCompilationEnvironment(const FGlobalShaderPermutationParameters& Parameters, FShaderCompilerEnvironment& OutEnvironment)
 	{
 		FGlobalShader::ModifyCompilationEnvironment(Parameters, OutEnvironment);
+		OutEnvironment.SetDefine(TEXT("SHADER_CATEGORIZATION"), 1);
 	}
 };
 IMPLEMENT_GLOBAL_SHADER(FStrataMaterialClassificationPassPS, "/Engine/Private/Strata/StrataMaterialClassification.usf", "MainPS", SF_Pixel);
 
-void AddStrataMaterialClassificationPass(FRDGBuilder& GraphBuilder, const TArray<FViewInfo>& Views)
+////////////////////////////////////////////////////////////////////////////////////////////////////////////////////////////////////////////////
+
+class FStrataMaterialTileClassificationPassCS : public FGlobalShader
+{
+	DECLARE_GLOBAL_SHADER(FStrataMaterialTileClassificationPassCS);
+	SHADER_USE_PARAMETER_STRUCT(FStrataMaterialTileClassificationPassCS, FGlobalShader);
+
+	using FPermutationDomain = TShaderPermutationDomain<>;
+
+	BEGIN_SHADER_PARAMETER_STRUCT(FParameters, )
+		SHADER_PARAMETER_STRUCT_REF(FViewUniformShaderParameters, ViewUniformBuffer)
+		SHADER_PARAMETER(int32, TileSize)
+		SHADER_PARAMETER(FIntPoint, ClassificationResolution)
+		SHADER_PARAMETER(FIntPoint, TileClassificationResolution)
+		SHADER_PARAMETER_RDG_TEXTURE(Texture2D, ClassificationTexture)
+		SHADER_PARAMETER_RDG_TEXTURE_UAV(RWTexture2D, TileClassificationTexture)
+	END_SHADER_PARAMETER_STRUCT()
+
+	static FPermutationDomain RemapPermutation(FPermutationDomain PermutationVector)
+	{
+		return PermutationVector;
+	}
+
+	static bool ShouldCompilePermutation(const FGlobalShaderPermutationParameters& Parameters)
+	{
+		return GetMaxSupportedFeatureLevel(Parameters.Platform) >= ERHIFeatureLevel::SM5 && Strata::IsStrataEnabled();
+	}
+
+	static void ModifyCompilationEnvironment(const FGlobalShaderPermutationParameters& Parameters, FShaderCompilerEnvironment& OutEnvironment)
+	{
+		FGlobalShader::ModifyCompilationEnvironment(Parameters, OutEnvironment);
+		OutEnvironment.SetDefine(TEXT("SHADER_TILE_CATEGORIZATION"), 1);
+	}
+};
+IMPLEMENT_GLOBAL_SHADER(FStrataMaterialTileClassificationPassCS, "/Engine/Private/Strata/StrataMaterialClassification.usf", "TileMainCS", SF_Compute);
+
+////////////////////////////////////////////////////////////////////////////////////////////////////////////////////////////////////////////////
+
+
+class FStrataMaterialStencilClassificationPassVS : public FGlobalShader
+{
+	DECLARE_GLOBAL_SHADER(FStrataMaterialStencilClassificationPassVS);
+	SHADER_USE_PARAMETER_STRUCT(FStrataMaterialStencilClassificationPassVS, FGlobalShader);
+
+	using FPermutationDomain = TShaderPermutationDomain<>;
+
+	BEGIN_SHADER_PARAMETER_STRUCT(FParameters, )
+		SHADER_PARAMETER_STRUCT_REF(FViewUniformShaderParameters, ViewUniformBuffer)
+		SHADER_PARAMETER(int32, TileSize)
+		SHADER_PARAMETER(FIntPoint, TileCount)
+		SHADER_PARAMETER(FIntPoint, OutputResolution)
+		SHADER_PARAMETER_RDG_TEXTURE(Texture2D, TileClassificationTexture)
+	END_SHADER_PARAMETER_STRUCT()
+
+	static FPermutationDomain RemapPermutation(FPermutationDomain PermutationVector)
+	{
+		return PermutationVector;
+	}
+
+	static bool ShouldCompilePermutation(const FGlobalShaderPermutationParameters& Parameters)
+	{
+		return GetMaxSupportedFeatureLevel(Parameters.Platform) >= ERHIFeatureLevel::SM5 && Strata::IsStrataEnabled();
+	}
+
+	static void ModifyCompilationEnvironment(const FGlobalShaderPermutationParameters& Parameters, FShaderCompilerEnvironment& OutEnvironment)
+	{
+		FGlobalShader::ModifyCompilationEnvironment(Parameters, OutEnvironment);
+		OutEnvironment.SetDefine(TEXT("SHADER_STENCIL_CATEGORIZATION"), 1);
+	}
+};
+
+class FStrataMaterialStencilClassificationPassPS : public FGlobalShader
+{
+	DECLARE_GLOBAL_SHADER(FStrataMaterialStencilClassificationPassPS);
+	SHADER_USE_PARAMETER_STRUCT(FStrataMaterialStencilClassificationPassPS, FGlobalShader);
+
+	using FPermutationDomain = TShaderPermutationDomain<>;
+
+	BEGIN_SHADER_PARAMETER_STRUCT(FParameters, )
+		SHADER_PARAMETER_STRUCT_REF(FViewUniformShaderParameters, ViewUniformBuffer)
+		SHADER_PARAMETER(int32, TileSize)
+		SHADER_PARAMETER(FIntPoint, TileCount)
+		SHADER_PARAMETER(FIntPoint, OutputResolution)
+		SHADER_PARAMETER_RDG_TEXTURE(Texture2D, TileClassificationTexture)
+		RENDER_TARGET_BINDING_SLOTS()
+	END_SHADER_PARAMETER_STRUCT()
+
+		static FPermutationDomain RemapPermutation(FPermutationDomain PermutationVector)
+	{
+		return PermutationVector;
+	}
+
+	static bool ShouldCompilePermutation(const FGlobalShaderPermutationParameters& Parameters)
+	{
+		return GetMaxSupportedFeatureLevel(Parameters.Platform) >= ERHIFeatureLevel::SM5 && Strata::IsStrataEnabled();
+	}
+
+	static void ModifyCompilationEnvironment(const FGlobalShaderPermutationParameters& Parameters, FShaderCompilerEnvironment& OutEnvironment)
+	{
+		FGlobalShader::ModifyCompilationEnvironment(Parameters, OutEnvironment);
+		OutEnvironment.SetDefine(TEXT("SHADER_STENCIL_CATEGORIZATION"), 1); 
+	}
+};
+
+IMPLEMENT_GLOBAL_SHADER(FStrataMaterialStencilClassificationPassVS, "/Engine/Private/Strata/StrataMaterialClassification.usf", "StencilMainVS", SF_Vertex);
+IMPLEMENT_GLOBAL_SHADER(FStrataMaterialStencilClassificationPassPS, "/Engine/Private/Strata/StrataMaterialClassification.usf", "StencilMainPS", SF_Pixel);
+
+static void AddStrataStencilPass(
+	FRDGBuilder& GraphBuilder,
+	const FViewInfo& View,
+	const FMinimalSceneTextures& SceneTextures,
+	FRDGTextureRef TileClassificationTexture)
+{	
+	const FIntPoint OutputResolution = View.ViewRect.Size();
+	const int32 TileSize = GetStrataTileSize();
+	const FIntPoint TileCount(FMath::DivideAndRoundUp(OutputResolution.X, TileSize), FMath::DivideAndRoundUp(OutputResolution.Y, TileSize));
+	const uint32 InstanceCount = TileCount.X * TileCount.Y;
+
+	FStrataMaterialStencilClassificationPassPS::FParameters* ParametersPS = GraphBuilder.AllocParameters<FStrataMaterialStencilClassificationPassPS::FParameters>();
+	ParametersPS->TileSize = TileSize;
+	ParametersPS->TileCount = TileCount;
+	ParametersPS->OutputResolution = OutputResolution;
+	ParametersPS->TileClassificationTexture = TileClassificationTexture;
+
+	TShaderMapRef<FStrataMaterialStencilClassificationPassVS> VertexShader(View.ShaderMap);
+	TShaderMapRef<FStrataMaterialStencilClassificationPassPS> PixelShader(View.ShaderMap);
+
+	// For debug purpose
+	#if 0
+	FRDGTextureDesc DummyDesc = FRDGTextureDesc::Create2D(SceneTextures.Depth.Target->Desc.Extent, EPixelFormat::PF_R8G8B8A8, FClearValueBinding::Black, ETextureCreateFlags::TexCreate_RenderTargetable);
+	FRDGTextureRef DummyTexture = GraphBuilder.CreateTexture(DummyDesc, TEXT("StencilClassificationOutput"));
+	ParametersPS->RenderTargets[0] = FRenderTargetBinding(DummyTexture, ERenderTargetLoadAction::EClear);
+	#endif
+
+	ParametersPS->RenderTargets.DepthStencil = FDepthStencilBinding(
+		SceneTextures.Depth.Target,
+		ERenderTargetLoadAction::ELoad,
+		ERenderTargetLoadAction::ELoad,
+		FExclusiveDepthStencil::DepthNop_StencilWrite);
+
+	GraphBuilder.AddPass(
+		RDG_EVENT_NAME("StrataStencilClassificationPass"),
+		ParametersPS,
+		ERDGPassFlags::Raster,
+		[ParametersPS, VertexShader, PixelShader, OutputResolution, InstanceCount](FRHICommandList& RHICmdList)
+		{
+			FStrataMaterialStencilClassificationPassVS::FParameters ParametersVS;
+			ParametersVS.TileSize = ParametersPS->TileSize;
+			ParametersVS.TileCount = ParametersPS->TileCount;
+			ParametersVS.OutputResolution = ParametersPS->OutputResolution;
+			ParametersVS.TileClassificationTexture = ParametersPS->TileClassificationTexture;
+
+			FGraphicsPipelineStateInitializer GraphicsPSOInit;
+			RHICmdList.ApplyCachedRenderTargets(GraphicsPSOInit);
+			GraphicsPSOInit.BlendState = TStaticBlendState<CW_RGBA, BO_Add, BF_One, BF_One, BO_Max, BF_SourceAlpha, BF_DestAlpha>::GetRHI();
+			GraphicsPSOInit.RasterizerState = TStaticRasterizerState<>::GetRHI();
+			GraphicsPSOInit.DepthStencilState = TStaticDepthStencilState<
+				false, CF_Always, 
+				true,  CF_Always, SO_Keep, SO_Keep, SO_Replace,
+				false, CF_Always, SO_Keep, SO_Keep, SO_Keep,
+				0xFF, StencilBit>::GetRHI();
+			GraphicsPSOInit.BoundShaderState.VertexDeclarationRHI = GFilterVertexDeclaration.VertexDeclarationRHI;
+			GraphicsPSOInit.BoundShaderState.VertexShaderRHI = VertexShader.GetVertexShader();
+			GraphicsPSOInit.BoundShaderState.PixelShaderRHI = PixelShader.GetPixelShader();
+			GraphicsPSOInit.PrimitiveType = PT_TriangleList;
+			SetGraphicsPipelineState(RHICmdList, GraphicsPSOInit);
+			SetShaderParameters(RHICmdList, VertexShader, VertexShader.GetVertexShader(), ParametersVS);
+			//SetShaderParameters(RHICmdList, PixelShader, PixelShader.GetPixelShader(), *ParametersPS);
+
+			RHICmdList.SetStencilRef(StencilBit);
+			RHICmdList.SetViewport(0, 0, 0.0f, OutputResolution.X, OutputResolution.Y, 1.0f);
+			RHICmdList.SetStreamSource(0, nullptr, 0);
+			RHICmdList.DrawPrimitive(0, 2, InstanceCount);
+		});
+}
+
+void AddStrataStencilPass(
+	FRDGBuilder& GraphBuilder,
+	const FViewInfo& View,
+	const FMinimalSceneTextures& SceneTextures)
+{
+	FRDGTextureRef TileClassificationTexture = GraphBuilder.RegisterExternalTexture(View.StrataSceneData->TileClassificationTexture);
+	AddStrataStencilPass(GraphBuilder, View, SceneTextures, TileClassificationTexture);
+}
+
+void AddStrataStencilPass(
+	FRDGBuilder& GraphBuilder,
+	const TArray<FViewInfo>& Views,
+	const FMinimalSceneTextures& SceneTextures)
+{
+	for (int32 i = 0; i < Views.Num(); ++i)
+	{
+		const FViewInfo& View = Views[i];
+		FRDGTextureRef TileClassificationTexture = GraphBuilder.RegisterExternalTexture(View.StrataSceneData->TileClassificationTexture);
+		AddStrataStencilPass(GraphBuilder, View, SceneTextures, TileClassificationTexture);
+	}
+}
+
+////////////////////////////////////////////////////////////////////////////////////////////////////////////////////////////////////////////////
+
+void AddStrataMaterialClassificationPass(FRDGBuilder& GraphBuilder, const FMinimalSceneTextures& SceneTextures, const TArray<FViewInfo>& Views)
 {
 	RDG_EVENT_SCOPE_CONDITIONAL(GraphBuilder, IsStrataEnabled() && Views.Num() > 0, "StrataMaterialClassification");
 	if (!IsStrataEnabled())
@@ -273,31 +506,60 @@ void AddStrataMaterialClassificationPass(FRDGBuilder& GraphBuilder, const TArray
 	{
 		const FViewInfo& View = Views[i];
 		
-		FStrataMaterialClassificationPassPS::FPermutationDomain PermutationVector;
-		TShaderMapRef<FStrataMaterialClassificationPassPS> PixelShader(View.ShaderMap, PermutationVector);
-		FStrataMaterialClassificationPassPS::FParameters* PassParameters = GraphBuilder.AllocParameters<FStrataMaterialClassificationPassPS::FParameters>();
-		PassParameters->ViewUniformBuffer = View.ViewUniformBuffer;
-		PassParameters->Strata = Strata::BindStrataGlobalUniformParameters(View);
-		PassParameters->SceneTextures = GetSceneTextureParameters(GraphBuilder);		
-		PassParameters->RenderTargets[0] = FRenderTargetBinding(GraphBuilder.RegisterExternalTexture(View.StrataSceneData->ClassificationTexture), ERenderTargetLoadAction::EClear);
-		PassParameters->RenderTargets[1] = FRenderTargetBinding(GraphBuilder.RegisterExternalTexture(View.StrataSceneData->TopLayerTexture), ERenderTargetLoadAction::EClear);
-
-		if (ShaderDrawDebug::IsShaderDrawDebugEnabled())
+		// Classification
+		FRDGTextureRef ClassificationTexture = GraphBuilder.RegisterExternalTexture(View.StrataSceneData->ClassificationTexture);
 		{
-			ShaderDrawDebug::SetParameters(GraphBuilder, View.ShaderDrawData, PassParameters->ShaderDrawParameters);
+			FStrataMaterialClassificationPassPS::FPermutationDomain PermutationVector;
+			TShaderMapRef<FStrataMaterialClassificationPassPS> PixelShader(View.ShaderMap, PermutationVector);
+			FStrataMaterialClassificationPassPS::FParameters* PassParameters = GraphBuilder.AllocParameters<FStrataMaterialClassificationPassPS::FParameters>();
+			PassParameters->ViewUniformBuffer = View.ViewUniformBuffer;
+			PassParameters->Strata = Strata::BindStrataGlobalUniformParameters(View);
+			PassParameters->SceneTextures = GetSceneTextureParameters(GraphBuilder);		
+			PassParameters->RenderTargets[0] = FRenderTargetBinding(ClassificationTexture, ERenderTargetLoadAction::EClear);
+			PassParameters->RenderTargets[1] = FRenderTargetBinding(GraphBuilder.RegisterExternalTexture(View.StrataSceneData->TopLayerTexture), ERenderTargetLoadAction::EClear);
+
+			if (ShaderDrawDebug::IsShaderDrawDebugEnabled())
+			{
+				ShaderDrawDebug::SetParameters(GraphBuilder, View.ShaderDrawData, PassParameters->ShaderDrawParameters);
+			}
+
+			FPixelShaderUtils::AddFullscreenPass<FStrataMaterialClassificationPassPS>(
+				GraphBuilder, 
+				View.ShaderMap, 
+				RDG_EVENT_NAME("StrataMaterialClassification"),
+				PixelShader, 
+				PassParameters, 
+				View.ViewRect);
 		}
 
-		FPixelShaderUtils::AddFullscreenPass<FStrataMaterialClassificationPassPS>(
-			GraphBuilder, 
-			View.ShaderMap, 
-			RDG_EVENT_NAME("StrataMaterialClassification"),
-			PixelShader, 
-			PassParameters, 
-			View.ViewRect);
+		// Downsampling
+		FRDGTextureRef TileClassificationTexture = GraphBuilder.RegisterExternalTexture(View.StrataSceneData->TileClassificationTexture);
+		if (IsClassificationEnabled())
+		{
+			const uint32 ClearValues[4] = { 0u, 0u, 0u, 0u };
+			FRDGTextureUAVRef TileClassificationTextureUAV = GraphBuilder.CreateUAV(TileClassificationTexture);
+			AddClearUAVPass(GraphBuilder, TileClassificationTextureUAV, ClearValues);
+
+			TShaderMapRef<FStrataMaterialTileClassificationPassCS> ComputeShader(View.ShaderMap);
+			FStrataMaterialTileClassificationPassCS::FParameters* PassParameters = GraphBuilder.AllocParameters<FStrataMaterialTileClassificationPassCS::FParameters>();
+			PassParameters->TileSize = GetStrataTileSize();
+			PassParameters->ClassificationResolution = ClassificationTexture->Desc.Extent;
+			PassParameters->TileClassificationResolution = TileClassificationTexture->Desc.Extent;
+			PassParameters->ClassificationTexture = ClassificationTexture;
+			PassParameters->TileClassificationTexture = TileClassificationTextureUAV;
+
+			// Add 64 threads permutation
+			const uint32 GroupSize = 8;
+			FComputeShaderUtils::AddPass(
+				GraphBuilder,
+				RDG_EVENT_NAME("StrataMaterialTileClassification"),
+				ComputeShader,
+				PassParameters,
+				FComputeShaderUtils::GetGroupCount(ClassificationTexture->Desc.Extent, GroupSize));
+		}
 	}
 }
 
-
 } // namespace Strata
 
 
diff --git a/Engine/Source/Runtime/Renderer/Private/Strata/Strata.h b/Engine/Source/Runtime/Renderer/Private/Strata/Strata.h
--- a/Engine/Source/Runtime/Renderer/Private/Strata/Strata.h
+++ b/Engine/Source/Runtime/Renderer/Private/Strata/Strata.h
@@ -12,8 +12,7 @@
 // Forward declarations.
 class FScene;
 class FRDGBuilder;
-
-
+struct FMinimalSceneTextures;
 
 BEGIN_SHADER_PARAMETER_STRUCT(FStrataOpaquePassUniformParameters, )
 	SHADER_PARAMETER(uint32, MaxBytesPerPixel)
@@ -30,9 +29,10 @@ END_GLOBAL_SHADER_PARAMETER_STRUCT()
 struct FStrataSceneData
 {
 	uint32 MaxBytesPerPixel;
-	FRWByteAddressBuffer MaterialLobesBuffer;				// This should be a RDG resource
-	TRefCountPtr<IPooledRenderTarget> ClassificationTexture;// This should be a RDG resource
-	TRefCountPtr<IPooledRenderTarget> TopLayerTexture;		// This should be a RDG resource
+	FRWByteAddressBuffer MaterialLobesBuffer;					// This should be a RDG resource
+	TRefCountPtr<IPooledRenderTarget> ClassificationTexture;	// This should be a RDG resource
+	TRefCountPtr<IPooledRenderTarget> TopLayerTexture;			// This should be a RDG resource
+	TRefCountPtr<IPooledRenderTarget> TileClassificationTexture;// This should be a RDG resource
 
 	TUniformBufferRef<FStrataGlobalUniformParameters> StrataGlobalUniformParameters;
 
@@ -43,8 +43,10 @@ struct FStrataSceneData
 
 namespace Strata
 {
+constexpr uint32 StencilBit = 0x80; // In sync with SceneRenderTargets.h - GET_STENCIL_BIT_MASK(STENCIL_STRATA_FASTPATH)
 
 bool IsStrataEnabled();
+bool IsClassificationEnabled();
 
 void InitialiseStrataFrameSceneData(FSceneRenderer& SceneRenderer, FRDGBuilder& GraphBuilder);
 
@@ -54,8 +56,10 @@ TUniformBufferRef<FStrataGlobalUniformParameters> BindStrataGlobalUniformParamet
 
 void AddVisualizeMaterialPasses(FRDGBuilder& GraphBuilder, const TArray<FViewInfo>& Views, FRDGTextureRef SceneColorTexture);
 
-void AddStrataMaterialClassificationPass(FRDGBuilder& GraphBuilder, const TArray<FViewInfo>& Views);
+void AddStrataMaterialClassificationPass(FRDGBuilder& GraphBuilder, const FMinimalSceneTextures& SceneTextures, const TArray<FViewInfo>& Views);
 
+void AddStrataStencilPass(FRDGBuilder& GraphBuilder, const FViewInfo& View, const FMinimalSceneTextures& SceneTextures);
+void AddStrataStencilPass(FRDGBuilder& GraphBuilder, const TArray<FViewInfo>& Views, const FMinimalSceneTextures& SceneTextures);
 };
 
 
