Removed hack which force depth to be kept on OpenGLES
Fixed issue with no Depth in distortion pass

diff --git a/Engine/Source/Runtime/Renderer/Private/MobileShadingRenderer.cpp b/Engine/Source/Runtime/Renderer/Private/MobileShadingRenderer.cpp
--- a/Engine/Source/Runtime/Renderer/Private/MobileShadingRenderer.cpp
+++ b/Engine/Source/Runtime/Renderer/Private/MobileShadingRenderer.cpp
@@ -383,17 +383,16 @@ void FMobileSceneRenderer::InitViews(FRDGBuilder& GraphBuilder, FSceneTexturesCo
 	// See CVarMobileForceDepthResolve use in ConditionalResolveSceneDepth.
 	const bool bForceDepthResolve = (CVarMobileForceDepthResolve.GetValueOnRenderThread() == 1);
 	const bool bSeparateTranslucencyActive = IsMobileSeparateTranslucencyActive(Views.GetData(), Views.Num()); 
-	const bool bPostProcessUsesSceneDepth = PostProcessUsesSceneDepth(Views[0]);
+	const bool bPostProcessUsesSceneDepth = PostProcessUsesSceneDepth(Views[0]) || IsMobileDistortionActive(Views[0]);
 	bRequiresMultiPass = RequiresMultiPass(RHICmdList, Views[0]);
-	bKeepDepthContent = 
-		bRequiresMultiPass || 
+	bKeepDepthContent =
+		bRequiresMultiPass ||
 		bForceDepthResolve ||
 		bRequriesAmbientOcclusionPass ||
 		bRequiresPixelProjectedPlanarRelfectionPass ||
 		bSeparateTranslucencyActive ||
 		Views[0].bIsReflectionCapture ||
-		(bDeferredShading && bPostProcessUsesSceneDepth) ||
-		IsAndroidOpenGLESPlatform(ShaderPlatform);
+		(bDeferredShading && bPostProcessUsesSceneDepth);
 	// never keep MSAA depth
 	bKeepDepthContent = (NumMSAASamples > 1 ? false : bKeepDepthContent);
 
@@ -1216,31 +1215,6 @@ void FMobileSceneRenderer::RenderDeferred(FRDGBuilder& GraphBuilder, const TArra
 			AddPass(GraphBuilder, PollOcclusionQueriesAndDispatchToRHIThreadPass);
 		}
 
-		/*if (IsAndroidOpenGLESPlatform(ShaderPlatform))
-		{
-			FRDGTextureRef DepthCopyTargets[] = 
-			{
-				SceneTextures.Depth.Target,
-			};
-			
-			TArrayView<FRDGTextureRef> DepthCopyTargetsView = MakeArrayView(DepthCopyTargets);
-
-			FRenderTargetBindingSlots DepthCopyRenderTargets = GetRenderTargetBindings(ERenderTargetLoadAction::EClear, DepthCopyTargetsView);
-			DepthCopyRenderTargets.DepthStencil = FDepthStencilBinding(SceneTextures.Depth.Target, ERenderTargetLoadAction::ELoad, FExclusiveDepthStencil::DepthRead_StencilRead);
-			DepthCopyRenderTargets.NumOcclusionQueries = 0;
-			DepthCopyRenderTargets.FoveationTexture = nullptr;
-			DepthCopyRenderTargets.MultiViewCount = 0;
-
-			auto* PassParameters = GraphBuilder.AllocParameters<FRenderTargetParameters>();
-			PassParameters->RenderTargets = DepthCopyRenderTargets;
-
-			GraphBuilder.AddPass(RDG_EVENT_NAME("FBFDepthCopy"), PassParameters, ERDGPassFlags::Raster,
-				[this](FRHICommandListImmediate& RHICmdList)
-			{
-				MobileDeferredCopyBuffer<FMobileDeferredCopyDepthPS>(RHICmdList, Views[0]);
-			});
-		}*/
-
 		AddPass(GraphBuilder, [](FRHICommandListImmediate& RHICmdList)
 		{
 			RHICmdList.EndRenderPass();
