Misc Instance culling fixes
- Fix buffer type error in instance culling.
- made sure non-deterministic instance culling code-path still works.

#rnx

diff --git a/Engine/Shaders/Private/InstanceCulling/BuildInstanceDrawCommands.usf b/Engine/Shaders/Private/InstanceCulling/BuildInstanceDrawCommands.usf
--- a/Engine/Shaders/Private/InstanceCulling/BuildInstanceDrawCommands.usf
+++ b/Engine/Shaders/Private/InstanceCulling/BuildInstanceDrawCommands.usf
@@ -222,7 +222,7 @@ void CalcOutputOffsets(uint GroupThreadIndex : SV_GroupIndex, uint GroupId : SV_
 	}
 }
 
-StructuredBuffer<uint> InstanceIdOffsetBuffer;
+Buffer<uint> InstanceIdOffsetBuffer;
 
 [numthreads(NUM_THREADS_PER_GROUP, 1, 1)]
 void OutputInstanceIdsAtOffset(uint GroupThreadIndex : SV_GroupIndex, uint GroupId : SV_GroupID)
diff --git a/Engine/Source/Runtime/Renderer/Private/InstanceCulling/InstanceCullingContext.cpp b/Engine/Source/Runtime/Renderer/Private/InstanceCulling/InstanceCullingContext.cpp
--- a/Engine/Source/Runtime/Renderer/Private/InstanceCulling/InstanceCullingContext.cpp
+++ b/Engine/Source/Runtime/Renderer/Private/InstanceCulling/InstanceCullingContext.cpp
@@ -184,7 +184,7 @@ public:
 	}
 
 	BEGIN_SHADER_PARAMETER_STRUCT(FParameters, )
-		SHADER_PARAMETER_RDG_BUFFER_SRV(StructuredBuffer<uint>, InstanceIdOffsetBuffer)
+		SHADER_PARAMETER_RDG_BUFFER_SRV(Buffer<uint>, InstanceIdOffsetBuffer)
 		SHADER_PARAMETER_RDG_BUFFER_SRV(StructuredBuffer<uint2>, InstanceCounts)
 
 		SHADER_PARAMETER_SRV(StructuredBuffer<float4>, GPUSceneInstanceSceneData)
@@ -225,7 +225,8 @@ public:
 
 	// GPUCULL_TODO: remove once buffer is somehow unified
 	class FOutputCommandIdDim : SHADER_PERMUTATION_BOOL("OUTPUT_COMMAND_IDS");
-	using FPermutationDomain = TShaderPermutationDomain<FOutputCommandIdDim>;
+	class FCullInstancesDim : SHADER_PERMUTATION_BOOL("CULL_INSTANCES");
+	using FPermutationDomain = TShaderPermutationDomain<FOutputCommandIdDim, FCullInstancesDim>;
 
 	static bool ShouldCompilePermutation(const FGlobalShaderPermutationParameters& Parameters)
 	{
@@ -261,7 +262,6 @@ public:
 		// Using the wrong kind of buffer for RDG...
 		SHADER_PARAMETER_UAV(RWBuffer<uint>, InstanceIdsBufferLegacyOut)
 
-		SHADER_PARAMETER_RDG_BUFFER_SRV(Buffer<uint>, InstanceIdOffsetBuffer)
 		SHADER_PARAMETER_RDG_BUFFER_UAV(RWBuffer<uint>, DrawIndirectArgsBufferOut)
 		SHADER_PARAMETER(int32, NumPrimitiveIds)
 		SHADER_PARAMETER(int32, NumInstanceRuns)		
@@ -405,7 +405,7 @@ void FInstanceCullingContext::BuildRenderingCommands(FRDGBuilder& GraphBuilder,
 
 		FComputeShaderUtils::AddPass(
 			GraphBuilder,
-			RDG_EVENT_NAME("CalcOutputOffsets"),
+			RDG_EVENT_NAME("OutputInstanceIdsAtOffset"),
 			ComputeShader,
 			PassParameters,
 			FIntVector(CullingCommands.Num(), 1, 1)
@@ -428,9 +428,6 @@ void FInstanceCullingContext::BuildRenderingCommands(FRDGBuilder& GraphBuilder,
 	PassParameters->PrimitiveIds = GraphBuilder.CreateSRV(CreateStructuredBuffer(GraphBuilder, TEXT("PrimitiveIds"), PrimitiveIds));
 	PassParameters->InstanceRuns = GraphBuilder.CreateSRV(CreateStructuredBuffer(GraphBuilder, TEXT("InstanceRuns"), InstanceRuns));
 
-	FRDGBufferRef InstanceIdOutOffsetBufferRDG = Intermediate.InstanceIdOutOffsetBuffer;
-	FRDGBufferRef VisibleInstanceFlagsRDG = Intermediate.VisibleInstanceFlags;
-
 	PassParameters->OutputOffsetBufferOut = GraphBuilder.CreateUAV(InstanceIdOutOffsetBufferRDG);
 
 	FRDGBufferRef DrawIndirectArgsRDG = GraphBuilder.CreateBuffer(FRDGBufferDesc::CreateIndirectDesc(IndirectArgsNumWords * CullingCommands.Num()), TEXT("DrawIndirectArgsBuffer"));
@@ -462,7 +459,7 @@ void FInstanceCullingContext::BuildRenderingCommands(FRDGBuilder& GraphBuilder,
 	FBuildInstanceIdBufferAndCommandsFromPrimitiveIdsCs::FPermutationDomain PermutationVector;
 	// NOTE: this also switches between legacy buffer and RDG for Id output
 	PermutationVector.Set<FBuildInstanceIdBufferAndCommandsFromPrimitiveIdsCs::FOutputCommandIdDim>(0);
-	PermutationVector.Set<FComputeInstanceIdOutputSizeCs::FCullInstancesDim>(bCullInstances);
+	PermutationVector.Set<FBuildInstanceIdBufferAndCommandsFromPrimitiveIdsCs::FCullInstancesDim>(bCullInstances);
 
 	auto ComputeShader = GetGlobalShaderMap(GMaxRHIFeatureLevel)->GetShader<FBuildInstanceIdBufferAndCommandsFromPrimitiveIdsCs>(PermutationVector);
 
@@ -577,89 +574,6 @@ void FInstanceCullingContext::BuildRenderingCommands(FRDGBuilder& GraphBuilder,
 	);
 }
 
-#if 0
-
-void FInstanceCullingContext::CreateLegacyPassParameters(FRDGBuilder& GraphBuilder, FInstanceCullingDrawParams& InstanceCullingDrawParamsOut, int32 NumDrawCommands)
-{
-	FRDGBufferRef DrawIndirectArgsBuffer = GraphBuilder.CreateBuffer(FRDGBufferDesc::CreateIndirectDesc(IndirectArgsNumWords * NumDrawCommands), TEXT("DrawIndirectArgsBuffer"));
-	FRDGBufferRef InstanceIdOffsetBuffer = GraphBuilder.CreateBuffer(FRDGBufferDesc::CreateBufferDesc(sizeof(uint32), NumDrawCommands), TEXT("InstanceIdOffsetBuffer"));
-
-	InstanceCullingDrawParamsOut.DrawIndirectArgsBuffer = DrawIndirectArgsBuffer;
-	InstanceCullingDrawParamsOut.DrawIndirectArgsBufferAccess = DrawIndirectArgsBuffer;
-	InstanceCullingDrawParamsOut.InstanceIdOffsetBuffer = InstanceIdOffsetBuffer;
-	InstanceCullingDrawParamsOut.InstanceIdOffsetBufferAccess = InstanceIdOffsetBuffer;
-}
-
-
-void FInstanceCullingContext::BuildLegacyRenderingCommands(const FScene &Scene, FRHICommandListImmediate& RHICmdList, const FGPUScenePrimitiveCollector &DynamicPrimitiveCollector, FRHIBuffer* &InstanceIdOffsetBufferOut, FRHIBuffer* &DrawIndirectArgsBufferOut)
-{
-	FInstanceCullingManagerResources::FLegacyContext LegacyContext = GInstanceCullingManagerResources.GetLegacyContext();
-	ensure(LegacyContext.MaxDrawCommands >= CullingCommands.Num());
-
-	uint32 *IndirectArgs = reinterpret_cast<uint32*>(RHICmdList.LockBuffer(LegacyContext.DrawIndirectArgsBuffer, 0, LegacyContext.DrawIndirectArgsBuffer->GetSize(), RLM_WriteOnly));
-	uint32 *InstanceIdOffsets = reinterpret_cast<uint32*>(RHICmdList.LockBuffer(LegacyContext.InstanceIdOffsetBuffer, 0, LegacyContext.InstanceIdOffsetBuffer->GetSize(), RLM_WriteOnly));
-
-	FRHIBuffer* InstancesIdBuffer = GInstanceCullingManagerResources.GetInstancesIdBuffer();
-	uint32* InstanceIds = reinterpret_cast<uint32*>(RHICmdList.LockBuffer(InstancesIdBuffer, 0, InstancesIdBuffer->GetSize(), RLM_WriteOnly));
-
-	uint32 InstanceIdWriteOffset = 0;
-	for (int32 CommandIndex = 0; CommandIndex < CullingCommands.Num(); ++CommandIndex)
-	{
-		int32 InstanceIdStartOffset = InstanceIdWriteOffset;
-		const FPrimCullingCommand& Cmd = CullingCommands[CommandIndex];
-
-		int32 InstanceRunEnd = CommandIndex < CullingCommands.Num() - 1 ? CullingCommands[CommandIndex + 1].FirstInstanceRunOffset : InstanceRuns.Num();
-		for (int32 InstanceRunOffset = Cmd.FirstInstanceRunOffset; InstanceRunOffset < InstanceRunEnd; ++InstanceRunOffset)
-		{
-			const FInstanceRun &InstanceRun = InstanceRuns[InstanceRunOffset];
-			for (uint32 InstanceId = InstanceRun.Start; InstanceId <= InstanceRun.EndInclusive; ++InstanceId)
-			{
-				for (int32 ViewIndex = 0; ViewIndex < ViewIds.Num(); ++ViewIndex)
-				{
-					InstanceIds[InstanceIdWriteOffset++] = InstanceId | (uint32(ViewIndex) << 28U);
-				}
-			}
-		}
-
-		const int32 NumScenePrimitives = Scene.Primitives.Num();
-
-		int32 PrimitiveIdEnd = CommandIndex < CullingCommands.Num() - 1 ? CullingCommands[CommandIndex + 1].FirstPrimitiveIdOffset : PrimitiveIds.Num();
-		for (int32 PrimitiveIdIndex = Cmd.FirstPrimitiveIdOffset; PrimitiveIdIndex < PrimitiveIdEnd; ++PrimitiveIdIndex)
-		{
-			int32 PrimitiveId = PrimitiveIds[PrimitiveIdIndex];
-
-			const int32 InstanceDataOffset = PrimitiveId < NumScenePrimitives ? Scene.Primitives[PrimitiveId]->GetInstanceDataOffset() : DynamicPrimitiveCollector.GetInstanceDataOffset(PrimitiveId);
-			const int32 InstanceDataEntries = PrimitiveId < NumScenePrimitives ? Scene.Primitives[PrimitiveId]->GetNumInstanceDataEntries() : 1;
-			for (int32 InstanceIndex = 0; InstanceIndex < InstanceDataEntries; ++InstanceIndex)
-			{
-				int32 InstanceId = InstanceDataEntries + InstanceIndex;
-				for (int32 ViewIndex = 0; ViewIndex < ViewIds.Num(); ++ViewIndex)
-				{
-					InstanceIds[InstanceIdWriteOffset++] = InstanceId | (uint32(ViewIndex) << 28U);
-				}
-			}
-		}
-
-		InstanceIdOffsets[CommandIndex] = InstanceIdStartOffset;
-
-		IndirectArgs[CommandIndex * IndirectArgsNumWords + 0] = Cmd.NumVerticesOrIndices;
-		IndirectArgs[CommandIndex * IndirectArgsNumWords + 1] = InstanceIdWriteOffset - InstanceIdStartOffset;
-		IndirectArgs[CommandIndex * IndirectArgsNumWords + 2] = Cmd.FirstIndex;
-		IndirectArgs[CommandIndex * IndirectArgsNumWords + 3] = Cmd.BaseVertexIndex;
-		IndirectArgs[CommandIndex * IndirectArgsNumWords + 4] = 0U;
-	}
-
-	RHICmdList.UnlockBuffer(LegacyContext.DrawIndirectArgsBuffer);
-	RHICmdList.UnlockBuffer(LegacyContext.InstanceIdOffsetBuffer);
-	RHICmdList.UnlockBuffer(GInstanceCullingManagerResources.GetInstancesIdBuffer());
-
-
-	InstanceIdOffsetBufferOut = LegacyContext.InstanceIdOffsetBuffer;
-	DrawIndirectArgsBufferOut = LegacyContext.DrawIndirectArgsBuffer;
-}
-
-#endif // 0
-
 #else // GPUCULL_TODO
 
 void FInstanceCullingContext::BuildRenderingCommands(FRDGBuilder& GraphBuilder, const FGPUScene& GPUScene, FInstanceCullingResult& Results) const
