Avoid compiling adaptive SMRT ray count (wave ops) in pixel shader projections.
Add cvar to enable/disable adaptive SMRT ray count.

diff --git a/Engine/Shaders/Private/VirtualShadowMaps/ProjectionSpot.ush b/Engine/Shaders/Private/VirtualShadowMaps/ProjectionSpot.ush
--- a/Engine/Shaders/Private/VirtualShadowMaps/ProjectionSpot.ush
+++ b/Engine/Shaders/Private/VirtualShadowMaps/ProjectionSpot.ush
@@ -12,6 +12,10 @@ ProjectionSpot.ush:
 #include "PageAccessCommon.ush"
 #include "ProjectionCommon.ush"
 
+#ifndef SMRT_ADAPTIVE_RAY_COUNT
+#define SMRT_ADAPTIVE_RAY_COUNT 0
+#endif
+
 /** Return float multiplier to scale RayStepScreen such that it clip it right at the edge of the screen. */
 float GetStepScreenFactorToClipAtScreenEdge(float2 RayStartScreen, float2 RayStepScreen)
 {
@@ -219,6 +223,7 @@ float TraceSpotLight(
 
 			RayMissFactor += bHit ? 0 : 1;
 			
+#if SMRT_ADAPTIVE_RAY_COUNT
 #if COMPILER_SUPPORTS_WAVE_VOTE
 			if( i == 0 )
 			{
@@ -245,6 +250,7 @@ float TraceSpotLight(
 					break;
 				}
 			}
+#endif
 #endif
 		}
 		SampleCount = min(i + 1, uint(RayCount));		// break vs regular for loop exit
diff --git a/Engine/Source/Runtime/Renderer/Private/VirtualShadowMaps/VirtualShadowMapProjection.cpp b/Engine/Source/Runtime/Renderer/Private/VirtualShadowMaps/VirtualShadowMapProjection.cpp
--- a/Engine/Source/Runtime/Renderer/Private/VirtualShadowMaps/VirtualShadowMapProjection.cpp
+++ b/Engine/Source/Runtime/Renderer/Private/VirtualShadowMaps/VirtualShadowMapProjection.cpp
@@ -90,6 +90,13 @@ static TAutoConsoleVariable<float> CVarSMRTRayLengthScaleDirectional(
 	ECVF_RenderThreadSafe
 );
 
+static TAutoConsoleVariable<int32> CVarSMRTAdaptiveRayCount(
+	TEXT( "r.Shadow.v.SMRT.AdaptiveRayCount" ),
+	1,
+	TEXT( "Shoot fewer rays in fully shadowed and unshadowed regions. Currently only supported with OnePassProjection. " ),
+	ECVF_RenderThreadSafe
+);
+
 BEGIN_SHADER_PARAMETER_STRUCT(FProjectionParameters, )
 	SHADER_PARAMETER_STRUCT_INCLUDE(FVirtualShadowMapSamplingParameters, ProjectionParameters)
 	SHADER_PARAMETER_STRUCT(FLightShaderParameters, Light)
@@ -442,6 +449,9 @@ class FVirtualShadowMapProjectionCS : public FGlobalShader
 	DECLARE_GLOBAL_SHADER(FVirtualShadowMapProjectionCS);
 	SHADER_USE_PARAMETER_STRUCT(FVirtualShadowMapProjectionCS, FGlobalShader)
 	
+	class FSMRTAdaptiveRayCountDim : SHADER_PERMUTATION_BOOL("SMRT_ADAPTIVE_RAY_COUNT");
+	using FPermutationDomain = TShaderPermutationDomain<FSMRTAdaptiveRayCountDim>;
+
 	BEGIN_SHADER_PARAMETER_STRUCT(FParameters, )
 		SHADER_PARAMETER_STRUCT_INCLUDE(FVirtualShadowMapSamplingParameters, ProjectionParameters)
 		SHADER_PARAMETER_RDG_UNIFORM_BUFFER(FSceneTextureUniformParameters, SceneTexturesStruct)
@@ -501,8 +511,10 @@ void RenderVirtualShadowMapProjection(
 	PassParameters->SMRTRayLengthScale = 0.0f;		// Currently unused in this path
 
 	PassParameters->RWShadowMaskBits = GraphBuilder.CreateUAV( ShadowMaskBits );
-				
-	auto ComputeShader = View.ShaderMap->GetShader< FVirtualShadowMapProjectionCS >();
+	
+	FVirtualShadowMapProjectionCS::FPermutationDomain PermutationVector;
+	PermutationVector.Set< FVirtualShadowMapProjectionCS::FSMRTAdaptiveRayCountDim >( CVarSMRTAdaptiveRayCount.GetValueOnRenderThread() != 0 );
+	auto ComputeShader = View.ShaderMap->GetShader< FVirtualShadowMapProjectionCS >( PermutationVector );
 
 	const FIntPoint GroupCount = FIntPoint::DivideAndRoundUp( View.ViewRect.Size(), 8 );
 
@@ -513,4 +525,4 @@ void RenderVirtualShadowMapProjection(
 		PassParameters,
 		FIntVector( GroupCount.X, GroupCount.Y, 1 )
 	);
-}
\ No newline at end of file
+}
