Potential crash fix with real time capture.

diff --git a/Engine/Source/Runtime/Renderer/Private/ReflectionEnvironmentRealTimeCapture.cpp b/Engine/Source/Runtime/Renderer/Private/ReflectionEnvironmentRealTimeCapture.cpp
--- a/Engine/Source/Runtime/Renderer/Private/ReflectionEnvironmentRealTimeCapture.cpp
+++ b/Engine/Source/Runtime/Renderer/Private/ReflectionEnvironmentRealTimeCapture.cpp
@@ -502,11 +502,15 @@ void FScene::AllocateAndCaptureFrameSkyEnvMap(
 					}
 				}
 
-				if (HasSkyAtmosphere() && (MainView.bSceneHasSkyMaterial || HasVolumetricCloud()))
+				if (HasSkyAtmosphere()&& (MainView.bSceneHasSkyMaterial || HasVolumetricCloud())
+					&& RealTimeReflectionCaptureCamera360APLutTexture.IsValid())	// we also check that because it seems it can happen for some view setup UE-107270, TODO find a repro for a proper fix.
 				{
 					CubeView.CachedViewUniformShaderParameters->CameraAerialPerspectiveVolume = RealTimeReflectionCaptureCamera360APLutTexture->GetRenderTargetItem().ShaderResourceTexture;
 				}
-				// Else we do nothing as we assume the MainView one will not be used
+				else
+				{
+					CubeView.CachedViewUniformShaderParameters->CameraAerialPerspectiveVolume = GSystemTextures.VolumetricBlackDummy->GetRenderTargetItem().ShaderResourceTexture;
+				}
 
 				TUniformBufferRef<FViewUniformShaderParameters> CubeViewUniformBuffer = TUniformBufferRef<FViewUniformShaderParameters>::CreateUniformBufferImmediate(*CubeView.CachedViewUniformShaderParameters, UniformBuffer_SingleFrame);
 				CubeView.ViewUniformBuffer = CubeViewUniformBuffer;
