Add faster classification for strata material

* Indirect tile draw
* Faster tile building using reduction
* Add rect primitive support

diff --git a/Engine/Shaders/Private/Strata/StrataMaterialClassification.usf b/Engine/Shaders/Private/Strata/StrataMaterialClassification.usf
--- a/Engine/Shaders/Private/Strata/StrataMaterialClassification.usf
+++ b/Engine/Shaders/Private/Strata/StrataMaterialClassification.usf
@@ -6,6 +6,9 @@
 #define STRATA_INLINE_SHADING 0
 #include "/Engine/Private/Strata/Strata.ush"
 
+#define TILE_SIZE 8
+#define GROUP_THREAD_COUNT (TILE_SIZE*TILE_SIZE)
+
 #if SHADER_CATEGORIZATION
 
 // Note: 
@@ -169,42 +172,79 @@ void MainPS(
 
 ////////////////////////////////////////////////////////////////////////////////////////////////////////////
 
+uint EncodeTile(uint2 TileCoord)
+{
+	return TileCoord.x | (TileCoord.y << 16); // assumes 16bit is enough to represent a tiled resolution up to 65,535 :)
+}
+
+uint2 DecodeTile(uint In)
+{
+	return uint2(In & 0xFFFF, In >> 16);
+}
+
 #if SHADER_TILE_CATEGORIZATION
-int  TileSize;
+int TileSize;
+int bRectPrimitive;
 int2 ClassificationResolution;
-int2 TileClassificationResolution;
 Texture2D<uint> ClassificationTexture;
+RWBuffer<uint> TileIndirectData;
+RWBuffer<uint> TileListData;
 
-#define GROUP_SIZE 8
-#define GROUP_THREAD_COUNT (GROUP_SIZE*GROUP_SIZE)
-groupshared uint s_TileType;
-RWTexture2D<uint> TileClassificationTexture;
-[numthreads(8, 8, 1)]
-void TileMainCS(uint2 DispatchThreadId : SV_DispatchThreadID, uint GroupIndex : SV_GroupIndex)
+groupshared bool s_IsSimpleTile[GROUP_THREAD_COUNT];
+
+[numthreads(TILE_SIZE, TILE_SIZE, 1)]
+void TileMainCS(uint2 DispatchThreadId : SV_DispatchThreadID, uint LinearIndex : SV_GroupIndex, uint3 GroupId : SV_GroupID)
 {
-	if (GroupIndex == 0)
+	if (all(DispatchThreadId == 0))
 	{
-		s_TileType = 0;
+		TileIndirectData[0] = bRectPrimitive > 0 ? 4 : 6;
+		TileIndirectData[2] = 0;
+		TileIndirectData[3] = 0;
 	}
 
+	// TODO: add a SM6 permutation with ballot?
 	const uint2 PixelCoord = DispatchThreadId;
-	const uint PackClassification = ClassificationTexture.Load(uint3(PixelCoord, 0));
-	FStrataClassification Data = StrataUnpackClassificationData(PackClassification);
-
-	if (all(PixelCoord < ClassificationResolution))
+	bool bIsSimpleTile = false;
+	if (all(PixelCoord < uint2(ClassificationResolution)))
 	{
-		InterlockedAdd(s_TileType, Data.bIsSimple ? 1 : 0);
+		const uint PackClassification = ClassificationTexture.Load(uint3(PixelCoord, 0));
+		FStrataClassification Data = StrataUnpackClassificationData(PackClassification);
+		bIsSimpleTile = Data.bIsSimple;
 	}
+	s_IsSimpleTile[LinearIndex] = bIsSimpleTile;
 
+	GroupMemoryBarrierWithGroupSync();
+	if (LinearIndex < 32)
+	{
+		s_IsSimpleTile[LinearIndex] = s_IsSimpleTile[LinearIndex] || s_IsSimpleTile[LinearIndex + 32];
+	}
+	GroupMemoryBarrierWithGroupSync();
+	if (LinearIndex < 16)
+	{
+		s_IsSimpleTile[LinearIndex] = s_IsSimpleTile[LinearIndex] || s_IsSimpleTile[LinearIndex + 16];
+	}
 	GroupMemoryBarrierWithGroupSync();
 
-	if (GroupIndex == 0)
+	if (LinearIndex < 8)
 	{
-		const uint2 TileCoord = PixelCoord / TileSize;
-		if (s_TileType == GROUP_THREAD_COUNT && all(TileCoord < TileClassificationResolution))
-		{
-			InterlockedAdd(TileClassificationTexture[TileCoord], s_TileType);
-		}
+		s_IsSimpleTile[LinearIndex] = s_IsSimpleTile[LinearIndex] || s_IsSimpleTile[LinearIndex + 8];
+	}
+	GroupMemoryBarrierWithGroupSync();
+	if (LinearIndex < 4)
+	{
+		s_IsSimpleTile[LinearIndex] = s_IsSimpleTile[LinearIndex] || s_IsSimpleTile[LinearIndex + 4];
+	}
+	GroupMemoryBarrierWithGroupSync();
+	if (LinearIndex < 2)
+	{
+		s_IsSimpleTile[LinearIndex] = s_IsSimpleTile[LinearIndex] || s_IsSimpleTile[LinearIndex + 2];
+	}
+	GroupMemoryBarrierWithGroupSync();
+	if (LinearIndex < 1 && (s_IsSimpleTile[LinearIndex] || s_IsSimpleTile[LinearIndex + 1]))
+	{
+		uint WriteToIndex;
+		InterlockedAdd(TileIndirectData[1], 1, WriteToIndex);
+		TileListData[WriteToIndex] = EncodeTile(GroupId.xy);
 	}
 }
 #endif // SHADER_TILE_CATEGORIZATION
@@ -215,39 +255,21 @@ void TileMainCS(uint2 DispatchThreadId : SV_DispatchThreadID, uint GroupIndex :
 
 int  TileSize;
 int2 TileCount;
+int	 bRectPrimitive;
 int2 OutputResolution;
-Texture2D<uint> TileClassificationTexture;
+Buffer<uint> TileListBuffer;
 
 void StencilMainVS(
 	in uint InVertexId : SV_VertexID,
 	in uint InInstanceId : SV_InstanceID,
 	out float4 OutPosition : SV_POSITION)
 {
-	uint2 TileCoord = uint2(InInstanceId % TileCount.x, InInstanceId / TileCount.y);
-	
-	// Only tile whose all the threads are flagged as simple will be set as valid
-	const uint bIsSimple = TileClassificationTexture.Load(uint3(TileCoord, 0)) == (TileSize*TileSize);
-	float3 P = float3(INFINITE_FLOAT, INFINITE_FLOAT, INFINITE_FLOAT);
-	if (bIsSimple)
-	{
-		float2 UV = 0;
-		switch (InVertexId)
-		{
-		case 0: UV = float2(0, 0); break;
-		case 1: UV = float2(1, 0); break;
-		case 2: UV = float2(1, 1); break;
-		case 3: UV = float2(0, 0); break;
-		case 4: UV = float2(1, 1); break;
-		case 5: UV = float2(0, 1); break;
-		}
+	const uint2 TileCoord = DecodeTile(TileListBuffer[InInstanceId]);
 
-		int LocalTileSize = TileSize;
-		P.xy = TileCoord * LocalTileSize + UV * LocalTileSize;
-		P.xy/= float2(OutputResolution);
-		P.xy = float2(P.x, (1 - P.y)) * 2 - 1;
-		P.z  = 0.5f;
-	}
-	OutPosition = float4(P, 1);
+	uint2 TileVertex = TileCoord * TILE_SIZE;
+	TileVertex.x += InVertexId == 1 || InVertexId == 2 || InVertexId == 4 ? TILE_SIZE : 0;
+	TileVertex.y += InVertexId == 2 || InVertexId == 4 || InVertexId == 5 ? TILE_SIZE : 0;
+	OutPosition = float4(float2(TileVertex) / OutputResolution * float2(2.0f, -2.0f) + float2(-1.0, 1.0f), 0.5f, 1.0f);
 }
 
 void StencilMainPS(
diff --git a/Engine/Source/Runtime/Renderer/Private/Strata/Strata.cpp b/Engine/Source/Runtime/Renderer/Private/Strata/Strata.cpp
--- a/Engine/Source/Runtime/Renderer/Private/Strata/Strata.cpp
+++ b/Engine/Source/Runtime/Renderer/Private/Strata/Strata.cpp
@@ -31,12 +31,6 @@ static TAutoConsoleVariable<int32> CVarStrataClassification(
 	TEXT("Enable strata classification to speed up lighting pass."),
 	ECVF_RenderThreadSafe); 
 
-static TAutoConsoleVariable<int32> CVarStrataTileInPixel(
-	TEXT("r.Strata.TileInPixel"),
-	64,
-	TEXT("Strata allocated byte per pixel to store materials data. Higher value means more complex material can be represented."),
-	ECVF_RenderThreadSafe);
-
 IMPLEMENT_GLOBAL_SHADER_PARAMETER_STRUCT(FStrataGlobalUniformParameters, "Strata");
 
 namespace Strata
@@ -54,7 +48,7 @@ bool IsClassificationEnabled()
 
 uint32 GetStrataTileSize()
 {
-	return FMath::RoundUpToPowerOfTwo(FMath::Clamp(CVarStrataTileInPixel.GetValueOnRenderThread(), 8, 128));
+	return 8;
 }
 
 void InitialiseStrataFrameSceneData(FSceneRenderer& SceneRenderer, FRDGBuilder& GraphBuilder)
@@ -84,15 +78,17 @@ void InitialiseStrataFrameSceneData(FSceneRenderer& SceneRenderer, FRDGBuilder&
 			StrataSceneData.ClassificationTexture = GraphBuilder.GetPooledTexture(Texture);
 		}
 
-		// Tile classification texture
+		// Tile classification buffers
 		{
 			const int32 TileInPixel = GetStrataTileSize();
-			const uint32 TileResolutionX = FMath::DivideAndRoundUp(BufferSizeXY.X, TileInPixel);
-			const uint32 TileResolutionY = FMath::DivideAndRoundUp(BufferSizeXY.Y, TileInPixel);
-
-			FRDGTextureRef Texture = GraphBuilder.CreateTexture(FRDGTextureDesc::Create2D(FIntPoint(TileResolutionX, TileResolutionY), PF_R32_UINT, FClearValueBinding::Black, TexCreate_ShaderResource | TexCreate_UAV), TEXT("StrataTileClassificationTexture"));
-			GraphBuilder.PreallocateTexture(Texture);
-			StrataSceneData.TileClassificationTexture = GraphBuilder.GetPooledTexture(Texture);
+			const FIntPoint TileResolution(FMath::DivideAndRoundUp(BufferSizeXY.X, TileInPixel), FMath::DivideAndRoundUp(BufferSizeXY.Y, TileInPixel));
+			FRDGBufferRef ClassificationTileListBuffer = GraphBuilder.CreateBuffer(FRDGBufferDesc::CreateBufferDesc(sizeof(uint32), TileResolution.X * TileResolution.Y), TEXT("StrataTileListBuffer"));
+			FRDGBufferRef ClassificationTileIndirectBuffer = GraphBuilder.CreateBuffer(FRDGBufferDesc::CreateIndirectDesc<FRHIDrawIndirectParameters>(), TEXT("StrataTileIndirectBuffer")); 
+
+			GraphBuilder.PreallocateBuffer(ClassificationTileListBuffer);
+			GraphBuilder.PreallocateBuffer(ClassificationTileIndirectBuffer);
+			StrataSceneData.ClassificationTileListBuffer = GraphBuilder.GetPooledBuffer(ClassificationTileListBuffer);
+			StrataSceneData.ClassificationTileIndirectBuffer = GraphBuilder.GetPooledBuffer(ClassificationTileIndirectBuffer);
 		}
 
 		// Top layer texture
@@ -306,10 +302,11 @@ class FStrataMaterialTileClassificationPassCS : public FGlobalShader
 	BEGIN_SHADER_PARAMETER_STRUCT(FParameters, )
 		SHADER_PARAMETER_STRUCT_REF(FViewUniformShaderParameters, ViewUniformBuffer)
 		SHADER_PARAMETER(int32, TileSize)
+		SHADER_PARAMETER(int32, bRectPrimitive)
 		SHADER_PARAMETER(FIntPoint, ClassificationResolution)
-		SHADER_PARAMETER(FIntPoint, TileClassificationResolution)
 		SHADER_PARAMETER_RDG_TEXTURE(Texture2D, ClassificationTexture)
-		SHADER_PARAMETER_RDG_TEXTURE_UAV(RWTexture2D, TileClassificationTexture)
+		SHADER_PARAMETER_RDG_BUFFER_UAV(RWBuffer, TileIndirectData)
+		SHADER_PARAMETER_RDG_BUFFER_UAV(RWBuffer, TileListData)
 	END_SHADER_PARAMETER_STRUCT()
 
 	static FPermutationDomain RemapPermutation(FPermutationDomain PermutationVector)
@@ -343,9 +340,10 @@ class FStrataMaterialStencilClassificationPassVS : public FGlobalShader
 	BEGIN_SHADER_PARAMETER_STRUCT(FParameters, )
 		SHADER_PARAMETER_STRUCT_REF(FViewUniformShaderParameters, ViewUniformBuffer)
 		SHADER_PARAMETER(int32, TileSize)
+		SHADER_PARAMETER(int32, bRectPrimitive)
 		SHADER_PARAMETER(FIntPoint, TileCount)
 		SHADER_PARAMETER(FIntPoint, OutputResolution)
-		SHADER_PARAMETER_RDG_TEXTURE(Texture2D, TileClassificationTexture)
+		SHADER_PARAMETER_RDG_BUFFER_SRV(Buffer, TileListBuffer)
 	END_SHADER_PARAMETER_STRUCT()
 
 	static FPermutationDomain RemapPermutation(FPermutationDomain PermutationVector)
@@ -375,13 +373,15 @@ class FStrataMaterialStencilClassificationPassPS : public FGlobalShader
 	BEGIN_SHADER_PARAMETER_STRUCT(FParameters, )
 		SHADER_PARAMETER_STRUCT_REF(FViewUniformShaderParameters, ViewUniformBuffer)
 		SHADER_PARAMETER(int32, TileSize)
+		SHADER_PARAMETER(int32, bRectPrimitive)
 		SHADER_PARAMETER(FIntPoint, TileCount)
 		SHADER_PARAMETER(FIntPoint, OutputResolution)
-		SHADER_PARAMETER_RDG_TEXTURE(Texture2D, TileClassificationTexture)
+		SHADER_PARAMETER_RDG_BUFFER_SRV(Buffer, TileListBuffer)
+		SHADER_PARAMETER_RDG_BUFFER(Buffer, TileIndirectBuffer)
 		RENDER_TARGET_BINDING_SLOTS()
 	END_SHADER_PARAMETER_STRUCT()
 
-		static FPermutationDomain RemapPermutation(FPermutationDomain PermutationVector)
+	static FPermutationDomain RemapPermutation(FPermutationDomain PermutationVector)
 	{
 		return PermutationVector;
 	}
@@ -405,18 +405,22 @@ static void AddStrataStencilPass(
 	FRDGBuilder& GraphBuilder,
 	const FViewInfo& View,
 	const FMinimalSceneTextures& SceneTextures,
-	FRDGTextureRef TileClassificationTexture)
+	FRDGBufferRef TileListBuffer,
+	FRDGBufferRef TileIndirectBuffer)
 {	
 	const FIntPoint OutputResolution = View.ViewRect.Size();
 	const int32 TileSize = GetStrataTileSize();
+	
 	const FIntPoint TileCount(FMath::DivideAndRoundUp(OutputResolution.X, TileSize), FMath::DivideAndRoundUp(OutputResolution.Y, TileSize));
 	const uint32 InstanceCount = TileCount.X * TileCount.Y;
 
 	FStrataMaterialStencilClassificationPassPS::FParameters* ParametersPS = GraphBuilder.AllocParameters<FStrataMaterialStencilClassificationPassPS::FParameters>();
 	ParametersPS->TileSize = TileSize;
 	ParametersPS->TileCount = TileCount;
+	ParametersPS->bRectPrimitive = GRHISupportsRectTopology ? 1 : 0;
 	ParametersPS->OutputResolution = OutputResolution;
-	ParametersPS->TileClassificationTexture = TileClassificationTexture;
+	ParametersPS->TileListBuffer = GraphBuilder.CreateSRV(TileListBuffer, PF_R32_UINT);
+	ParametersPS->TileIndirectBuffer = TileIndirectBuffer;
 
 	TShaderMapRef<FStrataMaterialStencilClassificationPassVS> VertexShader(View.ShaderMap);
 	TShaderMapRef<FStrataMaterialStencilClassificationPassPS> PixelShader(View.ShaderMap);
@@ -442,9 +446,10 @@ static void AddStrataStencilPass(
 		{
 			FStrataMaterialStencilClassificationPassVS::FParameters ParametersVS;
 			ParametersVS.TileSize = ParametersPS->TileSize;
+			ParametersVS.bRectPrimitive = ParametersPS->bRectPrimitive;
 			ParametersVS.TileCount = ParametersPS->TileCount;
 			ParametersVS.OutputResolution = ParametersPS->OutputResolution;
-			ParametersVS.TileClassificationTexture = ParametersPS->TileClassificationTexture;
+			ParametersVS.TileListBuffer = ParametersPS->TileListBuffer;
 
 			FGraphicsPipelineStateInitializer GraphicsPSOInit;
 			RHICmdList.ApplyCachedRenderTargets(GraphicsPSOInit);
@@ -458,7 +463,7 @@ static void AddStrataStencilPass(
 			GraphicsPSOInit.BoundShaderState.VertexDeclarationRHI = GFilterVertexDeclaration.VertexDeclarationRHI;
 			GraphicsPSOInit.BoundShaderState.VertexShaderRHI = VertexShader.GetVertexShader();
 			GraphicsPSOInit.BoundShaderState.PixelShaderRHI = PixelShader.GetPixelShader();
-			GraphicsPSOInit.PrimitiveType = PT_TriangleList;
+			GraphicsPSOInit.PrimitiveType = ParametersPS->bRectPrimitive > 0 ? PT_RectList : PT_TriangleList;
 			SetGraphicsPipelineState(RHICmdList, GraphicsPSOInit);
 			SetShaderParameters(RHICmdList, VertexShader, VertexShader.GetVertexShader(), ParametersVS);
 			//SetShaderParameters(RHICmdList, PixelShader, PixelShader.GetPixelShader(), *ParametersPS);
@@ -466,7 +471,7 @@ static void AddStrataStencilPass(
 			RHICmdList.SetStencilRef(StencilBit);
 			RHICmdList.SetViewport(0, 0, 0.0f, OutputResolution.X, OutputResolution.Y, 1.0f);
 			RHICmdList.SetStreamSource(0, nullptr, 0);
-			RHICmdList.DrawPrimitive(0, 2, InstanceCount);
+			RHICmdList.DrawPrimitiveIndirect(ParametersPS->TileIndirectBuffer->GetRHI(), 0);
 		});
 }
 
@@ -475,8 +480,9 @@ void AddStrataStencilPass(
 	const FViewInfo& View,
 	const FMinimalSceneTextures& SceneTextures)
 {
-	FRDGTextureRef TileClassificationTexture = GraphBuilder.RegisterExternalTexture(View.StrataSceneData->TileClassificationTexture);
-	AddStrataStencilPass(GraphBuilder, View, SceneTextures, TileClassificationTexture);
+	FRDGBufferRef TileListBuffer = GraphBuilder.RegisterExternalBuffer(View.StrataSceneData->ClassificationTileListBuffer);
+	FRDGBufferRef TileIndirectBuffer = GraphBuilder.RegisterExternalBuffer(View.StrataSceneData->ClassificationTileIndirectBuffer);
+	AddStrataStencilPass(GraphBuilder, View, SceneTextures, TileListBuffer, TileIndirectBuffer);
 }
 
 void AddStrataStencilPass(
@@ -487,8 +493,9 @@ void AddStrataStencilPass(
 	for (int32 i = 0; i < Views.Num(); ++i)
 	{
 		const FViewInfo& View = Views[i];
-		FRDGTextureRef TileClassificationTexture = GraphBuilder.RegisterExternalTexture(View.StrataSceneData->TileClassificationTexture);
-		AddStrataStencilPass(GraphBuilder, View, SceneTextures, TileClassificationTexture);
+		FRDGBufferRef TileListBuffer = GraphBuilder.RegisterExternalBuffer(View.StrataSceneData->ClassificationTileListBuffer);
+		FRDGBufferRef TileIndirectBuffer = GraphBuilder.RegisterExternalBuffer(View.StrataSceneData->ClassificationTileIndirectBuffer);
+		AddStrataStencilPass(GraphBuilder, View, SceneTextures, TileListBuffer, TileIndirectBuffer);
 	}
 }
 
@@ -533,20 +540,16 @@ void AddStrataMaterialClassificationPass(FRDGBuilder& GraphBuilder, const FMinim
 		}
 
 		// Downsampling
-		FRDGTextureRef TileClassificationTexture = GraphBuilder.RegisterExternalTexture(View.StrataSceneData->TileClassificationTexture);
 		if (IsClassificationEnabled())
 		{
-			const uint32 ClearValues[4] = { 0u, 0u, 0u, 0u };
-			FRDGTextureUAVRef TileClassificationTextureUAV = GraphBuilder.CreateUAV(TileClassificationTexture);
-			AddClearUAVPass(GraphBuilder, TileClassificationTextureUAV, ClearValues);
-
 			TShaderMapRef<FStrataMaterialTileClassificationPassCS> ComputeShader(View.ShaderMap);
 			FStrataMaterialTileClassificationPassCS::FParameters* PassParameters = GraphBuilder.AllocParameters<FStrataMaterialTileClassificationPassCS::FParameters>();
 			PassParameters->TileSize = GetStrataTileSize();
+			PassParameters->bRectPrimitive = GRHISupportsRectTopology ? 1 : 0;
 			PassParameters->ClassificationResolution = ClassificationTexture->Desc.Extent;
-			PassParameters->TileClassificationResolution = TileClassificationTexture->Desc.Extent;
 			PassParameters->ClassificationTexture = ClassificationTexture;
-			PassParameters->TileClassificationTexture = TileClassificationTextureUAV;
+			PassParameters->TileListData = GraphBuilder.CreateUAV(GraphBuilder.RegisterExternalBuffer(View.StrataSceneData->ClassificationTileListBuffer), PF_R32_UINT);
+			PassParameters->TileIndirectData = GraphBuilder.CreateUAV(GraphBuilder.RegisterExternalBuffer(View.StrataSceneData->ClassificationTileIndirectBuffer), PF_R32_UINT);
 
 			// Add 64 threads permutation
 			const uint32 GroupSize = 8;
diff --git a/Engine/Source/Runtime/Renderer/Private/Strata/Strata.h b/Engine/Source/Runtime/Renderer/Private/Strata/Strata.h
--- a/Engine/Source/Runtime/Renderer/Private/Strata/Strata.h
+++ b/Engine/Source/Runtime/Renderer/Private/Strata/Strata.h
@@ -32,7 +32,9 @@ struct FStrataSceneData
 	FRWByteAddressBuffer MaterialLobesBuffer;					// This should be a RDG resource
 	TRefCountPtr<IPooledRenderTarget> ClassificationTexture;	// This should be a RDG resource
 	TRefCountPtr<IPooledRenderTarget> TopLayerTexture;			// This should be a RDG resource
-	TRefCountPtr<IPooledRenderTarget> TileClassificationTexture;// This should be a RDG resource
+
+	TRefCountPtr<FRDGPooledBuffer> ClassificationTileListBuffer;
+	TRefCountPtr<FRDGPooledBuffer> ClassificationTileIndirectBuffer;
 
 	TUniformBufferRef<FStrataGlobalUniformParameters> StrataGlobalUniformParameters;
 
