Adds details in TAA draw events about Gen4 or Gen5

diff --git a/Engine/Source/Runtime/Renderer/Private/PostProcess/PostProcessing.cpp b/Engine/Source/Runtime/Renderer/Private/PostProcess/PostProcessing.cpp
--- a/Engine/Source/Runtime/Renderer/Private/PostProcess/PostProcessing.cpp
+++ b/Engine/Source/Runtime/Renderer/Private/PostProcess/PostProcessing.cpp
@@ -540,7 +540,7 @@ void AddPostProcessingPasses(FRDGBuilder& GraphBuilder, const FViewInfo& View, c
 			RDG_EVENT_SCOPE_CONDITIONAL(
 				GraphBuilder,
 				UpscalerToUse != DefaultTemporalUpscaler,
-				"ThirdParty %s %dx%d -> %dx%d",
+				"ThirdParty TAA %s %dx%d -> %dx%d",
 				UpscalerToUse->GetDebugName(),
 				View.ViewRect.Width(), View.ViewRect.Height(),
 				View.GetSecondaryViewRectSize().X, View.GetSecondaryViewRectSize().Y);
diff --git a/Engine/Source/Runtime/Renderer/Private/PostProcess/TemporalAA.cpp b/Engine/Source/Runtime/Renderer/Private/PostProcess/TemporalAA.cpp
--- a/Engine/Source/Runtime/Renderer/Private/PostProcess/TemporalAA.cpp
+++ b/Engine/Source/Runtime/Renderer/Private/PostProcess/TemporalAA.cpp
@@ -556,9 +556,9 @@ const TCHAR* const kTAAOutputNames[] = {
 };
 
 const TCHAR* const kTAAPassNames[] = {
-	TEXT("Main"),
-	TEXT("MainUpsampling"),
-	TEXT("MainSuperSampling"),
+	TEXT("Gen4 Main"),
+	TEXT("Gen4 MainUpsampling"),
+	TEXT("Gen4 MainSuperSampling"),
 	TEXT("ScreenSpaceReflections"),
 	TEXT("LightShaft"),
 	TEXT("DOF"),
@@ -1031,7 +1031,7 @@ static void AddGen5MainTemporalAAPasses(
 			FMath::Max(InputExtent.Y, QuantizedHistoryViewSize.Y));
 	}
 
-	RDG_EVENT_SCOPE(GraphBuilder, "TAAU %dx%d -> %dx%d", InputRect.Width(), InputRect.Height(), OutputRect.Width(), OutputRect.Height());
+	RDG_EVENT_SCOPE(GraphBuilder, "TAAU Gen5 %dx%d -> %dx%d", InputRect.Width(), InputRect.Height(), OutputRect.Width(), OutputRect.Height());
 	RDG_GPU_STAT_SCOPE(GraphBuilder, TAA);
 
 	FRDGTextureRef BlackDummy = GraphBuilder.RegisterExternalTexture(GSystemTextures.BlackDummy);
