Add strata classification for sky reflection pass.

diff --git a/Engine/Shaders/Private/Strata/StrataEnvironmentLighting.ush b/Engine/Shaders/Private/Strata/StrataEnvironmentLighting.ush
--- a/Engine/Shaders/Private/Strata/StrataEnvironmentLighting.ush
+++ b/Engine/Shaders/Private/Strata/StrataEnvironmentLighting.ush
@@ -2,6 +2,10 @@
 
 #pragma once
 
+#ifndef STRATA_FASTPATH 
+#define STRATA_FASTPATH 0
+#endif
+
 // Env. lighting evaluation for strata material.
 // Unpack BSDF on-the-fly
 float4 StrataEnvironmentLighting(
@@ -27,9 +31,22 @@ float4 StrataEnvironmentLighting(
 	CommonFactor *= View.PreExposure;
 #endif
 
+#if STRATA_FASTPATH
+	uint BSDFIndex = 0;
+	if (StrataPixelHeader.BSDFCount == 1)
+#else
 	for (uint BSDFIndex = 0; BSDFIndex < StrataPixelHeader.BSDFCount; ++BSDFIndex)
+#endif
 	{
-		const FStrataBSDF BSDF = UnpackStrataBSDFIn(StrataDataBuffer, PixelStrataDataByteOffset);
+		FStrataBSDF BSDF = UnpackStrataBSDFIn(StrataDataBuffer, PixelStrataDataByteOffset);
+#if STRATA_FASTPATH
+		BSDF_SETTYPE(BSDF, STRATA_BSDF_TYPE_SLAB);
+		BSDF_SETANISOTROPY(BSDF, 0);
+		BSDF_SETISTOPLAYER(BSDF, 0);
+		BSDF_SETHASSCATTERING(BSDF, 0);
+		BSDF_SETHASSSSPROFILE(BSDF, 0);
+		BSDF_SETHASHAZINESS(BSDF, 0);
+#endif
 		const float3 BSDFThroughput = BSDF.Weight;
 
 		// Create the BSDF context
diff --git a/Engine/Source/Runtime/Renderer/Private/IndirectLightRendering.cpp b/Engine/Source/Runtime/Renderer/Private/IndirectLightRendering.cpp
--- a/Engine/Source/Runtime/Renderer/Private/IndirectLightRendering.cpp
+++ b/Engine/Source/Runtime/Renderer/Private/IndirectLightRendering.cpp
@@ -169,6 +169,7 @@ class FReflectionEnvironmentSkyLightingPS : public FGlobalShader
 	class FSkyShadowing : SHADER_PERMUTATION_BOOL("APPLY_SKY_SHADOWING");
 	class FRayTracedReflections : SHADER_PERMUTATION_BOOL("RAY_TRACED_REFLECTIONS");
 	class FStrata : SHADER_PERMUTATION_BOOL("STRATA_ENABLED");
+	class FStrataFastPath : SHADER_PERMUTATION_BOOL("STRATA_FASTPATH");
 
 	using FPermutationDomain = TShaderPermutationDomain<
 		FHasBoxCaptures,
@@ -178,7 +179,8 @@ class FReflectionEnvironmentSkyLightingPS : public FGlobalShader
 		FDynamicSkyLight,
 		FSkyShadowing,
 		FRayTracedReflections,
-		FStrata>;
+		FStrata,
+		FStrataFastPath>;
 
 	static FPermutationDomain RemapPermutation(FPermutationDomain PermutationVector)
 	{
@@ -199,10 +201,15 @@ class FReflectionEnvironmentSkyLightingPS : public FGlobalShader
 			PermutationVector.Set<FStrata>(false);
 		}
 
+		if (PermutationVector.Get<FStrataFastPath>() && (!Strata::IsStrataEnabled() || !Strata::IsClassificationEnabled()))
+		{
+			PermutationVector.Set<FStrataFastPath>(false);
+		}
+
 		return PermutationVector;
 	}
 
-	static FPermutationDomain BuildPermutationVector(const FViewInfo& View, bool bBoxCapturesOnly, bool bSphereCapturesOnly, bool bSupportDFAOIndirectOcclusion, bool bEnableSkyLight, bool bEnableDynamicSkyLight, bool bApplySkyShadowing, bool bRayTracedReflections)
+	static FPermutationDomain BuildPermutationVector(const FViewInfo& View, bool bBoxCapturesOnly, bool bSphereCapturesOnly, bool bSupportDFAOIndirectOcclusion, bool bEnableSkyLight, bool bEnableDynamicSkyLight, bool bApplySkyShadowing, bool bRayTracedReflections, bool bStrataFastPath)
 	{
 		FPermutationDomain PermutationVector;
 
@@ -214,6 +221,7 @@ class FReflectionEnvironmentSkyLightingPS : public FGlobalShader
 		PermutationVector.Set<FSkyShadowing>(bApplySkyShadowing);
 		PermutationVector.Set<FRayTracedReflections>(bRayTracedReflections);
 		PermutationVector.Set<FStrata>(Strata::IsStrataEnabled());
+		PermutationVector.Set<FStrataFastPath>(Strata::IsStrataEnabled() && Strata::IsClassificationEnabled() && bStrataFastPath);
 
 		return RemapPermutation(PermutationVector);
 	}
@@ -1090,6 +1098,196 @@ void FDeferredShadingSceneRenderer::RenderDiffuseIndirectAndAmbientOcclusion(
 	} // for (FViewInfo& View : Views)
 }
 
+static void AddSkyReflectionPass(
+	FRDGBuilder& GraphBuilder, 
+	FViewInfo& View, 
+	FScene* Scene, 
+	const FSceneTextures& SceneTextures,
+	FRDGTextureRef DynamicBentNormalAOTexture,
+	FRDGTextureRef ReflectionsColor,
+	const FRayTracingReflectionOptions& RayTracingReflectionOptions,
+	FSceneTextureParameters& SceneTextureParameters,
+	bool bSkyLight, 
+	bool bDynamicSkyLight, 
+	bool bApplySkyShadowing,
+	bool bStrataFastPath)
+{
+	// Render the reflection environment with tiled deferred culling
+	bool bHasBoxCaptures = (View.NumBoxReflectionCaptures > 0);
+	bool bHasSphereCaptures = (View.NumSphereReflectionCaptures > 0);
+
+	float DynamicBentNormalAO = 0.0f;
+	const FRDGSystemTextures& SystemTextures = FRDGSystemTextures::Get(GraphBuilder);
+	FRDGTextureRef AmbientOcclusionTexture = GetScreenSpaceAOFallback(SystemTextures);
+	if (HasBeenProduced(SceneTextures.ScreenSpaceAO))
+	{
+		DynamicBentNormalAO = 1.0f;
+		AmbientOcclusionTexture = SceneTextures.ScreenSpaceAO;
+	}
+
+	const auto& SceneColorTexture = SceneTextures.Color;
+
+	FReflectionEnvironmentSkyLightingPS::FParameters* PassParameters = GraphBuilder.AllocParameters<FReflectionEnvironmentSkyLightingPS::FParameters>();
+
+	// Setup the parameters of the shader.
+	{
+		// Setups all shader parameters related to skylight.
+		{
+			FSkyLightSceneProxy* SkyLight = Scene->SkyLight;
+
+			float SkyLightContrast = 0.01f;
+			float SkyLightOcclusionExponent = 1.0f;
+			FVector4 SkyLightOcclusionTintAndMinOcclusion(0.0f, 0.0f, 0.0f, 0.0f);
+			EOcclusionCombineMode SkyLightOcclusionCombineMode = EOcclusionCombineMode::OCM_MAX;
+			if (SkyLight)
+			{
+				FDistanceFieldAOParameters Parameters(SkyLight->OcclusionMaxDistance, SkyLight->Contrast);
+				SkyLightContrast = Parameters.Contrast;
+				SkyLightOcclusionExponent = SkyLight->OcclusionExponent;
+				SkyLightOcclusionTintAndMinOcclusion = FVector4(SkyLight->OcclusionTint);
+				SkyLightOcclusionTintAndMinOcclusion.W = SkyLight->MinOcclusion;
+				SkyLightOcclusionCombineMode = SkyLight->OcclusionCombineMode;
+			}
+
+			// Scale and bias to remap the contrast curve to [0,1]
+			const float Min = 1 / (1 + FMath::Exp(-SkyLightContrast * (0 * 10 - 5)));
+			const float Max = 1 / (1 + FMath::Exp(-SkyLightContrast * (1 * 10 - 5)));
+			const float Mul = 1.0f / (Max - Min);
+			const float Add = -Min / (Max - Min);
+
+			PassParameters->OcclusionTintAndMinOcclusion = SkyLightOcclusionTintAndMinOcclusion;
+			PassParameters->ContrastAndNormalizeMulAdd = FVector(SkyLightContrast, Mul, Add);
+			PassParameters->OcclusionExponent = SkyLightOcclusionExponent;
+			PassParameters->OcclusionCombineMode = SkyLightOcclusionCombineMode == OCM_Minimum ? 0.0f : 1.0f;
+			PassParameters->ApplyBentNormalAO = DynamicBentNormalAO;
+			PassParameters->InvSkySpecularOcclusionStrength = 1.0f / FMath::Max(CVarSkySpecularOcclusionStrength.GetValueOnRenderThread(), 0.1f);
+		}
+
+		// Setups all shader parameters related to distance field AO
+		{
+			FIntPoint AOBufferSize = GetBufferSizeForAO();
+			PassParameters->AOBufferBilinearUVMax = FVector2D(
+				(View.ViewRect.Width() / GAODownsampleFactor - 0.51f) / AOBufferSize.X, // 0.51 - so bilateral gather4 won't sample invalid texels
+				(View.ViewRect.Height() / GAODownsampleFactor - 0.51f) / AOBufferSize.Y);
+
+			extern float GAOViewFadeDistanceScale;
+			PassParameters->AOMaxViewDistance = GetMaxAOViewDistance();
+			PassParameters->DistanceFadeScale = 1.0f / ((1.0f - GAOViewFadeDistanceScale) * GetMaxAOViewDistance());
+
+			PassParameters->BentNormalAOTexture = DynamicBentNormalAOTexture;
+			PassParameters->BentNormalAOSampler = TStaticSamplerState<SF_Bilinear>::GetRHI();
+		}
+
+		PassParameters->AmbientOcclusionTexture = AmbientOcclusionTexture;
+		PassParameters->AmbientOcclusionSampler = TStaticSamplerState<SF_Point>::GetRHI();
+
+		PassParameters->ScreenSpaceReflectionsTexture = ReflectionsColor ? ReflectionsColor : SystemTextures.Black;
+		PassParameters->ScreenSpaceReflectionsSampler = TStaticSamplerState<SF_Point>::GetRHI();
+
+		if (Scene->HasVolumetricCloud())
+		{
+			FVolumetricCloudRenderSceneInfo* CloudInfo = Scene->GetVolumetricCloudSceneInfo();
+
+			PassParameters->CloudSkyAOTexture = View.VolumetricCloudSkyAO != nullptr ? View.VolumetricCloudSkyAO : SystemTextures.Black;
+			PassParameters->CloudSkyAOWorldToLightClipMatrix = CloudInfo->GetVolumetricCloudCommonShaderParameters().CloudSkyAOWorldToLightClipMatrix;
+			PassParameters->CloudSkyAOFarDepthKm = CloudInfo->GetVolumetricCloudCommonShaderParameters().CloudSkyAOFarDepthKm;
+			PassParameters->CloudSkyAOEnabled = 1;
+		}
+		else
+		{
+			PassParameters->CloudSkyAOTexture = SystemTextures.Black;
+			PassParameters->CloudSkyAOEnabled = 0;
+		}
+		PassParameters->CloudSkyAOSampler = TStaticSamplerState<SF_Bilinear>::GetRHI();
+
+		PassParameters->PreIntegratedGF = GSystemTextures.PreintegratedGF->GetRenderTargetItem().ShaderResourceTexture;
+		PassParameters->PreIntegratedGFSampler = TStaticSamplerState<SF_Bilinear, AM_Clamp, AM_Clamp, AM_Clamp>::GetRHI();
+
+		PassParameters->SceneTextures = SceneTextureParameters;
+
+		PassParameters->ViewUniformBuffer = View.ViewUniformBuffer;
+		PassParameters->ReflectionCaptureData = View.ReflectionCaptureUniformBuffer;
+		{
+			FReflectionUniformParameters ReflectionUniformParameters;
+			SetupReflectionUniformParameters(View, ReflectionUniformParameters);
+			PassParameters->ReflectionsParameters = CreateUniformBufferImmediate(ReflectionUniformParameters, UniformBuffer_SingleDraw);
+		}
+		PassParameters->ForwardLightData = View.ForwardLightingResources->ForwardLightDataUniformBuffer;
+
+		PassParameters->Strata =  Strata::BindStrataGlobalUniformParameters(View);
+	}
+
+	PassParameters->RenderTargets[0] = FRenderTargetBinding(SceneColorTexture.Target, ERenderTargetLoadAction::ELoad);
+	if (Strata::IsStrataEnabled() && Strata::IsClassificationEnabled())
+	{		
+		PassParameters->RenderTargets.DepthStencil = FDepthStencilBinding(SceneTextureParameters.SceneDepthTexture, ERenderTargetLoadAction::ELoad, ERenderTargetLoadAction::ELoad, FExclusiveDepthStencil::DepthRead_StencilWrite);
+	}
+
+	// Bind hair data
+	const bool bCheckerboardSubsurfaceRendering = IsSubsurfaceCheckerboardFormat(SceneColorTexture.Target->Desc.Format);
+
+	// ScreenSpace and SortedDeferred ray traced reflections use the same reflection environment shader,
+	// but main RT reflection shader requires a custom path as it evaluates the clear coat BRDF differently.
+	const bool bRequiresSpecializedReflectionEnvironmentShader = RayTracingReflectionOptions.bEnabled
+		&& RayTracingReflectionOptions.Algorithm != FRayTracingReflectionOptions::EAlgorithm::SortedDeferred;
+
+	auto PermutationVector = FReflectionEnvironmentSkyLightingPS::BuildPermutationVector(
+		View, bHasBoxCaptures, bHasSphereCaptures, DynamicBentNormalAO != 0.0f,
+		bSkyLight, bDynamicSkyLight, bApplySkyShadowing,
+		bRequiresSpecializedReflectionEnvironmentShader,
+		bStrataFastPath);
+
+	TShaderMapRef<FReflectionEnvironmentSkyLightingPS> PixelShader(View.ShaderMap, PermutationVector);
+	ClearUnusedGraphResources(PixelShader, PassParameters);
+
+	GraphBuilder.AddPass(
+		RDG_EVENT_NAME("ReflectionEnvironmentAndSky %dx%d", View.ViewRect.Width(), View.ViewRect.Height()),
+		PassParameters,
+		ERDGPassFlags::Raster,
+		[PassParameters, &View, PixelShader, bCheckerboardSubsurfaceRendering, bStrataFastPath](FRHICommandList& InRHICmdList)
+	{
+		InRHICmdList.SetViewport(View.ViewRect.Min.X, View.ViewRect.Min.Y, 0.0f, View.ViewRect.Max.X, View.ViewRect.Max.Y, 1.0f);
+
+		FGraphicsPipelineStateInitializer GraphicsPSOInit;
+		FPixelShaderUtils::InitFullscreenPipelineState(InRHICmdList, View.ShaderMap, PixelShader, GraphicsPSOInit);
+
+		if (Strata::IsStrataEnabled() && Strata::IsClassificationEnabled())
+		{
+			GraphicsPSOInit.DepthStencilState = TStaticDepthStencilState<
+				false, CF_Always, 
+				true,  CF_Equal, SO_Keep, SO_Keep, SO_Keep, 
+				false, CF_Equal, SO_Keep, SO_Keep, SO_Keep,
+				Strata::StencilBit, Strata::StencilBit>::GetRHI();
+		}
+
+		extern int32 GAOOverwriteSceneColor;
+		if (GetReflectionEnvironmentCVar() == 2 || GAOOverwriteSceneColor)
+		{
+			// override scene color for debugging
+			GraphicsPSOInit.BlendState = TStaticBlendState<>::GetRHI();
+		}
+		else
+		{
+			if (bCheckerboardSubsurfaceRendering)
+			{
+				GraphicsPSOInit.BlendState = TStaticBlendState<CW_RGB, BO_Add, BF_One, BF_One>::GetRHI();
+			}
+			else
+			{
+				GraphicsPSOInit.BlendState = TStaticBlendState<CW_RGBA, BO_Add, BF_One, BF_One, BO_Add, BF_One, BF_One>::GetRHI();
+			}
+		}
+
+		SetGraphicsPipelineState(InRHICmdList, GraphicsPSOInit);
+		SetShaderParameters(InRHICmdList, PixelShader, PixelShader.GetPixelShader(), *PassParameters);
+		if (Strata::IsStrataEnabled() && Strata::IsClassificationEnabled())
+		{
+			InRHICmdList.SetStencilRef(bStrataFastPath ? Strata::StencilBit : 0x0);
+		}
+		FPixelShaderUtils::DrawFullscreenTriangle(InRHICmdList);
+	});
+}
+
 void FDeferredShadingSceneRenderer::RenderDeferredReflectionsAndSkyLighting(
 	FRDGBuilder& GraphBuilder,
 	const FSceneTextures& SceneTextures,
@@ -1147,19 +1345,8 @@ void FDeferredShadingSceneRenderer::RenderDeferredReflectionsAndSkyLighting(
 
 	RDG_EVENT_SCOPE(GraphBuilder, "ReflectionIndirect");
 
-	const FRDGSystemTextures& SystemTextures = FRDGSystemTextures::Get(GraphBuilder);
-
 	const bool bReflectionEnv = ShouldDoReflectionEnvironment();
 
-	float DynamicBentNormalAO = 0.0f;
-	FRDGTextureRef AmbientOcclusionTexture = GetScreenSpaceAOFallback(SystemTextures);
-	
-	if (HasBeenProduced(SceneTextures.ScreenSpaceAO))
-	{
-		DynamicBentNormalAO = 1.0f;
-		AmbientOcclusionTexture = SceneTextures.ScreenSpaceAO;
-	}
-
 	FSceneTextureParameters SceneTextureParameters = GetSceneTextureParameters(GraphBuilder);
 	const auto& SceneColorTexture = SceneTextures.Color;
 
@@ -1319,157 +1506,40 @@ void FDeferredShadingSceneRenderer::RenderDeferredReflectionsAndSkyLighting(
 			RenderDeferredPlanarReflections(GraphBuilder, SceneTextureParameters, View, /* inout */ ReflectionsColor);
 		}
 
-		bool bRequiresApply = ReflectionsColor != nullptr || bSkyLight || bDynamicSkyLight || bReflectionEnv;
-
+		const bool bRequiresApply = ReflectionsColor != nullptr || bSkyLight || bDynamicSkyLight || bReflectionEnv;
 		if (bRequiresApply)
 		{
 			RDG_GPU_STAT_SCOPE(GraphBuilder, ReflectionEnvironment);
 
-			// Render the reflection environment with tiled deferred culling
-			bool bHasBoxCaptures = (View.NumBoxReflectionCaptures > 0);
-			bool bHasSphereCaptures = (View.NumSphereReflectionCaptures > 0);
-
-			FReflectionEnvironmentSkyLightingPS::FParameters* PassParameters = GraphBuilder.AllocParameters<FReflectionEnvironmentSkyLightingPS::FParameters>();
-
-			// Setup the parameters of the shader.
+			AddSkyReflectionPass(
+				GraphBuilder,
+				View,
+				Scene,
+				SceneTextures,
+				DynamicBentNormalAOTexture,
+				ReflectionsColor,
+				RayTracingReflectionOptions,
+				SceneTextureParameters,
+				bSkyLight,
+				bDynamicSkyLight,
+				bApplySkyShadowing,
+				false);
+			if (Strata::IsStrataEnabled() && Strata::IsClassificationEnabled())
 			{
-				// Setups all shader parameters related to skylight.
-				{
-					FSkyLightSceneProxy* SkyLight = Scene->SkyLight;
-
-					float SkyLightContrast = 0.01f;
-					float SkyLightOcclusionExponent = 1.0f;
-					FVector4 SkyLightOcclusionTintAndMinOcclusion(0.0f, 0.0f, 0.0f, 0.0f);
-					EOcclusionCombineMode SkyLightOcclusionCombineMode = EOcclusionCombineMode::OCM_MAX;
-					if (SkyLight)
-					{
-						FDistanceFieldAOParameters Parameters(SkyLight->OcclusionMaxDistance, SkyLight->Contrast);
-						SkyLightContrast = Parameters.Contrast;
-						SkyLightOcclusionExponent = SkyLight->OcclusionExponent;
-						SkyLightOcclusionTintAndMinOcclusion = FVector4(SkyLight->OcclusionTint);
-						SkyLightOcclusionTintAndMinOcclusion.W = SkyLight->MinOcclusion;
-						SkyLightOcclusionCombineMode = SkyLight->OcclusionCombineMode;
-					}
-
-					// Scale and bias to remap the contrast curve to [0,1]
-					const float Min = 1 / (1 + FMath::Exp(-SkyLightContrast * (0 * 10 - 5)));
-					const float Max = 1 / (1 + FMath::Exp(-SkyLightContrast * (1 * 10 - 5)));
-					const float Mul = 1.0f / (Max - Min);
-					const float Add = -Min / (Max - Min);
-
-					PassParameters->OcclusionTintAndMinOcclusion = SkyLightOcclusionTintAndMinOcclusion;
-					PassParameters->ContrastAndNormalizeMulAdd = FVector(SkyLightContrast, Mul, Add);
-					PassParameters->OcclusionExponent = SkyLightOcclusionExponent;
-					PassParameters->OcclusionCombineMode = SkyLightOcclusionCombineMode == OCM_Minimum ? 0.0f : 1.0f;
-					PassParameters->ApplyBentNormalAO = DynamicBentNormalAO;
-					PassParameters->InvSkySpecularOcclusionStrength = 1.0f / FMath::Max(CVarSkySpecularOcclusionStrength.GetValueOnRenderThread(), 0.1f);
-				}
-
-				// Setups all shader parameters related to distance field AO
-				{
-					FIntPoint AOBufferSize = GetBufferSizeForAO();
-					PassParameters->AOBufferBilinearUVMax = FVector2D(
-						(View.ViewRect.Width() / GAODownsampleFactor - 0.51f) / AOBufferSize.X, // 0.51 - so bilateral gather4 won't sample invalid texels
-						(View.ViewRect.Height() / GAODownsampleFactor - 0.51f) / AOBufferSize.Y);
-
-					extern float GAOViewFadeDistanceScale;
-					PassParameters->AOMaxViewDistance = GetMaxAOViewDistance();
-					PassParameters->DistanceFadeScale = 1.0f / ((1.0f - GAOViewFadeDistanceScale) * GetMaxAOViewDistance());
-
-					PassParameters->BentNormalAOTexture = DynamicBentNormalAOTexture;
-					PassParameters->BentNormalAOSampler = TStaticSamplerState<SF_Bilinear>::GetRHI();
-				}
-
-				PassParameters->AmbientOcclusionTexture = AmbientOcclusionTexture;
-				PassParameters->AmbientOcclusionSampler = TStaticSamplerState<SF_Point>::GetRHI();
-
-				PassParameters->ScreenSpaceReflectionsTexture = ReflectionsColor ? ReflectionsColor : SystemTextures.Black;
-				PassParameters->ScreenSpaceReflectionsSampler = TStaticSamplerState<SF_Point>::GetRHI();
-
-				if (Scene->HasVolumetricCloud())
-				{
-					FVolumetricCloudRenderSceneInfo* CloudInfo = Scene->GetVolumetricCloudSceneInfo();
-
-					PassParameters->CloudSkyAOTexture = View.VolumetricCloudSkyAO != nullptr ? View.VolumetricCloudSkyAO : SystemTextures.Black;
-					PassParameters->CloudSkyAOWorldToLightClipMatrix = CloudInfo->GetVolumetricCloudCommonShaderParameters().CloudSkyAOWorldToLightClipMatrix;
-					PassParameters->CloudSkyAOFarDepthKm = CloudInfo->GetVolumetricCloudCommonShaderParameters().CloudSkyAOFarDepthKm;
-					PassParameters->CloudSkyAOEnabled = 1;
-				}
-				else
-				{
-					PassParameters->CloudSkyAOTexture = SystemTextures.Black;
-					PassParameters->CloudSkyAOEnabled = 0;
-				}
-				PassParameters->CloudSkyAOSampler = TStaticSamplerState<SF_Bilinear>::GetRHI();
-
-				PassParameters->PreIntegratedGF = GSystemTextures.PreintegratedGF->GetRenderTargetItem().ShaderResourceTexture;
-				PassParameters->PreIntegratedGFSampler = TStaticSamplerState<SF_Bilinear, AM_Clamp, AM_Clamp, AM_Clamp>::GetRHI();
-
-				PassParameters->SceneTextures = SceneTextureParameters;
-
-				PassParameters->ViewUniformBuffer = View.ViewUniformBuffer;
-				PassParameters->ReflectionCaptureData = View.ReflectionCaptureUniformBuffer;
-				{
-					FReflectionUniformParameters ReflectionUniformParameters;
-					SetupReflectionUniformParameters(View, ReflectionUniformParameters);
-					PassParameters->ReflectionsParameters = CreateUniformBufferImmediate(ReflectionUniformParameters, UniformBuffer_SingleDraw);
-				}
-				PassParameters->ForwardLightData = View.ForwardLightingResources->ForwardLightDataUniformBuffer;
-
-				PassParameters->Strata =  Strata::BindStrataGlobalUniformParameters(View);
+				AddSkyReflectionPass(
+					GraphBuilder,
+					View,
+					Scene,
+					SceneTextures,
+					DynamicBentNormalAOTexture,
+					ReflectionsColor,
+					RayTracingReflectionOptions,
+					SceneTextureParameters,
+					bSkyLight,
+					bDynamicSkyLight,
+					bApplySkyShadowing,
+					true);
 			}
-
-			PassParameters->RenderTargets[0] = FRenderTargetBinding(SceneColorTexture.Target, ERenderTargetLoadAction::ELoad);
-
-			// Bind hair data
-			const bool bCheckerboardSubsurfaceRendering = IsSubsurfaceCheckerboardFormat(SceneColorTexture.Target->Desc.Format);
-
-			// ScreenSpace and SortedDeferred ray traced reflections use the same reflection environment shader,
-			// but main RT reflection shader requires a custom path as it evaluates the clear coat BRDF differently.
-			const bool bRequiresSpecializedReflectionEnvironmentShader = RayTracingReflectionOptions.bEnabled
-				&& RayTracingReflectionOptions.Algorithm != FRayTracingReflectionOptions::EAlgorithm::SortedDeferred;
-
-			auto PermutationVector = FReflectionEnvironmentSkyLightingPS::BuildPermutationVector(
-				View, bHasBoxCaptures, bHasSphereCaptures, DynamicBentNormalAO != 0.0f,
-				bSkyLight, bDynamicSkyLight, bApplySkyShadowing,
-				bRequiresSpecializedReflectionEnvironmentShader);
-
-			TShaderMapRef<FReflectionEnvironmentSkyLightingPS> PixelShader(View.ShaderMap, PermutationVector);
-			ClearUnusedGraphResources(PixelShader, PassParameters);
-
-			GraphBuilder.AddPass(
-				RDG_EVENT_NAME("ReflectionEnvironmentAndSky %dx%d", View.ViewRect.Width(), View.ViewRect.Height()),
-				PassParameters,
-				ERDGPassFlags::Raster,
-				[PassParameters, &View, PixelShader, bCheckerboardSubsurfaceRendering](FRHICommandList& InRHICmdList)
-			{
-				InRHICmdList.SetViewport(View.ViewRect.Min.X, View.ViewRect.Min.Y, 0.0f, View.ViewRect.Max.X, View.ViewRect.Max.Y, 1.0f);
-
-				FGraphicsPipelineStateInitializer GraphicsPSOInit;
-				FPixelShaderUtils::InitFullscreenPipelineState(InRHICmdList, View.ShaderMap, PixelShader, GraphicsPSOInit);
-
-				extern int32 GAOOverwriteSceneColor;
-				if (GetReflectionEnvironmentCVar() == 2 || GAOOverwriteSceneColor)
-				{
-					// override scene color for debugging
-					GraphicsPSOInit.BlendState = TStaticBlendState<>::GetRHI();
-				}
-				else
-				{
-					if (bCheckerboardSubsurfaceRendering)
-					{
-						GraphicsPSOInit.BlendState = TStaticBlendState<CW_RGB, BO_Add, BF_One, BF_One>::GetRHI();
-					}
-					else
-					{
-						GraphicsPSOInit.BlendState = TStaticBlendState<CW_RGBA, BO_Add, BF_One, BF_One, BO_Add, BF_One, BF_One>::GetRHI();
-					}
-				}
-
-				SetGraphicsPipelineState(InRHICmdList, GraphicsPSOInit);
-				SetShaderParameters(InRHICmdList, PixelShader, PixelShader.GetPixelShader(), *PassParameters);
-				FPixelShaderUtils::DrawFullscreenTriangle(InRHICmdList);
-			});
 		}
 
 		const bool bIsHairSkyLightingEnabled = HairDatas && (bSkyLight || bDynamicSkyLight || bReflectionEnv);
