Only add wave ops CFLAG to the relevant permutation.
- DX11 now at least works if you disable r.Shadow.v.SMRT.AdaptiveRayCount

diff --git a/Engine/Source/Runtime/Renderer/Private/VirtualShadowMaps/VirtualShadowMapProjection.cpp b/Engine/Source/Runtime/Renderer/Private/VirtualShadowMaps/VirtualShadowMapProjection.cpp
--- a/Engine/Source/Runtime/Renderer/Private/VirtualShadowMaps/VirtualShadowMapProjection.cpp
+++ b/Engine/Source/Runtime/Renderer/Private/VirtualShadowMaps/VirtualShadowMapProjection.cpp
@@ -476,7 +476,12 @@ class FVirtualShadowMapProjectionCS : public FGlobalShader
 		FVirtualShadowMapArray::SetShaderDefines(OutEnvironment);
 		FForwardLightingParameters::ModifyCompilationEnvironment(Parameters.Platform, OutEnvironment);
 
-		OutEnvironment.CompilerFlags.Add(CFLAG_WaveOperations);
+		FPermutationDomain PermutationVector(Parameters.PermutationId);
+
+		if (PermutationVector.Get<FSMRTAdaptiveRayCountDim>())
+		{
+			OutEnvironment.CompilerFlags.Add(CFLAG_WaveOperations);
+		}
 	}
 
 	static bool ShouldCompilePermutation(const FGlobalShaderPermutationParameters& Parameters)
