Adds prefixes on MotionBlur, DOF and Denoiser RDG resources

diff --git a/Engine/Source/Runtime/Renderer/Private/PostProcess/DiaphragmDOF.cpp b/Engine/Source/Runtime/Renderer/Private/PostProcess/DiaphragmDOF.cpp
--- a/Engine/Source/Runtime/Renderer/Private/PostProcess/DiaphragmDOF.cpp
+++ b/Engine/Source/Runtime/Renderer/Private/PostProcess/DiaphragmDOF.cpp
@@ -1531,8 +1531,8 @@ FRDGTextureRef DiaphragmDOF::AddPasses(
 
 	// Setup at lower resolution from full resolution scene color and scene depth.
 	{
-		FullResGatherInputTextures = CreateTextures(GraphBuilder, FullResGatherInputDescs, TEXT("DOFFullResSetup"));
-		HalfResGatherInputTextures = CreateTextures(GraphBuilder, HalfResGatherInputDescs, TEXT("DOFHalfResSetup"));
+		FullResGatherInputTextures = CreateTextures(GraphBuilder, FullResGatherInputDescs, TEXT("DOF.FullResSetup"));
+		HalfResGatherInputTextures = CreateTextures(GraphBuilder, HalfResGatherInputDescs, TEXT("DOF.HalfResSetup"));
 
 		bool bOutputFullResolution = bRecombineDoesSlightOutOfFocus && !bProcessSceneAlpha;
 		bool bOutputHalfResolution = true; // TODO: there is a useless shader permutation.
@@ -1659,8 +1659,8 @@ FRDGTextureRef DiaphragmDOF::AddPasses(
 			FIntPoint SrcSize = HalfResGatherInputTextures.SceneColor->Desc.Extent;
 
 			static const TCHAR* OutputDebugNames[2] = {
-				TEXT("DOFFlattenFgdCoc"),
-				TEXT("DOFFlattenBgdCoc"),
+				TEXT("DOF.FlattenFgdCoc"),
+				TEXT("DOF.FlattenBgdCoc"),
 			};
 			FlattenedTileClassificationTextures = CreateTextures(GraphBuilder, TileClassificationDescs, OutputDebugNames);
 
@@ -1706,13 +1706,13 @@ FRDGTextureRef DiaphragmDOF::AddPasses(
 		{
 			FDOFTileClassificationDescs OutputDescs = TileClassificationDescs;
 			const TCHAR* OutputDebugNames[2] = {
-				TEXT("DOFDilateFgdCoc"),
-				TEXT("DOFDilateBgdCoc"),
+				TEXT("DOF.DilateFgdCoc"),
+				TEXT("DOF.DilateBgdCoc"),
 			};
 			if (Mode == EDiaphragmDOFDilateCocMode::MinForegroundAndMaxBackground)
 			{
-				OutputDebugNames[0] = TEXT("DOFDilateMinFgdCoc");
-				OutputDebugNames[1] = TEXT("DOFDilateMaxBgdCoc");
+				OutputDebugNames[0] = TEXT("DOF.DilateMinFgdCoc");
+				OutputDebugNames[1] = TEXT("DOF.DilateMaxBgdCoc");
 				OutputDescs.Foreground.Format = PF_R16F;
 				OutputDescs.Background.Format = PF_R16F;
 			}
@@ -1870,7 +1870,7 @@ FRDGTextureRef DiaphragmDOF::AddPasses(
 				ReducedGatherInputDescs.SeparateCoc.Format = PF_R16F;
 			}
 
-			ReducedGatherInputTextures = CreateTextures(GraphBuilder, ReducedGatherInputDescs, TEXT("DOFReduce"));
+			ReducedGatherInputTextures = CreateTextures(GraphBuilder, ReducedGatherInputDescs, TEXT("DOF.Reduce"));
 		}
 
 		// Downsample the gather color setup to have faster neighborhood comparisons.
@@ -1887,7 +1887,7 @@ FRDGTextureRef DiaphragmDOF::AddPasses(
 				if (bRGBBufferSeparateCocBuffer && !bProcessSceneAlpha)
 					QuarterResGatherInputDescs.SceneColor.Format = PF_FloatR11G11B10;
 			
-				QuarterResGatherInputTextures = CreateTextures(GraphBuilder, QuarterResGatherInputDescs, TEXT("DOFDownsample"));
+				QuarterResGatherInputTextures = CreateTextures(GraphBuilder, QuarterResGatherInputDescs, TEXT("DOF.Downsample"));
 			}
 
 			FIntPoint PassViewSize = FIntPoint::DivideAndRoundUp(PreprocessViewSize, 2);
@@ -1921,9 +1921,9 @@ FRDGTextureRef DiaphragmDOF::AddPasses(
 
 			FRDGBufferDesc DrawListDescs = FRDGBufferDesc::CreateStructuredDesc(sizeof(float) * 4, 5 * MaxScatteringGroupCount);
 			if (bForegroundHybridScattering)
-				ForegroundScatterDrawListBuffer = GraphBuilder.CreateBuffer(DrawListDescs, TEXT("DOFForegroundDrawList"));
+				ForegroundScatterDrawListBuffer = GraphBuilder.CreateBuffer(DrawListDescs, TEXT("DOF.ForegroundDrawList"));
 			if (bBackgroundHybridScattering)
-				BackgroundScatterDrawListBuffer = GraphBuilder.CreateBuffer(DrawListDescs, TEXT("DOFBackgroundDrawList"));
+				BackgroundScatterDrawListBuffer = GraphBuilder.CreateBuffer(DrawListDescs, TEXT("DOF.BackgroundDrawList"));
 		}
 		
 		// Number of mip level that has been reduced.
@@ -2077,9 +2077,9 @@ FRDGTextureRef DiaphragmDOF::AddPasses(
 		}
 
 		static const TCHAR* const DebugNames[] = {
-			TEXT("DOFScatterBokehLUT"),
-			TEXT("DOFRecombineBokehLUT"),
-			TEXT("DOFGatherBokehLUT"),
+			TEXT("DOF.ScatterBokehLUT"),
+			TEXT("DOF.RecombineBokehLUT"),
+			TEXT("DOF.GatherBokehLUT"),
 		};
 
 		FRDGTextureDesc BokehLUTDesc = FRDGTextureDesc::Create2D(
@@ -2171,13 +2171,13 @@ FRDGTextureRef DiaphragmDOF::AddPasses(
 				{
 					const TCHAR* DebugName = nullptr;
 					if (ConvolutionSettings.LayerProcessing == EDiaphragmDOFLayerProcessing::ForegroundOnly)
-						DebugName = TEXT("DOFGatherForeground");
+						DebugName = TEXT("DOF.GatherForeground");
 					else if (ConvolutionSettings.LayerProcessing == EDiaphragmDOFLayerProcessing::ForegroundHoleFilling)
-						DebugName = TEXT("DOFGatherForegroundFill");
+						DebugName = TEXT("DOF.GatherForegroundFill");
 					else if (ConvolutionSettings.LayerProcessing == EDiaphragmDOFLayerProcessing::BackgroundOnly)
-						DebugName = TEXT("DOFGatherBackground");
+						DebugName = TEXT("DOF.GatherBackground");
 					else if (ConvolutionSettings.LayerProcessing == EDiaphragmDOFLayerProcessing::SlightOutOfFocus)
-						DebugName = TEXT("DOFGatherFocus");
+						DebugName = TEXT("DOF.GatherFocus");
 					else
 						check(0);
 
@@ -2203,7 +2203,7 @@ FRDGTextureRef DiaphragmDOF::AddPasses(
 			
 					const TCHAR* DebugName = nullptr;
 					if (ConvolutionSettings.LayerProcessing == EDiaphragmDOFLayerProcessing::BackgroundOnly)
-						DebugName = TEXT("DOFScatterOcclusionBackground");
+						DebugName = TEXT("DOF.ScatterOcclusionBackground");
 					else
 						check(0);
 			
@@ -2586,7 +2586,7 @@ FRDGTextureRef DiaphragmDOF::AddPasses(
 			FRDGTextureDesc Desc = InputSceneColor->Desc;
 			Desc.Reset();
 			Desc.Flags |= TexCreate_UAV;
-			NewSceneColor = GraphBuilder.CreateTexture(Desc, TEXT("DOFRecombine"));
+			NewSceneColor = GraphBuilder.CreateTexture(Desc, TEXT("DOF.Recombine"));
 		}
 
 		const FSeparateTranslucencyDimensions SeparateTranslucencyDimensions = SeparateTranslucencyTextures.GetDimensions();
diff --git a/Engine/Source/Runtime/Renderer/Private/PostProcess/PostProcessMotionBlur.cpp b/Engine/Source/Runtime/Renderer/Private/PostProcess/PostProcessMotionBlur.cpp
--- a/Engine/Source/Runtime/Renderer/Private/PostProcess/PostProcessMotionBlur.cpp
+++ b/Engine/Source/Runtime/Renderer/Private/PostProcess/PostProcessMotionBlur.cpp
@@ -461,7 +461,7 @@ void AddMotionBlurVelocityPass(
 				PF_FloatR11G11B10,
 				FClearValueBinding::None,
 				GFastVRamConfig.VelocityFlat | TexCreate_ShaderResource | TexCreate_UAV),
-			TEXT("VelocityFlat"));
+			TEXT("MotionBlur.VelocityFlat"));
 
 	FRDGTextureRef VelocityTileTextureSetup =
 		GraphBuilder.CreateTexture(
@@ -470,7 +470,7 @@ void AddMotionBlurVelocityPass(
 				PF_FloatRGBA,
 				FClearValueBinding::None,
 				GFastVRamConfig.VelocityMax | TexCreate_ShaderResource | TexCreate_UAV),
-			TEXT("VelocityTile"));
+			TEXT("MotionBlur.VelocityTile"));
 
 	const FMotionBlurParameters MotionBlurParametersNoScale = GetMotionBlurParameters(View, Viewports.Color.Rect.Size(), 1.0f);
 
@@ -503,7 +503,7 @@ void AddMotionBlurVelocityPass(
 				PF_FloatRGBA,
 				FClearValueBinding::None,
 				GFastVRamConfig.MotionBlur | TexCreate_ShaderResource | (ScatterDilatation ? TexCreate_RenderTargetable : TexCreate_UAV)),
-			TEXT("DilatedVelocityTile"));
+			TEXT("MotionBlur.DilatedVelocityTile"));
 
 	FMotionBlurVelocityDilateParameters VelocityDilateParameters;
 	VelocityDilateParameters.MotionBlur = MotionBlurParametersNoScale;
@@ -519,7 +519,7 @@ void AddMotionBlurVelocityPass(
 					PF_ShadowDepth,
 					FClearValueBinding::DepthOne,
 					TexCreate_DepthStencilTargetable),
-				TEXT("DilatedVelocityDepth"));
+				TEXT("MotionBlur.DilatedVelocityDepth"));
 
 		FMotionBlurVelocityDilateScatterParameters* PassParameters = GraphBuilder.AllocParameters<FMotionBlurVelocityDilateScatterParameters>();
 		PassParameters->Dilate = VelocityDilateParameters;
@@ -681,7 +681,7 @@ FRDGTextureRef AddMotionBlurFilterPass(
 	OutColorDesc.Flags |= GFastVRamConfig.MotionBlur;
 	OutColorDesc.Format =  IsPostProcessingWithAlphaChannelSupported() ? PF_FloatRGBA : PF_FloatRGB;
 
-	FRDGTextureRef ColorTextureOutput = GraphBuilder.CreateTexture(OutColorDesc, TEXT("MotionBlur"));
+	FRDGTextureRef ColorTextureOutput = GraphBuilder.CreateTexture(OutColorDesc, TEXT("MotionBlur.SceneColor"));
 
 	FMotionBlurFilterPermutationDomain PermutationVector;
 	PermutationVector.Set<FMotionBlurQualityDimension>(MotionBlurQuality);
@@ -734,7 +734,7 @@ FScreenPassTexture AddVisualizeMotionBlurPass(FRDGBuilder& GraphBuilder, const F
 
 	if (!Output.IsValid())
 	{
-		Output = FScreenPassRenderTarget::CreateFromInput(GraphBuilder, Inputs.SceneColor, View.GetOverwriteLoadAction(), TEXT("VisualizeMotionBlur"));
+		Output = FScreenPassRenderTarget::CreateFromInput(GraphBuilder, Inputs.SceneColor, View.GetOverwriteLoadAction(), TEXT("MotionBlur.Visualize"));
 	}
 
 	// NOTE: Scene depth is used as the velocity viewport because velocity can actually be a 1x1 black texture.
diff --git a/Engine/Source/Runtime/Renderer/Private/PostProcess/TemporalAA.cpp b/Engine/Source/Runtime/Renderer/Private/PostProcess/TemporalAA.cpp
--- a/Engine/Source/Runtime/Renderer/Private/PostProcess/TemporalAA.cpp
+++ b/Engine/Source/Runtime/Renderer/Private/PostProcess/TemporalAA.cpp
@@ -546,13 +546,13 @@ void SetupSampleWeightParameters(FTAAStandaloneCS::FParameters* OutTAAParameters
 DECLARE_GPU_STAT(TAA)
 
 const TCHAR* const kTAAOutputNames[] = {
-	TEXT("TemporalAA"),
-	TEXT("TemporalAA"),
-	TEXT("TemporalAA"),
-	TEXT("SSRTemporalAA"),
-	TEXT("LightShaftTemporalAA"),
-	TEXT("DOFTemporalAA"),
-	TEXT("DOFTemporalAA"),
+	TEXT("TAA.History"),
+	TEXT("TAA.History"),
+	TEXT("TAA.History"),
+	TEXT("SSR.TemporalAA"),
+	TEXT("LightShaft.TemporalAA"),
+	TEXT("DOF.TemporalAA"),
+	TEXT("DOF.TemporalAA"),
 };
 
 const TCHAR* const kTAAPassNames[] = {
diff --git a/Engine/Source/Runtime/Renderer/Private/ScreenSpaceDenoise.cpp b/Engine/Source/Runtime/Renderer/Private/ScreenSpaceDenoise.cpp
--- a/Engine/Source/Runtime/Renderer/Private/ScreenSpaceDenoise.cpp
+++ b/Engine/Source/Runtime/Renderer/Private/ScreenSpaceDenoise.cpp
@@ -370,8 +370,8 @@ class FMultiSPPDim : SHADER_PERMUTATION_BOOL("DIM_MULTI_SPP");
 
 const TCHAR* const kInjestResourceNames[] = {
 	// ShadowVisibilityMask
-	TEXT("ShadowDenoiserInjest0"),
-	TEXT("ShadowDenoiserInjest1"),
+	TEXT("Shadow.Denoiser.Injest0"),
+	TEXT("Shadow.Denoiser.Injest1"),
 	nullptr,
 	nullptr,
 
@@ -418,8 +418,8 @@ const TCHAR* const kInjestResourceNames[] = {
 	nullptr,
 
 	// VirtualShadowMapMask
-	TEXT("VirtualShadowDenoiserInjest0"),
-	TEXT("VirtualShadowDenoiserInjest1"),
+	TEXT("VirtualShadow.Denoiser.Injest0"),
+	TEXT("VirtualShadow.Denoiser.Injest1"),
 	nullptr,
 	nullptr,
 };
@@ -482,10 +482,10 @@ const TCHAR* const kReduceResourceNames[] = {
 
 const TCHAR* const kReconstructionResourceNames[] = {
 	// ShadowVisibilityMask
-	TEXT("ShadowReconstruction0"),
-	TEXT("ShadowReconstruction1"),
-	TEXT("ShadowReconstruction2"),
-	TEXT("ShadowReconstruction3"),
+	TEXT("Shadow.Denoiser.Reconstruction0"),
+	TEXT("Shadow.Denoiser.Reconstruction1"),
+	TEXT("Shadow.Denoiser.Reconstruction2"),
+	TEXT("Shadow.Denoiser.Reconstruction3"),
 
 	// PolychromaticPenumbraHarmonic
 	TEXT("PolychromaticPenumbraHarmonicReconstruction0"),
@@ -494,13 +494,13 @@ const TCHAR* const kReconstructionResourceNames[] = {
 	TEXT("PolychromaticPenumbraHarmonicReconstruction3"),
 
 	// Reflections
-	TEXT("ReflectionsReconstruction0"),
-	TEXT("ReflectionsReconstruction1"),
+	TEXT("Reflections.Denoiser.Reconstruction0"),
+	TEXT("Reflections.Denoiser.Reconstruction1"),
 	nullptr,
 	nullptr,
 
 	// AmbientOcclusion
-	TEXT("AOReconstruction0"),
+	TEXT("AO.Denoiser.Reconstruction0"),
 	nullptr,
 	nullptr,
 	nullptr,
@@ -518,8 +518,8 @@ const TCHAR* const kReconstructionResourceNames[] = {
 	TEXT("DiffuseHarmonicReconstruction3"),
 
 	// ScreenSpaceDiffuseIndirect
-	TEXT("SSGIReconstruction0"),
-	TEXT("SSGIReconstruction1"),
+	TEXT("SSGI.Denoiser.Reconstruction0"),
+	TEXT("SSGI.Denoiser.Reconstruction1"),
 	nullptr,
 	nullptr,
 
@@ -530,18 +530,18 @@ const TCHAR* const kReconstructionResourceNames[] = {
 	nullptr,
 
 	// VirtualShadowMapMask
-	TEXT("VirtualShadowReconstruction0"),
-	TEXT("VirtualShadowReconstruction1"),
-	TEXT("VirtualShadowReconstruction2"),
-	TEXT("VirtualShadowReconstruction3"),
+	TEXT("VirtualShadow.Denoiser.Reconstruction0"),
+	TEXT("VirtualShadow.Denoiser.Reconstruction1"),
+	TEXT("VirtualShadow.Denoiser.Reconstruction2"),
+	TEXT("VirtualShadow.Denoiser.Reconstruction3"),
 };
 
 const TCHAR* const kPreConvolutionResourceNames[] = {
 	// ShadowVisibilityMask
-	TEXT("ShadowPreConvolution0"),
-	TEXT("ShadowPreConvolution1"),
-	TEXT("ShadowPreConvolution2"),
-	TEXT("ShadowPreConvolution3"),
+	TEXT("Shadow.Denoiser.PreConvolution0"),
+	TEXT("Shadow.Denoiser.PreConvolution1"),
+	TEXT("Shadow.Denoiser.PreConvolution2"),
+	TEXT("Shadow.Denoiser.PreConvolution3"),
 
 	// PolychromaticPenumbraHarmonic
 	nullptr,
@@ -550,13 +550,13 @@ const TCHAR* const kPreConvolutionResourceNames[] = {
 	nullptr,
 
 	// Reflections
-	TEXT("ReflectionsPreConvolution0"),
-	TEXT("ReflectionsPreConvolution1"),
+	TEXT("Reflections.Denoiser.PreConvolution0"),
+	TEXT("Reflections.Denoiser.PreConvolution1"),
 	nullptr,
 	nullptr,
 
 	// AmbientOcclusion
-	TEXT("AOPreConvolution0"),
+	TEXT("AO.Denoiser.PreConvolution0"),
 	nullptr,
 	nullptr,
 	nullptr,
@@ -586,18 +586,18 @@ const TCHAR* const kPreConvolutionResourceNames[] = {
 	nullptr,
 
 	// VirtualShadowMapMask
-	TEXT("VirtualShadowPreConvolution0"),
-	TEXT("VirtualShadowPreConvolution1"),
-	TEXT("VirtualShadowPreConvolution2"),
-	TEXT("VirtualShadowPreConvolution3"),
+	TEXT("VirtualShadow.Denoiser.PreConvolution0"),
+	TEXT("VirtualShadow.Denoiser.PreConvolution1"),
+	TEXT("VirtualShadow.Denoiser.PreConvolution2"),
+	TEXT("VirtualShadow.Denoiser.PreConvolution3"),
 };
 
 const TCHAR* const kRejectionPreConvolutionResourceNames[] = {
 	// ShadowVisibilityMask
-	TEXT("ShadowRejectionPreConvolution0"),
-	TEXT("ShadowRejectionPreConvolution1"),
-	TEXT("ShadowRejectionPreConvolution2"),
-	TEXT("ShadowRejectionPreConvolution3"),
+	TEXT("Shadow.Denoiser.RejectionPreConvolution0"),
+	TEXT("Shadow.Denoiser.RejectionPreConvolution1"),
+	TEXT("Shadow.Denoiser.RejectionPreConvolution2"),
+	TEXT("Shadow.Denoiser.RejectionPreConvolution3"),
 
 	// PolychromaticPenumbraHarmonic
 	nullptr,
@@ -606,13 +606,13 @@ const TCHAR* const kRejectionPreConvolutionResourceNames[] = {
 	nullptr,
 
 	// Reflections
-	TEXT("ReflectionsRejectionPreConvolution0"),
-	TEXT("ReflectionsRejectionPreConvolution1"),
-	TEXT("ReflectionsRejectionPreConvolution2"),
+	TEXT("Reflections.Denoiser.RejectionPreConvolution0"),
+	TEXT("Reflections.Denoiser.RejectionPreConvolution1"),
+	TEXT("Reflections.Denoiser.RejectionPreConvolution2"),
 	nullptr,
 
 	// AmbientOcclusion
-	TEXT("AORejectionPreConvolution0"),
+	TEXT("AO.Denoiser.RejectionPreConvolution0"),
 	nullptr,
 	nullptr,
 	nullptr,
@@ -642,18 +642,18 @@ const TCHAR* const kRejectionPreConvolutionResourceNames[] = {
 	nullptr,
 
 	// VirtualShadowMapMask
-	TEXT("VirtualShadowRejectionPreConvolution0"),
-	TEXT("VirtualShadowRejectionPreConvolution1"),
-	TEXT("VirtualShadowRejectionPreConvolution2"),
-	TEXT("VirtualShadowRejectionPreConvolution3"),
+	TEXT("VirtualShadow.Denoiser.RejectionPreConvolution0"),
+	TEXT("VirtualShadow.Denoiser.RejectionPreConvolution1"),
+	TEXT("VirtualShadow.Denoiser.RejectionPreConvolution2"),
+	TEXT("VirtualShadow.Denoiser.RejectionPreConvolution3"),
 };
 
 const TCHAR* const kTemporalAccumulationResourceNames[] = {
 	// ShadowVisibilityMask
-	TEXT("ShadowTemporalAccumulation0"),
-	TEXT("ShadowTemporalAccumulation1"),
-	TEXT("ShadowTemporalAccumulation2"),
-	TEXT("ShadowTemporalAccumulation3"),
+	TEXT("Shadow.Denoiser.TemporalAccumulation0"),
+	TEXT("Shadow.Denoiser.TemporalAccumulation1"),
+	TEXT("Shadow.Denoiser.TemporalAccumulation2"),
+	TEXT("Shadow.Denoiser.TemporalAccumulation3"),
 
 	// PolychromaticPenumbraHarmonic
 	TEXT("PolychromaticPenumbraHistory0"),
@@ -662,20 +662,20 @@ const TCHAR* const kTemporalAccumulationResourceNames[] = {
 	nullptr,
 
 	// Reflections
-	TEXT("ReflectionsTemporalAccumulation0"),
-	TEXT("ReflectionsTemporalAccumulation1"),
+	TEXT("Reflections.Denoiser.TemporalAccumulation0"),
+	TEXT("Reflections.Denoiser.TemporalAccumulation1"),
 	nullptr,
 	nullptr,
 
 	// AmbientOcclusion
-	TEXT("AOTemporalAccumulation0"),
+	TEXT("AO.Denoiser.TemporalAccumulation0"),
 	nullptr,
 	nullptr,
 	nullptr,
 
 	// DiffuseIndirect
-	TEXT("DiffuseIndirectTemporalAccumulation0"),
-	TEXT("DiffuseIndirectTemporalAccumulation1"),
+	TEXT("DiffuseIndirect.Denoiser.TemporalAccumulation0"),
+	TEXT("DiffuseIndirect.Denoiser.TemporalAccumulation1"),
 	nullptr,
 	nullptr,
 
@@ -686,8 +686,8 @@ const TCHAR* const kTemporalAccumulationResourceNames[] = {
 	TEXT("DiffuseHarmonicTemporalAccumulation3"),
 
 	// ScreenSpaceDiffuseIndirect
-	TEXT("SSGITemporalAccumulation0"),
-	TEXT("SSGITemporalAccumulation1"),
+	TEXT("SSGI.Denoiser.TemporalAccumulation0"),
+	TEXT("SSGI.Denoiser.TemporalAccumulation1"),
 	nullptr,
 	nullptr,
 
@@ -698,18 +698,18 @@ const TCHAR* const kTemporalAccumulationResourceNames[] = {
 	nullptr,
 
 	// VirtualShadowMapMask
-	TEXT("VirtualShadowTemporalAccumulation0"),
-	TEXT("VirtualShadowTemporalAccumulation1"),
-	TEXT("VirtualShadowTemporalAccumulation2"),
-	TEXT("VirtualShadowTemporalAccumulation3"),
+	TEXT("VirtualShadow.Denoiser.TemporalAccumulation0"),
+	TEXT("VirtualShadow.Denoiser.TemporalAccumulation1"),
+	TEXT("VirtualShadow.Denoiser.TemporalAccumulation2"),
+	TEXT("VirtualShadow.Denoiser.TemporalAccumulation3"),
 };
 
 const TCHAR* const kHistoryConvolutionResourceNames[] = {
 	// ShadowVisibilityMask
-	TEXT("ShadowHistoryConvolution0"),
-	TEXT("ShadowHistoryConvolution1"),
-	TEXT("ShadowHistoryConvolution2"),
-	TEXT("ShadowHistoryConvolution3"),
+	TEXT("Shadow.Denoiser.HistoryConvolution0"),
+	TEXT("Shadow.Denoiser.HistoryConvolution1"),
+	TEXT("Shadow.Denoiser.HistoryConvolution2"),
+	TEXT("Shadow.Denoiser.HistoryConvolution3"),
 
 	// PolychromaticPenumbraHarmonic
 	nullptr,
@@ -718,20 +718,20 @@ const TCHAR* const kHistoryConvolutionResourceNames[] = {
 	nullptr,
 
 	// Reflections
-	TEXT("ReflectionsHistoryConvolution0"),
-	TEXT("ReflectionsHistoryConvolution1"),
+	TEXT("Reflections.Denoiser.HistoryConvolution0"),
+	TEXT("Reflections.Denoiser.HistoryConvolution1"),
 	nullptr,
 	nullptr,
 
 	// AmbientOcclusion
-	TEXT("AOHistoryConvolution0"),
+	TEXT("AO.Denoiser.HistoryConvolution0"),
 	nullptr,
 	nullptr,
 	nullptr,
 
 	// DiffuseIndirect
-	TEXT("DiffuseIndirectHistoryConvolution0"),
-	TEXT("DiffuseIndirectHistoryConvolution1"),
+	TEXT("DiffuseIndirect.Denoiser.HistoryConvolution0"),
+	TEXT("DiffuseIndirect.Denoiser.HistoryConvolution1"),
 	nullptr,
 	nullptr,
 
@@ -754,18 +754,18 @@ const TCHAR* const kHistoryConvolutionResourceNames[] = {
 	nullptr,
 
 	// VirtualShadowMapMask
-	TEXT("VirtualShadowHistoryConvolution0"),
-	TEXT("VirtualShadowHistoryConvolution1"),
-	TEXT("VirtualShadowHistoryConvolution2"),
-	TEXT("VirtualShadowHistoryConvolution3"),
+	TEXT("VirtualShadow.Denoiser.HistoryConvolution0"),
+	TEXT("VirtualShadow.Denoiser.HistoryConvolution1"),
+	TEXT("VirtualShadow.Denoiser.HistoryConvolution2"),
+	TEXT("VirtualShadow.Denoiser.HistoryConvolution3"),
 };
 
 const TCHAR* const kDenoiserOutputResourceNames[] = {
 	// ShadowVisibilityMask
-	TEXT("ShadowDenoiserOutput0"),
-	TEXT("ShadowDenoiserOutput1"),
-	TEXT("ShadowDenoiserOutput2"),
-	TEXT("ShadowDenoiserOutput3"),
+	TEXT("Shadow.Denoiser.DenoiserOutput0"),
+	TEXT("Shadow.Denoiser.DenoiserOutput1"),
+	TEXT("Shadow.Denoiser.DenoiserOutput2"),
+	TEXT("Shadow.Denoiser.DenoiserOutput3"),
 
 	// PolychromaticPenumbraHarmonic
 	nullptr,
@@ -810,10 +810,10 @@ const TCHAR* const kDenoiserOutputResourceNames[] = {
 	nullptr,
 
 	// VirtualShadowMapMask
-	TEXT("VirtualShadowDenoiserOutput0"),
-	TEXT("VirtualShadowDenoiserOutput1"),
-	TEXT("VirtualShadowDenoiserOutput2"),
-	TEXT("VirtualShadowDenoiserOutput3"),
+	TEXT("VirtualShadow.Denoiser.DenoiserOutput0"),
+	TEXT("VirtualShadow.Denoiser.DenoiserOutput1"),
+	TEXT("VirtualShadow.Denoiser.DenoiserOutput2"),
+	TEXT("VirtualShadow.Denoiser.DenoiserOutput3"),
 };
 
 static_assert(UE_ARRAY_COUNT(kReconstructionResourceNames) == int32(ESignalProcessing::MAX) * kMaxBufferProcessingCount, "You forgot me!");
