Fixed missing SetStencilRef for starta classification when rendering non directional light.

diff --git a/Engine/Source/Runtime/Renderer/Private/LightRendering.cpp b/Engine/Source/Runtime/Renderer/Private/LightRendering.cpp
--- a/Engine/Source/Runtime/Renderer/Private/LightRendering.cpp
+++ b/Engine/Source/Runtime/Renderer/Private/LightRendering.cpp
@@ -2243,6 +2243,19 @@ void FDeferredShadingSceneRenderer::RenderLight(
 		RenderLightParams.HairCategorizationTexture = InHairVisibilityViews->HairDatas[ViewIndex].CategorizationTexture->GetPooledRenderTarget();
 	}
 
+	if (Strata::IsStrataEnabled() && Strata::IsClassificationEnabled())
+	{
+		GraphicsPSOInit.DepthStencilState = TStaticDepthStencilState<
+			false, CF_Always,
+			true, CF_Equal, SO_Keep, SO_Keep, SO_Keep,
+			true, CF_Equal, SO_Keep, SO_Keep, SO_Keep,
+			Strata::StencilBit, 0x0>::GetRHI();
+	}
+	else
+	{
+		GraphicsPSOInit.DepthStencilState = TStaticDepthStencilState<false, CF_Always>::GetRHI();
+	}
+
 	if (LightProxy->GetLightType() == LightType_Directional)
 	{
 		// Turn DBT back off
@@ -2250,16 +2263,6 @@ void FDeferredShadingSceneRenderer::RenderLight(
 		TShaderMapRef<TDeferredLightVS<false> > VertexShader(View.ShaderMap);
 
 		GraphicsPSOInit.RasterizerState = TStaticRasterizerState<FM_Solid, CM_None>::GetRHI();
-		GraphicsPSOInit.DepthStencilState = TStaticDepthStencilState<false, CF_Always>::GetRHI();
-
-		if (Strata::IsStrataEnabled() && Strata::IsClassificationEnabled())
-		{
-			GraphicsPSOInit.DepthStencilState = TStaticDepthStencilState<
-				false, CF_Always,
-				true, CF_Equal, SO_Keep, SO_Keep, SO_Keep,
-				true, CF_Equal, SO_Keep, SO_Keep, SO_Keep,
-				Strata::StencilBit, 0x0>::GetRHI();
-		}
 
 		if (bRenderOverlap)
 		{
@@ -2323,9 +2326,9 @@ void FDeferredShadingSceneRenderer::RenderLight(
 		}
 
 		VertexShader->SetParameters(RHICmdList, View, LightSceneInfo);
+		RHICmdList.SetStencilRef(bStrataFastPath ? Strata::StencilBit : 0u);
 
 		// Apply the directional light as a full screen quad
-		RHICmdList.SetStencilRef(bStrataFastPath ? Strata::StencilBit : 0u);
 		DrawRectangle(
 			RHICmdList,
 			0, 0,
@@ -2390,6 +2393,7 @@ void FDeferredShadingSceneRenderer::RenderLight(
 		}
 
 		VertexShader->SetParameters(RHICmdList, View, LightSceneInfo);
+		RHICmdList.SetStencilRef(bStrataFastPath ? Strata::StencilBit : 0u);
 
 		// Use DBT to allow work culling on shadow lights
 		if (GraphicsPSOInit.bDepthBounds)
