Fixed missing starta permutation specification for dynamic light pass.

diff --git a/Engine/Source/Runtime/Renderer/Private/LightRendering.cpp b/Engine/Source/Runtime/Renderer/Private/LightRendering.cpp
--- a/Engine/Source/Runtime/Renderer/Private/LightRendering.cpp
+++ b/Engine/Source/Runtime/Renderer/Private/LightRendering.cpp
@@ -2381,7 +2381,7 @@ void FDeferredShadingSceneRenderer::RenderLight(
 			PermutationVector.Set < FDeferredLightPS::FAtmosphereTransmittance >(false);
 			PermutationVector.Set< FDeferredLightPS::FCloudTransmittance >(false);
 			PermutationVector.Set< FDeferredLightPS::FStrata >(Strata::IsStrataEnabled());
-			PermutationVector.Set< FDeferredLightPS::FStrata >(Strata::IsStrataEnabled() && Strata::IsClassificationEnabled() && bStrataFastPath);
+			PermutationVector.Set< FDeferredLightPS::FStrataFastPath >(Strata::IsStrataEnabled() && Strata::IsClassificationEnabled() && bStrataFastPath);
 
 			TShaderMapRef< FDeferredLightPS > PixelShader( View.ShaderMap, PermutationVector );
 			GraphicsPSOInit.BoundShaderState.VertexDeclarationRHI = GetVertexDeclarationFVector4();
