Fix a bug that the MobileSceneTextures are not updated on mobile preview(D3D).

diff --git a/Engine/Source/Runtime/Renderer/Private/MobileShadingRenderer.cpp b/Engine/Source/Runtime/Renderer/Private/MobileShadingRenderer.cpp
--- a/Engine/Source/Runtime/Renderer/Private/MobileShadingRenderer.cpp
+++ b/Engine/Source/Runtime/Renderer/Private/MobileShadingRenderer.cpp
@@ -760,7 +760,7 @@ BEGIN_SHADER_PARAMETER_STRUCT(FMobilePostBasePassViewExtensionParameters, )
 	RENDER_TARGET_BINDING_SLOTS()
 END_SHADER_PARAMETER_STRUCT()
 
-void FMobileSceneRenderer::RenderForward(FRDGBuilder& GraphBuilder, const TArrayView<const FViewInfo> ViewList, FRDGTextureRef ViewFamilyTexture, const FSceneTextures& SceneTextures)
+void FMobileSceneRenderer::RenderForward(FRDGBuilder& GraphBuilder, const TArrayView<const FViewInfo> ViewList, FRDGTextureRef ViewFamilyTexture, FSceneTextures& SceneTextures)
 {
 	const FViewInfo& View = ViewList[0];
 
@@ -899,6 +899,9 @@ void FMobileSceneRenderer::RenderForward(FRDGBuilder& GraphBuilder, const TArray
 	{
 		// Make a copy of the scene depth if the current hardware doesn't support reading and writing to the same depth buffer
 		ConditionalResolveSceneDepth(GraphBuilder, View, SceneTextures.Depth);
+
+		SceneTextures.MobileSetupMode = EMobileSceneTextureSetupMode::SceneDepth | EMobileSceneTextureSetupMode::SceneDepthAux | EMobileSceneTextureSetupMode::CustomDepth;
+		SceneTextures.MobileUniformBuffer = CreateMobileSceneTextureUniformBuffer(GraphBuilder, SceneTextures.MobileSetupMode);
 	}
 
 	if (bRequiresPixelProjectedPlanarRelfectionPass)
@@ -1092,7 +1095,7 @@ void MobileDeferredCopyBuffer(FRHICommandListImmediate& RHICmdList, const FViewI
 		VertexShader);
 }
 
-void FMobileSceneRenderer::RenderDeferred(FRDGBuilder& GraphBuilder, const TArrayView<const FViewInfo> ViewList, const FSortedLightSetSceneInfo& SortedLightSet, FRDGTextureRef ViewFamilyTexture, const FSceneTextures& SceneTextures)
+void FMobileSceneRenderer::RenderDeferred(FRDGBuilder& GraphBuilder, const TArrayView<const FViewInfo> ViewList, const FSortedLightSetSceneInfo& SortedLightSet, FRDGTextureRef ViewFamilyTexture, FSceneTextures& SceneTextures)
 {
 	TArray<FRDGTextureRef, TInlineAllocator<5>> ColorTargets;
 
@@ -1231,6 +1234,9 @@ void FMobileSceneRenderer::RenderDeferred(FRDGBuilder& GraphBuilder, const TArra
 		
 		// SceneColor + GBuffer write, SceneDepth is read only
 		{
+			SceneTextures.MobileSetupMode = EMobileSceneTextureSetupMode::SceneDepth | EMobileSceneTextureSetupMode::SceneDepthAux | EMobileSceneTextureSetupMode::CustomDepth;
+			SceneTextures.MobileUniformBuffer = CreateMobileSceneTextureUniformBuffer(GraphBuilder, SceneTextures.MobileSetupMode);
+
 			for (int32 i = 0; i < ColorTargets.Num(); ++i)
 			{
 				BasePassRenderTargets[i].SetLoadAction(ERenderTargetLoadAction::ELoad);
@@ -1268,6 +1274,9 @@ void FMobileSceneRenderer::RenderDeferred(FRDGBuilder& GraphBuilder, const TArra
 
 		// SceneColor write, SceneDepth is read only
 		{
+			SceneTextures.MobileSetupMode = EMobileSceneTextureSetupMode::SceneDepth | EMobileSceneTextureSetupMode::SceneDepthAux | EMobileSceneTextureSetupMode::GBuffers | EMobileSceneTextureSetupMode::CustomDepth;
+			SceneTextures.MobileUniformBuffer = CreateMobileSceneTextureUniformBuffer(GraphBuilder, SceneTextures.MobileSetupMode);
+
 			for (int32 i = 1; i < ColorTargets.Num(); ++i)
 			{
 				BasePassRenderTargets[i] = FRenderTargetBinding();
diff --git a/Engine/Source/Runtime/Renderer/Private/SceneRendering.h b/Engine/Source/Runtime/Renderer/Private/SceneRendering.h
--- a/Engine/Source/Runtime/Renderer/Private/SceneRendering.h
+++ b/Engine/Source/Runtime/Renderer/Private/SceneRendering.h
@@ -2143,8 +2143,8 @@ protected:
 	void UpdateDirectionalLightUniformBuffers(FRDGBuilder& GraphBuilder, const FViewInfo& View);
 	void UpdateSkyReflectionUniformBuffer();
 
-	void RenderForward(FRDGBuilder& GraphBuilder, const TArrayView<const FViewInfo> ViewList, FRDGTextureRef ViewFamilyTexture, const FSceneTextures& SceneTextures);
-	void RenderDeferred(FRDGBuilder& GraphBuilder, const TArrayView<const FViewInfo> ViewList, const FSortedLightSetSceneInfo& SortedLightSet, FRDGTextureRef ViewFamilyTexture, const FSceneTextures& SceneTextures);
+	void RenderForward(FRDGBuilder& GraphBuilder, const TArrayView<const FViewInfo> ViewList, FRDGTextureRef ViewFamilyTexture, FSceneTextures& SceneTextures);
+	void RenderDeferred(FRDGBuilder& GraphBuilder, const TArrayView<const FViewInfo> ViewList, const FSortedLightSetSceneInfo& SortedLightSet, FRDGTextureRef ViewFamilyTexture, FSceneTextures& SceneTextures);
 	
 	void RenderAmbientOcclusion(FRDGBuilder& GraphBuilder, FRDGTextureRef SceneDepthTexture, FRDGTextureRef AmbientOcclusionTexture);
 
