Only tag stencil once for start simple/complex tile categorisation for all light passes.

diff --git a/Engine/Source/Runtime/Renderer/Private/LightRendering.cpp b/Engine/Source/Runtime/Renderer/Private/LightRendering.cpp
--- a/Engine/Source/Runtime/Renderer/Private/LightRendering.cpp
+++ b/Engine/Source/Runtime/Renderer/Private/LightRendering.cpp
@@ -1260,6 +1260,14 @@ void FDeferredShadingSceneRenderer::RenderLights(
 	{
 		RDG_EVENT_SCOPE(GraphBuilder, "DirectLighting");
 
+		// STRATA_TODO move right after stencil clear so that it is also common with EnvLight pass
+		if (ViewFamily.EngineShowFlags.DirectLighting &&
+			Strata::IsStrataEnabled() && Strata::IsClassificationEnabled())
+		{
+			// Update the stencil buffer, marking simple/complex strata material only once for all the following passes.
+			Strata::AddStrataStencilPass(GraphBuilder, Views, SceneTextures);
+		}
+
 		if(ViewFamily.EngineShowFlags.DirectLighting)
 		{
 			RDG_EVENT_SCOPE(GraphBuilder, "NonShadowedLights");
@@ -1340,11 +1348,6 @@ void FDeferredShadingSceneRenderer::RenderLights(
 
 			if (!bUseHairLighting)
 			{
-				if (Strata::IsStrataEnabled() && Strata::IsClassificationEnabled())
-				{
-					Strata::AddStrataStencilPass(GraphBuilder, Views, SceneTextures);
-				}
-
 				FRenderLightParameters* PassParameters = GraphBuilder.AllocParameters<FRenderLightParameters>();
 				GetRenderLightParameters(SceneTextures, nullptr, LightingChannelsTexture, {}, HairDatas ? &HairDatas->HairVisibilityViews : nullptr, *PassParameters);
 
@@ -2453,11 +2456,6 @@ void FDeferredShadingSceneRenderer::RenderLight(
 		RDG_EVENT_SCOPE_CONDITIONAL(GraphBuilder, ViewCount > 1, "View%d", ViewIndex);
 		const FViewInfo& View = Views[ViewIndex];
 
-		if (Strata::IsStrataEnabled() && Strata::IsClassificationEnabled())
-		{
-			Strata::AddStrataStencilPass(GraphBuilder, View, SceneTextures);
-		}
-
 		FRenderLightParameters* PassParameters = GraphBuilder.AllocParameters<FRenderLightParameters>();
 		GetRenderLightParameters(SceneTextures, ScreenShadowMaskTexture, LightingChannelsTexture, GetCloudShadowAOParameters(GraphBuilder, View, CloudInfo), InHairVisibilityViews, *PassParameters);
 
