Restore UE4 HackOverrideViewMatrixForShadows and fix uses in Nanite shadow rendering
- This fixes a number (unknown) of issues where the ShadowView was assumed to be mostly unmodified (e.g., HISM LOD selection).
- Also fixing an unrelated issue in cubemap rendering using instance culling ang geometry shader.

diff --git a/Engine/Shaders/Private/ShadowDepthVertexShader.usf b/Engine/Shaders/Private/ShadowDepthVertexShader.usf
--- a/Engine/Shaders/Private/ShadowDepthVertexShader.usf
+++ b/Engine/Shaders/Private/ShadowDepthVertexShader.usf
@@ -465,9 +465,9 @@ void Main(
 
 	OutPosition = OutParameters.GSPosition = WorldPos;
 	OutParameters.WorldNormal = WorldNormal;
-#if defined(GPUCULL_TODO) && ENABLE_FALLBACK_POINTLIGHT_SHADOW_GS
+#if defined(GPUCULL_TODO) && ENABLE_FALLBACK_POINTLIGHT_SHADOW_GS && IS_FOR_GEOMETRY_SHADER
 	OutParameters.ViewIndex = VertexFactoryGetViewIndex(VFIntermediates);
-#endif // !ENABLE_FALLBACK_POINTLIGHT_SHADOW_GS
+#endif // defined(GPUCULL_TODO) && ENABLE_FALLBACK_POINTLIGHT_SHADOW_GS && IS_FOR_GEOMETRY_SHADER
 
 	#if INTERPOLATE_VF_ATTRIBUTES
 		// Masked materials need texture coords to clip
diff --git a/Engine/Source/Runtime/Engine/Public/SceneView.h b/Engine/Source/Runtime/Engine/Public/SceneView.h
--- a/Engine/Source/Runtime/Engine/Public/SceneView.h
+++ b/Engine/Source/Runtime/Engine/Public/SceneView.h
@@ -406,45 +406,10 @@ public:
 		return ProjectionMatrix.M[3][3] < 1.0f;
 	}
 
-	inline void HackOverrideMatrixForShadows(
-		const FMatrix& InViewRotationMatrix,
-		const FMatrix& InProjectionMatrix,
-		const FVector& InViewOrigin
-		)
+	inline void HackOverrideViewMatrixForShadows(const FMatrix& InViewMatrix)
 	{
-		ViewOrigin = InViewOrigin;
-
-		FMatrix ViewRotationMatrix = InViewRotationMatrix;
-		if (!ViewRotationMatrix.GetOrigin().IsNearlyZero(0.0f))
-		{
-			ViewOrigin += ViewRotationMatrix.InverseTransformPosition(FVector::ZeroVector);
-			ViewRotationMatrix = ViewRotationMatrix.RemoveTranslation();
-		}
-
-		ViewMatrix = FTranslationMatrix(-ViewOrigin) * ViewRotationMatrix;
-
-		ProjectionMatrix = InProjectionMatrix;
-		InvProjectionMatrix = InvertProjectionMatrix( ProjectionMatrix );
-
-		// For precision reasons the view matrix inverse is calculated independently.
-		InvViewMatrix = ViewRotationMatrix.GetTransposed() * FTranslationMatrix( ViewOrigin );
-
-		// Translate world-space so its origin is at ViewOrigin for improved precision.
-		// Note that this isn't exactly right for orthogonal projections (See the above special case), but we still use ViewOrigin
-		// in that case so the same value may be used in shaders for both the world-space translation and the camera's world position.
-		PreViewTranslation = -FVector(ViewOrigin);
-
-		// Compute a transform from view origin centered world-space to clip space.
-		OverriddenTranslatedViewMatrix = ViewRotationMatrix;
-		OverriddenInvTranslatedViewMatrix = ViewRotationMatrix.GetTransposed();
-
-		// Compute the view projection matrix and its inverse.
-		ViewProjectionMatrix = GetViewMatrix() * GetProjectionMatrix();
-		InvViewProjectionMatrix = GetInvProjectionMatrix() * GetInvViewMatrix();
-
-		// Compute a transform from view origin centered world-space to clip space.
-		TranslatedViewProjectionMatrix = OverriddenTranslatedViewMatrix * GetProjectionMatrix();
-		InvTranslatedViewProjectionMatrix = GetInvProjectionMatrix() * OverriddenInvTranslatedViewMatrix;
+		OverriddenTranslatedViewMatrix = ViewMatrix = InViewMatrix;
+		OverriddenInvTranslatedViewMatrix = InViewMatrix.Inverse();
 	}
 
 	void SaveProjectionNoAAMatrix()
diff --git a/Engine/Source/Runtime/Renderer/Private/ShadowDepthRendering.cpp b/Engine/Source/Runtime/Renderer/Private/ShadowDepthRendering.cpp
--- a/Engine/Source/Runtime/Renderer/Private/ShadowDepthRendering.cpp
+++ b/Engine/Source/Runtime/Renderer/Private/ShadowDepthRendering.cpp
@@ -323,6 +323,9 @@ public:
 		OutEnvironment.SetDefine(TEXT("USING_VERTEX_SHADER_LAYER"), (uint32)(ShaderMode == VertexShadowDepth_VSLayer));
 		OutEnvironment.SetDefine(TEXT("POSITION_ONLY"), (uint32)bUsePositionOnlyStream);
 		OutEnvironment.SetDefine(TEXT("IS_FOR_GEOMETRY_SHADER"), (uint32)bIsForGeometryShader);
+#if defined(GPUCULL_TODO)
+		OutEnvironment.SetDefine(TEXT("ENABLE_FALLBACK_POINTLIGHT_SHADOW_GS"), UseGPUScene(Parameters.Platform) ? 1 : 0);
+#endif // defined(GPUCULL_TODO)
 
 		uint32 bEnableNonNaniteVSM = (uint32)(ENABLE_NON_NANITE_VSM != 0 && UseGPUScene(Parameters.Platform));
 		OutEnvironment.SetDefine(TEXT("ENABLE_NON_NANITE_VSM"), bEnableNonNaniteVSM);
@@ -1223,7 +1226,7 @@ void FProjectedShadowInfo::ModifyViewForShadow(FRHICommandList& RHICmdList, FVie
 	FIntRect OriginalViewRect = FoundView->ViewRect;
 	FoundView->ViewRect = GetOuterViewRect();
 
-	//FoundView->ViewMatrices.HackRemoveTemporalAAProjectionJitter();
+	FoundView->ViewMatrices.HackRemoveTemporalAAProjectionJitter();
 
 	if (CascadeSettings.bFarShadowCascade)
 	{
@@ -1236,9 +1239,7 @@ void FProjectedShadowInfo::ModifyViewForShadow(FRHICommandList& RHICmdList, FVie
 	FoundView->CachedViewUniformShaderParameters = MakeUnique<FViewUniformShaderParameters>();
 
 	// Override the view matrix so that billboarding primitives will be aligned to the light
-	FoundView->ViewMatrices.HackOverrideMatrixForShadows(TranslatedWorldToView, ViewToClipOuter, -PreShadowTranslation);
-	FoundView->PrevViewInfo.ViewMatrices.HackOverrideMatrixForShadows(TranslatedWorldToView, ViewToClipOuter, -PreShadowTranslation);
-
+	FoundView->ViewMatrices.HackOverrideViewMatrixForShadows(TranslatedWorldToView);
 	FBox VolumeBounds[TVC_MAX];
 	FoundView->SetupUniformBufferParameters(
 		VolumeBounds,
@@ -1371,10 +1372,10 @@ static void UpdatePackedViewParamsFromPrevShadowState(Nanite::FPackedViewParams&
 	}
 }
 
-static void UpdateCurrentFrameHZB(FLightSceneInfo& LightSceneInfo, const FPersistentShadowStateKey& ShadowKey, const FProjectedShadowInfo* ProjectedShadowInfo, const TRefCountPtr<IPooledRenderTarget>& HZB)
+static void UpdateCurrentFrameHZB(FLightSceneInfo& LightSceneInfo, const FPersistentShadowStateKey& ShadowKey, const FProjectedShadowInfo* ProjectedShadowInfo, const TRefCountPtr<IPooledRenderTarget>& HZB, int32 CubeFaceIndex = -1)
 {
 	FPersistentShadowState State;
-	State.ViewMatrices = ProjectedShadowInfo->ShadowDepthView->ViewMatrices;
+	State.ViewMatrices = ProjectedShadowInfo->GetShadowDepthRenderingViewMatrices(CubeFaceIndex);
 	State.HZBTestViewRect = ProjectedShadowInfo->GetInnerViewRect();
 	State.HZB = HZB;
 	LightSceneInfo.PersistentShadows.Add(ShadowKey, State);
@@ -1411,7 +1412,7 @@ static void RenderShadowDepthAtlasNanite(
 		}
 
 		Nanite::FPackedViewParams Initializer;
-		Initializer.ViewMatrices = ProjectedShadowInfo->ShadowDepthView->ViewMatrices;
+		Initializer.ViewMatrices = ProjectedShadowInfo->GetShadowDepthRenderingViewMatrices();
 		Initializer.ViewRect = ProjectedShadowInfo->GetOuterViewRect();
 		Initializer.RasterContextSize = AtlasSize;
 		Initializer.LODScaleFactor = ComputeNaniteShadowsLODScaleFactor();
@@ -1531,7 +1532,7 @@ static void RenderShadowDepthAtlasNanite(
 				ShadowMap,
 				AtlasViewRect,
 				AtlasViewRect.Min,
-				ProjectedShadowInfo->ShadowDepthView->ViewMatrices.GetProjectionMatrix(),
+				ProjectedShadowInfo->GetShadowDepthRenderingViewMatrices().GetProjectionMatrix(),
 				ProjectedShadowInfo->GetShaderDepthBias(),
 				ProjectedShadowInfo->bDirectionalLight
 			);
@@ -1910,20 +1911,7 @@ void FSceneRenderer::RenderShadowDepthMaps(FRDGBuilder& GraphBuilder, FInstanceC
 						{
 							Nanite::FPackedViewParams Params = BaseParams;
 							Params.TargetLayerIndex = ProjectedShadowInfo->VirtualShadowMaps[i]->ID;
-
-							if( ProjectedShadowInfo->bOnePassPointLightShadow )
-							{
-								FViewMatrices::FMinimalInitializer MatricesInitializer;
-								MatricesInitializer.ViewOrigin = -ProjectedShadowInfo->PreShadowTranslation;
-								MatricesInitializer.ViewRotationMatrix = ProjectedShadowInfo->OnePassShadowViewMatrices[i] * FScaleMatrix(FVector(1, -1, 1));	// TODO add cull direction to FPackedView instead of this scale nonsense.
-								MatricesInitializer.ProjectionMatrix = ProjectedShadowInfo->OnePassShadowFaceProjectionMatrix;
-								MatricesInitializer.ConstrainedViewRect = BaseParams.ViewRect;
-								Params.ViewMatrices = FViewMatrices(MatricesInitializer);
-							}
-							else
-							{
-								Params.ViewMatrices = ProjectedShadowInfo->ShadowDepthView->ViewMatrices;
-							}
+							Params.ViewMatrices = ProjectedShadowInfo->GetShadowDepthRenderingViewMatrices(i, true);
 
 							int32 HZBKey = ProjectedShadowInfo->GetLightSceneInfo().Id + (i << 24);
 							FVirtualShadowMapHZBMetadata& PrevHZB = CacheManager->HZBMetadata.FindOrAdd(HZBKey);
@@ -2120,14 +2108,8 @@ void FSceneRenderer::RenderShadowDepthMaps(FRDGBuilder& GraphBuilder, FInstanceC
 					// Setup packed view
 					TArray<Nanite::FPackedView, SceneRenderingAllocator> PackedViews;
 					{
-						FViewMatrices::FMinimalInitializer MatricesInitializer;
-						MatricesInitializer.ViewOrigin = -ProjectedShadowInfo->PreShadowTranslation;
-						MatricesInitializer.ViewRotationMatrix = ProjectedShadowInfo->OnePassShadowViewMatrices[CubemapFaceIndex];
-						MatricesInitializer.ProjectionMatrix = ProjectedShadowInfo->OnePassShadowFaceProjectionMatrix;
-						MatricesInitializer.ConstrainedViewRect = ShadowViewRect;
-
 						Nanite::FPackedViewParams Params;
-						Params.ViewMatrices = FViewMatrices(MatricesInitializer);
+						Params.ViewMatrices = ProjectedShadowInfo->GetShadowDepthRenderingViewMatrices(CubemapFaceIndex);
 						Params.ViewRect = ShadowViewRect;
 						Params.RasterContextSize = TargetSize;
 						Params.LODScaleFactor = ComputeNaniteShadowsLODScaleFactor();
@@ -2181,7 +2163,7 @@ void FSceneRenderer::RenderShadowDepthMaps(FRDGBuilder& GraphBuilder, FInstanceC
 
 						ConvertToExternalTexture(GraphBuilder, FurthestHZBTexture, HZB);
 					}
-					UpdateCurrentFrameHZB(LightSceneInfo, ShadowKey, ProjectedShadowInfo, HZB);
+					UpdateCurrentFrameHZB(LightSceneInfo, ShadowKey, ProjectedShadowInfo, HZB, CubemapFaceIndex);
 				}
 			}
 		}
diff --git a/Engine/Source/Runtime/Renderer/Private/ShadowRendering.cpp b/Engine/Source/Runtime/Renderer/Private/ShadowRendering.cpp
--- a/Engine/Source/Runtime/Renderer/Private/ShadowRendering.cpp
+++ b/Engine/Source/Runtime/Renderer/Private/ShadowRendering.cpp
@@ -2268,5 +2268,36 @@ void FEmptyTranslucentSelfShadowUniformBuffer::InitDynamicRHI()
 	Super::InitDynamicRHI();
 }
 
+
+FViewMatrices FProjectedShadowInfo::GetShadowDepthRenderingViewMatrices(int32 CubeFaceIndex, bool bUseForVSMCubeFaceWorkaround) const
+{
+	FViewMatrices::FMinimalInitializer MatricesInitializer;
+	MatricesInitializer.ViewOrigin = -PreShadowTranslation;
+	MatricesInitializer.ConstrainedViewRect = GetOuterViewRect();
+
+	ensure(!bOnePassPointLightShadow || CubeFaceIndex >= 0 && CubeFaceIndex < 6);
+	if (bOnePassPointLightShadow && CubeFaceIndex >= 0 && CubeFaceIndex < 6)
+	{
+		if (bUseForVSMCubeFaceWorkaround)
+		{
+			// TODO add cull direction to FPackedView instead of this scale nonsense.
+			MatricesInitializer.ViewRotationMatrix = OnePassShadowViewMatrices[CubeFaceIndex] * FScaleMatrix(FVector(1, -1, 1));
+		}
+		else
+		{
+			MatricesInitializer.ViewRotationMatrix = OnePassShadowViewMatrices[CubeFaceIndex];
+		}
+		MatricesInitializer.ProjectionMatrix = OnePassShadowFaceProjectionMatrix;
+	}
+	else
+	{
+		MatricesInitializer.ViewRotationMatrix = TranslatedWorldToView;
+		MatricesInitializer.ProjectionMatrix = ViewToClipOuter;
+	}
+
+	return FViewMatrices(MatricesInitializer);
+}
+
+
 /** */
 TGlobalResource< FEmptyTranslucentSelfShadowUniformBuffer > GEmptyTranslucentSelfShadowUniformBuffer;
diff --git a/Engine/Source/Runtime/Renderer/Private/ShadowRendering.h b/Engine/Source/Runtime/Renderer/Private/ShadowRendering.h
--- a/Engine/Source/Runtime/Renderer/Private/ShadowRendering.h
+++ b/Engine/Source/Runtime/Renderer/Private/ShadowRendering.h
@@ -247,7 +247,15 @@ class FProjectedShadowInfo : public FRefCountedObject
 public:
 	typedef TArray<const FPrimitiveSceneInfo*,SceneRenderingAllocator> PrimitiveArrayType;
 
-	/** The view to be used when rendering this shadow's depths. */
+	/** 
+	 * The view to be used when rendering this shadow's depths. 
+	 * WARNING: This view is a bastardization of the 'main view' - i.e., the view used to render the visible geometry (or one of them in the case of multi-view, e.g., split screen)
+	 * the only substantial differences are that the _view_ matrix is overridden (using the 'HackOverrideViewMatrixForShadows')
+	 * and the 'HackRemoveTemporalAAProjectionJitter' has been called, to presumably remove the AA Jitter.
+	 * IT DOES NOT contain the correct shadow projection matrix, PreViewTranslation, or much else that you would expect.
+	 * Instead use the GetShadowDepthRenderingViewMatrices() function to get these for rendering shadows.
+	 * Main reason for this state of affairs is to facilitate code that depends on the 'main' view matrices for LOD calculations in various places.
+	 */
 	FViewInfo* ShadowDepthView;
 
 	/** The depth or color targets this shadow was rendered to. */
@@ -461,6 +469,13 @@ public:
 		);
 #endif // ENABLE_NON_NANITE_VSM
 
+	/**
+	 * Get `FViewMatrices` instance set up for drawing geometry to the SM, using the outer projection,
+	 * @parameter bUseForVSMCubeFaceWorkaround - If true, when get a cube face the OnePassShadowViewMatrices is flipped in Y to undo the 
+	 *                                           flip used when creating OnePassShadowViewMatrices. This is also done when sampling the VSM.
+	 */
+	FViewMatrices GetShadowDepthRenderingViewMatrices(int32 CubeFaceIndex = -1, bool bUseForVSMCubeFaceWorkaround = false) const;
+
 	float GetShaderDepthBias() const { return ShaderDepthBias; }
 	float GetShaderSlopeDepthBias() const { return ShaderSlopeDepthBias; }
 	float GetShaderMaxSlopeDepthBias() const { return ShaderMaxSlopeDepthBias; }
diff --git a/Engine/Source/Runtime/Renderer/Private/ShadowSetup.cpp b/Engine/Source/Runtime/Renderer/Private/ShadowSetup.cpp
--- a/Engine/Source/Runtime/Renderer/Private/ShadowSetup.cpp
+++ b/Engine/Source/Runtime/Renderer/Private/ShadowSetup.cpp
@@ -2043,14 +2043,8 @@ void FProjectedShadowInfo::SetupMeshDrawCommandsForShadowDepth(FSceneRenderer& R
 			// Setup packed view
 			TArray<Nanite::FPackedView, SceneRenderingAllocator> PackedViews;
 			{
-				FViewMatrices::FMinimalInitializer MatricesInitializer;
-				MatricesInitializer.ViewOrigin = -PreShadowTranslation;
-				MatricesInitializer.ViewRotationMatrix = OnePassShadowViewMatrices[CubemapFaceIndex];
-				MatricesInitializer.ProjectionMatrix = OnePassShadowFaceProjectionMatrix;
-				MatricesInitializer.ConstrainedViewRect = ShadowViewRect;
-
 				Nanite::FPackedViewParams Params;
-				Params.ViewMatrices = FViewMatrices(MatricesInitializer);
+				Params.ViewMatrices = GetShadowDepthRenderingViewMatrices(CubemapFaceIndex);
 				// TODO: Real prev frame matrices
 				Params.PrevViewMatrices = Params.ViewMatrices;	
 				Params.ViewRect = ShadowViewRect;
@@ -2077,7 +2071,7 @@ void FProjectedShadowInfo::SetupMeshDrawCommandsForShadowDepth(FSceneRenderer& R
 	else
 	{
 		Nanite::FPackedViewParams Params;
-		Params.ViewMatrices = ShadowDepthView->ViewMatrices;
+		Params.ViewMatrices = GetShadowDepthRenderingViewMatrices();
 		// TODO: Real prev frame matrices
 		Params.PrevViewMatrices = Params.ViewMatrices;
 		Params.ViewRect = GetInnerViewRect();
@@ -2250,7 +2244,7 @@ void FProjectedShadowInfo::GatherDynamicMeshElements(FSceneRenderer& Renderer, F
 		FMatrix OriginalViewMatrix = ShadowDepthView->ViewMatrices.GetViewMatrix();
 
 		// Override the view matrix so that billboarding primitives will be aligned to the light
-		ShadowDepthView->ViewMatrices.HackOverrideMatrixForShadows(TranslatedWorldToView, ViewToClipOuter, -PreShadowTranslation);
+		ShadowDepthView->ViewMatrices.HackOverrideViewMatrixForShadows(TranslatedWorldToView);
 
 		ReusedViewsArray[0] = ShadowDepthView;
 
