Fix for tile classification on some platforms.

diff --git a/Engine/Source/Runtime/Renderer/Private/Strata/Strata.cpp b/Engine/Source/Runtime/Renderer/Private/Strata/Strata.cpp
--- a/Engine/Source/Runtime/Renderer/Private/Strata/Strata.cpp
+++ b/Engine/Source/Runtime/Renderer/Private/Strata/Strata.cpp
@@ -286,6 +286,8 @@ class FStrataMaterialClassificationPassPS : public FGlobalShader
 	{
 		FGlobalShader::ModifyCompilationEnvironment(Parameters, OutEnvironment);
 		OutEnvironment.SetDefine(TEXT("SHADER_CATEGORIZATION"), 1);
+		OutEnvironment.SetRenderTargetOutputFormat(0, PF_R32_UINT);
+		OutEnvironment.SetRenderTargetOutputFormat(1, PF_R32G32_UINT);
 	}
 };
 IMPLEMENT_GLOBAL_SHADER(FStrataMaterialClassificationPassPS, "/Engine/Private/Strata/StrataMaterialClassification.usf", "MainPS", SF_Pixel);
@@ -304,7 +306,7 @@ class FStrataMaterialTileClassificationPassCS : public FGlobalShader
 		SHADER_PARAMETER(int32, TileSize)
 		SHADER_PARAMETER(int32, bRectPrimitive)
 		SHADER_PARAMETER(FIntPoint, ViewResolution)
-		SHADER_PARAMETER_RDG_TEXTURE(Texture2D, ClassificationTexture)
+		SHADER_PARAMETER_RDG_TEXTURE(Texture2D<uint>, ClassificationTexture)
 		SHADER_PARAMETER_RDG_BUFFER_UAV(RWBuffer, TileIndirectData)
 		SHADER_PARAMETER_RDG_BUFFER_UAV(RWBuffer, TileListData)
 	END_SHADER_PARAMETER_STRUCT()
