Decouple RVT streaming mips count from validity of streaming mips texture.
This allows multiple RVT volumes to point at the same streamin mips, which will be useful for landscape streaming scenarios.

diff --git a/Engine/Source/Runtime/Engine/Classes/Components/RuntimeVirtualTextureComponent.h b/Engine/Source/Runtime/Engine/Classes/Components/RuntimeVirtualTextureComponent.h
--- a/Engine/Source/Runtime/Engine/Classes/Components/RuntimeVirtualTextureComponent.h
+++ b/Engine/Source/Runtime/Engine/Classes/Components/RuntimeVirtualTextureComponent.h
@@ -45,26 +45,26 @@ protected:
 	UPROPERTY(EditAnywhere, AdvancedDisplay, Category = VirtualTexture)
 	bool bHidePrimitives = false;
 
-	/** Texture object containing streamed low mips. */
+	/** Texture object containing streamed low mips. This can reduce rendering update cost. */
 	UPROPERTY(EditAnywhere, BlueprintReadOnly, NonPIEDuplicateTransient, Category = VirtualTextureBuild)
 	TObjectPtr<UVirtualTextureBuilder> StreamingTexture = nullptr;
 
-	/** Number of low mips to serialize and stream for the virtual texture. This can reduce rendering update cost. */
-	UPROPERTY(EditAnywhere, Category = VirtualTextureBuild, meta = (UIMin = "0", UIMax = "12", DisplayName = "Streaming Levels"))
+	/** Number of streaming low mips to build for the virtual texture. */
+	UPROPERTY(EditAnywhere, Category = VirtualTextureBuild, meta = (UIMin = "0", UIMax = "12", DisplayName = "Build Levels"))
 	int32 StreamLowMips = 0;
 
 	/** Placeholder for details customization button. */
 	UPROPERTY(VisibleAnywhere, Transient, Category = VirtualTextureBuild)
 	bool bBuildStreamingMipsButton;
 
+	/** Use streaming low mips when rendering in editor. Set true to view and debug the baked streaming low mips. */
+	UPROPERTY(EditAnywhere, AdvancedDisplay, Category = VirtualTextureBuild, meta = (DisplayName = "View in Editor"))
+	bool bUseStreamingLowMipsInEditor = false;
+
 	/** Enable Crunch texture compression for the streaming low mips. Generic ZLib compression is used when Crunch is disabled. */
 	UPROPERTY(EditAnywhere, AdvancedDisplay, Category = VirtualTextureBuild, meta = (DisplayName = "Enable Crunch"))
 	bool bEnableCompressCrunch = false;
 
-	/** Use any streaming low mips when rendering in editor. Set true to view and debug the baked streaming low mips. */
-	UPROPERTY(EditAnywhere, AdvancedDisplay, Category = VirtualTextureBuild, meta = (DisplayName = "View Streaming Mips in Editor"))
-	bool bUseStreamingLowMipsInEditor = false;
-
 	/** Build the streaming low mips using debug coloring. This can help show where streaming mips are being used. */
 	UPROPERTY(EditAnywhere, BlueprintReadOnly, Transient, AdvancedDisplay, Category = VirtualTextureBuild, meta = (DisplayName = "Build Debug"))
 	bool bBuildDebugStreamingMips = false;
diff --git a/Engine/Source/Runtime/Engine/Classes/VT/RuntimeVirtualTexture.h b/Engine/Source/Runtime/Engine/Classes/VT/RuntimeVirtualTexture.h
--- a/Engine/Source/Runtime/Engine/Classes/VT/RuntimeVirtualTexture.h
+++ b/Engine/Source/Runtime/Engine/Classes/VT/RuntimeVirtualTexture.h
@@ -181,11 +181,15 @@ class UVirtualTexture2D;
 
 namespace RuntimeVirtualTexture
 {
-	/** Helper function to wrap a runtime virtual texture producer with a streaming producer. */
+	/** Helper function to create a streaming virtual texture producer. */
 	ENGINE_API IVirtualTexture* CreateStreamingTextureProducer(
-		IVirtualTexture* InProducer,
-		FVTProducerDescription const& InProducerDesc,
 		UVirtualTexture2D* InStreamingTexture,
-		int32 InMaxLevel,
-		int32& OutTransitionLevel);
+		FVTProducerDescription const& InOwnerProducerDesc,
+		FVTProducerDescription& OutStreamingProducerDesc);
+
+	/** Helper function to bind a runtime virtual texture producer to a streaming producer. Returns the new combined producer. */
+	ENGINE_API IVirtualTexture* BindStreamingTextureProducer(
+		IVirtualTexture* InProducer,
+		IVirtualTexture* InStreamingProducer,
+		int32 InTransitionLevel);
 }
diff --git a/Engine/Source/Runtime/Engine/Private/Components/RuntimeVirtualTextureComponent.cpp b/Engine/Source/Runtime/Engine/Private/Components/RuntimeVirtualTextureComponent.cpp
--- a/Engine/Source/Runtime/Engine/Private/Components/RuntimeVirtualTextureComponent.cpp
+++ b/Engine/Source/Runtime/Engine/Private/Components/RuntimeVirtualTextureComponent.cpp
@@ -120,7 +120,6 @@ uint64 URuntimeVirtualTextureComponent::CalculateStreamingTextureSettingsHash()
 			uint32 MaterialType : 4;
 			uint32 TileSize : 12;
 			uint32 TileBorderSize : 4;
-			uint32 StreamLowMips : 4;
 			uint32 LODGroup : 8;
 			uint32 CompressTextures : 1;
 			uint32 SinglePhysicalSpace : 1;
@@ -134,7 +133,6 @@ uint64 URuntimeVirtualTextureComponent::CalculateStreamingTextureSettingsHash()
 	Settings.MaterialType = (uint32)VirtualTexture->GetMaterialType();
 	Settings.TileSize = (uint32)VirtualTexture->GetTileSize();
 	Settings.TileBorderSize = (uint32)VirtualTexture->GetTileBorderSize();
-	Settings.StreamLowMips = (uint32)StreamLowMips;
 	Settings.LODGroup = (uint32)VirtualTexture->GetLODGroup();
 	Settings.CompressTextures = (uint32)VirtualTexture->GetCompressTextures();
 	Settings.ContinuousUpdate = (uint32)VirtualTexture->GetContinuousUpdate();
@@ -157,7 +155,7 @@ bool URuntimeVirtualTextureComponent::IsStreamingLowMips() const
 		return false;
 	}
 #endif
-	return StreamLowMips > 0 && IsStreamingTextureValid();
+	return IsStreamingTextureValid();
 }
 
 #if WITH_EDITOR
@@ -214,13 +212,17 @@ void URuntimeVirtualTextureComponent::InitializeStreamingTexture(uint32 InSizeX,
 bool URuntimeVirtualTextureComponent::CanEditChange(const FProperty* InProperty) const
 {
 	bool bCanEdit = Super::CanEditChange(InProperty);
-	if (InProperty->GetFName() == TEXT("bEnableCompressCrunch"))
+	if (InProperty->GetFName() == TEXT("bUseStreamingLowMipsInEditor"))
 	{
-		bCanEdit &= NumStreamingMips() > 0 && GetVirtualTexture() != nullptr && GetVirtualTexture()->GetCompressTextures();
+		bCanEdit &= GetVirtualTexture() != nullptr && GetStreamingTexture() != nullptr;
 	}
-	else if (InProperty->GetFName() == TEXT("bUseStreamingLowMipsInEditor"))
+	else if (InProperty->GetFName() == TEXT("bEnableCompressCrunch"))
 	{
-		bCanEdit &= GetStreamingTexture() != nullptr && NumStreamingMips() > 0;
+		bCanEdit &= GetVirtualTexture() != nullptr && GetVirtualTexture()->GetCompressTextures();
+	}
+	else if (InProperty->GetFName() == TEXT("bBuildDebugStreamingMips"))
+	{
+		bCanEdit = GetVirtualTexture() != nullptr && GetVirtualTexture()->GetMaterialType() != ERuntimeVirtualTextureMaterialType::WorldHeight;
 	}
 	return bCanEdit;
 }
diff --git a/Engine/Source/Runtime/Engine/Private/VT/RuntimeVirtualTexture.cpp b/Engine/Source/Runtime/Engine/Private/VT/RuntimeVirtualTexture.cpp
--- a/Engine/Source/Runtime/Engine/Private/VT/RuntimeVirtualTexture.cpp
+++ b/Engine/Source/Runtime/Engine/Private/VT/RuntimeVirtualTexture.cpp
@@ -591,42 +591,51 @@ void URuntimeVirtualTexture::PostEditChangeProperty(FPropertyChangedEvent& Prope
 namespace RuntimeVirtualTexture
 {
 	IVirtualTexture* CreateStreamingTextureProducer(
-		IVirtualTexture* InProducer,
-		FVTProducerDescription const& InProducerDesc,
 		UVirtualTexture2D* InStreamingTexture,
-		int32 InMaxLevel,
-		int32& OutTransitionLevel)
+		FVTProducerDescription const& InOwnerProducerDesc,
+		FVTProducerDescription& OutStreamingProducerDesc)
 	{
-		if (InProducer != nullptr && InStreamingTexture != nullptr)
+		OutStreamingProducerDesc = InOwnerProducerDesc;
+
+		if(InStreamingTexture != nullptr)
 		{
+			OutStreamingProducerDesc.Name = InStreamingTexture->GetFName();
+
 			FTexturePlatformData* StreamingTextureData = InStreamingTexture->GetPlatformData();
-			if (StreamingTextureData != nullptr)
+			if(ensure(StreamingTextureData != nullptr))
 			{
 				FVirtualTextureBuiltData* VTData = StreamingTextureData->VTData;
-
-				ensure(InProducerDesc.TileSize == VTData->TileSize);
-				ensure(InProducerDesc.TileBorderSize == VTData->TileBorderSize);
-				if (InProducerDesc.TileSize == VTData->TileSize && InProducerDesc.TileBorderSize == VTData->TileBorderSize)
+				if(ensure(VTData != nullptr))
 				{
+					ensure(InOwnerProducerDesc.TileSize == VTData->TileSize);
+					ensure(InOwnerProducerDesc.TileBorderSize == VTData->TileBorderSize);
+					ensure(VTData->GetNumMips() > 0);
+
 					// Note that streaming data may have mips added/removed during cook.
-					const uint32 Size = FMath::Max(VTData->Width, VTData->Height);
-					const uint32 NumTiles = FMath::DivideAndRoundUp(Size, VTData->TileSize);
-					const uint32 NumMips = FMath::CeilLogTwo(NumTiles) + 1;
-
-					// If the streaming texture is bigger then the runtime virtual texture then offset the first mip.
-					const int32 TransitionLevel = InMaxLevel - (int32)NumMips + 1;
-					const int32 FirstStreamingMip = TransitionLevel < 0 ? -TransitionLevel : 0;
-					const int32 AdjustedTransitionLevel = TransitionLevel + FirstStreamingMip;
-					OutTransitionLevel = TransitionLevel;
-
-					IVirtualTexture* StreamingProducer = new FUploadingVirtualTexture(InStreamingTexture->GetFName(), VTData, FirstStreamingMip);
-					return new FVirtualTextureLevelRedirector(InProducer, StreamingProducer, AdjustedTransitionLevel);
+					const uint32 BlockWidthInTiles = VTData->GetWidthInTiles();
+					const uint32 BlockHeightInTiles = VTData->GetHeightInTiles();
+					const uint32 MaxLevel = FMath::CeilLogTwo(FMath::Max(BlockWidthInTiles, BlockHeightInTiles));
+
+					// Clamp the streaming texture size to the runtime virtual texture.
+					const uint32 FirstMipToUse = MaxLevel > InOwnerProducerDesc.MaxLevel ? MaxLevel - InOwnerProducerDesc.MaxLevel : 0;
+
+					OutStreamingProducerDesc.BlockWidthInTiles = BlockWidthInTiles >> FirstMipToUse;
+					OutStreamingProducerDesc.BlockHeightInTiles = BlockHeightInTiles >> FirstMipToUse;
+					OutStreamingProducerDesc.MaxLevel = MaxLevel - FirstMipToUse;
+
+					return new FUploadingVirtualTexture(InStreamingTexture->GetFName(), VTData, FirstMipToUse);
 				}
 			}
 		}
 
-		// Can't create a streaming producer so return original producer.
-		OutTransitionLevel = InMaxLevel;
-		return InProducer;
+		return nullptr;
+	}
+
+	IVirtualTexture* BindStreamingTextureProducer(
+		IVirtualTexture* InProducer,
+		IVirtualTexture* InStreamingProducer,
+		int32 InTransitionLevel)
+	{
+		return (InStreamingProducer == nullptr) ? InProducer : new FVirtualTextureLevelRedirector(InProducer, InStreamingProducer, InTransitionLevel);
 	}
 }
diff --git a/Engine/Source/Runtime/Renderer/Private/VT/RuntimeVirtualTextureSceneProxy.cpp b/Engine/Source/Runtime/Renderer/Private/VT/RuntimeVirtualTextureSceneProxy.cpp
--- a/Engine/Source/Runtime/Renderer/Private/VT/RuntimeVirtualTextureSceneProxy.cpp
+++ b/Engine/Source/Runtime/Renderer/Private/VT/RuntimeVirtualTextureSceneProxy.cpp
@@ -37,32 +37,38 @@ FRuntimeVirtualTextureSceneProxy::FRuntimeVirtualTextureSceneProxy(URuntimeVirtu
 		const FBox Bounds = InComponent->Bounds.GetBox();
 
 		// The producer description is calculated using the transform to determine the aspect ratio
-		FVTProducerDescription Desc;
-		VirtualTexture->GetProducerDescription(Desc, InitSettings, Transform);
-		VirtualTextureSize = FIntPoint(Desc.BlockWidthInTiles * Desc.TileSize, Desc.BlockHeightInTiles * Desc.TileSize);
+		FVTProducerDescription ProducerDesc;
+		VirtualTexture->GetProducerDescription(ProducerDesc, InitSettings, Transform);
+		VirtualTextureSize = FIntPoint(ProducerDesc.BlockWidthInTiles * ProducerDesc.TileSize, ProducerDesc.BlockHeightInTiles * ProducerDesc.TileSize);
 		// We only need to dirty flush up to the producer description MaxLevel which accounts for the RemoveLowMips
-		MaxDirtyLevel = Desc.MaxLevel;
+		MaxDirtyLevel = ProducerDesc.MaxLevel;
 
 		const ERuntimeVirtualTextureMaterialType MaterialType = VirtualTexture->GetMaterialType();
 		const bool bClearTextures = VirtualTexture->GetClearTextures();
 
-		// The Producer object created here will be passed into the virtual texture system which will take ownership.
-		IVirtualTexture* Producer = new FRuntimeVirtualTextureProducer(Desc, ProducerId, MaterialType, bClearTextures, InComponent->GetScene(), Transform, Bounds);
+		// The producer object created here will be passed into the virtual texture system which will take ownership.
+		IVirtualTexture* Producer = new FRuntimeVirtualTextureProducer(ProducerDesc, ProducerId, MaterialType, bClearTextures, InComponent->GetScene(), Transform, Bounds);
 
+		// Create a producer for the streaming low mips. 
+		// This is bound with the main producer so that one allocated VT can use both runtime or streaming producers dependent on mip level.
 		if (InComponent->IsStreamingLowMips())
 		{
 			UVirtualTexture2D* StreamingTexture = InComponent->GetStreamingTexture()->Texture;
-			// Streaming mips start from the MaxLevel before taking into account the RemoveLowMips
-			const int32 MaxLevel = FMath::CeilLogTwo(FMath::Max(Desc.BlockWidthInTiles, Desc.BlockHeightInTiles));
-			// Wrap our producer to use a streaming producer for low mips
-			int32 StreamingTransitionLevel;
-			Producer = RuntimeVirtualTexture::CreateStreamingTextureProducer(Producer, Desc, StreamingTexture, MaxLevel, StreamingTransitionLevel);
+
+			FVTProducerDescription StreamingProducerDesc;
+			IVirtualTexture* StreamingProducer = RuntimeVirtualTexture::CreateStreamingTextureProducer(StreamingTexture, ProducerDesc, StreamingProducerDesc);
+
+			ensure(ProducerDesc.MaxLevel >= StreamingProducerDesc.MaxLevel);
+			const int32 TransitionLevel = ProducerDesc.MaxLevel - StreamingProducerDesc.MaxLevel;
+
+			Producer = RuntimeVirtualTexture::BindStreamingTextureProducer(Producer, StreamingProducer, TransitionLevel);
+
 			// Any dirty flushes don't need to flush the streaming mips (they only change with a build step).
-			MaxDirtyLevel = FMath::Clamp(StreamingTransitionLevel, 0, MaxDirtyLevel);
+			MaxDirtyLevel = TransitionLevel - 1;
 		}
 
 		// The Initialize() call will allocate the virtual texture by spawning work on the render thread.
-		VirtualTexture->Initialize(Producer, Desc, Transform, Bounds);
+		VirtualTexture->Initialize(Producer, ProducerDesc, Transform, Bounds);
 	}
 }
 
@@ -122,20 +128,24 @@ void FRuntimeVirtualTextureSceneProxy::FlushDirtyPages()
 	// If Producer handle is not initialized yet it's safe to do nothing because we won't have rendered anything to the VT that needs flushing.
 	if (ProducerHandle.PackedValue != 0)
 	{
-		//todo[vt]: 
-		// Profile to work out best heuristic for when we should use the CombinedDirtyRect
-		// Also consider using some other structure to represent dirty area such as a course 2D bitfield
-		const bool bCombinedFlush = (DirtyRects.Num() > 2 || CombinedDirtyRect == FIntRect(0, 0, VirtualTextureSize.X, VirtualTextureSize.Y));
-
-		if (bCombinedFlush)
+		// Don't do any work if we won't mark anything dirty.
+		if (MaxDirtyLevel >= 0 && CombinedDirtyRect.Area() > 0)
 		{
-			FVirtualTextureSystem::Get().FlushCache(ProducerHandle, CombinedDirtyRect, MaxDirtyLevel);
-		}
-		else
-		{
-			for (FIntRect Rect : DirtyRects)
+			//todo[vt]: 
+			// Profile to work out best heuristic for when we should use the CombinedDirtyRect
+			// Also consider using some other structure to represent dirty area such as a course 2D bitfield
+			const bool bCombinedFlush = (DirtyRects.Num() > 2 || CombinedDirtyRect == FIntRect(0, 0, VirtualTextureSize.X, VirtualTextureSize.Y));
+
+			if (bCombinedFlush)
+			{
+				FVirtualTextureSystem::Get().FlushCache(ProducerHandle, CombinedDirtyRect, MaxDirtyLevel);
+			}
+			else
 			{
-				FVirtualTextureSystem::Get().FlushCache(ProducerHandle, Rect, MaxDirtyLevel);
+				for (FIntRect Rect : DirtyRects)
+				{
+					FVirtualTextureSystem::Get().FlushCache(ProducerHandle, Rect, MaxDirtyLevel);
+				}
 			}
 		}
 	}
