Remove Legacy tonemapper

diff --git a/Engine/Shaders/Private/DeviceEncodingOnlyCommon.ush b/Engine/Shaders/Private/DeviceEncodingOnlyCommon.ush
--- a/Engine/Shaders/Private/DeviceEncodingOnlyCommon.ush
+++ b/Engine/Shaders/Private/DeviceEncodingOnlyCommon.ush
@@ -18,65 +18,6 @@ half3 InverseGamma;
 static const float LinearToNitsScale = 100.0;
 static const float LinearToNitsScaleInverse = 1.0 / 100.0;
 
-// Film tonal and color control.
-half4 ColorMatrixR_ColorCurveCd1;
-half4 ColorMatrixG_ColorCurveCd3Cm3;
-half4 ColorMatrixB_ColorCurveCm2;
-half4 ColorCurve_Cm0Cd0_Cd2_Ch0Cm1_Ch3;
-half4 ColorCurve_Ch1_Ch2;
-half4 ColorShadow_Luma;
-half4 ColorShadow_Tint1;
-half4 ColorShadow_Tint2;
-
-half3 FilmPostProcess(half3 LinearColor) 
-{
-	// Color and exposure control.
-	half3 MatrixColor;
-	#if USE_COLOR_MATRIX == 1
-		// Apply color matrix (channel mixer, exposure, saturation).
-		MatrixColor.r = dot(LinearColor, ColorMatrixR_ColorCurveCd1.rgb);
-		MatrixColor.g = dot(LinearColor, ColorMatrixG_ColorCurveCd3Cm3.rgb);
-		MatrixColor.b = dot(LinearColor, ColorMatrixB_ColorCurveCm2.rgb);
-		#if USE_SHADOW_TINT == 1
-			MatrixColor *= ColorShadow_Tint1.rgb + ColorShadow_Tint2.rgb * rcp(dot(LinearColor, ColorShadow_Luma.rgb) + 1.0);
-		#endif
-		// Required to insure saturation doesn't create negative colors!
-		MatrixColor = max(half3(0.0, 0.0, 0.0), MatrixColor);
-	#else
-		// Less expensive route when not using saturation and channel mixer.
-		#if USE_SHADOW_TINT == 1
-			MatrixColor = LinearColor * (ColorShadow_Tint1.rgb + ColorShadow_Tint2.rgb * rcp(dot(LinearColor, ColorShadow_Luma.rgb) + 1.0));
-		#else
-			MatrixColor = LinearColor * ColorMatrixB_ColorCurveCm2.rgb;
-		#endif
-	#endif
-	// Apply color curve (includes tonemapping).
-	#if USE_CONTRAST == 1
-		// Full path.
-		half3 MatrixColorD = max(0, ColorCurve_Cm0Cd0_Cd2_Ch0Cm1_Ch3.xxx - MatrixColor);
-		half3 MatrixColorH = max(MatrixColor, ColorCurve_Cm0Cd0_Cd2_Ch0Cm1_Ch3.zzz);
-		half3 MatrixColorM = clamp(MatrixColor, ColorCurve_Cm0Cd0_Cd2_Ch0Cm1_Ch3.xxx, ColorCurve_Cm0Cd0_Cd2_Ch0Cm1_Ch3.zzz);
-		half3 CurveColor = 
-			// Highlights
-			(MatrixColorH*ColorCurve_Ch1_Ch2.xxx + ColorCurve_Ch1_Ch2.yyy) * rcp(MatrixColorH + ColorCurve_Cm0Cd0_Cd2_Ch0Cm1_Ch3.www) + 
-				// Midtones
-				((MatrixColorM*ColorMatrixB_ColorCurveCm2.aaa + 
-					// Darks
-					((MatrixColorD*ColorMatrixR_ColorCurveCd1.aaa) * rcp(MatrixColorD + ColorCurve_Cm0Cd0_Cd2_Ch0Cm1_Ch3.yyy) + ColorMatrixG_ColorCurveCd3Cm3.aaa)));
-	#else
-		// This is for mobile, it assumes color is not negative.
-		// Fast path when contrast=1, can remove the dark part of the curve.
-		half3 MatrixColorH = max(MatrixColor, ColorCurve_Cm0Cd0_Cd2_Ch0Cm1_Ch3.zzz);
-		half3 MatrixColorM = min(MatrixColor, ColorCurve_Cm0Cd0_Cd2_Ch0Cm1_Ch3.zzz);
-		half3 CurveColor = (MatrixColorH*ColorCurve_Ch1_Ch2.xxx + ColorCurve_Ch1_Ch2.yyy) * rcp(MatrixColorH + ColorCurve_Cm0Cd0_Cd2_Ch0Cm1_Ch3.www) + MatrixColorM;
-	#endif
-
-	// Must hit black by log10(-2.474)
-	CurveColor -= 0.002;
-
-	return CurveColor;
-}
-
 half3 DeviceEncodingOnlyAndGammaCorrect(half3 LinearColor)
 {
 	// Clamp input to be positive
diff --git a/Engine/Shaders/Private/PostProcessCombineLUTs.usf b/Engine/Shaders/Private/PostProcessCombineLUTs.usf
--- a/Engine/Shaders/Private/PostProcessCombineLUTs.usf
+++ b/Engine/Shaders/Private/PostProcessCombineLUTs.usf
@@ -31,8 +31,6 @@ float LUTWeights[5];
 half3 ColorScale;
 half4 OverlayColor;
 
-uint bUseMobileTonemapper;
-
 static const float LUTSize = 32;
 
 // Accurate for 1000K < Temp < 15000K
@@ -275,37 +273,34 @@ float4 CombineLUTsCommon(float2 InUV, uint InLayerIndex)
 	float3 ColorAP1 = mul( sRGB_2_AP1, BalancedColor );
 
 	// Expand bright saturated colors outside the sRGB gamut to fake wide gamut rendering.
-    if (!bUseMobileTonemapper)
+ 	float  LumaAP1 = dot( ColorAP1, AP1_RGB2Y );
+	float3 ChromaAP1 = ColorAP1 / LumaAP1;
+
+	float ChromaDistSqr = dot( ChromaAP1 - 1, ChromaAP1 - 1 );
+	float ExpandAmount = ( 1 - exp2( -4 * ChromaDistSqr ) ) * ( 1 - exp2( -4 * ExpandGamut * LumaAP1*LumaAP1 ) );
+
+	// Bizarre matrix but this expands sRGB to between P3 and AP1
+	// CIE 1931 chromaticities:	x		y
+	//				Red:		0.6965	0.3065
+	//				Green:		0.245	0.718
+	//				Blue:		0.1302	0.0456
+	//				White:		0.3127	0.329
+	const float3x3 Wide_2_XYZ_MAT = 
 	{
-		float  LumaAP1 = dot( ColorAP1, AP1_RGB2Y );
-		float3 ChromaAP1 = ColorAP1 / LumaAP1;
-
-		float ChromaDistSqr = dot( ChromaAP1 - 1, ChromaAP1 - 1 );
-		float ExpandAmount = ( 1 - exp2( -4 * ChromaDistSqr ) ) * ( 1 - exp2( -4 * ExpandGamut * LumaAP1*LumaAP1 ) );
-
-		// Bizarre matrix but this expands sRGB to between P3 and AP1
-		// CIE 1931 chromaticities:	x		y
-		//				Red:		0.6965	0.3065
-		//				Green:		0.245	0.718
-		//				Blue:		0.1302	0.0456
-		//				White:		0.3127	0.329
-		const float3x3 Wide_2_XYZ_MAT = 
-		{
-			0.5441691,  0.2395926,  0.1666943,
-			0.2394656,  0.7021530,  0.0583814,
-			-0.0023439,  0.0361834,  1.0552183,
-		};
-
-		const float3x3 Wide_2_AP1 = mul( XYZ_2_AP1_MAT, Wide_2_XYZ_MAT );
-		const float3x3 ExpandMat = mul( Wide_2_AP1, AP1_2_sRGB );
-
-		float3 ColorExpand = mul( ExpandMat, ColorAP1 );
-		ColorAP1 = lerp( ColorAP1, ColorExpand, ExpandAmount );
-	}
+		0.5441691,  0.2395926,  0.1666943,
+		0.2394656,  0.7021530,  0.0583814,
+		-0.0023439,  0.0361834,  1.0552183,
+	};
+
+	const float3x3 Wide_2_AP1 = mul( XYZ_2_AP1_MAT, Wide_2_XYZ_MAT );
+	const float3x3 ExpandMat = mul( Wide_2_AP1, AP1_2_sRGB );
+
+	float3 ColorExpand = mul( ExpandMat, ColorAP1 );
+	ColorAP1 = lerp( ColorAP1, ColorExpand, ExpandAmount );
 
 	ColorAP1 = ColorCorrectAll( ColorAP1 );
 
-	// Store for Legacy tonemap later and for Linear HDR output without tone curve
+	// Store for Linear HDR output without tone curve
 	float3 GradedColor = mul( AP1_2_sRGB, ColorAP1 );
 
 	const float3x3 BlueCorrect =
@@ -336,15 +331,6 @@ float4 CombineLUTsCommon(float2 InUV, uint InLayerIndex)
 	// Convert from AP1 to sRGB and clip out-of-gamut values
 	float3 FilmColor = max(0, mul( AP1_2_sRGB, ColorAP1 ));
 
-#if FEATURE_LEVEL > FEATURE_LEVEL_ES3_1 // Mobile path uses separate shader for legacy tone mapper
-	BRANCH
-	if (bUseMobileTonemapper)
-	{
-		// Legacy tone mapper
-		FilmColor = FilmPostProcess(GradedColor);
-	}
-#endif // FEATURE_LEVEL > FEATURE_LEVEL_ES3_1
-
 #if BLENDCOUNT > 1
 	{
 		// Legacy LDR LUT color grading
diff --git a/Engine/Shaders/Private/PostProcessTonemap.usf b/Engine/Shaders/Private/PostProcessTonemap.usf
--- a/Engine/Shaders/Private/PostProcessTonemap.usf
+++ b/Engine/Shaders/Private/PostProcessTonemap.usf
@@ -701,248 +701,3 @@ void MainCS(uint2 DispatchThreadId : SV_DispatchThreadID)
 }
 #endif
 
-
-// Mobile version
-// TODO: Convert float to half.
-
-float FringeIntensity;
-
-float4 BufferSizeAndInvSize;
-float4 DofBlurSizeAndInvSize;
-
-// vertex shader entry point
-void MainVS_Mobile(
-	in float4 InPosition : ATTRIBUTE0,
-	in float2 InTexCoord : ATTRIBUTE1,
-	out float4 OutTexCoord : TEXCOORD0,
-	out float4 OutFineDofGrain : TEXCOORD1,
-	out float4 OutFullViewUV : TEXCOORD2,
-	out float3 OutVignetteExposureScale : TEXCOORD3,
-	out float2 OutTexCoords[4] : TEXCOORD4,
-	out float4 OutFringeTexCoords : FRINGE_COORDS,
-	out float4 OutPosition : SV_POSITION
-	)
-{
-#if NEEDTOSWITCHVERTICLEAXIS
-	// This is currently the last pass, so flip the texture on V to match D3D
-	InTexCoord.y = 1.0-InTexCoord.y;
-#endif
-
-#if EYEADAPTATION_EXPOSURE_FIX == 1
-	// texture can be GWhiteTexture which is 1x1. It's important we don't read outside bounds.
-	OutVignetteExposureScale.z = EyeAdaptationBuffer[0].x;
-#else
-	// If eye adaptation is disabled (e.g. fixed exposure level ) or  not available.  
-	OutVignetteExposureScale.z = DefaultEyeExposure;
-#endif
-
-	DrawRectangle(InPosition, InTexCoord, OutPosition, OutTexCoord.xy);
-	OutTexCoord = float4(OutTexCoord.xy, OutPosition.xy);
-
-	// Fine adjustment is inside the possible non-full viewport in the full resolution texture.
-	OutFineDofGrain.xy = OutTexCoord.xy + BufferSizeAndInvSize.zw * float2(-0.5,0.5);
-	// Want grain and a second UV based on the knowledge that the source texture has a full viewport.
-	OutFullViewUV.xy = OutPosition.xy * float2(0.5,-0.5) + 0.5;
-	#if NEEDTOSWITCHVERTICLEAXIS
-		// This is currently the last pass, so flip the texture on V to match D3D
-		OutFullViewUV.y = 1.0 - OutFullViewUV.y;
-	#endif
-	// For DOF attempt to undo sampling bias for the first transition region.
-	// This is better for the fine transition, breaks down for the larger bokeh.
-	// This is the best compromise for mobile using 4 bilinear taps only.
-	OutFullViewUV.zw = OutFullViewUV.xy + DofBlurSizeAndInvSize.zw * float2(0.25,-0.5);
-	OutFineDofGrain.zw = OutFullViewUV.xy + GrainRandomFull.xy;
-	// NEWS
-	OutTexCoords[0] = OutTexCoord.xy + BufferSizeAndInvSize.zw * float2( 0,-1);
-	OutTexCoords[1] = OutTexCoord.xy + BufferSizeAndInvSize.zw * float2( 1, 0);
-	OutTexCoords[2] = OutTexCoord.xy + BufferSizeAndInvSize.zw * float2(-1, 0);
-	OutTexCoords[3] = OutTexCoord.xy + BufferSizeAndInvSize.zw * float2( 0, 1);
-
-	float2 ColorViewportPos = UVToScreenPos(OutTexCoord.xy, BufferSizeAndInvSize.xy);
-
-	// Scale vignette to always be a circle with consistent corner intensity.
-	float2 LensViewportPos = ConvertScreenViewportSpaceToLensViewportSpace(ColorViewportPos);
-	OutVignetteExposureScale.xy = VignetteSpace(LensViewportPos);
-
-	// Compute fringe tex coord offsets
-	BRANCH
-	if (FringeIntensity > 0.0)
-	{
-		// Wavelength of primaries in nm
-		half PrimaryR = 611.3f;
-		half PrimaryG = 549.1f;
-		half PrimaryB = 464.3f;
-
-		// Simple lens chromatic aberration is roughly linear in wavelength
-		float ScaleR = 0.007f * (PrimaryR - PrimaryB);
-		float ScaleG = 0.007f * (PrimaryG - PrimaryB);
-
-		// Fringe
-		float2 ScreenPos = UVToScreenPos(OutTexCoord.xy, BufferSizeAndInvSize.xy);
-		OutFringeTexCoords.xy = ScreenPosToUV(ScreenPos * (1.0/(1.0 + (FringeIntensity * ScaleG))), BufferSizeAndInvSize.zw);
-		OutFringeTexCoords.zw = ScreenPosToUV(ScreenPos * (1.0/(1.0 + (FringeIntensity * ScaleR))), BufferSizeAndInvSize.zw);
-	}
-	else
-	{
-		OutFringeTexCoords = 0.0;
-	}
-}
-
-float4 OverlayColor;
-float SRGBAwareTarget;
-
-Texture2D SceneColorTexture;
-SamplerState SceneColorSampler;
-
-Texture2D DofBlurTexture;
-SamplerState DofBlurSampler;
-
-Texture2D SunShaftAndDofTexture;
-SamplerState SunShaftAndDofSampler;
-
-void MainPS_Mobile(
-	in float4 UVAndScreenPos : TEXCOORD0,
-	in float4 FineDofGrain : TEXCOORD1,
-	in float4 FullViewUV : TEXCOORD2,
-	in float3 InVignetteExposureScale : TEXCOORD3,
-	in float2 InTexCoords[4] : TEXCOORD4,
-	in float4 InFringeTexCoords : FRINGE_COORDS,
-	out half4 OutColor : SV_Target0
-	)
-{
-	float2 UV = UVAndScreenPos.xy;
-	float2 ScreenSpacePos = UVAndScreenPos.zw;
-
-	half4 SceneColor;
-
-#if USE_PREEXPOSURE
-	const float OneOverPreExposure = View.OneOverPreExposure;
-#else
-	const float OneOverPreExposure = 1.f;
-#endif
-	
-	BRANCH
-	if (FringeIntensity > 0.0)
-	{
-		SceneColor.r = Texture2DSample(SceneColorTexture, SceneColorSampler, InFringeTexCoords.zw).r;
-		SceneColor.g = Texture2DSample(SceneColorTexture, SceneColorSampler, InFringeTexCoords.xy).g;
-		SceneColor.b = Texture2DSample(SceneColorTexture, SceneColorSampler, UV).b;
-		SceneColor.a = 1.0;
-	}
-	else
-	{
-		SceneColor = Texture2DSample(SceneColorTexture, SceneColorSampler, UV);
-	}
-
-	#if USE_GAMMA_ONLY
-		
-		#if FEATURE_LEVEL == FEATURE_LEVEL_ES3_1
-			OutColor.rgb = pow(SceneColor.rgb, InverseGamma.x);
-		#else
-			OutColor.rgb = sqrt(SceneColor.rgb);
-		#endif
-
-	#else
-
-		#if USE_GRAIN_JITTER || USE_GRAIN_INTENSITY || USE_GRAIN_QUANTIZATION
-			half Grain = GrainFromUV(FineDofGrain.zw);
-		#endif
-
-		#if USE_DOF
-			half4 DofFine = Texture2DSample(SceneColorTexture, SceneColorSampler, FineDofGrain.xy);
-			half4 Dof = Texture2DSample(DofBlurTexture, DofBlurSampler, FullViewUV.zw);
-			half DofCoc = Texture2DSample(SunShaftAndDofTexture, SunShaftAndDofSampler, UV).r;
-			// Convert alpha back into circle of confusion.
-			SceneColor.a = max(Dof.a, abs(DofCoc * 2.0 - 1.0));
-			// Convert circle of confusion into blend factors.		
-			half2 ScaleBias = CocBlendScaleBias(); // Constant.
-			half DofAmount = saturate(SceneColor.a * ScaleBias.x + ScaleBias.y);
-			half2 ScaleBias2 = CocBlendScaleBiasFine(); // Constant.
-			half DofAmountFine = saturate(SceneColor.a * ScaleBias2.x + ScaleBias2.y);
-			#if USE_GRAIN_JITTER
-				// Grain can increase fine DOF.
-				DofAmountFine = max((1.0-Grain*Grain) * GrainScaleBiasJitter.z, DofAmountFine);
-			#endif
-			// Blend in fine DOF.
-			SceneColor.rgb = lerp(SceneColor.rgb, DofFine.rgb, DofAmountFine);
-			// Blend in coarse DOF.
-			SceneColor.rgb = lerp(SceneColor.rgb, Dof.rgb, DofAmount);
-		#else
-			#if USE_GRAIN_JITTER
-				// Do jitter for grain.
-				half4 DofFine = Texture2DSample(SceneColorTexture, SceneColorSampler, FineDofGrain.xy);
-				// Grain jitter.
-				SceneColor.rgb = lerp(SceneColor.rgb, DofFine.rgb, (1.0-Grain*Grain) * GrainScaleBiasJitter.z);
-			#endif
-		#endif
-
-		#if METAL_MSAA_HDR_DECODE
-			// Do after jitter for grain as an optimization.
-			SceneColor.rgb *= rcp(SceneColor.r*(-0.299) + SceneColor.g*(-0.587) + SceneColor.b*(-0.114) + 1.0);
-			// Try to kill negatives and NaNs here
-			SceneColor.rgb = max(SceneColor.rgb, 0);
-		#endif
-
-		#if !USE_DOF
-			// Set so temporal AA shader knows everything is in focus.
-			SceneColor.a = 0.0;
-		#endif
-
-		// Match PC naming.
-		half3 LinearColor = SceneColor.rgb * OneOverPreExposure;
-
-		float ExposureScale = InVignetteExposureScale.z;
-		
-		#if USE_VIGNETTE
-			LinearColor.rgb *= ComputeVignetteMask(InVignetteExposureScale.xy, TonemapperParams.x );
-		#endif
-
-		// It is faster to do vignette as a texture lookup + mad because this is an ALU bound shader.
-		#if (USE_BLOOM || USE_LIGHT_SHAFTS || USE_VIGNETTE)
-		{
-			half4 CombinedBloomSunVignette = Texture2DSample(BloomTexture, BloomSampler, FullViewUV.xy);
-			LinearColor.rgb = LinearColor.rgb * CombinedBloomSunVignette.a + CombinedBloomSunVignette.rgb * OneOverPreExposure;
-		}
-		#endif
-
-		LinearColor *= ExposureScale;
-
-		#if USE_GRAIN_INTENSITY
-			// Needs to go before tonemapping.
-			half GrainMul = Grain * GrainScaleBiasJitter.x + GrainScaleBiasJitter.y;
-			LinearColor.rgb *= GrainMul;
-		#endif
-
-		half3 FilmColor = FilmPostProcess(LinearColor.rgb);
-		half3 TonemappedColor;
-		if (SRGBAwareTarget > 0.5)
-		{
-			// skip sRGB conversion if HW conversion will take place.
-			TonemappedColor = FilmColor;
-		}
-		else
-		{
-			// Apply conversion to sRGB (this must be an exact sRGB conversion else darks are bad).
-			TonemappedColor = LinearToPostTonemapSpace(FilmColor);
-		}
-
-
-		// Blend with custom LDR color, used for Fade track in Matinee.
-		// This is the 101% wrong way to do this,
-		//  - It adds an extra redundant lerp.
-		//  - It is not going to work with the future-forward ES3 fast path of sRGB output.
-		//  - And it does super ugly non-linear blending.
-		// The right way is to adjust exposure instead.
-		TonemappedColor = lerp(TonemappedColor.rgb, OverlayColor.rgb, OverlayColor.a);
-
-		#if USE_GRAIN_QUANTIZATION
-			// Needs to go after tonemapping.
-			half GrainQuantization = 1.0/256.0;
-			half GrainAdd = (Grain * GrainQuantization) + (-0.5 * GrainQuantization);
-			TonemappedColor.rgb += GrainAdd;
-		#endif
-			
-			OutColor = half4(TonemappedColor, SceneColor.a);
-
-	#endif
-}
diff --git a/Engine/Shaders/Private/PostProcessVisualizeHDR.usf b/Engine/Shaders/Private/PostProcessVisualizeHDR.usf
--- a/Engine/Shaders/Private/PostProcessVisualizeHDR.usf
+++ b/Engine/Shaders/Private/PostProcessVisualizeHDR.usf
@@ -384,7 +384,7 @@ void MainPS(noperspective float4 UVAndScreenPos : TEXCOORD0, float4 SvPosition :
 		// HDR > 0
 		float3 AdpatedLuminance = EyeAdaptationResult * float3(LuminanceVal, LuminanceVal, LuminanceVal);
 		// 0..1
-		float3 TonemappedLuminance = FilmPostProcess(AdpatedLuminance);
+		float3 TonemappedLuminance = FilmToneMap(AdpatedLuminance);
 		float3 DistMask = saturate(1.0 - 100.0 * abs(TonemappedLuminance - (1.0 - InsetUV.y)));
 		OutColor = lerp(OutColor, float4(1, 1, 1, 0), float4(DistMask, 0.0));
 	}
diff --git a/Engine/Shaders/Private/TonemapCommon.ush b/Engine/Shaders/Private/TonemapCommon.ush
--- a/Engine/Shaders/Private/TonemapCommon.ush
+++ b/Engine/Shaders/Private/TonemapCommon.ush
@@ -18,65 +18,6 @@ half3 InverseGamma;
 static const float LinearToNitsScale = 100.0;
 static const float LinearToNitsScaleInverse = 1.0 / 100.0;
 
-// Film tonal and color control.
-half4 ColorMatrixR_ColorCurveCd1;
-half4 ColorMatrixG_ColorCurveCd3Cm3;
-half4 ColorMatrixB_ColorCurveCm2;
-half4 ColorCurve_Cm0Cd0_Cd2_Ch0Cm1_Ch3;
-half4 ColorCurve_Ch1_Ch2;
-half4 ColorShadow_Luma;
-half4 ColorShadow_Tint1;
-half4 ColorShadow_Tint2;
-
-half3 FilmPostProcess(half3 LinearColor) 
-{
-	// Color and exposure control.
-	half3 MatrixColor;
-	#if USE_COLOR_MATRIX == 1
-		// Apply color matrix (channel mixer, exposure, saturation).
-		MatrixColor.r = dot(LinearColor, ColorMatrixR_ColorCurveCd1.rgb);
-		MatrixColor.g = dot(LinearColor, ColorMatrixG_ColorCurveCd3Cm3.rgb);
-		MatrixColor.b = dot(LinearColor, ColorMatrixB_ColorCurveCm2.rgb);
-		#if USE_SHADOW_TINT == 1
-			MatrixColor *= ColorShadow_Tint1.rgb + ColorShadow_Tint2.rgb * rcp(dot(LinearColor, ColorShadow_Luma.rgb) + 1.0);
-		#endif
-		// Required to insure saturation doesn't create negative colors!
-		MatrixColor = max(half3(0.0, 0.0, 0.0), MatrixColor);
-	#else
-		// Less expensive route when not using saturation and channel mixer.
-		#if USE_SHADOW_TINT == 1
-			MatrixColor = LinearColor * (ColorShadow_Tint1.rgb + ColorShadow_Tint2.rgb * rcp(dot(LinearColor, ColorShadow_Luma.rgb) + 1.0));
-		#else
-			MatrixColor = LinearColor * ColorMatrixB_ColorCurveCm2.rgb;
-		#endif
-	#endif
-	// Apply color curve (includes tonemapping).
-	#if USE_CONTRAST == 1
-		// Full path.
-		half3 MatrixColorD = max(0, ColorCurve_Cm0Cd0_Cd2_Ch0Cm1_Ch3.xxx - MatrixColor);
-		half3 MatrixColorH = max(MatrixColor, ColorCurve_Cm0Cd0_Cd2_Ch0Cm1_Ch3.zzz);
-		half3 MatrixColorM = clamp(MatrixColor, ColorCurve_Cm0Cd0_Cd2_Ch0Cm1_Ch3.xxx, ColorCurve_Cm0Cd0_Cd2_Ch0Cm1_Ch3.zzz);
-		half3 CurveColor = 
-			// Highlights
-			(MatrixColorH*ColorCurve_Ch1_Ch2.xxx + ColorCurve_Ch1_Ch2.yyy) * rcp(MatrixColorH + ColorCurve_Cm0Cd0_Cd2_Ch0Cm1_Ch3.www) + 
-				// Midtones
-				((MatrixColorM*ColorMatrixB_ColorCurveCm2.aaa + 
-					// Darks
-					((MatrixColorD*ColorMatrixR_ColorCurveCd1.aaa) * rcp(MatrixColorD + ColorCurve_Cm0Cd0_Cd2_Ch0Cm1_Ch3.yyy) + ColorMatrixG_ColorCurveCd3Cm3.aaa)));
-	#else
-		// This is for mobile, it assumes color is not negative.
-		// Fast path when contrast=1, can remove the dark part of the curve.
-		half3 MatrixColorH = max(MatrixColor, ColorCurve_Cm0Cd0_Cd2_Ch0Cm1_Ch3.zzz);
-		half3 MatrixColorM = min(MatrixColor, ColorCurve_Cm0Cd0_Cd2_Ch0Cm1_Ch3.zzz);
-		half3 CurveColor = (MatrixColorH*ColorCurve_Ch1_Ch2.xxx + ColorCurve_Ch1_Ch2.yyy) * rcp(MatrixColorH + ColorCurve_Cm0Cd0_Cd2_Ch0Cm1_Ch3.www) + MatrixColorM;
-	#endif
-
-	// Must hit black by log10(-2.474)
-	CurveColor -= 0.002;
-
-	return CurveColor;
-}
-
 half3 TonemapAndGammaCorrect(half3 LinearColor)
 {
 	// Clamp input to be positive
diff --git a/Engine/Source/Editor/DetailCustomizations/Private/PostProcessSettingsCustomization.cpp b/Engine/Source/Editor/DetailCustomizations/Private/PostProcessSettingsCustomization.cpp
--- a/Engine/Source/Editor/DetailCustomizations/Private/PostProcessSettingsCustomization.cpp
+++ b/Engine/Source/Editor/DetailCustomizations/Private/PostProcessSettingsCustomization.cpp
@@ -117,16 +117,6 @@ void FPostProcessSettingsCustomization::CustomizeChildren( TSharedRef<IPropertyH
 	TMap<FString, FCategoryOrGroup> NameToCategoryBuilderMap;
 	TMap<FString, FPostProcessGroup> NameToGroupMap;
 
-	static const auto VarTonemapperFilm = IConsoleManager::Get().FindTConsoleVariableDataInt(TEXT("r.TonemapperFilm"));
-	static const auto VarMobileTonemapperFilm = IConsoleManager::Get().FindTConsoleVariableDataInt(TEXT("r.Mobile.TonemapperFilm"));
-	static const FName LegacyTonemapperName("LegacyTonemapper");
-	static const FName TonemapperCategory("Film");
-	static const FName MobileTonemapperCategory("Mobile Tonemapper");
-	const bool bDesktopTonemapperFilm = VarTonemapperFilm->GetValueOnGameThread() == 1;
-	const bool bMobileTonemapperFilm = VarMobileTonemapperFilm->GetValueOnGameThread() == 1;
-	const bool bUsingFilmTonemapper = bDesktopTonemapperFilm || bMobileTonemapperFilm;		// Are any platforms use film tonemapper
-	const bool bUsingLegacyTonemapper = !bDesktopTonemapperFilm || !bMobileTonemapperFilm;	// Are any platforms use legacy tonemapper
-
 	static const auto VarDefaultAutoExposureExtendDefaultLuminanceRange = IConsoleManager::Get().FindTConsoleVariableDataInt(TEXT("r.DefaultFeature.AutoExposure.ExtendDefaultLuminanceRange"));
 	const bool bExtendedLuminanceRange = VarDefaultAutoExposureExtendDefaultLuminanceRange->GetValueOnGameThread() == 1;
 	static const FName ExposureCategory("Lens|Exposure");
@@ -147,28 +137,7 @@ void FPostProcessSettingsCustomization::CustomizeChildren( TSharedRef<IPropertyH
 
 				FName CategoryFName = FObjectEditorUtils::GetCategoryFName(Property);
 					
-				if (CategoryFName == TonemapperCategory)
-				{
-					bool bIsLegacyTonemapperPropery = ChildHandle->HasMetaData(LegacyTonemapperName);
-
-					// Hide in case no platforms use legacy tonemapper
-					// Hide in case no platforms use film tonemapper
-					if ((bIsLegacyTonemapperPropery && !bUsingLegacyTonemapper) || (!bIsLegacyTonemapperPropery && !bUsingFilmTonemapper))
-					{
-						ChildHandle->MarkHiddenByCustomization();
-						continue;
-					}
-
-					// In case platforms use different tonemappers, place mobile settings into separate category
-					if (bMobileTonemapperFilm != bDesktopTonemapperFilm)
-					{
-						if (bMobileTonemapperFilm == !bIsLegacyTonemapperPropery)
-						{
-							CategoryFName = MobileTonemapperCategory;
-						}
-					}
-				}
-				else if (CategoryFName == ExposureCategory && bExtendedLuminanceRange)
+				if (CategoryFName == ExposureCategory && bExtendedLuminanceRange)
 				{
 					if (Property->GetName() == TEXT("AutoExposureMinBrightness"))
 					{
diff --git a/Engine/Source/Runtime/Renderer/Private/PostProcess/PostProcessCombineLUTs.cpp b/Engine/Source/Runtime/Renderer/Private/PostProcess/PostProcessCombineLUTs.cpp
--- a/Engine/Source/Runtime/Renderer/Private/PostProcess/PostProcessCombineLUTs.cpp
+++ b/Engine/Source/Runtime/Renderer/Private/PostProcess/PostProcessCombineLUTs.cpp
@@ -35,12 +35,6 @@ FAutoConsoleVariableRef CVarLUTSize(
 	TEXT("Size of film LUT"),
 	ECVF_RenderThreadSafe);
 
-TAutoConsoleVariable<int32> CVarTonemapperFilm(
-	TEXT("r.TonemapperFilm"),
-	1,
-	TEXT("Use new film tone mapper"),
-	ECVF_RenderThreadSafe);
-
 // Including the neutral one at index 0
 const uint32 GMaxLUTBlendCount = 5;
 
@@ -129,10 +123,8 @@ BEGIN_SHADER_PARAMETER_STRUCT(FCombineLUTParameters, )
 	SHADER_PARAMETER(float, FilmShoulder)
 	SHADER_PARAMETER(float, FilmBlackClip)
 	SHADER_PARAMETER(float, FilmWhiteClip)
-	SHADER_PARAMETER(uint32, bUseMobileTonemapper)
 	SHADER_PARAMETER_STRUCT_INCLUDE(FColorRemapParameters, ColorRemap)
 	SHADER_PARAMETER_STRUCT_INCLUDE(FTonemapperOutputDeviceParameters, OutputDevice)
-	SHADER_PARAMETER_STRUCT_INCLUDE(FMobileFilmTonemapParameters, MobileFilmTonemap)
 END_SHADER_PARAMETER_STRUCT()
 
 void GetCombineLUTParameters(
@@ -208,24 +200,8 @@ void GetCombineLUTParameters(
 	Parameters.FilmShoulder = Settings.FilmShoulder;
 	Parameters.FilmBlackClip = Settings.FilmBlackClip;
 	Parameters.FilmWhiteClip = Settings.FilmWhiteClip;
-	Parameters.bUseMobileTonemapper = CVarTonemapperFilm.GetValueOnRenderThread() == 0;
-
-
-	Parameters.MobileFilmTonemap = GetMobileFilmTonemapParameters(
-		Settings,
-		/* UseColorMatrix = */ true,
-		/* UseShadowTint = */ true,
-		/* UseContrast = */ true);
-
-	Parameters.bUseMobileTonemapper = CVarTonemapperFilm.GetValueOnRenderThread() == 0;
 
 	Parameters.OutputDevice = GetTonemapperOutputDeviceParameters(ViewFamily);
-
-	Parameters.MobileFilmTonemap = GetMobileFilmTonemapParameters(
-		Settings,
-		/* UseColorMatrix = */ true,
-		/* UseShadowTint = */ true,
-		/* UseContrast = */ true);
 }
 
 class FLUTBlenderShader : public FGlobalShader
diff --git a/Engine/Source/Runtime/Renderer/Private/PostProcess/PostProcessTonemap.cpp b/Engine/Source/Runtime/Renderer/Private/PostProcess/PostProcessTonemap.cpp
--- a/Engine/Source/Runtime/Renderer/Private/PostProcess/PostProcessTonemap.cpp
+++ b/Engine/Source/Runtime/Renderer/Private/PostProcess/PostProcessTonemap.cpp
@@ -87,12 +87,6 @@ TAutoConsoleVariable<float> CVarGamma(
 	TEXT("Gamma on output"),
 	ECVF_RenderThreadSafe);
 
-TAutoConsoleVariable<int32> CVarMobileTonemapperFilm(
-	TEXT("r.Mobile.TonemapperFilm"),
-	1,
-	TEXT("Whether mobile platforms should use new film tone mapper"),
-	ECVF_RenderThreadSafe);
-
 const int32 GTonemapComputeTileSizeX = 8;
 const int32 GTonemapComputeTileSizeY = 8;
 
@@ -263,174 +257,6 @@ void GrainPostSettings(FVector* RESTRICT const Constant, const FPostProcessSetti
 	Constant->Z = GrainJitter;
 }
 
-FMobileFilmTonemapParameters GetMobileFilmTonemapParameters(const FPostProcessSettings& PostProcessSettings, bool bUseColorMatrix, bool bUseShadowTint, bool bUseContrast)
-{
-	// Must insure inputs are in correct range (else possible generation of NaNs).
-	float InExposure = 1.0f;
-	FVector InWhitePoint(PostProcessSettings.FilmWhitePoint);
-	float InSaturation = FMath::Clamp(PostProcessSettings.FilmSaturation, 0.0f, 2.0f);
-	FVector InLuma = FVector(1.0f / 3.0f, 1.0f / 3.0f, 1.0f / 3.0f);
-	FVector InMatrixR(PostProcessSettings.FilmChannelMixerRed);
-	FVector InMatrixG(PostProcessSettings.FilmChannelMixerGreen);
-	FVector InMatrixB(PostProcessSettings.FilmChannelMixerBlue);
-	float InContrast = FMath::Clamp(PostProcessSettings.FilmContrast, 0.0f, 1.0f) + 1.0f;
-	float InDynamicRange = powf(2.0f, FMath::Clamp(PostProcessSettings.FilmDynamicRange, 1.0f, 4.0f));
-	float InToe = (1.0f - FMath::Clamp(PostProcessSettings.FilmToeAmount, 0.0f, 1.0f)) * 0.18f;
-	InToe = FMath::Clamp(InToe, 0.18f / 8.0f, 0.18f * (15.0f / 16.0f));
-	float InHeal = 1.0f - (FMath::Max(1.0f / 32.0f, 1.0f - FMath::Clamp(PostProcessSettings.FilmHealAmount, 0.0f, 1.0f)) * (1.0f - 0.18f));
-	FVector InShadowTint(PostProcessSettings.FilmShadowTint);
-	float InShadowTintBlend = FMath::Clamp(PostProcessSettings.FilmShadowTintBlend, 0.0f, 1.0f) * 64.0f;
-
-	// Shadow tint amount enables turning off shadow tinting.
-	float InShadowTintAmount = FMath::Clamp(PostProcessSettings.FilmShadowTintAmount, 0.0f, 1.0f);
-	InShadowTint = InWhitePoint + (InShadowTint - InWhitePoint) * InShadowTintAmount;
-
-	// Make sure channel mixer inputs sum to 1 (+ smart dealing with all zeros).
-	InMatrixR.X += 1.0f / (256.0f*256.0f*32.0f);
-	InMatrixG.Y += 1.0f / (256.0f*256.0f*32.0f);
-	InMatrixB.Z += 1.0f / (256.0f*256.0f*32.0f);
-	InMatrixR *= 1.0f / FVector::DotProduct(InMatrixR, FVector(1.0f));
-	InMatrixG *= 1.0f / FVector::DotProduct(InMatrixG, FVector(1.0f));
-	InMatrixB *= 1.0f / FVector::DotProduct(InMatrixB, FVector(1.0f));
-
-	// Conversion from linear rgb to luma (using HDTV coef).
-	FVector LumaWeights = FVector(0.2126f, 0.7152f, 0.0722f);
-
-	// Make sure white point has 1.0 as luma (so adjusting white point doesn't change exposure).
-	// Make sure {0.0,0.0,0.0} inputs do something sane (default to white).
-	InWhitePoint += FVector(1.0f / (256.0f*256.0f*32.0f));
-	InWhitePoint *= 1.0f / FVector::DotProduct(InWhitePoint, LumaWeights);
-	InShadowTint += FVector(1.0f / (256.0f*256.0f*32.0f));
-	InShadowTint *= 1.0f / FVector::DotProduct(InShadowTint, LumaWeights);
-
-	// Grey after color matrix is applied.
-	FVector ColorMatrixLuma = FVector(
-		FVector::DotProduct(InLuma.X * FVector(InMatrixR.X, InMatrixG.X, InMatrixB.X), FVector(1.0f)),
-		FVector::DotProduct(InLuma.Y * FVector(InMatrixR.Y, InMatrixG.Y, InMatrixB.Y), FVector(1.0f)),
-		FVector::DotProduct(InLuma.Z * FVector(InMatrixR.Z, InMatrixG.Z, InMatrixB.Z), FVector(1.0f)));
-
-	FVector OutMatrixR = FVector(0.0f);
-	FVector OutMatrixG = FVector(0.0f);
-	FVector OutMatrixB = FVector(0.0f);
-	FVector OutColorShadow_Luma = LumaWeights * InShadowTintBlend;
-	FVector OutColorShadow_Tint1 = InWhitePoint;
-	FVector OutColorShadow_Tint2 = InShadowTint - InWhitePoint;
-
-	if (bUseColorMatrix)
-	{
-		// Final color matrix effected by saturation and exposure.
-		OutMatrixR = (ColorMatrixLuma + ((InMatrixR - ColorMatrixLuma) * InSaturation)) * InExposure;
-		OutMatrixG = (ColorMatrixLuma + ((InMatrixG - ColorMatrixLuma) * InSaturation)) * InExposure;
-		OutMatrixB = (ColorMatrixLuma + ((InMatrixB - ColorMatrixLuma) * InSaturation)) * InExposure;
-		if (!bUseShadowTint)
-		{
-			OutMatrixR = OutMatrixR * InWhitePoint.X;
-			OutMatrixG = OutMatrixG * InWhitePoint.Y;
-			OutMatrixB = OutMatrixB * InWhitePoint.Z;
-		}
-	}
-	else
-	{
-		// No color matrix fast path.
-		if (!bUseShadowTint)
-		{
-			OutMatrixB = InExposure * InWhitePoint;
-		}
-		else
-		{
-			// Need to drop exposure in.
-			OutColorShadow_Luma *= InExposure;
-			OutColorShadow_Tint1 *= InExposure;
-			OutColorShadow_Tint2 *= InExposure;
-		}
-	}
-
-	// Curve constants.
-	float OutColorCurveCh3;
-	float OutColorCurveCh0Cm1;
-	float OutColorCurveCd2;
-	float OutColorCurveCm0Cd0;
-	float OutColorCurveCh1;
-	float OutColorCurveCh2;
-	float OutColorCurveCd1;
-	float OutColorCurveCd3Cm3;
-	float OutColorCurveCm2;
-
-	// Line for linear section.
-	float FilmLineOffset = 0.18f - 0.18f*InContrast;
-	float FilmXAtY0 = -FilmLineOffset / InContrast;
-	float FilmXAtY1 = (1.0f - FilmLineOffset) / InContrast;
-	float FilmXS = FilmXAtY1 - FilmXAtY0;
-
-	// Coordinates of linear section.
-	float FilmHiX = FilmXAtY0 + InHeal * FilmXS;
-	float FilmHiY = FilmHiX * InContrast + FilmLineOffset;
-	float FilmLoX = FilmXAtY0 + InToe * FilmXS;
-	float FilmLoY = FilmLoX * InContrast + FilmLineOffset;
-	// Supported exposure range before clipping.
-	float FilmHeal = InDynamicRange - FilmHiX;
-	// Intermediates.
-	float FilmMidXS = FilmHiX - FilmLoX;
-	float FilmMidYS = FilmHiY - FilmLoY;
-	float FilmSlope = FilmMidYS / (FilmMidXS);
-	float FilmHiYS = 1.0f - FilmHiY;
-	float FilmLoYS = FilmLoY;
-	float FilmToe = FilmLoX;
-	float FilmHiG = (-FilmHiYS + (FilmSlope*FilmHeal)) / (FilmSlope*FilmHeal);
-	float FilmLoG = (-FilmLoYS + (FilmSlope*FilmToe)) / (FilmSlope*FilmToe);
-
-	if (bUseContrast)
-	{
-		// Constants.
-		OutColorCurveCh1 = FilmHiYS / FilmHiG;
-		OutColorCurveCh2 = -FilmHiX * (FilmHiYS / FilmHiG);
-		OutColorCurveCh3 = FilmHiYS / (FilmSlope*FilmHiG) - FilmHiX;
-		OutColorCurveCh0Cm1 = FilmHiX;
-		OutColorCurveCm2 = FilmSlope;
-		OutColorCurveCm0Cd0 = FilmLoX;
-		OutColorCurveCd3Cm3 = FilmLoY - FilmLoX * FilmSlope;
-		// Handle these separate in case of FilmLoG being 0.
-		if (FilmLoG != 0.0f)
-		{
-			OutColorCurveCd1 = -FilmLoYS / FilmLoG;
-			OutColorCurveCd2 = FilmLoYS / (FilmSlope*FilmLoG);
-		}
-		else
-		{
-			// FilmLoG being zero means dark region is a linear segment (so just continue the middle section).
-			OutColorCurveCd1 = 0.0f;
-			OutColorCurveCd2 = 1.0f;
-			OutColorCurveCm0Cd0 = 0.0f;
-			OutColorCurveCd3Cm3 = 0.0f;
-		}
-	}
-	else
-	{
-		// Simplified for no dark segment.
-		OutColorCurveCh1 = FilmHiYS / FilmHiG;
-		OutColorCurveCh2 = -FilmHiX * (FilmHiYS / FilmHiG);
-		OutColorCurveCh3 = FilmHiYS / (FilmSlope*FilmHiG) - FilmHiX;
-		OutColorCurveCh0Cm1 = FilmHiX;
-		// Not used.
-		OutColorCurveCm2 = 0.0f;
-		OutColorCurveCm0Cd0 = 0.0f;
-		OutColorCurveCd3Cm3 = 0.0f;
-		OutColorCurveCd1 = 0.0f;
-		OutColorCurveCd2 = 0.0f;
-	}
-
-	FMobileFilmTonemapParameters Parameters;
-	Parameters.ColorMatrixR_ColorCurveCd1 = FVector4(OutMatrixR, OutColorCurveCd1);
-	Parameters.ColorMatrixG_ColorCurveCd3Cm3 = FVector4(OutMatrixG, OutColorCurveCd3Cm3);
-	Parameters.ColorMatrixB_ColorCurveCm2 = FVector4(OutMatrixB, OutColorCurveCm2);
-	Parameters.ColorCurve_Cm0Cd0_Cd2_Ch0Cm1_Ch3 = FVector4(OutColorCurveCm0Cd0, OutColorCurveCd2, OutColorCurveCh0Cm1, OutColorCurveCh3);
-	Parameters.ColorCurve_Ch1_Ch2 = FVector4(OutColorCurveCh1, OutColorCurveCh2, 0.0f, 0.0f);
-	Parameters.ColorShadow_Luma = FVector4(OutColorShadow_Luma, 0.0f);
-	Parameters.ColorShadow_Tint1 = FVector4(OutColorShadow_Tint1, 0.0f);
-	Parameters.ColorShadow_Tint2 = FVector4(OutColorShadow_Tint2, 0.0f);
-	return Parameters;
-}
-
 FTonemapperOutputDeviceParameters GetTonemapperOutputDeviceParameters(const FSceneViewFamily& Family)
 {
 	static TConsoleVariableData<int32>* CVarOutputGamut = IConsoleManager::Get().FindTConsoleVariableDataInt(TEXT("r.HDR.Display.ColorGamut"));
@@ -867,425 +693,3 @@ FScreenPassTexture AddTonemapPass(FRDGBuilder& GraphBuilder, const FViewInfo& Vi
 	return MoveTemp(Output);
 }
 
-// ES2 version
-
-class FMobileTonemapPS : public FGlobalShader
-{
-	DECLARE_GLOBAL_SHADER(FMobileTonemapPS);
-	SHADER_USE_PARAMETER_STRUCT(FMobileTonemapPS, FGlobalShader);
-
-	BEGIN_SHADER_PARAMETER_STRUCT(FParameters, )
-		SHADER_PARAMETER(FVector4, ColorScale0)
-		SHADER_PARAMETER(FVector4, ColorScale1)
-		SHADER_PARAMETER(FVector4, TonemapperParams)
-		SHADER_PARAMETER(FVector4, OverlayColor)
-		SHADER_PARAMETER_EX(FVector4, ColorMatrixR_ColorCurveCd1, EShaderPrecisionModifier::Half)
-		SHADER_PARAMETER_EX(FVector4, ColorMatrixG_ColorCurveCd3Cm3, EShaderPrecisionModifier::Half)
-		SHADER_PARAMETER_EX(FVector4, ColorMatrixB_ColorCurveCm2, EShaderPrecisionModifier::Half)
-		SHADER_PARAMETER_EX(FVector4, ColorCurve_Cm0Cd0_Cd2_Ch0Cm1_Ch3, EShaderPrecisionModifier::Half)
-		SHADER_PARAMETER_EX(FVector4, ColorCurve_Ch1_Ch2, EShaderPrecisionModifier::Half)
-		SHADER_PARAMETER_EX(FVector4, ColorShadow_Luma, EShaderPrecisionModifier::Half)
-		SHADER_PARAMETER_EX(FVector4, ColorShadow_Tint1, EShaderPrecisionModifier::Half)
-		SHADER_PARAMETER_EX(FVector4, ColorShadow_Tint2, EShaderPrecisionModifier::Half)
-		SHADER_PARAMETER(FVector, GrainScaleBiasJitter)
-		SHADER_PARAMETER_EX(FVector, InverseGamma, EShaderPrecisionModifier::Half)
-		SHADER_PARAMETER(float, FringeIntensity)
-		SHADER_PARAMETER(float, SRGBAwareTarget)
-		SHADER_PARAMETER_STRUCT_REF(FViewUniformShaderParameters, View)
-		SHADER_PARAMETER_RDG_TEXTURE(Texture2D, SceneColorTexture)
-		SHADER_PARAMETER_SAMPLER(SamplerState, SceneColorSampler)
-		SHADER_PARAMETER_RDG_TEXTURE(Texture2D, BloomTexture)
-		SHADER_PARAMETER_SAMPLER(SamplerState, BloomSampler)
-		SHADER_PARAMETER_RDG_TEXTURE(Texture2D, DofBlurTexture)
-		SHADER_PARAMETER_SAMPLER(SamplerState, DofBlurSampler)
-		SHADER_PARAMETER_RDG_TEXTURE(Texture2D, SunShaftAndDofTexture)
-		SHADER_PARAMETER_SAMPLER(SamplerState, SunShaftAndDofSampler)
-		RENDER_TARGET_BINDING_SLOTS()
-	END_SHADER_PARAMETER_STRUCT()
-
-	// Mobile renderer specific permutation dimensions.
-	class FTonemapperDOFDim         : SHADER_PERMUTATION_BOOL("USE_DOF");
-	class FTonemapperLightShaftsDim : SHADER_PERMUTATION_BOOL("USE_LIGHT_SHAFTS");
-	class FTonemapperColorMatrixDim : SHADER_PERMUTATION_BOOL("USE_COLOR_MATRIX");
-	class FTonemapperShadowTintDim  : SHADER_PERMUTATION_BOOL("USE_SHADOW_TINT");
-	class FTonemapperContrastDim    : SHADER_PERMUTATION_BOOL("USE_CONTRAST");
-	
-	using FPermutationDomain = TShaderPermutationDomain<
-		TonemapperPermutation::FCommonDomain,
-		FTonemapperDOFDim,
-		FTonemapperLightShaftsDim,
-		FTonemapperColorMatrixDim,
-		FTonemapperShadowTintDim,
-		FTonemapperContrastDim>;
-	
-	template<class TPermDim, class TPermVector>
-	static void EnableIfSet(const TPermVector& SourceDomain, TPermVector& DestDomain)
-	{
-		if (SourceDomain.template Get<TPermDim>())
-		{
-			DestDomain.template Set<TPermDim>(true);
-		}
-	}
-
-	// Reduce the number of permutations by combining common states
-	static FPermutationDomain RemapPermutationVector(FPermutationDomain WantedPermutationVector)
-	{
-		TonemapperPermutation::FCommonDomain WantedCommonPermutationVector = WantedPermutationVector.Get<TonemapperPermutation::FCommonDomain>();
-		FPermutationDomain RemappedPermutationVector;
-		TonemapperPermutation::FCommonDomain RemappedCommonPermutationVector;
-
-		// Note: FTonemapperSharpenDim, FTonemapperGrainJitterDim are not supported.
-
-		// Gamma only
-		if (WantedCommonPermutationVector.Get<TonemapperPermutation::FTonemapperGammaOnlyDim>())
-		{
-			RemappedCommonPermutationVector.Set<TonemapperPermutation::FTonemapperGammaOnlyDim>(true);
-
-			// Mutually exclusive - clear the wanted vector. 
-			WantedPermutationVector = FPermutationDomain();
-			WantedCommonPermutationVector = WantedPermutationVector.Get<TonemapperPermutation::FCommonDomain>();
-		}
-		else
-		{
-			// Always enable contrast.
-			RemappedPermutationVector.Set<FTonemapperContrastDim>(true);
-		}
-
- 		// Bloom permutation
-		EnableIfSet<TonemapperPermutation::FTonemapperBloomDim>(WantedCommonPermutationVector, RemappedCommonPermutationVector);
- 		// Vignette permutation
-		EnableIfSet<TonemapperPermutation::FTonemapperVignetteDim>(WantedCommonPermutationVector, RemappedCommonPermutationVector);
- 		// Grain intensity permutation
-		EnableIfSet<TonemapperPermutation::FTonemapperGrainIntensityDim>(WantedCommonPermutationVector, RemappedCommonPermutationVector);
-		// Switch Y axis
-		EnableIfSet<TonemapperPermutation::FTonemapperSwitchAxis>(WantedCommonPermutationVector, RemappedCommonPermutationVector);
- 		// msaa permutation.
-		EnableIfSet<TonemapperPermutation::FTonemapperMsaaDim>(WantedCommonPermutationVector, RemappedCommonPermutationVector);
-
- 		// Color matrix
-		EnableIfSet<FTonemapperColorMatrixDim>(WantedPermutationVector, RemappedPermutationVector);
-
-		// DoF
-		if (WantedPermutationVector.Get<FTonemapperDOFDim>())
-		{
-			RemappedPermutationVector.Set<FTonemapperDOFDim>(true);
-			RemappedPermutationVector.Set<FTonemapperLightShaftsDim>(true);
-			RemappedCommonPermutationVector.Set<TonemapperPermutation::FTonemapperVignetteDim>(true);
-			RemappedCommonPermutationVector.Set<TonemapperPermutation::FTonemapperBloomDim>(true);
-		}
-
-		// light shafts
-		if (WantedPermutationVector.Get<FTonemapperLightShaftsDim>())
-		{
-			RemappedPermutationVector.Set<FTonemapperLightShaftsDim>(true);
-			RemappedCommonPermutationVector.Set<TonemapperPermutation::FTonemapperVignetteDim>(true);
-			RemappedCommonPermutationVector.Set<TonemapperPermutation::FTonemapperBloomDim>(true);
-		}
-
-		// Shadow tint
-		if (WantedPermutationVector.Get<FTonemapperShadowTintDim>())
-		{
-			RemappedPermutationVector.Set<FTonemapperShadowTintDim>(true);
-			RemappedPermutationVector.Set<FTonemapperColorMatrixDim>(true);
-		}
-
-		RemappedPermutationVector.Set<TonemapperPermutation::FCommonDomain>(RemappedCommonPermutationVector);
-		return RemappedPermutationVector;
-	}
-
-	static FPermutationDomain BuildPermutationVector(const FViewInfo& View, bool bNeedsToSwitchVerticalAxis, bool bMetalMSAAHDRDecode)
-	{
-		TonemapperPermutation::FCommonDomain CommonPermutationVector = TonemapperPermutation::BuildCommonPermutationDomain(View, /* bGammaOnly = */ false, bNeedsToSwitchVerticalAxis, bMetalMSAAHDRDecode);
-
-		FMobileTonemapPS::FPermutationDomain MobilePermutationVector;
-		MobilePermutationVector.Set<TonemapperPermutation::FCommonDomain>(CommonPermutationVector);
-
-		// Must early exit if gamma only.
-		if (CommonPermutationVector.Get<TonemapperPermutation::FTonemapperGammaOnlyDim>())
-		{
-			return RemapPermutationVector(MobilePermutationVector);
-		}
-
-		const FFinalPostProcessSettings& Settings = View.FinalPostProcessSettings;
-		{
-			FVector MixerR(Settings.FilmChannelMixerRed);
-			FVector MixerG(Settings.FilmChannelMixerGreen);
-			FVector MixerB(Settings.FilmChannelMixerBlue);
-			if (
-				(Settings.FilmSaturation != 1.0f) ||
-				((MixerR - FVector(1.0f, 0.0f, 0.0f)).GetAbsMax() != 0.0f) ||
-				((MixerG - FVector(0.0f, 1.0f, 0.0f)).GetAbsMax() != 0.0f) ||
-				((MixerB - FVector(0.0f, 0.0f, 1.0f)).GetAbsMax() != 0.0f))
-			{
-				MobilePermutationVector.Set<FTonemapperColorMatrixDim>(true);
-			}
-		}
-		MobilePermutationVector.Set<FTonemapperShadowTintDim>(Settings.FilmShadowTintAmount > 0.0f);
-		MobilePermutationVector.Set<FTonemapperContrastDim>(Settings.FilmContrast > 0.0f);
-
-		if (GSupportsRenderTargetFormat_PF_FloatRGBA)
-		{
-			bool bUseDof = GetMobileDepthOfFieldScale(View) > 0.0f && !Settings.bMobileHQGaussian;
-
-			MobilePermutationVector.Set<FTonemapperDOFDim>(bUseDof);
-			MobilePermutationVector.Set<FTonemapperLightShaftsDim>(View.MobileLightShaft.IsSet());
-		}
-		else
-		{
-			// Override Bloom because is not supported.
-			CommonPermutationVector.Set<TonemapperPermutation::FTonemapperBloomDim>(false);
-		}
-
-		// Mobile is not currently supporting these.
-		CommonPermutationVector.Set<TonemapperPermutation::FTonemapperGrainJitterDim>(false);
-		CommonPermutationVector.Set<TonemapperPermutation::FTonemapperSharpenDim>(false);
-		MobilePermutationVector.Set<TonemapperPermutation::FCommonDomain>(CommonPermutationVector);
-
-		// We're not supporting every possible permutation, remap the permutation vector to combine common effects.
-		return RemapPermutationVector(MobilePermutationVector);
-	}
-
-	static bool ShouldCompilePermutation(const FGlobalShaderPermutationParameters& Parameters)
-	{
-		FPermutationDomain PermutationVector(Parameters.PermutationId);
-
-		auto CommonPermutationVector = PermutationVector.Get<TonemapperPermutation::FCommonDomain>();
-		if (!TonemapperPermutation::ShouldCompileCommonPermutation(Parameters, CommonPermutationVector))
-		{
-			return false;
-		}
-
-		// If this permutation vector is remapped at runtime, we can avoid the compile.
-		if (RemapPermutationVector(PermutationVector) != PermutationVector)
-		{
-			return false;
-		}
-
-		// Only cache for ES 3.1 shader platforms, and only compile 32bpp shaders for Android or PC emulation
-		return IsMobilePlatform(Parameters.Platform);
-	}
-
-	static void ModifyCompilationEnvironment(const FGlobalShaderPermutationParameters& Parameters, FShaderCompilerEnvironment& OutEnvironment)
-	{
-		//Need to hack in exposure scale for < SM5
-		OutEnvironment.SetDefine(TEXT("NO_EYEADAPTATION_EXPOSURE_FIX"), 1);
-	}
-};
-
-IMPLEMENT_GLOBAL_SHADER(FMobileTonemapPS, "/Engine/Private/PostProcessTonemap.usf", "MainPS_Mobile", SF_Pixel);
-
-class FMobileTonemapVS : public FGlobalShader
-{
-	DECLARE_GLOBAL_SHADER(FMobileTonemapVS);
-	SHADER_USE_PARAMETER_STRUCT_WITH_LEGACY_BASE(FMobileTonemapVS, FGlobalShader);
-
-	BEGIN_SHADER_PARAMETER_STRUCT(FParameters, )
-		SHADER_PARAMETER(FVector4, BufferSizeAndInvSize)
-		SHADER_PARAMETER(FVector4, DofBlurSizeAndInvSize)
-		SHADER_PARAMETER(FVector4, ViewportSize)
-		SHADER_PARAMETER(FVector, GrainRandomFull)
-		SHADER_PARAMETER(float, FringeIntensity)
-		SHADER_PARAMETER(FVector2D, Color_ScreenPosToViewportBias)
-		SHADER_PARAMETER(FVector2D, Color_ScreenPosToViewportScale)
-		SHADER_PARAMETER(FVector4, LensPrincipalPointOffsetScale)
-		SHADER_PARAMETER(float, DefaultEyeExposure)
-		SHADER_PARAMETER_STRUCT_REF(FViewUniformShaderParameters, View)
-		SHADER_PARAMETER_SRV(Buffer<float4>, EyeAdaptationBuffer)
-	END_SHADER_PARAMETER_STRUCT()
-
-	using FPermutationDomain = TShaderPermutationDomain<
-		TonemapperPermutation::FTonemapperSwitchAxis, TonemapperPermutation::FTonemapperEyeAdaptationDim
-	>;
-
-	static bool ShouldCompilePermutation(const FGlobalShaderPermutationParameters& Parameters)
-	{
-		FPermutationDomain PermutationVector(Parameters.PermutationId);
-
-		// Prevent switch axis permutation on platforms that dont require it.
-		if (PermutationVector.Get<TonemapperPermutation::FTonemapperSwitchAxis>() && !RHINeedsToSwitchVerticalAxis(Parameters.Platform))
-		{
-			return false;
-		}
-
-		return IsMobilePlatform(Parameters.Platform);
-	}
-
-	static void ModifyCompilationEnvironment(const FGlobalShaderPermutationParameters& Parameters, FShaderCompilerEnvironment& OutEnvironment)
-	{
-		FGlobalShader::ModifyCompilationEnvironment(Parameters, OutEnvironment);
-	}
-
-	static FPermutationDomain BuildPermutationVector(bool bNeedsToSwitchVerticalAxis, bool bEyeAdaptation)
-	{
-		FPermutationDomain PermutationVector;
-		PermutationVector.Set<TonemapperPermutation::FTonemapperSwitchAxis>(bNeedsToSwitchVerticalAxis);
-		PermutationVector.Set<TonemapperPermutation::FTonemapperEyeAdaptationDim>(bEyeAdaptation);
-		return PermutationVector;
-	}
-};
-
-IMPLEMENT_GLOBAL_SHADER(FMobileTonemapVS, "/Engine/Private/PostProcessTonemap.usf", "MainVS_Mobile", SF_Vertex);
-
-FScreenPassTexture AddMobileTonemapperPass(FRDGBuilder& GraphBuilder, const FViewInfo& View, const FMobileTonemapperInputs& Inputs)
-{
-	FScreenPassRenderTarget Output = Inputs.OverrideOutput;
-
-	if (!Output.IsValid())
-	{
-		FRDGTextureDesc OutputDesc = Inputs.SceneColor.Texture->Desc;
-		OutputDesc.Reset();
-		OutputDesc.Flags |= View.bUseComputePasses ? TexCreate_UAV : TexCreate_RenderTargetable;
-		OutputDesc.Flags |= GFastVRamConfig.Tonemap;
-
-		// RGB is the color in LDR, A is the luminance for PostprocessAA
-		OutputDesc.Format = Inputs.bOutputInHDR ? GRHIHDRDisplayOutputFormat : PF_B8G8R8A8;
-		OutputDesc.ClearValue = FClearValueBinding(FLinearColor(0, 0, 0, 0));
-
-		const FTonemapperOutputDeviceParameters OutputDeviceParameters = GetTonemapperOutputDeviceParameters(*View.Family);
-		const ETonemapperOutputDevice OutputDevice = static_cast<ETonemapperOutputDevice>(OutputDeviceParameters.OutputDevice);
-
-		if (OutputDevice == ETonemapperOutputDevice::LinearEXR)
-		{
-			OutputDesc.Format = PF_A32B32G32R32F;
-		}
-		if (OutputDevice == ETonemapperOutputDevice::LinearNoToneCurve || OutputDevice == ETonemapperOutputDevice::LinearWithToneCurve)
-		{
-			OutputDesc.Format = PF_FloatRGBA;
-		}
-
-		Output = FScreenPassRenderTarget(
-			GraphBuilder.CreateTexture(OutputDesc, TEXT("Tonemap")),
-			Inputs.SceneColor.ViewRect,
-			View.IsFirstInFamily() ? ERenderTargetLoadAction::EClear : ERenderTargetLoadAction::ELoad);
-	}
-
-	const FIntPoint& BufferSize = Inputs.SceneColor.Texture->Desc.Extent;
-	const FIntPoint DofBlurSize = Inputs.DofOutput.IsValid() ? Inputs.DofOutput.Texture->Desc.Extent : FIntPoint(1);
-
-	const FPostProcessSettings& Settings = View.FinalPostProcessSettings;
-	const FSceneViewFamily& ViewFamily = *(View.Family);
-
- 	FVector GrainRandomFullValue;
- 	{
- 		uint8 FrameIndexMod8 = 0;
- 		if (View.State)
- 		{
- 			FrameIndexMod8 = View.ViewState->GetFrameIndex(8);
- 		}
- 		GrainRandomFromFrame(&GrainRandomFullValue, FrameIndexMod8);
- 	}
-
-	FIntPoint ViewportOffset = Inputs.SceneColor.ViewRect.Min;
-	FIntPoint ViewportExtent = Inputs.SceneColor.ViewRect.Size();
-	FVector4 ScreenPosToScenePixelValue(
-		ViewportExtent.X * 0.5f,
-		-ViewportExtent.Y * 0.5f,
-		ViewportExtent.X * 0.5f + ViewportOffset.X,
-		ViewportExtent.Y * 0.5f + ViewportOffset.Y);
-
-	float Sharpen = FMath::Clamp(CVarTonemapperSharpen.GetValueOnRenderThread(), 0.0f, 10.0f);
-
-	FVector GrainValue;
-	GrainPostSettings(&GrainValue, &Settings);
-
-	FVector InvDisplayGammaValue;
-	InvDisplayGammaValue.X = 1.0f / ViewFamily.RenderTarget->GetDisplayGamma();
-	InvDisplayGammaValue.Y = 2.2f / ViewFamily.RenderTarget->GetDisplayGamma();
-	InvDisplayGammaValue.Z = 1.0; // Unused on mobile.
-
-	auto VertexShaderPermutationVector = FMobileTonemapVS::BuildPermutationVector(Inputs.bFlipYAxis, Inputs.bUseEyeAdaptation);
-	TShaderMapRef<FMobileTonemapVS> VertexShader(View.ShaderMap, VertexShaderPermutationVector);
-
-	FMobileTonemapVS::FParameters* VSShaderParameters = GraphBuilder.AllocParameters<FMobileTonemapVS::FParameters>();
-
-	VSShaderParameters->View = View.ViewUniformBuffer;
-	VSShaderParameters->BufferSizeAndInvSize = FVector4(BufferSize.X, BufferSize.Y, 1.0f / BufferSize.X, 1.0f / BufferSize.Y);
-	VSShaderParameters->DofBlurSizeAndInvSize = FVector4(DofBlurSize.X, DofBlurSize.Y, 1.0f / DofBlurSize.X, 1.0f / DofBlurSize.Y);
-	VSShaderParameters->ViewportSize = FVector4(ViewportExtent.X, ViewportExtent.Y, 1.0f / ViewportExtent.X, 1.0f / ViewportExtent.Y);
-	VSShaderParameters->GrainRandomFull = GrainRandomFullValue;
-	VSShaderParameters->FringeIntensity = fabsf(Settings.SceneFringeIntensity) * 0.01f;
-	VSShaderParameters->Color_ScreenPosToViewportScale = FVector2D(ScreenPosToScenePixelValue.X, ScreenPosToScenePixelValue.Y);
-	VSShaderParameters->Color_ScreenPosToViewportBias = FVector2D(ScreenPosToScenePixelValue.Z, ScreenPosToScenePixelValue.W);
-	VSShaderParameters->DefaultEyeExposure = Inputs.bUseEyeAdaptation ? 0.0f : GetEyeAdaptationFixedExposure(View);
-	VSShaderParameters->EyeAdaptationBuffer = Inputs.bUseEyeAdaptation && View.GetLastEyeAdaptationBuffer(GraphBuilder.RHICmdList) ? View.GetLastEyeAdaptationBuffer(GraphBuilder.RHICmdList)->SRV : nullptr;
-	VSShaderParameters->LensPrincipalPointOffsetScale = View.LensPrincipalPointOffsetScale;
-
-	auto PixelShaderPermutationVector = FMobileTonemapPS::BuildPermutationVector(View, Inputs.bFlipYAxis, Inputs.bMetalMSAAHDRDecode);
-
-	const FMobileFilmTonemapParameters Parameters = GetMobileFilmTonemapParameters(
-		Settings,
-		/* UseColorMatrix = */ PixelShaderPermutationVector.Get<FMobileTonemapPS::FTonemapperColorMatrixDim>(),
-		/* UseShadowTint = */ PixelShaderPermutationVector.Get<FMobileTonemapPS::FTonemapperShadowTintDim>(),
-		/* UseContrast = */ PixelShaderPermutationVector.Get<FMobileTonemapPS::FTonemapperContrastDim>());
-
-	TShaderMapRef<FMobileTonemapPS> PixelShader(View.ShaderMap, PixelShaderPermutationVector);
-
-	FMobileTonemapPS::FParameters* PSShaderParameters = GraphBuilder.AllocParameters<FMobileTonemapPS::FParameters>();
-	PSShaderParameters->RenderTargets[0] = Output.GetRenderTargetBinding();
-	PSShaderParameters->View = View.ViewUniformBuffer;
-	PSShaderParameters->ColorScale0 = FVector4(Settings.SceneColorTint.R, Settings.SceneColorTint.G, Settings.SceneColorTint.B, 0);
-	PSShaderParameters->ColorScale1 = FVector4(Settings.BloomIntensity, Settings.BloomIntensity, Settings.BloomIntensity, 0);
-	PSShaderParameters->TonemapperParams = FVector4(Settings.VignetteIntensity, Sharpen, 0.0f, 0.0f);
-	PSShaderParameters->OverlayColor = View.OverlayColor;
-	PSShaderParameters->ColorMatrixR_ColorCurveCd1 = Parameters.ColorMatrixR_ColorCurveCd1;
-	PSShaderParameters->ColorMatrixG_ColorCurveCd3Cm3 = Parameters.ColorMatrixG_ColorCurveCd3Cm3;
-	PSShaderParameters->ColorMatrixB_ColorCurveCm2 = Parameters.ColorMatrixB_ColorCurveCm2;
-	PSShaderParameters->ColorCurve_Cm0Cd0_Cd2_Ch0Cm1_Ch3 = Parameters.ColorCurve_Cm0Cd0_Cd2_Ch0Cm1_Ch3;
-	PSShaderParameters->ColorCurve_Ch1_Ch2 = Parameters.ColorCurve_Ch1_Ch2;
-	PSShaderParameters->ColorShadow_Luma = Parameters.ColorShadow_Luma;
-	PSShaderParameters->ColorShadow_Tint1 = Parameters.ColorShadow_Tint1;
-	PSShaderParameters->ColorShadow_Tint2 = Parameters.ColorShadow_Tint2;
-	PSShaderParameters->GrainScaleBiasJitter = GrainValue;
-	PSShaderParameters->InverseGamma = InvDisplayGammaValue;
-	PSShaderParameters->FringeIntensity = fabsf(Settings.SceneFringeIntensity) * 0.01f;
-	PSShaderParameters->SRGBAwareTarget = Inputs.bSRGBAwareTarget ? 1.0f : 0.0f;
-	PSShaderParameters->SceneColorTexture = Inputs.SceneColor.Texture;
-	PSShaderParameters->SceneColorSampler = TStaticSamplerState<SF_Bilinear, AM_Clamp, AM_Clamp, AM_Clamp>::GetRHI();
-	PSShaderParameters->BloomTexture = Inputs.BloomOutput.Texture;
-	PSShaderParameters->BloomSampler = TStaticSamplerState<SF_Bilinear, AM_Clamp, AM_Clamp, AM_Clamp>::GetRHI();
-	PSShaderParameters->DofBlurTexture = Inputs.DofOutput.Texture;
-	PSShaderParameters->DofBlurSampler = TStaticSamplerState<SF_Bilinear, AM_Clamp, AM_Clamp, AM_Clamp>::GetRHI();
-	PSShaderParameters->SunShaftAndDofTexture = Inputs.SunShaftAndDof.Texture;
-	PSShaderParameters->SunShaftAndDofSampler = TStaticSamplerState<SF_Bilinear, AM_Clamp, AM_Clamp, AM_Clamp>::GetRHI();
-
-	const FScreenPassTextureViewport OutputViewport(Output);
-
-	GraphBuilder.AddPass(
-		RDG_EVENT_NAME("Tonemapper %dx%d (PS)", OutputViewport.Extent.X, OutputViewport.Extent.Y),
-		PSShaderParameters,
-		ERDGPassFlags::Raster,
-		[VertexShader, VSShaderParameters, PixelShader, PSShaderParameters, OutputViewport](FRHICommandList& RHICmdList)
-	{
-		RHICmdList.SetViewport(OutputViewport.Rect.Min.X, OutputViewport.Rect.Min.Y, 0.0f, OutputViewport.Rect.Max.X, OutputViewport.Rect.Max.Y, 1.0f);
-
-		FGraphicsPipelineStateInitializer GraphicsPSOInit;
-		RHICmdList.ApplyCachedRenderTargets(GraphicsPSOInit);
-
-		GraphicsPSOInit.BlendState = TStaticBlendState<>::GetRHI();
-		GraphicsPSOInit.RasterizerState = TStaticRasterizerState<>::GetRHI();
-		GraphicsPSOInit.DepthStencilState = TStaticDepthStencilState<false, CF_Always>::GetRHI();
-
-		GraphicsPSOInit.BoundShaderState.VertexDeclarationRHI = GFilterVertexDeclaration.VertexDeclarationRHI;
-		GraphicsPSOInit.BoundShaderState.VertexShaderRHI = VertexShader.GetVertexShader();
-		GraphicsPSOInit.BoundShaderState.PixelShaderRHI = PixelShader.GetPixelShader();
-		GraphicsPSOInit.PrimitiveType = PT_TriangleList;
-
-		SetGraphicsPipelineState(RHICmdList, GraphicsPSOInit);
-
-		SetShaderParameters(RHICmdList, VertexShader, VertexShader.GetVertexShader(), *VSShaderParameters);
-		SetShaderParameters(RHICmdList, PixelShader, PixelShader.GetPixelShader(), *PSShaderParameters);
-
-		DrawRectangle(
-			RHICmdList,
-			0, 0,
-			OutputViewport.Extent.X, OutputViewport.Extent.Y,
-			0, 0,
-			OutputViewport.Rect.Width(), OutputViewport.Rect.Height(),
-			OutputViewport.Extent,
-			OutputViewport.Extent,
-			VertexShader,
-			EDRF_UseTriangleOptimization);
-	});
-
-	return MoveTemp(Output);
-}
diff --git a/Engine/Source/Runtime/Renderer/Private/PostProcess/PostProcessTonemap.h b/Engine/Source/Runtime/Renderer/Private/PostProcess/PostProcessTonemap.h
--- a/Engine/Source/Runtime/Renderer/Private/PostProcess/PostProcessTonemap.h
+++ b/Engine/Source/Runtime/Renderer/Private/PostProcess/PostProcessTonemap.h
@@ -38,19 +38,6 @@ static void GrainRandomFromFrame(FVector* RESTRICT const Constant, uint32 FrameN
 	Constant->Y = Halton(FrameNumber & 1023, 3);
 }
 
-BEGIN_SHADER_PARAMETER_STRUCT(FMobileFilmTonemapParameters, )
-	SHADER_PARAMETER(FVector4, ColorMatrixR_ColorCurveCd1)
-	SHADER_PARAMETER(FVector4, ColorMatrixG_ColorCurveCd3Cm3)
-	SHADER_PARAMETER(FVector4, ColorMatrixB_ColorCurveCm2)
-	SHADER_PARAMETER(FVector4, ColorCurve_Cm0Cd0_Cd2_Ch0Cm1_Ch3)
-	SHADER_PARAMETER(FVector4, ColorCurve_Ch1_Ch2)
-	SHADER_PARAMETER(FVector4, ColorShadow_Luma)
-	SHADER_PARAMETER(FVector4, ColorShadow_Tint1)
-	SHADER_PARAMETER(FVector4, ColorShadow_Tint2)
-END_SHADER_PARAMETER_STRUCT()
-
-FMobileFilmTonemapParameters GetMobileFilmTonemapParameters(const FPostProcessSettings& PostProcessSettings, bool bUseColorMatrix, bool bUseShadowTint, bool bUseContrast);
-
 struct FTonemapInputs
 {
 	// [Optional] Render to the specified output. If invalid, a new texture is created and returned.
diff --git a/Engine/Source/Runtime/Renderer/Private/PostProcess/PostProcessVisualizeHDR.cpp b/Engine/Source/Runtime/Renderer/Private/PostProcess/PostProcessVisualizeHDR.cpp
--- a/Engine/Source/Runtime/Renderer/Private/PostProcess/PostProcessVisualizeHDR.cpp
+++ b/Engine/Source/Runtime/Renderer/Private/PostProcess/PostProcessVisualizeHDR.cpp
@@ -16,7 +16,6 @@ public:
 	BEGIN_SHADER_PARAMETER_STRUCT(FParameters, )
 		SHADER_PARAMETER_STRUCT_REF(FViewUniformShaderParameters, View)
 		SHADER_PARAMETER_STRUCT(FEyeAdaptationParameters, EyeAdaptation)
-		SHADER_PARAMETER_STRUCT_INCLUDE(FMobileFilmTonemapParameters, MobileTonemap)
 		SHADER_PARAMETER_STRUCT_INCLUDE(FTonemapperOutputDeviceParameters, OutputDevice)
 		SHADER_PARAMETER_STRUCT(FScreenPassTextureViewportParameters, Input)
 		SHADER_PARAMETER_STRUCT(FScreenPassTextureViewportParameters, Output)
@@ -27,6 +26,11 @@ public:
 		SHADER_PARAMETER_SAMPLER(SamplerState, HDRSceneColorSampler)
 		SHADER_PARAMETER_SAMPLER(SamplerState, SceneColorSampler)
 		SHADER_PARAMETER_TEXTURE(Texture2D, MiniFontTexture)
+		SHADER_PARAMETER(float, FilmSlope)
+		SHADER_PARAMETER(float, FilmToe)
+		SHADER_PARAMETER(float, FilmShoulder)
+		SHADER_PARAMETER(float, FilmBlackClip)
+		SHADER_PARAMETER(float, FilmWhiteClip)
 		RENDER_TARGET_BINDING_SLOTS()
 	END_SHADER_PARAMETER_STRUCT()
 
@@ -67,6 +71,8 @@ FScreenPassTexture AddVisualizeHDRPass(FRDGBuilder& GraphBuilder, const FViewInf
 
 	FRHISamplerState* BilinearClampSampler = TStaticSamplerState<SF_Bilinear, AM_Clamp, AM_Clamp, AM_Clamp>::GetRHI();
 
+	const FPostProcessSettings& Settings = View.FinalPostProcessSettings;
+
 	FVisualizeHDRPS::FParameters* PassParameters = GraphBuilder.AllocParameters<FVisualizeHDRPS::FParameters>();
 	PassParameters->RenderTargets[0] = Output.GetRenderTargetBinding();
 	PassParameters->View = View.ViewUniformBuffer;
@@ -80,12 +86,12 @@ FScreenPassTexture AddVisualizeHDRPass(FRDGBuilder& GraphBuilder, const FViewInf
 	PassParameters->EyeAdaptationTexture = Inputs.EyeAdaptationTexture;
 	PassParameters->EyeAdaptation = *Inputs.EyeAdaptationParameters;
 	PassParameters->OutputDevice = GetTonemapperOutputDeviceParameters(*View.Family);
-	PassParameters->MobileTonemap = GetMobileFilmTonemapParameters(
-		View.FinalPostProcessSettings,
-		/* UseColorMatrix = */ true,
-		/* UseShadowTint = */ true,
-		/* UseContrast = */ true);
 	PassParameters->MiniFontTexture = GetMiniFontTexture();
+	PassParameters->FilmSlope = Settings.FilmSlope;
+	PassParameters->FilmToe = Settings.FilmToe;
+	PassParameters->FilmShoulder = Settings.FilmShoulder;
+	PassParameters->FilmBlackClip = Settings.FilmBlackClip;
+	PassParameters->FilmWhiteClip = Settings.FilmWhiteClip;
 
 	TShaderMapRef<FVisualizeHDRPS> PixelShader(View.ShaderMap);
 
diff --git a/Engine/Source/Runtime/Renderer/Private/PostProcess/PostProcessing.cpp b/Engine/Source/Runtime/Renderer/Private/PostProcess/PostProcessing.cpp
--- a/Engine/Source/Runtime/Renderer/Private/PostProcess/PostProcessing.cpp
+++ b/Engine/Source/Runtime/Renderer/Private/PostProcess/PostProcessing.cpp
@@ -1409,10 +1409,7 @@ void AddMobilePostProcessingPasses(FRDGBuilder& GraphBuilder, const FViewInfo& V
 	bool bSRGBAwareTarget = View.Family->RenderTarget->GetDisplayGamma() == 1.0f
 		&& View.bIsSceneCapture
 		&& IsMetalMobilePlatform(View.GetShaderPlatform());
-
-	static const auto VarTonemapperFilm = IConsoleManager::Get().FindTConsoleVariableDataInt(TEXT("r.Mobile.TonemapperFilm"));
-	const bool bUseTonemapperFilm = View.GetFeatureLevel() == ERHIFeatureLevel::ES3_1 && GSupportsRenderTargetFormat_PF_FloatRGBA && (VarTonemapperFilm && VarTonemapperFilm->GetValueOnRenderThread());
-		
+	
 	const EAutoExposureMethod AutoExposureMethod = GetAutoExposureMethod(View);
 	const bool bUseEyeAdaptation = IsMobileEyeAdaptationEnabled(View);
 
@@ -1623,15 +1620,12 @@ void AddMobilePostProcessingPasses(FRDGBuilder& GraphBuilder, const FViewInfo& V
 
 				DofOutput = DofBlurOutputs.DofBlur;
 
-				if (bUseTonemapperFilm)
-				{
-					FMobileIntegrateDofInputs IntegrateDofInputs;
-					IntegrateDofInputs.DofBlur = DofBlurOutputs.DofBlur;
-					IntegrateDofInputs.SceneColor = SceneColor;
-					IntegrateDofInputs.SunShaftAndDof = PostProcessSunShaftAndDof;
+				FMobileIntegrateDofInputs IntegrateDofInputs;
+				IntegrateDofInputs.DofBlur = DofBlurOutputs.DofBlur;
+				IntegrateDofInputs.SceneColor = SceneColor;
+				IntegrateDofInputs.SunShaftAndDof = PostProcessSunShaftAndDof;
 
-					SceneColor = AddMobileIntegrateDofPass(GraphBuilder, View, IntegrateDofInputs);
-				}
+				SceneColor = AddMobileIntegrateDofPass(GraphBuilder, View, IntegrateDofInputs);
 			}
 			else
 			{
@@ -1878,95 +1872,58 @@ void AddMobilePostProcessingPasses(FRDGBuilder& GraphBuilder, const FViewInfo& V
 			BloomOutput = BlackAlphaOneDummy;
 		}
 
-		if (bUseTonemapperFilm)
-		{
-			bool bDoGammaOnly = false;
+		bool bDoGammaOnly = false;
 
-			FRDGTextureRef ColorGradingTexture = nullptr;
+		FRDGTextureRef ColorGradingTexture = nullptr;
 
-			if (IStereoRendering::IsAPrimaryView(View))
-			{
-				ColorGradingTexture = AddCombineLUTPass(GraphBuilder, View);
-			}
-			// We can re-use the color grading texture from the primary view.
-			else if (View.GetTonemappingLUT())
-			{
-				ColorGradingTexture = TryRegisterExternalTexture(GraphBuilder, View.GetTonemappingLUT());
-			}
-			else
-			{
-				const FViewInfo* PrimaryView = static_cast<const FViewInfo*>(View.Family->Views[0]);
-				ColorGradingTexture = TryRegisterExternalTexture(GraphBuilder, PrimaryView->GetTonemappingLUT());
-			}
+		if (IStereoRendering::IsAPrimaryView(View))
+		{
+			ColorGradingTexture = AddCombineLUTPass(GraphBuilder, View);
+		}
+		// We can re-use the color grading texture from the primary view.
+		else if (View.GetTonemappingLUT())
+		{
+			ColorGradingTexture = TryRegisterExternalTexture(GraphBuilder, View.GetTonemappingLUT());
+		}
+		else
+		{
+			const FViewInfo* PrimaryView = static_cast<const FViewInfo*>(View.Family->Views[0]);
+			ColorGradingTexture = TryRegisterExternalTexture(GraphBuilder, PrimaryView->GetTonemappingLUT());
+		}
 
-			FTonemapInputs TonemapperInputs;
-			PassSequence.AcceptOverrideIfLastPass(EPass::Tonemap, TonemapperInputs.OverrideOutput);
+		FTonemapInputs TonemapperInputs;
+		PassSequence.AcceptOverrideIfLastPass(EPass::Tonemap, TonemapperInputs.OverrideOutput);
 
-			// This is the view family render target.
-			if (TonemapperInputs.OverrideOutput.Texture)
+		// This is the view family render target.
+		if (TonemapperInputs.OverrideOutput.Texture)
+		{
+			FIntRect OutputViewRect;
+			if (View.PrimaryScreenPercentageMethod == EPrimaryScreenPercentageMethod::RawOutput)
 			{
-				FIntRect OutputViewRect;
-				if (View.PrimaryScreenPercentageMethod == EPrimaryScreenPercentageMethod::RawOutput)
-				{
-					OutputViewRect = View.ViewRect;
-				}
-				else
-				{
-					OutputViewRect = View.UnscaledViewRect;
-				}
-				ERenderTargetLoadAction  OutputLoadAction = View.IsFirstInFamily() ? ERenderTargetLoadAction::EClear : ERenderTargetLoadAction::ELoad;
-
-				TonemapperInputs.OverrideOutput.ViewRect = OutputViewRect;
-				TonemapperInputs.OverrideOutput.LoadAction = OutputLoadAction;
+				OutputViewRect = View.ViewRect;
 			}
-			
-			TonemapperInputs.SceneColor = SceneColor;
-			TonemapperInputs.Bloom = BloomOutput;
-			TonemapperInputs.EyeAdaptationTexture = nullptr;
-			TonemapperInputs.ColorGradingTexture = ColorGradingTexture;
-			TonemapperInputs.bWriteAlphaChannel = View.AntiAliasingMethod == AAM_FXAA || IsPostProcessingWithAlphaChannelSupported() || bUseMobileDof;
-			TonemapperInputs.bFlipYAxis = RHINeedsToSwitchVerticalAxis(View.GetShaderPlatform()) && !PassSequence.IsEnabled(EPass::PostProcessMaterialAfterTonemapping);
-			TonemapperInputs.bOutputInHDR = bHDRTonemapperOutput;
-			TonemapperInputs.bGammaOnly = bDoGammaOnly;
-			TonemapperInputs.bMetalMSAAHDRDecode = bMetalMSAAHDRDecode;
-			TonemapperInputs.EyeAdaptationBuffer = bUseEyeAdaptation && View.GetLastEyeAdaptationBuffer(GraphBuilder.RHICmdList) ? View.GetLastEyeAdaptationBuffer(GraphBuilder.RHICmdList)->SRV : nullptr;
-
-			SceneColor = AddTonemapPass(GraphBuilder, View, TonemapperInputs);
-		}
-		else
-		{
-			FMobileTonemapperInputs TonemapperInputs;
-			PassSequence.AcceptOverrideIfLastPass(EPass::Tonemap, TonemapperInputs.OverrideOutput);
-			if (TonemapperInputs.OverrideOutput.Texture)
+			else
 			{
-				FIntRect OutputViewRect;
-				if (View.PrimaryScreenPercentageMethod == EPrimaryScreenPercentageMethod::RawOutput)
-				{
-					OutputViewRect = View.ViewRect;
-				}
-				else
-				{
-					OutputViewRect = View.UnscaledViewRect;
-				}
-				ERenderTargetLoadAction  OutputLoadAction = View.IsFirstInFamily() ? ERenderTargetLoadAction::EClear : ERenderTargetLoadAction::ELoad;
-
-				TonemapperInputs.OverrideOutput.ViewRect = OutputViewRect;
-				TonemapperInputs.OverrideOutput.LoadAction = OutputLoadAction;
+				OutputViewRect = View.UnscaledViewRect;
 			}
+			ERenderTargetLoadAction  OutputLoadAction = View.IsFirstInFamily() ? ERenderTargetLoadAction::EClear : ERenderTargetLoadAction::ELoad;
 
-			TonemapperInputs.bFlipYAxis = RHINeedsToSwitchVerticalAxis(View.GetShaderPlatform()) && !PassSequence.IsEnabled(EPass::PostProcessMaterialAfterTonemapping);
-			TonemapperInputs.bMetalMSAAHDRDecode = bMetalMSAAHDRDecode;
-			TonemapperInputs.bOutputInHDR = bHDRTonemapperOutput;
-			TonemapperInputs.bSRGBAwareTarget = bSRGBAwareTarget;
-			TonemapperInputs.bUseEyeAdaptation = bUseEyeAdaptation;
-			TonemapperInputs.SceneColor = SceneColor;
-			TonemapperInputs.BloomOutput = BloomOutput;
-			TonemapperInputs.DofOutput = DofOutput;
-			TonemapperInputs.SunShaftAndDof = PostProcessSunShaftAndDof;
-			TonemapperInputs.EyeAdaptationBuffer = bUseEyeAdaptation && View.GetLastEyeAdaptationBuffer(GraphBuilder.RHICmdList) ? View.GetLastEyeAdaptationBuffer(GraphBuilder.RHICmdList)->SRV : nullptr;
-
-			SceneColor = AddMobileTonemapperPass(GraphBuilder, View, TonemapperInputs);
+			TonemapperInputs.OverrideOutput.ViewRect = OutputViewRect;
+			TonemapperInputs.OverrideOutput.LoadAction = OutputLoadAction;
 		}
+			
+		TonemapperInputs.SceneColor = SceneColor;
+		TonemapperInputs.Bloom = BloomOutput;
+		TonemapperInputs.EyeAdaptationTexture = nullptr;
+		TonemapperInputs.ColorGradingTexture = ColorGradingTexture;
+		TonemapperInputs.bWriteAlphaChannel = View.AntiAliasingMethod == AAM_FXAA || IsPostProcessingWithAlphaChannelSupported() || bUseMobileDof;
+		TonemapperInputs.bFlipYAxis = RHINeedsToSwitchVerticalAxis(View.GetShaderPlatform()) && !PassSequence.IsEnabled(EPass::PostProcessMaterialAfterTonemapping);
+		TonemapperInputs.bOutputInHDR = bHDRTonemapperOutput;
+		TonemapperInputs.bGammaOnly = bDoGammaOnly;
+		TonemapperInputs.bMetalMSAAHDRDecode = bMetalMSAAHDRDecode;
+		TonemapperInputs.EyeAdaptationBuffer = bUseEyeAdaptation && View.GetLastEyeAdaptationBuffer(GraphBuilder.RHICmdList) ? View.GetLastEyeAdaptationBuffer(GraphBuilder.RHICmdList)->SRV : nullptr;
+
+		SceneColor = AddTonemapPass(GraphBuilder, View, TonemapperInputs);
 
 		//The output color should been decoded to linear space after tone mapper apparently
 		bMetalMSAAHDRDecode = false;
