Fix for 'Unknown shader parameter type' shader error.

#rnx

diff --git a/Engine/Source/Runtime/Renderer/Private/SceneRendering.cpp b/Engine/Source/Runtime/Renderer/Private/SceneRendering.cpp
--- a/Engine/Source/Runtime/Renderer/Private/SceneRendering.cpp
+++ b/Engine/Source/Runtime/Renderer/Private/SceneRendering.cpp
@@ -4290,8 +4290,8 @@ TGlobalResource<FDummySceneColorResolveBuffer> GResolveDummyVertexBuffer;
 extern int32 GAllowCustomMSAAResolves;
 
 BEGIN_SHADER_PARAMETER_STRUCT(FResolveSceneColorParameters, )
-	SHADER_PARAMETER_RDG_TEXTURE(, SceneColor)
-	SHADER_PARAMETER_RDG_TEXTURE_SRV(, SceneColorFMask)
+	RDG_TEXTURE_ACCESS(SceneColor, ERHIAccess::SRVGraphics)
+	SHADER_PARAMETER_RDG_TEXTURE_SRV(Texture2D, SceneColorFMask)
 	RENDER_TARGET_BINDING_SLOTS()
 END_SHADER_PARAMETER_STRUCT()
 
