Instrumented DrawTileMesh RDG passes.

diff --git a/Engine/Source/Runtime/Renderer/Private/Renderer.cpp b/Engine/Source/Runtime/Renderer/Private/Renderer.cpp
--- a/Engine/Source/Runtime/Renderer/Private/Renderer.cpp
+++ b/Engine/Source/Runtime/Renderer/Private/Renderer.cpp
@@ -209,6 +209,8 @@ void FRendererModule::DrawTileMesh(FCanvasRenderContext& RenderContext, FMeshPas
 			FRDGSystemTextures::Create(GraphBuilder);
 		}
 
+		RDG_EVENT_SCOPE(GraphBuilder, "DrawTileMesh");
+
 		auto* PassParameters = GraphBuilder.AllocParameters<FDrawTileMeshPassParameters>();
 		PassParameters->RenderTargets[0] = FRenderTargetBinding(RenderContext.GetRenderTarget(), ERenderTargetLoadAction::ELoad);
 		PassParameters->View = View.ViewUniformBuffer;
@@ -226,7 +228,7 @@ void FRendererModule::DrawTileMesh(FCanvasRenderContext& RenderContext, FMeshPas
 			{
 				PassParameters->DebugViewMode = CreateDebugViewModePassUniformBuffer(GraphBuilder, View, nullptr);
 
-				RenderContext.AddPass({}, PassParameters,
+				RenderContext.AddPass(RDG_EVENT_NAME("OutputMaterialTextureScales"), PassParameters,
 					[Scene, &View, &Mesh](FRHICommandListImmediate& RHICmdList)
 				{
 					DrawDynamicMeshPass(View, RHICmdList, [&](FMeshPassDrawListContext* InDrawListContext)
@@ -250,7 +252,7 @@ void FRendererModule::DrawTileMesh(FCanvasRenderContext& RenderContext, FMeshPas
 			{
 				PassParameters->TranslucentBasePass = CreateTranslucentBasePassUniformBuffer(GraphBuilder, View);
 
-				RenderContext.AddPass({}, PassParameters,
+				RenderContext.AddPass(RDG_EVENT_NAME("TranslucentDeferred"), PassParameters,
 					[Scene, &View, &Mesh, DrawRenderState, bUse128bitRT](FRHICommandListImmediate& RHICmdList)
 				{
 					DrawDynamicMeshPass(View, RHICmdList, [&](FDynamicPassMeshDrawListContext* DynamicMeshPassContext)
@@ -273,7 +275,7 @@ void FRendererModule::DrawTileMesh(FCanvasRenderContext& RenderContext, FMeshPas
 			{
 				PassParameters->MobileBasePass = CreateMobileBasePassUniformBuffer(GraphBuilder, View, EMobileBasePass::Translucent);
 
-				RenderContext.AddPass({}, PassParameters,
+				RenderContext.AddPass(RDG_EVENT_NAME("TranslucentMobile"), PassParameters,
 					[Scene, &View, DrawRenderState, &Mesh](FRHICommandListImmediate& RHICmdList)
 				{
 					DrawDynamicMeshPass(View, RHICmdList, [&](FDynamicPassMeshDrawListContext* DynamicMeshPassContext)
@@ -305,7 +307,7 @@ void FRendererModule::DrawTileMesh(FCanvasRenderContext& RenderContext, FMeshPas
 				ensureMsgf(HitProxyId == Mesh.BatchHitProxyId, TEXT("Only Mesh.BatchHitProxyId is used for hit testing."));
 
 #if WITH_EDITOR
-				RenderContext.AddPass({}, PassParameters,
+				RenderContext.AddPass(RDG_EVENT_NAME("HitTesting"), PassParameters,
 					[Scene, &View, DrawRenderState, &Mesh](FRHICommandListImmediate& RHICmdList)
 				{
 					DrawDynamicMeshPass(View, RHICmdList, [&](FDynamicPassMeshDrawListContext* DynamicMeshPassContext)
@@ -329,7 +331,7 @@ void FRendererModule::DrawTileMesh(FCanvasRenderContext& RenderContext, FMeshPas
 				{
 					PassParameters->OpaqueBasePass = CreateOpaqueBasePassUniformBuffer(GraphBuilder, View);
 
-					RenderContext.AddPass({}, PassParameters,
+					RenderContext.AddPass(RDG_EVENT_NAME("OpaqueDeferred"), PassParameters,
 						[Scene, &View, DrawRenderState, &Mesh, bUse128bitRT](FRHICommandListImmediate& RHICmdList)
 					{
 						DrawDynamicMeshPass(View, RHICmdList,
@@ -352,7 +354,7 @@ void FRendererModule::DrawTileMesh(FCanvasRenderContext& RenderContext, FMeshPas
 				{
 					PassParameters->MobileBasePass = CreateMobileBasePassUniformBuffer(GraphBuilder, View, EMobileBasePass::Opaque);
 
-					RenderContext.AddPass({}, PassParameters,
+					RenderContext.AddPass(RDG_EVENT_NAME("OpaqueMobile"), PassParameters,
 						[Scene, &View, DrawRenderState, &Mesh](FRHICommandListImmediate& RHICmdList)
 					{
 						DrawDynamicMeshPass(View, RHICmdList, [&](FDynamicPassMeshDrawListContext* DynamicMeshPassContext)
