Port all mobile eye-adaptation resources to RDG.

diff --git a/Engine/Source/Runtime/Engine/Private/SceneView.cpp b/Engine/Source/Runtime/Engine/Private/SceneView.cpp
--- a/Engine/Source/Runtime/Engine/Private/SceneView.cpp
+++ b/Engine/Source/Runtime/Engine/Private/SceneView.cpp
@@ -2470,7 +2470,7 @@ IPooledRenderTarget* FSceneView::GetEyeAdaptationTexture() const
 	return nullptr;
 }
 
-const FExposureBufferData* FSceneView::GetEyeAdaptationBuffer() const
+FRDGPooledBuffer* FSceneView::GetEyeAdaptationBuffer() const
 {
 	checkf(FeatureLevel == ERHIFeatureLevel::ES3_1, TEXT("ES3_1 use RWBuffer for read back"));
 	if (EyeAdaptationViewState)
diff --git a/Engine/Source/Runtime/Engine/Public/SceneManagement.h b/Engine/Source/Runtime/Engine/Public/SceneManagement.h
--- a/Engine/Source/Runtime/Engine/Public/SceneManagement.h
+++ b/Engine/Source/Runtime/Engine/Public/SceneManagement.h
@@ -222,7 +222,7 @@ public:
 	virtual IPooledRenderTarget* GetCurrentEyeAdaptationTexture() const = 0;
 
 	/** Returns the eye adaptation buffer (mobile only). */
-	virtual const FExposureBufferData* GetCurrentEyeAdaptationBuffer() const = 0;
+	virtual FRDGPooledBuffer* GetCurrentEyeAdaptationBuffer() const = 0;
 
 	virtual void SetSequencerState(ESequencerState InSequencerState) = 0;
 
diff --git a/Engine/Source/Runtime/Engine/Public/SceneView.h b/Engine/Source/Runtime/Engine/Public/SceneView.h
--- a/Engine/Source/Runtime/Engine/Public/SceneView.h
+++ b/Engine/Source/Runtime/Engine/Public/SceneView.h
@@ -28,7 +28,6 @@ class FSceneViewFamily;
 class FVolumetricFogViewResources;
 class FIESLightProfileResource;
 class ITemporalUpscaler;
-struct FExposureBufferData;
 
 enum class ERayTracingRenderMode
 {
@@ -1382,7 +1381,7 @@ public:
 	IPooledRenderTarget* GetEyeAdaptationTexture() const;
 
 	/** Returns the eye adaptation buffer (mobile) or null if it doesn't exist. */
-	const FExposureBufferData* GetEyeAdaptationBuffer() const;
+	FRDGPooledBuffer* GetEyeAdaptationBuffer() const;
 
 
 protected:
diff --git a/Engine/Source/Runtime/RenderCore/Private/RenderUtils.cpp b/Engine/Source/Runtime/RenderCore/Private/RenderUtils.cpp
--- a/Engine/Source/Runtime/RenderCore/Private/RenderUtils.cpp
+++ b/Engine/Source/Runtime/RenderCore/Private/RenderUtils.cpp
@@ -8,6 +8,7 @@
 #include "RenderGraphUtils.h"
 #include "PipelineStateCache.h"
 #include "Misc/ConfigCacheIni.h"
+#include "RenderGraphResourcePool.h"
 
 #if WITH_EDITOR
 #include "Misc/CoreMisc.h"
@@ -44,6 +45,17 @@ void FTextureWithRDG::InitRDG(const TCHAR* Name)
 	RenderTarget = CreateRenderTarget(TextureRHI, Name);
 }
 
+FBufferWithRDG::FBufferWithRDG() = default;
+FBufferWithRDG::FBufferWithRDG(const FBufferWithRDG & Other) = default;
+FBufferWithRDG& FBufferWithRDG::operator=(const FBufferWithRDG & Other) = default;
+FBufferWithRDG::~FBufferWithRDG() = default;
+
+void FBufferWithRDG::ReleaseRHI()
+{
+	Buffer = nullptr;
+	FRenderResource::ReleaseRHI();
+}
+
 const uint16 GCubeIndices[12*3] =
 {
 	0, 2, 3,
@@ -237,6 +249,31 @@ public:
 
 FVertexBufferWithSRV* GWhiteVertexBufferWithSRV = new TGlobalResource<FWhiteVertexBuffer>;
 
+class FWhiteVertexBufferWithRDG : public FBufferWithRDG
+{
+public:
+
+	/**
+	 * Initialize RHI resources.
+	 */
+	virtual void InitRHI() override
+	{
+		if (!Buffer.IsValid())
+		{
+			FRHICommandList* UnusedCmdList = new FRHICommandList(FRHIGPUMask::All());
+			GetPooledFreeBuffer(*UnusedCmdList, FRDGBufferDesc::CreateBufferDesc(sizeof(FVector4), 1), Buffer, TEXT("WhiteVertexBufferWithRDG"));
+
+			FVector4* BufferData = (FVector4*)RHILockBuffer(Buffer->GetRHI(), 0, sizeof(FVector4), RLM_WriteOnly);
+			*BufferData = FVector4(1.0f, 1.0f, 1.0f, 1.0f);
+			RHIUnlockBuffer(Buffer->GetRHI());
+			delete UnusedCmdList;
+			UnusedCmdList = nullptr;
+		}
+	}
+};
+
+FBufferWithRDG* GWhiteVertexBufferWithRDG = new TGlobalResource<FWhiteVertexBufferWithRDG>();
+
 /**
  * A class representing a 1x1x1 black volume texture.
  */
diff --git a/Engine/Source/Runtime/RenderCore/Public/RenderUtils.h b/Engine/Source/Runtime/RenderCore/Public/RenderUtils.h
--- a/Engine/Source/Runtime/RenderCore/Public/RenderUtils.h
+++ b/Engine/Source/Runtime/RenderCore/Public/RenderUtils.h
@@ -9,6 +9,9 @@
 #include "RenderResource.h"
 #include "RHIDefinitions.h"
 
+class FBufferWithRDG;
+class FRDGBuffer;
+class FRDGPooledBuffer;
 class FTextureWithRDG;
 class FRDGTexture;
 class FRDGBuilder;
@@ -35,6 +38,19 @@ private:
 	TRefCountPtr<IPooledRenderTarget> RenderTarget;
 };
 
+class RENDERCORE_API FBufferWithRDG : public FRenderResource
+{
+public:
+	FBufferWithRDG();
+	FBufferWithRDG(const FBufferWithRDG& Other);
+	FBufferWithRDG& operator=(const FBufferWithRDG& Other);
+	~FBufferWithRDG() override;
+
+	void ReleaseRHI() override;
+
+	TRefCountPtr<FRDGPooledBuffer> Buffer;
+};
+
 extern RENDERCORE_API void RenderUtilsInit();
 
 /**
@@ -138,6 +154,8 @@ extern RENDERCORE_API class FVertexBufferWithSRV* GEmptyVertexBufferWithUAV;
 
 extern RENDERCORE_API class FVertexBufferWithSRV* GWhiteVertexBufferWithSRV;
 
+extern RENDERCORE_API class FBufferWithRDG* GWhiteVertexBufferWithRDG;
+
 /** A global black array texture. */
 extern RENDERCORE_API class FTexture* GBlackArrayTexture;
 
diff --git a/Engine/Source/Runtime/Renderer/Private/MobileBasePassRendering.cpp b/Engine/Source/Runtime/Renderer/Private/MobileBasePassRendering.cpp
--- a/Engine/Source/Runtime/Renderer/Private/MobileBasePassRendering.cpp
+++ b/Engine/Source/Runtime/Renderer/Private/MobileBasePassRendering.cpp
@@ -190,7 +190,7 @@ void SetupMobileBasePassUniformParameters(
 		}
 	}
 
-	BasePassParameters.EyeAdaptationBuffer = GetEyeAdaptationBuffer(View);
+	BasePassParameters.EyeAdaptationBuffer = GraphBuilder.CreateSRV(GetEyeAdaptationBuffer(GraphBuilder, View), PF_A32B32G32R32F);
 
 	if (BasePass == EMobileBasePass::Opaque && HasBeenProduced(ScreenSpaceAOTexture))
 	{
diff --git a/Engine/Source/Runtime/Renderer/Private/MobileBasePassRendering.h b/Engine/Source/Runtime/Renderer/Private/MobileBasePassRendering.h
--- a/Engine/Source/Runtime/Renderer/Private/MobileBasePassRendering.h
+++ b/Engine/Source/Runtime/Renderer/Private/MobileBasePassRendering.h
@@ -31,7 +31,7 @@ BEGIN_GLOBAL_SHADER_PARAMETER_STRUCT(FMobileBasePassUniformParameters, )
 	SHADER_PARAMETER_STRUCT(FMobileSceneTextureUniformParameters, SceneTextures)
 	SHADER_PARAMETER_TEXTURE(Texture2D, PreIntegratedGFTexture)
 	SHADER_PARAMETER_SAMPLER(SamplerState, PreIntegratedGFSampler)
-	SHADER_PARAMETER_SRV(Buffer<float4>, EyeAdaptationBuffer)
+	SHADER_PARAMETER_RDG_BUFFER_SRV(Buffer<float4>, EyeAdaptationBuffer)
 	SHADER_PARAMETER_RDG_TEXTURE(Texture2D, AmbientOcclusionTexture)
 	SHADER_PARAMETER_SAMPLER(SamplerState, AmbientOcclusionSampler)
 END_GLOBAL_SHADER_PARAMETER_STRUCT()
diff --git a/Engine/Source/Runtime/Renderer/Private/MobileShadingRenderer.cpp b/Engine/Source/Runtime/Renderer/Private/MobileShadingRenderer.cpp
--- a/Engine/Source/Runtime/Renderer/Private/MobileShadingRenderer.cpp
+++ b/Engine/Source/Runtime/Renderer/Private/MobileShadingRenderer.cpp
@@ -478,12 +478,6 @@ void FMobileSceneRenderer::InitViews(FRDGBuilder& GraphBuilder, FSceneTexturesCo
 
 		// Initialize the view's RHI resources.
 		View.InitRHIResources();
-
-		// Get the custom 1x1 target used to store exposure value and Toggle the two render targets used to store new and old.
-		if (IsMobileEyeAdaptationEnabled(View))
-		{
-			View.SwapEyeAdaptationBuffers();
-		}
 	}
 
 	Scene->GPUScene.Update(GraphBuilder, *Scene);
diff --git a/Engine/Source/Runtime/Renderer/Private/PostProcess/PostProcessEyeAdaptation.cpp b/Engine/Source/Runtime/Renderer/Private/PostProcess/PostProcessEyeAdaptation.cpp
--- a/Engine/Source/Runtime/Renderer/Private/PostProcess/PostProcessEyeAdaptation.cpp
+++ b/Engine/Source/Runtime/Renderer/Private/PostProcess/PostProcessEyeAdaptation.cpp
@@ -822,59 +822,59 @@ const TRefCountPtr<IPooledRenderTarget>& FSceneViewState::FEyeAdaptationManager:
 	return PooledRenderTarget[TextureIndex];
 }
 
-const FExposureBufferData* FSceneViewState::FEyeAdaptationManager::GetBuffer(uint32 BufferIndex) const
+const TRefCountPtr<FRDGPooledBuffer>& FSceneViewState::FEyeAdaptationManager::GetBuffer(uint32 BufferIndex) const
 {
 	check(BufferIndex == 0 || BufferIndex == 1);
 
-	const FExposureBufferData& Instance = ExposureBufferData[BufferIndex];
-	return Instance.IsValid() ? &Instance : nullptr;
+	return ExposureBufferData[BufferIndex];
 }
 
-FExposureBufferData* FSceneViewState::FEyeAdaptationManager::GetOrCreateBuffer(FRHICommandListImmediate& RHICmdList, uint32 BufferIndex)
+const TRefCountPtr<FRDGPooledBuffer>& FSceneViewState::FEyeAdaptationManager::GetOrCreateBuffer(FRDGBuilder& GraphBuilder, uint32 BufferIndex)
 {
 	check(BufferIndex == 0 || BufferIndex == 1);
 
 	// Create textures if needed.
 	if (!ExposureBufferData[BufferIndex].IsValid())
 	{
-		FRHIResourceCreateInfo CreateInfo(TEXT("ExposureBuffer"));
-		ExposureBufferData[BufferIndex].Buffer = RHICreateVertexBuffer(sizeof(FVector4), BUF_Static | BUF_ShaderResource | BUF_UnorderedAccess | BUF_SourceCopy, CreateInfo);
+		FRDGBufferDesc RDGBufferDesc = FRDGBufferDesc::CreateBufferDesc(sizeof(FVector4), 1);
+		RDGBufferDesc.Usage |= BUF_SourceCopy;
+		FRDGBufferRef RDGBuffer = GraphBuilder.CreateBuffer(RDGBufferDesc, TEXT("EyeAdaptationBuffer"), ERDGBufferFlags::MultiFrame);
 
-		FVector4* BufferData = (FVector4*)RHICmdList.LockBuffer(ExposureBufferData[BufferIndex].Buffer, 0, sizeof(FVector4), RLM_WriteOnly);
-		*BufferData = FVector4(1.0f, 1.0f, 1.0f, 1.0f);
-		RHICmdList.UnlockBuffer(ExposureBufferData[BufferIndex].Buffer);
-
-		ExposureBufferData[BufferIndex].SRV = RHICmdList.CreateShaderResourceView(ExposureBufferData[BufferIndex].Buffer, sizeof(FVector4), PF_A32B32G32R32F);
-		ExposureBufferData[BufferIndex].UAV = RHICmdList.CreateUnorderedAccessView(ExposureBufferData[BufferIndex].Buffer, PF_A32B32G32R32F);
+		ConvertToExternalBuffer(GraphBuilder, RDGBuffer, ExposureBufferData[BufferIndex]);
 
-		RHICmdList.Transition(FRHITransitionInfo(ExposureBufferData[BufferIndex].UAV, ERHIAccess::UAVMask, ERHIAccess::SRVMask));
+		FVector4* BufferData = (FVector4*)GraphBuilder.RHICmdList.LockBuffer(ExposureBufferData[BufferIndex]->GetRHI(), 0, sizeof(FVector4), RLM_WriteOnly);
+		*BufferData = FVector4(1.0f, 1.0f, 1.0f, 1.0f);
+		GraphBuilder.RHICmdList.UnlockBuffer(ExposureBufferData[BufferIndex]->GetRHI());
 	}
 
-	return &ExposureBufferData[BufferIndex];
+	return ExposureBufferData[BufferIndex];
 }
 
-void FSceneViewState::FEyeAdaptationManager::SwapBuffers(bool bInUpdateLastExposure)
+void FSceneViewState::FEyeAdaptationManager::SwapBuffers(FRDGBuilder& GraphBuilder, bool bInUpdateLastExposure)
 {
 	QUICK_SCOPE_CYCLE_COUNTER(STAT_FEyeAdaptationManager_SwapBuffers);
 
 	FRHICommandListImmediate& RHICmdList = FRHICommandListExecutor::GetImmediateCommandList();
-	
 
 	check(CurrentBuffer == 0 || CurrentBuffer == 1);
 
+	// When reading back last frame's exposure data in AFR, make sure we do it on the
+	// GPUs that were actually used last frame.
+	const FRHIGPUMask ThisFrameGPUMask = RHICmdList.GetGPUMask();
+	const FRHIGPUMask LastFrameGPUMask = AFRUtils::GetPrevSiblingGPUMask(ThisFrameGPUMask);
+	const FRHIGPUMask LastLastFrameGPUMask = AFRUtils::GetPrevSiblingGPUMask(LastFrameGPUMask);
+
 	if (bInUpdateLastExposure && ExposureBufferData[CurrentBuffer].IsValid() && (GIsEditor || CVarEnablePreExposureOnlyInTheEditor.GetValueOnRenderThread() == 0))
 	{
+		FRDGBufferRef CurrentRDGBuffer = GraphBuilder.RegisterExternalBuffer(ExposureBufferData[CurrentBuffer], ERDGBufferFlags::MultiFrame);
+
 		bool bReadbackCompleted = false;
 		if (ExposureBufferReadback != nullptr && ExposureBufferReadback->IsReady())
 		{
 			// Workaround until FRHIGPUTextureReadback::Lock has multigpu support
-			FRHIGPUMask ReadBackGPUMask = RHICmdList.GetGPUMask();
-			if (!ReadBackGPUMask.HasSingleIndex())
-			{
-				ReadBackGPUMask = FRHIGPUMask::GPU0();
-			}
+			FRHIGPUMask ReadBackGPUMask = LastLastFrameGPUMask;
 
-			SCOPED_GPU_MASK(RHICmdList, ReadBackGPUMask);
+			RDG_GPU_MASK_SCOPE(GraphBuilder, ReadBackGPUMask);
 
 			// Read the last request results.
 			FVector4* ReadbackData = (FVector4*)ExposureBufferReadback->Lock(sizeof(FVector4));
@@ -897,11 +897,7 @@ void FSceneViewState::FEyeAdaptationManager::SwapBuffers(bool bInUpdateLastExpos
 				ExposureBufferReadback.Reset(new FRHIGPUBufferReadback(ExposureValueName));
 			}
 
-			RHICmdList.Transition(FRHITransitionInfo(ExposureBufferData[CurrentBuffer].UAV, ERHIAccess::UAVMask, ERHIAccess::CopySrc));
-
-			ExposureBufferReadback->EnqueueCopy(RHICmdList, ExposureBufferData[CurrentBuffer].Buffer);
-
-			RHICmdList.Transition(FRHITransitionInfo(ExposureBufferData[CurrentBuffer].UAV, ERHIAccess::CopySrc, ERHIAccess::SRVMask));
+			AddEnqueueCopyPass(GraphBuilder, ExposureBufferReadback.Get(), CurrentRDGBuffer, 0);
 		}
 	}
 
diff --git a/Engine/Source/Runtime/Renderer/Private/PostProcess/PostProcessMaterial.cpp b/Engine/Source/Runtime/Renderer/Private/PostProcess/PostProcessMaterial.cpp
--- a/Engine/Source/Runtime/Renderer/Private/PostProcess/PostProcessMaterial.cpp
+++ b/Engine/Source/Runtime/Renderer/Private/PostProcess/PostProcessMaterial.cpp
@@ -478,7 +478,7 @@ FScreenPassTexture AddPostProcessMaterialPass(
 	PostProcessMaterialParameters->View = View.ViewUniformBuffer;
 	if (bMobilePlatform)
 	{
-		PostProcessMaterialParameters->EyeAdaptationBuffer = GetEyeAdaptationBuffer(View);
+		PostProcessMaterialParameters->EyeAdaptationBuffer = GraphBuilder.CreateSRV(GetEyeAdaptationBuffer(GraphBuilder, View), PF_A32B32G32R32F);
 	}
 	else
 	{
diff --git a/Engine/Source/Runtime/Renderer/Private/PostProcess/PostProcessMaterial.h b/Engine/Source/Runtime/Renderer/Private/PostProcess/PostProcessMaterial.h
--- a/Engine/Source/Runtime/Renderer/Private/PostProcess/PostProcessMaterial.h
+++ b/Engine/Source/Runtime/Renderer/Private/PostProcess/PostProcessMaterial.h
@@ -44,7 +44,7 @@ BEGIN_SHADER_PARAMETER_STRUCT(FPostProcessMaterialParameters, )
 	SHADER_PARAMETER_RDG_TEXTURE(Texture2D, MobileCustomStencilTexture)
 	SHADER_PARAMETER_SAMPLER(SamplerState, MobileCustomStencilTextureSampler)
 	SHADER_PARAMETER_RDG_TEXTURE(Texture2D, EyeAdaptationTexture)
-	SHADER_PARAMETER_SRV(Buffer<float4>, EyeAdaptationBuffer)
+	SHADER_PARAMETER_RDG_BUFFER_SRV(Buffer<float4>, EyeAdaptationBuffer)
 	SHADER_PARAMETER(int32, MobileStencilValueRef)
 	SHADER_PARAMETER(uint32, bFlipYAxis)
 	SHADER_PARAMETER(uint32, bMetalMSAAHDRDecode)
diff --git a/Engine/Source/Runtime/Renderer/Private/PostProcess/PostProcessMobile.cpp b/Engine/Source/Runtime/Renderer/Private/PostProcess/PostProcessMobile.cpp
--- a/Engine/Source/Runtime/Renderer/Private/PostProcess/PostProcessMobile.cpp
+++ b/Engine/Source/Runtime/Renderer/Private/PostProcess/PostProcessMobile.cpp
@@ -2135,9 +2135,9 @@ public:
 	BEGIN_SHADER_PARAMETER_STRUCT(FParameters, )
 		SHADER_PARAMETER_STRUCT_REF(FViewUniformShaderParameters, View)
 		SHADER_PARAMETER_STRUCT(FEyeAdaptationParameters, EyeAdaptation)
-		SHADER_PARAMETER_SRV(Buffer<float4>, EyeAdaptationBuffer)
+		SHADER_PARAMETER_RDG_BUFFER_SRV(Buffer<float4>, EyeAdaptationBuffer)
 		SHADER_PARAMETER_RDG_BUFFER_SRV(Buffer<uint>, LogLuminanceWeightBuffer)
-		SHADER_PARAMETER_UAV(RWBuffer<float4>, OutputBuffer)
+		SHADER_PARAMETER_RDG_BUFFER_UAV(RWBuffer<float4>, OutputBuffer)
 	END_SHADER_PARAMETER_STRUCT()
 
 	/** Static Shader boilerplate */
@@ -2168,9 +2168,9 @@ class FMobileHistogramEyeAdaptationCS : public FGlobalShader
 public:
 	BEGIN_SHADER_PARAMETER_STRUCT(FParameters, )
 		SHADER_PARAMETER_STRUCT(FEyeAdaptationParameters, EyeAdaptation)
-		SHADER_PARAMETER_SRV(Buffer<float4>, EyeAdaptationBuffer)
+		SHADER_PARAMETER_RDG_BUFFER_SRV(Buffer<float4>, EyeAdaptationBuffer)
 		SHADER_PARAMETER_RDG_BUFFER_SRV(Buffer<uint>, HistogramBuffer)
-		SHADER_PARAMETER_UAV(RWBuffer<float4>, OutputBuffer)
+		SHADER_PARAMETER_RDG_BUFFER_UAV(RWBuffer<float4>, OutputBuffer)
 	END_SHADER_PARAMETER_STRUCT()
 
 	static bool ShouldCompilePermutation(const FGlobalShaderPermutationParameters& Parameters)
@@ -2190,15 +2190,14 @@ IMPLEMENT_GLOBAL_SHADER(FMobileHistogramEyeAdaptationCS, "/Engine/Private/PostPr
 
 void AddMobileEyeAdaptationPass(FRDGBuilder& GraphBuilder, const FViewInfo& View, const FEyeAdaptationParameters& EyeAdaptationParameters, const FMobileEyeAdaptationInputs& Inputs)
 {
-	const FExposureBufferData* EyeAdaptationThisFrameBuffer = View.GetEyeAdaptationBuffer(GraphBuilder.RHICmdList);
-	const FExposureBufferData* EyeAdaptationLastFrameBuffer = View.GetLastEyeAdaptationBuffer(GraphBuilder.RHICmdList);
+	// Get the custom 1x1 target used to store exposure value and Toggle the two render targets used to store new and old.
+	View.SwapEyeAdaptationBuffers(GraphBuilder);
 
-	check(EyeAdaptationThisFrameBuffer && EyeAdaptationLastFrameBuffer);
+	FRDGBufferRef EyeAdaptationBuffer = GraphBuilder.RegisterExternalBuffer(View.GetLastEyeAdaptationBuffer(GraphBuilder), ERDGBufferFlags::MultiFrame);
+	FRDGBufferSRVRef EyeAdaptationBufferSRV = GraphBuilder.CreateSRV(EyeAdaptationBuffer, PF_A32B32G32R32F);
 
-	AddPass(GraphBuilder, [EyeAdaptationThisFrameBuffer](FRHIComputeCommandList& RHICmdList)
-	{
-		RHICmdList.Transition(FRHITransitionInfo(EyeAdaptationThisFrameBuffer->UAV, ERHIAccess::SRVMask, ERHIAccess::UAVMask));
-	});
+	FRDGBufferRef OutputBuffer = GraphBuilder.RegisterExternalBuffer(View.GetEyeAdaptationBuffer(GraphBuilder), ERDGBufferFlags::MultiFrame);
+	FRDGBufferUAVRef OutputBufferUAV = GraphBuilder.CreateUAV(OutputBuffer, PF_A32B32G32R32F);
 
 	if (Inputs.bUseBasicEyeAdaptation)
 	{
@@ -2206,9 +2205,9 @@ void AddMobileEyeAdaptationPass(FRDGBuilder& GraphBuilder, const FViewInfo& View
 
 		PassParameters->EyeAdaptation = EyeAdaptationParameters;
 		PassParameters->View = View.ViewUniformBuffer;
-		PassParameters->EyeAdaptationBuffer = EyeAdaptationLastFrameBuffer->SRV;
+		PassParameters->EyeAdaptationBuffer = EyeAdaptationBufferSRV;
 		PassParameters->LogLuminanceWeightBuffer = Inputs.EyeAdaptationSetupSRV;
-		PassParameters->OutputBuffer = EyeAdaptationThisFrameBuffer->UAV;
+		PassParameters->OutputBuffer = OutputBufferUAV;
 
 		TShaderMapRef<FMobileBasicEyeAdaptationCS> ComputeShader(View.ShaderMap);
 
@@ -2224,9 +2223,9 @@ void AddMobileEyeAdaptationPass(FRDGBuilder& GraphBuilder, const FViewInfo& View
 		FMobileHistogramEyeAdaptationCS::FParameters* PassParameters = GraphBuilder.AllocParameters<FMobileHistogramEyeAdaptationCS::FParameters>();
 
 		PassParameters->EyeAdaptation = EyeAdaptationParameters;
-		PassParameters->EyeAdaptationBuffer = EyeAdaptationLastFrameBuffer->SRV;
+		PassParameters->EyeAdaptationBuffer = EyeAdaptationBufferSRV;
 		PassParameters->HistogramBuffer = Inputs.EyeAdaptationSetupSRV;
-		PassParameters->OutputBuffer = EyeAdaptationThisFrameBuffer->UAV;
+		PassParameters->OutputBuffer = OutputBufferUAV;
 
 		TShaderMapRef<FMobileHistogramEyeAdaptationCS> ComputeShader(View.ShaderMap);
 
diff --git a/Engine/Source/Runtime/Renderer/Private/PostProcess/PostProcessTonemap.cpp b/Engine/Source/Runtime/Renderer/Private/PostProcess/PostProcessTonemap.cpp
--- a/Engine/Source/Runtime/Renderer/Private/PostProcess/PostProcessTonemap.cpp
+++ b/Engine/Source/Runtime/Renderer/Private/PostProcess/PostProcessTonemap.cpp
@@ -364,7 +364,7 @@ BEGIN_SHADER_PARAMETER_STRUCT(FTonemapParameters, )
 	SHADER_PARAMETER(float, EditorNITLevel)
 	SHADER_PARAMETER(uint32, bOutputInHDR)
 	// ES3_1 uses EyeAdaptationBuffer
-	SHADER_PARAMETER_SRV(Buffer<float4>, EyeAdaptationBuffer)
+	SHADER_PARAMETER_RDG_BUFFER_SRV(Buffer<float4>, EyeAdaptationBuffer)
 END_SHADER_PARAMETER_STRUCT()
 
 class FTonemapVS : public FGlobalShader
@@ -605,7 +605,10 @@ FScreenPassTexture AddTonemapPass(FRDGBuilder& GraphBuilder, const FViewInfo& Vi
 	CommonParameters.LensPrincipalPointOffsetScaleInverse.Y = -View.LensPrincipalPointOffsetScale.Y / View.LensPrincipalPointOffsetScale.W;
 	CommonParameters.LensPrincipalPointOffsetScaleInverse.Z = 1.0f / View.LensPrincipalPointOffsetScale.Z;
 	CommonParameters.LensPrincipalPointOffsetScaleInverse.W = 1.0f / View.LensPrincipalPointOffsetScale.W;
-	CommonParameters.EyeAdaptationBuffer = Inputs.EyeAdaptationBuffer;
+	
+	FRDGBufferSRVRef EyeAdaptationBufferSRV = Inputs.EyeAdaptationBuffer != nullptr ? GraphBuilder.CreateSRV(Inputs.EyeAdaptationBuffer, PF_A32B32G32R32F) : nullptr;
+
+	CommonParameters.EyeAdaptationBuffer = EyeAdaptationBufferSRV;
 
 	// Generate permutation vector for the desktop tonemapper.
 	TonemapperPermutation::FDesktopDomain DesktopPermutationVector;
diff --git a/Engine/Source/Runtime/Renderer/Private/PostProcess/PostProcessTonemap.h b/Engine/Source/Runtime/Renderer/Private/PostProcess/PostProcessTonemap.h
--- a/Engine/Source/Runtime/Renderer/Private/PostProcess/PostProcessTonemap.h
+++ b/Engine/Source/Runtime/Renderer/Private/PostProcess/PostProcessTonemap.h
@@ -56,7 +56,7 @@ struct FTonemapInputs
 	FRDGTextureRef EyeAdaptationTexture = nullptr;
 
 	// [Optional, ES31] Eye adaptation buffer used to compute exposure. 
-	FRHIShaderResourceView* EyeAdaptationBuffer = nullptr;
+	FRDGBufferRef EyeAdaptationBuffer = nullptr;
 
 	// [Raster Only, Mobile] Flips the image vertically on output.
 	bool bFlipYAxis = false;
@@ -90,7 +90,7 @@ struct FMobileTonemapperInputs
 
 	FScreenPassTexture SunShaftAndDof;
 
-	FRHIShaderResourceView* EyeAdaptationBuffer = nullptr;
+	FRDGBufferRef EyeAdaptationBuffer = nullptr;
 
 	// [Raster Only, Mobile] Flips the image vertically on output.
 	bool bFlipYAxis = false;
diff --git a/Engine/Source/Runtime/Renderer/Private/PostProcess/PostProcessing.cpp b/Engine/Source/Runtime/Renderer/Private/PostProcess/PostProcessing.cpp
--- a/Engine/Source/Runtime/Renderer/Private/PostProcess/PostProcessing.cpp
+++ b/Engine/Source/Runtime/Renderer/Private/PostProcess/PostProcessing.cpp
@@ -1921,7 +1921,7 @@ void AddMobilePostProcessingPasses(FRDGBuilder& GraphBuilder, const FViewInfo& V
 		TonemapperInputs.bOutputInHDR = bHDRTonemapperOutput;
 		TonemapperInputs.bGammaOnly = bDoGammaOnly;
 		TonemapperInputs.bMetalMSAAHDRDecode = bMetalMSAAHDRDecode;
-		TonemapperInputs.EyeAdaptationBuffer = bUseEyeAdaptation && View.GetLastEyeAdaptationBuffer(GraphBuilder.RHICmdList) ? View.GetLastEyeAdaptationBuffer(GraphBuilder.RHICmdList)->SRV : nullptr;
+		TonemapperInputs.EyeAdaptationBuffer = bUseEyeAdaptation ? TryRegisterExternalBuffer(GraphBuilder, View.GetLastEyeAdaptationBuffer(GraphBuilder), ERDGBufferFlags::MultiFrame) : nullptr;
 
 		SceneColor = AddTonemapPass(GraphBuilder, View, TonemapperInputs);
 
diff --git a/Engine/Source/Runtime/Renderer/Private/ScenePrivate.h b/Engine/Source/Runtime/Renderer/Private/ScenePrivate.h
--- a/Engine/Source/Runtime/Renderer/Private/ScenePrivate.h
+++ b/Engine/Source/Runtime/Renderer/Private/ScenePrivate.h
@@ -692,25 +692,6 @@ struct FHLODSceneNodeVisibilityState
 	uint16 bIsFading	: 1;
 };
 
-struct FExposureBufferData
-{
-	FBufferRHIRef Buffer;
-	FShaderResourceViewRHIRef SRV;
-	FUnorderedAccessViewRHIRef UAV;
-
-	bool IsValid() const
-	{
-		return Buffer.IsValid();
-	}
-
-	void SafeRelease()
-	{
-		Buffer.SafeRelease();
-		SRV.SafeRelease();
-		UAV.SafeRelease();
-	}
-};
-
 /**
  * The scene manager's private implementation of persistent view state.
  * This class is associated with a particular camera across multiple frames by the game thread.
@@ -883,28 +864,28 @@ private:
 		/** Get the last frame average scene luminance (used for exposure compensation curve) */
 		float GetLastAverageSceneLuminance() const { return LastAverageSceneLuminance; }
 
-		const FExposureBufferData* GetCurrentBuffer() const
+		const TRefCountPtr<FRDGPooledBuffer>& GetCurrentBuffer() const
 		{
 			return GetBuffer(CurrentBuffer);
 		}
 
-		const FExposureBufferData* GetCurrentBuffer(FRHICommandListImmediate& RHICmdList)
+		const TRefCountPtr<FRDGPooledBuffer>& GetCurrentBuffer(FRDGBuilder& GraphBuilder)
 		{
-			return GetOrCreateBuffer(RHICmdList, CurrentBuffer);
+			return GetOrCreateBuffer(GraphBuilder, CurrentBuffer);
 		}
 
-		const FExposureBufferData* GetLastBuffer(FRHICommandListImmediate& RHICmdList)
+		const TRefCountPtr<FRDGPooledBuffer>& GetLastBuffer(FRDGBuilder& GraphBuilder)
 		{
-			return GetOrCreateBuffer(RHICmdList, 1 - CurrentBuffer);
+			return GetOrCreateBuffer(GraphBuilder, 1 - CurrentBuffer);
 		}
 
-		void SwapBuffers(bool bUpdateLastExposure);
+		void SwapBuffers(FRDGBuilder& GraphBuilder, bool bUpdateLastExposure);
 	private:
 		const TRefCountPtr<IPooledRenderTarget>& GetTexture(uint32 TextureIndex) const;
 		const TRefCountPtr<IPooledRenderTarget>& GetOrCreateTexture(FRHICommandList& RHICmdList, uint32 TextureIndex);
 
-		const FExposureBufferData* GetBuffer(uint32 BufferIndex) const;
-		FExposureBufferData* GetOrCreateBuffer(FRHICommandListImmediate& RHICmdList, uint32 BufferIndex);
+		const TRefCountPtr<FRDGPooledBuffer>& GetBuffer(uint32 BufferIndex) const;
+		const TRefCountPtr<FRDGPooledBuffer>& GetOrCreateBuffer(FRDGBuilder& GraphBuilder, uint32 BufferIndex);
 
 		int32 CurrentBuffer = 0;
 
@@ -918,7 +899,7 @@ private:
 		TUniquePtr<FRHIGPUTextureReadback> ExposureTextureReadback[3];
 
 		// ES3.1 feature level. For efficent readback use buffers instead of textures
-		FExposureBufferData ExposureBufferData[2];
+		TRefCountPtr<FRDGPooledBuffer> ExposureBufferData[2];
 		TUniquePtr<FRHIGPUBufferReadback> ExposureBufferReadback;
 
 	} EyeAdaptationManager;
@@ -1281,27 +1262,27 @@ public:
 		EyeAdaptationManager.SwapTextures(GraphBuilder, bUpdateLastExposure && bValidEyeAdaptationTexture);
 	}
 
-	const FExposureBufferData* GetCurrentEyeAdaptationBuffer() const override final
+	FRDGPooledBuffer* GetCurrentEyeAdaptationBuffer() const override final
 	{
-		const FExposureBufferData* Buffer = EyeAdaptationManager.GetCurrentBuffer();
+		FRDGPooledBuffer* Buffer = EyeAdaptationManager.GetCurrentBuffer().GetReference();
 		check(bValidEyeAdaptationBuffer && Buffer);
 		return Buffer;
 	}
 
-	const FExposureBufferData* GetCurrentEyeAdaptationBuffer(FRHICommandListImmediate& RHICmdList)
+	FRDGPooledBuffer* GetCurrentEyeAdaptationBuffer(FRDGBuilder& GraphBuilder)
 	{
 		bValidEyeAdaptationBuffer = true;
-		return EyeAdaptationManager.GetCurrentBuffer(RHICmdList);
+		return EyeAdaptationManager.GetCurrentBuffer(GraphBuilder).GetReference();
 	}
 
-	const FExposureBufferData* GetLastEyeAdaptationBuffer(FRHICommandListImmediate& RHICmdList)
+	FRDGPooledBuffer* GetLastEyeAdaptationBuffer(FRDGBuilder& GraphBuilder)
 	{
-		return EyeAdaptationManager.GetLastBuffer(RHICmdList);
+		return EyeAdaptationManager.GetLastBuffer(GraphBuilder).GetReference();
 	}
 
-	void SwapEyeAdaptationBuffers()
+	void SwapEyeAdaptationBuffers(FRDGBuilder& GraphBuilder)
 	{
-		EyeAdaptationManager.SwapBuffers(bUpdateLastExposure && bValidEyeAdaptationBuffer);
+		EyeAdaptationManager.SwapBuffers(GraphBuilder, bUpdateLastExposure && bValidEyeAdaptationBuffer);
 	}
 
 #if WITH_MGPU
diff --git a/Engine/Source/Runtime/Renderer/Private/SceneRendering.cpp b/Engine/Source/Runtime/Renderer/Private/SceneRendering.cpp
--- a/Engine/Source/Runtime/Renderer/Private/SceneRendering.cpp
+++ b/Engine/Source/Runtime/Renderer/Private/SceneRendering.cpp
@@ -2018,35 +2018,35 @@ void FViewInfo::SwapEyeAdaptationTextures(FRDGBuilder& GraphBuilder) const
 	}
 }
 
-const FExposureBufferData* FViewInfo::GetEyeAdaptationBuffer(FRHICommandListImmediate& RHICmdList) const
+FRDGPooledBuffer* FViewInfo::GetEyeAdaptationBuffer(FRDGBuilder& GraphBuilder) const
 {
 	checkf(FeatureLevel == ERHIFeatureLevel::ES3_1, TEXT("ES3_1 use RWBuffer for read back"));
 
 	if (FSceneViewState* EffectiveViewState = GetEyeAdaptationViewState())
 	{
-		return EffectiveViewState->GetCurrentEyeAdaptationBuffer(RHICmdList);
+		return EffectiveViewState->GetCurrentEyeAdaptationBuffer(GraphBuilder);
 	}
 	return nullptr;
 }
 
-const FExposureBufferData* FViewInfo::GetLastEyeAdaptationBuffer(FRHICommandListImmediate& RHICmdList) const
+FRDGPooledBuffer* FViewInfo::GetLastEyeAdaptationBuffer(FRDGBuilder& GraphBuilder) const
 {
 	checkf(FeatureLevel == ERHIFeatureLevel::ES3_1, TEXT("ES3_1 use RWBuffer for read back"));
 
 	if (FSceneViewState* EffectiveViewState = GetEyeAdaptationViewState())
 	{
-		return EffectiveViewState->GetLastEyeAdaptationBuffer(RHICmdList);
+		return EffectiveViewState->GetLastEyeAdaptationBuffer(GraphBuilder);
 	}
 	return nullptr;
 }
 
-void FViewInfo::SwapEyeAdaptationBuffers() const
+void FViewInfo::SwapEyeAdaptationBuffers(FRDGBuilder& GraphBuilder) const
 {
 	checkf(FeatureLevel == ERHIFeatureLevel::ES3_1, TEXT("ES3_1 use RWBuffer for read back"));
 
 	if (FSceneViewState* EffectiveViewState = GetEyeAdaptationViewState())
 	{
-		EffectiveViewState->SwapEyeAdaptationBuffers();
+		EffectiveViewState->SwapEyeAdaptationBuffers(GraphBuilder);
 	}
 }
 
diff --git a/Engine/Source/Runtime/Renderer/Private/SceneRendering.h b/Engine/Source/Runtime/Renderer/Private/SceneRendering.h
--- a/Engine/Source/Runtime/Renderer/Private/SceneRendering.h
+++ b/Engine/Source/Runtime/Renderer/Private/SceneRendering.h
@@ -1477,8 +1477,8 @@ public:
 
 	/** Allocates and returns the current eye adaptation buffer. */
 	using FSceneView::GetEyeAdaptationBuffer;
-	const FExposureBufferData* GetEyeAdaptationBuffer(FRHICommandListImmediate& RHICmdList) const;
-	const FExposureBufferData* GetLastEyeAdaptationBuffer(FRHICommandListImmediate& RHICmdList) const;
+	FRDGPooledBuffer* GetEyeAdaptationBuffer(FRDGBuilder& GraphBuilder) const;
+	FRDGPooledBuffer* GetLastEyeAdaptationBuffer(FRDGBuilder& GraphBuilder) const;
 
 #if WITH_MGPU
 	void BroadcastEyeAdaptationTemporalEffect(FRHICommandList& RHICmdList);
@@ -1493,7 +1493,7 @@ public:
 
 	/**Swap the order of the two eye adaptation targets in the double buffer system */
 	void SwapEyeAdaptationTextures(FRDGBuilder& GraphBuilder) const;
-	void SwapEyeAdaptationBuffers() const;
+	void SwapEyeAdaptationBuffers(FRDGBuilder& GraphBuilder) const;
 	
 	/** Returns the load action to use when overwriting all pixels of a target that you intend to read from. Takes into account the HMD hidden area mesh. */
 	ERenderTargetLoadAction GetOverwriteLoadAction() const;
diff --git a/Engine/Source/Runtime/Renderer/Private/SceneTextureParameters.cpp b/Engine/Source/Runtime/Renderer/Private/SceneTextureParameters.cpp
--- a/Engine/Source/Runtime/Renderer/Private/SceneTextureParameters.cpp
+++ b/Engine/Source/Runtime/Renderer/Private/SceneTextureParameters.cpp
@@ -83,14 +83,14 @@ FRDGTextureRef GetEyeAdaptationTexture(FRDGBuilder& GraphBuilder, const FSceneVi
 	}
 }
 
-FRHIShaderResourceView* GetEyeAdaptationBuffer(const FSceneView& View)
+FRDGBufferRef GetEyeAdaptationBuffer(FRDGBuilder& GraphBuilder, const FSceneView& View)
 {
 	if (View.HasValidEyeAdaptationBuffer())
 	{
-		return View.GetEyeAdaptationBuffer()->SRV;
+		return GraphBuilder.RegisterExternalBuffer(View.GetEyeAdaptationBuffer(), ERDGBufferFlags::MultiFrame);
 	}
 	else
 	{
-		return GWhiteVertexBufferWithSRV->ShaderResourceViewRHI;
+		return GraphBuilder.RegisterExternalBuffer(GWhiteVertexBufferWithRDG->Buffer);
 	}
 }
\ No newline at end of file
diff --git a/Engine/Source/Runtime/Renderer/Private/SceneTextureParameters.h b/Engine/Source/Runtime/Renderer/Private/SceneTextureParameters.h
--- a/Engine/Source/Runtime/Renderer/Private/SceneTextureParameters.h
+++ b/Engine/Source/Runtime/Renderer/Private/SceneTextureParameters.h
@@ -42,4 +42,4 @@ FSceneLightingChannelParameters GetSceneLightingChannelParameters(FRDGBuilder& G
 RENDERER_API FRDGTextureRef GetEyeAdaptationTexture(FRDGBuilder& GraphBuilder, const FSceneView& View);
 
 /** Returns a render graph buffer resource reference onto the eye adaptation or fallback. */
-RENDERER_API FRHIShaderResourceView* GetEyeAdaptationBuffer(const FSceneView& View);
+RENDERER_API FRDGBufferRef GetEyeAdaptationBuffer(FRDGBuilder& GraphBuilder, const FSceneView& View);
