Add precision of the input resolution on the spatial upscale pass

diff --git a/Engine/Source/Runtime/Renderer/Private/PostProcess/PostProcessUpscale.cpp b/Engine/Source/Runtime/Renderer/Private/PostProcess/PostProcessUpscale.cpp
--- a/Engine/Source/Runtime/Renderer/Private/PostProcess/PostProcessUpscale.cpp
+++ b/Engine/Source/Runtime/Renderer/Private/PostProcess/PostProcessUpscale.cpp
@@ -228,7 +228,7 @@ FScreenPassTexture AddUpscalePass(FRDGBuilder& GraphBuilder, const FViewInfo& Vi
 	const TCHAR* StageName = StageNames[static_cast<uint32>(Inputs.Stage)];
 
 	GraphBuilder.AddPass(
-		RDG_EVENT_NAME("Upscale (%s) %dx%d", StageName, Output.ViewRect.Width(), Output.ViewRect.Height()),
+		RDG_EVENT_NAME("Upscale (%s) %dx%d -> %dx%d", StageName, Inputs.SceneColor.ViewRect.Width(), Inputs.SceneColor.ViewRect.Height(), Output.ViewRect.Width(), Output.ViewRect.Height()),
 		PassParameters,
 		ERDGPassFlags::Raster,
 		[&View, bUsePaniniProjection, PixelShader, PassParameters, InputViewport, OutputViewport](FRHICommandList& RHICmdList)
