Fix draw indirect buffer accumulating tiles indifinitely.

diff --git a/Engine/Shaders/Private/Strata/StrataMaterialClassification.usf b/Engine/Shaders/Private/Strata/StrataMaterialClassification.usf
--- a/Engine/Shaders/Private/Strata/StrataMaterialClassification.usf
+++ b/Engine/Shaders/Private/Strata/StrataMaterialClassification.usf
@@ -201,8 +201,6 @@ void TileMainCS(uint2 DispatchThreadId : SV_DispatchThreadID, uint LinearIndex :
 	if (all(DispatchThreadId == 0))
 	{
 		TileIndirectData[0] = bRectPrimitive > 0 ? 4 : 6;
-		TileIndirectData[2] = 0;
-		TileIndirectData[3] = 0;
 	}
 
 	// TODO: add a SM6 permutation with ballot?
diff --git a/Engine/Source/Runtime/Renderer/Private/Strata/Strata.cpp b/Engine/Source/Runtime/Renderer/Private/Strata/Strata.cpp
--- a/Engine/Source/Runtime/Renderer/Private/Strata/Strata.cpp
+++ b/Engine/Source/Runtime/Renderer/Private/Strata/Strata.cpp
@@ -89,6 +89,8 @@ void InitialiseStrataFrameSceneData(FSceneRenderer& SceneRenderer, FRDGBuilder&
 			GraphBuilder.PreallocateBuffer(ClassificationTileIndirectBuffer);
 			StrataSceneData.ClassificationTileListBuffer = GraphBuilder.GetPooledBuffer(ClassificationTileListBuffer);
 			StrataSceneData.ClassificationTileIndirectBuffer = GraphBuilder.GetPooledBuffer(ClassificationTileIndirectBuffer);
+
+			AddClearUAVPass(GraphBuilder, GraphBuilder.CreateUAV(GraphBuilder.RegisterExternalBuffer(StrataSceneData.ClassificationTileIndirectBuffer), PF_R32_UINT), 0);
 		}
 
 		// Top layer texture
