Implement multi-view support for VSM clipmaps (i.e., split screen and stereo)

diff --git a/Engine/Shaders/Private/VirtualShadowMaps/PageManagement.usf b/Engine/Shaders/Private/VirtualShadowMaps/PageManagement.usf
--- a/Engine/Shaders/Private/VirtualShadowMaps/PageManagement.usf
+++ b/Engine/Shaders/Private/VirtualShadowMaps/PageManagement.usf
@@ -62,7 +62,7 @@ float GetDeviceZFromPixel(uint2 PixelPos)
 	float DeviceZ = 0.0f;
 #if LOAD_DEPTH_FROM_NANITE_BUFFER
 	// Load Nanite Depth
-	UlongType VisPixel = VisBuffer64[PixelPos];
+	UlongType VisPixel = VisBuffer64[uint2(View.ViewRectMin.xy) + PixelPos];
 	uint DepthInt = 0;
 	uint VisibleClusterIndex = 0;
 	uint TriIndex = 0;
diff --git a/Engine/Source/Runtime/Renderer/Private/SceneRendering.cpp b/Engine/Source/Runtime/Renderer/Private/SceneRendering.cpp
--- a/Engine/Source/Runtime/Renderer/Private/SceneRendering.cpp
+++ b/Engine/Source/Runtime/Renderer/Private/SceneRendering.cpp
@@ -61,6 +61,7 @@
 #include "SingleLayerWaterRendering.h"
 #include "HairStrands/HairStrandsVisibility.h"
 #include "SystemTextures.h"
+#include "VirtualShadowMaps/VirtualShadowMapClipmap.h"
 #if WITH_EDITOR
 #include "Rendering/StaticLightingSystemInterface.h"
 #endif
@@ -438,6 +439,28 @@ bool IsStaticLightingAllowed()
 	return CVar->GetValueOnRenderThread() != 0;
 }
 
+
+TSharedPtr<FVirtualShadowMapClipmap> FVisibleLightInfo::FindShadowClipmapForView(const FViewInfo* View) const
+{
+	for (const auto& Clipmap : VirtualShadowMapClipmaps)
+	{
+		if (Clipmap->GetDependentView() == View)
+		{
+			return Clipmap;
+		}
+	}
+	
+	// This has to mirror the if (IStereoRendering::IsAPrimaryView(View)) test in ShadowSetup.cpp, which ensures only one view dependent shadow is set up for a stereo pair.
+	// TODO: this should very much be explicitly linked.
+	if (!IStereoRendering::IsAPrimaryView(*View) && VirtualShadowMapClipmaps.Num() > 0)
+	{
+		return VirtualShadowMapClipmaps[0];
+	}
+
+	return TSharedPtr<FVirtualShadowMapClipmap>();
+}
+
+
 #if !UE_BUILD_SHIPPING
 namespace
 {
diff --git a/Engine/Source/Runtime/Renderer/Private/SceneRendering.h b/Engine/Source/Runtime/Renderer/Private/SceneRendering.h
--- a/Engine/Source/Runtime/Renderer/Private/SceneRendering.h
+++ b/Engine/Source/Runtime/Renderer/Private/SceneRendering.h
@@ -191,6 +191,12 @@ public:
 
 	/** FProjectedShadowInfo set up to receive a dense copy of all types of geometry (nanite+trad) to be able to sample from e.g., volumetric effects */
 	TArray<FProjectedShadowInfo*,SceneRenderingAllocator> CompleteProjectedShadows;
+
+	/**
+	 * Search VirtualShadowMapClipmaps for the one that was set up with the given view. 
+	 * TODO: should VirtualShadowMapClipmap move to FVisibleLightViewInfo?
+	 */
+	TSharedPtr<FVirtualShadowMapClipmap> FindShadowClipmapForView(const FViewInfo* View) const;
 };
 
 // Stores the primitive count of each translucency pass (redundant, could be computed after sorting but this way we touch less memory)
diff --git a/Engine/Source/Runtime/Renderer/Private/ShadowRendering.cpp b/Engine/Source/Runtime/Renderer/Private/ShadowRendering.cpp
--- a/Engine/Source/Runtime/Renderer/Private/ShadowRendering.cpp
+++ b/Engine/Source/Runtime/Renderer/Private/ShadowRendering.cpp
@@ -1926,8 +1926,6 @@ void FDeferredShadingSceneRenderer::RenderShadowProjections(
 	if (VirtualShadowMaps.Num() > 0 || VisibleLightInfo.VirtualShadowMapClipmaps.Num() > 0)
 	{
 		check(!bProjectingForForwardShading);		// Not yet implemented/tested
-		check(Views.Num() == 1);					// TODO: Test multiview
-		
 
 		if (VirtualShadowMapArray.HasAnyShadowData())
 		{
@@ -1968,9 +1966,8 @@ void FDeferredShadingSceneRenderer::RenderShadowProjections(
 						bool bUseTemporal = false;
 						if (VisibleLightInfo.VirtualShadowMapClipmaps.Num() > 0)
 						{
-							check(VisibleLightInfo.VirtualShadowMapClipmaps.Num() == 1);
 							RenderVirtualShadowMapProjectionForDenoising(
-								VisibleLightInfo.VirtualShadowMapClipmaps[0],
+								VisibleLightInfo.FindShadowClipmapForView(&View),
 								GraphBuilder,
 								View,
 								VirtualShadowMapArray,
@@ -2021,9 +2018,8 @@ void FDeferredShadingSceneRenderer::RenderShadowProjections(
 					{
 						if (VisibleLightInfo.VirtualShadowMapClipmaps.Num() > 0)
 						{
-							check(VisibleLightInfo.VirtualShadowMapClipmaps.Num() == 1);
 							RenderVirtualShadowMapProjection(
-								VisibleLightInfo.VirtualShadowMapClipmaps[0],
+								VisibleLightInfo.FindShadowClipmapForView(&View),
 								GraphBuilder,
 								View,
 								VirtualShadowMapArray,
diff --git a/Engine/Source/Runtime/Renderer/Private/ShadowSetup.cpp b/Engine/Source/Runtime/Renderer/Private/ShadowSetup.cpp
--- a/Engine/Source/Runtime/Renderer/Private/ShadowSetup.cpp
+++ b/Engine/Source/Runtime/Renderer/Private/ShadowSetup.cpp
@@ -4459,7 +4459,8 @@ void FSceneRenderer::AddViewDependentWholeSceneShadowsForView(
 					LightSceneInfo,
 					WorldToLight,
 					View.ViewMatrices,
-					View.ViewRect.Size()
+					View.ViewRect.Size(),
+					&View
 				));
 
 				// TODO: Not clear we need both of these in this path, but keep it consistent for now
diff --git a/Engine/Source/Runtime/Renderer/Private/VirtualShadowMaps/VirtualShadowMapClipmap.cpp b/Engine/Source/Runtime/Renderer/Private/VirtualShadowMaps/VirtualShadowMapClipmap.cpp
--- a/Engine/Source/Runtime/Renderer/Private/VirtualShadowMaps/VirtualShadowMapClipmap.cpp
+++ b/Engine/Source/Runtime/Renderer/Private/VirtualShadowMaps/VirtualShadowMapClipmap.cpp
@@ -44,8 +44,10 @@ FVirtualShadowMapClipmap::FVirtualShadowMapClipmap(
 	const FLightSceneInfo& InLightSceneInfo,
 	const FMatrix& WorldToLightRotationMatrix,
 	const FViewMatrices& CameraViewMatrices,
-	FIntPoint CameraViewRectSize)
-	: LightSceneInfo(InLightSceneInfo)
+	FIntPoint CameraViewRectSize,
+	const FViewInfo* InDependentView)
+	: LightSceneInfo(InLightSceneInfo),
+	  DependentView(InDependentView)
 {
 	check(WorldToLightRotationMatrix.GetOrigin() == FVector(0, 0, 0));	// Should not contain translation or scaling
 
diff --git a/Engine/Source/Runtime/Renderer/Private/VirtualShadowMaps/VirtualShadowMapClipmap.h b/Engine/Source/Runtime/Renderer/Private/VirtualShadowMaps/VirtualShadowMapClipmap.h
--- a/Engine/Source/Runtime/Renderer/Private/VirtualShadowMaps/VirtualShadowMapClipmap.h
+++ b/Engine/Source/Runtime/Renderer/Private/VirtualShadowMaps/VirtualShadowMapClipmap.h
@@ -23,7 +23,8 @@ public:
 		const FLightSceneInfo& InLightSceneInfo,
 		const FMatrix& WorldToLightRotationMatrix,
 		const FViewMatrices& CameraViewMatrices,
-		FIntPoint CameraViewRectSize
+		FIntPoint CameraViewRectSize,
+		const FViewInfo* InDependentView
 	);
 
 	FViewMatrices GetViewMatrices(int32 ClipmapIndex) const;
@@ -58,9 +59,17 @@ public:
 
 	float GetMaxRadius() const { return MaxRadius; }
 
+	const FViewInfo* GetDependentView() const { return DependentView; }
+
 private:
 	const FLightSceneInfo& LightSceneInfo;
 
+	/**
+	 * DependentView is the 'main' or visible geometry view that this view-dependent clipmap was created for. Should only be used to 
+	 * identify the view during shadow projection (note: this should be refactored to be more explicit instead).
+	 */
+	const FViewInfo* DependentView;
+
 	/** Origin of the clipmap in world space
 	* Usually aligns with the camera position from which it was created.
 	* Note that the centers of each of the levels can be different as they are snapped to page alignment at their respective scales
