U5 - Swap the order because this function is called from a Taskthread and IsInRenderingThread will ensure because a TaskThread is never the RenderingThread, but it had been this way in the past (when UE3 did not had task threads) and we try to detect other places where it was used incorrectly.

diff --git a/Engine/Source/Runtime/Renderer/Private/RendererScene.cpp b/Engine/Source/Runtime/Renderer/Private/RendererScene.cpp
--- a/Engine/Source/Runtime/Renderer/Private/RendererScene.cpp
+++ b/Engine/Source/Runtime/Renderer/Private/RendererScene.cpp
@@ -4345,7 +4345,7 @@ void FScene::FlushAsyncLightPrimitiveInteractionCreation() const
 
 bool FScene::IsPrimitiveBeingRemoved(FPrimitiveSceneInfo* PrimitiveSceneInfo) const
 {
-	check(IsInRenderingThread() || IsInParallelRenderingThread());
+	check(IsInParallelRenderingThread() || IsInRenderingThread());
 	return RemovedPrimitiveSceneInfos.Contains(PrimitiveSceneInfo);
 }
 
