Improve sampling integration between path tracer and movie render queue

AA Camera jittering is now automatically disabled when using the path tracer in MRQ

The exact sample count is passed through the View object to the path tracer to ensure the MRQ settings win over what might be configured in the post process volume or in cvars

Disable movement based path tracer invalidation when offline rendering is enabled

Rename some path tracer state variables for clarity and ensure different views can have their own target sample count

Fix off-by-one error in the path tracer progress bar display. It now dissapears on the last sample, which allows it to stay on even when rendering with MovieRenderQueue

diff --git a/Engine/Plugins/MovieScene/MovieRenderPipeline/Source/MovieRenderPipelineRenderPasses/Private/MoviePipelineDeferredPasses.cpp b/Engine/Plugins/MovieScene/MovieRenderPipeline/Source/MovieRenderPipelineRenderPasses/Private/MoviePipelineDeferredPasses.cpp
--- a/Engine/Plugins/MovieScene/MovieRenderPipeline/Source/MovieRenderPipelineRenderPasses/Private/MoviePipelineDeferredPasses.cpp
+++ b/Engine/Plugins/MovieScene/MovieRenderPipeline/Source/MovieRenderPipelineRenderPasses/Private/MoviePipelineDeferredPasses.cpp
@@ -761,8 +761,9 @@ namespace MoviePipeline
 			return;
 		}
 
-		// For the first sample in a new output, we allocate memory
-		if (NewPayload->IsFirstTile() && NewPayload->IsFirstTemporalSample())
+		// Allocate memory if the ImageAccumulator has not been initialized yet for this output
+		// This usually happens on the first sample (regular case), or on the last spatial sample of the first temporal sample (path tracer)
+		if (InParams.ImageAccumulator->NumChannels == 0)
 		{
 			int32 ChannelCount = InParams.bAccumulateAlpha ? 4 : 3;
 			InParams.ImageAccumulator->InitMemory(FIntPoint(NewPayload->SampleState.TileSize.X * NewPayload->SampleState.TileCounts.X, NewPayload->SampleState.TileSize.Y * NewPayload->SampleState.TileCounts.Y), ChannelCount);
diff --git a/Engine/Plugins/MovieScene/MovieRenderPipeline/Source/MovieRenderPipelineRenderPasses/Private/MoviePipelineImagePassBase.cpp b/Engine/Plugins/MovieScene/MovieRenderPipeline/Source/MovieRenderPipelineRenderPasses/Private/MoviePipelineImagePassBase.cpp
--- a/Engine/Plugins/MovieScene/MovieRenderPipeline/Source/MovieRenderPipelineRenderPasses/Private/MoviePipelineImagePassBase.cpp
+++ b/Engine/Plugins/MovieScene/MovieRenderPipeline/Source/MovieRenderPipelineRenderPasses/Private/MoviePipelineImagePassBase.cpp
@@ -141,14 +141,43 @@ TSharedPtr<FSceneViewFamilyContext> UMoviePipelineImagePassBase::CalculateViewFa
 
 	// Anti Aliasing
 	{
-		// If we're not using Temporal Anti-Aliasing we will apply the View Matrix projection jitter. Normally TAA sets this
-		// inside FSceneRenderer::PreVisibilityFrameSetup..
-		if (View->AntiAliasingMethod != EAntiAliasingMethod::AAM_TemporalAA)
+		// If we're not using Temporal Anti-Aliasing or Path Tracing we will apply the View Matrix projection jitter. Normally TAA sets this
+		// inside FSceneRenderer::PreVisibilityFrameSetup. Path Tracing does its own anti-aliasing internally.
+		if (View->AntiAliasingMethod != EAntiAliasingMethod::AAM_TemporalAA && !OutViewFamily->EngineShowFlags.PathTracing)
 		{
 			View->ViewMatrices.HackAddTemporalAAProjectionJitter(InOutSampleState.ProjectionMatrixJitterAmount);
 		}
 	}
 
+	// Path Tracer Sampling
+	if (OutViewFamily->EngineShowFlags.PathTracing)
+	{
+		// override whatever settings came from PostProcessVolume or Camera
+		int32 SampleCount = InOutSampleState.TemporalSampleCount * InOutSampleState.SpatialSampleCount;
+		int32 SampleIndex = InOutSampleState.TemporalSampleIndex * InOutSampleState.SpatialSampleCount + InOutSampleState.SpatialSampleIndex;
+
+		// TODO: pass along FrameIndex (which includes SampleIndex) to make sure sampling is fully deterministic
+
+		// Overwrite whatever sampling count came from the PostProcessVolume
+		View->FinalPostProcessSettings.bOverride_PathTracingSamplesPerPixel = true;
+		View->FinalPostProcessSettings.PathTracingSamplesPerPixel = InOutSampleState.SpatialSampleCount;
+		if (InOutSampleState.SpatialSampleIndex == 0)
+		{
+			// reset path tracer's accumulation at the start of each temporal sample
+			View->bForceCameraVisibilityReset = true;
+		}
+
+		// discard the result, unless its the last spatial sample
+		InOutSampleState.bDiscardResult |= !(InOutSampleState.SpatialSampleIndex == InOutSampleState.SpatialSampleCount - 1);
+
+		if (InOutSampleState.GetTileCount() > 1 && InOutSampleState.GetTileIndex() == 0 && SampleIndex == 0)
+		{
+			// warn the user
+			UE_LOG(LogMovieRenderPipeline, Warning, TEXT("Path tracer does not support tiling correctly yet!"));
+			// FIXME: find a way to disable tiling earlier so this doesn't happen
+		}
+	}
+
 	// Object Occlusion/Histories
 	{
 		// If we're using tiling, we force the reset of histories each frame so that we don't use the previous tile's
diff --git a/Engine/Plugins/MovieScene/MovieRenderPipeline/Source/MovieRenderPipelineRenderPasses/Public/MoviePipelineDeferredPasses.h b/Engine/Plugins/MovieScene/MovieRenderPipeline/Source/MovieRenderPipelineRenderPasses/Public/MoviePipelineDeferredPasses.h
--- a/Engine/Plugins/MovieScene/MovieRenderPipeline/Source/MovieRenderPipelineRenderPasses/Public/MoviePipelineDeferredPasses.h
+++ b/Engine/Plugins/MovieScene/MovieRenderPipeline/Source/MovieRenderPipelineRenderPasses/Public/MoviePipelineDeferredPasses.h
@@ -243,6 +243,8 @@ public:
 		OutViewModeIndex = EViewModeIndex::VMI_PathTracing;
 	}
 	virtual int32 GetOutputFileSortingOrder() const override { return 2; }
+
+	virtual bool IsAntiAliasingSupported() const { return false; }
 };
 
 struct MOVIERENDERPIPELINERENDERPASSES_API FAccumulatorPool : public TSharedFromThis<FAccumulatorPool>
diff --git a/Engine/Shaders/Private/PathTracing/PathTracingCompositingPixelShader.usf b/Engine/Shaders/Private/PathTracing/PathTracingCompositingPixelShader.usf
--- a/Engine/Shaders/Private/PathTracing/PathTracingCompositingPixelShader.usf
+++ b/Engine/Shaders/Private/PathTracing/PathTracingCompositingPixelShader.usf
@@ -48,7 +48,7 @@ void CompositeMain(
 
 	OutColor = float4(Radiance.rgb * View.PreExposure, 1.0);
 
-	if (ProgressDisplayEnabled && Iteration < MaxSamples)
+	if (ProgressDisplayEnabled && (Iteration + 1) < MaxSamples)
 	{
 		// Draw a progress meter for how far along we are to the max number of passes
 		float Aspect = View.BufferSizeAndInvSize.y * View.BufferSizeAndInvSize.z;
@@ -60,6 +60,7 @@ void CompositeMain(
 		// center progress in middle-bottom of the frame
 		float2 Center = float2(lerp(RectMin.x, RectMax.x, 0.5), RectMax.y) * View.BufferSizeAndInvSize.zz - float2(0, Size * 0.75);
 		P -= Center;
+		float Fraction = saturate(float(Iteration + 1) / float(MaxSamples));
 #if 0
 		// measure progress like a clock
 		float R2 = dot(P, P);
@@ -71,7 +72,6 @@ void CompositeMain(
 		// point is on the progress ring
 		if (R2 < ROuter2 + R2f && R2 > RInner2 - R2f)
 		{
-			float Fraction = saturate(float(Iteration) / float(MaxSamples));
 			float Angle = (atan2(-P.x, P.y) + PI) / (2 * PI);
 			// anti-alias progress bar edge
 			float Blend = smoothstep(Fraction - fwidth(Angle), Fraction + fwidth(Angle), Angle);
@@ -89,7 +89,6 @@ void CompositeMain(
 		{
 			P += ProgressSize;
 			P /= ProgressSize * 2;
-			float Fraction = saturate(float(Iteration) / float(MaxSamples));
 			float Blend = smoothstep(Fraction - fwidth(P.x), Fraction + fwidth(P.x), P.x);
 			// Grey/Red blend
 			float4 ProgressBarColor = lerp(float4(1.0, 0.1, 0.1, 1), float4(0.5, 0.5, 0.5, 1), Blend);
diff --git a/Engine/Source/Runtime/Renderer/Private/PathTracing.cpp b/Engine/Source/Runtime/Renderer/Private/PathTracing.cpp
--- a/Engine/Source/Runtime/Renderer/Private/PathTracing.cpp
+++ b/Engine/Source/Runtime/Renderer/Private/PathTracing.cpp
@@ -452,7 +452,7 @@ void FSceneViewState::PathTracingInvalidate()
 	PathTracingRadianceRT.SafeRelease();
 	VarianceMipTreeDimensions = FIntVector(0);
 	TotalRayCount = 0;
-	PathTracingSPP = 0;
+	PathTracingSampleIndex = 0;
 }
 
 DECLARE_GPU_STAT_NAMED(Stat_GPU_PathTracing, TEXT("Path Tracing"));
@@ -469,13 +469,14 @@ void FDeferredShadingSceneRenderer::RenderPathTracing(
 	bool bArgsChanged = false;
 
 	// Get current value of MaxSPP and reset render if it has changed
-	int32 SamplesPerPixelCVar = CVarPathTracingSamplesPerPixel.GetValueOnRenderThread();
+	// NOTE: we ignore the CVar when using offline rendering
+	int32 SamplesPerPixelCVar = View.bIsOfflineRender ? -1 : CVarPathTracingSamplesPerPixel.GetValueOnRenderThread();
 	uint32 MaxSPP = SamplesPerPixelCVar > -1 ? SamplesPerPixelCVar : View.FinalPostProcessSettings.PathTracingSamplesPerPixel;
 	MaxSPP = FMath::Max(MaxSPP, 1u);
-	static uint32 PreviousMaxSPP = MaxSPP;
-	if (PreviousMaxSPP != MaxSPP)
+	if (View.ViewState->PathTracingTargetSPP != MaxSPP)
 	{
-		PreviousMaxSPP = MaxSPP;
+		// Store MaxSPP in the view state because we may have multiple views, each targetting a different sample count
+		View.ViewState->PathTracingTargetSPP = MaxSPP;
 		bArgsChanged = true;
 	}
 	// Changing FrameIndependentTemporalSeed requires starting over
@@ -511,14 +512,14 @@ void FDeferredShadingSceneRenderer::RenderPathTracing(
 	if (LockedSamplingPattern)
 	{
 		// Count samples from 0 for deterministic results
-		PathTracingData.TemporalSeed = View.ViewState->PathTracingSPP;
+		PathTracingData.TemporalSeed = View.ViewState->PathTracingSampleIndex;
 	}
 	else
 	{
 		// Count samples from an ever-increasing counter to avoid screen-door effect
-		PathTracingData.TemporalSeed = View.ViewState->PathTracingFrameIndependentTemporalSeed;
+		PathTracingData.TemporalSeed = View.ViewState->PathTracingFrameIndex;
 	}
-	PathTracingData.Iteration = View.ViewState->PathTracingSPP;
+	PathTracingData.Iteration = View.ViewState->PathTracingSampleIndex;
 	PathTracingData.MaxSamples = MaxSPP;
 
 	// Prepare radiance buffer (will be shared with display pass)
@@ -562,7 +563,7 @@ void FDeferredShadingSceneRenderer::RenderPathTracing(
 		TShaderMapRef<FPathTracingRG> RayGenShader(View.ShaderMap);
 		ClearUnusedGraphResources(RayGenShader, PassParameters);
 		GraphBuilder.AddPass(
-			RDG_EVENT_NAME("Path Tracer Compute (%d x %d) Sample=%d/%d", View.ViewRect.Size().X, View.ViewRect.Size().Y, View.ViewState->PathTracingSPP, MaxSPP),
+			RDG_EVENT_NAME("Path Tracer Compute (%d x %d) Sample=%d/%d", View.ViewRect.Size().X, View.ViewRect.Size().Y, View.ViewState->PathTracingSampleIndex, MaxSPP),
 			PassParameters,
 			ERDGPassFlags::Compute,
 			[PassParameters, RayGenShader, &View](FRHICommandListImmediate& RHICmdList)
@@ -619,8 +620,8 @@ void FDeferredShadingSceneRenderer::RenderPathTracing(
 	);
 
 	// Bump counters for next frame
-	++View.ViewState->PathTracingSPP;
-	++View.ViewState->PathTracingFrameIndependentTemporalSeed;
+	++View.ViewState->PathTracingSampleIndex;
+	++View.ViewState->PathTracingFrameIndex;
 }
 
 #endif
diff --git a/Engine/Source/Runtime/Renderer/Private/RendererScene.cpp b/Engine/Source/Runtime/Renderer/Private/RendererScene.cpp
--- a/Engine/Source/Runtime/Renderer/Private/RendererScene.cpp
+++ b/Engine/Source/Runtime/Renderer/Private/RendererScene.cpp
@@ -240,8 +240,9 @@ FSceneViewState::FSceneViewState()
 	VarianceMipTreeDimensions = FIntVector(0);
 	VarianceMipTree = new FRWBuffer;
 	PathTracingRect = FIntRect(0, 0, 0, 0);
-	PathTracingSPP = 0;
-	PathTracingFrameIndependentTemporalSeed = 0;
+	PathTracingTargetSPP = 0;
+	PathTracingSampleIndex = 0;
+	PathTracingFrameIndex = 0;
 	TotalRayCount = 0;
 	TotalRayCountBuffer = new FRWBuffer;
 	if (GetFeatureLevel() >= ERHIFeatureLevel::SM5)
diff --git a/Engine/Source/Runtime/Renderer/Private/ScenePrivate.h b/Engine/Source/Runtime/Renderer/Private/ScenePrivate.h
--- a/Engine/Source/Runtime/Renderer/Private/ScenePrivate.h
+++ b/Engine/Source/Runtime/Renderer/Private/ScenePrivate.h
@@ -1000,15 +1000,16 @@ public:
 	// Reference path tracing cached results
 	TRefCountPtr<IPooledRenderTarget> PathTracingRadianceRT;
 	FIntRect PathTracingRect;
+	uint32   PathTracingTargetSPP;
 	FRWBuffer* VarianceMipTree;
 	FIntVector VarianceMipTreeDimensions;
 
 	// Path tracer samples/pixel
-	uint32_t PathTracingSPP;
+	uint32 PathTracingSampleIndex;
 
-	// The frame independent temporal seed, not reset on invalidation unlike PathTracingSPP to avoid
+	// Path tracer frame index, not reset on invalidation unlike PathTracingSampleIndex to avoid
 	// the "screen door" effect and reduce temporal aliasing
-	uint32_t PathTracingFrameIndependentTemporalSeed;
+	uint32_t PathTracingFrameIndex;
 
 	// Path tracer ray counter
 	uint32 TotalRayCount;
diff --git a/Engine/Source/Runtime/Renderer/Private/SceneVisibility.cpp b/Engine/Source/Runtime/Renderer/Private/SceneVisibility.cpp
--- a/Engine/Source/Runtime/Renderer/Private/SceneVisibility.cpp
+++ b/Engine/Source/Runtime/Renderer/Private/SceneVisibility.cpp
@@ -3428,17 +3428,27 @@ void FSceneRenderer::PreVisibilityFrameSetup(FRDGBuilder& GraphBuilder, const FS
 			
 #if RHI_RAYTRACING
 			// Note: 0.18 deg is the minimum angle for avoiding numerical precision issue (which would cause constant invalidation)
-			const bool bIsThereALargeMomvement= IsLargeCameraMovement(
-				View, ViewState->PrevFrameViewInfo.ViewMatrices.GetViewMatrix(),
+			const bool bIsCameraMove = IsLargeCameraMovement(
+				View,
+				ViewState->PrevFrameViewInfo.ViewMatrices.GetViewMatrix(),
 				ViewState->PrevFrameViewInfo.ViewMatrices.GetViewOrigin(),
 				0.18f /*degree*/, 0.1f /*cm*/);
 			const bool bIsProjMatrixDifferent = View.ViewMatrices.GetProjectionNoAAMatrix() != View.ViewState->PrevFrameViewInfo.ViewMatrices.GetProjectionNoAAMatrix();
-			const bool bInvalidatePathTracer = (
-				bResetCamera ||
-				Scene->bPathTracingNeedsInvalidation ||
-				bIsProjMatrixDifferent ||
-				bIsThereALargeMomvement
-			);
+			bool bInvalidatePathTracer = false;
+			
+			if (View.bIsOfflineRender)
+			{
+				// for offline renders, we want to allow motion blur so don't let geometry changes or camera moves invalidate the path tracer accumulation
+				bInvalidatePathTracer = View.bCameraCut || View.bForceCameraVisibilityReset;
+			}
+			else
+			{
+				// for interactive usage - any movement or scene change should restart the path tracer
+				bInvalidatePathTracer = bResetCamera ||
+					Scene->bPathTracingNeedsInvalidation ||
+					bIsProjMatrixDifferent ||
+					bIsCameraMove;
+			}
 
 			if (bInvalidatePathTracer)
 			{
