Rename path tracer resources with "PathTracer." for easier filtering

diff --git a/Engine/Source/Runtime/Renderer/Private/PathTracing.cpp b/Engine/Source/Runtime/Renderer/Private/PathTracing.cpp
--- a/Engine/Source/Runtime/Renderer/Private/PathTracing.cpp
+++ b/Engine/Source/Runtime/Renderer/Private/PathTracing.cpp
@@ -361,7 +361,7 @@ void SetLightParameters(FRDGBuilder& GraphBuilder, FPathTracingRG::FParameters*
 	{
 		// Upload the buffer of lights to the GPU
 		size_t DataSize = sizeof(FPathTracingLight) * FMath::Max(LightCount, 1u);
-		PassParameters->SceneLights = GraphBuilder.CreateSRV(FRDGBufferSRVDesc(CreateStructuredBuffer(GraphBuilder, TEXT("PathTracingLightsBuffer"), sizeof(FPathTracingLight), FMath::Max(LightCount, 1u), Lights, DataSize)));
+		PassParameters->SceneLights = GraphBuilder.CreateSRV(FRDGBufferSRVDesc(CreateStructuredBuffer(GraphBuilder, TEXT("PathTracer.LightsBuffer"), sizeof(FPathTracingLight), FMath::Max(LightCount, 1u), Lights, DataSize)));
 		PassParameters->SceneLightCount = LightCount;
 	}
 
@@ -388,7 +388,7 @@ void SetLightParameters(FRDGBuilder& GraphBuilder, FPathTracingRG::FParameters*
 			FClearValueBinding::None,
 			TexCreate_ShaderResource | TexCreate_UAV,
 			NumSlices);
-		FRDGTexture* IESTexture = GraphBuilder.CreateTexture(IESTextureDesc, TEXT("PathTracerIESAtlas"), ERDGTextureFlags::None);
+		FRDGTexture* IESTexture = GraphBuilder.CreateTexture(IESTextureDesc, TEXT("PathTracer.IESAtlas"), ERDGTextureFlags::None);
 
 		for (auto&& Entry : IESLightProfilesMap)
 		{
@@ -527,7 +527,7 @@ void FDeferredShadingSceneRenderer::RenderPathTracing(
 	if (View.ViewState->PathTracingRadianceRT)
 	{
 		// we already have a valid radiance texture, re-use it
-		RadianceTexture = GraphBuilder.RegisterExternalTexture(View.ViewState->PathTracingRadianceRT, TEXT("PathTracerRadiance"));
+		RadianceTexture = GraphBuilder.RegisterExternalTexture(View.ViewState->PathTracingRadianceRT, TEXT("PathTracer.Radiance"));
 	}
 	else
 	{
@@ -537,7 +537,7 @@ void FDeferredShadingSceneRenderer::RenderPathTracing(
 			PF_A32B32G32R32F,
 			FClearValueBinding::None,
 			TexCreate_ShaderResource | TexCreate_UAV);
-		RadianceTexture = GraphBuilder.CreateTexture(RadianceTextureDesc, TEXT("PathTracerRadiance"), ERDGTextureFlags::MultiFrame);
+		RadianceTexture = GraphBuilder.CreateTexture(RadianceTextureDesc, TEXT("PathTracer.Radiance"), ERDGTextureFlags::MultiFrame);
 	}
 	bool NeedsMoreRays = PathTracingData.Iteration < MaxSPP;
 
