Implements Gen5 TAA's rejection post filtering pass and history's parrallax factor rejection

diff --git a/Engine/Shaders/Private/TemporalAA/TAACommon.ush b/Engine/Shaders/Private/TemporalAA/TAACommon.ush
--- a/Engine/Shaders/Private/TemporalAA/TAACommon.ush
+++ b/Engine/Shaders/Private/TemporalAA/TAACommon.ush
@@ -38,6 +38,8 @@
 
 #define PARALLAX_REJECTION_MASK_THRESHOLD 0.5
 
+#define MAX_PARALLAX_FACTOR 8.0
+
 // K = Center of the nearest input pixel.
 // O = Center of the output pixel.
 //
@@ -70,6 +72,26 @@ static const int2 kOffsets3x3[9] =
 	int2(1,  1),
 };
 
+// T = Center of the nearest top left pixel input pixel.
+// O = Center of the output pixel.
+//
+//          | 
+//    T     |     .
+//          | 
+//       O  | 
+//  --------+--------
+//          | 
+//          | 
+//    .     |     .
+//          | 
+static const int2 Offsets2x2[4] =
+{
+	int2( 0,  0), // T
+	int2( 1,  0),
+	int2( 0,  1),
+	int2( 1,  1),
+};
+
 // Indexes of the 3x3 square.
 static const uint kSquareIndexes3x3[9] = { 4, 0, 1, 2, 3, 5, 6, 7, 8 };
 
@@ -287,50 +309,6 @@ bool TakeOnlyOneSamplePair(float2 Offset)
 	return Offset.y > 0.0 || (Offset.x > 0.0 && Offset.y == 0.0);
 }
 
-
-// For motion vector, use camera/dynamic motion from min depth pixel in pattern around pixel.
-// This enables better quality outline on foreground against different motion background.
-// Larger 2 pixel distance "x" works best (because AA dilates surface).
-void FindClosestDepthOffset(Texture2D Texture, uint2 PixelPos, float DeviceZ, int PixelRadius, out float2 VelocityPixelOffset, out float ClosestDeviceZ)
-{
-	float4 Depths;
-	Depths.x = Texture[int2(PixelPos) + int2(-PixelRadius, -PixelRadius)].r;
-	Depths.y = Texture[int2(PixelPos) + int2( PixelRadius, -PixelRadius)].r;
-	Depths.z = Texture[int2(PixelPos) + int2(-PixelRadius,  PixelRadius)].r;
-	Depths.w = Texture[int2(PixelPos) + int2( PixelRadius,  PixelRadius)].r;
-
-	float2 DepthOffset = float2(PixelRadius, PixelRadius);
-	float DepthOffsetXx = float(PixelRadius);
-	// Nearest depth is the largest depth (depth surface 0=far, 1=near).
-	if(Depths.x > Depths.y) 
-	{
-		DepthOffsetXx = -PixelRadius;
-	}
-	if(Depths.z > Depths.w) 
-	{
-		DepthOffset.x = -PixelRadius;
-	}
-	float DepthsXY = max(Depths.x, Depths.y);
-	float DepthsZW = max(Depths.z, Depths.w);
-	if (DepthsXY > DepthsZW) 
-	{
-		DepthOffset.y = -PixelRadius;
-		DepthOffset.x = DepthOffsetXx; 
-	}
-	float DepthsXYZW = max(DepthsXY, DepthsZW);
-	
-	ClosestDeviceZ = DeviceZ;
-	VelocityPixelOffset = 0.0;
-
-	// If the pixel is closer in neighborhood, take it.
-	FLATTEN
-	if(DepthsXYZW > DeviceZ)
-	{
-		VelocityPixelOffset = DepthOffset;
-		ClosestDeviceZ = DepthsXYZW;
-	}
-}
-
 float2 ComputeStaticVelocity(float2 ScreenPos, float DeviceZ)
 {
 	float3 PosN = float3(ScreenPos, DeviceZ);
@@ -385,3 +363,82 @@ taa_half MeasureRejectionFactor(taa_half3 PrevYCoCg, taa_half3 ClampedPrevYCoCg,
 	
 	return Rejection;
 }
+
+#define BITS_PER_SUB_PIXEL 4
+#define SUB_PIXEL_MASK ((1 << BITS_PER_SUB_PIXEL) - 1)
+
+#define SUB_PIXEL_GRID_SIZE 2
+#define SUB_PIXEL_COUNT (SUB_PIXEL_GRID_SIZE * SUB_PIXEL_GRID_SIZE)
+
+struct FSubpixelNeighborhood
+{
+	// 2x2 quad of subpixels grid
+	uint4 SampleArray;
+
+	// Offset of subpixel
+	uint2 SubpixelOffset;
+};
+
+FSubpixelNeighborhood GatherPrevSubpixelNeighborhood(Texture2D<uint> PrevSubPixelDetailTexture, float2 PrevHistoryBufferUV)
+{
+	float2 PixelCoord = PrevHistoryBufferUV * PrevHistoryInfo_Extent;
+	float2 TopLeftPixelCoord = floor(PixelCoord - 0.5) + 0.5;
+	
+	FSubpixelNeighborhood Neighborhood;
+	#if 1 // TODO
+	{
+		Neighborhood.SampleArray[0] = PrevSubPixelDetailTexture[uint2(TopLeftPixelCoord + float2(0.0, 0.0))];
+		Neighborhood.SampleArray[1] = PrevSubPixelDetailTexture[uint2(TopLeftPixelCoord + float2(1.0, 0.0))];
+		Neighborhood.SampleArray[2] = PrevSubPixelDetailTexture[uint2(TopLeftPixelCoord + float2(0.0, 1.0))];
+		Neighborhood.SampleArray[3] = PrevSubPixelDetailTexture[uint2(TopLeftPixelCoord + float2(1.0, 1.0))];
+	}
+	#else
+	{
+		// (-,+),(+,+),(+,-),(-,-),
+		/**
+		 *  3 2
+		 *  0 1
+		 */
+		uint4 Samples = PrevSubPixelDetailTexture.Gather(GlobalPointClampedSampler, PrevHistoryBufferUV);
+
+		Neighborhood.SampleArray[0] = Samples[3]; // 00
+		Neighborhood.SampleArray[1] = Samples[2]; // 10
+		Neighborhood.SampleArray[2] = Samples[0]; // 01
+		Neighborhood.SampleArray[3] = Samples[1]; // 11
+	}
+	#endif
+
+	taa_half2 Interp = taa_half2(PixelCoord - TopLeftPixelCoord);
+	Neighborhood.SubpixelOffset = uint2(round(Interp * 2.0));
+
+	return Neighborhood;
+}
+
+uint GetSubpixelPayload(FSubpixelNeighborhood Neighborhood, uint SubpixelId)
+{
+	const uint2 SubpixelCoord = uint2(SubpixelId % SUB_PIXEL_GRID_SIZE, SubpixelId / SUB_PIXEL_GRID_SIZE);
+
+	uint2 InputSubpixelCoord = Neighborhood.SubpixelOffset + SubpixelCoord;
+	uint2 InputTexelCoord = InputSubpixelCoord / SUB_PIXEL_GRID_SIZE;
+	uint2 QuadSubpixelCoord = InputSubpixelCoord % SUB_PIXEL_GRID_SIZE;
+
+
+	uint QuadSampleId = dot(InputTexelCoord, uint2(1, 2));
+	uint InputSample = Neighborhood.SampleArray[QuadSampleId];
+
+	uint BitShift = BITS_PER_SUB_PIXEL * dot(QuadSubpixelCoord, uint2(1, SUB_PIXEL_GRID_SIZE));
+
+	uint SubpixelPayload = (InputSample >> BitShift) & SUB_PIXEL_MASK;
+
+	return SubpixelPayload;
+}
+
+uint CompressParallaxFactor(taa_half ParallaxFactor)
+{
+	return clamp(uint((ParallaxFactor - taa_half(1.0)) * taa_half(float(SUB_PIXEL_MASK) / MAX_PARALLAX_FACTOR)), 0, SUB_PIXEL_MASK);
+}
+
+taa_half UncompressParallaxFactor(uint CompressedParallaxFactor)
+{
+	return taa_half(CompressedParallaxFactor & SUB_PIXEL_MASK) * taa_half(MAX_PARALLAX_FACTOR / float(SUB_PIXEL_MASK)) + taa_half(1.0);
+}
diff --git a/Engine/Shaders/Private/TemporalAA/TAADecimateHistory.usf b/Engine/Shaders/Private/TemporalAA/TAADecimateHistory.usf
--- a/Engine/Shaders/Private/TemporalAA/TAADecimateHistory.usf
+++ b/Engine/Shaders/Private/TemporalAA/TAADecimateHistory.usf
@@ -12,6 +12,7 @@
 
 //------------------------------------------------------- PARAMETERS
 
+float4x4 RotationalClipToPrevClip;
 float3 OutputQuantizationError;
 float HistoryPreExposureCorrection;
 uint bCameraCut;
@@ -20,10 +21,12 @@ Texture2D<float2> DilatedVelocityTexture;
 Texture2D<float> ClosestDepthTexture;
 Texture2D<uint> PrevUseCountTexture;
 Texture2D<uint> PrevClosestDepthTexture;
+Texture2D<uint> ParallaxFactorTexture;
 
 Texture2D<taa_half3> PrevHistory_Textures_0;
 Texture2D<taa_half3> PrevHistory_Textures_1;
 Texture2D<taa_half2> PrevHistory_Textures_2;
+Texture2D<uint> PrevHistory_Textures_3;
 
 RWTexture2D<taa_half3> PredictionSceneColorOutput;
 RWTexture2D<taa_half> ParallaxRejectionMaskOutput;
@@ -60,11 +63,13 @@ void MainCS(
 	
 	float2 EncodedVelocity = DilatedVelocityTexture[InputPixelPos];
 	float DeviceZ = ClosestDepthTexture[InputPixelPos];
+	taa_half ParallaxFactor = UncompressParallaxFactor((ParallaxFactorTexture[InputPixelPos] >> 4) & 0xF);
+
 
 	float WorldDepth = ConvertFromDeviceZ(DeviceZ);
 	
 	// Compute the parralax rejection mask
-	float ParallaxRejectionMask = 0.0;
+	taa_half ParallaxRejectionMask = 0.0;
 	{
 		float2 PrevScreenPos = ScreenPos - DecodeVelocityFromTexture(float4(EncodedVelocity, 0.0, 0.0)).xy;
 		float2 PrevInputBufferUV = (InputInfo_ScreenPosToViewportScale * PrevScreenPos + InputInfo_ScreenPosToViewportBias) * InputInfo_ExtentInverse;
@@ -83,7 +88,7 @@ void MainCS(
 			float HistoryClosestDeviceZ = f16tof32(PrevClosestDepthTexture[PixelPos]);
 			float HistoryClosestWorldDepth = ConvertFromDeviceZ(HistoryClosestDeviceZ);
 
-			float DepthRejection;
+			taa_half DepthRejection;
 			{
 				const float WorldDepthToPixelWorldRadius = GetTanHalfFieldOfView().x * InputInfo_ExtentInverse.x;
 				const float PixelDepthError = 3.0;
@@ -92,14 +97,14 @@ void MainCS(
 				float WorldDepthEpsilon = (WorldDepthToPixelWorldRadius * RadiusToDiamater * PixelDepthError) * HistoryClosestWorldDepth; 
 				float DeltaDepth = abs(HistoryClosestWorldDepth - WorldDepth);
 
-				DepthRejection = saturate(2.0 - DeltaDepth / WorldDepthEpsilon);
+				DepthRejection = saturate(taa_half(2.0) - taa_half(DeltaDepth / WorldDepthEpsilon));
 			}
 
-			float HistoryGhostingRejection = saturate(1.0 - 3.0 * abs(HistoryUseCount * rcp(256) - 1.0));
+			taa_half HistoryGhostingRejection = saturate(taa_half(1.0) - taa_half(3.0) * abs(taa_half(HistoryUseCount) * rcp(taa_half(256)) - taa_half(1.0)));
 
 			FLATTEN
 			if (bValidPixelPos)
-				ParallaxRejectionMask += BilinearWeight * max(DepthRejection, HistoryGhostingRejection);
+				ParallaxRejectionMask += taa_half(BilinearWeight) * max(DepthRejection, HistoryGhostingRejection);
 		}
 	}
 
@@ -121,8 +126,10 @@ void MainCS(
 	// Vector in pixel between pixel K -> O.
 	taa_half2 dKO = taa_half2(HistoryPPCo - HistoryPPCk);
 	
+	// Sample the history
 	taa_half3 DecimatedColor;
 	taa_half HistoryValidity;
+	taa_half HistoryMaxRejectionFactor;
 	{
 		float2 KernelHistoryBufferUV = clamp(
 			HistoryPPCk * PrevHistoryInfo_ExtentInverse,
@@ -130,8 +137,39 @@ void MainCS(
 			PrevHistoryInfo_UVViewportBilinearMax);
 
 		DecimatedColor = PrevHistory_Textures_0.SampleLevel(GlobalBilinearClampedSampler, KernelHistoryBufferUV, 0);
+
+		FSubpixelNeighborhood SubpixelNeigborhood = GatherPrevSubpixelNeighborhood(PrevHistory_Textures_3, KernelHistoryBufferUV);
+		{
+			uint EncodedHistoryMaxRejectionFactor = 0;
+			UNROLL_N(SUB_PIXEL_COUNT)
+			for (uint SubpixelId = 0; SubpixelId < SUB_PIXEL_COUNT; SubpixelId++)
+			{
+				uint SubpixelPayload = GetSubpixelPayload(SubpixelNeigborhood, SubpixelId);
+				EncodedHistoryMaxRejectionFactor = max(EncodedHistoryMaxRejectionFactor, SubpixelPayload);
+			}
+
+			HistoryMaxRejectionFactor = UncompressParallaxFactor(EncodedHistoryMaxRejectionFactor);
+		}
 	}
+	
+	// Reject subpixel details that might be in history but not input just based on their parrallax
+	{
+		taa_half HistorySubdetailRejectionFactor = max(HistoryMaxRejectionFactor / ParallaxFactor, taa_half(1.0));
 
+		float4 ThisClip = float4(ScreenPos, DeviceZ, 1);
+		float4 PrevClip = mul(ThisClip, View.ClipToPrevClip);
+		float2 PrevScreen = PrevClip.xy / PrevClip.w;
+			
+		float4 RotationalPrevClip = mul(ThisClip, RotationalClipToPrevClip);
+		float2 RotationalPrevScreen = RotationalPrevClip.xy / RotationalPrevClip.w;
+			
+		taa_half PixelParallaxVelocity = taa_half(0.5) * length(taa_half2(RotationalPrevScreen - PrevScreen) * InputInfo_Extent);
+		taa_half HistorySubdetailAdditionalParralax = (HistorySubdetailRejectionFactor - taa_half(1.0)) * PixelParallaxVelocity;
+
+		taa_half HistoryParallaxRejectionMask = saturate(taa_half(1.0) - HistorySubdetailAdditionalParralax);
+
+		ParallaxRejectionMask = min(ParallaxRejectionMask, HistoryParallaxRejectionMask);
+	}
 
 	#if CONFIG_ENABLE_STOCASTIC_QUANTIZATION
 	{
diff --git a/Engine/Shaders/Private/TemporalAA/TAADilateVelocity.usf b/Engine/Shaders/Private/TemporalAA/TAADilateVelocity.usf
--- a/Engine/Shaders/Private/TemporalAA/TAADilateVelocity.usf
+++ b/Engine/Shaders/Private/TemporalAA/TAADilateVelocity.usf
@@ -19,6 +19,22 @@ RWTexture2D<float2> DilatedVelocityOutput;
 RWTexture2D<float> ClosestDepthOutput;
 RWTexture2D<uint> PrevUseCountOutput;
 RWTexture2D<uint> PrevClosestDepthOutput;
+RWTexture2D<uint> ParallaxFactorOutput;
+
+
+//------------------------------------------------------- ENTRY POINT
+
+uint PackDepthAndOffset(const int2 Offset, float DeviceZ)
+{
+	return ((asuint(DeviceZ) << 2) & ~0xF) | (uint(1 + Offset.x) << 0) | (uint(1 + Offset.y) << 2);
+}
+
+void UnpackDepthAndOffset(uint Packed, out int2 Offset, out float DeviceZ)
+{
+	DeviceZ = asfloat((Packed & ~0xF) >> 2);
+	Offset.x = int((Packed >> 0) & 0x3) - 1;
+	Offset.y = int((Packed >> 2) & 0x3) - 1;
+}
 
 
 //------------------------------------------------------- ENTRY POINT
@@ -44,11 +60,48 @@ void MainCS(
 
 	float ClosestDeviceZ = DeviceZ;
 	float2 VelocityPixelOffset = float2(0.0, 0.0);
+	float ParallaxFactorForHistory = 1.0;
+	float ParallaxFactorForHistoryRejection = 1.0;
 	#if AA_CROSS
 	{
-		FindClosestDepthOffset(SceneDepthTexture, InputPixelPos, DeviceZ, AA_CROSS, /* out */ VelocityPixelOffset, /* out */ ClosestDeviceZ);
+		const int2 kOffsets[4] = {
+			int2( 1,  0),
+			int2( 1,  1),
+			int2( 0,  1),
+			int2(-1,  1),
+		};
+
+		uint PackedDepthOffset = PackDepthAndOffset(/* Offset = */ int2(0, 0), DeviceZ);
+		
+		float FurthestPairDeviceZ = DeviceZ;
+		float FurthestDeviceZ = DeviceZ;
+
+		#if 1
+			UNROLL_N(4)
+			for (uint i = 0; i < 4; i++)
+		#else
+			UNROLL_N(2)
+			for (uint i = 1; i < 4; i+= 2)
+		#endif
+		{
+			const int2 Offset = kOffsets[i];
+
+			float SampleDeviceZ0 = SceneDepthTexture[int2(InputPixelPos) + Offset].r;
+			float SampleDeviceZ1 = SceneDepthTexture[int2(InputPixelPos) - Offset].r;
+			
+			PackedDepthOffset = max(PackedDepthOffset, PackDepthAndOffset(+Offset, SampleDeviceZ0));
+			PackedDepthOffset = max(PackedDepthOffset, PackDepthAndOffset(-Offset, SampleDeviceZ1));
+
+			FurthestPairDeviceZ = min(FurthestPairDeviceZ, max(SampleDeviceZ0, SampleDeviceZ1));
+			FurthestDeviceZ = min(FurthestDeviceZ, min(SampleDeviceZ0, SampleDeviceZ1));
+		}
+
+		ParallaxFactorForHistory = clamp(DeviceZ / FurthestPairDeviceZ, 0.0, MAX_PARALLAX_FACTOR);
+		ParallaxFactorForHistoryRejection = clamp(DeviceZ / FurthestDeviceZ, 0.0, MAX_PARALLAX_FACTOR);
+
+		UnpackDepthAndOffset(PackedDepthOffset, /* out */ VelocityPixelOffset, /* out */ ClosestDeviceZ);
 	}
-	#endif	// AA_CROSS
+	#endif
 	
 	// Camera motion for pixel.
 	float2 FinalScreenVelocity = 0.0;
@@ -106,6 +159,7 @@ void MainCS(
 	{
 		DilatedVelocityOutput[InputPixelPos] = EncodeVelocityToTexture(float3(FinalScreenVelocity, 0.0)).xy;
 		ClosestDepthOutput[InputPixelPos] = ClosestDeviceZ;
+		ParallaxFactorOutput[InputPixelPos] = CompressParallaxFactor(ParallaxFactorForHistory) | (CompressParallaxFactor(ParallaxFactorForHistoryRejection) << 4);
 
 		#if DEBUG_OUTPUT
 		{
diff --git a/Engine/Shaders/Private/TemporalAA/TAAPostfilterRejection.usf b/Engine/Shaders/Private/TemporalAA/TAAPostfilterRejection.usf
new file mode 100644
--- /dev/null
+++ b/Engine/Shaders/Private/TemporalAA/TAAPostfilterRejection.usf
@@ -0,0 +1,108 @@
+// Copyright Epic Games, Inc. All Rights Reserved.
+
+#include "TAACommon.ush"
+
+
+//------------------------------------------------------- CONFIG
+
+#define TILE_SIZE 8
+
+
+//------------------------------------------------------- PARAMETERS
+
+Texture2D<taa_half> HistoryRejectionTexture;
+RWTexture2D<taa_half> HistoryRejectionOutput;
+
+
+//------------------------------------------------------- FUNCTIONS
+
+// Operator that output, lowest, median and highest values from 3 input values.
+void LMHOperator(taa_half A, taa_half B, taa_half C, out taa_half L, out taa_half M, out taa_half H)
+{
+	// TODO: min3, max3, med3.
+	taa_half X = min(B, C);
+	taa_half Y = max(B, C);
+	
+	L = min(A, X);
+	taa_half Z = max(A, X);
+
+	M = min(Z, Y);
+	H = max(Z, Y);
+}
+
+// 3 samples median.
+taa_half Median(taa_half A, taa_half B, taa_half C)
+{
+	// TODO: med3.
+	taa_half L, M, H;
+	LMHOperator(A, B, C, L, M, H);
+	return M;
+}
+
+// 9 samples median using Smith1996.
+taa_half Median9(taa_half In[9])
+{
+	// First layer.
+	taa_half L0[3];
+	taa_half M0[3];
+	taa_half H0[3];
+	for (uint j = 0; j < 3; j++)
+	{
+		LMHOperator(In[j * 3 + 0], In[j * 3 + 1], In[j * 3 + 2], L0[j], M0[j], H0[j]);
+	}
+	
+	// Second layer.
+	// TODO: min3, max3.
+	taa_half M1[3];
+	M1[0] = max(max(L0[0], L0[1]), L0[2]);
+	M1[1] = Median(M0[0], M0[1], M0[2]);
+	M1[2] = min(min(H0[0], H0[1]), H0[2]);
+	
+	// Third layer.
+	return Median(M1[0], M1[1], M1[2]);
+}
+
+
+//------------------------------------------------------- ENTRY POINT
+
+[numthreads(TILE_SIZE, TILE_SIZE, 1)]
+void MainCS(
+	uint2 GroupId : SV_GroupID,
+	uint GroupThreadIndex : SV_GroupIndex)
+{
+	float4 Debug = 0.0;
+
+	uint2 DispatchThreadId = (
+		ZOrder2D(GroupThreadIndex, log2(TILE_SIZE)) +
+		GroupId * uint2(TILE_SIZE, TILE_SIZE));
+
+	uint2 RejectionPixelPos = DispatchThreadId;
+
+	float2 RejectionBufferUV = (RejectionPixelPos + 0.5) * RejectionInfo_ExtentInverse;
+
+	// Fetch neighborhood of the kernel
+	taa_half InSamples[9];
+	UNROLL
+	for (uint i = 0; i < 9; i++)
+	{
+		float2 Offset = float2(kOffsets3x3[i]);
+		
+		float2 SampleRejectionBufferUV = RejectionBufferUV + Offset * RejectionInfo_ExtentInverse;
+		SampleRejectionBufferUV = clamp(SampleRejectionBufferUV, RejectionInfo_UVViewportBilinearMin, RejectionInfo_UVViewportBilinearMax);
+
+		InSamples[i] = HistoryRejectionTexture.SampleLevel(GlobalPointClampedSampler, SampleRejectionBufferUV, 0);
+	}
+
+	taa_half Rejection = Median9(InSamples);
+
+	if (all(RejectionPixelPos < RejectionInfo_ViewportSize))
+	{
+		HistoryRejectionOutput[RejectionPixelPos] = Rejection;
+
+		#if DEBUG_OUTPUT
+		{
+			DebugOutput[RejectionPixelPos] = Debug;
+		}
+		#endif
+	}
+}
diff --git a/Engine/Shaders/Private/TemporalAA/TAAUpdateHistory.usf b/Engine/Shaders/Private/TemporalAA/TAAUpdateHistory.usf
--- a/Engine/Shaders/Private/TemporalAA/TAAUpdateHistory.usf
+++ b/Engine/Shaders/Private/TemporalAA/TAAUpdateHistory.usf
@@ -31,6 +31,7 @@ Texture2D<uint2> InputSceneStencilTexture;
 Texture2D<taa_half> HistoryRejectionTexture;
 
 Texture2D<float2> DilatedVelocityTexture;
+Texture2D<uint> ParallaxFactorTexture;
 Texture2D<taa_half> ParallaxRejectionMaskTexture;
 
 float3 HistoryQuantizationError;
@@ -40,10 +41,12 @@ uint bCameraCut;
 Texture2D<taa_half3> PrevHistory_Textures_0;
 Texture2D<taa_half3> PrevHistory_Textures_1;
 Texture2D<taa_half2> PrevHistory_Textures_2;
+Texture2D<uint>      PrevHistory_Textures_3;
 
 RWTexture2D<taa_half3> HistoryOutput_Textures_0;
 RWTexture2D<taa_half3> HistoryOutput_Textures_1;
 RWTexture2D<taa_half2> HistoryOutput_Textures_2;
+RWTexture2D<uint>      HistoryOutput_Textures_3;
 
 RWTexture2D<taa_half3> SceneColorOutput;
 
@@ -241,6 +244,8 @@ void MainCS(
 	taa_half3 PrevFallbackColor;
 	taa_half PrevFallbackWeight;
 	
+	uint PrevSubpixelDetails;
+
 	ISOLATE
 	{
 		// Reproject history
@@ -276,6 +281,17 @@ void MainCS(
 				RawHistory2 = PrevHistory_Textures_2.SampleLevel(GlobalBilinearClampedSampler, PrevHistoryBufferUV, 0);
 			}
 			#endif
+			
+			FSubpixelNeighborhood SubpixelNeighborhood = GatherPrevSubpixelNeighborhood(PrevHistory_Textures_3, PrevHistoryBufferUV);
+			{
+				PrevSubpixelDetails = 0;
+				UNROLL_N(SUB_PIXEL_COUNT)
+				for (uint SubpixelId = 0; SubpixelId < SUB_PIXEL_COUNT; SubpixelId++)
+				{
+					uint SubpixelPayload = GetSubpixelPayload(SubpixelNeighborhood, SubpixelId);
+					PrevSubpixelDetails |= SubpixelPayload << (BITS_PER_SUB_PIXEL * SubpixelId);
+				}
+			}
 
 			RawHistory0 = -min(-RawHistory0, taa_half(0.0));
 			RawHistory1 = -min(-RawHistory1, taa_half(0.0));
@@ -549,6 +565,35 @@ void MainCS(
 		}
 	}
 
+	// Update subpixel detail
+	uint FinalSubpixelDetails;
+	{
+		taa_half2 dKO = taa_half2(PPCo - PPCk);
+
+		bool bUpdate = all(abs(dKO) < 0.5 * (InputInfo_Extent.x * HistoryInfo_ExtentInverse.x));
+
+		FinalSubpixelDetails = PrevSubpixelDetails;
+		
+		FLATTEN
+		if (bOffScreen)
+		{
+			FinalSubpixelDetails = 0;
+		}
+
+		FLATTEN
+		if (bUpdate)
+		{
+			bool2 bBool = dKO < 0.0;
+
+			uint SubpixelId = dot(uint2(bBool), uint2(1, SUB_PIXEL_GRID_SIZE));
+			uint SubpixelShift = SubpixelId * BITS_PER_SUB_PIXEL;
+
+			uint SubpixelPayload = ParallaxFactorTexture[InputPixelPos] & SUB_PIXEL_MASK;
+
+			FinalSubpixelDetails = (FinalSubpixelDetails & (~(SUB_PIXEL_MASK << SubpixelShift))) | (SubpixelPayload << SubpixelShift);
+		}
+	}
+
 	// Compute final output
 	taa_half3 FinalOutputColor;
 	taa_half FinalOutputValidity;
@@ -665,6 +710,7 @@ void MainCS(
 			HistoryOutput_Textures_0[LocalHistoryPixelPos] = FinalFallbackColor;
 			HistoryOutput_Textures_1[LocalHistoryPixelPos] = FinalHistoryMoment1;
 			HistoryOutput_Textures_2[LocalHistoryPixelPos] = taa_half2(FinalFallbackWeight, FinalHistoryValidity);
+			HistoryOutput_Textures_3[LocalHistoryPixelPos] = FinalSubpixelDetails;
 
 			#if DEBUG_OUTPUT
 			{
diff --git a/Engine/Source/Runtime/Renderer/Private/PostProcess/TemporalAA.cpp b/Engine/Source/Runtime/Renderer/Private/PostProcess/TemporalAA.cpp
--- a/Engine/Source/Runtime/Renderer/Private/PostProcess/TemporalAA.cpp
+++ b/Engine/Source/Runtime/Renderer/Private/PostProcess/TemporalAA.cpp
@@ -17,7 +17,7 @@ namespace
 const int32 GTemporalAATileSizeX = 8;
 const int32 GTemporalAATileSizeY = 8;
 
-constexpr int32 kHistoryTextures = 3;
+constexpr int32 kHistoryTextures = 4;
 
 TAutoConsoleVariable<int32> CVarTAAAlgorithm(
 	TEXT("r.TemporalAA.Algorithm"), 0,
@@ -81,6 +81,11 @@ TAutoConsoleVariable<int32> CVarTAAR11G11B10History(
 	TEXT("Select the bitdepth of the history."),
 	ECVF_RenderThreadSafe);
 
+TAutoConsoleVariable<int32> CVarTAAFilterHistoryRejection(
+	TEXT("r.TemporalAA.FilterHistoryRejection"), 1,
+	TEXT("Whether the history rejection statistical filtering pass should be enabled or not."),
+	ECVF_RenderThreadSafe);
+
 inline bool DoesPlatformSupportTemporalHistoryUpscale(EShaderPlatform Platform)
 {
 	return (IsPCPlatform(Platform) || FDataDrivenShaderPlatformInfo::GetSupportsTemporalHistoryUpscale(Platform))
@@ -333,6 +338,7 @@ class FTAADilateVelocityCS : public FTAAGen5Shader
 		SHADER_PARAMETER_RDG_TEXTURE_UAV(RWTexture2D, ClosestDepthOutput)
 		SHADER_PARAMETER_RDG_TEXTURE_UAV(RWTexture2D, PrevUseCountOutput)
 		SHADER_PARAMETER_RDG_TEXTURE_UAV(RWTexture2D, PrevClosestDepthOutput)
+		SHADER_PARAMETER_RDG_TEXTURE_UAV(RWTexture2D, ParallaxFactorOutput)
 		SHADER_PARAMETER_RDG_TEXTURE_UAV(RWTexture2D, DebugOutput)
 	END_SHADER_PARAMETER_STRUCT()
 }; // class FTAADilateVelocityCS
@@ -344,6 +350,7 @@ class FTAADecimateHistoryCS : public FTAAGen5Shader
 
 	BEGIN_SHADER_PARAMETER_STRUCT(FParameters, )
 		SHADER_PARAMETER_STRUCT_INCLUDE(FTAACommonParameters, CommonParameters)
+		SHADER_PARAMETER(FMatrix, RotationalClipToPrevClip)
 		SHADER_PARAMETER(FVector, OutputQuantizationError)
 		SHADER_PARAMETER(float, HistoryPreExposureCorrection)
 		SHADER_PARAMETER(float, WorldDepthToPixelWorldRadius)
@@ -353,6 +360,7 @@ class FTAADecimateHistoryCS : public FTAAGen5Shader
 		SHADER_PARAMETER_RDG_TEXTURE(Texture2D, ClosestDepthTexture)
 		SHADER_PARAMETER_RDG_TEXTURE(Texture2D, PrevUseCountTexture)
 		SHADER_PARAMETER_RDG_TEXTURE(Texture2D, PrevClosestDepthTexture)
+		SHADER_PARAMETER_RDG_TEXTURE(Texture2D, ParallaxFactorTexture)
 
 		SHADER_PARAMETER_STRUCT(FScreenPassTextureViewportParameters, PrevHistoryInfo)
 		SHADER_PARAMETER_STRUCT(FTAAHistoryTextures, PrevHistory)
@@ -397,6 +405,19 @@ class FTAACompareHistoryCS : public FTAAGen5Shader
 	END_SHADER_PARAMETER_STRUCT()
 }; // class FTAACompareHistoryCS
 
+class FTAAPostfilterRejectionCS : public FTAAGen5Shader
+{
+	DECLARE_GLOBAL_SHADER(FTAAPostfilterRejectionCS);
+	SHADER_USE_PARAMETER_STRUCT(FTAAPostfilterRejectionCS, FTAAGen5Shader);
+
+	BEGIN_SHADER_PARAMETER_STRUCT(FParameters, )
+		SHADER_PARAMETER_STRUCT_INCLUDE(FTAACommonParameters, CommonParameters)
+		SHADER_PARAMETER_RDG_TEXTURE(Texture2D, HistoryRejectionTexture)
+		SHADER_PARAMETER_RDG_TEXTURE_UAV(RWTexture2D, HistoryRejectionOutput)
+		SHADER_PARAMETER_RDG_TEXTURE_UAV(RWTexture2D, DebugOutput)
+	END_SHADER_PARAMETER_STRUCT()
+}; // class FTAAPostfilterRejectionCS
+
 class FTAADilateRejectionCS : public FTAAGen5Shader
 {
 	DECLARE_GLOBAL_SHADER(FTAADilateRejectionCS);
@@ -422,6 +443,7 @@ class FTAAUpdateHistoryCS : public FTAAGen5Shader
 		SHADER_PARAMETER_RDG_TEXTURE(Texture2D, HistoryRejectionTexture)
 
 		SHADER_PARAMETER_RDG_TEXTURE(Texture2D, DilatedVelocityTexture)
+		SHADER_PARAMETER_RDG_TEXTURE(Texture2D, ParallaxFactorTexture)
 		SHADER_PARAMETER_RDG_TEXTURE(Texture2D, ParallaxRejectionMaskTexture)
 
 		SHADER_PARAMETER(FVector, HistoryQuantizationError)
@@ -443,6 +465,7 @@ IMPLEMENT_GLOBAL_SHADER(FTAADilateVelocityCS, "/Engine/Private/TemporalAA/TAADil
 IMPLEMENT_GLOBAL_SHADER(FTAADecimateHistoryCS, "/Engine/Private/TemporalAA/TAADecimateHistory.usf", "MainCS", SF_Compute);
 IMPLEMENT_GLOBAL_SHADER(FTAAFilterFrequenciesCS, "/Engine/Private/TemporalAA/TAAFilterFrequencies.usf", "MainCS", SF_Compute);
 IMPLEMENT_GLOBAL_SHADER(FTAACompareHistoryCS, "/Engine/Private/TemporalAA/TAACompareHistory.usf", "MainCS", SF_Compute);
+IMPLEMENT_GLOBAL_SHADER(FTAAPostfilterRejectionCS, "/Engine/Private/TemporalAA/TAAPostfilterRejection.usf", "MainCS", SF_Compute);
 IMPLEMENT_GLOBAL_SHADER(FTAADilateRejectionCS, "/Engine/Private/TemporalAA/TAADilateRejection.usf", "MainCS", SF_Compute);
 IMPLEMENT_GLOBAL_SHADER(FTAAUpdateHistoryCS, "/Engine/Private/TemporalAA/TAAUpdateHistory.usf", "MainCS", SF_Compute);
 
@@ -1067,9 +1090,7 @@ static void AddGen5MainTemporalAAPasses(
 		return GraphBuilder.CreateUAV(DebugTexture);
 	};
 
-	// Dilate the velocity texture & build the parallax rejection mask
-	FRDGTextureRef DilatedVelocityTexture;
-	FRDGTextureRef ClosestDepthTexture;
+	// Clear atomic scattered texture.
 	FRDGTextureRef PrevUseCountTexture;
 	FRDGTextureRef PrevClosestDepthTexture;
 	{
@@ -1098,6 +1119,10 @@ static void AddGen5MainTemporalAAPasses(
 			FComputeShaderUtils::GetGroupCount(InputRect.Size(), 8));
 	}
 
+	// Dilate the velocity texture & scatter reprojection into previous frame
+	FRDGTextureRef DilatedVelocityTexture;
+	FRDGTextureRef ClosestDepthTexture;
+	FRDGTextureRef ParallaxFactorTexture;
 	{
 		{
 			FRDGTextureDesc Desc = FRDGTextureDesc::Create2D(
@@ -1110,6 +1135,9 @@ static void AddGen5MainTemporalAAPasses(
 
 			Desc.Format = PF_R16F;
 			ClosestDepthTexture = GraphBuilder.CreateTexture(Desc, TEXT("TAA.ClosestDepthTexture"));
+
+			Desc.Format = PF_R8_UINT;
+			ParallaxFactorTexture = GraphBuilder.CreateTexture(Desc, TEXT("TAA.ParallaxFactor"));
 		}
 
 		FTAADilateVelocityCS::FParameters* PassParameters = GraphBuilder.AllocParameters<FTAADilateVelocityCS::FParameters>();
@@ -1120,6 +1148,7 @@ static void AddGen5MainTemporalAAPasses(
 		PassParameters->ClosestDepthOutput = GraphBuilder.CreateUAV(ClosestDepthTexture);
 		PassParameters->PrevUseCountOutput = GraphBuilder.CreateUAV(PrevUseCountTexture);
 		PassParameters->PrevClosestDepthOutput = GraphBuilder.CreateUAV(PrevClosestDepthTexture);
+		PassParameters->ParallaxFactorOutput = GraphBuilder.CreateUAV(ParallaxFactorTexture);
 		PassParameters->DebugOutput = CreateDebugUAV(InputExtent, TEXT("Debug.TAA.DilateVelocity"));
 
 		TShaderMapRef<FTAADilateVelocityCS> ComputeShader(View.ShaderMap);
@@ -1184,6 +1213,15 @@ static void AddGen5MainTemporalAAPasses(
 
 		FTAADecimateHistoryCS::FParameters* PassParameters = GraphBuilder.AllocParameters<FTAADecimateHistoryCS::FParameters>();
 		PassParameters->CommonParameters = CommonParameters;
+		{
+			const FViewMatrices& ViewMatrices = View.ViewMatrices;
+			const FViewMatrices& PrevViewMatrices = View.PrevViewInfo.ViewMatrices;
+
+			FMatrix RotationalInvViewProj = ViewMatrices.ComputeInvProjectionNoAAMatrix() * (ViewMatrices.GetTranslatedViewMatrix().RemoveTranslation().GetTransposed());
+			FMatrix RotationalPrevViewProj = (PrevViewMatrices.GetTranslatedViewMatrix().RemoveTranslation()) * PrevViewMatrices.ComputeProjectionNoAAMatrix();
+
+			PassParameters->RotationalClipToPrevClip = RotationalInvViewProj * RotationalPrevViewProj;
+		}
 		PassParameters->OutputQuantizationError = ComputePixelFormatQuantizationError(PredictionSceneColorTexture->Desc.Format);
 		PassParameters->HistoryPreExposureCorrection = View.PreExposure / View.PrevViewInfo.SceneColorPreExposure;
 		{
@@ -1198,6 +1236,7 @@ static void AddGen5MainTemporalAAPasses(
 		PassParameters->ClosestDepthTexture = ClosestDepthTexture;
 		PassParameters->PrevUseCountTexture = PrevUseCountTexture;
 		PassParameters->PrevClosestDepthTexture = PrevClosestDepthTexture;
+		PassParameters->ParallaxFactorTexture = ParallaxFactorTexture;
 
 		PassParameters->PrevHistoryInfo = PrevHistoryInfo;
 		PassParameters->PrevHistory = PrevHistory;
@@ -1285,6 +1324,28 @@ static void AddGen5MainTemporalAAPasses(
 		}
 	}
 
+	// Post filter the rejection.
+	if (CVarTAAFilterHistoryRejection.GetValueOnRenderThread() != 0)
+	{
+		FRDGTextureRef FilteredHistoryRejectionTexture = GraphBuilder.CreateTexture(HistoryRejectionTexture->Desc, TEXT("TAA.FilteredHistoryRejection"));
+
+		FTAAPostfilterRejectionCS::FParameters* PassParameters = GraphBuilder.AllocParameters<FTAAPostfilterRejectionCS::FParameters>();
+		PassParameters->CommonParameters = CommonParameters;
+		PassParameters->HistoryRejectionTexture = HistoryRejectionTexture;
+		PassParameters->HistoryRejectionOutput = GraphBuilder.CreateUAV(FilteredHistoryRejectionTexture);
+		PassParameters->DebugOutput = CreateDebugUAV(RejectionExtent, TEXT("Debug.TAA.PostfilterRejection"));
+
+		TShaderMapRef<FTAAPostfilterRejectionCS> ComputeShader(View.ShaderMap);
+		FComputeShaderUtils::AddPass(
+			GraphBuilder,
+			RDG_EVENT_NAME("TAA PostfilterRejection %dx%d", RejectionRect.Width(), RejectionRect.Height()),
+			ComputeShader,
+			PassParameters,
+			FComputeShaderUtils::GetGroupCount(RejectionRect.Size(), 8));
+
+		HistoryRejectionTexture = FilteredHistoryRejectionTexture;
+	}
+
 	// Dilate the rejection.
 	FRDGTextureRef DilatedHistoryRejectionTexture;
 	{
@@ -1322,6 +1383,9 @@ static void AddGen5MainTemporalAAPasses(
 
 			Desc.Format = PF_R8G8;
 			History.Textures[2] = GraphBuilder.CreateTexture(Desc, TEXT("TAA.History.Metadata"));
+
+			Desc.Format = PF_R16_UINT;
+			History.Textures[3] = GraphBuilder.CreateTexture(Desc, TEXT("TAA.History.SubpixelInfo"));
 		}
 
 		// Allocate output
@@ -1342,6 +1406,7 @@ static void AddGen5MainTemporalAAPasses(
 			FRDGTextureSRVDesc::CreateWithPixelFormat(PassInputs.SceneDepthTexture, PF_X24_G8));
 		PassParameters->HistoryRejectionTexture = DilatedHistoryRejectionTexture;
 		PassParameters->DilatedVelocityTexture = DilatedVelocityTexture;
+		PassParameters->ParallaxFactorTexture = ParallaxFactorTexture;
 		PassParameters->ParallaxRejectionMaskTexture = ParallaxRejectionMaskTexture;
 
 		PassParameters->HistoryQuantizationError = ComputePixelFormatQuantizationError(History.Textures[0]->Desc.Format);
@@ -1392,6 +1457,7 @@ static void AddGen5MainTemporalAAPasses(
 	{
 		OutputHistory->SafeRelease();
 
+		check(kHistoryTextures <= OutputHistory->RT.Num());
 		for (int32 i = 0; i < kHistoryTextures; i++)
 		{
 			if (ExtractHistory[i])
