Make InstanceCullingContext local to mesh passes to ensure correct life-time.
- Also stripped out instance culling manager and some other superseded code.

diff --git a/Engine/Source/Runtime/Renderer/Private/InstanceCulling/InstanceCullingContext.cpp b/Engine/Source/Runtime/Renderer/Private/InstanceCulling/InstanceCullingContext.cpp
--- a/Engine/Source/Runtime/Renderer/Private/InstanceCulling/InstanceCullingContext.cpp
+++ b/Engine/Source/Runtime/Renderer/Private/InstanceCulling/InstanceCullingContext.cpp
@@ -14,6 +14,14 @@
 
 #define ENABLE_DETERMINISTIC_INSTANCE_CULLING 1
 
+
+FInstanceCullingContext::FInstanceCullingContext(FInstanceCullingManager* InInstanceCullingManager, TArrayView<const int32> InViewIds) :
+	InstanceCullingManager(InInstanceCullingManager),
+	ViewIds(InViewIds),
+	bIsEnabled(InInstanceCullingManager == nullptr || InInstanceCullingManager->IsEnabled())
+{
+}
+
 void FInstanceCullingContext::BeginCullingCommand(EPrimitiveType BatchType, uint32 BaseVertexIndex, uint32 FirstIndex, uint32 NumPrimitives)
 {
 #if defined(GPUCULL_TODO)
@@ -63,12 +71,6 @@ void FInstanceCullingContext::AddInstanceRunToCullingCommand(int32 ScenePrimitiv
 #endif
 }
 
-int32 FInstanceCullingContext::AllocateArgsSlotRange(uint32 NumSlots)
-{
-	//CullingCommands.SetNumZeroed(NumSlots);
-	return 0;
-}
-
 #if defined(GPUCULL_TODO)
 
 #if ENABLE_DETERMINISTIC_INSTANCE_CULLING
diff --git a/Engine/Source/Runtime/Renderer/Private/InstanceCulling/InstanceCullingManager.cpp b/Engine/Source/Runtime/Renderer/Private/InstanceCulling/InstanceCullingManager.cpp
--- a/Engine/Source/Runtime/Renderer/Private/InstanceCulling/InstanceCullingManager.cpp
+++ b/Engine/Source/Runtime/Renderer/Private/InstanceCulling/InstanceCullingManager.cpp
@@ -22,11 +22,6 @@ static TAutoConsoleVariable<int32> CVarCullInstances(
 
 FInstanceCullingManager::~FInstanceCullingManager()
 {
-	for (FInstanceCullingContext* InstanceCullingContext : CullingContexts)
-	{
-		// Memstack allocated, so no destructor is called.
-		InstanceCullingContext->~FInstanceCullingContext();
-	}
 }
 
 int32 FInstanceCullingManager::RegisterView(const FViewInfo& ViewInfo)
@@ -165,20 +160,3 @@ void FInstanceCullingManager::CullInstances(FRDGBuilder& GraphBuilder, FGPUScene
 	}
 #endif // defined(GPUCULL_TODO)
 }
-
-
-FInstanceCullingContext* FInstanceCullingManager::CreateContext(const int32* ViewIds, int32 NumViews)
-{
-	if (bIsEnabled)
-	{
-		FInstanceCullingContext* InstanceCullingContext = new(FMemStack::Get()) FInstanceCullingContext;
-		CullingContexts.Add(InstanceCullingContext);
-		InstanceCullingContext->ViewIds.Insert(ViewIds, NumViews, 0);
-		InstanceCullingContext->InstanceCullingManager = this;
-		return InstanceCullingContext;
-	}
-	else
-	{
-		return nullptr;
-	}
-}
diff --git a/Engine/Source/Runtime/Renderer/Private/InstanceCulling/InstanceCullingManager.h b/Engine/Source/Runtime/Renderer/Private/InstanceCulling/InstanceCullingManager.h
--- a/Engine/Source/Runtime/Renderer/Private/InstanceCulling/InstanceCullingManager.h
+++ b/Engine/Source/Runtime/Renderer/Private/InstanceCulling/InstanceCullingManager.h
@@ -73,6 +73,8 @@ class FInstanceCullingManager
 public:
 	~FInstanceCullingManager();
 
+	bool IsEnabled() const { return bIsEnabled; }
+
 	// Max average number of instances that primitives are expanded to. GPUCULL_TODO: Not very robust
 	static constexpr uint32 MaxAverageInstanceFactor = 128;
 
@@ -96,12 +98,6 @@ public:
 	{
 	}
 
-	/**
-	 * Create a context for the instance culling, used to collect draw commands for a render pass. The context lifetime is managed by the Manager, and in general is a single-frame construct
-	 * Once the render pass setup is done, the Culling Context should be queued up for final processing (where the culled instances are used to fill in the draw commands).
-	 */
-	FInstanceCullingContext* CreateContext(const int32* ViewIds, int32 NumViews);
-
 	// Populated by CullInstances, used when performing final culling & rendering 
 	FInstanceCullingIntermediate CullingIntermediate;
 
@@ -111,7 +107,6 @@ private:
 
 	FInstanceCullingManagerResources& Resources;
 	TArray<Nanite::FPackedView> CullingViews;
-	TArray<FInstanceCullingContext*> CullingContexts;
 	bool bIsEnabled;
 };
 
diff --git a/Engine/Source/Runtime/Renderer/Private/MeshDrawCommands.cpp b/Engine/Source/Runtime/Renderer/Private/MeshDrawCommands.cpp
--- a/Engine/Source/Runtime/Renderer/Private/MeshDrawCommands.cpp
+++ b/Engine/Source/Runtime/Renderer/Private/MeshDrawCommands.cpp
@@ -539,7 +539,6 @@ void SetupGPUInstancedDraws(
 	const int32 NumDrawCommandsIn = VisibleMeshDrawCommandsInOut.Num();
 	int32 NumDrawCommandsOut = 0;
 	// Allocate conservatively for all commands, may not use all.
-	int32 BaseArgsSlotIndex = InstanceCullingContext.AllocateArgsSlotRange(NumDrawCommandsIn);
 	for (int32 DrawCommandIndex = 0; DrawCommandIndex < NumDrawCommandsIn; ++DrawCommandIndex)
 	{
 		const FVisibleMeshDrawCommand& RESTRICT VisibleMeshDrawCommand = PassVisibleMeshDrawCommands[DrawCommandIndex];
@@ -971,7 +970,7 @@ public:
 			{
 #if defined(GPUCULL_TODO)
 				// GPUCULL_TODO: Make a switch to control old / new behaviour, determine minimum reqs.
-				SetupGPUInstancedDraws(*Context.InstanceCullingContext, Context.MeshDrawCommands, Context.MaxInstances, Context.VisibleMeshDrawCommandsNum, Context.NewPassVisibleMeshDrawCommandsNum);
+				SetupGPUInstancedDraws(Context.InstanceCullingContext, Context.MeshDrawCommands, Context.MaxInstances, Context.VisibleMeshDrawCommandsNum, Context.NewPassVisibleMeshDrawCommandsNum);
 #else //!defined(GPUCULL_TODO)
 				BuildMeshDrawCommandPrimitiveIdBuffer(
 					Context.bDynamicInstancing,
@@ -1113,41 +1112,10 @@ void SortAndMergeDynamicPassMeshDrawCommands(
 
 
 
-/*
- * Used by various dynamic passes to sort/merge mesh draw commands immediately on a rendering thread.
- */
-void SortAndMergeDynamicPassMeshDrawCommands(
-	ERHIFeatureLevel::Type FeatureLevel,
-	FMeshCommandOneFrameArray& VisibleMeshDrawCommands,
-	FDynamicMeshDrawCommandStorage& MeshDrawCommandStorage,
-	FInstanceCullingContext& InstanceCullingContext)
-{
-#if defined(GPUCULL_TODO)
-	const bool bUseGPUScene = UseGPUScene(GMaxRHIShaderPlatform, FeatureLevel);
-
-	const int32 NumDrawCommands = VisibleMeshDrawCommands.Num();
-	if (NumDrawCommands > 0)
-	{
-		FMeshCommandOneFrameArray NewPassVisibleMeshDrawCommands;
-		int32 MaxInstances = 1;
-		int32 VisibleMeshDrawCommandsNum = 0;
-		int32 NewPassVisibleMeshDrawCommandsNum = 0;
-
-		VisibleMeshDrawCommands.Sort(FCompareFMeshDrawCommands());
-
-		if (bUseGPUScene)
-		{
-			SetupGPUInstancedDraws(InstanceCullingContext, VisibleMeshDrawCommands, MaxInstances, VisibleMeshDrawCommandsNum, NewPassVisibleMeshDrawCommandsNum);
-		}
-	}
-#endif // GPUCULL_TODO
-}
-
 void FParallelMeshDrawCommandPass::DispatchPassSetup(
 	FScene* Scene,
 	const FViewInfo& View,
-	FInstanceCullingContext* InstanceCullingContext,
-	FInstanceCullingManager* InstanceCullingManager,
+	FInstanceCullingContext&& InstanceCullingContext,
 	EMeshPass::Type PassType,
 	FExclusiveDepthStencil::Type BasePassDepthStencilAccess,
 	FMeshPassProcessor* MeshPassProcessor,
@@ -1198,8 +1166,7 @@ void FParallelMeshDrawCommandPass::DispatchPassSetup(
 	TaskContext.InstanceFactor = (bIsMainViewPass && View.IsInstancedStereoPass()) ? 2 : 1;
 #endif // defined(GPUCULL_TODO)
 
-	TaskContext.InstanceCullingContext = InstanceCullingContext; 
-	TaskContext.InstanceCullingManager = InstanceCullingManager;
+	TaskContext.InstanceCullingContext = MoveTemp(InstanceCullingContext); 
 
 	// Setup translucency sort key update pass based on view.
 	TaskContext.TranslucencyPass = ETranslucencyPass::TPT_MAX;
@@ -1467,13 +1434,13 @@ public:
 void FParallelMeshDrawCommandPass::BuildRenderingCommands(FRDGBuilder& GraphBuilder, FGPUScene& GPUScene, FInstanceCullingDrawParams& OutInstanceCullingDrawParams)
 {
 #if defined(GPUCULL_TODO)
-	if (TaskContext.InstanceCullingContext)
+	if (TaskContext.InstanceCullingContext.IsEnabled())
 	{
 		WaitForMeshPassSetupTask();
-		if (MaxNumDraws > 0 && TaskContext.InstanceCullingContext->HasCullingCommands())
+		if (MaxNumDraws > 0 && TaskContext.InstanceCullingContext.HasCullingCommands())
 		{
 			// 2. Run finalize culling commands pass
-			TaskContext.InstanceCullingContext->BuildRenderingCommands(GraphBuilder, GPUScene, TaskContext.InstanceCullingResult);
+			TaskContext.InstanceCullingContext.BuildRenderingCommands(GraphBuilder, GPUScene, TaskContext.InstanceCullingResult);
 			TaskContext.InstanceCullingResult.GetDrawParameters(OutInstanceCullingDrawParams);
 			return;
 		}
@@ -1489,7 +1456,7 @@ void FParallelMeshDrawCommandPass::BuildRenderingCommands(FRDGBuilder& GraphBuil
 void FParallelMeshDrawCommandPass::BuildInstanceList(FRDGBuilder& GraphBuilder, FGPUScene& GPUScene, FInstanceCullingRdgParams& OutParams)
 {
 #if defined(GPUCULL_TODO)
-	if (TaskContext.InstanceCullingContext)
+	if (TaskContext.InstanceCullingContext.IsEnabled())
 	{
 		WaitForMeshPassSetupTask();
 		if (MaxNumDraws <= 0)
@@ -1497,7 +1464,7 @@ void FParallelMeshDrawCommandPass::BuildInstanceList(FRDGBuilder& GraphBuilder,
 			return;
 		}
 		// Run pass to build ID lists (temporary)
-		TaskContext.InstanceCullingContext->BuildRenderingCommands(GraphBuilder, GPUScene, OutParams);
+		TaskContext.InstanceCullingContext.BuildRenderingCommands(GraphBuilder, GPUScene, OutParams);
 	}
 #endif
 }
diff --git a/Engine/Source/Runtime/Renderer/Private/MeshDrawCommands.h b/Engine/Source/Runtime/Renderer/Private/MeshDrawCommands.h
--- a/Engine/Source/Runtime/Renderer/Private/MeshDrawCommands.h
+++ b/Engine/Source/Runtime/Renderer/Private/MeshDrawCommands.h
@@ -56,7 +56,6 @@ public:
 	FMeshDrawCommandPassSetupTaskContext()
 		: View(nullptr)
 		, Scene(nullptr)
-		, InstanceCullingManager(nullptr)
 		, ShadingPath(EShadingPath::Num)
 		, PassType(EMeshPass::Num)
 		, bUseGPUScene(false)
@@ -77,13 +76,11 @@ public:
 		, VisibleMeshDrawCommandsNum(0)
 		, NewPassVisibleMeshDrawCommandsNum(0)
 		, MaxInstances(1)
-		, InstanceCullingContext(nullptr)
 	{
 	}
 
 	const FViewInfo* View;
 	const FScene* Scene;
-	FInstanceCullingManager* InstanceCullingManager;
 	EShadingPath ShadingPath;
 	EShaderPlatform ShaderPlatform;
 	EMeshPass::Type PassType;
@@ -130,7 +127,7 @@ public:
 	int32 NewPassVisibleMeshDrawCommandsNum;
 	int32 MaxInstances;
 
-	FInstanceCullingContext* InstanceCullingContext;
+	FInstanceCullingContext InstanceCullingContext;
 	FInstanceCullingResult InstanceCullingResult;
 };
 
@@ -156,8 +153,7 @@ public:
 	void DispatchPassSetup(
 		FScene* Scene,
 		const FViewInfo& View, 
-		FInstanceCullingContext* InstanceCullingContext,
-		FInstanceCullingManager* InstanceCullingManager,
+		FInstanceCullingContext &&InstanceCullingContext,
 		EMeshPass::Type PassType, 
 		FExclusiveDepthStencil::Type BasePassDepthStencilAccess,
 		FMeshPassProcessor* MeshPassProcessor,
@@ -204,7 +200,7 @@ public:
 
 	static bool IsOnDemandShaderCreationEnabled();
 
-	FInstanceCullingContext* GetInstanceCullingContext() { return TaskContext.InstanceCullingContext; }
+	FInstanceCullingContext* GetInstanceCullingContext() { return &TaskContext.InstanceCullingContext; }
 
 private:
 	FPrimitiveIdVertexBufferPoolEntry PrimitiveIdVertexBufferPoolEntry;
@@ -230,12 +226,6 @@ extern void SortAndMergeDynamicPassMeshDrawCommands(
 	FRHIBuffer*& OutPrimitiveIdVertexBuffer,
 	uint32 InstanceFactor);
 
-extern void SortAndMergeDynamicPassMeshDrawCommands(
-	ERHIFeatureLevel::Type FeatureLevel,
-	FMeshCommandOneFrameArray& VisibleMeshDrawCommands,
-	FDynamicMeshDrawCommandStorage& MeshDrawCommandStorage,
-	FInstanceCullingContext* InstanceCullingContext);
-
 void SubmitGPUInstancedMeshDrawCommandsRange(
 	const FMeshCommandOneFrameArray& VisibleMeshDrawCommands,
 	const FGraphicsMinimalPipelineStateSet& GraphicsMinimalPipelineStateSet,
diff --git a/Engine/Source/Runtime/Renderer/Private/MeshPassProcessor.cpp b/Engine/Source/Runtime/Renderer/Private/MeshPassProcessor.cpp
--- a/Engine/Source/Runtime/Renderer/Private/MeshPassProcessor.cpp
+++ b/Engine/Source/Runtime/Renderer/Private/MeshPassProcessor.cpp
@@ -1481,8 +1481,7 @@ FSimpleMeshDrawCommandPass::FSimpleMeshDrawCommandPass(const FSceneView& View, F
 {
 	check(View.bIsViewInfo);
 	const FViewInfo* ViewInfo = static_cast<const FViewInfo*>(&View);
-	InstanceCullingContext.ViewIds.Add(ViewInfo->GPUSceneViewId);
-	InstanceCullingContext.InstanceCullingManager = InstanceCullingManager;
+	InstanceCullingContext = FInstanceCullingContext(InstanceCullingManager, TArrayView<const int32>(&ViewInfo->GPUSceneViewId, 1));
 }
 
 void FSimpleMeshDrawCommandPass::BuildRenderingCommands(FRDGBuilder& GraphBuilder, const FSceneView& View, const FGPUScene& GPUScene, FInstanceCullingDrawParams& OutInstanceCullingDrawParams)
diff --git a/Engine/Source/Runtime/Renderer/Private/MobileShadingRenderer.cpp b/Engine/Source/Runtime/Renderer/Private/MobileShadingRenderer.cpp
--- a/Engine/Source/Runtime/Renderer/Private/MobileShadingRenderer.cpp
+++ b/Engine/Source/Runtime/Renderer/Private/MobileShadingRenderer.cpp
@@ -287,8 +287,7 @@ void FMobileSceneRenderer::SetupMobileBasePassAfterShadowInit(FExclusiveDepthSte
 		Pass.DispatchPassSetup(
 			Scene,
 			View,
-			nullptr, // GPUCULL_TODO: Fix mobile!
-			nullptr, // GPUCULL_TODO: Fix mobile!
+			FInstanceCullingContext(), // GPUCULL_TODO: Fix mobile!
 			EMeshPass::BasePass,
 			BasePassDepthStencilAccess,
 			MeshPassProcessor,
diff --git a/Engine/Source/Runtime/Renderer/Private/SceneRendering.cpp b/Engine/Source/Runtime/Renderer/Private/SceneRendering.cpp
--- a/Engine/Source/Runtime/Renderer/Private/SceneRendering.cpp
+++ b/Engine/Source/Runtime/Renderer/Private/SceneRendering.cpp
@@ -3262,13 +3262,10 @@ void FSceneRenderer::SetupMeshPass(FViewInfo& View, FExclusiveDepthStencil::Type
 				View.GPUSceneViewId + 1,
 			};
 
-			FInstanceCullingContext* InstanceCullingContext = InstanceCullingManager.CreateContext(ViewIds, NumViews);
-
 			Pass.DispatchPassSetup(
 				Scene,
 				View,
-				InstanceCullingContext,
-				&InstanceCullingManager,
+				FInstanceCullingContext(&InstanceCullingManager, TArrayView<int32>(ViewIds, NumViews)),
 				PassType,
 				BasePassDepthStencilAccess,
 				MeshPassProcessor,
diff --git a/Engine/Source/Runtime/Renderer/Private/ShadowSetup.cpp b/Engine/Source/Runtime/Renderer/Private/ShadowSetup.cpp
--- a/Engine/Source/Runtime/Renderer/Private/ShadowSetup.cpp
+++ b/Engine/Source/Runtime/Renderer/Private/ShadowSetup.cpp
@@ -2080,9 +2080,6 @@ void FProjectedShadowInfo::SetupMeshDrawCommandsForShadowDepth(FSceneRenderer& R
 	}
 	// GPUCULL_TODO: Pass along any custom culling planes or whatever here (e.g., cacade bounds):
 	// GPUCULL_TODO: Add debug tags to context and views (so compute passes can be understood)
-
-	FInstanceCullingContext* InstanceCullingContext = InstanceCullingManager.CreateContext(ViewIds.GetData(), ViewIds.Num());
-
 	extern int32 GShadowUseGS;
 #if defined(GPUCULL_TODO)
 	// GPUCULL_TODO: Needed to support legacy, non-GPU-Scene culled, primitives, this is merely used to allocate enough space for CPU-side replication.
@@ -2096,8 +2093,7 @@ void FProjectedShadowInfo::SetupMeshDrawCommandsForShadowDepth(FSceneRenderer& R
 	ShadowDepthPass.DispatchPassSetup(
 		Renderer.Scene,
 		*ShadowDepthView,
-		InstanceCullingContext,
-		&InstanceCullingManager,
+		FInstanceCullingContext(&InstanceCullingManager, ViewIds),
 		EMeshPass::Num,
 		FExclusiveDepthStencil::DepthNop_StencilNop,
 		MeshPassProcessor,
diff --git a/Engine/Source/Runtime/Renderer/Public/InstanceCulling/InstanceCullingContext.h b/Engine/Source/Runtime/Renderer/Public/InstanceCulling/InstanceCullingContext.h
--- a/Engine/Source/Runtime/Renderer/Public/InstanceCulling/InstanceCullingContext.h
+++ b/Engine/Source/Runtime/Renderer/Public/InstanceCulling/InstanceCullingContext.h
@@ -63,6 +63,8 @@ public:
 
 	FInstanceCullingContext() {}
 
+	FInstanceCullingContext(FInstanceCullingManager* InInstanceCullingManager, TArrayView<const int32> InViewIds);
+
 	struct FPrimCullingCommand
 	{
 		uint32 BaseVertexIndex;
@@ -72,6 +74,8 @@ public:
 		uint32 FirstInstanceRunOffset;
 	};
 
+	bool IsEnabled() const { return bIsEnabled; }
+
 	/**
 	 * Begin a new command. Allocates a slot and is referenced by subsequent AddPrimitiveToCullingCommand and AddInstanceRunToCullingCommand.
 	 */
@@ -89,12 +93,6 @@ public:
 	 */
 	void AddInstanceRunToCullingCommand(int32 ScenePrimitiveId, const uint32* Runs, uint32 NumRuns);
 
-	/**
-	 * Allocate a range from the indirect args buffer (and instance index offset buffer).
-	 * This function is thread safe. The alloction is only valid for a single frame.
-	 */
-	int32 AllocateArgsSlotRange(uint32 NumSlots);
-
 
 	void BuildRenderingCommands(FRDGBuilder& GraphBuilder, const FGPUScene& GPUScene, FInstanceCullingResult& Results) const;
 
@@ -105,20 +103,6 @@ public:
 
 	inline bool HasCullingCommands() const { return CullingCommands.Num() > 0; 	}
 
-#if 0
-	/**
-	 * Creates buffers of fixed size that will be filled in on the CPU during a dynamic pass, need to have a conservative size.
-	 * @param NumDrawCommands should include any instancing e.g., for ISR, it is fine to be an overestimate.
-	 */
-	static void CreateLegacyPassParameters(FRDGBuilder& GraphBuilder, FInstanceCullingDrawParams& InstanceCullingDrawParamsOut, int32 NumDrawCommands);
-
-
-	/**
-	 *
-	 */
-	void BuildLegacyRenderingCommands(const FScene& Scene, FRHICommandListImmediate& RHICmdList, const FGPUScenePrimitiveCollector& DynamicPrimitiveCollector, FRHIBuffer*& InstanceIdOffsetBufferOut, FRHIBuffer*& DrawIndirectArgsBufferOut);
-#endif // 0
-
 	// GPUCULL_TODO: These should not be dynamically heap-allocated, all except instance runs are easy to pre-size on memstack.
 	//           Must be presized as populated from task threads.
 	TArray<FPrimCullingCommand/*, SceneRenderingAllocator*/> CullingCommands;
@@ -131,8 +115,9 @@ public:
 		int32 PrimitiveId;
 	};
 
+	FInstanceCullingManager* InstanceCullingManager = nullptr;
 	TArray<FInstanceRun/*, SceneRenderingAllocator*/> InstanceRuns;
 	TArray<int32/*, SceneRenderingAllocator*/> ViewIds;
-	FInstanceCullingManager* InstanceCullingManager = nullptr;
+	bool bIsEnabled = false;
 };
 
diff --git a/Engine/Source/Runtime/Renderer/Public/MeshPassProcessor.h b/Engine/Source/Runtime/Renderer/Public/MeshPassProcessor.h
--- a/Engine/Source/Runtime/Renderer/Public/MeshPassProcessor.h
+++ b/Engine/Source/Runtime/Renderer/Public/MeshPassProcessor.h
@@ -1973,21 +1973,6 @@ private:
 
 #if defined(GPUCULL_TODO)
 
-#if 0
-/**
- * Same as the parallel one, but intended for use with simpler tasks where the overhead and complexity of parallel is not justified.
- */
-class SimpleMeshDrawCommandPass
-{
-public:
-
-private:
-	FMeshCommandOneFrameArray VisibleMeshDrawCommands;
-	FInstanceCullingContext InstanceCullingContext;
-};
-
-#endif
-
 /**
  * Same as the parallel one, but intended for use with simpler tasks where the overhead and complexity of parallel is not justified.
  */
