Fixes composite editor primitives when using temporal upsampling

diff --git a/Engine/Source/Runtime/Renderer/Private/PostProcess/PostProcessCompositeEditorPrimitives.cpp b/Engine/Source/Runtime/Renderer/Private/PostProcess/PostProcessCompositeEditorPrimitives.cpp
--- a/Engine/Source/Runtime/Renderer/Private/PostProcess/PostProcessCompositeEditorPrimitives.cpp
+++ b/Engine/Source/Runtime/Renderer/Private/PostProcess/PostProcessCompositeEditorPrimitives.cpp
@@ -255,8 +255,45 @@ FScreenPassTexture AddEditorPrimitivePass(
 
 	// Load the color target if it already exists.
 	const bool bProducedByPriorPass = HasBeenProduced(SceneTextures.EditorPrimitiveColor);
-	ERenderTargetLoadAction CompositeLoadAction = bProducedByPriorPass ? ERenderTargetLoadAction::ELoad : ERenderTargetLoadAction::EClear;
-	const FScreenPassTextureViewport EditorPrimitivesViewport(SceneTextures.EditorPrimitiveColor, Inputs.SceneColor.ViewRect);
+
+	FRDGTextureRef EditorPrimitiveColor;
+	FRDGTextureRef EditorPrimitiveDepth;
+	if (bProducedByPriorPass)
+	{
+		ensureMsgf(
+			Inputs.SceneColor.ViewRect == Inputs.SceneDepth.ViewRect,
+			TEXT("Temporal upsampling should be disabled when drawing directly to EditorPrimitivesColor."));
+		EditorPrimitiveColor = SceneTextures.EditorPrimitiveColor;
+		EditorPrimitiveDepth = SceneTextures.EditorPrimitiveDepth;
+	}
+	else
+	{
+		const FSceneTexturesConfig& Config = SceneTextures.Config;
+
+		FIntPoint Extent = Inputs.SceneColor.Texture->Desc.Extent;
+
+		const FRDGTextureDesc ColorDesc = FRDGTextureDesc::Create2D(
+			Extent,
+			PF_B8G8R8A8,
+			FClearValueBinding::Transparent,
+			TexCreate_ShaderResource | TexCreate_RenderTargetable,
+			1,
+			Config.EditorPrimitiveNumSamples);
+
+		const FRDGTextureDesc DepthDesc = FRDGTextureDesc::Create2D(
+			Extent,
+			PF_DepthStencil,
+			FClearValueBinding::DepthFar,
+			TexCreate_ShaderResource | TexCreate_DepthStencilTargetable,
+			1,
+			Config.EditorPrimitiveNumSamples);
+
+		EditorPrimitiveColor = GraphBuilder.CreateTexture(ColorDesc, TEXT("Editor.PrimitivesDepth"));
+		EditorPrimitiveDepth = GraphBuilder.CreateTexture(DepthDesc, TEXT("Editor.PrimitivesColor"));
+	}
+
+	// Load the color target if it already exists.
+	const FScreenPassTextureViewport EditorPrimitivesViewport(EditorPrimitiveColor, Inputs.SceneColor.ViewRect);
 
 	// The editor primitive composition pass is also used when rendering VMI_WIREFRAME in order to use MSAA.
 	// So we need to check whether the editor primitives are enabled inside this function.
@@ -272,8 +309,8 @@ FScreenPassTexture AddEditorPrimitivePass(
 			PassParameters->Depth = GetScreenPassTextureViewportParameters(SceneDepthViewport);
 			PassParameters->DepthTexture = Inputs.SceneDepth.Texture;
 			PassParameters->DepthSampler = TStaticSamplerState<SF_Point, AM_Clamp, AM_Clamp, AM_Clamp>::GetRHI();
-			PassParameters->RenderTargets[0] = FRenderTargetBinding(SceneTextures.EditorPrimitiveColor, ERenderTargetLoadAction::EClear);
-			PassParameters->RenderTargets.DepthStencil = FDepthStencilBinding(SceneTextures.EditorPrimitiveDepth, ERenderTargetLoadAction::EClear, ERenderTargetLoadAction::EClear, FExclusiveDepthStencil::DepthWrite_StencilWrite);
+			PassParameters->RenderTargets[0] = FRenderTargetBinding(EditorPrimitiveColor, ERenderTargetLoadAction::EClear);
+			PassParameters->RenderTargets.DepthStencil = FDepthStencilBinding(EditorPrimitiveDepth, ERenderTargetLoadAction::EClear, ERenderTargetLoadAction::EClear, FExclusiveDepthStencil::DepthWrite_StencilWrite);
 
 			FPopulateEditorDepthPS::FPermutationDomain PermutationVector;
 			PermutationVector.Set<FPopulateEditorDepthPS::FUseMSAADimension>(NumSamples > 1);
@@ -290,15 +327,13 @@ FScreenPassTexture AddEditorPrimitivePass(
 				PopulateDepthPixelShader,
 				TStaticDepthStencilState<true, CF_Always>::GetRHI(),
 				PassParameters);
-
-			CompositeLoadAction = ERenderTargetLoadAction::ELoad;
 		}
 
 		{
 			FEditorPrimitivesPassParameters* PassParameters = GraphBuilder.AllocParameters<FEditorPrimitivesPassParameters>();
 			PassParameters->View = EditorView->ViewUniformBuffer;
-			PassParameters->RenderTargets[0] = FRenderTargetBinding(SceneTextures.EditorPrimitiveColor, CompositeLoadAction);
-			PassParameters->RenderTargets.DepthStencil = FDepthStencilBinding(SceneTextures.EditorPrimitiveDepth, ERenderTargetLoadAction::ELoad, ERenderTargetLoadAction::ELoad, FExclusiveDepthStencil::DepthWrite_StencilWrite);
+			PassParameters->RenderTargets[0] = FRenderTargetBinding(EditorPrimitiveColor, ERenderTargetLoadAction::ELoad);
+			PassParameters->RenderTargets.DepthStencil = FDepthStencilBinding(EditorPrimitiveDepth, ERenderTargetLoadAction::ELoad, ERenderTargetLoadAction::ELoad, FExclusiveDepthStencil::DepthWrite_StencilWrite);
 
 			const FEditorPrimitiveInputs::EBasePassType BasePassType = Inputs.BasePassType;
 
@@ -369,8 +404,8 @@ FScreenPassTexture AddEditorPrimitivePass(
 	PassParameters->ColorSampler = PointClampSampler;
 	PassParameters->DepthTexture = Inputs.SceneDepth.Texture;
 	PassParameters->DepthSampler = PointClampSampler;
-	PassParameters->EditorPrimitivesDepth = SceneTextures.EditorPrimitiveDepth;
-	PassParameters->EditorPrimitivesColor = SceneTextures.EditorPrimitiveColor;
+	PassParameters->EditorPrimitivesDepth = EditorPrimitiveDepth;
+	PassParameters->EditorPrimitivesColor = EditorPrimitiveColor;
 	PassParameters->bOpaqueEditorGizmo = bOpaqueEditorGizmo;
 	PassParameters->bCompositeAnyNonNullDepth = bProducedByPriorPass;
 
diff --git a/Engine/Source/Runtime/Renderer/Private/PostProcess/PostProcessSelectionOutline.cpp b/Engine/Source/Runtime/Renderer/Private/PostProcess/PostProcessSelectionOutline.cpp
--- a/Engine/Source/Runtime/Renderer/Private/PostProcess/PostProcessSelectionOutline.cpp
+++ b/Engine/Source/Runtime/Renderer/Private/PostProcess/PostProcessSelectionOutline.cpp
@@ -76,7 +76,7 @@ FScreenPassTexture AddSelectionOutlinePass(FRDGBuilder& GraphBuilder, const FVie
 			DepthStencilDesc.Flags = TexCreate_DepthStencilTargetable | TexCreate_ShaderResource;
 			DepthStencilDesc.NumSamples = NumSamples;
 
-			DepthStencilTexture = GraphBuilder.CreateTexture(DepthStencilDesc, TEXT("SelectionOutline"));
+			DepthStencilTexture = GraphBuilder.CreateTexture(DepthStencilDesc, TEXT("Editor.SelectionOutline"));
 		}
 
 		auto* PassParameters = GraphBuilder.AllocParameters<FSelectionOutlinePassParameters>();
diff --git a/Engine/Source/Runtime/Renderer/Private/SceneTextures.cpp b/Engine/Source/Runtime/Renderer/Private/SceneTextures.cpp
--- a/Engine/Source/Runtime/Renderer/Private/SceneTextures.cpp
+++ b/Engine/Source/Runtime/Renderer/Private/SceneTextures.cpp
@@ -830,10 +830,10 @@ FSceneTextures& FSceneTextures::Create(FRDGBuilder& GraphBuilder, const FSceneTe
 #if WITH_EDITOR
 	{
 		const FRDGTextureDesc ColorDesc(FRDGTextureDesc::Create2D(Config.Extent, PF_B8G8R8A8, FClearValueBinding::Transparent, TexCreate_ShaderResource | TexCreate_RenderTargetable, 1, Config.EditorPrimitiveNumSamples));
-		SceneTextures.EditorPrimitiveColor = GraphBuilder.CreateTexture(ColorDesc, TEXT("EditorPrimitivesColor"));
+		SceneTextures.EditorPrimitiveColor = GraphBuilder.CreateTexture(ColorDesc, TEXT("Editor.PrimitivesColor"));
 
 		const FRDGTextureDesc DepthDesc(FRDGTextureDesc::Create2D(Config.Extent, PF_DepthStencil, FClearValueBinding::DepthFar, TexCreate_ShaderResource | TexCreate_DepthStencilTargetable, 1, Config.EditorPrimitiveNumSamples));
-		SceneTextures.EditorPrimitiveDepth = GraphBuilder.CreateTexture(DepthDesc, TEXT("EditorPrimitivesDepth"));
+		SceneTextures.EditorPrimitiveDepth = GraphBuilder.CreateTexture(DepthDesc, TEXT("Editor.PrimitivesDepth"));
 	}
 #endif
 
