Added common util class for DepthBounds tests setup.
Added a function to exclude sky pixels used by IndirectLight/SkyEnvLightPixel and SkyAtmosphere.
This reduces pointless GPU wavefront with early exit on sky pixels.

diff --git a/Engine/Source/Runtime/Renderer/Private/IndirectLightRendering.cpp b/Engine/Source/Runtime/Renderer/Private/IndirectLightRendering.cpp
--- a/Engine/Source/Runtime/Renderer/Private/IndirectLightRendering.cpp
+++ b/Engine/Source/Runtime/Renderer/Private/IndirectLightRendering.cpp
@@ -22,6 +22,7 @@
 #include "Math/Halton.h"
 #include "DistanceFieldAmbientOcclusion.h"
 #include "Strata/Strata.h"
+#include "RendererUtils.h"
 
 // Must match EDynamicGlobalIlluminationMethod
 // Note: Default for new projects set by GameProjectUtils
@@ -1293,6 +1294,15 @@ static void AddSkyReflectionPass(
 			}
 		}
 
+		if (GSupportsDepthBoundsTest)
+		{
+			// We do not want to process sky pixels so we take advantage of depth bound test when available to skip launching pointless GPU wavefront/work.
+			GraphicsPSOInit.bDepthBounds = true;
+
+			FDepthBounds::FDepthBoundsValues Values = FDepthBounds::CalculateNearFarDepthExcludingSky();
+			InRHICmdList.SetDepthBounds(Values.MinDepth, Values.MaxDepth);
+		}
+
 		SetGraphicsPipelineState(InRHICmdList, GraphicsPSOInit);
 		SetShaderParameters(InRHICmdList, PixelShader, PixelShader.GetPixelShader(), *PassParameters);
 		if (Strata::IsStrataEnabled() && Strata::IsClassificationEnabled())
diff --git a/Engine/Source/Runtime/Renderer/Private/RendererUtils.cpp b/Engine/Source/Runtime/Renderer/Private/RendererUtils.cpp
--- a/Engine/Source/Runtime/Renderer/Private/RendererUtils.cpp
+++ b/Engine/Source/Runtime/Renderer/Private/RendererUtils.cpp
@@ -2,6 +2,7 @@
 
 #include "RendererUtils.h"
 #include "RenderTargetPool.h"
+#include "RHIDefinitions.h"
 #include "VisualizeTexture.h"
 
 class FRTWriteMaskDecodeCS : public FGlobalShader
@@ -162,4 +163,28 @@ void FRenderTargetWriteMask::Decode(
 			FMath::DivideAndRoundUp((uint32)RTWriteMaskDims.Y, FRTWriteMaskDecodeCS::ThreadGroupSizeY),
 			1);
 	});
-}
\ No newline at end of file
+}
+
+FDepthBounds::FDepthBoundsValues FDepthBounds::CalculateNearFarDepthExcludingSky()
+{
+	FDepthBounds::FDepthBoundsValues Values;
+
+	if (bool(ERHIZBuffer::IsInverted))
+	{
+		//const float SmallestFloatAbove0 = 1.1754943508e-38;		// 32bit float depth
+		const float SmallestFloatAbove0 = 1.0f / 16777215.0f;		// 24bit norm depth
+
+		Values.MinDepth = SmallestFloatAbove0;
+		Values.MaxDepth = float(ERHIZBuffer::NearPlane);
+	}
+	else
+	{
+		//const float SmallestFloatBelow1 = 0.9999999404;			// 32bit float depth
+		const float SmallestFloatBelow1 = 16777214.0f / 16777215.0f;// 24bit norm depth
+
+		Values.MinDepth = float(ERHIZBuffer::NearPlane);
+		Values.MaxDepth = SmallestFloatBelow1;
+	}
+
+	return Values;
+}
diff --git a/Engine/Source/Runtime/Renderer/Private/SkyAtmosphereRendering.cpp b/Engine/Source/Runtime/Renderer/Private/SkyAtmosphereRendering.cpp
--- a/Engine/Source/Runtime/Renderer/Private/SkyAtmosphereRendering.cpp
+++ b/Engine/Source/Runtime/Renderer/Private/SkyAtmosphereRendering.cpp
@@ -16,6 +16,7 @@
 #include "SceneRenderTargetParameters.h"
 #include "VolumeLighting.h"
 #include "VolumetricCloudRendering.h"
+#include "RendererUtils.h"
 
 
 //PRAGMA_DISABLE_OPTIMIZATION
@@ -1724,20 +1725,11 @@ void FSceneRenderer::RenderSkyAtmosphereInternal(
 
 			if (!bRenderSkyPixel && GSupportsDepthBoundsTest)
 			{
-				// When we do not render the sky in the sky pass and depth bound test is supported, we take advantage of it in order to skip the processing of sky pixels.
+				// We do not want to process sky pixels so we take advantage of depth bound test when available to skip launching pointless GPU wavefront/work.
 				GraphicsPSOInit.bDepthBounds = true;
-				if (bool(ERHIZBuffer::IsInverted))
-				{
-					//const float SmallestFloatAbove0 = 1.1754943508e-38;		// 32bit float depth
-					const float SmallestFloatAbove0 = 1.0f / 16777215.0f;		// 24bit norm depth
-					RHICmdListLambda.SetDepthBounds(SmallestFloatAbove0, 1.0f);	// Tested on dx12 PC
-				}
-				else
-				{
-					//const float SmallestFloatBelow1 = 0.9999999404;			// 32bit float depth
-					const float SmallestFloatBelow1 = 16777214.0f / 16777215.0f;// 24bit norm depth
-					RHICmdListLambda.SetDepthBounds(0.0f, SmallestFloatBelow1);	// Untested
-				}
+
+				FDepthBounds::FDepthBoundsValues Values = FDepthBounds::CalculateNearFarDepthExcludingSky();
+				RHICmdListLambda.SetDepthBounds(Values.MinDepth, Values.MaxDepth);
 			}
 
 			SetGraphicsPipelineState(RHICmdListLambda, GraphicsPSOInit);
diff --git a/Engine/Source/Runtime/Renderer/Public/RendererUtils.h b/Engine/Source/Runtime/Renderer/Public/RendererUtils.h
--- a/Engine/Source/Runtime/Renderer/Public/RendererUtils.h
+++ b/Engine/Source/Runtime/Renderer/Public/RendererUtils.h
@@ -23,4 +23,17 @@ public:
 		FRDGTextureRef& OutRTWriteMask,
 		ETextureCreateFlags RTWriteMaskFastVRamConfig,
 		const TCHAR* RTWriteMaskDebugName);
-};	
\ No newline at end of file
+};
+
+class RENDERER_API FDepthBounds
+{
+public:
+
+	struct FDepthBoundsValues
+	{
+		float MinDepth;
+		float MaxDepth;
+	};
+
+	static FDepthBoundsValues CalculateNearFarDepthExcludingSky();
+};
