Allow backface culling on primary rays in path tracer

This is required for rendering the Archvis sample correctly as the cameras sometimes go through walls (which the raster path is ok with). A Cvar lets the user disable this behavior if it isn't needed.

Also minor fix to pickup pre-integrated skin SSS color like we do for Subsurface materials until we can implement these properly

diff --git a/Engine/Shaders/Private/PathTracing/PathTracing.usf b/Engine/Shaders/Private/PathTracing/PathTracing.usf
--- a/Engine/Shaders/Private/PathTracing/PathTracing.usf
+++ b/Engine/Shaders/Private/PathTracing/PathTracing.usf
@@ -40,14 +40,14 @@ void AccumulateRadiance(inout float3 TotalRadiance, float3 PathRadiance)
 	}
 }
 
-FMaterialClosestHitPayload TraceTransparentRay(RayDesc Ray, FRayCone RayCone, bool LastBounce, uint2 LaunchIndex, uint NumLights, inout RandomSequence RandSequence, inout float3 PathThroughput, inout float3 Radiance)
+FMaterialClosestHitPayload TraceTransparentRay(RayDesc Ray, FRayCone RayCone, bool IsCameraRay, bool LastBounce, uint2 LaunchIndex, uint NumLights, inout RandomSequence RandSequence, inout float3 PathThroughput, inout float3 Radiance)
 {
 	float SelectionWeightSum = 0;
 	float3 PayloadThroughput = PathThroughput;
 	FMaterialClosestHitPayload Payload;
 	for (;;)
 	{
-		const uint InitialRayFlags = 0;
+		const uint InitialRayFlags = IsCameraRay && PathTracingData.EnableCameraBackfaceCulling ? RAY_FLAG_CULL_BACK_FACING_TRIANGLES : 0;
 		const bool bInitialDisableSkyLightContribution = false;
 		const bool bIgnoreTranslucentMaterials = false;
 		FMaterialClosestHitPayload HitPayload = TraceMaterialRay(
@@ -76,7 +76,8 @@ FMaterialClosestHitPayload TraceTransparentRay(RayDesc Ray, FRayCone RayCone, bo
 		}
 
 		// material tweaks to workaround temporary limitations
-		if (HitPayload.ShadingModelID == SHADINGMODELID_SUBSURFACE)
+		if (HitPayload.ShadingModelID == SHADINGMODELID_SUBSURFACE ||
+			HitPayload.ShadingModelID == SHADINGMODELID_PREINTEGRATED_SKIN)
 		{
 			// for now just approximate SSS by moving it into the diffuse color
 			HitPayload.DiffuseColor += HitPayload.CustomData.xyz;
@@ -265,7 +266,7 @@ RAY_TRACING_ENTRY_RAYGEN(PathTracingMainRG)
 
 	for (int Bounce = 0; Bounce <= PathTracingData.MaxBounces; Bounce++)
 	{
-		FMaterialClosestHitPayload Payload = TraceTransparentRay(Ray, RayCone, Bounce == PathTracingData.MaxBounces, LaunchIndex, NumLights, RandSequence, PathThroughput, Radiance);
+		FMaterialClosestHitPayload Payload = TraceTransparentRay(Ray, RayCone, Bounce == 0, Bounce == PathTracingData.MaxBounces, LaunchIndex, NumLights, RandSequence, PathThroughput, Radiance);
 
 		if (Payload.IsMiss())
 		{
diff --git a/Engine/Source/Runtime/Renderer/Private/DeferredShadingRenderer.cpp b/Engine/Source/Runtime/Renderer/Private/DeferredShadingRenderer.cpp
--- a/Engine/Source/Runtime/Renderer/Private/DeferredShadingRenderer.cpp
+++ b/Engine/Source/Runtime/Renderer/Private/DeferredShadingRenderer.cpp
@@ -617,6 +617,7 @@ bool FDeferredShadingSceneRenderer::GatherRayTracingWorldInstances(FRHICommandLi
 		bool bAllSegmentsOpaque = true;
 		bool bAnySegmentsCastShadow = false;
 		bool bAnySegmentsDecal = false;
+		bool bTwoSided = false;
 
 		uint64 InstancingKey() const
 		{
@@ -625,6 +626,7 @@ bool FDeferredShadingSceneRenderer::GatherRayTracingWorldInstances(FRHICommandLi
 			Key ^= bAllSegmentsOpaque ? 0x1ull << 40 : 0x0;
 			Key ^= bAnySegmentsCastShadow ? 0x1ull << 41 : 0x0;
 			Key ^= bAnySegmentsDecal ? 0x1ull << 42 : 0x0;
+			Key ^= bTwoSided ? 0x1ull << 43 : 0x0;
 			return Key ^ reinterpret_cast<uint64>(RayTracingGeometryRHI);
 		}
 	};
@@ -846,6 +848,7 @@ bool FDeferredShadingSceneRenderer::GatherRayTracingWorldInstances(FRHICommandLi
 								RelevantPrimitive.bAllSegmentsOpaque &= RayTracingMeshCommand.bOpaque;
 								RelevantPrimitive.bAnySegmentsCastShadow |= RayTracingMeshCommand.bCastRayTracedShadows;
 								RelevantPrimitive.bAnySegmentsDecal |= RayTracingMeshCommand.bDecal;
+								RelevantPrimitive.bTwoSided |= RayTracingMeshCommand.bTwoSided;
 							}
 							else
 							{
@@ -929,7 +932,12 @@ bool FDeferredShadingSceneRenderer::GatherRayTracingWorldInstances(FRHICommandLi
 						RayTracingInstance.Transforms.SetNumUninitialized(Instance.InstanceTransforms.Num());
 						FMemory::Memcpy(RayTracingInstance.Transforms.GetData(), Instance.InstanceTransforms.GetData(), Instance.InstanceTransforms.Num() * sizeof(RayTracingInstance.Transforms[0]));
 						static_assert(TIsSame<decltype(RayTracingInstance.Transforms[0]), decltype(Instance.InstanceTransforms[0])>::Value, "Unexpected transform type");
-					}			
+					}
+					for (int32 SegmentIndex = 0; SegmentIndex < Instance.Materials.Num(); SegmentIndex++)
+					{
+						FMeshBatch& MeshBatch = Instance.Materials[SegmentIndex];
+						RayTracingInstance.bDoubleSided |= MeshBatch.bDisableBackfaceCulling;
+					}
 
 					uint32 InstanceIndex = ReferenceView.RayTracingGeometryInstances.Add(RayTracingInstance);
 
@@ -1103,6 +1111,7 @@ bool FDeferredShadingSceneRenderer::GatherRayTracingWorldInstances(FRHICommandLi
 				RayTracingInstance.UserData[0] = (uint32)PrimitiveIndex;
 				RayTracingInstance.Mask = RelevantPrimitive.InstanceMask; // When no cached command is found, InstanceMask == 0 and the instance is effectively filtered out
 				RayTracingInstance.bForceOpaque = RelevantPrimitive.bAllSegmentsOpaque;
+				RayTracingInstance.bDoubleSided = RelevantPrimitive.bTwoSided;
 			}
 		}
 	}
diff --git a/Engine/Source/Runtime/Renderer/Private/PathTracing.cpp b/Engine/Source/Runtime/Renderer/Private/PathTracing.cpp
--- a/Engine/Source/Runtime/Renderer/Private/PathTracing.cpp
+++ b/Engine/Source/Runtime/Renderer/Private/PathTracing.cpp
@@ -78,6 +78,13 @@ TAutoConsoleVariable<int32> CVarPathTracingApproximateCaustics(
 	ECVF_RenderThreadSafe
 );
 
+TAutoConsoleVariable<int32> CVarPathTracingEnableCameraBackfaceCulling(
+	TEXT("r.PathTracing.EnableCameraBackfaceCulling"),
+	1,
+	TEXT("When non-zero, the path tracer will skip over backfacing triangles when tracing primary rays from the camera. (default = 1 (enabled))"),
+	ECVF_RenderThreadSafe
+);
+
 
 TAutoConsoleVariable<int32> CVarPathTracingFrameIndependentTemporalSeed(
 	TEXT("r.PathTracing.FrameIndependentTemporalSeed"),
@@ -128,6 +135,7 @@ static bool PrepareShaderArgs(const FViewInfo& View, FPathTracingData& PathTraci
 	PathTracingData.MaxPathIntensity = CVarPathTracingMaxPathIntensity.GetValueOnRenderThread();
 	PathTracingData.UseErrorDiffusion = CVarPathTracingUseErrorDiffusion.GetValueOnRenderThread();
 	PathTracingData.ApproximateCaustics = CVarPathTracingApproximateCaustics.GetValueOnRenderThread();
+	PathTracingData.EnableCameraBackfaceCulling = CVarPathTracingEnableCameraBackfaceCulling.GetValueOnRenderThread();
 	float FilterWidth = CVarPathTracingFilterWidth.GetValueOnRenderThread();
 	if (FilterWidth < 0)
 	{
@@ -195,6 +203,14 @@ static bool PrepareShaderArgs(const FViewInfo& View, FPathTracingData& PathTraci
 		PreviousFilterWidth = PathTracingData.FilterWidth;
 	}
 
+	// Changing backface culling status requires starting over
+	static uint32 PreviousBackfaceCulling = PathTracingData.EnableCameraBackfaceCulling;
+	if (PreviousBackfaceCulling != PathTracingData.EnableCameraBackfaceCulling)
+	{
+		NeedInvalidation = true;
+		PreviousBackfaceCulling = PathTracingData.EnableCameraBackfaceCulling;
+	}
+
 	// the rest of PathTracingData and AdaptiveSamplingData is filled in by SetParameters below
 	return NeedInvalidation;
 }
diff --git a/Engine/Source/Runtime/Renderer/Private/PathTracingUniformBuffers.h b/Engine/Source/Runtime/Renderer/Private/PathTracingUniformBuffers.h
--- a/Engine/Source/Runtime/Renderer/Private/PathTracingUniformBuffers.h
+++ b/Engine/Source/Runtime/Renderer/Private/PathTracingUniformBuffers.h
@@ -18,6 +18,7 @@ BEGIN_GLOBAL_SHADER_PARAMETER_STRUCT(FPathTracingData, )
 	SHADER_PARAMETER(uint32, MISMode)
 	SHADER_PARAMETER(uint32, VisibleLights)
 	SHADER_PARAMETER(uint32, ApproximateCaustics)
+	SHADER_PARAMETER(uint32, EnableCameraBackfaceCulling)
 	SHADER_PARAMETER(float, MaxPathIntensity)
 	SHADER_PARAMETER(float, MaxNormalBias)
 	SHADER_PARAMETER(float, FilterWidth)
diff --git a/Engine/Source/Runtime/Renderer/Private/RayTracing/RayTracingMaterialHitShaders.h b/Engine/Source/Runtime/Renderer/Private/RayTracing/RayTracingMaterialHitShaders.h
--- a/Engine/Source/Runtime/Renderer/Private/RayTracing/RayTracingMaterialHitShaders.h
+++ b/Engine/Source/Runtime/Renderer/Private/RayTracing/RayTracingMaterialHitShaders.h
@@ -66,6 +66,7 @@ protected:
 		SharedCommand.bCastRayTracedShadows = MeshBatch.CastRayTracedShadow && MaterialResource.CastsRayTracedShadows();
 		SharedCommand.bOpaque = MaterialResource.GetBlendMode() == EBlendMode::BLEND_Opaque;
 		SharedCommand.bDecal = MaterialResource.GetMaterialDomain() == EMaterialDomain::MD_DeferredDecal;
+		SharedCommand.bTwoSided = MaterialResource.IsTwoSided();
 
 		FVertexInputStreamArray VertexStreams;
 		VertexFactory->GetStreams(ERHIFeatureLevel::SM5, EVertexInputStreamType::Default, VertexStreams);
diff --git a/Engine/Source/Runtime/Renderer/Public/MeshPassProcessor.h b/Engine/Source/Runtime/Renderer/Public/MeshPassProcessor.h
--- a/Engine/Source/Runtime/Renderer/Public/MeshPassProcessor.h
+++ b/Engine/Source/Runtime/Renderer/Public/MeshPassProcessor.h
@@ -1870,6 +1870,7 @@ public:
 	bool bCastRayTracedShadows = true;
 	bool bOpaque = true;
 	bool bDecal = false;
+	bool bTwoSided = false;
 
 	/** Sets ray hit group shaders on the mesh command and allocates room for the shader bindings. */
 	RENDERER_API void SetShaders(const FMeshProcessorShaders& Shaders);
