Re-enabled default denoiser for SSGI

diff --git a/Engine/Source/Runtime/Renderer/Private/IndirectLightRendering.cpp b/Engine/Source/Runtime/Renderer/Private/IndirectLightRendering.cpp
--- a/Engine/Source/Runtime/Renderer/Private/IndirectLightRendering.cpp
+++ b/Engine/Source/Runtime/Renderer/Private/IndirectLightRendering.cpp
@@ -341,6 +341,7 @@ void FDeferredShadingSceneRenderer::CommitIndirectLightingState()
 		else if (ScreenSpaceRayTracing::IsScreenSpaceDiffuseIndirectSupported(View))
 		{
 			DiffuseIndirectMethod = EDiffuseIndirectMethod::SSGI;
+			DiffuseIndirectDenoiser = IScreenSpaceDenoiser::GetDenoiserMode(CVarDiffuseIndirectDenoiser);
 		}
 		else if (ShouldRenderRayTracingGlobalIllumination(View))
 		{
