Implements r.TemporalAA.ShadingRejection.HalfRes

diff --git a/Engine/Shaders/Private/TemporalAA/TAACompareHistory.usf b/Engine/Shaders/Private/TemporalAA/TAACompareHistory.usf
--- a/Engine/Shaders/Private/TemporalAA/TAACompareHistory.usf
+++ b/Engine/Shaders/Private/TemporalAA/TAACompareHistory.usf
@@ -45,10 +45,10 @@ void MainCS(
 		ZOrder2D(GroupThreadIndex, log2(TILE_SIZE)) +
 		GroupId * uint2(TILE_SIZE, TILE_SIZE));
 
-	uint2 InputPixelPos = InputInfo_ViewportMin + DispatchThreadId;
+	uint2 InputPixelPos = LowFrequencyInfo_ViewportMin + DispatchThreadId;
 
-	float2 InputBufferUV = (InputPixelPos + 0.5) * InputInfo_ExtentInverse;
-	InputBufferUV = clamp(InputBufferUV, InputInfo_UVViewportBilinearMin, InputInfo_UVViewportBilinearMax);
+	float2 InputBufferUV = (InputPixelPos + 0.5) * LowFrequencyInfo_ExtentInverse;
+	InputBufferUV = clamp(InputBufferUV, LowFrequencyInfo_UVViewportBilinearMin, LowFrequencyInfo_UVViewportBilinearMax);
 
 	taa_half3 InputMinColor;
 	taa_half3 InputMaxColor;
@@ -71,8 +71,8 @@ void MainCS(
 			continue;
 		}
 
-		float2 SampleInputBufferUV = InputBufferUV + float2(kOffsets3x3[i]) * InputInfo_ExtentInverse;
-		SampleInputBufferUV = clamp(SampleInputBufferUV, InputInfo_UVViewportBilinearMin, InputInfo_UVViewportBilinearMax);
+		float2 SampleInputBufferUV = InputBufferUV + float2(kOffsets3x3[i]) * LowFrequencyInfo_ExtentInverse;
+		SampleInputBufferUV = clamp(SampleInputBufferUV, LowFrequencyInfo_UVViewportBilinearMin, LowFrequencyInfo_UVViewportBilinearMax);
 
 		taa_half3 Color = TransformColorForClampingBox(FilteredInputTexture.SampleLevel(GlobalPointClampedSampler, SampleInputBufferUV, 0));
 
@@ -115,24 +115,24 @@ void MainCS(
 		}
 
 		
-		if (all(InputPixelPos < InputInfo_ViewportMax) && all(((InputPixelPos - InputInfo_ViewportMin) % 2) == 0))
+		if (all(InputPixelPos < LowFrequencyInfo_ViewportMax) && all(((InputPixelPos - LowFrequencyInfo_ViewportMin) % 2) == 0))
 		{
-			HistoryRejectionOutput[(InputPixelPos - InputInfo_ViewportMin) / 2] = Rejection;
+			HistoryRejectionOutput[(InputPixelPos - LowFrequencyInfo_ViewportMin) / 2] = Rejection;
 		}
 	}
 	#endif
 
-	if (all(InputPixelPos < InputInfo_ViewportMax))
+	if (all(InputPixelPos < LowFrequencyInfo_ViewportMax))
 	{
 		#if !CONFIG_OUTPUT_HALF_RES
 		{
-			HistoryRejectionOutput[(InputPixelPos - InputInfo_ViewportMin)] = Rejection;
+			HistoryRejectionOutput[(InputPixelPos - LowFrequencyInfo_ViewportMin)] = Rejection;
 		}
 		#endif
 
 		#if DEBUG_OUTPUT
 		{
-			DebugOutput[InputPixelPos - InputInfo_ViewportMin] = Debug;
+			DebugOutput[InputPixelPos - LowFrequencyInfo_ViewportMin] = Debug;
 		}
 		#endif
 	}
diff --git a/Engine/Shaders/Private/TemporalAA/TAADecimateHistory.usf b/Engine/Shaders/Private/TemporalAA/TAADecimateHistory.usf
--- a/Engine/Shaders/Private/TemporalAA/TAADecimateHistory.usf
+++ b/Engine/Shaders/Private/TemporalAA/TAADecimateHistory.usf
@@ -9,6 +9,8 @@
 
 #define CONFIG_KERNEL 0
 
+#define CONFIG_OUTPUT_HALF_RES (DIM_OUTPUT_HALF_RES)
+
 
 //------------------------------------------------------- PARAMETERS
 
@@ -17,6 +19,7 @@ float3 OutputQuantizationError;
 float HistoryPreExposureCorrection;
 uint bCameraCut;
 
+Texture2D<taa_half3> InputSceneColorTexture;
 Texture2D<float2> DilatedVelocityTexture;
 Texture2D<float> ClosestDepthTexture;
 Texture2D<uint> PrevUseCountTexture;
@@ -28,10 +31,20 @@ Texture2D<taa_half3> PrevHistory_Textures_1;
 Texture2D<taa_half2> PrevHistory_Textures_2;
 Texture2D<uint> PrevHistory_Textures_3;
 
+RWTexture2D<taa_half3> HalfResSceneColorOutput;
+RWTexture2D<taa_half3> HalfResPredictionSceneColorOutput;
+RWTexture2D<taa_half> HalfResParallaxRejectionMaskOutput;
+
 RWTexture2D<taa_half3> PredictionSceneColorOutput;
 RWTexture2D<taa_half> ParallaxRejectionMaskOutput;
 
 
+//------------------------------------------------------- LDS
+
+groupshared taa_half3 SharedInputColor[TILE_SIZE * TILE_SIZE];
+groupshared taa_half4 SharedHistoryColorAndMask[TILE_SIZE * TILE_SIZE];
+
+
 //------------------------------------------------------- ENTRY POINT
 
 taa_half3 TransformToComparisonColorSpace(taa_half3 Color)
@@ -64,6 +77,7 @@ void MainCS(
 	float2 EncodedVelocity = DilatedVelocityTexture[InputPixelPos];
 	float DeviceZ = ClosestDepthTexture[InputPixelPos];
 	taa_half ParallaxFactor = UncompressParallaxFactor((ParallaxFactorTexture[InputPixelPos] >> 4) & 0xF);
+	taa_half3 InputSceneColor = InputSceneColorTexture[InputPixelPos];
 
 
 	float WorldDepth = ConvertFromDeviceZ(DeviceZ);
@@ -171,15 +185,6 @@ void MainCS(
 		ParallaxRejectionMask = min(ParallaxRejectionMask, HistoryParallaxRejectionMask);
 	}
 
-	#if CONFIG_ENABLE_STOCASTIC_QUANTIZATION
-	{
-		uint2 Random = Rand3DPCG16(int3(InputPixelPos - InputInfo_ViewportMin, View.StateFrameIndexMod8)).xy;
-		float2 E = Hammersley16(0, 1, Random);
-
-		DecimatedColor = QuantizeForFloatRenderTarget(DecimatedColor, E.x, OutputQuantizationError);
-	}
-	#endif
-
 	// Immediately reject based on the rejection mask
 	FLATTEN
 	if ((PrevOutputScreenPos.x < -1.0) ||
@@ -191,9 +196,77 @@ void MainCS(
 		HistoryValidity = taa_half(0.0);
 	}
 
+	
+	#if CONFIG_OUTPUT_HALF_RES
+	{
+		SharedInputColor[GroupThreadIndex] = InputSceneColor;
+		SharedHistoryColorAndMask[GroupThreadIndex] = taa_half4(DecimatedColor, ParallaxRejectionMask);
+
+		GroupMemoryBarrierWithGroupSync();
+
+		bool bCheckerboard = (((InputPixelPos.x ^ InputPixelPos.y) >> 1) & 0x1) != 0;
+
+		half3 HalfResInputSceneColor = InputSceneColor;
+		half3 HalfResDecimatedColor = DecimatedColor;
+		half HalfResParallaxRejectionMask = ParallaxRejectionMask;
+
+		UNROLL
+		for (uint i = 1; i < 4; i++)
+		{
+			uint ButterflySwap = i;
+			
+			if (bCheckerboard)
+			{
+				HalfResInputSceneColor = min(HalfResInputSceneColor, SharedInputColor[GroupThreadIndex ^ ButterflySwap]);
+			}
+			else
+			{
+				HalfResInputSceneColor = max(HalfResInputSceneColor, SharedInputColor[GroupThreadIndex ^ ButterflySwap]);
+			}
+			HalfResDecimatedColor += SharedHistoryColorAndMask[GroupThreadIndex ^ ButterflySwap].rgb;
+
+			HalfResParallaxRejectionMask += SharedHistoryColorAndMask[GroupThreadIndex ^ ButterflySwap].a;
+		}
+
+		HalfResDecimatedColor *= 0.25;
+		HalfResParallaxRejectionMask *= 0.25;
+
+		#if CONFIG_ENABLE_STOCASTIC_QUANTIZATION
+		{
+			uint2 Random = Rand3DPCG16(int3((InputPixelPos - InputInfo_ViewportMin) / 2, View.StateFrameIndexMod8)).xy;
+			float2 E = Hammersley16(0, 1, Random);
+
+			HalfResInputSceneColor = QuantizeForFloatRenderTarget(HalfResInputSceneColor, E.x, OutputQuantizationError);
+			HalfResDecimatedColor = QuantizeForFloatRenderTarget(HalfResDecimatedColor, E.x, OutputQuantizationError);
+		}
+		#endif
+
+		if (all(InputPixelPos < InputInfo_ViewportMax) && all(((InputPixelPos - InputInfo_ViewportMin) % 2) == 0))
+		{
+			HalfResSceneColorOutput[InputPixelPos / 2] = HalfResInputSceneColor;
+			HalfResPredictionSceneColorOutput[InputPixelPos / 2] = HalfResDecimatedColor;
+			HalfResParallaxRejectionMaskOutput[InputPixelPos / 2] = HalfResParallaxRejectionMask;
+		}
+	}
+	#endif
+
+	#if CONFIG_ENABLE_STOCASTIC_QUANTIZATION
+	{
+		uint2 Random = Rand3DPCG16(int3(InputPixelPos - InputInfo_ViewportMin, View.StateFrameIndexMod8)).xy;
+		float2 E = Hammersley16(0, 1, Random);
+
+		DecimatedColor = QuantizeForFloatRenderTarget(DecimatedColor, E.x, OutputQuantizationError);
+	}
+	#endif
+
 	if (all(InputPixelPos < InputInfo_ViewportMax))
 	{
-		PredictionSceneColorOutput[InputPixelPos] = DecimatedColor;
+		#if !CONFIG_OUTPUT_HALF_RES
+		{
+			PredictionSceneColorOutput[InputPixelPos] = DecimatedColor;
+		}
+		#endif
+		
 		ParallaxRejectionMaskOutput[InputPixelPos] = ParallaxRejectionMask;
 		
 		#if DEBUG_OUTPUT
diff --git a/Engine/Shaders/Private/TemporalAA/TAAFilterFrequencies.usf b/Engine/Shaders/Private/TemporalAA/TAAFilterFrequencies.usf
--- a/Engine/Shaders/Private/TemporalAA/TAAFilterFrequencies.usf
+++ b/Engine/Shaders/Private/TemporalAA/TAAFilterFrequencies.usf
@@ -129,9 +129,9 @@ void MainCS(
 		ZOrder2D(GroupThreadIndex, log2(TILE_SIZE)) +
 		GroupId * uint2(TILE_SIZE, TILE_SIZE));
 
-	uint2 InputPixelPos = InputInfo_ViewportMin + DispatchThreadId;
+	uint2 InputPixelPos = LowFrequencyInfo_ViewportMin + DispatchThreadId;
 
-	float2 InputBufferUV = (InputPixelPos + 0.5) * InputInfo_ExtentInverse;
+	float2 InputBufferUV = (InputPixelPos + 0.5) * LowFrequencyInfo_ExtentInverse;
 
 	// Fetch center of the kernel
 	taa_half3 CenterInput;
@@ -186,7 +186,7 @@ void MainCS(
 	{
 		StartSampleId = 1;
 		SampleCount = 16;
-		KernelCenterUV = float2(InputPixelPos | 1) * InputInfo_ExtentInverse;
+		KernelCenterUV = float2(InputPixelPos | 1) * LowFrequencyInfo_ExtentInverse;
 		SampleTrackId = ((InputPixelPos.x & 1) | ((InputPixelPos.y & 1) << 1));
 	}
 	#else
@@ -246,8 +246,8 @@ void MainCS(
 			#error Unknown kernel
 		#endif
 
-		float2 SampleInputBufferUV = InputBufferUV + float2(CONFIG_KERNEL_SPREAD_FACTOR * Offset) * InputInfo_ExtentInverse;
-		SampleInputBufferUV = clamp(SampleInputBufferUV, InputInfo_UVViewportBilinearMin, InputInfo_UVViewportBilinearMax);
+		float2 SampleInputBufferUV = InputBufferUV + float2(CONFIG_KERNEL_SPREAD_FACTOR * Offset) * LowFrequencyInfo_ExtentInverse;
+		SampleInputBufferUV = clamp(SampleInputBufferUV, LowFrequencyInfo_UVViewportBilinearMin, LowFrequencyInfo_UVViewportBilinearMax);
 
 		taa_half3 InputSample = InputTexture.SampleLevel(GlobalPointClampedSampler, SampleInputBufferUV, 0);
 		taa_half3 HistorySample = PredictionSceneColorTexture.SampleLevel(GlobalPointClampedSampler, SampleInputBufferUV, 0);
@@ -275,7 +275,7 @@ void MainCS(
 	}
 	
 	#if 1
-		uint2 OutputPixelPos = uint2(InputBufferUV * InputInfo_Extent);
+		uint2 OutputPixelPos = uint2(InputBufferUV * LowFrequencyInfo_Extent);
 	#else
 		uint2 OutputPixelPos = InputPixelPos;
 	#endif
@@ -291,14 +291,14 @@ void MainCS(
 	#endif
 
 
-	if (all(OutputPixelPos < InputInfo_ViewportMax))
+	if (all(OutputPixelPos < LowFrequencyInfo_ViewportMax))
 	{
 		FilteredInputOutput[OutputPixelPos] = FilteredInput; // * (HdrWeight4(FilteredInput) * 10);
 		FilteredPredictionSceneColorOutput[OutputPixelPos] = FilteredHistory; // * (HdrWeight4(FilteredHistory) * 10);
 		
 		#if DEBUG_OUTPUT
 		{
-			DebugOutput[OutputPixelPos - InputInfo_ViewportMin] = Debug;
+			DebugOutput[OutputPixelPos - LowFrequencyInfo_ViewportMin] = Debug;
 		}
 		#endif
 	}
diff --git a/Engine/Source/Runtime/Renderer/Private/PostProcess/TemporalAA.cpp b/Engine/Source/Runtime/Renderer/Private/PostProcess/TemporalAA.cpp
--- a/Engine/Source/Runtime/Renderer/Private/PostProcess/TemporalAA.cpp
+++ b/Engine/Source/Runtime/Renderer/Private/PostProcess/TemporalAA.cpp
@@ -81,9 +81,14 @@ TAutoConsoleVariable<int32> CVarTAAR11G11B10History(
 	TEXT("Select the bitdepth of the history."),
 	ECVF_RenderThreadSafe);
 
-TAutoConsoleVariable<int32> CVarTAAFilterHistoryRejection(
-	TEXT("r.TemporalAA.FilterHistoryRejection"), 1,
-	TEXT("Whether the history rejection statistical filtering pass should be enabled or not."),
+TAutoConsoleVariable<int32> CVarTAAHalfResShadingRejection(
+	TEXT("r.TemporalAA.ShadingRejection.HalfRes"), 0,
+	TEXT("Whether the shading rejection should be done at half res. Saves performance but may introduce back some flickering (default = 0)."),
+	ECVF_RenderThreadSafe);
+
+TAutoConsoleVariable<int32> CVarTAAFilterShadingRejection(
+	TEXT("r.TemporalAA.ShadingRejection.SpatialFilter"), 1,
+	TEXT("Whether the shading rejection should have spatial statistical filtering pass to reduce flickering (default = 1)."),
 	ECVF_RenderThreadSafe);
 
 inline bool DoesPlatformSupportTemporalHistoryUpscale(EShaderPlatform Platform)
@@ -348,6 +353,10 @@ class FTAADecimateHistoryCS : public FTAAGen5Shader
 	DECLARE_GLOBAL_SHADER(FTAADecimateHistoryCS);
 	SHADER_USE_PARAMETER_STRUCT(FTAADecimateHistoryCS, FTAAGen5Shader);
 
+	class FOutputHalfRes : SHADER_PERMUTATION_BOOL("DIM_OUTPUT_HALF_RES");
+
+	using FPermutationDomain = TShaderPermutationDomain<FOutputHalfRes>;
+
 	BEGIN_SHADER_PARAMETER_STRUCT(FParameters, )
 		SHADER_PARAMETER_STRUCT_INCLUDE(FTAACommonParameters, CommonParameters)
 		SHADER_PARAMETER(FMatrix, RotationalClipToPrevClip)
@@ -356,6 +365,7 @@ class FTAADecimateHistoryCS : public FTAAGen5Shader
 		SHADER_PARAMETER(float, WorldDepthToPixelWorldRadius)
 		SHADER_PARAMETER(int32, bCameraCut)
 
+		SHADER_PARAMETER_RDG_TEXTURE(Texture2D, InputSceneColorTexture)
 		SHADER_PARAMETER_RDG_TEXTURE(Texture2D, DilatedVelocityTexture)
 		SHADER_PARAMETER_RDG_TEXTURE(Texture2D, ClosestDepthTexture)
 		SHADER_PARAMETER_RDG_TEXTURE(Texture2D, PrevUseCountTexture)
@@ -365,6 +375,9 @@ class FTAADecimateHistoryCS : public FTAAGen5Shader
 		SHADER_PARAMETER_STRUCT(FScreenPassTextureViewportParameters, PrevHistoryInfo)
 		SHADER_PARAMETER_STRUCT(FTAAHistoryTextures, PrevHistory)
 
+		SHADER_PARAMETER_RDG_TEXTURE_UAV(RWTexture2D, HalfResSceneColorOutput)
+		SHADER_PARAMETER_RDG_TEXTURE_UAV(RWTexture2D, HalfResPredictionSceneColorOutput)
+		SHADER_PARAMETER_RDG_TEXTURE_UAV(RWTexture2D, HalfResParallaxRejectionMaskOutput)
 		SHADER_PARAMETER_RDG_TEXTURE_UAV(RWTexture2D, PredictionSceneColorOutput)
 		SHADER_PARAMETER_RDG_TEXTURE_UAV(RWTexture2D, ParallaxRejectionMaskOutput)
 		SHADER_PARAMETER_RDG_TEXTURE_UAV(RWTexture2D, DebugOutput)
@@ -1007,12 +1020,20 @@ static void AddGen5MainTemporalAAPasses(
 	// Whether to use camera cut shader permutation or not.
 	bool bCameraCut = !InputHistory.IsValid() || View.bCameraCut;
 
+	bool bHalfResLowFrequency = CVarTAAHalfResShadingRejection.GetValueOnRenderThread() != 0;
+
 	FIntPoint InputExtent = PassInputs.SceneColorTexture->Desc.Extent;
 	FIntRect InputRect = View.ViewRect;
 
 	FIntPoint LowFrequencyExtent = InputExtent;
 	FIntRect LowFrequencyRect = FIntRect(FIntPoint(0, 0), InputRect.Size());
 
+	if (bHalfResLowFrequency)
+	{
+		LowFrequencyExtent = InputExtent / 2;
+		LowFrequencyRect = FIntRect(FIntPoint(0, 0), FIntPoint::DivideAndRoundUp(InputRect.Size(), 2));
+	}
+
 	FIntPoint RejectionExtent = LowFrequencyExtent / 2;
 	FIntRect RejectionRect = FIntRect(FIntPoint(0, 0), FIntPoint::DivideAndRoundUp(LowFrequencyRect.Size(), 2));
 
@@ -1195,19 +1216,19 @@ static void AddGen5MainTemporalAAPasses(
 	}
 
 	// Decimate input to flicker at same frequency as input.
-	FRDGTextureRef PredictionSceneColorTexture;
-	FRDGTextureRef ParallaxRejectionMaskTexture;
+	FRDGTextureRef HalfResInputSceneColorTexture = nullptr;
+	FRDGTextureRef HalfResPredictionSceneColorTexture = nullptr;
+	FRDGTextureRef HalfResParallaxRejectionMaskTexture = nullptr;
+	FRDGTextureRef PredictionSceneColorTexture = nullptr;
+	FRDGTextureRef ParallaxRejectionMaskTexture = nullptr;
 	{
 		{
 			FRDGTextureDesc Desc = FRDGTextureDesc::Create2D(
-				LowFrequencyExtent,
-				PF_FloatR11G11B10,
+				InputExtent,
+				PF_R8,
 				FClearValueBinding::None,
 				/* InFlags = */ TexCreate_ShaderResource | TexCreate_UAV);
 
-			PredictionSceneColorTexture = GraphBuilder.CreateTexture(Desc, TEXT("TAA.Decimated.SceneColor"));
-
-			Desc.Format = PF_R8;
 			ParallaxRejectionMaskTexture = GraphBuilder.CreateTexture(Desc, TEXT("TAA.ParallaxRejectionMask"));
 		}
 
@@ -1222,7 +1243,7 @@ static void AddGen5MainTemporalAAPasses(
 
 			PassParameters->RotationalClipToPrevClip = RotationalInvViewProj * RotationalPrevViewProj;
 		}
-		PassParameters->OutputQuantizationError = ComputePixelFormatQuantizationError(PredictionSceneColorTexture->Desc.Format);
+		PassParameters->OutputQuantizationError = ComputePixelFormatQuantizationError(PF_FloatR11G11B10);
 		PassParameters->HistoryPreExposureCorrection = View.PreExposure / View.PrevViewInfo.SceneColorPreExposure;
 		{
 			float TanHalfFieldOfView = View.ViewMatrices.GetInvProjectionMatrix().M[0][0];
@@ -1232,6 +1253,7 @@ static void AddGen5MainTemporalAAPasses(
 		}
 		PassParameters->bCameraCut = bCameraCut;
 
+		PassParameters->InputSceneColorTexture = PassInputs.SceneColorTexture;
 		PassParameters->DilatedVelocityTexture = DilatedVelocityTexture;
 		PassParameters->ClosestDepthTexture = ClosestDepthTexture;
 		PassParameters->PrevUseCountTexture = PrevUseCountTexture;
@@ -1241,14 +1263,49 @@ static void AddGen5MainTemporalAAPasses(
 		PassParameters->PrevHistoryInfo = PrevHistoryInfo;
 		PassParameters->PrevHistory = PrevHistory;
 
-		PassParameters->PredictionSceneColorOutput = GraphBuilder.CreateUAV(PredictionSceneColorTexture);
+		if (bHalfResLowFrequency)
+		{
+			FRDGTextureDesc Desc = FRDGTextureDesc::Create2D(
+				LowFrequencyExtent,
+				PF_FloatR11G11B10,
+				FClearValueBinding::None,
+				/* InFlags = */ TexCreate_ShaderResource | TexCreate_UAV);
+
+			HalfResInputSceneColorTexture = GraphBuilder.CreateTexture(Desc, TEXT("TAA.HalfResInput"));
+			HalfResPredictionSceneColorTexture = GraphBuilder.CreateTexture(Desc, TEXT("TAA.Prediction.SceneColor"));
+
+			Desc.Format = PF_R8;
+			HalfResParallaxRejectionMaskTexture = GraphBuilder.CreateTexture(Desc, TEXT("TAA.HalfResParallaxRejectionMask"));
+
+			PassParameters->HalfResSceneColorOutput = GraphBuilder.CreateUAV(HalfResInputSceneColorTexture);
+			PassParameters->HalfResPredictionSceneColorOutput = GraphBuilder.CreateUAV(HalfResPredictionSceneColorTexture);
+			PassParameters->HalfResParallaxRejectionMaskOutput = GraphBuilder.CreateUAV(HalfResParallaxRejectionMaskTexture);
+		}
+		else
+		{
+			FRDGTextureDesc Desc = FRDGTextureDesc::Create2D(
+				InputExtent,
+				PF_FloatR11G11B10,
+				FClearValueBinding::None,
+				/* InFlags = */ TexCreate_ShaderResource | TexCreate_UAV);
+
+			PredictionSceneColorTexture = GraphBuilder.CreateTexture(Desc, TEXT("TAA.Prediction.SceneColor"));
+
+			PassParameters->PredictionSceneColorOutput = GraphBuilder.CreateUAV(PredictionSceneColorTexture);
+		}
+
 		PassParameters->ParallaxRejectionMaskOutput = GraphBuilder.CreateUAV(ParallaxRejectionMaskTexture);
 		PassParameters->DebugOutput = CreateDebugUAV(LowFrequencyExtent, TEXT("Debug.TAA.DecimateHistory"));
 
-		TShaderMapRef<FTAADecimateHistoryCS> ComputeShader(View.ShaderMap);
+		FTAADecimateHistoryCS::FPermutationDomain PermutationVector;
+		PermutationVector.Set<FTAADecimateHistoryCS::FOutputHalfRes>(bHalfResLowFrequency);
+
+		TShaderMapRef<FTAADecimateHistoryCS> ComputeShader(View.ShaderMap, PermutationVector);
 		FComputeShaderUtils::AddPass(
 			GraphBuilder,
-			RDG_EVENT_NAME("TAA DecimateHistory %dx%d", InputRect.Width(), InputRect.Height()),
+			RDG_EVENT_NAME("TAA DecimateHistory(%s) %dx%d",
+				bHalfResLowFrequency ? TEXT("HalfResShadingOutput") : TEXT(""),
+				InputRect.Width(), InputRect.Height()),
 			ComputeShader,
 			PassParameters,
 			FComputeShaderUtils::GetGroupCount(InputRect.Size(), 8));
@@ -1276,9 +1333,18 @@ static void AddGen5MainTemporalAAPasses(
 			PassParameters->CommonParameters = CommonParameters;
 			PassParameters->OutputQuantizationError = ComputePixelFormatQuantizationError(FilteredInputTexture->Desc.Format);
 
-			PassParameters->InputTexture = PassInputs.SceneColorTexture;
-			PassParameters->PredictionSceneColorTexture = PredictionSceneColorTexture;
-			PassParameters->ParallaxRejectionMaskTexture = ParallaxRejectionMaskTexture;
+			if (bHalfResLowFrequency)
+			{
+				PassParameters->InputTexture = HalfResInputSceneColorTexture;
+				PassParameters->PredictionSceneColorTexture = HalfResPredictionSceneColorTexture;
+				PassParameters->ParallaxRejectionMaskTexture = HalfResParallaxRejectionMaskTexture;
+			}
+			else
+			{
+				PassParameters->InputTexture = PassInputs.SceneColorTexture;
+				PassParameters->PredictionSceneColorTexture = PredictionSceneColorTexture;
+				PassParameters->ParallaxRejectionMaskTexture = ParallaxRejectionMaskTexture;
+			}
 
 			PassParameters->FilteredInputOutput = GraphBuilder.CreateUAV(FilteredInputTexture);
 			PassParameters->FilteredPredictionSceneColorOutput = GraphBuilder.CreateUAV(FilteredPredictionSceneColorTexture);
@@ -1325,7 +1391,7 @@ static void AddGen5MainTemporalAAPasses(
 	}
 
 	// Post filter the rejection.
-	if (CVarTAAFilterHistoryRejection.GetValueOnRenderThread() != 0)
+	if (CVarTAAFilterShadingRejection.GetValueOnRenderThread() != 0)
 	{
 		FRDGTextureRef FilteredHistoryRejectionTexture = GraphBuilder.CreateTexture(HistoryRejectionTexture->Desc, TEXT("TAA.FilteredHistoryRejection"));
 
@@ -1445,8 +1511,8 @@ static void AddGen5MainTemporalAAPasses(
 
 		FComputeShaderUtils::AddPass(
 			GraphBuilder,
-			RDG_EVENT_NAME("TAA UpdateHistory%s %dx%d", 
-				History.Textures[0]->Desc.Format == PF_FloatR11G11B10 ? TEXT(" R11G11B10") : TEXT(""),
+			RDG_EVENT_NAME("TAA UpdateHistory(%s) %dx%d", 
+				History.Textures[0]->Desc.Format == PF_FloatR11G11B10 ? TEXT("R11G11B10") : TEXT(""),
 				HistorySize.X, HistorySize.Y),
 			ComputeShader,
 			PassParameters,
