Fixed cubemap shadows being inverted and flipped cull winding for !Nanite VSM.

diff --git a/Engine/Shaders/Private/VirtualShadowMaps/PageManagement.usf b/Engine/Shaders/Private/VirtualShadowMaps/PageManagement.usf
--- a/Engine/Shaders/Private/VirtualShadowMaps/PageManagement.usf
+++ b/Engine/Shaders/Private/VirtualShadowMaps/PageManagement.usf
@@ -248,7 +248,7 @@ void GeneratePageFlagsFromPixels(uint3 GroupId : SV_GroupID, uint GroupIndex : S
 				bool bSpotLight = LightData.SpotAnglesAndSourceRadiusPacked.x > -2.0f;
 				if( !bSpotLight )
 				{
-					VirtualShadowMapId += GetCubeFace( ToLight );
+					VirtualShadowMapId += GetCubeFace( -ToLight );
 				}
 
 				uint MipLevel = GetMipLevel(VirtualShadowMapId, Pixel);
diff --git a/Engine/Shaders/Private/VirtualShadowMaps/ProjectionCommon.ush b/Engine/Shaders/Private/VirtualShadowMaps/ProjectionCommon.ush
--- a/Engine/Shaders/Private/VirtualShadowMaps/ProjectionCommon.ush
+++ b/Engine/Shaders/Private/VirtualShadowMaps/ProjectionCommon.ush
@@ -262,7 +262,7 @@ FVirtualShadowMapSampleResult SampleVirtualShadowMap(int VirtualShadowMapId, flo
 
 		if( BaseProjectionData.LightType != LIGHT_TYPE_SPOT )
 		{
-			VirtualShadowMapId += GetCubeFace( -TranslatedWorldPosition );
+			VirtualShadowMapId += GetCubeFace( TranslatedWorldPosition );
 			BaseProjectionData = GetVirtualShadowMapProjectionData( VirtualShadowMapId );
 		}
 
diff --git a/Engine/Shaders/Private/VirtualShadowMaps/ProjectionSpot.ush b/Engine/Shaders/Private/VirtualShadowMaps/ProjectionSpot.ush
--- a/Engine/Shaders/Private/VirtualShadowMaps/ProjectionSpot.ush
+++ b/Engine/Shaders/Private/VirtualShadowMaps/ProjectionSpot.ush
@@ -159,8 +159,8 @@ FSMRTResult ShadowRayCastPointLight(
 {
 	FVirtualShadowMapProjectionShaderData BaseProjectionData = GetVirtualShadowMapProjectionData(VirtualShadowMapId);
 
-	uint RayStartFace = GetCubeFace(-(RayStart + BaseProjectionData.ShadowPreViewTranslation));
-	uint RayEndFace   = GetCubeFace(-(RayEnd   + BaseProjectionData.ShadowPreViewTranslation));
+	uint RayStartFace = GetCubeFace( RayStart + BaseProjectionData.ShadowPreViewTranslation );
+	uint RayEndFace   = GetCubeFace( RayEnd   + BaseProjectionData.ShadowPreViewTranslation );
 		
 	FSMRTResult Result;	
 	if (RayStartFace != RayEndFace)
@@ -356,6 +356,8 @@ float TraceLocalLight(
 				}
 			}
 #endif
+			uint Seed = asuint( StepOffset );
+			StepOffset = ( EvolveSobolSeed( Seed ) >> 8 ) * 5.96046447754e-08; // * 2^-24;
 		}
 		SampleCount = min(i + 1, uint(RayCount));		// break vs regular for loop exit
 		RayMissFactor /= float(SampleCount);
diff --git a/Engine/Source/Runtime/Renderer/Private/ShadowDepthRendering.cpp b/Engine/Source/Runtime/Renderer/Private/ShadowDepthRendering.cpp
--- a/Engine/Source/Runtime/Renderer/Private/ShadowDepthRendering.cpp
+++ b/Engine/Source/Runtime/Renderer/Private/ShadowDepthRendering.cpp
@@ -2342,7 +2342,12 @@ bool FShadowDepthPassMeshProcessor::TryAddMeshBatch(const FMeshBatch& RESTRICT M
 		const bool bPlatformReversesCulling = (RHINeedsToSwitchVerticalAxis(ShaderPlatform) && MobileHDRCvar->GetValueOnAnyThread() == 0);
 
 		const bool bRenderSceneTwoSided = bTwoSided;
-		const bool bReverseCullMode = XOR(bPlatformReversesCulling, ShadowDepthType.bOnePassPointLightShadow);
+#if ENABLE_NON_NANITE_VSM
+		const bool bShadowReversesCulling = MeshPassTargetType == EMeshPass::VSMShadowDepth ? false : ShadowDepthType.bOnePassPointLightShadow;
+#else
+		const bool bShadowReversesCulling = ShadowDepthType.bOnePassPointLightShadow;
+#endif
+		const bool bReverseCullMode = XOR(bPlatformReversesCulling, bShadowReversesCulling);
 
 		FinalCullMode = bRenderSceneTwoSided ? CM_None : bReverseCullMode ? InverseCullMode(MeshCullMode) : MeshCullMode;
 	}
diff --git a/Engine/Source/Runtime/Renderer/Private/ShadowRendering.cpp b/Engine/Source/Runtime/Renderer/Private/ShadowRendering.cpp
--- a/Engine/Source/Runtime/Renderer/Private/ShadowRendering.cpp
+++ b/Engine/Source/Runtime/Renderer/Private/ShadowRendering.cpp
@@ -2276,8 +2276,7 @@ FViewMatrices FProjectedShadowInfo::GetShadowDepthRenderingViewMatrices(int32 Cu
 	{
 		if (bUseForVSMCubeFaceWorkaround)
 		{
-			// TODO add cull direction to FPackedView instead of this scale nonsense.
-			MatricesInitializer.ViewRotationMatrix = OnePassShadowViewMatrices[CubeFaceIndex] * FScaleMatrix(FVector(1, -1, 1));
+			MatricesInitializer.ViewRotationMatrix = OnePassShadowViewMatrices[CubeFaceIndex] * FScaleMatrix(FVector(1, 1, -1));
 		}
 		else
 		{
diff --git a/Engine/Source/Runtime/Renderer/Private/VirtualShadowMaps/VirtualShadowMapArray.cpp b/Engine/Source/Runtime/Renderer/Private/VirtualShadowMaps/VirtualShadowMapArray.cpp
--- a/Engine/Source/Runtime/Renderer/Private/VirtualShadowMaps/VirtualShadowMapArray.cpp
+++ b/Engine/Source/Runtime/Renderer/Private/VirtualShadowMaps/VirtualShadowMapArray.cpp
@@ -736,38 +736,20 @@ void FVirtualShadowMapArray::BuildPageAllocations(
 						check(ShadowViewRect.Max.Y == FVirtualShadowMap::VirtualMaxResolutionXY);
 					}
 
-					if (ProjectedShadowInfo->bOnePassPointLightShadow)
+					int32 NumMaps = ProjectedShadowInfo->bOnePassPointLightShadow ? 6 : 1;
+					for( int32 i = 0; i < NumMaps; i++ )
 					{
-						for (int32 CubeFace = 0; CubeFace < 6; CubeFace++)
-						{
-							FMatrix ViewToClip = ProjectedShadowInfo->OnePassShadowFaceProjectionMatrix;
-							FMatrix TranslatedWorldToView = ProjectedShadowInfo->OnePassShadowViewMatrices[ CubeFace ] * FScaleMatrix(FVector(1, -1, 1));	// TODO add cull direction to FPackedView instead of this scale nonsense.
-	
-							FVirtualShadowMapProjectionShaderData& Data = ShadowMapProjectionData[ ProjectedShadowInfo->VirtualShadowMaps[ CubeFace ]->ID ];
-							Data.TranslatedWorldToShadowViewMatrix = TranslatedWorldToView;
-							Data.ShadowViewToClipMatrix = ViewToClip;
-							Data.TranslatedWorldToShadowUVMatrix = CalcTranslatedWorldToShadowUVMatrix(TranslatedWorldToView, ViewToClip);
-							Data.TranslatedWorldToShadowUVNormalMatrix = CalcTranslatedWorldToShadowUVNormalMatrix(TranslatedWorldToView, ViewToClip);
-							Data.ShadowPreViewTranslation = FVector(ProjectedShadowInfo->PreShadowTranslation);
-							Data.VirtualShadowMapId = ProjectedShadowInfo->VirtualShadowMaps[ CubeFace ]->ID;
-							Data.LightType = ProjectedShadowInfo->GetLightSceneInfo().Proxy->GetLightType();
-				}
-			}
-					else
-					{
-						FMatrix ViewToClip = ProjectedShadowInfo->ViewToClipInner;
-	
-						FVirtualShadowMapProjectionShaderData& Data = ShadowMapProjectionData[ ProjectedShadowInfo->VirtualShadowMaps[0]->ID ];
-						Data.TranslatedWorldToShadowViewMatrix = ProjectedShadowInfo->TranslatedWorldToView;
-						Data.ShadowViewToClipMatrix = ViewToClip;
-						Data.TranslatedWorldToShadowUVMatrix = CalcTranslatedWorldToShadowUVMatrix(ProjectedShadowInfo->TranslatedWorldToView, ViewToClip);
-						Data.TranslatedWorldToShadowUVNormalMatrix = CalcTranslatedWorldToShadowUVNormalMatrix(ProjectedShadowInfo->TranslatedWorldToView, ViewToClip);
-						Data.ShadowPreViewTranslation = FVector(ProjectedShadowInfo->PreShadowTranslation);
-						Data.VirtualShadowMapId = ProjectedShadowInfo->VirtualShadowMaps[0]->ID;
-						Data.LightType = ProjectedShadowInfo->GetLightSceneInfo().Proxy->GetLightType();
-
-						ensure(!ProjectedShadowInfo->bDirectionalLight);
-		}
+						FViewMatrices ViewMatrices = ProjectedShadowInfo->GetShadowDepthRenderingViewMatrices( i, true );
+
+						FVirtualShadowMapProjectionShaderData& Data = ShadowMapProjectionData[ ProjectedShadowInfo->VirtualShadowMaps[i]->ID ];
+						Data.TranslatedWorldToShadowViewMatrix		= ViewMatrices.GetTranslatedViewMatrix();
+						Data.ShadowViewToClipMatrix					= ViewMatrices.GetProjectionMatrix();
+						Data.TranslatedWorldToShadowUVMatrix		= CalcTranslatedWorldToShadowUVMatrix( ViewMatrices.GetTranslatedViewMatrix(), ViewMatrices.GetProjectionMatrix() );
+						Data.TranslatedWorldToShadowUVNormalMatrix	= CalcTranslatedWorldToShadowUVNormalMatrix( ViewMatrices.GetTranslatedViewMatrix(), ViewMatrices.GetProjectionMatrix() );
+						Data.ShadowPreViewTranslation				= FVector(ProjectedShadowInfo->PreShadowTranslation);
+						Data.VirtualShadowMapId						= ProjectedShadowInfo->VirtualShadowMaps[i]->ID;
+						Data.LightType								= ProjectedShadowInfo->GetLightSceneInfo().Proxy->GetLightType();
+					}
 				}
 			}
 		}
