Fixed Lumen voxel lighting cache bug when distance field indices are swapped during primitive removal

diff --git a/Engine/Source/Runtime/Renderer/Private/DistanceFieldObjectManagement.cpp b/Engine/Source/Runtime/Renderer/Private/DistanceFieldObjectManagement.cpp
--- a/Engine/Source/Runtime/Renderer/Private/DistanceFieldObjectManagement.cpp
+++ b/Engine/Source/Runtime/Renderer/Private/DistanceFieldObjectManagement.cpp
@@ -99,13 +99,17 @@ static int32 PartitionUpdateRangesDFO(FParallelUpdateRangesDFO& Ranges, int32 It
 	return Ranges.Range[3].ItemCount > 0 ? 4 : 3;
 }
 
-void AddModifiedBounds(FScene* Scene, FGlobalDFCacheType CacheType, const FBox& Bounds)
+void AddModifiedBoundsForLumen(FScene* Scene, FGlobalDFCacheType CacheType, const FBox& Bounds)
 {
-	FDistanceFieldSceneData& DistanceFieldData = Scene->DistanceFieldSceneData;
 	FLumenSceneData& LumenData = *Scene->LumenSceneData;
+	LumenData.PrimitiveModifiedBounds.Add(Bounds);
+}
 
+void AddModifiedBounds(FScene* Scene, FGlobalDFCacheType CacheType, const FBox& Bounds)
+{
+	FDistanceFieldSceneData& DistanceFieldData = Scene->DistanceFieldSceneData;
 	DistanceFieldData.PrimitiveModifiedBounds[CacheType].Add(Bounds);
-	LumenData.PrimitiveModifiedBounds.Add(Bounds);
+	AddModifiedBoundsForLumen(Scene, CacheType, Bounds);
 }
 
 void UpdateGlobalDistanceFieldObjectRemoves(FScene* Scene)
@@ -162,6 +166,14 @@ void UpdateGlobalDistanceFieldObjectRemoves(FScene* Scene)
 
 				LumenUpdateDFObjectIndex(Scene, RemoveIndex);
 				DistanceFieldSceneData.IndicesToUpdateInObjectBuffers.Add(RemoveIndex);
+
+				// Lumen caches distance field indices, which requires an update after a swap
+				if (RemoveIndex < DistanceFieldSceneData.PrimitiveInstanceMapping.Num())
+				{
+					const FPrimitiveAndInstance& Swapped = DistanceFieldSceneData.PrimitiveInstanceMapping[RemoveIndex];
+					const FGlobalDFCacheType CacheType = Swapped.Primitive->Proxy->IsOftenMoving() ? GDF_Full : GDF_MostlyStatic;
+					AddModifiedBoundsForLumen(Scene, CacheType, Swapped.WorldBounds);
+				}
 			}
 
 			PendingRemoveOperations.Reset();
