Made Strata visualize shader pc/editor only.

diff --git a/Engine/Shaders/Private/Strata/StrataVisualize.usf b/Engine/Shaders/Private/Strata/StrataVisualize.usf
--- a/Engine/Shaders/Private/Strata/StrataVisualize.usf
+++ b/Engine/Shaders/Private/Strata/StrataVisualize.usf
@@ -68,7 +68,6 @@ void VisualizeMaterialPS(
 	float4 SVPos : SV_POSITION,
 	out float4 OutColor : SV_Target0)
 {
-#if 0	// Here to disabled on some platforms having problems with that shader. STRATA_TODO investigate why this shader is spilling like crazy again
 	const uint2 PixelPos = uint2(SVPos.xy);
 	const uint2 MaterialPos = View.CursorPosition;
 	float2 BufferUV = SvPositionToBufferUV(float4(MaterialPos, SVPos.zw));
@@ -462,7 +461,6 @@ void VisualizeMaterialPS(
 		}
 #endif
 	}
-#endif
 }
 
 
diff --git a/Engine/Source/Runtime/Renderer/Private/DeferredShadingRenderer.cpp b/Engine/Source/Runtime/Renderer/Private/DeferredShadingRenderer.cpp
--- a/Engine/Source/Runtime/Renderer/Private/DeferredShadingRenderer.cpp
+++ b/Engine/Source/Runtime/Renderer/Private/DeferredShadingRenderer.cpp
@@ -2704,7 +2704,7 @@ void FDeferredShadingSceneRenderer::Render(FRDGBuilder& GraphBuilder)
 
 	if (bStrataEnabled)
 	{
-		Strata::AddVisualizeMaterialPasses(GraphBuilder, Views, SceneTextures.Color.Target);
+		Strata::AddVisualizeMaterialPasses(GraphBuilder, Views, SceneTextures.Color.Target, Scene->GetShaderPlatform());
 	}
 
 	if (bCanOverlayRayTracingOutput && ViewFamily.EngineShowFlags.LightShafts)
diff --git a/Engine/Source/Runtime/Renderer/Private/Strata/Strata.cpp b/Engine/Source/Runtime/Renderer/Private/Strata/Strata.cpp
--- a/Engine/Source/Runtime/Renderer/Private/Strata/Strata.cpp
+++ b/Engine/Source/Runtime/Renderer/Private/Strata/Strata.cpp
@@ -199,9 +199,16 @@ class FVisualizeMaterialPS : public FGlobalShader
 		return PermutationVector;
 	}
 
+	static bool CanRunStrataVizualizeMaterial(EShaderPlatform Platform)
+	{
+		// On some consoles, this ALU heavy shader (and with optimisation disables for the sake of low compilation time) would spill registers. So only keep it for the editor.
+		return IsPCPlatform(Platform);
+	}
+
 	static bool ShouldCompilePermutation(const FGlobalShaderPermutationParameters& Parameters)
 	{
-		return GetMaxSupportedFeatureLevel(Parameters.Platform) >= ERHIFeatureLevel::SM5 && Strata::IsStrataEnabled();
+		return GetMaxSupportedFeatureLevel(Parameters.Platform) >= ERHIFeatureLevel::SM5 && Strata::IsStrataEnabled()
+			&& CanRunStrataVizualizeMaterial(Parameters.Platform);
 	}
 
 	static void ModifyCompilationEnvironment(const FGlobalShaderPermutationParameters& Parameters, FShaderCompilerEnvironment& OutEnvironment)
@@ -214,10 +221,10 @@ class FVisualizeMaterialPS : public FGlobalShader
 };
 IMPLEMENT_GLOBAL_SHADER(FVisualizeMaterialPS, "/Engine/Private/Strata/StrataVisualize.usf", "VisualizeMaterialPS", SF_Pixel);
 
-void AddVisualizeMaterialPasses(FRDGBuilder& GraphBuilder, const TArray<FViewInfo>& Views, FRDGTextureRef SceneColorTexture)
+void AddVisualizeMaterialPasses(FRDGBuilder& GraphBuilder, const TArray<FViewInfo>& Views, FRDGTextureRef SceneColorTexture, EShaderPlatform Platform)
 {
 	RDG_EVENT_SCOPE_CONDITIONAL(GraphBuilder, IsStrataEnabled() && Views.Num() > 0, "StrataVisualizeMaterial");
-	if (!IsStrataEnabled())
+	if (!IsStrataEnabled() || !FVisualizeMaterialPS::CanRunStrataVizualizeMaterial(Platform))
 	{
 		return;
 	}
diff --git a/Engine/Source/Runtime/Renderer/Private/Strata/Strata.h b/Engine/Source/Runtime/Renderer/Private/Strata/Strata.h
--- a/Engine/Source/Runtime/Renderer/Private/Strata/Strata.h
+++ b/Engine/Source/Runtime/Renderer/Private/Strata/Strata.h
@@ -63,7 +63,7 @@ void AddStrataMaterialClassificationPass(FRDGBuilder& GraphBuilder, const FMinim
 void AddStrataStencilPass(FRDGBuilder& GraphBuilder, const FViewInfo& View, const FMinimalSceneTextures& SceneTextures);
 void AddStrataStencilPass(FRDGBuilder& GraphBuilder, const TArray<FViewInfo>& Views, const FMinimalSceneTextures& SceneTextures);
 
-void AddVisualizeMaterialPasses(FRDGBuilder& GraphBuilder, const TArray<FViewInfo>& Views, FRDGTextureRef SceneColorTexture);
+void AddVisualizeMaterialPasses(FRDGBuilder& GraphBuilder, const TArray<FViewInfo>& Views, FRDGTextureRef SceneColorTexture, EShaderPlatform Platform);
 };
 
 
