Automatically use r.SSS.Quality=1 recombine when using Gen5 TAA

diff --git a/Engine/Source/Runtime/Renderer/Private/PostProcess/PostProcessSubsurface.cpp b/Engine/Source/Runtime/Renderer/Private/PostProcess/PostProcessSubsurface.cpp
--- a/Engine/Source/Runtime/Renderer/Private/PostProcess/PostProcessSubsurface.cpp
+++ b/Engine/Source/Runtime/Renderer/Private/PostProcess/PostProcessSubsurface.cpp
@@ -19,6 +19,7 @@
 #include "GenerateMips.h"
 #include "ClearQuad.h"
 #include "Strata/Strata.h"
+#include "TemporalAA.h"
 
 namespace
 {
@@ -792,12 +793,31 @@ class FSubsurfaceRecombinePS : public FSubsurfaceShader
 	{
 		const uint32 QualityCVar = GetSSSQuality();
 
+		EMainTAAPassConfig MainTAAConfig = ITemporalUpscaler::GetMainTAAPassConfig(View);
+
+		// Low quality is really bad with modern temporal upscalers.
+		bool bAllowLowQuality = MainTAAConfig == EMainTAAPassConfig::Disabled || MainTAAConfig == EMainTAAPassConfig::Gen4;
+		if (!bAllowLowQuality)
+		{
+			return EQuality::High;
+		}
+
 		// Quality is forced to high when the CVar is set to 'auto' and TAA is NOT enabled.
-		// TAA improves quality through temporal filtering, making it less necessary to use
+		// Gen4 TAA improves quality through temporal filtering and clamping box, making it less necessary to use
 		// high quality mode.
-		const bool bUseHighQuality = (QualityCVar == -1 && View.AntiAliasingMethod != AAM_TemporalAA);
+		if (QualityCVar == -1)
+		{
+			if (MainTAAConfig == EMainTAAPassConfig::Gen4)
+			{
+				return EQuality::Low;
+			}
+			else
+			{
+				return EQuality::High;
+			}
+		}
 
-		if (QualityCVar == 1 || bUseHighQuality || View.Family->GetTemporalUpscalerInterface() != nullptr)
+		if (QualityCVar == 1)
 		{
 			return EQuality::High;
 		}
diff --git a/Engine/Source/Runtime/Renderer/Private/PostProcess/PostProcessing.cpp b/Engine/Source/Runtime/Renderer/Private/PostProcess/PostProcessing.cpp
--- a/Engine/Source/Runtime/Renderer/Private/PostProcess/PostProcessing.cpp
+++ b/Engine/Source/Runtime/Renderer/Private/PostProcess/PostProcessing.cpp
@@ -518,7 +518,8 @@ void AddPostProcessingPasses(FRDGBuilder& GraphBuilder, const FViewInfo& View, c
 		FIntRect SecondaryViewRect = PrimaryViewRect;
 
 		// Temporal Anti-aliasing. Also may perform a temporal upsample from primary to secondary view rect.
-		if (AntiAliasingMethod == AAM_TemporalAA)
+		EMainTAAPassConfig TAAConfig = ITemporalUpscaler::GetMainTAAPassConfig(View);
+		if (TAAConfig != EMainTAAPassConfig::Disabled)
 		{
 			// Whether we allow the temporal AA pass to downsample scene color. It may choose not to based on internal context,
 			// in which case the output half resolution texture will remain null.
@@ -529,17 +530,15 @@ void AddPostProcessingPasses(FRDGBuilder& GraphBuilder, const FViewInfo& View, c
 				// TemporalAA is only able to match the low quality mode (box filter).
 				GetDownsampleQuality() == EDownsampleQuality::Low;
 
-			int32 UpscaleMode = ITemporalUpscaler::GetTemporalUpscalerMode();
-
-			const ITemporalUpscaler* DefaultTemporalUpscaler = ITemporalUpscaler::GetDefaultTemporalUpscaler();
-			const ITemporalUpscaler* UpscalerToUse = ( UpscaleMode == 0 || !View.Family->GetTemporalUpscalerInterface())? DefaultTemporalUpscaler : View.Family->GetTemporalUpscalerInterface();
-
+			const ITemporalUpscaler* UpscalerToUse = (TAAConfig == EMainTAAPassConfig::ThirdParty) ? View.Family->GetTemporalUpscalerInterface() : ITemporalUpscaler::GetDefaultTemporalUpscaler();
+			check(UpscalerToUse);
 			const TCHAR* UpscalerName = UpscalerToUse->GetDebugName();
 
-			// Standard event scope for temporal upscaler to have all profiling information not matter what, and with explicit detection of third party.
+			// Standard event scope for temporal upscaler to have all profiling information not matter what,
+			// and with explicit detection of third party.
 			RDG_EVENT_SCOPE_CONDITIONAL(
 				GraphBuilder,
-				UpscalerToUse != DefaultTemporalUpscaler,
+				TAAConfig == EMainTAAPassConfig::ThirdParty,
 				"ThirdParty TAA %s %dx%d -> %dx%d",
 				UpscalerToUse->GetDebugName(),
 				View.ViewRect.Width(), View.ViewRect.Height(),
diff --git a/Engine/Source/Runtime/Renderer/Private/PostProcess/TemporalAA.cpp b/Engine/Source/Runtime/Renderer/Private/PostProcess/TemporalAA.cpp
--- a/Engine/Source/Runtime/Renderer/Private/PostProcess/TemporalAA.cpp
+++ b/Engine/Source/Runtime/Renderer/Private/PostProcess/TemporalAA.cpp
@@ -1699,7 +1699,30 @@ const ITemporalUpscaler* ITemporalUpscaler::GetDefaultTemporalUpscaler()
 	return &DefaultTemporalUpscaler;
 }
 
-int ITemporalUpscaler::GetTemporalUpscalerMode()
+//static
+EMainTAAPassConfig ITemporalUpscaler::GetMainTAAPassConfig(const FViewInfo& View)
 {
-	return CVarUseTemporalAAUpscaler.GetValueOnRenderThread();
+	if (!IsPostProcessingEnabled(View))
+	{
+		return EMainTAAPassConfig::Disabled;
+	}
+	else if (View.AntiAliasingMethod != AAM_TemporalAA)
+	{
+		return EMainTAAPassConfig::Disabled;
+	}
+
+	int32 CustomUpscalerMode = CVarUseTemporalAAUpscaler.GetValueOnRenderThread();
+
+	if (View.Family->GetTemporalUpscalerInterface() && CustomUpscalerMode != 0)
+	{
+		return EMainTAAPassConfig::ThirdParty;
+	}
+	else if (CVarTAAAlgorithm.GetValueOnRenderThread() && DoesPlatformSupportGen5TAA(View.GetShaderPlatform()))
+	{
+		return EMainTAAPassConfig::Gen5;
+	}
+	else
+	{
+		return EMainTAAPassConfig::Gen4;
+	}
 }
diff --git a/Engine/Source/Runtime/Renderer/Private/PostProcess/TemporalAA.h b/Engine/Source/Runtime/Renderer/Private/PostProcess/TemporalAA.h
--- a/Engine/Source/Runtime/Renderer/Private/PostProcess/TemporalAA.h
+++ b/Engine/Source/Runtime/Renderer/Private/PostProcess/TemporalAA.h
@@ -6,6 +6,24 @@
 
 struct FTemporalAAHistory;
 
+
+/** Configuration of the main temporal AA pass. */
+enum class EMainTAAPassConfig : uint8
+{
+	// TAA is disabled.
+	Disabled,
+
+	// Uses old UE4's Temporal AA maintained for Gen4 consoles
+	Gen4,
+
+	// Uses Gen5 Temporal AA
+	Gen5,
+
+	// Uses third party View.Family->GetTemporalUpscalerInterface()
+	ThirdParty,
+};
+
+
 /** List of TAA configurations. */
 enum class ETAAPassConfig
 {
@@ -181,7 +199,8 @@ public:
 	virtual float GetMaxUpsampleResolutionFraction() const = 0;
 
 	static const ITemporalUpscaler* GetDefaultTemporalUpscaler();
-	static int GetTemporalUpscalerMode();
+
+	static EMainTAAPassConfig GetMainTAAPassConfig(const FViewInfo& View);
 }; 
 
 extern RENDERER_API const ITemporalUpscaler* GTemporalUpscaler;
