Fixed NonUnity Compile UnrealEditor Mac warning

diff --git a/Engine/Source/Runtime/Renderer/Private/DynamicBVH.cpp b/Engine/Source/Runtime/Renderer/Private/DynamicBVH.cpp
--- a/Engine/Source/Runtime/Renderer/Private/DynamicBVH.cpp
+++ b/Engine/Source/Runtime/Renderer/Private/DynamicBVH.cpp
@@ -1,7 +1,5 @@
 // Copyright Epic Games, Inc. All Rights Reserved.
 
-#pragma once
-
 #include "DynamicBVH.h"
 #include "Misc/AutomationTest.h"
 #include "Misc/OutputDeviceRedirector.h"
