Disable adaptive ray count projection shader when wave ops are not supported on a given RHI.

diff --git a/Engine/Source/Runtime/Renderer/Private/VirtualShadowMaps/VirtualShadowMapProjection.cpp b/Engine/Source/Runtime/Renderer/Private/VirtualShadowMaps/VirtualShadowMapProjection.cpp
--- a/Engine/Source/Runtime/Renderer/Private/VirtualShadowMaps/VirtualShadowMapProjection.cpp
+++ b/Engine/Source/Runtime/Renderer/Private/VirtualShadowMaps/VirtualShadowMapProjection.cpp
@@ -531,15 +531,17 @@ void RenderVirtualShadowMapProjection(
 
 	PassParameters->RWShadowMaskBits = GraphBuilder.CreateUAV( ShadowMaskBits );
 	
+	bool bAdaptiveRayCount = GRHISupportsWaveOperations && CVarSMRTAdaptiveRayCount.GetValueOnRenderThread() != 0;
+
 	FVirtualShadowMapProjectionCS::FPermutationDomain PermutationVector;
-	PermutationVector.Set< FVirtualShadowMapProjectionCS::FSMRTAdaptiveRayCountDim >( CVarSMRTAdaptiveRayCount.GetValueOnRenderThread() != 0 );
+	PermutationVector.Set< FVirtualShadowMapProjectionCS::FSMRTAdaptiveRayCountDim >( bAdaptiveRayCount );
 	auto ComputeShader = View.ShaderMap->GetShader< FVirtualShadowMapProjectionCS >( PermutationVector );
 
 	const FIntPoint GroupCount = FIntPoint::DivideAndRoundUp( View.ViewRect.Size(), 8 );
 
 	FComputeShaderUtils::AddPass(
 		GraphBuilder,
-		RDG_EVENT_NAME("VirtualShadowMapProjection"),
+		bAdaptiveRayCount ? RDG_EVENT_NAME("VirtualShadowMapProjection") : RDG_EVENT_NAME("VirtualShadowMapProjection (StaticRayCount)"),
 		ComputeShader,
 		PassParameters,
 		FIntVector( GroupCount.X, GroupCount.Y, 1 )
