Strata
- Basepass shader no longer expoert all rendertarget. Only scene color, preshado, velocity if enabled.
- More decoupling from GBuffer. Now normal gbuffer is not used by any of SSAO, DFAO, HBAO and SSR. Achieved  with permutations.
- Strata header now contains PrecShadow bits.

Shooter game compiles on pc and console. Tested with Strata disabled.

diff --git a/Engine/Shaders/Private/BasePassPixelShader.usf b/Engine/Shaders/Private/BasePassPixelShader.usf
--- a/Engine/Shaders/Private/BasePassPixelShader.usf
+++ b/Engine/Shaders/Private/BasePassPixelShader.usf
@@ -977,11 +977,18 @@ void FPixelShaderInOut_MainPS(
 	// Initialise a Strata header with normal in registers
 	FStrataPixelHeader StrataPixelHeader = InitialiseStrataPixelHeader();
 	UNROLL
-		for (int l = 0; l < Strata.LayerCount; ++l)
-		{
-			StrataPixelHeader.BSDFCount += Strata.Layers[l].BSDFCount;
-		}
+	for (int l = 0; l < Strata.LayerCount; ++l)
+	{
+		StrataPixelHeader.BSDFCount += Strata.Layers[l].BSDFCount;
+	}
 	StrataPixelHeader.SharedNormals = MaterialParameters.SharedNormals;
+
+#if !GBUFFER_HAS_PRECSHADOWFACTOR
+	StrataPixelHeader.SkipPrecShadowMask = true;
+#endif
+#if (GBUFFER_HAS_PRECSHADOWFACTOR && WRITES_PRECSHADOWFACTOR_ZERO)
+	StrataPixelHeader.ZeroPrecShadowMask = true;
+#endif
 #endif
 
 #if STRATA_ENABLED && STRATA_INLINE_SINGLELAYERWATER
diff --git a/Engine/Shaders/Private/DeferredLightPixelShaders.usf b/Engine/Shaders/Private/DeferredLightPixelShaders.usf
--- a/Engine/Shaders/Private/DeferredLightPixelShaders.usf
+++ b/Engine/Shaders/Private/DeferredLightPixelShaders.usf
@@ -140,6 +140,26 @@ float3 GetAtmosphericLightTransmittance(in float3 StableLargeWorldPosition, in f
 
 ////////////////////////////////////////////////////////////////////////////////////////////////////
 
+#if STRATA_ENABLED
+
+float4 StrataReadPrecomputedShadowFactors(FStrataPixelHeader StrataPixelHeader, int2 PixelPos)
+{
+	if (!StrataPixelHeader.SkipPrecShadowMask)
+	{
+#if ALLOW_STATIC_LIGHTING
+		float4 GBufferE = SceneTexturesStruct.GBufferETexture.Load(int3(PixelPos, 0));
+#else
+		float4 GBufferE = 1;
+#endif
+		return GBufferE;
+	}
+	return StrataPixelHeader.ZeroPrecShadowMask ? 0.0f : 1.0f;
+}
+
+#endif // STRATA_ENABLED
+
+////////////////////////////////////////////////////////////////////////////////////////////////////
+
 #if USE_HAIR_LIGHTING == 0
 
 void DeferredLightPixelMain(
@@ -169,10 +189,67 @@ void DeferredLightPixelMain(
 	InputParams.ScreenVector	= ScreenVector;
 #endif
 
-	FScreenSpaceData ScreenSpaceData = GetScreenSpaceData(InputParams.ScreenUV);
+#if STRATA_ENABLED
+
+	FStrataAddressing StrataAddressing = GetStrataPixelDataByteOffset(PixelPos, uint2(View.BufferSizeAndInvSize.xy), Strata.MaxBytesPerPixel);
+	FStrataPixelHeader StrataPixelHeader = UnpackStrataHeaderIn(Strata.MaterialLobesBuffer, StrataAddressing);
+
+	BRANCH
+	if (StrataPixelHeader.BSDFCount > 0	// This test is also enough to exclude sky pixels
+#if USE_LIGHTING_CHANNELS
+		&& (GetLightingChannelMask(InputParams.ScreenUV) & DeferredLightUniforms.LightingChannelMask)
+#endif
+		) 
+	{
+		const float SceneDepth = CalcSceneDepth(InputParams.ScreenUV);
+		const FDerivedParams DerivedParams = GetDerivedParams(InputParams, SceneDepth);
+
+		FDeferredLightData LightData = SetupLightDataForStandardDeferred();
 
+		FRectTexture RectTexture = InitRectTexture(DeferredLightUniforms.SourceTexture);
+		float SurfaceShadow = 1.0f;
+
+		float3 V =-DerivedParams.CameraVector;
+		float3 L = LightData.Direction;	// Already normalized
+		float3 ToLight = L;
+		float LightMask = 1;
+		if (LightData.bRadialLight)
+		{
+			LightMask = GetLocalLightAttenuation(DerivedParams.WorldPosition, LightData, ToLight, L);
+		}
+
+		float4 Radiance = 0;
+		if (LightMask > 0)
+		{
+			FShadowTerms ShadowTerms = { StrataPixelHeader.MaterialAO, 1.0, 1.0, InitHairTransmittanceData() };
+			float4 LightAttenuation = GetPerPixelLightAttenuation(InputParams.ScreenUV);
+
+			float Dither = InterleavedGradientNoise(InputParams.PixelPos, View.StateFrameIndexMod8);
+			float4 PrecomputedShadowFactors = StrataReadPrecomputedShadowFactors(StrataPixelHeader, PixelPos);
+			const uint FakeShadingModelID = 0;
+			const float FakeContactShadowOpacity = 1.0f;
+			GetShadowTerms(SceneDepth, PrecomputedShadowFactors, FakeShadingModelID, FakeContactShadowOpacity,
+				LightData, DerivedParams.WorldPosition, L, LightAttenuation, Dither, ShadowTerms);
+			SurfaceShadow = ShadowTerms.SurfaceShadow;
+
+			Radiance = StrataDeferredLighting(
+				LightData,
+				V,
+				L,
+				ToLight,
+				LightMask,
+				ShadowTerms,
+				RectTexture,
+				Strata.MaterialLobesBuffer,
+				StrataAddressing,
+				StrataPixelHeader);
+		}
+
+#else // STRATA_ENABLED
+
+	FScreenSpaceData ScreenSpaceData = GetScreenSpaceData(InputParams.ScreenUV);
 	// Only light pixels marked as using deferred shading
-	BRANCH if( ScreenSpaceData.GBuffer.ShadingModelID > 0 
+	BRANCH if (ScreenSpaceData.GBuffer.ShadingModelID > 0
 #if USE_LIGHTING_CHANNELS
 		&& (GetLightingChannelMask(InputParams.ScreenUV) & DeferredLightUniforms.LightingChannelMask)
 #endif
@@ -182,70 +259,27 @@ void DeferredLightPixelMain(
 		const FDerivedParams DerivedParams = GetDerivedParams(InputParams, SceneDepth);
 
 		FDeferredLightData LightData = SetupLightDataForStandardDeferred();
-#if STRATA_ENABLED == 0 || (STRATA_ENABLED && STRATA_FASTPATH == 0)
+
 		if (ScreenSpaceData.GBuffer.ShadingModelID == SHADINGMODELID_HAIR && ShouldUseHairComplexTransmittance(ScreenSpaceData.GBuffer))
 		{
 			// Prevent numerical issue causing NaN during hair lighting (i.e., during baseColor/luma(baseColor) computation)
 			ScreenSpaceData.GBuffer.BaseColor = max(0.0001f, ScreenSpaceData.GBuffer.BaseColor);
 			LightData.HairTransmittance = EvaluateDualScattering(ScreenSpaceData.GBuffer, DerivedParams.CameraVector, -DeferredLightUniforms.Direction);
 		}
-#endif
 
-		float Dither = InterleavedGradientNoise(InputParams.PixelPos, View.StateFrameIndexMod8 );
+		float Dither = InterleavedGradientNoise(InputParams.PixelPos, View.StateFrameIndexMod8);
 
 		FRectTexture RectTexture = InitRectTexture(DeferredLightUniforms.SourceTexture);
 		float SurfaceShadow = 1.0f;
 
-		float4 Radiance = 0;
+		float4 Radiance = GetDynamicLighting(DerivedParams.WorldPosition, DerivedParams.CameraVector, ScreenSpaceData.GBuffer, ScreenSpaceData.AmbientOcclusion, ScreenSpaceData.GBuffer.ShadingModelID, LightData, GetPerPixelLightAttenuation(InputParams.ScreenUV), Dither, uint2(InputParams.PixelPos), RectTexture, SurfaceShadow);
 
-#if STRATA_ENABLED
-		FStrataAddressing StrataAddressing = GetStrataPixelDataByteOffset(PixelPos, uint2(View.BufferSizeAndInvSize.xy), Strata.MaxBytesPerPixel);
-		FStrataPixelHeader StrataPixelHeader = UnpackStrataHeaderIn(Strata.MaterialLobesBuffer, StrataAddressing);
-		BRANCH
-		if (StrataPixelHeader.BSDFCount > 0) // This test is also enough to exclude sky pixels
-		{
-			float3 V =-DerivedParams.CameraVector;
-			float3 L = LightData.Direction;	// Already normalized
-			float3 ToLight = L;
-			float LightMask = 1;
-			if (LightData.bRadialLight)
-			{
-				LightMask = GetLocalLightAttenuation(DerivedParams.WorldPosition, LightData, ToLight, L);
-			}
-
-			if (LightMask > 0)
-			{
-				uint2 PixelPos = InputParams.PixelPos;
-
-				FShadowTerms ShadowTerms = { ScreenSpaceData.AmbientOcclusion, 1.0, 1.0, InitHairTransmittanceData() };
-				float4 LightAttenuation = GetPerPixelLightAttenuation(InputParams.ScreenUV);
-				GetShadowTerms(ScreenSpaceData.GBuffer, LightData, DerivedParams.WorldPosition, L, LightAttenuation, Dither, ShadowTerms);
-				SurfaceShadow = ShadowTerms.SurfaceShadow;
-
-				Radiance = StrataDeferredLighting(
-					LightData,
-					V,
-					L,
-					ToLight,
-					LightMask,
-					ShadowTerms,
-					RectTexture,
-					Strata.MaterialLobesBuffer,
-					StrataAddressing,
-					StrataPixelHeader);
-			}
-		}
-#else
-		{
-			Radiance = GetDynamicLighting(DerivedParams.WorldPosition, DerivedParams.CameraVector, ScreenSpaceData.GBuffer, ScreenSpaceData.AmbientOcclusion, ScreenSpaceData.GBuffer.ShadingModelID, LightData, GetPerPixelLightAttenuation(InputParams.ScreenUV), Dither, uint2(InputParams.PixelPos), RectTexture, SurfaceShadow);
-		}
-#endif
+#endif // STRATA_ENABLED
 
 		const float  Attenuation = ComputeLightProfileMultiplier(DerivedParams.WorldPosition, DeferredLightUniforms.Position, -DeferredLightUniforms.Direction, DeferredLightUniforms.Tangent);
 
 		OutColor += (Radiance * Attenuation);
 
-
 	#if USE_ATMOSPHERE_TRANSMITTANCE || USE_CLOUD_TRANSMITTANCE
 		float DeviceZ = LookupDeviceZ(ScreenUV);
 		float3 StableLargeWorldPosition = SvPositionToWorld(float4(SVPos.xy, DeviceZ, 1.0));
diff --git a/Engine/Shaders/Private/DeferredLightingCommon.ush b/Engine/Shaders/Private/DeferredLightingCommon.ush
--- a/Engine/Shaders/Private/DeferredLightingCommon.ush
+++ b/Engine/Shaders/Private/DeferredLightingCommon.ush
@@ -12,6 +12,10 @@
 #include "CapsuleLightIntegrate.ush"
 #include "RectLightIntegrate.ush"
 
+#ifndef STRATA_ENABLED
+#define STRATA_ENABLED 0
+#endif
+
 /** 
  * Data about a single light.
  * Putting the light data in this struct allows the same lighting code to be used between standard deferred, 
@@ -145,10 +149,10 @@ float ShadowRayCast(
 #error "Must set SUPPORT_CONTACT_SHADOWS"
 #endif
 
-void GetShadowTerms(FGBufferData GBuffer, FDeferredLightData LightData, float3 WorldPosition, float3 L, float4 LightAttenuation, float Dither, inout FShadowTerms Shadow)
+void GetShadowTerms(float SceneDepth, float4 PrecomputedShadowFactors, uint ShadingModelID, float ContactShadowOpacity, FDeferredLightData LightData, float3 WorldPosition, float3 L, float4 LightAttenuation, float Dither, inout FShadowTerms Shadow)
 {
 	float ContactShadowLength = 0.0f;
-	const float ContactShadowLengthScreenScale = View.ClipToView[1][1] * GBuffer.Depth;
+	const float ContactShadowLengthScreenScale = View.ClipToView[1][1] * SceneDepth;
 
 	BRANCH
 	if (LightData.ShadowedBits)
@@ -159,7 +163,7 @@ void GetShadowTerms(FGBufferData GBuffer, FDeferredLightData LightData, float3 W
 		// Whole scene directional light shadows in x, whole scene directional light SSS shadows in y
 		// Get static shadowing from the appropriate GBuffer channel
 		float UsesStaticShadowMap = dot(LightData.ShadowMapChannelMask, float4(1, 1, 1, 1));
-		float StaticShadowing = lerp(1, dot(GBuffer.PrecomputedShadowFactors, LightData.ShadowMapChannelMask), UsesStaticShadowMap);
+		float StaticShadowing = lerp(1, dot(PrecomputedShadowFactors, LightData.ShadowMapChannelMask), UsesStaticShadowMap);
 
 		if (LightData.bRadialLight)
 		{
@@ -178,7 +182,7 @@ void GetShadowTerms(FGBufferData GBuffer, FDeferredLightData LightData, float3 W
 			// Also fix up the fade between dynamic and static shadows
 			// to work with plane splits rather than spheres.
 
-			float DynamicShadowFraction = DistanceFromCameraFade(GBuffer.Depth, LightData, WorldPosition, View.WorldCameraOrigin);
+			float DynamicShadowFraction = DistanceFromCameraFade(SceneDepth, LightData, WorldPosition, View.WorldCameraOrigin);
 			// For a directional light, fade between static shadowing and the whole scene dynamic shadowing based on distance + per object shadows
 			Shadow.SurfaceShadow = lerp(LightAttenuation.x, StaticShadowing, DynamicShadowFraction);
 			// Fade between SSS dynamic shadowing and static shadowing based on distance
@@ -198,13 +202,13 @@ void GetShadowTerms(FGBufferData GBuffer, FDeferredLightData LightData, float3 W
 		}
 	}
 
-#if SUPPORT_CONTACT_SHADOWS
-	if (LightData.ShadowedBits < 2 && (GBuffer.ShadingModelID == SHADINGMODELID_HAIR))
+#if SUPPORT_CONTACT_SHADOWS && STRATA_ENABLED == 0
+	if (LightData.ShadowedBits < 2 && (ShadingModelID == SHADINGMODELID_HAIR))
 	{
 		ContactShadowLength = 0.2 * ContactShadowLengthScreenScale;
 	}
 	// World space distance to cover eyelids and eyelashes but not beyond
-	if (GBuffer.ShadingModelID == SHADINGMODELID_EYE)
+	if (ShadingModelID == SHADINGMODELID_EYE)
 	{
 		ContactShadowLength = 0.5;
 		
@@ -229,16 +233,15 @@ void GetShadowTerms(FGBufferData GBuffer, FDeferredLightData LightData, float3 W
 			// thickness (closest-point-from-light to shading-point), and this creates light leaks. Instead we consider first hit as a blocker (old behavior)
 			BRANCH
 			if (ContactShadowOcclusion > 0.0 && 
-				IsSubsurfaceModel(GBuffer.ShadingModelID) && 
-				GBuffer.ShadingModelID != SHADINGMODELID_HAIR && 
-				GBuffer.ShadingModelID != SHADINGMODELID_EYE && 
-				GBuffer.ShadingModelID != SHADINGMODELID_SUBSURFACE_PROFILE)
+				IsSubsurfaceModel(ShadingModelID) &&
+				ShadingModelID != SHADINGMODELID_HAIR &&
+				ShadingModelID != SHADINGMODELID_EYE &&
+				ShadingModelID != SHADINGMODELID_SUBSURFACE_PROFILE)
 			{
 				// Reduce the intensity of the shadow similar to the subsurface approximation used by the shadow maps path
 				// Note that this is imperfect as we don't really have the "nearest occluder to the light", but this should at least
 				// ensure that we don't darken-out the subsurface term with the contact shadows
-				float Opacity = GBuffer.CustomData.a;
-				float Density = SubsurfaceDensityFromOpacity( Opacity );
+				float Density = SubsurfaceDensityFromOpacity(ContactShadowOpacity);
 				ContactShadowOcclusion *= 1.0 - saturate( exp( -Density * HitDistance ) );
 			}
 
@@ -352,7 +355,9 @@ FDeferredLightingSplit GetDynamicLightingSplit(
 		Shadow.TransmissionShadow = 1;
 		Shadow.TransmissionThickness = 1;
 		Shadow.HairTransmittance.OpaqueVisibility = 1;
-		GetShadowTerms(GBuffer, LightData, WorldPosition, L, LightAttenuation, Dither, Shadow);
+		const float ContactShadowOpacity = GBuffer.CustomData.a;
+		GetShadowTerms(GBuffer.Depth, GBuffer.PrecomputedShadowFactors, GBuffer.ShadingModelID, ContactShadowOpacity,
+			LightData, WorldPosition, L, LightAttenuation, Dither, Shadow);
 		SurfaceShadow = Shadow.SurfaceShadow;
 
 		LightAccumulator.EstimatedCost += 0.3f;		// add the cost of getting the shadow terms
diff --git a/Engine/Shaders/Private/DistanceFieldAOShared.ush b/Engine/Shaders/Private/DistanceFieldAOShared.ush
--- a/Engine/Shaders/Private/DistanceFieldAOShared.ush
+++ b/Engine/Shaders/Private/DistanceFieldAOShared.ush
@@ -59,13 +59,13 @@ uint2 TileListGroupSize;
 Texture2D DistanceFieldNormalTexture;
 SamplerState DistanceFieldNormalSampler;
 
-float4 EncodeDownsampledGBuffer(FGBufferData GBufferData, float SceneDepth)
+float4 EncodeDownsampledGBuffer(in float3 WorldNormal, float SceneDepth)
 {
 	#if METAL_PROFILE 
 		// clamp max depth to avoid #inf
 		SceneDepth = min(SceneDepth, 65500.0f);
 	#endif
-	return float4(GBufferData.WorldNormal.xyz, SceneDepth);
+	return float4(WorldNormal, SceneDepth);
 } 
 
 void GetDownsampledGBuffer(float2 ScreenUV, out float3 OutNormal, out float OutDepth)
diff --git a/Engine/Shaders/Private/DistanceFieldScreenGridLighting.usf b/Engine/Shaders/Private/DistanceFieldScreenGridLighting.usf
--- a/Engine/Shaders/Private/DistanceFieldScreenGridLighting.usf
+++ b/Engine/Shaders/Private/DistanceFieldScreenGridLighting.usf
@@ -10,6 +10,11 @@
 #include "DistanceFieldAOShared.ush"
 #include "GlobalDistanceFieldShared.ush"
 
+#if STRATA_ENABLED
+#define STRATA_INLINE_SHADING 0
+#include "/Engine/Private/Strata/Strata.ush"
+#endif
+
 /** Computes the distance field normal from the GBuffer. */
 void ComputeDistanceFieldNormalPS(
 	in float4 UVAndScreenPos : TEXCOORD0, 
@@ -17,11 +22,18 @@ void ComputeDistanceFieldNormalPS(
 	out float4 OutColor : SV_Target0)
 {
 	// Sample from the center of the top left full resolution texel
-	float2 ScreenUV = float2((floor(SVPos.xy) * DOWNSAMPLE_FACTOR + View.ViewRectMin.xy + .5f) * View.BufferSizeAndInvSize.zw);
+	float2 PixelCoord = floor(SVPos.xy) * DOWNSAMPLE_FACTOR + View.ViewRectMin.xy + .5f;
+	float2 ScreenUV = float2(PixelCoord * View.BufferSizeAndInvSize.zw);
+
 	float SceneDepth = CalcSceneDepth(ScreenUV);
-	FGBufferData GBufferData = GetGBufferData(ScreenUV);
 
-	OutColor = EncodeDownsampledGBuffer(GBufferData, SceneDepth);
+#if STRATA_ENABLED
+	float3 WorldNormal = StrataUnpackTopLayerData(Strata.TopLayerTexture.Load(uint3(PixelCoord, 0))).Normal;
+#else
+	float3 WorldNormal = GetGBufferData(ScreenUV).WorldNormal.xyz;
+#endif
+
+	OutColor = EncodeDownsampledGBuffer(WorldNormal, SceneDepth);
 }
 
 RWTexture2D<float4> RWDistanceFieldNormal;
@@ -32,12 +44,18 @@ void ComputeDistanceFieldNormalCS(
 	uint3 DispatchThreadId : SV_DispatchThreadID,
     uint3 GroupThreadId : SV_GroupThreadID) 
 {
-	float2 ScreenUV = float2((DispatchThreadId.xy * DOWNSAMPLE_FACTOR + View.ViewRectMin.xy + .5f) * View.BufferSizeAndInvSize.zw);
+	float2 PixelCoord = DispatchThreadId.xy * DOWNSAMPLE_FACTOR + View.ViewRectMin.xy + .5f;
+	float2 ScreenUV = float2(PixelCoord * View.BufferSizeAndInvSize.zw);
+
 	float SceneDepth = CalcSceneDepth(ScreenUV);
-	FGBufferData GBufferData = GetGBufferData(ScreenUV);
 
-	float4 OutValue = EncodeDownsampledGBuffer(GBufferData, SceneDepth);
-	RWDistanceFieldNormal[DispatchThreadId.xy] = OutValue;
+#if STRATA_ENABLED
+	float3 WorldNormal = StrataUnpackTopLayerData(Strata.TopLayerTexture.Load(uint3(PixelCoord, 0))).Normal;
+#else
+	float3 WorldNormal = GetGBufferData(ScreenUV).WorldNormal.xyz;
+#endif
+
+	RWDistanceFieldNormal[DispatchThreadId.xy] = EncodeDownsampledGBuffer(WorldNormal, SceneDepth);
 }
 
 Buffer<float4> TileConeDepthRanges;
diff --git a/Engine/Shaders/Private/PostProcessAmbientOcclusion.usf b/Engine/Shaders/Private/PostProcessAmbientOcclusion.usf
--- a/Engine/Shaders/Private/PostProcessAmbientOcclusion.usf
+++ b/Engine/Shaders/Private/PostProcessAmbientOcclusion.usf
@@ -9,6 +9,11 @@
 #include "PostProcessCommon.ush"				
 #include "DeferredShadingCommon.ush"
 
+#if STRATA_ENABLED
+#define STRATA_INLINE_SHADING 0
+#include "/Engine/Private/Strata/Strata.ush"
+#endif
+
 // set by C++:
 //
 // 0:low / 1: medium / 2:high / 4:very high
@@ -257,8 +262,14 @@ void MainSetupPS(in noperspective float4 UVAndScreenPos : TEXCOORD0, float4 SvPo
 #if COMPUTE_SHADER || FORWARD_SHADING
 		// Async compute and forward shading don't have access to the gbuffer.
 		Samples[i].rgb = normalize(ReconstructNormalFromDepthBuffer(float4(UV[i] * ViewPortSize, SvPosition.zw))) * 0.5f + 0.5f;
+#else
+#if STRATA_ENABLED
+		// STRATA_TODO TopLayerTexture is a UINT2 texture is that is going to be bandwith heavy as compared to non strata path. We should separate Normal and BaseColor-Roughness into two textures.
+		const FStrataTopLayerData TopLayerData = StrataUnpackTopLayerData(Strata.TopLayerTexture.Load(uint3(clamp(UV[i] * View.BufferSizeAndInvSize.xy, 0.0f, View.BufferSizeAndInvSize.xy - 1.0f), 0)));
+		Samples[i].rgb = TopLayerData.Normal * 0.5f + 0.5f;
 #else
 		Samples[i].rgb = GetGBufferData(UV[i], true).WorldNormal * 0.5f + 0.5f;
+#endif
 #endif
 		Samples[i].a = CalcSceneDepth(UV[i]);
 	}
@@ -307,12 +318,22 @@ float3 GetWorldSpaceNormalFromAOInput(float2 UV, float4 SvPosition)
 		// Async compute and forward shading don't have access to the gbuffer.
 		WorldNormal = ReconstructNormalFromDepthBuffer(SvPosition);
 	#else
+
+#if STRATA_ENABLED
+		const FStrataClassification StrataData = StrataUnpackClassificationData(Strata.ClassificationTexture.Load(uint3(SvPosition.xy, 0)));
+		if (StrataData.ShadingModels != STRATA_BSDF_TYPE_HAIR)
+		{
+			const FStrataTopLayerData TopLayerData = StrataUnpackTopLayerData(Strata.TopLayerTexture.Load(uint3(SvPosition.xy, 0)));
+			WorldNormal = TopLayerData.Normal;
+		}
+#else
 		// Otherwise sample Gbuffer normals if the shader model has normals.
 		FGBufferData GBuffer = GetGBufferData(UV, false);
 		if( GBuffer.ShadingModelID != SHADINGMODELID_HAIR )
 		{
 			WorldNormal = GetGBufferData(UV, false).WorldNormal;
 		}
+#endif
 	#endif
 	}
 
@@ -961,8 +982,13 @@ float3 GetNormal(float2 UV, float3 ViewSpacePosMid)
 
 #if USE_NORMALBUFFER
 
+#if STRATA_ENABLED
+	const FStrataTopLayerData TopLayerData = StrataUnpackTopLayerData(Strata.TopLayerTexture.Load(uint3(clamp(UV * View.BufferSizeAndInvSize.xy, 0.0f, View.BufferSizeAndInvSize.xy-1.0f), 0)));
+	float3 WorldNormal = TopLayerData.Normal;
+#else
 	// Get the normal from the normal buffer
-	float3 WorldNormal		= GetGBufferData(UV, false).WorldNormal;
+	float3 WorldNormal = GetGBufferData(UV, false).WorldNormal;
+#endif
 	ViewSpaceNormal	= normalize(mul(WorldNormal, (float3x3)View.TranslatedWorldToView));
 
 #else
@@ -1334,7 +1360,12 @@ float GTAOInnerIntegratePSandCS(in float2 UV, in uint2 iPos)
 
 	// ViewspacePos and Normal
 	float3 ViewSpacePos		= ScreenToViewPos(UV, SceneDepth);
+#if STRATA_ENABLED
+	const FStrataTopLayerData TopLayerData = StrataUnpackTopLayerData(Strata.TopLayerTexture.Load(uint3(iPos, 0)));
+	float3 WorldNormal = TopLayerData.Normal;
+#else
 	float3 WorldNormal		= GetGBufferData(UV, false).WorldNormal;
+#endif
 	float3 ViewSpaceNormal	= normalize(mul(WorldNormal, (float3x3)View.TranslatedWorldToView));
 	float3 ViewDir			= -normalize(ViewSpacePos.xyz);	// TODO - This is a function of UV only.
 
diff --git a/Engine/Shaders/Private/SSRT/SSRTReflections.usf b/Engine/Shaders/Private/SSRT/SSRTReflections.usf
--- a/Engine/Shaders/Private/SSRT/SSRTReflections.usf
+++ b/Engine/Shaders/Private/SSRT/SSRTReflections.usf
@@ -7,6 +7,11 @@
 #include "../SceneTextureParameters.ush"
 #include "../ScreenSpaceDenoise/SSDPublic.ush"
 
+#if STRATA_ENABLED
+#define STRATA_INLINE_SHADING 0
+#include "/Engine/Private/Strata/Strata.ush"
+#endif
+
 // Generate vector truncation warnings to errors.
 #pragma warning(error: 3206)
 
@@ -68,13 +73,27 @@ void ScreenSpaceReflectionsStencilPS(
 {
 	float2 UV = SvPosition.xy * View.BufferSizeAndInvSize.zw;
 	
+#if STRATA_ENABLED
+
+	uint2 PixelCoord = SvPosition.xy;
+	const FStrataClassification StrataData = StrataUnpackClassificationData(Strata.ClassificationTexture.Load(uint3(PixelCoord, 0)));
+	const FStrataTopLayerData TopLayerData = StrataUnpackTopLayerData(Strata.TopLayerTexture.Load(uint3(PixelCoord, 0)));
+
+	const float Roughness = TopLayerData.Roughness;
+	const bool bNoMaterial = !StrataData.bIsStrataMaterial;
+
+#else // STRATA_ENABLED
+
 	FGBufferData GBuffer = GetGBufferDataFromSceneTextures(UV);
-	
 	const float Roughness = GetRoughness(GBuffer);
+	const bool bNoMaterial = GBuffer.ShadingModelID == 0;
+
+#endif // STRATA_ENABLED
+
 	const float RoughnessFade = GetRoughnessFade(Roughness);
-	
+
 #if !defined(TILE_COMPUTE_SSR)
-	if (RoughnessFade > 0.0 && GBuffer.ShadingModelID != 0)
+	if (RoughnessFade > 0.0 && bNoMaterial)
 	{
 		// we are going to compute SSR for this pixel, so we discard this
 		// pixel shader invocation to not overwrite the stencil buffer and
@@ -82,7 +101,7 @@ void ScreenSpaceReflectionsStencilPS(
 		discard;
 	}
 #endif
-	
+
 	// we are not going to compute SSR for this pixel, so we clear the color
 	// since ScreenSpaceReflectionsPS() won't be executed in this pixel.
 	OutColor = 0;
@@ -111,20 +130,46 @@ void ScreenSpaceReflections(
 		OutClosestHitDistance = -2.0;
 	#endif
 
+#if STRATA_ENABLED
+
+	uint2 PixelCoord = SvPosition.xy;
+	const FStrataClassification StrataData = StrataUnpackClassificationData(Strata.ClassificationTexture.Load(uint3(PixelCoord, 0)));
+	const FStrataTopLayerData TopLayerData = StrataUnpackTopLayerData(Strata.TopLayerTexture.Load(uint3(PixelCoord, 0)));
+
+	float3 N = TopLayerData.Normal;
+	float DeviceZ = SampleDeviceZFromSceneTextures(UV);
+	float SceneDepth = ConvertFromDeviceZ(DeviceZ);
+
+	const float3 PositionTranslatedWorld = mul(float4(ScreenPos * SceneDepth, SceneDepth, 1), View.ScreenToTranslatedWorld).xyz;
+	const float3 V = normalize(View.TranslatedWorldCameraOrigin - PositionTranslatedWorld);
+
+	// STRATA_TODO Handle anisotrpoy using ModifyGGXAnisotropicNormalRoughness or similar approach
+
+	// For now, SSR is only used to trace the top material layer.
+	// STRATA_TODO try stochastic selection
+	const float Roughness = TopLayerData.Roughness;
+	const bool bNoMaterial = !StrataData.bIsStrataMaterial;
+
+#else // STRATA_ENABLED
+
 	FGBufferData GBuffer = GetGBufferDataFromSceneTextures(UV);
 
 	float3 N = GBuffer.WorldNormal;
 	const float SceneDepth = GBuffer.Depth;
-	const float3 PositionTranslatedWorld = mul( float4( ScreenPos * SceneDepth, SceneDepth, 1 ), View.ScreenToTranslatedWorld ).xyz;
+	const float3 PositionTranslatedWorld = mul(float4(ScreenPos * SceneDepth, SceneDepth, 1), View.ScreenToTranslatedWorld).xyz;
 	const float3 V = normalize(View.TranslatedWorldCameraOrigin - PositionTranslatedWorld);
 
 	ModifyGGXAnisotropicNormalRoughness(GBuffer.WorldTangent, GBuffer.Anisotropy, GBuffer.Roughness, N, V);
-	
+
 	float Roughness = GetRoughness(GBuffer);
+	const bool bNoMaterial = GBuffer.ShadingModelID == 0;
+
+#endif // STRATA_ENABLED
+
 	float RoughnessFade = GetRoughnessFade(Roughness);
 
 	// Early out. Useless if using the stencil prepass.
-	BRANCH if( RoughnessFade <= 0.0 || GBuffer.ShadingModelID == 0 )
+	BRANCH if( RoughnessFade <= 0.0 || bNoMaterial)
 	{
 		return;
 	}
diff --git a/Engine/Shaders/Private/ShadowProjectionPixelShader.usf b/Engine/Shaders/Private/ShadowProjectionPixelShader.usf
--- a/Engine/Shaders/Private/ShadowProjectionPixelShader.usf
+++ b/Engine/Shaders/Private/ShadowProjectionPixelShader.usf
@@ -96,7 +96,7 @@ Texture2D<uint4> HairCategorizationTexture;
 
 #if STRATA_ENABLED
 #define STRATA_INLINE_SHADING 0
-#include "Strata/StrataEvaluation.ush"
+#include "/Engine/Private/Strata/Strata.ush"
 #endif
 
 /**
@@ -213,10 +213,18 @@ void Main(
 	{
 		#if SHADING_PATH_DEFERRED && !FORWARD_SHADING && !SUBPIXEL_SHADOW
 			// Attenuate soft transition based on the angle with the light and the shading normal (acts as a receiver bias)
-			FGBufferData GBufferData = GetGBufferData(ScreenUV);
 			const bool bIsDirectional = LightPositionOrDirection.w == 0;
 			const float3 LightDirection = bIsDirectional ? -LightPositionOrDirection.xyz : normalize(LightPositionOrDirection.xyz - WorldPosition);
-			const float NoL = saturate(dot(GBufferData.WorldNormal, LightDirection));
+
+		#if STRATA_ENABLED
+			const FStrataTopLayerData TopLayerData = StrataUnpackTopLayerData(Strata.TopLayerTexture.Load(uint3(SVPos.xy, 0)));
+			const float3 WorldNormal = TopLayerData.Normal;
+		#else
+			FGBufferData GBufferData = GetGBufferData(ScreenUV);
+			const float3 WorldNormal = GBufferData.WorldNormal;
+		#endif
+
+			const float NoL = saturate(dot(WorldNormal, LightDirection));
 		#endif
 		
 		FPCFSamplerSettings Settings;
diff --git a/Engine/Shaders/Private/Strata/Strata.ush b/Engine/Shaders/Private/Strata/Strata.ush
--- a/Engine/Shaders/Private/Strata/Strata.ush
+++ b/Engine/Shaders/Private/Strata/Strata.ush
@@ -57,12 +57,16 @@ struct FStrataBSDF
 #define HEADER_BIT_COUNT_SHAREDNORMAL_COUNT		3
 #define HEADER_BIT_COUNT_SHAREDNORMAL_TYPE		4
 #define HEADER_BIT_COUNT_AO						8
-#define HEADER_BIT_COUNT___UNUSED___			13
+#define HEADER_BIT_COUNT_SKIP_PRECSHADOW_MASK	1
+#define HEADER_BIT_COUNT_ZERO_PRECSHADOW_MASK	1
+#define HEADER_BIT_COUNT___UNUSED___			11
 
 #define HEADER_BIT_OFFSET_BSDF_COUNT			0
 #define HEADER_BIT_OFFSET_SHAREDNORMAL_COUNT	(HEADER_BIT_OFFSET_BSDF_COUNT				+ HEADER_BIT_COUNT_BSDF_COUNT)
 #define HEADER_BIT_OFFSET_SHAREDNORMAL_TYPE		(HEADER_BIT_OFFSET_SHAREDNORMAL_COUNT		+ HEADER_BIT_COUNT_SHAREDNORMAL_COUNT)
 #define HEADER_BIT_OFFSET_AO					(HEADER_BIT_OFFSET_SHAREDNORMAL_TYPE		+ HEADER_BIT_COUNT_SHAREDNORMAL_TYPE)
+#define HEADER_BIT_OFFSET_SKIP_PRECSHADOW_MASK	(HEADER_BIT_OFFSET_AO						+ HEADER_BIT_COUNT_AO)
+#define HEADER_BIT_OFFSET_ZERO_PRECSHADOW_MASK	(HEADER_BIT_OFFSET_SKIP_PRECSHADOW_MASK		+ HEADER_BIT_COUNT_SKIP_PRECSHADOW_MASK)
 
 #define HEADER_GETBSDFCOUNT(X)					READ_BITS( X, HEADER_BIT_COUNT_BSDF_COUNT, HEADER_BIT_OFFSET_BSDF_COUNT)
 #define HEADER_SETBSDFCOUNT(X, NewValue)		WRITE_BITS(X, HEADER_BIT_COUNT_BSDF_COUNT, HEADER_BIT_OFFSET_BSDF_COUNT, NewValue)
@@ -76,6 +80,12 @@ struct FStrataBSDF
 #define HEADER_GETAO(X)							READ_BITS( X, HEADER_BIT_COUNT_AO, HEADER_BIT_OFFSET_AO)
 #define HEADER_SETAO(X, NewValue)				WRITE_BITS(X, HEADER_BIT_COUNT_AO, HEADER_BIT_OFFSET_AO, NewValue)
 
+#define HEADER_GETSKIPPRECSHADOWMASK(X)			READ_BITS(  X, HEADER_BIT_COUNT_SKIP_PRECSHADOW_MASK, HEADER_BIT_OFFSET_SKIP_PRECSHADOW_MASK)
+#define HEADER_SETSKIPPRECSHADOWMASK(X, NewValue)WRITE_BITS(X, HEADER_BIT_COUNT_SKIP_PRECSHADOW_MASK, HEADER_BIT_OFFSET_SKIP_PRECSHADOW_MASK, NewValue)
+
+#define HEADER_GETZEROPRECSHADOWMASK(X)			READ_BITS(  X, HEADER_BIT_COUNT_ZERO_PRECSHADOW_MASK, HEADER_BIT_OFFSET_ZERO_PRECSHADOW_MASK)
+#define HEADER_SETZEROPRECSHADOWMASK(X, NewValue)WRITE_BITS(X, HEADER_BIT_COUNT_ZERO_PRECSHADOW_MASK, HEADER_BIT_OFFSET_ZERO_PRECSHADOW_MASK, NewValue)
+
 // Bit count for each component stored in a Strata BSDF state
 #define STATE_BIT_COUNT_BSDF					3
 #define STATE_BIT_COUNT_SHAREDNORMALID			2
@@ -638,6 +648,8 @@ struct FStrataPixelHeader
 #endif
 
 	float MaterialAO;
+	bool SkipPrecShadowMask;
+	bool ZeroPrecShadowMask;
 };
 
 FStrataPixelHeader InitialiseStrataPixelHeader()
@@ -651,6 +663,8 @@ FStrataPixelHeader InitialiseStrataPixelHeader()
 	Out.SharedNormalTypes = 0;
 #endif
 	Out.MaterialAO = 1.0f;
+	Out.SkipPrecShadowMask = false;
+	Out.ZeroPrecShadowMask = false;
 	return Out;
 }
 
@@ -1424,6 +1438,9 @@ uint PackStrataHeaderIn(uint InBSDFCount, FStrataPixelHeader InHeader)
 #endif
 
 	HEADER_SETAO(Out, uint(InHeader.MaterialAO * 255.0f));
+
+	HEADER_SETSKIPPRECSHADOWMASK(Out, InHeader.SkipPrecShadowMask ? 1u : 0u);
+	HEADER_SETZEROPRECSHADOWMASK(Out, InHeader.ZeroPrecShadowMask ? 1u : 0u);
 	return Out;
 }
 
@@ -1454,6 +1471,9 @@ FStrataPixelHeader UnpackStrataHeaderIn(ByteAddressBuffer StrataBuffer, inout FS
 
 	Out.MaterialAO = float(HEADER_GETAO(PackedHeader)) * (1.0f / 255.0f);
 
+	Out.SkipPrecShadowMask = HEADER_GETSKIPPRECSHADOWMASK(PackedHeader);
+	Out.ZeroPrecShadowMask = HEADER_GETSKIPPRECSHADOWMASK(PackedHeader);
+
 	return Out;
 }
 
diff --git a/Engine/Shaders/Private/Strata/StrataMaterialClassification.usf b/Engine/Shaders/Private/Strata/StrataMaterialClassification.usf
--- a/Engine/Shaders/Private/Strata/StrataMaterialClassification.usf
+++ b/Engine/Shaders/Private/Strata/StrataMaterialClassification.usf
@@ -160,6 +160,10 @@ void MainPS(
 
 		FinalizeTopLayerData(TopLayerData);
 	}
+	else
+	{
+		TopLayerData.Normal = normalize(float3(-1.0f, -1.0f, -1.0f));
+	}
 
 	OutColor0 = StrataPackClassificationData(ClassificationData);
 	OutColor1 = StrataPackTopLayerData(TopLayerData);
diff --git a/Engine/Shaders/Private/Strata/StrataVisualize.usf b/Engine/Shaders/Private/Strata/StrataVisualize.usf
--- a/Engine/Shaders/Private/Strata/StrataVisualize.usf
+++ b/Engine/Shaders/Private/Strata/StrataVisualize.usf
@@ -189,6 +189,8 @@ void VisualizeMaterialPS(
 
 #define PRINTBSDFCOUNT()	PRINTBSDFCHAR(min(_0_+Header.BSDFCount,_F_)); PRINTBSDFCHAR(_SPC_); PRINTBSDFCHAR(_B_); PRINTBSDFCHAR(_S_); PRINTBSDFCHAR(_D_); PRINTBSDFCHAR(_F_);
 #define PRINTMATERIALAO()	PRINTBSDFCHAR(_A_); PRINTBSDFCHAR(_O_); PRINTBSDFCHAR(_SPC_); PRINTFLOAT1(Header.MaterialAO);
+#define PRINTMATERIALSKIPPRECSHADOWMASK()	PRINTBSDFCHAR(_S_); PRINTBSDFCHAR(_K_); PRINTBSDFCHAR(_I_); PRINTBSDFCHAR(_P_); PRINTBSDFCHAR(_P_); PRINTBSDFCHAR(_R_); PRINTBSDFCHAR(_E_); PRINTBSDFCHAR(_S_); PRINTBSDFCHAR(_H_); PRINTBSDFCHAR(_A_); PRINTBSDFCHAR(_D_); PRINTBSDFCHAR(_SPC_); PRINTBSDFCHAR(_0_+(Header.SkipPrecShadowMask?1:0));
+#define PRINTMATERIALZEROPRECSHADOWMASK()	PRINTBSDFCHAR(_Z_); PRINTBSDFCHAR(_E_); PRINTBSDFCHAR(_R_); PRINTBSDFCHAR(_O_); PRINTBSDFCHAR(_P_); PRINTBSDFCHAR(_R_); PRINTBSDFCHAR(_E_); PRINTBSDFCHAR(_S_); PRINTBSDFCHAR(_H_); PRINTBSDFCHAR(_A_); PRINTBSDFCHAR(_D_); PRINTBSDFCHAR(_SPC_); PRINTBSDFCHAR(_0_+(Header.ZeroPrecShadowMask?1:0));
 
 #define NEXTLINE  TopLeft.y += 8; RectMax.x = max(RectMax.x, TopLeft.x); RectMax.y = max(RectMax.y, TopLeft.y); TopLeft.x = TopLeftInit.x; 
 #define HORIZONTALSEPARATION  TopLeft.y += 4;
@@ -202,12 +204,18 @@ void VisualizeMaterialPS(
 		NEXTLINE;
 		PRINTMATERIALAO();
 		NEXTLINE;
+		PRINTMATERIALSKIPPRECSHADOWMASK();
+		NEXTLINE;
+		PRINTMATERIALZEROPRECSHADOWMASK();
+		NEXTLINE;
 		HORIZONTALSEPARATION;
 #else
 		NEXTLINE;
 		HORIZONTALSEPARATION;
 		NEXTLINE;
 		NEXTLINE;
+		NEXTLINE;
+		NEXTLINE;
 		HORIZONTALSEPARATION;
 #endif
 
diff --git a/Engine/Source/Runtime/Engine/Private/ShaderCompiler/ShaderGenerationUtil.cpp b/Engine/Source/Runtime/Engine/Private/ShaderCompiler/ShaderGenerationUtil.cpp
--- a/Engine/Source/Runtime/Engine/Private/ShaderCompiler/ShaderGenerationUtil.cpp
+++ b/Engine/Source/Runtime/Engine/Private/ShaderCompiler/ShaderGenerationUtil.cpp
@@ -1666,23 +1666,24 @@ static void SetSlotsForShadingModelType(bool Slots[], EMaterialShadingModel Shad
 
 
 
-static void SetStandardGBufferSlots(bool Slots[], bool bWriteEmissive, bool bHasTangent, bool bHasVelocity, bool bHasStaticLighting)
+static void SetStandardGBufferSlots(bool Slots[], bool bWriteEmissive, bool bHasTangent, bool bHasVelocity, bool bHasStaticLighting, bool bIsStrataMaterial)
 {
 	Slots[GBS_SceneColor] = bWriteEmissive;
-	Slots[GBS_WorldNormal] = true;
-	Slots[GBS_PerObjectGBufferData] = true;
-	Slots[GBS_Metallic] = true;
-	Slots[GBS_Specular] = true;
-	Slots[GBS_Roughness] = true;
-	Slots[GBS_ShadingModelId] = true;
-	Slots[GBS_SelectiveOutputMask] = true;
-	Slots[GBS_BaseColor] = true;
-	Slots[GBS_GenericAO] = true;
-	Slots[GBS_AO] = false;// true;
 	Slots[GBS_Velocity] = bHasVelocity;
 	Slots[GBS_PrecomputedShadowFactor] = bHasStaticLighting;
-	Slots[GBS_WorldTangent] = bHasTangent;
-	Slots[GBS_Anisotropy] = bHasTangent;
+
+	Slots[GBS_WorldNormal] =			bIsStrataMaterial ? false : true;
+	Slots[GBS_PerObjectGBufferData] =	bIsStrataMaterial ? false : true;
+	Slots[GBS_Metallic] =				bIsStrataMaterial ? false : true;
+	Slots[GBS_Specular] =				bIsStrataMaterial ? false : true;
+	Slots[GBS_Roughness] =				bIsStrataMaterial ? false : true;
+	Slots[GBS_ShadingModelId] =			bIsStrataMaterial ? false : true;
+	Slots[GBS_SelectiveOutputMask] =	bIsStrataMaterial ? false : true;
+	Slots[GBS_BaseColor] =				bIsStrataMaterial ? false : true;
+	Slots[GBS_GenericAO] =				bIsStrataMaterial ? false : true;
+	Slots[GBS_AO] =						bIsStrataMaterial ? false : false;// true;		// Why false?
+	Slots[GBS_WorldTangent] =			bIsStrataMaterial ? false : bHasTangent;
+	Slots[GBS_Anisotropy] =				bIsStrataMaterial ? false : bHasTangent;
 }
 
 static void DetermineUsedMaterialSlots(
@@ -1698,6 +1699,7 @@ static void DetermineUsedMaterialSlots(
 	bool bHasTangent = SrcGlobal.GBUFFER_HAS_TANGENT;
 	bool bHasVelocity = Dst.WRITES_VELOCITY_TO_GBUFFER;
 	bool bHasStaticLighting = Dst.GBUFFER_HAS_PRECSHADOWFACTOR || Dst.WRITES_PRECSHADOWFACTOR_TO_GBUFFER;
+	bool bIsStrataMaterial = Mat.PROJECT_STRATA && Mat.MATERIAL_IS_STRATA;
 
 	// we have to use if statements, not switch or if/else statements because we can have multiple shader model ids.
 	if (Mat.MATERIAL_SHADINGMODEL_UNLIT)
@@ -1707,61 +1709,61 @@ static void DetermineUsedMaterialSlots(
 
 	if (Mat.MATERIAL_SHADINGMODEL_DEFAULT_LIT)
 	{
-		SetStandardGBufferSlots(Slots, bWriteEmissive, bHasTangent, bHasVelocity, bHasStaticLighting);
+		SetStandardGBufferSlots(Slots, bWriteEmissive, bHasTangent, bHasVelocity, bHasStaticLighting, bIsStrataMaterial);
 	}
 
 	if (Mat.MATERIAL_SHADINGMODEL_SUBSURFACE)
 	{
-		SetStandardGBufferSlots(Slots, bWriteEmissive, bHasTangent, bHasVelocity, bHasStaticLighting);
+		SetStandardGBufferSlots(Slots, bWriteEmissive, bHasTangent, bHasVelocity, bHasStaticLighting, bIsStrataMaterial);
 		Slots[GBS_CustomData] = true;
 	}
 
 	if (Mat.MATERIAL_SHADINGMODEL_PREINTEGRATED_SKIN)
 	{
-		SetStandardGBufferSlots(Slots, bWriteEmissive, bHasTangent, bHasVelocity, bHasStaticLighting);
+		SetStandardGBufferSlots(Slots, bWriteEmissive, bHasTangent, bHasVelocity, bHasStaticLighting, bIsStrataMaterial);
 		Slots[GBS_CustomData] = true;
 	}
 
 	if (Mat.MATERIAL_SHADINGMODEL_SUBSURFACE_PROFILE)
 	{
-		SetStandardGBufferSlots(Slots, bWriteEmissive, bHasTangent, bHasVelocity, bHasStaticLighting);
+		SetStandardGBufferSlots(Slots, bWriteEmissive, bHasTangent, bHasVelocity, bHasStaticLighting, bIsStrataMaterial);
 		Slots[GBS_CustomData] = true;
 	}
 
 	if (Mat.MATERIAL_SHADINGMODEL_CLEAR_COAT)
 	{
-		SetStandardGBufferSlots(Slots, bWriteEmissive, bHasTangent, bHasVelocity, bHasStaticLighting);
+		SetStandardGBufferSlots(Slots, bWriteEmissive, bHasTangent, bHasVelocity, bHasStaticLighting, bIsStrataMaterial);
 		Slots[GBS_CustomData] = true;
 	}
 
 	if (Mat.MATERIAL_SHADINGMODEL_TWOSIDED_FOLIAGE)
 	{
-		SetStandardGBufferSlots(Slots, bWriteEmissive, bHasTangent, bHasVelocity, bHasStaticLighting);
+		SetStandardGBufferSlots(Slots, bWriteEmissive, bHasTangent, bHasVelocity, bHasStaticLighting, bIsStrataMaterial);
 		Slots[GBS_CustomData] = true;
 	}
 
 	if (Mat.MATERIAL_SHADINGMODEL_HAIR)
 	{
-		SetStandardGBufferSlots(Slots, bWriteEmissive, bHasTangent, bHasVelocity, bHasStaticLighting);
+		SetStandardGBufferSlots(Slots, bWriteEmissive, bHasTangent, bHasVelocity, bHasStaticLighting, bIsStrataMaterial);
 		Slots[GBS_CustomData] = true;
 	}
 
 	if (Mat.MATERIAL_SHADINGMODEL_CLOTH)
 	{
-		SetStandardGBufferSlots(Slots, bWriteEmissive, bHasTangent, bHasVelocity, bHasStaticLighting);
+		SetStandardGBufferSlots(Slots, bWriteEmissive, bHasTangent, bHasVelocity, bHasStaticLighting, bIsStrataMaterial);
 		Slots[GBS_CustomData] = true;
 	}
 
 	if (Mat.MATERIAL_SHADINGMODEL_EYE)
 	{
-		SetStandardGBufferSlots(Slots, bWriteEmissive, bHasTangent, bHasVelocity, bHasStaticLighting);
+		SetStandardGBufferSlots(Slots, bWriteEmissive, bHasTangent, bHasVelocity, bHasStaticLighting, bIsStrataMaterial);
 		Slots[GBS_CustomData] = true;
 	}
 
 	if (Mat.MATERIAL_SHADINGMODEL_SINGLELAYERWATER)
 	{
 		// single layer water uses standard slots
-		SetStandardGBufferSlots(Slots, bWriteEmissive, bHasTangent, bHasVelocity, bHasStaticLighting);
+		SetStandardGBufferSlots(Slots, bWriteEmissive, bHasTangent, bHasVelocity, bHasStaticLighting, bIsStrataMaterial);
 	}
 
 	// doesn't write to GBuffer
diff --git a/Engine/Source/Runtime/Renderer/Private/CompositionLighting/PostProcessAmbientOcclusion.cpp b/Engine/Source/Runtime/Renderer/Private/CompositionLighting/PostProcessAmbientOcclusion.cpp
--- a/Engine/Source/Runtime/Renderer/Private/CompositionLighting/PostProcessAmbientOcclusion.cpp
+++ b/Engine/Source/Runtime/Renderer/Private/CompositionLighting/PostProcessAmbientOcclusion.cpp
@@ -8,6 +8,7 @@
 #include "CompositionLighting/CompositionLighting.h"
 #include "SceneTextureParameters.h"
 #include "ScenePrivate.h"
+#include "Strata/Strata.h"
 
 DECLARE_GPU_STAT_NAMED(SSAOSetup, TEXT("ScreenSpace AO Setup") );
 DECLARE_GPU_STAT_NAMED(SSAO, TEXT("ScreenSpace AO") );
@@ -514,6 +515,9 @@ public:
 	DECLARE_GLOBAL_SHADER(FAmbientOcclusionSetupPS);
 	SHADER_USE_PARAMETER_STRUCT(FAmbientOcclusionSetupPS, FGlobalShader);
 
+	class FStrata : SHADER_PERMUTATION_BOOL("STRATA_ENABLED");
+	using FPermutationDomain = TShaderPermutationDomain<FStrata>;
+
 	static bool ShouldCompilePermutation(const FGlobalShaderPermutationParameters& Parameters)
 	{
 		return IsFeatureLevelSupported(Parameters.Platform, ERHIFeatureLevel::SM5);
@@ -528,6 +532,8 @@ public:
 		SHADER_PARAMETER(float, ThresholdInverse)
 		SHADER_PARAMETER(FVector2D, InputExtentInverse)
 
+		SHADER_PARAMETER_STRUCT_REF(FStrataGlobalUniformParameters, Strata)
+
 		RENDER_TARGET_BINDING_SLOTS()
 	END_SHADER_PARAMETER_STRUCT();
 };
@@ -563,15 +569,20 @@ FScreenPassTexture AddAmbientOcclusionSetupPass(
 	const FFinalPostProcessSettings& Settings = View.FinalPostProcessSettings;
 	const float ThresholdInverseValue = Settings.AmbientOcclusionMipThreshold * ((float)OutputViewport.Extent.X / (float)CommonParameters.SceneTexturesViewport.Extent.X);
 
+
 	FAmbientOcclusionSetupPS::FParameters* PassParameters = GraphBuilder.AllocParameters<FAmbientOcclusionSetupPS::FParameters>();
 	PassParameters->View = View.ViewUniformBuffer;
 	PassParameters->SceneTextures = CommonParameters.SceneTexturesUniformBuffer;
 	PassParameters->SSAOParameters = GetSSAOShaderParameters(View, InputViewport, OutputViewport, CommonParameters.SceneTexturesViewport, EAOTechnique::SSAO);
 	PassParameters->ThresholdInverse = ThresholdInverseValue;
 	PassParameters->InputExtentInverse = FVector2D(1.0f) / FVector2D(InputViewport.Extent);
+	PassParameters->Strata = Strata::BindStrataGlobalUniformParameters(View);
 	PassParameters->RenderTargets[0] = Output.GetRenderTargetBinding();
 
-	TShaderMapRef<FAmbientOcclusionSetupPS> PixelShader(View.ShaderMap);
+	FAmbientOcclusionSetupPS::FPermutationDomain PermutationVector;
+	PermutationVector.Set<FAmbientOcclusionSetupPS::FStrata>(Strata::IsStrataEnabled());
+	TShaderMapRef<FAmbientOcclusionSetupPS> PixelShader(View.ShaderMap, PermutationVector);
+
 	AddDrawScreenPass(
 		GraphBuilder,
 		RDG_EVENT_NAME("AmbientOcclusionSetup %dx%d", OutputViewport.Rect.Width(), OutputViewport.Rect.Height()),
@@ -697,8 +708,9 @@ public:
 	class FUseUpsampleDim       : SHADER_PERMUTATION_BOOL("USE_UPSAMPLE");
 	class FUseAoSetupAsInputDim : SHADER_PERMUTATION_BOOL("USE_AO_SETUP_AS_INPUT");
 	class FShaderQualityDim     : SHADER_PERMUTATION_INT("SHADER_QUALITY", 5);
+	class FStrata				: SHADER_PERMUTATION_BOOL("STRATA_ENABLED");
 
-	using FPermutationDomain = TShaderPermutationDomain<FUseUpsampleDim, FUseAoSetupAsInputDim, FShaderQualityDim>;
+	using FPermutationDomain = TShaderPermutationDomain<FUseUpsampleDim, FUseAoSetupAsInputDim, FShaderQualityDim, FStrata>;
 
 	static void ModifyCompilationEnvironment(const FGlobalShaderPermutationParameters& Parameters, FShaderCompilerEnvironment& OutEnvironment)
 	{
@@ -713,7 +725,7 @@ public:
 
 	BEGIN_SHADER_PARAMETER_STRUCT(FParameters, )
 		SHADER_PARAMETER_STRUCT_INCLUDE(FAmbientOcclusionParameters, SharedParameters)
-
+		SHADER_PARAMETER_STRUCT_REF(FStrataGlobalUniformParameters, Strata)
 		RENDER_TARGET_BINDING_SLOTS()
 	END_GLOBAL_SHADER_PARAMETER_STRUCT();
 };
@@ -728,8 +740,9 @@ public:
 	class FUseUpsampleDim       : SHADER_PERMUTATION_BOOL("USE_UPSAMPLE");
 	class FUseAoSetupAsInputDim : SHADER_PERMUTATION_BOOL("USE_AO_SETUP_AS_INPUT");
 	class FShaderQualityDim     : SHADER_PERMUTATION_INT("SHADER_QUALITY", 5);
+	class FStrata				: SHADER_PERMUTATION_BOOL("STRATA_ENABLED");
 
-	using FPermutationDomain = TShaderPermutationDomain<FUseUpsampleDim, FUseAoSetupAsInputDim, FShaderQualityDim>;
+	using FPermutationDomain = TShaderPermutationDomain<FUseUpsampleDim, FUseAoSetupAsInputDim, FShaderQualityDim, FStrata>;
 
 	static void ModifyCompilationEnvironment(const FGlobalShaderPermutationParameters& Parameters, FShaderCompilerEnvironment& OutEnvironment)
 	{
@@ -746,7 +759,7 @@ public:
 
 	BEGIN_SHADER_PARAMETER_STRUCT(FParameters, )
 		SHADER_PARAMETER_STRUCT_INCLUDE(FAmbientOcclusionParameters, SharedParameters)
-
+		SHADER_PARAMETER_STRUCT_REF(FStrataGlobalUniformParameters, Strata)
 		SHADER_PARAMETER_RDG_TEXTURE_UAV(RWTexture2D, OutTexture)
 	END_GLOBAL_SHADER_PARAMETER_STRUCT();
 };
@@ -844,12 +857,14 @@ FScreenPassTexture AddAmbientOcclusionPass(
 		// Compute Shader Path
 		FAmbientOcclusionCS::FParameters* PassParameters = GraphBuilder.AllocParameters<FAmbientOcclusionCS::FParameters>();
 		PassParameters->SharedParameters = MoveTemp(SharedParameters);
+		PassParameters->Strata = Strata::BindStrataGlobalUniformParameters(View);
 		PassParameters->OutTexture = GraphBuilder.CreateUAV(Output.Texture);
 
 		FAmbientOcclusionCS::FPermutationDomain PermutationVector;
 		PermutationVector.Set<FAmbientOcclusionCS::FUseUpsampleDim>(bDoUpsample);
 		PermutationVector.Set<FAmbientOcclusionCS::FUseAoSetupAsInputDim>(bAOSetupAsInput);
 		PermutationVector.Set<FAmbientOcclusionCS::FShaderQualityDim>(CommonParameters.ShaderQuality);
+		PermutationVector.Set<FAmbientOcclusionCS::FStrata>(Strata::IsStrataEnabled());
 
 		TShaderMapRef<FAmbientOcclusionCS> ComputeShader(View.ShaderMap, PermutationVector);
 		FComputeShaderUtils::AddPass(
@@ -865,12 +880,14 @@ FScreenPassTexture AddAmbientOcclusionPass(
 		// Pixel Shader Path
 		FAmbientOcclusionPS::FParameters* PassParameters = GraphBuilder.AllocParameters<FAmbientOcclusionPS::FParameters>();
 		PassParameters->SharedParameters = MoveTemp(SharedParameters);
+		PassParameters->Strata = Strata::BindStrataGlobalUniformParameters(View);
 		PassParameters->RenderTargets[0] = Output.GetRenderTargetBinding();
 
 		FAmbientOcclusionPS::FPermutationDomain PermutationVector;
 		PermutationVector.Set<FAmbientOcclusionPS::FUseUpsampleDim>(bDoUpsample);
 		PermutationVector.Set<FAmbientOcclusionPS::FUseAoSetupAsInputDim>(bAOSetupAsInput);
 		PermutationVector.Set<FAmbientOcclusionPS::FShaderQualityDim>(CommonParameters.ShaderQuality);
+		PermutationVector.Set<FAmbientOcclusionPS::FStrata>(Strata::IsStrataEnabled());
 
 		TShaderMapRef<FAmbientOcclusionPS> PixelShader(View.ShaderMap, PermutationVector);
 		AddDrawScreenPass(
@@ -957,8 +974,9 @@ public:
 
 	class FShaderQualityDim : SHADER_PERMUTATION_INT("SHADER_QUALITY", 5);
 	class FUseNormalBufferDim : SHADER_PERMUTATION_BOOL("USE_NORMALBUFFER");
+	class FStrata : SHADER_PERMUTATION_BOOL("STRATA_ENABLED");
 
-	using FPermutationDomain = TShaderPermutationDomain<FShaderQualityDim, FUseNormalBufferDim>;
+	using FPermutationDomain = TShaderPermutationDomain<FShaderQualityDim, FUseNormalBufferDim, FStrata>;
 
 	static void ModifyCompilationEnvironment(const FGlobalShaderPermutationParameters& Parameters, FShaderCompilerEnvironment& OutEnvironment)
 	{
@@ -977,6 +995,7 @@ public:
 		SHADER_PARAMETER_STRUCT_REF(FViewUniformShaderParameters, View)
 		SHADER_PARAMETER_RDG_UNIFORM_BUFFER(FSceneTextureUniformParameters, SceneTextures)
 		SHADER_PARAMETER_STRUCT_INCLUDE(FHZBParameters, HZBParameters)
+		SHADER_PARAMETER_STRUCT_REF(FStrataGlobalUniformParameters, Strata)
 
 		SHADER_PARAMETER_STRUCT_INCLUDE(FSSAOShaderParameters, SSAOParameters)
 		SHADER_PARAMETER_STRUCT_INCLUDE(FGTAOShaderParameters, GTAOParameters)
@@ -1021,6 +1040,7 @@ FGTAOHorizonSearchOutputs AddGTAOHorizonSearchIntegratePass(
 	PassParameters->View = View.ViewUniformBuffer;
 	PassParameters->SceneTextures = CommonParameters.SceneTexturesUniformBuffer;
 	PassParameters->HZBParameters = GetHZBParameters(View, HZBInput, SceneViewport.Extent, EAOTechnique::GTAO);
+	PassParameters->Strata = Strata::BindStrataGlobalUniformParameters(View);
 	PassParameters->SSAOParameters = GetSSAOShaderParameters(View, SceneViewport, OutputViewport, CommonParameters.SceneTexturesViewport, EAOTechnique::GTAO);
 	PassParameters->GTAOParameters = GetGTAOShaderParameters(View, OutputViewport.Extent);
 
@@ -1029,6 +1049,7 @@ FGTAOHorizonSearchOutputs AddGTAOHorizonSearchIntegratePass(
 	FGTAOHorizonSearchAndIntegrateCS::FPermutationDomain PermutationVector;
 	PermutationVector.Set<FGTAOHorizonSearchAndIntegrateCS::FShaderQualityDim>(CommonParameters.ShaderQuality);
 	PermutationVector.Set<FGTAOHorizonSearchAndIntegrateCS::FUseNormalBufferDim>(bUseNormals);
+	PermutationVector.Set<FGTAOHorizonSearchAndIntegrateCS::FStrata>(Strata::IsStrataEnabled());
 
 	TShaderMapRef<FGTAOHorizonSearchAndIntegrateCS> ComputeShader(View.ShaderMap, PermutationVector);
 	FComputeShaderUtils::AddPass(
@@ -1052,6 +1073,10 @@ public:
 	DECLARE_GLOBAL_SHADER(FGTAOInnerIntegratePS);
 	SHADER_USE_PARAMETER_STRUCT(FGTAOInnerIntegratePS, FGlobalShader);
 
+	class FStrata : SHADER_PERMUTATION_BOOL("STRATA_ENABLED");
+
+	using FPermutationDomain = TShaderPermutationDomain<FStrata>;
+
 	static void ModifyCompilationEnvironment(const FGlobalShaderPermutationParameters& Parameters, FShaderCompilerEnvironment& OutEnvironment)
 	{
 		FGlobalShader::ModifyCompilationEnvironment(Parameters, OutEnvironment);
@@ -1066,6 +1091,7 @@ public:
 	BEGIN_SHADER_PARAMETER_STRUCT(FParameters, )
 		SHADER_PARAMETER_STRUCT_REF(FViewUniformShaderParameters, View)
 		SHADER_PARAMETER_RDG_UNIFORM_BUFFER(FSceneTextureUniformParameters, SceneTextures)
+		SHADER_PARAMETER_STRUCT_REF(FStrataGlobalUniformParameters, Strata)
 
 		SHADER_PARAMETER_STRUCT_INCLUDE(FSSAOShaderParameters, SSAOParameters)
 		SHADER_PARAMETER_STRUCT_INCLUDE(FGTAOShaderParameters, GTAOParameters)
@@ -1110,6 +1136,7 @@ FScreenPassTexture AddGTAOInnerIntegratePass(
 
 	PassParameters->View = View.ViewUniformBuffer;
 	PassParameters->SceneTextures = CommonParameters.SceneTexturesUniformBuffer;
+	PassParameters->Strata = Strata::BindStrataGlobalUniformParameters(View);
 	PassParameters->SSAOParameters = GetSSAOShaderParameters(View, InputViewport, OutputViewport, CommonParameters.SceneTexturesViewport, EAOTechnique::GTAO);
 	PassParameters->GTAOParameters = GetGTAOShaderParameters(View, OutputViewport.Extent);
 
@@ -1118,7 +1145,10 @@ FScreenPassTexture AddGTAOInnerIntegratePass(
 
 	PassParameters->RenderTargets[0] = Output.GetRenderTargetBinding();
 
-	TShaderMapRef<FGTAOInnerIntegratePS> PixelShader(View.ShaderMap);
+	FGTAOInnerIntegratePS::FPermutationDomain PermutationVector;
+	PermutationVector.Set<FGTAOInnerIntegratePS::FStrata>(Strata::IsStrataEnabled());
+	TShaderMapRef<FGTAOInnerIntegratePS> PixelShader(View.ShaderMap, PermutationVector);
+
 	AddDrawScreenPass(
 		GraphBuilder,
 		RDG_EVENT_NAME("GTAOInnerIntegratePS %dx%d Downscale=%d", OutputViewport.Rect.Width(), OutputViewport.Rect.Height(), CommonParameters.DownscaleFactor),
@@ -1140,8 +1170,9 @@ public:
 	SHADER_USE_PARAMETER_STRUCT(FGTAOHorizonSearchCS, FGlobalShader);
 
 	class FShaderQualityDim : SHADER_PERMUTATION_INT("SHADER_QUALITY", 5);
+	class FStrata : SHADER_PERMUTATION_BOOL("STRATA_ENABLED");
 
-	using FPermutationDomain = TShaderPermutationDomain<FShaderQualityDim>;
+	using FPermutationDomain = TShaderPermutationDomain<FShaderQualityDim, FStrata>;
 
 	static void ModifyCompilationEnvironment(const FGlobalShaderPermutationParameters& Parameters, FShaderCompilerEnvironment& OutEnvironment)
 	{
@@ -1160,6 +1191,7 @@ public:
 		SHADER_PARAMETER_STRUCT_REF(FViewUniformShaderParameters, View)
 		SHADER_PARAMETER_RDG_UNIFORM_BUFFER(FSceneTextureUniformParameters, SceneTextures)
 		SHADER_PARAMETER_STRUCT_INCLUDE(FHZBParameters, HZBParameters)
+		SHADER_PARAMETER_STRUCT_REF(FStrataGlobalUniformParameters, Strata)
 
 		SHADER_PARAMETER_STRUCT_INCLUDE(FSSAOShaderParameters, SSAOParameters)
 		SHADER_PARAMETER_STRUCT_INCLUDE(FGTAOShaderParameters, GTAOParameters)
@@ -1188,6 +1220,7 @@ FGTAOHorizonSearchOutputs AddGTAOHorizonSearchPass(
 	PassParameters->View = View.ViewUniformBuffer;
 	PassParameters->SceneTextures = CommonParameters.SceneTexturesUniformBuffer;
 	PassParameters->HZBParameters = GetHZBParameters(View, HZBInput, SceneViewport.Extent, EAOTechnique::GTAO);
+	PassParameters->Strata = Strata::BindStrataGlobalUniformParameters(View);
 	PassParameters->SSAOParameters = GetSSAOShaderParameters(View, SceneViewport, OutputViewport, CommonParameters.SceneTexturesViewport, EAOTechnique::GTAO);
 	PassParameters->GTAOParameters = GetGTAOShaderParameters(View, OutputViewport.Extent);
 
@@ -1195,6 +1228,7 @@ FGTAOHorizonSearchOutputs AddGTAOHorizonSearchPass(
 
 	FGTAOHorizonSearchCS::FPermutationDomain PermutationVector;
 	PermutationVector.Set<FGTAOHorizonSearchCS::FShaderQualityDim>(CommonParameters.ShaderQuality);
+	PermutationVector.Set<FGTAOHorizonSearchCS::FStrata>(Strata::IsStrataEnabled());
 
 	TShaderMapRef<FGTAOHorizonSearchCS> ComputeShader(View.ShaderMap, PermutationVector);
 	FComputeShaderUtils::AddPass(
diff --git a/Engine/Source/Runtime/Renderer/Private/DistanceFieldAmbientOcclusion.cpp b/Engine/Source/Runtime/Renderer/Private/DistanceFieldAmbientOcclusion.cpp
--- a/Engine/Source/Runtime/Renderer/Private/DistanceFieldAmbientOcclusion.cpp
+++ b/Engine/Source/Runtime/Renderer/Private/DistanceFieldAmbientOcclusion.cpp
@@ -19,6 +19,7 @@
 #include "VisualizeTexture.h"
 #include "RayTracing/RaytracingOptions.h"
 #include "Lumen/Lumen.h"
+#include "Strata/Strata.h"
 
 IMPLEMENT_TYPE_LAYOUT(FAOParameters);
 IMPLEMENT_TYPE_LAYOUT(FDFAOUpsampleParameters);
@@ -274,9 +275,13 @@ class FComputeDistanceFieldNormalPS : public FGlobalShader
 public:
 	DECLARE_GLOBAL_SHADER(FComputeDistanceFieldNormalPS);
 
+	class FStrata : SHADER_PERMUTATION_BOOL("STRATA_ENABLED");
+	using FPermutationDomain = TShaderPermutationDomain<FStrata>;
+
 	BEGIN_SHADER_PARAMETER_STRUCT(FParameters, )
 		SHADER_PARAMETER_STRUCT_REF(FViewUniformShaderParameters, View)
 		SHADER_PARAMETER_RDG_UNIFORM_BUFFER(FSceneTextureUniformParameters, SceneTextures)
+		SHADER_PARAMETER_STRUCT_REF(FStrataGlobalUniformParameters, Strata)
 		RENDER_TARGET_BINDING_SLOTS()
 	END_SHADER_PARAMETER_STRUCT()
 
@@ -318,9 +323,13 @@ class FComputeDistanceFieldNormalCS : public FGlobalShader
 public:
 	DECLARE_GLOBAL_SHADER(FComputeDistanceFieldNormalCS);
 
+	class FStrata : SHADER_PERMUTATION_BOOL("STRATA_ENABLED");
+	using FPermutationDomain = TShaderPermutationDomain<FStrata>;
+
 	BEGIN_SHADER_PARAMETER_STRUCT(FParameters, )
 		SHADER_PARAMETER_STRUCT_REF(FViewUniformShaderParameters, View)
 		SHADER_PARAMETER_RDG_UNIFORM_BUFFER(FSceneTextureUniformParameters, SceneTextures)
+		SHADER_PARAMETER_STRUCT_REF(FStrataGlobalUniformParameters, Strata)
 		SHADER_PARAMETER_RDG_TEXTURE_UAV(RWTexture2D<float4>, RWDistanceFieldNormal)
 	END_SHADER_PARAMETER_STRUCT()
 
@@ -376,9 +385,12 @@ void ComputeDistanceFieldNormal(
 			auto* PassParameters = GraphBuilder.AllocParameters<FComputeDistanceFieldNormalCS::FParameters>();
 			PassParameters->View = View.ViewUniformBuffer;
 			PassParameters->SceneTextures = SceneTexturesUniformBuffer;
+			PassParameters->Strata = Strata::BindStrataGlobalUniformParameters(View);
 			PassParameters->RWDistanceFieldNormal = GraphBuilder.CreateUAV(DistanceFieldNormal);
 
-			TShaderMapRef<FComputeDistanceFieldNormalCS> ComputeShader(View.ShaderMap);
+			FComputeDistanceFieldNormalCS::FPermutationDomain PermutationVector;
+			PermutationVector.Set<FComputeDistanceFieldNormalCS::FStrata>(Strata::IsStrataEnabled());
+			TShaderMapRef<FComputeDistanceFieldNormalCS> ComputeShader(View.ShaderMap, PermutationVector);
 
 			GraphBuilder.AddPass(
 				RDG_EVENT_NAME("ComputeNormalCS"),
@@ -405,6 +417,7 @@ void ComputeDistanceFieldNormal(
 			auto* PassParameters = GraphBuilder.AllocParameters<FComputeDistanceFieldNormalPS::FParameters>();
 			PassParameters->View = View.ViewUniformBuffer;
 			PassParameters->SceneTextures = SceneTexturesUniformBuffer;
+			PassParameters->Strata = Strata::BindStrataGlobalUniformParameters(View);
 			PassParameters->RenderTargets[0] = FRenderTargetBinding(DistanceFieldNormal, ViewIndex == 0 ? ERenderTargetLoadAction::EClear : ERenderTargetLoadAction::ELoad);
 
 			GraphBuilder.AddPass(
@@ -416,7 +429,10 @@ void ComputeDistanceFieldNormal(
 				RHICmdList.SetViewport(0, 0, 0.0f, View.ViewRect.Width() / GAODownsampleFactor, View.ViewRect.Height() / GAODownsampleFactor, 1.0f);
 
 				TShaderMapRef<FPostProcessVS> VertexShader(View.ShaderMap);
-				TShaderMapRef<FComputeDistanceFieldNormalPS> PixelShader(View.ShaderMap);
+
+				FComputeDistanceFieldNormalPS::FPermutationDomain PermutationVector;
+				PermutationVector.Set<FComputeDistanceFieldNormalPS::FStrata>(Strata::IsStrataEnabled());
+				TShaderMapRef<FComputeDistanceFieldNormalPS> PixelShader(View.ShaderMap, PermutationVector);
 
 				FGraphicsPipelineStateInitializer GraphicsPSOInit;
 				RHICmdList.ApplyCachedRenderTargets(GraphicsPSOInit);
diff --git a/Engine/Source/Runtime/Renderer/Private/ScreenSpaceRayTracing.cpp b/Engine/Source/Runtime/Renderer/Private/ScreenSpaceRayTracing.cpp
--- a/Engine/Source/Runtime/Renderer/Private/ScreenSpaceRayTracing.cpp
+++ b/Engine/Source/Runtime/Renderer/Private/ScreenSpaceRayTracing.cpp
@@ -6,6 +6,7 @@
 #include "ScreenPass.h"
 #include "ScenePrivate.h"
 #include "SceneTextureParameters.h"
+#include "Strata/Strata.h"
 
 
 static TAutoConsoleVariable<int32> CVarSSRQuality(
@@ -360,6 +361,7 @@ enum class ELightingTerm
 class FSSRQualityDim : SHADER_PERMUTATION_ENUM_CLASS("SSR_QUALITY", ESSRQuality);
 class FSSROutputForDenoiser : SHADER_PERMUTATION_BOOL("SSR_OUTPUT_FOR_DENOISER");
 class FLightingTermDim : SHADER_PERMUTATION_ENUM_CLASS("DIM_LIGHTING_TERM", ELightingTerm);
+class FStrata : SHADER_PERMUTATION_BOOL("STRATA_ENABLED");
 
 class FSSRTPrevFrameReductionCS : public FGlobalShader
 {
@@ -443,7 +445,7 @@ class FScreenSpaceReflectionsStencilPS : public FGlobalShader
 	DECLARE_GLOBAL_SHADER(FScreenSpaceReflectionsStencilPS);
 	SHADER_USE_PARAMETER_STRUCT(FScreenSpaceReflectionsStencilPS, FGlobalShader);
 
-	using FPermutationDomain = TShaderPermutationDomain<FSSROutputForDenoiser>;
+	using FPermutationDomain = TShaderPermutationDomain<FSSROutputForDenoiser, FStrata>;
 
 	static bool ShouldCompilePermutation(const FGlobalShaderPermutationParameters& Parameters)
 	{
@@ -458,6 +460,7 @@ class FScreenSpaceReflectionsStencilPS : public FGlobalShader
 	
 	BEGIN_SHADER_PARAMETER_STRUCT(FParameters, )
 		SHADER_PARAMETER_STRUCT_INCLUDE(FSSRCommonParameters, CommonParameters)
+		SHADER_PARAMETER_STRUCT_REF(FStrataGlobalUniformParameters, Strata)
 		RENDER_TARGET_BINDING_SLOTS()
 	END_SHADER_PARAMETER_STRUCT()
 };
@@ -467,7 +470,7 @@ class FScreenSpaceReflectionsPS : public FGlobalShader
 	DECLARE_GLOBAL_SHADER(FScreenSpaceReflectionsPS);
 	SHADER_USE_PARAMETER_STRUCT(FScreenSpaceReflectionsPS, FGlobalShader);
 
-	using FPermutationDomain = TShaderPermutationDomain<FSSRQualityDim, FSSROutputForDenoiser>;
+	using FPermutationDomain = TShaderPermutationDomain<FSSRQualityDim, FSSROutputForDenoiser, FStrata>;
 
 	static bool ShouldCompilePermutation(const FGlobalShaderPermutationParameters& Parameters)
 	{
@@ -480,6 +483,7 @@ class FScreenSpaceReflectionsPS : public FGlobalShader
 		SHADER_PARAMETER_STRUCT_INCLUDE(FSSRPassCommonParameters, SSRPassCommonParameter)
 		SHADER_PARAMETER_RDG_BUFFER(Buffer<uint>, IndirectDrawParameter)			// FScreenSpaceReflectionsTileVS
 		SHADER_PARAMETER_RDG_BUFFER_SRV(StructuredBuffer<uint>, TileListData)		// FScreenSpaceReflectionsTileVS
+		SHADER_PARAMETER_STRUCT_REF(FStrataGlobalUniformParameters, Strata)
 		RENDER_TARGET_BINDING_SLOTS()
 	END_SHADER_PARAMETER_STRUCT()
 };
@@ -1086,9 +1090,11 @@ void RenderScreenSpaceReflections(
 
 		FScreenSpaceReflectionsStencilPS::FPermutationDomain PermutationVector;
 		PermutationVector.Set<FSSROutputForDenoiser>(bDenoiser);
+		PermutationVector.Set<FStrata>(Strata::IsStrataEnabled());
 
 		FScreenSpaceReflectionsStencilPS::FParameters* PassParameters = GraphBuilder.AllocParameters<FScreenSpaceReflectionsStencilPS::FParameters>();
 		PassParameters->CommonParameters = CommonParameters;
+		PassParameters->Strata = Strata::BindStrataGlobalUniformParameters(View);
 		PassParameters->RenderTargets = RenderTargets;
 		
 		TShaderMapRef<FScreenSpaceReflectionsStencilPS> PixelShader(View.ShaderMap, PermutationVector);
@@ -1166,10 +1172,12 @@ void RenderScreenSpaceReflections(
 	FScreenSpaceReflectionsPS::FPermutationDomain PermutationVector;
 	PermutationVector.Set<FSSRQualityDim>(SSRQuality);
 	PermutationVector.Set<FSSROutputForDenoiser>(bDenoiser);
+	PermutationVector.Set<FStrata>(Strata::IsStrataEnabled());
 		
 	FScreenSpaceReflectionsPS::FParameters* PassParameters = GraphBuilder.AllocParameters<FScreenSpaceReflectionsPS::FParameters>();
 	PassParameters->CommonParameters = CommonParameters;
 	SetSSRParameters(&PassParameters->SSRPassCommonParameter);
+	PassParameters->Strata = Strata::BindStrataGlobalUniformParameters(View);
 	PassParameters->RenderTargets = RenderTargets;
 
 	TShaderMapRef<FScreenSpaceReflectionsPS> PixelShader(View.ShaderMap, PermutationVector);
