Removed Strata permutation from all shaders since it is part of the ShaderMapAppendKeyString and that will recompile all shaders (global and materials)
We still keep the FastPath permutations of course.

diff --git a/Engine/Shaders/Public/ShaderVersion.ush b/Engine/Shaders/Public/ShaderVersion.ush
--- a/Engine/Shaders/Public/ShaderVersion.ush
+++ b/Engine/Shaders/Public/ShaderVersion.ush
@@ -3,5 +3,5 @@
 // in Platform.ush (which should be included in any shader) it allows to invalidate the shader DDC.
 // 
 // If you are merging streams and there is a conflict with this GUID you should make a new GUID rather than taking one or the other.
-// GUID = D8A972B4-2356-49BC-9CA8-ABF9D6792E1B
+// GUID = 62353968-C93C-410B-961B-766739FA9F5B
 
diff --git a/Engine/Source/Runtime/Renderer/Private/CompositionLighting/PostProcessAmbientOcclusion.cpp b/Engine/Source/Runtime/Renderer/Private/CompositionLighting/PostProcessAmbientOcclusion.cpp
--- a/Engine/Source/Runtime/Renderer/Private/CompositionLighting/PostProcessAmbientOcclusion.cpp
+++ b/Engine/Source/Runtime/Renderer/Private/CompositionLighting/PostProcessAmbientOcclusion.cpp
@@ -515,14 +515,16 @@ public:
 	DECLARE_GLOBAL_SHADER(FAmbientOcclusionSetupPS);
 	SHADER_USE_PARAMETER_STRUCT(FAmbientOcclusionSetupPS, FGlobalShader);
 
-	class FStrata : SHADER_PERMUTATION_BOOL("STRATA_ENABLED");
-	using FPermutationDomain = TShaderPermutationDomain<FStrata>;
-
 	static bool ShouldCompilePermutation(const FGlobalShaderPermutationParameters& Parameters)
 	{
 		return IsFeatureLevelSupported(Parameters.Platform, ERHIFeatureLevel::SM5);
 	}
 
+	static void ModifyCompilationEnvironment(const FGlobalShaderPermutationParameters& Parameters, FShaderCompilerEnvironment& OutEnvironment)
+	{
+		OutEnvironment.SetDefine(TEXT("STRATA_ENABLED"), Strata::IsStrataEnabled() ? 1u : 0u);
+	}
+
 	BEGIN_SHADER_PARAMETER_STRUCT(FParameters, )
 		SHADER_PARAMETER_STRUCT_REF(FViewUniformShaderParameters, View)
 		SHADER_PARAMETER_RDG_UNIFORM_BUFFER(FSceneTextureUniformParameters, SceneTextures)
@@ -579,9 +581,7 @@ FScreenPassTexture AddAmbientOcclusionSetupPass(
 	PassParameters->Strata = Strata::BindStrataGlobalUniformParameters(View);
 	PassParameters->RenderTargets[0] = Output.GetRenderTargetBinding();
 
-	FAmbientOcclusionSetupPS::FPermutationDomain PermutationVector;
-	PermutationVector.Set<FAmbientOcclusionSetupPS::FStrata>(Strata::IsStrataEnabled());
-	TShaderMapRef<FAmbientOcclusionSetupPS> PixelShader(View.ShaderMap, PermutationVector);
+	TShaderMapRef<FAmbientOcclusionSetupPS> PixelShader(View.ShaderMap);
 
 	AddDrawScreenPass(
 		GraphBuilder,
@@ -708,14 +708,14 @@ public:
 	class FUseUpsampleDim       : SHADER_PERMUTATION_BOOL("USE_UPSAMPLE");
 	class FUseAoSetupAsInputDim : SHADER_PERMUTATION_BOOL("USE_AO_SETUP_AS_INPUT");
 	class FShaderQualityDim     : SHADER_PERMUTATION_INT("SHADER_QUALITY", 5);
-	class FStrata				: SHADER_PERMUTATION_BOOL("STRATA_ENABLED");
 
-	using FPermutationDomain = TShaderPermutationDomain<FUseUpsampleDim, FUseAoSetupAsInputDim, FShaderQualityDim, FStrata>;
+	using FPermutationDomain = TShaderPermutationDomain<FUseUpsampleDim, FUseAoSetupAsInputDim, FShaderQualityDim>;
 
 	static void ModifyCompilationEnvironment(const FGlobalShaderPermutationParameters& Parameters, FShaderCompilerEnvironment& OutEnvironment)
 	{
 		FGlobalShader::ModifyCompilationEnvironment(Parameters, OutEnvironment);
 		OutEnvironment.SetDefine(TEXT("COMPUTE_SHADER"), 0);
+		OutEnvironment.SetDefine(TEXT("STRATA_ENABLED"), Strata::IsStrataEnabled() ? 1u : 0u);
 	}
 
 	static bool ShouldCompilePermutation(const FGlobalShaderPermutationParameters& Parameters)
@@ -740,9 +740,8 @@ public:
 	class FUseUpsampleDim       : SHADER_PERMUTATION_BOOL("USE_UPSAMPLE");
 	class FUseAoSetupAsInputDim : SHADER_PERMUTATION_BOOL("USE_AO_SETUP_AS_INPUT");
 	class FShaderQualityDim     : SHADER_PERMUTATION_INT("SHADER_QUALITY", 5);
-	class FStrata				: SHADER_PERMUTATION_BOOL("STRATA_ENABLED");
 
-	using FPermutationDomain = TShaderPermutationDomain<FUseUpsampleDim, FUseAoSetupAsInputDim, FShaderQualityDim, FStrata>;
+	using FPermutationDomain = TShaderPermutationDomain<FUseUpsampleDim, FUseAoSetupAsInputDim, FShaderQualityDim>;
 
 	static void ModifyCompilationEnvironment(const FGlobalShaderPermutationParameters& Parameters, FShaderCompilerEnvironment& OutEnvironment)
 	{
@@ -750,6 +749,7 @@ public:
 		OutEnvironment.SetDefine(TEXT("COMPUTE_SHADER"), 1);
 		OutEnvironment.SetDefine(TEXT("THREADGROUP_SIZEX"), 16);
 		OutEnvironment.SetDefine(TEXT("THREADGROUP_SIZEY"), 16);
+		OutEnvironment.SetDefine(TEXT("STRATA_ENABLED"), Strata::IsStrataEnabled() ? 1u : 0u);
 	}
 
 	static bool ShouldCompilePermutation(const FGlobalShaderPermutationParameters& Parameters)
@@ -864,7 +864,6 @@ FScreenPassTexture AddAmbientOcclusionPass(
 		PermutationVector.Set<FAmbientOcclusionCS::FUseUpsampleDim>(bDoUpsample);
 		PermutationVector.Set<FAmbientOcclusionCS::FUseAoSetupAsInputDim>(bAOSetupAsInput);
 		PermutationVector.Set<FAmbientOcclusionCS::FShaderQualityDim>(CommonParameters.ShaderQuality);
-		PermutationVector.Set<FAmbientOcclusionCS::FStrata>(Strata::IsStrataEnabled());
 
 		TShaderMapRef<FAmbientOcclusionCS> ComputeShader(View.ShaderMap, PermutationVector);
 		FComputeShaderUtils::AddPass(
@@ -887,7 +886,6 @@ FScreenPassTexture AddAmbientOcclusionPass(
 		PermutationVector.Set<FAmbientOcclusionPS::FUseUpsampleDim>(bDoUpsample);
 		PermutationVector.Set<FAmbientOcclusionPS::FUseAoSetupAsInputDim>(bAOSetupAsInput);
 		PermutationVector.Set<FAmbientOcclusionPS::FShaderQualityDim>(CommonParameters.ShaderQuality);
-		PermutationVector.Set<FAmbientOcclusionPS::FStrata>(Strata::IsStrataEnabled());
 
 		TShaderMapRef<FAmbientOcclusionPS> PixelShader(View.ShaderMap, PermutationVector);
 		AddDrawScreenPass(
@@ -974,9 +972,8 @@ public:
 
 	class FShaderQualityDim : SHADER_PERMUTATION_INT("SHADER_QUALITY", 5);
 	class FUseNormalBufferDim : SHADER_PERMUTATION_BOOL("USE_NORMALBUFFER");
-	class FStrata : SHADER_PERMUTATION_BOOL("STRATA_ENABLED");
 
-	using FPermutationDomain = TShaderPermutationDomain<FShaderQualityDim, FUseNormalBufferDim, FStrata>;
+	using FPermutationDomain = TShaderPermutationDomain<FShaderQualityDim, FUseNormalBufferDim>;
 
 	static void ModifyCompilationEnvironment(const FGlobalShaderPermutationParameters& Parameters, FShaderCompilerEnvironment& OutEnvironment)
 	{
@@ -984,6 +981,7 @@ public:
 		OutEnvironment.SetDefine(TEXT("COMPUTE_SHADER"), 1);
 		OutEnvironment.SetDefine(TEXT("THREADGROUP_SIZEX"), 8);
 		OutEnvironment.SetDefine(TEXT("THREADGROUP_SIZEY"), 8);
+		OutEnvironment.SetDefine(TEXT("STRATA_ENABLED"), Strata::IsStrataEnabled() ? 1u : 0u);
 	}
 
 	static bool ShouldCompilePermutation(const FGlobalShaderPermutationParameters& Parameters)
@@ -1049,7 +1047,6 @@ FGTAOHorizonSearchOutputs AddGTAOHorizonSearchIntegratePass(
 	FGTAOHorizonSearchAndIntegrateCS::FPermutationDomain PermutationVector;
 	PermutationVector.Set<FGTAOHorizonSearchAndIntegrateCS::FShaderQualityDim>(CommonParameters.ShaderQuality);
 	PermutationVector.Set<FGTAOHorizonSearchAndIntegrateCS::FUseNormalBufferDim>(bUseNormals);
-	PermutationVector.Set<FGTAOHorizonSearchAndIntegrateCS::FStrata>(Strata::IsStrataEnabled());
 
 	TShaderMapRef<FGTAOHorizonSearchAndIntegrateCS> ComputeShader(View.ShaderMap, PermutationVector);
 	FComputeShaderUtils::AddPass(
@@ -1073,14 +1070,11 @@ public:
 	DECLARE_GLOBAL_SHADER(FGTAOInnerIntegratePS);
 	SHADER_USE_PARAMETER_STRUCT(FGTAOInnerIntegratePS, FGlobalShader);
 
-	class FStrata : SHADER_PERMUTATION_BOOL("STRATA_ENABLED");
-
-	using FPermutationDomain = TShaderPermutationDomain<FStrata>;
-
 	static void ModifyCompilationEnvironment(const FGlobalShaderPermutationParameters& Parameters, FShaderCompilerEnvironment& OutEnvironment)
 	{
 		FGlobalShader::ModifyCompilationEnvironment(Parameters, OutEnvironment);
 		OutEnvironment.SetDefine(TEXT("COMPUTE_SHADER"), 0);
+		OutEnvironment.SetDefine(TEXT("STRATA_ENABLED"), Strata::IsStrataEnabled() ? 1u : 0u);
 	}
 
 	static bool ShouldCompilePermutation(const FGlobalShaderPermutationParameters& Parameters)
@@ -1145,9 +1139,7 @@ FScreenPassTexture AddGTAOInnerIntegratePass(
 
 	PassParameters->RenderTargets[0] = Output.GetRenderTargetBinding();
 
-	FGTAOInnerIntegratePS::FPermutationDomain PermutationVector;
-	PermutationVector.Set<FGTAOInnerIntegratePS::FStrata>(Strata::IsStrataEnabled());
-	TShaderMapRef<FGTAOInnerIntegratePS> PixelShader(View.ShaderMap, PermutationVector);
+	TShaderMapRef<FGTAOInnerIntegratePS> PixelShader(View.ShaderMap);
 
 	AddDrawScreenPass(
 		GraphBuilder,
@@ -1170,9 +1162,8 @@ public:
 	SHADER_USE_PARAMETER_STRUCT(FGTAOHorizonSearchCS, FGlobalShader);
 
 	class FShaderQualityDim : SHADER_PERMUTATION_INT("SHADER_QUALITY", 5);
-	class FStrata : SHADER_PERMUTATION_BOOL("STRATA_ENABLED");
 
-	using FPermutationDomain = TShaderPermutationDomain<FShaderQualityDim, FStrata>;
+	using FPermutationDomain = TShaderPermutationDomain<FShaderQualityDim>;
 
 	static void ModifyCompilationEnvironment(const FGlobalShaderPermutationParameters& Parameters, FShaderCompilerEnvironment& OutEnvironment)
 	{
@@ -1180,6 +1171,7 @@ public:
 		OutEnvironment.SetDefine(TEXT("COMPUTE_SHADER"), 1);
 		OutEnvironment.SetDefine(TEXT("THREADGROUP_SIZEX"), 8);
 		OutEnvironment.SetDefine(TEXT("THREADGROUP_SIZEY"), 8);
+		OutEnvironment.SetDefine(TEXT("STRATA_ENABLED"), Strata::IsStrataEnabled() ? 1u : 0u);
 	}
 
 	static bool ShouldCompilePermutation(const FGlobalShaderPermutationParameters& Parameters)
@@ -1228,7 +1220,6 @@ FGTAOHorizonSearchOutputs AddGTAOHorizonSearchPass(
 
 	FGTAOHorizonSearchCS::FPermutationDomain PermutationVector;
 	PermutationVector.Set<FGTAOHorizonSearchCS::FShaderQualityDim>(CommonParameters.ShaderQuality);
-	PermutationVector.Set<FGTAOHorizonSearchCS::FStrata>(Strata::IsStrataEnabled());
 
 	TShaderMapRef<FGTAOHorizonSearchCS> ComputeShader(View.ShaderMap, PermutationVector);
 	FComputeShaderUtils::AddPass(
diff --git a/Engine/Source/Runtime/Renderer/Private/DistanceFieldAmbientOcclusion.cpp b/Engine/Source/Runtime/Renderer/Private/DistanceFieldAmbientOcclusion.cpp
--- a/Engine/Source/Runtime/Renderer/Private/DistanceFieldAmbientOcclusion.cpp
+++ b/Engine/Source/Runtime/Renderer/Private/DistanceFieldAmbientOcclusion.cpp
@@ -275,9 +275,6 @@ class FComputeDistanceFieldNormalPS : public FGlobalShader
 public:
 	DECLARE_GLOBAL_SHADER(FComputeDistanceFieldNormalPS);
 
-	class FStrata : SHADER_PERMUTATION_BOOL("STRATA_ENABLED");
-	using FPermutationDomain = TShaderPermutationDomain<FStrata>;
-
 	BEGIN_SHADER_PARAMETER_STRUCT(FParameters, )
 		SHADER_PARAMETER_STRUCT_REF(FViewUniformShaderParameters, View)
 		SHADER_PARAMETER_RDG_UNIFORM_BUFFER(FSceneTextureUniformParameters, SceneTextures)
@@ -295,6 +292,7 @@ public:
 		OutEnvironment.SetDefine(TEXT("DOWNSAMPLE_FACTOR"), GAODownsampleFactor);
 		OutEnvironment.SetDefine(TEXT("THREADGROUP_SIZEX"), GDistanceFieldAOTileSizeX);
 		OutEnvironment.SetDefine(TEXT("THREADGROUP_SIZEY"), GDistanceFieldAOTileSizeY);
+		OutEnvironment.SetDefine(TEXT("STRATA_ENABLED"), Strata::IsStrataEnabled() ? 1u : 0u);
 	}
 
 	FComputeDistanceFieldNormalPS() = default;
@@ -323,9 +321,6 @@ class FComputeDistanceFieldNormalCS : public FGlobalShader
 public:
 	DECLARE_GLOBAL_SHADER(FComputeDistanceFieldNormalCS);
 
-	class FStrata : SHADER_PERMUTATION_BOOL("STRATA_ENABLED");
-	using FPermutationDomain = TShaderPermutationDomain<FStrata>;
-
 	BEGIN_SHADER_PARAMETER_STRUCT(FParameters, )
 		SHADER_PARAMETER_STRUCT_REF(FViewUniformShaderParameters, View)
 		SHADER_PARAMETER_RDG_UNIFORM_BUFFER(FSceneTextureUniformParameters, SceneTextures)
@@ -343,6 +338,7 @@ public:
 		OutEnvironment.SetDefine(TEXT("DOWNSAMPLE_FACTOR"), GAODownsampleFactor);
 		OutEnvironment.SetDefine(TEXT("THREADGROUP_SIZEX"), GDistanceFieldAOTileSizeX);
 		OutEnvironment.SetDefine(TEXT("THREADGROUP_SIZEY"), GDistanceFieldAOTileSizeY);
+		OutEnvironment.SetDefine(TEXT("STRATA_ENABLED"), Strata::IsStrataEnabled() ? 1u : 0u);
 	}
 
 	FComputeDistanceFieldNormalCS() = default;
@@ -388,9 +384,7 @@ void ComputeDistanceFieldNormal(
 			PassParameters->Strata = Strata::BindStrataGlobalUniformParameters(View);
 			PassParameters->RWDistanceFieldNormal = GraphBuilder.CreateUAV(DistanceFieldNormal);
 
-			FComputeDistanceFieldNormalCS::FPermutationDomain PermutationVector;
-			PermutationVector.Set<FComputeDistanceFieldNormalCS::FStrata>(Strata::IsStrataEnabled());
-			TShaderMapRef<FComputeDistanceFieldNormalCS> ComputeShader(View.ShaderMap, PermutationVector);
+			TShaderMapRef<FComputeDistanceFieldNormalCS> ComputeShader(View.ShaderMap);
 
 			GraphBuilder.AddPass(
 				RDG_EVENT_NAME("ComputeNormalCS"),
@@ -430,9 +424,7 @@ void ComputeDistanceFieldNormal(
 
 				TShaderMapRef<FPostProcessVS> VertexShader(View.ShaderMap);
 
-				FComputeDistanceFieldNormalPS::FPermutationDomain PermutationVector;
-				PermutationVector.Set<FComputeDistanceFieldNormalPS::FStrata>(Strata::IsStrataEnabled());
-				TShaderMapRef<FComputeDistanceFieldNormalPS> PixelShader(View.ShaderMap, PermutationVector);
+				TShaderMapRef<FComputeDistanceFieldNormalPS> PixelShader(View.ShaderMap);
 
 				FGraphicsPipelineStateInitializer GraphicsPSOInit;
 				RHICmdList.ApplyCachedRenderTargets(GraphicsPSOInit);
diff --git a/Engine/Source/Runtime/Renderer/Private/IndirectLightRendering.cpp b/Engine/Source/Runtime/Renderer/Private/IndirectLightRendering.cpp
--- a/Engine/Source/Runtime/Renderer/Private/IndirectLightRendering.cpp
+++ b/Engine/Source/Runtime/Renderer/Private/IndirectLightRendering.cpp
@@ -189,7 +189,6 @@ class FReflectionEnvironmentSkyLightingPS : public FGlobalShader
 	class FDynamicSkyLight : SHADER_PERMUTATION_BOOL("ENABLE_DYNAMIC_SKY_LIGHT");
 	class FSkyShadowing : SHADER_PERMUTATION_BOOL("APPLY_SKY_SHADOWING");
 	class FRayTracedReflections : SHADER_PERMUTATION_BOOL("RAY_TRACED_REFLECTIONS");
-	class FStrata : SHADER_PERMUTATION_BOOL("STRATA_ENABLED");
 	class FStrataFastPath : SHADER_PERMUTATION_BOOL("STRATA_FASTPATH");
 
 	using FPermutationDomain = TShaderPermutationDomain<
@@ -200,7 +199,6 @@ class FReflectionEnvironmentSkyLightingPS : public FGlobalShader
 		FDynamicSkyLight,
 		FSkyShadowing,
 		FRayTracedReflections,
-		FStrata,
 		FStrataFastPath>;
 
 	static FPermutationDomain RemapPermutation(FPermutationDomain PermutationVector)
@@ -217,11 +215,6 @@ class FReflectionEnvironmentSkyLightingPS : public FGlobalShader
 			PermutationVector.Set<FSkyShadowing>(false);
 		}
 
-		if (PermutationVector.Get<FStrata>() && !Strata::IsStrataEnabled())
-		{
-			PermutationVector.Set<FStrata>(false);
-		}
-
 		if (PermutationVector.Get<FStrataFastPath>() && (!Strata::IsStrataEnabled() || !Strata::IsClassificationEnabled()))
 		{
 			PermutationVector.Set<FStrataFastPath>(false);
@@ -241,7 +234,6 @@ class FReflectionEnvironmentSkyLightingPS : public FGlobalShader
 		PermutationVector.Set<FDynamicSkyLight>(bEnableDynamicSkyLight);
 		PermutationVector.Set<FSkyShadowing>(bApplySkyShadowing);
 		PermutationVector.Set<FRayTracedReflections>(bRayTracedReflections);
-		PermutationVector.Set<FStrata>(Strata::IsStrataEnabled());
 		PermutationVector.Set<FStrataFastPath>(Strata::IsStrataEnabled() && Strata::IsClassificationEnabled() && bStrataFastPath);
 
 		return RemapPermutation(PermutationVector);
@@ -262,6 +254,7 @@ class FReflectionEnvironmentSkyLightingPS : public FGlobalShader
 	{
 		FGlobalShader::ModifyCompilationEnvironment(Parameters, OutEnvironment);
 		OutEnvironment.SetDefine(TEXT("MAX_CAPTURES"), GMaxNumReflectionCaptures);
+		OutEnvironment.SetDefine(TEXT("STRATA_ENABLED"), Strata::IsStrataEnabled() ? 1u : 0u);
 		OutEnvironment.CompilerFlags.Add(CFLAG_StandardOptimization);
 		FForwardLightingParameters::ModifyCompilationEnvironment(Parameters.Platform, OutEnvironment);
 	}
diff --git a/Engine/Source/Runtime/Renderer/Private/LightRendering.cpp b/Engine/Source/Runtime/Renderer/Private/LightRendering.cpp
--- a/Engine/Source/Runtime/Renderer/Private/LightRendering.cpp
+++ b/Engine/Source/Runtime/Renderer/Private/LightRendering.cpp
@@ -441,7 +441,6 @@ class FDeferredLightPS : public FGlobalShader
 	class FAtmosphereTransmittance : SHADER_PERMUTATION_BOOL("USE_ATMOSPHERE_TRANSMITTANCE");
 	class FCloudTransmittance 	: SHADER_PERMUTATION_BOOL("USE_CLOUD_TRANSMITTANCE");
 	class FAnistropicMaterials 	: SHADER_PERMUTATION_BOOL("SUPPORTS_ANISOTROPIC_MATERIALS");
-	class FStrata				: SHADER_PERMUTATION_BOOL("STRATA_ENABLED");
 	class FStrataFastPath		: SHADER_PERMUTATION_BOOL("STRATA_FASTPATH");
 
 	using FPermutationDomain = TShaderPermutationDomain<
@@ -456,7 +455,6 @@ class FDeferredLightPS : public FGlobalShader
 		FAtmosphereTransmittance,
 		FCloudTransmittance,
 		FAnistropicMaterials,
-		FStrata,
 		FStrataFastPath>;
 
 	static bool ShouldCompilePermutation(const FGlobalShaderPermutationParameters& Parameters)
@@ -519,18 +517,19 @@ class FDeferredLightPS : public FGlobalShader
 			}
 		}
 
-		if (PermutationVector.Get<FStrata>() && !Strata::IsStrataEnabled())
-		{
-			return false;
-		}
-
-		if (PermutationVector.Get<FStrataFastPath>() && (!Strata::IsStrataEnabled() || !PermutationVector.Get<FStrata>()))
+		if (PermutationVector.Get<FStrataFastPath>() && !Strata::IsStrataEnabled())
 		{
 			return false;
 		}
 		return IsFeatureLevelSupported(Parameters.Platform, ERHIFeatureLevel::SM5);
 	}
 
+	static void ModifyCompilationEnvironment(const FGlobalShaderPermutationParameters& Parameters, FShaderCompilerEnvironment& OutEnvironment)
+	{
+		FGlobalShader::ModifyCompilationEnvironment(Parameters, OutEnvironment);
+		OutEnvironment.SetDefine(TEXT("STRATA_ENABLED"), Strata::IsStrataEnabled() ? 1u : 0u);
+	}
+
 	FDeferredLightPS(const ShaderMetaType::CompiledShaderInitializerType& Initializer)
 	:	FGlobalShader(Initializer)
 	{
@@ -2145,7 +2144,6 @@ static void SetShaderTemplLightingSimple(
 	PermutationVector.Set< FDeferredLightPS::FHairLighting>( 0 );
 	PermutationVector.Set< FDeferredLightPS::FAtmosphereTransmittance >( false );
 	PermutationVector.Set< FDeferredLightPS::FCloudTransmittance >(false);
-	PermutationVector.Set< FDeferredLightPS::FStrata >(Strata::IsStrataEnabled());
 	PermutationVector.Set< FDeferredLightPS::FStrataFastPath>(false);
 
 	TShaderMapRef< FDeferredLightPS > PixelShader( View.ShaderMap, PermutationVector );
@@ -2316,7 +2314,6 @@ void FDeferredShadingSceneRenderer::RenderLight(
 			// Only directional lights are rendered in this path, so we only need to check if it is use to light the atmosphere
 			PermutationVector.Set< FDeferredLightPS::FAtmosphereTransmittance >(bAtmospherePerPixelTransmittance);
 			PermutationVector.Set< FDeferredLightPS::FCloudTransmittance >(bLight0CloudPerPixelTransmittance || bLight1CloudPerPixelTransmittance);
-			PermutationVector.Set< FDeferredLightPS::FStrata >(Strata::IsStrataEnabled());
 			PermutationVector.Set< FDeferredLightPS::FStrataFastPath >(Strata::IsStrataEnabled() && Strata::IsClassificationEnabled() && bStrataFastPath);
 
 			TShaderMapRef< FDeferredLightPS > PixelShader( View.ShaderMap, PermutationVector );
@@ -2383,7 +2380,6 @@ void FDeferredShadingSceneRenderer::RenderLight(
 			PermutationVector.Set< FDeferredLightPS::FHairLighting>(0);
 			PermutationVector.Set < FDeferredLightPS::FAtmosphereTransmittance >(false);
 			PermutationVector.Set< FDeferredLightPS::FCloudTransmittance >(false);
-			PermutationVector.Set< FDeferredLightPS::FStrata >(Strata::IsStrataEnabled());
 			PermutationVector.Set< FDeferredLightPS::FStrataFastPath >(Strata::IsStrataEnabled() && Strata::IsClassificationEnabled() && bStrataFastPath);
 
 			TShaderMapRef< FDeferredLightPS > PixelShader( View.ShaderMap, PermutationVector );
@@ -2587,7 +2583,6 @@ void FDeferredShadingSceneRenderer::RenderLightForHair(
 			PermutationVector.Set< FDeferredLightPS::FLightingChannelsDim >(View.bUsesLightingChannels);
 			PermutationVector.Set< FDeferredLightPS::FVisualizeCullingDim >(false);
 			PermutationVector.Set< FDeferredLightPS::FTransmissionDim >(false);
-			PermutationVector.Set< FDeferredLightPS::FStrata >(Strata::IsStrataEnabled());
 			PermutationVector.Set< FDeferredLightPS::FHairLighting>(1);
 
 			TShaderMapRef<TDeferredLightHairVS> VertexShader(View.ShaderMap);
diff --git a/Engine/Source/Runtime/Renderer/Private/ScreenSpaceRayTracing.cpp b/Engine/Source/Runtime/Renderer/Private/ScreenSpaceRayTracing.cpp
--- a/Engine/Source/Runtime/Renderer/Private/ScreenSpaceRayTracing.cpp
+++ b/Engine/Source/Runtime/Renderer/Private/ScreenSpaceRayTracing.cpp
@@ -361,7 +361,6 @@ enum class ELightingTerm
 class FSSRQualityDim : SHADER_PERMUTATION_ENUM_CLASS("SSR_QUALITY", ESSRQuality);
 class FSSROutputForDenoiser : SHADER_PERMUTATION_BOOL("SSR_OUTPUT_FOR_DENOISER");
 class FLightingTermDim : SHADER_PERMUTATION_ENUM_CLASS("DIM_LIGHTING_TERM", ELightingTerm);
-class FStrata : SHADER_PERMUTATION_BOOL("STRATA_ENABLED");
 
 class FSSRTPrevFrameReductionCS : public FGlobalShader
 {
@@ -445,7 +444,7 @@ class FScreenSpaceReflectionsStencilPS : public FGlobalShader
 	DECLARE_GLOBAL_SHADER(FScreenSpaceReflectionsStencilPS);
 	SHADER_USE_PARAMETER_STRUCT(FScreenSpaceReflectionsStencilPS, FGlobalShader);
 
-	using FPermutationDomain = TShaderPermutationDomain<FSSROutputForDenoiser, FStrata>;
+	using FPermutationDomain = TShaderPermutationDomain<FSSROutputForDenoiser>;
 
 	static bool ShouldCompilePermutation(const FGlobalShaderPermutationParameters& Parameters)
 	{
@@ -455,7 +454,8 @@ class FScreenSpaceReflectionsStencilPS : public FGlobalShader
 	static void ModifyCompilationEnvironment(const FGlobalShaderPermutationParameters& Parameters, FShaderCompilerEnvironment& OutEnvironment)
 	{
 		FGlobalShader::ModifyCompilationEnvironment(Parameters, OutEnvironment);
-		OutEnvironment.SetDefine( TEXT("SSR_QUALITY"), uint32(0) );
+		OutEnvironment.SetDefine(TEXT("SSR_QUALITY"), uint32(0));
+		OutEnvironment.SetDefine(TEXT("STRATA_ENABLED"), Strata::IsStrataEnabled() ? 1u : 0u);
 	}
 	
 	BEGIN_SHADER_PARAMETER_STRUCT(FParameters, )
@@ -470,13 +470,19 @@ class FScreenSpaceReflectionsPS : public FGlobalShader
 	DECLARE_GLOBAL_SHADER(FScreenSpaceReflectionsPS);
 	SHADER_USE_PARAMETER_STRUCT(FScreenSpaceReflectionsPS, FGlobalShader);
 
-	using FPermutationDomain = TShaderPermutationDomain<FSSRQualityDim, FSSROutputForDenoiser, FStrata>;
+	using FPermutationDomain = TShaderPermutationDomain<FSSRQualityDim, FSSROutputForDenoiser>;
 
 	static bool ShouldCompilePermutation(const FGlobalShaderPermutationParameters& Parameters)
 	{
 		FPermutationDomain PermutationVector(Parameters.PermutationId);
 		return IsFeatureLevelSupported(Parameters.Platform, ERHIFeatureLevel::SM5);
 	}
+
+	static void ModifyCompilationEnvironment(const FGlobalShaderPermutationParameters& Parameters, FShaderCompilerEnvironment& OutEnvironment)
+	{
+		FGlobalShader::ModifyCompilationEnvironment(Parameters, OutEnvironment);
+		OutEnvironment.SetDefine(TEXT("STRATA_ENABLED"), Strata::IsStrataEnabled() ? 1u : 0u);
+	}
 	
 	BEGIN_SHADER_PARAMETER_STRUCT(FParameters, )
 		SHADER_PARAMETER_STRUCT_INCLUDE(FSSRCommonParameters, CommonParameters)
@@ -1090,7 +1096,6 @@ void RenderScreenSpaceReflections(
 
 		FScreenSpaceReflectionsStencilPS::FPermutationDomain PermutationVector;
 		PermutationVector.Set<FSSROutputForDenoiser>(bDenoiser);
-		PermutationVector.Set<FStrata>(Strata::IsStrataEnabled());
 
 		FScreenSpaceReflectionsStencilPS::FParameters* PassParameters = GraphBuilder.AllocParameters<FScreenSpaceReflectionsStencilPS::FParameters>();
 		PassParameters->CommonParameters = CommonParameters;
@@ -1172,7 +1177,6 @@ void RenderScreenSpaceReflections(
 	FScreenSpaceReflectionsPS::FPermutationDomain PermutationVector;
 	PermutationVector.Set<FSSRQualityDim>(SSRQuality);
 	PermutationVector.Set<FSSROutputForDenoiser>(bDenoiser);
-	PermutationVector.Set<FStrata>(Strata::IsStrataEnabled());
 		
 	FScreenSpaceReflectionsPS::FParameters* PassParameters = GraphBuilder.AllocParameters<FScreenSpaceReflectionsPS::FParameters>();
 	PassParameters->CommonParameters = CommonParameters;
diff --git a/Engine/Source/Runtime/Renderer/Private/SingleLayerWaterRendering.cpp b/Engine/Source/Runtime/Renderer/Private/SingleLayerWaterRendering.cpp
--- a/Engine/Source/Runtime/Renderer/Private/SingleLayerWaterRendering.cpp
+++ b/Engine/Source/Runtime/Renderer/Private/SingleLayerWaterRendering.cpp
@@ -149,8 +149,7 @@ class FSingleLayerWaterCompositePS : public FGlobalShader
 	class FScreenSpaceReflections : SHADER_PERMUTATION_BOOL("SCREEN_SPACE_REFLECTION");
 	class FHasBoxCaptures : SHADER_PERMUTATION_BOOL("REFLECTION_COMPOSITE_HAS_BOX_CAPTURES");
 	class FHasSphereCaptures : SHADER_PERMUTATION_BOOL("REFLECTION_COMPOSITE_HAS_SPHERE_CAPTURES");
-	class FStrata : SHADER_PERMUTATION_BOOL("STRATA_ENABLED");
-	using FPermutationDomain = TShaderPermutationDomain<FScreenSpaceReflections, FHasBoxCaptures, FHasSphereCaptures, FStrata>;
+	using FPermutationDomain = TShaderPermutationDomain<FScreenSpaceReflections, FHasBoxCaptures, FHasSphereCaptures>;
 
 	BEGIN_SHADER_PARAMETER_STRUCT(FParameters, )
 		SHADER_PARAMETER_STRUCT_INCLUDE(FSingleLayerWaterCommonShaderParameters, CommonParameters)
@@ -169,6 +168,7 @@ class FSingleLayerWaterCompositePS : public FGlobalShader
 	static void ModifyCompilationEnvironment(const FGlobalShaderPermutationParameters& Parameters, FShaderCompilerEnvironment& OutEnvironment)
 	{
 		FGlobalShader::ModifyCompilationEnvironment(Parameters, OutEnvironment);
+		OutEnvironment.SetDefine(TEXT("STRATA_ENABLED"), Strata::IsStrataEnabled() ? 1u : 0u);
 	}
 };
 
@@ -179,9 +179,6 @@ class FWaterTileCategorisationCS : public FGlobalShader
 	DECLARE_GLOBAL_SHADER(FWaterTileCategorisationCS);
 	SHADER_USE_PARAMETER_STRUCT(FWaterTileCategorisationCS, FGlobalShader)
 
-	class FStrata : SHADER_PERMUTATION_BOOL("STRATA_ENABLED");
-	using FPermutationDomain = TShaderPermutationDomain<FStrata>;
-
 	static int32 GetTileSize()
 	{
 		return 8;
@@ -208,6 +205,7 @@ class FWaterTileCategorisationCS : public FGlobalShader
 	{
 		OutEnvironment.SetDefine(TEXT("TILE_CATERGORISATION_SHADER"), 1.0f);
 		OutEnvironment.SetDefine(TEXT("WORK_TILE_SIZE"), GetTileSize());
+		OutEnvironment.SetDefine(TEXT("STRATA_ENABLED"), Strata::IsStrataEnabled() ? 1u : 0u);
 		FGlobalShader::ModifyCompilationEnvironment(Parameters, OutEnvironment);
 	}
 };
@@ -441,9 +439,7 @@ void FDeferredShadingSceneRenderer::RenderSingleLayerWaterReflections(
 
 			// Categorization based on SHADING_MODEL_ID
 			{
-				FWaterTileCategorisationCS::FPermutationDomain PermutationVector;
-				PermutationVector.Set<FWaterTileCategorisationCS::FStrata>(Strata::IsStrataEnabled());
-				TShaderMapRef<FWaterTileCategorisationCS> ComputeShader(View.ShaderMap, PermutationVector);
+				TShaderMapRef<FWaterTileCategorisationCS> ComputeShader(View.ShaderMap);
 
 				FWaterTileCategorisationCS::FParameters* PassParameters = GraphBuilder.AllocParameters<FWaterTileCategorisationCS::FParameters>();
 				SetCommonParameters(PassParameters->CommonParameters);
@@ -589,7 +585,6 @@ void FDeferredShadingSceneRenderer::RenderSingleLayerWaterReflections(
 			PermutationVector.Set<FSingleLayerWaterCompositePS::FScreenSpaceReflections>(bEnableSSR || bComposeLumenReflections);
 			PermutationVector.Set<FSingleLayerWaterCompositePS::FHasBoxCaptures>(bHasBoxCaptures);
 			PermutationVector.Set<FSingleLayerWaterCompositePS::FHasSphereCaptures>(bHasSphereCaptures);
-			PermutationVector.Set<FSingleLayerWaterCompositePS::FStrata>(Strata::IsStrataEnabled());
 			TShaderMapRef<FSingleLayerWaterCompositePS> PixelShader(View.ShaderMap, PermutationVector);
 
 			FWaterCompositeParameters* PassParameters = GraphBuilder.AllocParameters<FWaterCompositeParameters>();
