Fixed NonUnity Compile error

diff --git a/Engine/Source/Runtime/Renderer/Private/RendererModule.h b/Engine/Source/Runtime/Renderer/Private/RendererModule.h
--- a/Engine/Source/Runtime/Renderer/Private/RendererModule.h
+++ b/Engine/Source/Runtime/Renderer/Private/RendererModule.h
@@ -18,6 +18,7 @@ class FSceneInterface;
 class FSceneView;
 class FSceneViewFamily;
 class FSceneViewStateInterface;
+class FViewInfo;
 struct FMeshBatch;
 struct FSynthBenchmarkResults;
 struct FSceneTextures;
