Fix non-unity compile error due to missing forward declaration of FDistanceFieldSceneData.

diff --git a/Engine/Source/Runtime/Renderer/Private/DistanceFieldLightingShared.h b/Engine/Source/Runtime/Renderer/Private/DistanceFieldLightingShared.h
--- a/Engine/Source/Runtime/Renderer/Private/DistanceFieldLightingShared.h
+++ b/Engine/Source/Runtime/Renderer/Private/DistanceFieldLightingShared.h
@@ -22,6 +22,7 @@ class FPrimitiveSceneInfo;
 class FSceneRenderer;
 class FShaderParameterMap;
 class FViewInfo;
+class FDistanceFieldSceneData;
 
 DECLARE_LOG_CATEGORY_EXTERN(LogDistanceField, Log, All);
 
