Fixed AllocatePagesUsingRects with Nanite

diff --git a/Engine/Source/Runtime/Renderer/Private/ShadowDepthRendering.cpp b/Engine/Source/Runtime/Renderer/Private/ShadowDepthRendering.cpp
--- a/Engine/Source/Runtime/Renderer/Private/ShadowDepthRendering.cpp
+++ b/Engine/Source/Runtime/Renderer/Private/ShadowDepthRendering.cpp
@@ -1760,6 +1760,8 @@ void FSceneRenderer::RenderShadowDepthMapAtlases(FRDGBuilder& GraphBuilder)
 	}
 }
 
+extern TAutoConsoleVariable<int32> CVarAllocatePagesUsingRects;
+
 void FSceneRenderer::RenderShadowDepthMaps(FRDGBuilder& GraphBuilder, FInstanceCullingManager &InstanceCullingManager)
 {
 	CSV_SCOPED_TIMING_STAT_EXCLUSIVE(RenderShadows);
@@ -1779,6 +1781,7 @@ void FSceneRenderer::RenderShadowDepthMaps(FRDGBuilder& GraphBuilder, FInstanceC
 	const bool bHasVSMClipMaps = SortedShadowsForShadowDepthPass.VirtualShadowMapClipmaps.Num() > 0;
 	const bool bNaniteEnabled = UseNanite(ShaderPlatform) && ViewFamily.EngineShowFlags.NaniteMeshes;
 	const bool bUseHZB = (CVarNaniteShadowsUseHZB.GetValueOnRenderThread() != 0);
+	const bool bAllocatePageRectAtlas = CVarAllocatePagesUsingRects.GetValueOnRenderThread() != 0;
 
 #if ENABLE_NON_NANITE_VSM
 	VirtualShadowMapArray.RenderVirtualShadowMapsHw(GraphBuilder, SortedShadowsForShadowDepthPass.VirtualShadowMapShadows, *Scene);
@@ -1807,9 +1810,12 @@ void FSceneRenderer::RenderShadowDepthMaps(FRDGBuilder& GraphBuilder, FInstanceC
 			const FIntPoint VirtualShadowSize = VirtualShadowMapArray.GetPhysicalPoolSize();
 			const FIntRect VirtualShadowViewRect = FIntRect(0, 0, VirtualShadowSize.X, VirtualShadowSize.Y);
 
-			Nanite::FRasterContext RasterContext = Nanite::InitRasterContext(GraphBuilder, VirtualShadowSize, Nanite::EOutputBufferMode::DepthOnly, false);
+			Nanite::FRasterContext RasterContext = Nanite::InitRasterContext(GraphBuilder, VirtualShadowSize, Nanite::EOutputBufferMode::DepthOnly, bAllocatePageRectAtlas);
 
-			VirtualShadowMapArray.ClearPhysicalMemory(GraphBuilder, RasterContext.DepthBuffer, Scene->VirtualShadowMapArrayCacheManager);
+			if( !bAllocatePageRectAtlas )
+			{
+				VirtualShadowMapArray.ClearPhysicalMemory(GraphBuilder, RasterContext.DepthBuffer, Scene->VirtualShadowMapArrayCacheManager);
+			}
 
 			const bool bUpdateStreaming = CVarNaniteShadowsUpdateStreaming.GetValueOnRenderThread() != 0;
 
diff --git a/Engine/Source/Runtime/Renderer/Private/VirtualShadowMaps/VirtualShadowMapArray.cpp b/Engine/Source/Runtime/Renderer/Private/VirtualShadowMaps/VirtualShadowMapArray.cpp
--- a/Engine/Source/Runtime/Renderer/Private/VirtualShadowMaps/VirtualShadowMapArray.cpp
+++ b/Engine/Source/Runtime/Renderer/Private/VirtualShadowMaps/VirtualShadowMapArray.cpp
@@ -98,15 +98,15 @@ FMatrix CalcTranslatedWorldToShadowUVMatrix(
 	FMatrix TranslatedWorldToShadowUv = TranslatedWorldToShadowClip * ScaleAndBiasToSmUV;
 	return TranslatedWorldToShadowUv;
 }
-#if ENABLE_NON_NANITE_VSM
 
-static TAutoConsoleVariable<int32> CVarAllocatePagesUsingRects(
+TAutoConsoleVariable<int32> CVarAllocatePagesUsingRects(
 	TEXT("r.Shadow.v.AllocatePagesUsingRects"),
 	0,
 	TEXT("If set to 1 then pages are allocated in a contigious block to support the bounding rectangle of allocated virtual pages leading to a trivial mapping but higher memory overhead."),
 	ECVF_RenderThreadSafe
 );
 
+#if ENABLE_NON_NANITE_VSM
 
 static TAutoConsoleVariable<int32> CVarShowClipmapStats(
 	TEXT("r.Shadow.v.ShowClipmapStats"),
