* Add strata energy LUT for computing directional alebdo for reflection and refraction terms
* Add simple furnace test for validating energy conservation
* Add multiple scattering support for Strata slab with analytical lighting.

Current state:
* Multiple scattering is disabled by defaullt
* Multiple scattering is computed only for analytical light
* Energy compensation is only correct for pure conductor (metallic = 1).

diff --git a/Engine/Shaders/Private/Strata/StrataEvaluation.ush b/Engine/Shaders/Private/Strata/StrataEvaluation.ush
--- a/Engine/Shaders/Private/Strata/StrataEvaluation.ush
+++ b/Engine/Shaders/Private/Strata/StrataEvaluation.ush
@@ -20,6 +20,9 @@
 #define STRATA_USE_ENV_BRDF_APPROX 0
 #endif
 
+#ifndef STRATA_MICROFACET_MULTIPLESCATTERING
+#define STRATA_MICROFACET_MULTIPLESCATTERING 0
+#endif
 
 // STRATA_TODO put in a common file
 // Point lobe in off-specular peak direction
@@ -92,7 +95,7 @@ void StrataGetBSDFRoughness(in FStrataBSDF BSDF, inout float R)
 	}
 }
 
-FStrataBSDFContext StrataCreateBSDFContext(FStrataPixelHeader StrataPixelHeader, FStrataBSDF BSDF, const FStrataAddressing StrataAddressing, float3 V, float3 L, bool bOverrideWithLocalFrame=false)
+FStrataBSDFContext StrataCreateBSDFContext(float3x3 TangentBasis, FStrataBSDF BSDF, float3 V, float3 L, bool bOverrideWithLocalFrame = false)
 {
 	FStrataBSDFContext BSDFContext = (FStrataBSDFContext)0;
 
@@ -101,7 +104,6 @@ FStrataBSDFContext StrataCreateBSDFContext(FStrataPixelHeader StrataPixelHeader,
 	BSDFContext.V = V;
 	BSDFContext.L = L;
 
-	float3x3 TangentBasis = StrataGetBSDFSharedBasis(StrataPixelHeader, BSDF, StrataAddressing);
 	BSDFContext.X = TangentBasis[0]; 
 	BSDFContext.Y = TangentBasis[1];
 	BSDFContext.N = TangentBasis[2]; 
@@ -126,8 +128,8 @@ FStrataBSDFContext StrataCreateBSDFContext(FStrataPixelHeader StrataPixelHeader,
 	}
 	BSDFContext.TangentBasis = float3x3(BSDFContext.X, BSDFContext.Y, BSDFContext.N);
 
-	BSDFContext.TangentV = mul(BSDFContext.TangentBasis, BSDFContext.V);
-	BSDFContext.TangentH = mul(BSDFContext.TangentBasis, BSDFContext.H);
+	BSDFContext.TangentV = mul(BSDFContext.V, BSDFContext.TangentBasis);
+	BSDFContext.TangentH = mul(BSDFContext.H, BSDFContext.TangentBasis);
 
 	BSDFContext.SatNoL = saturate(BSDFContext.Context.NoL);
 	BSDFContext.SatNoV = saturate(BSDFContext.Context.NoV);
@@ -135,12 +137,18 @@ FStrataBSDFContext StrataCreateBSDFContext(FStrataPixelHeader StrataPixelHeader,
 	return BSDFContext;
 }
 
+FStrataBSDFContext StrataCreateBSDFContext(FStrataPixelHeader StrataPixelHeader, FStrataBSDF BSDF, const FStrataAddressing StrataAddressing, float3 V, float3 L, bool bOverrideWithLocalFrame = false)
+{
+	float3x3 TangentBasis = StrataGetBSDFSharedBasis(StrataPixelHeader, BSDF, StrataAddressing);
+	return StrataCreateBSDFContext(TangentBasis, BSDF, V, L, bOverrideWithLocalFrame);
+}
+
 void StrataUpdateBSDFContext(inout FStrataBSDFContext BSDFContext, float3 NewL)
 {
 	// Update all the data related to L
 	Init(BSDFContext.Context, BSDFContext.N, BSDFContext.V, NewL);
 	BSDFContext.H = normalize(BSDFContext.V + BSDFContext.L);
-	BSDFContext.TangentH = mul(BSDFContext.TangentBasis, BSDFContext.H);
+	BSDFContext.TangentH = mul(BSDFContext.H, BSDFContext.TangentBasis);
 }
 
 struct FStrataEvaluateResult
@@ -182,6 +190,23 @@ float3 StrataEnvBRDF(float3 F0, float3 F90, float SafeRoughnessX, float SatNoV)
 #endif
 }
 
+float3 StrataEvaluateMicrofacetMultipleScattering(float3 F0, float NoV, float Metallic, float SafeRoughnessX, float SafeRoughnessY)
+{
+#if STRATA_MICROFACET_MULTIPLESCATTERING
+	// Use the main Roughness term to drive the energy compensation.
+	// "Revisiting Physically Based Shading at Imageworks" Kulla 2017
+	const float SafeRoughness = SafeRoughnessX;
+
+	const FStrataLUTResult StrataLUTSample = StrataSampleGGXLUT(Strata.GGXEnergyLUT2DTexture, Strata.GGXEnergyLUT3DTexture, Strata.GGXEnergyLUTSampler, Strata.GGXEnergyLUTScaleBias, NoV, SafeRoughness, F0);
+	const float3 Fs_0 = StrataEvaluateMultipleScatteringDielectric3D(F0, StrataLUTSample);
+	const float3 Fs_1 = StrataEvaluateMultipleScatteringConductor(F0, StrataLUTSample);
+	const float3 Fs   = lerp(Fs_0, Fs_1, Metallic);
+
+	return Fs;
+#else
+	return 1.0f;
+#endif
+}
 
 FStrataEvaluateResult StrataEvaluateBSDF(FStrataBSDFContext BSDFContext, FShadowTerms ShadowTerms)
 {
@@ -210,6 +235,9 @@ FStrataEvaluateResult StrataEvaluateBSDF(FStrataBSDFContext BSDFContext, FShadow
 			const float SafeRoughnessX = MakeRoughnessSafe(SLAB_ROUGHNESSX(BSDFContext.BSDF));
 			const float SafeRoughnessY = MakeRoughnessSafe(SLAB_ROUGHNESSY(BSDFContext.BSDF));
 
+			float HazeSafeRoughnessX = SafeRoughnessX;
+			float HazeSafeRoughnessY = SafeRoughnessY;
+
 			const bool bHasAnisotropy = BSDF_GETANISOTROPY(BSDFContext.BSDF);
 			const bool bHaziness = BSDF_GETHASHAZINESS(BSDFContext.BSDF);
 			float D = 0;
@@ -219,50 +247,74 @@ FStrataEvaluateResult StrataEvaluateBSDF(FStrataBSDFContext BSDFContext, FShadow
 			float HazeVis = 0;
 			float HazePDF = 0;
 
-#if STRATA_FASTPATH==0
+			#if STRATA_FASTPATH==0
 			BRANCH
 			if (bHasAnisotropy)
 			{
 				// Generalized microfacet specular
-				float Alpha0 = Pow2(SafeRoughnessX);
-				float Alpha1 = Pow2(SafeRoughnessY);
-				D = D_GGXaniso(Alpha0, Alpha1, BSDFContext.Context.NoH, BSDFContext.Context.XoH, BSDFContext.Context.YoH);
-				Vis = Vis_SmithJointAniso(Alpha0, Alpha1, BSDFContext.Context.NoV, BSDFContext.Context.NoL, BSDFContext.Context.XoV, BSDFContext.Context.XoL, BSDFContext.Context.YoV, BSDFContext.Context.YoL);
-				PDF = VisibleGGXPDF_aniso(BSDFContext.TangentV, BSDFContext.TangentH, float2(Alpha0, Alpha1));
+				{
+					const float Alpha0 = Pow2(SafeRoughnessX);
+					const float Alpha1 = Pow2(SafeRoughnessY);
+					D   = D_GGXaniso(Alpha0, Alpha1, BSDFContext.Context.NoH, BSDFContext.Context.XoH, BSDFContext.Context.YoH);
+					Vis = Vis_SmithJointAniso(Alpha0, Alpha1, BSDFContext.Context.NoV, BSDFContext.Context.NoL, BSDFContext.Context.XoV, BSDFContext.Context.XoL, BSDFContext.Context.YoV, BSDFContext.Context.YoL);
+					PDF = VisibleGGXPDF_aniso(BSDFContext.TangentV, BSDFContext.TangentH, float2(Alpha0, Alpha1));
+				}
 
 				if (bHaziness)
 				{
-					Alpha0 = Pow2(ComputeHazyLobeRoughness(SafeRoughnessX, SLAB_HAZINESS(BSDFContext.BSDF)));
-					Alpha1 = Pow2(ComputeHazyLobeRoughness(SafeRoughnessY, SLAB_HAZINESS(BSDFContext.BSDF)));
-					HazeD = D_GGXaniso(Alpha0, Alpha1, BSDFContext.Context.NoH, BSDFContext.Context.XoH, BSDFContext.Context.YoH);
-					HazeVis = Vis_SmithJointAniso(Alpha0, Alpha1, BSDFContext.Context.NoV, BSDFContext.Context.NoL, BSDFContext.Context.XoV, BSDFContext.Context.XoL, BSDFContext.Context.YoV, BSDFContext.Context.YoL);
-					HazePDF = VisibleGGXPDF_aniso(BSDFContext.TangentV, BSDFContext.TangentH, float2(Alpha0, Alpha1));
+					HazeSafeRoughnessX = ComputeHazyLobeRoughness(SafeRoughnessX, SLAB_HAZINESS(BSDFContext.BSDF));
+					HazeSafeRoughnessY = ComputeHazyLobeRoughness(SafeRoughnessY, SLAB_HAZINESS(BSDFContext.BSDF));
+
+					const float HazeAlpha0 = Pow2(HazeSafeRoughnessX);
+					const float HazeAlpha1 = Pow2(HazeSafeRoughnessY);
+					HazeD   = D_GGXaniso(HazeAlpha0, HazeAlpha1, BSDFContext.Context.NoH, BSDFContext.Context.XoH, BSDFContext.Context.YoH);
+					HazeVis = Vis_SmithJointAniso(HazeAlpha0, HazeAlpha1, BSDFContext.Context.NoV, BSDFContext.Context.NoL, BSDFContext.Context.XoV, BSDFContext.Context.XoL, BSDFContext.Context.YoV, BSDFContext.Context.YoL);
+					HazePDF = VisibleGGXPDF_aniso(BSDFContext.TangentV, BSDFContext.TangentH, float2(HazeAlpha0, HazeAlpha1));
 				}
 			}
 			else
-#endif
+			#endif
 			{
 				// Special override for roughness for supporting area light integrator with Sphere/Tube/Disk light, which modifies/increase roughness.
 				const float a2 = BSDFContext.a2Override > 0 ? BSDFContext.a2Override : Pow4(SafeRoughnessX);
 
 				// Generalized microfacet specular
-				D = D_GGX(a2, BSDFContext.Context.NoH);
-				Vis = Vis_SmithJointApprox(a2, BSDFContext.Context.NoV, BSDFContext.SatNoL);
-				PDF = VisibleGGXPDF(BSDFContext.TangentV, BSDFContext.TangentH, a2);
+				{
+					D   = D_GGX(a2, BSDFContext.Context.NoH);
+					Vis = Vis_SmithJoint(a2, BSDFContext.Context.NoV, BSDFContext.SatNoL);
+					PDF = VisibleGGXPDF(BSDFContext.TangentV, BSDFContext.TangentH, a2);
+				}
 
-#if STRATA_FASTPATH==0
+				#if STRATA_FASTPATH==0
 				if (bHaziness)
 				{
-					const float a2Haze = BSDFContext.a2OverrideHazy > 0 ? BSDFContext.a2OverrideHazy : Pow4(ComputeHazyLobeRoughness(SafeRoughnessX, SLAB_HAZINESS(BSDFContext.BSDF)));
-					HazeD = D_GGX(a2Haze, BSDFContext.Context.NoH);
-					HazeVis = Vis_SmithJointApprox(a2Haze, BSDFContext.Context.NoV, BSDFContext.SatNoL);
+					HazeSafeRoughnessX = ComputeHazyLobeRoughness(SafeRoughnessX, SLAB_HAZINESS(BSDFContext.BSDF));
+
+					const float a2Haze = BSDFContext.a2OverrideHazy > 0 ? BSDFContext.a2OverrideHazy : Pow4(HazeSafeRoughnessX);
+					HazeD   = D_GGX(a2Haze, BSDFContext.Context.NoH);
+					HazeVis = Vis_SmithJoint(a2Haze, BSDFContext.Context.NoV, BSDFContext.SatNoL);
 					HazePDF = VisibleGGXPDF(BSDFContext.TangentV, BSDFContext.TangentH, a2Haze);
+
 				}
-#endif
+				#endif
+			}
+
+			float3 MSScale = 1;
+			float3 HazeMSScale = 1;
+			#if STRATA_MICROFACET_MULTIPLESCATTERING
+			{
+				MSScale	= StrataEvaluateMicrofacetMultipleScattering(F0, BSDFContext.Context.NoV, Metallic, SafeRoughnessX, SafeRoughnessY);
+				#if STRATA_FASTPATH==0
+				if (bHaziness)
+				{
+					HazeMSScale = StrataEvaluateMicrofacetMultipleScattering(F0, BSDFContext.Context.NoV, Metallic, HazeSafeRoughnessX, HazeSafeRoughnessY);
+				}
+				#endif
 			}
+			#endif
 
 			float3 F = 0;
-#if STRATA_FASTPATH==0
+			#if STRATA_FASTPATH==0
 			const bool bHasThinFilm = BSDF_GETHASTHINFILM(BSDFContext.BSDF);
 			BRANCH
 			if (bHasThinFilm)
@@ -272,18 +324,18 @@ FStrataEvaluateResult StrataEvaluateBSDF(FStrataBSDFContext BSDFContext, FShadow
 				F = F_ThinFilm(BSDFContext.Context.NoV, BSDFContext.Context.NoL, BSDFContext.Context.VoH, F0, F90, ThinFilmIOR, ThinFilmThickness);
 			}
 			else
-#endif
+			#endif
 			{
 				F = F_Schlick(F0, F90, BSDFContext.Context.VoH);
 			}
 
-			Sample.SpecularPathProbability	= bHaziness		? (1.0f - SPECULAR_HAZINESS_WEIGHT) : 1.0f;
-			Sample.SpecularHazePathProbability= bHaziness	? SPECULAR_HAZINESS_WEIGHT			: 0.0f;
+			Sample.SpecularPathProbability		= bHaziness	? (1.0f - SPECULAR_HAZINESS_WEIGHT) : 1.0f;
+			Sample.SpecularHazePathProbability	= bHaziness	? SPECULAR_HAZINESS_WEIGHT			: 0.0f;
 
-			Sample.SpecularPathValue		= D * Vis * F;
+			Sample.SpecularPathValue		= D * Vis * F * MSScale;
 			Sample.SpecularPDF				= PDF * Sample.SpecularPathProbability;
 
-			Sample.SpecularHazePathValue	= HazeD * HazeVis * F;
+			Sample.SpecularHazePathValue	= HazeD * HazeVis * F * HazeMSScale;
 			Sample.SpecularHazePDF			= HazePDF * Sample.SpecularHazePathProbability;
 
 			Sample.EmissivePathValue= BSDF_GETEMISSIVE(BSDFContext.BSDF);
@@ -292,7 +344,7 @@ FStrataEvaluateResult StrataEvaluateBSDF(FStrataBSDFContext BSDFContext, FShadow
 			break;
 		}
 
-#if STRATA_FASTPATH==0
+		#if STRATA_FASTPATH==0
 		case STRATA_BSDF_TYPE_SHEEN:
 		{
 			const float3 F0  = SHEEN_BASECOLOR(BSDFContext.BSDF);
@@ -337,7 +389,7 @@ FStrataEvaluateResult StrataEvaluateBSDF(FStrataBSDFContext BSDFContext, FShadow
 		//case STRATA_BSDF_TYPE_SINGLELAYERWATER:
 		//Nothing to do in this case because these BSDF are evaluated in other specialised passes.
 		//break;
-#endif
+		#endif
 	}
 
 	return Sample;
diff --git a/Engine/Shaders/Private/Strata/StrataFurnaceTest.usf b/Engine/Shaders/Private/Strata/StrataFurnaceTest.usf
new file mode 100644
--- /dev/null
+++ b/Engine/Shaders/Private/Strata/StrataFurnaceTest.usf
@@ -0,0 +1,154 @@
+// Copyright Epic Games, Inc. All Rights Reserved.
+
+#include "/Engine/Private/Common.ush"
+
+#define STRATA_INLINE_SHADING 0
+#define STRATA_MICROFACET_MULTIPLESCATTERING 1
+
+#include "StrataMultipleScattering.ush"
+#include "StrataEvaluation.ush"
+
+#if SHADER_FURNACE_ANALYTIC
+
+uint SceneType;
+uint NumSamples;
+int2 OutputResolution;
+
+void MainPS(
+	float4 SVPos : SV_POSITION,
+	out float3 OutColor : SV_Target0)
+{
+	const float3 InRadiance = 0.5f;
+	OutColor = InRadiance;
+
+	const uint2 TileCount = 8;
+	const uint2 TileResolution = OutputResolution / TileCount;
+
+	const uint2 PixelPos = uint2(SVPos.xy);	
+	const float2 UV = PixelPos / float2(OutputResolution);
+	const float2 TileUV = frac(UV * TileCount);
+	const uint2 TileCoord = floor(UV * TileCount);
+	const uint2 TilePixelCoord = TileUV * TileResolution;
+
+	float3 V = float3(0, 0, -1);
+	float3 N = 0;
+	bool bEvaluate = false;
+	{
+		const float3 RayO = float3(TileUV, 0);
+		const float3 RayD = float3(0, 0, 1);
+		const float4 Sphere = float4(0.5f, 0.5f, 1.f, 0.4f);
+		const float2 Hit = RayIntersectSphere(RayO, RayD, Sphere);
+		if (Hit.x >= 0)
+		{
+			const float3 P = RayO + RayD * Hit.x;
+			N = normalize(P - Sphere.xyz);
+			bEvaluate = true;
+		}
+	}
+
+	// Draw a simple 2D grid of spheres 
+	if (bEvaluate)
+	{
+		// Build a single layer BSDF with white albedo
+		FStrataBSDF BSDF = (FStrataBSDF)0;
+		{
+			const float MinRoughness = 0.02;
+
+			float3 BaseColor = 1;
+			float3 EdgeColor = 1;
+			float Roughness  = 0.5f;
+			float Metallic   = 0.0f;
+			float Specular   = 0.5f;
+
+			if (SceneType == 1)
+			{
+				//  ------> Roughness
+				// |
+				// |
+				// v
+				// Metallic
+				Roughness = max(MinRoughness, (TileCoord.x) / float(TileCount.x - 1));
+				Metallic = (TileCoord.y) / float(TileCount.y - 1);
+			}
+			else if (SceneType == 2)
+			{
+				// All metalic with different roughness & metal type (F0)
+				//  ------> Roughness
+				//  -> Silver
+				//  -> Aluminium
+				//  -> Gold
+				//  -> Copper
+				const float3 Silver      = float3(0.971519,    0.959915,    0.915324);
+				const float3 Aluminium   = float3(0.913183,    0.921494,    0.924524);
+				const float3 Gold        = float3(1       ,    0.765557,    0.336057);
+				const float3 Copper      = float3(0.955008,    0.637427,    0.538163);
+				const float3 Chromium    = float3(0.549585,    0.556114,    0.554256);
+				const float3 Nickel      = float3(0.659777,    0.608679,    0.525649);
+				const float3 Titanium    = float3(0.541931,    0.496791,    0.449419);
+				const float3 Cobalt      = float3(0.662124,    0.654864,    0.633732);
+				const float3 Platinum    = float3(0.672411,    0.637331,    0.585456);
+
+				Roughness = max(MinRoughness, (TileCoord.x) / float(TileCount.x - 1));
+				Metallic = 1;
+				switch (TileCoord.y)
+				{
+					case 0: BaseColor = Silver; break;
+					case 1: BaseColor = Aluminium; break;
+					case 2: BaseColor = Gold; break;
+					case 3: BaseColor = Copper; break;
+					case 4: BaseColor = Chromium; break;
+					case 5: BaseColor = Nickel; break;
+					case 6: BaseColor = Titanium; break;
+					case 7: BaseColor = Cobalt; break;
+				}
+			}
+
+			BSDF_SETTYPE(BSDF, STRATA_BSDF_TYPE_SLAB);
+			BSDF_SETNORMALID(BSDF, 0);
+			BSDF_SETANISOTROPY(BSDF, 0);
+			BSDF_SETISTOPLAYER(BSDF, 1);
+			BSDF_SETHASEDGECOLOR(BSDF, 0);
+			BSDF_SETHASSCATTERING(BSDF, 0);
+			BSDF_SETHASSSSPROFILE(BSDF, 0);
+			BSDF_SETHASHAZINESS(BSDF, 0);
+			BSDF_SETHASTHINFILM(BSDF, 0);
+			BSDF.Weight = 1.0f;
+
+			SLAB_BASECOLOR(BSDF)	= BaseColor;
+			SLAB_EDGECOLOR(BSDF)	= EdgeColor;
+			SLAB_ROUGHNESSX(BSDF)	= Roughness;
+			SLAB_ROUGHNESSY(BSDF)	= Roughness;
+			SLAB_SPECULAR(BSDF)		= Specular;
+			SLAB_METALLIC(BSDF)		= Metallic;
+			SLAB_HAZINESS(BSDF)		= 0;
+		}
+
+		// Sample uniformly the upper hemisphere of each pixel
+		float3 Acc = 0;
+		const FShadowTerms ShadowTerms = { 1.0f, 1.0, 1.0, InitHairTransmittanceData() };
+		LOOP
+		for (uint SampleIt = 0; SampleIt < NumSamples; ++SampleIt)
+		{
+			const float3x3 TangentBasis = GetTangentBasis(N);
+			const float2 U = Hammersley(SampleIt, NumSamples, 0);
+
+			const float4 S = UniformSampleHemisphere(U);
+			const float3 L = mul(S.xyz, TangentBasis);
+
+			const float SatNoL = saturate(dot(L, N));
+
+			FStrataBSDFContext BSDFContext = StrataCreateBSDFContext(TangentBasis, BSDF, V, L);
+			FStrataEvaluateResult BSDFEvaluate = StrataEvaluateBSDF(BSDFContext, ShadowTerms);
+
+			const float3 DiffusePathWeight  = BSDF.Weight * BSDFEvaluate.DiffusePathValue;
+			const float3 SpecularPathWeight = BSDF.Weight * BSDFEvaluate.SpecularPathValue;
+
+			Acc += SatNoL * InRadiance * (DiffusePathWeight + SpecularPathWeight) / (S.w != 0 ?  S.w : 1);
+		}
+		OutColor = Acc / max(1, SampleIt); 
+	}
+}
+
+#endif // SHADER_FURNACE_ANALYTIC
+
+////////////////////////////////////////////////////////////////////////////////////////////////////////////
\ No newline at end of file
diff --git a/Engine/Shaders/Private/Strata/StrataLUT.usf b/Engine/Shaders/Private/Strata/StrataLUT.usf
new file mode 100644
--- /dev/null
+++ b/Engine/Shaders/Private/Strata/StrataLUT.usf
@@ -0,0 +1,215 @@
+// Copyright Epic Games, Inc. All Rights Reserved.
+
+#include "/Engine/Private/Common.ush"
+#include "/Engine/Private/MonteCarlo.ush"
+#include "/Engine/Private/ShadingCommon.ush"
+#include "/Engine/Private/BRDF.ush"
+#include "/Engine/Private/PathTracing/Material/PathTracingGlossy.ush"
+
+//////////////////////////////////////////////////////////////////////////////////////////////////////////////
+
+struct FDirectionResult
+{
+	float3 Direction;
+	float Fresnel;
+	bool bIsRefacted;
+};
+
+// Compute the outgoing direction for a microfacet
+FDirectionResult ComputeDirection(float3 RayDirection, float3 Normal, float Fresnel, float Eta, bool bRefractEnable)
+{
+	FDirectionResult Out;
+	Out.Direction = 0;
+	Out.Fresnel = 0;
+	Out.bIsRefacted = false;
+
+	// Compute refracted direction and Fresnel term at the same time to re-use intermediate values
+	// and to get matching results for the TIR condition
+	float CosI = abs(dot(RayDirection, Normal));
+	float g2 = Eta * Eta - 1.0 + CosI * CosI;
+
+	if (g2 >= 0.0)
+	{
+		const float c = CosI;
+		const float g = sqrt(g2);
+		if (bRefractEnable)
+		{
+			Out.Fresnel = 1 - Fresnel;
+			Out.Direction = normalize(RayDirection + Normal * (c - g));
+			return Out;
+		}
+		Out.Fresnel = Fresnel;
+	}
+	else
+	{
+		// TIR
+		Out.Fresnel = 1.0;
+	}
+	// either TIR, or sampled reflection
+	Out.Direction = RayDirection + (2 * CosI) * Normal;
+	return Out;
+}
+
+//////////////////////////////////////////////////////////////////////////////////////////////////////////////
+
+#if SHADER_LUT
+
+uint NumSamples;
+uint EnergyLUTResolution;
+int2 SliceXYCount;
+
+RWTexture3D<float2> OutLUT3D;
+RWTexture2D<float4> OutLUT2D;
+
+void MainPS(
+	float4 SVPos : SV_POSITION, 
+	out float2 OutColor : SV_Target0)
+{
+	const uint3 LUTResolution = EnergyLUTResolution.xxx;
+	const uint2 PixelPos = uint2(SVPos.xy);
+
+	const uint2 SliceCoord = uint2(
+		PixelPos.x / LUTResolution.x,
+		PixelPos.y / LUTResolution.y);
+
+	const uint3 Coord = uint3(
+		PixelPos.x % LUTResolution.x,
+		PixelPos.y % LUTResolution.y,
+		SliceCoord.x + SliceCoord.y * SliceXYCount.x);
+
+	float A_R = 0.0f;
+	float B_R = 0.0f;
+
+	float A_T = 0.0f;
+	float B_T = 0.0f;
+
+	float E_R = 0.0f;
+	float E_T = 0.0f;
+
+	// Roughness
+	const float Roughness = float(Coord.y + 0.5f) / LUTResolution.y;
+	const float m = Roughness * Roughness;
+	const float m2 = m * m;
+
+	// Incident angle
+	const float NoV = float(Coord.x + 0.5f) / LUTResolution.x;
+	const float3 V = float3(sqrt(1.0f - NoV * NoV),	0.0f, NoV);
+
+	LOOP
+	for (uint i = 0; i < NumSamples; i++)
+	{
+		const float2 E = Hammersley(i, NumSamples, uint2(0, 0));
+
+		// Decomposition based on Schlick's fresnel formula
+		// F_Schlick = F0 + (1-F0) (1-cos)^5 = F0 + (1-F0) Fc
+		// 
+		// UPPER \int{Omera_{+}}
+		// =====================
+		// Decomposition in A B terms of the directional albedo for microfacter for *upper* hemisphere
+		// E_R
+		// = \int F.D.G
+		// = \int (F0 + (1-F0) Fc).D.G
+		// = \int F0.D.G + \int (1-F0).Fc.D.G
+		// = F0 \int D.G + (1-F0) \int Fc.D.G
+		//      ========           ==========
+		//         A_R                 B_R
+		//
+		// Reconstruction of the directional Albedo for upper hemisphere
+		// E_R = F0.A_R + (1-F0).B_R
+		//
+		//
+		// LOWER \int{Omera_{-}}
+		// =====================
+		// Decomposition in A B terms of the directional albedo for microfacter for *lower* hemisphere
+		//
+		// E_T
+		// = \int (1-F).D.G
+		// = \int (1- (F0 + (1-F0) Fc)).D.G
+		// = \int (1 - F0 - (1-F0) Fc).D.G
+		// = \int D.G - F0.D.G - (1-F0) Fc.D.G
+		// = \int D.G - \int F0.D.G - \int (1-F0) Fc.D.G
+		// = \int D.G - F0 \int D.G - (1-F0) \int Fc.D.G
+		//   ========      ========          ===========
+		//     A_T            A_T                B_T
+		//
+		//
+		// Reconstruction of the directional Albedo for lower hemisphere
+		// E_T = A_T + F0.A_T + (1-F0).B_T
+		//
+		
+		// "Based on "Practical multiple scattering compensation for microfacet models". E. Turquin
+		//
+		// Dieletric
+		// ==========
+		// Microfacet multiple scattering term for *dielietric*
+		//  f_ss(wo,wi) = F(wo,wi) . D(wo,wi) . G(wo,wi)
+		//  f_ms(wo,wi) = F_ss / (E_T + E_R)
+		//
+		// Conductor
+		// =========
+		// Microfacet multiple scattering term for *conductor*
+		//  f_ss(wo,wi) = F(wo,wi) . D(wo,wi) . G(wo,wi)
+		//  f_ms(wo,wi) = F_ss . (1 +  F0 (1-E_R)/ E_R)
+
+		// TODO: parametrize F0 with something which preserve precision in the lower part e.g., sqrt(F0)
+		const float F0 = float(Coord.z) / LUTResolution.z;
+		const float F90 = 1;
+		{
+			float4 GGXSample = ImportanceSampleVisibleGGX(UniformSampleDisk(E), m2, V);
+			const float3 H  = GGXSample.xyz;
+			const float VoH = saturate(dot(V, H));
+			const float F	= F_Schlick(F0, F90, VoH).x;
+			
+			const float Eta_Incident = 1.0f;
+			const float Eta_Outgoing = DielectricF0ToIor(F0);
+			const float Eta = Eta_Outgoing / Eta_Incident;
+
+			const FDirectionResult Sample_Lr = ComputeDirection(-V, H, F, Eta, false);
+			const FDirectionResult Sample_Lt = ComputeDirection(-V, H, F, Eta, true);
+
+			const float NoL_r = Sample_Lr.Direction.z;
+			const float NoL_t = Sample_Lt.Direction.z;
+
+			// Schlick's Fresnel sub-term: Fc = (1-VoH)^5 
+			const float Fc = Pow5(1.0f - VoH);
+
+			const float R_WeightPdf = GGXEvalReflection(Sample_Lr.Direction, V, H, float2(m, m)).x;
+			const float T_WeightPdf = GGXEvalRefraction(Sample_Lt.Direction, V, H, float2(m, m), Eta).x;
+
+			E_R += Sample_Lr.Fresnel * R_WeightPdf;
+			E_T += Sample_Lt.Fresnel * T_WeightPdf;
+
+			A_R += R_WeightPdf;
+			B_R += R_WeightPdf * Fc;
+
+			A_T += T_WeightPdf;
+			B_T += T_WeightPdf * Fc;
+		}
+	}
+
+	A_R /= NumSamples;
+	B_R /= NumSamples;
+			
+	A_T /= NumSamples;
+	B_T /= NumSamples;
+
+	E_R /= NumSamples;
+	E_T /= NumSamples;
+
+	const float2 Out3D = float2(E_R, E_T);
+	const float4 Out2D = float4(A_R, B_R, A_T, B_T);
+
+	// At the moment ouput 2 LUTs:
+	// * A 2D LUT which factor out IOR/F0 data
+	// * A 3D LUT which compute the directional albedo for all F0 value (for debug purpose)
+	if (Coord.z == 2) // ~ IOR 1.5 for 64x64 LUT (this is still experimental)
+	{
+		OutLUT2D[Coord.xy]	= Out2D;
+	}
+	OutLUT3D[Coord]		= Out3D;
+	OutColor			= Out3D;
+}
+
+#endif // SHADER_LUT
+
+////////////////////////////////////////////////////////////////////////////////////////////////////////////
\ No newline at end of file
diff --git a/Engine/Shaders/Private/Strata/StrataMultipleScattering.ush b/Engine/Shaders/Private/Strata/StrataMultipleScattering.ush
new file mode 100644
--- /dev/null
+++ b/Engine/Shaders/Private/Strata/StrataMultipleScattering.ush
@@ -0,0 +1,95 @@
+// Copyright Epic Games, Inc. All Rights Reserved.
+
+#pragma once 
+
+#include "/Engine/Private/Common.ush"
+
+// Multiple scattering computation use a precomputed LUT. There are two variant: 
+//  1) Using 2D LUT using Schlick Fresnel decomposition (view angle, roughness)
+//  2) Using 3D LUT building a full 3D LUT (view angle, roughness, F0)
+// 
+// These two variants exist only for testing purpose at the moment
+
+struct FStrataLUTResult
+{
+	// Fresnel A/B term for reconstructing directional albedo for reflection
+	float A_R;
+	float B_R;
+
+	// Fresnel A/B term for reconstructing directional albedo for refraction
+	float A_T;
+	float B_T;
+
+	// Complete directional albedo for reflection and refraction
+	float E_R;
+	float E_T;
+};
+
+FStrataLUTResult StrataSampleGGXLUT(Texture2D<float4> LUT2D, Texture3D<float2> LUT3D, SamplerState LUTSampler, float2 LUTScaleBias, float NoV, float Roughness, float3 F0)
+{
+	FStrataLUTResult Out = (FStrataLUTResult)0;
+
+	// 2D LUT sampling
+	{
+		const float2 UV = float2(saturate(NoV), saturate(Roughness)) * LUTScaleBias.x + LUTScaleBias.yy;
+		const float4 S = LUT2D.SampleLevel(LUTSampler, UV, 0);
+		Out.A_R = S.x;
+		Out.B_R = S.y;
+
+		Out.A_T = S.z;
+		Out.B_T = S.w;
+	}
+
+	// 3D LUT sampling
+	{
+		const float3 UVW = float3(saturate(NoV), saturate(Roughness), saturate(F0.y)) * LUTScaleBias.x + LUTScaleBias.yyy;
+		const float2 S = LUT3D.SampleLevel(LUTSampler, UVW, 0);
+		Out.E_R = S.x;
+		Out.E_R = S.y;
+	}
+
+	return Out;
+}
+
+// From "Practical multiple scattering compensation for microfacet models." E. Turquin
+float3 StrataEvaluateMultipleScatteringConductor(float3 F0, FStrataLUTResult LUT)
+{
+	// Reconstruction of the directional Albedo for upper hemisphere
+	// E_R = F0.A_R + (1-F0).B_R
+	// 
+	// Microfacet multiple scattering term for *conductor*
+	//  f_ms(wo,wi) = F_ss . (1 +  F0 (1-E_R)/ E_R)
+	
+	const float3 E_R = F0 * LUT.A_R + (1 - F0) * LUT.B_R;
+
+	return (1 + F0 * (1 - E_R) / E_R);
+}
+
+float3 StrataEvaluateMultipleScatteringDielectric(float3 F0, FStrataLUTResult LUT)
+{
+	// Reconstruction of the directional Albedo for upper hemisphere
+	// E_R = F0.A_R + (1-F0).B_R
+	//
+	// Reconstruction of the directional Albedo for lower hemisphere
+	// E_T = A_T + F0.A_T + (1-F0).B_T
+	//
+	// Microfacet multiple scattering term for *dielietric*
+	//  f_ms(wo,wi) = F_ss / (E_T + E_R)
+
+	const float3 E_R = F0 * LUT.A_R + (1 - F0) * LUT.B_R;
+	const float3 E_T = LUT.A_T + F0 * LUT.A_T + (1 - F0) * LUT.B_T;
+
+	return 1.f / (E_R + E_T);
+}
+
+float3 StrataEvaluateMultipleScatteringConductor3D(float3 F0, FStrataLUTResult LUT)
+{
+	return (1 + F0 * (1 - LUT.E_R) / LUT.E_R);
+}
+
+float3 StrataEvaluateMultipleScatteringDielectric3D(float3 F0, FStrataLUTResult LUT)
+{
+	return 1.f / (LUT.E_R + LUT.E_T);
+}
+
+////////////////////////////////////////////////////////////////////////////////////////////////////////////
\ No newline at end of file
diff --git a/Engine/Source/Runtime/Renderer/Private/Strata/Strata.cpp b/Engine/Source/Runtime/Renderer/Private/Strata/Strata.cpp
--- a/Engine/Source/Runtime/Renderer/Private/Strata/Strata.cpp
+++ b/Engine/Source/Runtime/Renderer/Private/Strata/Strata.cpp
@@ -32,11 +32,54 @@ static TAutoConsoleVariable<int32> CVarStrataClassification(
 	TEXT("Enable strata classification to speed up lighting pass."),
 	ECVF_RenderThreadSafe); 
 
+static TAutoConsoleVariable<int32> CVarStrataLUTResolution(
+	TEXT("r.Strata.LUT.Resolution"),
+	64,
+	TEXT("Resolution of the GGX energy LUT."),
+	ECVF_RenderThreadSafe);
+
+static TAutoConsoleVariable<int32> CVarStrataLUTSampleCount(
+	TEXT("r.Strata.LUT.SampleCount"),
+	128,
+	TEXT("Number of sample used for computing the energy LUT."),
+	ECVF_RenderThreadSafe);
+
+static TAutoConsoleVariable<int32> CVarStrataLUTContinousUpdate(
+	TEXT("r.Strata.LUT.ContinousUpdate"),
+	0,
+	TEXT("Update Strata energy LUT every frame (for debug purpose)."),
+	ECVF_RenderThreadSafe);
+
+static TAutoConsoleVariable<int32> CVarStrataFurnaceTest(
+	TEXT("r.Strata.FurnaceTest"),
+	0,
+	TEXT("Enable Strata furnace test (for debug purpose) 1: roughness/metallic 2: a selection of conductors."),
+	ECVF_RenderThreadSafe);
+
+static TAutoConsoleVariable<int32> CVarStrataFurnaceTestSampleCount(
+	TEXT("r.Strata.FurnaceTest.SampleCount"),
+	1024,
+	TEXT("Number of sample used for furnace test."),
+	ECVF_RenderThreadSafe);
+
 IMPLEMENT_GLOBAL_SHADER_PARAMETER_STRUCT(FStrataGlobalUniformParameters, "Strata");
 
 namespace Strata
 {
 
+static uint32 GetStrataGGXEnergyLUTResolution()
+{
+	return FMath::RoundUpToPowerOfTwo(FMath::Clamp(CVarStrataLUTResolution.GetValueOnAnyThread(), 16, 256));
+}
+
+static FVector2D GetStrataGGXEnergyLUTScaleBias()
+{
+	const float Resolution = GetStrataGGXEnergyLUTResolution();
+	const float Scale = (Resolution - 1) / Resolution;
+	const float Bias  = 0.5f / (Resolution - 1);
+	return FVector2D(Scale, Bias);
+}
+
 bool IsStrataEnabled()
 {
 	return CVarStrata.GetValueOnAnyThread() > 0;
@@ -52,13 +95,16 @@ uint32 GetStrataTileSize()
 	return 8;
 }
 
+static void AddStrataLUTPass(FRDGBuilder& GraphBuilder, const FViewInfo& View);
+static void AddStrataFurnacePass(FRDGBuilder& GraphBuilder, const FViewInfo& View);
+
 void InitialiseStrataFrameSceneData(FSceneRenderer& SceneRenderer, FRDGBuilder& GraphBuilder)
 {
 	FStrataSceneData& StrataSceneData = SceneRenderer.Scene->StrataSceneData;
 
 	uint32 ResolutionX = 1;
 	uint32 ResolutionY = 1;
-
+	bool bUpdateLUT = false;
 	if (IsStrataEnabled())
 	{
 		FIntPoint BufferSizeXY = GetSceneTextureExtent();
@@ -101,6 +147,22 @@ void InitialiseStrataFrameSceneData(FSceneRenderer& SceneRenderer, FRDGBuilder&
 			GraphBuilder.PreallocateTexture(Texture);
 			StrataSceneData.TopLayerNormalTexture = GraphBuilder.GetPooledTexture(Texture);
 		}
+
+		// Energy LUT
+		const uint32 LUTResolution = GetStrataGGXEnergyLUTResolution();
+		bUpdateLUT = StrataSceneData.GGXEnergyLUT2DTexture == nullptr || StrataSceneData.GGXEnergyLUT2DTexture->GetDesc().Extent.X != LUTResolution || CVarStrataLUTContinousUpdate.GetValueOnAnyThread() > 0;
+		if (bUpdateLUT)
+		{
+			FRDGTextureDesc Desc3D = FRDGTextureDesc::Create3D(FIntVector(LUTResolution, LUTResolution, LUTResolution), EPixelFormat::PF_G16R16F, FClearValueBinding::Black, ETextureCreateFlags::TexCreate_ShaderResource | ETextureCreateFlags::TexCreate_RenderTargetable | ETextureCreateFlags::TexCreate_UAV);
+			FRDGTextureDesc Desc2D = FRDGTextureDesc::Create2D(FIntPoint(LUTResolution, LUTResolution), EPixelFormat::PF_FloatRGBA, FClearValueBinding::Black, ETextureCreateFlags::TexCreate_ShaderResource | ETextureCreateFlags::TexCreate_RenderTargetable | ETextureCreateFlags::TexCreate_UAV);
+			FRDGTextureRef OutLUT3D = GraphBuilder.CreateTexture(Desc3D, TEXT("StrataEnergyLUT3D"));
+			FRDGTextureRef OutLUT2D = GraphBuilder.CreateTexture(Desc2D, TEXT("StrataEnergyLUT2D"));
+
+			GraphBuilder.PreallocateTexture(OutLUT3D);
+			GraphBuilder.PreallocateTexture(OutLUT2D);
+			StrataSceneData.GGXEnergyLUT3DTexture = GraphBuilder.GetPooledTexture(OutLUT3D);
+			StrataSceneData.GGXEnergyLUT2DTexture = GraphBuilder.GetPooledTexture(OutLUT2D);
+		}
 	}
 	else
 	{
@@ -128,6 +190,24 @@ void InitialiseStrataFrameSceneData(FSceneRenderer& SceneRenderer, FRDGBuilder&
 	SceneRenderer.Scene->StrataSceneData.StrataGlobalUniformParameters.SafeRelease();
 
 	AddStrataClearMaterialBufferPass(GraphBuilder, StrataSceneData.MaterialLobesBuffer.UAV, StrataSceneData.MaxBytesPerPixel, FIntPoint(ResolutionX, ResolutionY));
+
+	const bool bFurnaceTest = CVarStrataFurnaceTest.GetValueOnAnyThread() > 0;
+	if (bUpdateLUT || bFurnaceTest)
+	{
+		for (int32 ViewIndex = 0; ViewIndex < SceneRenderer.Views.Num(); ViewIndex++)
+		{
+			FViewInfo& View = SceneRenderer.Views[ViewIndex];
+			if (bUpdateLUT)
+			{
+				AddStrataLUTPass(GraphBuilder, View);
+			}
+
+			if (bFurnaceTest)
+			{
+				AddStrataFurnacePass(GraphBuilder, View);
+			}
+		}
+	}
 }
 
 void BindStrataBasePassUniformParameters(const FViewInfo& View, FStrataOpaquePassUniformParameters& OutStrataUniformParameters)
@@ -159,7 +239,10 @@ TUniformBufferRef<FStrataGlobalUniformParameters> BindStrataGlobalUniformParamet
 		StrataUniformParameters.MaterialLobesBuffer = View.StrataSceneData->MaterialLobesBuffer.SRV;
 		StrataUniformParameters.ClassificationTexture = View.StrataSceneData->ClassificationTexture->GetRenderTargetItem().ShaderResourceTexture;
 		StrataUniformParameters.TopLayerNormalTexture = View.StrataSceneData->TopLayerNormalTexture->GetRenderTargetItem().ShaderResourceTexture;
-
+		StrataUniformParameters.GGXEnergyLUT3DTexture = View.StrataSceneData->GGXEnergyLUT3DTexture->GetRenderTargetItem().ShaderResourceTexture;
+		StrataUniformParameters.GGXEnergyLUT2DTexture = View.StrataSceneData->GGXEnergyLUT2DTexture->GetRenderTargetItem().ShaderResourceTexture;
+		StrataUniformParameters.GGXEnergyLUTSampler = TStaticSamplerState<SF_Bilinear, AM_Clamp, AM_Clamp, AM_Clamp>::GetRHI();
+		StrataUniformParameters.GGXEnergyLUTScaleBias = GetStrataGGXEnergyLUTScaleBias();
 		View.StrataSceneData->StrataGlobalUniformParameters = CreateUniformBufferImmediate(StrataUniformParameters, UniformBuffer_SingleFrame);
 		return View.StrataSceneData->StrataGlobalUniformParameters;
 	}
@@ -170,6 +253,10 @@ TUniformBufferRef<FStrataGlobalUniformParameters> BindStrataGlobalUniformParamet
 		StrataUniformParameters.MaterialLobesBuffer = GEmptyVertexBufferWithUAV->ShaderResourceViewRHI;
 		StrataUniformParameters.ClassificationTexture = GSystemTextures.BlackDummy->GetRenderTargetItem().ShaderResourceTexture;
 		StrataUniformParameters.TopLayerNormalTexture = GSystemTextures.BlackDummy->GetRenderTargetItem().ShaderResourceTexture;
+		StrataUniformParameters.GGXEnergyLUT3DTexture = GSystemTextures.BlackDummy->GetRenderTargetItem().ShaderResourceTexture;
+		StrataUniformParameters.GGXEnergyLUT2DTexture = GSystemTextures.BlackDummy->GetRenderTargetItem().ShaderResourceTexture;
+		StrataUniformParameters.GGXEnergyLUTSampler = TStaticSamplerState<SF_Bilinear, AM_Clamp, AM_Clamp, AM_Clamp>::GetRHI();
+		StrataUniformParameters.GGXEnergyLUTScaleBias = GetStrataGGXEnergyLUTScaleBias();
 		return CreateUniformBufferImmediate(StrataUniformParameters, UniformBuffer_SingleDraw);
 	}
 }
@@ -623,6 +710,145 @@ void AddStrataClearMaterialBufferPass(FRDGBuilder& GraphBuilder, FUnorderedAcces
 		FComputeShaderUtils::GetGroupCount(TiledViewBufferResolution, GroupSize));
 }
 
+////////////////////////////////////////////////////////////////////////////////////////////////////////////////////////////////////////////////
+
+class FStrataLUTPassPS : public FGlobalShader
+{
+	DECLARE_GLOBAL_SHADER(FStrataLUTPassPS);
+	SHADER_USE_PARAMETER_STRUCT(FStrataLUTPassPS, FGlobalShader);
+
+	using FPermutationDomain = TShaderPermutationDomain<>;
+
+	BEGIN_SHADER_PARAMETER_STRUCT(FParameters, )
+		SHADER_PARAMETER_STRUCT_REF(FViewUniformShaderParameters, ViewUniformBuffer)
+		SHADER_PARAMETER(FIntPoint, SliceXYCount)
+		SHADER_PARAMETER(uint32, EnergyLUTResolution)
+		SHADER_PARAMETER(uint32, NumSamples)
+		SHADER_PARAMETER_RDG_TEXTURE_UAV(Texture3D, OutLUT3D)
+		SHADER_PARAMETER_RDG_TEXTURE_UAV(Texture2D, OutLUT2D)
+		RENDER_TARGET_BINDING_SLOTS()
+	END_SHADER_PARAMETER_STRUCT()
+
+	static FPermutationDomain RemapPermutation(FPermutationDomain PermutationVector)
+	{
+		return PermutationVector;
+	}
+
+	static bool ShouldCompilePermutation(const FGlobalShaderPermutationParameters& Parameters)
+	{
+		return GetMaxSupportedFeatureLevel(Parameters.Platform) >= ERHIFeatureLevel::SM5 && Strata::IsStrataEnabled();
+	}
+
+	static void ModifyCompilationEnvironment(const FGlobalShaderPermutationParameters& Parameters, FShaderCompilerEnvironment& OutEnvironment)
+	{
+		FGlobalShader::ModifyCompilationEnvironment(Parameters, OutEnvironment);
+		OutEnvironment.SetDefine(TEXT("SHADER_LUT"), 1);
+	}
+};
+
+IMPLEMENT_GLOBAL_SHADER(FStrataLUTPassPS, "/Engine/Private/Strata/StrataLUT.usf", "MainPS", SF_Pixel);
+
+static void AddStrataLUTPass(FRDGBuilder& GraphBuilder, const FViewInfo& View)
+{
+	const uint32 LUTResolution = GetStrataGGXEnergyLUTResolution();
+	const uint32 SliceResolution = FMath::CeilToInt(FMath::Sqrt(float(LUTResolution)));
+	
+	FIntPoint OutputResolutionRT;
+	OutputResolutionRT.X = LUTResolution * SliceResolution;
+	OutputResolutionRT.Y = LUTResolution * SliceResolution;
+
+	FRDGTextureRef OutLUT3D = GraphBuilder.RegisterExternalTexture(View.StrataSceneData->GGXEnergyLUT3DTexture);
+	FRDGTextureRef OutLUT2D = GraphBuilder.RegisterExternalTexture(View.StrataSceneData->GGXEnergyLUT2DTexture);
+
+	// For debug purpose	
+	FRDGTextureDesc UnfoldLUTDesc = FRDGTextureDesc::Create2D(OutputResolutionRT, EPixelFormat::PF_G16R16F, FClearValueBinding::Black, ETextureCreateFlags::TexCreate_RenderTargetable);
+	FRDGTextureRef UnfoldLUTTexture = GraphBuilder.CreateTexture(UnfoldLUTDesc, TEXT("StrataEnergyUnfoldLUT"));
+
+	TShaderMapRef<FStrataLUTPassPS> PixelShader(View.ShaderMap);
+	FStrataLUTPassPS::FParameters* Parameters = GraphBuilder.AllocParameters<FStrataLUTPassPS::FParameters>();
+	Parameters->ViewUniformBuffer = View.ViewUniformBuffer;
+	Parameters->NumSamples = FMath::Clamp(CVarStrataLUTSampleCount.GetValueOnAnyThread(), 16, 2048);
+	Parameters->EnergyLUTResolution = LUTResolution;
+	Parameters->SliceXYCount = FIntPoint(SliceResolution, SliceResolution);
+	Parameters->OutLUT2D = GraphBuilder.CreateUAV(OutLUT2D);
+	Parameters->OutLUT3D = GraphBuilder.CreateUAV(OutLUT3D);
+	Parameters->RenderTargets[0] = FRenderTargetBinding(UnfoldLUTTexture, ERenderTargetLoadAction::EClear);
+
+	FPixelShaderUtils::AddFullscreenPass<FStrataLUTPassPS>(
+		GraphBuilder,
+		View.ShaderMap,
+		RDG_EVENT_NAME("StrataLUT"),
+		PixelShader,
+		Parameters,
+		View.ViewRect);
+}
+
+////////////////////////////////////////////////////////////////////////////////////////////////////////////////////////////////////////////////
+
+class FStrataFurnaceTestPassPS : public FGlobalShader
+{
+	DECLARE_GLOBAL_SHADER(FStrataFurnaceTestPassPS);
+	SHADER_USE_PARAMETER_STRUCT(FStrataFurnaceTestPassPS, FGlobalShader);
+
+	using FPermutationDomain = TShaderPermutationDomain<>;
+
+	BEGIN_SHADER_PARAMETER_STRUCT(FParameters, )
+		SHADER_PARAMETER_STRUCT_REF(FViewUniformShaderParameters, ViewUniformBuffer)
+		SHADER_PARAMETER_STRUCT_REF(FStrataGlobalUniformParameters, Strata)
+		SHADER_PARAMETER(FIntPoint, OutputResolution)
+		SHADER_PARAMETER(uint32, NumSamples)
+		SHADER_PARAMETER(uint32, SceneType)
+		SHADER_PARAMETER_RDG_TEXTURE_UAV(Texture3D, OutLUT3D)
+		SHADER_PARAMETER_RDG_TEXTURE_UAV(Texture2D, OutLUT2D)
+		RENDER_TARGET_BINDING_SLOTS()
+	END_SHADER_PARAMETER_STRUCT()
+
+		static FPermutationDomain RemapPermutation(FPermutationDomain PermutationVector)
+	{
+		return PermutationVector;
+	}
+
+	static bool ShouldCompilePermutation(const FGlobalShaderPermutationParameters& Parameters)
+	{
+		return GetMaxSupportedFeatureLevel(Parameters.Platform) >= ERHIFeatureLevel::SM5 && Strata::IsStrataEnabled();
+	}
+
+	static void ModifyCompilationEnvironment(const FGlobalShaderPermutationParameters& Parameters, FShaderCompilerEnvironment& OutEnvironment)
+	{
+		FGlobalShader::ModifyCompilationEnvironment(Parameters, OutEnvironment);
+		OutEnvironment.SetDefine(TEXT("SHADER_FURNACE_ANALYTIC"), 1);
+	}
+};
+
+IMPLEMENT_GLOBAL_SHADER(FStrataFurnaceTestPassPS, "/Engine/Private/Strata/StrataFurnaceTest.usf", "MainPS", SF_Pixel);
+
+static void AddStrataFurnacePass(FRDGBuilder& GraphBuilder, const FViewInfo& View)
+{
+	FIntPoint OutputResolution;
+	OutputResolution.X = 1024;
+	OutputResolution.Y = 1024;
+
+	FRDGTextureDesc Desc = FRDGTextureDesc::Create2D(OutputResolution, EPixelFormat::PF_FloatRGBA, FClearValueBinding::Black, ETextureCreateFlags::TexCreate_RenderTargetable);
+	FRDGTextureRef Texture = GraphBuilder.CreateTexture(Desc, TEXT("StrataFurnaceTest"));
+
+	TShaderMapRef<FStrataFurnaceTestPassPS> PixelShader(View.ShaderMap);
+	FStrataFurnaceTestPassPS::FParameters* Parameters = GraphBuilder.AllocParameters<FStrataFurnaceTestPassPS::FParameters>();
+	Parameters->ViewUniformBuffer = View.ViewUniformBuffer;
+	Parameters->Strata = Strata::BindStrataGlobalUniformParameters(View);
+	Parameters->SceneType = FMath::Clamp(CVarStrataFurnaceTest.GetValueOnAnyThread(), 1, 2);
+	Parameters->NumSamples = FMath::Clamp(CVarStrataFurnaceTestSampleCount.GetValueOnAnyThread(), 16, 2048);
+	Parameters->OutputResolution = OutputResolution;
+	Parameters->RenderTargets[0] = FRenderTargetBinding(Texture, ERenderTargetLoadAction::EClear);
+
+	FPixelShaderUtils::AddFullscreenPass<FStrataFurnaceTestPassPS>(
+		GraphBuilder,
+		View.ShaderMap,
+		RDG_EVENT_NAME("StrataFurnaceTest"),
+		PixelShader,
+		Parameters,
+		View.ViewRect);
+}
+
 } // namespace Strata
 
 
diff --git a/Engine/Source/Runtime/Renderer/Private/Strata/Strata.h b/Engine/Source/Runtime/Renderer/Private/Strata/Strata.h
--- a/Engine/Source/Runtime/Renderer/Private/Strata/Strata.h
+++ b/Engine/Source/Runtime/Renderer/Private/Strata/Strata.h
@@ -21,9 +21,13 @@ END_SHADER_PARAMETER_STRUCT()
 
 BEGIN_GLOBAL_SHADER_PARAMETER_STRUCT(FStrataGlobalUniformParameters, )
 	SHADER_PARAMETER(uint32, MaxBytesPerPixel)
+	SHADER_PARAMETER(FVector2D, GGXEnergyLUTScaleBias)
 	SHADER_PARAMETER_SRV(ByteAddressBuffer, MaterialLobesBuffer)
 	SHADER_PARAMETER_TEXTURE(Texture2D<uint>, ClassificationTexture)
 	SHADER_PARAMETER_TEXTURE(Texture2D<uint>, TopLayerNormalTexture)
+	SHADER_PARAMETER_TEXTURE(Texture3D<float2>, GGXEnergyLUT3DTexture)
+	SHADER_PARAMETER_TEXTURE(Texture2D<float4>, GGXEnergyLUT2DTexture)
+	SHADER_PARAMETER_SAMPLER(SamplerState, GGXEnergyLUTSampler)
 END_GLOBAL_SHADER_PARAMETER_STRUCT()
 
 struct FStrataSceneData
@@ -32,6 +36,8 @@ struct FStrataSceneData
 	FRWByteAddressBuffer MaterialLobesBuffer;					// This should be a RDG resource
 	TRefCountPtr<IPooledRenderTarget> ClassificationTexture;	// This should be a RDG resource
 	TRefCountPtr<IPooledRenderTarget> TopLayerNormalTexture;	// This should be a RDG resource
+	TRefCountPtr<IPooledRenderTarget> GGXEnergyLUT3DTexture;
+	TRefCountPtr<IPooledRenderTarget> GGXEnergyLUT2DTexture;
 
 	TRefCountPtr<FRDGPooledBuffer> ClassificationTileListBuffer;
 	TRefCountPtr<FRDGPooledBuffer> ClassificationTileIndirectBuffer;
