Fixed default ggxenergy 3D lut texture to use a volume texture.

diff --git a/Engine/Source/Runtime/Renderer/Private/Strata/Strata.cpp b/Engine/Source/Runtime/Renderer/Private/Strata/Strata.cpp
--- a/Engine/Source/Runtime/Renderer/Private/Strata/Strata.cpp
+++ b/Engine/Source/Runtime/Renderer/Private/Strata/Strata.cpp
@@ -262,7 +262,7 @@ TUniformBufferRef<FStrataGlobalUniformParameters> BindStrataGlobalUniformParamet
 		StrataUniformParameters.ClassificationTexture = GSystemTextures.ZeroUIntDummy->GetRenderTargetItem().ShaderResourceTexture;
 		StrataUniformParameters.TopLayerNormalTexture = GSystemTextures.ZeroUIntDummy->GetRenderTargetItem().ShaderResourceTexture;
 		StrataUniformParameters.SSSTexture = GSystemTextures.ZeroUIntDummy->GetRenderTargetItem().ShaderResourceTexture;
-		StrataUniformParameters.GGXEnergyLUT3DTexture = GSystemTextures.BlackDummy->GetRenderTargetItem().ShaderResourceTexture;
+		StrataUniformParameters.GGXEnergyLUT3DTexture = GSystemTextures.VolumetricBlackDummy->GetRenderTargetItem().ShaderResourceTexture;
 		StrataUniformParameters.GGXEnergyLUT2DTexture = GSystemTextures.BlackDummy->GetRenderTargetItem().ShaderResourceTexture;
 		StrataUniformParameters.GGXEnergyLUTSampler = TStaticSamplerState<SF_Bilinear, AM_Clamp, AM_Clamp, AM_Clamp>::GetRHI();
 		StrataUniformParameters.GGXEnergyLUTScaleBias = GetStrataGGXEnergyLUTScaleBias();
