OnePassPointShadowProj uses Strata global uniform buffer instead of isolated textures.

diff --git a/Engine/Source/Runtime/Renderer/Private/ShadowRendering.h b/Engine/Source/Runtime/Renderer/Private/ShadowRendering.h
--- a/Engine/Source/Runtime/Renderer/Private/ShadowRendering.h
+++ b/Engine/Source/Runtime/Renderer/Private/ShadowRendering.h
@@ -1464,6 +1464,7 @@ public:
 	TOnePassPointShadowProjectionPS(const ShaderMetaType::CompiledShaderInitializerType& Initializer):
 		FGlobalShader(Initializer)
 	{
+		StrataGlobalParameters.Bind(Initializer.ParameterMap, FStrataGlobalUniformParameters::StaticStructMetadata.GetShaderVariableName());
 		OnePassShadowParameters.Bind(Initializer.ParameterMap);
 		ShadowDepthTextureSampler.Bind(Initializer.ParameterMap,TEXT("ShadowDepthTextureSampler"));
 		LightPosition.Bind(Initializer.ParameterMap,TEXT("LightPositionAndInvRadius"));
@@ -1473,9 +1474,6 @@ public:
 		PointLightProjParameters.Bind(Initializer.ParameterMap, TEXT("PointLightProjParameters"));
 		TransmissionProfilesTexture.Bind(Initializer.ParameterMap, TEXT("SSProfilesTexture"));
 		HairCategorizationTexture.Bind(Initializer.ParameterMap, TEXT("HairCategorizationTexture"));
-		ClassificationTexture.Bind(Initializer.ParameterMap, TEXT("StrataClassificationTexture"));
-		TopLayerNormalTexture.Bind(Initializer.ParameterMap, TEXT("StrataTopLayerNormalTexture"));
-		SSSTexture.Bind(Initializer.ParameterMap, TEXT("StrataSSSTexture"));
 	}
 
 	static void ModifyCompilationEnvironment(const FGlobalShaderPermutationParameters& Parameters, FShaderCompilerEnvironment& OutEnvironment)
@@ -1522,11 +1520,10 @@ public:
 			SetTextureParameter(RHICmdList, ShaderRHI, HairCategorizationTexture, HairVisibilityData->CategorizationTexture->GetPooledRenderTarget()->GetRenderTargetItem().ShaderResourceTexture);
 		}
 
-		if (Strata::IsStrataEnabled())
+		if (StrataGlobalParameters.IsBound())
 		{
-			SetTextureParameter(RHICmdList, ShaderRHI, ClassificationTexture, Strata::GetClassificationTexture(View));
-			SetTextureParameter(RHICmdList, ShaderRHI, TopLayerNormalTexture, Strata::GetTopLayerNormalTexture(View));
-			SetTextureParameter(RHICmdList, ShaderRHI, SSSTexture, Strata::GetSSSTexture(View));
+			TUniformBufferRef<FStrataGlobalUniformParameters> StrataGlobalUniformParameters = Strata::BindStrataGlobalUniformParameters(View);
+			SetUniformBufferParameter(RHICmdList, ShaderRHI, StrataGlobalParameters, StrataGlobalUniformParameters);
 		}
 
 		{
@@ -1570,9 +1567,7 @@ private:
 	LAYOUT_FIELD(FShaderParameter, PointLightProjParameters);
 	LAYOUT_FIELD(FShaderResourceParameter, TransmissionProfilesTexture);
 	LAYOUT_FIELD(FShaderResourceParameter, HairCategorizationTexture);
-	LAYOUT_FIELD(FShaderResourceParameter, ClassificationTexture);
-	LAYOUT_FIELD(FShaderResourceParameter, TopLayerNormalTexture);
-	LAYOUT_FIELD(FShaderResourceParameter, SSSTexture);
+	LAYOUT_FIELD(FShaderUniformBufferParameter, StrataGlobalParameters);
 };
 
 // Reversed Z
diff --git a/Engine/Source/Runtime/Renderer/Private/Strata/Strata.cpp b/Engine/Source/Runtime/Renderer/Private/Strata/Strata.cpp
--- a/Engine/Source/Runtime/Renderer/Private/Strata/Strata.cpp
+++ b/Engine/Source/Runtime/Renderer/Private/Strata/Strata.cpp
@@ -270,31 +270,6 @@ TUniformBufferRef<FStrataGlobalUniformParameters> BindStrataGlobalUniformParamet
 	}
 }
 
-FTextureRHIRef GetClassificationTexture(const FViewInfo& View)
-{
-	if (View.StrataSceneData)
-	{
-		return View.StrataSceneData->ClassificationTexture->GetRenderTargetItem().ShaderResourceTexture;
-	}
-	return GSystemTextures.ZeroUIntDummy->GetRenderTargetItem().ShaderResourceTexture;
-}
-FTextureRHIRef GetTopLayerNormalTexture(const FViewInfo& View)
-{
-	if (View.StrataSceneData)
-	{
-		return View.StrataSceneData->TopLayerNormalTexture->GetRenderTargetItem().ShaderResourceTexture;
-	}
-	return GSystemTextures.ZeroUIntDummy->GetRenderTargetItem().ShaderResourceTexture;
-}
-FTextureRHIRef GetSSSTexture(const FViewInfo& View)
-{
-	if (View.StrataSceneData)
-	{
-		return View.StrataSceneData->SSSTexture->GetRenderTargetItem().ShaderResourceTexture;
-	}
-	return GSystemTextures.ZeroUIntDummy->GetRenderTargetItem().ShaderResourceTexture;
-}
-
 ////////////////////////////////////////////////////////////////////////// 
 // Debug
 
diff --git a/Engine/Source/Runtime/Renderer/Private/Strata/Strata.h b/Engine/Source/Runtime/Renderer/Private/Strata/Strata.h
--- a/Engine/Source/Runtime/Renderer/Private/Strata/Strata.h
+++ b/Engine/Source/Runtime/Renderer/Private/Strata/Strata.h
@@ -25,10 +25,10 @@ BEGIN_GLOBAL_SHADER_PARAMETER_STRUCT(FStrataGlobalUniformParameters, )
 	SHADER_PARAMETER_SRV(ByteAddressBuffer, MaterialLobesBuffer)
 	SHADER_PARAMETER_TEXTURE(Texture2D<uint>, ClassificationTexture)
 	SHADER_PARAMETER_TEXTURE(Texture2D<uint>, TopLayerNormalTexture)
+	SHADER_PARAMETER_TEXTURE(Texture2D<uint2>, SSSTexture)
 	SHADER_PARAMETER_TEXTURE(Texture3D<float2>, GGXEnergyLUT3DTexture)
 	SHADER_PARAMETER_TEXTURE(Texture2D<float4>, GGXEnergyLUT2DTexture)
 	SHADER_PARAMETER_SAMPLER(SamplerState, GGXEnergyLUTSampler)
-	SHADER_PARAMETER_TEXTURE(Texture2D<uint2>, SSSTexture)
 END_GLOBAL_SHADER_PARAMETER_STRUCT()
 
 struct FStrataSceneData
@@ -66,10 +66,6 @@ void BindStrataBasePassUniformParameters(const FViewInfo& View, FStrataOpaquePas
 
 TUniformBufferRef<FStrataGlobalUniformParameters> BindStrataGlobalUniformParameters(const FViewInfo& View);
 
-FTextureRHIRef GetClassificationTexture(const FViewInfo& View);
-FTextureRHIRef GetTopLayerNormalTexture(const FViewInfo& View);
-FTextureRHIRef GetSSSTexture(const FViewInfo& View);
-
 void AddStrataMaterialClassificationPass(FRDGBuilder& GraphBuilder, const FMinimalSceneTextures& SceneTextures, const TArray<FViewInfo>& Views);
 
 void AddStrataStencilPass(FRDGBuilder& GraphBuilder, const FViewInfo& View, const FMinimalSceneTextures& SceneTextures);
