Downgrade a Virtual Texture check() into a warning, add more info

diff --git a/Engine/Source/Runtime/Renderer/Private/VT/AllocatedVirtualTexture.cpp b/Engine/Source/Runtime/Renderer/Private/VT/AllocatedVirtualTexture.cpp
--- a/Engine/Source/Runtime/Renderer/Private/VT/AllocatedVirtualTexture.cpp
+++ b/Engine/Source/Runtime/Renderer/Private/VT/AllocatedVirtualTexture.cpp
@@ -199,11 +199,32 @@ void FAllocatedVirtualTexture::Release(FVirtualTextureSystem* System)
 			UniquePageTableLayers[PageTableIndex].PhysicalSpace.SafeRelease();
 		}
 
+#if DO_CHECK
 		for (uint32 LayerIndex = 0; LayerIndex < Space->GetNumPageTableLayers(); ++LayerIndex)
 		{
 			const FTexturePageMap& PageMap = Space->GetPageMapForPageTableLayer(LayerIndex);
-			PageMap.VerifyAddressRangeUnmapped(VirtualAddress, BlockSize);
+
+			TArray<FMappedTexturePage> MappedPages;
+			PageMap.GetMappedPagesInRange(VirtualAddress, BlockSize, MappedPages);
+			if (MappedPages.Num() > 0)
+			{
+				TStringBuilder<2048> Message;
+				Message.Appendf(TEXT("Mapped pages remain after releasing AllocatedVT - vAddress: %d, Size: %d, PhysicalSpaces: ["), VirtualAddress, BlockSize);
+				for (FVirtualTexturePhysicalSpace* PhysicalSpace : UniquePhysicalSpaces)
+				{
+					Message.Appendf(TEXT("%d "), PhysicalSpace->GetID());
+				}
+				Message.Appendf(TEXT("], MappedPages: ["));
+
+				for (const FMappedTexturePage& MappedPage : MappedPages)
+				{
+					Message.Appendf(TEXT("(vAddress: %d, PhysicalSpace: %d) "), MappedPage.Page.vAddress, MappedPage.PhysicalSpaceID);
+				}
+				Message.Appendf(TEXT("]"));
+				UE_LOG(LogVirtualTexturing, Warning, TEXT("%s"), Message.ToString());
+			}
 		}
+#endif // DO_CHECK
 	}
 
 	Space->FreeVirtualTexture(this);
diff --git a/Engine/Source/Runtime/Renderer/Private/VT/TexturePageMap.cpp b/Engine/Source/Runtime/Renderer/Private/VT/TexturePageMap.cpp
--- a/Engine/Source/Runtime/Renderer/Private/VT/TexturePageMap.cpp
+++ b/Engine/Source/Runtime/Renderer/Private/VT/TexturePageMap.cpp
@@ -203,7 +203,7 @@ void FTexturePageMap::MapPage(FVirtualTextureSpace* Space, FVirtualTexturePhysic
 	SortedKeysDirty = true;
 }
 
-void FTexturePageMap::VerifyAddressRangeUnmapped(uint32 vAddress, uint32 Size) const
+void FTexturePageMap::GetMappedPagesInRange(uint32 vAddress, uint32 Size, TArray<FMappedTexturePage>& OutMappedPages) const
 {
 	const uint32 vAddressMax = vAddress + FMath::Square(Size);
 
@@ -214,7 +214,11 @@ void FTexturePageMap::VerifyAddressRangeUnmapped(uint32 vAddress, uint32 Size) c
 		const FPageEntry& Entry = Pages[PageIndex];
 		if (Entry.Page.vAddress >= vAddress && Entry.Page.vAddress < vAddressMax)
 		{
-			check(false);
+			FMappedTexturePage& OutPage = OutMappedPages.AddDefaulted_GetRef();
+			OutPage.Page = Entry.Page;
+			OutPage.pAddress = Entry.pAddress;
+			OutPage.PhysicalSpaceID = Entry.PhysicalSpaceID;
+			OutPage.vLevel = Entry.vLevel;
 		}
 
 		PageIndex = Entry.NextIndex;
diff --git a/Engine/Source/Runtime/Renderer/Private/VT/TexturePageMap.h b/Engine/Source/Runtime/Renderer/Private/VT/TexturePageMap.h
--- a/Engine/Source/Runtime/Renderer/Private/VT/TexturePageMap.h
+++ b/Engine/Source/Runtime/Renderer/Private/VT/TexturePageMap.h
@@ -42,6 +42,14 @@ union FPhysicalSpaceIDAndAddress
 	};
 };
 
+struct FMappedTexturePage
+{
+	FTexturePage Page;
+	uint32 pAddress : 16;
+	uint32 PhysicalSpaceID : 12;
+	uint32 vLevel : 4;
+};
+
 /**
  * Manages a single layer of a VT page table, contains mappings of virtual->physical address
  * Pages should not be directly mapped/unmapped from this class, this should instead go through FTexturePagePool
@@ -90,7 +98,7 @@ public:
 	*/
 	void		MapPage(FVirtualTextureSpace* Space, FVirtualTexturePhysicalSpace* PhysicalSpace, uint8 vLogSize, uint32 vAddress, uint8 vLevel, uint16 pAddress);
 
-	void		VerifyAddressRangeUnmapped(uint32 vAddress, uint32 Size) const;
+	void		GetMappedPagesInRange(uint32 vAddress, uint32 Size, TArray<FMappedTexturePage>& OutMappedPages) const;
 
 	void		RefreshEntirePageTable(FVirtualTextureSystem* System, TArray< FPageTableUpdate >* Output);
 	void		ExpandPageTableUpdatePainters(FVirtualTextureSystem* System, FPageTableUpdate Update, TArray< FPageTableUpdate >* Output);
