Fix crash when running MH sample in UE5.

Crash happens when data are still loading, and initialization is not fully completed.

diff --git a/Engine/Source/Runtime/Renderer/Private/HairStrands/HairStrandsVoxelization.cpp b/Engine/Source/Runtime/Renderer/Private/HairStrands/HairStrandsVoxelization.cpp
--- a/Engine/Source/Runtime/Renderer/Private/HairStrands/HairStrandsVoxelization.cpp
+++ b/Engine/Source/Runtime/Renderer/Private/HairStrands/HairStrandsVoxelization.cpp
@@ -1049,6 +1049,10 @@ static void AddVirtualVoxelizationComputeRasterPass(
 			if (HairGroupPublicData->DoesSupportVoxelization())
 			{
 				const FHairGroupPublicData::FVertexFactoryInput& VFInput = HairGroupPublicData->VFInput;
+				if (HairGroupPublicData->VFInput.Strands.PositionBuffer == nullptr)
+				{
+					continue;
+				}
 
 				FVoxelRasterComputeCS::FParameters* PassParameters = GraphBuilder.AllocParameters<FVoxelRasterComputeCS::FParameters>();
 				PassParameters->MaxRasterCount = FMath::Clamp(GHairStrandsVoxelComputeRasterMaxVoxelCount, 1, 256);
