MaterialBaking - Support materials with VTs
* Added required FVirtualTextureSystem calls to FRendererModule::DrawTileMesh()

diff --git a/Engine/Source/Runtime/Renderer/Private/Renderer.cpp b/Engine/Source/Runtime/Renderer/Private/Renderer.cpp
--- a/Engine/Source/Runtime/Renderer/Private/Renderer.cpp
+++ b/Engine/Source/Runtime/Renderer/Private/Renderer.cpp
@@ -197,6 +197,15 @@ void FRendererModule::DrawTileMesh(FCanvasRenderContext& RenderContext, FMeshPas
 
 		FRDGBuilder& GraphBuilder = RenderContext.GraphBuilder;
 
+		const bool bUseVirtualTexturing = UseVirtualTexturing(FeatureLevel);
+		if (bUseVirtualTexturing)
+		{
+			RDG_GPU_STAT_SCOPE(GraphBuilder, VirtualTextureUpdate);
+			FVirtualTextureSystem::Get().AllocateResources(GraphBuilder, FeatureLevel);
+			FVirtualTextureSystem::Get().CallPendingCallbacks();
+			FVirtualTextureSystem::Get().Update(GraphBuilder, FeatureLevel, Scene);
+		}
+
 		if (!FRDGSystemTextures::IsValid(GraphBuilder))
 		{
 			FRDGSystemTextures::Create(GraphBuilder);
