Fix RDG validation error.

diff --git a/Engine/Source/Runtime/Renderer/Private/Strata/Strata.cpp b/Engine/Source/Runtime/Renderer/Private/Strata/Strata.cpp
--- a/Engine/Source/Runtime/Renderer/Private/Strata/Strata.cpp
+++ b/Engine/Source/Runtime/Renderer/Private/Strata/Strata.cpp
@@ -735,8 +735,8 @@ class FStrataLUTPassPS : public FGlobalShader
 		SHADER_PARAMETER(FIntPoint, SliceXYCount)
 		SHADER_PARAMETER(uint32, EnergyLUTResolution)
 		SHADER_PARAMETER(uint32, NumSamples)
-		SHADER_PARAMETER_RDG_TEXTURE_UAV(Texture3D, OutLUT3D)
-		SHADER_PARAMETER_RDG_TEXTURE_UAV(Texture2D, OutLUT2D)
+		SHADER_PARAMETER_RDG_TEXTURE_UAV(RWTexture3D, OutLUT3D)
+		SHADER_PARAMETER_RDG_TEXTURE_UAV(RWTexture2D, OutLUT2D)
 		RENDER_TARGET_BINDING_SLOTS()
 	END_SHADER_PARAMETER_STRUCT()
 
@@ -809,8 +809,8 @@ class FStrataFurnaceTestPassPS : public FGlobalShader
 		SHADER_PARAMETER(FIntPoint, OutputResolution)
 		SHADER_PARAMETER(uint32, NumSamples)
 		SHADER_PARAMETER(uint32, SceneType)
-		SHADER_PARAMETER_RDG_TEXTURE_UAV(Texture3D, OutLUT3D)
-		SHADER_PARAMETER_RDG_TEXTURE_UAV(Texture2D, OutLUT2D)
+		SHADER_PARAMETER_RDG_TEXTURE(Texture3D, OutLUT3D)
+		SHADER_PARAMETER_RDG_TEXTURE(Texture2D, OutLUT2D)
 		RENDER_TARGET_BINDING_SLOTS()
 	END_SHADER_PARAMETER_STRUCT()
 
