Fixed "uninitialized class member 'XXXX' is used to initialize the 'DynamicPassMeshDrawListContext' member" errors in MeshPassProcessor

diff --git a/Engine/Source/Runtime/Renderer/Public/MeshPassProcessor.h b/Engine/Source/Runtime/Renderer/Public/MeshPassProcessor.h
--- a/Engine/Source/Runtime/Renderer/Public/MeshPassProcessor.h
+++ b/Engine/Source/Runtime/Renderer/Public/MeshPassProcessor.h
@@ -1993,16 +1993,17 @@ public:
 	FDynamicPassMeshDrawListContext* GetDynamicPassMeshDrawListContext() { return &DynamicPassMeshDrawListContext; }
 
 private:
-	bool bNeedsInitialization = false;
-	FDynamicPassMeshDrawListContext DynamicPassMeshDrawListContext;
 	FMeshCommandOneFrameArray VisibleMeshDrawCommands;
 	FInstanceCullingContext InstanceCullingContext;
-	FDynamicMeshDrawCommandStorage DynamicMeshDrawCommandStorage;
 	FGraphicsMinimalPipelineStateSet GraphicsMinimalPipelineStateSet;
+	FDynamicMeshDrawCommandStorage DynamicMeshDrawCommandStorage;
+	FDynamicPassMeshDrawListContext DynamicPassMeshDrawListContext;
 
 	// Is set to true if and only if the BuildRenderingCommands has been called with an enabled GPU scene (which implies a valid Scene etc).
 	// Is used to check that we don't submit any draw commands that require a GPU scene without supplying one.
 	bool bSupportsScenePrimitives = false;
+
+	bool bNeedsInitialization = false;
 };
 
 // GPUCULL_TODO: Write documentation
