Add Strata classification visualization.

diff --git a/Engine/Source/Runtime/Renderer/Private/DeferredShadingRenderer.cpp b/Engine/Source/Runtime/Renderer/Private/DeferredShadingRenderer.cpp
--- a/Engine/Source/Runtime/Renderer/Private/DeferredShadingRenderer.cpp
+++ b/Engine/Source/Runtime/Renderer/Private/DeferredShadingRenderer.cpp
@@ -2705,6 +2705,7 @@ void FDeferredShadingSceneRenderer::Render(FRDGBuilder& GraphBuilder)
 	if (bStrataEnabled)
 	{
 		Strata::AddVisualizeMaterialPasses(GraphBuilder, Views, SceneTextures.Color.Target, Scene->GetShaderPlatform());
+		Strata::AddVisualizeClassificationPass(GraphBuilder, Views, SceneTextures.Color.Target);
 	}
 
 	if (bCanOverlayRayTracingOutput && ViewFamily.EngineShowFlags.LightShafts)
diff --git a/Engine/Source/Runtime/Renderer/Private/Strata/Strata.cpp b/Engine/Source/Runtime/Renderer/Private/Strata/Strata.cpp
--- a/Engine/Source/Runtime/Renderer/Private/Strata/Strata.cpp
+++ b/Engine/Source/Runtime/Renderer/Private/Strata/Strata.cpp
@@ -32,6 +32,12 @@ static TAutoConsoleVariable<int32> CVarStrataClassification(
 	TEXT("Enable strata classification to speed up lighting pass."),
 	ECVF_RenderThreadSafe); 
 
+static TAutoConsoleVariable<int32> CVarStrataClassificationDebug(
+	TEXT("r.Strata.Classification.Debug"),
+	0,
+	TEXT("Enable strata classification visualization."),
+	ECVF_RenderThreadSafe);
+
 static TAutoConsoleVariable<int32> CVarStrataLUTResolution(
 	TEXT("r.Strata.LUT.Resolution"),
 	64,
@@ -541,10 +547,11 @@ class FStrataMaterialStencilClassificationPassPS : public FGlobalShader
 IMPLEMENT_GLOBAL_SHADER(FStrataMaterialStencilClassificationPassVS, "/Engine/Private/Strata/StrataMaterialClassification.usf", "StencilMainVS", SF_Vertex);
 IMPLEMENT_GLOBAL_SHADER(FStrataMaterialStencilClassificationPassPS, "/Engine/Private/Strata/StrataMaterialClassification.usf", "StencilMainPS", SF_Pixel);
 
-static void AddStrataStencilPass(
+static void AddStrataInternalClassifedTilePass(
 	FRDGBuilder& GraphBuilder,
 	const FViewInfo& View,
-	const FMinimalSceneTextures& SceneTextures,
+	const FRDGTextureRef* DepthTexture,
+	const FRDGTextureRef* ColorTexture,
 	FRDGBufferRef TileListBuffer,
 	FRDGBufferRef TileIndirectBuffer)
 {	
@@ -566,23 +573,28 @@ static void AddStrataStencilPass(
 	TShaderMapRef<FStrataMaterialStencilClassificationPassPS> PixelShader(View.ShaderMap);
 
 	// For debug purpose
-	#if 0
-	FRDGTextureDesc DummyDesc = FRDGTextureDesc::Create2D(SceneTextures.Depth.Target->Desc.Extent, EPixelFormat::PF_R8G8B8A8, FClearValueBinding::Black, ETextureCreateFlags::TexCreate_RenderTargetable);
-	FRDGTextureRef DummyTexture = GraphBuilder.CreateTexture(DummyDesc, TEXT("StencilClassificationOutput"));
-	ParametersPS->RenderTargets[0] = FRenderTargetBinding(DummyTexture, ERenderTargetLoadAction::EClear);
-	#endif
-
-	ParametersPS->RenderTargets.DepthStencil = FDepthStencilBinding(
-		SceneTextures.Depth.Target,
-		ERenderTargetLoadAction::ELoad,
-		ERenderTargetLoadAction::ELoad,
-		FExclusiveDepthStencil::DepthNop_StencilWrite);
-
+	const bool bDebug = ColorTexture != nullptr;
+	if (bDebug)
+	{
+		ParametersPS->RenderTargets[0] = FRenderTargetBinding(*ColorTexture, ERenderTargetLoadAction::ELoad);
+	}
+	else
+	{
+		check(DepthTexture);
+		ParametersPS->RenderTargets.DepthStencil = FDepthStencilBinding(
+			*DepthTexture,
+			ERenderTargetLoadAction::ELoad,
+			ERenderTargetLoadAction::ELoad,
+			FExclusiveDepthStencil::DepthNop_StencilWrite);
+	}
+	
 	GraphBuilder.AddPass(
+		bDebug ? 
+		RDG_EVENT_NAME("StrataDebugClassificationPass") :
 		RDG_EVENT_NAME("StrataStencilClassificationPass"),
 		ParametersPS,
 		ERDGPassFlags::Raster,
-		[ParametersPS, VertexShader, PixelShader, OutputResolution, InstanceCount](FRHICommandList& RHICmdList)
+		[ParametersPS, VertexShader, PixelShader, OutputResolution, InstanceCount, bDebug](FRHICommandList& RHICmdList)
 		{
 			FStrataMaterialStencilClassificationPassVS::FParameters ParametersVS;
 			ParametersVS.TileSize = ParametersPS->TileSize;
@@ -595,11 +607,18 @@ static void AddStrataStencilPass(
 			RHICmdList.ApplyCachedRenderTargets(GraphicsPSOInit);
 			GraphicsPSOInit.BlendState = TStaticBlendState<CW_RGBA, BO_Add, BF_One, BF_One, BO_Max, BF_SourceAlpha, BF_DestAlpha>::GetRHI();
 			GraphicsPSOInit.RasterizerState = TStaticRasterizerState<>::GetRHI();
-			GraphicsPSOInit.DepthStencilState = TStaticDepthStencilState<
-				false, CF_Always, 
-				true,  CF_Always, SO_Keep, SO_Keep, SO_Replace,
-				false, CF_Always, SO_Keep, SO_Keep, SO_Keep,
-				0xFF, StencilBit>::GetRHI();
+			if (bDebug)
+			{
+				GraphicsPSOInit.DepthStencilState = TStaticDepthStencilState<false, CF_Always>::GetRHI();
+			}
+			else
+			{
+				GraphicsPSOInit.DepthStencilState = TStaticDepthStencilState<
+					false, CF_Always, 
+					true,  CF_Always, SO_Keep, SO_Keep, SO_Replace,
+					false, CF_Always, SO_Keep, SO_Keep, SO_Keep,
+					0xFF, StencilBit>::GetRHI();
+			}
 			GraphicsPSOInit.BoundShaderState.VertexDeclarationRHI = GFilterVertexDeclaration.VertexDeclarationRHI;
 			GraphicsPSOInit.BoundShaderState.VertexShaderRHI = VertexShader.GetVertexShader();
 			GraphicsPSOInit.BoundShaderState.PixelShaderRHI = PixelShader.GetPixelShader();
@@ -622,7 +641,7 @@ void AddStrataStencilPass(
 {
 	FRDGBufferRef TileListBuffer = GraphBuilder.RegisterExternalBuffer(View.StrataSceneData->ClassificationTileListBuffer);
 	FRDGBufferRef TileIndirectBuffer = GraphBuilder.RegisterExternalBuffer(View.StrataSceneData->ClassificationTileIndirectBuffer);
-	AddStrataStencilPass(GraphBuilder, View, SceneTextures, TileListBuffer, TileIndirectBuffer);
+	AddStrataInternalClassifedTilePass(GraphBuilder, View, &SceneTextures.Depth.Target, nullptr, TileListBuffer, TileIndirectBuffer);
 }
 
 void AddStrataStencilPass(
@@ -635,7 +654,24 @@ void AddStrataStencilPass(
 		const FViewInfo& View = Views[i];
 		FRDGBufferRef TileListBuffer = GraphBuilder.RegisterExternalBuffer(View.StrataSceneData->ClassificationTileListBuffer);
 		FRDGBufferRef TileIndirectBuffer = GraphBuilder.RegisterExternalBuffer(View.StrataSceneData->ClassificationTileIndirectBuffer);
-		AddStrataStencilPass(GraphBuilder, View, SceneTextures, TileListBuffer, TileIndirectBuffer);
+		AddStrataInternalClassifedTilePass(GraphBuilder, View, &SceneTextures.Depth.Target, nullptr, TileListBuffer, TileIndirectBuffer);
+	}
+}
+
+void AddVisualizeClassificationPass(
+	FRDGBuilder& GraphBuilder, 
+	const TArray<FViewInfo>& Views, 
+	FRDGTextureRef SceneColorTexture)
+{
+	if (IsClassificationEnabled() && CVarStrataClassificationDebug.GetValueOnAnyThread() > 0)
+	{
+		for (int32 i = 0; i < Views.Num(); ++i)
+		{
+			const FViewInfo& View = Views[i];
+			FRDGBufferRef TileListBuffer = GraphBuilder.RegisterExternalBuffer(View.StrataSceneData->ClassificationTileListBuffer);
+			FRDGBufferRef TileIndirectBuffer = GraphBuilder.RegisterExternalBuffer(View.StrataSceneData->ClassificationTileIndirectBuffer);
+			AddStrataInternalClassifedTilePass(GraphBuilder, View, nullptr, &SceneColorTexture, TileListBuffer, TileIndirectBuffer);
+		}
 	}
 }
 
diff --git a/Engine/Source/Runtime/Renderer/Private/Strata/Strata.h b/Engine/Source/Runtime/Renderer/Private/Strata/Strata.h
--- a/Engine/Source/Runtime/Renderer/Private/Strata/Strata.h
+++ b/Engine/Source/Runtime/Renderer/Private/Strata/Strata.h
@@ -72,6 +72,7 @@ void AddStrataStencilPass(FRDGBuilder& GraphBuilder, const FViewInfo& View, cons
 void AddStrataStencilPass(FRDGBuilder& GraphBuilder, const TArray<FViewInfo>& Views, const FMinimalSceneTextures& SceneTextures);
 
 void AddVisualizeMaterialPasses(FRDGBuilder& GraphBuilder, const TArray<FViewInfo>& Views, FRDGTextureRef SceneColorTexture, EShaderPlatform Platform);
+void AddVisualizeClassificationPass(FRDGBuilder& GraphBuilder, const TArray<FViewInfo>& Views, FRDGTextureRef SceneColorTexture);
 };
 
 
