Move all strata debug passes into a single function.

diff --git a/Engine/Shaders/Private/Strata/StrataFurnaceTest.usf b/Engine/Shaders/Private/Strata/StrataFurnaceTest.usf
--- a/Engine/Shaders/Private/Strata/StrataFurnaceTest.usf
+++ b/Engine/Shaders/Private/Strata/StrataFurnaceTest.usf
@@ -12,7 +12,6 @@
 
 uint SceneType;
 uint NumSamples;
-int2 OutputResolution;
 
 void MainPS(
 	float4 SVPos : SV_POSITION,
@@ -21,11 +20,12 @@ void MainPS(
 	const float3 InRadiance = 0.5f;
 	OutColor = InRadiance;
 
-	const uint2 TileCount = 8;
-	const uint2 TileResolution = OutputResolution / TileCount;
+	const uint2 ViewResolution = View.ViewSizeAndInvSize.xy;
+	const uint TileResolution = 128;
+	const uint2 TileCount = ViewResolution / TileResolution;
 
 	const uint2 PixelPos = uint2(SVPos.xy);	
-	const float2 UV = PixelPos / float2(OutputResolution);
+	const float2 UV = PixelPos / float2(ViewResolution);
 	const float2 TileUV = frac(UV * TileCount);
 	const uint2 TileCoord = floor(UV * TileCount);
 	const uint2 TilePixelCoord = TileUV * TileResolution;
diff --git a/Engine/Source/Runtime/Renderer/Private/DeferredShadingRenderer.cpp b/Engine/Source/Runtime/Renderer/Private/DeferredShadingRenderer.cpp
--- a/Engine/Source/Runtime/Renderer/Private/DeferredShadingRenderer.cpp
+++ b/Engine/Source/Runtime/Renderer/Private/DeferredShadingRenderer.cpp
@@ -2704,8 +2704,7 @@ void FDeferredShadingSceneRenderer::Render(FRDGBuilder& GraphBuilder)
 
 	if (bStrataEnabled)
 	{
-		Strata::AddVisualizeMaterialPasses(GraphBuilder, Views, SceneTextures.Color.Target, Scene->GetShaderPlatform());
-		Strata::AddVisualizeClassificationPass(GraphBuilder, Views, SceneTextures.Color.Target);
+		Strata::AddStrataDebugPasses(GraphBuilder, Views, SceneTextures.Color.Target, Scene->GetShaderPlatform());
 	}
 
 	if (bCanOverlayRayTracingOutput && ViewFamily.EngineShowFlags.LightShafts)
diff --git a/Engine/Source/Runtime/Renderer/Private/Strata/Strata.cpp b/Engine/Source/Runtime/Renderer/Private/Strata/Strata.cpp
--- a/Engine/Source/Runtime/Renderer/Private/Strata/Strata.cpp
+++ b/Engine/Source/Runtime/Renderer/Private/Strata/Strata.cpp
@@ -102,7 +102,6 @@ uint32 GetStrataTileSize()
 }
 
 static void AddStrataLUTPass(FRDGBuilder& GraphBuilder, const FViewInfo& View);
-static void AddStrataFurnacePass(FRDGBuilder& GraphBuilder, const FViewInfo& View);
 
 void InitialiseStrataFrameSceneData(FSceneRenderer& SceneRenderer, FRDGBuilder& GraphBuilder)
 {
@@ -204,21 +203,12 @@ void InitialiseStrataFrameSceneData(FSceneRenderer& SceneRenderer, FRDGBuilder&
 
 	AddStrataClearMaterialBufferPass(GraphBuilder, StrataSceneData.MaterialLobesBuffer.UAV, StrataSceneData.MaxBytesPerPixel, FIntPoint(ResolutionX, ResolutionY));
 
-	const bool bFurnaceTest = CVarStrataFurnaceTest.GetValueOnAnyThread() > 0;
-	if (bUpdateLUT || bFurnaceTest)
+	if (bUpdateLUT)
 	{
 		for (int32 ViewIndex = 0; ViewIndex < SceneRenderer.Views.Num(); ViewIndex++)
 		{
 			FViewInfo& View = SceneRenderer.Views[ViewIndex];
-			if (bUpdateLUT)
-			{
-				AddStrataLUTPass(GraphBuilder, View);
-			}
-
-			if (bFurnaceTest)
-			{
-				AddStrataFurnacePass(GraphBuilder, View);
-			}
+			AddStrataLUTPass(GraphBuilder, View);
 		}
 	}
 }
@@ -323,43 +313,31 @@ class FVisualizeMaterialPS : public FGlobalShader
 };
 IMPLEMENT_GLOBAL_SHADER(FVisualizeMaterialPS, "/Engine/Private/Strata/StrataVisualize.usf", "VisualizeMaterialPS", SF_Pixel);
 
-void AddVisualizeMaterialPasses(FRDGBuilder& GraphBuilder, const TArray<FViewInfo>& Views, FRDGTextureRef SceneColorTexture, EShaderPlatform Platform)
+static void AddVisualizeMaterialPasses(FRDGBuilder& GraphBuilder, const FViewInfo& View, FRDGTextureRef SceneColorTexture, EShaderPlatform Platform)
 {
-	RDG_EVENT_SCOPE_CONDITIONAL(GraphBuilder, IsStrataEnabled() && Views.Num() > 0, "StrataVisualizeMaterial");
-	if (!IsStrataEnabled() || !FVisualizeMaterialPS::CanRunStrataVizualizeMaterial(Platform))
-	{
-		return;
-	}
-
 	FRHIBlendState* PreMultipliedColorTransmittanceBlend = TStaticBlendState<CW_RGB, BO_Add, BF_One, BF_SourceAlpha, BO_Add, BF_Zero, BF_One>::GetRHI();
-
-	for (int32 i = 0; i < Views.Num(); ++i)
+	if (View.Family->EngineShowFlags.VisualizeStrataMaterial)
 	{
-		const FViewInfo& View = Views[i];
+		FVisualizeMaterialPS::FParameters* PassParameters = GraphBuilder.AllocParameters<FVisualizeMaterialPS::FParameters>();
+		PassParameters->ViewUniformBuffer = View.ViewUniformBuffer;
+		PassParameters->Strata = Strata::BindStrataGlobalUniformParameters(View);
+		PassParameters->MiniFontTexture = GetMiniFontTexture();
+		PassParameters->SceneTextures = GetSceneTextureParameters(GraphBuilder);
+		PassParameters->RenderTargets[0] = FRenderTargetBinding(SceneColorTexture, ERenderTargetLoadAction::ELoad);
 
-		if (View.Family->EngineShowFlags.VisualizeStrataMaterial)
+		if (ShaderDrawDebug::IsShaderDrawDebugEnabled())
 		{
-			FVisualizeMaterialPS::FParameters* PassParameters = GraphBuilder.AllocParameters<FVisualizeMaterialPS::FParameters>();
-			PassParameters->ViewUniformBuffer = View.ViewUniformBuffer;
-			PassParameters->Strata = Strata::BindStrataGlobalUniformParameters(View);
-			PassParameters->MiniFontTexture = GetMiniFontTexture();
-			PassParameters->SceneTextures = GetSceneTextureParameters(GraphBuilder);
-			PassParameters->RenderTargets[0] = FRenderTargetBinding(SceneColorTexture, ERenderTargetLoadAction::ELoad);
-
-			if (ShaderDrawDebug::IsShaderDrawDebugEnabled())
-			{
-				ShaderDrawDebug::SetParameters(GraphBuilder, View.ShaderDrawData, PassParameters->ShaderDrawParameters);
-			}
+			ShaderDrawDebug::SetParameters(GraphBuilder, View.ShaderDrawData, PassParameters->ShaderDrawParameters);
+		}
 
-			for (uint32 j = 0; j < 4; ++j)
-			{
-				FVisualizeMaterialPS::FPermutationDomain PermutationVector;
-				PermutationVector.Set<typename FVisualizeMaterialPS::FBSDFPass>(j);
-				TShaderMapRef<FVisualizeMaterialPS> PixelShader(View.ShaderMap, PermutationVector);
+		for (uint32 j = 0; j < 4; ++j)
+		{
+			FVisualizeMaterialPS::FPermutationDomain PermutationVector;
+			PermutationVector.Set<typename FVisualizeMaterialPS::FBSDFPass>(j);
+			TShaderMapRef<FVisualizeMaterialPS> PixelShader(View.ShaderMap, PermutationVector);
 
-				FPixelShaderUtils::AddFullscreenPass<FVisualizeMaterialPS>(GraphBuilder, View.ShaderMap, RDG_EVENT_NAME("StrataVisualizeMaterial"),
-					PixelShader, PassParameters, View.ViewRect, PreMultipliedColorTransmittanceBlend);
-			}
+			FPixelShaderUtils::AddFullscreenPass<FVisualizeMaterialPS>(GraphBuilder, View.ShaderMap, RDG_EVENT_NAME("StrataVisualizeMaterial"),
+				PixelShader, PassParameters, View.ViewRect, PreMultipliedColorTransmittanceBlend);
 		}
 	}
 }
@@ -658,23 +636,6 @@ void AddStrataStencilPass(
 	}
 }
 
-void AddVisualizeClassificationPass(
-	FRDGBuilder& GraphBuilder, 
-	const TArray<FViewInfo>& Views, 
-	FRDGTextureRef SceneColorTexture)
-{
-	if (IsClassificationEnabled() && CVarStrataClassificationDebug.GetValueOnAnyThread() > 0)
-	{
-		for (int32 i = 0; i < Views.Num(); ++i)
-		{
-			const FViewInfo& View = Views[i];
-			FRDGBufferRef TileListBuffer = GraphBuilder.RegisterExternalBuffer(View.StrataSceneData->ClassificationTileListBuffer);
-			FRDGBufferRef TileIndirectBuffer = GraphBuilder.RegisterExternalBuffer(View.StrataSceneData->ClassificationTileIndirectBuffer);
-			AddStrataInternalClassifedTilePass(GraphBuilder, View, nullptr, &SceneColorTexture, TileListBuffer, TileIndirectBuffer);
-		}
-	}
-}
-
 ////////////////////////////////////////////////////////////////////////////////////////////////////////////////////////////////////////////////
 
 void AddStrataMaterialClassificationPass(FRDGBuilder& GraphBuilder, const FMinimalSceneTextures& SceneTextures, const TArray<FViewInfo>& Views)
@@ -842,7 +803,6 @@ class FStrataFurnaceTestPassPS : public FGlobalShader
 	BEGIN_SHADER_PARAMETER_STRUCT(FParameters, )
 		SHADER_PARAMETER_STRUCT_REF(FViewUniformShaderParameters, ViewUniformBuffer)
 		SHADER_PARAMETER_STRUCT_REF(FStrataGlobalUniformParameters, Strata)
-		SHADER_PARAMETER(FIntPoint, OutputResolution)
 		SHADER_PARAMETER(uint32, NumSamples)
 		SHADER_PARAMETER(uint32, SceneType)
 		SHADER_PARAMETER_RDG_TEXTURE(Texture3D, OutLUT3D)
@@ -869,23 +829,15 @@ class FStrataFurnaceTestPassPS : public FGlobalShader
 
 IMPLEMENT_GLOBAL_SHADER(FStrataFurnaceTestPassPS, "/Engine/Private/Strata/StrataFurnaceTest.usf", "MainPS", SF_Pixel);
 
-static void AddStrataFurnacePass(FRDGBuilder& GraphBuilder, const FViewInfo& View)
+static void AddStrataFurnacePass(FRDGBuilder& GraphBuilder, const FViewInfo& View, FRDGTextureRef OutTexture)
 {
-	FIntPoint OutputResolution;
-	OutputResolution.X = 1024;
-	OutputResolution.Y = 1024;
-
-	FRDGTextureDesc Desc = FRDGTextureDesc::Create2D(OutputResolution, EPixelFormat::PF_FloatRGBA, FClearValueBinding::Black, ETextureCreateFlags::TexCreate_RenderTargetable);
-	FRDGTextureRef Texture = GraphBuilder.CreateTexture(Desc, TEXT("StrataFurnaceTest"));
-
 	TShaderMapRef<FStrataFurnaceTestPassPS> PixelShader(View.ShaderMap);
 	FStrataFurnaceTestPassPS::FParameters* Parameters = GraphBuilder.AllocParameters<FStrataFurnaceTestPassPS::FParameters>();
 	Parameters->ViewUniformBuffer = View.ViewUniformBuffer;
 	Parameters->Strata = Strata::BindStrataGlobalUniformParameters(View);
 	Parameters->SceneType = FMath::Clamp(CVarStrataFurnaceTest.GetValueOnAnyThread(), 1, 2);
 	Parameters->NumSamples = FMath::Clamp(CVarStrataFurnaceTestSampleCount.GetValueOnAnyThread(), 16, 2048);
-	Parameters->OutputResolution = OutputResolution;
-	Parameters->RenderTargets[0] = FRenderTargetBinding(Texture, ERenderTargetLoadAction::EClear);
+	Parameters->RenderTargets[0] = FRenderTargetBinding(OutTexture, ERenderTargetLoadAction::ELoad);
 
 	FPixelShaderUtils::AddFullscreenPass<FStrataFurnaceTestPassPS>(
 		GraphBuilder,
@@ -896,6 +848,44 @@ static void AddStrataFurnacePass(FRDGBuilder& GraphBuilder, const FViewInfo& Vie
 		View.ViewRect);
 }
 
+void AddStrataDebugPasses(FRDGBuilder& GraphBuilder, const TArray<FViewInfo>& Views, FRDGTextureRef SceneColorTexture, EShaderPlatform Platform)
+{
+	if (!IsStrataEnabled())
+		return;
+
+	if (FVisualizeMaterialPS::CanRunStrataVizualizeMaterial(Platform))
+	{
+		RDG_EVENT_SCOPE(GraphBuilder, "StrataVisualizeMaterial");
+		for (int32 i = 0; i < Views.Num(); ++i)
+		{
+			const FViewInfo& View = Views[i];
+			AddVisualizeMaterialPasses(GraphBuilder, View, SceneColorTexture, Platform);
+		}
+	}
+
+	if (IsClassificationEnabled() && CVarStrataClassificationDebug.GetValueOnAnyThread() > 0)
+	{
+		RDG_EVENT_SCOPE(GraphBuilder, "StrataVisualizeClassification");
+		for (int32 i = 0; i < Views.Num(); ++i)
+		{
+			const FViewInfo& View = Views[i];
+			FRDGBufferRef TileListBuffer = GraphBuilder.RegisterExternalBuffer(View.StrataSceneData->ClassificationTileListBuffer);
+			FRDGBufferRef TileIndirectBuffer = GraphBuilder.RegisterExternalBuffer(View.StrataSceneData->ClassificationTileIndirectBuffer);
+			AddStrataInternalClassifedTilePass(GraphBuilder, View, nullptr, &SceneColorTexture, TileListBuffer, TileIndirectBuffer);
+		}
+	}
+
+	if (CVarStrataFurnaceTest.GetValueOnAnyThread() > 0)
+	{
+		RDG_EVENT_SCOPE(GraphBuilder, "StrataVisualizeFurnaceTest");
+		for (int32 i = 0; i < Views.Num(); ++i)
+		{
+			const FViewInfo& View = Views[i];
+			AddStrataFurnacePass(GraphBuilder, View, SceneColorTexture);
+		}
+	}
+}
+
 } // namespace Strata
 
 
diff --git a/Engine/Source/Runtime/Renderer/Private/Strata/Strata.h b/Engine/Source/Runtime/Renderer/Private/Strata/Strata.h
--- a/Engine/Source/Runtime/Renderer/Private/Strata/Strata.h
+++ b/Engine/Source/Runtime/Renderer/Private/Strata/Strata.h
@@ -71,8 +71,7 @@ void AddStrataMaterialClassificationPass(FRDGBuilder& GraphBuilder, const FMinim
 void AddStrataStencilPass(FRDGBuilder& GraphBuilder, const FViewInfo& View, const FMinimalSceneTextures& SceneTextures);
 void AddStrataStencilPass(FRDGBuilder& GraphBuilder, const TArray<FViewInfo>& Views, const FMinimalSceneTextures& SceneTextures);
 
-void AddVisualizeMaterialPasses(FRDGBuilder& GraphBuilder, const TArray<FViewInfo>& Views, FRDGTextureRef SceneColorTexture, EShaderPlatform Platform);
-void AddVisualizeClassificationPass(FRDGBuilder& GraphBuilder, const TArray<FViewInfo>& Views, FRDGTextureRef SceneColorTexture);
+void AddStrataDebugPasses(FRDGBuilder& GraphBuilder, const TArray<FViewInfo>& Views, FRDGTextureRef SceneColorTexture, EShaderPlatform Platform);
 };
 
 
