Defer ID translation for dynamic primitives such that the commit can happen at any time before draw submit.
- dynamic primitive collectors now auto-commit when uploading dynamic primitives, simplifying use.
- fixed shadow stencil passes to properly use GPU-Scene instance culling (and instancing), now using FSimpleMeshDrawCommandPass
- Also move translucency shadow rendering to FSimpleMeshDrawCommandPass

diff --git a/Engine/Shaders/Private/InstanceCulling/BuildInstanceDrawCommands.usf b/Engine/Shaders/Private/InstanceCulling/BuildInstanceDrawCommands.usf
--- a/Engine/Shaders/Private/InstanceCulling/BuildInstanceDrawCommands.usf
+++ b/Engine/Shaders/Private/InstanceCulling/BuildInstanceDrawCommands.usf
@@ -50,6 +50,31 @@ uint NumPrimitiveIds;
 uint NumInstanceRuns;
 int NumCommands;
 uint NumViewIds;
+int DynamicPrimitiveIdOffset;
+int DynamicPrimitiveIdMax;
+
+
+/**
+ * Check if the PRIM_ID_DYNAMIC_FLAG bit is set on the Primitive ID and if so add the dynamic primitive offset to get the GPU-Scene ID.
+ */
+uint TranslatePrimitiveId(uint PrimitiveIdIn)
+{
+	// This means we defer the translation to later
+	if (DynamicPrimitiveIdOffset == -1)
+	{
+		return PrimitiveIdIn;
+	}
+
+	uint PrimitiveId = PrimitiveIdIn;
+	if ((PrimitiveIdIn & PRIM_ID_DYNAMIC_FLAG) != 0)
+	{
+		uint DynamicPrimitiveIndex = PrimitiveIdIn & (~PRIM_ID_DYNAMIC_FLAG);
+		PrimitiveId = uint(DynamicPrimitiveIdOffset) + DynamicPrimitiveIndex;
+	}
+
+	return PrimitiveId;
+}
+
 
 bool TestInstanceVisibilityBit(uint InstanceId, uint ViewId)
 {
@@ -108,7 +133,7 @@ void ComputeInstanceIdOutputSize(uint GroupThreadIndex : SV_GroupIndex, uint Gro
 		uint Num = 0;
 		if (Index < NumCmdPrimitiveIds)
 		{
-			uint PrimitiveId = PrimitiveIds[Cmd.FirstPrimitiveIdOffset + Index];
+			uint PrimitiveId = TranslatePrimitiveId(PrimitiveIds[Cmd.FirstPrimitiveIdOffset + Index]);
 			if (GetPrimitiveData(PrimitiveId).InstanceDataOffset >= 0)
 			{
 				Num = GetPrimitiveData(PrimitiveId).NumInstanceDataEntries;
@@ -246,7 +271,7 @@ void OutputInstanceIdsAtOffset(uint GroupThreadIndex : SV_GroupIndex, uint Group
 	if (CommandInstanceCount > 0)
 	{
 		// 3. Broadcast to all threads
-		uint PrimitiveId = PrimitiveIds[Cmd.FirstPrimitiveIdOffset];
+		uint PrimitiveId = TranslatePrimitiveId(PrimitiveIds[Cmd.FirstPrimitiveIdOffset]);
 		if (GroupThreadIndex == 0)
 		{
 			WriteInstanceId(InstanceIdOutputOffset, GetPrimitiveData(PrimitiveId).InstanceDataOffset, 0U, CommandSlotIndex);
@@ -272,7 +297,7 @@ void OutputInstanceIdsAtOffset(uint GroupThreadIndex : SV_GroupIndex, uint Group
 				while (OutputIndex >= CurrentRangeEnd)
 				{
 					CurrentPrimitiveIdIndex += 1;
-					uint PrimitiveId = PrimitiveIds[Cmd.FirstPrimitiveIdOffset + CurrentPrimitiveIdIndex];
+					uint PrimitiveId = TranslatePrimitiveId(PrimitiveIds[Cmd.FirstPrimitiveIdOffset + CurrentPrimitiveIdIndex]);
 
 					CurrentRangeStart = CurrentRangeEnd;
 					InstanceDataOffset = GetPrimitiveData(PrimitiveId).InstanceDataOffset;
@@ -390,7 +415,7 @@ void BuildInstanceIdBufferAndCommandsFromPrimitiveIdsCs(uint GroupThreadIndex :
 		uint Num = 0;
 		if (Index < NumCmdPrimitiveIds)
 		{
-			uint PrimitiveId = PrimitiveIds[Cmd.FirstPrimitiveIdOffset + Index];
+			uint PrimitiveId = TranslatePrimitiveId(PrimitiveIds[Cmd.FirstPrimitiveIdOffset + Index]);
 			if (GetPrimitiveData(PrimitiveId).InstanceDataOffset >= 0)
 			{
 				Num = GetPrimitiveData(PrimitiveId).NumInstanceDataEntries;
@@ -433,7 +458,7 @@ void BuildInstanceIdBufferAndCommandsFromPrimitiveIdsCs(uint GroupThreadIndex :
 
 		// 3. Broadcast to all threads
 		uint InstanceIdOutputOffset = InstanceIdOffsetGroupShared;
-		uint PrimitiveId = PrimitiveIds[Cmd.FirstPrimitiveIdOffset];
+		uint PrimitiveId = TranslatePrimitiveId(PrimitiveIds[Cmd.FirstPrimitiveIdOffset]);
 		WriteInstanceId(InstanceIdOutputOffset, GetPrimitiveData(PrimitiveId).InstanceDataOffset, 0U, CommandSlotIndex);
 		GroupMemoryBarrierWithGroupSync();
 
@@ -450,7 +475,7 @@ void BuildInstanceIdBufferAndCommandsFromPrimitiveIdsCs(uint GroupThreadIndex :
 			while (OutputIndex1 >= CurrentRangeEnd)
 			{
 				CurrentPrimitiveIdIndex += 1;
-				uint PrimitiveId = PrimitiveIds[Cmd.FirstPrimitiveIdOffset + CurrentPrimitiveIdIndex];
+				uint PrimitiveId = TranslatePrimitiveId(PrimitiveIds[Cmd.FirstPrimitiveIdOffset + CurrentPrimitiveIdIndex]);
 
 				CurrentRangeStart = CurrentRangeEnd;
 				InstanceDataOffset = GetPrimitiveData(PrimitiveId).InstanceDataOffset;
diff --git a/Engine/Source/Runtime/Engine/Public/MeshBatch.h b/Engine/Source/Runtime/Engine/Public/MeshBatch.h
--- a/Engine/Source/Runtime/Engine/Public/MeshBatch.h
+++ b/Engine/Source/Runtime/Engine/Public/MeshBatch.h
@@ -38,6 +38,9 @@ enum EPrimitiveIdMode
 	PrimID_NumBits						= 2,
 };
 
+// Flag used to mark a primtive ID as dynamic, and thus needing translation (by adding the offset from the dynamic primitive collector).
+static constexpr int32 GPrimIDDynamicFlag = 1 << 31;
+
 /**
  * A batch mesh element definition.
  */
diff --git a/Engine/Source/Runtime/Renderer/Private/GPUScene.cpp b/Engine/Source/Runtime/Renderer/Private/GPUScene.cpp
--- a/Engine/Source/Runtime/Renderer/Private/GPUScene.cpp
+++ b/Engine/Source/Runtime/Renderer/Private/GPUScene.cpp
@@ -1050,6 +1050,13 @@ void FGPUScene::UploadDynamicPrimitiveShaderDataForViewInternal(FRHICommandListI
 
 	FGPUScenePrimitiveCollector& Collector = View.DynamicPrimitiveCollector;
 
+	// Auto-commit if not done (should usually not be done, but sometimes the UploadDynamicPrimitiveShaderDataForViewInternal is called to ensure the 
+	// CachedViewUniformShaderParameters is set on the view.
+	if (!Collector.bCommitted)
+	{
+		Collector.Commit();
+	}
+
 	const int32 NumPrimitiveDataUploads = Collector.Num();
 	ensure(Collector.GetPrimitiveIdRange().Size<int32>() == NumPrimitiveDataUploads);
 	
diff --git a/Engine/Source/Runtime/Renderer/Private/GPUScene.h b/Engine/Source/Runtime/Renderer/Private/GPUScene.h
--- a/Engine/Source/Runtime/Renderer/Private/GPUScene.h
+++ b/Engine/Source/Runtime/Renderer/Private/GPUScene.h
@@ -73,8 +73,13 @@ public:
 	RENDERER_API void Commit();
 
 	/**
+	 * Get the range of Primitive IDs in GPU-Scene for this batch of dynamic primitives, only valid to call after commit.
 	 */
-	const TRange<int32> &GetPrimitiveIdRange() const { return PrimitiveIdRange; }
+	FORCEINLINE const TRange<int32> &GetPrimitiveIdRange() const
+	{ 
+		ensure(bCommitted || UploadData == nullptr);
+		return PrimitiveIdRange; 
+	}
 
 	FORCEINLINE int32 GetInstanceDataOffset(int32 PrimitiveId) const
 	{
diff --git a/Engine/Source/Runtime/Renderer/Private/HairStrands/HairStrandsRasterCommon.cpp b/Engine/Source/Runtime/Renderer/Private/HairStrands/HairStrandsRasterCommon.cpp
--- a/Engine/Source/Runtime/Renderer/Private/HairStrands/HairStrandsRasterCommon.cpp
+++ b/Engine/Source/Runtime/Renderer/Private/HairStrands/HairStrandsRasterCommon.cpp
@@ -427,7 +427,7 @@ void AddHairStrandsRasterPass(
 		if (VisibleMeshDrawCommands.Num() > 0)
 		{
 			FRHIBuffer* PrimitiveIdVertexBuffer = nullptr;
-			SortAndMergeDynamicPassMeshDrawCommands(ViewInfo->GetFeatureLevel(), VisibleMeshDrawCommands, DynamicMeshDrawCommandStorage, PrimitiveIdVertexBuffer, 1);
+			SortAndMergeDynamicPassMeshDrawCommands(ViewInfo->GetFeatureLevel(), VisibleMeshDrawCommands, DynamicMeshDrawCommandStorage, PrimitiveIdVertexBuffer, 1, ViewInfo->DynamicPrimitiveCollector.GetPrimitiveIdRange());
 			SubmitMeshDrawCommands(VisibleMeshDrawCommands, GraphicsMinimalPipelineStateSet, PrimitiveIdVertexBuffer, 0, false, 1, RHICmdList);
 		}
 	});
diff --git a/Engine/Source/Runtime/Renderer/Private/HairStrands/HairStrandsVisibility.cpp b/Engine/Source/Runtime/Renderer/Private/HairStrands/HairStrandsVisibility.cpp
--- a/Engine/Source/Runtime/Renderer/Private/HairStrands/HairStrandsVisibility.cpp
+++ b/Engine/Source/Runtime/Renderer/Private/HairStrands/HairStrandsVisibility.cpp
@@ -738,7 +738,7 @@ static void AddHairMaterialGBufferPass(
 				if (VisibleMeshDrawCommands.Num() > 0)
 				{
 					FRHIBuffer* PrimitiveIdVertexBuffer = nullptr;
-					SortAndMergeDynamicPassMeshDrawCommands(ViewInfo->GetFeatureLevel(), VisibleMeshDrawCommands, DynamicMeshDrawCommandStorage, PrimitiveIdVertexBuffer, 1);
+					SortAndMergeDynamicPassMeshDrawCommands(ViewInfo->GetFeatureLevel(), VisibleMeshDrawCommands, DynamicMeshDrawCommandStorage, PrimitiveIdVertexBuffer, 1, ViewInfo->DynamicPrimitiveCollector.GetPrimitiveIdRange());
 					SubmitMeshDrawCommands(VisibleMeshDrawCommands, PipelineStateSet, PrimitiveIdVertexBuffer, 0, false, 1, RHICmdList);
 				}
 			}
@@ -1181,7 +1181,7 @@ static FMaterialPassOutput AddHairMaterialPass(
 			if (VisibleMeshDrawCommands.Num() > 0)
 			{
 				FRHIBuffer* PrimitiveIdVertexBuffer = nullptr;
-				SortAndMergeDynamicPassMeshDrawCommands(ViewInfo->GetFeatureLevel(), VisibleMeshDrawCommands, DynamicMeshDrawCommandStorage, PrimitiveIdVertexBuffer, 1);
+				SortAndMergeDynamicPassMeshDrawCommands(ViewInfo->GetFeatureLevel(), VisibleMeshDrawCommands, DynamicMeshDrawCommandStorage, PrimitiveIdVertexBuffer, 1, ViewInfo->DynamicPrimitiveCollector.GetPrimitiveIdRange());
 				SubmitMeshDrawCommands(VisibleMeshDrawCommands, PipelineStateSet, PrimitiveIdVertexBuffer, 0, false, 1, RHICmdList);
 			}
 		}
@@ -2466,7 +2466,7 @@ static void AddHairVisibilityCommonPass(
 			if (VisibleMeshDrawCommands.Num() > 0)
 			{
 				FRHIBuffer* PrimitiveIdVertexBuffer = nullptr;
-				SortAndMergeDynamicPassMeshDrawCommands(ViewInfo->GetFeatureLevel(), VisibleMeshDrawCommands, DynamicMeshDrawCommandStorage, PrimitiveIdVertexBuffer, 1);
+				SortAndMergeDynamicPassMeshDrawCommands(ViewInfo->GetFeatureLevel(), VisibleMeshDrawCommands, DynamicMeshDrawCommandStorage, PrimitiveIdVertexBuffer, 1, ViewInfo->DynamicPrimitiveCollector.GetPrimitiveIdRange());
 				SubmitMeshDrawCommands(VisibleMeshDrawCommands, PipelineStateSet, PrimitiveIdVertexBuffer, 0, false, 1, RHICmdList);
 			}
 		}
diff --git a/Engine/Source/Runtime/Renderer/Private/InstanceCulling/InstanceCullingContext.cpp b/Engine/Source/Runtime/Renderer/Private/InstanceCulling/InstanceCullingContext.cpp
--- a/Engine/Source/Runtime/Renderer/Private/InstanceCulling/InstanceCullingContext.cpp
+++ b/Engine/Source/Runtime/Renderer/Private/InstanceCulling/InstanceCullingContext.cpp
@@ -25,15 +25,18 @@ FInstanceCullingContext::FInstanceCullingContext(FInstanceCullingManager* InInst
 void FInstanceCullingContext::BeginCullingCommand(EPrimitiveType BatchType, uint32 BaseVertexIndex, uint32 FirstIndex, uint32 NumPrimitives)
 {
 #if defined(GPUCULL_TODO)
-	if (ensure(BatchType == PT_TriangleList || BatchType == PT_LineList || BatchType == PT_PointList || BatchType == PT_QuadList))
+	if (ensure(BatchType < PT_Num))
 	{
-		// default to PT_TriangleList
+		// default to PT_TriangleList & PT_RectList
 		int32 NumVerticesOrIndices = NumPrimitives * 3;
 		switch (BatchType)
 		{
 		case PT_QuadList:
 			NumVerticesOrIndices = NumPrimitives * 4;
 			break;
+		case PT_TriangleStrip:
+			NumVerticesOrIndices = NumPrimitives + 2;
+			break;
 		case PT_LineList:
 			NumVerticesOrIndices = NumPrimitives * 2;
 			break;
@@ -99,6 +102,7 @@ public:
 		OutEnvironment.SetDefine(TEXT("NUM_THREADS_PER_GROUP"), NumThreadsPerGroup);
 		OutEnvironment.SetDefine(TEXT("NANITE_MULTI_VIEW"), 1);
 		OutEnvironment.SetDefine(TEXT("ENABLE_DETERMINISTIC_INSTANCE_CULLING"), 1);
+		OutEnvironment.SetDefine(TEXT("PRIM_ID_DYNAMIC_FLAG"), GPrimIDDynamicFlag);
 	}
 
 	BEGIN_SHADER_PARAMETER_STRUCT(FParameters, )
@@ -120,6 +124,9 @@ public:
 		SHADER_PARAMETER(int32, NumCommands)
 		SHADER_PARAMETER(uint32, NumInstanceFlagWords)
 		SHADER_PARAMETER(int32, NumViewIds)
+
+		SHADER_PARAMETER(int32, DynamicPrimitiveIdOffset)
+		SHADER_PARAMETER(int32, DynamicPrimitiveIdMax)
 	END_SHADER_PARAMETER_STRUCT()
 };
 IMPLEMENT_GLOBAL_SHADER(FComputeInstanceIdOutputSizeCs, "/Engine/Private/InstanceCulling/BuildInstanceDrawCommands.usf", "ComputeInstanceIdOutputSize", SF_Compute);
@@ -145,6 +152,7 @@ public:
 		OutEnvironment.SetDefine(TEXT("NUM_THREADS_PER_GROUP"), NumThreadsPerGroup);
 		OutEnvironment.SetDefine(TEXT("NANITE_MULTI_VIEW"), 1);
 		OutEnvironment.SetDefine(TEXT("ENABLE_DETERMINISTIC_INSTANCE_CULLING"), 1);
+		OutEnvironment.SetDefine(TEXT("PRIM_ID_DYNAMIC_FLAG"), GPrimIDDynamicFlag);
 	}
 
 	BEGIN_SHADER_PARAMETER_STRUCT(FParameters, )
@@ -183,6 +191,7 @@ public:
 		OutEnvironment.SetDefine(TEXT("NUM_THREADS_PER_GROUP"), NumThreadsPerGroup);
 		OutEnvironment.SetDefine(TEXT("NANITE_MULTI_VIEW"), 1);
 		OutEnvironment.SetDefine(TEXT("ENABLE_DETERMINISTIC_INSTANCE_CULLING"), 1);
+		OutEnvironment.SetDefine(TEXT("PRIM_ID_DYNAMIC_FLAG"), GPrimIDDynamicFlag);
 	}
 
 	BEGIN_SHADER_PARAMETER_STRUCT(FParameters, )
@@ -211,6 +220,9 @@ public:
 		SHADER_PARAMETER(uint32, NumInstanceFlagWords)
 		SHADER_PARAMETER(int32, NumViewIds)
 
+		SHADER_PARAMETER(int32, DynamicPrimitiveIdOffset)
+		SHADER_PARAMETER(int32, DynamicPrimitiveIdMax)
+
 	END_SHADER_PARAMETER_STRUCT()
 };
 IMPLEMENT_GLOBAL_SHADER(FOutputInstanceIdsAtOffsetCs, "/Engine/Private/InstanceCulling/BuildInstanceDrawCommands.usf", "OutputInstanceIdsAtOffset", SF_Compute);
@@ -243,6 +255,7 @@ public:
 		OutEnvironment.SetDefine(TEXT("USE_GLOBAL_GPU_SCENE_DATA"), 1);
 		OutEnvironment.SetDefine(TEXT("NUM_THREADS_PER_GROUP"), NumThreadsPerGroup);
 		OutEnvironment.SetDefine(TEXT("NANITE_MULTI_VIEW"), 1);
+		OutEnvironment.SetDefine(TEXT("PRIM_ID_DYNAMIC_FLAG"), GPrimIDDynamicFlag);
 	}
 
 	BEGIN_SHADER_PARAMETER_STRUCT(FParameters, )
@@ -270,11 +283,14 @@ public:
 		SHADER_PARAMETER(int32, NumCommands)
 		SHADER_PARAMETER(uint32, NumInstanceFlagWords)
 		SHADER_PARAMETER(int32, NumViewIds)
+
+		SHADER_PARAMETER(int32, DynamicPrimitiveIdOffset)
+		SHADER_PARAMETER(int32, DynamicPrimitiveIdMax)
 	END_SHADER_PARAMETER_STRUCT()
 };
 IMPLEMENT_GLOBAL_SHADER(FBuildInstanceIdBufferAndCommandsFromPrimitiveIdsCs, "/Engine/Private/InstanceCulling/BuildInstanceDrawCommands.usf", "BuildInstanceIdBufferAndCommandsFromPrimitiveIdsCs", SF_Compute);
 
-void FInstanceCullingContext::BuildRenderingCommands(FRDGBuilder& GraphBuilder, const FGPUScene& GPUScene, FInstanceCullingResult& Results) const
+void FInstanceCullingContext::BuildRenderingCommands(FRDGBuilder& GraphBuilder, const FGPUScene& GPUScene, const TRange<int32>& DynamicPrimitiveIdRange, FInstanceCullingResult& Results) const
 {
 	Results = FInstanceCullingResult();
 	if (CullingCommands.Num() == 0)
@@ -315,6 +331,9 @@ void FInstanceCullingContext::BuildRenderingCommands(FRDGBuilder& GraphBuilder,
 		PassParameters->PrimitiveCullingCommands = GraphBuilder.CreateSRV(CreateStructuredBuffer(GraphBuilder, TEXT("PrimitiveCullingCommands"), CullingCommands));
 
 		PassParameters->PrimitiveIds = GraphBuilder.CreateSRV(CreateStructuredBuffer(GraphBuilder, TEXT("PrimitiveIds"), PrimitiveIds));
+		PassParameters->DynamicPrimitiveIdOffset = DynamicPrimitiveIdRange.GetLowerBoundValue();
+		PassParameters->DynamicPrimitiveIdMax = DynamicPrimitiveIdRange.GetUpperBoundValue();
+
 		PassParameters->InstanceRuns = GraphBuilder.CreateSRV(CreateStructuredBuffer(GraphBuilder, TEXT("InstanceRuns"), InstanceRuns));
 
 		PassParameters->OutputOffsetBufferOut = GraphBuilder.CreateUAV(InstanceIdOutOffsetBufferRDG);
@@ -378,6 +397,9 @@ void FInstanceCullingContext::BuildRenderingCommands(FRDGBuilder& GraphBuilder,
 		PassParameters->PrimitiveCullingCommands = GraphBuilder.CreateSRV(CreateStructuredBuffer(GraphBuilder, TEXT("PrimitiveCullingCommands"), CullingCommands));
 
 		PassParameters->PrimitiveIds = GraphBuilder.CreateSRV(CreateStructuredBuffer(GraphBuilder, TEXT("PrimitiveIds"), PrimitiveIds));
+		PassParameters->DynamicPrimitiveIdOffset = DynamicPrimitiveIdRange.GetLowerBoundValue();
+		PassParameters->DynamicPrimitiveIdMax = DynamicPrimitiveIdRange.GetUpperBoundValue();
+
 		PassParameters->InstanceRuns = GraphBuilder.CreateSRV(CreateStructuredBuffer(GraphBuilder, TEXT("InstanceRuns"), InstanceRuns));
 
 		PassParameters->ViewIds = GraphBuilder.CreateSRV(CreateStructuredBuffer(GraphBuilder, TEXT("ViewIds"), ViewIds));
@@ -428,6 +450,9 @@ void FInstanceCullingContext::BuildRenderingCommands(FRDGBuilder& GraphBuilder,
 	PassParameters->PrimitiveCullingCommands = GraphBuilder.CreateSRV(CreateStructuredBuffer(GraphBuilder, TEXT("PrimitiveCullingCommands"), CullingCommands));
 
 	PassParameters->PrimitiveIds = GraphBuilder.CreateSRV(CreateStructuredBuffer(GraphBuilder, TEXT("PrimitiveIds"), PrimitiveIds));
+	PassParameters->DynamicPrimitiveIdOffset = DynamicPrimitiveIdRange.GetLowerBoundValue();
+	PassParameters->DynamicPrimitiveIdMax = DynamicPrimitiveIdRange.GetUpperBoundValue();
+
 	PassParameters->InstanceRuns = GraphBuilder.CreateSRV(CreateStructuredBuffer(GraphBuilder, TEXT("InstanceRuns"), InstanceRuns));
 
 	PassParameters->OutputOffsetBufferOut = GraphBuilder.CreateUAV(InstanceIdOutOffsetBufferRDG);
@@ -490,7 +515,7 @@ void FInstanceCullingContext::BuildRenderingCommands(FRDGBuilder& GraphBuilder,
 	});
 }
 
-void FInstanceCullingContext::BuildRenderingCommands(FRDGBuilder& GraphBuilder, FGPUScene& GPUScene, FInstanceCullingRdgParams& Params) const
+void FInstanceCullingContext::BuildRenderingCommands(FRDGBuilder& GraphBuilder, FGPUScene& GPUScene, const TRange<int32>& DynamicPrimitiveIdRange, FInstanceCullingRdgParams& Params) const
 {
 	if (CullingCommands.Num() == 0)
 	{
@@ -512,6 +537,9 @@ void FInstanceCullingContext::BuildRenderingCommands(FRDGBuilder& GraphBuilder,
 	PassParameters->PrimitiveCullingCommands = GraphBuilder.CreateSRV(CreateStructuredBuffer(GraphBuilder, TEXT("PrimitiveCullingCommands"), CullingCommands));
 
 	PassParameters->PrimitiveIds = GraphBuilder.CreateSRV(CreateStructuredBuffer(GraphBuilder, TEXT("PrimitiveIds"), PrimitiveIds));
+	PassParameters->DynamicPrimitiveIdOffset = DynamicPrimitiveIdRange.GetLowerBoundValue();
+	PassParameters->DynamicPrimitiveIdMax = DynamicPrimitiveIdRange.GetUpperBoundValue();
+
 	PassParameters->InstanceRuns = GraphBuilder.CreateSRV(CreateStructuredBuffer(GraphBuilder, TEXT("InstanceRuns"), InstanceRuns));
 
 
@@ -578,11 +606,11 @@ void FInstanceCullingContext::BuildRenderingCommands(FRDGBuilder& GraphBuilder,
 
 #else // GPUCULL_TODO
 
-void FInstanceCullingContext::BuildRenderingCommands(FRDGBuilder& GraphBuilder, const FGPUScene& GPUScene, FInstanceCullingResult& Results) const
+void FInstanceCullingContext::BuildRenderingCommands(FRDGBuilder& GraphBuilder, const FGPUScene& GPUScene, const TRange<int32>& DynamicPrimitiveIdRange, FInstanceCullingResult& Results) const
 {
 }
 
-void FInstanceCullingContext::BuildRenderingCommands(FRDGBuilder& GraphBuilder, FGPUScene& GPUScene, FInstanceCullingRdgParams& Params) const
+void FInstanceCullingContext::BuildRenderingCommands(FRDGBuilder& GraphBuilder, FGPUScene& GPUScene, const TRange<int32>& DynamicPrimitiveIdRange, FInstanceCullingRdgParams& Params) const
 {
 }
 
diff --git a/Engine/Source/Runtime/Renderer/Private/MeshDrawCommands.cpp b/Engine/Source/Runtime/Renderer/Private/MeshDrawCommands.cpp
--- a/Engine/Source/Runtime/Renderer/Private/MeshDrawCommands.cpp
+++ b/Engine/Source/Runtime/Renderer/Private/MeshDrawCommands.cpp
@@ -361,11 +361,37 @@ void UpdateMobileBasePassMeshSortKeys(
 	}
 }
 
+
+FORCEINLINE int32 TranslatePrimitiveId(int32 DrawPrimitiveIdIn, int32 DynamicPrimitiveIdOffset, int32 DynamicPrimitiveIdMax)
+{
+	// INDEX_NONE means we defer the translation to later
+	if (DynamicPrimitiveIdOffset == INDEX_NONE)
+	{
+		return DrawPrimitiveIdIn;
+	}
+	int32 DrawPrimitiveId = DrawPrimitiveIdIn;
+
+	if ((DrawPrimitiveIdIn & GPrimIDDynamicFlag) != 0)
+	{
+		int32 DynamicPrimitiveIndex = DrawPrimitiveIdIn & (~GPrimIDDynamicFlag);
+		DrawPrimitiveId = DynamicPrimitiveIdOffset + DynamicPrimitiveIndex;
+		checkSlow(DrawPrimitiveId < DynamicPrimitiveIdMax);
+	}
+
+#if defined(GPUCULL_TODO)
+	// Append flag to mark this as a non-instance data index (which is then treated as a primitive ID in the SceneData.ush loading
+	return DrawPrimitiveId | (1U << 31U);
+#else
+	return DrawPrimitiveId;
+#endif
+
+}
+
 /**
  * Build mesh draw command primitive Id buffer for instancing.
  * TempVisibleMeshDrawCommands must be presized for NewPassVisibleMeshDrawCommands.
  */
-void BuildMeshDrawCommandPrimitiveIdBuffer(
+static void BuildMeshDrawCommandPrimitiveIdBuffer(
 	bool bDynamicInstancing,
 	FMeshCommandOneFrameArray& VisibleMeshDrawCommands,
 	FDynamicMeshDrawCommandStorage& MeshDrawCommandStorage,
@@ -376,8 +402,9 @@ void BuildMeshDrawCommandPrimitiveIdBuffer(
 	int32& VisibleMeshDrawCommandsNum,
 	int32& NewPassVisibleMeshDrawCommandsNum,
 	EShaderPlatform ShaderPlatform,
-	uint32 InstanceFactor = 1
-	)
+	uint32 InstanceFactor,
+	int32 DynamicPrimitiveIdOffset,
+	int32 DynamicPrimitiveIdMax)
 {
 	QUICK_SCOPE_CYCLE_COUNTER(STAT_BuildMeshDrawCommandPrimitiveIdBuffer);
 	check(PrimitiveIdData && PrimitiveIdDataSize > 0);
@@ -386,7 +413,7 @@ void BuildMeshDrawCommandPrimitiveIdBuffer(
 	const int32 NumDrawCommands = VisibleMeshDrawCommands.Num();
 
 	uint32 PrimitiveIdIndex = 0;
-	int32* RESTRICT PrimitiveIds = (int32*)PrimitiveIdData;
+	int32* RESTRICT PrimitiveIds = reinterpret_cast<int32*>(PrimitiveIdData);
 	const uint32 MaxPrimitiveId = PrimitiveIdDataSize / sizeof(int32);
 
 	if (bDynamicInstancing)
@@ -466,9 +493,8 @@ void BuildMeshDrawCommandPrimitiveIdBuffer(
 			//@todo - refactor into instance step rate in the RHI
 			for (uint32 InstanceFactorIndex = 0; InstanceFactorIndex < InstanceFactor; InstanceFactorIndex++, PrimitiveIdIndex++)
 			{
-				//@todo - refactor into memcpy
 				checkSlow(PrimitiveIdIndex < MaxPrimitiveId);
-				PrimitiveIds[PrimitiveIdIndex] = VisibleMeshDrawCommand.DrawPrimitiveId;
+				PrimitiveIds[PrimitiveIdIndex] = TranslatePrimitiveId(VisibleMeshDrawCommand.DrawPrimitiveId, DynamicPrimitiveIdOffset, DynamicPrimitiveIdMax);
 			}
 		}
 
@@ -490,12 +516,7 @@ void BuildMeshDrawCommandPrimitiveIdBuffer(
 			for (uint32 InstanceFactorIndex = 0; InstanceFactorIndex < InstanceFactor; InstanceFactorIndex++, PrimitiveIdIndex++)
 			{
 				checkSlow(PrimitiveIdIndex < MaxPrimitiveId);
-			#if defined(GPUCULL_TODO)
-				// Append flag to mark this as a non-instance data index (which is then treated as a primitive ID in the SceneData.ush loading
-				PrimitiveIds[PrimitiveIdIndex] = VisibleMeshDrawCommand.DrawPrimitiveId | (1U << 31U);
-			#else
-				PrimitiveIds[PrimitiveIdIndex] = VisibleMeshDrawCommand.DrawPrimitiveId;
-			#endif
+				PrimitiveIds[PrimitiveIdIndex] = TranslatePrimitiveId(VisibleMeshDrawCommand.DrawPrimitiveId, DynamicPrimitiveIdOffset, DynamicPrimitiveIdMax);
 			}
 		}
 	}
@@ -974,7 +995,9 @@ public:
 					Context.VisibleMeshDrawCommandsNum,
 					Context.NewPassVisibleMeshDrawCommandsNum,
 					Context.ShaderPlatform,
-					Context.InstanceFactor
+					Context.InstanceFactor,
+					INDEX_NONE, // Defer the translation until submit
+					0
 				);
 #endif // defined(GPUCULL_TODO)
 			}
@@ -1047,7 +1070,8 @@ void SortAndMergeDynamicPassMeshDrawCommands(
 	FMeshCommandOneFrameArray& VisibleMeshDrawCommands,
 	FDynamicMeshDrawCommandStorage& MeshDrawCommandStorage,
 	FRHIBuffer*& OutPrimitiveIdVertexBuffer,
-	uint32 InstanceFactor)
+	uint32 InstanceFactor,
+	const TRange<int32>& DynamicPrimitiveIdRange)
 {
 	const bool bUseGPUScene = UseGPUScene(GMaxRHIShaderPlatform, FeatureLevel);
 
@@ -1092,8 +1116,9 @@ void SortAndMergeDynamicPassMeshDrawCommands(
 				VisibleMeshDrawCommandsNum,
 				NewPassVisibleMeshDrawCommandsNum,
 				GShaderPlatformForFeatureLevel[FeatureLevel],
-				InstanceFactor
-			);
+				InstanceFactor,
+				DynamicPrimitiveIdRange.GetLowerBoundValue(),
+				DynamicPrimitiveIdRange.GetUpperBoundValue());
 
 			RHIUnlockBuffer(OutPrimitiveIdVertexBuffer);
 			GPrimitiveIdVertexBufferPool.ReturnToFreeList(Entry);
@@ -1194,6 +1219,9 @@ void FParallelMeshDrawCommandPass::DispatchPassSetup(
 		bPrimitiveIdBufferDataOwnedByRHIThread = false;
 		TaskContext.PrimitiveIdBufferDataSize = TaskContext.InstanceFactor * MaxNumDraws * sizeof(int32);
 		TaskContext.PrimitiveIdBufferData = FMemory::Malloc(TaskContext.PrimitiveIdBufferDataSize);
+#if DO_GUARD_SLOW
+		FMemory::Memzero(TaskContext.PrimitiveIdBufferData, TaskContext.PrimitiveIdBufferDataSize);
+#endif // DO_GUARD_SLOW
 		PrimitiveIdVertexBufferPoolEntry = GPrimitiveIdVertexBufferPool.Allocate(TaskContext.PrimitiveIdBufferDataSize);
 		TaskContext.MeshDrawCommands.Reserve(MaxNumDraws);
 		TaskContext.TempVisibleMeshDrawCommands.Reserve(MaxNumDraws);
@@ -1431,7 +1459,7 @@ void FParallelMeshDrawCommandPass::BuildRenderingCommands(FRDGBuilder& GraphBuil
 		if (MaxNumDraws > 0 && TaskContext.InstanceCullingContext.HasCullingCommands())
 		{
 			// 2. Run finalize culling commands pass
-			TaskContext.InstanceCullingContext.BuildRenderingCommands(GraphBuilder, GPUScene, TaskContext.InstanceCullingResult);
+			TaskContext.InstanceCullingContext.BuildRenderingCommands(GraphBuilder, GPUScene, TaskContext.View->DynamicPrimitiveCollector.GetPrimitiveIdRange(), TaskContext.InstanceCullingResult);
 			TaskContext.InstanceCullingResult.GetDrawParameters(OutInstanceCullingDrawParams);
 			return;
 		}
@@ -1455,10 +1483,20 @@ void FParallelMeshDrawCommandPass::BuildInstanceList(FRDGBuilder& GraphBuilder,
 			return;
 		}
 		// Run pass to build ID lists (temporary)
-		TaskContext.InstanceCullingContext.BuildRenderingCommands(GraphBuilder, GPUScene, OutParams);
+		TaskContext.InstanceCullingContext.BuildRenderingCommands(GraphBuilder, GPUScene, TaskContext.View->DynamicPrimitiveCollector.GetPrimitiveIdRange(), OutParams);
 	}
 #endif
 }
+/**
+ * Helper to upload and translate a primitive ID buffer.
+ */
+static void CopyPrimitiveIdBuffer(const int32* RESTRICT SrcData, int32* RESTRICT Data, int32 NumIds, const TRange<int32>& DynamicPrimitiveIdRange)
+{
+	for (int32 Index = 0; Index < NumIds; ++Index)
+	{
+		Data[Index] = TranslatePrimitiveId(SrcData[Index], DynamicPrimitiveIdRange.GetLowerBoundValue(), DynamicPrimitiveIdRange.GetUpperBoundValue());
+	}
+}
 
 void FParallelMeshDrawCommandPass::DispatchDraw(FParallelCommandListSet* ParallelCommandListSet, FRHICommandList& RHICmdList, const FInstanceCullingDrawParams* InstanceCullingDrawParams) const
 {
@@ -1502,11 +1540,15 @@ void FParallelMeshDrawCommandPass::DispatchDraw(FParallelCommandListSet* Paralle
 				VertexBuffer = PrimitiveIdsBuffer,
 				VertexBufferData = TaskContext.PrimitiveIdBufferData, 
 				VertexBufferDataSize = TaskContext.PrimitiveIdBufferDataSize,
-				PrimitiveIdVertexBufferPoolEntry = PrimitiveIdVertexBufferPoolEntry](FRHICommandListImmediate& CmdList)
+				PrimitiveIdVertexBufferPoolEntry = PrimitiveIdVertexBufferPoolEntry,
+				DynamicPrimitiveIdRange = TaskContext.View->DynamicPrimitiveCollector.GetPrimitiveIdRange()](FRHICommandListImmediate& CmdList)
 			{
 				// Upload vertex buffer data.
-				void* RESTRICT Data = (void* RESTRICT)CmdList.LockBuffer(VertexBuffer, 0, VertexBufferDataSize, RLM_WriteOnly);
-				FMemory::Memcpy(Data, VertexBufferData, VertexBufferDataSize);
+				CopyPrimitiveIdBuffer(
+					reinterpret_cast<const int32 * RESTRICT>(VertexBufferData), 
+					reinterpret_cast<int32 * RESTRICT>(CmdList.LockBuffer(VertexBuffer, 0, VertexBufferDataSize, RLM_WriteOnly)),
+					VertexBufferDataSize / sizeof(int32),
+					DynamicPrimitiveIdRange);
 				CmdList.UnlockBuffer(VertexBuffer);
 
 				FMemory::Free(VertexBufferData);
@@ -1585,8 +1627,11 @@ void FParallelMeshDrawCommandPass::DispatchDraw(FParallelCommandListSet* Paralle
 		if (TaskContext.bUseGPUScene)
 		{
 			// Can immediately upload vertex buffer data, as there is no parallel draw task.
-			void* RESTRICT Data = RHILockBuffer(PrimitiveIdVertexBufferPoolEntry.BufferRHI, 0, TaskContext.PrimitiveIdBufferDataSize, RLM_WriteOnly);
-			FMemory::Memcpy(Data, TaskContext.PrimitiveIdBufferData, TaskContext.PrimitiveIdBufferDataSize);
+			CopyPrimitiveIdBuffer(
+				reinterpret_cast<const int32* RESTRICT>(TaskContext.PrimitiveIdBufferData),
+				reinterpret_cast<int32* RESTRICT>(RHILockBuffer(PrimitiveIdVertexBufferPoolEntry.BufferRHI, 0, TaskContext.PrimitiveIdBufferDataSize, RLM_WriteOnly)),
+				TaskContext.PrimitiveIdBufferDataSize / sizeof(int32),
+				TaskContext.View->DynamicPrimitiveCollector.GetPrimitiveIdRange());
 			RHIUnlockBuffer(PrimitiveIdVertexBufferPoolEntry.BufferRHI);
 		}
 
diff --git a/Engine/Source/Runtime/Renderer/Private/MeshDrawCommands.h b/Engine/Source/Runtime/Renderer/Private/MeshDrawCommands.h
--- a/Engine/Source/Runtime/Renderer/Private/MeshDrawCommands.h
+++ b/Engine/Source/Runtime/Renderer/Private/MeshDrawCommands.h
@@ -224,7 +224,8 @@ extern void SortAndMergeDynamicPassMeshDrawCommands(
 	FMeshCommandOneFrameArray& VisibleMeshDrawCommands,
 	FDynamicMeshDrawCommandStorage& MeshDrawCommandStorage,
 	FRHIBuffer*& OutPrimitiveIdVertexBuffer,
-	uint32 InstanceFactor);
+	uint32 InstanceFactor,
+	const TRange<int32>& DynamicPrimitiveIdRange);
 
 void SubmitGPUInstancedMeshDrawCommandsRange(
 	const FMeshCommandOneFrameArray& VisibleMeshDrawCommands,
diff --git a/Engine/Source/Runtime/Renderer/Private/MeshPassProcessor.cpp b/Engine/Source/Runtime/Renderer/Private/MeshPassProcessor.cpp
--- a/Engine/Source/Runtime/Renderer/Private/MeshPassProcessor.cpp
+++ b/Engine/Source/Runtime/Renderer/Private/MeshPassProcessor.cpp
@@ -1304,7 +1304,10 @@ void DrawDynamicMeshPassPrivate(
 		FRHIBuffer* PrimitiveIdVertexBuffer = nullptr;
 
 		ApplyViewOverridesToMeshDrawCommands(View, VisibleMeshDrawCommands, DynamicMeshDrawCommandStorage, GraphicsMinimalPipelineStateSet, InNeedsShaderInitialisation);
-		SortAndMergeDynamicPassMeshDrawCommands(View.GetFeatureLevel(), VisibleMeshDrawCommands, DynamicMeshDrawCommandStorage, PrimitiveIdVertexBuffer, InstanceFactor);
+
+		check(View.bIsViewInfo);
+		const FViewInfo* ViewInfo = static_cast<const FViewInfo*>(&View);
+		SortAndMergeDynamicPassMeshDrawCommands(View.GetFeatureLevel(), VisibleMeshDrawCommands, DynamicMeshDrawCommandStorage, PrimitiveIdVertexBuffer, InstanceFactor, ViewInfo->DynamicPrimitiveCollector.GetPrimitiveIdRange());
 
 		SubmitMeshDrawCommandsRange(VisibleMeshDrawCommands, GraphicsMinimalPipelineStateSet, PrimitiveIdVertexBuffer, 0, bDynamicInstancing, 0, VisibleMeshDrawCommands.Num(), InstanceFactor, RHICmdList);
 	}
@@ -1375,9 +1378,8 @@ void FMeshPassProcessor::GetDrawCommandPrimitiveId(
 		}
 		else if (BatchElement.PrimitiveIdMode == PrimID_DynamicPrimitiveShaderData && ViewIfDynamicMeshCommand != nullptr)
 		{
-			checkSlow(ViewIfDynamicMeshCommand->bIsViewInfo);
-			const FViewInfo* ViewInfo = static_cast<const FViewInfo*>(ViewIfDynamicMeshCommand);
-			DrawPrimitiveId = ViewInfo->DynamicPrimitiveCollector.GetFinalId(BatchElement.DynamicPrimitiveShaderDataIndex);
+			// Mark using GPrimIDDynamicFlag (top bit) as we defer this to later.
+			DrawPrimitiveId = BatchElement.DynamicPrimitiveShaderDataIndex | GPrimIDDynamicFlag;
 		}
 		else
 		{
@@ -1473,15 +1475,21 @@ void FCachedPassMeshDrawListContext::FinalizeCommand(
 PassProcessorCreateFunction FPassProcessorManager::JumpTable[(int32)EShadingPath::Num][EMeshPass::Num] = {};
 EMeshPassFlags FPassProcessorManager::Flags[(int32)EShadingPath::Num][EMeshPass::Num] = {};
 
-#if defined(GPUCULL_TODO)
 
-FSimpleMeshDrawCommandPass::FSimpleMeshDrawCommandPass(const FSceneView& View, FInstanceCullingManager* InstanceCullingManager) :
+FSimpleMeshDrawCommandPass::FSimpleMeshDrawCommandPass(const FSceneView& View, FInstanceCullingManager* InstanceCullingManager, uint32 InstanceFactorIn) :
 	bNeedsInitialization(false),
 	DynamicPassMeshDrawListContext(DynamicMeshDrawCommandStorage, VisibleMeshDrawCommands, GraphicsMinimalPipelineStateSet, bNeedsInitialization)
 {
 	check(View.bIsViewInfo);
 	const FViewInfo* ViewInfo = static_cast<const FViewInfo*>(&View);
+	// GPUCULL_TODO: Non-Instanced stereo?
 	InstanceCullingContext = FInstanceCullingContext(InstanceCullingManager, TArrayView<const int32>(&ViewInfo->GPUSceneViewId, 1));
+	bDynamicInstancing = IsDynamicInstancingEnabled(ViewInfo->GetFeatureLevel());
+
+#if defined(GPUCULL_TODO)
+	// GPUCULL_TODO: Instanced stereo?
+	InstanceFactor = InstanceFactorIn;
+#endif // defined(GPUCULL_TODO)
 }
 
 void FSimpleMeshDrawCommandPass::BuildRenderingCommands(FRDGBuilder& GraphBuilder, const FSceneView& View, const FGPUScene& GPUScene, FInstanceCullingDrawParams& OutInstanceCullingDrawParams)
@@ -1489,11 +1497,11 @@ void FSimpleMeshDrawCommandPass::BuildRenderingCommands(FRDGBuilder& GraphBuilde
 	// NOTE: Everything up to InstanceCullingContext.BuildRenderingCommands could be peeled off into an async task.
 	ApplyViewOverridesToMeshDrawCommands(View, VisibleMeshDrawCommands, DynamicMeshDrawCommandStorage, GraphicsMinimalPipelineStateSet, bNeedsInitialization);
 
-	VisibleMeshDrawCommands.Sort(FCompareFMeshDrawCommands());
-
 	FInstanceCullingResult InstanceCullingResult;
+	VisibleMeshDrawCommands.Sort(FCompareFMeshDrawCommands());
 	if (GPUScene.IsEnabled())
 	{
+#if defined(GPUCULL_TODO)
 		int32 MaxInstances = 0;
 		int32 VisibleMeshDrawCommandsNum = 0;
 		int32 NewPassVisibleMeshDrawCommandsNum = 0;
@@ -1502,8 +1510,16 @@ void FSimpleMeshDrawCommandPass::BuildRenderingCommands(FRDGBuilder& GraphBuilde
 		SetupGPUInstancedDraws(InstanceCullingContext, VisibleMeshDrawCommands, MaxInstances, VisibleMeshDrawCommandsNum, NewPassVisibleMeshDrawCommandsNum);
 
 		// 2. Run finalize culling commands pass
-		InstanceCullingContext.BuildRenderingCommands(GraphBuilder, GPUScene, InstanceCullingResult);
+		check(View.bIsViewInfo);
+		const FViewInfo* ViewInfo = static_cast<const FViewInfo*>(&View);
+		InstanceCullingContext.BuildRenderingCommands(GraphBuilder, GPUScene, ViewInfo->DynamicPrimitiveCollector.GetPrimitiveIdRange(), InstanceCullingResult);
 		
+#else // defined(GPUCULL_TODO)
+		check(View.bIsViewInfo);
+		const FViewInfo* ViewInfo = static_cast<const FViewInfo*>(&View);
+		SortAndMergeDynamicPassMeshDrawCommands(View.GetFeatureLevel(), VisibleMeshDrawCommands, DynamicMeshDrawCommandStorage, PrimitiveIdVertexBuffer, InstanceFactor, ViewInfo->DynamicPrimitiveCollector.GetPrimitiveIdRange());
+
+#endif // defined(GPUCULL_TODO)
 		// Signal that scene primitives are supported, used for validation, the existence of a valid InstanceCullingResult is the required signal
 		bSupportsScenePrimitives = true;
 	}
@@ -1520,24 +1536,32 @@ void FSimpleMeshDrawCommandPass::SubmitDraw(FRHICommandListImmediate& RHICmdList
 {
 	if (VisibleMeshDrawCommands.Num() > 0)
 	{
-		FRHIBuffer* DrawIndirectArgsBuffer = nullptr;
-		FRHIBuffer* InstanceIdOffsetBuffer = nullptr;
+#if defined(GPUCULL_TODO)
+		if (bSupportsScenePrimitives)
+		{
+			FRHIBuffer* DrawIndirectArgsBuffer = nullptr;
+			FRHIBuffer* InstanceIdOffsetBuffer = nullptr;
+
+			if (InstanceCullingDrawParams.DrawIndirectArgsBuffer != nullptr && InstanceCullingDrawParams.InstanceIdOffsetBuffer != nullptr)
+			{
+				DrawIndirectArgsBuffer = InstanceCullingDrawParams.DrawIndirectArgsBuffer->GetRHI();
+				InstanceIdOffsetBuffer = InstanceCullingDrawParams.InstanceIdOffsetBuffer->GetRHI();
+			}
 
-		if (InstanceCullingDrawParams.DrawIndirectArgsBuffer != nullptr && InstanceCullingDrawParams.InstanceIdOffsetBuffer != nullptr)
+			SubmitGPUInstancedMeshDrawCommandsRange(
+				VisibleMeshDrawCommands,
+				GraphicsMinimalPipelineStateSet,
+				0,
+				VisibleMeshDrawCommands.Num(),
+				InstanceIdOffsetBuffer,
+				DrawIndirectArgsBuffer,
+				RHICmdList);
+		}
+		else
+#endif // defined(GPUCULL_TODO)
 		{
-			DrawIndirectArgsBuffer = InstanceCullingDrawParams.DrawIndirectArgsBuffer->GetRHI();
-			InstanceIdOffsetBuffer = InstanceCullingDrawParams.InstanceIdOffsetBuffer->GetRHI();
+			SubmitMeshDrawCommandsRange(VisibleMeshDrawCommands, GraphicsMinimalPipelineStateSet, PrimitiveIdVertexBuffer, 0, bDynamicInstancing, 0, VisibleMeshDrawCommands.Num(), InstanceFactor, RHICmdList);
 		}
-
-		SubmitGPUInstancedMeshDrawCommandsRange(
-			VisibleMeshDrawCommands,
-			GraphicsMinimalPipelineStateSet,
-			0,
-			VisibleMeshDrawCommands.Num(),
-			InstanceIdOffsetBuffer,
-			DrawIndirectArgsBuffer,
-			RHICmdList);
 	}
 }
 
-#endif // GPUCULL_TODO
diff --git a/Engine/Source/Runtime/Renderer/Private/SceneRendering.cpp b/Engine/Source/Runtime/Renderer/Private/SceneRendering.cpp
--- a/Engine/Source/Runtime/Renderer/Private/SceneRendering.cpp
+++ b/Engine/Source/Runtime/Renderer/Private/SceneRendering.cpp
@@ -1947,6 +1947,10 @@ FViewInfo* FViewInfo::CreateSnapshot() const
 		Result->ParallelMeshDrawCommandPasses[i].InitCreateSnapshot();
 	}
 	
+	// Ensure the internal state is maintained, needed because we've just Memcpy'd the member data.
+	static_assert(TIsTriviallyDestructible<FGPUScenePrimitiveCollector>::Value != 0, "The destructor is not invoked properly because of FMemory::Memcpy(*Result, *this) above");
+	Result->DynamicPrimitiveCollector = FGPUScenePrimitiveCollector(DynamicPrimitiveCollector);
+
 	Result->bIsSnapshot = true;
 	ViewInfoSnapshots.Add(Result);
 	return Result;
diff --git a/Engine/Source/Runtime/Renderer/Private/SceneVisibility.cpp b/Engine/Source/Runtime/Renderer/Private/SceneVisibility.cpp
--- a/Engine/Source/Runtime/Renderer/Private/SceneVisibility.cpp
+++ b/Engine/Source/Runtime/Renderer/Private/SceneVisibility.cpp
@@ -3053,11 +3053,6 @@ void FSceneRenderer::GatherDynamicMeshElements(
 		}
 	}
 	MeshCollector.ProcessTasks();
-
-	for (FViewInfo& View : InViews)
-	{
-		View.DynamicPrimitiveCollector.Commit();
-	}
 }
 
 /**
diff --git a/Engine/Source/Runtime/Renderer/Private/ShadowRendering.cpp b/Engine/Source/Runtime/Renderer/Private/ShadowRendering.cpp
--- a/Engine/Source/Runtime/Renderer/Private/ShadowRendering.cpp
+++ b/Engine/Source/Runtime/Renderer/Private/ShadowRendering.cpp
@@ -837,10 +837,9 @@ void FProjectedShadowInfo::SetupProjectionStencilMask(
 			RHICmdList.SetScissorRect(true, View->ViewRect.Min.X, View->ViewRect.Min.Y, View->ViewRect.Max.X, View->ViewRect.Max.Y);
 		}
 
-		const FShadowMeshDrawCommandPass& ProjectionStencilingPass = ProjectionStencilingPasses[ViewIndex];
-		if (ProjectionStencilingPass.VisibleMeshDrawCommands.Num() > 0)
+		if (ViewIndex < ProjectionStencilingPasses.Num())
 		{
-			SubmitMeshDrawCommands(ProjectionStencilingPass.VisibleMeshDrawCommands, GraphicsMinimalPipelineStateSet, ProjectionStencilingPass.PrimitiveIdVertexBuffer, 0, bDynamicInstancing, bIsInstancedStereoEmulated ? 2 : 1, RHICmdList);
+			ProjectionStencilingPasses[ViewIndex]->SubmitDraw(RHICmdList, InstanceCullingDrawParams);
 		}
 
 		// Restore viewport
@@ -978,10 +977,9 @@ void FProjectedShadowInfo::SetupProjectionStencilMask(
 		// if rendering modulated shadows mask out subject mesh elements to prevent self shadowing.
 		if (bMobileModulatedProjections && !CVarEnableModulatedSelfShadow.GetValueOnRenderThread())
 		{
-			const FShadowMeshDrawCommandPass& ProjectionStencilingPass = ProjectionStencilingPasses[ViewIndex];
-			if (ProjectionStencilingPass.VisibleMeshDrawCommands.Num() > 0)
+			if (ViewIndex < ProjectionStencilingPasses.Num())
 			{
-				SubmitMeshDrawCommands(ProjectionStencilingPass.VisibleMeshDrawCommands, GraphicsMinimalPipelineStateSet, ProjectionStencilingPass.PrimitiveIdVertexBuffer, 0, bDynamicInstancing, 1, RHICmdList);
+				ProjectionStencilingPasses[ViewIndex]->SubmitDraw(RHICmdList, InstanceCullingDrawParams);
 			}
 		}
 	}
@@ -1102,11 +1100,11 @@ void FProjectedShadowInfo::RenderProjection(
 		PassParameters->ShadowTexture1 = GraphBuilder.RegisterExternalTexture(RenderTargets.ColorTargets[1]);
 	}
 
-	// GPUCULL_TODO: get rid of const cast
 	if (ViewIndex < ProjectionStencilingPasses.Num())
 	{
-		FShadowMeshDrawCommandPass& ProjectionStencilingPass = const_cast<FShadowMeshDrawCommandPass&>(ProjectionStencilingPasses[ViewIndex]);
-		// GPUCULL_TODO: ProjectionStencilingPass.BuildRenderingCommands(GraphBuilder, *View, *SceneRender->Scene, PassParameters->InstanceCullingDrawParams);
+		// GPUCULL_TODO: get rid of const cast
+		FSimpleMeshDrawCommandPass& ProjectionStencilingPass = *const_cast<FSimpleMeshDrawCommandPass*>(ProjectionStencilingPasses[ViewIndex]);
+		ProjectionStencilingPass.BuildRenderingCommands(GraphBuilder, *View, *SceneRender->Scene, PassParameters->InstanceCullingDrawParams);
 	}
 
 	GraphBuilder.AddPass(
diff --git a/Engine/Source/Runtime/Renderer/Private/ShadowRendering.h b/Engine/Source/Runtime/Renderer/Private/ShadowRendering.h
--- a/Engine/Source/Runtime/Renderer/Private/ShadowRendering.h
+++ b/Engine/Source/Runtime/Renderer/Private/ShadowRendering.h
@@ -232,13 +232,6 @@ BEGIN_GLOBAL_SHADER_PARAMETER_STRUCT(FMobileShadowDepthPassUniformParameters,)
 	SHADER_PARAMETER_ARRAY(FMatrix, ShadowViewProjectionMatrices, [6])
 END_GLOBAL_SHADER_PARAMETER_STRUCT()
 
-class FShadowMeshDrawCommandPass
-{
-public:
-	FMeshCommandOneFrameArray VisibleMeshDrawCommands;
-	FRHIBuffer* PrimitiveIdVertexBuffer;
-};
-
 /**
  * Information about a projected shadow.
  */
@@ -733,7 +726,7 @@ private:
 	FMeshCommandOneFrameArray ShadowDepthPassVisibleCommands;
 	FParallelMeshDrawCommandPass ShadowDepthPass;
 
-	TArray<FShadowMeshDrawCommandPass, TInlineAllocator<2>> ProjectionStencilingPasses;
+	TArray<FSimpleMeshDrawCommandPass *, TInlineAllocator<8>> ProjectionStencilingPasses;
 
 	FDynamicMeshDrawCommandStorage DynamicMeshDrawCommandStorage;
 	FGraphicsMinimalPipelineStateSet GraphicsMinimalPipelineStateSet;
diff --git a/Engine/Source/Runtime/Renderer/Private/ShadowSetup.cpp b/Engine/Source/Runtime/Renderer/Private/ShadowSetup.cpp
--- a/Engine/Source/Runtime/Renderer/Private/ShadowSetup.cpp
+++ b/Engine/Source/Runtime/Renderer/Private/ShadowSetup.cpp
@@ -2121,10 +2121,11 @@ void FProjectedShadowInfo::SetupMeshDrawCommandsForProjectionStenciling(FSceneRe
 		for (int32 ViewIndex = 0; ViewIndex < Renderer.Views.Num(); ViewIndex++)
 		{
 			const FViewInfo& View = Renderer.Views[ViewIndex];
-			ProjectionStencilingPasses.Add(FShadowMeshDrawCommandPass());
-			FShadowMeshDrawCommandPass& ProjectionStencilingPass = ProjectionStencilingPasses[ViewIndex];
 
-			FDynamicPassMeshDrawListContext ProjectionStencilingContext(DynamicMeshDrawCommandStorage, ProjectionStencilingPass.VisibleMeshDrawCommands, GraphicsMinimalPipelineStateSet, NeedsShaderInitialisation);
+			const uint32 InstanceFactor = View.bIsInstancedStereoEnabled && !View.bIsMultiViewEnabled && IStereoRendering::IsStereoEyeView(View) ? 2 : 1;
+			// Note: not using FMemstack as the data is not cleared in the right order.
+			ProjectionStencilingPasses.Add(new FSimpleMeshDrawCommandPass(View, &InstanceCullingManager, InstanceFactor));
+			FSimpleMeshDrawCommandPass& ProjectionStencilingPass = *ProjectionStencilingPasses[ViewIndex];
 
 			FMeshPassProcessorRenderState DrawRenderState;
 			DrawRenderState.SetBlendState(TStaticBlendState<CW_NONE>::GetRHI());
@@ -2164,7 +2165,7 @@ void FProjectedShadowInfo::SetupMeshDrawCommandsForProjectionStenciling(FSceneRe
 				DDM_AllOccluders,
 				false,
 				false,
-				&ProjectionStencilingContext);
+				ProjectionStencilingPass.GetDynamicPassMeshDrawListContext());
 
 			// Pre-shadows mask by receiver elements, self-shadow mask by subject elements.
 			// Note that self-shadow pre-shadows still mask by receiver elements.
@@ -2216,13 +2217,6 @@ void FProjectedShadowInfo::SetupMeshDrawCommandsForProjectionStenciling(FSceneRe
 					DepthPassMeshProcessor.AddMeshBatch(StaticMesh, DefaultBatchElementMask, StaticMesh.PrimitiveSceneInfo->Proxy);
 				}
 			}
-
-			ApplyViewOverridesToMeshDrawCommands(View, ProjectionStencilingPass.VisibleMeshDrawCommands, DynamicMeshDrawCommandStorage, GraphicsMinimalPipelineStateSet, NeedsShaderInitialisation);
-
-			// If instanced stereo is enabled, we need to render each view of the stereo pair using the instanced stereo transform to avoid bias issues.
-			// TODO: Support instanced stereo properly in the projection stenciling pass.
-			const uint32 InstanceFactor = View.bIsInstancedStereoEnabled && !View.bIsMultiViewEnabled && IStereoRendering::IsStereoEyeView(View) ? 2 : 1;
-			SortAndMergeDynamicPassMeshDrawCommands(Renderer.FeatureLevel, ProjectionStencilingPass.VisibleMeshDrawCommands, DynamicMeshDrawCommandStorage, ProjectionStencilingPass.PrimitiveIdVertexBuffer, InstanceFactor);
 		}
 	}
 }
@@ -2280,8 +2274,6 @@ void FProjectedShadowInfo::GatherDynamicMeshElements(FSceneRenderer& Renderer, F
 		Renderer.MeshCollector.ProcessTasks();
 	}
 
-	ShadowDepthView->DynamicPrimitiveCollector.Commit();
-
 	SetupMeshDrawCommandsForShadowDepth(Renderer, InstanceCullingManager);
 	SetupMeshDrawCommandsForProjectionStenciling(Renderer, InstanceCullingManager);
 }
@@ -2377,6 +2369,10 @@ void FProjectedShadowInfo::ClearTransientArrays()
 
 	SubjectMeshCommandBuildRequests.Empty();
 
+	for (auto ProjectionStencilingPass : ProjectionStencilingPasses)
+	{
+		delete ProjectionStencilingPass;
+	}
 	ProjectionStencilingPasses.Reset();
 
 	DynamicMeshDrawCommandStorage.MeshDrawCommands.Empty();
diff --git a/Engine/Source/Runtime/Renderer/Private/TranslucentLighting.cpp b/Engine/Source/Runtime/Renderer/Private/TranslucentLighting.cpp
--- a/Engine/Source/Runtime/Renderer/Private/TranslucentLighting.cpp
+++ b/Engine/Source/Runtime/Renderer/Private/TranslucentLighting.cpp
@@ -485,7 +485,55 @@ void FProjectedShadowInfo::RenderTranslucencyDepths(FRDGBuilder& GraphBuilder, F
 	PassParameters->View = ShadowDepthView->ViewUniformBuffer;
 	PassParameters->PassUniformBuffer = PassUniformBuffer;
 	PassParameters->RenderTargets = InRenderTargets;
-	PassParameters->View = ShadowDepthView->ViewUniformBuffer;
+
+	FSimpleMeshDrawCommandPass* SimpleMeshDrawCommandPass = GraphBuilder.AllocObject<FSimpleMeshDrawCommandPass>(*ShadowDepthView, &InstanceCullingManager);
+
+	FMeshPassProcessorRenderState DrawRenderState;
+	DrawRenderState.SetDepthStencilState(TStaticDepthStencilState<false, CF_Always>::GetRHI());
+	DrawRenderState.SetBlendState(TStaticBlendState<
+		CW_RGBA, BO_Add, BF_One, BF_One, BO_Add, BF_One, BF_One,
+		CW_RGBA, BO_Add, BF_One, BF_One, BO_Add, BF_One, BF_One>::GetRHI());
+
+
+	FTranslucencyDepthPassMeshProcessor TranslucencyDepthPassMeshProcessor(
+		SceneRenderer->Scene,
+		ShadowDepthView,
+		DrawRenderState,
+		this,
+		SimpleMeshDrawCommandPass->GetDynamicPassMeshDrawListContext());
+
+	for (int32 MeshBatchIndex = 0; MeshBatchIndex < DynamicSubjectTranslucentMeshElements.Num(); MeshBatchIndex++)
+	{
+		const FMeshBatchAndRelevance& MeshAndRelevance = DynamicSubjectTranslucentMeshElements[MeshBatchIndex];
+		const uint64 BatchElementMask = ~0ull;
+		TranslucencyDepthPassMeshProcessor.AddMeshBatch(*MeshAndRelevance.Mesh, BatchElementMask, MeshAndRelevance.PrimitiveSceneProxy);
+	}
+
+	for (int32 PrimitiveIndex = 0; PrimitiveIndex < SubjectTranslucentPrimitives.Num(); PrimitiveIndex++)
+	{
+		const FPrimitiveSceneInfo* PrimitiveSceneInfo = SubjectTranslucentPrimitives[PrimitiveIndex];
+		int32 PrimitiveId = PrimitiveSceneInfo->GetIndex();
+		FPrimitiveViewRelevance ViewRelevance = ShadowDepthView->PrimitiveViewRelevanceMap[PrimitiveId];
+
+		if (!ViewRelevance.bInitializedThisFrame)
+		{
+			// Compute the subject primitive's view relevance since it wasn't cached
+			ViewRelevance = PrimitiveSceneInfo->Proxy->GetViewRelevance(ShadowDepthView);
+		}
+
+		if (ViewRelevance.bDrawRelevance && ViewRelevance.bStaticRelevance)
+		{
+			for (int32 MeshIndex = 0; MeshIndex < PrimitiveSceneInfo->StaticMeshes.Num(); MeshIndex++)
+			{
+				const FStaticMeshBatch& StaticMeshBatch = PrimitiveSceneInfo->StaticMeshes[MeshIndex];
+				const uint64 DefaultBatchElementMask = ~0ul;
+				TranslucencyDepthPassMeshProcessor.AddMeshBatch(StaticMeshBatch, DefaultBatchElementMask, StaticMeshBatch.PrimitiveSceneInfo->Proxy, StaticMeshBatch.Id);
+			}
+		}
+	}
+
+	SimpleMeshDrawCommandPass->BuildRenderingCommands(GraphBuilder, *ShadowDepthView, SceneRenderer->Scene->GPUScene, PassParameters->InstanceCullingDrawParams);
+
 
 	FString EventName;
 #if WANTS_DRAW_MESH_EVENTS
@@ -499,7 +547,7 @@ void FProjectedShadowInfo::RenderTranslucencyDepths(FRDGBuilder& GraphBuilder, F
 		RDG_EVENT_NAME("%s", *EventName),
 		PassParameters,
 		ERDGPassFlags::Raster,
-		[this, SceneRenderer](FRHICommandListImmediate& RHICmdList)
+		[this, SimpleMeshDrawCommandPass, PassParameters](FRHICommandListImmediate& RHICmdList)
 	{
 		FMeshPassProcessorRenderState DrawRenderState;
 
@@ -525,62 +573,7 @@ void FProjectedShadowInfo::RenderTranslucencyDepths(FRDGBuilder& GraphBuilder, F
 			(Y + BorderSize + ResolutionY),
 			1.0f
 		);
-
-		DrawRenderState.SetDepthStencilState(TStaticDepthStencilState<false, CF_Always>::GetRHI());
-		DrawRenderState.SetBlendState(TStaticBlendState<
-			CW_RGBA, BO_Add, BF_One, BF_One, BO_Add, BF_One, BF_One,
-			CW_RGBA, BO_Add, BF_One, BF_One, BO_Add, BF_One, BF_One>::GetRHI());
-
-
-		FMeshCommandOneFrameArray VisibleMeshDrawCommands;
-		FDynamicPassMeshDrawListContext TranslucencyDepthContext(DynamicMeshDrawCommandStorage, VisibleMeshDrawCommands, GraphicsMinimalPipelineStateSet, NeedsShaderInitialisation);
-
-		FTranslucencyDepthPassMeshProcessor TranslucencyDepthPassMeshProcessor(
-			SceneRenderer->Scene,
-			ShadowDepthView,
-			DrawRenderState,
-			this,
-			&TranslucencyDepthContext);
-
-		for (int32 MeshBatchIndex = 0; MeshBatchIndex < DynamicSubjectTranslucentMeshElements.Num(); MeshBatchIndex++)
-		{
-			const FMeshBatchAndRelevance& MeshAndRelevance = DynamicSubjectTranslucentMeshElements[MeshBatchIndex];
-			const uint64 BatchElementMask = ~0ull;
-			TranslucencyDepthPassMeshProcessor.AddMeshBatch(*MeshAndRelevance.Mesh, BatchElementMask, MeshAndRelevance.PrimitiveSceneProxy);
-		}
-
-		for (int32 PrimitiveIndex = 0; PrimitiveIndex < SubjectTranslucentPrimitives.Num(); PrimitiveIndex++)
-		{
-			const FPrimitiveSceneInfo* PrimitiveSceneInfo = SubjectTranslucentPrimitives[PrimitiveIndex];
-			int32 PrimitiveId = PrimitiveSceneInfo->GetIndex();
-			FPrimitiveViewRelevance ViewRelevance = ShadowDepthView->PrimitiveViewRelevanceMap[PrimitiveId];
-
-			if (!ViewRelevance.bInitializedThisFrame)
-			{
-				// Compute the subject primitive's view relevance since it wasn't cached
-				ViewRelevance = PrimitiveSceneInfo->Proxy->GetViewRelevance(ShadowDepthView);
-			}
-
-			if (ViewRelevance.bDrawRelevance && ViewRelevance.bStaticRelevance)
-			{
-				for (int32 MeshIndex = 0; MeshIndex < PrimitiveSceneInfo->StaticMeshes.Num(); MeshIndex++)
-				{
-					const FStaticMeshBatch& StaticMeshBatch = PrimitiveSceneInfo->StaticMeshes[MeshIndex];
-					const uint64 DefaultBatchElementMask = ~0ul;
-					TranslucencyDepthPassMeshProcessor.AddMeshBatch(StaticMeshBatch, DefaultBatchElementMask, StaticMeshBatch.PrimitiveSceneInfo->Proxy, StaticMeshBatch.Id);
-				}
-			}
-		}
-
-		if (VisibleMeshDrawCommands.Num() > 0)
-		{
-			const bool bDynamicInstancing = IsDynamicInstancingEnabled(ShadowDepthView->FeatureLevel);
-
-			FRHIBuffer* PrimitiveIdVertexBuffer = nullptr;
-			ApplyViewOverridesToMeshDrawCommands(*ShadowDepthView, VisibleMeshDrawCommands, DynamicMeshDrawCommandStorage, GraphicsMinimalPipelineStateSet, NeedsShaderInitialisation);
-			SortAndMergeDynamicPassMeshDrawCommands(SceneRenderer->FeatureLevel, VisibleMeshDrawCommands, DynamicMeshDrawCommandStorage, PrimitiveIdVertexBuffer, 1);
-			SubmitMeshDrawCommands(VisibleMeshDrawCommands, GraphicsMinimalPipelineStateSet, PrimitiveIdVertexBuffer, 0, bDynamicInstancing, 1, RHICmdList);
-		}
+		SimpleMeshDrawCommandPass->SubmitDraw(RHICmdList, PassParameters->InstanceCullingDrawParams);
 	});
 }
 
diff --git a/Engine/Source/Runtime/Renderer/Private/VT/RuntimeVirtualTextureRender.cpp b/Engine/Source/Runtime/Renderer/Private/VT/RuntimeVirtualTextureRender.cpp
--- a/Engine/Source/Runtime/Renderer/Private/VT/RuntimeVirtualTextureRender.cpp
+++ b/Engine/Source/Runtime/Renderer/Private/VT/RuntimeVirtualTextureRender.cpp
@@ -544,10 +544,18 @@ namespace RuntimeVirtualTexture
 		if (AllocatedCommands.Num() > 0)
 		{
 			FRHIBuffer* PrimitiveIdsBuffer;
+
+#if defined(GPUCULL_TODO)
+			// GPUCULL_TODO: workaround for the fact that DrawDynamicMeshPassPrivate et al. don't work with GPU-Scene instancing
+			//               we don't support dynamic instancing for this path since we require one primitive per draw command
+			//               This is because the stride on the instance data buffer is set to 0 so only the first will ever be fetched.
+			const bool bDynamicInstancing = false;
+#else
 			const bool bDynamicInstancing = IsDynamicInstancingEnabled(View->FeatureLevel);
+#endif
 			const uint32 InstanceFactor = 1;
 
-			SortAndMergeDynamicPassMeshDrawCommands(View->FeatureLevel, AllocatedCommands, MeshDrawCommandStorage, PrimitiveIdsBuffer, InstanceFactor);
+			SortAndMergeDynamicPassMeshDrawCommands(View->FeatureLevel, AllocatedCommands, MeshDrawCommandStorage, PrimitiveIdsBuffer, InstanceFactor, View->DynamicPrimitiveCollector.GetPrimitiveIdRange());
 			SubmitMeshDrawCommands(AllocatedCommands, GraphicsMinimalPipelineStateSet, PrimitiveIdsBuffer, 0, bDynamicInstancing, InstanceFactor, RHICmdList);
 		}
 	}
diff --git a/Engine/Source/Runtime/Renderer/Public/InstanceCulling/InstanceCullingContext.h b/Engine/Source/Runtime/Renderer/Public/InstanceCulling/InstanceCullingContext.h
--- a/Engine/Source/Runtime/Renderer/Public/InstanceCulling/InstanceCullingContext.h
+++ b/Engine/Source/Runtime/Renderer/Public/InstanceCulling/InstanceCullingContext.h
@@ -94,12 +94,12 @@ public:
 	void AddInstanceRunToCullingCommand(int32 ScenePrimitiveId, const uint32* Runs, uint32 NumRuns);
 
 
-	void BuildRenderingCommands(FRDGBuilder& GraphBuilder, const FGPUScene& GPUScene, FInstanceCullingResult& Results) const;
+	void BuildRenderingCommands(FRDGBuilder& GraphBuilder, const FGPUScene& GPUScene, const TRange<int32> &DynamicPrimitiveIdRange, FInstanceCullingResult& Results) const;
 
 	/**
 	 * Build instance ID lists & render commands, but without extracting results and also accepting RGD resources for output buffers etc (rather than registering global resources)
 	 */
-	void BuildRenderingCommands(FRDGBuilder& GraphBuilder, FGPUScene& GPUScene, FInstanceCullingRdgParams& Params) const;
+	void BuildRenderingCommands(FRDGBuilder& GraphBuilder, FGPUScene& GPUScene, const TRange<int32>& DynamicPrimitiveIdRange, FInstanceCullingRdgParams& Params) const;
 
 	inline bool HasCullingCommands() const { return CullingCommands.Num() > 0; 	}
 
diff --git a/Engine/Source/Runtime/Renderer/Public/MeshPassProcessor.h b/Engine/Source/Runtime/Renderer/Public/MeshPassProcessor.h
--- a/Engine/Source/Runtime/Renderer/Public/MeshPassProcessor.h
+++ b/Engine/Source/Runtime/Renderer/Public/MeshPassProcessor.h
@@ -1972,15 +1972,13 @@ private:
 	uint32 RayTracingInstanceIndex;
 };
 
-#if defined(GPUCULL_TODO)
-
 /**
- * Same as the parallel one, but intended for use with simpler tasks where the overhead and complexity of parallel is not justified.
+ * Similar to the parallel one, but intended for use with simpler tasks where the overhead and complexity of parallel is not justified.
  */
 class FSimpleMeshDrawCommandPass
 {
 public:
-	RENDERER_API FSimpleMeshDrawCommandPass(const FSceneView& View, FInstanceCullingManager* InstanceCullingManager);
+	RENDERER_API FSimpleMeshDrawCommandPass(const FSceneView& View, FInstanceCullingManager* InstanceCullingManager, uint32 InstanceFactorIn = 1);
 
 	/**
 	 * Run post-instance culling job to create the render commands and instance ID lists and optionally vertex instance data.
@@ -1993,6 +1991,8 @@ public:
 	FDynamicPassMeshDrawListContext* GetDynamicPassMeshDrawListContext() { return &DynamicPassMeshDrawListContext; }
 
 private:
+	FSimpleMeshDrawCommandPass() = delete;
+	FSimpleMeshDrawCommandPass(const FSimpleMeshDrawCommandPass&) = delete;
 	FMeshCommandOneFrameArray VisibleMeshDrawCommands;
 	FInstanceCullingContext InstanceCullingContext;
 	FGraphicsMinimalPipelineStateSet GraphicsMinimalPipelineStateSet;
@@ -2004,6 +2004,11 @@ private:
 	bool bSupportsScenePrimitives = false;
 
 	bool bNeedsInitialization = false;
+	bool bDynamicInstancing = false;
+	uint32 InstanceFactor = 1;
+
+	// GPUCULL_TODO: Only for legacy path
+	FRHIBuffer* PrimitiveIdVertexBuffer = nullptr;
 };
 
 // GPUCULL_TODO: Write documentation
@@ -2142,5 +2147,3 @@ void AddSimpleMeshPass(FRDGBuilder& GraphBuilder, PassParametersType* PassParame
 }
 
 #endif
-
-#endif // GPUCULL_TODO
