CIS fix initialization order

diff --git a/Engine/Source/Runtime/Renderer/Private/MeshPassProcessor.cpp b/Engine/Source/Runtime/Renderer/Private/MeshPassProcessor.cpp
--- a/Engine/Source/Runtime/Renderer/Private/MeshPassProcessor.cpp
+++ b/Engine/Source/Runtime/Renderer/Private/MeshPassProcessor.cpp
@@ -1477,7 +1477,6 @@ EMeshPassFlags FPassProcessorManager::Flags[(int32)EShadingPath::Num][EMeshPass:
 
 
 FSimpleMeshDrawCommandPass::FSimpleMeshDrawCommandPass(const FSceneView& View, FInstanceCullingManager* InstanceCullingManager, uint32 InstanceFactorIn) :
-	bNeedsInitialization(false),
 	DynamicPassMeshDrawListContext(DynamicMeshDrawCommandStorage, VisibleMeshDrawCommands, GraphicsMinimalPipelineStateSet, bNeedsInitialization)
 {
 	check(View.bIsViewInfo);
