Added new ray tracing instance culling mode r.RayTracing.Culling=2, which allows instances that are in front of the camera to be culled. Previous mode r.RayTracing.Culling=1 only allowed culling for instances behind the camera.

diff --git a/Engine/Source/Runtime/Renderer/Private/DeferredShadingRenderer.cpp b/Engine/Source/Runtime/Renderer/Private/DeferredShadingRenderer.cpp
--- a/Engine/Source/Runtime/Renderer/Private/DeferredShadingRenderer.cpp
+++ b/Engine/Source/Runtime/Renderer/Private/DeferredShadingRenderer.cpp
@@ -173,7 +173,8 @@ static TAutoConsoleVariable<int32> CVarRayTracingCulling(
 	0,
 	TEXT("Enable culling in ray tracing for objects that are behind the camera\n")
 	TEXT(" 0: Culling disabled (default)\n")
-	TEXT(" 1: Culling by distance and solid angle enabled"),
+	TEXT(" 1: Culling by distance and solid angle enabled. Only cull objects behind camera.\n")
+	TEXT(" 2: Culling by distance and solid angle enabled. Cull objects in front and behind camera."),
 	ECVF_RenderThreadSafe);
 
 static TAutoConsoleVariable<float> CVarRayTracingCullingRadius(
@@ -645,6 +646,7 @@ bool FDeferredShadingSceneRenderer::GatherRayTracingWorldInstances(FRHICommandLi
 		const float AngleThresholdRatio = FMath::Tan(CullAngleThreshold * PI / 180.0f);
 		const FVector ViewOrigin = ReferenceView.ViewMatrices.GetViewOrigin();
 		const FVector ViewDirection = ReferenceView.GetViewDirection();
+		const bool bCullAllObjects = CullInRayTracing == 2;
 
 		for (int PrimitiveIndex = 0; PrimitiveIndex < Scene->PrimitiveSceneProxies.Num(); PrimitiveIndex++)
 		{
@@ -681,9 +683,9 @@ bool FDeferredShadingSceneRenderer::GatherRayTracingWorldInstances(FRHICommandLi
 				const FVector ObjectCenter = ObjectBounds.Origin + 0.5*ObjectBounds.BoxExtent;
 				const FVector CameraToObjectCenter = FVector(ObjectCenter - ViewOrigin);
 
-				const bool bIsBehindCamera = FVector::DotProduct(ViewDirection, CameraToObjectCenter) < -ObjectRadius;
+				const bool bConsiderCulling = bCullAllObjects || FVector::DotProduct(ViewDirection, CameraToObjectCenter) < -ObjectRadius;
 
-				if (bIsBehindCamera)
+				if (bConsiderCulling)
 				{
 					const float CameraToObjectCenterLength = CameraToObjectCenter.Size();
 					const bool bIsFarEnoughToCull = CameraToObjectCenterLength > (CullingRadius + ObjectRadius);
