Fix single layer water rendering with GPU-Scene instance culling.

diff --git a/Engine/Source/Runtime/Renderer/Private/SingleLayerWaterRendering.cpp b/Engine/Source/Runtime/Renderer/Private/SingleLayerWaterRendering.cpp
--- a/Engine/Source/Runtime/Renderer/Private/SingleLayerWaterRendering.cpp
+++ b/Engine/Source/Runtime/Renderer/Private/SingleLayerWaterRendering.cpp
@@ -687,6 +687,7 @@ BEGIN_SHADER_PARAMETER_STRUCT(FSingleLayerWaterPassParameters, )
 	SHADER_PARAMETER_STRUCT_INCLUDE(FViewShaderParameters, View)
 	SHADER_PARAMETER_STRUCT_REF(FReflectionCaptureShaderData, ReflectionCapture)
 	SHADER_PARAMETER_RDG_UNIFORM_BUFFER(FOpaqueBasePassUniformParameters, BasePass)
+	SHADER_PARAMETER_STRUCT_INCLUDE(FInstanceCullingDrawParams, InstanceCullingDrawParams)
 	RENDER_TARGET_BINDING_SLOTS()
 END_SHADER_PARAMETER_STRUCT()
 
@@ -724,15 +725,14 @@ void FDeferredShadingSceneRenderer::RenderSingleLayerWaterInner(
 		RDG_EVENT_SCOPE_CONDITIONAL(GraphBuilder, Views.Num() > 1, "View%d", ViewIndex);
 		View.BeginRenderView();
 
-		// GPUCULL_TODO: Ignore for now as water does not work with instancing anyway
-		// View.ParallelMeshDrawCommandPasses[EMeshPass::SingleLayerWaterPass].BuildRenderingCommands(GraphBuilder, Scene->GPUScene);
-
 		FSingleLayerWaterPassParameters* PassParameters = GraphBuilder.AllocParameters<FSingleLayerWaterPassParameters>();
 		PassParameters->View = View.GetShaderParameters();
 		PassParameters->ReflectionCapture = View.ReflectionCaptureUniformBuffer;
 		PassParameters->BasePass = CreateOpaqueBasePassUniformBuffer(GraphBuilder, View, ViewIndex, {}, {}, &SceneWithoutWaterTextures);
 		PassParameters->RenderTargets = RenderTargets;
 
+		View.ParallelMeshDrawCommandPasses[EMeshPass::SingleLayerWaterPass].BuildRenderingCommands(GraphBuilder, Scene->GPUScene, PassParameters->InstanceCullingDrawParams);
+
 		if (bRenderInParallel)
 		{
 			GraphBuilder.AddPass(
@@ -742,7 +742,7 @@ void FDeferredShadingSceneRenderer::RenderSingleLayerWaterInner(
 				[this, &View, PassParameters](FRHICommandListImmediate& RHICmdList)
 			{
 				FRDGParallelCommandListSet ParallelCommandListSet(RHICmdList, GET_STATID(STAT_CLP_WaterSingleLayerPass), *this, View, FParallelCommandListBindings(PassParameters));
-				View.ParallelMeshDrawCommandPasses[EMeshPass::SingleLayerWaterPass].DispatchDraw(&ParallelCommandListSet, RHICmdList);
+				View.ParallelMeshDrawCommandPasses[EMeshPass::SingleLayerWaterPass].DispatchDraw(&ParallelCommandListSet, RHICmdList, &PassParameters->InstanceCullingDrawParams);
 			});
 		}
 		else
@@ -751,10 +751,10 @@ void FDeferredShadingSceneRenderer::RenderSingleLayerWaterInner(
 				RDG_EVENT_NAME("SingleLayerWater"),
 				PassParameters,
 				ERDGPassFlags::Raster,
-				[this, &View](FRHICommandListImmediate& RHICmdList)
+				[this, &View, PassParameters](FRHICommandListImmediate& RHICmdList)
 			{
 				SetStereoViewport(RHICmdList, View, 1.0f);
-				View.ParallelMeshDrawCommandPasses[EMeshPass::SingleLayerWaterPass].DispatchDraw(nullptr, RHICmdList);
+				View.ParallelMeshDrawCommandPasses[EMeshPass::SingleLayerWaterPass].DispatchDraw(nullptr, RHICmdList, &PassParameters->InstanceCullingDrawParams);
 			});
 		}
 	}
