Removed RDG helper passes from shadow rendering to reduce number of total number of RDG passes in lighting phase.

#rnx

diff --git a/Engine/Source/Runtime/Renderer/Private/ShadowRendering.cpp b/Engine/Source/Runtime/Renderer/Private/ShadowRendering.cpp
--- a/Engine/Source/Runtime/Renderer/Private/ShadowRendering.cpp
+++ b/Engine/Source/Runtime/Renderer/Private/ShadowRendering.cpp
@@ -1000,6 +1000,7 @@ void FProjectedShadowInfo::RenderProjection(
 	const FShadowProjectionPassParameters& CommonPassParameters,
 	int32 ViewIndex,
 	const FViewInfo* View,
+	FIntRect ScissorRect,
 	const FSceneRenderer* SceneRender,
 	bool bProjectingForForwardShading,
 	bool bMobileModulatedProjections,
@@ -1111,8 +1112,9 @@ void FProjectedShadowInfo::RenderProjection(
 		RDG_EVENT_NAME("%s", *EventName),
 		PassParameters,
 		ERDGPassFlags::Raster | PassFlags,
-		[this, SceneRender, View, ViewIndex, HairVisibilityData, bProjectingForForwardShading, bMobileModulatedProjections, PassParameters](FRHICommandListImmediate& RHICmdList)
+		[this, SceneRender, View, ViewIndex, ScissorRect, HairVisibilityData, bProjectingForForwardShading, bMobileModulatedProjections, PassParameters](FRHICommandListImmediate& RHICmdList)
 	{
+		RHICmdList.SetScissorRect(true, ScissorRect.Min.X, ScissorRect.Min.Y, ScissorRect.Max.X, ScissorRect.Max.Y);
 		RHICmdList.SetViewport(View->ViewRect.Min.X, View->ViewRect.Min.Y, 0.0f, View->ViewRect.Max.X, View->ViewRect.Max.Y, 1.0f);
 
 		FScopeCycleCounter Scope(bWholeSceneShadow ? GET_STATID(STAT_RenderWholeSceneShadowProjectionsTime) : GET_STATID(STAT_RenderPerObjectShadowProjectionsTime));
@@ -1249,6 +1251,8 @@ void FProjectedShadowInfo::RenderProjection(
 				DrawClearQuad(RHICmdList, false, FLinearColor::Transparent, false, 0, true, 0);
 			}
 		}
+
+		RHICmdList.SetScissorRect(false, 0, 0, 0, 0);
 	});
 }
 
@@ -1273,6 +1277,7 @@ void FProjectedShadowInfo::RenderOnePassPointLightProjection(
 	const FShadowProjectionPassParameters& CommonPassParameters,
 	int32 ViewIndex,
 	const FViewInfo& View,
+	FIntRect ScissorRect,
 	bool bProjectingForForwardShading,
 	const FHairStrandsVisibilityData* HairVisibilityData,
 	const FHairStrandsMacroGroupDatas* HairMacroGroupData) const
@@ -1303,11 +1308,12 @@ void FProjectedShadowInfo::RenderOnePassPointLightProjection(
 	}
 
 	GraphBuilder.AddPass(
-		{},
+		RDG_EVENT_NAME("OnePassPointLightProjection"),
 		PassParameters,
 		ERDGPassFlags::Raster | PassFlags,
-		[this, &View, HairVisibilityData, HairMacroGroupData, LightBounds, bProjectingForForwardShading, bCameraInsideLightGeometry, bUseTransmission, ViewIndex](FRHICommandList& RHICmdList)
+		[this, &View, HairVisibilityData, HairMacroGroupData, LightBounds, bProjectingForForwardShading, bCameraInsideLightGeometry, bUseTransmission, ViewIndex, ScissorRect](FRHICommandList& RHICmdList)
 	{
+		RHICmdList.SetScissorRect(true, ScissorRect.Min.X, ScissorRect.Min.Y, ScissorRect.Max.X, ScissorRect.Max.Y);
 		RHICmdList.SetViewport(View.ViewRect.Min.X, View.ViewRect.Min.Y, 0.0f, View.ViewRect.Max.X, View.ViewRect.Max.Y, 1.0f);
 
 		FGraphicsPipelineStateInitializer GraphicsPSOInit;
@@ -1418,6 +1424,7 @@ void FProjectedShadowInfo::RenderOnePassPointLightProjection(
 		// Project the point light shadow with some approximately bounding geometry, 
 		// So we can get speedups from depth testing and not processing pixels outside of the light's influence.
 		StencilingGeometry::DrawSphere(RHICmdList);
+		RHICmdList.SetScissorRect(false, 0, 0, 0, 0);
 	});
 }
 
@@ -1790,8 +1797,6 @@ void FSceneRenderer::RenderShadowProjections(
 {
 	ensureMsgf(bShadowDepthRenderCompleted, TEXT("Shadow depth rendering was not done before shadow projections, this will cause severe shadow artifacts and indicates an engine bug (pass ordering)"));
 
-	FPersistentUniformBuffers& UniformBuffers = Scene->UniformBuffers;
-
 	for (int32 ViewIndex = 0; ViewIndex < Views.Num(); ViewIndex++)
 	{
 		const FViewInfo& View = Views[ViewIndex];
@@ -1799,10 +1804,8 @@ void FSceneRenderer::RenderShadowProjections(
 		RDG_GPU_MASK_SCOPE(GraphBuilder, View.GPUMask);
 		RDG_EVENT_SCOPE_CONDITIONAL(GraphBuilder, Views.Num() > 1, "View%d", ViewIndex);
 
-		AddPass(GraphBuilder, [&UniformBuffers, &View, LightSceneProxy](FRHICommandList& RHICmdList)
-		{
-			LightSceneProxy->SetScissorRect(RHICmdList, View, View.ViewRect);
-		});
+		FIntRect ScissorRect;
+		LightSceneProxy->GetScissorRect(ScissorRect, View, View.ViewRect);
 
 		const FHairStrandsVisibilityData* HairVisibilityData = nullptr;
 		const FHairStrandsMacroGroupDatas* HairMacroGroupData = nullptr;
@@ -1822,21 +1825,16 @@ void FSceneRenderer::RenderShadowProjections(
 				{
 					if (ProjectedShadowInfo->bOnePassPointLightShadow)
 					{
-						ProjectedShadowInfo->RenderOnePassPointLightProjection(GraphBuilder, CommonPassParameters, ViewIndex, View, bProjectingForForwardShading, HairVisibilityData, HairMacroGroupData);
+						ProjectedShadowInfo->RenderOnePassPointLightProjection(GraphBuilder, CommonPassParameters, ViewIndex, View, ScissorRect, bProjectingForForwardShading, HairVisibilityData, HairMacroGroupData);
 					}
 					else
 					{
-						ProjectedShadowInfo->RenderProjection(GraphBuilder, CommonPassParameters, ViewIndex, &View, this, bProjectingForForwardShading, bMobileModulatedProjections, HairVisibilityData, HairMacroGroupData);
+						ProjectedShadowInfo->RenderProjection(GraphBuilder, CommonPassParameters, ViewIndex, &View, ScissorRect, this, bProjectingForForwardShading, bMobileModulatedProjections, HairVisibilityData, HairMacroGroupData);
 					}
 				}
 			}
 		}
 	}
-
-	AddPass(GraphBuilder, [](FRHICommandList& RHICmdList)
-	{
-		RHICmdList.SetScissorRect(false, 0, 0, 0, 0);
-	});
 }
 
 void FDeferredShadingSceneRenderer::RenderShadowProjections(
diff --git a/Engine/Source/Runtime/Renderer/Private/ShadowRendering.h b/Engine/Source/Runtime/Renderer/Private/ShadowRendering.h
--- a/Engine/Source/Runtime/Renderer/Private/ShadowRendering.h
+++ b/Engine/Source/Runtime/Renderer/Private/ShadowRendering.h
@@ -523,6 +523,7 @@ public:
 		const FShadowProjectionPassParameters& CommonPassParameters,
 		int32 ViewIndex,
 		const FViewInfo* View,
+		FIntRect ScissorRect,
 		const FSceneRenderer* SceneRender,
 		bool bProjectingForForwardShading,
 		bool bMobileModulatedProjections,
@@ -558,6 +559,7 @@ public:
 		const FShadowProjectionPassParameters& CommonPassParameters,
 		int32 ViewIndex,
 		const FViewInfo& View,
+		FIntRect ScissorRect,
 		bool bProjectingForForwardShading,
 		const FHairStrandsVisibilityData* HairVisibilityData,
 		const struct FHairStrandsMacroGroupDatas* HairMacroGroupData) const;
