Fixed invalid scissor rect ensure with point light shadow rendering.

#rnx

diff --git a/Engine/Source/Runtime/Renderer/Private/ShadowRendering.cpp b/Engine/Source/Runtime/Renderer/Private/ShadowRendering.cpp
--- a/Engine/Source/Runtime/Renderer/Private/ShadowRendering.cpp
+++ b/Engine/Source/Runtime/Renderer/Private/ShadowRendering.cpp
@@ -1000,7 +1000,7 @@ void FProjectedShadowInfo::RenderProjection(
 	const FShadowProjectionPassParameters& CommonPassParameters,
 	int32 ViewIndex,
 	const FViewInfo* View,
-	FIntRect ScissorRect,
+	const FLightSceneProxy* LightSceneProxy,
 	const FSceneRenderer* SceneRender,
 	bool bProjectingForForwardShading,
 	bool bMobileModulatedProjections,
@@ -1112,10 +1112,10 @@ void FProjectedShadowInfo::RenderProjection(
 		RDG_EVENT_NAME("%s", *EventName),
 		PassParameters,
 		ERDGPassFlags::Raster | PassFlags,
-		[this, SceneRender, View, ViewIndex, ScissorRect, HairVisibilityData, bProjectingForForwardShading, bMobileModulatedProjections, PassParameters](FRHICommandListImmediate& RHICmdList)
+		[this, SceneRender, View, ViewIndex, LightSceneProxy, HairVisibilityData, bProjectingForForwardShading, bMobileModulatedProjections, PassParameters](FRHICommandListImmediate& RHICmdList)
 	{
-		RHICmdList.SetScissorRect(true, ScissorRect.Min.X, ScissorRect.Min.Y, ScissorRect.Max.X, ScissorRect.Max.Y);
 		RHICmdList.SetViewport(View->ViewRect.Min.X, View->ViewRect.Min.Y, 0.0f, View->ViewRect.Max.X, View->ViewRect.Max.Y, 1.0f);
+		LightSceneProxy->SetScissorRect(RHICmdList, *View, View->ViewRect);
 
 		FScopeCycleCounter Scope(bWholeSceneShadow ? GET_STATID(STAT_RenderWholeSceneShadowProjectionsTime) : GET_STATID(STAT_RenderPerObjectShadowProjectionsTime));
 
@@ -1277,7 +1277,7 @@ void FProjectedShadowInfo::RenderOnePassPointLightProjection(
 	const FShadowProjectionPassParameters& CommonPassParameters,
 	int32 ViewIndex,
 	const FViewInfo& View,
-	FIntRect ScissorRect,
+	const FLightSceneProxy* LightSceneProxy,
 	bool bProjectingForForwardShading,
 	const FHairStrandsVisibilityData* HairVisibilityData,
 	const FHairStrandsMacroGroupDatas* HairMacroGroupData) const
@@ -1311,10 +1311,10 @@ void FProjectedShadowInfo::RenderOnePassPointLightProjection(
 		RDG_EVENT_NAME("OnePassPointLightProjection"),
 		PassParameters,
 		ERDGPassFlags::Raster | PassFlags,
-		[this, &View, HairVisibilityData, HairMacroGroupData, LightBounds, bProjectingForForwardShading, bCameraInsideLightGeometry, bUseTransmission, ViewIndex, ScissorRect](FRHICommandList& RHICmdList)
+		[this, &View, HairVisibilityData, HairMacroGroupData, LightBounds, bProjectingForForwardShading, bCameraInsideLightGeometry, bUseTransmission, ViewIndex, LightSceneProxy](FRHICommandList& RHICmdList)
 	{
-		RHICmdList.SetScissorRect(true, ScissorRect.Min.X, ScissorRect.Min.Y, ScissorRect.Max.X, ScissorRect.Max.Y);
 		RHICmdList.SetViewport(View.ViewRect.Min.X, View.ViewRect.Min.Y, 0.0f, View.ViewRect.Max.X, View.ViewRect.Max.Y, 1.0f);
+		LightSceneProxy->SetScissorRect(RHICmdList, View, View.ViewRect);
 
 		FGraphicsPipelineStateInitializer GraphicsPSOInit;
 		RHICmdList.ApplyCachedRenderTargets(GraphicsPSOInit);
@@ -1804,9 +1804,6 @@ void FSceneRenderer::RenderShadowProjections(
 		RDG_GPU_MASK_SCOPE(GraphBuilder, View.GPUMask);
 		RDG_EVENT_SCOPE_CONDITIONAL(GraphBuilder, Views.Num() > 1, "View%d", ViewIndex);
 
-		FIntRect ScissorRect;
-		LightSceneProxy->GetScissorRect(ScissorRect, View, View.ViewRect);
-
 		const FHairStrandsVisibilityData* HairVisibilityData = nullptr;
 		const FHairStrandsMacroGroupDatas* HairMacroGroupData = nullptr;
 		if (HairDatas)
@@ -1825,11 +1822,11 @@ void FSceneRenderer::RenderShadowProjections(
 				{
 					if (ProjectedShadowInfo->bOnePassPointLightShadow)
 					{
-						ProjectedShadowInfo->RenderOnePassPointLightProjection(GraphBuilder, CommonPassParameters, ViewIndex, View, ScissorRect, bProjectingForForwardShading, HairVisibilityData, HairMacroGroupData);
+						ProjectedShadowInfo->RenderOnePassPointLightProjection(GraphBuilder, CommonPassParameters, ViewIndex, View, LightSceneProxy, bProjectingForForwardShading, HairVisibilityData, HairMacroGroupData);
 					}
 					else
 					{
-						ProjectedShadowInfo->RenderProjection(GraphBuilder, CommonPassParameters, ViewIndex, &View, ScissorRect, this, bProjectingForForwardShading, bMobileModulatedProjections, HairVisibilityData, HairMacroGroupData);
+						ProjectedShadowInfo->RenderProjection(GraphBuilder, CommonPassParameters, ViewIndex, &View, LightSceneProxy, this, bProjectingForForwardShading, bMobileModulatedProjections, HairVisibilityData, HairMacroGroupData);
 					}
 				}
 			}
diff --git a/Engine/Source/Runtime/Renderer/Private/ShadowRendering.h b/Engine/Source/Runtime/Renderer/Private/ShadowRendering.h
--- a/Engine/Source/Runtime/Renderer/Private/ShadowRendering.h
+++ b/Engine/Source/Runtime/Renderer/Private/ShadowRendering.h
@@ -523,7 +523,7 @@ public:
 		const FShadowProjectionPassParameters& CommonPassParameters,
 		int32 ViewIndex,
 		const FViewInfo* View,
-		FIntRect ScissorRect,
+		const FLightSceneProxy* LightSceneProxy,
 		const FSceneRenderer* SceneRender,
 		bool bProjectingForForwardShading,
 		bool bMobileModulatedProjections,
@@ -559,7 +559,7 @@ public:
 		const FShadowProjectionPassParameters& CommonPassParameters,
 		int32 ViewIndex,
 		const FViewInfo& View,
-		FIntRect ScissorRect,
+		const FLightSceneProxy* LightSceneProxy,
 		bool bProjectingForForwardShading,
 		const FHairStrandsVisibilityData* HairVisibilityData,
 		const struct FHairStrandsMacroGroupDatas* HairMacroGroupData) const;
