Fixed RTWriteMaskInputs shader parameter error.

#rnx

diff --git a/Engine/Source/Runtime/Renderer/Private/RendererUtils.cpp b/Engine/Source/Runtime/Renderer/Private/RendererUtils.cpp
--- a/Engine/Source/Runtime/Renderer/Private/RendererUtils.cpp
+++ b/Engine/Source/Runtime/Renderer/Private/RendererUtils.cpp
@@ -19,7 +19,7 @@ public:
 
 	BEGIN_SHADER_PARAMETER_STRUCT(FParameters, )
 		SHADER_PARAMETER_RDG_TEXTURE(Texture2D, ReferenceInput)
-		SHADER_PARAMETER_RDG_TEXTURE_SRV_ARRAY(Buffer<uint>, RTWriteMaskInputs, [MaxRenderTargetCount])
+		SHADER_PARAMETER_RDG_TEXTURE_SRV_ARRAY(TextureMetadata, RTWriteMaskInputs, [MaxRenderTargetCount])
 		SHADER_PARAMETER_RDG_TEXTURE_UAV(RWTexture2D<uint>, OutCombinedRTWriteMask)
 	END_SHADER_PARAMETER_STRUCT()
 
