Support Desktop TAA with the mobile renderer( reimplement in UE5)

diff --git a/Engine/Shaders/Private/PostProcessMobile.usf b/Engine/Shaders/Private/PostProcessMobile.usf
--- a/Engine/Shaders/Private/PostProcessMobile.usf
+++ b/Engine/Shaders/Private/PostProcessMobile.usf
@@ -1136,106 +1136,6 @@ void SunMergePS_Mobile(
 #undef SUN_P5
 #undef SUN_P6
 
-//
-// Average results of bloom/sunshaft/vignette for 2 frames (used for temporal AA).
-//
-
-void SunAvgPS_Mobile(
-	float4 InUV : TEXCOORD0,
-	out HALF4_TYPE OutColor : SV_Target0
-	)
-{
-	OutColor = 
-		SunMergeTexture.Sample(SunMergeSampler, InUV.xy).rgba * 0.5 +
-		LastFrameSunMergeTexture.Sample(LastFrameSunMergeSampler, InUV.xy).rgba * 0.5;
-}
-
-
-
-//
-// Ultra simple temporal antialiasing shader which works without motion vectors.
-//
-
-float AaBlendAmount;
-
-void AaVS_Mobile(
-	in float4 InPosition : ATTRIBUTE0,
-	in float2 InTexCoord : ATTRIBUTE1,
-	out float2 OutUV[6] : TEXCOORD0,
-	out float4 OutPosition : SV_POSITION
-	)
-{
-	DrawRectangle(InPosition, InTexCoord, OutPosition, InTexCoord);
-
-	// AA texture coordinates.
-	float AaX = View.TemporalAAParams.x;
-	// Y is flipped here in comparison with the coords in SceneVisibility (GL -y flip).
-	float2 Aa1 = float2(0.5,-0.0);
-	float2 Aa2 = float2(-0.0,0.5);
-	if(abs(AaX) == 0.0) Aa1 = float2(-0.0,0.5);
-	if(abs(AaX) == 0.0) Aa2 = float2(0.5,-0.0);
-	OutUV[0] = InTexCoord.xy + (Aa1 + float2(-0.75, 0.0)) * BufferSizeAndInvSize.zw;
-	OutUV[1] = InTexCoord.xy + (Aa1 + float2( 0.0,-0.75)) * BufferSizeAndInvSize.zw;
-	OutUV[2] = InTexCoord.xy + Aa1 * BufferSizeAndInvSize.zw;
-	OutUV[3] = InTexCoord.xy + (Aa1 + float2( 0.0, 0.75)) * BufferSizeAndInvSize.zw;
-	OutUV[4] = InTexCoord.xy + (Aa1 + float2( 0.75, 0.0)) * BufferSizeAndInvSize.zw;
-	OutUV[5] = InTexCoord.xy + Aa2 * BufferSizeAndInvSize.zw;
-}
-
-void AaPS_Mobile(
-	float2 UV[6] : TEXCOORD0,
-	out HALF4_TYPE OutColor : SV_Target0
-	)
-{
-	//   N
-	// w M E            P = previous middle
-	//   S
-	half4 ColorN = SceneColorTexture.Sample(SceneColorSampler, UV[0].xy);
-	half4 ColorW = SceneColorTexture.Sample(SceneColorSampler, UV[1].xy);
-	half4 ColorM = SceneColorTexture.Sample(SceneColorSampler, UV[2].xy);
-	half4 ColorE = SceneColorTexture.Sample(SceneColorSampler, UV[3].xy);
-	half4 ColorS = SceneColorTexture.Sample(SceneColorSampler, UV[4].xy);
-	half4 ColorP = LastFrameSceneColorTexture.Sample(LastFrameSceneColorSampler, UV[5].xy);
-
-	half LumaN = ColorN.g;
-	half LumaW = ColorW.g;
-	half LumaM = ColorM.g;
-	half LumaE = ColorE.g;
-	half LumaS = ColorS.g;
-	half LumaP = ColorP.g;
-
-	half LumaMin = min(LumaM, min(min(LumaN, LumaS), min(LumaW, LumaE)));
-	half LumaMax = max(LumaM, max(max(LumaN, LumaS), max(LumaW, LumaE)));
-
-	// Using circle of confusion size in alpha to increase the blend between two frames.
-	// Depth of field would flicker otherwise.
-	half BlendAmount = saturate(min(ColorM.a, ColorP.a) * 4.0); 
-	BlendAmount = max(BlendAmount * 0.5, AaBlendAmount);
-
-	// BlendAmount = 0.5 -> Full ghosting (average of two frames).
-	// BlendAmount = 0.0 -> No ghosting but lots of flicker.
-	LumaMin = min(LumaM, LumaP) * BlendAmount + LumaMin * (1.0 - BlendAmount);
-	LumaMax = max(LumaM, LumaP) * BlendAmount + LumaMax * (1.0 - BlendAmount);
-
-	half LumaT = LumaM * 0.5 + LumaP * 0.5;
-
-	LumaT = max(LumaT, LumaMin);
-	LumaT = min(LumaT, LumaMax);
-	
-#if METAL_PROFILE
-	// avoid division by zero 
-	// normally saturate should kill NaNs, but Metal fast math does not do it
-	// this can be fixed by instructing cross-compiler to use precise math for saturate
-	half LumaD = LumaP - LumaM;
-	half BlendFinal = saturate((LumaT - LumaM) * (abs(LumaD) < 0.001 ? 1000.0 : rcp(LumaD)));
-#else
-	half BlendFinal = saturate((LumaT - LumaM) * rcp(LumaP - LumaM));
-#endif
-			
-	OutColor.rgb = lerp(ColorM.rgb, ColorP.rgb, BlendFinal);
-	OutColor.a = 0;
-}
-
 // EyeAdaptation
 
 Buffer<float4> EyeAdaptationBuffer;
diff --git a/Engine/Shaders/Private/TemporalAA/TAAStandalone.usf b/Engine/Shaders/Private/TemporalAA/TAAStandalone.usf
--- a/Engine/Shaders/Private/TemporalAA/TAAStandalone.usf
+++ b/Engine/Shaders/Private/TemporalAA/TAAStandalone.usf
@@ -59,7 +59,7 @@
 
 // Compute shaders always do responsive TAA in a single pass.
 #if COMPUTESHADER
-	#define AA_SINGLE_PASS_RESPONSIVE 1
+	#define AA_SINGLE_PASS_RESPONSIVE SHADING_PATH_DEFERRED //The StencilTexture is not available on the mobile platform
 #endif
 
 #if TAA_PASS_CONFIG == 0 // Main
@@ -81,7 +81,11 @@
 	#endif
 
 	#if COMPUTESHADER
-		#define AA_SAMPLE_CACHE_METHOD (AA_SAMPLE_CACHE_METHOD_LDS)
+		#if AA_MOBILE_CONFIG
+			#define AA_SAMPLE_CACHE_METHOD (AA_SAMPLE_CACHE_METHOD_DISABLE) // the shared memory is not as efficient as expected on mobile devices, try not to use it on mobile devices.
+		#else
+			#define AA_SAMPLE_CACHE_METHOD (AA_SAMPLE_CACHE_METHOD_LDS)
+		#endif
 	#endif
 
 #elif TAA_PASS_CONFIG == 3 // Used for SSR
diff --git a/Engine/Source/Runtime/RHI/Public/RHIDefinitions.h b/Engine/Source/Runtime/RHI/Public/RHIDefinitions.h
--- a/Engine/Source/Runtime/RHI/Public/RHIDefinitions.h
+++ b/Engine/Source/Runtime/RHI/Public/RHIDefinitions.h
@@ -1871,6 +1871,8 @@ inline bool RHIHasTiledGPU(const FStaticShaderPlatform Platform)
 	return Platform == SP_METAL || Platform == SP_METAL_TVOS
 		|| Platform == SP_OPENGL_ES3_1_ANDROID
 		|| Platform == SP_VULKAN_ES3_1_ANDROID
+		|| Platform == SP_METAL_MRT || Platform == SP_METAL_MRT_TVOS
+		|| Platform == SP_VULKAN_SM5_ANDROID
 		|| FDataDrivenShaderPlatformInfo::GetTargetsTiledGPU(Platform);
 }
 
diff --git a/Engine/Source/Runtime/RenderCore/Private/RenderUtils.cpp b/Engine/Source/Runtime/RenderCore/Private/RenderUtils.cpp
--- a/Engine/Source/Runtime/RenderCore/Private/RenderUtils.cpp
+++ b/Engine/Source/Runtime/RenderCore/Private/RenderUtils.cpp
@@ -1240,6 +1240,16 @@ RENDERCORE_API void RenderUtilsInit()
 	}
 #endif // WITH_EDITOR
 
+	for (uint32 ShaderPlatformIndex = 0; ShaderPlatformIndex < SP_NumPlatforms; ++ShaderPlatformIndex)
+	{
+		EShaderPlatform ShaderPlatform = EShaderPlatform(ShaderPlatformIndex);
+		if (IsMobilePlatform(ShaderPlatform))
+		{
+			uint64 Mask = 1ull << ShaderPlatformIndex;
+			GBasePassVelocityPlatformMask &= ~Mask;
+		}
+	}
+
 	// Run-time ray tracing support depends on the following factors:
 	// - Ray tracing must be enabled for the project
 	// - Skin cache must be enabled for the project
diff --git a/Engine/Source/Runtime/Renderer/Private/DeferredShadingRenderer.cpp b/Engine/Source/Runtime/Renderer/Private/DeferredShadingRenderer.cpp
--- a/Engine/Source/Runtime/Renderer/Private/DeferredShadingRenderer.cpp
+++ b/Engine/Source/Runtime/Renderer/Private/DeferredShadingRenderer.cpp
@@ -324,8 +324,7 @@ TGlobalResource<FGlobalDynamicReadBuffer> FDeferredShadingSceneRenderer::Dynamic
  */
 static FORCEINLINE bool NeedsPrePass(const FDeferredShadingSceneRenderer* Renderer)
 {
-	return (RHIHasTiledGPU(Renderer->ViewFamily.GetShaderPlatform()) == false) && 
-		(Renderer->DepthPass.EarlyZPassMode != DDM_None || Renderer->DepthPass.bEarlyZPassMovable != 0);
+	return (Renderer->DepthPass.EarlyZPassMode != DDM_None || Renderer->DepthPass.bEarlyZPassMovable != 0);
 }
 
 bool FDeferredShadingSceneRenderer::RenderHzb(FRDGBuilder& GraphBuilder, FRDGTextureRef SceneDepthTexture)
diff --git a/Engine/Source/Runtime/Renderer/Private/DeferredShadingRenderer.h b/Engine/Source/Runtime/Renderer/Private/DeferredShadingRenderer.h
--- a/Engine/Source/Runtime/Renderer/Private/DeferredShadingRenderer.h
+++ b/Engine/Source/Runtime/Renderer/Private/DeferredShadingRenderer.h
@@ -23,7 +23,6 @@
 #include "RenderGraphUtils.h"
 
 enum class ERayTracingPrimaryRaysFlag : uint32;
-enum class EVelocityPass : uint32;
 
 class FSceneTextureParameters;
 class FDistanceFieldAOParameters;
@@ -252,6 +251,8 @@ public:
 	/** Render the view family's hit proxies. */
 	virtual void RenderHitProxies(FRDGBuilder& GraphBuilder) override;
 
+	virtual bool ShouldRenderVelocities() const override;
+
 #if !(UE_BUILD_SHIPPING || UE_BUILD_TEST)
 	void RenderVisualizeTexturePool(FRHICommandListImmediate& RHICmdList);
 #endif
@@ -612,14 +613,6 @@ private:
 		const FMinimalSceneTextures& SceneTextures,
 		FSeparateTranslucencyTextures& OutSeparateTranslucencyTextures);
 
-	bool ShouldRenderVelocities() const;
-
-	void RenderVelocities(
-		FRDGBuilder& GraphBuilder,
-		const FSceneTextures& SceneTextures,
-		EVelocityPass VelocityPass,
-		bool bForceVelocity);
-
 	bool ShouldRenderDistortion() const;
 	void RenderDistortion(FRDGBuilder& GraphBuilder, FRDGTextureRef SceneColorTexture, FRDGTextureRef SceneDepthTexture);
 
diff --git a/Engine/Source/Runtime/Renderer/Private/MobileShadingRenderer.cpp b/Engine/Source/Runtime/Renderer/Private/MobileShadingRenderer.cpp
--- a/Engine/Source/Runtime/Renderer/Private/MobileShadingRenderer.cpp
+++ b/Engine/Source/Runtime/Renderer/Private/MobileShadingRenderer.cpp
@@ -106,6 +106,9 @@ DECLARE_CYCLE_STAT(TEXT("Translucency"), STAT_CLMM_Translucency, STATGROUP_Comma
 DECLARE_CYCLE_STAT(TEXT("Shadows"), STAT_CLMM_Shadows, STATGROUP_CommandListMarkers);
 DECLARE_CYCLE_STAT(TEXT("SceneSimulation"), STAT_CLMM_SceneSim, STATGROUP_CommandListMarkers);
 DECLARE_CYCLE_STAT(TEXT("PrePass"), STAT_CLM_MobilePrePass, STATGROUP_CommandListMarkers);
+DECLARE_CYCLE_STAT(TEXT("Velocity"), STAT_CLMM_Velocity, STATGROUP_CommandListMarkers);
+DECLARE_CYCLE_STAT(TEXT("AfterVelocity"), STAT_CLMM_AfterVelocity, STATGROUP_CommandListMarkers);
+DECLARE_CYCLE_STAT(TEXT("TranslucentVelocity"), STAT_CLMM_TranslucentVelocity, STATGROUP_CommandListMarkers);
 
 FGlobalDynamicIndexBuffer FMobileSceneRenderer::DynamicIndexBuffer;
 FGlobalDynamicVertexBuffer FMobileSceneRenderer::DynamicVertexBuffer;
@@ -377,6 +380,8 @@ void FMobileSceneRenderer::InitViews(FRDGBuilder& GraphBuilder, FSceneTexturesCo
 		// Only support forward shading, we don't want to break tiled deferred shading.
 		&& !bDeferredShading;
 
+	bShouldRenderVelocities = ShouldRenderVelocities();
+
 	// Whether we need to store depth for post-processing
 	// On PowerVR we see flickering of shadows and depths not updating correctly if targets are discarded.
 	// See CVarMobileForceDepthResolve use in ConditionalResolveSceneDepth.
@@ -391,12 +396,19 @@ void FMobileSceneRenderer::InitViews(FRDGBuilder& GraphBuilder, FSceneTexturesCo
 		bRequiresPixelProjectedPlanarRelfectionPass ||
 		bSeparateTranslucencyActive ||
 		Views[0].bIsReflectionCapture ||
-		(bDeferredShading && bPostProcessUsesSceneDepth);
+		(bDeferredShading && bPostProcessUsesSceneDepth) ||
+		bShouldRenderVelocities;
 	// never keep MSAA depth
 	bKeepDepthContent = (NumMSAASamples > 1 ? false : bKeepDepthContent);
 
 	// Finalize and set the scene textures config.
 	SceneTexturesConfig.bKeepDepthContent = bKeepDepthContent;
+
+	if (!SupportsMSAA())
+	{
+		SceneTexturesConfig.NumSamples = 1;
+	}
+	
 	FSceneTexturesConfig::Set(SceneTexturesConfig);
 
 	// Initialise Sky/View resources before the view global uniform buffer is built.
@@ -662,8 +674,23 @@ void FMobileSceneRenderer::Render(FRDGBuilder& GraphBuilder)
 	}
 
 	SceneTextures.MobileSetupMode = EMobileSceneTextureSetupMode::All;
+	SceneTextures.MobileSetupMode &= ~EMobileSceneTextureSetupMode::SceneVelocity;
 	SceneTextures.MobileUniformBuffer = CreateMobileSceneTextureUniformBuffer(GraphBuilder, SceneTextures.MobileSetupMode);
 
+	if (bShouldRenderVelocities)
+	{
+		// Render the velocities of movable objects
+		GraphBuilder.SetCommandListStat(GET_STATID(STAT_CLMM_Velocity));
+		RenderVelocities(GraphBuilder, SceneTextures, EVelocityPass::Opaque, false);
+		GraphBuilder.SetCommandListStat(GET_STATID(STAT_CLMM_AfterVelocity));
+
+		GraphBuilder.SetCommandListStat(GET_STATID(STAT_CLMM_TranslucentVelocity));
+		RenderVelocities(GraphBuilder, SceneTextures, EVelocityPass::Translucent, false);
+
+		SceneTextures.MobileSetupMode = EMobileSceneTextureSetupMode::All;
+		SceneTextures.MobileUniformBuffer = CreateMobileSceneTextureUniformBuffer(GraphBuilder, SceneTextures.MobileSetupMode);
+	}
+
 	if (FXSystem && Views.IsValidIndex(0))
 	{
 		auto* PassParameters = GraphBuilder.AllocParameters<FMobileRenderOpaqueFXPassParameters>();
@@ -1659,3 +1686,8 @@ void FMobileSceneRenderer::UpdateMovablePointLightUniformBufferAndShadowInfo()
 		}
 	}
 }
+
+bool FMobileSceneRenderer::SupportsMSAA() const
+{
+	return !(IsUsingMobilePixelProjectedReflection(ShaderPlatform) || IsUsingMobileAmbientOcclusion(ShaderPlatform) || ShouldRenderVelocities() || bDeferredShading);
+}
diff --git a/Engine/Source/Runtime/Renderer/Private/PostProcess/PostProcessMobile.cpp b/Engine/Source/Runtime/Renderer/Private/PostProcess/PostProcessMobile.cpp
--- a/Engine/Source/Runtime/Renderer/Private/PostProcess/PostProcessMobile.cpp
+++ b/Engine/Source/Runtime/Renderer/Private/PostProcess/PostProcessMobile.cpp
@@ -1679,250 +1679,6 @@ FScreenPassTexture AddMobileIntegrateDofPass(FRDGBuilder& GraphBuilder, const FV
 	return MoveTemp(IntegrateDofOutput);
 }
 
-//
-// SUN AVG
-//
-
-class FMobileSunAvgPS : public FGlobalShader
-{
-public:
-	DECLARE_GLOBAL_SHADER(FMobileSunAvgPS);
-	SHADER_USE_PARAMETER_STRUCT(FMobileSunAvgPS, FGlobalShader);
-
-	BEGIN_SHADER_PARAMETER_STRUCT(FParameters, )
-		SHADER_PARAMETER_STRUCT_REF(FViewUniformShaderParameters, View)
-		SHADER_PARAMETER_RDG_TEXTURE(Texture2D, SunMergeTexture)
-		SHADER_PARAMETER_SAMPLER(SamplerState, SunMergeSampler)
-		SHADER_PARAMETER_RDG_TEXTURE(Texture2D, LastFrameSunMergeTexture)
-		SHADER_PARAMETER_SAMPLER(SamplerState, LastFrameSunMergeSampler)
-		RENDER_TARGET_BINDING_SLOTS()
-	END_SHADER_PARAMETER_STRUCT()
-
-	static bool ShouldCompilePermutation(const FGlobalShaderPermutationParameters& Parameters)
-	{
-		return IsMobilePlatform(Parameters.Platform);
-	}
-};
-
-IMPLEMENT_GLOBAL_SHADER(FMobileSunAvgPS, "/Engine/Private/PostProcessMobile.usf", "SunAvgPS_Mobile", SF_Pixel);
-
-FScreenPassTexture AddMobileSunAvgPass(FRDGBuilder& GraphBuilder, const FViewInfo& View, const FMobileSunAvgInputs& Inputs)
-{
-	FRDGTextureDesc SunAvgDesc = Inputs.SunMerge.Texture->Desc;
-
-	FScreenPassRenderTarget SunAvgOutput = FScreenPassRenderTarget(GraphBuilder.CreateTexture(SunAvgDesc, TEXT("SunAvg")), ERenderTargetLoadAction::EClear);
-
-
-	FMobileSunAvgPS::FParameters* PassParameters = GraphBuilder.AllocParameters<FMobileSunAvgPS::FParameters>();
-	PassParameters->RenderTargets[0] = SunAvgOutput.GetRenderTargetBinding();
-
-	PassParameters->View = View.ViewUniformBuffer;
-	PassParameters->SunMergeTexture = Inputs.SunMerge.Texture;
-	PassParameters->SunMergeSampler = TStaticSamplerState<SF_Bilinear, AM_Clamp, AM_Clamp, AM_Clamp>::GetRHI();
-	PassParameters->LastFrameSunMergeTexture = Inputs.LastFrameSunMerge.Texture;
-	PassParameters->LastFrameSunMergeSampler = TStaticSamplerState<SF_Bilinear, AM_Clamp, AM_Clamp, AM_Clamp>::GetRHI();
-
-
-	TShaderMapRef<FMobileSunAvgPS> PixelShader(View.ShaderMap);
-
-	const FScreenPassTextureViewport OutputViewport(SunAvgOutput);
-
-	AddDrawScreenPass(GraphBuilder, RDG_EVENT_NAME("MobileSunAvg"), View, OutputViewport, OutputViewport, PixelShader, PassParameters);
-
-	return MoveTemp(SunAvgOutput);
-}
-
-//
-// MOBILE AA
-//
-
-class FMobileTAAPS : public FGlobalShader
-{
-	DECLARE_GLOBAL_SHADER(FMobileTAAPS);
-	SHADER_USE_PARAMETER_STRUCT(FMobileTAAPS, FGlobalShader);
-
-	BEGIN_SHADER_PARAMETER_STRUCT(FParameters, )
-		SHADER_PARAMETER_STRUCT_REF(FViewUniformShaderParameters, View)
-		SHADER_PARAMETER(float, AaBlendAmount)
-		SHADER_PARAMETER_RDG_TEXTURE(Texture2D, SceneColorTexture)
-		SHADER_PARAMETER_SAMPLER(SamplerState, SceneColorSampler)
-		SHADER_PARAMETER_RDG_TEXTURE(Texture2D, LastFrameSceneColorTexture)
-		SHADER_PARAMETER_SAMPLER(SamplerState, LastFrameSceneColorSampler)
-		RENDER_TARGET_BINDING_SLOTS()
-	END_SHADER_PARAMETER_STRUCT()
-
-	static bool ShouldCompilePermutation(const FGlobalShaderPermutationParameters& Parameters)
-	{
-		return IsMobilePlatform(Parameters.Platform);
-	}
-};
-
-IMPLEMENT_GLOBAL_SHADER(FMobileTAAPS, "/Engine/Private/PostProcessMobile.usf", "AaPS_Mobile", SF_Pixel);
-
-class FMobileTAAVS : public FGlobalShader
-{
-public:
-	DECLARE_GLOBAL_SHADER(FMobileTAAVS);
-	SHADER_USE_PARAMETER_STRUCT_WITH_LEGACY_BASE(FMobileTAAVS, FGlobalShader);
-
-	BEGIN_SHADER_PARAMETER_STRUCT(FParameters, )
-		SHADER_PARAMETER_STRUCT_REF(FViewUniformShaderParameters, View)
-		SHADER_PARAMETER(FVector4, BufferSizeAndInvSize)
-	END_SHADER_PARAMETER_STRUCT()
-
-	static bool ShouldCompilePermutation(const FGlobalShaderPermutationParameters& Parameters)
-	{
-		return IsMobilePlatform(Parameters.Platform);
-	}
-};
-
-IMPLEMENT_GLOBAL_SHADER(FMobileTAAVS, "/Engine/Private/PostProcessMobile.usf", "AaVS_Mobile", SF_Vertex);
-
-FScreenPassTexture AddMobileTAAPass(FRDGBuilder& GraphBuilder, const FViewInfo& View, const FMobileTAAInputs& Inputs)
-{
-	const FIntPoint& BufferSize = Inputs.SceneColor.Texture->Desc.Extent;
-
-	FScreenPassRenderTarget TAAOutput = Inputs.OverrideOutput;
-
-	if (!TAAOutput.IsValid())
-	{
-		FRDGTextureDesc TAADesc = FRDGTextureDesc::Create2D(BufferSize, PF_B8G8R8A8, FClearValueBinding::Black, TexCreate_RenderTargetable | TexCreate_ShaderResource);
-
-		TAAOutput = FScreenPassRenderTarget(GraphBuilder.CreateTexture(TAADesc, TEXT("TAA")), Inputs.SceneColor.ViewRect, View.IsFirstInFamily() ? ERenderTargetLoadAction::EClear : ERenderTargetLoadAction::ELoad);
-	}
-
-	TShaderMapRef<FMobileTAAVS> VertexShader(View.ShaderMap);
-
-	FMobileTAAVS::FParameters VSShaderParameters;
-
-	VSShaderParameters.View = View.ViewUniformBuffer;
-	VSShaderParameters.BufferSizeAndInvSize = FVector4(BufferSize.X, BufferSize.Y, 1.0f / BufferSize.X, 1.0f / BufferSize.Y);
-
-	// Compute the blend factor which decides the trade off between ghosting in motion and flicker when not moving.
-	// This works by computing the screen space motion vector of distant point at the center of the screen.
-	// This factor will effectively provide an idea of the amount of camera rotation.
-	// Higher camera rotation = less blend factor (0.0).
-	// Lower or no camera rotation = high blend factor (0.25).
-	float BlendAmount = 0.0;
-	FSceneViewState* ViewState = (FSceneViewState*)View.State;
-	if(ViewState)
-	{
-		FMatrix Proj = View.ViewMatrices.ComputeProjectionNoAAMatrix();
-		FMatrix PrevProj = View.PrevViewInfo.ViewMatrices.ComputeProjectionNoAAMatrix();
-
-		FMatrix ViewProj = ( View.ViewMatrices.GetViewMatrix() * Proj ).GetTransposed();
-		FMatrix PrevViewProj = ( View.PrevViewInfo.ViewMatrices.GetViewMatrix() * PrevProj ).GetTransposed();
-
-		double InvViewProj[16];
-		Inverse4x4( InvViewProj, (float*)ViewProj.M );
-
-		const float* p = (float*)PrevViewProj.M;
-
-		const double cxx = InvViewProj[ 0]; const double cxy = InvViewProj[ 1]; const double cxz = InvViewProj[ 2]; const double cxw = InvViewProj[ 3];
-		const double cyx = InvViewProj[ 4]; const double cyy = InvViewProj[ 5]; const double cyz = InvViewProj[ 6]; const double cyw = InvViewProj[ 7];
-		const double czx = InvViewProj[ 8]; const double czy = InvViewProj[ 9]; const double czz = InvViewProj[10]; const double czw = InvViewProj[11];
-		const double cwx = InvViewProj[12]; const double cwy = InvViewProj[13]; const double cwz = InvViewProj[14]; const double cww = InvViewProj[15];
-
-		const double pxx = (double)(p[ 0]); const double pxy = (double)(p[ 1]); const double pxz = (double)(p[ 2]); const double pxw = (double)(p[ 3]);
-		const double pyx = (double)(p[ 4]); const double pyy = (double)(p[ 5]); const double pyz = (double)(p[ 6]); const double pyw = (double)(p[ 7]);
-		const double pwx = (double)(p[12]); const double pwy = (double)(p[13]); const double pwz = (double)(p[14]); const double pww = (double)(p[15]);
-
-		float CameraMotion0W = (float)(2.0*(cww*pww - cwx*pww + cwy*pww + (cxw - cxx + cxy)*pwx + (cyw - cyx + cyy)*pwy + (czw - czx + czy)*pwz));
-		float CameraMotion2Z = (float)(cwy*pww + cwy*pxw + cww*(pww + pxw) - cwx*(pww + pxw) + (cxw - cxx + cxy)*(pwx + pxx) + (cyw - cyx + cyy)*(pwy + pxy) + (czw - czx + czy)*(pwz + pxz));
-		float CameraMotion4Z = (float)(cwy*pww + cww*(pww - pyw) - cwy*pyw + cwx*((-pww) + pyw) + (cxw - cxx + cxy)*(pwx - pyx) + (cyw - cyx + cyy)*(pwy - pyy) + (czw - czx + czy)*(pwz - pyz));
-
-		// Depth surface 0=far, 1=near.
-		// This is simplified to compute camera motion with depth = 0.0 (infinitely far away).
-		// Camera motion for pixel (in ScreenPos space).
-		float ScaleM = 1.0f / CameraMotion0W;
-		// Back projection value (projected screen space).
-		float BackX = CameraMotion2Z * ScaleM;
-		float BackY = CameraMotion4Z * ScaleM;
-
-		// Start with the distance in screen space.
-		BlendAmount = BackX * BackX + BackY * BackY;
-		if(BlendAmount > 0.0f)
-		{
-			BlendAmount = sqrt(BlendAmount);
-		}
-			
-		// Higher numbers truncate anti-aliasing and ghosting faster.
-		float BlendEffect = 8.0f;
-		BlendAmount = 0.25f - BlendAmount * BlendEffect;
-		if(BlendAmount < 0.0f)
-		{
-			BlendAmount = 0.0f;
-		}
-	}
-
-	TShaderMapRef<FMobileTAAPS> PixelShader(View.ShaderMap);
-
-	FMobileTAAPS::FParameters* PSShaderParameters = GraphBuilder.AllocParameters<FMobileTAAPS::FParameters>();
-	PSShaderParameters->RenderTargets[0] = TAAOutput.GetRenderTargetBinding();
-	PSShaderParameters->View = View.ViewUniformBuffer;
-	PSShaderParameters->AaBlendAmount = BlendAmount;
-
-	PSShaderParameters->SceneColorTexture = Inputs.SceneColor.Texture;
-	PSShaderParameters->SceneColorSampler = TStaticSamplerState<SF_Bilinear, AM_Clamp, AM_Clamp, AM_Clamp>::GetRHI();
-	PSShaderParameters->LastFrameSceneColorTexture = Inputs.LastFrameSceneColor.Texture;
-	PSShaderParameters->LastFrameSceneColorSampler = TStaticSamplerState<SF_Bilinear, AM_Clamp, AM_Clamp, AM_Clamp>::GetRHI();
-
-	const FScreenPassTextureViewport InputViewport(Inputs.SceneColor);
-	const FScreenPassTextureViewport OutputViewport(TAAOutput);
-
-	GraphBuilder.AddPass(
-		RDG_EVENT_NAME("TAA %dx%d (PS)", OutputViewport.Extent.X, OutputViewport.Extent.Y),
-		PSShaderParameters,
-		ERDGPassFlags::Raster,
-		[VertexShader, VSShaderParameters, PixelShader, PSShaderParameters, InputViewport, OutputViewport, &View](FRHICommandList& RHICmdList)
-	{
-		RHICmdList.SetViewport(OutputViewport.Rect.Min.X, OutputViewport.Rect.Min.Y, 0.0f, OutputViewport.Rect.Max.X, OutputViewport.Rect.Max.Y, 1.0f);
-
-		FGraphicsPipelineStateInitializer GraphicsPSOInit;
-		RHICmdList.ApplyCachedRenderTargets(GraphicsPSOInit);
-
-		GraphicsPSOInit.BlendState = TStaticBlendState<>::GetRHI();
-		GraphicsPSOInit.RasterizerState = TStaticRasterizerState<>::GetRHI();
-		GraphicsPSOInit.DepthStencilState = TStaticDepthStencilState<false, CF_Always>::GetRHI();
-
-		GraphicsPSOInit.BoundShaderState.VertexDeclarationRHI = GFilterVertexDeclaration.VertexDeclarationRHI;
-		GraphicsPSOInit.BoundShaderState.VertexShaderRHI = VertexShader.GetVertexShader();
-		GraphicsPSOInit.BoundShaderState.PixelShaderRHI = PixelShader.GetPixelShader();
-		GraphicsPSOInit.PrimitiveType = PT_TriangleList;
-
-		SetGraphicsPipelineState(RHICmdList, GraphicsPSOInit);
-
-		SetShaderParameters(RHICmdList, VertexShader, VertexShader.GetVertexShader(), VSShaderParameters);
-		SetShaderParameters(RHICmdList, PixelShader, PixelShader.GetPixelShader(), *PSShaderParameters);
-
-		DrawRectangle(
-			RHICmdList,
-			0, 0,
-			OutputViewport.Extent.X, OutputViewport.Extent.Y,
-			InputViewport.Rect.Min.X, InputViewport.Rect.Min.Y,
-			InputViewport.Rect.Width(), InputViewport.Rect.Height(),
-			OutputViewport.Extent,
-			InputViewport.Extent,
-			VertexShader,
-			EDRF_UseTriangleOptimization);
-
-		if (FSceneRenderer::ShouldCompositeEditorPrimitives(View))
-		{
-			FRHICommandListExecutor::GetImmediateCommandList().ImmediateFlush(EImmediateFlushType::WaitForOutstandingTasksOnly);
-			// because of the flush it's ok to remove the const, this is not ideal as the flush can cost performance
-			FViewInfo& NonConstView = (FViewInfo&)View;
-
-			// Remove jitter (ensures editor prims are stable.)
-			NonConstView.ViewMatrices.HackRemoveTemporalAAProjectionJitter();
-
-			NonConstView.InitRHIResources();
-		}
-	});
-
-	return MoveTemp(TAAOutput);
-}
-
-
 /** Encapsulates the average luminance compute shader. */
 class FMobileAverageLuminanceCS : public FGlobalShader
 {
diff --git a/Engine/Source/Runtime/Renderer/Private/PostProcess/PostProcessMobile.h b/Engine/Source/Runtime/Renderer/Private/PostProcess/PostProcessMobile.h
--- a/Engine/Source/Runtime/Renderer/Private/PostProcess/PostProcessMobile.h
+++ b/Engine/Source/Runtime/Renderer/Private/PostProcess/PostProcessMobile.h
@@ -151,23 +151,6 @@ struct FMobileSunMergeInputs
 
 FScreenPassTexture AddMobileSunMergePass(FRDGBuilder& GraphBuilder, const FViewInfo& View, const FMobileSunMergeInputs& Inputs);
 
-struct FMobileSunAvgInputs
-{
-	FScreenPassTexture SunMerge;
-	FScreenPassTexture LastFrameSunMerge;
-};
-
-FScreenPassTexture AddMobileSunAvgPass(FRDGBuilder& GraphBuilder, const FViewInfo& View, const FMobileSunAvgInputs& Inputs);
-
-struct FMobileTAAInputs
-{
-	FScreenPassRenderTarget OverrideOutput;
-	FScreenPassTexture SceneColor;
-	FScreenPassTexture LastFrameSceneColor;
-};
-
-FScreenPassTexture AddMobileTAAPass(FRDGBuilder& GraphBuilder, const FViewInfo& View, const FMobileTAAInputs& Inputs);
-
 struct FMobileEyeAdaptationSetupInputs
 {
 	FScreenPassTexture BloomSetup_EyeAdaptation;
diff --git a/Engine/Source/Runtime/Renderer/Private/PostProcess/PostProcessing.cpp b/Engine/Source/Runtime/Renderer/Private/PostProcess/PostProcessing.cpp
--- a/Engine/Source/Runtime/Renderer/Private/PostProcess/PostProcessing.cpp
+++ b/Engine/Source/Runtime/Renderer/Private/PostProcess/PostProcessing.cpp
@@ -551,7 +551,6 @@ void AddPostProcessingPasses(FRDGBuilder& GraphBuilder, const FViewInfo& View, c
 			UpscalerPassInputs.SceneColorTexture = SceneColor.Texture;
 			UpscalerPassInputs.SceneDepthTexture = SceneDepth.Texture;
 			UpscalerPassInputs.SceneVelocityTexture = Velocity.Texture;
-			UpscalerPassInputs.EyeAdaptationTexture = GetEyeAdaptationTexture(GraphBuilder, View);
 
 			UpscalerToUse->AddPasses(
 				GraphBuilder,
@@ -1334,6 +1333,7 @@ void AddMobilePostProcessingPasses(FRDGBuilder& GraphBuilder, const FViewInfo& V
 	const FScreenPassRenderTarget ViewFamilyOutput = FScreenPassRenderTarget::CreateViewFamilyOutput(Inputs.ViewFamilyTexture, View);
 	const FScreenPassTexture SceneDepth((*Inputs.SceneTextures)->SceneDepthTexture, FinalOutputViewRect);
 	const FScreenPassTexture CustomDepth((*Inputs.SceneTextures)->CustomDepthTexture, FinalOutputViewRect);
+	const FScreenPassTexture Velocity((*Inputs.SceneTextures)->SceneVelocityTexture, FinalOutputViewRect);
 	const FScreenPassTexture BlackAlphaOneDummy(GSystemTextures.GetBlackAlphaOneDummy(GraphBuilder));
 
 	// Scene color is updated incrementally through the post process pipeline.
@@ -1352,9 +1352,9 @@ void AddMobilePostProcessingPasses(FRDGBuilder& GraphBuilder, const FViewInfo& V
 		EyeAdaptation,
 		SunMerge,
 		SeparateTranslucency,
+		TAA,
 		Tonemap,
 		PostProcessMaterialAfterTonemapping,
-		TAA,
 		FXAA,
 		HighResolutionScreenshotMask,
 		SelectionOutline,
@@ -1375,9 +1375,9 @@ void AddMobilePostProcessingPasses(FRDGBuilder& GraphBuilder, const FViewInfo& V
 		TEXT("EyeAdaptation"),
 		TEXT("SunMerge"),
 		TEXT("SeparateTranslucency"),
+		TEXT("TAA"),
 		TEXT("Tonemap"),
 		TEXT("PostProcessMaterial (AfterTonemapping)"),
-		TEXT("TAA"),
 		TEXT("FXAA"),
 		TEXT("HighResolutionScreenshotMask"),
 		TEXT("SelectionOutline"),
@@ -1514,8 +1514,8 @@ void AddMobilePostProcessingPasses(FRDGBuilder& GraphBuilder, const FViewInfo& V
 		PassSequence.SetEnabled(EPass::EyeAdaptation, bUseEyeAdaptation);
 		PassSequence.SetEnabled(EPass::SunMerge, bUseBloom || bUseSun);
 		PassSequence.SetEnabled(EPass::SeparateTranslucency, bUseSeparateTranslucency);
-		PassSequence.SetEnabled(EPass::PostProcessMaterialAfterTonemapping, PostProcessMaterialAfterTonemappingChain.Num() != 0);
 		PassSequence.SetEnabled(EPass::TAA, bUseTAA);
+		PassSequence.SetEnabled(EPass::PostProcessMaterialAfterTonemapping, PostProcessMaterialAfterTonemappingChain.Num() != 0);
 		PassSequence.SetEnabled(EPass::FXAA, View.AntiAliasingMethod == AAM_FXAA);
 		PassSequence.Finalize();
 			
@@ -1816,24 +1816,6 @@ void AddMobilePostProcessingPasses(FRDGBuilder& GraphBuilder, const FViewInfo& V
 
 			BloomOutput = AddMobileSunMergePass(GraphBuilder, View, SunMergeInputs);
 
-			if (bUseTAA && View.ViewState && !View.bStatePrevViewInfoIsReadOnly)
-			{
-				GraphBuilder.QueueTextureExtraction(BloomOutput.Texture, &View.ViewState->PrevFrameViewInfo.MobileAaBloomSunVignette);
-			}
-
-			// Mobile temporal AA requires a composite of two of these frames.
-			if (bUseTAA)
-			{
-				FMobileSunAvgInputs SunAvgInputs;
-				SunAvgInputs.SunMerge = BloomOutput;
-				SunAvgInputs.LastFrameSunMerge = FScreenPassTexture(TryRegisterExternalTexture(GraphBuilder, View.PrevViewInfo.MobileAaBloomSunVignette));
-				if (!SunAvgInputs.LastFrameSunMerge.IsValid())
-				{
-					SunAvgInputs.LastFrameSunMerge = BloomOutput;
-				}
-
-				BloomOutput = AddMobileSunAvgPass(GraphBuilder, View, SunAvgInputs);
-			}
 		}
 
 		// mobile separate translucency 
@@ -1848,6 +1830,47 @@ void AddMobilePostProcessingPasses(FRDGBuilder& GraphBuilder, const FViewInfo& V
 		}
 
 		AddPostProcessMaterialPass(BL_BeforeTonemapping, false);
+
+		// Temporal Anti-aliasing. Also may perform a temporal upsample from primary to secondary view rect.
+		if (PassSequence.IsEnabled(EPass::TAA))
+		{
+			PassSequence.AcceptPass(EPass::TAA);
+
+			EMainTAAPassConfig TAAConfig = ITemporalUpscaler::GetMainTAAPassConfig(View);
+			checkSlow(TAAConfig != EMainTAAPassConfig::Disabled);
+
+			const ITemporalUpscaler* UpscalerToUse = (TAAConfig == EMainTAAPassConfig::ThirdParty) ? View.Family->GetTemporalUpscalerInterface() : ITemporalUpscaler::GetDefaultTemporalUpscaler();
+
+			const TCHAR* UpscalerName = UpscalerToUse->GetDebugName();
+
+			// Standard event scope for temporal upscaler to have all profiling information not matter what, and with explicit detection of third party.
+			RDG_EVENT_SCOPE_CONDITIONAL(
+				GraphBuilder,
+				TAAConfig == EMainTAAPassConfig::ThirdParty,
+				"ThirdParty %s %dx%d -> %dx%d",
+				UpscalerToUse->GetDebugName(),
+				View.ViewRect.Width(), View.ViewRect.Height(),
+				View.GetSecondaryViewRectSize().X, View.GetSecondaryViewRectSize().Y);
+
+			ITemporalUpscaler::FPassInputs UpscalerPassInputs;
+
+			UpscalerPassInputs.bAllowDownsampleSceneColor = false;
+			UpscalerPassInputs.SceneColorTexture = SceneColor.Texture;
+			UpscalerPassInputs.SceneDepthTexture = SceneDepth.Texture;
+			UpscalerPassInputs.SceneVelocityTexture = Velocity.Texture;
+
+			FIntRect SecondaryViewRect;
+			FScreenPassTexture HalfResolutionSceneColor;
+
+			UpscalerToUse->AddPasses(
+				GraphBuilder,
+				View,
+				UpscalerPassInputs,
+				&SceneColor.Texture,
+				&SecondaryViewRect,
+				&HalfResolutionSceneColor.Texture,
+				&HalfResolutionSceneColor.ViewRect);
+		}
 	}
 	else
 	{
@@ -1934,26 +1957,6 @@ void AddMobilePostProcessingPasses(FRDGBuilder& GraphBuilder, const FViewInfo& V
 	{
 		AddPostProcessMaterialPass(BL_AfterTonemapping, true);
 
-		if (PassSequence.IsEnabled(EPass::TAA))
-		{
-			if (View.ViewState && !View.bStatePrevViewInfoIsReadOnly)
-			{
-				GraphBuilder.QueueTextureExtraction(SceneColor.Texture, &View.ViewState->PrevFrameViewInfo.MobileAaColor);
-			}
-
-			FMobileTAAInputs TAAInputs;
-			PassSequence.AcceptOverrideIfLastPass(EPass::TAA, TAAInputs.OverrideOutput);
-			TAAInputs.OverrideOutput.LoadAction = View.IsFirstInFamily() ? ERenderTargetLoadAction::EClear : ERenderTargetLoadAction::ELoad;
-			TAAInputs.SceneColor = SceneColor;
-			TAAInputs.LastFrameSceneColor = FScreenPassTexture(TryRegisterExternalTexture(GraphBuilder, View.PrevViewInfo.MobileAaColor));
-			if (!TAAInputs.LastFrameSceneColor.IsValid())
-			{
-				TAAInputs.LastFrameSceneColor = SceneColor;
-			}
-
-			SceneColor = AddMobileTAAPass(GraphBuilder, View, TAAInputs);
-		}
-
 		if (PassSequence.IsEnabled(EPass::FXAA))
 		{
 			FFXAAInputs PassInputs;
diff --git a/Engine/Source/Runtime/Renderer/Private/PostProcess/TemporalAA.cpp b/Engine/Source/Runtime/Renderer/Private/PostProcess/TemporalAA.cpp
--- a/Engine/Source/Runtime/Renderer/Private/PostProcess/TemporalAA.cpp
+++ b/Engine/Source/Runtime/Renderer/Private/PostProcess/TemporalAA.cpp
@@ -197,6 +197,7 @@ class FTAAStandaloneCS : public FGlobalShader
 		SHADER_PARAMETER(FVector4, ScreenPosToHistoryBufferUV)
 
 		SHADER_PARAMETER_RDG_TEXTURE(Texture2D, EyeAdaptationTexture)
+		SHADER_PARAMETER_RDG_BUFFER_SRV(Buffer<float4>, EyeAdaptationBuffer)
 
 		// Inputs
 		SHADER_PARAMETER_RDG_TEXTURE(Texture2D, InputSceneColor)
@@ -304,14 +305,20 @@ class FTAAStandaloneCS : public FGlobalShader
 		{
 			return false;
 		}
-
-		return IsFeatureLevelSupported(Parameters.Platform, ERHIFeatureLevel::SM5);
+		
+		//Only Main config without DownSample permutations are supported on mobile platform.
+		return IsFeatureLevelSupported(Parameters.Platform, ERHIFeatureLevel::SM5) || (PermutationVector.Get<FTAAPassConfigDim>() == ETAAPassConfig::Main && !PermutationVector.Get<FTAADownsampleDim>());
 	}
 
 	static void ModifyCompilationEnvironment(const FGlobalShaderPermutationParameters& Parameters, FShaderCompilerEnvironment& OutEnvironment)
 	{
 		OutEnvironment.SetDefine(TEXT("THREADGROUP_SIZEX"), GTemporalAATileSizeX);
 		OutEnvironment.SetDefine(TEXT("THREADGROUP_SIZEY"), GTemporalAATileSizeY);
+
+		bool bIsMobileTiledGPU = RHIHasTiledGPU(Parameters.Platform) || IsSimulatedPlatform(Parameters.Platform);
+
+		// There are some mobile specific shader optimizations need to be set in the shader, such as disable shared memory usage, disable stencil texture sampling.
+		OutEnvironment.SetDefine(TEXT("AA_MOBILE_CONFIG"), bIsMobileTiledGPU ? 1 : 0);
 	}
 }; // class FTAAStandaloneCS
 
@@ -920,7 +927,14 @@ FTAAOutputs AddTemporalAAPass(
 				ResDivisorInv * InputViewRect.Min.Y * InvSizeY);
 		}
 
-		PassParameters->EyeAdaptationTexture = GetEyeAdaptationTexture(GraphBuilder, View);
+		if (View.GetFeatureLevel() <= ERHIFeatureLevel::ES3_1)
+		{
+			PassParameters->EyeAdaptationBuffer = GraphBuilder.CreateSRV(GetEyeAdaptationBuffer(GraphBuilder, View), PF_A32B32G32R32F);
+		}
+		else
+		{
+			PassParameters->EyeAdaptationTexture = GetEyeAdaptationTexture(GraphBuilder, View);
+		}
 
 		// Temporal upsample specific shader parameters.
 		{
diff --git a/Engine/Source/Runtime/Renderer/Private/PostProcess/TemporalAA.h b/Engine/Source/Runtime/Renderer/Private/PostProcess/TemporalAA.h
--- a/Engine/Source/Runtime/Renderer/Private/PostProcess/TemporalAA.h
+++ b/Engine/Source/Runtime/Renderer/Private/PostProcess/TemporalAA.h
@@ -173,7 +173,6 @@ public:
 		FRDGTextureRef SceneColorTexture;
 		FRDGTextureRef SceneDepthTexture;
 		FRDGTextureRef SceneVelocityTexture;
-		FRDGTextureRef EyeAdaptationTexture;
 	};
 
 	virtual ~ITemporalUpscaler() {};
diff --git a/Engine/Source/Runtime/Renderer/Private/RendererScene.cpp b/Engine/Source/Runtime/Renderer/Private/RendererScene.cpp
--- a/Engine/Source/Runtime/Renderer/Private/RendererScene.cpp
+++ b/Engine/Source/Runtime/Renderer/Private/RendererScene.cpp
@@ -4111,7 +4111,7 @@ void FScene::UpdateAllPrimitiveSceneInfos(FRDGBuilder& GraphBuilder, bool bAsync
 				FPrimitiveSceneInfo* PrimitiveSceneInfo = AddedLocalPrimitiveSceneInfos[AddIndex];
 				int32 PrimitiveIndex = PrimitiveSceneInfo->PackedIndex;
 
-				if (ShouldPrimitiveOutputVelocity(PrimitiveSceneInfo->Proxy) && GetFeatureLevel() > ERHIFeatureLevel::ES3_1)
+				if (ShouldPrimitiveOutputVelocity(PrimitiveSceneInfo->Proxy))
 				{
 					// We must register the initial LocalToWorld with the velocity state. 
 					// In the case of a moving component with MarkRenderStateDirty() called every frame, UpdateTransform will never happen.
@@ -4180,7 +4180,7 @@ void FScene::UpdateAllPrimitiveSceneInfos(FRDGBuilder& GraphBuilder, bool bAsync
 			// (note that the octree update relies on the bounds not being modified yet).
 			PrimitiveSceneInfo->RemoveFromScene(bUpdateStaticDrawLists);
 
-			if (ShouldPrimitiveOutputVelocity(PrimitiveSceneInfo->Proxy) && GetFeatureLevel() > ERHIFeatureLevel::ES3_1)
+			if (ShouldPrimitiveOutputVelocity(PrimitiveSceneInfo->Proxy))
 			{
 				VelocityData.UpdateTransform(PrimitiveSceneInfo, LocalToWorld, PrimitiveSceneProxy->GetLocalToWorld());
 			}
diff --git a/Engine/Source/Runtime/Renderer/Private/SceneRendering.h b/Engine/Source/Runtime/Renderer/Private/SceneRendering.h
--- a/Engine/Source/Runtime/Renderer/Private/SceneRendering.h
+++ b/Engine/Source/Runtime/Renderer/Private/SceneRendering.h
@@ -56,6 +56,7 @@ struct FSceneWithoutWaterTextures;
 struct FHairStrandsVisibilityViews;
 struct FSortedLightSetSceneInfo;
 struct FHairStrandsRenderingData;
+enum class EVelocityPass : uint32;
 struct FStrataSceneData;
 
 struct FSceneTexturesConfig;
@@ -986,9 +987,6 @@ struct FPreviousViewInfo
 
 	// Mobile bloom setup eye adaptation surface.
 	TRefCountPtr<IPooledRenderTarget> MobileBloomSetup_EyeAdaptation;
-	// Mobile temporal AA surface.
-	TRefCountPtr<IPooledRenderTarget> MobileAaBloomSunVignette;
-	TRefCountPtr<IPooledRenderTarget> MobileAaColor;
 };
 
 class FViewCommands
@@ -1784,6 +1782,8 @@ public:
 
 	virtual void Render(FRDGBuilder& GraphBuilder) = 0;
 	virtual void RenderHitProxies(FRDGBuilder& GraphBuilder) {}
+	virtual bool ShouldRenderVelocities() const { return false; }
+	virtual bool SupportsMSAA() const { return true; }
 
 	/** Creates a scene renderer based on the current feature level. */
 	static FSceneRenderer* CreateSceneRenderer(const FSceneViewFamily* InViewFamily, FHitProxyConsumer* HitProxyConsumer);
@@ -2069,6 +2069,12 @@ protected:
 
 	void SetupSceneReflectionCaptureBuffer(FRHICommandListImmediate& RHICmdList);
 
+	void RenderVelocities(
+		FRDGBuilder& GraphBuilder,
+		const FSceneTextures& SceneTextures,
+		EVelocityPass VelocityPass,
+		bool bForceVelocity);
+
 protected:
 	FGPUSceneDynamicContext GPUSceneDynamicContext;
 private:
@@ -2096,6 +2102,10 @@ public:
 
 	virtual void RenderHitProxies(FRDGBuilder& GraphBuilder) override;
 
+	virtual bool ShouldRenderVelocities() const override;
+
+	virtual bool SupportsMSAA() const override;
+
 	void RenderInverseOpacity(FRDGBuilder& GraphBuilder, const FViewInfo& View);
 
 protected:
@@ -2173,6 +2183,7 @@ private:
 	bool bShouldRenderCustomDepth;
 	bool bRequiresPixelProjectedPlanarRelfectionPass;
 	bool bRequriesAmbientOcclusionPass;
+	bool bShouldRenderVelocities;
 	static FGlobalDynamicIndexBuffer DynamicIndexBuffer;
 	static FGlobalDynamicVertexBuffer DynamicVertexBuffer;
 	static TGlobalResource<FGlobalDynamicReadBuffer> DynamicReadBuffer;
diff --git a/Engine/Source/Runtime/Renderer/Private/SceneTextures.cpp b/Engine/Source/Runtime/Renderer/Private/SceneTextures.cpp
--- a/Engine/Source/Runtime/Renderer/Private/SceneTextures.cpp
+++ b/Engine/Source/Runtime/Renderer/Private/SceneTextures.cpp
@@ -754,9 +754,6 @@ FSceneTextures& FSceneTextures::Create(FRDGBuilder& GraphBuilder, const FSceneTe
 		// Screen Space Ambient Occlusion
 		SceneTextures.ScreenSpaceAO = CreateScreenSpaceAOTexture(GraphBuilder, Config.Extent);
 
-		// Velocity
-		SceneTextures.Velocity = GraphBuilder.CreateTexture(FVelocityRendering::GetRenderTargetDesc(Config.ShaderPlatform, Config.Extent), TEXT("SceneVelocity"));
-
 		// Small Depth
 		const FIntPoint SmallDepthExtent = GetDownscaledExtent(Config.Extent, Config.SmallDepthDownsampleFactor);
 		const FRDGTextureDesc SmallDepthDesc(FRDGTextureDesc::Create2D(SmallDepthExtent, PF_DepthStencil, FClearValueBinding::None, TexCreate_DepthStencilTargetable | TexCreate_ShaderResource));
@@ -768,6 +765,9 @@ FSceneTextures& FSceneTextures::Create(FRDGBuilder& GraphBuilder, const FSceneTe
 		SceneTextures.ScreenSpaceAO = CreateMobileScreenSpaceAOTexture(GraphBuilder, Config.Extent);
 	}
 
+	// Velocity
+	SceneTextures.Velocity = GraphBuilder.CreateTexture(FVelocityRendering::GetRenderTargetDesc(Config.ShaderPlatform, Config.Extent), TEXT("SceneVelocity"));
+
 	if (Config.bIsUsingGBuffers)
 	{
 		ETextureCreateFlags FlagsToAdd = TexCreate_None;
@@ -1126,6 +1126,8 @@ void SetupMobileSceneTextureUniformParameters(
 	SceneTextureParameters.CustomDepthTextureSampler = TStaticSamplerState<>::GetRHI();
 	SceneTextureParameters.MobileCustomStencilTexture = SystemTextures.Black;
 	SceneTextureParameters.MobileCustomStencilTextureSampler = TStaticSamplerState<SF_Point, AM_Clamp, AM_Clamp, AM_Clamp>::GetRHI();
+	SceneTextureParameters.SceneVelocityTexture = SystemTextures.Black;
+	SceneTextureParameters.SceneVelocityTextureSampler = TStaticSamplerState<SF_Point, AM_Clamp, AM_Clamp, AM_Clamp>::GetRHI();
 	SceneTextureParameters.GBufferATexture = SystemTextures.Black;
 	SceneTextureParameters.GBufferBTexture = SystemTextures.Black;
 	SceneTextureParameters.GBufferCTexture = SystemTextures.Black;
@@ -1199,6 +1201,14 @@ void SetupMobileSceneTextureUniformParameters(
 				SceneTextureParameters.MobileCustomStencilTexture = CustomDepthTextures.MobileStencil;
 			}
 		}
+
+		if (EnumHasAnyFlags(SetupMode, EMobileSceneTextureSetupMode::SceneVelocity))
+		{
+			if (HasBeenProduced(SceneTextures->Velocity))
+			{
+				SceneTextureParameters.SceneVelocityTexture = SceneTextures->Velocity;
+			}
+		}
 	}
 }
 
diff --git a/Engine/Source/Runtime/Renderer/Private/VelocityRendering.cpp b/Engine/Source/Runtime/Renderer/Private/VelocityRendering.cpp
--- a/Engine/Source/Runtime/Renderer/Private/VelocityRendering.cpp
+++ b/Engine/Source/Runtime/Renderer/Private/VelocityRendering.cpp
@@ -87,9 +87,6 @@ public:
 		// Compile for materials which modify meshes.
 		const bool bMayModifyMeshes = Parameters.MaterialParameters.bMaterialMayModifyMeshPosition;
 
-		// Compile if supported by the hardware.
-		const bool bIsFeatureSupported = IsFeatureLevelSupported(Parameters.Platform, ERHIFeatureLevel::SM5);
-
 		/**
 		 * Any material with a vertex factory incompatible with base pass velocity generation must generate
 		 * permutations for this shader. Shaders which don't fall into this set are considered "simple" enough
@@ -101,7 +98,7 @@ public:
 		// The material may explicitly override and request that it be rendered into the velocity pass.
 		const bool bIsSeparateVelocityPassRequiredByMaterial = Parameters.MaterialParameters.bIsTranslucencyWritingVelocity;
 
-		return bIsFeatureSupported && (bIsSeparateVelocityPassRequired || bIsSeparateVelocityPassRequiredByMaterial);
+		return (bIsSeparateVelocityPassRequired || bIsSeparateVelocityPassRequiredByMaterial);
 	}
 
 	FVelocityVS() = default;
@@ -217,6 +214,26 @@ bool FDeferredShadingSceneRenderer::ShouldRenderVelocities() const
 	return bNeedsVelocity;
 }
 
+bool FMobileSceneRenderer::ShouldRenderVelocities() const
+{
+	if (!FVelocityRendering::IsSeparateVelocityPassSupported(ShaderPlatform) || ViewFamily.UseDebugViewPS())
+	{
+		return false;
+	}
+
+	bool bNeedsVelocity = false;
+	for (int32 ViewIndex = 0; ViewIndex < Views.Num(); ViewIndex++)
+	{
+		const FViewInfo& View = Views[ViewIndex];
+
+		bool bTemporalAA = (View.AntiAliasingMethod == AAM_TemporalAA) && !View.bCameraCut;
+
+		bNeedsVelocity |= bTemporalAA;
+	}
+
+	return bNeedsVelocity;
+}
+
 BEGIN_SHADER_PARAMETER_STRUCT(FVelocityPassParameters, )
 	SHADER_PARAMETER_STRUCT_INCLUDE(FViewShaderParameters, View)
 	SHADER_PARAMETER_RDG_UNIFORM_BUFFER(FSceneTextureUniformParameters, SceneTextures)
@@ -224,7 +241,7 @@ BEGIN_SHADER_PARAMETER_STRUCT(FVelocityPassParameters, )
 	RENDER_TARGET_BINDING_SLOTS()
 END_SHADER_PARAMETER_STRUCT()
 
-void FDeferredShadingSceneRenderer::RenderVelocities(
+void FSceneRenderer::RenderVelocities(
 	FRDGBuilder& GraphBuilder,
 	const FSceneTextures& SceneTextures,
 	EVelocityPass VelocityPass,
@@ -236,7 +253,7 @@ void FDeferredShadingSceneRenderer::RenderVelocities(
 	}
 
 	RDG_CSV_STAT_EXCLUSIVE_SCOPE(GraphBuilder, RenderVelocities);
-	SCOPED_NAMED_EVENT(FDeferredShadingSceneRenderer_RenderVelocities, FColor::Emerald);
+	SCOPED_NAMED_EVENT(FSceneRenderer_RenderVelocities, FColor::Emerald);
 	SCOPE_CYCLE_COUNTER(STAT_RenderVelocities);
 
 	ERenderTargetLoadAction VelocityLoadAction = HasBeenProduced(SceneTextures.Velocity)
@@ -264,18 +281,19 @@ void FDeferredShadingSceneRenderer::RenderVelocities(
 
 			RDG_GPU_MASK_SCOPE(GraphBuilder, View.GPUMask);
 
-			if (VelocityLoadAction == ERenderTargetLoadAction::EClear)
+			// Avoid adding a separate clear pass, clear it in the velocity pass if it is possible.
+			if (bForceVelocity && !bHasAnyDraw)
 			{
-				AddClearRenderTargetPass(GraphBuilder, SceneTextures.Velocity);
-
-				if (!View.Family->bMultiGPUForkAndJoin)
+				if (VelocityLoadAction == ERenderTargetLoadAction::EClear)
 				{
-					VelocityLoadAction = ERenderTargetLoadAction::ELoad;
+					AddClearRenderTargetPass(GraphBuilder, SceneTextures.Velocity);
+
+					if (!View.Family->bMultiGPUForkAndJoin)
+					{
+						VelocityLoadAction = ERenderTargetLoadAction::ELoad;
+					}
 				}
-			}
 
-			if (!bHasAnyDraw)
-			{
 				continue;
 			}
 
@@ -293,10 +311,10 @@ void FDeferredShadingSceneRenderer::RenderVelocities(
 					? FExclusiveDepthStencil::DepthRead_StencilWrite
 					: FExclusiveDepthStencil::DepthWrite_StencilWrite);
 
+			PassParameters->RenderTargets[0] = FRenderTargetBinding(SceneTextures.Velocity, VelocityLoadAction);
+
 			if (IsParallelVelocity())
 			{
-				PassParameters->RenderTargets[0] = FRenderTargetBinding(SceneTextures.Velocity, ERenderTargetLoadAction::ELoad);
-
 				GraphBuilder.AddPass(
 					RDG_EVENT_NAME("VelocityParallel"),
 					PassParameters,
@@ -309,7 +327,6 @@ void FDeferredShadingSceneRenderer::RenderVelocities(
 			}
 			else
 			{
-				PassParameters->RenderTargets[0] = FRenderTargetBinding(SceneTextures.Velocity, ERenderTargetLoadAction::ELoad);
 
 				GraphBuilder.AddPass(
 					RDG_EVENT_NAME("Velocity"),
@@ -331,7 +348,9 @@ EPixelFormat FVelocityRendering::GetFormat(EShaderPlatform ShaderPlatform)
 	static const auto CMeshSDFVar = IConsoleManager::Get().FindTConsoleVariableDataInt(TEXT("r.GenerateMeshDistanceFields"));
 	const bool bNeedVelocityDepth = (CMeshSDFVar->GetValueOnRenderThread() != 0 && FDataDrivenShaderPlatformInfo::GetSupportsLumenGI(ShaderPlatform)) 
 		|| FDataDrivenShaderPlatformInfo::GetSupportsRayTracing(ShaderPlatform);
-	return bNeedVelocityDepth ? PF_A16B16G16R16 : PF_G16R16;
+
+	// Android platform doesn't support unorm G16R16 format, use G16R16F instead.
+	return bNeedVelocityDepth ? PF_A16B16G16R16 : (IsAndroidOpenGLESPlatform(ShaderPlatform) ? PF_G16R16F : PF_G16R16);
 }
 
 FRDGTextureDesc FVelocityRendering::GetRenderTargetDesc(EShaderPlatform ShaderPlatform, FIntPoint Extent)
@@ -724,6 +743,8 @@ FMeshPassProcessor* CreateVelocityPassProcessor(const FScene* Scene, const FScen
 }
 
 FRegisterPassProcessorCreateFunction RegisterVelocityPass(&CreateVelocityPassProcessor, EShadingPath::Deferred,  EMeshPass::Velocity, EMeshPassFlags::CachedMeshCommands | EMeshPassFlags::MainView);
+FRegisterPassProcessorCreateFunction MobileRegisterVelocityPass(&CreateVelocityPassProcessor, EShadingPath::Mobile, EMeshPass::Velocity, EMeshPassFlags::CachedMeshCommands | EMeshPassFlags::MainView);
+
 FTranslucentVelocityMeshProcessor::FTranslucentVelocityMeshProcessor(const FScene* Scene, const FSceneView* InViewIfDynamicMeshCommand, const FMeshPassProcessorRenderState& InPassDrawRenderState, FMeshPassDrawListContext* InDrawListContext)
 	: FVelocityMeshProcessor(Scene, InViewIfDynamicMeshCommand, InPassDrawRenderState, InDrawListContext)
 {}
@@ -737,4 +758,5 @@ FMeshPassProcessor* CreateTranslucentVelocityPassProcessor(const FScene* Scene,
 	return new(FMemStack::Get()) FTranslucentVelocityMeshProcessor(Scene, InViewIfDynamicMeshCommand, VelocityPassState, InDrawListContext);
 }
 
-FRegisterPassProcessorCreateFunction RegisterTranslucentVelocityPass(&CreateTranslucentVelocityPassProcessor, EShadingPath::Deferred, EMeshPass::TranslucentVelocity, EMeshPassFlags::CachedMeshCommands | EMeshPassFlags::MainView);
\ No newline at end of file
+FRegisterPassProcessorCreateFunction RegisterTranslucentVelocityPass(&CreateTranslucentVelocityPassProcessor, EShadingPath::Deferred, EMeshPass::TranslucentVelocity, EMeshPassFlags::CachedMeshCommands | EMeshPassFlags::MainView);
+FRegisterPassProcessorCreateFunction RegisterMobileTranslucentVelocityPass(&CreateTranslucentVelocityPassProcessor, EShadingPath::Mobile, EMeshPass::TranslucentVelocity, EMeshPassFlags::CachedMeshCommands | EMeshPassFlags::MainView);
diff --git a/Engine/Source/Runtime/Renderer/Public/SceneRenderTargetParameters.h b/Engine/Source/Runtime/Renderer/Public/SceneRenderTargetParameters.h
--- a/Engine/Source/Runtime/Renderer/Public/SceneRenderTargetParameters.h
+++ b/Engine/Source/Runtime/Renderer/Public/SceneRenderTargetParameters.h
@@ -71,6 +71,8 @@ BEGIN_GLOBAL_SHADER_PARAMETER_STRUCT(FMobileSceneTextureUniformParameters, RENDE
 	SHADER_PARAMETER_SAMPLER(SamplerState, CustomDepthTextureSampler)
 	SHADER_PARAMETER_RDG_TEXTURE(Texture2D, MobileCustomStencilTexture)
 	SHADER_PARAMETER_SAMPLER(SamplerState, MobileCustomStencilTextureSampler)
+	SHADER_PARAMETER_RDG_TEXTURE(Texture2D, SceneVelocityTexture)
+	SHADER_PARAMETER_SAMPLER(SamplerState, SceneVelocityTextureSampler)
 	// GBuffer
 	SHADER_PARAMETER_RDG_TEXTURE(Texture2D, GBufferATexture)
 	SHADER_PARAMETER_RDG_TEXTURE(Texture2D, GBufferBTexture)
@@ -95,8 +97,9 @@ enum class EMobileSceneTextureSetupMode : uint32
 	GBufferC		= 1 << 5,
 	GBufferD		= 1 << 6,
 	SceneDepthAux	= 1 << 7,
+	SceneVelocity	= 1 << 8,
 	GBuffers		= GBufferA | GBufferB | GBufferC | GBufferD | SceneDepthAux,
-	All				= SceneColor | SceneDepth | CustomDepth | GBuffers
+	All				= SceneColor | SceneDepth | CustomDepth | GBuffers | SceneVelocity
 };
 ENUM_CLASS_FLAGS(EMobileSceneTextureSetupMode);
 
