DistanceField Shadow Mobile(disabled by default for the moment)

diff --git a/Engine/Config/Android/AndroidEngine.ini b/Engine/Config/Android/AndroidEngine.ini
--- a/Engine/Config/Android/AndroidEngine.ini
+++ b/Engine/Config/Android/AndroidEngine.ini
@@ -108,3 +108,6 @@ bUseGooglePlayBillingApiV2=true
 +SensorLocations=/sys/devices/platform/s5p-tmu/temperature
 +SensorLocations=/sys/devices/platform/s5p-tmu/curr_temp
 +SensorLocations=/sys/htc/cpu_temp
+
+[/Script/Engine.RendererSettings]
+r.DistanceFields=0
\ No newline at end of file
diff --git a/Engine/Config/IOS/IOSEngine.ini b/Engine/Config/IOS/IOSEngine.ini
--- a/Engine/Config/IOS/IOSEngine.ini
+++ b/Engine/Config/IOS/IOSEngine.ini
@@ -78,3 +78,6 @@ RetryResumeDataLimit=3;
 [SystemSettings]
 ; iOS should always have IPv6 support on
 net.DisableIPv6=0
+
+[/Script/Engine.RendererSettings]
+r.DistanceFields=0
\ No newline at end of file
diff --git a/Engine/Shaders/Private/DistanceFieldShadowing.usf b/Engine/Shaders/Private/DistanceFieldShadowing.usf
--- a/Engine/Shaders/Private/DistanceFieldShadowing.usf
+++ b/Engine/Shaders/Private/DistanceFieldShadowing.usf
@@ -705,7 +705,7 @@ void DistanceFieldShadowingCS(
 			float SubsurfaceDensity = 0;
 			bool bUseSubsurfaceTransmission = false;
 
-#if !FORWARD_SHADING && DF_SHADOW_QUALITY > 1
+#if !FORWARD_SHADING && DF_SHADOW_QUALITY > 1 && !SHADING_PATH_MOBILE
 			FGBufferData GBufferData = GetGBufferData(ScreenUV);
 
 			BRANCH
diff --git a/Engine/Shaders/Private/MobileBasePassPixelShader.usf b/Engine/Shaders/Private/MobileBasePassPixelShader.usf
--- a/Engine/Shaders/Private/MobileBasePassPixelShader.usf
+++ b/Engine/Shaders/Private/MobileBasePassPixelShader.usf
@@ -731,11 +731,12 @@ void Main(
 	Color *= MaterialAO;
 	IndirectIrradiance *= MaterialAO;
 
+	float  ShadowPositionZ = 0;
 #if DIRECTIONAL_LIGHT_CSM && !MATERIAL_SHADINGMODEL_SINGLELAYERWATER
 	// Cascaded Shadow Map
 	if (UseCSM)
 	{
-		half ShadowMap = MobileDirectionalLightCSM(MaterialParameters.ScreenPosition.xy, MaterialParameters.ScreenPosition.w);
+		half ShadowMap = MobileDirectionalLightCSM(MaterialParameters.ScreenPosition.xy, MaterialParameters.ScreenPosition.w, ShadowPositionZ);
 	#if MOVABLE_DIRECTIONAL_LIGHT
 		Shadow = ShadowMap;
 	#else
@@ -744,6 +745,13 @@ void Main(
 	}
 #endif /* DIRECTIONAL_LIGHT_CSM */
 
+#if (MATERIALBLENDING_SOLID || MATERIALBLENDING_MASKED) && ENABLE_DISTANCE_FIELD
+	if (ShadowPositionZ == 0)
+	{
+		Shadow = Texture2DSample(MobileBasePass.ScreenSpaceShadowMaskTexture, MobileBasePass.ScreenSpaceShadowMaskSampler, SvPositionToBufferUV(SvPosition));
+	}
+#endif
+
 	half NoL = max(0, dot(MaterialParameters.WorldNormal, MobileDirectionalLight.DirectionalLightDirectionAndShadowTransition.xyz));
 	half3 H = normalize(MaterialParameters.CameraVector + MobileDirectionalLight.DirectionalLightDirectionAndShadowTransition.xyz);
 	half NoH = max(0, dot(MaterialParameters.WorldNormal, H));
diff --git a/Engine/Shaders/Private/MobileDeferredShading.usf b/Engine/Shaders/Private/MobileDeferredShading.usf
--- a/Engine/Shaders/Private/MobileDeferredShading.usf
+++ b/Engine/Shaders/Private/MobileDeferredShading.usf
@@ -259,9 +259,10 @@ void MobileDirectLightPS(
 	half Shadow = GBuffer.PrecomputedShadowFactors.r;
 
 #if APPLY_CSM  
-	// TODO: we don't need to fully compute ScreenPosition here 
+	// TODO: we don't need to fully compute ScreenPosition here
+	float  ShadowPositionZ = 0;
 	float4 ScreenPosition = SvPositionToScreenPosition(float4(SvPosition.xyz,SceneDepth));
-	float ShadowMap = MobileDirectionalLightCSM(ScreenPosition.xy, SceneDepth);
+	float ShadowMap = MobileDirectionalLightCSM(ScreenPosition.xy, SceneDepth, ShadowPositionZ);
 
 	// Check movable light param to determine if we should be using precomputed shadows
 	Shadow = LightFunctionParameters2.z > 0.0f ? ShadowMap : min(ShadowMap, Shadow);
diff --git a/Engine/Shaders/Private/MobileShadingModels.ush b/Engine/Shaders/Private/MobileShadingModels.ush
--- a/Engine/Shaders/Private/MobileShadingModels.ush
+++ b/Engine/Shaders/Private/MobileShadingModels.ush
@@ -231,9 +231,10 @@ return ShadowMap;
 // Add fading CSM plane:
 #define FADE_CSM 1
 
-half MobileDirectionalLightCSM(float2 ScreenPosition, float SceneDepth)
+half MobileDirectionalLightCSM(float2 ScreenPosition, float SceneDepth, inout float ShadowPositionZ)
 {
 	half ShadowMap = 1;
+	ShadowPositionZ = 0;
 	FPCFSamplerSettings Settings;
 	Settings.ShadowDepthTexture = MobileDirectionalLight.DirectionalLightShadowTexture;
 	Settings.ShadowDepthTextureSampler = MobileDirectionalLight.DirectionalLightShadowSampler;
@@ -255,6 +256,7 @@ half MobileDirectionalLightCSM(float2 ScreenPosition, float SceneDepth)
 		#else
 			ShadowPosition = mul(float4(ScreenPosition.x, ScreenPosition.y, SceneDepth, 1), MobileDirectionalLight.DirectionalLightScreenToShadow[i]);
 		#endif
+			ShadowPositionZ = ShadowPosition.z;
 			break; // position found.
 		}
 	}
diff --git a/Engine/Shaders/Private/SceneTexturesCommon.ush b/Engine/Shaders/Private/SceneTexturesCommon.ush
--- a/Engine/Shaders/Private/SceneTexturesCommon.ush
+++ b/Engine/Shaders/Private/SceneTexturesCommon.ush
@@ -130,6 +130,9 @@ float LookupDeviceZ( float2 ScreenUV )
 {
 #if	SCENE_TEXTURES_DISABLED
 	return FarDepthValue;
+#elif FORCE_DEPTH_TEXTURE_READS 
+	// native Depth buffer lookup
+	return Texture2DSampleLevel(MobileSceneTextures.SceneDepthTexture, MobileSceneTextures.SceneDepthTextureSampler, ScreenUV, 0).r;
 #elif (POST_PROCESS_MATERIAL || POST_PROCESS_MATERIAL_MOBILE) && !POST_PROCESS_AR_PASSTHROUGH
 	#if MOBILE_DEFERRED_SHADING
 		return Texture2DSample(MobileSceneTextures.SceneDepthTexture, MobileSceneTextures.SceneDepthTextureSampler, ScreenUV).r;
diff --git a/Engine/Source/Developer/Android/AndroidTargetPlatform/Private/AndroidTargetPlatform.cpp b/Engine/Source/Developer/Android/AndroidTargetPlatform/Private/AndroidTargetPlatform.cpp
--- a/Engine/Source/Developer/Android/AndroidTargetPlatform/Private/AndroidTargetPlatform.cpp
+++ b/Engine/Source/Developer/Android/AndroidTargetPlatform/Private/AndroidTargetPlatform.cpp
@@ -193,11 +193,13 @@ static bool HasLicense()
 FAndroidTargetPlatform::FAndroidTargetPlatform(bool bInIsClient, const TCHAR* FlavorName, const TCHAR* OverrideIniPlatformName)
 	: TNonDesktopTargetPlatformBase(bInIsClient, FlavorName, OverrideIniPlatformName)
 	, DeviceDetection(nullptr)
+	, bDistanceField(false)
 
 {
 #if WITH_ENGINE
 	TextureLODSettings = nullptr; // These are registered by the device profile system.
 	StaticMeshLODSettings.Initialize(this);
+	GetConfigSystem()->GetBool(TEXT("/Script/Engine.RendererSettings"), TEXT("r.DistanceFields"), bDistanceField, GEngineIni);
 #endif
 
 	TickDelegate = FTickerDelegate::CreateRaw(this, &FAndroidTargetPlatform::HandleTicker);
@@ -366,6 +368,9 @@ bool FAndroidTargetPlatform::SupportsFeature( ETargetPlatformFeatures Feature )
 		case ETargetPlatformFeatures::LandscapeMeshLODStreaming:
 			return SupportsLandscapeMeshLODStreaming();
 
+		case ETargetPlatformFeatures::DistanceFieldAO:
+			return UsesDistanceFields();
+			
 		default:
 			break;
 	}
diff --git a/Engine/Source/Developer/Android/AndroidTargetPlatform/Private/AndroidTargetPlatform.h b/Engine/Source/Developer/Android/AndroidTargetPlatform/Private/AndroidTargetPlatform.h
--- a/Engine/Source/Developer/Android/AndroidTargetPlatform/Private/AndroidTargetPlatform.h
+++ b/Engine/Source/Developer/Android/AndroidTargetPlatform/Private/AndroidTargetPlatform.h
@@ -197,6 +197,11 @@ public:
 
 	virtual void InitializeDeviceDetection();
 	
+	virtual bool UsesDistanceFields() const override
+	{
+		return bDistanceField;
+	}
+
 protected:
 
 	/**
@@ -246,7 +251,10 @@ protected:
 
 	// Pointer to the device detection handler that grabs device ids in another thread
 	IAndroidDeviceDetection* DeviceDetection;
-
+	
+	// true if DistanceField is enabled
+	bool bDistanceField;
+	
 #if WITH_ENGINE
 	// Holds a cache of the target LOD settings.
 	const UTextureLODSettings* TextureLODSettings;
diff --git a/Engine/Source/Developer/IOS/IOSTargetPlatform/Private/IOSTargetPlatform.cpp b/Engine/Source/Developer/IOS/IOSTargetPlatform/Private/IOSTargetPlatform.cpp
--- a/Engine/Source/Developer/IOS/IOSTargetPlatform/Private/IOSTargetPlatform.cpp
+++ b/Engine/Source/Developer/IOS/IOSTargetPlatform/Private/IOSTargetPlatform.cpp
@@ -30,10 +30,12 @@ FIOSTargetPlatform::FIOSTargetPlatform(bool bInIsTVOS, bool bIsClientOnly)
 	// override the ini name up in the base classes, which will go into the FTargetPlatformInfo
 	: TNonDesktopTargetPlatformBase(bIsClientOnly, nullptr, bInIsTVOS ? TEXT("TVOS") : nullptr)
 	, bIsTVOS(bInIsTVOS)
+	, bDistanceField(false)
 {
 #if WITH_ENGINE
 	TextureLODSettings = nullptr; // TextureLODSettings are registered by the device profile.
 	StaticMeshLODSettings.Initialize(this);
+	GetConfigSystem()->GetBool(TEXT("/Script/Engine.RendererSettings"), TEXT("r.DistanceFields"), bDistanceField, GEngineIni);
 #endif // #if WITH_ENGINE
 
 	// initialize the connected device detector
@@ -483,6 +485,9 @@ bool FIOSTargetPlatform::SupportsFeature( ETargetPlatformFeatures Feature ) cons
 		case ETargetPlatformFeatures::LandscapeMeshLODStreaming:
 			return SupportsLandscapeMeshLODStreaming() && SupportsMetal();
 
+		case ETargetPlatformFeatures::DistanceFieldAO:
+			return UsesDistanceFields();
+
 		default:
 			break;
 	}
diff --git a/Engine/Source/Developer/IOS/IOSTargetPlatform/Private/IOSTargetPlatform.h b/Engine/Source/Developer/IOS/IOSTargetPlatform/Private/IOSTargetPlatform.h
--- a/Engine/Source/Developer/IOS/IOSTargetPlatform/Private/IOSTargetPlatform.h
+++ b/Engine/Source/Developer/IOS/IOSTargetPlatform/Private/IOSTargetPlatform.h
@@ -110,6 +110,10 @@ public:
 
 	//~ Begin ITargetPlatform Interface
 
+	virtual bool UsesDistanceFields() const override
+	{
+		return bDistanceField;
+	}
 private:
 
 	// Handles received pong messages from the LauncherDaemon.
@@ -129,6 +133,9 @@ private:
 	// Holds the message endpoint used for communicating with the LaunchDaemon.
 	TSharedPtr<FMessageEndpoint, ESPMode::ThreadSafe> MessageEndpoint;
 
+	// true if DistanceField is enabled
+	bool bDistanceField;
+
 #if WITH_ENGINE
 	// Holds the cache of the target LOD settings.
 	const UTextureLODSettings* TextureLODSettings;
diff --git a/Engine/Source/Developer/VulkanShaderFormat/Private/VulkanShaderCompiler.cpp b/Engine/Source/Developer/VulkanShaderFormat/Private/VulkanShaderCompiler.cpp
--- a/Engine/Source/Developer/VulkanShaderFormat/Private/VulkanShaderCompiler.cpp
+++ b/Engine/Source/Developer/VulkanShaderFormat/Private/VulkanShaderCompiler.cpp
@@ -1773,7 +1773,7 @@ static bool CompileWithHlslcc(const FString& PreprocessedShader, FVulkanBindingT
 		FHlslCrossCompilerContext CrossCompilerContext(CompilerInfo.CCFlags, CompilerInfo.Frequency, HlslCompilerTarget);
 
 		const bool bShareSamplers = true;
-		const bool bRequiresOESExtensions = (CompilerInfo.Input.Target.Platform == SP_VULKAN_ES3_1_LUMIN);
+		const bool bRequiresOESExtensions = true;
 
 		FVulkanLanguageSpec VulkanLanguageSpec(bShareSamplers, bRequiresOESExtensions);
 		int32 Result = 0;
diff --git a/Engine/Source/Runtime/Apple/MetalRHI/Private/MetalRHI.cpp b/Engine/Source/Runtime/Apple/MetalRHI/Private/MetalRHI.cpp
--- a/Engine/Source/Runtime/Apple/MetalRHI/Private/MetalRHI.cpp
+++ b/Engine/Source/Runtime/Apple/MetalRHI/Private/MetalRHI.cpp
@@ -196,9 +196,11 @@ FMetalDynamicRHI::FMetalDynamicRHI(ERHIFeatureLevel::Type RequestedFeatureLevel)
 	// A8 can use 256 bits of MRTs
 #if PLATFORM_TVOS
 	GRHISupportsDrawIndirect = Device.SupportsFeatureSet(mtlpp::FeatureSet::tvOS_GPUFamily2_v1);
+	GRHISupportsPixelShaderUAVs = Device.SupportsFeatureSet(mtlpp::FeatureSet::tvOS_GPUFamily2_v1);
 #else
 	GRHISupportsRWTextureBuffers = Device.SupportsFeatureSet(mtlpp::FeatureSet::iOS_GPUFamily4_v1);
 	GRHISupportsDrawIndirect = Device.SupportsFeatureSet(mtlpp::FeatureSet::iOS_GPUFamily3_v1);
+	GRHISupportsPixelShaderUAVs = Device.SupportsFeatureSet(mtlpp::FeatureSet::iOS_GPUFamily3_v1);
 
 	const mtlpp::FeatureSet FeatureSets[] = {
 		mtlpp::FeatureSet::iOS_GPUFamily1_v1,
diff --git a/Engine/Source/Runtime/Engine/Private/InstancedStaticMesh.cpp b/Engine/Source/Runtime/Engine/Private/InstancedStaticMesh.cpp
--- a/Engine/Source/Runtime/Engine/Private/InstancedStaticMesh.cpp
+++ b/Engine/Source/Runtime/Engine/Private/InstancedStaticMesh.cpp
@@ -2908,7 +2908,7 @@ static bool ComponentRequestsCPUAccess(UInstancedStaticMeshComponent* InComponen
 	// Check mesh distance fields
 	if (StaticMesh)
 	{
-		if (FeatureLevel > ERHIFeatureLevel::ES3_1)
+		if ((FeatureLevel > ERHIFeatureLevel::ES3_1) || IsMobileDistanceFieldEnabled(GMaxRHIShaderPlatform))
 		{
 			static const auto CVar = IConsoleManager::Get().FindTConsoleVariableDataInt(TEXT("r.GenerateMeshDistanceFields"));
 
diff --git a/Engine/Source/Runtime/Engine/Private/StaticMesh.cpp b/Engine/Source/Runtime/Engine/Private/StaticMesh.cpp
--- a/Engine/Source/Runtime/Engine/Private/StaticMesh.cpp
+++ b/Engine/Source/Runtime/Engine/Private/StaticMesh.cpp
@@ -1628,7 +1628,7 @@ void FStaticMeshRenderData::Serialize(FArchive& Ar, UStaticMesh* Owner, bool bCo
 		uint8 ClassDataStripFlags = 0;
 
 #if WITH_EDITOR
-		const bool bWantToStripDistanceFieldData = Ar.IsCooking() && (!Ar.CookingTarget()->SupportsFeature(ETargetPlatformFeatures::DeferredRendering) || !Ar.CookingTarget()->SupportsFeature(ETargetPlatformFeatures::DistanceFieldAO));
+		const bool bWantToStripDistanceFieldData = Ar.IsCooking() && !Ar.CookingTarget()->SupportsFeature(ETargetPlatformFeatures::DistanceFieldAO);
 
 		ClassDataStripFlags |= (bWantToStripDistanceFieldData ? DistanceFieldDataStripFlag : 0);
 #endif
diff --git a/Engine/Source/Runtime/Landscape/Private/LandscapeRender.cpp b/Engine/Source/Runtime/Landscape/Private/LandscapeRender.cpp
--- a/Engine/Source/Runtime/Landscape/Private/LandscapeRender.cpp
+++ b/Engine/Source/Runtime/Landscape/Private/LandscapeRender.cpp
@@ -1446,7 +1446,7 @@ PRAGMA_ENABLE_DEPRECATION_WARNINGS
 	SharedBuffersKey = (SubsectionSizeLog2 & 0xf) | ((NumSubsections & 0xf) << 4) |
 		(FeatureLevel <= ERHIFeatureLevel::ES3_1 ? 0 : 1 << 30) | (XYOffsetmapTexture == nullptr ? 0 : 1 << 31);
 
-	bSupportsHeightfieldRepresentation = true;
+	bSupportsHeightfieldRepresentation = FeatureLevel > ERHIFeatureLevel::ES3_1;
 
 #if WITH_EDITOR
 	TArray<FWeightmapLayerAllocationInfo>& ComponentWeightmapLayerAllocations = InComponent->GetWeightmapLayerAllocations();
diff --git a/Engine/Source/Runtime/RHI/Private/RHI.cpp b/Engine/Source/Runtime/RHI/Private/RHI.cpp
--- a/Engine/Source/Runtime/RHI/Private/RHI.cpp
+++ b/Engine/Source/Runtime/RHI/Private/RHI.cpp
@@ -961,6 +961,8 @@ bool GRHISupportsConservativeRasterization = false;
 /** Whether we are profiling GPU hitches. */
 bool GTriggerGPUHitchProfile = false;
 
+bool GRHISupportsPixelShaderUAVs = true;
+
 FVertexElementTypeSupportInfo GVertexElementTypeSupport;
 
 RHI_API int32 volatile GCurrentTextureMemorySize = 0;
@@ -1316,15 +1318,6 @@ RHI_API bool RHISupportsTessellation(const FStaticShaderPlatform Platform)
 	return false;
 }
 
-RHI_API bool RHISupportsPixelShaderUAVs(const FStaticShaderPlatform Platform)
-{
-	if (IsFeatureLevelSupported(Platform, ERHIFeatureLevel::SM5))
-	{
-		return true;
-	}
-	return false;
-}
-
 RHI_API bool RHISupportsIndexBufferUAVs(const FStaticShaderPlatform Platform)
 {
 	return Platform == SP_PCD3D_SM5 || IsVulkanPlatform(Platform) || IsMetalSM5Platform(Platform)
diff --git a/Engine/Source/Runtime/RHI/Public/RHI.h b/Engine/Source/Runtime/RHI/Public/RHI.h
--- a/Engine/Source/Runtime/RHI/Public/RHI.h
+++ b/Engine/Source/Runtime/RHI/Public/RHI.h
@@ -103,6 +103,9 @@ extern RHI_API uint32 GRHIDeviceRevision;
 // 0 means not defined yet, use functions like IsRHIDeviceAMD() to access
 extern RHI_API uint32 GRHIVendorId;
 
+// true if the RHI supports Pixel Shader UAV
+extern RHI_API bool GRHISupportsPixelShaderUAVs;
+
 // to trigger GPU specific optimizations and fallbacks
 RHI_API bool IsRHIDeviceAMD();
 
diff --git a/Engine/Source/Runtime/RenderCore/Private/RenderUtils.cpp b/Engine/Source/Runtime/RenderCore/Private/RenderUtils.cpp
--- a/Engine/Source/Runtime/RenderCore/Private/RenderUtils.cpp
+++ b/Engine/Source/Runtime/RenderCore/Private/RenderUtils.cpp
@@ -1053,6 +1053,19 @@ RENDERCORE_API bool AllowPixelDepthOffset(const FStaticShaderPlatform Platform)
 	return true;
 }
 
+RENDERCORE_API bool IsMobileDistanceFieldEnabled(const FStaticShaderPlatform Platform)
+{
+	return IsMobilePlatform(Platform) && IsSwitchPlatform(Platform) && IsUsingDistanceFields(Platform);
+}
+
+RENDERCORE_API bool IsMobileDistanceFieldShadowingEnabled(const FStaticShaderPlatform Platform)
+{
+	static auto CVar = IConsoleManager::Get().FindConsoleVariable(TEXT("r.DistanceFieldShadowing"));
+	bool bDistanceFieldShadowingEnabled = CVar && (CVar->GetInt() != 0);
+
+	return GRHISupportsPixelShaderUAVs && bDistanceFieldShadowingEnabled && IsMobileDistanceFieldEnabled(Platform);
+}
+
 RENDERCORE_API int32 GUseForwardShading = 0;
 static FAutoConsoleVariableRef CVarForwardShading(
 	TEXT("r.ForwardShading"),
diff --git a/Engine/Source/Runtime/RenderCore/Private/Shader.cpp b/Engine/Source/Runtime/RenderCore/Private/Shader.cpp
--- a/Engine/Source/Runtime/RenderCore/Private/Shader.cpp
+++ b/Engine/Source/Runtime/RenderCore/Private/Shader.cpp
@@ -1542,6 +1542,10 @@ void ShaderMapAppendKeyString(EShaderPlatform Platform, FString& KeyString)
 			int32 GTAOPreIntegratedTextureType = MobileGTAOPreIntegratedTextureTypeCVar ? MobileGTAOPreIntegratedTextureTypeCVar->GetInt() : 0;
 			KeyString += ((MobileAmbientOcclusionCVar && MobileAmbientOcclusionCVar->GetInt() != 0) && (MobileHDRCVar && MobileHDRCVar->GetInt() !=0)) ? FString::Printf(TEXT("_MobileAO_%d"), GTAOPreIntegratedTextureType) : TEXT("");
 		}
+
+		{
+			KeyString += IsMobileDistanceFieldEnabled(Platform) ? TEXT("_MobSDF") : TEXT("");
+		}
 	}
 	else
 	{
diff --git a/Engine/Source/Runtime/RenderCore/Public/RenderUtils.h b/Engine/Source/Runtime/RenderCore/Public/RenderUtils.h
--- a/Engine/Source/Runtime/RenderCore/Public/RenderUtils.h
+++ b/Engine/Source/Runtime/RenderCore/Public/RenderUtils.h
@@ -502,6 +502,10 @@ RENDERCORE_API bool MaskedInEarlyPass(const FStaticShaderPlatform Platform);
 
 RENDERCORE_API bool AllowPixelDepthOffset(const FStaticShaderPlatform Platform);
 
+RENDERCORE_API bool IsMobileDistanceFieldEnabled(const FStaticShaderPlatform Platform);
+
+RENDERCORE_API bool IsMobileDistanceFieldShadowingEnabled(const FStaticShaderPlatform Platform);
+
 /** Returns if ForwardShading is enabled. Only valid for the current platform (otherwise call ITargetPlatform::UsesForwardShading()). */
 inline bool IsForwardShadingEnabled(const FStaticShaderPlatform Platform)
 {
diff --git a/Engine/Source/Runtime/Renderer/Private/CapsuleShadowRendering.cpp b/Engine/Source/Runtime/Renderer/Private/CapsuleShadowRendering.cpp
--- a/Engine/Source/Runtime/Renderer/Private/CapsuleShadowRendering.cpp
+++ b/Engine/Source/Runtime/Renderer/Private/CapsuleShadowRendering.cpp
@@ -1396,7 +1396,7 @@ void FDeferredShadingSceneRenderer::RenderIndirectCapsuleShadows(FRDGBuilder& Gr
 	}
 }
 
-bool FDeferredShadingSceneRenderer::ShouldPrepareForDFInsetIndirectShadow() const
+bool FSceneRenderer::ShouldPrepareForDFInsetIndirectShadow() const
 {
 	bool bSceneHasInsetDFPrimitives = false;
 
diff --git a/Engine/Source/Runtime/Renderer/Private/DeferredShadingRenderer.cpp b/Engine/Source/Runtime/Renderer/Private/DeferredShadingRenderer.cpp
--- a/Engine/Source/Runtime/Renderer/Private/DeferredShadingRenderer.cpp
+++ b/Engine/Source/Runtime/Renderer/Private/DeferredShadingRenderer.cpp
@@ -219,7 +219,6 @@ DECLARE_CYCLE_STAT(TEXT("InitViews Intentional Stall"), STAT_InitViews_Intention
 DECLARE_CYCLE_STAT(TEXT("DeferredShadingSceneRenderer UpdateDownsampledDepthSurface"), STAT_FDeferredShadingSceneRenderer_UpdateDownsampledDepthSurface, STATGROUP_SceneRendering);
 DECLARE_CYCLE_STAT(TEXT("DeferredShadingSceneRenderer Render Init"), STAT_FDeferredShadingSceneRenderer_Render_Init, STATGROUP_SceneRendering);
 DECLARE_CYCLE_STAT(TEXT("DeferredShadingSceneRenderer Render ServiceLocalQueue"), STAT_FDeferredShadingSceneRenderer_Render_ServiceLocalQueue, STATGROUP_SceneRendering);
-DECLARE_CYCLE_STAT(TEXT("DeferredShadingSceneRenderer DistanceFieldAO Init"), STAT_FDeferredShadingSceneRenderer_DistanceFieldAO_Init, STATGROUP_SceneRendering);
 DECLARE_CYCLE_STAT(TEXT("DeferredShadingSceneRenderer FGlobalDynamicVertexBuffer Commit"), STAT_FDeferredShadingSceneRenderer_FGlobalDynamicVertexBuffer_Commit, STATGROUP_SceneRendering);
 DECLARE_CYCLE_STAT(TEXT("DeferredShadingSceneRenderer FXSystem PreRender"), STAT_FDeferredShadingSceneRenderer_FXSystem_PreRender, STATGROUP_SceneRendering);
 DECLARE_CYCLE_STAT(TEXT("DeferredShadingSceneRenderer AllocGBufferTargets"), STAT_FDeferredShadingSceneRenderer_AllocGBufferTargets, STATGROUP_SceneRendering);
@@ -457,70 +456,6 @@ static void RenderOpaqueFX(
 	}
 }
 
-void FDeferredShadingSceneRenderer::PrepareDistanceFieldScene(FRDGBuilder& GraphBuilder, bool bSplitDispatch)
-{
-	CSV_SCOPED_TIMING_STAT_EXCLUSIVE(RenderDFAO);
-	SCOPE_CYCLE_COUNTER(STAT_FDeferredShadingSceneRenderer_DistanceFieldAO_Init);
-
-	const bool bShouldPrepareHeightFieldScene = ShouldPrepareHeightFieldScene();
-	const bool bShouldPrepareDistanceFieldScene = ShouldPrepareDistanceFieldScene();
-
-	if (bShouldPrepareHeightFieldScene)
-	{
-		extern int32 GHFShadowQuality;
-		if (GHFShadowQuality > 2)
-		{
-			GHFVisibilityTextureAtlas.UpdateAllocations(GraphBuilder, FeatureLevel);
-		}
-		GHeightFieldTextureAtlas.UpdateAllocations(GraphBuilder, FeatureLevel);
-		UpdateGlobalHeightFieldObjectBuffers(GraphBuilder);
-	}
-	else if (bShouldPrepareDistanceFieldScene)
-	{
-		AddOrRemoveSceneHeightFieldPrimitives();
-	}
-
-	if (bShouldPrepareDistanceFieldScene)
-	{
-		auto DispatchToRHIThreadPass = [](FRHICommandListImmediate& RHICmdList)
-		{
-			RHICmdList.ImmediateFlush(EImmediateFlushType::DispatchToRHIThread);
-		};
-
-		GDistanceFieldVolumeTextureAtlas.UpdateAllocations(GraphBuilder, FeatureLevel);
-		UpdateGlobalDistanceFieldObjectBuffers(GraphBuilder);
-		if (bSplitDispatch)
-		{
-			AddPass(GraphBuilder, DispatchToRHIThreadPass);
-		}
-		for (int32 ViewIndex = 0; ViewIndex < Views.Num(); ViewIndex++)
-		{
-			FViewInfo& View = Views[ViewIndex];
-
-			RDG_GPU_MASK_SCOPE(GraphBuilder, View.GPUMask);
-
-			View.HeightfieldLightingViewInfo.SetupVisibleHeightfields(View, GraphBuilder);
-
-			if (ShouldPrepareGlobalDistanceField())
-			{
-				float OcclusionMaxDistance = Scene->DefaultMaxDistanceFieldOcclusionDistance;
-
-				// Use the skylight's max distance if there is one
-				if (Scene->SkyLight && Scene->SkyLight->bCastShadows && !Scene->SkyLight->bWantsStaticShadowing)
-				{
-					OcclusionMaxDistance = Scene->SkyLight->OcclusionMaxDistance;
-				}
-
-				const bool bLumenEnabled = GetViewPipelineState(View).DiffuseIndirectMethod == EDiffuseIndirectMethod::Lumen || GetViewPipelineState(View).ReflectionsMethod == EReflectionsMethod::Lumen;
-				UpdateGlobalDistanceFieldVolume(GraphBuilder, Views[ViewIndex], Scene, OcclusionMaxDistance, bLumenEnabled, Views[ViewIndex].GlobalDistanceFieldInfo);
-			}
-		}
-		if (!bSplitDispatch)
-		{
-			AddPass(GraphBuilder, DispatchToRHIThreadPass);
-		}
-	}
-}
 
 #if RHI_RAYTRACING
 
diff --git a/Engine/Source/Runtime/Renderer/Private/DeferredShadingRenderer.h b/Engine/Source/Runtime/Renderer/Private/DeferredShadingRenderer.h
--- a/Engine/Source/Runtime/Renderer/Private/DeferredShadingRenderer.h
+++ b/Engine/Source/Runtime/Renderer/Private/DeferredShadingRenderer.h
@@ -330,6 +330,26 @@ private:
 		return *ViewPipelineStates[ViewIndex];
 	}
 
+	virtual bool IsLumenEnabled(const FViewInfo& View) const override
+	{ 
+		return (GetViewPipelineState(View).DiffuseIndirectMethod == EDiffuseIndirectMethod::Lumen || GetViewPipelineState(View).ReflectionsMethod == EReflectionsMethod::Lumen);
+	}
+
+	virtual bool AnyViewHasSupportingGIMethod() const override
+	{
+		bool bAnyViewHasSupportingGIMethod = false;
+
+		for (int32 ViewIndex = 0; ViewIndex < Views.Num(); ViewIndex++)
+		{
+			if (GetViewPipelineState(Views[ViewIndex]).DiffuseIndirectMethod == EDiffuseIndirectMethod::Disabled || GetViewPipelineState(Views[ViewIndex]).DiffuseIndirectMethod == EDiffuseIndirectMethod::SSGI)
+			{
+				bAnyViewHasSupportingGIMethod = true;
+			}
+		}
+
+		return bAnyViewHasSupportingGIMethod;
+	}
+
 	static FGraphEventRef TranslucencyTimestampQuerySubmittedFence[FOcclusionQueryHelpers::MaxBufferedOcclusionFrames + 1];
 	static FGlobalDynamicIndexBuffer DynamicIndexBufferForInitViews;
 	static FGlobalDynamicIndexBuffer DynamicIndexBufferForInitShadows;
@@ -794,20 +814,6 @@ private:
 	
 	bool ShouldRenderDistanceFieldAO() const;
 
-	/** Whether distance field global data structures should be prepared for features that use it. */
-	bool ShouldPrepareForDistanceFieldShadows() const;
-	bool ShouldPrepareForDistanceFieldAO() const;
-	bool ShouldPrepareForDFInsetIndirectShadow() const;
-
-	bool ShouldPrepareDistanceFieldScene() const;
-	bool ShouldPrepareGlobalDistanceField() const;
-	bool ShouldPrepareHeightFieldScene() const;
-
-	void UpdateGlobalDistanceFieldObjectBuffers(FRDGBuilder& GraphBuilder);
-	void UpdateGlobalHeightFieldObjectBuffers(FRDGBuilder& GraphBuilder);
-	void AddOrRemoveSceneHeightFieldPrimitives(bool bSkipAdd = false);
-	void PrepareDistanceFieldScene(FRDGBuilder& GraphBuilder, bool bSplitDispatch);
-
 	void CopySceneCaptureComponentToTarget(
 		FRDGBuilder& GraphBuilder,
 		TRDGUniformBufferRef<FSceneTextureUniformParameters> SceneTexturesUniformBuffer,
diff --git a/Engine/Source/Runtime/Renderer/Private/DistanceFieldAmbientOcclusion.cpp b/Engine/Source/Runtime/Renderer/Private/DistanceFieldAmbientOcclusion.cpp
--- a/Engine/Source/Runtime/Renderer/Private/DistanceFieldAmbientOcclusion.cpp
+++ b/Engine/Source/Runtime/Renderer/Private/DistanceFieldAmbientOcclusion.cpp
@@ -672,17 +672,10 @@ bool ShouldRenderDeferredDynamicSkyLight(const FScene* Scene, const FSceneViewFa
 		&& !ShouldRenderRayTracingSkyLight(Scene->SkyLight); // Disable diffuse sky contribution if evaluated by RT Sky;
 }
 
-bool FDeferredShadingSceneRenderer::ShouldPrepareForDistanceFieldAO() const
-{
-	bool bAnyViewHasSupportingGIMethod = false;
 
-	for (int32 ViewIndex = 0; ViewIndex < Views.Num(); ViewIndex++)
-	{
-		if (GetViewPipelineState(Views[ViewIndex]).DiffuseIndirectMethod == EDiffuseIndirectMethod::Disabled || GetViewPipelineState(Views[ViewIndex]).DiffuseIndirectMethod == EDiffuseIndirectMethod::SSGI)
-		{
-			bAnyViewHasSupportingGIMethod = true;
-		}
-	}
+bool FSceneRenderer::ShouldPrepareForDistanceFieldAO() const
+{
+	bool bAnyViewHasSupportingGIMethod = AnyViewHasSupportingGIMethod();
 
 	return SupportsDistanceFieldAO(Scene->GetFeatureLevel(), Scene->GetShaderPlatform())
 		&& ((ShouldRenderDeferredDynamicSkyLight(Scene, ViewFamily) && bAnyViewHasSupportingGIMethod && Scene->SkyLight->bCastShadows && ViewFamily.EngineShowFlags.DistanceFieldAO)
@@ -692,7 +685,7 @@ bool FDeferredShadingSceneRenderer::ShouldPrepareForDistanceFieldAO() const
 			|| (GDistanceFieldAOApplyToStaticIndirect && bAnyViewHasSupportingGIMethod && ViewFamily.EngineShowFlags.DistanceFieldAO));
 }
 
-bool FDeferredShadingSceneRenderer::ShouldPrepareDistanceFieldScene() const
+bool FSceneRenderer::ShouldPrepareDistanceFieldScene() const
 {
 	if (!ensure(Scene != nullptr))
 	{
@@ -713,7 +706,7 @@ bool FDeferredShadingSceneRenderer::ShouldPrepareDistanceFieldScene() const
 	return bShouldPrepareGlobalDistanceField || bShouldPrepareForAO || ShouldPrepareForDistanceFieldShadows() || bShouldPrepareForDFInsetIndirectShadow;
 }
 
-bool FDeferredShadingSceneRenderer::ShouldPrepareGlobalDistanceField() const
+bool FSceneRenderer::ShouldPrepareGlobalDistanceField() const
 {
 	if (!ensure(Scene != nullptr))
 	{
@@ -725,9 +718,7 @@ bool FDeferredShadingSceneRenderer::ShouldPrepareGlobalDistanceField() const
 			|| ((Views.Num() > 0) && Views[0].bUsesGlobalDistanceField)
 			|| ((FXSystem != nullptr) && FXSystem->UsesGlobalDistanceField()));
 
-	bShouldPrepareForAO = bShouldPrepareForAO 
-		|| GetViewPipelineState(Views[0]).DiffuseIndirectMethod == EDiffuseIndirectMethod::Lumen 
-		|| GetViewPipelineState(Views[0]).ReflectionsMethod == EReflectionsMethod::Lumen;
+	bShouldPrepareForAO = bShouldPrepareForAO || IsLumenEnabled(Views[0]);
 
 	return bShouldPrepareForAO && UseGlobalDistanceField();
 }
diff --git a/Engine/Source/Runtime/Renderer/Private/DistanceFieldObjectCulling.cpp b/Engine/Source/Runtime/Renderer/Private/DistanceFieldObjectCulling.cpp
--- a/Engine/Source/Runtime/Renderer/Private/DistanceFieldObjectCulling.cpp
+++ b/Engine/Source/Runtime/Renderer/Private/DistanceFieldObjectCulling.cpp
@@ -381,7 +381,7 @@ public:
 
 	static bool ShouldCompilePermutation(const FGlobalShaderPermutationParameters& Parameters)
 	{
-		return IsFeatureLevelSupported(Parameters.Platform, ERHIFeatureLevel::SM5) && DoesPlatformSupportDistanceFieldAO(Parameters.Platform) && RHISupportsPixelShaderUAVs(Parameters.Platform) && IsUsingDistanceFields(Parameters.Platform);
+		return IsFeatureLevelSupported(Parameters.Platform, ERHIFeatureLevel::SM5) && DoesPlatformSupportDistanceFieldAO(Parameters.Platform) && IsUsingDistanceFields(Parameters.Platform);
 	}
 
 	static void ModifyCompilationEnvironment(const FGlobalShaderPermutationParameters& Parameters, FShaderCompilerEnvironment& OutEnvironment)
diff --git a/Engine/Source/Runtime/Renderer/Private/DistanceFieldObjectManagement.cpp b/Engine/Source/Runtime/Renderer/Private/DistanceFieldObjectManagement.cpp
--- a/Engine/Source/Runtime/Renderer/Private/DistanceFieldObjectManagement.cpp
+++ b/Engine/Source/Runtime/Renderer/Private/DistanceFieldObjectManagement.cpp
@@ -19,6 +19,7 @@
 #include "ScenePrivate.h"
 #include "DistanceFieldLightingShared.h"
 #include "DistanceFieldAmbientOcclusion.h"
+#include "GlobalDistanceField.h"
 
 extern void LumenUpdateDFObjectIndex(FScene* Scene, int32 MappingIndex);
 
@@ -38,6 +39,7 @@ FAutoConsoleVariableRef CVarAOLogObjectBufferReallocation(
 	ECVF_RenderThreadSafe
 	);
 
+DECLARE_CYCLE_STAT(TEXT("SceneRenderer DistanceFieldAO Init"), STAT_FSceneRenderer_DistanceFieldAO_Init, STATGROUP_SceneRendering);
 int32 GDFParallelUpdate = 0;
 FAutoConsoleVariableRef CVarDFParallelUpdate(
 	TEXT("r.DF.ParallelUpdate"),
@@ -444,7 +446,7 @@ bool ProcessHeightFieldPrimitiveUpdate(
 
 bool bVerifySceneIntegrity = false;
 
-void FDeferredShadingSceneRenderer::UpdateGlobalDistanceFieldObjectBuffers(FRDGBuilder& GraphBuilder)
+void FSceneRenderer::UpdateGlobalDistanceFieldObjectBuffers(FRDGBuilder& GraphBuilder)
 {
 	const bool bExecuteInParallel = GDFParallelUpdate != 0 && FApp::ShouldUseThreadingForPerformance();
 
@@ -750,7 +752,7 @@ void FDeferredShadingSceneRenderer::UpdateGlobalDistanceFieldObjectBuffers(FRDGB
 	DistanceFieldSceneData.IndicesToUpdateInObjectBuffers.Reset();
 }
 
-void FDeferredShadingSceneRenderer::UpdateGlobalHeightFieldObjectBuffers(FRDGBuilder& GraphBuilder)
+void FSceneRenderer::UpdateGlobalHeightFieldObjectBuffers(FRDGBuilder& GraphBuilder)
 {
 	FDistanceFieldSceneData& DistanceFieldSceneData = Scene->DistanceFieldSceneData;
 
@@ -927,7 +929,71 @@ void FDeferredShadingSceneRenderer::UpdateGlobalHeightFieldObjectBuffers(FRDGBui
 	}
 }
 
-void FDeferredShadingSceneRenderer::AddOrRemoveSceneHeightFieldPrimitives(bool bSkipAdd)
+void FSceneRenderer::PrepareDistanceFieldScene(FRDGBuilder& GraphBuilder, bool bSplitDispatch)
+{
+	CSV_SCOPED_TIMING_STAT_EXCLUSIVE(RenderDFAO);
+	SCOPE_CYCLE_COUNTER(STAT_FSceneRenderer_DistanceFieldAO_Init);
+
+	const bool bShouldPrepareHeightFieldScene = ShouldPrepareHeightFieldScene();
+	const bool bShouldPrepareDistanceFieldScene = ShouldPrepareDistanceFieldScene();
+
+	if (bShouldPrepareHeightFieldScene)
+	{
+		extern int32 GHFShadowQuality;
+		if (GHFShadowQuality > 2)
+		{
+			GHFVisibilityTextureAtlas.UpdateAllocations(GraphBuilder, FeatureLevel);
+		}
+		GHeightFieldTextureAtlas.UpdateAllocations(GraphBuilder, FeatureLevel);
+		UpdateGlobalHeightFieldObjectBuffers(GraphBuilder);
+	}
+	else if (bShouldPrepareDistanceFieldScene)
+	{
+		AddOrRemoveSceneHeightFieldPrimitives();
+	}
+
+	if (bShouldPrepareDistanceFieldScene)
+	{
+		auto DispatchToRHIThreadPass = [](FRHICommandListImmediate& RHICmdList)
+		{
+			RHICmdList.ImmediateFlush(EImmediateFlushType::DispatchToRHIThread);
+		};
+
+		GDistanceFieldVolumeTextureAtlas.UpdateAllocations(GraphBuilder, FeatureLevel);
+		UpdateGlobalDistanceFieldObjectBuffers(GraphBuilder);
+		if (bSplitDispatch)
+		{
+			AddPass(GraphBuilder, DispatchToRHIThreadPass);
+		}
+		for (int32 ViewIndex = 0; ViewIndex < Views.Num(); ViewIndex++)
+		{
+			FViewInfo& View = Views[ViewIndex];
+
+			RDG_GPU_MASK_SCOPE(GraphBuilder, View.GPUMask);
+
+			View.HeightfieldLightingViewInfo.SetupVisibleHeightfields(View, GraphBuilder);
+
+			if (ShouldPrepareGlobalDistanceField())
+			{
+				float OcclusionMaxDistance = Scene->DefaultMaxDistanceFieldOcclusionDistance;
+
+				// Use the skylight's max distance if there is one
+				if (Scene->SkyLight && Scene->SkyLight->bCastShadows && !Scene->SkyLight->bWantsStaticShadowing)
+				{
+					OcclusionMaxDistance = Scene->SkyLight->OcclusionMaxDistance;
+				}
+
+				UpdateGlobalDistanceFieldVolume(GraphBuilder, Views[ViewIndex], Scene, OcclusionMaxDistance, IsLumenEnabled(View), Views[ViewIndex].GlobalDistanceFieldInfo);
+			}
+		}
+		if (!bSplitDispatch)
+		{
+			AddPass(GraphBuilder, DispatchToRHIThreadPass);
+		}
+	}
+}
+
+void FSceneRenderer::AddOrRemoveSceneHeightFieldPrimitives(bool bSkipAdd)
 {
 	FDistanceFieldSceneData& SceneData = Scene->DistanceFieldSceneData;
 
diff --git a/Engine/Source/Runtime/Renderer/Private/DistanceFieldShadowing.cpp b/Engine/Source/Runtime/Renderer/Private/DistanceFieldShadowing.cpp
--- a/Engine/Source/Runtime/Renderer/Private/DistanceFieldShadowing.cpp
+++ b/Engine/Source/Runtime/Renderer/Private/DistanceFieldShadowing.cpp
@@ -151,7 +151,7 @@ class FCullObjectsForShadowCS : public FGlobalShader
 
 	static bool ShouldCompilePermutation(const FGlobalShaderPermutationParameters& Parameters)
 	{
-		return IsFeatureLevelSupported(Parameters.Platform, ERHIFeatureLevel::SM5) && DoesPlatformSupportDistanceFieldShadowing(Parameters.Platform);
+		return DoesPlatformSupportDistanceFieldShadowing(Parameters.Platform);
 	}
 
 	static void ModifyCompilationEnvironment(const FGlobalShaderPermutationParameters& Parameters, FShaderCompilerEnvironment& OutEnvironment)
@@ -181,7 +181,7 @@ class FShadowObjectCullVS : public FGlobalShader
 
 	static bool ShouldCompilePermutation(const FGlobalShaderPermutationParameters& Parameters) 
 	{ 
-		return IsFeatureLevelSupported(Parameters.Platform, ERHIFeatureLevel::SM5) && DoesPlatformSupportDistanceFieldShadowing(Parameters.Platform); 
+		return DoesPlatformSupportDistanceFieldShadowing(Parameters.Platform); 
 	}
 };
 
@@ -204,7 +204,7 @@ class FShadowObjectCullPS : public FGlobalShader
 
 	static bool ShouldCompilePermutation(const FGlobalShaderPermutationParameters& Parameters)
 	{
-		return IsFeatureLevelSupported(Parameters.Platform, ERHIFeatureLevel::SM5) && DoesPlatformSupportDistanceFieldShadowing(Parameters.Platform) && RHISupportsPixelShaderUAVs(Parameters.Platform);
+		return DoesPlatformSupportDistanceFieldShadowing(Parameters.Platform);
 	}
 
 	static void ModifyCompilationEnvironment(const FGlobalShaderPermutationParameters& Parameters, FShaderCompilerEnvironment& OutEnvironment)
@@ -238,7 +238,7 @@ class FDistanceFieldShadowingCS : public FGlobalShader
 	BEGIN_SHADER_PARAMETER_STRUCT(FParameters, )
 		SHADER_PARAMETER_RDG_TEXTURE_UAV(RWTexture2D<float2>, RWShadowFactors)
 		SHADER_PARAMETER_STRUCT_REF(FViewUniformShaderParameters, View)
-		SHADER_PARAMETER_RDG_UNIFORM_BUFFER(FSceneTextureUniformParameters, SceneTextures)
+		SHADER_PARAMETER_STRUCT_INCLUDE(FSceneTextureShaderParameters, SceneTextures)
 		SHADER_PARAMETER(FVector2D, NumGroups)
 		SHADER_PARAMETER(FVector, LightDirection)
 		SHADER_PARAMETER(FVector4, LightPositionAndInvRadius)
@@ -269,7 +269,7 @@ class FDistanceFieldShadowingCS : public FGlobalShader
 
 	static bool ShouldCompilePermutation(const FGlobalShaderPermutationParameters& Parameters)
 	{
-		return IsFeatureLevelSupported(Parameters.Platform, ERHIFeatureLevel::SM5) && DoesPlatformSupportDistanceFieldShadowing(Parameters.Platform);
+		return DoesPlatformSupportDistanceFieldShadowing(Parameters.Platform);
 	}
 
 	static void ModifyCompilationEnvironment(const FGlobalShaderPermutationParameters& Parameters, FShaderCompilerEnvironment& OutEnvironment)
@@ -288,7 +288,7 @@ class FDistanceFieldShadowingUpsamplePS : public FGlobalShader
 	SHADER_USE_PARAMETER_STRUCT(FDistanceFieldShadowingUpsamplePS, FGlobalShader);
 
 	BEGIN_SHADER_PARAMETER_STRUCT(FParameters, )
-		SHADER_PARAMETER_RDG_UNIFORM_BUFFER(FSceneTextureUniformParameters, SceneTextures)
+		SHADER_PARAMETER_STRUCT_INCLUDE(FSceneTextureShaderParameters, SceneTextures)
 		SHADER_PARAMETER_STRUCT_REF(FViewUniformShaderParameters, View)
 		SHADER_PARAMETER_RDG_TEXTURE(Texture2D, ShadowFactorsTexture)
 		SHADER_PARAMETER_SAMPLER(SamplerState, ShadowFactorsSampler)
@@ -304,12 +304,13 @@ class FDistanceFieldShadowingUpsamplePS : public FGlobalShader
 
 	static bool ShouldCompilePermutation(const FGlobalShaderPermutationParameters& Parameters)
 	{
-		return IsFeatureLevelSupported(Parameters.Platform, ERHIFeatureLevel::SM5) && DoesPlatformSupportDistanceFieldShadowing(Parameters.Platform);
+		return DoesPlatformSupportDistanceFieldShadowing(Parameters.Platform);
 	}
 
 	static void ModifyCompilationEnvironment(const FGlobalShaderPermutationParameters& Parameters, FShaderCompilerEnvironment& OutEnvironment)
 	{
 		OutEnvironment.SetDefine(TEXT("DOWNSAMPLE_FACTOR"), GAODownsampleFactor);
+		OutEnvironment.SetDefine(TEXT("FORCE_DEPTH_TEXTURE_READS"), 1);
 	}
 };
 
@@ -329,7 +330,7 @@ class FComputeCulledObjectStartOffsetCS : public FGlobalShader
 
 	static bool ShouldCompilePermutation(const FGlobalShaderPermutationParameters& Parameters)
 	{
-		return IsFeatureLevelSupported(Parameters.Platform, ERHIFeatureLevel::SM5) && DoesPlatformSupportDistanceFieldShadowing(Parameters.Platform);
+		return DoesPlatformSupportDistanceFieldShadowing(Parameters.Platform);
 	}
 
 	static void ModifyCompilationEnvironment(const FGlobalShaderPermutationParameters& Parameters, FShaderCompilerEnvironment& OutEnvironment)
@@ -609,7 +610,6 @@ bool SupportsDistanceFieldShadows(ERHIFeatureLevel::Type FeatureLevel, EShaderPl
 {
 	return GDistanceFieldShadowing 
 		&& GetDFShadowQuality() > 0
-		&& FeatureLevel >= ERHIFeatureLevel::SM5
 		&& DoesPlatformSupportDistanceFieldShadowing(ShaderPlatform);
 }
 
@@ -621,7 +621,7 @@ bool SupportsHeightFieldShadows(ERHIFeatureLevel::Type FeatureLevel, EShaderPlat
 		&& DoesPlatformSupportDistanceFieldShadowing(ShaderPlatform);
 }
 
-bool FDeferredShadingSceneRenderer::ShouldPrepareForDistanceFieldShadows() const 
+bool FSceneRenderer::ShouldPrepareForDistanceFieldShadows() const
 {
 	bool bSceneHasRayTracedDFShadows = false;
 
@@ -652,7 +652,7 @@ bool FDeferredShadingSceneRenderer::ShouldPrepareForDistanceFieldShadows() const
 		&& SupportsDistanceFieldShadows(Scene->GetFeatureLevel(), Scene->GetShaderPlatform());
 }
 
-bool FDeferredShadingSceneRenderer::ShouldPrepareHeightFieldScene() const
+bool FSceneRenderer::ShouldPrepareHeightFieldScene() const
 {
 	return Scene
 		&& ViewFamily.EngineShowFlags.DynamicShadows
@@ -661,7 +661,7 @@ bool FDeferredShadingSceneRenderer::ShouldPrepareHeightFieldScene() const
 
 void RayTraceShadows(
 	FRDGBuilder& GraphBuilder,
-	TRDGUniformBufferRef<FSceneTextureUniformParameters> SceneTexturesUniformBuffer,
+	const FMinimalSceneTextures& SceneTextures,
 	FRDGTextureRef RayTracedShadowsTexture,
 	const FViewInfo& View,
 	const FProjectedShadowInfo* ProjectedShadowInfo,
@@ -716,7 +716,7 @@ void RayTraceShadows(
 			
 		PassParameters->RWShadowFactors = GraphBuilder.CreateUAV(RayTracedShadowsTexture);
 		PassParameters->View = View.ViewUniformBuffer;
-		PassParameters->SceneTextures = SceneTexturesUniformBuffer;
+		PassParameters->SceneTextures = SceneTextures.GetSceneTextureShaderParameters(View.GetFeatureLevel());
 
 		const FLightSceneProxy& LightProxy = *(ProjectedShadowInfo->GetLightSceneInfo().Proxy);
 		FLightShaderParameters LightParameters;
@@ -783,7 +783,7 @@ void RayTraceShadows(
 
 FRDGTextureRef FProjectedShadowInfo::BeginRenderRayTracedDistanceFieldProjection(
 	FRDGBuilder& GraphBuilder,
-	TRDGUniformBufferRef<FSceneTextureUniformParameters> SceneTexturesUniformBuffer,
+	const FMinimalSceneTextures& SceneTextures,
 	const FViewInfo& View) const
 {
 	const bool bDFShadowSupported = SupportsDistanceFieldShadows(View.GetFeatureLevel(), View.GetShaderPlatform());
@@ -853,7 +853,7 @@ FRDGTextureRef FProjectedShadowInfo::BeginRenderRayTracedDistanceFieldProjection
 				RayTracedShadowsTexture = GraphBuilder.CreateTexture(Desc, TEXT("RayTracedShadows"));
 			}
 
-			RayTraceShadows(GraphBuilder, SceneTexturesUniformBuffer, RayTracedShadowsTexture, View, this, DFPT_SignedDistanceField, false, nullptr, ObjectBufferParameters, CulledObjectBufferParameters, LightTileIntersectionParameters);
+			RayTraceShadows(GraphBuilder, SceneTextures, RayTracedShadowsTexture, View, this, DFPT_SignedDistanceField, false, nullptr, ObjectBufferParameters, CulledObjectBufferParameters, LightTileIntersectionParameters);
 		}
 	}
 
@@ -914,7 +914,7 @@ FRDGTextureRef FProjectedShadowInfo::BeginRenderRayTracedDistanceFieldProjection
 			RayTracedShadowsTexture = GraphBuilder.CreateTexture(Desc, TEXT("RayTracedShadows"));
 		}
 
-		RayTraceShadows(GraphBuilder, SceneTexturesUniformBuffer, RayTracedShadowsTexture, View, this, DFPT_HeightField, bHasPrevOutput, PrevOutputTexture, ObjectBufferParameters, CulledObjectBufferParameters, LightTileIntersectionParameters);
+		RayTraceShadows(GraphBuilder, SceneTextures, RayTracedShadowsTexture, View, this, DFPT_HeightField, bHasPrevOutput, PrevOutputTexture, ObjectBufferParameters, CulledObjectBufferParameters, LightTileIntersectionParameters);
 	}
 
 	return RayTracedShadowsTexture;
@@ -935,7 +935,7 @@ void FProjectedShadowInfo::RenderRayTracedDistanceFieldProjection(
 {
 	check(ScissorRect.Area() > 0);
 
-	FRDGTextureRef RayTracedShadowsTexture = BeginRenderRayTracedDistanceFieldProjection(GraphBuilder, SceneTextures.UniformBuffer, View);
+	FRDGTextureRef RayTracedShadowsTexture = BeginRenderRayTracedDistanceFieldProjection(GraphBuilder, SceneTextures, View);
 
 	if (RayTracedShadowsTexture)
 	{
@@ -944,7 +944,7 @@ void FProjectedShadowInfo::RenderRayTracedDistanceFieldProjection(
 		PassParameters->RenderTargets.DepthStencil = FDepthStencilBinding(SceneTextures.Depth.Target, ERenderTargetLoadAction::ELoad, ERenderTargetLoadAction::ELoad, FExclusiveDepthStencil::DepthRead_StencilRead);
 		
 		PassParameters->PS.View = View.ViewUniformBuffer;
-		PassParameters->PS.SceneTextures = SceneTextures.UniformBuffer;
+		PassParameters->PS.SceneTextures = SceneTextures.GetSceneTextureShaderParameters(View.GetFeatureLevel());
 		PassParameters->PS.ShadowFactorsTexture = RayTracedShadowsTexture;
 		PassParameters->PS.ShadowFactorsSampler = TStaticSamplerState<SF_Bilinear>::GetRHI();
 		PassParameters->PS.ScissorRectMinAndSize = FIntRect(ScissorRect.Min, ScissorRect.Size());
diff --git a/Engine/Source/Runtime/Renderer/Private/MobileBasePassRendering.cpp b/Engine/Source/Runtime/Renderer/Private/MobileBasePassRendering.cpp
--- a/Engine/Source/Runtime/Renderer/Private/MobileBasePassRendering.cpp
+++ b/Engine/Source/Runtime/Renderer/Private/MobileBasePassRendering.cpp
@@ -202,6 +202,20 @@ void SetupMobileBasePassUniformParameters(
 	}
 	BasePassParameters.AmbientOcclusionSampler = TStaticSamplerState<SF_Bilinear, AM_Clamp, AM_Clamp, AM_Clamp>::GetRHI();
 	BasePassParameters.AmbientOcclusionStaticFraction = FMath::Clamp(View.FinalPostProcessSettings.AmbientOcclusionStaticFraction, 0.0f, 1.0f);
+
+	bool bRequiresDistanceFieldShadowingPass = IsMobileDistanceFieldShadowingEnabled(View.GetShaderPlatform());
+	
+	if (bRequiresDistanceFieldShadowingPass && GScreenSpaceShadowMaskTextureMobileOutputs.ScreenSpaceShadowMaskTextureMobile.IsValid())
+	{
+		FRDGTextureRef ScreenShadowMaskTexture = GraphBuilder.RegisterExternalTexture(GScreenSpaceShadowMaskTextureMobileOutputs.ScreenSpaceShadowMaskTextureMobile, TEXT("ScreenSpaceShadowMaskTextureMobile"));
+		BasePassParameters.ScreenSpaceShadowMaskTexture = ScreenShadowMaskTexture;
+		BasePassParameters.ScreenSpaceShadowMaskSampler = TStaticSamplerState<SF_Point, AM_Clamp, AM_Clamp, AM_Clamp>::GetRHI();
+	}
+	else
+	{
+		BasePassParameters.ScreenSpaceShadowMaskTexture = GSystemTextures.GetWhiteDummy(GraphBuilder);
+		BasePassParameters.ScreenSpaceShadowMaskSampler = TStaticSamplerState<SF_Point, AM_Clamp, AM_Clamp, AM_Clamp>::GetRHI();
+	}
 }
 
 TRDGUniformBufferRef<FMobileBasePassUniformParameters> CreateMobileBasePassUniformBuffer(
@@ -242,17 +256,7 @@ void SetupMobileDirectionalLightUniformParameters(
 		if (bDynamicShadows && VisibleLightInfos.IsValidIndex(Light->Id) && VisibleLightInfos[Light->Id].AllProjectedShadows.Num() > 0)
 		{
 			const TArray<FProjectedShadowInfo*, SceneRenderingAllocator>& DirectionalLightShadowInfos = VisibleLightInfos[Light->Id].AllProjectedShadows;
-
 			static_assert(MAX_MOBILE_SHADOWCASCADES <= 4, "more than 4 cascades not supported by the shader and uniform buffer");
-			{
-				const FProjectedShadowInfo* ShadowInfo = DirectionalLightShadowInfos[0];
-				const FIntPoint ShadowBufferResolution = ShadowInfo->GetShadowBufferResolution();
-				const FVector4 ShadowBufferSizeValue((float)ShadowBufferResolution.X, (float)ShadowBufferResolution.Y, 1.0f / (float)ShadowBufferResolution.X, 1.0f / (float)ShadowBufferResolution.Y);
-
-				Params.DirectionalLightShadowTexture = ShadowInfo->RenderTargets.DepthTarget->GetRenderTargetItem().ShaderResourceTexture.GetReference();
-				Params.DirectionalLightDirectionAndShadowTransition.W = 1.0f / ShadowInfo->ComputeTransitionSize();
-				Params.DirectionalLightShadowSize = ShadowBufferSizeValue;
-			}
 
 			const int32 NumShadowsToCopy = FMath::Min(DirectionalLightShadowInfos.Num(), MAX_MOBILE_SHADOWCASCADES);
 			int32_t OutShadowIndex = 0;
@@ -260,8 +264,17 @@ void SetupMobileDirectionalLightUniformParameters(
 			{
 				const FProjectedShadowInfo* ShadowInfo = DirectionalLightShadowInfos[i];
 
-				if (ShadowInfo->ShadowDepthView)
+				if (ShadowInfo->ShadowDepthView && !ShadowInfo->bRayTracedDistanceField)
 				{
+					if (OutShadowIndex == 0)
+					{
+						const FIntPoint ShadowBufferResolution = ShadowInfo->GetShadowBufferResolution();
+						const FVector4 ShadowBufferSizeValue((float)ShadowBufferResolution.X, (float)ShadowBufferResolution.Y, 1.0f / (float)ShadowBufferResolution.X, 1.0f / (float)ShadowBufferResolution.Y);
+
+						Params.DirectionalLightShadowTexture = ShadowInfo->RenderTargets.DepthTarget->GetRenderTargetItem().ShaderResourceTexture.GetReference();
+						Params.DirectionalLightDirectionAndShadowTransition.W = 1.0f / ShadowInfo->ComputeTransitionSize();
+						Params.DirectionalLightShadowSize = ShadowBufferSizeValue;
+					}
 					Params.DirectionalLightScreenToShadow[OutShadowIndex] = ShadowInfo->GetScreenToShadowMatrix(SceneView);
 					Params.DirectionalLightShadowDistances[OutShadowIndex] = ShadowInfo->CascadeSettings.SplitFar;
 					OutShadowIndex++;
diff --git a/Engine/Source/Runtime/Renderer/Private/MobileBasePassRendering.h b/Engine/Source/Runtime/Renderer/Private/MobileBasePassRendering.h
--- a/Engine/Source/Runtime/Renderer/Private/MobileBasePassRendering.h
+++ b/Engine/Source/Runtime/Renderer/Private/MobileBasePassRendering.h
@@ -34,6 +34,8 @@ BEGIN_GLOBAL_SHADER_PARAMETER_STRUCT(FMobileBasePassUniformParameters, )
 	SHADER_PARAMETER_RDG_BUFFER_SRV(Buffer<float4>, EyeAdaptationBuffer)
 	SHADER_PARAMETER_RDG_TEXTURE(Texture2D, AmbientOcclusionTexture)
 	SHADER_PARAMETER_SAMPLER(SamplerState, AmbientOcclusionSampler)
+	SHADER_PARAMETER_RDG_TEXTURE(Texture2D, ScreenSpaceShadowMaskTexture)
+	SHADER_PARAMETER_SAMPLER(SamplerState, ScreenSpaceShadowMaskSampler)
 END_GLOBAL_SHADER_PARAMETER_STRUCT()
 
 BEGIN_SHADER_PARAMETER_STRUCT(FMobileBasePassParameters, )
@@ -437,6 +439,8 @@ public:
 		}
 
 		OutEnvironment.SetDefine(TEXT("ENABLE_AMBIENT_OCCLUSION"), IsMobileAmbientOcclusionEnabled(Parameters.Platform) ? 1u : 0u);
+
+		OutEnvironment.SetDefine(TEXT("ENABLE_DISTANCE_FIELD"), IsMobileDistanceFieldEnabled(Parameters.Platform));
 	}
 	
 	/** Initialization constructor. */
diff --git a/Engine/Source/Runtime/Renderer/Private/MobileShadingRenderer.cpp b/Engine/Source/Runtime/Renderer/Private/MobileShadingRenderer.cpp
--- a/Engine/Source/Runtime/Renderer/Private/MobileShadingRenderer.cpp
+++ b/Engine/Source/Runtime/Renderer/Private/MobileShadingRenderer.cpp
@@ -203,6 +203,7 @@ FMobileSceneRenderer::FMobileSceneRenderer(const FSceneViewFamily* InViewFamily,
 	bShouldRenderCustomDepth = false;
 	bRequiresPixelProjectedPlanarRelfectionPass = false;
 	bRequriesAmbientOcclusionPass = false;
+	bRequiresDistanceFieldShadowingPass = false;
 
 	// Don't do occlusion queries when doing scene captures
 	for (FViewInfo& View : Views)
@@ -381,6 +382,17 @@ void FMobileSceneRenderer::InitViews(FRDGBuilder& GraphBuilder, FSceneTexturesCo
 		&& !bDeferredShading;
 
 	bShouldRenderVelocities = ShouldRenderVelocities();
+	bRequiresDistanceField = IsMobileDistanceFieldEnabled(ShaderPlatform)
+		&& ViewFamily.EngineShowFlags.Lighting
+		&& !Views[0].bIsReflectionCapture
+		&& !Views[0].bIsPlanarReflection
+		&& !ViewFamily.EngineShowFlags.HitProxies
+		&& !ViewFamily.EngineShowFlags.VisualizeLightCulling
+		&& !ViewFamily.UseDebugViewPS()
+		&& !bDeferredShading;
+
+	bRequiresDistanceFieldShadowingPass = bRequiresDistanceField && IsMobileDistanceFieldShadowingEnabled(ShaderPlatform);
+		
 
 	// Whether we need to store depth for post-processing
 	// On PowerVR we see flickering of shadows and depths not updating correctly if targets are discarded.
@@ -393,6 +405,7 @@ void FMobileSceneRenderer::InitViews(FRDGBuilder& GraphBuilder, FSceneTexturesCo
 		bRequiresMultiPass ||
 		bForceDepthResolve ||
 		bRequriesAmbientOcclusionPass ||
+		bRequiresDistanceFieldShadowingPass ||
 		bRequiresPixelProjectedPlanarRelfectionPass ||
 		bSeparateTranslucencyActive ||
 		Views[0].bIsReflectionCapture ||
@@ -425,6 +438,15 @@ void FMobileSceneRenderer::InitViews(FRDGBuilder& GraphBuilder, FSceneTexturesCo
 		ReleasePixelProjectedReflectionOutputs();
 	}
 
+	if (bRequiresDistanceFieldShadowingPass)
+	{
+		InitMobileSDFShadowingOutputs(RHICmdList, SceneTexturesConfig.Extent);
+	}
+	else
+	{
+		ReleaseMobileSDFShadowingOutputs();
+	}
+
 	// Find out whether custom depth pass should be rendered.
 	{
 		bool bCouldUseCustomDepthStencil = !bGammaSpace && (!Scene->World || (Scene->World->WorldType != EWorldType::EditorPreview && Scene->World->WorldType != EWorldType::Inactive));
@@ -496,6 +518,11 @@ void FMobileSceneRenderer::InitViews(FRDGBuilder& GraphBuilder, FSceneTexturesCo
 		Scene->GPUScene.UploadDynamicPrimitiveShaderDataForView(RHICmdList, Scene, Views[ViewIndex]);
 	}
 
+	if (bRequiresDistanceField)
+	{
+		PrepareDistanceFieldScene(GraphBuilder, false);
+	}
+
 	{
 		// GPUCULL_TODO: Possibly fold into unpack step
 		InstanceCullingManager.CullInstances(GraphBuilder, Scene->GPUScene);
@@ -730,6 +757,12 @@ void FMobileSceneRenderer::Render(FRDGBuilder& GraphBuilder)
 		});
 	}
 
+	if (bRequiresDistanceFieldShadowingPass)
+	{
+		CSV_SCOPED_TIMING_STAT_EXCLUSIVE(RenderSDFShadowing);
+		RenderMobileSDFShadowing(GraphBuilder, SceneTextures.Depth.Resolve, Scene, Views, VisibleLightInfos);
+	}
+
 	if (bRequriesAmbientOcclusionPass)
 	{
 		RenderAmbientOcclusion(GraphBuilder, SceneTextures.Depth.Resolve, SceneTextures.ScreenSpaceAO);
diff --git a/Engine/Source/Runtime/Renderer/Private/PostProcess/PostProcessAmbientOcclusionMobile.cpp b/Engine/Source/Runtime/Renderer/Private/PostProcess/PostProcessAmbientOcclusionMobile.cpp
--- a/Engine/Source/Runtime/Renderer/Private/PostProcess/PostProcessAmbientOcclusionMobile.cpp
+++ b/Engine/Source/Runtime/Renderer/Private/PostProcess/PostProcessAmbientOcclusionMobile.cpp
@@ -351,6 +351,7 @@ public:
 		Super::ModifyCompilationEnvironment(Parameters, OutEnvironment);
 
 		OutEnvironment.SetDefine(TEXT("HORIZONSEARCH_INTEGRAL_PIXEL_SHADER"), 1u);
+		OutEnvironment.SetDefine(TEXT("FORCE_DEPTH_TEXTURE_READS"), 1);
 	}
 
 	static FPermutationDomain BuildPermutationVector(int32 LUTTextureType, int32 ShaderQuality)
@@ -384,6 +385,7 @@ public:
 	{
 		Super::ModifyCompilationEnvironment(Parameters, OutEnvironment);
 		OutEnvironment.SetDefine(TEXT("SPATIALFILTER_PIXEL_SHADER"), 1u);
+		OutEnvironment.SetDefine(TEXT("FORCE_DEPTH_TEXTURE_READS"), 1);
 	}
 };
 
diff --git a/Engine/Source/Runtime/Renderer/Private/SceneRendering.h b/Engine/Source/Runtime/Renderer/Private/SceneRendering.h
--- a/Engine/Source/Runtime/Renderer/Private/SceneRendering.h
+++ b/Engine/Source/Runtime/Renderer/Private/SceneRendering.h
@@ -1864,8 +1864,23 @@ public:
 	/** Cache the FXSystem value from the Scene. Must be ran on the renderthread to ensure it is valid throughout rendering. */
 	void InitFXSystem();
 
+	/** Whether distance field global data structures should be prepared for features that use it. */
+	bool ShouldPrepareForDistanceFieldShadows() const;
+	bool ShouldPrepareForDistanceFieldAO() const;
+	bool ShouldPrepareForDFInsetIndirectShadow() const;
 	FGPUSceneDynamicContext& GetGPUSceneDynamicContext() { return GPUSceneDynamicContext; }
 
+	bool ShouldPrepareDistanceFieldScene() const;
+	bool ShouldPrepareGlobalDistanceField() const;
+	bool ShouldPrepareHeightFieldScene() const;
+
+	void UpdateGlobalDistanceFieldObjectBuffers(FRDGBuilder& GraphBuilder);
+	void UpdateGlobalHeightFieldObjectBuffers(FRDGBuilder& GraphBuilder);
+	void AddOrRemoveSceneHeightFieldPrimitives(bool bSkipAdd = false);
+	void PrepareDistanceFieldScene(FRDGBuilder& GraphBuilder, bool bSplitDispatch);
+
+	virtual bool IsLumenEnabled(const FViewInfo& View) const { return false; }
+	virtual bool AnyViewHasSupportingGIMethod() const { return false; }
 protected:
 
 	/** Size of the family. */
@@ -2087,6 +2102,23 @@ private:
 	bool bShadowDepthRenderCompleted;
 };
 
+struct FForwardScreenSpaceShadowMaskTextureMobileOutputs
+{
+	TRefCountPtr<IPooledRenderTarget> ScreenSpaceShadowMaskTextureMobile;
+
+	bool IsValid()
+	{
+		return ScreenSpaceShadowMaskTextureMobile.IsValid();
+	}
+
+	void Release()
+	{
+		ScreenSpaceShadowMaskTextureMobile.SafeRelease();
+	}
+};
+
+extern FForwardScreenSpaceShadowMaskTextureMobileOutputs GScreenSpaceShadowMaskTextureMobileOutputs;
+
 /**
  * Renderer that implements simple forward shading and associated features.
  */
@@ -2184,6 +2216,8 @@ private:
 	bool bRequiresPixelProjectedPlanarRelfectionPass;
 	bool bRequriesAmbientOcclusionPass;
 	bool bShouldRenderVelocities;
+	bool bRequiresDistanceField;
+	bool bRequiresDistanceFieldShadowingPass;
 	static FGlobalDynamicIndexBuffer DynamicIndexBuffer;
 	static FGlobalDynamicVertexBuffer DynamicVertexBuffer;
 	static TGlobalResource<FGlobalDynamicReadBuffer> DynamicReadBuffer;
diff --git a/Engine/Source/Runtime/Renderer/Private/SceneTextures.cpp b/Engine/Source/Runtime/Renderer/Private/SceneTextures.cpp
--- a/Engine/Source/Runtime/Renderer/Private/SceneTextures.cpp
+++ b/Engine/Source/Runtime/Renderer/Private/SceneTextures.cpp
@@ -740,6 +740,20 @@ FSceneTextures& FMinimalSceneTextures::Create(FRDGBuilder& GraphBuilder, const F
 	return SceneTextures;
 }
 
+FSceneTextureShaderParameters FMinimalSceneTextures::GetSceneTextureShaderParameters(ERHIFeatureLevel::Type FeatureLevel) const
+{
+	FSceneTextureShaderParameters OutSceneTextureShaderParameters;
+	if (FeatureLevel >= ERHIFeatureLevel::SM5)
+	{
+		OutSceneTextureShaderParameters.SceneTextures = UniformBuffer;
+	}
+	else
+	{
+		OutSceneTextureShaderParameters.MobileSceneTextures = MobileUniformBuffer;
+	}
+	return OutSceneTextureShaderParameters;
+}
+
 FSceneTextures& FSceneTextures::Create(FRDGBuilder& GraphBuilder, const FSceneTexturesConfig& Config)
 {
 	FSceneTextures& SceneTextures = FMinimalSceneTextures::Create(GraphBuilder, Config);
diff --git a/Engine/Source/Runtime/Renderer/Private/SceneTextures.h b/Engine/Source/Runtime/Renderer/Private/SceneTextures.h
--- a/Engine/Source/Runtime/Renderer/Private/SceneTextures.h
+++ b/Engine/Source/Runtime/Renderer/Private/SceneTextures.h
@@ -126,6 +126,8 @@ struct RENDERER_API FMinimalSceneTextures
 	// Textures containing depth / stencil information from the custom depth pass.
 	FCustomDepthTextures CustomDepth{};
 
+	FSceneTextureShaderParameters GetSceneTextureShaderParameters(ERHIFeatureLevel::Type FeatureLevel) const;
+
 protected:
 	FMinimalSceneTextures(const FSceneTexturesConfig& InConfig)
 		: Config(InConfig)
diff --git a/Engine/Source/Runtime/Renderer/Private/ShadowRendering.cpp b/Engine/Source/Runtime/Renderer/Private/ShadowRendering.cpp
--- a/Engine/Source/Runtime/Renderer/Private/ShadowRendering.cpp
+++ b/Engine/Source/Runtime/Renderer/Private/ShadowRendering.cpp
@@ -188,6 +188,8 @@ static TAutoConsoleVariable<float> CVarShadowMaxSlopeScaleDepthBias(
 	TEXT("Higher values give better self-shadowing, but increase self-shadowing artifacts"),
 	ECVF_RenderThreadSafe);
 
+FForwardScreenSpaceShadowMaskTextureMobileOutputs GScreenSpaceShadowMaskTextureMobileOutputs;
+
 ///////////////////////////////////////////////////////////////////////////////////////////////////
 // Hair
 static TAutoConsoleVariable<int32> CVarHairStrandsCullPerObjectShadowCaster(
@@ -2212,6 +2214,90 @@ void FMobileSceneRenderer::RenderModulatedShadowProjections(FRDGBuilder& GraphBu
 	}
 }
 
+void InitMobileSDFShadowingOutputs(FRHICommandListImmediate& RHICmdList, const FIntPoint& Extent)
+{
+	const FIntPoint& BufferSize = Extent;
+
+	if (!GScreenSpaceShadowMaskTextureMobileOutputs.IsValid() || GScreenSpaceShadowMaskTextureMobileOutputs.ScreenSpaceShadowMaskTextureMobile->GetDesc().Extent != BufferSize)
+	{
+		GScreenSpaceShadowMaskTextureMobileOutputs.ScreenSpaceShadowMaskTextureMobile.SafeRelease();
+		GRenderTargetPool.FindFreeElement(RHICmdList, FPooledRenderTargetDesc::Create2DDesc(BufferSize, PF_B8G8R8A8, FClearValueBinding::White, TexCreate_None, TexCreate_RenderTargetable | TexCreate_ShaderResource, false, 1, false), GScreenSpaceShadowMaskTextureMobileOutputs.ScreenSpaceShadowMaskTextureMobile, TEXT("ForwardScreenSpaceShadowMaskTextureTexture"));
+	}
+}
+
+void RenderMobileSDFShadowing(
+	FRDGBuilder& GraphBuilder, 
+	FRDGTextureRef SceneDepthTexture, 
+	const FScene* Scene,
+	const TArrayView<const FViewInfo> Views,
+	const TArrayView<const FVisibleLightInfo> VisibleLightInfos)
+{
+	FRDGTextureRef ScreenShadowMaskTexture = GraphBuilder.RegisterExternalTexture(GScreenSpaceShadowMaskTextureMobileOutputs.ScreenSpaceShadowMaskTextureMobile, TEXT("ScreenSpaceShadowMaskTextureMobile"));
+	AddClearRenderTargetPass(GraphBuilder, ScreenShadowMaskTexture);
+
+	const FMinimalSceneTextures& SceneTextures = FSceneTextures::Get(GraphBuilder);
+	for (TSparseArray<FLightSceneInfoCompact>::TConstIterator LightIt(Scene->Lights); LightIt; ++LightIt)
+	{
+		const FLightSceneInfoCompact& LightSceneInfoCompact = *LightIt;
+		const FLightSceneInfo* LightSceneInfo = LightSceneInfoCompact.LightSceneInfo;
+
+		const FVisibleLightInfo& VisibleLightInfo = VisibleLightInfos[LightSceneInfo->Id];
+		const FLightSceneProxy* LightSceneProxy = LightSceneInfo->Proxy;
+
+		// Allocate arrays using the graph allocator so we can safely reference them in passes.
+		using FProjectedShadowInfoArray = TArray<const FProjectedShadowInfo*, SceneRenderingAllocator>;
+		auto& DistanceFieldShadows = *GraphBuilder.AllocObject<FProjectedShadowInfoArray>();
+
+		for (int32 ShadowIndex = 0; ShadowIndex < VisibleLightInfo.ShadowsToProject.Num(); ShadowIndex++)
+		{
+			const FProjectedShadowInfo* ProjectedShadowInfo = VisibleLightInfo.ShadowsToProject[ShadowIndex];
+			if (ProjectedShadowInfo->bRayTracedDistanceField)
+			{
+				DistanceFieldShadows.Add(ProjectedShadowInfo);
+			}
+		}
+
+		if (DistanceFieldShadows.Num() > 0)
+		{
+			RDG_EVENT_SCOPE(GraphBuilder, "DistanceFieldShadows");
+
+			// Distance field shadows need to be renderer last as they blend over far shadow cascades.
+			for (int32 ViewIndex = 0; ViewIndex < Views.Num(); ViewIndex++)
+			{
+				const FViewInfo& View = Views[ViewIndex];
+				RDG_GPU_MASK_SCOPE(GraphBuilder, View.GPUMask);
+				RDG_EVENT_SCOPE_CONDITIONAL(GraphBuilder, Views.Num() > 1, "View%d", ViewIndex);
+
+				FIntRect ScissorRect;
+				if (!LightSceneProxy->GetScissorRect(ScissorRect, View, View.ViewRect))
+				{
+					ScissorRect = View.ViewRect;
+				}
+
+				if (ScissorRect.Area() > 0)
+				{
+					for (int32 ShadowIndex = 0; ShadowIndex < DistanceFieldShadows.Num(); ShadowIndex++)
+					{
+						const FProjectedShadowInfo* ProjectedShadowInfo = DistanceFieldShadows[ShadowIndex];
+						ProjectedShadowInfo->RenderRayTracedDistanceFieldProjection(
+							GraphBuilder,
+							SceneTextures,
+							ScreenShadowMaskTexture,
+							View,
+							ScissorRect,
+							false);
+					}
+				}
+			}
+		}
+	}
+}
+
+void ReleaseMobileSDFShadowingOutputs()
+{
+	GScreenSpaceShadowMaskTextureMobileOutputs.Release();
+}
+
 IMPLEMENT_GLOBAL_SHADER_PARAMETER_STRUCT(FTranslucentSelfShadowUniformParameters, "TranslucentSelfShadow");
 
 void SetupTranslucentSelfShadowUniformParameters(const FProjectedShadowInfo* ShadowInfo, FTranslucentSelfShadowUniformParameters& OutParameters)
diff --git a/Engine/Source/Runtime/Renderer/Private/ShadowRendering.h b/Engine/Source/Runtime/Renderer/Private/ShadowRendering.h
--- a/Engine/Source/Runtime/Renderer/Private/ShadowRendering.h
+++ b/Engine/Source/Runtime/Renderer/Private/ShadowRendering.h
@@ -58,6 +58,14 @@ void OverrideWithDefaultMaterialForShadowDepth(
 	ERHIFeatureLevel::Type InFeatureLevel
 	);
 
+void InitMobileSDFShadowingOutputs(FRHICommandListImmediate& RHICmdList, const FIntPoint& Extent);
+void RenderMobileSDFShadowing(FRDGBuilder& GraphBuilder, 
+	FRDGTextureRef SceneDepthTexture, 
+	const FScene* Scene,
+	const TArrayView<const FViewInfo> Views,
+	const TArrayView<const FVisibleLightInfo> VisibleLightInfos);
+void ReleaseMobileSDFShadowingOutputs();
+
 enum EShadowDepthRenderMode
 {
 	/** The render mode used by regular shadows */
@@ -541,7 +549,7 @@ public:
 
 	FRDGTextureRef BeginRenderRayTracedDistanceFieldProjection(
 		FRDGBuilder& GraphBuilder,
-		TRDGUniformBufferRef<FSceneTextureUniformParameters> SceneTexturesUniformBuffer,
+		const FMinimalSceneTextures& SceneTextures,
 		const FViewInfo& View) const;
 
 	/** Renders ray traced distance field shadows. */
diff --git a/Engine/Source/Runtime/Renderer/Private/ShadowSetup.cpp b/Engine/Source/Runtime/Renderer/Private/ShadowSetup.cpp
--- a/Engine/Source/Runtime/Renderer/Private/ShadowSetup.cpp
+++ b/Engine/Source/Runtime/Renderer/Private/ShadowSetup.cpp
@@ -4574,9 +4574,7 @@ void FSceneRenderer::AllocateShadowDepthTargets(FRHICommandListImmediate& RHICmd
 					INC_DWORD_STAT(STAT_PerObjectShadows);
 				}
 
-				bool bNeedsProjection = ProjectedShadowInfo->CacheMode != SDCM_StaticPrimitivesOnly
-					// Mobile rendering only projects opaque per object shadows.
-					&& (FeatureLevel >= ERHIFeatureLevel::SM5 || ProjectedShadowInfo->bPerObjectOpaqueShadow);
+				bool bNeedsProjection = ProjectedShadowInfo->CacheMode != SDCM_StaticPrimitivesOnly;
 
 				if (bNeedsProjection)
 				{if (ProjectedShadowInfo->bCapsuleShadow)
