Add a CVar r.DecalDepthBias to allow global control over decal depth bias.

diff --git a/Engine/Shaders/Private/MeshDecals.usf b/Engine/Shaders/Private/MeshDecals.usf
--- a/Engine/Shaders/Private/MeshDecals.usf
+++ b/Engine/Shaders/Private/MeshDecals.usf
@@ -140,10 +140,10 @@ void MainVS(
 
 #endif // USING_TESSELLATION
 
-// Move all geometry a little bit towards the camera (to not get z fighting with existing zbuffer of the same meshes rendered with slightly different float computations).
-#if !USING_TESSELLATION
-	// todo: tweak or expose (cvar or per mesh/material setting), if needed the artist always can create a WorldPositionOffset adjustment in the shader (e.g. -0.1*CameraDirectionVector )
-	Output.Position.z *= 1.001;
+#if !USING_TESSELLATION && HAS_INVERTED_Z_BUFFER
+	// Move all geometry a small amount towards the camera to avoid z-fighting.
+	// Assuming a reverse z buffer this pushes the near plane a fixed small distance.
+	Output.Position.z += ResolvedView.DecalDepthBias;
 #endif
 
 	OutputVertexID( Output );
diff --git a/Engine/Source/Runtime/Engine/Public/SceneView.h b/Engine/Source/Runtime/Engine/Public/SceneView.h
--- a/Engine/Source/Runtime/Engine/Public/SceneView.h
+++ b/Engine/Source/Runtime/Engine/Public/SceneView.h
@@ -680,6 +680,7 @@ enum ETranslucencyVolumeCascade
 	VIEW_UNIFORM_BUFFER_MEMBER(float, GeneralPurposeTweak) \
 	VIEW_UNIFORM_BUFFER_MEMBER(float, GeneralPurposeTweak2) \
 	VIEW_UNIFORM_BUFFER_MEMBER_EX(float, DemosaicVposOffset, EShaderPrecisionModifier::Half) \
+	VIEW_UNIFORM_BUFFER_MEMBER(float, DecalDepthBias) \
 	VIEW_UNIFORM_BUFFER_MEMBER(FVector, IndirectLightingColorScale) \
 	VIEW_UNIFORM_BUFFER_MEMBER(FVector, PrecomputedIndirectLightingColorScale) \
 	VIEW_UNIFORM_BUFFER_MEMBER_EX(float, AtmosphericFogSunPower, EShaderPrecisionModifier::Half) \
diff --git a/Engine/Source/Runtime/Renderer/Private/SceneRendering.cpp b/Engine/Source/Runtime/Renderer/Private/SceneRendering.cpp
--- a/Engine/Source/Runtime/Renderer/Private/SceneRendering.cpp
+++ b/Engine/Source/Runtime/Renderer/Private/SceneRendering.cpp
@@ -134,6 +134,12 @@ static TAutoConsoleVariable<float> CVarDemosaicVposOffset(
 	TEXT("This offset is added to the rasterized position used for demosaic in the mobile tonemapping shader. It exists to workaround driver bugs on some Android devices that have a half-pixel offset."),
 	ECVF_RenderThreadSafe);
 
+static TAutoConsoleVariable<float> CVarDecalDepthBias(
+	TEXT("r.DecalDepthBias"),
+	0.005f,
+	TEXT("Global depth bias used by mesh decals. Default is 0.005"),
+	ECVF_RenderThreadSafe);
+
 static TAutoConsoleVariable<int32> CVarRefractionQuality(
 	TEXT("r.RefractionQuality"),
 	2,
@@ -1601,6 +1607,8 @@ void FViewInfo::SetupUniformBufferParameters(
 		ViewUniformShaderParameters.DemosaicVposOffset = CVarDemosaicVposOffset.GetValueOnRenderThread();
 	}
 
+	ViewUniformShaderParameters.DecalDepthBias = CVarDecalDepthBias.GetValueOnRenderThread();
+
 	ViewUniformShaderParameters.IndirectLightingColorScale = FVector(FinalPostProcessSettings.IndirectLightingColor.R * FinalPostProcessSettings.IndirectLightingIntensity,
 		FinalPostProcessSettings.IndirectLightingColor.G * FinalPostProcessSettings.IndirectLightingIntensity,
 		FinalPostProcessSettings.IndirectLightingColor.B * FinalPostProcessSettings.IndirectLightingIntensity);
