Remove STRATA shading model from subsurface list.

diff --git a/Engine/Shaders/Private/DeferredLightingCommon.ush b/Engine/Shaders/Private/DeferredLightingCommon.ush
--- a/Engine/Shaders/Private/DeferredLightingCommon.ush
+++ b/Engine/Shaders/Private/DeferredLightingCommon.ush
@@ -229,6 +229,7 @@ void GetShadowTerms(float SceneDepth, float4 PrecomputedShadowFactors, uint Shad
 		{
 			float ContactShadowOcclusion = bHitCastContactShadow ? 1.0 : LightData.ContactShadowNonShadowCastingIntensity;
 
+			// STRATA_TODO: manage SSS transmission path here. IsSubsurfaceModel does not take into account strata
 			// Exponential attenuation is not applied on hair/eye/SSS-profile here, as the hit distance (shading-point to blocker) is different from the estimated 
 			// thickness (closest-point-from-light to shading-point), and this creates light leaks. Instead we consider first hit as a blocker (old behavior)
 			BRANCH
diff --git a/Engine/Shaders/Private/DeferredShadingCommon.ush b/Engine/Shaders/Private/DeferredShadingCommon.ush
--- a/Engine/Shaders/Private/DeferredShadingCommon.ush
+++ b/Engine/Shaders/Private/DeferredShadingCommon.ush
@@ -349,13 +349,12 @@ bool IsSubsurfaceModel(int ShadingModel)
 		|| ShadingModel == SHADINGMODELID_SUBSURFACE_PROFILE
 		|| ShadingModel == SHADINGMODELID_TWOSIDED_FOLIAGE
 		|| ShadingModel == SHADINGMODELID_HAIR
-		|| ShadingModel == SHADINGMODELID_EYE 
-		|| ShadingModel == SHADINGMODELID_STRATA;
+		|| ShadingModel == SHADINGMODELID_EYE;
 }
 
 bool UseSubsurfaceProfile(int ShadingModel)
 {
-	return ShadingModel == SHADINGMODELID_SUBSURFACE_PROFILE || ShadingModel == SHADINGMODELID_EYE || ShadingModel == SHADINGMODELID_STRATA;
+	return ShadingModel == SHADINGMODELID_SUBSURFACE_PROFILE || ShadingModel == SHADINGMODELID_EYE;
 }
 
 bool HasCustomGBufferData(int ShadingModelID)
diff --git a/Engine/Shaders/Private/DistanceFieldShadowing.usf b/Engine/Shaders/Private/DistanceFieldShadowing.usf
--- a/Engine/Shaders/Private/DistanceFieldShadowing.usf
+++ b/Engine/Shaders/Private/DistanceFieldShadowing.usf
@@ -706,14 +706,23 @@ void DistanceFieldShadowingCS(
 			bool bUseSubsurfaceTransmission = false;
 
 #if !FORWARD_SHADING && DF_SHADOW_QUALITY > 1 && !SHADING_PATH_MOBILE
+		#if STRATA_ENABLED
+			const FStrataSubsurfaceData SSSData = StrataUnpackSSSData(Strata.SSSTexture.Load(uint3(ScreenPosition, 0)));
+			BRANCH
+			if (SSSData.bIsValid)
+			{
+				SubsurfaceDensity = SSSData.ProfileRadiusScale;
+				bUseSubsurfaceTransmission = true;
+			}
+		#else
 			FGBufferData GBufferData = GetGBufferData(ScreenUV);
-
 			BRANCH
 			if (IsSubsurfaceModel(GBufferData.ShadingModelID))
 			{
 				SubsurfaceDensity = SubsurfaceDensityFromOpacity(GBufferData.CustomData.a);
 				bUseSubsurfaceTransmission = true;
 			}
+		#endif
 #endif
 
 			Result = ShadowRayTraceThroughCulledObjects(
diff --git a/Engine/Source/Runtime/Renderer/Private/DistanceFieldShadowing.cpp b/Engine/Source/Runtime/Renderer/Private/DistanceFieldShadowing.cpp
--- a/Engine/Source/Runtime/Renderer/Private/DistanceFieldShadowing.cpp
+++ b/Engine/Source/Runtime/Renderer/Private/DistanceFieldShadowing.cpp
@@ -27,6 +27,7 @@
 #include "DistanceFieldAmbientOcclusion.h"
 #include "PipelineStateCache.h"
 #include "ClearQuad.h"
+#include "Strata/Strata.h"
 
 int32 GDistanceFieldShadowing = 1;
 FAutoConsoleVariableRef CVarDistanceFieldShadowing(
@@ -137,6 +138,7 @@ class FCullObjectsForShadowCS : public FGlobalShader
 
 	BEGIN_SHADER_PARAMETER_STRUCT(FParameters, )
 		SHADER_PARAMETER_STRUCT_REF(FViewUniformShaderParameters, View)
+		SHADER_PARAMETER_STRUCT_REF(FStrataGlobalUniformParameters, Strata)
 		SHADER_PARAMETER_STRUCT_INCLUDE(FDistanceFieldObjectBufferParameters, ObjectBufferParameters)
 		SHADER_PARAMETER_STRUCT_INCLUDE(FDistanceFieldCulledObjectBufferParameters, CulledObjectBufferParameters)
 		SHADER_PARAMETER(uint32, ObjectBoundingGeometryIndexCount)
@@ -158,6 +160,7 @@ class FCullObjectsForShadowCS : public FGlobalShader
 	{
 		FGlobalShader::ModifyCompilationEnvironment(Parameters,OutEnvironment);
 		OutEnvironment.SetDefine(TEXT("UPDATEOBJECTS_THREADGROUP_SIZE"), UpdateObjectsGroupSize);
+		OutEnvironment.SetDefine(TEXT("STRATA_ENABLED"), Strata::IsStrataEnabled() ? 1 : 0);
 	}
 };
 
