Implements Gen5TAA's AntiInterference heuristic

diff --git a/Engine/Shaders/Private/TemporalAA/TAACommon.ush b/Engine/Shaders/Private/TemporalAA/TAACommon.ush
--- a/Engine/Shaders/Private/TemporalAA/TAACommon.ush
+++ b/Engine/Shaders/Private/TemporalAA/TAACommon.ush
@@ -40,6 +40,23 @@
 
 #define MAX_PARALLAX_FACTOR 8.0
 
+// Size and payload size of the subpixel grid in the history
+#define SUB_PIXEL_GRID_SIZE 2
+#define SUB_PIXEL_COUNT (SUB_PIXEL_GRID_SIZE * SUB_PIXEL_GRID_SIZE)
+
+#define SUB_PIXEL_BIT_COUNT 8
+#define SUB_PIXEL_BIT_MASK ((1 << SUB_PIXEL_BIT_COUNT) - 1)
+
+#define SUB_PIXEL_PARALLAX_FACTOR_BIT_COUNT 4
+#define SUB_PIXEL_PARALLAX_FACTOR_BIT_OFFSET 0
+#define SUB_PIXEL_PARALLAX_FACTOR_BIT_MASK ((1 << SUB_PIXEL_PARALLAX_FACTOR_BIT_COUNT) - 1)
+
+#define SUB_PIXEL_INTERFERENCE_BIT_COUNT 4
+#define SUB_PIXEL_INTERFERENCE_BIT_OFFSET 4
+#define SUB_PIXEL_INTERFERENCE_BIT_MASK ((1 << SUB_PIXEL_INTERFERENCE_BIT_COUNT) - 1)
+
+
+
 // K = Center of the nearest input pixel.
 // O = Center of the output pixel.
 //
@@ -111,6 +128,9 @@ static const uint kPlusIndexes3x3[5] = { 4, 1, 3, 5, 7 };
 	#define taa_half4 float4
 #endif
 
+#define taa_subpixel_details uint
+#define taa_subpixel_payload uint // TODO: uint16_t
+
 
 //------------------------------------------------------- PARAMETERS
 
@@ -191,6 +211,9 @@ float2 PrevHistoryInfo_UVViewportSizeInverse;
 float2 PrevHistoryInfo_UVViewportBilinearMin;
 float2 PrevHistoryInfo_UVViewportBilinearMax;
 
+uint bCameraCut;
+uint bEnableInterferenceHeuristic;
+
 #if DEBUG_OUTPUT
 	RWTexture2D<float4> DebugOutput;
 #endif
@@ -364,19 +387,16 @@ taa_half MeasureRejectionFactor(taa_half3 PrevYCoCg, taa_half3 ClampedPrevYCoCg,
 	return Rejection;
 }
 
-#define BITS_PER_SUB_PIXEL 4
-#define SUB_PIXEL_MASK ((1 << BITS_PER_SUB_PIXEL) - 1)
-
-#define SUB_PIXEL_GRID_SIZE 2
-#define SUB_PIXEL_COUNT (SUB_PIXEL_GRID_SIZE * SUB_PIXEL_GRID_SIZE)
-
 struct FSubpixelNeighborhood
 {
 	// 2x2 quad of subpixels grid
-	uint4 SampleArray;
+	taa_subpixel_details SampleArray[4];
 
 	// Offset of subpixel
 	uint2 SubpixelOffset;
+
+	// Offset of the closest subpixel
+	bool2 ClosestSubpixelOffset;
 };
 
 FSubpixelNeighborhood GatherPrevSubpixelNeighborhood(Texture2D<uint> PrevSubPixelDetailTexture, float2 PrevHistoryBufferUV)
@@ -410,11 +430,12 @@ FSubpixelNeighborhood GatherPrevSubpixelNeighborhood(Texture2D<uint> PrevSubPixe
 
 	taa_half2 Interp = taa_half2(PixelCoord - TopLeftPixelCoord);
 	Neighborhood.SubpixelOffset = uint2(round(Interp * 2.0));
+	Neighborhood.ClosestSubpixelOffset = Interp > 0.5;
 
 	return Neighborhood;
 }
 
-uint GetSubpixelPayload(FSubpixelNeighborhood Neighborhood, uint SubpixelId)
+taa_subpixel_payload GetSubpixelPayload(FSubpixelNeighborhood Neighborhood, uint SubpixelId)
 {
 	const uint2 SubpixelCoord = uint2(SubpixelId % SUB_PIXEL_GRID_SIZE, SubpixelId / SUB_PIXEL_GRID_SIZE);
 
@@ -424,23 +445,59 @@ uint GetSubpixelPayload(FSubpixelNeighborhood Neighborhood, uint SubpixelId)
 
 
 	uint QuadSampleId = dot(InputTexelCoord, uint2(1, 2));
-	uint InputSample = Neighborhood.SampleArray[QuadSampleId];
+	taa_subpixel_details InputSample = Neighborhood.SampleArray[QuadSampleId];
 
-	uint BitShift = BITS_PER_SUB_PIXEL * dot(QuadSubpixelCoord, uint2(1, SUB_PIXEL_GRID_SIZE));
+	uint BitShift = SUB_PIXEL_BIT_COUNT * dot(QuadSubpixelCoord, uint2(1, SUB_PIXEL_GRID_SIZE));
 
-	uint SubpixelPayload = (InputSample >> BitShift) & SUB_PIXEL_MASK;
+	taa_subpixel_payload SubpixelPayload = taa_subpixel_payload(InputSample >> BitShift) & SUB_PIXEL_BIT_MASK;
 
 	return SubpixelPayload;
 }
 
-uint CompressParallaxFactor(taa_half ParallaxFactor)
+taa_subpixel_payload GetClosestSubpixelPayload(FSubpixelNeighborhood Neighborhood)
+{
+	uint QuadSampleId = dot(uint2(Neighborhood.ClosestSubpixelOffset), uint2(1, 2));
+	taa_subpixel_details InputSample = Neighborhood.SampleArray[QuadSampleId];
+
+	uint BitShift = SUB_PIXEL_BIT_COUNT * (3 - QuadSampleId);
+
+	taa_subpixel_payload SubpixelPayload = taa_subpixel_payload(InputSample >> BitShift) & SUB_PIXEL_BIT_MASK;
+
+	return SubpixelPayload;
+}
+
+taa_subpixel_payload CompressParallaxFactor(taa_half ParallaxFactor)
+{
+	return clamp(
+		taa_subpixel_payload((ParallaxFactor - taa_half(1.0)) * taa_half(float(SUB_PIXEL_PARALLAX_FACTOR_BIT_MASK) / MAX_PARALLAX_FACTOR)),
+		taa_subpixel_payload(0), taa_subpixel_payload(SUB_PIXEL_PARALLAX_FACTOR_BIT_MASK));
+}
+
+taa_half UncompressParallaxFactor(taa_subpixel_payload CompressedParallaxFactor)
+{
+	return 
+		taa_half(CompressedParallaxFactor & SUB_PIXEL_PARALLAX_FACTOR_BIT_MASK) *
+		taa_half(MAX_PARALLAX_FACTOR / float(SUB_PIXEL_PARALLAX_FACTOR_BIT_MASK)) +
+		taa_half(1.0);
+}
+
+taa_half LumaToInterferenceSeed(taa_half Luma)
+{
+	return Luma * HdrWeightY(Luma);
+}
+
+taa_subpixel_payload CompressInteferenceSeed(taa_half InteferenceSeed)
 {
-	return clamp(uint((ParallaxFactor - taa_half(1.0)) * taa_half(float(SUB_PIXEL_MASK) / MAX_PARALLAX_FACTOR)), 0, SUB_PIXEL_MASK);
+	return clamp(
+		taa_subpixel_payload(InteferenceSeed * taa_half(SUB_PIXEL_INTERFERENCE_BIT_MASK)),
+		taa_subpixel_payload(0), taa_subpixel_payload(SUB_PIXEL_INTERFERENCE_BIT_MASK));
 }
 
-taa_half UncompressParallaxFactor(uint CompressedParallaxFactor)
+taa_half UncompressInteferenceSeed(taa_subpixel_payload CompressedInteferenceSeed)
 {
-	return taa_half(CompressedParallaxFactor & SUB_PIXEL_MASK) * taa_half(MAX_PARALLAX_FACTOR / float(SUB_PIXEL_MASK)) + taa_half(1.0);
+	return 
+		taa_half(CompressedInteferenceSeed & SUB_PIXEL_INTERFERENCE_BIT_MASK) *
+		taa_half(1.0 / float(SUB_PIXEL_INTERFERENCE_BIT_MASK));
 }
 
 taa_half3 QuantizeForFloatRenderTarget(taa_half3 Color, taa_half E, float3 QuantizationError)
diff --git a/Engine/Shaders/Private/TemporalAA/TAACompareHistory.usf b/Engine/Shaders/Private/TemporalAA/TAACompareHistory.usf
--- a/Engine/Shaders/Private/TemporalAA/TAACompareHistory.usf
+++ b/Engine/Shaders/Private/TemporalAA/TAACompareHistory.usf
@@ -15,6 +15,7 @@
 Texture2D<taa_half> ParallaxRejectionMaskTexture;
 Texture2D<taa_half3> FilteredInputTexture;
 Texture2D<taa_half3> FilteredPredictionSceneColorTexture;
+Texture2D<taa_half> InterferenceWeightTexture;
 
 RWTexture2D<taa_half> HistoryRejectionOutput;
 
@@ -54,13 +55,15 @@ void MainCS(
 	taa_half3 InputMaxColor;
 	taa_half3 InputCenterColor;
 	taa_half ParallaxRejectionMask;
-
+	taa_half InterferenceWeight;
 	{
 		InputCenterColor = TransformColorForClampingBox(FilteredInputTexture.SampleLevel(GlobalPointClampedSampler, InputBufferUV, 0));
 		ParallaxRejectionMask = ParallaxRejectionMaskTexture[InputPixelPos];
 
 		InputMinColor = InputCenterColor;
 		InputMaxColor = InputCenterColor;
+		
+		InterferenceWeight = InterferenceWeightTexture[InputPixelPos];
 	}
 
 	UNROLL
@@ -95,6 +98,11 @@ void MainCS(
 	}
 	#endif
 	
+	if (bEnableInterferenceHeuristic)
+	{
+		Rejection = saturate(Rejection + 1.0 - InterferenceWeight);
+	}
+
 	if (ParallaxRejectionMask < taa_half(PARALLAX_REJECTION_MASK_THRESHOLD))
 	{
 		Rejection = taa_half(1.0);
diff --git a/Engine/Shaders/Private/TemporalAA/TAADebugHistory.usf b/Engine/Shaders/Private/TemporalAA/TAADebugHistory.usf
new file mode 100644
--- /dev/null
+++ b/Engine/Shaders/Private/TemporalAA/TAADebugHistory.usf
@@ -0,0 +1,70 @@
+// Copyright Epic Games, Inc. All Rights Reserved.
+
+#include "TAACommon.ush"
+#include "../Random.ush"
+
+
+//------------------------------------------------------- CONSTANTS
+
+#define TILE_SIZE 8
+
+
+//------------------------------------------------------- PARAMETERS
+
+Texture2D<taa_half3> History_Textures_0;
+Texture2D<taa_half3> History_Textures_1;
+Texture2D<taa_half2> History_Textures_2;
+Texture2D<taa_subpixel_details> History_Textures_3;
+
+Texture2D<taa_half3> PrevHistory_Textures_0;
+Texture2D<taa_half3> PrevHistory_Textures_1;
+Texture2D<taa_half2> PrevHistory_Textures_2;
+Texture2D<taa_subpixel_details> PrevHistory_Textures_3;
+
+
+//------------------------------------------------------- ENTRY POINT
+
+[numthreads(TILE_SIZE, TILE_SIZE, 1)]
+void MainCS(
+	uint2 GroupId : SV_GroupID,
+	uint GroupThreadIndex : SV_GroupIndex)
+{
+	float4 Debug = 0.0;
+
+	uint2 DispatchThreadId = (
+		ZOrder2D(GroupThreadIndex, log2(TILE_SIZE)) +
+		GroupId * uint2(TILE_SIZE, TILE_SIZE));
+
+	uint2 HistoryPixelPos = DispatchThreadId / SUB_PIXEL_GRID_SIZE;
+	float2 ViewportUV = (float2(HistoryPixelPos) + 0.5f) * HistoryInfo_ViewportSizeInverse;
+	float2 ScreenPos = ViewportUVToScreenPos(ViewportUV);
+
+	uint2 SubpixelPos = DispatchThreadId % SUB_PIXEL_GRID_SIZE;
+	uint SubpixelId = dot(SubpixelPos, uint2(1, SUB_PIXEL_GRID_SIZE));
+	uint SubpixelShift = SubpixelId * SUB_PIXEL_BIT_COUNT;
+	
+#if 1
+	taa_half3 RawHistory0 = History_Textures_0[HistoryPixelPos];
+	taa_half3 RawHistory1 = History_Textures_1[HistoryPixelPos];
+	taa_half2 RawHistory2 = History_Textures_2[HistoryPixelPos];
+	taa_subpixel_details RawHistory3 = History_Textures_3[HistoryPixelPos];
+#else
+	taa_half3 RawHistory0 = PrevHistory_Textures_0[HistoryPixelPos];
+	taa_half3 RawHistory1 = PrevHistory_Textures_1[HistoryPixelPos];
+	taa_half2 RawHistory2 = PrevHistory_Textures_2[HistoryPixelPos];
+	taa_subpixel_details RawHistory3 = PrevHistory_Textures_3[HistoryPixelPos];
+#endif
+
+	taa_subpixel_payload SubpixelDetails = taa_subpixel_payload(RawHistory3 >> SubpixelShift) & SUB_PIXEL_BIT_MASK;
+
+	Debug.r = UncompressParallaxFactor(SubpixelDetails);
+	Debug.g = UncompressInteferenceSeed(SubpixelDetails >> SUB_PIXEL_INTERFERENCE_BIT_OFFSET);
+	Debug.b = SubpixelId;
+	
+	#if DEBUG_OUTPUT
+	if (all(DispatchThreadId < HistoryInfo_ViewportMax * SUB_PIXEL_GRID_SIZE))
+	{
+		DebugOutput[DispatchThreadId] = Debug;
+	}
+	#endif
+}
diff --git a/Engine/Shaders/Private/TemporalAA/TAADecimateHistory.usf b/Engine/Shaders/Private/TemporalAA/TAADecimateHistory.usf
--- a/Engine/Shaders/Private/TemporalAA/TAADecimateHistory.usf
+++ b/Engine/Shaders/Private/TemporalAA/TAADecimateHistory.usf
@@ -17,7 +17,6 @@
 float4x4 RotationalClipToPrevClip;
 float3 OutputQuantizationError;
 float HistoryPreExposureCorrection;
-uint bCameraCut;
 
 Texture2D<taa_half3> InputSceneColorTexture;
 Texture2D<float2> DilatedVelocityTexture;
@@ -29,7 +28,8 @@ Texture2D<uint> ParallaxFactorTexture;
 Texture2D<taa_half3> PrevHistory_Textures_0;
 Texture2D<taa_half3> PrevHistory_Textures_1;
 Texture2D<taa_half2> PrevHistory_Textures_2;
-Texture2D<uint> PrevHistory_Textures_3;
+Texture2D<uint>      PrevHistory_Textures_3;
+Texture2D<taa_half>  PrevHistory_LowResTextures_0;
 
 RWTexture2D<taa_half3> HalfResSceneColorOutput;
 RWTexture2D<taa_half3> HalfResPredictionSceneColorOutput;
@@ -37,6 +37,7 @@ RWTexture2D<taa_half> HalfResParallaxRejectionMaskOutput;
 
 RWTexture2D<taa_half3> PredictionSceneColorOutput;
 RWTexture2D<taa_half> ParallaxRejectionMaskOutput;
+RWTexture2D<taa_half4> InterferenceSeedOutput;
 
 
 //------------------------------------------------------- LDS
@@ -78,6 +79,8 @@ void MainCS(
 	float DeviceZ = ClosestDepthTexture[InputPixelPos];
 	taa_half ParallaxFactor = UncompressParallaxFactor((ParallaxFactorTexture[InputPixelPos] >> 4) & 0xF);
 	taa_half3 InputSceneColor = InputSceneColorTexture[InputPixelPos];
+	
+	taa_half ClosestInputLuma4 = Luma4(InputSceneColor);
 
 
 	float WorldDepth = ConvertFromDeviceZ(DeviceZ);
@@ -142,28 +145,119 @@ void MainCS(
 	
 	// Sample the history
 	taa_half3 DecimatedColor;
-	taa_half HistoryValidity;
 	taa_half HistoryMaxRejectionFactor;
+
+	taa_half ClosestInterferenceSeed;
+	taa_half IntegratedInterferenceSeed;
+	taa_half MinInterferenceSeed;
+	taa_half MaxInterferenceSeed;
 	{
 		float2 KernelHistoryBufferUV = clamp(
 			HistoryPPCk * PrevHistoryInfo_ExtentInverse,
 			PrevHistoryInfo_UVViewportBilinearMin,
 			PrevHistoryInfo_UVViewportBilinearMax);
+		
+		float2 PrevInputBufferUV = (InputInfo_ScreenPosToViewportScale * PrevOutputScreenPos + InputInfo_ScreenPosToViewportBias) * InputInfo_ExtentInverse;
+
+		taa_half PrevUninterferenceSeed = PrevHistory_LowResTextures_0.SampleLevel(GlobalBilinearClampedSampler, PrevInputBufferUV, 0);
+		taa_half PrevUninterferenceLuma = PrevUninterferenceSeed * HdrWeightInvY(PrevUninterferenceSeed);
 
-		DecimatedColor = PrevHistory_Textures_0.SampleLevel(GlobalBilinearClampedSampler, KernelHistoryBufferUV, 0);
+		#if 0
+		{
+			DecimatedColor = PrevHistory_Textures_0.SampleLevel(GlobalBilinearClampedSampler, KernelHistoryBufferUV, 0);
+			
+			// Cancels out the flickering Luma component with the interference free luma from the inteference seed of the history.
+			if (bEnableInterferenceHeuristic)
+			{
+				DecimatedColor *= PrevUninterferenceLuma * SafeRcp(Luma4(DecimatedColor));
+			}
+		}
+		#else
+		{
+			DecimatedColor = 0.0;
+			UNROLL_N(4)
+			for (uint i = 0; i < 4; i++)
+			{
+				const float2 PixelOffset = float2(Offsets2x2[i]) - 0.5;
+
+				float2 SampleBufferUV = KernelHistoryBufferUV + PixelOffset * PrevHistoryInfo_ExtentInverse; 
+				SampleBufferUV = clamp(
+					SampleBufferUV,
+					PrevHistoryInfo_UVViewportBilinearMin,
+					PrevHistoryInfo_UVViewportBilinearMax);
+				
+				taa_half3 HistorySample = PrevHistory_Textures_0.SampleLevel(GlobalBilinearClampedSampler, SampleBufferUV, 0);
+
+				// Cancels out the flickering Luma component with the interference free luma from the inteference seed of the history.
+				if (bEnableInterferenceHeuristic)
+				{
+					HistorySample *= PrevUninterferenceLuma * SafeRcp(Luma4(HistorySample));
+				}
+
+				DecimatedColor += HistorySample;
+			}
 
+			DecimatedColor *= rcp(4.0);
+		}
+		#endif
+
+		// Compute the parallax rejection factor of the details in the history.
 		FSubpixelNeighborhood SubpixelNeigborhood = GatherPrevSubpixelNeighborhood(PrevHistory_Textures_3, KernelHistoryBufferUV);
 		{
 			uint EncodedHistoryMaxRejectionFactor = 0;
+
 			UNROLL_N(SUB_PIXEL_COUNT)
 			for (uint SubpixelId = 0; SubpixelId < SUB_PIXEL_COUNT; SubpixelId++)
 			{
-				uint SubpixelPayload = GetSubpixelPayload(SubpixelNeigborhood, SubpixelId);
-				EncodedHistoryMaxRejectionFactor = max(EncodedHistoryMaxRejectionFactor, SubpixelPayload);
+				taa_subpixel_payload SubpixelPayload = GetSubpixelPayload(SubpixelNeigborhood, SubpixelId);
+				taa_subpixel_payload EncodedRejectionFactor = (SubpixelPayload >> SUB_PIXEL_PARALLAX_FACTOR_BIT_OFFSET) & SUB_PIXEL_PARALLAX_FACTOR_BIT_MASK;
+				taa_subpixel_payload EncodedInterference = (SubpixelPayload >> SUB_PIXEL_INTERFERENCE_BIT_OFFSET) & SUB_PIXEL_INTERFERENCE_BIT_MASK;
+
+				EncodedHistoryMaxRejectionFactor = max(EncodedHistoryMaxRejectionFactor, EncodedRejectionFactor);
 			}
 
 			HistoryMaxRejectionFactor = UncompressParallaxFactor(EncodedHistoryMaxRejectionFactor);
 		}
+
+		// Compute the inteference seeds for detection of the interference
+		{
+			const taa_half RefInterferenceSeed = LumaToInterferenceSeed(ClosestInputLuma4);
+			ClosestInterferenceSeed = 0.0;
+			taa_half RefInterferenceSeedDist = 2.0;
+
+			MinInterferenceSeed = 1.0;
+			MaxInterferenceSeed = 0.0;
+
+			// Searches in the entire 2x2 neighborhood instead because this history is nearest reproject that can causes subpixel
+			// shifts over multiple frames.
+			UNROLL_N(4)
+			for (uint SampleId = 0; SampleId < 4; SampleId++)
+			{
+				UNROLL_N(SUB_PIXEL_COUNT)
+				for (uint SubpixelId = 0; SubpixelId < SUB_PIXEL_COUNT; SubpixelId++)
+				{
+					taa_subpixel_payload SubpixelPayload = SubpixelNeigborhood.SampleArray[SampleId] >> (SubpixelId * SUB_PIXEL_BIT_COUNT);
+					taa_subpixel_payload EncodedInterference = (SubpixelPayload >> SUB_PIXEL_INTERFERENCE_BIT_OFFSET) & SUB_PIXEL_INTERFERENCE_BIT_MASK;
+					
+					// TODO(G5TAA): move out of the loop.
+					taa_half SubpixelSeed = UncompressInteferenceSeed(EncodedInterference);
+
+					FLATTEN
+					if (abs(RefInterferenceSeed - SubpixelSeed) < RefInterferenceSeedDist)
+					{
+						ClosestInterferenceSeed = SubpixelSeed;
+						RefInterferenceSeedDist = abs(RefInterferenceSeed - SubpixelSeed);
+					}
+
+					MinInterferenceSeed = min(MinInterferenceSeed, SubpixelSeed);
+					MaxInterferenceSeed = max(MaxInterferenceSeed, SubpixelSeed);
+				}
+			}
+			
+			// actually integrating SubpixelSeed stil suffer interference that causes instability. So instead take
+			// average of min and max that turns out interference free.
+			IntegratedInterferenceSeed = (MinInterferenceSeed + MaxInterferenceSeed) * 0.5;
+		}
 	}
 	
 	// Reject subpixel details that might be in history but not input just based on their parrallax
@@ -193,7 +287,7 @@ void MainCS(
 		(PrevOutputScreenPos.y > 1.0))
 	//if (IsOffScreen2(bCameraCut, PrevOutputScreenPos, /* ParallaxRejectionMask = */ 1.0))
 	{
-		HistoryValidity = taa_half(0.0);
+		ParallaxRejectionMask = taa_half(0.0);
 	}
 
 	
@@ -268,7 +362,8 @@ void MainCS(
 		#endif
 		
 		ParallaxRejectionMaskOutput[InputPixelPos] = ParallaxRejectionMask;
-		
+		InterferenceSeedOutput[InputPixelPos] = taa_half4(ClosestInterferenceSeed, IntegratedInterferenceSeed, MinInterferenceSeed, MaxInterferenceSeed);
+
 		#if DEBUG_OUTPUT
 		{
 			DebugOutput[InputPixelPos] = Debug;
diff --git a/Engine/Shaders/Private/TemporalAA/TAADetectInterference.usf b/Engine/Shaders/Private/TemporalAA/TAADetectInterference.usf
new file mode 100644
--- /dev/null
+++ b/Engine/Shaders/Private/TemporalAA/TAADetectInterference.usf
@@ -0,0 +1,221 @@
+// Copyright Epic Games, Inc. All Rights Reserved.
+
+#include "TAACommon.ush"
+
+
+//------------------------------------------------------- CONFIG
+
+#define TILE_SIZE 8
+
+#define CONFIG_KERNEL 0
+
+#define CONFIG_OUTPUT_HALF_RES (DIM_OUTPUT_HALF_RES)
+
+#define CONFIG_SAMPLES 5
+
+
+//------------------------------------------------------- PARAMETERS
+
+float4x4 RotationalClipToPrevClip;
+float3 OutputQuantizationError;
+float HistoryPreExposureCorrection;
+
+Texture2D<taa_half3> InputSceneColorTexture;
+Texture2D<taa_half3> PredictionSceneColorTexture;
+Texture2D<taa_half> ParallaxRejectionMaskTexture;
+Texture2D<taa_half4> InterferenceSeedTexture;
+
+Texture2D<float2> DilatedVelocityTexture;
+
+Texture2D<taa_half> PrevHistory_LowResTextures_0;
+
+RWTexture2D<taa_half2> InterferenceWeightOutput;
+RWTexture2D<taa_half2> InterferenceSeedUpdateOutput;
+RWTexture2D<taa_half2> HistoryOutput_LowResTextures_0;
+
+
+//------------------------------------------------------- ENTRY POINT
+
+[numthreads(TILE_SIZE, TILE_SIZE, 1)]
+void MainCS(
+	uint2 GroupId : SV_GroupID,
+	uint GroupThreadIndex : SV_GroupIndex)
+{
+	float4 Debug = 0;
+	
+	uint2 DispatchThreadId = (
+		ZOrder2D(GroupThreadIndex, log2(TILE_SIZE)) +
+		GroupId * uint2(TILE_SIZE, TILE_SIZE));
+
+	uint2 InputPixelPos = InputInfo_ViewportMin + DispatchThreadId;
+	
+	float2 ViewportUV = (float2(DispatchThreadId) + 0.5f) * InputInfo_ViewportSizeInverse;
+	float2 ScreenPos = ViewportUVToScreenPos(ViewportUV);
+
+	float2 EncodedVelocity = DilatedVelocityTexture[InputPixelPos];
+	
+	taa_half ParallaxRejectionMask;
+	taa_half InputLuma;
+	taa_half InputLumaMin;
+	taa_half InputLumaMax;
+	taa_half ClosestInterferenceSeedMin;
+	taa_half ClosestInterferenceSeedMax;
+	taa_half IntegratedInterferenceSeed;
+	taa_half IntegratedInterferenceSeedMin;
+	taa_half IntegratedInterferenceSeedMax;
+	taa_half MinInterferenceSeed;
+	taa_half MaxInterferenceSeed;
+	{
+		ParallaxRejectionMask = ParallaxRejectionMaskTexture[InputPixelPos];
+
+		InputLuma = Luma4(InputSceneColorTexture[InputPixelPos]);
+		InputLumaMin = InputLuma;
+		InputLumaMax = InputLuma;
+
+		taa_half4 InterferenceSeed = InterferenceSeedTexture[InputPixelPos];
+
+		ClosestInterferenceSeedMin = InterferenceSeed.x;
+		ClosestInterferenceSeedMax = InterferenceSeed.x; 
+		IntegratedInterferenceSeed = InterferenceSeed.y;
+		IntegratedInterferenceSeedMin = InterferenceSeed.y;
+		IntegratedInterferenceSeedMax = InterferenceSeed.y;
+		MinInterferenceSeed = InterferenceSeed.z;
+		MaxInterferenceSeed = InterferenceSeed.w;
+	}
+	
+	float2 PrevScreenPos = ScreenPos - DecodeVelocityFromTexture(float4(EncodedVelocity, 0.0, 0.0)).xy;
+	bool bOffScreen = IsOffScreen(bCameraCut, PrevScreenPos, ParallaxRejectionMask);
+
+	taa_half PrevUninterferenceSeed;
+	{
+		float2 PrevInputBufferUV = (InputInfo_ScreenPosToViewportScale * PrevScreenPos + InputInfo_ScreenPosToViewportBias) * InputInfo_ExtentInverse;
+		
+		PrevUninterferenceSeed = PrevHistory_LowResTextures_0.SampleLevel(GlobalBilinearClampedSampler, PrevInputBufferUV, 0);
+	}
+
+
+	taa_half NeigborMaxInterferenceSeedGradient = 0.0;
+	
+	UNROLL_N(CONFIG_SAMPLES - 1)
+	for (uint SampleId = 1; SampleId < CONFIG_SAMPLES; SampleId++)
+	{
+		taa_half2 PixelOffset;
+			
+		#if CONFIG_SAMPLES == 9
+		{
+			PixelOffset = taa_half2(kOffsets3x3[kSquareIndexes3x3[SampleId]]);
+		}
+		#elif CONFIG_SAMPLES == 5 || CONFIG_SAMPLES == 6
+		{
+			if (SampleId == 5)
+			{
+				PixelOffset = taa_half2(SignFastInt(InputJitter)); // TODO: shader parameter
+			}
+			else
+			{
+				PixelOffset = taa_half2(kOffsets3x3[kPlusIndexes3x3[SampleId]]);
+			}
+		}
+		#else
+			#error Unknown sample count
+		#endif
+
+		float2 SampleInputBufferUV = (InputPixelPos + PixelOffset + 0.5) * InputInfo_ExtentInverse;
+		SampleInputBufferUV = clamp(SampleInputBufferUV, InputInfo_UVViewportBilinearMin, InputInfo_UVViewportBilinearMax);
+
+		taa_half NeighborInputLuma = Luma4(InputSceneColorTexture.SampleLevel(GlobalPointClampedSampler, SampleInputBufferUV, 0));
+		taa_half4 InterferenceSeed = InterferenceSeedTexture.SampleLevel(GlobalPointClampedSampler, SampleInputBufferUV, 0);
+		
+		taa_half InterferenceSeedGradient = InterferenceSeed.w - InterferenceSeed.z;
+
+		InputLumaMin = min(InputLumaMin, NeighborInputLuma);
+		InputLumaMax = max(InputLumaMax, NeighborInputLuma);
+
+		ClosestInterferenceSeedMin = min(ClosestInterferenceSeedMin, InterferenceSeed.x);
+		ClosestInterferenceSeedMax = max(ClosestInterferenceSeedMax, InterferenceSeed.x); 
+		
+		// TODO(G5TAA): this is InterferenceSeed.w - InterferenceSeed.z
+		IntegratedInterferenceSeedMin = min(IntegratedInterferenceSeedMin, InterferenceSeed.y);
+		IntegratedInterferenceSeedMax = max(IntegratedInterferenceSeedMax, InterferenceSeed.y);
+		
+		NeigborMaxInterferenceSeedGradient = max(NeigborMaxInterferenceSeedGradient, InterferenceSeedGradient); 
+	}
+
+	taa_half UninterferenceSeed;
+	bool bRejectHistoryInterferenceSeed;
+	bool bIsInterference;
+	{
+		// Error due to the low bitdepth of the inteference seeds in the history.
+		const taa_half InterferenceSeedEncodingError = rcp(taa_half(SUB_PIXEL_INTERFERENCE_BIT_MASK));
+
+		taa_half InputSeedMin = LumaToInterferenceSeed(InputLumaMin);
+		taa_half InputSeedMax = LumaToInterferenceSeed(InputLumaMax);
+
+		// Clamp the uninteference seed with current frame input to see what happens
+		taa_half ClampedPrevUninterferenceSeed = clamp(PrevUninterferenceSeed, InputSeedMin, InputSeedMax);
+		
+		// Measure the low frequency gradien in the seed. Needs to happen both on current frame input to avoid ghosting on
+		// current frame changes with high contrast geometric edges, and previous frames to avoid ghosting on area that had 
+		// high contrast geometric edges but no longer had in current frame due to shading change.
+		taa_half InterferenceSeedGradient = max(InputSeedMax - InputSeedMin, IntegratedInterferenceSeedMax - IntegratedInterferenceSeedMin);
+
+		// Test whether the current frame will be rejected.
+		bool bRejectPrevScene = abs(ClampedPrevUninterferenceSeed - PrevUninterferenceSeed) > InterferenceSeedEncodingError;
+		
+		// Detect if there is any inteference based on any subpixel detail range from the lowest LDR to highest LDR.
+		bool bRejectInterferenceSeedWithGradient = (
+			(MinInterferenceSeed < ClosestInterferenceSeedMin - InterferenceSeedEncodingError - InterferenceSeedGradient) ||
+			(MaxInterferenceSeed > ClosestInterferenceSeedMax + InterferenceSeedEncodingError + InterferenceSeedGradient));
+			
+		bool bRejectInterferenceSeed = (
+			(MinInterferenceSeed < ClosestInterferenceSeedMin - InterferenceSeedEncodingError) ||
+			(MaxInterferenceSeed > ClosestInterferenceSeedMax + InterferenceSeedEncodingError));
+
+		// Detect interference solely based on the inteference seed, but need to take into account low frequency
+		// gradien that there may be in the neighborhood to avoid ghosting on high contrast geometric edges.
+		bIsInterference = bRejectInterferenceSeedWithGradient;
+		
+		// If the input is not clamping, but the history's interference seeds are, that means the later are ghosting.
+		// Only care about bRejectInterferenceSeed here because it is only about not clamping inteference seeds in the history
+		// that may not be in the input whereas bRejectInterferenceSeedWithGradient is ghosting the seed to much in the history.
+		bRejectHistoryInterferenceSeed = !bRejectPrevScene && bRejectInterferenceSeed;
+
+		#if 0
+			bIsInterference = false;
+			bRejectHistoryInterferenceSeed = false;
+		#endif
+
+		// Don't clamp the seed history if there is an interference to keep luma stability in the input
+		// of the shading heuristic next frame.
+		FLATTEN
+		if (bIsInterference)
+		{
+			ClampedPrevUninterferenceSeed = PrevUninterferenceSeed;
+		}
+
+		FLATTEN
+		if (bOffScreen)
+		{
+			bRejectHistoryInterferenceSeed = true;
+			bIsInterference = false;
+			ClampedPrevUninterferenceSeed = LumaToInterferenceSeed(InputLuma);
+		}
+
+		UninterferenceSeed = lerp(ClampedPrevUninterferenceSeed, LumaToInterferenceSeed(InputLuma), 0.05);
+	}
+
+	// TODO(G5TAA): optimise VGPR here
+	if (all(InputPixelPos < InputInfo_ViewportMax))
+	{
+		InterferenceWeightOutput[InputPixelPos] = bIsInterference ? 0.0 : 1.0;
+		InterferenceSeedUpdateOutput[InputPixelPos] = bRejectHistoryInterferenceSeed ? 1.0 : 0.0;
+		
+		HistoryOutput_LowResTextures_0[InputPixelPos] = UninterferenceSeed;
+
+		#if DEBUG_OUTPUT
+		{
+			DebugOutput[InputPixelPos] = Debug;
+		}
+		#endif
+	}
+}
diff --git a/Engine/Shaders/Private/TemporalAA/TAAFilterFrequencies.usf b/Engine/Shaders/Private/TemporalAA/TAAFilterFrequencies.usf
--- a/Engine/Shaders/Private/TemporalAA/TAAFilterFrequencies.usf
+++ b/Engine/Shaders/Private/TemporalAA/TAAFilterFrequencies.usf
@@ -19,9 +19,16 @@
 
 #define CONFIG_ACCUMULATE_KARIS_WEIGHTING 1
 
-#define CONFIG_KERNEL_SPREAD_FACTOR 3
 #define CONFIG_KERNEL (KERNEL_3X3)
 
+// Use a spread factor the size of the clamping box in the history comparison to avoid one flickering input pixel
+// to flicker the entire comparison neighborhood
+#define CONFIG_KERNEL_SPREAD_FACTOR 3
+
+// Uses a spread factor that varies slightly between 2x2 pixel neighborhood to reduce chances of inteference
+// against structured input.
+#define CONFIG_USE_DYNAMIC_SPREAD_FACTOR 1
+
 
 //------------------------------------------------------- CONSTANTS
 
@@ -93,6 +100,7 @@ float3 OutputQuantizationError;
 Texture2D<taa_half3> InputTexture;
 Texture2D<taa_half3> PredictionSceneColorTexture;
 Texture2D<taa_half> ParallaxRejectionMaskTexture;
+Texture2D<taa_half> InterferenceWeightTexture;
 
 RWTexture2D<taa_half3> FilteredInputOutput;
 RWTexture2D<taa_half3> FilteredPredictionSceneColorOutput;
@@ -193,6 +201,12 @@ void MainCS(
 		#error Unknown kernel
 	#endif
 
+	#if CONFIG_USE_DYNAMIC_SPREAD_FACTOR
+		float2 DynamicSpreadFactor = CONFIG_KERNEL_SPREAD_FACTOR + float2((InputPixelPos % 2) ^ (InputPixelPos.yx % 2));
+	#else
+		float2 DynamicSpreadFactor = CONFIG_KERNEL_SPREAD_FACTOR;
+	#endif
+
 	// Fetch neighborhood of the kernel
 	UNROLL
 	for (uint SampleId = StartSampleId; SampleId < SampleCount; SampleId++)
@@ -246,7 +260,7 @@ void MainCS(
 			#error Unknown kernel
 		#endif
 
-		float2 SampleInputBufferUV = InputBufferUV + float2(CONFIG_KERNEL_SPREAD_FACTOR * Offset) * LowFrequencyInfo_ExtentInverse;
+		float2 SampleInputBufferUV = InputBufferUV + (DynamicSpreadFactor * Offset) * LowFrequencyInfo_ExtentInverse;
 		SampleInputBufferUV = clamp(SampleInputBufferUV, LowFrequencyInfo_UVViewportBilinearMin, LowFrequencyInfo_UVViewportBilinearMax);
 
 		taa_half3 InputSample = InputTexture.SampleLevel(GlobalPointClampedSampler, SampleInputBufferUV, 0);
@@ -254,8 +268,15 @@ void MainCS(
 		taa_half HistorySampleValidity = ParallaxRejectionMaskToHistoryValidity(
 			ParallaxRejectionMaskTexture.SampleLevel(GlobalPointClampedSampler, SampleInputBufferUV, 0));
 		
-		taa_half InputWeight = HistorySampleValidity * ComputeColorWeight(InputSample);
-		taa_half HistoryWeight = HistorySampleValidity * ComputeColorWeight(HistorySample);
+		taa_half SampleInterferenceWeight = InterferenceWeightTexture.SampleLevel(GlobalPointClampedSampler, SampleInputBufferUV, 0);
+		
+		if (!bEnableInterferenceHeuristic)
+		{
+			SampleInterferenceWeight = 1.0;
+		}
+
+		taa_half InputWeight = (HistorySampleValidity * SampleInterferenceWeight) * ComputeColorWeight(InputSample);
+		taa_half HistoryWeight = (HistorySampleValidity * SampleInterferenceWeight) * ComputeColorWeight(HistorySample);
 		
 		#if CONFIG_ACCUMULATE_KARIS_COLOR
 			InputSample *= 10.0 * HdrWeight4(InputSample);
diff --git a/Engine/Shaders/Private/TemporalAA/TAAUpdateHistory.usf b/Engine/Shaders/Private/TemporalAA/TAAUpdateHistory.usf
--- a/Engine/Shaders/Private/TemporalAA/TAAUpdateHistory.usf
+++ b/Engine/Shaders/Private/TemporalAA/TAAUpdateHistory.usf
@@ -33,20 +33,20 @@ Texture2D<taa_half> HistoryRejectionTexture;
 Texture2D<float2> DilatedVelocityTexture;
 Texture2D<uint> ParallaxFactorTexture;
 Texture2D<taa_half> ParallaxRejectionMaskTexture;
+Texture2D<taa_half> InterferenceSeedUpdateTexture;
 
 float3 HistoryQuantizationError;
 float HistoryPreExposureCorrection;
-uint bCameraCut;
 
 Texture2D<taa_half3> PrevHistory_Textures_0;
 Texture2D<taa_half3> PrevHistory_Textures_1;
 Texture2D<taa_half2> PrevHistory_Textures_2;
-Texture2D<uint>      PrevHistory_Textures_3;
+Texture2D<taa_subpixel_details> PrevHistory_Textures_3;
 
 RWTexture2D<taa_half3> HistoryOutput_Textures_0;
 RWTexture2D<taa_half3> HistoryOutput_Textures_1;
 RWTexture2D<taa_half2> HistoryOutput_Textures_2;
-RWTexture2D<uint>      HistoryOutput_Textures_3;
+RWTexture2D<taa_subpixel_details> HistoryOutput_Textures_3;
 
 RWTexture2D<taa_half3> SceneColorOutput;
 
@@ -114,18 +114,29 @@ void MainCS(
 	float2 PrevScreenPos = ScreenPos;
 	taa_half ParallaxRejectionMask = taa_half(1.0);
 	taa_half LowFrequencyRejection;
+	bool bRejectHistoryInterferenceSeed;
 	#if 1
 	{
-		float2 EncodedVelocity = DilatedVelocityTexture[InputPixelPos];
-		ParallaxRejectionMask = ParallaxRejectionMaskTexture[InputPixelPos];
-		PrevScreenPos = ScreenPos - DecodeVelocityFromTexture(float4(EncodedVelocity, 0.0, 0.0)).xy;
-
 		float2 RejectionBilinearBufferUV = clamp(
 			PPCo * InputInfo_ExtentInverse,
 			RejectionInfo_UVViewportBilinearMin,
 			RejectionInfo_UVViewportBilinearMax);
 		
+		float2 EncodedVelocity = DilatedVelocityTexture[InputPixelPos];
+		ParallaxRejectionMask = ParallaxRejectionMaskTexture[InputPixelPos];
+		PrevScreenPos = ScreenPos - DecodeVelocityFromTexture(float4(EncodedVelocity, 0.0, 0.0)).xy;
+
 		LowFrequencyRejection = HistoryRejectionTexture.SampleLevel(GlobalBilinearClampedSampler, RejectionBilinearBufferUV, 0);
+		
+		bRejectHistoryInterferenceSeed = InterferenceSeedUpdateTexture.SampleLevel(GlobalBilinearClampedSampler, RejectionBilinearBufferUV, 0) > 0.5;
+
+		#if !CONFIG_CLAMP
+		{
+			ParallaxRejectionMask = taa_half(1.0);
+			LowFrequencyRejection = taa_half(1.0);
+			bRejectHistoryInterferenceSeed = false;
+		}
+		#endif
 	}
 	#endif
 
@@ -155,6 +166,7 @@ void MainCS(
 	taa_half3 InputMinColor;
 	taa_half3 InputMaxColor;
 	taa_half InputPixelAlignement;
+	taa_half ClosestInputLuma4;
 	
 	ISOLATE
 	{
@@ -177,14 +189,14 @@ void MainCS(
 		InputMinColor = taa_half(INFINITE_FLOAT);
 		InputMaxColor = taa_half(-INFINITE_FLOAT);
 
-		UNROLL
+		UNROLL_N(CONFIG_SAMPLES)
 		for (uint SampleId = 0; SampleId < CONFIG_SAMPLES; SampleId++)
 		{
 			taa_half2 PixelOffset;
 			
 			#if CONFIG_SAMPLES == 9
 			{
-				PixelOffset = taa_half2(kOffsets3x3[SampleId]);
+				PixelOffset = taa_half2(kOffsets3x3[kSquareIndexes3x3[SampleId]]);
 			}
 			#elif CONFIG_SAMPLES == 5 || CONFIG_SAMPLES == 6
 			{
@@ -201,12 +213,6 @@ void MainCS(
 				#error Unknown sample count
 			#endif
 
-			//if (any(PixelOffset != 0.0))
-			//	continue;
-
-			//if (i != 4)
-			//	continue;
-
 			float2 SampleInputBufferUV = (InputInfo_ViewportMin + PPCk + PixelOffset) * InputInfo_ExtentInverse;
 			SampleInputBufferUV = clamp(SampleInputBufferUV, InputInfo_UVViewportBilinearMin, InputInfo_UVViewportBilinearMax);
 
@@ -222,6 +228,11 @@ void MainCS(
 
 			InputMinColor = min(InputMinColor, TransformColorForClampingBox(InputColor));
 			InputMaxColor = max(InputMaxColor, TransformColorForClampingBox(InputColor));
+
+			if (SampleId == 0)
+			{
+				ClosestInputLuma4 = Luma4(InputColor);
+			}
 		}
 		
 		FilteredInputColor *= SafeRcp(FilteredInputColorWeight);
@@ -248,7 +259,7 @@ void MainCS(
 	taa_half3 PrevFallbackColor;
 	taa_half PrevFallbackWeight;
 	
-	uint PrevSubpixelDetails;
+	taa_subpixel_details PrevSubpixelDetails;
 
 	ISOLATE
 	{
@@ -292,8 +303,8 @@ void MainCS(
 				UNROLL_N(SUB_PIXEL_COUNT)
 				for (uint SubpixelId = 0; SubpixelId < SUB_PIXEL_COUNT; SubpixelId++)
 				{
-					uint SubpixelPayload = GetSubpixelPayload(SubpixelNeighborhood, SubpixelId);
-					PrevSubpixelDetails |= SubpixelPayload << (BITS_PER_SUB_PIXEL * SubpixelId);
+					taa_subpixel_payload SubpixelPayload = GetSubpixelPayload(SubpixelNeighborhood, SubpixelId);
+					PrevSubpixelDetails |= SubpixelPayload << (SUB_PIXEL_BIT_COUNT * SubpixelId);
 				}
 			}
 
@@ -360,16 +371,13 @@ void MainCS(
 			PrevYCoCg, ClampedPrevYCoCg,
 			InputCenterYCoCg, InputMinColor, InputMaxColor);
 
-		//Debug.x = FallbackRejection;
-	}
-
-	#if !CONFIG_CLAMP
-	{
-		ClampedFallbackColor = PrevFallbackColor;
-		FallbackRejection = taa_half(1.0);
-		LowFrequencyRejection = taa_half(1.0);
+		#if !CONFIG_CLAMP
+		{
+			ClampedFallbackColor = PrevFallbackColor;
+			FallbackRejection = taa_half(1.0);
+		}
+		#endif
 	}
-	#endif
 
 	// Also reject sub pixel detail low frequency if detect mid frequency changes.
 	#if 0
@@ -566,18 +574,71 @@ void MainCS(
 	}
 
 	// Update subpixel detail
-	uint FinalSubpixelDetails;
+	taa_subpixel_details FinalSubpixelDetails;
 	{
 		taa_half2 dKO = taa_half2(PPCo - PPCk);
 
-		bool bUpdate = all(abs(dKO) < 0.5 * (InputInfo_Extent.x * HistoryInfo_ExtentInverse.x));
+		bool bUpdate = all(abs(dKO) < 0.5 * (InputInfo_ViewportSize.x * HistoryInfo_ViewportSizeInverse.x));
 
 		FinalSubpixelDetails = PrevSubpixelDetails;
-		
-		FLATTEN
-		if (bOffScreen)
+
+		// TODO Save VGPR here
+		taa_subpixel_payload ParallaxFactorBits = ParallaxFactorTexture[InputPixelPos] & SUB_PIXEL_PARALLAX_FACTOR_BIT_MASK;
+		taa_half InterferenceSeed = LumaToInterferenceSeed(ClosestInputLuma4);
+		taa_subpixel_payload InterferenceBits = CompressInteferenceSeed(InterferenceSeed);
+
 		{
-			FinalSubpixelDetails = 0;
+			const uint ParallaxFactorMask = (
+				(SUB_PIXEL_PARALLAX_FACTOR_BIT_MASK << (SUB_PIXEL_PARALLAX_FACTOR_BIT_OFFSET + 0 * SUB_PIXEL_BIT_COUNT)) | 
+				(SUB_PIXEL_PARALLAX_FACTOR_BIT_MASK << (SUB_PIXEL_PARALLAX_FACTOR_BIT_OFFSET + 1 * SUB_PIXEL_BIT_COUNT)) | 
+				(SUB_PIXEL_PARALLAX_FACTOR_BIT_MASK << (SUB_PIXEL_PARALLAX_FACTOR_BIT_OFFSET + 2 * SUB_PIXEL_BIT_COUNT)) | 
+				(SUB_PIXEL_PARALLAX_FACTOR_BIT_MASK << (SUB_PIXEL_PARALLAX_FACTOR_BIT_OFFSET + 3 * SUB_PIXEL_BIT_COUNT)) | 
+				0x0);
+
+			const uint InterferenceMask = (
+				(SUB_PIXEL_INTERFERENCE_BIT_MASK << (SUB_PIXEL_INTERFERENCE_BIT_OFFSET + 0 * SUB_PIXEL_BIT_COUNT)) | 
+				(SUB_PIXEL_INTERFERENCE_BIT_MASK << (SUB_PIXEL_INTERFERENCE_BIT_OFFSET + 1 * SUB_PIXEL_BIT_COUNT)) | 
+				(SUB_PIXEL_INTERFERENCE_BIT_MASK << (SUB_PIXEL_INTERFERENCE_BIT_OFFSET + 2 * SUB_PIXEL_BIT_COUNT)) | 
+				(SUB_PIXEL_INTERFERENCE_BIT_MASK << (SUB_PIXEL_INTERFERENCE_BIT_OFFSET + 3 * SUB_PIXEL_BIT_COUNT)) | 
+				0x0);
+
+			uint ResetInterferenceSeed = (
+				(uint(InterferenceBits) << (SUB_PIXEL_INTERFERENCE_BIT_OFFSET + 0 * SUB_PIXEL_BIT_COUNT)) | 
+				(uint(InterferenceBits) << (SUB_PIXEL_INTERFERENCE_BIT_OFFSET + 1 * SUB_PIXEL_BIT_COUNT)) | 
+				(uint(InterferenceBits) << (SUB_PIXEL_INTERFERENCE_BIT_OFFSET + 2 * SUB_PIXEL_BIT_COUNT)) | 
+				(uint(InterferenceBits) << (SUB_PIXEL_INTERFERENCE_BIT_OFFSET + 3 * SUB_PIXEL_BIT_COUNT)) | 
+				0x0);
+			
+			// Reset parallax factor
+			FLATTEN
+			if (bOffScreen)
+			{
+				FinalSubpixelDetails = FinalSubpixelDetails & ~ParallaxFactorMask;
+			}
+
+			// Reset inteference bit
+			FLATTEN
+			if (bRejectHistoryInterferenceSeed)
+			{
+				uint NewSubpixelDetails = FinalSubpixelDetails & ~InterferenceMask;
+				
+				taa_subpixel_payload InterferenceSeedMin = CompressInteferenceSeed(LumaToInterferenceSeed(LumaMin));
+				taa_subpixel_payload InterferenceSeedMax = CompressInteferenceSeed(LumaToInterferenceSeed(LumaMax));
+
+				UNROLL_N(SUB_PIXEL_COUNT)
+				for (uint SubpixelId = 0; SubpixelId < SUB_PIXEL_COUNT; SubpixelId++)
+				{
+					const uint SubpixelShift = SubpixelId * SUB_PIXEL_BIT_COUNT + SUB_PIXEL_INTERFERENCE_BIT_OFFSET;
+
+					taa_subpixel_payload SubpixelInterferenceSeed = (FinalSubpixelDetails >> SubpixelShift) & SUB_PIXEL_INTERFERENCE_BIT_MASK;
+
+					taa_subpixel_payload ClampedSubpixelInterferenceSeed = clamp(SubpixelInterferenceSeed, InterferenceSeedMin, InterferenceSeedMax);
+
+					NewSubpixelDetails |= ClampedSubpixelInterferenceSeed << SubpixelShift;
+				}
+
+				FinalSubpixelDetails = NewSubpixelDetails;
+			}
 		}
 
 		FLATTEN
@@ -586,11 +647,13 @@ void MainCS(
 			bool2 bBool = dKO < 0.0;
 
 			uint SubpixelId = dot(uint2(bBool), uint2(1, SUB_PIXEL_GRID_SIZE));
-			uint SubpixelShift = SubpixelId * BITS_PER_SUB_PIXEL;
+			uint SubpixelShift = SubpixelId * SUB_PIXEL_BIT_COUNT;
 
-			uint SubpixelPayload = ParallaxFactorTexture[InputPixelPos] & SUB_PIXEL_MASK;
+			taa_subpixel_payload SubpixelPayload = (
+				(ParallaxFactorBits << SUB_PIXEL_PARALLAX_FACTOR_BIT_OFFSET) |
+				(InterferenceBits << SUB_PIXEL_INTERFERENCE_BIT_OFFSET));
 
-			FinalSubpixelDetails = (FinalSubpixelDetails & (~(SUB_PIXEL_MASK << SubpixelShift))) | (SubpixelPayload << SubpixelShift);
+			FinalSubpixelDetails = (FinalSubpixelDetails & (~(SUB_PIXEL_BIT_MASK << SubpixelShift))) | (SubpixelPayload << SubpixelShift);
 		}
 	}
 
diff --git a/Engine/Source/Runtime/Renderer/Private/PostProcess/TemporalAA.cpp b/Engine/Source/Runtime/Renderer/Private/PostProcess/TemporalAA.cpp
--- a/Engine/Source/Runtime/Renderer/Private/PostProcess/TemporalAA.cpp
+++ b/Engine/Source/Runtime/Renderer/Private/PostProcess/TemporalAA.cpp
@@ -11,14 +11,14 @@
 #include "ScenePrivate.h"
 #include "RendererModule.h"
 
+#define COMPILE_TAA_DEBUG_PASSES (!UE_BUILD_SHIPPING)
+
 namespace
 {
 
 const int32 GTemporalAATileSizeX = 8;
 const int32 GTemporalAATileSizeY = 8;
 
-constexpr int32 kHistoryTextures = 4;
-
 TAutoConsoleVariable<int32> CVarTAAAlgorithm(
 	TEXT("r.TemporalAA.Algorithm"), 0,
 	TEXT("Algorithm to use for Temporal AA\n")
@@ -91,6 +91,11 @@ TAutoConsoleVariable<int32> CVarTAAFilterShadingRejection(
 	TEXT("Whether the shading rejection should have spatial statistical filtering pass to reduce flickering (default = 1)."),
 	ECVF_RenderThreadSafe);
 
+TAutoConsoleVariable<int32> CVarTAAEnableAntiInterference(
+	TEXT("r.TemporalAA.AntiInterference"), 1,
+	TEXT("Enable heuristic to detect geometric interference between input pixel grid alignement and structured geometry."),
+	ECVF_RenderThreadSafe);
+
 inline bool DoesPlatformSupportTemporalHistoryUpscale(EShaderPlatform Platform)
 {
 	return (IsPCPlatform(Platform) || FDataDrivenShaderPlatformInfo::GetSupportsTemporalHistoryUpscale(Platform))
@@ -111,22 +116,30 @@ BEGIN_SHADER_PARAMETER_STRUCT(FTAACommonParameters, )
 	SHADER_PARAMETER_STRUCT(FScreenPassTextureViewportParameters, HistoryInfo)
 
 	SHADER_PARAMETER(FVector2D, InputJitter)
+	SHADER_PARAMETER(int32, bCameraCut)
+	SHADER_PARAMETER(int32, bEnableInterferenceHeuristic)
 
 	SHADER_PARAMETER_STRUCT_REF(FViewUniformShaderParameters, ViewUniformBuffer)
 END_SHADER_PARAMETER_STRUCT()
 
 BEGIN_SHADER_PARAMETER_STRUCT(FTAAHistoryTextures, )
-	SHADER_PARAMETER_RDG_TEXTURE_ARRAY(Texture2D, Textures, [kHistoryTextures])
+	SHADER_PARAMETER_RDG_TEXTURE_ARRAY(Texture2D, LowResTextures, [FTemporalAAHistory::kLowResRenderTargetCount])
+	SHADER_PARAMETER_RDG_TEXTURE_ARRAY(Texture2D, Textures, [FTemporalAAHistory::kRenderTargetCount])
 END_SHADER_PARAMETER_STRUCT()
 
 BEGIN_SHADER_PARAMETER_STRUCT(FTAAHistoryUAVs, )
-	SHADER_PARAMETER_RDG_TEXTURE_UAV_ARRAY(RWTexture2D, Textures, [kHistoryTextures])
+	SHADER_PARAMETER_RDG_TEXTURE_UAV_ARRAY(RWTexture2D, LowResTextures, [FTemporalAAHistory::kLowResRenderTargetCount])
+	SHADER_PARAMETER_RDG_TEXTURE_UAV_ARRAY(RWTexture2D, Textures, [FTemporalAAHistory::kRenderTargetCount])
 END_SHADER_PARAMETER_STRUCT()
 
 FTAAHistoryUAVs CreateUAVs(FRDGBuilder& GraphBuilder, const FTAAHistoryTextures& Textures)
 {
 	FTAAHistoryUAVs UAVs;
-	for (int32 i = 0; i < kHistoryTextures; i++)
+	for (int32 i = 0; i < Textures.LowResTextures.Num(); i++)
+	{
+		UAVs.LowResTextures[i] = GraphBuilder.CreateUAV(Textures.LowResTextures[i]);
+	}
+	for (int32 i = 0; i < Textures.Textures.Num(); i++)
 	{
 		UAVs.Textures[i] = GraphBuilder.CreateUAV(Textures.Textures[i]);
 	}
@@ -370,7 +383,6 @@ class FTAADecimateHistoryCS : public FTAAGen5Shader
 		SHADER_PARAMETER(FVector, OutputQuantizationError)
 		SHADER_PARAMETER(float, HistoryPreExposureCorrection)
 		SHADER_PARAMETER(float, WorldDepthToPixelWorldRadius)
-		SHADER_PARAMETER(int32, bCameraCut)
 
 		SHADER_PARAMETER_RDG_TEXTURE(Texture2D, InputSceneColorTexture)
 		SHADER_PARAMETER_RDG_TEXTURE(Texture2D, DilatedVelocityTexture)
@@ -387,10 +399,35 @@ class FTAADecimateHistoryCS : public FTAAGen5Shader
 		SHADER_PARAMETER_RDG_TEXTURE_UAV(RWTexture2D, HalfResParallaxRejectionMaskOutput)
 		SHADER_PARAMETER_RDG_TEXTURE_UAV(RWTexture2D, PredictionSceneColorOutput)
 		SHADER_PARAMETER_RDG_TEXTURE_UAV(RWTexture2D, ParallaxRejectionMaskOutput)
+		SHADER_PARAMETER_RDG_TEXTURE_UAV(RWTexture2D, InterferenceSeedOutput)
 		SHADER_PARAMETER_RDG_TEXTURE_UAV(RWTexture2D, DebugOutput)
 	END_SHADER_PARAMETER_STRUCT()
 }; // class FTAADecimateHistoryCS
 
+class FTAADetectInterferenceCS : public FTAAGen5Shader
+{
+	DECLARE_GLOBAL_SHADER(FTAADetectInterferenceCS);
+	SHADER_USE_PARAMETER_STRUCT(FTAADetectInterferenceCS, FTAAGen5Shader);
+	
+	BEGIN_SHADER_PARAMETER_STRUCT(FParameters, )
+		SHADER_PARAMETER_STRUCT_INCLUDE(FTAACommonParameters, CommonParameters)
+
+		SHADER_PARAMETER_RDG_TEXTURE(Texture2D, InputSceneColorTexture)
+		SHADER_PARAMETER_RDG_TEXTURE(Texture2D, DilatedVelocityTexture)
+		SHADER_PARAMETER_RDG_TEXTURE(Texture2D, PredictionSceneColorTexture)
+		SHADER_PARAMETER_RDG_TEXTURE(Texture2D, ParallaxRejectionMaskTexture)
+		SHADER_PARAMETER_RDG_TEXTURE(Texture2D, InterferenceSeedTexture)
+
+		SHADER_PARAMETER_STRUCT(FScreenPassTextureViewportParameters, PrevHistoryInfo)
+		SHADER_PARAMETER_STRUCT(FTAAHistoryTextures, PrevHistory)
+
+		SHADER_PARAMETER_STRUCT(FTAAHistoryUAVs, HistoryOutput)
+		SHADER_PARAMETER_RDG_TEXTURE_UAV(RWTexture2D, InterferenceWeightOutput)
+		SHADER_PARAMETER_RDG_TEXTURE_UAV(RWTexture2D, InterferenceSeedUpdateOutput)
+		SHADER_PARAMETER_RDG_TEXTURE_UAV(RWTexture2D, DebugOutput)
+	END_SHADER_PARAMETER_STRUCT()
+}; // class FTAADetectInterferenceCS
+
 class FTAAFilterFrequenciesCS : public FTAAGen5Shader
 {
 	DECLARE_GLOBAL_SHADER(FTAAFilterFrequenciesCS);
@@ -403,6 +440,7 @@ class FTAAFilterFrequenciesCS : public FTAAGen5Shader
 		SHADER_PARAMETER_RDG_TEXTURE(Texture2D, InputTexture)
 		SHADER_PARAMETER_RDG_TEXTURE(Texture2D, PredictionSceneColorTexture)
 		SHADER_PARAMETER_RDG_TEXTURE(Texture2D, ParallaxRejectionMaskTexture)
+		SHADER_PARAMETER_RDG_TEXTURE(Texture2D, InterferenceWeightTexture)
 
 		SHADER_PARAMETER_RDG_TEXTURE_UAV(RWTexture2D, FilteredInputOutput)
 		SHADER_PARAMETER_RDG_TEXTURE_UAV(RWTexture2D, FilteredPredictionSceneColorOutput)
@@ -420,6 +458,7 @@ class FTAACompareHistoryCS : public FTAAGen5Shader
 		SHADER_PARAMETER_RDG_TEXTURE(Texture2D, ParallaxRejectionMaskTexture)
 		SHADER_PARAMETER_RDG_TEXTURE(Texture2D, FilteredInputTexture)
 		SHADER_PARAMETER_RDG_TEXTURE(Texture2D, FilteredPredictionSceneColorTexture)
+		SHADER_PARAMETER_RDG_TEXTURE(Texture2D, InterferenceWeightTexture)
 		SHADER_PARAMETER_RDG_TEXTURE_UAV(RWTexture2D, HistoryRejectionOutput)
 		SHADER_PARAMETER_RDG_TEXTURE_UAV(RWTexture2D, DebugOutput)
 	END_SHADER_PARAMETER_STRUCT()
@@ -465,10 +504,10 @@ class FTAAUpdateHistoryCS : public FTAAGen5Shader
 		SHADER_PARAMETER_RDG_TEXTURE(Texture2D, DilatedVelocityTexture)
 		SHADER_PARAMETER_RDG_TEXTURE(Texture2D, ParallaxFactorTexture)
 		SHADER_PARAMETER_RDG_TEXTURE(Texture2D, ParallaxRejectionMaskTexture)
+		SHADER_PARAMETER_RDG_TEXTURE(Texture2D, InterferenceSeedUpdateTexture)
 
 		SHADER_PARAMETER(FVector, HistoryQuantizationError)
 		SHADER_PARAMETER(float, HistoryPreExposureCorrection)
-		SHADER_PARAMETER(int32, bCameraCut)
 
 		SHADER_PARAMETER_STRUCT(FScreenPassTextureViewportParameters, PrevHistoryInfo)
 		SHADER_PARAMETER_STRUCT(FTAAHistoryTextures, PrevHistory)
@@ -483,12 +522,32 @@ IMPLEMENT_GLOBAL_SHADER(FTAAStandaloneCS, "/Engine/Private/TemporalAA/TAAStandal
 IMPLEMENT_GLOBAL_SHADER(FTAAClearPrevTexturesCS, "/Engine/Private/TemporalAA/TAAClearPrevTextures.usf", "MainCS", SF_Compute);
 IMPLEMENT_GLOBAL_SHADER(FTAADilateVelocityCS, "/Engine/Private/TemporalAA/TAADilateVelocity.usf", "MainCS", SF_Compute);
 IMPLEMENT_GLOBAL_SHADER(FTAADecimateHistoryCS, "/Engine/Private/TemporalAA/TAADecimateHistory.usf", "MainCS", SF_Compute);
+IMPLEMENT_GLOBAL_SHADER(FTAADetectInterferenceCS, "/Engine/Private/TemporalAA/TAADetectInterference.usf", "MainCS", SF_Compute);
 IMPLEMENT_GLOBAL_SHADER(FTAAFilterFrequenciesCS, "/Engine/Private/TemporalAA/TAAFilterFrequencies.usf", "MainCS", SF_Compute);
 IMPLEMENT_GLOBAL_SHADER(FTAACompareHistoryCS, "/Engine/Private/TemporalAA/TAACompareHistory.usf", "MainCS", SF_Compute);
 IMPLEMENT_GLOBAL_SHADER(FTAAPostfilterRejectionCS, "/Engine/Private/TemporalAA/TAAPostfilterRejection.usf", "MainCS", SF_Compute);
 IMPLEMENT_GLOBAL_SHADER(FTAADilateRejectionCS, "/Engine/Private/TemporalAA/TAADilateRejection.usf", "MainCS", SF_Compute);
 IMPLEMENT_GLOBAL_SHADER(FTAAUpdateHistoryCS, "/Engine/Private/TemporalAA/TAAUpdateHistory.usf", "MainCS", SF_Compute);
 
+#if COMPILE_TAA_DEBUG_PASSES
+
+class FTAADebugHistoryCS : public FTAAGen5Shader
+{
+	DECLARE_GLOBAL_SHADER(FTAADebugHistoryCS);
+	SHADER_USE_PARAMETER_STRUCT(FTAADebugHistoryCS, FTAAGen5Shader);
+
+	BEGIN_SHADER_PARAMETER_STRUCT(FParameters, )
+		SHADER_PARAMETER_STRUCT_INCLUDE(FTAACommonParameters, CommonParameters)
+		SHADER_PARAMETER_STRUCT(FTAAHistoryTextures, History)
+		SHADER_PARAMETER_STRUCT(FTAAHistoryTextures, PrevHistory)
+		SHADER_PARAMETER_RDG_TEXTURE_UAV(RWTexture2D, DebugOutput)
+	END_SHADER_PARAMETER_STRUCT()
+}; // class FTAADebugHistoryCS
+
+IMPLEMENT_GLOBAL_SHADER(FTAADebugHistoryCS, "/Engine/Private/TemporalAA/TAADebugHistory.usf", "MainCS", SF_Compute);
+
+#endif
+
 float CatmullRom(float x)
 {
 	float ax = FMath::Abs(x);
@@ -1036,6 +1095,8 @@ static void AddGen5MainTemporalAAPasses(
 
 	bool bHalfResLowFrequency = CVarTAAHalfResShadingRejection.GetValueOnRenderThread() != 0;
 
+	bool bEnableInterferenceHeuristic = CVarTAAEnableAntiInterference.GetValueOnRenderThread() != 0;
+
 	FIntPoint InputExtent = PassInputs.SceneColorTexture->Desc.Extent;
 	FIntRect InputRect = View.ViewRect;
 
@@ -1093,6 +1154,7 @@ static void AddGen5MainTemporalAAPasses(
 	RDG_GPU_STAT_SCOPE(GraphBuilder, TAA);
 
 	FRDGTextureRef BlackDummy = GraphBuilder.RegisterExternalTexture(GSystemTextures.BlackDummy);
+	FRDGTextureRef WhiteDummy = GraphBuilder.RegisterExternalTexture(GSystemTextures.WhiteDummy);
 
 	FTAACommonParameters CommonParameters;
 	{
@@ -1109,6 +1171,8 @@ static void AddGen5MainTemporalAAPasses(
 			HistoryExtent, FIntRect(FIntPoint(0, 0), HistorySize)));
 
 		CommonParameters.InputJitter = View.TemporalJitterPixels;
+		CommonParameters.bCameraCut = bCameraCut;
+		CommonParameters.bEnableInterferenceHeuristic = bEnableInterferenceHeuristic;
 		CommonParameters.ViewUniformBuffer = View.ViewUniformBuffer;
 	}
 
@@ -1203,7 +1267,12 @@ static void AddGen5MainTemporalAAPasses(
 		PrevHistoryInfo = GetScreenPassTextureViewportParameters(FScreenPassTextureViewport(
 			FIntPoint(1, 1), FIntRect(FIntPoint(0, 0), FIntPoint(1, 1))));
 
-		for (int32 i = 0; i < kHistoryTextures; i++)
+		for (int32 i = 0; i < PrevHistory.LowResTextures.Num(); i++)
+		{
+			PrevHistory.LowResTextures[i] = BlackDummy;
+		}
+
+		for (int32 i = 0; i < PrevHistory.Textures.Num(); i++)
 		{
 			PrevHistory.Textures[i] = BlackDummy;
 		}
@@ -1214,7 +1283,19 @@ static void AddGen5MainTemporalAAPasses(
 			InputHistory.ReferenceBufferSize,
 			InputHistory.ViewportRect));
 
-		for (int32 i = 0; i < kHistoryTextures; i++)
+		for (int32 i = 0; i < PrevHistory.LowResTextures.Num(); i++)
+		{
+			if (InputHistory.LowResRT[i].IsValid())
+			{
+				PrevHistory.LowResTextures[i] = GraphBuilder.RegisterExternalTexture(InputHistory.LowResRT[i]);
+			}
+			else
+			{
+				PrevHistory.LowResTextures[i] = BlackDummy;
+			}
+		}
+
+		for (int32 i = 0; i < PrevHistory.Textures.Num(); i++)
 		{
 			if (InputHistory.RT[i].IsValid())
 			{
@@ -1229,12 +1310,40 @@ static void AddGen5MainTemporalAAPasses(
 		// InputHistory.SafeRelease(); TODO
 	}
 
+
+	FTAAHistoryTextures History;
+	{
+		FRDGTextureDesc Desc = FRDGTextureDesc::Create2D(
+			HistoryExtent,
+			(CVarTAAR11G11B10History.GetValueOnRenderThread() != 0) ? PF_FloatR11G11B10 : PF_FloatRGBA,
+			FClearValueBinding::None,
+			TexCreate_ShaderResource | TexCreate_UAV);
+
+		FRDGTextureDesc LowResDesc = FRDGTextureDesc::Create2D(
+			InputExtent,
+			PF_R8,
+			FClearValueBinding::None,
+			TexCreate_ShaderResource | TexCreate_UAV);
+
+		History.Textures[0] = GraphBuilder.CreateTexture(Desc, TEXT("TAA.History.LowFrequencies"));
+		History.Textures[1] = GraphBuilder.CreateTexture(Desc, TEXT("TAA.History.HighFrequencies"));
+
+		Desc.Format = PF_R8G8;
+		History.Textures[2] = GraphBuilder.CreateTexture(Desc, TEXT("TAA.History.Metadata"));
+
+		Desc.Format = PF_R32_UINT;
+		History.Textures[3] = GraphBuilder.CreateTexture(Desc, TEXT("TAA.History.SubpixelInfo"));
+
+		History.LowResTextures[0] = GraphBuilder.CreateTexture(LowResDesc, TEXT("TAA.History.LowResMetadata"));
+	}
+
 	// Decimate input to flicker at same frequency as input.
 	FRDGTextureRef HalfResInputSceneColorTexture = nullptr;
 	FRDGTextureRef HalfResPredictionSceneColorTexture = nullptr;
 	FRDGTextureRef HalfResParallaxRejectionMaskTexture = nullptr;
 	FRDGTextureRef PredictionSceneColorTexture = nullptr;
 	FRDGTextureRef ParallaxRejectionMaskTexture = nullptr;
+	FRDGTextureRef InterferenceSeedTexture = nullptr;
 	{
 		{
 			FRDGTextureDesc Desc = FRDGTextureDesc::Create2D(
@@ -1244,6 +1353,10 @@ static void AddGen5MainTemporalAAPasses(
 				/* InFlags = */ TexCreate_ShaderResource | TexCreate_UAV);
 
 			ParallaxRejectionMaskTexture = GraphBuilder.CreateTexture(Desc, TEXT("TAA.ParallaxRejectionMask"));
+
+			// TODO(G5TAA): can compress to the history seed's 4bit per pixel
+			Desc.Format = PF_R8G8B8A8;
+			InterferenceSeedTexture = GraphBuilder.CreateTexture(Desc, TEXT("TAA.Interference.Seed"));
 		}
 
 		FTAADecimateHistoryCS::FParameters* PassParameters = GraphBuilder.AllocParameters<FTAADecimateHistoryCS::FParameters>();
@@ -1265,7 +1378,6 @@ static void AddGen5MainTemporalAAPasses(
 			// Should be multiplied 0.5* for the diameter to radius, and by 2.0 because GetTanHalfFieldOfView() cover only half of the pixels.
 			PassParameters->WorldDepthToPixelWorldRadius = TanHalfFieldOfView / float(View.ViewRect.Width());
 		}
-		PassParameters->bCameraCut = bCameraCut;
 
 		PassParameters->InputSceneColorTexture = PassInputs.SceneColorTexture;
 		PassParameters->DilatedVelocityTexture = DilatedVelocityTexture;
@@ -1309,6 +1421,7 @@ static void AddGen5MainTemporalAAPasses(
 		}
 
 		PassParameters->ParallaxRejectionMaskOutput = GraphBuilder.CreateUAV(ParallaxRejectionMaskTexture);
+		PassParameters->InterferenceSeedOutput = GraphBuilder.CreateUAV(InterferenceSeedTexture);
 		PassParameters->DebugOutput = CreateDebugUAV(LowFrequencyExtent, TEXT("Debug.TAA.DecimateHistory"));
 
 		FTAADecimateHistoryCS::FPermutationDomain PermutationVector;
@@ -1325,6 +1438,87 @@ static void AddGen5MainTemporalAAPasses(
 			FComputeShaderUtils::GetGroupCount(InputRect.Size(), 8));
 	}
 
+	// Detect interference between geometry and alignement of input texel centers. It is not about awnser whether an
+	// interference has happened in the past, because interference change based on input resolution or camera position.
+	// So to remain stable on camera movement and input resolution change, it is about awnsering the question on whether
+	// an interference is possible.
+	// TODO(G5TAA): Could sample the interference seed in the DilateVelocity and detect interference in the decimate.
+	TStaticArray<bool, FTemporalAAHistory::kLowResRenderTargetCount> ExtractLowResHistoryTexture;
+	FRDGTextureRef InterferenceWeightTexture;
+	FRDGTextureRef InterferenceSeedUpdateTexture;
+	if (bEnableInterferenceHeuristic)
+	{
+		{
+			// TODO(G5TAA): Compress to 1bit per pixel
+			FRDGTextureDesc Desc = FRDGTextureDesc::Create2D(
+				InputExtent,
+				PF_R8,
+				FClearValueBinding::None,
+				/* InFlags = */ TexCreate_ShaderResource | TexCreate_UAV);
+
+			InterferenceWeightTexture = GraphBuilder.CreateTexture(Desc, TEXT("TAA.Interference.Weight"));
+			InterferenceSeedUpdateTexture = GraphBuilder.CreateTexture(Desc, TEXT("TAA.Interference.SeedUpdateRequest"));
+		}
+
+		FTAADetectInterferenceCS::FParameters* PassParameters = GraphBuilder.AllocParameters<FTAADetectInterferenceCS::FParameters>();
+		PassParameters->CommonParameters = CommonParameters;
+
+		PassParameters->InputSceneColorTexture = PassInputs.SceneColorTexture;
+		PassParameters->DilatedVelocityTexture = DilatedVelocityTexture;
+		PassParameters->PredictionSceneColorTexture = PredictionSceneColorTexture;
+		PassParameters->ParallaxRejectionMaskTexture = ParallaxRejectionMaskTexture;
+		PassParameters->InterferenceSeedTexture = InterferenceSeedTexture;
+
+		PassParameters->PrevHistoryInfo = PrevHistoryInfo;
+		PassParameters->PrevHistory = PrevHistory;
+
+		PassParameters->HistoryOutput = CreateUAVs(GraphBuilder, History);
+		PassParameters->InterferenceWeightOutput = GraphBuilder.CreateUAV(InterferenceWeightTexture);
+		PassParameters->InterferenceSeedUpdateOutput = GraphBuilder.CreateUAV(InterferenceSeedUpdateTexture);
+		PassParameters->DebugOutput = CreateDebugUAV(LowFrequencyExtent, TEXT("Debug.TAA.DetectInterference"));
+
+		TShaderMapRef<FTAADetectInterferenceCS> ComputeShader(View.ShaderMap);
+		ClearUnusedGraphResources(ComputeShader, PassParameters);
+
+		for (int32 i = 0; i < PrevHistory.LowResTextures.Num(); i++)
+		{
+			bool bNeedsExtractForNextFrame = PassParameters->PrevHistory.LowResTextures[i] != nullptr;
+			bool bPrevFrameIsntAvailable = PassParameters->PrevHistory.LowResTextures[i] == BlackDummy;
+			bool bOutputHistory = PassParameters->HistoryOutput.LowResTextures[i] != nullptr;
+
+			ExtractLowResHistoryTexture[i] = bNeedsExtractForNextFrame;
+
+			if (bPrevFrameIsntAvailable && !PassParameters->CommonParameters.bCameraCut)
+			{
+				//ensureMsgf(false, TEXT("Shaders read PrevHistory[%d] but doesn't write HistoryOutput[%d]"), i, i);
+				PassParameters->CommonParameters.bCameraCut = true;
+			}
+
+			if (bOutputHistory && !bNeedsExtractForNextFrame)
+			{
+				ensureMsgf(false, TEXT("Shaders write HistoryOutput[%d] but doesn't read PrevHistory[%d]"), i, i);
+			}
+		}
+
+		FComputeShaderUtils::AddPass(
+			GraphBuilder,
+			RDG_EVENT_NAME("TAA DetectInterference %dx%d", InputRect.Width(), InputRect.Height()),
+			ComputeShader,
+			PassParameters,
+			FComputeShaderUtils::GetGroupCount(InputRect.Size(), 8));
+	}
+	else
+	{
+		for (int32 i = 0; i < ExtractLowResHistoryTexture.Num(); i++)
+		{
+			ExtractLowResHistoryTexture[i] = false;
+		}
+
+		// TODO(G5TAA): Shader permutation.
+		InterferenceWeightTexture = WhiteDummy;
+		InterferenceSeedUpdateTexture = BlackDummy;
+	}
+
 	// Reject the history with frequency decomposition.
 	FRDGTextureRef HistoryRejectionTexture;
 	{
@@ -1358,6 +1552,7 @@ static void AddGen5MainTemporalAAPasses(
 				PassParameters->InputTexture = PassInputs.SceneColorTexture;
 				PassParameters->PredictionSceneColorTexture = PredictionSceneColorTexture;
 				PassParameters->ParallaxRejectionMaskTexture = ParallaxRejectionMaskTexture;
+				PassParameters->InterferenceWeightTexture = InterferenceWeightTexture;
 			}
 
 			PassParameters->FilteredInputOutput = GraphBuilder.CreateUAV(FilteredInputTexture);
@@ -1390,6 +1585,7 @@ static void AddGen5MainTemporalAAPasses(
 			PassParameters->ParallaxRejectionMaskTexture = ParallaxRejectionMaskTexture;
 			PassParameters->FilteredInputTexture = FilteredInputTexture;
 			PassParameters->FilteredPredictionSceneColorTexture = FilteredPredictionSceneColorTexture;
+			PassParameters->InterferenceWeightTexture = InterferenceWeightTexture;
 
 			PassParameters->HistoryRejectionOutput = GraphBuilder.CreateUAV(HistoryRejectionTexture);
 			PassParameters->DebugOutput = CreateDebugUAV(LowFrequencyExtent, TEXT("Debug.TAA.CompareHistory"));
@@ -1446,28 +1642,9 @@ static void AddGen5MainTemporalAAPasses(
 			FComputeShaderUtils::GetGroupCount(RejectionRect.Size(), 8));
 	}
 
-	TStaticArray<bool, kHistoryTextures> ExtractHistory;
+	TStaticArray<bool, FTemporalAAHistory::kRenderTargetCount> ExtractHistoryTexture;
 	FRDGTextureRef SceneColorOutputTexture;
-	FTAAHistoryTextures History;
 	{
-		// Allocate a new history
-		{
-			FRDGTextureDesc Desc = FRDGTextureDesc::Create2D(
-				HistoryExtent,
-				(CVarTAAR11G11B10History.GetValueOnRenderThread() != 0) ? PF_FloatR11G11B10 : PF_FloatRGBA,
-				FClearValueBinding::None,
-				TexCreate_ShaderResource | TexCreate_UAV);
-
-			History.Textures[0] = GraphBuilder.CreateTexture(Desc, TEXT("TAA.History.LowFrequencies"));
-			History.Textures[1] = GraphBuilder.CreateTexture(Desc, TEXT("TAA.History.HighFrequencies"));
-
-			Desc.Format = PF_R8G8;
-			History.Textures[2] = GraphBuilder.CreateTexture(Desc, TEXT("TAA.History.Metadata"));
-
-			Desc.Format = PF_R16_UINT;
-			History.Textures[3] = GraphBuilder.CreateTexture(Desc, TEXT("TAA.History.SubpixelInfo"));
-		}
-
 		// Allocate output
 		{
 			FRDGTextureDesc Desc = FRDGTextureDesc::Create2D(
@@ -1488,10 +1665,10 @@ static void AddGen5MainTemporalAAPasses(
 		PassParameters->DilatedVelocityTexture = DilatedVelocityTexture;
 		PassParameters->ParallaxFactorTexture = ParallaxFactorTexture;
 		PassParameters->ParallaxRejectionMaskTexture = ParallaxRejectionMaskTexture;
+		PassParameters->InterferenceSeedUpdateTexture = InterferenceSeedUpdateTexture;
 
 		PassParameters->HistoryQuantizationError = ComputePixelFormatQuantizationError(History.Textures[0]->Desc.Format);
 		PassParameters->HistoryPreExposureCorrection = View.PreExposure / View.PrevViewInfo.SceneColorPreExposure;
-		PassParameters->bCameraCut = bCameraCut;
 
 		PassParameters->PrevHistoryInfo = PrevHistoryInfo;
 		PassParameters->PrevHistory = PrevHistory;
@@ -1503,18 +1680,18 @@ static void AddGen5MainTemporalAAPasses(
 		TShaderMapRef<FTAAUpdateHistoryCS> ComputeShader(View.ShaderMap);
 		ClearUnusedGraphResources(ComputeShader, PassParameters);
 
-		for (int32 i = 0; i < kHistoryTextures; i++)
+		for (int32 i = 0; i < PrevHistory.Textures.Num(); i++)
 		{
 			bool bNeedsExtractForNextFrame = PassParameters->PrevHistory.Textures[i] != nullptr;
 			bool bPrevFrameIsntAvailable = PassParameters->PrevHistory.Textures[i] == BlackDummy;
 			bool bOutputHistory = PassParameters->HistoryOutput.Textures[i] != nullptr;
 
-			ExtractHistory[i] = bNeedsExtractForNextFrame;
+			ExtractHistoryTexture[i] = bNeedsExtractForNextFrame;
 
-			if (bPrevFrameIsntAvailable && !PassParameters->bCameraCut)
+			if (bPrevFrameIsntAvailable && !bCameraCut)
 			{
 				//ensureMsgf(false, TEXT("Shaders read PrevHistory[%d] but doesn't write HistoryOutput[%d]"), i, i);
-				PassParameters->bCameraCut = true;
+				PassParameters->CommonParameters.bCameraCut = true;
 			}
 
 			if (bOutputHistory && !bNeedsExtractForNextFrame)
@@ -1533,14 +1710,42 @@ static void AddGen5MainTemporalAAPasses(
 			FComputeShaderUtils::GetGroupCount(HistorySize, 8));
 	}
 
+	// Debug the history.
+	#if COMPILE_TAA_DEBUG_PASSES
+	{
+		const int32 kHistoryUpscalingFactor = 2;
+
+		FTAADebugHistoryCS::FParameters* PassParameters = GraphBuilder.AllocParameters<FTAADebugHistoryCS::FParameters>();
+		PassParameters->CommonParameters = CommonParameters;
+		PassParameters->History = History;
+		PassParameters->PrevHistory = PrevHistory;
+		PassParameters->DebugOutput = CreateDebugUAV(HistoryExtent * kHistoryUpscalingFactor, TEXT("Debug.TAA.History"));
+
+		TShaderMapRef<FTAADebugHistoryCS> ComputeShader(View.ShaderMap);
+		FComputeShaderUtils::AddPass(
+			GraphBuilder,
+			RDG_EVENT_NAME("TAA DebugHistory %dx%d", RejectionRect.Width(), RejectionRect.Height()),
+			ComputeShader,
+			PassParameters,
+			FComputeShaderUtils::GetGroupCount(HistorySize * kHistoryUpscalingFactor, 8));
+	}
+	#endif
+
 	if (!View.bStatePrevViewInfoIsReadOnly)
 	{
 		OutputHistory->SafeRelease();
 
-		check(kHistoryTextures <= OutputHistory->RT.Num());
-		for (int32 i = 0; i < kHistoryTextures; i++)
+		for (int32 i = 0; i < History.LowResTextures.Num(); i++)
+		{
+			if (ExtractLowResHistoryTexture[i])
+			{
+				GraphBuilder.QueueTextureExtraction(History.LowResTextures[i], &OutputHistory->LowResRT[i]);
+			}
+		}
+
+		for (int32 i = 0; i < History.Textures.Num(); i++)
 		{
-			if (ExtractHistory[i])
+			if (ExtractHistoryTexture[i])
 			{
 				GraphBuilder.QueueTextureExtraction(History.Textures[i], &OutputHistory->RT[i]);
 			}
diff --git a/Engine/Source/Runtime/Renderer/Private/SceneRendering.h b/Engine/Source/Runtime/Renderer/Private/SceneRendering.h
--- a/Engine/Source/Runtime/Renderer/Private/SceneRendering.h
+++ b/Engine/Source/Runtime/Renderer/Private/SceneRendering.h
@@ -810,22 +810,31 @@ END_GLOBAL_SHADER_PARAMETER_STRUCT()
 struct FTemporalAAHistory
 {
 	// Number of render target in the history.
-	static constexpr uint32 kRenderTargetCount = 4;
+	static constexpr int32 kLowResRenderTargetCount = 1;
+
+	// Number of render target in the history.
+	static constexpr int32 kRenderTargetCount = 4;
 
 	// Render targets holding's pixel history.
-	//  scene color's RGBA are in RT[0].
+	//  scene color's RGBA are in OutputRT[0].
+	TStaticArray<TRefCountPtr<IPooledRenderTarget>, kLowResRenderTargetCount> LowResRT;
 	TStaticArray<TRefCountPtr<IPooledRenderTarget>, kRenderTargetCount> RT;
 
 	// Reference size of RT. Might be different than RT's actual size to handle down res.
 	FIntPoint ReferenceBufferSize;
 
 	// Viewport coordinate of the history in RT according to ReferenceBufferSize.
+	FIntRect LowResViewportRect;
 	FIntRect ViewportRect;
 
 
 	void SafeRelease()
 	{
-		for (uint32 i = 0; i < kRenderTargetCount; i++)
+		for (int32 i = 0; i < LowResRT.Num(); i++)
+		{
+			LowResRT[i].SafeRelease();
+		}
+		for (int32 i = 0; i < RT.Num(); i++)
 		{
 			RT[i].SafeRelease();
 		}
diff --git a/Engine/Source/Runtime/Renderer/Private/SceneVisibility.cpp b/Engine/Source/Runtime/Renderer/Private/SceneVisibility.cpp
--- a/Engine/Source/Runtime/Renderer/Private/SceneVisibility.cpp
+++ b/Engine/Source/Runtime/Renderer/Private/SceneVisibility.cpp
@@ -243,6 +243,11 @@ static FAutoConsoleVariableRef CVarLightMaxDrawDistanceScale(
 
 #if !UE_BUILD_SHIPPING
 
+static TAutoConsoleVariable<int32> CVarTAADebugOverrideTemporalIndex(
+	TEXT("r.TemporalAA.Debug.OverrideTemporalIndex"), -1,
+	TEXT("Override the temporal index for debugging purposes."),
+	ECVF_RenderThreadSafe);
+
 static TAutoConsoleVariable<float> CVarFreezeTemporalSequences(
 	TEXT("r.Test.FreezeTemporalSequences"), 0,
 	TEXT("Freezes all temporal sequences."),
@@ -3263,7 +3268,7 @@ void FSceneRenderer::PreVisibilityFrameSetup(FRDGBuilder& GraphBuilder, const FS
 				{
 					// When doing TAA upsample with screen percentage < 100%, we need extra temporal samples to have a
 					// constant temporal sample density for final output pixels to avoid output pixel aligned converging issues.
-					TemporalAASamples = float(TemporalAASamples) * FMath::Max(1.f, 1.f / (EffectivePrimaryResolutionFraction * EffectivePrimaryResolutionFraction));
+					TemporalAASamples = FMath::RoundToInt(float(TemporalAASamples) * FMath::Max(1.f, 1.f / (EffectivePrimaryResolutionFraction * EffectivePrimaryResolutionFraction)));
 				}
 				else if (CVarTemporalAASamplesValue == 5)
 				{
@@ -3275,11 +3280,18 @@ void FSceneRenderer::PreVisibilityFrameSetup(FRDGBuilder& GraphBuilder, const FS
 
 			// Compute the new sample index in the temporal sequence.
 			int32 TemporalSampleIndex			= ViewState->TemporalAASampleIndex + 1;
-			if(TemporalSampleIndex >= TemporalAASamples || View.bCameraCut)
+			if (TemporalSampleIndex >= TemporalAASamples || View.bCameraCut)
 			{
 				TemporalSampleIndex = 0;
 			}
 
+			#if !UE_BUILD_SHIPPING
+			if (CVarTAADebugOverrideTemporalIndex.GetValueOnRenderThread() >= 0)
+			{
+				TemporalSampleIndex = CVarTAADebugOverrideTemporalIndex.GetValueOnRenderThread();
+			}
+			#endif
+
 			// Updates view state.
 			if (!View.bStatePrevViewInfoIsReadOnly && !bFreezeTemporalSequences)
 			{
