Missing include

diff --git a/Engine/Source/Runtime/Renderer/Private/VT/AllocatedVirtualTexture.cpp b/Engine/Source/Runtime/Renderer/Private/VT/AllocatedVirtualTexture.cpp
--- a/Engine/Source/Runtime/Renderer/Private/VT/AllocatedVirtualTexture.cpp
+++ b/Engine/Source/Runtime/Renderer/Private/VT/AllocatedVirtualTexture.cpp
@@ -6,6 +6,7 @@
 #include "VT/VirtualTextureSystem.h"
 #include "VT/VirtualTextureSpace.h"
 #include "VT/VirtualTexturePhysicalSpace.h"
+#include "Misc/StringBuilder.h"
 
 FAllocatedVirtualTexture::FAllocatedVirtualTexture(FVirtualTextureSystem* InSystem,
 	uint32 InFrame,
