ShowFlag support in path tracer

Add support for light types and direct vs. gi flags

diff --git a/Engine/Shaders/Private/PathTracing/PathTracing.usf b/Engine/Shaders/Private/PathTracing/PathTracing.usf
--- a/Engine/Shaders/Private/PathTracing/PathTracing.usf
+++ b/Engine/Shaders/Private/PathTracing/PathTracing.usf
@@ -40,7 +40,7 @@ void AccumulateRadiance(inout float3 TotalRadiance, float3 PathRadiance)
 	}
 }
 
-FMaterialClosestHitPayload TraceTransparentRay(RayDesc Ray, FRayCone RayCone, bool IsCameraRay, bool LastBounce, uint2 LaunchIndex, uint NumLights, inout RandomSequence RandSequence, inout float3 PathThroughput, inout float3 Radiance)
+FMaterialClosestHitPayload TraceTransparentRay(RayDesc Ray, FRayCone RayCone, bool IsCameraRay, bool LastBounce, bool IncludeEmission, uint2 LaunchIndex, uint NumLights, inout RandomSequence RandSequence, inout float3 PathThroughput, inout float3 Radiance)
 {
 	float SelectionWeightSum = 0;
 	float3 PayloadThroughput = PathThroughput;
@@ -95,7 +95,7 @@ FMaterialClosestHitPayload TraceTransparentRay(RayDesc Ray, FRayCone RayCone, bo
 		}
 
 		// add in surface emission (except for modulate blend mode which uses Radiance as the transparency multiplier)
-		if (HitPayload.BlendingMode != RAY_TRACING_BLEND_MODE_MODULATE)
+		if (IncludeEmission && HitPayload.BlendingMode != RAY_TRACING_BLEND_MODE_MODULATE)
 		{
 			AccumulateRadiance(Radiance, PathThroughput * HitPayload.Radiance);
 		}
@@ -267,7 +267,7 @@ RAY_TRACING_ENTRY_RAYGEN(PathTracingMainRG)
 
 	for (int Bounce = 0; Bounce <= PathTracingData.MaxBounces; Bounce++)
 	{
-		FMaterialClosestHitPayload Payload = TraceTransparentRay(Ray, RayCone, Bounce == 0, Bounce == PathTracingData.MaxBounces, LaunchIndex, NumLights, RandSequence, PathThroughput, Radiance);
+		FMaterialClosestHitPayload Payload = TraceTransparentRay(Ray, RayCone, Bounce == 0, Bounce == PathTracingData.MaxBounces, PathTracingData.SkipDirectLighting == 0 || Bounce > 1, LaunchIndex, NumLights, RandSequence, PathThroughput, Radiance);
 
 		if (Payload.IsMiss())
 		{
@@ -299,10 +299,12 @@ RAY_TRACING_ENTRY_RAYGEN(PathTracingMainRG)
 		float4 RandSample = RandomSequence_GenerateSample4D(RandSequence);
 
 		// Does this material require NEE? (will be false if MaterialPdf is always +inf)
-		bool bIsNeeValid = IsNeeValidMaterial(Payload);
+		const bool bIsNeeValid = IsNeeValidMaterial(Payload);
+
+		const bool bDoLightLoop = PathTracingData.SkipDirectLighting == 0 || Bounce > 0;
 
 		// If we are using Light sampling and the material can use it ...
-		if (MISMode != 0 && bIsNeeValid && SceneLightCount > 0)
+		if (MISMode != 0 && bIsNeeValid && SceneLightCount > 0 && bDoLightLoop)
 		{
 			// Choose a light and sample it
 			float3 WorldPos = Payload.WorldPos;
@@ -414,7 +416,7 @@ RAY_TRACING_ENTRY_RAYGEN(PathTracingMainRG)
 		ApplyPositionBias(Ray, MaterialSample.PositionBiasSign * Payload.WorldNormal, PathTracingData.MaxNormalBias);
 
 		// If we are using Material sampling for lights
-		if (MISMode != 1)
+		if (MISMode != 1 && bDoLightLoop)
 		{
 			// Check which lights can be seen by the material ray and trace a dedicated shadow ray
 			// While it would be possible to just loop around and use the indirect ray to do this, it would prevent the application
diff --git a/Engine/Source/Runtime/Renderer/Private/PathTracing.cpp b/Engine/Source/Runtime/Renderer/Private/PathTracing.cpp
--- a/Engine/Source/Runtime/Renderer/Private/PathTracing.cpp
+++ b/Engine/Source/Runtime/Renderer/Private/PathTracing.cpp
@@ -123,12 +123,35 @@ IMPLEMENT_GLOBAL_SHADER_PARAMETER_STRUCT(FPathTracingData, "PathTracingData");
 
 // This function prepares the portion of shader arguments that may involve invalidating the path traced state
 static bool PrepareShaderArgs(const FViewInfo& View, FPathTracingData& PathTracingData) {
-	int32 PathTracingMaxBounces = CVarPathTracingMaxBounces.GetValueOnRenderThread();
-	if (PathTracingMaxBounces < 0)
+	PathTracingData.SkipDirectLighting = false;
+	int32 MaxBounces = CVarPathTracingMaxBounces.GetValueOnRenderThread();
+	if (MaxBounces < 0)
 	{
-		PathTracingMaxBounces = View.FinalPostProcessSettings.PathTracingMaxBounces;
+		MaxBounces = View.FinalPostProcessSettings.PathTracingMaxBounces;
 	}
-	PathTracingData.MaxBounces = PathTracingMaxBounces;
+	if (View.Family->EngineShowFlags.DirectLighting)
+	{
+		if (!View.Family->EngineShowFlags.GlobalIllumination)
+		{
+			// direct lighting, but no GI
+			MaxBounces = 1;
+		}
+	}
+	else
+	{
+		if (View.Family->EngineShowFlags.GlobalIllumination)
+		{
+			// skip direct lighting, but still do the full bounces
+			PathTracingData.SkipDirectLighting = true;
+		}
+		else
+		{
+			// neither direct, nor GI is on
+			MaxBounces = 0;
+		}
+	}
+
+	PathTracingData.MaxBounces = MaxBounces;
 	PathTracingData.MaxNormalBias = GetRaytracingMaxNormalBias();
 	PathTracingData.MISMode = CVarPathTracingMISMode.GetValueOnRenderThread();
 	PathTracingData.VisibleLights = CVarPathTracingVisibleLights.GetValueOnRenderThread();
@@ -148,7 +171,7 @@ static bool PrepareShaderArgs(const FViewInfo& View, FPathTracingData& PathTraci
 	// If any of the parameters above changed since last time -- reset the accumulation
 	// FIXME: find something cleaner than just using static variables here. Should all
 	// the state used for comparison go into ViewState?
-	static uint32 PrevMaxBounces = PathTracingMaxBounces;
+	static uint32 PrevMaxBounces = PathTracingData.MaxBounces;
 	if (PathTracingData.MaxBounces != PrevMaxBounces)
 	{
 		NeedInvalidation = true;
@@ -211,6 +234,14 @@ static bool PrepareShaderArgs(const FViewInfo& View, FPathTracingData& PathTraci
 		PreviousBackfaceCulling = PathTracingData.EnableCameraBackfaceCulling;
 	}
 
+	// Changing direct lighting skipping requires starting over
+	static uint32 PreviousSkipDirectLighting = PathTracingData.SkipDirectLighting;
+	if (PreviousSkipDirectLighting != PathTracingData.SkipDirectLighting)
+	{
+		NeedInvalidation = true;
+		PreviousSkipDirectLighting = PathTracingData.SkipDirectLighting;
+	}
+
 	// the rest of PathTracingData and AdaptiveSamplingData is filled in by SetParameters below
 	return NeedInvalidation;
 }
@@ -269,7 +300,7 @@ class FPathTracingIESAtlasCS : public FGlobalShader
 };
 IMPLEMENT_SHADER_TYPE(, FPathTracingIESAtlasCS, TEXT("/Engine/Private/PathTracing/PathTracingIESAtlas.usf"), TEXT("PathTracingIESAtlasCS"), SF_Compute);
 
-void SetLightParameters(FRDGBuilder& GraphBuilder, FPathTracingRG::FParameters* PassParameters, FSkyLightData& SkyLightData, FScene* Scene, const FViewInfo& View, bool UseLightProfiles)
+void SetLightParameters(FRDGBuilder& GraphBuilder, FPathTracingRG::FParameters* PassParameters, FSkyLightData& SkyLightData, FScene* Scene, const FViewInfo& View)
 {
 	// Sky light
 	bool IsSkyLightValid = SetupSkyLightParameters(*Scene, &SkyLightData);
@@ -279,7 +310,7 @@ void SetLightParameters(FRDGBuilder& GraphBuilder, FPathTracingRG::FParameters*
 	unsigned LightCount = 0;
 
 	// Prepend SkyLight to light buffer since it is not part of the regular light list
-	if (IsSkyLightValid)
+	if (IsSkyLightValid && View.Family->EngineShowFlags.SkyLighting)
 	{
 		FPathTracingLight& DestLight = Lights[LightCount];
 		DestLight.Color = FVector(SkyLightData.Color);
@@ -296,6 +327,17 @@ void SetLightParameters(FRDGBuilder& GraphBuilder, FPathTracingRG::FParameters*
 	{
 		if (LightCount >= RAY_TRACING_LIGHT_COUNT_MAXIMUM) break;
 
+		ELightComponentType LightComponentType = (ELightComponentType)Light.LightSceneInfo->Proxy->GetLightType();
+
+		if (((LightComponentType == LightType_Directional) && !View.Family->EngineShowFlags.DirectionalLights) ||
+			((LightComponentType == LightType_Rect       ) && !View.Family->EngineShowFlags.RectLights       ) ||
+			((LightComponentType == LightType_Spot       ) && !View.Family->EngineShowFlags.SpotLights       ) ||
+			((LightComponentType == LightType_Point      ) && !View.Family->EngineShowFlags.PointLights      ))
+		{
+			// This light type is not currently enabled
+			continue;
+		}
+
 		FPathTracingLight& DestLight = Lights[LightCount];
 
 		FLightShaderParameters LightParameters;
@@ -308,7 +350,7 @@ void SetLightParameters(FRDGBuilder& GraphBuilder, FPathTracingRG::FParameters*
 		DestLight.Flags |= Light.LightSceneInfo->Proxy->CastsDynamicShadow() ? PATHTRACER_FLAG_CAST_SHADOW_MASK : 0;
 		DestLight.IESTextureSlice = -1;
 
-		if (UseLightProfiles)
+		if (View.Family->EngineShowFlags.TexturedLightProfiles)
 		{
 			FTexture* IESTexture = Light.LightSceneInfo->Proxy->GetIESTextureResource();
 			if (IESTexture != nullptr)
@@ -327,7 +369,6 @@ void SetLightParameters(FRDGBuilder& GraphBuilder, FPathTracingRG::FParameters*
 		DestLight.Attenuation = LightParameters.InvRadius;
 		DestLight.FalloffExponent = 0;
 
-		ELightComponentType LightComponentType = (ELightComponentType)Light.LightSceneInfo->Proxy->GetLightType();
 		switch (LightComponentType)
 		{
 			case LightType_Directional:
@@ -504,11 +545,19 @@ void FDeferredShadingSceneRenderer::RenderPathTracing(
 		bArgsChanged = true;
 	}
 
-	bool UseLightProfiles = View.Family->EngineShowFlags.TexturedLightProfiles;
-	static bool PreviousUseLightProfiles = UseLightProfiles;
-	if (PreviousUseLightProfiles != UseLightProfiles)
+	// compute an integer code of what show flags related to lights are currently enabled so we can detect changes
+	int CurrentLightShowFlags = 0;
+	CurrentLightShowFlags |= View.Family->EngineShowFlags.SkyLighting           ? 1 << 0 : 0;
+	CurrentLightShowFlags |= View.Family->EngineShowFlags.DirectionalLights     ? 1 << 1 : 0;
+	CurrentLightShowFlags |= View.Family->EngineShowFlags.RectLights            ? 1 << 2 : 0;
+	CurrentLightShowFlags |= View.Family->EngineShowFlags.SpotLights            ? 1 << 3 : 0;
+	CurrentLightShowFlags |= View.Family->EngineShowFlags.PointLights           ? 1 << 4 : 0;
+	CurrentLightShowFlags |= View.Family->EngineShowFlags.TexturedLightProfiles ? 1 << 5 : 0;
+
+	static int PreviousLightShowFlags = CurrentLightShowFlags;
+	if (PreviousLightShowFlags != CurrentLightShowFlags)
 	{
-		PreviousUseLightProfiles = UseLightProfiles;
+		PreviousLightShowFlags = CurrentLightShowFlags;
 		bArgsChanged = true;
 	}
 
@@ -566,7 +615,7 @@ void FDeferredShadingSceneRenderer::RenderPathTracing(
 		{
 			// upload sky/lights data
 			FSkyLightData SkyLightData;
-			SetLightParameters(GraphBuilder, PassParameters, SkyLightData, Scene, View, UseLightProfiles);
+			SetLightParameters(GraphBuilder, PassParameters, SkyLightData, Scene, View);
 			PassParameters->SkyLightData = CreateUniformBufferImmediate(SkyLightData, EUniformBufferUsage::UniformBuffer_SingleFrame);
 		}
 		PassParameters->IESTextureSampler = TStaticSamplerState<SF_Bilinear, AM_Clamp, AM_Clamp, AM_Clamp>::GetRHI();
diff --git a/Engine/Source/Runtime/Renderer/Private/PathTracingUniformBuffers.h b/Engine/Source/Runtime/Renderer/Private/PathTracingUniformBuffers.h
--- a/Engine/Source/Runtime/Renderer/Private/PathTracingUniformBuffers.h
+++ b/Engine/Source/Runtime/Renderer/Private/PathTracingUniformBuffers.h
@@ -19,6 +19,7 @@ BEGIN_GLOBAL_SHADER_PARAMETER_STRUCT(FPathTracingData, )
 	SHADER_PARAMETER(uint32, VisibleLights)
 	SHADER_PARAMETER(uint32, ApproximateCaustics)
 	SHADER_PARAMETER(uint32, EnableCameraBackfaceCulling)
+	SHADER_PARAMETER(uint32, SkipDirectLighting)
 	SHADER_PARAMETER(float, MaxPathIntensity)
 	SHADER_PARAMETER(float, MaxNormalBias)
 	SHADER_PARAMETER(float, FilterWidth)
