VSM atomic depth writes

diff --git a/Engine/Shaders/Private/ShadowDepthPixelShader.usf b/Engine/Shaders/Private/ShadowDepthPixelShader.usf
--- a/Engine/Shaders/Private/ShadowDepthPixelShader.usf
+++ b/Engine/Shaders/Private/ShadowDepthPixelShader.usf
@@ -24,10 +24,19 @@
 #include "/Engine/Generated/VertexFactory.ush"
 #include "ShadowDepthCommon.ush"
 
+#if ENABLE_NON_NANITE_VSM
+#include "Nanite/NaniteDataDecode.ush"
+#include "SceneData.ush"
+#include "VirtualShadowMaps/PageAccessCommon.ush"
+#endif
+
 #define SECONDARY_OCCLUSION			1
 
 void Main( 
 	FShadowDepthVSToPS Inputs,
+#if VIRTUAL_TEXTURE_TARGET
+	nointerpolation uint ViewId : TEXCOORD8,
+#endif
 	in float4 SvPosition : SV_Position		// after all interpolators
 #if PERSPECTIVE_CORRECT_DEPTH || COMPILER_METAL
 	,out float OutDepth : SV_DEPTH
@@ -70,4 +79,24 @@ void Main(
 	// so output depth value explicitly if this shader permuation was not discarded
 	OutDepth = SvPosition.z;
 #endif
+
+#if VIRTUAL_TEXTURE_TARGET
+	uint2 vAddress = (uint2)SvPosition.xy;
+	float DeviceZ = SvPosition.z;
+
+#if !NEAR_CLIP
+	//DeviceZ = max(0.0f, DeviceZ);
+#endif
+
+	FNaniteView NaniteView = UnpackNaniteView( PassStruct.PackedNaniteViews[ ViewId ] );
+
+	FShadowPhysicalPage Page = ShadowDecodePageTable( PassStruct.VirtualSmPageTable[ CalcPageOffset( NaniteView.TargetLayerIndex, NaniteView.TargetMipLevel, vAddress >> VSM_LOG2_PAGE_SIZE ) ] );
+
+	if( Page.bThisLODValid )
+	{
+		uint2 pAddress = Page.PageIndex * VSM_PAGE_SIZE + (vAddress & VSM_PAGE_SIZE_MASK);
+
+		InterlockedMax( PassStruct.OutDepthBuffer[ pAddress ], asuint( DeviceZ ) );
+	}
+#endif
 }
diff --git a/Engine/Shaders/Private/ShadowDepthVertexShader.usf b/Engine/Shaders/Private/ShadowDepthVertexShader.usf
--- a/Engine/Shaders/Private/ShadowDepthVertexShader.usf
+++ b/Engine/Shaders/Private/ShadowDepthVertexShader.usf
@@ -14,10 +14,10 @@
 
 
 #if ENABLE_NON_NANITE_VSM
-	#define VIRTUAL_SM_ENABLED (!(USING_TESSELLATION || ONEPASS_POINTLIGHT_SHADOW))
-#else //!ENABLE_NON_NANITE_VSM
+	#define VIRTUAL_SM_ENABLED (!(USING_TESSELLATION || ONEPASS_POINTLIGHT_SHADOW || PERSPECTIVE_CORRECT_DEPTH))
+#else
 	#define VIRTUAL_SM_ENABLED 0
-#endif // ENABLE_NON_NANITE_VSM
+#endif
 
 // Need to get the definition of nanite packed view for generated unform buffer (which is pulled in by Common.ush)
 #include "Nanite/PackedNaniteView.ush"
@@ -123,16 +123,29 @@ FPageInfo UnPackPageInfo(uint PackedData)
 void ScaleBiasClipToPhysicalSmPage(FNaniteView NaniteView, inout float4 PositionClip, inout float4 ClipPlanesInOut, FPageInfo PageInfo)
 {
 	uint2 vPage = PageInfo.VirtualPage;
-	uint2 pPage = ShadowDecodePageTable( PassStruct.VirtualSmPageTable[ CalcPageOffset( NaniteView.TargetLayerIndex, NaniteView.TargetMipLevel, vPage ) ] ).PageIndex;
 
 	float2 vClip = PositionClip.xy;
 	float2 vUV = vClip * float2(0.5, -0.5) + 0.5 * PositionClip.w;
 	float2 vPixels = vUV * ( VSM_VIRTUAL_MAX_RESOLUTION_XY >> NaniteView.TargetMipLevel );
-	float2 pPixels = vPixels + ( (float2)pPage - (float2)vPage ) * VSM_PAGE_SIZE * PositionClip.w;
-	float2 pUV = pPixels * VirtualShadowMap.RecPhysicalPoolSize.xy;
-	float2 pClip = pUV * float2(2, -2) + float2(-1, 1) * PositionClip.w;
 
-	PositionClip.xy = pClip;
+	if( PassStruct.bAtomicWrites )
+	{
+		//PointClip.xy = NaniteView.ClipSpaceScaleOffset.xy * PointClip.xy + NaniteView.ClipSpaceScaleOffset.zw * PointClip.w;
+
+		float2 Scale = exp2( -NaniteView.TargetMipLevel );
+		PositionClip.xy *= Scale;
+		PositionClip.xy += ( Scale * float2(1,-1) + float2(-1,1) ) * PositionClip.w;
+	}
+	else
+	{
+		uint2 pPage = ShadowDecodePageTable( PassStruct.VirtualSmPageTable[ CalcPageOffset( NaniteView.TargetLayerIndex, NaniteView.TargetMipLevel, vPage ) ] ).PageIndex;
+
+		float2 pPixels = vPixels + ( (float2)pPage - (float2)vPage ) * VSM_PAGE_SIZE * PositionClip.w;
+		float2 pUV = pPixels * VirtualShadowMap.RecPhysicalPoolSize.xy;
+		float2 pClip = pUV * float2(2, -2) + float2(-1, 1) * PositionClip.w;
+
+		PositionClip.xy = pClip;
+	}
 
 	if( PassStruct.bInstancePerPage )
 	{
@@ -155,23 +168,23 @@ void ScaleBiasClipToPhysicalSmPage(FNaniteView NaniteView, inout float4 Position
 	}
 }
 
-void TransformToVirtualSmPage(inout float4 PosCsHInOut, inout float4 ClipPlanesInOut, FPageInfo PageInfo, float3 TranslatedWorldPos)
+void TransformToVirtualSmPage(inout float4 PointClip, inout float4 ClipPlanesInOut, FPageInfo PageInfo, float3 PointTranslatedWorld)
 {
 	// Load nanite view and perform transform to SM space.
 	FNaniteView NaniteView = UnpackNaniteView(PassStruct.PackedNaniteViews[PageInfo.ViewId]);
-	PosCsHInOut = mul(float4(TranslatedWorldPos - ResolvedView.PreViewTranslation.xyz + NaniteView.PreViewTranslation.xyz, 1.0f), NaniteView.TranslatedWorldToClip);
 
+	PointTranslatedWorld += ( -ResolvedView.PreViewTranslation.xyz + NaniteView.PreViewTranslation.xyz );
+	PointClip = mul( float4( PointTranslatedWorld, 1 ), NaniteView.TranslatedWorldToClip );
 
 	// Clamp the vertex to the near plane if it is in front of the near plane
 	// This has problems if some vertices of a triangle get clamped and others do not, also causes artifacts with non-ortho projections
-	if (PassStruct.bClampToNearPlane > 0 && PosCsHInOut.z > PosCsHInOut.w)
+	if (PassStruct.bClampToNearPlane > 0 && PointClip.z > PointClip.w)
 	{
-		PosCsHInOut.z = 0.999999f;
-		PosCsHInOut.w = 1.0f;
+		PointClip.z = 0.999999f;
+		PointClip.w = 1.0f;
 	}
 
-	// 
-	ScaleBiasClipToPhysicalSmPage(NaniteView, PosCsHInOut, ClipPlanesInOut, PageInfo);
+	ScaleBiasClipToPhysicalSmPage(NaniteView, PointClip, ClipPlanesInOut, PageInfo);
 }
 
 #endif // VIRTUAL_SM_ENABLED
@@ -389,6 +402,9 @@ void Main(
 #endif
 #else
 	out FShadowDepthVSToPS OutParameters,
+#endif
+#if VIRTUAL_SM_ENABLED
+	out nointerpolation uint ViewId : TEXCOORD8,
 #endif
 	OPTIONAL_VertexID
 	out float4 OutPosition : SV_POSITION
@@ -486,6 +502,8 @@ void Main(
 #endif
 
 #if VIRTUAL_SM_ENABLED
+	ViewId = 0;
+	
 	OutVirtualSmPageClip = float4(1.0f, 1.0f, 1.0f, 1.0f);
 	if (PassStruct.bRenderToVirtualShadowMap != 0)
 	{
@@ -494,6 +512,8 @@ void Main(
 		FPageInfo PageInfo = UnPackPageInfo(View.PageInfoBuffer[InstanceIdIndex]);
 
 		TransformToVirtualSmPage(OutPosition, OutVirtualSmPageClip, PageInfo, WorldPos.xyz);
+		
+		ViewId = PageInfo.ViewId;
 	}
 #endif // VIRTUAL_SM_ENABLED
 
@@ -551,6 +571,9 @@ void PositionOnlyMain(
 	out FShadowDepthVSToGS OutParameters,
 #else
 	out FShadowDepthVSToPS OutParameters,
+#endif
+#if VIRTUAL_SM_ENABLED
+	out nointerpolation uint ViewId : TEXCOORD8,
 #endif
 	out float4 OutPosition : SV_POSITION
 #if VIRTUAL_SM_ENABLED	
@@ -614,6 +637,8 @@ void PositionOnlyMain(
 #endif
 
 #if VIRTUAL_SM_ENABLED
+	ViewId = 0;
+
 	OutVirtualSmPageClip = float4(1.0f, 1.0f, 1.0f, 1.0f);
 	if (PassStruct.bRenderToVirtualShadowMap != 0)
 	{
@@ -623,6 +648,8 @@ void PositionOnlyMain(
 		//       Or perhaps both are different, as they are managed somewhat differently anyway, maybe.
 		FPageInfo PageInfo = UnPackPageInfo(View.PageInfoBuffer[InstanceIdIndex]);
 		TransformToVirtualSmPage(OutPosition, OutVirtualSmPageClip, PageInfo, WorldPos.xyz);
+
+		ViewId = PageInfo.ViewId;
 	}
 #endif // VIRTUAL_SM_ENABLED
 }
diff --git a/Engine/Source/Runtime/Renderer/Private/ShadowDepthRendering.cpp b/Engine/Source/Runtime/Renderer/Private/ShadowDepthRendering.cpp
--- a/Engine/Source/Runtime/Renderer/Private/ShadowDepthRendering.cpp
+++ b/Engine/Source/Runtime/Renderer/Private/ShadowDepthRendering.cpp
@@ -117,7 +117,7 @@ static FAutoConsoleVariableRef CVarShadowShadowUseGS(
 	TEXT("Use geometry shaders to render cube map shadows."),
 	ECVF_RenderThreadSafe);
 
-
+extern int32 GVirtualShadowMapAtomicWrites;
 extern int32 GNaniteDebugFlags;
 extern int32 GNaniteShowStats;
 
@@ -164,9 +164,14 @@ void SetupShadowDepthPassUniformBuffer(
 #if ENABLE_NON_NANITE_VSM
 	ShadowDepthPassParameters.bRenderToVirtualShadowMap = false;
 	ShadowDepthPassParameters.bInstancePerPage = false;
+	ShadowDepthPassParameters.bAtomicWrites = false;
 	ShadowDepthPassParameters.VirtualSmPageTable = GraphBuilder.CreateSRV(CreateStructuredBuffer(GraphBuilder, TEXT("Dummy-VirtualSmPageTable"), TArray<uint32>()));
 	ShadowDepthPassParameters.PackedNaniteViews = GraphBuilder.CreateSRV(CreateStructuredBuffer(GraphBuilder, TEXT("Dummy-PackedNaniteViews"), TArray<Nanite::FPackedView>()));
 	ShadowDepthPassParameters.PageRectBounds = GraphBuilder.CreateSRV(CreateStructuredBuffer(GraphBuilder, TEXT("Dummy-PageRectBounds"), TArray<FIntVector4>()));
+
+	FRDGTextureRef DepthBuffer = GraphBuilder.CreateTexture( FRDGTextureDesc::Create2D( FIntPoint(4,4), PF_R32_UINT, FClearValueBinding::None, TexCreate_ShaderResource | TexCreate_UAV ), TEXT("Dummy-OutDepthBuffer") );
+
+	ShadowDepthPassParameters.OutDepthBuffer = GraphBuilder.CreateUAV( DepthBuffer );
 #endif // ENABLE_NON_NANITE_VSM
 }
 
@@ -554,7 +559,8 @@ enum EShadowDepthPixelShaderMode
 {
 	PixelShadowDepth_NonPerspectiveCorrect,
 	PixelShadowDepth_PerspectiveCorrect,
-	PixelShadowDepth_OnePassPointLight
+	PixelShadowDepth_OnePassPointLight,
+	PixelShadowDepth_VirtualShadowMap
 };
 
 template <EShadowDepthPixelShaderMode ShaderMode>
@@ -566,32 +572,29 @@ public:
 	static bool ShouldCompilePermutation(const FMeshMaterialShaderPermutationParameters& Parameters)
 	{
 		const EShaderPlatform Platform = Parameters.Platform;
+		
+		// Only compile one pass point light shaders for feature levels >= SM5
+		if (ShaderMode == PixelShadowDepth_OnePassPointLight && !IsFeatureLevelSupported(Platform, ERHIFeatureLevel::SM5))
+		{
+			return false;
+		}
 
-		if (!IsFeatureLevelSupported(Platform, ERHIFeatureLevel::SM5))
-		{
-			return (Parameters.MaterialParameters.bIsSpecialEngineMaterial
-				// Only compile for masked or lit translucent materials
-				|| !Parameters.MaterialParameters.bWritesEveryPixelShadowPass
-				|| (Parameters.MaterialParameters.bMaterialMayModifyMeshPosition && Parameters.MaterialParameters.bIsUsedWithInstancedStaticMeshes)
-				// Perspective correct rendering needs a pixel shader and WPO materials can't be overridden with default material.
-				|| (ShaderMode == PixelShadowDepth_PerspectiveCorrect && Parameters.MaterialParameters.bMaterialMayModifyMeshPosition))
-				&& ShaderMode != PixelShadowDepth_OnePassPointLight
-				// Don't render ShadowDepth for translucent unlit materials
-				&& Parameters.MaterialParameters.bShouldCastDynamicShadows;
+		if (ShaderMode == PixelShadowDepth_VirtualShadowMap && !IsFeatureLevelSupported(Platform, ERHIFeatureLevel::SM5))
+		{
+			return false;
 		}
 
+		bool bModeRequiresPS = ShaderMode == PixelShadowDepth_PerspectiveCorrect || ShaderMode == PixelShadowDepth_VirtualShadowMap;
+
 		//Note: This logic needs to stay in sync with OverrideWithDefaultMaterialForShadowDepth!
 		return (Parameters.MaterialParameters.bIsSpecialEngineMaterial
 			// Only compile for masked or lit translucent materials
 			|| !Parameters.MaterialParameters.bWritesEveryPixelShadowPass
 			|| (Parameters.MaterialParameters.bMaterialMayModifyMeshPosition && Parameters.MaterialParameters.bIsUsedWithInstancedStaticMeshes)
-			// Perspective correct rendering needs a pixel shader and WPO materials can't be overridden with default material.
-			|| (ShaderMode == PixelShadowDepth_PerspectiveCorrect && Parameters.MaterialParameters.bMaterialMayModifyMeshPosition))
-			// Only compile one pass point light shaders for feature levels >= SM5
-			&& (ShaderMode != PixelShadowDepth_OnePassPointLight || IsFeatureLevelSupported(Platform, ERHIFeatureLevel::SM5))
+			// This mode needs a pixel shader and WPO materials can't be overridden with default material.
+			|| (bModeRequiresPS && Parameters.MaterialParameters.bMaterialMayModifyMeshPosition))
 			// Don't render ShadowDepth for translucent unlit materials
-			&& Parameters.MaterialParameters.bShouldCastDynamicShadows
-			&& IsFeatureLevelSupported(Platform, ERHIFeatureLevel::SM5);
+			&& Parameters.MaterialParameters.bShouldCastDynamicShadows;
 	}
 
 	static void ModifyCompilationEnvironment(const FMaterialShaderPermutationParameters& Parameters, FShaderCompilerEnvironment& OutEnvironment)
@@ -599,6 +602,16 @@ public:
 		FShadowDepthBasePS::ModifyCompilationEnvironment(Parameters, OutEnvironment);
 		OutEnvironment.SetDefine(TEXT("PERSPECTIVE_CORRECT_DEPTH"), (uint32)(ShaderMode == PixelShadowDepth_PerspectiveCorrect));
 		OutEnvironment.SetDefine(TEXT("ONEPASS_POINTLIGHT_SHADOW"), (uint32)(ShaderMode == PixelShadowDepth_OnePassPointLight));
+		OutEnvironment.SetDefine(TEXT("VIRTUAL_TEXTURE_TARGET"), (uint32)(ShaderMode == PixelShadowDepth_VirtualShadowMap));
+
+		uint32 bEnableNonNaniteVSM = (uint32)(ENABLE_NON_NANITE_VSM != 0 && UseGPUScene(Parameters.Platform));
+		OutEnvironment.SetDefine(TEXT("ENABLE_NON_NANITE_VSM"), bEnableNonNaniteVSM);
+#if ENABLE_NON_NANITE_VSM
+		if (bEnableNonNaniteVSM != 0)
+		{
+			FVirtualShadowMapArray::SetShaderDefines(OutEnvironment);
+		}
+#endif
 	}
 
 	TShadowDepthPS()
@@ -619,6 +632,9 @@ public:
 IMPLEMENT_SHADOWDEPTHPASS_PIXELSHADER_TYPE(PixelShadowDepth_NonPerspectiveCorrect);
 IMPLEMENT_SHADOWDEPTHPASS_PIXELSHADER_TYPE(PixelShadowDepth_PerspectiveCorrect);
 IMPLEMENT_SHADOWDEPTHPASS_PIXELSHADER_TYPE(PixelShadowDepth_OnePassPointLight);
+#if ENABLE_NON_NANITE_VSM
+IMPLEMENT_SHADOWDEPTHPASS_PIXELSHADER_TYPE(PixelShadowDepth_VirtualShadowMap);
+#endif
 
 /**
 * Overrides a material used for shadow depth rendering with the default material when appropriate.
@@ -652,6 +668,7 @@ bool GetShadowDepthPassShaders(
 	bool bOnePassPointLightShadow,
 	bool bPositionOnlyVS,
 	bool bUsePerspectiveCorrectShadowDepths,
+	bool bAtomicWrites,
 	TShaderRef<FShadowDepthVS>& VertexShader,
 	TShaderRef<FBaseHS>& HullShader,
 	TShaderRef<FBaseDS>& DomainShader,
@@ -742,9 +759,16 @@ bool GetShadowDepthPassShaders(
 	}
 
 	// Pixel shaders
-	const bool bNullPixelShader = Material.WritesEveryPixel(true) && !bUsePerspectiveCorrectShadowDepths && VertexFactory->SupportsNullPixelShader();
+	const bool bNullPixelShader = Material.WritesEveryPixel(true) && !bUsePerspectiveCorrectShadowDepths && !bAtomicWrites && VertexFactory->SupportsNullPixelShader();
 	if (!bNullPixelShader)
 	{
+#if ENABLE_NON_NANITE_VSM
+		if( bAtomicWrites )
+		{
+			ShaderTypes.AddShaderType<TShadowDepthPS<PixelShadowDepth_VirtualShadowMap>>();
+		}
+		else
+#endif
 		if (bUsePerspectiveCorrectShadowDepths)
 		{
 			ShaderTypes.AddShaderType<TShadowDepthPS<PixelShadowDepth_PerspectiveCorrect>>();
@@ -837,7 +861,11 @@ void SetStateForShadowDepth(bool bOnePassPointLightShadow, bool bDirectionalLigh
 	DrawRenderState.SetBlendState(TStaticBlendState<CW_NONE>::GetRHI());
 
 #if ENABLE_NON_NANITE_VSM
-	if (bOnePassPointLightShadow || InMeshPassTargetType == EMeshPass::VSMShadowDepth)
+	if( InMeshPassTargetType == EMeshPass::VSMShadowDepth && GVirtualShadowMapAtomicWrites )
+	{
+		DrawRenderState.SetDepthStencilState(TStaticDepthStencilState<false, CF_Always>::GetRHI());
+	}
+	else if (bOnePassPointLightShadow || InMeshPassTargetType == EMeshPass::VSMShadowDepth)
 #else //!ENABLE_NON_NANITE_VSM
 	if (bOnePassPointLightShadow)
 #endif // ENABLE_NON_NANITE_VSM
@@ -1115,7 +1143,7 @@ BEGIN_SHADER_PARAMETER_STRUCT(FShadowDepthPassParameters, )
 	SHADER_PARAMETER_RDG_UNIFORM_BUFFER(FShadowDepthPassUniformParameters, DeferredPassUniformBuffer)
 #if ENABLE_NON_NANITE_VSM
 	SHADER_PARAMETER_RDG_UNIFORM_BUFFER(FVirtualShadowMapCommonParameters, VirtualSmCommon)
-#endif // ENABLE_NON_NANITE_VSM
+#endif
 	SHADER_PARAMETER_STRUCT_INCLUDE(FInstanceCullingDrawParams, InstanceCullingDrawParams)
 	RENDER_TARGET_BINDING_SLOTS()
 END_SHADER_PARAMETER_STRUCT()
@@ -1783,10 +1811,6 @@ void FSceneRenderer::RenderShadowDepthMaps(FRDGBuilder& GraphBuilder, FInstanceC
 	const bool bUseHZB = (CVarNaniteShadowsUseHZB.GetValueOnRenderThread() != 0);
 	const bool bAllocatePageRectAtlas = CVarAllocatePagesUsingRects.GetValueOnRenderThread() != 0;
 
-#if ENABLE_NON_NANITE_VSM
-	VirtualShadowMapArray.RenderVirtualShadowMapsHw(GraphBuilder, SortedShadowsForShadowDepthPass.VirtualShadowMapShadows, *Scene);
-#endif
-
 	if (bNaniteEnabled && (bHasVSMShadows || bHasVSMClipMaps))
 	{
 		if (bUseHZB)
@@ -2023,6 +2047,10 @@ void FSceneRenderer::RenderShadowDepthMaps(FRDGBuilder& GraphBuilder, FInstanceC
 		GraphBuilder.QueueBufferExtraction(VirtualShadowMapArray.PageTableRDG, &Scene->VirtualShadowMapArrayCacheManager->HZBPageTable);
 	}
 
+#if ENABLE_NON_NANITE_VSM
+	VirtualShadowMapArray.RenderVirtualShadowMapsHw(GraphBuilder, SortedShadowsForShadowDepthPass.VirtualShadowMapShadows, *Scene);
+#endif
+
 	// Render non-VSM shadows
 	FSceneRenderer::RenderShadowDepthMapAtlases(GraphBuilder);
 
@@ -2256,12 +2284,14 @@ bool FShadowDepthPassMeshProcessor::Process(
 	// One pass point lights don't output a linear depth, so they are already perspective correct.
 	bool bUsePerspectiveCorrectShadowDepths = !ShadowDepthType.bDirectionalLight && !ShadowDepthType.bOnePassPointLightShadow;
 	bool bOnePassPointLightShadow = ShadowDepthType.bOnePassPointLightShadow;
+	bool bAtomicWrites = false;
 
 #if ENABLE_NON_NANITE_VSM
 	if( MeshPassTargetType == EMeshPass::VSMShadowDepth )
 	{
 		bUsePerspectiveCorrectShadowDepths = false;
 		bOnePassPointLightShadow = false;
+		bAtomicWrites = GVirtualShadowMapAtomicWrites != 0;
 	}
 #endif
 
@@ -2273,6 +2303,7 @@ bool FShadowDepthPassMeshProcessor::Process(
 		bOnePassPointLightShadow,
 		bUsePositionOnlyVS,
 		bUsePerspectiveCorrectShadowDepths,
+		bAtomicWrites,
 		ShadowDepthPassShaders.VertexShader,
 		ShadowDepthPassShaders.HullShader,
 		ShadowDepthPassShaders.DomainShader,
diff --git a/Engine/Source/Runtime/Renderer/Private/ShadowRendering.h b/Engine/Source/Runtime/Renderer/Private/ShadowRendering.h
--- a/Engine/Source/Runtime/Renderer/Private/ShadowRendering.h
+++ b/Engine/Source/Runtime/Renderer/Private/ShadowRendering.h
@@ -222,10 +222,12 @@ BEGIN_GLOBAL_SHADER_PARAMETER_STRUCT(FShadowDepthPassUniformParameters,)
 	// GPUCULL_TODO: ?
 	SHADER_PARAMETER(int, bRenderToVirtualShadowMap)
 	SHADER_PARAMETER(int, bInstancePerPage)
+	SHADER_PARAMETER(int, bAtomicWrites)
 
 	SHADER_PARAMETER_RDG_BUFFER_SRV(StructuredBuffer< uint >, VirtualSmPageTable)
 	SHADER_PARAMETER_RDG_BUFFER_SRV(StructuredBuffer< FPackedNaniteView >, PackedNaniteViews)
 	SHADER_PARAMETER_RDG_BUFFER_SRV(StructuredBuffer< uint4 >, PageRectBounds)
+	SHADER_PARAMETER_RDG_TEXTURE_UAV( RWTexture2D< uint >, OutDepthBuffer )
 
 	//SHADER_PARAMETER_RDG_BUFFER_SRV(Buffer< uint >, PackedPageInfoBuffer)
 #endif // ENABLE_NON_NANITE_VSM
diff --git a/Engine/Source/Runtime/Renderer/Private/VirtualShadowMaps/VirtualShadowMapArray.cpp b/Engine/Source/Runtime/Renderer/Private/VirtualShadowMaps/VirtualShadowMapArray.cpp
--- a/Engine/Source/Runtime/Renderer/Private/VirtualShadowMaps/VirtualShadowMapArray.cpp
+++ b/Engine/Source/Runtime/Renderer/Private/VirtualShadowMaps/VirtualShadowMapArray.cpp
@@ -117,6 +117,15 @@ TAutoConsoleVariable<int32> CVarVirtualShadowMapClipmapFirstCoarseLevel(
 	ECVF_RenderThreadSafe
 );
 
+int32 GVirtualShadowMapAtomicWrites = 1;
+FAutoConsoleVariableRef CVarAtomicWrites(
+	TEXT("r.Shadow.Virtual.AtomicWrites"),
+	GVirtualShadowMapAtomicWrites,
+	TEXT("Use per pixel page table routing and atomic writes instead of an instance per page."),
+	FConsoleVariableDelegate::CreateStatic(&RecreateGlobalRenderState),
+	ECVF_RenderThreadSafe
+);
+
 #if ENABLE_NON_NANITE_VSM
 
 static TAutoConsoleVariable<int32> CVarShowClipmapStats(
@@ -1486,21 +1495,27 @@ void FVirtualShadowMapArray::RenderVirtualShadowMapsHw(FRDGBuilder& GraphBuilder
 	}
 
 	RDG_EVENT_SCOPE(GraphBuilder, "RenderVirtualShadowMapsHw");
-	// Create physical page pool
-	PhysicalPagePoolHw = GraphBuilder.CreateTexture(FRDGTextureDesc::Create2D(GetPhysicalPoolSize(), PF_DepthStencil, FClearValueBinding::DepthFar, TexCreate_DepthStencilTargetable | TexCreate_ShaderResource), TEXT("Shadow.Virtual.PhysicalPagePoolHw"));
+
+	const bool bAllocatePageRectAtlas = CVarAllocatePagesUsingRects.GetValueOnRenderThread() != 0;
+	const bool bAtomicWrites = GVirtualShadowMapAtomicWrites != 0;
 
 	FGPUScene& GPUScene = Scene.GPUScene;
 	// Create that pass struct and whatever I guess.
 
 	// track whether the page pool was cleared
-	bool bWasCleared = false;
-
+	bool bWasCleared = bAtomicWrites;
 
+	if( !bAtomicWrites )
 	{
-		AddInitializePhysicalPagesHwPass(GraphBuilder, Scene.VirtualShadowMapArrayCacheManager);
-		bWasCleared = true;
+		// Create physical page pool
+		PhysicalPagePoolHw = GraphBuilder.CreateTexture(FRDGTextureDesc::Create2D(GetPhysicalPoolSize(), PF_DepthStencil, FClearValueBinding::DepthFar, TexCreate_DepthStencilTargetable | TexCreate_ShaderResource), TEXT("Shadow.Virtual.PhysicalPagePoolHw"));
+
+		if( !bAllocatePageRectAtlas )
+		{
+			AddInitializePhysicalPagesHwPass(GraphBuilder, Scene.VirtualShadowMapArrayCacheManager);
+			bWasCleared = true;
+		}
 	}
-	const bool bAllocatePageRectAtlas = CVarAllocatePagesUsingRects.GetValueOnRenderThread() != 0;
 
 	TArray<uint32, SceneRenderingAllocator> VirtualShadowMapFlags;
 	VirtualShadowMapFlags.AddZeroed(ShadowMaps.Num());
@@ -1592,8 +1607,8 @@ void FVirtualShadowMapArray::RenderVirtualShadowMapsHw(FRDGBuilder& GraphBuilder
 		{
 			TArray<int32, SceneRenderingAllocator> DrawCommandInstanceCountTmp;
 			DrawCommandInstanceCountTmp.AddZeroed(NumDrawCommands);
-			FRDGBufferRef DrawCommandInstanceCountRdg = CreateStructuredBuffer(GraphBuilder, TEXT("Shadow.Virtual.DrawCommandInstanceCount"), DrawCommandInstanceCountTmp);
-			FRDGBufferRef TmpInstanceIdOffsetBufferRdg = CreateStructuredBuffer(GraphBuilder, TEXT("Shadow.Virtual.TmpInstanceIdOffsetBufferRdg"), DrawCommandInstanceCountTmp);
+			FRDGBufferRef DrawCommandInstanceCountRDG = CreateStructuredBuffer(GraphBuilder, TEXT("Shadow.Virtual.DrawCommandInstanceCount"), DrawCommandInstanceCountTmp);
+			FRDGBufferRef TmpInstanceIdOffsetBufferRDG = CreateStructuredBuffer(GraphBuilder, TEXT("Shadow.Virtual.TmpInstanceIdOffsetBuffer"), DrawCommandInstanceCountTmp);
 			//AddBuildInstanceListPass(GraphBuilder, GPUScene, InstanceIdBuffer, InstanceCountBuffer);
 			// 2. Run page culling & command building pass(es)
 			// This takes the per high-level view instance list and performs page overlap tests and amplification/compaction for all target mip levels OR clip levels.
@@ -1605,9 +1620,9 @@ void FVirtualShadowMapArray::RenderVirtualShadowMapsHw(FRDGBuilder& GraphBuilder
 
 			TArray<uint32> NullArray;
 			NullArray.AddZeroed(1);
-			FRDGBufferRef VisibleInstanceWriteOffsetRdg = CreateStructuredBuffer(GraphBuilder, TEXT("Shadow.Virtual.VisibleInstanceWriteOffset"), NullArray);
-			FRDGBufferRef OutputOffsetBufferRdg = CreateStructuredBuffer(GraphBuilder, TEXT("Shadow.Virtual.OutputOffsetBuffer"), NullArray);
-			FRDGBufferRef VirtualShadowViewsRdg = CreateStructuredBuffer(GraphBuilder, TEXT("Shadow.Virtual.VirtualShadowViews"), VirtualShadowViews);
+			FRDGBufferRef VisibleInstanceWriteOffsetRDG = CreateStructuredBuffer(GraphBuilder, TEXT("Shadow.Virtual.VisibleInstanceWriteOffset"), NullArray);
+			FRDGBufferRef OutputOffsetBufferRDG = CreateStructuredBuffer(GraphBuilder, TEXT("Shadow.Virtual.OutputOffsetBuffer"), NullArray);
+			FRDGBufferRef VirtualShadowViewsRDG = CreateStructuredBuffer(GraphBuilder, TEXT("Shadow.Virtual.VirtualShadowViews"), VirtualShadowViews);
 
 
 			// TODO: Remove this when everything is properly RDG'd
@@ -1638,14 +1653,14 @@ void FVirtualShadowMapArray::RenderVirtualShadowMapsHw(FRDGBuilder& GraphBuilder
 				PassParameters->NumInstanceIdsBuffer = GraphBuilder.CreateSRV(Params.InstanceIdWriteOffsetBuffer);
 				PassParameters->FirstPrimaryView = 0;
 				PassParameters->NumPrimaryViews = NumPrimaryViews;
-				PassParameters->InViews = GraphBuilder.CreateSRV(VirtualShadowViewsRdg);
+				PassParameters->InViews = GraphBuilder.CreateSRV(VirtualShadowViewsRDG);
 				PassParameters->PrimitiveCullingCommands = GraphBuilder.CreateSRV(Params.PrimitiveCullingCommands);
 				PassParameters->VisibleInstancesOut = GraphBuilder.CreateUAV(VisibleInstancesRdg);
-				PassParameters->DrawCommandInstanceCountBufferOut = GraphBuilder.CreateUAV(DrawCommandInstanceCountRdg);
-				PassParameters->VisibleInstanceCountBufferOut = GraphBuilder.CreateUAV(VisibleInstanceWriteOffsetRdg);
+				PassParameters->DrawCommandInstanceCountBufferOut = GraphBuilder.CreateUAV(DrawCommandInstanceCountRDG);
+				PassParameters->VisibleInstanceCountBufferOut = GraphBuilder.CreateUAV(VisibleInstanceWriteOffsetRDG);
 				PassParameters->OutDynamicCasterFlags = GraphBuilder.CreateUAV(DynamicCasterPageFlagsRDG);
 				PassParameters->IndirectArgs = IndirectArgs;
-				PassParameters->bInstancePerPage = !bAllocatePageRectAtlas;
+				PassParameters->bInstancePerPage = !bAllocatePageRectAtlas && !bAtomicWrites;
 
 				FCullPerPageDrawCommandsCs::FPermutationDomain PermutationVector;
 				PermutationVector.Set< FCullPerPageDrawCommandsCs::FNearClipDim >( !Clipmap.IsValid() );
@@ -1672,8 +1687,8 @@ void FVirtualShadowMapArray::RenderVirtualShadowMapsHw(FRDGBuilder& GraphBuilder
 				PassParameters->DrawIndirectArgsBufferOut = GraphBuilder.CreateUAV(Params.DrawIndirectArgs, PF_R32_UINT);
 				PassParameters->InstanceIdOffsetBufferOut = GraphBuilder.CreateUAV(Params.InstanceIdStartOffsetBuffer, PF_R32_UINT);;
 				PassParameters->OutputOffsetBufferOut = GraphBuilder.CreateUAV(InstanceIdOutOffsetBufferRDG);
-				PassParameters->TmpInstanceIdOffsetBufferOut = GraphBuilder.CreateUAV(TmpInstanceIdOffsetBufferRdg);
-				PassParameters->DrawCommandInstanceCountBuffer = GraphBuilder.CreateSRV(DrawCommandInstanceCountRdg);
+				PassParameters->TmpInstanceIdOffsetBufferOut = GraphBuilder.CreateUAV(TmpInstanceIdOffsetBufferRDG);
+				PassParameters->DrawCommandInstanceCountBuffer = GraphBuilder.CreateSRV(DrawCommandInstanceCountRDG);
 
 				auto ComputeShader = GetGlobalShaderMap(GMaxRHIFeatureLevel)->GetShader<FAllocateCommandInstanceOutputSpaceCs>();
 
@@ -1688,15 +1703,15 @@ void FVirtualShadowMapArray::RenderVirtualShadowMapsHw(FRDGBuilder& GraphBuilder
 			}
 			// 2.3. Perform final pass to re-shuffle the instance ID's to their final resting places
 			{
-				FRDGBufferRef IndirectArgs = AddIndirectArgsSetupCsPass1D(GraphBuilder, VisibleInstanceWriteOffsetRdg, 1, FOutputCommandInstanceListsCs::NumThreadsPerGroup);
+				FRDGBufferRef IndirectArgs = AddIndirectArgsSetupCsPass1D(GraphBuilder, VisibleInstanceWriteOffsetRDG, 1, FOutputCommandInstanceListsCs::NumThreadsPerGroup);
 
 				FOutputCommandInstanceListsCs::FParameters* PassParameters = GraphBuilder.AllocParameters<FOutputCommandInstanceListsCs::FParameters>();
 
 				PassParameters->VisibleInstances = GraphBuilder.CreateSRV(VisibleInstancesRdg);
 				PassParameters->PageInfoBufferLegacyOut = GInstanceCullingManagerResources.GetPageInfoBufferUav();
 				PassParameters->InstanceIdsBufferLegacyOut = GInstanceCullingManagerResources.GetInstancesIdBufferUav();
-				PassParameters->TmpInstanceIdOffsetBufferOut = GraphBuilder.CreateUAV(TmpInstanceIdOffsetBufferRdg);
-				PassParameters->VisibleInstanceCountBuffer = GraphBuilder.CreateSRV(VisibleInstanceWriteOffsetRdg);
+				PassParameters->TmpInstanceIdOffsetBufferOut = GraphBuilder.CreateUAV(TmpInstanceIdOffsetBufferRDG);
+				PassParameters->VisibleInstanceCountBuffer = GraphBuilder.CreateSRV(VisibleInstanceWriteOffsetRDG);
 				PassParameters->IndirectArgs = IndirectArgs;
 
 				auto ComputeShader = GetGlobalShaderMap(GMaxRHIFeatureLevel)->GetShader<FOutputCommandInstanceListsCs>();
@@ -1711,10 +1726,21 @@ void FVirtualShadowMapArray::RenderVirtualShadowMapsHw(FRDGBuilder& GraphBuilder
 				);
 			}
 
+			if( PhysicalPagePoolRDG == nullptr )
+			{
+				PhysicalPagePoolRDG = GraphBuilder.CreateTexture( FRDGTextureDesc::Create2D( GetPhysicalPoolSize(), PF_R32_UINT, FClearValueBinding::None, TexCreate_ShaderResource | TexCreate_UAV ), TEXT("DepthBuffer32") );
+
+				ClearPhysicalMemory(GraphBuilder, PhysicalPagePoolRDG, Scene.VirtualShadowMapArrayCacheManager);
+			}
+
 			FVirtualShadowDepthPassParameters* PassParameters = GraphBuilder.AllocParameters<FVirtualShadowDepthPassParameters>();
 			PassParameters->View = ShadowDepthView->ViewUniformBuffer;
-			PassParameters->RenderTargets.DepthStencil = FDepthStencilBinding(PhysicalPagePoolHw, bWasCleared ? ERenderTargetLoadAction::ELoad : ERenderTargetLoadAction::EClear, FExclusiveDepthStencil::DepthWrite_StencilNop);
-			bWasCleared = true;
+
+			if( !bAtomicWrites )
+			{
+				PassParameters->RenderTargets.DepthStencil = FDepthStencilBinding(PhysicalPagePoolHw, bWasCleared ? ERenderTargetLoadAction::ELoad : ERenderTargetLoadAction::EClear, FExclusiveDepthStencil::DepthWrite_StencilNop);
+				bWasCleared = true;
+			}
 
 			FShadowDepthPassUniformParameters* ShadowDepthPassParameters = GraphBuilder.AllocParameters<FShadowDepthPassUniformParameters>();
 
@@ -1727,24 +1753,25 @@ void FVirtualShadowMapArray::RenderVirtualShadowMapsHw(FRDGBuilder& GraphBuilder
 			ShadowDepthPassParameters->ViewMatrix = FMatrix::Identity;
 			ShadowDepthPassParameters->ShadowParams = FVector4(0.0f, 0.0f, 0.0f, 1.0f);
 			ShadowDepthPassParameters->bRenderToVirtualShadowMap = true;
-			ShadowDepthPassParameters->VirtualSmPageTable = GraphBuilder.CreateSRV(PageTableRDG);
-			ShadowDepthPassParameters->PackedNaniteViews = GraphBuilder.CreateSRV(VirtualShadowViewsRdg);
-			ShadowDepthPassParameters->PageRectBounds = GraphBuilder.CreateSRV(PageRectBoundsRDG);
-			ShadowDepthPassParameters->bInstancePerPage = !bAllocatePageRectAtlas;
+			ShadowDepthPassParameters->bInstancePerPage = !bAllocatePageRectAtlas && !bAtomicWrites;
+			ShadowDepthPassParameters->bAtomicWrites = bAtomicWrites;
+			
+			ShadowDepthPassParameters->VirtualSmPageTable	= GraphBuilder.CreateSRV( PageTableRDG );
+			ShadowDepthPassParameters->PackedNaniteViews	= GraphBuilder.CreateSRV( VirtualShadowViewsRDG );
+			ShadowDepthPassParameters->PageRectBounds		= GraphBuilder.CreateSRV( PageRectBoundsRDG );
+			ShadowDepthPassParameters->OutDepthBuffer		= GraphBuilder.CreateUAV( PhysicalPagePoolRDG );
 			
-			PassParameters->ShadowDepthPass = GraphBuilder.CreateUniformBuffer(ShadowDepthPassParameters);;
+			PassParameters->ShadowDepthPass = GraphBuilder.CreateUniformBuffer(ShadowDepthPassParameters);
 			PassParameters->VirtualSmCommon = GetCommonUniformBuffer(GraphBuilder);
-			PassParameters->InViews = GraphBuilder.CreateSRV(VirtualShadowViewsRdg);
+			PassParameters->InViews = GraphBuilder.CreateSRV(VirtualShadowViewsRDG);
 
 			FInstanceCullingResult InstanceCullingResult;
 			InstanceCullingResult.DrawIndirectArgsBuffer = Params.DrawIndirectArgs;
 			InstanceCullingResult.InstanceIdOffsetBuffer = Params.InstanceIdStartOffsetBuffer;
 			InstanceCullingResult.GetDrawParameters(PassParameters->InstanceCullingDrawParams);
 
-			FMeshPassProcessorRenderState DrawRenderState;
-			DrawRenderState.SetBlendState(TStaticBlendState<CW_NONE>::GetRHI());
-			DrawRenderState.SetDepthStencilState(TStaticDepthStencilState<true, CF_DepthNearOrEqual>::GetRHI());
-
+			FIntRect ViewRect;
+			ViewRect.Max = bAtomicWrites ? FVirtualShadowMap::VirtualMaxResolutionXY : GetPhysicalPoolSize();
 
 			// TODO: Remove this when everything is properly RDG'd
 			AddPass(GraphBuilder, [](FRHICommandList& RHICmdList)
@@ -1762,16 +1789,10 @@ void FVirtualShadowMapArray::RenderVirtualShadowMapsHw(FRDGBuilder& GraphBuilder
 				RDG_EVENT_NAME("RenderVirtualShadowMapsHw"),
 				PassParameters,
 				ERDGPassFlags::Raster,
-				[this, ProjectedShadowInfo , &MeshCommandPass, PassParameters, ShadowDepthPassParameters](FRHICommandList& RHICmdList)
+				[this, ProjectedShadowInfo , &MeshCommandPass, PassParameters, ShadowDepthPassParameters, ViewRect](FRHICommandList& RHICmdList)
 				{
-					RHICmdList.SetViewport(
-						0.0f,
-						0.0f,
-						0.0f,
-						GetPhysicalPoolSize().X,
-						GetPhysicalPoolSize().Y,
-						1.0f
-					);
+					RHICmdList.SetViewport( ViewRect.Min.X, ViewRect.Min.Y, 0.0f, FMath::Min( ViewRect.Max.X, 32767 ), FMath::Min( ViewRect.Max.Y, 32767 ), 1.0f );
+
 					MeshCommandPass.DispatchDraw(nullptr, RHICmdList, &PassParameters->InstanceCullingDrawParams);
 				});
 		}
@@ -1900,7 +1921,7 @@ void FVirtualShadowMapArray::AddInitializePhysicalPagesHwPass(FRDGBuilder& Graph
 		ParametersPS->IndirectArgsBuffer = IndirectArgsBuffer;
 		ParametersPS->bRectPrimitive = GRHISupportsRectTopology ? 1 : 0;
 
-		bool bCacheDataAvailable = VirtualShadowMapArrayCacheManager && VirtualShadowMapArrayCacheManager->PrevPhysicalPageMetaData;
+		bool bCacheDataAvailable = VirtualShadowMapArrayCacheManager && VirtualShadowMapArrayCacheManager->PrevPhysicalPagePoolHw;
 		if (bCacheDataAvailable)
 		{
 			ParametersPS->PrevPhysicalPagePoolHw = GraphBuilder.RegisterExternalTexture(VirtualShadowMapArrayCacheManager->PrevPhysicalPagePoolHw);
diff --git a/Engine/Source/Runtime/Renderer/Private/VirtualShadowMaps/VirtualShadowMapCacheManager.cpp b/Engine/Source/Runtime/Renderer/Private/VirtualShadowMaps/VirtualShadowMapCacheManager.cpp
--- a/Engine/Source/Runtime/Renderer/Private/VirtualShadowMaps/VirtualShadowMapCacheManager.cpp
+++ b/Engine/Source/Runtime/Renderer/Private/VirtualShadowMaps/VirtualShadowMapCacheManager.cpp
@@ -140,8 +140,15 @@ void FVirtualShadowMapArrayCacheManager::ExtractFrameData(FVirtualShadowMapArray
 		
 		GraphBuilder.QueueTextureExtraction(VirtualShadowMapArray.PhysicalPagePoolRDG, &PrevPhysicalPagePool);
 #if ENABLE_NON_NANITE_VSM
-		GraphBuilder.QueueTextureExtraction(VirtualShadowMapArray.PhysicalPagePoolHw, &PrevPhysicalPagePoolHw);
-#endif //ENABLE_NON_NANITE_VSM
+		if( VirtualShadowMapArray.PhysicalPagePoolHw )
+		{
+			GraphBuilder.QueueTextureExtraction(VirtualShadowMapArray.PhysicalPagePoolHw, &PrevPhysicalPagePoolHw);
+		}
+		else
+		{
+			PrevPhysicalPagePoolHw = TRefCountPtr<IPooledRenderTarget>();
+		}
+#endif
 		GraphBuilder.QueueBufferExtraction(VirtualShadowMapArray.PhysicalPageMetaDataRDG, &PrevPhysicalPageMetaData);
 		GraphBuilder.QueueBufferExtraction(VirtualShadowMapArray.DynamicCasterPageFlagsRDG, &PrevDynamicCasterPageFlags);
 		GraphBuilder.QueueBufferExtraction(VirtualShadowMapArray.ShadowMapProjectionDataRDG, &PrevShadowMapProjectionDataBuffer);
