Add nullptr check
This corresponds to CL15495515 in UE4 which won't automatically robomerge due to many changes in this file.

diff --git a/Engine/Source/Runtime/Renderer/Private/GlobalDistanceField.cpp b/Engine/Source/Runtime/Renderer/Private/GlobalDistanceField.cpp
--- a/Engine/Source/Runtime/Renderer/Private/GlobalDistanceField.cpp
+++ b/Engine/Source/Runtime/Renderer/Private/GlobalDistanceField.cpp
@@ -1786,7 +1786,7 @@ void UpdateGlobalDistanceFieldVolume(
 									FHeightfieldComponentDescription NewComponentDescription(HeightfieldPrimitive->Proxy->GetLocalToWorld());
 									HeightfieldPrimitive->Proxy->GetHeightfieldRepresentation(HeightfieldTexture, DiffuseColorTexture, VisibilityTexture, NewComponentDescription);
 
-									if (HeightfieldTexture && HeightfieldTexture->Resource->TextureRHI)
+									if (HeightfieldTexture && HeightfieldTexture->Resource && HeightfieldTexture->Resource->TextureRHI)
 									{
 										const FIntPoint HeightfieldSize = NewComponentDescription.HeightfieldRect.Size();
 
