Remove double physical texture sample if only a single physical texture is written to.
Disable temporal denoiser pass until crash bug is fixed.

diff --git a/Engine/Shaders/Private/VirtualShadowMaps/ProjectionCommon.ush b/Engine/Shaders/Private/VirtualShadowMaps/ProjectionCommon.ush
--- a/Engine/Shaders/Private/VirtualShadowMaps/ProjectionCommon.ush
+++ b/Engine/Shaders/Private/VirtualShadowMaps/ProjectionCommon.ush
@@ -89,16 +89,20 @@ FVirtualShadowMapSample InitVirtualShadowMapSample()
 	return Result;
 }
 
+#ifndef TWO_PHYSICAL_TEXTURES
+#define TWO_PHYSICAL_TEXTURES 1
+#endif
+
 float SampleVirtualShadowMapPhysicalDepth(uint2 PhysicalTexelAddress)
 {
-#if ENABLE_NON_NANITE_VSM
+#if ENABLE_NON_NANITE_VSM && TWO_PHYSICAL_TEXTURES
 	return max(
 		asfloat(PhysicalPagePool.Load(int3(PhysicalTexelAddress, 0))),
 		PhysicalPagePoolHw.Load(int3(PhysicalTexelAddress, 0))
 	);
-#else //!ENABLE_NON_NANITE_VSM
+#else
 	return asfloat(PhysicalPagePool.Load(int3(PhysicalTexelAddress, 0)));
-#endif // ENABLE_NON_NANITE_VSM
+#endif
 }
 
 FVirtualShadowMapSample SampleVirtualShadowMapLevel(int VirtualShadowMapId, float2 ShadowMapUV, uint MipLevel)
diff --git a/Engine/Source/Runtime/Renderer/Private/ShadowRendering.cpp b/Engine/Source/Runtime/Renderer/Private/ShadowRendering.cpp
--- a/Engine/Source/Runtime/Renderer/Private/ShadowRendering.cpp
+++ b/Engine/Source/Runtime/Renderer/Private/ShadowRendering.cpp
@@ -206,7 +206,7 @@ static TAutoConsoleVariable<int32> CVarVirtualSmDenoiser(
 
 static TAutoConsoleVariable<int32> CVarVirtualSmDenoiserTemporal(
 	TEXT("r.Shadow.Virtual.Denoiser.Temporal"),
-	1,
+	0,
 	TEXT("Use temporal reprojection (when applicable)."),
 	ECVF_RenderThreadSafe);
 
diff --git a/Engine/Source/Runtime/Renderer/Private/VirtualShadowMaps/VirtualShadowMapProjection.cpp b/Engine/Source/Runtime/Renderer/Private/VirtualShadowMaps/VirtualShadowMapProjection.cpp
--- a/Engine/Source/Runtime/Renderer/Private/VirtualShadowMaps/VirtualShadowMapProjection.cpp
+++ b/Engine/Source/Runtime/Renderer/Private/VirtualShadowMaps/VirtualShadowMapProjection.cpp
@@ -463,8 +463,11 @@ class FVirtualShadowMapProjectionCS : public FGlobalShader
 	DECLARE_GLOBAL_SHADER(FVirtualShadowMapProjectionCS);
 	SHADER_USE_PARAMETER_STRUCT(FVirtualShadowMapProjectionCS, FGlobalShader)
 	
-	class FSMRTAdaptiveRayCountDim : SHADER_PERMUTATION_BOOL("SMRT_ADAPTIVE_RAY_COUNT");
-	using FPermutationDomain = TShaderPermutationDomain<FSMRTAdaptiveRayCountDim>;
+	class FSMRTAdaptiveRayCountDim	: SHADER_PERMUTATION_BOOL("SMRT_ADAPTIVE_RAY_COUNT");
+	class FTwoPhysicalTexturesDim	: SHADER_PERMUTATION_BOOL("TWO_PHYSICAL_TEXTURES");
+	using FPermutationDomain = TShaderPermutationDomain<
+		FSMRTAdaptiveRayCountDim,
+		FTwoPhysicalTexturesDim >;
 
 	BEGIN_SHADER_PARAMETER_STRUCT(FParameters, )
 		SHADER_PARAMETER_STRUCT_INCLUDE(FVirtualShadowMapSamplingParameters, ProjectionParameters)
@@ -506,6 +509,8 @@ class FVirtualShadowMapProjectionCS : public FGlobalShader
 };
 IMPLEMENT_GLOBAL_SHADER(FVirtualShadowMapProjectionCS, "/Engine/Private/VirtualShadowMaps/VirtualShadowMapProjection.usf", "VirtualShadowMapProjection", SF_Compute);
 
+extern int32 GVirtualShadowMapAtomicWrites;
+
 void RenderVirtualShadowMapProjection(
 	FRDGBuilder& GraphBuilder,
 	const FMinimalSceneTextures& SceneTextures,
@@ -537,6 +542,8 @@ void RenderVirtualShadowMapProjection(
 
 	FVirtualShadowMapProjectionCS::FPermutationDomain PermutationVector;
 	PermutationVector.Set< FVirtualShadowMapProjectionCS::FSMRTAdaptiveRayCountDim >( bAdaptiveRayCount );
+	PermutationVector.Set< FVirtualShadowMapProjectionCS::FTwoPhysicalTexturesDim >( GVirtualShadowMapAtomicWrites == 0 );
+
 	auto ComputeShader = View.ShaderMap->GetShader< FVirtualShadowMapProjectionCS >( PermutationVector );
 
 	const FIntPoint GroupCount = FIntPoint::DivideAndRoundUp( View.ViewRect.Size(), 8 );
