Remove unused variance tree and ray count code from path tracer

diff --git a/Engine/Shaders/Private/PathTracing/BuildVarianceMipTreeComputeShader.usf b/Engine/Shaders/Private/PathTracing/BuildVarianceMipTreeComputeShader.usf
deleted file mode 100644
diff --git a/Engine/Shaders/Private/PathTracing/PathCompaction.usf b/Engine/Shaders/Private/PathTracing/PathCompaction.usf
deleted file mode 100644
diff --git a/Engine/Shaders/Private/PathTracing/PathTracingCamera.ush b/Engine/Shaders/Private/PathTracing/PathTracingCamera.ush
--- a/Engine/Shaders/Private/PathTracing/PathTracingCamera.ush
+++ b/Engine/Shaders/Private/PathTracing/PathTracingCamera.ush
@@ -8,7 +8,6 @@
 
 #include "PathTracingCommon.ush"
 #include "Sensor/PathTracingPinholeCamera.ush"
-//#include "Sensor/PathTracingAdaptivePinholeCamera.ush"
 
 void SampleEmitter(
 	uint2 LaunchIndex,
diff --git a/Engine/Shaders/Private/PathTracing/PathTracingCompositingPixelShader.usf b/Engine/Shaders/Private/PathTracing/PathTracingCompositingPixelShader.usf
--- a/Engine/Shaders/Private/PathTracing/PathTracingCompositingPixelShader.usf
+++ b/Engine/Shaders/Private/PathTracing/PathTracingCompositingPixelShader.usf
@@ -5,7 +5,6 @@ ReferencePathTracinPixelShader.ush: Reference path tracing  compositing pixel sh
 ===============================================================================================*/
 
 #include "../Common.ush"
-#include "Utilities/PathTracingFixedPointArithmetic.ush"
 
 uint Iteration;
 uint MaxSamples;
diff --git a/Engine/Shaders/Private/PathTracing/PathTracingRayCounterComputeShader.usf b/Engine/Shaders/Private/PathTracing/PathTracingRayCounterComputeShader.usf
deleted file mode 100644
diff --git a/Engine/Shaders/Private/PathTracing/Sensor/PathTracingAdaptivePinholeCamera.ush b/Engine/Shaders/Private/PathTracing/Sensor/PathTracingAdaptivePinholeCamera.ush
deleted file mode 100644
diff --git a/Engine/Shaders/Private/PathTracing/Utilities/PathTracingFixedPointArithmetic.ush b/Engine/Shaders/Private/PathTracing/Utilities/PathTracingFixedPointArithmetic.ush
deleted file mode 100644
diff --git a/Engine/Source/Runtime/Renderer/Private/DeferredShadingRenderer.h b/Engine/Source/Runtime/Renderer/Private/DeferredShadingRenderer.h
--- a/Engine/Source/Runtime/Renderer/Private/DeferredShadingRenderer.h
+++ b/Engine/Source/Runtime/Renderer/Private/DeferredShadingRenderer.h
@@ -956,8 +956,6 @@ private:
 	void ComputePathCompaction(FRHICommandListImmediate& RHICmdList, const FViewInfo& View, FRHITexture* RadianceTexture, FRHITexture* SampleCountTexture, FRHITexture* PixelPositionTexture,
 		FRHIUnorderedAccessView* RadianceSortedRedUAV, FRHIUnorderedAccessView* RadianceSortedGreenUAV, FRHIUnorderedAccessView* RadianceSortedBlueUAV, FRHIUnorderedAccessView* RadianceSortedAlphaUAV, FRHIUnorderedAccessView* SampleCountSortedUAV);
 
-	void ComputeRayCount(FRHICommandListImmediate& RHICmdList, const FViewInfo& View, FRHITexture* RayCountPerPixelTexture);
-
 	void WaitForRayTracingScene(FRDGBuilder& GraphBuilder);
 
 	/** Debug ray tracing functions. */
diff --git a/Engine/Source/Runtime/Renderer/Private/PathTracing.cpp b/Engine/Source/Runtime/Renderer/Private/PathTracing.cpp
--- a/Engine/Source/Runtime/Renderer/Private/PathTracing.cpp
+++ b/Engine/Source/Runtime/Renderer/Private/PathTracing.cpp
@@ -507,8 +507,6 @@ void FDeferredShadingSceneRenderer::PreparePathTracing(const FViewInfo& View, TA
 void FSceneViewState::PathTracingInvalidate()
 {
 	PathTracingRadianceRT.SafeRelease();
-	VarianceMipTreeDimensions = FIntVector(0);
-	TotalRayCount = 0;
 	PathTracingSampleIndex = 0;
 }
 
diff --git a/Engine/Source/Runtime/Renderer/Private/PathTracing/PathCompactionCompute.cpp b/Engine/Source/Runtime/Renderer/Private/PathTracing/PathCompactionCompute.cpp
deleted file mode 100644
diff --git a/Engine/Source/Runtime/Renderer/Private/PathTracing/RayCounterCompute.cpp b/Engine/Source/Runtime/Renderer/Private/PathTracing/RayCounterCompute.cpp
deleted file mode 100644
diff --git a/Engine/Source/Runtime/Renderer/Private/RendererScene.cpp b/Engine/Source/Runtime/Renderer/Private/RendererScene.cpp
--- a/Engine/Source/Runtime/Renderer/Private/RendererScene.cpp
+++ b/Engine/Source/Runtime/Renderer/Private/RendererScene.cpp
@@ -237,24 +237,10 @@ FSceneViewState::FSceneViewState()
 	bUpdateLastExposure = false;
 
 #if RHI_RAYTRACING
-	VarianceMipTreeDimensions = FIntVector(0);
-	VarianceMipTree = new FRWBuffer;
 	PathTracingRect = FIntRect(0, 0, 0, 0);
 	PathTracingTargetSPP = 0;
 	PathTracingSampleIndex = 0;
 	PathTracingFrameIndex = 0;
-	TotalRayCount = 0;
-	TotalRayCountBuffer = new FRWBuffer;
-	if (GetFeatureLevel() >= ERHIFeatureLevel::SM5)
-	{
-		ENQUEUE_RENDER_COMMAND(InitializeSceneViewStateRWBuffer)(
-			[this](FRHICommandList&)
-			{
-				TotalRayCountBuffer->Initialize(TEXT("TotalRayCountBuffer"), sizeof(uint32), 1, PF_R32_UINT, BUF_SourceCopy);
-			});
-	}
-	bReadbackInitialized = false;
-	RayCountGPUReadback = new FRHIGPUBufferReadback(TEXT("Ray Count Readback"));
 
 	GatherPointsBuffer = nullptr;
 	GatherPointsResolution = FIntVector(0, 0, 0);
@@ -298,17 +284,6 @@ FSceneViewState::~FSceneViewState()
 	AOTileIntersectionResources = NULL;
 	DestroyRenderResource(AOScreenGridResources);
 	AOScreenGridResources = NULL;
-
-#if RHI_RAYTRACING
-	DestroyRWBuffer(VarianceMipTree);
-	DestroyRWBuffer(TotalRayCountBuffer);
-
-	ENQUEUE_RENDER_COMMAND(FDeleteGpuReadback)(
-		[DeleteMe = RayCountGPUReadback](FRHICommandList&)
-	{
-		delete DeleteMe;
-	});
-#endif // RHI_RAYTRACING
 }
 
 #if WITH_EDITOR
diff --git a/Engine/Source/Runtime/Renderer/Private/ScenePrivate.h b/Engine/Source/Runtime/Renderer/Private/ScenePrivate.h
--- a/Engine/Source/Runtime/Renderer/Private/ScenePrivate.h
+++ b/Engine/Source/Runtime/Renderer/Private/ScenePrivate.h
@@ -1001,8 +1001,6 @@ public:
 	TRefCountPtr<IPooledRenderTarget> PathTracingRadianceRT;
 	FIntRect PathTracingRect;
 	uint32   PathTracingTargetSPP;
-	FRWBuffer* VarianceMipTree;
-	FIntVector VarianceMipTreeDimensions;
 
 	// Path tracer samples/pixel
 	uint32 PathTracingSampleIndex;
@@ -1011,14 +1009,6 @@ public:
 	// the "screen door" effect and reduce temporal aliasing
 	uint32_t PathTracingFrameIndex;
 
-	// Path tracer ray counter
-	uint32 TotalRayCount;
-	FRWBuffer* TotalRayCountBuffer;
-
-	// Ray Count readback:
-	FRHIGPUBufferReadback* RayCountGPUReadback;
-	bool bReadbackInitialized = false;
-
 	// IES light profiles
 	FIESLightProfileResource IESLightProfileResources;
 
