Disabling early occlusion pass if lumen is enabled.

diff --git a/Engine/Source/Runtime/Renderer/Private/DeferredShadingRenderer.cpp b/Engine/Source/Runtime/Renderer/Private/DeferredShadingRenderer.cpp
--- a/Engine/Source/Runtime/Renderer/Private/DeferredShadingRenderer.cpp
+++ b/Engine/Source/Runtime/Renderer/Private/DeferredShadingRenderer.cpp
@@ -1689,6 +1689,7 @@ void FDeferredShadingSceneRenderer::Render(FRDGBuilder& GraphBuilder)
 		&& bCanOverlayRayTracingOutput;
 
 	bool bComputeLightGrid = false;
+	bool bAnyLumenEnabled = false;
 	// Simple forward shading doesn't support local lights. No need to compute light grid
 	if (!IsSimpleForwardShadingEnabled(ShaderPlatform))
 	{
@@ -1701,8 +1702,6 @@ void FDeferredShadingSceneRenderer::Render(FRDGBuilder& GraphBuilder)
 			bComputeLightGrid = ViewFamily.EngineShowFlags.Lighting;
 		}
 
-		bool bAnyLumenEnabled = false;
-
 		for (int32 ViewIndex = 0; ViewIndex < Views.Num(); ViewIndex++)
 		{
 			FViewInfo& View = Views[ViewIndex];
@@ -2024,7 +2023,7 @@ void FDeferredShadingSceneRenderer::Render(FRDGBuilder& GraphBuilder)
 	};
 
 	// Early occlusion queries
-	const bool bOcclusionBeforeBasePass = !bNaniteEnabled && ((DepthPass.EarlyZPassMode == EDepthDrawingMode::DDM_AllOccluders) || bIsEarlyDepthComplete);
+	const bool bOcclusionBeforeBasePass = !bNaniteEnabled && !bAnyLumenEnabled && ((DepthPass.EarlyZPassMode == EDepthDrawingMode::DDM_AllOccluders) || bIsEarlyDepthComplete);
 
 	if (bOcclusionBeforeBasePass)
 	{
