Removed "Lumen Reflections are enabled without Lumen Global Illumination" screen message as it is confusing in asset editors

diff --git a/Engine/Source/Runtime/Renderer/Private/SceneRendering.cpp b/Engine/Source/Runtime/Renderer/Private/SceneRendering.cpp
--- a/Engine/Source/Runtime/Renderer/Private/SceneRendering.cpp
+++ b/Engine/Source/Runtime/Renderer/Private/SceneRendering.cpp
@@ -2918,7 +2918,6 @@ void FSceneRenderer::RenderFinish(FRDGBuilder& GraphBuilder, FRDGTextureRef View
 		const bool bShowShadowedLightOverflowWarning = Scene->OverflowingDynamicShadowedLights.Num() > 0;
 
 		bool bLumenEnabledButNoSoftwareTracing = false;
-		bool bLumenReflectionsEnabledWithoutLumenGI = false;
 
 		for (int32 ViewIndex = 0;ViewIndex < Views.Num();ViewIndex++)
 		{	
@@ -2926,8 +2925,6 @@ void FSceneRenderer::RenderFinish(FRDGBuilder& GraphBuilder, FRDGTextureRef View
 			bLumenEnabledButNoSoftwareTracing = bLumenEnabledButNoSoftwareTracing 
 				|| (ShouldRenderLumenDiffuseGI(Scene, View, false) && !ShouldRenderLumenDiffuseGI(Scene, View, true))
 				|| (ShouldRenderLumenReflections(View, false) && !ShouldRenderLumenReflections(View, true));	
-			bLumenReflectionsEnabledWithoutLumenGI = bLumenReflectionsEnabledWithoutLumenGI
-				|| (ShouldRenderLumenReflections(View, true) && !ShouldRenderLumenDiffuseGI(Scene, View, true));
 		}
 
 		// Mobile-specific warnings
@@ -2954,7 +2951,7 @@ void FSceneRenderer::RenderFinish(FRDGBuilder& GraphBuilder, FRDGTextureRef View
 		const bool bAnyWarning = bShowPrecomputedVisibilityWarning || bShowGlobalClipPlaneWarning || bShowAtmosphericFogWarning || bShowSkylightWarning || bShowPointLightWarning 
 			|| bShowDFAODisabledWarning || bShowShadowedLightOverflowWarning || bShowMobileDynamicCSMWarning || bShowMobileLowQualityLightmapWarning || bShowMobileMovableDirectionalLightWarning
 			|| bMobileShowVertexFogWarning || bShowSkinCacheOOM || bSingleLayerWaterWarning || bShowDFDisabledWarning || bShowNoSkyAtmosphereComponentWarning || bFxDebugDraw 
-			|| bShowSkyAtmosphereFogComponentsConflicts || bLumenEnabledButNoSoftwareTracing || bLumenReflectionsEnabledWithoutLumenGI || bRealTimeSkyCaptureButNothingToCapture;
+			|| bShowSkyAtmosphereFogComponentsConflicts || bLumenEnabledButNoSoftwareTracing || bRealTimeSkyCaptureButNothingToCapture;
 
 		for(int32 ViewIndex = 0;ViewIndex < Views.Num();ViewIndex++)
 		{	
@@ -2977,7 +2974,7 @@ void FSceneRenderer::RenderFinish(FRDGBuilder& GraphBuilder, FRDGTextureRef View
 						bShowAtmosphericFogWarning, bViewParentOrFrozen, bShowSkylightWarning, bShowPointLightWarning, bShowShadowedLightOverflowWarning,
 						bShowMobileLowQualityLightmapWarning, bShowMobileMovableDirectionalLightWarning, bShowMobileDynamicCSMWarning, bMobileShowVertexFogWarning,
 						bShowSkinCacheOOM, bSingleLayerWaterWarning, bShowNoSkyAtmosphereComponentWarning, bFxDebugDraw, FXInterface, bShowSkyAtmosphereFogComponentsConflicts, 
-						bLumenEnabledButNoSoftwareTracing, bLumenReflectionsEnabledWithoutLumenGI, bRealTimeSkyCaptureButNothingToCapture]
+						bLumenEnabledButNoSoftwareTracing, bRealTimeSkyCaptureButNothingToCapture]
 						(FCanvas& Canvas)
 					{
 						// so it can get the screen size
@@ -3127,13 +3124,6 @@ void FSceneRenderer::RenderFinish(FRDGBuilder& GraphBuilder, FRDGTextureRef View
 							Canvas.DrawShadowedText(10, Y, Message, GetStatsFont(), FLinearColor(1.0, 0.05, 0.05, 1.0));
 							Y += 14;
 						}
-
-						if (bLumenReflectionsEnabledWithoutLumenGI)
-						{
-							static const FText Message = NSLOCTEXT("Renderer", "LumenReflectionsWithoutGI", "Lumen Reflections are enabled without Lumen Global Illumination, they will not render correctly.");
-							Canvas.DrawShadowedText(10, Y, Message, GetStatsFont(), FLinearColor(1.0, 0.05, 0.05, 1.0));
-							Y += 14;
-						}
 					});
 					if (bFxDebugDraw)
 					{
