Added (read only) cvar r.Shadow.Virtual.NonNaniteVSM to toggle non-nanite VSM support

diff --git a/Engine/Shaders/Private/VirtualShadowMaps/PageManagement.usf b/Engine/Shaders/Private/VirtualShadowMaps/PageManagement.usf
--- a/Engine/Shaders/Private/VirtualShadowMaps/PageManagement.usf
+++ b/Engine/Shaders/Private/VirtualShadowMaps/PageManagement.usf
@@ -851,8 +851,6 @@ void UpdatePhysicalPageMetaDataClearState(uint PageIndex : SV_DispatchThreadID)
 
 
 
-
-#if ENABLE_NON_NANITE_VSM
 Texture2D< float > PrevPhysicalPagePoolHw;
 
 
@@ -911,8 +909,6 @@ void InitializePhysicalPagesPS(uint PageIndex : TEXCOORD0, float4 SVPos : SV_POS
 
 	OutDepth = ClearValue;
 }
-#endif // ENABLE_NON_NANITE_VSM
-
 
 RWStructuredBuffer<FPhysicalPageMetaData> InOutPhysicalPageMetaData;
 StructuredBuffer<uint> VirtualShadowMapFlags;
@@ -971,8 +967,6 @@ void InitPhysicalPageMetaData(uint3 Index : SV_DispatchThreadID)
 	}
 }
 
-#if ENABLE_NON_NANITE_VSM
-
 
 StructuredBuffer<uint> InputCountBuffer;
 RWBuffer<uint> IndirectDispatchArgsOut;
@@ -1195,5 +1189,3 @@ void AllocatePagesUsingRects(uint GroupThreadId : SV_GroupThreadID, uint2 Dispat
 		GroupMemoryBarrierWithGroupSync();
 	}
 }
-
-#endif // ENABLE_NON_NANITE_VSM
diff --git a/Engine/Shaders/Private/VirtualShadowMaps/PrintStats.usf b/Engine/Shaders/Private/VirtualShadowMaps/PrintStats.usf
--- a/Engine/Shaders/Private/VirtualShadowMaps/PrintStats.usf
+++ b/Engine/Shaders/Private/VirtualShadowMaps/PrintStats.usf
@@ -6,10 +6,7 @@
 
 // Just the page count for now
 StructuredBuffer< uint >	InStatsBuffer;
-
-#if ENABLE_NON_NANITE_VSM
 StructuredBuffer<uint4> AllocatedPageRectBounds;
-#endif // ENABLE_NON_NANITE_VSM
 
 float2 ShaderPrintLeftAlign(float2 ScreenPos, uint Value)
 {
@@ -132,8 +129,6 @@ void PrintStats()
 	}
 }
 
-#if ENABLE_NON_NANITE_VSM
-
 
 uint ShadowMapIdRangeStart;
 uint ShadowMapIdRangeEnd;
@@ -241,4 +236,3 @@ void PrintClipmapStats()
 	}
 
 }
-#endif // ENABLE_NON_NANITE_VSM
diff --git a/Engine/Source/Runtime/Renderer/Private/MeshPassProcessor.cpp b/Engine/Source/Runtime/Renderer/Private/MeshPassProcessor.cpp
--- a/Engine/Source/Runtime/Renderer/Private/MeshPassProcessor.cpp
+++ b/Engine/Source/Runtime/Renderer/Private/MeshPassProcessor.cpp
@@ -996,8 +996,9 @@ void FMeshDrawCommand::SubmitDrawBegin(
 {
 	checkSlow(MeshDrawCommand.CachedPipelineId.IsValid());
 #if GPUCULL_TODO
+	// GPUCULL_TODO: Can't do this check as the VFs are created with GMaxRHIFeatureLevel (so may support PrimitiveIdStreamIndex even for preview platforms)
 	// Want to be sure that we supply GPU-scene instance data if required.
-	checkSlow(MeshDrawCommand.PrimitiveIdStreamIndex == -1 || ScenePrimitiveIdsBuffer != nullptr);
+	// checkSlow(MeshDrawCommand.PrimitiveIdStreamIndex == -1 || ScenePrimitiveIdsBuffer != nullptr);
 #endif // GPUCULL_TODO
 
 #if WANTS_DRAW_MESH_EVENTS
@@ -1268,6 +1269,7 @@ void DrawDynamicMeshPassPrivate(
 		const FViewInfo* ViewInfo = static_cast<const FViewInfo*>(&View);
 #if GPUCULL_TODO
 #if DO_GUARD_SLOW
+		if (UseGPUScene(View.GetShaderPlatform(), View.GetFeatureLevel()))
 		{
 			bool bNeedsGPUSceneData = false;
 			for (const auto& VisibleMeshDrawCommand : VisibleMeshDrawCommands)
diff --git a/Engine/Source/Runtime/Renderer/Private/SceneRendering.h b/Engine/Source/Runtime/Renderer/Private/SceneRendering.h
--- a/Engine/Source/Runtime/Renderer/Private/SceneRendering.h
+++ b/Engine/Source/Runtime/Renderer/Private/SceneRendering.h
@@ -1694,9 +1694,7 @@ struct FSortedShadowMaps
 	TArray<FProjectedShadowInfo*, SceneRenderingAllocator> VirtualShadowMapShadows;
 
 	TArray<TSharedPtr<FVirtualShadowMapClipmap>, SceneRenderingAllocator> VirtualShadowMapClipmaps;
-#if ENABLE_NON_NANITE_VSM
 	TArray<FProjectedShadowInfo*, SceneRenderingAllocator> VirtualShadowClipmapsHw;
-#endif // ENABLE_NON_NANITE_VSM
 
 	TArray<FSortedShadowMapAtlas,SceneRenderingAllocator> CompleteShadowMapAtlases;
 
diff --git a/Engine/Source/Runtime/Renderer/Private/ShadowDepthRendering.cpp b/Engine/Source/Runtime/Renderer/Private/ShadowDepthRendering.cpp
--- a/Engine/Source/Runtime/Renderer/Private/ShadowDepthRendering.cpp
+++ b/Engine/Source/Runtime/Renderer/Private/ShadowDepthRendering.cpp
@@ -120,6 +120,7 @@ static FAutoConsoleVariableRef CVarShadowShadowUseGS(
 extern int32 GVirtualShadowMapAtomicWrites;
 extern int32 GNaniteDebugFlags;
 extern int32 GNaniteShowStats;
+extern int32 GEnableNonNaniteVSM;
 
 namespace Nanite
 {
@@ -161,7 +162,7 @@ void SetupShadowDepthPassUniformBuffer(
 			ShadowDepthPassParameters.ShadowViewMatrices[FaceIndex] = Translation * ShadowInfo->OnePassShadowViewMatrices[FaceIndex];
 		}
 	}
-#if ENABLE_NON_NANITE_VSM
+
 	ShadowDepthPassParameters.bRenderToVirtualShadowMap = false;
 	ShadowDepthPassParameters.bInstancePerPage = false;
 	ShadowDepthPassParameters.bAtomicWrites = false;
@@ -172,7 +173,6 @@ void SetupShadowDepthPassUniformBuffer(
 	FRDGTextureRef DepthBuffer = GraphBuilder.CreateTexture( FRDGTextureDesc::Create2D( FIntPoint(4,4), PF_R32_UINT, FClearValueBinding::None, TexCreate_ShaderResource | TexCreate_UAV ), TEXT("Dummy-OutDepthBuffer") );
 
 	ShadowDepthPassParameters.OutDepthBuffer = GraphBuilder.CreateUAV( DepthBuffer );
-#endif // ENABLE_NON_NANITE_VSM
 }
 
 void SetupShadowDepthPassUniformBuffer(
@@ -332,14 +332,13 @@ public:
 		OutEnvironment.SetDefine(TEXT("ENABLE_FALLBACK_POINTLIGHT_SHADOW_GS"), UseGPUScene(Parameters.Platform) ? 1 : 0);
 #endif // GPUCULL_TODO
 
-		uint32 bEnableNonNaniteVSM = (uint32)(ENABLE_NON_NANITE_VSM != 0 && UseGPUScene(Parameters.Platform));
+		uint32 bEnableNonNaniteVSM = (uint32)(GEnableNonNaniteVSM != 0 && UseGPUScene(Parameters.Platform));
 		OutEnvironment.SetDefine(TEXT("ENABLE_NON_NANITE_VSM"), bEnableNonNaniteVSM);
-#if ENABLE_NON_NANITE_VSM
 		if (bEnableNonNaniteVSM != 0)
 		{
 			FVirtualShadowMapArray::SetShaderDefines(OutEnvironment);
 		}
-#endif // ENABLE_NON_NANITE_VSM
+
 		if (bIsForGeometryShader)
 		{
 			OutEnvironment.CompilerFlags.Add(CFLAG_VertexToGeometryShader);
@@ -506,7 +505,7 @@ public:
 			return false;
 		}
 
-		if (ShaderMode == PixelShadowDepth_VirtualShadowMap && !IsFeatureLevelSupported(Platform, ERHIFeatureLevel::SM5))
+		if (ShaderMode == PixelShadowDepth_VirtualShadowMap && (!IsFeatureLevelSupported(Platform, ERHIFeatureLevel::SM5) || GEnableNonNaniteVSM == 0))
 		{
 			return false;
 		}
@@ -531,14 +530,12 @@ public:
 		OutEnvironment.SetDefine(TEXT("ONEPASS_POINTLIGHT_SHADOW"), (uint32)(ShaderMode == PixelShadowDepth_OnePassPointLight));
 		OutEnvironment.SetDefine(TEXT("VIRTUAL_TEXTURE_TARGET"), (uint32)(ShaderMode == PixelShadowDepth_VirtualShadowMap));
 
-		uint32 bEnableNonNaniteVSM = (uint32)(ENABLE_NON_NANITE_VSM != 0 && UseGPUScene(Parameters.Platform));
+		uint32 bEnableNonNaniteVSM = (uint32)(GEnableNonNaniteVSM != 0 && UseGPUScene(Parameters.Platform));
 		OutEnvironment.SetDefine(TEXT("ENABLE_NON_NANITE_VSM"), bEnableNonNaniteVSM);
-#if ENABLE_NON_NANITE_VSM
 		if (bEnableNonNaniteVSM != 0)
 		{
 			FVirtualShadowMapArray::SetShaderDefines(OutEnvironment);
 		}
-#endif
 	}
 
 	TShadowDepthPS()
@@ -559,9 +556,7 @@ public:
 IMPLEMENT_SHADOWDEPTHPASS_PIXELSHADER_TYPE(PixelShadowDepth_NonPerspectiveCorrect);
 IMPLEMENT_SHADOWDEPTHPASS_PIXELSHADER_TYPE(PixelShadowDepth_PerspectiveCorrect);
 IMPLEMENT_SHADOWDEPTHPASS_PIXELSHADER_TYPE(PixelShadowDepth_OnePassPointLight);
-#if ENABLE_NON_NANITE_VSM
 IMPLEMENT_SHADOWDEPTHPASS_PIXELSHADER_TYPE(PixelShadowDepth_VirtualShadowMap);
-#endif
 
 /**
 * Overrides a material used for shadow depth rendering with the default material when appropriate.
@@ -664,14 +659,11 @@ bool GetShadowDepthPassShaders(
 	const bool bNullPixelShader = Material.WritesEveryPixel(true) && !bUsePerspectiveCorrectShadowDepths && !bAtomicWrites && VertexFactory->SupportsNullPixelShader();
 	if (!bNullPixelShader)
 	{
-#if ENABLE_NON_NANITE_VSM
-		if( bAtomicWrites )
+		if(UseNonNaniteVirtualShadowMaps(GShaderPlatformForFeatureLevel[FeatureLevel], FeatureLevel) && bAtomicWrites )
 		{
 			ShaderTypes.AddShaderType<TShadowDepthPS<PixelShadowDepth_VirtualShadowMap>>();
 		}
-		else
-#endif
-		if (bUsePerspectiveCorrectShadowDepths)
+		else if (bUsePerspectiveCorrectShadowDepths)
 		{
 			ShaderTypes.AddShaderType<TShadowDepthPS<PixelShadowDepth_PerspectiveCorrect>>();
 		}
@@ -760,15 +752,11 @@ void SetStateForShadowDepth(bool bOnePassPointLightShadow, bool bDirectionalLigh
 	// Disable color writes
 	DrawRenderState.SetBlendState(TStaticBlendState<CW_NONE>::GetRHI());
 
-#if ENABLE_NON_NANITE_VSM
 	if( InMeshPassTargetType == EMeshPass::VSMShadowDepth && GVirtualShadowMapAtomicWrites )
 	{
 		DrawRenderState.SetDepthStencilState(TStaticDepthStencilState<false, CF_Always>::GetRHI());
 	}
 	else if (bOnePassPointLightShadow || InMeshPassTargetType == EMeshPass::VSMShadowDepth)
-#else //!ENABLE_NON_NANITE_VSM
-	if (bOnePassPointLightShadow)
-#endif // ENABLE_NON_NANITE_VSM
 	{
 		// Point lights use reverse Z depth maps
 		DrawRenderState.SetDepthStencilState(TStaticDepthStencilState<true, CF_DepthNearOrEqual>::GetRHI());
@@ -1041,9 +1029,7 @@ BEGIN_SHADER_PARAMETER_STRUCT(FShadowDepthPassParameters, )
 	SHADER_PARAMETER_STRUCT_REF(FViewUniformShaderParameters, View)
 	SHADER_PARAMETER_RDG_UNIFORM_BUFFER(FMobileShadowDepthPassUniformParameters, MobilePassUniformBuffer)
 	SHADER_PARAMETER_RDG_UNIFORM_BUFFER(FShadowDepthPassUniformParameters, DeferredPassUniformBuffer)
-#if ENABLE_NON_NANITE_VSM
 	SHADER_PARAMETER_RDG_UNIFORM_BUFFER(FVirtualShadowMapCommonParameters, VirtualSmCommon)
-#endif
 	SHADER_PARAMETER_STRUCT_INCLUDE(FInstanceCullingDrawParams, InstanceCullingDrawParams)
 	RENDER_TARGET_BINDING_SLOTS()
 END_SHADER_PARAMETER_STRUCT()
@@ -1089,10 +1075,8 @@ void FProjectedShadowInfo::RenderDepth(
 		SetStateForShadowDepth(bOnePassPointLightShadow, bDirectionalLight, DrawRenderState, MeshPassTargetType);
 		CopyCachedShadowMap(GraphBuilder, *ShadowDepthView, SceneRenderer, PassParameters->RenderTargets, DrawRenderState);
 	}
-#if ENABLE_NON_NANITE_VSM
 	FVirtualShadowMapCommonParameters DummyVsmParameters = FVirtualShadowMapCommonParameters();
 	PassParameters->VirtualSmCommon = GraphBuilder.CreateUniformBuffer(&DummyVsmParameters);
-#endif // ENABLE_NON_NANITE_VSM
 
 	switch (FSceneInterface::GetShadingPath(FeatureLevel))
 	{
@@ -1947,9 +1931,10 @@ void FSceneRenderer::RenderShadowDepthMaps(FRDGBuilder& GraphBuilder, FInstanceC
 		GraphBuilder.QueueBufferExtraction(VirtualShadowMapArray.PageTableRDG, &Scene->VirtualShadowMapArrayCacheManager->HZBPageTable);
 	}
 
-#if ENABLE_NON_NANITE_VSM
-	VirtualShadowMapArray.RenderVirtualShadowMapsHw(GraphBuilder, SortedShadowsForShadowDepthPass.VirtualShadowMapShadows, *Scene);
-#endif
+	if (UseNonNaniteVirtualShadowMaps(ShaderPlatform, FeatureLevel))
+	{
+		VirtualShadowMapArray.RenderVirtualShadowMapsHw(GraphBuilder, SortedShadowsForShadowDepthPass.VirtualShadowMapShadows, *Scene);
+	}
 
 	// Render non-VSM shadows
 	FSceneRenderer::RenderShadowDepthMapAtlases(GraphBuilder);
@@ -2184,14 +2169,12 @@ bool FShadowDepthPassMeshProcessor::Process(
 	bool bOnePassPointLightShadow = ShadowDepthType.bOnePassPointLightShadow;
 	bool bAtomicWrites = false;
 
-#if ENABLE_NON_NANITE_VSM
 	if( MeshPassTargetType == EMeshPass::VSMShadowDepth )
 	{
 		bUsePerspectiveCorrectShadowDepths = false;
 		bOnePassPointLightShadow = false;
 		bAtomicWrites = GVirtualShadowMapAtomicWrites != 0;
 	}
-#endif
 
 	if (!GetShadowDepthPassShaders(
 		MaterialResource,
@@ -2275,11 +2258,7 @@ bool FShadowDepthPassMeshProcessor::TryAddMeshBatch(const FMeshBatch& RESTRICT M
 		const bool bPlatformReversesCulling = (RHINeedsToSwitchVerticalAxis(ShaderPlatform) && MobileHDRCvar->GetValueOnAnyThread() == 0);
 
 		const bool bRenderSceneTwoSided = bTwoSided;
-#if ENABLE_NON_NANITE_VSM
 		const bool bShadowReversesCulling = MeshPassTargetType == EMeshPass::VSMShadowDepth ? false : ShadowDepthType.bOnePassPointLightShadow;
-#else
-		const bool bShadowReversesCulling = ShadowDepthType.bOnePassPointLightShadow;
-#endif
 		const bool bReverseCullMode = XOR(bPlatformReversesCulling, bShadowReversesCulling);
 
 		FinalCullMode = bRenderSceneTwoSided ? CM_None : bReverseCullMode ? InverseCullMode(MeshCullMode) : MeshCullMode;
@@ -2295,28 +2274,31 @@ bool FShadowDepthPassMeshProcessor::TryAddMeshBatch(const FMeshBatch& RESTRICT M
 
 		OverrideWithDefaultMaterialForShadowDepth(EffectiveMaterialRenderProxy, EffectiveMaterial, FeatureLevel);
 
-#if ENABLE_NON_NANITE_VSM
-		// TODO: This uses a lot of indirections and complex logic, optimize by precomputing as far as possible
-		//       E.g., maybe store  bSupportsGpuSceneInstancing as a flag in the MeshBatch?
-		const FVertexFactory* VertexFactory = MeshBatch.VertexFactory;
-		const bool bUsePositionOnlyVS =
-			VertexFactory->SupportsPositionAndNormalOnlyStream()
-			&& EffectiveMaterial->WritesEveryPixel(true)
-			&& !EffectiveMaterial->MaterialModifiesMeshPosition_RenderThread();
-
-		// TODO: Store in MeshBatch?
-		const bool bSupportsGpuSceneInstancing = UseGPUScene(GShaderPlatformForFeatureLevel[FeatureLevel], FeatureLevel)
-			&& VertexFactory->GetPrimitiveIdStreamIndex(bUsePositionOnlyVS ? EVertexInputStreamType::PositionAndNormalOnly : EVertexInputStreamType::Default) != INDEX_NONE;
-
-		// EMeshPass::CSMShadowDepth: If no VSM: include everything, else only !bSupportsGpuSceneInstancing
-		// EMeshPass::VSMShadowDepth: If VSM: only bSupportsGpuSceneInstancing, else nothing needs to go in.
-		// TODO: I'm sure I can reduce this logic
-		const bool bDraw = MeshPassTargetType == EMeshPass::CSMShadowDepth ?
-			(!UseVirtualShadowMaps(GShaderPlatformForFeatureLevel[FeatureLevel], FeatureLevel) || !bSupportsGpuSceneInstancing) :
-			(UseVirtualShadowMaps(GShaderPlatformForFeatureLevel[FeatureLevel], FeatureLevel)  && bSupportsGpuSceneInstancing);
-		
+		bool bDraw = true;
+#if GPUCULL_TODO
+		if (UseNonNaniteVirtualShadowMaps(GShaderPlatformForFeatureLevel[FeatureLevel], FeatureLevel))
+		{
+			// TODO: This uses a lot of indirections and complex logic, optimize by precomputing as far as possible
+			//       E.g., maybe store  bSupportsGpuSceneInstancing as a flag in the MeshBatch?
+			const FVertexFactory* VertexFactory = MeshBatch.VertexFactory;
+			const bool bUsePositionOnlyVS =
+				VertexFactory->SupportsPositionAndNormalOnlyStream()
+				&& EffectiveMaterial->WritesEveryPixel(true)
+				&& !EffectiveMaterial->MaterialModifiesMeshPosition_RenderThread();
+
+			// TODO: Store in MeshBatch?
+			const bool bSupportsGpuSceneInstancing = UseGPUScene(GShaderPlatformForFeatureLevel[FeatureLevel], FeatureLevel)
+				&& VertexFactory->GetPrimitiveIdStreamIndex(bUsePositionOnlyVS ? EVertexInputStreamType::PositionAndNormalOnly : EVertexInputStreamType::Default) != INDEX_NONE;
+
+			// EMeshPass::CSMShadowDepth: If no VSM: include everything, else only !bSupportsGpuSceneInstancing
+			// EMeshPass::VSMShadowDepth: If VSM: only bSupportsGpuSceneInstancing, else nothing needs to go in.
+			// TODO: I'm sure I can reduce this logic
+			bDraw = MeshPassTargetType == EMeshPass::CSMShadowDepth ?
+				(!UseVirtualShadowMaps(GShaderPlatformForFeatureLevel[FeatureLevel], FeatureLevel) || !bSupportsGpuSceneInstancing) :
+				(UseVirtualShadowMaps(GShaderPlatformForFeatureLevel[FeatureLevel], FeatureLevel) && bSupportsGpuSceneInstancing);
+		}
+#endif // GPUCULL_TODO
 		if (bDraw)
-#endif
 		{
 			bResult = Process(MeshBatch, BatchElementMask, StaticMeshId, PrimitiveSceneProxy, *EffectiveMaterialRenderProxy, *EffectiveMaterial, MeshFillMode, FinalCullMode);
 		}
@@ -2374,9 +2356,9 @@ FRegisterPassProcessorCreateFunction RegisterCSMShadowDepthPass(&CreateCSMShadow
 
 FMeshPassProcessor* CreateVSMShadowDepthPassProcessor(const FScene* Scene, const FSceneView* InViewIfDynamicMeshCommand, FMeshPassDrawListContext* InDrawListContext)
 {
-#if ENABLE_NON_NANITE_VSM
+#if GPUCULL_TODO
 	// Only create the mesh pass processor if VSMs are not enabled as this prevents wasting time caching the SM draw commands
-	if (UseVirtualShadowMaps(Scene->GetShaderPlatform(), Scene->GetFeatureLevel()))
+	if (UseNonNaniteVirtualShadowMaps(Scene->GetShaderPlatform(), Scene->GetFeatureLevel()))
 	{
 		return new(FMemStack::Get()) FShadowDepthPassMeshProcessor(
 			Scene,
@@ -2385,7 +2367,7 @@ FMeshPassProcessor* CreateVSMShadowDepthPassProcessor(const FScene* Scene, const
 			InDrawListContext,
 			EMeshPass::VSMShadowDepth);
 	}
-#endif // ENABLE_NON_NANITE_VSM
+#endif // GPUCULL_TODO
 	return nullptr;
 }
 FRegisterPassProcessorCreateFunction RegisterVSMShadowDepthPass(&CreateVSMShadowDepthPassProcessor, EShadingPath::Deferred, EMeshPass::VSMShadowDepth, EMeshPassFlags::CachedMeshCommands);
diff --git a/Engine/Source/Runtime/Renderer/Private/ShadowRendering.h b/Engine/Source/Runtime/Renderer/Private/ShadowRendering.h
--- a/Engine/Source/Runtime/Renderer/Private/ShadowRendering.h
+++ b/Engine/Source/Runtime/Renderer/Private/ShadowRendering.h
@@ -218,7 +218,6 @@ BEGIN_GLOBAL_SHADER_PARAMETER_STRUCT(FShadowDepthPassUniformParameters,)
 	SHADER_PARAMETER(float, bClampToNearPlane)
 	SHADER_PARAMETER_ARRAY(FMatrix, ShadowViewProjectionMatrices, [6])
 	SHADER_PARAMETER_ARRAY(FMatrix, ShadowViewMatrices, [6])
-#if ENABLE_NON_NANITE_VSM
 	// GPUCULL_TODO: ?
 	SHADER_PARAMETER(int, bRenderToVirtualShadowMap)
 	SHADER_PARAMETER(int, bInstancePerPage)
@@ -228,9 +227,6 @@ BEGIN_GLOBAL_SHADER_PARAMETER_STRUCT(FShadowDepthPassUniformParameters,)
 	SHADER_PARAMETER_RDG_BUFFER_SRV(StructuredBuffer< FPackedNaniteView >, PackedNaniteViews)
 	SHADER_PARAMETER_RDG_BUFFER_SRV(StructuredBuffer< uint4 >, PageRectBounds)
 	SHADER_PARAMETER_RDG_TEXTURE_UAV( RWTexture2D< uint >, OutDepthBuffer )
-
-	//SHADER_PARAMETER_RDG_BUFFER_SRV(Buffer< uint >, PackedPageInfoBuffer)
-#endif // ENABLE_NON_NANITE_VSM
 END_GLOBAL_SHADER_PARAMETER_STRUCT()
 
 BEGIN_GLOBAL_SHADER_PARAMETER_STRUCT(FMobileShadowDepthPassUniformParameters,)
@@ -424,9 +420,7 @@ public:
 	int32 SubjectPrimitiveComponentIndex = -1;
 
 	TArray<int32> ViewIds;
-#if ENABLE_NON_NANITE_VSM
 	TSharedPtr<FVirtualShadowMapClipmap> VirtualShadowMapClipmap;
-#endif // ENABLE_NON_NANITE_VSM
 public:
 
 	// default constructor
@@ -462,7 +456,7 @@ public:
 		uint32 InSnapResolutionY,
 		uint32 InBorderSize
 		);
-#if ENABLE_NON_NANITE_VSM
+
 	/** for a clipmap shadow. */
 	void SetupClipmapProjection(
 		FLightSceneInfo* InLightSceneInfo,
@@ -470,7 +464,6 @@ public:
 		const TSharedPtr<FVirtualShadowMapClipmap> &VirtualShadowMapClipmap,
 		float InMaxNonFarCascadeDistance
 		);
-#endif // ENABLE_NON_NANITE_VSM
 
 	/**
 	 * Get `FViewMatrices` instance set up for drawing geometry to the SM, using the outer projection,
@@ -701,9 +694,7 @@ public:
 
 	FParallelMeshDrawCommandPass& GetShadowDepthPass() { return ShadowDepthPass; }
 
-#if ENABLE_NON_NANITE_VSM
 	float GetMaxNonFarCascadeDistance() const { return MaxNonFarCascadeDistance; }
-#endif // ENABLE_NON_NANITE_VSM
 private:
 	// 0 if Setup...() wasn't called yet
 	FLightSceneInfo* LightSceneInfo;
@@ -743,10 +734,8 @@ private:
 	FDynamicMeshDrawCommandStorage DynamicMeshDrawCommandStorage;
 	FGraphicsMinimalPipelineStateSet GraphicsMinimalPipelineStateSet;
 	bool NeedsShaderInitialisation;
-#if ENABLE_NON_NANITE_VSM
 	// If >= 0.0f then this records the cascade distance to the farthest non-'far' shadow cascade. Only used for clipmaps at the moment.
 	float MaxNonFarCascadeDistance = -1.0f;
-#endif // ENABLE_NON_NANITE_VSM
 
 	/**
 	 * Bias during in shadowmap rendering, stored redundantly for better performance 
diff --git a/Engine/Source/Runtime/Renderer/Private/ShadowSetup.cpp b/Engine/Source/Runtime/Renderer/Private/ShadowSetup.cpp
--- a/Engine/Source/Runtime/Renderer/Private/ShadowSetup.cpp
+++ b/Engine/Source/Runtime/Renderer/Private/ShadowSetup.cpp
@@ -274,18 +274,13 @@ static TAutoConsoleVariable<int32> CVarAlwaysAllocateMaxResolutionAtlases(
 	ECVF_RenderThreadSafe
 );
 
-
-#if ENABLE_NON_NANITE_VSM
-
 static TAutoConsoleVariable<int32> CVarVsmUseFarShadowRules(
 	TEXT("r.Shadow.Virtual.UseFarShadowCulling"),
-	1,
+	0,
 	TEXT("Switch between implementing the far shadow culling logic for VSMs."),
 	ECVF_RenderThreadSafe
 );
 
-#endif // ENABLE_NON_NANITE_VSM
-
 CSV_DECLARE_CATEGORY_EXTERN(LightCount);
 
 #if !UE_BUILD_SHIPPING
@@ -986,8 +981,6 @@ void FProjectedShadowInfo::SetupWholeSceneProjection(
 }
 
 
-#if ENABLE_NON_NANITE_VSM
-
 void FProjectedShadowInfo::SetupClipmapProjection(FLightSceneInfo* InLightSceneInfo, FViewInfo* InDependentView, const TSharedPtr<FVirtualShadowMapClipmap>& InVirtualShadowMapClipmap, float InMaxNonFarCascadeDistance)
 {
 	VirtualShadowMapClipmap = InVirtualShadowMapClipmap;
@@ -1022,8 +1015,6 @@ void FProjectedShadowInfo::SetupClipmapProjection(FLightSceneInfo* InLightSceneI
 	TranslatedWorldToClipOuterMatrix = TranslatedWorldToView * ViewToClipOuter;
 }
 
-#endif // ENABLE_NON_NANITE_VSM
-
 void FProjectedShadowInfo::AddCachedMeshDrawCommandsForPass(
 	int32 PrimitiveIndex,
 	const FPrimitiveSceneInfo* InPrimitiveSceneInfo,
@@ -1685,22 +1676,22 @@ bool FProjectedShadowInfo::TestPrimitiveFarCascadeConditions(bool bPrimitiveCast
 		return bPrimitiveCastsFarShadow;
 	}
 
-#if ENABLE_NON_NANITE_VSM
-
-	const bool bWholeSceneDirectionalShadow = IsWholeSceneDirectionalShadow();
-
-	if (bWholeSceneDirectionalShadow && MaxNonFarCascadeDistance > 0.0f)
+	if (GEnableNonNaniteVSM != 0 && MeshPassTargetType == EMeshPass::VSMShadowDepth)
 	{
-		check(DependentView);
+		const bool bWholeSceneDirectionalShadow = IsWholeSceneDirectionalShadow();
 
-		FVector ViewOrigin = DependentView->ShadowViewMatrices.GetViewOrigin();
-		float ViewDistance = (Bounds.Origin - ViewOrigin) | DependentView->GetViewDirection();
-		if (ViewDistance - Bounds.SphereRadius > MaxNonFarCascadeDistance)
+		if (bWholeSceneDirectionalShadow && MaxNonFarCascadeDistance > 0.0f)
 		{
-			return bPrimitiveCastsFarShadow;
+			check(DependentView);
+
+			FVector ViewOrigin = DependentView->ShadowViewMatrices.GetViewOrigin();
+			float ViewDistance = (Bounds.Origin - ViewOrigin) | DependentView->GetViewDirection();
+			if (ViewDistance - Bounds.SphereRadius > MaxNonFarCascadeDistance)
+			{
+				return bPrimitiveCastsFarShadow;
+			}
 		}
 	}
-#endif // ENABLE_NON_NANITE_VSM
 	return true;
 }
 
@@ -2054,8 +2045,6 @@ void FProjectedShadowInfo::SetupMeshDrawCommandsForShadowDepth(FSceneRenderer& R
 			}
 		}
 	}
-
-#if ENABLE_NON_NANITE_VSM
 	else if (VirtualShadowMapClipmap.IsValid())
 	{
 		// TODO: Register view per clip level such that they are culled early (?)
@@ -2068,7 +2057,6 @@ void FProjectedShadowInfo::SetupMeshDrawCommandsForShadowDepth(FSceneRenderer& R
 		Params.RasterContextSize = FIntPoint(ResolutionX, ResolutionY);
 		ViewIds.Add(InstanceCullingManager.RegisterView(Params));
 	}
-#endif // ENABLE_NON_NANITE_VSM
 	else
 	{
 		Nanite::FPackedViewParams Params;
@@ -2547,11 +2535,7 @@ void FSceneRenderer::SetupInteractionShadows(
 	if (!bShadowHandledByParent)
 	{
 		const bool bCreateTranslucentObjectShadow = GUseTranslucencyShadowDepths && Interaction->HasTranslucentObjectShadow();
-#if ENABLE_NON_NANITE_VSM
-		const bool bCreateInsetObjectShadow = !UseVirtualShadowMaps(ShaderPlatform, FeatureLevel) &&  Interaction->HasInsetObjectShadow();
-#else //!ENABLE_NON_NANITE_VSM
-		const bool bCreateInsetObjectShadow = Interaction->HasInsetObjectShadow();
-#endif // ENABLE_NON_NANITE_VSM
+		const bool bCreateInsetObjectShadow = Interaction->HasInsetObjectShadow() && !UseNonNaniteVirtualShadowMaps(ShaderPlatform, FeatureLevel);
 		const bool bCreateObjectShadowForStationaryLight = ShouldCreateObjectShadowForStationaryLight(Interaction->GetLight(), PrimitiveSceneInfo->Proxy, Interaction->IsShadowMapped());
 
 		if (Interaction->HasShadow() 
@@ -3865,15 +3849,16 @@ void FSceneRenderer::GatherShadowDynamicMeshElements(FGlobalDynamicIndexBuffer&
 			ProjectedShadowInfo->GatherDynamicMeshElements(*this, VisibleLightInfo, ReusedViewsArray, DynamicIndexBuffer, DynamicVertexBuffer, DynamicReadBuffer, InstanceCullingManager);
 		}
 	}
-#if ENABLE_NON_NANITE_VSM
 
-	// GPUCULL_TODO: Replace with new shadow culling processor thingo
-	for (FProjectedShadowInfo* ProjectedShadowInfo : SortedShadowsForShadowDepthPass.VirtualShadowMapShadows)
+	if (UseNonNaniteVirtualShadowMaps(ShaderPlatform, FeatureLevel))
 	{
-		FVisibleLightInfo& VisibleLightInfo = VisibleLightInfos[ProjectedShadowInfo->GetLightSceneInfo().Id];
-		ProjectedShadowInfo->GatherDynamicMeshElements(*this, VisibleLightInfo, ReusedViewsArray, DynamicIndexBuffer, DynamicVertexBuffer, DynamicReadBuffer, InstanceCullingManager);
+		// GPUCULL_TODO: Replace with new shadow culling processor thing
+		for (FProjectedShadowInfo* ProjectedShadowInfo : SortedShadowsForShadowDepthPass.VirtualShadowMapShadows)
+		{
+			FVisibleLightInfo& VisibleLightInfo = VisibleLightInfos[ProjectedShadowInfo->GetLightSceneInfo().Id];
+			ProjectedShadowInfo->GatherDynamicMeshElements(*this, VisibleLightInfo, ReusedViewsArray, DynamicIndexBuffer, DynamicVertexBuffer, DynamicReadBuffer, InstanceCullingManager);
+		}
 	}
-#endif // defined(ENABLE_NON_NANITE_VSM)
 }
 
 typedef TArray<FAddSubjectPrimitiveOp> FShadowSubjectPrimitives;
@@ -3989,12 +3974,10 @@ struct FGatherShadowPrimitivesPacket
 
 	bool DoesPrimitiveCastInsetShadow(const FPrimitiveSceneInfo* PrimitiveSceneInfo, const FPrimitiveSceneProxy* PrimitiveProxy) const
 	{
-#if ENABLE_NON_NANITE_VSM
-		if (UseVirtualShadowMaps(GShaderPlatformForFeatureLevel[FeatureLevel], FeatureLevel))
+		if (UseNonNaniteVirtualShadowMaps(GShaderPlatformForFeatureLevel[FeatureLevel], FeatureLevel))
 		{
 			return false;
 		}
-#endif // ENABLE_NON_NANITE_VSM
 
 		// If light attachment root is valid, we're in a group and need to get the flag from the root.
 		if (PrimitiveSceneInfo->LightingAttachmentRoot.IsValid())
@@ -4341,9 +4324,7 @@ void FSceneRenderer::AddViewDependentWholeSceneShadowsForView(
 
 			static_assert(INDEX_NONE == -1, "INDEX_NONE != -1!");
 
-#if ENABLE_NON_NANITE_VSM
 			float MaxNonFarCascadeDistance = 0.0f;
-#endif // ENABLE_NON_NANITE_VSM
 
 			// todo: this code can be simplified by computing all the distances in one place - avoiding some redundant work and complexity
 			for (int32 Index = 0; Index < ProjectionCount; Index++)
@@ -4360,12 +4341,10 @@ void FSceneRenderer::AddViewDependentWholeSceneShadowsForView(
 
 				if (LightSceneInfo.Proxy->GetViewDependentWholeSceneProjectedShadowInitializer(View, LocalIndex, LightSceneInfo.IsPrecomputedLightingValid(), ProjectedShadowInitializer))
 				{
-#if ENABLE_NON_NANITE_VSM
 					if (!ProjectedShadowInitializer.CascadeSettings.bFarShadowCascade)
 					{
 						MaxNonFarCascadeDistance = FMath::Max(MaxNonFarCascadeDistance, ProjectedShadowInitializer.CascadeSettings.SplitFar);
 					}
-#endif // ENABLE_NON_NANITE_VSM
 					uint32 ShadowBorder = NeedsUnatlasedCSMDepthsWorkaround( FeatureLevel ) ? 0 : SHADOW_BORDER;
 					
 					const int32 MaxCSMResolution = GetCachedScalabilityCVars().MaxCSMShadowResolution;
@@ -4466,15 +4445,16 @@ void FSceneRenderer::AddViewDependentWholeSceneShadowsForView(
 				VisibleLightInfo.VirtualShadowMapClipmaps.Add(VirtualShadowMapClipmap);
 				SortedShadowsForShadowDepthPass.VirtualShadowMapClipmaps.Add(VirtualShadowMapClipmap);
 
-#if ENABLE_NON_NANITE_VSM
-				// Create the projected shadow info to make sure that culling happens.
-				FProjectedShadowInfo* ProjectedShadowInfo = new(FMemStack::Get(), 1, 16) FProjectedShadowInfo;
-				// Add to remember-to-call-dtor list
-				VisibleLightInfo.MemStackProjectedShadows.Add(ProjectedShadowInfo);
-				ProjectedShadowInfo->SetupClipmapProjection(&LightSceneInfo, &View, VirtualShadowMapClipmap, CVarVsmUseFarShadowRules.GetValueOnRenderThread() != 0 ? MaxNonFarCascadeDistance : -1.0f);
-				VisibleLightInfo.AllProjectedShadows.Add(ProjectedShadowInfo);
-				ShadowInfosThatNeedCulling.Add(ProjectedShadowInfo);
-#endif // ENABLE_NON_NANITE_VSM
+				if (GEnableNonNaniteVSM != 0)
+				{
+					// Create the projected shadow info to make sure that culling happens.
+					FProjectedShadowInfo* ProjectedShadowInfo = new(FMemStack::Get(), 1, 16) FProjectedShadowInfo;
+					// Add to remember-to-call-dtor list
+					VisibleLightInfo.MemStackProjectedShadows.Add(ProjectedShadowInfo);
+					ProjectedShadowInfo->SetupClipmapProjection(&LightSceneInfo, &View, VirtualShadowMapClipmap, CVarVsmUseFarShadowRules.GetValueOnRenderThread() != 0 ? MaxNonFarCascadeDistance : -1.0f);
+					VisibleLightInfo.AllProjectedShadows.Add(ProjectedShadowInfo);
+					ShadowInfosThatNeedCulling.Add(ProjectedShadowInfo);
+				}
 			}
 		}
 	}
@@ -4591,11 +4571,7 @@ void FSceneRenderer::AllocateShadowDepthTargets(FRHICommandListImmediate& RHICmd
 
 				if (bNeedsShadowmapSetup)
 				{
-#if ENABLE_NON_NANITE_VSM
 					if (ProjectedShadowInfo->HasVirtualShadowMap() || ProjectedShadowInfo->VirtualShadowMapClipmap.IsValid())
-#else
-					if (ProjectedShadowInfo->HasVirtualShadowMap())
-#endif
 					{
 						ProjectedShadowInfo->SetupShadowDepthView(RHICmdList, this);
 						SortedShadowsForShadowDepthPass.VirtualShadowMapShadows.Add(ProjectedShadowInfo);
diff --git a/Engine/Source/Runtime/Renderer/Private/VirtualShadowMaps/VirtualShadowMapArray.cpp b/Engine/Source/Runtime/Renderer/Private/VirtualShadowMaps/VirtualShadowMapArray.cpp
--- a/Engine/Source/Runtime/Renderer/Private/VirtualShadowMaps/VirtualShadowMapArray.cpp
+++ b/Engine/Source/Runtime/Renderer/Private/VirtualShadowMaps/VirtualShadowMapArray.cpp
@@ -43,10 +43,8 @@ struct FCachedPageInfo
 
 static void RecreateGlobalRenderState(IConsoleVariable* Var)
 {
-#if ENABLE_NON_NANITE_VSM
 	// Needed because the depth state changes with method (so cached draw commands must be re-created) see SetStateForShadowDepth
 	FGlobalComponentRecreateRenderStateContext Context;
-#endif
 }
 
 int32 GEnableVirtualShadowMaps = 0;
@@ -126,8 +124,6 @@ FAutoConsoleVariableRef CVarAtomicWrites(
 	ECVF_RenderThreadSafe
 );
 
-#if ENABLE_NON_NANITE_VSM
-
 static TAutoConsoleVariable<int32> CVarShowClipmapStats(
 	TEXT("r.Shadow.Virtual.ShowClipmapStats"),
 	-1,
@@ -135,7 +131,17 @@ static TAutoConsoleVariable<int32> CVarShowClipmapStats(
 	ECVF_RenderThreadSafe
 );
 
-#endif // ENABLE_NON_NANITE_VSM
+int32 GEnableNonNaniteVSM = 0;
+
+#if GPUCULL_TODO
+FAutoConsoleVariableRef CVarEnableNonNaniteVSM(
+	TEXT("r.Shadow.Virtual.NonNaniteVSM"),
+	GEnableNonNaniteVSM,
+	TEXT("EXPERIMENTAL: Enable support for non-nanite Virtual Shadow Maps, this has some performance impact on regular shadow maps.")
+	TEXT("Read-only and to be set in a config file (requires restart)."),
+	ECVF_RenderThreadSafe | ECVF_ReadOnly
+);
+#endif // GPUCULL_TODO
 
 FMatrix CalcTranslatedWorldToShadowUVMatrix(
 	const FMatrix& TranslatedWorldToShadowView,
@@ -237,10 +243,8 @@ void FVirtualShadowMapArray::Initialize(bool bInEnabled)
 
 void FVirtualShadowMapArray::SetShaderDefines(FShaderCompilerEnvironment& OutEnvironment)
 {
-#if ENABLE_NON_NANITE_VSM
 	static_assert(FVirtualShadowMap::Log2Level0DimPagesXY * 2U + MAX_VIEWS_PER_CULL_RASTERIZE_PASS_BITS <= 32U, "Page indirection plus view index must fit into 32-bits for page-routing storage!");
-#endif // ENABLE_NON_NANITE_VSM
-	OutEnvironment.SetDefine(TEXT("ENABLE_NON_NANITE_VSM"), int32(ENABLE_NON_NANITE_VSM));
+	OutEnvironment.SetDefine(TEXT("ENABLE_NON_NANITE_VSM"), GEnableNonNaniteVSM);
 
 	OutEnvironment.SetDefine(TEXT("VSM_PAGE_SIZE"), FVirtualShadowMap::PageSize);
 	OutEnvironment.SetDefine(TEXT("VSM_PAGE_SIZE_MASK"), FVirtualShadowMap::PageSizeMask);
@@ -275,11 +279,7 @@ public:
 
 	static bool ShouldCompilePermutation(const FGlobalShaderPermutationParameters& Parameters)
 	{
-#if ENABLE_NON_NANITE_VSM
-		return FDataDrivenShaderPlatformInfo::GetSupportsGPUScene(Parameters.Platform);
-#else //!ENABLE_NON_NANITE_VSM
 		return DoesPlatformSupportNanite(Parameters.Platform);
-#endif // ENABLE_NON_NANITE_VSM
 	}
 
 	/**
@@ -442,8 +442,6 @@ class FClearPhysicalPagesCS : public FVirtualPageManagementShader
 IMPLEMENT_GLOBAL_SHADER(FClearPhysicalPagesCS, "/Engine/Private/VirtualShadowMaps/PageManagement.usf", "ClearPhysicalPages", SF_Compute);
 
 
-#if ENABLE_NON_NANITE_VSM
-
 class FInitIndirectArgs1DCS : public FVirtualPageManagementShader
 {
 	DECLARE_GLOBAL_SHADER(FInitIndirectArgs1DCS);
@@ -486,7 +484,6 @@ FRDGBufferRef AddIndirectArgsSetupCsPass1D(FRDGBuilder& GraphBuilder, FRDGBuffer
 	return IndirectArgsBuffer;
 }
 
-#endif // ENABLE_NON_NANITE_VSM
 
 void FVirtualShadowMapArray::ClearPhysicalMemory(FRDGBuilder& GraphBuilder, FRDGTextureRef& PhysicalTexture, FVirtualShadowMapArrayCacheManager *VirtualShadowMapArrayCacheManager)
 {
@@ -625,8 +622,6 @@ class FInitPageRectBoundsCS : public FVirtualPageManagementShader
 IMPLEMENT_GLOBAL_SHADER(FInitPageRectBoundsCS, "/Engine/Private/VirtualShadowMaps/PageManagement.usf", "InitPageRectBounds", SF_Compute);
 
 
-#if ENABLE_NON_NANITE_VSM
-
 
 
 
@@ -679,7 +674,6 @@ class FAllocatePagesUsingRectsCS : public FVirtualPageManagementShader
 	END_SHADER_PARAMETER_STRUCT()
 };
 IMPLEMENT_GLOBAL_SHADER(FAllocatePagesUsingRectsCS, "/Engine/Private/VirtualShadowMaps/PageManagement.usf", "AllocatePagesUsingRects", SF_Compute);
-#endif // ENABLE_NON_NANITE_VSM
 
 void FVirtualShadowMapArray::BuildPageAllocations(
 	FRDGBuilder& GraphBuilder,
@@ -922,9 +916,7 @@ void FVirtualShadowMapArray::BuildPageAllocations(
 		CachedPageInfosRDG = GraphBuilder.CreateBuffer(FRDGBufferDesc::CreateStructuredDesc(sizeof(FCachedPageInfo), CommonParameters.MaxPhysicalPages), TEXT("Shadow.Virtual.CachedPageInfos"));
 		PhysicalPageMetaDataRDG = GraphBuilder.CreateBuffer(FRDGBufferDesc::CreateStructuredDesc(sizeof(FPhysicalPageMetaData), CommonParameters.MaxPhysicalPages), TEXT("Shadow.Virtual.PhysicalPageMetaData"));
 
-#if ENABLE_NON_NANITE_VSM
 		AllocatedPageRectBoundsRDG = GraphBuilder.CreateBuffer(FRDGBufferDesc::CreateStructuredDesc(sizeof(FIntVector4), NumPageRects), TEXT("Shadow.Virtual.AllocatedPageRectBounds"));
-#endif // ENABLE_NON_NANITE_VSM
 
 		{
 			FInitPhysicalPageMetaData::FParameters* PassParameters = GraphBuilder.AllocParameters< FInitPhysicalPageMetaData::FParameters >();
@@ -941,8 +933,6 @@ void FVirtualShadowMapArray::BuildPageAllocations(
 			);
 		}
 		
-#if ENABLE_NON_NANITE_VSM
-
 		const bool bAllocatePageRectAtlas = CVarAllocatePagesUsingRects.GetValueOnRenderThread() != 0;
 		if (bAllocatePageRectAtlas)
 		{
@@ -988,8 +978,6 @@ void FVirtualShadowMapArray::BuildPageAllocations(
 			}
 		}
 		else
-#endif // ENABLE_NON_NANITE_VSM
-
 		{
 			// Note: does not actually need mip0 so can be trimmed down a bit
 			FRDGBufferRef CoverageSummary = GraphBuilder.CreateBuffer( FRDGBufferDesc::CreateStructuredDesc( 4, NumPageFlags ), TEXT("Shadow.Virtual.CoverageSummary") );
@@ -1074,13 +1062,9 @@ void FVirtualShadowMapArray::SetProjectionParameters(FRDGBuilder& GraphBuilder,
 {
 	OutParameters.VirtualSmCommon = GetCommonUniformBuffer(GraphBuilder);
 	OutParameters.ShadowPageTable = GraphBuilder.CreateSRV(PageTableRDG);
-#if ENABLE_NON_NANITE_VSM
 	// GPUCULL_TODO: Add flags to avoid sampling if not present?
 	OutParameters.PhysicalPagePool = PhysicalPagePoolRDG != nullptr ? PhysicalPagePoolRDG : GraphBuilder.RegisterExternalTexture(GSystemTextures.BlackDummy);
 	OutParameters.PhysicalPagePoolHw = PhysicalPagePoolHw != nullptr ? PhysicalPagePoolHw : GraphBuilder.RegisterExternalTexture(GSystemTextures.BlackDummy);
-#else //!ENABLE_NON_NANITE_VSM
-	OutParameters.PhysicalPagePool = PhysicalPagePoolRDG != nullptr ? PhysicalPagePoolRDG : GraphBuilder.RegisterExternalTexture(GSystemTextures.BlackDummy);
-#endif // ENABLE_NON_NANITE_VSM
 	OutParameters.VirtualShadowMapProjectionData = GraphBuilder.CreateSRV(ShadowMapProjectionDataRDG);
 }
 
@@ -1193,9 +1177,7 @@ class FVirtualSmPrintStatsCS : public FVirtualPageManagementShader
 		SHADER_PARAMETER_RDG_UNIFORM_BUFFER(FVirtualShadowMapCommonParameters, VirtualSmCommon)
 		SHADER_PARAMETER_STRUCT_INCLUDE( ShaderPrint::FShaderParameters, ShaderPrintStruct )
 		SHADER_PARAMETER_RDG_BUFFER_SRV(StructuredBuffer< uint >, InStatsBuffer)
-#if ENABLE_NON_NANITE_VSM
 		SHADER_PARAMETER_RDG_BUFFER_SRV(StructuredBuffer<FIntVector4>, AllocatedPageRectBounds)
-#endif // ENABLE_NON_NANITE_VSM
 	END_SHADER_PARAMETER_STRUCT()
 };
 IMPLEMENT_GLOBAL_SHADER(FVirtualSmPrintStatsCS, "/Engine/Private/VirtualShadowMaps/PrintStats.usf", "PrintStats", SF_Compute);
@@ -1303,9 +1285,6 @@ void FVirtualShadowMapArray::CreateMipViews( TArray<Nanite::FPackedView, SceneRe
 }
 
 
-#if ENABLE_NON_NANITE_VSM
-
-
 class FVirtualSmPrintClipmapStatsCS : public FVirtualPageManagementShader
 {
 	DECLARE_GLOBAL_SHADER(FVirtualSmPrintClipmapStatsCS);
@@ -1998,6 +1977,3 @@ void FVirtualShadowMapArray::AddInitializePhysicalPagesHwPass(FRDGBuilder& Graph
 	}
 
 }
-
-
-#endif // ENABLE_NON_NANITE_VSM
diff --git a/Engine/Source/Runtime/Renderer/Private/VirtualShadowMaps/VirtualShadowMapArray.h b/Engine/Source/Runtime/Renderer/Private/VirtualShadowMaps/VirtualShadowMapArray.h
--- a/Engine/Source/Runtime/Renderer/Private/VirtualShadowMaps/VirtualShadowMapArray.h
+++ b/Engine/Source/Runtime/Renderer/Private/VirtualShadowMaps/VirtualShadowMapArray.h
@@ -131,9 +131,7 @@ BEGIN_SHADER_PARAMETER_STRUCT(FVirtualShadowMapSamplingParameters, )
 	SHADER_PARAMETER_RDG_BUFFER_SRV(StructuredBuffer<uint>, ShadowPageTable)
 	SHADER_PARAMETER_RDG_TEXTURE(Texture2D<uint>, PhysicalPagePool)
 	SHADER_PARAMETER_RDG_BUFFER_SRV(StructuredBuffer< FVirtualShadowMapProjectionShaderData >, VirtualShadowMapProjectionData)
-#if ENABLE_NON_NANITE_VSM
 	SHADER_PARAMETER_RDG_TEXTURE(Texture2D<float>, PhysicalPagePoolHw)
-#endif // ENABLE_NON_NANITE_VSM
 END_SHADER_PARAMETER_STRUCT()
 
 /**
@@ -203,14 +201,12 @@ public:
 
 	void CreateMipViews( TArray<Nanite::FPackedView, SceneRenderingAllocator>& Views ) const;
 
-#if ENABLE_NON_NANITE_VSM
 	/**
 	 * Draw old-school hardware based shadow map tiles into virtual SM.
 	 */
 	void RenderVirtualShadowMapsHw(FRDGBuilder& GraphBuilder, const TArray<FProjectedShadowInfo*, SceneRenderingAllocator>& VirtualSmMeshCommandPasses, FScene& Scene);
 
 	void AddInitializePhysicalPagesHwPass(FRDGBuilder& GraphBuilder, FVirtualShadowMapArrayCacheManager* VirtualShadowMapArrayCacheManager);
-#endif
 
 	// Draw debug info into render target 'VSMDebug' of screen-size, the mode is controlled by 'r.Shadow.Virtual.DebugVisualize'.
 	void RenderDebugInfo(FRDGBuilder& GraphBuilder, FVirtualShadowMapArrayCacheManager *VirtualShadowMapArrayCacheManager);
@@ -225,11 +221,7 @@ public:
 		return GraphBuilder.CreateUniformBuffer(&CommonParameters);
 	}
 
-#if ENABLE_NON_NANITE_VSM
 	bool HasAnyShadowData() const { return PhysicalPagePoolRDG != nullptr || PhysicalPagePoolHw != nullptr;  }
-#else
-	bool HasAnyShadowData() const { return PhysicalPagePoolRDG != nullptr;  }
-#endif
 
 	void GetPageTableParameters(FRDGBuilder& GraphBuilder, FVirtualShadowMapPageTableParameters& OutParameters);
 
@@ -267,9 +259,7 @@ public:
 	// uint4 buffer with one rect for each mip level in all SMs, calculated to bound committed pages
 	// Used to clip the rect size of clusters during culling.
 	FRDGBufferRef PageRectBoundsRDG = nullptr;
-#if ENABLE_NON_NANITE_VSM
 	FRDGBufferRef AllocatedPageRectBoundsRDG = nullptr;
-#endif // ENABLE_NON_NANITE_VSM
 	FRDGBufferRef ShadowMapProjectionDataRDG = nullptr;
 	TArray< FRDGBufferRef, SceneRenderingAllocator > VirtualShadowMapIdRemapRDG;
 
@@ -287,9 +277,7 @@ public:
 	//FRDGBufferRef RDG = nullptr;
 	TRefCountPtr<IPooledRenderTarget>	DebugVisualizationProjectionOutput;
 	//FRDGBufferRef RDG = nullptr;
-#if ENABLE_NON_NANITE_VSM
 	// page pool for HW rasterized shadow data. 
 	// Mirrors the regular one, but uses a shadow depth target format
 	FRDGTextureRef PhysicalPagePoolHw = nullptr;
-#endif // ENABLE_NON_NANITE_VSM
 };
diff --git a/Engine/Source/Runtime/Renderer/Private/VirtualShadowMaps/VirtualShadowMapCacheManager.cpp b/Engine/Source/Runtime/Renderer/Private/VirtualShadowMaps/VirtualShadowMapCacheManager.cpp
--- a/Engine/Source/Runtime/Renderer/Private/VirtualShadowMaps/VirtualShadowMapCacheManager.cpp
+++ b/Engine/Source/Runtime/Renderer/Private/VirtualShadowMaps/VirtualShadowMapCacheManager.cpp
@@ -139,7 +139,7 @@ void FVirtualShadowMapArrayCacheManager::ExtractFrameData(FVirtualShadowMapArray
 		GraphBuilder.QueueBufferExtraction(VirtualShadowMapArray.HPageFlagsRDG, &PrevHPageFlags);
 		
 		GraphBuilder.QueueTextureExtraction(VirtualShadowMapArray.PhysicalPagePoolRDG, &PrevPhysicalPagePool);
-#if ENABLE_NON_NANITE_VSM
+
 		if( VirtualShadowMapArray.PhysicalPagePoolHw )
 		{
 			GraphBuilder.QueueTextureExtraction(VirtualShadowMapArray.PhysicalPagePoolHw, &PrevPhysicalPagePoolHw);
@@ -148,7 +148,7 @@ void FVirtualShadowMapArrayCacheManager::ExtractFrameData(FVirtualShadowMapArray
 		{
 			PrevPhysicalPagePoolHw = TRefCountPtr<IPooledRenderTarget>();
 		}
-#endif
+
 		GraphBuilder.QueueBufferExtraction(VirtualShadowMapArray.PhysicalPageMetaDataRDG, &PrevPhysicalPageMetaData);
 		GraphBuilder.QueueBufferExtraction(VirtualShadowMapArray.DynamicCasterPageFlagsRDG, &PrevDynamicCasterPageFlags);
 		GraphBuilder.QueueBufferExtraction(VirtualShadowMapArray.ShadowMapProjectionDataRDG, &PrevShadowMapProjectionDataBuffer);
@@ -165,9 +165,7 @@ void FVirtualShadowMapArrayCacheManager::ExtractFrameData(FVirtualShadowMapArray
 		PrevHPageFlags = TRefCountPtr<FRDGPooledBuffer>();
 
 		PrevPhysicalPagePool = TRefCountPtr<IPooledRenderTarget>();
-#if ENABLE_NON_NANITE_VSM
 		PrevPhysicalPagePoolHw = TRefCountPtr<IPooledRenderTarget>();
-#endif //ENABLE_NON_NANITE_VSM
 		PrevPhysicalPageMetaData = TRefCountPtr<FRDGPooledBuffer>();
 		PrevDynamicCasterPageFlags = TRefCountPtr<FRDGPooledBuffer>();
 		PrevShadowMapProjectionDataBuffer = TRefCountPtr<FRDGPooledBuffer>();
@@ -300,11 +298,7 @@ bool FVirtualShadowMapArrayCacheManager::IsValid()
 	return CVarCacheVirtualSMs.GetValueOnRenderThread() != 0
 		&& PrevPageTable
 		&& PrevPageFlags
-#if ENABLE_NON_NANITE_VSM
 		&& (PrevPhysicalPagePool || PrevPhysicalPagePoolHw)
-#else // !ENABLE_NON_NANITE_VSM
-		&& PrevPhysicalPagePool
-#endif //ENABLE_NON_NANITE_VSM
 		&& PrevPhysicalPageMetaData
 		&& PrevDynamicCasterPageFlags;
 }
diff --git a/Engine/Source/Runtime/Renderer/Private/VirtualShadowMaps/VirtualShadowMapCacheManager.h b/Engine/Source/Runtime/Renderer/Private/VirtualShadowMaps/VirtualShadowMapCacheManager.h
--- a/Engine/Source/Runtime/Renderer/Private/VirtualShadowMaps/VirtualShadowMapCacheManager.h
+++ b/Engine/Source/Runtime/Renderer/Private/VirtualShadowMaps/VirtualShadowMapCacheManager.h
@@ -135,10 +135,7 @@ public:
 	TRefCountPtr<IPooledRenderTarget>	PrevPhysicalPagePool;
 	TRefCountPtr<FRDGPooledBuffer>		PrevPhysicalPageMetaData;
 
-#if ENABLE_NON_NANITE_VSM
 	TRefCountPtr<IPooledRenderTarget>		PrevPhysicalPagePoolHw;
-#endif // ENABLE_NON_NANITE_VSM
-
 
 	TRefCountPtr<FRDGPooledBuffer>		PrevShadowMapProjectionDataBuffer;
 	TRefCountPtr<FRDGPooledBuffer>		PrevPageRectBounds;
diff --git a/Engine/Source/Runtime/Renderer/Private/VirtualShadowMaps/VirtualShadowMapConfig.h b/Engine/Source/Runtime/Renderer/Private/VirtualShadowMaps/VirtualShadowMapConfig.h
--- a/Engine/Source/Runtime/Renderer/Private/VirtualShadowMaps/VirtualShadowMapConfig.h
+++ b/Engine/Source/Runtime/Renderer/Private/VirtualShadowMaps/VirtualShadowMapConfig.h
@@ -3,11 +3,21 @@
 /*=============================================================================
 	VirtualShadowMapConfig.h
 =============================================================================*/
-
 #pragma once
 
+#include "RenderUtils.h"
+
+extern int32 GEnableNonNaniteVSM;
+
+/**
+ * Returns true if non-nanite virtual shadow maps are enbled by CVar r.Shadow.Virtual.NonNaniteVSM 
+ * and UseVirtualShadowMaps is true for the given platform and feature level.
+ */
+inline bool UseNonNaniteVirtualShadowMaps(EShaderPlatform ShaderPlatform, const FStaticFeatureLevel FeatureLevel)
+{
 #if GPUCULL_TODO
-	#define ENABLE_NON_NANITE_VSM 0
-#else
-	#define ENABLE_NON_NANITE_VSM 0
-#endif
+	return GEnableNonNaniteVSM != 0 && UseVirtualShadowMaps(ShaderPlatform, FeatureLevel);
+#else // !GPUCULL_TODO
+	return false;
+#endif // GPUCULL_TODO
+}
