Applies post process settings color grading based on showflag and r.Color.Grading (for ini files)

diff --git a/Engine/Source/Runtime/Renderer/Private/PostProcess/PostProcessCombineLUTs.cpp b/Engine/Source/Runtime/Renderer/Private/PostProcess/PostProcessCombineLUTs.cpp
--- a/Engine/Source/Runtime/Renderer/Private/PostProcess/PostProcessCombineLUTs.cpp
+++ b/Engine/Source/Runtime/Renderer/Private/PostProcess/PostProcessCombineLUTs.cpp
@@ -35,6 +35,11 @@ FAutoConsoleVariableRef CVarLUTSize(
 	TEXT("Size of film LUT"),
 	ECVF_RenderThreadSafe);
 
+TAutoConsoleVariable<int32> CVarColorGrading(
+	TEXT("r.Color.Grading"), 1,
+	TEXT("Controls whether post process settings's color grading settings should be applied."),
+	ECVF_RenderThreadSafe);
+
 // Including the neutral one at index 0
 const uint32 GMaxLUTBlendCount = 5;
 
@@ -137,8 +142,13 @@ void GetCombineLUTParameters(
 	check(Textures);
 	check(Weights);
 
-	const FPostProcessSettings& Settings = View.FinalPostProcessSettings;
+	static const FPostProcessSettings DefaultSettings;
+
 	const FSceneViewFamily& ViewFamily = *(View.Family);
+	const FPostProcessSettings& Settings = (
+		ViewFamily.EngineShowFlags.ColorGrading && CVarColorGrading.GetValueOnRenderThread() != 0)
+		? View.FinalPostProcessSettings
+		: DefaultSettings;
 
 	for (uint32 BlendIndex = 0; BlendIndex < BlendCount; ++BlendIndex)
 	{
