Implements r.Editor.TemporalUpsampleDepth

diff --git a/Engine/Shaders/Private/PostProcessCompositeEditorPrimitives.usf b/Engine/Shaders/Private/PostProcessCompositeEditorPrimitives.usf
--- a/Engine/Shaders/Private/PostProcessCompositeEditorPrimitives.usf
+++ b/Engine/Shaders/Private/PostProcessCompositeEditorPrimitives.usf
@@ -8,6 +8,66 @@
 #define MSAA_SAMPLE_COUNT 1
 #endif
 
+#define CONFIG_SAMPLES 6
+
+// K = Center of the nearest input pixel.
+// O = Center of the output pixel.
+//
+//          |           |
+//    0     |     1     |     2
+//          |           |
+//          |           |
+//  --------+-----------+--------
+//          |           |
+//          | O         |
+//    3     |     K     |     5
+//          |           |
+//          |           |
+//  --------+-----------+--------
+//          |           |
+//          |           |
+//    6     |     7     |     8
+//          |           |
+//
+static const int2 kOffsets3x3[9] =
+{
+	int2(-1, -1),
+	int2(0, -1),
+	int2(1, -1),
+	int2(-1,  0),
+	int2(0,  0), // K
+	int2(1,  0),
+	int2(-1,  1),
+	int2(0,  1),
+	int2(1,  1),
+};
+
+// T = Center of the nearest top left pixel input pixel.
+// O = Center of the output pixel.
+//
+//          | 
+//    T     |     .
+//          | 
+//       O  | 
+//  --------+--------
+//          | 
+//          | 
+//    .     |     .
+//          | 
+static const int2 Offsets2x2[4] =
+{
+	int2( 0,  0), // T
+	int2( 1,  0),
+	int2( 0,  1),
+	int2( 1,  1),
+};
+
+// Indexes of the 3x3 square.
+static const uint kSquareIndexes3x3[9] = { 4, 0, 1, 2, 3, 5, 6, 7, 8 };
+
+// Indexes of the offsets to have plus + shape.
+static const uint kPlusIndexes3x3[5] = { 4, 1, 3, 5, 7 };
+
 #if MSAA_SAMPLE_COUNT > 1
 Texture2DMS<float4, MSAA_SAMPLE_COUNT> EditorPrimitivesColor;
 Texture2DMS<float, MSAA_SAMPLE_COUNT> EditorPrimitivesDepth;
@@ -22,28 +82,185 @@ SamplerState ColorSampler;
 Texture2D DepthTexture;
 SamplerState DepthSampler;
 
+Texture2D VelocityTexture;
+SamplerState VelocitySampler;
+
+Texture2D PrevHistoryTexture;
+SamplerState PrevHistorySampler;
+
 SCREEN_PASS_TEXTURE_VIEWPORT(Color)
 SCREEN_PASS_TEXTURE_VIEWPORT(Depth)
+SCREEN_PASS_TEXTURE_VIEWPORT(PrevHistory)
+SCREEN_PASS_TEXTURE_VIEWPORT(History)
+SCREEN_PASS_TEXTURE_VIEWPORT(Output)
 SCREEN_PASS_TEXTURE_VIEWPORT_TRANSFORM(ColorToDepth)
 
+float2 DepthTextureJitter;
+uint bCameraCut;
+
+float4 SampleOffsetArray[MSAA_SAMPLE_COUNT];
+
 uint bOpaqueEditorGizmo;
 uint bCompositeAnyNonNullDepth;
 
+void MainTemporalUpsampleEditorDepthPS(
+	float4 SvPosition : SV_POSITION,
+	out float OutDeviceZ : SV_Target0)
+{
+	float2 ViewportUV = (SvPosition.xy - History_ViewportMin) * History_ViewportSizeInverse;
+	float2 ScreenPos = ViewportUVToScreenPos(ViewportUV);
+	
+	// Pixel coordinate of the center of output pixel O in the input viewport.
+	float2 PPCo = ViewportUV * Depth_ViewportSize + DepthTextureJitter;
+
+	// Pixel coordinate of the center of the nearest input pixel K.
+	float2 PPCk = floor(PPCo) + 0.5;
+	
+	// Sample nearest depth
+	float InputDeviceZ;
+	{
+		float2 SampleInputBufferUV = (Depth_ViewportMin + PPCk) * Depth_ExtentInverse;
+		SampleInputBufferUV = clamp(SampleInputBufferUV, Depth_UVViewportBilinearMin, Depth_UVViewportBilinearMax);
+			
+		InputDeviceZ = Texture2DSampleLevel(DepthTexture, DepthSampler, SampleInputBufferUV, 0).r;
+	}
+
+	// Compute the screen position to sample in history
+	float2 PrevScreenPos = ScreenPos;
+	#if 1
+	{
+		float4 ThisClip = float4(ScreenPos, InputDeviceZ, 1);
+		float4 PrevClip = mul(ThisClip, View.ClipToPrevClip);
+
+		PrevScreenPos = PrevClip.xy / PrevClip.w;
+
+		#if 0
+		{
+			float2 SampleInputBufferUV = (Depth_ViewportMin + PPCk) * Depth_ExtentInverse;
+			SampleInputBufferUV = clamp(SampleInputBufferUV, Depth_UVViewportBilinearMin, Depth_UVViewportBilinearMax);
+			
+			float4 EncodedVelocity = VelocityTexture.SampleLevel(VelocitySampler, SampleInputBufferUV, 0);
+			bool DynamicN = EncodedVelocity.x > 0.0;
+			if(DynamicN)
+			{
+				PrevScreenPos = ScreenPos - DecodeVelocityFromTexture(EncodedVelocity).xy;
+			}
+		}
+		#endif
+	}
+	#endif
+
+	// Sample the history.
+	float PrevHistoryDeviceZ;
+	{
+		float2 PrevHistoryBufferUV = (PrevHistory_ScreenPosToViewportScale * PrevScreenPos + PrevHistory_ScreenPosToViewportBias) * PrevHistory_ExtentInverse;
+		PrevHistoryBufferUV = clamp(PrevHistoryBufferUV, PrevHistory_UVViewportBilinearMin, PrevHistory_UVViewportBilinearMax);
+
+		PrevHistoryDeviceZ = Texture2DSampleLevel(PrevHistoryTexture, PrevHistorySampler, PrevHistoryBufferUV, 0).r;
+	}
+
+	// Correct the DeviceZ from the previous depth to the current depth.
+	float CorrectedPrevHistoryDeviceZ;
+	#if 1
+	{
+		float PrevHistoryDepth = ConvertFromDeviceZ(PrevHistoryDeviceZ);
+
+		const float3 PrevHistoryTranslatedWorldPosition = mul(float4(PrevScreenPos * PrevHistoryDepth, PrevHistoryDepth, 1), View.PrevScreenToTranslatedWorld).xyz + (View.PreViewTranslation.xyz - View.PrevPreViewTranslation.xyz);
+
+		float CorrectedPrevHistoryDepth = mul(float4(PrevHistoryTranslatedWorldPosition, 1.0), View.TranslatedWorldToView).z;
+		
+		CorrectedPrevHistoryDeviceZ = ConvertToDeviceZ(CorrectedPrevHistoryDepth);
+	}
+	#else
+	{
+		CorrectedPrevHistoryDeviceZ = PrevHistoryDeviceZ;
+	}
+	#endif
+
+	// Replaces NaN from history with 0.0
+	CorrectedPrevHistoryDeviceZ = -min(-CorrectedPrevHistoryDeviceZ, 0.0);
+
+	// Sample current frame
+	bool bUpdateHistory;
+	float InputMinDeviceZ;
+	float InputMaxDeviceZ;
+	{
+		// Vector in pixel between pixel K -> O.
+		float2 dKO = float2(PPCo - PPCk);
+		
+		bUpdateHistory = all(abs(dKO) < 0.5 * (Depth_ViewportSize.x * History_ViewportSizeInverse.x)) || bCameraCut != 0;
+
+		InputMinDeviceZ = InputDeviceZ;
+		InputMaxDeviceZ = InputDeviceZ;
+
+		UNROLL_N(CONFIG_SAMPLES - 1)
+		for (uint SampleId = 1; SampleId < CONFIG_SAMPLES; SampleId++)
+		{
+			float2 PixelOffset;
+			
+			#if CONFIG_SAMPLES == 9
+			{
+				PixelOffset = kOffsets3x3[kSquareIndexes3x3[SampleId]];
+			}
+			#elif CONFIG_SAMPLES == 5 || CONFIG_SAMPLES == 6
+			{
+				if (SampleId == 5)
+				{
+					PixelOffset = SignFastInt(float2(dKO));
+				}
+				else
+				{
+					PixelOffset = kOffsets3x3[kPlusIndexes3x3[SampleId]];
+				}
+			}
+			#else
+				#error Unknown sample count
+			#endif
+			
+			float2 SampleInputBufferUV = (Depth_ViewportMin + PPCk + PixelOffset) * Depth_ExtentInverse;
+			SampleInputBufferUV = clamp(SampleInputBufferUV, Depth_UVViewportBilinearMin, Depth_UVViewportBilinearMax);
+
+			float SampleDeviceZ = Texture2DSampleLevel(DepthTexture, DepthSampler, SampleInputBufferUV, 0).r;
+			
+			InputMinDeviceZ = min(InputMinDeviceZ, SampleDeviceZ);
+			InputMaxDeviceZ = max(InputMaxDeviceZ, SampleDeviceZ);
+		}
+	}
+
+	// Reject history based on neighborhood
+	float ClampedHistoryDeviceZ = clamp(CorrectedPrevHistoryDeviceZ, InputMinDeviceZ, InputMaxDeviceZ);
+
+	// Output history
+	float FinalHistoryDeviceZ = bUpdateHistory ? InputDeviceZ : ClampedHistoryDeviceZ;
+	
+	OutDeviceZ = -min(-FinalHistoryDeviceZ, 0.0);
+}
+
 void MainPopulateSceneDepthPS(
 	#if USE_MSAA && !COMPILER_HLSLCC
-		sample noperspective float4 UVAndScreenPos : TEXCOORD0,
+		sample noperspective float4 SvPosition : SV_POSITION,
 	#else
-		noperspective float4 UVAndScreenPos : TEXCOORD0,
+		noperspective float4 SvPosition : SV_POSITION,
 	#endif
 	out float4 OutColor : SV_Target0,
 	out float OutDepth : SV_DEPTH)
 {
+	float2 ViewportUV = (SvPosition.xy - Color_ViewportMin) * Color_ViewportSizeInverse;
+	
 	float2 DepthUV = clamp(
-		UVAndScreenPos.xy - View.TemporalAAParams.zw * Depth_ExtentInverse,
+		(Depth_ViewportMin + ViewportUV * Depth_ViewportSize - DepthTextureJitter) * Depth_ExtentInverse,
 		Depth_UVViewportBilinearMin, Depth_UVViewportBilinearMax);
 
 	OutDepth = Texture2DSampleLevel(DepthTexture, DepthSampler, DepthUV, 0).r;
-	OutColor = 0;
+	OutColor = 0.0;
+}
+
+float ComputeDepthMask(float SceneDeviceZ, float EditorDeviceZ)
+{
+	// Soft Bias with SceneDeviceZ for best quality
+	const float DeviceDepthFade = 0.00005f;
+
+	return saturate(1.0f - (SceneDeviceZ - EditorDeviceZ) / DeviceDepthFade);
 }
 
 void ResolveEditorPrimitiveColor(int2 PixelPos, out float4 OutColor, out float OutDeviceZ)
@@ -51,88 +268,109 @@ void ResolveEditorPrimitiveColor(int2 PixelPos, out float4 OutColor, out float O
 	// Furthest device Z or 0 if there is none.
 	OutDeviceZ = 0;
 	OutColor = 0;
+	
+
+	// Bring out of premultiplied.
+	OutColor.rgb /= max(OutColor.a, 0.0001f);
+	// Fix gamma.
+	OutColor.rgb = pow(OutColor.rgb, 1.0f / 2.2f);
+	// Bring back to premultiplied
+	OutColor.rgb *= OutColor.a;
+}
+
+void MainCompositeEditorPrimitivesPS(
+	float4 SvPosition : SV_POSITION,
+	out float4 OutColor : SV_Target0)
+{
+	float2 ViewportUV = (SvPosition.xy - Output_ViewportMin) * Output_ViewportSizeInverse;
+	
+	const float2 ColorUV = (ViewportUV * Color_ViewportSize + Color_ViewportMin) * Color_ExtentInverse;
+	const int2 ColorPixelPos = int2(ColorUV * Color_Extent);
+	
+	float4 SceneColor = Texture2DSample(ColorTexture, ColorSampler, ColorUV);
+
+#if OUTPUT_SRGB_BUFFER
+	SceneColor.rgb = pow(SceneColor.rgb, 1.0f / 2.2f);
+#endif
+
+	// Resolve editor primitive scene color and depth.
+	float4 EditorPrimitiveColor;
+	float DepthMask;
 
 	#if MSAA_SAMPLE_COUNT > 1
 	{
-		float SumWeight = 0.00001f;
+		EditorPrimitiveColor = 0.0;
+		DepthMask = 0.0;
+		float DepthMaskWeight = 0.0;
 
+		UNROLL_N(MSAA_SAMPLE_COUNT)
 		for (uint SampleIndex = 0; SampleIndex < MSAA_SAMPLE_COUNT; ++SampleIndex)
 		{
-			float4 Sample = EditorPrimitivesColor.Load(PixelPos, SampleIndex);
-			float SampleDeviceZ = EditorPrimitivesDepth.Load(PixelPos, SampleIndex).r;
+			const float2 SampleOffset = SampleOffsetArray[SampleIndex].xy;
+
+			float2 SampleDepthUV = (ViewportUV * Depth_ViewportSize + Depth_ViewportMin + SampleOffset - DepthTextureJitter) * Depth_ExtentInverse;
+			SampleDepthUV = clamp(SampleDepthUV, Depth_UVViewportBilinearMin, Depth_UVViewportBilinearMax);
+			float SampleSceneDeviceZ = Texture2DSampleLevel(DepthTexture, DepthSampler, SampleDepthUV, 0).r;
+
+			float4 Sample = EditorPrimitivesColor.Load(ColorPixelPos, SampleIndex);
+			float SampleEditorDeviceZ = EditorPrimitivesDepth.Load(ColorPixelPos, SampleIndex).r;
 
 			// Check if any color was applied to this pixel. Note: This prevents actual black pixels from being visible.
 			float Weight = Sample.a;
 
-			if (SampleDeviceZ > 0.0 && bCompositeAnyNonNullDepth)
+			float SampleDepthMask = ComputeDepthMask(SampleSceneDeviceZ, SampleEditorDeviceZ);
+
+			if (SampleEditorDeviceZ > 0.0 && bCompositeAnyNonNullDepth)
 			{
 				Weight = 1;
 			}
 
-			FLATTEN if(Weight)
+			FLATTEN
+			if (Weight)
 			{
-				OutDeviceZ = max(OutDeviceZ, SampleDeviceZ);
-
-				OutColor += float4(Sample.rgb, Weight);
+				EditorPrimitiveColor += float4(Sample.rgb, Weight);
 
-				SumWeight += Weight;
+				DepthMask += SampleDepthMask * Weight;
+				DepthMaskWeight += Weight;
 			}
 		}
 
-		OutColor.rgb /= MSAA_SAMPLE_COUNT;
-		OutColor.a /= MSAA_SAMPLE_COUNT;
+		EditorPrimitiveColor.rgb /= MSAA_SAMPLE_COUNT;
+		EditorPrimitiveColor.a /= MSAA_SAMPLE_COUNT;
+		
+		DepthMask *= DepthMaskWeight > 0 ? rcp(DepthMaskWeight) : 0.0;
 	}
 	#else
 	{
-		OutColor = EditorPrimitivesColor.Load(int3(PixelPos, 0));
-		OutDeviceZ = EditorPrimitivesDepth.Load(int3(PixelPos, 0)).r;
+		// De-jitter the sample position and make a filtered lookup - for planes this allows to reconstruct a much less jittery depth comparison function. It doesn't fix silhouettes, however.
+		float2 SampleDepthUV = (ViewportUV * Depth_ViewportSize + Depth_ViewportMin - DepthTextureJitter) * Depth_ExtentInverse;
+		SampleDepthUV = clamp(SampleDepthUV, Depth_UVViewportBilinearMin, Depth_UVViewportBilinearMax);
+
+		float SceneDeviceZ = Texture2DSampleLevel(DepthTexture, DepthSampler, SampleDepthUV, 0).r;
+		float EditorDeviceZ = EditorPrimitivesDepth.Load(int3(ColorPixelPos, 0)).r;
+		
+		EditorPrimitiveColor = EditorPrimitivesColor.Load(int3(ColorPixelPos, 0));
+		DepthMask = ComputeDepthMask(SceneDeviceZ, EditorDeviceZ);
 
-		if (OutDeviceZ > 0.0 && bCompositeAnyNonNullDepth)
+		if (EditorDeviceZ > 0.0 && bCompositeAnyNonNullDepth)
 		{
-			OutColor.a = 1;
+			EditorPrimitiveColor.a = 1;
 		}
 	}
 	#endif
-
-	// Bring out of premultiplied.
-	OutColor.rgb /= max(OutColor.a, 0.0001f);
-	// Fix gamma.
-	OutColor.rgb = pow(OutColor.rgb, 1.0f / 2.2f);
-	// Bring back to premultiplied
-	OutColor.rgb *= OutColor.a;
-}
-
-void MainCompositeEditorPrimitivesPS(
-	noperspective float4 UVAndScreenPos : TEXCOORD0,
-	float4 SvPosition : SV_POSITION,
-	out float4 OutColor : SV_Target0)
-{
-	const float2 ColorUV = UVAndScreenPos.xy;
-	const int2 ColorPixelPos = int2(ColorUV * Color_Extent);
-
-	float4 Color = Texture2DSample(ColorTexture, ColorSampler, ColorUV);
-
-#if OUTPUT_SRGB_BUFFER
-	Color.rgb = pow(Color.rgb, 1.0f / 2.2f);
-#endif
-
-	// Resolve editor primitive scene color and depth.
-	float4 EditorPrimitiveColor;
-	float EditorMaxDeviceZ;
-	ResolveEditorPrimitiveColor(ColorPixelPos, /* out = */ EditorPrimitiveColor, /* out = */ EditorMaxDeviceZ);
+	
+	// Fixes the gama of editor primitives
+	{
+		// Bring out of premultiplied.
+		EditorPrimitiveColor.rgb /= max(EditorPrimitiveColor.a, 0.0001f);
+		// Fix gamma.
+		EditorPrimitiveColor.rgb = pow(EditorPrimitiveColor.rgb, 1.0f / 2.2f);
+		// Bring back to premultiplied
+		EditorPrimitiveColor.rgb *= EditorPrimitiveColor.a;
+	}
 
 	float2 DepthUV = ColorUV * ColorToDepth_Scale + ColorToDepth_Bias;
 
-	// De-jitter the sample position and make a filtered lookup - for planes this allows to reconstruct a much less jittery depth comparison function. It doesn't fix silhouettes, however.
-	DepthUV -= View.TemporalAAParams.zw * Depth_ExtentInverse;
-	DepthUV = clamp(DepthUV, Depth_UVViewportBilinearMin, Depth_UVViewportBilinearMax);
-
-	float SceneDeviceZ = Texture2DSampleLevel(DepthTexture, DepthSampler, DepthUV, 0).r;
-
-	// Soft Bias with SceneDeviceZ for best quality
-	const float DeviceDepthFade = 0.00005f;
-	float DepthMask = saturate(1.0f - (SceneDeviceZ - EditorMaxDeviceZ) / DeviceDepthFade);
-
 	// Generate 2x2 square tiling pattern for foreground primitive that end up behind scene opaque primitives.
 	float PatternMask = ((ColorPixelPos.x/2 + ColorPixelPos.y/2) % 2) * 0.7f;
 
@@ -145,10 +383,10 @@ void MainCompositeEditorPrimitivesPS(
 	float DarkenMask = lerp(LowContrastPatternMask, 1.0f, DepthMask);
 
 	// Blend editor ptimitives with scene color.
-	OutColor.rgb = Color.rgb * (1 - EditorPrimitiveColor.a) + EditorPrimitiveColor.rgb * (DarkenMask * HiddenMask);
+	OutColor.rgb = SceneColor.rgb * (1 - EditorPrimitiveColor.a) + EditorPrimitiveColor.rgb * (DarkenMask * HiddenMask);
 
 	#if POST_PROCESS_ALPHA
-		OutColor.a = lerp(Color.a, 1, EditorPrimitiveColor.a * DarkenMask * HiddenMask);
+		OutColor.a = lerp(SceneColor.a, 1, EditorPrimitiveColor.a * DarkenMask * HiddenMask);
 	#else
 		OutColor.a = 0;
 	#endif
diff --git a/Engine/Source/Editor/UnrealEd/Private/EditorViewportClient.cpp b/Engine/Source/Editor/UnrealEd/Private/EditorViewportClient.cpp
--- a/Engine/Source/Editor/UnrealEd/Private/EditorViewportClient.cpp
+++ b/Engine/Source/Editor/UnrealEd/Private/EditorViewportClient.cpp
@@ -83,6 +83,12 @@ static TAutoConsoleVariable<int32> CVarEditorViewportTest(
 	TEXT("1..7: Various Configuations"),
 	ECVF_RenderThreadSafe);
 
+static TAutoConsoleVariable<float> CVarEditorViewportScreenPercentage(
+	TEXT("r.Editor.ViewportScreenPercentage"),
+	100.0f,
+	TEXT("Default screen percentage of the editor's viewports"),
+	ECVF_RenderThreadSafe);
+
 static bool GetDefaultLowDPIPreviewValue()
 {
 	static auto CVarEnableEditorScreenPercentageOverride = IConsoleManager::Get().FindConsoleVariable(TEXT("Editor.OverrideDPIBasedEditorViewportScaling"));
@@ -443,7 +449,7 @@ FEditorViewportClient::FEditorViewportClient(FEditorModeTools* InModeTools, FPre
 	, MovingPreviewLightSavedScreenPos(ForceInitToZero)
 	, MovingPreviewLightTimer(0.0f)
 	, bLockFlightCamera(false)
-	, PreviewResolutionFraction(1.0f)
+	, PreviewResolutionFraction(CVarEditorViewportScreenPercentage.GetValueOnGameThread() / 100.0f)
 	, SceneDPIMode(ESceneDPIMode::EditorDefault)
 	, PerspViewModeIndex(DefaultPerspectiveViewMode)
 	, OrthoViewModeIndex(DefaultOrthoViewMode)
@@ -1231,6 +1237,63 @@ FSceneView* FEditorViewportClient::CalcSceneView(FSceneViewFamily* ViewFamily, c
 	ViewInitOptions.OverrideFarClippingPlaneDistance = FarPlane;
 	ViewInitOptions.CursorPos = CurrentMousePos;
 
+	{
+		static const auto CVarVSync = IConsoleManager::Get().FindTConsoleVariableDataInt(TEXT("r.Test.ConstrainedView"));
+		int32 Value = CVarVSync->GetValueOnGameThread();
+
+		if (Value)
+		{
+			const FIntRect& ViewRect = ViewInitOptions.GetViewRect();
+			FIntRect ConstrainedViewRect = ViewInitOptions.GetConstrainedViewRect();
+
+			int InsetX = ConstrainedViewRect.Width() / 4;
+			int InsetY = ConstrainedViewRect.Height() / 4;
+
+			// this allows to test various typical view port situations
+			switch (Value)
+			{
+			case 1:
+				ConstrainedViewRect.Min.X += InsetX;
+				ConstrainedViewRect.Min.Y += InsetY;
+				ConstrainedViewRect.Max.X -= InsetX;
+				ConstrainedViewRect.Max.Y -= InsetY;
+				break;
+
+			case 2:
+				ConstrainedViewRect.Min.Y += InsetY;
+				ConstrainedViewRect.Max.Y -= InsetY;
+				break;
+
+			case 3:
+				ConstrainedViewRect.Min.X += InsetX;
+				ConstrainedViewRect.Max.X -= InsetX;
+				break;
+
+			case 4:
+				ConstrainedViewRect.Max.X -= 2 * InsetX;
+				ConstrainedViewRect.Max.Y -= 2 * InsetY;
+				break;
+
+			case 5:
+				ConstrainedViewRect.Min.X += 2 * InsetX;
+				ConstrainedViewRect.Max.Y -= 2 * InsetY;
+				break;
+
+			case 6:
+				ConstrainedViewRect.Max.X -= 2 * InsetX;
+				ConstrainedViewRect.Min.Y += 2 * InsetY;
+				break;
+
+			case 7:
+				ConstrainedViewRect.Min.X += 2 * InsetX;
+				ConstrainedViewRect.Min.Y += 2 * InsetY;
+				break;
+			}
+
+			ViewInitOptions.SetConstrainedViewRectangle(ConstrainedViewRect);
+		}
+	}
+
 	FSceneView* View = new FSceneView(ViewInitOptions);
 
 	View->ViewLocation = ModifiedViewLocation;
diff --git a/Engine/Source/Runtime/Renderer/Private/PostProcess/PostProcessCompositeEditorPrimitives.cpp b/Engine/Source/Runtime/Renderer/Private/PostProcess/PostProcessCompositeEditorPrimitives.cpp
--- a/Engine/Source/Runtime/Renderer/Private/PostProcess/PostProcessCompositeEditorPrimitives.cpp
+++ b/Engine/Source/Runtime/Renderer/Private/PostProcess/PostProcessCompositeEditorPrimitives.cpp
@@ -8,6 +8,8 @@
 #include "BasePassRendering.h"
 #include "MobileBasePassRendering.h"
 #include "SceneRenderingUtils.h"
+#include "PixelShaderUtils.h"
+#include "ScenePrivate.h"
 
 namespace
 {
@@ -17,6 +19,42 @@ TAutoConsoleVariable<float> CVarEditorOpaqueGizmo(
 	TEXT("0..1\n0: occluded gizmo is partly transparent (default), 1:gizmo is never occluded"),
 	ECVF_RenderThreadSafe);
 
+TAutoConsoleVariable<int32> CVarEditorTemporalUpsampleDepth(
+	TEXT("r.Editor.TemporalUpsampleDepth"), 2,
+	TEXT("Temporal upsample factor of the depth buffer for depth testing editor primitives against."),
+	ECVF_RenderThreadSafe);
+
+class FTemporalUpsampleEditorDepthPS : public FGlobalShader
+{
+public:
+	DECLARE_GLOBAL_SHADER(FTemporalUpsampleEditorDepthPS);
+	SHADER_USE_PARAMETER_STRUCT(FTemporalUpsampleEditorDepthPS, FGlobalShader);
+
+	BEGIN_SHADER_PARAMETER_STRUCT(FParameters, )
+		SHADER_PARAMETER_STRUCT_REF(FViewUniformShaderParameters, View)
+		SHADER_PARAMETER_STRUCT(FScreenPassTextureViewportParameters, Depth)
+		SHADER_PARAMETER_STRUCT(FScreenPassTextureViewportParameters, PrevHistory)
+		SHADER_PARAMETER_STRUCT(FScreenPassTextureViewportParameters, History)
+		SHADER_PARAMETER(FVector2D, DepthTextureJitter)
+		SHADER_PARAMETER(int32, bCameraCut)
+
+		SHADER_PARAMETER_RDG_TEXTURE(Texture2D, DepthTexture)
+		SHADER_PARAMETER_SAMPLER(SamplerState, DepthSampler)
+		SHADER_PARAMETER_RDG_TEXTURE(Texture2D, VelocityTexture)
+		SHADER_PARAMETER_SAMPLER(SamplerState, VelocitySampler)
+		SHADER_PARAMETER_RDG_TEXTURE(Texture2D, PrevHistoryTexture)
+		SHADER_PARAMETER_SAMPLER(SamplerState, PrevHistorySampler)
+
+		RENDER_TARGET_BINDING_SLOTS()
+	END_SHADER_PARAMETER_STRUCT()
+
+	static bool ShouldCompilePermutation(const FGlobalShaderPermutationParameters& Parameters)
+	{
+		// Only PC platforms render editor primitives.
+		return IsPCPlatform(Parameters.Platform);
+	}
+};
+
 class FPopulateEditorDepthPS : public FGlobalShader
 {
 public:
@@ -28,7 +66,9 @@ public:
 
 	BEGIN_SHADER_PARAMETER_STRUCT(FParameters, )
 		SHADER_PARAMETER_STRUCT_REF(FViewUniformShaderParameters, View)
+		SHADER_PARAMETER_STRUCT(FScreenPassTextureViewportParameters, Color)
 		SHADER_PARAMETER_STRUCT(FScreenPassTextureViewportParameters, Depth)
+		SHADER_PARAMETER(FVector2D, DepthTextureJitter)
 		SHADER_PARAMETER_RDG_TEXTURE(Texture2D, DepthTexture)
 		SHADER_PARAMETER_SAMPLER(SamplerState, DepthSampler)
 		RENDER_TARGET_BINDING_SLOTS()
@@ -50,8 +90,6 @@ public:
 	}
 };
 
-IMPLEMENT_GLOBAL_SHADER(FPopulateEditorDepthPS, "/Engine/Private/PostProcessCompositeEditorPrimitives.usf", "MainPopulateSceneDepthPS", SF_Pixel);
-
 class FCompositeEditorPrimitivesPS : public FEditorPrimitiveShader
 {
 public:
@@ -62,7 +100,8 @@ public:
 		SHADER_PARAMETER_STRUCT_REF(FViewUniformShaderParameters, View)
 		SHADER_PARAMETER_STRUCT(FScreenPassTextureViewportParameters, Color)
 		SHADER_PARAMETER_STRUCT(FScreenPassTextureViewportParameters, Depth)
-		SHADER_PARAMETER_STRUCT(FScreenPassTextureViewportTransform, ColorToDepth)
+		SHADER_PARAMETER_STRUCT(FScreenPassTextureViewportParameters, Output)
+		SHADER_PARAMETER_ARRAY(FVector4, SampleOffsetArray, [FEditorPrimitiveShader::kEditorMSAASampleCountMax])
 		SHADER_PARAMETER_RDG_TEXTURE(Texture2D, EditorPrimitivesDepth)
 		SHADER_PARAMETER_RDG_TEXTURE(Texture2D, EditorPrimitivesColor)
 		SHADER_PARAMETER_RDG_TEXTURE(Texture2D, ColorTexture)
@@ -71,6 +110,7 @@ public:
 		SHADER_PARAMETER_SAMPLER(SamplerState, DepthSampler)
 		SHADER_PARAMETER(uint32, bOpaqueEditorGizmo)
 		SHADER_PARAMETER(uint32, bCompositeAnyNonNullDepth)
+		SHADER_PARAMETER(FVector2D, DepthTextureJitter)
 		RENDER_TARGET_BINDING_SLOTS()
 	END_SHADER_PARAMETER_STRUCT()
 
@@ -81,6 +121,8 @@ public:
 	}
 };
 
+IMPLEMENT_GLOBAL_SHADER(FTemporalUpsampleEditorDepthPS, "/Engine/Private/PostProcessCompositeEditorPrimitives.usf", "MainTemporalUpsampleEditorDepthPS", SF_Pixel);
+IMPLEMENT_GLOBAL_SHADER(FPopulateEditorDepthPS, "/Engine/Private/PostProcessCompositeEditorPrimitives.usf", "MainPopulateSceneDepthPS", SF_Pixel);
 IMPLEMENT_GLOBAL_SHADER(FCompositeEditorPrimitivesPS, "/Engine/Private/PostProcessCompositeEditorPrimitives.usf", "MainCompositeEditorPrimitivesPS", SF_Pixel);
 
 void RenderEditorPrimitives(FRHICommandListImmediate& RHICmdList, const FViewInfo& View, FMeshPassProcessorRenderState& DrawRenderState, FInstanceCullingManager& InstanceCullingManager)
@@ -247,8 +289,6 @@ FScreenPassTexture AddEditorPrimitivePass(
 	check(Inputs.SceneDepth.IsValid());
 	check(Inputs.BasePassType != FEditorPrimitiveInputs::EBasePassType::MAX);
 
-	RDG_EVENT_SCOPE(GraphBuilder, "CompositeEditorPrimitives");
-
 	const FSceneTextures& SceneTextures = FSceneTextures::Get(GraphBuilder);
 	const uint32 NumSamples = SceneTextures.Config.EditorPrimitiveNumSamples;
 	const FViewInfo* EditorView = CreateEditorPrimitiveView(View, Inputs.SceneColor.ViewRect, NumSamples);
@@ -288,47 +328,145 @@ FScreenPassTexture AddEditorPrimitivePass(
 			1,
 			Config.EditorPrimitiveNumSamples);
 
-		EditorPrimitiveColor = GraphBuilder.CreateTexture(ColorDesc, TEXT("Editor.PrimitivesDepth"));
-		EditorPrimitiveDepth = GraphBuilder.CreateTexture(DepthDesc, TEXT("Editor.PrimitivesColor"));
+		EditorPrimitiveColor = GraphBuilder.CreateTexture(ColorDesc, TEXT("Editor.PrimitivesColor"));
+		EditorPrimitiveDepth = GraphBuilder.CreateTexture(DepthDesc, TEXT("Editor.PrimitivesDepth"));
 	}
 
 	// Load the color target if it already exists.
 	const FScreenPassTextureViewport EditorPrimitivesViewport(EditorPrimitiveColor, Inputs.SceneColor.ViewRect);
 
+	RDG_EVENT_SCOPE(GraphBuilder, "CompositeEditorPrimitives %dx%d MSAA=%d",
+		EditorPrimitivesViewport.Rect.Width(),
+		EditorPrimitivesViewport.Rect.Height(),
+		NumSamples);
+
+	FScreenPassTexture SceneDepth = Inputs.SceneDepth;
+	FVector2D SceneDepthJitter = View.TemporalJitterPixels;
+
 	// The editor primitive composition pass is also used when rendering VMI_WIREFRAME in order to use MSAA.
 	// So we need to check whether the editor primitives are enabled inside this function.
 	if (View.Family->EngineShowFlags.CompositeEditorPrimitives)
 	{
-		const FScreenPassTextureViewport SceneDepthViewport(Inputs.SceneDepth);
-
 		// Populate depth if a prior pass did not already do it.
 		if (!bProducedByPriorPass)
 		{
-			FPopulateEditorDepthPS::FParameters* PassParameters = GraphBuilder.AllocParameters<FPopulateEditorDepthPS::FParameters>();
-			PassParameters->View = View.ViewUniformBuffer;
-			PassParameters->Depth = GetScreenPassTextureViewportParameters(SceneDepthViewport);
-			PassParameters->DepthTexture = Inputs.SceneDepth.Texture;
-			PassParameters->DepthSampler = TStaticSamplerState<SF_Point, AM_Clamp, AM_Clamp, AM_Clamp>::GetRHI();
-			PassParameters->RenderTargets[0] = FRenderTargetBinding(EditorPrimitiveColor, ERenderTargetLoadAction::EClear);
-			PassParameters->RenderTargets.DepthStencil = FDepthStencilBinding(EditorPrimitiveDepth, ERenderTargetLoadAction::EClear, ERenderTargetLoadAction::EClear, FExclusiveDepthStencil::DepthWrite_StencilWrite);
-
-			FPopulateEditorDepthPS::FPermutationDomain PermutationVector;
-			PermutationVector.Set<FPopulateEditorDepthPS::FUseMSAADimension>(NumSamples > 1);
-			TShaderMapRef<FPopulateEditorDepthPS> PopulateDepthPixelShader(View.ShaderMap, PermutationVector);
-			TShaderMapRef<FScreenPassVS> PopulateDepthVertexShader(View.ShaderMap);
-
-			AddDrawScreenPass(
-				GraphBuilder,
-				RDG_EVENT_NAME("PopulateDepth"),
-				View,
-				EditorPrimitivesViewport,
-				SceneDepthViewport,
-				PopulateDepthVertexShader,
-				PopulateDepthPixelShader,
-				TStaticDepthStencilState<true, CF_Always>::GetRHI(),
-				PassParameters);
+			// Upscale factor of the depth buffer that might be needed.
+			const float UpscaleFactor = float(EditorPrimitivesViewport.Rect.Width()) / float(Inputs.SceneDepth.ViewRect.Width());
+
+			// Upscale factor shouldn't be higher than there is TAA samples, or that means there will be unrendered pixels.
+			const int32 ComputeMaxUpsampleFactorDueToTAA = FMath::FloorToInt(FMath::Sqrt(float(View.TemporalJitterSequenceLength) / (UpscaleFactor * UpscaleFactor)));
+			
+			const int32 ComputeMaxUpsampleFactor = FMath::Clamp(ComputeMaxUpsampleFactorDueToTAA, 0, 4);
+
+			const int32 DepthUpsampleFactor = FMath::Clamp(CVarEditorTemporalUpsampleDepth.GetValueOnRenderThread(), 0, ComputeMaxUpsampleFactor);
+
+			// Upsample the depth at higher resolution to reduce depth intersection instability of editor primitives.
+			if (DepthUpsampleFactor > 0 && View.ViewState && View.AntiAliasingMethod == AAM_TemporalAA)
+			{
+				FScreenPassTexture History;
+				{
+					const FRDGTextureDesc Desc = FRDGTextureDesc::Create2D(
+						EditorPrimitivesViewport.Extent * DepthUpsampleFactor,
+						PF_R32_FLOAT,
+						FClearValueBinding::None,
+						TexCreate_ShaderResource | TexCreate_RenderTargetable);
+
+					History.Texture = GraphBuilder.CreateTexture(Desc, TEXT("Editor.PrimitivesDepthHistory"));
+					History.ViewRect = EditorPrimitivesViewport.Rect * DepthUpsampleFactor;
+				}
+
+				FTemporalUpsampleEditorDepthPS::FParameters* PassParameters = GraphBuilder.AllocParameters<FTemporalUpsampleEditorDepthPS::FParameters>();
+				PassParameters->View = View.ViewUniformBuffer;
+				PassParameters->Depth = GetScreenPassTextureViewportParameters(FScreenPassTextureViewport(Inputs.SceneDepth));
+				PassParameters->History = GetScreenPassTextureViewportParameters(FScreenPassTextureViewport(History));
+				PassParameters->DepthTextureJitter = View.TemporalJitterPixels;
+
+				PassParameters->DepthTexture = Inputs.SceneDepth.Texture;
+				PassParameters->DepthSampler = TStaticSamplerState<SF_Point, AM_Clamp, AM_Clamp, AM_Clamp>::GetRHI();
+
+				// TODO
+				PassParameters->VelocityTexture = GraphBuilder.RegisterExternalTexture(GSystemTextures.BlackDummy);
+				PassParameters->VelocitySampler = TStaticSamplerState<SF_Point, AM_Clamp, AM_Clamp, AM_Clamp>::GetRHI();
+
+				if (View.PrevViewInfo.EditorPrimtiveDepthHistory.IsValid())
+				{
+					PassParameters->PrevHistory = GetScreenPassTextureViewportParameters(FScreenPassTextureViewport(
+						View.PrevViewInfo.EditorPrimtiveDepthHistory.RT[0]->GetDesc().Extent, View.PrevViewInfo.EditorPrimtiveDepthHistory.ViewportRect));
+					PassParameters->PrevHistoryTexture = GraphBuilder.RegisterExternalTexture(View.PrevViewInfo.EditorPrimtiveDepthHistory.RT[0]);
+					PassParameters->PrevHistorySampler = TStaticSamplerState<SF_Bilinear, AM_Clamp, AM_Clamp, AM_Clamp>::GetRHI();
+					PassParameters->bCameraCut = false;
+				}
+				else
+				{
+					PassParameters->PrevHistory = GetScreenPassTextureViewportParameters(FScreenPassTextureViewport(
+						FIntPoint(1, 1), FIntRect(FIntPoint(0, 0), FIntPoint(1, 1))));
+					PassParameters->PrevHistoryTexture = GraphBuilder.RegisterExternalTexture(GSystemTextures.BlackDummy);
+					PassParameters->PrevHistorySampler = TStaticSamplerState<SF_Point, AM_Clamp, AM_Clamp, AM_Clamp>::GetRHI();
+					PassParameters->bCameraCut = true;
+				}
+
+				PassParameters->RenderTargets[0] = FRenderTargetBinding(History.Texture, ERenderTargetLoadAction::ENoAction);
+
+				TShaderMapRef<FTemporalUpsampleEditorDepthPS> PixelShader(View.ShaderMap);
+				FPixelShaderUtils::AddFullscreenPass(
+					GraphBuilder,
+					View.ShaderMap,
+					RDG_EVENT_NAME("TemporalUpsampleDepth %dx%d -> %dx%d",
+						Inputs.SceneDepth.ViewRect.Width(),
+						Inputs.SceneDepth.ViewRect.Height(),
+						History.ViewRect.Width(),
+						History.ViewRect.Height()),
+					PixelShader,
+					PassParameters,
+					History.ViewRect);
+
+				if (!View.bStatePrevViewInfoIsReadOnly)
+				{
+					FTemporalAAHistory* OutputHistory = &View.ViewState->PrevFrameViewInfo.EditorPrimtiveDepthHistory;
+					OutputHistory->SafeRelease();
+
+					GraphBuilder.QueueTextureExtraction(History.Texture, &OutputHistory->RT[0]);
+					OutputHistory->ViewportRect = History.ViewRect;
+					OutputHistory->ReferenceBufferSize = History.Texture->Desc.Extent;
+				}
+
+				SceneDepth = History;
+				SceneDepthJitter = FVector2D::ZeroVector;
+			}
+
+			// Populate the MSAA depth buffer from depth buffer or temporally upscaled depth buffer
+			{
+				FPopulateEditorDepthPS::FParameters* PassParameters = GraphBuilder.AllocParameters<FPopulateEditorDepthPS::FParameters>();
+				PassParameters->View = View.ViewUniformBuffer;
+				PassParameters->Color = GetScreenPassTextureViewportParameters(FScreenPassTextureViewport(Inputs.SceneColor));
+				PassParameters->Depth = GetScreenPassTextureViewportParameters(FScreenPassTextureViewport(SceneDepth));
+				PassParameters->DepthTextureJitter = SceneDepthJitter;
+				PassParameters->DepthTexture = SceneDepth.Texture;
+				PassParameters->DepthSampler = TStaticSamplerState<SF_Point, AM_Clamp, AM_Clamp, AM_Clamp>::GetRHI();
+				PassParameters->RenderTargets[0] = FRenderTargetBinding(EditorPrimitiveColor, ERenderTargetLoadAction::EClear);
+				PassParameters->RenderTargets.DepthStencil = FDepthStencilBinding(EditorPrimitiveDepth, ERenderTargetLoadAction::EClear, ERenderTargetLoadAction::EClear, FExclusiveDepthStencil::DepthWrite_StencilWrite);
+
+				FPopulateEditorDepthPS::FPermutationDomain PermutationVector;
+				PermutationVector.Set<FPopulateEditorDepthPS::FUseMSAADimension>(NumSamples > 1);
+				TShaderMapRef<FPopulateEditorDepthPS> PixelShader(View.ShaderMap, PermutationVector);
+
+				FPixelShaderUtils::AddFullscreenPass(
+					GraphBuilder,
+					View.ShaderMap,
+					RDG_EVENT_NAME("PopulateDepth %dx%d%s",
+						EditorPrimitivesViewport.Rect.Width(),
+						EditorPrimitivesViewport.Rect.Height(),
+						NumSamples > 1 ? TEXT(" MSAA") : TEXT("")),
+					PixelShader,
+					PassParameters,
+					EditorPrimitivesViewport.Rect,
+					/* BlendState = */ nullptr,
+					/* RasterizerState = */ nullptr,
+					TStaticDepthStencilState<true, CF_Always>::GetRHI());
+			}
 		}
 
+		// Draws the editors primitives
 		{
 			FEditorPrimitivesPassParameters* PassParameters = GraphBuilder.AllocParameters<FEditorPrimitivesPassParameters>();
 			PassParameters->View = EditorView->ViewUniformBuffer;
@@ -347,10 +485,13 @@ FScreenPassTexture AddEditorPrimitivePass(
 			}
 
 			GraphBuilder.AddPass(
-				RDG_EVENT_NAME("EditorPrimitives"),
+				RDG_EVENT_NAME("EditorPrimitives %dx%d MSAA=%d",
+					EditorPrimitivesViewport.Rect.Width(),
+					EditorPrimitivesViewport.Rect.Height(),
+					NumSamples),
 				PassParameters,
 				ERDGPassFlags::Raster,
-				[&View, &InstanceCullingManager, PassParameters, EditorView, EditorPrimitivesViewport, SceneDepthViewport, BasePassType, NumSamples](FRHICommandListImmediate& RHICmdList)
+				[&View, &InstanceCullingManager, PassParameters, EditorView, EditorPrimitivesViewport, BasePassType, NumSamples](FRHICommandListImmediate& RHICmdList)
 			{
 				RHICmdList.SetViewport(EditorPrimitivesViewport.Rect.Min.X, EditorPrimitivesViewport.Rect.Min.Y, 0.0f, EditorPrimitivesViewport.Rect.Max.X, EditorPrimitivesViewport.Rect.Max.Y, 1.0f);
 
@@ -383,44 +524,48 @@ FScreenPassTexture AddEditorPrimitivePass(
 
 	if (!Output.IsValid())
 	{
-		Output = FScreenPassRenderTarget::CreateFromInput(GraphBuilder, Inputs.SceneColor, View.GetOverwriteLoadAction(), TEXT("EditorPrimitives"));
+		Output = FScreenPassRenderTarget::CreateFromInput(GraphBuilder, Inputs.SceneColor, View.GetOverwriteLoadAction(), TEXT("Editor.Primitives"));
 	}
 
-	const FScreenPassTextureViewport OutputViewport(Output);
-	const FScreenPassTextureViewport ColorViewport(Inputs.SceneColor);
-	const FScreenPassTextureViewport DepthViewport(Inputs.SceneDepth);
-
-	FRHISamplerState* PointClampSampler = TStaticSamplerState<SF_Point, AM_Clamp, AM_Clamp, AM_Clamp>::GetRHI();
-
-	const bool bOpaqueEditorGizmo = CVarEditorOpaqueGizmo.GetValueOnRenderThread() != 0 || View.Family->EngineShowFlags.Wireframe;
-
-	FCompositeEditorPrimitivesPS::FParameters* PassParameters = GraphBuilder.AllocParameters<FCompositeEditorPrimitivesPS::FParameters>();
-	PassParameters->RenderTargets[0] = Output.GetRenderTargetBinding();
-	PassParameters->View = View.ViewUniformBuffer;
-	PassParameters->Color = GetScreenPassTextureViewportParameters(ColorViewport);
-	PassParameters->Depth = GetScreenPassTextureViewportParameters(DepthViewport);
-	PassParameters->ColorToDepth = GetScreenPassTextureViewportTransform(PassParameters->Color, PassParameters->Depth);
-	PassParameters->ColorTexture = Inputs.SceneColor.Texture;
-	PassParameters->ColorSampler = PointClampSampler;
-	PassParameters->DepthTexture = Inputs.SceneDepth.Texture;
-	PassParameters->DepthSampler = PointClampSampler;
-	PassParameters->EditorPrimitivesDepth = EditorPrimitiveDepth;
-	PassParameters->EditorPrimitivesColor = EditorPrimitiveColor;
-	PassParameters->bOpaqueEditorGizmo = bOpaqueEditorGizmo;
-	PassParameters->bCompositeAnyNonNullDepth = bProducedByPriorPass;
-
-	FCompositeEditorPrimitivesPS::FPermutationDomain PermutationVector;
-	PermutationVector.Set<FCompositeEditorPrimitivesPS::FSampleCountDimension>(NumSamples);
-
-	TShaderMapRef<FCompositeEditorPrimitivesPS> PixelShader(View.ShaderMap, PermutationVector);
-	AddDrawScreenPass(
-		GraphBuilder,
-		RDG_EVENT_NAME("Composite %dx%d msaa=%d", OutputViewport.Rect.Width(), OutputViewport.Rect.Height(), NumSamples),
-		View,
-		OutputViewport,
-		ColorViewport,
-		PixelShader,
-		PassParameters);
+	{
+		FRHISamplerState* PointClampSampler = TStaticSamplerState<SF_Point, AM_Clamp, AM_Clamp, AM_Clamp>::GetRHI();
+
+		const bool bOpaqueEditorGizmo = CVarEditorOpaqueGizmo.GetValueOnRenderThread() != 0 || View.Family->EngineShowFlags.Wireframe;
+
+		FCompositeEditorPrimitivesPS::FParameters* PassParameters = GraphBuilder.AllocParameters<FCompositeEditorPrimitivesPS::FParameters>();
+		PassParameters->RenderTargets[0] = Output.GetRenderTargetBinding();
+		PassParameters->View = View.ViewUniformBuffer;
+		PassParameters->Color = GetScreenPassTextureViewportParameters(FScreenPassTextureViewport(Inputs.SceneColor));
+		PassParameters->Depth = GetScreenPassTextureViewportParameters(FScreenPassTextureViewport(SceneDepth));
+		PassParameters->Output = GetScreenPassTextureViewportParameters(FScreenPassTextureViewport(Output));
+		PassParameters->ColorTexture = Inputs.SceneColor.Texture;
+		PassParameters->ColorSampler = PointClampSampler;
+		PassParameters->DepthTexture = SceneDepth.Texture;
+		PassParameters->DepthSampler = PointClampSampler;
+		PassParameters->EditorPrimitivesDepth = EditorPrimitiveDepth;
+		PassParameters->EditorPrimitivesColor = EditorPrimitiveColor;
+		PassParameters->bOpaqueEditorGizmo = bOpaqueEditorGizmo;
+		PassParameters->bCompositeAnyNonNullDepth = bProducedByPriorPass;
+		PassParameters->DepthTextureJitter = SceneDepthJitter;
+
+		for (int32 i = 0; i < int32(NumSamples); i++)
+		{
+			PassParameters->SampleOffsetArray[i].X = GetMSAASampleOffsets(NumSamples, i).X;
+			PassParameters->SampleOffsetArray[i].Y = GetMSAASampleOffsets(NumSamples, i).Y;
+		}
+
+		FCompositeEditorPrimitivesPS::FPermutationDomain PermutationVector;
+		PermutationVector.Set<FCompositeEditorPrimitivesPS::FSampleCountDimension>(NumSamples);
+
+		TShaderMapRef<FCompositeEditorPrimitivesPS> PixelShader(View.ShaderMap, PermutationVector);
+		FPixelShaderUtils::AddFullscreenPass(
+			GraphBuilder,
+			View.ShaderMap,
+			RDG_EVENT_NAME("Composite %dx%d MSAA=%d", Output.ViewRect.Width(), Output.ViewRect.Height(), NumSamples),
+			PixelShader,
+			PassParameters,
+			Output.ViewRect);
+	}
 
 	return MoveTemp(Output);
 }
diff --git a/Engine/Source/Runtime/Renderer/Private/SceneRendering.h b/Engine/Source/Runtime/Renderer/Private/SceneRendering.h
--- a/Engine/Source/Runtime/Renderer/Private/SceneRendering.h
+++ b/Engine/Source/Runtime/Renderer/Private/SceneRendering.h
@@ -962,6 +962,9 @@ struct FPreviousViewInfo
 	FTemporalAAHistory SSRHistory;
 	FTemporalAAHistory WaterSSRHistory;
 
+	// Temporal AA history for the editor primitive depth upsampling
+	FTemporalAAHistory EditorPrimtiveDepthHistory;
+
 	// Scene color input for SSR, that can be different from TemporalAAHistory.RT[0] if there is a SSR
 	// input post process material.
 	TRefCountPtr<IPooledRenderTarget> CustomSSRInput;
