New skydome importance sampling for path tracer

Change the skydome importance sampling table used in the path tracer to use an equi-area mapping instead of cube maps. This simplifies the importance sampling step, preserves stratification better, and gives us an opportunity to sample the correct sky cubemap. The previous code only worked in static cases and did not support the real time capture mode or blending.

Currently this importance table is rebuilt every frame (by the path tracer), but this can be improved in a subsequent step. The skylight prep is relatively cheap (compared to the main path tracer kernel at least). For a 1k cubemap (which builds a 2k importance table) this pass costs roughly 0.4ms on a 2080 Ti.

The old cubemap based approach is still around for now to keep RT Skylight working as-is, but I'm hoping to share the code in a follow-up step.

Also disable the special case making the skylight always visible to camera. The new behavior better matches the raster path.

diff --git a/Engine/Shaders/Private/MonteCarlo.ush b/Engine/Shaders/Private/MonteCarlo.ush
--- a/Engine/Shaders/Private/MonteCarlo.ush
+++ b/Engine/Shaders/Private/MonteCarlo.ush
@@ -151,6 +151,60 @@ float2 UniformSampleDiskConcentricApprox( float2 E )
 	return sf;
 }
 
+
+// Based on: [Clarberg 2008, "Fast Equal-Area Mapping of the (Hemi)Sphere using SIMD"]
+// Fixed sign bit for UV.y == 0 and removed branch before division by using a small epsilon
+// https://fileadmin.cs.lth.se/graphics/research/papers/2008/simdmapping/clarberg_simdmapping08_preprint.pdf
+float3 EquiAreaSphericalMapping(float2 UV)
+{
+	UV = 2 * UV - 1;
+	float D = 1 - (abs(UV.x) + abs(UV.y));
+	float R = 1 - abs(D);
+	float Epsilon = 5.42101086243e-20; // 2^-64 (this avoids 0/0 without changing the rest of the mapping)
+	float Phi = (PI / 4) * ((abs(UV.y) - abs(UV.x)) / (R + Epsilon) + 1);
+	float F = R * sqrt(2 - R * R);
+	return float3(
+		F * sign(UV.x) * abs(cos(Phi)),
+		F * sign(UV.y) * abs(sin(Phi)),
+		sign(D) * (1 - R * R)
+	);
+}
+
+// Based on: [Clarberg 2008, "Fast Equal-Area Mapping of the (Hemi)Sphere using SIMD"]
+// Removed branch before division by using a small epsilon
+// https://fileadmin.cs.lth.se/graphics/research/papers/2008/simdmapping/clarberg_simdmapping08_preprint.pdf
+float2 InverseEquiAreaSphericalMapping(float3 Direction)
+{
+	float3 AbsDir = abs(Direction);
+	float R = sqrt(1 - AbsDir.z);
+	float Epsilon = 5.42101086243e-20; // 2^-64 (this avoids 0/0 without changing the rest of the mapping)
+	float x = min(AbsDir.x, AbsDir.y) / (max(AbsDir.x, AbsDir.y) + Epsilon);
+
+	// Coefficients for 6th degree minimax approximation of atan(x)*2/pi, x=[0,1].
+	const float t1 = 0.406758566246788489601959989e-5f;
+	const float t2 = 0.636226545274016134946890922156f;
+	const float t3 = 0.61572017898280213493197203466e-2f;
+	const float t4 = -0.247333733281268944196501420480f;
+	const float t5 = 0.881770664775316294736387951347e-1f;
+	const float t6 = 0.419038818029165735901852432784e-1f;
+	const float t7 = -0.251390972343483509333252996350e-1f;
+
+	// Polynomial approximation of atan(x)*2/pi
+	float Phi = t6 + t7 * x;
+	Phi = t5 + Phi * x;
+	Phi = t4 + Phi * x;
+	Phi = t3 + Phi * x;
+	Phi = t2 + Phi * x;
+	Phi = t1 + Phi * x;
+
+	Phi = (AbsDir.x < AbsDir.y) ? 1 - Phi : Phi;
+	float2 UV = float2(R - Phi * R, Phi * R);
+	UV = (Direction.z < 0) ? 1 - UV.yx : UV;
+	UV = asfloat(asuint(UV) ^ (asuint(Direction.xy) & 0x80000000u));
+	return UV * 0.5 + 0.5;
+}
+
+
 // PDF = 1 / (4 * PI)
 float4 UniformSampleSphere( float2 E )
 {
diff --git a/Engine/Shaders/Private/PathTracing/Light/PathTracingSkyLight.ush b/Engine/Shaders/Private/PathTracing/Light/PathTracingSkyLight.ush
--- a/Engine/Shaders/Private/PathTracing/Light/PathTracingSkyLight.ush
+++ b/Engine/Shaders/Private/PathTracing/Light/PathTracingSkyLight.ush
@@ -6,8 +6,16 @@ SkyLight.usf: Light sampling functions for SkyLight implementation
 
 #pragma once
 
+#ifdef USE_NEW_SKYDOME
+Texture2D SkylightTexture;
+Texture2D<float> SkylightPdf;
+SamplerState SkylightTextureSampler;
+float SkylightInvResolution;
+int SkylightMipCount;
+#else
 #include "../../RayTracing/SkyLightMipTreeCommon.ush"
 #include "../../ReflectionEnvironmentShared.ush"
+#endif
 
 #define USE_HIERARCHICAL_IMPORTANCE_SAMPLING
 
@@ -18,27 +26,39 @@ float SkyLight_EstimateLight(
 	bool IsTransmissiveMaterial
 )
 {
+#ifdef USE_NEW_SKYDOME
+	// cancels out constant factor in pdf below
+	return 4 * PI * SkylightPdf.Load(int3(0, 0, SkylightMipCount - 1));
+#else
 	// Modifying diffuse contribution to prevent false negatives with PDF estimation
 	float3 SkySHDiffuse = max(GetSkySHDiffuse(WorldNormal), 0.01);
 	float3 SkyLightIrradiance = GetColor(LightId) * SkySHDiffuse;
 	float SkyLightIlluminance = Luminance(SkyLightIrradiance);
 
 	return SkyLightIlluminance * 2.0f * PI;
+#endif
 }
 
 FLightHit SkyLight_TraceLight(RayDesc Ray, int LightId) {
 	if (Ray.TMax == POSITIVE_INFINITY)
 	{
-		// LightUV convention in the code below is just the world space ray direction
-		// so its safe to call into the method below for now ...
-
-		// #dxr_todo: write optimized form that combines radiance and pdf lookup in a single step
+#ifdef USE_NEW_SKYDOME
+		// NOTE: assumes direction is normalized
 
-		float3 Radiance = 0.0;
-		float Pdf = 0.0;
+		float3 Dir = Ray.Direction;
+		float2 UV = InverseEquiAreaSphericalMapping(Dir.yzx);
+		float4 Result = SkylightTexture.SampleLevel(SkylightTextureSampler, UV, 0);
+		float3 Radiance = Result.xyz;
+#ifndef USE_HIERARCHICAL_IMPORTANCE_SAMPLING
+		float Pdf = 1.0 / (2.0 * PI);
+#else
+		float Pdf = Result.w / (4 * PI * SkylightPdf.Load(int3(0, 0, SkylightMipCount - 1)));
+#endif
 
+#else
 		float3 LightUV = normalize(Ray.Direction);
-		Radiance = SkyLight.Color * SkyLight.Texture.SampleLevel(SkyLight.TextureSampler, LightUV, 0).rgb;
+		float3 Radiance = SkyLight.Color * SkyLight.Texture.SampleLevel(SkyLight.TextureSampler, LightUV, 0).rgb;
+		float Pdf = 0;
 		if (any(Radiance > 0))
 		{
 #ifndef USE_HIERARCHICAL_IMPORTANCE_SAMPLING
@@ -63,7 +83,7 @@ FLightHit SkyLight_TraceLight(RayDesc Ray, int LightId) {
 			Pdf = FacePdf * MipPdf * AreaPdf;
 #endif
 		}
-		
+#endif
 		return CreateLightHit(Radiance, Pdf, POSITIVE_INFINITY);
 	}
 	return NullLightHit();
@@ -84,10 +104,78 @@ FLightSample SkyLight_SampleLight(
 	Payload.WorldNormal = WorldNormal;
 
 
+#ifdef USE_NEW_SKYDOME
+
 #ifndef USE_HIERARCHICAL_IMPORTANCE_SAMPLING
 	float4 Direction_Tangent = UniformSampleHemisphere(RandSample.yz);
-	OutDirection = TangentToWorld(Direction_Tangent.xyz, WorldNormal);
-	OutPdf = Direction_Tangent.w;
+	float3 OutDirection = TangentToWorld(Direction_Tangent.xyz, WorldNormal);
+	float OutPdf = Direction_Tangent.w;
+	float2 UV = InverseEquiAreaSphericalMapping(OutDirection.yzx);
+	float3 Radiance = SkylightTexture.SampleLevel(SkylightTextureSampler, UV, 0).xyz;
+	return CreateLightSample(Radiance / OutPdf, OutPdf, OutDirection, POSITIVE_INFINITY);
+#else
+	float2 UV = RandSample.yz;
+
+	int3 Pixel = int3(0, 0, SkylightMipCount - 2);
+	for (; Pixel.z >= 0; Pixel.z--)
+	{
+		Pixel.xy *= 2;
+		float P00 = SkylightPdf.Load(Pixel + int3(0, 0, 0));
+		float P10 = SkylightPdf.Load(Pixel + int3(1, 0, 0));
+		float P01 = SkylightPdf.Load(Pixel + int3(0, 1, 0));
+		float P11 = SkylightPdf.Load(Pixel + int3(1, 1, 0));
+
+		float L = P00 + P01;
+		float R = P10 + P11;
+
+		float ProbX = L / (L + R);
+		if (UV.x < ProbX)
+		{
+			UV.x /= ProbX;
+			float ProbY = P00 / L;
+			if (UV.y < ProbY)
+			{
+				UV.y /= ProbY;
+			}
+			else
+			{
+				Pixel.y++;
+				UV.y = (UV.y - ProbY) / (1 - ProbY);
+			}
+		}
+		else
+		{
+			Pixel.x++;
+			UV.x = (UV.x - ProbX) / (1 - ProbX);
+			float ProbY = P10 / R;
+			if (UV.y < ProbY)
+			{
+				UV.y /= ProbY;
+			}
+			else
+			{
+				Pixel.y++;
+				UV.y = (UV.y - ProbY) / (1 - ProbY);
+			}
+		}
+	}
+
+	Pixel.z = 0;
+	float4 Result = SkylightTexture.Load(Pixel);
+	float3 Radiance = Result.xyz;
+	float OutPdf = Result.w / (4 * PI * SkylightPdf.Load(int3(0, 0, SkylightMipCount - 1)));
+
+	UV = (float2(Pixel.xy) + UV) * SkylightInvResolution;
+
+	return CreateLightSample(Radiance / OutPdf, OutPdf, EquiAreaSphericalMapping(UV).zxy, POSITIVE_INFINITY);
+#endif
+
+#else
+
+#ifndef USE_HIERARCHICAL_IMPORTANCE_SAMPLING
+	float4 Direction_Tangent = UniformSampleHemisphere(RandSample.yz);
+	float3 OutDirection = TangentToWorld(Direction_Tangent.xyz, WorldNormal);
+	float OutPdf = Direction_Tangent.w;
 #else
 	// Determine the appropriate face to use
 	//uint MipCount = Log2(SkyLight.MipDimensions.x);
@@ -162,4 +250,6 @@ FLightSample SkyLight_SampleLight(
 	float3 Radiance = SkyLight.Color * SkyLight.Texture.SampleLevel(SkyLight.TextureSampler, OutDirection, 0).rgb;
 
 	return CreateLightSample(Radiance / OutPdf, OutPdf, OutDirection, POSITIVE_INFINITY);
+#endif
+
 }
diff --git a/Engine/Shaders/Private/PathTracing/PathTracing.usf b/Engine/Shaders/Private/PathTracing/PathTracing.usf
--- a/Engine/Shaders/Private/PathTracing/PathTracing.usf
+++ b/Engine/Shaders/Private/PathTracing/PathTracing.usf
@@ -250,20 +250,7 @@ RAY_TRACING_ENTRY_RAYGEN(PathTracingMainRG)
 	const int MISMode = PathTracingData.MISMode;
 
 	// number of directly visible lights for the first bounce
-	uint NumLights = 0;
-	if (SceneLightCount > 0)
-	{
-		if (PathTracingData.VisibleLights)
-		{
-			// user wants to see all lights
-			NumLights = SceneLightCount;
-		}
-		else if (IsEnvironmentLight(0))
-		{
-			// user doesn't want to see all lights - still check the first one in case its a skydome (there can be only one, and it always comes first if it is there)
-			NumLights = 1;
-		}
-	}
+	uint NumLights = PathTracingData.VisibleLights ? SceneLightCount : 0;
 
 	for (int Bounce = 0; Bounce <= PathTracingData.MaxBounces; Bounce++)
 	{
diff --git a/Engine/Shaders/Private/PathTracing/PathTracingSkylightPrepare.usf b/Engine/Shaders/Private/PathTracing/PathTracingSkylightPrepare.usf
new file mode 100644
--- /dev/null
+++ b/Engine/Shaders/Private/PathTracing/PathTracingSkylightPrepare.usf
@@ -0,0 +1,36 @@
+// Copyright Epic Games, Inc. All Rights Reserved.
+
+#include "../Common.ush"
+#include "../MonteCarlo.ush"
+
+TextureCube SkyLightCubemap0;
+TextureCube SkyLightCubemap1;
+SamplerState SkyLightCubemapSampler0;
+SamplerState SkyLightCubemapSampler1;
+float SkylightBlendFactor;
+float SkylightInvResolution;
+
+RWTexture2D<float4> SkylightTextureOutput;
+RWTexture2D<float> SkylightTexturePdf;
+float3 SkyColor;
+
+[numthreads(THREADGROUPSIZE_X, THREADGROUPSIZE_Y, 1)]
+void PathTracingSkylightPrepareCS(uint2 DispatchThreadId : SV_DispatchThreadID)
+{
+	float2 UV = (DispatchThreadId + 0.5) * SkylightInvResolution;
+	// Highest resolution level -- capture the actual cube map(s) data
+	float3 ReflectionVector = EquiAreaSphericalMapping(UV).zxy;
+	float3 Reflection = TextureCubeSample(SkyLightCubemap0, SkyLightCubemapSampler0, ReflectionVector).rgb;
+	if (SkylightBlendFactor > 0)
+	{
+		float3 Reflection1 = TextureCubeSample(SkyLightCubemap1, SkyLightCubemapSampler1, ReflectionVector).rgb;
+
+		Reflection = lerp(Reflection, Reflection1, SkylightBlendFactor);
+	}
+
+	float3 Pixel = max(Reflection * SkyColor, 0.0);
+	float Prob = max(Pixel.r, max(Pixel.g, Pixel.b));
+	// record color and probability density (as the max of each component)
+	SkylightTextureOutput[DispatchThreadId] = float4(Pixel, Prob);
+	SkylightTexturePdf[DispatchThreadId] = Prob;
+}
diff --git a/Engine/Source/Runtime/Renderer/Private/PathTracing.cpp b/Engine/Source/Runtime/Renderer/Private/PathTracing.cpp
--- a/Engine/Source/Runtime/Renderer/Private/PathTracing.cpp
+++ b/Engine/Source/Runtime/Renderer/Private/PathTracing.cpp
@@ -7,14 +7,10 @@
 #include "RendererPrivate.h"
 #include "GlobalShader.h"
 #include "DeferredShadingRenderer.h"
-#include "PostProcess/PostProcessing.h"
-#include "PostProcess/SceneFilterRendering.h"
 #include "PathTracingUniformBuffers.h"
-#include "RHI/Public/PipelineStateCache.h"
-#include "RayTracing/RayTracingSkyLight.h"
-#include "RayTracing/RaytracingOptions.h"
 #include "HAL/PlatformApplicationMisc.h"
 #include "RayTracingTypes.h"
+#include "RenderCore/Public/GenerateMips.h"
 
 TAutoConsoleVariable<int32> CVarPathTracingMaxBounces(
 	TEXT("r.PathTracing.MaxBounces"),
@@ -246,6 +242,37 @@ static bool PrepareShaderArgs(const FViewInfo& View, FPathTracingData& PathTraci
 	return NeedInvalidation;
 }
 
+class FPathTracingSkylightPrepareCS : public FGlobalShader
+{
+	DECLARE_GLOBAL_SHADER(FPathTracingSkylightPrepareCS)
+	SHADER_USE_PARAMETER_STRUCT(FPathTracingSkylightPrepareCS, FGlobalShader)
+
+	static bool ShouldCompilePermutation(const FGlobalShaderPermutationParameters& Parameters)
+	{
+		return ShouldCompileRayTracingShadersForProject(Parameters.Platform);
+	}
+
+	static void ModifyCompilationEnvironment(const FGlobalShaderPermutationParameters& Parameters, FShaderCompilerEnvironment& OutEnvironment)
+	{
+		OutEnvironment.CompilerFlags.Add(CFLAG_WarningsAsErrors);
+		OutEnvironment.SetDefine(TEXT("THREADGROUPSIZE_X"), FComputeShaderUtils::kGolden2DGroupSize);
+		OutEnvironment.SetDefine(TEXT("THREADGROUPSIZE_Y"), FComputeShaderUtils::kGolden2DGroupSize);
+	}
+
+	BEGIN_SHADER_PARAMETER_STRUCT(FParameters, )
+		SHADER_PARAMETER_TEXTURE(TextureCube, SkyLightCubemap0)
+		SHADER_PARAMETER_TEXTURE(TextureCube, SkyLightCubemap1)
+		SHADER_PARAMETER_SAMPLER(SamplerState, SkyLightCubemapSampler0)
+		SHADER_PARAMETER_SAMPLER(SamplerState, SkyLightCubemapSampler1)
+		SHADER_PARAMETER(float, SkylightBlendFactor)
+		SHADER_PARAMETER(float, SkylightInvResolution)
+		SHADER_PARAMETER_RDG_TEXTURE_UAV(RWTexture2D, SkylightTextureOutput)
+		SHADER_PARAMETER_RDG_TEXTURE_UAV(RWTexture2D, SkylightTexturePdf)
+		SHADER_PARAMETER(FVector, SkyColor)
+	END_SHADER_PARAMETER_STRUCT()
+};
+IMPLEMENT_SHADER_TYPE(, FPathTracingSkylightPrepareCS, TEXT("/Engine/Private/PathTracing/PathTracingSkylightPrepare.usf"), TEXT("PathTracingSkylightPrepareCS"), SF_Compute);
+
 class FPathTracingRG : public FGlobalShader
 {
 	DECLARE_GLOBAL_SHADER(FPathTracingRG)
@@ -256,15 +283,27 @@ class FPathTracingRG : public FGlobalShader
 		return ShouldCompileRayTracingShadersForProject(Parameters.Platform);
 	}
 
+	static void ModifyCompilationEnvironment(const FGlobalShaderPermutationParameters& Parameters, FShaderCompilerEnvironment& OutEnvironment)
+	{
+		OutEnvironment.CompilerFlags.Add(CFLAG_WarningsAsErrors);
+		OutEnvironment.SetDefine(TEXT("USE_NEW_SKYDOME"), 1);
+	}
+
+
 	BEGIN_SHADER_PARAMETER_STRUCT(FParameters, )
 		SHADER_PARAMETER_RDG_TEXTURE_UAV(RWTexture2D<float4>, RadianceTexture)
 		SHADER_PARAMETER_SRV(RaytracingAccelerationStructure, TLAS)
 
 		SHADER_PARAMETER_STRUCT_REF(FViewUniformShaderParameters, ViewUniformBuffer)
-		SHADER_PARAMETER_STRUCT_REF(FSkyLightData, SkyLightData)
 		SHADER_PARAMETER_STRUCT_REF(FPathTracingData, PathTracingData)
 		SHADER_PARAMETER_RDG_BUFFER_SRV(StructuredBuffer<FPathTracingLight>, SceneLights)
 		SHADER_PARAMETER(uint32, SceneLightCount)
+		// Skylight
+		SHADER_PARAMETER_RDG_TEXTURE(Texture2D, SkylightTexture)
+		SHADER_PARAMETER_RDG_TEXTURE(Texture2D, SkylightPdf)
+		SHADER_PARAMETER_SAMPLER(SamplerState, SkylightTextureSampler)
+		SHADER_PARAMETER(float, SkylightInvResolution)
+		SHADER_PARAMETER(int32, SkylightMipCount)
 		// IES Profiles
 		SHADER_PARAMETER_RDG_TEXTURE(Texture2DArray, IESTexture)
 		SHADER_PARAMETER_SAMPLER(SamplerState, IESTextureSampler)
@@ -286,7 +325,7 @@ class FPathTracingIESAtlasCS : public FGlobalShader
 
 	static void ModifyCompilationEnvironment(const FGlobalShaderPermutationParameters& Parameters, FShaderCompilerEnvironment& OutEnvironment)
 	{
-		FGlobalShader::ModifyCompilationEnvironment(Parameters, OutEnvironment);
+		OutEnvironment.CompilerFlags.Add(CFLAG_WarningsAsErrors);
 		OutEnvironment.SetDefine(TEXT("THREADGROUPSIZE_X"), FComputeShaderUtils::kGolden2DGroupSize);
 		OutEnvironment.SetDefine(TEXT("THREADGROUPSIZE_Y"), FComputeShaderUtils::kGolden2DGroupSize);
 	}
@@ -300,21 +339,88 @@ class FPathTracingIESAtlasCS : public FGlobalShader
 };
 IMPLEMENT_SHADER_TYPE(, FPathTracingIESAtlasCS, TEXT("/Engine/Private/PathTracing/PathTracingIESAtlas.usf"), TEXT("PathTracingIESAtlasCS"), SF_Compute);
 
-void SetLightParameters(FRDGBuilder& GraphBuilder, FPathTracingRG::FParameters* PassParameters, FSkyLightData& SkyLightData, FScene* Scene, const FViewInfo& View)
+bool PrepareSkyTexture(FRDGBuilder& GraphBuilder, FScene* Scene, const FViewInfo& View, FPathTracingRG::FParameters* PathTracingParameters)
 {
-	// Sky light
-	bool IsSkyLightValid = SetupSkyLightParameters(*Scene, &SkyLightData);
+	PathTracingParameters->SkylightTextureSampler = TStaticSamplerState<SF_Bilinear, AM_Clamp, AM_Clamp, AM_Clamp>::GetRHI();
+
+	FReflectionUniformParameters Parameters;
+	SetupReflectionUniformParameters(View, Parameters);
+	if (!(Parameters.SkyLightParameters.Y > 0))
+	{
+		// textures not ready, or skylight not active
+		// just put in a placeholder
+		PathTracingParameters->SkylightTexture = GraphBuilder.RegisterExternalTexture(GSystemTextures.BlackDummy);
+		PathTracingParameters->SkylightPdf = GraphBuilder.RegisterExternalTexture(GSystemTextures.BlackDummy);
+		PathTracingParameters->SkylightInvResolution = 0;
+		PathTracingParameters->SkylightMipCount = 0;
+		return false;
+	}
+
+	// TODO: this current runs every frame, but should only run when something has been updated
+	// This pass is not too expensive (at least compared to the cost of the path tracer) but this should
+	// be cleaned up.
+	RDG_EVENT_SCOPE(GraphBuilder, "Path Tracing SkylightPrepare");
+
+	FLinearColor SkyColor = Scene->SkyLight->GetEffectiveLightColor();
+	// since we are resampled into an octahedral layout, we multiply the cubemap resolution by 2 to get roughly the same number of texels
+	uint32 Size = FMath::RoundUpToPowerOfTwo(2 * Scene->SkyLight->CaptureCubeMapResolution);
+
+	FRDGTextureDesc SkylightTextureDesc = FRDGTextureDesc::Create2D(
+		FIntPoint(Size, Size),
+		PF_A32B32G32R32F, // half precision might be ok?
+		FClearValueBinding::None,
+		TexCreate_ShaderResource | TexCreate_UAV);
+	PathTracingParameters->SkylightTexture = GraphBuilder.CreateTexture(SkylightTextureDesc, TEXT("PathTracer.Skylight"), ERDGTextureFlags::None);
+	FRDGTextureDesc SkylightPdfDesc = FRDGTextureDesc::Create2D(
+		FIntPoint(Size, Size),
+		PF_R32_FLOAT, // half precision might be ok?
+		FClearValueBinding::None,
+		TexCreate_ShaderResource | TexCreate_UAV,
+		FMath::CeilLogTwo(Size) + 1);
+	PathTracingParameters->SkylightPdf = GraphBuilder.CreateTexture(SkylightPdfDesc, TEXT("PathTracer.SkylightPdf"), ERDGTextureFlags::None);
+
+	PathTracingParameters->SkylightInvResolution = 1.0f / Size;
+	PathTracingParameters->SkylightMipCount = SkylightPdfDesc.NumMips;
+
+	// run a simple compute shader to sample the cubemap and prep the top level of the mipmap hierarchy
+	{
+		TShaderMapRef<FPathTracingSkylightPrepareCS> ComputeShader(View.ShaderMap);
+		FPathTracingSkylightPrepareCS::FParameters* PassParameters = GraphBuilder.AllocParameters<FPathTracingSkylightPrepareCS::FParameters>();
+		PassParameters->SkyColor = FVector(SkyColor.R, SkyColor.G, SkyColor.B);
+		PassParameters->SkyLightCubemap0 = Parameters.SkyLightCubemap;
+		PassParameters->SkyLightCubemap1 = Parameters.SkyLightBlendDestinationCubemap;
+		PassParameters->SkyLightCubemapSampler0 = Parameters.SkyLightCubemapSampler;
+		PassParameters->SkyLightCubemapSampler1 = Parameters.SkyLightBlendDestinationCubemapSampler;
+		PassParameters->SkylightBlendFactor = Parameters.SkyLightParameters.W;
+		PassParameters->SkylightInvResolution = PathTracingParameters->SkylightInvResolution;
+		PassParameters->SkylightTextureOutput = GraphBuilder.CreateUAV(FRDGTextureUAVDesc(PathTracingParameters->SkylightTexture, 0));
+		PassParameters->SkylightTexturePdf = GraphBuilder.CreateUAV(FRDGTextureUAVDesc(PathTracingParameters->SkylightPdf, 0));
+		FComputeShaderUtils::AddPass(
+			GraphBuilder,
+			RDG_EVENT_NAME("SkylightPrepare"),
+			ComputeShader,
+			PassParameters,
+			FComputeShaderUtils::GetGroupCount(FIntPoint(Size, Size), FComputeShaderUtils::kGolden2DGroupSize));
+	}
+
+	FGenerateMips::ExecuteCompute(GraphBuilder, PathTracingParameters->SkylightPdf, TStaticSamplerState<SF_Bilinear, AM_Clamp, AM_Clamp, AM_Clamp>::GetRHI());
+
+	return true;
+}
 
+void SetLightParameters(FRDGBuilder& GraphBuilder, FPathTracingRG::FParameters* PassParameters, FScene* Scene, const FViewInfo& View)
+{
 	// Lights
 	FPathTracingLight Lights[RAY_TRACING_LIGHT_COUNT_MAXIMUM]; // Keep this on the stack for now -- eventually will need to make this dynamic to lift size limit (and also avoid uploading per frame ...)
 	unsigned LightCount = 0;
 
 	// Prepend SkyLight to light buffer since it is not part of the regular light list
-	if (IsSkyLightValid && View.Family->EngineShowFlags.SkyLighting)
+	if (PrepareSkyTexture(GraphBuilder, Scene, View, PassParameters))
 	{
+		check(Scene->SkyLight != nullptr);
 		FPathTracingLight& DestLight = Lights[LightCount];
-		DestLight.Color = FVector(SkyLightData.Color);
-		DestLight.Flags = SkyLightData.bTransmission ? PATHTRACER_FLAG_TRANSMISSION_MASK : 0;
+		DestLight.Color = FVector(1, 1, 1); // not used
+		DestLight.Flags = Scene->SkyLight->bTransmission ? PATHTRACER_FLAG_TRANSMISSION_MASK : 0;
 		DestLight.Flags |= PATHTRACER_FLAG_LIGHTING_CHANNEL_MASK;
 		DestLight.Flags |= PATHTRACING_LIGHT_SKY;
 		DestLight.Flags |= Scene->SkyLight->bCastShadows ? PATHTRACER_FLAG_CAST_SHADOW_MASK : 0;
@@ -610,12 +716,8 @@ void FDeferredShadingSceneRenderer::RenderPathTracing(
 		PassParameters->TLAS = View.RayTracingScene.RayTracingSceneRHI->GetShaderResourceView();
 		PassParameters->ViewUniformBuffer = View.ViewUniformBuffer;
 		PassParameters->PathTracingData = CreateUniformBufferImmediate(PathTracingData, EUniformBufferUsage::UniformBuffer_SingleFrame);
-		{
-			// upload sky/lights data
-			FSkyLightData SkyLightData;
-			SetLightParameters(GraphBuilder, PassParameters, SkyLightData, Scene, View);
-			PassParameters->SkyLightData = CreateUniformBufferImmediate(SkyLightData, EUniformBufferUsage::UniformBuffer_SingleFrame);
-		}
+		// upload sky/lights data
+		SetLightParameters(GraphBuilder, PassParameters, Scene, View);
 		PassParameters->IESTextureSampler = TStaticSamplerState<SF_Bilinear, AM_Clamp, AM_Clamp, AM_Clamp>::GetRHI();
 		PassParameters->RadianceTexture = GraphBuilder.CreateUAV(RadianceTexture);
 
