- supress static analysis warning

diff --git a/Engine/Source/Runtime/Renderer/Private/DynamicBVH.h b/Engine/Source/Runtime/Renderer/Private/DynamicBVH.h
--- a/Engine/Source/Runtime/Renderer/Private/DynamicBVH.h
+++ b/Engine/Source/Runtime/Renderer/Private/DynamicBVH.h
@@ -687,6 +687,7 @@ uint32 FDynamicBVH< MaxChildren, FCostMetric >::PromoteChild( uint32 NodeIndex )
 	// Promote from child to sibling
 	
 	// Remove from bounds
+	CA_SUPPRESS(6385);
 	SetBounds( NodeIndex, Excluded[ BestIndex & ChildMask ] );
 
 	// Add as sibling
