Added a missing blend mode remappings for a decalls on mobile

diff --git a/Engine/Source/Runtime/Renderer/Private/MobileDecalRendering.cpp b/Engine/Source/Runtime/Renderer/Private/MobileDecalRendering.cpp
--- a/Engine/Source/Runtime/Renderer/Private/MobileDecalRendering.cpp
+++ b/Engine/Source/Runtime/Renderer/Private/MobileDecalRendering.cpp
@@ -43,6 +43,8 @@ FRHIBlendState* MobileForward_GetDecalBlendState(EDecalBlendMode DecalBlendMode)
 		// Additive
 		return TStaticBlendState<CW_RGB, BO_Add, BF_SourceAlpha, BF_One>::GetRHI();
 	case DBM_AlphaComposite:
+	case DBM_DBuffer_AlphaComposite:
+	case DBM_DBuffer_EmissiveAlphaComposite:
 		// Premultiplied alpha
 		return TStaticBlendState<CW_RGB, BO_Add, BF_One, BF_InverseSourceAlpha>::GetRHI();
 	default:
