Disabled warnings-as-errors shader compilation setting in PathTracing.cpp to fix cook errors caused by tons of resultant shader warnings.

diff --git a/Engine/Source/Runtime/Renderer/Private/PathTracing.cpp b/Engine/Source/Runtime/Renderer/Private/PathTracing.cpp
--- a/Engine/Source/Runtime/Renderer/Private/PathTracing.cpp
+++ b/Engine/Source/Runtime/Renderer/Private/PathTracing.cpp
@@ -254,7 +254,7 @@ class FPathTracingSkylightPrepareCS : public FGlobalShader
 
 	static void ModifyCompilationEnvironment(const FGlobalShaderPermutationParameters& Parameters, FShaderCompilerEnvironment& OutEnvironment)
 	{
-		OutEnvironment.CompilerFlags.Add(CFLAG_WarningsAsErrors);
+		//OutEnvironment.CompilerFlags.Add(CFLAG_WarningsAsErrors);
 		OutEnvironment.SetDefine(TEXT("THREADGROUPSIZE_X"), FComputeShaderUtils::kGolden2DGroupSize);
 		OutEnvironment.SetDefine(TEXT("THREADGROUPSIZE_Y"), FComputeShaderUtils::kGolden2DGroupSize);
 	}
@@ -285,7 +285,7 @@ class FPathTracingRG : public FGlobalShader
 
 	static void ModifyCompilationEnvironment(const FGlobalShaderPermutationParameters& Parameters, FShaderCompilerEnvironment& OutEnvironment)
 	{
-		OutEnvironment.CompilerFlags.Add(CFLAG_WarningsAsErrors);
+		//OutEnvironment.CompilerFlags.Add(CFLAG_WarningsAsErrors);
 		OutEnvironment.SetDefine(TEXT("USE_NEW_SKYDOME"), 1);
 	}
 
@@ -325,7 +325,7 @@ class FPathTracingIESAtlasCS : public FGlobalShader
 
 	static void ModifyCompilationEnvironment(const FGlobalShaderPermutationParameters& Parameters, FShaderCompilerEnvironment& OutEnvironment)
 	{
-		OutEnvironment.CompilerFlags.Add(CFLAG_WarningsAsErrors);
+		//OutEnvironment.CompilerFlags.Add(CFLAG_WarningsAsErrors);
 		OutEnvironment.SetDefine(TEXT("THREADGROUPSIZE_X"), FComputeShaderUtils::kGolden2DGroupSize);
 		OutEnvironment.SetDefine(TEXT("THREADGROUPSIZE_Y"), FComputeShaderUtils::kGolden2DGroupSize);
 	}
