Actually use r.Strata.BytesPerPixel and also make the user facing setting safe.

diff --git a/Engine/Source/Runtime/Engine/Private/RendererSettings.cpp b/Engine/Source/Runtime/Engine/Private/RendererSettings.cpp
--- a/Engine/Source/Runtime/Engine/Private/RendererSettings.cpp
+++ b/Engine/Source/Runtime/Engine/Private/RendererSettings.cpp
@@ -209,6 +209,13 @@ void URendererSettings::PostEditChangeProperty(FPropertyChangedEvent& PropertyCh
 			}
 		}
 
+		if ((PropertyChangedEvent.Property->GetFName() == GET_MEMBER_NAME_CHECKED(URendererSettings, StrataBytePerPixel)))
+		{
+			// We enforce at least 20 bytes per pixel because this is the minimal Strata GBuffer footprint of the simplest material.
+			StrataBytePerPixel = FMath::Clamp(StrataBytePerPixel, 20u, 128u);
+
+		}
+
 		ExportValuesToConsoleVariables(PropertyChangedEvent.Property);
 
 		if ((PropertyChangedEvent.Property->GetFName() == GET_MEMBER_NAME_CHECKED(URendererSettings, ReflectionCaptureResolution) ||
diff --git a/Engine/Source/Runtime/Renderer/Private/Strata/Strata.cpp b/Engine/Source/Runtime/Renderer/Private/Strata/Strata.cpp
--- a/Engine/Source/Runtime/Renderer/Private/Strata/Strata.cpp
+++ b/Engine/Source/Runtime/Renderer/Private/Strata/Strata.cpp
@@ -124,9 +124,7 @@ void InitialiseStrataFrameSceneData(FSceneRenderer& SceneRenderer, FRDGBuilder&
 		ResolutionX = FMath::DivideAndRoundUp(BufferSizeXY.X, STRATA_DATA_TILE_SIZE) * STRATA_DATA_TILE_SIZE;
 		ResolutionY = FMath::DivideAndRoundUp(BufferSizeXY.Y, STRATA_DATA_TILE_SIZE) * STRATA_DATA_TILE_SIZE;
 
-		// Previous GBuffer when complete was 28bytes
-		// check out Strata.ush to see how this is computed
-		const uint32 MaterialConservativeByteCountPerPixel = 100u;
+		const uint32 MaterialConservativeByteCountPerPixel = CVarStrataBytePerPixel.GetValueOnAnyThread();
 		const uint32 RoundToValue = 4u;
 		StrataSceneData.MaxBytesPerPixel = FMath::DivideAndRoundUp(MaterialConservativeByteCountPerPixel, RoundToValue) * RoundToValue;
 
