Fix crash when attempting to add anisotropic mesh batches on platforms that don't have the feature enabled.

Adding a mesh batch for anistropy will attempt to get shaders and fail since we are no longer compiling variants.  Prevent mesh batches from being added by checking if the platform supports it similar to ShouldRenderAnisotropyPass().

diff --git a/Engine/Source/Runtime/Renderer/Private/AnisotropyRendering.cpp b/Engine/Source/Runtime/Renderer/Private/AnisotropyRendering.cpp
--- a/Engine/Source/Runtime/Renderer/Private/AnisotropyRendering.cpp
+++ b/Engine/Source/Runtime/Renderer/Private/AnisotropyRendering.cpp
@@ -142,18 +142,21 @@ void FAnisotropyMeshProcessor::AddMeshBatch(
 	int32 StaticMeshId /* = -1 */ 
 	)
 {
-	const FMaterialRenderProxy* MaterialRenderProxy = MeshBatch.MaterialRenderProxy;
-	const FMaterial* Material = &MaterialRenderProxy->GetMaterialWithFallback(FeatureLevel, MaterialRenderProxy);
-	const EBlendMode BlendMode = Material->GetBlendMode();
-	const bool bIsNotTranslucent = BlendMode == BLEND_Opaque || BlendMode == BLEND_Masked;
-
-	if (MeshBatch.bUseForMaterial && Material->MaterialUsesAnisotropy_RenderThread() && bIsNotTranslucent && Material->GetShadingModels().HasAnyShadingModel({ MSM_DefaultLit, MSM_ClearCoat }))
+	if (SupportsAnisotropicMaterials(FeatureLevel, GShaderPlatformForFeatureLevel[FeatureLevel]))
 	{
-		const FMeshDrawingPolicyOverrideSettings OverrideSettings = ComputeMeshOverrideSettings(MeshBatch);
-		const ERasterizerFillMode MeshFillMode = ComputeMeshFillMode(MeshBatch, *Material, OverrideSettings);
-		const ERasterizerCullMode MeshCullMode = ComputeMeshCullMode(MeshBatch, *Material, OverrideSettings);
+		const FMaterialRenderProxy* MaterialRenderProxy = MeshBatch.MaterialRenderProxy;
+		const FMaterial* Material = &MaterialRenderProxy->GetMaterialWithFallback(FeatureLevel, MaterialRenderProxy);
+		const EBlendMode BlendMode = Material->GetBlendMode();
+		const bool bIsNotTranslucent = BlendMode == BLEND_Opaque || BlendMode == BLEND_Masked;
+
+		if (MeshBatch.bUseForMaterial && Material->MaterialUsesAnisotropy_RenderThread() && bIsNotTranslucent && Material->GetShadingModels().HasAnyShadingModel({ MSM_DefaultLit, MSM_ClearCoat }))
+		{
+			const FMeshDrawingPolicyOverrideSettings OverrideSettings = ComputeMeshOverrideSettings(MeshBatch);
+			const ERasterizerFillMode MeshFillMode = ComputeMeshFillMode(MeshBatch, *Material, OverrideSettings);
+			const ERasterizerCullMode MeshCullMode = ComputeMeshCullMode(MeshBatch, *Material, OverrideSettings);
 
-		Process(MeshBatch, BatchElementMask, StaticMeshId, PrimitiveSceneProxy, *MaterialRenderProxy, *Material, MeshFillMode, MeshCullMode);
+			Process(MeshBatch, BatchElementMask, StaticMeshId, PrimitiveSceneProxy, *MaterialRenderProxy, *Material, MeshFillMode, MeshCullMode);
+		}
 	}
 }
 
