Avoid rebuilding Skylight importance tables for the path tracer each frame.

The path tracer's representation of the skylight texture is now cached across frames. Code such as the realtime capture system or the skylight cubemap management only need to worry avoid invalidating the handle to the texture if it needs to change. The path tracing code will automatically rebuild these textures when they are needed. I also added a cvar (off by default) to force these textures to always be rebuilt to make it easier to benchmark the performance of this step.

Implement MIS Compensation for Skylight importance tables. When MIS is being used, we can tweak the skylight pdf since we know the domain is already well covered by BxDF sampling.

diff --git a/Engine/Shaders/Private/PathTracing/Light/PathTracingSkyLight.ush b/Engine/Shaders/Private/PathTracing/Light/PathTracingSkyLight.ush
--- a/Engine/Shaders/Private/PathTracing/Light/PathTracingSkyLight.ush
+++ b/Engine/Shaders/Private/PathTracing/Light/PathTracingSkyLight.ush
@@ -52,7 +52,7 @@ FLightHit SkyLight_TraceLight(RayDesc Ray, int LightId) {
 #ifndef USE_HIERARCHICAL_IMPORTANCE_SAMPLING
 		float Pdf = 1.0 / (2.0 * PI);
 #else
-		float Pdf = Result.w / (4 * PI * SkylightPdf.Load(int3(0, 0, SkylightMipCount - 1)));
+		float Pdf = Result.w > 0 ? Result.w / (4 * PI * SkylightPdf.Load(int3(0, 0, SkylightMipCount - 1))) : 0.0;
 #endif
 
 #else
diff --git a/Engine/Shaders/Private/PathTracing/PathTracingSkylightMISCompensation.usf b/Engine/Shaders/Private/PathTracing/PathTracingSkylightMISCompensation.usf
new file mode 100644
--- /dev/null
+++ b/Engine/Shaders/Private/PathTracing/PathTracingSkylightMISCompensation.usf
@@ -0,0 +1,23 @@
+// Copyright Epic Games, Inc. All Rights Reserved.
+
+#include "../Common.ush"
+
+Texture2D<float> SkylightTexturePdfAverage;
+RWTexture2D<float4> SkylightTextureOutput;
+RWTexture2D<float> SkylightTexturePdf;
+
+
+[numthreads(THREADGROUPSIZE_X, THREADGROUPSIZE_Y, 1)]
+void PathTracingSkylightMISCompensationCS(uint2 DispatchThreadId : SV_DispatchThreadID)
+{
+	// anything smaller than the average value will be well handled by BSDF samping when MIS is enabled
+	// so subtract it from the base value
+	// A more complete theoretical derivation is given in:
+	// [Karlík et al, 2019] "MIS Compensation: Optimizing Sampling Techniques in Multiple Importance Sampling"
+	// https://cgg.mff.cuni.cz/~jaroslav/papers/2019-mis-compensation/index.html
+
+	float AverageValue = SkylightTexturePdfAverage.Load(0); // bound to 1x1 mip level
+	float Prob = max(SkylightTexturePdf[DispatchThreadId] - AverageValue, 0.0);
+	SkylightTextureOutput[DispatchThreadId] = float4(SkylightTextureOutput[DispatchThreadId].xyz, Prob);
+	SkylightTexturePdf[DispatchThreadId] = Prob;
+}
diff --git a/Engine/Source/Runtime/Renderer/Private/PathTracing.cpp b/Engine/Source/Runtime/Renderer/Private/PathTracing.cpp
--- a/Engine/Source/Runtime/Renderer/Private/PathTracing.cpp
+++ b/Engine/Source/Runtime/Renderer/Private/PathTracing.cpp
@@ -22,14 +22,14 @@ TAutoConsoleVariable<int32> CVarPathTracingMaxBounces(
 TAutoConsoleVariable<int32> CVarPathTracingSamplesPerPixel(
 	TEXT("r.PathTracing.SamplesPerPixel"),
 	-1,
-	TEXT("Defines the samples per pixel before resetting the simulation (default = -1 (driven by postprocesing volume))"),
+	TEXT("Sets the maximum number of samples per pixel (default = -1 (driven by postprocesing volume))"),
 	ECVF_RenderThreadSafe
 );
 
 TAutoConsoleVariable<float> CVarPathTracingFilterWidth(
 	TEXT("r.PathTracing.FilterWidth"),
 	-1,
-	TEXT("Define the anti-aliasing filter width (default = -1 (driven by postprocesing volume))"),
+	TEXT("Sets the anti-aliasing filter width (default = -1 (driven by postprocesing volume))"),
 	ECVF_RenderThreadSafe
 );
 
@@ -44,13 +44,29 @@ TAutoConsoleVariable<int32> CVarPathTracingUseErrorDiffusion(
 TAutoConsoleVariable<int32> CVarPathTracingMISMode(
 	TEXT("r.PathTracing.MISMode"),
 	2,
-	TEXT("Selects the sampling techniques (default = 2 (MIS enabled))\n")
+	TEXT("Selects the sampling technique for light integration (default = 2 (MIS enabled))\n")
 	TEXT("0: Material sampling\n")
 	TEXT("1: Light sampling\n")
 	TEXT("2: MIS betwen material and light sampling (default)\n"),
 	ECVF_RenderThreadSafe
 );
 
+TAutoConsoleVariable<int32> CVarPathTracingMISCompensation(
+	TEXT("r.PathTracing.MISCompensation"),
+	1,
+	TEXT("Activates MIS compensation for skylight importance sampling. (default = 1 (enabled))\n")
+	TEXT("This option only takes effect when r.PathTracing.MISMode = 2\n"),
+	ECVF_RenderThreadSafe
+);
+
+TAutoConsoleVariable<int32> CVarPathTracingSkylightCaching(
+	TEXT("r.PathTracing.SkylightCaching"),
+	1,
+	TEXT("Attempts to re-use skylight data between frames. (default = 1 (enabled))\n")
+	TEXT("When set to 0, the skylight texture and importance samping data will be regenerated every frame. This is mainly intended as a benchmarking and debugging aid\n"),
+	ECVF_RenderThreadSafe
+);
+
 TAutoConsoleVariable<int32> CVarPathTracingVisibleLights(
 	TEXT("r.PathTracing.VisibleLights"),
 	0,
@@ -273,6 +289,34 @@ class FPathTracingSkylightPrepareCS : public FGlobalShader
 };
 IMPLEMENT_SHADER_TYPE(, FPathTracingSkylightPrepareCS, TEXT("/Engine/Private/PathTracing/PathTracingSkylightPrepare.usf"), TEXT("PathTracingSkylightPrepareCS"), SF_Compute);
 
+class FPathTracingSkylightMISCompensationCS : public FGlobalShader
+{
+	DECLARE_GLOBAL_SHADER(FPathTracingSkylightMISCompensationCS)
+	SHADER_USE_PARAMETER_STRUCT(FPathTracingSkylightMISCompensationCS, FGlobalShader)
+
+	static bool ShouldCompilePermutation(const FGlobalShaderPermutationParameters& Parameters)
+	{
+		return ShouldCompileRayTracingShadersForProject(Parameters.Platform);
+	}
+
+	static void ModifyCompilationEnvironment(const FGlobalShaderPermutationParameters& Parameters, FShaderCompilerEnvironment& OutEnvironment)
+	{
+		//OutEnvironment.CompilerFlags.Add(CFLAG_WarningsAsErrors);
+		OutEnvironment.CompilerFlags.Add(CFLAG_AllowTypedUAVLoads);
+		OutEnvironment.SetDefine(TEXT("THREADGROUPSIZE_X"), FComputeShaderUtils::kGolden2DGroupSize);
+		OutEnvironment.SetDefine(TEXT("THREADGROUPSIZE_Y"), FComputeShaderUtils::kGolden2DGroupSize);
+	}
+
+	BEGIN_SHADER_PARAMETER_STRUCT(FParameters, )
+		SHADER_PARAMETER_RDG_TEXTURE_SRV(Texture2D, SkylightTexturePdfAverage)
+		SHADER_PARAMETER_RDG_TEXTURE_UAV(RWTexture2D, SkylightTextureOutput)
+		SHADER_PARAMETER_RDG_TEXTURE_UAV(RWTexture2D, SkylightTexturePdf)
+		SHADER_PARAMETER(FVector, SkyColor)
+		END_SHADER_PARAMETER_STRUCT()
+};
+IMPLEMENT_SHADER_TYPE(, FPathTracingSkylightMISCompensationCS, TEXT("/Engine/Private/PathTracing/PathTracingSkylightMISCompensation.usf"), TEXT("PathTracingSkylightMISCompensationCS"), SF_Compute);
+
+
 class FPathTracingRG : public FGlobalShader
 {
 	DECLARE_GLOBAL_SHADER(FPathTracingRG)
@@ -339,7 +383,7 @@ class FPathTracingIESAtlasCS : public FGlobalShader
 };
 IMPLEMENT_SHADER_TYPE(, FPathTracingIESAtlasCS, TEXT("/Engine/Private/PathTracing/PathTracingIESAtlas.usf"), TEXT("PathTracingIESAtlasCS"), SF_Compute);
 
-bool PrepareSkyTexture(FRDGBuilder& GraphBuilder, FScene* Scene, const FViewInfo& View, FPathTracingRG::FParameters* PathTracingParameters)
+bool PrepareSkyTexture(FRDGBuilder& GraphBuilder, FScene* Scene, const FViewInfo& View, bool UseMISCompensation, FPathTracingRG::FParameters* PathTracingParameters)
 {
 	PathTracingParameters->SkylightTextureSampler = TStaticSamplerState<SF_Bilinear, AM_Clamp, AM_Clamp, AM_Clamp>::GetRHI();
 
@@ -356,9 +400,26 @@ bool PrepareSkyTexture(FRDGBuilder& GraphBuilder, FScene* Scene, const FViewInfo
 		return false;
 	}
 
-	// TODO: this current runs every frame, but should only run when something has been updated
-	// This pass is not too expensive (at least compared to the cost of the path tracer) but this should
-	// be cleaned up.
+	const bool IsSkylightCachingEnabled = CVarPathTracingSkylightCaching.GetValueOnAnyThread() != 0;
+
+	if (!IsSkylightCachingEnabled)
+	{
+		// we don't want any caching - release what we might have been holding onto
+		Scene->PathTracingSkylightTexture.SafeRelease();
+		Scene->PathTracingSkylightPdf.SafeRelease();
+	}
+
+	if (Scene->PathTracingSkylightTexture.IsValid() &&
+		Scene->PathTracingSkylightPdf.IsValid())
+	{
+		// we already have a valid texture and pdf, just re-use them!
+		// it is the responsability of code that may invalidate the contents to reset these pointers
+		PathTracingParameters->SkylightTexture = GraphBuilder.RegisterExternalTexture(Scene->PathTracingSkylightTexture, TEXT("PathTracer.Skylight"));
+		PathTracingParameters->SkylightPdf = GraphBuilder.RegisterExternalTexture(Scene->PathTracingSkylightPdf, TEXT("PathTracer.SkylightPdf"));
+		PathTracingParameters->SkylightInvResolution = 1.0f / PathTracingParameters->SkylightTexture->Desc.GetSize().X;
+		PathTracingParameters->SkylightMipCount = PathTracingParameters->SkylightPdf->Desc.NumMips;
+		return true;
+	}
 	RDG_EVENT_SCOPE(GraphBuilder, "Path Tracing SkylightPrepare");
 
 	FLinearColor SkyColor = Scene->SkyLight->GetEffectiveLightColor();
@@ -402,24 +463,45 @@ bool PrepareSkyTexture(FRDGBuilder& GraphBuilder, FScene* Scene, const FViewInfo
 			PassParameters,
 			FComputeShaderUtils::GetGroupCount(FIntPoint(Size, Size), FComputeShaderUtils::kGolden2DGroupSize));
 	}
-
 	FGenerateMips::ExecuteCompute(GraphBuilder, PathTracingParameters->SkylightPdf, TStaticSamplerState<SF_Bilinear, AM_Clamp, AM_Clamp, AM_Clamp>::GetRHI());
 
+	if (UseMISCompensation)
+	{
+		TShaderMapRef<FPathTracingSkylightMISCompensationCS> ComputeShader(View.ShaderMap);
+		FPathTracingSkylightMISCompensationCS::FParameters* PassParameters = GraphBuilder.AllocParameters<FPathTracingSkylightMISCompensationCS::FParameters>();
+		PassParameters->SkylightTexturePdfAverage = GraphBuilder.CreateSRV(FRDGTextureSRVDesc::CreateForMipLevel(PathTracingParameters->SkylightPdf, PathTracingParameters->SkylightMipCount - 1));
+		PassParameters->SkylightTextureOutput = GraphBuilder.CreateUAV(FRDGTextureUAVDesc(PathTracingParameters->SkylightTexture, 0));
+		PassParameters->SkylightTexturePdf = GraphBuilder.CreateUAV(FRDGTextureUAVDesc(PathTracingParameters->SkylightPdf, 0));
+		FComputeShaderUtils::AddPass(
+			GraphBuilder,
+			RDG_EVENT_NAME("SkylightMISCompensation"),
+			ComputeShader,
+			PassParameters,
+			FComputeShaderUtils::GetGroupCount(FIntPoint(Size, Size), FComputeShaderUtils::kGolden2DGroupSize));
+		FGenerateMips::ExecuteCompute(GraphBuilder, PathTracingParameters->SkylightPdf, TStaticSamplerState<SF_Bilinear, AM_Clamp, AM_Clamp, AM_Clamp>::GetRHI());
+	}
+
+	// hang onto these for next time (if caching is enabled)
+	if (IsSkylightCachingEnabled)
+	{
+		GraphBuilder.QueueTextureExtraction(PathTracingParameters->SkylightTexture, &Scene->PathTracingSkylightTexture);
+		GraphBuilder.QueueTextureExtraction(PathTracingParameters->SkylightPdf, &Scene->PathTracingSkylightPdf);
+	}
 	return true;
 }
 
-void SetLightParameters(FRDGBuilder& GraphBuilder, FPathTracingRG::FParameters* PassParameters, FScene* Scene, const FViewInfo& View)
+void SetLightParameters(FRDGBuilder& GraphBuilder, FPathTracingRG::FParameters* PassParameters, FScene* Scene, const FViewInfo& View, bool UseMISCompensation)
 {
 	// Lights
 	FPathTracingLight Lights[RAY_TRACING_LIGHT_COUNT_MAXIMUM]; // Keep this on the stack for now -- eventually will need to make this dynamic to lift size limit (and also avoid uploading per frame ...)
 	unsigned LightCount = 0;
 
 	// Prepend SkyLight to light buffer since it is not part of the regular light list
-	if (PrepareSkyTexture(GraphBuilder, Scene, View, PassParameters))
+	if (PrepareSkyTexture(GraphBuilder, Scene, View, UseMISCompensation, PassParameters))
 	{
 		check(Scene->SkyLight != nullptr);
 		FPathTracingLight& DestLight = Lights[LightCount];
-		DestLight.Color = FVector(1, 1, 1); // not used
+		DestLight.Color = FVector(1, 1, 1); // not used (it is folded into the importance table directly)
 		DestLight.Flags = Scene->SkyLight->bTransmission ? PATHTRACER_FLAG_TRANSMISSION_MASK : 0;
 		DestLight.Flags |= PATHTRACER_FLAG_LIGHTING_CHANNEL_MASK;
 		DestLight.Flags |= PATHTRACING_LIGHT_SKY;
@@ -665,6 +747,17 @@ void FDeferredShadingSceneRenderer::RenderPathTracing(
 		bArgsChanged = true;
 	}
 
+	bool UseMISCompensation = CVarPathTracingMISMode.GetValueOnRenderThread() == 2 && CVarPathTracingMISCompensation.GetValueOnRenderThread() != 0;
+	static bool PreviousUseMISCompensation = UseMISCompensation;
+	if (PreviousUseMISCompensation != UseMISCompensation)
+	{
+		PreviousUseMISCompensation = UseMISCompensation;
+		bArgsChanged = true;
+		// if the mode changes we need to rebuild the importance table
+		Scene->PathTracingSkylightTexture.SafeRelease();
+		Scene->PathTracingSkylightPdf.SafeRelease();
+	}
+
 	// Get other basic path tracing settings and see if we need to invalidate the current state
 	FPathTracingData PathTracingData;
 	bArgsChanged |= PrepareShaderArgs(View, PathTracingData);
@@ -717,7 +810,7 @@ void FDeferredShadingSceneRenderer::RenderPathTracing(
 		PassParameters->ViewUniformBuffer = View.ViewUniformBuffer;
 		PassParameters->PathTracingData = CreateUniformBufferImmediate(PathTracingData, EUniformBufferUsage::UniformBuffer_SingleFrame);
 		// upload sky/lights data
-		SetLightParameters(GraphBuilder, PassParameters, Scene, View);
+		SetLightParameters(GraphBuilder, PassParameters, Scene, View, UseMISCompensation);
 		PassParameters->IESTextureSampler = TStaticSamplerState<SF_Bilinear, AM_Clamp, AM_Clamp, AM_Clamp>::GetRHI();
 		PassParameters->RadianceTexture = GraphBuilder.CreateUAV(RadianceTexture);
 
diff --git a/Engine/Source/Runtime/Renderer/Private/ReflectionEnvironmentCapture.cpp b/Engine/Source/Runtime/Renderer/Private/ReflectionEnvironmentCapture.cpp
--- a/Engine/Source/Runtime/Renderer/Private/ReflectionEnvironmentCapture.cpp
+++ b/Engine/Source/Runtime/Renderer/Private/ReflectionEnvironmentCapture.cpp
@@ -1768,5 +1768,14 @@ void FScene::UpdateSkyCaptureContents(
 				GRenderTargetPool.FreeUnusedResources();
 			});
 		}
+
+		// These textures should only be manipulated by the render thread,
+		// so enqueue a render command for them to be processed there
+		ENQUEUE_RENDER_COMMAND(ReleasePathTracerSkylightData)(
+			[this](FRHICommandListImmediate& RHICmdList)
+		{
+			PathTracingSkylightTexture.SafeRelease();
+			PathTracingSkylightPdf.SafeRelease();
+		});
 	}
 }
diff --git a/Engine/Source/Runtime/Renderer/Private/ReflectionEnvironmentRealTimeCapture.cpp b/Engine/Source/Runtime/Renderer/Private/ReflectionEnvironmentRealTimeCapture.cpp
--- a/Engine/Source/Runtime/Renderer/Private/ReflectionEnvironmentRealTimeCapture.cpp
+++ b/Engine/Source/Runtime/Renderer/Private/ReflectionEnvironmentRealTimeCapture.cpp
@@ -889,6 +889,10 @@ void FScene::AllocateAndCaptureFrameSkyEnvMap(
 
 		// Reset Scene time slicing state so that it starts from the beginning if/when we get out of non-time-sliced.
 		RealTimeSlicedReflectionCaptureState = -1; // Value of -1 indicates this is the first time-sliced iteration.
+
+		// The sky just changed, so invalidate these textures, so that the path tracer can rebuild them
+		PathTracingSkylightTexture.SafeRelease();
+		PathTracingSkylightPdf.SafeRelease();
 	}
 	else
 	{
@@ -1015,6 +1019,10 @@ void FScene::AllocateAndCaptureFrameSkyEnvMap(
 
 			// Now use the new cubemap
 			ConvolvedSkyRenderTargetReadyIndex = ConvolvedSkyRenderTargetWorkIndex;
+
+			// The sky just changed, so invalidate these textures, so that the path tracer can rebuild them
+			PathTracingSkylightTexture.SafeRelease();
+			PathTracingSkylightPdf.SafeRelease();
 		}
 
 #if DEBUG_TIME_SLICE
diff --git a/Engine/Source/Runtime/Renderer/Private/ScenePrivate.h b/Engine/Source/Runtime/Renderer/Private/ScenePrivate.h
--- a/Engine/Source/Runtime/Renderer/Private/ScenePrivate.h
+++ b/Engine/Source/Runtime/Renderer/Private/ScenePrivate.h
@@ -999,10 +999,12 @@ public:
 
 	// Reference path tracing cached results
 	TRefCountPtr<IPooledRenderTarget> PathTracingRadianceRT;
+	// Keep track of the rectangle of pixels the Radiance texture is valid for so that path tracing can restart if this changes
 	FIntRect PathTracingRect;
+	// Target sampling count for the path tracer - to allow different views to target different quality levels
 	uint32   PathTracingTargetSPP;
 
-	// Path tracer samples/pixel
+	// Current sample index to be rendered by the path tracer - this gets incremented each time the path tracer accumulates a frame of samples
 	uint32 PathTracingSampleIndex;
 
 	// Path tracer frame index, not reset on invalidation unlike PathTracingSampleIndex to avoid
@@ -2502,6 +2504,14 @@ public:
 	/** The current progress of the real time reflection capture when time sliced. */
 	int32 RealTimeSlicedReflectionCaptureState;
 
+	/**
+	 * The path tracer uses its own representation of the skylight. These textures
+	 * are updated lazily by the path tracer when missing. Any code that modifies
+	 * the skylight appearance should simplify reset these pointers.
+	 */
+	TRefCountPtr<IPooledRenderTarget> PathTracingSkylightTexture;
+	TRefCountPtr<IPooledRenderTarget> PathTracingSkylightPdf;
+
 	/** Used to track the order that skylights were enabled in. */
 	TArray<FSkyLightSceneProxy*> SkyLightStack;
 
