Reduce the amount of work performed during rendering when the path tracer is active

The raster passes "under" the path tracer now look much closer to the "unlit" mode. I tweaked the logic for bCanOverlayRayTracingOutput to ignore the wiper cvar so that this variable can be used to confirm what is happening under the path tracer.

diff --git a/Engine/Source/Runtime/Engine/Private/SceneView.cpp b/Engine/Source/Runtime/Engine/Private/SceneView.cpp
--- a/Engine/Source/Runtime/Engine/Private/SceneView.cpp
+++ b/Engine/Source/Runtime/Engine/Private/SceneView.cpp
@@ -2431,7 +2431,7 @@ void FSceneView::SetupCommonViewUniformBufferParameters(
 		InPrevViewMatrices.GetTemporalAAJitter().X, InPrevViewMatrices.GetTemporalAAJitter().Y );
 
 	ViewUniformShaderParameters.DebugViewModeMask = Family->UseDebugViewPS() ? 1 : 0;
-	ViewUniformShaderParameters.UnlitViewmodeMask = !Family->EngineShowFlags.Lighting ? 1 : 0;
+	ViewUniformShaderParameters.UnlitViewmodeMask = !Family->EngineShowFlags.Lighting || Family->EngineShowFlags.PathTracing ? 1 : 0;
 	ViewUniformShaderParameters.OutOfBoundsMask = Family->EngineShowFlags.VisualizeOutOfBoundsPixels ? 1 : 0;
 
 	ViewUniformShaderParameters.GameTime = Family->CurrentWorldTime;
diff --git a/Engine/Source/Runtime/Renderer/Private/DeferredShadingRenderer.cpp b/Engine/Source/Runtime/Renderer/Private/DeferredShadingRenderer.cpp
--- a/Engine/Source/Runtime/Renderer/Private/DeferredShadingRenderer.cpp
+++ b/Engine/Source/Runtime/Renderer/Private/DeferredShadingRenderer.cpp
@@ -2052,7 +2052,7 @@ void FDeferredShadingSceneRenderer::Render(FRDGBuilder& GraphBuilder)
 	// End early occlusion queries
 
 	// Early Shadow depth rendering
-	if (bOcclusionBeforeBasePass)
+	if (bCanOverlayRayTracingOutput && bOcclusionBeforeBasePass)
 	{
 		const bool bAfterBasePass = false;
 		AllocateVirtualShadowMaps(bAfterBasePass);
@@ -2068,7 +2068,7 @@ void FDeferredShadingSceneRenderer::Render(FRDGBuilder& GraphBuilder)
 	const bool bShouldVisualizeVolumetricCloud = bShouldRenderVolumetricCloudBase && !!ViewFamily.EngineShowFlags.VisualizeVolumetricCloudConservativeDensity;
 	bool bAsyncComputeVolumetricCloud = IsVolumetricRenderTargetEnabled() && IsVolumetricRenderTargetAsyncCompute();
 	bool bHasHalfResCheckerboardMinMaxDepth = false;
-	bool bVolumetricRenderTargetRequired = bShouldRenderVolumetricCloud;
+	bool bVolumetricRenderTargetRequired = bShouldRenderVolumetricCloud && bCanOverlayRayTracingOutput;
 
 	if (bShouldRenderVolumetricCloudBase)
 	{
@@ -2126,7 +2126,7 @@ void FDeferredShadingSceneRenderer::Render(FRDGBuilder& GraphBuilder)
 	FRDGTextureRef HalfResolutionDepthCheckerboardMinMaxTexture = nullptr;
 
 	// Kick off async compute cloud eraly if all depth has been written in the prepass
-	if (bShouldRenderVolumetricCloud && bAsyncComputeVolumetricCloud && DepthPass.EarlyZPassMode == DDM_AllOpaque)
+	if (bShouldRenderVolumetricCloud && bAsyncComputeVolumetricCloud && DepthPass.EarlyZPassMode == DDM_AllOpaque && bCanOverlayRayTracingOutput)
 	{
 		HalfResolutionDepthCheckerboardMinMaxTexture = CreateHalfResolutionDepthCheckerboardMinMax(GraphBuilder, Views, SceneTextures.Depth.Resolve);
 		bHasHalfResCheckerboardMinMaxDepth = true;
@@ -2278,7 +2278,7 @@ void FDeferredShadingSceneRenderer::Render(FRDGBuilder& GraphBuilder)
 #endif // RHI_RAYTRACING
 
 	// Shadow and fog after base pass
-	if (!bOcclusionBeforeBasePass)
+	if (bCanOverlayRayTracingOutput && !bOcclusionBeforeBasePass)
 	{
 		const bool bAfterBasePass = true;
 		AllocateVirtualShadowMaps(bAfterBasePass);
@@ -2291,22 +2291,14 @@ void FDeferredShadingSceneRenderer::Render(FRDGBuilder& GraphBuilder)
 		}
 
 		AddServiceLocalQueuePass(GraphBuilder);
+		ComputeVolumetricFog(GraphBuilder);
 	}
 	// End shadow and fog after base pass
 
 	if (bNaniteEnabled)
 	{
-		if (!bOcclusionBeforeBasePass)
-		{
-			ComputeVolumetricFog(GraphBuilder);
-		}
-
 		Nanite::GStreamingManager.SubmitFrameStreamingRequests(GraphBuilder);
 	}
-	else if (!bOcclusionBeforeBasePass)
-	{
-		ComputeVolumetricFog(GraphBuilder);
-	}
 
 	if (VirtualShadowMapArray.IsEnabled())
 	{
@@ -2324,7 +2316,7 @@ void FDeferredShadingSceneRenderer::Render(FRDGBuilder& GraphBuilder)
 	}
 
 	// If not all depth is written during the prepass, kick off async compute cloud after basepass
-	if (bShouldRenderVolumetricCloud && bAsyncComputeVolumetricCloud && DepthPass.EarlyZPassMode != DDM_AllOpaque)
+	if (bShouldRenderVolumetricCloud && bAsyncComputeVolumetricCloud && DepthPass.EarlyZPassMode != DDM_AllOpaque && bCanOverlayRayTracingOutput)
 	{
 		HalfResolutionDepthCheckerboardMinMaxTexture = CreateHalfResolutionDepthCheckerboardMinMax(GraphBuilder, Views, SceneTextures.Depth.Resolve);
 		bHasHalfResCheckerboardMinMaxDepth = true;
@@ -2483,12 +2475,12 @@ void FDeferredShadingSceneRenderer::Render(FRDGBuilder& GraphBuilder)
 		RenderDeferredReflectionsAndSkyLightingHair(GraphBuilder, HairDatas);
 	}
 
-	if (bShouldRenderVolumetricCloud && IsVolumetricRenderTargetEnabled() && !bHasHalfResCheckerboardMinMaxDepth)
+	if (bShouldRenderVolumetricCloud && IsVolumetricRenderTargetEnabled() && !bHasHalfResCheckerboardMinMaxDepth && bCanOverlayRayTracingOutput)
 	{
 		HalfResolutionDepthCheckerboardMinMaxTexture = CreateHalfResolutionDepthCheckerboardMinMax(GraphBuilder, Views, SceneTextures.Depth.Resolve);
 	}
 
-	if (bShouldRenderVolumetricCloud)
+	if (bShouldRenderVolumetricCloud && bCanOverlayRayTracingOutput)
 	{
 		if (!bAsyncComputeVolumetricCloud)
 		{
@@ -2501,7 +2493,7 @@ void FDeferredShadingSceneRenderer::Render(FRDGBuilder& GraphBuilder)
 		ReconstructVolumetricRenderTarget(GraphBuilder, Views, SceneTextures.Depth.Resolve, HalfResolutionDepthCheckerboardMinMaxTexture, bAsyncComputeVolumetricCloud);
 	}
 
-	const bool bShouldRenderTranslucency = ShouldRenderTranslucency();
+	const bool bShouldRenderTranslucency = bCanOverlayRayTracingOutput && ShouldRenderTranslucency();
 
 	// Union of all translucency view render flags.
 	ETranslucencyView TranslucencyViewsToRender = bShouldRenderTranslucency ? GetTranslucencyViews(Views) : ETranslucencyView::None;
@@ -2530,7 +2522,7 @@ void FDeferredShadingSceneRenderer::Render(FRDGBuilder& GraphBuilder)
 	FRDGTextureRef LightShaftOcclusionTexture = nullptr;
 
 	// Draw Lightshafts
-	if (ViewFamily.EngineShowFlags.LightShafts)
+	if (bCanOverlayRayTracingOutput && ViewFamily.EngineShowFlags.LightShafts)
 	{
 		SCOPE_CYCLE_COUNTER(STAT_FDeferredShadingSceneRenderer_RenderLightShaftOcclusion);
 		LightShaftOcclusionTexture = RenderLightShaftOcclusion(GraphBuilder, SceneTextures);
@@ -2544,7 +2536,7 @@ void FDeferredShadingSceneRenderer::Render(FRDGBuilder& GraphBuilder)
 	}
 
 	// Draw the sky atmosphere
-	if (bShouldRenderSkyAtmosphere)
+	if (bCanOverlayRayTracingOutput && bShouldRenderSkyAtmosphere)
 	{
 		SCOPE_CYCLE_COUNTER(STAT_FDeferredShadingSceneRenderer_RenderSkyAtmosphere);
 		RenderSkyAtmosphere(GraphBuilder, SceneTextures);
@@ -2559,7 +2551,7 @@ void FDeferredShadingSceneRenderer::Render(FRDGBuilder& GraphBuilder)
 	}
 
 	// After the height fog, Draw volumetric clouds (having fog applied on them already) when using per pixel tracing,
-	if (bShouldRenderVolumetricCloud)
+	if (bCanOverlayRayTracingOutput && bShouldRenderVolumetricCloud)
 	{
 		bool bSkipVolumetricRenderTarget = true;
 		bool bSkipPerPixelTracing = false;
@@ -2577,7 +2569,7 @@ void FDeferredShadingSceneRenderer::Render(FRDGBuilder& GraphBuilder)
 
 	RenderOpaqueFX(GraphBuilder, Views, FXSystem, SceneTextures.UniformBuffer);
 
-	if (bShouldRenderSkyAtmosphere)
+	if (bCanOverlayRayTracingOutput && bShouldRenderSkyAtmosphere)
 	{
 		// Debug the sky atmosphere. Critically rendered before translucency to avoid emissive leaking over visualization by writing depth. 
 		// Alternative: render in post process chain as VisualizeHDR.
@@ -2727,7 +2719,7 @@ void FDeferredShadingSceneRenderer::Render(FRDGBuilder& GraphBuilder)
 		AddServiceLocalQueuePass(GraphBuilder);
 	}
 
-	if (bShouldVisualizeVolumetricCloud)
+	if (bShouldVisualizeVolumetricCloud && bCanOverlayRayTracingOutput)
 	{
 		RenderVolumetricCloud(GraphBuilder, SceneTextures, false, true, HalfResolutionDepthCheckerboardMinMaxTexture, false, InstanceCullingManager);
 		ReconstructVolumetricRenderTarget(GraphBuilder, Views, SceneTextures.Depth.Resolve, HalfResolutionDepthCheckerboardMinMaxTexture, false);
@@ -2880,11 +2872,9 @@ bool ShouldRenderRayTracingEffect(bool bEffectEnabled)
 
 bool CanOverlayRayTracingOutput(const FViewInfo& View)
 {
-	static auto CVarWiper = IConsoleManager::Get().FindConsoleVariable(TEXT("r.PathTracing.WiperMode"));
-
-	return (
-		(View.RayTracingRenderMode != ERayTracingRenderMode::PathTracing || (View.RayTracingRenderMode == ERayTracingRenderMode::PathTracing && CVarWiper && CVarWiper->GetInt() > 0))
-		&& (View.RayTracingRenderMode != ERayTracingRenderMode::RayTracingDebug)
-		);
+	// Return false if a full screen ray tracing pass will be displayed on top of the raster pass
+	// This can be used to skip certain calculations
+	return (View.RayTracingRenderMode != ERayTracingRenderMode::PathTracing &&
+			View.RayTracingRenderMode != ERayTracingRenderMode::RayTracingDebug);
 }
 #endif // RHI_RAYTRACING
