Make skylight visible by default in path tracer when realtime capture mode is active

Simplify the logic for directly visible lights

Fix raycone not being propagated

Call TraceRay directly in path tracer for more granular control over the packed payload

diff --git a/Engine/Shaders/Private/PathTracing/PathTracing.usf b/Engine/Shaders/Private/PathTracing/PathTracing.usf
--- a/Engine/Shaders/Private/PathTracing/PathTracing.usf
+++ b/Engine/Shaders/Private/PathTracing/PathTracing.usf
@@ -24,6 +24,7 @@ PathTracingRayGenShader.usf: Reference path tracing
 RWTexture2D<float4> RadianceTexture;
 RaytracingAccelerationStructure TLAS;
 int3 TileOffset;
+uint SceneVisibleLightCount;
 
 void AccumulateRadiance(inout float3 TotalRadiance, float3 PathRadiance)
 {
@@ -42,40 +43,48 @@ void AccumulateRadiance(inout float3 TotalRadiance, float3 PathRadiance)
 
 FMaterialClosestHitPayload TraceTransparentRay(RayDesc Ray, FRayCone RayCone, bool IsCameraRay, bool LastBounce, bool IncludeEmission, uint2 LaunchIndex, uint NumLights, inout RandomSequence RandSequence, inout float3 PathThroughput, inout float3 Radiance)
 {
+	const uint RayFlags = IsCameraRay && PathTracingData.EnableCameraBackfaceCulling ? RAY_FLAG_CULL_BACK_FACING_TRIANGLES : 0;
+	const uint MissShaderIndex = 0;
 	float SelectionWeightSum = 0;
 	float3 PayloadThroughput = PathThroughput;
 	FMaterialClosestHitPayload Payload;
 	for (;;)
 	{
-		const uint InitialRayFlags = IsCameraRay && PathTracingData.EnableCameraBackfaceCulling ? RAY_FLAG_CULL_BACK_FACING_TRIANGLES : 0;
-		const bool bInitialDisableSkyLightContribution = false;
-		const bool bIgnoreTranslucentMaterials = false;
-		FMaterialClosestHitPayload HitPayload = TraceMaterialRay(
+		FPackedMaterialClosestHitPayload PackedPayload = (FPackedMaterialClosestHitPayload)0;
+		// Set payload inputs
+		PackedPayload.SetRayCone(RayCone);
+		PackedPayload.SetPixelCoord(LaunchIndex);
+		// Trace the ray
+		TraceRay(
 			TLAS,
-			InitialRayFlags,
+			RayFlags,
 			RAY_TRACING_MASK_ALL,
+			RAY_TRACING_SHADER_SLOT_MATERIAL,
+			RAY_TRACING_NUM_SHADER_SLOTS,
+			MissShaderIndex,
 			Ray,
-			RayCone,
-			LaunchIndex,
-			bInitialDisableSkyLightContribution,
-			bIgnoreTranslucentMaterials);
+			PackedPayload);
 
 		// Loop over lights to capture their contribution
 		// #dxr_todo: if we have lots of lights, having some hierarchical structure would be better ....
 		for (uint LightId = 0; LightId < NumLights; ++LightId)
 		{
 			RayDesc LightRay = Ray;
-			LightRay.TMax = HitPayload.IsMiss() ? Ray.TMax : HitPayload.HitT;
+			LightRay.TMax = PackedPayload.IsMiss() ? Ray.TMax : PackedPayload.HitT;
 			float3 LightRadiance = TraceLight(LightRay, LightId).Radiance;
 			AccumulateRadiance(Radiance, PathThroughput * LightRadiance);
 		}
 
-		if (HitPayload.IsMiss())
+		if (PackedPayload.IsMiss())
 		{
 			// Ray didn't hit any real geometry, so nothing left to do
 			break;
 		}
 
+		// Unpack the payload
+		FMaterialClosestHitPayload HitPayload = UnpackRayTracingPayload(PackedPayload, Ray);
+
+
 		// material tweaks to workaround temporary limitations
 		if (HitPayload.ShadingModelID == SHADINGMODELID_SUBSURFACE ||
 			HitPayload.ShadingModelID == SHADINGMODELID_PREINTEGRATED_SKIN)
@@ -250,11 +259,11 @@ RAY_TRACING_ENTRY_RAYGEN(PathTracingMainRG)
 	const int MISMode = PathTracingData.MISMode;
 
 	// number of directly visible lights for the first bounce
-	uint NumLights = PathTracingData.VisibleLights ? SceneLightCount : 0;
+	uint NumVisibleLights = SceneVisibleLightCount;
 
 	for (int Bounce = 0; Bounce <= PathTracingData.MaxBounces; Bounce++)
 	{
-		FMaterialClosestHitPayload Payload = TraceTransparentRay(Ray, RayCone, Bounce == 0, Bounce == PathTracingData.MaxBounces, PathTracingData.SkipDirectLighting == 0 || Bounce > 1, LaunchIndex, NumLights, RandSequence, PathThroughput, Radiance);
+		FMaterialClosestHitPayload Payload = TraceTransparentRay(Ray, RayCone, Bounce == 0, Bounce == PathTracingData.MaxBounces, PathTracingData.SkipDirectLighting == 0 || Bounce > 1, LaunchIndex, NumVisibleLights, RandSequence, PathThroughput, Radiance);
 
 		if (Payload.IsMiss())
 		{
@@ -445,7 +454,10 @@ RAY_TRACING_ENTRY_RAYGEN(PathTracingMainRG)
 
 		// from this point on, we don't need to include lights in the trace call
 		// because NEE handled it for us
-		NumLights = 0;
+		NumVisibleLights = 0;
+		
+		// update ray cone for the next bounce
+		RayCone = Payload.GetRayCone();
 	}
 
 
diff --git a/Engine/Source/Runtime/Renderer/Private/PathTracing.cpp b/Engine/Source/Runtime/Renderer/Private/PathTracing.cpp
--- a/Engine/Source/Runtime/Renderer/Private/PathTracing.cpp
+++ b/Engine/Source/Runtime/Renderer/Private/PathTracing.cpp
@@ -166,7 +166,7 @@ static bool PrepareShaderArgs(const FViewInfo& View, FPathTracingData& PathTraci
 	PathTracingData.MaxBounces = MaxBounces;
 	PathTracingData.MaxNormalBias = GetRaytracingMaxNormalBias();
 	PathTracingData.MISMode = CVarPathTracingMISMode.GetValueOnRenderThread();
-	PathTracingData.VisibleLights = CVarPathTracingVisibleLights.GetValueOnRenderThread();
+	uint32 VisibleLights = CVarPathTracingVisibleLights.GetValueOnRenderThread();
 	PathTracingData.MaxPathIntensity = CVarPathTracingMaxPathIntensity.GetValueOnRenderThread();
 	PathTracingData.UseErrorDiffusion = CVarPathTracingUseErrorDiffusion.GetValueOnRenderThread();
 	PathTracingData.ApproximateCaustics = CVarPathTracingApproximateCaustics.GetValueOnRenderThread();
@@ -199,11 +199,11 @@ static bool PrepareShaderArgs(const FViewInfo& View, FPathTracingData& PathTraci
 	}
 
 	// Changing VisibleLights requires starting over
-	static uint32 PreviousVisibleLights = PathTracingData.VisibleLights;
-	if (PreviousVisibleLights != PathTracingData.VisibleLights)
+	static uint32 PreviousVisibleLights = VisibleLights;
+	if (PreviousVisibleLights != VisibleLights)
 	{
 		NeedInvalidation = true;
-		PreviousVisibleLights = PathTracingData.VisibleLights;
+		PreviousVisibleLights = VisibleLights;
 	}
 
 	// Changing MaxPathIntensity requires starting over
@@ -342,6 +342,7 @@ class FPathTracingRG : public FGlobalShader
 		SHADER_PARAMETER_STRUCT_REF(FPathTracingData, PathTracingData)
 		SHADER_PARAMETER_RDG_BUFFER_SRV(StructuredBuffer<FPathTracingLight>, SceneLights)
 		SHADER_PARAMETER(uint32, SceneLightCount)
+		SHADER_PARAMETER(uint32, SceneVisibleLightCount)
 		// Skylight
 		SHADER_PARAMETER_RDG_TEXTURE(Texture2D, SkylightTexture)
 		SHADER_PARAMETER_RDG_TEXTURE(Texture2D, SkylightPdf)
@@ -492,6 +493,8 @@ bool PrepareSkyTexture(FRDGBuilder& GraphBuilder, FScene* Scene, const FViewInfo
 
 void SetLightParameters(FRDGBuilder& GraphBuilder, FPathTracingRG::FParameters* PassParameters, FScene* Scene, const FViewInfo& View, bool UseMISCompensation)
 {
+	PassParameters->SceneVisibleLightCount = 0;
+
 	// Lights
 	FPathTracingLight Lights[RAY_TRACING_LIGHT_COUNT_MAXIMUM]; // Keep this on the stack for now -- eventually will need to make this dynamic to lift size limit (and also avoid uploading per frame ...)
 	unsigned LightCount = 0;
@@ -507,6 +510,13 @@ void SetLightParameters(FRDGBuilder& GraphBuilder, FPathTracingRG::FParameters*
 		DestLight.Flags |= PATHTRACING_LIGHT_SKY;
 		DestLight.Flags |= Scene->SkyLight->bCastShadows ? PATHTRACER_FLAG_CAST_SHADOW_MASK : 0;
 		DestLight.IESTextureSlice = -1;
+		if (Scene->SkyLight->bRealTimeCaptureEnabled)
+		{
+			// When using the realtime capture system, always make the skylight visible
+			// because this is our only way of "seeing" the atmo/clouds at the moment
+			PassParameters->SceneVisibleLightCount = 1;
+		}
+
 		LightCount++;
 	}
 
@@ -610,6 +620,11 @@ void SetLightParameters(FRDGBuilder& GraphBuilder, FPathTracingRG::FParameters*
 		PassParameters->SceneLightCount = LightCount;
 	}
 
+	if (CVarPathTracingVisibleLights.GetValueOnRenderThread() != 0)
+	{
+		PassParameters->SceneVisibleLightCount = LightCount;
+	}
+
 	if (!IESLightProfilesMap.IsEmpty())
 	{
 		// We found some IES profiles to use -- upload them into a single atlas so we can access them easily in HLSL
@@ -821,7 +836,7 @@ void FDeferredShadingSceneRenderer::RenderPathTracing(
 		TShaderMapRef<FPathTracingRG> RayGenShader(View.ShaderMap);
 		ClearUnusedGraphResources(RayGenShader, PassParameters);
 		GraphBuilder.AddPass(
-			RDG_EVENT_NAME("Path Tracer Compute (%d x %d) Sample=%d/%d", View.ViewRect.Size().X, View.ViewRect.Size().Y, View.ViewState->PathTracingSampleIndex, MaxSPP),
+			RDG_EVENT_NAME("Path Tracer Compute (%d x %d) Sample=%d/%d NumLights=%d", View.ViewRect.Size().X, View.ViewRect.Size().Y, View.ViewState->PathTracingSampleIndex, MaxSPP, PassParameters->SceneLightCount),
 			PassParameters,
 			ERDGPassFlags::Compute,
 			[PassParameters, RayGenShader, &View](FRHICommandListImmediate& RHICmdList)
diff --git a/Engine/Source/Runtime/Renderer/Private/PathTracingUniformBuffers.h b/Engine/Source/Runtime/Renderer/Private/PathTracingUniformBuffers.h
--- a/Engine/Source/Runtime/Renderer/Private/PathTracingUniformBuffers.h
+++ b/Engine/Source/Runtime/Renderer/Private/PathTracingUniformBuffers.h
@@ -16,7 +16,6 @@ BEGIN_GLOBAL_SHADER_PARAMETER_STRUCT(FPathTracingData, )
 	SHADER_PARAMETER(uint32, UseErrorDiffusion)
 	SHADER_PARAMETER(uint32, MaxBounces)
 	SHADER_PARAMETER(uint32, MISMode)
-	SHADER_PARAMETER(uint32, VisibleLights)
 	SHADER_PARAMETER(uint32, ApproximateCaustics)
 	SHADER_PARAMETER(uint32, EnableCameraBackfaceCulling)
 	SHADER_PARAMETER(uint32, SkipDirectLighting)
