Increase default r.RayTracing.ParallelMeshBatchSize from 128 to 1024 to avoid some of the task overhead and resource contention.

diff --git a/Engine/Source/Runtime/Renderer/Private/DeferredShadingRenderer.cpp b/Engine/Source/Runtime/Renderer/Private/DeferredShadingRenderer.cpp
--- a/Engine/Source/Runtime/Renderer/Private/DeferredShadingRenderer.cpp
+++ b/Engine/Source/Runtime/Renderer/Private/DeferredShadingRenderer.cpp
@@ -156,7 +156,7 @@ static FAutoConsoleVariableRef CRayTracingParallelMeshBatchSetup(
 	TEXT("Whether to setup ray tracing materials via parallel jobs."),
 	ECVF_RenderThreadSafe);
 
-static int32 GRayTracingParallelMeshBatchSize = 128;
+static int32 GRayTracingParallelMeshBatchSize = 1024;
 static FAutoConsoleVariableRef CRayTracingParallelMeshBatchSize(
 	TEXT("r.RayTracing.ParallelMeshBatchSize"),
 	GRayTracingParallelMeshBatchSize,
@@ -973,7 +973,7 @@ bool FDeferredShadingSceneRenderer::GatherRayTracingWorldInstances(FRHICommandLi
 	if (ReferenceView.AddRayTracingMeshBatchData.Num() > 0)
 	{
 		const uint32 NumTotalItems = ReferenceView.AddRayTracingMeshBatchData.Num();
-		const uint32 TargetItemsPerTask = GRayTracingParallelMeshBatchSize;//128; // Granularity based on profiling several scenes
+		const uint32 TargetItemsPerTask = GRayTracingParallelMeshBatchSize;// 1024; // Granularity based on profiling several scenes
 		const uint32 NumTasks = FMath::Max(1u, FMath::DivideAndRoundUp(NumTotalItems, TargetItemsPerTask));
 		const uint32 BatchSize = FMath::DivideAndRoundUp(NumTotalItems, NumTasks);
 
