Added public ResetSceneTextureExtentHistory function to reset render target size after high resolution rendering.

#rnx

diff --git a/Engine/Source/Runtime/Renderer/Private/Renderer.cpp b/Engine/Source/Runtime/Renderer/Private/Renderer.cpp
--- a/Engine/Source/Runtime/Renderer/Private/Renderer.cpp
+++ b/Engine/Source/Runtime/Renderer/Private/Renderer.cpp
@@ -453,7 +453,6 @@ static void VisualizeTextureExec( const TCHAR* Cmd, FOutputDevice &Ar )
 }
 
 extern void NaniteStatsFilterExec(const TCHAR* Cmd, FOutputDevice& Ar);
-extern void ResetRenderTargetsExtent(FOutputDevice& Ar);
 
 static bool RendererExec( UWorld* InWorld, const TCHAR* Cmd, FOutputDevice& Ar )
 {
@@ -484,7 +483,8 @@ static bool RendererExec( UWorld* InWorld, const TCHAR* Cmd, FOutputDevice& Ar )
 	}
 	else if (FParse::Command(&Cmd, TEXT("r.ResetRenderTargetsExtent")))
 	{
-		ResetRenderTargetsExtent(Ar);
+		ResetSceneTextureExtentHistory();
+		Ar.Logf(TEXT("Scene texture extent history reset. Next scene render will reallocate textures at the requested size."));
 		return true;
 	}
 #endif
diff --git a/Engine/Source/Runtime/Renderer/Private/SceneTextures.cpp b/Engine/Source/Runtime/Renderer/Private/SceneTextures.cpp
--- a/Engine/Source/Runtime/Renderer/Private/SceneTextures.cpp
+++ b/Engine/Source/Runtime/Renderer/Private/SceneTextures.cpp
@@ -434,7 +434,7 @@ public:
 		return DesiredExtent;
 	}
 
-	void ResetRenderTargetsExtent()
+	void ResetHistory()
 	{
 		LastStereoExtent = FIntPoint(0, 0);
 		LastExtent = FIntPoint(0, 0);
@@ -488,11 +488,9 @@ private:
 	uint32 ThisFrameNumber = 0;
 };
 
-
-void ResetRenderTargetsExtent(FOutputDevice& Ar)
+void ResetSceneTextureExtentHistory()
 {
-	FSceneTextureExtentState::Get().ResetRenderTargetsExtent();
-	Ar.Logf(TEXT("Internal render target extents reseted"));
+	FSceneTextureExtentState::Get().ResetHistory();
 }
 
 ENUM_CLASS_FLAGS(FSceneTextureExtentState::ERenderTargetHistory);
diff --git a/Engine/Source/Runtime/Renderer/Public/SceneRenderTargetParameters.h b/Engine/Source/Runtime/Renderer/Public/SceneRenderTargetParameters.h
--- a/Engine/Source/Runtime/Renderer/Public/SceneRenderTargetParameters.h
+++ b/Engine/Source/Runtime/Renderer/Public/SceneRenderTargetParameters.h
@@ -197,6 +197,10 @@ extern RENDERER_API FIntPoint GetSceneTextureExtent();
 /** Returns the feature level being used by the renderer. */
 extern RENDERER_API ERHIFeatureLevel::Type GetSceneTextureFeatureLevel();
 
+/** Resets the scene texture extent history. Call this method after rendering with very large render
+ *  targets. The next scene render will create them at the requested size.
+ */
+extern RENDERER_API void ResetSceneTextureExtentHistory();
 
 ///////////////////////////////////////////////////////////////////////////
 // Deprecated APIs
