Remove a unused mobile base pass shader permutation class.
Wrap the 'min' instuction with ALLOW_STATIC_LIGHTING macro in mobile base pass pixel shader.

diff --git a/Engine/Shaders/Private/MobileBasePassPixelShader.usf b/Engine/Shaders/Private/MobileBasePassPixelShader.usf
--- a/Engine/Shaders/Private/MobileBasePassPixelShader.usf
+++ b/Engine/Shaders/Private/MobileBasePassPixelShader.usf
@@ -737,7 +737,8 @@ void Main(
 	if (UseCSM)
 	{
 		half ShadowMap = MobileDirectionalLightCSM(MaterialParameters.ScreenPosition.xy, MaterialParameters.ScreenPosition.w, ShadowPositionZ);
-	#if MOVABLE_DIRECTIONAL_LIGHT
+	// Skip the min instruction since the Shadow is always 1 either with movable directional light or the static lighting is disabled.
+	#if MOVABLE_DIRECTIONAL_LIGHT || !ALLOW_STATIC_LIGHTING
 		Shadow = ShadowMap;
 	#else
 		Shadow = min(ShadowMap, Shadow);
diff --git a/Engine/Source/Runtime/Landscape/Private/LandscapeRender.cpp b/Engine/Source/Runtime/Landscape/Private/LandscapeRender.cpp
--- a/Engine/Source/Runtime/Landscape/Private/LandscapeRender.cpp
+++ b/Engine/Source/Runtime/Landscape/Private/LandscapeRender.cpp
@@ -3549,11 +3549,6 @@ public:
 			FName(TEXT("TMobileBasePassPSFMobileMovableDirectionalLightCSMLightingPolicy0HDRLinear64Skylight")),
 			FName(TEXT("TMobileBasePassPSFMobileMovableDirectionalLightCSMLightingPolicy0HDRLinear64")),
 			FName(TEXT("TMobileBasePassVSFMobileMovableDirectionalLightCSMLightingPolicyHDRLinear64")),
-			FName(TEXT("TMobileBasePassPSFMobileMovableDirectionalLightLightingPolicyINT32_MAXHDRLinear64Skylight")),
-			FName(TEXT("TMobileBasePassPSFMobileMovableDirectionalLightLightingPolicyINT32_MAXHDRLinear64")),
-			FName(TEXT("TMobileBasePassPSFMobileMovableDirectionalLightLightingPolicy0HDRLinear64Skylight")),
-			FName(TEXT("TMobileBasePassPSFMobileMovableDirectionalLightLightingPolicy0HDRLinear64")),
-			FName(TEXT("TMobileBasePassVSFMobileMovableDirectionalLightLightingPolicyHDRLinear64")),
 			FName(TEXT("TMobileBasePassPSFMobileDirectionalLightCSMAndSHIndirectPolicyINT32_MAXHDRLinear64Skylight")),
 			FName(TEXT("TMobileBasePassPSFMobileDirectionalLightCSMAndSHIndirectPolicyINT32_MAXHDRLinear64")),
 			FName(TEXT("TMobileBasePassPSFMobileDirectionalLightCSMAndSHIndirectPolicy0HDRLinear64Skylight")),
diff --git a/Engine/Source/Runtime/Renderer/Private/LightMapRendering.h b/Engine/Source/Runtime/Renderer/Private/LightMapRendering.h
--- a/Engine/Source/Runtime/Renderer/Private/LightMapRendering.h
+++ b/Engine/Source/Runtime/Renderer/Private/LightMapRendering.h
@@ -570,34 +570,6 @@ struct FMobileMovableDirectionalLightAndSHIndirectPolicy : public FMobileDirecti
 	}
 };
 
-/** Mobile Specific */
-struct FMobileMovableDirectionalLightLightingPolicy
-{
-	static bool ShouldCompilePermutation(const FMeshMaterialShaderPermutationParameters& Parameters)
-	{
-		if (IsMobileDeferredShadingEnabled(Parameters.Platform) && !UsePermutationWithMobileDeferred(Parameters))
-		{
-			return false;
-		}
-		static auto* CVarMobileAllowMovableDirectionalLights = IConsoleManager::Get().FindTConsoleVariableDataInt(TEXT("r.Mobile.AllowMovableDirectionalLights"));
-		const bool bMobileAllowMovableDirectionalLights = CVarMobileAllowMovableDirectionalLights->GetValueOnAnyThread() != 0;
-		return bMobileAllowMovableDirectionalLights
-			&& Parameters.MaterialParameters.ShadingModels.IsLit();
-	}	
-
-	static void ModifyCompilationEnvironment(const FMaterialShaderPermutationParameters& Parameters, FShaderCompilerEnvironment& OutEnvironment)
-	{
-		OutEnvironment.SetDefine(TEXT("MOVABLE_DIRECTIONAL_LIGHT"), TEXT("1"));
-		FMobileDirectionalLightCSMPolicy::ModifyCompilationEnvironment(Parameters, OutEnvironment);
-		FNoLightMapPolicy::ModifyCompilationEnvironment(Parameters, OutEnvironment);
-	}
-
-	static bool RequiresSkylight()
-	{
-		return false;
-	}
-};
-
 /** Mobile Specific */
 struct FMobileMovableDirectionalLightWithLightmapPolicy : public TLightMapPolicy<LQ_LIGHTMAP>
 {
