Only add ray tracing geometries that previously used a shared vertex buffer buffer to the validation set.

diff --git a/Engine/Source/Runtime/Renderer/Private/DeferredShadingRenderer.cpp b/Engine/Source/Runtime/Renderer/Private/DeferredShadingRenderer.cpp
--- a/Engine/Source/Runtime/Renderer/Private/DeferredShadingRenderer.cpp
+++ b/Engine/Source/Runtime/Renderer/Private/DeferredShadingRenderer.cpp
@@ -919,10 +919,13 @@ bool FDeferredShadingSceneRenderer::GatherRayTracingWorldInstances(FRHICommandLi
 					Swap(ReferenceView.RayTracingGeometryInstances.Last(), RayTracingInstance);
 
 #if DO_CHECK
-					ReferenceView.RayTracingGeometries.Add(Instance.Geometry);
-					for (int32 ViewIndex = 1; ViewIndex < Views.Num(); ViewIndex++)
+					if (Instance.Geometry->DynamicGeometrySharedBufferGenerationID != FRayTracingGeometry::NonSharedVertexBuffers)
 					{
-						Views[ViewIndex].RayTracingGeometries.Add(Instance.Geometry);
+						ReferenceView.RayTracingGeometriesForValidation.Add(Instance.Geometry);
+						for (int32 ViewIndex = 1; ViewIndex < Views.Num(); ViewIndex++)
+						{
+							Views[ViewIndex].RayTracingGeometriesForValidation.Add(Instance.Geometry);
+						}
 					}
 #endif // DO_CHECK
 
@@ -1178,7 +1181,7 @@ bool FDeferredShadingSceneRenderer::DispatchRayTracingWorldUpdates(FRDGBuilder&
 #ifdef DO_CHECK
 		// Validate all the ray tracing geometries lifetimes
 		int64 SharedBufferGenerationID = (int64)Scene->GetRayTracingDynamicGeometryCollection()->GetSharedBufferGenerationID();
-		for (const FRayTracingGeometry* Geometry : View.RayTracingGeometries)
+		for (const FRayTracingGeometry* Geometry : View.RayTracingGeometriesForValidation)
 		{
 			check(Geometry->DynamicGeometrySharedBufferGenerationID == FRayTracingGeometry::NonSharedVertexBuffers || Geometry->DynamicGeometrySharedBufferGenerationID == SharedBufferGenerationID);
 		}
diff --git a/Engine/Source/Runtime/Renderer/Private/SceneRendering.h b/Engine/Source/Runtime/Renderer/Private/SceneRendering.h
--- a/Engine/Source/Runtime/Renderer/Private/SceneRendering.h
+++ b/Engine/Source/Runtime/Renderer/Private/SceneRendering.h
@@ -1407,7 +1407,7 @@ public:
 
 #ifdef DO_CHECK
 	// Keep track of all used RT Geometries which are used to validate the vertex buffer data (see FRayTracingGeometry::DynamicGeometrySharedBufferGenerationID)
-	TSet<const FRayTracingGeometry*> RayTracingGeometries;
+	TSet<const FRayTracingGeometry*> RayTracingGeometriesForValidation;
 #endif
 
 	// Ray tracing scene specific to this view
