Added EVertexFactoryFlags::SupportsRayTracing

diff --git a/Engine/Plugins/FX/Niagara/Source/NiagaraVertexFactories/Private/NiagaraRibbonVertexFactory.cpp b/Engine/Plugins/FX/Niagara/Source/NiagaraVertexFactories/Private/NiagaraRibbonVertexFactory.cpp
--- a/Engine/Plugins/FX/Niagara/Source/NiagaraVertexFactories/Private/NiagaraRibbonVertexFactory.cpp
+++ b/Engine/Plugins/FX/Niagara/Source/NiagaraVertexFactories/Private/NiagaraRibbonVertexFactory.cpp
@@ -186,4 +186,5 @@ IMPLEMENT_VERTEX_FACTORY_PARAMETER_TYPE(FNiagaraRibbonVertexFactory, SF_Pixel, F
 IMPLEMENT_VERTEX_FACTORY_TYPE(FNiagaraRibbonVertexFactory, "/Plugin/FX/Niagara/Private/NiagaraRibbonVertexFactory.ush",
 	  EVertexFactoryFlags::UsedWithMaterials 
 	| EVertexFactoryFlags::SupportsDynamicLighting
+	| EVertexFactoryFlags::SupportsRayTracing
 );
diff --git a/Engine/Plugins/FX/Niagara/Source/NiagaraVertexFactories/Private/NiagaraSpriteVertexFactory.cpp b/Engine/Plugins/FX/Niagara/Source/NiagaraVertexFactories/Private/NiagaraSpriteVertexFactory.cpp
--- a/Engine/Plugins/FX/Niagara/Source/NiagaraVertexFactories/Private/NiagaraSpriteVertexFactory.cpp
+++ b/Engine/Plugins/FX/Niagara/Source/NiagaraVertexFactories/Private/NiagaraSpriteVertexFactory.cpp
@@ -209,4 +209,5 @@ IMPLEMENT_VERTEX_FACTORY_PARAMETER_TYPE(FNiagaraSpriteVertexFactory, SF_RayHitGr
 IMPLEMENT_VERTEX_FACTORY_TYPE(FNiagaraSpriteVertexFactory,"/Plugin/FX/Niagara/Private/NiagaraSpriteVertexFactory.ush",
 	  EVertexFactoryFlags::UsedWithMaterials
 	| EVertexFactoryFlags::SupportsDynamicLighting
+	| EVertexFactoryFlags::SupportsRayTracing
 );
diff --git a/Engine/Plugins/Runtime/HairStrands/Source/HairStrandsCore/Private/HairCardsVertexFactory.cpp b/Engine/Plugins/Runtime/HairStrands/Source/HairStrandsCore/Private/HairCardsVertexFactory.cpp
--- a/Engine/Plugins/Runtime/HairStrands/Source/HairStrandsCore/Private/HairCardsVertexFactory.cpp
+++ b/Engine/Plugins/Runtime/HairStrands/Source/HairStrandsCore/Private/HairCardsVertexFactory.cpp
@@ -344,5 +344,6 @@ IMPLEMENT_VERTEX_FACTORY_TYPE(FHairCardsVertexFactory,"/Engine/Private/HairStran
 	| EVertexFactoryFlags::SupportsDynamicLighting
 	| EVertexFactoryFlags::SupportsPrecisePrevWorldPos
 	| EVertexFactoryFlags::SupportsCachingMeshDrawCommands
+	| EVertexFactoryFlags::SupportsRayTracing
 	| (VF_CARDS_SUPPORT_GPU_SCENE ? EVertexFactoryFlags::SupportsPrimitiveIdStream : EVertexFactoryFlags::None)
 );
diff --git a/Engine/Source/Runtime/Engine/Private/GPUSkinVertexFactory.cpp b/Engine/Source/Runtime/Engine/Private/GPUSkinVertexFactory.cpp
--- a/Engine/Source/Runtime/Engine/Private/GPUSkinVertexFactory.cpp
+++ b/Engine/Source/Runtime/Engine/Private/GPUSkinVertexFactory.cpp
@@ -833,6 +833,7 @@ IMPLEMENT_VERTEX_FACTORY_PARAMETER_TYPE(FGPUSkinPassthroughVertexFactory, SF_Com
 IMPLEMENT_VERTEX_FACTORY_TYPE(FGPUSkinPassthroughVertexFactory, "/Engine/Private/LocalVertexFactory.ush",
 	  EVertexFactoryFlags::UsedWithMaterials
 	| EVertexFactoryFlags::SupportsDynamicLighting
+	| EVertexFactoryFlags::SupportsRayTracing
 	| (GPUCULL_TODO ? EVertexFactoryFlags::SupportsPrimitiveIdStream : EVertexFactoryFlags::None)
 );
 
diff --git a/Engine/Source/Runtime/Engine/Private/GeometryCacheVertexFactory.cpp b/Engine/Source/Runtime/Engine/Private/GeometryCacheVertexFactory.cpp
--- a/Engine/Source/Runtime/Engine/Private/GeometryCacheVertexFactory.cpp
+++ b/Engine/Source/Runtime/Engine/Private/GeometryCacheVertexFactory.cpp
@@ -341,4 +341,5 @@ IMPLEMENT_VERTEX_FACTORY_TYPE(FGeometryCacheVertexVertexFactory, "/Engine/Privat
 	  EVertexFactoryFlags::UsedWithMaterials
 	| EVertexFactoryFlags::SupportsDynamicLighting
 	| EVertexFactoryFlags::SupportsPositionOnly
+	| EVertexFactoryFlags::SupportsRayTracing
 );
diff --git a/Engine/Source/Runtime/Engine/Private/InstancedStaticMesh.cpp b/Engine/Source/Runtime/Engine/Private/InstancedStaticMesh.cpp
--- a/Engine/Source/Runtime/Engine/Private/InstancedStaticMesh.cpp
+++ b/Engine/Source/Runtime/Engine/Private/InstancedStaticMesh.cpp
@@ -848,6 +848,7 @@ IMPLEMENT_VERTEX_FACTORY_TYPE(FInstancedStaticMeshVertexFactory,"/Engine/Private
 	| EVertexFactoryFlags::SupportsPrecisePrevWorldPos
 	| EVertexFactoryFlags::SupportsPositionOnly
 	| EVertexFactoryFlags::SupportsCachingMeshDrawCommands
+	| EVertexFactoryFlags::SupportsRayTracing
 	| (GPUCULL_TODO ? EVertexFactoryFlags::SupportsPrimitiveIdStream : EVertexFactoryFlags::None)
 );
 
diff --git a/Engine/Source/Runtime/Engine/Private/LocalVertexFactory.cpp b/Engine/Source/Runtime/Engine/Private/LocalVertexFactory.cpp
--- a/Engine/Source/Runtime/Engine/Private/LocalVertexFactory.cpp
+++ b/Engine/Source/Runtime/Engine/Private/LocalVertexFactory.cpp
@@ -404,4 +404,5 @@ IMPLEMENT_VERTEX_FACTORY_TYPE(FLocalVertexFactory,"/Engine/Private/LocalVertexFa
 	| EVertexFactoryFlags::SupportsPositionOnly
 	| EVertexFactoryFlags::SupportsCachingMeshDrawCommands
 	| EVertexFactoryFlags::SupportsPrimitiveIdStream
+	| EVertexFactoryFlags::SupportsRayTracing
 );
diff --git a/Engine/Source/Runtime/Landscape/Private/LandscapeRender.cpp b/Engine/Source/Runtime/Landscape/Private/LandscapeRender.cpp
--- a/Engine/Source/Runtime/Landscape/Private/LandscapeRender.cpp
+++ b/Engine/Source/Runtime/Landscape/Private/LandscapeRender.cpp
@@ -3231,6 +3231,7 @@ IMPLEMENT_VERTEX_FACTORY_TYPE(FLandscapeVertexFactory, "/Engine/Private/Landscap
 	| EVertexFactoryFlags::SupportsStaticLighting
 	| EVertexFactoryFlags::SupportsDynamicLighting
 	| EVertexFactoryFlags::SupportsCachingMeshDrawCommands
+	| EVertexFactoryFlags::SupportsRayTracing
 	| (GPUCULL_TODO ? EVertexFactoryFlags::SupportsPrimitiveIdStream : EVertexFactoryFlags::None)
 );
 
@@ -3273,6 +3274,7 @@ IMPLEMENT_VERTEX_FACTORY_TYPE(FLandscapeXYOffsetVertexFactory, "/Engine/Private/
 	| EVertexFactoryFlags::SupportsStaticLighting
 	| EVertexFactoryFlags::SupportsDynamicLighting
 	| EVertexFactoryFlags::SupportsCachingMeshDrawCommands
+	| EVertexFactoryFlags::SupportsRayTracing
 	| (GPUCULL_TODO ? EVertexFactoryFlags::SupportsPrimitiveIdStream : EVertexFactoryFlags::None)
 );
 
@@ -3298,6 +3300,7 @@ IMPLEMENT_VERTEX_FACTORY_TYPE(FLandscapeFixedGridVertexFactory, "/Engine/Private
 	| EVertexFactoryFlags::SupportsStaticLighting
 	| EVertexFactoryFlags::SupportsDynamicLighting
 	| EVertexFactoryFlags::SupportsCachingMeshDrawCommands
+	| EVertexFactoryFlags::SupportsRayTracing
 	| (GPUCULL_TODO ? EVertexFactoryFlags::SupportsPrimitiveIdStream : EVertexFactoryFlags::None)
 );
 
diff --git a/Engine/Source/Runtime/RenderCore/Public/VertexFactory.h b/Engine/Source/Runtime/RenderCore/Public/VertexFactory.h
--- a/Engine/Source/Runtime/RenderCore/Public/VertexFactory.h
+++ b/Engine/Source/Runtime/RenderCore/Public/VertexFactory.h
@@ -100,6 +100,7 @@ enum class EVertexFactoryFlags : uint32
 	SupportsCachingMeshDrawCommands = 1u << 6,
 	SupportsPrimitiveIdStream       = 1u << 7,
 	SupportsNaniteRendering         = 1u << 8,
+	SupportsRayTracing              = 1u << 9,
 };
 ENUM_CLASS_FLAGS(EVertexFactoryFlags);
 
@@ -399,6 +400,7 @@ public:
 	bool SupportsCachingMeshDrawCommands() const { return HasFlags(EVertexFactoryFlags::SupportsCachingMeshDrawCommands); }
 	bool SupportsPrimitiveIdStream() const       { return HasFlags(EVertexFactoryFlags::SupportsPrimitiveIdStream); }
 	bool SupportsNaniteRendering() const         { return HasFlags(EVertexFactoryFlags::SupportsNaniteRendering); }
+	bool SupportsRayTracing() const              { return HasFlags(EVertexFactoryFlags::SupportsRayTracing); }
 
 	// Hash function.
 	friend uint32 GetTypeHash(const FVertexFactoryType* Type)
diff --git a/Engine/Source/Runtime/Renderer/Private/RayTracing/RayTracingMaterialHitShaders.cpp b/Engine/Source/Runtime/Renderer/Private/RayTracing/RayTracingMaterialHitShaders.cpp
--- a/Engine/Source/Runtime/Renderer/Private/RayTracing/RayTracingMaterialHitShaders.cpp
+++ b/Engine/Source/Runtime/Renderer/Private/RayTracing/RayTracingMaterialHitShaders.cpp
@@ -52,27 +52,7 @@ extern int32 GRayTracingUseTextureLod;
 
 static bool IsSupportedVertexFactoryType(const FVertexFactoryType* VertexFactoryType)
 {
-	static FName LocalVfFname = FName(TEXT("FLocalVertexFactory"), FNAME_Find);
-	static FName LSkinnedVfFname = FName(TEXT("FGPUSkinPassthroughVertexFactory"), FNAME_Find);
-	static FName InstancedVfFname = FName(TEXT("FInstancedStaticMeshVertexFactory"), FNAME_Find);
-	static FName NiagaraRibbonVfFname = FName(TEXT("FNiagaraRibbonVertexFactory"), FNAME_Find);
-	static FName NiagaraSpriteVfFname = FName(TEXT("FNiagaraSpriteVertexFactory"), FNAME_Find);
-	static FName GeometryCacheVfFname = FName(TEXT("FGeometryCacheVertexVertexFactory"), FNAME_Find);
-	static FName LandscapeVfFname = FName(TEXT("FLandscapeVertexFactory"), FNAME_Find);
-	static FName LandscapeFixedGridVfFname = FName(TEXT("FLandscapeFixedGridVertexFactory"), FNAME_Find);
-	static FName LandscapeXYOffsetVfFname = FName(TEXT("FLandscapeXYOffsetVertexFactory"), FNAME_Find);
-	static FName HairCardVfFname = FName(TEXT("FHairCardsVertexFactory"), FNAME_Find);
-
-	return VertexFactoryType == FindVertexFactoryType(LocalVfFname)
-		|| VertexFactoryType == FindVertexFactoryType(LSkinnedVfFname)
-		|| VertexFactoryType == FindVertexFactoryType(InstancedVfFname)
-		|| VertexFactoryType == FindVertexFactoryType(NiagaraRibbonVfFname)
-		|| VertexFactoryType == FindVertexFactoryType(NiagaraSpriteVfFname)
-		|| VertexFactoryType == FindVertexFactoryType(GeometryCacheVfFname)
-		|| VertexFactoryType == FindVertexFactoryType(LandscapeVfFname)
-		|| VertexFactoryType == FindVertexFactoryType(LandscapeFixedGridVfFname)
-		|| VertexFactoryType == FindVertexFactoryType(LandscapeXYOffsetVfFname)
-		|| VertexFactoryType == FindVertexFactoryType(HairCardVfFname);
+	return VertexFactoryType->SupportsRayTracing();
 }
 
 class FMaterialCHS : public FMeshMaterialShader, public FUniformLightMapPolicyShaderParametersType
