Remove transient flag from Halton resource

diff --git a/Engine/Source/Runtime/Renderer/Private/Halton.cpp b/Engine/Source/Runtime/Renderer/Private/Halton.cpp
--- a/Engine/Source/Runtime/Renderer/Private/Halton.cpp
+++ b/Engine/Source/Runtime/Renderer/Private/Halton.cpp
@@ -162,7 +162,7 @@ void FHaltonPrimesResource::InitRHI()
 		};
 
 		uint32 ElementCount = FMath::DivideAndRoundUp(DimensionCount, 3u);
-		PrimesBuffer = RHICreateStructuredBuffer(sizeof(FPrimeData), ElementCount * sizeof(FPrimeData), BUF_Transient | BUF_FastVRAM | BUF_ShaderResource | BUF_UnorderedAccess, CreateInfo);
+		PrimesBuffer = RHICreateStructuredBuffer(sizeof(FPrimeData), ElementCount * sizeof(FPrimeData), BUF_ShaderResource | BUF_UnorderedAccess, CreateInfo);
 		uint32 Offset = 0;
 		void* BasePtr = RHILockBuffer(PrimesBuffer, Offset, ElementCount * sizeof(FPrimeData), RLM_WriteOnly);
 		FPlatformMemory::Memcpy(BasePtr, Primes.GetData(), DimensionCount * sizeof(int));
