Exclude HitProxy from being compiled for non-editor targets.

diff --git a/Engine/Source/Runtime/Renderer/Private/SceneHitProxyRendering.cpp b/Engine/Source/Runtime/Renderer/Private/SceneHitProxyRendering.cpp
--- a/Engine/Source/Runtime/Renderer/Private/SceneHitProxyRendering.cpp
+++ b/Engine/Source/Runtime/Renderer/Private/SceneHitProxyRendering.cpp
@@ -51,8 +51,8 @@ class FHitProxyVS : public FMeshMaterialShader
 public:
 	static bool ShouldCompilePermutation(const FMeshMaterialShaderPermutationParameters& Parameters)
 	{
-		// Only compile the hit proxy vertex shader on PC
-		return IsPCPlatform(Parameters.Platform)
+		// Only compile the hit proxy vertex shader on desktop editor platforms
+		return IsPCPlatform(Parameters.Platform) && EnumHasAllFlags(Parameters.Flags, EShaderPermutationFlags::HasEditorOnlyData)
 			// and only compile for the default material or materials that are masked.
 			&& (Parameters.MaterialParameters.bIsSpecialEngineMaterial ||
 				!Parameters.MaterialParameters.bWritesEveryPixel ||
@@ -108,8 +108,8 @@ public:
 
 	static bool ShouldCompilePermutation(const FMeshMaterialShaderPermutationParameters& Parameters)
 	{
-		// Only compile the hit proxy vertex shader on PC
-		return IsPCPlatform(Parameters.Platform) 
+		// Only compile the hit proxy vertex shader on desktop editor platforms
+		return IsPCPlatform(Parameters.Platform) && EnumHasAllFlags(Parameters.Flags, EShaderPermutationFlags::HasEditorOnlyData)
 			// and only compile for default materials or materials that are masked.
 			&& (Parameters.MaterialParameters.bIsSpecialEngineMaterial ||
 				!Parameters.MaterialParameters.bWritesEveryPixel ||
