Avoid redundantly touching FMeshBatch-es in GatherRayTracingWorldInstances(). Cache double-sided instance flag in FRayTracingMaskAndFlags.

diff --git a/Engine/Source/Runtime/Engine/Private/RayTracingInstance.cpp b/Engine/Source/Runtime/Engine/Private/RayTracingInstance.cpp
--- a/Engine/Source/Runtime/Engine/Private/RayTracingInstance.cpp
+++ b/Engine/Source/Runtime/Engine/Private/RayTracingInstance.cpp
@@ -17,6 +17,7 @@ void FRayTracingInstance::BuildInstanceMaskAndFlags()
 
 	Mask |= MaskAndFlags.Mask;
 	bForceOpaque = bForceOpaque || MaskAndFlags.bForceOpaque;
+	bDoubleSided = bDoubleSided || MaskAndFlags.bDoubleSided;
 }
 
 FRayTracingMaskAndFlags BuildRayTracingInstanceMaskAndFlags(TArrayView<const FMeshBatch> MeshBatches)
@@ -30,6 +31,7 @@ FRayTracingMaskAndFlags BuildRayTracingInstanceMaskAndFlags(TArrayView<const FMe
 	bool bAllSegmentsOpaque = true;
 	bool bAnySegmentsCastShadow = false;
 	bool bAllSegmentsCastShadow = true;
+	bool bDoubleSided = false;
 
 	for (int32 SegmentIndex = 0; SegmentIndex < MeshBatches.Num(); SegmentIndex++)
 	{
@@ -42,9 +44,11 @@ FRayTracingMaskAndFlags BuildRayTracingInstanceMaskAndFlags(TArrayView<const FMe
 		bAllSegmentsOpaque &= BlendMode == BLEND_Opaque;
 		bAnySegmentsCastShadow |= MeshBatch.CastRayTracedShadow && Material.CastsRayTracedShadows();
 		bAllSegmentsCastShadow &= MeshBatch.CastRayTracedShadow && Material.CastsRayTracedShadows();
+		bDoubleSided |= MeshBatch.bDisableBackfaceCulling;
 	}
 
 	Result.bForceOpaque = bAllSegmentsOpaque && bAllSegmentsCastShadow;
+	Result.bDoubleSided = bDoubleSided;
 	Result.Mask |= bAnySegmentsCastShadow ? RAY_TRACING_MASK_SHADOW : 0;
 
 	return Result;
diff --git a/Engine/Source/Runtime/Engine/Public/RayTracingInstance.h b/Engine/Source/Runtime/Engine/Public/RayTracingInstance.h
--- a/Engine/Source/Runtime/Engine/Public/RayTracingInstance.h
+++ b/Engine/Source/Runtime/Engine/Public/RayTracingInstance.h
@@ -14,6 +14,9 @@ struct FRayTracingMaskAndFlags
 
 	/** Whether the instance is forced opaque, i.e. anyhit shaders are disabled on this instance */
 	bool bForceOpaque = false;
+
+	/** Whether ray hits should be registered for front and back faces. */
+	bool bDoubleSided = false;
 };
 
 struct FRayTracingInstance
@@ -52,6 +55,9 @@ struct FRayTracingInstance
 	/** Whether the instance is forced opaque, i.e. anyhit shaders are disabled on this instance */
 	bool bForceOpaque = false;
 
+	/** Whether ray hits should be registered for front and back faces. */
+	bool bDoubleSided = false;
+
 	/** Instance mask that can be used to exclude the instance from specific effects (eg. ray traced shadows). */
 	uint8 Mask = 0xFF;
 
diff --git a/Engine/Source/Runtime/Renderer/Private/DeferredShadingRenderer.cpp b/Engine/Source/Runtime/Renderer/Private/DeferredShadingRenderer.cpp
--- a/Engine/Source/Runtime/Renderer/Private/DeferredShadingRenderer.cpp
+++ b/Engine/Source/Runtime/Renderer/Private/DeferredShadingRenderer.cpp
@@ -865,6 +865,7 @@ bool FDeferredShadingSceneRenderer::GatherRayTracingWorldInstances(FRHICommandLi
 					RayTracingInstance.UserData.Add((uint32)PrimitiveIndex);
 					RayTracingInstance.Mask = Instance.Mask;
 					RayTracingInstance.bForceOpaque = Instance.bForceOpaque;
+					RayTracingInstance.bDoubleSided = Instance.bDoubleSided;
 
 					// Thin geometries like hair don't have material, as they only support shadow at the moment.
 					if (!ensureMsgf(Instance.GetMaterials().Num() == Instance.Geometry->Initializer.Segments.Num() ||
@@ -903,13 +904,6 @@ bool FDeferredShadingSceneRenderer::GatherRayTracingWorldInstances(FRHICommandLi
 						}
 					}
 
-					TArrayView<const FMeshBatch> InstanceMaterials = Instance.GetMaterials();
-					for (int32 SegmentIndex = 0; SegmentIndex < InstanceMaterials.Num(); SegmentIndex++)
-					{
-						const FMeshBatch& MeshBatch = InstanceMaterials[SegmentIndex];
-						RayTracingInstance.bDoubleSided |= MeshBatch.bDisableBackfaceCulling;
-					}
-
 					for (int32 ViewIndex = 1; ViewIndex < Views.Num(); ViewIndex++)
 					{
 						Views[ViewIndex].RayTracingGeometryInstances.Add(RayTracingInstance);
@@ -944,6 +938,7 @@ bool FDeferredShadingSceneRenderer::GatherRayTracingWorldInstances(FRHICommandLi
 					}
 					else
 					{
+						TArrayView<const FMeshBatch> InstanceMaterials = Instance.GetMaterials();
 						for (int32 SegmentIndex = 0; SegmentIndex < InstanceMaterials.Num(); SegmentIndex++)
 						{
 							const FMeshBatch& MeshBatch = InstanceMaterials[SegmentIndex];
