Always draw base pass velocity when not using Gen4 TAA

diff --git a/Engine/Source/Runtime/Renderer/Private/SceneRendering.cpp b/Engine/Source/Runtime/Renderer/Private/SceneRendering.cpp
--- a/Engine/Source/Runtime/Renderer/Private/SceneRendering.cpp
+++ b/Engine/Source/Runtime/Renderer/Private/SceneRendering.cpp
@@ -1535,8 +1535,16 @@ void FViewInfo::SetupUniformBufferParameters(
 			TemporalJitterPixels.Y);
 	}
 
-	ViewUniformShaderParameters.ForceDrawAllVelocities = CVarBasePassForceOutputsVelocity.GetValueOnRenderThread();
-		
+	{
+		EMainTAAPassConfig MainTAAPass = ITemporalUpscaler::GetMainTAAPassConfig(*this);
+
+		// Gen4 TAA have the AA_DYNAMIC_ANTIGHOST heuristic that reject history based on whether the pixel is static or dynamic geometry
+		// through whether the velocity has been drawn by the base pass.
+		ViewUniformShaderParameters.ForceDrawAllVelocities = (
+			CVarBasePassForceOutputsVelocity.GetValueOnRenderThread() ||
+			MainTAAPass != EMainTAAPassConfig::Gen4);
+	}
+
 	uint32 FrameIndex = 0;
 	if (ViewState)
 	{
