Add hair strands view uniform buffer to simplify resources binding.

diff --git a/Engine/Source/Runtime/Renderer/Private/DeferredShadingRenderer.cpp b/Engine/Source/Runtime/Renderer/Private/DeferredShadingRenderer.cpp
--- a/Engine/Source/Runtime/Renderer/Private/DeferredShadingRenderer.cpp
+++ b/Engine/Source/Runtime/Renderer/Private/DeferredShadingRenderer.cpp
@@ -1810,7 +1810,6 @@ void FDeferredShadingSceneRenderer::Render(FRDGBuilder& GraphBuilder)
 	});
 
 	FHairStrandsBookmarkParameters& HairStrandsBookmarkParameters = *GraphBuilder.AllocObject<FHairStrandsBookmarkParameters>();
-
 	{
 		if (IsHairStrandsEnabled(EHairStrandsShaderType::All, Scene->GetShaderPlatform()))
 		{
@@ -1835,6 +1834,14 @@ void FDeferredShadingSceneRenderer::Render(FRDGBuilder& GraphBuilder)
 
 			RunHairStrandsBookmark(GraphBuilder, EHairStrandsBookmark::ProcessStrandsInterpolation, HairStrandsBookmarkParameters);
 		}
+		else
+		{
+			for (int32 ViewIndex = 0; ViewIndex < Views.Num(); ViewIndex++)
+			{
+				FViewInfo& View = Views[ViewIndex];
+				View.HairStrandsViewData.UniformBuffer = HairStrands::CreateDefaultHairStrandsViewUniformBuffer(GraphBuilder, View);
+			}
+		}
 	}
 
 	if (bNaniteEnabled)
@@ -1848,8 +1855,7 @@ void FDeferredShadingSceneRenderer::Render(FRDGBuilder& GraphBuilder)
 	const bool bShouldRenderVelocities = ShouldRenderVelocities();
 	const bool bBasePassCanOutputVelocity = FVelocityRendering::BasePassCanOutputVelocity(FeatureLevel);
 	const bool bUseSelectiveBasePassOutputs = IsUsingSelectiveBasePassOutputs(ShaderPlatform);
-	const bool bIsViewCompatible = Views.Num() > 0;
-	const bool bHairEnable = HairStrandsBookmarkParameters.bHasElements && bIsViewCompatible && IsHairStrandsEnabled(EHairStrandsShaderType::Strands, Views[0].GetShaderPlatform());
+	const bool bHairEnable = HairStrandsBookmarkParameters.bHasElements && Views.Num() > 0 && IsHairStrandsEnabled(EHairStrandsShaderType::Strands, Views[0].GetShaderPlatform());
 
 	{
 		GraphBuilder.SetCommandListStat(GET_STATID(STAT_CLM_PrePass));
@@ -2036,14 +2042,13 @@ void FDeferredShadingSceneRenderer::Render(FRDGBuilder& GraphBuilder)
 	CSV_CUSTOM_STAT(LightCount, ShadowOn, float(SortedLightSet.SortedLights.Num()) - float(SortedLightSet.AttenuationLightStart), ECsvCustomStatOp::Set);
 
 	// Local helper function to perform virtual shadow map allocation, which can occur early, or late.
-	FHairStrandsRenderingData* HairDatas = nullptr;
 	const auto AllocateVirtualShadowMaps = [&](bool bPostBasePass)
 	{
 		if (VirtualShadowMapArray.IsEnabled())
 		{
 			ensureMsgf(AreLightsInLightGrid(), TEXT("Virtual shadow map setup requires local lights to be injected into the light grid (this may be caused by 'r.LightCulling.Quality=0')."));
 			// ensure(ShadowMapSetupDone)
-			VirtualShadowMapArray.BuildPageAllocations(GraphBuilder, SceneTextures, Views, SortedLightSet, VisibleLightInfos, NaniteRasterResults, bPostBasePass, Scene->VirtualShadowMapArrayCacheManager, HairDatas);
+			VirtualShadowMapArray.BuildPageAllocations(GraphBuilder, SceneTextures, Views, SortedLightSet, VisibleLightInfos, NaniteRasterResults, bPostBasePass, Scene->VirtualShadowMapArrayCacheManager);
 		}
 	};
 
@@ -2155,8 +2160,8 @@ void FDeferredShadingSceneRenderer::Render(FRDGBuilder& GraphBuilder)
 		bAsyncComputeVolumetricCloud = RenderVolumetricCloud(GraphBuilder, SceneTextures, bSkipVolumetricRenderTarget, bSkipPerPixelTracing, HalfResolutionDepthCheckerboardMinMaxTexture, true, InstanceCullingManager);
 	}
 	
+	FHairStrandsRenderingData* HairDatas = nullptr;
 	FHairStrandsRenderingData* HairDatasStorage = GraphBuilder.AllocObject<FHairStrandsRenderingData>();
-
 	FRDGTextureRef ForwardScreenSpaceShadowMaskTexture = nullptr;
 	FRDGTextureRef ForwardScreenSpaceShadowMaskHairTexture = nullptr;
 	if (IsForwardShadingEnabled(ShaderPlatform))
diff --git a/Engine/Source/Runtime/Renderer/Private/HairStrands/HairStrandsData.h b/Engine/Source/Runtime/Renderer/Private/HairStrands/HairStrandsData.h
new file mode 100644
--- /dev/null
+++ b/Engine/Source/Runtime/Renderer/Private/HairStrands/HairStrandsData.h
@@ -0,0 +1,31 @@
+// Copyright Epic Games, Inc. All Rights Reserved.
+
+/*=============================================================================
+	HairRendering.h: Hair rendering implementation.
+=============================================================================*/
+
+#pragma once
+
+#include "CoreMinimal.h"
+#include "RendererInterface.h"
+#include "Shader.h"
+
+BEGIN_GLOBAL_SHADER_PARAMETER_STRUCT(FHairStrandsViewUniformParameters, )	
+	SHADER_PARAMETER_RDG_TEXTURE(Texture2D<uint4>, HairCategorizationTexture)
+	SHADER_PARAMETER_RDG_TEXTURE(Texture2D<float>, HairOnlyDepthTexture)
+END_GLOBAL_SHADER_PARAMETER_STRUCT()
+
+struct FHairStrandsViewData
+{
+	//FRDGTextureRef CategorizationTexture = nullptr;
+	//FRDGTextureRef HairOnlyDepthTexture = nullptr;
+	TRDGUniformBufferRef<FHairStrandsViewUniformParameters> UniformBuffer;
+	bool bIsValid = false;
+};
+
+namespace HairStrands
+{
+	TRDGUniformBufferRef<FHairStrandsViewUniformParameters> CreateDefaultHairStrandsViewUniformBuffer(FRDGBuilder& GraphBuilder, FViewInfo& View);
+	TRDGUniformBufferRef<FHairStrandsViewUniformParameters> BindHairStrandsViewUniformParameters(const FViewInfo& View);
+	bool HasViewHairStrandsData(const FViewInfo& View);
+}
\ No newline at end of file
diff --git a/Engine/Source/Runtime/Renderer/Private/HairStrands/HairStrandsRendering.cpp b/Engine/Source/Runtime/Renderer/Private/HairStrands/HairStrandsRendering.cpp
--- a/Engine/Source/Runtime/Renderer/Private/HairStrands/HairStrandsRendering.cpp
+++ b/Engine/Source/Runtime/Renderer/Private/HairStrands/HairStrandsRendering.cpp
@@ -4,6 +4,19 @@
 #include "SceneRendering.h"
 #include "ScenePrivate.h"
 
+static TRDGUniformBufferRef<FHairStrandsViewUniformParameters> InternalCreateHairStrandsViewUniformBuffer(
+	FRDGBuilder& GraphBuilder, 
+	FRDGTextureRef HairCategorizationTexture, 
+	FRDGTextureRef HairDepthOnlyTexture)
+{
+	FHairStrandsViewUniformParameters* Parameters = GraphBuilder.AllocParameters<FHairStrandsViewUniformParameters>();
+	Parameters->HairCategorizationTexture = HairCategorizationTexture ? HairCategorizationTexture : GSystemTextures.GetBlackDummy(GraphBuilder);
+	Parameters->HairOnlyDepthTexture = HairDepthOnlyTexture ? HairDepthOnlyTexture : GSystemTextures.GetDepthDummy(GraphBuilder);
+	return GraphBuilder.CreateUniformBuffer(Parameters);
+}
+
+IMPLEMENT_GLOBAL_SHADER_PARAMETER_STRUCT(FHairStrandsViewUniformParameters, "HairStrands");
+
 void AddServiceLocalQueuePass(FRDGBuilder& GraphBuilder);
 
 void RenderHairPrePass(
@@ -60,5 +73,37 @@ void RenderHairBasePass(
 			SceneTextures.Velocity,
 			InstanceCullingManager,
 			OutHairDatas.MacroGroupsPerViews);
+
+		// Create RDG uniform buffer for each view
+		for (int32 ViewIndex = 0; ViewIndex < Views.Num(); ViewIndex++)
+		{
+			const FViewInfo& View = Views[ViewIndex];
+			if (View.Family)
+			{
+				FHairStrandsVisibilityData VisibilityData = OutHairDatas.HairVisibilityViews.HairDatas[ViewIndex];
+				Views[ViewIndex].HairStrandsViewData.UniformBuffer = InternalCreateHairStrandsViewUniformBuffer(GraphBuilder, VisibilityData.CategorizationTexture, VisibilityData.HairOnlyDepthTexture);
+				Views[ViewIndex].HairStrandsViewData.bIsValid = true;
+			}
+		}
 	}
 }
+
+namespace HairStrands
+{
+
+TRDGUniformBufferRef<FHairStrandsViewUniformParameters> CreateDefaultHairStrandsViewUniformBuffer(FRDGBuilder& GraphBuilder, FViewInfo& View)
+{
+	return InternalCreateHairStrandsViewUniformBuffer(GraphBuilder, nullptr, nullptr);
+}
+
+TRDGUniformBufferRef<FHairStrandsViewUniformParameters> BindHairStrandsViewUniformParameters(const FViewInfo& View)
+{
+	return View.HairStrandsViewData.UniformBuffer;
+}
+
+bool HasViewHairStrandsData(const FViewInfo& View)
+{
+	return View.HairStrandsViewData.bIsValid;
+}
+
+}
\ No newline at end of file
