Simplify hair resources binding for virtual shadow map.

diff --git a/Engine/Shaders/Private/VirtualShadowMaps/PageManagement.usf b/Engine/Shaders/Private/VirtualShadowMaps/PageManagement.usf
--- a/Engine/Shaders/Private/VirtualShadowMaps/PageManagement.usf
+++ b/Engine/Shaders/Private/VirtualShadowMaps/PageManagement.usf
@@ -23,6 +23,10 @@
 RWStructuredBuffer<uint> OutStatsBuffer;
 #endif // VSM_GENERATE_STATS
 
+// Type of input data consume by the page allocation (i.e., data read from the source buffer: Gbuffer, HairStrands data, ...)
+#define INPUT_TYPE_GBUFFER 0
+#define INPUT_TYPE_HAIRSTRANDS 1
+
 // Flags generated by per-pixel pass to determine which pages are required to provide shadow for the visible geometry
 RWStructuredBuffer<uint> OutPageRequestFlags;
 
@@ -32,10 +36,7 @@ uint NumDirectionalLightSmInds;
 uint bPostBasePass;
 float LodFootprintScale;
 float PageDilationBorderSize;
-
-// HAIR_TODO: move all data into a HairStrands UB
-Texture2D<uint4> HairCategorizationTexture;
-uint bUseHairData;
+uint InputType;
 
 // Convenience to group up a bunch of data computed for a screen pixel
 struct FScreenPixelData
@@ -176,19 +177,19 @@ void GeneratePageFlagsFromPixels(uint3 GroupId : SV_GroupID, uint GroupIndex : S
 	{
 		return;
 	}
-	
+
 	float3 WorldPosition = 0;
 	float  SceneDepth    = 0;
 	bool   bIsValid		 = false;
-	if (bUseHairData)
+	if (InputType == INPUT_TYPE_HAIRSTRANDS)
 	{
-		FCategorizationData CatData = DecodeCategorizationData(HairCategorizationTexture.Load(uint3(PixelPos, 0)));
+		FCategorizationData CatData = DecodeCategorizationData(HairStrands.HairCategorizationTexture.Load(uint3(PixelPos, 0)));
 		FScreenPixelData HairPixel = GetScreenPixelData(PixelPos, CatData.ClosestDepth);
 		WorldPosition = HairPixel.WorldPosition;
 		SceneDepth = HairPixel.SceneDepth;
 		bIsValid = CatData.PixelCoverage > 0;
 	}
-	else
+	else// (InputType == INPUT_TYPE_GBUFFER)
 	{
 		FScreenPixelData Pixel = GetScreenPixelData(PixelPos, GetDeviceZFromPixel(PixelPos));
 		WorldPosition = Pixel.WorldPosition;
diff --git a/Engine/Shaders/Private/VirtualShadowMaps/VirtualShadowMapProjection.usf b/Engine/Shaders/Private/VirtualShadowMaps/VirtualShadowMapProjection.usf
--- a/Engine/Shaders/Private/VirtualShadowMaps/VirtualShadowMapProjection.usf
+++ b/Engine/Shaders/Private/VirtualShadowMaps/VirtualShadowMapProjection.usf
@@ -253,9 +253,10 @@ void VirtualShadowMapProjectionSpotPS(
 	OutputTarget = EncodeOutput(TranslatedWorldPosition, Light, LightType, RayMissFactor, HitDistance, SampleCount, DebugOutput);
 }
 
-// HAIR_TODO: move all data into a HairStrands UB
-uint bUseHairData;
-Texture2D<uint4> HairCategorizationTexture;
+// Type of input data consume by the page allocation (i.e., data read from the source buffer: Gbuffer, HairStrands data, ...)
+#define INPUT_TYPE_GBUFFER 0
+#define INPUT_TYPE_HAIRSTRANDS 1
+uint InputType;
 
 void VirtualShadowMapProjectionDirectionalPS(
 	in float4 SvPosition : SV_Position,
@@ -265,9 +266,9 @@ void VirtualShadowMapProjectionDirectionalPS(
 	// TODO: Unify this with the logic in PageManagement.usf - we need to get the same values in both passes
 	const float2 ScreenUV = SvPositionToBufferUV(SvPosition);
 	float DeviceZ = LookupDeviceZ(ScreenUV);
-	if (bUseHairData > 0)
+	if (InputType == INPUT_TYPE_HAIRSTRANDS)
 	{
-		FCategorizationData CatData = DecodeCategorizationData(HairCategorizationTexture.Load(uint3(SvPosition.xy, 0)));
+		FCategorizationData CatData = DecodeCategorizationData(HairStrands.HairCategorizationTexture.Load(uint3(SvPosition.xy, 0)));
 		DeviceZ = CatData.ClosestDepth;
 	}
 	const float SceneDepth = ConvertFromDeviceZ(DeviceZ);
@@ -290,7 +291,7 @@ void VirtualShadowMapProjectionDirectionalPS(
 		const float ContactShadowLengthWorld = ContactShadowLength * View.ClipToView[1][1] * SceneDepth;
 
 		float3 EstimatedGeoWorldNormal = GetEstimatedGeoWorldNormal(TranslatedWorldPosition, ScreenUV);
-		if (bUseHairData > 0)
+		if (InputType == INPUT_TYPE_HAIRSTRANDS)
 		{			
 			EstimatedGeoWorldNormal = L;
 		}
diff --git a/Engine/Source/Runtime/Renderer/Private/SceneRendering.cpp b/Engine/Source/Runtime/Renderer/Private/SceneRendering.cpp
--- a/Engine/Source/Runtime/Renderer/Private/SceneRendering.cpp
+++ b/Engine/Source/Runtime/Renderer/Private/SceneRendering.cpp
@@ -954,7 +954,8 @@ void FViewInfo::Init()
 	NumVisibleDynamicEditorPrimitives = 0;
 
 	StrataSceneData = nullptr;
-	
+	HairStrandsViewData = FHairStrandsViewData();
+
 	GPUSceneViewId = INDEX_NONE;
 }
 
diff --git a/Engine/Source/Runtime/Renderer/Private/SceneRendering.h b/Engine/Source/Runtime/Renderer/Private/SceneRendering.h
--- a/Engine/Source/Runtime/Renderer/Private/SceneRendering.h
+++ b/Engine/Source/Runtime/Renderer/Private/SceneRendering.h
@@ -36,6 +36,7 @@
 #include "VirtualShadowMaps/VirtualShadowMapConfig.h"
 #include "VirtualShadowMaps/VirtualShadowMapArray.h"
 #include "Lumen/LumenTranslucencyVolumeLighting.h"
+#include "HairStrands/HairStrandsData.h"
 #include "GPUScene.h"
 
 // Forward declarations.
@@ -1234,6 +1235,8 @@ public:
 
 	FStrataSceneData* StrataSceneData;
 
+	FHairStrandsViewData HairStrandsViewData;
+
 	//Spotlight shadow info for mobile.
 	FMobileMovableSpotLightsShadowInfo MobileMovableSpotLightsShadowInfo;
 
diff --git a/Engine/Source/Runtime/Renderer/Private/ShadowRendering.cpp b/Engine/Source/Runtime/Renderer/Private/ShadowRendering.cpp
--- a/Engine/Source/Runtime/Renderer/Private/ShadowRendering.cpp
+++ b/Engine/Source/Runtime/Renderer/Private/ShadowRendering.cpp
@@ -2021,7 +2021,7 @@ void FDeferredShadingSceneRenderer::RenderShadowProjections(
 								VirtualShadowMapArray,
 								ScissorRect,
 								ScreenShadowMaskTexture,
-								nullptr,
+								EVirtualShadowMapProjectionInputType::GBuffer,
 								bProjectingForForwardShading);
 						}
 						else
@@ -2034,13 +2034,13 @@ void FDeferredShadingSceneRenderer::RenderShadowProjections(
 								VirtualShadowMapArray,
 								ScissorRect,
 								ScreenShadowMaskTexture,
-								nullptr,
+								EVirtualShadowMapProjectionInputType::GBuffer,
 								bProjectingForForwardShading);
 						}
 					}
 
 					// Sub-pixel shadow (no filtering for hair)
-					if (ScreenShadowMaskSubPixelTexture && HairCatgorizationTexture)
+					if (HairStrands::HasViewHairStrandsData(View) && ScreenShadowMaskSubPixelTexture)
 					{
 						if (VisibleLightInfo.VirtualShadowMapClipmaps.Num() > 0)
 						{
@@ -2051,7 +2051,7 @@ void FDeferredShadingSceneRenderer::RenderShadowProjections(
 								VirtualShadowMapArray,
 								ScissorRect,
 								ScreenShadowMaskSubPixelTexture,
-								HairCatgorizationTexture,
+								EVirtualShadowMapProjectionInputType::HairStrands,
 								bProjectingForForwardShading);
 						}
 						else
@@ -2064,7 +2064,7 @@ void FDeferredShadingSceneRenderer::RenderShadowProjections(
 								VirtualShadowMapArray,
 								ScissorRect,
 								ScreenShadowMaskSubPixelTexture,
-								HairCatgorizationTexture,
+								EVirtualShadowMapProjectionInputType::HairStrands,
 								bProjectingForForwardShading);
 						}
 					}
diff --git a/Engine/Source/Runtime/Renderer/Private/VirtualShadowMaps/VirtualShadowMapArray.cpp b/Engine/Source/Runtime/Renderer/Private/VirtualShadowMaps/VirtualShadowMapArray.cpp
--- a/Engine/Source/Runtime/Renderer/Private/VirtualShadowMaps/VirtualShadowMapArray.cpp
+++ b/Engine/Source/Runtime/Renderer/Private/VirtualShadowMaps/VirtualShadowMapArray.cpp
@@ -11,7 +11,7 @@
 #include "VirtualShadowMapCacheManager.h"
 #include "VirtualShadowMapClipmap.h"
 #include "ComponentRecreateRenderStateContext.h"
-#include "HairStrands/HairStrandsRendering.h"
+#include "HairStrands/HairStrandsData.h"
 
 IMPLEMENT_STATIC_UNIFORM_BUFFER_SLOT(VirtualShadowMapUbSlot);
 
@@ -314,14 +314,14 @@ class FGeneratePageFlagsFromPixelsCS : public FVirtualPageManagementShader
 	BEGIN_SHADER_PARAMETER_STRUCT(FParameters, )
 		SHADER_PARAMETER_RDG_UNIFORM_BUFFER(FVirtualShadowMapCommonParameters, VirtualSmCommon)
 		SHADER_PARAMETER_RDG_UNIFORM_BUFFER(FSceneTextureUniformParameters, SceneTexturesStruct)
+		SHADER_PARAMETER_RDG_UNIFORM_BUFFER(FHairStrandsViewUniformParameters, HairStrands)
 		SHADER_PARAMETER_STRUCT_REF(FViewUniformShaderParameters, View)
 		SHADER_PARAMETER_STRUCT_REF(FForwardLightData, ForwardLightData)
 		SHADER_PARAMETER_RDG_TEXTURE(Texture2D<uint2>, VisBuffer64)
-		SHADER_PARAMETER_RDG_TEXTURE(Texture2D<uint4>, HairCategorizationTexture)
 		SHADER_PARAMETER_RDG_BUFFER_UAV(RWStructuredBuffer<uint>, OutPageRequestFlags)
 		SHADER_PARAMETER_RDG_BUFFER_SRV(StructuredBuffer< FVirtualShadowMapProjectionShaderData >, VirtualShadowMapProjectionData)
 		SHADER_PARAMETER_RDG_BUFFER_SRV(StructuredBuffer< uint >, VirtualShadowMapIdRemap)
-		SHADER_PARAMETER(uint32, bUseHairData)
+		SHADER_PARAMETER(uint32, InputType)
 		SHADER_PARAMETER(uint32, NumDirectionalLightSmInds)
 		SHADER_PARAMETER(uint32, bPostBasePass)
 		SHADER_PARAMETER(float, LodFootprintScale)
@@ -686,8 +686,7 @@ void FVirtualShadowMapArray::BuildPageAllocations(
 	const TArray<FVisibleLightInfo, SceneRenderingAllocator>& VisibleLightInfos,
 	const TArray<Nanite::FRasterResults, TInlineAllocator<2>>& NaniteRasterResults,
 	bool bPostBasePass,
-	FVirtualShadowMapArrayCacheManager* VirtualShadowMapArrayCacheManager,
-	FHairStrandsRenderingData* HairDatas)
+	FVirtualShadowMapArrayCacheManager* VirtualShadowMapArrayCacheManager)
 {
 	check(IsEnabled());
 	RDG_EVENT_SCOPE(GraphBuilder, "FVirtualShadowMapArray::GeneratePageFlagsFromLightGrid");
@@ -850,7 +849,7 @@ void FVirtualShadowMapArray::BuildPageAllocations(
 				// Mark pages based on projected depth buffer pixels
 				if (CVarMarkPixelPages.GetValueOnRenderThread() != 0)
 				{
-					auto GeneratePageFlags = [&](FRDGTextureRef HairCategorizationTexture)
+					auto GeneratePageFlags = [&](bool bHairPass)
 					{
 						FGeneratePageFlagsFromPixelsCS::FPermutationDomain PermutationVector;
 						PermutationVector.Set<FGeneratePageFlagsFromPixelsCS::FNaniteDepthBufferDim>(NaniteVisBuffer64 != nullptr && !bPostBasePass);
@@ -862,8 +861,8 @@ void FVirtualShadowMapArray::BuildPageAllocations(
 						PassParameters->bPostBasePass = bPostBasePass;
 
 						PassParameters->VisBuffer64 = VisBuffer64;
-						PassParameters->bUseHairData = HairCategorizationTexture != nullptr ? 1 : 0;
-						PassParameters->HairCategorizationTexture = HairCategorizationTexture ? HairCategorizationTexture : GSystemTextures.GetBlackAlphaOneDummy(GraphBuilder);
+						PassParameters->InputType = bHairPass ? 1 : 0;
+						PassParameters->HairStrands = HairStrands::BindHairStrandsViewUniformParameters(View);
 						PassParameters->View = View.ViewUniformBuffer;
 						PassParameters->OutPageRequestFlags = PageRequestFlagsUAV;
 						PassParameters->ForwardLightData = View.ForwardLightingResources->ForwardLightDataUniformBuffer;
@@ -887,11 +886,10 @@ void FVirtualShadowMapArray::BuildPageAllocations(
 						);
 					};
 
-					GeneratePageFlags(nullptr);
-					FRDGTextureRef HairCategorizationTexture = HairDatas && ViewIndex < HairDatas->HairVisibilityViews.HairDatas.Num() ? HairDatas->HairVisibilityViews.HairDatas[ViewIndex].CategorizationTexture : nullptr;
-					if (HairCategorizationTexture)
+					GeneratePageFlags(false);
+					if (HairStrands::HasViewHairStrandsData(View))
 					{
-						GeneratePageFlags(HairCategorizationTexture);
+						GeneratePageFlags(true);
 					}
 				}
 				// Mark coarse pages
diff --git a/Engine/Source/Runtime/Renderer/Private/VirtualShadowMaps/VirtualShadowMapArray.h b/Engine/Source/Runtime/Renderer/Private/VirtualShadowMaps/VirtualShadowMapArray.h
--- a/Engine/Source/Runtime/Renderer/Private/VirtualShadowMaps/VirtualShadowMapArray.h
+++ b/Engine/Source/Runtime/Renderer/Private/VirtualShadowMaps/VirtualShadowMapArray.h
@@ -192,8 +192,7 @@ public:
 		const TArray<FVisibleLightInfo, SceneRenderingAllocator> &VisibleLightInfos, 
 		const TArray<Nanite::FRasterResults, TInlineAllocator<2>> &NaniteRasterResults, 
 		bool bPostBasePass, 
-		FVirtualShadowMapArrayCacheManager *VirtualShadowMapArrayCacheManager,
-		struct FHairStrandsRenderingData* HairDatas);
+		FVirtualShadowMapArrayCacheManager *VirtualShadowMapArrayCacheManager);
 
 	bool IsAllocated() const
 	{
diff --git a/Engine/Source/Runtime/Renderer/Private/VirtualShadowMaps/VirtualShadowMapProjection.cpp b/Engine/Source/Runtime/Renderer/Private/VirtualShadowMaps/VirtualShadowMapProjection.cpp
--- a/Engine/Source/Runtime/Renderer/Private/VirtualShadowMaps/VirtualShadowMapProjection.cpp
+++ b/Engine/Source/Runtime/Renderer/Private/VirtualShadowMaps/VirtualShadowMapProjection.cpp
@@ -23,6 +23,7 @@
 #include "ShadowRendering.h"
 #include "SceneRendering.h"
 #include "VirtualShadowMapClipmap.h"
+#include "HairStrands/HairStrandsData.h"
 
 static TAutoConsoleVariable<float> CVarContactShadowLength(
 	TEXT( "r.Shadow.Virtual.ContactShadowLength" ),
@@ -110,8 +111,8 @@ BEGIN_SHADER_PARAMETER_STRUCT(FProjectionParameters, )
 	SHADER_PARAMETER_STRUCT_INCLUDE(FVirtualShadowMapSamplingParameters, ProjectionParameters)
 	SHADER_PARAMETER_STRUCT(FLightShaderParameters, Light)
 	SHADER_PARAMETER_RDG_UNIFORM_BUFFER(FSceneTextureUniformParameters, SceneTexturesStruct)
+	SHADER_PARAMETER_RDG_UNIFORM_BUFFER(FHairStrandsViewUniformParameters, HairStrands)
 	SHADER_PARAMETER_STRUCT_REF(FViewUniformShaderParameters, View)
-	SHADER_PARAMETER_RDG_TEXTURE(Texture2D<uint4>, HairCategorizationTexture)
 	SHADER_PARAMETER(int32, VirtualShadowMapId)
 	SHADER_PARAMETER(float, ContactShadowLength)
 	SHADER_PARAMETER(float, NormalOffsetWorld)
@@ -120,7 +121,7 @@ BEGIN_SHADER_PARAMETER_STRUCT(FProjectionParameters, )
 	SHADER_PARAMETER(int32, SMRTSamplesPerRay)
 	SHADER_PARAMETER(float, SMRTRayLengthScale)
 	SHADER_PARAMETER(float, SMRTCotMaxRayAngleFromLight)
-	SHADER_PARAMETER(uint32, bUseHairData)
+	SHADER_PARAMETER(uint32, InputType)
 	RENDER_TARGET_BINDING_SLOTS()
 END_SHADER_PARAMETER_STRUCT()
 
@@ -205,7 +206,7 @@ static bool AddPass_RenderVirtualShadowMapProjection(
 	FVirtualShadowMapArray& VirtualShadowMapArray,
 	int VirtualShadowMapId,
 	const FIntRect ScissorRect,
-	FRDGTextureRef HairCategorization,
+	EVirtualShadowMapProjectionInputType InputType,
 	EVirtualShadowMapProjectionOutputType OutputType,
 	const FRenderTargetBinding& RenderTargetBinding,
 	FRHIBlendState* BlendState)
@@ -226,9 +227,8 @@ static bool AddPass_RenderVirtualShadowMapProjection(
 	PassParameters->DebugOutputType = CVarVirtualShadowMapDebugProjection.GetValueOnRenderThread();
 	PassParameters->ContactShadowLength = CVarContactShadowLength.GetValueOnRenderThread();
 	PassParameters->NormalOffsetWorld = CVarNormalOffsetWorld.GetValueOnRenderThread();
-	PassParameters->bUseHairData = HairCategorization != nullptr ? 1 : 0;
-	PassParameters->HairCategorizationTexture = HairCategorization ? HairCategorization : GSystemTextures.GetBlackDummy(GraphBuilder);
-
+	PassParameters->InputType = uint32(InputType);
+	PassParameters->HairStrands = HairStrands::BindHairStrandsViewUniformParameters(View);
 	PassParameters->RenderTargets[0] = RenderTargetBinding;
 
 	if (LightProxy->GetLightType() == LightType_Directional)
@@ -303,7 +303,7 @@ static bool RenderVirtualShadowMapProjectionCommon(
 	FVirtualShadowMapArray& VirtualShadowMapArray,
 	int32 VirtualShadowMapId,
 	const FIntRect ScissorRect,
-	FRDGTextureRef HairCategorization,
+	EVirtualShadowMapProjectionInputType InputType,
 	EVirtualShadowMapProjectionOutputType OutputType,
 	const FRenderTargetBinding& RenderTargetBinding,
 	FRHIBlendState* BlendState)
@@ -316,7 +316,7 @@ static bool RenderVirtualShadowMapProjectionCommon(
 		VirtualShadowMapArray,
 		VirtualShadowMapId,
 		ScissorRect,
-		HairCategorization,
+		InputType,
 		OutputType,
 		RenderTargetBinding,
 		BlendState);
@@ -335,7 +335,7 @@ static bool RenderVirtualShadowMapProjectionCommon(
 			VirtualShadowMapArray,
 			VirtualShadowMapId,
 			ScissorRect,
-			HairCategorization,
+			InputType,
 			EVirtualShadowMapProjectionOutputType::Debug,
 			DebugRenderTargetBinding,
 			nullptr);
@@ -363,7 +363,7 @@ void RenderVirtualShadowMapProjectionForDenoising(
 		VirtualShadowMapArray,
 		ShadowInfo->VirtualShadowMaps[0]->ID,
 		ScissorRect,
-		nullptr,
+		EVirtualShadowMapProjectionInputType::GBuffer,
 		EVirtualShadowMapProjectionOutputType::Denoiser,
 		RenderTargetBinding,
 		nullptr);
@@ -386,7 +386,7 @@ void RenderVirtualShadowMapProjectionForDenoising(
 		VirtualShadowMapArray,
 		Clipmap->GetVirtualShadowMap()->ID,
 		ScissorRect,
-		nullptr,
+		EVirtualShadowMapProjectionInputType::GBuffer,
 		EVirtualShadowMapProjectionOutputType::Denoiser,
 		RenderTargetBinding,
 		nullptr);
@@ -399,7 +399,7 @@ void RenderVirtualShadowMapProjection(
 	FVirtualShadowMapArray& VirtualShadowMapArray,
 	const FIntRect ScissorRect,
 	FRDGTextureRef ScreenShadowMaskTexture,
-	FRDGTextureRef HairCategorization,
+	EVirtualShadowMapProjectionInputType InputType,
 	bool bProjectingForForwardShading)
 {
 	FRenderTargetBinding RenderTargetBinding(ScreenShadowMaskTexture, ERenderTargetLoadAction::ELoad);
@@ -411,7 +411,7 @@ void RenderVirtualShadowMapProjection(
 		VirtualShadowMapArray,
 		ShadowInfo->VirtualShadowMaps[0]->ID,
 		ScissorRect,
-		HairCategorization,
+		InputType,
 		EVirtualShadowMapProjectionOutputType::ScreenShadowMask,
 		RenderTargetBinding,
 		BlendState);
@@ -424,7 +424,7 @@ void RenderVirtualShadowMapProjection(
 	FVirtualShadowMapArray& VirtualShadowMapArray,
 	const FIntRect ScissorRect,
 	FRDGTextureRef ScreenShadowMaskTexture,
-	FRDGTextureRef HairCategorization,
+	EVirtualShadowMapProjectionInputType InputType,
 	bool bProjectingForForwardShading)
 {
 	FRenderTargetBinding RenderTargetBinding(ScreenShadowMaskTexture, ERenderTargetLoadAction::ELoad);	
@@ -437,7 +437,7 @@ void RenderVirtualShadowMapProjection(
 		VirtualShadowMapArray,
 		Clipmap->GetVirtualShadowMap()->ID,
 		ScissorRect,
-		HairCategorization,
+		InputType,
 		EVirtualShadowMapProjectionOutputType::ScreenShadowMask,
 		RenderTargetBinding,
 		BlendState);
diff --git a/Engine/Source/Runtime/Renderer/Private/VirtualShadowMaps/VirtualShadowMapProjection.h b/Engine/Source/Runtime/Renderer/Private/VirtualShadowMaps/VirtualShadowMapProjection.h
--- a/Engine/Source/Runtime/Renderer/Private/VirtualShadowMaps/VirtualShadowMapProjection.h
+++ b/Engine/Source/Runtime/Renderer/Private/VirtualShadowMaps/VirtualShadowMapProjection.h
@@ -13,6 +13,12 @@
 
 class FVirtualShadowMapClipmap;
 
+enum class EVirtualShadowMapProjectionInputType
+{
+	GBuffer = 0,
+	HairStrands = 1
+};
+
 void RenderVirtualShadowMapProjectionForDenoising(
 	const TSharedPtr<FVirtualShadowMapClipmap>& Clipmap,
 	FRDGBuilder& GraphBuilder,
@@ -38,7 +44,7 @@ void RenderVirtualShadowMapProjection(
 	FVirtualShadowMapArray& VirtualShadowMapArray,
 	const FIntRect ScissorRect,
 	FRDGTextureRef ScreenShadowMaskTexture,
-	FRDGTextureRef HairCategorization,
+	EVirtualShadowMapProjectionInputType InputType,
 	bool bProjectingForForwardShading);
 
 void RenderVirtualShadowMapProjection(
@@ -48,7 +54,7 @@ void RenderVirtualShadowMapProjection(
 	FVirtualShadowMapArray& VirtualShadowMapArray,
 	const FIntRect ScissorRect,
 	FRDGTextureRef ScreenShadowMaskTexture,
-	FRDGTextureRef HairCategorization,
+	EVirtualShadowMapProjectionInputType InputType,
 	bool bProjectingForForwardShading);
 
 void RenderVirtualShadowMapProjection(
