Fix CIS unity build issue.

diff --git a/Engine/Source/Runtime/Renderer/Private/DeferredShadingRenderer.cpp b/Engine/Source/Runtime/Renderer/Private/DeferredShadingRenderer.cpp
--- a/Engine/Source/Runtime/Renderer/Private/DeferredShadingRenderer.cpp
+++ b/Engine/Source/Runtime/Renderer/Private/DeferredShadingRenderer.cpp
@@ -48,6 +48,7 @@
 #include "GpuDebugRendering.h"
 #include "GPUSortManager.h"
 #include "HairStrands/HairStrandsRendering.h"
+#include "HairStrands/HairStrandsData.h"
 #include "PhysicsField/PhysicsFieldComponent.h"
 #include "GPUSortManager.h"
 #include "Rendering/NaniteResources.h"
diff --git a/Engine/Source/Runtime/Renderer/Private/HairStrands/HairStrandsRendering.cpp b/Engine/Source/Runtime/Renderer/Private/HairStrands/HairStrandsRendering.cpp
--- a/Engine/Source/Runtime/Renderer/Private/HairStrands/HairStrandsRendering.cpp
+++ b/Engine/Source/Runtime/Renderer/Private/HairStrands/HairStrandsRendering.cpp
@@ -1,6 +1,7 @@
 // Copyright Epic Games, Inc. All Rights Reserved.
 
 #include "HairStrandsRendering.h"
+#include "HairStrandsData.h"
 #include "SceneRendering.h"
 #include "ScenePrivate.h"
 
