Implements r.TemporalAA.AntiInterference with age based bilateral history decimation to kill ghosting on VFX

diff --git a/Engine/Shaders/Private/TemporalAA/TAACommon.ush b/Engine/Shaders/Private/TemporalAA/TAACommon.ush
--- a/Engine/Shaders/Private/TemporalAA/TAACommon.ush
+++ b/Engine/Shaders/Private/TemporalAA/TAACommon.ush
@@ -33,9 +33,6 @@
 
 #define MAX_FALLBACK_SAMPLE_COUNT 20
 
-#define PREDICTOR_RES_MULTIPLIER 2
-
-
 #define PARALLAX_REJECTION_MASK_THRESHOLD 0.5
 
 #define MAX_PARALLAX_FACTOR 8.0
@@ -44,16 +41,28 @@
 #define SUB_PIXEL_GRID_SIZE 2
 #define SUB_PIXEL_COUNT (SUB_PIXEL_GRID_SIZE * SUB_PIXEL_GRID_SIZE)
 
-#define SUB_PIXEL_BIT_COUNT 8
+#define SUB_PIXEL_BIT_COUNT 4
 #define SUB_PIXEL_BIT_MASK ((1 << SUB_PIXEL_BIT_COUNT) - 1)
 
 #define SUB_PIXEL_PARALLAX_FACTOR_BIT_COUNT 4
 #define SUB_PIXEL_PARALLAX_FACTOR_BIT_OFFSET 0
 #define SUB_PIXEL_PARALLAX_FACTOR_BIT_MASK ((1 << SUB_PIXEL_PARALLAX_FACTOR_BIT_COUNT) - 1)
 
-#define SUB_PIXEL_INTERFERENCE_BIT_COUNT 4
-#define SUB_PIXEL_INTERFERENCE_BIT_OFFSET 4
-#define SUB_PIXEL_INTERFERENCE_BIT_MASK ((1 << SUB_PIXEL_INTERFERENCE_BIT_COUNT) - 1)
+
+// Size and payload of the super resolution history
+#define SUPERRES_RES_MULTIPLIER 2
+
+#define SUPERRES_AGE_BITS 5
+#define SUPERRES_AGE_BITMASK ((1 << SUPERRES_AGE_BITS) - 1)
+
+#define SUPERRES_INTERFERENCE_BITS 4
+#define SUPERRES_INTERFERENCE_BITMASK ((1 << SUPERRES_INTERFERENCE_BITS) - 1)
+
+#define SUPERRES_OFFSET_BITS 4
+#define SUPERRES_OFFSET_BITMASK ((1 << SUPERRES_OFFSET_BITS) - 1)
+#define SUPERRES_OFFSET_OFFSET (SUPERRES_OFFSET_BITMASK / 2)
+#define SUPERRES_OFFSET_MIN (-SUPERRES_OFFSET_OFFSET)
+#define SUPERRES_OFFSET_MAX (SUPERRES_OFFSET_BITMASK - SUPERRES_OFFSET_OFFSET)
 
 
 
@@ -486,20 +495,6 @@ taa_half LumaToInterferenceSeed(taa_half Luma)
 	return Luma * HdrWeightY(Luma);
 }
 
-taa_subpixel_payload CompressInteferenceSeed(taa_half InteferenceSeed)
-{
-	return clamp(
-		taa_subpixel_payload(InteferenceSeed * taa_half(SUB_PIXEL_INTERFERENCE_BIT_MASK)),
-		taa_subpixel_payload(0), taa_subpixel_payload(SUB_PIXEL_INTERFERENCE_BIT_MASK));
-}
-
-taa_half UncompressInteferenceSeed(taa_subpixel_payload CompressedInteferenceSeed)
-{
-	return 
-		taa_half(CompressedInteferenceSeed & SUB_PIXEL_INTERFERENCE_BIT_MASK) *
-		taa_half(1.0 / float(SUB_PIXEL_INTERFERENCE_BIT_MASK));
-}
-
 taa_half3 QuantizeForFloatRenderTarget(taa_half3 Color, taa_half E, float3 QuantizationError)
 {
 	taa_half3 Error = Color * taa_half3(QuantizationError);
@@ -522,3 +517,42 @@ taa_half3 QuantizeForFloatRenderTarget(taa_half3 Color, taa_half E, float3 Quant
 	
 	return Color + Error * E;
 }
+
+
+
+struct FSuperResHistory
+{
+	int2 TexelToSample;
+	uint Age;
+	taa_half InterferenceSeed;
+};
+
+uint2 PackSuperResHistory(FSuperResHistory SuperRes)
+{
+	uint CompressedInterferenceSeed = clamp(
+		uint(SuperRes.InterferenceSeed * taa_half(SUPERRES_INTERFERENCE_BITMASK)),
+		uint(0), uint(SUPERRES_INTERFERENCE_BITMASK));
+
+	uint2 Packed = 0;
+
+	Packed[0] |= ((SuperRes.Age) & SUPERRES_AGE_BITMASK) << 0;
+	Packed[0] |= ((CompressedInterferenceSeed) & SUPERRES_INTERFERENCE_BITMASK) << 5;
+
+	Packed[1] |= (uint(SuperRes.TexelToSample.x + SUPERRES_OFFSET_OFFSET) & SUPERRES_OFFSET_BITS) << 0 * SUPERRES_OFFSET_BITS;
+	Packed[1] |= (uint(SuperRes.TexelToSample.y + SUPERRES_OFFSET_OFFSET) & SUPERRES_OFFSET_BITS) << 1 * SUPERRES_OFFSET_BITS;
+
+	return Packed;
+}
+
+FSuperResHistory UnpackSuperResHistory(uint2 Packed)
+{
+	FSuperResHistory SuperRes;
+
+	SuperRes.Age = ((Packed[0] >> 0) & SUPERRES_AGE_BITMASK);
+	SuperRes.InterferenceSeed = taa_half((Packed[0] >> 5) & SUPERRES_INTERFERENCE_BITMASK) * rcp(taa_half(SUPERRES_INTERFERENCE_BITMASK));
+
+	SuperRes.TexelToSample.x = ((Packed[1] >> 0 * SUPERRES_OFFSET_BITS) & SUPERRES_OFFSET_BITMASK) - SUPERRES_OFFSET_OFFSET;
+	SuperRes.TexelToSample.y = ((Packed[1] >> 1 * SUPERRES_OFFSET_BITS) & SUPERRES_OFFSET_BITMASK) - SUPERRES_OFFSET_OFFSET;
+
+	return SuperRes;
+}
diff --git a/Engine/Shaders/Private/TemporalAA/TAADebugHistory.usf b/Engine/Shaders/Private/TemporalAA/TAADebugHistory.usf
--- a/Engine/Shaders/Private/TemporalAA/TAADebugHistory.usf
+++ b/Engine/Shaders/Private/TemporalAA/TAADebugHistory.usf
@@ -22,6 +22,43 @@ Texture2D<taa_half2> PrevHistory_Textures_2;
 Texture2D<taa_subpixel_details> PrevHistory_Textures_3;
 
 
+//------------------------------------------------------- FUNCTION
+
+
+void LMHOperator(float A, float B, float C, out float L, out float M, out float H)
+{
+	float X = min(B, C);
+	float Y = max(B, C);
+	
+	L = min(A, X);
+	float Z = max(A, X);
+
+	M = min(Z, Y);
+	H = max(Z, Y);
+}
+
+float Median(float A, float B, float C)
+{
+	float L, M, H;
+	LMHOperator(A, B, C, L, M, H);
+	return M;
+}
+
+
+taa_subpixel_payload ReadSubpixelDetails(uint2 SubpixelPos)
+{
+	uint2 HistoryPixelPos = SubpixelPos / SUB_PIXEL_GRID_SIZE;
+
+	uint2 SubpixelRelPos = SubpixelPos % SUB_PIXEL_GRID_SIZE;
+	uint SubpixelId = dot(SubpixelRelPos, uint2(1, SUB_PIXEL_GRID_SIZE));
+	uint SubpixelShift = SubpixelId * SUB_PIXEL_BIT_COUNT;
+	
+	taa_subpixel_details RawHistory3 = History_Textures_3[HistoryPixelPos];
+
+	return taa_subpixel_payload(RawHistory3 >> SubpixelShift) & SUB_PIXEL_BIT_MASK;
+}
+
+
 //------------------------------------------------------- ENTRY POINT
 
 [numthreads(TILE_SIZE, TILE_SIZE, 1)]
@@ -43,24 +80,25 @@ void MainCS(
 	uint SubpixelId = dot(SubpixelPos, uint2(1, SUB_PIXEL_GRID_SIZE));
 	uint SubpixelShift = SubpixelId * SUB_PIXEL_BIT_COUNT;
 	
-#if 1
-	taa_half3 RawHistory0 = History_Textures_0[HistoryPixelPos];
-	taa_half3 RawHistory1 = History_Textures_1[HistoryPixelPos];
-	taa_half2 RawHistory2 = History_Textures_2[HistoryPixelPos];
-	taa_subpixel_details RawHistory3 = History_Textures_3[HistoryPixelPos];
-#else
-	taa_half3 RawHistory0 = PrevHistory_Textures_0[HistoryPixelPos];
-	taa_half3 RawHistory1 = PrevHistory_Textures_1[HistoryPixelPos];
-	taa_half2 RawHistory2 = PrevHistory_Textures_2[HistoryPixelPos];
-	taa_subpixel_details RawHistory3 = PrevHistory_Textures_3[HistoryPixelPos];
-#endif
-
-	taa_subpixel_payload SubpixelDetails = taa_subpixel_payload(RawHistory3 >> SubpixelShift) & SUB_PIXEL_BIT_MASK;
-
-	Debug.r = UncompressParallaxFactor(SubpixelDetails);
-	Debug.g = UncompressInteferenceSeed(SubpixelDetails >> SUB_PIXEL_INTERFERENCE_BIT_OFFSET);
-	Debug.b = SubpixelId;
-	
+	{
+		#if 1
+			taa_half3 RawHistory0 = History_Textures_0[HistoryPixelPos];
+			taa_half3 RawHistory1 = History_Textures_1[HistoryPixelPos];
+			taa_half2 RawHistory2 = History_Textures_2[HistoryPixelPos];
+			taa_subpixel_details RawHistory3 = History_Textures_3[HistoryPixelPos];
+		#else
+			taa_half3 RawHistory0 = PrevHistory_Textures_0[HistoryPixelPos];
+			taa_half3 RawHistory1 = PrevHistory_Textures_1[HistoryPixelPos];
+			taa_half2 RawHistory2 = PrevHistory_Textures_2[HistoryPixelPos];
+			taa_subpixel_details RawHistory3 = PrevHistory_Textures_3[HistoryPixelPos];
+		#endif
+
+		taa_subpixel_payload SubpixelDetails = taa_subpixel_payload(RawHistory3 >> SubpixelShift) & SUB_PIXEL_BIT_MASK;
+
+		Debug.r = UncompressParallaxFactor(SubpixelDetails);
+		Debug.b = SubpixelId;
+	}
+
 	#if DEBUG_OUTPUT
 	if (all(DispatchThreadId < HistoryInfo_ViewportMax * SUB_PIXEL_GRID_SIZE))
 	{
diff --git a/Engine/Shaders/Private/TemporalAA/TAADecimateHistory.usf b/Engine/Shaders/Private/TemporalAA/TAADecimateHistory.usf
--- a/Engine/Shaders/Private/TemporalAA/TAADecimateHistory.usf
+++ b/Engine/Shaders/Private/TemporalAA/TAADecimateHistory.usf
@@ -30,6 +30,8 @@ Texture2D<taa_half3> PrevHistory_Textures_1;
 Texture2D<taa_half2> PrevHistory_Textures_2;
 Texture2D<uint>      PrevHistory_Textures_3;
 Texture2D<taa_half>  PrevHistory_LowResTextures_0;
+Texture2D<uint>      PrevHistory_SuperResTextures_0;
+Texture2D<uint>      PrevHistory_SuperResTextures_1;
 
 RWTexture2D<taa_half3> HalfResSceneColorOutput;
 RWTexture2D<taa_half3> HalfResPredictionSceneColorOutput;
@@ -81,7 +83,7 @@ void MainCS(
 	taa_half3 InputSceneColor = InputSceneColorTexture[InputPixelPos];
 	
 	taa_half ClosestInputLuma4 = Luma4(InputSceneColor);
-
+	taa_half InputInterferenceSeed = LumaToInterferenceSeed(ClosestInputLuma4);
 
 	float WorldDepth = ConvertFromDeviceZ(DeviceZ);
 	
@@ -159,6 +161,76 @@ void MainCS(
 		
 		float2 PrevInputBufferUV = (InputInfo_ScreenPosToViewportScale * PrevOutputScreenPos + InputInfo_ScreenPosToViewportBias) * InputInfo_ExtentInverse;
 
+		{
+			const uint KernelSize = 4;
+
+			uint SuperResPayloads[16];
+
+			uint ClosestPayload = ~0;
+			uint ClosestAge = 0;
+
+			// Sample the super res history and find an aproximative age of the input pixel.
+			UNROLL
+			for (uint y = 0; y < KernelSize; y++)
+			{
+				UNROLL
+				for (uint x = 0; x < KernelSize; x++)
+				{
+					const float2 Offset = float2(x, y) - (KernelSize - 1.0) * 0.5;
+					
+					float2 SampleHistoryBufferUV = clamp(
+						(HistoryPPCk * 2.0 + Offset) * PrevHistoryInfo_ExtentInverse * rcp(float(SUPERRES_RES_MULTIPLIER)),
+						PrevHistoryInfo_UVViewportBilinearMin,
+						PrevHistoryInfo_UVViewportBilinearMax);
+
+					uint2 SuperResPixelCoord = uint2(SampleHistoryBufferUV * PrevHistoryInfo_Extent * float(SUPERRES_RES_MULTIPLIER));
+
+					uint SubpixelPayload = PrevHistory_SuperResTextures_0[SuperResPixelCoord];
+					FSuperResHistory UnpackedSubpixel = UnpackSuperResHistory(SubpixelPayload);
+
+					SuperResPayloads[y * 4 + x] = SubpixelPayload;
+
+					taa_half DistanceToSeed = abs(InputInterferenceSeed - UnpackedSubpixel.InterferenceSeed);
+
+					FLATTEN
+					if (DistanceToSeed < 0.1 && UnpackedSubpixel.Age > ClosestAge)
+					{
+						ClosestPayload = SubpixelPayload;
+						ClosestAge = UnpackedSubpixel.Age;
+					}
+				}
+			}
+			
+			uint AgeMin = UnpackSuperResHistory(ClosestPayload).Age;
+			uint AgeMax = UnpackSuperResHistory(ClosestPayload).Age + 4;
+
+			ClosestInterferenceSeed = UnpackSuperResHistory(ClosestPayload).InterferenceSeed;
+
+			MinInterferenceSeed = 1.0;
+			MaxInterferenceSeed = 0.0;
+
+			// Find out if there is high frequency details in the history with that age.
+			UNROLL
+			for (uint i = 0; i < KernelSize * KernelSize; i++)
+			{
+				FSuperResHistory UnpackedSubpixel = UnpackSuperResHistory(SuperResPayloads[i]);
+
+				bool bAgeBilateralWeight = UnpackedSubpixel.Age >= AgeMin && UnpackedSubpixel.Age <= AgeMax;
+
+				FLATTEN
+				if (bAgeBilateralWeight && ClosestPayload != ~0)
+				{
+					MinInterferenceSeed = min(MinInterferenceSeed, UnpackedSubpixel.InterferenceSeed);
+					MaxInterferenceSeed = max(MaxInterferenceSeed, UnpackedSubpixel.InterferenceSeed);
+				}
+
+			}
+			
+			// actually integrating SubpixelSeed stil suffer interference that causes instability. So instead take
+			// average of min and max that turns out interference free.
+			IntegratedInterferenceSeed = (MinInterferenceSeed + MaxInterferenceSeed) * 0.5;
+		}
+
 		taa_half PrevUninterferenceSeed = PrevHistory_LowResTextures_0.SampleLevel(GlobalBilinearClampedSampler, PrevInputBufferUV, 0);
 		taa_half PrevUninterferenceLuma = PrevUninterferenceSeed * HdrWeightInvY(PrevUninterferenceSeed);
 
@@ -211,53 +283,12 @@ void MainCS(
 			{
 				taa_subpixel_payload SubpixelPayload = GetSubpixelPayload(SubpixelNeigborhood, SubpixelId);
 				taa_subpixel_payload EncodedRejectionFactor = (SubpixelPayload >> SUB_PIXEL_PARALLAX_FACTOR_BIT_OFFSET) & SUB_PIXEL_PARALLAX_FACTOR_BIT_MASK;
-				taa_subpixel_payload EncodedInterference = (SubpixelPayload >> SUB_PIXEL_INTERFERENCE_BIT_OFFSET) & SUB_PIXEL_INTERFERENCE_BIT_MASK;
 
 				EncodedHistoryMaxRejectionFactor = max(EncodedHistoryMaxRejectionFactor, EncodedRejectionFactor);
 			}
 
 			HistoryMaxRejectionFactor = UncompressParallaxFactor(EncodedHistoryMaxRejectionFactor);
 		}
-
-		// Compute the inteference seeds for detection of the interference
-		{
-			const taa_half RefInterferenceSeed = LumaToInterferenceSeed(ClosestInputLuma4);
-			ClosestInterferenceSeed = 0.0;
-			taa_half RefInterferenceSeedDist = 2.0;
-
-			MinInterferenceSeed = 1.0;
-			MaxInterferenceSeed = 0.0;
-
-			// Searches in the entire 2x2 neighborhood instead because this history is nearest reproject that can causes subpixel
-			// shifts over multiple frames.
-			UNROLL_N(4)
-			for (uint SampleId = 0; SampleId < 4; SampleId++)
-			{
-				UNROLL_N(SUB_PIXEL_COUNT)
-				for (uint SubpixelId = 0; SubpixelId < SUB_PIXEL_COUNT; SubpixelId++)
-				{
-					taa_subpixel_payload SubpixelPayload = SubpixelNeigborhood.SampleArray[SampleId] >> (SubpixelId * SUB_PIXEL_BIT_COUNT);
-					taa_subpixel_payload EncodedInterference = (SubpixelPayload >> SUB_PIXEL_INTERFERENCE_BIT_OFFSET) & SUB_PIXEL_INTERFERENCE_BIT_MASK;
-					
-					// TODO(G5TAA): move out of the loop.
-					taa_half SubpixelSeed = UncompressInteferenceSeed(EncodedInterference);
-
-					FLATTEN
-					if (abs(RefInterferenceSeed - SubpixelSeed) < RefInterferenceSeedDist)
-					{
-						ClosestInterferenceSeed = SubpixelSeed;
-						RefInterferenceSeedDist = abs(RefInterferenceSeed - SubpixelSeed);
-					}
-
-					MinInterferenceSeed = min(MinInterferenceSeed, SubpixelSeed);
-					MaxInterferenceSeed = max(MaxInterferenceSeed, SubpixelSeed);
-				}
-			}
-			
-			// actually integrating SubpixelSeed stil suffer interference that causes instability. So instead take
-			// average of min and max that turns out interference free.
-			IntegratedInterferenceSeed = (MinInterferenceSeed + MaxInterferenceSeed) * 0.5;
-		}
 	}
 	
 	// Reject subpixel details that might be in history but not input just based on their parrallax
diff --git a/Engine/Shaders/Private/TemporalAA/TAADetectInterference.usf b/Engine/Shaders/Private/TemporalAA/TAADetectInterference.usf
--- a/Engine/Shaders/Private/TemporalAA/TAADetectInterference.usf
+++ b/Engine/Shaders/Private/TemporalAA/TAADetectInterference.usf
@@ -142,11 +142,10 @@ void MainCS(
 	}
 
 	taa_half UninterferenceSeed;
-	bool bRejectHistoryInterferenceSeed;
 	bool bIsInterference;
 	{
 		// Error due to the low bitdepth of the inteference seeds in the history.
-		const taa_half InterferenceSeedEncodingError = rcp(taa_half(SUB_PIXEL_INTERFERENCE_BIT_MASK));
+		const taa_half InterferenceSeedEncodingError = rcp(taa_half(SUPERRES_INTERFERENCE_BITMASK));
 
 		taa_half InputSeedMin = LumaToInterferenceSeed(InputLumaMin);
 		taa_half InputSeedMax = LumaToInterferenceSeed(InputLumaMax);
@@ -174,16 +173,6 @@ void MainCS(
 		// Detect interference solely based on the inteference seed, but need to take into account low frequency
 		// gradien that there may be in the neighborhood to avoid ghosting on high contrast geometric edges.
 		bIsInterference = bRejectInterferenceSeedWithGradient;
-		
-		// If the input is not clamping, but the history's interference seeds are, that means the later are ghosting.
-		// Only care about bRejectInterferenceSeed here because it is only about not clamping inteference seeds in the history
-		// that may not be in the input whereas bRejectInterferenceSeedWithGradient is ghosting the seed to much in the history.
-		bRejectHistoryInterferenceSeed = !bRejectPrevScene && bRejectInterferenceSeed;
-
-		#if 0
-			bIsInterference = false;
-			bRejectHistoryInterferenceSeed = false;
-		#endif
 
 		// Don't clamp the seed history if there is an interference to keep luma stability in the input
 		// of the shading heuristic next frame.
@@ -196,7 +185,6 @@ void MainCS(
 		FLATTEN
 		if (bOffScreen)
 		{
-			bRejectHistoryInterferenceSeed = true;
 			bIsInterference = false;
 			ClampedPrevUninterferenceSeed = LumaToInterferenceSeed(InputLuma);
 		}
@@ -208,7 +196,6 @@ void MainCS(
 	if (all(InputPixelPos < InputInfo_ViewportMax))
 	{
 		InterferenceWeightOutput[InputPixelPos] = bIsInterference ? 0.0 : 1.0;
-		InterferenceSeedUpdateOutput[InputPixelPos] = bRejectHistoryInterferenceSeed ? 1.0 : 0.0;
 		
 		HistoryOutput_LowResTextures_0[InputPixelPos] = UninterferenceSeed;
 
diff --git a/Engine/Shaders/Private/TemporalAA/TAAUpdateHistory.usf b/Engine/Shaders/Private/TemporalAA/TAAUpdateHistory.usf
--- a/Engine/Shaders/Private/TemporalAA/TAAUpdateHistory.usf
+++ b/Engine/Shaders/Private/TemporalAA/TAAUpdateHistory.usf
@@ -33,7 +33,6 @@ Texture2D<taa_half> HistoryRejectionTexture;
 Texture2D<float2> DilatedVelocityTexture;
 Texture2D<uint> ParallaxFactorTexture;
 Texture2D<taa_half> ParallaxRejectionMaskTexture;
-Texture2D<taa_half> InterferenceSeedUpdateTexture;
 
 float3 HistoryQuantizationError;
 float HistoryPreExposureCorrection;
@@ -114,7 +113,6 @@ void MainCS(
 	float2 PrevScreenPos = ScreenPos;
 	taa_half ParallaxRejectionMask = taa_half(1.0);
 	taa_half LowFrequencyRejection;
-	bool bRejectHistoryInterferenceSeed;
 	#if 1
 	{
 		float2 RejectionBilinearBufferUV = clamp(
@@ -128,13 +126,10 @@ void MainCS(
 
 		LowFrequencyRejection = HistoryRejectionTexture.SampleLevel(GlobalBilinearClampedSampler, RejectionBilinearBufferUV, 0);
 		
-		bRejectHistoryInterferenceSeed = InterferenceSeedUpdateTexture.SampleLevel(GlobalBilinearClampedSampler, RejectionBilinearBufferUV, 0) > 0.5;
-
 		#if !CONFIG_CLAMP
 		{
 			ParallaxRejectionMask = taa_half(1.0);
 			LowFrequencyRejection = taa_half(1.0);
-			bRejectHistoryInterferenceSeed = false;
 		}
 		#endif
 	}
@@ -584,8 +579,6 @@ void MainCS(
 
 		// TODO Save VGPR here
 		taa_subpixel_payload ParallaxFactorBits = ParallaxFactorTexture[InputPixelPos] & SUB_PIXEL_PARALLAX_FACTOR_BIT_MASK;
-		taa_half InterferenceSeed = LumaToInterferenceSeed(ClosestInputLuma4);
-		taa_subpixel_payload InterferenceBits = CompressInteferenceSeed(InterferenceSeed);
 
 		{
 			const uint ParallaxFactorMask = (
@@ -594,20 +587,6 @@ void MainCS(
 				(SUB_PIXEL_PARALLAX_FACTOR_BIT_MASK << (SUB_PIXEL_PARALLAX_FACTOR_BIT_OFFSET + 2 * SUB_PIXEL_BIT_COUNT)) | 
 				(SUB_PIXEL_PARALLAX_FACTOR_BIT_MASK << (SUB_PIXEL_PARALLAX_FACTOR_BIT_OFFSET + 3 * SUB_PIXEL_BIT_COUNT)) | 
 				0x0);
-
-			const uint InterferenceMask = (
-				(SUB_PIXEL_INTERFERENCE_BIT_MASK << (SUB_PIXEL_INTERFERENCE_BIT_OFFSET + 0 * SUB_PIXEL_BIT_COUNT)) | 
-				(SUB_PIXEL_INTERFERENCE_BIT_MASK << (SUB_PIXEL_INTERFERENCE_BIT_OFFSET + 1 * SUB_PIXEL_BIT_COUNT)) | 
-				(SUB_PIXEL_INTERFERENCE_BIT_MASK << (SUB_PIXEL_INTERFERENCE_BIT_OFFSET + 2 * SUB_PIXEL_BIT_COUNT)) | 
-				(SUB_PIXEL_INTERFERENCE_BIT_MASK << (SUB_PIXEL_INTERFERENCE_BIT_OFFSET + 3 * SUB_PIXEL_BIT_COUNT)) | 
-				0x0);
-
-			uint ResetInterferenceSeed = (
-				(uint(InterferenceBits) << (SUB_PIXEL_INTERFERENCE_BIT_OFFSET + 0 * SUB_PIXEL_BIT_COUNT)) | 
-				(uint(InterferenceBits) << (SUB_PIXEL_INTERFERENCE_BIT_OFFSET + 1 * SUB_PIXEL_BIT_COUNT)) | 
-				(uint(InterferenceBits) << (SUB_PIXEL_INTERFERENCE_BIT_OFFSET + 2 * SUB_PIXEL_BIT_COUNT)) | 
-				(uint(InterferenceBits) << (SUB_PIXEL_INTERFERENCE_BIT_OFFSET + 3 * SUB_PIXEL_BIT_COUNT)) | 
-				0x0);
 			
 			// Reset parallax factor
 			FLATTEN
@@ -615,30 +594,6 @@ void MainCS(
 			{
 				FinalSubpixelDetails = FinalSubpixelDetails & ~ParallaxFactorMask;
 			}
-
-			// Reset inteference bit
-			FLATTEN
-			if (bRejectHistoryInterferenceSeed)
-			{
-				uint NewSubpixelDetails = FinalSubpixelDetails & ~InterferenceMask;
-				
-				taa_subpixel_payload InterferenceSeedMin = CompressInteferenceSeed(LumaToInterferenceSeed(LumaMin));
-				taa_subpixel_payload InterferenceSeedMax = CompressInteferenceSeed(LumaToInterferenceSeed(LumaMax));
-
-				UNROLL_N(SUB_PIXEL_COUNT)
-				for (uint SubpixelId = 0; SubpixelId < SUB_PIXEL_COUNT; SubpixelId++)
-				{
-					const uint SubpixelShift = SubpixelId * SUB_PIXEL_BIT_COUNT + SUB_PIXEL_INTERFERENCE_BIT_OFFSET;
-
-					taa_subpixel_payload SubpixelInterferenceSeed = (FinalSubpixelDetails >> SubpixelShift) & SUB_PIXEL_INTERFERENCE_BIT_MASK;
-
-					taa_subpixel_payload ClampedSubpixelInterferenceSeed = clamp(SubpixelInterferenceSeed, InterferenceSeedMin, InterferenceSeedMax);
-
-					NewSubpixelDetails |= ClampedSubpixelInterferenceSeed << SubpixelShift;
-				}
-
-				FinalSubpixelDetails = NewSubpixelDetails;
-			}
 		}
 
 		FLATTEN
@@ -649,9 +604,7 @@ void MainCS(
 			uint SubpixelId = dot(uint2(bBool), uint2(1, SUB_PIXEL_GRID_SIZE));
 			uint SubpixelShift = SubpixelId * SUB_PIXEL_BIT_COUNT;
 
-			taa_subpixel_payload SubpixelPayload = (
-				(ParallaxFactorBits << SUB_PIXEL_PARALLAX_FACTOR_BIT_OFFSET) |
-				(InterferenceBits << SUB_PIXEL_INTERFERENCE_BIT_OFFSET));
+			taa_subpixel_payload SubpixelPayload = (ParallaxFactorBits << SUB_PIXEL_PARALLAX_FACTOR_BIT_OFFSET);
 
 			FinalSubpixelDetails = (FinalSubpixelDetails & (~(SUB_PIXEL_BIT_MASK << SubpixelShift))) | (SubpixelPayload << SubpixelShift);
 		}
diff --git a/Engine/Shaders/Private/TemporalAA/TAAUpdateSuperResHistory.usf b/Engine/Shaders/Private/TemporalAA/TAAUpdateSuperResHistory.usf
new file mode 100644
--- /dev/null
+++ b/Engine/Shaders/Private/TemporalAA/TAAUpdateSuperResHistory.usf
@@ -0,0 +1,220 @@
+// Copyright Epic Games, Inc. All Rights Reserved.
+
+#include "TAACommon.ush"
+#include "../Random.ush"
+
+
+//------------------------------------------------------- CONSTANTS
+
+#define TILE_SIZE 8
+
+
+//------------------------------------------------------- PARAMETERS
+
+Texture2D<taa_half3> InputSceneColorTexture;
+
+Texture2D<float2> DilatedVelocityTexture;
+Texture2D<taa_half> ParallaxRejectionMaskTexture;
+
+Texture2D<uint> PrevHistory_SuperResTextures_0;
+Texture2D<uint> PrevHistory_SuperResTextures_1;
+
+RWTexture2D<uint> HistoryOutput_SuperResTextures_0;
+RWTexture2D<uint> HistoryOutput_SuperResTextures_1;
+
+
+//------------------------------------------------------- FUNCTIONS
+
+uint GetGroupThreadIndex(uint GroupThreadIndex)
+#if COMPILER_SUPPORTS_WAVE_ONCE
+{
+	ISOLATE
+	{
+		return WaveGetLaneIndex();
+	}
+}
+#else
+{
+	return GroupThreadIndex;
+}
+#endif
+
+
+//------------------------------------------------------- ENTRY POINT
+
+[numthreads(TILE_SIZE, TILE_SIZE, 1)]
+void MainCS(
+	uint2 GroupId : SV_GroupID,
+	uint GroupThreadIndex : SV_GroupIndex)
+{
+	float4 Debug = 0.0;
+
+	uint2 DispatchThreadId = (
+		ZOrder2D(GroupThreadIndex, log2(TILE_SIZE)) +
+		GroupId * uint2(TILE_SIZE, TILE_SIZE));
+
+	uint2 HistorySuperPixelPos = DispatchThreadId;
+	float2 ViewportUV = (float2(HistorySuperPixelPos) + 0.5f) * HistoryInfo_ViewportSizeInverse * rcp(taa_half(SUPERRES_RES_MULTIPLIER));
+	float2 ScreenPos = ViewportUVToScreenPos(ViewportUV);
+	
+	// Pixel coordinate of the center of output pixel O in the input viewport.
+	float2 PPCo = ViewportUV * InputInfo_ViewportSize + InputJitter;
+
+	// Pixel coordinate of the center of the nearest input pixel K.
+	float2 PPCk = floor(PPCo) + 0.5;
+	
+	// Vector in pixel between pixel K -> O in input pixel size.
+	taa_half2 idKO = taa_half2(PPCo - PPCk);
+	
+	// Vector in pixel between pixel K -> O in output pixel size.
+	taa_half2 odKO = idKO * (InputInfo_ViewportSizeInverse.x * HistoryInfo_ViewportSize.x * SUPERRES_RES_MULTIPLIER);
+
+	uint2 InputPixelPos = InputInfo_ViewportMin + uint2(PPCk);
+
+	// Fetch reprojections related informations.
+	float2 PrevScreenPos = ScreenPos;
+	taa_half ParallaxRejectionMask = taa_half(1.0);
+	#if 1
+	{
+		float2 RejectionBilinearBufferUV = clamp(
+			PPCo * InputInfo_ExtentInverse,
+			RejectionInfo_UVViewportBilinearMin,
+			RejectionInfo_UVViewportBilinearMax);
+		
+		float2 EncodedVelocity = DilatedVelocityTexture[InputPixelPos];
+		ParallaxRejectionMask = ParallaxRejectionMaskTexture[InputPixelPos];
+		PrevScreenPos = ScreenPos - DecodeVelocityFromTexture(float4(EncodedVelocity, 0.0, 0.0)).xy;
+	}
+	#endif
+
+	float2 ScreenVelocity = ScreenPos - PrevScreenPos;
+	
+	// Detect if HistoryBufferUV would be outside of the viewport.
+	bool bOffScreen = IsOffScreen(bCameraCut, PrevScreenPos, ParallaxRejectionMask);
+	
+	// Read input color
+	taa_half ClosestInterferenceSeed;
+	ISOLATE
+	{
+		float2 SampleInputBufferUV = (InputInfo_ViewportMin + PPCk) * InputInfo_ExtentInverse;
+		SampleInputBufferUV = clamp(SampleInputBufferUV, InputInfo_UVViewportBilinearMin, InputInfo_UVViewportBilinearMax);
+
+		taa_half3 InputColor = InputSceneColorTexture.SampleLevel(GlobalPointClampedSampler, SampleInputBufferUV, 0);
+
+		taa_half ClosestInputLuma4 = Luma4(InputColor);
+		ClosestInterferenceSeed = LumaToInterferenceSeed(ClosestInputLuma4);
+	}
+	
+	// Reproject history
+	uint OldestAge = 0;
+	FSuperResHistory SubpixelHistory;
+	ISOLATE
+	{
+		float2 PrevHistoryBufferUV = (PrevHistoryInfo_ScreenPosToViewportScale * PrevScreenPos + PrevHistoryInfo_ScreenPosToViewportBias) * PrevHistoryInfo_ExtentInverse;
+		PrevHistoryBufferUV = clamp(PrevHistoryBufferUV, PrevHistoryInfo_UVViewportBilinearMin, PrevHistoryInfo_UVViewportBilinearMax);
+
+		FBilinearSampleInfos SampleInfos = GetBilinearSampleLevelInfos(
+			PrevHistoryBufferUV,
+			PrevHistoryInfo_Extent * float(SUPERRES_RES_MULTIPLIER),
+			PrevHistoryInfo_ExtentInverse * rcp(float(SUPERRES_RES_MULTIPLIER)));
+
+		uint2 BestPayload = 0;
+		int2 BestKernelToSample = 0;
+		uint BestKernelToSampleDistSquare = dot(uint2(SUPERRES_OFFSET_MAX, SUPERRES_OFFSET_MAX), uint2(SUPERRES_OFFSET_MAX, SUPERRES_OFFSET_MAX));
+
+
+		UNROLL_N(4)
+		for (uint SampleId = 0; SampleId < 4; SampleId++)
+		{
+			uint2 PixelCoord = GetSamplePixelCoord(SampleInfos, SampleId);
+
+			uint2 SubpixelPayload;
+			SubpixelPayload[0] = PrevHistory_SuperResTextures_0[PixelCoord];
+			SubpixelPayload[1] = PrevHistory_SuperResTextures_1[PixelCoord];
+
+			FSuperResHistory UnpackedSubpixel = UnpackSuperResHistory(SubpixelPayload);
+
+			const float2 PixelOffset = float2(BilinearSamplingOffsets2x2[SampleId]);
+
+			int2 KernelToTexel = int2(round((PixelOffset - SampleInfos.BilinearInterp) * 0x20));
+
+			int2 KernelToSample = KernelToTexel + UnpackedSubpixel.TexelToSample;
+
+			uint KernelToSampleDistSquare = uint(dot(KernelToSample, KernelToSample));
+
+			FLATTEN
+			if (KernelToSampleDistSquare < BestKernelToSampleDistSquare || SampleId == 0)
+			{
+				BestPayload = SubpixelPayload;
+				BestKernelToSample = KernelToSample;
+				BestKernelToSampleDistSquare = KernelToSampleDistSquare;
+			}
+
+			bool bIsCloseEnough = abs(UnpackedSubpixel.InterferenceSeed - ClosestInterferenceSeed) < 0.1;
+
+			if (bIsCloseEnough)
+			{
+				OldestAge = max(OldestAge, UnpackedSubpixel.Age);
+			}
+		}
+
+		SubpixelHistory = UnpackSuperResHistory(BestPayload);
+		SubpixelHistory.Age = min(SubpixelHistory.Age + 1, SUPERRES_AGE_BITMASK);
+		SubpixelHistory.TexelToSample = clamp(BestKernelToSample, SUPERRES_OFFSET_MIN, SUPERRES_OFFSET_MAX);
+	}
+		
+	int2 CurrentTexelToSample = int2(round(odKO * (-2 * SUPERRES_OFFSET_MAX)));
+	
+	bool bUpdate = all(CurrentTexelToSample >= SUPERRES_OFFSET_MIN && CurrentTexelToSample <= SUPERRES_OFFSET_MAX);
+
+	FLATTEN
+	if (bOffScreen || bUpdate)
+	{
+		uint2 Random = Rand3DPCG16(int3(InputPixelPos, View.StateFrameIndexMod8)).xy;
+		float2 E = Hammersley16(0, 1, Random);
+
+		SubpixelHistory.TexelToSample = clamp(CurrentTexelToSample, SUPERRES_OFFSET_MIN, SUPERRES_OFFSET_MAX);
+		SubpixelHistory.Age = bOffScreen ? 0 : OldestAge;
+		SubpixelHistory.InterferenceSeed = saturate(ClosestInterferenceSeed + (E.x - 0.5) * rcp(taa_half(SUPERRES_INTERFERENCE_BITMASK)));
+	}
+
+	#if 0
+	{
+		Debug.r = SubpixelHistory.InterferenceSeed;
+		Debug.g = SubpixelHistory.Age * rcp(SUPERRES_AGE_BITMASK);
+		//Debug.b = SubpixelHistory.TexelToSample.x * rcp(16.0);
+		//Debug.a = SubpixelHistory.TexelToSample.y * rcp(16.0);
+
+		Debug.b = bOffScreen ? 1 : 0;
+		//Debug.b = ParallaxRejectionMask;
+	}
+	#endif
+
+	uint2 PackedSubpixelHistory = PackSuperResHistory(SubpixelHistory);
+
+	ISOLATE
+	{
+		#if 1
+			uint LocalGroupThreadIndex = GetGroupThreadIndex(GroupThreadIndex);
+			uint2 LocalHistoryPixelPos = (
+				ZOrder2D(LocalGroupThreadIndex, log2(TILE_SIZE)) +
+				GroupId * uint2(TILE_SIZE, TILE_SIZE));
+		#else
+			uint LocalGroupThreadIndex = HistorySuperPixelPos;
+		#endif
+			
+		// Output final history.
+		bool bIsValidhistoryPixel = all(LocalHistoryPixelPos < SUPERRES_RES_MULTIPLIER * HistoryInfo_ViewportMax);
+		if (bIsValidhistoryPixel)
+		{
+			HistoryOutput_SuperResTextures_0[LocalHistoryPixelPos] = PackedSubpixelHistory[0];
+			HistoryOutput_SuperResTextures_1[LocalHistoryPixelPos] = PackedSubpixelHistory[1];
+
+			#if DEBUG_OUTPUT
+			{
+				DebugOutput[LocalHistoryPixelPos] = Debug;
+			}
+			#endif
+		}
+	}
+}
diff --git a/Engine/Source/Runtime/Renderer/Private/PostProcess/TemporalAA.cpp b/Engine/Source/Runtime/Renderer/Private/PostProcess/TemporalAA.cpp
--- a/Engine/Source/Runtime/Renderer/Private/PostProcess/TemporalAA.cpp
+++ b/Engine/Source/Runtime/Renderer/Private/PostProcess/TemporalAA.cpp
@@ -125,11 +125,13 @@ END_SHADER_PARAMETER_STRUCT()
 BEGIN_SHADER_PARAMETER_STRUCT(FTAAHistoryTextures, )
 	SHADER_PARAMETER_RDG_TEXTURE_ARRAY(Texture2D, LowResTextures, [FTemporalAAHistory::kLowResRenderTargetCount])
 	SHADER_PARAMETER_RDG_TEXTURE_ARRAY(Texture2D, Textures, [FTemporalAAHistory::kRenderTargetCount])
+	SHADER_PARAMETER_RDG_TEXTURE_ARRAY(Texture2D, SuperResTextures, [FTemporalAAHistory::kSuperResRenderTargetCount])
 END_SHADER_PARAMETER_STRUCT()
 
 BEGIN_SHADER_PARAMETER_STRUCT(FTAAHistoryUAVs, )
 	SHADER_PARAMETER_RDG_TEXTURE_UAV_ARRAY(RWTexture2D, LowResTextures, [FTemporalAAHistory::kLowResRenderTargetCount])
 	SHADER_PARAMETER_RDG_TEXTURE_UAV_ARRAY(RWTexture2D, Textures, [FTemporalAAHistory::kRenderTargetCount])
+	SHADER_PARAMETER_RDG_TEXTURE_UAV_ARRAY(RWTexture2D, SuperResTextures, [FTemporalAAHistory::kSuperResRenderTargetCount])
 END_SHADER_PARAMETER_STRUCT()
 
 FTAAHistoryUAVs CreateUAVs(FRDGBuilder& GraphBuilder, const FTAAHistoryTextures& Textures)
@@ -143,6 +145,10 @@ FTAAHistoryUAVs CreateUAVs(FRDGBuilder& GraphBuilder, const FTAAHistoryTextures&
 	{
 		UAVs.Textures[i] = GraphBuilder.CreateUAV(Textures.Textures[i]);
 	}
+	for (int32 i = 0; i < Textures.SuperResTextures.Num(); i++)
+	{
+		UAVs.SuperResTextures[i] = GraphBuilder.CreateUAV(Textures.SuperResTextures[i]);
+	}
 	return UAVs;
 }
 
@@ -423,7 +429,6 @@ class FTAADetectInterferenceCS : public FTAAGen5Shader
 
 		SHADER_PARAMETER_STRUCT(FTAAHistoryUAVs, HistoryOutput)
 		SHADER_PARAMETER_RDG_TEXTURE_UAV(RWTexture2D, InterferenceWeightOutput)
-		SHADER_PARAMETER_RDG_TEXTURE_UAV(RWTexture2D, InterferenceSeedUpdateOutput)
 		SHADER_PARAMETER_RDG_TEXTURE_UAV(RWTexture2D, DebugOutput)
 	END_SHADER_PARAMETER_STRUCT()
 }; // class FTAADetectInterferenceCS
@@ -490,6 +495,26 @@ class FTAADilateRejectionCS : public FTAAGen5Shader
 	END_SHADER_PARAMETER_STRUCT()
 }; // class FTAADilateRejectionCS
 
+class FTAAUpdateSuperResHistoryCS : public FTAAGen5Shader
+{
+	DECLARE_GLOBAL_SHADER(FTAAUpdateSuperResHistoryCS);
+	SHADER_USE_PARAMETER_STRUCT(FTAAUpdateSuperResHistoryCS, FTAAGen5Shader);
+
+	BEGIN_SHADER_PARAMETER_STRUCT(FParameters, )
+		SHADER_PARAMETER_STRUCT_INCLUDE(FTAACommonParameters, CommonParameters)
+		SHADER_PARAMETER_RDG_TEXTURE(Texture2D, InputSceneColorTexture)
+
+		SHADER_PARAMETER_RDG_TEXTURE(Texture2D, DilatedVelocityTexture)
+		SHADER_PARAMETER_RDG_TEXTURE(Texture2D, ParallaxRejectionMaskTexture)
+
+		SHADER_PARAMETER_STRUCT(FScreenPassTextureViewportParameters, PrevHistoryInfo)
+		SHADER_PARAMETER_STRUCT(FTAAHistoryTextures, PrevHistory)
+
+		SHADER_PARAMETER_STRUCT(FTAAHistoryUAVs, HistoryOutput)
+		SHADER_PARAMETER_RDG_TEXTURE_UAV(RWTexture2D, DebugOutput)
+	END_SHADER_PARAMETER_STRUCT()
+}; // class FTAAUpdateSuperResHistoryCS
+
 class FTAAUpdateHistoryCS : public FTAAGen5Shader
 {
 	DECLARE_GLOBAL_SHADER(FTAAUpdateHistoryCS);
@@ -504,7 +529,6 @@ class FTAAUpdateHistoryCS : public FTAAGen5Shader
 		SHADER_PARAMETER_RDG_TEXTURE(Texture2D, DilatedVelocityTexture)
 		SHADER_PARAMETER_RDG_TEXTURE(Texture2D, ParallaxFactorTexture)
 		SHADER_PARAMETER_RDG_TEXTURE(Texture2D, ParallaxRejectionMaskTexture)
-		SHADER_PARAMETER_RDG_TEXTURE(Texture2D, InterferenceSeedUpdateTexture)
 
 		SHADER_PARAMETER(FVector, HistoryQuantizationError)
 		SHADER_PARAMETER(float, HistoryPreExposureCorrection)
@@ -527,6 +551,7 @@ IMPLEMENT_GLOBAL_SHADER(FTAAFilterFrequenciesCS, "/Engine/Private/TemporalAA/TAA
 IMPLEMENT_GLOBAL_SHADER(FTAACompareHistoryCS, "/Engine/Private/TemporalAA/TAACompareHistory.usf", "MainCS", SF_Compute);
 IMPLEMENT_GLOBAL_SHADER(FTAAPostfilterRejectionCS, "/Engine/Private/TemporalAA/TAAPostfilterRejection.usf", "MainCS", SF_Compute);
 IMPLEMENT_GLOBAL_SHADER(FTAADilateRejectionCS, "/Engine/Private/TemporalAA/TAADilateRejection.usf", "MainCS", SF_Compute);
+IMPLEMENT_GLOBAL_SHADER(FTAAUpdateSuperResHistoryCS, "/Engine/Private/TemporalAA/TAAUpdateSuperResHistory.usf", "MainCS", SF_Compute);
 IMPLEMENT_GLOBAL_SHADER(FTAAUpdateHistoryCS, "/Engine/Private/TemporalAA/TAAUpdateHistory.usf", "MainCS", SF_Compute);
 
 #if COMPILE_TAA_DEBUG_PASSES
@@ -1276,6 +1301,11 @@ static void AddGen5MainTemporalAAPasses(
 		{
 			PrevHistory.Textures[i] = BlackDummy;
 		}
+
+		for (int32 i = 0; i < PrevHistory.SuperResTextures.Num(); i++)
+		{
+			PrevHistory.SuperResTextures[i] = BlackDummy;
+		}
 	}
 	else
 	{
@@ -1307,34 +1337,61 @@ static void AddGen5MainTemporalAAPasses(
 			}
 		}
 
+		for (int32 i = 0; i < PrevHistory.SuperResTextures.Num(); i++)
+		{
+			if (InputHistory.SuperResRT[i].IsValid())
+			{
+				PrevHistory.SuperResTextures[i] = GraphBuilder.RegisterExternalTexture(InputHistory.SuperResRT[i]);
+			}
+			else
+			{
+				PrevHistory.SuperResTextures[i] = BlackDummy;
+			}
+		}
+
 		// InputHistory.SafeRelease(); TODO
 	}
 
 
 	FTAAHistoryTextures History;
 	{
-		FRDGTextureDesc Desc = FRDGTextureDesc::Create2D(
-			HistoryExtent,
-			(CVarTAAR11G11B10History.GetValueOnRenderThread() != 0) ? PF_FloatR11G11B10 : PF_FloatRGBA,
-			FClearValueBinding::None,
-			TexCreate_ShaderResource | TexCreate_UAV);
-
 		FRDGTextureDesc LowResDesc = FRDGTextureDesc::Create2D(
 			InputExtent,
 			PF_R8,
 			FClearValueBinding::None,
 			TexCreate_ShaderResource | TexCreate_UAV);
 
+		History.LowResTextures[0] = GraphBuilder.CreateTexture(LowResDesc, TEXT("TAA.History.LowResMetadata[0]"));
+	}
+
+	{
+		FRDGTextureDesc Desc = FRDGTextureDesc::Create2D(
+			HistoryExtent,
+			(CVarTAAR11G11B10History.GetValueOnRenderThread() != 0) ? PF_FloatR11G11B10 : PF_FloatRGBA,
+			FClearValueBinding::None,
+			TexCreate_ShaderResource | TexCreate_UAV);
+
 		History.Textures[0] = GraphBuilder.CreateTexture(Desc, TEXT("TAA.History.LowFrequencies"));
 		History.Textures[1] = GraphBuilder.CreateTexture(Desc, TEXT("TAA.History.HighFrequencies"));
 
 		Desc.Format = PF_R8G8;
 		History.Textures[2] = GraphBuilder.CreateTexture(Desc, TEXT("TAA.History.Metadata"));
 
-		Desc.Format = PF_R32_UINT;
+		Desc.Format = PF_R16_UINT;
 		History.Textures[3] = GraphBuilder.CreateTexture(Desc, TEXT("TAA.History.SubpixelInfo"));
+	}
 
-		History.LowResTextures[0] = GraphBuilder.CreateTexture(LowResDesc, TEXT("TAA.History.LowResMetadata"));
+	{
+		FRDGTextureDesc SuperResDesc = FRDGTextureDesc::Create2D(
+			HistoryExtent * 2,
+			PF_R16_UINT,
+			FClearValueBinding::None,
+			TexCreate_ShaderResource | TexCreate_UAV);
+
+		History.SuperResTextures[0] = GraphBuilder.CreateTexture(SuperResDesc, TEXT("TAA.History.SuperResMetadata[0]"));
+
+		SuperResDesc.Format = PF_R8_UINT;
+		History.SuperResTextures[1] = GraphBuilder.CreateTexture(SuperResDesc, TEXT("TAA.History.SuperResMetadata[1]"));
 	}
 
 	// Decimate input to flicker at same frequency as input.
@@ -1445,7 +1502,6 @@ static void AddGen5MainTemporalAAPasses(
 	// TODO(G5TAA): Could sample the interference seed in the DilateVelocity and detect interference in the decimate.
 	TStaticArray<bool, FTemporalAAHistory::kLowResRenderTargetCount> ExtractLowResHistoryTexture;
 	FRDGTextureRef InterferenceWeightTexture;
-	FRDGTextureRef InterferenceSeedUpdateTexture;
 	if (bEnableInterferenceHeuristic)
 	{
 		{
@@ -1457,7 +1513,6 @@ static void AddGen5MainTemporalAAPasses(
 				/* InFlags = */ TexCreate_ShaderResource | TexCreate_UAV);
 
 			InterferenceWeightTexture = GraphBuilder.CreateTexture(Desc, TEXT("TAA.Interference.Weight"));
-			InterferenceSeedUpdateTexture = GraphBuilder.CreateTexture(Desc, TEXT("TAA.Interference.SeedUpdateRequest"));
 		}
 
 		FTAADetectInterferenceCS::FParameters* PassParameters = GraphBuilder.AllocParameters<FTAADetectInterferenceCS::FParameters>();
@@ -1474,7 +1529,6 @@ static void AddGen5MainTemporalAAPasses(
 
 		PassParameters->HistoryOutput = CreateUAVs(GraphBuilder, History);
 		PassParameters->InterferenceWeightOutput = GraphBuilder.CreateUAV(InterferenceWeightTexture);
-		PassParameters->InterferenceSeedUpdateOutput = GraphBuilder.CreateUAV(InterferenceSeedUpdateTexture);
 		PassParameters->DebugOutput = CreateDebugUAV(LowFrequencyExtent, TEXT("Debug.TAA.DetectInterference"));
 
 		TShaderMapRef<FTAADetectInterferenceCS> ComputeShader(View.ShaderMap);
@@ -1516,7 +1570,6 @@ static void AddGen5MainTemporalAAPasses(
 
 		// TODO(G5TAA): Shader permutation.
 		InterferenceWeightTexture = WhiteDummy;
-		InterferenceSeedUpdateTexture = BlackDummy;
 	}
 
 	// Reject the history with frequency decomposition.
@@ -1642,6 +1695,59 @@ static void AddGen5MainTemporalAAPasses(
 			FComputeShaderUtils::GetGroupCount(RejectionRect.Size(), 8));
 	}
 
+	TStaticArray<bool, FTemporalAAHistory::kSuperResRenderTargetCount> ExtractSuperResHistoryTexture;
+	if (bEnableInterferenceHeuristic)
+	{
+		FTAAUpdateSuperResHistoryCS::FParameters* PassParameters = GraphBuilder.AllocParameters<FTAAUpdateSuperResHistoryCS::FParameters>();
+		PassParameters->CommonParameters = CommonParameters;
+		PassParameters->InputSceneColorTexture = PassInputs.SceneColorTexture;
+		PassParameters->DilatedVelocityTexture = DilatedVelocityTexture;
+		PassParameters->ParallaxRejectionMaskTexture = ParallaxRejectionMaskTexture;
+
+		PassParameters->PrevHistoryInfo = PrevHistoryInfo;
+		PassParameters->PrevHistory = PrevHistory;
+
+		PassParameters->HistoryOutput = CreateUAVs(GraphBuilder, History);
+		PassParameters->DebugOutput = CreateDebugUAV(HistoryExtent * 2, TEXT("Debug.TAA.UpdateSuperResHistory"));
+
+		TShaderMapRef<FTAAUpdateSuperResHistoryCS> ComputeShader(View.ShaderMap);
+		ClearUnusedGraphResources(ComputeShader, PassParameters);
+
+		for (int32 i = 0; i < PrevHistory.SuperResTextures.Num(); i++)
+		{
+			bool bNeedsExtractForNextFrame = PassParameters->PrevHistory.SuperResTextures[i] != nullptr;
+			bool bPrevFrameIsntAvailable = PassParameters->PrevHistory.SuperResTextures[i] == BlackDummy;
+			bool bOutputHistory = PassParameters->HistoryOutput.SuperResTextures[i] != nullptr;
+
+			ExtractSuperResHistoryTexture[i] = bNeedsExtractForNextFrame;
+
+			if (bPrevFrameIsntAvailable && !PassParameters->CommonParameters.bCameraCut)
+			{
+				//ensureMsgf(false, TEXT("Shaders read PrevHistory[%d] but doesn't write HistoryOutput[%d]"), i, i);
+				PassParameters->CommonParameters.bCameraCut = true;
+			}
+
+			if (bOutputHistory && !bNeedsExtractForNextFrame)
+			{
+				ensureMsgf(false, TEXT("Shaders write HistoryOutput[%d] but doesn't read PrevHistory[%d]"), i, i);
+			}
+		}
+
+		FComputeShaderUtils::AddPass(
+			GraphBuilder,
+			RDG_EVENT_NAME("TAA UpdateSuperResHistory %dx%d", HistorySize.X * 2, HistorySize.Y * 2),
+			ComputeShader,
+			PassParameters,
+			FComputeShaderUtils::GetGroupCount(HistorySize * 2, 8));
+	}
+	else
+	{
+		for (int32 i = 0; i < ExtractSuperResHistoryTexture.Num(); i++)
+		{
+			ExtractSuperResHistoryTexture[i] = false;
+		}
+	}
+
 	TStaticArray<bool, FTemporalAAHistory::kRenderTargetCount> ExtractHistoryTexture;
 	FRDGTextureRef SceneColorOutputTexture;
 	{
@@ -1665,7 +1771,6 @@ static void AddGen5MainTemporalAAPasses(
 		PassParameters->DilatedVelocityTexture = DilatedVelocityTexture;
 		PassParameters->ParallaxFactorTexture = ParallaxFactorTexture;
 		PassParameters->ParallaxRejectionMaskTexture = ParallaxRejectionMaskTexture;
-		PassParameters->InterferenceSeedUpdateTexture = InterferenceSeedUpdateTexture;
 
 		PassParameters->HistoryQuantizationError = ComputePixelFormatQuantizationError(History.Textures[0]->Desc.Format);
 		PassParameters->HistoryPreExposureCorrection = View.PreExposure / View.PrevViewInfo.SceneColorPreExposure;
@@ -1751,6 +1856,14 @@ static void AddGen5MainTemporalAAPasses(
 			}
 		}
 
+		for (int32 i = 0; i < History.SuperResTextures.Num(); i++)
+		{
+			if (ExtractSuperResHistoryTexture[i])
+			{
+				GraphBuilder.QueueTextureExtraction(History.SuperResTextures[i], &OutputHistory->SuperResRT[i]);
+			}
+		}
+
 		OutputHistory->ViewportRect = FIntRect(FIntPoint(0, 0), HistorySize);
 		OutputHistory->ReferenceBufferSize = HistoryExtent;
 	}
diff --git a/Engine/Source/Runtime/Renderer/Private/SceneRendering.h b/Engine/Source/Runtime/Renderer/Private/SceneRendering.h
--- a/Engine/Source/Runtime/Renderer/Private/SceneRendering.h
+++ b/Engine/Source/Runtime/Renderer/Private/SceneRendering.h
@@ -816,10 +816,14 @@ struct FTemporalAAHistory
 	// Number of render target in the history.
 	static constexpr int32 kRenderTargetCount = 4;
 
+	// Number of render target in the history.
+	static constexpr int32 kSuperResRenderTargetCount = 2;
+
 	// Render targets holding's pixel history.
 	//  scene color's RGBA are in OutputRT[0].
 	TStaticArray<TRefCountPtr<IPooledRenderTarget>, kLowResRenderTargetCount> LowResRT;
 	TStaticArray<TRefCountPtr<IPooledRenderTarget>, kRenderTargetCount> RT;
+	TStaticArray<TRefCountPtr<IPooledRenderTarget>, kSuperResRenderTargetCount> SuperResRT;
 
 	// Reference size of RT. Might be different than RT's actual size to handle down res.
 	FIntPoint ReferenceBufferSize;
@@ -839,6 +843,10 @@ struct FTemporalAAHistory
 		{
 			RT[i].SafeRelease();
 		}
+		for (int32 i = 0; i < SuperResRT.Num(); i++)
+		{
+			SuperResRT[i].SafeRelease();
+		}
 	}
 
 	bool IsValid() const
