Fixes ensure due to temporal upsampling when using ray tracing debugging view mode

diff --git a/Engine/Source/Runtime/Renderer/Private/PostProcess/PostProcessing.cpp b/Engine/Source/Runtime/Renderer/Private/PostProcess/PostProcessing.cpp
--- a/Engine/Source/Runtime/Renderer/Private/PostProcess/PostProcessing.cpp
+++ b/Engine/Source/Runtime/Renderer/Private/PostProcess/PostProcessing.cpp
@@ -1086,8 +1086,6 @@ void AddDebugViewPostProcessingPasses(FRDGBuilder& GraphBuilder, const FViewInfo
 	const FScreenPassTexture SceneDepth(SceneTextureParameters.SceneDepthTexture, PrimaryViewRect);
 	FScreenPassTexture SceneColor((*Inputs.SceneTextures)->SceneColorTexture, PrimaryViewRect);
 
-	ensure(View.PrimaryScreenPercentageMethod != EPrimaryScreenPercentageMethod::TemporalUpscale);
-
 	// Some view modes do not actually output a color so they should not be tonemapped.
 	const bool bTonemapAfter = View.Family->EngineShowFlags.RayTracingDebug;
 	const bool bTonemapBefore = !bTonemapAfter && !View.Family->EngineShowFlags.ShaderComplexity;
@@ -1119,7 +1117,7 @@ void AddDebugViewPostProcessingPasses(FRDGBuilder& GraphBuilder, const FViewInfo
 	PassSequence.SetEnabled(EPass::Visualize, true);
 	PassSequence.SetEnabled(EPass::TonemapAfter, bTonemapAfter);
 	PassSequence.SetEnabled(EPass::SelectionOutline, GIsEditor);
-	PassSequence.SetEnabled(EPass::PrimaryUpscale, View.ViewRect.Size() != View.GetSecondaryViewRectSize());
+	PassSequence.SetEnabled(EPass::PrimaryUpscale, View.ViewRect.Size() != View.GetSecondaryViewRectSize() && View.PrimaryScreenPercentageMethod != EPrimaryScreenPercentageMethod::TemporalUpscale);
 	PassSequence.SetEnabled(EPass::SecondaryUpscale, View.RequiresSecondaryUpscale());
 	PassSequence.Finalize();
 
@@ -1190,12 +1188,17 @@ void AddDebugViewPostProcessingPasses(FRDGBuilder& GraphBuilder, const FViewInfo
 			Parameters.SceneDepthTexture = SceneTextureParameters.SceneDepthTexture;
 			Parameters.SceneVelocityTexture = SceneTextureParameters.GBufferVelocityTexture;
 			Parameters.SceneColorInput = SceneColor.Texture;
+			Parameters.Pass = View.PrimaryScreenPercentageMethod == EPrimaryScreenPercentageMethod::TemporalUpscale
+				? ETAAPassConfig::MainUpsampling
+				: ETAAPassConfig::Main;
+			Parameters.SetupViewRect(View);
 
 			const FTemporalAAHistory& InputHistory = View.PrevViewInfo.TemporalAAHistory;
 			FTemporalAAHistory* OutputHistory = &View.ViewState->PrevFrameViewInfo.TemporalAAHistory;
 
 			FTAAOutputs Outputs = AddTemporalAAPass(GraphBuilder, View, Parameters, InputHistory, OutputHistory);
 			SceneColor.Texture = Outputs.SceneColor;
+			SceneColor.ViewRect = Parameters.OutputViewRect;
 
 			break;
 		}
