Removed FRDGTexture::GetPooledRenderTarget() usage and deprecated in prepration for transient allocator support.

diff --git a/Engine/Plugins/Runtime/TextureShare/Source/TextureShare/Private/TextureShareModule.cpp b/Engine/Plugins/Runtime/TextureShare/Source/TextureShare/Private/TextureShareModule.cpp
--- a/Engine/Plugins/Runtime/TextureShare/Source/TextureShare/Private/TextureShareModule.cpp
+++ b/Engine/Plugins/Runtime/TextureShare/Source/TextureShare/Private/TextureShareModule.cpp
@@ -253,7 +253,7 @@ bool FTextureShareModule::SendSceneContext_RenderThread(FRDGBuilder& GraphBuilde
 					ERDGPassFlags::Copy | ERDGPassFlags::NeverCull,
 					[this, Texture, ShareItem, ShareName](FRHICommandListImmediate& RHICmdList)
 				{
-					SendTexture_RenderThread(RHICmdList, ShareItem, ShareName, Texture->GetPooledRenderTarget());
+					SendTexture_RenderThread(RHICmdList, ShareItem, ShareName, Texture->GetRHI());
 				});
 			};
 
@@ -304,16 +304,6 @@ bool FTextureShareModule::SendPostRender_RenderThread(FRHICommandListImmediate&
 	return false;
 }
 
-bool FTextureShareModule::SendTexture_RenderThread(FRHICommandListImmediate& RHICmdList, const TSharedPtr<ITextureShareItem>& ShareItem, const FString& TextureName, const TRefCountPtr<IPooledRenderTarget>& PooledRenderTargetRef)
-{
-	if (PooledRenderTargetRef.IsValid())
-	{
-		const FTexture2DRHIRef& RHITexture = (const FTexture2DRHIRef&)PooledRenderTargetRef->GetRenderTargetItem().ShaderResourceTexture;
-		return SendTexture_RenderThread(RHICmdList, ShareItem, TextureName, RHITexture.GetReference());
-	}
-	return false;
-}
-
 bool FTextureShareModule::SendTexture_RenderThread(FRHICommandListImmediate& RHICmdList, const TSharedPtr<ITextureShareItem>& ShareItem, const FString& InTextureName, FRHITexture* RHITexture, const FIntRect* SrcTextureRect)
 {
 	if (ShareItem.IsValid() && ShareItem->IsValid())
diff --git a/Engine/Plugins/Runtime/TextureShare/Source/TextureShare/Private/TextureShareModule.h b/Engine/Plugins/Runtime/TextureShare/Source/TextureShare/Private/TextureShareModule.h
--- a/Engine/Plugins/Runtime/TextureShare/Source/TextureShare/Private/TextureShareModule.h
+++ b/Engine/Plugins/Runtime/TextureShare/Source/TextureShare/Private/TextureShareModule.h
@@ -45,9 +45,6 @@ public:
 	void OnResolvedSceneColor_RenderThread(FRDGBuilder& GraphBuilder, const FSceneTextures& SceneTextures, class FSceneViewFamily& ViewFamily);
 	void OnPostRenderViewFamily_RenderThread(FRHICommandListImmediate& RHICmdList, class FSceneViewFamily& ViewFamily);
 
-protected:
-	bool SendTexture_RenderThread(FRHICommandListImmediate& RHICmdList, const TSharedPtr<ITextureShareItem>& ShareItem, const FString& TextureName, const TRefCountPtr<IPooledRenderTarget>& PooledRenderTargetRef);
-
 protected:
 	bool SendSceneContext_RenderThread(FRDGBuilder& GraphBuilder, const FString& ShareName, const FSceneTextures& SceneTextures, class FSceneViewFamily& ViewFamily);
 	bool SendPostRender_RenderThread(FRHICommandListImmediate& RHICmdList, const FString& ShareName, class FSceneViewFamily& ViewFamily);
diff --git a/Engine/Source/Runtime/RenderCore/Public/RenderGraphResources.h b/Engine/Source/Runtime/RenderCore/Public/RenderGraphResources.h
--- a/Engine/Source/Runtime/RenderCore/Public/RenderGraphResources.h
+++ b/Engine/Source/Runtime/RenderCore/Public/RenderGraphResources.h
@@ -604,10 +604,10 @@ public:
 	//! The following methods may only be called during pass execution.
 
 	/** Returns the allocated pooled render target. */
+	UE_DEPRECATED(5.0, "Accessing the underlying pooled render target has been deprecated. Use GetRHI() instead.")
 	IPooledRenderTarget* GetPooledRenderTarget() const
 	{
 		IF_RDG_ENABLE_DEBUG(ValidateRHIAccess());
-		check(PooledRenderTarget);
 		return PooledRenderTarget;
 	}
 
@@ -689,7 +689,7 @@ private:
 	IPooledRenderTarget* PooledRenderTarget = nullptr;
 
 	/** The assigned pooled texture to use during execution. Never reset. */
-	FRDGPooledTexture* PooledTexture = nullptr;
+	TRefCountPtr<FRDGPooledTexture> PooledTexture;
 
 	/** Cached state pointer from the pooled texture. */
 	FRDGTextureSubresourceState* State = nullptr;
diff --git a/Engine/Source/Runtime/RenderCore/Public/RenderGraphUtils.h b/Engine/Source/Runtime/RenderCore/Public/RenderGraphUtils.h
--- a/Engine/Source/Runtime/RenderCore/Public/RenderGraphUtils.h
+++ b/Engine/Source/Runtime/RenderCore/Public/RenderGraphUtils.h
@@ -48,9 +48,12 @@ inline FRHITexture* TryGetRHI(FRDGTextureRef Texture)
 }
 
 /** Returns the pooled render target from an RDG texture if it exists, or null otherwise. */
+UE_DEPRECATED(5.0, "Accessing the underlying pooled render target has been deprecated. Use TryGetRHI() instead.")
 inline IPooledRenderTarget* TryGetPooledRenderTarget(FRDGTextureRef Texture)
 {
+PRAGMA_DISABLE_DEPRECATION_WARNINGS
 	return Texture ? Texture->GetPooledRenderTarget() : nullptr;
+PRAGMA_ENABLE_DEPRECATION_WARNINGS
 }
 
 inline FRenderTargetBindingSlots GetRenderTargetBindings(ERenderTargetLoadAction ColorLoadAction, TArrayView<FRDGTextureRef> ColorTextures)
diff --git a/Engine/Source/Runtime/Renderer/Private/CapsuleShadowRendering.cpp b/Engine/Source/Runtime/Renderer/Private/CapsuleShadowRendering.cpp
--- a/Engine/Source/Runtime/Renderer/Private/CapsuleShadowRendering.cpp
+++ b/Engine/Source/Runtime/Renderer/Private/CapsuleShadowRendering.cpp
@@ -256,7 +256,8 @@ public:
 		FScene* Scene,
 		const FSceneView& View, 
 		const FLightSceneInfo* LightSceneInfo,
-		const FSceneRenderTargetItem& OutputTexture, 
+		FRHITexture* OutputTexture,
+		FRHIUnorderedAccessView* OutputTextureUAV, 
 		FIntPoint TileDimensionsValue,
 		const FRWBuffer* TileIntersectionCountsBuffer,
 		FVector2D NumGroupsValue,
@@ -282,12 +283,12 @@ public:
 		if (ShadowingType == ShapeShadow_MovableSkylightTiledCulling)
 		{
 			check(!ShadowFactors.IsBound());
-			BentNormalTexture.SetTexture(RHICmdList, ShaderRHI, OutputTexture.ShaderResourceTexture, OutputTexture.UAV);
+			BentNormalTexture.SetTexture(RHICmdList, ShaderRHI, OutputTexture, OutputTextureUAV);
 		}
 		else
 		{
 			check(!BentNormalTexture.IsBound());
-			ShadowFactors.SetTexture(RHICmdList, ShaderRHI, OutputTexture.ShaderResourceTexture, OutputTexture.UAV);
+			ShadowFactors.SetTexture(RHICmdList, ShaderRHI, OutputTexture, OutputTextureUAV);
 		}
 		
 		if (TileIntersectionCountsBuffer)
@@ -557,13 +558,13 @@ public:
 		OutputtingToLightAttenuation.Bind(Initializer.ParameterMap,TEXT("OutputtingToLightAttenuation"));
 	}
 
-	void SetParameters(FRHICommandList& RHICmdList, const FSceneView& View, const FIntRect& ScissorRect, IPooledRenderTarget* ShadowFactorsTextureValue, bool bOutputtingToLightAttenuation)
+	void SetParameters(FRHICommandList& RHICmdList, const FSceneView& View, const FIntRect& ScissorRect, FRHITexture* ShadowFactorsTextureRHI, bool bOutputtingToLightAttenuation)
 	{
 		FRHIPixelShader* ShaderRHI = RHICmdList.GetBoundPixelShader();
 
 		FGlobalShader::SetParameters<FViewUniformShaderParameters>(RHICmdList, ShaderRHI, View.ViewUniformBuffer);
 
-		SetTextureParameter(RHICmdList, ShaderRHI, ShadowFactorsTexture, ShadowFactorsSampler, TStaticSamplerState<SF_Bilinear>::GetRHI(), ShadowFactorsTextureValue->GetRenderTargetItem().ShaderResourceTexture);
+		SetTextureParameter(RHICmdList, ShaderRHI, ShadowFactorsTexture, ShadowFactorsSampler, TStaticSamplerState<SF_Bilinear>::GetRHI(), ShadowFactorsTextureRHI);
 	
 		SetShaderValue(RHICmdList, ShaderRHI, ScissorRectMinAndSize, FIntRect(ScissorRect.Min, ScissorRect.Size()));
 		SetShaderValue(RHICmdList, ShaderRHI, OutputtingToLightAttenuation, bOutputtingToLightAttenuation ? 1.0f : 0.0f);
@@ -600,6 +601,7 @@ void AllocateCapsuleTileIntersectionCountsBuffer(FIntPoint GroupSize, FSceneView
 // TODO(RDG) Move these into the shader FParameters.
 BEGIN_SHADER_PARAMETER_STRUCT(FTiledCapsuleShadowParameters, )
 	SHADER_PARAMETER_RDG_UNIFORM_BUFFER(FSceneTextureUniformParameters, SceneTextures)
+	SHADER_PARAMETER_RDG_TEXTURE_UAV(RWTexture2D, RayTracedShadowsUAV)
 	RDG_TEXTURE_ACCESS(RayTracedShadows, ERHIAccess::UAVCompute)
 END_SHADER_PARAMETER_STRUCT()
 
@@ -725,13 +727,14 @@ bool FDeferredShadingSceneRenderer::RenderCapsuleDirectShadows(
 			{
 				FTiledCapsuleShadowParameters* PassParameters = GraphBuilder.AllocParameters<FTiledCapsuleShadowParameters>();
 				PassParameters->RayTracedShadows = RayTracedShadowsRT;
+				PassParameters->RayTracedShadowsUAV = GraphBuilder.CreateUAV(RayTracedShadowsRT);
 				PassParameters->SceneTextures = SceneTexturesUniformBuffer;
 
 				GraphBuilder.AddPass(
 					RDG_EVENT_NAME("TiledCapsuleShadowing"),
 					PassParameters,
 					ERDGPassFlags::Compute,
-					[this, &View, &LightSceneInfo, RayTracedShadowsRT, GroupSize, ScissorRect, bDirectionalLight, NumCapsuleShapeData](FRHIComputeCommandList& RHICmdList)
+					[this, &View, &LightSceneInfo, RayTracedShadows = PassParameters->RayTracedShadows.GetTexture(), RayTracedShadowsUAV = PassParameters->RayTracedShadowsUAV, GroupSize, ScissorRect, bDirectionalLight, NumCapsuleShapeData](FRHIComputeCommandList& RHICmdList)
 				{
 					if (bDirectionalLight)
 					{
@@ -743,7 +746,8 @@ bool FDeferredShadingSceneRenderer::RenderCapsuleDirectShadows(
 							Scene,
 							View,
 							&LightSceneInfo,
-							RayTracedShadowsRT->GetPooledRenderTarget()->GetRenderTargetItem(),
+							RayTracedShadows->GetRHI(),
+							RayTracedShadowsUAV->GetRHI(),
 							GroupSize,
 							&View.ViewState->CapsuleTileIntersectionCountsBuffer,
 							FVector2D(GroupSize.X, GroupSize.Y),
@@ -770,7 +774,8 @@ bool FDeferredShadingSceneRenderer::RenderCapsuleDirectShadows(
 							Scene,
 							View,
 							&LightSceneInfo,
-							RayTracedShadowsRT->GetPooledRenderTarget()->GetRenderTargetItem(),
+							RayTracedShadows->GetRHI(),
+							RayTracedShadowsUAV->GetRHI(),
 							GroupSize,
 							&View.ViewState->CapsuleTileIntersectionCountsBuffer,
 							FVector2D(GroupSize.X, GroupSize.Y),
@@ -830,7 +835,7 @@ bool FDeferredShadingSceneRenderer::RenderCapsuleDirectShadows(
 						SetGraphicsPipelineState(RHICmdList, GraphicsPSOInit);
 
 						VertexShader->SetParameters(RHICmdList, View, GroupSize, ScissorRect, View.ViewState->CapsuleTileIntersectionCountsBuffer);
-						PixelShader->SetParameters(RHICmdList, View, ScissorRect, RayTracedShadowsRT->GetPooledRenderTarget(), true);
+						PixelShader->SetParameters(RHICmdList, View, ScissorRect, RayTracedShadowsRT->GetRHI(), true);
 					}
 					else
 					{
@@ -843,7 +848,7 @@ bool FDeferredShadingSceneRenderer::RenderCapsuleDirectShadows(
 						SetGraphicsPipelineState(RHICmdList, GraphicsPSOInit);
 
 						VertexShader->SetParameters(RHICmdList, View, GroupSize, ScissorRect, View.ViewState->CapsuleTileIntersectionCountsBuffer);
-						PixelShader->SetParameters(RHICmdList, View, ScissorRect, RayTracedShadowsRT->GetPooledRenderTarget(), true);
+						PixelShader->SetParameters(RHICmdList, View, ScissorRect, RayTracedShadowsRT->GetRHI(), true);
 					}
 
 					RHICmdList.SetStreamSource(0, GTileTexCoordVertexBuffer.VertexBufferRHI, 0);
@@ -1248,13 +1253,14 @@ void FDeferredShadingSceneRenderer::RenderIndirectCapsuleShadows(FRDGBuilder& Gr
 			{
 				FTiledCapsuleShadowParameters* PassParameters = GraphBuilder.AllocParameters<FTiledCapsuleShadowParameters>();
 				PassParameters->RayTracedShadows = RayTracedShadowsRT;
+				PassParameters->RayTracedShadowsUAV = GraphBuilder.CreateUAV(RayTracedShadowsRT);
 				PassParameters->SceneTextures = SceneTextures.UniformBuffer;
 
 				GraphBuilder.AddPass(
 					RDG_EVENT_NAME("TiledCapsuleShadowing %u capsules among %u meshes", NumCapsuleShapes, NumMeshesWithCapsules),
 					PassParameters,
 					ERDGPassFlags::Compute,
-					[this, &View, RayTracedShadowsRT, IndirectShadowLightDirectionSRV, GroupSize, ScissorRect, NumCapsuleShapes, NumMeshDistanceFieldCasters](FRHIComputeCommandList& RHICmdList)
+					[this, &View, RayTracedShadows = PassParameters->RayTracedShadows.GetTexture(), RayTracedShadowsUAV = PassParameters->RayTracedShadowsUAV, IndirectShadowLightDirectionSRV, GroupSize, ScissorRect, NumCapsuleShapes, NumMeshDistanceFieldCasters](FRHIComputeCommandList& RHICmdList)
 				{
 					TShaderRef<TCapsuleShadowingBaseCS<ShapeShadow_IndirectTiledCulling>> ComputeShaderBase;
 
@@ -1282,7 +1288,8 @@ void FDeferredShadingSceneRenderer::RenderIndirectCapsuleShadows(FRDGBuilder& Gr
 						Scene,
 						View,
 						NULL,
-						RayTracedShadowsRT->GetPooledRenderTarget()->GetRenderTargetItem(),
+						RayTracedShadows->GetRHI(),
+						RayTracedShadowsUAV->GetRHI(),
 						GroupSize,
 						&View.ViewState->CapsuleTileIntersectionCountsBuffer,
 						FVector2D(GroupSize.X, GroupSize.Y),
@@ -1351,7 +1358,7 @@ void FDeferredShadingSceneRenderer::RenderIndirectCapsuleShadows(FRDGBuilder& Gr
 							GraphicsPSOInit.BoundShaderState.PixelShaderRHI = PixelShader.GetPixelShader();
 							SetGraphicsPipelineState(RHICmdList, GraphicsPSOInit);
 							VertexShader->SetParameters(RHICmdList, View, GroupSize, ScissorRect, View.ViewState->CapsuleTileIntersectionCountsBuffer);
-							PixelShader->SetParameters(RHICmdList, View, ScissorRect, RayTracedShadowsRT->GetPooledRenderTarget(), false);
+							PixelShader->SetParameters(RHICmdList, View, ScissorRect, RayTracedShadowsRT->GetRHI(), false);
 						}
 						else
 						{
@@ -1359,7 +1366,7 @@ void FDeferredShadingSceneRenderer::RenderIndirectCapsuleShadows(FRDGBuilder& Gr
 							GraphicsPSOInit.BoundShaderState.PixelShaderRHI = PixelShader.GetPixelShader();
 							SetGraphicsPipelineState(RHICmdList, GraphicsPSOInit);
 							VertexShader->SetParameters(RHICmdList, View, GroupSize, ScissorRect, View.ViewState->CapsuleTileIntersectionCountsBuffer);
-							PixelShader->SetParameters(RHICmdList, View, ScissorRect, RayTracedShadowsRT->GetPooledRenderTarget(), false);
+							PixelShader->SetParameters(RHICmdList, View, ScissorRect, RayTracedShadowsRT->GetRHI(), false);
 						}
 					}
 					else
@@ -1370,7 +1377,7 @@ void FDeferredShadingSceneRenderer::RenderIndirectCapsuleShadows(FRDGBuilder& Gr
 							GraphicsPSOInit.BoundShaderState.PixelShaderRHI = PixelShader.GetPixelShader();
 							SetGraphicsPipelineState(RHICmdList, GraphicsPSOInit);
 							VertexShader->SetParameters(RHICmdList, View, GroupSize, ScissorRect, View.ViewState->CapsuleTileIntersectionCountsBuffer);
-							PixelShader->SetParameters(RHICmdList, View, ScissorRect, RayTracedShadowsRT->GetPooledRenderTarget(), false);
+							PixelShader->SetParameters(RHICmdList, View, ScissorRect, RayTracedShadowsRT->GetRHI(), false);
 						}
 						else
 						{
@@ -1378,7 +1385,7 @@ void FDeferredShadingSceneRenderer::RenderIndirectCapsuleShadows(FRDGBuilder& Gr
 							GraphicsPSOInit.BoundShaderState.PixelShaderRHI = PixelShader.GetPixelShader();
 							SetGraphicsPipelineState(RHICmdList, GraphicsPSOInit);
 							VertexShader->SetParameters(RHICmdList, View, GroupSize, ScissorRect, View.ViewState->CapsuleTileIntersectionCountsBuffer);
-							PixelShader->SetParameters(RHICmdList, View, ScissorRect, RayTracedShadowsRT->GetPooledRenderTarget(), false);
+							PixelShader->SetParameters(RHICmdList, View, ScissorRect, RayTracedShadowsRT->GetRHI(), false);
 						}
 					}
 
@@ -1429,6 +1436,7 @@ BEGIN_SHADER_PARAMETER_STRUCT(FCapsuleShadowsForMovableSkylightParameters, )
 	SHADER_PARAMETER_RDG_UNIFORM_BUFFER(FSceneTextureUniformParameters, SceneTextures)
 	RDG_TEXTURE_ACCESS(BentNormalInput, ERHIAccess::SRVCompute)
 	RDG_TEXTURE_ACCESS(BentNormalOutput, ERHIAccess::UAVCompute)
+	SHADER_PARAMETER_RDG_TEXTURE_UAV(RWTexture2D, BentNormalOutputUAV)
 END_SHADER_PARAMETER_STRUCT()
 
 void FDeferredShadingSceneRenderer::RenderCapsuleShadowsForMovableSkylight(
@@ -1492,17 +1500,16 @@ void FDeferredShadingSceneRenderer::RenderCapsuleShadowsForMovableSkylight(
 							auto* PassParameters = GraphBuilder.AllocParameters< FCapsuleShadowsForMovableSkylightParameters>();
 							PassParameters->BentNormalInput = BentNormalOutput;
 							PassParameters->BentNormalOutput = NewBentNormal;
+							PassParameters->BentNormalOutputUAV = GraphBuilder.CreateUAV(NewBentNormal);
 							PassParameters->SceneTextures = SceneTexturesUniformBuffer;
 
 							GraphBuilder.AddPass(
 								RDG_EVENT_NAME("TiledCapsuleShadowing % u capsules among % u meshes", NumCapsuleShapes, NumMeshesWithCapsules),
 								PassParameters,
 								ERDGPassFlags::Compute,
-								[this, &View, NewBentNormal, BentNormalOutput, GroupSizeX, GroupSizeY, NumCapsuleShapes, NumMeshDistanceFieldCasters, ScissorRect, IndirectShadowLightDirectionSRV]
+								[this, &View, NewBentNormal, PassParameters, GroupSizeX, GroupSizeY, NumCapsuleShapes, NumMeshDistanceFieldCasters, ScissorRect, IndirectShadowLightDirectionSRV]
 								(FRHICommandList& RHICmdList)
 							{
-								FSceneRenderTargetItem& RayTracedShadowsRTI = NewBentNormal->GetPooledRenderTarget()->GetRenderTargetItem();
-
 								{
 									TShaderRef<TCapsuleShadowingBaseCS<ShapeShadow_MovableSkylightTiledCulling>> ComputeShaderBase;
 									if (NumCapsuleShapes > 0 && NumMeshDistanceFieldCasters > 0)
@@ -1528,7 +1535,8 @@ void FDeferredShadingSceneRenderer::RenderCapsuleShadowsForMovableSkylight(
 										Scene,
 										View,
 										NULL,
-										RayTracedShadowsRTI,
+										PassParameters->BentNormalOutput->GetRHI(),
+										PassParameters->BentNormalOutputUAV->GetRHI(),
 										FIntPoint(GroupSizeX, GroupSizeY),
 										NULL,
 										FVector2D(GroupSizeX, GroupSizeY),
@@ -1540,7 +1548,7 @@ void FDeferredShadingSceneRenderer::RenderCapsuleShadowsForMovableSkylight(
 										NumMeshDistanceFieldCasters,
 										View.ViewState->IndirectShadowMeshDistanceFieldCasterIndicesSRV ? View.ViewState->IndirectShadowMeshDistanceFieldCasterIndicesSRV.GetReference() : NULL,
 										IndirectShadowLightDirectionSRV,
-										BentNormalOutput->GetRHI());
+										PassParameters->BentNormalInput->GetRHI());
 
 									DispatchComputeShader(RHICmdList, ComputeShaderBase.GetShader(), GroupSizeX, GroupSizeY, 1);
 									ComputeShaderBase->UnsetParameters(RHICmdList, nullptr);
diff --git a/Engine/Source/Runtime/Renderer/Private/DistanceFieldAmbientOcclusion.h b/Engine/Source/Runtime/Renderer/Private/DistanceFieldAmbientOcclusion.h
--- a/Engine/Source/Runtime/Renderer/Private/DistanceFieldAmbientOcclusion.h
+++ b/Engine/Source/Runtime/Renderer/Private/DistanceFieldAmbientOcclusion.h
@@ -354,9 +354,9 @@ public:
 		AOMaxViewDistance.Bind(ParameterMap, TEXT("AOMaxViewDistance"));
 	}
 
-	void Set(FRHICommandList& RHICmdList, FRHIPixelShader* ShaderRHI, const FViewInfo& View, const TRefCountPtr<IPooledRenderTarget>& DistanceFieldAOBentNormal)
+	void Set(FRHICommandList& RHICmdList, FRHIPixelShader* ShaderRHI, const FViewInfo& View, FRHITexture* DistanceFieldAOBentNormal)
 	{
-		FRHITexture* BentNormalAO = DistanceFieldAOBentNormal ? DistanceFieldAOBentNormal->GetRenderTargetItem().ShaderResourceTexture : GWhiteTexture->TextureRHI;
+		FRHITexture* BentNormalAO = DistanceFieldAOBentNormal ? DistanceFieldAOBentNormal : (FRHITexture*)GWhiteTexture->TextureRHI;
 		SetTextureParameter(RHICmdList, ShaderRHI, BentNormalAOTexture, BentNormalAOSampler, TStaticSamplerState<SF_Bilinear>::GetRHI(), BentNormalAO);
 
 		FIntPoint const AOBufferSize = GetBufferSizeForAO();
@@ -450,7 +450,7 @@ public:
 	}
 
 	template<typename TParamRef>
-	void Set(FRHICommandList& RHICmdList, const TParamRef& ShaderRHI, const FViewInfo& View, FSceneRenderTargetItem& DistanceFieldNormal)
+	void Set(FRHICommandList& RHICmdList, const TParamRef& ShaderRHI, const FViewInfo& View, FRHITexture* DistanceFieldNormal)
 	{
 		const FIntPoint DownsampledBufferSize = GetBufferSizeForAO();
 		const FVector2D BaseLevelTexelSizeValue(1.0f / DownsampledBufferSize.X, 1.0f / DownsampledBufferSize.Y);
@@ -469,7 +469,7 @@ public:
 			DistanceFieldNormalTexture,
 			DistanceFieldNormalSampler,
 			TStaticSamplerState<SF_Point,AM_Wrap,AM_Wrap,AM_Wrap>::GetRHI(),
-			DistanceFieldNormal.ShaderResourceTexture
+			DistanceFieldNormal
 			);
 	}
 
diff --git a/Engine/Source/Runtime/Renderer/Private/DistanceFieldLightingPost.cpp b/Engine/Source/Runtime/Renderer/Private/DistanceFieldLightingPost.cpp
--- a/Engine/Source/Runtime/Renderer/Private/DistanceFieldLightingPost.cpp
+++ b/Engine/Source/Runtime/Renderer/Private/DistanceFieldLightingPost.cpp
@@ -88,7 +88,7 @@ public:
 		JitterOffset.Bind(ParameterMap, TEXT("JitterOffset"));
 	}
 
-	void Set(FRHICommandList& RHICmdList, FRHIPixelShader* ShaderRHI, const FViewInfo& View, FSceneRenderTargetItem& DistanceFieldNormal, FSceneRenderTargetItem& DistanceFieldAOBentNormal)
+	void Set(FRHICommandList& RHICmdList, FRHIPixelShader* ShaderRHI, const FViewInfo& View, FRHITexture* DistanceFieldNormal, FRHITexture* DistanceFieldAOBentNormal)
 	{
 		SetTextureParameter(
 			RHICmdList,
@@ -96,7 +96,7 @@ public:
 			DistanceFieldNormalTexture,
 			DistanceFieldNormalSampler,
 			TStaticSamplerState<SF_Bilinear, AM_Clamp, AM_Clamp, AM_Clamp>::GetRHI(),
-			DistanceFieldNormal.ShaderResourceTexture
+			DistanceFieldNormal
 		);
 
 		SetTextureParameter(
@@ -105,7 +105,7 @@ public:
 			BentNormalAOTexture,
 			BentNormalAOSampler,
 			TStaticSamplerState<SF_Bilinear, AM_Clamp, AM_Clamp, AM_Clamp>::GetRHI(),
-			DistanceFieldAOBentNormal.ShaderResourceTexture
+			DistanceFieldAOBentNormal
 		);
 
 		extern FVector2D GetJitterOffset(int32 SampleIndex);
@@ -214,10 +214,10 @@ public:
 		FRHICommandList& RHICmdList, 
 		const FViewInfo& View,
 		const FIntRect& HistoryViewRect,
-		FSceneRenderTargetItem& DistanceFieldNormal, 
-		FSceneRenderTargetItem& DistanceFieldAOBentNormal,
-		FSceneRenderTargetItem& BentNormalHistoryTextureValue, 
-		IPooledRenderTarget* VelocityTextureValue,
+		FRHITexture* DistanceFieldNormal, 
+		FRHITexture* DistanceFieldAOBentNormal,
+		FRHITexture* BentNormalHistoryTextureValue, 
+		FRHITexture* VelocityTextureValue,
 		const FDistanceFieldAOParameters& Parameters)
 	{
 		FRHIPixelShader* ShaderRHI = RHICmdList.GetBoundPixelShader();
@@ -232,7 +232,7 @@ public:
 			BentNormalHistoryTexture,
 			BentNormalHistorySampler,
 			TStaticSamplerState<SF_Bilinear,AM_Clamp,AM_Clamp,AM_Clamp>::GetRHI(),
-			BentNormalHistoryTextureValue.ShaderResourceTexture
+			BentNormalHistoryTextureValue
 			);
 
 		SetShaderValue(RHICmdList, ShaderRHI, HistoryWeight, GAOHistoryWeight);
@@ -245,7 +245,7 @@ public:
 			VelocityTexture,
 			VelocityTextureSampler,
 			TStaticSamplerState<SF_Bilinear,AM_Clamp,AM_Clamp,AM_Clamp>::GetRHI(),
-			VelocityTextureValue ? VelocityTextureValue->GetRenderTargetItem().ShaderResourceTexture : GBlackTexture->TextureRHI
+			VelocityTextureValue ? (FRHITexture2D*)VelocityTextureValue : GBlackTexture->TextureRHI
 			);
 
 		{
@@ -333,8 +333,8 @@ public:
 	void SetParameters(
 		FRHICommandList& RHICmdList, 
 		const FViewInfo& View, 
-		FSceneRenderTargetItem& DistanceFieldNormal, 
-		FSceneRenderTargetItem& BentNormalHistoryTextureValue)
+		FRHITexture* DistanceFieldNormal, 
+		FRHITexture* BentNormalHistoryTextureValue)
 	{
 		FRHIPixelShader* ShaderRHI = RHICmdList.GetBoundPixelShader();
 
@@ -346,7 +346,7 @@ public:
 			BentNormalAOTexture,
 			BentNormalAOSampler,
 			TStaticSamplerState<SF_Bilinear,AM_Clamp,AM_Clamp,AM_Clamp>::GetRHI(),
-			BentNormalHistoryTextureValue.ShaderResourceTexture
+			BentNormalHistoryTextureValue
 			);
 
 		SetTextureParameter(
@@ -355,7 +355,7 @@ public:
 			DistanceFieldNormalTexture,
 			DistanceFieldNormalSampler,
 			TStaticSamplerState<SF_Bilinear,AM_Clamp,AM_Clamp,AM_Clamp>::GetRHI(),
-			DistanceFieldNormal.ShaderResourceTexture
+			DistanceFieldNormal
 			);
 
 		SetShaderValue(RHICmdList, ShaderRHI, HistoryWeight, GAOHistoryWeight);
@@ -429,8 +429,8 @@ public:
 	void SetParameters(
 		FRHICommandList& RHICmdList,
 		const FViewInfo& View,
-		FSceneRenderTargetItem& DistanceFieldNormal,
-		FSceneRenderTargetItem& DistanceFieldAOBentNormal,
+		FRHITexture* DistanceFieldNormal,
+		FRHITexture* DistanceFieldAOBentNormal,
 		const FDistanceFieldAOParameters& Parameters)
 	{
 		FRHIPixelShader* ShaderRHI = RHICmdList.GetBoundPixelShader();
@@ -488,7 +488,7 @@ void GeometryAwareUpsample(FRDGBuilder& GraphBuilder, const FViewInfo& View, FRD
 		TShaderMapRef<FGeometryAwareUpsamplePS> PixelShader(View.ShaderMap);
 		GraphicsPSOInit.BoundShaderState.PixelShaderRHI = PixelShader.GetPixelShader();
 		SetGraphicsPipelineState(RHICmdList, GraphicsPSOInit);
-		PixelShader->SetParameters(RHICmdList, View, DistanceFieldNormal->GetPooledRenderTarget()->GetRenderTargetItem(), BentNormalInterpolation->GetPooledRenderTarget()->GetRenderTargetItem(), Parameters);
+		PixelShader->SetParameters(RHICmdList, View, DistanceFieldNormal->GetRHI(), BentNormalInterpolation->GetRHI(), Parameters);
 
 		DrawRectangle(
 			RHICmdList,
@@ -580,10 +580,10 @@ void UpdateHistory(
 					SetGraphicsPipelineState(RHICmdList, GraphicsPSOInit);
 
 					PixelShader->SetParameters(RHICmdList, View, *DistanceFieldAOHistoryViewRect,
-						DistanceFieldNormal->GetPooledRenderTarget()->GetRenderTargetItem(),
-						BentNormalInterpolation->GetPooledRenderTarget()->GetRenderTargetItem(),
-						BentNormalHistoryTexture->GetPooledRenderTarget()->GetRenderTargetItem(),
-						VelocityTexture ? VelocityTexture->GetPooledRenderTarget() : nullptr,
+						DistanceFieldNormal->GetRHI(),
+						BentNormalInterpolation->GetRHI(),
+						BentNormalHistoryTexture->GetRHI(),
+						VelocityTexture ? VelocityTexture->GetRHI() : nullptr,
 						Parameters);
 
 					DrawRectangle(
@@ -642,7 +642,7 @@ void UpdateHistory(
 
 						GraphicsPSOInit.BoundShaderState.PixelShaderRHI = PixelShader.GetPixelShader();
 						SetGraphicsPipelineState(RHICmdList, GraphicsPSOInit);
-						PixelShader->SetParameters(RHICmdList, View, DistanceFieldNormal->GetPooledRenderTarget()->GetRenderTargetItem(), NewBentNormalHistory->GetPooledRenderTarget()->GetRenderTargetItem());
+						PixelShader->SetParameters(RHICmdList, View, DistanceFieldNormal->GetRHI(), NewBentNormalHistory->GetRHI());
 					}
 					else
 					{
@@ -650,7 +650,7 @@ void UpdateHistory(
 
 						GraphicsPSOInit.BoundShaderState.PixelShaderRHI = PixelShader.GetPixelShader();
 						SetGraphicsPipelineState(RHICmdList, GraphicsPSOInit);
-						PixelShader->SetParameters(RHICmdList, View, DistanceFieldNormal->GetPooledRenderTarget()->GetRenderTargetItem(), NewBentNormalHistory->GetPooledRenderTarget()->GetRenderTargetItem());
+						PixelShader->SetParameters(RHICmdList, View, DistanceFieldNormal->GetRHI(), NewBentNormalHistory->GetRHI());
 					}
 
 					DrawRectangle(
@@ -736,7 +736,7 @@ public:
 		MinIndirectDiffuseOcclusion.Bind(Initializer.ParameterMap,TEXT("MinIndirectDiffuseOcclusion"));
 	}
 
-	void SetParameters(FRHICommandList& RHICmdList, const FViewInfo& View, const TRefCountPtr<IPooledRenderTarget>& DistanceFieldAOBentNormal)
+	void SetParameters(FRHICommandList& RHICmdList, const FViewInfo& View, FRHITexture* DistanceFieldAOBentNormal)
 	{
 		FRHIPixelShader* ShaderRHI = RHICmdList.GetBoundPixelShader();
 
@@ -806,7 +806,7 @@ void UpsampleBentNormalAO(
 			GraphicsPSOInit.BoundShaderState.VertexShaderRHI = VertexShader.GetVertexShader();
 			GraphicsPSOInit.BoundShaderState.PixelShaderRHI = PixelShader.GetPixelShader();
 			SetGraphicsPipelineState(RHICmdList, GraphicsPSOInit);
-			PixelShader->SetParameters(RHICmdList, View, DistanceFieldAOBentNormal->GetPooledRenderTarget());
+			PixelShader->SetParameters(RHICmdList, View, DistanceFieldAOBentNormal->GetRHI());
 			SetShaderParameters(RHICmdList, PixelShader, PixelShader.GetPixelShader(), *PassParameters);
 
 			DrawRectangle(
diff --git a/Engine/Source/Runtime/Renderer/Private/DistanceFieldObjectCulling.cpp b/Engine/Source/Runtime/Renderer/Private/DistanceFieldObjectCulling.cpp
--- a/Engine/Source/Runtime/Renderer/Private/DistanceFieldObjectCulling.cpp
+++ b/Engine/Source/Runtime/Renderer/Private/DistanceFieldObjectCulling.cpp
@@ -261,7 +261,7 @@ public:
 	FBuildTileConesCS()
 	{
 	}
-	void SetParameters(FRHICommandList& RHICmdList, const FViewInfo& View, FSceneRenderTargetItem& DistanceFieldNormal, FScene* Scene, FVector2D NumGroupsValue, const FDistanceFieldAOParameters& Parameters)
+	void SetParameters(FRHICommandList& RHICmdList, const FViewInfo& View, FRHITexture* DistanceFieldNormal, FScene* Scene, FVector2D NumGroupsValue, const FDistanceFieldAOParameters& Parameters)
 	{
 		FRHIComputeShader* ShaderRHI = RHICmdList.GetBoundComputeShader();
 
@@ -288,7 +288,7 @@ public:
 			DistanceFieldNormalTexture,
 			DistanceFieldNormalSampler,
 			TStaticSamplerState<SF_Point,AM_Clamp,AM_Clamp,AM_Clamp>::GetRHI(),
-			DistanceFieldNormal.ShaderResourceTexture
+			DistanceFieldNormal
 			);
 	}
 
@@ -666,7 +666,7 @@ void BuildTileObjectLists(FRDGBuilder& GraphBuilder, FScene* Scene, TArray<FView
 					TShaderMapRef<FBuildTileConesCS> ComputeShader(View.ShaderMap);
 
 					RHICmdList.SetComputeShader(ComputeShader.GetComputeShader());
-					ComputeShader->SetParameters(RHICmdList, View, DistanceFieldNormal->GetPooledRenderTarget()->GetRenderTargetItem(), Scene, FVector2D(TileListGroupSize.X, TileListGroupSize.Y), Parameters);
+					ComputeShader->SetParameters(RHICmdList, View, DistanceFieldNormal->GetRHI(), Scene, FVector2D(TileListGroupSize.X, TileListGroupSize.Y), Parameters);
 					DispatchComputeShader(RHICmdList, ComputeShader.GetShader(), TileListGroupSize.X, TileListGroupSize.Y, 1);
 
 					ComputeShader->UnsetParameters(RHICmdList, View);
diff --git a/Engine/Source/Runtime/Renderer/Private/DistanceFieldScreenGridLighting.cpp b/Engine/Source/Runtime/Renderer/Private/DistanceFieldScreenGridLighting.cpp
--- a/Engine/Source/Runtime/Renderer/Private/DistanceFieldScreenGridLighting.cpp
+++ b/Engine/Source/Runtime/Renderer/Private/DistanceFieldScreenGridLighting.cpp
@@ -113,7 +113,7 @@ public:
 	void SetParameters(
 		FRHICommandList& RHICmdList, 
 		const FViewInfo& View, 
-		FSceneRenderTargetItem& DistanceFieldNormal, 
+		FRHITexture* DistanceFieldNormal, 
 		const FDistanceFieldAOParameters& Parameters,
 		const FGlobalDistanceFieldInfo& GlobalDistanceFieldInfo)
 	{
@@ -251,7 +251,7 @@ public:
 		FRHICommandList& RHICmdList, 
 		const FViewInfo& View, 
 		FIntPoint TileListGroupSizeValue, 
-		FSceneRenderTargetItem& DistanceFieldNormal, 
+		FRHITexture* DistanceFieldNormal, 
 		const FDistanceFieldAOParameters& Parameters,
 		const FGlobalDistanceFieldInfo& GlobalDistanceFieldInfo)
 	{
@@ -385,8 +385,9 @@ public:
 	void SetParameters(
 		FRHICommandList& RHICmdList, 
 		const FViewInfo& View, 
-		FSceneRenderTargetItem& DistanceFieldNormal, 
-		FSceneRenderTargetItem& DownsampledBentNormal)
+		FRHITexture* DistanceFieldNormal, 
+		FRHITexture* DownsampledBentNormal,
+		FRHIUnorderedAccessView* DownsampleBentNormalUAV)
 	{
 		FRHIComputeShader* ShaderRHI = RHICmdList.GetBoundComputeShader();
 		FGlobalShader::SetParameters<FViewUniformShaderParameters>(RHICmdList, ShaderRHI, View.ViewUniformBuffer);
@@ -416,7 +417,7 @@ public:
 
 		FAOScreenGridResources* ScreenGridResources = View.ViewState->AOScreenGridResources;
 
-		DistanceFieldBentNormal.SetTexture(RHICmdList, ShaderRHI, DownsampledBentNormal.ShaderResourceTexture, DownsampledBentNormal.UAV);
+		DistanceFieldBentNormal.SetTexture(RHICmdList, ShaderRHI, DownsampledBentNormal, DownsampleBentNormalUAV);
 
 		SetSRVParameter(RHICmdList, ShaderRHI, ScreenGridConeVisibility, ScreenGridResources->ScreenGridConeVisibility.SRV);
 
@@ -479,6 +480,7 @@ END_SHADER_PARAMETER_STRUCT()
 // TODO(RDG) Replace with individual passes.
 BEGIN_SHADER_PARAMETER_STRUCT(FDistanceFieldAOScreenGridCombineConeVisibilityParameters, )
 	RDG_TEXTURE_ACCESS(DownsampledBentNormal, ERHIAccess::UAVCompute)
+	SHADER_PARAMETER_RDG_TEXTURE_UAV(RWTexture2D, DownsampledBentNormalUAV)
 	RDG_TEXTURE_ACCESS(DistanceFieldNormal, ERHIAccess::SRVCompute)
 END_SHADER_PARAMETER_STRUCT()
 
@@ -549,7 +551,7 @@ void FDeferredShadingSceneRenderer::RenderDistanceFieldAOScreenGrid(
 					TShaderMapRef<TConeTraceScreenGridGlobalOcclusionCS<true> > ComputeShader(View.ShaderMap);
 
 					RHICmdList.SetComputeShader(ComputeShader.GetComputeShader());
-					ComputeShader->SetParameters(RHICmdList, View, TileListGroupSize, DistanceFieldNormal->GetPooledRenderTarget()->GetRenderTargetItem(), Parameters, View.GlobalDistanceFieldInfo);
+					ComputeShader->SetParameters(RHICmdList, View, TileListGroupSize, DistanceFieldNormal->GetRHI(), Parameters, View.GlobalDistanceFieldInfo);
 					DispatchComputeShader(RHICmdList, ComputeShader.GetShader(), GroupSizeX, GroupSizeY, 1);
 					ComputeShader->UnsetParameters(RHICmdList, View);
 				}
@@ -558,7 +560,7 @@ void FDeferredShadingSceneRenderer::RenderDistanceFieldAOScreenGrid(
 					TShaderMapRef<TConeTraceScreenGridGlobalOcclusionCS<false> > ComputeShader(View.ShaderMap);
 
 					RHICmdList.SetComputeShader(ComputeShader.GetComputeShader());
-					ComputeShader->SetParameters(RHICmdList, View, TileListGroupSize, DistanceFieldNormal->GetPooledRenderTarget()->GetRenderTargetItem(), Parameters, View.GlobalDistanceFieldInfo);
+					ComputeShader->SetParameters(RHICmdList, View, TileListGroupSize, DistanceFieldNormal->GetRHI(), Parameters, View.GlobalDistanceFieldInfo);
 					DispatchComputeShader(RHICmdList, ComputeShader.GetShader(), GroupSizeX, GroupSizeY, 1);
 					ComputeShader->UnsetParameters(RHICmdList, View);
 				}
@@ -579,7 +581,7 @@ void FDeferredShadingSceneRenderer::RenderDistanceFieldAOScreenGrid(
 						TShaderMapRef<TConeTraceScreenGridObjectOcclusionCS<true> > ComputeShader(View.ShaderMap);
 
 						RHICmdList.SetComputeShader(ComputeShader.GetComputeShader());
-						ComputeShader->SetParameters(RHICmdList, View, DistanceFieldNormal->GetPooledRenderTarget()->GetRenderTargetItem(), Parameters, View.GlobalDistanceFieldInfo);
+						ComputeShader->SetParameters(RHICmdList, View, DistanceFieldNormal->GetRHI(), Parameters, View.GlobalDistanceFieldInfo);
 						DispatchIndirectComputeShader(RHICmdList, ComputeShader.GetShader(), TileIntersectionResources->ObjectTilesIndirectArguments.Buffer, 0);
 						ComputeShader->UnsetParameters(RHICmdList, View);
 					}
@@ -590,7 +592,7 @@ void FDeferredShadingSceneRenderer::RenderDistanceFieldAOScreenGrid(
 						TShaderMapRef<TConeTraceScreenGridObjectOcclusionCS<false> > ComputeShader(View.ShaderMap);
 
 						RHICmdList.SetComputeShader(ComputeShader.GetComputeShader());
-						ComputeShader->SetParameters(RHICmdList, View, DistanceFieldNormal->GetPooledRenderTarget()->GetRenderTargetItem(), Parameters, View.GlobalDistanceFieldInfo);
+						ComputeShader->SetParameters(RHICmdList, View, DistanceFieldNormal->GetRHI(), Parameters, View.GlobalDistanceFieldInfo);
 						DispatchIndirectComputeShader(RHICmdList, ComputeShader.GetShader(), TileIntersectionResources->ObjectTilesIndirectArguments.Buffer, 0);
 						ComputeShader->UnsetParameters(RHICmdList, View);
 					}
@@ -598,7 +600,7 @@ void FDeferredShadingSceneRenderer::RenderDistanceFieldAOScreenGrid(
 			}
 
 			// Compute heightfield occlusion after heightfield GI, otherwise it self-shadows incorrectly
-			View.HeightfieldLightingViewInfo.ComputeOcclusionForScreenGrid(View, RHICmdList, DistanceFieldNormal->GetPooledRenderTarget()->GetRenderTargetItem(), *ScreenGridResources, Parameters);
+			View.HeightfieldLightingViewInfo.ComputeOcclusionForScreenGrid(View, RHICmdList, DistanceFieldNormal->GetRHI(), *ScreenGridResources, Parameters);
 
 			RHICmdList.Transition(FRHITransitionInfo(ScreenGridResources->ScreenGridConeVisibility.UAV, ERHIAccess::UAVCompute, ERHIAccess::SRVCompute));
 		});
@@ -619,16 +621,17 @@ void FDeferredShadingSceneRenderer::RenderDistanceFieldAOScreenGrid(
 
 		auto* PassParameters = GraphBuilder.AllocParameters<FDistanceFieldAOScreenGridCombineConeVisibilityParameters>();
 		PassParameters->DownsampledBentNormal = DownsampledBentNormal;
+		PassParameters->DownsampledBentNormalUAV = GraphBuilder.CreateUAV(DownsampledBentNormal);
 		PassParameters->DistanceFieldNormal = DistanceFieldNormal;
 
 		GraphBuilder.AddPass(
 			RDG_EVENT_NAME("CombineCones"),
 			PassParameters,
 			ERDGPassFlags::Compute,
-			[&View, ComputeShader, DistanceFieldNormal, DownsampledBentNormal, GroupSizeX, GroupSizeY](FRHICommandList& RHICmdList)
+			[&View, ComputeShader, PassParameters, GroupSizeX, GroupSizeY](FRHICommandList& RHICmdList)
 		{
 			RHICmdList.SetComputeShader(ComputeShader.GetComputeShader());
-			ComputeShader->SetParameters(RHICmdList, View, DistanceFieldNormal->GetPooledRenderTarget()->GetRenderTargetItem(), DownsampledBentNormal->GetPooledRenderTarget()->GetRenderTargetItem());
+			ComputeShader->SetParameters(RHICmdList, View, PassParameters->DistanceFieldNormal->GetRHI(), PassParameters->DownsampledBentNormal->GetRHI(), PassParameters->DownsampledBentNormalUAV->GetRHI());
 			DispatchComputeShader(RHICmdList, ComputeShader.GetShader(), GroupSizeX, GroupSizeY, 1);
 		});
 	}
diff --git a/Engine/Source/Runtime/Renderer/Private/HeightfieldLighting.cpp b/Engine/Source/Runtime/Renderer/Private/HeightfieldLighting.cpp
--- a/Engine/Source/Runtime/Renderer/Private/HeightfieldLighting.cpp
+++ b/Engine/Source/Runtime/Renderer/Private/HeightfieldLighting.cpp
@@ -1151,7 +1151,7 @@ public:
 		const FViewInfo& View,
 		UTexture2D* HeightfieldTextureValue,
 		int32 NumHeightfieldsValue,
-		FSceneRenderTargetItem& DistanceFieldNormal,
+		FRHITexture* DistanceFieldNormal,
 		const FAOScreenGridResources& ScreenGridResources,
 		const FDistanceFieldAOParameters& Parameters)
 	{
@@ -1204,7 +1204,7 @@ IMPLEMENT_SHADER_TYPE(, FCalculateHeightfieldOcclusionScreenGridCS, TEXT("/Engin
 void FHeightfieldLightingViewInfo::ComputeOcclusionForScreenGrid(
 	const FViewInfo& View,
 	FRHICommandListImmediate& RHICmdList,
-	FSceneRenderTargetItem& DistanceFieldNormal,
+	FRHITexture* DistanceFieldNormal,
 	const FAOScreenGridResources& ScreenGridResources,
 	const FDistanceFieldAOParameters& Parameters) const
 {
diff --git a/Engine/Source/Runtime/Renderer/Private/HeightfieldLighting.h b/Engine/Source/Runtime/Renderer/Private/HeightfieldLighting.h
--- a/Engine/Source/Runtime/Renderer/Private/HeightfieldLighting.h
+++ b/Engine/Source/Runtime/Renderer/Private/HeightfieldLighting.h
@@ -111,14 +111,14 @@ public:
 	void ComputeOcclusionForScreenGrid(
 		const FViewInfo& View, 
 		FRHICommandListImmediate& RHICmdList, 
-		FSceneRenderTargetItem& DistanceFieldNormal,
+		FRHITexture* DistanceFieldNormal,
 		const class FAOScreenGridResources& ScreenGridResources,
 		const class FDistanceFieldAOParameters& Parameters) const;
 
 	void ComputeIrradianceForScreenGrid(
 		const FViewInfo& View, 
-		FRHICommandListImmediate& RHICmdList, 
-		FSceneRenderTargetItem& DistanceFieldNormal, 
+		FRHICommandListImmediate& RHICmdList,
+		FRHITexture* DistanceFieldNormal,
 		const FAOScreenGridResources& ScreenGridResources,
 		const FDistanceFieldAOParameters& Parameters) const;
 
diff --git a/Engine/Source/Runtime/Renderer/Private/LightRendering.cpp b/Engine/Source/Runtime/Renderer/Private/LightRendering.cpp
--- a/Engine/Source/Runtime/Renderer/Private/LightRendering.cpp
+++ b/Engine/Source/Runtime/Renderer/Private/LightRendering.cpp
@@ -349,18 +349,18 @@ struct FRenderLightParams
 	uint32 DeepShadow_TransmittanceMaskBufferMaxCount = 0;
 	
 	// Visibility buffer data
-	IPooledRenderTarget* HairCategorizationTexture = nullptr;
-	IPooledRenderTarget* HairVisibilityNodeOffsetAndCount = nullptr;
-	IPooledRenderTarget* HairVisibilityNodeCount = nullptr;
+	FRHITexture* HairCategorizationTexture = nullptr;
+	FRHITexture* HairVisibilityNodeOffsetAndCount = nullptr;
+	FRHITexture* HairVisibilityNodeCount = nullptr;
 	FShaderResourceViewRHIRef HairVisibilityNodeCoordsSRV = nullptr;
 	FShaderResourceViewRHIRef HairVisibilityNodeDataSRV = nullptr;
 
-	IPooledRenderTarget* ScreenShadowMaskSubPixelTexture = nullptr;
+	FRHITexture* ScreenShadowMaskSubPixelTexture = nullptr;
 
 	// Cloud shadow data
 	FMatrix Cloud_WorldToLightClipShadowMatrix;
 	float Cloud_ShadowmapFarDepthKm = 0.0f;
-	IPooledRenderTarget* Cloud_ShadowmapTexture = nullptr;
+	FRHITexture* Cloud_ShadowmapTexture = nullptr;
 	float Cloud_ShadowmapStrength = 0.0f;
 };
 
@@ -625,7 +625,7 @@ private:
 			LTCMatTexture,
 			LTCMatSampler,
 			TStaticSamplerState<SF_Bilinear,AM_Clamp,AM_Clamp,AM_Clamp>::GetRHI(),
-			GSystemTextures.LTCMat->GetRenderTargetItem().ShaderResourceTexture
+			GSystemTextures.LTCMat->GetShaderResourceRHI()
 			);
 
 		SetTextureParameter(
@@ -634,7 +634,7 @@ private:
 			LTCAmpTexture,
 			LTCAmpSampler,
 			TStaticSamplerState<SF_Bilinear,AM_Clamp,AM_Clamp,AM_Clamp>::GetRHI(),
-			GSystemTextures.LTCAmp->GetRenderTargetItem().ShaderResourceTexture
+			GSystemTextures.LTCAmp->GetShaderResourceRHI()
 			);
 
 		{
@@ -704,7 +704,7 @@ private:
 					ScreenShadowMaskSubPixelTexture,
 					LightAttenuationTextureSampler,
 					TStaticSamplerState<SF_Point, AM_Clamp, AM_Clamp, AM_Clamp>::GetRHI(),
-					(RenderLightParams && RenderLightParams->ScreenShadowMaskSubPixelTexture) ? RenderLightParams->ScreenShadowMaskSubPixelTexture->GetRenderTargetItem().ShaderResourceTexture : GWhiteTexture->TextureRHI);
+					(RenderLightParams && RenderLightParams->ScreenShadowMaskSubPixelTexture) ? (FRHITexture2D*)RenderLightParams->ScreenShadowMaskSubPixelTexture : GWhiteTexture->TextureRHI);
 
 				uint32 InHairShadowMaskValid = RenderLightParams->ScreenShadowMaskSubPixelTexture ? 1 : 0;
 				SetShaderValue(
@@ -725,7 +725,7 @@ private:
 					HairCategorizationTexture,
 					LightAttenuationTextureSampler,
 					TStaticSamplerState<SF_Point, AM_Clamp, AM_Clamp, AM_Clamp>::GetRHI(),
-					RenderLightParams->HairCategorizationTexture->GetRenderTargetItem().TargetableTexture);
+					RenderLightParams->HairCategorizationTexture);
 			}
 		}
 
@@ -739,7 +739,7 @@ private:
 					HairVisibilityNodeOffsetAndCount,
 					LightAttenuationTextureSampler,
 					TStaticSamplerState<SF_Point, AM_Clamp, AM_Clamp, AM_Clamp>::GetRHI(),
-					RenderLightParams->HairVisibilityNodeOffsetAndCount->GetRenderTargetItem().TargetableTexture);
+					RenderLightParams->HairVisibilityNodeOffsetAndCount);
 			}
 		}
 		
@@ -787,7 +787,7 @@ private:
 				DummyRectLightTextureForCapsuleCompilerWarning,
 				LTCMatSampler,
 				TStaticSamplerState<SF_Bilinear, AM_Clamp, AM_Clamp, AM_Clamp>::GetRHI(),
-				GSystemTextures.DepthDummy->GetRenderTargetItem().ShaderResourceTexture
+				GSystemTextures.DepthDummy->GetShaderResourceRHI()
 			);
 		}
 
@@ -801,7 +801,7 @@ private:
 					CloudShadowmapTexture,
 					CloudShadowmapSampler,
 					TStaticSamplerState<SF_Bilinear, AM_Clamp, AM_Clamp, AM_Clamp>::GetRHI(),
-					RenderLightParams->Cloud_ShadowmapTexture ? RenderLightParams->Cloud_ShadowmapTexture->GetRenderTargetItem().ShaderResourceTexture : GBlackVolumeTexture->TextureRHI);
+					RenderLightParams->Cloud_ShadowmapTexture ? (FRHITexture2D*)RenderLightParams->Cloud_ShadowmapTexture : GBlackVolumeTexture->TextureRHI);
 
 				SetShaderValue(
 					RHICmdList,
@@ -2248,7 +2248,7 @@ void FDeferredShadingSceneRenderer::RenderLight(
 	const bool bHairLighting = InHairVisibilityViews && ViewIndex < InHairVisibilityViews->HairDatas.Num() && InHairVisibilityViews->HairDatas[ViewIndex].CategorizationTexture != nullptr;
 	if (bHairLighting)
 	{
-		RenderLightParams.HairCategorizationTexture = InHairVisibilityViews->HairDatas[ViewIndex].CategorizationTexture->GetPooledRenderTarget();
+		RenderLightParams.HairCategorizationTexture = InHairVisibilityViews->HairDatas[ViewIndex].CategorizationTexture->GetRHI();
 	}
 
 	if (Strata::IsStrataEnabled() && Strata::IsClassificationEnabled())
@@ -2296,14 +2296,14 @@ void FDeferredShadingSceneRenderer::RenderLight(
 			const bool bLight1CloudPerPixelTransmittance = CloudInfo && VolumetricCloudShadowMap1Valid && AtmosphereLight1Proxy == LightSceneInfo->Proxy && AtmosphereLight1Proxy && AtmosphereLight1Proxy->GetCloudShadowOnSurfaceStrength() > 0.0f;
 			if (bLight0CloudPerPixelTransmittance)
 			{
-				RenderLightParams.Cloud_ShadowmapTexture = TryGetPooledRenderTarget(View.VolumetricCloudShadowRenderTarget[0]);
+				RenderLightParams.Cloud_ShadowmapTexture = TryGetRHI(View.VolumetricCloudShadowRenderTarget[0]);
 				RenderLightParams.Cloud_ShadowmapFarDepthKm = CloudInfo->GetVolumetricCloudCommonShaderParameters().CloudShadowmapFarDepthKm[0];
 				RenderLightParams.Cloud_WorldToLightClipShadowMatrix = CloudInfo->GetVolumetricCloudCommonShaderParameters().CloudShadowmapWorldToLightClipMatrix[0];
 				RenderLightParams.Cloud_ShadowmapStrength = AtmosphereLight0Proxy->GetCloudShadowOnSurfaceStrength();
 			}
 			else if(bLight1CloudPerPixelTransmittance)
 			{
-				RenderLightParams.Cloud_ShadowmapTexture = TryGetPooledRenderTarget(View.VolumetricCloudShadowRenderTarget[1]);
+				RenderLightParams.Cloud_ShadowmapTexture = TryGetRHI(View.VolumetricCloudShadowRenderTarget[1]);
 				RenderLightParams.Cloud_ShadowmapFarDepthKm = CloudInfo->GetVolumetricCloudCommonShaderParameters().CloudShadowmapFarDepthKm[1];
 				RenderLightParams.Cloud_WorldToLightClipShadowMatrix = CloudInfo->GetVolumetricCloudCommonShaderParameters().CloudShadowmapWorldToLightClipMatrix[1];
 				RenderLightParams.Cloud_ShadowmapStrength = AtmosphereLight1Proxy->GetCloudShadowOnSurfaceStrength();
@@ -2559,12 +2559,12 @@ void FDeferredShadingSceneRenderer::RenderLightForHair(
 
 			FRenderLightParams RenderLightParams;
 			RenderLightParams.DeepShadow_TransmittanceMaskBufferMaxCount = MaxTransmittanceElementCount;
-			RenderLightParams.ScreenShadowMaskSubPixelTexture = bIsShadowMaskValid ? PassParameters->Light.ShadowMaskTexture->GetPooledRenderTarget() : GSystemTextures.WhiteDummy.GetReference();
+			RenderLightParams.ScreenShadowMaskSubPixelTexture = bIsShadowMaskValid ? PassParameters->Light.ShadowMaskTexture->GetRHI() : GSystemTextures.WhiteDummy->GetShaderResourceRHI();
 			RenderLightParams.DeepShadow_TransmittanceMaskBuffer = PassParameters->HairTransmittanceMaskSRV->GetRHI();
-			RenderLightParams.HairVisibilityNodeOffsetAndCount = PassParameters->HairIndexAndCountTexture->GetPooledRenderTarget();
+			RenderLightParams.HairVisibilityNodeOffsetAndCount = PassParameters->HairIndexAndCountTexture->GetRHI();
 			RenderLightParams.HairVisibilityNodeDataSRV = PassParameters->HairVisibilityNodeDataSRV->GetRHI();
 			RenderLightParams.HairVisibilityNodeCoordsSRV = PassParameters->HairVisibilityNodeCoordsSRV->GetRHI();
-			RenderLightParams.HairCategorizationTexture = PassParameters->Light.HairCategorizationTexture->GetPooledRenderTarget();
+			RenderLightParams.HairCategorizationTexture = PassParameters->Light.HairCategorizationTexture->GetRHI();
 
 			FGraphicsPipelineStateInitializer GraphicsPSOInit;
 			RHICmdList.ApplyCachedRenderTargets(GraphicsPSOInit);
diff --git a/Engine/Source/Runtime/Renderer/Private/PostProcess/PostProcessFFTBloom.cpp b/Engine/Source/Runtime/Renderer/Private/PostProcess/PostProcessFFTBloom.cpp
--- a/Engine/Source/Runtime/Renderer/Private/PostProcess/PostProcessFFTBloom.cpp
+++ b/Engine/Source/Runtime/Renderer/Private/PostProcess/PostProcessFFTBloom.cpp
@@ -458,15 +458,15 @@ bool TransformKernelFFT(
 
 struct FFFTBloomIntermediates
 {
-	const FSceneRenderTargetItem* FullResolutionTarget = nullptr;
+	FRHITexture* FullResolutionTarget = nullptr;
 	FIntRect FullResolutionViewRect;
 	FIntPoint FullResolutionSize;
 
-	const FSceneRenderTargetItem* HalfResolutionTarget = nullptr;
+	FRHITexture* HalfResolutionTarget = nullptr;
 	FIntRect HalfResolutionViewRect;
 	FIntPoint HalfResolutionSize;
 
-	const FSceneRenderTargetItem* InputTarget = nullptr;
+	FRHITexture* InputTarget = nullptr;
 	FRHIUnorderedAccessView* OutputUAV = nullptr;
 
 	// The size of the input buffer.
@@ -530,14 +530,14 @@ FFFTBloomIntermediates GetFFTBloomIntermediates(
 	const bool bHalfResolutionFFT = IsFFTBloomHalfResolutionEnabled();
 
 	FFFTBloomIntermediates Intermediates;
-	Intermediates.FullResolutionTarget = &Inputs.FullResolutionTexture->GetPooledRenderTarget()->GetRenderTargetItem();
+	Intermediates.FullResolutionTarget = Inputs.FullResolutionTexture->GetRHI();
 	Intermediates.FullResolutionViewRect = Inputs.FullResolutionViewRect;
 	Intermediates.FullResolutionSize = Intermediates.FullResolutionViewRect.Size();
 	Intermediates.bHalfResolutionFFT = bHalfResolutionFFT;
 
 	if (bHalfResolutionFFT)
 	{
-		Intermediates.HalfResolutionTarget = &Inputs.HalfResolutionTexture->GetPooledRenderTarget()->GetRenderTargetItem();
+		Intermediates.HalfResolutionTarget = Inputs.HalfResolutionTexture->GetRHI();
 		Intermediates.HalfResolutionViewRect = Inputs.HalfResolutionViewRect;
 		Intermediates.HalfResolutionSize = Inputs.HalfResolutionViewRect.Size();
 
@@ -613,7 +613,7 @@ void ConvolveWithKernel(
 	GRenderTargetPool.FindFreeElement(RHICmdList, Desc, TmpTargets[1], TEXT("Tmp FFT Buffer B"));
 
 	// Get the source
-	const FTextureRHIRef& InputTexture = Intermediates.InputTarget->ShaderResourceTexture;
+	FRHITexture* InputTexture = Intermediates.InputTarget;
 
 	GPUFFT::ConvolutionWithTextureImage2D(FFTContext, Intermediates.FrequencySize, Intermediates.bDoHorizontalFirst,
 		SpectralKernelTexture,
@@ -811,7 +811,7 @@ FRDGTextureRef AddFFTBloomPass(FRDGBuilder& GraphBuilder, const FViewInfo& View,
 			const FTextureRHIRef& CenterWeightTexture = FFTKernel.CenterWeight->GetRenderTargetItem().ShaderResourceTexture;
 
 			// Get full resolution source
-			const FTextureRHIRef& FullResResourceTexture = Intermediates.FullResolutionTarget->ShaderResourceTexture;
+			FRHITexture* FullResResourceTexture = Intermediates.FullResolutionTarget;
 
 			// Blend with  alpha * SrcBuffer + betta * BloomedBuffer  where alpha = Weights[0], beta = Weights[1]
 			const FIntPoint HalfResBufferSize = Intermediates.InputBufferSize;
diff --git a/Engine/Source/Runtime/Renderer/Private/ShadowRendering.h b/Engine/Source/Runtime/Renderer/Private/ShadowRendering.h
--- a/Engine/Source/Runtime/Renderer/Private/ShadowRendering.h
+++ b/Engine/Source/Runtime/Renderer/Private/ShadowRendering.h
@@ -1063,7 +1063,7 @@ public:
 
 		if (HairVisibilityData && HairVisibilityData->CategorizationTexture)
 		{
-			SetTextureParameter(RHICmdList, ShaderRHI, HairCategorizationTexture, HairVisibilityData->CategorizationTexture->GetPooledRenderTarget()->GetRenderTargetItem().ShaderResourceTexture);
+			SetTextureParameter(RHICmdList, ShaderRHI, HairCategorizationTexture, HairVisibilityData->CategorizationTexture->GetRHI());
 
 			float DeviceZNear = 1;
 			float DeviceZFar = 0;
@@ -1511,7 +1511,7 @@ public:
 
 		if (bUseSubPixel && HairVisibilityData && HairVisibilityData->CategorizationTexture)
 		{
-			SetTextureParameter(RHICmdList, ShaderRHI, HairCategorizationTexture, HairVisibilityData->CategorizationTexture->GetPooledRenderTarget()->GetRenderTargetItem().ShaderResourceTexture);
+			SetTextureParameter(RHICmdList, ShaderRHI, HairCategorizationTexture, HairVisibilityData->CategorizationTexture->GetRHI());
 		}
 
 		if (StrataGlobalParameters.IsBound())
diff --git a/Engine/Source/Runtime/Renderer/Private/TiledDeferredLightRendering.cpp b/Engine/Source/Runtime/Renderer/Private/TiledDeferredLightRendering.cpp
--- a/Engine/Source/Runtime/Renderer/Private/TiledDeferredLightRendering.cpp
+++ b/Engine/Source/Runtime/Renderer/Private/TiledDeferredLightRendering.cpp
@@ -129,14 +129,14 @@ public:
 		const FSimpleLightArray& SimpleLights,
 		int32 StartIndex,
 		int32 NumThisPass,
-		const IPooledRenderTarget& InTextureValue,
-		const IPooledRenderTarget& OutTextureValue)
+		FRHITexture* InTextureValue,
+		FRHIUnorderedAccessView* OutTextureValue)
 	{
 		FRHIComputeShader* ShaderRHI = RHICmdList.GetBoundComputeShader();
 
 		FGlobalShader::SetParameters<FViewUniformShaderParameters>(RHICmdList, ShaderRHI, View.ViewUniformBuffer);
-		SetTextureParameter(RHICmdList, ShaderRHI, InTexture, InTextureValue.GetRenderTargetItem().ShaderResourceTexture);
-		OutTexture.SetTexture(RHICmdList, ShaderRHI, 0, OutTextureValue.GetRenderTargetItem().UAV);
+		SetTextureParameter(RHICmdList, ShaderRHI, InTexture, InTextureValue);
+		OutTexture.SetTexture(RHICmdList, ShaderRHI, 0, OutTextureValue);
 
 		SetShaderValue(RHICmdList, ShaderRHI, ViewDimensions, View.ViewRect);
 
@@ -146,7 +146,7 @@ public:
 			LTCMatTexture,
 			LTCMatSampler,
 			TStaticSamplerState<SF_Bilinear,AM_Clamp,AM_Clamp,AM_Clamp>::GetRHI(),
-			GSystemTextures.LTCMat->GetRenderTargetItem().ShaderResourceTexture
+			GSystemTextures.LTCMat->GetShaderResourceRHI()
 			);
 
 		SetTextureParameter(
@@ -155,7 +155,7 @@ public:
 			LTCAmpTexture,
 			LTCAmpSampler,
 			TStaticSamplerState<SF_Bilinear,AM_Clamp,AM_Clamp,AM_Clamp>::GetRHI(),
-			GSystemTextures.LTCAmp->GetRenderTargetItem().ShaderResourceTexture
+			GSystemTextures.LTCAmp->GetShaderResourceRHI()
 			);
 
 		const int32 NumLightsToRenderInSortedLights = TiledDeferredLightsEnd - TiledDeferredLightsStart;
@@ -254,7 +254,7 @@ public:
 		SetShaderValue(RHICmdList, ShaderRHI, NumLights, NumThisPass);
 	}
 
-	void UnsetParameters(FRHIComputeCommandList& RHICmdList, const IPooledRenderTarget& OutTextureValue)
+	void UnsetParameters(FRHIComputeCommandList& RHICmdList)
 	{
 		OutTexture.UnsetUAV(RHICmdList, RHICmdList.GetBoundComputeShader());
 	}
@@ -313,8 +313,8 @@ static void SetShaderTemplTiledLighting(
 	const FSimpleLightArray& SimpleLights,
 	int32 StartIndex,
 	int32 NumThisPass,
-	const IPooledRenderTarget& InTexture,
-	const IPooledRenderTarget& OutTexture)
+	FRHITexture* InTexture,
+	FRHIUnorderedAccessView* OutTexture)
 {
 	TShaderMapRef<FTiledDeferredLightingCS<bVisualizeLightCulling> > ComputeShader(View.ShaderMap);
 	RHICmdList.SetComputeShader(ComputeShader.GetComputeShader());
@@ -325,12 +325,13 @@ static void SetShaderTemplTiledLighting(
 	uint32 GroupSizeY = (View.ViewRect.Size().Y + GDeferredLightTileSizeY - 1) / GDeferredLightTileSizeY;
 	DispatchComputeShader(RHICmdList, ComputeShader.GetShader(), GroupSizeX, GroupSizeY, 1);
 
-	ComputeShader->UnsetParameters(RHICmdList, OutTexture);
+	ComputeShader->UnsetParameters(RHICmdList);
 }
 
 BEGIN_SHADER_PARAMETER_STRUCT(FTiledDeferredLightingParameters, )
 	RDG_TEXTURE_ACCESS(SceneColorInput, ERHIAccess::SRVCompute)
 	RDG_TEXTURE_ACCESS(SceneColorOutput, ERHIAccess::UAVCompute)
+	SHADER_PARAMETER_RDG_TEXTURE_UAV(RWTexture2D, SceneColorOutputUAV)
 	SHADER_PARAMETER_RDG_UNIFORM_BUFFER(FSceneTextureUniformParameters, SceneTextures)
 END_SHADER_PARAMETER_STRUCT()
 
@@ -373,6 +374,8 @@ void FDeferredShadingSceneRenderer::RenderTiledDeferredLighting(
 				SceneColorOutputTexture = GraphBuilder.CreateTexture(Desc, TEXT("SceneColorTiled"));
 			}
 
+			FRDGTextureUAVRef SceneColorOutputUAV = GraphBuilder.CreateUAV(SceneColorOutputTexture);
+
 			const int32 ViewCount = Views.Num();
 
 			for (int32 ViewIndex = 0; ViewIndex < ViewCount; ViewIndex++)
@@ -382,21 +385,22 @@ void FDeferredShadingSceneRenderer::RenderTiledDeferredLighting(
 				FTiledDeferredLightingParameters* PassParameters = GraphBuilder.AllocParameters<FTiledDeferredLightingParameters>();
 				PassParameters->SceneColorInput = SceneTextures.Color.Target;
 				PassParameters->SceneColorOutput = SceneColorOutputTexture;
+				PassParameters->SceneColorOutputUAV = SceneColorOutputUAV;
 				PassParameters->SceneTextures = SceneTextures.UniformBuffer;
 
 				GraphBuilder.AddPass(
 					RDG_EVENT_NAME("TiledDeferredLighting"),
 					PassParameters,
 					ERDGPassFlags::Compute,
-					[&View, ViewIndex, ViewCount, &SortedLights, TiledDeferredLightsStart, TiledDeferredLightsEnd, &SimpleLights, StartIndex, NumThisPass, SceneColorTexture = SceneTextures.Color.Target, SceneColorOutputTexture](FRHIComputeCommandList& RHICmdList)
+					[&View, ViewIndex, ViewCount, &SortedLights, TiledDeferredLightsStart, TiledDeferredLightsEnd, &SimpleLights, StartIndex, NumThisPass, SceneColorOutputTexture, SceneColorTexture = SceneTextures.Color.Target, SceneColorOutputUAV](FRHIComputeCommandList& RHICmdList)
 				{
 					if (View.Family->EngineShowFlags.VisualizeLightCulling)
 					{
-						SetShaderTemplTiledLighting<1>(RHICmdList, View, ViewIndex, ViewCount, SortedLights, TiledDeferredLightsStart, TiledDeferredLightsEnd, SimpleLights, StartIndex, NumThisPass, *SceneColorTexture->GetPooledRenderTarget(), *SceneColorOutputTexture->GetPooledRenderTarget());
+						SetShaderTemplTiledLighting<1>(RHICmdList, View, ViewIndex, ViewCount, SortedLights, TiledDeferredLightsStart, TiledDeferredLightsEnd, SimpleLights, StartIndex, NumThisPass, SceneColorTexture->GetRHI(), SceneColorOutputUAV->GetRHI());
 					}
 					else
 					{
-						SetShaderTemplTiledLighting<0>(RHICmdList, View, ViewIndex, ViewCount, SortedLights, TiledDeferredLightsStart, TiledDeferredLightsEnd, SimpleLights, StartIndex, NumThisPass, *SceneColorTexture->GetPooledRenderTarget(), *SceneColorOutputTexture->GetPooledRenderTarget());
+						SetShaderTemplTiledLighting<0>(RHICmdList, View, ViewIndex, ViewCount, SortedLights, TiledDeferredLightsStart, TiledDeferredLightsEnd, SimpleLights, StartIndex, NumThisPass, SceneColorTexture->GetRHI(), SceneColorOutputUAV->GetRHI());
 					}
 
 #if PLATFORM_REQUIRES_UAV_TO_RTV_TEXTURE_CACHE_FLUSH_WORKAROUND
diff --git a/Engine/Source/Runtime/Renderer/Private/VolumetricFog.cpp b/Engine/Source/Runtime/Renderer/Private/VolumetricFog.cpp
--- a/Engine/Source/Runtime/Renderer/Private/VolumetricFog.cpp
+++ b/Engine/Source/Runtime/Renderer/Private/VolumetricFog.cpp
@@ -326,7 +326,7 @@ public:
 		SetShaderValue(RHICmdList, ShaderRHI, InverseSquaredLightDistanceBiasScale, GInverseSquaredLightDistanceBiasScale);
 
 		SetShaderValue(RHICmdList, ShaderRHI, UseConservativeDepthTexture, GVolumetricFogConservativeDepth > 0 ? 1 : 0);
-		SetTextureParameter(RHICmdList, ShaderRHI, ConservativeDepthTextureParam, ConservativeDepthTexture->GetPooledRenderTarget()->GetRenderTargetItem().TargetableTexture);
+		SetTextureParameter(RHICmdList, ShaderRHI, ConservativeDepthTextureParam, ConservativeDepthTexture->GetRHI());
 
 		VolumeShadowingParameters.Set(RHICmdList, ShaderRHI, View, LightSceneInfo, ShadowMap, INDEX_NONE, bDynamicallyShadowed);
 	}
