Fix Mobile Multiview rendering (Quest, etc) by ensuring depth (and possibly color targets) are created as TextureArrays when multiview is enabled.

diff --git a/Engine/Source/Runtime/Renderer/Private/SceneTextures.cpp b/Engine/Source/Runtime/Renderer/Private/SceneTextures.cpp
--- a/Engine/Source/Runtime/Renderer/Private/SceneTextures.cpp
+++ b/Engine/Source/Runtime/Renderer/Private/SceneTextures.cpp
@@ -599,9 +599,11 @@ FSceneTextures& FMinimalSceneTextures::Create(FRDGBuilder& GraphBuilder, const F
 			Flags |= TexCreate_UAV;
 		}
 
-		FRDGTextureDesc Desc(FRDGTextureDesc::Create2D(SceneTextures.Config.Extent, PF_DepthStencil, Config.DepthClearValue, Flags));
+		// TODO: Array-size could be values > 2, on upcoming XR devices. This should be part of the config.
+		FRDGTextureDesc Desc(Config.bRequireMultiView ?
+							 FRDGTextureDesc::Create2DArray(SceneTextures.Config.Extent, PF_DepthStencil, Config.DepthClearValue, Flags, 2) :
+							 FRDGTextureDesc::Create2D(SceneTextures.Config.Extent, PF_DepthStencil, Config.DepthClearValue, Flags));
 		Desc.NumSamples = Config.NumSamples;
-		Desc.ArraySize = Config.bRequireMultiView ? 2 : 1;
 		SceneTextures.Depth = GraphBuilder.CreateTexture(Desc, TEXT("SceneDepthZ"));
 
 		if (Desc.NumSamples > 1)
@@ -635,10 +637,11 @@ FSceneTextures& FMinimalSceneTextures::Create(FRDGBuilder& GraphBuilder, const F
 		const TCHAR* SceneColorName = bIsMobilePlatform ? TEXT("SceneColorMobile") : TEXT("SceneColorDeferred");
 
 		// Create the scene color.
-		FRDGTextureDesc Desc(FRDGTextureDesc::Create2D(Config.Extent, Config.ColorFormat, Config.ColorClearValue, Flags));
+		// TODO: Array-size could be values > 2, on upcoming XR devices. This should be part of the config.
+		FRDGTextureDesc Desc(Config.bRequireMultiView ?
+							 FRDGTextureDesc::Create2DArray(Config.Extent, Config.ColorFormat, Config.ColorClearValue, Flags, 2) :
+							 FRDGTextureDesc::Create2D(Config.Extent, Config.ColorFormat, Config.ColorClearValue, Flags));
 		Desc.NumSamples = Config.NumSamples;
-		Desc.ArraySize = Config.bRequireMultiView ? 2 : 1;
-
 		SceneTextures.Color = GraphBuilder.CreateTexture(Desc, SceneColorName);
 
 		if (Desc.NumSamples > 1)
