Sync GPU Niagara evaluation with RDG timeline

diff --git a/Engine/Source/Runtime/Renderer/Private/SceneVisibility.cpp b/Engine/Source/Runtime/Renderer/Private/SceneVisibility.cpp
--- a/Engine/Source/Runtime/Renderer/Private/SceneVisibility.cpp
+++ b/Engine/Source/Runtime/Renderer/Private/SceneVisibility.cpp
@@ -3156,7 +3156,10 @@ void FSceneRenderer::PreVisibilityFrameSetup(FRDGBuilder& GraphBuilder, const FS
 
 	if (FXSystem && Views.IsValidIndex(0))
 	{
-		FXSystem->PreInitViews(RHICmdList, Views[0].AllowGPUParticleUpdate() && !ViewFamily.EngineShowFlags.HitProxies);
+		AddPass(GraphBuilder, [this](FRHICommandListImmediate& InRHICmdList)
+		{
+			FXSystem->PreInitViews(InRHICmdList, Views[0].AllowGPUParticleUpdate() && !ViewFamily.EngineShowFlags.HitProxies);
+		});
 	}
 
 	// Draw lines to lights affecting this mesh if its selected.
