Fixed occasional crashes in TM-ShaderModels mobile preview related to reflection captures
TextureResource can be updated inside UTexture and cached FTexture pointer inside reflection capture will get stale

diff --git a/Engine/Source/Runtime/Engine/Private/Components/ReflectionCaptureComponent.cpp b/Engine/Source/Runtime/Engine/Private/Components/ReflectionCaptureComponent.cpp
--- a/Engine/Source/Runtime/Engine/Private/Components/ReflectionCaptureComponent.cpp
+++ b/Engine/Source/Runtime/Engine/Private/Components/ReflectionCaptureComponent.cpp
@@ -1242,10 +1242,7 @@ FReflectionCaptureProxy::FReflectionCaptureProxy(const UReflectionCaptureCompone
 	Component = InComponent;
 	const FReflectionCaptureMapBuildData* MapBuildData = InComponent->GetMapBuildData();
 
-	EncodedHDRCubemap = Component->CachedEncodedHDRCubemap != nullptr ? Component->CachedEncodedHDRCubemap->Resource: nullptr;
-
-
-	
+	EncodedHDRCubemap = Component->CachedEncodedHDRCubemap;
 	EncodedHDRAverageBrightness = Component->CachedAverageBrightness;
 	MaxValueRGBM = Component->MaxValueRGBM;
 	SetTransform(InComponent->GetComponentTransform().ToMatrixWithScale());
@@ -1280,8 +1277,7 @@ void FReflectionCaptureProxy::UpdateMobileUniformBuffer()
 	FTexture* CaptureTexture = GBlackTextureCube;
 	if (EncodedHDRCubemap)
 	{
-		check(EncodedHDRCubemap->IsInitialized());
-		CaptureTexture = EncodedHDRCubemap;
+		CaptureTexture = EncodedHDRCubemap->GetResource();
 	}
 		
 	FMobileReflectionCaptureShaderParameters Parameters;
diff --git a/Engine/Source/Runtime/Engine/Public/SceneManagement.h b/Engine/Source/Runtime/Engine/Public/SceneManagement.h
--- a/Engine/Source/Runtime/Engine/Public/SceneManagement.h
+++ b/Engine/Source/Runtime/Engine/Public/SceneManagement.h
@@ -1820,7 +1820,7 @@ public:
 
 	/** Used with mobile renderer */
 	TUniformBufferRef<FMobileReflectionCaptureShaderParameters> MobileUniformBuffer;
-	FTexture* EncodedHDRCubemap;
+	UTextureCube* EncodedHDRCubemap;
 	float EncodedHDRAverageBrightness;
 	float MaxValueRGBM;
 
diff --git a/Engine/Source/Runtime/Renderer/Private/MobileBasePass.cpp b/Engine/Source/Runtime/Renderer/Private/MobileBasePass.cpp
--- a/Engine/Source/Runtime/Renderer/Private/MobileBasePass.cpp
+++ b/Engine/Source/Runtime/Renderer/Private/MobileBasePass.cpp
@@ -12,6 +12,7 @@
 #include "MaterialShaderQualitySettings.h"
 #include "PrimitiveSceneInfo.h"
 #include "MeshPassProcessor.inl"
+#include "Engine/TextureCube.h"
 
 template <ELightMapPolicyType Policy, int32 NumMovablePointLights>
 bool GetUniformMobileBasePassShaders(
@@ -496,9 +497,9 @@ void TMobileBasePassPSPolicyParamType<FUniformLightMapPolicy>::GetShaderBindings
 					{
 						CapturePositions[i] = ReflectionProxy->Position;
 						CapturePositions[i].W = ReflectionProxy->InfluenceRadius;
-						if (ReflectionProxy->EncodedHDRCubemap && ReflectionProxy->EncodedHDRCubemap->IsInitialized())
+						if (ReflectionProxy->EncodedHDRCubemap)
 						{
-							ReflectionCubemapTextures[i] = PrimitiveSceneInfo->CachedReflectionCaptureProxies[i]->EncodedHDRCubemap;
+							ReflectionCubemapTextures[i] = ReflectionProxy->EncodedHDRCubemap->GetResource();
 						}
 						//To keep ImageBasedReflectionLighting coherence with PC, use AverageBrightness instead of InvAverageBrightness to calculate the IBL contribution
 						ReflectionParams[i] = ReflectionProxy->EncodedHDRAverageBrightness;
diff --git a/Engine/Source/Runtime/Renderer/Private/RendererScene.cpp b/Engine/Source/Runtime/Renderer/Private/RendererScene.cpp
--- a/Engine/Source/Runtime/Renderer/Private/RendererScene.cpp
+++ b/Engine/Source/Runtime/Renderer/Private/RendererScene.cpp
@@ -54,7 +54,6 @@
 #include "Lumen/LumenSceneRendering.h"
 #include "PlanarReflectionSceneProxy.h"
 #include "Engine/StaticMesh.h"
-#include "Engine/TextureCube.h"
 #include "GPUSkinCache.h"
 #include "ComputeFramework/ComputeFramework.h"
 #include "DynamicShadowMapChannelBindingHelper.h"
@@ -2049,16 +2048,11 @@ void FScene::UpdateReflectionCaptureTransform(UReflectionCaptureComponent* Compo
 
 		FScene* Scene = this;
 		FReflectionCaptureProxy* Proxy = Component->SceneProxy;
-
-		FTexture* EncodedHDRCubemapResource = Component->CachedEncodedHDRCubemap != nullptr ? Component->CachedEncodedHDRCubemap->Resource : nullptr;;
 		FMatrix Transform = Component->GetComponentTransform().ToMatrixWithScale();
 
 		ENQUEUE_RENDER_COMMAND(FUpdateTransformCommand)
-			([Scene, Proxy, Transform, bUsingPreviewCaptureData, EncodedHDRCubemapResource](FRHICommandListImmediate& RHICmdList)
+			([Scene, Proxy, Transform, bUsingPreviewCaptureData](FRHICommandListImmediate& RHICmdList)
 		{
-#if WITH_EDITOR
-			Proxy->EncodedHDRCubemap = EncodedHDRCubemapResource;
-#endif
 			if (Proxy->bUsingPreviewCaptureData)
 			{
 				FPlatformAtomics::InterlockedDecrement(&Scene->NumUnbuiltReflectionCaptures);
