1.Convert mobile base passes to single pass to avoid leveraging on the RDG pass merge which could let the resource transitions occurred inside a pass.
2.Make pixel projected reflection using last frame scenecolor and scenedepth to avoid to split mobile base pass and translucency pass.
3.Support GTAO and PPR for mobile deferred renderer.

diff --git a/Engine/Shaders/Private/MobileBasePassPixelShader.usf b/Engine/Shaders/Private/MobileBasePassPixelShader.usf
--- a/Engine/Shaders/Private/MobileBasePassPixelShader.usf
+++ b/Engine/Shaders/Private/MobileBasePassPixelShader.usf
@@ -136,17 +136,6 @@ int NumDynamicPointLights;
 #endif
 
 #if !FULLY_ROUGH
-half4 GetPlanarReflection(float3 WorldPosition, half3 WorldNormal, half Roughness)
-{
-	half4 PlanarReflection = ComputePlanarReflections(WorldPosition, WorldNormal, Roughness, PlanarReflectionStruct.PlanarReflectionSampler);
-#if OUTPUT_GAMMA_SPACE
-	// the capture will also be in gamma space, convert to linear:
-	PlanarReflection.rgb *= PlanarReflection.rgb;
-#endif
-
-	return PlanarReflection;
-}
-
 // Half precision version for mobile's high quality reflection
 half MobileComputeMixingWeight(half IndirectIrradiance, half AverageBrightness, half Roughness)
 {
@@ -732,6 +721,9 @@ void Main(
 
 	half Shadow = GetPrimaryPrecomputedShadowMask(LightmapVTPageTableResult, Interpolants, MaterialParameters).r;
 
+	// Apply MaterialAO to Color since we don't have the DiffuseIndirectComposite pass on mobile deferred.
+	Color *= MaterialAO;
+
 #if DEFERRED_SHADING_PATH
 	float4 OutGBufferD;
 	float4 OutGBufferE;
@@ -758,8 +750,7 @@ void Main(
 	half3 DiffuseLookup = SkyDiffuseLighting * ResolvedView.SkyLightColor.rgb;
 	IndirectIrradiance += Luminance(DiffuseLookup);
 #endif
-			
-	Color *= MaterialAO;
+
 	IndirectIrradiance *= MaterialAO;
 
 	float  ShadowPositionZ = 0;
diff --git a/Engine/Shaders/Private/MobileDeferredShading.usf b/Engine/Shaders/Private/MobileDeferredShading.usf
--- a/Engine/Shaders/Private/MobileDeferredShading.usf
+++ b/Engine/Shaders/Private/MobileDeferredShading.usf
@@ -8,6 +8,7 @@
 #include "MobileShadingModels.ush"
 #include "IESLightProfilesCommon.ush"
 #include "LightFunctionCommon.ush"
+#include "PlanarReflectionShared.ush"
 
 #if APPLY_SKY_REFLECTION
 #define ENABLE_SKY_LIGHT 1
@@ -238,6 +239,8 @@ void MobileDirectLightPS(
 	float4 SvPosition : SV_POSITION, 
 	out HALF4_TYPE OutColor : SV_Target0)
 {
+	ResolvedView = ResolveView();
+
 	FGBufferData GBuffer = (FGBufferData)0;
 	float SceneDepth = 0; 
 	{
@@ -313,6 +316,14 @@ void MobileDirectLightPS(
 		DataStartIndex,
 		0,
 		true);
+
+	BRANCH
+	if (abs(dot(PlanarReflectionStruct.ReflectionPlane.xyz, 1)) > .0001f)
+	{
+		half4 PlanarReflection = GetPlanarReflection(WorldPosition, GBuffer.WorldNormal, GBuffer.Roughness);
+		// Planar reflections win over reflection environment
+		SpecularIBL = lerp(SpecularIBL, PlanarReflection.rgb, PlanarReflection.a);
+	}
 		
 	Color += SpecularIBL * ShadingModelContext.SpecularColor;
 #elif APPLY_SKY_REFLECTION
diff --git a/Engine/Shaders/Private/PlanarReflectionShared.ush b/Engine/Shaders/Private/PlanarReflectionShared.ush
--- a/Engine/Shaders/Private/PlanarReflectionShared.ush
+++ b/Engine/Shaders/Private/PlanarReflectionShared.ush
@@ -72,3 +72,16 @@ half4 ComputePlanarReflections(float3 WorldPosition, half3 WorldNormal, half Rou
 
 	return OutPlanarReflection;
 }
+
+#if (FEATURE_LEVEL <= FEATURE_LEVEL_ES3_1)
+half4 GetPlanarReflection(float3 WorldPosition, half3 WorldNormal, half Roughness)
+{
+	half4 PlanarReflection = ComputePlanarReflections(WorldPosition, WorldNormal, Roughness, PlanarReflectionStruct.PlanarReflectionSampler);
+#if OUTPUT_GAMMA_SPACE
+	// the capture will also be in gamma space, convert to linear:
+	PlanarReflection.rgb *= PlanarReflection.rgb;
+#endif
+
+	return PlanarReflection;
+}
+#endif
\ No newline at end of file
diff --git a/Engine/Shaders/Private/PostProcessPixelProjectedReflectionMobile.usf b/Engine/Shaders/Private/PostProcessPixelProjectedReflectionMobile.usf
--- a/Engine/Shaders/Private/PostProcessPixelProjectedReflectionMobile.usf
+++ b/Engine/Shaders/Private/PostProcessPixelProjectedReflectionMobile.usf
@@ -161,7 +161,7 @@ void ProjectionPassCS(
 
 		float PlaneDistance = dot(ReflectionPlane, float4(ViewPosition, -1.0f));
 
-		if (PlaneDistance > 0.01f )
+		if (PlaneDistance > 0.5f )
 		{
 			float3 MirroredViewPosition = ViewPosition - ReflectionPlane.xyz * (2.f * PlaneDistance);
 
diff --git a/Engine/Source/Runtime/Renderer/Private/CompositionLighting/PostProcessMeshDecals.cpp b/Engine/Source/Runtime/Renderer/Private/CompositionLighting/PostProcessMeshDecals.cpp
--- a/Engine/Source/Runtime/Renderer/Private/CompositionLighting/PostProcessMeshDecals.cpp
+++ b/Engine/Source/Runtime/Renderer/Private/CompositionLighting/PostProcessMeshDecals.cpp
@@ -359,40 +359,32 @@ void RenderMeshDecals(
 	}
 }
 
-void RenderMeshDecalsMobile(FRDGBuilder& GraphBuilder, FRenderTargetBindingSlots& BasePassRenderTargets, const FViewInfo& View)
+void RenderMeshDecalsMobile(FRHICommandListImmediate& RHICmdList, const FViewInfo& View)
 {
 	FDecalRenderingCommon::ERenderTargetMode RenderTargetMode = IsMobileDeferredShadingEnabled(View.GetShaderPlatform()) ? 
 		FDecalRenderingCommon::RTM_SceneColorAndGBufferWithNormal : 
 		FDecalRenderingCommon::RTM_SceneColor;
 
-	auto* PassParameters = GraphBuilder.AllocParameters<FMobileBasePassParameters>();
-	PassParameters->View = View.GetShaderParameters();
-	PassParameters->RenderTargets = BasePassRenderTargets;
+	FGraphicsPipelineStateInitializer GraphicsPSOInit;
+	RHICmdList.SetViewport(View.ViewRect.Min.X, View.ViewRect.Min.Y, 0, View.ViewRect.Max.X, View.ViewRect.Max.Y, 1);
+	RHICmdList.ApplyCachedRenderTargets(GraphicsPSOInit);
 
-	GraphBuilder.AddPass(RDG_EVENT_NAME("MeshDecalsPass"), PassParameters, ERDGPassFlags::Raster | ERDGPassFlags::SkipRenderPass,
-		[RenderTargetMode, &View](FRHICommandList& RHICmdList)
+	DrawDynamicMeshPass(View, RHICmdList, [&View, RenderTargetMode](FDynamicPassMeshDrawListContext* DynamicMeshPassContext)
 	{
-		FGraphicsPipelineStateInitializer GraphicsPSOInit;
-		RHICmdList.SetViewport(View.ViewRect.Min.X, View.ViewRect.Min.Y, 0, View.ViewRect.Max.X, View.ViewRect.Max.Y, 1);
-		RHICmdList.ApplyCachedRenderTargets(GraphicsPSOInit);
-
-		DrawDynamicMeshPass(View, RHICmdList, [&View, RenderTargetMode](FDynamicPassMeshDrawListContext* DynamicMeshPassContext)
+		FMeshDecalMeshProcessor PassMeshProcessor(
+			View.Family->Scene->GetRenderScene(),
+			&View,
+			DRS_Mobile,
+			RenderTargetMode,
+			DynamicMeshPassContext);
+
+		for (int32 MeshBatchIndex = 0; MeshBatchIndex < View.MeshDecalBatches.Num(); ++MeshBatchIndex)
 		{
-			FMeshDecalMeshProcessor PassMeshProcessor(
-				View.Family->Scene->GetRenderScene(),
-				&View,
-				DRS_Mobile,
-				RenderTargetMode,
-				DynamicMeshPassContext);
+			const FMeshBatch* Mesh = View.MeshDecalBatches[MeshBatchIndex].Mesh;
+			const FPrimitiveSceneProxy* PrimitiveSceneProxy = View.MeshDecalBatches[MeshBatchIndex].Proxy;
+			const uint64 DefaultBatchElementMask = ~0ull;
 
-			for (int32 MeshBatchIndex = 0; MeshBatchIndex < View.MeshDecalBatches.Num(); ++MeshBatchIndex)
-			{
-				const FMeshBatch* Mesh = View.MeshDecalBatches[MeshBatchIndex].Mesh;
-				const FPrimitiveSceneProxy* PrimitiveSceneProxy = View.MeshDecalBatches[MeshBatchIndex].Proxy;
-				const uint64 DefaultBatchElementMask = ~0ull;
-
-				PassMeshProcessor.AddMeshBatch(*Mesh, DefaultBatchElementMask, PrimitiveSceneProxy);
-			}
-		}, true);
-	});
+			PassMeshProcessor.AddMeshBatch(*Mesh, DefaultBatchElementMask, PrimitiveSceneProxy);
+		}
+	}, true);
 }
\ No newline at end of file
diff --git a/Engine/Source/Runtime/Renderer/Private/DeferredShadingRenderer.cpp b/Engine/Source/Runtime/Renderer/Private/DeferredShadingRenderer.cpp
--- a/Engine/Source/Runtime/Renderer/Private/DeferredShadingRenderer.cpp
+++ b/Engine/Source/Runtime/Renderer/Private/DeferredShadingRenderer.cpp
@@ -322,9 +322,9 @@ TGlobalResource<FGlobalDynamicReadBuffer> FDeferredShadingSceneRenderer::Dynamic
 /**
  * Returns true if the depth Prepass needs to run
  */
-static FORCEINLINE bool NeedsPrePass(const FDeferredShadingSceneRenderer* Renderer)
+bool FDeferredShadingSceneRenderer::ShouldRenderPrePass() const
 {
-	return (Renderer->DepthPass.EarlyZPassMode != DDM_None || Renderer->DepthPass.bEarlyZPassMovable != 0);
+	return (DepthPass.EarlyZPassMode != DDM_None || DepthPass.bEarlyZPassMovable != 0);
 }
 
 bool FDeferredShadingSceneRenderer::RenderHzb(FRDGBuilder& GraphBuilder, FRDGTextureRef SceneDepthTexture)
@@ -1767,7 +1767,7 @@ void FDeferredShadingSceneRenderer::Render(FRDGBuilder& GraphBuilder)
 
 	
 	const bool bIsOcclusionTesting = DoOcclusionQueries(FeatureLevel) && (!ViewFamily.EngineShowFlags.Wireframe || bIsViewFrozen || bHasViewParent);
-	const bool bNeedsPrePass = NeedsPrePass(this);
+	const bool bNeedsPrePass = ShouldRenderPrePass();
 
 	GEngine->GetPreRenderDelegateEx().Broadcast(GraphBuilder);
 
diff --git a/Engine/Source/Runtime/Renderer/Private/DeferredShadingRenderer.h b/Engine/Source/Runtime/Renderer/Private/DeferredShadingRenderer.h
--- a/Engine/Source/Runtime/Renderer/Private/DeferredShadingRenderer.h
+++ b/Engine/Source/Runtime/Renderer/Private/DeferredShadingRenderer.h
@@ -253,6 +253,8 @@ public:
 
 	virtual bool ShouldRenderVelocities() const override;
 
+	virtual bool ShouldRenderPrePass() const override;
+
 #if !(UE_BUILD_SHIPPING || UE_BUILD_TEST)
 	void RenderVisualizeTexturePool(FRHICommandListImmediate& RHICmdList);
 #endif
diff --git a/Engine/Source/Runtime/Renderer/Private/DepthRendering.cpp b/Engine/Source/Runtime/Renderer/Private/DepthRendering.cpp
--- a/Engine/Source/Runtime/Renderer/Private/DepthRendering.cpp
+++ b/Engine/Source/Runtime/Renderer/Private/DepthRendering.cpp
@@ -607,53 +607,29 @@ void FDeferredShadingSceneRenderer::RenderPrePass(FRDGBuilder& GraphBuilder, FRD
 	}
 }
 
-BEGIN_SHADER_PARAMETER_STRUCT(FMobileDepthPassParameters, )
-	SHADER_PARAMETER_STRUCT_INCLUDE(FViewShaderParameters, View)
-	SHADER_PARAMETER_STRUCT_INCLUDE(FInstanceCullingDrawParams, InstanceCullingDrawParams)
-	RENDER_TARGET_BINDING_SLOTS()
-END_SHADER_PARAMETER_STRUCT()
-
-void FMobileSceneRenderer::RenderPrePass(FRDGBuilder& GraphBuilder, FRenderTargetBindingSlots& BasePassRenderTargets)
+bool FMobileSceneRenderer::ShouldRenderPrePass() const
 {
-	SCOPED_NAMED_EVENT(FMobileSceneRenderer_RenderPrePass, FColor::Emerald);
-	RDG_EVENT_SCOPE(GraphBuilder, "MobileRenderPrePass");
-
-	SCOPE_CYCLE_COUNTER(STAT_DepthDrawTime);
-	CSV_SCOPED_TIMING_STAT_EXCLUSIVE(RenderPrePass);
-	RDG_GPU_STAT_SCOPE(GraphBuilder, Prepass);
-
 	// Draw a depth pass to avoid overdraw in the other passes.
 	// Mobile only does MaskedOnly DepthPass for the moment
-	if (Scene->EarlyZPassMode == DDM_MaskedOnly)
-	{
-		for (int32 ViewIndex = 0; ViewIndex < Views.Num(); ViewIndex++)
-		{
-			FViewInfo& View = Views[ViewIndex];
+	bool bShouldRenderPrePass = Scene->EarlyZPassMode == DDM_MaskedOnly;
 
-			RDG_GPU_MASK_SCOPE(GraphBuilder, !View.IsInstancedStereoPass() ? View.GPUMask : (Views[0].GPUMask | Views[1].GPUMask));
-			RDG_EVENT_SCOPE_CONDITIONAL(GraphBuilder, Views.Num() > 1, "View%d", ViewIndex);
-
-			if (!View.ShouldRenderView())
-			{
-				continue;
-			}
+	return bShouldRenderPrePass;
+}
 
-			View.BeginRenderView();
+void FMobileSceneRenderer::RenderPrePass(FRHICommandListImmediate& RHICmdList, const FViewInfo& View, const FInstanceCullingDrawParams* InstanceCullingDrawParams/* = nullptr*/)
+{
+	checkSlow(RHICmdList.IsInsideRenderPass());
+	checkSlow(ShouldRenderPrePass());
 
-			auto* PassParameters = GraphBuilder.AllocParameters<FMobileDepthPassParameters>();
-			PassParameters->View = View.GetShaderParameters();
-			View.ParallelMeshDrawCommandPasses[EMeshPass::DepthPass].BuildRenderingCommands(GraphBuilder, Scene->GPUScene, PassParameters->InstanceCullingDrawParams);
+	SCOPED_NAMED_EVENT(FMobileSceneRenderer_RenderPrePass, FColor::Emerald);
+	SCOPED_DRAW_EVENT(RHICmdList, MobileRenderPrePass);
 
-			PassParameters->RenderTargets = BasePassRenderTargets;
+	SCOPE_CYCLE_COUNTER(STAT_DepthDrawTime);
+	CSV_SCOPED_TIMING_STAT_EXCLUSIVE(RenderPrePass);
+	SCOPED_GPU_STAT(RHICmdList, Prepass);
 
-			GraphBuilder.AddPass(RDG_EVENT_NAME("RenderPrePass"), PassParameters, ERDGPassFlags::Raster | ERDGPassFlags::SkipRenderPass,
-				[this, &View, PassParameters](FRHICommandListImmediate& RHICmdList)
-			{
-				SetStereoViewport(RHICmdList, View);
-				View.ParallelMeshDrawCommandPasses[EMeshPass::DepthPass].DispatchDraw(nullptr, RHICmdList, &PassParameters->InstanceCullingDrawParams);
-			});
-		}
-	}
+	SetStereoViewport(RHICmdList, View);
+	View.ParallelMeshDrawCommandPasses[EMeshPass::DepthPass].DispatchDraw(nullptr, RHICmdList, InstanceCullingDrawParams);
 }
 
 void FDeferredShadingSceneRenderer::RenderPrePassHMD(FRDGBuilder& GraphBuilder, FRDGTextureRef DepthTexture)
diff --git a/Engine/Source/Runtime/Renderer/Private/MobileBasePass.cpp b/Engine/Source/Runtime/Renderer/Private/MobileBasePass.cpp
--- a/Engine/Source/Runtime/Renderer/Private/MobileBasePass.cpp
+++ b/Engine/Source/Runtime/Renderer/Private/MobileBasePass.cpp
@@ -399,19 +399,6 @@ void MobileBasePass::SetTranslucentRenderState(FMeshPassProcessorRenderState& Dr
 				// Single layer water is an opaque marerial rendered as translucent on Mobile. We force pre-multiplied alpha to achieve water depth based transmittance.
 				DrawRenderState.SetBlendState(TStaticBlendState<CW_RGB, BO_Add, BF_One, BF_InverseSourceAlpha, BO_Add, BF_Zero, BF_InverseSourceAlpha>::GetRHI());
 			}
-			else if (bIsUsingMobilePixelProjectedReflection)
-			{
-				if (GetMobilePixelProjectedReflectionQuality() == EMobilePixelProjectedReflectionQuality::BestPerformance)
-				{
-					// We only render the meshes used for mobile pixel projected reflection once for the BestPerformance quality level, so we have to write the depth to alpha channel in translucent pass.
-					DrawRenderState.SetBlendState(TStaticBlendStateWriteMask<CW_RGBA>::GetRHI());
-				}
-				else
-				{
-					// We render the meshes used for mobile pixel projected reflection twice for the BetterPerformance or greater quality level, so we don't have to write the depth to alpha channel in translucent pass.
-					DrawRenderState.SetBlendState(TStaticBlendStateWriteMask<CW_RGB>::GetRHI());
-				}
-			}
 			else
 			{
 				check(0);
@@ -423,14 +410,6 @@ void MobileBasePass::SetTranslucentRenderState(FMeshPassProcessorRenderState& Dr
 	{
 		DrawRenderState.SetDepthStencilState(TStaticDepthStencilState<false, CF_Always>::GetRHI());
 	}
-	else if (bIsUsingMobilePixelProjectedReflection && !IsTranslucentBlendMode(Material.GetBlendMode()))
-	{
-		if (GetMobilePixelProjectedReflectionQuality() == EMobilePixelProjectedReflectionQuality::BestPerformance)
-		{
-			// We only render the opaque meshes used for mobile pixel projected reflection once for the BestPerformance quality level, so we have to write the depth to depth RT in translucent pass.
-			DrawRenderState.SetDepthStencilState(TStaticDepthStencilState<true, CF_DepthNearOrEqual>::GetRHI());
-		}
-	}
 }
 
 bool MobileBasePass::StationarySkyLightHasBeenApplied(const FScene* Scene, ELightMapPolicyType LightMapPolicyType)
@@ -619,15 +598,13 @@ bool FMobileBasePassMeshProcessor::TryAddMeshBatch(const FMeshBatch& RESTRICT Me
 	const FMaterialShadingModelField ShadingModels = Material.GetShadingModels();
 	const bool bIsTranslucent = IsTranslucentBlendMode(BlendMode);
 	const bool bUsesWaterMaterial = ShadingModels.HasShadingModel(MSM_SingleLayerWater); // Water goes into the translucent pass
-	const bool bIsUsingMobilePixelProjectedReflection = Material.IsUsingPlanarForwardReflections() 
-														&& IsUsingMobilePixelProjectedReflection(GetFeatureLevelShaderPlatform(Material.GetFeatureLevel()));
 	const bool bCanReceiveCSM = ((Flags & EFlags::CanReceiveCSM) == EFlags::CanReceiveCSM);
 
 	bool bResult = true;
 	if (bTranslucentBasePass)
 	{
 		// Skipping TPT_TranslucencyAfterDOFModulate. That pass is only needed for Dual Blending, which is not supported on Mobile.
-		bool bShouldDraw = (bIsTranslucent || bUsesWaterMaterial || bIsUsingMobilePixelProjectedReflection) &&
+		bool bShouldDraw = (bIsTranslucent || bUsesWaterMaterial) &&
 		(TranslucencyPassType == ETranslucencyPass::TPT_AllTranslucency
 		|| (TranslucencyPassType == ETranslucencyPass::TPT_StandardTranslucency && !Material.IsMobileSeparateTranslucencyEnabled())
 		|| (TranslucencyPassType == ETranslucencyPass::TPT_TranslucencyAfterDOF && Material.IsMobileSeparateTranslucencyEnabled()));
@@ -637,15 +614,14 @@ bool FMobileBasePassMeshProcessor::TryAddMeshBatch(const FMeshBatch& RESTRICT Me
 			check(bCanReceiveCSM == false);
 			const FLightSceneInfo* MobileDirectionalLight = MobileBasePass::GetDirectionalLightInfo(Scene, PrimitiveSceneProxy);
 			// Opaque meshes used for mobile pixel projected reflection could receive CSM in translucent pass.
-			ELightMapPolicyType LightmapPolicyType = MobileBasePass::SelectMeshLightmapPolicy(Scene, MeshBatch, PrimitiveSceneProxy, MobileDirectionalLight, ShadingModels, bCanReceiveCSM || (!bIsTranslucent && bIsUsingMobilePixelProjectedReflection), false, FeatureLevel, BlendMode);
+			ELightMapPolicyType LightmapPolicyType = MobileBasePass::SelectMeshLightmapPolicy(Scene, MeshBatch, PrimitiveSceneProxy, MobileDirectionalLight, ShadingModels, bCanReceiveCSM, false, FeatureLevel, BlendMode);
 			bResult = Process(MeshBatch, BatchElementMask, StaticMeshId, PrimitiveSceneProxy, MaterialRenderProxy, Material, BlendMode, ShadingModels, LightmapPolicyType, bCanReceiveCSM, MeshBatch.LCI);
 		}
 	}
 	else
 	{
 		// opaque materials.
-		// We have to render the opaque meshes used for mobile pixel projected reflection both in opaque and translucent pass if the quality level is greater than BestPerformance
-		if (!bIsTranslucent && !bUsesWaterMaterial && (!bIsUsingMobilePixelProjectedReflection || GetMobilePixelProjectedReflectionQuality() > EMobilePixelProjectedReflectionQuality::BestPerformance))
+		if (!bIsTranslucent && !bUsesWaterMaterial)
 		{
 			const FLightSceneInfo* MobileDirectionalLight = MobileBasePass::GetDirectionalLightInfo(Scene, PrimitiveSceneProxy);
 			ELightMapPolicyType LightmapPolicyType = MobileBasePass::SelectMeshLightmapPolicy(Scene, MeshBatch, PrimitiveSceneProxy, MobileDirectionalLight, ShadingModels, bCanReceiveCSM, bUsesDeferredShading, FeatureLevel, BlendMode);
diff --git a/Engine/Source/Runtime/Renderer/Private/MobileBasePassRendering.cpp b/Engine/Source/Runtime/Renderer/Private/MobileBasePassRendering.cpp
--- a/Engine/Source/Runtime/Renderer/Private/MobileBasePassRendering.cpp
+++ b/Engine/Source/Runtime/Renderer/Private/MobileBasePassRendering.cpp
@@ -43,7 +43,9 @@ static TAutoConsoleVariable<int32> CVarMobileMaxVisibleMovableSpotLightsShadow(
 	TEXT("The max number of visible spotlighs can cast shadow sorted by screen size, should be as less as possible for performance reason"),
 	ECVF_RenderThreadSafe);
 
-IMPLEMENT_STATIC_UNIFORM_BUFFER_STRUCT(FMobileBasePassUniformParameters, "MobileBasePass", SceneTextures);
+// Specify a unique slot for mobile base pass because some rendering in the mobile base pass (e.g. modulated shadow and ViewExtensions) use SceneTextures uniform buffer, but the SceneTextures uniform buffer and MobileBasePass uniform buffer share the same slot
+IMPLEMENT_STATIC_UNIFORM_BUFFER_SLOT(MobileBasePass);
+IMPLEMENT_STATIC_UNIFORM_BUFFER_STRUCT(FMobileBasePassUniformParameters, "MobileBasePass", MobileBasePass);
 
 static TAutoConsoleVariable<int32> CVarMobileUseHWsRGBEncoding(
 	TEXT("r.Mobile.UseHWsRGBEncoding"),
@@ -142,6 +144,7 @@ void SetupMobileBasePassUniformParameters(
 	const FViewInfo& View,
 	EMobileBasePass BasePass,
 	FRDGTextureRef ScreenSpaceAOTexture,
+	FTextureRHIRef PixelProjectedReflectionTexture,
 	FMobileBasePassUniformParameters& BasePassParameters)
 {
 	SetupFogUniformParameters(GraphBuilder, View, BasePassParameters.Fog);
@@ -167,32 +170,15 @@ void SetupMobileBasePassUniformParameters(
 	BasePassParameters.PreIntegratedGFTexture = GSystemTextures.PreintegratedGF->GetRenderTargetItem().ShaderResourceTexture;
 	BasePassParameters.PreIntegratedGFSampler = TStaticSamplerState<SF_Bilinear, AM_Clamp, AM_Clamp, AM_Clamp>::GetRHI();
 
-	if (GPixelProjectedReflectionMobileOutputs.IsValid())
+	if (PixelProjectedReflectionTexture.IsValid())
 	{
-		if (BasePass == EMobileBasePass::Translucent)
-		{
-			BasePassParameters.PlanarReflection.PlanarReflectionTexture = GPixelProjectedReflectionMobileOutputs.PixelProjectedReflectionTexture->GetRenderTargetItem().ShaderResourceTexture;
-			if (GetMobilePixelProjectedReflectionQuality() <= EMobilePixelProjectedReflectionQuality::BestPerformance)
-			{
-				// We only render the meshes used for pixel projected reflection once and it could cause color bleeding artifact if we use bilinear filter.
-				BasePassParameters.PlanarReflection.PlanarReflectionSampler = TStaticSamplerState<SF_Point, AM_Clamp, AM_Clamp, AM_Clamp>::GetRHI();
-			}
-			else
-			{
-				// We render the meshes used for pixel projected reflection twice, so we could use bilinear filter.
-				BasePassParameters.PlanarReflection.PlanarReflectionSampler = TStaticSamplerState<SF_Bilinear, AM_Clamp, AM_Clamp, AM_Clamp>::GetRHI();
-			}
-		}
-		else
-		{
-			// Clear the ReflectionPlane to skip planar reflection on opaque mesh when the PPR is enabled because we render the reflection meshes used for pixel projected reflection in the translucent pass.
-			BasePassParameters.PlanarReflection.ReflectionPlane.Set(0.0f, 0.0f, 0.0f, 0.0f);
-		}
+		BasePassParameters.PlanarReflection.PlanarReflectionTexture = PixelProjectedReflectionTexture;
+		BasePassParameters.PlanarReflection.PlanarReflectionSampler = TStaticSamplerState<SF_Bilinear, AM_Clamp, AM_Clamp, AM_Clamp>::GetRHI();
 	}
 
 	BasePassParameters.EyeAdaptationBuffer = GraphBuilder.CreateSRV(GetEyeAdaptationBuffer(GraphBuilder, View), PF_A32B32G32R32F);
 
-	if (BasePass == EMobileBasePass::Opaque && HasBeenProduced(ScreenSpaceAOTexture))
+	if (BasePass == EMobileBasePass::Opaque && ScreenSpaceAOTexture != nullptr)
 	{
 		BasePassParameters.AmbientOcclusionTexture = ScreenSpaceAOTexture;
 	}
@@ -222,10 +208,11 @@ TRDGUniformBufferRef<FMobileBasePassUniformParameters> CreateMobileBasePassUnifo
 	FRDGBuilder& GraphBuilder,
 	const FViewInfo& View,
 	EMobileBasePass BasePass,
-	FRDGTextureRef ScreenSpaceAO)
+	FRDGTextureRef ScreenSpaceAO,
+	FTextureRHIRef PixelProjectedReflectionTexture)
 {
 	FMobileBasePassUniformParameters* BasePassParameters = GraphBuilder.AllocParameters<FMobileBasePassUniformParameters>();
-	SetupMobileBasePassUniformParameters(GraphBuilder, View, BasePass, ScreenSpaceAO, *BasePassParameters);
+	SetupMobileBasePassUniformParameters(GraphBuilder, View, BasePass, ScreenSpaceAO, PixelProjectedReflectionTexture, *BasePassParameters);
 	return GraphBuilder.CreateUniformBuffer(BasePassParameters);
 }
 
@@ -312,57 +299,27 @@ void SetupMobileSkyReflectionUniformParameters(FSkyLightSceneProxy* SkyLight, FM
 	Parameters.TextureSampler = CaptureTexture->SamplerStateRHI;
 }
 
-void FMobileSceneRenderer::RenderMobileBasePass(
-	FRDGBuilder& GraphBuilder,
-	FRenderTargetBindingSlots& BasePassRenderTargets,
-	const TArrayView<const FViewInfo> PassViews,
-	FRDGTextureRef ScreenSpaceAO)
+void FMobileSceneRenderer::RenderMobileBasePass(FRHICommandListImmediate& RHICmdList, const FViewInfo& View)
 {
 	CSV_SCOPED_TIMING_STAT_EXCLUSIVE(RenderBasePass);
-	RDG_EVENT_SCOPE(GraphBuilder, "MobileBasePass");
+	SCOPED_DRAW_EVENT(RHICmdList, MobileBasePass);
 	SCOPE_CYCLE_COUNTER(STAT_BasePassDrawTime);
-	RDG_GPU_STAT_SCOPE(GraphBuilder, Basepass);
-
-	for (int32 ViewIndex = 0; ViewIndex < PassViews.Num(); ViewIndex++)
-	{
-		RDG_EVENT_SCOPE_CONDITIONAL(GraphBuilder, Views.Num() > 1, "View%d", ViewIndex);
-		const FViewInfo& View = PassViews[ViewIndex];
-		if (!View.ShouldRenderView())
-		{
-			continue;
-		}
-
-		View.BeginRenderView();
+	SCOPED_GPU_STAT(RHICmdList, Basepass);
 
-		UpdateDirectionalLightUniformBuffers(GraphBuilder, View);
+	RHICmdList.SetViewport(View.ViewRect.Min.X, View.ViewRect.Min.Y, 0, View.ViewRect.Max.X, View.ViewRect.Max.Y, 1);
+	View.ParallelMeshDrawCommandPasses[EMeshPass::BasePass].DispatchDraw(nullptr, RHICmdList);
 
-		auto* OpaqueBasePassParameters = GraphBuilder.AllocParameters<FMobileBasePassParameters>();
-		OpaqueBasePassParameters->View = View.GetShaderParameters();
-		OpaqueBasePassParameters->RenderTargets = BasePassRenderTargets;
-		OpaqueBasePassParameters->MobileBasePass = CreateMobileBasePassUniformBuffer(GraphBuilder, View, EMobileBasePass::Opaque, ScreenSpaceAO);
-
-		GraphBuilder.AddPass(RDG_EVENT_NAME("RenderOpaqueBasePass"), OpaqueBasePassParameters, ERDGPassFlags::Raster | ERDGPassFlags::SkipRenderPass,
-			[this, &View, OpaqueBasePassParameters](FRHICommandListImmediate& RHICmdList)
-		{
-			RHICmdList.SetViewport(View.ViewRect.Min.X, View.ViewRect.Min.Y, 0, View.ViewRect.Max.X, View.ViewRect.Max.Y, 1);
-			View.ParallelMeshDrawCommandPasses[EMeshPass::BasePass].DispatchDraw(nullptr, RHICmdList);
-		
-		RHICmdList.SetViewport(View.ViewRect.Min.X, View.ViewRect.Min.Y, 0, View.ViewRect.Max.X, View.ViewRect.Max.Y, 1);
-		View.ParallelMeshDrawCommandPasses[EMeshPass::BasePass].DispatchDraw(nullptr, RHICmdList);
-		
-		if (View.Family->EngineShowFlags.Atmosphere)
-		{
-			View.ParallelMeshDrawCommandPasses[EMeshPass::SkyPass].DispatchDraw(nullptr, RHICmdList);
-		}
-
-		// editor primitives
-		FMeshPassProcessorRenderState DrawRenderState;
-		DrawRenderState.SetBlendState(TStaticBlendStateWriteMask<CW_RGBA>::GetRHI());
-		DrawRenderState.SetDepthStencilAccess(Scene->DefaultBasePassDepthStencilAccess);
-		DrawRenderState.SetDepthStencilState(TStaticDepthStencilState<true, CF_DepthNearOrEqual>::GetRHI());
-		RenderMobileEditorPrimitives(RHICmdList, View, DrawRenderState);
-		});
+	if (View.Family->EngineShowFlags.Atmosphere)
+	{
+		View.ParallelMeshDrawCommandPasses[EMeshPass::SkyPass].DispatchDraw(nullptr, RHICmdList);
 	}
+
+	// editor primitives
+	FMeshPassProcessorRenderState DrawRenderState;
+	DrawRenderState.SetBlendState(TStaticBlendStateWriteMask<CW_RGBA>::GetRHI());
+	DrawRenderState.SetDepthStencilAccess(Scene->DefaultBasePassDepthStencilAccess);
+	DrawRenderState.SetDepthStencilState(TStaticDepthStencilState<true, CF_DepthNearOrEqual>::GetRHI());
+	RenderMobileEditorPrimitives(RHICmdList, View, DrawRenderState);
 }
 
 void FMobileSceneRenderer::RenderMobileEditorPrimitives(FRHICommandList& RHICmdList, const FViewInfo& View, const FMeshPassProcessorRenderState& DrawRenderState)
diff --git a/Engine/Source/Runtime/Renderer/Private/MobileBasePassRendering.h b/Engine/Source/Runtime/Renderer/Private/MobileBasePassRendering.h
--- a/Engine/Source/Runtime/Renderer/Private/MobileBasePassRendering.h
+++ b/Engine/Source/Runtime/Renderer/Private/MobileBasePassRendering.h
@@ -61,7 +61,8 @@ extern TRDGUniformBufferRef<FMobileBasePassUniformParameters> CreateMobileBasePa
 	FRDGBuilder& GraphBuilder,
 	const FViewInfo& View,
 	EMobileBasePass BasePass,
-	FRDGTextureRef ScreenSpaceAO = nullptr);
+	FRDGTextureRef ScreenSpaceAO = nullptr,
+	FTextureRHIRef PixelProjectedReflectionTexture = nullptr);
 
 extern void SetupMobileDirectionalLightUniformParameters(
 	const FScene& Scene,
diff --git a/Engine/Source/Runtime/Renderer/Private/MobileDecalRendering.cpp b/Engine/Source/Runtime/Renderer/Private/MobileDecalRendering.cpp
--- a/Engine/Source/Runtime/Renderer/Private/MobileDecalRendering.cpp
+++ b/Engine/Source/Runtime/Renderer/Private/MobileDecalRendering.cpp
@@ -15,15 +15,9 @@
 #include "ScenePrivate.h"
 #include "DecalRenderingShared.h"
 
-BEGIN_SHADER_PARAMETER_STRUCT(FMobileDecalPassParameters, )
-	SHADER_PARAMETER_STRUCT_REF(FViewUniformShaderParameters, View)
-	SHADER_PARAMETER_RDG_UNIFORM_BUFFER(FMobileSceneTextureUniformParameters, MobileSceneTextures)
-	RENDER_TARGET_BINDING_SLOTS()
-END_SHADER_PARAMETER_STRUCT()
-
 extern FRHIRasterizerState* GetDecalRasterizerState(EDecalRasterizerState DecalRasterizerState);
-extern void RenderMeshDecalsMobile(FRDGBuilder& GraphBuilder, FRenderTargetBindingSlots& BasePassRenderTargets, const FViewInfo& View);
-void RenderDeferredDecalsMobile(FRDGBuilder& GraphBuilder, FRenderTargetBindingSlots& BasePassRenderTargets, TRDGUniformBufferRef<FMobileSceneTextureUniformParameters> MobileSceneTextures, const FScene& Scene, const FViewInfo& View);
+extern void RenderMeshDecalsMobile(FRHICommandListImmediate& RHICmdList, const FViewInfo& View);
+void RenderDeferredDecalsMobile(FRHICommandListImmediate& RHICmdList, const FScene& Scene, const FViewInfo& View);
 
 FRHIBlendState* MobileForward_GetDecalBlendState(EDecalBlendMode DecalBlendMode)
 {
@@ -191,7 +185,7 @@ FRHIBlendState* MobileDeferred_GetDecalBlendState(EDecalBlendMode DecalBlendMode
 	return TStaticBlendState<>::GetRHI(); 
 }
 
-void FMobileSceneRenderer::RenderDecals(FRDGBuilder& GraphBuilder, FRenderTargetBindingSlots& BasePassRenderTargets, TRDGUniformBufferRef<FMobileSceneTextureUniformParameters> MobileSceneTextures)
+void FMobileSceneRenderer::RenderDecals(FRHICommandListImmediate& RHICmdList, const FViewInfo& View)
 {
 	if (!IsMobileHDR())
 	{
@@ -203,107 +197,88 @@ void FMobileSceneRenderer::RenderDecals(FRDGBuilder& GraphBuilder, FRenderTarget
 	// Deferred decals
 	if (Scene->Decals.Num() > 0)
 	{
-		for (int32 ViewIndex = 0; ViewIndex < Views.Num(); ViewIndex++)
-		{
-			const FViewInfo& View = Views[ViewIndex];
-			RenderDeferredDecalsMobile(GraphBuilder, BasePassRenderTargets, MobileSceneTextures, *Scene, View);
-		}
+		RenderDeferredDecalsMobile(RHICmdList, *Scene, View);
 	}
 
 	// Mesh decals
-	for (int32 ViewIndex = 0; ViewIndex < Views.Num(); ViewIndex++)
+	if (View.MeshDecalBatches.Num() > 0)
 	{
-		const FViewInfo& View = Views[ViewIndex];
-		if (View.MeshDecalBatches.Num() > 0)
-		{
-			RenderMeshDecalsMobile(GraphBuilder, BasePassRenderTargets, View);
-		}
+		RenderMeshDecalsMobile(RHICmdList, View);
 	}
 }
 
-void RenderDeferredDecalsMobile(FRDGBuilder& GraphBuilder, FRenderTargetBindingSlots& BasePassRenderTargets, TRDGUniformBufferRef<FMobileSceneTextureUniformParameters> MobileSceneTextures, const FScene& Scene, const FViewInfo& View)
+void RenderDeferredDecalsMobile(FRHICommandListImmediate& RHICmdList, const FScene& Scene, const FViewInfo& View)
 {
 	const uint32 DecalCount = Scene.Decals.Num();
 	int32 SortedDecalCount = 0;
-	FTransientDecalRenderDataList* SortedDecals = nullptr;
+	FTransientDecalRenderDataList SortedDecals;
 
 	if (DecalCount > 0)
 	{
 		// Build a list of decals that need to be rendered for this view
-		SortedDecals = GraphBuilder.AllocObject<FTransientDecalRenderDataList>();
-		FDecalRendering::BuildVisibleDecalList(Scene, View, DRS_Mobile, SortedDecals);
-		SortedDecalCount = SortedDecals->Num();
+		FDecalRendering::BuildVisibleDecalList(Scene, View, DRS_Mobile, &SortedDecals);
+		SortedDecalCount = SortedDecals.Num();
 		INC_DWORD_STAT_BY(STAT_Decals, SortedDecalCount);
 	}
 
 	if (SortedDecalCount > 0)
 	{
-		auto* PassParameters = GraphBuilder.AllocParameters<FMobileDecalPassParameters>();
-		PassParameters->View = View.ViewUniformBuffer;
-		PassParameters->RenderTargets = BasePassRenderTargets;
-		PassParameters->MobileSceneTextures = MobileSceneTextures;
+		const bool bDeferredShading = IsMobileDeferredShadingEnabled(View.GetShaderPlatform());
 
-		GraphBuilder.AddPass(
-			RDG_EVENT_NAME("MobileDeferredDecalsPass"), PassParameters, ERDGPassFlags::Raster | ERDGPassFlags::SkipRenderPass,
-			[&View, SortedDecalCount, SortedDecals](FRHICommandListImmediate& RHICmdList)
-		{
-			const bool bDeferredShading = IsMobileDeferredShadingEnabled(View.GetShaderPlatform());
+		FGraphicsPipelineStateInitializer GraphicsPSOInit;
+		RHICmdList.ApplyCachedRenderTargets(GraphicsPSOInit);
 
-			FGraphicsPipelineStateInitializer GraphicsPSOInit;
-			RHICmdList.ApplyCachedRenderTargets(GraphicsPSOInit);
+		RHICmdList.SetViewport(View.ViewRect.Min.X, View.ViewRect.Min.Y, 0, View.ViewRect.Max.X, View.ViewRect.Max.Y, 1);
+		RHICmdList.SetStreamSource(0, GetUnitCubeVertexBuffer(), 0);
 
-			RHICmdList.SetViewport(View.ViewRect.Min.X, View.ViewRect.Min.Y, 0, View.ViewRect.Max.X, View.ViewRect.Max.Y, 1);
-			RHICmdList.SetStreamSource(0, GetUnitCubeVertexBuffer(), 0);
+		for (int32 DecalIndex = 0; DecalIndex < SortedDecalCount; DecalIndex++)
+		{
+			const FTransientDecalRenderData& DecalData = SortedDecals[DecalIndex];
+			const FDeferredDecalProxy& DecalProxy = *DecalData.DecalProxy;
+			const FMatrix ComponentToWorldMatrix = DecalProxy.ComponentTrans.ToMatrixWithScale();
+			const FMatrix FrustumComponentToClip = FDecalRendering::ComputeComponentToClipMatrix(View, ComponentToWorldMatrix);
 
-			for (int32 DecalIndex = 0; DecalIndex < SortedDecalCount; DecalIndex++)
+			const float ConservativeRadius = DecalData.ConservativeRadius;
+			const bool bInsideDecal = ((FVector)View.ViewMatrices.GetViewOrigin() - ComponentToWorldMatrix.GetOrigin()).SizeSquared() < FMath::Square(ConservativeRadius * 1.05f + View.NearClippingDistance * 2.0f);
+			bool bReverseHanded = false;
 			{
-				const FTransientDecalRenderData& DecalData = (*SortedDecals)[DecalIndex];
-				const FDeferredDecalProxy& DecalProxy = *DecalData.DecalProxy;
-				const FMatrix ComponentToWorldMatrix = DecalProxy.ComponentTrans.ToMatrixWithScale();
-				const FMatrix FrustumComponentToClip = FDecalRendering::ComputeComponentToClipMatrix(View, ComponentToWorldMatrix);
-
-				const float ConservativeRadius = DecalData.ConservativeRadius;
-				const bool bInsideDecal = ((FVector)View.ViewMatrices.GetViewOrigin() - ComponentToWorldMatrix.GetOrigin()).SizeSquared() < FMath::Square(ConservativeRadius * 1.05f + View.NearClippingDistance * 2.0f);
-				bool bReverseHanded = false;
-				{
-					// Account for the reversal of handedness caused by negative scale on the decal
-					const auto& Scale3d = DecalProxy.ComponentTrans.GetScale3D();
-					bReverseHanded = Scale3d[0] * Scale3d[1] * Scale3d[2] < 0.f;
-				}
-				EDecalRasterizerState DecalRasterizerState = FDecalRenderingCommon::ComputeDecalRasterizerState(bInsideDecal, bReverseHanded, View.bReverseCulling);
-				GraphicsPSOInit.RasterizerState = GetDecalRasterizerState(DecalRasterizerState);
+				// Account for the reversal of handedness caused by negative scale on the decal
+				const auto& Scale3d = DecalProxy.ComponentTrans.GetScale3D();
+				bReverseHanded = Scale3d[0] * Scale3d[1] * Scale3d[2] < 0.f;
+			}
+			EDecalRasterizerState DecalRasterizerState = FDecalRenderingCommon::ComputeDecalRasterizerState(bInsideDecal, bReverseHanded, View.bReverseCulling);
+			GraphicsPSOInit.RasterizerState = GetDecalRasterizerState(DecalRasterizerState);
 
-				if (bInsideDecal)
-				{
-					GraphicsPSOInit.DepthStencilState = TStaticDepthStencilState<
-						false, CF_Always,
-						true, CF_Equal, SO_Keep, SO_Keep, SO_Keep,
-						false, CF_Always, SO_Keep, SO_Keep, SO_Keep,
-						GET_STENCIL_BIT_MASK(RECEIVE_DECAL, 1), 0x00>::GetRHI();
-				}
-				else
-				{
-					GraphicsPSOInit.DepthStencilState = TStaticDepthStencilState<
-						false, CF_DepthNearOrEqual,
-						true, CF_Equal, SO_Keep, SO_Keep, SO_Keep,
-						false, CF_Always, SO_Keep, SO_Keep, SO_Keep,
-						GET_STENCIL_BIT_MASK(RECEIVE_DECAL, 1), 0x00>::GetRHI();
-				}
+			if (bInsideDecal)
+			{
+				GraphicsPSOInit.DepthStencilState = TStaticDepthStencilState<
+					false, CF_Always,
+					true, CF_Equal, SO_Keep, SO_Keep, SO_Keep,
+					false, CF_Always, SO_Keep, SO_Keep, SO_Keep,
+					GET_STENCIL_BIT_MASK(RECEIVE_DECAL, 1), 0x00>::GetRHI();
+			}
+			else
+			{
+				GraphicsPSOInit.DepthStencilState = TStaticDepthStencilState<
+					false, CF_DepthNearOrEqual,
+					true, CF_Equal, SO_Keep, SO_Keep, SO_Keep,
+					false, CF_Always, SO_Keep, SO_Keep, SO_Keep,
+					GET_STENCIL_BIT_MASK(RECEIVE_DECAL, 1), 0x00>::GetRHI();
+			}
 
-				if (bDeferredShading)
-				{
-					GraphicsPSOInit.BlendState = MobileDeferred_GetDecalBlendState(DecalData.FinalDecalBlendMode, DecalData.bHasNormal);
-				}
-				else
-				{
-					GraphicsPSOInit.BlendState = MobileForward_GetDecalBlendState(DecalData.FinalDecalBlendMode);
-				}
+			if (bDeferredShading)
+			{
+				GraphicsPSOInit.BlendState = MobileDeferred_GetDecalBlendState(DecalData.FinalDecalBlendMode, DecalData.bHasNormal);
+			}
+			else
+			{
+				GraphicsPSOInit.BlendState = MobileForward_GetDecalBlendState(DecalData.FinalDecalBlendMode);
+			}
 
-				// Set shader params
-				FDecalRendering::SetShader(RHICmdList, GraphicsPSOInit, View, DecalData, DRS_Mobile, FrustumComponentToClip);
+			// Set shader params
+			FDecalRendering::SetShader(RHICmdList, GraphicsPSOInit, View, DecalData, DRS_Mobile, FrustumComponentToClip);
 
-				RHICmdList.DrawIndexedPrimitive(GetUnitCubeIndexBuffer(), 0, 0, 8, 0, UE_ARRAY_COUNT(GCubeIndices) / 3, 1);
-			}
-		});
+			RHICmdList.DrawIndexedPrimitive(GetUnitCubeIndexBuffer(), 0, 0, 8, 0, UE_ARRAY_COUNT(GCubeIndices) / 3, 1);
+		}
 	}
 }
diff --git a/Engine/Source/Runtime/Renderer/Private/MobileDeferredShadingPass.cpp b/Engine/Source/Runtime/Renderer/Private/MobileDeferredShadingPass.cpp
--- a/Engine/Source/Runtime/Renderer/Private/MobileDeferredShadingPass.cpp
+++ b/Engine/Source/Runtime/Renderer/Private/MobileDeferredShadingPass.cpp
@@ -6,6 +6,7 @@
 #include "PostProcess/PostProcessing.h"
 #include "PostProcess/SceneFilterRendering.h"
 #include "PipelineStateCache.h"
+#include "PlanarReflectionRendering.h"
 
 int32 GMobileUseClusteredDeferredShading = 0;
 static FAutoConsoleVariableRef CVarMobileUseClusteredDeferredShading(
@@ -45,6 +46,7 @@ class FMobileDirectLightFunctionPS : public FMaterialShader
 		SHADER_PARAMETER_STRUCT_REF(FMobileDirectionalLightShaderParameters, MobileDirectionalLight)
 		SHADER_PARAMETER_STRUCT_REF(FReflectionUniformParameters, ReflectionsParameters)
 		SHADER_PARAMETER_STRUCT_REF(FReflectionCaptureShaderData, ReflectionCaptureData)
+		SHADER_PARAMETER_STRUCT_REF(FPlanarReflectionUniformParameters, PlanarReflection) // Single global planar reflection.
 		SHADER_PARAMETER(FMatrix, WorldToLight)
 		SHADER_PARAMETER(FVector4, LightFunctionParameters)
 		SHADER_PARAMETER(FVector, LightFunctionParameters2)
@@ -106,7 +108,6 @@ class FMobileDirectLightFunctionPS : public FMaterialShader
 	{
 		FMaterialShader* MaterialShader = Shader.GetShader();
 		FRHIPixelShader* ShaderRHI = Shader.GetPixelShader();
-		MaterialShader->SetViewParameters(RHICmdList, ShaderRHI, View, View.ViewUniformBuffer);
 		MaterialShader->SetParameters(RHICmdList, ShaderRHI, Proxy, Material, View);
 		SetShaderParameters(RHICmdList, Shader, ShaderRHI, Parameters);
 	}
@@ -201,7 +202,7 @@ static void GetLightMaterial(const FCachedLightMaterial& DefaultLightMaterial, c
 	OutShader = MaterialShaderMap->GetShader<ShaderType>(PermutationId);
 }
 
-static void RenderDirectLight(FRHICommandListImmediate& RHICmdList, const FScene& Scene, const FViewInfo& View, const FCachedLightMaterial& DefaultLightMaterial)
+static void RenderDirectLight(FRHICommandListImmediate& RHICmdList, const FScene& Scene, const FViewInfo& View, const FCachedLightMaterial& DefaultLightMaterial, const FTextureRHIRef MobilePixelProjectedReflection)
 {
 	FLightSceneInfo* DirectionalLight = nullptr;
 	for (int32 ChannelIdx = 0; ChannelIdx < UE_ARRAY_COUNT(Scene.MobileDirectionalLights) && !DirectionalLight; ChannelIdx++)
@@ -254,6 +255,16 @@ static void RenderDirectLight(FRHICommandListImmediate& RHICmdList, const FScene
 	PassParameters.ReflectionsParameters = CreateUniformBufferImmediate(ReflectionUniformParameters, UniformBuffer_SingleDraw);
 	PassParameters.LightFunctionParameters = FVector4(1.0f, 1.0f, 0.0f, 0.0f);
 
+	const FPlanarReflectionSceneProxy* ReflectionSceneProxy = Scene.GetForwardPassGlobalPlanarReflection();
+	FPlanarReflectionUniformParameters PlanarReflectionUniformParameters;
+	SetupPlanarReflectionUniformParameters(View, ReflectionSceneProxy, PlanarReflectionUniformParameters);
+
+	if (MobilePixelProjectedReflection.IsValid())
+	{
+		PlanarReflectionUniformParameters.PlanarReflectionTexture = MobilePixelProjectedReflection;
+	}
+	PassParameters.PlanarReflection = TUniformBufferRef<FPlanarReflectionUniformParameters>::CreateUniformBufferImmediate(PlanarReflectionUniformParameters, UniformBuffer_SingleDraw);
+
 	const bool bUseMovableLight = DirectionalLight && !DirectionalLight->Proxy->HasStaticShadowing();
 	PassParameters.LightFunctionParameters2 = FVector(DirectionalLight->Proxy->GetLightFunctionFadeDistance(), DirectionalLight->Proxy->GetLightFunctionDisabledBrightness(), bUseMovableLight ? 1.0f : 0.0f);
 	const FVector Scale = DirectionalLight->Proxy->GetLightFunctionScale();
@@ -470,15 +481,15 @@ static void SetupSimpleLightPSO(
 static void RenderSimpleLights(
 	FRHICommandListImmediate& RHICmdList, 
 	const FScene& Scene, 
-	const TArrayView<const FViewInfo> PassViews, 
+	int32 ViewIndex,
+	int32 NumViews,
+	const FViewInfo& View,
 	const FSortedLightSetSceneInfo &SortedLightSet, 
 	const FCachedLightMaterial& DefaultMaterial)
 {
 	const FSimpleLightArray& SimpleLights = SortedLightSet.SimpleLights;
-	const int32 NumViews = PassViews.Num();
-	const FViewInfo& View0 = PassViews[0];
 
-	TShaderMapRef<TDeferredLightVS<true>> VertexShader(View0.ShaderMap);
+	TShaderMapRef<TDeferredLightVS<true>> VertexShader(View.ShaderMap);
 	TShaderRef<FMobileRadialLightFunctionPS> PixelShaders[2];
 	{
 		const FMaterialShaderMap* MaterialShaderMap = DefaultMaterial.Material->GetRenderingThreadShaderMap();
@@ -494,8 +505,8 @@ static void RenderSimpleLights(
 	// Setup PSOs we going to use for light rendering 
 	FGraphicsPipelineStateInitializer GraphicsPSOLight[2];
 	{
-		SetupSimpleLightPSO(RHICmdList, View0, VertexShader, PixelShaders[0], GraphicsPSOLight[0]);
-		SetupSimpleLightPSO(RHICmdList, View0, VertexShader, PixelShaders[1], GraphicsPSOLight[1]);
+		SetupSimpleLightPSO(RHICmdList, View, VertexShader, PixelShaders[0], GraphicsPSOLight[0]);
+		SetupSimpleLightPSO(RHICmdList, View, VertexShader, PixelShaders[1], GraphicsPSOLight[1]);
 	}
 	// Setup stencil mask PSO
 	FGraphicsPipelineStateInitializer GraphicsPSOLightMask;
@@ -503,7 +514,7 @@ static void RenderSimpleLights(
 		RHICmdList.ApplyCachedRenderTargets(GraphicsPSOLightMask);
 		GraphicsPSOLightMask.PrimitiveType = PT_TriangleList;
 		GraphicsPSOLightMask.BlendState = TStaticBlendStateWriteMask<CW_NONE, CW_NONE, CW_NONE, CW_NONE, CW_NONE, CW_NONE, CW_NONE, CW_NONE>::GetRHI();
-		GraphicsPSOLightMask.RasterizerState = View0.bReverseCulling ? TStaticRasterizerState<FM_Solid, CM_CCW>::GetRHI() : TStaticRasterizerState<FM_Solid, CM_CW>::GetRHI();
+		GraphicsPSOLightMask.RasterizerState = View.bReverseCulling ? TStaticRasterizerState<FM_Solid, CM_CCW>::GetRHI() : TStaticRasterizerState<FM_Solid, CM_CW>::GetRHI();
 		// set stencil to 1 where depth test fails
 		GraphicsPSOLightMask.DepthStencilState = TStaticDepthStencilState<
 			false, CF_DepthNearOrEqual,
@@ -518,101 +529,88 @@ static void RenderSimpleLights(
 	for (int32 LightIndex = 0; LightIndex < SimpleLights.InstanceData.Num(); LightIndex++)
 	{
 		const FSimpleLightEntry& SimpleLight = SimpleLights.InstanceData[LightIndex];
-		for (int32 ViewIndex = 0; ViewIndex < NumViews; ViewIndex++)
+		const FSimpleLightPerViewEntry& SimpleLightPerViewData = SimpleLights.GetViewDependentData(LightIndex, ViewIndex, NumViews);
+		const FSphere LightBounds(SimpleLightPerViewData.Position, SimpleLight.Radius);
+
+		if (NumViews > 1)
 		{
-			const FViewInfo& View = PassViews[ViewIndex];
-			const FSimpleLightPerViewEntry& SimpleLightPerViewData = SimpleLights.GetViewDependentData(LightIndex, ViewIndex, NumViews);
-			const FSphere LightBounds(SimpleLightPerViewData.Position, SimpleLight.Radius);
-			
-			if (NumViews > 1)
-			{
-				// set viewports only we we have more than one 
-				// otherwise it is set at the start of the pass
-				RHICmdList.SetViewport(View.ViewRect.Min.X, View.ViewRect.Min.Y, 0.0f, View.ViewRect.Max.X, View.ViewRect.Max.Y, 1.0f);
-			}
+			// set viewports only we we have more than one 
+			// otherwise it is set at the start of the pass
+			RHICmdList.SetViewport(View.ViewRect.Min.X, View.ViewRect.Min.Y, 0.0f, View.ViewRect.Max.X, View.ViewRect.Max.Y, 1.0f);
+		}
 
-			// Render light mask
-			SetGraphicsPipelineState(RHICmdList, GraphicsPSOLightMask);
-			VertexShader->SetSimpleLightParameters(RHICmdList, View, LightBounds);
-			RHICmdList.SetStencilRef(1);
-			StencilingGeometry::DrawSphere(RHICmdList);
-						
-			// Render light
-			FMobileRadialLightFunctionPS::FParameters PassParameters;
-			FDeferredLightUniformStruct DeferredLightUniformsValue;
-			SetupSimpleDeferredLightParameters(SimpleLight, SimpleLightPerViewData, DeferredLightUniformsValue);
-			PassParameters.DeferredLightUniforms = TUniformBufferRef<FDeferredLightUniformStruct>::CreateUniformBufferImmediate(DeferredLightUniformsValue, EUniformBufferUsage::UniformBuffer_SingleFrame);
-			PassParameters.IESTexture = GWhiteTexture->TextureRHI;
-			PassParameters.IESTextureSampler = GWhiteTexture->SamplerStateRHI;
-			if (SimpleLight.Exponent == 0)
-			{
-				SetGraphicsPipelineState(RHICmdList, GraphicsPSOLight[1]);
-				FMobileRadialLightFunctionPS::SetParameters(RHICmdList, PixelShaders[1], View, DefaultMaterial.MaterialProxy, *DefaultMaterial.Material, PassParameters);
-			}
-			else
-			{
-				SetGraphicsPipelineState(RHICmdList, GraphicsPSOLight[0]);
-				FMobileRadialLightFunctionPS::SetParameters(RHICmdList, PixelShaders[0], View, DefaultMaterial.MaterialProxy, *DefaultMaterial.Material, PassParameters);
-			}
-			VertexShader->SetSimpleLightParameters(RHICmdList, View, LightBounds);
-			
-			// Shade only MSM_DefaultLit pixels
-			uint8 StencilRef = GET_STENCIL_MOBILE_SM_MASK(MSM_DefaultLit);
-			RHICmdList.SetStencilRef(StencilRef);
+		// Render light mask
+		SetGraphicsPipelineState(RHICmdList, GraphicsPSOLightMask);
+		VertexShader->SetSimpleLightParameters(RHICmdList, View, LightBounds);
+		RHICmdList.SetStencilRef(1);
+		StencilingGeometry::DrawSphere(RHICmdList);
 
-			// Apply the point or spot light with some approximately bounding geometry,
-			// So we can get speedups from depth testing and not processing pixels outside of the light's influence.
-			StencilingGeometry::DrawSphere(RHICmdList);
+		// Render light
+		FMobileRadialLightFunctionPS::FParameters PassParameters;
+		FDeferredLightUniformStruct DeferredLightUniformsValue;
+		SetupSimpleDeferredLightParameters(SimpleLight, SimpleLightPerViewData, DeferredLightUniformsValue);
+		PassParameters.DeferredLightUniforms = TUniformBufferRef<FDeferredLightUniformStruct>::CreateUniformBufferImmediate(DeferredLightUniformsValue, EUniformBufferUsage::UniformBuffer_SingleFrame);
+		PassParameters.IESTexture = GWhiteTexture->TextureRHI;
+		PassParameters.IESTextureSampler = GWhiteTexture->SamplerStateRHI;
+		if (SimpleLight.Exponent == 0)
+		{
+			SetGraphicsPipelineState(RHICmdList, GraphicsPSOLight[1]);
+			FMobileRadialLightFunctionPS::SetParameters(RHICmdList, PixelShaders[1], View, DefaultMaterial.MaterialProxy, *DefaultMaterial.Material, PassParameters);
 		}
+		else
+		{
+			SetGraphicsPipelineState(RHICmdList, GraphicsPSOLight[0]);
+			FMobileRadialLightFunctionPS::SetParameters(RHICmdList, PixelShaders[0], View, DefaultMaterial.MaterialProxy, *DefaultMaterial.Material, PassParameters);
+		}
+		VertexShader->SetSimpleLightParameters(RHICmdList, View, LightBounds);
+
+		// Shade only MSM_DefaultLit pixels
+		uint8 StencilRef = GET_STENCIL_MOBILE_SM_MASK(MSM_DefaultLit);
+		RHICmdList.SetStencilRef(StencilRef);
+
+		// Apply the point or spot light with some approximately bounding geometry,
+		// So we can get speedups from depth testing and not processing pixels outside of the light's influence.
+		StencilingGeometry::DrawSphere(RHICmdList);
 	}
 }
 
 void MobileDeferredShadingPass(
-	FRDGBuilder& GraphBuilder,
-	FRenderTargetBindingSlots& BasePassRenderTargets,
-	TRDGUniformBufferRef<FMobileSceneTextureUniformParameters> MobileSceneTextures,
+	FRHICommandListImmediate& RHICmdList,
+	int32 ViewIndex,
+	int32 NumViews,
+	const FViewInfo& View,
 	const FScene& Scene, 
-	const TArrayView<const FViewInfo> PassViews, 
-	const FSortedLightSetSceneInfo &SortedLightSet)
+	const FSortedLightSetSceneInfo &SortedLightSet,
+	const FTextureRHIRef MobilePixelProjectedReflection)
 {
-	auto* PassParameters = GraphBuilder.AllocParameters<FMobileDeferredPassParameters>();
-	PassParameters->RenderTargets = BasePassRenderTargets;
-	PassParameters->MobileSceneTextures = MobileSceneTextures;
+	RHICmdList.SetViewport(View.ViewRect.Min.X, View.ViewRect.Min.Y, 0.0f, View.ViewRect.Max.X, View.ViewRect.Max.Y, 1.0f);
 
-	GraphBuilder.AddPass(
-		RDG_EVENT_NAME("MobileDeferredShadingPass"), PassParameters, ERDGPassFlags::Raster | ERDGPassFlags::SkipRenderPass,
-		[&Scene, PassViews, &SortedLightSet](FRHICommandListImmediate& RHICmdList)
-	{
-		const FViewInfo& View0 = PassViews[0];
-		RHICmdList.SetViewport(View0.ViewRect.Min.X, View0.ViewRect.Min.Y, 0.0f, View0.ViewRect.Max.X, View0.ViewRect.Max.Y, 1.0f);
-
-		// Default material for light rendering
-		FCachedLightMaterial DefaultMaterial;
-		DefaultMaterial.MaterialProxy = UMaterial::GetDefaultMaterial(MD_LightFunction)->GetRenderProxy();
-		DefaultMaterial.Material = DefaultMaterial.MaterialProxy->GetMaterialNoFallback(ERHIFeatureLevel::ES3_1);
-		check(DefaultMaterial.Material);
+	// Default material for light rendering
+	FCachedLightMaterial DefaultMaterial;
+	DefaultMaterial.MaterialProxy = UMaterial::GetDefaultMaterial(MD_LightFunction)->GetRenderProxy();
+	DefaultMaterial.Material = DefaultMaterial.MaterialProxy->GetMaterialNoFallback(ERHIFeatureLevel::ES3_1);
+	check(DefaultMaterial.Material);
 
-		RenderDirectLight(RHICmdList, Scene, View0, DefaultMaterial);
+	RenderDirectLight(RHICmdList, Scene, View, DefaultMaterial, MobilePixelProjectedReflection);
 
-		if (GMobileUseClusteredDeferredShading == 0)
-		{
-			// Render non-clustered simple lights
-			RenderSimpleLights(RHICmdList, Scene, PassViews, SortedLightSet, DefaultMaterial);
-		}
+	if (GMobileUseClusteredDeferredShading == 0)
+	{
+		// Render non-clustered simple lights
+		RenderSimpleLights(RHICmdList, Scene, ViewIndex, NumViews, View, SortedLightSet, DefaultMaterial);
+	}
 
-		// Render non-clustered local lights
-		int32 NumLights = SortedLightSet.SortedLights.Num();
-		int32 StandardDeferredStart = SortedLightSet.SimpleLightsEnd;
-		if (GMobileUseClusteredDeferredShading != 0)
-		{
-			StandardDeferredStart = SortedLightSet.ClusteredSupportedEnd;
-		}
+	// Render non-clustered local lights
+	int32 NumLights = SortedLightSet.SortedLights.Num();
+	int32 StandardDeferredStart = SortedLightSet.SimpleLightsEnd;
+	if (GMobileUseClusteredDeferredShading != 0)
+	{
+		StandardDeferredStart = SortedLightSet.ClusteredSupportedEnd;
+	}
 
-		for (int32 LightIdx = StandardDeferredStart; LightIdx < NumLights; ++LightIdx)
-		{
-			const FSortedLightSceneInfo& SortedLight = SortedLightSet.SortedLights[LightIdx];
-			const FLightSceneInfo& LightSceneInfo = *SortedLight.LightSceneInfo;
-			RenderLocalLight(RHICmdList, Scene, View0, LightSceneInfo, DefaultMaterial);
-		}
-	});
+	for (int32 LightIdx = StandardDeferredStart; LightIdx < NumLights; ++LightIdx)
+	{
+		const FSortedLightSceneInfo& SortedLight = SortedLightSet.SortedLights[LightIdx];
+		const FLightSceneInfo& LightSceneInfo = *SortedLight.LightSceneInfo;
+		RenderLocalLight(RHICmdList, Scene, View, LightSceneInfo, DefaultMaterial);
+	}
 }
\ No newline at end of file
diff --git a/Engine/Source/Runtime/Renderer/Private/MobileDeferredShadingPass.h b/Engine/Source/Runtime/Renderer/Private/MobileDeferredShadingPass.h
--- a/Engine/Source/Runtime/Renderer/Private/MobileDeferredShadingPass.h
+++ b/Engine/Source/Runtime/Renderer/Private/MobileDeferredShadingPass.h
@@ -16,9 +16,10 @@ struct FSortedLightSetSceneInfo;
 extern int32 GMobileUseClusteredDeferredShading;
 
 void MobileDeferredShadingPass(
-	FRDGBuilder& GraphBuilder,
-	FRenderTargetBindingSlots& BasePassRenderTargets,
-	TRDGUniformBufferRef<FMobileSceneTextureUniformParameters> MobileSceneTextures,
+	FRHICommandListImmediate& RHICmdList,
+	int32 ViewIndex,
+	int32 NumViews,
+	const FViewInfo& View,
 	const FScene& Scene, 
-	const TArrayView<const FViewInfo> PassViews, 
-	const FSortedLightSetSceneInfo &SortedLightSet);
+	const FSortedLightSetSceneInfo &SortedLightSet,
+	const FTextureRHIRef MobilePixelProjectedReflection);
diff --git a/Engine/Source/Runtime/Renderer/Private/MobileShadingRenderer.cpp b/Engine/Source/Runtime/Renderer/Private/MobileShadingRenderer.cpp
--- a/Engine/Source/Runtime/Renderer/Private/MobileShadingRenderer.cpp
+++ b/Engine/Source/Runtime/Renderer/Private/MobileShadingRenderer.cpp
@@ -167,6 +167,15 @@ BEGIN_SHADER_PARAMETER_STRUCT(FMobileRenderOpaqueFXPassParameters, )
 	SHADER_PARAMETER_RDG_UNIFORM_BUFFER(FMobileSceneTextureUniformParameters, SceneTextures)
 END_SHADER_PARAMETER_STRUCT()
 
+BEGIN_SHADER_PARAMETER_STRUCT(FMobileRenderPassParameters, )
+	SHADER_PARAMETER_STRUCT_INCLUDE(FViewShaderParameters, View)
+	SHADER_PARAMETER_STRUCT_ARRAY(FInstanceCullingDrawParams, InstanceCullingDrawParams, [EMeshPass::Num])
+	SHADER_PARAMETER_RDG_UNIFORM_BUFFER(FMobileSceneTextureUniformParameters, SceneTextures)
+	SHADER_PARAMETER_RDG_UNIFORM_BUFFER(FMobileBasePassUniformParameters, MobileBasePass)
+	SHADER_PARAMETER_RDG_UNIFORM_BUFFER(FDebugViewModePassUniformParameters, DebugViewMode)
+	RENDER_TARGET_BINDING_SLOTS()
+END_SHADER_PARAMETER_STRUCT()
+
 static bool PostProcessUsesSceneDepth(const FViewInfo& View)
 {
 	// Find out whether post-process materials use CustomDepth/Stencil lookups
@@ -365,10 +374,7 @@ void FMobileSceneRenderer::InitViews(FRDGBuilder& GraphBuilder, FSceneTexturesCo
 		&& !ViewFamily.EngineShowFlags.HitProxies
 		&& ViewFamily.EngineShowFlags.Lighting
 		&& !ViewFamily.EngineShowFlags.VisualizeLightCulling
-		&& !ViewFamily.UseDebugViewPS()
-		// Only support forward shading, we don't want to break tiled deferred shading.
-		&& !bDeferredShading;
-
+		&& !ViewFamily.UseDebugViewPS();
 
 	bRequiresAmbientOcclusionPass = IsUsingMobileAmbientOcclusion(ShaderPlatform)
 		&& Views[0].FinalPostProcessSettings.AmbientOcclusionIntensity > 0
@@ -378,11 +384,10 @@ void FMobileSceneRenderer::InitViews(FRDGBuilder& GraphBuilder, FSceneTexturesCo
 		&& !Views[0].bIsPlanarReflection
 		&& !ViewFamily.EngineShowFlags.HitProxies
 		&& !ViewFamily.EngineShowFlags.VisualizeLightCulling
-		&& !ViewFamily.UseDebugViewPS()
-		// Only support forward shading, we don't want to break tiled deferred shading.
-		&& !bDeferredShading;
+		&& !ViewFamily.UseDebugViewPS();
 
 	bShouldRenderVelocities = ShouldRenderVelocities();
+
 	bRequiresDistanceField = IsMobileDistanceFieldEnabled(ShaderPlatform)
 		&& ViewFamily.EngineShowFlags.Lighting
 		&& !Views[0].bIsReflectionCapture
@@ -417,9 +422,20 @@ void FMobileSceneRenderer::InitViews(FRDGBuilder& GraphBuilder, FSceneTexturesCo
 	// never keep MSAA depth
 	bKeepDepthContent = (NumMSAASamples > 1 ? false : bKeepDepthContent);
 
-	// Finalize and set the scene textures config.
+	// Update the bKeepDepthContent based on the mobile renderer status.
 	SceneTexturesConfig.bKeepDepthContent = bKeepDepthContent;
+
+	// Update the pixel projected reflection extent according to the settings in the PlanarReflectionComponent.
+	if (bRequiresPixelProjectedPlanarRelfectionPass)
+	{
+		SceneTexturesConfig.MobilePixelProjectedReflectionExtent = PlanarReflectionSceneProxy->RenderTarget->GetSizeXY();
+	}
+	else
+	{
+		SceneTexturesConfig.MobilePixelProjectedReflectionExtent = FIntPoint::ZeroValue;
+	}
 	
+	// Finalize and set the scene textures config.
 	FSceneTexturesConfig::Set(SceneTexturesConfig);
 
 	// Initialise Sky/View resources before the view global uniform buffer is built.
@@ -427,14 +443,6 @@ void FMobileSceneRenderer::InitViews(FRDGBuilder& GraphBuilder, FSceneTexturesCo
 	{
 		InitSkyAtmosphereForViews(RHICmdList);
 	}
-	if (bRequiresPixelProjectedPlanarRelfectionPass)
-	{
-		InitPixelProjectedReflectionOutputs(RHICmdList, PlanarReflectionSceneProxy->RenderTarget->GetSizeXY());
-	}
-	else
-	{
-		ReleasePixelProjectedReflectionOutputs();
-	}
 
 	if (bRequiresDistanceFieldShadowingPass)
 	{
@@ -691,11 +699,11 @@ void FMobileSceneRenderer::Render(FRDGBuilder& GraphBuilder)
 
 	if (bDeferredShading)
 	{
-		RenderDeferred(GraphBuilder, Views, SortedLightSet, ViewFamilyTexture, SceneTextures);
+		RenderDeferred(GraphBuilder, SortedLightSet, ViewFamilyTexture, SceneTextures);
 	}
 	else
 	{
-		RenderForward(GraphBuilder, Views, ViewFamilyTexture, SceneTextures);
+		RenderForward(GraphBuilder, ViewFamilyTexture, SceneTextures);
 	}
 
 	SceneTextures.MobileSetupMode = EMobileSceneTextureSetupMode::All;
@@ -771,6 +779,13 @@ void FMobileSceneRenderer::Render(FRDGBuilder& GraphBuilder)
 		RenderAmbientOcclusion(GraphBuilder, SceneTextures.Depth.Resolve, SceneTextures.ScreenSpaceAO);
 	}
 
+	if (bRequiresPixelProjectedPlanarRelfectionPass)
+	{
+		const FPlanarReflectionSceneProxy* PlanarReflectionSceneProxy = Scene ? Scene->GetForwardPassGlobalPlanarReflection() : nullptr;
+
+		RenderPixelProjectedReflection(GraphBuilder, SceneTextures.Color.Resolve, SceneTextures.Depth.Resolve, SceneTextures.PixelProjectedReflection, PlanarReflectionSceneProxy);
+	}
+
 	GraphBuilder.SetCommandListStat(GET_STATID(STAT_CLMM_Post));
 
 	if (bUseVirtualTexturing)
@@ -810,14 +825,9 @@ void FMobileSceneRenderer::Render(FRDGBuilder& GraphBuilder)
 	AddPass(GraphBuilder, PollOcclusionQueriesAndDispatchToRHIThreadPass);
 }
 
-BEGIN_SHADER_PARAMETER_STRUCT(FMobilePostBasePassViewExtensionParameters, )
-	SHADER_PARAMETER_RDG_UNIFORM_BUFFER(FMobileSceneTextureUniformParameters, SceneTextures)
-	RENDER_TARGET_BINDING_SLOTS()
-END_SHADER_PARAMETER_STRUCT()
-
-void FMobileSceneRenderer::RenderForward(FRDGBuilder& GraphBuilder, const TArrayView<const FViewInfo> ViewList, FRDGTextureRef ViewFamilyTexture, FSceneTextures& SceneTextures)
+void FMobileSceneRenderer::RenderForward(FRDGBuilder& GraphBuilder, FRDGTextureRef ViewFamilyTexture, FSceneTextures& SceneTextures)
 {
-	const FViewInfo& View = ViewList[0];
+	const FViewInfo& MainView = Views[0];
 
 	FRDGTextureRef SceneColor = nullptr;
 	FRDGTextureRef SceneColorResolve = nullptr;
@@ -854,216 +864,328 @@ void FMobileSceneRenderer::RenderForward(FRDGBuilder& GraphBuilder, const TArray
 	FRenderTargetBindingSlots BasePassRenderTargets;
 	BasePassRenderTargets[0] = FRenderTargetBinding(SceneColor, SceneColorResolve, ERenderTargetLoadAction::EClear);
 	BasePassRenderTargets.DepthStencil = FDepthStencilBinding(SceneDepth, ERenderTargetLoadAction::EClear, FExclusiveDepthStencil::DepthWrite_StencilWrite);
-	BasePassRenderTargets.ShadingRateTexture = (!View.bIsSceneCapture && !View.bIsReflectionCapture) ? SceneTextures.ShadingRate : nullptr;
+	BasePassRenderTargets.ShadingRateTexture = (!MainView.bIsSceneCapture && !MainView.bIsReflectionCapture) ? SceneTextures.ShadingRate : nullptr;
 	BasePassRenderTargets.SubpassHint = ESubpassHint::DepthReadSubpass;
 	BasePassRenderTargets.NumOcclusionQueries = ComputeNumOcclusionQueriesToBatch();
 
 	//if the scenecolor isn't multiview but the app is, need to render as a single-view multiview due to shaders
-	BasePassRenderTargets.MultiViewCount = View.bIsMobileMultiViewEnabled ? 2 : (bIsMultiViewApplication ? 1 : 0);
+	BasePassRenderTargets.MultiViewCount = MainView.bIsMobileMultiViewEnabled ? 2 : (bIsMultiViewApplication ? 1 : 0);
 
-	auto* OpaqueBasePassParameters = GraphBuilder.AllocParameters<FRenderTargetParameters>();
-	OpaqueBasePassParameters->RenderTargets = BasePassRenderTargets;
+	const bool bShouldRenderPrePass = ShouldRenderPrePass();
 
-	GraphBuilder.AddPass(
-		{},
-		OpaqueBasePassParameters,
-		ERDGPassFlags::Raster | ERDGPassFlags::SkipRenderPass,
-		[OpaqueBasePassParameters](FRHICommandListImmediate& RHICmdList)
-	{
-		RHICmdList.BeginRenderPass(GetRenderPassInfo(OpaqueBasePassParameters), TEXT("SceneColorRendering"));
-	});
+	const bool bShouldRenderMobileDebugView = ViewFamily.UseDebugViewPS();
+
+	const bool bAdrenoOcclusionMode = CVarMobileAdrenoOcclusionMode.GetValueOnRenderThread() != 0;
+
+	const bool bShouldRenderOcclusionPass = DoOcclusionQueries(FeatureLevel);
+
+	ETranslucencyPass::Type TranslucencyPass = ViewFamily.AllowTranslucencyAfterDOF() ? ETranslucencyPass::TPT_StandardTranslucency : ETranslucencyPass::TPT_AllTranslucency;
+
+	const EMeshPass::Type TranslucencyMeshPass = TranslucencyPassToMeshPass(TranslucencyPass);
 
-	if (GIsEditor && !View.bIsSceneCapture)
+	const bool bShouldRenderTranslucency = ShouldRenderTranslucency(TranslucencyPass) && ViewFamily.EngineShowFlags.Translucency;
+
+	const FRDGSystemTextures& SystemTextures = FRDGSystemTextures::Get(GraphBuilder);
+
+	for (int32 ViewIndex = 0; ViewIndex < Views.Num(); ViewIndex++)
 	{
-		AddPass(GraphBuilder, [this](FRHICommandListImmediate& RHICmdList)
+		FViewInfo& View = Views[ViewIndex];
+
+		SCOPED_GPU_MASK(GraphBuilder.RHICmdList, !View.IsInstancedStereoPass() ? View.GPUMask : (Views[0].GPUMask | Views[1].GPUMask));
+		SCOPED_CONDITIONAL_DRAW_EVENTF(GraphBuilder.RHICmdList, EventView, Views.Num() > 1, TEXT("View%d"), ViewIndex);
+
+		if (!View.ShouldRenderView())
 		{
-			DrawClearQuad(RHICmdList, Views[0].BackgroundColor);
-		});
+			continue;
+		}
+
+		View.BeginRenderView();
+
+		UpdateDirectionalLightUniformBuffers(GraphBuilder, View);
+
+		FRDGTextureRef LastFrameScreenSpaceAO = TryRegisterExternalTexture(GraphBuilder, View.PrevViewInfo.MobileAmbientOcclusion); 
+
+		if (LastFrameScreenSpaceAO == nullptr)
+		{
+			LastFrameScreenSpaceAO = SystemTextures.White;
+		}
+
+		FTextureRHIRef MobilePixelProjectedReflection = View.PrevViewInfo.MobilePixelProjectedReflection.IsValid() ? View.PrevViewInfo.MobilePixelProjectedReflection->GetRenderTargetItem().ShaderResourceTexture : nullptr;
+
+		auto* PassParameters = GraphBuilder.AllocParameters<FMobileRenderPassParameters>();
+		PassParameters->View = View.GetShaderParameters();
+		PassParameters->MobileBasePass = CreateMobileBasePassUniformBuffer(GraphBuilder, View, EMobileBasePass::Opaque, LastFrameScreenSpaceAO, MobilePixelProjectedReflection);
+		PassParameters->SceneTextures = SceneTextures.MobileUniformBuffer;
+		PassParameters->RenderTargets = BasePassRenderTargets;
+
+		if (bShouldRenderPrePass)
+		{
+//			View.ParallelMeshDrawCommandPasses[EMeshPass::DepthPass].BuildRenderingCommands(GraphBuilder, Scene->GPUScene, PassParameters->InstanceCullingDrawParams[EMeshPass::DepthPass]);
+		}
+
+//		View.ParallelMeshDrawCommandPasses[EMeshPass::BasePass].BuildRenderingCommands(GraphBuilder, Scene->GPUScene, PassParameters->InstanceCullingDrawParams[EMeshPass::BasePass]);
+
+		if (bShouldRenderMobileDebugView)
+		{
+			// GPUCULL_TODO: View.ParallelMeshDrawCommandPasses[EMeshPass::DebugViewMode].BuildRenderingCommands(GraphBuilder, Scene->GPUScene, PassParameters->InstanceCullingDrawParams[EMeshPass::DebugViewMode]);
+			PassParameters->DebugViewMode = CreateDebugViewModePassUniformBuffer(GraphBuilder, View, SceneTextures.QuadOverdraw);
+		}
+
+		if (bShouldRenderTranslucency)
+		{
+			// GPUCULL_TODO: View.ParallelMeshDrawCommandPasses[TranslucencyMeshPass].BuildRenderingCommands(GraphBuilder, Scene->GPUScene, PassParameters->InstanceCullingDrawParams[TranslucencyMeshPass]);
+		}
+
+		// Split if we need to render translucency in a separate render pass
+		if (bRequiresMultiPass)
+		{
+			RenderForwardMultiPass(GraphBuilder, PassParameters, BasePassRenderTargets, ViewIndex, View, SceneTextures, bShouldRenderPrePass, bShouldRenderMobileDebugView, bAdrenoOcclusionMode, bShouldRenderOcclusionPass, bShouldRenderTranslucency, TranslucencyMeshPass);
+		}
+		else
+		{
+			RenderForwardSinglePass(GraphBuilder, PassParameters, ViewIndex, View, SceneTextures, bShouldRenderPrePass, bShouldRenderMobileDebugView, bAdrenoOcclusionMode, bShouldRenderOcclusionPass, bShouldRenderTranslucency, TranslucencyMeshPass);
+		}
 	}
 
-	// Depth pre-pass
-	GraphBuilder.SetCommandListStat(GET_STATID(STAT_CLM_MobilePrePass));
-	RenderPrePass(GraphBuilder, BasePassRenderTargets);
-	
-	// Opaque and masked
-	GraphBuilder.SetCommandListStat(GET_STATID(STAT_CLMM_Opaque));
-	RenderMobileBasePass(GraphBuilder, BasePassRenderTargets, ViewList, SceneTextures.ScreenSpaceAO);
-	AddDispatchToRHIThreadPass(GraphBuilder);
+	QueueSceneTextureExtractions(GraphBuilder, SceneTextures);
+}
 
-#if WITH_DEBUG_VIEW_MODES
-	if (ViewFamily.UseDebugViewPS())
+void FMobileSceneRenderer::RenderForwardSinglePass(FRDGBuilder& GraphBuilder, class FMobileRenderPassParameters* PassParameters, int32 ViewIndex, FViewInfo& View, FSceneTextures& SceneTextures, bool bShouldRenderPrePass, bool bShouldRenderMobileDebugView, bool bAdrenoOcclusionMode, bool bShouldRenderOcclusionPass, bool bShouldRenderTranslucency, EMeshPass::Type TranslucencyMeshPass)
+{
+	GraphBuilder.AddPass(
+		RDG_EVENT_NAME("SceneColorRendering"),
+		PassParameters,
+		ERDGPassFlags::Raster,
+		[this, PassParameters, ViewIndex, &View, &SceneTextures, bShouldRenderPrePass, bShouldRenderMobileDebugView, bAdrenoOcclusionMode, bShouldRenderOcclusionPass, bShouldRenderTranslucency, TranslucencyMeshPass](FRHICommandListImmediate& RHICmdList)
 	{
-		// Here we use the base pass depth result to get z culling for opaque and masque.
-		// The color needs to be cleared at this point since shader complexity renders in additive.
-		AddPass(GraphBuilder, [](FRHICommandListImmediate& RHICmdList)
+		if (GIsEditor && !View.bIsSceneCapture)
 		{
+			DrawClearQuad(RHICmdList, View.BackgroundColor);
+		}
+
+		RHICmdList.SetCurrentStat(GET_STATID(STAT_CLM_MobilePrePass));
+		if (bShouldRenderPrePass)
+		{
+			// Depth pre-pass
+			RenderPrePass(RHICmdList, View, &PassParameters->InstanceCullingDrawParams[EMeshPass::DepthPass]);
+		}
+
+		// Opaque and masked
+		RHICmdList.SetCurrentStat(GET_STATID(STAT_CLMM_Opaque));
+		RenderMobileBasePass(RHICmdList, View);
+		RHICmdList.ImmediateFlush(EImmediateFlushType::DispatchToRHIThread);
+
+#if WITH_DEBUG_VIEW_MODES
+		if (bShouldRenderMobileDebugView)
+		{
+			// Here we use the base pass depth result to get z culling for opaque and masque.
+			// The color needs to be cleared at this point since shader complexity renders in additive.
 			DrawClearQuad(RHICmdList, FLinearColor::Black);
-		});
-		RenderMobileDebugView(GraphBuilder, ViewList, SceneTextures.QuadOverdraw);
-		AddDispatchToRHIThreadPass(GraphBuilder);
-	}
+			RenderMobileDebugView(RHICmdList, View);
+			RHICmdList.ImmediateFlush(EImmediateFlushType::DispatchToRHIThread);
+		}
 #endif // WITH_DEBUG_VIEW_MODES
 
-	const bool bAdrenoOcclusionMode = CVarMobileAdrenoOcclusionMode.GetValueOnRenderThread() != 0;
-	if (!bAdrenoOcclusionMode)
-	{
-		// Issue occlusion queries
-		GraphBuilder.SetCommandListStat(GET_STATID(STAT_CLMM_Occlusion));
-		RenderOcclusion(GraphBuilder, BasePassRenderTargets);
-		AddDispatchToRHIThreadPass(GraphBuilder);
-	}
-
-	if (ViewFamily.ViewExtensions.Num() > 1)
-	{
-		auto* ViewExtensionPassParameters = GraphBuilder.AllocParameters<FMobilePostBasePassViewExtensionParameters>();
-		ViewExtensionPassParameters->RenderTargets = BasePassRenderTargets;
-		ViewExtensionPassParameters->SceneTextures = SceneTextures.MobileUniformBuffer;
+		if (!bAdrenoOcclusionMode && bShouldRenderOcclusionPass)
+		{
+			// Issue occlusion queries
+			RHICmdList.SetCurrentStat(GET_STATID(STAT_CLMM_Occlusion));
+			RenderOcclusion(RHICmdList, View);
+			RHICmdList.ImmediateFlush(EImmediateFlushType::DispatchToRHIThread);
+		}
 
-		GraphBuilder.AddPass(
-			{},
-			ViewExtensionPassParameters,
-			ERDGPassFlags::Raster | ERDGPassFlags::SkipRenderPass | ERDGPassFlags::NeverCull,
-			[this](FRHICommandListImmediate& RHICmdList)
+		if (ViewFamily.ViewExtensions.Num() > 1)
 		{
 			CSV_SCOPED_TIMING_STAT_EXCLUSIVE(ViewExtensionPostRenderBasePass);
 			QUICK_SCOPE_CYCLE_COUNTER(STAT_FMobileSceneRenderer_ViewExtensionPostRenderBasePass);
 			for (int32 ViewExt = 0; ViewExt < ViewFamily.ViewExtensions.Num(); ++ViewExt)
 			{
-				for (int32 ViewIndex = 0; ViewIndex < ViewFamily.Views.Num(); ++ViewIndex)
-				{
-					ViewFamily.ViewExtensions[ViewExt]->PostRenderBasePass_RenderThread(RHICmdList, Views[ViewIndex]);
-				}
+				ViewFamily.ViewExtensions[ViewExt]->PostRenderBasePass_RenderThread(RHICmdList, View);
 			}
-		});
-	}
+		}
 
-	// Split if we need to render translucency in a separate render pass
-	// Split if we need to render pixel projected reflection
-	if (bRequiresMultiPass || bRequiresPixelProjectedPlanarRelfectionPass)
-	{
-		AddPass(GraphBuilder, [](FRHICommandListImmediate& RHICmdList)
-		{
-			RHICmdList.EndRenderPass();
-		});
-	}
-	   
-	GraphBuilder.SetCommandListStat(GET_STATID(STAT_CLMM_Translucency));
+		RHICmdList.SetCurrentStat(GET_STATID(STAT_CLMM_Translucency));
 
-	// Restart translucency render pass if needed
-	if (bRequiresMultiPass || bRequiresPixelProjectedPlanarRelfectionPass)
-	{
-		// Make a copy of the scene depth if the current hardware doesn't support reading and writing to the same depth buffer
-		ConditionalResolveSceneDepth(GraphBuilder, View, SceneTextures.Depth);
+		// scene depth is read only and can be fetched
+		RHICmdList.NextSubpass();
 
-		SceneTextures.MobileSetupMode = EMobileSceneTextureSetupMode::SceneDepth | EMobileSceneTextureSetupMode::SceneDepthAux | EMobileSceneTextureSetupMode::CustomDepth;
-		SceneTextures.MobileUniformBuffer = CreateMobileSceneTextureUniformBuffer(GraphBuilder, SceneTextures.MobileSetupMode);
-	}
+		if (!View.bIsPlanarReflection)
+		{
+			if (ViewFamily.EngineShowFlags.Decals)
+			{
+				CSV_SCOPED_TIMING_STAT_EXCLUSIVE(RenderDecals);
+				RenderDecals(RHICmdList, View);
+			}
 
-	if (bRequiresPixelProjectedPlanarRelfectionPass)
-	{
-		const FPlanarReflectionSceneProxy* PlanarReflectionSceneProxy = Scene ? Scene->GetForwardPassGlobalPlanarReflection() : nullptr;
+			if (ViewFamily.EngineShowFlags.DynamicShadows)
+			{
+				CSV_SCOPED_TIMING_STAT_EXCLUSIVE(RenderShadowProjections);
+				RenderModulatedShadowProjections(RHICmdList, ViewIndex, View);
+			}
+		}
 
-		FRDGTextureRef PixelProjectedReflectionTexture = GraphBuilder.RegisterExternalTexture(GPixelProjectedReflectionMobileOutputs.PixelProjectedReflectionTexture);
+		// Draw translucency.
+		if (bShouldRenderTranslucency)
+		{
+			CSV_SCOPED_TIMING_STAT_EXCLUSIVE(RenderTranslucency);
+			SCOPE_CYCLE_COUNTER(STAT_TranslucencyDrawTime);
+			RenderTranslucency(RHICmdList, View, TranslucencyMeshPass);
+			FRHICommandListExecutor::GetImmediateCommandList().PollOcclusionQueries();
+			RHICmdList.ImmediateFlush(EImmediateFlushType::DispatchToRHIThread);
+		}
 
-		RenderPixelProjectedReflection(GraphBuilder, SceneTextures.Color.Resolve, SceneTextures.Depth.Resolve, PixelProjectedReflectionTexture, PlanarReflectionSceneProxy);
-	}
+		if (bAdrenoOcclusionMode)
+		{
+			RHICmdList.SetCurrentStat(GET_STATID(STAT_CLMM_Occlusion));
+			// flush
+			RHICmdList.SubmitCommandsHint();
+			bSubmitOffscreenRendering = false; // submit once
+			// Issue occlusion queries
+			RenderOcclusion(RHICmdList, View);
+			RHICmdList.ImmediateFlush(EImmediateFlushType::DispatchToRHIThread);
+		}
+
+		// Pre-tonemap before MSAA resolve (iOS only)
+		if (!bGammaSpace)
+		{
+			PreTonemapMSAA(RHICmdList, SceneTextures);
+		}
+	});
+}
 
-	if (bRequiresMultiPass || bRequiresPixelProjectedPlanarRelfectionPass)
+void FMobileSceneRenderer::RenderForwardMultiPass(FRDGBuilder& GraphBuilder, class FMobileRenderPassParameters* PassParameters, FRenderTargetBindingSlots& BasePassRenderTargets, int32 ViewIndex, FViewInfo& View, FSceneTextures& SceneTextures, bool bShouldRenderPrePass, bool bShouldRenderMobileDebugView, bool bAdrenoOcclusionMode, bool bShouldRenderOcclusionPass, bool bShouldRenderTranslucency, EMeshPass::Type TranslucencyMeshPass)
+{
+	GraphBuilder.AddPass(
+		RDG_EVENT_NAME("SceneColorRendering"),
+		PassParameters,
+		ERDGPassFlags::Raster,
+		[this, PassParameters, &View, &SceneTextures, bShouldRenderPrePass, bShouldRenderMobileDebugView, bAdrenoOcclusionMode, bShouldRenderOcclusionPass](FRHICommandListImmediate& RHICmdList)
 	{
-		FExclusiveDepthStencil::Type ExclusiveDepthStencil = FExclusiveDepthStencil::DepthRead_StencilRead;
-		if (bModulatedShadowsInUse)
+		if (GIsEditor && !View.bIsSceneCapture)
 		{
-			// FIXME: modulated shadows write to stencil
-			ExclusiveDepthStencil = FExclusiveDepthStencil::DepthRead_StencilWrite;
+			DrawClearQuad(RHICmdList, View.BackgroundColor);
 		}
 
-		// The opaque meshes used for mobile pixel projected reflection have to write depth to depth RT, since we only render the meshes once if the quality level is less or equal to BestPerformance 
-		if (IsMobilePixelProjectedReflectionEnabled(View.GetShaderPlatform())
-			&& GetMobilePixelProjectedReflectionQuality() == EMobilePixelProjectedReflectionQuality::BestPerformance)
+		RHICmdList.SetCurrentStat(GET_STATID(STAT_CLM_MobilePrePass));
+		if (bShouldRenderPrePass)
 		{
-			ExclusiveDepthStencil = FExclusiveDepthStencil::DepthWrite_StencilWrite;
+			// Depth pre-pass
+			RenderPrePass(RHICmdList, View, &PassParameters->InstanceCullingDrawParams[EMeshPass::DepthPass]);
 		}
 
-		BasePassRenderTargets[0].SetLoadAction(ERenderTargetLoadAction::ELoad);
-		BasePassRenderTargets.DepthStencil.SetDepthLoadAction(ERenderTargetLoadAction::ELoad);
-		BasePassRenderTargets.DepthStencil.SetStencilLoadAction(ERenderTargetLoadAction::ELoad);
-		BasePassRenderTargets.DepthStencil.SetDepthStencilAccess(ExclusiveDepthStencil);
-		BasePassRenderTargets.NumOcclusionQueries = 0;
-		BasePassRenderTargets.SubpassHint = ESubpassHint::DepthReadSubpass;
+		// Opaque and masked
+		RHICmdList.SetCurrentStat(GET_STATID(STAT_CLMM_Opaque));
+		RenderMobileBasePass(RHICmdList, View);
+		RHICmdList.ImmediateFlush(EImmediateFlushType::DispatchToRHIThread);
 
-		auto* TranslucencyBasePassParameters = GraphBuilder.AllocParameters<FRenderTargetParameters>();
-		TranslucencyBasePassParameters->RenderTargets = BasePassRenderTargets;
+#if WITH_DEBUG_VIEW_MODES
+		if (bShouldRenderMobileDebugView)
+		{
+			// Here we use the base pass depth result to get z culling for opaque and masque.
+			// The color needs to be cleared at this point since shader complexity renders in additive.
+			DrawClearQuad(RHICmdList, FLinearColor::Black);
+			RenderMobileDebugView(RHICmdList, View);
+			RHICmdList.ImmediateFlush(EImmediateFlushType::DispatchToRHIThread);
+		}
+#endif // WITH_DEBUG_VIEW_MODES
 
-		GraphBuilder.AddPass(
-			{},
-			TranslucencyBasePassParameters,
-			ERDGPassFlags::Raster | ERDGPassFlags::SkipRenderPass,
-			[TranslucencyBasePassParameters](FRHICommandListImmediate& RHICmdList)
+		if (!bAdrenoOcclusionMode && bShouldRenderOcclusionPass)
 		{
-			RHICmdList.BeginRenderPass(GetRenderPassInfo(TranslucencyBasePassParameters), TEXT("SceneColorTranslucencyRendering"));
-		});
+			// Issue occlusion queries
+			RHICmdList.SetCurrentStat(GET_STATID(STAT_CLMM_Occlusion));
+			RenderOcclusion(RHICmdList, View);
+			RHICmdList.ImmediateFlush(EImmediateFlushType::DispatchToRHIThread);
+		}
+
+		if (ViewFamily.ViewExtensions.Num() > 1)
+		{
+			CSV_SCOPED_TIMING_STAT_EXCLUSIVE(ViewExtensionPostRenderBasePass);
+			QUICK_SCOPE_CYCLE_COUNTER(STAT_FMobileSceneRenderer_ViewExtensionPostRenderBasePass);
+			for (int32 ViewExt = 0; ViewExt < ViewFamily.ViewExtensions.Num(); ++ViewExt)
+			{
+				ViewFamily.ViewExtensions[ViewExt]->PostRenderBasePass_RenderThread(RHICmdList, View);
+			}
+		}
+	});
+
+	// Make a copy of the scene depth if the current hardware doesn't support reading and writing to the same depth buffer
+	ConditionalResolveSceneDepth(GraphBuilder, View, SceneTextures.Depth);
+
+	SceneTextures.MobileSetupMode = EMobileSceneTextureSetupMode::SceneDepth | EMobileSceneTextureSetupMode::SceneDepthAux | EMobileSceneTextureSetupMode::CustomDepth;
+	SceneTextures.MobileUniformBuffer = CreateMobileSceneTextureUniformBuffer(GraphBuilder, SceneTextures.MobileSetupMode);
+
+	FExclusiveDepthStencil::Type ExclusiveDepthStencil = FExclusiveDepthStencil::DepthRead_StencilRead;
+	if (bModulatedShadowsInUse)
+	{
+		// FIXME: modulated shadows write to stencil
+		ExclusiveDepthStencil = FExclusiveDepthStencil::DepthRead_StencilWrite;
 	}
 
-	AddPass(GraphBuilder, [](FRHICommandListImmediate& RHICmdList)
+	BasePassRenderTargets[0].SetLoadAction(ERenderTargetLoadAction::ELoad);
+	BasePassRenderTargets.DepthStencil.SetDepthLoadAction(ERenderTargetLoadAction::ELoad);
+	BasePassRenderTargets.DepthStencil.SetStencilLoadAction(ERenderTargetLoadAction::ELoad);
+	BasePassRenderTargets.DepthStencil.SetDepthStencilAccess(ExclusiveDepthStencil);
+	BasePassRenderTargets.NumOcclusionQueries = 0;
+	BasePassRenderTargets.SubpassHint = ESubpassHint::DepthReadSubpass;
+
+	auto* SecondPassParameters = GraphBuilder.AllocParameters<FMobileRenderPassParameters>();
+	*SecondPassParameters = *PassParameters;
+	SecondPassParameters->SceneTextures = SceneTextures.MobileUniformBuffer;
+	SecondPassParameters->RenderTargets = BasePassRenderTargets;
+
+	GraphBuilder.AddPass(
+		{},
+		SecondPassParameters,
+		ERDGPassFlags::Raster,
+		[this, SecondPassParameters, ViewIndex, &View, &SceneTextures, bAdrenoOcclusionMode, bShouldRenderTranslucency, TranslucencyMeshPass](FRHICommandListImmediate& RHICmdList)
 	{
+		RHICmdList.SetCurrentStat(GET_STATID(STAT_CLMM_Translucency));
+
 		// scene depth is read only and can be fetched
 		RHICmdList.NextSubpass();
-	});
 
-	if (!View.bIsPlanarReflection)
-	{
-		if (ViewFamily.EngineShowFlags.Decals)
+		if (!View.bIsPlanarReflection)
 		{
-			CSV_SCOPED_TIMING_STAT_EXCLUSIVE(RenderDecals);
-			RenderDecals(GraphBuilder, BasePassRenderTargets, SceneTextures.MobileUniformBuffer);
-		}
+			if (ViewFamily.EngineShowFlags.Decals)
+			{
+				CSV_SCOPED_TIMING_STAT_EXCLUSIVE(RenderDecals);
+				RenderDecals(RHICmdList, View);
+			}
 
-		if (ViewFamily.EngineShowFlags.DynamicShadows)
-		{
-			CSV_SCOPED_TIMING_STAT_EXCLUSIVE(RenderShadowProjections);
-			RenderModulatedShadowProjections(GraphBuilder, SceneTextures.MobileUniformBuffer);
+			if (ViewFamily.EngineShowFlags.DynamicShadows)
+			{
+				CSV_SCOPED_TIMING_STAT_EXCLUSIVE(RenderShadowProjections);
+				RenderModulatedShadowProjections(RHICmdList, ViewIndex, View);
+			}
 		}
-	}
-	
-	// Draw translucency.
-	if (ViewFamily.EngineShowFlags.Translucency)
-	{
-		CSV_SCOPED_TIMING_STAT_EXCLUSIVE(RenderTranslucency);
-		SCOPE_CYCLE_COUNTER(STAT_TranslucencyDrawTime);
-		RenderTranslucency(GraphBuilder, BasePassRenderTargets, ViewList, SceneTextures.ScreenSpaceAO);
-		AddPass(GraphBuilder, [](FRHICommandListImmediate& RHICmdList)
+
+		// Draw translucency.
+		if (bShouldRenderTranslucency)
 		{
+			CSV_SCOPED_TIMING_STAT_EXCLUSIVE(RenderTranslucency);
+			SCOPE_CYCLE_COUNTER(STAT_TranslucencyDrawTime);
+			RenderTranslucency(RHICmdList, View, TranslucencyMeshPass);
 			FRHICommandListExecutor::GetImmediateCommandList().PollOcclusionQueries();
 			RHICmdList.ImmediateFlush(EImmediateFlushType::DispatchToRHIThread);
-		});
-	}
+		}
 
-	if (bAdrenoOcclusionMode)
-	{
-		GraphBuilder.SetCommandListStat(GET_STATID(STAT_CLMM_Occlusion));
-		AddPass(GraphBuilder, [](FRHICommandListImmediate& RHICmdList)
+		if (bAdrenoOcclusionMode)
 		{
+			RHICmdList.SetCurrentStat(GET_STATID(STAT_CLMM_Occlusion));
 			// flush
 			RHICmdList.SubmitCommandsHint();
-		});
-		bSubmitOffscreenRendering = false; // submit once
-		// Issue occlusion queries
-		RenderOcclusion(GraphBuilder, BasePassRenderTargets);
-		AddDispatchToRHIThreadPass(GraphBuilder);
-	}
-
-	// Pre-tonemap before MSAA resolve (iOS only)
-	if (!bGammaSpace)
-	{
-		PreTonemapMSAA(GraphBuilder, BasePassRenderTargets, SceneTextures);
-	}
+			bSubmitOffscreenRendering = false; // submit once
+			// Issue occlusion queries
+			RenderOcclusion(RHICmdList, View);
+			RHICmdList.ImmediateFlush(EImmediateFlushType::DispatchToRHIThread);
+		}
 
-	AddPass(GraphBuilder, [](FRHICommandListImmediate& RHICmdList)
-	{
-		RHICmdList.EndRenderPass();
+		// Pre-tonemap before MSAA resolve (iOS only)
+		if (!bGammaSpace)
+		{
+			PreTonemapMSAA(RHICmdList, SceneTextures);
+		}
 	});
-
-	QueueSceneTextureExtractions(GraphBuilder, SceneTextures);
 }
 
 class FMobileDeferredCopyPLSPS : public FGlobalShader
@@ -1150,7 +1272,7 @@ void MobileDeferredCopyBuffer(FRHICommandListImmediate& RHICmdList, const FViewI
 		VertexShader);
 }
 
-void FMobileSceneRenderer::RenderDeferred(FRDGBuilder& GraphBuilder, const TArrayView<const FViewInfo> ViewList, const FSortedLightSetSceneInfo& SortedLightSet, FRDGTextureRef ViewFamilyTexture, FSceneTextures& SceneTextures)
+void FMobileSceneRenderer::RenderDeferred(FRDGBuilder& GraphBuilder, const FSortedLightSetSceneInfo& SortedLightSet, FRDGTextureRef ViewFamilyTexture, FSceneTextures& SceneTextures)
 {
 	TArray<FRDGTextureRef, TInlineAllocator<5>> ColorTargets;
 
@@ -1182,227 +1304,255 @@ void FMobileSceneRenderer::RenderDeferred(FRDGBuilder& GraphBuilder, const TArra
 	BasePassRenderTargets.ShadingRateTexture = nullptr;
 	BasePassRenderTargets.MultiViewCount = 0;
 
-	auto* OpaqueBasePassParameters = GraphBuilder.AllocParameters<FRenderTargetParameters>();
-	OpaqueBasePassParameters->RenderTargets = BasePassRenderTargets;
+	const bool bShouldRenderPrePass = ShouldRenderPrePass();
 
-	GraphBuilder.AddPass(
-		{},
-		OpaqueBasePassParameters,
-		ERDGPassFlags::Raster | ERDGPassFlags::SkipRenderPass,
-		[OpaqueBasePassParameters](FRHICommandListImmediate& RHICmdList)
-	{
-		RHICmdList.BeginRenderPass(GetRenderPassInfo(OpaqueBasePassParameters), TEXT("BasePassRendering"));
-	});
-	
-	if (GIsEditor && !Views[0].bIsSceneCapture)
+	const bool bShouldRenderOcclusionPass = DoOcclusionQueries(FeatureLevel);
+
+	ETranslucencyPass::Type TranslucencyPass = ViewFamily.AllowTranslucencyAfterDOF() ? ETranslucencyPass::TPT_StandardTranslucency : ETranslucencyPass::TPT_AllTranslucency;
+
+	const EMeshPass::Type TranslucencyMeshPass = TranslucencyPassToMeshPass(TranslucencyPass);
+
+	const bool bShouldRenderTranslucency = ShouldRenderTranslucency(TranslucencyPass) && ViewFamily.EngineShowFlags.Translucency;
+
+	const FRDGSystemTextures& SystemTextures = FRDGSystemTextures::Get(GraphBuilder);
+
+	int32 NumViews = Views.Num();
+
+	for (int32 ViewIndex = 0; ViewIndex < NumViews; ViewIndex++)
 	{
-		AddPass(GraphBuilder, [this](FRHICommandListImmediate& RHICmdList)
+		FViewInfo& View = Views[ViewIndex];
+
+		SCOPED_GPU_MASK(GraphBuilder.RHICmdList, !View.IsInstancedStereoPass() ? View.GPUMask : (Views[0].GPUMask | Views[1].GPUMask));
+		SCOPED_CONDITIONAL_DRAW_EVENTF(GraphBuilder.RHICmdList, EventView, NumViews > 1, TEXT("View%d"), ViewIndex);
+
+		if (!View.ShouldRenderView())
 		{
-			DrawClearQuad(RHICmdList, Views[0].BackgroundColor);
-		});
-	}
+			continue;
+		}
 
-	GraphBuilder.SetCommandListStat(GET_STATID(STAT_CLM_MobilePrePass));
-	// Depth pre-pass
-	RenderPrePass(GraphBuilder, BasePassRenderTargets);
-	
-	// Opaque and masked
-	GraphBuilder.SetCommandListStat(GET_STATID(STAT_CLMM_Opaque));
-	RenderMobileBasePass(GraphBuilder, BasePassRenderTargets, ViewList, SceneTextures.ScreenSpaceAO);
-	AddDispatchToRHIThreadPass(GraphBuilder);
+		View.BeginRenderView();
 
-	// Issue occlusion queries
-	GraphBuilder.SetCommandListStat(GET_STATID(STAT_CLMM_Occlusion));
-	RenderOcclusion(GraphBuilder, BasePassRenderTargets);
-	AddDispatchToRHIThreadPass(GraphBuilder);
+		UpdateDirectionalLightUniformBuffers(GraphBuilder, View);
 
-	auto PollOcclusionQueriesAndDispatchToRHIThreadPass = [](FRHICommandListImmediate& InRHICmdList)
-	{
-		FRHICommandListExecutor::GetImmediateCommandList().PollOcclusionQueries();
-		InRHICmdList.ImmediateFlush(EImmediateFlushType::DispatchToRHIThread);
-	};
+		FRDGTextureRef LastFrameScreenSpaceAO = TryRegisterExternalTexture(GraphBuilder, View.PrevViewInfo.MobileAmbientOcclusion);
 
-	if (!bRequiresMultiPass)
-	{
-		AddPass(GraphBuilder, [](FRHICommandListImmediate& RHICmdList)
+		if (LastFrameScreenSpaceAO == nullptr)
 		{
-			// SceneColor + GBuffer write, SceneDepth is read only
-			RHICmdList.NextSubpass();
-		});
-		
-		if (ViewFamily.EngineShowFlags.Decals)
+			LastFrameScreenSpaceAO = SystemTextures.White;
+		}
+
+		FTextureRHIRef MobilePixelProjectedReflection = View.PrevViewInfo.MobilePixelProjectedReflection.IsValid() ? View.PrevViewInfo.MobilePixelProjectedReflection->GetRenderTargetItem().ShaderResourceTexture : nullptr;
+
+		auto* PassParameters = GraphBuilder.AllocParameters<FMobileRenderPassParameters>();
+		PassParameters->View = View.GetShaderParameters();
+		PassParameters->MobileBasePass = CreateMobileBasePassUniformBuffer(GraphBuilder, View, EMobileBasePass::Opaque, LastFrameScreenSpaceAO, MobilePixelProjectedReflection);
+		PassParameters->SceneTextures = SceneTextures.MobileUniformBuffer;
+		PassParameters->RenderTargets = BasePassRenderTargets;
+
+		if (bShouldRenderPrePass)
 		{
-			CSV_SCOPED_TIMING_STAT_EXCLUSIVE(RenderDecals);
-			RenderDecals(GraphBuilder, BasePassRenderTargets, SceneTextures.MobileUniformBuffer);
+			View.ParallelMeshDrawCommandPasses[EMeshPass::DepthPass].BuildRenderingCommands(GraphBuilder, Scene->GPUScene, PassParameters->InstanceCullingDrawParams[EMeshPass::DepthPass]);
 		}
 
-		AddPass(GraphBuilder, [](FRHICommandListImmediate& RHICmdList)
+		if (bRequiresMultiPass)
 		{
-			// SceneColor write, SceneDepth is read only
-			RHICmdList.NextSubpass();
-		});
+			RenderDeferredMultiPass(GraphBuilder, PassParameters, BasePassRenderTargets, ColorTargets.Num(), ViewIndex, NumViews, View, SceneTextures, SortedLightSet, MobilePixelProjectedReflection, bShouldRenderPrePass, bShouldRenderOcclusionPass, bShouldRenderTranslucency, TranslucencyMeshPass);
+		}
+		else
+		{
+			RenderDeferredSinglePass(GraphBuilder, PassParameters, ViewIndex, NumViews, View, SceneTextures, SortedLightSet, MobilePixelProjectedReflection, bShouldRenderPrePass, bShouldRenderOcclusionPass, bUsingPixelLocalStorage, bShouldRenderTranslucency, TranslucencyMeshPass);
+		}
+	}
 
-		if (bUsingPixelLocalStorage)
+	QueueSceneTextureExtractions(GraphBuilder, SceneTextures);
+}
+
+void FMobileSceneRenderer::RenderDeferredSinglePass(FRDGBuilder& GraphBuilder, class FMobileRenderPassParameters* PassParameters, int32 ViewIndex, int32 NumViews, FViewInfo& View, FSceneTextures& SceneTextures, const FSortedLightSetSceneInfo& SortedLightSet, FTextureRHIRef MobilePixelProjectedReflection, bool bShouldRenderPrePass, bool bShouldRenderOcclusionPass, bool bUsingPixelLocalStorage, bool bShouldRenderTranslucency, EMeshPass::Type TranslucencyMeshPass)
+{
+	GraphBuilder.AddPass(
+		RDG_EVENT_NAME("SceneColorRendering"),
+		PassParameters,
+		ERDGPassFlags::Raster,
+		[this, PassParameters, ViewIndex, NumViews, &View, &SceneTextures, &SortedLightSet, MobilePixelProjectedReflection, bShouldRenderPrePass, bShouldRenderOcclusionPass, bUsingPixelLocalStorage, bShouldRenderTranslucency, TranslucencyMeshPass](FRHICommandListImmediate& RHICmdList)
+	{
+		if (GIsEditor && !View.bIsSceneCapture)
 		{
-			{
-				auto* PassParameters = GraphBuilder.AllocParameters<FRenderTargetParameters>();
-				PassParameters->RenderTargets = BasePassRenderTargets;
+			DrawClearQuad(RHICmdList, View.BackgroundColor);
+		}
 
-				AddPass(GraphBuilder, [this](FRHICommandListImmediate& RHICmdList)
-				{
-					MobileDeferredCopyBuffer<FMobileDeferredCopyPLSPS>(RHICmdList, Views[0]);
-				});
-			}
+		RHICmdList.SetCurrentStat(GET_STATID(STAT_CLM_MobilePrePass));
+		if (bShouldRenderPrePass)
+		{
+			// Depth pre-pass
+			RenderPrePass(RHICmdList, View, &PassParameters->InstanceCullingDrawParams[EMeshPass::DepthPass]);
+		}
+
+		// Opaque and masked
+		RHICmdList.SetCurrentStat(GET_STATID(STAT_CLMM_Opaque));
+		RenderMobileBasePass(RHICmdList, View);
+		RHICmdList.ImmediateFlush(EImmediateFlushType::DispatchToRHIThread);
+
+		if (bShouldRenderOcclusionPass)
+		{
+			// Issue occlusion queries
+			RHICmdList.SetCurrentStat(GET_STATID(STAT_CLMM_Occlusion));
+			RenderOcclusion(RHICmdList, View);
+			RHICmdList.ImmediateFlush(EImmediateFlushType::DispatchToRHIThread);
+		}
+
+		// SceneColor + GBuffer write, SceneDepth is read only
+		RHICmdList.NextSubpass();
 
+		if (!View.bIsPlanarReflection)
+		{
+			if (ViewFamily.EngineShowFlags.Decals)
 			{
-				AddPass(GraphBuilder, [](FRHICommandListImmediate& RHICmdList)
-				{
-					// SceneColor write, SceneDepth is read only
-					RHICmdList.NextSubpass();
-				});
+				CSV_SCOPED_TIMING_STAT_EXCLUSIVE(RenderDecals);
+				RenderDecals(RHICmdList, View);
 			}
 		}
-		
-		MobileDeferredShadingPass(GraphBuilder, BasePassRenderTargets, SceneTextures.MobileUniformBuffer, *Scene, ViewList, SortedLightSet);
+
+		// SceneColor write, SceneDepth is read only
+		RHICmdList.NextSubpass();
+
+		if (bUsingPixelLocalStorage)
+		{
+			MobileDeferredCopyBuffer<FMobileDeferredCopyPLSPS>(RHICmdList, View);
+
+			// SceneColor write, SceneDepth is read only
+			RHICmdList.NextSubpass();
+		}
+
+		MobileDeferredShadingPass(RHICmdList, ViewIndex, NumViews, View, *Scene, SortedLightSet, MobilePixelProjectedReflection);
+
 		// Draw translucency.
-		if (ViewFamily.EngineShowFlags.Translucency)
+		if (bShouldRenderTranslucency)
 		{
 			CSV_SCOPED_TIMING_STAT_EXCLUSIVE(RenderTranslucency);
 			SCOPE_CYCLE_COUNTER(STAT_TranslucencyDrawTime);
-			RenderTranslucency(GraphBuilder, BasePassRenderTargets, ViewList, SceneTextures.ScreenSpaceAO);
-			AddPass(GraphBuilder, PollOcclusionQueriesAndDispatchToRHIThreadPass);
+			RenderTranslucency(RHICmdList, View, TranslucencyMeshPass);
+			FRHICommandListExecutor::GetImmediateCommandList().PollOcclusionQueries();
+			RHICmdList.ImmediateFlush(EImmediateFlushType::DispatchToRHIThread);
 		}
+	});
+}
 
-		AddPass(GraphBuilder, [](FRHICommandListImmediate& RHICmdList)
-		{
-			RHICmdList.EndRenderPass();
-		});
-	}
-	else
+void FMobileSceneRenderer::RenderDeferredMultiPass(FRDGBuilder& GraphBuilder, class FMobileRenderPassParameters* PassParameters, FRenderTargetBindingSlots& BasePassRenderTargets, int32 NumColorTargets, int32 ViewIndex, int32 NumViews, FViewInfo& View, FSceneTextures& SceneTextures, const FSortedLightSetSceneInfo& SortedLightSet, FTextureRHIRef MobilePixelProjectedReflection, bool bShouldRenderPrePass, bool bShouldRenderOcclusionPass, bool bShouldRenderTranslucency, EMeshPass::Type TranslucencyMeshPass)
+{
+	GraphBuilder.AddPass(
+		RDG_EVENT_NAME("SceneColorRendering"),
+		PassParameters,
+		ERDGPassFlags::Raster,
+		[this, PassParameters, &View, &SceneTextures, bShouldRenderPrePass, bShouldRenderOcclusionPass](FRHICommandListImmediate& RHICmdList)
 	{
-		AddPass(GraphBuilder, [](FRHICommandListImmediate& RHICmdList)
+		if (GIsEditor && !View.bIsSceneCapture)
 		{
-			RHICmdList.NextSubpass();
-			RHICmdList.NextSubpass();
-			RHICmdList.EndRenderPass();
-		});
-		
-		// SceneColor + GBuffer write, SceneDepth is read only
+			DrawClearQuad(RHICmdList, View.BackgroundColor);
+		}
+
+		RHICmdList.SetCurrentStat(GET_STATID(STAT_CLM_MobilePrePass));
+		if (bShouldRenderPrePass)
 		{
-			SceneTextures.MobileSetupMode = EMobileSceneTextureSetupMode::SceneDepth | EMobileSceneTextureSetupMode::SceneDepthAux | EMobileSceneTextureSetupMode::CustomDepth;
-			SceneTextures.MobileUniformBuffer = CreateMobileSceneTextureUniformBuffer(GraphBuilder, SceneTextures.MobileSetupMode);
+			// Depth pre-pass
+			RenderPrePass(RHICmdList, View, &PassParameters->InstanceCullingDrawParams[EMeshPass::DepthPass]);
+		}
 
-			for (int32 i = 0; i < ColorTargets.Num(); ++i)
-			{
-				BasePassRenderTargets[i].SetLoadAction(ERenderTargetLoadAction::ELoad);
-			}
+		// Opaque and masked
+		RHICmdList.SetCurrentStat(GET_STATID(STAT_CLMM_Opaque));
+		RenderMobileBasePass(RHICmdList, View);
+		RHICmdList.ImmediateFlush(EImmediateFlushType::DispatchToRHIThread);
 
-			BasePassRenderTargets.DepthStencil.SetDepthLoadAction(ERenderTargetLoadAction::ELoad);
-			BasePassRenderTargets.DepthStencil.SetStencilLoadAction(ERenderTargetLoadAction::ELoad);
-			BasePassRenderTargets.DepthStencil.SetDepthStencilAccess(FExclusiveDepthStencil::DepthRead_StencilRead);
-			BasePassRenderTargets.SubpassHint = ESubpassHint::None;
-			BasePassRenderTargets.NumOcclusionQueries = 0;
+		if (bShouldRenderOcclusionPass)
+		{
+			// Issue occlusion queries
+			RHICmdList.SetCurrentStat(GET_STATID(STAT_CLMM_Occlusion));
+			RenderOcclusion(RHICmdList, View);
+			RHICmdList.ImmediateFlush(EImmediateFlushType::DispatchToRHIThread);
+		}
 
-			auto* DecalPassParameters = GraphBuilder.AllocParameters<FRenderTargetParameters>();
-			DecalPassParameters->RenderTargets = BasePassRenderTargets;
+		RHICmdList.NextSubpass();
+		RHICmdList.NextSubpass();
+	});
 
-			GraphBuilder.AddPass(
-				{},
-				DecalPassParameters,
-				ERDGPassFlags::Raster | ERDGPassFlags::SkipRenderPass,
-				[DecalPassParameters](FRHICommandListImmediate& RHICmdList)
-			{
-				RHICmdList.BeginRenderPass(GetRenderPassInfo(DecalPassParameters), TEXT("AfterBasePass"));
-			});
+	// SceneColor + GBuffer write, SceneDepth is read only
+	SceneTextures.MobileSetupMode = EMobileSceneTextureSetupMode::SceneDepth | EMobileSceneTextureSetupMode::SceneDepthAux | EMobileSceneTextureSetupMode::CustomDepth;
+	SceneTextures.MobileUniformBuffer = CreateMobileSceneTextureUniformBuffer(GraphBuilder, SceneTextures.MobileSetupMode);
+
+	for (int32 i = 0; i < NumColorTargets; ++i)
+	{
+		BasePassRenderTargets[i].SetLoadAction(ERenderTargetLoadAction::ELoad);
+	}
 
+	BasePassRenderTargets.DepthStencil.SetDepthLoadAction(ERenderTargetLoadAction::ELoad);
+	BasePassRenderTargets.DepthStencil.SetStencilLoadAction(ERenderTargetLoadAction::ELoad);
+	BasePassRenderTargets.DepthStencil.SetDepthStencilAccess(FExclusiveDepthStencil::DepthRead_StencilRead);
+	BasePassRenderTargets.SubpassHint = ESubpassHint::None;
+	BasePassRenderTargets.NumOcclusionQueries = 0;
+
+	auto* SecondPassParameters = GraphBuilder.AllocParameters<FMobileRenderPassParameters>();
+	*SecondPassParameters = *PassParameters;
+	SecondPassParameters->SceneTextures = SceneTextures.MobileUniformBuffer;
+	SecondPassParameters->RenderTargets = BasePassRenderTargets;
+
+	GraphBuilder.AddPass(
+		{},
+		SecondPassParameters,
+		ERDGPassFlags::Raster,
+		[this, SecondPassParameters, &View](FRHICommandListImmediate& RHICmdList)
+	{
+		if (!View.bIsPlanarReflection)
+		{
 			if (ViewFamily.EngineShowFlags.Decals)
 			{
 				CSV_SCOPED_TIMING_STAT_EXCLUSIVE(RenderDecals);
-				RenderDecals(GraphBuilder, BasePassRenderTargets, SceneTextures.MobileUniformBuffer);
+				RenderDecals(RHICmdList, View);
 			}
-
-			AddPass(GraphBuilder, [](FRHICommandListImmediate& RHICmdList)
-			{
-				RHICmdList.EndRenderPass();
-			});
 		}
+	});
 
-		// SceneColor write, SceneDepth is read only
-		{
-			SceneTextures.MobileSetupMode = EMobileSceneTextureSetupMode::SceneDepth | EMobileSceneTextureSetupMode::SceneDepthAux | EMobileSceneTextureSetupMode::GBuffers | EMobileSceneTextureSetupMode::CustomDepth;
-			SceneTextures.MobileUniformBuffer = CreateMobileSceneTextureUniformBuffer(GraphBuilder, SceneTextures.MobileSetupMode);
-
-			for (int32 i = 1; i < ColorTargets.Num(); ++i)
-			{
-				BasePassRenderTargets[i] = FRenderTargetBinding();
-			}
-			BasePassRenderTargets.DepthStencil.SetDepthStencilAccess(FExclusiveDepthStencil::DepthRead_StencilWrite);
+	// SceneColor write, SceneDepth is read only
+	SceneTextures.MobileSetupMode = EMobileSceneTextureSetupMode::SceneDepth | EMobileSceneTextureSetupMode::SceneDepthAux | EMobileSceneTextureSetupMode::GBuffers | EMobileSceneTextureSetupMode::CustomDepth;
+	SceneTextures.MobileUniformBuffer = CreateMobileSceneTextureUniformBuffer(GraphBuilder, SceneTextures.MobileSetupMode);
 
-			auto* DeferredShadingPassParameters = GraphBuilder.AllocParameters<FRenderTargetParameters>();
-			DeferredShadingPassParameters->RenderTargets = BasePassRenderTargets;
+	for (int32 i = 1; i < NumColorTargets; ++i)
+	{
+		BasePassRenderTargets[i] = FRenderTargetBinding();
+	}
+	BasePassRenderTargets.DepthStencil.SetDepthStencilAccess(FExclusiveDepthStencil::DepthRead_StencilWrite);
+	auto* ThirdPassParameters = GraphBuilder.AllocParameters<FMobileRenderPassParameters>();
+	*ThirdPassParameters = *PassParameters;
+	ThirdPassParameters->SceneTextures = SceneTextures.MobileUniformBuffer;
+	ThirdPassParameters->RenderTargets = BasePassRenderTargets;
 
-			GraphBuilder.AddPass(
-				{},
-				DeferredShadingPassParameters,
-				ERDGPassFlags::Raster | ERDGPassFlags::SkipRenderPass,
-				[DeferredShadingPassParameters](FRHICommandListImmediate& RHICmdList)
-			{
-				RHICmdList.BeginRenderPass(GetRenderPassInfo(DeferredShadingPassParameters), TEXT("MobileShadingPass"));
-			});
-			
-			MobileDeferredShadingPass(GraphBuilder, BasePassRenderTargets, SceneTextures.MobileUniformBuffer, *Scene, ViewList, SortedLightSet);
-			// Draw translucency.
-			if (ViewFamily.EngineShowFlags.Translucency)
-			{
-				CSV_SCOPED_TIMING_STAT_EXCLUSIVE(RenderTranslucency);
-				SCOPE_CYCLE_COUNTER(STAT_TranslucencyDrawTime);
-				RenderTranslucency(GraphBuilder, BasePassRenderTargets, ViewList, SceneTextures.ScreenSpaceAO);
-				AddPass(GraphBuilder, PollOcclusionQueriesAndDispatchToRHIThreadPass);
-			}
+	GraphBuilder.AddPass(
+		{},
+		ThirdPassParameters,
+		ERDGPassFlags::Raster,
+		[this, ThirdPassParameters, ViewIndex, NumViews, &View, &SceneTextures, &SortedLightSet, MobilePixelProjectedReflection, bShouldRenderTranslucency, TranslucencyMeshPass](FRHICommandListImmediate& RHICmdList)
+	{
+		MobileDeferredShadingPass(RHICmdList, ViewIndex, NumViews, View, *Scene, SortedLightSet, MobilePixelProjectedReflection);
 
-			AddPass(GraphBuilder, [](FRHICommandListImmediate& RHICmdList)
-			{
-				RHICmdList.EndRenderPass();
-			});
+		// Draw translucency.
+		if (bShouldRenderTranslucency)
+		{
+			CSV_SCOPED_TIMING_STAT_EXCLUSIVE(RenderTranslucency);
+			SCOPE_CYCLE_COUNTER(STAT_TranslucencyDrawTime);
+			RenderTranslucency(RHICmdList, View, TranslucencyMeshPass);
+			FRHICommandListExecutor::GetImmediateCommandList().PollOcclusionQueries();
+			RHICmdList.ImmediateFlush(EImmediateFlushType::DispatchToRHIThread);
 		}
-	}
+	});
 }
 
-BEGIN_SHADER_PARAMETER_STRUCT(FMobileDebugViewPassParameters, )
-	SHADER_PARAMETER_RDG_UNIFORM_BUFFER(FDebugViewModePassUniformParameters, DebugViewMode)
-END_SHADER_PARAMETER_STRUCT()
-
-void FMobileSceneRenderer::RenderMobileDebugView(FRDGBuilder& GraphBuilder, const TArrayView<const FViewInfo> PassViews, FRDGTextureRef QuadOverdrawTexture)
+void FMobileSceneRenderer::RenderMobileDebugView(FRHICommandListImmediate& RHICmdList, const FViewInfo& View)
 {
 #if WITH_DEBUG_VIEW_MODES
 	CSV_SCOPED_TIMING_STAT_EXCLUSIVE(RenderDebugView);
-	RDG_EVENT_SCOPE(GraphBuilder, "MobileDebugView");
+	SCOPED_DRAW_EVENT(RHICmdList, MobileDebugView);
 	SCOPE_CYCLE_COUNTER(STAT_BasePassDrawTime);
 
-	for (int32 ViewIndex = 0; ViewIndex < PassViews.Num(); ViewIndex++)
-	{
-		RDG_EVENT_SCOPE_CONDITIONAL(GraphBuilder, Views.Num() > 1, "View%d", ViewIndex);
-		const FViewInfo& View = PassViews[ViewIndex];
-		if (!View.ShouldRenderView())
-		{
-			continue;
-		}
-
-		// GPUCULL_TODO: View.ParallelMeshDrawCommandPasses[EMeshPass::DebugViewMode].BuildRenderingCommands(GraphBuilder, Scene->GPUScene);
-
-		auto* PassParameters = GraphBuilder.AllocParameters<FMobileDebugViewPassParameters>();
-		PassParameters->DebugViewMode = CreateDebugViewModePassUniformBuffer(GraphBuilder, View, QuadOverdrawTexture);
-
-		GraphBuilder.AddPass(
-			{},
-			PassParameters,
-			ERDGPassFlags::Raster | ERDGPassFlags::SkipRenderPass | ERDGPassFlags::NeverCull,
-			[&View](FRHICommandList& RHICmdList)
-		{
-			RHICmdList.SetViewport(View.ViewRect.Min.X, View.ViewRect.Min.Y, 0, View.ViewRect.Max.X, View.ViewRect.Max.Y, 1);
-			View.ParallelMeshDrawCommandPasses[EMeshPass::DebugViewMode].DispatchDraw(nullptr, RHICmdList);
-		});
-	}
+	RHICmdList.SetViewport(View.ViewRect.Min.X, View.ViewRect.Min.Y, 0, View.ViewRect.Max.X, View.ViewRect.Max.Y, 1);
+	View.ParallelMeshDrawCommandPasses[EMeshPass::DebugViewMode].DispatchDraw(nullptr, RHICmdList);
 #endif // WITH_DEBUG_VIEW_MODES
 }
 
@@ -1488,7 +1638,7 @@ void FMobileSceneRenderer::ConditionalResolveSceneDepth(FRDGBuilder& GraphBuilde
 		bool bDecals = ViewFamily.EngineShowFlags.Decals && Scene->Decals.Num();
 		bool bModulatedShadows = ViewFamily.EngineShowFlags.DynamicShadows && bModulatedShadowsInUse;
 
-		if (bDecals || bModulatedShadows || bAlwaysResolveDepth || View.bUsesSceneDepth || bRequiresPixelProjectedPlanarRelfectionPass)
+		if (bDecals || bModulatedShadows || bAlwaysResolveDepth || View.bUsesSceneDepth)
 		{
 			// WEBGL copies depth from SceneColor alpha to a separate texture
 			// Switch target to force hardware flush current depth to texture
@@ -1601,7 +1751,7 @@ public:
 
 IMPLEMENT_SHADER_TYPE(, FPreTonemapMSAA_Mobile,TEXT("/Engine/Private/PostProcessMobile.usf"),TEXT("PreTonemapMSAA_Mobile"),SF_Pixel);
 
-void FMobileSceneRenderer::PreTonemapMSAA(FRDGBuilder& GraphBuilder, FRenderTargetBindingSlots& BasePassRenderTargets, const FMinimalSceneTextures& SceneTextures)
+void FMobileSceneRenderer::PreTonemapMSAA(FRHICommandListImmediate& RHICmdList, const FMinimalSceneTextures& SceneTextures)
 {
 	// iOS only
 	bool bOnChipPP = GSupportsRenderTargetFormat_PF_FloatRGBA && GSupportsShaderFramebufferFetch &&	ViewFamily.EngineShowFlags.PostProcessing;
@@ -1613,45 +1763,35 @@ void FMobileSceneRenderer::PreTonemapMSAA(FRDGBuilder& GraphBuilder, FRenderTarg
 
 	const FIntPoint TargetSize = SceneTextures.Config.Extent;
 
-	auto* PassParameters = GraphBuilder.AllocParameters<FRenderTargetParameters>();
-	PassParameters->RenderTargets = BasePassRenderTargets;
-
 	const auto ShaderMap = GetGlobalShaderMap(FeatureLevel);
 	TShaderMapRef<FScreenVS> VertexShader(ShaderMap);
 	TShaderMapRef<FPreTonemapMSAA_Mobile> PixelShader(ShaderMap);
 
-	GraphBuilder.AddPass(
-		RDG_EVENT_NAME("PreTonemapMSAAPass"),
-		PassParameters,
-		ERDGPassFlags::Raster | ERDGPassFlags::SkipRenderPass,
-		[this, TargetSize, VertexShader, PixelShader](FRHICommandListImmediate& RHICmdList)
-	{
-		FGraphicsPipelineStateInitializer GraphicsPSOInit;
-		RHICmdList.ApplyCachedRenderTargets(GraphicsPSOInit);
-		GraphicsPSOInit.BlendState = TStaticBlendState<>::GetRHI();
-		GraphicsPSOInit.RasterizerState = TStaticRasterizerState<>::GetRHI();
-		GraphicsPSOInit.DepthStencilState = TStaticDepthStencilState<false, CF_Always>::GetRHI();
-
-		GraphicsPSOInit.BoundShaderState.VertexDeclarationRHI = GFilterVertexDeclaration.VertexDeclarationRHI;
-		GraphicsPSOInit.BoundShaderState.VertexShaderRHI = VertexShader.GetVertexShader();
-		GraphicsPSOInit.BoundShaderState.PixelShaderRHI = PixelShader.GetPixelShader();
-		GraphicsPSOInit.PrimitiveType = PT_TriangleList;
-
-		SetGraphicsPipelineState(RHICmdList, GraphicsPSOInit);
-
-		RHICmdList.SetViewport(0, 0, 0.0f, TargetSize.X, TargetSize.Y, 1.0f);
-
-		DrawRectangle(
-			RHICmdList,
-			0, 0,
-			TargetSize.X, TargetSize.Y,
-			0, 0,
-			TargetSize.X, TargetSize.Y,
-			TargetSize,
-			TargetSize,
-			VertexShader,
-			EDRF_UseTriangleOptimization);
-	});
+	FGraphicsPipelineStateInitializer GraphicsPSOInit;
+	RHICmdList.ApplyCachedRenderTargets(GraphicsPSOInit);
+	GraphicsPSOInit.BlendState = TStaticBlendState<>::GetRHI();
+	GraphicsPSOInit.RasterizerState = TStaticRasterizerState<>::GetRHI();
+	GraphicsPSOInit.DepthStencilState = TStaticDepthStencilState<false, CF_Always>::GetRHI();
+
+	GraphicsPSOInit.BoundShaderState.VertexDeclarationRHI = GFilterVertexDeclaration.VertexDeclarationRHI;
+	GraphicsPSOInit.BoundShaderState.VertexShaderRHI = VertexShader.GetVertexShader();
+	GraphicsPSOInit.BoundShaderState.PixelShaderRHI = PixelShader.GetPixelShader();
+	GraphicsPSOInit.PrimitiveType = PT_TriangleList;
+
+	SetGraphicsPipelineState(RHICmdList, GraphicsPSOInit);
+
+	RHICmdList.SetViewport(0, 0, 0.0f, TargetSize.X, TargetSize.Y, 1.0f);
+
+	DrawRectangle(
+		RHICmdList,
+		0, 0,
+		TargetSize.X, TargetSize.Y,
+		0, 0,
+		TargetSize.X, TargetSize.Y,
+		TargetSize,
+		TargetSize,
+		VertexShader,
+		EDRF_UseTriangleOptimization);
 }
 
 /** Before SetupMobileBasePassAfterShadowInit, we need to update the uniform buffer and shadow info for all movable point lights.*/
diff --git a/Engine/Source/Runtime/Renderer/Private/MobileTranslucentRendering.cpp b/Engine/Source/Runtime/Renderer/Private/MobileTranslucentRendering.cpp
--- a/Engine/Source/Runtime/Renderer/Private/MobileTranslucentRendering.cpp
+++ b/Engine/Source/Runtime/Renderer/Private/MobileTranslucentRendering.cpp
@@ -31,48 +31,16 @@
 #include "MeshPassProcessor.inl"
 #include "ClearQuad.h"
 
-void FMobileSceneRenderer::RenderTranslucency(FRDGBuilder& GraphBuilder, FRenderTargetBindingSlots& BasePassRenderTargets, const TArrayView<const FViewInfo> PassViews, FRDGTextureRef ScreenSpaceAO)
+void FMobileSceneRenderer::RenderTranslucency(FRHICommandListImmediate& RHICmdList, const FViewInfo& View, const EMeshPass::Type TranslucencyMeshPass)
 {
-	ETranslucencyPass::Type TranslucencyPass = 
-		ViewFamily.AllowTranslucencyAfterDOF() ? ETranslucencyPass::TPT_StandardTranslucency : ETranslucencyPass::TPT_AllTranslucency;
-		
-	if (ShouldRenderTranslucency(TranslucencyPass))
-	{
-		RDG_EVENT_SCOPE(GraphBuilder, "Translucency");
-		RDG_GPU_STAT_SCOPE(GraphBuilder, Translucency);
-
-		for (int32 ViewIndex = 0; ViewIndex < PassViews.Num(); ViewIndex++)
-		{
-			RDG_EVENT_SCOPE_CONDITIONAL(GraphBuilder, Views.Num() > 1, "View%d", ViewIndex);
-
-			const FViewInfo& View = PassViews[ViewIndex];
-			if (!View.ShouldRenderView())
-			{
-				continue;
-			}
-
-			// GPUCULL_TODO: View.ParallelMeshDrawCommandPasses[MeshPass].BuildRenderingCommands(GraphBuilder, Scene->GPUScene);
+	SCOPED_DRAW_EVENT(RHICmdList, Translucency);
+	SCOPED_GPU_STAT(RHICmdList, Translucency);
 
-			View.BeginRenderView();
-			UpdateDirectionalLightUniformBuffers(GraphBuilder, View);
+	RHICmdList.SetViewport(View.ViewRect.Min.X, View.ViewRect.Min.Y, 0.0f, View.ViewRect.Max.X, View.ViewRect.Max.Y, 1.0f);
 
-			auto* TranslucencyBasePassParameters = GraphBuilder.AllocParameters<FMobileBasePassParameters>();
-			TranslucencyBasePassParameters->View = View.GetShaderParameters();
-			TranslucencyBasePassParameters->MobileBasePass = CreateMobileBasePassUniformBuffer(GraphBuilder, View, EMobileBasePass::Translucent, ScreenSpaceAO);
-			TranslucencyBasePassParameters->RenderTargets = BasePassRenderTargets;
-
-			GraphBuilder.AddPass(RDG_EVENT_NAME("RenderTranslucencyBasePass"), TranslucencyBasePassParameters, ERDGPassFlags::Raster | ERDGPassFlags::SkipRenderPass,
-				[this, &View, TranslucencyPass](FRHICommandListImmediate& RHICmdList)
-			{
-				RHICmdList.SetViewport(View.ViewRect.Min.X, View.ViewRect.Min.Y, 0.0f, View.ViewRect.Max.X, View.ViewRect.Max.Y, 1.0f);
-
-				if (!View.Family->UseDebugViewPS())
-				{
-					const EMeshPass::Type MeshPass = TranslucencyPassToMeshPass(TranslucencyPass);
-					View.ParallelMeshDrawCommandPasses[MeshPass].DispatchDraw(nullptr, RHICmdList);
-				}
-			});
-		}
+	if (!View.Family->UseDebugViewPS())
+	{
+		View.ParallelMeshDrawCommandPasses[TranslucencyMeshPass].DispatchDraw(nullptr, RHICmdList);
 	}
 }
 
@@ -85,7 +53,7 @@ void FMobileSceneRenderer::RenderInverseOpacity(FRDGBuilder& GraphBuilder, const
 
 	auto* InverseOpacityParameters = GraphBuilder.AllocParameters<FMobileBasePassParameters>();
 	InverseOpacityParameters->View = View.GetShaderParameters();
-	InverseOpacityParameters->MobileBasePass = CreateMobileBasePassUniformBuffer(GraphBuilder, View, EMobileBasePass::Translucent, SceneTextures.ScreenSpaceAO);
+	InverseOpacityParameters->MobileBasePass = CreateMobileBasePassUniformBuffer(GraphBuilder, View, EMobileBasePass::Translucent);
 	InverseOpacityParameters->RenderTargets[0] = FRenderTargetBinding(SceneTextures.Color.Target, SceneTextures.Color.Resolve, ERenderTargetLoadAction::EClear);
 	InverseOpacityParameters->RenderTargets.DepthStencil = FDepthStencilBinding(SceneTextures.Depth.Target, ERenderTargetLoadAction::EClear, FExclusiveDepthStencil::DepthWrite_StencilWrite);
 	// Opacity could fetch depth as we use exactly the same shaders as in base pass
diff --git a/Engine/Source/Runtime/Renderer/Private/PostProcess/PostProcessAmbientOcclusionMobile.cpp b/Engine/Source/Runtime/Renderer/Private/PostProcess/PostProcessAmbientOcclusionMobile.cpp
--- a/Engine/Source/Runtime/Renderer/Private/PostProcess/PostProcessAmbientOcclusionMobile.cpp
+++ b/Engine/Source/Runtime/Renderer/Private/PostProcess/PostProcessAmbientOcclusionMobile.cpp
@@ -501,6 +501,7 @@ static void RenderGTAO(FRDGBuilder& GraphBuilder, FRDGTextureRef SceneDepthTextu
 			FComputeShaderUtils::AddPass(
 				GraphBuilder,
 				RDG_EVENT_NAME("AmbientOcclusion_HorizonSearchIntegralSpatialFilter %dx%d (CS)", ViewRect.Width(), ViewRect.Height()),
+				ERDGPassFlags::Compute | ERDGPassFlags::NeverCull,
 				ComputeShader,
 				HorizonSearchIntegralSpatialFilterParameters,
 				FComputeShaderUtils::GetGroupCount(ViewRect.Size(), FGTAOMobile_HorizonSearchIntegralSpatialFilterCS::TexelsPerThreadGroup));
@@ -524,7 +525,7 @@ static void RenderGTAO(FRDGBuilder& GraphBuilder, FRDGTextureRef SceneDepthTextu
 			GraphBuilder.AddPass(
 				RDG_EVENT_NAME("AmbientOcclusion_HorizonSearchIntegral %dx%d (PS)", ViewRect.Width(), ViewRect.Height()),
 				HorizonSearchIntegralParameters,
-				ERDGPassFlags::Raster,
+				ERDGPassFlags::Raster | ERDGPassFlags::NeverCull,
 				[VertexShader, HorizonSearchIntegralShader, HorizonSearchIntegralParameters, ViewRect, BufferSize](FRHICommandList& RHICmdList)
 			{
 				RHICmdList.SetViewport(ViewRect.Min.X, ViewRect.Min.Y, 0.0f, ViewRect.Max.X, ViewRect.Max.Y, 1.0f);
@@ -570,7 +571,7 @@ static void RenderGTAO(FRDGBuilder& GraphBuilder, FRDGTextureRef SceneDepthTextu
 			GraphBuilder.AddPass(
 				RDG_EVENT_NAME("AmbientOcclusion_SpatialFilter %dx%d (PS)", ViewRect.Width(), ViewRect.Height()),
 				SpatialFilterParameters,
-				ERDGPassFlags::Raster,
+				ERDGPassFlags::Raster | ERDGPassFlags::NeverCull,
 				[VertexShader, SpatialFilterShader, SpatialFilterParameters, ViewRect, BufferSize](FRHICommandList& RHICmdList)
 			{
 				RHICmdList.SetViewport(ViewRect.Min.X, ViewRect.Min.Y, 0.0f, ViewRect.Max.X, ViewRect.Max.Y, 1.0f);
@@ -616,6 +617,7 @@ static void RenderGTAO(FRDGBuilder& GraphBuilder, FRDGTextureRef SceneDepthTextu
 			FComputeShaderUtils::AddPass(
 				GraphBuilder,
 				RDG_EVENT_NAME("AmbientOcclusion_HorizonSearchIntegral %dx%d (CS)", ViewRect.Width(), ViewRect.Height()),
+				ERDGPassFlags::Compute | ERDGPassFlags::NeverCull,
 				HorizonSearchIntegralShader,
 				HorizonSearchIntegralParameters,
 				FComputeShaderUtils::GetGroupCount(ViewRect.Size(), FGTAOMobile_HorizonSearchIntegralCS::TexelsPerThreadGroup));
@@ -629,10 +631,16 @@ static void RenderGTAO(FRDGBuilder& GraphBuilder, FRDGTextureRef SceneDepthTextu
 			FComputeShaderUtils::AddPass(
 				GraphBuilder,
 				RDG_EVENT_NAME("AmbientOcclusion_SpatialFilter %dx%d (CS)", ViewRect.Width(), ViewRect.Height()),
+				ERDGPassFlags::Compute | ERDGPassFlags::NeverCull,
 				SpatialFilterShader,
 				SpatialFilterParameters,
 				FComputeShaderUtils::GetGroupCount(ViewRect.Size(), FGTAOMobile_SpatialFilterCS::TexelsPerThreadGroup));
 		}
+
+		if (View.ViewState && !View.bStatePrevViewInfoIsReadOnly)
+		{
+			GraphBuilder.QueueTextureExtraction(AmbientOcclusionTexture, &View.ViewState->PrevFrameViewInfo.MobileAmbientOcclusion);
+		}
 	}
 }
 
@@ -960,8 +968,6 @@ static void RenderSSAO(FRDGBuilder& GraphBuilder, FRDGTextureRef SceneDepthTextu
 // --------------------------------------------------------------------------------------------------------------------
 void FMobileSceneRenderer::RenderAmbientOcclusion(FRDGBuilder& GraphBuilder, FRDGTextureRef SceneDepthTexture, FRDGTextureRef AmbientOcclusionTexture)
 {
-	FMemMark Mark(FMemStack::Get());
-
 	const int32 Technique = CVarMobileAmbientOcclusionTechnique.GetValueOnRenderThread();
 	switch (Technique)
 	{
@@ -972,6 +978,4 @@ void FMobileSceneRenderer::RenderAmbientOcclusion(FRDGBuilder& GraphBuilder, FRD
 		RenderSSAO(GraphBuilder, SceneDepthTexture, AmbientOcclusionTexture, Views);
 		break;
 	}
-
-	GraphBuilder.Execute();
 }
diff --git a/Engine/Source/Runtime/Renderer/Private/PostProcess/PostProcessPixelProjectedReflectionMobile.cpp b/Engine/Source/Runtime/Renderer/Private/PostProcess/PostProcessPixelProjectedReflectionMobile.cpp
--- a/Engine/Source/Runtime/Renderer/Private/PostProcess/PostProcessPixelProjectedReflectionMobile.cpp
+++ b/Engine/Source/Runtime/Renderer/Private/PostProcess/PostProcessPixelProjectedReflectionMobile.cpp
@@ -11,6 +11,7 @@
 #include "MobileBasePassRendering.h"
 #include "ScreenPass.h"
 #include "RenderTargetPool.h"
+#include "ScenePrivate.h"
 
 
 static TAutoConsoleVariable<int32> CVarMobilePlanarReflectionMode(
@@ -29,7 +30,7 @@ static TAutoConsoleVariable<int32> CVarMobilePixelProjectedReflectionQuality(
 	TEXT("The quality of pixel projected reflection on mobile platform.\n")
 	TEXT("0: Disabled\n")
 	TEXT("1: Best performance but may have some artifacts in some view angles. [default]\n")
-	TEXT("2: Better quality and reasonable performance and could fix some artifacts, but the PlanarReflection mesh has to render twice.\n")
+	TEXT("2: Better quality and reasonable performance and could fix some artifacts.\n")
 	TEXT("3: Best quality but will be much heavier.\n"),
 	ECVF_Scalability | ECVF_RenderThreadSafe
 );
@@ -161,8 +162,6 @@ public:
 
 IMPLEMENT_GLOBAL_SHADER(FPixelProjectedReflectionMobile_ReflectionPassPS, "/Engine/Private/PostProcessPixelProjectedReflectionMobile.usf", "ReflectionPassPS", SF_Pixel);
 
-FPixelProjectedReflectionMobileOutputs GPixelProjectedReflectionMobileOutputs;
-
 class FReflectionPlaneVertexDeclaration : public FRenderResource
 {
 public:
@@ -186,19 +185,11 @@ public:
 
 TGlobalResource<FReflectionPlaneVertexDeclaration> GReflectionPlaneVertexDeclaration;
 
-void FMobileSceneRenderer::InitPixelProjectedReflectionOutputs(FRHICommandListImmediate& RHICmdList, const FIntPoint& BufferSize)
+FRDGTextureRef CreateMobilePixelProjectedReflectionTexture(FRDGBuilder& GraphBuilder, FIntPoint Extent)
 {
-	if (!GPixelProjectedReflectionMobileOutputs.IsValid() || GPixelProjectedReflectionMobileOutputs.PixelProjectedReflectionTexture->GetDesc().Extent != BufferSize)
-	{
-		GPixelProjectedReflectionMobileOutputs.PixelProjectedReflectionTexture.SafeRelease();
-
-		GRenderTargetPool.FindFreeElement(RHICmdList, FPooledRenderTargetDesc::Create2DDesc(BufferSize, PF_FloatRGBA, FClearValueBinding::Transparent, TexCreate_None, TexCreate_RenderTargetable | TexCreate_ShaderResource | TexCreate_UAV, false, 1, false), GPixelProjectedReflectionMobileOutputs.PixelProjectedReflectionTexture, TEXT("PixelProjectedReflectionTexture"));
-	}
-}
-
-void FMobileSceneRenderer::ReleasePixelProjectedReflectionOutputs()
-{
-	GPixelProjectedReflectionMobileOutputs.Release();
+	return GraphBuilder.CreateTexture(
+		FRDGTextureDesc::Create2D(Extent, PF_FloatRGBA, FClearValueBinding::Transparent, TexCreate_ShaderResource | TexCreate_RenderTargetable | TexCreate_UAV),
+		TEXT("PixelProjectedReflectionTexture"));
 }
 
 void FMobileSceneRenderer::RenderPixelProjectedReflection(FRDGBuilder& GraphBuilder, FRDGTextureRef SceneColorTexture, FRDGTextureRef SceneDepthTexture, FRDGTextureRef PixelProjectedReflectionTexture, const FPlanarReflectionSceneProxy* PlanarReflectionSceneProxy)
@@ -247,6 +238,7 @@ void FMobileSceneRenderer::RenderPixelProjectedReflection(FRDGBuilder& GraphBuil
 		FComputeShaderUtils::AddPass(
 			GraphBuilder,
 			RDG_EVENT_NAME("PixelProjectedReflection_Projection %dx%d (CS)", ViewRect.Width(), ViewRect.Height()),
+			ERDGPassFlags::Compute | ERDGPassFlags::NeverCull,
 			ComputeShader,
 			PassParameters,
 			FComputeShaderUtils::GetGroupCount(ViewRect.Size(), FPixelProjectedReflectionMobile_ProjectionPassCS::TexelsPerThreadGroup));
@@ -301,7 +293,7 @@ void FMobileSceneRenderer::RenderPixelProjectedReflection(FRDGBuilder& GraphBuil
 		GraphBuilder.AddPass(
 			RDG_EVENT_NAME("PixelProjectedReflection_Reflection %dx%d (PS)", ViewRect.Width(), ViewRect.Height()),
 			PSShaderParameters,
-			ERDGPassFlags::Raster,
+			ERDGPassFlags::Raster | ERDGPassFlags::NeverCull,
 			[VertexShader, VSShaderParameters, PixelShader, PSShaderParameters, ViewRect](FRHICommandList& RHICmdList)
 		{
 			RHICmdList.SetViewport(ViewRect.Min.X, ViewRect.Min.Y, 0.0f, ViewRect.Max.X, ViewRect.Max.Y, 1.0f);
@@ -326,5 +318,10 @@ void FMobileSceneRenderer::RenderPixelProjectedReflection(FRDGBuilder& GraphBuil
 			RHICmdList.SetStreamSource(0, GScreenSpaceVertexBuffer.VertexBufferRHI, 0);
 			RHICmdList.DrawPrimitive(0, 2, 1);
 		});
+
+		if (View.ViewState && !View.bStatePrevViewInfoIsReadOnly)
+		{
+			GraphBuilder.QueueTextureExtraction(PixelProjectedReflectionTexture, &View.ViewState->PrevFrameViewInfo.MobilePixelProjectedReflection);
+		}
 	}
 }
\ No newline at end of file
diff --git a/Engine/Source/Runtime/Renderer/Private/PostProcess/PostProcessPixelProjectedReflectionMobile.h b/Engine/Source/Runtime/Renderer/Private/PostProcess/PostProcessPixelProjectedReflectionMobile.h
--- a/Engine/Source/Runtime/Renderer/Private/PostProcess/PostProcessPixelProjectedReflectionMobile.h
+++ b/Engine/Source/Runtime/Renderer/Private/PostProcess/PostProcessPixelProjectedReflectionMobile.h
@@ -18,19 +18,4 @@
 #define PROJECTION_OUTPUT_TYPE PROJECTION_OUTPUT_TYPE_TEXTURE
 #endif
 
-struct FPixelProjectedReflectionMobileOutputs
-{
-	TRefCountPtr<IPooledRenderTarget> PixelProjectedReflectionTexture;
-
-	bool IsValid()
-	{
-		return PixelProjectedReflectionTexture.IsValid();
-	}
-
-	void Release()
-	{
-		PixelProjectedReflectionTexture.SafeRelease();
-	}
-};
-
-extern FPixelProjectedReflectionMobileOutputs GPixelProjectedReflectionMobileOutputs;
\ No newline at end of file
+FRDGTextureRef CreateMobilePixelProjectedReflectionTexture(FRDGBuilder& GraphBuilder, FIntPoint Extent);
\ No newline at end of file
diff --git a/Engine/Source/Runtime/Renderer/Private/SceneOcclusion.cpp b/Engine/Source/Runtime/Renderer/Private/SceneOcclusion.cpp
--- a/Engine/Source/Runtime/Renderer/Private/SceneOcclusion.cpp
+++ b/Engine/Source/Runtime/Renderer/Private/SceneOcclusion.cpp
@@ -1515,52 +1515,48 @@ void FDeferredShadingSceneRenderer::RenderOcclusion(
 	}
 }
 
-void FMobileSceneRenderer::RenderOcclusion(FRDGBuilder& GraphBuilder, FRenderTargetBindingSlots& BasePassRenderTargets)
+void FMobileSceneRenderer::RenderOcclusion(FRHICommandListImmediate& RHICmdList, const FViewInfo& View)
 {
-	if (!DoOcclusionQueries(FeatureLevel))
-	{
-		return;
-	}
-
 	{
 		SCOPED_NAMED_EVENT(FMobileSceneRenderer_BeginOcclusionTests, FColor::Emerald);
-		FViewOcclusionQueriesPerView QueriesPerView = AllocateOcclusionTests(Scene, VisibleLightInfos, Views);
+		const FViewOcclusionQueriesPerView QueriesPerView = AllocateOcclusionTests(Scene, VisibleLightInfos, Views);
 
 		if (QueriesPerView.Num())
 		{
-			auto* PassParameters = GraphBuilder.AllocParameters<FRenderTargetParameters>();
-			PassParameters->RenderTargets = BasePassRenderTargets;
-
-			GraphBuilder.AddPass(RDG_EVENT_NAME("BeginOcclusionTestsPass"), PassParameters, ERDGPassFlags::Raster | ERDGPassFlags::SkipRenderPass,
-				[this, LocalQueriesPerView = MoveTemp(QueriesPerView)](FRHICommandListImmediate& RHICmdList)
-			{
-				BeginOcclusionTests(RHICmdList, Views, FeatureLevel, LocalQueriesPerView, 1.0f);
-			});
+			BeginOcclusionTests(RHICmdList, Views, FeatureLevel, QueriesPerView, 1.0f);
 		}
 	}
 
-	FenceOcclusionTests(GraphBuilder);
+	if (IsRunningRHIInSeparateThread())
+	{
+		FenceOcclusionTestsInternal(RHICmdList);
+	}
 }
 
 DECLARE_CYCLE_STAT(TEXT("OcclusionSubmittedFence Dispatch"), STAT_OcclusionSubmittedFence_Dispatch, STATGROUP_SceneRendering);
 DECLARE_CYCLE_STAT(TEXT("OcclusionSubmittedFence Wait"), STAT_OcclusionSubmittedFence_Wait, STATGROUP_SceneRendering);
 
+void FSceneRenderer::FenceOcclusionTestsInternal(FRHICommandListImmediate& RHICmdList)
+{
+	SCOPE_CYCLE_COUNTER(STAT_OcclusionSubmittedFence_Dispatch);
+	int32 NumFrames = FOcclusionQueryHelpers::GetNumBufferedFrames(FeatureLevel);
+	for (int32 Dest = NumFrames - 1; Dest >= 1; Dest--)
+	{
+		CA_SUPPRESS(6385);
+		OcclusionSubmittedFence[Dest] = OcclusionSubmittedFence[Dest - 1];
+	}
+	OcclusionSubmittedFence[0] = RHICmdList.RHIThreadFence();
+	RHICmdList.ImmediateFlush(EImmediateFlushType::DispatchToRHIThread);
+	RHICmdList.PollRenderQueryResults();
+}
+
 void FSceneRenderer::FenceOcclusionTests(FRDGBuilder& GraphBuilder)
 {
 	if (IsRunningRHIInSeparateThread())
 	{
 		AddPass(GraphBuilder, [this](FRHICommandListImmediate& RHICmdList)
 		{
-			SCOPE_CYCLE_COUNTER(STAT_OcclusionSubmittedFence_Dispatch);
-			int32 NumFrames = FOcclusionQueryHelpers::GetNumBufferedFrames(FeatureLevel);
-			for (int32 Dest = NumFrames - 1; Dest >= 1; Dest--)
-			{
-				CA_SUPPRESS(6385);
-				OcclusionSubmittedFence[Dest] = OcclusionSubmittedFence[Dest - 1];
-			}
-			OcclusionSubmittedFence[0] = RHICmdList.RHIThreadFence();
-			RHICmdList.ImmediateFlush(EImmediateFlushType::DispatchToRHIThread);
-			RHICmdList.PollRenderQueryResults();
+			FenceOcclusionTestsInternal(RHICmdList);
 		});
 	}
 }
diff --git a/Engine/Source/Runtime/Renderer/Private/SceneOcclusion.h b/Engine/Source/Runtime/Renderer/Private/SceneOcclusion.h
--- a/Engine/Source/Runtime/Renderer/Private/SceneOcclusion.h
+++ b/Engine/Source/Runtime/Renderer/Private/SceneOcclusion.h
@@ -83,8 +83,6 @@ public:
 
 	void SetParameters(FRHICommandList& RHICmdList, const FViewInfo& View)
 	{
-		FGlobalShader::SetParameters<FViewUniformShaderParameters>(RHICmdList, RHICmdList.GetBoundVertexShader(),View.ViewUniformBuffer);
-
 		// Don't transform if rendering frustum
 		StencilingGeometryParameters.Set(RHICmdList, this, FVector4(0,0,0,1));
 
diff --git a/Engine/Source/Runtime/Renderer/Private/SceneRendering.cpp b/Engine/Source/Runtime/Renderer/Private/SceneRendering.cpp
--- a/Engine/Source/Runtime/Renderer/Private/SceneRendering.cpp
+++ b/Engine/Source/Runtime/Renderer/Private/SceneRendering.cpp
@@ -3755,9 +3755,6 @@ FRendererModule::FRendererModule()
 	static auto CVarVertexDeformationOutputsVelocity = IConsoleManager::Get().FindConsoleVariable(TEXT("r.VertexDeformationOutputsVelocity"));
 	CVarVertexDeformationOutputsVelocity->SetOnChangedCallback(FConsoleVariableDelegate::CreateStatic(&OnChangeCVarRequiringRecreateRenderState));
 
-	static auto MobilePixelProjectedReflectionQualityCVar = IConsoleManager::Get().FindConsoleVariable(TEXT("r.Mobile.PixelProjectedReflectionQuality"));
-	MobilePixelProjectedReflectionQualityCVar->SetOnChangedCallback(FConsoleVariableDelegate::CreateStatic(&OnChangeCVarRequiringRecreateRenderState));
-
 #if !(UE_BUILD_SHIPPING || UE_BUILD_TEST)
 	void InitDebugViewModeInterfaces();
 	InitDebugViewModeInterfaces();
diff --git a/Engine/Source/Runtime/Renderer/Private/SceneRendering.h b/Engine/Source/Runtime/Renderer/Private/SceneRendering.h
--- a/Engine/Source/Runtime/Renderer/Private/SceneRendering.h
+++ b/Engine/Source/Runtime/Renderer/Private/SceneRendering.h
@@ -1035,6 +1035,12 @@ struct FPreviousViewInfo
 
 	// Mobile bloom setup eye adaptation surface.
 	TRefCountPtr<IPooledRenderTarget> MobileBloomSetup_EyeAdaptation;
+
+	// Mobile pixel projected reflection textureused for next frame.
+	TRefCountPtr<IPooledRenderTarget> MobilePixelProjectedReflection = nullptr;
+
+	// Mobile ambient occlusion texture used for next frame.
+	TRefCountPtr<IPooledRenderTarget> MobileAmbientOcclusion = nullptr;
 };
 
 class FViewCommands
@@ -1859,6 +1865,7 @@ public:
 	virtual void Render(FRDGBuilder& GraphBuilder) = 0;
 	virtual void RenderHitProxies(FRDGBuilder& GraphBuilder) {}
 	virtual bool ShouldRenderVelocities() const { return false; }
+	virtual bool ShouldRenderPrePass() const { return false; }
 
 	/** Creates a scene renderer based on the current feature level. */
 	static FSceneRenderer* CreateSceneRenderer(const FSceneViewFamily* InViewFamily, FHitProxyConsumer* HitProxyConsumer);
@@ -1877,7 +1884,8 @@ public:
 	bool DoOcclusionQueries(ERHIFeatureLevel::Type InFeatureLevel) const;
 
 	void FenceOcclusionTests(FRDGBuilder& GraphBuilder);
-	void WaitOcclusionTests(FRHICommandListImmediate& GraphBuilder);
+	void FenceOcclusionTestsInternal(FRHICommandListImmediate& RHICmdList);
+	void WaitOcclusionTests(FRHICommandListImmediate& RHICmdList);
 
 	// fences to make sure the rhi thread has digested the occlusion query renders before we attempt to read them back async
 	static FGraphEventRef OcclusionSubmittedFence[FOcclusionQueryHelpers::MaxBufferedOcclusionFrames];
@@ -2211,6 +2219,8 @@ public:
 
 	virtual bool ShouldRenderVelocities() const override;
 
+	virtual bool ShouldRenderPrePass() const override;
+
 	void RenderInverseOpacity(FRDGBuilder& GraphBuilder, const FViewInfo& View);
 
 protected:
@@ -2224,24 +2234,24 @@ protected:
 
 	void InitViews(FRDGBuilder& GraphBuilder, FSceneTexturesConfig& SceneTexturesConfig, FInstanceCullingManager& InstanceCullingManager);
 
-	void RenderPrePass(FRDGBuilder& GraphBuilder, FRenderTargetBindingSlots& BasePassRenderTargets);
+	void RenderPrePass(FRHICommandListImmediate& RHICmdList, const FViewInfo& View, const FInstanceCullingDrawParams* InstanceCullingDrawParams = nullptr);
 
 	/** Renders the opaque base pass for mobile. */
-	void RenderMobileBasePass(FRDGBuilder& GraphBuilder, FRenderTargetBindingSlots& BasePassRenderTargets, const TArrayView<const FViewInfo> PassViews, FRDGTextureRef ScreenSpaceAO);
+	void RenderMobileBasePass(FRHICommandListImmediate& RHICmdList, const FViewInfo& View);
 
 	void RenderMobileEditorPrimitives(FRHICommandList& RHICmdList, const FViewInfo& View, const FMeshPassProcessorRenderState& DrawRenderState);
 
 	/** Renders the debug view pass for mobile. */
-	void RenderMobileDebugView(FRDGBuilder& GraphBuilder, const TArrayView<const FViewInfo> PassViews, FRDGTextureRef QuadOverdrawTexture);
+	void RenderMobileDebugView(FRHICommandListImmediate& RHICmdList, const FViewInfo& View);
 
 	/** Render modulated shadow projections in to the scene, loops over any unrendered shadows until all are processed.*/
-	void RenderModulatedShadowProjections(FRDGBuilder& GraphBuilder, TRDGUniformBufferRef<FMobileSceneTextureUniformParameters> MobileSceneTextures);
+	void RenderModulatedShadowProjections(FRHICommandListImmediate& RHICmdList, int32 ViewIndex, const FViewInfo& View);
 
 	/** Resolves scene depth in case hardware does not support reading depth in the shader */
 	void ConditionalResolveSceneDepth(FRDGBuilder& GraphBuilder, const FViewInfo& View, FRDGTextureMSAA SceneDepth);
 
 	/** Issues occlusion queries */
-	void RenderOcclusion(FRDGBuilder& GraphBuilder, FRenderTargetBindingSlots& BasePassRenderTargets);
+	void RenderOcclusion(FRHICommandListImmediate& RHICmdList, const FViewInfo& View);
 	
 	bool ShouldRenderHZB();
 
@@ -2257,13 +2267,13 @@ protected:
 	bool RequiresMultiPass(FRHICommandListImmediate& RHICmdList, const FViewInfo& View) const;
 
 	/** Renders decals. */
-	void RenderDecals(FRDGBuilder& GraphBuilder, FRenderTargetBindingSlots& BasePassRenderTargets, TRDGUniformBufferRef<FMobileSceneTextureUniformParameters> MobileSceneTextures);
+	void RenderDecals(FRHICommandListImmediate& RHICmdList, const FViewInfo& View);
 
 	/** Renders the base pass for translucency. */
-	void RenderTranslucency(FRDGBuilder& GraphBuilder, FRenderTargetBindingSlots& BasePassRenderTargets, const TArrayView<const FViewInfo> PassViews, FRDGTextureRef ScreenSpaceAO);
+	void RenderTranslucency(FRHICommandListImmediate& RHICmdList, const FViewInfo& View, const EMeshPass::Type TranslucencyMeshPass);
 
 	/** On chip pre-tonemap before scene color MSAA resolve (iOS only) */
-	void PreTonemapMSAA(FRDGBuilder& GraphBuilder, FRenderTargetBindingSlots& BasePassRenderTargets, const FMinimalSceneTextures& SceneTextures);
+	void PreTonemapMSAA(FRHICommandListImmediate& RHICmdList, const FMinimalSceneTextures& SceneTextures);
 
 	void SortMobileBasePassAfterShadowInit(FExclusiveDepthStencil::Type BasePassDepthStencilAccess, FViewVisibleCommandsPerView& ViewCommandsPerView);
 	void SetupMobileBasePassAfterShadowInit(FExclusiveDepthStencil::Type BasePassDepthStencilAccess, FViewVisibleCommandsPerView& ViewCommandsPerView);
@@ -2271,14 +2281,16 @@ protected:
 	void UpdateDirectionalLightUniformBuffers(FRDGBuilder& GraphBuilder, const FViewInfo& View);
 	void UpdateSkyReflectionUniformBuffer();
 
-	void RenderForward(FRDGBuilder& GraphBuilder, const TArrayView<const FViewInfo> ViewList, FRDGTextureRef ViewFamilyTexture, FSceneTextures& SceneTextures);
-	void RenderDeferred(FRDGBuilder& GraphBuilder, const TArrayView<const FViewInfo> ViewList, const FSortedLightSetSceneInfo& SortedLightSet, FRDGTextureRef ViewFamilyTexture, FSceneTextures& SceneTextures);
-	
+	void RenderForward(FRDGBuilder& GraphBuilder, FRDGTextureRef ViewFamilyTexture, FSceneTextures& SceneTextures);
+	void RenderForwardSinglePass(FRDGBuilder& GraphBuilder, class FMobileRenderPassParameters* PassParameters, int32 ViewIndex, FViewInfo& View, FSceneTextures& SceneTextures, bool bShouldRenderPrePass, bool bShouldRenderMobileDebugView, bool bAdrenoOcclusionMode, bool bShouldRenderOcclusionPass, bool bShouldRenderTranslucency, EMeshPass::Type TranslucencyMeshPass);
+	void RenderForwardMultiPass(FRDGBuilder& GraphBuilder, class FMobileRenderPassParameters* PassParameters, FRenderTargetBindingSlots& BasePassRenderTargets, int32 ViewIndex, FViewInfo& View, FSceneTextures& SceneTextures, bool bShouldRenderPrePass, bool bShouldRenderMobileDebugView, bool bAdrenoOcclusionMode, bool bShouldRenderOcclusionPass, bool bShouldRenderTranslucency, EMeshPass::Type TranslucencyMeshPass);
+	void RenderDeferred(FRDGBuilder& GraphBuilder, const FSortedLightSetSceneInfo& SortedLightSet, FRDGTextureRef ViewFamilyTexture, FSceneTextures& SceneTextures);
+	void RenderDeferredSinglePass(FRDGBuilder& GraphBuilder, class FMobileRenderPassParameters* PassParameters, int32 ViewIndex, int32 NumViews, FViewInfo& View, FSceneTextures& SceneTextures, const FSortedLightSetSceneInfo& SortedLightSet, FTextureRHIRef MobilePixelProjectedReflection, bool bShouldRenderPrePass, bool bShouldRenderOcclusionPass, bool bUsingPixelLocalStorage, bool bShouldRenderTranslucency, EMeshPass::Type TranslucencyMeshPass);
+	void RenderDeferredMultiPass(FRDGBuilder& GraphBuilder, class FMobileRenderPassParameters* PassParameters, FRenderTargetBindingSlots& BasePassRenderTargets, int32 NumColorTargets, int32 ViewIndex, int32 NumViews, FViewInfo& View, FSceneTextures& SceneTextures, const FSortedLightSetSceneInfo& SortedLightSet, FTextureRHIRef MobilePixelProjectedReflection, bool bShouldRenderPrePass, bool bShouldRenderOcclusionPass, bool bShouldRenderTranslucency, EMeshPass::Type TranslucencyMeshPass);
+
 	void RenderAmbientOcclusion(FRDGBuilder& GraphBuilder, FRDGTextureRef SceneDepthTexture, FRDGTextureRef AmbientOcclusionTexture);
 
-	void InitPixelProjectedReflectionOutputs(FRHICommandListImmediate& RHICmdList, const FIntPoint& BufferSize);
 	void RenderPixelProjectedReflection(FRDGBuilder& GraphBuilder, FRDGTextureRef SceneColorTexture, FRDGTextureRef SceneDepthTexture, FRDGTextureRef PixelProjectedReflectionTexture, const FPlanarReflectionSceneProxy* PlanarReflectionSceneProxy);
-	void ReleasePixelProjectedReflectionOutputs();
 
 	/** Before SetupMobileBasePassAfterShadowInit, we need to update the uniform buffer and shadow info for all movable point lights.*/
 	void UpdateMovablePointLightUniformBufferAndShadowInfo();
diff --git a/Engine/Source/Runtime/Renderer/Private/SceneTextures.cpp b/Engine/Source/Runtime/Renderer/Private/SceneTextures.cpp
--- a/Engine/Source/Runtime/Renderer/Private/SceneTextures.cpp
+++ b/Engine/Source/Runtime/Renderer/Private/SceneTextures.cpp
@@ -16,6 +16,8 @@
 #include "PostProcessCompositeEditorPrimitives.h"
 #include "ShaderCompiler.h"
 #include "SystemTextures.h"
+#include "PostProcess/PostProcessAmbientOcclusionMobile.h"
+#include "PostProcess/PostProcessPixelProjectedReflectionMobile.h"
 
 static TAutoConsoleVariable<int32> CVarSceneTargetsResizeMethod(
 	TEXT("r.SceneRenderTargetResizeMethod"),
@@ -702,6 +704,11 @@ FSceneTextures& FSceneTextures::Create(FRDGBuilder& GraphBuilder, const FSceneTe
 	{
 		// Mobile Screen Space Ambient Occlusion
 		SceneTextures.ScreenSpaceAO = CreateMobileScreenSpaceAOTexture(GraphBuilder, Config.Extent);
+
+		if (Config.MobilePixelProjectedReflectionExtent != FIntPoint::ZeroValue)
+		{
+			SceneTextures.PixelProjectedReflection = CreateMobilePixelProjectedReflectionTexture(GraphBuilder, Config.MobilePixelProjectedReflectionExtent);
+		}
 	}
 
 	// Velocity
diff --git a/Engine/Source/Runtime/Renderer/Private/SceneTextures.h b/Engine/Source/Runtime/Renderer/Private/SceneTextures.h
--- a/Engine/Source/Runtime/Renderer/Private/SceneTextures.h
+++ b/Engine/Source/Runtime/Renderer/Private/SceneTextures.h
@@ -58,6 +58,9 @@ struct RENDERER_API FSceneTexturesConfig
 	// Extent of all full-resolution textures.
 	FIntPoint Extent = FIntPoint::ZeroValue;
 
+	// Extend of the mobile Pixel Projected Reflection texture
+	FIntPoint MobilePixelProjectedReflectionExtent = FIntPoint::ZeroValue;
+
 	// Downsample factors to divide against the full resolution texture extent.
 	uint32 SmallDepthDownsampleFactor = 2;
 	uint32 CustomDepthDownsampleFactor = 1;
@@ -161,7 +164,7 @@ struct RENDERER_API FSceneTextures : public FMinimalSceneTextures
 	// (Mobile Deferred) Additional GBuffer texture used by mobile
 	FRDGTextureRef DepthAux{};
 
-	// (Deferred) Texture containing dynamic motion vectors. Can be bound by the base pass or its own velocity pass.
+	// Texture containing dynamic motion vectors. Can be bound by the base pass or its own velocity pass.
 	FRDGTextureRef Velocity{};
 
 	// Texture containing the screen space ambient occlusion result.
@@ -173,6 +176,9 @@ struct RENDERER_API FSceneTextures : public FMinimalSceneTextures
 	// Texture used to control shading rate for rendering.
 	FRDGTextureRef ShadingRate{};
 
+	// (Mobile) Texture used by mobile PPR in the next frame.
+	FRDGTextureRef PixelProjectedReflection{};
+
 	// Textures used to composite editor primitives. Also used by the base pass when in wireframe mode.
 #if WITH_EDITOR
 	FRDGTextureRef EditorPrimitiveColor{};
diff --git a/Engine/Source/Runtime/Renderer/Private/ShadowRendering.cpp b/Engine/Source/Runtime/Renderer/Private/ShadowRendering.cpp
--- a/Engine/Source/Runtime/Renderer/Private/ShadowRendering.cpp
+++ b/Engine/Source/Runtime/Renderer/Private/ShadowRendering.cpp
@@ -1110,152 +1110,205 @@ void FProjectedShadowInfo::RenderProjection(
 		ProjectionStencilingPass.BuildRenderingCommands(GraphBuilder, *View, *SceneRender->Scene, PassParameters->InstanceCullingDrawParams);
 	}
 
+	const FInstanceCullingDrawParams& InstanceCullingDrawParams = PassParameters->InstanceCullingDrawParams;
+
 	GraphBuilder.AddPass(
 		RDG_EVENT_NAME("%s", *EventName),
 		PassParameters,
 		ERDGPassFlags::Raster | PassFlags,
-		[this, SceneRender, View, ViewIndex, LightSceneProxy, HairVisibilityData, bProjectingForForwardShading, bMobileModulatedProjections, PassParameters](FRHICommandListImmediate& RHICmdList)
+		[this, SceneRender, View, ViewIndex, LightSceneProxy, HairVisibilityData, bProjectingForForwardShading, bMobileModulatedProjections, &InstanceCullingDrawParams](FRHICommandListImmediate& RHICmdList)
 	{
-		RHICmdList.SetViewport(View->ViewRect.Min.X, View->ViewRect.Min.Y, 0.0f, View->ViewRect.Max.X, View->ViewRect.Max.Y, 1.0f);
-		LightSceneProxy->SetScissorRect(RHICmdList, *View, View->ViewRect);
+		RenderProjectionInternal(RHICmdList, ViewIndex, View, LightSceneProxy, SceneRender, bProjectingForForwardShading, bMobileModulatedProjections, HairVisibilityData, InstanceCullingDrawParams);
+	});
+}
 
-		FScopeCycleCounter Scope(bWholeSceneShadow ? GET_STATID(STAT_RenderWholeSceneShadowProjectionsTime) : GET_STATID(STAT_RenderPerObjectShadowProjectionsTime));
+void FProjectedShadowInfo::RenderProjectionInternal(
+	FRHICommandListImmediate& RHICmdList,
+	int32 ViewIndex,
+	const FViewInfo* View,
+	const FLightSceneProxy* LightSceneProxy,
+	const FSceneRenderer* SceneRender,
+	bool bProjectingForForwardShading,
+	bool bMobileModulatedProjections,
+	const FHairStrandsVisibilityData* HairVisibilityData,
+	const FInstanceCullingDrawParams& InstanceCullingDrawParams) const
+{
+	RHICmdList.SetViewport(View->ViewRect.Min.X, View->ViewRect.Min.Y, 0.0f, View->ViewRect.Max.X, View->ViewRect.Max.Y, 1.0f);
+	LightSceneProxy->SetScissorRect(RHICmdList, *View, View->ViewRect);
 
-		FGraphicsPipelineStateInitializer GraphicsPSOInit;
-		RHICmdList.ApplyCachedRenderTargets(GraphicsPSOInit);
+	FScopeCycleCounter Scope(bWholeSceneShadow ? GET_STATID(STAT_RenderWholeSceneShadowProjectionsTime) : GET_STATID(STAT_RenderPerObjectShadowProjectionsTime));
 
-		FPlane OutFrustmPlanes[6];
-		bool bCameraInsideShadowFrustum;
-		TArray<FVector4, TInlineAllocator<8>> FrustumVertices;
-		SetupFrustumForProjection(View, FrustumVertices, bCameraInsideShadowFrustum, OutFrustmPlanes);
+	FGraphicsPipelineStateInitializer GraphicsPSOInit;
+	RHICmdList.ApplyCachedRenderTargets(GraphicsPSOInit);
 
-		const bool bSubPixelSupport = HairVisibilityData != nullptr;
-		const bool bStencilTestEnabled = !bSubPixelSupport && GShadowStencilCulling;
-		const bool bDepthBoundsTestEnabled = IsWholeSceneDirectionalShadow() && GSupportsDepthBoundsTest && CVarCSMDepthBoundsTest.GetValueOnRenderThread() != 0 && !bSubPixelSupport;
+	FPlane OutFrustmPlanes[6];
+	bool bCameraInsideShadowFrustum;
+	TArray<FVector4, TInlineAllocator<8>> FrustumVertices;
+	SetupFrustumForProjection(View, FrustumVertices, bCameraInsideShadowFrustum, OutFrustmPlanes);
 
-		if (!bDepthBoundsTestEnabled && bStencilTestEnabled)
-		{
-			SetupProjectionStencilMask(RHICmdList, View, ViewIndex, SceneRender, FrustumVertices, bMobileModulatedProjections, bCameraInsideShadowFrustum, PassParameters->InstanceCullingDrawParams);
-		}
+	const bool bSubPixelSupport = HairVisibilityData != nullptr;
+	const bool bStencilTestEnabled = !bSubPixelSupport && GShadowStencilCulling;
+	const bool bDepthBoundsTestEnabled = IsWholeSceneDirectionalShadow() && GSupportsDepthBoundsTest && CVarCSMDepthBoundsTest.GetValueOnRenderThread() != 0 && !bSubPixelSupport;
 
-		// solid rasterization w/ back-face culling.
-		GraphicsPSOInit.RasterizerState = (View->bReverseCulling || IsWholeSceneDirectionalShadow()) ? TStaticRasterizerState<FM_Solid, CM_CCW>::GetRHI() : TStaticRasterizerState<FM_Solid, CM_CW>::GetRHI();
+	if (!bDepthBoundsTestEnabled && bStencilTestEnabled)
+	{
+		SetupProjectionStencilMask(RHICmdList, View, ViewIndex, SceneRender, FrustumVertices, bMobileModulatedProjections, bCameraInsideShadowFrustum, InstanceCullingDrawParams);
+	}
 
-		GraphicsPSOInit.bDepthBounds = bDepthBoundsTestEnabled;
-		if (bDepthBoundsTestEnabled)
-		{
-			// no depth test or writes
-			GraphicsPSOInit.DepthStencilState = TStaticDepthStencilState<false, CF_Always>::GetRHI();
-		}
-		else if (bStencilTestEnabled)
+	// solid rasterization w/ back-face culling.
+	GraphicsPSOInit.RasterizerState = (View->bReverseCulling || IsWholeSceneDirectionalShadow()) ? TStaticRasterizerState<FM_Solid, CM_CCW>::GetRHI() : TStaticRasterizerState<FM_Solid, CM_CW>::GetRHI();
+
+	GraphicsPSOInit.bDepthBounds = bDepthBoundsTestEnabled;
+	if (bDepthBoundsTestEnabled)
+	{
+		// no depth test or writes
+		GraphicsPSOInit.DepthStencilState = TStaticDepthStencilState<false, CF_Always>::GetRHI();
+	}
+	else if (bStencilTestEnabled)
+	{
+		if (GStencilOptimization)
 		{
-			if (GStencilOptimization)
-			{
-				// No depth test or writes, zero the stencil
-				// Note: this will disable hi-stencil on many GPUs, but still seems 
-				// to be faster. However, early stencil still works 
-				GraphicsPSOInit.DepthStencilState =
-					TStaticDepthStencilState<
-					false, CF_Always,
-					true, CF_NotEqual, SO_Zero, SO_Zero, SO_Zero,
-					false, CF_Always, SO_Zero, SO_Zero, SO_Zero,
-					ShadowStencilMask, ShadowStencilMask
-					>::GetRHI();
-			}
-			else
-			{
-				// no depth test or writes, Test stencil for non-zero.
-				GraphicsPSOInit.DepthStencilState =
-					TStaticDepthStencilState<
-					false, CF_Always,
-					true, CF_NotEqual, SO_Keep, SO_Keep, SO_Keep,
-					false, CF_Always, SO_Keep, SO_Keep, SO_Keep,
-					ShadowStencilMask, ShadowStencilMask
-					>::GetRHI();
-			}
+			// No depth test or writes, zero the stencil
+			// Note: this will disable hi-stencil on many GPUs, but still seems 
+			// to be faster. However, early stencil still works 
+			GraphicsPSOInit.DepthStencilState =
+				TStaticDepthStencilState<
+				false, CF_Always,
+				true, CF_NotEqual, SO_Zero, SO_Zero, SO_Zero,
+				false, CF_Always, SO_Zero, SO_Zero, SO_Zero,
+				ShadowStencilMask, ShadowStencilMask
+				>::GetRHI();
 		}
 		else
 		{
-			GraphicsPSOInit.DepthStencilState = TStaticDepthStencilState<false, CF_DepthFartherOrEqual>::GetRHI();
+			// no depth test or writes, Test stencil for non-zero.
+			GraphicsPSOInit.DepthStencilState =
+				TStaticDepthStencilState<
+				false, CF_Always,
+				true, CF_NotEqual, SO_Keep, SO_Keep, SO_Keep,
+				false, CF_Always, SO_Keep, SO_Keep, SO_Keep,
+				ShadowStencilMask, ShadowStencilMask
+				>::GetRHI();
 		}
+	}
+	else
+	{
+		GraphicsPSOInit.DepthStencilState = TStaticDepthStencilState<false, CF_DepthFartherOrEqual>::GetRHI();
+	}
 
-		GraphicsPSOInit.BlendState = GetBlendStateForProjection(bProjectingForForwardShading, bMobileModulatedProjections);
+	GraphicsPSOInit.BlendState = GetBlendStateForProjection(bProjectingForForwardShading, bMobileModulatedProjections);
 
-		GraphicsPSOInit.PrimitiveType = IsWholeSceneDirectionalShadow() ? PT_TriangleStrip : PT_TriangleList;
+	GraphicsPSOInit.PrimitiveType = IsWholeSceneDirectionalShadow() ? PT_TriangleStrip : PT_TriangleList;
 
-		{
-			uint32 LocalQuality = GetShadowQuality();
+	{
+		uint32 LocalQuality = GetShadowQuality();
 
-			if (LocalQuality > 1)
+		if (LocalQuality > 1)
+		{
+			if (IsWholeSceneDirectionalShadow() && CascadeSettings.ShadowSplitIndex > 0)
 			{
-				if (IsWholeSceneDirectionalShadow() && CascadeSettings.ShadowSplitIndex > 0)
-				{
-					// adjust kernel size so that the penumbra size of distant splits will better match up with the closer ones
-					const float SizeScale = CascadeSettings.ShadowSplitIndex / FMath::Max(0.001f, CVarCSMSplitPenumbraScale.GetValueOnRenderThread());
-				}
-				else if (LocalQuality > 2 && !bWholeSceneShadow)
-				{
-					static auto CVarPreShadowResolutionFactor = IConsoleManager::Get().FindTConsoleVariableDataFloat(TEXT("r.Shadow.PreShadowResolutionFactor"));
-					const int32 TargetResolution = bPreShadow ? FMath::TruncToInt(512 * CVarPreShadowResolutionFactor->GetValueOnRenderThread()) : 512;
+				// adjust kernel size so that the penumbra size of distant splits will better match up with the closer ones
+				const float SizeScale = CascadeSettings.ShadowSplitIndex / FMath::Max(0.001f, CVarCSMSplitPenumbraScale.GetValueOnRenderThread());
+			}
+			else if (LocalQuality > 2 && !bWholeSceneShadow)
+			{
+				static auto CVarPreShadowResolutionFactor = IConsoleManager::Get().FindTConsoleVariableDataFloat(TEXT("r.Shadow.PreShadowResolutionFactor"));
+				const int32 TargetResolution = bPreShadow ? FMath::TruncToInt(512 * CVarPreShadowResolutionFactor->GetValueOnRenderThread()) : 512;
 
-					int32 Reduce = 0;
+				int32 Reduce = 0;
 
-					{
-						int32 Res = ResolutionX;
+				{
+					int32 Res = ResolutionX;
 
-						while (Res < TargetResolution)
-						{
-							Res *= 2;
-							++Reduce;
-						}
+					while (Res < TargetResolution)
+					{
+						Res *= 2;
+						++Reduce;
 					}
-
-					// Never drop to quality 1 due to low resolution, aliasing is too bad
-					LocalQuality = FMath::Clamp((int32)LocalQuality - Reduce, 3, 5);
 				}
-			}
-
-			GraphicsPSOInit.BoundShaderState.VertexDeclarationRHI = GetVertexDeclarationFVector4();
-			BindShadowProjectionShaders(LocalQuality, RHICmdList, GraphicsPSOInit, ViewIndex, *View, HairVisibilityData, this, bMobileModulatedProjections);
 
-			if (bDepthBoundsTestEnabled)
-			{
-				SetDepthBoundsTest(RHICmdList, CascadeSettings.SplitNear, CascadeSettings.SplitFar, View->ViewMatrices.GetProjectionMatrix());
+				// Never drop to quality 1 due to low resolution, aliasing is too bad
+				LocalQuality = FMath::Clamp((int32)LocalQuality - Reduce, 3, 5);
 			}
-
-			RHICmdList.SetStencilRef(0);
 		}
 
-		if (IsWholeSceneDirectionalShadow())
+		GraphicsPSOInit.BoundShaderState.VertexDeclarationRHI = GetVertexDeclarationFVector4();
+		BindShadowProjectionShaders(LocalQuality, RHICmdList, GraphicsPSOInit, ViewIndex, *View, HairVisibilityData, this, bMobileModulatedProjections);
+
+		if (bDepthBoundsTestEnabled)
 		{
-			RHICmdList.SetStreamSource(0, GClearVertexBuffer.VertexBufferRHI, 0);
-			RHICmdList.DrawPrimitive(0, 2, 1);
+			SetDepthBoundsTest(RHICmdList, CascadeSettings.SplitNear, CascadeSettings.SplitFar, View->ViewMatrices.GetProjectionMatrix());
 		}
-		else
-		{
-			FRHIResourceCreateInfo CreateInfo(TEXT("FProjectedShadowInfoFrustum"));
-			FBufferRHIRef VertexBufferRHI = RHICreateVertexBuffer(sizeof(FVector4) * FrustumVertices.Num(), BUF_Volatile, CreateInfo);
-			void* VoidPtr = RHILockBuffer(VertexBufferRHI, 0, sizeof(FVector4) * FrustumVertices.Num(), RLM_WriteOnly);
-			FPlatformMemory::Memcpy(VoidPtr, FrustumVertices.GetData(), sizeof(FVector4) * FrustumVertices.Num());
-			RHIUnlockBuffer(VertexBufferRHI);
 
-			RHICmdList.SetStreamSource(0, VertexBufferRHI, 0);
-			// Draw the frustum using the projection shader..
-			RHICmdList.DrawIndexedPrimitive(GCubeIndexBuffer.IndexBufferRHI, 0, 0, 8, 0, 12, 1);
-			VertexBufferRHI.SafeRelease();
+		RHICmdList.SetStencilRef(0);
+	}
+
+	if (IsWholeSceneDirectionalShadow())
+	{
+		RHICmdList.SetStreamSource(0, GClearVertexBuffer.VertexBufferRHI, 0);
+		RHICmdList.DrawPrimitive(0, 2, 1);
+	}
+	else
+	{
+		FRHIResourceCreateInfo CreateInfo(TEXT("FProjectedShadowInfoFrustum"));
+		FBufferRHIRef VertexBufferRHI = RHICreateVertexBuffer(sizeof(FVector4) * FrustumVertices.Num(), BUF_Volatile, CreateInfo);
+		void* VoidPtr = RHILockBuffer(VertexBufferRHI, 0, sizeof(FVector4) * FrustumVertices.Num(), RLM_WriteOnly);
+		FPlatformMemory::Memcpy(VoidPtr, FrustumVertices.GetData(), sizeof(FVector4) * FrustumVertices.Num());
+		RHIUnlockBuffer(VertexBufferRHI);
+
+		RHICmdList.SetStreamSource(0, VertexBufferRHI, 0);
+		// Draw the frustum using the projection shader..
+		RHICmdList.DrawIndexedPrimitive(GCubeIndexBuffer.IndexBufferRHI, 0, 0, 8, 0, 12, 1);
+		VertexBufferRHI.SafeRelease();
+	}
+
+	if (!bDepthBoundsTestEnabled && bStencilTestEnabled)
+	{
+		// Clear the stencil buffer to 0.
+		if (!GStencilOptimization)
+		{
+			DrawClearQuad(RHICmdList, false, FLinearColor::Transparent, false, 0, true, 0);
 		}
+	}
+
+	RHICmdList.SetScissorRect(false, 0, 0, 0, 0);
+}
+
+void FProjectedShadowInfo::RenderMobileModulatedShadowProjection(
+	FRHICommandListImmediate& RHICmdList,
+	int32 ViewIndex,
+	const FViewInfo* View,
+	const FLightSceneProxy* LightSceneProxy,
+	const FSceneRenderer* SceneRender) const
+{
+	// Find the shadow's view relevance.
+	const FVisibleLightViewInfo& VisibleLightViewInfo = View->VisibleLightInfos[LightSceneInfo->Id];
+	{
+		FPrimitiveViewRelevance ViewRelevance = VisibleLightViewInfo.ProjectedShadowViewRelevanceMap[ShadowId];
 
-		if (!bDepthBoundsTestEnabled && bStencilTestEnabled)
+		// Don't render shadows for subjects which aren't view relevant.
+		if (ViewRelevance.bShadowRelevance == false)
 		{
-			// Clear the stencil buffer to 0.
-			if (!GStencilOptimization)
-			{
-				DrawClearQuad(RHICmdList, false, FLinearColor::Transparent, false, 0, true, 0);
-			}
+			return;
 		}
+	}
 
-		RHICmdList.SetScissorRect(false, 0, 0, 0, 0);
-	});
+	FString EventName;
+
+#if WANTS_DRAW_MESH_EVENTS
+	if (GetEmitDrawEvents())
+	{
+		GetShadowTypeNameForDrawEvent(EventName);
+	}
+#endif
+
+	const bool bProjectingForForwardShading = false;
+	const bool bMobileModulatedProjections = true;
+
+	// GPUCULL_TODO, mobile modulated shadow is inside the mobile render pass, probably couldn't work with the GPUCull.
+	FInstanceCullingDrawParams InstanceCullingDrawParams;
+
+	RenderProjectionInternal(RHICmdList, ViewIndex, View, LightSceneProxy, SceneRender, bProjectingForForwardShading, bMobileModulatedProjections, nullptr, InstanceCullingDrawParams);
 }
 
 template <uint32 Quality, bool bUseTransmission, bool bUseSubPixel>
@@ -2202,7 +2255,7 @@ void FDeferredShadingSceneRenderer::RenderDeferredShadowProjections(
 	}
 }
 
-void FMobileSceneRenderer::RenderModulatedShadowProjections(FRDGBuilder& GraphBuilder, TRDGUniformBufferRef<FMobileSceneTextureUniformParameters> MobileSceneTextures)
+void FMobileSceneRenderer::RenderModulatedShadowProjections(FRHICommandListImmediate& RHICmdList, int32 ViewIndex, const FViewInfo& View)
 {
 	if (IsSimpleForwardShadingEnabled(ShaderPlatform) || !ViewFamily.EngineShowFlags.DynamicShadows)
 	{
@@ -2211,11 +2264,8 @@ void FMobileSceneRenderer::RenderModulatedShadowProjections(FRDGBuilder& GraphBu
 
 	SCOPED_NAMED_EVENT(FMobileSceneRenderer_RenderModulatedShadowProjections, FColor::Emerald);
 	SCOPE_CYCLE_COUNTER(STAT_ProjectedShadowDrawTime);
-	RDG_EVENT_SCOPE(GraphBuilder, "ShadowProjectionOnOpaque");
-	RDG_GPU_STAT_SCOPE(GraphBuilder, ShadowProjection);
-
-	FShadowProjectionPassParameters CommonPassParameters;
-	CommonPassParameters.SceneTextures.MobileSceneTextures = MobileSceneTextures;
+	SCOPED_DRAW_EVENT(RHICmdList, ShadowProjectionOnOpaque);
+	SCOPED_GPU_STAT(RHICmdList, ShadowProjection);
 
 	// render shadowmaps for relevant lights.
 	for (TSparseArray<FLightSceneInfoCompact>::TConstIterator LightIt(Scene->Lights); LightIt; ++LightIt)
@@ -2229,17 +2279,29 @@ void FMobileSceneRenderer::RenderModulatedShadowProjections(FRDGBuilder& GraphBu
 			const FVisibleLightInfo& VisibleLightInfo = VisibleLightInfos[LightSceneInfo->Id];
 
 			if (VisibleLightInfo.ShadowsToProject.Num() > 0)
-			{
-				const bool bProjectingForForwardShading = false;
-				const bool bMobileModulatedProjections = true;
-				RenderShadowProjections(
-					GraphBuilder,
-					CommonPassParameters,
-					LightSceneProxy,
-					nullptr,
-					VisibleLightInfo.ShadowsToProject,
-					bProjectingForForwardShading,
-					bMobileModulatedProjections);
+			{			
+				LightSceneProxy->SetScissorRect(RHICmdList, View, View.ViewRect);
+
+				const FHairStrandsVisibilityData* HairVisibilityData = nullptr;
+				const FHairStrandsMacroGroupDatas* HairMacroGroupData = nullptr;
+
+				TArrayView<const FProjectedShadowInfo* const> Shadows = VisibleLightInfo.ShadowsToProject;
+
+				// Project the shadow depth buffers onto the scene.
+				for (const FProjectedShadowInfo* ProjectedShadowInfo : Shadows)
+				{
+					if (ProjectedShadowInfo->bAllocated)
+					{
+						// Only project the shadow if it's large enough in this particular view (split screen, etc... may have shadows that are large in one view but irrelevantly small in others)
+						if (ProjectedShadowInfo->FadeAlphas[ViewIndex] > 1.0f / 256.0f)
+						{
+							checkSlow(!ProjectedShadowInfo->bOnePassPointLightShadow);
+							ProjectedShadowInfo->RenderMobileModulatedShadowProjection(RHICmdList, ViewIndex, &View, LightSceneProxy, this);
+						}
+					}
+				}
+
+				RHICmdList.SetScissorRect(false, 0, 0, 0, 0);
 			}
 		}
 	}
diff --git a/Engine/Source/Runtime/Renderer/Private/ShadowRendering.h b/Engine/Source/Runtime/Renderer/Private/ShadowRendering.h
--- a/Engine/Source/Runtime/Renderer/Private/ShadowRendering.h
+++ b/Engine/Source/Runtime/Renderer/Private/ShadowRendering.h
@@ -527,6 +527,27 @@ public:
 		const FHairStrandsVisibilityData* HairVisibilityData,
 		const struct FHairStrandsMacroGroupDatas* HairMacroGroupData) const;
 
+	void RenderProjectionInternal(
+		FRHICommandListImmediate& RHICmdList,
+		int32 ViewIndex,
+		const FViewInfo* View,
+		const FLightSceneProxy* LightSceneProxy,
+		const FSceneRenderer* SceneRender,
+		bool bProjectingForForwardShading,
+		bool bMobileModulatedProjections,
+		const FHairStrandsVisibilityData* HairVisibilityData,
+		const FInstanceCullingDrawParams& InstanceCullingDrawParams) const;
+
+	/**
+	 * Projects the mobile modulated shadow onto the scene for a particular view.
+	 */
+	void RenderMobileModulatedShadowProjection(
+		FRHICommandListImmediate& RHICmdList,
+		int32 ViewIndex,
+		const FViewInfo* View,
+		const FLightSceneProxy* LightSceneProxy,
+		const FSceneRenderer* SceneRender) const;
+
 	/**
 	* Renders the shadow subject depth, to a particular hacked view
 	*/
