Fix a compiling bug.

diff --git a/Engine/Source/Runtime/Renderer/Private/MobileShadingRenderer.cpp b/Engine/Source/Runtime/Renderer/Private/MobileShadingRenderer.cpp
--- a/Engine/Source/Runtime/Renderer/Private/MobileShadingRenderer.cpp
+++ b/Engine/Source/Runtime/Renderer/Private/MobileShadingRenderer.cpp
@@ -925,11 +925,13 @@ void FMobileSceneRenderer::RenderForward(FRDGBuilder& GraphBuilder, FRDGTextureR
 
 //		View.ParallelMeshDrawCommandPasses[EMeshPass::BasePass].BuildRenderingCommands(GraphBuilder, Scene->GPUScene, PassParameters->InstanceCullingDrawParams[EMeshPass::BasePass]);
 
+#if WITH_DEBUG_VIEW_MODES
 		if (bShouldRenderMobileDebugView)
 		{
 			// GPUCULL_TODO: View.ParallelMeshDrawCommandPasses[EMeshPass::DebugViewMode].BuildRenderingCommands(GraphBuilder, Scene->GPUScene, PassParameters->InstanceCullingDrawParams[EMeshPass::DebugViewMode]);
 			PassParameters->DebugViewMode = CreateDebugViewModePassUniformBuffer(GraphBuilder, View, SceneTextures.QuadOverdraw);
 		}
+#endif
 
 		if (bShouldRenderTranslucency)
 		{
