Fix Scene Occlusion by restoring the setting of FViewUniformShaderParameters

diff --git a/Engine/Source/Runtime/Renderer/Private/SceneOcclusion.h b/Engine/Source/Runtime/Renderer/Private/SceneOcclusion.h
--- a/Engine/Source/Runtime/Renderer/Private/SceneOcclusion.h
+++ b/Engine/Source/Runtime/Renderer/Private/SceneOcclusion.h
@@ -83,6 +83,8 @@ public:
 
 	void SetParameters(FRHICommandList& RHICmdList, const FViewInfo& View)
 	{
+		FGlobalShader::SetParameters<FViewUniformShaderParameters>(RHICmdList, RHICmdList.GetBoundVertexShader(),View.ViewUniformBuffer);
+
 		// Don't transform if rendering frustum
 		StencilingGeometryParameters.Set(RHICmdList, this, FVector4(0,0,0,1));
 
