Add Virtual Shadow Map Array initialization for Mobile renderer to ensure dummy buffers are set up.

diff --git a/Engine/Source/Runtime/Renderer/Private/MobileShadingRenderer.cpp b/Engine/Source/Runtime/Renderer/Private/MobileShadingRenderer.cpp
--- a/Engine/Source/Runtime/Renderer/Private/MobileShadingRenderer.cpp
+++ b/Engine/Source/Runtime/Renderer/Private/MobileShadingRenderer.cpp
@@ -613,6 +613,11 @@ void FMobileSceneRenderer::Render(FRDGBuilder& GraphBuilder)
 
 	FInstanceCullingManager InstanceCullingManager(GInstanceCullingManagerResources, Scene->GPUScene.IsEnabled());
 
+	// Important that this uses consistent logic throughout the frame, so evaluate once and pass in the flag from here
+	// NOTE: Must be done after  system texture initialization
+	VirtualShadowMapArray.Initialize(GraphBuilder, UseVirtualShadowMaps(ShaderPlatform, FeatureLevel));
+
+
 	// Find the visible primitives and prepare targets and buffers for rendering
 	InitViews(GraphBuilder, SceneTexturesConfig, InstanceCullingManager);
 
