Calculate accurate total number of ray tracing instances, taking into account that a single FRayTracingGeometryInstance may contain multiple instanced geometries.

diff --git a/Engine/Source/Runtime/Renderer/Private/DeferredShadingRenderer.cpp b/Engine/Source/Runtime/Renderer/Private/DeferredShadingRenderer.cpp
--- a/Engine/Source/Runtime/Renderer/Private/DeferredShadingRenderer.cpp
+++ b/Engine/Source/Runtime/Renderer/Private/DeferredShadingRenderer.cpp
@@ -1171,7 +1171,16 @@ bool FDeferredShadingSceneRenderer::DispatchRayTracingWorldUpdates(FRDGBuilder&
 	for (int32 ViewIndex = 0; ViewIndex < Views.Num(); ++ViewIndex)
 	{
 		FViewInfo& View = Views[ViewIndex];
-		SET_DWORD_STAT(STAT_RayTracingInstances, View.RayTracingGeometryInstances.Num());
+
+#if STATS
+		uint32 TotalInstances = 0;
+		for (const FRayTracingGeometryInstance& Instance : View.RayTracingGeometryInstances)
+		{
+			ensure(Instance.NumTransforms >= uint32(Instance.GetTransforms().Num()));
+			TotalInstances += Instance.NumTransforms;
+		}
+		SET_DWORD_STAT(STAT_RayTracingInstances, TotalInstances);
+#endif // STATS
 
 #ifdef DO_CHECK
 		// Validate all the ray tracing geometries lifetimes
