Fix hair strands HitProxy and outline rendering.

diff --git a/Engine/Source/Runtime/Renderer/Private/MeshPassProcessor.cpp b/Engine/Source/Runtime/Renderer/Private/MeshPassProcessor.cpp
--- a/Engine/Source/Runtime/Renderer/Private/MeshPassProcessor.cpp
+++ b/Engine/Source/Runtime/Renderer/Private/MeshPassProcessor.cpp
@@ -1072,7 +1072,12 @@ void FMeshDrawCommand::SubmitDrawEnd(const FMeshDrawCommand& MeshDrawCommand, ui
 	FRHIBuffer* IndirectArgsOverrideBuffer,
 	uint32 IndirectArgsOverrideByteOffset)
 {
-	const bool bDoOverrideArgs = IndirectArgsOverrideBuffer != nullptr && MeshDrawCommand.PrimitiveIdStreamIndex >= 0;
+	// Do not use the output result of the instance culling for MeshDrawCommand which initially provides indirect draw arguments. Otherwise these MeshDrawCommands
+	// wouldn't be drawn, as their NumPrimitives is initially set to 0. This would cause the instance culling to issue indirect args (stored in IndirectArgfsOverrideBuffer) 
+	// with NumPrimitives set to 0
+	const bool bDoOverrideArgs = 
+		IndirectArgsOverrideBuffer != nullptr && MeshDrawCommand.PrimitiveIdStreamIndex >= 0 &&
+		MeshDrawCommand.NumPrimitives > 0 && MeshDrawCommand.IndirectArgs.Buffer == nullptr;
 
 	if (MeshDrawCommand.IndexBuffer)
 	{
