Fix packaging (shader compilation) errors packaging for Android (Virtual Shadowmap Sampling shouldn't be enabled on mobile devices).
#rnx

diff --git a/Engine/Source/Runtime/Renderer/Private/SkyAtmosphereRendering.cpp b/Engine/Source/Runtime/Renderer/Private/SkyAtmosphereRendering.cpp
--- a/Engine/Source/Runtime/Renderer/Private/SkyAtmosphereRendering.cpp
+++ b/Engine/Source/Runtime/Renderer/Private/SkyAtmosphereRendering.cpp
@@ -284,6 +284,10 @@ static bool ShouldPipelineCompileSkyAtmosphereShader(EShaderPlatform ShaderPlatf
 	return RHISupportsComputeShaders(ShaderPlatform);
 }
 
+static bool VirtualShadowMapSamplingSupported(EShaderPlatform ShaderPlatform)
+{
+	return GetMaxSupportedFeatureLevel(ShaderPlatform) >= ERHIFeatureLevel::SM5;
+}
 
 
 bool ShouldSkySampleAtmosphereLightsOpaqueShadow(const FScene& Scene, const TArray<FVisibleLightInfo, SceneRenderingAllocator>& VisibleLightInfos, SkyAtmosphereLightShadowData& LightShadowData)
@@ -679,7 +683,8 @@ class FRenderSkyAtmospherePS : public FGlobalShader
 		OutEnvironment.SetDefine(TEXT("MULTISCATTERING_APPROX_SAMPLING_ENABLED"), 1);
 		OutEnvironment.SetDefine(TEXT("SOURCE_DISK_ENABLED"), 1);
 
-		if (PermutationVector.Get<FSampleOpaqueShadow>())
+		
+		if (PermutationVector.Get<FSampleOpaqueShadow>() && VirtualShadowMapSamplingSupported(Parameters.Platform))
 		{
 			OutEnvironment.SetDefine(TEXT("VIRTUAL_SHADOW_MAP"), 1);
 			FVirtualShadowMapArray::SetShaderDefines(OutEnvironment);
@@ -853,7 +858,7 @@ public:
 		OutEnvironment.SetDefine(TEXT("MULTISCATTERING_APPROX_SAMPLING_ENABLED"), 1);
 		OutEnvironment.SetDefine(TEXT("SOURCE_DISK_ENABLED"), 1);
 
-		if (PermutationVector.Get<FSampleOpaqueShadow>())
+		if (PermutationVector.Get<FSampleOpaqueShadow>() && VirtualShadowMapSamplingSupported(Parameters.Platform))
 		{
 			OutEnvironment.SetDefine(TEXT("VIRTUAL_SHADOW_MAP"), 1);
 			FVirtualShadowMapArray::SetShaderDefines(OutEnvironment);
@@ -914,7 +919,7 @@ public:
 		OutEnvironment.SetDefine(TEXT("THREADGROUP_SIZE"), GroupSize);
 		OutEnvironment.SetDefine(TEXT("MULTISCATTERING_APPROX_SAMPLING_ENABLED"), 1);
 
-		if (PermutationVector.Get<FSampleOpaqueShadow>())
+		if (PermutationVector.Get<FSampleOpaqueShadow>() && VirtualShadowMapSamplingSupported(Parameters.Platform))
 		{
 			OutEnvironment.SetDefine(TEXT("VIRTUAL_SHADOW_MAP"), 1);
 			FVirtualShadowMapArray::SetShaderDefines(OutEnvironment);
