Fix race condition in async light primitive creation.

diff --git a/Engine/Source/Runtime/Renderer/Private/RendererScene.cpp b/Engine/Source/Runtime/Renderer/Private/RendererScene.cpp
--- a/Engine/Source/Runtime/Renderer/Private/RendererScene.cpp
+++ b/Engine/Source/Runtime/Renderer/Private/RendererScene.cpp
@@ -3701,6 +3701,10 @@ void FScene::UpdateAllPrimitiveSceneInfos(FRDGBuilder& GraphBuilder, bool bAsync
 	RDG_EVENT_SCOPE(GraphBuilder, "UpdateAllPrimitiveSceneInfos");
 
 	TArray<FPrimitiveSceneInfo*> RemovedLocalPrimitiveSceneInfos(RemovedPrimitiveSceneInfos.Array());
+	// NOTE: We clear this early because IsPrimitiveBeingRemoved gets called from the CreateLightPrimitiveInteraction (to make sure that old primitives are not accessed) 
+	// we cannot safely kick off the AsyncCreateLightPrimitiveInteractionsTask before the RemovedPrimitiveSceneInfos has been cleared.
+	RemovedPrimitiveSceneInfos.Reset();
+
 	RemovedLocalPrimitiveSceneInfos.Sort(FPrimitiveArraySortKey());
 
 	if (VirtualShadowMapArrayCacheManager)
@@ -4267,7 +4271,6 @@ void FScene::UpdateAllPrimitiveSceneInfos(FRDGBuilder& GraphBuilder, bool bAsync
 	UpdatedCustomPrimitiveParams.Reset();
 	OverridenPreviousTransforms.Reset();
 	DistanceFieldSceneDataUpdates.Reset();
-	RemovedPrimitiveSceneInfos.Reset();
 	AddedPrimitiveSceneInfos.Reset();
 }
 
