Fixed crash when toggling properties on dir light.

diff --git a/Engine/Source/Runtime/Renderer/Private/VolumetricCloudRendering.cpp b/Engine/Source/Runtime/Renderer/Private/VolumetricCloudRendering.cpp
--- a/Engine/Source/Runtime/Renderer/Private/VolumetricCloudRendering.cpp
+++ b/Engine/Source/Runtime/Renderer/Private/VolumetricCloudRendering.cpp
@@ -1298,11 +1298,11 @@ void FSceneRenderer::InitVolumetricCloudsForViews(FRDGBuilder& GraphBuilder, boo
 			for (int32 ViewIndex = 0; ViewIndex < Views.Num(); ViewIndex++)
 			{
 				FViewInfo& ViewInfo = Views[ViewIndex];
-				if (ViewInfo.ViewState != nullptr)
+				for (int LightIndex = 0; LightIndex < NUM_ATMOSPHERE_LIGHTS; ++LightIndex)
 				{
-					for (int LightIndex = 0; LightIndex < NUM_ATMOSPHERE_LIGHTS; ++LightIndex)
+					ViewInfo.VolumetricCloudShadowRenderTarget[LightIndex] = nullptr;
+					if (ViewInfo.ViewState != nullptr)
 					{
-						ViewInfo.VolumetricCloudShadowRenderTarget[LightIndex] = nullptr;
 						ViewInfo.ViewState->VolumetricCloudShadowRenderTarget[LightIndex].Reset();
 					}
 				}
@@ -1317,6 +1317,20 @@ void FSceneRenderer::InitVolumetricCloudsForViews(FRDGBuilder& GraphBuilder, boo
 		return;
 	}
 
+	// First make sure we always clear the texture on  views to make sure no dangling texture pointers are ever used
+	AddPass(GraphBuilder, [&Views = Views](FRHICommandList&)
+	{
+		for (int32 ViewIndex = 0; ViewIndex < Views.Num(); ViewIndex++)
+		{
+			FViewInfo& ViewInfo = Views[ViewIndex];
+			for (int LightIndex = 0; LightIndex < NUM_ATMOSPHERE_LIGHTS; ++LightIndex)
+			{
+				ViewInfo.VolumetricCloudSkyAO = nullptr;
+				ViewInfo.VolumetricCloudShadowRenderTarget[LightIndex] = nullptr;
+			}
+		}
+	});
+
 	if (Scene)
 	{
 		check(ShouldRenderVolumetricCloud(Scene, ViewFamily.EngineShowFlags)); // This should not be called if we should not render SkyAtmosphere
@@ -1620,7 +1634,7 @@ void FSceneRenderer::InitVolumetricCloudsForViews(FRDGBuilder& GraphBuilder, boo
 
 						// We need to make a copy of the parameters on CPU to morph them because the creation is deferred.
 						FRenderVolumetricCloudGlobalParameters& VolumetricCloudParamsAO = *GraphBuilder.AllocParameters<FRenderVolumetricCloudGlobalParameters>();
-						VolumetricCloudParamsAO = VolumetricCloudParams;	// Use the same parameter as for the SkyAO
+						VolumetricCloudParamsAO = VolumetricCloudParams;	// Use the same parameter as for the directional light shadow
 						VolumetricCloudParamsAO.TraceShadowmap = 0;			// Notify that this pass is for SkyAO (avoid to use another shader permutation)
 						TRDGUniformBufferRef<FRenderVolumetricCloudGlobalParameters> TraceVolumetricCloudSkyAOParamsUB = GraphBuilder.CreateUniformBuffer(&VolumetricCloudParamsAO);
 						TraceCloudTexture(CloudSkyAOTexture, true, TraceVolumetricCloudSkyAOParamsUB);
@@ -1753,9 +1767,9 @@ void FSceneRenderer::InitVolumetricCloudsForViews(FRDGBuilder& GraphBuilder, boo
 						}
 						else
 						{
+							ViewInfo.VolumetricCloudShadowRenderTarget[LightIndex] = nullptr;
 							if (ViewInfo.ViewState)
 							{
-								ViewInfo.VolumetricCloudShadowRenderTarget[LightIndex] = nullptr;
 								ViewInfo.ViewState->VolumetricCloudShadowRenderTarget[LightIndex].Reset();
 							}
 						}
