Fix duplicated FShaderPlatformCachedIniValue caches for TMobileBasePassPS template permutations causing long editor load times

diff --git a/Engine/Source/Runtime/Renderer/Private/MobileBasePassRendering.cpp b/Engine/Source/Runtime/Renderer/Private/MobileBasePassRendering.cpp
--- a/Engine/Source/Runtime/Renderer/Private/MobileBasePassRendering.cpp
+++ b/Engine/Source/Runtime/Renderer/Private/MobileBasePassRendering.cpp
@@ -54,6 +54,11 @@ static TAutoConsoleVariable<int32> CVarMobileUseHWsRGBEncoding(
 	TEXT("1: Use GPU HW to convert linear to sRGB automatically (device must support sRGB write control)\n"),
 	ECVF_RenderThreadSafe);
 
+namespace MobileBasePass
+{
+	FShaderPlatformCachedIniValue<int32> MobileDynamicPointLightsUseStaticBranchIniValue(TEXT("/Script/Engine.RendererSettings"), TEXT("r.MobileDynamicPointLightsUseStaticBranch"));
+	FShaderPlatformCachedIniValue<int32> MobileNumDynamicPointLightsIniValue(TEXT("/Script/Engine.RendererSettings"), TEXT("r.MobileNumDynamicPointLights"));
+};
 
 #define IMPLEMENT_MOBILE_SHADING_BASEPASS_LIGHTMAPPED_VERTEX_SHADER_TYPE(LightMapPolicyType,LightMapPolicyName) \
 	typedef TMobileBasePassVS< LightMapPolicyType, LDR_GAMMA_32 > TMobileBasePassVS##LightMapPolicyName##LDRGamma32; \
diff --git a/Engine/Source/Runtime/Renderer/Private/MobileBasePassRendering.h b/Engine/Source/Runtime/Renderer/Private/MobileBasePassRendering.h
--- a/Engine/Source/Runtime/Renderer/Private/MobileBasePassRendering.h
+++ b/Engine/Source/Runtime/Renderer/Private/MobileBasePassRendering.h
@@ -370,6 +370,9 @@ namespace MobileBasePass
 	void SetTranslucentRenderState(FMeshPassProcessorRenderState& DrawRenderState, const FMaterial& Material);
 
 	bool StationarySkyLightHasBeenApplied(const FScene* Scene, ELightMapPolicyType LightMapPolicyType);
+
+	extern FShaderPlatformCachedIniValue<int32> MobileDynamicPointLightsUseStaticBranchIniValue;
+	extern FShaderPlatformCachedIniValue<int32> MobileNumDynamicPointLightsIniValue;
 };
 
 
@@ -394,12 +397,10 @@ public:
 	static bool ShouldCompilePermutation(const FMeshMaterialShaderPermutationParameters& Parameters)
 	{		
 		// We compile the point light shader combinations based on the project settings
-		static FShaderPlatformCachedIniValue<int32> MobileDynamicPointLightsUseStaticBranchIniValue(TEXT("/Script/Engine.RendererSettings"), TEXT("r.MobileDynamicPointLightsUseStaticBranch"));
-		static FShaderPlatformCachedIniValue<int32> MobileNumDynamicPointLightsIniValue(TEXT("/Script/Engine.RendererSettings"), TEXT("r.MobileNumDynamicPointLights"));
 		static auto* MobileSkyLightPermutationCVar = IConsoleManager::Get().FindTConsoleVariableDataInt(TEXT("r.Mobile.SkyLightPermutation"));
 
-		const bool bMobileDynamicPointLightsUseStaticBranch = (MobileDynamicPointLightsUseStaticBranchIniValue.Get(Parameters.Platform) == 1);
-		const int32 MobileNumDynamicPointLights = MobileNumDynamicPointLightsIniValue.Get(Parameters.Platform);
+		const bool bMobileDynamicPointLightsUseStaticBranch = (MobileBasePass::MobileDynamicPointLightsUseStaticBranchIniValue.Get(Parameters.Platform) == 1);
+		const int32 MobileNumDynamicPointLights = MobileBasePass::MobileNumDynamicPointLightsIniValue.Get(Parameters.Platform);
 		const int32 MobileSkyLightPermutationOptions = MobileSkyLightPermutationCVar->GetValueOnAnyThread();
 
 		const bool bIsUnlit = Parameters.MaterialParameters.ShadingModels.IsUnlit();
