Compilation fix after a robomerge

diff --git a/Engine/Source/Runtime/Renderer/Private/VolumetricCloudRendering.cpp b/Engine/Source/Runtime/Renderer/Private/VolumetricCloudRendering.cpp
--- a/Engine/Source/Runtime/Renderer/Private/VolumetricCloudRendering.cpp
+++ b/Engine/Source/Runtime/Renderer/Private/VolumetricCloudRendering.cpp
@@ -2367,7 +2367,7 @@ bool SetupLightCloudTransmittanceParameters(const FScene* Scene, const FViewInfo
 	{
 		OutParameters.CloudShadowmapTexture = View.VolumetricCloudShadowRenderTarget[0];
 		OutParameters.CloudShadowmapSampler = TStaticSamplerState<SF_Bilinear, AM_Clamp, AM_Clamp, AM_Clamp>::GetRHI();
-		OutParameters.CloudShadowmapFarDepthKm = CloudInfo->GetVolumetricCloudCommonShaderParameters().CloudShadowmapFarDepthKm[0];
+		OutParameters.CloudShadowmapFarDepthKm = CloudInfo->GetVolumetricCloudCommonShaderParameters().CloudShadowmapFarDepthKm[0].X;
 		OutParameters.CloudShadowmapWorldToLightClipMatrix = CloudInfo->GetVolumetricCloudCommonShaderParameters().CloudShadowmapWorldToLightClipMatrix[0];
 		OutParameters.CloudShadowmapStrength = AtmosphereLight0Proxy->GetCloudShadowOnSurfaceStrength();
 	}
@@ -2375,7 +2375,7 @@ bool SetupLightCloudTransmittanceParameters(const FScene* Scene, const FViewInfo
 	{
 		OutParameters.CloudShadowmapTexture = View.VolumetricCloudShadowRenderTarget[1];
 		OutParameters.CloudShadowmapSampler = TStaticSamplerState<SF_Bilinear, AM_Clamp, AM_Clamp, AM_Clamp>::GetRHI();
-		OutParameters.CloudShadowmapFarDepthKm = CloudInfo->GetVolumetricCloudCommonShaderParameters().CloudShadowmapFarDepthKm[1];
+		OutParameters.CloudShadowmapFarDepthKm = CloudInfo->GetVolumetricCloudCommonShaderParameters().CloudShadowmapFarDepthKm[1].X;
 		OutParameters.CloudShadowmapWorldToLightClipMatrix = CloudInfo->GetVolumetricCloudCommonShaderParameters().CloudShadowmapWorldToLightClipMatrix[1];
 		OutParameters.CloudShadowmapStrength = AtmosphereLight1Proxy->GetCloudShadowOnSurfaceStrength();
 	}
