Fix for PlanarReflections not transitioning resource to SRV

diff --git a/Engine/Source/Runtime/Renderer/Private/PlanarReflectionRendering.cpp b/Engine/Source/Runtime/Renderer/Private/PlanarReflectionRendering.cpp
--- a/Engine/Source/Runtime/Renderer/Private/PlanarReflectionRendering.cpp
+++ b/Engine/Source/Runtime/Renderer/Private/PlanarReflectionRendering.cpp
@@ -387,6 +387,8 @@ static void UpdatePlanarReflectionContents_RenderThread(
 				}
 
 				FRDGTextureRef ReflectionOutputTexture = GraphBuilder.RegisterExternalTexture(CreateRenderTarget(RenderTarget->TextureRHI, TEXT("ReflectionOutputTexture")));
+				GraphBuilder.SetTextureAccessFinal(ReflectionOutputTexture, ERHIAccess::SRVGraphics);
+
 				FSceneTextureShaderParameters SceneTextureParameters = CreateSceneTextureShaderParameters(GraphBuilder, SceneRenderer->FeatureLevel, ESceneTextureSetupMode::SceneDepth);
 				const FMinimalSceneTextures& SceneTextures = FSceneTextures::Get(GraphBuilder);
 
