Separate ray tracing geometry & scene build GPU stats.

diff --git a/Engine/Source/Runtime/Renderer/Private/DeferredShadingRenderer.cpp b/Engine/Source/Runtime/Renderer/Private/DeferredShadingRenderer.cpp
--- a/Engine/Source/Runtime/Renderer/Private/DeferredShadingRenderer.cpp
+++ b/Engine/Source/Runtime/Renderer/Private/DeferredShadingRenderer.cpp
@@ -236,8 +236,8 @@ DECLARE_CYCLE_STAT(TEXT("DeferredShadingSceneRenderer RenderFog"), STAT_FDeferre
 DECLARE_CYCLE_STAT(TEXT("DeferredShadingSceneRenderer RenderLightShaftBloom"), STAT_FDeferredShadingSceneRenderer_RenderLightShaftBloom, STATGROUP_SceneRendering);
 DECLARE_CYCLE_STAT(TEXT("DeferredShadingSceneRenderer RenderFinish"), STAT_FDeferredShadingSceneRenderer_RenderFinish, STATGROUP_SceneRendering);
 
-DECLARE_GPU_STAT_NAMED(RayTracingAS, TEXT("Ray Tracing Acceleration Structure Update/Refit"));
-DECLARE_GPU_STAT_NAMED(RayTracingDynamicGeom, TEXT("Ray Tracing Dynamic Geometry Update"));
+DECLARE_GPU_STAT(RayTracingScene);
+DECLARE_GPU_STAT(RayTracingGeometry);
 
 DECLARE_GPU_STAT(Postprocessing);
 DECLARE_GPU_STAT(VisibilityCommands);
@@ -1163,9 +1163,6 @@ bool FDeferredShadingSceneRenderer::DispatchRayTracingWorldUpdates(FRDGBuilder&
 		ReferenceView.AddRayTracingMeshBatchData.Empty();
 	}
 
-	bool bAsyncUpdateGeometry = (CVarRayTracingAsyncBuild.GetValueOnRenderThread() != 0)
-		&& GRHISupportsRayTracingAsyncBuildAccelerationStructure;
-
 	RDG_GPU_MASK_SCOPE(GraphBuilder, FRHIGPUMask::All());
 
 	for (int32 ViewIndex = 0; ViewIndex < Views.Num(); ++ViewIndex)
@@ -1238,26 +1235,10 @@ bool FDeferredShadingSceneRenderer::DispatchRayTracingWorldUpdates(FRDGBuilder&
 			EUniformBufferUsage::UniformBuffer_SingleFrame);
 	}
 
-	if (!bAsyncUpdateGeometry)
-	{
-		RDG_GPU_STAT_SCOPE(GraphBuilder, RayTracingAS);
-		RDG_GPU_STAT_SCOPE(GraphBuilder, RayTracingDynamicGeom);
 
-		AddPass(GraphBuilder, [this](FRHICommandListImmediate& RHICmdList)
-		{
-			Scene->GetRayTracingDynamicGeometryCollection()->DispatchUpdates(RHICmdList);
-		});
+	const bool bRayTracingAsyncBuild = CVarRayTracingAsyncBuild.GetValueOnRenderThread() != 0;
 
-		AddPass(GraphBuilder, RDG_EVENT_NAME("BuildRayTracingScene"), [this](FRHICommandList& RHICmdList)
-		{
-			for (int32 ViewIndex = 0; ViewIndex < Views.Num(); ++ViewIndex)
-			{
-				FViewInfo& View = Views[ViewIndex];
-				RHICmdList.BuildAccelerationStructure(View.RayTracingScene.RayTracingSceneRHI);
-			}
-		});
-	}
-	else
+	if (bRayTracingAsyncBuild && GRHISupportsRayTracingAsyncBuildAccelerationStructure)
 	{
 		AddPass(GraphBuilder, [this](FRHICommandList& RHICmdList)
 		{
@@ -1285,6 +1266,28 @@ bool FDeferredShadingSceneRenderer::DispatchRayTracingWorldUpdates(FRDGBuilder&
 			FRHIAsyncComputeCommandListImmediate::ImmediateDispatch(RHIAsyncCmdList);
 		});
 	}
+	else
+	{
+		{
+			RDG_GPU_STAT_SCOPE(GraphBuilder, RayTracingGeometry);
+			AddPass(GraphBuilder, RDG_EVENT_NAME("RayTracingGeometry"), [this](FRHICommandListImmediate& RHICmdList)
+			{
+				Scene->GetRayTracingDynamicGeometryCollection()->DispatchUpdates(RHICmdList);
+			});
+		}
+
+		{
+			RDG_GPU_STAT_SCOPE(GraphBuilder, RayTracingScene);
+			AddPass(GraphBuilder, RDG_EVENT_NAME("RayTracingScene"), [this](FRHICommandList& RHICmdList)
+			{
+				for (int32 ViewIndex = 0; ViewIndex < Views.Num(); ++ViewIndex)
+				{
+					FViewInfo& View = Views[ViewIndex];
+					RHICmdList.BuildAccelerationStructure(View.RayTracingScene.RayTracingSceneRHI);
+				}
+			});
+		}
+	}
 
 	AddPass(GraphBuilder, [this](FRHICommandListImmediate& RHICmdList)
 	{
