Rename instance culling buffers to match naming convention.

diff --git a/Engine/Source/Runtime/Renderer/Private/InstanceCulling/InstanceCullingContext.cpp b/Engine/Source/Runtime/Renderer/Private/InstanceCulling/InstanceCullingContext.cpp
--- a/Engine/Source/Runtime/Renderer/Private/InstanceCulling/InstanceCullingContext.cpp
+++ b/Engine/Source/Runtime/Renderer/Private/InstanceCulling/InstanceCullingContext.cpp
@@ -307,7 +307,7 @@ void FInstanceCullingContext::BuildRenderingCommands(FRDGBuilder& GraphBuilder,
 	// Note: use start at zero offset if there is no instance culling manager, this means each build rendering commands pass will overwrite the same ID range. Which is only ok assuming correct barriers (should be erring on this side by default).
 	TArray<uint32> NullArray;
 	NullArray.AddZeroed(1);
-	FRDGBufferRef InstanceIdOutOffsetBufferRDG = InstanceCullingManager ? InstanceCullingManager->CullingIntermediate.InstanceIdOutOffsetBuffer : CreateStructuredBuffer(GraphBuilder, TEXT("OutputOffsetBufferOutTransient"), NullArray);
+	FRDGBufferRef InstanceIdOutOffsetBufferRDG = InstanceCullingManager ? InstanceCullingManager->CullingIntermediate.InstanceIdOutOffsetBuffer : CreateStructuredBuffer(GraphBuilder, TEXT("InstanceCulling.OutputOffsetBufferOutTransient"), NullArray);
 	// If there is no manager, then there is no data on culling, so set flag to skip that and ignore buffers.
 	FRDGBufferRef VisibleInstanceFlagsRDG = InstanceCullingManager != nullptr ? InstanceCullingManager->CullingIntermediate.VisibleInstanceFlags : nullptr;
 	const bool bCullInstances = InstanceCullingManager != nullptr;
@@ -316,9 +316,9 @@ void FInstanceCullingContext::BuildRenderingCommands(FRDGBuilder& GraphBuilder,
 
 #if ENABLE_DETERMINISTIC_INSTANCE_CULLING
 	// 1. Compute output sizes for all commands
-	FRDGBufferRef InstanceCountsRDG = GraphBuilder.CreateBuffer(FRDGBufferDesc::CreateStructuredDesc(sizeof(FIntVector2), CullingCommands.Num()), TEXT("InstanceCounts"));
-	FRDGBufferRef InstanceIdOffsetBufferRDG = GraphBuilder.CreateBuffer(FRDGBufferDesc::CreateBufferDesc(sizeof(uint32), CullingCommands.Num()), TEXT("InstanceIdOffsetBuffer"));
-	FRDGBufferRef DrawIndirectArgsRDG = GraphBuilder.CreateBuffer(FRDGBufferDesc::CreateIndirectDesc(IndirectArgsNumWords * CullingCommands.Num()), TEXT("DrawIndirectArgsBuffer"));
+	FRDGBufferRef InstanceCountsRDG = GraphBuilder.CreateBuffer(FRDGBufferDesc::CreateStructuredDesc(sizeof(FIntVector2), CullingCommands.Num()), TEXT("InstanceCulling.InstanceCounts"));
+	FRDGBufferRef InstanceIdOffsetBufferRDG = GraphBuilder.CreateBuffer(FRDGBufferDesc::CreateBufferDesc(sizeof(uint32), CullingCommands.Num()), TEXT("InstanceCulling.InstanceIdOffsetBuffer"));
+	FRDGBufferRef DrawIndirectArgsRDG = GraphBuilder.CreateBuffer(FRDGBufferDesc::CreateIndirectDesc(IndirectArgsNumWords * CullingCommands.Num()), TEXT("InstanceCulling.DrawIndirectArgsBuffer"));
 
 	{
 		FComputeInstanceIdOutputSizeCs::FParameters* PassParameters = GraphBuilder.AllocParameters<FComputeInstanceIdOutputSizeCs::FParameters>();
@@ -332,17 +332,17 @@ void FInstanceCullingContext::BuildRenderingCommands(FRDGBuilder& GraphBuilder,
 		PassParameters->GPUScenePrimitiveSceneData = GPUScene.PrimitiveBuffer.SRV;
 		PassParameters->InstanceDataSOAStride = GPUScene.InstanceDataSOAStride;
 		// Upload data etc
-		PassParameters->PrimitiveCullingCommands = GraphBuilder.CreateSRV(CreateStructuredBuffer(GraphBuilder, TEXT("PrimitiveCullingCommands"), CullingCommands));
+		PassParameters->PrimitiveCullingCommands = GraphBuilder.CreateSRV(CreateStructuredBuffer(GraphBuilder, TEXT("InstanceCulling.PrimitiveCullingCommands"), CullingCommands));
 
-		PassParameters->PrimitiveIds = GraphBuilder.CreateSRV(CreateStructuredBuffer(GraphBuilder, TEXT("PrimitiveIds"), PrimitiveIds));
+		PassParameters->PrimitiveIds = GraphBuilder.CreateSRV(CreateStructuredBuffer(GraphBuilder, TEXT("InstanceCulling.PrimitiveIds"), PrimitiveIds));
 		PassParameters->DynamicPrimitiveIdOffset = DynamicPrimitiveIdRange.GetLowerBoundValue();
 		PassParameters->DynamicPrimitiveIdMax = DynamicPrimitiveIdRange.GetUpperBoundValue();
 
-		PassParameters->InstanceRuns = GraphBuilder.CreateSRV(CreateStructuredBuffer(GraphBuilder, TEXT("InstanceRuns"), InstanceRuns));
+		PassParameters->InstanceRuns = GraphBuilder.CreateSRV(CreateStructuredBuffer(GraphBuilder, TEXT("InstanceCulling.InstanceRuns"), InstanceRuns));
 
 		PassParameters->OutputOffsetBufferOut = GraphBuilder.CreateUAV(InstanceIdOutOffsetBufferRDG);
 
-		PassParameters->ViewIds = GraphBuilder.CreateSRV(CreateStructuredBuffer(GraphBuilder, TEXT("ViewIds"), ViewIds));
+		PassParameters->ViewIds = GraphBuilder.CreateSRV(CreateStructuredBuffer(GraphBuilder, TEXT("InstanceCulling.ViewIds"), ViewIds));
 		PassParameters->NumViewIds = ViewIds.Num();
 		PassParameters->NumPrimitiveIds = PrimitiveIds.Num();
 		PassParameters->NumInstanceRuns = InstanceRuns.Num();
@@ -398,15 +398,15 @@ void FInstanceCullingContext::BuildRenderingCommands(FRDGBuilder& GraphBuilder,
 		PassParameters->GPUScenePrimitiveSceneData = GPUScene.PrimitiveBuffer.SRV;
 		PassParameters->InstanceDataSOAStride = GPUScene.InstanceDataSOAStride;
 		// Upload data etc
-		PassParameters->PrimitiveCullingCommands = GraphBuilder.CreateSRV(CreateStructuredBuffer(GraphBuilder, TEXT("PrimitiveCullingCommands"), CullingCommands));
+		PassParameters->PrimitiveCullingCommands = GraphBuilder.CreateSRV(CreateStructuredBuffer(GraphBuilder, TEXT("InstanceCulling.PrimitiveCullingCommands"), CullingCommands));
 
-		PassParameters->PrimitiveIds = GraphBuilder.CreateSRV(CreateStructuredBuffer(GraphBuilder, TEXT("PrimitiveIds"), PrimitiveIds));
+		PassParameters->PrimitiveIds = GraphBuilder.CreateSRV(CreateStructuredBuffer(GraphBuilder, TEXT("InstanceCulling.PrimitiveIds"), PrimitiveIds));
 		PassParameters->DynamicPrimitiveIdOffset = DynamicPrimitiveIdRange.GetLowerBoundValue();
 		PassParameters->DynamicPrimitiveIdMax = DynamicPrimitiveIdRange.GetUpperBoundValue();
 
-		PassParameters->InstanceRuns = GraphBuilder.CreateSRV(CreateStructuredBuffer(GraphBuilder, TEXT("InstanceRuns"), InstanceRuns));
+		PassParameters->InstanceRuns = GraphBuilder.CreateSRV(CreateStructuredBuffer(GraphBuilder, TEXT("InstanceCulling.InstanceRuns"), InstanceRuns));
 
-		PassParameters->ViewIds = GraphBuilder.CreateSRV(CreateStructuredBuffer(GraphBuilder, TEXT("ViewIds"), ViewIds));
+		PassParameters->ViewIds = GraphBuilder.CreateSRV(CreateStructuredBuffer(GraphBuilder, TEXT("InstanceCulling.ViewIds"), ViewIds));
 		PassParameters->NumViewIds = ViewIds.Num();
 
 		// TODO: Remove this when everything is properly RDG'd
@@ -451,22 +451,22 @@ void FInstanceCullingContext::BuildRenderingCommands(FRDGBuilder& GraphBuilder,
 	PassParameters->GPUScenePrimitiveSceneData = GPUScene.PrimitiveBuffer.SRV;
 	PassParameters->InstanceDataSOAStride = GPUScene.InstanceDataSOAStride;
 	// Upload data etc
-	PassParameters->PrimitiveCullingCommands = GraphBuilder.CreateSRV(CreateStructuredBuffer(GraphBuilder, TEXT("PrimitiveCullingCommands"), CullingCommands));
+	PassParameters->PrimitiveCullingCommands = GraphBuilder.CreateSRV(CreateStructuredBuffer(GraphBuilder, TEXT("InstanceCulling.PrimitiveCullingCommands"), CullingCommands));
 
-	PassParameters->PrimitiveIds = GraphBuilder.CreateSRV(CreateStructuredBuffer(GraphBuilder, TEXT("PrimitiveIds"), PrimitiveIds));
+	PassParameters->PrimitiveIds = GraphBuilder.CreateSRV(CreateStructuredBuffer(GraphBuilder, TEXT("InstanceCulling.PrimitiveIds"), PrimitiveIds));
 	PassParameters->DynamicPrimitiveIdOffset = DynamicPrimitiveIdRange.GetLowerBoundValue();
 	PassParameters->DynamicPrimitiveIdMax = DynamicPrimitiveIdRange.GetUpperBoundValue();
 
-	PassParameters->InstanceRuns = GraphBuilder.CreateSRV(CreateStructuredBuffer(GraphBuilder, TEXT("InstanceRuns"), InstanceRuns));
+	PassParameters->InstanceRuns = GraphBuilder.CreateSRV(CreateStructuredBuffer(GraphBuilder, TEXT("InstanceCulling.InstanceRuns"), InstanceRuns));
 
 	PassParameters->OutputOffsetBufferOut = GraphBuilder.CreateUAV(InstanceIdOutOffsetBufferRDG);
 
-	FRDGBufferRef DrawIndirectArgsRDG = GraphBuilder.CreateBuffer(FRDGBufferDesc::CreateIndirectDesc(IndirectArgsNumWords * CullingCommands.Num()), TEXT("DrawIndirectArgsBuffer"));
+	FRDGBufferRef DrawIndirectArgsRDG = GraphBuilder.CreateBuffer(FRDGBufferDesc::CreateIndirectDesc(IndirectArgsNumWords * CullingCommands.Num()), TEXT("InstanceCulling.DrawIndirectArgsBuffer"));
 	// not using structured buffer as we want/have toget at it as a vertex buffer 
-	//FRDGBufferRef InstanceIdsBufferRDG = GraphBuilder.CreateBuffer(FRDGBufferDesc::CreateBufferDesc(sizeof(uint32), PrimitiveIds.Num() * FInstanceCullingManager::MaxAverageInstanceFactor), TEXT("InstanceIdsBuffer"));
-	FRDGBufferRef InstanceIdOffsetBufferRDG = GraphBuilder.CreateBuffer(FRDGBufferDesc::CreateBufferDesc(sizeof(uint32), CullingCommands.Num()), TEXT("InstanceIdOffsetBuffer"));
+	//FRDGBufferRef InstanceIdsBufferRDG = GraphBuilder.CreateBuffer(FRDGBufferDesc::CreateBufferDesc(sizeof(uint32), PrimitiveIds.Num() * FInstanceCullingManager::MaxAverageInstanceFactor), TEXT("InstanceCulling.InstanceIdsBuffer"));
+	FRDGBufferRef InstanceIdOffsetBufferRDG = GraphBuilder.CreateBuffer(FRDGBufferDesc::CreateBufferDesc(sizeof(uint32), CullingCommands.Num()), TEXT("InstanceCulling.InstanceIdOffsetBuffer"));
 
-	PassParameters->ViewIds = GraphBuilder.CreateSRV(CreateStructuredBuffer(GraphBuilder, TEXT("ViewIds"), ViewIds));
+	PassParameters->ViewIds = GraphBuilder.CreateSRV(CreateStructuredBuffer(GraphBuilder, TEXT("InstanceCulling.ViewIds"), ViewIds));
 	PassParameters->NumViewIds = ViewIds.Num();
 
 	// TODO: Remove this when everything is properly RDG'd
@@ -541,14 +541,14 @@ void FInstanceCullingContext::BuildRenderingCommands(FRDGBuilder& GraphBuilder,
 	PassParameters->GPUScenePrimitiveSceneData = GPUScene.PrimitiveBuffer.SRV;
 	PassParameters->InstanceDataSOAStride = GPUScene.InstanceDataSOAStride;
 	// Upload data etc
-	Params.PrimitiveCullingCommands = CreateStructuredBuffer(GraphBuilder, TEXT("PrimitiveCullingCommands"), CullingCommands);
+	Params.PrimitiveCullingCommands = CreateStructuredBuffer(GraphBuilder, TEXT("InstanceCulling.PrimitiveCullingCommands"), CullingCommands);
 	PassParameters->PrimitiveCullingCommands = GraphBuilder.CreateSRV(Params.PrimitiveCullingCommands);
 
-	PassParameters->PrimitiveIds = GraphBuilder.CreateSRV(CreateStructuredBuffer(GraphBuilder, TEXT("PrimitiveIds"), PrimitiveIds));
+	PassParameters->PrimitiveIds = GraphBuilder.CreateSRV(CreateStructuredBuffer(GraphBuilder, TEXT("InstanceCulling.PrimitiveIds"), PrimitiveIds));
 	PassParameters->DynamicPrimitiveIdOffset = DynamicPrimitiveIdRange.GetLowerBoundValue();
 	PassParameters->DynamicPrimitiveIdMax = DynamicPrimitiveIdRange.GetUpperBoundValue();
 
-	PassParameters->InstanceRuns = GraphBuilder.CreateSRV(CreateStructuredBuffer(GraphBuilder, TEXT("InstanceRuns"), InstanceRuns));
+	PassParameters->InstanceRuns = GraphBuilder.CreateSRV(CreateStructuredBuffer(GraphBuilder, TEXT("InstanceCulling.InstanceRuns"), InstanceRuns));
 
 
 	FRDGBufferRef InstanceIdOutOffsetBufferRDG = Intermediate.InstanceIdOutOffsetBuffer;
@@ -559,17 +559,17 @@ void FInstanceCullingContext::BuildRenderingCommands(FRDGBuilder& GraphBuilder,
 	{
 		TArray<uint32> NullArray;
 		NullArray.AddZeroed(1);
-		Params.InstanceIdWriteOffsetBuffer = CreateStructuredBuffer(GraphBuilder, TEXT("InstanceIdWriteOffsetBuffer"), NullArray);
+		Params.InstanceIdWriteOffsetBuffer = CreateStructuredBuffer(GraphBuilder, TEXT("InstanceCulling.InstanceIdWriteOffsetBuffer"), NullArray);
 	}
 
 	PassParameters->OutputOffsetBufferOut = GraphBuilder.CreateUAV(Params.InstanceIdWriteOffsetBuffer);
 
-	Params.DrawIndirectArgs = GraphBuilder.CreateBuffer(FRDGBufferDesc::CreateIndirectDesc(IndirectArgsNumWords * CullingCommands.Num()), TEXT("DrawIndirectArgsBuffer"));
+	Params.DrawIndirectArgs = GraphBuilder.CreateBuffer(FRDGBufferDesc::CreateIndirectDesc(IndirectArgsNumWords * CullingCommands.Num()), TEXT("InstanceCulling.DrawIndirectArgsBuffer"));
 	// not using structured buffer as we want/have to get at it as a vertex buffer 
 	//FRDGBufferRef InstanceIdsBufferRDG = GraphBuilder.CreateBuffer(FRDGBufferDesc::CreateBufferDesc(sizeof(uint32), PrimitiveIds.Num() * FInstanceCullingManager::MaxAverageInstanceFactor), TEXT("InstanceIdsBuffer"));
-	Params.InstanceIdStartOffsetBuffer = GraphBuilder.CreateBuffer(FRDGBufferDesc::CreateBufferDesc(sizeof(uint32), CullingCommands.Num()), TEXT("InstanceIdOffsetBuffer"));
+	Params.InstanceIdStartOffsetBuffer = GraphBuilder.CreateBuffer(FRDGBufferDesc::CreateBufferDesc(sizeof(uint32), CullingCommands.Num()), TEXT("InstanceCulling.InstanceIdOffsetBuffer"));
 
-	PassParameters->ViewIds = GraphBuilder.CreateSRV(CreateStructuredBuffer(GraphBuilder, TEXT("ViewIds"), ViewIds));
+	PassParameters->ViewIds = GraphBuilder.CreateSRV(CreateStructuredBuffer(GraphBuilder, TEXT("InstanceCulling.ViewIds"), ViewIds));
 	PassParameters->NumViewIds = ViewIds.Num();
 
 
@@ -586,8 +586,8 @@ void FInstanceCullingContext::BuildRenderingCommands(FRDGBuilder& GraphBuilder,
 	{
 		// TODO: we could compute the max instance count from the MDCs.
 		const int32 InstanceIdBufferSize = CullingCommands.Num() * FInstanceCullingManager::MaxAverageInstanceFactor * 64;
-		Params.InstanceIdsBuffer = GraphBuilder.CreateBuffer(FRDGBufferDesc::CreateBufferDesc(sizeof(uint32), InstanceIdBufferSize), TEXT("InstanceIdsBuffer"));;
-		Params.DrawCommandIdsBuffer = GraphBuilder.CreateBuffer(FRDGBufferDesc::CreateBufferDesc(sizeof(uint32), InstanceIdBufferSize), TEXT("DrawCommandIdsBuffer"));;
+		Params.InstanceIdsBuffer = GraphBuilder.CreateBuffer(FRDGBufferDesc::CreateBufferDesc(sizeof(uint32), InstanceIdBufferSize), TEXT("InstanceCulling.InstanceIdsBuffer"));
+		Params.DrawCommandIdsBuffer = GraphBuilder.CreateBuffer(FRDGBufferDesc::CreateBufferDesc(sizeof(uint32), InstanceIdBufferSize), TEXT("InstanceCulling.DrawCommandIdsBuffer"));
 
 	}
 	
diff --git a/Engine/Source/Runtime/Renderer/Private/InstanceCulling/InstanceCullingManager.cpp b/Engine/Source/Runtime/Renderer/Private/InstanceCulling/InstanceCullingManager.cpp
--- a/Engine/Source/Runtime/Renderer/Private/InstanceCulling/InstanceCullingManager.cpp
+++ b/Engine/Source/Runtime/Renderer/Private/InstanceCulling/InstanceCullingManager.cpp
@@ -109,7 +109,7 @@ void FInstanceCullingManager::CullInstances(FRDGBuilder& GraphBuilder, FGPUScene
 	TArray<uint32> NullArray;
 	NullArray.AddZeroed(1);
 
-	CullingIntermediate.InstanceIdOutOffsetBuffer = CreateStructuredBuffer(GraphBuilder, TEXT("OutputOffsetBufferOut"), NullArray);
+	CullingIntermediate.InstanceIdOutOffsetBuffer = CreateStructuredBuffer(GraphBuilder, TEXT("InstanceCulling.OutputOffsetBufferOut"), NullArray);
 
 	int32 NumViews = CullingViews.Num();
 	int32 NumInstances = GPUScene.InstanceDataAllocator.GetMaxSize();
@@ -121,7 +121,7 @@ void FInstanceCullingManager::CullInstances(FRDGBuilder& GraphBuilder, FGPUScene
 	if (NumInstances && NumViews)
 	{
 		// Create a buffer to record one bit for each instance per view
-		CullingIntermediate.VisibleInstanceFlags = GraphBuilder.CreateBuffer(FRDGBufferDesc::CreateStructuredDesc(sizeof(uint32), NumInstanceFlagWords * NumViews), TEXT("VisibleInstanceFlags"));
+		CullingIntermediate.VisibleInstanceFlags = GraphBuilder.CreateBuffer(FRDGBufferDesc::CreateStructuredDesc(sizeof(uint32), NumInstanceFlagWords * NumViews), TEXT("InstanceCulling.VisibleInstanceFlags"));
 		FRDGBufferUAVRef VisibleInstanceFlagsUAV = GraphBuilder.CreateUAV(CullingIntermediate.VisibleInstanceFlags);
 
 		if (CVarCullInstances.GetValueOnRenderThread() != 0)
@@ -136,7 +136,7 @@ void FInstanceCullingManager::CullInstances(FRDGBuilder& GraphBuilder, FGPUScene
 			PassParameters->NumInstances = NumInstances;
 			PassParameters->NumInstanceFlagWords = NumInstanceFlagWords;
 
-			PassParameters->InViews = GraphBuilder.CreateSRV(CreateStructuredBuffer(GraphBuilder, TEXT("CullingViews"), CullingViews));
+			PassParameters->InViews = GraphBuilder.CreateSRV(CreateStructuredBuffer(GraphBuilder, TEXT("InstanceCulling.CullingViews"), CullingViews));
 			PassParameters->NumViews = NumViews;
 
 
