Fix VT page table memory stat not being updated after first alloc for each space.
Added AVT indirection texture into page table memory stat.

diff --git a/Engine/Source/Runtime/Renderer/Private/VT/VirtualTextureSpace.cpp b/Engine/Source/Runtime/Renderer/Private/VT/VirtualTextureSpace.cpp
--- a/Engine/Source/Runtime/Renderer/Private/VT/VirtualTextureSpace.cpp
+++ b/Engine/Source/Runtime/Renderer/Private/VT/VirtualTextureSpace.cpp
@@ -142,6 +142,9 @@ uint32 FVirtualTextureSpace::GetSizeInBytes() const
 		const SIZE_T TextureSize = CalculateImageBytes(PageTableSize, PageTableSize, 0, TexturePixelFormat[TextureIndex]);
 		TotalSize += TextureSize;
 	}
+	
+	TotalSize += CalculateImageBytes(Description.IndirectionTextureSize, Description.IndirectionTextureSize, 0, PF_R32_UINT);
+
 	return TotalSize;
 }
 
diff --git a/Engine/Source/Runtime/Renderer/Private/VT/VirtualTextureSystem.cpp b/Engine/Source/Runtime/Renderer/Private/VT/VirtualTextureSystem.cpp
--- a/Engine/Source/Runtime/Renderer/Private/VT/VirtualTextureSystem.cpp
+++ b/Engine/Source/Runtime/Renderer/Private/VT/VirtualTextureSystem.cpp
@@ -625,9 +625,13 @@ FVirtualTextureSpace* FVirtualTextureSystem::AcquireSpace(const FVTSpaceDescript
 				FVirtualTextureSpace* Space = Spaces[SpaceIndex].Get();
 				if (Space && Space->GetDescription() == InDesc)
 				{
+					const int32 PagetableMemory = Space->GetSizeInBytes();
 					const uint32 vAddress = Space->AllocateVirtualTexture(AllocatedVT);
 					if (vAddress != ~0u)
 					{
+						const int32 NewPagetableMemory = Space->GetSizeInBytes();
+						INC_MEMORY_STAT_BY(STAT_TotalPagetableMemory, NewPagetableMemory - PagetableMemory);
+
 						AllocatedVT->AssignVirtualAddress(vAddress);
 						Space->AddRef();
 						return Space;
