Always use compute shader for VSM projection
- Adaptive ray counts for SMRT now work even with pass per light
- Removes a few divergences between the logic in the two paths

Remove debug projection stuff for now; will reimplement with a separate UAV write later

Remove temporal filtering for denoiser (it was broken and would need to be implemented slightly differently with the new path anyways)

Split cvar to allow separate denoiser toggle for directional (default on) and local (default off) lights

Enable SMRT for local lights and VSM caching by default

diff --git a/Engine/Shaders/Private/ClusteredDeferredShadingPixelShader.usf b/Engine/Shaders/Private/ClusteredDeferredShadingPixelShader.usf
--- a/Engine/Shaders/Private/ClusteredDeferredShadingPixelShader.usf
+++ b/Engine/Shaders/Private/ClusteredDeferredShadingPixelShader.usf
@@ -112,7 +112,12 @@ float4 GetLightGridLocalLighting(const FCulledLightsGridData InLightGridData, fl
 		float4 LightAttenuation = 1.0f;
 
 		LightData.ShadowedBits = 1;
-		LightAttenuation = ( ( ShadowMask >> (LocalLightListIndex * 3) ) & 7 ) / 7.0;
+		// Be careful with the shift here as the shift amount is mod 32, so >>32 is a noop rather than zero
+		// We treat lights past the 10th as unshadowed
+		if (LocalLightListIndex < 10)
+		{
+			LightAttenuation = ( ( ShadowMask >> (LocalLightListIndex * 3) ) & 7 ) / 7.0;
+		}
 
 #if 1
 		FRectTexture RectTexture = InitRectTexture(ForwardLightData.DummyRectLightSourceTexture);
diff --git a/Engine/Shaders/Private/VirtualShadowMaps/PageManagement.usf b/Engine/Shaders/Private/VirtualShadowMaps/PageManagement.usf
--- a/Engine/Shaders/Private/VirtualShadowMaps/PageManagement.usf
+++ b/Engine/Shaders/Private/VirtualShadowMaps/PageManagement.usf
@@ -216,8 +216,6 @@ void GeneratePageFlagsFromPixels(uint3 GroupId : SV_GroupID, uint GroupIndex : S
 	// Directional lights
 	for (uint Index = 0; Index < NumDirectionalLightSmInds; ++Index)
 	{
-		// Mark one page per directional light. Lights can have varying numbers of clipmap levels allocated which
-		// will all be contiguous in the array, so we hop between them using the counts.
 		int ClipmapStartId = VirtualShadowMapIdRemap[Index];
 		FVirtualShadowMapProjectionShaderData BaseProjectionData = GetVirtualShadowMapProjectionData(ClipmapStartId);
 
diff --git a/Engine/Shaders/Private/VirtualShadowMaps/ProjectionDirectional.ush b/Engine/Shaders/Private/VirtualShadowMaps/ProjectionDirectional.ush
--- a/Engine/Shaders/Private/VirtualShadowMaps/ProjectionDirectional.ush
+++ b/Engine/Shaders/Private/VirtualShadowMaps/ProjectionDirectional.ush
@@ -151,15 +151,14 @@ float3 GetRandomDirectionalLightRayDir(FLightShaderParameters Light, float2 E)
 float TraceDirectional(
 	int VirtualShadowMapId,
 	FLightShaderParameters Light,
-	float4 SvPosition,
+	uint2 PixelPos,
 	const float SceneDepth,
 	float3 TranslatedWorldPosition,
 	float RayStartOffset,
-	int RayCount,
-	int SamplesPerRay,
+	uint RayCount,
+	uint SamplesPerRay,
 	float RayLengthScale,
-	int DebugOutputType,
-	inout float3 DebugOutput,
+	const half Noise,
 	inout float InOutHitDistance)
 {
 	float3 WorldPosition = TranslatedWorldPosition - View.PreViewTranslation;
@@ -173,13 +172,14 @@ float TraceDirectional(
 	// Too low values will greatly restrict how large penumbras can be in screen space.
 	float RayLength = RayLengthScale * DistanceFromViewOrigin;
 
-	float StepOffset = InterleavedGradientNoise(SvPosition.xy, View.StateFrameIndexMod8);
+	float StepOffset = Noise;
 
-	float RayMissFactor = 0;
-	for (int i = 0; i < RayCount; ++i)
+	uint RayMissCount = 0;
+	uint i = 0;
+	for ( ; i < RayCount; i++)
 	{
 		RandomSequence RandSequence;
-		RandomSequence_Initialize(RandSequence, uint2(SvPosition.xy), i, View.StateFrameIndexMod8, RayCount);
+		RandomSequence_Initialize(RandSequence, PixelPos, i, View.StateFrameIndexMod8, RayCount);
 #if 1
 		float2 E = RandomSequence_GenerateSample2D(RandSequence);
 #else
@@ -198,24 +198,43 @@ float TraceDirectional(
 			SamplesPerRay,
 			StepOffset);
 
-		float OutputDebugSignal = Result.bValidHit ? 0.0f : 1.0f;
-		if (DebugOutputType == 1)	// Invalid page
-		{
-			DebugOutput = lerp(
-				OutputDebugSignal.xxx,
-				lerp( float3( 0, 1, 0 ), float3( 1, 0, 1 ), Result.SampleMissCount / 4.0f ),
-				0.75f );
-		}
-
 		if (Result.bValidHit)
 		{
 			// TODO: min over all the rays if we keep denoiser with > 1 spp?
 			InOutHitDistance = Result.OccluderDistance;
 		}
 
-		RayMissFactor += Result.bValidHit ? 0 : 1;
-	}
-	RayMissFactor /= float(RayCount);
+		RayMissCount += Result.bValidHit ? 0 : 1;
 
+#if SMRT_ADAPTIVE_RAY_COUNT && COMPILER_SUPPORTS_WAVE_VOTE
+		if( i == 0 )
+		{
+			bool bHit = Result.bValidHit;
+			if( bHit )
+			{
+				// TODO: Some simple heuristic for directional lights
+				//bHit = Result.OccluderDistance * Light.SourceRadius * ( 0.25 * View.ViewSizeAndInvSize.y / ( SceneDepth * View.ClipToView[1][1] ) ) > max( DistToLight - Result.OccluderDistance, 0.001 );
+			}
+
+			// All lanes miss or hit something close enough to not be noisy.
+			bool bAllLanesMiss = WaveActiveAllTrue( !bHit );
+			if( bAllLanesMiss )
+			{
+				break;
+			}
+		}
+		else if( i == 1 )
+		{
+			// After 2 iterations and all have hit, must be in umbra.
+			bool bAllLanesHit = WaveActiveAllTrue( RayMissCount == 0 );
+			if( bAllLanesHit )
+			{
+				break;
+			}
+		}
+#endif
+	}
+	uint SampleCount = min(i + 1, RayCount);		// break vs regular for loop exit
+	float RayMissFactor = float(RayMissCount) / float(SampleCount);
 	return RayMissFactor;
 }
diff --git a/Engine/Shaders/Private/VirtualShadowMaps/ProjectionSpot.ush b/Engine/Shaders/Private/VirtualShadowMaps/ProjectionSpot.ush
--- a/Engine/Shaders/Private/VirtualShadowMaps/ProjectionSpot.ush
+++ b/Engine/Shaders/Private/VirtualShadowMaps/ProjectionSpot.ush
@@ -185,7 +185,7 @@ FSMRTResult ShadowRayCastPointLight(
 float4 SobolSampler(uint SampleIndex, inout uint Seed, uint MaxBits)
 {
 	// first scramble the index to decorelate from other 4-tuples
-	uint SobolIndex = FastOwenScrambling(SampleIndex, EvolveSobolSeed(Seed)) & ( (1 << MaxBits) - 1 );
+	uint SobolIndex = FastOwenScrambling(SampleIndex, EvolveSobolSeed(Seed)) & ( (1U << MaxBits) - 1U );
 	// now get Sobol' point from this index
 	uint4 Result = uint4(SobolIndex, 0, 0, 0);
 	UNROLL for (uint b = 0, v = 1; b < MaxBits; b++)
@@ -213,19 +213,16 @@ float TraceLocalLight(
 	int VirtualShadowMapId,
 	FLightShaderParameters Light,
 	uint LightType,
-	float4 SvPosition,
+	uint2 PixelPos,
 	const float SceneDepth,
 	float3 TranslatedWorldPosition,
 	float RayStartOffset,
-	int RayCount,
-	int SamplesPerRay,
+	uint RayCount,
+	uint SamplesPerRay,
 	float CotMaxRayAngleFromLight,
-	int DebugOutputType,
-	inout float3 DebugOutput,
+	const half Noise,
 	inout float InOutHitDistance)
 {
-	uint2 PixelPos = uint2( SvPosition.xy );
-
 	const half3 N = GetGBufferDataUint( PixelPos, true ).WorldNormal;
 	const float3 WorldPosition = TranslatedWorldPosition - View.PreViewTranslation;
 
@@ -243,10 +240,9 @@ float TraceLocalLight(
 	bool bBackFace = dot( N, ConeAxis ) < -SinAlpha;
 
 	half RayMissFactor = 0;
-	uint SampleCount = 0;
 	if (!bBackFace)
 	{
-		float StepOffset = InterleavedGradientNoise(SvPosition.xy, View.StateFrameIndexMod8);
+		float StepOffset = Noise;
 
 		half ConeSin = Light.SourceRadius * InvDist;
 		half ConeCos = sqrt( 1 - Pow2( ConeSin ) );
@@ -254,7 +250,7 @@ float TraceLocalLight(
 		uint TileID = MortonEncode( PixelPos );
 
 		uint i = 0;
-		for( ; i < uint(RayCount); i++ )
+		for( ; i < RayCount; i++ )
 		{
 			RandomSequence RandSequence;
 #if 0
@@ -310,15 +306,6 @@ float TraceLocalLight(
 					StepOffset);
 			}
 
-			float OutputDebugSignal = Result.bValidHit ? 0.0f : 1.0f;
-			if (DebugOutputType == 1)	// Invalid page
-			{
-				DebugOutput = lerp(	
-					OutputDebugSignal.xxx,
-					lerp( float3( 0, 1, 0 ), float3( 1, 0, 1 ), Result.SampleMissCount / 4.0f ),
-					0.75f );
-			}
-
 			if ( Result.bValidHit )
 			{
 				// TODO: Min over all the rays with > 1 spp?
@@ -364,15 +351,9 @@ float TraceLocalLight(
 			uint Seed = asuint( StepOffset );
 			StepOffset = ( EvolveSobolSeed( Seed ) >> 8 ) * 5.96046447754e-08; // * 2^-24;
 		}
-		SampleCount = min(i + 1, uint(RayCount));		// break vs regular for loop exit
+		uint SampleCount = min(i + 1, RayCount);		// break vs regular for loop exit
 		RayMissFactor /= float(SampleCount);
 	}
 
-	if (DebugOutputType == 2)	// Sample count
-	{
-		float Ratio = float(SampleCount) / float(RayCount);
-		DebugOutput = lerp(RayMissFactor.xxx, lerp(float3(0, 1, 0), float3(1, 0, 0), Ratio), 1.0f);
-	}
-
 	return RayMissFactor;
 }
diff --git a/Engine/Shaders/Private/VirtualShadowMaps/VirtualShadowMapProjection.usf b/Engine/Shaders/Private/VirtualShadowMaps/VirtualShadowMapProjection.usf
--- a/Engine/Shaders/Private/VirtualShadowMaps/VirtualShadowMapProjection.usf
+++ b/Engine/Shaders/Private/VirtualShadowMaps/VirtualShadowMapProjection.usf
@@ -21,26 +21,47 @@
 #include "ProjectionDirectional.ush"
 #include "ProjectionSpot.ush"
 
-// See VirtualShadowMapProjection.cpp
-// 0 = Output to denoiser
-// 1 = Output to screen shadow mask
-// 2 = Output debug
-#ifndef OUTPUT_TYPE
-#define OUTPUT_TYPE 0
+#ifndef DIRECTIONAL_LIGHT
+#define DIRECTIONAL_LIGHT 0
 #endif
 
-int VirtualShadowMapId;		// For the first layer for clipmaps
-int DebugOutputType;
-float ContactShadowLength;
-float NormalOffsetWorld;	// Normal offset in world space
+#define CONTACT_SHADOW_SAMPLES 8
 
-int SMRTRayCount;	// 0 = off
-int SMRTSamplesPerRay;
-float SMRTRayLengthScale;			// Directional lights
-float SMRTCotMaxRayAngleFromLight;	// Spot/point lights
+float4 EncodeLightAttenuationFromMask(float ShadowMask)
+{
+	const float ShadowFadeFraction = 1;
+	float SSSTransmission = ShadowMask;
 
-#define CONTACT_SHADOW_SAMPLES 8
+	// 0 is shadowed, 1 is unshadowed
+	// RETURN_COLOR not needed unless writing to SceneColor;
+	//float FadedShadow = lerp(1.0f, Square(ShadowMask), ShadowFadeFraction);
+	//float FadedSSSShadow = lerp(1.0f, Square(SSSTransmission), ShadowFadeFraction);
+	float FadedShadow = lerp(1.0f, ShadowMask, ShadowFadeFraction);
+	float FadedSSSShadow = lerp(1.0f, SSSTransmission, ShadowFadeFraction);
 
+	// the channel assignment is documented in ShadowRendering.cpp (look for Light Attenuation channel assignment)
+	return EncodeLightAttenuation(half4(FadedShadow, FadedSSSShadow, FadedShadow, FadedSSSShadow));
+}
+
+Texture2D<float4> InputSignal;
+
+void VirtualShadowMapCompositePS(
+	in float4 SvPosition : SV_Position,
+	out float4 OutShadowMask : SV_Target
+	)
+{
+	// NOTE: The signal is encoded as per SSDSignalBufferEncoding.ush, SIGNAL_BUFFER_LAYOUT_PENUMBRA_RECONSTRUCTION
+	float4 Input = InputSignal.Load(int3(SvPosition.xy, 0));
+	float SampleCount = Input.g;
+	float ShadowFactor = Input.r;
+	float Shadow = (SampleCount > 0 ? ShadowFactor : 1.0f);
+	
+	OutShadowMask = EncodeLightAttenuationFromMask(Shadow);
+}
+
+
+SCHEDULER_MIN_PRESSURE
+MAX_OCCUPENCY
 
 // Returns hit distance or negative if miss
 float VirtualShadowMapScreenRayCast(
@@ -61,7 +82,6 @@ float ProjectSingleSample(
 	float3 TranslatedWorldPosition,
 	float3 EstimatedGeoWorldNormal,
 	float ContactShadowLengthWorld,
-	inout float3 DebugOutput,
 	inout float InOutHitDistance)
 {
 	float3 WorldPosition = TranslatedWorldPosition - View.PreViewTranslation;
@@ -74,77 +94,29 @@ float ProjectSingleSample(
 		InOutHitDistance = Sample.OccluderDistance;
 	}
 
-	if (DebugOutputType == 1)	// Invalid page
-	{
-		DebugOutput = Sample.bValid ? (0.25f * Sample.ShadowFactor) : float3(1, 0, 1);
-	}
-	else if (DebugOutputType == 2)	// ClipmapIndex/Mip
-	{
-		DebugOutput = lerp(Sample.ShadowFactor, VirtualShadowMapDebugIntToColor(Sample.ClipmapIndexOrMipLevel), 0.75f);
-	}
-	else if (DebugOutputType == 3) // Occluder distance
-	{
-		DebugOutput = lerp(Sample.ShadowFactor, InOutHitDistance.xxx / 1000.0f, 0.75f);
-	}
-
 	return Sample.ShadowFactor;
 }
 
-
-float4 EncodeLightAttenuationFromMask(float ShadowMask)
-{
-	const float ShadowFadeFraction = 1;
-	float SSSTransmission = ShadowMask;
-
-	// 0 is shadowed, 1 is unshadowed
-	// RETURN_COLOR not needed unless writing to SceneColor;
-	//float FadedShadow = lerp(1.0f, Square(ShadowMask), ShadowFadeFraction);
-	//float FadedSSSShadow = lerp(1.0f, Square(SSSTransmission), ShadowFadeFraction);
-	float FadedShadow = lerp(1.0f, ShadowMask, ShadowFadeFraction);
-	float FadedSSSShadow = lerp(1.0f, SSSTransmission, ShadowFadeFraction);
-
-	// the channel assignment is documented in ShadowRendering.cpp (look for Light Attenuation channel assignment)
-	return EncodeLightAttenuation(half4(FadedShadow, FadedSSSShadow, FadedShadow, FadedSSSShadow));
-}
-
 float4 EncodeOutput(
 	float3 TranslatedWorldPosition,
 	FLightShaderParameters Light,
 	uint LightType,
 	float RayMissFactor,
 	float HitDistance,
-	int SampleCount,
-	float3 DebugOutput)
+	int SampleCount)
 {
 	// See signal encoding for SIGNAL_BUFFER_LAYOUT_PENUMBRA_RECONSTRUCTION in SSDSignalBufferEncoding.ush
 	// We fuse that logic in here to avoid needing another pass, but it's hard to call the functions directly as they depend on a large
 	// number of preprocessor defines and logic that we don't want to mirror here.
 	// Unfortunately that means this logic can get out of date when the denoiser changes, but that's the trade-off for now.
-	float4 OutputSignal = float4(0, 0, 0, 0);
+	float4 OutputSignal = float4(RayMissFactor, SampleCount, WORLD_RADIUS_MISS, RayMissFactor);
 	if (RayMissFactor < 1.0f)
 	{
-		float WorldBluringRadius = ComputeLightSampleWorldBluringRadius(TranslatedWorldPosition, LightType, Light, HitDistance);
-		OutputSignal = float4(RayMissFactor, SampleCount, WorldBluringRadius, RayMissFactor);
+		OutputSignal.z = ComputeLightSampleWorldBluringRadius(TranslatedWorldPosition, LightType, Light, HitDistance);
 	}
-	else
-	{
-		// Miss
-		OutputSignal = float4(1, SampleCount, WORLD_RADIUS_MISS, 1);
-	}
-
-#if OUTPUT_TYPE == 0
-	// Output feeds into denoiser
 	return OutputSignal;
-#elif OUTPUT_TYPE == 1
-	// Output to screen shadow mask
-	return EncodeLightAttenuationFromMask(OutputSignal.r);
-#elif OUTPUT_TYPE == 2
-	// Output debug for "vis".
-	return float4(DebugOutput, 1.0f);
-#endif // OUTPUT_TYPE
 }
 
-
 float3 GetEstimatedGeoWorldNormal(float3 TranslatedWorldPosition, float2 ScreenUV)
 {
 	// Figure out slope, we do world space since that is the space where we might override using the shading normal...
@@ -170,217 +142,155 @@ float3 GetEstimatedGeoWorldNormal(float3 TranslatedWorldPosition, float2 ScreenU
 	return EstimatedGeoWorldNormal;
 }
 
-
-void VirtualShadowMapProjectionSpotPS(
-	in float4 SvPosition : SV_Position,
-	out float4 OutputTarget : SV_Target)
+FLightShaderParameters ConvertFromLocal( const FLocalLightData LightData )
 {
-	// TODO: Unify this with the logic in PageManagement.usf - we need to get the same values in both passes
-	const float2 ScreenUV = SvPositionToBufferUV(SvPosition);
-	const float DeviceZ = LookupDeviceZ(ScreenUV);
-	const float SceneDepth = ConvertFromDeviceZ(DeviceZ);	
-	float3 TranslatedWorldPosition = SvPositionToTranslatedWorld(float4(SvPosition.xy, DeviceZ, 1.0f));
-	
-	FLightShaderParameters Light = GetRootLightShaderParameters();
-
-	// Early out if we're outside the spot light cone		
-	float3 ToLight = Light.Position - (TranslatedWorldPosition - View.PreViewTranslation);
-	bool bInLightRadius = length(ToLight) <= rcp(Light_InvRadius);
-	float3 L = normalize(ToLight);
-	bool bInLightCone = dot(L, Light.Direction) >= Light.SpotAngles.x;
-	bool bInLightRegion = bInLightRadius && bInLightCone;
-
-	float3 DebugOutput = float3(0, 0, 0);
-	int SampleCount = 1;
-	float HitDistance = WORLD_RADIUS_MISS;
-	float RayMissFactor = 0.0f;
-
-	uint LightType = Light.SpotAngles.x > -2.0f ? LIGHT_TYPE_SPOT : LIGHT_TYPE_POINT;
-
-	if (bInLightRadius && bInLightCone)
-	{
-		const float ContactShadowLengthWorld = ContactShadowLength * View.ClipToView[1][1] * SceneDepth;
+	FLightShaderParameters Light = (FLightShaderParameters)0;
+	Light.Position			= LightData.LightPositionAndInvRadius.xyz;
+	Light.InvRadius			= LightData.LightPositionAndInvRadius.w;
+	Light.Color				= LightData.LightColorAndFalloffExponent.xyz;
+	Light.FalloffExponent	= LightData.LightColorAndFalloffExponent.w;
+	Light.Direction			= LightData.LightDirectionAndShadowMask.xyz;
+	Light.Tangent			= LightData.LightTangentAndSoftSourceRadius.xyz;
+	Light.SpotAngles		= LightData.SpotAnglesAndSourceRadiusPacked.xy;
+	Light.SpecularScale		= 1;
+	Light.SourceRadius		= LightData.SpotAnglesAndSourceRadiusPacked.z;
+	Light.SoftSourceRadius	= LightData.LightTangentAndSoftSourceRadius.w;
+	Light.SourceLength		= f16tof32(asuint(LightData.SpotAnglesAndSourceRadiusPacked.w));
+	return Light;
+}
 
-		float3 EstimatedGeoWorldNormal = GetEstimatedGeoWorldNormal(TranslatedWorldPosition, ScreenUV);
 
-		if (SMRTRayCount > 0)
-		{
-			RayMissFactor = TraceLocalLight(
-				VirtualShadowMapId,
-				Light,
-				LightType,
-				SvPosition,
-				SceneDepth,
-				TranslatedWorldPosition + EstimatedGeoWorldNormal * NormalOffsetWorld,
-				ContactShadowLengthWorld,
-				SMRTRayCount,
-				SMRTSamplesPerRay,
-				SMRTCotMaxRayAngleFromLight,
-				DebugOutputType,
-				DebugOutput,
-				HitDistance);
-		}
-		else
-		{
-			RayMissFactor = ProjectSingleSample(
-				VirtualShadowMapId,
-				L,
-				TranslatedWorldPosition + EstimatedGeoWorldNormal * NormalOffsetWorld,
-				EstimatedGeoWorldNormal,
-				ContactShadowLengthWorld,
-				/* out */ DebugOutput,
-				HitDistance);
-		}
+int4 ProjectionRect;
+float ContactShadowLength;
+float NormalOffsetWorld;			// Normal offset in world space
+int SMRTRayCount;					// 0 = off
+int SMRTSamplesPerRay;
+float SMRTRayLengthScale;			// Directional lights
+float SMRTCotMaxRayAngleFromLight;	// Spot/point lights
+// HAIR_TODO: move all data into a HairStrands UB
+uint bUseHairData;
+Texture2D<uint4> HairCategorizationTexture;
 
-		if (ContactShadowLengthWorld > 0.0f)
-		{
-			float StepOffset = InterleavedGradientNoise(SvPosition.xy, View.StateFrameIndexMod8);
-			float ContactShadowHitDistance = VirtualShadowMapScreenRayCast(TranslatedWorldPosition, L, ContactShadowLengthWorld, CONTACT_SHADOW_SAMPLES, StepOffset);		
-			if (ContactShadowHitDistance >= 0.0f)
-			{
-				RayMissFactor = 0.0f;
-				HitDistance = ContactShadowHitDistance;
-			}
-		}
-	}
+// One pass projection
+StructuredBuffer< int > VirtualShadowMapIdRemap;
+uint NumDirectionalLightSmInds;
+RWTexture2D< uint > RWShadowMaskBits;
 
-	OutputTarget = EncodeOutput(TranslatedWorldPosition, Light, LightType, RayMissFactor, HitDistance, SampleCount, DebugOutput);
-}
+// Single light per pass
+// Light parameters loaded via GetRootLightShaderParameters();
+int LightUniformVirtualShadowMapId;
+RWTexture2D< float4 > RWSignal;
 
-// Type of input data consume by the page allocation (i.e., data read from the source buffer: Gbuffer, HairStrands data, ...)
-#define INPUT_TYPE_GBUFFER 0
-#define INPUT_TYPE_HAIRSTRANDS 1
-uint InputType;
+struct FProjectLightResult
+{
+	half ShadowFactor;
+	// NOTE: The signal is encoded as per SSDSignalBufferEncoding.ush, SIGNAL_BUFFER_LAYOUT_PENUMBRA_RECONSTRUCTION
+	float4 Signal;
+};
 
-void VirtualShadowMapProjectionDirectionalPS(
-	in float4 SvPosition : SV_Position,
-	out float4 OutputTarget : SV_Target
-)
+FProjectLightResult ProjectLight(
+	int VirtualShadowMapId,
+	FLightShaderParameters Light,
+	uint2 PixelPos,
+	float SceneDepth,
+	float ContactShadowLengthWorld,
+	float3 TranslatedWorldPosition,
+	half3 GBufferNormal,
+	const half Noise)
 {
-	// TODO: Unify this with the logic in PageManagement.usf - we need to get the same values in both passes
-	const float2 ScreenUV = SvPositionToBufferUV(SvPosition);
-	float DeviceZ = LookupDeviceZ(ScreenUV);
-	if (InputType == INPUT_TYPE_HAIRSTRANDS)
-	{
-		FCategorizationData CatData = DecodeCategorizationData(HairStrands.HairCategorizationTexture.Load(uint3(SvPosition.xy, 0)));
-		DeviceZ = CatData.ClosestDepth;
-	}
-	const float SceneDepth = ConvertFromDeviceZ(DeviceZ);
-	float3 TranslatedWorldPosition = SvPositionToTranslatedWorld(float4(SvPosition.xy, DeviceZ, 1.0f));
+	const float3 WorldPosition = TranslatedWorldPosition - View.PreViewTranslation;
 
-	FLightShaderParameters Light = GetRootLightShaderParameters();
-		
+#if DIRECTIONAL_LIGHT
+	uint LightType = LIGHT_TYPE_DIRECTIONAL;
 	float3 L = Light.Direction;
+	// TODO: Find a way to branch out backfacing surfaces again that doesn't make assumptions
+	// about the Gbuffer normal that break with hair
+	bool bInLightRegion = true;
+#else
+	uint LightType = Light.SpotAngles.x > -2.0f ? LIGHT_TYPE_SPOT : LIGHT_TYPE_POINT;
+	float3 ToLight = Light.Position - WorldPosition;
+	float d2 = dot( ToLight, ToLight );
+	float InvDist = rsqrt( d2 );
+	half3 L = ToLight * InvDist;
 	
-	float3 DebugOutput = float3(0, 0, 0);
-	int SampleCount = 1;
+	bool bInLightRadius = InvDist >= Light.InvRadius;
+	bool bInLightCone = dot( L, Light.Direction ) >= Light.SpotAngles.x;
+	bool bInLightRegion = bInLightRadius && bInLightCone;
+	
+#endif // DIRECTIONAL_LIGHT
+
+	float ShadowFactor = 1.0f;
 	float HitDistance = WORLD_RADIUS_MISS;
-	float RayMissFactor = 0.0f;
+	uint SampleCount = 0;
 
-	// NOTE: This logic is unreliable for hair and potentially other materials with alternate gbuffer encodings	
-	// TODO: Revisit other ways to get performance back here.
-	//bool bFrontFace = dot(ToLightDir, GBufferData.WorldNormal) > 0.0f;
-	//if (bFrontFace)
+	BRANCH
+	if (bInLightRegion)
 	{
-		const float ContactShadowLengthWorld = ContactShadowLength * View.ClipToView[1][1] * SceneDepth;
+		const half3 EstimatedGeoWorldNormal = bUseHairData > 0 ? L : GBufferNormal;
 
-		float3 EstimatedGeoWorldNormal = GetEstimatedGeoWorldNormal(TranslatedWorldPosition, ScreenUV);
-		if (InputType == INPUT_TYPE_HAIRSTRANDS)
-		{			
-			EstimatedGeoWorldNormal = L;
-		}
+		float3 OffsetTranslatedWorldPosition = TranslatedWorldPosition + EstimatedGeoWorldNormal * NormalOffsetWorld;
 
 		if (SMRTRayCount > 0)
 		{
 			SampleCount = SMRTRayCount;
-			RayMissFactor = TraceDirectional(
+		#if DIRECTIONAL_LIGHT
+			ShadowFactor = TraceDirectional(
 				VirtualShadowMapId,
 				Light,
-				SvPosition,
+				PixelPos,
 				SceneDepth,
-				TranslatedWorldPosition + EstimatedGeoWorldNormal * NormalOffsetWorld,
+				OffsetTranslatedWorldPosition,
 				ContactShadowLengthWorld,
 				SMRTRayCount,
 				SMRTSamplesPerRay,
 				SMRTRayLengthScale,
-				DebugOutputType,
-				DebugOutput,
-				HitDistance);
+				Noise,
+				/* out */ HitDistance);
+		#else
+			ShadowFactor = TraceLocalLight(
+				VirtualShadowMapId,
+				Light,
+				LightType,
+				PixelPos,
+				SceneDepth,
+				OffsetTranslatedWorldPosition,
+				ContactShadowLengthWorld,
+				SMRTRayCount,
+				SMRTSamplesPerRay,
+				SMRTCotMaxRayAngleFromLight,
+				Noise,
+				/* out */ HitDistance);
+		#endif // DIRECTIONAL_LIGHT
 		}
 		else
 		{
 			SampleCount = 1;
-			RayMissFactor = ProjectSingleSample(
+			ShadowFactor = ProjectSingleSample(
 				VirtualShadowMapId,
 				L,
-				TranslatedWorldPosition + EstimatedGeoWorldNormal * NormalOffsetWorld,
+				OffsetTranslatedWorldPosition,
 				EstimatedGeoWorldNormal,
 				ContactShadowLengthWorld,
-				/* out */ DebugOutput,
 				HitDistance);
 		}
 
 		if (ContactShadowLengthWorld > 0.0f)
 		{
-			float StepOffset = InterleavedGradientNoise(SvPosition.xy, View.StateFrameIndexMod8);
-			float ContactShadowHitDistance = VirtualShadowMapScreenRayCast(TranslatedWorldPosition, L, ContactShadowLengthWorld, CONTACT_SHADOW_SAMPLES, StepOffset);		
+			half ContactShadowHitDistance = VirtualShadowMapScreenRayCast(TranslatedWorldPosition, L, ContactShadowLengthWorld, CONTACT_SHADOW_SAMPLES, Noise);		
 			if (ContactShadowHitDistance >= 0.0f)
 			{
-				RayMissFactor = 0.0f;
+				ShadowFactor = 0.0f;
 				HitDistance = ContactShadowHitDistance;
 			}
 		}
 	}
 
-	OutputTarget = EncodeOutput(TranslatedWorldPosition, Light, LIGHT_TYPE_DIRECTIONAL, RayMissFactor, HitDistance, SampleCount, DebugOutput);
+	FProjectLightResult Result;
+	Result.ShadowFactor = ShadowFactor;
+	Result.Signal = EncodeOutput(TranslatedWorldPosition, Light, LightType, ShadowFactor, HitDistance, SampleCount);
+	return Result;
 }
 
-
-
-
-
-
-Texture2D<float4> InputSignal;
-
-void VirtualShadowMapCompositePS(
-	in float4 SvPosition : SV_Position,
-	out float4 OutShadowMask : SV_Target
-	)
-{
-	// NOTE: The signal is encoded as per SSDSignalBufferEncoding.ush, SIGNAL_BUFFER_LAYOUT_PENUMBRA_RECONSTRUCTION
-	float4 Input = InputSignal.Load(int3(SvPosition.xy, 0));
-	float SampleCount = Input.g;
-	float MissFactor = Input.r;
-	float Shadow = (SampleCount > 0 ? MissFactor : 1.0f);
-	
-	OutShadowMask = EncodeLightAttenuationFromMask(Shadow);
-}
-
-SCHEDULER_MIN_PRESSURE
-MAX_OCCUPENCY
-
-FLightShaderParameters ConvertFromLocal( const FLocalLightData LightData )
-{
-	FLightShaderParameters Light = (FLightShaderParameters)0;
-	Light.Position			= LightData.LightPositionAndInvRadius.xyz;
-	Light.InvRadius			= LightData.LightPositionAndInvRadius.w;
-	Light.Color				= LightData.LightColorAndFalloffExponent.xyz;
-	Light.FalloffExponent	= LightData.LightColorAndFalloffExponent.w;
-	Light.Direction			= LightData.LightDirectionAndShadowMask.xyz;
-	Light.Tangent			= LightData.LightTangentAndSoftSourceRadius.xyz;
-	Light.SpotAngles		= LightData.SpotAnglesAndSourceRadiusPacked.xy;
-	Light.SpecularScale		= 1;
-	Light.SourceRadius		= LightData.SpotAnglesAndSourceRadiusPacked.z;
-	Light.SoftSourceRadius	= LightData.LightTangentAndSoftSourceRadius.w;
-	Light.SourceLength		= f16tof32(asuint(LightData.SpotAnglesAndSourceRadiusPacked.w));
-	return Light;
-}
-
-StructuredBuffer< int > VirtualShadowMapIdRemap;
-uint NumDirectionalLightSmInds;
-
-RWTexture2D< uint > RWShadowMaskBits;
-
 [numthreads(8, 8, 1)]
 void VirtualShadowMapProjection(
 	uint3	GroupId				: SV_GroupID,
@@ -388,106 +298,81 @@ void VirtualShadowMapProjection(
 	uint3	DispatchThreadId	: SV_DispatchThreadID )
 {
 	// Morton order within a group so page access/atomics are more coherent and wave-swizzled gradients are possible.
-	uint2 LocalPixelPos = 8 * GroupId.xy + MortonDecode(GroupIndex);
-	uint2 PixelPos = LocalPixelPos + uint2( View.ViewRectMin.xy );
-
-	if (any(LocalPixelPos >= uint2(View.ViewSizeAndInvSize.xy)))
+	uint2 LocalPixelPos = 8 * GroupId.xy + MortonDecode( GroupIndex );
+	uint2 PixelPos = LocalPixelPos + uint2( ProjectionRect.xy );
+	if ( any( PixelPos >= uint2( ProjectionRect.zw ) ) )
 	{
 		return;
 	}
 	
-	const float DeviceZ = SceneTexturesStruct.SceneDepthTexture.Load( int3( PixelPos, 0 ) ).r;
+	float DeviceZ = SceneTexturesStruct.SceneDepthTexture.Load( int3( PixelPos, 0 ) ).r;
+	if (bUseHairData > 0)
+	{
+		FCategorizationData CatData = DecodeCategorizationData( HairCategorizationTexture.Load( int3( PixelPos, 0 ) ) );
+		DeviceZ = CatData.ClosestDepth;
+	}
 	const float SceneDepth = ConvertFromDeviceZ( DeviceZ );
 
 	const float4 SvPosition = float4( float2( PixelPos ) + 0.5, DeviceZ, 1.0f );
-	//const float2 ScreenUV = SvPositionToBufferUV( SvPosition );
-
 	const float3 TranslatedWorldPosition = SvPositionToTranslatedWorld( SvPosition );
-	const float3 WorldPosition = TranslatedWorldPosition - View.PreViewTranslation;
 
 	const float ContactShadowLengthWorld = ContactShadowLength * View.ClipToView[1][1] * SceneDepth;
 	const half Noise = InterleavedGradientNoise( SvPosition.xy, View.StateFrameIndexMod8 );
 
-	//const half3 EstimatedGeoWorldNormal = GetEstimatedGeoWorldNormal(TranslatedWorldPosition, ScreenUV);
-	const half3 EstimatedGeoWorldNormal = GetGBufferDataUint( PixelPos, true ).WorldNormal;
-
-	uint ShadowMask = 0xffffffff;
-	
-	uint EyeIndex = 0; // ??
+	const half3 GBufferNormal = GetGBufferDataUint( PixelPos, true ).WorldNormal;
+		
+#if ONE_PASS_PROJECTION
+	uint EyeIndex = 0; // TODO: Instanced stereo
 	uint GridLinearIndex = ComputeLightGridCellIndex( LocalPixelPos, SceneDepth, EyeIndex );
 	const FCulledLightsGridData CulledLightGridData = GetCulledLightsGrid( GridLinearIndex, EyeIndex );
 
+	// We can only handle 10 lights in our output encoding right now, so no purpose in computing more
+	const uint LightCount = min( 10, CulledLightGridData.NumLocalLights );
+
+	uint ShadowMask = 0xffffffff;
+
 	LOOP
-	for (uint Index = 0; Index < CulledLightGridData.NumLocalLights; ++Index)
+	for (uint Index = 0; Index < LightCount; ++Index)
 	{
+		// The Virtual Shadow Remap stores directional lights first
 		const uint LocalLightIndex = ForwardLightData.CulledLightDataGrid[CulledLightGridData.DataStartIndex + Index];
-		const FLocalLightData LightData = GetLocalLightData(CulledLightGridData.DataStartIndex + Index, EyeIndex);
+		int VirtualShadowMapId = VirtualShadowMapIdRemap[VirtualShadowMap.NumDirectionalLights + LocalLightIndex];
 
-		// The Virtual Shadow Remap stores directional lights first
-		int VirtualShadowMapId = VirtualShadowMapIdRemap[NumDirectionalLightSmInds + LocalLightIndex];
 		if (VirtualShadowMapId != INDEX_NONE)
 		{
-			float3 ToLight = LightData.LightPositionAndInvRadius.xyz - WorldPosition;
-			float d2 = dot( ToLight, ToLight );
-			float InvDist = rsqrt( d2 );
-			half3 L = ToLight * InvDist;
-			
-			bool bInLightRadius = InvDist >= LightData.LightPositionAndInvRadius.w;
-			bool bInLightCone = dot( L, LightData.LightDirectionAndShadowMask.xyz ) >= LightData.SpotAnglesAndSourceRadiusPacked.x;
-
-			if (bInLightRadius && bInLightCone)
-			{
-				half RayMissFactor = 1.0f;
-				half HitDistance = WORLD_RADIUS_MISS;
-				float3 DebugOutput = float3(0, 0, 0);
-
-				FLightShaderParameters Light = ConvertFromLocal( LightData );
-
-				uint LightType = Light.SpotAngles.x > -2.0f ? LIGHT_TYPE_SPOT : LIGHT_TYPE_POINT;
-
-				if (SMRTRayCount > 0)
-				{
-					RayMissFactor = TraceLocalLight(
-						VirtualShadowMapId,
-						Light,
-						LightType,
-						SvPosition,
-						SceneDepth,
-						TranslatedWorldPosition + EstimatedGeoWorldNormal * NormalOffsetWorld,
-						ContactShadowLengthWorld,
-						SMRTRayCount,
-						SMRTSamplesPerRay,
-						SMRTCotMaxRayAngleFromLight,
-						DebugOutputType,
-						DebugOutput,
-						HitDistance);
-				}
-				else
-				{
-					RayMissFactor = ProjectSingleSample(
-						VirtualShadowMapId,
-						L,
-						TranslatedWorldPosition + EstimatedGeoWorldNormal * NormalOffsetWorld,
-						EstimatedGeoWorldNormal,
-						ContactShadowLengthWorld,
-						/* out */ DebugOutput,
-						HitDistance);
-				}
-
-				if (ContactShadowLengthWorld > 0.0f)
-				{
-					half ContactShadowHitDistance = VirtualShadowMapScreenRayCast(TranslatedWorldPosition, L, ContactShadowLengthWorld, CONTACT_SHADOW_SAMPLES, Noise);		
-					if (ContactShadowHitDistance >= 0.0f)
-					{
-						RayMissFactor = 0.0f;
-						HitDistance = ContactShadowHitDistance;
-					}
-				}
-
-				ShadowMask ^= ( ~uint( round( RayMissFactor * 7.0 ) ) & 7 ) << (Index*3);
-			}
+			const FLocalLightData LightData = GetLocalLightData( CulledLightGridData.DataStartIndex + Index, EyeIndex );
+			FLightShaderParameters Light = ConvertFromLocal( LightData );
+
+			FProjectLightResult Result = ProjectLight(
+				VirtualShadowMapId,
+				Light,
+				PixelPos,
+				SceneDepth,
+				ContactShadowLengthWorld,
+				TranslatedWorldPosition,
+				GBufferNormal,
+				Noise);
+
+			ShadowMask ^= ( ~uint( round( Result.ShadowFactor * 7.0 ) ) & 7 ) << (Index*3);
 		}
 	}
 
 	RWShadowMaskBits[ PixelPos ] = ~ShadowMask;
-}
\ No newline at end of file
+
+#else // !ONE_PASS_PROJECTION
+	int VirtualShadowMapId = LightUniformVirtualShadowMapId;
+	FLightShaderParameters Light = GetRootLightShaderParameters();
+
+	FProjectLightResult Result = ProjectLight(
+		VirtualShadowMapId,
+		Light,
+		PixelPos,
+		SceneDepth,
+		ContactShadowLengthWorld,
+		TranslatedWorldPosition,
+		GBufferNormal,
+		Noise);
+
+	RWSignal[ PixelPos ] = Result.Signal;
+#endif // ONE_PASS_PROJECTION
+}
diff --git a/Engine/Source/Runtime/Renderer/Private/LightRendering.cpp b/Engine/Source/Runtime/Renderer/Private/LightRendering.cpp
--- a/Engine/Source/Runtime/Renderer/Private/LightRendering.cpp
+++ b/Engine/Source/Runtime/Renderer/Private/LightRendering.cpp
@@ -1293,22 +1293,14 @@ void FDeferredShadingSceneRenderer::RenderLights(
 				FRDGTextureRef ShadowMaskBits = nullptr;
 				if( VirtualShadowMapArray.IsAllocated() && CVarVirtualShadowOnePassProjection.GetValueOnRenderThread() )
 				{
-					const FRDGTextureDesc ShadowMaskDesc = FRDGTextureDesc::Create2D(
-						SceneTextures.Config.Extent,
-						PF_R32_UINT,
-						FClearValueBinding::None,
-						TexCreate_ShaderResource | TexCreate_UAV );
-
-					ShadowMaskBits = GraphBuilder.CreateTexture( ShadowMaskDesc, TEXT("ShadowMaskBits") );
-
+					// TODO: This needs to move into the view loop in clustered deferred shading pass
 					for (int32 ViewIndex = 0, Num = Views.Num(); ViewIndex < Num; ViewIndex++)
 					{
-						RenderVirtualShadowMapProjection(
+						ShadowMaskBits = RenderVirtualShadowMapProjectionOnePass(
 							GraphBuilder,
 							SceneTextures,
 							Views[ ViewIndex ],
-							VirtualShadowMapArray,
-							ShadowMaskBits );
+							VirtualShadowMapArray );
 					}
 				}
 				else
diff --git a/Engine/Source/Runtime/Renderer/Private/SceneRendering.h b/Engine/Source/Runtime/Renderer/Private/SceneRendering.h
--- a/Engine/Source/Runtime/Renderer/Private/SceneRendering.h
+++ b/Engine/Source/Runtime/Renderer/Private/SceneRendering.h
@@ -1027,9 +1027,6 @@ struct FPreviousViewInfo
 	// History for shadow denoising.
 	TMap<const ULightComponent*, FScreenSpaceDenoiserHistory> ShadowHistories;
 
-	// History for virtual shadow denoising.
-	TMap<const ULightComponent*, FScreenSpaceDenoiserHistory> VirtualShadowHistories;
-
 	// History for denoising all lights penumbra at once.
 	FScreenSpaceDenoiserHistory PolychromaticPenumbraHarmonicsHistory;
 
diff --git a/Engine/Source/Runtime/Renderer/Private/ScreenSpaceDenoise.cpp b/Engine/Source/Runtime/Renderer/Private/ScreenSpaceDenoise.cpp
--- a/Engine/Source/Runtime/Renderer/Private/ScreenSpaceDenoise.cpp
+++ b/Engine/Source/Runtime/Renderer/Private/ScreenSpaceDenoise.cpp
@@ -3054,12 +3054,10 @@ FSSDSignalTextures IScreenSpaceDenoiser::DenoiseIndirectProbeHierarchy(
 FSSDSignalTextures IScreenSpaceDenoiser::DenoiseVirtualShadowMapMask(
 	FRDGBuilder& GraphBuilder,
 	const FViewInfo& View,
-	FPreviousViewInfo* PreviousViewInfos,
 	const FSceneTextureParameters& SceneTextures,
 	const FLightSceneInfo* LightSceneInfo,
 	FIntRect LightScissorRect,
-	const FVirtualShadowMapMaskInputs& InputParameters,
-	bool bUseTemporal)
+	const FVirtualShadowMapMaskInputs& InputParameters)
 {
 	RDG_GPU_STAT_SCOPE(GraphBuilder, VirtualShadowsDenoiser);
 	const FLightSceneProxy* Proxy = LightSceneInfo->Proxy;
@@ -3070,7 +3068,7 @@ FSSDSignalTextures IScreenSpaceDenoiser::DenoiseVirtualShadowMapMask(
 	Settings.SignalProcessing = ESignalProcessing::VirtualShadowMapMask;
 	Settings.ReconstructionSamples = CVarVirtualShadowReconstructionSampleCount.GetValueOnRenderThread();
 	Settings.PreConvolutionCount = CVarVirtualShadowPreConvolutionCount.GetValueOnRenderThread();
-	Settings.bUseTemporalAccumulation = bUseTemporal;
+	Settings.bUseTemporalAccumulation = false;
 	
 	Settings.LightSceneInfo[0] = LightSceneInfo;
 	Settings.SignalScissor[0] = LightScissorRect;
@@ -3082,13 +3080,8 @@ FSSDSignalTextures IScreenSpaceDenoiser::DenoiseVirtualShadowMapMask(
 	// No temporal reprojection
 	TStaticArray<FScreenSpaceDenoiserHistory*, IScreenSpaceDenoiser::kMaxBatchSize> PrevHistories;
 	TStaticArray<FScreenSpaceDenoiserHistory*, IScreenSpaceDenoiser::kMaxBatchSize> NewHistories;
-	PrevHistories[0] = PreviousViewInfos->VirtualShadowHistories.Find(LightSceneInfo->Proxy->GetLightComponent());
+	PrevHistories[0] = nullptr;
 	NewHistories[0] = nullptr;
-	if (!View.bStatePrevViewInfoIsReadOnly)
-	{
-		check(View.ViewState);
-		NewHistories[0] = &View.ViewState->PrevFrameViewInfo.VirtualShadowHistories.FindOrAdd(LightSceneInfo->Proxy->GetLightComponent());
-	}
 
 	FSSDSignalTextures InputSignal;
 	InputSignal.Textures[0] = InputParameters.Signal;
diff --git a/Engine/Source/Runtime/Renderer/Private/ScreenSpaceDenoise.h b/Engine/Source/Runtime/Renderer/Private/ScreenSpaceDenoise.h
--- a/Engine/Source/Runtime/Renderer/Private/ScreenSpaceDenoise.h
+++ b/Engine/Source/Runtime/Renderer/Private/ScreenSpaceDenoise.h
@@ -344,12 +344,10 @@ public:
 	static FSSDSignalTextures DenoiseVirtualShadowMapMask(
 		FRDGBuilder& GraphBuilder,
 		const FViewInfo& View,
-		FPreviousViewInfo* PreviousViewInfos,
 		const FSceneTextureParameters& SceneTextures,
 		const FLightSceneInfo* LightSceneInfo,
 		FIntRect LightScissorRect,
-		const FVirtualShadowMapMaskInputs& InputParameters,
-		bool bUseTemporal);
+		const FVirtualShadowMapMaskInputs& InputParameters);
 
 	/** Returns the interface of the default denoiser of the renderer. */
 	static const IScreenSpaceDenoiser* GetDefaultDenoiser();
diff --git a/Engine/Source/Runtime/Renderer/Private/ShadowRendering.cpp b/Engine/Source/Runtime/Renderer/Private/ShadowRendering.cpp
--- a/Engine/Source/Runtime/Renderer/Private/ShadowRendering.cpp
+++ b/Engine/Source/Runtime/Renderer/Private/ShadowRendering.cpp
@@ -188,8 +188,19 @@ static TAutoConsoleVariable<float> CVarShadowMaxSlopeScaleDepthBias(
 	TEXT("Higher values give better self-shadowing, but increase self-shadowing artifacts"),
 	ECVF_RenderThreadSafe);
 
+static TAutoConsoleVariable<int32> CVarVirtualSmDenoiserDirectional(
+	TEXT("r.Shadow.Virtual.Denoiser.EnableDirectional"),
+	1,
+	TEXT("Apply denoiser to directional light virtual shadow maps."),
+	ECVF_RenderThreadSafe);
 FForwardScreenSpaceShadowMaskTextureMobileOutputs GScreenSpaceShadowMaskTextureMobileOutputs;
 
+static TAutoConsoleVariable<int32> CVarVirtualSmDenoiserLocal(
+	TEXT("r.Shadow.Virtual.Denoiser.EnableLocal"),
+	0,
+	TEXT("Apply denoiser to local light virtual shadow maps."),
+	ECVF_RenderThreadSafe);
+
 ///////////////////////////////////////////////////////////////////////////////////////////////////
 // Hair
 static TAutoConsoleVariable<int32> CVarHairStrandsCullPerObjectShadowCaster(
@@ -198,18 +209,6 @@ static TAutoConsoleVariable<int32> CVarHairStrandsCullPerObjectShadowCaster(
 	TEXT("Enable CPU culling of object casting per-object shadow (stationnary object)"),
 	ECVF_RenderThreadSafe);
 
-static TAutoConsoleVariable<int32> CVarVirtualSmDenoiser(
-	TEXT("r.Shadow.Virtual.Denoiser"),
-	1,
-	TEXT("Apply denoiser on virtual shadow map masks."),
-	ECVF_RenderThreadSafe);
-
-static TAutoConsoleVariable<int32> CVarVirtualSmDenoiserTemporal(
-	TEXT("r.Shadow.Virtual.Denoiser.Temporal"),
-	0,
-	TEXT("Use temporal reprojection (when applicable)."),
-	ECVF_RenderThreadSafe);
-
 DEFINE_GPU_DRAWCALL_STAT(ShadowProjection);
 
 // Use Shadow stencil mask is set to 0x7F instead of 0xFF so that that last bit can be used for strata classification without clearing the stencil bit for pre-shadow/per-object static shadow mask
@@ -1993,133 +1992,92 @@ void FDeferredShadingSceneRenderer::RenderShadowProjections(
 
 				if (ScissorRect.Area() > 0)
 				{
-					if (CVarVirtualSmDenoiser.GetValueOnRenderThread() != 0)
+					// Project virtual shadow maps
+					FRDGTextureRef SignalTexture;
+					bool bUseDenoiser = false;
+					if (VisibleLightInfo.VirtualShadowMapClipmaps.Num() > 0)
 					{
-						const IScreenSpaceDenoiser* Denoiser = IScreenSpaceDenoiser::GetDefaultDenoiser();
-
-
-						// Virtual SMs draw into a texture formatted for the denoiser. See VirtualShadowMapProjection.usf
-						FRDGTextureRef SignalTexture;
-						const FLinearColor ClearColor(0.0f, 0.0f, 0.0f, 0.0f);
-						{
-							FRDGTextureDesc Desc = FRDGTextureDesc::Create2D(
-								SceneTextures.Config.Extent,
-								PF_FloatRGBA,
-								FClearValueBinding(ClearColor),
-								TexCreate_ShaderResource | TexCreate_RenderTargetable);
-
-							SignalTexture = GraphBuilder.CreateTexture(Desc, TEXT("VirtualShadowMap_Signal"));
-						}
-
-						// Project virtual shadow maps
-						bool bUseTemporal = false;
-						if (VisibleLightInfo.VirtualShadowMapClipmaps.Num() > 0)
-						{
-							RenderVirtualShadowMapProjectionForDenoising(
-								VisibleLightInfo.FindShadowClipmapForView(&View),
-								GraphBuilder,
-								View,
-								VirtualShadowMapArray,
-								ScissorRect,
-								SignalTexture,
-								bUseTemporal);
-						}
-						else
-						{
-							check(VirtualShadowMaps.Num() == 1);
-							RenderVirtualShadowMapProjectionForDenoising(
-								VirtualShadowMaps[0],
-								GraphBuilder,
-								View,
-								VirtualShadowMapArray,
-								ScissorRect,
-								SignalTexture,
-								bUseTemporal);
-						}
+						bUseDenoiser = CVarVirtualSmDenoiserDirectional.GetValueOnRenderThread() != 0;
+						SignalTexture = RenderVirtualShadowMapProjection(
+							GraphBuilder,
+							SceneTextures,
+							View,
+							VirtualShadowMapArray,
+							ScissorRect,
+							nullptr,	// Hair
+							VisibleLightInfo.FindShadowClipmapForView(&View));
+					}
+					else
+					{
+						bUseDenoiser = CVarVirtualSmDenoiserLocal.GetValueOnRenderThread() != 0;
+						check(VirtualShadowMaps.Num() == 1);
+						SignalTexture = RenderVirtualShadowMapProjection(
+							GraphBuilder,
+							SceneTextures,
+							View,
+							VirtualShadowMapArray,
+							ScissorRect,
+							nullptr,	// Hair
+							VirtualShadowMaps[0]);
+					}
 
-						bUseTemporal = bUseTemporal && (CVarVirtualSmDenoiserTemporal.GetValueOnRenderThread() != 0);
+					// Shadow filtering via denoiser
+					if (bUseDenoiser)
+					{
+						const IScreenSpaceDenoiser* Denoiser = IScreenSpaceDenoiser::GetDefaultDenoiser();
 
-						// Shadow filtering via denoiser
-						FSSDSignalTextures DenoisedSignal;
-						{
-							RDG_EVENT_SCOPE(GraphBuilder, "%s %dx%d",
-								Denoiser->GetDebugName(),
-								View.ViewRect.Width(), View.ViewRect.Height());
+						RDG_EVENT_SCOPE(GraphBuilder, "%s %dx%d",
+							Denoiser->GetDebugName(),
+							View.ViewRect.Width(), View.ViewRect.Height());
 
-							IScreenSpaceDenoiser::FVirtualShadowMapMaskInputs Inputs;
-							Inputs.Signal = SignalTexture;
+						IScreenSpaceDenoiser::FVirtualShadowMapMaskInputs Inputs;
+						Inputs.Signal = SignalTexture;
 
-							DenoisedSignal = Denoiser->DenoiseVirtualShadowMapMask(
-								GraphBuilder,
-								View,
-								&View.PrevViewInfo,
-								SceneTextureParameters,
-								LightSceneInfo,
-								ScissorRect,
-								Inputs,
-								bUseTemporal);
-						}
+						FSSDSignalTextures DenoisedSignal = Denoiser->DenoiseVirtualShadowMapMask(
+							GraphBuilder,
+							View,
+							SceneTextureParameters,
+							LightSceneInfo,
+							ScissorRect,
+							Inputs);
 
-						// Composite into light's screen shadow mask
-						CompositeVirtualShadowMapMask(GraphBuilder, ScissorRect, DenoisedSignal, ScreenShadowMaskTexture);
-					}
-					else
-					{
-						if (VisibleLightInfo.VirtualShadowMapClipmaps.Num() > 0)
-						{
-							RenderVirtualShadowMapProjection(
-								VisibleLightInfo.FindShadowClipmapForView(&View),
-								GraphBuilder,
-								View,
-								VirtualShadowMapArray,
-								ScissorRect,
-								ScreenShadowMaskTexture,
-								EVirtualShadowMapProjectionInputType::GBuffer,
-								bProjectingForForwardShading);
-						}
-						else
-						{
-							check(VirtualShadowMaps.Num() == 1);
-							RenderVirtualShadowMapProjection(
-								VirtualShadowMaps[0],
-								GraphBuilder,
-								View,
-								VirtualShadowMapArray,
-								ScissorRect,
-								ScreenShadowMaskTexture,
-								EVirtualShadowMapProjectionInputType::GBuffer,
-								bProjectingForForwardShading);
-						}
+						SignalTexture = DenoisedSignal.Textures[0];
 					}
 
-					// Sub-pixel shadow (no filtering for hair)
+					// Composite into screen shadow mask
+					CompositeVirtualShadowMapMask(GraphBuilder, ScissorRect, SignalTexture, ScreenShadowMaskTexture);
+
+					// Sub-pixel shadow (no denoising for hair)
 					if (HairStrands::HasViewHairStrandsData(View) && ScreenShadowMaskSubPixelTexture)
 					{
+						// TODO: This is ugly, but just a temporary workaround for branch divergence!
+						FRDGTextureRef HairCategorizationTexture = (*View.HairStrandsViewData.UniformBuffer)->HairCategorizationTexture;
 						if (VisibleLightInfo.VirtualShadowMapClipmaps.Num() > 0)
 						{
-							RenderVirtualShadowMapProjection(
-								VisibleLightInfo.FindShadowClipmapForView(&View),
+							SignalTexture = RenderVirtualShadowMapProjection(
 								GraphBuilder,
+								SceneTextures,
 								View,
 								VirtualShadowMapArray,
 								ScissorRect,
-								ScreenShadowMaskSubPixelTexture,
-								EVirtualShadowMapProjectionInputType::HairStrands,
-								bProjectingForForwardShading);
+								HairCategorizationTexture,
+								VisibleLightInfo.FindShadowClipmapForView(&View));
 						}
 						else
 						{
 							check(VirtualShadowMaps.Num() == 1);
-							RenderVirtualShadowMapProjection(
-								VirtualShadowMaps[0],
+							SignalTexture = RenderVirtualShadowMapProjection(
 								GraphBuilder,
+								SceneTextures,
 								View,
 								VirtualShadowMapArray,
 								ScissorRect,
-								ScreenShadowMaskSubPixelTexture,
-								EVirtualShadowMapProjectionInputType::HairStrands,
-								bProjectingForForwardShading);
+								HairCategorizationTexture,
+								VirtualShadowMaps[0]);
 						}
+
+						// Composite into sub pixel mask
+						CompositeVirtualShadowMapMask(GraphBuilder, ScissorRect, SignalTexture, ScreenShadowMaskSubPixelTexture);
 					}
 				}
 			}
diff --git a/Engine/Source/Runtime/Renderer/Private/ShadowRendering.h b/Engine/Source/Runtime/Renderer/Private/ShadowRendering.h
--- a/Engine/Source/Runtime/Renderer/Private/ShadowRendering.h
+++ b/Engine/Source/Runtime/Renderer/Private/ShadowRendering.h
@@ -155,21 +155,6 @@ inline bool IsShadowCacheModeOcclusionQueryable(EShadowDepthCacheMode CacheMode)
 	return CacheMode != SDCM_StaticPrimitivesOnly;
 }
 
-// Matches logic in VirtualShadowMapProjection.usf
-enum class EVirtualShadowMapProjectionOutputType
-{
-	/** Output signal for consumption by virtual shadow denoiser */
-	Denoiser = 0,
-
-	/** Output values for blending into screen shadow mask */
-	ScreenShadowMask,
-
-	/** Output debug visualization */
-	Debug,
-
-	MAX
-};
-
 class FShadowMapRenderTargets
 {
 public:
diff --git a/Engine/Source/Runtime/Renderer/Private/VirtualShadowMaps/VirtualShadowMapArray.cpp b/Engine/Source/Runtime/Renderer/Private/VirtualShadowMaps/VirtualShadowMapArray.cpp
--- a/Engine/Source/Runtime/Renderer/Private/VirtualShadowMaps/VirtualShadowMapArray.cpp
+++ b/Engine/Source/Runtime/Renderer/Private/VirtualShadowMaps/VirtualShadowMapArray.cpp
@@ -262,6 +262,7 @@ void FVirtualShadowMapArray::Initialize(FRDGBuilder& GraphBuilder, bool bInEnabl
 	bEnabled = bInEnabled;
 
 	UniformParameters.NumShadowMaps = 0;
+	UniformParameters.NumDirectionalLights = 0;
 
 	// Reference dummy data in the UB initially
 	TArray<uint32, SceneRenderingAllocator> DummyPageTable;
@@ -761,15 +762,11 @@ void FVirtualShadowMapArray::BuildPageAllocations(
 			for (const TSharedPtr<FVirtualShadowMapClipmap>& Clipmap : VisibleLightInfo.VirtualShadowMapClipmaps)
 			{
 				// NOTE: Shader assumes all levels from a given clipmap are contiguous in both the remap and projection arrays
+				int32 ClipmapID = Clipmap->GetVirtualShadowMap()->ID;
+				DirectionalLightSmInds.Add(ClipmapID);
 				for (int32 ClipmapLevel = 0; ClipmapLevel < Clipmap->GetLevelCount(); ++ClipmapLevel)
 				{
-					int32 ID = Clipmap->GetVirtualShadowMap(ClipmapLevel)->ID;
-					ShadowMapProjectionData[ID] = Clipmap->GetProjectionShaderData(ClipmapLevel);
-
-					if( ClipmapLevel == 0 )
-					{
-						DirectionalLightSmInds.Add(ID);
-					}
+					ShadowMapProjectionData[ClipmapID + ClipmapLevel] = Clipmap->GetProjectionShaderData(ClipmapLevel);
 				}
 			}
 
@@ -812,6 +809,9 @@ void FVirtualShadowMapArray::BuildPageAllocations(
 			}
 		}
 
+		UniformParameters.NumShadowMaps = ShadowMaps.Num();
+		UniformParameters.NumDirectionalLights = DirectionalLightSmInds.Num();
+
 		ShadowMapProjectionDataRDG = CreateProjectionDataBuffer(GraphBuilder, TEXT("Shadow.Virtual.ProjectionData"), ShadowMapProjectionData);
 		UniformParameters.ProjectionData = GraphBuilder.CreateSRV(ShadowMapProjectionDataRDG);
 
@@ -820,8 +820,7 @@ void FVirtualShadowMapArray::BuildPageAllocations(
 			StatsBufferRDG = GraphBuilder.CreateBuffer(FRDGBufferDesc::CreateStructuredDesc(sizeof(uint32), NumStats), TEXT("Shadow.Virtual.StatsBuffer"));
 			AddClearUAVPass(GraphBuilder, GraphBuilder.CreateUAV(StatsBufferRDG), 0);
 		}
-
-		UniformParameters.NumShadowMaps = ShadowMaps.Num();
+		
 		// Create and clear the requested page flags
 		const uint32 NumPageFlags = ShadowMaps.Num() * UniformParameters.PageTableSize;
 		FRDGBufferRef PageRequestFlagsRDG = GraphBuilder.CreateBuffer(FRDGBufferDesc::CreateStructuredDesc(sizeof(uint32), NumPageFlags), TEXT("Shadow.Virtual.PageRequestFlags"));
diff --git a/Engine/Source/Runtime/Renderer/Private/VirtualShadowMaps/VirtualShadowMapArray.h b/Engine/Source/Runtime/Renderer/Private/VirtualShadowMaps/VirtualShadowMapArray.h
--- a/Engine/Source/Runtime/Renderer/Private/VirtualShadowMaps/VirtualShadowMapArray.h
+++ b/Engine/Source/Runtime/Renderer/Private/VirtualShadowMaps/VirtualShadowMapArray.h
@@ -116,6 +116,7 @@ BEGIN_GLOBAL_SHADER_PARAMETER_STRUCT(FVirtualShadowMapUniformParameters, )
 	SHADER_PARAMETER(uint32, PageTableSize)
 	SHADER_PARAMETER(uint32, HPageTableSize)
 	SHADER_PARAMETER(uint32, NumShadowMaps)
+	SHADER_PARAMETER(uint32, NumDirectionalLights)
 	SHADER_PARAMETER(uint32, MaxPhysicalPages)
 	// use to map linear index to x,y page coord
 	SHADER_PARAMETER(uint32, PhysicalPageRowMask)
@@ -271,8 +272,6 @@ public:
 	FRDGBufferRef ShadowMapProjectionDataRDG = nullptr;
 	TArray< FRDGBufferRef, SceneRenderingAllocator > VirtualShadowMapIdRemapRDG;
 
-	int32 NumDirectionalLights = 0;
-
 	// These are not created or managed by this class - references are copied from the cache manager
 	// when needed as they come from the previous frame.
 	TRefCountPtr<IPooledRenderTarget>	HZBPhysical;
diff --git a/Engine/Source/Runtime/Renderer/Private/VirtualShadowMaps/VirtualShadowMapCacheManager.cpp b/Engine/Source/Runtime/Renderer/Private/VirtualShadowMaps/VirtualShadowMapCacheManager.cpp
--- a/Engine/Source/Runtime/Renderer/Private/VirtualShadowMaps/VirtualShadowMapCacheManager.cpp
+++ b/Engine/Source/Runtime/Renderer/Private/VirtualShadowMaps/VirtualShadowMapCacheManager.cpp
@@ -22,7 +22,7 @@ static TAutoConsoleVariable<int32> CVarAccumulateStats(
 
 static TAutoConsoleVariable<int32> CVarCacheVirtualSMs(
 	TEXT("r.Shadow.Virtual.Cache"),
-	0,
+	1,
 	TEXT("Turn on to enable caching"),
 	ECVF_RenderThreadSafe
 );
diff --git a/Engine/Source/Runtime/Renderer/Private/VirtualShadowMaps/VirtualShadowMapProjection.cpp b/Engine/Source/Runtime/Renderer/Private/VirtualShadowMaps/VirtualShadowMapProjection.cpp
--- a/Engine/Source/Runtime/Renderer/Private/VirtualShadowMaps/VirtualShadowMapProjection.cpp
+++ b/Engine/Source/Runtime/Renderer/Private/VirtualShadowMaps/VirtualShadowMapProjection.cpp
@@ -23,7 +23,6 @@
 #include "ShadowRendering.h"
 #include "SceneRendering.h"
 #include "VirtualShadowMapClipmap.h"
-#include "HairStrands/HairStrandsData.h"
 
 static TAutoConsoleVariable<float> CVarContactShadowLength(
 	TEXT( "r.Shadow.Virtual.ContactShadowLength" ),
@@ -56,14 +55,14 @@ TAutoConsoleVariable<int32> CVarVirtualShadowOnePassProjection(
 
 static TAutoConsoleVariable<int32> CVarSMRTRayCountLocal(
 	TEXT( "r.Shadow.Virtual.SMRT.RayCountLocal" ),
-	0,
+	7,
 	TEXT( "Ray count for shadow map tracing of local lights. 0 = disabled." ),
 	ECVF_RenderThreadSafe
 );
 
 static TAutoConsoleVariable<int32> CVarSMRTSamplesPerRayLocal(
 	TEXT( "r.Shadow.Virtual.SMRT.SamplesPerRayLocal" ),
-	16,
+	8,
 	TEXT( "Shadow map samples per ray for local lights" ),
 	ECVF_RenderThreadSafe
 );
@@ -86,14 +85,14 @@ static TAutoConsoleVariable<int32> CVarSMRTRayCountDirectional(
 
 static TAutoConsoleVariable<int32> CVarSMRTSamplesPerRayDirectional(
 	TEXT( "r.Shadow.Virtual.SMRT.SamplesPerRayDirectional" ),
-	12,
+	8,
 	TEXT( "Shadow map samples per ray for directional lights" ),
 	ECVF_RenderThreadSafe
 );
 
 static TAutoConsoleVariable<float> CVarSMRTRayLengthScaleDirectional(
 	TEXT( "r.Shadow.Virtual.SMRT.RayLengthScaleDirectional" ),
-	1.0f,
+	1.5f,
 	TEXT( "Length of ray to shoot for directional lights, scaled by distance to camera." )
 	TEXT( "Shorter rays limit the screen space size of shadow penumbra. " )
 	TEXT( "Longer rays require more samples to avoid shadows disconnecting from contact points. " ),
@@ -107,74 +106,6 @@ static TAutoConsoleVariable<int32> CVarSMRTAdaptiveRayCount(
 	ECVF_RenderThreadSafe
 );
 
-BEGIN_SHADER_PARAMETER_STRUCT(FProjectionParameters, )
-	SHADER_PARAMETER_STRUCT_INCLUDE(FVirtualShadowMapSamplingParameters, SamplingParameters)
-	SHADER_PARAMETER_STRUCT(FLightShaderParameters, Light)
-	SHADER_PARAMETER_RDG_UNIFORM_BUFFER(FSceneTextureUniformParameters, SceneTexturesStruct)
-	SHADER_PARAMETER_RDG_UNIFORM_BUFFER(FHairStrandsViewUniformParameters, HairStrands)
-	SHADER_PARAMETER_STRUCT_REF(FViewUniformShaderParameters, View)
-	SHADER_PARAMETER(int32, VirtualShadowMapId)
-	SHADER_PARAMETER(float, ContactShadowLength)
-	SHADER_PARAMETER(float, NormalOffsetWorld)
-	SHADER_PARAMETER(int32, DebugOutputType)
-	SHADER_PARAMETER(int32, SMRTRayCount)
-	SHADER_PARAMETER(int32, SMRTSamplesPerRay)
-	SHADER_PARAMETER(float, SMRTRayLengthScale)
-	SHADER_PARAMETER(float, SMRTCotMaxRayAngleFromLight)
-	SHADER_PARAMETER(uint32, InputType)
-	RENDER_TARGET_BINDING_SLOTS()
-END_SHADER_PARAMETER_STRUCT()
-
-class FVirtualShadowMapProjectionSpotPS : public FGlobalShader
-{
-	DECLARE_GLOBAL_SHADER(FVirtualShadowMapProjectionSpotPS);
-	SHADER_USE_PARAMETER_STRUCT(FVirtualShadowMapProjectionSpotPS, FGlobalShader);
-	
-	class FOutputTypeDim : SHADER_PERMUTATION_ENUM_CLASS("OUTPUT_TYPE", EVirtualShadowMapProjectionOutputType);
-	using FPermutationDomain = TShaderPermutationDomain<FOutputTypeDim>;
-
-	using FParameters = FProjectionParameters;
-
-	static bool ShouldCompilePermutation(const FGlobalShaderPermutationParameters& Parameters)
-	{
-		return DoesPlatformSupportNanite(Parameters.Platform);
-	}
-
-	static void ModifyCompilationEnvironment(const FGlobalShaderPermutationParameters& Parameters, FShaderCompilerEnvironment& OutEnvironment)
-	{
-		FVirtualShadowMapArray::SetShaderDefines(OutEnvironment);
-	}
-};
-IMPLEMENT_GLOBAL_SHADER(FVirtualShadowMapProjectionSpotPS, "/Engine/Private/VirtualShadowMaps/VirtualShadowMapProjection.usf", "VirtualShadowMapProjectionSpotPS", SF_Pixel);
-
-class FVirtualShadowMapProjectionDirectionalPS : public FGlobalShader
-{
-	DECLARE_GLOBAL_SHADER(FVirtualShadowMapProjectionDirectionalPS);
-	SHADER_USE_PARAMETER_STRUCT(FVirtualShadowMapProjectionDirectionalPS, FGlobalShader);
-
-	class FOutputTypeDim : SHADER_PERMUTATION_ENUM_CLASS("OUTPUT_TYPE", EVirtualShadowMapProjectionOutputType);
-	using FPermutationDomain = TShaderPermutationDomain<FOutputTypeDim>;
-
-	using FParameters = FProjectionParameters;
-
-	static bool ShouldCompilePermutation(const FGlobalShaderPermutationParameters& Parameters)
-	{
-		return DoesPlatformSupportNanite(Parameters.Platform);
-	}
-
-	static void ModifyCompilationEnvironment(const FGlobalShaderPermutationParameters& Parameters, FShaderCompilerEnvironment& OutEnvironment)
-	{
-		FVirtualShadowMapArray::SetShaderDefines(OutEnvironment);
-	}
-};
-IMPLEMENT_GLOBAL_SHADER(FVirtualShadowMapProjectionDirectionalPS, "/Engine/Private/VirtualShadowMaps/VirtualShadowMapProjection.usf", "VirtualShadowMapProjectionDirectionalPS", SF_Pixel);
-
-
-
-BEGIN_SHADER_PARAMETER_STRUCT(FVirtualShadowMapProjectionCompositeParameters, )
-	SHADER_PARAMETER_RDG_TEXTURE(Texture2D<float4>, InputSignal)
-	RENDER_TARGET_BINDING_SLOTS()
-END_SHADER_PARAMETER_STRUCT()
 
 // Composite denoised shadow projection mask onto the light's shadow mask
 // Basically just a copy shader with a special blend mode
@@ -182,8 +113,11 @@ class FVirtualShadowMapProjectionCompositePS : public FGlobalShader
 {
 	DECLARE_GLOBAL_SHADER(FVirtualShadowMapProjectionCompositePS);
 	SHADER_USE_PARAMETER_STRUCT(FVirtualShadowMapProjectionCompositePS, FGlobalShader);
-		
-	using FParameters = FVirtualShadowMapProjectionCompositeParameters;
+
+	BEGIN_SHADER_PARAMETER_STRUCT(FParameters, )
+		SHADER_PARAMETER_RDG_TEXTURE(Texture2D<float4>, InputSignal)
+		RENDER_TARGET_BINDING_SLOTS()
+	END_SHADER_PARAMETER_STRUCT()
 
 	static bool ShouldCompilePermutation(const FGlobalShaderPermutationParameters& Parameters)
 	{
@@ -198,262 +132,16 @@ class FVirtualShadowMapProjectionCompositePS : public FGlobalShader
 };
 IMPLEMENT_GLOBAL_SHADER(FVirtualShadowMapProjectionCompositePS, "/Engine/Private/VirtualShadowMaps/VirtualShadowMapProjection.usf", "VirtualShadowMapCompositePS", SF_Pixel);
 
-// Returns true if temporal denoising should be used
-static bool AddPass_RenderVirtualShadowMapProjection(
-	const FLightSceneProxy* LightProxy,
-	FRDGBuilder& GraphBuilder,
-	const FViewInfo& View, 
-	FVirtualShadowMapArray& VirtualShadowMapArray,
-	int VirtualShadowMapId,
-	const FIntRect ScissorRect,
-	EVirtualShadowMapProjectionInputType InputType,
-	EVirtualShadowMapProjectionOutputType OutputType,
-	const FRenderTargetBinding& RenderTargetBinding,
-	FRHIBlendState* BlendState)
-{
-	check(ScissorRect.Area() > 0);
-		
-	FGlobalShaderMap* ShaderMap = View.ShaderMap;
-
-	FProjectionParameters* PassParameters = GraphBuilder.AllocParameters<FProjectionParameters>();
-	PassParameters->SamplingParameters = VirtualShadowMapArray.GetSamplingParameters(GraphBuilder);
-	PassParameters->SceneTexturesStruct = CreateSceneTextureUniformBuffer(GraphBuilder, View.FeatureLevel, ESceneTextureSetupMode::GBuffers | ESceneTextureSetupMode::SceneDepth);
-	PassParameters->View = View.ViewUniformBuffer;
-		
-	FLightShaderParameters LightParameters;
-	LightProxy->GetLightShaderParameters(LightParameters);
-	PassParameters->Light = LightParameters;
-	PassParameters->VirtualShadowMapId = VirtualShadowMapId;
-	PassParameters->DebugOutputType = CVarVirtualShadowMapDebugProjection.GetValueOnRenderThread();
-	PassParameters->ContactShadowLength = CVarContactShadowLength.GetValueOnRenderThread();
-	PassParameters->NormalOffsetWorld = CVarNormalOffsetWorld.GetValueOnRenderThread();
-	PassParameters->InputType = uint32(InputType);
-	PassParameters->HairStrands = HairStrands::BindHairStrandsViewUniformParameters(View);
-	PassParameters->RenderTargets[0] = RenderTargetBinding;
-
-	if (LightProxy->GetLightType() == LightType_Directional)
-	{
-		PassParameters->SMRTRayCount = CVarSMRTRayCountDirectional.GetValueOnRenderThread();
-		PassParameters->SMRTSamplesPerRay = CVarSMRTSamplesPerRayDirectional.GetValueOnRenderThread();
-		PassParameters->SMRTRayLengthScale = CVarSMRTRayLengthScaleDirectional.GetValueOnRenderThread();
-		PassParameters->SMRTCotMaxRayAngleFromLight = 0.0f;	// unused in this path
-
-		FVirtualShadowMapProjectionDirectionalPS::FPermutationDomain PermutationVector;
-		PermutationVector.Set<FVirtualShadowMapProjectionDirectionalPS::FOutputTypeDim>(OutputType);
-		auto PixelShader = ShaderMap->GetShader<FVirtualShadowMapProjectionDirectionalPS>(PermutationVector);
-		ValidateShaderParameters(PixelShader, *PassParameters);
-
-		// NOTE: We use SV_Position in the shader, so we don't need separate scissor/view rect
-		FPixelShaderUtils::AddFullscreenPass(GraphBuilder,
-			ShaderMap,
-			(OutputType == EVirtualShadowMapProjectionOutputType::Debug) ? RDG_EVENT_NAME("Debug Directional Projection") : RDG_EVENT_NAME("Directional Projection"),
-			PixelShader,
-			PassParameters,
-			ScissorRect,
-			BlendState);
-
-		// Use temporal denoising with SMRT
-		return PassParameters->SMRTRayCount > 0;
-	}
-	else
-	{
-		PassParameters->SMRTRayCount = CVarSMRTRayCountLocal.GetValueOnRenderThread();
-		PassParameters->SMRTSamplesPerRay = CVarSMRTSamplesPerRayLocal.GetValueOnRenderThread();
-		PassParameters->SMRTRayLengthScale = 0.0f;		// unused in this path
-		PassParameters->SMRTCotMaxRayAngleFromLight = 1.0f / FMath::Tan(CVarSMRTMaxRayAngleFromLight.GetValueOnRenderThread());
-
-		FVirtualShadowMapProjectionSpotPS::FPermutationDomain PermutationVector;
-		PermutationVector.Set<FVirtualShadowMapProjectionSpotPS::FOutputTypeDim>(OutputType);
-		auto PixelShader = ShaderMap->GetShader<FVirtualShadowMapProjectionSpotPS>(PermutationVector);
-		ValidateShaderParameters(PixelShader, *PassParameters);
-
-		// NOTE: We use SV_Position in the shader, so we don't need separate scissor/view rect
-		FPixelShaderUtils::AddFullscreenPass(GraphBuilder,
-			ShaderMap,
-			(OutputType == EVirtualShadowMapProjectionOutputType::Debug) ? RDG_EVENT_NAME("Debug Spot Projection") : RDG_EVENT_NAME("Spot Projection"),
-			PixelShader,
-			PassParameters,
-			ScissorRect,
-			BlendState);
-
-		// Use temporal denoising with SMRT
-		return PassParameters->SMRTRayCount > 0;
-	}
-
-	return false;
-}
-
-static FRDGTextureRef CreateDebugOutput(FRDGBuilder& GraphBuilder, FIntPoint Extent)
-{
-	FRDGTextureDesc DebugOutputDesc = FRDGTextureDesc::Create2D(
-		Extent,
-		PF_A32B32G32R32F,
-		FClearValueBinding::Transparent,
-		TexCreate_ShaderResource | TexCreate_RenderTargetable
-	);
-
-	return GraphBuilder.CreateTexture(DebugOutputDesc, TEXT("Shadow.Virtual.DebugProjection"));
-}
-
-// Returns true if temporal denoising should be used
-static bool RenderVirtualShadowMapProjectionCommon(
-	const FLightSceneProxy* LightProxy,
-	FRDGBuilder& GraphBuilder,
-	const FViewInfo& View, 
-	FVirtualShadowMapArray& VirtualShadowMapArray,
-	int32 VirtualShadowMapId,
-	const FIntRect ScissorRect,
-	EVirtualShadowMapProjectionInputType InputType,
-	EVirtualShadowMapProjectionOutputType OutputType,
-	const FRenderTargetBinding& RenderTargetBinding,
-	FRHIBlendState* BlendState)
-{
-	// Main Pass
-	bool bOutUseTemporalDenoising = AddPass_RenderVirtualShadowMapProjection(
-		LightProxy,
-		GraphBuilder,
-		View,
-		VirtualShadowMapArray,
-		VirtualShadowMapId,
-		ScissorRect,
-		InputType,
-		OutputType,
-		RenderTargetBinding,
-		BlendState);
-
-	// Debug visualization
-	const int32 DebugVisualize = CVarVirtualShadowMapDebugProjection.GetValueOnRenderThread();
-	if (DebugVisualize > 0)
-	{
-		FRDGTextureRef DebugOutput = CreateDebugOutput(GraphBuilder, View.ViewRect.Max);
-		FRenderTargetBinding DebugRenderTargetBinding(DebugOutput, ERenderTargetLoadAction::EClear);
-
-		AddPass_RenderVirtualShadowMapProjection(
-			LightProxy,
-			GraphBuilder,
-			View,
-			VirtualShadowMapArray,
-			VirtualShadowMapId,
-			ScissorRect,
-			InputType,
-			EVirtualShadowMapProjectionOutputType::Debug,
-			DebugRenderTargetBinding,
-			nullptr);
-
-		GraphBuilder.QueueTextureExtraction(DebugOutput, &VirtualShadowMapArray.DebugVisualizationProjectionOutput);
-	}
-
-	return bOutUseTemporalDenoising;
-}
-
-void RenderVirtualShadowMapProjectionForDenoising(
-	FProjectedShadowInfo* ShadowInfo,
-	FRDGBuilder& GraphBuilder,
-	const FViewInfo& View,
-	FVirtualShadowMapArray& VirtualShadowMapArray,
-	const FIntRect ScissorRect,
-	FRDGTextureRef SignalTexture,
-	bool& bOutUseTemporalDenoising)
-{
-	FRenderTargetBinding RenderTargetBinding(SignalTexture, ERenderTargetLoadAction::EClear);
-	bOutUseTemporalDenoising = RenderVirtualShadowMapProjectionCommon(
-		ShadowInfo->GetLightSceneInfo().Proxy,
-		GraphBuilder,
-		View,
-		VirtualShadowMapArray,
-		ShadowInfo->VirtualShadowMaps[0]->ID,
-		ScissorRect,
-		EVirtualShadowMapProjectionInputType::GBuffer,
-		EVirtualShadowMapProjectionOutputType::Denoiser,
-		RenderTargetBinding,
-		nullptr);
-}
-
-void RenderVirtualShadowMapProjectionForDenoising(
-	const TSharedPtr<FVirtualShadowMapClipmap>& Clipmap,
-	FRDGBuilder& GraphBuilder,
-	const FViewInfo& View,
-	FVirtualShadowMapArray& VirtualShadowMapArray,
-	const FIntRect ScissorRect,
-	FRDGTextureRef SignalTexture,
-	bool& bOutUseTemporalDenoising)
-{	
-	FRenderTargetBinding RenderTargetBinding(SignalTexture, ERenderTargetLoadAction::EClear);
-	bOutUseTemporalDenoising = RenderVirtualShadowMapProjectionCommon(
-		Clipmap->GetLightSceneInfo().Proxy,
-		GraphBuilder,
-		View,
-		VirtualShadowMapArray,
-		Clipmap->GetVirtualShadowMap()->ID,
-		ScissorRect,
-		EVirtualShadowMapProjectionInputType::GBuffer,
-		EVirtualShadowMapProjectionOutputType::Denoiser,
-		RenderTargetBinding,
-		nullptr);
-}
-
-void RenderVirtualShadowMapProjection(
-	FProjectedShadowInfo* ShadowInfo,
-	FRDGBuilder& GraphBuilder,
-	const FViewInfo& View,
-	FVirtualShadowMapArray& VirtualShadowMapArray,
-	const FIntRect ScissorRect,
-	FRDGTextureRef ScreenShadowMaskTexture,
-	EVirtualShadowMapProjectionInputType InputType,
-	bool bProjectingForForwardShading)
-{
-	FRenderTargetBinding RenderTargetBinding(ScreenShadowMaskTexture, ERenderTargetLoadAction::ELoad);
-	FRHIBlendState* BlendState = ShadowInfo->GetBlendStateForProjection(bProjectingForForwardShading, false);
-	RenderVirtualShadowMapProjectionCommon(
-		ShadowInfo->GetLightSceneInfo().Proxy,
-		GraphBuilder,
-		View,
-		VirtualShadowMapArray,
-		ShadowInfo->VirtualShadowMaps[0]->ID,
-		ScissorRect,
-		InputType,
-		EVirtualShadowMapProjectionOutputType::ScreenShadowMask,
-		RenderTargetBinding,
-		BlendState);
-}
-
-void RenderVirtualShadowMapProjection(
-	const TSharedPtr<FVirtualShadowMapClipmap>& Clipmap,
-	FRDGBuilder& GraphBuilder,
-	const FViewInfo& View,
-	FVirtualShadowMapArray& VirtualShadowMapArray,
-	const FIntRect ScissorRect,
-	FRDGTextureRef ScreenShadowMaskTexture,
-	EVirtualShadowMapProjectionInputType InputType,
-	bool bProjectingForForwardShading)
-{
-	FRenderTargetBinding RenderTargetBinding(ScreenShadowMaskTexture, ERenderTargetLoadAction::ELoad);	
-	// See FProjectedShadowInfo::GetBlendStateForProjection. TODO: Support other modes in this path?
-	FRHIBlendState* BlendState = TStaticBlendState<CW_BA, BO_Min, BF_One, BF_One, BO_Min, BF_One, BF_One>::GetRHI();
-	RenderVirtualShadowMapProjectionCommon(
-		Clipmap->GetLightSceneInfo().Proxy,
-		GraphBuilder,
-		View,
-		VirtualShadowMapArray,
-		Clipmap->GetVirtualShadowMap()->ID,
-		ScissorRect,
-		InputType,
-		EVirtualShadowMapProjectionOutputType::ScreenShadowMask,
-		RenderTargetBinding,
-		BlendState);
-}
-
-
 void CompositeVirtualShadowMapMask(
 	FRDGBuilder& GraphBuilder,
 	const FIntRect ScissorRect,
-	const FSSDSignalTextures& InputSignal,
+	const FRDGTextureRef Input,
 	FRDGTextureRef OutputShadowMaskTexture)
 {
 	auto ShaderMap = GetGlobalShaderMap(GMaxRHIFeatureLevel);
 
-	FVirtualShadowMapProjectionCompositeParameters* PassParameters = GraphBuilder.AllocParameters<FVirtualShadowMapProjectionCompositeParameters>();
-	PassParameters->InputSignal = InputSignal.Textures[0];
+	FVirtualShadowMapProjectionCompositePS::FParameters* PassParameters = GraphBuilder.AllocParameters<FVirtualShadowMapProjectionCompositePS::FParameters>();
+	PassParameters->InputSignal = Input;
 
 	PassParameters->RenderTargets[0] = FRenderTargetBinding(OutputShadowMaskTexture, ERenderTargetLoadAction::ELoad);
 
@@ -472,33 +160,45 @@ void CompositeVirtualShadowMapMask(
 		BlendState);
 }
 
+
 class FVirtualShadowMapProjectionCS : public FGlobalShader
 {
 	DECLARE_GLOBAL_SHADER(FVirtualShadowMapProjectionCS);
 	SHADER_USE_PARAMETER_STRUCT(FVirtualShadowMapProjectionCS, FGlobalShader)
 	
+	class FDirectionalLightDim		: SHADER_PERMUTATION_BOOL("DIRECTIONAL_LIGHT");
 	class FSMRTAdaptiveRayCountDim	: SHADER_PERMUTATION_BOOL("SMRT_ADAPTIVE_RAY_COUNT");
 	class FTwoPhysicalTexturesDim	: SHADER_PERMUTATION_BOOL("TWO_PHYSICAL_TEXTURES");
+	class FOnePassProjectionDim		: SHADER_PERMUTATION_BOOL("ONE_PASS_PROJECTION");
+
 	using FPermutationDomain = TShaderPermutationDomain<
+		FDirectionalLightDim,
+		FOnePassProjectionDim,
 		FSMRTAdaptiveRayCountDim,
-		FTwoPhysicalTexturesDim >;
+		FTwoPhysicalTexturesDim	>;
 
 	BEGIN_SHADER_PARAMETER_STRUCT(FParameters, )
 		SHADER_PARAMETER_STRUCT_INCLUDE(FVirtualShadowMapSamplingParameters, SamplingParameters)
 		SHADER_PARAMETER_RDG_UNIFORM_BUFFER(FSceneTextureUniformParameters, SceneTexturesStruct)
-		SHADER_PARAMETER_STRUCT_REF(FViewUniformShaderParameters, View)
-		SHADER_PARAMETER_STRUCT_REF(FForwardLightData, ForwardLightData)
-		SHADER_PARAMETER_RDG_BUFFER_SRV(StructuredBuffer< FVirtualShadowMapProjectionShaderData >, ShadowMapProjectionData)
-		SHADER_PARAMETER_RDG_BUFFER_SRV(StructuredBuffer< uint >, VirtualShadowMapIdRemap)
-		SHADER_PARAMETER_RDG_TEXTURE_UAV(RWTexture2D< uint >, RWShadowMaskBits)
-		SHADER_PARAMETER(uint32, NumDirectionalLightSmInds)
+		SHADER_PARAMETER_STRUCT_REF(FViewUniformShaderParameters, View)		
+		SHADER_PARAMETER(FIntVector4, ProjectionRect)
 		SHADER_PARAMETER(float, ContactShadowLength)
 		SHADER_PARAMETER(float, NormalOffsetWorld)
-		SHADER_PARAMETER(int32, DebugOutputType)
-		SHADER_PARAMETER(int32, SMRTRayCount)
-		SHADER_PARAMETER(int32, SMRTSamplesPerRay)
+		SHADER_PARAMETER(uint32, SMRTRayCount)
+		SHADER_PARAMETER(uint32, SMRTSamplesPerRay)
 		SHADER_PARAMETER(float, SMRTRayLengthScale)
 		SHADER_PARAMETER(float, SMRTCotMaxRayAngleFromLight)
+		SHADER_PARAMETER(int32, DebugOutputType)
+		SHADER_PARAMETER_RDG_TEXTURE(Texture2D<uint4>, HairCategorizationTexture)
+		SHADER_PARAMETER(uint32, bUseHairData)
+		// One pass projection parameters
+		SHADER_PARAMETER_STRUCT_REF(FForwardLightData, ForwardLightData)
+		SHADER_PARAMETER_RDG_BUFFER_SRV(StructuredBuffer< uint >, VirtualShadowMapIdRemap)	// TODO: Move to VSM UB?
+		SHADER_PARAMETER_RDG_TEXTURE_UAV(RWTexture2D< uint >, RWShadowMaskBits)
+		// Pass per light parameters
+		SHADER_PARAMETER_STRUCT(FLightShaderParameters, Light)
+		SHADER_PARAMETER(int32, LightUniformVirtualShadowMapId)
+		SHADER_PARAMETER_RDG_TEXTURE_UAV(RWTexture2D, RWSignal)
 	END_SHADER_PARAMETER_STRUCT()
 
 	static void ModifyCompilationEnvironment( const FGlobalShaderPermutationParameters& Parameters, FShaderCompilerEnvironment& OutEnvironment )
@@ -521,6 +221,15 @@ class FVirtualShadowMapProjectionCS : public FGlobalShader
 
 	static bool ShouldCompilePermutation(const FGlobalShaderPermutationParameters& Parameters)
 	{
+		FPermutationDomain PermutationVector(Parameters.PermutationId);
+
+		// Directional lights are always in separate passes as forward light data structure currently
+		// only contains a single single directional light.
+		if( PermutationVector.Get< FDirectionalLightDim >() && PermutationVector.Get< FOnePassProjectionDim >() )
+		{
+			return false;
+		}
+		
 		return DoesPlatformSupportNanite(Parameters.Platform);
 	}
 };
@@ -528,48 +237,170 @@ IMPLEMENT_GLOBAL_SHADER(FVirtualShadowMapProjectionCS, "/Engine/Private/VirtualS
 
 extern int32 GVirtualShadowMapAtomicWrites;
 
-void RenderVirtualShadowMapProjection(
+static void RenderVirtualShadowMapProjectionCommon(
 	FRDGBuilder& GraphBuilder,
 	const FMinimalSceneTextures& SceneTextures,
 	const FViewInfo& View,
 	FVirtualShadowMapArray& VirtualShadowMapArray,
-	FRDGTextureRef ShadowMaskBits )
+	const FIntRect ProjectionRect,
+	FRDGTextureRef HairCategorization,
+	FRDGTextureRef OutputTexture,
+	const FLightSceneProxy* LightProxy = nullptr,
+	int32 VirtualShadowMapId = INDEX_NONE)
 {
 	FVirtualShadowMapProjectionCS::FParameters* PassParameters = GraphBuilder.AllocParameters< FVirtualShadowMapProjectionCS::FParameters >();
 	PassParameters->SamplingParameters = VirtualShadowMapArray.GetSamplingParameters(GraphBuilder);
-
 	PassParameters->SceneTexturesStruct = SceneTextures.UniformBuffer;
 	PassParameters->View = View.ViewUniformBuffer;
-	PassParameters->ForwardLightData = View.ForwardLightingResources->ForwardLightDataUniformBuffer;
-	PassParameters->VirtualShadowMapIdRemap = GraphBuilder.CreateSRV( VirtualShadowMapArray.VirtualShadowMapIdRemapRDG[0] );	// FIXME Index proper view
-	PassParameters->NumDirectionalLightSmInds = VirtualShadowMapArray.NumDirectionalLights;
-
+	PassParameters->ProjectionRect = FIntVector4(ProjectionRect.Min.X, ProjectionRect.Min.Y, ProjectionRect.Max.X, ProjectionRect.Max.Y);
 	PassParameters->DebugOutputType = CVarVirtualShadowMapDebugProjection.GetValueOnRenderThread();
 	PassParameters->ContactShadowLength = CVarContactShadowLength.GetValueOnRenderThread();
 	PassParameters->NormalOffsetWorld = CVarNormalOffsetWorld.GetValueOnRenderThread();
+	PassParameters->bUseHairData = HairCategorization != nullptr ? 1 : 0;
+	PassParameters->HairCategorizationTexture = HairCategorization ? HairCategorization : GSystemTextures.GetBlackDummy(GraphBuilder);
 
-	PassParameters->SMRTRayCount = CVarSMRTRayCountLocal.GetValueOnRenderThread();
-	PassParameters->SMRTSamplesPerRay = CVarSMRTSamplesPerRayLocal.GetValueOnRenderThread();
-	PassParameters->SMRTRayLengthScale = 0.0f;		// Currently unused in this path
-	PassParameters->SMRTCotMaxRayAngleFromLight = 1.0f / FMath::Tan(CVarSMRTMaxRayAngleFromLight.GetValueOnRenderThread());
-
-	PassParameters->RWShadowMaskBits = GraphBuilder.CreateUAV( ShadowMaskBits );
+	bool bDirectionalLight = false;
+	bool bOnePassProjection = LightProxy == nullptr;
+	if (bOnePassProjection)
+	{
+		// One pass projection
+		PassParameters->ForwardLightData = View.ForwardLightingResources->ForwardLightDataUniformBuffer;
+		PassParameters->VirtualShadowMapIdRemap = GraphBuilder.CreateSRV( VirtualShadowMapArray.VirtualShadowMapIdRemapRDG[0] );	// FIXME Index proper view
+		PassParameters->RWShadowMaskBits = GraphBuilder.CreateUAV( OutputTexture );	
+	}
+	else
+	{
+		// Pass per light
+		bDirectionalLight = LightProxy->GetLightType() == LightType_Directional;
+		FLightShaderParameters LightParameters;
+		LightProxy->GetLightShaderParameters(LightParameters);
+		PassParameters->Light = LightParameters;
+		PassParameters->LightUniformVirtualShadowMapId = VirtualShadowMapId;
+		PassParameters->RWSignal = GraphBuilder.CreateUAV( OutputTexture );
+	}
+ 
+	if (bDirectionalLight)
+	{
+		PassParameters->SMRTRayCount = CVarSMRTRayCountDirectional.GetValueOnRenderThread();
+		PassParameters->SMRTSamplesPerRay = CVarSMRTSamplesPerRayDirectional.GetValueOnRenderThread();
+		PassParameters->SMRTRayLengthScale = CVarSMRTRayLengthScaleDirectional.GetValueOnRenderThread();
+		PassParameters->SMRTCotMaxRayAngleFromLight = 0.0f;	// unused in this path
+	}
+	else
+	{
+		PassParameters->SMRTRayCount = CVarSMRTRayCountLocal.GetValueOnRenderThread();
+		PassParameters->SMRTSamplesPerRay = CVarSMRTSamplesPerRayLocal.GetValueOnRenderThread();
+		PassParameters->SMRTRayLengthScale = 0.0f;		// unused in this path
+		PassParameters->SMRTCotMaxRayAngleFromLight = 1.0f / FMath::Tan(CVarSMRTMaxRayAngleFromLight.GetValueOnRenderThread());
+	}
 	
 	bool bAdaptiveRayCount = GRHISupportsWaveOperations && CVarSMRTAdaptiveRayCount.GetValueOnRenderThread() != 0;
 
 	FVirtualShadowMapProjectionCS::FPermutationDomain PermutationVector;
+	PermutationVector.Set< FVirtualShadowMapProjectionCS::FDirectionalLightDim >( bDirectionalLight );
+	PermutationVector.Set< FVirtualShadowMapProjectionCS::FOnePassProjectionDim >( bOnePassProjection );
 	PermutationVector.Set< FVirtualShadowMapProjectionCS::FSMRTAdaptiveRayCountDim >( bAdaptiveRayCount );
 	PermutationVector.Set< FVirtualShadowMapProjectionCS::FTwoPhysicalTexturesDim >( GVirtualShadowMapAtomicWrites == 0 );
 
 	auto ComputeShader = View.ShaderMap->GetShader< FVirtualShadowMapProjectionCS >( PermutationVector );
+	ClearUnusedGraphResources( ComputeShader, PassParameters );
+	ValidateShaderParameters( ComputeShader, *PassParameters );
 
-	const FIntPoint GroupCount = FIntPoint::DivideAndRoundUp( View.ViewRect.Size(), 8 );
-
+	const FIntPoint GroupCount = FIntPoint::DivideAndRoundUp( ProjectionRect.Size(), 8 );
 	FComputeShaderUtils::AddPass(
 		GraphBuilder,
-		bAdaptiveRayCount ? RDG_EVENT_NAME("VirtualShadowMapProjection") : RDG_EVENT_NAME("VirtualShadowMapProjection (StaticRayCount)"),
+		bAdaptiveRayCount ? RDG_EVENT_NAME("VirtualShadowMapProjection") : RDG_EVENT_NAME("VirtualShadowMapProjection_StaticRayCount"),
 		ComputeShader,
 		PassParameters,
 		FIntVector( GroupCount.X, GroupCount.Y, 1 )
 	);
 }
+
+FRDGTextureRef RenderVirtualShadowMapProjectionOnePass(
+	FRDGBuilder& GraphBuilder,
+	const FMinimalSceneTextures& SceneTextures,
+	const FViewInfo& View,
+	FVirtualShadowMapArray& VirtualShadowMapArray )
+{
+	FIntRect ProjectionRect = View.ViewRect;
+
+	const FRDGTextureDesc ShadowMaskDesc = FRDGTextureDesc::Create2D(
+		SceneTextures.Config.Extent,
+		PF_R32_UINT,
+		FClearValueBinding::None,
+		TexCreate_ShaderResource | TexCreate_UAV );
+
+	FRDGTextureRef ShadowMaskBits = GraphBuilder.CreateTexture( ShadowMaskDesc, TEXT("ShadowMaskBits") );
+	
+	RenderVirtualShadowMapProjectionCommon(
+		GraphBuilder,
+		SceneTextures,
+		View,
+		VirtualShadowMapArray,
+		ProjectionRect,
+		nullptr,		// TODO: Hair
+		ShadowMaskBits);
+
+	return ShadowMaskBits;
+}
+
+static FRDGTextureRef CreateSignalTexture(FRDGBuilder& GraphBuilder, FIntPoint Extent)
+{
+	const FLinearColor ClearColor(0.0f, 0.0f, 0.0f, 0.0f);
+
+	FRDGTextureDesc Desc = FRDGTextureDesc::Create2D(
+		Extent,
+		PF_FloatRGBA,
+		FClearValueBinding(ClearColor),
+		TexCreate_ShaderResource | TexCreate_UAV);
+
+	FRDGTextureRef Texture = GraphBuilder.CreateTexture(Desc, TEXT("Shadow.Virtual.Signal"));
+	AddClearUAVPass(GraphBuilder, GraphBuilder.CreateUAV(Texture), ClearColor);
+	return Texture;
+}
+
+FRDGTextureRef RenderVirtualShadowMapProjection(
+	FRDGBuilder& GraphBuilder,
+	const FMinimalSceneTextures& SceneTextures,
+	const FViewInfo& View,
+	FVirtualShadowMapArray& VirtualShadowMapArray,
+	const FIntRect ScissorRect,
+	FRDGTextureRef HairCategorization,
+	FProjectedShadowInfo* ShadowInfo)
+{	
+	FRDGTextureRef SignalTexture = CreateSignalTexture(GraphBuilder, SceneTextures.Config.Extent);
+	RenderVirtualShadowMapProjectionCommon(
+		GraphBuilder,
+		SceneTextures,
+		View,
+		VirtualShadowMapArray,
+		ScissorRect,
+		HairCategorization,
+		SignalTexture,
+		ShadowInfo->GetLightSceneInfo().Proxy,
+		ShadowInfo->VirtualShadowMaps[0]->ID);
+	return SignalTexture;
+}
+
+FRDGTextureRef RenderVirtualShadowMapProjection(
+	FRDGBuilder& GraphBuilder,
+	const FMinimalSceneTextures& SceneTextures,
+	const FViewInfo& View,
+	FVirtualShadowMapArray& VirtualShadowMapArray,
+	const FIntRect ScissorRect,
+	FRDGTextureRef HairCategorization,
+	const TSharedPtr<FVirtualShadowMapClipmap>& Clipmap)
+{
+	FRDGTextureRef SignalTexture = CreateSignalTexture(GraphBuilder, SceneTextures.Config.Extent);
+	RenderVirtualShadowMapProjectionCommon(
+		GraphBuilder,
+		SceneTextures,
+		View,
+		VirtualShadowMapArray,		
+		ScissorRect,
+		HairCategorization,
+		SignalTexture,
+		Clipmap->GetLightSceneInfo().Proxy,
+		Clipmap->GetVirtualShadowMap()->ID);
+	return SignalTexture;
+}
diff --git a/Engine/Source/Runtime/Renderer/Private/VirtualShadowMaps/VirtualShadowMapProjection.h b/Engine/Source/Runtime/Renderer/Private/VirtualShadowMaps/VirtualShadowMapProjection.h
--- a/Engine/Source/Runtime/Renderer/Private/VirtualShadowMaps/VirtualShadowMapProjection.h
+++ b/Engine/Source/Runtime/Renderer/Private/VirtualShadowMaps/VirtualShadowMapProjection.h
@@ -13,59 +13,32 @@
 
 class FVirtualShadowMapClipmap;
 
-enum class EVirtualShadowMapProjectionInputType
-{
-	GBuffer = 0,
-	HairStrands = 1
-};
-
-void RenderVirtualShadowMapProjectionForDenoising(
-	const TSharedPtr<FVirtualShadowMapClipmap>& Clipmap,
-	FRDGBuilder& GraphBuilder,
-	const FViewInfo& View,
-	FVirtualShadowMapArray& VirtualShadowMapArray,
-	const FIntRect ScissorRect,
-	FRDGTextureRef SignalTexture,
-	bool& bOutUseTemporalDenoising);
-
-void RenderVirtualShadowMapProjectionForDenoising(
-	FProjectedShadowInfo* ShadowInfo,
-	FRDGBuilder& GraphBuilder,
-	const FViewInfo& View,
-	FVirtualShadowMapArray& VirtualShadowMapArray,
-	const FIntRect ScissorRect,
-	FRDGTextureRef SignalTexture,
-	bool& bOutUseTemporalDenoising);
-
-void RenderVirtualShadowMapProjection(
-	const TSharedPtr<FVirtualShadowMapClipmap>& Clipmap,
+FRDGTextureRef RenderVirtualShadowMapProjection(
 	FRDGBuilder& GraphBuilder,
+	const FMinimalSceneTextures& SceneTextures,
 	const FViewInfo& View,
 	FVirtualShadowMapArray& VirtualShadowMapArray,
 	const FIntRect ScissorRect,
-	FRDGTextureRef ScreenShadowMaskTexture,
-	EVirtualShadowMapProjectionInputType InputType,
-	bool bProjectingForForwardShading);
+	FRDGTextureRef HairCategorization,
+	const TSharedPtr<FVirtualShadowMapClipmap>& Clipmap);
 
-void RenderVirtualShadowMapProjection(
-	FProjectedShadowInfo* ShadowInfo,
+FRDGTextureRef RenderVirtualShadowMapProjection(
 	FRDGBuilder& GraphBuilder,
+	const FMinimalSceneTextures& SceneTextures,
 	const FViewInfo& View,
 	FVirtualShadowMapArray& VirtualShadowMapArray,
 	const FIntRect ScissorRect,
-	FRDGTextureRef ScreenShadowMaskTexture,
-	EVirtualShadowMapProjectionInputType InputType,
-	bool bProjectingForForwardShading);
+	FRDGTextureRef HairCategorization,
+	FProjectedShadowInfo* ShadowInfo);
 
-void RenderVirtualShadowMapProjection(
+FRDGTextureRef RenderVirtualShadowMapProjectionOnePass(
 	FRDGBuilder& GraphBuilder,
 	const FMinimalSceneTextures& SceneTextures,
 	const FViewInfo& View,
-	FVirtualShadowMapArray& VirtualShadowMapArray,
-	FRDGTextureRef ShadowMaskBits );
+	FVirtualShadowMapArray& VirtualShadowMapArray );
 
 void CompositeVirtualShadowMapMask(
 	FRDGBuilder& GraphBuilder,
 	const FIntRect ScissorRect,
-	const FSSDSignalTextures& InputSignal,
+	const FRDGTextureRef Input,
 	FRDGTextureRef OutputShadowMaskTexture);
