Fix transition errors in GPU-Scene and UnifiedBuffer.

diff --git a/Engine/Source/Runtime/RenderCore/Private/UnifiedBuffer.cpp b/Engine/Source/Runtime/RenderCore/Private/UnifiedBuffer.cpp
--- a/Engine/Source/Runtime/RenderCore/Private/UnifiedBuffer.cpp
+++ b/Engine/Source/Runtime/RenderCore/Private/UnifiedBuffer.cpp
@@ -361,6 +361,9 @@ RENDERCORE_API bool ResizeResourceIfNeeded<FRWBufferStructured>(FRHICommandList&
 		FRWBufferStructured NewBuffer;
 		NewBuffer.Initialize(DebugName, 16, NumBytes / 16, 0);
 
+		RHICmdList.Transition(FRHITransitionInfo(Buffer.UAV, ERHIAccess::Unknown, ERHIAccess::SRVCompute));
+		RHICmdList.Transition(FRHITransitionInfo(NewBuffer.UAV, ERHIAccess::Unknown, ERHIAccess::UAVCompute));
+
 		// Copy data to new buffer
 		uint32 CopyBytes = FMath::Min(NumBytes, Buffer.NumBytes);
 		MemcpyResource(RHICmdList, NewBuffer, Buffer, CopyBytes);
@@ -419,15 +422,20 @@ RENDERCORE_API bool ResizeResourceSOAIfNeeded<FRWBufferStructured>(FRHICommandLi
 		FRWBufferStructured NewBuffer;
 		NewBuffer.Initialize(DebugName, 16, NumBytes / 16, 0);
 
+		RHICmdList.Transition(FRHITransitionInfo(Buffer.UAV, ERHIAccess::Unknown, ERHIAccess::SRVCompute));
+		RHICmdList.Transition(FRHITransitionInfo(NewBuffer.UAV, ERHIAccess::Unknown, ERHIAccess::UAVCompute));
+
 		// Copy data to new buffer
 		uint32 OldArraySize = Buffer.NumBytes / NumArrays;
 		uint32 NewArraySize = NumBytes / NumArrays;
-		
+
+		RHICmdList.BeginUAVOverlap(NewBuffer.UAV);
 		uint32 CopyBytes = FMath::Min(NewArraySize, OldArraySize);
 		for( uint32 i = 0; i < NumArrays; i++ )
 		{
 			MemcpyResource( RHICmdList, NewBuffer, Buffer, CopyBytes, i * NewArraySize, i * OldArraySize );
 		}
+		RHICmdList.EndUAVOverlap(NewBuffer.UAV);
 
 		Buffer = NewBuffer;
 		return true;
diff --git a/Engine/Source/Runtime/Renderer/Private/GPUScene.cpp b/Engine/Source/Runtime/Renderer/Private/GPUScene.cpp
--- a/Engine/Source/Runtime/Renderer/Private/GPUScene.cpp
+++ b/Engine/Source/Runtime/Renderer/Private/GPUScene.cpp
@@ -474,6 +474,46 @@ void FGPUScene::UpdateInternal(FRHICommandListImmediate& RHICmdList, FScene& Sce
 	checkSlow(PrimitivesToUpdate.Num() == 0);
 }
 
+namespace 
+{
+
+	class FUAVTransitionStateScopeHelper
+	{
+	public:
+		FUAVTransitionStateScopeHelper(FRHICommandListImmediate& InRHICmdList, const FUnorderedAccessViewRHIRef& InUAV, ERHIAccess InitialState, ERHIAccess InFinalState = ERHIAccess::None) :
+			RHICmdList(InRHICmdList),
+			UAV(InUAV),
+			CurrentState(InitialState),
+			FinalState(InFinalState)
+		{
+		}
+
+		~FUAVTransitionStateScopeHelper()
+		{
+			if (FinalState != ERHIAccess::None)
+			{
+				TransitionTo(FinalState);
+			}
+		}
+
+		void TransitionTo(ERHIAccess NewState)
+		{
+			if (CurrentState != NewState)
+			{
+				RHICmdList.Transition(FRHITransitionInfo(UAV, CurrentState, NewState));
+				CurrentState = NewState;
+			}
+
+		}
+
+		FRHICommandListImmediate& RHICmdList;
+		FUnorderedAccessViewRHIRef UAV;
+		ERHIAccess CurrentState;
+		ERHIAccess FinalState;
+	};
+
+};
+
 template<typename FUploadDataSourceAdapter>
 void FGPUScene::UploadGeneral(FRHICommandListImmediate& RHICmdList, FScene *Scene, FUploadDataSourceAdapter &UploadDataSourceAdapter)
 {
@@ -494,11 +534,12 @@ void FGPUScene::UploadGeneral(FRHICommandListImmediate& RHICmdList, FScene *Scen
 		FRWBufferStructured* MirrorResourceGPU = GetMirrorGPU(*this);
 
 		const uint32 SizeReserve = FMath::RoundUpToPowerOfTwo(FMath::Max(DynamicPrimitivesOffset, 256));
-		const bool bResizedPrimitiveData = ResizeResourceIfNeeded(RHICmdList, *MirrorResourceGPU, SizeReserve * sizeof(FPrimitiveSceneShaderData::Data), TEXT("PrimitiveData"));
+		const bool bResizedPrimitiveData = ResizeResourceIfNeeded(RHICmdList, *MirrorResourceGPU, SizeReserve * sizeof(FPrimitiveSceneShaderData::Data), TEXT("GPUScene.PrimitiveData"));
 
 		const uint32 InstanceDataNumArrays = FInstanceSceneShaderData::InstanceDataStrideInFloat4s;
 		const uint32 InstanceDataSizeReserve = FMath::RoundUpToPowerOfTwo(FMath::Max(InstanceDataAllocator.GetMaxSize(), 256));
-		const bool bResizedInstanceData = ResizeResourceSOAIfNeeded(RHICmdList, InstanceDataBuffer, InstanceDataSizeReserve * sizeof(FInstanceSceneShaderData::Data), InstanceDataNumArrays, TEXT("InstanceData"));
+		const bool bResizedInstanceData = ResizeResourceSOAIfNeeded(RHICmdList, InstanceDataBuffer, InstanceDataSizeReserve * sizeof(FInstanceSceneShaderData::Data), InstanceDataNumArrays, TEXT("GPUScene.InstanceData"));
+		FUAVTransitionStateScopeHelper InstancedataTransitionHelper(RHICmdList, InstanceDataBuffer.UAV, bResizedInstanceData ? ERHIAccess::Unknown : ERHIAccess::SRVMask, ERHIAccess::SRVMask);
 		InstanceDataSOAStride = InstanceDataSizeReserve;
 
 		if (Scene != nullptr)
@@ -508,7 +549,7 @@ void FGPUScene::UploadGeneral(FRHICommandListImmediate& RHICmdList, FScene *Scen
 		}
 
 		const uint32 LightMapDataBufferSize = FMath::RoundUpToPowerOfTwo(FMath::Max(LightmapDataAllocator.GetMaxSize(), 256));
-		const bool bResizedLightmapData = ResizeResourceIfNeeded(RHICmdList, LightmapDataBuffer, LightMapDataBufferSize * sizeof(FLightmapSceneShaderData::Data), TEXT("LightmapData"));
+		const bool bResizedLightmapData = ResizeResourceIfNeeded(RHICmdList, LightmapDataBuffer, LightMapDataBufferSize * sizeof(FLightmapSceneShaderData::Data), TEXT("GPUScene.LightmapData"));
 
 		// These should always be in sync with each other.
 		check(InstanceDataToClear.Num() == InstanceDataAllocator.GetMaxSize());
@@ -634,7 +675,7 @@ void FGPUScene::UploadGeneral(FRHICommandListImmediate& RHICmdList, FScene *Scen
 				}
 			};
 
-			ERHIAccess CurrentAccess = ERHIAccess::Unknown;
+			FUAVTransitionStateScopeHelper PrimitiveDataTransitionHelper(RHICmdList, MirrorResourceGPU->UAV, ERHIAccess::Unknown, ERHIAccess::SRVMask);
 
 			const int32 MaxPrimitivesUploads = GetMaxPrimitivesUpdate(NumPrimitiveDataUploads, FPrimitiveSceneShaderData::PrimitiveDataStrideInFloat4s);
 			if (MaxPrimitivesUploads == NumPrimitiveDataUploads)
@@ -657,9 +698,7 @@ void FGPUScene::UploadGeneral(FRHICommandListImmediate& RHICmdList, FScene *Scen
 					RangeCount == 1
 				);
 
-				RHICmdList.Transition(FRHITransitionInfo(MirrorResourceGPU->UAV, CurrentAccess, ERHIAccess::UAVCompute));
-				CurrentAccess = ERHIAccess::UAVCompute;
-
+				PrimitiveDataTransitionHelper.TransitionTo(ERHIAccess::UAVCompute);
 				PrimitiveUploadBuffer.ResourceUploadTo(RHICmdList, *MirrorResourceGPU, true);
 			}
 			else
@@ -677,8 +716,7 @@ void FGPUScene::UploadGeneral(FRHICommandListImmediate& RHICmdList, FScene *Scen
 						ProcessPrimitiveFn(ItemIndex, false /* threaded */);
 					}
 
-					RHICmdList.Transition(FRHITransitionInfo(MirrorResourceGPU->UAV, CurrentAccess, ERHIAccess::UAVCompute));
-					CurrentAccess = ERHIAccess::UAVCompute;
+					PrimitiveDataTransitionHelper.TransitionTo(ERHIAccess::UAVCompute);
 
 					{
 						QUICK_SCOPE_CYCLE_COUNTER(UploadTo);
@@ -686,8 +724,6 @@ void FGPUScene::UploadGeneral(FRHICommandListImmediate& RHICmdList, FScene *Scen
 					}
 				}
 			}
-
-			RHICmdList.Transition(FRHITransitionInfo(MirrorResourceGPU->UAV, CurrentAccess, ERHIAccess::SRVMask));
 		}
 
 		if (Scene != nullptr)
@@ -870,13 +906,8 @@ void FGPUScene::UploadGeneral(FRHICommandListImmediate& RHICmdList, FScene *Scen
 						RangeCount == 1
 					);
 				}
-				RHICmdList.Transition(FRHITransitionInfo(InstanceDataBuffer.UAV, ERHIAccess::Unknown, ERHIAccess::UAVCompute));
+				InstancedataTransitionHelper.TransitionTo(ERHIAccess::UAVCompute);
 				InstanceUploadBuffer.ResourceUploadTo(RHICmdList, InstanceDataBuffer, false);
-				RHICmdList.Transition(FRHITransitionInfo(InstanceDataBuffer.UAV, ERHIAccess::UAVCompute, ERHIAccess::SRVMask));
-			}
-			else if (bResizedInstanceData)
-			{
-				RHICmdList.Transition(FRHITransitionInfo(InstanceDataBuffer.UAV, ERHIAccess::Unknown, ERHIAccess::SRVMask));
 			}
 
 			if( Scene != nullptr && Scene->InstanceBVH.GetNumDirty() > 0 )
@@ -910,7 +941,7 @@ void FGPUScene::UploadGeneral(FRHICommandListImmediate& RHICmdList, FScene *Scen
 
 			if (NumLightmapDataUploads > 0)
 			{
-				ERHIAccess CurrentAccess = ERHIAccess::Unknown;
+				FUAVTransitionStateScopeHelper LightMapTransitionHelper(RHICmdList, LightmapDataBuffer.UAV, ERHIAccess::Unknown, ERHIAccess::SRVMask);
 
 				// GPUCULL_TODO: This code is wrong: the intention is to break it up into batches such that the uploaded data fits in the max buffer size
 				//               However, what it does do is break it up into batches of MaxLightmapsUploads (while iterating over primitives). This is bad
@@ -935,13 +966,9 @@ void FGPUScene::UploadGeneral(FRHICommandListImmediate& RHICmdList, FScene *Scen
 						}
 					}
 
-					RHICmdList.Transition(FRHITransitionInfo(LightmapDataBuffer.UAV, CurrentAccess, ERHIAccess::UAVCompute));
-					CurrentAccess = ERHIAccess::UAVCompute;
-
+					LightMapTransitionHelper.TransitionTo(ERHIAccess::UAVCompute);
 					LightmapUploadBuffer.ResourceUploadTo(RHICmdList, LightmapDataBuffer, false);
 				}
-
-				RHICmdList.Transition(FRHITransitionInfo(LightmapDataBuffer.UAV, CurrentAccess, ERHIAccess::SRVMask));
 			}
 
 			if (PrimitiveUploadBuffer.GetNumBytes() > (uint32)GGPUSceneMaxPooledUploadBufferSize)
