* Add hair strands view uniform buffer to simplify resources binding.
* Simplify hair resources binding for virtual shadow map.

diff --git a/Engine/Shaders/Private/VirtualShadowMaps/VirtualShadowMapProjection.usf b/Engine/Shaders/Private/VirtualShadowMaps/VirtualShadowMapProjection.usf
--- a/Engine/Shaders/Private/VirtualShadowMaps/VirtualShadowMapProjection.usf
+++ b/Engine/Shaders/Private/VirtualShadowMaps/VirtualShadowMapProjection.usf
@@ -167,9 +167,6 @@ int SMRTRayCount;					// 0 = off
 int SMRTSamplesPerRay;
 float SMRTRayLengthScale;			// Directional lights
 float SMRTCotMaxRayAngleFromLight;	// Spot/point lights
-// HAIR_TODO: move all data into a HairStrands UB
-uint bUseHairData;
-Texture2D<uint4> HairCategorizationTexture;
 
 // One pass projection
 StructuredBuffer< int > VirtualShadowMapIdRemap;
@@ -181,6 +178,11 @@ RWTexture2D< uint > RWShadowMaskBits;
 int LightUniformVirtualShadowMapId;
 RWTexture2D< float4 > RWSignal;
 
+// Type of input data consume by the page allocation (i.e., data read from the source buffer: Gbuffer, HairStrands data, ...)
+#define INPUT_TYPE_GBUFFER 0
+#define INPUT_TYPE_HAIRSTRANDS 1
+uint InputType;
+
 struct FProjectLightResult
 {
 	half ShadowFactor;
@@ -226,7 +228,8 @@ FProjectLightResult ProjectLight(
 	BRANCH
 	if (bInLightRegion)
 	{
-		const half3 EstimatedGeoWorldNormal = bUseHairData > 0 ? L : GBufferNormal;
+		const bool bIsHair = InputType == INPUT_TYPE_HAIRSTRANDS;
+		const half3 EstimatedGeoWorldNormal = bIsHair ? L : GBufferNormal;
 
 		float3 OffsetTranslatedWorldPosition = TranslatedWorldPosition + EstimatedGeoWorldNormal * NormalOffsetWorld;
 
@@ -298,6 +301,7 @@ void VirtualShadowMapProjection(
 	uint3	DispatchThreadId	: SV_DispatchThreadID )
 {
 	// Morton order within a group so page access/atomics are more coherent and wave-swizzled gradients are possible.
+	uint2 PixelCoord = DispatchThreadId.xy;
 	uint2 LocalPixelPos = 8 * GroupId.xy + MortonDecode( GroupIndex );
 	uint2 PixelPos = LocalPixelPos + uint2( ProjectionRect.xy );
 	if ( any( PixelPos >= uint2( ProjectionRect.zw ) ) )
@@ -306,10 +310,14 @@ void VirtualShadowMapProjection(
 	}
 	
 	float DeviceZ = SceneTexturesStruct.SceneDepthTexture.Load( int3( PixelPos, 0 ) ).r;
-	if (bUseHairData > 0)
+	if (InputType == INPUT_TYPE_HAIRSTRANDS)
 	{
-		FCategorizationData CatData = DecodeCategorizationData( HairCategorizationTexture.Load( int3( PixelPos, 0 ) ) );
+		FCategorizationData CatData = DecodeCategorizationData(HairStrands.HairCategorizationTexture.Load( int3(PixelPos, 0 ) ) );
 		DeviceZ = CatData.ClosestDepth;
+		if (CatData.PixelCoverage < 0)
+		{
+			return;
+		}
 	}
 	const float SceneDepth = ConvertFromDeviceZ( DeviceZ );
 
diff --git a/Engine/Source/Runtime/Renderer/Private/ShadowRendering.cpp b/Engine/Source/Runtime/Renderer/Private/ShadowRendering.cpp
--- a/Engine/Source/Runtime/Renderer/Private/ShadowRendering.cpp
+++ b/Engine/Source/Runtime/Renderer/Private/ShadowRendering.cpp
@@ -2004,7 +2004,7 @@ void FDeferredShadingSceneRenderer::RenderShadowProjections(
 							View,
 							VirtualShadowMapArray,
 							ScissorRect,
-							nullptr,	// Hair
+							EVirtualShadowMapProjectionInputType::GBuffer,
 							VisibleLightInfo.FindShadowClipmapForView(&View));
 					}
 					else
@@ -2017,7 +2017,7 @@ void FDeferredShadingSceneRenderer::RenderShadowProjections(
 							View,
 							VirtualShadowMapArray,
 							ScissorRect,
-							nullptr,	// Hair
+							EVirtualShadowMapProjectionInputType::GBuffer,
 							VirtualShadowMaps[0]);
 					}
 
@@ -2050,34 +2050,33 @@ void FDeferredShadingSceneRenderer::RenderShadowProjections(
 					// Sub-pixel shadow (no denoising for hair)
 					if (HairStrands::HasViewHairStrandsData(View) && ScreenShadowMaskSubPixelTexture)
 					{
-						// TODO: This is ugly, but just a temporary workaround for branch divergence!
-						FRDGTextureRef HairCategorizationTexture = (*View.HairStrandsViewData.UniformBuffer)->HairCategorizationTexture;
+						FRDGTextureRef HairSignalTexture;
 						if (VisibleLightInfo.VirtualShadowMapClipmaps.Num() > 0)
 						{
-							SignalTexture = RenderVirtualShadowMapProjection(
+							HairSignalTexture = RenderVirtualShadowMapProjection(
 								GraphBuilder,
 								SceneTextures,
 								View,
 								VirtualShadowMapArray,
 								ScissorRect,
-								HairCategorizationTexture,
+								EVirtualShadowMapProjectionInputType::HairStrands,
 								VisibleLightInfo.FindShadowClipmapForView(&View));
 						}
 						else
 						{
 							check(VirtualShadowMaps.Num() == 1);
-							SignalTexture = RenderVirtualShadowMapProjection(
+							HairSignalTexture = RenderVirtualShadowMapProjection(
 								GraphBuilder,
 								SceneTextures,
 								View,
 								VirtualShadowMapArray,
 								ScissorRect,
-								HairCategorizationTexture,
+								EVirtualShadowMapProjectionInputType::HairStrands,
 								VirtualShadowMaps[0]);
 						}
 
 						// Composite into sub pixel mask
-						CompositeVirtualShadowMapMask(GraphBuilder, ScissorRect, SignalTexture, ScreenShadowMaskSubPixelTexture);
+						CompositeVirtualShadowMapMask(GraphBuilder, ScissorRect, HairSignalTexture, ScreenShadowMaskSubPixelTexture);
 					}
 				}
 			}
diff --git a/Engine/Source/Runtime/Renderer/Private/VirtualShadowMaps/VirtualShadowMapProjection.cpp b/Engine/Source/Runtime/Renderer/Private/VirtualShadowMaps/VirtualShadowMapProjection.cpp
--- a/Engine/Source/Runtime/Renderer/Private/VirtualShadowMaps/VirtualShadowMapProjection.cpp
+++ b/Engine/Source/Runtime/Renderer/Private/VirtualShadowMaps/VirtualShadowMapProjection.cpp
@@ -23,6 +23,7 @@
 #include "ShadowRendering.h"
 #include "SceneRendering.h"
 #include "VirtualShadowMapClipmap.h"
+#include "HairStrands/HairStrandsData.h"
 
 static TAutoConsoleVariable<float> CVarContactShadowLength(
 	TEXT( "r.Shadow.Virtual.ContactShadowLength" ),
@@ -180,6 +181,7 @@ class FVirtualShadowMapProjectionCS : public FGlobalShader
 	BEGIN_SHADER_PARAMETER_STRUCT(FParameters, )
 		SHADER_PARAMETER_STRUCT_INCLUDE(FVirtualShadowMapSamplingParameters, SamplingParameters)
 		SHADER_PARAMETER_RDG_UNIFORM_BUFFER(FSceneTextureUniformParameters, SceneTexturesStruct)
+		SHADER_PARAMETER_RDG_UNIFORM_BUFFER(FHairStrandsViewUniformParameters, HairStrands)
 		SHADER_PARAMETER_STRUCT_REF(FViewUniformShaderParameters, View)		
 		SHADER_PARAMETER(FIntVector4, ProjectionRect)
 		SHADER_PARAMETER(float, ContactShadowLength)
@@ -189,8 +191,7 @@ class FVirtualShadowMapProjectionCS : public FGlobalShader
 		SHADER_PARAMETER(float, SMRTRayLengthScale)
 		SHADER_PARAMETER(float, SMRTCotMaxRayAngleFromLight)
 		SHADER_PARAMETER(int32, DebugOutputType)
-		SHADER_PARAMETER_RDG_TEXTURE(Texture2D<uint4>, HairCategorizationTexture)
-		SHADER_PARAMETER(uint32, bUseHairData)
+		SHADER_PARAMETER(uint32, InputType)
 		// One pass projection parameters
 		SHADER_PARAMETER_STRUCT_REF(FForwardLightData, ForwardLightData)
 		SHADER_PARAMETER_RDG_BUFFER_SRV(StructuredBuffer< uint >, VirtualShadowMapIdRemap)	// TODO: Move to VSM UB?
@@ -243,7 +244,7 @@ static void RenderVirtualShadowMapProjectionCommon(
 	const FViewInfo& View,
 	FVirtualShadowMapArray& VirtualShadowMapArray,
 	const FIntRect ProjectionRect,
-	FRDGTextureRef HairCategorization,
+	EVirtualShadowMapProjectionInputType InputType,
 	FRDGTextureRef OutputTexture,
 	const FLightSceneProxy* LightProxy = nullptr,
 	int32 VirtualShadowMapId = INDEX_NONE)
@@ -256,8 +257,8 @@ static void RenderVirtualShadowMapProjectionCommon(
 	PassParameters->DebugOutputType = CVarVirtualShadowMapDebugProjection.GetValueOnRenderThread();
 	PassParameters->ContactShadowLength = CVarContactShadowLength.GetValueOnRenderThread();
 	PassParameters->NormalOffsetWorld = CVarNormalOffsetWorld.GetValueOnRenderThread();
-	PassParameters->bUseHairData = HairCategorization != nullptr ? 1 : 0;
-	PassParameters->HairCategorizationTexture = HairCategorization ? HairCategorization : GSystemTextures.GetBlackDummy(GraphBuilder);
+	PassParameters->InputType = uint32(InputType);
+	PassParameters->HairStrands = HairStrands::BindHairStrandsViewUniformParameters(View);
 
 	bool bDirectionalLight = false;
 	bool bOnePassProjection = LightProxy == nullptr;
@@ -338,7 +339,7 @@ FRDGTextureRef RenderVirtualShadowMapProjectionOnePass(
 		View,
 		VirtualShadowMapArray,
 		ProjectionRect,
-		nullptr,		// TODO: Hair
+		EVirtualShadowMapProjectionInputType::GBuffer, // TODO: Hair
 		ShadowMaskBits);
 
 	return ShadowMaskBits;
@@ -365,7 +366,7 @@ FRDGTextureRef RenderVirtualShadowMapProjection(
 	const FViewInfo& View,
 	FVirtualShadowMapArray& VirtualShadowMapArray,
 	const FIntRect ScissorRect,
-	FRDGTextureRef HairCategorization,
+	EVirtualShadowMapProjectionInputType InputType,
 	FProjectedShadowInfo* ShadowInfo)
 {	
 	FRDGTextureRef SignalTexture = CreateSignalTexture(GraphBuilder, SceneTextures.Config.Extent);
@@ -375,7 +376,7 @@ FRDGTextureRef RenderVirtualShadowMapProjection(
 		View,
 		VirtualShadowMapArray,
 		ScissorRect,
-		HairCategorization,
+		InputType,
 		SignalTexture,
 		ShadowInfo->GetLightSceneInfo().Proxy,
 		ShadowInfo->VirtualShadowMaps[0]->ID);
@@ -388,7 +389,7 @@ FRDGTextureRef RenderVirtualShadowMapProjection(
 	const FViewInfo& View,
 	FVirtualShadowMapArray& VirtualShadowMapArray,
 	const FIntRect ScissorRect,
-	FRDGTextureRef HairCategorization,
+	EVirtualShadowMapProjectionInputType InputType,
 	const TSharedPtr<FVirtualShadowMapClipmap>& Clipmap)
 {
 	FRDGTextureRef SignalTexture = CreateSignalTexture(GraphBuilder, SceneTextures.Config.Extent);
@@ -398,7 +399,7 @@ FRDGTextureRef RenderVirtualShadowMapProjection(
 		View,
 		VirtualShadowMapArray,		
 		ScissorRect,
-		HairCategorization,
+		InputType,
 		SignalTexture,
 		Clipmap->GetLightSceneInfo().Proxy,
 		Clipmap->GetVirtualShadowMap()->ID);
diff --git a/Engine/Source/Runtime/Renderer/Private/VirtualShadowMaps/VirtualShadowMapProjection.h b/Engine/Source/Runtime/Renderer/Private/VirtualShadowMaps/VirtualShadowMapProjection.h
--- a/Engine/Source/Runtime/Renderer/Private/VirtualShadowMaps/VirtualShadowMapProjection.h
+++ b/Engine/Source/Runtime/Renderer/Private/VirtualShadowMaps/VirtualShadowMapProjection.h
@@ -13,13 +13,19 @@
 
 class FVirtualShadowMapClipmap;
 
+enum class EVirtualShadowMapProjectionInputType
+{
+	GBuffer = 0,
+	HairStrands = 1
+};
+
 FRDGTextureRef RenderVirtualShadowMapProjection(
 	FRDGBuilder& GraphBuilder,
 	const FMinimalSceneTextures& SceneTextures,
 	const FViewInfo& View,
 	FVirtualShadowMapArray& VirtualShadowMapArray,
 	const FIntRect ScissorRect,
-	FRDGTextureRef HairCategorization,
+	EVirtualShadowMapProjectionInputType InputType,
 	const TSharedPtr<FVirtualShadowMapClipmap>& Clipmap);
 
 FRDGTextureRef RenderVirtualShadowMapProjection(
@@ -28,7 +34,7 @@ FRDGTextureRef RenderVirtualShadowMapProjection(
 	const FViewInfo& View,
 	FVirtualShadowMapArray& VirtualShadowMapArray,
 	const FIntRect ScissorRect,
-	FRDGTextureRef HairCategorization,
+	EVirtualShadowMapProjectionInputType InputType,
 	FProjectedShadowInfo* ShadowInfo);
 
 FRDGTextureRef RenderVirtualShadowMapProjectionOnePass(
