[ODSC] - FRayTracingMeshProcessor now implements the TryGetShaders interface.

- Extended FMaterialShaderTypes and FMaterialShaders to support all graphics frequences.  This allows FRayTracingMeshProcessor to use the RayGen, RayMiss, RayHitGroup, and RayCallable material shaders.
- FLightmapRayTracingMeshProcessor does not implement the TryGetShaders functionality because when rendering lightmaps it needs the actual shaders and not ones from default material.
- In FMaterial::TryGetShaders don't assert on a complete shader map.  This is valid in the following scenario:
    1. We don't find our shaders in our material map.
    2. We don't find the shaders we want on the default material because the default material doesn't have the appropriate usage flags.
    3. In this case we are going to request a shader compile for a variant that should not be cached for this material, shader, vertex factory type.
    4. From the material's perspective the shader map is complete.

todo:
- In the future it makes sense to block on shader compilation in-line to ensure we get TryGetShaders when we can't fail

diff --git a/Engine/Plugins/Experimental/GPULightmass/Source/GPULightmass/Private/LightmapRayTracing.cpp b/Engine/Plugins/Experimental/GPULightmass/Source/GPULightmass/Private/LightmapRayTracing.cpp
--- a/Engine/Plugins/Experimental/GPULightmass/Source/GPULightmass/Private/LightmapRayTracing.cpp
+++ b/Engine/Plugins/Experimental/GPULightmass/Source/GPULightmass/Private/LightmapRayTracing.cpp
@@ -144,7 +144,7 @@ static TShaderRef<FLightmapMaterialCHS> GetLightmapMaterialHitShader(const FMate
 	}
 }
 
-void FLightmapRayTracingMeshProcessor::Process(
+bool FLightmapRayTracingMeshProcessor::Process(
 	const FMeshBatch& RESTRICT MeshBatch,
 	uint64 BatchElementMask,
 	const FPrimitiveSceneProxy* RESTRICT PrimitiveSceneProxy,
@@ -180,6 +180,8 @@ void FLightmapRayTracingMeshProcessor::Process(
 		PassDrawRenderState,
 		RayTracingShaders,
 		ShaderElementData);
+
+	return true;
 }
 
 #endif
diff --git a/Engine/Plugins/Experimental/GPULightmass/Source/GPULightmass/Private/LightmapRayTracing.h b/Engine/Plugins/Experimental/GPULightmass/Source/GPULightmass/Private/LightmapRayTracing.h
--- a/Engine/Plugins/Experimental/GPULightmass/Source/GPULightmass/Private/LightmapRayTracing.h
+++ b/Engine/Plugins/Experimental/GPULightmass/Source/GPULightmass/Private/LightmapRayTracing.h
@@ -21,7 +21,7 @@ public:
 		: FRayTracingMeshProcessor(InCommandContext, nullptr, nullptr, InPassDrawRenderState)
 	{}
 
-	virtual void Process(
+	virtual bool Process(
 		const FMeshBatch& RESTRICT MeshBatch,
 		uint64 BatchElementMask,
 		const FPrimitiveSceneProxy* RESTRICT PrimitiveSceneProxy,
diff --git a/Engine/Source/Runtime/Engine/Private/Materials/MaterialShared.cpp b/Engine/Source/Runtime/Engine/Private/Materials/MaterialShared.cpp
--- a/Engine/Source/Runtime/Engine/Private/Materials/MaterialShared.cpp
+++ b/Engine/Source/Runtime/Engine/Private/Materials/MaterialShared.cpp
@@ -2635,14 +2635,16 @@ bool FMaterial::TryGetShaders(const FMaterialShaderTypes& InTypes, const FVertex
 #endif
 				if (CompilingShaderMapId != 0u)
 				{
-					check(!bShaderMapComplete);
-					if (InVertexFactoryType)
+					if (!bShaderMapComplete)
 					{
-						FMeshMaterialShaderType::BeginCompileShaderPipeline(EShaderCompileJobPriority::ForceLocal, CompilingShaderMapId, kUniqueShaderPermutationId, ShaderPlatform, PermutationFlags, this, RenderingThreadPendingCompilerEnvironment, InVertexFactoryType, InTypes.PipelineType, CompileJobs, nullptr, nullptr);
-					}
-					else
-					{
-						FMaterialShaderType::BeginCompileShaderPipeline(EShaderCompileJobPriority::ForceLocal, CompilingShaderMapId, ShaderPlatform, PermutationFlags, this, RenderingThreadPendingCompilerEnvironment, InTypes.PipelineType, CompileJobs, nullptr, nullptr);
+						if (InVertexFactoryType)
+						{
+							FMeshMaterialShaderType::BeginCompileShaderPipeline(EShaderCompileJobPriority::ForceLocal, CompilingShaderMapId, kUniqueShaderPermutationId, ShaderPlatform, PermutationFlags, this, RenderingThreadPendingCompilerEnvironment, InVertexFactoryType, InTypes.PipelineType, CompileJobs, nullptr, nullptr);
+						}
+						else
+						{
+							FMaterialShaderType::BeginCompileShaderPipeline(EShaderCompileJobPriority::ForceLocal, CompilingShaderMapId, ShaderPlatform, PermutationFlags, this, RenderingThreadPendingCompilerEnvironment, InTypes.PipelineType, CompileJobs, nullptr, nullptr);
+						}
 					}
 				}
 			}
@@ -2650,7 +2652,7 @@ bool FMaterial::TryGetShaders(const FMaterialShaderTypes& InTypes, const FVertex
 	}
 	else
 	{
-		for (int32 FrequencyIndex = 0; FrequencyIndex < SF_NumGraphicsFrequencies; ++FrequencyIndex)
+		for (int32 FrequencyIndex = 0; FrequencyIndex < SF_NumFrequencies; ++FrequencyIndex)
 		{
 			const FShaderType* ShaderType = InTypes.ShaderType[FrequencyIndex];
 			if (ShaderType)
@@ -2663,10 +2665,11 @@ bool FMaterial::TryGetShaders(const FMaterialShaderTypes& InTypes, const FVertex
 				}
 				else
 				{
+					bMissingShader = true;
+
 #if WITH_ODSC
 					if (FPlatformProperties::RequiresCookedData())
 					{
-						bMissingShader = true;
 
 						const FString MaterialName = GetFullPath();
 						const FString VFTypeName(InVertexFactoryType ? InVertexFactoryType->GetName() : TEXT(""));
@@ -2678,10 +2681,9 @@ bool FMaterial::TryGetShaders(const FMaterialShaderTypes& InTypes, const FVertex
 					}
 					else
 #endif
+					if (CompilingShaderMapId != 0u)
 					{
-						check(!bShaderMapComplete);
-						bMissingShader = true;
-						if (CompilingShaderMapId != 0u)
+						if (!bShaderMapComplete)
 						{
 							if (InVertexFactoryType)
 							{
diff --git a/Engine/Source/Runtime/Engine/Public/MaterialShaderType.h b/Engine/Source/Runtime/Engine/Public/MaterialShaderType.h
--- a/Engine/Source/Runtime/Engine/Public/MaterialShaderType.h
+++ b/Engine/Source/Runtime/Engine/Public/MaterialShaderType.h
@@ -182,8 +182,8 @@ public:
 struct FMaterialShaderTypes
 {
 	const FShaderPipelineType* PipelineType;
-	const FShaderType* ShaderType[SF_NumGraphicsFrequencies];
-	int32 PermutationId[SF_NumGraphicsFrequencies];
+	const FShaderType* ShaderType[SF_NumFrequencies];
+	int32 PermutationId[SF_NumFrequencies];
 
 	inline FMaterialShaderTypes() : PipelineType(nullptr) { FMemory::Memzero(ShaderType); FMemory::Memzero(PermutationId); }
 
@@ -207,7 +207,7 @@ struct FMaterialShaders
 {
 	const FShaderMapBase* ShaderMap;
 	FShaderPipeline* Pipeline;
-	FShader* Shaders[SF_NumGraphicsFrequencies];
+	FShader* Shaders[SF_NumFrequencies];
 
 	inline FMaterialShaders() : ShaderMap(nullptr), Pipeline(nullptr) { FMemory::Memzero(Shaders); }
 
diff --git a/Engine/Source/Runtime/Renderer/Private/RayTracing/RayTracingMaterialHitShaders.cpp b/Engine/Source/Runtime/Renderer/Private/RayTracing/RayTracingMaterialHitShaders.cpp
--- a/Engine/Source/Runtime/Renderer/Private/RayTracing/RayTracingMaterialHitShaders.cpp
+++ b/Engine/Source/Runtime/Renderer/Private/RayTracing/RayTracingMaterialHitShaders.cpp
@@ -217,36 +217,46 @@ IMPLEMENT_GLOBAL_SHADER(FHiddenMaterialHitGroup, "/Engine/Private/RayTracing/Ray
 IMPLEMENT_GLOBAL_SHADER(FOpaqueShadowHitGroup, "/Engine/Private/RayTracing/RayTracingMaterialDefaultHitShaders.usf", "closesthit=OpaqueShadowCHS", SF_RayHitGroup);
 
 template<typename LightMapPolicyType>
-static TShaderRef<FMaterialCHS> GetMaterialHitShader(const FMaterial& RESTRICT MaterialResource, const FVertexFactory* VertexFactory, bool UseTextureLod)
+static bool GetMaterialHitShader(const FMaterial& RESTRICT MaterialResource, const FVertexFactory* VertexFactory, bool UseTextureLod, TShaderRef<FMaterialCHS>& OutShader)
 {
 	const bool bMaterialsCompiled = GCompileRayTracingMaterialAHS || GCompileRayTracingMaterialCHS;
 	checkf(bMaterialsCompiled, TEXT(""));
 
+	FMaterialShaderTypes ShaderTypes;
 	if ((MaterialResource.IsMasked() || MaterialResource.GetBlendMode() != BLEND_Opaque) && GCompileRayTracingMaterialAHS)
 	{
 		if(UseTextureLod)
 		{ 
-			return MaterialResource.GetShader<TMaterialCHS<LightMapPolicyType, true, true>>(VertexFactory->GetType());
+			ShaderTypes.AddShaderType<TMaterialCHS<LightMapPolicyType, true, true>>();
 		}
 		else
 		{
-			return MaterialResource.GetShader<TMaterialCHS<LightMapPolicyType, true, false>>(VertexFactory->GetType());
+			ShaderTypes.AddShaderType<TMaterialCHS<LightMapPolicyType, true, false>>();
 		}
 	}
 	else
 	{
 		if (UseTextureLod)
 		{
-			return MaterialResource.GetShader<TMaterialCHS<LightMapPolicyType, false, true>>(VertexFactory->GetType());
+			ShaderTypes.AddShaderType<TMaterialCHS<LightMapPolicyType, false, true>>();
 		}
 		else
 		{
-			return MaterialResource.GetShader<TMaterialCHS<LightMapPolicyType, false, false>>(VertexFactory->GetType());
+			ShaderTypes.AddShaderType<TMaterialCHS<LightMapPolicyType, false, false>>();
 		}
 	}
+
+	FMaterialShaders Shaders;
+	if (!MaterialResource.TryGetShaders(ShaderTypes, VertexFactory->GetType(), Shaders))
+	{
+		return false;
+	}
+
+	Shaders.TryGetShader(SF_RayHitGroup, OutShader);
+	return true;
 }
 
-void FRayTracingMeshProcessor::Process(
+bool FRayTracingMeshProcessor::Process(
 	const FMeshBatch& RESTRICT MeshBatch,
 	uint64 BatchElementMask,
 	const FPrimitiveSceneProxy* RESTRICT PrimitiveSceneProxy,
@@ -274,19 +284,34 @@ void FRayTracingMeshProcessor::Process(
 		switch (LightMapPolicy.GetIndirectPolicy())
 		{
 		case LMP_PRECOMPUTED_IRRADIANCE_VOLUME_INDIRECT_LIGHTING:
-			RayTracingShaders.RayHitGroupShader = GetMaterialHitShader<TUniformLightMapPolicy<LMP_PRECOMPUTED_IRRADIANCE_VOLUME_INDIRECT_LIGHTING>>(MaterialResource, VertexFactory, bUseTextureLOD);
+			if (!GetMaterialHitShader<TUniformLightMapPolicy<LMP_PRECOMPUTED_IRRADIANCE_VOLUME_INDIRECT_LIGHTING>>(MaterialResource, VertexFactory, bUseTextureLOD, RayTracingShaders.RayHitGroupShader))
+			{
+				return false;
+			}
 			break;
 		case LMP_LQ_LIGHTMAP:
-			RayTracingShaders.RayHitGroupShader = GetMaterialHitShader<TUniformLightMapPolicy<LMP_LQ_LIGHTMAP>>(MaterialResource, VertexFactory, bUseTextureLOD);
+			if (!GetMaterialHitShader<TUniformLightMapPolicy<LMP_LQ_LIGHTMAP>>(MaterialResource, VertexFactory, bUseTextureLOD, RayTracingShaders.RayHitGroupShader))
+			{
+				return false;
+			}
 			break;
 		case LMP_HQ_LIGHTMAP:
-			RayTracingShaders.RayHitGroupShader = GetMaterialHitShader<TUniformLightMapPolicy<LMP_HQ_LIGHTMAP>>(MaterialResource, VertexFactory, bUseTextureLOD);
+			if (!GetMaterialHitShader<TUniformLightMapPolicy<LMP_HQ_LIGHTMAP>>(MaterialResource, VertexFactory, bUseTextureLOD, RayTracingShaders.RayHitGroupShader))
+			{
+				return false;
+			}
 			break;
 		case LMP_DISTANCE_FIELD_SHADOWS_AND_HQ_LIGHTMAP:
-			RayTracingShaders.RayHitGroupShader = GetMaterialHitShader<TUniformLightMapPolicy<LMP_DISTANCE_FIELD_SHADOWS_AND_HQ_LIGHTMAP>>(MaterialResource, VertexFactory, bUseTextureLOD);
+			if (!GetMaterialHitShader<TUniformLightMapPolicy<LMP_DISTANCE_FIELD_SHADOWS_AND_HQ_LIGHTMAP>>(MaterialResource, VertexFactory, bUseTextureLOD, RayTracingShaders.RayHitGroupShader))
+			{
+				return false;
+			}
 			break;
 		case LMP_NO_LIGHTMAP:
-			RayTracingShaders.RayHitGroupShader = GetMaterialHitShader<TUniformLightMapPolicy<LMP_NO_LIGHTMAP>>(MaterialResource, VertexFactory, bUseTextureLOD);
+			if (!GetMaterialHitShader<TUniformLightMapPolicy<LMP_NO_LIGHTMAP>>(MaterialResource, VertexFactory, bUseTextureLOD, RayTracingShaders.RayHitGroupShader))
+			{
+				return false;
+			}
 			break;
 		default:
 			check(false);
@@ -294,7 +319,16 @@ void FRayTracingMeshProcessor::Process(
 	}
 	else
 	{
-		RayTracingShaders.RayHitGroupShader = MaterialResource.GetShader<FTrivialMaterialCHS>(VertexFactory->GetType());
+		FMaterialShaderTypes ShaderTypes;
+		ShaderTypes.AddShaderType<FTrivialMaterialCHS>();
+
+		FMaterialShaders Shaders;
+		if (!MaterialResource.TryGetShaders(ShaderTypes, VertexFactory->GetType(), Shaders))
+		{
+			return false;
+		}
+
+		Shaders.TryGetShader(SF_RayHitGroup, RayTracingShaders.RayHitGroupShader);
 	}
 
 	PassDrawRenderState.SetBlendState(TStaticBlendState<CW_RGBA, BO_Add, BF_One, BF_One, BO_Add, BF_Zero, BF_One>::GetRHI());
@@ -312,143 +346,163 @@ void FRayTracingMeshProcessor::Process(
 		PassDrawRenderState,
 		RayTracingShaders,
 		ShaderElementData);
+
+	return true;
 }
 
 void FRayTracingMeshProcessor::AddMeshBatch(const FMeshBatch& RESTRICT MeshBatch, uint64 BatchElementMask, const FPrimitiveSceneProxy* RESTRICT PrimitiveSceneProxy)
 {
-	// Caveat: there are also branches not emitting any MDC
-	// Mesh Batches can "null" when they have zero triangles.  Check the MaterialRenderProxy before accessing.
-	if (MeshBatch.bUseForMaterial && MeshBatch.MaterialRenderProxy && IsSupportedVertexFactoryType(MeshBatch.VertexFactory->GetType()))
+	if (!MeshBatch.bUseForMaterial || !IsSupportedVertexFactoryType(MeshBatch.VertexFactory->GetType()))
 	{
-		// Determine the mesh's material and blend mode.
-		const FMaterialRenderProxy* FallbackMaterialRenderProxyPtr = nullptr;
-		const FMaterial& Material = MeshBatch.MaterialRenderProxy->GetMaterialWithFallback(FeatureLevel, FallbackMaterialRenderProxyPtr);
+		return;
+	}
 
-		const FMaterialRenderProxy& MaterialRenderProxy = FallbackMaterialRenderProxyPtr ? *FallbackMaterialRenderProxyPtr : *MeshBatch.MaterialRenderProxy;
+	const FMaterialRenderProxy* FallbackMaterialRenderProxyPtr = MeshBatch.MaterialRenderProxy;
+	while (FallbackMaterialRenderProxyPtr)
+	{
+		const FMaterial* Material = FallbackMaterialRenderProxyPtr->GetMaterialNoFallback(FeatureLevel);
+		if (TryAddMeshBatch(MeshBatch, BatchElementMask, PrimitiveSceneProxy, -1, *FallbackMaterialRenderProxyPtr, *Material))
+		{
+			break;
+		}
+		FallbackMaterialRenderProxyPtr = FallbackMaterialRenderProxyPtr->GetFallback(FeatureLevel);
+	}
+}
 
-		const EBlendMode BlendMode = Material.GetBlendMode();
-		const FMaterialShadingModelField ShadingModels = Material.GetShadingModels();
+bool FRayTracingMeshProcessor::TryAddMeshBatch(
+	const FMeshBatch& RESTRICT MeshBatch,
+	uint64 BatchElementMask,
+	const FPrimitiveSceneProxy* RESTRICT PrimitiveSceneProxy,
+	int32 StaticMeshId,
+	const FMaterialRenderProxy& MaterialRenderProxy,
+	const FMaterial& Material
+)
+{
+	// Determine the mesh's material and blend mode.
+	const EBlendMode BlendMode = Material.GetBlendMode();
+	const FMaterialShadingModelField ShadingModels = Material.GetShadingModels();
 
-		// Only draw opaque materials.
-		if ((!PrimitiveSceneProxy || PrimitiveSceneProxy->ShouldRenderInMainPass())
-			&& ShouldIncludeDomainInMeshPass(Material.GetMaterialDomain()))
-		{
-			// Check for a cached light-map.
-			const bool bIsLitMaterial = ShadingModels.IsLit();
-			static const auto AllowStaticLightingVar = IConsoleManager::Get().FindTConsoleVariableDataInt(TEXT("r.AllowStaticLighting"));
-			const bool bAllowStaticLighting = (!AllowStaticLightingVar || AllowStaticLightingVar->GetValueOnRenderThread() != 0);
+	// Only draw opaque materials.
+	if ((!PrimitiveSceneProxy || PrimitiveSceneProxy->ShouldRenderInMainPass())
+		&& ShouldIncludeDomainInMeshPass(Material.GetMaterialDomain()))
+	{
+		// Check for a cached light-map.
+		const bool bIsLitMaterial = ShadingModels.IsLit();
+		static const auto AllowStaticLightingVar = IConsoleManager::Get().FindTConsoleVariableDataInt(TEXT("r.AllowStaticLighting"));
+		const bool bAllowStaticLighting = (!AllowStaticLightingVar || AllowStaticLightingVar->GetValueOnRenderThread() != 0);
 
-			const FLightMapInteraction LightMapInteraction = (bAllowStaticLighting && MeshBatch.LCI && bIsLitMaterial)
-				? MeshBatch.LCI->GetLightMapInteraction(FeatureLevel)
-				: FLightMapInteraction();
+		const FLightMapInteraction LightMapInteraction = (bAllowStaticLighting && MeshBatch.LCI && bIsLitMaterial)
+			? MeshBatch.LCI->GetLightMapInteraction(FeatureLevel)
+			: FLightMapInteraction();
 
-			// force LQ lightmaps based on system settings
-			const bool bPlatformAllowsHighQualityLightMaps = AllowHighQualityLightmaps(FeatureLevel);
-			const bool bAllowHighQualityLightMaps = bPlatformAllowsHighQualityLightMaps && LightMapInteraction.AllowsHighQualityLightmaps();
+		// force LQ lightmaps based on system settings
+		const bool bPlatformAllowsHighQualityLightMaps = AllowHighQualityLightmaps(FeatureLevel);
+		const bool bAllowHighQualityLightMaps = bPlatformAllowsHighQualityLightMaps && LightMapInteraction.AllowsHighQualityLightmaps();
 
-			const bool bAllowIndirectLightingCache = Scene && Scene->PrecomputedLightVolumes.Num() > 0;
-			const bool bUseVolumetricLightmap = Scene && Scene->VolumetricLightmapSceneData.HasData();
+		const bool bAllowIndirectLightingCache = Scene && Scene->PrecomputedLightVolumes.Num() > 0;
+		const bool bUseVolumetricLightmap = Scene && Scene->VolumetricLightmapSceneData.HasData();
 
-			{
-				static const auto CVarSupportLowQualityLightmap = IConsoleManager::Get().FindTConsoleVariableDataInt(TEXT("r.SupportLowQualityLightmaps"));
-				const bool bAllowLowQualityLightMaps = (!CVarSupportLowQualityLightmap) || (CVarSupportLowQualityLightmap->GetValueOnAnyThread() != 0);
+		{
+			static const auto CVarSupportLowQualityLightmap = IConsoleManager::Get().FindTConsoleVariableDataInt(TEXT("r.SupportLowQualityLightmaps"));
+			const bool bAllowLowQualityLightMaps = (!CVarSupportLowQualityLightmap) || (CVarSupportLowQualityLightmap->GetValueOnAnyThread() != 0);
 
-				switch (LightMapInteraction.GetType())
+			switch (LightMapInteraction.GetType())
+			{
+			case LMIT_Texture:
+				if (bAllowHighQualityLightMaps)
 				{
-				case LMIT_Texture:
-					if (bAllowHighQualityLightMaps)
-					{
-						const FShadowMapInteraction ShadowMapInteraction = (bAllowStaticLighting && MeshBatch.LCI && bIsLitMaterial)
-							? MeshBatch.LCI->GetShadowMapInteraction(FeatureLevel)
-							: FShadowMapInteraction();
-
-						if (ShadowMapInteraction.GetType() == SMIT_Texture)
-						{
-							Process(
-								MeshBatch,
-								BatchElementMask,
-								PrimitiveSceneProxy,
-								MaterialRenderProxy,
-								Material,
-								ShadingModels,
-								FUniformLightMapPolicy(LMP_DISTANCE_FIELD_SHADOWS_AND_HQ_LIGHTMAP),
-								MeshBatch.LCI);
-						}
-						else
-						{
-							Process(
-								MeshBatch,
-								BatchElementMask,
-								PrimitiveSceneProxy,
-								MaterialRenderProxy,
-								Material,
-								ShadingModels,
-								FUniformLightMapPolicy(LMP_HQ_LIGHTMAP),
-								MeshBatch.LCI);
-						}
-					}
-					else if (bAllowLowQualityLightMaps)
-					{
-						Process(
-							MeshBatch,
-							BatchElementMask,
-							PrimitiveSceneProxy,
-							MaterialRenderProxy,
-							Material,
-							ShadingModels,
-							FUniformLightMapPolicy(LMP_LQ_LIGHTMAP),
-							MeshBatch.LCI);
-					}
-					else
-					{
-						Process(
-							MeshBatch,
-							BatchElementMask,
-							PrimitiveSceneProxy,
-							MaterialRenderProxy,
-							Material,
-							ShadingModels,
-							FUniformLightMapPolicy(LMP_NO_LIGHTMAP),
-							MeshBatch.LCI);
-					}
-					break;
-				default:
-					if (bIsLitMaterial
-						&& bAllowStaticLighting
-						&& Scene
-						&& Scene->VolumetricLightmapSceneData.HasData()
-						&& PrimitiveSceneProxy
-						&& (PrimitiveSceneProxy->IsMovable()
-							|| PrimitiveSceneProxy->NeedsUnbuiltPreviewLighting()
-							|| PrimitiveSceneProxy->GetLightmapType() == ELightmapType::ForceVolumetric))
+					const FShadowMapInteraction ShadowMapInteraction = (bAllowStaticLighting && MeshBatch.LCI && bIsLitMaterial)
+						? MeshBatch.LCI->GetShadowMapInteraction(FeatureLevel)
+						: FShadowMapInteraction();
+
+					if (ShadowMapInteraction.GetType() == SMIT_Texture)
 					{
-						Process(
+						return Process(
 							MeshBatch,
 							BatchElementMask,
 							PrimitiveSceneProxy,
 							MaterialRenderProxy,
 							Material,
 							ShadingModels,
-							FUniformLightMapPolicy(LMP_PRECOMPUTED_IRRADIANCE_VOLUME_INDIRECT_LIGHTING),
+							FUniformLightMapPolicy(LMP_DISTANCE_FIELD_SHADOWS_AND_HQ_LIGHTMAP),
 							MeshBatch.LCI);
 					}
 					else
 					{
-						Process(
+						return Process(
 							MeshBatch,
 							BatchElementMask,
 							PrimitiveSceneProxy,
 							MaterialRenderProxy,
 							Material,
 							ShadingModels,
-							FUniformLightMapPolicy(LMP_NO_LIGHTMAP),
+							FUniformLightMapPolicy(LMP_HQ_LIGHTMAP),
 							MeshBatch.LCI);
 					}
-					break;
-				};
-			}
+				}
+				else if (bAllowLowQualityLightMaps)
+				{
+					return Process(
+						MeshBatch,
+						BatchElementMask,
+						PrimitiveSceneProxy,
+						MaterialRenderProxy,
+						Material,
+						ShadingModels,
+						FUniformLightMapPolicy(LMP_LQ_LIGHTMAP),
+						MeshBatch.LCI);
+				}
+				else
+				{
+					return Process(
+						MeshBatch,
+						BatchElementMask,
+						PrimitiveSceneProxy,
+						MaterialRenderProxy,
+						Material,
+						ShadingModels,
+						FUniformLightMapPolicy(LMP_NO_LIGHTMAP),
+						MeshBatch.LCI);
+				}
+				break;
+			default:
+				if (bIsLitMaterial
+					&& bAllowStaticLighting
+					&& Scene
+					&& Scene->VolumetricLightmapSceneData.HasData()
+					&& PrimitiveSceneProxy
+					&& (PrimitiveSceneProxy->IsMovable()
+						|| PrimitiveSceneProxy->NeedsUnbuiltPreviewLighting()
+						|| PrimitiveSceneProxy->GetLightmapType() == ELightmapType::ForceVolumetric))
+				{
+					return Process(
+						MeshBatch,
+						BatchElementMask,
+						PrimitiveSceneProxy,
+						MaterialRenderProxy,
+						Material,
+						ShadingModels,
+						FUniformLightMapPolicy(LMP_PRECOMPUTED_IRRADIANCE_VOLUME_INDIRECT_LIGHTING),
+						MeshBatch.LCI);
+				}
+				else
+				{
+					return Process(
+						MeshBatch,
+						BatchElementMask,
+						PrimitiveSceneProxy,
+						MaterialRenderProxy,
+						Material,
+						ShadingModels,
+						FUniformLightMapPolicy(LMP_NO_LIGHTMAP),
+						MeshBatch.LCI);
+				}
+				break;
+			};
 		}
 	}
+
+	return true;
 }
 
 static bool IsCompatibleFallbackPipelineSignature(FRayTracingPipelineStateSignature& B, FRayTracingPipelineStateSignature& A)
diff --git a/Engine/Source/Runtime/Renderer/Private/RayTracing/RayTracingMaterialHitShaders.h b/Engine/Source/Runtime/Renderer/Private/RayTracing/RayTracingMaterialHitShaders.h
--- a/Engine/Source/Runtime/Renderer/Private/RayTracing/RayTracingMaterialHitShaders.h
+++ b/Engine/Source/Runtime/Renderer/Private/RayTracing/RayTracingMaterialHitShaders.h
@@ -34,7 +34,7 @@ protected:
 	ERHIFeatureLevel::Type FeatureLevel;
 	FMeshPassProcessorRenderState PassDrawRenderState;
 
-	virtual void Process(
+	virtual bool Process(
 		const FMeshBatch& RESTRICT MeshBatch,
 		uint64 BatchElementMask,
 		const FPrimitiveSceneProxy* RESTRICT PrimitiveSceneProxy,
@@ -100,6 +100,17 @@ protected:
 			}
 		}
 	}
+
+private:
+
+	bool TryAddMeshBatch(
+		const FMeshBatch& RESTRICT MeshBatch,
+		uint64 BatchElementMask,
+		const FPrimitiveSceneProxy* RESTRICT PrimitiveSceneProxy,
+		int32 StaticMeshId,
+		const FMaterialRenderProxy& MaterialRenderProxy,
+		const FMaterial& Material
+	);
 };
 
 class RENDERER_API FHiddenMaterialHitGroup : public FGlobalShader
