Replaces SHADER_PARAMETER_RDG_BUFFER() with RDG_BUFFER_ACCESS() in Hair

diff --git a/Engine/Plugins/Runtime/HairStrands/Source/HairStrandsCore/Private/HairStrandsRendering.cpp b/Engine/Plugins/Runtime/HairStrands/Source/HairStrandsCore/Private/HairStrandsRendering.cpp
--- a/Engine/Plugins/Runtime/HairStrands/Source/HairStrandsCore/Private/HairStrandsRendering.cpp
+++ b/Engine/Plugins/Runtime/HairStrands/Source/HairStrandsCore/Private/HairStrandsRendering.cpp
@@ -435,7 +435,7 @@ class FHairInterpolationCS : public FGlobalShader
 		SHADER_PARAMETER_RDG_BUFFER_SRV(Buffer<uint>, HairStrandsVF_CullingIndirectBuffer)
 		SHADER_PARAMETER_RDG_BUFFER_SRV(Buffer<uint>, HairStrandsVF_CullingIndexBuffer)
 		SHADER_PARAMETER_RDG_BUFFER_SRV(Buffer<float>, HairStrandsVF_CullingRadiusScaleBuffer)
-		SHADER_PARAMETER_RDG_BUFFER(Buffer<uint>, HairStrandsVF_CullingIndirectBufferArgs)
+		RDG_BUFFER_ACCESS(HairStrandsVF_CullingIndirectBufferArgs, ERHIAccess::IndirectArgs)
 		
 		END_SHADER_PARAMETER_STRUCT()
 
@@ -932,7 +932,7 @@ class FHairTangentCS : public FGlobalShader
 		SHADER_PARAMETER_RDG_BUFFER_SRV(Buffer<uint>,	HairStrandsVF_CullingIndirectBuffer)
 		SHADER_PARAMETER_RDG_BUFFER_SRV(Buffer<uint>,	HairStrandsVF_CullingIndexBuffer)
 		SHADER_PARAMETER_RDG_BUFFER_UAV(RWBuffer, OutputTangentBuffer)
-		SHADER_PARAMETER_RDG_BUFFER(Buffer, IndirectBufferArgs)
+		RDG_BUFFER_ACCESS(IndirectBufferArgs, ERHIAccess::IndirectArgs)
 	END_SHADER_PARAMETER_STRUCT()
 
 public:
@@ -1013,7 +1013,7 @@ class FHairRaytracingGeometryCS : public FGlobalShader
 		SHADER_PARAMETER_RDG_BUFFER_SRV(Buffer,	HairStrandsVF_CullingIndirectBuffer)
 		SHADER_PARAMETER_RDG_BUFFER_SRV(Buffer,	HairStrandsVF_CullingIndexBuffer)
 		SHADER_PARAMETER_RDG_BUFFER_SRV(Buffer,	HairStrandsVF_CullingRadiusScaleBuffer)
-		SHADER_PARAMETER_RDG_BUFFER(Buffer,	HairStrandsVF_CullingIndirectBufferArgs)
+		RDG_BUFFER_ACCESS(HairStrandsVF_CullingIndirectBufferArgs, ERHIAccess::IndirectArgs)
 
 		SHADER_PARAMETER_RDG_BUFFER_SRV(Buffer, PositionOffsetBuffer)
 		SHADER_PARAMETER_RDG_BUFFER_SRV(Buffer, PositionBuffer)
diff --git a/Engine/Source/Runtime/Renderer/Private/HairStrands/HairStrandsClusters.cpp b/Engine/Source/Runtime/Renderer/Private/HairStrands/HairStrandsClusters.cpp
--- a/Engine/Source/Runtime/Renderer/Private/HairStrands/HairStrandsClusters.cpp
+++ b/Engine/Source/Runtime/Renderer/Private/HairStrands/HairStrandsClusters.cpp
@@ -164,7 +164,7 @@ class FHairClusterCullingLocalBlockPreFixSumCS : public FGlobalShader
 	SHADER_USE_PARAMETER_STRUCT(FHairClusterCullingLocalBlockPreFixSumCS, FGlobalShader);
 
 	BEGIN_SHADER_PARAMETER_STRUCT(FParameters, )
-		SHADER_PARAMETER_RDG_BUFFER(Buffer, DispatchIndirectParametersClusterCountDiv512)
+		RDG_BUFFER_ACCESS(DispatchIndirectParametersClusterCountDiv512, ERHIAccess::IndirectArgs)
 		SHADER_PARAMETER_RDG_BUFFER_SRV(Buffer, DispatchIndirectParametersClusterCount)
 		SHADER_PARAMETER_RDG_BUFFER_SRV(Buffer, GlobalIndexCountBuffer)
 		SHADER_PARAMETER_RDG_BUFFER_UAV(RWBuffer, PerBlocklIndexCountPreFixSumBuffer)
@@ -190,7 +190,7 @@ class FHairClusterCullingCompactVertexIdsLocalBlockCS : public FGlobalShader
 	SHADER_USE_PARAMETER_STRUCT(FHairClusterCullingCompactVertexIdsLocalBlockCS, FGlobalShader);
 
 	BEGIN_SHADER_PARAMETER_STRUCT(FParameters, )
-		SHADER_PARAMETER_RDG_BUFFER(Buffer, DispatchIndirectParametersBuffer)
+		RDG_BUFFER_ACCESS(DispatchIndirectParametersBuffer, ERHIAccess::IndirectArgs)
 		SHADER_PARAMETER_RDG_BUFFER_SRV(Buffer, DispatchIndirectParametersClusterCount)
 		SHADER_PARAMETER_RDG_BUFFER_SRV(Buffer, DispatchIndirectParametersClusterCount2D)
 		SHADER_PARAMETER_RDG_BUFFER_SRV(Buffer, DispatchIndirectParametersClusterCountDiv512)
@@ -223,7 +223,7 @@ class FHairClusterCullingPreFixSumCS : public FGlobalShader
 	SHADER_USE_PARAMETER_STRUCT(FHairClusterCullingPreFixSumCS, FGlobalShader);
 
 	BEGIN_SHADER_PARAMETER_STRUCT(FParameters, )
-		SHADER_PARAMETER_RDG_BUFFER(Buffer, DispatchIndirectParameters)
+		RDG_BUFFER_ACCESS(DispatchIndirectParameters, ERHIAccess::IndirectArgs)
 		SHADER_PARAMETER_RDG_BUFFER_SRV(Buffer, DispatchIndirectParametersClusterCount)
 		SHADER_PARAMETER_RDG_BUFFER_SRV(Buffer, GlobalIndexCountBuffer)
 		SHADER_PARAMETER_RDG_BUFFER_UAV(RWBuffer, GlobalIndexCountPreFixSumBuffer)
diff --git a/Engine/Source/Runtime/Renderer/Private/HairStrands/HairStrandsTransmittance.cpp b/Engine/Source/Runtime/Renderer/Private/HairStrands/HairStrandsTransmittance.cpp
--- a/Engine/Source/Runtime/Renderer/Private/HairStrands/HairStrandsTransmittance.cpp
+++ b/Engine/Source/Runtime/Renderer/Private/HairStrands/HairStrandsTransmittance.cpp
@@ -201,7 +201,7 @@ class FDeepTransmittanceMaskCS : public FGlobalShader
 		SHADER_PARAMETER_RDG_TEXTURE(Texture3D, HairLUTTexture)
 		SHADER_PARAMETER_RDG_BUFFER_SRV(StructuredBuffer, HairVisibilityNodeData)
 		SHADER_PARAMETER_RDG_BUFFER_SRV(StructuredBuffer, HairVisibilityNodeCoord)
-		SHADER_PARAMETER_RDG_BUFFER(StructuredBuffer, IndirectArgsBuffer)
+		RDG_BUFFER_ACCESS(IndirectArgsBuffer, ERHIAccess::IndirectArgs)
 
 		SHADER_PARAMETER_RDG_BUFFER_UAV(RWStructuredBuffer, OutputColor)
 		SHADER_PARAMETER_SAMPLER(SamplerState, LinearSampler)
diff --git a/Engine/Source/Runtime/Renderer/Private/HairStrands/HairStrandsVisibility.cpp b/Engine/Source/Runtime/Renderer/Private/HairStrands/HairStrandsVisibility.cpp
--- a/Engine/Source/Runtime/Renderer/Private/HairStrands/HairStrandsVisibility.cpp
+++ b/Engine/Source/Runtime/Renderer/Private/HairStrands/HairStrandsVisibility.cpp
@@ -3011,7 +3011,7 @@ class FVisiblityRasterComputeCS : public FGlobalShader
 		SHADER_PARAMETER_RDG_BUFFER_SRV(Buffer, HairStrandsVF_CullingIndirectBuffer)
 		SHADER_PARAMETER_RDG_BUFFER_SRV(Buffer, HairStrandsVF_CullingIndexBuffer)
 		SHADER_PARAMETER_RDG_BUFFER_SRV(Buffer, HairStrandsVF_CullingRadiusScaleBuffer)
-		SHADER_PARAMETER_RDG_BUFFER(Buffer, IndirectBufferArgs)
+		RDG_BUFFER_ACCESS(IndirectBufferArgs, ERHIAccess::IndirectArgs)
 		SHADER_PARAMETER_RDG_TEXTURE(Texture2D, SceneDepthTexture)
 		SHADER_PARAMETER_RDG_TEXTURE_UAV(RWTexture2D, OutHairCountTexture)
 		SHADER_PARAMETER_RDG_TEXTURE_UAV(RWTexture2D, OutVisibilityTexture0)
diff --git a/Engine/Source/Runtime/Renderer/Private/HairStrands/HairStrandsVoxelization.cpp b/Engine/Source/Runtime/Renderer/Private/HairStrands/HairStrandsVoxelization.cpp
--- a/Engine/Source/Runtime/Renderer/Private/HairStrands/HairStrandsVoxelization.cpp
+++ b/Engine/Source/Runtime/Renderer/Private/HairStrands/HairStrandsVoxelization.cpp
@@ -141,7 +141,7 @@ class FVirtualVoxelInjectOpaqueCS : public FGlobalShader
 		SHADER_PARAMETER(FVector2D, SceneDepthResolution)
 		SHADER_PARAMETER(uint32, VoxelBiasCount)
 		SHADER_PARAMETER(uint32, VoxelMarkCount)
-		SHADER_PARAMETER_RDG_BUFFER(StructuredBuffer, IndirectDispatchArgs)
+		RDG_BUFFER_ACCESS(IndirectDispatchArgs, ERHIAccess::IndirectArgs)
 		SHADER_PARAMETER_RDG_TEXTURE_UAV(RWTexture3D, OutPageTexture)
 
 		SHADER_PARAMETER_STRUCT_REF(FViewUniformShaderParameters, ViewUniformBuffer)
@@ -284,7 +284,7 @@ class FVoxelMarkValidPageIndex_ScatterCS : public FGlobalShader
 	SHADER_USE_PARAMETER_STRUCT(FVoxelMarkValidPageIndex_ScatterCS, FGlobalShader);
 
 	BEGIN_SHADER_PARAMETER_STRUCT(FParameters, )
-		SHADER_PARAMETER_RDG_BUFFER(Buffer<int>,		IndirectBufferArgs)
+		RDG_BUFFER_ACCESS(IndirectBufferArgs, ERHIAccess::IndirectArgs)
 		SHADER_PARAMETER_RDG_BUFFER_SRV(Buffer<uint4>,	PageIndexResolutionAndOffsetBuffer)
 		SHADER_PARAMETER_RDG_BUFFER_SRV(Buffer<uint>,	DeferredScatterCounter)
 		SHADER_PARAMETER_RDG_BUFFER_SRV(Buffer<uint2>,  DeferredScatterBuffer)
@@ -343,8 +343,8 @@ class FVoxelAllocateVoxelPageCS : public FGlobalShader
 		SHADER_PARAMETER(uint32, MacroGroupId)
 		SHADER_PARAMETER(uint32, PageCount)
 		SHADER_PARAMETER(uint32, CPU_PageIndexCount)
-		SHADER_PARAMETER(uint32, CPU_PageIndexOffset)		
-		SHADER_PARAMETER_RDG_BUFFER(Buffer, IndirectBufferArgs)
+		SHADER_PARAMETER(uint32, CPU_PageIndexOffset)
+		RDG_BUFFER_ACCESS(IndirectBufferArgs, ERHIAccess::IndirectArgs)
 		SHADER_PARAMETER_RDG_BUFFER_SRV(Buffer, PageIndexResolutionAndOffsetBuffer)
 		SHADER_PARAMETER_RDG_BUFFER_UAV(RWBuffer, PageIndexGlobalCounter)
 		SHADER_PARAMETER_RDG_BUFFER_UAV(RWBuffer, PageIndexBuffer)
@@ -440,7 +440,7 @@ class FVoxelIndPageClearCS : public FGlobalShader
 	BEGIN_SHADER_PARAMETER_STRUCT(FParameters, )
 		SHADER_PARAMETER_STRUCT(FVirtualVoxelCommonParameters, VirtualVoxelParams)
 		SHADER_PARAMETER_RDG_TEXTURE_UAV(RWTexture3D, OutPageTexture)
-		SHADER_PARAMETER_RDG_BUFFER(Buffer, IndirectDispatchBuffer)
+		RDG_BUFFER_ACCESS(IndirectDispatchBuffer, ERHIAccess::IndirectArgs)
 	END_SHADER_PARAMETER_STRUCT()
 
 public:
@@ -998,7 +998,7 @@ class FVoxelRasterComputeCS : public FGlobalShader
 		SHADER_PARAMETER_RDG_BUFFER_SRV(Buffer, HairStrandsVF_CullingIndexBuffer)
 		SHADER_PARAMETER_RDG_BUFFER_SRV(Buffer, HairStrandsVF_CullingRadiusScaleBuffer)
 		SHADER_PARAMETER_RDG_BUFFER_SRV(Buffer, VoxelizationViewInfoBuffer)
-		SHADER_PARAMETER_RDG_BUFFER(Buffer, IndirectBufferArgs)
+		RDG_BUFFER_ACCESS(IndirectBufferArgs, ERHIAccess::IndirectArgs)
 		SHADER_PARAMETER_RDG_TEXTURE_UAV(RWTexture3D, OutPageTexture)
 	END_SHADER_PARAMETER_STRUCT()
 
@@ -1229,7 +1229,7 @@ class FVirtualVoxelGenerateMipCS : public FGlobalShader
 		SHADER_PARAMETER(uint32, SourceMip)
 		SHADER_PARAMETER(uint32, TargetMip)
 
-		SHADER_PARAMETER_RDG_BUFFER(StructuredBuffer, IndirectDispatchArgs)
+		RDG_BUFFER_ACCESS(IndirectDispatchArgs, ERHIAccess::IndirectArgs)
 		SHADER_PARAMETER_RDG_TEXTURE_SRV(Texture3D, InDensityTexture)
 		SHADER_PARAMETER_RDG_TEXTURE_UAV(RWTexture3D, OutDensityTexture)
 
@@ -1278,7 +1278,7 @@ class FVirtualVoxelPatchPageIndexWithMipDataCS : public FGlobalShader
 		SHADER_PARAMETER(uint32, MipIt)
 
 		SHADER_PARAMETER_RDG_TEXTURE(Texture3D, DensityTexture)
-		SHADER_PARAMETER_RDG_BUFFER(Buffer, IndirectDispatchArgs)
+		RDG_BUFFER_ACCESS(IndirectDispatchArgs, ERHIAccess::IndirectArgs)
 		SHADER_PARAMETER_RDG_BUFFER_SRV(Buffer, PageIndexGlobalCounter)
 		SHADER_PARAMETER_RDG_BUFFER_SRV(Buffer<uint>, PageToPageIndexBuffer)
 		SHADER_PARAMETER_RDG_BUFFER_UAV(RWBuffer<uint>, OutPageIndexBuffer)
