Replaces SHADER_PARAMETER_RDG_BUFFER() with RDG_BUFFER_ACCESS() in Strata

diff --git a/Engine/Source/Runtime/Renderer/Private/Strata/Strata.cpp b/Engine/Source/Runtime/Renderer/Private/Strata/Strata.cpp
--- a/Engine/Source/Runtime/Renderer/Private/Strata/Strata.cpp
+++ b/Engine/Source/Runtime/Renderer/Private/Strata/Strata.cpp
@@ -505,7 +505,7 @@ class FStrataMaterialStencilClassificationPassPS : public FGlobalShader
 		SHADER_PARAMETER(FIntPoint, TileCount)
 		SHADER_PARAMETER(FIntPoint, OutputResolution)
 		SHADER_PARAMETER_RDG_BUFFER_SRV(Buffer, TileListBuffer)
-		SHADER_PARAMETER_RDG_BUFFER(Buffer, TileIndirectBuffer)
+		RDG_BUFFER_ACCESS(TileIndirectBuffer, ERHIAccess::IndirectArgs)
 		RENDER_TARGET_BINDING_SLOTS()
 	END_SHADER_PARAMETER_STRUCT()
 
