Replaces SHADER_PARAMETER_RDG_BUFFER() with RDG_BUFFER_ACCESS() in remaining files

diff --git a/Engine/Source/Runtime/Renderer/Private/GpuDebugRendering.cpp b/Engine/Source/Runtime/Renderer/Private/GpuDebugRendering.cpp
--- a/Engine/Source/Runtime/Renderer/Private/GpuDebugRendering.cpp
+++ b/Engine/Source/Runtime/Renderer/Private/GpuDebugRendering.cpp
@@ -178,7 +178,7 @@ namespace ShaderDrawDebug
 			SHADER_PARAMETER_STRUCT_REF(FViewUniformShaderParameters, View)
 			SHADER_PARAMETER_SRV(StructuredBuffer, LockedShaderDrawDebugPrimitive)
 			SHADER_PARAMETER_RDG_BUFFER_SRV(StructuredBuffer, ShaderDrawDebugPrimitive)
-			SHADER_PARAMETER_RDG_BUFFER(StructuredBuffer, IndirectBuffer)
+			RDG_BUFFER_ACCESS(IndirectBuffer, ERHIAccess::IndirectArgs)
 		END_SHADER_PARAMETER_STRUCT()
 
 		static bool ShouldCompilePermutation(const FGlobalShaderPermutationParameters& Parameters)
diff --git a/Engine/Source/Runtime/Renderer/Private/PostProcess/DiaphragmDOF.cpp b/Engine/Source/Runtime/Renderer/Private/PostProcess/DiaphragmDOF.cpp
--- a/Engine/Source/Runtime/Renderer/Private/PostProcess/DiaphragmDOF.cpp
+++ b/Engine/Source/Runtime/Renderer/Private/PostProcess/DiaphragmDOF.cpp
@@ -1159,8 +1159,8 @@ BEGIN_SHADER_PARAMETER_STRUCT(FDOFHybridScatterParameters, )
 
 	SHADER_PARAMETER_RDG_TEXTURE(Texture2D, BokehLUT)
 
-	SHADER_PARAMETER_RDG_BUFFER(Buffer<uint>, IndirectDrawParameter)
 	SHADER_PARAMETER_RDG_BUFFER_SRV(StructuredBuffer<float4>, ScatterDrawList)
+	RDG_BUFFER_ACCESS(IndirectDrawParameter, ERHIAccess::IndirectArgs)
 
 	RENDER_TARGET_BINDING_SLOTS()
 END_SHADER_PARAMETER_STRUCT()
diff --git a/Engine/Source/Runtime/Renderer/Private/PostProcess/PostProcessSubsurface.cpp b/Engine/Source/Runtime/Renderer/Private/PostProcess/PostProcessSubsurface.cpp
--- a/Engine/Source/Runtime/Renderer/Private/PostProcess/PostProcessSubsurface.cpp
+++ b/Engine/Source/Runtime/Renderer/Private/PostProcess/PostProcessSubsurface.cpp
@@ -533,7 +533,7 @@ public:
 		SHADER_PARAMETER_RDG_TEXTURE_UAV(RWTexture2D, SSSColorUAV)
 		SHADER_PARAMETER_RDG_TEXTURE_UAV(RWTexture2D, HistoryUAV)
 		SHADER_PARAMETER_RDG_BUFFER_SRV(Buffer<uint>, GroupBuffer)
-		SHADER_PARAMETER_RDG_BUFFER(Buffer<uint>, IndirectDispatchArgsBuffer)
+		RDG_BUFFER_ACCESS(IndirectDispatchArgsBuffer, ERHIAccess::IndirectArgs)
 		SHADER_PARAMETER_STRUCT(FSubsurfaceInput, SubsurfaceInput0)
 		SHADER_PARAMETER_SAMPLER(SamplerState, SubsurfaceSampler0)
 		SHADER_PARAMETER_STRUCT(FSubsurfaceInput, SubsurfaceInput1)	// History
diff --git a/Engine/Source/Runtime/Renderer/Private/ScreenSpaceRayTracing.cpp b/Engine/Source/Runtime/Renderer/Private/ScreenSpaceRayTracing.cpp
--- a/Engine/Source/Runtime/Renderer/Private/ScreenSpaceRayTracing.cpp
+++ b/Engine/Source/Runtime/Renderer/Private/ScreenSpaceRayTracing.cpp
@@ -489,7 +489,7 @@ class FScreenSpaceReflectionsPS : public FGlobalShader
 	BEGIN_SHADER_PARAMETER_STRUCT(FParameters, )
 		SHADER_PARAMETER_STRUCT_INCLUDE(FSSRCommonParameters, CommonParameters)
 		SHADER_PARAMETER_STRUCT_INCLUDE(FSSRPassCommonParameters, SSRPassCommonParameter)
-		SHADER_PARAMETER_RDG_BUFFER(Buffer<uint>, IndirectDrawParameter)			// FScreenSpaceReflectionsTileVS
+		RDG_BUFFER_ACCESS(IndirectDrawParameter, ERHIAccess::IndirectArgs)
 		SHADER_PARAMETER_RDG_BUFFER_SRV(StructuredBuffer<uint>, TileListData)		// FScreenSpaceReflectionsTileVS
 		SHADER_PARAMETER_STRUCT_REF(FStrataGlobalUniformParameters, Strata)
 		RENDER_TARGET_BINDING_SLOTS()
@@ -570,7 +570,7 @@ class FScreenSpaceTraceProbeCS : public FGlobalShader
 		SHADER_PARAMETER_STRUCT_INCLUDE(LumenProbeHierarchy::FHierarchyParameters, HierarchyParameters)
 		SHADER_PARAMETER_STRUCT_INCLUDE(LumenProbeHierarchy::FHierarchyLevelParameters, LevelParameters)
 		SHADER_PARAMETER(float, FurthestHZBStartMipLevel)
-		SHADER_PARAMETER_RDG_BUFFER(Buffer<uint>, DispatchParameters)
+		RDG_BUFFER_ACCESS(DispatchParameters, ERHIAccess::IndirectArgs)
 		SHADER_PARAMETER_RDG_TEXTURE_UAV(RWTexture2D, ProbeAtlasColorOutput)
 		SHADER_PARAMETER_RDG_TEXTURE_UAV(RWTexture2D<uint>, ProbeAtlasSampleMaskOutput)
 	END_SHADER_PARAMETER_STRUCT()
@@ -617,7 +617,7 @@ class FScreenSpaceCastProbeOcclusionCS : public FGlobalShader
 		SHADER_PARAMETER_STRUCT_INCLUDE(LumenProbeHierarchy::FIndirectLightingProbeOcclusionParameters, ProbeOcclusionParameters)
 		SHADER_PARAMETER_STRUCT_INCLUDE(LumenProbeHierarchy::FIndirectLightingProbeOcclusionOutputParameters, ProbeOcclusionOutputParameters)
 		SHADER_PARAMETER(int32, DispatchOffset)
-		SHADER_PARAMETER_RDG_BUFFER(Buffer<uint>, DispatchParameters)
+		RDG_BUFFER_ACCESS(DispatchParameters, ERHIAccess::IndirectArgs)
 	END_SHADER_PARAMETER_STRUCT()
 		
 	class FTileClassificationDim : SHADER_PERMUTATION_ENUM_CLASS("DIM_PROBE_OCCLUSION_CLASSIFICATION", LumenProbeHierarchy::EProbeOcclusionClassification);
diff --git a/Engine/Source/Runtime/Renderer/Private/ShaderPrint.cpp b/Engine/Source/Runtime/Renderer/Private/ShaderPrint.cpp
--- a/Engine/Source/Runtime/Renderer/Private/ShaderPrint.cpp
+++ b/Engine/Source/Runtime/Renderer/Private/ShaderPrint.cpp
@@ -185,7 +185,7 @@ namespace ShaderPrint
 			SHADER_PARAMETER_STRUCT_REF(FUniformBufferParameters, UniformBufferParameters)
 			SHADER_PARAMETER_RDG_BUFFER_SRV(StructuredBuffer<ShaderPrintItem>, ValuesBuffer)
 			SHADER_PARAMETER_RDG_BUFFER_UAV(RWStructuredBuffer<ShaderPrintItem>, RWSymbolsBuffer)
-			SHADER_PARAMETER_RDG_BUFFER(StructuredBuffer<uint>, IndirectDispatchArgsBuffer)
+			RDG_BUFFER_ACCESS(IndirectDispatchArgsBuffer, ERHIAccess::IndirectArgs)
 		END_SHADER_PARAMETER_STRUCT()
 
 		static bool ShouldCompilePermutation(FGlobalShaderPermutationParameters const& Parameters)
@@ -228,7 +228,7 @@ namespace ShaderPrint
 			SHADER_PARAMETER_STRUCT_REF(FUniformBufferParameters, UniformBufferParameters)
 			SHADER_PARAMETER_TEXTURE(Texture2D, MiniFontTexture)
 			SHADER_PARAMETER_RDG_BUFFER_SRV(StructuredBuffer<ShaderPrintItem>, SymbolsBuffer)
-			SHADER_PARAMETER_RDG_BUFFER(StructuredBuffer<uint>, IndirectDrawArgsBuffer)
+			RDG_BUFFER_ACCESS(IndirectDrawArgsBuffer, ERHIAccess::IndirectArgs)
 		END_SHADER_PARAMETER_STRUCT()
 
 		static bool ShouldCompilePermutation(FGlobalShaderPermutationParameters const& Parameters)
diff --git a/Engine/Source/Runtime/Renderer/Private/SingleLayerWaterRendering.cpp b/Engine/Source/Runtime/Renderer/Private/SingleLayerWaterRendering.cpp
--- a/Engine/Source/Runtime/Renderer/Private/SingleLayerWaterRendering.cpp
+++ b/Engine/Source/Runtime/Renderer/Private/SingleLayerWaterRendering.cpp
@@ -358,7 +358,7 @@ static FSceneWithoutWaterTextures AddCopySceneWithoutWaterPass(
 BEGIN_SHADER_PARAMETER_STRUCT(FWaterCompositeParameters, )
 	SHADER_PARAMETER_STRUCT_INCLUDE(FWaterTileVS::FParameters, VS)
 	SHADER_PARAMETER_STRUCT_INCLUDE(FSingleLayerWaterCompositePS::FParameters, PS)
-	SHADER_PARAMETER_RDG_BUFFER(Buffer<uint>, IndirectDrawParameter)
+	RDG_BUFFER_ACCESS(IndirectDrawParameter, ERHIAccess::IndirectArgs)
 	RENDER_TARGET_BINDING_SLOTS()
 END_SHADER_PARAMETER_STRUCT()
 
