Removes the last uses of SHADER_PARAMETER_RDG_BUFFER()

diff --git a/Engine/Source/Runtime/Renderer/Private/InstanceCulling/InstanceCullingManager.h b/Engine/Source/Runtime/Renderer/Private/InstanceCulling/InstanceCullingManager.h
--- a/Engine/Source/Runtime/Renderer/Private/InstanceCulling/InstanceCullingManager.h
+++ b/Engine/Source/Runtime/Renderer/Private/InstanceCulling/InstanceCullingManager.h
@@ -44,9 +44,7 @@ struct FInstanceCullingResult
 	{
 		// GPUCULL_TODO: Maybe get dummy buffers?
 		OutParams.DrawIndirectArgsBuffer = DrawIndirectArgsBuffer;
-		OutParams.DrawIndirectArgsBufferAccess = DrawIndirectArgsBuffer;
 		OutParams.InstanceIdOffsetBuffer = InstanceIdOffsetBuffer;
-		OutParams.InstanceIdOffsetBufferAccess = InstanceIdOffsetBuffer;
 	}
 
 	static void CondGetDrawParameters(const FInstanceCullingResult* InstanceCullingResult, FInstanceCullingDrawParams& OutParams)
@@ -58,9 +56,7 @@ struct FInstanceCullingResult
 		else
 		{
 			OutParams.DrawIndirectArgsBuffer = nullptr;
-			OutParams.DrawIndirectArgsBufferAccess = nullptr;
 			OutParams.InstanceIdOffsetBuffer = nullptr;
-			OutParams.InstanceIdOffsetBufferAccess = nullptr;
 		}
 	}
 };
diff --git a/Engine/Source/Runtime/Renderer/Private/MeshDrawCommands.cpp b/Engine/Source/Runtime/Renderer/Private/MeshDrawCommands.cpp
--- a/Engine/Source/Runtime/Renderer/Private/MeshDrawCommands.cpp
+++ b/Engine/Source/Runtime/Renderer/Private/MeshDrawCommands.cpp
@@ -1478,9 +1478,7 @@ void FParallelMeshDrawCommandPass::BuildRenderingCommands(FRDGBuilder& GraphBuil
 		}
 	}
 	OutInstanceCullingDrawParams.DrawIndirectArgsBuffer = nullptr;
-	OutInstanceCullingDrawParams.DrawIndirectArgsBufferAccess = nullptr;
 	OutInstanceCullingDrawParams.InstanceIdOffsetBuffer = nullptr;
-	OutInstanceCullingDrawParams.InstanceIdOffsetBufferAccess = nullptr;
 
 #endif
 }
diff --git a/Engine/Source/Runtime/Renderer/Public/InstanceCulling/InstanceCullingContext.h b/Engine/Source/Runtime/Renderer/Public/InstanceCulling/InstanceCullingContext.h
--- a/Engine/Source/Runtime/Renderer/Public/InstanceCulling/InstanceCullingContext.h
+++ b/Engine/Source/Runtime/Renderer/Public/InstanceCulling/InstanceCullingContext.h
@@ -52,10 +52,8 @@ struct FInstanceCullingRdgParams
 };
 
 BEGIN_SHADER_PARAMETER_STRUCT(FInstanceCullingDrawParams, )
-	RDG_BUFFER_ACCESS(DrawIndirectArgsBufferAccess, ERHIAccess::IndirectArgs)
-	SHADER_PARAMETER_RDG_BUFFER(Buffer<uint32>, DrawIndirectArgsBuffer)
-	RDG_BUFFER_ACCESS(InstanceIdOffsetBufferAccess, ERHIAccess::VertexOrIndexBuffer)
-	SHADER_PARAMETER_RDG_BUFFER(Buffer<uint32>, InstanceIdOffsetBuffer)
+	RDG_BUFFER_ACCESS(DrawIndirectArgsBuffer, ERHIAccess::IndirectArgs)
+	RDG_BUFFER_ACCESS(InstanceIdOffsetBuffer, ERHIAccess::VertexOrIndexBuffer)
 END_SHADER_PARAMETER_STRUCT()
 
 /**
