Ensure that we mask off unwanted bits in the UDIM wrapped address.
Fixes bugs in VT mapping.

diff --git a/Engine/Source/Runtime/Renderer/Private/VT/VirtualTextureSystem.cpp b/Engine/Source/Runtime/Renderer/Private/VT/VirtualTextureSystem.cpp
--- a/Engine/Source/Runtime/Renderer/Private/VT/VirtualTextureSystem.cpp
+++ b/Engine/Source/Runtime/Renderer/Private/VT/VirtualTextureSystem.cpp
@@ -1539,6 +1539,7 @@ void FVirtualTextureSystem::GatherRequestsTask(const FGatherRequestsParameters&
 					Local_vPageX %= ProducerWidthInPages;
 					Local_vTileY %= ProducerHeightInPages;
 					Wrapped_vAddress = FMath::MortonCode2(Local_vPageX + AllocatedPageX) | (FMath::MortonCode2(Local_vTileY + AllocatedPageY) << 1);
+					Wrapped_vAddress &= (0xffffffff << (vLevel * vDimensions));
 				}
 			}
 
