[ODSC] - Fix crash in FRayTracingMeshProcessor due to null material.  Changing the code to soft fail similar to the other mesh processors.

 See FBasePassMeshProcessor::AddMeshBatch()

diff --git a/Engine/Source/Runtime/Renderer/Private/RayTracing/RayTracingMaterialHitShaders.cpp b/Engine/Source/Runtime/Renderer/Private/RayTracing/RayTracingMaterialHitShaders.cpp
--- a/Engine/Source/Runtime/Renderer/Private/RayTracing/RayTracingMaterialHitShaders.cpp
+++ b/Engine/Source/Runtime/Renderer/Private/RayTracing/RayTracingMaterialHitShaders.cpp
@@ -361,9 +361,12 @@ void FRayTracingMeshProcessor::AddMeshBatch(const FMeshBatch& RESTRICT MeshBatch
 	while (FallbackMaterialRenderProxyPtr)
 	{
 		const FMaterial* Material = FallbackMaterialRenderProxyPtr->GetMaterialNoFallback(FeatureLevel);
-		if (TryAddMeshBatch(MeshBatch, BatchElementMask, PrimitiveSceneProxy, -1, *FallbackMaterialRenderProxyPtr, *Material))
+		if (Material && Material->GetRenderingThreadShaderMap())
 		{
-			break;
+			if (TryAddMeshBatch(MeshBatch, BatchElementMask, PrimitiveSceneProxy, -1, *FallbackMaterialRenderProxyPtr, *Material))
+			{
+				break;
+			}
 		}
 		FallbackMaterialRenderProxyPtr = FallbackMaterialRenderProxyPtr->GetFallback(FeatureLevel);
 	}
