Fix issue with FShaderPlatformCachedIniValue reading several bool cvars as ints, disabling them if they are assigned =True in the ini file.

diff --git a/Engine/Source/Runtime/Core/Public/HAL/IConsoleManager.h b/Engine/Source/Runtime/Core/Public/HAL/IConsoleManager.h
--- a/Engine/Source/Runtime/Core/Public/HAL/IConsoleManager.h
+++ b/Engine/Source/Runtime/Core/Public/HAL/IConsoleManager.h
@@ -423,6 +423,10 @@ public:
 	{
 		OutIntValue = GetInt();
 	}
+	void GetValue(bool& OutBoolValue)
+	{
+		OutBoolValue = GetBool();
+	}
 	void GetValue(float& OutFloatValue)
 	{
 		OutFloatValue = GetFloat();
diff --git a/Engine/Source/Runtime/Engine/Private/ShaderCompiler/ShaderCompiler.cpp b/Engine/Source/Runtime/Engine/Private/ShaderCompiler/ShaderCompiler.cpp
--- a/Engine/Source/Runtime/Engine/Private/ShaderCompiler/ShaderCompiler.cpp
+++ b/Engine/Source/Runtime/Engine/Private/ShaderCompiler/ShaderCompiler.cpp
@@ -5077,8 +5077,8 @@ void GlobalBeginCompileShader(
 
 	if (IsMobilePlatform((EShaderPlatform)Target.Platform))
 	{
-		static FShaderPlatformCachedIniValue<int32> MobileEnableMovableSpotlightsIniValue(TEXT("/Script/Engine.RendererSettings"), TEXT("r.Mobile.EnableMovableSpotlights"));
-		static FShaderPlatformCachedIniValue<int32> MobileEnableMovableSpotlightsShadowIniValue(TEXT("/Script/Engine.RendererSettings"), TEXT("r.Mobile.EnableMovableSpotlightsShadow"));
+		static FShaderPlatformCachedIniValue<bool> MobileEnableMovableSpotlightsIniValue(TEXT("/Script/Engine.RendererSettings"), TEXT("r.Mobile.EnableMovableSpotlights"));
+		static FShaderPlatformCachedIniValue<bool> MobileEnableMovableSpotlightsShadowIniValue(TEXT("/Script/Engine.RendererSettings"), TEXT("r.Mobile.EnableMovableSpotlightsShadow"));
 
 		bool bMobileEnableMovableSpotlights = (MobileEnableMovableSpotlightsIniValue.Get((EShaderPlatform)Target.Platform) != 0);
 		Input.Environment.SetDefine(TEXT("PROJECT_MOBILE_ENABLE_MOVABLE_SPOTLIGHTS"), bMobileEnableMovableSpotlights ? 1 : 0);
diff --git a/Engine/Source/Runtime/RenderCore/Private/RenderUtils.cpp b/Engine/Source/Runtime/RenderCore/Private/RenderUtils.cpp
--- a/Engine/Source/Runtime/RenderCore/Private/RenderUtils.cpp
+++ b/Engine/Source/Runtime/RenderCore/Private/RenderUtils.cpp
@@ -1201,6 +1201,7 @@ Type FShaderPlatformCachedIniValue<Type>::Get(EShaderPlatform ShaderPlatform)
 template struct FShaderPlatformCachedIniValue<FString>;
 template struct FShaderPlatformCachedIniValue<int32>;
 template struct FShaderPlatformCachedIniValue<float>;
+template struct FShaderPlatformCachedIniValue<bool>;
 
 RENDERCORE_API int32 GUseForwardShading = 0;
 static FAutoConsoleVariableRef CVarForwardShading(
diff --git a/Engine/Source/Runtime/RenderCore/Private/Shader.cpp b/Engine/Source/Runtime/RenderCore/Private/Shader.cpp
--- a/Engine/Source/Runtime/RenderCore/Private/Shader.cpp
+++ b/Engine/Source/Runtime/RenderCore/Private/Shader.cpp
@@ -1508,8 +1508,8 @@ void ShaderMapAppendKeyString(EShaderPlatform Platform, FString& KeyString)
 		}
 
 		{
-			static FShaderPlatformCachedIniValue<int32> MobileEnableMovableSpotlightsIniValue(TEXT("/Script/Engine.RendererSettings"), TEXT("r.Mobile.EnableMovableSpotlights"));
-			static FShaderPlatformCachedIniValue<int32> MobileEnableMovableSpotlightsShadowIniValue(TEXT("/Script/Engine.RendererSettings"), TEXT("r.Mobile.EnableMovableSpotlightsShadow"));
+			static FShaderPlatformCachedIniValue<bool> MobileEnableMovableSpotlightsIniValue(TEXT("/Script/Engine.RendererSettings"), TEXT("r.Mobile.EnableMovableSpotlights"));
+			static FShaderPlatformCachedIniValue<bool> MobileEnableMovableSpotlightsShadowIniValue(TEXT("/Script/Engine.RendererSettings"), TEXT("r.Mobile.EnableMovableSpotlightsShadow"));
 
 			bool bMobileEnableMovableSpotlights = (MobileEnableMovableSpotlightsIniValue.Get(Platform) != 0);
 			KeyString += (bMobileEnableMovableSpotlights) ? TEXT("_MSPTL") : TEXT("");
diff --git a/Engine/Source/Runtime/Renderer/Private/MobileBasePassRendering.cpp b/Engine/Source/Runtime/Renderer/Private/MobileBasePassRendering.cpp
--- a/Engine/Source/Runtime/Renderer/Private/MobileBasePassRendering.cpp
+++ b/Engine/Source/Runtime/Renderer/Private/MobileBasePassRendering.cpp
@@ -56,7 +56,7 @@ static TAutoConsoleVariable<int32> CVarMobileUseHWsRGBEncoding(
 
 namespace MobileBasePass
 {
-	FShaderPlatformCachedIniValue<int32> MobileDynamicPointLightsUseStaticBranchIniValue(TEXT("/Script/Engine.RendererSettings"), TEXT("r.MobileDynamicPointLightsUseStaticBranch"));
+	FShaderPlatformCachedIniValue<bool> MobileDynamicPointLightsUseStaticBranchIniValue(TEXT("/Script/Engine.RendererSettings"), TEXT("r.MobileDynamicPointLightsUseStaticBranch"));
 	FShaderPlatformCachedIniValue<int32> MobileNumDynamicPointLightsIniValue(TEXT("/Script/Engine.RendererSettings"), TEXT("r.MobileNumDynamicPointLights"));
 };
 
@@ -146,7 +146,7 @@ FMobileBasePassMovableLightInfo::FMobileBasePassMovableLightInfo(const FPrimitiv
 
 void SetupMobileBasePassUniformParameters(
 	FRDGBuilder& GraphBuilder,
-	const FViewInfo& View,
+	const FViewInfo& View, 
 	EMobileBasePass BasePass,
 	FRDGTextureRef ScreenSpaceAOTexture,
 	FTextureRHIRef PixelProjectedReflectionTexture,
@@ -176,10 +176,10 @@ void SetupMobileBasePassUniformParameters(
 	BasePassParameters.PreIntegratedGFSampler = TStaticSamplerState<SF_Bilinear, AM_Clamp, AM_Clamp, AM_Clamp>::GetRHI();
 
 	if (PixelProjectedReflectionTexture.IsValid())
-	{
+			{
 		BasePassParameters.PlanarReflection.PlanarReflectionTexture = PixelProjectedReflectionTexture;
-		BasePassParameters.PlanarReflection.PlanarReflectionSampler = TStaticSamplerState<SF_Bilinear, AM_Clamp, AM_Clamp, AM_Clamp>::GetRHI();
-	}
+				BasePassParameters.PlanarReflection.PlanarReflectionSampler = TStaticSamplerState<SF_Bilinear, AM_Clamp, AM_Clamp, AM_Clamp>::GetRHI();
+			}
 
 	BasePassParameters.EyeAdaptationBuffer = GraphBuilder.CreateSRV(GetEyeAdaptationBuffer(GraphBuilder, View), PF_A32B32G32R32F);
 
@@ -311,21 +311,21 @@ void FMobileSceneRenderer::RenderMobileBasePass(FRHICommandListImmediate& RHICmd
 	SCOPE_CYCLE_COUNTER(STAT_BasePassDrawTime);
 	SCOPED_GPU_STAT(RHICmdList, Basepass);
 
-	RHICmdList.SetViewport(View.ViewRect.Min.X, View.ViewRect.Min.Y, 0, View.ViewRect.Max.X, View.ViewRect.Max.Y, 1);
-	View.ParallelMeshDrawCommandPasses[EMeshPass::BasePass].DispatchDraw(nullptr, RHICmdList);
-
-	if (View.Family->EngineShowFlags.Atmosphere)
-	{
-		View.ParallelMeshDrawCommandPasses[EMeshPass::SkyPass].DispatchDraw(nullptr, RHICmdList);
-	}
+		RHICmdList.SetViewport(View.ViewRect.Min.X, View.ViewRect.Min.Y, 0, View.ViewRect.Max.X, View.ViewRect.Max.Y, 1);
+		View.ParallelMeshDrawCommandPasses[EMeshPass::BasePass].DispatchDraw(nullptr, RHICmdList);
+		
+		if (View.Family->EngineShowFlags.Atmosphere)
+		{
+			View.ParallelMeshDrawCommandPasses[EMeshPass::SkyPass].DispatchDraw(nullptr, RHICmdList);
+		}
 
-	// editor primitives
+		// editor primitives
 	FMeshPassProcessorRenderState DrawRenderState;
-	DrawRenderState.SetBlendState(TStaticBlendStateWriteMask<CW_RGBA>::GetRHI());
-	DrawRenderState.SetDepthStencilAccess(Scene->DefaultBasePassDepthStencilAccess);
-	DrawRenderState.SetDepthStencilState(TStaticDepthStencilState<true, CF_DepthNearOrEqual>::GetRHI());
-	RenderMobileEditorPrimitives(RHICmdList, View, DrawRenderState);
-}
+			DrawRenderState.SetBlendState(TStaticBlendStateWriteMask<CW_RGBA>::GetRHI());
+			DrawRenderState.SetDepthStencilAccess(Scene->DefaultBasePassDepthStencilAccess);
+			DrawRenderState.SetDepthStencilState(TStaticDepthStencilState<true, CF_DepthNearOrEqual>::GetRHI());
+			RenderMobileEditorPrimitives(RHICmdList, View, DrawRenderState);
+		}
 
 void FMobileSceneRenderer::RenderMobileEditorPrimitives(FRHICommandList& RHICmdList, const FViewInfo& View, const FMeshPassProcessorRenderState& DrawRenderState)
 {
diff --git a/Engine/Source/Runtime/Renderer/Private/MobileBasePassRendering.h b/Engine/Source/Runtime/Renderer/Private/MobileBasePassRendering.h
--- a/Engine/Source/Runtime/Renderer/Private/MobileBasePassRendering.h
+++ b/Engine/Source/Runtime/Renderer/Private/MobileBasePassRendering.h
@@ -272,7 +272,7 @@ public:
 		HQReflectionCaptureBoxScalesArray.Bind(Initializer.ParameterMap, TEXT("CaptureBoxScalesArray"));
 
 		NumDynamicPointLightsParameter.Bind(Initializer.ParameterMap, TEXT("NumDynamicPointLights"));
-
+						
 		CSMDebugHintParams.Bind(Initializer.ParameterMap, TEXT("CSMDebugHint"));
 		UseCSMParameter.Bind(Initializer.ParameterMap, TEXT("UseCSM"));
 	}
@@ -371,7 +371,7 @@ namespace MobileBasePass
 
 	bool StationarySkyLightHasBeenApplied(const FScene* Scene, ELightMapPolicyType LightMapPolicyType);
 
-	extern FShaderPlatformCachedIniValue<int32> MobileDynamicPointLightsUseStaticBranchIniValue;
+	extern FShaderPlatformCachedIniValue<bool> MobileDynamicPointLightsUseStaticBranchIniValue;
 	extern FShaderPlatformCachedIniValue<int32> MobileNumDynamicPointLightsIniValue;
 };
 
@@ -399,7 +399,7 @@ public:
 		// We compile the point light shader combinations based on the project settings
 		static auto* MobileSkyLightPermutationCVar = IConsoleManager::Get().FindTConsoleVariableDataInt(TEXT("r.Mobile.SkyLightPermutation"));
 
-		const bool bMobileDynamicPointLightsUseStaticBranch = (MobileBasePass::MobileDynamicPointLightsUseStaticBranchIniValue.Get(Parameters.Platform) == 1);
+		const bool bMobileDynamicPointLightsUseStaticBranch = MobileBasePass::MobileDynamicPointLightsUseStaticBranchIniValue.Get(Parameters.Platform);
 		const int32 MobileNumDynamicPointLights = MobileBasePass::MobileNumDynamicPointLightsIniValue.Get(Parameters.Platform);
 		const int32 MobileSkyLightPermutationOptions = MobileSkyLightPermutationCVar->GetValueOnAnyThread();
 
