Fix a crash bug of rendering modulated shadow when the directional light set to dynamic.

diff --git a/Engine/Source/Runtime/Renderer/Private/ShadowRendering.cpp b/Engine/Source/Runtime/Renderer/Private/ShadowRendering.cpp
--- a/Engine/Source/Runtime/Renderer/Private/ShadowRendering.cpp
+++ b/Engine/Source/Runtime/Renderer/Private/ShadowRendering.cpp
@@ -2250,7 +2250,9 @@ void FMobileSceneRenderer::RenderModulatedShadowProjections(FRHICommandListImmed
 					if (ProjectedShadowInfo->bAllocated)
 					{
 						// Only project the shadow if it's large enough in this particular view (split screen, etc... may have shadows that are large in one view but irrelevantly small in others)
-						if (ProjectedShadowInfo->FadeAlphas[ViewIndex] > 1.0f / 256.0f)
+						if (ProjectedShadowInfo->FadeAlphas[ViewIndex] > 1.0f / 256.0f 
+							// Skip, if it is a whole scene directional shadow
+							&& !ProjectedShadowInfo->IsWholeSceneDirectionalShadow())
 						{
 							checkSlow(!ProjectedShadowInfo->bOnePassPointLightShadow);
 							ProjectedShadowInfo->RenderMobileModulatedShadowProjection(RHICmdList, ViewIndex, &View, LightSceneProxy, this);
