r.VT.ListPhysicalPools also reports total memory usage

diff --git a/Engine/Source/Runtime/Renderer/Private/VT/VirtualTextureSystem.cpp b/Engine/Source/Runtime/Renderer/Private/VT/VirtualTextureSystem.cpp
--- a/Engine/Source/Runtime/Renderer/Private/VT/VirtualTextureSystem.cpp
+++ b/Engine/Source/Runtime/Renderer/Private/VT/VirtualTextureSystem.cpp
@@ -322,6 +322,7 @@ void FVirtualTextureSystem::DumpFromConsole()
 
 void FVirtualTextureSystem::ListPhysicalPoolsFromConsole()
 {
+	uint64 TotalPhysicalMemory = 0u;
 	for(int32 i = 0; i < PhysicalSpaces.Num(); ++i)
 	{
 		if (PhysicalSpaces[i])
@@ -352,9 +353,12 @@ void FVirtualTextureSystem::ListPhysicalPoolsFromConsole()
 			UE_LOG(LogConsoleResponse, Display, TEXT("  Tiles Allocated= %i (%fMB)"), AllocatedTiles, AllocatedMemory);
 			UE_LOG(LogConsoleResponse, Display, TEXT("  Tiles Locked= %i (%fMB)"), LockedTiles, LockedMemory);
 			UE_LOG(LogConsoleResponse, Display, TEXT("  Tiles Mapped= %i"), PagePool.GetNumMappedPages());
+
+			TotalPhysicalMemory += TotalSizeInBytes;
 		}
 	}
 
+	uint64 TotalPageTableMemory = 0u;
 	for (int ID = 0; ID < 16; ID++)
 	{
 		const FVirtualTextureSpace* Space = Spaces[ID].Get();
@@ -386,7 +390,12 @@ void FVirtualTextureSystem::ListPhysicalPoolsFromConsole()
 			Allocator.GetNumAllocations(),
 			(int)(AllocatedRatio * 100.0),
 			(float)(AllocatedRatio * TotalSizeInBytes / 1024.0 / 1024.0));
+
+		TotalPageTableMemory += TotalSizeInBytes;
 	}
+
+	UE_LOG(LogConsoleResponse, Display, TEXT("TotalPageTableMemory: %fMB"), (double)TotalPageTableMemory / 1024.0 / 1024.0);
+	UE_LOG(LogConsoleResponse, Display, TEXT("TotalPhysicalMemory: %fMB"), (double)TotalPhysicalMemory / 1024.0 / 1024.0);
 }
 
 IAllocatedVirtualTexture* FVirtualTextureSystem::AllocateVirtualTexture(const FAllocatedVTDescription& Desc)
