Make sure stereo depth swapchain texture is used if appropriate (with no MSAA).

diff --git a/Engine/Source/Runtime/Renderer/Private/SceneTextures.cpp b/Engine/Source/Runtime/Renderer/Private/SceneTextures.cpp
--- a/Engine/Source/Runtime/Renderer/Private/SceneTextures.cpp
+++ b/Engine/Source/Runtime/Renderer/Private/SceneTextures.cpp
@@ -588,6 +588,15 @@ FSceneTextures& FMinimalSceneTextures::Create(FRDGBuilder& GraphBuilder, const F
 	FSceneTextures& SceneTextures = GraphBuilder.Blackboard.Create<FSceneTextures>(Config);
 
 	// Scene Depth
+
+	// If not using MSAA, we need to make sure to grab the stereo depth texture if appropriate.
+	FTexture2DRHIRef StereoDepthRHI;
+	if (Config.NumSamples == 1 && (StereoDepthRHI = FindStereoDepthTexture(Config.Extent, Config.NumSamples)) != nullptr)
+	{
+		SceneTextures.Depth = RegisterExternalTexture(GraphBuilder, StereoDepthRHI, TEXT("SceneDepthZ"));
+		SceneTextures.Stencil = GraphBuilder.CreateSRV(FRDGTextureSRVDesc::CreateWithPixelFormat(SceneTextures.Depth.Target, PF_X24_G8));
+	}
+	else
 	{
 		ETextureCreateFlags Flags = TexCreate_DepthStencilTargetable | TexCreate_ShaderResource | TexCreate_InputAttachmentRead | GFastVRamConfig.SceneDepth;
 
@@ -612,7 +621,7 @@ FSceneTextures& FMinimalSceneTextures::Create(FRDGBuilder& GraphBuilder, const F
 		{
 			Desc.NumSamples = 1;
 
-			if (FTexture2DRHIRef StereoDepthRHI = FindStereoDepthTexture(Config.Extent, Desc.NumSamples))
+			if ((StereoDepthRHI = FindStereoDepthTexture(Config.Extent, Desc.NumSamples)) != nullptr)
 			{
 				SceneTextures.Depth.Resolve = RegisterExternalTexture(GraphBuilder, StereoDepthRHI, TEXT("SceneDepthZ"));
 			}
