Submit potentially-expensive BVH build commands to the GPU as soon as possible. Avoids a GPU bubble in some CPU-limited cases.

diff --git a/Engine/Source/Runtime/Renderer/Private/DeferredShadingRenderer.cpp b/Engine/Source/Runtime/Renderer/Private/DeferredShadingRenderer.cpp
--- a/Engine/Source/Runtime/Renderer/Private/DeferredShadingRenderer.cpp
+++ b/Engine/Source/Runtime/Renderer/Private/DeferredShadingRenderer.cpp
@@ -1273,6 +1273,10 @@ bool FDeferredShadingSceneRenderer::DispatchRayTracingWorldUpdates(FRDGBuilder&
 					FViewInfo& View = Views[ViewIndex];
 					RHICmdList.BuildAccelerationStructure(View.RayTracingScene.RayTracingSceneRHI);
 				}
+
+				// Submit potentially expensive BVH build commands to the GPU as soon as possible.
+				// Avoids a GPU bubble in some CPU-limited cases.
+				RHICmdList.SubmitCommandsHint();
 			});
 		}
 	}
