Temporily disable Gen5TAA on PCD3D_SM6 to avoid deal breaker artifacts

diff --git a/Engine/Source/Runtime/Renderer/Private/PostProcess/TemporalAA.cpp b/Engine/Source/Runtime/Renderer/Private/PostProcess/TemporalAA.cpp
--- a/Engine/Source/Runtime/Renderer/Private/PostProcess/TemporalAA.cpp
+++ b/Engine/Source/Runtime/Renderer/Private/PostProcess/TemporalAA.cpp
@@ -104,8 +104,9 @@ inline bool DoesPlatformSupportTemporalHistoryUpscale(EShaderPlatform Platform)
 
 inline bool DoesPlatformSupportGen5TAA(EShaderPlatform Platform)
 {
+	// TODO(UE-111389): Fixes artifact on SP_PCD3D_SM6
 	return (
-		(IsPCPlatform(Platform) && IsFeatureLevelSupported(Platform, ERHIFeatureLevel::SM5)) ||
+		(IsPCPlatform(Platform) && IsFeatureLevelSupported(Platform, ERHIFeatureLevel::SM5) && Platform != SP_PCD3D_SM6) ||
 		FDataDrivenShaderPlatformInfo::GetSupportsGen5TemporalAA(Platform));
 }
 
