VT allocator supports non-square/non-power-2 allocations
- Still built on top of a quad-tree that manages square power-2 memory blocks
- Allocations can now be made of multiple blocks of different sizes, which are subdivided to fit the requested size

diff --git a/Engine/Source/Runtime/Renderer/Private/VT/VirtualTextureAllocator.cpp b/Engine/Source/Runtime/Renderer/Private/VT/VirtualTextureAllocator.cpp
--- a/Engine/Source/Runtime/Renderer/Private/VT/VirtualTextureAllocator.cpp
+++ b/Engine/Source/Runtime/Renderer/Private/VT/VirtualTextureAllocator.cpp
@@ -23,7 +23,10 @@ void FVirtualTextureAllocator::Initialize(uint32 InSize)
 	FreeList.Reset(LogSize + 1);
 
 	// Start with one empty block
-	AddressBlocks.Add(FAddressBlock(LogSize));
+	FAddressBlock DefaultBlock(LogSize);
+	DefaultBlock.State = EBlockState::FreeList;
+	AddressBlocks.Add(DefaultBlock);
+
 	SortedAddresses.Add(0u);
 	SortedIndices.Add(0u);
 
@@ -59,7 +62,9 @@ void FVirtualTextureAllocator::Grow()
 
 	// Add new root block
 	FAddressBlock RootBlock(LogSize);
+	RootBlock.State = EBlockState::AllocatedChildren;
 	RootBlock.FirstChild = OldRootIndex;
+
 	RootIndex = AcquireBlock();
 	AddressBlocks[RootIndex] = RootBlock;
 
@@ -93,6 +98,9 @@ void FVirtualTextureAllocator::Grow()
 		FAddressBlock& AddressBlock = AddressBlocks[Index];
 
 		// Place on free list
+		check(AddressBlock.State == EBlockState::None);
+		AddressBlock.State = EBlockState::FreeList;
+
 		AddressBlock.NextFree = FreeList[AddressBlock.vLogSize];
 		if (AddressBlock.NextFree != 0xffff)
 		{
@@ -112,20 +120,26 @@ void FVirtualTextureAllocator::Grow()
 
 int32 FVirtualTextureAllocator::AcquireBlock()
 {
-	if (GlobalFreeList == 0xffff)
+	int32 Index = GlobalFreeList;
+	if (Index == 0xffff)
 	{
-		return AddressBlocks.AddUninitialized();
+		Index = AddressBlocks.AddUninitialized();
 	}
-
-	int32 FreeBlock = GlobalFreeList;
-
-	GlobalFreeList = AddressBlocks[GlobalFreeList].NextFree;
-	if (GlobalFreeList != 0xffff)
+	else
 	{
-		AddressBlocks[GlobalFreeList].PrevFree = 0xffff;
+		check(AddressBlocks[Index].State == EBlockState::GlobalFreeList);
+		GlobalFreeList = AddressBlocks[Index].NextFree;
+		if (GlobalFreeList != 0xffff)
+		{
+			AddressBlocks[GlobalFreeList].PrevFree = 0xffff;
+		}
 	}
 
-	return FreeBlock;
+	// Debug fill memory to invalid value
+	FAddressBlock& AddressBlock = AddressBlocks[Index];
+	FMemory::Memset(AddressBlock, 0xCC);
+
+	return Index;
 }
 
 uint32 FVirtualTextureAllocator::FindAddressBlock(uint32 vAddress) const
@@ -156,17 +170,20 @@ FAllocatedVirtualTexture* FVirtualTextureAllocator::Find(uint32 vAddress, uint32
 	const FAddressBlock& AddressBlock = AddressBlocks[Index];
 	check(SortedAddresses[SortedIndex] == AddressBlock.vAddress);
 
+	FAllocatedVirtualTexture* AllocatedVT = nullptr;
 	const uint32 BlockSize = 1 << (vDimensions * AddressBlock.vLogSize);
 	if (vAddress >= AddressBlock.vAddress &&
 		vAddress < AddressBlock.vAddress + BlockSize)
 	{
-		check(!AddressBlock.VT || AddressBlock.VT->GetVirtualAddress() == AddressBlock.vAddress);
-		OutLocal_vAddress = vAddress - AddressBlock.vAddress;
+		AllocatedVT = AddressBlock.VT;
+		if (AllocatedVT)
+		{
+			OutLocal_vAddress = vAddress - AllocatedVT->GetVirtualAddress();
+		}
 		// TODO mip bias
-		return AddressBlock.VT;
 	}
 
-	return nullptr;
+	return AllocatedVT;
 }
 
 bool FVirtualTextureAllocator::TryAlloc(uint32 InLogSize)
@@ -182,128 +199,173 @@ bool FVirtualTextureAllocator::TryAlloc(uint32 InLogSize)
 	return false;
 }
 
-uint32 FVirtualTextureAllocator::Alloc(FAllocatedVirtualTexture* VT)
+void FVirtualTextureAllocator::SubdivideBlock(uint32 ParentIndex)
 {
-	// Pad out to square power of 2
-	const uint32 BlockSize = FMath::Max(VT->GetWidthInTiles(), VT->GetHeightInTiles());
-	const uint8 vLogSize = FMath::CeilLogTwo(BlockSize);
+	const uint32 NumChildren = (1 << vDimensions);
+	
+	const uint32 vLogSize = AddressBlocks[ParentIndex].vLogSize - 1u;
+	const uint32 vAddress = AddressBlocks[ParentIndex].vAddress;
 
-	// Find smallest free that fits
-	for (int i = vLogSize; i < FreeList.Num(); i++)
+	const int32 SortedIndex = FindAddressBlock(vAddress);
+	check(vAddress == SortedAddresses[SortedIndex]);
+
+	// Make room for newly added
+	SortedAddresses.InsertUninitialized(SortedIndex, NumChildren - 1u);
+	SortedIndices.InsertUninitialized(SortedIndex, NumChildren - 1u);
+	check(SortedAddresses.Num() == SortedIndices.Num());
+
+	uint16 FirstSiblingIndex = 0xffff;
+	uint16 PrevChildIndex = 0xffff;
+	for (uint32 Sibling = 0; Sibling < NumChildren; Sibling++)
 	{
-		uint16 FreeIndex = FreeList[i];
-		if (FreeIndex != 0xffff)
-		{
-			// Found free
-			uint16 AllocIndex = FreeIndex;
-			FAddressBlock* AllocBlock = &AddressBlocks[AllocIndex];
-			check(AllocBlock->VT == nullptr);
-			check(AllocBlock->PrevFree == 0xffff);
+		const int32 ChildBlockIndex = AcquireBlock();
+		const uint32 vChildAddress = vAddress + (Sibling << (vDimensions * vLogSize));
 
-			// Remove from free list
-			FreeList[i] = AllocBlock->NextFree;
-			if (AllocBlock->NextFree != 0xffff)
-			{
-				AddressBlocks[AllocBlock->NextFree].PrevFree = 0xffff;
-				AllocBlock->NextFree = 0xffff;
-			}
+		if (Sibling == 0u)
+		{
+			FirstSiblingIndex = ChildBlockIndex;
+			AddressBlocks[ParentIndex].FirstChild = FirstSiblingIndex;
+			// Replace parent block with new allocated block
+			SortedIndices[SortedIndex] = FirstSiblingIndex;
+		}
+		else
+		{
+			const int32 SortedIndexOffset = NumChildren - 1 - Sibling;
+			SortedAddresses[SortedIndex + SortedIndexOffset] = vChildAddress;
+			SortedIndices[SortedIndex + SortedIndexOffset] = ChildBlockIndex;
+			AddressBlocks[PrevChildIndex].NextSibling = ChildBlockIndex;
+		}
+	
+		FAddressBlock ChildBlock(vLogSize);
+		ChildBlock.vAddress = vChildAddress;
+		ChildBlock.Parent = ParentIndex;
+		ChildBlock.FirstSibling = FirstSiblingIndex;
+		ChildBlock.NextSibling = 0xffff;
+
+		// New child blocks start out on the free list
+		ChildBlock.State = EBlockState::FreeList;
+		ChildBlock.NextFree = FreeList[vLogSize];
+		if (ChildBlock.NextFree != 0xffff)
+		{
+			AddressBlocks[ChildBlock.NextFree].PrevFree = ChildBlockIndex;
+		}
+		FreeList[vLogSize] = ChildBlockIndex;
 
-			// Recursive subdivide to requested size
-			TArray<int32, TInlineAllocator<32>> NewBlocks;
-			while (AllocBlock->vLogSize > vLogSize)
-			{
-				const uint32 NumChildren = (1 << vDimensions);
+		AddressBlocks[ChildBlockIndex] = ChildBlock;
+		PrevChildIndex = ChildBlockIndex;
+	}
+}
 
-				// Create child blocks
-				const int32 FirstChildIndex = AcquireBlock();
-				int32 NextSibling = AcquireBlock();
+void FVirtualTextureAllocator::MarkBlockAllocated(uint32 Index, uint32 vAllocatedTileX, uint32 vAllocatedTileY, FAllocatedVirtualTexture* VT)
+{
+	FAddressBlock* AllocBlock = &AddressBlocks[Index];
+	check(AllocBlock->State == EBlockState::FreeList);
+	check(AllocBlock->FirstChild == 0xffff);
 
-				AllocBlock = &AddressBlocks[AllocIndex];
-				AllocBlock->FirstChild = FirstChildIndex;
+	const uint32 vLogSize = AllocBlock->vLogSize;
 
-				FAddressBlock FirstChildBlock(AllocBlock->vLogSize - 1);
-				FirstChildBlock.vAddress = AllocBlock->vAddress;
-				FirstChildBlock.Parent = AllocIndex;
-				FirstChildBlock.FirstSibling = FirstChildIndex;
-				FirstChildBlock.NextSibling = NextSibling;
-				AddressBlocks[FirstChildIndex] = FirstChildBlock;
+	// check to see if block is in the correct position
+	const uint32 vAllocatedTileX1 = vAllocatedTileX + VT->GetWidthInTiles();
+	const uint32 vAllocatedTileY1 = vAllocatedTileY + VT->GetHeightInTiles();
+	const uint32 vBlockAddress = AllocBlock->vAddress;
+	const uint32 vBlockTileX = FMath::ReverseMortonCode2(vBlockAddress);
+	const uint32 vBlockTileY = FMath::ReverseMortonCode2(vBlockAddress >> 1);
+	const uint32 BlockSize = (1u << vLogSize);
 
-				for (uint32 Sibling = 1; Sibling < NumChildren; Sibling++)
-				{
-					const int32 BlockIndex = NextSibling;
-					NextSibling = (Sibling + 1 < NumChildren) ? AcquireBlock() : 0xffff;
+	check(vBlockTileX >= vAllocatedTileX);
+	check(vBlockTileY >= vAllocatedTileY);
 
-					FAddressBlock Block(FirstChildBlock, Sibling, vDimensions);
-					Block.NextSibling = NextSibling;
-					AddressBlocks[BlockIndex] = Block;
-					NewBlocks.Add(BlockIndex);
-				}
+	if (vBlockTileX < vAllocatedTileX1 && vBlockTileY < vAllocatedTileY1)
+	{
+		// Block overlaps the VT we're trying to allocate
 
-				AllocIndex = FirstChildIndex;
-				AllocBlock = &AddressBlocks[AllocIndex];
+		// Remove from free list
+		{
+			const uint32 PrevFreeIndex = AllocBlock->PrevFree;
+			const uint32 NextFreeIndex = AllocBlock->NextFree;
+			if (PrevFreeIndex != 0xffff)
+			{
+				AddressBlocks[PrevFreeIndex].NextFree = NextFreeIndex;
+				AllocBlock->PrevFree = 0xffff;
 			}
-
-			// If new blocks were generated add them to the lists
-			if (NewBlocks.Num())
+			if (NextFreeIndex != 0xffff)
 			{
-				FAddressBlock* ParentBlock = &AddressBlocks[FreeIndex];
-				check(ParentBlock->vAddress == AllocBlock->vAddress);
-
-				const int32 SortedIndex = FindAddressBlock(AllocBlock->vAddress);
-				check(AllocBlock->vAddress == SortedAddresses[SortedIndex]);
-
-				// Replace parent block with new allocated block
-				SortedIndices[SortedIndex] = AllocIndex;
-				// Make room for newly added
-				SortedAddresses.InsertUninitialized(SortedIndex, NewBlocks.Num());
-				SortedIndices.InsertUninitialized(SortedIndex, NewBlocks.Num());
-				check(SortedAddresses.Num() == SortedIndices.Num());
-
-				// Add new blocks to lists
-				int32 DepthCount = 0;
-				while (ParentBlock->FirstChild != 0xffff)
-				{
-					// Add siblings at this size
-					int32 Sibling = 1;
-					uint16 Index = AddressBlocks[ParentBlock->FirstChild].NextSibling;
-					while (Index != 0xffff)
-					{
-						FAddressBlock& AddressBlock = AddressBlocks[Index];
-
-						// Place on free list
-						AddressBlock.NextFree = FreeList[AddressBlock.vLogSize];
-						if (AddressBlock.NextFree != 0xffff)
-						{
-							AddressBlocks[AddressBlock.NextFree].PrevFree = Index;
-						}
-						FreeList[AddressBlock.vLogSize] = Index;
-
-						// Add to sorted list
-						// We are inserting sorted so need to be careful to add siblings in reverse order
-						const int32 NumChildren = (1 << vDimensions);
-						const int32 SortedIndexOffset = (DepthCount + 1) * (NumChildren - 1) - Sibling;
-						SortedAddresses[SortedIndex + SortedIndexOffset] = AddressBlock.vAddress;
-						SortedIndices[SortedIndex + SortedIndexOffset] = Index;
-
-						Index = AddressBlock.NextSibling;
-						Sibling++;
-					}
-
-					// Now handle child siblings
-					ParentBlock = &AddressBlocks[ParentBlock->FirstChild];
-					DepthCount++;
-				}
+				AddressBlocks[NextFreeIndex].PrevFree = PrevFreeIndex;
+				AllocBlock->NextFree = 0xffff;
+			}
+			if (FreeList[vLogSize] == Index)
+			{
+				FreeList[vLogSize] = NextFreeIndex;
 			}
+			AllocBlock->State = EBlockState::None;
+		}
 
+		if (vBlockTileX + BlockSize <= vAllocatedTileX1 || vBlockTileY + BlockSize <= vAllocatedTileY1)
+		{
+			// Block is entirely contained within the VT we're trying to allocate
 			++NumAllocations;
 			NumAllocatedPages += 1u << (vDimensions * vLogSize);
 
 			// Add to hash table
 			uint16 Key = reinterpret_cast<UPTRINT>(VT) / 16;
-			HashTable.Add(Key, AllocIndex);
+			HashTable.Add(Key, Index);
 
 			AllocBlock->VT = VT;
-			return AllocBlock->vAddress;
+			AllocBlock->State = EBlockState::AllocatedTexture;
+		}
+		else
+		{
+			// Block intersects the VT, subdivide further
+			AllocBlock->State = EBlockState::AllocatedChildren;
+			AllocBlock = nullptr; // list will be potentially reallocated
+
+			SubdivideBlock(Index);
+	
+			uint32 NumChildren = 0u;
+			uint16 ChildIndex = AddressBlocks[Index].FirstChild;
+			check(ChildIndex == AddressBlocks[ChildIndex].FirstSibling);
+			while (ChildIndex != 0xffff)
+			{
+				check(AddressBlocks[ChildIndex].Parent == Index);
+
+				MarkBlockAllocated(ChildIndex, vAllocatedTileX, vAllocatedTileY, VT);
+
+				ChildIndex = AddressBlocks[ChildIndex].NextSibling;
+				NumChildren++;
+			}
+			check(NumChildren == (1u << vDimensions));
+		}
+	}
+}
+
+uint32 FVirtualTextureAllocator::Alloc(FAllocatedVirtualTexture* VT)
+{
+	// Pad out to square power of 2
+	const uint32 MinSize = FMath::Min(VT->GetWidthInTiles(), VT->GetHeightInTiles());
+	const uint32 MaxSize = FMath::Max(VT->GetWidthInTiles(), VT->GetHeightInTiles());
+	const uint8 vLogMaxSize = FMath::CeilLogTwo(MaxSize);
+
+	// Find smallest free that fits
+	for (int i = vLogMaxSize; i < FreeList.Num(); i++)
+	{
+		uint16 FreeIndex = FreeList[i];
+		if (FreeIndex != 0xffff)
+		{
+			// Found free
+			uint16 AllocIndex = FreeIndex;
+			FAddressBlock& AllocBlock = AddressBlocks[AllocIndex];
+			check(AllocBlock.VT == nullptr);
+			check(AllocBlock.State == EBlockState::FreeList);
+
+			const uint32 vAddress = AllocBlock.vAddress;
+			const uint32 vTileX = FMath::ReverseMortonCode2(vAddress);
+			const uint32 vTileY = FMath::ReverseMortonCode2(vAddress >> 1);
+
+			MarkBlockAllocated(AllocIndex, vTileX, vTileY, VT);
+
+			check(AddressBlocks[AllocIndex].State != EBlockState::FreeList);
+
+			return vAddress;
 		}
 	}
 
@@ -314,33 +376,33 @@ void FVirtualTextureAllocator::Free(FAllocatedVirtualTexture* VT)
 {
 	// Find block index
 	uint16 Key = reinterpret_cast<UPTRINT>(VT) / 16;
-	uint32 Index;
-	for (Index = HashTable.First(Key); HashTable.IsValid(Index); Index = HashTable.Next(Index))
-	{
-		if (AddressBlocks[Index].VT == VT)
-		{
-			break;
-		}
-	}
-	if (HashTable.IsValid(Index))
+	uint32 Index = HashTable.First(Key);
+	while (HashTable.IsValid(Index))
 	{
+		const uint32 NextIndex = HashTable.Next(Index);
 		FAddressBlock& AddressBlock = AddressBlocks[Index];
-		check(AddressBlock.VT == VT);
-		AddressBlock.VT = nullptr;
+		if (AddressBlock.VT == VT)
+		{
+			check(AddressBlock.State == EBlockState::AllocatedTexture);
+			check(AddressBlock.FirstChild == 0xffff); // texture allocation should be leaf
+			AddressBlock.State = EBlockState::AllocatedChildren;
+			AddressBlock.VT = nullptr;
 
-		check(NumAllocations > 0u);
-		--NumAllocations;
+			check(NumAllocations > 0u);
+			--NumAllocations;
 
-		const uint32 NumPagesForBlock = 1u << (vDimensions * AddressBlock.vLogSize);
-		check(NumAllocatedPages >= NumPagesForBlock);
-		NumAllocatedPages -= NumPagesForBlock;
+			const uint32 NumPagesForBlock = 1u << (vDimensions * AddressBlock.vLogSize);
+			check(NumAllocatedPages >= NumPagesForBlock);
+			NumAllocatedPages -= NumPagesForBlock;
 
-		// Add block to free list
-		// This handles merging free siblings
-		FreeAddressBlock(Index);
+			// Add block to free list
+			// This handles merging free siblings
+			FreeAddressBlock(Index);
 
-		// Remove the index from the hash table as it may be reused later
-		HashTable.Remove(Key, Index);
+			// Remove the index from the hash table as it may be reused later
+			HashTable.Remove(Key, Index);
+		}
+		Index = NextIndex;
 	}
 }
 
@@ -348,6 +410,7 @@ void FVirtualTextureAllocator::FreeAddressBlock(uint32 Index)
 {
 	FAddressBlock& AddressBlock = AddressBlocks[Index];
 	check(AddressBlock.VT == nullptr);
+	check(AddressBlock.State == EBlockState::AllocatedChildren);
 	check(AddressBlock.NextFree == 0xffff);
 	check(AddressBlock.PrevFree == 0xffff);
 
@@ -362,7 +425,11 @@ void FVirtualTextureAllocator::FreeAddressBlock(uint32 Index)
 
 	if (!bConsolidateSiblings)
 	{
+		// If we got here, the block's children have already been consolidated/removed
+		AddressBlock.FirstChild = 0xffff;
+
 		// Simply place this block on the free list
+		AddressBlock.State = EBlockState::FreeList;
 		AddressBlock.NextFree = FreeList[AddressBlock.vLogSize];
 		if (AddressBlock.NextFree != 0xffff)
 		{
@@ -377,21 +444,28 @@ void FVirtualTextureAllocator::FreeAddressBlock(uint32 Index)
 		while (FreeIndex != 0xffff)
 		{
 			FAddressBlock& FreeBlock = AddressBlocks[FreeIndex];
-			const uint32 PrevFreeIndex = FreeBlock.PrevFree;
-			const uint32 NextFreeIndex = FreeBlock.NextFree;
-			if (PrevFreeIndex != 0xffff)
-			{
-				AddressBlocks[PrevFreeIndex].NextFree = NextFreeIndex;
-				FreeBlock.PrevFree = 0xffff;
-			}
-			if (NextFreeIndex != 0xffff)
-			{
-				AddressBlocks[NextFreeIndex].PrevFree = PrevFreeIndex;
-				FreeBlock.NextFree = 0xffff;
-			}
-			if (FreeList[AddressBlock.vLogSize] == FreeIndex)
+
+			if (FreeIndex != Index)
 			{
-				FreeList[AddressBlock.vLogSize] = NextFreeIndex;
+				// All our siblings must be free (we checked above to get into this case)
+				check(FreeBlock.State == EBlockState::FreeList);
+
+				const uint32 PrevFreeIndex = FreeBlock.PrevFree;
+				const uint32 NextFreeIndex = FreeBlock.NextFree;
+				if (PrevFreeIndex != 0xffff)
+				{
+					AddressBlocks[PrevFreeIndex].NextFree = NextFreeIndex;
+					FreeBlock.PrevFree = 0xffff;
+				}
+				if (NextFreeIndex != 0xffff)
+				{
+					AddressBlocks[NextFreeIndex].PrevFree = PrevFreeIndex;
+					FreeBlock.NextFree = 0xffff;
+				}
+				if (FreeList[AddressBlock.vLogSize] == FreeIndex)
+				{
+					FreeList[AddressBlock.vLogSize] = NextFreeIndex;
+				}
 			}
 
 			if (GlobalFreeList != 0xffff)
@@ -400,6 +474,7 @@ void FVirtualTextureAllocator::FreeAddressBlock(uint32 Index)
 				AddressBlocks[GlobalFreeList].PrevFree = FreeIndex;
 			}
 			GlobalFreeList = FreeIndex;
+			FreeBlock.State = EBlockState::GlobalFreeList;
 
 			FreeIndex = FreeBlock.NextSibling;
 		}
diff --git a/Engine/Source/Runtime/Renderer/Private/VT/VirtualTextureAllocator.h b/Engine/Source/Runtime/Renderer/Private/VT/VirtualTextureAllocator.h
--- a/Engine/Source/Runtime/Renderer/Private/VT/VirtualTextureAllocator.h
+++ b/Engine/Source/Runtime/Renderer/Private/VT/VirtualTextureAllocator.h
@@ -63,6 +63,19 @@ private:
 	void FreeAddressBlock(uint32 Index);
 	uint32 FindAddressBlock(uint32 vAddress) const;
 
+	void SubdivideBlock(uint32 ParentIndex);
+
+	void MarkBlockAllocated(uint32 Index, uint32 vAllocatedTileX, uint32 vAllocatedTileY, FAllocatedVirtualTexture* VT);
+
+	enum class EBlockState : uint8
+	{
+		None,
+		GlobalFreeList,
+		FreeList,
+		AllocatedChildren,
+		AllocatedTexture,
+	};
+
 	struct FAddressBlock
 	{
 		FAllocatedVirtualTexture*	VT;
@@ -75,6 +88,7 @@ private:
 		uint16						NextSibling;
 		uint16						NextFree;
 		uint16						PrevFree;
+		EBlockState                 State;
 
 		FAddressBlock()
 		{}
@@ -90,6 +104,7 @@ private:
 			, NextSibling(0xffff)
 			, NextFree(0xffff)
 			, PrevFree(0xffff)
+			, State(EBlockState::None)
 		{}
 
 		FAddressBlock(const FAddressBlock& Block, uint32 Offset, uint32 Dimensions)
@@ -103,6 +118,7 @@ private:
 			, NextSibling(0xffff)
 			, NextFree(0xffff)
 			, PrevFree(0xffff)
+			, State(EBlockState::None)
 		{}
 	};
 
