Added bSupportsPathTracing bit to DDPI, which is only set on Windows.

diff --git a/Engine/Config/Windows/DataDrivenPlatformInfo.ini b/Engine/Config/Windows/DataDrivenPlatformInfo.ini
--- a/Engine/Config/Windows/DataDrivenPlatformInfo.ini
+++ b/Engine/Config/Windows/DataDrivenPlatformInfo.ini
@@ -28,6 +28,7 @@ MaxFeatureLevel=SM5
 bIsPC=true
 bSupportsRayTracing=true
 bSupportsRayTracingIndirectInstanceData=true
+bSupportsPathTracing=true
 bSupportsWaveOperations=false
 bSupportsGPUScene=true
 bSupportsNanite=true
@@ -53,6 +54,7 @@ MaxFeatureLevel=SM6
 bIsPC=true
 bSupportsRayTracing=true
 bSupportsRayTracingIndirectInstanceData=true
+bSupportsPathTracing=true
 bSupportsGPUScene=true
 bSupportsNanite=true
 bSupportsComputeFramework=true
diff --git a/Engine/Shaders/Private/PathTracing/PathTracing.usf b/Engine/Shaders/Private/PathTracing/PathTracing.usf
--- a/Engine/Shaders/Private/PathTracing/PathTracing.usf
+++ b/Engine/Shaders/Private/PathTracing/PathTracing.usf
@@ -6,6 +6,8 @@ PathTracingRayGenShader.usf: Reference path tracing
 
 #define PATH_TRACING  
 
+#define PATH_TRACING_SHADER_VERSION 0x51478b68 // Bump to force-compile this shader
+
 #define RANDSEQ_UNROLL_SOBOL 0 // unrolling sobol sampler doesn't seem to pay off in the path tracer
 
 #include "../Common.ush"
diff --git a/Engine/Source/Runtime/Engine/Private/SceneView.cpp b/Engine/Source/Runtime/Engine/Private/SceneView.cpp
--- a/Engine/Source/Runtime/Engine/Private/SceneView.cpp
+++ b/Engine/Source/Runtime/Engine/Private/SceneView.cpp
@@ -2682,7 +2682,7 @@ void FSceneView::SetupRayTracedRendering()
 
 	const FEngineShowFlags& ShowFlags = Family->EngineShowFlags;
 
-	if (ShowFlags.PathTracing)
+	if (ShowFlags.PathTracing && FDataDrivenShaderPlatformInfo::GetSupportsPathTracing(GetShaderPlatform()))
 	{
 		RayTracingRenderMode = ERayTracingRenderMode::PathTracing;
 	}	
diff --git a/Engine/Source/Runtime/RHI/Private/RHI.cpp b/Engine/Source/Runtime/RHI/Private/RHI.cpp
--- a/Engine/Source/Runtime/RHI/Private/RHI.cpp
+++ b/Engine/Source/Runtime/RHI/Private/RHI.cpp
@@ -1734,6 +1734,7 @@ void FGenericDataDrivenShaderPlatformInfo::ParseDataDrivenShaderInfo(const FConf
 	Info.bSupportsRenderTargetWriteMask = GetSectionBool(Section, "bSupportsRenderTargetWriteMask");
 	Info.bSupportsRayTracing = GetSectionBool(Section, "bSupportsRayTracing");
 	Info.bSupportsRayTracingIndirectInstanceData = GetSectionBool(Section, "bSupportsRayTracingIndirectInstanceData");
+	Info.bSupportsPathTracing = GetSectionBool(Section, "bSupportsPathTracing");
 	Info.bSupportsGPUSkinCache = GetSectionBool(Section, "bSupportsGPUSkinCache");
 	Info.bSupportsByteBufferComputeShaders = GetSectionBool(Section, "bSupportsByteBufferComputeShaders");
 	Info.bSupportsGPUScene = GetSectionBool(Section, "bSupportsGPUScene");
diff --git a/Engine/Source/Runtime/RHI/Public/RHIDefinitions.h b/Engine/Source/Runtime/RHI/Public/RHIDefinitions.h
--- a/Engine/Source/Runtime/RHI/Public/RHIDefinitions.h
+++ b/Engine/Source/Runtime/RHI/Public/RHIDefinitions.h
@@ -340,6 +340,7 @@ class RHI_API FGenericDataDrivenShaderPlatformInfo
 	uint32 bSupportsRenderTargetWriteMask: 1;
 	uint32 bSupportsRayTracing: 1;
 	uint32 bSupportsRayTracingIndirectInstanceData : 1; // Whether instance transforms can be copied from the GPU to the TLAS instances buffer
+	uint32 bSupportsPathTracing : 1; // Whether real-time path tracer is supported on this platform (avoids compiling unnecessary shaders)
 	uint32 bSupportsGPUSkinCache: 1;
 	uint32 bSupportsGPUScene : 1;
 	uint32 bSupportsByteBufferComputeShaders : 1;
@@ -522,6 +523,11 @@ public:
 		return Infos[Platform].bSupportsRayTracingIndirectInstanceData;
 	}
 
+	static FORCEINLINE_DEBUGGABLE const bool GetSupportsPathTracing(const FStaticShaderPlatform Platform)
+	{
+		return Infos[Platform].bSupportsPathTracing;
+	}
+
 	static FORCEINLINE_DEBUGGABLE const bool GetSupportsGPUSkinCache(const FStaticShaderPlatform Platform)
 	{
 		return Infos[Platform].bSupportsGPUSkinCache;
diff --git a/Engine/Source/Runtime/Renderer/Private/DeferredShadingRenderer.cpp b/Engine/Source/Runtime/Renderer/Private/DeferredShadingRenderer.cpp
--- a/Engine/Source/Runtime/Renderer/Private/DeferredShadingRenderer.cpp
+++ b/Engine/Source/Runtime/Renderer/Private/DeferredShadingRenderer.cpp
@@ -2718,7 +2718,8 @@ void FDeferredShadingSceneRenderer::Render(FRDGBuilder& GraphBuilder)
 	{
 		for (int32 ViewIndex = 0; ViewIndex < Views.Num(); ViewIndex++)
 		{
-			if (Views[ViewIndex].RayTracingRenderMode == ERayTracingRenderMode::PathTracing)
+			if (Views[ViewIndex].RayTracingRenderMode == ERayTracingRenderMode::PathTracing
+				&& FDataDrivenShaderPlatformInfo::GetSupportsPathTracing(Views[ViewIndex].GetShaderPlatform()))
 			{
 				RenderPathTracing(GraphBuilder, Views[ViewIndex], SceneTextures.UniformBuffer, SceneTextures.Color.Target);
 			}
diff --git a/Engine/Source/Runtime/Renderer/Private/PathTracing.cpp b/Engine/Source/Runtime/Renderer/Private/PathTracing.cpp
--- a/Engine/Source/Runtime/Renderer/Private/PathTracing.cpp
+++ b/Engine/Source/Runtime/Renderer/Private/PathTracing.cpp
@@ -340,7 +340,8 @@ class FPathTracingRG : public FGlobalShader
 
 	static bool ShouldCompilePermutation(const FGlobalShaderPermutationParameters& Parameters)
 	{
-		return ShouldCompileRayTracingShadersForProject(Parameters.Platform);
+		return ShouldCompileRayTracingShadersForProject(Parameters.Platform)
+			&& FDataDrivenShaderPlatformInfo::GetSupportsPathTracing(Parameters.Platform);
 	}
 
 	static void ModifyCompilationEnvironment(const FGlobalShaderPermutationParameters& Parameters, FShaderCompilerEnvironment& OutEnvironment)
@@ -761,7 +762,8 @@ IMPLEMENT_SHADER_TYPE(, FPathTracingCompositorPS, TEXT("/Engine/Private/PathTrac
 
 void FDeferredShadingSceneRenderer::PreparePathTracing(const FViewInfo& View, TArray<FRHIRayTracingShader*>& OutRayGenShaders)
 {
-	if (View.RayTracingRenderMode == ERayTracingRenderMode::PathTracing)
+	if (View.RayTracingRenderMode == ERayTracingRenderMode::PathTracing
+		&& FDataDrivenShaderPlatformInfo::GetSupportsPathTracing(View.GetShaderPlatform()))
 	{
 		// Declare all RayGen shaders that require material closest hit shaders to be bound
 		auto RayGenShader = View.ShaderMap->GetShader<FPathTracingRG>();
@@ -785,6 +787,11 @@ void FDeferredShadingSceneRenderer::RenderPathTracing(
 	RDG_GPU_STAT_SCOPE(GraphBuilder, Stat_GPU_PathTracing);
 	RDG_EVENT_SCOPE(GraphBuilder, "Path Tracing");
 
+	if (!ensureMsgf(FDataDrivenShaderPlatformInfo::GetSupportsPathTracing(View.GetShaderPlatform()),
+		TEXT("Attempting to use path tracing on unsupported platform.")))
+	{
+		return;
+	}
 
 	bool bArgsChanged = false;
 
