Mobile Full DepthPrepass UE5

diff --git a/Engine/Source/Runtime/Renderer/Private/DepthRendering.cpp b/Engine/Source/Runtime/Renderer/Private/DepthRendering.cpp
--- a/Engine/Source/Runtime/Renderer/Private/DepthRendering.cpp
+++ b/Engine/Source/Runtime/Renderer/Private/DepthRendering.cpp
@@ -610,10 +610,7 @@ void FDeferredShadingSceneRenderer::RenderPrePass(FRDGBuilder& GraphBuilder, FRD
 bool FMobileSceneRenderer::ShouldRenderPrePass() const
 {
 	// Draw a depth pass to avoid overdraw in the other passes.
-	// Mobile only does MaskedOnly DepthPass for the moment
-	bool bShouldRenderPrePass = Scene->EarlyZPassMode == DDM_MaskedOnly;
-
-	return bShouldRenderPrePass;
+	return Scene->EarlyZPassMode == DDM_MaskedOnly || Scene->EarlyZPassMode == DDM_AllOpaque;
 }
 
 void FMobileSceneRenderer::RenderPrePass(FRHICommandListImmediate& RHICmdList, const FViewInfo& View, const FInstanceCullingDrawParams* InstanceCullingDrawParams/* = nullptr*/)
diff --git a/Engine/Source/Runtime/Renderer/Private/MobileBasePass.cpp b/Engine/Source/Runtime/Renderer/Private/MobileBasePass.cpp
--- a/Engine/Source/Runtime/Renderer/Private/MobileBasePass.cpp
+++ b/Engine/Source/Runtime/Renderer/Private/MobileBasePass.cpp
@@ -713,7 +713,7 @@ bool FMobileBasePassMeshProcessor::Process(
 		{
 			MobileBasePass::SetTranslucentRenderState(DrawRenderState, MaterialResource);
 		}
-		else if (bMaskedInEarlyPass)
+		else if((MeshBatch.bUseForDepthPass && Scene->EarlyZPassMode == DDM_AllOpaque) || bMaskedInEarlyPass)
 		{
 			DrawRenderState.SetDepthStencilState(TStaticDepthStencilState<false, CF_Equal>::GetRHI());
 		}
diff --git a/Engine/Source/Runtime/Renderer/Private/MobileShadingRenderer.cpp b/Engine/Source/Runtime/Renderer/Private/MobileShadingRenderer.cpp
--- a/Engine/Source/Runtime/Renderer/Private/MobileShadingRenderer.cpp
+++ b/Engine/Source/Runtime/Renderer/Private/MobileShadingRenderer.cpp
@@ -215,7 +215,8 @@ FMobileSceneRenderer::FMobileSceneRenderer(const FSceneViewFamily* InViewFamily,
 	bRequiresPixelProjectedPlanarRelfectionPass = false;
 	bRequiresAmbientOcclusionPass = false;
 	bRequiresDistanceFieldShadowingPass = false;
-
+	bIsFullDepthPrepassEnabled = Scene->EarlyZPassMode == DDM_AllOpaque;
+	bIsMaskedOnlyDepthPrepassEnabled = Scene->EarlyZPassMode == DDM_MaskedOnly;
 	// Don't do occlusion queries when doing scene captures
 	for (FViewInfo& View : Views)
 	{
@@ -561,6 +562,39 @@ void FMobileSceneRenderer::InitViews(FRDGBuilder& GraphBuilder, FSceneTexturesCo
 	bSubmitOffscreenRendering = (!bGammaSpace || bRenderToSceneColor) && CVarMobileFlushSceneColorRendering.GetValueOnAnyThread() != 0;
 }
 
+/*
+* Renders the Full Depth Prepass
+*/
+void FMobileSceneRenderer::RenderFullDepthPrepass(FRDGBuilder& GraphBuilder, FSceneTextures& SceneTextures)
+{
+	FRenderTargetBindingSlots BasePassRenderTargets;
+	BasePassRenderTargets.DepthStencil = FDepthStencilBinding(SceneTextures.Depth.Target, ERenderTargetLoadAction::EClear, FExclusiveDepthStencil::DepthWrite_StencilWrite);
+
+	for (int32 ViewIndex = 0; ViewIndex < Views.Num(); ViewIndex++)
+	{
+		const FViewInfo& View = Views[ViewIndex];
+		if (!View.ShouldRenderView())
+		{
+			continue;
+		}
+
+		View.BeginRenderView();
+
+		auto* PassParameters = GraphBuilder.AllocParameters<FMobileRenderPassParameters>();
+		PassParameters->RenderTargets = BasePassRenderTargets;
+		PassParameters->View = View.GetShaderParameters();
+
+		GraphBuilder.AddPass(
+			RDG_EVENT_NAME("FullDepthPrepass"),
+			PassParameters,
+			ERDGPassFlags::Raster,
+			[this, PassParameters, &View](FRHICommandListImmediate& RHICmdList)
+			{
+				RenderPrePass(RHICmdList, View, &PassParameters->InstanceCullingDrawParams[EMeshPass::DepthPass]);
+			});
+	}
+}
+
 /** 
 * Renders the view family. 
 */
@@ -698,6 +732,11 @@ void FMobileSceneRenderer::Render(FRDGBuilder& GraphBuilder)
 	SceneTextures.MobileUniformBuffer = CreateMobileSceneTextureUniformBuffer(GraphBuilder, SceneTextures.MobileSetupMode);
 
 	FRDGTextureRef ViewFamilyTexture = TryCreateViewFamilyTexture(GraphBuilder, ViewFamily);
+	
+	if (bIsFullDepthPrepassEnabled)
+	{
+		RenderFullDepthPrepass(GraphBuilder, SceneTextures);
+	}
 
 	if (bDeferredShading)
 	{
@@ -857,7 +896,7 @@ void FMobileSceneRenderer::RenderForward(FRDGBuilder& GraphBuilder, FRDGTextureR
 
 	FRenderTargetBindingSlots BasePassRenderTargets;
 	BasePassRenderTargets[0] = FRenderTargetBinding(SceneColor, SceneColorResolve, ERenderTargetLoadAction::EClear);
-	BasePassRenderTargets.DepthStencil = FDepthStencilBinding(SceneDepth, ERenderTargetLoadAction::EClear, FExclusiveDepthStencil::DepthWrite_StencilWrite);
+	BasePassRenderTargets.DepthStencil = FDepthStencilBinding(SceneDepth, bIsFullDepthPrepassEnabled ? ERenderTargetLoadAction::ELoad: ERenderTargetLoadAction::EClear, FExclusiveDepthStencil::DepthWrite_StencilWrite);
 	BasePassRenderTargets.ShadingRateTexture = (!MainView.bIsSceneCapture && !MainView.bIsReflectionCapture) ? SceneTextures.ShadingRate : nullptr;
 	BasePassRenderTargets.SubpassHint = ESubpassHint::DepthReadSubpass;
 	BasePassRenderTargets.NumOcclusionQueries = ComputeNumOcclusionQueriesToBatch();
@@ -865,8 +904,6 @@ void FMobileSceneRenderer::RenderForward(FRDGBuilder& GraphBuilder, FRDGTextureR
 	//if the scenecolor isn't multiview but the app is, need to render as a single-view multiview due to shaders
 	BasePassRenderTargets.MultiViewCount = MainView.bIsMobileMultiViewEnabled ? 2 : (bIsMultiViewApplication ? 1 : 0);
 
-	const bool bShouldRenderPrePass = ShouldRenderPrePass();
-
 	const bool bShouldRenderMobileDebugView = ViewFamily.UseDebugViewPS();
 
 	const bool bAdrenoOcclusionMode = CVarMobileAdrenoOcclusionMode.GetValueOnRenderThread() != 0;
@@ -912,7 +949,7 @@ void FMobileSceneRenderer::RenderForward(FRDGBuilder& GraphBuilder, FRDGTextureR
 		PassParameters->SceneTextures = SceneTextures.MobileUniformBuffer;
 		PassParameters->RenderTargets = BasePassRenderTargets;
 
-		if (bShouldRenderPrePass)
+		if (bIsMaskedOnlyDepthPrepassEnabled)
 		{
 //			View.ParallelMeshDrawCommandPasses[EMeshPass::DepthPass].BuildRenderingCommands(GraphBuilder, Scene->GPUScene, PassParameters->InstanceCullingDrawParams[EMeshPass::DepthPass]);
 		}
@@ -935,24 +972,24 @@ void FMobileSceneRenderer::RenderForward(FRDGBuilder& GraphBuilder, FRDGTextureR
 		// Split if we need to render translucency in a separate render pass
 		if (bRequiresMultiPass)
 		{
-			RenderForwardMultiPass(GraphBuilder, PassParameters, BasePassRenderTargets, ViewIndex, View, SceneTextures, bShouldRenderPrePass, bShouldRenderMobileDebugView, bAdrenoOcclusionMode, bShouldRenderOcclusionPass, bShouldRenderTranslucency, TranslucencyMeshPass);
+			RenderForwardMultiPass(GraphBuilder, PassParameters, BasePassRenderTargets, ViewIndex, View, SceneTextures, bShouldRenderMobileDebugView, bAdrenoOcclusionMode, bShouldRenderOcclusionPass, bShouldRenderTranslucency, TranslucencyMeshPass);
 		}
 		else
 		{
-			RenderForwardSinglePass(GraphBuilder, PassParameters, ViewIndex, View, SceneTextures, bShouldRenderPrePass, bShouldRenderMobileDebugView, bAdrenoOcclusionMode, bShouldRenderOcclusionPass, bShouldRenderTranslucency, TranslucencyMeshPass);
+			RenderForwardSinglePass(GraphBuilder, PassParameters, ViewIndex, View, SceneTextures, bShouldRenderMobileDebugView, bAdrenoOcclusionMode, bShouldRenderOcclusionPass, bShouldRenderTranslucency, TranslucencyMeshPass);
 		}
 	}
 
 	QueueSceneTextureExtractions(GraphBuilder, SceneTextures);
 }
 
-void FMobileSceneRenderer::RenderForwardSinglePass(FRDGBuilder& GraphBuilder, class FMobileRenderPassParameters* PassParameters, int32 ViewIndex, FViewInfo& View, FSceneTextures& SceneTextures, bool bShouldRenderPrePass, bool bShouldRenderMobileDebugView, bool bAdrenoOcclusionMode, bool bShouldRenderOcclusionPass, bool bShouldRenderTranslucency, EMeshPass::Type TranslucencyMeshPass)
+void FMobileSceneRenderer::RenderForwardSinglePass(FRDGBuilder& GraphBuilder, class FMobileRenderPassParameters* PassParameters, int32 ViewIndex, FViewInfo& View, FSceneTextures& SceneTextures, bool bShouldRenderMobileDebugView, bool bAdrenoOcclusionMode, bool bShouldRenderOcclusionPass, bool bShouldRenderTranslucency, EMeshPass::Type TranslucencyMeshPass)
 {
 	GraphBuilder.AddPass(
 		RDG_EVENT_NAME("SceneColorRendering"),
 		PassParameters,
 		ERDGPassFlags::Raster,
-		[this, PassParameters, ViewIndex, &View, &SceneTextures, bShouldRenderPrePass, bShouldRenderMobileDebugView, bAdrenoOcclusionMode, bShouldRenderOcclusionPass, bShouldRenderTranslucency, TranslucencyMeshPass](FRHICommandListImmediate& RHICmdList)
+		[this, PassParameters, ViewIndex, &View, &SceneTextures, bShouldRenderMobileDebugView, bAdrenoOcclusionMode, bShouldRenderOcclusionPass, bShouldRenderTranslucency, TranslucencyMeshPass](FRHICommandListImmediate& RHICmdList)
 	{
 		if (GIsEditor && !View.bIsSceneCapture)
 		{
@@ -960,7 +997,7 @@ void FMobileSceneRenderer::RenderForwardSinglePass(FRDGBuilder& GraphBuilder, cl
 		}
 
 		RHICmdList.SetCurrentStat(GET_STATID(STAT_CLM_MobilePrePass));
-		if (bShouldRenderPrePass)
+		if (bIsMaskedOnlyDepthPrepassEnabled)
 		{
 			// Depth pre-pass
 			RenderPrePass(RHICmdList, View, &PassParameters->InstanceCullingDrawParams[EMeshPass::DepthPass]);
@@ -1049,13 +1086,13 @@ void FMobileSceneRenderer::RenderForwardSinglePass(FRDGBuilder& GraphBuilder, cl
 	});
 }
 
-void FMobileSceneRenderer::RenderForwardMultiPass(FRDGBuilder& GraphBuilder, class FMobileRenderPassParameters* PassParameters, FRenderTargetBindingSlots& BasePassRenderTargets, int32 ViewIndex, FViewInfo& View, FSceneTextures& SceneTextures, bool bShouldRenderPrePass, bool bShouldRenderMobileDebugView, bool bAdrenoOcclusionMode, bool bShouldRenderOcclusionPass, bool bShouldRenderTranslucency, EMeshPass::Type TranslucencyMeshPass)
+void FMobileSceneRenderer::RenderForwardMultiPass(FRDGBuilder& GraphBuilder, class FMobileRenderPassParameters* PassParameters, FRenderTargetBindingSlots& BasePassRenderTargets, int32 ViewIndex, FViewInfo& View, FSceneTextures& SceneTextures, bool bShouldRenderMobileDebugView, bool bAdrenoOcclusionMode, bool bShouldRenderOcclusionPass, bool bShouldRenderTranslucency, EMeshPass::Type TranslucencyMeshPass)
 {
 	GraphBuilder.AddPass(
 		RDG_EVENT_NAME("SceneColorRendering"),
 		PassParameters,
 		ERDGPassFlags::Raster,
-		[this, PassParameters, &View, &SceneTextures, bShouldRenderPrePass, bShouldRenderMobileDebugView, bAdrenoOcclusionMode, bShouldRenderOcclusionPass](FRHICommandListImmediate& RHICmdList)
+		[this, PassParameters, &View, &SceneTextures, bShouldRenderMobileDebugView, bAdrenoOcclusionMode, bShouldRenderOcclusionPass](FRHICommandListImmediate& RHICmdList)
 	{
 		if (GIsEditor && !View.bIsSceneCapture)
 		{
@@ -1063,7 +1100,7 @@ void FMobileSceneRenderer::RenderForwardMultiPass(FRDGBuilder& GraphBuilder, cla
 		}
 
 		RHICmdList.SetCurrentStat(GET_STATID(STAT_CLM_MobilePrePass));
-		if (bShouldRenderPrePass)
+		if (bIsMaskedOnlyDepthPrepassEnabled)
 		{
 			// Depth pre-pass
 			RenderPrePass(RHICmdList, View, &PassParameters->InstanceCullingDrawParams[EMeshPass::DepthPass]);
@@ -1300,8 +1337,6 @@ void FMobileSceneRenderer::RenderDeferred(FRDGBuilder& GraphBuilder, const FSort
 	BasePassRenderTargets.ShadingRateTexture = nullptr;
 	BasePassRenderTargets.MultiViewCount = 0;
 
-	const bool bShouldRenderPrePass = ShouldRenderPrePass();
-
 	const bool bShouldRenderOcclusionPass = DoOcclusionQueries(FeatureLevel);
 
 	ETranslucencyPass::Type TranslucencyPass = ViewFamily.AllowTranslucencyAfterDOF() ? ETranslucencyPass::TPT_StandardTranslucency : ETranslucencyPass::TPT_AllTranslucency;
@@ -1345,31 +1380,31 @@ void FMobileSceneRenderer::RenderDeferred(FRDGBuilder& GraphBuilder, const FSort
 		PassParameters->SceneTextures = SceneTextures.MobileUniformBuffer;
 		PassParameters->RenderTargets = BasePassRenderTargets;
 
-		if (bShouldRenderPrePass)
+		if (bIsMaskedOnlyDepthPrepassEnabled)
 		{
 			View.ParallelMeshDrawCommandPasses[EMeshPass::DepthPass].BuildRenderingCommands(GraphBuilder, Scene->GPUScene, PassParameters->InstanceCullingDrawParams[EMeshPass::DepthPass]);
 		}
 
 		if (bRequiresMultiPass)
 		{
-			RenderDeferredMultiPass(GraphBuilder, PassParameters, BasePassRenderTargets, ColorTargets.Num(), ViewIndex, NumViews, View, SceneTextures, SortedLightSet, MobilePixelProjectedReflection, bShouldRenderPrePass, bShouldRenderOcclusionPass, bShouldRenderTranslucency, TranslucencyMeshPass);
+			RenderDeferredMultiPass(GraphBuilder, PassParameters, BasePassRenderTargets, ColorTargets.Num(), ViewIndex, NumViews, View, SceneTextures, SortedLightSet, MobilePixelProjectedReflection, bShouldRenderOcclusionPass, bShouldRenderTranslucency, TranslucencyMeshPass);
 		}
 		else
 		{
-			RenderDeferredSinglePass(GraphBuilder, PassParameters, ViewIndex, NumViews, View, SceneTextures, SortedLightSet, MobilePixelProjectedReflection, bShouldRenderPrePass, bShouldRenderOcclusionPass, bUsingPixelLocalStorage, bShouldRenderTranslucency, TranslucencyMeshPass);
+			RenderDeferredSinglePass(GraphBuilder, PassParameters, ViewIndex, NumViews, View, SceneTextures, SortedLightSet, MobilePixelProjectedReflection, bShouldRenderOcclusionPass, bUsingPixelLocalStorage, bShouldRenderTranslucency, TranslucencyMeshPass);
 		}
 	}
 
 	QueueSceneTextureExtractions(GraphBuilder, SceneTextures);
 }
 
-void FMobileSceneRenderer::RenderDeferredSinglePass(FRDGBuilder& GraphBuilder, class FMobileRenderPassParameters* PassParameters, int32 ViewIndex, int32 NumViews, FViewInfo& View, FSceneTextures& SceneTextures, const FSortedLightSetSceneInfo& SortedLightSet, FTextureRHIRef MobilePixelProjectedReflection, bool bShouldRenderPrePass, bool bShouldRenderOcclusionPass, bool bUsingPixelLocalStorage, bool bShouldRenderTranslucency, EMeshPass::Type TranslucencyMeshPass)
+void FMobileSceneRenderer::RenderDeferredSinglePass(FRDGBuilder& GraphBuilder, class FMobileRenderPassParameters* PassParameters, int32 ViewIndex, int32 NumViews, FViewInfo& View, FSceneTextures& SceneTextures, const FSortedLightSetSceneInfo& SortedLightSet, FTextureRHIRef MobilePixelProjectedReflection, bool bShouldRenderOcclusionPass, bool bUsingPixelLocalStorage, bool bShouldRenderTranslucency, EMeshPass::Type TranslucencyMeshPass)
 {
 	GraphBuilder.AddPass(
 		RDG_EVENT_NAME("SceneColorRendering"),
 		PassParameters,
 		ERDGPassFlags::Raster,
-		[this, PassParameters, ViewIndex, NumViews, &View, &SceneTextures, &SortedLightSet, MobilePixelProjectedReflection, bShouldRenderPrePass, bShouldRenderOcclusionPass, bUsingPixelLocalStorage, bShouldRenderTranslucency, TranslucencyMeshPass](FRHICommandListImmediate& RHICmdList)
+		[this, PassParameters, ViewIndex, NumViews, &View, &SceneTextures, &SortedLightSet, MobilePixelProjectedReflection, bShouldRenderOcclusionPass, bUsingPixelLocalStorage, bShouldRenderTranslucency, TranslucencyMeshPass](FRHICommandListImmediate& RHICmdList)
 	{
 		if (GIsEditor && !View.bIsSceneCapture)
 		{
@@ -1377,7 +1412,7 @@ void FMobileSceneRenderer::RenderDeferredSinglePass(FRDGBuilder& GraphBuilder, c
 		}
 
 		RHICmdList.SetCurrentStat(GET_STATID(STAT_CLM_MobilePrePass));
-		if (bShouldRenderPrePass)
+		if (bIsMaskedOnlyDepthPrepassEnabled)
 		{
 			// Depth pre-pass
 			RenderPrePass(RHICmdList, View, &PassParameters->InstanceCullingDrawParams[EMeshPass::DepthPass]);
@@ -1433,13 +1468,13 @@ void FMobileSceneRenderer::RenderDeferredSinglePass(FRDGBuilder& GraphBuilder, c
 	});
 }
 
-void FMobileSceneRenderer::RenderDeferredMultiPass(FRDGBuilder& GraphBuilder, class FMobileRenderPassParameters* PassParameters, FRenderTargetBindingSlots& BasePassRenderTargets, int32 NumColorTargets, int32 ViewIndex, int32 NumViews, FViewInfo& View, FSceneTextures& SceneTextures, const FSortedLightSetSceneInfo& SortedLightSet, FTextureRHIRef MobilePixelProjectedReflection, bool bShouldRenderPrePass, bool bShouldRenderOcclusionPass, bool bShouldRenderTranslucency, EMeshPass::Type TranslucencyMeshPass)
+void FMobileSceneRenderer::RenderDeferredMultiPass(FRDGBuilder& GraphBuilder, class FMobileRenderPassParameters* PassParameters, FRenderTargetBindingSlots& BasePassRenderTargets, int32 NumColorTargets, int32 ViewIndex, int32 NumViews, FViewInfo& View, FSceneTextures& SceneTextures, const FSortedLightSetSceneInfo& SortedLightSet, FTextureRHIRef MobilePixelProjectedReflection, bool bShouldRenderOcclusionPass, bool bShouldRenderTranslucency, EMeshPass::Type TranslucencyMeshPass)
 {
 	GraphBuilder.AddPass(
 		RDG_EVENT_NAME("SceneColorRendering"),
 		PassParameters,
 		ERDGPassFlags::Raster,
-		[this, PassParameters, &View, &SceneTextures, bShouldRenderPrePass, bShouldRenderOcclusionPass](FRHICommandListImmediate& RHICmdList)
+		[this, PassParameters, &View, &SceneTextures, bShouldRenderOcclusionPass](FRHICommandListImmediate& RHICmdList)
 	{
 		if (GIsEditor && !View.bIsSceneCapture)
 		{
@@ -1447,7 +1482,7 @@ void FMobileSceneRenderer::RenderDeferredMultiPass(FRDGBuilder& GraphBuilder, cl
 		}
 
 		RHICmdList.SetCurrentStat(GET_STATID(STAT_CLM_MobilePrePass));
-		if (bShouldRenderPrePass)
+		if (bIsMaskedOnlyDepthPrepassEnabled)
 		{
 			// Depth pre-pass
 			RenderPrePass(RHICmdList, View, &PassParameters->InstanceCullingDrawParams[EMeshPass::DepthPass]);
diff --git a/Engine/Source/Runtime/Renderer/Private/RendererScene.cpp b/Engine/Source/Runtime/Renderer/Private/RendererScene.cpp
--- a/Engine/Source/Runtime/Renderer/Private/RendererScene.cpp
+++ b/Engine/Source/Runtime/Renderer/Private/RendererScene.cpp
@@ -3355,13 +3355,16 @@ void FScene::UpdateEarlyZPassMode()
 	}
 	else if (GetShadingPath(GetFeatureLevel()) == EShadingPath::Mobile)
 	{
+		EarlyZPassMode = DDM_None;
+
 		if (MaskedInEarlyPass(GetFeatureLevelShaderPlatform(FeatureLevel)))
 		{
 			EarlyZPassMode = DDM_MaskedOnly;
 		}
-		else
+
+		if (IsMobileDistanceFieldEnabled(GetShaderPlatform()) || IsMobileAmbientOcclusionEnabled(GetShaderPlatform()))
 		{
-			EarlyZPassMode = DDM_None;
+			EarlyZPassMode = DDM_AllOpaque;
 		}
 	}
 }
diff --git a/Engine/Source/Runtime/Renderer/Private/SceneRendering.h b/Engine/Source/Runtime/Renderer/Private/SceneRendering.h
--- a/Engine/Source/Runtime/Renderer/Private/SceneRendering.h
+++ b/Engine/Source/Runtime/Renderer/Private/SceneRendering.h
@@ -2233,6 +2233,8 @@ protected:
 
 	void RenderPrePass(FRHICommandListImmediate& RHICmdList, const FViewInfo& View, const FInstanceCullingDrawParams* InstanceCullingDrawParams = nullptr);
 
+	void RenderFullDepthPrepass(FRDGBuilder& GraphBuilder, FSceneTextures& SceneTextures);
+
 	/** Renders the opaque base pass for mobile. */
 	void RenderMobileBasePass(FRHICommandListImmediate& RHICmdList, const FViewInfo& View);
 
@@ -2279,11 +2281,11 @@ protected:
 	void UpdateSkyReflectionUniformBuffer();
 
 	void RenderForward(FRDGBuilder& GraphBuilder, FRDGTextureRef ViewFamilyTexture, FSceneTextures& SceneTextures);
-	void RenderForwardSinglePass(FRDGBuilder& GraphBuilder, class FMobileRenderPassParameters* PassParameters, int32 ViewIndex, FViewInfo& View, FSceneTextures& SceneTextures, bool bShouldRenderPrePass, bool bShouldRenderMobileDebugView, bool bAdrenoOcclusionMode, bool bShouldRenderOcclusionPass, bool bShouldRenderTranslucency, EMeshPass::Type TranslucencyMeshPass);
-	void RenderForwardMultiPass(FRDGBuilder& GraphBuilder, class FMobileRenderPassParameters* PassParameters, FRenderTargetBindingSlots& BasePassRenderTargets, int32 ViewIndex, FViewInfo& View, FSceneTextures& SceneTextures, bool bShouldRenderPrePass, bool bShouldRenderMobileDebugView, bool bAdrenoOcclusionMode, bool bShouldRenderOcclusionPass, bool bShouldRenderTranslucency, EMeshPass::Type TranslucencyMeshPass);
+	void RenderForwardSinglePass(FRDGBuilder& GraphBuilder, class FMobileRenderPassParameters* PassParameters, int32 ViewIndex, FViewInfo& View, FSceneTextures& SceneTextures, bool bShouldRenderMobileDebugView, bool bAdrenoOcclusionMode, bool bShouldRenderOcclusionPass, bool bShouldRenderTranslucency, EMeshPass::Type TranslucencyMeshPass);
+	void RenderForwardMultiPass(FRDGBuilder& GraphBuilder, class FMobileRenderPassParameters* PassParameters, FRenderTargetBindingSlots& BasePassRenderTargets, int32 ViewIndex, FViewInfo& View, FSceneTextures& SceneTextures, bool bShouldRenderMobileDebugView, bool bAdrenoOcclusionMode, bool bShouldRenderOcclusionPass, bool bShouldRenderTranslucency, EMeshPass::Type TranslucencyMeshPass);
 	void RenderDeferred(FRDGBuilder& GraphBuilder, const FSortedLightSetSceneInfo& SortedLightSet, FRDGTextureRef ViewFamilyTexture, FSceneTextures& SceneTextures);
-	void RenderDeferredSinglePass(FRDGBuilder& GraphBuilder, class FMobileRenderPassParameters* PassParameters, int32 ViewIndex, int32 NumViews, FViewInfo& View, FSceneTextures& SceneTextures, const FSortedLightSetSceneInfo& SortedLightSet, FTextureRHIRef MobilePixelProjectedReflection, bool bShouldRenderPrePass, bool bShouldRenderOcclusionPass, bool bUsingPixelLocalStorage, bool bShouldRenderTranslucency, EMeshPass::Type TranslucencyMeshPass);
-	void RenderDeferredMultiPass(FRDGBuilder& GraphBuilder, class FMobileRenderPassParameters* PassParameters, FRenderTargetBindingSlots& BasePassRenderTargets, int32 NumColorTargets, int32 ViewIndex, int32 NumViews, FViewInfo& View, FSceneTextures& SceneTextures, const FSortedLightSetSceneInfo& SortedLightSet, FTextureRHIRef MobilePixelProjectedReflection, bool bShouldRenderPrePass, bool bShouldRenderOcclusionPass, bool bShouldRenderTranslucency, EMeshPass::Type TranslucencyMeshPass);
+	void RenderDeferredSinglePass(FRDGBuilder& GraphBuilder, class FMobileRenderPassParameters* PassParameters, int32 ViewIndex, int32 NumViews, FViewInfo& View, FSceneTextures& SceneTextures, const FSortedLightSetSceneInfo& SortedLightSet, FTextureRHIRef MobilePixelProjectedReflection, bool bShouldRenderOcclusionPass, bool bUsingPixelLocalStorage, bool bShouldRenderTranslucency, EMeshPass::Type TranslucencyMeshPass);
+	void RenderDeferredMultiPass(FRDGBuilder& GraphBuilder, class FMobileRenderPassParameters* PassParameters, FRenderTargetBindingSlots& BasePassRenderTargets, int32 NumColorTargets, int32 ViewIndex, int32 NumViews, FViewInfo& View, FSceneTextures& SceneTextures, const FSortedLightSetSceneInfo& SortedLightSet, FTextureRHIRef MobilePixelProjectedReflection, bool bShouldRenderOcclusionPass, bool bShouldRenderTranslucency, EMeshPass::Type TranslucencyMeshPass);
 
 	void RenderAmbientOcclusion(FRDGBuilder& GraphBuilder, FRDGTextureRef SceneDepthTexture, FRDGTextureRef AmbientOcclusionTexture);
 
@@ -2308,6 +2310,8 @@ private:
 	bool bShouldRenderHZB;
 	bool bRequiresDistanceField;
 	bool bRequiresDistanceFieldShadowingPass;
+	bool bIsFullDepthPrepassEnabled;
+	bool bIsMaskedOnlyDepthPrepassEnabled;
 	static FGlobalDynamicIndexBuffer DynamicIndexBuffer;
 	static FGlobalDynamicVertexBuffer DynamicVertexBuffer;
 	static TGlobalResource<FGlobalDynamicReadBuffer> DynamicReadBuffer;
diff --git a/Engine/Source/Runtime/Renderer/Private/SceneVisibility.cpp b/Engine/Source/Runtime/Renderer/Private/SceneVisibility.cpp
--- a/Engine/Source/Runtime/Renderer/Private/SceneVisibility.cpp
+++ b/Engine/Source/Runtime/Renderer/Private/SceneVisibility.cpp
@@ -1817,7 +1817,16 @@ struct FMarkRelevantStaticMeshesForViewData
 		MinScreenRadiusForDepthPrepassSquared = GMinScreenRadiusForDepthPrepass * GMinScreenRadiusForDepthPrepass;
 
 		extern bool ShouldForceFullDepthPass(EShaderPlatform ShaderPlatform);
-		bFullEarlyZPass = ShouldForceFullDepthPass(View.GetShaderPlatform());
+		EShaderPlatform ShaderPlatform = View.GetShaderPlatform();
+		if (IsMobilePlatform(ShaderPlatform))
+		{
+			FScene* Scene = View.Family->Scene->GetRenderScene();
+			bFullEarlyZPass = (Scene && Scene->EarlyZPassMode == DDM_AllOpaque);
+		}
+		else
+		{
+			bFullEarlyZPass = ShouldForceFullDepthPass(ShaderPlatform);
+		}
 	}
 };
 
@@ -2202,10 +2211,11 @@ struct FRelevancePacket
 		FViewInfo& WriteView = const_cast<FViewInfo&>(View);
 		const FSceneViewState* ViewState = (FSceneViewState*)View.State;
 		const EShadingPath ShadingPath = Scene->GetShadingPath();
-
+		const bool bMobileMaskedInEarlyPass = (ShadingPath == EShadingPath::Mobile) &&  Scene->EarlyZPassMode == DDM_MaskedOnly;
 		const bool bHLODActive = Scene->SceneLODHierarchy.IsActive();
 		const FHLODVisibilityState* const HLODState = bHLODActive && ViewState ? &ViewState->HLODVisibilityState : nullptr;
 
+		
 		for (int32 StaticPrimIndex = 0, Num = RelevantStaticPrimitives.NumPrims; StaticPrimIndex < Num; ++StaticPrimIndex)
 		{
 			int32 PrimitiveIndex = RelevantStaticPrimitives.Prims[StaticPrimIndex];
@@ -2228,10 +2238,9 @@ struct FRelevancePacket
 			float DistanceSquared = (Bounds.BoxSphereBounds.Origin - ViewData.ViewOrigin).SizeSquared();
 			const float LODFactorDistanceSquared = DistanceSquared * FMath::Square(ViewData.LODScale);
 			const bool bDrawShadowDepth = FMath::Square(Bounds.BoxSphereBounds.SphereRadius) > ViewData.MinScreenRadiusForCSMDepthSquared * LODFactorDistanceSquared;
-			const bool bDrawDepthOnly = ViewData.bFullEarlyZPass || FMath::Square(Bounds.BoxSphereBounds.SphereRadius) > GMinScreenRadiusForDepthPrepass * GMinScreenRadiusForDepthPrepass * LODFactorDistanceSquared;
+			const bool bDrawDepthOnly = ViewData.bFullEarlyZPass || ((ShadingPath != EShadingPath::Mobile) && (FMath::Square(Bounds.BoxSphereBounds.SphereRadius) > GMinScreenRadiusForDepthPrepass * GMinScreenRadiusForDepthPrepass * LODFactorDistanceSquared));
 
 			const bool bAddLightmapDensityCommands = View.Family->EngineShowFlags.LightMapDensity && AllowDebugViewmodes();
-			const bool bMobileMaskedInEarlyPass = MaskedInEarlyPass(Scene->GetShaderPlatform()) && Scene->EarlyZPassMode == DDM_MaskedOnly;
 
 			const int32 NumStaticMeshes = PrimitiveSceneInfo->StaticMeshRelevances.Num();
 			for(int32 MeshIndex = 0;MeshIndex < NumStaticMeshes;MeshIndex++)
@@ -2291,8 +2300,7 @@ struct FRelevancePacket
 							&& (ViewRelevance.bRenderInMainPass || ViewRelevance.bRenderCustomDepth || ViewRelevance.bRenderInDepthPass)
 							&& !bHiddenByHLODFade)
 						{
-							bool bMobileIsInDepthPassMaskedMesh = (bMobileMaskedInEarlyPass && ViewRelevance.bMasked) && ShadingPath == EShadingPath::Mobile;
-							if ((StaticMeshRelevance.bUseForDepthPass && bDrawDepthOnly && ShadingPath != EShadingPath::Mobile) || bMobileIsInDepthPassMaskedMesh)
+							if (StaticMeshRelevance.bUseForDepthPass && (bDrawDepthOnly || (bMobileMaskedInEarlyPass && ViewRelevance.bMasked)))
 							{
 								DrawCommandPacket.AddCommandsForMesh(PrimitiveIndex, PrimitiveSceneInfo, StaticMeshRelevance, StaticMesh, Scene, bCanCache, EMeshPass::DepthPass);
 
