Hair binding update.

diff --git a/Engine/Shaders/Private/DeferredLightPixelShaders.usf b/Engine/Shaders/Private/DeferredLightPixelShaders.usf
--- a/Engine/Shaders/Private/DeferredLightPixelShaders.usf
+++ b/Engine/Shaders/Private/DeferredLightPixelShaders.usf
@@ -6,6 +6,7 @@
 
 #define SUPPORT_CONTACT_SHADOWS 1
 
+#include "HairStrands/HairStrandsVisibilityCommonStruct.ush" 
 #include "Common.ush"
 
 #include "DeferredShadingCommon.ush"
@@ -302,17 +303,11 @@ void DeferredLightPixelMain(
 
 #if USE_HAIR_LIGHTING == 1
 
-Texture2D<uint> HairVisibilityNodeOffsetAndCount;
-StructuredBuffer<FPackedHairSample> HairVisibilityNodeData;
-StructuredBuffer<uint>	HairVisibilityNodeCoords;
-
 Texture2D<float4> ScreenShadowMaskSubPixelTexture;
 
 uint HairTransmittanceBufferMaxCount;
 StructuredBuffer<FPackedHairTransmittanceMask> HairTransmittanceBuffer;
 uint HairShadowMaskValid;
-uint HairComponents; // R, TT, TRT, LocalScattering, GlobalScattering, MultipleScattering
-float HairDualScatteringRoughnessOverride;
 
 void DeferredLightPixelMain(
 	float4 SVPos : SV_POSITION,
@@ -329,12 +324,12 @@ void DeferredLightPixelMain(
 		return;
 	}
 
-	const uint PackedCoord = HairVisibilityNodeCoords[SampleIndex];
+	const uint PackedCoord = HairStrands.HairSampleCoords[SampleIndex];
 	const uint2 PixelCoord = uint2((PackedCoord & 0xFFFF), ((PackedCoord >> 16) & 0xFFFF));
 	const float2 UV = (PixelCoord + float2(0.5f, 0.5f)) / float2(View.BufferSizeAndInvSize.xy);
 	const float2 ScreenPosition = (UV - View.ScreenPositionScaleBias.wz) / View.ScreenPositionScaleBias.xy;
 	
-	const FPackedHairSample PackedSample = HairVisibilityNodeData[SampleIndex];
+	const FPackedHairSample PackedSample = HairStrands.HairSampleData[SampleIndex];
 	const FHairSample Sample = UnpackHairSample(PackedSample);
 
 	#if USE_LIGHTING_CHANNELS
@@ -358,7 +353,7 @@ void DeferredLightPixelMain(
 	HairScreenSpaceData.GBuffer.Metallic = 0;
 	HairScreenSpaceData.GBuffer.Depth = SceneDepth;
 	HairScreenSpaceData.GBuffer.GBufferAO = 1;
-	HairScreenSpaceData.GBuffer.CustomData = float4(HairDualScatteringRoughnessOverride, 0, Sample.Backlit, 0);
+	HairScreenSpaceData.GBuffer.CustomData = float4(HairStrands.HairDualScatteringRoughnessOverride, 0, Sample.Backlit, 0);
 	HairScreenSpaceData.GBuffer.IndirectIrradiance = 0;
 	HairScreenSpaceData.GBuffer.PrecomputedShadowFactors = 1;
 	HairScreenSpaceData.GBuffer.PerObjectGBufferData = 0;
diff --git a/Engine/Shaders/Private/DeferredLightVertexShaders.usf b/Engine/Shaders/Private/DeferredLightVertexShaders.usf
--- a/Engine/Shaders/Private/DeferredLightVertexShaders.usf
+++ b/Engine/Shaders/Private/DeferredLightVertexShaders.usf
@@ -4,6 +4,7 @@
 	DeferredLightVertexShaders.usf: 
 =============================================================================*/
 
+#include "HairStrands/HairStrandsVisibilityCommonStruct.ush" 
 #include "Common.ush"
 
 #if defined(SHADER_RADIAL_LIGHT)
@@ -100,20 +101,17 @@ void RadialVertexMain(
 
 #if defined(SHADER_HAIR) && SHADER_HAIR == 1
 
-int2 MaxViewportResolution;
-Texture2D<uint> HairVisibilityNodeCount;
-
 void HairVertexMain(
 	in uint InVertexId : SV_VertexID,
 	out float4 OutPosition : SV_POSITION,
 	nointerpolation out uint OutNodeCount : DISPATCH_NODECOUNT,
 	nointerpolation out uint2 OutResolution : DISPATCH_RESOLUTION)
 {
-	OutNodeCount = HairVisibilityNodeCount.Load(uint3(0,0,0));
+	OutNodeCount = HairStrands.HairSampleCount.Load(uint3(0,0,0));
 	OutResolution.x = ceil(sqrt(OutNodeCount));
 	OutResolution.y = OutResolution.x;
 
-	const float2 ClipCoord = ((OutResolution + 0.5f) / float2(MaxViewportResolution)) * 2;
+	const float2 ClipCoord = ((OutResolution + 0.5f) / float2(HairStrands.HairSampleViewportResolution)) * 2;
 
 	const float2 UV = float2(InVertexId & 1, InVertexId >> 1);
 	const float2 Pos = float2(-UV.x, UV.y) * 4 + float2(-1, +1) + float2(ClipCoord.x, -ClipCoord.y);
diff --git a/Engine/Shaders/Private/HairStrands/HairStrandsVisibilityCommon.ush b/Engine/Shaders/Private/HairStrands/HairStrandsVisibilityCommon.ush
--- a/Engine/Shaders/Private/HairStrands/HairStrandsVisibilityCommon.ush
+++ b/Engine/Shaders/Private/HairStrands/HairStrandsVisibilityCommon.ush
@@ -115,16 +115,6 @@ float2 UnpackVelocity(uint PackedVelocity)
 // Option to reduce the footprint: 
 // * we don't use PrimitiveID at the moment which is encoded on 3bytes
 // * Deferred the evaluation of Coverage & material properties into some other pass?
-
-struct FPackedHairSample
-{
-	float Depth;
-	uint  PrimitiveID_MacroGroupID;
-	uint  Tangent_Coverage8bit;
-	uint  BaseColor_Roughness;
-	uint  Specular_LightChannels_Backlit;
-};
-
 struct FHairSample
 {
 	float	Depth;
diff --git a/Engine/Shaders/Private/HairStrands/HairStrandsVisibilityCommonStruct.ush b/Engine/Shaders/Private/HairStrands/HairStrandsVisibilityCommonStruct.ush
--- a/Engine/Shaders/Private/HairStrands/HairStrandsVisibilityCommonStruct.ush
+++ b/Engine/Shaders/Private/HairStrands/HairStrandsVisibilityCommonStruct.ush
@@ -13,3 +13,12 @@ struct FPPLLNodeData
 	uint NextNodeIndex;
 	uint PackedVelocity;
 };
+
+struct FPackedHairSample
+{
+	float Depth;
+	uint  PrimitiveID_MacroGroupID;
+	uint  Tangent_Coverage8bit;
+	uint  BaseColor_Roughness;
+	uint  Specular_LightChannels_Backlit;
+};
diff --git a/Engine/Shaders/Private/HairStrands/HairStrandsVoxelPageRayMarching.usf b/Engine/Shaders/Private/HairStrands/HairStrandsVoxelPageRayMarching.usf
--- a/Engine/Shaders/Private/HairStrands/HairStrandsVoxelPageRayMarching.usf
+++ b/Engine/Shaders/Private/HairStrands/HairStrandsVoxelPageRayMarching.usf
@@ -410,7 +410,7 @@ void MainCS(uint3 DispatchThreadId : SV_DispatchThreadID)
 		//  Virtual Resolution   : X Y Z
 		//  Page Index Offset    : O
 		{
-			float2 StartPos = float2(50.0f, 350.0f + float(MacroGroupId) * 100.0f) / OutputResolution;
+			float2 StartPos = float2(50.0f, 420.0f + float(MacroGroupId) * 100.0f) / OutputResolution;
 			float2 Pos = StartPos;
 
 			Pos = ShaderPrintSymbol(Pos, _M_);
diff --git a/Engine/Shaders/Private/ShadowProjectionPixelShader.usf b/Engine/Shaders/Private/ShadowProjectionPixelShader.usf
--- a/Engine/Shaders/Private/ShadowProjectionPixelShader.usf
+++ b/Engine/Shaders/Private/ShadowProjectionPixelShader.usf
@@ -28,7 +28,7 @@
 #define SPOT_LIGHT_PCSS 0
 #endif
 
-
+#include "HairStrands/HairStrandsVisibilityCommonStruct.ush"
 #include "Common.ush"
 
 #if FEATURE_LEVEL >= FEATURE_LEVEL_SM5
@@ -91,7 +91,6 @@ float3 ModulatedShadowBlendOp(float3 Source)
 #if SUBPIXEL_SHADOW
 #include "HairStrands/HairStrandsVisibilityCommon.ush"
 float2 ShadowNearAndFarDepth;
-Texture2D<uint4> HairCategorizationTexture;
 #endif
 
 #if STRATA_ENABLED
@@ -123,7 +122,7 @@ void Main(
 	const uint3 PixelCoord = uint3(floor(SVPos.xy), 0);
 #if SUBPIXEL_SHADOW
 	const bool bUseFadePlane = bCascadeUseFadePlane > 0;
-	const FCategorizationData CategorizationData = DecodeCategorizationData(HairCategorizationTexture.Load(PixelCoord));
+	const FCategorizationData CategorizationData = DecodeCategorizationData(HairStrands.HairCategorizationTexture.Load(PixelCoord));
 	SceneW = ConvertFromDeviceZ(CategorizationData.ClosestDepth);
 	bIsSubsurfaceCompatible = false;
 #endif
@@ -411,7 +410,7 @@ void MainOnePassPointLightPS(
 
 #if SUBPIXEL_SHADOW
 	const uint3 PixelCoord = uint3(floor(SVPos.xy), 0);
-	const FCategorizationData CategorizationData = DecodeCategorizationData(HairCategorizationTexture.Load(PixelCoord));
+	const FCategorizationData CategorizationData = DecodeCategorizationData(HairStrands.HairCategorizationTexture.Load(PixelCoord));
 	if (CategorizationData.PixelCoverage > 0)
 	{
 		SceneW = ConvertFromDeviceZ(CategorizationData.ClosestDepth);
diff --git a/Engine/Shaders/Private/VirtualShadowMaps/PageManagement.usf b/Engine/Shaders/Private/VirtualShadowMaps/PageManagement.usf
--- a/Engine/Shaders/Private/VirtualShadowMaps/PageManagement.usf
+++ b/Engine/Shaders/Private/VirtualShadowMaps/PageManagement.usf
@@ -4,6 +4,7 @@
 	PageManagement.usf: 
 =============================================================================*/
 
+#include "../HairStrands/HairStrandsVisibilityCommonStruct.ush"
 #include "../Common.ush"
 #include "../LightGridCommon.ush"
 #include "../SceneTexturesCommon.ush"
diff --git a/Engine/Shaders/Private/VirtualShadowMaps/VirtualShadowMapProjection.usf b/Engine/Shaders/Private/VirtualShadowMaps/VirtualShadowMapProjection.usf
--- a/Engine/Shaders/Private/VirtualShadowMaps/VirtualShadowMapProjection.usf
+++ b/Engine/Shaders/Private/VirtualShadowMaps/VirtualShadowMapProjection.usf
@@ -5,7 +5,7 @@
 =============================================================================*/
 
 #define SUPPORT_CONTACT_SHADOWS 1
-
+#include "../HairStrands/HairStrandsVisibilityCommonStruct.ush"
 #include "../Common.ush"
 #include "../SceneTexturesCommon.ush"
 #include "../DeferredShadingCommon.ush"
diff --git a/Engine/Source/Runtime/Renderer/Private/DeferredShadingRenderer.cpp b/Engine/Source/Runtime/Renderer/Private/DeferredShadingRenderer.cpp
--- a/Engine/Source/Runtime/Renderer/Private/DeferredShadingRenderer.cpp
+++ b/Engine/Source/Runtime/Renderer/Private/DeferredShadingRenderer.cpp
@@ -2183,20 +2183,17 @@ void FDeferredShadingSceneRenderer::Render(FRDGBuilder& GraphBuilder)
 		bAsyncComputeVolumetricCloud = RenderVolumetricCloud(GraphBuilder, SceneTextures, bSkipVolumetricRenderTarget, bSkipPerPixelTracing, HalfResolutionDepthCheckerboardMinMaxTexture, true, InstanceCullingManager);
 	}
 	
-	FHairStrandsRenderingData* HairDatas = nullptr;
-	FHairStrandsRenderingData* HairDatasStorage = GraphBuilder.AllocObject<FHairStrandsRenderingData>();
 	FRDGTextureRef ForwardScreenSpaceShadowMaskTexture = nullptr;
 	FRDGTextureRef ForwardScreenSpaceShadowMaskHairTexture = nullptr;
 	if (IsForwardShadingEnabled(ShaderPlatform))
 	{
 		if (bHairEnable)
 		{
-			RenderHairPrePass(GraphBuilder, Scene, Views, InstanceCullingManager, *HairDatasStorage);
-			RenderHairBasePass(GraphBuilder, Scene, SceneTextures, Views, InstanceCullingManager, *HairDatasStorage);
-			HairDatas = HairDatasStorage;
+			RenderHairPrePass(GraphBuilder, Scene, Views, InstanceCullingManager);
+			RenderHairBasePass(GraphBuilder, Scene, SceneTextures, Views, InstanceCullingManager);
 		}
 
-		RenderForwardShadowProjections(GraphBuilder, SceneTextures, ForwardScreenSpaceShadowMaskTexture, ForwardScreenSpaceShadowMaskHairTexture, HairDatas);
+		RenderForwardShadowProjections(GraphBuilder, SceneTextures, ForwardScreenSpaceShadowMaskTexture, ForwardScreenSpaceShadowMaskHairTexture);
 	}
 
 	FDBufferTextures DBufferTextures = CreateDBufferTextures(GraphBuilder, SceneTextures.Config.Extent, ShaderPlatform);
@@ -2318,9 +2315,8 @@ void FDeferredShadingSceneRenderer::Render(FRDGBuilder& GraphBuilder)
 	// Render hair
 	if (bHairEnable && !IsForwardShadingEnabled(ShaderPlatform))
 	{
-		RenderHairPrePass(GraphBuilder, Scene, Views, InstanceCullingManager, *HairDatasStorage);
-		RenderHairBasePass(GraphBuilder, Scene, SceneTextures, Views, InstanceCullingManager, *HairDatasStorage);
-		HairDatas = HairDatasStorage;
+		RenderHairPrePass(GraphBuilder, Scene, Views, InstanceCullingManager);
+		RenderHairBasePass(GraphBuilder, Scene, SceneTextures, Views, InstanceCullingManager);
 	}
 
 #if RHI_RAYTRACING
@@ -2454,7 +2450,7 @@ void FDeferredShadingSceneRenderer::Render(FRDGBuilder& GraphBuilder)
 		SCOPE_CYCLE_COUNTER(STAT_FDeferredShadingSceneRenderer_Lighting);
 
 		FRDGTextureRef DynamicBentNormalAOTexture = nullptr;
-		RenderDiffuseIndirectAndAmbientOcclusion(GraphBuilder, SceneTextures, LightingChannelsTexture, HairDatas, /* bIsVisualizePass = */ false);
+		RenderDiffuseIndirectAndAmbientOcclusion(GraphBuilder, SceneTextures, LightingChannelsTexture, /* bIsVisualizePass = */ false);
 
 		// These modulate the scenecolor output from the basepass, which is assumed to be indirect lighting
 		if (bAllowStaticLighting)
@@ -2479,7 +2475,7 @@ void FDeferredShadingSceneRenderer::Render(FRDGBuilder& GraphBuilder)
 #endif
 
 		GraphBuilder.SetCommandListStat(GET_STATID(STAT_CLM_Lighting));
-		RenderLights(GraphBuilder, SceneTextures, TranslucencyLightingVolumeTextures, LightingChannelsTexture, SortedLightSet, HairDatas);
+		RenderLights(GraphBuilder, SceneTextures, TranslucencyLightingVolumeTextures, LightingChannelsTexture, SortedLightSet);
 		GraphBuilder.SetCommandListStat(GET_STATID(STAT_CLM_AfterLighting));
 		AddServiceLocalQueuePass(GraphBuilder);
 
@@ -2488,13 +2484,12 @@ void FDeferredShadingSceneRenderer::Render(FRDGBuilder& GraphBuilder)
 		AddServiceLocalQueuePass(GraphBuilder);
 
 		// Render diffuse sky lighting and reflections that only operate on opaque pixels
-		RenderDeferredReflectionsAndSkyLighting(GraphBuilder, SceneTextures, DynamicBentNormalAOTexture, HairDatas);
+		RenderDeferredReflectionsAndSkyLighting(GraphBuilder, SceneTextures, DynamicBentNormalAOTexture);
 
 		AddSubsurfacePass(GraphBuilder, SceneTextures, Views);
 
-		if (HairDatas)
 		{
-			RenderHairStrandsSceneColorScattering(GraphBuilder, SceneTextures.Color.Target, Scene, Views, HairDatas);
+			RenderHairStrandsSceneColorScattering(GraphBuilder, SceneTextures.Color.Target, Scene, Views);
 		}
 
 	#if RHI_RAYTRACING
@@ -2502,17 +2497,17 @@ void FDeferredShadingSceneRenderer::Render(FRDGBuilder& GraphBuilder)
 		{
 			FRDGTextureRef SkyLightTexture = nullptr;
 			FRDGTextureRef SkyLightHitDistanceTexture = nullptr;
-			RenderRayTracingSkyLight(GraphBuilder, SceneTextures.Color.Target, SkyLightTexture, SkyLightHitDistanceTexture, HairDatas);
+			RenderRayTracingSkyLight(GraphBuilder, SceneTextures.Color.Target, SkyLightTexture, SkyLightHitDistanceTexture);
 			CompositeRayTracingSkyLight(GraphBuilder, SceneTextures, SkyLightTexture, SkyLightHitDistanceTexture);
 		}
 	#endif
 
 		AddServiceLocalQueuePass(GraphBuilder);
 	}
-	else if (HairDatas)
+	else if (HairStrands::HasViewHairStrandsData(Views))
 	{
-		RenderLightsForHair(GraphBuilder, SceneTextures.UniformBuffer, SortedLightSet, HairDatas, ForwardScreenSpaceShadowMaskHairTexture, LightingChannelsTexture);
-		RenderDeferredReflectionsAndSkyLightingHair(GraphBuilder, HairDatas);
+		RenderLightsForHair(GraphBuilder, SceneTextures.UniformBuffer, SortedLightSet, ForwardScreenSpaceShadowMaskHairTexture, LightingChannelsTexture);
+		RenderDeferredReflectionsAndSkyLightingHair(GraphBuilder);
 	}
 
 	if (bShouldRenderVolumetricCloud && IsVolumetricRenderTargetEnabled() && !bHasHalfResCheckerboardMinMaxDepth && bCanOverlayRayTracingOutput)
@@ -2616,9 +2611,9 @@ void FDeferredShadingSceneRenderer::Render(FRDGBuilder& GraphBuilder)
 		RenderDebugSkyAtmosphere(GraphBuilder, SceneTextures.Color.Target, SceneTextures.Depth.Target);
 	}
 
-	if (HairDatas && !IsHairStrandsComposeAfterTranslucency())
+	if (!IsHairStrandsComposeAfterTranslucency())
 	{
-		RenderHairComposition(GraphBuilder, Views, HairDatas, SceneTextures.Color.Target, SceneTextures.Depth.Target);
+		RenderHairComposition(GraphBuilder, Views, SceneTextures.Color.Target, SceneTextures.Depth.Target);
 	}
 
 	FSeparateTranslucencyTextures SeparateTranslucencyTextures(SeparateTranslucencyDimensions);
@@ -2683,14 +2678,14 @@ void FDeferredShadingSceneRenderer::Render(FRDGBuilder& GraphBuilder)
 
 	{
 		RDG_GPU_STAT_SCOPE(GraphBuilder, HairRendering);
-		if (HairDatas && IsHairStrandsComposeAfterTranslucency())
+		if (IsHairStrandsComposeAfterTranslucency())
 		{
-			RenderHairComposition(GraphBuilder, Views, HairDatas, SceneTextures.Color.Target, SceneTextures.Depth.Target);
+			RenderHairComposition(GraphBuilder, Views, SceneTextures.Color.Target, SceneTextures.Depth.Target);
 		}
 
 		if (HairStrandsBookmarkParameters.bHasElements)
 		{
-			RenderHairStrandsDebugInfo(GraphBuilder, Views, HairDatas, HairStrandsBookmarkParameters.HairClusterData, SceneTextures.Color.Target);
+			RenderHairStrandsDebugInfo(GraphBuilder, Views, HairStrandsBookmarkParameters.HairClusterData, SceneTextures.Color.Target);
 		}
 	}
 
@@ -2752,7 +2747,7 @@ void FDeferredShadingSceneRenderer::Render(FRDGBuilder& GraphBuilder)
 	}
 
 	RenderLumenSceneVisualization(GraphBuilder, SceneTextures);
-	RenderDiffuseIndirectAndAmbientOcclusion(GraphBuilder, SceneTextures, LightingChannelsTexture, HairDatas, /* bIsVisualizePass = */ true);
+	RenderDiffuseIndirectAndAmbientOcclusion(GraphBuilder, SceneTextures, LightingChannelsTexture, /* bIsVisualizePass = */ true);
 
 	if (ViewFamily.EngineShowFlags.StationaryLightOverlap)
 	{
diff --git a/Engine/Source/Runtime/Renderer/Private/DeferredShadingRenderer.h b/Engine/Source/Runtime/Renderer/Private/DeferredShadingRenderer.h
--- a/Engine/Source/Runtime/Renderer/Private/DeferredShadingRenderer.h
+++ b/Engine/Source/Runtime/Renderer/Private/DeferredShadingRenderer.h
@@ -40,7 +40,6 @@ struct FSceneWithoutWaterTextures;
 struct FHeightFogRenderingParameters;
 struct FRayTracingReflectionOptions;
 struct FHairStrandsTransmittanceMaskData;
-struct FHairStrandsRenderingData;
 
 struct FTranslucencyLightingVolumeTextures;
 
@@ -458,17 +457,15 @@ private:
 		FRDGBuilder& GraphBuilder,
 		FSceneTextures& SceneTextures,
 		FRDGTextureRef LightingChannelsTexture,
-		FHairStrandsRenderingData* HairDatas,
 		bool bIsVisualizePass);
 
 	/** Renders sky lighting and reflections that can be done in a deferred pass. */
 	void RenderDeferredReflectionsAndSkyLighting(
 		FRDGBuilder& GraphBuilder,
 		const FSceneTextures& SceneTextures,
-		FRDGTextureRef DynamicBentNormalAOTexture,
-		struct FHairStrandsRenderingData* HairDatas);
+		FRDGTextureRef DynamicBentNormalAOTexture);
 
-	void RenderDeferredReflectionsAndSkyLightingHair(FRDGBuilder& GraphBuilder, struct FHairStrandsRenderingData* HairDatas);
+	void RenderDeferredReflectionsAndSkyLightingHair(FRDGBuilder& GraphBuilder);
 
 	/** Computes DFAO, modulates it to scene color (which is assumed to contain diffuse indirect lighting), and stores the output bent normal for use occluding specular. */
 	void RenderDFAOAsIndirectShadowing(
@@ -590,8 +587,7 @@ private:
 		FMinimalSceneTextures& SceneTextures,
 		const FTranslucencyLightingVolumeTextures& TranslucencyLightingVolumeTextures,
 		FRDGTextureRef LightingChannelsTexture,
-		FSortedLightSetSceneInfo& SortedLightSet,
-		const FHairStrandsRenderingData* HairDatas);
+		FSortedLightSetSceneInfo& SortedLightSet);
 
 	/** Renders an array of lights for the stationary light overlap viewmode. */
 	void RenderLightArrayForOverlapViewmode(
@@ -671,7 +667,6 @@ private:
 		FRDGTextureRef ScreenShadowMaskTexture,
 		FRDGTextureRef ScreenShadowMaskSubPixelTexture,
 		const FLightSceneInfo* LightSceneInfo,
-		const FHairStrandsRenderingData* HairDatas,
 		bool bProjectingForForwardShading);
 
 	void RenderDeferredShadowProjections(
@@ -681,15 +676,13 @@ private:
 		const FLightSceneInfo* LightSceneInfo,
 		FRDGTextureRef ScreenShadowMaskTexture,
 		FRDGTextureRef ScreenShadowMaskSubPixelTexture,
-		const FHairStrandsRenderingData* HairDatas,
 		bool& bInjectedTranslucentVolume);
 
 	void RenderForwardShadowProjections(
 		FRDGBuilder& GraphBuilder,
 		const FMinimalSceneTextures& SceneTextures,
 		FRDGTextureRef& ForwardScreenSpaceShadowMask,
-		FRDGTextureRef& ForwardScreenSpaceShadowMaskSubPixel,
-		const FHairStrandsRenderingData* InHairDatas);
+		FRDGTextureRef& ForwardScreenSpaceShadowMaskSubPixel);
 
 	/** Used by RenderLights to render a light function to the attenuation buffer. */
 	bool RenderLightFunction(
@@ -723,11 +716,9 @@ private:
 	void RenderLight(
 		FRHICommandList& RHICmdList,
 		const FViewInfo& View,
-		int32 ViewIndex,
 		const FLightSceneInfo* LightSceneInfo,
 		FRHITexture* ScreenShadowMaskTexture,
 		FRHITexture* LightingChannelTexture,
-		const struct FHairStrandsVisibilityViews* InHairVisibilityViews,
 		bool bRenderOverlap,
 		bool bIssueDrawEvent);
 
@@ -737,26 +728,24 @@ private:
 		const FLightSceneInfo* LightSceneInfo,
 		FRDGTextureRef ScreenShadowMaskTexture,
 		FRDGTextureRef LightingChannelsTexture,
-		const FHairStrandsVisibilityViews* InHairVisibilityViews,
 		bool bRenderOverlap);
 
 	void RenderLightsForHair(
 		FRDGBuilder& GraphBuilder,
 		TRDGUniformBufferRef<FSceneTextureUniformParameters> SceneTexturesUniformBuffer,
 		FSortedLightSetSceneInfo& SortedLightSet,
-		const FHairStrandsRenderingData* HairDatas,
 		FRDGTextureRef InScreenShadowMaskSubPixelTexture,
 		FRDGTextureRef LightingChannelsTexture);
 
 	/** Specialized version of RenderLight for hair (run lighting evaluation on at sub-pixel rate, without depth bound) */
 	void RenderLightForHair(
 		FRDGBuilder& GraphBuilder,
+		FViewInfo& View,
 		TRDGUniformBufferRef<FSceneTextureUniformParameters> SceneTexturesUniformBuffer,
 		const FLightSceneInfo* LightSceneInfo,
 		FRDGTextureRef ScreenShadowMaskSubPixelTexture,
 		FRDGTextureRef LightingChannelsTexture,
-		const FHairStrandsTransmittanceMaskData& InTransmittanceMaskData,
-		const struct FHairStrandsVisibilityViews* InHairVisibilityViews);
+		const FHairStrandsTransmittanceMaskData& InTransmittanceMaskData);
 
 	/** Renders an array of simple lights using standard deferred shading. */
 	void RenderSimpleLightsStandardDeferred(
@@ -921,8 +910,7 @@ private:
 		FRDGBuilder& GraphBuilder,
 		FRDGTextureRef SceneColorTexture,
 		FRDGTextureRef& OutSkyLightTexture,
-		FRDGTextureRef& OutHitDistanceTexture,
-		const FHairStrandsRenderingData* HairDatas);
+		FRDGTextureRef& OutHitDistanceTexture);
 
 	void RenderRayTracingPrimaryRaysView(
 		FRDGBuilder& GraphBuilder,
diff --git a/Engine/Source/Runtime/Renderer/Private/HairStrands/HairStrandsCluster.cpp b/Engine/Source/Runtime/Renderer/Private/HairStrands/HairStrandsCluster.cpp
--- a/Engine/Source/Runtime/Renderer/Private/HairStrands/HairStrandsCluster.cpp
+++ b/Engine/Source/Runtime/Renderer/Private/HairStrands/HairStrandsCluster.cpp
@@ -154,90 +154,82 @@ static void InternalUpdateMacroGroup(FHairStrandsMacroGroupData& MacroGroup, int
 	}
 }
 
-FHairStrandsMacroGroupViews CreateHairStrandsMacroGroups(
+void CreateHairStrandsMacroGroups(
 	FRDGBuilder& GraphBuilder,
 	const FScene* Scene,
-	const TArray<FViewInfo>& Views)
+	FViewInfo& View)
 {
+	if (!View.Family || View.HairStrandsMeshElements.Num() == 0 || View.bIsPlanarReflection || View.bIsReflectionCapture)
+	{
+		return;
+	}
+
 	static const FVertexFactoryType* CompatibleVF = FVertexFactoryType::GetVFByName(TEXT("FHairStrandsVertexFactory"));
 
-	FHairStrandsMacroGroupViews MacroGroupsViews;
-	for (int32 ViewIndex = 0; ViewIndex < Views.Num(); ViewIndex++)
-	{
-		const FViewInfo& View = Views[ViewIndex];
-		if (View.Family)
-		{
-			int32 MaterialId = 0;
-			FHairStrandsMacroGroupDatas& MacroGroups = MacroGroupsViews.Views.AddDefaulted_GetRef();
+	TArray<FHairStrandsMacroGroupData, SceneRenderingAllocator>& MacroGroups = View.HairStrandsViewData.MacroGroupDatas;
+	FHairStrandsMacroGroupResources& MacroGroupResources = View.HairStrandsViewData.MacroGroupResources;
 
-			if (View.HairStrandsMeshElements.Num() == 0 || View.bIsPlanarReflection || View.bIsReflectionCapture)
-			{
-				continue;
-			}
+	int32 MaterialId = 0;
 
-			// Aggregate all hair primitives within the same area into macro groups, for allocating/rendering DOM/voxel
-			uint32 MacroGroupId = 0;
-			auto UpdateMacroGroup = [&MacroGroups, &View, &MacroGroupId, &MaterialId](const FMeshBatchAndRelevance* MeshBatchAndRelevance, const FPrimitiveSceneProxy* Proxy)
-			{
-				const bool bIsHairStrandsFactory = MeshBatchAndRelevance->Mesh->VertexFactory->GetType()->GetHashedName() == CompatibleVF->GetHashedName();
-				if (!bIsHairStrandsFactory)
-					return;
-
-				const FBoxSphereBounds& PrimitiveBounds = Proxy->GetBounds();
-
-				bool bFound = false;
-				for (FHairStrandsMacroGroupData& MacroGroup : MacroGroups.Datas)
-				{
-					const bool bIntersect = FBoxSphereBounds::SpheresIntersect(MacroGroup.Bounds, PrimitiveBounds);
-					if (bIntersect)
-					{
-						MacroGroup.Bounds = Union(MacroGroup.Bounds, PrimitiveBounds);
-						InternalUpdateMacroGroup(MacroGroup, MaterialId, MeshBatchAndRelevance, Proxy);
-						bFound = true;
-						break;
-					}
-				}
-
-				if (!bFound)
-				{
-					FHairStrandsMacroGroupData MacroGroup;
-					MacroGroup.MacroGroupId = MacroGroupId++;
-					InternalUpdateMacroGroup(MacroGroup, MaterialId, MeshBatchAndRelevance, Proxy);
-					MacroGroup.Bounds = PrimitiveBounds;
-					MacroGroups.Datas.Add(MacroGroup);
-				}
-			};
-
-			for (const FMeshBatchAndRelevance& MeshBatchAndRelevance : View.HairStrandsMeshElements)
-			{
-				UpdateMacroGroup(&MeshBatchAndRelevance, MeshBatchAndRelevance.PrimitiveSceneProxy);
-			}
+	// Aggregate all hair primitives within the same area into macro groups, for allocating/rendering DOM/voxel
+	uint32 MacroGroupId = 0;
+	auto UpdateMacroGroup = [&MacroGroups, &MacroGroupId, &MaterialId](const FMeshBatchAndRelevance* MeshBatchAndRelevance, const FPrimitiveSceneProxy* Proxy)
+	{
+		const bool bIsHairStrandsFactory = MeshBatchAndRelevance->Mesh->VertexFactory->GetType()->GetHashedName() == CompatibleVF->GetHashedName();
+		if (!bIsHairStrandsFactory)
+			return;
 
-			for (FHairStrandsMacroGroupData& MacroGroup : MacroGroups.Datas)
-			{
-				MacroGroup.ScreenRect = ComputeProjectedScreenRect(MacroGroup.Bounds.GetBox(), View);
-			}
+		const FBoxSphereBounds& PrimitiveBounds = Proxy->GetBounds();
 
-			// Build hair macro group AABBB
-			const uint32 MacroGroupCount = MacroGroups.Datas.Num();
-			if (MacroGroupCount > 0)
+		bool bFound = false;
+		for (FHairStrandsMacroGroupData& MacroGroup : MacroGroups)
+		{
+			const bool bIntersect = FBoxSphereBounds::SpheresIntersect(MacroGroup.Bounds, PrimitiveBounds);
+			if (bIntersect)
 			{
-				DECLARE_GPU_STAT(HairStrandsAABB);
-				RDG_EVENT_SCOPE(GraphBuilder, "HairStrandsAABB");
-				RDG_GPU_STAT_SCOPE(GraphBuilder, HairStrandsAABB);
-
-				MacroGroups.MacroGroupResources.MacroGroupAABBsBuffer = GraphBuilder.CreateBuffer(FRDGBufferDesc::CreateBufferDesc(4, 6 * MacroGroupCount), TEXT("Hair.MacroGroupAABBBuffer"));
-				FRDGBufferUAVRef MacroGroupAABBBufferUAV = GraphBuilder.CreateUAV(MacroGroups.MacroGroupResources.MacroGroupAABBsBuffer, PF_R32_SINT);
-				for (FHairStrandsMacroGroupData& MacroGroup : MacroGroups.Datas)
-				{				
-					AddHairMacroGroupAABBPass(GraphBuilder, View.ShaderMap, MacroGroup, MacroGroupAABBBufferUAV);
-				}
-				MacroGroups.MacroGroupResources.MacroGroupCount = MacroGroups.Datas.Num();
+				MacroGroup.Bounds = Union(MacroGroup.Bounds, PrimitiveBounds);
+				InternalUpdateMacroGroup(MacroGroup, MaterialId, MeshBatchAndRelevance, Proxy);
+				bFound = true;
+				break;
 			}
 		}
+
+		if (!bFound)
+		{
+			FHairStrandsMacroGroupData MacroGroup;
+			MacroGroup.MacroGroupId = MacroGroupId++;
+			InternalUpdateMacroGroup(MacroGroup, MaterialId, MeshBatchAndRelevance, Proxy);
+			MacroGroup.Bounds = PrimitiveBounds;
+			MacroGroups.Add(MacroGroup);
+		}
+	};
+
+	for (const FMeshBatchAndRelevance& MeshBatchAndRelevance : View.HairStrandsMeshElements)
+	{
+		UpdateMacroGroup(&MeshBatchAndRelevance, MeshBatchAndRelevance.PrimitiveSceneProxy);
 	}
 
-	return MacroGroupsViews;
+	for (FHairStrandsMacroGroupData& MacroGroup : MacroGroups)
+	{
+		MacroGroup.ScreenRect = ComputeProjectedScreenRect(MacroGroup.Bounds.GetBox(), View);
+	}
+
+	// Build hair macro group AABBB
+	const uint32 MacroGroupCount = MacroGroups.Num();
+	if (MacroGroupCount > 0)
+	{
+		DECLARE_GPU_STAT(HairStrandsAABB);
+		RDG_EVENT_SCOPE(GraphBuilder, "HairStrandsAABB");
+		RDG_GPU_STAT_SCOPE(GraphBuilder, HairStrandsAABB);
+
+		MacroGroupResources.MacroGroupAABBsBuffer = GraphBuilder.CreateBuffer(FRDGBufferDesc::CreateBufferDesc(4, 6 * MacroGroupCount), TEXT("Hair.MacroGroupAABBBuffer"));
+		FRDGBufferUAVRef MacroGroupAABBBufferUAV = GraphBuilder.CreateUAV(MacroGroupResources.MacroGroupAABBsBuffer, PF_R32_SINT);
+		for (FHairStrandsMacroGroupData& MacroGroup : MacroGroups)
+		{				
+			AddHairMacroGroupAABBPass(GraphBuilder, View.ShaderMap, MacroGroup, MacroGroupAABBBufferUAV);
+		}
+		MacroGroupResources.MacroGroupCount = MacroGroups.Num();
+	}
 }
 
 bool FHairStrandsMacroGroupData::PrimitiveInfo::IsCullingEnable() const
diff --git a/Engine/Source/Runtime/Renderer/Private/HairStrands/HairStrandsCluster.h b/Engine/Source/Runtime/Renderer/Private/HairStrands/HairStrandsCluster.h
--- a/Engine/Source/Runtime/Renderer/Private/HairStrands/HairStrandsCluster.h
+++ b/Engine/Source/Runtime/Renderer/Private/HairStrands/HairStrandsCluster.h
@@ -12,65 +12,7 @@
 #include "HairStrandsDeepShadow.h"
 #include "SceneManagement.h"
 
-// -----------------------------
-// Hair Data structures overview 
-// -----------------------------
-// A groom component contains one or several HairGroup. These hair group are send to the 
-// render as mesh batches. These meshes batches are filtered/culled per view, and regroup 
-// into HairMacroGroup for computing voxelization/DOM data, ...
-//
-// The hierarchy of the data structure is as follow:
-//  * HairMacroGroup
-//  * HairGroup
-//  * HairCluster
-
-struct FHairMacroGroupAABBData
-{
-	uint32 MacroGroupCount = 0;
-	FRDGBufferRef MacroGroupAABBsBuffer = nullptr;
-};
-
-/// Hair macro group infos
-struct FHairStrandsMacroGroupData
-{
-	// List of primitive/mesh batch within an instance group
-	struct PrimitiveInfo
-	{
-		FMeshBatchAndRelevance MeshBatchAndRelevance;
-		uint32 MaterialId;
-		uint32 ResourceId;
-		uint32 GroupIndex;		
-		FHairGroupPublicData* PublicDataPtr = nullptr;
-		bool IsCullingEnable() const;
-	};
-	typedef TArray<PrimitiveInfo, SceneRenderingAllocator> TPrimitiveInfos;
-
-	FVirtualVoxelNodeDesc VirtualVoxelNodeDesc;
-	FHairStrandsDeepShadowDatas DeepShadowDatas;
-	TPrimitiveInfos PrimitivesInfos;
-	FBoxSphereBounds Bounds;
-	FIntRect ScreenRect;
-	uint32 MacroGroupId;
-
-	bool bNeedScatterSceneLighting = false;
-};
-
-/// Store all hair strands macro group infos for a given view
-struct FHairStrandsMacroGroupDatas
-{
-	TArray<FHairStrandsMacroGroupData, SceneRenderingAllocator> Datas;
-	FDeepShadowResources DeepShadowResources;
-	FVirtualVoxelResources VirtualVoxelResources;
-	FHairMacroGroupAABBData MacroGroupResources;
-};
-
-/// Store all hair strands macro group info for all views
-struct FHairStrandsMacroGroupViews
-{
-	TArray<FHairStrandsMacroGroupDatas, SceneRenderingAllocator> Views;
-};
-
-FHairStrandsMacroGroupViews CreateHairStrandsMacroGroups(
+void CreateHairStrandsMacroGroups(
 	FRDGBuilder& GraphBuilder,
 	const FScene* Scene,
-	const TArray<FViewInfo>& Views);
\ No newline at end of file
+	FViewInfo& View);
\ No newline at end of file
diff --git a/Engine/Source/Runtime/Renderer/Private/HairStrands/HairStrandsComposition.cpp b/Engine/Source/Runtime/Renderer/Private/HairStrands/HairStrandsComposition.cpp
--- a/Engine/Source/Runtime/Renderer/Private/HairStrands/HairStrandsComposition.cpp
+++ b/Engine/Source/Runtime/Renderer/Private/HairStrands/HairStrandsComposition.cpp
@@ -587,47 +587,26 @@ static void AddHairVisibilityGBufferWritePass(
 		});
 }
 
-void RenderHairComposition(
+static void InternalRenderHairComposition(
 	FRDGBuilder& GraphBuilder,
-	const TArray<FViewInfo>& Views,
-	const FHairStrandsRenderingData* HairDatas,
+	const FViewInfo& View,
 	FRDGTextureRef SceneColorTexture,
 	FRDGTextureRef SceneDepthTexture)
 {
-	if (!HairDatas || HairDatas->HairVisibilityViews.HairDatas.Num() == 0 || Views.Num() == 0)
-		return;
-
-	const FHairStrandsVisibilityViews& HairVisibilityViews = HairDatas->HairVisibilityViews;
-
 	DECLARE_GPU_STAT(HairStrandsComposition);
 	RDG_EVENT_SCOPE(GraphBuilder, "HairStrandsComposition");
 	RDG_GPU_STAT_SCOPE(GraphBuilder, HairStrandsComposition);
 
-	for (int32 ViewIndex = 0; ViewIndex < Views.Num(); ViewIndex++)
 	{
-		const FViewInfo& View = Views[ViewIndex];
-		if (View.Family)
 		{
-			if (ViewIndex < HairVisibilityViews.HairDatas.Num())
 			{
-				const FHairStrandsMacroGroupDatas& MacroGroupDatas = HairDatas->MacroGroupsPerViews.Views[ViewIndex];
-				const FHairStrandsVisibilityData& VisibilityData = HairVisibilityViews.HairDatas[ViewIndex];
+				const FHairStrandsVisibilityData& VisibilityData = View.HairStrandsViewData.VisibilityData;
 
 				if (!VisibilityData.CategorizationTexture)
 				{
-					continue; // Automatically skip for any view not rendering hair
+					return; // Automatically skip for any view not rendering hair
 				}
 
-				// todo: rehook the diffusion pass
-				//AddHairDiffusionPass(
-				//	GraphBuilder,
-				//	View,
-				//	VisibilityData,
-				//	MacroGroupDatas.VirtualVoxelResources,
-				//	SceneDepthTexture,
-				//	SceneColorSubPixelTexture,
-				//	SceneColorTexture);
-
 				FRDGTextureRef DOFDepth = nullptr;
 				const bool bHairDOF = GHairStrandsComposeDOFDepth > 0 ? 1 : 0;
 				if (bHairDOF)
@@ -650,20 +629,21 @@ void RenderHairComposition(
 					SceneColorTexture,
 					SceneDepthTexture);
 
-				if (HairVisibilityViews.HairDatas[ViewIndex].VelocityTexture)
+				
+				if (VisibilityData.VelocityTexture)
 				{
 					AddHairVisibilityFastResolveMSAAPass(
 						GraphBuilder,
 						View,
-						HairVisibilityViews.HairDatas[ViewIndex].VelocityTexture,
+						VisibilityData.VelocityTexture,
 						SceneDepthTexture);
 				}
-				else if (HairVisibilityViews.HairDatas[ViewIndex].ResolveMaskTexture)
+				else if (VisibilityData.ResolveMaskTexture)
 				{
 					AddHairVisibilityFastResolveMaskPass(
 						GraphBuilder,
 						View,
-						HairVisibilityViews.HairDatas[ViewIndex].ResolveMaskTexture,
+						VisibilityData.ResolveMaskTexture,
 						SceneDepthTexture);
 				}
 
@@ -715,3 +695,22 @@ void RenderHairComposition(
 		}
 	}
 }
+
+void RenderHairComposition(
+	FRDGBuilder& GraphBuilder,
+	const TArray<FViewInfo>& Views,
+	FRDGTextureRef SceneColorTexture,
+	FRDGTextureRef SceneDepthTexture)
+{
+	for (const FViewInfo& View : Views)
+	{		
+		if (View.Family && HairStrands::HasViewHairStrandsData(View))
+		{			
+			InternalRenderHairComposition(
+				GraphBuilder,
+				View,
+				SceneColorTexture,
+				SceneDepthTexture);
+		}
+	}
+}
diff --git a/Engine/Source/Runtime/Renderer/Private/HairStrands/HairStrandsComposition.h b/Engine/Source/Runtime/Renderer/Private/HairStrands/HairStrandsComposition.h
--- a/Engine/Source/Runtime/Renderer/Private/HairStrands/HairStrandsComposition.h
+++ b/Engine/Source/Runtime/Renderer/Private/HairStrands/HairStrandsComposition.h
@@ -12,6 +12,5 @@
 void RenderHairComposition(
 	FRDGBuilder& GraphBuilder, 
 	const TArray<FViewInfo>& Views,
-	const struct FHairStrandsRenderingData* HairDatas,
 	FRDGTextureRef SceneColorTexture,
 	FRDGTextureRef SceneDepthTexture);
diff --git a/Engine/Source/Runtime/Renderer/Private/HairStrands/HairStrandsData.h b/Engine/Source/Runtime/Renderer/Private/HairStrands/HairStrandsData.h
--- a/Engine/Source/Runtime/Renderer/Private/HairStrands/HairStrandsData.h
+++ b/Engine/Source/Runtime/Renderer/Private/HairStrands/HairStrandsData.h
@@ -10,17 +10,332 @@
 #include "RendererInterface.h"
 #include "Shader.h"
 
+////////////////////////////////////////////////////////////////////////////////////
+// HairStrands uniform buffer
+
 BEGIN_GLOBAL_SHADER_PARAMETER_STRUCT(FHairStrandsViewUniformParameters, )	
-	SHADER_PARAMETER_RDG_TEXTURE(Texture2D<uint4>, HairCategorizationTexture)
-	SHADER_PARAMETER_RDG_TEXTURE(Texture2D<float>, HairOnlyDepthTexture)
+//	SHADER_PARAMETER(uint32, SampleCount)
+	SHADER_PARAMETER(float, HairDualScatteringRoughnessOverride)						// Override the roughness used for dual scattering (for hack/test purpose only)
+	SHADER_PARAMETER(FIntPoint, HairSampleViewportResolution)							// Maximum viewport resolution of the sample space
+	SHADER_PARAMETER_RDG_TEXTURE(Texture2D<uint4>, HairCategorizationTexture)			// Categorization texture aggregating hair info in screen space (closest depth, coverage, ...)
+	SHADER_PARAMETER_RDG_TEXTURE(Texture2D<float>, HairOnlyDepthTexture)				// Depth texture containing only hair depth
+	SHADER_PARAMETER_RDG_TEXTURE(Texture2D<uint>, HairSampleOffset)						// Offset & count, for accessing pixel's samples, based on screen pixel position
+	SHADER_PARAMETER_RDG_TEXTURE(Texture2D<uint>, HairSampleCount)						// Total count of hair sample, in sample space
+	SHADER_PARAMETER_RDG_BUFFER_SRV(StructuredBuffer<FPackedHairSample>, HairSampleData)// Sample data (coverage, tangent, base color, ...), in sample space // HAIRSTRANDS_TODO: change this to be a uint4 so that we don't have to include the type for generated contant buffer
+	SHADER_PARAMETER_RDG_BUFFER_SRV(StructuredBuffer<uint>, HairSampleCoords)			// Screen pixel coordinate of each sample, in sample space
+END_GLOBAL_SHADER_PARAMETER_STRUCT()
+
+////////////////////////////////////////////////////////////////////////////////////
+// Visibility Data
+
+struct FHairStrandsVisibilityData
+{
+	FRDGTextureRef DepthTexture = nullptr;
+	FRDGTextureRef IDTexture = nullptr;
+	FRDGTextureRef MaterialTexture = nullptr;
+	FRDGTextureRef AttributeTexture = nullptr;
+	FRDGTextureRef VelocityTexture = nullptr;
+	FRDGTextureRef ResolveMaskTexture = nullptr;
+	FRDGTextureRef CategorizationTexture = nullptr;
+	FRDGTextureRef ViewHairCountTexture = nullptr;
+	FRDGTextureRef ViewHairCountUintTexture = nullptr;
+	FRDGTextureRef DepthTextureUint = nullptr;
+	FRDGTextureRef EmissiveTexture = nullptr;
+
+	FRDGTextureRef HairOnlyDepthTexture = nullptr;
+
+	FRDGTextureRef ViewHairVisibilityTexture0 = nullptr;
+	FRDGTextureRef ViewHairVisibilityTexture1 = nullptr;
+	FRDGTextureRef ViewHairVisibilityTexture2 = nullptr;
+	FRDGTextureRef ViewHairVisibilityTexture3 = nullptr;
+
+	FRDGTextureRef LightChannelMaskTexture = nullptr;
+
+	FRDGTextureRef	PPLLNodeCounterTexture = nullptr;
+	FRDGTextureRef	PPLLNodeIndexTexture = nullptr;
+	FRDGBufferRef	PPLLNodeDataBuffer = nullptr;
+	uint32			MaxPPLLNodePerPixelCount = 0;
+	uint32			MaxPPLLNodeCount = 0;
+
+	FRDGTextureRef	TileIndexTexture = nullptr;
+	FRDGBufferRef	TileBuffer = nullptr;
+	FRDGBufferRef	TileIndirectArgs = nullptr;
+	uint32			TileSize = 8;
+	uint32			TileThreadGroupSize = 32;
+
+	uint32			MaxSampleCount = 8;
+	uint32			MaxNodeCount = 0;
+	FRDGTextureRef	NodeCount = nullptr;
+	FRDGTextureRef	NodeIndex = nullptr;
+	FRDGBufferRef	NodeData = nullptr;
+	FRDGBufferRef	NodeCoord = nullptr;
+	FRDGBufferRef	NodeIndirectArg = nullptr;
+	uint32			NodeGroupSize = 0;
+
+	// Hair lighting is accumulated within this buffer
+	// Allocated conservatively
+	// User indirect dispatch for accumulating contribution
+	FIntPoint SampleLightingViewportResolution;
+	FRDGTextureRef SampleLightingBuffer = nullptr;
+};
+
+////////////////////////////////////////////////////////////////////////////////////
+// Voxel data
+
+struct FHairStrandsVoxelNodeDesc
+{
+	FVector WorldMinAABB = FVector::ZeroVector;
+	FVector WorldMaxAABB = FVector::ZeroVector;
+	FIntVector PageIndexResolution = FIntVector::ZeroValue;
+	FMatrix WorldToClip;
+};
+
+struct FPackedVirtualVoxelNodeDesc
+{
+	// This is just a placeholder having the correct size. The actual definition is in HairStradsNVoxelPageCommon.ush
+	const static EPixelFormat Format = PF_R32G32B32A32_UINT;
+	const static uint32 ComponentCount = 2;
+
+	// Shader View is struct { uint4; uint4; }
+	FVector	MinAABB;
+	uint32	PackedPageIndexResolution;
+	FVector	MaxAABB;
+	uint32	PageIndexOffset;
+};
+
+// PixelRadiusAtDepth1 shouldn't be stored into this structure should be view independent, 
+// but is put here for convenience at the moment since multiple views are not supported 
+// at the moment
+BEGIN_SHADER_PARAMETER_STRUCT(FHairStrandsVoxelCommonParameters, )
+	SHADER_PARAMETER(FIntVector, PageCountResolution)
+	SHADER_PARAMETER(float, VoxelWorldSize)
+	SHADER_PARAMETER(FIntVector, PageTextureResolution)
+	SHADER_PARAMETER(uint32, PageCount)
+	SHADER_PARAMETER(uint32, PageResolution)
+	SHADER_PARAMETER(uint32, PageIndexCount)
+	SHADER_PARAMETER(uint32, IndirectDispatchGroupSize)
+	SHADER_PARAMETER(uint32, NodeDescCount)
+
+	SHADER_PARAMETER(float, DensityScale)
+	SHADER_PARAMETER(float, DensityScale_AO)
+	SHADER_PARAMETER(float, DensityScale_Shadow)
+	SHADER_PARAMETER(float, DensityScale_Transmittance)
+	SHADER_PARAMETER(float, DensityScale_Environment)
+	SHADER_PARAMETER(float, DensityScale_Raytracing)
+
+	SHADER_PARAMETER(float, DepthBiasScale_Shadow)
+	SHADER_PARAMETER(float, DepthBiasScale_Transmittance)
+	SHADER_PARAMETER(float, DepthBiasScale_Environment)
+
+	SHADER_PARAMETER(float, SteppingScale_Shadow)
+	SHADER_PARAMETER(float, SteppingScale_Transmittance)
+	SHADER_PARAMETER(float, SteppingScale_Environment)
+	SHADER_PARAMETER(float, SteppingScale_Raytracing)
+
+	SHADER_PARAMETER(float, HairCoveragePixelRadiusAtDepth1)
+	SHADER_PARAMETER(float, Raytracing_ShadowOcclusionThreshold)
+	SHADER_PARAMETER(float, Raytracing_SkyOcclusionThreshold)
+	SHADER_PARAMETER_RDG_BUFFER_SRV(Buffer<uint>, PageIndexBuffer)
+	SHADER_PARAMETER_RDG_BUFFER_SRV(Buffer<uint2>, PageIndexOccupancyBuffer)
+	SHADER_PARAMETER_RDG_BUFFER_SRV(Buffer<uint>, PageIndexCoordBuffer)
+	SHADER_PARAMETER_RDG_BUFFER_SRV(StructuredBuffer<FPackedVirtualVoxelNodeDesc>, NodeDescBuffer) // Packed into 2 x uint4
+	SHADER_PARAMETER_RDG_TEXTURE(Texture2D<float>, HairCoverageLUT)
+	SHADER_PARAMETER_SAMPLER(SamplerState, HairCoverageSampler)
+END_SHADER_PARAMETER_STRUCT()
+
+BEGIN_GLOBAL_SHADER_PARAMETER_STRUCT(FVirtualVoxelParameters, RENDERER_API)
+	SHADER_PARAMETER_STRUCT_INCLUDE(FHairStrandsVoxelCommonParameters, Common)
+	SHADER_PARAMETER_RDG_TEXTURE(Texture3D<uint>, PageTexture)
 END_GLOBAL_SHADER_PARAMETER_STRUCT()
 
+struct FHairStrandsVoxelResources
+{
+	FVirtualVoxelParameters	Parameters;
+	TRDGUniformBufferRef<FVirtualVoxelParameters> UniformBuffer;
+	FRDGTextureRef PageTexture = nullptr;
+	FRDGBufferRef PageIndexBuffer = nullptr;
+	FRDGBufferRef PageIndexOccupancyBuffer = nullptr;
+	FRDGBufferRef NodeDescBuffer = nullptr;
+	FRDGBufferRef PageIndexCoordBuffer = nullptr;
+	FRDGBufferRef IndirectArgsBuffer = nullptr;
+	FRDGBufferRef PageIndexGlobalCounter = nullptr;
+	FRDGBufferRef VoxelizationViewInfoBuffer = nullptr;
+
+	const bool IsValid() const { return UniformBuffer != nullptr && PageTexture != nullptr && NodeDescBuffer != nullptr; }
+};
+
+
+////////////////////////////////////////////////////////////////////////////////////
+// Deep shadow data
+
+struct FMinHairRadiusAtDepth1
+{
+	float Primary = 1;
+	float Velocity = 1;
+	float Stable = 1;
+};
+
+/// Hold deep shadow information for a given light.
+struct FHairStrandsDeepShadowData
+{
+	static const uint32 MaxMacroGroupCount = 16u;
+
+	FMatrix CPU_WorldToLightTransform;
+	FMinHairRadiusAtDepth1 CPU_MinStrandRadiusAtDepth1;
+	FIntRect AtlasRect;
+	uint32 MacroGroupId = ~0;
+	uint32 AtlasSlotIndex = 0;
+
+	FIntPoint ShadowResolution = FIntPoint::ZeroValue;
+	uint32 LightId = ~0;
+	ELightComponentType LightType = LightType_MAX;
+	FVector  LightDirection;
+	FVector4 LightPosition;
+	FLinearColor LightLuminance;
+	float LayerDistribution;
+
+	FBoxSphereBounds Bounds;
+};
+
+typedef TArray<FHairStrandsDeepShadowData, SceneRenderingAllocator> FHairStrandsDeepShadowDatas;
+
+struct FHairStrandsDeepShadowResources
+{
+	// Limit the number of atlas slot to 32, in order to create the view info per slot in single compute
+	// This limitation can be alleviate, and is just here for convenience (see FDeepShadowCreateViewInfoCS)
+	static const uint32 MaxAtlasSlotCount = 32u;
+
+	uint32 TotalAtlasSlotCount = 0;
+	FIntPoint AtlasSlotResolution;
+	bool bIsGPUDriven = false;
+
+	FRDGTextureRef DepthAtlasTexture = nullptr;
+	FRDGTextureRef LayersAtlasTexture = nullptr;
+	FRDGBufferRef DeepShadowWorldToLightTransforms = nullptr;
+};
+
+////////////////////////////////////////////////////////////////////////////////////
+// Cluster data
+
+// A groom component contains one or several HairGroup. These hair group are send to the 
+// render as mesh batches. These meshes batches are filtered/culled per view, and regroup 
+// into HairMacroGroup for computing voxelization/DOM data, ...
+//
+// The hierarchy of the data structure is as follow:
+//  * HairMacroGroup
+//  * HairGroup
+//  * HairCluster
+
+struct FHairStrandsMacroGroupResources
+{
+	uint32 MacroGroupCount = 0;
+	FRDGBufferRef MacroGroupAABBsBuffer = nullptr;
+};
+
+class FHairGroupPublicData;
+
+/// Hair macro group infos
+struct FHairStrandsMacroGroupData
+{
+	// List of primitive/mesh batch within an instance group
+	struct PrimitiveInfo
+	{
+		FMeshBatchAndRelevance MeshBatchAndRelevance;
+		uint32 MaterialId;
+		uint32 ResourceId;
+		uint32 GroupIndex;
+		FHairGroupPublicData* PublicDataPtr = nullptr;
+		bool IsCullingEnable() const;
+	};
+	typedef TArray<PrimitiveInfo, SceneRenderingAllocator> TPrimitiveInfos;
+
+	FHairStrandsVoxelNodeDesc VirtualVoxelNodeDesc;
+	FHairStrandsDeepShadowDatas DeepShadowDatas;
+	TPrimitiveInfos PrimitivesInfos;
+	FBoxSphereBounds Bounds;
+	FIntRect ScreenRect;
+	uint32 MacroGroupId;
+
+	bool bNeedScatterSceneLighting = false;
+};
+
+////////////////////////////////////////////////////////////////////////////////////
+// Debug data
+
+struct FHairStrandsDebugData
+{
+	BEGIN_SHADER_PARAMETER_STRUCT(FWriteParameters, )
+		SHADER_PARAMETER(uint32, Debug_MaxShadingPointCount)
+		SHADER_PARAMETER(uint32, Debug_MaxSampleCount)
+		SHADER_PARAMETER_RDG_BUFFER_UAV(RWStructuredBuffer, Debug_ShadingPointBuffer)
+		SHADER_PARAMETER_RDG_BUFFER_UAV(RWBuffer, Debug_ShadingPointCounter)
+		SHADER_PARAMETER_RDG_BUFFER_UAV(RWStructuredBuffer, Debug_SampleBuffer)
+		SHADER_PARAMETER_RDG_BUFFER_UAV(RWBuffer, Debug_SampleCounter)
+		END_SHADER_PARAMETER_STRUCT()
+
+		BEGIN_SHADER_PARAMETER_STRUCT(FReadParameters, )
+		SHADER_PARAMETER(uint32, Debug_MaxShadingPointCount)
+		SHADER_PARAMETER(uint32, Debug_MaxSampleCount)
+		SHADER_PARAMETER_RDG_BUFFER_SRV(StructuredBuffer, Debug_ShadingPointBuffer)
+		SHADER_PARAMETER_RDG_BUFFER_SRV(Buffer, Debug_ShadingPointCounter)
+		SHADER_PARAMETER_RDG_BUFFER_SRV(StructuredBuffer, Debug_SampleBuffer)
+		SHADER_PARAMETER_RDG_BUFFER_SRV(Buffer, Debug_SampleCounter)
+	END_SHADER_PARAMETER_STRUCT()
+
+	struct ShadingInfo
+	{
+		FVector BaseColor;
+		float	Roughness;
+		FVector T;
+		uint32	SampleCount;
+		FVector V;
+		float	SampleOffset;
+	};
+
+	struct Sample
+	{
+		FVector Direction;
+		float	Pdf;
+		FVector Weights;
+		float	Pad;
+	};
+	static const uint32 MaxShadingPointCount = 32;
+	static const uint32 MaxSampleCount = 1024 * 32;
+
+	struct Data
+	{
+		FRDGBufferRef ShadingPointBuffer = nullptr;
+		FRDGBufferRef ShadingPointCounter = nullptr;
+		FRDGBufferRef SampleBuffer = nullptr;
+		FRDGBufferRef SampleCounter = nullptr;
+	} Resources;
+
+	bool IsValid() const
+	{
+		return Resources.ShadingPointBuffer && Resources.ShadingPointCounter && Resources.SampleBuffer && Resources.SampleCounter;
+	}
+
+	static Data CreateData(FRDGBuilder& GraphBuilder);
+	static void SetParameters(FRDGBuilder& GraphBuilder, Data& In, FWriteParameters& Out);
+	static void SetParameters(FRDGBuilder& GraphBuilder, const Data& In, FReadParameters& Out);
+};
+
+typedef TArray<FHairStrandsMacroGroupData, SceneRenderingAllocator> FHairStrandsMacroGroupDatas;
+
+////////////////////////////////////////////////////////////////////////////////////
+// View Data
 struct FHairStrandsViewData
 {
-	//FRDGTextureRef CategorizationTexture = nullptr;
-	//FRDGTextureRef HairOnlyDepthTexture = nullptr;
 	TRDGUniformBufferRef<FHairStrandsViewUniformParameters> UniformBuffer;
 	bool bIsValid = false;
+
+	// Internal data
+	FHairStrandsVisibilityData VisibilityData;
+	FHairStrandsMacroGroupDatas MacroGroupDatas;
+	FHairStrandsDeepShadowResources DeepShadowResources;
+	FHairStrandsVoxelResources VirtualVoxelResources;
+	FHairStrandsMacroGroupResources MacroGroupResources;
+	FHairStrandsDebugData DebugData;
 };
 
 namespace HairStrands
@@ -28,4 +343,5 @@ namespace HairStrands
 	TRDGUniformBufferRef<FHairStrandsViewUniformParameters> CreateDefaultHairStrandsViewUniformBuffer(FRDGBuilder& GraphBuilder, FViewInfo& View);
 	TRDGUniformBufferRef<FHairStrandsViewUniformParameters> BindHairStrandsViewUniformParameters(const FViewInfo& View);
 	bool HasViewHairStrandsData(const FViewInfo& View);
+	bool HasViewHairStrandsData(const TArrayView<FViewInfo>& Views);
 }
\ No newline at end of file
diff --git a/Engine/Source/Runtime/Renderer/Private/HairStrands/HairStrandsDebug.cpp b/Engine/Source/Runtime/Renderer/Private/HairStrands/HairStrandsDebug.cpp
--- a/Engine/Source/Runtime/Renderer/Private/HairStrands/HairStrandsDebug.cpp
+++ b/Engine/Source/Runtime/Renderer/Private/HairStrands/HairStrandsDebug.cpp
@@ -79,16 +79,6 @@ static bool IsDebugDrawAndDebugPrintEnabled(const FViewInfo& View)
 	return ShaderDrawDebug::IsShaderDrawDebugEnabled() && ShaderPrint::IsEnabled() && ShaderPrint::IsSupported(View);
 }
 
-static bool IsDebugDrawAndDebugPrintEnabled(const TArray<FViewInfo>& Views)
-{
-	return Views.Num() > 0 && IsDebugDrawAndDebugPrintEnabled(Views[0]);
-}
-
-static bool IsDebugDrawAndDebugPrintEnabled(const TArrayView<FViewInfo>& Views)
-{
-	return Views.Num() > 0 && IsDebugDrawAndDebugPrintEnabled(Views[0]);
-}
-
 bool IsHairStrandsClusterDebugEnable()
 {
 	return GHairStrandsClusterDebug > 0;
@@ -293,7 +283,7 @@ static void AddDebugHairPrintPass(
 	const FViewInfo* View,
 	const EHairDebugMode InDebugMode,
 	const FHairStrandsVisibilityData& VisibilityData,
-	const FHairStrandsMacroGroupDatas& MacroGroupDatas,
+	const FHairStrandsMacroGroupResources& MacroGroupResources,
 	FRDGTextureSRVRef InDepthStencilTexture)
 {
 	if (!VisibilityData.CategorizationTexture || !VisibilityData.NodeIndex || !VisibilityData.NodeData || !InDepthStencilTexture) return;
@@ -318,9 +308,9 @@ static void AddDebugHairPrintPass(
 	Parameters->HairVisibilityNodeGroupSize = VisibilityData.NodeGroupSize;
 	Parameters->DepthStencilTexture = InDepthStencilTexture;
 	Parameters->LinearSampler = TStaticSamplerState<SF_Bilinear, AM_Clamp, AM_Clamp, AM_Clamp>::GetRHI();
-	Parameters->HairMacroGroupCount = MacroGroupDatas.Datas.Num();
+	Parameters->HairMacroGroupCount = MacroGroupResources.MacroGroupCount;
 	Parameters->MaxSampleCount = VisibilityData.MaxSampleCount;
-	Parameters->HairMacroGroupAABBBuffer = GraphBuilder.CreateSRV(MacroGroupDatas.MacroGroupResources.MacroGroupAABBsBuffer, PF_R32_SINT);
+	Parameters->HairMacroGroupAABBBuffer = GraphBuilder.CreateSRV(MacroGroupResources.MacroGroupAABBsBuffer, PF_R32_SINT);
 	ShaderPrint::SetParameters(GraphBuilder, *View, Parameters->ShaderPrintUniformBuffer);
 	ShaderDrawDebug::SetParameters(GraphBuilder, View->ShaderDrawData, Parameters->ShaderDrawUniformBuffer);
 	TShaderMapRef<FHairDebugPrintCS> ComputeShader(View->ShaderMap);
@@ -508,7 +498,7 @@ static void AddDebugDeepShadowTexturePass(
 	const FViewInfo* View,
 	const FIntRect& HairViewRect,
 	const FHairStrandsDeepShadowData* ShadowData,
-	const FDeepShadowResources* Resources,
+	const FHairStrandsDeepShadowResources* Resources,
 	FRDGTextureRef& OutTarget)
 {
 	check(OutTarget);
@@ -610,10 +600,11 @@ IMPLEMENT_GLOBAL_SHADER(FDeepShadowInfoCS, "/Engine/Private/HairStrands/HairStra
 static void AddDeepShadowInfoPass(
 	FRDGBuilder& GraphBuilder,
 	const FViewInfo& View,
-	const FHairStrandsMacroGroupDatas& MacroGroupDatas,
+	const FHairStrandsDeepShadowResources& DeepShadowResources,
+	const FHairStrandsMacroGroupResources& MacroGroupResources,
 	FRDGTextureRef& OutputTexture)
 {
-	if (MacroGroupDatas.DeepShadowResources.TotalAtlasSlotCount == 0)
+	if (DeepShadowResources.TotalAtlasSlotCount == 0)
 	{
 		return;
 	}
@@ -624,11 +615,11 @@ static void AddDeepShadowInfoPass(
 	FDeepShadowInfoCS::FParameters* Parameters = GraphBuilder.AllocParameters<FDeepShadowInfoCS::FParameters>();
 	Parameters->ViewUniformBuffer = View.ViewUniformBuffer;
 	Parameters->OutputResolution = Resolution;
-	Parameters->AllocatedSlotCount = MacroGroupDatas.DeepShadowResources.TotalAtlasSlotCount;
-	Parameters->MacroGroupCount = MacroGroupDatas.MacroGroupResources.MacroGroupCount;
+	Parameters->AllocatedSlotCount = DeepShadowResources.TotalAtlasSlotCount;
+	Parameters->MacroGroupCount = MacroGroupResources.MacroGroupCount;
 	Parameters->SceneTextures = SceneTextures;
-	Parameters->MacroGroupAABBBuffer = GraphBuilder.CreateSRV(MacroGroupDatas.MacroGroupResources.MacroGroupAABBsBuffer, PF_R32_SINT);
-	Parameters->ShadowWorldToLightTransformBuffer = GraphBuilder.CreateSRV(MacroGroupDatas.DeepShadowResources.DeepShadowWorldToLightTransforms);
+	Parameters->MacroGroupAABBBuffer = GraphBuilder.CreateSRV(MacroGroupResources.MacroGroupAABBsBuffer, PF_R32_SINT);
+	Parameters->ShadowWorldToLightTransformBuffer = GraphBuilder.CreateSRV(DeepShadowResources.DeepShadowWorldToLightTransforms);
 	ShaderDrawDebug::SetParameters(GraphBuilder, View.ShaderDrawData, Parameters->ShaderDrawParameters);
 	ShaderPrint::SetParameters(GraphBuilder, View, Parameters->ShaderPrintParameters);
 	Parameters->OutputTexture = GraphBuilder.CreateUAV(OutputTexture);
@@ -680,6 +671,7 @@ static void AddVoxelPageRaymarchingPass(
 	FRDGBuilder& GraphBuilder,
 	const FViewInfo& View,
 	const FHairStrandsMacroGroupDatas& MacroGroupDatas,
+	const FHairStrandsVoxelResources& VoxelResources,
 	FRDGTextureRef& OutputTexture)
 {
 	if (!IsDebugDrawAndDebugPrintEnabled(View))
@@ -689,9 +681,8 @@ static void AddVoxelPageRaymarchingPass(
 
 	FSceneTextureParameters SceneTextures = GetSceneTextureParameters(GraphBuilder);
 
-	const FVirtualVoxelResources& VoxelResources = MacroGroupDatas.VirtualVoxelResources;
 	const FIntPoint Resolution(OutputTexture->Desc.Extent);
-	for (const FHairStrandsMacroGroupData& MacroGroupData : MacroGroupDatas.Datas)
+	for (const FHairStrandsMacroGroupData& MacroGroupData : MacroGroupDatas)
 	{
 		FVoxelVirtualRaymarchingCS::FParameters* Parameters = GraphBuilder.AllocParameters<FVoxelVirtualRaymarchingCS::FParameters>();
 		Parameters->ViewUniformBuffer		= View.ViewUniformBuffer;
@@ -700,7 +691,7 @@ static void AddVoxelPageRaymarchingPass(
 		Parameters->bDrawPage				= FMath::Clamp(GHairVirtualVoxel_DrawDebugPage, 0, 2);
 		Parameters->ForcedMipLevel			= FMath::Clamp(GHairVirtualVoxel_ForceMipLevel, -1, 5);
 		Parameters->MacroGroupId			= MacroGroupData.MacroGroupId;
-		Parameters->MacroGroupCount			= MacroGroupDatas.Datas.Num();
+		Parameters->MacroGroupCount			= MacroGroupDatas.Num();
 		Parameters->MaxTotalPageIndexCount  = VoxelResources.Parameters.Common.PageIndexCount;
 		Parameters->VirtualVoxel			= VoxelResources.UniformBuffer;
 		Parameters->TotalValidPageCounter	= GraphBuilder.CreateSRV(VoxelResources.PageIndexGlobalCounter, PF_R32_UINT);
@@ -988,11 +979,10 @@ bool IsHairStrandsClusterCullingEnable();
 
 static void AddDrawDebugClusterPass(
 	FRDGBuilder& GraphBuilder,
-	const FHairStrandsMacroGroupViews& InMacroGroupViews,
 	const FHairStrandClusterData& HairClusterData,
-	const TArrayView<FViewInfo>& Views)
+	const FViewInfo& View)
 {
-	if (!IsDebugDrawAndDebugPrintEnabled(Views))
+	if (!IsDebugDrawAndDebugPrintEnabled(View))
 	{
 		return;
 	}
@@ -1003,12 +993,10 @@ static void AddDrawDebugClusterPass(
 	{
 		return;
 	}
-
-	for (int32 ViewIndex = 0; ViewIndex < Views.Num(); ++ViewIndex)
+	
 	{
-		if (ViewIndex < InMacroGroupViews.Views.Num())
 		{
-			for (const FHairStrandsMacroGroupData& MacroGroupData : InMacroGroupViews.Views[ViewIndex].Datas)
+			for (const FHairStrandsMacroGroupData& MacroGroupData : View.HairStrandsViewData.MacroGroupDatas)
 			{
 				const bool bDebugAABB = IsHairStrandsClusterDebugAABBEnable();
 
@@ -1025,18 +1013,17 @@ static void AddDrawDebugClusterPass(
 						if (HairGroupPublicData != HairGroupClusters.HairGroupPublicPtr)
 							continue;
 
-						FViewInfo& ViewInfo = Views[ViewIndex];
-						if (ShaderDrawDebug::IsShaderDrawDebugEnabled(ViewInfo) && HairGroupClusters.CulledDispatchIndirectParametersClusterCount)
+						if (ShaderDrawDebug::IsShaderDrawDebugEnabled(View) && HairGroupClusters.CulledDispatchIndirectParametersClusterCount)
 						{
 							FRDGBufferRef CulledDispatchIndirectParametersClusterCount = GraphBuilder.RegisterExternalBuffer(HairGroupClusters.CulledDispatchIndirectParametersClusterCount);
 							FRDGExternalBuffer& DrawIndirectBuffer = HairGroupClusters.HairGroupPublicPtr->GetDrawIndirectBuffer();
 
 							FDrawDebugClusterAABBCS::FPermutationDomain Permutation;
 							Permutation.Set<FDrawDebugClusterAABBCS::FDebugAABBBuffer>(bDebugAABB ? 1 : 0);
-							TShaderMapRef<FDrawDebugClusterAABBCS> ComputeShader(ViewInfo.ShaderMap, Permutation);
+							TShaderMapRef<FDrawDebugClusterAABBCS> ComputeShader(View.ShaderMap, Permutation);
 
 							FDrawDebugClusterAABBCS::FParameters* Parameters = GraphBuilder.AllocParameters<FDrawDebugClusterAABBCS::FParameters>();
-							Parameters->ViewUniformBuffer = ViewInfo.ViewUniformBuffer;
+							Parameters->ViewUniformBuffer = View.ViewUniformBuffer;
 							Parameters->ClusterCount = HairGroupClusters.ClusterCount;
 							Parameters->TriangleCount = HairGroupClusters.VertexCount * 2; // VertexCount is actually the number of control points
 							Parameters->HairGroupId = DataIndex;
@@ -1051,8 +1038,8 @@ static void AddDrawDebugClusterPass(
 								FRDGBufferRef ClusterDebugInfoBuffer = GraphBuilder.RegisterExternalBuffer(HairGroupClusters.ClusterDebugInfoBuffer);
 								Parameters->ClusterDebugInfoBuffer = GraphBuilder.CreateSRV(ClusterDebugInfoBuffer);
 							}
-							ShaderDrawDebug::SetParameters(GraphBuilder, ViewInfo.ShaderDrawData, Parameters->ShaderDrawParameters);
-							ShaderPrint::SetParameters(GraphBuilder, ViewInfo, Parameters->ShaderPrintParameters);
+							ShaderDrawDebug::SetParameters(GraphBuilder, View.ShaderDrawData, Parameters->ShaderDrawParameters);
+							ShaderPrint::SetParameters(GraphBuilder, View, Parameters->ShaderPrintParameters);
 
 							check(Parameters->ClusterCount / 64 <= 65535);
 							const FIntVector DispatchCount = DispatchCount.DivideAndRoundUp(FIntVector(Parameters->ClusterCount, 1, 1), FIntVector(64, 1, 1));// FIX ME, this could get over 65535
@@ -1069,37 +1056,31 @@ static void AddDrawDebugClusterPass(
 
 ///////////////////////////////////////////////////////////////////////////////////////////////////
 
-void RenderHairStrandsDebugInfo(
+static void InternalRenderHairStrandsDebugInfo(
 	FRDGBuilder& GraphBuilder,
-	TArrayView<FViewInfo> Views,
-	const FHairStrandsRenderingData* HairDatas,
+	FViewInfo& View,
 	const struct FHairStrandClusterData& HairClusterData,
 	FRDGTextureRef SceneColorTexture)
 {
+	check(HairStrands::HasViewHairStrandsData(View));
+
 	const float YStep = 14;
 	const float ColumnWidth = 200;
 
-	if (Views.Num() == 0)
-		return;
-
 	RDG_EVENT_SCOPE(GraphBuilder, "HairStrandsDebug");
 
 	if (GHairStrandsUpdateCullIndex)
 	{
-		const FViewInfo& View = Views[0];
 		const float TotalPixelCount = View.ViewRect.Width() * View.ViewRect.Height();
 		const float Index = View.CursorPos.X + View.CursorPos.Y * View.ViewRect.Width();
 		GHairStrandsCullNormalizedIndex = Index / TotalPixelCount;
 	}
 
 	// Only render debug information for the main view
-	const uint32 ViewIndex = 0;
-	FViewInfo& View = Views[ViewIndex];
-	const FSceneViewFamily& ViewFamily = *(View.Family);
 	const FSceneTextures& SceneTextures = FSceneTextures::Get(GraphBuilder);
 
 	{
-		FHairStrandsBookmarkParameters Params = CreateHairStrandsBookmarkParameters(Views[0]);
+		FHairStrandsBookmarkParameters Params = CreateHairStrandsBookmarkParameters(View);
 		Params.SceneColorTexture = SceneColorTexture;
 		RunHairStrandsBookmark(GraphBuilder, EHairStrandsBookmark::ProcessDebug, Params);
 	}
@@ -1110,39 +1091,34 @@ void RenderHairStrandsDebugInfo(
 	const EHairStrandsDebugMode StrandsDebugMode = GetHairStrandsDebugStrandsMode();
 	const EHairDebugMode HairDebugMode = GetHairStrandsDebugMode();
 
-	if (HairDatas && (GHairStrandsDebugPlotBsdf > 0 || HairDatas->DebugData.IsValid()))
+	
+	const FHairStrandsViewData& HairData = View.HairStrandsViewData;
+
+	if (GHairStrandsDebugPlotBsdf > 0 || HairData.DebugData.IsValid())
 	{
 		if (GHairStrandsDebugPlotBsdf > 0)
 		{
 			AddPlotBSDFPass(GraphBuilder, View, SceneColorTexture);
 		}
-		if (HairDatas->DebugData.IsValid())
+		if (HairData.DebugData.IsValid())
 		{
-			AddPlotSamplePass(GraphBuilder, View, HairDatas->DebugData.Resources, SceneColorTexture);
+			AddPlotSamplePass(GraphBuilder, View, HairData.DebugData.Resources, SceneColorTexture);
 		}	
 	}
 
 	float ClusterY = 38;
-	if (!HairDatas)
-		return;
-
-	const FHairStrandsMacroGroupViews& InMacroGroupViews = HairDatas->MacroGroupsPerViews;
 
 	if (HairDebugMode == EHairDebugMode::MacroGroups)
 	{
-		if (ViewIndex < uint32(HairDatas->HairVisibilityViews.HairDatas.Num()))
-		{
-			const FHairStrandsVisibilityData& VisibilityData = HairDatas->HairVisibilityViews.HairDatas[ViewIndex];
-			const FHairStrandsMacroGroupDatas& MacroGroupDatas = InMacroGroupViews.Views[ViewIndex];
-			AddDebugHairPrintPass(GraphBuilder, &View, HairDebugMode, VisibilityData, MacroGroupDatas, SceneTextures.Stencil);
+		{		
+			AddDebugHairPrintPass(GraphBuilder, &View, HairDebugMode, HairData.VisibilityData, HairData.MacroGroupResources, SceneTextures.Stencil);
 		}
 
 		// CPU bound of macro groups
 		FViewElementPDI ShadowFrustumPDI(&View, nullptr, nullptr);
-		const FHairStrandsMacroGroupDatas& MacroGroupDatas = InMacroGroupViews.Views[ViewIndex];
-		if (MacroGroupDatas.VirtualVoxelResources.IsValid())
+		if (HairData.VirtualVoxelResources.IsValid())
 		{
-			for (const FHairStrandsMacroGroupData& MacroGroupData : MacroGroupDatas.Datas)
+			for (const FHairStrandsMacroGroupData& MacroGroupData : HairData.MacroGroupDatas)
 			{
 				const FBox Bound(MacroGroupData.VirtualVoxelNodeDesc.WorldMinAABB, MacroGroupData.VirtualVoxelNodeDesc.WorldMaxAABB);
 				DrawWireBox(&ShadowFrustumPDI, Bound, FColor::Red, 0);
@@ -1160,9 +1136,9 @@ void RenderHairStrandsDebugInfo(
 			Line = FString::Printf(TEXT("----------------------------------------------------------------"));
 			Canvas.DrawShadowedString(X, Y += YStep, *Line, GetStatsFont(), DebugColor);
 
-			Line = FString::Printf(TEXT("Macro group count : %d"), MacroGroupDatas.Datas.Num());
+			Line = FString::Printf(TEXT("Macro group count : %d"), HairData.MacroGroupDatas.Num());
 			Canvas.DrawShadowedString(X, Y += YStep, *Line, GetStatsFont(), DebugColor);
-			for (const FHairStrandsMacroGroupData& MacroGroupData : MacroGroupDatas.Datas)
+			for (const FHairStrandsMacroGroupData& MacroGroupData : HairData.MacroGroupDatas)
 			{
 				Line = FString::Printf(TEXT(" %d - Bound Radus: %f.2m (%dx%d)"), MacroGroupData.MacroGroupId, MacroGroupData.Bounds.GetSphere().W);
 				Canvas.DrawShadowedString(X, Y += YStep, *Line, GetStatsFont(), DebugColor);
@@ -1173,23 +1149,21 @@ void RenderHairStrandsDebugInfo(
 
 	if (HairDebugMode == EHairDebugMode::DeepOpacityMaps)
 	{
-		if (ViewIndex < uint32(InMacroGroupViews.Views.Num()))
 		{
-			const FHairStrandsMacroGroupDatas& MacroGroupDatas = InMacroGroupViews.Views[ViewIndex];
-			for (const FHairStrandsMacroGroupData& MacroGroup : MacroGroupDatas.Datas)
+			for (const FHairStrandsMacroGroupData& MacroGroup : HairData.MacroGroupDatas)
 			{
-				if (!MacroGroupDatas.DeepShadowResources.DepthAtlasTexture || !MacroGroupDatas.DeepShadowResources.LayersAtlasTexture)
+				if (!HairData.DeepShadowResources.DepthAtlasTexture || !HairData.DeepShadowResources.LayersAtlasTexture)
 				{
 					continue;
 				}
 
-				for (const FHairStrandsDeepShadowData& DeepShadowData : MacroGroup.DeepShadowDatas.Datas)
+				for (const FHairStrandsDeepShadowData& DeepShadowData : MacroGroup.DeepShadowDatas)
 				{
 					const uint32 DomIndex = GDeepShadowDebugIndex;
 					if (DeepShadowData.AtlasSlotIndex != DomIndex)
 						continue;
 
-					AddDebugDeepShadowTexturePass(GraphBuilder, &View, FIntRect(), &DeepShadowData, &MacroGroupDatas.DeepShadowResources, SceneColorTexture);
+					AddDebugDeepShadowTexturePass(GraphBuilder, &View, FIntRect(), &DeepShadowData, &HairData.DeepShadowResources, SceneColorTexture);
 				}
 			}
 		}
@@ -1198,15 +1172,13 @@ void RenderHairStrandsDebugInfo(
 	// View Rect
 	if (IsHairStrandsViewRectOptimEnable() && HairDebugMode == EHairDebugMode::MacroGroupScreenRect)
 	{
-		if (ViewIndex < uint32(InMacroGroupViews.Views.Num()))
 		{
-			const FHairStrandsMacroGroupDatas& MacroGroupDatas = InMacroGroupViews.Views[ViewIndex];
-			for (const FHairStrandsMacroGroupData& MacroGroupData : MacroGroupDatas.Datas)
+			for (const FHairStrandsMacroGroupData& MacroGroupData : HairData.MacroGroupDatas)
 			{
 				AddDebugDeepShadowTexturePass(GraphBuilder, &View, MacroGroupData.ScreenRect, nullptr, nullptr, SceneColorTexture);
 			}
 
-			const FIntRect TotalRect = ComputeVisibleHairStrandsMacroGroupsRect(View.ViewRect, MacroGroupDatas);
+			const FIntRect TotalRect = ComputeVisibleHairStrandsMacroGroupsRect(View.ViewRect, HairData.MacroGroupDatas);
 			AddDebugDeepShadowTexturePass(GraphBuilder, &View, TotalRect, nullptr, nullptr, SceneColorTexture);
 		}
 	}
@@ -1215,14 +1187,12 @@ void RenderHairStrandsDebugInfo(
 
 	// Render Frustum for all lights & macro groups 
 	{
-		if ((HairDebugMode == EHairDebugMode::LightBounds || HairDebugMode == EHairDebugMode::DeepOpacityMaps) && ViewIndex < uint32(InMacroGroupViews.Views.Num()))
+		if ((HairDebugMode == EHairDebugMode::LightBounds || HairDebugMode == EHairDebugMode::DeepOpacityMaps))
 		{
-			if (ViewIndex < uint32(InMacroGroupViews.Views.Num()))
 			{
-				const FHairStrandsMacroGroupDatas& MacroGroupDatas = InMacroGroupViews.Views[ViewIndex];
-				if (MacroGroupDatas.DeepShadowResources.bIsGPUDriven)
+				if (HairData.DeepShadowResources.bIsGPUDriven)
 				{
-					AddDeepShadowInfoPass(GraphBuilder, View, MacroGroupDatas, SceneColorTexture);
+					AddDeepShadowInfoPass(GraphBuilder, View, HairData.DeepShadowResources, HairData.MacroGroupResources, SceneColorTexture);
 				}
 			}
 		}
@@ -1230,13 +1200,13 @@ void RenderHairStrandsDebugInfo(
 		FViewElementPDI ShadowFrustumPDI(&View, nullptr, nullptr);
 
 		// All DOMs
-		if (HairDebugMode == EHairDebugMode::LightBounds && ViewIndex < uint32(InMacroGroupViews.Views.Num()))
+		if (HairDebugMode == EHairDebugMode::LightBounds)
 		{
-			if (!InMacroGroupViews.Views[ViewIndex].DeepShadowResources.bIsGPUDriven)
+			if (!HairData.DeepShadowResources.bIsGPUDriven)
 			{
-				for (const FHairStrandsMacroGroupData& MacroGroupData : InMacroGroupViews.Views[ViewIndex].Datas)
+				for (const FHairStrandsMacroGroupData& MacroGroupData : HairData.MacroGroupDatas)
 				{
-					for (const FHairStrandsDeepShadowData& DomData : MacroGroupData.DeepShadowDatas.Datas)
+					for (const FHairStrandsDeepShadowData& DomData : MacroGroupData.DeepShadowDatas)
 					{
 						DrawFrustumWireframe(&ShadowFrustumPDI, DomData.CPU_WorldToLightTransform.Inverse(), FColor::Emerald, 0);
 						DrawWireBox(&ShadowFrustumPDI, DomData.Bounds.GetBox(), FColor::Yellow, 0);
@@ -1246,14 +1216,14 @@ void RenderHairStrandsDebugInfo(
 		}
 
 		// Current DOM
-		if (HairDebugMode == EHairDebugMode::DeepOpacityMaps && ViewIndex < uint32(InMacroGroupViews.Views.Num()))
+		if (HairDebugMode == EHairDebugMode::DeepOpacityMaps)
 		{
-			if (!InMacroGroupViews.Views[ViewIndex].DeepShadowResources.bIsGPUDriven)
+			if (!HairData.DeepShadowResources.bIsGPUDriven)
 			{
 				const int32 CurrentIndex = FMath::Max(0, GDeepShadowDebugIndex);
-				for (const FHairStrandsMacroGroupData& MacroGroupData : InMacroGroupViews.Views[ViewIndex].Datas)
+				for (const FHairStrandsMacroGroupData& MacroGroupData : HairData.MacroGroupDatas)
 				{
-					for (const FHairStrandsDeepShadowData& DomData : MacroGroupData.DeepShadowDatas.Datas)
+					for (const FHairStrandsDeepShadowData& DomData : MacroGroupData.DeepShadowDatas)
 					{
 						if (DomData.AtlasSlotIndex == CurrentIndex)
 						{
@@ -1266,12 +1236,11 @@ void RenderHairStrandsDebugInfo(
 		}
 
 		// Voxelization
-		if (bIsVoxelMode && ViewIndex < uint32(InMacroGroupViews.Views.Num()))
+		if (bIsVoxelMode)
 		{
-			const FHairStrandsMacroGroupDatas& MacroGroupDatas = InMacroGroupViews.Views[ViewIndex];
-			if (MacroGroupDatas.VirtualVoxelResources.IsValid())
+			if (HairData.VirtualVoxelResources.IsValid())
 			{
-				for (const FHairStrandsMacroGroupData& MacroGroupData : MacroGroupDatas.Datas)
+				for (const FHairStrandsMacroGroupData& MacroGroupData : HairData.MacroGroupDatas)
 				{
 					const FBox Bound(MacroGroupData.VirtualVoxelNodeDesc.WorldMinAABB, MacroGroupData.VirtualVoxelNodeDesc.WorldMaxAABB);
 					DrawWireBox(&ShadowFrustumPDI, Bound, FColor::Red, 0);
@@ -1294,30 +1263,25 @@ void RenderHairStrandsDebugInfo(
 		HairDebugMode == EHairDebugMode::Tile;
 	if (bRunDebugPass)
 	{
-		if (ViewIndex < uint32(HairDatas->HairVisibilityViews.HairDatas.Num()))
 		{
-			const FHairStrandsVisibilityData& VisibilityData = HairDatas->HairVisibilityViews.HairDatas[ViewIndex];
-			const FHairStrandsMacroGroupDatas& MacroGroupDatas = InMacroGroupViews.Views[ViewIndex];
+			const FHairStrandsVisibilityData& VisibilityData = HairData.VisibilityData;
 			AddDebugHairPass(GraphBuilder, &View, HairDebugMode, VisibilityData, SceneTextures.Stencil, SceneColorTexture);
-			AddDebugHairPrintPass(GraphBuilder, &View, HairDebugMode, VisibilityData, MacroGroupDatas, SceneTextures.Stencil);
+			AddDebugHairPrintPass(GraphBuilder, &View, HairDebugMode, VisibilityData, HairData.MacroGroupResources, SceneTextures.Stencil);
 		}
 	}
 
 	if (bIsVoxelMode)
 	{
-		if (ViewIndex < uint32(InMacroGroupViews.Views.Num()))
 		{
-			const FHairStrandsMacroGroupDatas& MacroGroupDatas = InMacroGroupViews.Views[ViewIndex];
-			if (MacroGroupDatas.VirtualVoxelResources.IsValid())
+			if (HairData.VirtualVoxelResources.IsValid())
 			{
-				AddVoxelPageRaymarchingPass(GraphBuilder, View, MacroGroupDatas, SceneColorTexture);
+				AddVoxelPageRaymarchingPass(GraphBuilder, View, HairData.MacroGroupDatas, HairData.VirtualVoxelResources, SceneColorTexture);
 			}
 		}
 	}
 
-	if (ViewIndex < uint32(HairDatas->HairVisibilityViews.HairDatas.Num()))
 	{
-		const FHairStrandsVisibilityData& VisibilityData = HairDatas->HairVisibilityViews.HairDatas[ViewIndex];
+		const FHairStrandsVisibilityData& VisibilityData = HairData.VisibilityData;
 		if (GHairStrandsDebugPPLL && VisibilityData.PPLLNodeCounterTexture) // Check if PPLL rendering is used and its debug view is enabled.
 		{
 			const FIntPoint PPLLResolution = VisibilityData.PPLLNodeIndexTexture->Desc.Extent;
@@ -1341,22 +1305,21 @@ void RenderHairStrandsDebugInfo(
 
 	if (GHairStrandsClusterDebug > 0)
 	{
-		AddDrawDebugClusterPass(GraphBuilder, InMacroGroupViews, HairClusterData, Views);
+		AddDrawDebugClusterPass(GraphBuilder, HairClusterData, View);
 	}
 
 	// Text
 	if (HairDebugMode == EHairDebugMode::LightBounds || HairDebugMode == EHairDebugMode::DeepOpacityMaps)
 	{
-		AddDrawCanvasPass(GraphBuilder, {}, View, SceneColor, [&InMacroGroupViews, YStep, ViewIndex, StrandsDebugMode, HairDebugMode](FCanvas& Canvas)
+		AddDrawCanvasPass(GraphBuilder, {}, View, SceneColor, [&HairData, YStep, StrandsDebugMode, HairDebugMode](FCanvas& Canvas)
 		{
-			const uint32 AtlasTotalSlotCount = FDeepShadowResources::MaxAtlasSlotCount;
+			const uint32 AtlasTotalSlotCount = FHairStrandsDeepShadowResources::MaxAtlasSlotCount;
 			uint32 AtlasAllocatedSlot = 0;
 			FIntPoint AtlasSlotResolution = FIntPoint(0, 0);
 			FIntPoint AtlasResolution = FIntPoint(0, 0);
 			bool bIsGPUDriven = false;
-			if (ViewIndex < uint32(InMacroGroupViews.Views.Num()))
 			{
-				const FDeepShadowResources& Resources = InMacroGroupViews.Views[ViewIndex].DeepShadowResources;
+				const FHairStrandsDeepShadowResources& Resources = HairData.DeepShadowResources;
 				AtlasResolution = Resources.DepthAtlasTexture ? Resources.DepthAtlasTexture->Desc.Extent : FIntPoint(0, 0);
 				AtlasAllocatedSlot = Resources.TotalAtlasSlotCount;
 				bIsGPUDriven = Resources.bIsGPUDriven;
@@ -1383,11 +1346,10 @@ void RenderHairStrandsDebugInfo(
 			Line = FString::Printf(TEXT("DOM Texture Index     : %d/%d"), DomTextureIndex, AtlasAllocatedSlot);						Canvas.DrawShadowedString(X, Y += YStep, *Line, GetStatsFont(), DebugColor);
 			Line = FString::Printf(TEXT("DOM GPU driven        : %s"), bIsGPUDriven ? TEXT("On") : TEXT("Off"));					Canvas.DrawShadowedString(X, Y += YStep, *Line, GetStatsFont(), DebugColor);
 
-			if (ViewIndex < uint32(InMacroGroupViews.Views.Num()))
 			{
-				for (const FHairStrandsMacroGroupData& MacroGroupData : InMacroGroupViews.Views[ViewIndex].Datas)
+				for (const FHairStrandsMacroGroupData& MacroGroupData : HairData.MacroGroupDatas)
 				{
-					for (const FHairStrandsDeepShadowData& DomData : MacroGroupData.DeepShadowDatas.Datas)
+					for (const FHairStrandsDeepShadowData& DomData : MacroGroupData.DeepShadowDatas)
 					{
 						Line = FString::Printf(TEXT(" %d - Bound Radus: %f.2m (%dx%d)"), DomData.AtlasSlotIndex, DomData.Bounds.GetSphere().W / 10.f, DomData.ShadowResolution.X, DomData.ShadowResolution.Y);
 						Canvas.DrawShadowedString(X, Y += YStep, *Line, GetStatsFont(), DebugColor);
@@ -1413,3 +1375,19 @@ void RenderHairStrandsDebugInfo(
 		});
 	}
 }
+
+void RenderHairStrandsDebugInfo(
+	FRDGBuilder& GraphBuilder,
+	TArrayView<FViewInfo> Views,
+	const struct FHairStrandClusterData& HairClusterData,
+	FRDGTextureRef SceneColorTexture)
+{
+	bool bHasHairData = false;
+	for (FViewInfo& View : Views)
+	{
+		if (HairStrands::HasViewHairStrandsData(View))
+		{
+			InternalRenderHairStrandsDebugInfo(GraphBuilder, View, HairClusterData, SceneColorTexture);
+		}
+	}
+}
\ No newline at end of file
diff --git a/Engine/Source/Runtime/Renderer/Private/HairStrands/HairStrandsDebug.h b/Engine/Source/Runtime/Renderer/Private/HairStrands/HairStrandsDebug.h
--- a/Engine/Source/Runtime/Renderer/Private/HairStrands/HairStrandsDebug.h
+++ b/Engine/Source/Runtime/Renderer/Private/HairStrands/HairStrandsDebug.h
@@ -9,67 +9,8 @@
 #include "RenderGraph.h"
 #include "Renderer/Private/SceneRendering.h"
 
-struct FHairStrandsDebugData
-{
-	BEGIN_SHADER_PARAMETER_STRUCT(FWriteParameters, )
-		SHADER_PARAMETER(uint32,							Debug_MaxShadingPointCount)
-		SHADER_PARAMETER(uint32,							Debug_MaxSampleCount)
-		SHADER_PARAMETER_RDG_BUFFER_UAV(RWStructuredBuffer, Debug_ShadingPointBuffer)
-		SHADER_PARAMETER_RDG_BUFFER_UAV(RWBuffer,			Debug_ShadingPointCounter)
-		SHADER_PARAMETER_RDG_BUFFER_UAV(RWStructuredBuffer, Debug_SampleBuffer)
-		SHADER_PARAMETER_RDG_BUFFER_UAV(RWBuffer,			Debug_SampleCounter)
-	END_SHADER_PARAMETER_STRUCT()
-
-	BEGIN_SHADER_PARAMETER_STRUCT(FReadParameters, )
-		SHADER_PARAMETER(uint32,							Debug_MaxShadingPointCount)
-		SHADER_PARAMETER(uint32,							Debug_MaxSampleCount)
-		SHADER_PARAMETER_RDG_BUFFER_SRV(StructuredBuffer,	Debug_ShadingPointBuffer)
-		SHADER_PARAMETER_RDG_BUFFER_SRV(Buffer,				Debug_ShadingPointCounter)
-		SHADER_PARAMETER_RDG_BUFFER_SRV(StructuredBuffer,	Debug_SampleBuffer)
-		SHADER_PARAMETER_RDG_BUFFER_SRV(Buffer,				Debug_SampleCounter)
-	END_SHADER_PARAMETER_STRUCT()
-
-	struct ShadingInfo
-	{
-		FVector BaseColor;
-		float	Roughness;
-		FVector T;
-		uint32	SampleCount;
-		FVector V;
-		float	SampleOffset;
-	};
-
-	struct Sample
-	{
-		FVector Direction;
-		float	Pdf;
-		FVector Weights;
-		float	Pad;
-	};
-	static const uint32 MaxShadingPointCount = 32;
-	static const uint32 MaxSampleCount = 1024 * 32;
-
-	struct Data
-	{
-		FRDGBufferRef ShadingPointBuffer = nullptr;
-		FRDGBufferRef ShadingPointCounter = nullptr;
-		FRDGBufferRef SampleBuffer = nullptr;
-		FRDGBufferRef SampleCounter = nullptr;
-	} Resources;
-
-	bool IsValid() const
-	{
-		return Resources.ShadingPointBuffer && Resources.ShadingPointCounter && Resources.SampleBuffer && Resources.SampleCounter;
-	}
-
-	static Data CreateData(FRDGBuilder& GraphBuilder);
-	static void SetParameters(FRDGBuilder& GraphBuilder, Data& In, FWriteParameters& Out);
-	static void SetParameters(FRDGBuilder& GraphBuilder, const Data& In, FReadParameters& Out);
-};
-
 void RenderHairStrandsDebugInfo(
 	FRDGBuilder& GraphBuilder,
 	TArrayView<FViewInfo> Views,
-	const struct FHairStrandsRenderingData* HairDatas,
 	const struct FHairStrandClusterData& HairClusterData,
 	FRDGTextureRef SceneColorTexture);
\ No newline at end of file
diff --git a/Engine/Source/Runtime/Renderer/Private/HairStrands/HairStrandsDeepShadow.cpp b/Engine/Source/Runtime/Renderer/Private/HairStrands/HairStrandsDeepShadow.cpp
--- a/Engine/Source/Runtime/Renderer/Private/HairStrands/HairStrandsDeepShadow.cpp
+++ b/Engine/Source/Runtime/Renderer/Private/HairStrands/HairStrandsDeepShadow.cpp
@@ -4,6 +4,7 @@
 #include "HairStrandsRasterCommon.h"
 #include "HairStrandsCluster.h"
 #include "HairStrandsUtils.h"
+#include "HairStrandsData.h"
 #include "LightSceneInfo.h"
 #include "ScenePrivate.h"
 
@@ -87,7 +88,7 @@ void AddInjectHairVoxelShadowCaster(
 	FIntRect AtlasRect,
 	uint32 AtlasSlotIndex,
 	FIntPoint AtlasSlotResolution,
-	FVirtualVoxelResources& VoxelResources,
+	FHairStrandsVoxelResources& VoxelResources,
 	FRDGBufferSRVRef DeepShadowViewInfoBufferSRV,
 	FRDGTextureRef OutDepthTexture)
 {
@@ -144,13 +145,11 @@ void AddInjectHairVoxelShadowCaster(
 ///////////////////////////////////////////////////////////////////////////////////////////////////
 
 typedef TArray<const FLightSceneInfo*, SceneRenderingAllocator> FLightSceneInfos;
-typedef TArray<FLightSceneInfos, SceneRenderingAllocator> FLightSceneInfosArray;
 
-static FLightSceneInfosArray GetVisibleDeepShadowLights(const FScene* Scene, const TArray<FViewInfo>& Views)
+static FLightSceneInfos GetVisibleDeepShadowLights(const FScene* Scene, const FViewInfo& View)
 {
-	// Collect all visible lights per view
-	FLightSceneInfosArray VisibleLightsPerView;
-	VisibleLightsPerView.SetNum(Views.Num());
+	// Collect all visible lights for the current view
+	FLightSceneInfos Out;
 	for (TSparseArray<FLightSceneInfoCompact>::TConstIterator LightIt(Scene->Lights); LightIt; ++LightIt)
 	{
 		const FLightSceneInfoCompact& LightSceneInfoCompact = *LightIt;
@@ -160,17 +159,16 @@ static FLightSceneInfosArray GetVisibleDeepShadowLights(const FScene* Scene, con
 			continue;
 
 		// Check if the light is visible in any of the views.
-		for (int32 ViewIndex = 0; ViewIndex < Views.Num(); ViewIndex++)
 		{
-			const bool bIsCompatible = LightSceneInfo->ShouldRenderLight(Views[ViewIndex]) && LightSceneInfo->Proxy->CastsHairStrandsDeepShadow();
+			const bool bIsCompatible = LightSceneInfo->ShouldRenderLight(View) && LightSceneInfo->Proxy->CastsHairStrandsDeepShadow();
 			if (!bIsCompatible)
 				continue;
 
-			VisibleLightsPerView[ViewIndex].Add(LightSceneInfo);
+			Out.Add(LightSceneInfo);
 		}
 	}
 
-	return VisibleLightsPerView;
+	return Out;
 }
 
 ///////////////////////////////////////////////////////////////////////////////////////////////////
@@ -223,52 +221,50 @@ FVector4 ComputeDeepShadowLayerDepths(float LayerDistribution);
 void RenderHairStrandsDeepShadows(
 	FRDGBuilder& GraphBuilder,
 	const FScene* Scene,
-	const TArray<FViewInfo>& Views,
-	FInstanceCullingManager& InstanceCullingManager,
-	FHairStrandsMacroGroupViews& MacroGroupsViews)
+	FViewInfo& View,
+	FInstanceCullingManager& InstanceCullingManager)
 {
 	QUICK_SCOPE_CYCLE_COUNTER(STAT_CLM_RenderDeepShadow);
 	DECLARE_GPU_STAT(HairStrandsDeepShadow);
 	RDG_EVENT_SCOPE(GraphBuilder, "HairStrandsDeepShadow");
 	RDG_GPU_STAT_SCOPE(GraphBuilder, HairStrandsDeepShadow);
 
-	FLightSceneInfosArray VisibleLightsPerView = GetVisibleDeepShadowLights(Scene, Views);
+	const FLightSceneInfos VisibleLights = GetVisibleDeepShadowLights(Scene, View);
+	FHairStrandsMacroGroupDatas& MacroGroupDatas = View.HairStrandsViewData.MacroGroupDatas;
+	FHairStrandsMacroGroupResources MacroGroupResources = View.HairStrandsViewData.MacroGroupResources;
+	FHairStrandsDeepShadowResources& DeepShadowResources = View.HairStrandsViewData.DeepShadowResources;
+	FHairStrandsVoxelResources VirtualVoxelResources = View.HairStrandsViewData.VirtualVoxelResources;
 
-	for (int32 ViewIndex = 0; ViewIndex < Views.Num(); ViewIndex++)
 	{
-		const FViewInfo& View = Views[ViewIndex];
-		if (!View.Family || ViewIndex >= MacroGroupsViews.Views.Num())
+		if (!View.Family)
 		{
-			continue;
+			return;
 		}
 
-		FHairStrandsMacroGroupDatas& MacroGroupDatas = MacroGroupsViews.Views[ViewIndex];
-		if (MacroGroupDatas.Datas.Num() == 0 || 
-			VisibleLightsPerView[ViewIndex].Num() == 0 ||
+		if (MacroGroupDatas.Num() == 0 || 
+			VisibleLights.Num() == 0 ||
 			IsHairStrandsForVoxelTransmittanceAndShadowEnable()) 
 		{
-			continue; 
+			return; 
 		}
 
 		// Compute the number of DOM which need to be created and insert default value
 		uint32 DOMSlotCount = 0;
-		for (int32 MacroGroupIt = 0; MacroGroupIt < MacroGroupDatas.Datas.Num(); ++MacroGroupIt)
-		{
-			const FHairStrandsMacroGroupData& MacroGroup = MacroGroupDatas.Datas[MacroGroupIt];
+		for (const FHairStrandsMacroGroupData& MacroGroup : MacroGroupDatas)
+		{			
 			const FBoxSphereBounds MacroGroupBounds = MacroGroup.Bounds;
-
-			for (const FLightSceneInfo* LightInfo : VisibleLightsPerView[ViewIndex])
+			for (const FLightSceneInfo* LightInfo : VisibleLights)
 			{
 				const FLightSceneProxy* LightProxy = LightInfo->Proxy;
 				if (!LightProxy->AffectsBounds(MacroGroupBounds))
 				{
-					continue;
+					return;
 				}
 
 				// Run out of atlas slot
-				if (DOMSlotCount >= FDeepShadowResources::MaxAtlasSlotCount)
+				if (DOMSlotCount >= FHairStrandsDeepShadowResources::MaxAtlasSlotCount)
 				{
-					continue;
+					return;
 				}
 
 				DOMSlotCount++;
@@ -276,15 +272,15 @@ void RenderHairStrandsDeepShadows(
 		}
 
 		if (DOMSlotCount == 0)
-			continue;
+			return;
 
 		const uint32 AtlasSlotX = FGenericPlatformMath::CeilToInt(FMath::Sqrt(static_cast<float>(DOMSlotCount)));
 		const FIntPoint AtlasSlotDimension(AtlasSlotX, AtlasSlotX == DOMSlotCount ? 1 : AtlasSlotX);
 		const FIntPoint AtlasSlotResolution(GDeepShadowResolution, GDeepShadowResolution);
 		const FIntPoint AtlasResolution(AtlasSlotResolution.X * AtlasSlotDimension.X, AtlasSlotResolution.Y * AtlasSlotDimension.Y);
 
-		FDeepShadowResources& Resources = MacroGroupDatas.DeepShadowResources;
-		Resources.TotalAtlasSlotCount = 0;
+		
+		DeepShadowResources.TotalAtlasSlotCount = 0;
 
 		// Create Atlas resources for DOM. It is shared for all lights, across all views
 		bool bClear = true;
@@ -294,12 +290,10 @@ void RenderHairStrandsDeepShadows(
 		// TODO add support for multiple view: need to deduplicate light which are visible accross several views
 		// Allocate atlas CPU slot
 		uint32 TotalAtlasSlotIndex = 0;
-		for (int32 MacroGroupIt = 0; MacroGroupIt < MacroGroupDatas.Datas.Num(); ++MacroGroupIt)
+		for (FHairStrandsMacroGroupData& MacroGroup : MacroGroupDatas)
 		{
-			FHairStrandsMacroGroupData& MacroGroup = MacroGroupDatas.Datas[MacroGroupIt];
-
 			// List of all the light in the scene.
-			for (const FLightSceneInfo* LightInfo : VisibleLightsPerView[ViewIndex])
+			for (const FLightSceneInfo* LightInfo : VisibleLights)
 			{
 				FBoxSphereBounds MacroGroupBounds = MacroGroup.Bounds;
 
@@ -309,7 +303,7 @@ void RenderHairStrandsDeepShadows(
 					continue;
 				}
 					
-				if (TotalAtlasSlotIndex >= FDeepShadowResources::MaxAtlasSlotCount)
+				if (TotalAtlasSlotIndex >= FHairStrandsDeepShadowResources::MaxAtlasSlotCount)
 				{
 					continue;
 				}
@@ -322,7 +316,7 @@ void RenderHairStrandsDeepShadows(
 					(TotalAtlasSlotIndex / AtlasSlotDimension.X) * AtlasSlotResolution.Y);
 
 				// Note: LightPosition.W is used in the transmittance mask shader to differentiate between directional and local lights.
-				FHairStrandsDeepShadowData& DomData = MacroGroup.DeepShadowDatas.Datas.AddZeroed_GetRef();
+				FHairStrandsDeepShadowData& DomData = MacroGroup.DeepShadowDatas.AddZeroed_GetRef();
 				ComputeWorldToLightClip(DomData.CPU_WorldToLightTransform, MinStrandRadiusAtDepth1, MacroGroupBounds, *LightProxy, LightType, AtlasSlotResolution);
 				DomData.LightDirection = LightProxy->GetDirection();
 				DomData.LightPosition = FVector4(LightProxy->GetPosition(), LightType == ELightComponentType::LightType_Directional ? 0 : 1);
@@ -343,25 +337,25 @@ void RenderHairStrandsDeepShadows(
 		// Sanity check
 		check(DOMSlotCount == TotalAtlasSlotIndex); 
 
-		Resources.TotalAtlasSlotCount = TotalAtlasSlotIndex;
-		Resources.AtlasSlotResolution = AtlasSlotResolution;
+		DeepShadowResources.TotalAtlasSlotCount = TotalAtlasSlotIndex;
+		DeepShadowResources.AtlasSlotResolution = AtlasSlotResolution;
 
 		FRDGBufferRef DeepShadowViewInfoBuffer = GraphBuilder.CreateBuffer(FRDGBufferDesc::CreateStructuredDesc(20 * sizeof(float), FMath::Max(1u, TotalAtlasSlotIndex)), TEXT("Hair.DeepShadowViewInfo"));
 		FRDGBufferRef DeepShadowWorldToLightBuffer = GraphBuilder.CreateBuffer(FRDGBufferDesc::CreateStructuredDesc(16 * sizeof(float), FMath::Max(1u, TotalAtlasSlotIndex)), TEXT("Hair.DeepShadowWorldToLightTransform"));
 		FRDGBufferSRVRef DeepShadowViewInfoBufferSRV = GraphBuilder.CreateSRV(DeepShadowViewInfoBuffer);
 
-		Resources.bIsGPUDriven = GDeepShadowGPUDriven > 0;
+		DeepShadowResources.bIsGPUDriven = GDeepShadowGPUDriven > 0;
 		{
-			check(TotalAtlasSlotIndex < FDeepShadowResources::MaxAtlasSlotCount);
+			check(TotalAtlasSlotIndex < FHairStrandsDeepShadowResources::MaxAtlasSlotCount);
 
 			// Allocate and create projection matrix and Min radius
 			// Stored FDeepShadowViewInfo structs
 			// See HairStrandsDeepShadowCommonStruct.ush for more details
 			FDeepShadowCreateViewInfoCS::FParameters* Parameters = GraphBuilder.AllocParameters<FDeepShadowCreateViewInfoCS::FParameters>();
 
-			for (FHairStrandsMacroGroupData& MacroGroup : MacroGroupDatas.Datas)
+			for (FHairStrandsMacroGroupData& MacroGroup : MacroGroupDatas)
 			{
-				for (FHairStrandsDeepShadowData& DomData : MacroGroup.DeepShadowDatas.Datas)
+				for (FHairStrandsDeepShadowData& DomData : MacroGroup.DeepShadowDatas)
 				{				
 					Parameters->LightDirections[DomData.AtlasSlotIndex]		= FVector4(DomData.LightDirection.X, DomData.LightDirection.Y, DomData.LightDirection.Z, 0);
 					Parameters->LightPositions[DomData.AtlasSlotIndex]		= FVector4(DomData.LightPosition.X, DomData.LightPosition.Y, DomData.LightPosition.Z, DomData.LightType == LightType_Directional ? 0 : 1);
@@ -369,10 +363,10 @@ void RenderHairStrandsDeepShadows(
 				}
 			}
 
-			Parameters->SlotResolution = Resources.AtlasSlotResolution;
-			Parameters->SlotIndexCount = Resources.TotalAtlasSlotCount;
-			Parameters->MacroGroupCount = MacroGroupDatas.Datas.Num();
-			Parameters->MacroGroupAABBBuffer = GraphBuilder.CreateSRV(MacroGroupDatas.MacroGroupResources.MacroGroupAABBsBuffer, PF_R32_SINT);
+			Parameters->SlotResolution = DeepShadowResources.AtlasSlotResolution;
+			Parameters->SlotIndexCount = DeepShadowResources.TotalAtlasSlotCount;
+			Parameters->MacroGroupCount = MacroGroupDatas.Num();
+			Parameters->MacroGroupAABBBuffer = GraphBuilder.CreateSRV(MacroGroupResources.MacroGroupAABBsBuffer, PF_R32_SINT);
 			Parameters->OutShadowViewInfoBuffer = GraphBuilder.CreateUAV(DeepShadowViewInfoBuffer);
 			Parameters->OutShadowWorldToLightTransformBuffer = GraphBuilder.CreateUAV(DeepShadowWorldToLightBuffer);
 
@@ -394,13 +388,13 @@ void RenderHairStrandsDeepShadows(
 		}
 
 		// Render deep shadows
-		for (FHairStrandsMacroGroupData& MacroGroup : MacroGroupDatas.Datas)
+		for (FHairStrandsMacroGroupData& MacroGroup : MacroGroupDatas)
 		{
-			for (FHairStrandsDeepShadowData& DomData : MacroGroup.DeepShadowDatas.Datas)
+			for (FHairStrandsDeepShadowData& DomData : MacroGroup.DeepShadowDatas)
 			{
 				const bool bIsOrtho = DomData.LightType == ELightComponentType::LightType_Directional;
 				const FVector4 HairRenderInfo = PackHairRenderInfo(DomData.CPU_MinStrandRadiusAtDepth1.Primary, DomData.CPU_MinStrandRadiusAtDepth1.Stable, DomData.CPU_MinStrandRadiusAtDepth1.Primary, 1);
-				const uint32 HairRenderInfoBits = PackHairRenderInfoBits(bIsOrtho, Resources.bIsGPUDriven);
+				const uint32 HairRenderInfoBits = PackHairRenderInfoBits(bIsOrtho, DeepShadowResources.bIsGPUDriven);
 					
 				const bool bDeepShadow = GDeepShadowInjectVoxelDepth == 0;
 				// Inject voxel result into the deep shadow
@@ -419,7 +413,7 @@ void RenderHairStrandsDeepShadows(
 						DomData.AtlasRect,
 						DomData.AtlasSlotIndex,
 						AtlasSlotResolution,
-						MacroGroupDatas.VirtualVoxelResources,
+						VirtualVoxelResources,
 						DeepShadowViewInfoBufferSRV,
 						FrontDepthAtlasTexture);
 
@@ -514,8 +508,8 @@ void RenderHairStrandsDeepShadows(
 				bClear = false;
 			}
 		}
-		Resources.DepthAtlasTexture = FrontDepthAtlasTexture;
-		Resources.LayersAtlasTexture = DeepShadowLayersAtlasTexture;
-		Resources.DeepShadowWorldToLightTransforms = DeepShadowWorldToLightBuffer;
+		DeepShadowResources.DepthAtlasTexture = FrontDepthAtlasTexture;
+		DeepShadowResources.LayersAtlasTexture = DeepShadowLayersAtlasTexture;
+		DeepShadowResources.DeepShadowWorldToLightTransforms = DeepShadowWorldToLightBuffer;
 	}
 }
diff --git a/Engine/Source/Runtime/Renderer/Private/HairStrands/HairStrandsDeepShadow.h b/Engine/Source/Runtime/Renderer/Private/HairStrands/HairStrandsDeepShadow.h
--- a/Engine/Source/Runtime/Renderer/Private/HairStrands/HairStrandsDeepShadow.h
+++ b/Engine/Source/Runtime/Renderer/Private/HairStrands/HairStrandsDeepShadow.h
@@ -13,52 +13,8 @@
 #include "SceneTypes.h"
 #include "Renderer/Private/SceneRendering.h"
 
-/// Hold deep shadow information for a given light.
-struct FHairStrandsDeepShadowData
-{
-	static const uint32 MaxMacroGroupCount = 16u;
-
-	FMatrix CPU_WorldToLightTransform;
-	FMinHairRadiusAtDepth1 CPU_MinStrandRadiusAtDepth1;
-	FIntRect AtlasRect;
-	uint32 MacroGroupId = ~0;
-	uint32 AtlasSlotIndex = 0;
-
-	FIntPoint ShadowResolution = FIntPoint::ZeroValue;
-	uint32 LightId = ~0;
-	ELightComponentType LightType = LightType_MAX;
-	FVector  LightDirection;
-	FVector4 LightPosition;
-	FLinearColor LightLuminance;
-	float LayerDistribution;
-
-	FBoxSphereBounds Bounds;
-};
-
-struct FDeepShadowResources
-{
-	// Limit the number of atlas slot to 32, in order to create the view info per slot in single compute
-	// This limitation can be alleviate, and is just here for convenience (see FDeepShadowCreateViewInfoCS)
-	static const uint32 MaxAtlasSlotCount = 32u;
-
-	uint32 TotalAtlasSlotCount = 0;
-	FIntPoint AtlasSlotResolution;
-	bool bIsGPUDriven = false;
-
-	FRDGTextureRef DepthAtlasTexture = nullptr;
-	FRDGTextureRef LayersAtlasTexture = nullptr;
-	FRDGBufferRef DeepShadowWorldToLightTransforms = nullptr;
-};
-
-/// Store all deep shadows infos for a given view
-struct FHairStrandsDeepShadowDatas
-{
-	TArray<FHairStrandsDeepShadowData, SceneRenderingAllocator> Datas;
-};
-
 void RenderHairStrandsDeepShadows(
 	FRDGBuilder& GraphBuilder,
 	const class FScene* Scene,
-	const TArray<FViewInfo>& Views,
-	FInstanceCullingManager& InstanceCullingManager,
-	struct FHairStrandsMacroGroupViews& MacroGroupsViews);
\ No newline at end of file
+	FViewInfo& View,
+	FInstanceCullingManager& InstanceCullingManager);
\ No newline at end of file
diff --git a/Engine/Source/Runtime/Renderer/Private/HairStrands/HairStrandsEnvironment.cpp b/Engine/Source/Runtime/Renderer/Private/HairStrands/HairStrandsEnvironment.cpp
--- a/Engine/Source/Runtime/Renderer/Private/HairStrands/HairStrandsEnvironment.cpp
+++ b/Engine/Source/Runtime/Renderer/Private/HairStrands/HairStrandsEnvironment.cpp
@@ -117,8 +117,8 @@ IMPLEMENT_GLOBAL_SHADER(FHairEnvironmentAO, "/Engine/Private/HairStrands/HairStr
 static void AddHairStrandsEnvironmentAOPass(
 	FRDGBuilder& GraphBuilder,
 	const FViewInfo& View,
-	const FHairStrandsVisibilityData& VisibilityData,
-	const FHairStrandsMacroGroupDatas& MacroGroupDatas,
+	const FHairStrandsVisibilityData& VisibilityResources,
+	const FHairStrandsVoxelResources& VoxelResources,
 	const FHairStrandsMacroGroupData& MacroGroupData,
 	FRDGTextureRef Output)
 {
@@ -133,10 +133,10 @@ static void AddHairStrandsEnvironmentAOPass(
 	PassParameters->Voxel_MacroGroupId = MacroGroupData.MacroGroupId;
 	PassParameters->Voxel_TanConeAngle = FMath::Tan(FMath::DegreesToRadians(GetHairStrandsSkyLightingConeAngle()));
 	PassParameters->SceneTextures = SceneTextures;
-	PassParameters->VirtualVoxel = MacroGroupDatas.VirtualVoxelResources.UniformBuffer;
+	PassParameters->VirtualVoxel = VoxelResources.UniformBuffer;
 
 	PassParameters->ViewUniformBuffer = View.ViewUniformBuffer;
-	PassParameters->HairCategorizationTexture = VisibilityData.CategorizationTexture;
+	PassParameters->HairCategorizationTexture = VisibilityResources.CategorizationTexture;
 	const FFinalPostProcessSettings& Settings = View.FinalPostProcessSettings;
 	PassParameters->AO_Power = Settings.AmbientOcclusionPower;
 	PassParameters->AO_Intensity = Settings.AmbientOcclusionIntensity;
@@ -315,13 +315,13 @@ static void AddHairStrandsEnvironmentLightingPassPS(
 	const FScene* Scene,
 	const FViewInfo& View,
 	const FHairStrandsVisibilityData& VisibilityData,
-	const FHairStrandsMacroGroupDatas& MacroGroupDatas,
+	const FHairStrandsVoxelResources& VirtualVoxelResources,
 	const FRDGTextureRef SceneColorTexture,
 	FHairStrandsDebugData::Data* DebugData)
 {
 	FSceneTextureParameters SceneTextures = GetSceneTextureParameters(GraphBuilder);
 
-	check(MacroGroupDatas.VirtualVoxelResources.IsValid());
+	check(VirtualVoxelResources.IsValid());
 
 	// Render the reflection environment with tiled deferred culling
 	const bool bHasBoxCaptures = (View.NumBoxReflectionCaptures > 0);
@@ -381,7 +381,7 @@ static void AddHairStrandsEnvironmentLightingPassPS(
 	PassParameters->PreIntegratedGFSampler = TStaticSamplerState<SF_Bilinear, AM_Clamp, AM_Clamp, AM_Clamp>::GetRHI();
 	PassParameters->HairCountTexture = VisibilityData.ViewHairCountTexture ? VisibilityData.ViewHairCountTexture : GSystemTextures.GetBlackDummy(GraphBuilder);
 	PassParameters->SceneTextures = SceneTextures;
-	PassParameters->VirtualVoxel = MacroGroupDatas.VirtualVoxelResources.UniformBuffer;
+	PassParameters->VirtualVoxel = VirtualVoxelResources.UniformBuffer;
 	PassParameters->ViewUniformBuffer = View.ViewUniformBuffer;
 	PassParameters->ReflectionCaptureData = View.ReflectionCaptureUniformBuffer;
 	{
@@ -454,30 +454,25 @@ void RenderHairStrandsSceneColorScattering(
 	FRDGBuilder& GraphBuilder,
 	FRDGTextureRef SceneColorTexture,
 	const FScene* Scene,
-	TArrayView<const FViewInfo> Views,
-	FHairStrandsRenderingData* HairDatas)
+	TArrayView<const FViewInfo> Views)
 {
-	if (Views.Num() == 0 || !HairDatas || GHairScatterSceneLighting <= 0)
+	if (Views.Num() == 0 || GHairScatterSceneLighting <= 0)
 		return;
 
-	for (uint32 ViewIndex = 0; ViewIndex < uint32(Views.Num()); ++ViewIndex)
+	for (const FViewInfo& View : Views)
 	{
-		check(ViewIndex < uint32(HairDatas->HairVisibilityViews.HairDatas.Num()));
-		const FHairStrandsVisibilityData& VisibilityData = HairDatas->HairVisibilityViews.HairDatas[ViewIndex];
-		const bool bRenderHairLighting = VisibilityData.NodeIndex && VisibilityData.NodeData && HairDatas->MacroGroupsPerViews.Views[ViewIndex].VirtualVoxelResources.IsValid();
-		if (!bRenderHairLighting)
+		if (!HairStrands::HasViewHairStrandsData(View))
 		{
 			continue;
 		}
 
-		if (ViewIndex > uint32(HairDatas->MacroGroupsPerViews.Views.Num()))
-			continue;
-
-		const FViewInfo& View = Views[ViewIndex];
-		const FHairStrandsMacroGroupDatas& MacroGroupDatas = HairDatas->MacroGroupsPerViews.Views[ViewIndex];
+		const FHairStrandsVisibilityData& VisibilityData = View.HairStrandsViewData.VisibilityData;
+		const FHairStrandsVoxelResources& VoxelResources = View.HairStrandsViewData.VirtualVoxelResources;
+		const bool bRenderHairLighting = VisibilityData.NodeIndex && VisibilityData.NodeData && VoxelResources.IsValid();
+		check(bRenderHairLighting);
 
 		bool bNeedScatterSceneLighting = false;
-		for (const FHairStrandsMacroGroupData& MacroGroupData : HairDatas->MacroGroupsPerViews.Views[ViewIndex].Datas)
+		for (const FHairStrandsMacroGroupData& MacroGroupData : View.HairStrandsViewData.MacroGroupDatas)
 		{
 			if (MacroGroupData.bNeedScatterSceneLighting)
 			{
@@ -488,7 +483,7 @@ void RenderHairStrandsSceneColorScattering(
 
 		if (bNeedScatterSceneLighting)
 		{
-			AddHairStrandsEnvironmentLightingPassPS(GraphBuilder, Scene, View, VisibilityData, MacroGroupDatas, SceneColorTexture, nullptr);
+			AddHairStrandsEnvironmentLightingPassPS(GraphBuilder, Scene, View, VisibilityData, VoxelResources, SceneColorTexture, nullptr);
 		}
 	}
 }
@@ -496,11 +491,9 @@ void RenderHairStrandsSceneColorScattering(
 void RenderHairStrandsEnvironmentLighting(
 	FRDGBuilder& GraphBuilder,
 	const FScene* Scene,
-	const uint32 ViewIndex,
-	TArrayView<const FViewInfo> Views,
-	FHairStrandsRenderingData* HairDatas)
+	FViewInfo& View)
 {
-	if (!GetHairStrandsSkyLightingEnable() || !HairDatas)
+	if (!GetHairStrandsSkyLightingEnable() || !HairStrands::HasViewHairStrandsData(View))
 		return;
 
 	struct FHairEnvDebugData
@@ -514,55 +507,38 @@ void RenderHairStrandsEnvironmentLighting(
 		FRDGBufferUAVRef ShadingPointCounterUAV = nullptr;
 	};
 
-
-	check(ViewIndex < uint32(Views.Num()));
-	check(ViewIndex < uint32(HairDatas->HairVisibilityViews.HairDatas.Num()));	
-	const FHairStrandsVisibilityData& VisibilityData = HairDatas->HairVisibilityViews.HairDatas[ViewIndex];
-	const bool bRenderHairLighting = VisibilityData.NodeIndex && VisibilityData.NodeData && HairDatas->MacroGroupsPerViews.Views[ViewIndex].VirtualVoxelResources.IsValid() && VisibilityData.CategorizationTexture;
+	const FHairStrandsVisibilityData& VisibilityData = View.HairStrandsViewData.VisibilityData;
+	const FHairStrandsVoxelResources& VoxelResources = View.HairStrandsViewData.VirtualVoxelResources;
+	const bool bRenderHairLighting = VisibilityData.NodeIndex && VisibilityData.NodeData && VoxelResources.IsValid() && VisibilityData.CategorizationTexture;
 	if (!bRenderHairLighting)
 	{
 		return;
 	}
 
 	const bool bDebugSamplingEnable = GHairStrandsSkyLighting_DebugSample > 0;
-	FHairStrandsDebugData::Data DebugData;
 	if (bDebugSamplingEnable)
 	{
-		HairDatas->DebugData.Resources = FHairStrandsDebugData::CreateData(GraphBuilder);
+		View.HairStrandsViewData.DebugData.Resources = FHairStrandsDebugData::CreateData(GraphBuilder);
 	}
 
-	const FViewInfo& View = Views[ViewIndex];
-	const FHairStrandsMacroGroupDatas& MacroGroupDatas = HairDatas->MacroGroupsPerViews.Views[ViewIndex];
-	AddHairStrandsEnvironmentLightingPassPS(GraphBuilder, Scene, View, VisibilityData, MacroGroupDatas, nullptr, bDebugSamplingEnable ? &HairDatas->DebugData.Resources : nullptr);
+	AddHairStrandsEnvironmentLightingPassPS(GraphBuilder, Scene, View, VisibilityData, VoxelResources, nullptr, bDebugSamplingEnable ? &View.HairStrandsViewData.DebugData.Resources : nullptr);
 }
 
 void RenderHairStrandsAmbientOcclusion(
 	FRDGBuilder& GraphBuilder,
-	TArrayView<const FViewInfo> Views,
-	const FHairStrandsRenderingData* HairDatas,
+	const FViewInfo& View,
 	const FRDGTextureRef& InAOTexture)
 {
-	if (!GetHairStrandsSkyAOEnable() || Views.Num() == 0 || !InAOTexture || !HairDatas)
+	if (!GetHairStrandsSkyAOEnable() || !HairStrands::HasViewHairStrandsData(View) || !InAOTexture)
 		return;
 
-	for (uint32 ViewIndex = 0; ViewIndex < uint32(Views.Num()); ++ViewIndex)
-	{
-		check(ViewIndex < uint32(HairDatas->HairVisibilityViews.HairDatas.Num()));
-		const FHairStrandsVisibilityData& VisibilityData = HairDatas->HairVisibilityViews.HairDatas[ViewIndex];
-		const bool bRenderHairLighting = VisibilityData.NodeIndex && VisibilityData.NodeData && HairDatas->MacroGroupsPerViews.Views[ViewIndex].VirtualVoxelResources.IsValid() && VisibilityData.CategorizationTexture;
-		if (!bRenderHairLighting)
-		{
-			continue;
-		}
-
-		if (ViewIndex > uint32(HairDatas->MacroGroupsPerViews.Views.Num()))
-			continue;
+	const FHairStrandsVisibilityData& VisibilityData = View.HairStrandsViewData.VisibilityData;
+	const FHairStrandsVoxelResources& VoxelResources = View.HairStrandsViewData.VirtualVoxelResources;
+	const bool bRenderHairLighting = VisibilityData.NodeIndex && VisibilityData.NodeData && VoxelResources.IsValid() && VisibilityData.CategorizationTexture;
+	check(bRenderHairLighting);
 
-		const FViewInfo& View = Views[ViewIndex];
-		const FHairStrandsMacroGroupDatas& MacroGroupDatas = HairDatas->MacroGroupsPerViews.Views[ViewIndex];
-		for (const FHairStrandsMacroGroupData& MacroGroupData : HairDatas->MacroGroupsPerViews.Views[ViewIndex].Datas)
-		{
-			AddHairStrandsEnvironmentAOPass(GraphBuilder, View, VisibilityData, MacroGroupDatas, MacroGroupData, InAOTexture);
-		}
+	for (const FHairStrandsMacroGroupData& MacroGroupData : View.HairStrandsViewData.MacroGroupDatas)
+	{
+		AddHairStrandsEnvironmentAOPass(GraphBuilder, View, VisibilityData, VoxelResources, MacroGroupData, InAOTexture);
 	}
 }
diff --git a/Engine/Source/Runtime/Renderer/Private/HairStrands/HairStrandsEnvironment.h b/Engine/Source/Runtime/Renderer/Private/HairStrands/HairStrandsEnvironment.h
--- a/Engine/Source/Runtime/Renderer/Private/HairStrands/HairStrandsEnvironment.h
+++ b/Engine/Source/Runtime/Renderer/Private/HairStrands/HairStrandsEnvironment.h
@@ -12,24 +12,19 @@
 
 class FScene;
 class FViewInfo;
-struct FHairStrandsRenderingData;
 
 void RenderHairStrandsAmbientOcclusion(
 	FRDGBuilder& GraphBuilder,
-	TArrayView<const FViewInfo> Views,
-	const FHairStrandsRenderingData* HairDatas,
+	const FViewInfo& View,
 	const FRDGTextureRef& InAOTexture);
 
 void RenderHairStrandsEnvironmentLighting(
 	FRDGBuilder& GraphBuilder,
 	const FScene* Scene,
-	const uint32 ViewIndex,
-	TArrayView<const FViewInfo> Views,
-	FHairStrandsRenderingData* HairDatas);
+	FViewInfo& View);
 
 void RenderHairStrandsSceneColorScattering(
 	FRDGBuilder& GraphBuilder,
 	FRDGTextureRef SceneColorTexture,
 	const FScene* Scene,
-	TArrayView<const FViewInfo> Views,
-	FHairStrandsRenderingData* HairDatas);
+	TArrayView<const FViewInfo> Views);
diff --git a/Engine/Source/Runtime/Renderer/Private/HairStrands/HairStrandsRasterCommon.h b/Engine/Source/Runtime/Renderer/Private/HairStrands/HairStrandsRasterCommon.h
--- a/Engine/Source/Runtime/Renderer/Private/HairStrands/HairStrandsRasterCommon.h
+++ b/Engine/Source/Runtime/Renderer/Private/HairStrands/HairStrandsRasterCommon.h
@@ -57,7 +57,7 @@ void AddHairDeepShadowRasterPass(
 // ////////////////////////////////////////////////////////////////
 // Voxelization raster pass
 BEGIN_GLOBAL_SHADER_PARAMETER_STRUCT(FHairVoxelizationRasterUniformParameters, )
-	SHADER_PARAMETER_STRUCT(FVirtualVoxelCommonParameters, VirtualVoxel)
+	SHADER_PARAMETER_STRUCT(FHairStrandsVoxelCommonParameters, VirtualVoxel)
 	SHADER_PARAMETER(FMatrix, WorldToClipMatrix)
 	SHADER_PARAMETER(FVector, VoxelMinAABB)
 	SHADER_PARAMETER(FVector, VoxelMaxAABB)
diff --git a/Engine/Source/Runtime/Renderer/Private/HairStrands/HairStrandsRendering.cpp b/Engine/Source/Runtime/Renderer/Private/HairStrands/HairStrandsRendering.cpp
--- a/Engine/Source/Runtime/Renderer/Private/HairStrands/HairStrandsRendering.cpp
+++ b/Engine/Source/Runtime/Renderer/Private/HairStrands/HairStrandsRendering.cpp
@@ -7,12 +7,37 @@
 
 static TRDGUniformBufferRef<FHairStrandsViewUniformParameters> InternalCreateHairStrandsViewUniformBuffer(
 	FRDGBuilder& GraphBuilder, 
-	FRDGTextureRef HairCategorizationTexture, 
-	FRDGTextureRef HairDepthOnlyTexture)
+	FHairStrandsVisibilityData* In)
 {
 	FHairStrandsViewUniformParameters* Parameters = GraphBuilder.AllocParameters<FHairStrandsViewUniformParameters>();
-	Parameters->HairCategorizationTexture = HairCategorizationTexture ? HairCategorizationTexture : GSystemTextures.GetBlackDummy(GraphBuilder);
-	Parameters->HairOnlyDepthTexture = HairDepthOnlyTexture ? HairDepthOnlyTexture : GSystemTextures.GetDepthDummy(GraphBuilder);
+	Parameters->HairDualScatteringRoughnessOverride = GetHairDualScatteringRoughnessOverride();
+	if (In && In->CategorizationTexture)
+	{
+		Parameters->HairCategorizationTexture = In->CategorizationTexture;
+		Parameters->HairOnlyDepthTexture = In->HairOnlyDepthTexture;
+		Parameters->HairSampleOffset = In->NodeIndex;
+		Parameters->HairSampleData = GraphBuilder.CreateSRV(In->NodeData);
+		Parameters->HairSampleCoords = GraphBuilder.CreateSRV(In->NodeCoord);
+		Parameters->HairSampleCount = In->NodeCount;
+		Parameters->HairSampleViewportResolution = In->SampleLightingViewportResolution;
+	}
+	else
+	{
+		FRDGBufferRef DummyBuffer = GraphBuilder.CreateBuffer(FRDGBufferDesc::CreateStructuredDesc(4,1), TEXT("Hair.DummyBuffer"));
+		FRDGBufferRef DummyNodeBuffer = GraphBuilder.CreateBuffer(FRDGBufferDesc::CreateStructuredDesc(20, 1), TEXT("Hair.DummyNodeBuffer"));
+		FRDGTextureRef BlackTexture = GSystemTextures.GetBlackDummy(GraphBuilder);
+		AddClearUAVPass(GraphBuilder, GraphBuilder.CreateUAV(DummyNodeBuffer), 0);
+		AddClearUAVPass(GraphBuilder, GraphBuilder.CreateUAV(DummyBuffer), 0);
+
+		Parameters->HairOnlyDepthTexture = GSystemTextures.GetDepthDummy(GraphBuilder);
+		Parameters->HairCategorizationTexture = BlackTexture;
+		Parameters->HairSampleCount = BlackTexture;
+		Parameters->HairSampleOffset = BlackTexture;
+		Parameters->HairSampleCoords = GraphBuilder.CreateSRV(DummyBuffer);
+		Parameters->HairSampleData	 = GraphBuilder.CreateSRV(DummyNodeBuffer);
+		Parameters->HairSampleViewportResolution = FIntPoint(0, 0);
+	}
+
 	return GraphBuilder.CreateUniformBuffer(Parameters);
 }
 
@@ -24,22 +49,24 @@ void RenderHairPrePass(
 	FRDGBuilder& GraphBuilder,
 	FScene* Scene,
 	TArray<FViewInfo>& Views,
-	FInstanceCullingManager& InstanceCullingManager,
-	FHairStrandsRenderingData& OutHairDatas)
+	FInstanceCullingManager& InstanceCullingManager)
 {
 	// #hair_todo: Add multi-view
-	const bool bIsViewCompatible = Views.Num() > 0 && IsHairStrandsEnabled(EHairStrandsShaderType::Strands, Views[0].GetShaderPlatform());
-	if (bIsViewCompatible)
+	for (FViewInfo& View : Views)
 	{
+		const bool bIsViewCompatible = IsHairStrandsEnabled(EHairStrandsShaderType::Strands, View.GetShaderPlatform());
+		if (!View.Family || !bIsViewCompatible)
+			continue;
+
 		const ERHIFeatureLevel::Type FeatureLevel = Scene->GetFeatureLevel();
 
 		//SCOPED_GPU_STAT(RHICmdList, HairRendering);
-		OutHairDatas.MacroGroupsPerViews = CreateHairStrandsMacroGroups(GraphBuilder, Scene, Views);
+		CreateHairStrandsMacroGroups(GraphBuilder, Scene, View);
 		AddServiceLocalQueuePass(GraphBuilder);
 
 		// Voxelization and Deep Opacity Maps
-		VoxelizeHairStrands(GraphBuilder, Scene, Views, InstanceCullingManager, OutHairDatas.MacroGroupsPerViews);
-		RenderHairStrandsDeepShadows(GraphBuilder, Scene, Views, InstanceCullingManager, OutHairDatas.MacroGroupsPerViews);
+		VoxelizeHairStrands(GraphBuilder, Scene, View, InstanceCullingManager);
+		RenderHairStrandsDeepShadows(GraphBuilder, Scene, View, InstanceCullingManager);
 
 		AddServiceLocalQueuePass(GraphBuilder);
 	}
@@ -50,41 +77,38 @@ void RenderHairBasePass(
 	FScene* Scene,
 	const FSceneTextures& SceneTextures,
 	TArray<FViewInfo>& Views,
-	FInstanceCullingManager& InstanceCullingManager,
-	FHairStrandsRenderingData& OutHairDatas)
+	FInstanceCullingManager& InstanceCullingManager)
 {
-	// #hair_todo: Add multi-view
-	const bool bIsViewCompatible = Views.Num() > 0 && IsHairStrandsEnabled(EHairStrandsShaderType::Strands, Views[0].GetShaderPlatform());
-	if (bIsViewCompatible)
+	for (FViewInfo& View : Views)
 	{
-		const ERHIFeatureLevel::Type FeatureLevel = Scene->GetFeatureLevel();
-
-		// Hair visibility pass
-		OutHairDatas.HairVisibilityViews = RenderHairStrandsVisibilityBuffer(
-			GraphBuilder, 
-			Scene, 
-			Views, 
-			SceneTextures.GBufferA,
-			SceneTextures.GBufferB,
-			SceneTextures.GBufferC,
-			SceneTextures.GBufferD,
-			SceneTextures.GBufferE,
-			SceneTextures.Color.Resolve,
-			SceneTextures.Depth.Resolve,
-			SceneTextures.Velocity,
-			InstanceCullingManager,
-			OutHairDatas.MacroGroupsPerViews);
-
-		// Create RDG uniform buffer for each view
-		for (int32 ViewIndex = 0; ViewIndex < Views.Num(); ViewIndex++)
+		const bool bIsViewCompatible = IsHairStrandsEnabled(EHairStrandsShaderType::Strands, View.GetShaderPlatform());
+		if (View.Family && bIsViewCompatible && View.HairStrandsViewData.MacroGroupDatas.Num() > 0)
 		{
-			const FViewInfo& View = Views[ViewIndex];
-			if (View.Family)
-			{
-				FHairStrandsVisibilityData VisibilityData = OutHairDatas.HairVisibilityViews.HairDatas[ViewIndex];
-				Views[ViewIndex].HairStrandsViewData.UniformBuffer = InternalCreateHairStrandsViewUniformBuffer(GraphBuilder, VisibilityData.CategorizationTexture, VisibilityData.HairOnlyDepthTexture);
-				Views[ViewIndex].HairStrandsViewData.bIsValid = true;
-			}
+			RenderHairStrandsVisibilityBuffer(
+				GraphBuilder, 
+				Scene, 
+				View, 
+				SceneTextures.GBufferA,
+				SceneTextures.GBufferB,
+				SceneTextures.GBufferC,
+				SceneTextures.GBufferD,
+				SceneTextures.GBufferE,
+				SceneTextures.Color.Resolve,
+				SceneTextures.Depth.Resolve,
+				SceneTextures.Velocity,
+				InstanceCullingManager);
+			
+		}
+		
+		if (View.HairStrandsViewData.VisibilityData.CategorizationTexture)
+		{			
+			View.HairStrandsViewData.UniformBuffer = InternalCreateHairStrandsViewUniformBuffer(GraphBuilder, &View.HairStrandsViewData.VisibilityData);
+			View.HairStrandsViewData.bIsValid = true;
+		}
+		else
+		{
+			View.HairStrandsViewData.UniformBuffer = InternalCreateHairStrandsViewUniformBuffer(GraphBuilder, nullptr);
+			View.HairStrandsViewData.bIsValid = false;
 		}
 	}
 }
@@ -94,7 +118,7 @@ namespace HairStrands
 
 TRDGUniformBufferRef<FHairStrandsViewUniformParameters> CreateDefaultHairStrandsViewUniformBuffer(FRDGBuilder& GraphBuilder, FViewInfo& View)
 {
-	return InternalCreateHairStrandsViewUniformBuffer(GraphBuilder, nullptr, nullptr);
+	return InternalCreateHairStrandsViewUniformBuffer(GraphBuilder, nullptr);
 }
 
 TRDGUniformBufferRef<FHairStrandsViewUniformParameters> BindHairStrandsViewUniformParameters(const FViewInfo& View)
@@ -107,4 +131,16 @@ bool HasViewHairStrandsData(const FViewInfo& View)
 	return View.HairStrandsViewData.bIsValid;
 }
 
+bool HasViewHairStrandsData(const TArrayView<FViewInfo>& Views)
+{
+	for (const FViewInfo& View : Views)
+	{
+		if (View.HairStrandsViewData.bIsValid)
+		{
+			return true;
+		}
+	}
+	return false;
+}
+
 }
\ No newline at end of file
diff --git a/Engine/Source/Runtime/Renderer/Private/HairStrands/HairStrandsRendering.h b/Engine/Source/Runtime/Renderer/Private/HairStrands/HairStrandsRendering.h
--- a/Engine/Source/Runtime/Renderer/Private/HairStrands/HairStrandsRendering.h
+++ b/Engine/Source/Runtime/Renderer/Private/HairStrands/HairStrandsRendering.h
@@ -21,29 +21,20 @@
 #include "HairStrandsComposition.h"
 #include "HairStrandsDebug.h"
 #include "HairStrandsInterface.h"
-
-/// Hold all the hair strands data
-struct FHairStrandsRenderingData
-{
-	FHairStrandsVisibilityViews HairVisibilityViews;
-	FHairStrandsMacroGroupViews MacroGroupsPerViews;
-	FHairStrandsDebugData DebugData;
-};
+#include "HairStrandsData.h"
 
 void RenderHairPrePass(
 	FRDGBuilder& GraphBuilder,
 	FScene* Scene,
 	TArray<FViewInfo>& Views,
-	FInstanceCullingManager& InstanceCullingManager,
-	FHairStrandsRenderingData& OutHairDatas);
+	FInstanceCullingManager& InstanceCullingManager);
 
 void RenderHairBasePass(
 	FRDGBuilder& GraphBuilder,
 	FScene* Scene,
 	const FSceneTextures& SceneTextures,
 	TArray<FViewInfo>& Views,
-	FInstanceCullingManager& InstanceCullingManager,
-	FHairStrandsRenderingData& OutHairDatas);
+	FInstanceCullingManager& InstanceCullingManager);
 
 void RunHairStrandsBookmark(
 	FRDGBuilder& GraphBuilder, 
diff --git a/Engine/Source/Runtime/Renderer/Private/HairStrands/HairStrandsScatter.cpp b/Engine/Source/Runtime/Renderer/Private/HairStrands/HairStrandsScatter.cpp
--- a/Engine/Source/Runtime/Renderer/Private/HairStrands/HairStrandsScatter.cpp
+++ b/Engine/Source/Runtime/Renderer/Private/HairStrands/HairStrandsScatter.cpp
@@ -146,7 +146,7 @@ IMPLEMENT_GLOBAL_SHADER(FHairScatterPS, "/Engine/Private/HairStrands/HairScatter
 static FRDGTextureRef AddScatterPass(
 	FRDGBuilder& GraphBuilder,
 	const FViewInfo& View,
-	const FVirtualVoxelResources& VoxelResources,
+	const FHairStrandsVoxelResources& VoxelResources,
 	const FRDGTextureRef& InVisibilityNodeIndex,
 	const FRDGBufferRef&  InVisibilityNodeData,
 	const FRDGTextureRef& InCategorizationTexture,
@@ -254,7 +254,7 @@ void AddHairDiffusionPass(
 	FRDGBuilder& GraphBuilder,
 	const FViewInfo& View,
 	const FHairStrandsVisibilityData& VisibilityData,
-	const FVirtualVoxelResources& VoxelResources,
+	const FHairStrandsVoxelResources& VoxelResources,
 	const FRDGTextureRef SceneColorDepth,
 	FRDGTextureRef OutLightSampleTexture)
 {
diff --git a/Engine/Source/Runtime/Renderer/Private/HairStrands/HairStrandsScatter.h b/Engine/Source/Runtime/Renderer/Private/HairStrands/HairStrandsScatter.h
--- a/Engine/Source/Runtime/Renderer/Private/HairStrands/HairStrandsScatter.h
+++ b/Engine/Source/Runtime/Renderer/Private/HairStrands/HairStrandsScatter.h
@@ -11,6 +11,6 @@ void AddHairDiffusionPass(
 	FRDGBuilder& GraphBuilder,
 	const FViewInfo& View,
 	const struct FHairStrandsVisibilityData& VisibilityData,
-	const struct FVirtualVoxelResources& VoxelResources,
+	const struct FHairStrandsVoxelResources& VoxelResources,
 	const FRDGTextureRef SceneColorDepth,
 	FRDGTextureRef OutSceneColorTexture);
\ No newline at end of file
diff --git a/Engine/Source/Runtime/Renderer/Private/HairStrands/HairStrandsTransmittance.cpp b/Engine/Source/Runtime/Renderer/Private/HairStrands/HairStrandsTransmittance.cpp
--- a/Engine/Source/Runtime/Renderer/Private/HairStrands/HairStrandsTransmittance.cpp
+++ b/Engine/Source/Runtime/Renderer/Private/HairStrands/HairStrandsTransmittance.cpp
@@ -56,9 +56,9 @@ static float GetDeepShadowDepthBiasScale() { return FMath::Max(0.0f, GDeepShadow
 
 static bool HasDeepShadowData(const FLightSceneInfo* LightSceneInfo, const FHairStrandsMacroGroupDatas& InDatas)
 {
-	for (const FHairStrandsMacroGroupData& MacroGroupData : InDatas.Datas)
+	for (const FHairStrandsMacroGroupData& MacroGroupData : InDatas)
 	{
-		for (const FHairStrandsDeepShadowData& DomData : MacroGroupData.DeepShadowDatas.Datas)
+		for (const FHairStrandsDeepShadowData& DomData : MacroGroupData.DeepShadowDatas)
 		{
 			if (DomData.LightId == LightSceneInfo->Id)
 				return true;
@@ -265,7 +265,7 @@ struct FDeepShadowTransmittanceParams
 	FRDGBufferRef HairVisibilityNodeData = nullptr;
 	FRDGBufferRef HairVisibilityNodeCoord = nullptr;
 
-	const FVirtualVoxelResources* VirtualVoxelResources = nullptr;
+	const FHairStrandsVoxelResources* VirtualVoxelResources = nullptr;
 };
 
 // Transmittance mask
@@ -431,7 +431,7 @@ struct FDeepShadowOpaqueParams
 	FVector			Voxel_LightDirection = FVector::ZeroVector;
 	FVector4		Voxel_LightPosition = FVector4(0, 0, 0, 0);
 	uint32			Voxel_MacroGroupId;
-	const FVirtualVoxelResources* Voxel_VirtualVoxel = nullptr;
+	const FHairStrandsVoxelResources* Voxel_VirtualVoxel = nullptr;
 };
 
 // Opaque mask
@@ -548,11 +548,13 @@ static FHairStrandsTransmittanceMaskData RenderHairStrandsTransmittanceMask(
 	FRDGBuilder& GraphBuilder,
 	const FViewInfo& View,
 	const FLightSceneInfo* LightSceneInfo,
-	const FHairStrandsMacroGroupDatas& MacroGroupDatas,
 	const FHairStrandsVisibilityData& VisibilityData,
+	const FHairStrandsMacroGroupDatas& MacroGroupDatas,
+	const FHairStrandsDeepShadowResources& DeepShadowResources,
+	const FHairStrandsVoxelResources& VoxelResources,
 	FRDGTextureRef ScreenShadowMaskSubPixelTexture)
 {
-	if (MacroGroupDatas.Datas.Num() == 0)
+	if (MacroGroupDatas.Num() == 0)
 		return FHairStrandsTransmittanceMaskData();
 
 	if (!HasDeepShadowData(LightSceneInfo, MacroGroupDatas) && !IsHairStrandsVoxelizationEnable())
@@ -583,20 +585,20 @@ static FHairStrandsTransmittanceMaskData RenderHairStrandsTransmittanceMask(
 	if (!IsHairStrandsForVoxelTransmittanceAndShadowEnable())
 	{
 		FRDGBufferSRVRef DeepShadow_WorldToLightTransformBufferSRV = nullptr;
-		for (const FHairStrandsMacroGroupData& MacroGroupData : MacroGroupDatas.Datas)
+		for (const FHairStrandsMacroGroupData& MacroGroupData : MacroGroupDatas)
 		{
-			for (const FHairStrandsDeepShadowData& DeepShadowData : MacroGroupData.DeepShadowDatas.Datas)
+			for (const FHairStrandsDeepShadowData& DeepShadowData : MacroGroupData.DeepShadowDatas)
 			{
 				if (DeepShadowData.LightId == LightSceneInfo->Id)
 				{
 					if (DeepShadow_WorldToLightTransformBufferSRV == nullptr)
 					{
-						DeepShadow_WorldToLightTransformBufferSRV = GraphBuilder.CreateSRV(MacroGroupDatas.DeepShadowResources.DeepShadowWorldToLightTransforms);
+						DeepShadow_WorldToLightTransformBufferSRV = GraphBuilder.CreateSRV(DeepShadowResources.DeepShadowWorldToLightTransforms);
 					}
 
 					bHasFoundLight = true;
-					Params.DeepShadow_FrontDepthTexture = MacroGroupDatas.DeepShadowResources.DepthAtlasTexture;
-					Params.DeepShadow_DomTexture = MacroGroupDatas.DeepShadowResources.LayersAtlasTexture;
+					Params.DeepShadow_FrontDepthTexture = DeepShadowResources.DepthAtlasTexture;
+					Params.DeepShadow_DomTexture = DeepShadowResources.LayersAtlasTexture;
 					Params.DeepShadow_Resolution = DeepShadowData.ShadowResolution;
 					Params.LightDirection = DeepShadowData.LightDirection;
 					Params.LightPosition = DeepShadowData.LightPosition;
@@ -606,7 +608,7 @@ static FHairStrandsTransmittanceMaskData RenderHairStrandsTransmittanceMask(
 					Params.DeepShadow_AtlasSlotOffsets_AtlasSlotIndex[DeepShadowData.MacroGroupId] = FIntVector4(DeepShadowData.AtlasRect.Min.X, DeepShadowData.AtlasRect.Min.Y, DeepShadowData.AtlasSlotIndex, 0);
 					Params.DeepShadow_CPUWorldToLightTransforms[DeepShadowData.MacroGroupId] = DeepShadowData.CPU_WorldToLightTransform;
 					Params.DeepShadow_WorldToLightTransformBuffer = DeepShadow_WorldToLightTransformBufferSRV;
-					Params.DeepShadow_bIsGPUDriven = MacroGroupDatas.DeepShadowResources.bIsGPUDriven;
+					Params.DeepShadow_bIsGPUDriven = DeepShadowResources.bIsGPUDriven;
 				}
 			}
 		}
@@ -628,7 +630,7 @@ static FHairStrandsTransmittanceMaskData RenderHairStrandsTransmittanceMask(
 		}
 	}
 
-	if (!bHasFoundLight && MacroGroupDatas.VirtualVoxelResources.IsValid())
+	if (!bHasFoundLight && VoxelResources.IsValid())
 	{
 		FLightShaderParameters LightParameters;
 		LightSceneInfo->Proxy->GetLightShaderParameters(LightParameters);
@@ -639,7 +641,7 @@ static FHairStrandsTransmittanceMaskData RenderHairStrandsTransmittanceMask(
 		Params.LightPosition = FVector4(LightSceneInfo->Proxy->GetPosition(), LightSceneInfo->Proxy->GetLightType() == ELightComponentType::LightType_Directional ? 0 : 1);
 		Params.LightChannelMask = LightSceneInfo->Proxy->GetLightingChannelMask();
 		Params.LightRadius = FMath::Max(LightParameters.SourceLength, LightParameters.SourceRadius);
-		Params.VirtualVoxelResources = &MacroGroupDatas.VirtualVoxelResources;
+		Params.VirtualVoxelResources = &VoxelResources;
 
 		OutShadowMask = AddDeepShadowTransmittanceMaskPass(
 			GraphBuilder,
@@ -661,22 +663,22 @@ static FHairStrandsTransmittanceMaskData RenderHairStrandsTransmittanceMask(
 
 FHairStrandsTransmittanceMaskData RenderHairStrandsTransmittanceMask(
 	FRDGBuilder& GraphBuilder,
-	const TArray<FViewInfo>& Views,
+	const FViewInfo& View,
 	const FLightSceneInfo* LightSceneInfo,
-	const FHairStrandsRenderingData* HairDatas,
 	FRDGTextureRef ScreenShadowMaskSubPixelTexture)
 {
 	FHairStrandsTransmittanceMaskData TransmittanceMaskData;
-	if (HairDatas)
+	if (HairStrands::HasViewHairStrandsData(View))
 	{
-		for (int32 ViewIndex = 0; ViewIndex < Views.Num(); ViewIndex++)
-		{
-			const FViewInfo& View = Views[ViewIndex];
-			const FHairStrandsVisibilityData& InHairVisibilityData = HairDatas->HairVisibilityViews.HairDatas[ViewIndex];
-			const FHairStrandsMacroGroupDatas& InMacroGroupDatas = HairDatas->MacroGroupsPerViews.Views[ViewIndex];
-
-			TransmittanceMaskData = RenderHairStrandsTransmittanceMask(GraphBuilder, View, LightSceneInfo, InMacroGroupDatas, InHairVisibilityData, ScreenShadowMaskSubPixelTexture);
-		}
+		TransmittanceMaskData = RenderHairStrandsTransmittanceMask(
+			GraphBuilder, 
+			View, 
+			LightSceneInfo, 
+			View.HairStrandsViewData.VisibilityData,
+			View.HairStrandsViewData.MacroGroupDatas,
+			View.HairStrandsViewData.DeepShadowResources,
+			View.HairStrandsViewData.VirtualVoxelResources,
+			ScreenShadowMaskSubPixelTexture);
 	}
 	return TransmittanceMaskData;
 }
@@ -689,9 +691,11 @@ static void RenderHairStrandsShadowMask(
 	const FLightSceneInfo* LightSceneInfo,
 	const FHairStrandsVisibilityData& InVisibilityData,
 	const FHairStrandsMacroGroupDatas& InMacroGroupDatas,
+	const FHairStrandsDeepShadowResources& DeepShadowResources,
+	const FHairStrandsVoxelResources& VoxelResources,
 	FRDGTextureRef OutShadowMask)
 {
-	if (InMacroGroupDatas.Datas.Num() == 0)
+	if (InMacroGroupDatas.Num() == 0)
 		return;
 
 	if (!HasDeepShadowData(LightSceneInfo, InMacroGroupDatas) && !IsHairStrandsVoxelizationEnable())
@@ -708,16 +712,16 @@ static void RenderHairStrandsShadowMask(
 		FRDGBufferRef DeepShadow_WorldToLightTransformBuffer = nullptr;
 		FRDGBufferSRVRef DeepShadow_WorldToLightTransformBufferSRV = nullptr;
 
-		for (const FHairStrandsMacroGroupData& MacroGroupData : InMacroGroupDatas.Datas)
+		for (const FHairStrandsMacroGroupData& MacroGroupData : InMacroGroupDatas)
 		{
-			for (const FHairStrandsDeepShadowData& DomData : MacroGroupData.DeepShadowDatas.Datas)
+			for (const FHairStrandsDeepShadowData& DomData : MacroGroupData.DeepShadowDatas)
 			{
 				if (DomData.LightId != LightSceneInfo->Id)
 					continue;
 
 				if (DeepShadow_WorldToLightTransformBuffer == nullptr)
 				{
-					DeepShadow_WorldToLightTransformBuffer = InMacroGroupDatas.DeepShadowResources.DeepShadowWorldToLightTransforms;
+					DeepShadow_WorldToLightTransformBuffer = DeepShadowResources.DeepShadowWorldToLightTransforms;
 					DeepShadow_WorldToLightTransformBufferSRV = GraphBuilder.CreateSRV(DeepShadow_WorldToLightTransformBuffer);
 				}
 
@@ -728,11 +732,11 @@ static void RenderHairStrandsShadowMask(
 				Params.CategorizationTexture = InVisibilityData.CategorizationTexture;
 				Params.DeepShadow_AtlasSlotIndex = DomData.AtlasSlotIndex;
 				Params.DeepShadow_WorldToLightTransformBuffer = DeepShadow_WorldToLightTransformBufferSRV;
-				Params.DeepShadow_bIsGPUDriven = InMacroGroupDatas.DeepShadowResources.bIsGPUDriven ? 1 : 0;
+				Params.DeepShadow_bIsGPUDriven = DeepShadowResources.bIsGPUDriven ? 1 : 0;
 				Params.DeepShadow_CPUWorldToLightTransform = DomData.CPU_WorldToLightTransform;
 				Params.DeepShadow_AtlasRect = DomData.AtlasRect;
-				Params.DeepShadow_FrontDepthTexture = InMacroGroupDatas.DeepShadowResources.DepthAtlasTexture;
-				Params.DeepShadow_LayerTexture = InMacroGroupDatas.DeepShadowResources.LayersAtlasTexture;
+				Params.DeepShadow_FrontDepthTexture = DeepShadowResources.DepthAtlasTexture;
+				Params.DeepShadow_LayerTexture = DeepShadowResources.LayersAtlasTexture;
 				Params.bIsWholeSceneLight = bIsWholeSceneLight;
 				Params.DeepShadow_DepthBiasScale = GetDeepShadowDepthBiasScale();
 				Params.DeepShadow_DensityScale = GetDeepShadowDensityScale();
@@ -750,10 +754,10 @@ static void RenderHairStrandsShadowMask(
 
 	// Code is disabled for now until we have the full DOM/voxel fallback logic
 	// If there is no deep shadow for this light, fallback on the voxel representation
-	if (!bHasDeepShadow && InMacroGroupDatas.VirtualVoxelResources.IsValid())
+	if (!bHasDeepShadow && VoxelResources.IsValid())
 	{
 		// TODO: Change this to be a single pass with virtual voxel?
-		for (const FHairStrandsMacroGroupData& MacroGroupData : InMacroGroupDatas.Datas)
+		for (const FHairStrandsMacroGroupData& MacroGroupData : InMacroGroupDatas)
 		{
 			const bool bIsWholeSceneLight = LightSceneInfo->Proxy->GetLightType() == LightType_Directional;
 
@@ -764,7 +768,7 @@ static void RenderHairStrandsShadowMask(
 			Params.Voxel_MacroGroupId = MacroGroupData.MacroGroupId;
 			Params.bIsWholeSceneLight = bIsWholeSceneLight ? 1 : 0;
 
-			Params.Voxel_VirtualVoxel = &InMacroGroupDatas.VirtualVoxelResources;
+			Params.Voxel_VirtualVoxel = &VoxelResources;
 			AddDeepShadowOpaqueMaskPass(
 				GraphBuilder,
 				SceneTextures.Depth.Resolve,
@@ -780,27 +784,27 @@ void RenderHairStrandsShadowMask(
 	FRDGBuilder& GraphBuilder,
 	const TArray<FViewInfo>& Views,
 	const FLightSceneInfo* LightSceneInfo,
-	const FHairStrandsRenderingData* HairDatas,
 	FRDGTextureRef OutShadowMask)
 {
-	if (Views.Num() == 0 || HairDatas == nullptr || OutShadowMask == nullptr)
+	if (Views.Num() == 0 || OutShadowMask == nullptr)
 	{
 		return;
 	}
 
-	const FHairStrandsVisibilityViews& HairVisibilityViews = HairDatas->HairVisibilityViews;
-	const FHairStrandsMacroGroupViews& MacroGroupViews = HairDatas->MacroGroupsPerViews;
-
-	for (int32 ViewIndex = 0; ViewIndex < Views.Num(); ++ViewIndex)
+	for (const FViewInfo& View : Views)
 	{
-		if (ViewIndex < HairVisibilityViews.HairDatas.Num() && ViewIndex < MacroGroupViews.Views.Num())
+		if (HairStrands::HasViewHairStrandsData(View))
 		{
-			const FHairStrandsVisibilityData& HairVisibilityData = HairVisibilityViews.HairDatas[ViewIndex];
-			const FHairStrandsMacroGroupDatas& MacroGroupDatas = MacroGroupViews.Views[ViewIndex];
-			if (HairVisibilityData.CategorizationTexture)
-			{
-				RenderHairStrandsShadowMask(GraphBuilder, Views[ViewIndex], LightSceneInfo, HairVisibilityData, MacroGroupDatas, OutShadowMask);
-			}
+			check(View.HairStrandsViewData.VisibilityData.CategorizationTexture);
+			RenderHairStrandsShadowMask(
+				GraphBuilder, 
+				View, 
+				LightSceneInfo, 
+				View.HairStrandsViewData.VisibilityData,
+				View.HairStrandsViewData.MacroGroupDatas,
+				View.HairStrandsViewData.DeepShadowResources,
+				View.HairStrandsViewData.VirtualVoxelResources,
+				OutShadowMask);
 		}
 	}
 }
\ No newline at end of file
diff --git a/Engine/Source/Runtime/Renderer/Private/HairStrands/HairStrandsTransmittance.h b/Engine/Source/Runtime/Renderer/Private/HairStrands/HairStrandsTransmittance.h
--- a/Engine/Source/Runtime/Renderer/Private/HairStrands/HairStrandsTransmittance.h
+++ b/Engine/Source/Runtime/Renderer/Private/HairStrands/HairStrandsTransmittance.h
@@ -24,13 +24,11 @@ void RenderHairStrandsShadowMask(
 	FRDGBuilder& GraphBuilder,
 	const TArray<FViewInfo>& Views,
 	const FLightSceneInfo* LightSceneInfo,
-	const struct FHairStrandsRenderingData* HairDatas,
 	FRDGTextureRef ScreenShadowMaskTexture); 
 
 /// Write hair transmittance onto screen shadow mask
 FHairStrandsTransmittanceMaskData RenderHairStrandsTransmittanceMask(
 	FRDGBuilder& GraphBuilder,
-	const TArray<FViewInfo>& Views,
+	const FViewInfo& View,
 	const class FLightSceneInfo* LightSceneInfo,
-	const struct FHairStrandsRenderingData* Hairdatas,
 	FRDGTextureRef ScreenShadowMaskSubPixelTexture);
\ No newline at end of file
diff --git a/Engine/Source/Runtime/Renderer/Private/HairStrands/HairStrandsUtils.cpp b/Engine/Source/Runtime/Renderer/Private/HairStrands/HairStrandsUtils.cpp
--- a/Engine/Source/Runtime/Renderer/Private/HairStrands/HairStrandsUtils.cpp
+++ b/Engine/Source/Runtime/Renderer/Private/HairStrands/HairStrandsUtils.cpp
@@ -4,6 +4,7 @@
 #include "ScenePrivate.h"
 #include "HairStrandsCluster.h"
 #include "Rendering/SkeletalMeshRenderData.h"
+#include "HairStrandsData.h"
 
 static float GHairR = 1;
 static float GHairTT = 1;
@@ -244,7 +245,7 @@ FIntRect ComputeVisibleHairStrandsMacroGroupsRect(const FIntRect& ViewRect, cons
 	FIntRect TotalRect(INT_MAX, INT_MAX, -INT_MAX, -INT_MAX);
 	if (IsHairStrandsViewRectOptimEnable())
 	{
-		for (const FHairStrandsMacroGroupData& Data : Datas.Datas)
+		for (const FHairStrandsMacroGroupData& Data : Datas)
 		{
 			TotalRect.Union(Data.ScreenRect);
 		}
diff --git a/Engine/Source/Runtime/Renderer/Private/HairStrands/HairStrandsUtils.h b/Engine/Source/Runtime/Renderer/Private/HairStrands/HairStrandsUtils.h
--- a/Engine/Source/Runtime/Renderer/Private/HairStrands/HairStrandsUtils.h
+++ b/Engine/Source/Runtime/Renderer/Private/HairStrands/HairStrandsUtils.h
@@ -10,9 +10,16 @@
 #include "RendererInterface.h"
 #include "HairStrandsInterface.h"
 #include "SceneTypes.h"
+#include "HairStrandsData.h"
 
 class FViewInfo;
 
+FMinHairRadiusAtDepth1 ComputeMinStrandRadiusAtDepth1(
+	const FIntPoint& Resolution,
+	const float FOV,
+	const uint32 SampleCount,
+	const float OverrideStrandHairRasterizationScale);
+
 FIntRect ComputeProjectedScreenRect(const FBox& B, const FViewInfo& View);
 
 void ComputeWorldToLightClip(
@@ -39,7 +46,7 @@ float GetHairDualScatteringRoughnessOverride();
 
 float SampleCountToSubPixelSize(uint32 SamplePerPixelCount);
 
-FIntRect ComputeVisibleHairStrandsMacroGroupsRect(const FIntRect& ViewRect, const struct FHairStrandsMacroGroupDatas& Datas);
+FIntRect ComputeVisibleHairStrandsMacroGroupsRect(const FIntRect& ViewRect, const FHairStrandsMacroGroupDatas& Datas);
 
 bool IsHairStrandsViewRectOptimEnable();
 
diff --git a/Engine/Source/Runtime/Renderer/Private/HairStrands/HairStrandsVisibility.cpp b/Engine/Source/Runtime/Renderer/Private/HairStrands/HairStrandsVisibility.cpp
--- a/Engine/Source/Runtime/Renderer/Private/HairStrands/HairStrandsVisibility.cpp
+++ b/Engine/Source/Runtime/Renderer/Private/HairStrands/HairStrandsVisibility.cpp
@@ -723,7 +723,7 @@ static void AddHairMaterialGBufferPass(
 				FDynamicPassMeshDrawListContext ShadowContext(DynamicMeshDrawCommandStorage, VisibleMeshDrawCommands, PipelineStateSet, NeedsShaderInitialization);
 				FHairMaterialGBufferProcessor MeshProcessor(Scene, ViewInfo, DrawRenderState, &ShadowContext);
 
-				for (const FHairStrandsMacroGroupData& MacroGroupData : MacroGroupDatas.Datas)
+				for (const FHairStrandsMacroGroupData& MacroGroupData : MacroGroupDatas)
 				{
 					for (const FHairStrandsMacroGroupData::PrimitiveInfo& PrimitiveInfo : MacroGroupData.PrimitivesInfos)
 					{
@@ -1164,7 +1164,7 @@ static FMaterialPassOutput AddHairMaterialPass(
 			FDynamicPassMeshDrawListContext ShadowContext(DynamicMeshDrawCommandStorage, VisibleMeshDrawCommands, PipelineStateSet, NeedsShaderInitialization);
 			FHairMaterialProcessor MeshProcessor(Scene, ViewInfo, DrawRenderState, &ShadowContext);
 
-			for (const FHairStrandsMacroGroupData& MacroGroupData : MacroGroupDatas.Datas)
+			for (const FHairStrandsMacroGroupData& MacroGroupData : MacroGroupDatas)
 			{
 				for (const FHairStrandsMacroGroupData::PrimitiveInfo& PrimitiveInfo : MacroGroupData.PrimitivesInfos)
 				{
@@ -2064,7 +2064,6 @@ IMPLEMENT_GLOBAL_SHADER(FHairVisibilityCompactionComputeRasterCS, "/Engine/Priva
 static void AddHairVisibilityCompactionComputeRasterPass(
 	FRDGBuilder& GraphBuilder,
 	const FViewInfo& View,
-	const FHairStrandsMacroGroupDatas& MacroGroupDatas,
 	const uint32 NodeGroupSize,
 	const uint32 SamplerPerPixel,
 	const FRasterComputeOutput& RasterComputeData,
@@ -2275,7 +2274,7 @@ static FRDGTextureRef AddHairVisibilityFillOpaqueDepth(
 	TArray<FIntRect> MacroGroupRects;
 	if (IsHairStrandsViewRectOptimEnable())
 	{
-		for (const FHairStrandsMacroGroupData& MacroGroupData : MacroGroupDatas.Datas)
+		for (const FHairStrandsMacroGroupData& MacroGroupData : MacroGroupDatas)
 		{
 			MacroGroupRects.Add(MacroGroupData.ScreenRect);
 		}
@@ -2336,7 +2335,7 @@ static void AddHairCulledVertexResourcesTransitionPass(
 	const FHairStrandsMacroGroupDatas& MacroGroupDatas)
 {
 	FBufferTransitionQueue TransitionQueue;
-	for (const FHairStrandsMacroGroupData& MacroGroupData : MacroGroupDatas.Datas)
+	for (const FHairStrandsMacroGroupData& MacroGroupData : MacroGroupDatas)
 	{
 		for (const FHairStrandsMacroGroupData::PrimitiveInfo& PrimitiveInfo : MacroGroupData.PrimitivesInfos)
 		{
@@ -2447,7 +2446,7 @@ static void AddHairVisibilityCommonPass(
 			FDynamicPassMeshDrawListContext ShadowContext(DynamicMeshDrawCommandStorage, VisibleMeshDrawCommands, PipelineStateSet, NeedsShaderInitialization);
 			FHairVisibilityProcessor MeshProcessor(Scene, ViewInfo, DrawRenderState, RenderMode, &ShadowContext);
 			
-			for (const FHairStrandsMacroGroupData& MacroGroupData : MacroGroupDatas.Datas)
+			for (const FHairStrandsMacroGroupData& MacroGroupData : MacroGroupDatas)
 			{
 				for (const FHairStrandsMacroGroupData::PrimitiveInfo& PrimitiveInfo : MacroGroupData.PrimitivesInfos)
 				{
@@ -3158,7 +3157,7 @@ static FRasterComputeOutput AddVisibilityComputeRasterPass(
 	TShaderMapRef<FVisiblityRasterComputeCS> ComputeShader_CullingOff(ViewInfo.ShaderMap, PermutationVector0);
 	TShaderMapRef<FVisiblityRasterComputeCS> ComputeShader_CullingOn (ViewInfo.ShaderMap, PermutationVector1);
 
-	for (const FHairStrandsMacroGroupData& MacroGroup : MacroGroupDatas.Datas)
+	for (const FHairStrandsMacroGroupData& MacroGroup : MacroGroupDatas)
 	{
 		const FHairStrandsMacroGroupData::TPrimitiveInfos& PrimitiveSceneInfos = MacroGroup.PrimitivesInfos;
 
@@ -3221,10 +3220,10 @@ static FRasterComputeOutput AddVisibilityComputeRasterPass(
 ///////////////////////////////////////////////////////////////////////////////////////////////////
 bool GetHairStrandsSkyLightingEnable();
 
-FHairStrandsVisibilityViews RenderHairStrandsVisibilityBuffer(
+void RenderHairStrandsVisibilityBuffer(
 	FRDGBuilder& GraphBuilder,
 	const FScene* Scene,
-	const TArray<FViewInfo>& Views,
+	FViewInfo& View,
 	FRDGTextureRef SceneGBufferATexture,
 	FRDGTextureRef SceneGBufferBTexture,
 	FRDGTextureRef SceneGBufferCTexture,
@@ -3233,28 +3232,24 @@ FHairStrandsVisibilityViews RenderHairStrandsVisibilityBuffer(
 	FRDGTextureRef SceneColorTexture,
 	FRDGTextureRef SceneDepthTexture,
 	FRDGTextureRef SceneVelocityTexture,
-	FInstanceCullingManager& InstanceCullingManager,
-	const FHairStrandsMacroGroupViews& MacroGroupViews)
+	FInstanceCullingManager& InstanceCullingManager)
 {
 	QUICK_SCOPE_CYCLE_COUNTER(STAT_CLM_RenderHairStrandsVisibility);
 	RDG_EVENT_SCOPE(GraphBuilder, "HairStrandsVisibility");
 	RDG_GPU_STAT_SCOPE(GraphBuilder, HairStrandsVisibility);
 
-	FHairStrandsVisibilityViews Output;
-	for (int32 ViewIndex = 0; ViewIndex < Views.Num(); ViewIndex++)
+	FHairStrandsMacroGroupDatas& MacroGroupDatas = View.HairStrandsViewData.MacroGroupDatas;
+	check(View.Family);
+	check(MacroGroupDatas.Num() > 0);
+
 	{
-		const FViewInfo& View = Views[ViewIndex];
-		if (View.Family)
+		
 		{
 			FHairLUT HairLUT = GetHairLUT(GraphBuilder, View);
 
-			FHairStrandsVisibilityData& VisibilityData = Output.HairDatas.AddDefaulted_GetRef();
+			FHairStrandsVisibilityData& VisibilityData = View.HairStrandsViewData.VisibilityData;
 			VisibilityData.NodeGroupSize = GetVendorOptimalGroupSize1D();
 			VisibilityData.MaxSampleCount = GetMaxSamplePerPixel();
-			const FHairStrandsMacroGroupDatas& MacroGroupDatas = MacroGroupViews.Views[ViewIndex];
-
-			if (MacroGroupDatas.Datas.Num() == 0)
-				continue;
 
 			// Use the scene color for computing target resolution as the View.ViewRect, 
 			// doesn't include the actual resolution padding which make buffer size 
@@ -3304,7 +3299,6 @@ FHairStrandsVisibilityViews RenderHairStrandsVisibilityBuffer(
 						AddHairVisibilityCompactionComputeRasterPass(
 							GraphBuilder,
 							View,
-							MacroGroupDatas,
 							VisibilityData.NodeGroupSize,
 							VisibilityData.MaxSampleCount,
 							RasterOutput,
@@ -3754,6 +3748,4 @@ FHairStrandsVisibilityViews RenderHairStrandsVisibilityBuffer(
 			}
 		}
 	}
-
-	return Output;
 }
diff --git a/Engine/Source/Runtime/Renderer/Private/HairStrands/HairStrandsVisibility.h b/Engine/Source/Runtime/Renderer/Private/HairStrands/HairStrandsVisibility.h
--- a/Engine/Source/Runtime/Renderer/Private/HairStrands/HairStrandsVisibility.h
+++ b/Engine/Source/Runtime/Renderer/Private/HairStrands/HairStrandsVisibility.h
@@ -11,70 +11,12 @@
 #include "RenderGraphResources.h"
 
 class FViewInfo;
-struct FVector;
-struct FIntPoint;
 class FInstanceCullingManager;
 
-struct FHairStrandsVisibilityData
-{
-	FRDGTextureRef DepthTexture = nullptr;
-	FRDGTextureRef IDTexture = nullptr;
-	FRDGTextureRef MaterialTexture = nullptr;
-	FRDGTextureRef AttributeTexture = nullptr;
-	FRDGTextureRef VelocityTexture = nullptr;
-	FRDGTextureRef ResolveMaskTexture = nullptr;
-	FRDGTextureRef CategorizationTexture = nullptr;
-	FRDGTextureRef ViewHairCountTexture = nullptr;
-	FRDGTextureRef ViewHairCountUintTexture = nullptr;
-	FRDGTextureRef DepthTextureUint = nullptr;
-	FRDGTextureRef EmissiveTexture = nullptr;
-
-	FRDGTextureRef HairOnlyDepthTexture = nullptr;
-
-	FRDGTextureRef ViewHairVisibilityTexture0 = nullptr;
-	FRDGTextureRef ViewHairVisibilityTexture1 = nullptr;
-	FRDGTextureRef ViewHairVisibilityTexture2 = nullptr;
-	FRDGTextureRef ViewHairVisibilityTexture3 = nullptr;
-
-	FRDGTextureRef LightChannelMaskTexture = nullptr;
-
-	FRDGTextureRef	PPLLNodeCounterTexture = nullptr;
-	FRDGTextureRef	PPLLNodeIndexTexture = nullptr;
-	FRDGBufferRef	PPLLNodeDataBuffer = nullptr;
-	uint32			MaxPPLLNodePerPixelCount = 0;
-	uint32			MaxPPLLNodeCount = 0;
-
-	FRDGTextureRef	TileIndexTexture = nullptr;
-	FRDGBufferRef	TileBuffer = nullptr;
-	FRDGBufferRef	TileIndirectArgs = nullptr;
-	const uint32	TileSize = 8;
-	const uint32	TileThreadGroupSize = 32;
-
-	uint32			MaxSampleCount = 8;
-	uint32			MaxNodeCount = 0;
-	FRDGTextureRef	NodeCount = nullptr;
-	FRDGTextureRef	NodeIndex = nullptr;
-	FRDGBufferRef	NodeData = nullptr;
-	FRDGBufferRef	NodeCoord = nullptr;
-	FRDGBufferRef	NodeIndirectArg = nullptr;
-	uint32			NodeGroupSize = 0;
-
-	// Hair lighting is accumulated within this buffer
-	// Allocated conservatively
-	// User indirect dispatch for accumulating contribution
-	FIntPoint SampleLightingViewportResolution;
-	FRDGTextureRef SampleLightingBuffer = nullptr;
-};
-
-struct FHairStrandsVisibilityViews
-{
-	TArray<FHairStrandsVisibilityData, SceneRenderingAllocator> HairDatas;
-};
-
-FHairStrandsVisibilityViews RenderHairStrandsVisibilityBuffer(
+void RenderHairStrandsVisibilityBuffer(
 	FRDGBuilder& GraphBuilder,
 	const class FScene* Scene,
-	const TArray<FViewInfo>& Views,
+	FViewInfo& View,
 	FRDGTextureRef SceneGBufferATexture,
 	FRDGTextureRef SceneGBufferBTexture,
 	FRDGTextureRef SceneGBufferCTexture,
@@ -83,8 +25,7 @@ FHairStrandsVisibilityViews RenderHairStrandsVisibilityBuffer(
 	FRDGTextureRef SceneColorTexture,
 	FRDGTextureRef SceneDepthTexture,
 	FRDGTextureRef SceneVelocityTexture,
-	FInstanceCullingManager& InstanceCullingManager,
-	const struct FHairStrandsMacroGroupViews& MacroGroupViews);
+	FInstanceCullingManager& InstanceCullingManager);
 
 void SetUpViewHairRenderInfo(const FViewInfo& ViewInfo, FVector4& OutHairRenderInfo, uint32& OutHairRenderInfoBits, uint32& OutHairComponents);
 
diff --git a/Engine/Source/Runtime/Renderer/Private/HairStrands/HairStrandsVoxelization.cpp b/Engine/Source/Runtime/Renderer/Private/HairStrands/HairStrandsVoxelization.cpp
--- a/Engine/Source/Runtime/Renderer/Private/HairStrands/HairStrandsVoxelization.cpp
+++ b/Engine/Source/Runtime/Renderer/Private/HairStrands/HairStrandsVoxelization.cpp
@@ -135,7 +135,7 @@ class FVirtualVoxelInjectOpaqueCS : public FGlobalShader
 
 	BEGIN_SHADER_PARAMETER_STRUCT(FParameters, )
 		SHADER_PARAMETER_STRUCT_INCLUDE(FSceneTextureParameters, SceneTextures)
-		SHADER_PARAMETER_STRUCT(FVirtualVoxelCommonParameters, VirtualVoxelParams)
+		SHADER_PARAMETER_STRUCT(FHairStrandsVoxelCommonParameters, VirtualVoxelParams)
 		SHADER_PARAMETER(FIntVector, DispatchedPageIndexResolution)
 		SHADER_PARAMETER(uint32, MacroGroupId)
 		SHADER_PARAMETER(FVector2D, SceneDepthResolution)
@@ -161,7 +161,7 @@ IMPLEMENT_GLOBAL_SHADER(FVirtualVoxelInjectOpaqueCS, "/Engine/Private/HairStrand
 static void AddVirtualVoxelInjectOpaquePass(
 	FRDGBuilder& GraphBuilder,
 	const FViewInfo& View,
-	const FVirtualVoxelResources& VoxelResources,
+	const FHairStrandsVoxelResources& VoxelResources,
 	const FHairStrandsMacroGroupData& MacroGroup)
 {
 	FSceneTextureParameters SceneTextures = GetSceneTextureParameters(GraphBuilder);
@@ -438,7 +438,7 @@ class FVoxelIndPageClearCS : public FGlobalShader
 	SHADER_USE_PARAMETER_STRUCT(FVoxelIndPageClearCS, FGlobalShader);
 
 	BEGIN_SHADER_PARAMETER_STRUCT(FParameters, )
-		SHADER_PARAMETER_STRUCT(FVirtualVoxelCommonParameters, VirtualVoxelParams)
+		SHADER_PARAMETER_STRUCT(FHairStrandsVoxelCommonParameters, VirtualVoxelParams)
 		SHADER_PARAMETER_RDG_TEXTURE_UAV(RWTexture3D, OutPageTexture)
 		RDG_BUFFER_ACCESS(IndirectDispatchBuffer, ERHIAccess::IndirectArgs)
 	END_SHADER_PARAMETER_STRUCT()
@@ -472,7 +472,8 @@ inline FIntVector CeilToInt(const FVector& V)
 static void AddAllocateVoxelPagesPass(
 	FRDGBuilder& GraphBuilder,
 	const FViewInfo& View,
-	FHairStrandsMacroGroupDatas& MacroGroups,
+	FHairStrandsMacroGroupDatas& MacroGroupDatas,
+	FHairStrandsMacroGroupResources& MacroGroupResources,
 	const FIntVector PageCountResolution,
 	const uint32 PageCount,
 	const float VoxelWorldSize,
@@ -491,7 +492,7 @@ static void AddAllocateVoxelPagesPass(
 {
 	const uint32 GroupSize = 32;
 	const bool bIsGPUDriven = GHairVirtualVoxelGPUDriven > 0;
-	const uint32 MacroGroupCount = MacroGroups.Datas.Num();
+	const uint32 MacroGroupCount = MacroGroupDatas.Num();
 	if (MacroGroupCount == 0)
 		return;
 
@@ -508,7 +509,7 @@ static void AddAllocateVoxelPagesPass(
 
 	OutTotalPageIndexCount = 0;
 	TArray<FCPUMacroGroupAllocation> CPUAllocationDescs;	
-	for (FHairStrandsMacroGroupData& MacroGroup : MacroGroups.Datas)
+	for (FHairStrandsMacroGroupData& MacroGroup : MacroGroupDatas)
 	{
 		// Snap the max AABB to the voxel size
 		// Scale the bounding box in place of proper GPU driven AABB for now
@@ -578,7 +579,7 @@ static void AddAllocateVoxelPagesPass(
 		Parameters->TotalPageIndexCount = OutTotalPageIndexCount;
 		Parameters->PageResolution = PageResolution;
 		Parameters->MacroGroupCount = MacroGroupCount;
-		Parameters->MacroGroupAABBBuffer = GraphBuilder.CreateUAV(MacroGroups.MacroGroupResources.MacroGroupAABBsBuffer, PF_R32_SINT);
+		Parameters->MacroGroupAABBBuffer = GraphBuilder.CreateUAV(MacroGroupResources.MacroGroupAABBsBuffer, PF_R32_SINT);
 		Parameters->IndirectDispatchGroupSize = GroupSize; // This is the GroupSize used for FVoxelAllocateVoxelPageCS
 		Parameters->OutPageIndexResolutionAndOffsetBuffer = GraphBuilder.CreateUAV(PageIndexResolutionBuffer, PF_R32G32B32A32_UINT);
 		Parameters->OutVoxelizationViewInfoBuffer = GraphBuilder.CreateUAV(VoxelizationViewInfoBuffer);
@@ -606,7 +607,7 @@ static void AddAllocateVoxelPagesPass(
 	uint32 TotalClusterCount = 0;
 	for (uint32 MacroGroupIt = 0; MacroGroupIt < MacroGroupCount; ++MacroGroupIt)
 	{
-		const FHairStrandsMacroGroupData& MacroGroup = MacroGroups.Datas[MacroGroupIt];
+		const FHairStrandsMacroGroupData& MacroGroup = MacroGroupDatas[MacroGroupIt];
 		for (const FHairStrandsMacroGroupData::PrimitiveInfo& PrimitiveInfo : MacroGroup.PrimitivesInfos)
 		{
 			FHairGroupPublicData* HairGroupData = PrimitiveInfo.PublicDataPtr;
@@ -621,7 +622,7 @@ static void AddAllocateVoxelPagesPass(
 		SCOPED_DRAW_EVENT(GraphBuilder.RHICmdList, HairStrandsAllocateMacroGroup);
 		SCOPED_GPU_STAT(GraphBuilder.RHICmdList, HairStrandsAllocateMacroGroup);
 
-		const FHairStrandsMacroGroupData& MacroGroup = MacroGroups.Datas[MacroGroupIt];
+		const FHairStrandsMacroGroupData& MacroGroup = MacroGroupDatas[MacroGroupIt];
 		FCPUMacroGroupAllocation& CPUAllocationDesc = CPUAllocationDescs[MacroGroupIt];
 
 		const bool bUseIndirectScatter = GHairVirtualVoxelUseIndirectScatterPageAllocation>0 && bIsGPUDriven;
@@ -649,7 +650,7 @@ static void AddAllocateVoxelPagesPass(
 				Parameters->bForceDirectPageAllocation	= GHairVirtualVoxelUseImmediatePageAllocation > 0 ? 1 : 0;
 
 				Parameters->ClusterAABBsBuffer		= HairGroupData->GetClusterAABBBuffer().SRV;
-				Parameters->MacroGroupAABBBuffer	= GraphBuilder.CreateSRV(MacroGroups.MacroGroupResources.MacroGroupAABBsBuffer, PF_R32_SINT);
+				Parameters->MacroGroupAABBBuffer	= GraphBuilder.CreateSRV(MacroGroupResources.MacroGroupAABBsBuffer, PF_R32_SINT);
 				Parameters->PageIndexResolutionAndOffsetBuffer = PageIndexResolutionAndOffsetBufferSRV;
 
 				Parameters->OutDeferredScatterCounter	= ScatterCounterUAV;
@@ -717,7 +718,7 @@ static void AddAllocateVoxelPagesPass(
 
 				if (bIsGPUDriven)
 				{
-					Parameters->MacroGroupAABBBuffer = GraphBuilder.CreateSRV(MacroGroups.MacroGroupResources.MacroGroupAABBsBuffer, PF_R32_SINT);
+					Parameters->MacroGroupAABBBuffer = GraphBuilder.CreateSRV(MacroGroupResources.MacroGroupAABBsBuffer, PF_R32_SINT);
 					Parameters->PageIndexResolutionAndOffsetBuffer = PageIndexResolutionAndOffsetBufferSRV;
 				}
 
@@ -752,7 +753,7 @@ static void AddAllocateVoxelPagesPass(
 
 			if (bIsGPUDriven)
 			{
-				Parameters->MacroGroupAABBBuffer = GraphBuilder.CreateSRV(MacroGroups.MacroGroupResources.MacroGroupAABBsBuffer, PF_R32_SINT);
+				Parameters->MacroGroupAABBBuffer = GraphBuilder.CreateSRV(MacroGroupResources.MacroGroupAABBsBuffer, PF_R32_SINT);
 				Parameters->PageIndexResolutionAndOffsetBuffer = PageIndexResolutionAndOffsetBufferSRV;
 			}
 
@@ -827,17 +828,18 @@ static void AddAllocateVoxelPagesPass(
 	OutVoxelizationViewInfoBuffer = VoxelizationViewInfoBuffer;
 }
 
-FVirtualVoxelResources AllocateVirtualVoxelResources(
+FHairStrandsVoxelResources AllocateVirtualVoxelResources(
 	FRDGBuilder& GraphBuilder,
 	const FViewInfo& View,
-	FHairStrandsMacroGroupDatas& MacroGroups, 
+	FHairStrandsMacroGroupDatas& MacroGroupDatas,
+	FHairStrandsMacroGroupResources& MacroGroupResources,
 	FRDGBufferRef& PageToPageIndexBuffer)
 {
 	DECLARE_GPU_STAT(HairStrandsVoxelPageAllocation);
 	RDG_EVENT_SCOPE(GraphBuilder, "HairStrandsVoxelPageAllocation");
 	RDG_GPU_STAT_SCOPE(GraphBuilder, HairStrandsVoxelPageAllocation);
 
-	FVirtualVoxelResources Out;
+	FHairStrandsVoxelResources Out;
 
 	Out.Parameters.Common.PageCountResolution		= FIntVector(GHairVirtualVoxel_PageCountPerDim, GHairVirtualVoxel_PageCountPerDim, GHairVirtualVoxel_PageCountPerDim);
 	Out.Parameters.Common.PageCount					= Out.Parameters.Common.PageCountResolution.X * Out.Parameters.Common.PageCountResolution.Y * Out.Parameters.Common.PageCountResolution.Z;
@@ -861,7 +863,7 @@ FVirtualVoxelResources AllocateVirtualVoxelResources(
 	Out.Parameters.Common.SteppingScale_Environment	  = FMath::Clamp(GHairStransVoxelRaymarchingSteppingScale_Environment		>= 0 ? GHairStransVoxelRaymarchingSteppingScale_Environment		: GHairStransVoxelRaymarchingSteppingScale, 1.f, 10.f);
 	Out.Parameters.Common.SteppingScale_Raytracing	  = FMath::Clamp(GHairStransVoxelRaymarchingSteppingScale_Raytracing		>= 0 ? GHairStransVoxelRaymarchingSteppingScale_Raytracing		: GHairStransVoxelRaymarchingSteppingScale, 1.f, 10.f);
 
-	Out.Parameters.Common.NodeDescCount				= MacroGroups.Datas.Num();
+	Out.Parameters.Common.NodeDescCount				= MacroGroupDatas.Num();
 	Out.Parameters.Common.IndirectDispatchGroupSize = 64;	
 	Out.Parameters.Common.Raytracing_ShadowOcclusionThreshold	= FMath::Max(0.f, GHairVirtualVoxelRaytracing_ShadowOcclusionThreshold);
 	Out.Parameters.Common.Raytracing_SkyOcclusionThreshold		= FMath::Max(0.f, GHairVirtualVoxelRaytracing_SkyOcclusionThreshold);
@@ -873,7 +875,8 @@ FVirtualVoxelResources AllocateVirtualVoxelResources(
 	AddAllocateVoxelPagesPass(
 		GraphBuilder, 
 		View, 
-		MacroGroups,
+		MacroGroupDatas,
+		MacroGroupResources,
 		Out.Parameters.Common.PageCountResolution,
 		Out.Parameters.Common.PageCount,
 		Out.Parameters.Common.VoxelWorldSize,
@@ -925,7 +928,7 @@ FVirtualVoxelResources AllocateVirtualVoxelResources(
 static FRDGBufferRef IndirectVoxelPageClear(
 	FRDGBuilder& GraphBuilder,
 	const FViewInfo& ViewInfo,
-	FVirtualVoxelResources& VoxelResources)
+	FHairStrandsVoxelResources& VoxelResources)
 {
 	DECLARE_GPU_STAT(HairStrandsIndVoxelPageClear);
 	SCOPED_DRAW_EVENT(GraphBuilder.RHICmdList, HairStrandsIndVoxelPageClear);
@@ -980,7 +983,7 @@ class FVoxelRasterComputeCS : public FGlobalShader
 	using FPermutationDomain = TShaderPermutationDomain<FCulling>;
 
 	BEGIN_SHADER_PARAMETER_STRUCT(FParameters, )
-		SHADER_PARAMETER_STRUCT(FVirtualVoxelCommonParameters, VirtualVoxelParams)
+		SHADER_PARAMETER_STRUCT(FHairStrandsVoxelCommonParameters, VirtualVoxelParams)
 		SHADER_PARAMETER(uint32, MacroGroupId)
 		SHADER_PARAMETER(uint32, DispatchCountX)
 		SHADER_PARAMETER(uint32, MaxRasterCount)
@@ -1016,7 +1019,7 @@ IMPLEMENT_GLOBAL_SHADER(FVoxelRasterComputeCS, "/Engine/Private/HairStrands/Hair
 static void AddVirtualVoxelizationComputeRasterPass(
 	FRDGBuilder& GraphBuilder,
 	const FViewInfo* ViewInfo,
-	FVirtualVoxelResources& VoxelResources,
+	FHairStrandsVoxelResources& VoxelResources,
 	FHairStrandsMacroGroupData& MacroGroup)
 {
 	const bool bIsGPUDriven = GHairVirtualVoxelGPUDriven > 0;
@@ -1103,7 +1106,7 @@ static void AddVirtualVoxelizationRasterPass(
 	FRDGBuilder& GraphBuilder,
 	const FScene* Scene,
 	const FViewInfo* ViewInfo,
-	FVirtualVoxelResources& VoxelResources,
+	FHairStrandsVoxelResources& VoxelResources,
 	FInstanceCullingManager& InstanceCullingManager,
 	FHairStrandsMacroGroupData& MacroGroup)
 {
@@ -1304,19 +1307,17 @@ IMPLEMENT_GLOBAL_SHADER(FVirtualVoxelPatchPageIndexWithMipDataCS, "/Engine/Priva
 static void AddVirtualVoxelGenerateMipPass(
 	FRDGBuilder& GraphBuilder,
 	const FViewInfo& View,
-	FHairStrandsMacroGroupDatas& MacroGroups,
+	FHairStrandsVoxelResources& VoxelResources,
 	FRDGBufferRef IndirectArgsBuffer, 
 	FRDGBufferRef InPageToPageIndexBuffer)
 {
-	if (!MacroGroups.VirtualVoxelResources.IsValid())
+	if (!VoxelResources.IsValid())
 		return;
 
 	DECLARE_GPU_STAT(HairStrandsDensityMipGen);
 	SCOPED_DRAW_EVENT(GraphBuilder.RHICmdList, HairStrandsDensityMipGen);
 	SCOPED_GPU_STAT(GraphBuilder.RHICmdList, HairStrandsDensityMipGen);
 
-	FVirtualVoxelResources& VoxelResources = MacroGroups.VirtualVoxelResources;
-
 	const uint32 MipCount = VoxelResources.PageTexture->Desc.NumMips;
 
 	// Prepare indirect dispatch for all the pages this frame (allocated linearly in 3D DensityTexture)
@@ -1403,53 +1404,42 @@ static void AddVirtualVoxelGenerateMipPass(
 void VoxelizeHairStrands(
 	FRDGBuilder& GraphBuilder, 
 	const FScene* Scene, 
-	const TArray<FViewInfo>& Views,
-	FInstanceCullingManager& InstanceCullingManager,
-	FHairStrandsMacroGroupViews& MacroGroupsViews)
+	FViewInfo& View,
+	FInstanceCullingManager& InstanceCullingManager)
 {
 	if (!IsHairStrandsVoxelizationEnable())
 		return;
 
-	FHairStrandsMacroGroupViews PrimitivesClusterViews;
-	for (int32 ViewIndex = 0; ViewIndex < Views.Num(); ViewIndex++)
-	{			
-		if (ViewIndex >= MacroGroupsViews.Views.Num())
-			continue;
+	FHairStrandsMacroGroupDatas& MacroGroupDatas = View.HairStrandsViewData.MacroGroupDatas;
+	FHairStrandsVoxelResources& VirtualVoxelResources = View.HairStrandsViewData.VirtualVoxelResources;
+	FHairStrandsMacroGroupResources& MacroGroupResources = View.HairStrandsViewData.MacroGroupResources;
 
-		const FViewInfo& View = Views[ViewIndex];
-		FHairStrandsMacroGroupDatas& MacroGroupDatas = MacroGroupsViews.Views[ViewIndex];
-
-		if (MacroGroupDatas.Datas.Num() == 0)
-			continue;
+	if (MacroGroupDatas.Num() == 0)
+		return;
 
-		DECLARE_GPU_STAT(HairStrandsVoxelization);
-		RDG_EVENT_SCOPE(GraphBuilder, "HairStrandsVoxelization");
-		RDG_GPU_STAT_SCOPE(GraphBuilder, HairStrandsVoxelization);
+	DECLARE_GPU_STAT(HairStrandsVoxelization);
+	RDG_EVENT_SCOPE(GraphBuilder, "HairStrandsVoxelization");
+	RDG_GPU_STAT_SCOPE(GraphBuilder, HairStrandsVoxelization);
 
-		if (MacroGroupDatas.Datas.Num() > 0)
-		{
-			// Toto moves this function into the render graph. At the moment this is not possible as this functions 
-			// generates internally a non-transient constant buffer which initialized VirtualVoxelResources. This 
-			// needs to be rewritten/worked out.
-			FRDGBufferRef PageToPageIndexBuffer = nullptr;
-			MacroGroupDatas.VirtualVoxelResources = AllocateVirtualVoxelResources(GraphBuilder, View, MacroGroupDatas, PageToPageIndexBuffer);
+	{
+		FRDGBufferRef PageToPageIndexBuffer = nullptr;
+		VirtualVoxelResources = AllocateVirtualVoxelResources(GraphBuilder, View, MacroGroupDatas, MacroGroupResources, PageToPageIndexBuffer);
 
-			FRDGBufferRef ClearIndArgsBuffer = IndirectVoxelPageClear(GraphBuilder, View, MacroGroupDatas.VirtualVoxelResources);
+		FRDGBufferRef ClearIndArgsBuffer = IndirectVoxelPageClear(GraphBuilder, View, VirtualVoxelResources);
 
-			for (FHairStrandsMacroGroupData& MacroGroup : MacroGroupDatas.Datas)
-			{
-				AddVirtualVoxelizationRasterPass(GraphBuilder, Scene, &View, MacroGroupDatas.VirtualVoxelResources, InstanceCullingManager, MacroGroup);
-			}
+		for (FHairStrandsMacroGroupData& MacroGroup : MacroGroupDatas)
+		{
+			AddVirtualVoxelizationRasterPass(GraphBuilder, Scene, &View, VirtualVoxelResources, InstanceCullingManager, MacroGroup);
+		}
 
-			if (GHairVoxelInjectOpaqueDepthEnable > 0)
+		if (GHairVoxelInjectOpaqueDepthEnable > 0)
+		{
+			for (FHairStrandsMacroGroupData& MacroGroup : MacroGroupDatas)
 			{
-				for (FHairStrandsMacroGroupData& MacroGroup : MacroGroupDatas.Datas)
-				{
-					AddVirtualVoxelInjectOpaquePass(GraphBuilder, View, MacroGroupDatas.VirtualVoxelResources, MacroGroup);
-				}
+				AddVirtualVoxelInjectOpaquePass(GraphBuilder, View, VirtualVoxelResources, MacroGroup);
 			}
-
-			AddVirtualVoxelGenerateMipPass(GraphBuilder, View, MacroGroupDatas, ClearIndArgsBuffer, PageToPageIndexBuffer);
 		}
+
+		AddVirtualVoxelGenerateMipPass(GraphBuilder, View, VirtualVoxelResources, ClearIndArgsBuffer, PageToPageIndexBuffer);
 	}
 }
diff --git a/Engine/Source/Runtime/Renderer/Private/HairStrands/HairStrandsVoxelization.h b/Engine/Source/Runtime/Renderer/Private/HairStrands/HairStrandsVoxelization.h
--- a/Engine/Source/Runtime/Renderer/Private/HairStrands/HairStrandsVoxelization.h
+++ b/Engine/Source/Runtime/Renderer/Private/HairStrands/HairStrandsVoxelization.h
@@ -11,87 +11,6 @@
 #include "RenderGraphResources.h"
 #include "Renderer/Private/SceneRendering.h"
 
-struct FVirtualVoxelNodeDesc
-{
-	FVector WorldMinAABB = FVector::ZeroVector;
-	FVector WorldMaxAABB = FVector::ZeroVector;
-	FIntVector PageIndexResolution = FIntVector::ZeroValue;
-	FMatrix WorldToClip;
-};
-
-struct FPackedVirtualVoxelNodeDesc
-{
-	// This is just a placeholder having the correct size. The actual definition is in HairStradsNVoxelPageCommon.ush
-	const static EPixelFormat Format = PF_R32G32B32A32_UINT;
-	const static uint32 ComponentCount = 2;
-
-	// Shader View is struct { uint4; uint4; }
-	FVector	MinAABB;
-	uint32	PackedPageIndexResolution;
-	FVector	MaxAABB;
-	uint32	PageIndexOffset;
-};
-
-// PixelRadiusAtDepth1 shouldn't be stored into this structure should be view independent, 
-// but is put here for convenience at the moment since multiple views are not supported 
-// at the moment
-BEGIN_SHADER_PARAMETER_STRUCT(FVirtualVoxelCommonParameters, )
-	SHADER_PARAMETER(FIntVector, PageCountResolution)
-	SHADER_PARAMETER(float, VoxelWorldSize)
-	SHADER_PARAMETER(FIntVector, PageTextureResolution)
-	SHADER_PARAMETER(uint32, PageCount)
-	SHADER_PARAMETER(uint32, PageResolution)
-	SHADER_PARAMETER(uint32, PageIndexCount)
-	SHADER_PARAMETER(uint32, IndirectDispatchGroupSize)
-	SHADER_PARAMETER(uint32, NodeDescCount)
-
-	SHADER_PARAMETER(float, DensityScale)
-	SHADER_PARAMETER(float, DensityScale_AO)
-	SHADER_PARAMETER(float, DensityScale_Shadow)
-	SHADER_PARAMETER(float, DensityScale_Transmittance)
-	SHADER_PARAMETER(float, DensityScale_Environment)
-	SHADER_PARAMETER(float, DensityScale_Raytracing)
-
-	SHADER_PARAMETER(float, DepthBiasScale_Shadow)
-	SHADER_PARAMETER(float, DepthBiasScale_Transmittance)
-	SHADER_PARAMETER(float, DepthBiasScale_Environment)
-
-	SHADER_PARAMETER(float, SteppingScale_Shadow)
-	SHADER_PARAMETER(float, SteppingScale_Transmittance)
-	SHADER_PARAMETER(float, SteppingScale_Environment)
-	SHADER_PARAMETER(float, SteppingScale_Raytracing)
-
-	SHADER_PARAMETER(float, HairCoveragePixelRadiusAtDepth1) 
-	SHADER_PARAMETER(float, Raytracing_ShadowOcclusionThreshold)
-	SHADER_PARAMETER(float, Raytracing_SkyOcclusionThreshold)
-	SHADER_PARAMETER_RDG_BUFFER_SRV(Buffer<uint>, PageIndexBuffer)
-	SHADER_PARAMETER_RDG_BUFFER_SRV(Buffer<uint2>, PageIndexOccupancyBuffer)
-	SHADER_PARAMETER_RDG_BUFFER_SRV(Buffer<uint>, PageIndexCoordBuffer)
-	SHADER_PARAMETER_RDG_BUFFER_SRV(StructuredBuffer<FPackedVirtualVoxelNodeDesc>, NodeDescBuffer) // Packed into 2 x uint4
-	SHADER_PARAMETER_RDG_TEXTURE(Texture2D<float>, HairCoverageLUT)
-	SHADER_PARAMETER_SAMPLER(SamplerState, HairCoverageSampler)
-END_SHADER_PARAMETER_STRUCT()
-
-BEGIN_GLOBAL_SHADER_PARAMETER_STRUCT(FVirtualVoxelParameters, RENDERER_API)
-	SHADER_PARAMETER_STRUCT_INCLUDE(FVirtualVoxelCommonParameters, Common)
-	SHADER_PARAMETER_RDG_TEXTURE(Texture3D<uint>, PageTexture)
-END_GLOBAL_SHADER_PARAMETER_STRUCT()
-
-struct FVirtualVoxelResources
-{
-	FVirtualVoxelParameters	Parameters;
-	TRDGUniformBufferRef<FVirtualVoxelParameters> UniformBuffer;
-	FRDGTextureRef PageTexture = nullptr;
-	FRDGBufferRef PageIndexBuffer = nullptr;
-	FRDGBufferRef PageIndexOccupancyBuffer = nullptr;
-	FRDGBufferRef NodeDescBuffer = nullptr;
-	FRDGBufferRef PageIndexCoordBuffer = nullptr;
-	FRDGBufferRef IndirectArgsBuffer = nullptr;
-	FRDGBufferRef PageIndexGlobalCounter = nullptr;
-	FRDGBufferRef VoxelizationViewInfoBuffer = nullptr;
-
-	const bool IsValid() const { return UniformBuffer != nullptr && PageTexture != nullptr && NodeDescBuffer != nullptr; }
-};
 
 /// Global enable/disable for hair voxelization
 bool IsHairStrandsVoxelizationEnable();
@@ -100,15 +19,14 @@ bool IsHairStrandsForVoxelTransmittanceAndShadowEnable();
 void VoxelizeHairStrands(
 	FRDGBuilder& GraphBuilder,
 	const class FScene* Scene,
-	const TArray<FViewInfo>& Views,
-	FInstanceCullingManager& InstanceCullingManager,
-	struct FHairStrandsMacroGroupViews& MacroGroupViews);
+	FViewInfo& View,
+	FInstanceCullingManager& InstanceCullingManager);
 
 
 struct FHairStrandsOcclusionResources
 {
 	FRDGTextureRef CategorizationTexture = nullptr;
 	FRDGTextureRef LightChannelMaskTexture = nullptr;
-	const FVirtualVoxelResources* VoxelResources = nullptr;
+	const FHairStrandsVoxelResources* VoxelResources = nullptr;
 	bool bUseHairVoxel = false;
 };
\ No newline at end of file
diff --git a/Engine/Source/Runtime/Renderer/Private/IndirectLightRendering.cpp b/Engine/Source/Runtime/Renderer/Private/IndirectLightRendering.cpp
--- a/Engine/Source/Runtime/Renderer/Private/IndirectLightRendering.cpp
+++ b/Engine/Source/Runtime/Renderer/Private/IndirectLightRendering.cpp
@@ -708,7 +708,6 @@ void FDeferredShadingSceneRenderer::RenderDiffuseIndirectAndAmbientOcclusion(
 	FRDGBuilder& GraphBuilder,
 	FSceneTextures& SceneTextures,
 	FRDGTextureRef LightingChannelsTexture,
-	FHairStrandsRenderingData* InHairDatas,
 	bool bIsVisualizePass)
 {
 	using namespace HybridIndirectLighting;
@@ -902,14 +901,11 @@ void FDeferredShadingSceneRenderer::RenderDiffuseIndirectAndAmbientOcclusion(
 			SceneTextures.ScreenSpaceAO = AmbientOcclusionMask;
 		}
 
-		if (InHairDatas 
-			&& (   (ViewPipelineState.AmbientOcclusionMethod == EAmbientOcclusionMethod::SSGI) 
-				|| (ViewPipelineState.AmbientOcclusionMethod == EAmbientOcclusionMethod::SSAO)))
+		if (HairStrands::HasViewHairStrandsData(View) && (ViewPipelineState.AmbientOcclusionMethod == EAmbientOcclusionMethod::SSGI || ViewPipelineState.AmbientOcclusionMethod == EAmbientOcclusionMethod::SSAO))
 		{
 			RenderHairStrandsAmbientOcclusion(
 				GraphBuilder,
-				Views,
-				InHairDatas,
+				View,
 				AmbientOcclusionMask);
 		}
 
@@ -1317,8 +1313,7 @@ static void AddSkyReflectionPass(
 void FDeferredShadingSceneRenderer::RenderDeferredReflectionsAndSkyLighting(
 	FRDGBuilder& GraphBuilder,
 	const FSceneTextures& SceneTextures,
-	FRDGTextureRef DynamicBentNormalAOTexture,
-	FHairStrandsRenderingData* HairDatas)
+	FRDGTextureRef DynamicBentNormalAOTexture)
 {
 	if (ViewFamily.EngineShowFlags.VisualizeLightCulling 
 		|| ViewFamily.EngineShowFlags.RayTracingDebug
@@ -1571,58 +1566,44 @@ void FDeferredShadingSceneRenderer::RenderDeferredReflectionsAndSkyLighting(
 			}
 		}
 
-		const bool bIsHairSkyLightingEnabled = HairDatas && (bSkyLight || bDynamicSkyLight || bReflectionEnv);
+		const bool bIsHairSkyLightingEnabled = HairStrands::HasViewHairStrandsData(View) && (bSkyLight || bDynamicSkyLight || bReflectionEnv);
 		if (bIsHairSkyLightingEnabled)
 		{
 			RDG_GPU_STAT_SCOPE(GraphBuilder, HairSkyLighting);
-			RenderHairStrandsEnvironmentLighting(GraphBuilder, Scene, CurrentViewIndex, Views, HairDatas);
+			RenderHairStrandsEnvironmentLighting(GraphBuilder, Scene, View);
 		}
 	}
 
 	AddResolveSceneColorPass(GraphBuilder, Views, SceneColorTexture);
 }
 
-void FDeferredShadingSceneRenderer::RenderDeferredReflectionsAndSkyLightingHair(
-	FRDGBuilder& GraphBuilder,
-	FHairStrandsRenderingData* HairDatas)
+void FDeferredShadingSceneRenderer::RenderDeferredReflectionsAndSkyLightingHair(FRDGBuilder& GraphBuilder)
 {
 	if (ViewFamily.EngineShowFlags.VisualizeLightCulling || !ViewFamily.EngineShowFlags.Lighting)
 	{
 		return;
 	}
 
-	// If we're currently capturing a reflection capture, output SpecularColor * IndirectIrradiance for metals so they are not black in reflections,
-	// Since we don't have multiple bounce specular reflections
-	bool bReflectionCapture = false;
-	for (int32 ViewIndex = 0, Num = Views.Num(); ViewIndex < Num; ViewIndex++)
-	{
-		const FViewInfo& View = Views[ViewIndex];
-		bReflectionCapture = bReflectionCapture || View.bIsReflectionCapture;
-	}
-
-	if (bReflectionCapture)
+	for (FViewInfo& View : Views)
 	{
 		// if we are rendering a reflection capture then we can skip this pass entirely (no reflection and no sky contribution evaluated in this pass)
-		return;
-	}
+		if (View.bIsReflectionCapture)
+		{
+			continue;
+		}
 
-	// The specular sky light contribution is also needed by RT Reflections as a fallback.
-	const bool bSkyLight = Scene->SkyLight
-		&& Scene->SkyLight->ProcessedTexture
-		&& !Scene->SkyLight->bHasStaticLighting;
+		// The specular sky light contribution is also needed by RT Reflections as a fallback.
+		const bool bSkyLight = Scene->SkyLight
+			&& Scene->SkyLight->ProcessedTexture
+			&& !Scene->SkyLight->bHasStaticLighting;
 
-	bool bDynamicSkyLight = ShouldRenderDeferredDynamicSkyLight(Scene, ViewFamily);
-	bool bApplySkyShadowing = false;
-	const bool bReflectionEnv = ShouldDoReflectionEnvironment();
-
-	uint32 ViewIndex = 0;
-	for (FViewInfo& View : Views)
-	{
-		const uint32 CurrentViewIndex = ViewIndex++;
-		const bool bIsHairSkyLightingEnabled = HairDatas && (bSkyLight || bDynamicSkyLight || bReflectionEnv);
+		const bool bDynamicSkyLight = ShouldRenderDeferredDynamicSkyLight(Scene, ViewFamily);
+		const bool bReflectionEnv = ShouldDoReflectionEnvironment();
+		const bool bIsHairSkyLightingEnabled = HairStrands::HasViewHairStrandsData(View) && (bSkyLight || bDynamicSkyLight || bReflectionEnv);
 		if (bIsHairSkyLightingEnabled)
 		{
-			RenderHairStrandsEnvironmentLighting(GraphBuilder, Scene, CurrentViewIndex, Views, HairDatas);
+			RenderHairStrandsEnvironmentLighting(GraphBuilder, Scene, View);
 		}
 	}
+
 }
diff --git a/Engine/Source/Runtime/Renderer/Private/LightGridInjection.cpp b/Engine/Source/Runtime/Renderer/Private/LightGridInjection.cpp
--- a/Engine/Source/Runtime/Renderer/Private/LightGridInjection.cpp
+++ b/Engine/Source/Runtime/Renderer/Private/LightGridInjection.cpp
@@ -786,10 +786,9 @@ void FDeferredShadingSceneRenderer::RenderForwardShadowProjections(
 	FRDGBuilder& GraphBuilder,
 	const FMinimalSceneTextures& SceneTextures,
 	FRDGTextureRef& OutForwardScreenSpaceShadowMask,
-	FRDGTextureRef& OutForwardScreenSpaceShadowMaskSubPixel,
-	const FHairStrandsRenderingData* InHairDatas)
+	FRDGTextureRef& OutForwardScreenSpaceShadowMaskSubPixel)
 {
-	const bool bIsHairEnable = InHairDatas != nullptr;
+	const bool bIsHairEnable = HairStrands::HasViewHairStrandsData(Views);
 	bool bScreenShadowMaskNeeded = false;
 
 	FRDGTextureRef SceneDepthTexture = SceneTextures.Depth.Target;
@@ -852,12 +851,11 @@ void FDeferredShadingSceneRenderer::RenderForwardShadowProjections(
 					ForwardScreenSpaceShadowMask.Target,
 					ForwardScreenSpaceShadowMaskSubPixel.Target,
 					LightSceneInfo,
-					InHairDatas,
 					bProjectingForForwardShading);
 
-				if (InHairDatas)
+				if (bIsHairEnable)
 				{
-					RenderHairStrandsShadowMask(GraphBuilder, Views, LightSceneInfo, InHairDatas, ForwardScreenSpaceShadowMask.Target);
+					RenderHairStrandsShadowMask(GraphBuilder, Views, LightSceneInfo, ForwardScreenSpaceShadowMask.Target);
 				}
 			}
 
diff --git a/Engine/Source/Runtime/Renderer/Private/LightRendering.cpp b/Engine/Source/Runtime/Renderer/Private/LightRendering.cpp
--- a/Engine/Source/Runtime/Renderer/Private/LightRendering.cpp
+++ b/Engine/Source/Runtime/Renderer/Private/LightRendering.cpp
@@ -22,6 +22,7 @@
 #include "VolumetricCloudRendering.h"
 #include "Strata/Strata.h"
 #include "VirtualShadowMaps/VirtualShadowMapProjection.h"
+#include "HairStrands/HairStrandsData.h"
 
 // ENABLE_DEBUG_DISCARD_PROP is used to test the lighting code by allowing to discard lights to see how performance scales
 // It ought never to be enabled in a shipping build, and is probably only really useful when woring on the shading code.
@@ -347,14 +348,6 @@ struct FRenderLightParams
 	// Precompute transmittance
 	FShaderResourceViewRHIRef DeepShadow_TransmittanceMaskBuffer = nullptr;
 	uint32 DeepShadow_TransmittanceMaskBufferMaxCount = 0;
-	
-	// Visibility buffer data
-	FRHITexture* HairCategorizationTexture = nullptr;
-	FRHITexture* HairVisibilityNodeOffsetAndCount = nullptr;
-	FRHITexture* HairVisibilityNodeCount = nullptr;
-	FShaderResourceViewRHIRef HairVisibilityNodeCoordsSRV = nullptr;
-	FShaderResourceViewRHIRef HairVisibilityNodeDataSRV = nullptr;
-
 	FRHITexture* ScreenShadowMaskSubPixelTexture = nullptr;
 
 	// Cloud shadow data
@@ -385,31 +378,22 @@ public:
 	TDeferredLightHairVS(const ShaderMetaType::CompiledShaderInitializerType& Initializer) :
 		FGlobalShader(Initializer)
 	{
-		MaxViewportResolution.Bind(Initializer.ParameterMap, TEXT("MaxViewportResolution"));
-		HairVisibilityNodeCount.Bind(Initializer.ParameterMap, TEXT("HairVisibilityNodeCount"));
+		HairStrandsParameters.Bind(Initializer.ParameterMap, FHairStrandsViewUniformParameters::StaticStructMetadata.GetShaderVariableName());
 	}
 
-	void SetParameters(FRHICommandList& RHICmdList, const FViewInfo& View, const FHairStrandsVisibilityData* VisibilityData)
+	void SetParameters(FRHICommandList& RHICmdList, const FViewInfo& View, FRHIUniformBuffer* HairStrandsUniformBuffer)
 	{
 		FRHIVertexShader* ShaderRHI = RHICmdList.GetBoundVertexShader();
 		FGlobalShader::SetParameters<FViewUniformShaderParameters>(RHICmdList, ShaderRHI, View.ViewUniformBuffer);
 
-		if (!VisibilityData)
-		{
-			return;
-		}
-
-		if (HairVisibilityNodeCount.IsBound() && VisibilityData->NodeCount)
+		if (HairStrandsUniformBuffer)
 		{
-			SetTextureParameter(RHICmdList,	ShaderRHI, HairVisibilityNodeCount, TryGetRHI(VisibilityData->NodeCount));
+			SetUniformBufferParameter(RHICmdList, ShaderRHI, HairStrandsParameters, HairStrandsUniformBuffer);
 		}
-
-		SetShaderValue(RHICmdList, ShaderRHI, MaxViewportResolution, VisibilityData->SampleLightingViewportResolution);
 	}
 
 private:
-	LAYOUT_FIELD(FShaderParameter, MaxViewportResolution);
-	LAYOUT_FIELD(FShaderResourceParameter, HairVisibilityNodeCount);
+	LAYOUT_FIELD(FShaderUniformBufferParameter, HairStrandsParameters);
 };
 
 IMPLEMENT_SHADER_TYPE(, TDeferredLightHairVS, TEXT("/Engine/Private/DeferredLightVertexShaders.usf"), TEXT("HairVertexMain"), SF_Vertex);
@@ -550,12 +534,7 @@ class FDeferredLightPS : public FGlobalShader
 		ScreenShadowMaskSubPixelTexture.Bind(Initializer.ParameterMap, TEXT("ScreenShadowMaskSubPixelTexture")); // TODO hook the shader itself
 
 		HairShadowMaskValid.Bind(Initializer.ParameterMap, TEXT("HairShadowMaskValid"));
-		HairDualScatteringRoughnessOverride.Bind(Initializer.ParameterMap, TEXT("HairDualScatteringRoughnessOverride"));
-
-		HairCategorizationTexture.Bind(Initializer.ParameterMap, TEXT("HairCategorizationTexture"));
-		HairVisibilityNodeOffsetAndCount.Bind(Initializer.ParameterMap, TEXT("HairVisibilityNodeOffsetAndCount"));
-		HairVisibilityNodeCoords.Bind(Initializer.ParameterMap, TEXT("HairVisibilityNodeCoords"));
-		HairVisibilityNodeData.Bind(Initializer.ParameterMap, TEXT("HairVisibilityNodeData"));
+		HairStrandsParameters.Bind(Initializer.ParameterMap, FHairStrandsViewUniformParameters::StaticStructMetadata.GetShaderVariableName());
 
 		DummyRectLightTextureForCapsuleCompilerWarning.Bind(Initializer.ParameterMap, TEXT("DummyRectLightTextureForCapsuleCompilerWarning"));
 
@@ -576,11 +555,16 @@ public:
 		const FLightSceneInfo* LightSceneInfo,
 		FRHITexture* ScreenShadowMaskTexture,
 		FRHITexture* LightingChannelsTextureRHI,
-		FRenderLightParams* RenderLightParams)
+		FRenderLightParams* RenderLightParams,
+		FRHIUniformBuffer* HairStrandsUniformBuffer)
 	{
 		FRHIPixelShader* ShaderRHI = RHICmdList.GetBoundPixelShader();
 		SetParametersBase(RHICmdList, ShaderRHI, View, ScreenShadowMaskTexture, LightingChannelsTextureRHI, LightSceneInfo->Proxy->GetIESTextureResource(), RenderLightParams);
 		SetDeferredLightParameters(RHICmdList, ShaderRHI, GetUniformBufferParameter<FDeferredLightUniformStruct>(), LightSceneInfo, View);
+		if (HairStrandsUniformBuffer)
+		{
+			SetUniformBufferParameter(RHICmdList, ShaderRHI, HairStrandsParameters, HairStrandsUniformBuffer);
+		}
 	}
 
 	void SetParametersSimpleLight(FRHICommandList& RHICmdList, const FViewInfo& View, const FSimpleLightEntry& SimpleLight, const FSimpleLightPerViewEntry& SimpleLightPerViewData)
@@ -716,70 +700,6 @@ private:
 			}
 		}
 
-		if (HairCategorizationTexture.IsBound())
-		{
-			if (RenderLightParams && RenderLightParams->HairCategorizationTexture)
-			{
-				SetTextureParameter(
-					RHICmdList,
-					ShaderRHI,
-					HairCategorizationTexture,
-					LightAttenuationTextureSampler,
-					TStaticSamplerState<SF_Point, AM_Clamp, AM_Clamp, AM_Clamp>::GetRHI(),
-					RenderLightParams->HairCategorizationTexture);
-			}
-		}
-
-		if (HairVisibilityNodeOffsetAndCount.IsBound())
-		{
-			if (RenderLightParams && RenderLightParams->HairVisibilityNodeOffsetAndCount)
-			{
-				SetTextureParameter(
-					RHICmdList,
-					ShaderRHI,
-					HairVisibilityNodeOffsetAndCount,
-					LightAttenuationTextureSampler,
-					TStaticSamplerState<SF_Point, AM_Clamp, AM_Clamp, AM_Clamp>::GetRHI(),
-					RenderLightParams->HairVisibilityNodeOffsetAndCount);
-			}
-		}
-		
-		if (HairVisibilityNodeCoords.IsBound())
-		{
-			if (RenderLightParams && RenderLightParams->HairVisibilityNodeCoordsSRV)
-			{
-				FShaderResourceViewRHIRef SRV = RenderLightParams->HairVisibilityNodeCoordsSRV;
-				SetSRVParameter(
-					RHICmdList, 
-					ShaderRHI, 
-					HairVisibilityNodeCoords,
-					SRV);
-			}
-		}
-
-		if (HairVisibilityNodeData.IsBound())
-		{
-			if (RenderLightParams && RenderLightParams->HairVisibilityNodeDataSRV)
-			{
-				FShaderResourceViewRHIRef SRV = RenderLightParams->HairVisibilityNodeDataSRV;
-				SetSRVParameter(
-					RHICmdList, 
-					ShaderRHI, 
-					HairVisibilityNodeData, 
-					SRV);
-			}
-		}
-
-		if (HairDualScatteringRoughnessOverride.IsBound())
-		{
-			const float DualScatteringRoughness = GetHairDualScatteringRoughnessOverride();
-			SetShaderValue(
-				RHICmdList,
-				ShaderRHI,
-				HairDualScatteringRoughnessOverride,
-				DualScatteringRoughness);
-		}
-
 		if (DummyRectLightTextureForCapsuleCompilerWarning.IsBound())
 		{
 			SetTextureParameter(
@@ -840,14 +760,9 @@ private:
 
 	LAYOUT_FIELD(FShaderParameter, HairTransmittanceBufferMaxCount);
 	LAYOUT_FIELD(FShaderResourceParameter, HairTransmittanceBuffer);
-	LAYOUT_FIELD(FShaderResourceParameter, HairCategorizationTexture);
-	LAYOUT_FIELD(FShaderResourceParameter, HairVisibilityNodeOffsetAndCount);
-	LAYOUT_FIELD(FShaderResourceParameter, HairVisibilityNodeCoords);
-	LAYOUT_FIELD(FShaderResourceParameter, HairVisibilityNodeData);
+	LAYOUT_FIELD(FShaderUniformBufferParameter, HairStrandsParameters);
 	LAYOUT_FIELD(FShaderResourceParameter, ScreenShadowMaskSubPixelTexture);
-
 	LAYOUT_FIELD(FShaderParameter, HairShadowMaskValid);
-	LAYOUT_FIELD(FShaderParameter, HairDualScatteringRoughnessOverride);
 
 	LAYOUT_FIELD(FShaderResourceParameter, DummyRectLightTextureForCapsuleCompilerWarning);
 
@@ -1162,26 +1077,21 @@ void FSceneRenderer::GatherAndSortLights(FSortedLightSetSceneInfo& OutSortedLigh
 	check(OutSortedLights.AttenuationLightStart >= OutSortedLights.ClusteredSupportedEnd);
 }
 
-static bool HasHairStrandsClusters(int32 ViewIndex, const FHairStrandsRenderingData* HairDatas)
-{
-	return HairDatas && ViewIndex < HairDatas->MacroGroupsPerViews.Views.Num() && HairDatas->MacroGroupsPerViews.Views[ViewIndex].Datas.Num() > 0;
-};
-
-static FHairStrandsOcclusionResources GetHairStrandsResources(int32 ViewIndex, FRDGBuilder& GraphBuilder, const FHairStrandsRenderingData* HairDatas)
+static FHairStrandsOcclusionResources GetHairStrandsResources(FRDGBuilder& GraphBuilder, const FViewInfo& View)
 {
 	FHairStrandsOcclusionResources Out;
-	if (HairDatas && ViewIndex < HairDatas->HairVisibilityViews.HairDatas.Num())
+	if (HairStrands::HasViewHairStrandsData(View))
 	{
-		if (HairDatas->HairVisibilityViews.HairDatas[ViewIndex].CategorizationTexture)
+		if (FRDGTextureRef Texture = View.HairStrandsViewData.VisibilityData.CategorizationTexture)
 		{
-			Out.CategorizationTexture = HairDatas->HairVisibilityViews.HairDatas[ViewIndex].CategorizationTexture;
+			Out.CategorizationTexture = Texture;
 		}
-		if (HairDatas->HairVisibilityViews.HairDatas[ViewIndex].LightChannelMaskTexture)
+		if (FRDGTextureRef Texture = View.HairStrandsViewData.VisibilityData.LightChannelMaskTexture)
 		{
-			Out.LightChannelMaskTexture = HairDatas->HairVisibilityViews.HairDatas[ViewIndex].LightChannelMaskTexture;
+			Out.LightChannelMaskTexture = Texture;
 		}
 
-		Out.VoxelResources = &HairDatas->MacroGroupsPerViews.Views[ViewIndex].VirtualVoxelResources;
+		Out.VoxelResources = &View.HairStrandsViewData.VirtualVoxelResources;
 	}
 	return Out;
 }
@@ -1189,8 +1099,8 @@ static FHairStrandsOcclusionResources GetHairStrandsResources(int32 ViewIndex, F
 /** Shader parameters to use when creating a RenderLight(...) pass. */
 BEGIN_SHADER_PARAMETER_STRUCT(FRenderLightParameters, )
 	SHADER_PARAMETER_RDG_UNIFORM_BUFFER(FSceneTextureUniformParameters, SceneTextures)
+	SHADER_PARAMETER_RDG_UNIFORM_BUFFER(FHairStrandsViewUniformParameters, HairStrands)
 	SHADER_PARAMETER_STRUCT_INCLUDE(FVolumetricCloudShadowAOParameters, CloudShadowAO)
-	RDG_TEXTURE_ACCESS(HairCategorizationTexture, ERHIAccess::SRVGraphics)
 	RDG_TEXTURE_ACCESS(ShadowMaskTexture, ERHIAccess::SRVGraphics)
 	RDG_TEXTURE_ACCESS(LightingChannelsTexture, ERHIAccess::SRVGraphics)
 	RENDER_TARGET_BINDING_SLOTS()
@@ -1200,17 +1110,17 @@ void GetRenderLightParameters(
 	FRDGTextureRef SceneColorTexture,
 	FRDGTextureRef SceneDepthTexture,
 	TRDGUniformBufferRef<FSceneTextureUniformParameters> SceneTexturesUniformBuffer,
+	TRDGUniformBufferRef<FHairStrandsViewUniformParameters> HairStrandsUniformBuffer,
 	FRDGTextureRef ShadowMaskTexture,
 	FRDGTextureRef LightingChannelsTexture,
 	const FVolumetricCloudShadowAOParameters& CloudShadowAOParameters,
-	const FHairStrandsVisibilityViews* InHairVisibilityViews,
 	FRenderLightParameters& Parameters)
 {
 	Parameters.SceneTextures = SceneTexturesUniformBuffer;
+	Parameters.HairStrands = HairStrandsUniformBuffer;
 	Parameters.ShadowMaskTexture = ShadowMaskTexture;
 	Parameters.LightingChannelsTexture = LightingChannelsTexture;
 	Parameters.CloudShadowAO = CloudShadowAOParameters;
-	Parameters.HairCategorizationTexture = InHairVisibilityViews && InHairVisibilityViews->HairDatas.Num() > 0 ? InHairVisibilityViews->HairDatas[0].CategorizationTexture : nullptr;
 	Parameters.RenderTargets[0] = FRenderTargetBinding(SceneColorTexture, ERenderTargetLoadAction::ELoad);
 
 	if (SceneDepthTexture)
@@ -1223,13 +1133,13 @@ FHairStrandsTransmittanceMaskData CreateDummyHairStrandsTransmittanceMaskData(FR
 
 void GetRenderLightParameters(
 	const FMinimalSceneTextures& SceneTextures,
+	const FHairStrandsViewData& HairViewData,
 	FRDGTextureRef ShadowMaskTexture,
 	FRDGTextureRef LightingChannelsTexture,
 	const FVolumetricCloudShadowAOParameters& CloudShadowAOParameters,
-	const FHairStrandsVisibilityViews* InHairVisibilityViews,
 	FRenderLightParameters& Parameters)
 {
-	GetRenderLightParameters(SceneTextures.Color.Target, SceneTextures.Depth.Target, SceneTextures.UniformBuffer, ShadowMaskTexture, LightingChannelsTexture, CloudShadowAOParameters, InHairVisibilityViews, Parameters);
+	GetRenderLightParameters(SceneTextures.Color.Target, SceneTextures.Depth.Target, SceneTextures.UniformBuffer, HairViewData.UniformBuffer, ShadowMaskTexture, LightingChannelsTexture, CloudShadowAOParameters, Parameters);
 }
 
 void FDeferredShadingSceneRenderer::RenderLights(
@@ -1237,13 +1147,9 @@ void FDeferredShadingSceneRenderer::RenderLights(
 	FMinimalSceneTextures& SceneTextures,
 	const FTranslucencyLightingVolumeTextures& TranslucencyLightingVolumeTextures,
 	FRDGTextureRef LightingChannelsTexture,
-	FSortedLightSetSceneInfo &SortedLightSet,
-	const FHairStrandsRenderingData* HairDatas)
+	FSortedLightSetSceneInfo& SortedLightSet)
 {
-	const EShaderPlatform ShaderPlatformForFeatureLevel = GShaderPlatformForFeatureLevel[FeatureLevel];
-
-	const bool bUseHairLighting = HairDatas != nullptr && HairDatas->HairVisibilityViews.HairDatas.Num() > 0 && HairDatas->HairVisibilityViews.HairDatas[0].CategorizationTexture;
-	const FHairStrandsVisibilityViews* InHairVisibilityViews = bUseHairLighting ? &HairDatas->HairVisibilityViews : nullptr;
+	const bool bUseHairLighting = HairStrands::HasViewHairStrandsData(Views);
 
 	RDG_EVENT_SCOPE(GraphBuilder, "Lights");
 	RDG_GPU_STAT_SCOPE(GraphBuilder, Lights);
@@ -1258,7 +1164,7 @@ void FDeferredShadingSceneRenderer::RenderLights(
 	const int32 SimpleLightsEnd = SortedLightSet.SimpleLightsEnd;
 
 	FHairStrandsTransmittanceMaskData DummyTransmittanceMaskData;
-	if (bUseHairLighting && Views.Num() > 0)
+	if (HairStrands::HasViewHairStrandsData(Views) && Views.Num() > 0)
 	{
 		DummyTransmittanceMaskData = CreateDummyHairStrandsTransmittanceMaskData(GraphBuilder, Views[0].ShaderMap);
 	}
@@ -1345,43 +1251,50 @@ void FDeferredShadingSceneRenderer::RenderLights(
 			}
 
 			{
-				FRenderLightParameters* PassParameters = GraphBuilder.AllocParameters<FRenderLightParameters>();
-				GetRenderLightParameters(SceneTextures, nullptr, LightingChannelsTexture, {}, HairDatas ? &HairDatas->HairVisibilityViews : nullptr, *PassParameters);
-
-				GraphBuilder.AddPass(
-					RDG_EVENT_NAME("StandardDeferredLighting"),
-					PassParameters,
-					ERDGPassFlags::Raster,
-					[this, &SortedLights, LightingChannelsTexture, StandardDeferredStart, AttenuationLightStart](FRHICommandList& RHICmdList)
+				for (int32 ViewIndex = 0, ViewCount = Views.Num(); ViewIndex < ViewCount; ++ViewIndex)
 				{
-					// Draw non-shadowed non-light function lights without changing render targets between them
-					for (int32 LightIndex = StandardDeferredStart; LightIndex < AttenuationLightStart; LightIndex++)
+					const FViewInfo& View = Views[ViewIndex];
+					FRenderLightParameters* PassParameters = GraphBuilder.AllocParameters<FRenderLightParameters>();
+					GetRenderLightParameters(SceneTextures, View.HairStrandsViewData, nullptr, LightingChannelsTexture, {}, *PassParameters);
+
+					GraphBuilder.AddPass(
+						RDG_EVENT_NAME("StandardDeferredLighting"),
+						PassParameters,
+						ERDGPassFlags::Raster,
+						[this, &View, &SortedLights, LightingChannelsTexture, StandardDeferredStart, AttenuationLightStart](FRHICommandList& RHICmdList)
 					{
-						const FSortedLightSceneInfo& SortedLightInfo = SortedLights[LightIndex];
-						const FLightSceneInfo* const LightSceneInfo = SortedLightInfo.LightSceneInfo;
-
-						for (int32 ViewIndex = 0, ViewCount = Views.Num(); ViewIndex < ViewCount; ++ViewIndex)
+						// Draw non-shadowed non-light function lights without changing render targets between them
+						for (int32 LightIndex = StandardDeferredStart; LightIndex < AttenuationLightStart; LightIndex++)
 						{
-							const FViewInfo& View = Views[ViewIndex];
+							const FSortedLightSceneInfo& SortedLightInfo = SortedLights[LightIndex];
+							const FLightSceneInfo* const LightSceneInfo = SortedLightInfo.LightSceneInfo;
+
+							
 							SCOPED_GPU_MASK(RHICmdList, View.GPUMask);
 
 							// Render the light to the scene color buffer, using a 1x1 white texture as input
-							RenderLight(RHICmdList, View, ViewIndex, LightSceneInfo, nullptr, TryGetRHI(LightingChannelsTexture), nullptr, false, false);
+							RenderLight(RHICmdList, View, LightSceneInfo, nullptr, TryGetRHI(LightingChannelsTexture), false, false);
 						}
-					}
-				});
+					});
+				}
 			}
 
 			// Add a special version when hair rendering is enabled for getting lighting on hair. 
 			if (bUseHairLighting)
 			{
 				FRDGTextureRef NullScreenShadowMaskSubPixelTexture = nullptr;
-				// Draw non-shadowed non-light function lights without changing render targets between them
-				for (int32 LightIndex = StandardDeferredStart; LightIndex < AttenuationLightStart; LightIndex++)
+				for (FViewInfo& View : Views)
 				{
-					const FSortedLightSceneInfo& SortedLightInfo = SortedLights[LightIndex];
-					const FLightSceneInfo* const LightSceneInfo = SortedLightInfo.LightSceneInfo;
-					RenderLightForHair(GraphBuilder, SceneTextures.UniformBuffer, LightSceneInfo, NullScreenShadowMaskSubPixelTexture, LightingChannelsTexture, DummyTransmittanceMaskData, InHairVisibilityViews);
+					if (HairStrands::HasViewHairStrandsData(View))
+					{
+						// Draw non-shadowed non-light function lights without changing render targets between them
+						for (int32 LightIndex = StandardDeferredStart; LightIndex < AttenuationLightStart; LightIndex++)
+						{
+							const FSortedLightSceneInfo& SortedLightInfo = SortedLights[LightIndex];
+							const FLightSceneInfo* const LightSceneInfo = SortedLightInfo.LightSceneInfo;
+							RenderLightForHair(GraphBuilder, View, SceneTextures.UniformBuffer, LightSceneInfo, NullScreenShadowMaskSubPixelTexture, LightingChannelsTexture, DummyTransmittanceMaskData);
+						}
+					}
 				}
 			}
 
@@ -1455,7 +1368,7 @@ void FDeferredShadingSceneRenderer::RenderLights(
 
 				PreprocessedShadowMaskTextures.SetNum(SortedLights.Num());
 
-				if (HasHairStrandsClusters(ViewIndex, HairDatas))
+				if (HairStrands::HasViewHairStrandsData(View))
 				{ 
 					PreprocessedShadowMaskSubPixelTextures.SetNum(SortedLights.Num());
 				}
@@ -1655,7 +1568,7 @@ void FDeferredShadingSceneRenderer::RenderLights(
 
 								FRDGTextureUAV* RayTracingShadowMaskUAV = GraphBuilder.CreateUAV(FRDGTextureUAVDesc(RayTracingShadowMaskTexture));
 								FRDGTextureUAV* RayHitDistanceUAV = GraphBuilder.CreateUAV(FRDGTextureUAVDesc(RayDistanceTexture));
-								FHairStrandsOcclusionResources HairResources = GetHairStrandsResources(ViewIndex, GraphBuilder, HairDatas);
+								FHairStrandsOcclusionResources HairResources = GetHairStrandsResources(GraphBuilder, View);
 								HairResources.bUseHairVoxel = !BatchLightSceneInfo.Proxy->CastsHairStrandsDeepShadow();
 								{
 									RDG_EVENT_SCOPE(GraphBuilder, "%s", *BatchLightNameWithLevel);
@@ -1675,7 +1588,7 @@ void FDeferredShadingSceneRenderer::RenderLights(
 										RayHitDistanceUAV,
 										SubPixelRayTracingShadowMaskUAV);
 									
-									if (HasHairStrandsClusters(ViewIndex, HairDatas))
+									if (HairStrands::HasViewHairStrandsData(View))
 									{
 										FRDGTextureRef& RefDestination = PreprocessedShadowMaskSubPixelTextures[LightBatchIndex - AttenuationLightStart];
 										check(RefDestination == nullptr);
@@ -1749,7 +1662,7 @@ void FDeferredShadingSceneRenderer::RenderLights(
 						// Inject deep shadow mask if the light supports it
 						if (bUseHairDeepShadow)
 						{
-							RenderHairStrandsShadowMask(GraphBuilder, Views, &LightSceneInfo, HairDatas, ScreenShadowMaskTexture);
+							RenderHairStrandsShadowMask(GraphBuilder, Views, &LightSceneInfo, ScreenShadowMaskTexture);
 						}
 					}
 					else if (OcclusionType == FLightOcclusionType::Raytraced)
@@ -1805,9 +1718,8 @@ void FDeferredShadingSceneRenderer::RenderLights(
 
 						bool bIsMultiview = Views.Num() > 0;
 
-						for (int32 ViewIndex = 0; ViewIndex < Views.Num(); ViewIndex++)
+						for (FViewInfo& View : Views)
 						{
-							FViewInfo& View = Views[ViewIndex];
 							RDG_GPU_MASK_SCOPE(GraphBuilder, View.GPUMask);
 
 							IScreenSpaceDenoiser::FShadowRayTracingConfig RayTracingConfig;
@@ -1819,7 +1731,7 @@ void FDeferredShadingSceneRenderer::RenderLights(
 								DenoiserRequirements = DenoiserToUse->GetShadowRequirements(View, LightSceneInfo, RayTracingConfig);
 							}
 
-							FHairStrandsOcclusionResources HairResources = GetHairStrandsResources(ViewIndex, GraphBuilder, HairDatas);
+							FHairStrandsOcclusionResources HairResources = GetHairStrandsResources(GraphBuilder, View);
 							HairResources.bUseHairVoxel = !bUseHairDeepShadow;
 
 							RenderRayTracingShadows(
@@ -1877,16 +1789,16 @@ void FDeferredShadingSceneRenderer::RenderLights(
 								ScreenShadowMaskTexture = RayTracingShadowMaskTexture;
 							}
 
-							if (HasHairStrandsClusters(ViewIndex, HairDatas))
+							if (HairStrands::HasViewHairStrandsData(View))
 							{
 								ScreenShadowMaskSubPixelTexture = SubPixelRayTracingShadowMaskTexture;
 							}
 						}
 
 						// Inject deep shadow mask if the light supports it
-						if (HairDatas && bUseHairDeepShadow)
+						if (bUseHairDeepShadow)
 						{
-							RenderHairStrandsShadowMask(GraphBuilder, Views, &LightSceneInfo, HairDatas, ScreenShadowMaskTexture);
+							RenderHairStrandsShadowMask(GraphBuilder, Views, &LightSceneInfo, ScreenShadowMaskTexture);
 						}
 					}
 					else // (OcclusionType == FOcclusionType::Shadowmap)
@@ -1949,7 +1861,7 @@ void FDeferredShadingSceneRenderer::RenderLights(
 							ClearShadowMask(ScreenShadowMaskSubPixelTexture);
 						}
 
-						RenderDeferredShadowProjections(GraphBuilder, SceneTextures, TranslucencyLightingVolumeTextures, &LightSceneInfo, ScreenShadowMaskTexture, ScreenShadowMaskSubPixelTexture, HairDatas, bInjectedTranslucentVolume);
+						RenderDeferredShadowProjections(GraphBuilder, SceneTextures, TranslucencyLightingVolumeTextures, &LightSceneInfo, ScreenShadowMaskTexture, ScreenShadowMaskSubPixelTexture, bInjectedTranslucentVolume);
 					}
 
 					bUsedShadowMaskTexture = true;
@@ -2005,19 +1917,22 @@ void FDeferredShadingSceneRenderer::RenderLights(
 				if (bDirectLighting)
 				{
 					const bool bRenderOverlap = false;
-					RenderLight(GraphBuilder, SceneTextures, &LightSceneInfo, ScreenShadowMaskTexture, LightingChannelsTexture, InHairVisibilityViews, bRenderOverlap);
+					RenderLight(GraphBuilder, SceneTextures, &LightSceneInfo, ScreenShadowMaskTexture, LightingChannelsTexture, bRenderOverlap);
 				}
 
 				if (bUseHairLighting)
 				{
-					FHairStrandsTransmittanceMaskData TransmittanceMaskData = DummyTransmittanceMaskData;
-					if (bDrawHairShadow)
+					for (FViewInfo& View : Views)
 					{
-						TransmittanceMaskData = RenderHairStrandsTransmittanceMask(GraphBuilder, Views, &LightSceneInfo, HairDatas, ScreenShadowMaskSubPixelTexture);
+						FHairStrandsTransmittanceMaskData TransmittanceMaskData = DummyTransmittanceMaskData;
+						if (bDrawHairShadow && HairStrands::HasViewHairStrandsData(View))
+						{
+							TransmittanceMaskData = RenderHairStrandsTransmittanceMask(GraphBuilder, View, &LightSceneInfo, ScreenShadowMaskSubPixelTexture);
 
-						// Note: ideally the light should still be evaluated for hair when not casting shadow, but for preserving the old behavior, and not adding 
-						// any perf. regression, we disable this light for hair rendering 
-						RenderLightForHair(GraphBuilder, SceneTextures.UniformBuffer, &LightSceneInfo, ScreenShadowMaskSubPixelTexture, LightingChannelsTexture, TransmittanceMaskData, InHairVisibilityViews);
+							// Note: ideally the light should still be evaluated for hair when not casting shadow, but for preserving the old behavior, and not adding 
+							// any perf. regression, we disable this light for hair rendering 
+							RenderLightForHair(GraphBuilder, View, SceneTextures.UniformBuffer, &LightSceneInfo, ScreenShadowMaskSubPixelTexture, LightingChannelsTexture, TransmittanceMaskData);
+						}
 					}
 				}
 			}
@@ -2050,12 +1965,10 @@ void FDeferredShadingSceneRenderer::RenderLightArrayForOverlapViewmode(
 		}
 
 		// Check if the light is visible in any of the views.
-		for (int32 ViewIndex = 0; ViewIndex < Views.Num(); ViewIndex++)
+		for (const FViewInfo& View : Views)
 		{
-			const FViewInfo& View = Views[ViewIndex];
 			SCOPED_GPU_MASK(RHICmdList, View.GPUMask);
-
-			RenderLight(RHICmdList, View, ViewIndex, LightSceneInfo, nullptr, LightingChannelsTexture, nullptr, true, false);
+			RenderLight(RHICmdList, View, LightSceneInfo, nullptr, LightingChannelsTexture, true, false);
 		}
 	}
 }
@@ -2068,7 +1981,7 @@ void FDeferredShadingSceneRenderer::RenderStationaryLightOverlap(
 	if (Scene->bIsEditorScene)
 	{
 		FRenderLightParameters* PassParameters = GraphBuilder.AllocParameters<FRenderLightParameters>();
-		GetRenderLightParameters(SceneTextures, nullptr, LightingChannelsTexture, {}, nullptr, *PassParameters);
+		GetRenderLightParameters(SceneTextures, Views[0].HairStrandsViewData, nullptr, LightingChannelsTexture, {}, *PassParameters);
 
 		GraphBuilder.AddPass(
 			RDG_EVENT_NAME("StationaryLightOverlap"),
@@ -2200,11 +2113,9 @@ void CalculateLightNearFarDepthFromBounds(const FViewInfo& View, const FSphere &
 void FDeferredShadingSceneRenderer::RenderLight(
 	FRHICommandList& RHICmdList,
 	const FViewInfo& View,
-	int32 ViewIndex,
 	const FLightSceneInfo* LightSceneInfo,
 	FRHITexture* ScreenShadowMaskTexture,
 	FRHITexture* LightingChannelsTexture,
-	const FHairStrandsVisibilityViews* InHairVisibilityViews,
 	bool bRenderOverlap, bool bIssueDrawEvent)
 {
 	// Ensure the light is valid for this view
@@ -2242,12 +2153,6 @@ void FDeferredShadingSceneRenderer::RenderLight(
 	RHICmdList.SetViewport(View.ViewRect.Min.X, View.ViewRect.Min.Y, 0.0f, View.ViewRect.Max.X, View.ViewRect.Max.Y, 1.0f);
 
 	FRenderLightParams RenderLightParams;
-	const bool bHairLighting = InHairVisibilityViews && ViewIndex < InHairVisibilityViews->HairDatas.Num() && InHairVisibilityViews->HairDatas[ViewIndex].CategorizationTexture != nullptr;
-	if (bHairLighting)
-	{
-		RenderLightParams.HairCategorizationTexture = InHairVisibilityViews->HairDatas[ViewIndex].CategorizationTexture->GetRHI();
-	}
-
 	if (Strata::IsStrataEnabled() && Strata::IsClassificationEnabled())
 	{
 		GraphicsPSOInit.DepthStencilState = TStaticDepthStencilState<
@@ -2325,7 +2230,7 @@ void FDeferredShadingSceneRenderer::RenderLight(
 			GraphicsPSOInit.BoundShaderState.PixelShaderRHI = PixelShader.GetPixelShader();
 
 			SetGraphicsPipelineState(RHICmdList, GraphicsPSOInit);
-			PixelShader->SetParameters(RHICmdList, View, LightSceneInfo, ScreenShadowMaskTexture, LightingChannelsTexture, &RenderLightParams);
+			PixelShader->SetParameters(RHICmdList, View, LightSceneInfo, ScreenShadowMaskTexture, LightingChannelsTexture, &RenderLightParams, nullptr);
 		}
 
 		VertexShader->SetParameters(RHICmdList, View, LightSceneInfo);
@@ -2391,7 +2296,7 @@ void FDeferredShadingSceneRenderer::RenderLight(
 			GraphicsPSOInit.BoundShaderState.PixelShaderRHI = PixelShader.GetPixelShader();
 
 			SetGraphicsPipelineState(RHICmdList, GraphicsPSOInit);
-			PixelShader->SetParameters(RHICmdList, View, LightSceneInfo, ScreenShadowMaskTexture, LightingChannelsTexture, (bHairLighting) ? &RenderLightParams : nullptr);
+			PixelShader->SetParameters(RHICmdList, View, LightSceneInfo, ScreenShadowMaskTexture, LightingChannelsTexture, &RenderLightParams, nullptr);
 		}
 
 		VertexShader->SetParameters(RHICmdList, View, LightSceneInfo);
@@ -2443,7 +2348,6 @@ void FDeferredShadingSceneRenderer::RenderLight(
 	const FLightSceneInfo* LightSceneInfo,
 	FRDGTextureRef ScreenShadowMaskTexture,
 	FRDGTextureRef LightingChannelsTexture,
-	const FHairStrandsVisibilityViews* InHairVisibilityViews,
 	bool bRenderOverlap)
 {
 	ERDGPassFlags PassFlags = ERDGPassFlags::Raster;
@@ -2456,22 +2360,20 @@ void FDeferredShadingSceneRenderer::RenderLight(
 		const FViewInfo& View = Views[ViewIndex];
 
 		FRenderLightParameters* PassParameters = GraphBuilder.AllocParameters<FRenderLightParameters>();
-		GetRenderLightParameters(SceneTextures, ScreenShadowMaskTexture, LightingChannelsTexture, GetCloudShadowAOParameters(GraphBuilder, View, CloudInfo), InHairVisibilityViews, *PassParameters);
+		GetRenderLightParameters(SceneTextures, View.HairStrandsViewData, ScreenShadowMaskTexture, LightingChannelsTexture, GetCloudShadowAOParameters(GraphBuilder, View, CloudInfo), *PassParameters);
 
 		GraphBuilder.AddPass(
 			RDG_EVENT_NAME("StandardDeferredLighting"),
 			PassParameters,
 			PassFlags,
-			[this, &View, ViewIndex, LightSceneInfo, ScreenShadowMaskTexture, LightingChannelsTexture, InHairVisibilityViews, bRenderOverlap](FRHICommandList& RHICmdList)
+			[this, &View, LightSceneInfo, ScreenShadowMaskTexture, LightingChannelsTexture, bRenderOverlap](FRHICommandList& RHICmdList)
 		{
 			RenderLight(
 				RHICmdList,
 				View,
-				ViewIndex,
 				LightSceneInfo,
 				TryGetRHI(ScreenShadowMaskTexture),
 				TryGetRHI(LightingChannelsTexture),
-				InHairVisibilityViews,
 				bRenderOverlap,
 				false);
 		});
@@ -2480,26 +2382,20 @@ void FDeferredShadingSceneRenderer::RenderLight(
 
 BEGIN_SHADER_PARAMETER_STRUCT(FRenderLightForHairParameters, )
 	SHADER_PARAMETER_STRUCT_INCLUDE(FRenderLightParameters, Light)
-	RDG_TEXTURE_ACCESS(HairIndexAndCountTexture, ERHIAccess::SRVGraphics)
-	RDG_TEXTURE_ACCESS(HairNodeCount, ERHIAccess::SRVGraphics)
-	RDG_BUFFER_ACCESS(HairTransmittanceMask, ERHIAccess::SRVGraphics)
-	RDG_BUFFER_ACCESS(HairVisibilityNodeData, ERHIAccess::SRVGraphics)
-	RDG_BUFFER_ACCESS(HairVisibilityNodeCoords, ERHIAccess::SRVGraphics)
+	SHADER_PARAMETER_RDG_UNIFORM_BUFFER(FHairStrandsViewUniformParameters, HairStrands)
 	SHADER_PARAMETER_RDG_BUFFER_SRV(Buffer, HairTransmittanceMaskSRV)
-	SHADER_PARAMETER_RDG_BUFFER_SRV(Buffer, HairVisibilityNodeDataSRV)
-	SHADER_PARAMETER_RDG_BUFFER_SRV(Buffer, HairVisibilityNodeCoordsSRV)
 END_SHADER_PARAMETER_STRUCT()
 
 void FDeferredShadingSceneRenderer::RenderLightForHair(
 	FRDGBuilder& GraphBuilder,
+	FViewInfo& View,
 	TRDGUniformBufferRef<FSceneTextureUniformParameters> SceneTexturesUniformBuffer,
 	const FLightSceneInfo* LightSceneInfo,
 	FRDGTextureRef HairShadowMaskTexture,
 	FRDGTextureRef LightingChannelsTexture,
-	const FHairStrandsTransmittanceMaskData& InTransmittanceMaskData,
-	const FHairStrandsVisibilityViews* InHairVisibilityViews)
+	const FHairStrandsTransmittanceMaskData& InTransmittanceMaskData)
 {
-	const bool bHairRenderingEnabled = InHairVisibilityViews && (LightSceneInfo->Proxy->CastsHairStrandsDeepShadow() || IsHairStrandsVoxelizationEnable());
+	const bool bHairRenderingEnabled = HairStrands::HasViewHairStrandsData(View) && (LightSceneInfo->Proxy->CastsHairStrandsDeepShadow() || IsHairStrandsVoxelizationEnable());
 	if (!bHairRenderingEnabled)
 	{
 		return;
@@ -2514,34 +2410,35 @@ void FDeferredShadingSceneRenderer::RenderLightForHair(
 
 	const FVolumetricCloudRenderSceneInfo* CloudInfo = Scene->GetVolumetricCloudSceneInfo();
 
-	for (int32 ViewIndex = 0; ViewIndex < Views.Num(); ViewIndex++)
+
 	{
-		FViewInfo& View = Views[ViewIndex];
 		RDG_GPU_MASK_SCOPE(GraphBuilder, View.GPUMask);
 
 		// Ensure the light is valid for this view
-		if (!LightSceneInfo->ShouldRenderLight(View) || ViewIndex >= InHairVisibilityViews->HairDatas.Num())
+		if (!LightSceneInfo->ShouldRenderLight(View))
 		{
-			continue;
+			return;
 		}
 
-		const FHairStrandsVisibilityData& HairVisibilityData = InHairVisibilityViews->HairDatas[ViewIndex];
+		const FHairStrandsVisibilityData& HairVisibilityData = View.HairStrandsViewData.VisibilityData;
 		if (!HairVisibilityData.SampleLightingBuffer)
 		{
-			continue;
+			return;
 		}
 
 		FRenderLightForHairParameters* PassParameters = GraphBuilder.AllocParameters<FRenderLightForHairParameters>();
-		GetRenderLightParameters(HairVisibilityData.SampleLightingBuffer, nullptr, SceneTexturesUniformBuffer, HairShadowMaskTexture, LightingChannelsTexture, GetCloudShadowAOParameters(GraphBuilder, View, CloudInfo), InHairVisibilityViews, PassParameters->Light);
-		PassParameters->HairIndexAndCountTexture = HairVisibilityData.NodeIndex;
-
-		PassParameters->HairTransmittanceMask = InTransmittanceMaskData.TransmittanceMask;
-		PassParameters->HairVisibilityNodeData = HairVisibilityData.NodeData;
-		PassParameters->HairVisibilityNodeCoords = HairVisibilityData.NodeCoord;
-		PassParameters->HairNodeCount = HairVisibilityData.NodeCount;
+		GetRenderLightParameters(
+			HairVisibilityData.SampleLightingBuffer, 
+			nullptr, 
+			SceneTexturesUniformBuffer, 
+			HairStrands::BindHairStrandsViewUniformParameters(View), 
+			HairShadowMaskTexture, 
+			LightingChannelsTexture, 
+			GetCloudShadowAOParameters(GraphBuilder, View, CloudInfo), 
+			PassParameters->Light);
+
+		PassParameters->HairStrands = HairStrands::BindHairStrandsViewUniformParameters(View);
 		PassParameters->HairTransmittanceMaskSRV = GraphBuilder.CreateSRV(InTransmittanceMaskData.TransmittanceMask);
-		PassParameters->HairVisibilityNodeDataSRV = GraphBuilder.CreateSRV(HairVisibilityData.NodeData);
-		PassParameters->HairVisibilityNodeCoordsSRV = GraphBuilder.CreateSRV(HairVisibilityData.NodeCoord);
 
 		const bool bIsShadowMaskValid = !!PassParameters->Light.ShadowMaskTexture;
 		const uint32 MaxTransmittanceElementCount = InTransmittanceMaskData.TransmittanceMask ? InTransmittanceMaskData.TransmittanceMask->Desc.NumElements : 0;
@@ -2557,10 +2454,6 @@ void FDeferredShadingSceneRenderer::RenderLightForHair(
 			RenderLightParams.DeepShadow_TransmittanceMaskBufferMaxCount = MaxTransmittanceElementCount;
 			RenderLightParams.ScreenShadowMaskSubPixelTexture = bIsShadowMaskValid ? PassParameters->Light.ShadowMaskTexture->GetRHI() : GSystemTextures.WhiteDummy->GetShaderResourceRHI();
 			RenderLightParams.DeepShadow_TransmittanceMaskBuffer = PassParameters->HairTransmittanceMaskSRV->GetRHI();
-			RenderLightParams.HairVisibilityNodeOffsetAndCount = PassParameters->HairIndexAndCountTexture->GetRHI();
-			RenderLightParams.HairVisibilityNodeDataSRV = PassParameters->HairVisibilityNodeDataSRV->GetRHI();
-			RenderLightParams.HairVisibilityNodeCoordsSRV = PassParameters->HairVisibilityNodeCoordsSRV->GetRHI();
-			RenderLightParams.HairCategorizationTexture = PassParameters->Light.HairCategorizationTexture->GetRHI();
 
 			FGraphicsPipelineStateInitializer GraphicsPSOInit;
 			RHICmdList.ApplyCachedRenderTargets(GraphicsPSOInit);
@@ -2600,14 +2493,15 @@ void FDeferredShadingSceneRenderer::RenderLightForHair(
 			GraphicsPSOInit.PrimitiveType = PT_TriangleList;
 			SetGraphicsPipelineState(RHICmdList, GraphicsPSOInit);
 
-			VertexShader->SetParameters(RHICmdList, View, &HairVisibilityData);
+			VertexShader->SetParameters(RHICmdList, View, PassParameters->HairStrands->GetRHI());
 			PixelShader->SetParameters(
 				RHICmdList,
 				View,
 				LightSceneInfo,
 				TryGetRHI(HairShadowMaskTexture),
 				TryGetRHI(LightingChannelsTexture),
-				&RenderLightParams);
+				&RenderLightParams,
+				PassParameters->HairStrands->GetRHI());
 
 			RHICmdList.SetStreamSource(0, nullptr, 0);
 			RHICmdList.DrawPrimitive(0, 1, 1);
@@ -2620,7 +2514,6 @@ void FDeferredShadingSceneRenderer::RenderLightsForHair(
 	FRDGBuilder& GraphBuilder,
 	TRDGUniformBufferRef<FSceneTextureUniformParameters> SceneTexturesUniformBuffer,
 	FSortedLightSetSceneInfo &SortedLightSet,
-	const FHairStrandsRenderingData* HairDatas,
 	FRDGTextureRef ScreenShadowMaskSubPixelTexture,
 	FRDGTextureRef LightingChannelsTexture)
 {
@@ -2629,32 +2522,39 @@ void FDeferredShadingSceneRenderer::RenderLightsForHair(
 	const int32 AttenuationLightStart = SortedLightSet.AttenuationLightStart;
 	const int32 SimpleLightsEnd = SortedLightSet.SimpleLightsEnd;
 
-	const bool bUseHairLighting = HairDatas != nullptr;
-	if (ViewFamily.EngineShowFlags.DirectLighting && bUseHairLighting)
+	if (ViewFamily.EngineShowFlags.DirectLighting)
 	{
 		RDG_EVENT_SCOPE(GraphBuilder, "DirectLighting");
 
-		for (int32 LightIndex = AttenuationLightStart; LightIndex < SortedLights.Num(); LightIndex++)
+		for (FViewInfo& View : Views)
 		{
-			const FSortedLightSceneInfo& SortedLightInfo = SortedLights[LightIndex];
-			const FLightSceneInfo& LightSceneInfo = *SortedLightInfo.LightSceneInfo;
-			if (LightSceneInfo.Proxy)
+			if (!HairStrands::HasViewHairStrandsData(View))
 			{
-				const bool bDrawHairShadow = SortedLightInfo.SortKey.Fields.bShadowed;
-				FHairStrandsTransmittanceMaskData TransmittanceMaskData;
-				if (bDrawHairShadow)
+				continue;
+			}
+
+			for (int32 LightIndex = AttenuationLightStart; LightIndex < SortedLights.Num(); LightIndex++)
+			{
+				const FSortedLightSceneInfo& SortedLightInfo = SortedLights[LightIndex];
+				const FLightSceneInfo& LightSceneInfo = *SortedLightInfo.LightSceneInfo;
+				if (LightSceneInfo.Proxy)
 				{
-					TransmittanceMaskData = RenderHairStrandsTransmittanceMask(GraphBuilder, Views, &LightSceneInfo, HairDatas, ScreenShadowMaskSubPixelTexture);
-				}
+					const bool bDrawHairShadow = SortedLightInfo.SortKey.Fields.bShadowed;
+					FHairStrandsTransmittanceMaskData TransmittanceMaskData;
+					if (bDrawHairShadow)
+					{
+						TransmittanceMaskData = RenderHairStrandsTransmittanceMask(GraphBuilder, View, &LightSceneInfo, ScreenShadowMaskSubPixelTexture);
+					}
 
-				RenderLightForHair(
-					GraphBuilder,
-					SceneTexturesUniformBuffer,
-					&LightSceneInfo,
-					ScreenShadowMaskSubPixelTexture,
-					LightingChannelsTexture,
-					TransmittanceMaskData,
-					&HairDatas->HairVisibilityViews);
+					RenderLightForHair(
+						GraphBuilder,
+						View,
+						SceneTexturesUniformBuffer,
+						&LightSceneInfo,
+						ScreenShadowMaskSubPixelTexture,
+						LightingChannelsTexture,
+						TransmittanceMaskData);
+				}
 			}
 		}
 	}
diff --git a/Engine/Source/Runtime/Renderer/Private/RayTracing/RaytracingSkylight.cpp b/Engine/Source/Runtime/Renderer/Private/RayTracing/RaytracingSkylight.cpp
--- a/Engine/Source/Runtime/Renderer/Private/RayTracing/RaytracingSkylight.cpp
+++ b/Engine/Source/Runtime/Renderer/Private/RayTracing/RaytracingSkylight.cpp
@@ -410,8 +410,7 @@ void FDeferredShadingSceneRenderer::RenderRayTracingSkyLight(
 	FRDGBuilder& GraphBuilder,
 	FRDGTextureRef SceneColorTexture,
 	FRDGTextureRef& OutSkyLightTexture,
-	FRDGTextureRef& OutHitDistanceTexture,
-	const FHairStrandsRenderingData* HairDatas)
+	FRDGTextureRef& OutHitDistanceTexture)
 {
 	FSkyLightSceneProxy* SkyLight = Scene->SkyLight;
 	if (!ShouldRenderRayTracingSkyLight(SkyLight))
@@ -502,12 +501,12 @@ void FDeferredShadingSceneRenderer::RenderRayTracingSkyLight(
 		PassParameters->ViewUniformBuffer = View.ViewUniformBuffer;
 
 		const bool bUseHairLighting = 
-			HairDatas && ViewIndex < HairDatas->MacroGroupsPerViews.Views.Num() && 
-			HairDatas->MacroGroupsPerViews.Views[ViewIndex].VirtualVoxelResources.IsValid() && 
+			HairStrands::HasViewHairStrandsData(View) && 
+			View.HairStrandsViewData.VirtualVoxelResources.IsValid() &&
 			CVarRayTracingSkyLightEnableHairVoxel.GetValueOnRenderThread() > 0;
 		if (bUseHairLighting)
 		{
-			PassParameters->VirtualVoxel = HairDatas->MacroGroupsPerViews.Views[ViewIndex].VirtualVoxelResources.UniformBuffer;
+			PassParameters->VirtualVoxel = View.HairStrandsViewData.VirtualVoxelResources.UniformBuffer;
 		}
 
 		FRayTracingSkyLightRGS::FPermutationDomain PermutationVector;
diff --git a/Engine/Source/Runtime/Renderer/Private/SceneRendering.h b/Engine/Source/Runtime/Renderer/Private/SceneRendering.h
--- a/Engine/Source/Runtime/Renderer/Private/SceneRendering.h
+++ b/Engine/Source/Runtime/Renderer/Private/SceneRendering.h
@@ -56,7 +56,6 @@ struct FCloudRenderContext;
 struct FSceneWithoutWaterTextures;
 struct FHairStrandsVisibilityViews;
 struct FSortedLightSetSceneInfo;
-struct FHairStrandsRenderingData;
 enum class EVelocityPass : uint32;
 struct FStrataSceneData;
 
@@ -1984,10 +1983,11 @@ protected:
 
 	void RenderShadowProjections(
 		FRDGBuilder& GraphBuilder,
-		const FShadowProjectionPassParameters& CommonPassParameters,
+		FRDGTextureRef OutputTexture,
+		const FMinimalSceneTextures& SceneTextures,
 		const FLightSceneProxy* LightSceneProxy,
-		const FHairStrandsRenderingData* HairDatas,
 		TArrayView<const FProjectedShadowInfo* const> Shadows,
+		bool bSubPixelShadow,
 		bool bProjectingForForwardShading,
 		bool bMobileModulatedProjections);
 
diff --git a/Engine/Source/Runtime/Renderer/Private/ShadowRendering.cpp b/Engine/Source/Runtime/Renderer/Private/ShadowRendering.cpp
--- a/Engine/Source/Runtime/Renderer/Private/ShadowRendering.cpp
+++ b/Engine/Source/Runtime/Renderer/Private/ShadowRendering.cpp
@@ -407,7 +407,7 @@ IMPLEMENT_SHADOW_PROJECTION_PIXEL_SHADER(5, true);
 
 template<typename VertexShaderType, typename PixelShaderType>
 static void BindShaderShaders(FRHICommandList& RHICmdList, FGraphicsPipelineStateInitializer& GraphicsPSOInit,
-	int32 ViewIndex, const FViewInfo& View, const FHairStrandsVisibilityData* HairVisibilityData, const FProjectedShadowInfo* ShadowInfo)
+	int32 ViewIndex, const FViewInfo& View, const FProjectedShadowInfo* ShadowInfo, FRHIUniformBuffer* HairStrandsUniformBuffer = nullptr)
 {
 	TShaderRef<VertexShaderType> VertexShader = View.ShaderMap->GetShader<VertexShaderType>();
 	TShaderRef<PixelShaderType> PixelShader = View.ShaderMap->GetShader<PixelShaderType>();
@@ -417,7 +417,7 @@ static void BindShaderShaders(FRHICommandList& RHICmdList, FGraphicsPipelineStat
 	SetGraphicsPipelineState(RHICmdList, GraphicsPSOInit);
 
 	VertexShader->SetParameters(RHICmdList, View, ShadowInfo);
-	PixelShader->SetParameters(RHICmdList, ViewIndex, View, HairVisibilityData, ShadowInfo);
+	PixelShader->SetParameters(RHICmdList, ViewIndex, View, ShadowInfo);
 
 	if (Strata::IsStrataEnabled())
 	{
@@ -425,13 +425,20 @@ static void BindShaderShaders(FRHICommandList& RHICmdList, FGraphicsPipelineStat
 		FRHIUniformBuffer* StrataUniformBuffer = Strata::BindStrataGlobalUniformParameters(View);
 		PixelShader->FGlobalShader::template SetParameters<FStrataGlobalUniformParameters>(RHICmdList, ShaderRHI, StrataUniformBuffer);
 	}
+
+	if (HairStrandsUniformBuffer)
+	{
+		FRHIPixelShader* ShaderRHI = PixelShader.GetPixelShader();
+		PixelShader->FGlobalShader::template SetParameters<FHairStrandsViewUniformParameters>(RHICmdList, ShaderRHI, HairStrandsUniformBuffer);
+	}
 }
 
 
 static void BindShadowProjectionShaders(int32 Quality, FRHICommandList& RHICmdList, FGraphicsPipelineStateInitializer GraphicsPSOInit, int32 ViewIndex, const FViewInfo& View,
-	const FHairStrandsVisibilityData* HairVisibilityData, const FProjectedShadowInfo* ShadowInfo, bool bMobileModulatedProjections)
+	const FProjectedShadowInfo* ShadowInfo, bool bMobileModulatedProjections, FRHIUniformBuffer* HairStrandsUniformBuffer)
 {
-	if (HairVisibilityData)
+	const bool bSubPixelShadow = HairStrandsUniformBuffer != nullptr;
+	if (bSubPixelShadow)
 	{
 		check(!bMobileModulatedProjections);
 
@@ -439,11 +446,11 @@ static void BindShadowProjectionShaders(int32 Quality, FRHICommandList& RHICmdLi
 		{
 			switch (Quality)
 			{
-			case 1: BindShaderShaders<FShadowProjectionNoTransformVS, TShadowProjectionPS<1, false, false, false, true> >(RHICmdList, GraphicsPSOInit, ViewIndex, View, HairVisibilityData, ShadowInfo); break;
-			case 2: BindShaderShaders<FShadowProjectionNoTransformVS, TShadowProjectionPS<2, false, false, false, true> >(RHICmdList, GraphicsPSOInit, ViewIndex, View, HairVisibilityData, ShadowInfo); break;
-			case 3: BindShaderShaders<FShadowProjectionNoTransformVS, TShadowProjectionPS<3, false, false, false, true> >(RHICmdList, GraphicsPSOInit, ViewIndex, View, HairVisibilityData, ShadowInfo); break;
-			case 4: BindShaderShaders<FShadowProjectionNoTransformVS, TShadowProjectionPS<4, false, false, false, true> >(RHICmdList, GraphicsPSOInit, ViewIndex, View, HairVisibilityData, ShadowInfo); break;
-			case 5: BindShaderShaders<FShadowProjectionNoTransformVS, TShadowProjectionPS<5, false, false, false, true> >(RHICmdList, GraphicsPSOInit, ViewIndex, View, HairVisibilityData, ShadowInfo); break;
+			case 1: BindShaderShaders<FShadowProjectionNoTransformVS, TShadowProjectionPS<1, false, false, false, true> >(RHICmdList, GraphicsPSOInit, ViewIndex, View, ShadowInfo, HairStrandsUniformBuffer); break;
+			case 2: BindShaderShaders<FShadowProjectionNoTransformVS, TShadowProjectionPS<2, false, false, false, true> >(RHICmdList, GraphicsPSOInit, ViewIndex, View, ShadowInfo, HairStrandsUniformBuffer); break;
+			case 3: BindShaderShaders<FShadowProjectionNoTransformVS, TShadowProjectionPS<3, false, false, false, true> >(RHICmdList, GraphicsPSOInit, ViewIndex, View, ShadowInfo, HairStrandsUniformBuffer); break;
+			case 4: BindShaderShaders<FShadowProjectionNoTransformVS, TShadowProjectionPS<4, false, false, false, true> >(RHICmdList, GraphicsPSOInit, ViewIndex, View, ShadowInfo, HairStrandsUniformBuffer); break;
+			case 5: BindShaderShaders<FShadowProjectionNoTransformVS, TShadowProjectionPS<5, false, false, false, true> >(RHICmdList, GraphicsPSOInit, ViewIndex, View, ShadowInfo, HairStrandsUniformBuffer); break;
 			default:
 				check(0);
 			}
@@ -452,11 +459,11 @@ static void BindShadowProjectionShaders(int32 Quality, FRHICommandList& RHICmdLi
 		{
 			switch (Quality)
 			{
-			case 1: BindShaderShaders<FShadowVolumeBoundProjectionVS, TShadowProjectionPS<1, false, false, false, true> >(RHICmdList, GraphicsPSOInit, ViewIndex, View, HairVisibilityData, ShadowInfo); break;
-			case 2: BindShaderShaders<FShadowVolumeBoundProjectionVS, TShadowProjectionPS<2, false, false, false, true> >(RHICmdList, GraphicsPSOInit, ViewIndex, View, HairVisibilityData, ShadowInfo); break;
-			case 3: BindShaderShaders<FShadowVolumeBoundProjectionVS, TShadowProjectionPS<3, false, false, false, true> >(RHICmdList, GraphicsPSOInit, ViewIndex, View, HairVisibilityData, ShadowInfo); break;
-			case 4: BindShaderShaders<FShadowVolumeBoundProjectionVS, TShadowProjectionPS<4, false, false, false, true> >(RHICmdList, GraphicsPSOInit, ViewIndex, View, HairVisibilityData, ShadowInfo); break;
-			case 5: BindShaderShaders<FShadowVolumeBoundProjectionVS, TShadowProjectionPS<5, false, false, false, true> >(RHICmdList, GraphicsPSOInit, ViewIndex, View, HairVisibilityData, ShadowInfo); break;
+			case 1: BindShaderShaders<FShadowVolumeBoundProjectionVS, TShadowProjectionPS<1, false, false, false, true> >(RHICmdList, GraphicsPSOInit, ViewIndex, View, ShadowInfo, HairStrandsUniformBuffer); break;
+			case 2: BindShaderShaders<FShadowVolumeBoundProjectionVS, TShadowProjectionPS<2, false, false, false, true> >(RHICmdList, GraphicsPSOInit, ViewIndex, View, ShadowInfo, HairStrandsUniformBuffer); break;
+			case 3: BindShaderShaders<FShadowVolumeBoundProjectionVS, TShadowProjectionPS<3, false, false, false, true> >(RHICmdList, GraphicsPSOInit, ViewIndex, View, ShadowInfo, HairStrandsUniformBuffer); break;
+			case 4: BindShaderShaders<FShadowVolumeBoundProjectionVS, TShadowProjectionPS<4, false, false, false, true> >(RHICmdList, GraphicsPSOInit, ViewIndex, View, ShadowInfo, HairStrandsUniformBuffer); break;
+			case 5: BindShaderShaders<FShadowVolumeBoundProjectionVS, TShadowProjectionPS<5, false, false, false, true> >(RHICmdList, GraphicsPSOInit, ViewIndex, View, ShadowInfo, HairStrandsUniformBuffer); break;
 			default:
 				check(0);
 			}
@@ -468,11 +475,11 @@ static void BindShadowProjectionShaders(int32 Quality, FRHICommandList& RHICmdLi
 	{
 		switch (Quality)
 		{
-		case 1: BindShaderShaders<FShadowVolumeBoundProjectionVS, TShadowProjectionFromTranslucencyPS<1> >(RHICmdList, GraphicsPSOInit, ViewIndex, View, HairVisibilityData, ShadowInfo); break;
-		case 2: BindShaderShaders<FShadowVolumeBoundProjectionVS, TShadowProjectionFromTranslucencyPS<2> >(RHICmdList, GraphicsPSOInit, ViewIndex, View, HairVisibilityData, ShadowInfo); break;
-		case 3: BindShaderShaders<FShadowVolumeBoundProjectionVS, TShadowProjectionFromTranslucencyPS<3> >(RHICmdList, GraphicsPSOInit, ViewIndex, View, HairVisibilityData, ShadowInfo); break;
-		case 4: BindShaderShaders<FShadowVolumeBoundProjectionVS, TShadowProjectionFromTranslucencyPS<4> >(RHICmdList, GraphicsPSOInit, ViewIndex, View, HairVisibilityData, ShadowInfo); break;
-		case 5: BindShaderShaders<FShadowVolumeBoundProjectionVS, TShadowProjectionFromTranslucencyPS<5> >(RHICmdList, GraphicsPSOInit, ViewIndex, View, HairVisibilityData, ShadowInfo); break;
+		case 1: BindShaderShaders<FShadowVolumeBoundProjectionVS, TShadowProjectionFromTranslucencyPS<1> >(RHICmdList, GraphicsPSOInit, ViewIndex, View, ShadowInfo); break;
+		case 2: BindShaderShaders<FShadowVolumeBoundProjectionVS, TShadowProjectionFromTranslucencyPS<2> >(RHICmdList, GraphicsPSOInit, ViewIndex, View, ShadowInfo); break;
+		case 3: BindShaderShaders<FShadowVolumeBoundProjectionVS, TShadowProjectionFromTranslucencyPS<3> >(RHICmdList, GraphicsPSOInit, ViewIndex, View, ShadowInfo); break;
+		case 4: BindShaderShaders<FShadowVolumeBoundProjectionVS, TShadowProjectionFromTranslucencyPS<4> >(RHICmdList, GraphicsPSOInit, ViewIndex, View, ShadowInfo); break;
+		case 5: BindShaderShaders<FShadowVolumeBoundProjectionVS, TShadowProjectionFromTranslucencyPS<5> >(RHICmdList, GraphicsPSOInit, ViewIndex, View, ShadowInfo); break;
 		default:
 			check(0);
 		}
@@ -482,9 +489,9 @@ static void BindShadowProjectionShaders(int32 Quality, FRHICommandList& RHICmdLi
 		if (CVarFilterMethod.GetValueOnRenderThread() == 1)
 		{
 			if (ShadowInfo->CascadeSettings.FadePlaneLength > 0)
-				BindShaderShaders<FShadowProjectionNoTransformVS, TDirectionalPercentageCloserShadowProjectionPS<5, true> >(RHICmdList, GraphicsPSOInit, ViewIndex, View, HairVisibilityData, ShadowInfo);
+				BindShaderShaders<FShadowProjectionNoTransformVS, TDirectionalPercentageCloserShadowProjectionPS<5, true> >(RHICmdList, GraphicsPSOInit, ViewIndex, View, ShadowInfo);
 			else
-				BindShaderShaders<FShadowProjectionNoTransformVS, TDirectionalPercentageCloserShadowProjectionPS<5, false> >(RHICmdList, GraphicsPSOInit, ViewIndex, View, HairVisibilityData, ShadowInfo);
+				BindShaderShaders<FShadowProjectionNoTransformVS, TDirectionalPercentageCloserShadowProjectionPS<5, false> >(RHICmdList, GraphicsPSOInit, ViewIndex, View, ShadowInfo);
 		}
 		else if (ShadowInfo->CascadeSettings.FadePlaneLength > 0)
 		{
@@ -492,11 +499,11 @@ static void BindShadowProjectionShaders(int32 Quality, FRHICommandList& RHICmdLi
 			{
 				switch (Quality)
 				{
-				case 1: BindShaderShaders<FShadowProjectionNoTransformVS, TShadowProjectionPS<1, true, false, true> >(RHICmdList, GraphicsPSOInit, ViewIndex, View, HairVisibilityData, ShadowInfo); break;
-				case 2: BindShaderShaders<FShadowProjectionNoTransformVS, TShadowProjectionPS<2, true, false, true> >(RHICmdList, GraphicsPSOInit, ViewIndex, View, HairVisibilityData, ShadowInfo); break;
-				case 3: BindShaderShaders<FShadowProjectionNoTransformVS, TShadowProjectionPS<3, true, false, true> >(RHICmdList, GraphicsPSOInit, ViewIndex, View, HairVisibilityData, ShadowInfo); break;
-				case 4: BindShaderShaders<FShadowProjectionNoTransformVS, TShadowProjectionPS<4, true, false, true> >(RHICmdList, GraphicsPSOInit, ViewIndex, View, HairVisibilityData, ShadowInfo); break;
-				case 5: BindShaderShaders<FShadowProjectionNoTransformVS, TShadowProjectionPS<5, true, false, true> >(RHICmdList, GraphicsPSOInit, ViewIndex, View, HairVisibilityData, ShadowInfo); break;
+				case 1: BindShaderShaders<FShadowProjectionNoTransformVS, TShadowProjectionPS<1, true, false, true> >(RHICmdList, GraphicsPSOInit, ViewIndex, View, ShadowInfo); break;
+				case 2: BindShaderShaders<FShadowProjectionNoTransformVS, TShadowProjectionPS<2, true, false, true> >(RHICmdList, GraphicsPSOInit, ViewIndex, View, ShadowInfo); break;
+				case 3: BindShaderShaders<FShadowProjectionNoTransformVS, TShadowProjectionPS<3, true, false, true> >(RHICmdList, GraphicsPSOInit, ViewIndex, View, ShadowInfo); break;
+				case 4: BindShaderShaders<FShadowProjectionNoTransformVS, TShadowProjectionPS<4, true, false, true> >(RHICmdList, GraphicsPSOInit, ViewIndex, View, ShadowInfo); break;
+				case 5: BindShaderShaders<FShadowProjectionNoTransformVS, TShadowProjectionPS<5, true, false, true> >(RHICmdList, GraphicsPSOInit, ViewIndex, View, ShadowInfo); break;
 				default:
 					check(0);
 				}
@@ -505,11 +512,11 @@ static void BindShadowProjectionShaders(int32 Quality, FRHICommandList& RHICmdLi
 			{
 				switch (Quality)
 				{
-				case 1: BindShaderShaders<FShadowProjectionNoTransformVS, TShadowProjectionPS<1, true> >(RHICmdList, GraphicsPSOInit, ViewIndex, View, HairVisibilityData, ShadowInfo); break;
-				case 2: BindShaderShaders<FShadowProjectionNoTransformVS, TShadowProjectionPS<2, true> >(RHICmdList, GraphicsPSOInit, ViewIndex, View, HairVisibilityData, ShadowInfo); break;
-				case 3: BindShaderShaders<FShadowProjectionNoTransformVS, TShadowProjectionPS<3, true> >(RHICmdList, GraphicsPSOInit, ViewIndex, View, HairVisibilityData, ShadowInfo); break;
-				case 4: BindShaderShaders<FShadowProjectionNoTransformVS, TShadowProjectionPS<4, true> >(RHICmdList, GraphicsPSOInit, ViewIndex, View, HairVisibilityData, ShadowInfo); break;
-				case 5: BindShaderShaders<FShadowProjectionNoTransformVS, TShadowProjectionPS<5, true> >(RHICmdList, GraphicsPSOInit, ViewIndex, View, HairVisibilityData, ShadowInfo); break;
+				case 1: BindShaderShaders<FShadowProjectionNoTransformVS, TShadowProjectionPS<1, true> >(RHICmdList, GraphicsPSOInit, ViewIndex, View, ShadowInfo); break;
+				case 2: BindShaderShaders<FShadowProjectionNoTransformVS, TShadowProjectionPS<2, true> >(RHICmdList, GraphicsPSOInit, ViewIndex, View, ShadowInfo); break;
+				case 3: BindShaderShaders<FShadowProjectionNoTransformVS, TShadowProjectionPS<3, true> >(RHICmdList, GraphicsPSOInit, ViewIndex, View, ShadowInfo); break;
+				case 4: BindShaderShaders<FShadowProjectionNoTransformVS, TShadowProjectionPS<4, true> >(RHICmdList, GraphicsPSOInit, ViewIndex, View, ShadowInfo); break;
+				case 5: BindShaderShaders<FShadowProjectionNoTransformVS, TShadowProjectionPS<5, true> >(RHICmdList, GraphicsPSOInit, ViewIndex, View, ShadowInfo); break;
 				default:
 					check(0);
 				}
@@ -521,11 +528,11 @@ static void BindShadowProjectionShaders(int32 Quality, FRHICommandList& RHICmdLi
 			{
 				switch (Quality)
 				{
-				case 1: BindShaderShaders<FShadowProjectionNoTransformVS, TShadowProjectionPS<1, false, false, true> >(RHICmdList, GraphicsPSOInit, ViewIndex, View, HairVisibilityData, ShadowInfo); break;
-				case 2: BindShaderShaders<FShadowProjectionNoTransformVS, TShadowProjectionPS<2, false, false, true> >(RHICmdList, GraphicsPSOInit, ViewIndex, View, HairVisibilityData, ShadowInfo); break;
-				case 3: BindShaderShaders<FShadowProjectionNoTransformVS, TShadowProjectionPS<3, false, false, true> >(RHICmdList, GraphicsPSOInit, ViewIndex, View, HairVisibilityData, ShadowInfo); break;
-				case 4: BindShaderShaders<FShadowProjectionNoTransformVS, TShadowProjectionPS<4, false, false, true> >(RHICmdList, GraphicsPSOInit, ViewIndex, View, HairVisibilityData, ShadowInfo); break;
-				case 5: BindShaderShaders<FShadowProjectionNoTransformVS, TShadowProjectionPS<5, false, false, true> >(RHICmdList, GraphicsPSOInit, ViewIndex, View, HairVisibilityData, ShadowInfo); break;
+				case 1: BindShaderShaders<FShadowProjectionNoTransformVS, TShadowProjectionPS<1, false, false, true> >(RHICmdList, GraphicsPSOInit, ViewIndex, View, ShadowInfo); break;
+				case 2: BindShaderShaders<FShadowProjectionNoTransformVS, TShadowProjectionPS<2, false, false, true> >(RHICmdList, GraphicsPSOInit, ViewIndex, View, ShadowInfo); break;
+				case 3: BindShaderShaders<FShadowProjectionNoTransformVS, TShadowProjectionPS<3, false, false, true> >(RHICmdList, GraphicsPSOInit, ViewIndex, View, ShadowInfo); break;
+				case 4: BindShaderShaders<FShadowProjectionNoTransformVS, TShadowProjectionPS<4, false, false, true> >(RHICmdList, GraphicsPSOInit, ViewIndex, View, ShadowInfo); break;
+				case 5: BindShaderShaders<FShadowProjectionNoTransformVS, TShadowProjectionPS<5, false, false, true> >(RHICmdList, GraphicsPSOInit, ViewIndex, View, ShadowInfo); break;
 				default:
 					check(0);
 				}
@@ -534,11 +541,11 @@ static void BindShadowProjectionShaders(int32 Quality, FRHICommandList& RHICmdLi
 			{ 
 				switch (Quality)
 				{
-				case 1: BindShaderShaders<FShadowProjectionNoTransformVS, TShadowProjectionPS<1, false> >(RHICmdList, GraphicsPSOInit, ViewIndex, View, HairVisibilityData, ShadowInfo); break;
-				case 2: BindShaderShaders<FShadowProjectionNoTransformVS, TShadowProjectionPS<2, false> >(RHICmdList, GraphicsPSOInit, ViewIndex, View, HairVisibilityData, ShadowInfo); break;
-				case 3: BindShaderShaders<FShadowProjectionNoTransformVS, TShadowProjectionPS<3, false> >(RHICmdList, GraphicsPSOInit, ViewIndex, View, HairVisibilityData, ShadowInfo); break;
-				case 4: BindShaderShaders<FShadowProjectionNoTransformVS, TShadowProjectionPS<4, false> >(RHICmdList, GraphicsPSOInit, ViewIndex, View, HairVisibilityData, ShadowInfo); break;
-				case 5: BindShaderShaders<FShadowProjectionNoTransformVS, TShadowProjectionPS<5, false> >(RHICmdList, GraphicsPSOInit, ViewIndex, View, HairVisibilityData, ShadowInfo); break;
+				case 1: BindShaderShaders<FShadowProjectionNoTransformVS, TShadowProjectionPS<1, false> >(RHICmdList, GraphicsPSOInit, ViewIndex, View, ShadowInfo); break;
+				case 2: BindShaderShaders<FShadowProjectionNoTransformVS, TShadowProjectionPS<2, false> >(RHICmdList, GraphicsPSOInit, ViewIndex, View, ShadowInfo); break;
+				case 3: BindShaderShaders<FShadowProjectionNoTransformVS, TShadowProjectionPS<3, false> >(RHICmdList, GraphicsPSOInit, ViewIndex, View, ShadowInfo); break;
+				case 4: BindShaderShaders<FShadowProjectionNoTransformVS, TShadowProjectionPS<4, false> >(RHICmdList, GraphicsPSOInit, ViewIndex, View, ShadowInfo); break;
+				case 5: BindShaderShaders<FShadowProjectionNoTransformVS, TShadowProjectionPS<5, false> >(RHICmdList, GraphicsPSOInit, ViewIndex, View, ShadowInfo); break;
 				default:
 					check(0);
 				}
@@ -551,11 +558,11 @@ static void BindShadowProjectionShaders(int32 Quality, FRHICommandList& RHICmdLi
 		{
 			switch (Quality)
 			{
-			case 1: BindShaderShaders<FShadowVolumeBoundProjectionVS, TModulatedShadowProjection<1> >(RHICmdList, GraphicsPSOInit, ViewIndex, View, HairVisibilityData, ShadowInfo); break;
-			case 2: BindShaderShaders<FShadowVolumeBoundProjectionVS, TModulatedShadowProjection<2> >(RHICmdList, GraphicsPSOInit, ViewIndex, View, HairVisibilityData, ShadowInfo); break;
-			case 3: BindShaderShaders<FShadowVolumeBoundProjectionVS, TModulatedShadowProjection<3> >(RHICmdList, GraphicsPSOInit, ViewIndex, View, HairVisibilityData, ShadowInfo); break;
-			case 4: BindShaderShaders<FShadowVolumeBoundProjectionVS, TModulatedShadowProjection<4> >(RHICmdList, GraphicsPSOInit, ViewIndex, View, HairVisibilityData, ShadowInfo); break;
-			case 5: BindShaderShaders<FShadowVolumeBoundProjectionVS, TModulatedShadowProjection<5> >(RHICmdList, GraphicsPSOInit, ViewIndex, View, HairVisibilityData, ShadowInfo); break;
+			case 1: BindShaderShaders<FShadowVolumeBoundProjectionVS, TModulatedShadowProjection<1> >(RHICmdList, GraphicsPSOInit, ViewIndex, View, ShadowInfo); break;
+			case 2: BindShaderShaders<FShadowVolumeBoundProjectionVS, TModulatedShadowProjection<2> >(RHICmdList, GraphicsPSOInit, ViewIndex, View, ShadowInfo); break;
+			case 3: BindShaderShaders<FShadowVolumeBoundProjectionVS, TModulatedShadowProjection<3> >(RHICmdList, GraphicsPSOInit, ViewIndex, View, ShadowInfo); break;
+			case 4: BindShaderShaders<FShadowVolumeBoundProjectionVS, TModulatedShadowProjection<4> >(RHICmdList, GraphicsPSOInit, ViewIndex, View, ShadowInfo); break;
+			case 5: BindShaderShaders<FShadowVolumeBoundProjectionVS, TModulatedShadowProjection<5> >(RHICmdList, GraphicsPSOInit, ViewIndex, View, ShadowInfo); break;
 			default:
 				check(0);
 			}
@@ -564,11 +571,11 @@ static void BindShadowProjectionShaders(int32 Quality, FRHICommandList& RHICmdLi
 		{
 			switch (Quality)
 			{
-			case 1: BindShaderShaders<FShadowVolumeBoundProjectionVS, TShadowProjectionPS<1, false, false, true> >(RHICmdList, GraphicsPSOInit, ViewIndex, View, HairVisibilityData, ShadowInfo); break;
-			case 2: BindShaderShaders<FShadowVolumeBoundProjectionVS, TShadowProjectionPS<2, false, false, true> >(RHICmdList, GraphicsPSOInit, ViewIndex, View, HairVisibilityData, ShadowInfo); break;
-			case 3: BindShaderShaders<FShadowVolumeBoundProjectionVS, TShadowProjectionPS<3, false, false, true> >(RHICmdList, GraphicsPSOInit, ViewIndex, View, HairVisibilityData, ShadowInfo); break;
-			case 4: BindShaderShaders<FShadowVolumeBoundProjectionVS, TShadowProjectionPS<4, false, false, true> >(RHICmdList, GraphicsPSOInit, ViewIndex, View, HairVisibilityData, ShadowInfo); break;
-			case 5: BindShaderShaders<FShadowVolumeBoundProjectionVS, TShadowProjectionPS<5, false, false, true> >(RHICmdList, GraphicsPSOInit, ViewIndex, View, HairVisibilityData, ShadowInfo); break;
+			case 1: BindShaderShaders<FShadowVolumeBoundProjectionVS, TShadowProjectionPS<1, false, false, true> >(RHICmdList, GraphicsPSOInit, ViewIndex, View, ShadowInfo); break;
+			case 2: BindShaderShaders<FShadowVolumeBoundProjectionVS, TShadowProjectionPS<2, false, false, true> >(RHICmdList, GraphicsPSOInit, ViewIndex, View, ShadowInfo); break;
+			case 3: BindShaderShaders<FShadowVolumeBoundProjectionVS, TShadowProjectionPS<3, false, false, true> >(RHICmdList, GraphicsPSOInit, ViewIndex, View, ShadowInfo); break;
+			case 4: BindShaderShaders<FShadowVolumeBoundProjectionVS, TShadowProjectionPS<4, false, false, true> >(RHICmdList, GraphicsPSOInit, ViewIndex, View, ShadowInfo); break;
+			case 5: BindShaderShaders<FShadowVolumeBoundProjectionVS, TShadowProjectionPS<5, false, false, true> >(RHICmdList, GraphicsPSOInit, ViewIndex, View, ShadowInfo); break;
 			default:
 				check(0);
 			}
@@ -577,17 +584,17 @@ static void BindShadowProjectionShaders(int32 Quality, FRHICommandList& RHICmdLi
 		{
 			if (CVarFilterMethod.GetValueOnRenderThread() == 1 && ShadowInfo->GetLightSceneInfo().Proxy->GetLightType() == LightType_Spot)
 			{
-				BindShaderShaders<FShadowVolumeBoundProjectionVS, TSpotPercentageCloserShadowProjectionPS<5, false> >(RHICmdList, GraphicsPSOInit, ViewIndex, View, HairVisibilityData, ShadowInfo);
+				BindShaderShaders<FShadowVolumeBoundProjectionVS, TSpotPercentageCloserShadowProjectionPS<5, false> >(RHICmdList, GraphicsPSOInit, ViewIndex, View, ShadowInfo);
 			}
 			else
 			{
 				switch (Quality)
 				{
-				case 1: BindShaderShaders<FShadowVolumeBoundProjectionVS, TShadowProjectionPS<1, false> >(RHICmdList, GraphicsPSOInit, ViewIndex, View, HairVisibilityData, ShadowInfo); break;
-				case 2: BindShaderShaders<FShadowVolumeBoundProjectionVS, TShadowProjectionPS<2, false> >(RHICmdList, GraphicsPSOInit, ViewIndex, View, HairVisibilityData, ShadowInfo); break;
-				case 3: BindShaderShaders<FShadowVolumeBoundProjectionVS, TShadowProjectionPS<3, false> >(RHICmdList, GraphicsPSOInit, ViewIndex, View, HairVisibilityData, ShadowInfo); break;
-				case 4: BindShaderShaders<FShadowVolumeBoundProjectionVS, TShadowProjectionPS<4, false> >(RHICmdList, GraphicsPSOInit, ViewIndex, View, HairVisibilityData, ShadowInfo); break;
-				case 5: BindShaderShaders<FShadowVolumeBoundProjectionVS, TShadowProjectionPS<5, false> >(RHICmdList, GraphicsPSOInit, ViewIndex, View, HairVisibilityData, ShadowInfo); break;
+				case 1: BindShaderShaders<FShadowVolumeBoundProjectionVS, TShadowProjectionPS<1, false> >(RHICmdList, GraphicsPSOInit, ViewIndex, View, ShadowInfo); break;
+				case 2: BindShaderShaders<FShadowVolumeBoundProjectionVS, TShadowProjectionPS<2, false> >(RHICmdList, GraphicsPSOInit, ViewIndex, View, ShadowInfo); break;
+				case 3: BindShaderShaders<FShadowVolumeBoundProjectionVS, TShadowProjectionPS<3, false> >(RHICmdList, GraphicsPSOInit, ViewIndex, View, ShadowInfo); break;
+				case 4: BindShaderShaders<FShadowVolumeBoundProjectionVS, TShadowProjectionPS<4, false> >(RHICmdList, GraphicsPSOInit, ViewIndex, View, ShadowInfo); break;
+				case 5: BindShaderShaders<FShadowVolumeBoundProjectionVS, TShadowProjectionPS<5, false> >(RHICmdList, GraphicsPSOInit, ViewIndex, View, ShadowInfo); break;
 				default:
 					check(0);
 				}
@@ -981,9 +988,9 @@ void FProjectedShadowInfo::SetupProjectionStencilMask(
 
 BEGIN_SHADER_PARAMETER_STRUCT(FShadowProjectionPassParameters, )
 	SHADER_PARAMETER_STRUCT_INCLUDE(FViewShaderParameters, View)
+	SHADER_PARAMETER_RDG_UNIFORM_BUFFER(FHairStrandsViewUniformParameters, HairStrands)
 	SHADER_PARAMETER_STRUCT_INCLUDE(FSceneTextureShaderParameters, SceneTextures)
 	SHADER_PARAMETER_STRUCT_INCLUDE(FInstanceCullingDrawParams, InstanceCullingDrawParams)
-	RDG_TEXTURE_ACCESS(HairCategorizationTexture, ERHIAccess::SRVGraphics)
 	RDG_TEXTURE_ACCESS(ShadowTexture0, ERHIAccess::SRVGraphics)
 	RDG_TEXTURE_ACCESS(ShadowTexture1, ERHIAccess::SRVGraphics)
 	RENDER_TARGET_BINDING_SLOTS()
@@ -998,8 +1005,7 @@ void FProjectedShadowInfo::RenderProjection(
 	const FSceneRenderer* SceneRender,
 	bool bProjectingForForwardShading,
 	bool bMobileModulatedProjections,
-	const FHairStrandsVisibilityData* HairVisibilityData,
-	 const FHairStrandsMacroGroupDatas* HairMacroGroupData) const
+	bool bSubPixelShadow) const
 {
 	// Find the shadow's view relevance.
 	const FVisibleLightViewInfo& VisibleLightViewInfo = View->VisibleLightInfos[LightSceneInfo->Id];
@@ -1018,11 +1024,9 @@ void FProjectedShadowInfo::RenderProjection(
 	FPlane OutPlanes[6];
 	SetupFrustumForProjection(View, FrustumVertices, bCameraInsideShadowFrustum, OutPlanes);
 
-	const bool bSubPixelSupport = HairVisibilityData != nullptr;
-	const bool bStencilTestEnabled = true;// !bSubPixelSupport;
 	const bool bDepthBoundsTestEnabled = IsWholeSceneDirectionalShadow() && GSupportsDepthBoundsTest && CVarCSMDepthBoundsTest.GetValueOnRenderThread() != 0;// && !bSubPixelSupport;
 
-	if (bSubPixelSupport)
+	if (bSubPixelShadow)
 	{
 		// Do not apply pre-shadow on hair, as this is intended only for targed opaque geometry
 		if (bPreShadow)
@@ -1035,7 +1039,7 @@ void FProjectedShadowInfo::RenderProjection(
 		{
 			// Skip volume which does not intersect hair clusters
 			bool bIntersect = bValidPlanes;
-			for (const FHairStrandsMacroGroupData& Data : HairMacroGroupData->Datas)
+			for (const FHairStrandsMacroGroupData& Data : View->HairStrandsViewData.MacroGroupDatas)
 			{
 				const FSphere BoundSphere = Data.Bounds.GetSphere();
 				// Return the signed distance to the plane. The planes are pointing inward
@@ -1121,9 +1125,9 @@ void FProjectedShadowInfo::RenderProjection(
 		RDG_EVENT_NAME("%s", *EventName),
 		PassParameters,
 		ERDGPassFlags::Raster | PassFlags,
-		[this, SceneRender, View, ViewIndex, LightSceneProxy, HairVisibilityData, bProjectingForForwardShading, bMobileModulatedProjections, &InstanceCullingDrawParams](FRHICommandListImmediate& RHICmdList)
+		[this, SceneRender, View, ViewIndex, LightSceneProxy, bProjectingForForwardShading, bMobileModulatedProjections, &InstanceCullingDrawParams, bSubPixelShadow, PassParameters](FRHICommandListImmediate& RHICmdList)
 	{
-		RenderProjectionInternal(RHICmdList, ViewIndex, View, LightSceneProxy, SceneRender, bProjectingForForwardShading, bMobileModulatedProjections, HairVisibilityData, InstanceCullingDrawParams);
+		RenderProjectionInternal(RHICmdList, ViewIndex, View, LightSceneProxy, SceneRender, bProjectingForForwardShading, bMobileModulatedProjections, InstanceCullingDrawParams, bSubPixelShadow && PassParameters->HairStrands ? PassParameters->HairStrands.GetUniformBuffer()->GetRHI() : nullptr);
 	});
 }
 
@@ -1135,8 +1139,8 @@ void FProjectedShadowInfo::RenderProjectionInternal(
 	const FSceneRenderer* SceneRender,
 	bool bProjectingForForwardShading,
 	bool bMobileModulatedProjections,
-	const FHairStrandsVisibilityData* HairVisibilityData,
-	const FInstanceCullingDrawParams& InstanceCullingDrawParams) const
+	const FInstanceCullingDrawParams& InstanceCullingDrawParams, 
+	FRHIUniformBuffer* HairStrandsUniformBuffer) const
 {
 	RHICmdList.SetViewport(View->ViewRect.Min.X, View->ViewRect.Min.Y, 0.0f, View->ViewRect.Max.X, View->ViewRect.Max.Y, 1.0f);
 	LightSceneProxy->SetScissorRect(RHICmdList, *View, View->ViewRect);
@@ -1151,7 +1155,7 @@ void FProjectedShadowInfo::RenderProjectionInternal(
 	TArray<FVector4, TInlineAllocator<8>> FrustumVertices;
 	SetupFrustumForProjection(View, FrustumVertices, bCameraInsideShadowFrustum, OutFrustmPlanes);
 
-	const bool bSubPixelSupport = HairVisibilityData != nullptr;
+	const bool bSubPixelSupport = HairStrandsUniformBuffer != nullptr;// HairStrands::HasViewHairStrandsData(*View);
 	const bool bStencilTestEnabled = !bSubPixelSupport && GShadowStencilCulling;
 	const bool bDepthBoundsTestEnabled = IsWholeSceneDirectionalShadow() && GSupportsDepthBoundsTest && CVarCSMDepthBoundsTest.GetValueOnRenderThread() != 0 && !bSubPixelSupport;
 
@@ -1238,7 +1242,7 @@ void FProjectedShadowInfo::RenderProjectionInternal(
 		}
 
 		GraphicsPSOInit.BoundShaderState.VertexDeclarationRHI = GetVertexDeclarationFVector4();
-		BindShadowProjectionShaders(LocalQuality, RHICmdList, GraphicsPSOInit, ViewIndex, *View, HairVisibilityData, this, bMobileModulatedProjections);
+		BindShadowProjectionShaders(LocalQuality, RHICmdList, GraphicsPSOInit, ViewIndex, *View, this, bMobileModulatedProjections, HairStrandsUniformBuffer);
 
 		if (bDepthBoundsTestEnabled)
 		{
@@ -1312,12 +1316,11 @@ void FProjectedShadowInfo::RenderMobileModulatedShadowProjection(
 
 	// GPUCULL_TODO, mobile modulated shadow is inside the mobile render pass, probably couldn't work with the GPUCull.
 	FInstanceCullingDrawParams InstanceCullingDrawParams;
-
-	RenderProjectionInternal(RHICmdList, ViewIndex, View, LightSceneProxy, SceneRender, bProjectingForForwardShading, bMobileModulatedProjections, nullptr, InstanceCullingDrawParams);
+	RenderProjectionInternal(RHICmdList, ViewIndex, View, LightSceneProxy, SceneRender, bProjectingForForwardShading, bMobileModulatedProjections, InstanceCullingDrawParams, nullptr);
 }
 
 template <uint32 Quality, bool bUseTransmission, bool bUseSubPixel>
-static void SetPointLightShaderTempl(FRHICommandList& RHICmdList, FGraphicsPipelineStateInitializer& GraphicsPSOInit, int32 ViewIndex, const FViewInfo& View, const FProjectedShadowInfo* ShadowInfo, const FHairStrandsVisibilityData* HairVisibilityData)
+static void SetPointLightShaderTempl(FRHICommandList& RHICmdList, FGraphicsPipelineStateInitializer& GraphicsPSOInit, int32 ViewIndex, const FViewInfo& View, const FProjectedShadowInfo* ShadowInfo, FRHIUniformBuffer* HairStrandsUniformBuffer = nullptr)
 {
 	TShaderMapRef<FShadowVolumeBoundProjectionVS> VertexShader(View.ShaderMap);
 	TShaderMapRef<TOnePassPointShadowProjectionPS<Quality,bUseTransmission,bUseSubPixel> > PixelShader(View.ShaderMap);
@@ -1329,7 +1332,7 @@ static void SetPointLightShaderTempl(FRHICommandList& RHICmdList, FGraphicsPipel
 	SetGraphicsPipelineState(RHICmdList, GraphicsPSOInit);
 	
 	VertexShader->SetParameters(RHICmdList, View, ShadowInfo);
-	PixelShader->SetParameters(RHICmdList, ViewIndex, View, HairVisibilityData, ShadowInfo);
+	PixelShader->SetParameters(RHICmdList, ViewIndex, View, ShadowInfo, HairStrandsUniformBuffer);
 }
 
 void FProjectedShadowInfo::RenderOnePassPointLightProjection(
@@ -1339,8 +1342,7 @@ void FProjectedShadowInfo::RenderOnePassPointLightProjection(
 	const FViewInfo& View,
 	const FLightSceneProxy* LightSceneProxy,
 	bool bProjectingForForwardShading,
-	const FHairStrandsVisibilityData* HairVisibilityData,
-	const FHairStrandsMacroGroupDatas* HairMacroGroupData) const
+	bool bSubPixelShadow) const
 {
 	SCOPE_CYCLE_COUNTER(STAT_RenderWholeSceneShadowProjectionsTime);
 
@@ -1371,7 +1373,7 @@ void FProjectedShadowInfo::RenderOnePassPointLightProjection(
 		RDG_EVENT_NAME("OnePassPointLightProjection"),
 		PassParameters,
 		ERDGPassFlags::Raster | PassFlags,
-		[this, &View, HairVisibilityData, HairMacroGroupData, LightBounds, bProjectingForForwardShading, bCameraInsideLightGeometry, bUseTransmission, ViewIndex, LightSceneProxy](FRHICommandList& RHICmdList)
+		[this, &View, LightBounds, bProjectingForForwardShading, bCameraInsideLightGeometry, bUseTransmission, ViewIndex, LightSceneProxy, bSubPixelShadow, PassParameters](FRHICommandList& RHICmdList)
 	{
 		RHICmdList.SetViewport(View.ViewRect.Min.X, View.ViewRect.Min.Y, 0.0f, View.ViewRect.Max.X, View.ViewRect.Max.Y, 1.0f);
 		LightSceneProxy->SetScissorRect(RHICmdList, View, View.ViewRect);
@@ -1413,8 +1415,7 @@ void FProjectedShadowInfo::RenderOnePassPointLightProjection(
 			}
 		}
 
-		const bool bSubPixelSupport = HairVisibilityData != nullptr;
-		if (bSubPixelSupport)
+		if (bSubPixelShadow)
 		{
 			// Do not apply pre-shadow on hair, as this is intended only for targed opaque geometry
 			if (bPreShadow)
@@ -1426,7 +1427,7 @@ void FProjectedShadowInfo::RenderOnePassPointLightProjection(
 			bool bIntersect = false;
 			if (CVarHairStrandsCullPerObjectShadowCaster.GetValueOnRenderThread() > 0)
 			{
-				for (const FHairStrandsMacroGroupData& Data : HairMacroGroupData->Datas)
+				for (const FHairStrandsMacroGroupData& Data : View.HairStrandsViewData.MacroGroupDatas)
 				{
 					const FSphere BoundSphere = Data.Bounds.GetSphere();
 					if (BoundSphere.Intersects(LightBounds))
@@ -1443,13 +1444,14 @@ void FProjectedShadowInfo::RenderOnePassPointLightProjection(
 				}
 			}
 
+			FRHIUniformBuffer* HairStrandsUniformBuffer = PassParameters->HairStrands.GetUniformBuffer()->GetRHI();
 			switch (LocalQuality)
 			{
-			case 1: SetPointLightShaderTempl<1, false, true>(RHICmdList, GraphicsPSOInit, ViewIndex, View, this, HairVisibilityData); break;
-			case 2: SetPointLightShaderTempl<2, false, true>(RHICmdList, GraphicsPSOInit, ViewIndex, View, this, HairVisibilityData); break;
-			case 3: SetPointLightShaderTempl<3, false, true>(RHICmdList, GraphicsPSOInit, ViewIndex, View, this, HairVisibilityData); break;
-			case 4: SetPointLightShaderTempl<4, false, true>(RHICmdList, GraphicsPSOInit, ViewIndex, View, this, HairVisibilityData); break;
-			case 5: SetPointLightShaderTempl<5, false, true>(RHICmdList, GraphicsPSOInit, ViewIndex, View, this, HairVisibilityData); break;
+			case 1: SetPointLightShaderTempl<1, false, true>(RHICmdList, GraphicsPSOInit, ViewIndex, View, this, HairStrandsUniformBuffer); break;
+			case 2: SetPointLightShaderTempl<2, false, true>(RHICmdList, GraphicsPSOInit, ViewIndex, View, this, HairStrandsUniformBuffer); break;
+			case 3: SetPointLightShaderTempl<3, false, true>(RHICmdList, GraphicsPSOInit, ViewIndex, View, this, HairStrandsUniformBuffer); break;
+			case 4: SetPointLightShaderTempl<4, false, true>(RHICmdList, GraphicsPSOInit, ViewIndex, View, this, HairStrandsUniformBuffer); break;
+			case 5: SetPointLightShaderTempl<5, false, true>(RHICmdList, GraphicsPSOInit, ViewIndex, View, this, HairStrandsUniformBuffer); break;
 			default:
 				check(0);
 			}
@@ -1458,11 +1460,11 @@ void FProjectedShadowInfo::RenderOnePassPointLightProjection(
 		{
 			switch (LocalQuality)
 			{
-			case 1: SetPointLightShaderTempl<1, true, false>(RHICmdList, GraphicsPSOInit, ViewIndex, View, this, nullptr); break;
-			case 2: SetPointLightShaderTempl<2, true, false>(RHICmdList, GraphicsPSOInit, ViewIndex, View, this, nullptr); break;
-			case 3: SetPointLightShaderTempl<3, true, false>(RHICmdList, GraphicsPSOInit, ViewIndex, View, this, nullptr); break;
-			case 4: SetPointLightShaderTempl<4, true, false>(RHICmdList, GraphicsPSOInit, ViewIndex, View, this, nullptr); break;
-			case 5: SetPointLightShaderTempl<5, true, false>(RHICmdList, GraphicsPSOInit, ViewIndex, View, this, nullptr); break;
+			case 1: SetPointLightShaderTempl<1, true, false>(RHICmdList, GraphicsPSOInit, ViewIndex, View, this); break;
+			case 2: SetPointLightShaderTempl<2, true, false>(RHICmdList, GraphicsPSOInit, ViewIndex, View, this); break;
+			case 3: SetPointLightShaderTempl<3, true, false>(RHICmdList, GraphicsPSOInit, ViewIndex, View, this); break;
+			case 4: SetPointLightShaderTempl<4, true, false>(RHICmdList, GraphicsPSOInit, ViewIndex, View, this); break;
+			case 5: SetPointLightShaderTempl<5, true, false>(RHICmdList, GraphicsPSOInit, ViewIndex, View, this); break;
 			default:
 				check(0);
 			}
@@ -1471,11 +1473,11 @@ void FProjectedShadowInfo::RenderOnePassPointLightProjection(
 		{
 			switch (LocalQuality)
 			{
-			case 1: SetPointLightShaderTempl<1, false, false>(RHICmdList, GraphicsPSOInit, ViewIndex, View, this, nullptr); break;
-			case 2: SetPointLightShaderTempl<2, false, false>(RHICmdList, GraphicsPSOInit, ViewIndex, View, this, nullptr); break;
-			case 3: SetPointLightShaderTempl<3, false, false>(RHICmdList, GraphicsPSOInit, ViewIndex, View, this, nullptr); break;
-			case 4: SetPointLightShaderTempl<4, false, false>(RHICmdList, GraphicsPSOInit, ViewIndex, View, this, nullptr); break;
-			case 5: SetPointLightShaderTempl<5, false, false>(RHICmdList, GraphicsPSOInit, ViewIndex, View, this, nullptr); break;
+			case 1: SetPointLightShaderTempl<1, false, false>(RHICmdList, GraphicsPSOInit, ViewIndex, View, this); break;
+			case 2: SetPointLightShaderTempl<2, false, false>(RHICmdList, GraphicsPSOInit, ViewIndex, View, this); break;
+			case 3: SetPointLightShaderTempl<3, false, false>(RHICmdList, GraphicsPSOInit, ViewIndex, View, this); break;
+			case 4: SetPointLightShaderTempl<4, false, false>(RHICmdList, GraphicsPSOInit, ViewIndex, View, this); break;
+			case 5: SetPointLightShaderTempl<5, false, false>(RHICmdList, GraphicsPSOInit, ViewIndex, View, this); break;
 			default:
 				check(0);
 			}
@@ -1846,12 +1848,14 @@ bool FSceneRenderer::CheckForProjectedShadows( const FLightSceneInfo* LightScene
 	return false;
 }
 
+
 void FSceneRenderer::RenderShadowProjections(
 	FRDGBuilder& GraphBuilder,
-	const FShadowProjectionPassParameters& CommonPassParameters,
+	FRDGTextureRef OutputTexture,
+	const FMinimalSceneTextures& SceneTextures,
 	const FLightSceneProxy* LightSceneProxy,
-	const FHairStrandsRenderingData* HairDatas,
 	TArrayView<const FProjectedShadowInfo* const> Shadows,
+	bool bSubPixelShadow,
 	bool bProjectingForForwardShading,
 	bool bMobileModulatedProjections)
 {
@@ -1860,18 +1864,31 @@ void FSceneRenderer::RenderShadowProjections(
 	for (int32 ViewIndex = 0; ViewIndex < Views.Num(); ViewIndex++)
 	{
 		const FViewInfo& View = Views[ViewIndex];
+		const FExclusiveDepthStencil ExclusiveDepthStencil = FExclusiveDepthStencil::DepthRead_StencilWrite;
 
-		RDG_GPU_MASK_SCOPE(GraphBuilder, View.GPUMask);
-		RDG_EVENT_SCOPE_CONDITIONAL(GraphBuilder, Views.Num() > 1, "View%d", ViewIndex);
-
-		const FHairStrandsVisibilityData* HairVisibilityData = nullptr;
-		const FHairStrandsMacroGroupDatas* HairMacroGroupData = nullptr;
-		if (HairDatas)
+		if (bSubPixelShadow && !HairStrands::HasViewHairStrandsData(View))
 		{
-			HairVisibilityData = &(HairDatas->HairVisibilityViews.HairDatas[ViewIndex]);
-			HairMacroGroupData = &(HairDatas->MacroGroupsPerViews.Views[ViewIndex]);
+			continue;
+		}
+		
+		// Sanity check
+		if (bSubPixelShadow)
+		{
+			check(View.HairStrandsViewData.VisibilityData.HairOnlyDepthTexture);
 		}
 
+		FShadowProjectionPassParameters CommonPassParameters;
+		CommonPassParameters.SceneTextures = GetSceneTextureShaderParameters(SceneTextures.UniformBuffer);
+		CommonPassParameters.HairStrands = HairStrands::BindHairStrandsViewUniformParameters(View);
+		CommonPassParameters.RenderTargets[0] = FRenderTargetBinding(OutputTexture, ERenderTargetLoadAction::ELoad);
+		CommonPassParameters.RenderTargets.DepthStencil =
+			bSubPixelShadow ?
+			FDepthStencilBinding(View.HairStrandsViewData.VisibilityData.HairOnlyDepthTexture, ERenderTargetLoadAction::ELoad, ERenderTargetLoadAction::ELoad, ExclusiveDepthStencil) :
+			FDepthStencilBinding(SceneTextures.Depth.Target, ERenderTargetLoadAction::ELoad, ERenderTargetLoadAction::ELoad, ExclusiveDepthStencil);
+
+		RDG_GPU_MASK_SCOPE(GraphBuilder, View.GPUMask);
+		RDG_EVENT_SCOPE_CONDITIONAL(GraphBuilder, Views.Num() > 1, "View%d", ViewIndex);
+
 		// Project the shadow depth buffers onto the scene.
 		for (const FProjectedShadowInfo* ProjectedShadowInfo : Shadows)
 		{
@@ -1882,11 +1899,11 @@ void FSceneRenderer::RenderShadowProjections(
 				{
 					if (ProjectedShadowInfo->bOnePassPointLightShadow)
 					{
-						ProjectedShadowInfo->RenderOnePassPointLightProjection(GraphBuilder, CommonPassParameters, ViewIndex, View, LightSceneProxy, bProjectingForForwardShading, HairVisibilityData, HairMacroGroupData);
+						ProjectedShadowInfo->RenderOnePassPointLightProjection(GraphBuilder, CommonPassParameters, ViewIndex, View, LightSceneProxy, bProjectingForForwardShading, bSubPixelShadow);
 					}
 					else
 					{
-						ProjectedShadowInfo->RenderProjection(GraphBuilder, CommonPassParameters, ViewIndex, &View, LightSceneProxy, this, bProjectingForForwardShading, bMobileModulatedProjections, HairVisibilityData, HairMacroGroupData);
+						ProjectedShadowInfo->RenderProjection(GraphBuilder, CommonPassParameters, ViewIndex, &View, LightSceneProxy, this, bProjectingForForwardShading, bMobileModulatedProjections, bSubPixelShadow);
 					}
 				}
 			}
@@ -1900,7 +1917,6 @@ void FDeferredShadingSceneRenderer::RenderShadowProjections(
 	FRDGTextureRef ScreenShadowMaskTexture,
 	FRDGTextureRef ScreenShadowMaskSubPixelTexture,
 	const FLightSceneInfo* LightSceneInfo,
-	const FHairStrandsRenderingData* HairDatas,
 	bool bProjectingForForwardShading)
 {
 	const FVisibleLightInfo& VisibleLightInfo = VisibleLightInfos[LightSceneInfo->Id];
@@ -1934,46 +1950,35 @@ void FDeferredShadingSceneRenderer::RenderShadowProjections(
 		}
 	}
 
-	FRDGTextureRef HairCatgorizationTexture = HairDatas && HairDatas->HairVisibilityViews.HairDatas.Num() > 0 ? HairDatas->HairVisibilityViews.HairDatas[0].CategorizationTexture : nullptr;
 	if (NormalShadows.Num() > 0)
 	{
-		const auto RenderNormalShadows = [&](FRDGTextureRef OutputTexture, FExclusiveDepthStencil ExclusiveDepthStencil, bool bSubPixel)
+		const auto RenderNormalShadows = [&](FRDGTextureRef OutputTexture, bool bSubPixel)
 		{
 			FString LightNameWithLevel;
 			GetLightNameForDrawEvent(LightSceneProxy, LightNameWithLevel);
 			RDG_EVENT_SCOPE(GraphBuilder, "%s", *LightNameWithLevel);
 
-			FShadowProjectionPassParameters CommonPassParameters;
-			CommonPassParameters.SceneTextures = GetSceneTextureShaderParameters(SceneTextures.UniformBuffer);
-			CommonPassParameters.HairCategorizationTexture = bSubPixel && HairCatgorizationTexture ? HairCatgorizationTexture : nullptr;
-			CommonPassParameters.RenderTargets[0] = FRenderTargetBinding(OutputTexture, ERenderTargetLoadAction::ELoad);
-			CommonPassParameters.RenderTargets.DepthStencil = 
-				bSubPixel && HairDatas && HairDatas->HairVisibilityViews.HairDatas.Num() > 0 ?
-				FDepthStencilBinding(HairDatas->HairVisibilityViews.HairDatas[0].HairOnlyDepthTexture, ERenderTargetLoadAction::ELoad, ERenderTargetLoadAction::ELoad, ExclusiveDepthStencil) :
-				FDepthStencilBinding(SceneTextures.Depth.Target, ERenderTargetLoadAction::ELoad, ERenderTargetLoadAction::ELoad, ExclusiveDepthStencil);
-
 			const bool bMobileModulatedProjections = false;
 			FSceneRenderer::RenderShadowProjections(
 				GraphBuilder,
-				CommonPassParameters,
+				OutputTexture,
+				SceneTextures,
 				LightSceneProxy,
-				bSubPixel ? HairDatas : nullptr,
 				NormalShadows,
+				bSubPixel,
 				bProjectingForForwardShading,
 				bMobileModulatedProjections);
 		};
 
 		{
 			RDG_EVENT_SCOPE(GraphBuilder, "Shadows");
-			RenderNormalShadows(ScreenShadowMaskTexture, FExclusiveDepthStencil::DepthRead_StencilWrite, false);
+			RenderNormalShadows(ScreenShadowMaskTexture, false);
 		}
 
-		if (ScreenShadowMaskSubPixelTexture && HairCatgorizationTexture)
+		if (ScreenShadowMaskSubPixelTexture)
 		{
 			RDG_EVENT_SCOPE(GraphBuilder, "SubPixelShadows");
-
-			// Sub-pixel shadows don't use stencil.
-			RenderNormalShadows(ScreenShadowMaskSubPixelTexture, FExclusiveDepthStencil::DepthRead_StencilWrite, true);
+			RenderNormalShadows(ScreenShadowMaskSubPixelTexture, true);
 		}
 	}
 
@@ -2131,7 +2136,6 @@ void FDeferredShadingSceneRenderer::RenderDeferredShadowProjections(
 	const FLightSceneInfo* LightSceneInfo,
 	FRDGTextureRef ScreenShadowMaskTexture,
 	FRDGTextureRef ScreenShadowMaskSubPixelTexture,
-	const FHairStrandsRenderingData* HairDatas,
 	bool& bInjectedTranslucentVolume)
 {
 	SCOPED_NAMED_EVENT(FDeferredShadingSceneRenderer_RenderShadowProjections, FColor::Emerald);
@@ -2142,7 +2146,7 @@ void FDeferredShadingSceneRenderer::RenderDeferredShadowProjections(
 	const FVisibleLightInfo& VisibleLightInfo = VisibleLightInfos[LightSceneInfo->Id];
 	
 	const bool bProjectingForForwardShading = false;
-	RenderShadowProjections(GraphBuilder, SceneTextures, ScreenShadowMaskTexture, ScreenShadowMaskSubPixelTexture, LightSceneInfo, HairDatas, bProjectingForForwardShading);
+	RenderShadowProjections(GraphBuilder, SceneTextures, ScreenShadowMaskTexture, ScreenShadowMaskSubPixelTexture, LightSceneInfo, bProjectingForForwardShading);
 
 	// Perform injection on translucent lighting volume
 	{
@@ -2212,9 +2216,9 @@ void FDeferredShadingSceneRenderer::RenderDeferredShadowProjections(
 	}
 
 	// Inject deep shadow mask
-	if (HairDatas)
+	if (HairStrands::HasViewHairStrandsData(Views))
 	{
-		RenderHairStrandsShadowMask(GraphBuilder, Views, LightSceneInfo, HairDatas, ScreenShadowMaskTexture);
+		RenderHairStrandsShadowMask(GraphBuilder, Views, LightSceneInfo, ScreenShadowMaskTexture);
 	}
 }
 
@@ -2244,10 +2248,6 @@ void FMobileSceneRenderer::RenderModulatedShadowProjections(FRHICommandListImmed
 			if (VisibleLightInfo.ShadowsToProject.Num() > 0)
 			{			
 				LightSceneProxy->SetScissorRect(RHICmdList, View, View.ViewRect);
-
-				const FHairStrandsVisibilityData* HairVisibilityData = nullptr;
-				const FHairStrandsMacroGroupDatas* HairMacroGroupData = nullptr;
-
 				TArrayView<const FProjectedShadowInfo* const> Shadows = VisibleLightInfo.ShadowsToProject;
 
 				// Project the shadow depth buffers onto the scene.
diff --git a/Engine/Source/Runtime/Renderer/Private/ShadowRendering.h b/Engine/Source/Runtime/Renderer/Private/ShadowRendering.h
--- a/Engine/Source/Runtime/Renderer/Private/ShadowRendering.h
+++ b/Engine/Source/Runtime/Renderer/Private/ShadowRendering.h
@@ -509,8 +509,7 @@ public:
 		const FSceneRenderer* SceneRender,
 		bool bProjectingForForwardShading,
 		bool bMobileModulatedProjections,
-		const FHairStrandsVisibilityData* HairVisibilityData,
-		const struct FHairStrandsMacroGroupDatas* HairMacroGroupData) const;
+		bool bSubPixelShadow) const;
 
 	void RenderProjectionInternal(
 		FRHICommandListImmediate& RHICmdList,
@@ -520,8 +519,8 @@ public:
 		const FSceneRenderer* SceneRender,
 		bool bProjectingForForwardShading,
 		bool bMobileModulatedProjections,
-		const FHairStrandsVisibilityData* HairVisibilityData,
-		const FInstanceCullingDrawParams& InstanceCullingDrawParams) const;
+		const FInstanceCullingDrawParams& InstanceCullingDrawParams, 
+		FRHIUniformBuffer* HairStrandsUniformBuffer) const;
 
 	/**
 	 * Projects the mobile modulated shadow onto the scene for a particular view.
@@ -564,8 +563,7 @@ public:
 		const FViewInfo& View,
 		const FLightSceneProxy* LightSceneProxy,
 		bool bProjectingForForwardShading,
-		const FHairStrandsVisibilityData* HairVisibilityData,
-		const struct FHairStrandsMacroGroupDatas* HairMacroGroupData) const;
+		bool bSubPixelShadow) const;
 
 	/**
 	 * Renders the projected shadow's frustum wireframe with the given FPrimitiveDrawInterface.
@@ -967,14 +965,13 @@ public:
 		InvFadePlaneLength.Bind(ParameterMap,TEXT("InvFadePlaneLength"));
 		ShadowTileOffsetAndSizeParam.Bind(ParameterMap, TEXT("ShadowTileOffsetAndSize"));
 		LightPositionOrDirection.Bind(ParameterMap, TEXT("LightPositionOrDirection"));
-		HairCategorizationTexture.Bind(ParameterMap, TEXT("HairCategorizationTexture"));
 		PerObjectShadowFadeStart.Bind(ParameterMap, TEXT("PerObjectShadowFadeStart"));
 		InvPerObjectShadowFadeLength.Bind(ParameterMap, TEXT("InvPerObjectShadowFadeLength"));
 		ShadowNearAndFarDepth.Bind(ParameterMap, TEXT("ShadowNearAndFarDepth"));
 		bCascadeUseFadePlane.Bind(ParameterMap, TEXT("bCascadeUseFadePlane"));
 	}
 
-	void Set(FRHICommandList& RHICmdList, FShader* Shader, const FSceneView& View, const FProjectedShadowInfo* ShadowInfo, const FHairStrandsVisibilityData* HairVisibilityData, bool bModulatedShadows, bool bUseFadePlane)
+	void Set(FRHICommandList& RHICmdList, FShader* Shader, const FSceneView& View, const FProjectedShadowInfo* ShadowInfo, bool bModulatedShadows, bool bUseFadePlane, bool SubPixelShadow)
 	{
 		FRHIPixelShader* ShaderRHI = RHICmdList.GetBoundPixelShader();
 
@@ -1061,10 +1058,8 @@ public:
 			SetShaderValue(RHICmdList, ShaderRHI, LightPositionOrDirection, bIsDirectional ? FVector4(LightDirection,0) : FVector4(LightPosition,1));
 		}
 
-		if (HairVisibilityData && HairVisibilityData->CategorizationTexture)
+		if (SubPixelShadow)
 		{
-			SetTextureParameter(RHICmdList, ShaderRHI, HairCategorizationTexture, HairVisibilityData->CategorizationTexture->GetRHI());
-
 			float DeviceZNear = 1;
 			float DeviceZFar = 0;
 			const bool bIsCascadedShadow = ShadowInfo->bDirectionalLight && !(ShadowInfo->bPerObjectOpaqueShadow || ShadowInfo->bPreShadow);
@@ -1098,7 +1093,6 @@ private:
 	LAYOUT_FIELD(FShaderParameter, InvFadePlaneLength);
 	LAYOUT_FIELD(FShaderParameter, ShadowTileOffsetAndSizeParam);
 	LAYOUT_FIELD(FShaderParameter, LightPositionOrDirection);
-	LAYOUT_FIELD(FShaderResourceParameter, HairCategorizationTexture);
 	LAYOUT_FIELD(FShaderParameter, PerObjectShadowFadeStart);
 	LAYOUT_FIELD(FShaderParameter, InvPerObjectShadowFadeLength);
 	LAYOUT_FIELD(FShaderParameter, ShadowNearAndFarDepth);
@@ -1163,14 +1157,13 @@ public:
 		FRHICommandList& RHICmdList, 
 		int32 ViewIndex,
 		const FSceneView& View,
-		const FHairStrandsVisibilityData* HairVisibilityData,
 		const FProjectedShadowInfo* ShadowInfo)
 	{
 		FRHIPixelShader* ShaderRHI = RHICmdList.GetBoundPixelShader();
 
 		const bool bUseFadePlaneEnable = ShadowInfo->CascadeSettings.FadePlaneLength > 0;
 
-		ProjectionParameters.Set(RHICmdList, this, View, ShadowInfo, HairVisibilityData, bModulatedShadows, bUseFadePlaneEnable);
+		ProjectionParameters.Set(RHICmdList, this, View, ShadowInfo, bModulatedShadows, bUseFadePlaneEnable, SubPixelShadow);
 		const FLightSceneProxy& LightProxy = *(ShadowInfo->GetLightSceneInfo().Proxy);
 
 		SetShaderValue(RHICmdList, ShaderRHI, ShadowFadeFraction, ShadowInfo->FadeAlphas[ViewIndex] );
@@ -1245,10 +1238,9 @@ public:
 		FRHICommandList& RHICmdList,
 		int32 ViewIndex,
 		const FSceneView& View,
-		const FHairStrandsVisibilityData* HairVisibilityData,
 		const FProjectedShadowInfo* ShadowInfo)
 	{
-		TShadowProjectionPS<Quality, false, true>::SetParameters(RHICmdList, ViewIndex, View, HairVisibilityData, ShadowInfo);
+		TShadowProjectionPS<Quality, false, true>::SetParameters(RHICmdList, ViewIndex, View, ShadowInfo);
 		FRHIPixelShader* ShaderRHI = RHICmdList.GetBoundPixelShader();
 		SetShaderValue(RHICmdList, ShaderRHI, ModulatedShadowColorParameter, ShadowInfo->GetLightSceneInfo().Proxy->GetModulatedShadowColor());
 	}
@@ -1313,10 +1305,9 @@ public:
 		FRHICommandList& RHICmdList, 
 		int32 ViewIndex,
 		const FSceneView& View,
-		const FHairStrandsVisibilityData* HairVisibilityData,
 		const FProjectedShadowInfo* ShadowInfo)
 	{
-		TShadowProjectionPS<Quality>::SetParameters(RHICmdList, ViewIndex, View, HairVisibilityData, ShadowInfo);
+		TShadowProjectionPS<Quality>::SetParameters(RHICmdList, ViewIndex, View, ShadowInfo);
 
 		FTranslucentSelfShadowUniformParameters TranslucentSelfShadowUniformParameters;
 		SetupTranslucentSelfShadowUniformParameters(ShadowInfo, TranslucentSelfShadowUniformParameters);
@@ -1458,6 +1449,7 @@ public:
 	TOnePassPointShadowProjectionPS(const ShaderMetaType::CompiledShaderInitializerType& Initializer):
 		FGlobalShader(Initializer)
 	{
+		HairStrandsParameters.Bind(Initializer.ParameterMap, FHairStrandsViewUniformParameters::StaticStructMetadata.GetShaderVariableName());
 		StrataGlobalParameters.Bind(Initializer.ParameterMap, FStrataGlobalUniformParameters::StaticStructMetadata.GetShaderVariableName());
 		OnePassShadowParameters.Bind(Initializer.ParameterMap);
 		ShadowDepthTextureSampler.Bind(Initializer.ParameterMap,TEXT("ShadowDepthTextureSampler"));
@@ -1467,7 +1459,6 @@ public:
 		PointLightDepthBias.Bind(Initializer.ParameterMap,TEXT("PointLightDepthBias"));
 		PointLightProjParameters.Bind(Initializer.ParameterMap, TEXT("PointLightProjParameters"));
 		TransmissionProfilesTexture.Bind(Initializer.ParameterMap, TEXT("SSProfilesTexture"));
-		HairCategorizationTexture.Bind(Initializer.ParameterMap, TEXT("HairCategorizationTexture"));
 	}
 
 	static void ModifyCompilationEnvironment(const FGlobalShaderPermutationParameters& Parameters, FShaderCompilerEnvironment& OutEnvironment)
@@ -1488,8 +1479,8 @@ public:
 		FRHICommandList& RHICmdList, 
 		int32 ViewIndex,
 		const FViewInfo& View,
-		const FHairStrandsVisibilityData* HairVisibilityData,
-		const FProjectedShadowInfo* ShadowInfo)
+		const FProjectedShadowInfo* ShadowInfo,
+		FRHIUniformBuffer* HairStrandsUniformBuffer)
 	{
 		FRHIPixelShader* ShaderRHI = RHICmdList.GetBoundPixelShader();
 
@@ -1509,9 +1500,9 @@ public:
 		SetShaderValue(RHICmdList, ShaderRHI, PointLightDepthBias, FVector(ShadowInfo->GetShaderDepthBias(), ShadowInfo->GetShaderSlopeDepthBias(), ShadowInfo->GetShaderMaxSlopeDepthBias()));
 		SetShaderValue(RHICmdList, ShaderRHI, PointLightProjParameters, FVector2D(projParam.X, projParam.Y));
 
-		if (bUseSubPixel && HairVisibilityData && HairVisibilityData->CategorizationTexture)
+		if (HairStrandsParameters.IsBound())
 		{
-			SetTextureParameter(RHICmdList, ShaderRHI, HairCategorizationTexture, HairVisibilityData->CategorizationTexture->GetRHI());
+			SetUniformBufferParameter(RHICmdList, ShaderRHI, HairStrandsParameters, HairStrandsUniformBuffer);
 		}
 
 		if (StrataGlobalParameters.IsBound())
@@ -1560,7 +1551,7 @@ private:
 	LAYOUT_FIELD(FShaderParameter, PointLightDepthBias);
 	LAYOUT_FIELD(FShaderParameter, PointLightProjParameters);
 	LAYOUT_FIELD(FShaderResourceParameter, TransmissionProfilesTexture);
-	LAYOUT_FIELD(FShaderResourceParameter, HairCategorizationTexture);
+	LAYOUT_FIELD(FShaderUniformBufferParameter, HairStrandsParameters);
 	LAYOUT_FIELD(FShaderUniformBufferParameter, StrataGlobalParameters);
 };
 
@@ -1626,10 +1617,9 @@ public:
 		FRHICommandList& RHICmdList,
 		int32 ViewIndex,
 		const FSceneView& View,
-		const FHairStrandsVisibilityData* HairVisibilityData,
 		const FProjectedShadowInfo* ShadowInfo)
 	{
-		TShadowProjectionPS<Quality, bUseFadePlane>::SetParameters(RHICmdList, ViewIndex, View, HairVisibilityData, ShadowInfo);
+		TShadowProjectionPS<Quality, bUseFadePlane>::SetParameters(RHICmdList, ViewIndex, View, ShadowInfo);
 
 		FRHIPixelShader* ShaderRHI = RHICmdList.GetBoundPixelShader();
 
@@ -1684,12 +1674,11 @@ public:
 		FRHICommandList& RHICmdList,
 		int32 ViewIndex,
 		const FSceneView& View,
-		const FHairStrandsVisibilityData* HairVisibilityData,
 		const FProjectedShadowInfo* ShadowInfo)
 	{
 		check(ShadowInfo->GetLightSceneInfo().Proxy->GetLightType() == LightType_Spot);
 
-		TShadowProjectionPS<Quality, bUseFadePlane>::SetParameters(RHICmdList, ViewIndex, View, HairVisibilityData, ShadowInfo);
+		TShadowProjectionPS<Quality, bUseFadePlane>::SetParameters(RHICmdList, ViewIndex, View, ShadowInfo);
 
 		FRHIPixelShader* ShaderRHI = RHICmdList.GetBoundPixelShader();
 
diff --git a/Engine/Source/Runtime/Renderer/Public/HairStrandsInterface.h b/Engine/Source/Runtime/Renderer/Public/HairStrandsInterface.h
--- a/Engine/Source/Runtime/Renderer/Public/HairStrandsInterface.h
+++ b/Engine/Source/Runtime/Renderer/Public/HairStrandsInterface.h
@@ -98,20 +98,6 @@ enum class EHairDebugMode : uint8
 RENDERER_API EHairStrandsDebugMode GetHairStrandsDebugStrandsMode();
 RENDERER_API EHairDebugMode GetHairStrandsDebugMode();
 
-struct FMinHairRadiusAtDepth1
-{
-	float Primary = 1;
-	float Velocity = 1;
-	float Stable = 1;
-};
-
-/// Compute the strand radius at a distance of 1 meter
-RENDERER_API FMinHairRadiusAtDepth1 ComputeMinStrandRadiusAtDepth1(
-	const FIntPoint& Resolution,
-	const float FOV,
-	const uint32 SampleCount,
-	const float OverrideStrandHairRasterizationScale);
-
 struct FHairStrandClusterCullingData;
 struct IPooledRenderTarget;
 struct FRWBuffer;
