Build and start ray tracing mesh batch processing tasks while looping over RelevantPrimitives in GatherRayTracingWorldInstances(), rather than only in the end.

This reduces the latency as we start the parallel work a bit earlier, reducing the possibility of a stall in DispatchRayTracingWorldUpdates(), which waits for AddRayTracingMeshBatchTaskList.

diff --git a/Engine/Source/Runtime/Renderer/Private/DeferredShadingRenderer.cpp b/Engine/Source/Runtime/Renderer/Private/DeferredShadingRenderer.cpp
--- a/Engine/Source/Runtime/Renderer/Private/DeferredShadingRenderer.cpp
+++ b/Engine/Source/Runtime/Renderer/Private/DeferredShadingRenderer.cpp
@@ -813,16 +813,94 @@ bool FDeferredShadingSceneRenderer::GatherRayTracingWorldInstances(FRHICommandLi
 
 		ReferenceView.RayTracingGeometryInstances.Reserve(RelevantPrimitives.Num());
 
-		if (bParallelMeshBatchSetup)
-		{
-			ReferenceView.AddRayTracingMeshBatchData.Reserve(RelevantPrimitives.Num());
-		}
-
 		for (int32 ViewIndex = 0; ViewIndex < Views.Num(); ViewIndex++)
 		{
 			Views[ViewIndex].RayTracingGeometryInstances.Reserve(RelevantPrimitives.Num());
 		}
 
+		struct FRayTracingMeshBatchWorkItem
+		{
+			const FPrimitiveSceneProxy* SceneProxy = nullptr;
+			TArray<FMeshBatch> MeshBatchesOwned;
+			TArrayView<const FMeshBatch> MeshBatchesView;
+			uint32 InstanceIndex = 0;
+
+			TArrayView<const FMeshBatch> GetMeshBatches() const
+			{
+				if (MeshBatchesOwned.Num())
+				{
+					check(MeshBatchesView.Num() == 0);
+					return TArrayView<const FMeshBatch>(MeshBatchesOwned);
+				}
+				else
+				{
+					check(MeshBatchesOwned.Num() == 0);
+					return MeshBatchesView;
+				}
+			}
+		};
+
+		static constexpr uint32 MaxWorkItemsPerPage = 128; // Try to keep individual pages small to avoid slow-path memory allocations
+		struct FRayTracingMeshBatchTaskPage
+		{
+			FRayTracingMeshBatchWorkItem WorkItems[MaxWorkItemsPerPage];
+			uint32 NumWorkItems = 0;
+			FRayTracingMeshBatchTaskPage* Next = nullptr;
+		};
+
+		FRayTracingMeshBatchTaskPage* MeshBatchTaskHead = nullptr;
+		FRayTracingMeshBatchTaskPage* MeshBatchTaskPage = nullptr;
+		uint32 NumPendingMeshBatches = 0;
+		const uint32 RayTracingParallelMeshBatchSize = GRayTracingParallelMeshBatchSize;
+
+		auto KickRayTracingMeshBatchTask = [&ReferenceView, &MeshBatchTaskHead, &MeshBatchTaskPage, &NumPendingMeshBatches, Scene = this->Scene]()
+		{
+			if (MeshBatchTaskHead)
+			{
+				FDynamicRayTracingMeshCommandStorage* TaskDynamicCommandStorage = new(FMemStack::Get()) FDynamicRayTracingMeshCommandStorage;
+				ReferenceView.DynamicRayTracingMeshCommandStoragePerTask.Add(TaskDynamicCommandStorage);
+
+				FRayTracingMeshCommandOneFrameArray* TaskVisibleCommands = new(FMemStack::Get()) FRayTracingMeshCommandOneFrameArray;
+				ReferenceView.VisibleRayTracingMeshCommandsPerTask.Add(TaskVisibleCommands);
+
+				const FViewInfo* ReferenceViewPtr = &ReferenceView;
+
+				ReferenceView.AddRayTracingMeshBatchTaskList.Add(FFunctionGraphTask::CreateAndDispatchWhenReady(
+					[TaskDataHead = MeshBatchTaskHead, ReferenceViewPtr, Scene, TaskDynamicCommandStorage, TaskVisibleCommands]()
+				{
+					FTaskTagScope TaskTagScope(ETaskTag::EParallelRenderingThread);
+					TRACE_CPUPROFILER_EVENT_SCOPE(RayTracingMeshBatchTask);
+					FRayTracingMeshBatchTaskPage* Page = TaskDataHead;
+					while (Page)
+					{
+						for (uint32 ItemIndex = 0; ItemIndex < Page->NumWorkItems; ++ItemIndex)
+						{
+							const FRayTracingMeshBatchWorkItem& WorkItem = Page->WorkItems[ItemIndex];
+							TArrayView<const FMeshBatch> MeshBatches = WorkItem.GetMeshBatches();
+							for (int32 SegmentIndex = 0; SegmentIndex < MeshBatches.Num(); SegmentIndex++)
+							{
+								const FMeshBatch& MeshBatch = MeshBatches[SegmentIndex];
+								FDynamicRayTracingMeshCommandContext CommandContext(
+									*TaskDynamicCommandStorage, *TaskVisibleCommands,
+									SegmentIndex, WorkItem.InstanceIndex);
+								FMeshPassProcessorRenderState PassDrawRenderState(Scene->UniformBuffers.ViewUniformBuffer);
+								FRayTracingMeshProcessor RayTracingMeshProcessor(&CommandContext, Scene, ReferenceViewPtr, PassDrawRenderState);
+								RayTracingMeshProcessor.AddMeshBatch(MeshBatch, 1, WorkItem.SceneProxy);
+							}
+						}
+						FRayTracingMeshBatchTaskPage* NextPage = Page->Next;
+						Page->~FRayTracingMeshBatchTaskPage();
+						Page = NextPage;
+					}
+				}, TStatId(), nullptr, ENamedThreads::AnyThread));
+			}
+
+			MeshBatchTaskHead = nullptr;
+			MeshBatchTaskPage = nullptr;
+			NumPendingMeshBatches = 0;
+		};
+
+
 		for (const FRelevantPrimitive& RelevantPrimitive : RelevantPrimitives)
 		{
 			const uint8 RayTracedMeshElementsMask = RelevantPrimitive.RayTracedMeshElementsMask;
@@ -926,16 +1004,41 @@ bool FDeferredShadingSceneRenderer::GatherRayTracingWorldInstances(FRHICommandLi
 
 					if (bParallelMeshBatchSetup)
 					{
+						if (NumPendingMeshBatches >= RayTracingParallelMeshBatchSize)
+						{
+							KickRayTracingMeshBatchTask();
+						}
+
+						if (MeshBatchTaskPage == nullptr || MeshBatchTaskPage->NumWorkItems == MaxWorkItemsPerPage)
+						{
+							FRayTracingMeshBatchTaskPage* NextPage = new(FMemStack::Get()) FRayTracingMeshBatchTaskPage;
+							if (MeshBatchTaskHead == nullptr)
+							{
+								MeshBatchTaskHead = NextPage;
+							}
+							if (MeshBatchTaskPage)
+							{
+								MeshBatchTaskPage->Next = NextPage;
+							}
+							MeshBatchTaskPage = NextPage;
+						}
+
+						FRayTracingMeshBatchWorkItem& WorkItem = MeshBatchTaskPage->WorkItems[MeshBatchTaskPage->NumWorkItems];
+						MeshBatchTaskPage->NumWorkItems++;
+
+						NumPendingMeshBatches += Instance.GetMaterials().Num();
+
 						if (Instance.OwnsMaterials())
 						{
-							check(Instance.MaterialsView.Num() == 0);
-							ReferenceView.AddRayTracingMeshBatchData.Emplace(Instance.Materials, SceneProxy, InstanceIndex);
+							Swap(WorkItem.MeshBatchesOwned, Instance.Materials);
 						}
 						else
 						{
-							check(Instance.Materials.Num() == 0);
-							ReferenceView.AddRayTracingMeshBatchData.Emplace(Instance.MaterialsView, SceneProxy, InstanceIndex);
+							WorkItem.MeshBatchesView = Instance.MaterialsView;
 						}
+
+						WorkItem.SceneProxy = SceneProxy;
+						WorkItem.InstanceIndex = InstanceIndex;
 					}
 					else
 					{
@@ -946,7 +1049,6 @@ bool FDeferredShadingSceneRenderer::GatherRayTracingWorldInstances(FRHICommandLi
 							FDynamicRayTracingMeshCommandContext CommandContext(ReferenceView.DynamicRayTracingMeshCommandStorage, ReferenceView.VisibleRayTracingMeshCommands, SegmentIndex, InstanceIndex);
 							FMeshPassProcessorRenderState PassDrawRenderState(Scene->UniformBuffers.ViewUniformBuffer);
 							FRayTracingMeshProcessor RayTracingMeshProcessor(&CommandContext, Scene, &ReferenceView, PassDrawRenderState);
-
 							RayTracingMeshProcessor.AddMeshBatch(MeshBatch, 1, SceneProxy);
 						}
 					}
@@ -965,52 +1067,8 @@ bool FDeferredShadingSceneRenderer::GatherRayTracingWorldInstances(FRHICommandLi
 				}
 			}
 		}
-	}
 
-	//
-
-	if (ReferenceView.AddRayTracingMeshBatchData.Num() > 0)
-	{
-		const uint32 NumTotalItems = ReferenceView.AddRayTracingMeshBatchData.Num();
-		const uint32 TargetItemsPerTask = GRayTracingParallelMeshBatchSize;// 1024; // Granularity based on profiling several scenes
-		const uint32 NumTasks = FMath::Max(1u, FMath::DivideAndRoundUp(NumTotalItems, TargetItemsPerTask));
-		const uint32 BatchSize = FMath::DivideAndRoundUp(NumTotalItems, NumTasks);
-
-		ReferenceView.DynamicRayTracingMeshCommandStorageParallel.Init(FDynamicRayTracingMeshCommandStorage(), NumTasks);
-		ReferenceView.VisibleRayTracingMeshCommandsParallel.Init(FRayTracingMeshCommandOneFrameArray(), NumTasks);
-
-		for (uint32 Batch = 0; Batch < NumTasks; Batch++)
-		{
-			uint32 BatchStart = Batch * BatchSize;
-			uint32 BatchEnd = FMath::Min(BatchStart + BatchSize, (uint32) ReferenceView.AddRayTracingMeshBatchData.Num());
-
-			ReferenceView.DynamicRayTracingMeshCommandStorageParallel[Batch].RayTracingMeshCommands.Reserve(Scene->Primitives.Num());
-			ReferenceView.VisibleRayTracingMeshCommandsParallel[Batch].Reserve(Scene->Primitives.Num());
-
-			ReferenceView.AddRayTracingMeshBatchTaskList.Add(FFunctionGraphTask::CreateAndDispatchWhenReady(
-				[&ReferenceView, Scene = this->Scene, Batch, BatchStart, BatchEnd]()
-			{
-				FTaskTagScope TaskTagScope(ETaskTag::EParallelRenderingThread);
-				TRACE_CPUPROFILER_EVENT_SCOPE(RayTracingMeshBatchTask);
-
-				for (uint32 Index = BatchStart; Index < BatchEnd; Index++)
-				{
-					FRayTracingMeshBatchWorkItem& MeshBatchJob = ReferenceView.AddRayTracingMeshBatchData[Index];
-					TArrayView<const FMeshBatch> MeshBatches = MeshBatchJob.GetMeshBatches();
-					for (int32 SegmentIndex = 0; SegmentIndex < MeshBatches.Num(); SegmentIndex++)
-					{
-						const FMeshBatch& MeshBatch = MeshBatches[SegmentIndex];
-						const FPrimitiveSceneProxy* SceneProxy = MeshBatchJob.SceneProxy;
-						const uint32 InstanceIndex = MeshBatchJob.InstanceIndex;
-						FDynamicRayTracingMeshCommandContext CommandContext(ReferenceView.DynamicRayTracingMeshCommandStorageParallel[Batch], ReferenceView.VisibleRayTracingMeshCommandsParallel[Batch], SegmentIndex, InstanceIndex);
-						FMeshPassProcessorRenderState PassDrawRenderState(Scene->UniformBuffers.ViewUniformBuffer);
-						FRayTracingMeshProcessor RayTracingMeshProcessor(&CommandContext, Scene, &ReferenceView, PassDrawRenderState);
-						RayTracingMeshProcessor.AddMeshBatch(MeshBatch, 1, SceneProxy);
-					}
-				}
-			},
-				TStatId(), nullptr, ENamedThreads::AnyThread));
-		}
+		KickRayTracingMeshBatchTask();
 	}
 
 	//
@@ -1154,13 +1212,12 @@ bool FDeferredShadingSceneRenderer::DispatchRayTracingWorldUpdates(FRDGBuilder&
 
 		FTaskGraphInterface::Get().WaitUntilTasksComplete(ReferenceView.AddRayTracingMeshBatchTaskList, ENamedThreads::GetRenderThread_Local());
 
-		for (int32 Batch = 0; Batch < ReferenceView.AddRayTracingMeshBatchTaskList.Num(); Batch++)
+		for (int32 TaskIndex = 0; TaskIndex < ReferenceView.AddRayTracingMeshBatchTaskList.Num(); TaskIndex++)
 		{
-			ReferenceView.VisibleRayTracingMeshCommands.Append(ReferenceView.VisibleRayTracingMeshCommandsParallel[Batch]);
+			ReferenceView.VisibleRayTracingMeshCommands.Append(*ReferenceView.VisibleRayTracingMeshCommandsPerTask[TaskIndex]);
 		}
 
 		ReferenceView.AddRayTracingMeshBatchTaskList.Empty();
-		ReferenceView.AddRayTracingMeshBatchData.Empty();
 	}
 
 	RDG_GPU_MASK_SCOPE(GraphBuilder, FRHIGPUMask::All());
diff --git a/Engine/Source/Runtime/Renderer/Private/SceneRendering.cpp b/Engine/Source/Runtime/Renderer/Private/SceneRendering.cpp
--- a/Engine/Source/Runtime/Renderer/Private/SceneRendering.cpp
+++ b/Engine/Source/Runtime/Renderer/Private/SceneRendering.cpp
@@ -990,6 +990,18 @@ FViewInfo::~FViewInfo()
 
 	//this uses memstack allocation for strongrefs, so we need to manually empty to get the destructor called to not leak the uniformbuffers stored here.
 	TranslucentSelfShadowUniformBufferMap.Empty();
+
+#if RHI_RAYTRACING
+	// Per-task structures are allocated using memstack so we have to call destructors explicitly.
+	for (FRayTracingMeshCommandOneFrameArray* It : VisibleRayTracingMeshCommandsPerTask)
+	{
+		It->~FRayTracingMeshCommandOneFrameArray();
+	}
+	for (FDynamicRayTracingMeshCommandStorage* It : DynamicRayTracingMeshCommandStoragePerTask)
+	{
+		It->~FDynamicRayTracingMeshCommandStorage();
+	}
+#endif // RHI_RAYTRACING
 }
 
 #if DO_CHECK || USING_CODE_ANALYSIS
diff --git a/Engine/Source/Runtime/Renderer/Private/SceneRendering.h b/Engine/Source/Runtime/Renderer/Private/SceneRendering.h
--- a/Engine/Source/Runtime/Renderer/Private/SceneRendering.h
+++ b/Engine/Source/Runtime/Renderer/Private/SceneRendering.h
@@ -1058,47 +1058,6 @@ public:
 typedef TArray<FViewCommands, TInlineAllocator<4>> FViewVisibleCommandsPerView;
 
 #if RHI_RAYTRACING
-struct FRayTracingMeshBatchWorkItem
-{
-	FRayTracingMeshBatchWorkItem(TArray<FMeshBatch>& InBatches, FPrimitiveSceneProxy* InSceneProxy, uint32 InInstanceIndex) 
-		: SceneProxy(InSceneProxy)
-		, InstanceIndex(InInstanceIndex)
-	{
-		Swap(MeshBatchesOwned, InBatches);
-	}
-
-	FRayTracingMeshBatchWorkItem(TArrayView<const FMeshBatch>& InBatches, FPrimitiveSceneProxy* InSceneProxy, uint32 InInstanceIndex) 
-		: SceneProxy(InSceneProxy)
-		, InstanceIndex(InInstanceIndex)
-	{
-		MeshBatchesView = InBatches;
-	}
-
-	bool OwnsMeshBatches() const
-	{
-		return MeshBatchesOwned.Num() != 0;
-	}
-
-	TArrayView<const FMeshBatch> GetMeshBatches() const
-	{
-		if (OwnsMeshBatches())
-		{
-			check(MeshBatchesView.Num() == 0);
-			return TArrayView<const FMeshBatch>(MeshBatchesOwned);
-		}
-		else
-		{
-			check(MeshBatchesOwned.Num() == 0);
-			return MeshBatchesView;
-		}
-	}
-
-	TArray<FMeshBatch> MeshBatchesOwned;
-	TArrayView<const FMeshBatch> MeshBatchesView;
-
-	FPrimitiveSceneProxy* SceneProxy = nullptr;
-	uint32 InstanceIndex = 0;
-};
 
 /** Convenience struct for all lighting data used by ray tracing effects using RayTracingLightingCommon.ush */
 struct FRayTracingLightData
@@ -1261,10 +1220,8 @@ public:
 	FDynamicRayTracingMeshCommandStorage DynamicRayTracingMeshCommandStorage;
 
 	FGraphEventArray AddRayTracingMeshBatchTaskList;
-	TArray<FRayTracingMeshBatchWorkItem, SceneRenderingAllocator> AddRayTracingMeshBatchData;
-
-	TArray<FRayTracingMeshCommandOneFrameArray, SceneRenderingAllocator> VisibleRayTracingMeshCommandsParallel;
-	TArray<FDynamicRayTracingMeshCommandStorage, SceneRenderingAllocator> DynamicRayTracingMeshCommandStorageParallel;
+	TArray<FRayTracingMeshCommandOneFrameArray*, SceneRenderingAllocator> VisibleRayTracingMeshCommandsPerTask;
+	TArray<FDynamicRayTracingMeshCommandStorage*, SceneRenderingAllocator> DynamicRayTracingMeshCommandStoragePerTask;
 #endif
 
 	// Used by mobile renderer to determine whether static meshes will be rendered with CSM shaders or not.
