Allocate FRayTracingMeshCommandOneFrameArray memory on render thread memstack. Fixes memory leak introduced in CL 15771478.

diff --git a/Engine/Source/Runtime/Renderer/Private/DeferredShadingRenderer.cpp b/Engine/Source/Runtime/Renderer/Private/DeferredShadingRenderer.cpp
--- a/Engine/Source/Runtime/Renderer/Private/DeferredShadingRenderer.cpp
+++ b/Engine/Source/Runtime/Renderer/Private/DeferredShadingRenderer.cpp
@@ -861,6 +861,7 @@ bool FDeferredShadingSceneRenderer::GatherRayTracingWorldInstances(FRHICommandLi
 				ReferenceView.DynamicRayTracingMeshCommandStoragePerTask.Add(TaskDynamicCommandStorage);
 
 				FRayTracingMeshCommandOneFrameArray* TaskVisibleCommands = new(FMemStack::Get()) FRayTracingMeshCommandOneFrameArray;
+				TaskVisibleCommands->Reserve(NumPendingMeshBatches);
 				ReferenceView.VisibleRayTracingMeshCommandsPerTask.Add(TaskVisibleCommands);
 
 				const FViewInfo* ReferenceViewPtr = &ReferenceView;
@@ -871,6 +872,7 @@ bool FDeferredShadingSceneRenderer::GatherRayTracingWorldInstances(FRHICommandLi
 					FTaskTagScope TaskTagScope(ETaskTag::EParallelRenderingThread);
 					TRACE_CPUPROFILER_EVENT_SCOPE(RayTracingMeshBatchTask);
 					FRayTracingMeshBatchTaskPage* Page = TaskDataHead;
+					const int32 ExpectedMaxVisibieCommands = TaskVisibleCommands->Max();
 					while (Page)
 					{
 						for (uint32 ItemIndex = 0; ItemIndex < Page->NumWorkItems; ++ItemIndex)
@@ -892,6 +894,7 @@ bool FDeferredShadingSceneRenderer::GatherRayTracingWorldInstances(FRHICommandLi
 						Page->~FRayTracingMeshBatchTaskPage();
 						Page = NextPage;
 					}
+					check(ExpectedMaxVisibieCommands <= TaskVisibleCommands->Max());
 				}, TStatId(), nullptr, ENamedThreads::AnyThread));
 			}
 
