Fixed: auto-instancing on mobile

diff --git a/Engine/Source/Runtime/Renderer/Private/DeferredShadingRenderer.cpp b/Engine/Source/Runtime/Renderer/Private/DeferredShadingRenderer.cpp
--- a/Engine/Source/Runtime/Renderer/Private/DeferredShadingRenderer.cpp
+++ b/Engine/Source/Runtime/Renderer/Private/DeferredShadingRenderer.cpp
@@ -1849,7 +1849,7 @@ void FDeferredShadingSceneRenderer::Render(FRDGBuilder& GraphBuilder)
 	#endif
 
 	
-	const bool bIsOcclusionTesting = DoOcclusionQueries(FeatureLevel) && (!ViewFamily.EngineShowFlags.Wireframe || bIsViewFrozen || bHasViewParent);
+	const bool bIsOcclusionTesting = DoOcclusionQueries() && (!ViewFamily.EngineShowFlags.Wireframe || bIsViewFrozen || bHasViewParent);
 	const bool bNeedsPrePass = ShouldRenderPrePass();
 
 	GEngine->GetPreRenderDelegateEx().Broadcast(GraphBuilder);
diff --git a/Engine/Source/Runtime/Renderer/Private/DepthRendering.cpp b/Engine/Source/Runtime/Renderer/Private/DepthRendering.cpp
--- a/Engine/Source/Runtime/Renderer/Private/DepthRendering.cpp
+++ b/Engine/Source/Runtime/Renderer/Private/DepthRendering.cpp
@@ -615,8 +615,12 @@ bool FMobileSceneRenderer::ShouldRenderPrePass() const
 
 void FMobileSceneRenderer::RenderPrePass(FRHICommandListImmediate& RHICmdList, const FViewInfo& View, const FInstanceCullingDrawParams* InstanceCullingDrawParams/* = nullptr*/)
 {
+	if (!ShouldRenderPrePass())
+	{
+		return;
+	}
+	
 	checkSlow(RHICmdList.IsInsideRenderPass());
-	checkSlow(ShouldRenderPrePass());
 
 	SCOPED_NAMED_EVENT(FMobileSceneRenderer_RenderPrePass, FColor::Emerald);
 	SCOPED_DRAW_EVENT(RHICmdList, MobileRenderPrePass);
diff --git a/Engine/Source/Runtime/Renderer/Private/MobileBasePassRendering.cpp b/Engine/Source/Runtime/Renderer/Private/MobileBasePassRendering.cpp
--- a/Engine/Source/Runtime/Renderer/Private/MobileBasePassRendering.cpp
+++ b/Engine/Source/Runtime/Renderer/Private/MobileBasePassRendering.cpp
@@ -304,28 +304,28 @@ void SetupMobileSkyReflectionUniformParameters(FSkyLightSceneProxy* SkyLight, FM
 	Parameters.TextureSampler = CaptureTexture->SamplerStateRHI;
 }
 
-void FMobileSceneRenderer::RenderMobileBasePass(FRHICommandListImmediate& RHICmdList, const FViewInfo& View)
+void FMobileSceneRenderer::RenderMobileBasePass(FRHICommandListImmediate& RHICmdList, const FViewInfo& View, const FInstanceCullingDrawParams* InstanceCullingDrawParams)
 {
 	CSV_SCOPED_TIMING_STAT_EXCLUSIVE(RenderBasePass);
 	SCOPED_DRAW_EVENT(RHICmdList, MobileBasePass);
 	SCOPE_CYCLE_COUNTER(STAT_BasePassDrawTime);
 	SCOPED_GPU_STAT(RHICmdList, Basepass);
 
-		RHICmdList.SetViewport(View.ViewRect.Min.X, View.ViewRect.Min.Y, 0, View.ViewRect.Max.X, View.ViewRect.Max.Y, 1);
-		View.ParallelMeshDrawCommandPasses[EMeshPass::BasePass].DispatchDraw(nullptr, RHICmdList);
+	RHICmdList.SetViewport(View.ViewRect.Min.X, View.ViewRect.Min.Y, 0, View.ViewRect.Max.X, View.ViewRect.Max.Y, 1);
+	View.ParallelMeshDrawCommandPasses[EMeshPass::BasePass].DispatchDraw(nullptr, RHICmdList, InstanceCullingDrawParams);
 		
-		if (View.Family->EngineShowFlags.Atmosphere)
-		{
-			View.ParallelMeshDrawCommandPasses[EMeshPass::SkyPass].DispatchDraw(nullptr, RHICmdList);
-		}
+	if (View.Family->EngineShowFlags.Atmosphere)
+	{
+		View.ParallelMeshDrawCommandPasses[EMeshPass::SkyPass].DispatchDraw(nullptr, RHICmdList);
+	}
 
-		// editor primitives
+	// editor primitives
 	FMeshPassProcessorRenderState DrawRenderState;
-			DrawRenderState.SetBlendState(TStaticBlendStateWriteMask<CW_RGBA>::GetRHI());
-			DrawRenderState.SetDepthStencilAccess(Scene->DefaultBasePassDepthStencilAccess);
-			DrawRenderState.SetDepthStencilState(TStaticDepthStencilState<true, CF_DepthNearOrEqual>::GetRHI());
-			RenderMobileEditorPrimitives(RHICmdList, View, DrawRenderState);
-		}
+	DrawRenderState.SetBlendState(TStaticBlendStateWriteMask<CW_RGBA>::GetRHI());
+	DrawRenderState.SetDepthStencilAccess(Scene->DefaultBasePassDepthStencilAccess);
+	DrawRenderState.SetDepthStencilState(TStaticDepthStencilState<true, CF_DepthNearOrEqual>::GetRHI());
+	RenderMobileEditorPrimitives(RHICmdList, View, DrawRenderState);
+}
 
 void FMobileSceneRenderer::RenderMobileEditorPrimitives(FRHICommandList& RHICmdList, const FViewInfo& View, const FMeshPassProcessorRenderState& DrawRenderState)
 {
diff --git a/Engine/Source/Runtime/Renderer/Private/MobileDecalRendering.cpp b/Engine/Source/Runtime/Renderer/Private/MobileDecalRendering.cpp
--- a/Engine/Source/Runtime/Renderer/Private/MobileDecalRendering.cpp
+++ b/Engine/Source/Runtime/Renderer/Private/MobileDecalRendering.cpp
@@ -187,11 +187,12 @@ FRHIBlendState* MobileDeferred_GetDecalBlendState(EDecalBlendMode DecalBlendMode
 
 void FMobileSceneRenderer::RenderDecals(FRHICommandListImmediate& RHICmdList, const FViewInfo& View)
 {
-	if (!IsMobileHDR())
+	if (!IsMobileHDR() || !ViewFamily.EngineShowFlags.Decals || View.bIsPlanarReflection)
 	{
 		return;
 	}
 
+	CSV_SCOPED_TIMING_STAT_EXCLUSIVE(RenderDecals);
 	SCOPE_CYCLE_COUNTER(STAT_DecalsDrawTime);
 
 	// Deferred decals
diff --git a/Engine/Source/Runtime/Renderer/Private/MobileDistortionPass.cpp b/Engine/Source/Runtime/Renderer/Private/MobileDistortionPass.cpp
--- a/Engine/Source/Runtime/Renderer/Private/MobileDistortionPass.cpp
+++ b/Engine/Source/Runtime/Renderer/Private/MobileDistortionPass.cpp
@@ -34,6 +34,7 @@ bool IsMobileDistortionActive(const FViewInfo& View)
 BEGIN_SHADER_PARAMETER_STRUCT(FMobileDistortionPassParameters, )
 	SHADER_PARAMETER_STRUCT_INCLUDE(FViewShaderParameters, View)
 	SHADER_PARAMETER_RDG_UNIFORM_BUFFER(FMobileDistortionPassUniformParameters, Pass)
+	SHADER_PARAMETER_STRUCT_INCLUDE(FInstanceCullingDrawParams, InstanceCullingDrawParams)
 	RENDER_TARGET_BINDING_SLOTS()
 END_SHADER_PARAMETER_STRUCT()
 
@@ -54,7 +55,7 @@ TRDGUniformBufferRef<FMobileDistortionPassUniformParameters> CreateMobileDistort
 	return GraphBuilder.CreateUniformBuffer(Parameters);
 }
 
-FMobileDistortionAccumulateOutputs AddMobileDistortionAccumulatePass(FRDGBuilder& GraphBuilder, const FViewInfo& View, const FMobileDistortionAccumulateInputs& Inputs)
+FMobileDistortionAccumulateOutputs AddMobileDistortionAccumulatePass(FRDGBuilder& GraphBuilder, FScene* Scene, const FViewInfo& View, const FMobileDistortionAccumulateInputs& Inputs)
 {
 	FRDGTextureDesc DistortionAccumulateDesc = FRDGTextureDesc::Create2D(Inputs.SceneColor.Texture->Desc.Extent, PF_B8G8R8A8, FClearValueBinding::Transparent, TexCreate_ShaderResource | TexCreate_RenderTargetable);
 
@@ -67,17 +68,16 @@ FMobileDistortionAccumulateOutputs AddMobileDistortionAccumulatePass(FRDGBuilder
 
 	const FScreenPassTextureViewport SceneColorViewport(Inputs.SceneColor);
 
-	// GPUCULL_TODO: View.ParallelMeshDrawCommandPasses[EMeshPass::Distortion].BuildRenderingCommands(GraphBuilder, Scene->GPUScene);
+	const_cast<FViewInfo&>(View).ParallelMeshDrawCommandPasses[EMeshPass::Distortion].BuildRenderingCommands(GraphBuilder, Scene->GPUScene, PassParameters->InstanceCullingDrawParams);
 
 	GraphBuilder.AddPass(
 		RDG_EVENT_NAME("DistortionAccumulate %dx%d", SceneColorViewport.Rect.Width(), SceneColorViewport.Rect.Height()),
 		PassParameters,
 		ERDGPassFlags::Raster,
-		[&View, SceneColorViewport](FRHICommandList& RHICmdList)
+		[&View, SceneColorViewport, PassParameters](FRHICommandList& RHICmdList)
 	{
 		RHICmdList.SetViewport(SceneColorViewport.Rect.Min.X, SceneColorViewport.Rect.Min.Y, 0.0f, SceneColorViewport.Rect.Max.X, SceneColorViewport.Rect.Max.Y, 1.0f);
-
-		View.ParallelMeshDrawCommandPasses[EMeshPass::Distortion].DispatchDraw(nullptr, RHICmdList);
+		View.ParallelMeshDrawCommandPasses[EMeshPass::Distortion].DispatchDraw(nullptr, RHICmdList, &PassParameters->InstanceCullingDrawParams);
 	});
 
 	FMobileDistortionAccumulateOutputs Outputs;
diff --git a/Engine/Source/Runtime/Renderer/Private/MobileDistortionPass.h b/Engine/Source/Runtime/Renderer/Private/MobileDistortionPass.h
--- a/Engine/Source/Runtime/Renderer/Private/MobileDistortionPass.h
+++ b/Engine/Source/Runtime/Renderer/Private/MobileDistortionPass.h
@@ -14,7 +14,7 @@ struct FMobileDistortionAccumulateOutputs
 	FScreenPassTexture DistortionAccumulate;
 };
 
-FMobileDistortionAccumulateOutputs AddMobileDistortionAccumulatePass(FRDGBuilder& GraphBuilder, const FViewInfo& View, const FMobileDistortionAccumulateInputs& Inputs);
+FMobileDistortionAccumulateOutputs AddMobileDistortionAccumulatePass(FRDGBuilder& GraphBuilder, FScene* Scene, const FViewInfo& View, const FMobileDistortionAccumulateInputs& Inputs);
 
 struct FMobileDistortionMergeInputs
 {
diff --git a/Engine/Source/Runtime/Renderer/Private/MobileShadingRenderer.cpp b/Engine/Source/Runtime/Renderer/Private/MobileShadingRenderer.cpp
--- a/Engine/Source/Runtime/Renderer/Private/MobileShadingRenderer.cpp
+++ b/Engine/Source/Runtime/Renderer/Private/MobileShadingRenderer.cpp
@@ -60,13 +60,6 @@
 
 uint32 GetShadowQuality();
 
-static TAutoConsoleVariable<int32> CVarMobileAlwaysResolveDepth(
-	TEXT("r.Mobile.AlwaysResolveDepth"),
-	0,
-	TEXT("0: Depth buffer is resolved after opaque pass only when decals or modulated shadows are in use. (Default)\n")
-	TEXT("1: Depth buffer is always resolved after opaque pass.\n"),
-	ECVF_Scalability | ECVF_RenderThreadSafe);
-
 static TAutoConsoleVariable<int32> CVarMobileForceDepthResolve(
 	TEXT("r.Mobile.ForceDepthResolve"),
 	0,
@@ -169,10 +162,14 @@ END_SHADER_PARAMETER_STRUCT()
 
 BEGIN_SHADER_PARAMETER_STRUCT(FMobileRenderPassParameters, )
 	SHADER_PARAMETER_STRUCT_INCLUDE(FViewShaderParameters, View)
-	SHADER_PARAMETER_STRUCT_ARRAY(FInstanceCullingDrawParams, InstanceCullingDrawParams, [EMeshPass::Num])
 	SHADER_PARAMETER_RDG_UNIFORM_BUFFER(FMobileSceneTextureUniformParameters, SceneTextures)
 	SHADER_PARAMETER_RDG_UNIFORM_BUFFER(FMobileBasePassUniformParameters, MobileBasePass)
 	SHADER_PARAMETER_RDG_UNIFORM_BUFFER(FDebugViewModePassUniformParameters, DebugViewMode)
+	SHADER_PARAMETER_STRUCT_INCLUDE(FInstanceCullingDrawParams, InstanceCullingDrawParamsDepth)
+	SHADER_PARAMETER_STRUCT_INCLUDE(FInstanceCullingDrawParams, InstanceCullingDrawParamsOpaque)
+	SHADER_PARAMETER_STRUCT_INCLUDE(FInstanceCullingDrawParams, InstanceCullingDrawParamsSky)
+	SHADER_PARAMETER_STRUCT_INCLUDE(FInstanceCullingDrawParams, InstanceCullingDrawParamsTranslucency)
+	SHADER_PARAMETER_STRUCT_INCLUDE(FInstanceCullingDrawParams, InstanceCullingDrawParamsDebugView)
 	RENDER_TARGET_BINDING_SLOTS()
 END_SHADER_PARAMETER_STRUCT()
 
@@ -217,6 +214,10 @@ FMobileSceneRenderer::FMobileSceneRenderer(const FSceneViewFamily* InViewFamily,
 	bRequiresDistanceFieldShadowingPass = false;
 	bIsFullDepthPrepassEnabled = Scene->EarlyZPassMode == DDM_AllOpaque;
 	bIsMaskedOnlyDepthPrepassEnabled = Scene->EarlyZPassMode == DDM_MaskedOnly;
+	
+	StandardTranslucencyPass = ViewFamily.AllowTranslucencyAfterDOF() ? ETranslucencyPass::TPT_StandardTranslucency : ETranslucencyPass::TPT_AllTranslucency;
+	StandardTranslucencyMeshPass = TranslucencyPassToMeshPass(StandardTranslucencyPass);
+
 	// Don't do occlusion queries when doing scene captures
 	for (FViewInfo& View : Views)
 	{
@@ -280,7 +281,7 @@ void FMobileSceneRenderer::PrepareViewVisibilityLists()
 	}
 }
 
-void FMobileSceneRenderer::SetupMobileBasePassAfterShadowInit(FExclusiveDepthStencil::Type BasePassDepthStencilAccess, FViewVisibleCommandsPerView& ViewCommandsPerView)
+void FMobileSceneRenderer::SetupMobileBasePassAfterShadowInit(FExclusiveDepthStencil::Type BasePassDepthStencilAccess, FViewVisibleCommandsPerView& ViewCommandsPerView, FInstanceCullingManager& InstanceCullingManager)
 {
 	// Sort front to back on all platforms, even HSR benefits from it
 	//const bool bWantsFrontToBackSorting = (GHardwareHiddenSurfaceRemoval == false);
@@ -297,12 +298,22 @@ void FMobileSceneRenderer::SetupMobileBasePassAfterShadowInit(FExclusiveDepthSte
 		PassProcessorCreateFunction BasePassCSMCreateFunction = FPassProcessorManager::GetCreateFunction(EShadingPath::Mobile, EMeshPass::MobileBasePassCSM);
 		FMeshPassProcessor* BasePassCSMMeshPassProcessor = BasePassCSMCreateFunction(Scene, &View, nullptr);
 
+		TArray<int32, TInlineAllocator<2> > ViewIds;
+		ViewIds.Add(View.GPUSceneViewId);
+		// Only apply instancing for ISR to main view passes
+		EInstanceCullingMode InstanceCullingMode = View.IsInstancedStereoPass() ? EInstanceCullingMode::Stereo : EInstanceCullingMode::Normal;
+		if (InstanceCullingMode == EInstanceCullingMode::Stereo)
+		{
+			check(View.GetInstancedView() != nullptr);
+			ViewIds.Add(View.GetInstancedView()->GPUSceneViewId);
+		}
+
 		// Run sorting on BasePass, as it's ignored inside FSceneRenderer::SetupMeshPass, so it can be done after shadow init on mobile.
 		FParallelMeshDrawCommandPass& Pass = View.ParallelMeshDrawCommandPasses[EMeshPass::BasePass];
 		Pass.DispatchPassSetup(
 			Scene,
 			View,
-			FInstanceCullingContext(), // GPUCULL_TODO: Fix mobile!
+			FInstanceCullingContext(&InstanceCullingManager, ViewIds, InstanceCullingMode),
 			EMeshPass::BasePass,
 			BasePassDepthStencilAccess,
 			MeshPassProcessor,
@@ -405,7 +416,6 @@ void FMobileSceneRenderer::InitViews(FRDGBuilder& GraphBuilder, FSceneTexturesCo
 
 	// Whether we need to store depth for post-processing
 	// On PowerVR we see flickering of shadows and depths not updating correctly if targets are discarded.
-	// See CVarMobileForceDepthResolve use in ConditionalResolveSceneDepth.
 	const bool bForceDepthResolve = (CVarMobileForceDepthResolve.GetValueOnRenderThread() == 1);
 	const bool bSeparateTranslucencyActive = IsMobileSeparateTranslucencyActive(Views.GetData(), Views.Num()); 
 	const bool bPostProcessUsesSceneDepth = PostProcessUsesSceneDepth(Views[0]) || IsMobileDistortionActive(Views[0]);
@@ -480,7 +490,7 @@ void FMobileSceneRenderer::InitViews(FRDGBuilder& GraphBuilder, FSceneTexturesCo
 	/** Before SetupMobileBasePassAfterShadowInit, we need to update the uniform buffer and shadow info for all movable point lights.*/
 	UpdateMovablePointLightUniformBufferAndShadowInfo();
 
-	SetupMobileBasePassAfterShadowInit(BasePassDepthStencilAccess, ViewCommandsPerView);
+	SetupMobileBasePassAfterShadowInit(BasePassDepthStencilAccess, ViewCommandsPerView, InstanceCullingManager);
 
 	// if we kicked off ILC update via task, wait and finalize.
 	if (ILCTaskData.TaskRef.IsValid())
@@ -529,7 +539,8 @@ void FMobileSceneRenderer::InitViews(FRDGBuilder& GraphBuilder, FSceneTexturesCo
 	{
 		PrepareDistanceFieldScene(GraphBuilder, false);
 	}
-
+	
+	if (InstanceCullingManager.IsEnabled())
 	{
 		// GPUCULL_TODO: Possibly fold into unpack step
 		InstanceCullingManager.CullInstances(GraphBuilder, Scene->GPUScene);
@@ -590,11 +601,19 @@ void FMobileSceneRenderer::RenderFullDepthPrepass(FRDGBuilder& GraphBuilder, FSc
 			ERDGPassFlags::Raster,
 			[this, PassParameters, &View](FRHICommandListImmediate& RHICmdList)
 			{
-				RenderPrePass(RHICmdList, View, &PassParameters->InstanceCullingDrawParams[EMeshPass::DepthPass]);
+				RenderPrePass(RHICmdList, View, &PassParameters->InstanceCullingDrawParamsDepth);
 			});
 	}
 }
 
+void FMobileSceneRenderer::RenderMaskedPrePass(FRHICommandListImmediate& RHICmdList, const FViewInfo& View, const FInstanceCullingDrawParams* InstanceCullingDrawParams)
+{
+	if (bIsMaskedOnlyDepthPrepassEnabled)
+	{
+		RenderPrePass(RHICmdList, View, InstanceCullingDrawParams);
+	}
+}
+
 /** 
 * Renders the view family. 
 */
@@ -651,7 +670,6 @@ void FMobileSceneRenderer::Render(FRDGBuilder& GraphBuilder)
 	// NOTE: Must be done after  system texture initialization
 	VirtualShadowMapArray.Initialize(GraphBuilder, UseVirtualShadowMaps(ShaderPlatform, FeatureLevel));
 
-
 	// Find the visible primitives and prepare targets and buffers for rendering
 	InitViews(GraphBuilder, SceneTexturesConfig, InstanceCullingManager);
 
@@ -843,7 +861,7 @@ void FMobileSceneRenderer::Render(FRDGBuilder& GraphBuilder)
 				{
 					RDG_EVENT_SCOPE_CONDITIONAL(GraphBuilder, Views.Num() > 1, "View%d", ViewIndex);
 					PostProcessingInputs.SceneTextures = SceneTextures.MobileUniformBuffer;
-					AddMobilePostProcessingPasses(GraphBuilder, Views[ViewIndex], PostProcessingInputs, InstanceCullingManager);
+					AddMobilePostProcessingPasses(GraphBuilder, Scene, Views[ViewIndex], PostProcessingInputs, InstanceCullingManager);
 				}
 			}
 		}
@@ -858,6 +876,18 @@ void FMobileSceneRenderer::Render(FRDGBuilder& GraphBuilder)
 	AddPass(GraphBuilder, PollOcclusionQueriesAndDispatchToRHIThreadPass);
 }
 
+void FMobileSceneRenderer::BuildInstanceCullingDrawParams(FRDGBuilder& GraphBuilder, FViewInfo& View, FMobileRenderPassParameters* PassParameters)
+{
+	View.ParallelMeshDrawCommandPasses[EMeshPass::DepthPass].BuildRenderingCommands(GraphBuilder, Scene->GPUScene, PassParameters->InstanceCullingDrawParamsDepth);
+	View.ParallelMeshDrawCommandPasses[EMeshPass::BasePass].BuildRenderingCommands(GraphBuilder, Scene->GPUScene, PassParameters->InstanceCullingDrawParamsOpaque);
+	View.ParallelMeshDrawCommandPasses[EMeshPass::SkyPass].BuildRenderingCommands(GraphBuilder, Scene->GPUScene, PassParameters->InstanceCullingDrawParamsSky);
+	View.ParallelMeshDrawCommandPasses[StandardTranslucencyMeshPass].BuildRenderingCommands(GraphBuilder, Scene->GPUScene, PassParameters->InstanceCullingDrawParamsTranslucency);
+	if (ViewFamily.UseDebugViewPS())
+	{
+		View.ParallelMeshDrawCommandPasses[EMeshPass::DebugViewMode].BuildRenderingCommands(GraphBuilder, Scene->GPUScene, PassParameters->InstanceCullingDrawParamsDebugView);
+	}
+}
+
 void FMobileSceneRenderer::RenderForward(FRDGBuilder& GraphBuilder, FRDGTextureRef ViewFamilyTexture, FSceneTextures& SceneTextures)
 {
 	const FViewInfo& MainView = Views[0];
@@ -904,18 +934,6 @@ void FMobileSceneRenderer::RenderForward(FRDGBuilder& GraphBuilder, FRDGTextureR
 	//if the scenecolor isn't multiview but the app is, need to render as a single-view multiview due to shaders
 	BasePassRenderTargets.MultiViewCount = MainView.bIsMobileMultiViewEnabled ? 2 : (bIsMultiViewApplication ? 1 : 0);
 
-	const bool bShouldRenderMobileDebugView = ViewFamily.UseDebugViewPS();
-
-	const bool bAdrenoOcclusionMode = CVarMobileAdrenoOcclusionMode.GetValueOnRenderThread() != 0;
-
-	const bool bShouldRenderOcclusionPass = DoOcclusionQueries(FeatureLevel);
-
-	ETranslucencyPass::Type TranslucencyPass = ViewFamily.AllowTranslucencyAfterDOF() ? ETranslucencyPass::TPT_StandardTranslucency : ETranslucencyPass::TPT_AllTranslucency;
-
-	const EMeshPass::Type TranslucencyMeshPass = TranslucencyPassToMeshPass(TranslucencyPass);
-
-	const bool bShouldRenderTranslucency = ShouldRenderTranslucency(TranslucencyPass) && ViewFamily.EngineShowFlags.Translucency;
-
 	const FRDGSystemTextures& SystemTextures = FRDGSystemTextures::Get(GraphBuilder);
 
 	for (int32 ViewIndex = 0; ViewIndex < Views.Num(); ViewIndex++)
@@ -948,125 +966,63 @@ void FMobileSceneRenderer::RenderForward(FRDGBuilder& GraphBuilder, FRDGTextureR
 		PassParameters->MobileBasePass = CreateMobileBasePassUniformBuffer(GraphBuilder, View, EMobileBasePass::Opaque, LastFrameScreenSpaceAO, MobilePixelProjectedReflection);
 		PassParameters->SceneTextures = SceneTextures.MobileUniformBuffer;
 		PassParameters->RenderTargets = BasePassRenderTargets;
-
-		if (bIsMaskedOnlyDepthPrepassEnabled)
-		{
-//			View.ParallelMeshDrawCommandPasses[EMeshPass::DepthPass].BuildRenderingCommands(GraphBuilder, Scene->GPUScene, PassParameters->InstanceCullingDrawParams[EMeshPass::DepthPass]);
-		}
-
-//		View.ParallelMeshDrawCommandPasses[EMeshPass::BasePass].BuildRenderingCommands(GraphBuilder, Scene->GPUScene, PassParameters->InstanceCullingDrawParams[EMeshPass::BasePass]);
-
-#if WITH_DEBUG_VIEW_MODES
-		if (bShouldRenderMobileDebugView)
+		if (ViewFamily.UseDebugViewPS())
 		{
-			// GPUCULL_TODO: View.ParallelMeshDrawCommandPasses[EMeshPass::DebugViewMode].BuildRenderingCommands(GraphBuilder, Scene->GPUScene, PassParameters->InstanceCullingDrawParams[EMeshPass::DebugViewMode]);
 			PassParameters->DebugViewMode = CreateDebugViewModePassUniformBuffer(GraphBuilder, View, SceneTextures.QuadOverdraw);
 		}
-#endif
-
-		if (bShouldRenderTranslucency)
-		{
-			// GPUCULL_TODO: View.ParallelMeshDrawCommandPasses[TranslucencyMeshPass].BuildRenderingCommands(GraphBuilder, Scene->GPUScene, PassParameters->InstanceCullingDrawParams[TranslucencyMeshPass]);
-		}
+		
+		BuildInstanceCullingDrawParams(GraphBuilder, View, PassParameters);
 
 		// Split if we need to render translucency in a separate render pass
 		if (bRequiresMultiPass)
 		{
-			RenderForwardMultiPass(GraphBuilder, PassParameters, BasePassRenderTargets, ViewIndex, View, SceneTextures, bShouldRenderMobileDebugView, bAdrenoOcclusionMode, bShouldRenderOcclusionPass, bShouldRenderTranslucency, TranslucencyMeshPass);
+			RenderForwardMultiPass(GraphBuilder, PassParameters, BasePassRenderTargets, ViewIndex, View, SceneTextures);
 		}
 		else
 		{
-			RenderForwardSinglePass(GraphBuilder, PassParameters, ViewIndex, View, SceneTextures, bShouldRenderMobileDebugView, bAdrenoOcclusionMode, bShouldRenderOcclusionPass, bShouldRenderTranslucency, TranslucencyMeshPass);
+			RenderForwardSinglePass(GraphBuilder, PassParameters, ViewIndex, View, SceneTextures);
 		}
 	}
 
 	QueueSceneTextureExtractions(GraphBuilder, SceneTextures);
 }
 
-void FMobileSceneRenderer::RenderForwardSinglePass(FRDGBuilder& GraphBuilder, class FMobileRenderPassParameters* PassParameters, int32 ViewIndex, FViewInfo& View, FSceneTextures& SceneTextures, bool bShouldRenderMobileDebugView, bool bAdrenoOcclusionMode, bool bShouldRenderOcclusionPass, bool bShouldRenderTranslucency, EMeshPass::Type TranslucencyMeshPass)
+void FMobileSceneRenderer::RenderForwardSinglePass(FRDGBuilder& GraphBuilder, FMobileRenderPassParameters* PassParameters, int32 ViewIndex, FViewInfo& View, FSceneTextures& SceneTextures)
 {
 	GraphBuilder.AddPass(
 		RDG_EVENT_NAME("SceneColorRendering"),
 		PassParameters,
 		ERDGPassFlags::Raster,
-		[this, PassParameters, ViewIndex, &View, &SceneTextures, bShouldRenderMobileDebugView, bAdrenoOcclusionMode, bShouldRenderOcclusionPass, bShouldRenderTranslucency, TranslucencyMeshPass](FRHICommandListImmediate& RHICmdList)
+		[this, PassParameters, ViewIndex, &View, &SceneTextures](FRHICommandListImmediate& RHICmdList)
 	{
 		if (GIsEditor && !View.bIsSceneCapture)
 		{
 			DrawClearQuad(RHICmdList, View.BackgroundColor);
 		}
 
+		// Depth pre-pass
 		RHICmdList.SetCurrentStat(GET_STATID(STAT_CLM_MobilePrePass));
-		if (bIsMaskedOnlyDepthPrepassEnabled)
-		{
-			// Depth pre-pass
-			RenderPrePass(RHICmdList, View, &PassParameters->InstanceCullingDrawParams[EMeshPass::DepthPass]);
-		}
-
+		RenderMaskedPrePass(RHICmdList, View, &PassParameters->InstanceCullingDrawParamsDepth);
 		// Opaque and masked
 		RHICmdList.SetCurrentStat(GET_STATID(STAT_CLMM_Opaque));
-		RenderMobileBasePass(RHICmdList, View);
+		RenderMobileBasePass(RHICmdList, View, &PassParameters->InstanceCullingDrawParamsOpaque);
+		RenderMobileDebugView(RHICmdList, View);
 		RHICmdList.ImmediateFlush(EImmediateFlushType::DispatchToRHIThread);
-
-#if WITH_DEBUG_VIEW_MODES
-		if (bShouldRenderMobileDebugView)
-		{
-			// Here we use the base pass depth result to get z culling for opaque and masque.
-			// The color needs to be cleared at this point since shader complexity renders in additive.
-			DrawClearQuad(RHICmdList, FLinearColor::Black);
-			RenderMobileDebugView(RHICmdList, View);
-			RHICmdList.ImmediateFlush(EImmediateFlushType::DispatchToRHIThread);
-		}
-#endif // WITH_DEBUG_VIEW_MODES
-
-		if (!bAdrenoOcclusionMode && bShouldRenderOcclusionPass)
+		const bool bAdrenoOcclusionMode = (CVarMobileAdrenoOcclusionMode.GetValueOnRenderThread() != 0 && IsOpenGLPlatform(ShaderPlatform));
+		if (!bAdrenoOcclusionMode)
 		{
 			// Issue occlusion queries
 			RHICmdList.SetCurrentStat(GET_STATID(STAT_CLMM_Occlusion));
 			RenderOcclusion(RHICmdList, View);
-			RHICmdList.ImmediateFlush(EImmediateFlushType::DispatchToRHIThread);
 		}
-
-		if (ViewFamily.ViewExtensions.Num() > 1)
-		{
-			CSV_SCOPED_TIMING_STAT_EXCLUSIVE(ViewExtensionPostRenderBasePass);
-			QUICK_SCOPE_CYCLE_COUNTER(STAT_FMobileSceneRenderer_ViewExtensionPostRenderBasePass);
-			for (int32 ViewExt = 0; ViewExt < ViewFamily.ViewExtensions.Num(); ++ViewExt)
-			{
-				ViewFamily.ViewExtensions[ViewExt]->PostRenderBasePass_RenderThread(RHICmdList, View);
-			}
-		}
-
-		RHICmdList.SetCurrentStat(GET_STATID(STAT_CLMM_Translucency));
-
+		PostRenderBasePass(RHICmdList, View);
 		// scene depth is read only and can be fetched
 		RHICmdList.NextSubpass();
-
-		if (!View.bIsPlanarReflection)
-		{
-			if (ViewFamily.EngineShowFlags.Decals)
-			{
-				CSV_SCOPED_TIMING_STAT_EXCLUSIVE(RenderDecals);
-				RenderDecals(RHICmdList, View);
-			}
-
-			if (ViewFamily.EngineShowFlags.DynamicShadows)
-			{
-				CSV_SCOPED_TIMING_STAT_EXCLUSIVE(RenderShadowProjections);
-				RenderModulatedShadowProjections(RHICmdList, ViewIndex, View);
-			}
-		}
-
+		RHICmdList.SetCurrentStat(GET_STATID(STAT_CLMM_Translucency));
+		RenderDecals(RHICmdList, View);
+		RenderModulatedShadowProjections(RHICmdList, ViewIndex, View);
 		// Draw translucency.
-		if (bShouldRenderTranslucency)
-		{
-			CSV_SCOPED_TIMING_STAT_EXCLUSIVE(RenderTranslucency);
-			SCOPE_CYCLE_COUNTER(STAT_TranslucencyDrawTime);
-			RenderTranslucency(RHICmdList, View, TranslucencyMeshPass);
-			FRHICommandListExecutor::GetImmediateCommandList().PollOcclusionQueries();
-			RHICmdList.ImmediateFlush(EImmediateFlushType::DispatchToRHIThread);
-		}
-
+		RenderTranslucency(RHICmdList, View, &PassParameters->InstanceCullingDrawParamsTranslucency);
 		if (bAdrenoOcclusionMode)
 		{
 			RHICmdList.SetCurrentStat(GET_STATID(STAT_CLMM_Occlusion));
@@ -1075,75 +1031,38 @@ void FMobileSceneRenderer::RenderForwardSinglePass(FRDGBuilder& GraphBuilder, cl
 			bSubmitOffscreenRendering = false; // submit once
 			// Issue occlusion queries
 			RenderOcclusion(RHICmdList, View);
-			RHICmdList.ImmediateFlush(EImmediateFlushType::DispatchToRHIThread);
 		}
-
 		// Pre-tonemap before MSAA resolve (iOS only)
-		if (!bGammaSpace)
-		{
-			PreTonemapMSAA(RHICmdList, SceneTextures);
-		}
+		PreTonemapMSAA(RHICmdList, SceneTextures);
 	});
 }
 
-void FMobileSceneRenderer::RenderForwardMultiPass(FRDGBuilder& GraphBuilder, class FMobileRenderPassParameters* PassParameters, FRenderTargetBindingSlots& BasePassRenderTargets, int32 ViewIndex, FViewInfo& View, FSceneTextures& SceneTextures, bool bShouldRenderMobileDebugView, bool bAdrenoOcclusionMode, bool bShouldRenderOcclusionPass, bool bShouldRenderTranslucency, EMeshPass::Type TranslucencyMeshPass)
+void FMobileSceneRenderer::RenderForwardMultiPass(FRDGBuilder& GraphBuilder, FMobileRenderPassParameters* PassParameters, FRenderTargetBindingSlots& BasePassRenderTargets, int32 ViewIndex, FViewInfo& View, FSceneTextures& SceneTextures)
 {
 	GraphBuilder.AddPass(
 		RDG_EVENT_NAME("SceneColorRendering"),
 		PassParameters,
 		ERDGPassFlags::Raster,
-		[this, PassParameters, &View, &SceneTextures, bShouldRenderMobileDebugView, bAdrenoOcclusionMode, bShouldRenderOcclusionPass](FRHICommandListImmediate& RHICmdList)
+		[this, PassParameters, &View, &SceneTextures](FRHICommandListImmediate& RHICmdList)
 	{
 		if (GIsEditor && !View.bIsSceneCapture)
 		{
 			DrawClearQuad(RHICmdList, View.BackgroundColor);
 		}
 
+		// Depth pre-pass
 		RHICmdList.SetCurrentStat(GET_STATID(STAT_CLM_MobilePrePass));
-		if (bIsMaskedOnlyDepthPrepassEnabled)
-		{
-			// Depth pre-pass
-			RenderPrePass(RHICmdList, View, &PassParameters->InstanceCullingDrawParams[EMeshPass::DepthPass]);
-		}
-
+		RenderMaskedPrePass(RHICmdList, View, &PassParameters->InstanceCullingDrawParamsDepth);
 		// Opaque and masked
 		RHICmdList.SetCurrentStat(GET_STATID(STAT_CLMM_Opaque));
-		RenderMobileBasePass(RHICmdList, View);
-		RHICmdList.ImmediateFlush(EImmediateFlushType::DispatchToRHIThread);
-
-#if WITH_DEBUG_VIEW_MODES
-		if (bShouldRenderMobileDebugView)
-		{
-			// Here we use the base pass depth result to get z culling for opaque and masque.
-			// The color needs to be cleared at this point since shader complexity renders in additive.
-			DrawClearQuad(RHICmdList, FLinearColor::Black);
-			RenderMobileDebugView(RHICmdList, View);
-			RHICmdList.ImmediateFlush(EImmediateFlushType::DispatchToRHIThread);
-		}
-#endif // WITH_DEBUG_VIEW_MODES
-
-		if (!bAdrenoOcclusionMode && bShouldRenderOcclusionPass)
-		{
-			// Issue occlusion queries
-			RHICmdList.SetCurrentStat(GET_STATID(STAT_CLMM_Occlusion));
-			RenderOcclusion(RHICmdList, View);
-			RHICmdList.ImmediateFlush(EImmediateFlushType::DispatchToRHIThread);
-		}
-
-		if (ViewFamily.ViewExtensions.Num() > 1)
-		{
-			CSV_SCOPED_TIMING_STAT_EXCLUSIVE(ViewExtensionPostRenderBasePass);
-			QUICK_SCOPE_CYCLE_COUNTER(STAT_FMobileSceneRenderer_ViewExtensionPostRenderBasePass);
-			for (int32 ViewExt = 0; ViewExt < ViewFamily.ViewExtensions.Num(); ++ViewExt)
-			{
-				ViewFamily.ViewExtensions[ViewExt]->PostRenderBasePass_RenderThread(RHICmdList, View);
-			}
-		}
+		RenderMobileBasePass(RHICmdList, View, &PassParameters->InstanceCullingDrawParamsOpaque);
+		RenderMobileDebugView(RHICmdList, View);
+		// Issue occlusion queries
+		RHICmdList.SetCurrentStat(GET_STATID(STAT_CLMM_Occlusion));
+		RenderOcclusion(RHICmdList, View);
+		PostRenderBasePass(RHICmdList, View);
 	});
 
-	// Make a copy of the scene depth if the current hardware doesn't support reading and writing to the same depth buffer
-	ConditionalResolveSceneDepth(GraphBuilder, View, SceneTextures.Depth);
-
 	SceneTextures.MobileSetupMode = EMobileSceneTextureSetupMode::SceneDepth | EMobileSceneTextureSetupMode::SceneDepthAux | EMobileSceneTextureSetupMode::CustomDepth;
 	SceneTextures.MobileUniformBuffer = CreateMobileSceneTextureUniformBuffer(GraphBuilder, SceneTextures.MobileSetupMode);
 
@@ -1170,54 +1089,17 @@ void FMobileSceneRenderer::RenderForwardMultiPass(FRDGBuilder& GraphBuilder, cla
 		{},
 		SecondPassParameters,
 		ERDGPassFlags::Raster,
-		[this, SecondPassParameters, ViewIndex, &View, &SceneTextures, bAdrenoOcclusionMode, bShouldRenderTranslucency, TranslucencyMeshPass](FRHICommandListImmediate& RHICmdList)
+		[this, SecondPassParameters, ViewIndex, &View, &SceneTextures](FRHICommandListImmediate& RHICmdList)
 	{
-		RHICmdList.SetCurrentStat(GET_STATID(STAT_CLMM_Translucency));
-
 		// scene depth is read only and can be fetched
 		RHICmdList.NextSubpass();
-
-		if (!View.bIsPlanarReflection)
-		{
-			if (ViewFamily.EngineShowFlags.Decals)
-			{
-				CSV_SCOPED_TIMING_STAT_EXCLUSIVE(RenderDecals);
-				RenderDecals(RHICmdList, View);
-			}
-
-			if (ViewFamily.EngineShowFlags.DynamicShadows)
-			{
-				CSV_SCOPED_TIMING_STAT_EXCLUSIVE(RenderShadowProjections);
-				RenderModulatedShadowProjections(RHICmdList, ViewIndex, View);
-			}
-		}
-
+		RHICmdList.SetCurrentStat(GET_STATID(STAT_CLMM_Translucency));
+		RenderDecals(RHICmdList, View);
+		RenderModulatedShadowProjections(RHICmdList, ViewIndex, View);
 		// Draw translucency.
-		if (bShouldRenderTranslucency)
-		{
-			CSV_SCOPED_TIMING_STAT_EXCLUSIVE(RenderTranslucency);
-			SCOPE_CYCLE_COUNTER(STAT_TranslucencyDrawTime);
-			RenderTranslucency(RHICmdList, View, TranslucencyMeshPass);
-			FRHICommandListExecutor::GetImmediateCommandList().PollOcclusionQueries();
-			RHICmdList.ImmediateFlush(EImmediateFlushType::DispatchToRHIThread);
-		}
-
-		if (bAdrenoOcclusionMode)
-		{
-			RHICmdList.SetCurrentStat(GET_STATID(STAT_CLMM_Occlusion));
-			// flush
-			RHICmdList.SubmitCommandsHint();
-			bSubmitOffscreenRendering = false; // submit once
-			// Issue occlusion queries
-			RenderOcclusion(RHICmdList, View);
-			RHICmdList.ImmediateFlush(EImmediateFlushType::DispatchToRHIThread);
-		}
-
+		RenderTranslucency(RHICmdList, View, &SecondPassParameters->InstanceCullingDrawParamsTranslucency);
 		// Pre-tonemap before MSAA resolve (iOS only)
-		if (!bGammaSpace)
-		{
-			PreTonemapMSAA(RHICmdList, SceneTextures);
-		}
+		PreTonemapMSAA(RHICmdList, SceneTextures);
 	});
 }
 
@@ -1337,14 +1219,6 @@ void FMobileSceneRenderer::RenderDeferred(FRDGBuilder& GraphBuilder, const FSort
 	BasePassRenderTargets.ShadingRateTexture = nullptr;
 	BasePassRenderTargets.MultiViewCount = 0;
 
-	const bool bShouldRenderOcclusionPass = DoOcclusionQueries(FeatureLevel);
-
-	ETranslucencyPass::Type TranslucencyPass = ViewFamily.AllowTranslucencyAfterDOF() ? ETranslucencyPass::TPT_StandardTranslucency : ETranslucencyPass::TPT_AllTranslucency;
-
-	const EMeshPass::Type TranslucencyMeshPass = TranslucencyPassToMeshPass(TranslucencyPass);
-
-	const bool bShouldRenderTranslucency = ShouldRenderTranslucency(TranslucencyPass) && ViewFamily.EngineShowFlags.Translucency;
-
 	const FRDGSystemTextures& SystemTextures = FRDGSystemTextures::Get(GraphBuilder);
 
 	int32 NumViews = Views.Num();
@@ -1379,130 +1253,89 @@ void FMobileSceneRenderer::RenderDeferred(FRDGBuilder& GraphBuilder, const FSort
 		PassParameters->MobileBasePass = CreateMobileBasePassUniformBuffer(GraphBuilder, View, EMobileBasePass::Opaque, LastFrameScreenSpaceAO, MobilePixelProjectedReflection);
 		PassParameters->SceneTextures = SceneTextures.MobileUniformBuffer;
 		PassParameters->RenderTargets = BasePassRenderTargets;
-
-		if (bIsMaskedOnlyDepthPrepassEnabled)
+		if (ViewFamily.UseDebugViewPS())
 		{
-			View.ParallelMeshDrawCommandPasses[EMeshPass::DepthPass].BuildRenderingCommands(GraphBuilder, Scene->GPUScene, PassParameters->InstanceCullingDrawParams[EMeshPass::DepthPass]);
+			PassParameters->DebugViewMode = CreateDebugViewModePassUniformBuffer(GraphBuilder, View, SceneTextures.QuadOverdraw);
 		}
+		
+		BuildInstanceCullingDrawParams(GraphBuilder, View, PassParameters);
 
 		if (bRequiresMultiPass)
 		{
-			RenderDeferredMultiPass(GraphBuilder, PassParameters, BasePassRenderTargets, ColorTargets.Num(), ViewIndex, NumViews, View, SceneTextures, SortedLightSet, MobilePixelProjectedReflection, bShouldRenderOcclusionPass, bShouldRenderTranslucency, TranslucencyMeshPass);
+			RenderDeferredMultiPass(GraphBuilder, PassParameters, BasePassRenderTargets, ColorTargets.Num(), ViewIndex, NumViews, View, SceneTextures, SortedLightSet, MobilePixelProjectedReflection);
 		}
 		else
 		{
-			RenderDeferredSinglePass(GraphBuilder, PassParameters, ViewIndex, NumViews, View, SceneTextures, SortedLightSet, MobilePixelProjectedReflection, bShouldRenderOcclusionPass, bUsingPixelLocalStorage, bShouldRenderTranslucency, TranslucencyMeshPass);
+			RenderDeferredSinglePass(GraphBuilder, PassParameters, ViewIndex, NumViews, View, SceneTextures, SortedLightSet, MobilePixelProjectedReflection, bUsingPixelLocalStorage);
 		}
 	}
 
 	QueueSceneTextureExtractions(GraphBuilder, SceneTextures);
 }
 
-void FMobileSceneRenderer::RenderDeferredSinglePass(FRDGBuilder& GraphBuilder, class FMobileRenderPassParameters* PassParameters, int32 ViewIndex, int32 NumViews, FViewInfo& View, FSceneTextures& SceneTextures, const FSortedLightSetSceneInfo& SortedLightSet, FTextureRHIRef MobilePixelProjectedReflection, bool bShouldRenderOcclusionPass, bool bUsingPixelLocalStorage, bool bShouldRenderTranslucency, EMeshPass::Type TranslucencyMeshPass)
+void FMobileSceneRenderer::RenderDeferredSinglePass(FRDGBuilder& GraphBuilder, class FMobileRenderPassParameters* PassParameters, int32 ViewIndex, int32 NumViews, FViewInfo& View, FSceneTextures& SceneTextures, const FSortedLightSetSceneInfo& SortedLightSet, FTextureRHIRef MobilePixelProjectedReflection, bool bUsingPixelLocalStorage)
 {
 	GraphBuilder.AddPass(
 		RDG_EVENT_NAME("SceneColorRendering"),
 		PassParameters,
 		ERDGPassFlags::Raster,
-		[this, PassParameters, ViewIndex, NumViews, &View, &SceneTextures, &SortedLightSet, MobilePixelProjectedReflection, bShouldRenderOcclusionPass, bUsingPixelLocalStorage, bShouldRenderTranslucency, TranslucencyMeshPass](FRHICommandListImmediate& RHICmdList)
+		[this, PassParameters, ViewIndex, NumViews, &View, &SceneTextures, &SortedLightSet, MobilePixelProjectedReflection, bUsingPixelLocalStorage](FRHICommandListImmediate& RHICmdList)
 	{
 		if (GIsEditor && !View.bIsSceneCapture)
 		{
 			DrawClearQuad(RHICmdList, View.BackgroundColor);
 		}
-
+		// Depth pre-pass
 		RHICmdList.SetCurrentStat(GET_STATID(STAT_CLM_MobilePrePass));
-		if (bIsMaskedOnlyDepthPrepassEnabled)
-		{
-			// Depth pre-pass
-			RenderPrePass(RHICmdList, View, &PassParameters->InstanceCullingDrawParams[EMeshPass::DepthPass]);
-		}
-
+		RenderMaskedPrePass(RHICmdList, View, &PassParameters->InstanceCullingDrawParamsDepth);
 		// Opaque and masked
 		RHICmdList.SetCurrentStat(GET_STATID(STAT_CLMM_Opaque));
-		RenderMobileBasePass(RHICmdList, View);
+		RenderMobileBasePass(RHICmdList, View, &PassParameters->InstanceCullingDrawParamsOpaque);
 		RHICmdList.ImmediateFlush(EImmediateFlushType::DispatchToRHIThread);
-
-		if (bShouldRenderOcclusionPass)
-		{
-			// Issue occlusion queries
-			RHICmdList.SetCurrentStat(GET_STATID(STAT_CLMM_Occlusion));
-			RenderOcclusion(RHICmdList, View);
-			RHICmdList.ImmediateFlush(EImmediateFlushType::DispatchToRHIThread);
-		}
-
+		// Issue occlusion queries
+		RHICmdList.SetCurrentStat(GET_STATID(STAT_CLMM_Occlusion));
+		RenderOcclusion(RHICmdList, View);
+		PostRenderBasePass(RHICmdList, View);
 		// SceneColor + GBuffer write, SceneDepth is read only
 		RHICmdList.NextSubpass();
-
-		if (!View.bIsPlanarReflection)
-		{
-			if (ViewFamily.EngineShowFlags.Decals)
-			{
-				CSV_SCOPED_TIMING_STAT_EXCLUSIVE(RenderDecals);
-				RenderDecals(RHICmdList, View);
-			}
-		}
-
+		RenderDecals(RHICmdList, View);
 		// SceneColor write, SceneDepth is read only
 		RHICmdList.NextSubpass();
-
 		if (bUsingPixelLocalStorage)
 		{
 			MobileDeferredCopyBuffer<FMobileDeferredCopyPLSPS>(RHICmdList, View);
-
 			// SceneColor write, SceneDepth is read only
 			RHICmdList.NextSubpass();
 		}
-
 		MobileDeferredShadingPass(RHICmdList, ViewIndex, NumViews, View, *Scene, SortedLightSet, MobilePixelProjectedReflection);
-
 		// Draw translucency.
-		if (bShouldRenderTranslucency)
-		{
-			CSV_SCOPED_TIMING_STAT_EXCLUSIVE(RenderTranslucency);
-			SCOPE_CYCLE_COUNTER(STAT_TranslucencyDrawTime);
-			RenderTranslucency(RHICmdList, View, TranslucencyMeshPass);
-			FRHICommandListExecutor::GetImmediateCommandList().PollOcclusionQueries();
-			RHICmdList.ImmediateFlush(EImmediateFlushType::DispatchToRHIThread);
-		}
+		RenderTranslucency(RHICmdList, View, &PassParameters->InstanceCullingDrawParamsTranslucency);
 	});
 }
 
-void FMobileSceneRenderer::RenderDeferredMultiPass(FRDGBuilder& GraphBuilder, class FMobileRenderPassParameters* PassParameters, FRenderTargetBindingSlots& BasePassRenderTargets, int32 NumColorTargets, int32 ViewIndex, int32 NumViews, FViewInfo& View, FSceneTextures& SceneTextures, const FSortedLightSetSceneInfo& SortedLightSet, FTextureRHIRef MobilePixelProjectedReflection, bool bShouldRenderOcclusionPass, bool bShouldRenderTranslucency, EMeshPass::Type TranslucencyMeshPass)
+void FMobileSceneRenderer::RenderDeferredMultiPass(FRDGBuilder& GraphBuilder, class FMobileRenderPassParameters* PassParameters, FRenderTargetBindingSlots& BasePassRenderTargets, int32 NumColorTargets, int32 ViewIndex, int32 NumViews, FViewInfo& View, FSceneTextures& SceneTextures, const FSortedLightSetSceneInfo& SortedLightSet, FTextureRHIRef MobilePixelProjectedReflection)
 {
 	GraphBuilder.AddPass(
 		RDG_EVENT_NAME("SceneColorRendering"),
 		PassParameters,
 		ERDGPassFlags::Raster,
-		[this, PassParameters, &View, &SceneTextures, bShouldRenderOcclusionPass](FRHICommandListImmediate& RHICmdList)
+		[this, PassParameters, &View, &SceneTextures](FRHICommandListImmediate& RHICmdList)
 	{
 		if (GIsEditor && !View.bIsSceneCapture)
 		{
 			DrawClearQuad(RHICmdList, View.BackgroundColor);
 		}
 
+		// Depth pre-pass
 		RHICmdList.SetCurrentStat(GET_STATID(STAT_CLM_MobilePrePass));
-		if (bIsMaskedOnlyDepthPrepassEnabled)
-		{
-			// Depth pre-pass
-			RenderPrePass(RHICmdList, View, &PassParameters->InstanceCullingDrawParams[EMeshPass::DepthPass]);
-		}
-
+		RenderMaskedPrePass(RHICmdList, View, &PassParameters->InstanceCullingDrawParamsDepth);
 		// Opaque and masked
 		RHICmdList.SetCurrentStat(GET_STATID(STAT_CLMM_Opaque));
-		RenderMobileBasePass(RHICmdList, View);
-		RHICmdList.ImmediateFlush(EImmediateFlushType::DispatchToRHIThread);
-
-		if (bShouldRenderOcclusionPass)
-		{
-			// Issue occlusion queries
-			RHICmdList.SetCurrentStat(GET_STATID(STAT_CLMM_Occlusion));
-			RenderOcclusion(RHICmdList, View);
-			RHICmdList.ImmediateFlush(EImmediateFlushType::DispatchToRHIThread);
-		}
-
-		RHICmdList.NextSubpass();
-		RHICmdList.NextSubpass();
+		RenderMobileBasePass(RHICmdList, View, &PassParameters->InstanceCullingDrawParamsOpaque);
+		// Issue occlusion queries
+		RHICmdList.SetCurrentStat(GET_STATID(STAT_CLMM_Occlusion));
+		RenderOcclusion(RHICmdList, View);
+		PostRenderBasePass(RHICmdList, View);
 	});
 
 	// SceneColor + GBuffer write, SceneDepth is read only
@@ -1531,14 +1364,8 @@ void FMobileSceneRenderer::RenderDeferredMultiPass(FRDGBuilder& GraphBuilder, cl
 		ERDGPassFlags::Raster,
 		[this, SecondPassParameters, &View](FRHICommandListImmediate& RHICmdList)
 	{
-		if (!View.bIsPlanarReflection)
-		{
-			if (ViewFamily.EngineShowFlags.Decals)
-			{
-				CSV_SCOPED_TIMING_STAT_EXCLUSIVE(RenderDecals);
-				RenderDecals(RHICmdList, View);
-			}
-		}
+		RHICmdList.NextSubpass();
+		RenderDecals(RHICmdList, View);
 	});
 
 	// SceneColor write, SceneDepth is read only
@@ -1559,31 +1386,44 @@ void FMobileSceneRenderer::RenderDeferredMultiPass(FRDGBuilder& GraphBuilder, cl
 		{},
 		ThirdPassParameters,
 		ERDGPassFlags::Raster,
-		[this, ThirdPassParameters, ViewIndex, NumViews, &View, &SceneTextures, &SortedLightSet, MobilePixelProjectedReflection, bShouldRenderTranslucency, TranslucencyMeshPass](FRHICommandListImmediate& RHICmdList)
+		[this, ThirdPassParameters, ViewIndex, NumViews, &View, &SceneTextures, &SortedLightSet, MobilePixelProjectedReflection](FRHICommandListImmediate& RHICmdList)
 	{
+		RHICmdList.NextSubpass();
 		MobileDeferredShadingPass(RHICmdList, ViewIndex, NumViews, View, *Scene, SortedLightSet, MobilePixelProjectedReflection);
-
 		// Draw translucency.
-		if (bShouldRenderTranslucency)
+		RenderTranslucency(RHICmdList, View, &ThirdPassParameters->InstanceCullingDrawParamsTranslucency);
+	});
+}
+
+void FMobileSceneRenderer::PostRenderBasePass(FRHICommandListImmediate& RHICmdList, FViewInfo& View)
+{
+	if (ViewFamily.ViewExtensions.Num() > 1)
+	{
+		CSV_SCOPED_TIMING_STAT_EXCLUSIVE(ViewExtensionPostRenderBasePass);
+		QUICK_SCOPE_CYCLE_COUNTER(STAT_FMobileSceneRenderer_ViewExtensionPostRenderBasePass);
+		for (int32 ViewExt = 0; ViewExt < ViewFamily.ViewExtensions.Num(); ++ViewExt)
 		{
-			CSV_SCOPED_TIMING_STAT_EXCLUSIVE(RenderTranslucency);
-			SCOPE_CYCLE_COUNTER(STAT_TranslucencyDrawTime);
-			RenderTranslucency(RHICmdList, View, TranslucencyMeshPass);
-			FRHICommandListExecutor::GetImmediateCommandList().PollOcclusionQueries();
-			RHICmdList.ImmediateFlush(EImmediateFlushType::DispatchToRHIThread);
+			ViewFamily.ViewExtensions[ViewExt]->PostRenderBasePass_RenderThread(RHICmdList, View);
 		}
-	});
+	}
 }
 
 void FMobileSceneRenderer::RenderMobileDebugView(FRHICommandListImmediate& RHICmdList, const FViewInfo& View)
 {
 #if WITH_DEBUG_VIEW_MODES
-	CSV_SCOPED_TIMING_STAT_EXCLUSIVE(RenderDebugView);
-	SCOPED_DRAW_EVENT(RHICmdList, MobileDebugView);
-	SCOPE_CYCLE_COUNTER(STAT_BasePassDrawTime);
+	if (ViewFamily.UseDebugViewPS())
+	{
+		CSV_SCOPED_TIMING_STAT_EXCLUSIVE(RenderDebugView);
+		SCOPED_DRAW_EVENT(RHICmdList, MobileDebugView);
+		SCOPE_CYCLE_COUNTER(STAT_BasePassDrawTime);
+
+		// Here we use the base pass depth result to get z culling for opaque and masque.
+		// The color needs to be cleared at this point since shader complexity renders in additive.
+		DrawClearQuad(RHICmdList, FLinearColor::Black);
 
-	RHICmdList.SetViewport(View.ViewRect.Min.X, View.ViewRect.Min.Y, 0, View.ViewRect.Max.X, View.ViewRect.Max.Y, 1);
-	View.ParallelMeshDrawCommandPasses[EMeshPass::DebugViewMode].DispatchDraw(nullptr, RHICmdList);
+		RHICmdList.SetViewport(View.ViewRect.Min.X, View.ViewRect.Min.Y, 0, View.ViewRect.Max.X, View.ViewRect.Max.Y, 1);
+		View.ParallelMeshDrawCommandPasses[EMeshPass::DebugViewMode].DispatchDraw(nullptr, RHICmdList);
+	}
 #endif // WITH_DEBUG_VIEW_MODES
 }
 
@@ -1627,9 +1467,9 @@ bool FMobileSceneRenderer::RequiresMultiPass(FRHICommandListImmediate& RHICmdLis
 		return false;
 	}
 
+	// Only Vulkan, iOS and some GL can do a single pass deferred shading, otherwise multipass
 	if (IsMobileDeferredShadingEnabled(ShaderPlatform))
 	{
-		// TODO: add GL support
 		return true;
 	}
 		
@@ -1654,75 +1494,6 @@ bool FMobileSceneRenderer::RequiresMultiPass(FRHICommandListImmediate& RHICmdLis
 	return true;
 }
 
-
-void FMobileSceneRenderer::ConditionalResolveSceneDepth(FRDGBuilder& GraphBuilder, const FViewInfo& View, FRDGTextureMSAA SceneDepth)
-{
-	if (IsSimulatedPlatform(ShaderPlatform)) // mobile emulation on PC
-	{
-		// resolve MSAA depth for translucency
-		AddResolveSceneDepthPass(GraphBuilder, View, SceneDepth);
-	}
-	else if (IsAndroidOpenGLESPlatform(ShaderPlatform))
-	{
-		const bool bAlwaysResolveDepth = CVarMobileAlwaysResolveDepth.GetValueOnRenderThread() == 1;
-		// Only these features require depth texture
-		bool bDecals = ViewFamily.EngineShowFlags.Decals && Scene->Decals.Num();
-		bool bModulatedShadows = ViewFamily.EngineShowFlags.DynamicShadows && bModulatedShadowsInUse;
-
-		if (bDecals || bModulatedShadows || bAlwaysResolveDepth || View.bUsesSceneDepth)
-		{
-			// WEBGL copies depth from SceneColor alpha to a separate texture
-			// Switch target to force hardware flush current depth to texture
-			FRDGTextureRef DummySceneColor = GraphBuilder.RegisterExternalTexture(GSystemTextures.BlackDummy, ERenderTargetTexture::Targetable);
-			FRDGTextureRef DummyDepthTarget = GraphBuilder.RegisterExternalTexture(GSystemTextures.DepthDummy, ERenderTargetTexture::Targetable);
-
-			auto* PassParameters = GraphBuilder.AllocParameters<FRenderTargetParameters>();
-
-			PassParameters->RenderTargets[0] = FRenderTargetBinding(DummySceneColor, ERenderTargetLoadAction::ENoAction);
-			PassParameters->RenderTargets.DepthStencil = FDepthStencilBinding(DummyDepthTarget, ERenderTargetLoadAction::EClear, FExclusiveDepthStencil::DepthWrite_StencilWrite);
-
-			GraphBuilder.AddPass(
-				RDG_EVENT_NAME("ResolveDepthPass"), PassParameters, ERDGPassFlags::Raster,
-				[&View, SceneDepth](FRHICommandListImmediate& RHICmdList)
-			{
-				FGraphicsPipelineStateInitializer GraphicsPSOInit;
-				RHICmdList.ApplyCachedRenderTargets(GraphicsPSOInit);
-				GraphicsPSOInit.BlendState = TStaticBlendState<>::GetRHI();
-				GraphicsPSOInit.RasterizerState = TStaticRasterizerState<>::GetRHI();
-				GraphicsPSOInit.DepthStencilState = TStaticDepthStencilState<false, CF_Always>::GetRHI();
-
-				// for devices that do not support framebuffer fetch we rely on undocumented behavior:
-				// Depth reading features will have the depth bound as an attachment AND as a sampler this means
-				// some driver implementations will ignore our attempts to resolve, here we draw with the depth texture to force a resolve.
-				// See UE-37809 for a description of the desired fix.
-				// The results of this draw are irrelevant.
-				TShaderMapRef<FScreenVS> ScreenVertexShader(View.ShaderMap);
-				TShaderMapRef<FScreenPS> PixelShader(View.ShaderMap);
-
-				GraphicsPSOInit.BoundShaderState.VertexDeclarationRHI = GFilterVertexDeclaration.VertexDeclarationRHI;
-				GraphicsPSOInit.BoundShaderState.VertexShaderRHI = ScreenVertexShader.GetVertexShader();
-				GraphicsPSOInit.BoundShaderState.PixelShaderRHI = PixelShader.GetPixelShader();
-				GraphicsPSOInit.PrimitiveType = PT_TriangleList;
-
-				SetGraphicsPipelineState(RHICmdList, GraphicsPSOInit);
-
-				ScreenVertexShader->SetParameters(RHICmdList, View.ViewUniformBuffer);
-				PixelShader->SetParameters(RHICmdList, TStaticSamplerState<SF_Point>::GetRHI(), SceneDepth.Target->GetRHI());
-				DrawRectangle(
-					RHICmdList,
-					0, 0,
-					0, 0,
-					0, 0,
-					1, 1,
-					FIntPoint(1, 1),
-					FIntPoint(1, 1),
-					ScreenVertexShader,
-					EDRF_UseTriangleOptimization);
-			}); // force depth resolve
-		}
-	}
-}
-
 void FMobileSceneRenderer::UpdateDirectionalLightUniformBuffers(FRDGBuilder& GraphBuilder, const FViewInfo& View)
 {
 	if (CachedView == &View)
@@ -1787,7 +1558,7 @@ void FMobileSceneRenderer::PreTonemapMSAA(FRHICommandListImmediate& RHICmdList,
 	// iOS only
 	bool bOnChipPP = GSupportsRenderTargetFormat_PF_FloatRGBA && GSupportsShaderFramebufferFetch &&	ViewFamily.EngineShowFlags.PostProcessing;
 	bool bOnChipPreTonemapMSAA = bOnChipPP && IsMetalMobilePlatform(ViewFamily.GetShaderPlatform()) && (NumMSAASamples > 1);
-	if (!bOnChipPreTonemapMSAA)
+	if (!bOnChipPreTonemapMSAA || bGammaSpace)
 	{
 		return;
 	}
diff --git a/Engine/Source/Runtime/Renderer/Private/MobileTranslucentRendering.cpp b/Engine/Source/Runtime/Renderer/Private/MobileTranslucentRendering.cpp
--- a/Engine/Source/Runtime/Renderer/Private/MobileTranslucentRendering.cpp
+++ b/Engine/Source/Runtime/Renderer/Private/MobileTranslucentRendering.cpp
@@ -31,16 +31,18 @@
 #include "MeshPassProcessor.inl"
 #include "ClearQuad.h"
 
-void FMobileSceneRenderer::RenderTranslucency(FRHICommandListImmediate& RHICmdList, const FViewInfo& View, const EMeshPass::Type TranslucencyMeshPass)
-{
-	SCOPED_DRAW_EVENT(RHICmdList, Translucency);
-	SCOPED_GPU_STAT(RHICmdList, Translucency);
-
-	RHICmdList.SetViewport(View.ViewRect.Min.X, View.ViewRect.Min.Y, 0.0f, View.ViewRect.Max.X, View.ViewRect.Max.Y, 1.0f);
-
-	if (!View.Family->UseDebugViewPS())
+void FMobileSceneRenderer::RenderTranslucency(FRHICommandListImmediate& RHICmdList, const FViewInfo& View, const FInstanceCullingDrawParams* InstanceCullingDrawParams)
+{	
+	const bool bShouldRenderTranslucency = ShouldRenderTranslucency(StandardTranslucencyPass) && ViewFamily.EngineShowFlags.Translucency && !ViewFamily.UseDebugViewPS();
+	if (bShouldRenderTranslucency)
 	{
-		View.ParallelMeshDrawCommandPasses[TranslucencyMeshPass].DispatchDraw(nullptr, RHICmdList);
+		CSV_SCOPED_TIMING_STAT_EXCLUSIVE(RenderTranslucency);
+		SCOPE_CYCLE_COUNTER(STAT_TranslucencyDrawTime);
+		SCOPED_DRAW_EVENT(RHICmdList, Translucency);
+		SCOPED_GPU_STAT(RHICmdList, Translucency);
+
+		RHICmdList.SetViewport(View.ViewRect.Min.X, View.ViewRect.Min.Y, 0.0f, View.ViewRect.Max.X, View.ViewRect.Max.Y, 1.0f);
+		View.ParallelMeshDrawCommandPasses[StandardTranslucencyMeshPass].DispatchDraw(nullptr, RHICmdList, InstanceCullingDrawParams);
 	}
 }
 
diff --git a/Engine/Source/Runtime/Renderer/Private/PostProcess/PostProcessing.cpp b/Engine/Source/Runtime/Renderer/Private/PostProcess/PostProcessing.cpp
--- a/Engine/Source/Runtime/Renderer/Private/PostProcess/PostProcessing.cpp
+++ b/Engine/Source/Runtime/Renderer/Private/PostProcess/PostProcessing.cpp
@@ -1323,7 +1323,7 @@ static bool IsGaussianActive(const FViewInfo& View)
 	return true;
 }
 
-void AddMobilePostProcessingPasses(FRDGBuilder& GraphBuilder, const FViewInfo& View, const FMobilePostProcessingInputs& Inputs, FInstanceCullingManager& InstanceCullingManager)
+void AddMobilePostProcessingPasses(FRDGBuilder& GraphBuilder, FScene* Scene, const FViewInfo& View, const FMobilePostProcessingInputs& Inputs, FInstanceCullingManager& InstanceCullingManager)
 {
 	CSV_SCOPED_TIMING_STAT_EXCLUSIVE(RenderPostProcessing);
 	QUICK_SCOPE_CYCLE_COUNTER(STAT_PostProcessing_Process);
@@ -1527,7 +1527,7 @@ void AddMobilePostProcessingPasses(FRDGBuilder& GraphBuilder, const FViewInfo& V
 			FMobileDistortionAccumulateInputs DistortionAccumulateInputs;
 			DistortionAccumulateInputs.SceneColor = SceneColor;
 
-			FMobileDistortionAccumulateOutputs DistortionAccumulateOutputs = AddMobileDistortionAccumulatePass(GraphBuilder, View, DistortionAccumulateInputs);
+			FMobileDistortionAccumulateOutputs DistortionAccumulateOutputs = AddMobileDistortionAccumulatePass(GraphBuilder, Scene, View, DistortionAccumulateInputs);
 
 			FMobileDistortionMergeInputs DistortionMergeInputs;
 			DistortionMergeInputs.SceneColor = SceneColor;
diff --git a/Engine/Source/Runtime/Renderer/Private/PostProcess/PostProcessing.h b/Engine/Source/Runtime/Renderer/Private/PostProcess/PostProcessing.h
--- a/Engine/Source/Runtime/Renderer/Private/PostProcess/PostProcessing.h
+++ b/Engine/Source/Runtime/Renderer/Private/PostProcess/PostProcessing.h
@@ -80,7 +80,7 @@ struct FMobilePostProcessingInputs
 	}
 };
 
-void AddMobilePostProcessingPasses(FRDGBuilder& GraphBuilder, const FViewInfo& View, const FMobilePostProcessingInputs& Inputs, FInstanceCullingManager& InstanceCullingManager);
+void AddMobilePostProcessingPasses(FRDGBuilder& GraphBuilder, FScene* Scene, const FViewInfo& View, const FMobilePostProcessingInputs& Inputs, FInstanceCullingManager& InstanceCullingManager);
 
 void AddBasicPostProcessPasses(FRDGBuilder& GraphBuilder, const FViewInfo& View);
 
diff --git a/Engine/Source/Runtime/Renderer/Private/SceneOcclusion.cpp b/Engine/Source/Runtime/Renderer/Private/SceneOcclusion.cpp
--- a/Engine/Source/Runtime/Renderer/Private/SceneOcclusion.cpp
+++ b/Engine/Source/Runtime/Renderer/Private/SceneOcclusion.cpp
@@ -1527,6 +1527,11 @@ void FDeferredShadingSceneRenderer::RenderOcclusion(
 
 void FMobileSceneRenderer::RenderOcclusion(FRHICommandListImmediate& RHICmdList, const FViewInfo& View)
 {
+	if (!DoOcclusionQueries())
+	{
+		return;
+	}
+
 	{
 		SCOPED_NAMED_EVENT(FMobileSceneRenderer_BeginOcclusionTests, FColor::Emerald);
 		const FViewOcclusionQueriesPerView QueriesPerView = AllocateOcclusionTests(Scene, VisibleLightInfos, Views);
diff --git a/Engine/Source/Runtime/Renderer/Private/SceneRendering.cpp b/Engine/Source/Runtime/Renderer/Private/SceneRendering.cpp
--- a/Engine/Source/Runtime/Renderer/Private/SceneRendering.cpp
+++ b/Engine/Source/Runtime/Renderer/Private/SceneRendering.cpp
@@ -2864,9 +2864,9 @@ void FSceneRenderer::ComputeFamilySize()
 	check(bInitializedExtents);
 }
 
-bool FSceneRenderer::DoOcclusionQueries(ERHIFeatureLevel::Type InFeatureLevel) const
+bool FSceneRenderer::DoOcclusionQueries() const
 {
-	return InFeatureLevel >= ERHIFeatureLevel::ES3_1 && CVarAllowOcclusionQueries.GetValueOnRenderThread() != 0;
+	return CVarAllowOcclusionQueries.GetValueOnRenderThread() != 0;
 }
 
 FSceneRenderer::~FSceneRenderer()
diff --git a/Engine/Source/Runtime/Renderer/Private/SceneRendering.h b/Engine/Source/Runtime/Renderer/Private/SceneRendering.h
--- a/Engine/Source/Runtime/Renderer/Private/SceneRendering.h
+++ b/Engine/Source/Runtime/Renderer/Private/SceneRendering.h
@@ -1834,7 +1834,7 @@ public:
 	/** Update the rendertarget with each view results.*/
 	void DoCrossGPUTransfers(FRDGBuilder& GraphBuilder, FRHIGPUMask RenderTargetGPUMask, FRDGTextureRef ViewFamilyTexture);
 
-	bool DoOcclusionQueries(ERHIFeatureLevel::Type InFeatureLevel) const;
+	bool DoOcclusionQueries() const;
 
 	void FenceOcclusionTests(FRDGBuilder& GraphBuilder);
 	void FenceOcclusionTestsInternal(FRHICommandListImmediate& RHICmdList);
@@ -2189,11 +2189,13 @@ protected:
 	void InitViews(FRDGBuilder& GraphBuilder, FSceneTexturesConfig& SceneTexturesConfig, FInstanceCullingManager& InstanceCullingManager);
 
 	void RenderPrePass(FRHICommandListImmediate& RHICmdList, const FViewInfo& View, const FInstanceCullingDrawParams* InstanceCullingDrawParams = nullptr);
-
+	void RenderMaskedPrePass(FRHICommandListImmediate& RHICmdList, const FViewInfo& View, const FInstanceCullingDrawParams* InstanceCullingDrawParams);
 	void RenderFullDepthPrepass(FRDGBuilder& GraphBuilder, FSceneTextures& SceneTextures);
 
 	/** Renders the opaque base pass for mobile. */
-	void RenderMobileBasePass(FRHICommandListImmediate& RHICmdList, const FViewInfo& View);
+	void RenderMobileBasePass(FRHICommandListImmediate& RHICmdList, const FViewInfo& View, const FInstanceCullingDrawParams* InstanceCullingDrawParams);
+
+	void PostRenderBasePass(FRHICommandListImmediate& RHICmdList, FViewInfo& View);
 
 	void RenderMobileEditorPrimitives(FRHICommandList& RHICmdList, const FViewInfo& View, const FMeshPassProcessorRenderState& DrawRenderState);
 
@@ -2203,9 +2205,6 @@ protected:
 	/** Render modulated shadow projections in to the scene, loops over any unrendered shadows until all are processed.*/
 	void RenderModulatedShadowProjections(FRHICommandListImmediate& RHICmdList, int32 ViewIndex, const FViewInfo& View);
 
-	/** Resolves scene depth in case hardware does not support reading depth in the shader */
-	void ConditionalResolveSceneDepth(FRDGBuilder& GraphBuilder, const FViewInfo& View, FRDGTextureMSAA SceneDepth);
-
 	/** Issues occlusion queries */
 	void RenderOcclusion(FRHICommandListImmediate& RHICmdList, const FViewInfo& View);
 	
@@ -2226,24 +2225,27 @@ protected:
 	void RenderDecals(FRHICommandListImmediate& RHICmdList, const FViewInfo& View);
 
 	/** Renders the base pass for translucency. */
-	void RenderTranslucency(FRHICommandListImmediate& RHICmdList, const FViewInfo& View, const EMeshPass::Type TranslucencyMeshPass);
+	void RenderTranslucency(FRHICommandListImmediate& RHICmdList, const FViewInfo& View, const FInstanceCullingDrawParams* InstanceCullingDrawParams);
 
 	/** On chip pre-tonemap before scene color MSAA resolve (iOS only) */
 	void PreTonemapMSAA(FRHICommandListImmediate& RHICmdList, const FMinimalSceneTextures& SceneTextures);
 
 	void SortMobileBasePassAfterShadowInit(FExclusiveDepthStencil::Type BasePassDepthStencilAccess, FViewVisibleCommandsPerView& ViewCommandsPerView);
-	void SetupMobileBasePassAfterShadowInit(FExclusiveDepthStencil::Type BasePassDepthStencilAccess, FViewVisibleCommandsPerView& ViewCommandsPerView);
+	void SetupMobileBasePassAfterShadowInit(FExclusiveDepthStencil::Type BasePassDepthStencilAccess, FViewVisibleCommandsPerView& ViewCommandsPerView, FInstanceCullingManager& InstanceCullingManager);
 
 	void UpdateDirectionalLightUniformBuffers(FRDGBuilder& GraphBuilder, const FViewInfo& View);
 	void UpdateSkyReflectionUniformBuffer();
+	
+	void BuildInstanceCullingDrawParams(FRDGBuilder& GraphBuilder, FViewInfo& View, class FMobileRenderPassParameters* PassParameters);
 
 	void RenderForward(FRDGBuilder& GraphBuilder, FRDGTextureRef ViewFamilyTexture, FSceneTextures& SceneTextures);
-	void RenderForwardSinglePass(FRDGBuilder& GraphBuilder, class FMobileRenderPassParameters* PassParameters, int32 ViewIndex, FViewInfo& View, FSceneTextures& SceneTextures, bool bShouldRenderMobileDebugView, bool bAdrenoOcclusionMode, bool bShouldRenderOcclusionPass, bool bShouldRenderTranslucency, EMeshPass::Type TranslucencyMeshPass);
-	void RenderForwardMultiPass(FRDGBuilder& GraphBuilder, class FMobileRenderPassParameters* PassParameters, FRenderTargetBindingSlots& BasePassRenderTargets, int32 ViewIndex, FViewInfo& View, FSceneTextures& SceneTextures, bool bShouldRenderMobileDebugView, bool bAdrenoOcclusionMode, bool bShouldRenderOcclusionPass, bool bShouldRenderTranslucency, EMeshPass::Type TranslucencyMeshPass);
+	void RenderForwardSinglePass(FRDGBuilder& GraphBuilder, class FMobileRenderPassParameters* PassParameters, int32 ViewIndex, FViewInfo& View, FSceneTextures& SceneTextures);
+	void RenderForwardMultiPass(FRDGBuilder& GraphBuilder, class FMobileRenderPassParameters* PassParameters, FRenderTargetBindingSlots& BasePassRenderTargets, int32 ViewIndex, FViewInfo& View, FSceneTextures& SceneTextures);
+	
 	void RenderDeferred(FRDGBuilder& GraphBuilder, const FSortedLightSetSceneInfo& SortedLightSet, FRDGTextureRef ViewFamilyTexture, FSceneTextures& SceneTextures);
-	void RenderDeferredSinglePass(FRDGBuilder& GraphBuilder, class FMobileRenderPassParameters* PassParameters, int32 ViewIndex, int32 NumViews, FViewInfo& View, FSceneTextures& SceneTextures, const FSortedLightSetSceneInfo& SortedLightSet, FTextureRHIRef MobilePixelProjectedReflection, bool bShouldRenderOcclusionPass, bool bUsingPixelLocalStorage, bool bShouldRenderTranslucency, EMeshPass::Type TranslucencyMeshPass);
-	void RenderDeferredMultiPass(FRDGBuilder& GraphBuilder, class FMobileRenderPassParameters* PassParameters, FRenderTargetBindingSlots& BasePassRenderTargets, int32 NumColorTargets, int32 ViewIndex, int32 NumViews, FViewInfo& View, FSceneTextures& SceneTextures, const FSortedLightSetSceneInfo& SortedLightSet, FTextureRHIRef MobilePixelProjectedReflection, bool bShouldRenderOcclusionPass, bool bShouldRenderTranslucency, EMeshPass::Type TranslucencyMeshPass);
-
+	void RenderDeferredSinglePass(FRDGBuilder& GraphBuilder, class FMobileRenderPassParameters* PassParameters, int32 ViewIndex, int32 NumViews, FViewInfo& View, FSceneTextures& SceneTextures, const FSortedLightSetSceneInfo& SortedLightSet, FTextureRHIRef MobilePixelProjectedReflection, bool bUsingPixelLocalStorage);
+	void RenderDeferredMultiPass(FRDGBuilder& GraphBuilder, class FMobileRenderPassParameters* PassParameters, FRenderTargetBindingSlots& BasePassRenderTargets, int32 NumColorTargets, int32 ViewIndex, int32 NumViews, FViewInfo& View, FSceneTextures& SceneTextures, const FSortedLightSetSceneInfo& SortedLightSet, FTextureRHIRef MobilePixelProjectedReflection);
+	
 	void RenderAmbientOcclusion(FRDGBuilder& GraphBuilder, FRDGTextureRef SceneDepthTexture, FRDGTextureRef AmbientOcclusionTexture);
 
 	void RenderPixelProjectedReflection(FRDGBuilder& GraphBuilder, FRDGTextureRef SceneColorTexture, FRDGTextureRef SceneDepthTexture, FRDGTextureRef PixelProjectedReflectionTexture, const FPlanarReflectionSceneProxy* PlanarReflectionSceneProxy);
@@ -2269,6 +2271,10 @@ private:
 	bool bRequiresDistanceFieldShadowingPass;
 	bool bIsFullDepthPrepassEnabled;
 	bool bIsMaskedOnlyDepthPrepassEnabled;
+
+	ETranslucencyPass::Type StandardTranslucencyPass;
+	EMeshPass::Type StandardTranslucencyMeshPass;
+
 	static FGlobalDynamicIndexBuffer DynamicIndexBuffer;
 	static FGlobalDynamicVertexBuffer DynamicVertexBuffer;
 	static TGlobalResource<FGlobalDynamicReadBuffer> DynamicReadBuffer;
diff --git a/Engine/Source/Runtime/Renderer/Private/SceneVisibility.cpp b/Engine/Source/Runtime/Renderer/Private/SceneVisibility.cpp
--- a/Engine/Source/Runtime/Renderer/Private/SceneVisibility.cpp
+++ b/Engine/Source/Runtime/Renderer/Private/SceneVisibility.cpp
@@ -3230,7 +3230,7 @@ void FSceneRenderer::PreVisibilityFrameSetup(FRDGBuilder& GraphBuilder, const FS
 		const bool bIsHitTesting = ViewFamily.EngineShowFlags.HitProxies;
 		// Don't test occlusion queries in collision viewmode as they can be bigger then the rendering bounds.
 		const bool bCollisionView = ViewFamily.EngineShowFlags.CollisionVisibility || ViewFamily.EngineShowFlags.CollisionPawn;
-		if (GIsHighResScreenshot || !DoOcclusionQueries(FeatureLevel) || bIsHitTesting || bCollisionView)
+		if (GIsHighResScreenshot || !DoOcclusionQueries() || bIsHitTesting || bCollisionView)
 		{
 			View.bDisableQuerySubmissions = true;
 			View.bIgnoreExistingQueries = true;
diff --git a/Engine/Source/Runtime/Renderer/Private/ShadowRendering.cpp b/Engine/Source/Runtime/Renderer/Private/ShadowRendering.cpp
--- a/Engine/Source/Runtime/Renderer/Private/ShadowRendering.cpp
+++ b/Engine/Source/Runtime/Renderer/Private/ShadowRendering.cpp
@@ -2224,11 +2224,12 @@ void FDeferredShadingSceneRenderer::RenderDeferredShadowProjections(
 
 void FMobileSceneRenderer::RenderModulatedShadowProjections(FRHICommandListImmediate& RHICmdList, int32 ViewIndex, const FViewInfo& View)
 {
-	if (IsSimpleForwardShadingEnabled(ShaderPlatform) || !ViewFamily.EngineShowFlags.DynamicShadows)
+	if (IsSimpleForwardShadingEnabled(ShaderPlatform) || !ViewFamily.EngineShowFlags.DynamicShadows || View.bIsPlanarReflection)
 	{
 		return;
 	}
 
+	CSV_SCOPED_TIMING_STAT_EXCLUSIVE(RenderShadowProjections);
 	SCOPED_NAMED_EVENT(FMobileSceneRenderer_RenderModulatedShadowProjections, FColor::Emerald);
 	SCOPE_CYCLE_COUNTER(STAT_ProjectedShadowDrawTime);
 	SCOPED_DRAW_EVENT(RHICmdList, ShadowProjectionOnOpaque);
