Fixed: incremental build error

diff --git a/Engine/Source/Runtime/Renderer/Private/MobileShadingRenderer.cpp b/Engine/Source/Runtime/Renderer/Private/MobileShadingRenderer.cpp
--- a/Engine/Source/Runtime/Renderer/Private/MobileShadingRenderer.cpp
+++ b/Engine/Source/Runtime/Renderer/Private/MobileShadingRenderer.cpp
@@ -966,10 +966,12 @@ void FMobileSceneRenderer::RenderForward(FRDGBuilder& GraphBuilder, FRDGTextureR
 		PassParameters->MobileBasePass = CreateMobileBasePassUniformBuffer(GraphBuilder, View, EMobileBasePass::Opaque, LastFrameScreenSpaceAO, MobilePixelProjectedReflection);
 		PassParameters->SceneTextures = SceneTextures.MobileUniformBuffer;
 		PassParameters->RenderTargets = BasePassRenderTargets;
+		#if WITH_DEBUG_VIEW_MODES
 		if (ViewFamily.UseDebugViewPS())
 		{
 			PassParameters->DebugViewMode = CreateDebugViewModePassUniformBuffer(GraphBuilder, View, SceneTextures.QuadOverdraw);
 		}
+		#endif
 		
 		BuildInstanceCullingDrawParams(GraphBuilder, View, PassParameters);
 
@@ -1253,10 +1255,12 @@ void FMobileSceneRenderer::RenderDeferred(FRDGBuilder& GraphBuilder, const FSort
 		PassParameters->MobileBasePass = CreateMobileBasePassUniformBuffer(GraphBuilder, View, EMobileBasePass::Opaque, LastFrameScreenSpaceAO, MobilePixelProjectedReflection);
 		PassParameters->SceneTextures = SceneTextures.MobileUniformBuffer;
 		PassParameters->RenderTargets = BasePassRenderTargets;
+		#if WITH_DEBUG_VIEW_MODES
 		if (ViewFamily.UseDebugViewPS())
 		{
 			PassParameters->DebugViewMode = CreateDebugViewModePassUniformBuffer(GraphBuilder, View, SceneTextures.QuadOverdraw);
 		}
+		#endif
 		
 		BuildInstanceCullingDrawParams(GraphBuilder, View, PassParameters);
 
