Change fallback texture Black Dummy to Volumetric Black Dummy (Texture3D) to match shader texture input.

diff --git a/Engine/Source/Runtime/Renderer/Private/TranslucentLighting.cpp b/Engine/Source/Runtime/Renderer/Private/TranslucentLighting.cpp
--- a/Engine/Source/Runtime/Renderer/Private/TranslucentLighting.cpp
+++ b/Engine/Source/Runtime/Renderer/Private/TranslucentLighting.cpp
@@ -955,10 +955,10 @@ FTranslucencyLightingVolumeParameters GetTranslucencyLightingVolumeParameters(FR
 	else
 	{
 		const FRDGSystemTextures& SystemTextures = FRDGSystemTextures::Get(GraphBuilder);
-		Parameters.TranslucencyLightingVolumeAmbientInner = SystemTextures.Black;
-		Parameters.TranslucencyLightingVolumeAmbientOuter = SystemTextures.Black;
-		Parameters.TranslucencyLightingVolumeDirectionalInner = SystemTextures.Black;
-		Parameters.TranslucencyLightingVolumeDirectionalOuter = SystemTextures.Black;
+		Parameters.TranslucencyLightingVolumeAmbientInner = SystemTextures.VolumetricBlack;
+		Parameters.TranslucencyLightingVolumeAmbientOuter = SystemTextures.VolumetricBlack;
+		Parameters.TranslucencyLightingVolumeDirectionalInner = SystemTextures.VolumetricBlack;
+		Parameters.TranslucencyLightingVolumeDirectionalOuter = SystemTextures.VolumetricBlack;
 	}
 	return Parameters;
 }
