Enable Metal SM5 dual-source blend targets for DiffuseIndirectComposite (indirect lighting).

#rnx

diff --git a/Engine/Shaders/Private/DiffuseIndirectComposite.usf b/Engine/Shaders/Private/DiffuseIndirectComposite.usf
--- a/Engine/Shaders/Private/DiffuseIndirectComposite.usf
+++ b/Engine/Shaders/Private/DiffuseIndirectComposite.usf
@@ -123,8 +123,8 @@ float3 CombineRoughSpecular(FGBufferData GBuffer, float NoV, float3 RayTracedRef
 
 void MainPS(
 	float4 SvPosition : SV_POSITION
-	, out float4 OutAddColor : SV_Target0
-	, out float4 OutMultiplyColor : SV_Target1
+	, out float4 OutAddColor      DUAL_SOURCE_BLENDING_SLOT(0) : SV_Target0
+	, out float4 OutMultiplyColor DUAL_SOURCE_BLENDING_SLOT(1) : SV_Target1
 )
 {
 	float2 SceneBufferUV = SvPositionToBufferUV(SvPosition);
diff --git a/Engine/Shaders/Public/Platform.ush b/Engine/Shaders/Public/Platform.ush
--- a/Engine/Shaders/Public/Platform.ush
+++ b/Engine/Shaders/Public/Platform.ush
@@ -496,6 +496,10 @@ void ShaderYield()
 #define EARLYDEPTHSTENCIL [earlydepthstencil]
 #endif
 
+#if METAL_MRT_PROFILE || METAL_SM5_PROFILE || METAL_SM5_NOTESS_PROFILE
+#define DUAL_SOURCE_BLENDING_SLOT(SLOT) [[vk::location(0), vk::index(SLOT)]]
+#endif
+
 // ---------------------------------------------------- Compiler attribute fallbacks
 
 #ifndef UNROLL
@@ -534,6 +538,10 @@ void ShaderYield()
 #define EARLYDEPTHSTENCIL
 #endif
 
+#ifndef DUAL_SOURCE_BLENDING_SLOT
+#define DUAL_SOURCE_BLENDING_SLOT(SLOT)
+#endif
+
 #ifndef DEPTHSTENCIL_EARLYTEST_LATEWRITE
 #define DEPTHSTENCIL_EARLYTEST_LATEWRITE
 #endif
diff --git a/Engine/Shaders/Public/Platform/Metal/MetalCommon.ush b/Engine/Shaders/Public/Platform/Metal/MetalCommon.ush
--- a/Engine/Shaders/Public/Platform/Metal/MetalCommon.ush
+++ b/Engine/Shaders/Public/Platform/Metal/MetalCommon.ush
@@ -5,7 +5,7 @@
 =============================================================================*/
 
 // Update this GUID to improve shader recompilation for Metal only shaders
-// GUID = A998F35E-2D63-47D6-BD71-B63DB703EE16
+// GUID = 829A0E2C-A667-4ABF-82A8-E27DA13447D2
 
 #pragma once
 
diff --git a/Engine/Source/Runtime/Renderer/Private/IndirectLightRendering.cpp b/Engine/Source/Runtime/Renderer/Private/IndirectLightRendering.cpp
--- a/Engine/Source/Runtime/Renderer/Private/IndirectLightRendering.cpp
+++ b/Engine/Source/Runtime/Renderer/Private/IndirectLightRendering.cpp
@@ -107,7 +107,7 @@ class FDiffuseIndirectCompositePS : public FGlobalShader
 
 	static bool ShouldCompilePermutation(const FGlobalShaderPermutationParameters& Parameters)
 	{
-		if (IsMetalPlatform(Parameters.Platform))
+		if (IsMetalPlatform(Parameters.Platform) && !IsMetalSM5Platform(Parameters.Platform))
 		{
 			return false;
 		}
@@ -911,7 +911,7 @@ void FDeferredShadingSceneRenderer::RenderDiffuseIndirectAndAmbientOcclusion(
 
 		// Applies diffuse indirect and ambient occlusion to the scene color.
 		if ((DenoiserOutputs.Textures[0] || AmbientOcclusionMask) && (!bIsVisualizePass || ViewPipelineState.DiffuseIndirectDenoiser != IScreenSpaceDenoiser::EMode::Disabled || ViewPipelineState.DiffuseIndirectMethod == EDiffuseIndirectMethod::Lumen)
-			&& !IsMetalPlatform(ShaderPlatform))
+			&& !(IsMetalPlatform(ShaderPlatform) && !IsMetalSM5Platform(ShaderPlatform)))
 		{
 			FDiffuseIndirectCompositePS::FParameters* PassParameters = GraphBuilder.AllocParameters<FDiffuseIndirectCompositePS::FParameters>();
 			
