Remove async LPI creation
#test compiled ue4editor; checked several levels in qagame

diff --git a/Engine/Source/Runtime/Renderer/Private/LightSceneInfo.cpp b/Engine/Source/Runtime/Renderer/Private/LightSceneInfo.cpp
--- a/Engine/Source/Runtime/Renderer/Private/LightSceneInfo.cpp
+++ b/Engine/Source/Runtime/Renderer/Private/LightSceneInfo.cpp
@@ -93,8 +93,6 @@ void FLightSceneInfo::AddToScene()
 		// Mobile path supports dynamic point/spot lights in the base pass using forward rendering, so we need to know the primitives
 		|| bIsValidLightTypeMobile)
 	{
-		Scene->FlushAsyncLightPrimitiveInteractionCreation();
-		
 		// Directional lights have no finite extent and cannot meaningfully be in the LocalShadowCastingLightOctree
 		if (LightSceneInfoCompact.LightType == LightType_Directional)
 		{
@@ -147,8 +145,6 @@ void FLightSceneInfo::CreateLightPrimitiveInteraction(const FLightSceneInfoCompa
 
 void FLightSceneInfo::RemoveFromScene()
 {
-	Scene->FlushAsyncLightPrimitiveInteractionCreation();
-
 	if (OctreeId.IsValidId())
 	{
 		// Remove the light from the octree.
@@ -232,30 +228,18 @@ bool FLightSceneInfo::IsPrecomputedLightingValid() const
 	return (bPrecomputedLightingIsValid && NumUnbuiltInteractions < GWholeSceneShadowUnbuiltInteractionThreshold) || !Proxy->HasStaticShadowing();
 }
 
-const TArray<FLightPrimitiveInteraction*>* FLightSceneInfo::GetInteractionShadowPrimitives(bool bSync) const
+const TArray<FLightPrimitiveInteraction*>* FLightSceneInfo::GetInteractionShadowPrimitives() const
 {
-	if (bSync)
-	{
-		Scene->FlushAsyncLightPrimitiveInteractionCreation();
-	}
 	return bRecordInteractionShadowPrimitives ? &InteractionShadowPrimitives : nullptr;
 }
 
-FLightPrimitiveInteraction* FLightSceneInfo::GetDynamicInteractionOftenMovingPrimitiveList(bool bSync) const
+FLightPrimitiveInteraction* FLightSceneInfo::GetDynamicInteractionOftenMovingPrimitiveList() const
 {
-	if (bSync)
-	{
-		Scene->FlushAsyncLightPrimitiveInteractionCreation();
-	}
 	return DynamicInteractionOftenMovingPrimitiveList;
 }
 
-FLightPrimitiveInteraction* FLightSceneInfo::GetDynamicInteractionStaticPrimitiveList(bool bSync) const
+FLightPrimitiveInteraction* FLightSceneInfo::GetDynamicInteractionStaticPrimitiveList() const
 {
-	if (bSync)
-	{
-		Scene->FlushAsyncLightPrimitiveInteractionCreation();
-	}
 	return DynamicInteractionStaticPrimitiveList;
 }
 
diff --git a/Engine/Source/Runtime/Renderer/Private/LightSceneInfo.h b/Engine/Source/Runtime/Renderer/Private/LightSceneInfo.h
--- a/Engine/Source/Runtime/Renderer/Private/LightSceneInfo.h
+++ b/Engine/Source/Runtime/Renderer/Private/LightSceneInfo.h
@@ -332,11 +332,11 @@ public:
 		return DynamicShadowMapChannel;
 	}
 
-	const TArray<FLightPrimitiveInteraction*>* GetInteractionShadowPrimitives(bool bSync = true) const;
+	const TArray<FLightPrimitiveInteraction*>* GetInteractionShadowPrimitives() const;
 
-	FLightPrimitiveInteraction* GetDynamicInteractionOftenMovingPrimitiveList(bool bSync = true) const;
+	FLightPrimitiveInteraction* GetDynamicInteractionOftenMovingPrimitiveList() const;
 
-	FLightPrimitiveInteraction* GetDynamicInteractionStaticPrimitiveList(bool bSync = true) const;
+	FLightPrimitiveInteraction* GetDynamicInteractionStaticPrimitiveList() const;
 
 	/** Hash function. */
 	friend uint32 GetTypeHash(const FLightSceneInfo* LightSceneInfo)
diff --git a/Engine/Source/Runtime/Renderer/Private/RendererScene.cpp b/Engine/Source/Runtime/Renderer/Private/RendererScene.cpp
--- a/Engine/Source/Runtime/Renderer/Private/RendererScene.cpp
+++ b/Engine/Source/Runtime/Renderer/Private/RendererScene.cpp
@@ -108,13 +108,6 @@ static TAutoConsoleVariable<int32> CVarBasePassWriteDepthEvenWithFullPrepass(
 	TEXT("0 to allow a readonly base pass, which skips an MSAA depth resolve, and allows masked materials to get EarlyZ (writing to depth while doing clip() disables EarlyZ) (default)\n")
 	TEXT("1 to force depth writes in the base pass.  Useful for debugging when the prepass and base pass don't match what they render."));
 
-static int32 GAsyncCreateLightPrimitiveInteractions = 1;
-static FAutoConsoleVariableRef CVarAsyncCreateLightPrimitiveInteractions(
-	TEXT("r.AsyncCreateLightPrimitiveInteractions"),
-	GAsyncCreateLightPrimitiveInteractions,
-	TEXT("Whether to create LPIs asynchronously."),
-	ECVF_RenderThreadSafe);
-
 static TAutoConsoleVariable<int32> CVarWPOPrimitivesOutputVelocity(
 	TEXT("r.WPOPrimitivesOutputVelocity"),
 	0,
@@ -1016,7 +1009,6 @@ FScene::FScene(UWorld* InWorld, bool bInRequiresHitProxies, bool bInIsEditorScen
 #if RHI_RAYTRACING
 ,	RayTracingDynamicGeometryCollection(nullptr)
 #endif
-,	AsyncCreateLightPrimitiveInteractionsTask(nullptr)
 ,	NumVisibleLights_GameThread(0)
 ,	NumEnabledSkylights_GameThread(0)
 ,	SceneFrameNumber(0)
@@ -1137,12 +1129,6 @@ FScene::~FScene()
 	BeginReleaseResource(&HaltonPrimesResource);
 #endif
 
-	if (AsyncCreateLightPrimitiveInteractionsTask)
-	{
-		delete AsyncCreateLightPrimitiveInteractionsTask;
-		AsyncCreateLightPrimitiveInteractionsTask = nullptr;
-	}
-
 	checkf(RemovedPrimitiveSceneInfos.Num() == 0, TEXT("Leaking %d FPrimitiveSceneInfo instances."), RemovedPrimitiveSceneInfos.Num()); // Ensure UpdateAllPrimitiveSceneInfos() is called before destruction.
 }
 
@@ -3524,7 +3510,6 @@ void FScene::ApplyWorldOffset_RenderThread(const FVector& InOffset)
 		(*It).LightSceneInfo->Proxy->ApplyWorldOffset(InOffset);
 	}
 
-	FlushAsyncLightPrimitiveInteractionCreation();
 	LocalShadowCastingLightOctree.ApplyOffset(InOffset, /*bGlobalOctee*/ true);
 
 	// Cached preshadows
@@ -3720,20 +3705,6 @@ void FScene::UpdateAllPrimitiveSceneInfos(FRDGBuilder& GraphBuilder, bool bAsync
 	TSet<FPrimitiveSceneInfo*> DeletedSceneInfos;
 	DeletedSceneInfos.Reserve(RemovedLocalPrimitiveSceneInfos.Num());
 
-	if (!!GAsyncCreateLightPrimitiveInteractions && !AsyncCreateLightPrimitiveInteractionsTask && FApp::ShouldUseThreadingForPerformance())
-	{
-		AsyncCreateLightPrimitiveInteractionsTask = new FAsyncTask<FAsyncCreateLightPrimitiveInteractionsTask>();
-	}
-
-	if (AsyncCreateLightPrimitiveInteractionsTask)
-	{
-		if (!AsyncCreateLightPrimitiveInteractionsTask->IsDone())
-		{
-			AsyncCreateLightPrimitiveInteractionsTask->EnsureCompletion();
-		}
-		AsyncCreateLightPrimitiveInteractionsTask->GetTask().Init(this);
-	}
-
 	{
 		SCOPED_NAMED_EVENT(FScene_RemovePrimitiveSceneInfos, FColor::Red);
 		SCOPE_CYCLE_COUNTER(STAT_RemoveScenePrimitiveTime);
@@ -4195,12 +4166,6 @@ void FScene::UpdateAllPrimitiveSceneInfos(FRDGBuilder& GraphBuilder, bool bAsync
 			}
 		}
 
-		if (AsyncCreateLightPrimitiveInteractionsTask && AsyncCreateLightPrimitiveInteractionsTask->GetTask().HasPendingPrimitives())
-		{
-			check(GAsyncCreateLightPrimitiveInteractions);
-			AsyncCreateLightPrimitiveInteractionsTask->StartBackgroundTask();
-		}
-
 		for (const auto& Transform : OverridenPreviousTransforms)
 		{
 			FPrimitiveSceneInfo* PrimitiveSceneInfo = Transform.Key;
@@ -4279,48 +4244,21 @@ void FScene::UpdateAllPrimitiveSceneInfos(FRDGBuilder& GraphBuilder, bool bAsync
 
 void FScene::CreateLightPrimitiveInteractionsForPrimitive(FPrimitiveSceneInfo* PrimitiveInfo, bool bAsyncCreateLPIs)
 {
-	if (!bAsyncCreateLPIs
-		|| !GAsyncCreateLightPrimitiveInteractions
-		|| !FApp::ShouldUseThreadingForPerformance()
-		|| GetShadingPath() == EShadingPath::Mobile
-		|| PrimitiveInfo->LODParentComponentId.IsValid()
-		|| !PrimitiveInfo->Proxy->IsMeshShapeOftenMoving())
-	{
-		FMemMark MemStackMark(FMemStack::Get());
-		const FBoxSphereBounds& Bounds = PrimitiveInfo->Proxy->GetBounds();
-		const FPrimitiveSceneInfoCompact PrimitiveSceneInfoCompact(PrimitiveInfo);
+	FMemMark MemStackMark(FMemStack::Get());
+	const FBoxSphereBounds& Bounds = PrimitiveInfo->Proxy->GetBounds();
+	const FPrimitiveSceneInfoCompact PrimitiveSceneInfoCompact(PrimitiveInfo);
 
-		// Find local lights that affect the primitive in the light octree.
-		LocalShadowCastingLightOctree.FindElementsWithBoundsTest(Bounds.GetBox(), [&PrimitiveSceneInfoCompact](const FLightSceneInfoCompact& LightSceneInfoCompact)
+	// Find local lights that affect the primitive in the light octree.
+	LocalShadowCastingLightOctree.FindElementsWithBoundsTest(Bounds.GetBox(), [&PrimitiveSceneInfoCompact](const FLightSceneInfoCompact& LightSceneInfoCompact)
 		{
 			LightSceneInfoCompact.LightSceneInfo->CreateLightPrimitiveInteraction(LightSceneInfoCompact, PrimitiveSceneInfoCompact);
 		});
-	
-		// Also loop through non-local (directional) shadow-casting lights
-		for (int32 LightID : DirectionalShadowCastingLightIDs)
-		{
-			const FLightSceneInfoCompact& LightSceneInfoCompact = Lights[LightID];
-			LightSceneInfoCompact.LightSceneInfo->CreateLightPrimitiveInteraction(LightSceneInfoCompact, PrimitiveSceneInfoCompact);
-		}
-	}
-	else
-	{
-		check(AsyncCreateLightPrimitiveInteractionsTask);
-		AsyncCreateLightPrimitiveInteractionsTask->GetTask().AddPrimitive(PrimitiveInfo);
-	}
-}
 
-void FScene::FlushAsyncLightPrimitiveInteractionCreation() const
-{
-	if (!GAsyncCreateLightPrimitiveInteractions || !FApp::ShouldUseThreadingForPerformance())
-	{
-		check(!AsyncCreateLightPrimitiveInteractionsTask || AsyncCreateLightPrimitiveInteractionsTask->IsIdle());
-	}
-	else if (AsyncCreateLightPrimitiveInteractionsTask && !AsyncCreateLightPrimitiveInteractionsTask->IsWorkDone())
+	// Also loop through non-local (directional) shadow-casting lights
+	for (int32 LightID : DirectionalShadowCastingLightIDs)
 	{
-		CSV_SCOPED_TIMING_STAT_EXCLUSIVE(FlushAsyncLPICreation);
-		SCOPE_CYCLE_COUNTER(STAT_FlushAsyncLPICreation);
-		AsyncCreateLightPrimitiveInteractionsTask->EnsureCompletion();
+		const FLightSceneInfoCompact& LightSceneInfoCompact = Lights[LightID];
+		LightSceneInfoCompact.LightSceneInfo->CreateLightPrimitiveInteraction(LightSceneInfoCompact, PrimitiveSceneInfoCompact);
 	}
 }
 
diff --git a/Engine/Source/Runtime/Renderer/Private/ScenePrivate.h b/Engine/Source/Runtime/Renderer/Private/ScenePrivate.h
--- a/Engine/Source/Runtime/Renderer/Private/ScenePrivate.h
+++ b/Engine/Source/Runtime/Renderer/Private/ScenePrivate.h
@@ -2969,8 +2969,6 @@ public:
 
 	void CreateLightPrimitiveInteractionsForPrimitive(FPrimitiveSceneInfo* PrimitiveInfo, bool bAsyncCreateLPIs);
 
-	void FlushAsyncLightPrimitiveInteractionCreation() const;
-
 	bool IsPrimitiveBeingRemoved(FPrimitiveSceneInfo* PrimitiveSceneInfo) const;
 
 protected:
@@ -3098,8 +3096,6 @@ private:
 	TSet<FPrimitiveSceneInfo*> RemovedPrimitiveSceneInfos;
 	TSet<FPrimitiveSceneInfo*> DistanceFieldSceneDataUpdates;
 
-	FAsyncTask<class FAsyncCreateLightPrimitiveInteractionsTask>* AsyncCreateLightPrimitiveInteractionsTask;
-
 	/** 
 	 * The number of visible lights in the scene
 	 * Note: This is tracked on the game thread!
diff --git a/Engine/Source/Runtime/Renderer/Private/ShadowSetup.cpp b/Engine/Source/Runtime/Renderer/Private/ShadowSetup.cpp
--- a/Engine/Source/Runtime/Renderer/Private/ShadowSetup.cpp
+++ b/Engine/Source/Runtime/Renderer/Private/ShadowSetup.cpp
@@ -3343,8 +3343,6 @@ void FSceneRenderer::CreateWholeSceneProjectedShadow(
 
 		if (MaxFadeAlpha > 1.0f / 256.0f)
 		{
-			Scene->FlushAsyncLightPrimitiveInteractionCreation();
-
 			for (int32 ShadowIndex = 0, ShadowCount = ProjectedShadowInitializers.Num(); ShadowIndex < ShadowCount; ShadowIndex++)
 			{
 				FWholeSceneProjectedShadowInitializer& ProjectedShadowInitializer = ProjectedShadowInitializers[ShadowIndex];
@@ -3481,8 +3479,8 @@ void FSceneRenderer::CreateWholeSceneProjectedShadow(
 					ProjectedShadowInfo->MeshPassTargetType = EMeshPass::VSMShadowDepth;
 
 					bool bContainsNaniteSubjects = false;
-					AddInteractingPrimitives(LightSceneInfo->GetDynamicInteractionOftenMovingPrimitiveList(false), ProjectedShadowInfo, bContainsNaniteSubjects);
-					AddInteractingPrimitives(LightSceneInfo->GetDynamicInteractionStaticPrimitiveList(false), ProjectedShadowInfo, bContainsNaniteSubjects);
+					AddInteractingPrimitives(LightSceneInfo->GetDynamicInteractionOftenMovingPrimitiveList(), ProjectedShadowInfo, bContainsNaniteSubjects);
+					AddInteractingPrimitives(LightSceneInfo->GetDynamicInteractionStaticPrimitiveList(), ProjectedShadowInfo, bContainsNaniteSubjects);
 
 					VisibleLightInfo.VirtualShadowMapId = ProjectedShadowInfo->VirtualShadowMaps[0]->ID;
 					VisibleLightInfo.AllProjectedShadows.Add(ProjectedShadowInfo);
@@ -3525,12 +3523,12 @@ void FSceneRenderer::CreateWholeSceneProjectedShadow(
 							&& (CacheMode[CacheModeIndex] != SDCM_MovablePrimitivesOnly || bCastCachedShadowFromMovablePrimitives))
 						{
 							// Add all the shadow casting primitives affected by the light to the shadow's subject primitive list.
-							AddInteractingPrimitives(LightSceneInfo->GetDynamicInteractionOftenMovingPrimitiveList(false), ProjectedShadowInfo, bContainsNaniteSubjects);
+							AddInteractingPrimitives(LightSceneInfo->GetDynamicInteractionOftenMovingPrimitiveList(), ProjectedShadowInfo, bContainsNaniteSubjects);
 						}
 						
 						if (CacheMode[CacheModeIndex] != SDCM_MovablePrimitivesOnly)
 						{
-							AddInteractingPrimitives(LightSceneInfo->GetDynamicInteractionStaticPrimitiveList(false), ProjectedShadowInfo, bContainsNaniteSubjects);
+							AddInteractingPrimitives(LightSceneInfo->GetDynamicInteractionStaticPrimitiveList(), ProjectedShadowInfo, bContainsNaniteSubjects);
 						}
 					}
 
@@ -5087,9 +5085,7 @@ void FSceneRenderer::InitDynamicShadows(FRHICommandListImmediate& RHICmdList, FG
 
 						if( !bMobile || (LightSceneInfo->Proxy->CastsModulatedShadows() && !LightSceneInfo->Proxy->UseCSMForDynamicObjects()))
 						{
-							Scene->FlushAsyncLightPrimitiveInteractionCreation();
-
-							const TArray<FLightPrimitiveInteraction*>* InteractionShadowPrimitives = LightSceneInfo->GetInteractionShadowPrimitives(false);
+							const TArray<FLightPrimitiveInteraction*>* InteractionShadowPrimitives = LightSceneInfo->GetInteractionShadowPrimitives();
 
 							if (InteractionShadowPrimitives)
 							{
@@ -5102,7 +5098,7 @@ void FSceneRenderer::InitDynamicShadows(FRHICommandListImmediate& RHICmdList, FG
 							else
 							{
 								// Look for individual primitives with a dynamic shadow.
-								for (FLightPrimitiveInteraction* Interaction = LightSceneInfo->GetDynamicInteractionOftenMovingPrimitiveList(false);
+								for (FLightPrimitiveInteraction* Interaction = LightSceneInfo->GetDynamicInteractionOftenMovingPrimitiveList();
 									Interaction;
 									Interaction = Interaction->GetNextPrimitive()
 									)
@@ -5110,7 +5106,7 @@ void FSceneRenderer::InitDynamicShadows(FRHICommandListImmediate& RHICmdList, FG
 									SetupInteractionShadows(RHICmdList, Interaction, VisibleLightInfo, bStaticSceneOnly, ViewDependentWholeSceneShadows, PreShadows);
 								}
 
-								for (FLightPrimitiveInteraction* Interaction = LightSceneInfo->GetDynamicInteractionStaticPrimitiveList(false);
+								for (FLightPrimitiveInteraction* Interaction = LightSceneInfo->GetDynamicInteractionStaticPrimitiveList();
 									Interaction;
 									Interaction = Interaction->GetNextPrimitive()
 									)
