Fix RHI validation errors

diff --git a/Engine/Source/Runtime/Renderer/Private/HairStrands/HairStrandsCluster.cpp b/Engine/Source/Runtime/Renderer/Private/HairStrands/HairStrandsCluster.cpp
--- a/Engine/Source/Runtime/Renderer/Private/HairStrands/HairStrandsCluster.cpp
+++ b/Engine/Source/Runtime/Renderer/Private/HairStrands/HairStrandsCluster.cpp
@@ -20,14 +20,14 @@ class FHairMacroGroupAABBCS : public FGlobalShader
 		SHADER_PARAMETER(uint32, MacroGroupId)
 		SHADER_PARAMETER(uint32, MacroGroupValid)
 		SHADER_PARAMETER(uint32, bClearBuffer)
-		SHADER_PARAMETER_SRV(Buffer, InGroupAABBBuffer0)
-		SHADER_PARAMETER_SRV(Buffer, InGroupAABBBuffer1)
-		SHADER_PARAMETER_SRV(Buffer, InGroupAABBBuffer2)
-		SHADER_PARAMETER_SRV(Buffer, InGroupAABBBuffer3)
-		SHADER_PARAMETER_SRV(Buffer, InGroupAABBBuffer4)
-		SHADER_PARAMETER_SRV(Buffer, InGroupAABBBuffer5)
-		SHADER_PARAMETER_SRV(Buffer, InGroupAABBBuffer6)
-		SHADER_PARAMETER_SRV(Buffer, InGroupAABBBuffer7)
+		SHADER_PARAMETER_RDG_BUFFER_SRV(Buffer, InGroupAABBBuffer0)
+		SHADER_PARAMETER_RDG_BUFFER_SRV(Buffer, InGroupAABBBuffer1)
+		SHADER_PARAMETER_RDG_BUFFER_SRV(Buffer, InGroupAABBBuffer2)
+		SHADER_PARAMETER_RDG_BUFFER_SRV(Buffer, InGroupAABBBuffer3)
+		SHADER_PARAMETER_RDG_BUFFER_SRV(Buffer, InGroupAABBBuffer4)
+		SHADER_PARAMETER_RDG_BUFFER_SRV(Buffer, InGroupAABBBuffer5)
+		SHADER_PARAMETER_RDG_BUFFER_SRV(Buffer, InGroupAABBBuffer6)
+		SHADER_PARAMETER_RDG_BUFFER_SRV(Buffer, InGroupAABBBuffer7)
 		SHADER_PARAMETER_RDG_BUFFER_UAV(RWBuffer, OutMacroGroupAABBBuffer)
 		END_SHADER_PARAMETER_STRUCT()
 
@@ -68,7 +68,7 @@ static void AddHairMacroGroupAABBPass(
 		{
 			const uint32 PrimitiveIndex = PassIt * GroupPerPass + PassPrimitiveIt;
 			const FHairStrandsMacroGroupData::PrimitiveInfo& PrimitiveInfo = MacroGroup.PrimitivesInfos[PrimitiveIndex];
-			FShaderResourceViewRHIRef GroupAABBBufferSRV = PrimitiveInfo.PublicDataPtr->GetGroupAABBBuffer().SRV;
+			FRDGBufferSRVRef GroupAABBBufferSRV = RegisterAsSRV(GraphBuilder, PrimitiveInfo.PublicDataPtr->GetGroupAABBBuffer());
 
 			// Default value
 			if (PassPrimitiveIt == 0 && PassPrimitiveCount != GroupPerPass)
diff --git a/Engine/Source/Runtime/Renderer/Private/HairStrands/HairStrandsVoxelization.cpp b/Engine/Source/Runtime/Renderer/Private/HairStrands/HairStrandsVoxelization.cpp
--- a/Engine/Source/Runtime/Renderer/Private/HairStrands/HairStrandsVoxelization.cpp
+++ b/Engine/Source/Runtime/Renderer/Private/HairStrands/HairStrandsVoxelization.cpp
@@ -242,7 +242,7 @@ class FVoxelMarkValidPageIndex_PrepareCS : public FGlobalShader
 		SHADER_PARAMETER(uint32, MaxScatterAllocationCount)
 		SHADER_PARAMETER(uint32, bForceDirectPageAllocation)
 
-		SHADER_PARAMETER_SRV(Buffer, ClusterAABBsBuffer)
+		SHADER_PARAMETER_RDG_BUFFER_SRV(Buffer, ClusterAABBsBuffer)
 		SHADER_PARAMETER_RDG_BUFFER_SRV(Buffer, MacroGroupAABBBuffer)
 		SHADER_PARAMETER_RDG_BUFFER_SRV(Buffer, PageIndexResolutionAndOffsetBuffer)
 		SHADER_PARAMETER_RDG_BUFFER_UAV(RWBuffer<uint>, OutValidPageIndexBuffer)
@@ -315,7 +315,7 @@ class FVoxelMarkValidPageIndexCS : public FGlobalShader
 		SHADER_PARAMETER(FVector, CPU_MaxAABB)
 		SHADER_PARAMETER(uint32, CPU_PageIndexOffset)
 		SHADER_PARAMETER(uint32, MacroGroupId)
-		SHADER_PARAMETER_SRV(Buffer, ClusterAABBsBuffer)
+		SHADER_PARAMETER_RDG_BUFFER_SRV(Buffer, ClusterAABBsBuffer)
 		SHADER_PARAMETER_RDG_BUFFER_SRV(Buffer, MacroGroupAABBBuffer)
 		SHADER_PARAMETER_RDG_BUFFER_SRV(Buffer, PageIndexResolutionAndOffsetBuffer)
 		SHADER_PARAMETER_RDG_BUFFER_UAV(RWBuffer<uint>, OutValidPageIndexBuffer)
@@ -649,7 +649,7 @@ static void AddAllocateVoxelPagesPass(
 				Parameters->MaxScatterAllocationCount	= MaxAllocationCount;
 				Parameters->bForceDirectPageAllocation	= GHairVirtualVoxelUseImmediatePageAllocation > 0 ? 1 : 0;
 
-				Parameters->ClusterAABBsBuffer		= HairGroupData->GetClusterAABBBuffer().SRV;
+				Parameters->ClusterAABBsBuffer		= RegisterAsSRV(GraphBuilder, HairGroupData->GetClusterAABBBuffer());
 				Parameters->MacroGroupAABBBuffer	= GraphBuilder.CreateSRV(MacroGroupResources.MacroGroupAABBsBuffer, PF_R32_SINT);
 				Parameters->PageIndexResolutionAndOffsetBuffer = PageIndexResolutionAndOffsetBufferSRV;
 
@@ -713,7 +713,7 @@ static void AddAllocateVoxelPagesPass(
 				Parameters->CPU_PageIndexOffset = CPUAllocationDesc.PageIndexOffset;
 				Parameters->CPU_MinAABB = CPUAllocationDesc.MinAABB;
 				Parameters->CPU_MaxAABB = CPUAllocationDesc.MaxAABB;
-				Parameters->ClusterAABBsBuffer = HairGroupData->GetClusterAABBBuffer().SRV;
+				Parameters->ClusterAABBsBuffer = RegisterAsSRV(GraphBuilder, HairGroupData->GetClusterAABBBuffer());
 				Parameters->OutValidPageIndexBuffer = PageIndexBufferUAV;
 
 				if (bIsGPUDriven)
