Add hair buffers allocation based on screen coverage.

This allows to reduce GPU memory allocation, with an estimate of the hair screen coverage. This scheme is only efficient with the new directx12 transient allocator.

diff --git a/Engine/Source/Runtime/Renderer/Private/HairStrands/HairStrandsVisibility.cpp b/Engine/Source/Runtime/Renderer/Private/HairStrands/HairStrandsVisibility.cpp
--- a/Engine/Source/Runtime/Renderer/Private/HairStrands/HairStrandsVisibility.cpp
+++ b/Engine/Source/Runtime/Renderer/Private/HairStrands/HairStrandsVisibility.cpp
@@ -1948,10 +1948,15 @@ static void AddHairVisibilityPrimitiveIdCompactionPass(
 	AddClearUAVPass(GraphBuilder, GraphBuilder.CreateUAV(OutCompactNodeIndex), ClearValues);
 	AddClearUAVPass(GraphBuilder, GraphBuilder.CreateUAV(OutCategorizationTexture), ClearValues);
 
+	// Adapt the buffer allocation based on the bounding box of the hair macro groups. This allows to reduce the overall allocation size
+	const FIntRect HairRect = ComputeVisibleHairStrandsMacroGroupsRect(View.ViewRect, MacroGroupDatas);
+	const FIntPoint EffectiveResolution(HairRect.Width(), HairRect.Height());
+
 	// Select render node count according to current mode
 	const uint32 MSAASampleCount = GetHairVisibilityRenderMode() == HairVisibilityRenderMode_MSAA ? GetMaxSamplePerPixel() : 1;
 	const uint32 PPLLMaxRenderNodePerPixel = GetMaxSamplePerPixel();
-	const uint32 MaxRenderNodeCount = GetTotalSampleCountForAllocation(Resolution);
+	const uint32 MaxRenderNodeCount = GetTotalSampleCountForAllocation(EffectiveResolution);
+
 	OutCompactNodeData = GraphBuilder.CreateBuffer(FRDGBufferDesc::CreateStructuredDesc(bUseVisibility ? sizeof(HairStrandsVisibilityInternal::NodeVis) : sizeof(HairStrandsVisibilityInternal::NodeData), MaxRenderNodeCount), TEXT("Hair.VisibilityPrimitiveIdCompactNodeData"));
 
 	{
@@ -2577,7 +2582,10 @@ static void AddHairVisibilityPPLLPass(
 		OutVisibilityPPLLNodeIndex = GraphBuilder.CreateTexture(Desc, TEXT("Hair.VisibilityPPLLNodeIndex"));
 	}
 
-	const uint32 PPLLMaxTotalListElementCount = GetTotalSampleCountForAllocation(Resolution);
+`	const FIntRect HairRect = ComputeVisibleHairStrandsMacroGroupsRect(ViewInfo->ViewRect, MacroGroupDatas);
+	const FIntPoint EffectiveResolution(HairRect.Width(), HairRect.Height());
+
+	const uint32 PPLLMaxTotalListElementCount = GetTotalSampleCountForAllocation(EffectiveResolution);
 	{
 		OutVisibilityPPLLNodeData = GraphBuilder.CreateBuffer(FRDGBufferDesc::CreateStructuredDesc(sizeof(PPLLNodeData), PPLLMaxTotalListElementCount), TEXT("Hair.VisibilityPPLLNodeData"));
 	}
@@ -3018,7 +3026,7 @@ class FVisiblityRasterComputeCS : public FGlobalShader
 		SHADER_PARAMETER_RDG_TEXTURE_UAV(RWTexture2D, OutVisibilityTexture2)
 		SHADER_PARAMETER_RDG_TEXTURE_UAV(RWTexture2D, OutVisibilityTexture3)
 		SHADER_PARAMETER_STRUCT_REF(FViewUniformShaderParameters, ViewUniformBuffer)
-		END_SHADER_PARAMETER_STRUCT()
+	END_SHADER_PARAMETER_STRUCT()
 
 public:
 	static bool ShouldCompilePermutation(const FGlobalShaderPermutationParameters& Parameters) 
