Fix typo.

diff --git a/Engine/Source/Runtime/Renderer/Private/HairStrands/HairStrandsVisibility.cpp b/Engine/Source/Runtime/Renderer/Private/HairStrands/HairStrandsVisibility.cpp
--- a/Engine/Source/Runtime/Renderer/Private/HairStrands/HairStrandsVisibility.cpp
+++ b/Engine/Source/Runtime/Renderer/Private/HairStrands/HairStrandsVisibility.cpp
@@ -2582,7 +2582,7 @@ static void AddHairVisibilityPPLLPass(
 		OutVisibilityPPLLNodeIndex = GraphBuilder.CreateTexture(Desc, TEXT("Hair.VisibilityPPLLNodeIndex"));
 	}
 
-`	const FIntRect HairRect = ComputeVisibleHairStrandsMacroGroupsRect(ViewInfo->ViewRect, MacroGroupDatas);
+	const FIntRect HairRect = ComputeVisibleHairStrandsMacroGroupsRect(ViewInfo->ViewRect, MacroGroupDatas);
 	const FIntPoint EffectiveResolution(HairRect.Width(), HairRect.Height());
 
 	const uint32 PPLLMaxTotalListElementCount = GetTotalSampleCountForAllocation(EffectiveResolution);
