Remove unecessary resources tracking from the HairVisibility data.

diff --git a/Engine/Source/Runtime/Renderer/Private/HairStrands/HairStrandsData.h b/Engine/Source/Runtime/Renderer/Private/HairStrands/HairStrandsData.h
--- a/Engine/Source/Runtime/Renderer/Private/HairStrands/HairStrandsData.h
+++ b/Engine/Source/Runtime/Renderer/Private/HairStrands/HairStrandsData.h
@@ -30,33 +30,16 @@ END_GLOBAL_SHADER_PARAMETER_STRUCT()
 
 struct FHairStrandsVisibilityData
 {
-	FRDGTextureRef DepthTexture = nullptr;
-	FRDGTextureRef IDTexture = nullptr;
-	FRDGTextureRef MaterialTexture = nullptr;
-	FRDGTextureRef AttributeTexture = nullptr;
 	FRDGTextureRef VelocityTexture = nullptr;
 	FRDGTextureRef ResolveMaskTexture = nullptr;
 	FRDGTextureRef CategorizationTexture = nullptr;
 	FRDGTextureRef ViewHairCountTexture = nullptr;
 	FRDGTextureRef ViewHairCountUintTexture = nullptr;
-	FRDGTextureRef DepthTextureUint = nullptr;
 	FRDGTextureRef EmissiveTexture = nullptr;
-
 	FRDGTextureRef HairOnlyDepthTexture = nullptr;
 
-	FRDGTextureRef ViewHairVisibilityTexture0 = nullptr;
-	FRDGTextureRef ViewHairVisibilityTexture1 = nullptr;
-	FRDGTextureRef ViewHairVisibilityTexture2 = nullptr;
-	FRDGTextureRef ViewHairVisibilityTexture3 = nullptr;
-
 	FRDGTextureRef LightChannelMaskTexture = nullptr;
 
-	FRDGTextureRef	PPLLNodeCounterTexture = nullptr;
-	FRDGTextureRef	PPLLNodeIndexTexture = nullptr;
-	FRDGBufferRef	PPLLNodeDataBuffer = nullptr;
-	uint32			MaxPPLLNodePerPixelCount = 0;
-	uint32			MaxPPLLNodeCount = 0;
-
 	FRDGTextureRef	TileIndexTexture = nullptr;
 	FRDGBufferRef	TileBuffer = nullptr;
 	FRDGBufferRef	TileIndirectArgs = nullptr;
@@ -310,7 +293,7 @@ struct FHairStrandsDebugData
 		FRDGBufferRef SampleCounter = nullptr;
 	} Resources;
 
-	bool IsValid() const
+	bool IsPlotDataValid() const
 	{
 		return Resources.ShadingPointBuffer && Resources.ShadingPointCounter && Resources.SampleBuffer && Resources.SampleCounter;
 	}
@@ -318,6 +301,16 @@ struct FHairStrandsDebugData
 	static Data CreateData(FRDGBuilder& GraphBuilder);
 	static void SetParameters(FRDGBuilder& GraphBuilder, Data& In, FWriteParameters& Out);
 	static void SetParameters(FRDGBuilder& GraphBuilder, const Data& In, FReadParameters& Out);
+
+	// PPLL debug data
+	bool IsPPLLDataValid() const 
+	{ 
+		return PPLLNodeCounterTexture && PPLLNodeIndexTexture && PPLLNodeDataBuffer; 
+	}
+
+	FRDGTextureRef	PPLLNodeCounterTexture = nullptr;
+	FRDGTextureRef	PPLLNodeIndexTexture = nullptr;
+	FRDGBufferRef	PPLLNodeDataBuffer = nullptr;
 };
 
 typedef TArray<FHairStrandsMacroGroupData, SceneRenderingAllocator> FHairStrandsMacroGroupDatas;
diff --git a/Engine/Source/Runtime/Renderer/Private/HairStrands/HairStrandsDebug.cpp b/Engine/Source/Runtime/Renderer/Private/HairStrands/HairStrandsDebug.cpp
--- a/Engine/Source/Runtime/Renderer/Private/HairStrands/HairStrandsDebug.cpp
+++ b/Engine/Source/Runtime/Renderer/Private/HairStrands/HairStrandsDebug.cpp
@@ -57,8 +57,6 @@ static FAutoConsoleVariableRef CVarHairStrandsDebugMaterialSampleIndex(TEXT("r.H
 static int32 GHairStrandsClusterDebug = 0;
 static FAutoConsoleVariableRef CVarHairStrandsDebugClusterAABB(TEXT("r.HairStrands.Cluster.Debug"), GHairStrandsClusterDebug, TEXT("Draw debug the world bounding box of hair clusters used for culling optimisation (0:off, 1:visible cluster, 2:culled cluster, 3:colored LOD, 4:LOD info)."));
 
-static int32 GHairStrandsDebugPPLL = 0;
-static FAutoConsoleVariableRef CVarHairStrandsDebugPPLL(									TEXT("r.HairStrands.DebugPPLL"),									GHairStrandsDebugPPLL, TEXT("Draw debug per pixel light list rendering."));
 static int32 GHairVirtualVoxel_DrawDebugPage = 0;
 static FAutoConsoleVariableRef CVarHairVirtualVoxel_DrawDebugPage(TEXT("r.HairStrands.Voxelization.Virtual.DrawDebugPage"), GHairVirtualVoxel_DrawDebugPage, TEXT("When voxel debug rendering is enable 1: render the page bounds, instead of the voxel 2: the occupancy within the page (i.e., 8x8x8 brick)"));
 static int32 GHairVirtualVoxel_ForceMipLevel = -1;
@@ -1055,7 +1053,7 @@ static void AddDrawDebugClusterPass(
 }
 
 ///////////////////////////////////////////////////////////////////////////////////////////////////
-
+uint32 GetHairStrandsMeanSamplePerPixel();
 static void InternalRenderHairStrandsDebugInfo(
 	FRDGBuilder& GraphBuilder,
 	FViewInfo& View,
@@ -1094,13 +1092,13 @@ static void InternalRenderHairStrandsDebugInfo(
 	
 	const FHairStrandsViewData& HairData = View.HairStrandsViewData;
 
-	if (GHairStrandsDebugPlotBsdf > 0 || HairData.DebugData.IsValid())
+	if (GHairStrandsDebugPlotBsdf > 0 || HairData.DebugData.IsPlotDataValid())
 	{
 		if (GHairStrandsDebugPlotBsdf > 0)
 		{
 			AddPlotBSDFPass(GraphBuilder, View, SceneColorTexture);
 		}
-		if (HairData.DebugData.IsValid())
+		if (HairData.DebugData.IsPlotDataValid())
 		{
 			AddPlotSamplePass(GraphBuilder, View, HairData.DebugData.Resources, SceneColorTexture);
 		}	
@@ -1281,16 +1279,15 @@ static void InternalRenderHairStrandsDebugInfo(
 	}
 
 	{
-		const FHairStrandsVisibilityData& VisibilityData = HairData.VisibilityData;
-		if (GHairStrandsDebugPPLL && VisibilityData.PPLLNodeCounterTexture) // Check if PPLL rendering is used and its debug view is enabled.
+		if (HairData.DebugData.IsPPLLDataValid()) // Check if PPLL rendering is used and its debug view is enabled.
 		{
-			const FIntPoint PPLLResolution = VisibilityData.PPLLNodeIndexTexture->Desc.Extent;
+			const FIntPoint PPLLResolution = HairData.DebugData.PPLLNodeIndexTexture->Desc.Extent;
 			FHairVisibilityDebugPPLLCS::FParameters* PassParameters = GraphBuilder.AllocParameters<FHairVisibilityDebugPPLLCS::FParameters>();
-			PassParameters->PPLLMeanListElementCountPerPixel = float(VisibilityData.MaxSampleCount);
-			PassParameters->PPLLMaxTotalListElementCount = float(VisibilityData.MaxSampleCount * PPLLResolution.X * PPLLResolution.Y);
-			PassParameters->PPLLCounter = VisibilityData.PPLLNodeCounterTexture;
-			PassParameters->PPLLNodeIndex = VisibilityData.PPLLNodeIndexTexture;
-			PassParameters->PPLLNodeData = GraphBuilder.CreateSRV(FRDGBufferSRVDesc(VisibilityData.PPLLNodeDataBuffer));
+			PassParameters->PPLLMeanListElementCountPerPixel = GetHairStrandsMeanSamplePerPixel();
+			PassParameters->PPLLMaxTotalListElementCount = HairData.DebugData.PPLLNodeDataBuffer->Desc.NumElements;
+			PassParameters->PPLLCounter = HairData.DebugData.PPLLNodeCounterTexture;
+			PassParameters->PPLLNodeIndex = HairData.DebugData.PPLLNodeIndexTexture;
+			PassParameters->PPLLNodeData = GraphBuilder.CreateSRV(FRDGBufferSRVDesc(HairData.DebugData.PPLLNodeDataBuffer));
 			PassParameters->ViewUniformBuffer = View.ViewUniformBuffer;
 			PassParameters->SceneColorTextureUAV = GraphBuilder.CreateUAV(SceneColorTexture);
 			ShaderPrint::SetParameters(GraphBuilder, View, PassParameters->ShaderPrintParameters);
diff --git a/Engine/Source/Runtime/Renderer/Private/HairStrands/HairStrandsVisibility.cpp b/Engine/Source/Runtime/Renderer/Private/HairStrands/HairStrandsVisibility.cpp
--- a/Engine/Source/Runtime/Renderer/Private/HairStrands/HairStrandsVisibility.cpp
+++ b/Engine/Source/Runtime/Renderer/Private/HairStrands/HairStrandsVisibility.cpp
@@ -82,6 +82,9 @@ static FAutoConsoleVariableRef CVarHairStrandsVisibility_UseFastPath(TEXT("r.Hai
 static int32 GHairStrandsVisibility_OutputEmissiveData = 0;
 static FAutoConsoleVariableRef CVarHairStrandsVisibility_OutputEmissiveData(TEXT("r.HairStrands.Visibility.Emissive"), GHairStrandsVisibility_OutputEmissiveData, TEXT("Enable emissive data during the material pass."));
 
+static int32 GHairStrandsDebugPPLL = 0;
+static FAutoConsoleVariableRef CVarHairStrandsDebugPPLL(TEXT("r.HairStrands.Visibility.PPLL.Debug"), GHairStrandsDebugPPLL, TEXT("Draw debug per pixel light list rendering."));
+
 /////////////////////////////////////////////////////////////////////////////////////////
 
 namespace HairStrandsVisibilityInternal
@@ -230,6 +233,11 @@ inline uint32 GetMeanSamplePerPixel()
 	return 1;
 }
 
+uint32 GetHairStrandsMeanSamplePerPixel()
+{
+	return GetMeanSamplePerPixel();
+}
+
 struct FRasterComputeOutput
 {
 	FIntPoint BaseResolution;
@@ -3365,11 +3373,6 @@ void RenderHairStrandsVisibilityBuffer(
 						VisibilityData.EmissiveTexture = PassOutput.EmissiveTexture;
 					}
 
-					VisibilityData.ViewHairVisibilityTexture0 = RasterOutput.VisibilityTexture0;
-					VisibilityData.ViewHairVisibilityTexture1 = RasterOutput.VisibilityTexture1;
-					VisibilityData.ViewHairVisibilityTexture2 = RasterOutput.VisibilityTexture2;
-					VisibilityData.ViewHairVisibilityTexture3 = RasterOutput.VisibilityTexture3;
-
 					// For fully covered pixels, write: 
 					// * black color into the scene color
 					// * closest depth
@@ -3484,13 +3487,9 @@ void RenderHairStrandsVisibilityBuffer(
 
 				// This is used when compaction is not enabled.
 				VisibilityData.MaxSampleCount = MsaaVisibilityResources.IdTexture->Desc.NumSamples;
-				VisibilityData.IDTexture = MsaaVisibilityResources.IdTexture;
-				VisibilityData.DepthTexture = MsaaVisibilityResources.DepthTexture;
 				VisibilityData.HairOnlyDepthTexture = HairOnlyDepthTexture;
 				if (!bIsVisiblityEnable)
 				{
-					VisibilityData.MaterialTexture = MsaaVisibilityResources.MaterialTexture;
-					VisibilityData.AttributeTexture = MsaaVisibilityResources.AttributeTexture;
 					VisibilityData.VelocityTexture = MsaaVisibilityResources.VelocityTexture;
 				}
 
@@ -3723,9 +3722,12 @@ void RenderHairStrandsVisibilityBuffer(
 
 			#if WITH_EDITOR
 				// Extract texture for debug visualization
-				VisibilityData.PPLLNodeCounterTexture = PPLLNodeCounterTexture;
-				VisibilityData.PPLLNodeIndexTexture = PPLLNodeIndexTexture;
-				VisibilityData.PPLLNodeDataBuffer = PPLLNodeDataBuffer;
+				if (GHairStrandsDebugPPLL > 0)
+				{
+					View.HairStrandsViewData.DebugData.PPLLNodeCounterTexture = PPLLNodeCounterTexture;
+					View.HairStrandsViewData.DebugData.PPLLNodeIndexTexture = PPLLNodeIndexTexture;
+					View.HairStrandsViewData.DebugData.PPLLNodeDataBuffer = PPLLNodeDataBuffer;
+				}
 			#endif
 			}
 
