Removed optional Flush right after scene color rendering on mobile (no benefits enabling it)

diff --git a/Engine/Source/Runtime/Renderer/Private/MobileShadingRenderer.cpp b/Engine/Source/Runtime/Renderer/Private/MobileShadingRenderer.cpp
--- a/Engine/Source/Runtime/Renderer/Private/MobileShadingRenderer.cpp
+++ b/Engine/Source/Runtime/Renderer/Private/MobileShadingRenderer.cpp
@@ -74,13 +74,6 @@ static TAutoConsoleVariable<int32> CVarMobileAdrenoOcclusionMode(
 	TEXT("1: Render occlusion queries after translucency and a flush, which can help Adreno devices in GL mode."),
 	ECVF_RenderThreadSafe);
 
-static TAutoConsoleVariable<int32> CVarMobileFlushSceneColorRendering(
-	TEXT("r.Mobile.FlushSceneColorRendering"),
-	1,
-	TEXT("0: Submmit command buffer after all rendering is finished.\n")
-	TEXT("1: Submmit command buffer (flush) before starting post-processing (default)"),
-	ECVF_RenderThreadSafe);
-
 static TAutoConsoleVariable<int32> CVarMobileCustomDepthForTranslucency(
 	TEXT("r.Mobile.CustomDepthForTranslucency"),
 	1,
@@ -206,7 +199,6 @@ FMobileSceneRenderer::FMobileSceneRenderer(const FSceneViewFamily* InViewFamily,
 	bRenderToSceneColor = false;
 	bRequiresMultiPass = false;
 	bKeepDepthContent = false;
-	bSubmitOffscreenRendering = false;
 	bModulatedShadowsInUse = false;
 	bShouldRenderCustomDepth = false;
 	bRequiresPixelProjectedPlanarRelfectionPass = false;
@@ -344,13 +336,11 @@ void FMobileSceneRenderer::InitViews(FRDGBuilder& GraphBuilder, FSceneTexturesCo
 
 	check(Scene);
 
-#if GPUCULL_TODO
 	// Create GPU-side reprensenation of the view for instance culling.
 	for (int32 ViewIndex = 0; ViewIndex < Views.Num(); ++ViewIndex)
 	{
 		Views[ViewIndex].GPUSceneViewId = InstanceCullingManager.RegisterView(Views[ViewIndex]);
 	}
-#endif //GPUCULL_TODO
 
 	if (bUseVirtualTexturing)
 	{
@@ -568,9 +558,6 @@ void FMobileSceneRenderer::InitViews(FRDGBuilder& GraphBuilder, FSceneTexturesCo
 	UpdatePrimitiveIndirectLightingCacheBuffers();
 	
 	OnStartRender(RHICmdList);
-
-	// Whether to submit cmdbuffer with offscreen rendering before doing post-processing
-	bSubmitOffscreenRendering = (!bGammaSpace || bRenderToSceneColor) && CVarMobileFlushSceneColorRendering.GetValueOnAnyThread() != 0;
 }
 
 /*
@@ -804,16 +791,6 @@ void FMobileSceneRenderer::Render(FRDGBuilder& GraphBuilder)
 		);
 	}
 
-	// Flush / submit cmdbuffer
-	if (bSubmitOffscreenRendering)
-	{
-		AddPass(GraphBuilder, [this](FRHICommandListImmediate& InRHICmdList)
-		{
-			InRHICmdList.SubmitCommandsHint();
-			InRHICmdList.ImmediateFlush(EImmediateFlushType::DispatchToRHIThread);
-		});
-	}
-
 	if (bRequiresDistanceFieldShadowingPass)
 	{
 		CSV_SCOPED_TIMING_STAT_EXCLUSIVE(RenderSDFShadowing);
@@ -1030,7 +1007,6 @@ void FMobileSceneRenderer::RenderForwardSinglePass(FRDGBuilder& GraphBuilder, FM
 			RHICmdList.SetCurrentStat(GET_STATID(STAT_CLMM_Occlusion));
 			// flush
 			RHICmdList.SubmitCommandsHint();
-			bSubmitOffscreenRendering = false; // submit once
 			// Issue occlusion queries
 			RenderOcclusion(RHICmdList, View);
 		}
diff --git a/Engine/Source/Runtime/Renderer/Private/SceneRendering.h b/Engine/Source/Runtime/Renderer/Private/SceneRendering.h
--- a/Engine/Source/Runtime/Renderer/Private/SceneRendering.h
+++ b/Engine/Source/Runtime/Renderer/Private/SceneRendering.h
@@ -2265,7 +2265,6 @@ private:
 	bool bRenderToSceneColor;
 	bool bRequiresMultiPass;
 	bool bKeepDepthContent;
-	bool bSubmitOffscreenRendering;
 	bool bModulatedShadowsInUse;
 	bool bShouldRenderCustomDepth;
 	bool bRequiresPixelProjectedPlanarRelfectionPass;
