Fixed CIS issue.

diff --git a/Engine/Source/Runtime/Renderer/Private/HairStrands/HairStrandsCluster.cpp b/Engine/Source/Runtime/Renderer/Private/HairStrands/HairStrandsCluster.cpp
--- a/Engine/Source/Runtime/Renderer/Private/HairStrands/HairStrandsCluster.cpp
+++ b/Engine/Source/Runtime/Renderer/Private/HairStrands/HairStrandsCluster.cpp
@@ -126,12 +126,12 @@ static bool DoesGroupExists(uint32 ResourceId, uint32 GroupIndex, const FHairStr
 	return false;
 }
 
-inline const FHairGroupPublicData* GetHairStrandsPublicData(const FMeshBatchAndRelevance& InMeshBatchAndRelevance)
+inline const FHairGroupPublicData* GetHairStrandsPublicData(const FMeshBatch* InMeshBatch)
 {
-	 return reinterpret_cast<const FHairGroupPublicData*>(InMeshBatchAndRelevance.Mesh->Elements[0].VertexFactoryUserData);
+	 return reinterpret_cast<const FHairGroupPublicData*>(InMeshBatch->Elements[0].VertexFactoryUserData);
 }
 
-static void InternalUpdateMacroGroup(FHairStrandsMacroGroupData& MacroGroup, int32& MaterialId, const FMeshBatchAndRelevance* MeshBatchAndRelevance, const FPrimitiveSceneProxy* Proxy)
+static void InternalUpdateMacroGroup(FHairStrandsMacroGroupData& MacroGroup, int32& MaterialId, const FMeshBatchAndRelevance* MeshBatchAndRelevance)
 {
 	if (MeshBatchAndRelevance)
 	{
@@ -145,7 +145,8 @@ static void InternalUpdateMacroGroup(FHairStrandsMacroGroupData& MacroGroup, int
 		}
 
 		FHairStrandsMacroGroupData::PrimitiveInfo& PrimitiveInfo = MacroGroup.PrimitivesInfos.AddZeroed_GetRef();
-		PrimitiveInfo.MeshBatchAndRelevance = *MeshBatchAndRelevance;
+		PrimitiveInfo.Mesh = MeshBatchAndRelevance->Mesh;
+		PrimitiveInfo.PrimitiveSceneProxy = MeshBatchAndRelevance->PrimitiveSceneProxy;
 		PrimitiveInfo.MaterialId = MaterialId++;
 		PrimitiveInfo.ResourceId = reinterpret_cast<uint64>(MeshBatchAndRelevance->Mesh->Elements[0].UserData);
 		PrimitiveInfo.GroupIndex = HairGroupPublicData->GetGroupIndex();
@@ -188,7 +189,7 @@ void CreateHairStrandsMacroGroups(
 			if (bIntersect)
 			{
 				MacroGroup.Bounds = Union(MacroGroup.Bounds, PrimitiveBounds);
-				InternalUpdateMacroGroup(MacroGroup, MaterialId, MeshBatchAndRelevance, Proxy);
+				InternalUpdateMacroGroup(MacroGroup, MaterialId, MeshBatchAndRelevance);
 				bFound = true;
 				break;
 			}
@@ -198,7 +199,7 @@ void CreateHairStrandsMacroGroups(
 		{
 			FHairStrandsMacroGroupData MacroGroup;
 			MacroGroup.MacroGroupId = MacroGroupId++;
-			InternalUpdateMacroGroup(MacroGroup, MaterialId, MeshBatchAndRelevance, Proxy);
+			InternalUpdateMacroGroup(MacroGroup, MaterialId, MeshBatchAndRelevance);
 			MacroGroup.Bounds = PrimitiveBounds;
 			MacroGroups.Add(MacroGroup);
 		}
@@ -234,6 +235,6 @@ void CreateHairStrandsMacroGroups(
 
 bool FHairStrandsMacroGroupData::PrimitiveInfo::IsCullingEnable() const
 {
-	const FHairGroupPublicData* HairGroupPublicData = GetHairStrandsPublicData(MeshBatchAndRelevance);
+	const FHairGroupPublicData* HairGroupPublicData = GetHairStrandsPublicData(Mesh);
 	return HairGroupPublicData->GetCullingResultAvailable();
 }
\ No newline at end of file
diff --git a/Engine/Source/Runtime/Renderer/Private/HairStrands/HairStrandsData.h b/Engine/Source/Runtime/Renderer/Private/HairStrands/HairStrandsData.h
--- a/Engine/Source/Runtime/Renderer/Private/HairStrands/HairStrandsData.h
+++ b/Engine/Source/Runtime/Renderer/Private/HairStrands/HairStrandsData.h
@@ -10,6 +10,9 @@
 #include "RendererInterface.h"
 #include "Shader.h"
 
+class FPrimitiveSceneProxy;
+struct FMeshBatch;
+
 ////////////////////////////////////////////////////////////////////////////////////
 // HairStrands uniform buffer
 
@@ -223,7 +226,8 @@ struct FHairStrandsMacroGroupData
 	// List of primitive/mesh batch within an instance group
 	struct PrimitiveInfo
 	{
-		FMeshBatchAndRelevance MeshBatchAndRelevance;
+		const FMeshBatch* Mesh = nullptr;
+		const FPrimitiveSceneProxy* PrimitiveSceneProxy = nullptr;
 		uint32 MaterialId;
 		uint32 ResourceId;
 		uint32 GroupIndex;
diff --git a/Engine/Source/Runtime/Renderer/Private/HairStrands/HairStrandsDebug.cpp b/Engine/Source/Runtime/Renderer/Private/HairStrands/HairStrandsDebug.cpp
--- a/Engine/Source/Runtime/Renderer/Private/HairStrands/HairStrandsDebug.cpp
+++ b/Engine/Source/Runtime/Renderer/Private/HairStrands/HairStrandsDebug.cpp
@@ -1000,8 +1000,8 @@ static void AddDrawDebugClusterPass(
 
 				for (const FHairStrandsMacroGroupData::PrimitiveInfo& PrimitiveInfo : MacroGroupData.PrimitivesInfos)
 				{
-					check(PrimitiveInfo.MeshBatchAndRelevance.Mesh && PrimitiveInfo.MeshBatchAndRelevance.Mesh->Elements.Num() > 0);
-					const FHairGroupPublicData* HairGroupPublicData = reinterpret_cast<const FHairGroupPublicData*>(PrimitiveInfo.MeshBatchAndRelevance.Mesh->Elements[0].VertexFactoryUserData);
+					check(PrimitiveInfo.Mesh && PrimitiveInfo.Mesh->Elements.Num() > 0);
+					const FHairGroupPublicData* HairGroupPublicData = reinterpret_cast<const FHairGroupPublicData*>(PrimitiveInfo.Mesh->Elements[0].VertexFactoryUserData);
 
 					for (int DataIndex = 0; DataIndex < HairClusterData.HairGroups.Num(); ++DataIndex)
 					{
diff --git a/Engine/Source/Runtime/Renderer/Private/HairStrands/HairStrandsRasterCommon.cpp b/Engine/Source/Runtime/Renderer/Private/HairStrands/HairStrandsRasterCommon.cpp
--- a/Engine/Source/Runtime/Renderer/Private/HairStrands/HairStrandsRasterCommon.cpp
+++ b/Engine/Source/Runtime/Renderer/Private/HairStrands/HairStrandsRasterCommon.cpp
@@ -414,9 +414,9 @@ void AddHairStrandsRasterPass(
 		for (const FHairStrandsMacroGroupData::PrimitiveInfo& PrimitiveInfo : PrimitiveSceneInfos)
 		{
 			const bool bCullingEnable = PrimitiveInfo.IsCullingEnable();
-			const FMeshBatch& MeshBatch = *PrimitiveInfo.MeshBatchAndRelevance.Mesh;
+			const FMeshBatch& MeshBatch = *PrimitiveInfo.Mesh;
 			const uint64 BatchElementMask = ~0ull;
-			HairRasterMeshProcessor.AddMeshBatch(MeshBatch, BatchElementMask, PrimitiveInfo.MeshBatchAndRelevance.PrimitiveSceneProxy, -1 , bCullingEnable);
+			HairRasterMeshProcessor.AddMeshBatch(MeshBatch, BatchElementMask, PrimitiveInfo.PrimitiveSceneProxy, -1 , bCullingEnable);
 		}
 
 		if (VisibleMeshDrawCommands.Num() > 0)
diff --git a/Engine/Source/Runtime/Renderer/Private/HairStrands/HairStrandsVisibility.cpp b/Engine/Source/Runtime/Renderer/Private/HairStrands/HairStrandsVisibility.cpp
--- a/Engine/Source/Runtime/Renderer/Private/HairStrands/HairStrandsVisibility.cpp
+++ b/Engine/Source/Runtime/Renderer/Private/HairStrands/HairStrandsVisibility.cpp
@@ -735,9 +735,9 @@ static void AddHairMaterialGBufferPass(
 				{
 					for (const FHairStrandsMacroGroupData::PrimitiveInfo& PrimitiveInfo : MacroGroupData.PrimitivesInfos)
 					{
-						const FMeshBatch& MeshBatch = *PrimitiveInfo.MeshBatchAndRelevance.Mesh;
+						const FMeshBatch& MeshBatch = *PrimitiveInfo.Mesh;
 						const uint64 BatchElementMask = ~0ull;
-						MeshProcessor.AddMeshBatch(MeshBatch, BatchElementMask, PrimitiveInfo.MeshBatchAndRelevance.PrimitiveSceneProxy, -1, MacroGroupData.MacroGroupId, PrimitiveInfo.MaterialId);
+						MeshProcessor.AddMeshBatch(MeshBatch, BatchElementMask, PrimitiveInfo.PrimitiveSceneProxy, -1, MacroGroupData.MacroGroupId, PrimitiveInfo.MaterialId);
 					}
 				}
 
@@ -1176,9 +1176,9 @@ static FMaterialPassOutput AddHairMaterialPass(
 			{
 				for (const FHairStrandsMacroGroupData::PrimitiveInfo& PrimitiveInfo : MacroGroupData.PrimitivesInfos)
 				{
-					const FMeshBatch& MeshBatch = *PrimitiveInfo.MeshBatchAndRelevance.Mesh;
+					const FMeshBatch& MeshBatch = *PrimitiveInfo.Mesh;
 					const uint64 BatchElementMask = ~0ull;
-					MeshProcessor.AddMeshBatch(MeshBatch, BatchElementMask, PrimitiveInfo.MeshBatchAndRelevance.PrimitiveSceneProxy, -1, MacroGroupData.MacroGroupId, PrimitiveInfo.MaterialId);
+					MeshProcessor.AddMeshBatch(MeshBatch, BatchElementMask, PrimitiveInfo.PrimitiveSceneProxy, -1, MacroGroupData.MacroGroupId, PrimitiveInfo.MaterialId);
 				}
 			}
 
@@ -2463,9 +2463,9 @@ static void AddHairVisibilityCommonPass(
 			{
 				for (const FHairStrandsMacroGroupData::PrimitiveInfo& PrimitiveInfo : MacroGroupData.PrimitivesInfos)
 				{
-					const FMeshBatch& MeshBatch = *PrimitiveInfo.MeshBatchAndRelevance.Mesh;
+					const FMeshBatch& MeshBatch = *PrimitiveInfo.Mesh;
 					const uint64 BatchElementMask = ~0ull;
-					MeshProcessor.AddMeshBatch(MeshBatch, BatchElementMask, PrimitiveInfo.MeshBatchAndRelevance.PrimitiveSceneProxy, -1, MacroGroupData.MacroGroupId, PrimitiveInfo.MaterialId, PrimitiveInfo.IsCullingEnable());
+					MeshProcessor.AddMeshBatch(MeshBatch, BatchElementMask, PrimitiveInfo.PrimitiveSceneProxy, -1, MacroGroupData.MacroGroupId, PrimitiveInfo.MaterialId, PrimitiveInfo.IsCullingEnable());
 				}
 			}
 
@@ -3195,8 +3195,8 @@ static FRasterComputeOutput AddVisibilityComputeRasterPass(
 			PassParameters->OutVisibilityTexture2 = VisibilityTexture2UAV;
 			PassParameters->OutVisibilityTexture3 = VisibilityTexture3UAV;
 			
-			check(PrimitiveInfo.MeshBatchAndRelevance.Mesh && PrimitiveInfo.MeshBatchAndRelevance.Mesh->Elements.Num() > 0);
-			const FHairGroupPublicData* HairGroupPublicData = reinterpret_cast<const FHairGroupPublicData*>(PrimitiveInfo.MeshBatchAndRelevance.Mesh->Elements[0].VertexFactoryUserData);
+			check(PrimitiveInfo.Mesh && PrimitiveInfo.Mesh->Elements.Num() > 0);
+			const FHairGroupPublicData* HairGroupPublicData = reinterpret_cast<const FHairGroupPublicData*>(PrimitiveInfo.Mesh->Elements[0].VertexFactoryUserData);
 			check(HairGroupPublicData);
 
 			const FHairGroupPublicData::FVertexFactoryInput& VFInput = HairGroupPublicData->VFInput;
diff --git a/Engine/Source/Runtime/Renderer/Private/HairStrands/HairStrandsVoxelization.cpp b/Engine/Source/Runtime/Renderer/Private/HairStrands/HairStrandsVoxelization.cpp
--- a/Engine/Source/Runtime/Renderer/Private/HairStrands/HairStrandsVoxelization.cpp
+++ b/Engine/Source/Runtime/Renderer/Private/HairStrands/HairStrandsVoxelization.cpp
@@ -1046,8 +1046,8 @@ static void AddVirtualVoxelizationComputeRasterPass(
 		TShaderMapRef<FVoxelRasterComputeCS> ComputeShader_CullingOn(ViewInfo->ShaderMap, PermutationVector_On);
 		for (const FHairStrandsMacroGroupData::PrimitiveInfo& PrimitiveInfo : PrimitiveSceneInfos)
 		{
-			check(PrimitiveInfo.MeshBatchAndRelevance.Mesh && PrimitiveInfo.MeshBatchAndRelevance.Mesh->Elements.Num() > 0);
-			const FHairGroupPublicData* HairGroupPublicData = reinterpret_cast<const FHairGroupPublicData*>(PrimitiveInfo.MeshBatchAndRelevance.Mesh->Elements[0].VertexFactoryUserData);
+			check(PrimitiveInfo.Mesh && PrimitiveInfo.Mesh->Elements.Num() > 0);
+			const FHairGroupPublicData* HairGroupPublicData = reinterpret_cast<const FHairGroupPublicData*>(PrimitiveInfo.Mesh->Elements[0].VertexFactoryUserData);
 
 			if (HairGroupPublicData->DoesSupportVoxelization())
 			{
