Move Nanite visualization into post processing pass creation to occur after tonemap and before shader print.

diff --git a/Engine/Source/Runtime/Renderer/Private/DeferredShadingRenderer.cpp b/Engine/Source/Runtime/Renderer/Private/DeferredShadingRenderer.cpp
--- a/Engine/Source/Runtime/Renderer/Private/DeferredShadingRenderer.cpp
+++ b/Engine/Source/Runtime/Renderer/Private/DeferredShadingRenderer.cpp
@@ -2812,41 +2812,6 @@ void FDeferredShadingSceneRenderer::Render(FRDGBuilder& GraphBuilder)
 	}
 #endif
 
-	// TODO: NANITE_VIEW_MODES: Ideal place to run the visualization passes, but need to remove NANITE_USE_SCRATCH_BUFFERS first
-	/*if (bVisualizeNanite)
-	{
-		Nanite::AddVisualizationPasses(
-			GraphBuilder,
-			Scene,
-			SceneTextures,
-			Views,
-			NaniteRasterResults
-		);
-	}*/
-	if (bVisualizeNanite)
-	{
-		if (Views.Num() > 0 && NaniteRasterResults.Num() > 0)
-		{
-			const FViewInfo& View = Views[0];
-			const Nanite::FRasterResults& RasterResults = NaniteRasterResults[0];
-			
-			// TODO: NANITE_VIEW_MODES: Add resample + tile grid and border for overview (multiple visualizations)
-			if (RasterResults.Visualizations.Num() == 1)
-			{
-				const Nanite::FVisualizeResult& Visualization = RasterResults.Visualizations[0];
-				AddDrawTexturePass(
-					GraphBuilder,
-					View,
-					Visualization.ModeOutput,
-					SceneTextures.Color.Target,
-					View.ViewRect.Min,
-					View.ViewRect.Min,
-					View.ViewRect.Size()
-				);
-			}
-		}
-	}
-
 	RendererModule.RenderOverlayExtensions(GraphBuilder, Views, SceneTextures);
 
 	if (ViewFamily.EngineShowFlags.VisualizeDistanceFieldAO && ShouldRenderDistanceFieldLighting())
diff --git a/Engine/Source/Runtime/Renderer/Private/PostProcess/PostProcessing.cpp b/Engine/Source/Runtime/Renderer/Private/PostProcess/PostProcessing.cpp
--- a/Engine/Source/Runtime/Renderer/Private/PostProcess/PostProcessing.cpp
+++ b/Engine/Source/Runtime/Renderer/Private/PostProcess/PostProcessing.cpp
@@ -49,6 +49,7 @@
 #include "PixelShaderUtils.h"
 #include "ScreenSpaceRayTracing.h"
 #include "SceneViewExtension.h"
+#include "NaniteVisualizationData.h"
 #include "FXSystem.h"
 
 bool IsMobileEyeAdaptationEnabled(const FViewInfo& View);
@@ -1014,6 +1015,25 @@ void AddPostProcessingPasses(FRDGBuilder& GraphBuilder, const FViewInfo& View, c
 		AddTestImagePass(GraphBuilder, View, SceneColor);
 	}
 
+	const FNaniteVisualizationData& NaniteVisualizationData = GetNaniteVisualizationData();
+	if (EngineShowFlags.VisualizeNanite && NaniteRasterResults != nullptr && NaniteVisualizationData.IsActive())
+	{
+		// TODO: NANITE_VIEW_MODES: Add resample + tile grid and border for overview (multiple visualizations)
+		if (NaniteRasterResults->Visualizations.Num() == 1)
+		{
+			const Nanite::FVisualizeResult& Visualization = NaniteRasterResults->Visualizations[0];
+			AddDrawTexturePass(
+				GraphBuilder,
+				View,
+				Visualization.ModeOutput,
+				SceneColor.Texture,
+				View.ViewRect.Min,
+				View.ViewRect.Min,
+				View.ViewRect.Size()
+			);
+		}
+	}
+
 	if (ShaderDrawDebug::IsShaderDrawDebugEnabled(View))
 	{
 		ShaderDrawDebug::DrawView(GraphBuilder, View, SceneColor.Texture, SceneDepth.Texture);
