Cubemap VSM no longer draws NumFaces^2 instances.

diff --git a/Engine/Shaders/Private/VirtualShadowMaps/BuildPerPageDrawCommands.usf b/Engine/Shaders/Private/VirtualShadowMaps/BuildPerPageDrawCommands.usf
--- a/Engine/Shaders/Private/VirtualShadowMaps/BuildPerPageDrawCommands.usf
+++ b/Engine/Shaders/Private/VirtualShadowMaps/BuildPerPageDrawCommands.usf
@@ -178,8 +178,12 @@ void CullPerPageDrawCommandsCs(uint InstanceIdIndex : SV_DispatchThreadID, uint
 	GroupMemoryBarrierWithGroupSync();
 
 	uint ThreadTotalForAllViews = 0;
+#if LOOP_OVER_VIEWS
 	// Loop over views and output visible instance (i.e., those that overlap a valid page)
 	for (uint PrimaryViewId = FirstPrimaryView; PrimaryViewId < FirstPrimaryView + NumPrimaryViews; ++PrimaryViewId)
+#else
+	uint PrimaryViewId = ViewIdIndex;
+#endif
 	{
 		uint VisibleMipsMask = 0;
 		uint ThreadTotalPageOverlaps = 0;
diff --git a/Engine/Source/Runtime/Renderer/Private/VirtualShadowMaps/VirtualShadowMapArray.cpp b/Engine/Source/Runtime/Renderer/Private/VirtualShadowMaps/VirtualShadowMapArray.cpp
--- a/Engine/Source/Runtime/Renderer/Private/VirtualShadowMaps/VirtualShadowMapArray.cpp
+++ b/Engine/Source/Runtime/Renderer/Private/VirtualShadowMaps/VirtualShadowMapArray.cpp
@@ -1409,8 +1409,9 @@ class FCullPerPageDrawCommandsCs : public FGlobalShader
 	DECLARE_GLOBAL_SHADER(FCullPerPageDrawCommandsCs);
 	SHADER_USE_PARAMETER_STRUCT(FCullPerPageDrawCommandsCs, FGlobalShader)
 
-	class FNearClipDim : SHADER_PERMUTATION_BOOL( "NEAR_CLIP" );
-	using FPermutationDomain = TShaderPermutationDomain< FNearClipDim >;
+	class FNearClipDim		: SHADER_PERMUTATION_BOOL( "NEAR_CLIP" );
+	class FLoopOverViewsDim	: SHADER_PERMUTATION_BOOL( "LOOP_OVER_VIEWS" );
+	using FPermutationDomain = TShaderPermutationDomain< FNearClipDim, FLoopOverViewsDim >;
 
 public:
 	static constexpr int32 NumThreadsPerGroup = 64;
@@ -1717,6 +1718,7 @@ void FVirtualShadowMapArray::RenderVirtualShadowMapsHw(FRDGBuilder& GraphBuilder
 
 				FCullPerPageDrawCommandsCs::FPermutationDomain PermutationVector;
 				PermutationVector.Set< FCullPerPageDrawCommandsCs::FNearClipDim >( !Clipmap.IsValid() );
+				PermutationVector.Set< FCullPerPageDrawCommandsCs::FLoopOverViewsDim >( Clipmap.IsValid() );
 
 				auto ComputeShader = GetGlobalShaderMap(GMaxRHIFeatureLevel)->GetShader<FCullPerPageDrawCommandsCs>( PermutationVector );
 
@@ -1812,7 +1814,7 @@ void FVirtualShadowMapArray::RenderVirtualShadowMapsHw(FRDGBuilder& GraphBuilder
 			ShadowDepthPassParameters->VirtualSmPageTable	= GraphBuilder.CreateSRV( PageTableRDG );
 			ShadowDepthPassParameters->PackedNaniteViews	= GraphBuilder.CreateSRV( VirtualShadowViewsRDG );
 			ShadowDepthPassParameters->PageRectBounds		= GraphBuilder.CreateSRV( PageRectBoundsRDG );
-			ShadowDepthPassParameters->OutDepthBuffer		= GraphBuilder.CreateUAV( PhysicalPagePoolRDG );
+			ShadowDepthPassParameters->OutDepthBuffer		= GraphBuilder.CreateUAV( PhysicalPagePoolRDG, ERDGUnorderedAccessViewFlags::SkipBarrier );
 			
 			PassParameters->ShadowDepthPass = GraphBuilder.CreateUniformBuffer(ShadowDepthPassParameters);
 			PassParameters->VirtualShadowMap = GetUniformBuffer(GraphBuilder);
