Fix unity compile error

diff --git a/Engine/Source/Runtime/Renderer/Private/DistanceFieldLightingShared.h b/Engine/Source/Runtime/Renderer/Private/DistanceFieldLightingShared.h
--- a/Engine/Source/Runtime/Renderer/Private/DistanceFieldLightingShared.h
+++ b/Engine/Source/Runtime/Renderer/Private/DistanceFieldLightingShared.h
@@ -15,6 +15,7 @@
 #include "DistanceFieldAtlas.h"
 #include "Templates/UniquePtr.h"
 #include "SceneRendering.h"
+#include "ScenePrivate.h"
 
 class FLightSceneProxy;
 class FMaterialRenderProxy;
