Add sanity check and minor clean up to hair transmittance mask.

diff --git a/Engine/Source/Runtime/Renderer/Private/HairStrands/HairStrandsTransmittance.cpp b/Engine/Source/Runtime/Renderer/Private/HairStrands/HairStrandsTransmittance.cpp
--- a/Engine/Source/Runtime/Renderer/Private/HairStrands/HairStrandsTransmittance.cpp
+++ b/Engine/Source/Runtime/Renderer/Private/HairStrands/HairStrandsTransmittance.cpp
@@ -145,7 +145,6 @@ FHairStrandsTransmittanceMaskData CreateDummyHairStrandsTransmittanceMaskData(FR
 {
 	FHairStrandsTransmittanceMaskData Out;
 	Out.TransmittanceMask = CreateHairStrandsTransmittanceMaskBuffer(GraphBuilder, 1);
-	Out.TransmittanceMaskSRV = GraphBuilder.CreateSRV(Out.TransmittanceMask);
 	AddHairStrandsClearTransmittanceMaskPass(GraphBuilder, ShaderMap, Out.TransmittanceMask);
 	return Out;
 }
@@ -554,11 +553,12 @@ static FHairStrandsTransmittanceMaskData RenderHairStrandsTransmittanceMask(
 	const FHairStrandsVoxelResources& VoxelResources,
 	FRDGTextureRef ScreenShadowMaskSubPixelTexture)
 {
+	FHairStrandsTransmittanceMaskData Out;
 	if (MacroGroupDatas.Num() == 0)
-		return FHairStrandsTransmittanceMaskData();
+		return Out;
 
 	if (!HasDeepShadowData(LightSceneInfo, MacroGroupDatas) && !IsHairStrandsVoxelizationEnable())
-		return FHairStrandsTransmittanceMaskData();
+		return Out;
 
 	DECLARE_GPU_STAT(HairStrandsTransmittanceMask);
 	RDG_EVENT_SCOPE(GraphBuilder, "HairStrandsTransmittanceMask");
@@ -580,7 +580,6 @@ static FHairStrandsTransmittanceMaskData RenderHairStrandsTransmittanceMask(
 	memset(Params.DeepShadow_CPUWorldToLightTransforms, 0, sizeof(Params.DeepShadow_CPUWorldToLightTransforms));
 
 
-	FRDGBufferRef OutShadowMask = nullptr;
 	bool bHasFoundLight = false;
 	if (!IsHairStrandsForVoxelTransmittanceAndShadowEnable())
 	{
@@ -617,7 +616,7 @@ static FHairStrandsTransmittanceMaskData RenderHairStrandsTransmittanceMask(
 		{
 			check(Params.DeepShadow_FrontDepthTexture);
 			check(Params.DeepShadow_DomTexture);
-			OutShadowMask = AddDeepShadowTransmittanceMaskPass(
+			Out.TransmittanceMask = AddDeepShadowTransmittanceMaskPass(
 				GraphBuilder,
 				SceneTextures,
 				View,
@@ -643,7 +642,7 @@ static FHairStrandsTransmittanceMaskData RenderHairStrandsTransmittanceMask(
 		Params.LightRadius = FMath::Max(LightParameters.SourceLength, LightParameters.SourceRadius);
 		Params.VirtualVoxelResources = &VoxelResources;
 
-		OutShadowMask = AddDeepShadowTransmittanceMaskPass(
+		Out.TransmittanceMask = AddDeepShadowTransmittanceMaskPass(
 			GraphBuilder,
 			SceneTextures,
 			View,
@@ -655,10 +654,7 @@ static FHairStrandsTransmittanceMaskData RenderHairStrandsTransmittanceMask(
 			ScreenShadowMaskSubPixelTexture);
 	}
 
-	FHairStrandsTransmittanceMaskData OutTransmittanceMaskData;
-	OutTransmittanceMaskData.TransmittanceMask = OutShadowMask;
-	OutTransmittanceMaskData.TransmittanceMaskSRV = GraphBuilder.CreateSRV(OutShadowMask);
-	return OutTransmittanceMaskData;
+	return Out;
 }
 
 FHairStrandsTransmittanceMaskData RenderHairStrandsTransmittanceMask(
diff --git a/Engine/Source/Runtime/Renderer/Private/HairStrands/HairStrandsTransmittance.h b/Engine/Source/Runtime/Renderer/Private/HairStrands/HairStrandsTransmittance.h
--- a/Engine/Source/Runtime/Renderer/Private/HairStrands/HairStrandsTransmittance.h
+++ b/Engine/Source/Runtime/Renderer/Private/HairStrands/HairStrandsTransmittance.h
@@ -16,7 +16,6 @@ class FViewInfo;
 struct FHairStrandsTransmittanceMaskData
 {
 	FRDGBufferRef TransmittanceMask = nullptr;
-	FRDGBufferSRVRef TransmittanceMaskSRV = nullptr;
 };
 
 /// Write opaque hair shadow onto screen shadow mask to have fine hair details cast onto opaque geometries
diff --git a/Engine/Source/Runtime/Renderer/Private/LightRendering.cpp b/Engine/Source/Runtime/Renderer/Private/LightRendering.cpp
--- a/Engine/Source/Runtime/Renderer/Private/LightRendering.cpp
+++ b/Engine/Source/Runtime/Renderer/Private/LightRendering.cpp
@@ -1164,7 +1164,7 @@ void FDeferredShadingSceneRenderer::RenderLights(
 	const int32 SimpleLightsEnd = SortedLightSet.SimpleLightsEnd;
 
 	FHairStrandsTransmittanceMaskData DummyTransmittanceMaskData;
-	if (HairStrands::HasViewHairStrandsData(Views) && Views.Num() > 0)
+	if (bUseHairLighting && Views.Num() > 0)
 	{
 		DummyTransmittanceMaskData = CreateDummyHairStrandsTransmittanceMaskData(GraphBuilder, Views[0].ShaderMap);
 	}
@@ -2437,6 +2437,9 @@ void FDeferredShadingSceneRenderer::RenderLightForHair(
 			GetCloudShadowAOParameters(GraphBuilder, View, CloudInfo), 
 			PassParameters->Light);
 
+		// Sanity check
+		check(InTransmittanceMaskData.TransmittanceMask);
+
 		PassParameters->HairStrands = HairStrands::BindHairStrandsViewUniformParameters(View);
 		PassParameters->HairTransmittanceMaskSRV = GraphBuilder.CreateSRV(InTransmittanceMaskData.TransmittanceMask);
 
