Mark GBuffer as memoryless on mobile

diff --git a/Engine/Source/Runtime/Renderer/Private/MobileShadingRenderer.cpp b/Engine/Source/Runtime/Renderer/Private/MobileShadingRenderer.cpp
--- a/Engine/Source/Runtime/Renderer/Private/MobileShadingRenderer.cpp
+++ b/Engine/Source/Runtime/Renderer/Private/MobileShadingRenderer.cpp
@@ -190,6 +190,22 @@ static bool PostProcessUsesSceneDepth(const FViewInfo& View)
 	return false;
 }
 
+static void SetupGBufferFlags(FSceneTexturesConfig& SceneTexturesConfig, bool bRequiresMultiPass)
+{
+	ETextureCreateFlags AddFlags = TexCreate_InputAttachmentRead;
+	if (!bRequiresMultiPass)
+	{
+		// use memoryless GBuffer when possible
+		AddFlags|= TexCreate_Memoryless;
+	}
+	
+	SceneTexturesConfig.GBufferA.Flags|= AddFlags;
+	SceneTexturesConfig.GBufferB.Flags|= AddFlags;
+	SceneTexturesConfig.GBufferC.Flags|= AddFlags;
+	SceneTexturesConfig.GBufferD.Flags|= AddFlags;
+	SceneTexturesConfig.GBufferE.Flags|= AddFlags;
+}
+
 FMobileSceneRenderer::FMobileSceneRenderer(const FSceneViewFamily* InViewFamily,FHitProxyConsumer* HitProxyConsumer)
 	: FSceneRenderer(InViewFamily, HitProxyConsumer)
 	, bGammaSpace(!IsMobileHDR())
@@ -425,6 +441,11 @@ void FMobileSceneRenderer::InitViews(FRDGBuilder& GraphBuilder, FSceneTexturesCo
 
 	// Update the bKeepDepthContent based on the mobile renderer status.
 	SceneTexturesConfig.bKeepDepthContent = bKeepDepthContent;
+	
+	if (bDeferredShading) 
+	{
+		SetupGBufferFlags(SceneTexturesConfig, bRequiresMultiPass);
+	}
 
 	// Update the pixel projected reflection extent according to the settings in the PlanarReflectionComponent.
 	if (bRequiresPixelProjectedPlanarRelfectionPass)
