Fix hair composition when running with effect writing velocity/depth data.

diff --git a/Engine/Source/Runtime/Renderer/Private/DeferredShadingRenderer.cpp b/Engine/Source/Runtime/Renderer/Private/DeferredShadingRenderer.cpp
--- a/Engine/Source/Runtime/Renderer/Private/DeferredShadingRenderer.cpp
+++ b/Engine/Source/Runtime/Renderer/Private/DeferredShadingRenderer.cpp
@@ -2698,6 +2698,7 @@ void FDeferredShadingSceneRenderer::Render(FRDGBuilder& GraphBuilder)
 
 	if (!IsHairStrandsComposeAfterTranslucency())
 	{
+		RDG_GPU_STAT_SCOPE(GraphBuilder, HairRendering);
 		RenderHairComposition(GraphBuilder, Views, SceneTextures.Color.Target, SceneTextures.Depth.Target);
 	}
 
@@ -2726,6 +2727,14 @@ void FDeferredShadingSceneRenderer::Render(FRDGBuilder& GraphBuilder)
 		AddServiceLocalQueuePass(GraphBuilder);
 		TranslucencyViewsToRender = ETranslucencyView::None;
 
+		// Compose hair before velocity/distortion pass since these pass write depth value, 
+		// and this would make the hair composition fails in this cases.
+		if (IsHairStrandsComposeAfterTranslucency())
+		{
+			RDG_GPU_STAT_SCOPE(GraphBuilder, HairRendering);
+			RenderHairComposition(GraphBuilder, Views, SceneTextures.Color.Target, SceneTextures.Depth.Target);
+		}
+
 		if (bShouldRenderDistortion)
 		{
 			GraphBuilder.SetCommandListStat(GET_STATID(STAT_CLM_Distortion));
@@ -2750,6 +2759,11 @@ void FDeferredShadingSceneRenderer::Render(FRDGBuilder& GraphBuilder)
 
 		GraphBuilder.SetCommandListStat(GET_STATID(STAT_CLM_AfterTranslucency));
 	}
+	else if (IsHairStrandsComposeAfterTranslucency())
+	{
+		RDG_GPU_STAT_SCOPE(GraphBuilder, HairRendering);
+		RenderHairComposition(GraphBuilder, Views, SceneTextures.Color.Target, SceneTextures.Depth.Target);
+	}
 
 #if !UE_BUILD_SHIPPING
 	if (CVarForceBlackVelocityBuffer.GetValueOnRenderThread())
@@ -2762,12 +2776,6 @@ void FDeferredShadingSceneRenderer::Render(FRDGBuilder& GraphBuilder)
 #endif
 
 	{
-		RDG_GPU_STAT_SCOPE(GraphBuilder, HairRendering);
-		if (IsHairStrandsComposeAfterTranslucency())
-		{
-			RenderHairComposition(GraphBuilder, Views, SceneTextures.Color.Target, SceneTextures.Depth.Target);
-		}
-
 		if (HairStrandsBookmarkParameters.bHasElements)
 		{
 			RenderHairStrandsDebugInfo(GraphBuilder, Views, HairStrandsBookmarkParameters.HairClusterData, SceneTextures.Color.Target);
