Changes the default of r.TemporalAA.R11G11B10History

diff --git a/Engine/Source/Runtime/Renderer/Private/PostProcess/TemporalAA.cpp b/Engine/Source/Runtime/Renderer/Private/PostProcess/TemporalAA.cpp
--- a/Engine/Source/Runtime/Renderer/Private/PostProcess/TemporalAA.cpp
+++ b/Engine/Source/Runtime/Renderer/Private/PostProcess/TemporalAA.cpp
@@ -77,7 +77,7 @@ static TAutoConsoleVariable<int32> CVarUseTemporalAAUpscaler(
 	ECVF_RenderThreadSafe);
 
 TAutoConsoleVariable<int32> CVarTAAR11G11B10History(
-	TEXT("r.TemporalAA.R11G11B10History"), 0,
+	TEXT("r.TemporalAA.R11G11B10History"), 1,
 	TEXT("Select the bitdepth of the history."),
 	ECVF_RenderThreadSafe);
 
