Fixed crash with meshes that have zero SDF bricks

diff --git a/Engine/Source/Runtime/Renderer/Private/DistanceFieldStreaming.cpp b/Engine/Source/Runtime/Renderer/Private/DistanceFieldStreaming.cpp
--- a/Engine/Source/Runtime/Renderer/Private/DistanceFieldStreaming.cpp
+++ b/Engine/Source/Runtime/Renderer/Private/DistanceFieldStreaming.cpp
@@ -505,8 +505,12 @@ void FDistanceFieldSceneData::ProcessStreamingRequestsFromGPU(
 						check(AssetState.ReversedMips.Num() > 1);
 						const FDistanceFieldAssetMipState MipState = AssetState.ReversedMips.Pop();
 						IndirectionTableAllocator.Free(MipState.IndirectionTableOffset, MipState.IndirectionDimensions.X * MipState.IndirectionDimensions.Y * MipState.IndirectionDimensions.Z);
-						check(MipState.AllocatedBlocks.Num() > 0);
-						DistanceFieldAtlasBlockAllocator.Free(MipState.AllocatedBlocks);
+						
+						if (MipState.NumBricks > 0)
+						{
+							check(MipState.AllocatedBlocks.Num() > 0);
+							DistanceFieldAtlasBlockAllocator.Free(MipState.AllocatedBlocks);
+						}
 
 						// Re-upload mip0 to push the new NumMips to the shader
 						AssetDataUploads.Add(FDistanceFieldAssetMipId(AssetSetId, 0));
@@ -782,8 +786,12 @@ void FDistanceFieldSceneData::UpdateDistanceFieldAtlas(
 		for (const FDistanceFieldAssetMipState& MipState : AssetState.ReversedMips)
 		{
 			IndirectionTableAllocator.Free(MipState.IndirectionTableOffset, MipState.IndirectionDimensions.X * MipState.IndirectionDimensions.Y * MipState.IndirectionDimensions.Z);
-			check(MipState.AllocatedBlocks.Num() > 0);
-			DistanceFieldAtlasBlockAllocator.Free(MipState.AllocatedBlocks);
+
+			if (MipState.NumBricks > 0)
+			{
+				check(MipState.AllocatedBlocks.Num() > 0);
+				DistanceFieldAtlasBlockAllocator.Free(MipState.AllocatedBlocks);
+			}
 		}
 		
 		AssetStateArray.Remove(AssetSetId);
