Add showflag to disable VSM caching for thumbnail scenes to avoid using huge amounts of memory

diff --git a/Engine/Source/Runtime/Engine/Public/ShowFlags.h b/Engine/Source/Runtime/Engine/Public/ShowFlags.h
--- a/Engine/Source/Runtime/Engine/Public/ShowFlags.h
+++ b/Engine/Source/Runtime/Engine/Public/ShowFlags.h
@@ -221,6 +221,9 @@ struct FEngineShowFlags
 		SetVolumetricLightmap(false);
 		SetLumenGlobalIllumination(false);
 		SetLumenReflections(false);
+
+		// TODO: Remove when Physical page pool size scales automatically with demand
+		SetVirtualShadowMapCaching(false);
 	}
 
 	void EnableAdvancedFeatures()
@@ -247,6 +250,9 @@ struct FEngineShowFlags
 		SetDistanceFieldAO(true);
 		SetLumenGlobalIllumination(false);
 		SetLumenReflections(false);
+
+		// TODO: Remove when Physical page pool size scales automatically with demand
+		SetVirtualShadowMapCaching(false);
 	}
 
 	bool IsVisualizeCalibrationEnabled() const
diff --git a/Engine/Source/Runtime/Engine/Public/ShowFlagsValues.inl b/Engine/Source/Runtime/Engine/Public/ShowFlagsValues.inl
--- a/Engine/Source/Runtime/Engine/Public/ShowFlagsValues.inl
+++ b/Engine/Source/Runtime/Engine/Public/ShowFlagsValues.inl
@@ -378,5 +378,7 @@ SHOWFLAG_FIXED_IN_SHIPPING(0, VisualizeVolumetricCloudConservativeDensity, SFG_V
 /** Debug the Strata material buffer content */
 SHOWFLAG_FIXED_IN_SHIPPING(0, VisualizeStrataMaterial, SFG_Visualize, NSLOCTEXT("UnrealEd", "StrataMaterial", "Strata Material"))
 
+SHOWFLAG_FIXED_IN_SHIPPING(1, VirtualShadowMapCaching, SFG_Developer, NSLOCTEXT("UnrealEd", "VirtualShadowMapCaching", "Cache Virtual Shadow Maps"))
+
 #undef SHOWFLAG_ALWAYS_ACCESSIBLE
 #undef SHOWFLAG_FIXED_IN_SHIPPING
diff --git a/Engine/Source/Runtime/Renderer/Private/DeferredShadingRenderer.cpp b/Engine/Source/Runtime/Renderer/Private/DeferredShadingRenderer.cpp
--- a/Engine/Source/Runtime/Renderer/Private/DeferredShadingRenderer.cpp
+++ b/Engine/Source/Runtime/Renderer/Private/DeferredShadingRenderer.cpp
@@ -2446,7 +2446,7 @@ void FDeferredShadingSceneRenderer::Render(FRDGBuilder& GraphBuilder)
 
 		if (Scene->VirtualShadowMapArrayCacheManager)
 		{
-			Scene->VirtualShadowMapArrayCacheManager->ExtractFrameData(VirtualShadowMapArray, GraphBuilder);
+			Scene->VirtualShadowMapArrayCacheManager->ExtractFrameData(ViewFamily.EngineShowFlags.VirtualShadowMapCaching, VirtualShadowMapArray, GraphBuilder);
 		}
 	}
 
diff --git a/Engine/Source/Runtime/Renderer/Private/VirtualShadowMaps/VirtualShadowMapCacheManager.cpp b/Engine/Source/Runtime/Renderer/Private/VirtualShadowMaps/VirtualShadowMapCacheManager.cpp
--- a/Engine/Source/Runtime/Renderer/Private/VirtualShadowMaps/VirtualShadowMapCacheManager.cpp
+++ b/Engine/Source/Runtime/Renderer/Private/VirtualShadowMaps/VirtualShadowMapCacheManager.cpp
@@ -129,13 +129,14 @@ public:
 };
 IMPLEMENT_GLOBAL_SHADER(FVirtualSmCopyStatsCS, "/Engine/Private/VirtualShadowMaps/CopyStats.usf", "CopyStatsCS", SF_Compute);
 
-void FVirtualShadowMapArrayCacheManager::ExtractFrameData(FVirtualShadowMapArray &VirtualShadowMapArray, FRDGBuilder& GraphBuilder)
+void FVirtualShadowMapArrayCacheManager::ExtractFrameData(bool bEnableCaching, FVirtualShadowMapArray &VirtualShadowMapArray, FRDGBuilder& GraphBuilder)
 {
 	// Drop all refs.
 	PrevBuffers = FVirtualShadowMapArrayFrameData();
 	PrevUniformParameters.NumShadowMaps = 0;
 
-	if (VirtualShadowMapArray.IsAllocated()
+	if (bEnableCaching 
+		&& VirtualShadowMapArray.IsAllocated()
 		&& CVarCacheVirtualSMs.GetValueOnRenderThread() != 0)
 	{
 		GraphBuilder.QueueBufferExtraction(VirtualShadowMapArray.PageTableRDG, &PrevBuffers.PageTable);
diff --git a/Engine/Source/Runtime/Renderer/Private/VirtualShadowMaps/VirtualShadowMapCacheManager.h b/Engine/Source/Runtime/Renderer/Private/VirtualShadowMaps/VirtualShadowMapCacheManager.h
--- a/Engine/Source/Runtime/Renderer/Private/VirtualShadowMaps/VirtualShadowMapCacheManager.h
+++ b/Engine/Source/Runtime/Renderer/Private/VirtualShadowMaps/VirtualShadowMapCacheManager.h
@@ -108,8 +108,11 @@ public:
 
 	static constexpr float ClipSpaceScaleFactor = float(EffectiveCacheResolutionTexels) / float(FVirtualShadowMap::VirtualMaxResolutionXY);
 
-	// Call at end of frame to extract resouces from the virtual SM array to preserve to next frame
-	void ExtractFrameData(FVirtualShadowMapArray &VirtualShadowMapArray, FRDGBuilder& GraphBuilder);
+	/**
+	 * Call at end of frame to extract resouces from the virtual SM array to preserve to next frame.
+	 * If bCachingEnabled is false, all cache data is dropped and cache data will not be available for the next frame.
+	 */ 
+	void ExtractFrameData(bool bEnableCaching, FVirtualShadowMapArray &VirtualShadowMapArray, FRDGBuilder& GraphBuilder);
 
 	/**
 	 * Finds an existing cache entry and moves to the active set or creates a fresh one.
