Fix crash when running hair strands with non-lighting view mode & unshadowed forward lights.

diff --git a/Engine/Source/Runtime/Renderer/Private/DeferredShadingRenderer.cpp b/Engine/Source/Runtime/Renderer/Private/DeferredShadingRenderer.cpp
--- a/Engine/Source/Runtime/Renderer/Private/DeferredShadingRenderer.cpp
+++ b/Engine/Source/Runtime/Renderer/Private/DeferredShadingRenderer.cpp
@@ -2589,7 +2589,7 @@ void FDeferredShadingSceneRenderer::Render(FRDGBuilder& GraphBuilder)
 
 		AddServiceLocalQueuePass(GraphBuilder);
 	}
-	else if (HairStrands::HasViewHairStrandsData(Views))
+	else if (HairStrands::HasViewHairStrandsData(Views) && ViewFamily.EngineShowFlags.Lighting)
 	{
 		RenderLightsForHair(GraphBuilder, SceneTextures.UniformBuffer, SortedLightSet, ForwardScreenSpaceShadowMaskHairTexture, LightingChannelsTexture);
 		RenderDeferredReflectionsAndSkyLightingHair(GraphBuilder);
diff --git a/Engine/Source/Runtime/Renderer/Private/LightRendering.cpp b/Engine/Source/Runtime/Renderer/Private/LightRendering.cpp
--- a/Engine/Source/Runtime/Renderer/Private/LightRendering.cpp
+++ b/Engine/Source/Runtime/Renderer/Private/LightRendering.cpp
@@ -2536,6 +2536,7 @@ void FDeferredShadingSceneRenderer::RenderLightsForHair(
 				continue;
 			}
 
+			FHairStrandsTransmittanceMaskData DummyTransmittanceMaskData = CreateDummyHairStrandsTransmittanceMaskData(GraphBuilder, View.ShaderMap);
 			for (int32 LightIndex = AttenuationLightStart; LightIndex < SortedLights.Num(); LightIndex++)
 			{
 				const FSortedLightSceneInfo& SortedLightInfo = SortedLights[LightIndex];
@@ -2543,7 +2544,7 @@ void FDeferredShadingSceneRenderer::RenderLightsForHair(
 				if (LightSceneInfo.Proxy)
 				{
 					const bool bDrawHairShadow = SortedLightInfo.SortKey.Fields.bShadowed;
-					FHairStrandsTransmittanceMaskData TransmittanceMaskData;
+					FHairStrandsTransmittanceMaskData TransmittanceMaskData = DummyTransmittanceMaskData;
 					if (bDrawHairShadow)
 					{
 						TransmittanceMaskData = RenderHairStrandsTransmittanceMask(GraphBuilder, View, &LightSceneInfo, ScreenShadowMaskSubPixelTexture);
