Modify Mobile SDF, SSAO and OcclusionQueries to be after FullPrepass
SDF and SSAO require FullPrepass
Addedr.Mobile.EarlyZPass to force FullPrepass

diff --git a/Engine/Shaders/Private/MobileBasePassPixelShader.usf b/Engine/Shaders/Private/MobileBasePassPixelShader.usf
--- a/Engine/Shaders/Private/MobileBasePassPixelShader.usf
+++ b/Engine/Shaders/Private/MobileBasePassPixelShader.usf
@@ -619,26 +619,9 @@ void Main(
 	GBuffer.ShadingModelID = GetMaterialShadingModel(PixelMaterialInputs);
 	half MaterialAO = GetMaterialAmbientOcclusion(PixelMaterialInputs);
 
-#if (APPLY_AO || APPLY_DISTANCE_FIELD) && !DEFERRED_SHADING_PATH
-	// Compute velocity screen UV used to sample AO or DF buffer generated from previous frame
-	float4 ScreenPosition = SvPositionToResolvedScreenPosition(float4(SvPosition.xyz, 1));
-	float4 CurrClip = float4(ScreenPosition.xy, SvPosition.z, 1); // Inverted Far Depth = 0.0f
-	float4 PrevClip = mul(CurrClip, View.ClipToPrevClip);
-	float2 PrevScreen = PrevClip.xy / PrevClip.w;
-	float2 ScreenVelocity = ScreenPosition.xy - PrevScreen;
-	// TODO Sample screen velocity when available
-
-	float2 PrevScreenPosition = (ScreenPosition - ScreenVelocity);// NDC in [-1,1]
-	float2 PrevScreenUV = ScreenPositionToBufferUV(float4(PrevScreen.x, PrevScreen.y, 1, 1));
-#endif
 
 #if APPLY_AO
-	#if DEFERRED_SHADING_PATH
 	half4 GatheredAmbientOcclusion = Texture2DSample(AmbientOcclusionTexture, AmbientOcclusionSampler, SvPositionToBufferUV(SvPosition));
-	#else
-		half4 GatheredAmbientOcclusion = Texture2DSample(AmbientOcclusionTexture, AmbientOcclusionSampler, PrevScreenUV);
-	#endif
-
 
 	MaterialAO *= GatheredAmbientOcclusion.r;
 #endif
@@ -771,11 +754,7 @@ void Main(
 #if APPLY_DISTANCE_FIELD
 	if (ShadowPositionZ == 0)
 	{
-		#if DEFERRED_SHADING_PATH
-			Shadow = Texture2DSample(MobileBasePass.ScreenSpaceShadowMaskTexture, MobileBasePass.ScreenSpaceShadowMaskSampler, SvPositionToBufferUV(SvPosition)).x;
-		#else
-			Shadow = Texture2DSample(MobileBasePass.ScreenSpaceShadowMaskTexture, MobileBasePass.ScreenSpaceShadowMaskSampler, PrevScreenUV).x;
-		#endif
+		Shadow = Texture2DSample(MobileBasePass.ScreenSpaceShadowMaskTexture, MobileBasePass.ScreenSpaceShadowMaskSampler, SvPositionToBufferUV(SvPosition)).x;
 	}
 #endif
 
diff --git a/Engine/Source/Runtime/Renderer/Private/MobileShadingRenderer.cpp b/Engine/Source/Runtime/Renderer/Private/MobileShadingRenderer.cpp
--- a/Engine/Source/Runtime/Renderer/Private/MobileShadingRenderer.cpp
+++ b/Engine/Source/Runtime/Renderer/Private/MobileShadingRenderer.cpp
@@ -588,6 +588,7 @@ void FMobileSceneRenderer::RenderFullDepthPrepass(FRDGBuilder& GraphBuilder, FSc
 {
 	FRenderTargetBindingSlots BasePassRenderTargets;
 	BasePassRenderTargets.DepthStencil = FDepthStencilBinding(SceneTextures.Depth.Target, ERenderTargetLoadAction::EClear, FExclusiveDepthStencil::DepthWrite_StencilWrite);
+	BasePassRenderTargets.NumOcclusionQueries = ComputeNumOcclusionQueriesToBatch();
 
 	for (int32 ViewIndex = 0; ViewIndex < Views.Num(); ViewIndex++)
 	{
@@ -610,6 +611,10 @@ void FMobileSceneRenderer::RenderFullDepthPrepass(FRDGBuilder& GraphBuilder, FSc
 			[this, PassParameters, &View](FRHICommandListImmediate& RHICmdList)
 			{
 				RenderPrePass(RHICmdList, View, &PassParameters->InstanceCullingDrawParamsDepth);
+
+				// Issue occlusion queries
+				RHICmdList.SetCurrentStat(GET_STATID(STAT_CLMM_Occlusion));
+				RenderOcclusion(RHICmdList, View);
 			});
 	}
 }
@@ -762,6 +767,25 @@ void FMobileSceneRenderer::Render(FRDGBuilder& GraphBuilder)
 	if (bIsFullDepthPrepassEnabled)
 	{
 		RenderFullDepthPrepass(GraphBuilder, SceneTextures);
+
+		SceneTextures.MobileSetupMode = EMobileSceneTextureSetupMode::SceneDepth;
+		SceneTextures.MobileUniformBuffer = CreateMobileSceneTextureUniformBuffer(GraphBuilder, SceneTextures.MobileSetupMode);
+
+		if (bRequiresDistanceFieldShadowingPass)
+		{
+			CSV_SCOPED_TIMING_STAT_EXCLUSIVE(RenderSDFShadowing);
+			RenderMobileSDFShadowing(GraphBuilder, SceneTextures.Depth.Resolve, Scene, Views, VisibleLightInfos);
+		}
+
+		if (bShouldRenderHZB)
+		{
+			RenderHZB(GraphBuilder, SceneTextures.Depth.Resolve);
+		}
+
+		if (bRequiresAmbientOcclusionPass)
+		{
+			RenderAmbientOcclusion(GraphBuilder, SceneTextures.Depth.Resolve, SceneTextures.ScreenSpaceAO);
+		}
 	}
 
 	if (bDeferredShading)
@@ -812,22 +836,6 @@ void FMobileSceneRenderer::Render(FRDGBuilder& GraphBuilder)
 		);
 	}
 
-	if (bRequiresDistanceFieldShadowingPass)
-	{
-		CSV_SCOPED_TIMING_STAT_EXCLUSIVE(RenderSDFShadowing);
-		RenderMobileSDFShadowing(GraphBuilder, SceneTextures.Depth.Resolve, Scene, Views, VisibleLightInfos);
-	}
-
-	if (bShouldRenderHZB)
-	{
-		RenderHZB(GraphBuilder, SceneTextures.Depth.Resolve);
-	}
-
-	if (bRequiresAmbientOcclusionPass)
-	{
-		RenderAmbientOcclusion(GraphBuilder, SceneTextures.Depth.Resolve, SceneTextures.ScreenSpaceAO);
-	}
-
 	if (bRequiresPixelProjectedPlanarRelfectionPass)
 	{
 		const FPlanarReflectionSceneProxy* PlanarReflectionSceneProxy = Scene ? Scene->GetForwardPassGlobalPlanarReflection() : nullptr;
@@ -927,7 +935,10 @@ void FMobileSceneRenderer::RenderForward(FRDGBuilder& GraphBuilder, FRDGTextureR
 	BasePassRenderTargets.DepthStencil = FDepthStencilBinding(SceneDepth, bIsFullDepthPrepassEnabled ? ERenderTargetLoadAction::ELoad: ERenderTargetLoadAction::EClear, FExclusiveDepthStencil::DepthWrite_StencilWrite);
 	BasePassRenderTargets.ShadingRateTexture = (!MainView.bIsSceneCapture && !MainView.bIsReflectionCapture) ? SceneTextures.ShadingRate : nullptr;
 	BasePassRenderTargets.SubpassHint = ESubpassHint::DepthReadSubpass;
-	BasePassRenderTargets.NumOcclusionQueries = ComputeNumOcclusionQueriesToBatch();
+	if (!bIsFullDepthPrepassEnabled)
+	{
+		BasePassRenderTargets.NumOcclusionQueries = ComputeNumOcclusionQueriesToBatch();
+	}
 
 	//if the scenecolor isn't multiview but the app is, need to render as a single-view multiview due to shaders
 	BasePassRenderTargets.MultiViewCount = MainView.bIsMobileMultiViewEnabled ? 2 : (bIsMultiViewApplication ? 1 : 0);
@@ -1212,9 +1223,12 @@ void FMobileSceneRenderer::RenderDeferred(FRDGBuilder& GraphBuilder, const FSort
 	TArrayView<FRDGTextureRef> BasePassTexturesView = MakeArrayView(ColorTargets);
 
 	FRenderTargetBindingSlots BasePassRenderTargets = GetRenderTargetBindings(ERenderTargetLoadAction::ENoAction, BasePassTexturesView);
-	BasePassRenderTargets.DepthStencil = FDepthStencilBinding(SceneTextures.Depth.Target, ERenderTargetLoadAction::EClear, FExclusiveDepthStencil::DepthWrite_StencilWrite);
+	BasePassRenderTargets.DepthStencil = FDepthStencilBinding(SceneTextures.Depth.Target, bIsFullDepthPrepassEnabled ? ERenderTargetLoadAction::ELoad : ERenderTargetLoadAction::EClear, FExclusiveDepthStencil::DepthWrite_StencilWrite);
 	BasePassRenderTargets.SubpassHint = ESubpassHint::DeferredShadingSubpass;
-	BasePassRenderTargets.NumOcclusionQueries = ComputeNumOcclusionQueriesToBatch();
+	if (!bIsFullDepthPrepassEnabled)
+	{
+		BasePassRenderTargets.NumOcclusionQueries = ComputeNumOcclusionQueriesToBatch();
+	}
 	BasePassRenderTargets.ShadingRateTexture = nullptr;
 	BasePassRenderTargets.MultiViewCount = 0;
 
diff --git a/Engine/Source/Runtime/Renderer/Private/RendererScene.cpp b/Engine/Source/Runtime/Renderer/Private/RendererScene.cpp
--- a/Engine/Source/Runtime/Renderer/Private/RendererScene.cpp
+++ b/Engine/Source/Runtime/Renderer/Private/RendererScene.cpp
@@ -102,6 +102,16 @@ TAutoConsoleVariable<int32> CVarEarlyZPass(
 	TEXT("  x: use built in heuristic (default is 3)"),
 	ECVF_Scalability);
 
+static TAutoConsoleVariable<int32> CVarMobileEarlyZPass(
+	TEXT("r.Mobile.EarlyZPass"),
+	0,
+	TEXT("Whether to use a depth only pass to initialize Z culling for the mobile base pass.\n")
+	TEXT("  0: off\n")
+	TEXT("  1: all opaque \n"),
+	ECVF_Scalability
+);
+
+
 static TAutoConsoleVariable<int32> CVarBasePassWriteDepthEvenWithFullPrepass(
 	TEXT("r.BasePassWriteDepthEvenWithFullPrepass"),
 	0,
@@ -687,12 +697,15 @@ static void UpdateEarlyZPassModeCVarSinkFunction()
 {
 	static int32 CachedEarlyZPass = CVarEarlyZPass.GetValueOnGameThread();
 	static int32 CachedBasePassWriteDepthEvenWithFullPrepass = CVarBasePassWriteDepthEvenWithFullPrepass.GetValueOnGameThread();
+	static int32 CachedMobileEarlyZPass = CVarMobileEarlyZPass.GetValueOnGameThread();
 
 	const int32 EarlyZPass = CVarEarlyZPass.GetValueOnGameThread();
 	const int32 BasePassWriteDepthEvenWithFullPrepass = CVarBasePassWriteDepthEvenWithFullPrepass.GetValueOnGameThread();
+	const int32 MobileEarlyZPass = CVarMobileEarlyZPass.GetValueOnGameThread();
 
 	if (EarlyZPass != CachedEarlyZPass
-		|| BasePassWriteDepthEvenWithFullPrepass != CachedBasePassWriteDepthEvenWithFullPrepass)
+		|| BasePassWriteDepthEvenWithFullPrepass != CachedBasePassWriteDepthEvenWithFullPrepass
+		|| MobileEarlyZPass != CachedMobileEarlyZPass)
 	{
 		for (TObjectIterator<UWorld> It; It; ++It)
 		{
@@ -706,6 +719,7 @@ static void UpdateEarlyZPassModeCVarSinkFunction()
 
 		CachedEarlyZPass = EarlyZPass;
 		CachedBasePassWriteDepthEvenWithFullPrepass = BasePassWriteDepthEvenWithFullPrepass;
+		CachedMobileEarlyZPass = MobileEarlyZPass;
 	}
 }
 
@@ -3364,6 +3378,12 @@ void FScene::UpdateEarlyZPassMode()
 		{
 			EarlyZPassMode = DDM_AllOpaque;
 		}
+
+		bool bMobileForceFullDepthPrepass = CVarMobileEarlyZPass.GetValueOnAnyThread() == 1;
+		if (bMobileForceFullDepthPrepass)
+		{
+			EarlyZPassMode = DDM_AllOpaque;
+		}
 	}
 }
 
diff --git a/Engine/Source/Runtime/Renderer/Private/SceneTextures.cpp b/Engine/Source/Runtime/Renderer/Private/SceneTextures.cpp
--- a/Engine/Source/Runtime/Renderer/Private/SceneTextures.cpp
+++ b/Engine/Source/Runtime/Renderer/Private/SceneTextures.cpp
@@ -1109,9 +1109,12 @@ void SetupMobileSceneTextureUniformParameters(
 			}
 		}
 
-		if (HasBeenProduced(SceneTextures->Color.Resolve) && !EnumHasAnyFlags(SceneTextures->Depth.Resolve->Desc.Flags, TexCreate_Memoryless))
+		if (EnumHasAnyFlags(SetupMode, EMobileSceneTextureSetupMode::SceneDepth))
 		{
-			SceneTextureParameters.SceneDepthTexture = SceneTextures->Depth.Resolve;
+			if ((HasBeenProduced(SceneTextures->Depth.Resolve) && !EnumHasAnyFlags(SceneTextures->Depth.Resolve->Desc.Flags, TexCreate_Memoryless)))
+			{
+				SceneTextureParameters.SceneDepthTexture = SceneTextures->Depth.Resolve;
+			}
 		}
 
 		if (SceneTextures->Config.bIsUsingGBuffers)
