Removed unused DynamicMeshBatchStartOffset and VisibleDrawCommandStartOffset locals from FDeferredShadingSceneRenderer::GatherRayTracingWorldInstances()

diff --git a/Engine/Source/Runtime/Renderer/Private/DeferredShadingRenderer.cpp b/Engine/Source/Runtime/Renderer/Private/DeferredShadingRenderer.cpp
--- a/Engine/Source/Runtime/Renderer/Private/DeferredShadingRenderer.cpp
+++ b/Engine/Source/Runtime/Renderer/Private/DeferredShadingRenderer.cpp
@@ -484,8 +484,6 @@ bool FDeferredShadingSceneRenderer::GatherRayTracingWorldInstances(FRHICommandLi
 	SCOPE_CYCLE_COUNTER(STAT_GatherRayTracingWorldInstances);
 
 	RayTracingCollector.ClearViewMeshArrays();
-	TArray<int> DynamicMeshBatchStartOffset;
-	TArray<int> VisibleDrawCommandStartOffset;
 
 	FGPUScenePrimitiveCollector DummyDynamicPrimitiveCollector;
 
@@ -494,8 +492,6 @@ bool FDeferredShadingSceneRenderer::GatherRayTracingWorldInstances(FRHICommandLi
 	for (int32 ViewIndex = 0; ViewIndex < Views.Num(); ViewIndex++)
 	{
 		FViewInfo& View = Views[ViewIndex];
-		DynamicMeshBatchStartOffset.Add(0);
-		VisibleDrawCommandStartOffset.Add(0);
 		View.RayTracingGeometryInstances.Reserve(Scene->Primitives.Num());
 
 		RayTracingCollector.AddViewMeshArrays(
@@ -1134,7 +1130,6 @@ bool FDeferredShadingSceneRenderer::GatherRayTracingWorldInstances(FRHICommandLi
 						NewVisibleMeshCommand.RayTracingMeshCommand = &Scene->CachedRayTracingMeshCommands.RayTracingMeshCommands[CommandIndex];
 						NewVisibleMeshCommand.InstanceIndex = NewInstanceIndex;
 						View.VisibleRayTracingMeshCommands.Add(NewVisibleMeshCommand);
-						VisibleDrawCommandStartOffset[ViewIndex]++;
 					}
 					else
 					{
