Fixed missing cloud shadow and crash with tweaking the directional light.

diff --git a/Engine/Source/Runtime/Renderer/Private/LightRendering.cpp b/Engine/Source/Runtime/Renderer/Private/LightRendering.cpp
--- a/Engine/Source/Runtime/Renderer/Private/LightRendering.cpp
+++ b/Engine/Source/Runtime/Renderer/Private/LightRendering.cpp
@@ -2191,20 +2191,20 @@ void FDeferredShadingSceneRenderer::RenderLight(
 			FLightSceneProxy* AtmosphereLight0Proxy = Scene->AtmosphereLights[0] ? Scene->AtmosphereLights[0]->Proxy : nullptr;
 			FLightSceneProxy* AtmosphereLight1Proxy = Scene->AtmosphereLights[1] ? Scene->AtmosphereLights[1]->Proxy : nullptr;
 			FVolumetricCloudRenderSceneInfo* CloudInfo = Scene->GetVolumetricCloudSceneInfo();
-			const bool VolumetricCloudShadowMap0Valid = View.VolumetricCloudShadowRenderTarget[0] != nullptr;
-			const bool VolumetricCloudShadowMap1Valid = View.VolumetricCloudShadowRenderTarget[1] != nullptr;
+			const bool VolumetricCloudShadowMap0Valid = View.VolumetricCloudShadowExtractedRenderTarget[0] != nullptr;
+			const bool VolumetricCloudShadowMap1Valid = View.VolumetricCloudShadowExtractedRenderTarget[1] != nullptr;
 			const bool bLight0CloudPerPixelTransmittance = CloudInfo && VolumetricCloudShadowMap0Valid && AtmosphereLight0Proxy == LightSceneInfo->Proxy && AtmosphereLight0Proxy && AtmosphereLight0Proxy->GetCloudShadowOnSurfaceStrength() > 0.0f;
 			const bool bLight1CloudPerPixelTransmittance = CloudInfo && VolumetricCloudShadowMap1Valid && AtmosphereLight1Proxy == LightSceneInfo->Proxy && AtmosphereLight1Proxy && AtmosphereLight1Proxy->GetCloudShadowOnSurfaceStrength() > 0.0f;
 			if (bLight0CloudPerPixelTransmittance)
 			{
-				RenderLightParams.Cloud_ShadowmapTexture = TryGetRHI(View.VolumetricCloudShadowRenderTarget[0]);
+				RenderLightParams.Cloud_ShadowmapTexture = View.VolumetricCloudShadowExtractedRenderTarget[0]->GetShaderResourceRHI();
 				RenderLightParams.Cloud_ShadowmapFarDepthKm = CloudInfo->GetVolumetricCloudCommonShaderParameters().CloudShadowmapFarDepthKm[0].X;
 				RenderLightParams.Cloud_WorldToLightClipShadowMatrix = CloudInfo->GetVolumetricCloudCommonShaderParameters().CloudShadowmapWorldToLightClipMatrix[0];
 				RenderLightParams.Cloud_ShadowmapStrength = AtmosphereLight0Proxy->GetCloudShadowOnSurfaceStrength();
 			}
 			else if(bLight1CloudPerPixelTransmittance)
 			{
-				RenderLightParams.Cloud_ShadowmapTexture = TryGetRHI(View.VolumetricCloudShadowRenderTarget[1]);
+				RenderLightParams.Cloud_ShadowmapTexture = View.VolumetricCloudShadowExtractedRenderTarget[1]->GetShaderResourceRHI();
 				RenderLightParams.Cloud_ShadowmapFarDepthKm = CloudInfo->GetVolumetricCloudCommonShaderParameters().CloudShadowmapFarDepthKm[1].X;
 				RenderLightParams.Cloud_WorldToLightClipShadowMatrix = CloudInfo->GetVolumetricCloudCommonShaderParameters().CloudShadowmapWorldToLightClipMatrix[1];
 				RenderLightParams.Cloud_ShadowmapStrength = AtmosphereLight1Proxy->GetCloudShadowOnSurfaceStrength();
diff --git a/Engine/Source/Runtime/Renderer/Private/SceneRendering.h b/Engine/Source/Runtime/Renderer/Private/SceneRendering.h
--- a/Engine/Source/Runtime/Renderer/Private/SceneRendering.h
+++ b/Engine/Source/Runtime/Renderer/Private/SceneRendering.h
@@ -1350,6 +1350,8 @@ public:
 	TUniformBufferRef<FViewUniformShaderParameters> VolumetricRenderTargetViewUniformBuffer;
 	// The effective cloud shadow target this frame independently of the fact that a view can have a state (primary view) or not (sky light reflection capture)
 	FRDGTextureRef VolumetricCloudShadowRenderTarget[NUM_ATMOSPHERE_LIGHTS] = {};
+	// We need to extract that RDG resource because the RHI must be accessed to setup FDeferredLightPS 
+	TRefCountPtr<IPooledRenderTarget> VolumetricCloudShadowExtractedRenderTarget[NUM_ATMOSPHERE_LIGHTS] = {};
 
 	/** Used when there is no view state, buffers reallocate every frame. */
 	TUniquePtr<FForwardLightingViewResources> ForwardLightingResourcesStorage;
diff --git a/Engine/Source/Runtime/Renderer/Private/TranslucentLighting.cpp b/Engine/Source/Runtime/Renderer/Private/TranslucentLighting.cpp
--- a/Engine/Source/Runtime/Renderer/Private/TranslucentLighting.cpp
+++ b/Engine/Source/Runtime/Renderer/Private/TranslucentLighting.cpp
@@ -785,7 +785,7 @@ public:
 					VolumetricCloudShadowmapTexture,
 					VolumetricCloudShadowmapTextureSampler,
 					TStaticSamplerState<SF_Bilinear, AM_Clamp, AM_Clamp, AM_Clamp>::GetRHI(),
-					View.VolumetricCloudShadowRenderTarget[LightIndex]->GetRHI());
+					View.VolumetricCloudShadowExtractedRenderTarget[LightIndex]->GetShaderResourceRHI());
 
 				if (bLight0CloudPerPixelTransmittance)
 				{
diff --git a/Engine/Source/Runtime/Renderer/Private/VolumetricCloudRendering.cpp b/Engine/Source/Runtime/Renderer/Private/VolumetricCloudRendering.cpp
--- a/Engine/Source/Runtime/Renderer/Private/VolumetricCloudRendering.cpp
+++ b/Engine/Source/Runtime/Renderer/Private/VolumetricCloudRendering.cpp
@@ -1289,48 +1289,43 @@ IMPLEMENT_GLOBAL_SHADER(FCloudShadowTemporalProcessCS, "/Engine/Private/Volumetr
 
 
 
-void FSceneRenderer::InitVolumetricCloudsForViews(FRDGBuilder& GraphBuilder, bool bShouldRenderVolumetricCloud, FInstanceCullingManager& InstanceCullingManager)
+void CleanUpCloudDataFunction(TArray<FViewInfo>& Views)
 {
-	auto CleanUpCloudData = [&Views = Views](FRDGBuilder& GraphBuilder)
+	for (int32 ViewIndex = 0; ViewIndex < Views.Num(); ViewIndex++)
 	{
-		AddPass(GraphBuilder, [&Views](FRHICommandList&)
+		FViewInfo& ViewInfo = Views[ViewIndex];
+		ViewInfo.VolumetricCloudSkyAO = nullptr;
+		for (int LightIndex = 0; LightIndex < NUM_ATMOSPHERE_LIGHTS; ++LightIndex)
 		{
-			for (int32 ViewIndex = 0; ViewIndex < Views.Num(); ViewIndex++)
+			ViewInfo.VolumetricCloudShadowRenderTarget[LightIndex] = nullptr;
+			ViewInfo.VolumetricCloudShadowExtractedRenderTarget[LightIndex] = nullptr;
+			if (ViewInfo.ViewState != nullptr)
 			{
-				FViewInfo& ViewInfo = Views[ViewIndex];
-				for (int LightIndex = 0; LightIndex < NUM_ATMOSPHERE_LIGHTS; ++LightIndex)
-				{
-					ViewInfo.VolumetricCloudShadowRenderTarget[LightIndex] = nullptr;
-					if (ViewInfo.ViewState != nullptr)
-					{
-						ViewInfo.ViewState->VolumetricCloudShadowRenderTarget[LightIndex].Reset();
-					}
-				}
+				ViewInfo.ViewState->VolumetricCloudShadowRenderTarget[LightIndex].Reset();
 			}
+		}
+	}
+};
+
+void FSceneRenderer::InitVolumetricCloudsForViews(FRDGBuilder& GraphBuilder, bool bShouldRenderVolumetricCloud, FInstanceCullingManager& InstanceCullingManager)
+{
+
+	auto CleanUpCloudDataPass = [&Views = Views](FRDGBuilder& GraphBuilder)
+	{
+		AddPass(GraphBuilder, [&Views](FRHICommandList&)
+		{
+			CleanUpCloudDataFunction(Views);
 		});
 	};
 
+	// First make sure we always clear the texture on views to make sure no dangling texture pointers are ever used
+	CleanUpCloudDataFunction(Views);
+
 	if (!bShouldRenderVolumetricCloud)
 	{
-		// Clean everything up
-		CleanUpCloudData(GraphBuilder);
 		return;
 	}
 
-	// First make sure we always clear the texture on  views to make sure no dangling texture pointers are ever used
-	AddPass(GraphBuilder, [&Views = Views](FRHICommandList&)
-	{
-		for (int32 ViewIndex = 0; ViewIndex < Views.Num(); ViewIndex++)
-		{
-			FViewInfo& ViewInfo = Views[ViewIndex];
-			for (int LightIndex = 0; LightIndex < NUM_ATMOSPHERE_LIGHTS; ++LightIndex)
-			{
-				ViewInfo.VolumetricCloudSkyAO = nullptr;
-				ViewInfo.VolumetricCloudShadowRenderTarget[LightIndex] = nullptr;
-			}
-		}
-	});
-
 	if (Scene)
 	{
 		check(ShouldRenderVolumetricCloud(Scene, ViewFamily.EngineShowFlags)); // This should not be called if we should not render SkyAtmosphere
@@ -1739,10 +1734,12 @@ void FSceneRenderer::InitVolumetricCloudsForViews(FRDGBuilder& GraphBuilder, boo
 										FilterTracedCloudTexture(&SpatiallyFilteredShadowTexture, VolumetricCloudParams.VolumetricCloud.CloudShadowmapSizeInvSize[LightIndex], CloudShadowTextureTexelWorldSizeInvSize, false);
 									}
 									ViewInfo.VolumetricCloudShadowRenderTarget[LightIndex] = SpatiallyFilteredShadowTexture;
+									ConvertToUntrackedExternalTexture(GraphBuilder, ViewInfo.VolumetricCloudShadowRenderTarget[LightIndex], ViewInfo.VolumetricCloudShadowExtractedRenderTarget[LightIndex], ERHIAccess::SRVMask);
 								}
 								else
 								{
 									ViewInfo.VolumetricCloudShadowRenderTarget[LightIndex] = GraphBuilder.RegisterExternalTexture(CloudShadowTemporalRT.CurrentRenderTarget());
+									ConvertToUntrackedExternalTexture(GraphBuilder, ViewInfo.VolumetricCloudShadowRenderTarget[LightIndex], ViewInfo.VolumetricCloudShadowExtractedRenderTarget[LightIndex], ERHIAccess::SRVMask);
 								}
 
 							}
@@ -1758,10 +1755,12 @@ void FSceneRenderer::InitVolumetricCloudsForViews(FRDGBuilder& GraphBuilder, boo
 										FilterTracedCloudTexture(&SpatiallyFilteredShadowTexture, VolumetricCloudParams.VolumetricCloud.CloudShadowmapSizeInvSize[LightIndex], CloudShadowTextureTexelWorldSizeInvSize, false);
 									}
 									ViewInfo.VolumetricCloudShadowRenderTarget[LightIndex] = SpatiallyFilteredShadowTexture;
+									ConvertToUntrackedExternalTexture(GraphBuilder, ViewInfo.VolumetricCloudShadowRenderTarget[LightIndex], ViewInfo.VolumetricCloudShadowExtractedRenderTarget[LightIndex], ERHIAccess::SRVMask);
 								}
 								else
 								{
 									ViewInfo.VolumetricCloudShadowRenderTarget[LightIndex] = NewCloudShadowTexture;
+									ConvertToUntrackedExternalTexture(GraphBuilder, ViewInfo.VolumetricCloudShadowRenderTarget[LightIndex], ViewInfo.VolumetricCloudShadowExtractedRenderTarget[LightIndex], ERHIAccess::SRVMask);
 								}
 							}
 						}
@@ -1789,17 +1788,17 @@ void FSceneRenderer::InitVolumetricCloudsForViews(FRDGBuilder& GraphBuilder, boo
 			}
 			else
 			{
-				CleanUpCloudData(GraphBuilder);
+				CleanUpCloudDataPass(GraphBuilder);
 			}
 		}
 		else
 		{
-			CleanUpCloudData(GraphBuilder);
+			CleanUpCloudDataPass(GraphBuilder);
 		}
 	}
 	else
 	{
-		CleanUpCloudData(GraphBuilder);
+		CleanUpCloudDataPass(GraphBuilder);
 	}
 }
 
