VT allocator improvements
- Check partially allocated block for fit, don't need to start with a completely free block for allocations
- Try to use the lowest possible vAddress for new allocations, this reduces the texture area covered by allocations
- Page table is no longer required to be square power-2 texture, instead allow it to grow organically as allocated area increases
- Allocator starts with maximum-sized block and tracks allocated area, no longer needs special logic to 'grow'
- Rectangular allocated VTs now mip based on min-size rather than max-size, relaxes alignment requirements and allows better packing

diff --git a/Engine/Shaders/Private/PageTableUpdate.usf b/Engine/Shaders/Private/PageTableUpdate.usf
--- a/Engine/Shaders/Private/PageTableUpdate.usf
+++ b/Engine/Shaders/Private/PageTableUpdate.usf
@@ -2,7 +2,7 @@
 
 #include "Common.ush"
 
-uint PageTableSize;
+uint2 PageTableSize;
 uint FirstUpdate;
 uint NumUpdates;
 
diff --git a/Engine/Source/Runtime/Renderer/Private/VT/AdaptiveVirtualTexture.cpp b/Engine/Source/Runtime/Renderer/Private/VT/AdaptiveVirtualTexture.cpp
--- a/Engine/Source/Runtime/Renderer/Private/VT/AdaptiveVirtualTexture.cpp
+++ b/Engine/Source/Runtime/Renderer/Private/VT/AdaptiveVirtualTexture.cpp
@@ -453,7 +453,7 @@ void FAdaptiveVirtualTexture::Allocate(FVirtualTextureSystem* InSystem, uint32 I
 
 	// Check if we have space in the page table to allocate. If not then hopefully we can allocate next frame.
 	FVirtualTextureSpace* Space = InSystem->GetSpace(GetSpaceID());
-	if (Space->GetPageTableSize() >= Space->GetDescription().MaxSpaceSize && !Space->GetAllocator().TryAlloc(NewLevel))
+	if (!Space->GetAllocator().TryAlloc(NewLevel))
 	{
 		return;
 	}
@@ -560,7 +560,7 @@ bool FAdaptiveVirtualTexture::FreeLRU(FVirtualTextureSystem* InSystem, uint32 In
 	int32 NewLevel = CurrentLevel - 1;
 	while (NewLevel > 0)
 	{
-		if (Space->GetPageTableSize() < Space->GetDescription().MaxSpaceSize || Space->GetAllocator().TryAlloc(NewLevel))
+		if (Space->GetAllocator().TryAlloc(NewLevel))
 		{
 			break;
 		}
diff --git a/Engine/Source/Runtime/Renderer/Private/VT/AllocatedVirtualTexture.cpp b/Engine/Source/Runtime/Renderer/Private/VT/AllocatedVirtualTexture.cpp
--- a/Engine/Source/Runtime/Renderer/Private/VT/AllocatedVirtualTexture.cpp
+++ b/Engine/Source/Runtime/Renderer/Private/VT/AllocatedVirtualTexture.cpp
@@ -56,7 +56,10 @@ FAllocatedVirtualTexture::FAllocatedVirtualTexture(FVirtualTextureSystem* InSyst
 
 	// Max level of overall allocated VT is limited by size in tiles
 	// With multiple layers of different sizes, some layers may have mips smaller than a single tile
-	MaxLevel = FMath::Min(MaxLevel, FMath::CeilLogTwo(FMath::Max(GetWidthInTiles(), GetHeightInTiles())));
+	// We can either use the Min or Max of Width/Height to determine the number of mips
+	// - Using Max will allow more mips for rectangular VTs, which could potentially reduce aliasing in certain situations
+	// - Using Min will relax alignment requirements for the page table allocator, which will tend to reduce overall VRAM usage
+	MaxLevel = FMath::Min(MaxLevel, FMath::CeilLogTwo(FMath::Min(GetWidthInTiles(), GetHeightInTiles())));
 
 	MaxLevel = FMath::Min(MaxLevel, VIRTUALTEXTURE_LOG2_MAX_PAGETABLE_SIZE - 1u);
 
diff --git a/Engine/Source/Runtime/Renderer/Private/VT/VirtualTextureAllocator.cpp b/Engine/Source/Runtime/Renderer/Private/VT/VirtualTextureAllocator.cpp
--- a/Engine/Source/Runtime/Renderer/Private/VT/VirtualTextureAllocator.cpp
+++ b/Engine/Source/Runtime/Renderer/Private/VT/VirtualTextureAllocator.cpp
@@ -4,26 +4,36 @@
 #include "AllocatedVirtualTexture.h"
 #include "VirtualTexturing.h"
 
+#if WITH_EDITOR
+#include "IImageWrapper.h"
+#include "IImageWrapperModule.h"
+#include "HAL/FileManager.h"
+#include "Misc/FileHelper.h"
+#endif // WITH_EDITOR
+
 FVirtualTextureAllocator::FVirtualTextureAllocator(uint32 Dimensions)
 	: vDimensions(Dimensions)
-	, LogSize(0u)
+	, AllocatedWidth(0u)
+	, AllocatedHeight(0u)
 	, NumAllocations(0u)
 	, NumAllocatedPages(0u)
 {
 }
 
-void FVirtualTextureAllocator::Initialize(uint32 InSize)
+void FVirtualTextureAllocator::Initialize(uint32 MaxSize)
 {
+	const uint32 vLogSize = FMath::CeilLogTwo(MaxSize);
+	check(vLogSize <= VIRTUALTEXTURE_LOG2_MAX_PAGETABLE_SIZE);
 	check(NumAllocations == 0);
-	LogSize = FMath::CeilLogTwo(InSize);
 
 	AddressBlocks.Reset(1);
 	SortedAddresses.Reset(1);
 	SortedIndices.Reset(1);
-	FreeList.Reset(LogSize + 1);
+	FreeList.Reset(vLogSize + 1);
+	PartiallyFreeList.Reset(vLogSize + 1);
 
 	// Start with one empty block
-	FAddressBlock DefaultBlock(LogSize);
+	FAddressBlock DefaultBlock(vLogSize);
 	DefaultBlock.State = EBlockState::FreeList;
 	AddressBlocks.Add(DefaultBlock);
 
@@ -31,12 +41,15 @@ void FVirtualTextureAllocator::Initialize(uint32 InSize)
 	SortedIndices.Add(0u);
 
 	// Init free list
-	FreeList.AddUninitialized(LogSize + 1);
-	for (uint8 i = 0; i < LogSize; i++)
+	FreeList.AddUninitialized(vLogSize + 1);
+	PartiallyFreeList.AddUninitialized(vLogSize + 1);
+	for (uint8 i = 0; i < vLogSize; i++)
 	{
 		FreeList[i] = 0xffff;
+		PartiallyFreeList[i] = 0xffff;
 	}
-	FreeList[LogSize] = 0;
+	FreeList[vLogSize] = 0;
+	PartiallyFreeList[vLogSize] = 0xffff;
 
 	// Init global free list
 	GlobalFreeList = 0xffff;
@@ -44,78 +57,46 @@ void FVirtualTextureAllocator::Initialize(uint32 InSize)
 	RootIndex = 0;
 }
 
-void FVirtualTextureAllocator::Grow()
+void FVirtualTextureAllocator::LinkFreeList(uint16& InOutListHead, EBlockState State, uint16 Index)
 {
-	// If we are empty then fast path is to reinitialize
-	if (NumAllocations == 0)
-	{
-		Initialize(1 << (LogSize + 1));
-		return;
-	}
-
-	++LogSize;
-
-	// Add entry for for free list of next LogSize (currently empty)
-	FreeList.Add(0xffff);
-
-	uint16 OldRootIndex = RootIndex;
-
-	// Add new root block
-	FAddressBlock RootBlock(LogSize);
-	RootBlock.State = EBlockState::AllocatedChildren;
-	RootBlock.FirstChild = OldRootIndex;
-
-	RootIndex = AcquireBlock();
-	AddressBlocks[RootIndex] = RootBlock;
+	FAddressBlock& AddressBlock = AddressBlocks[Index];
+	check(AddressBlock.State == EBlockState::None);
+	check(AddressBlock.NextFree == 0xffff);
+	check(AddressBlock.PrevFree == 0xffff);
 
-	// Reparent old root block
-	int32 NextSibling = AcquireBlock();
-	AddressBlocks[OldRootIndex].Parent = RootIndex;
-	AddressBlocks[OldRootIndex].FirstSibling = OldRootIndex;
-	AddressBlocks[OldRootIndex].NextSibling = NextSibling;
+	// Only the PartiallyFreeList free list is allowed to have children
+	check(State == EBlockState::PartiallyFreeList || AddressBlock.FirstChild == 0xffff);
 
-	// Add new siblings for old root block
-	const int32 NumChildren = (1 << vDimensions);
-	for (int32 Sibling = 1; Sibling < NumChildren; Sibling++)
+	AddressBlock.State = State;
+	AddressBlock.NextFree = InOutListHead;
+	if (AddressBlock.NextFree != 0xffff)
 	{
-		const int32 BlockIndex = NextSibling;
-		NextSibling = (Sibling + 1 < NumChildren) ? AcquireBlock() : 0xffff;
-
-		FAddressBlock Block(AddressBlocks[OldRootIndex], Sibling, vDimensions);
-		Block.NextSibling = NextSibling;
-		AddressBlocks[BlockIndex] = Block;
+		AddressBlocks[AddressBlock.NextFree].PrevFree = Index;
 	}
+	InOutListHead = Index;
+}
 
-	// Add new siblings to lists
-	SortedAddresses.InsertUninitialized(0, NumChildren - 1);
-	SortedIndices.InsertUninitialized(0, NumChildren - 1);
-	check(SortedAddresses.Num() == SortedIndices.Num());
-
-	int32 Sibling = 1;
-	uint16 Index = AddressBlocks[OldRootIndex].NextSibling;
-	while (Index != 0xffff)
+void FVirtualTextureAllocator::UnlinkFreeList(uint16& InOutListHead, EBlockState State, uint16 Index)
+{
+	FAddressBlock& AddressBlock = AddressBlocks[Index];
+	check(AddressBlock.State == State);
+	const uint32 PrevFreeIndex = AddressBlock.PrevFree;
+	const uint32 NextFreeIndex = AddressBlock.NextFree;
+	if (PrevFreeIndex != 0xffff)
 	{
-		FAddressBlock& AddressBlock = AddressBlocks[Index];
-
-		// Place on free list
-		check(AddressBlock.State == EBlockState::None);
-		AddressBlock.State = EBlockState::FreeList;
-
-		AddressBlock.NextFree = FreeList[AddressBlock.vLogSize];
-		if (AddressBlock.NextFree != 0xffff)
-		{
-			AddressBlocks[AddressBlock.NextFree].PrevFree = Index;
-		}
-		FreeList[AddressBlock.vLogSize] = Index;
-
-		// Add to sorted list
-		// We are inserting sorted so need to be careful to add siblings in reverse order
-		SortedAddresses[NumChildren - Sibling - 1] = AddressBlock.vAddress;
-		SortedIndices[NumChildren - Sibling - 1] = Index;
-
-		Index = AddressBlock.NextSibling;
-		Sibling++;
+		AddressBlocks[PrevFreeIndex].NextFree = NextFreeIndex;
+		AddressBlock.PrevFree = 0xffff;
 	}
+	if (NextFreeIndex != 0xffff)
+	{
+		AddressBlocks[NextFreeIndex].PrevFree = PrevFreeIndex;
+		AddressBlock.NextFree = 0xffff;
+	}
+	if (InOutListHead == Index)
+	{
+		InOutListHead = NextFreeIndex;
+	}
+	AddressBlock.State = EBlockState::None;
 }
 
 int32 FVirtualTextureAllocator::AcquireBlock()
@@ -124,15 +105,12 @@ int32 FVirtualTextureAllocator::AcquireBlock()
 	if (Index == 0xffff)
 	{
 		Index = AddressBlocks.AddUninitialized();
+		ensure(Index <= 0x8000); // make sure we're not getting close
+		check(Index <= 0xffff); // make sure we fit in 16bit index
 	}
 	else
 	{
-		check(AddressBlocks[Index].State == EBlockState::GlobalFreeList);
-		GlobalFreeList = AddressBlocks[Index].NextFree;
-		if (GlobalFreeList != 0xffff)
-		{
-			AddressBlocks[GlobalFreeList].PrevFree = 0xffff;
-		}
+		UnlinkFreeList(GlobalFreeList, EBlockState::GlobalFreeList, Index);
 	}
 
 	// Debug fill memory to invalid value
@@ -202,10 +180,17 @@ bool FVirtualTextureAllocator::TryAlloc(uint32 InLogSize)
 void FVirtualTextureAllocator::SubdivideBlock(uint32 ParentIndex)
 {
 	const uint32 NumChildren = (1 << vDimensions);
-	
-	const uint32 vLogSize = AddressBlocks[ParentIndex].vLogSize - 1u;
-	const uint32 vAddress = AddressBlocks[ParentIndex].vAddress;
 
+	const uint32 vParentLogSize = AddressBlocks[ParentIndex].vLogSize;
+	check(vParentLogSize > 0u);
+	const uint32 vChildLogSize = vParentLogSize - 1u;
+
+	// Only free blocks can be subdivided, move to the partially free list
+	check(AddressBlocks[ParentIndex].FirstChild == 0xffff);
+	UnlinkFreeList(FreeList[vParentLogSize], EBlockState::FreeList, ParentIndex);
+	LinkFreeList(PartiallyFreeList[vParentLogSize], EBlockState::PartiallyFreeList, ParentIndex);
+
+	const uint32 vAddress = AddressBlocks[ParentIndex].vAddress;
 	const int32 SortedIndex = FindAddressBlock(vAddress);
 	check(vAddress == SortedAddresses[SortedIndex]);
 
@@ -219,90 +204,71 @@ void FVirtualTextureAllocator::SubdivideBlock(uint32 ParentIndex)
 	for (uint32 Sibling = 0; Sibling < NumChildren; Sibling++)
 	{
 		const int32 ChildBlockIndex = AcquireBlock();
-		const uint32 vChildAddress = vAddress + (Sibling << (vDimensions * vLogSize));
+		const uint32 vChildAddress = vAddress + (Sibling << (vDimensions * vChildLogSize));
+
+		const int32 SortedIndexOffset = NumChildren - 1 - Sibling;
+		SortedAddresses[SortedIndex + SortedIndexOffset] = vChildAddress;
+		SortedIndices[SortedIndex + SortedIndexOffset] = ChildBlockIndex;
 
 		if (Sibling == 0u)
 		{
 			FirstSiblingIndex = ChildBlockIndex;
-			AddressBlocks[ParentIndex].FirstChild = FirstSiblingIndex;
-			// Replace parent block with new allocated block
-			SortedIndices[SortedIndex] = FirstSiblingIndex;
+			AddressBlocks[ParentIndex].FirstChild = ChildBlockIndex;
 		}
 		else
 		{
-			const int32 SortedIndexOffset = NumChildren - 1 - Sibling;
-			SortedAddresses[SortedIndex + SortedIndexOffset] = vChildAddress;
-			SortedIndices[SortedIndex + SortedIndexOffset] = ChildBlockIndex;
 			AddressBlocks[PrevChildIndex].NextSibling = ChildBlockIndex;
 		}
 	
-		FAddressBlock ChildBlock(vLogSize);
+		FAddressBlock ChildBlock(vChildLogSize);
 		ChildBlock.vAddress = vChildAddress;
 		ChildBlock.Parent = ParentIndex;
 		ChildBlock.FirstSibling = FirstSiblingIndex;
 		ChildBlock.NextSibling = 0xffff;
+		AddressBlocks[ChildBlockIndex] = ChildBlock;
 
 		// New child blocks start out on the free list
-		ChildBlock.State = EBlockState::FreeList;
-		ChildBlock.NextFree = FreeList[vLogSize];
-		if (ChildBlock.NextFree != 0xffff)
-		{
-			AddressBlocks[ChildBlock.NextFree].PrevFree = ChildBlockIndex;
-		}
-		FreeList[vLogSize] = ChildBlockIndex;
+		LinkFreeList(FreeList[vChildLogSize], EBlockState::FreeList, ChildBlockIndex);
 
-		AddressBlocks[ChildBlockIndex] = ChildBlock;
 		PrevChildIndex = ChildBlockIndex;
 	}
 }
 
-void FVirtualTextureAllocator::MarkBlockAllocated(uint32 Index, uint32 vAllocatedTileX, uint32 vAllocatedTileY, FAllocatedVirtualTexture* VT)
+void FVirtualTextureAllocator::MarkBlockAllocated(uint32 Index, uint32 vAllocatedTileX0, uint32 vAllocatedTileY0, FAllocatedVirtualTexture* VT)
 {
 	FAddressBlock* AllocBlock = &AddressBlocks[Index];
-	check(AllocBlock->State == EBlockState::FreeList);
-	check(AllocBlock->FirstChild == 0xffff);
+	check(AllocBlock->State != EBlockState::None);
+	check(AllocBlock->State != EBlockState::GlobalFreeList);
 
 	const uint32 vLogSize = AllocBlock->vLogSize;
 
 	// check to see if block is in the correct position
-	const uint32 vAllocatedTileX1 = vAllocatedTileX + VT->GetWidthInTiles();
-	const uint32 vAllocatedTileY1 = vAllocatedTileY + VT->GetHeightInTiles();
-	const uint32 vBlockAddress = AllocBlock->vAddress;
-	const uint32 vBlockTileX = FMath::ReverseMortonCode2(vBlockAddress);
-	const uint32 vBlockTileY = FMath::ReverseMortonCode2(vBlockAddress >> 1);
+	const uint32 vAllocatedTileX1 = vAllocatedTileX0 + VT->GetWidthInTiles();
+	const uint32 vAllocatedTileY1 = vAllocatedTileY0 + VT->GetHeightInTiles();
 	const uint32 BlockSize = (1u << vLogSize);
-
-	check(vBlockTileX >= vAllocatedTileX);
-	check(vBlockTileY >= vAllocatedTileY);
-
-	if (vBlockTileX < vAllocatedTileX1 && vBlockTileY < vAllocatedTileY1)
+	const uint32 vBlockAddress = AllocBlock->vAddress;
+	const uint32 vBlockTileX0 = FMath::ReverseMortonCode2(vBlockAddress);
+	const uint32 vBlockTileY0 = FMath::ReverseMortonCode2(vBlockAddress >> 1);
+	const uint32 vBlockTileX1 = vBlockTileX0 + BlockSize;
+	const uint32 vBlockTileY1 = vBlockTileY0 + BlockSize;
+
+	if (vAllocatedTileX1 > vBlockTileX0 &&
+		vAllocatedTileX0 < vBlockTileX1 &&
+		vAllocatedTileY1 > vBlockTileY0 &&
+		vAllocatedTileY0 < vBlockTileY1)
 	{
 		// Block overlaps the VT we're trying to allocate
-
-		// Remove from free list
-		{
-			const uint32 PrevFreeIndex = AllocBlock->PrevFree;
-			const uint32 NextFreeIndex = AllocBlock->NextFree;
-			if (PrevFreeIndex != 0xffff)
-			{
-				AddressBlocks[PrevFreeIndex].NextFree = NextFreeIndex;
-				AllocBlock->PrevFree = 0xffff;
-			}
-			if (NextFreeIndex != 0xffff)
-			{
-				AddressBlocks[NextFreeIndex].PrevFree = PrevFreeIndex;
-				AllocBlock->NextFree = 0xffff;
-			}
-			if (FreeList[vLogSize] == Index)
-			{
-				FreeList[vLogSize] = NextFreeIndex;
-			}
-			AllocBlock->State = EBlockState::None;
-		}
-
-		if (vBlockTileX + BlockSize <= vAllocatedTileX1 || vBlockTileY + BlockSize <= vAllocatedTileY1)
+		if (vBlockTileX0 >= vAllocatedTileX0 &&
+			vBlockTileX1 <= vAllocatedTileX1 &&
+			vBlockTileY0 >= vAllocatedTileY0 &&
+			vBlockTileY1 <= vAllocatedTileY1)
 		{
 			// Block is entirely contained within the VT we're trying to allocate
+
+			// In this case, block must be completely free (or else there's an error somewhere else)
+			check(AllocBlock->FirstChild == 0xffff);
+			UnlinkFreeList(FreeList[vLogSize], EBlockState::FreeList, Index);
+			
 			++NumAllocations;
 			NumAllocatedPages += 1u << (vDimensions * vLogSize);
 
@@ -315,11 +281,15 @@ void FVirtualTextureAllocator::MarkBlockAllocated(uint32 Index, uint32 vAllocate
 		}
 		else
 		{
-			// Block intersects the VT, subdivide further
-			AllocBlock->State = EBlockState::AllocatedChildren;
+			// Block intersects the VT
+			if (AllocBlock->State == EBlockState::FreeList)
+			{
+				// If block is completely free, need to subdivide further
+				SubdivideBlock(Index);
+			}
+			// otherwise block is already subdivided
 			AllocBlock = nullptr; // list will be potentially reallocated
-
-			SubdivideBlock(Index);
+			check(AddressBlocks[Index].State == EBlockState::PartiallyFreeList);
 	
 			uint32 NumChildren = 0u;
 			uint16 ChildIndex = AddressBlocks[Index].FirstChild;
@@ -328,7 +298,7 @@ void FVirtualTextureAllocator::MarkBlockAllocated(uint32 Index, uint32 vAllocate
 			{
 				check(AddressBlocks[ChildIndex].Parent == Index);
 
-				MarkBlockAllocated(ChildIndex, vAllocatedTileX, vAllocatedTileY, VT);
+				MarkBlockAllocated(ChildIndex, vAllocatedTileX0, vAllocatedTileY0, VT);
 
 				ChildIndex = AddressBlocks[ChildIndex].NextSibling;
 				NumChildren++;
@@ -338,38 +308,167 @@ void FVirtualTextureAllocator::MarkBlockAllocated(uint32 Index, uint32 vAllocate
 	}
 }
 
+bool FVirtualTextureAllocator::TestAllocation(uint32 Index, uint32 vAllocatedTileX0, uint32 vAllocatedTileY0, uint32 vAllocatedTileX1, uint32 vAllocatedTileY1) const
+{
+	const FAddressBlock& AllocBlock = AddressBlocks[Index];
+	const uint32 vLogSize = AllocBlock.vLogSize;
+	const uint32 BlockSize = (1u << vLogSize);
+
+	const uint32 vBlockAddress = AllocBlock.vAddress;
+	const uint32 vBlockTileX0 = FMath::ReverseMortonCode2(vBlockAddress);
+	const uint32 vBlockTileY0 = FMath::ReverseMortonCode2(vBlockAddress >> 1);
+	const uint32 vBlockTileX1 = vBlockTileX0 + BlockSize;
+	const uint32 vBlockTileY1 = vBlockTileY0 + BlockSize;
+
+	if (vAllocatedTileX1 > vBlockTileX0 &&
+		vAllocatedTileX0 < vBlockTileX1 &&
+		vAllocatedTileY1 > vBlockTileY0 &&
+		vAllocatedTileY0 < vBlockTileY1)
+	{
+		// Block overlaps the VT we're trying to allocate
+		if (AllocBlock.State == EBlockState::AllocatedTexture)
+		{
+			return false;
+		}
+		else
+		{
+			check(AllocBlock.State == EBlockState::PartiallyFreeList);
+			if (vBlockTileX0 >= vAllocatedTileX0 &&
+				vBlockTileX1 <= vAllocatedTileX1 &&
+				vBlockTileY0 >= vAllocatedTileY0 &&
+				vBlockTileY1 <= vAllocatedTileY1)
+			{
+				// If block is fully contained within the check region, don't need to search children, we are guaranteed to find an intersection
+				return false;
+			}
+
+			uint16 ChildIndex = AddressBlocks[Index].FirstChild;
+			check(ChildIndex == AddressBlocks[ChildIndex].FirstSibling);
+			while (ChildIndex != 0xffff)
+			{
+				const FAddressBlock& ChildBlock = AddressBlocks[ChildIndex];
+				check(ChildBlock.Parent == Index);
+				if (ChildBlock.State != EBlockState::FreeList &&
+					!TestAllocation(ChildIndex, vAllocatedTileX0, vAllocatedTileY0, vAllocatedTileX1, vAllocatedTileY1))
+				{
+					return false;
+				}
+				ChildIndex = ChildBlock.NextSibling;
+			}
+		}
+	}
+
+	return true;
+}
+
 uint32 FVirtualTextureAllocator::Alloc(FAllocatedVirtualTexture* VT)
 {
-	// Pad out to square power of 2
-	const uint32 MinSize = FMath::Min(VT->GetWidthInTiles(), VT->GetHeightInTiles());
-	const uint32 MaxSize = FMath::Max(VT->GetWidthInTiles(), VT->GetHeightInTiles());
-	const uint8 vLogMaxSize = FMath::CeilLogTwo(MaxSize);
+	const uint32 WidthInTiles = VT->GetWidthInTiles();
+	const uint32 HeightInTiles = VT->GetHeightInTiles();
+	const uint32 MinSize = FMath::Min(WidthInTiles, HeightInTiles);
+	const uint32 MaxSize = FMath::Max(WidthInTiles, HeightInTiles);
+	const int32 vLogMinSize = FMath::CeilLogTwo(MinSize);
+	const int32 vLogMaxSize = FMath::CeilLogTwo(MaxSize);
+
+	// Tile must be aligned to match the max level of the VT, otherwise tiles at lower mip levels may intersect neighboring regions
+	const uint32 MaxLevel = VT->GetMaxLevel();
+	const uint32 vAddressAlignment = 1u << (vDimensions * MaxLevel);
+
+	if (vLogMaxSize >= FreeList.Num())
+	{
+		// VT is larger than the entire page table
+		return ~0u;
+	}
+
+	uint16 AllocIndex = 0xffff;
+	uint32 vAddress = ~0u;
 
-	// Find smallest free that fits
-	for (int i = vLogMaxSize; i < FreeList.Num(); i++)
+	// See if we have any completely free blocks big enough
+	// Here we search all free blocks, including ones that are too large (large blocks will still be subdivided to fit)
+	for (int32 vLogSize = vLogMaxSize; vLogSize < FreeList.Num(); ++vLogSize)
 	{
-		uint16 FreeIndex = FreeList[i];
-		if (FreeIndex != 0xffff)
+		// Could avoid this loop if FreeList was kept sorted by vAddress
+		uint16 FreeIndex = FreeList[vLogSize];
+		while (FreeIndex != 0xffff)
 		{
-			// Found free
-			uint16 AllocIndex = FreeIndex;
-			FAddressBlock& AllocBlock = AddressBlocks[AllocIndex];
-			check(AllocBlock.VT == nullptr);
+			const FAddressBlock& AllocBlock = AddressBlocks[FreeIndex];
 			check(AllocBlock.State == EBlockState::FreeList);
+			if (AllocBlock.vAddress < vAddress)
+			{
+				AllocIndex = FreeIndex;
+				vAddress = AllocBlock.vAddress;
+			}
+			FreeIndex = AllocBlock.NextFree;
+		}
+	}
 
-			const uint32 vAddress = AllocBlock.vAddress;
-			const uint32 vTileX = FMath::ReverseMortonCode2(vAddress);
-			const uint32 vTileY = FMath::ReverseMortonCode2(vAddress >> 1);
+	// Look for a partially allocated block that has room for this allocation
+	// Only need to check partially allocated blocks of the correct size
+	// Larger partially allocated blocks will contain a child block that's completely free, that will already be discovered by the above search
+	{
+		uint16 FreeIndex = PartiallyFreeList[vLogMaxSize];
+		while (FreeIndex != 0xffff)
+		{
+			FAddressBlock& AllocBlock = AddressBlocks[FreeIndex];
+			if (AllocBlock.vAddress < vAddress)
+			{
+				check(AllocBlock.State == EBlockState::PartiallyFreeList);
+				const uint32 BlockSize = 1u << vLogMaxSize;
+				uint32 vCheckAddress = AllocBlock.vAddress;
+				const uint32 vBlockTileX0 = FMath::ReverseMortonCode2(vCheckAddress);
+				const uint32 vBlockTileY0 = FMath::ReverseMortonCode2(vCheckAddress >> 1);
+				const uint32 vBlockTileX1 = vBlockTileX0 + BlockSize;
+				const uint32 vBlockTileY1 = vBlockTileY0 + BlockSize;
+
+				// Search all valid positions within the block (in ascending morton order), looking for a fit for the texture we're trying to allocate
+				// Step size is driven by our alignment requirements
+				// Seems like there is probably a more clever way to accomplish this, if perf becomes an issue
+				while (true)
+				{
+					const uint32 vTileX0 = FMath::ReverseMortonCode2(vCheckAddress);
+					const uint32 vTileY0 = FMath::ReverseMortonCode2(vCheckAddress >> 1);
+					const uint32 vTileX1 = vTileX0 + WidthInTiles;
+					const uint32 vTileY1 = vTileY0 + HeightInTiles;
+					if (vTileY1 > vBlockTileY1)
+					{
+						break;
+					}
+
+					if (vTileX1 <= vBlockTileX1)
+					{
+						if (TestAllocation(FreeIndex, vTileX0, vTileY0, vTileX1, vTileY1))
+						{
+							// here AllocIndex won't point to exactly the correct block yet, but we don't want to subdivide yet, until we're sure this is the best fit
+							// MarkBlockAllocated will properly subdivide the initial block as needed
+							AllocIndex = FreeIndex;
+							vAddress = vCheckAddress;
+							break;
+						}
+					}
+
+					vCheckAddress += vAddressAlignment;
+				}
+			}
+			FreeIndex = AllocBlock.NextFree;
+		}
+	}
 
-			MarkBlockAllocated(AllocIndex, vTileX, vTileY, VT);
+	if (AllocIndex != 0xffff)
+	{
+		check(vAddress != ~0u);
+		FAddressBlock& AllocBlock = AddressBlocks[AllocIndex];
+		const uint32 vTileX = FMath::ReverseMortonCode2(vAddress);
+		const uint32 vTileY = FMath::ReverseMortonCode2(vAddress >> 1);
 
-			check(AddressBlocks[AllocIndex].State != EBlockState::FreeList);
+		MarkBlockAllocated(AllocIndex, vTileX, vTileY, VT);
 
-			return vAddress;
-		}
+		check(AddressBlocks[AllocIndex].State != EBlockState::FreeList);
+
+		AllocatedWidth = FMath::Max(AllocatedWidth, vTileX + WidthInTiles);
+		AllocatedHeight = FMath::Max(AllocatedHeight, vTileY + HeightInTiles);
 	}
 
-	return ~0u;
+	return vAddress;
 }
 
 void FVirtualTextureAllocator::Free(FAllocatedVirtualTexture* VT)
@@ -385,7 +484,7 @@ void FVirtualTextureAllocator::Free(FAllocatedVirtualTexture* VT)
 		{
 			check(AddressBlock.State == EBlockState::AllocatedTexture);
 			check(AddressBlock.FirstChild == 0xffff); // texture allocation should be leaf
-			AddressBlock.State = EBlockState::AllocatedChildren;
+			AddressBlock.State = EBlockState::None;
 			AddressBlock.VT = nullptr;
 
 			check(NumAllocations > 0u);
@@ -397,7 +496,7 @@ void FVirtualTextureAllocator::Free(FAllocatedVirtualTexture* VT)
 
 			// Add block to free list
 			// This handles merging free siblings
-			FreeAddressBlock(Index);
+			FreeAddressBlock(Index, true);
 
 			// Remove the index from the hash table as it may be reused later
 			HashTable.Remove(Key, Index);
@@ -406,36 +505,46 @@ void FVirtualTextureAllocator::Free(FAllocatedVirtualTexture* VT)
 	}
 }
 
-void FVirtualTextureAllocator::FreeAddressBlock(uint32 Index)
+void FVirtualTextureAllocator::FreeAddressBlock(uint32 Index, bool bTopLevelBlock)
 {
 	FAddressBlock& AddressBlock = AddressBlocks[Index];
+	if (bTopLevelBlock)
+	{
+		// Block was freed directly, should already be removed froms lists
+		check(AddressBlock.State == EBlockState::None);
+	}
+	else
+	{
+		// Block was freed by consolidating children
+		UnlinkFreeList(PartiallyFreeList[AddressBlock.vLogSize], EBlockState::PartiallyFreeList, Index);
+	}
+
 	check(AddressBlock.VT == nullptr);
-	check(AddressBlock.State == EBlockState::AllocatedChildren);
 	check(AddressBlock.NextFree == 0xffff);
 	check(AddressBlock.PrevFree == 0xffff);
 
+	// If we got here, the block's children have already been consolidated/removed
+	AddressBlock.FirstChild = 0xffff;
+
 	// If all siblings are free then we can merge them
 	uint32 SiblingIndex = AddressBlock.FirstSibling;
 	bool bConsolidateSiblings = SiblingIndex != 0xffff;
 	while (bConsolidateSiblings && SiblingIndex != 0xffff)
 	{
-		bConsolidateSiblings &= (SiblingIndex == Index || FreeList[AddressBlock.vLogSize] == SiblingIndex || AddressBlocks[SiblingIndex].PrevFree != 0xffff);
-		SiblingIndex = AddressBlocks[SiblingIndex].NextSibling;
+		const FAddressBlock& SiblingBlock = AddressBlocks[SiblingIndex];
+		if (SiblingIndex != Index)
+		{
+			check(SiblingBlock.State != EBlockState::None);
+			check(SiblingBlock.State != EBlockState::GlobalFreeList);
+			bConsolidateSiblings &= (SiblingBlock.State == EBlockState::FreeList);
+		}
+		SiblingIndex = SiblingBlock.NextSibling;
 	}
 
 	if (!bConsolidateSiblings)
 	{
-		// If we got here, the block's children have already been consolidated/removed
-		AddressBlock.FirstChild = 0xffff;
-
 		// Simply place this block on the free list
-		AddressBlock.State = EBlockState::FreeList;
-		AddressBlock.NextFree = FreeList[AddressBlock.vLogSize];
-		if (AddressBlock.NextFree != 0xffff)
-		{
-			AddressBlocks[AddressBlock.NextFree].PrevFree = Index;
-		}
-		FreeList[AddressBlock.vLogSize] = Index;
+		LinkFreeList(FreeList[AddressBlock.vLogSize], EBlockState::FreeList, Index);
 	}
 	else
 	{
@@ -448,37 +557,16 @@ void FVirtualTextureAllocator::FreeAddressBlock(uint32 Index)
 			if (FreeIndex != Index)
 			{
 				// All our siblings must be free (we checked above to get into this case)
-				check(FreeBlock.State == EBlockState::FreeList);
-
-				const uint32 PrevFreeIndex = FreeBlock.PrevFree;
-				const uint32 NextFreeIndex = FreeBlock.NextFree;
-				if (PrevFreeIndex != 0xffff)
-				{
-					AddressBlocks[PrevFreeIndex].NextFree = NextFreeIndex;
-					FreeBlock.PrevFree = 0xffff;
-				}
-				if (NextFreeIndex != 0xffff)
-				{
-					AddressBlocks[NextFreeIndex].PrevFree = PrevFreeIndex;
-					FreeBlock.NextFree = 0xffff;
-				}
-				if (FreeList[AddressBlock.vLogSize] == FreeIndex)
-				{
-					FreeList[AddressBlock.vLogSize] = NextFreeIndex;
-				}
+				UnlinkFreeList(FreeList[AddressBlock.vLogSize], EBlockState::FreeList, FreeIndex);
 			}
 
-			if (GlobalFreeList != 0xffff)
-			{
-				FreeBlock.NextFree = GlobalFreeList;
-				AddressBlocks[GlobalFreeList].PrevFree = FreeIndex;
-			}
-			GlobalFreeList = FreeIndex;
-			FreeBlock.State = EBlockState::GlobalFreeList;
+			LinkFreeList(GlobalFreeList, EBlockState::GlobalFreeList, FreeIndex);
 
 			FreeIndex = FreeBlock.NextSibling;
 		}
 
+		check(AddressBlock.State == EBlockState::GlobalFreeList);
+
 		// Remove this block and its siblings from the sorted lists
 		// We can assume that the sibling blocks are sequential in the sorted list since they are free and so have no children
 		// FirstSibling will be the last in the range of siblings in the sorted lists 
@@ -496,7 +584,7 @@ void FVirtualTextureAllocator::FreeAddressBlock(uint32 Index)
 		check(SortedAddresses[SortedIndexRangeStart] == AddressBlocks[AddressBlock.Parent].vAddress);
 
 		// Add parent block to free list (and possibly consolidate)
-		FreeAddressBlock(AddressBlock.Parent);
+		FreeAddressBlock(AddressBlock.Parent, false);
 	}
 }
 
@@ -527,3 +615,108 @@ void FVirtualTextureAllocator::DumpToConsole(bool verbose)
 		}
 	}
 }
+
+#if WITH_EDITOR
+void FVirtualTextureAllocator::FillDebugImage(uint32 Index, uint32* ImageData, TMap<FAllocatedVirtualTexture*, uint32>& ColorMap) const
+{
+	const FAddressBlock& AddressBlock = AddressBlocks[Index];
+	if (AddressBlock.State == EBlockState::AllocatedTexture || AddressBlock.State == EBlockState::FreeList)
+	{
+		const uint32 vTileX = FMath::ReverseMortonCode2(AddressBlock.vAddress);
+		const uint32 vTileY = FMath::ReverseMortonCode2(AddressBlock.vAddress >> 1);
+		const uint32 BlockSize = (1u << AddressBlock.vLogSize);
+
+		if (vTileX + BlockSize <= AllocatedWidth && vTileY + BlockSize <= AllocatedHeight)
+		{
+			uint32 Color;
+			uint32 BorderColor;
+			if (AddressBlock.State == EBlockState::FreeList)
+			{
+				// Free blocks are black, with grey border
+				Color = FColor::Black.ToPackedABGR();
+				BorderColor = FColor(100u, 100u, 100u).ToPackedABGR();
+			}
+			else
+			{
+				// Allocated blocks have white border, random color for each AllocatedVT
+				uint32* FoundColor = ColorMap.Find(AddressBlock.VT);
+				if (!FoundColor)
+				{
+					FoundColor = &ColorMap.Add(AddressBlock.VT, FColor::MakeRandomColor().ToPackedABGR());
+				}
+				Color = *FoundColor;
+				BorderColor = FColor::White.ToPackedABGR();
+			}
+
+			// Add top/bottom borders
+			for (uint32 X = 0u; X < BlockSize; ++X)
+			{
+				const uint32 ImageY0 = vTileY;
+				const uint32 ImageY1 = vTileY + BlockSize - 1u;
+				const uint32 ImageX = vTileX + X;
+				ImageData[ImageY0 * AllocatedWidth + ImageX] = BorderColor;
+				ImageData[ImageY1 * AllocatedWidth + ImageX] = BorderColor;
+			}
+
+			for (uint32 Y = 1u; Y < BlockSize - 1u; ++Y)
+			{
+				const uint32 ImageY = vTileY + Y;
+
+				// Add left/right borders
+				ImageData[ImageY * AllocatedWidth + vTileX] = BorderColor;
+				ImageData[ImageY * AllocatedWidth + vTileX + BlockSize - 1u] = BorderColor;
+
+				for (uint32 X = 1u; X < BlockSize - 1u; ++X)
+				{
+					const uint32 ImageX = vTileX + X;
+					ImageData[ImageY * AllocatedWidth + ImageX] = Color;
+				}
+			}
+		}
+		else
+		{
+			// If block is outside allocated size, it must be free
+			check(AddressBlock.State == EBlockState::FreeList);
+		}
+	}
+	else if (AddressBlock.State == EBlockState::PartiallyFreeList)
+	{
+		uint32 ChildIndex = AddressBlock.FirstChild;
+		while (ChildIndex != 0xffff)
+		{
+			FillDebugImage(ChildIndex, ImageData, ColorMap);
+			ChildIndex = AddressBlocks[ChildIndex].NextSibling;
+		}
+	}
+	else
+	{
+		// Blocks of this state should not be in the tree
+		checkf(false, TEXT("Invalid block state %d"), (int32)AddressBlock.State);
+	}
+}
+
+void FVirtualTextureAllocator::SaveDebugImage(const TCHAR* ImageName) const
+{
+	const uint32 EmptyColor = FColor(255, 0, 255).ToPackedABGR();
+	TArray<uint32> ImageData;
+	ImageData.Init(EmptyColor, AllocatedWidth * AllocatedHeight);
+	TMap<FAllocatedVirtualTexture*, uint32> ColorMap;
+	FillDebugImage(RootIndex, ImageData.GetData(), ColorMap);
+
+	IImageWrapperModule& ImageWrapperModule = FModuleManager::LoadModuleChecked<IImageWrapperModule>(FName("ImageWrapper"));
+	TSharedPtr<IImageWrapper> ImageWrapper = ImageWrapperModule.CreateImageWrapper(EImageFormat::PNG);
+	ImageWrapper->SetRaw(ImageData.GetData(), ImageData.Num() * 4, AllocatedWidth, AllocatedHeight, ERGBFormat::RGBA, 8);
+
+	// Compress and write image
+	IFileManager& FileManager = IFileManager::Get();
+	const FString BasePath = FPaths::ProjectUserDir();
+	const FString ImagePath = BasePath / ImageName;
+	FArchive* Ar = FileManager.CreateFileWriter(*ImagePath);
+	if (Ar)
+	{
+		const TArray64<uint8>& CompressedData = ImageWrapper->GetCompressed();
+		Ar->Serialize((void*)CompressedData.GetData(), CompressedData.Num());
+		delete Ar;
+	}
+}
+#endif // WITH_EDITOR
diff --git a/Engine/Source/Runtime/Renderer/Private/VT/VirtualTextureAllocator.h b/Engine/Source/Runtime/Renderer/Private/VT/VirtualTextureAllocator.h
--- a/Engine/Source/Runtime/Renderer/Private/VT/VirtualTextureAllocator.h
+++ b/Engine/Source/Runtime/Renderer/Private/VT/VirtualTextureAllocator.h
@@ -12,18 +12,16 @@ class FAllocatedVirtualTexture;
 class RENDERER_API FVirtualTextureAllocator
 {
 public:
-	FVirtualTextureAllocator(uint32 Dimensions);
+	explicit FVirtualTextureAllocator(uint32 Dimensions);
 	~FVirtualTextureAllocator() {}
 
 	/**
-	 * Initialise the allocator with the given initial size.
+	 * Initialise the allocator
 	 */
-	void Initialize(uint32 InSize);
+	void Initialize(uint32 MaxSize);
 
-	/**
-	 * Increase size of region managed by allocator by factor of 2 in each dimension.
-	 */
-	void Grow();
+	uint32 GetAllocatedWidth() const { return AllocatedWidth; }
+	uint32 GetAllocatedHeight() const { return AllocatedHeight; }
 
 	/**
 	 * Translate a virtual page address in the address space to a local page address within a virtual texture.
@@ -58,24 +56,46 @@ public:
 	/** Output debugging information to the console. */
 	void DumpToConsole(bool verbose);
 
-private:
-	int32 AcquireBlock();
-	void FreeAddressBlock(uint32 Index);
-	uint32 FindAddressBlock(uint32 vAddress) const;
-
-	void SubdivideBlock(uint32 ParentIndex);
-
-	void MarkBlockAllocated(uint32 Index, uint32 vAllocatedTileX, uint32 vAllocatedTileY, FAllocatedVirtualTexture* VT);
+#if WITH_EDITOR
+	void SaveDebugImage(const TCHAR* ImageName) const;
+#endif
 
+private:
 	enum class EBlockState : uint8
 	{
 		None,
 		GlobalFreeList,
 		FreeList,
-		AllocatedChildren,
+		PartiallyFreeList,
 		AllocatedTexture,
 	};
 
+	struct FTestRegion
+	{
+		uint32 BaseIndex;
+		uint32 vTileX0;
+		uint32 vTileY0;
+		uint32 vTileX1;
+		uint32 vTileY1;
+	};
+
+	void LinkFreeList(uint16& InOutListHead, EBlockState State, uint16 Index);
+	void UnlinkFreeList(uint16& InOutListHead, EBlockState State, uint16 Index);
+
+	int32 AcquireBlock();
+	void FreeAddressBlock(uint32 Index, bool bTopLevelBlock);
+	uint32 FindAddressBlock(uint32 vAddress) const;
+
+	void SubdivideBlock(uint32 ParentIndex);
+
+	void MarkBlockAllocated(uint32 Index, uint32 vAllocatedTileX, uint32 vAllocatedTileY, FAllocatedVirtualTexture* VT);
+
+	bool TestAllocation(uint32 Index, uint32 vTileX0, uint32 vTileY0, uint32 vTileX1, uint32 vTileY1) const;
+
+#if WITH_EDITOR
+	void FillDebugImage(uint32 Index, uint32* ImageData, TMap<FAllocatedVirtualTexture*, uint32>& ColorMap) const;
+#endif
+
 	struct FAddressBlock
 	{
 		FAllocatedVirtualTexture*	VT;
@@ -123,10 +143,12 @@ private:
 	};
 
 	const uint32			vDimensions;
-	uint32					LogSize;
+	uint32					AllocatedWidth;
+	uint32					AllocatedHeight;
 
 	TArray< FAddressBlock >	AddressBlocks;
 	TArray< uint16 >		FreeList;
+	TArray< uint16 >		PartiallyFreeList;
 	uint16					GlobalFreeList;
 	TArray< uint32 >		SortedAddresses;
 	TArray< uint16 >		SortedIndices;
diff --git a/Engine/Source/Runtime/Renderer/Private/VT/VirtualTextureSpace.cpp b/Engine/Source/Runtime/Renderer/Private/VT/VirtualTextureSpace.cpp
--- a/Engine/Source/Runtime/Renderer/Private/VT/VirtualTextureSpace.cpp
+++ b/Engine/Source/Runtime/Renderer/Private/VT/VirtualTextureSpace.cpp
@@ -13,6 +13,8 @@
 #include "HAL/IConsoleManager.h"
 #include "SceneUtils.h"
 
+#include "VT/AllocatedVirtualTexture.h"
+
 DEFINE_LOG_CATEGORY_STATIC(LogVirtualTextureSpace, Log, All);
 
 static TAutoConsoleVariable<int32> CVarVTRefreshEntirePageTable(
@@ -45,10 +47,11 @@ static EPixelFormat GetFormatForNumLayers(uint32 NumLayers, EVTPageTableFormat F
 FVirtualTextureSpace::FVirtualTextureSpace(FVirtualTextureSystem* InSystem, uint8 InID, const FVTSpaceDescription& InDesc, uint32 InSizeNeeded)
 	: Description(InDesc)
 	, Allocator(InDesc.Dimensions)
-	, PageTableSize(0u)
-	, NumPageTableLevels(0u)
 	, NumRefs(0u)
 	, ID(InID)
+	, CachedPageTableWidth(0u)
+	, CachedPageTableHeight(0u)
+	, CachedNumPageTableLevels(0u)
 	, bNeedToAllocatePageTable(true)
 	, bForceEntireUpdate(false)
 {
@@ -73,12 +76,12 @@ FVirtualTextureSpace::FVirtualTextureSpace(FVirtualTextureSystem* InSystem, uint
 		++PageTableIndex;
 	}
 
-	PageTableSize = FMath::Max(InSizeNeeded, VIRTUALTEXTURE_MIN_PAGETABLE_SIZE);
-	PageTableSize = FMath::RoundUpToPowerOfTwo(PageTableSize);
-	ensure(PageTableSize <= Description.MaxSpaceSize);
-	ensure(Description.MaxSpaceSize <= VIRTUALTEXTURE_MAX_PAGETABLE_SIZE);
-	NumPageTableLevels = FMath::FloorLog2(PageTableSize) + 1;
-	Allocator.Initialize(PageTableSize);
+	//PageTableSize = FMath::Max(InSizeNeeded, VIRTUALTEXTURE_MIN_PAGETABLE_SIZE);
+	//PageTableSize = FMath::RoundUpToPowerOfTwo(PageTableSize);
+	//ensure(PageTableSize <= Description.MaxSpaceSize);
+	//ensure(Description.MaxSpaceSize <= VIRTUALTEXTURE_MAX_PAGETABLE_SIZE);
+	//NumPageTableLevels = FMath::FloorLog2(PageTableSize) + 1;
+	Allocator.Initialize(Description.MaxSpaceSize);
 
 	bNeedToAllocatePageTableIndirection = InDesc.IndirectionTextureSize > 0;
 }
@@ -90,15 +93,11 @@ FVirtualTextureSpace::~FVirtualTextureSpace()
 uint32 FVirtualTextureSpace::AllocateVirtualTexture(FAllocatedVirtualTexture* VirtualTexture)
 {
 	uint32 vAddress = Allocator.Alloc(VirtualTexture);
-	while (vAddress == ~0u && PageTableSize < Description.MaxSpaceSize)
+	if (Allocator.GetAllocatedWidth() > CachedPageTableWidth || Allocator.GetAllocatedHeight() > CachedPageTableHeight)
 	{
-		// Allocation failed, expand the size of page table texture and try again
-		PageTableSize *= 2u;
-		++NumPageTableLevels;
 		bNeedToAllocatePageTable = true;
-		Allocator.Grow();
-		vAddress = Allocator.Alloc(VirtualTexture);
 	}
+
 	return vAddress;
 }
 
@@ -139,7 +138,7 @@ uint32 FVirtualTextureSpace::GetSizeInBytes() const
 	uint32 TotalSize = 0u;
 	for (uint32 TextureIndex = 0u; TextureIndex < GetNumPageTableTextures(); ++TextureIndex)
 	{
-		const SIZE_T TextureSize = CalculateImageBytes(PageTableSize, PageTableSize, 0, TexturePixelFormat[TextureIndex]);
+		const SIZE_T TextureSize = CalcTextureSize(CachedPageTableWidth, CachedPageTableHeight, TexturePixelFormat[TextureIndex], CachedNumPageTableLevels);
 		TotalSize += TextureSize;
 	}
 	
@@ -272,16 +271,22 @@ void FVirtualTextureSpace::AllocateTextures(FRDGBuilder& GraphBuilder)
 		const TCHAR* TextureNames[] = { TEXT("VirtualPageTable_0"), TEXT("VirtualPageTable_1") };
 		static_assert(UE_ARRAY_COUNT(TextureNames) == TextureCapacity, "");
 
+		CachedPageTableWidth = Align(Allocator.GetAllocatedWidth(), VIRTUALTEXTURE_MIN_PAGETABLE_SIZE);
+		CachedPageTableHeight = Align(Allocator.GetAllocatedHeight(), VIRTUALTEXTURE_MIN_PAGETABLE_SIZE);
+		//CachedPageTableWidth = FMath::Max(CachedPageTableWidth, VIRTUALTEXTURE_MIN_PAGETABLE_SIZE);
+		//CachedPageTableHeight = FMath::Max(CachedPageTableHeight, VIRTUALTEXTURE_MIN_PAGETABLE_SIZE);
+		CachedNumPageTableLevels = FMath::FloorLog2(FMath::Max(CachedPageTableWidth, CachedPageTableHeight)) + 1u;
+
 		for (uint32 TextureIndex = 0u; TextureIndex < GetNumPageTableTextures(); ++TextureIndex)
 		{
 			// Page Table
 			FTextureEntry& TextureEntry = PageTable[TextureIndex];
 			const FRDGTextureDesc Desc = FRDGTextureDesc::Create2D(
-				FIntPoint(PageTableSize, PageTableSize),
+				FIntPoint(CachedPageTableWidth, CachedPageTableHeight),
 				TexturePixelFormat[TextureIndex],
 				FClearValueBinding::None,
 				TexCreate_RenderTargetable | TexCreate_ShaderResource,
-				NumPageTableLevels);
+				CachedNumPageTableLevels);
 
 			FRDGTextureRef DstTexture = GraphBuilder.CreateTexture(Desc, TextureNames[TextureIndex]);
 
@@ -374,7 +379,7 @@ void FVirtualTextureSpace::ApplyUpdates(FVirtualTextureSystem* System, FRDGBuild
 	uint32 TotalNumUpdates = 0;
 	for (uint32 LayerIndex = 0u; LayerIndex < Description.NumPageTableLayers; ++LayerIndex)
 	{
-		for (uint32 Mip = 0; Mip < NumPageTableLevels; Mip++)
+		for (uint32 Mip = 0; Mip < CachedNumPageTableLevels; Mip++)
 		{
 			TotalNumUpdates += ExpandedUpdates[LayerIndex][Mip].Num();
 		}
@@ -405,7 +410,7 @@ void FVirtualTextureSpace::ApplyUpdates(FVirtualTextureSystem* System, FRDGBuild
 		uint8* Buffer = (uint8*)RHILockBuffer(UpdateBuffer, 0, TotalNumUpdates * sizeof(FPageTableUpdate), RLM_WriteOnly);
 		for (uint32 LayerIndex = 0u; LayerIndex < Description.NumPageTableLayers; ++LayerIndex)
 		{
-			for (uint32 Mip = 0; Mip < NumPageTableLevels; Mip++)
+			for (uint32 Mip = 0; Mip < CachedNumPageTableLevels; Mip++)
 			{
 				const uint32 NumUpdates = ExpandedUpdates[LayerIndex][Mip].Num();
 				if (NumUpdates)
@@ -467,8 +472,9 @@ void FVirtualTextureSpace::ApplyUpdates(FVirtualTextureSystem* System, FRDGBuild
 		}
 		check(PixelShader.IsValid());
 
-		uint32 MipSize = PageTableSize;
-		for (uint32 Mip = 0; Mip < NumPageTableLevels; Mip++)
+		uint32 MipWidth = CachedPageTableWidth;
+		uint32 MipHeight = CachedPageTableHeight;
+		for (uint32 Mip = 0; Mip < CachedNumPageTableLevels; Mip++)
 		{
 			const uint32 NumUpdates = ExpandedUpdates[LayerIndex][Mip].Num();
 			if (NumUpdates)
@@ -480,9 +486,9 @@ void FVirtualTextureSpace::ApplyUpdates(FVirtualTextureSystem* System, FRDGBuild
 					RDG_EVENT_NAME("PageTableUpdate (Mip: %d)", Mip),
 					PassParameters,
 					ERDGPassFlags::Raster,
-					[this, VertexShader, PixelShader, BlendStateRHI, FirstUpdate, NumUpdates, MipSize](FRHICommandList& RHICmdList)
+					[this, VertexShader, PixelShader, BlendStateRHI, FirstUpdate, NumUpdates, MipWidth, MipHeight](FRHICommandList& RHICmdList)
 				{
-					RHICmdList.SetViewport(0, 0, 0.0f, MipSize, MipSize, 1.0f);
+					RHICmdList.SetViewport(0, 0, 0.0f, MipWidth, MipHeight, 1.0f);
 
 					FGraphicsPipelineStateInitializer GraphicsPSOInit;
 					RHICmdList.ApplyCachedRenderTargets(GraphicsPSOInit);
@@ -500,7 +506,7 @@ void FVirtualTextureSpace::ApplyUpdates(FVirtualTextureSystem* System, FRDGBuild
 
 					{
 						FRHIVertexShader* ShaderRHI = VertexShader.GetVertexShader();
-						SetShaderValue(RHICmdList, ShaderRHI, VertexShader->PageTableSize, PageTableSize);
+						SetShaderValue(RHICmdList, ShaderRHI, VertexShader->PageTableSize, FIntPoint(CachedPageTableWidth, CachedPageTableHeight));
 						SetShaderValue(RHICmdList, ShaderRHI, VertexShader->FirstUpdate, FirstUpdate);
 						SetShaderValue(RHICmdList, ShaderRHI, VertexShader->NumUpdates, NumUpdates);
 						SetSRVParameter(RHICmdList, ShaderRHI, VertexShader->UpdateBuffer, UpdateBufferSRV);
@@ -517,7 +523,8 @@ void FVirtualTextureSpace::ApplyUpdates(FVirtualTextureSystem* System, FRDGBuild
 			}
 
 			FirstUpdate += NumUpdates;
-			MipSize >>= 1;
+			MipWidth = FMath::Max(MipWidth / 2u, 1u);
+			MipHeight = FMath::Max(MipHeight / 2u, 1u);
 		}
 
 		ConvertToUntrackedExternalTexture(GraphBuilder, PageTableTexture, PageTableEntry.RenderTarget, ERHIAccess::SRVMask);
@@ -529,3 +536,11 @@ void FVirtualTextureSpace::DumpToConsole(bool verbose)
 	UE_LOG(LogConsoleResponse, Display, TEXT("-= Space ID %i =-"), ID);
 	Allocator.DumpToConsole(verbose);
 }
+
+#if WITH_EDITOR
+void FVirtualTextureSpace::SaveAllocatorDebugImage() const
+{
+	const FString ImageName = FString::Printf(TEXT("Space%dAllocator.png"), ID);
+	Allocator.SaveDebugImage(*ImageName);
+}
+#endif
\ No newline at end of file
diff --git a/Engine/Source/Runtime/Renderer/Private/VT/VirtualTextureSpace.h b/Engine/Source/Runtime/Renderer/Private/VT/VirtualTextureSpace.h
--- a/Engine/Source/Runtime/Renderer/Private/VT/VirtualTextureSpace.h
+++ b/Engine/Source/Runtime/Renderer/Private/VT/VirtualTextureSpace.h
@@ -51,14 +51,15 @@ public:
 	virtual ~FVirtualTextureSpace();
 
 	inline const FVTSpaceDescription& GetDescription() const { return Description; }
-	inline uint32 GetPageTableSize() const { return PageTableSize; }
+	inline uint32 GetPageTableWidth() const { return CachedPageTableWidth; }
+	inline uint32 GetPageTableHeight() const { return CachedPageTableHeight; }
 	inline uint8 GetDimensions() const { return Description.Dimensions; }
 	inline EVTPageTableFormat GetPageTableFormat() const { return Description.PageTableFormat; }
 	inline uint32 GetNumPageTableLayers() const { return Description.NumPageTableLayers; }
 	inline uint32 GetNumPageTableTextures() const { return (Description.NumPageTableLayers + LayersPerPageTableTexture - 1u) / LayersPerPageTableTexture; }
 	inline uint8 GetID() const { return ID; }
 
-	inline uint32 GetNumPageTableLevels() const { return NumPageTableLevels; }
+	inline uint32 GetNumPageTableLevels() const { return CachedNumPageTableLevels; }
 	inline FVirtualTextureAllocator& GetAllocator() { return Allocator; }
 	inline const FVirtualTextureAllocator& GetAllocator() const { return Allocator; }
 	inline FTexturePageMap& GetPageMapForPageTableLayer(uint32 PageTableLayerIndex) { check(PageTableLayerIndex < Description.NumPageTableLayers); return PhysicalPageMap[PageTableLayerIndex]; }
@@ -96,6 +97,10 @@ public:
 
 	void DumpToConsole(bool verbose);
 
+#if WITH_EDITOR
+	void SaveAllocatorDebugImage() const;
+#endif // WITH_EDITOR
+
 private:
 	static const uint32 TextureCapacity = (VIRTUALTEXTURE_SPACE_MAXLAYERS + LayersPerPageTableTexture - 1u) / LayersPerPageTableTexture;
 
@@ -120,11 +125,14 @@ private:
 	FBufferRHIRef UpdateBuffer;
 	FShaderResourceViewRHIRef UpdateBufferSRV;
 
-	uint32 PageTableSize;
-	uint32 NumPageTableLevels;
+	//uint32 PageTableSize;
+	//uint32 NumPageTableLevels;
 	uint32 NumRefs;
 
 	uint8 ID;
+	uint32 CachedPageTableWidth;
+	uint32 CachedPageTableHeight;
+	uint32 CachedNumPageTableLevels;
 	bool bNeedToAllocatePageTable;
 	bool bNeedToAllocatePageTableIndirection;
 	bool bForceEntireUpdate;
diff --git a/Engine/Source/Runtime/Renderer/Private/VT/VirtualTextureSystem.cpp b/Engine/Source/Runtime/Renderer/Private/VT/VirtualTextureSystem.cpp
--- a/Engine/Source/Runtime/Renderer/Private/VT/VirtualTextureSystem.cpp
+++ b/Engine/Source/Runtime/Renderer/Private/VT/VirtualTextureSystem.cpp
@@ -222,14 +222,18 @@ FVirtualTextureSystem& FVirtualTextureSystem::Get()
 }
 
 FVirtualTextureSystem::FVirtualTextureSystem()
-	: Frame(1u), // Need to start on Frame 1, otherwise the first call to update will fail to allocate any pages
-	bFlushCaches(false),
-	FlushCachesCommand(TEXT("r.VT.Flush"), TEXT("Flush all the physical caches in the VT system."),
-		FConsoleCommandDelegate::CreateRaw(this, &FVirtualTextureSystem::FlushCachesFromConsole)),
-	DumpCommand(TEXT("r.VT.Dump"), TEXT("Lot a whole lot of info on the VT system state."),
-		FConsoleCommandDelegate::CreateRaw(this, &FVirtualTextureSystem::DumpFromConsole)),
-	ListPhysicalPools(TEXT("r.VT.ListPhysicalPools"), TEXT("Lot a whole lot of info on the VT system state."),
+	: Frame(1u) // Need to start on Frame 1, otherwise the first call to update will fail to allocate any pages
+	, bFlushCaches(false)
+	, FlushCachesCommand(TEXT("r.VT.Flush"), TEXT("Flush all the physical caches in the VT system."),
+		FConsoleCommandDelegate::CreateRaw(this, &FVirtualTextureSystem::FlushCachesFromConsole))
+	, DumpCommand(TEXT("r.VT.Dump"), TEXT("Lot a whole lot of info on the VT system state."),
+		FConsoleCommandDelegate::CreateRaw(this, &FVirtualTextureSystem::DumpFromConsole))
+	, ListPhysicalPools(TEXT("r.VT.ListPhysicalPools"), TEXT("Lot a whole lot of info on the VT system state."),
 		FConsoleCommandDelegate::CreateRaw(this, &FVirtualTextureSystem::ListPhysicalPoolsFromConsole))
+#if WITH_EDITOR
+	, SaveAllocatorImages(TEXT("r.VT.SaveAllocatorImages"), TEXT("Save images showing allocator usage."),
+		FConsoleCommandDelegate::CreateRaw(this, &FVirtualTextureSystem::SaveAllocatorImagesFromConsole))
+#endif
 {
 }
 
@@ -370,10 +374,11 @@ void FVirtualTextureSystem::ListPhysicalPoolsFromConsole()
 
 		const FVTSpaceDescription& Desc = Space->GetDescription();
 		const FVirtualTextureAllocator& Allocator = Space->GetAllocator();
-		const uint32 PageTableSize = Space->GetPageTableSize();
+		const uint32 PageTableWidth = Space->GetPageTableWidth();
+		const uint32 PageTableHeight = Space->GetPageTableHeight();
 		const uint32 TotalSizeInBytes = Space->GetSizeInBytes();
 		const uint32 NumAllocatedPages = Allocator.GetNumAllocatedPages();
-		const uint32 NumTotalPages = PageTableSize * PageTableSize;
+		const uint32 NumTotalPages = PageTableWidth * PageTableHeight;
 		const double AllocatedRatio = (double)NumAllocatedPages / NumTotalPages;
 
 		const uint32 PhysicalTileSize = Desc.TileSize + Desc.TileBorderSize * 2u;
@@ -386,7 +391,7 @@ void FVirtualTextureSystem::ListPhysicalPoolsFromConsole()
 		}
 
 		UE_LOG(LogConsoleResponse, Display, TEXT("Pool: [%i] %s (%ix%i) x %i:"), ID, FormatName, PhysicalTileSize, PhysicalTileSize, Desc.NumPageTableLayers);
-		UE_LOG(LogConsoleResponse, Display, TEXT("  PageTableSize= %ix%i"), PageTableSize, PageTableSize);
+		UE_LOG(LogConsoleResponse, Display, TEXT("  PageTableSize= %ix%i"), PageTableWidth, PageTableHeight);
 		UE_LOG(LogConsoleResponse, Display, TEXT("  Allocations= %i, %i%% (%fMB)"),
 			Allocator.GetNumAllocations(),
 			(int)(AllocatedRatio * 100.0),
@@ -399,6 +404,20 @@ void FVirtualTextureSystem::ListPhysicalPoolsFromConsole()
 	UE_LOG(LogConsoleResponse, Display, TEXT("TotalPhysicalMemory: %fMB"), (double)TotalPhysicalMemory / 1024.0 / 1024.0);
 }
 
+#if WITH_EDITOR
+void FVirtualTextureSystem::SaveAllocatorImagesFromConsole()
+{
+	for (int ID = 0; ID < 16; ID++)
+	{
+		const FVirtualTextureSpace* Space = Spaces[ID].Get();
+		if (Space)
+		{
+			Space->SaveAllocatorDebugImage();
+		}
+	}
+}
+#endif // WITH_EDITOR
+
 IAllocatedVirtualTexture* FVirtualTextureSystem::AllocateVirtualTexture(const FAllocatedVTDescription& Desc)
 {
 	check(Desc.NumTextureLayers <= VIRTUALTEXTURE_SPACE_MAXLAYERS);
diff --git a/Engine/Source/Runtime/Renderer/Private/VT/VirtualTextureSystem.h b/Engine/Source/Runtime/Renderer/Private/VT/VirtualTextureSystem.h
--- a/Engine/Source/Runtime/Renderer/Private/VT/VirtualTextureSystem.h
+++ b/Engine/Source/Runtime/Renderer/Private/VT/VirtualTextureSystem.h
@@ -126,6 +126,11 @@ private:
 	void ListPhysicalPoolsFromConsole();
 	FAutoConsoleCommand ListPhysicalPools;
 
+#if WITH_EDITOR
+	void SaveAllocatorImagesFromConsole();
+	FAutoConsoleCommand SaveAllocatorImages;
+#endif
+
 	FCriticalSection RequestedTilesLock;
 	TArray<uint32> RequestedPackedTiles;
 
