AllocatedVT no longer assumes a square power-2 allocation region when unmapping/evicting pages

diff --git a/Engine/Source/Runtime/Renderer/Private/VT/AllocatedVirtualTexture.cpp b/Engine/Source/Runtime/Renderer/Private/VT/AllocatedVirtualTexture.cpp
--- a/Engine/Source/Runtime/Renderer/Private/VT/AllocatedVirtualTexture.cpp
+++ b/Engine/Source/Runtime/Renderer/Private/VT/AllocatedVirtualTexture.cpp
@@ -181,8 +181,8 @@ void FAllocatedVirtualTexture::Release(FVirtualTextureSystem* System)
 
 	// Physical pool needs to evict all pages that belong to this VT
 	{
-		// Allocator always allocates square regions, need to update logic here if that changes
-		const uint32 BlockSize = FMath::Max(GetWidthInTiles(), GetHeightInTiles());
+		const uint32 WidthInTiles = GetWidthInTiles();
+		const uint32 HeightInTiles = GetHeightInTiles();
 
 		TArray<FVirtualTexturePhysicalSpace*> UniquePhysicalSpaces;
 		for (int32 PageTableIndex = 0u; PageTableIndex < UniquePageTableLayers.Num(); ++PageTableIndex)
@@ -195,7 +195,7 @@ void FAllocatedVirtualTexture::Release(FVirtualTextureSystem* System)
 
 		for (FVirtualTexturePhysicalSpace* PhysicalSpace : UniquePhysicalSpaces)
 		{
-			PhysicalSpace->GetPagePool().UnmapAllPagesForSpace(System, Space->GetID(), VirtualAddress, BlockSize, MaxLevel);
+			PhysicalSpace->GetPagePool().UnmapAllPagesForSpace(System, Space->GetID(), VirtualAddress, WidthInTiles, HeightInTiles, MaxLevel);
 		}
 
 		for (int32 PageTableIndex = 0u; PageTableIndex < UniquePageTableLayers.Num(); ++PageTableIndex)
@@ -209,11 +209,11 @@ void FAllocatedVirtualTexture::Release(FVirtualTextureSystem* System)
 			const FTexturePageMap& PageMap = Space->GetPageMapForPageTableLayer(LayerIndex);
 
 			TArray<FMappedTexturePage> MappedPages;
-			PageMap.GetMappedPagesInRange(VirtualAddress, BlockSize, MappedPages);
+			PageMap.GetMappedPagesInRange(VirtualAddress, WidthInTiles, HeightInTiles, MappedPages);
 			if (MappedPages.Num() > 0)
 			{
 				TStringBuilder<2048> Message;
-				Message.Appendf(TEXT("Mapped pages remain after releasing AllocatedVT - vAddress: %d, Size: %d, PhysicalSpaces: ["), VirtualAddress, BlockSize);
+				Message.Appendf(TEXT("Mapped pages remain after releasing AllocatedVT - vAddress: %d, Size: %d x %d, PhysicalSpaces: ["), VirtualAddress, WidthInTiles, HeightInTiles);
 				for (FVirtualTexturePhysicalSpace* PhysicalSpace : UniquePhysicalSpaces)
 				{
 					Message.Appendf(TEXT("%d "), PhysicalSpace->GetID());
diff --git a/Engine/Source/Runtime/Renderer/Private/VT/TexturePageMap.cpp b/Engine/Source/Runtime/Renderer/Private/VT/TexturePageMap.cpp
--- a/Engine/Source/Runtime/Renderer/Private/VT/TexturePageMap.cpp
+++ b/Engine/Source/Runtime/Renderer/Private/VT/TexturePageMap.cpp
@@ -203,9 +203,12 @@ void FTexturePageMap::MapPage(FVirtualTextureSpace* Space, FVirtualTexturePhysic
 	SortedKeysDirty = true;
 }
 
-void FTexturePageMap::GetMappedPagesInRange(uint32 vAddress, uint32 Size, TArray<FMappedTexturePage>& OutMappedPages) const
+void FTexturePageMap::GetMappedPagesInRange(uint32 vAddress, uint32 Width, uint32 Height, TArray<FMappedTexturePage>& OutMappedPages) const
 {
+	const uint32 Size = FMath::Max(Width, Height);
 	const uint32 vAddressMax = vAddress + FMath::Square(Size);
+	const uint32 vTileX0 = FMath::ReverseMortonCode2(vAddress);
+	const uint32 vTileY0 = FMath::ReverseMortonCode2(vAddress >> 1);
 
 	uint32 PageIndex = Pages[PageListHead_Mapped].NextIndex;
 	uint32 CheckPageCount = 0u;
@@ -214,11 +217,19 @@ void FTexturePageMap::GetMappedPagesInRange(uint32 vAddress, uint32 Size, TArray
 		const FPageEntry& Entry = Pages[PageIndex];
 		if (Entry.Page.vAddress >= vAddress && Entry.Page.vAddress < vAddressMax)
 		{
-			FMappedTexturePage& OutPage = OutMappedPages.AddDefaulted_GetRef();
-			OutPage.Page = Entry.Page;
-			OutPage.pAddress = Entry.pAddress;
-			OutPage.PhysicalSpaceID = Entry.PhysicalSpaceID;
-			OutPage.vLevel = Entry.vLevel;
+			const uint32 vTileX = FMath::ReverseMortonCode2(Entry.Page.vAddress);
+			const uint32 vTileY = FMath::ReverseMortonCode2(Entry.Page.vAddress >> 1);
+			if (vTileX >= vTileX0 &&
+				vTileY >= vTileY0 &&
+				vTileX < vTileX0 + Width &&
+				vTileY < vTileY0 + Height)
+			{
+				FMappedTexturePage& OutPage = OutMappedPages.AddDefaulted_GetRef();
+				OutPage.Page = Entry.Page;
+				OutPage.pAddress = Entry.pAddress;
+				OutPage.PhysicalSpaceID = Entry.PhysicalSpaceID;
+				OutPage.vLevel = Entry.vLevel;
+			}
 		}
 
 		PageIndex = Entry.NextIndex;
diff --git a/Engine/Source/Runtime/Renderer/Private/VT/TexturePageMap.h b/Engine/Source/Runtime/Renderer/Private/VT/TexturePageMap.h
--- a/Engine/Source/Runtime/Renderer/Private/VT/TexturePageMap.h
+++ b/Engine/Source/Runtime/Renderer/Private/VT/TexturePageMap.h
@@ -98,7 +98,7 @@ public:
 	*/
 	void		MapPage(FVirtualTextureSpace* Space, FVirtualTexturePhysicalSpace* PhysicalSpace, uint8 vLogSize, uint32 vAddress, uint8 vLevel, uint16 pAddress);
 
-	void		GetMappedPagesInRange(uint32 vAddress, uint32 Size, TArray<FMappedTexturePage>& OutMappedPages) const;
+	void		GetMappedPagesInRange(uint32 vAddress, uint32 Width, uint32 Height, TArray<FMappedTexturePage>& OutMappedPages) const;
 
 	void		RefreshEntirePageTable(FVirtualTextureSystem* System, TArray< FPageTableUpdate >* Output);
 	void		ExpandPageTableUpdatePainters(FVirtualTextureSystem* System, FPageTableUpdate Update, TArray< FPageTableUpdate >* Output);
diff --git a/Engine/Source/Runtime/Renderer/Private/VT/TexturePagePool.cpp b/Engine/Source/Runtime/Renderer/Private/VT/TexturePagePool.cpp
--- a/Engine/Source/Runtime/Renderer/Private/VT/TexturePagePool.cpp
+++ b/Engine/Source/Runtime/Renderer/Private/VT/TexturePagePool.cpp
@@ -58,12 +58,16 @@ void FTexturePagePool::EvictAllPages(FVirtualTextureSystem* System)
 	}
 }
 
-void FTexturePagePool::UnmapAllPagesForSpace(FVirtualTextureSystem* System, uint8 SpaceID, uint32 vAddress, uint32 Size, uint32 MaxLevel)
+void FTexturePagePool::UnmapAllPagesForSpace(FVirtualTextureSystem* System, uint8 SpaceID, uint32 vAddress, uint32 Width, uint32 Height, uint32 MaxLevel)
 {
-	check(Size > 0u);
+	check(Width > 0u);
+	check(Height > 0u);
 	checkf((vAddress & (0xffffffff << (MaxLevel * 2u))) == vAddress, TEXT("vAddress %08X is not aligned to max level %d"), vAddress, MaxLevel);
 
+	const uint32 Size = FMath::Max(Width, Height);
 	const uint32 vAddressMax = vAddress + FMath::Square(Size);
+	const uint32 vTileX0 = FMath::ReverseMortonCode2(vAddress);
+	const uint32 vTileY0 = FMath::ReverseMortonCode2(vAddress >> 1);
 
 	// walk through all of our current mapping entries, and unmap any that belong to the current space
 	for (int32 MappingIndex = NumPages + 1; MappingIndex < PageMapping.Num(); ++MappingIndex)
@@ -75,8 +79,16 @@ void FTexturePagePool::UnmapAllPagesForSpace(FVirtualTextureSystem* System, uint
 			Mapping.vAddress >= vAddress &&
 			Mapping.vAddress < vAddressMax)
 		{
-			// we're unmapping all pages for space, so don't try to map any ancestor pages...they'll be unmapped as well
-			UnmapPageMapping(System, MappingIndex, false);
+			const uint32 vTileX = FMath::ReverseMortonCode2(Mapping.vAddress);
+			const uint32 vTileY = FMath::ReverseMortonCode2(Mapping.vAddress >> 1);
+			if (vTileX >= vTileX0 &&
+				vTileY >= vTileY0 &&
+				vTileX < vTileX0 + Width &&
+				vTileY < vTileY0 + Height)
+			{
+				// we're unmapping all pages for space, so don't try to map any ancestor pages...they'll be unmapped as well
+				UnmapPageMapping(System, MappingIndex, false);
+			}
 		}
 	}
 }
diff --git a/Engine/Source/Runtime/Renderer/Private/VT/TexturePagePool.h b/Engine/Source/Runtime/Renderer/Private/VT/TexturePagePool.h
--- a/Engine/Source/Runtime/Renderer/Private/VT/TexturePagePool.h
+++ b/Engine/Source/Runtime/Renderer/Private/VT/TexturePagePool.h
@@ -54,7 +54,7 @@ public:
 	/**
 	* Unmap all pages from the given address range in a space. Pages will remain resident in the pool, but no longer by mapped to any page table.
 	*/
-	void UnmapAllPagesForSpace(FVirtualTextureSystem* System, uint8 SpaceID, uint32 vAddress, uint32 Size, uint32 MaxLevel);
+	void UnmapAllPagesForSpace(FVirtualTextureSystem* System, uint8 SpaceID, uint32 vAddress, uint32 Width, uint32 Height, uint32 MaxLevel);
 
 	/**
 	* Remap physical pages from one producer to another.
