Avoid an assertion in the FGPUSceneDynamicContext destructor when pathracing or raytracing debug modes are active

We are not rendering shadows in these cases but were still doing some initialization which was leaving FGPUSceneDynamicContext in an unexpected state.

diff --git a/Engine/Source/Runtime/Renderer/Private/SceneVisibility.cpp b/Engine/Source/Runtime/Renderer/Private/SceneVisibility.cpp
--- a/Engine/Source/Runtime/Renderer/Private/SceneVisibility.cpp
+++ b/Engine/Source/Runtime/Renderer/Private/SceneVisibility.cpp
@@ -4681,9 +4681,12 @@ void FDeferredShadingSceneRenderer::InitViewsPossiblyAfterPrepass(FRDGBuilder& G
 
 	FRHICommandListImmediate& RHICmdList = GraphBuilder.RHICmdList;
 
+	const bool bCanOverlayRayTracingOutput = CanOverlayRayTracingOutput(Views[0]);// #dxr_todo: UE-72557 multi-view case
+
 	if (ViewFamily.EngineShowFlags.DynamicShadows 
 		&& !IsSimpleForwardShadingEnabled(ShaderPlatform)
-		&& !ViewFamily.EngineShowFlags.HitProxies)
+		&& !ViewFamily.EngineShowFlags.HitProxies
+		&& bCanOverlayRayTracingOutput)
 	{
 		// Setup dynamic shadows.
 		InitDynamicShadows(RHICmdList, DynamicIndexBufferForInitShadows, DynamicVertexBufferForInitShadows, DynamicReadBufferForInitShadows, InstanceCullingManager);
