Speculative fix for crash streaming distance fields

diff --git a/Engine/Source/Runtime/Renderer/Private/DistanceFieldStreaming.cpp b/Engine/Source/Runtime/Renderer/Private/DistanceFieldStreaming.cpp
--- a/Engine/Source/Runtime/Renderer/Private/DistanceFieldStreaming.cpp
+++ b/Engine/Source/Runtime/Renderer/Private/DistanceFieldStreaming.cpp
@@ -351,7 +351,11 @@ void FDistanceFieldSceneData::AsyncUpdate(FDistanceFieldAsyncUpdateParameters Up
 			if (BrickIndex != DistanceField::InvalidBrickIndex)
 			{
 				const int32 BlockIndex = BrickIndex / DistanceFieldBlockAllocatorSizeInBricks;
-				GlobalBrickIndex = BrickIndex % DistanceFieldBlockAllocatorSizeInBricks + GlobalBlockOffsets[BlockIndex] * DistanceFieldBlockAllocatorSizeInBricks;
+
+				if (BlockIndex < MipState.AllocatedBlocks.Num())
+				{
+					GlobalBrickIndex = BrickIndex % DistanceFieldBlockAllocatorSizeInBricks + GlobalBlockOffsets[BlockIndex] * DistanceFieldBlockAllocatorSizeInBricks;
+				}
 			}
 			DestIndirectionTable[i] = GlobalBrickIndex;
 		}
