Fix default material compile error due to onepass light shader being compiled on mobile.

diff --git a/Engine/Source/Runtime/Renderer/Private/ShadowDepthRendering.cpp b/Engine/Source/Runtime/Renderer/Private/ShadowDepthRendering.cpp
--- a/Engine/Source/Runtime/Renderer/Private/ShadowDepthRendering.cpp
+++ b/Engine/Source/Runtime/Renderer/Private/ShadowDepthRendering.cpp
@@ -314,7 +314,7 @@ public:
 			// Masked and WPO materials need their shaders but cannot be used with a position only stream.
 			|| ((!Parameters.MaterialParameters.bWritesEveryPixelShadowPass || Parameters.MaterialParameters.bMaterialMayModifyMeshPosition) && !bUsePositionOnlyStream))
 			// Only compile one pass point light shaders for feature levels >= SM5
-			&& (ShaderMode != VertexShadowDepth_OnePassPointLight || IsFeatureLevelSupported(Platform, ERHIFeatureLevel::SM5))
+				&& ((ShaderMode != VertexShadowDepth_OnePassPointLight && ShaderMode != VertexShadowDepth_VSLayer) || IsFeatureLevelSupported(Platform, ERHIFeatureLevel::SM5))
 			// Only compile position-only shaders for vertex factories that support it. (Note: this assumes that a vertex factor which supports PositionOnly, supports also PositionAndNormalOnly)
 			&& (!bUsePositionOnlyStream || Parameters.VertexFactoryType->SupportsPositionOnly())
 			// Don't render ShadowDepth for translucent unlit materials
@@ -1597,13 +1597,13 @@ void FSceneRenderer::RenderShadowDepthMaps(FRDGBuilder& GraphBuilder, FInstanceC
 	{				
 		FVirtualShadowMapArrayCacheManager *CacheManager = Scene->VirtualShadowMapArrayCacheManager;
 		const TRefCountPtr<IPooledRenderTarget> PrevHZBPhysical = bUseHZB ? CacheManager->PrevBuffers.HZBPhysical : nullptr;
-
+		
 		{
 			RDG_EVENT_SCOPE(GraphBuilder, "RenderVirtualShadowMaps");
 
 			const FIntPoint VirtualShadowSize = VirtualShadowMapArray.GetPhysicalPoolSize();
 			const FIntRect VirtualShadowViewRect = FIntRect(0, 0, VirtualShadowSize.X, VirtualShadowSize.Y);
- 
+
 			VirtualShadowMapArray.PhysicalPagePoolRDG = GraphBuilder.CreateTexture( 
 				FRDGTextureDesc::Create2D(VirtualShadowSize, PF_R32_UINT, FClearValueBinding::None, TexCreate_ShaderResource | TexCreate_UAV),
 				TEXT("Shadow.Virtual.PhysicalTexture") );
@@ -1679,11 +1679,11 @@ void FSceneRenderer::RenderShadowDepthMaps(FRDGBuilder& GraphBuilder, FInstanceC
 							{
 								FVirtualShadowMapHZBMetadata* PrevHZBMeta = CacheManager->PrevBuffers.HZBMetadata.Find(HZBKey);
 								if (PrevHZBMeta)
-								{
+							{
 									Params.PrevTargetLayerIndex = PrevHZBMeta->TargetLayerIndex;
 									Params.PrevViewMatrices = PrevHZBMeta->ViewMatrices;
-									Params.Flags = VIEW_FLAG_HZBTEST;
-								}
+								Params.Flags = VIEW_FLAG_HZBTEST;
+							}
 							}
 
 							// If we're going to generate a new HZB this frame, save the associated metadata
@@ -1730,11 +1730,11 @@ void FSceneRenderer::RenderShadowDepthMaps(FRDGBuilder& GraphBuilder, FInstanceC
 							{
 								FVirtualShadowMapHZBMetadata* PrevHZBMeta = CacheManager->PrevBuffers.HZBMetadata.Find(HZBKey);
 								if (PrevHZBMeta)
-								{
+							{
 									Params.PrevTargetLayerIndex = PrevHZBMeta->TargetLayerIndex;
 									Params.PrevViewMatrices = PrevHZBMeta->ViewMatrices;
-									Params.Flags = VIEW_FLAG_HZBTEST;
-								}
+								Params.Flags = VIEW_FLAG_HZBTEST;
+							}
 							}
 
 							// If we're going to generate a new HZB this frame, save the associated metadata
@@ -1745,7 +1745,7 @@ void FSceneRenderer::RenderShadowDepthMaps(FRDGBuilder& GraphBuilder, FInstanceC
 								HZBMeta.ViewMatrices = Params.ViewMatrices;
 								HZBMeta.ViewRect = Params.ViewRect;
 							}
-
+						
 							Nanite::FPackedView View = Nanite::CreatePackedView(Params);
 							VirtualShadowViews.Add(View);
 							VirtualShadowMapFlags[ProjectedShadowInfo->VirtualShadowMaps[i]->ID] = 1;
