Add voxel binding helper function.

diff --git a/Engine/Shaders/Private/RayTracing/RayTracingOcclusionRGS.usf b/Engine/Shaders/Private/RayTracing/RayTracingOcclusionRGS.usf
--- a/Engine/Shaders/Private/RayTracing/RayTracingOcclusionRGS.usf
+++ b/Engine/Shaders/Private/RayTracing/RayTracingOcclusionRGS.usf
@@ -1,5 +1,6 @@
 // Copyright Epic Games, Inc. All Rights Reserved.
 
+#include "../HairStrands/HairStrandsVisibilityCommonStruct.ush"
 #include "../HairStrands/HairStrandsVoxelPageCommonStruct.ush"
 
 #include "../Common.ush"
@@ -52,7 +53,6 @@ uint bAcceptFirstHit;
 #include "../HairStrands/HairStrandsRaytracing.ush"
 uint				bUseHairVoxel;
 float				HairOcclusionThreshold;
-Texture2D<uint4>	HairCategorizationTexture;
 Texture2D<uint>		HairLightChannelMaskTexture;
 
 void NeedTraceHair(
@@ -60,7 +60,7 @@ void NeedTraceHair(
 	inout bool bTraceHairRay,
 	inout float HairDeviceZ)
 {
-	const FCategorizationData Categorization = DecodeCategorizationData(HairCategorizationTexture.Load(uint3(PixelCoord, 0)));
+	const FCategorizationData Categorization = DecodeCategorizationData(HairStrands.HairCategorizationTexture.Load(uint3(PixelCoord, 0)));
 	bTraceHairRay = false;
 	HairDeviceZ = 0;
 	if (Categorization.PixelCoverage > 0)
diff --git a/Engine/Shaders/Private/RayTracing/RayTracingSkyLightRGS.usf b/Engine/Shaders/Private/RayTracing/RayTracingSkyLightRGS.usf
--- a/Engine/Shaders/Private/RayTracing/RayTracingSkyLightRGS.usf
+++ b/Engine/Shaders/Private/RayTracing/RayTracingSkyLightRGS.usf
@@ -1,5 +1,6 @@
 // Copyright Epic Games, Inc. All Rights Reserved.
 
+#include "../HairStrands/HairStrandsVisibilityCommonStruct.ush"
 #include "../HairStrands/HairStrandsVoxelPageCommonStruct.ush"
 
 #include "../Common.ush"
diff --git a/Engine/Source/Runtime/Renderer/Private/DeferredShadingRenderer.h b/Engine/Source/Runtime/Renderer/Private/DeferredShadingRenderer.h
--- a/Engine/Source/Runtime/Renderer/Private/DeferredShadingRenderer.h
+++ b/Engine/Source/Runtime/Renderer/Private/DeferredShadingRenderer.h
@@ -840,7 +840,6 @@ private:
 		const FLightSceneInfo& LightSceneInfo,
 		const IScreenSpaceDenoiser::FShadowRayTracingConfig& RayTracingConfig,
 		const IScreenSpaceDenoiser::EShadowRequirements DenoiserRequirements,
-		const struct FHairStrandsOcclusionResources* HairResources,
 		FRDGTextureRef LightingChannelsTexture,
 		FRDGTextureUAV* OutShadowMaskUAV,
 		FRDGTextureUAV* OutRayHitDistanceUAV,
diff --git a/Engine/Source/Runtime/Renderer/Private/HairStrands/HairStrandsData.h b/Engine/Source/Runtime/Renderer/Private/HairStrands/HairStrandsData.h
--- a/Engine/Source/Runtime/Renderer/Private/HairStrands/HairStrandsData.h
+++ b/Engine/Source/Runtime/Renderer/Private/HairStrands/HairStrandsData.h
@@ -341,6 +341,8 @@ namespace HairStrands
 {
 	TRDGUniformBufferRef<FHairStrandsViewUniformParameters> CreateDefaultHairStrandsViewUniformBuffer(FRDGBuilder& GraphBuilder, FViewInfo& View);
 	TRDGUniformBufferRef<FHairStrandsViewUniformParameters> BindHairStrandsViewUniformParameters(const FViewInfo& View);
+	TRDGUniformBufferRef<FVirtualVoxelParameters> BindHairStrandsVoxelUniformParameters(const FViewInfo& View);
 	bool HasViewHairStrandsData(const FViewInfo& View);
 	bool HasViewHairStrandsData(const TArrayView<FViewInfo>& Views);
+	bool HasViewHairStrandsVoxelData(const FViewInfo& View);
 }
\ No newline at end of file
diff --git a/Engine/Source/Runtime/Renderer/Private/HairStrands/HairStrandsRendering.cpp b/Engine/Source/Runtime/Renderer/Private/HairStrands/HairStrandsRendering.cpp
--- a/Engine/Source/Runtime/Renderer/Private/HairStrands/HairStrandsRendering.cpp
+++ b/Engine/Source/Runtime/Renderer/Private/HairStrands/HairStrandsRendering.cpp
@@ -126,11 +126,23 @@ TRDGUniformBufferRef<FHairStrandsViewUniformParameters> BindHairStrandsViewUnifo
 	return View.HairStrandsViewData.UniformBuffer;
 }
 
+TRDGUniformBufferRef<FVirtualVoxelParameters> BindHairStrandsVoxelUniformParameters(const FViewInfo& View)
+{
+	// Voxel uniform buffer exist only if the view has hair strands data
+	check(View.HairStrandsViewData.bIsValid && View.HairStrandsViewData.VirtualVoxelResources.IsValid());
+	return View.HairStrandsViewData.VirtualVoxelResources.UniformBuffer;
+}
+
 bool HasViewHairStrandsData(const FViewInfo& View)
 {
 	return View.HairStrandsViewData.bIsValid;
 }
 
+bool HasViewHairStrandsVoxelData(const FViewInfo& View)
+{
+	return View.HairStrandsViewData.bIsValid && View.HairStrandsViewData.VirtualVoxelResources.IsValid();
+}
+
 bool HasViewHairStrandsData(const TArrayView<FViewInfo>& Views)
 {
 	for (const FViewInfo& View : Views)
diff --git a/Engine/Source/Runtime/Renderer/Private/HairStrands/HairStrandsVoxelization.h b/Engine/Source/Runtime/Renderer/Private/HairStrands/HairStrandsVoxelization.h
--- a/Engine/Source/Runtime/Renderer/Private/HairStrands/HairStrandsVoxelization.h
+++ b/Engine/Source/Runtime/Renderer/Private/HairStrands/HairStrandsVoxelization.h
@@ -20,13 +20,4 @@ void VoxelizeHairStrands(
 	FRDGBuilder& GraphBuilder,
 	const class FScene* Scene,
 	FViewInfo& View,
-	FInstanceCullingManager& InstanceCullingManager);
-
-
-struct FHairStrandsOcclusionResources
-{
-	FRDGTextureRef CategorizationTexture = nullptr;
-	FRDGTextureRef LightChannelMaskTexture = nullptr;
-	const FHairStrandsVoxelResources* VoxelResources = nullptr;
-	bool bUseHairVoxel = false;
-};
\ No newline at end of file
+	FInstanceCullingManager& InstanceCullingManager);
\ No newline at end of file
diff --git a/Engine/Source/Runtime/Renderer/Private/LightRendering.cpp b/Engine/Source/Runtime/Renderer/Private/LightRendering.cpp
--- a/Engine/Source/Runtime/Renderer/Private/LightRendering.cpp
+++ b/Engine/Source/Runtime/Renderer/Private/LightRendering.cpp
@@ -1077,25 +1077,6 @@ void FSceneRenderer::GatherAndSortLights(FSortedLightSetSceneInfo& OutSortedLigh
 	check(OutSortedLights.AttenuationLightStart >= OutSortedLights.ClusteredSupportedEnd);
 }
 
-static FHairStrandsOcclusionResources GetHairStrandsResources(FRDGBuilder& GraphBuilder, const FViewInfo& View)
-{
-	FHairStrandsOcclusionResources Out;
-	if (HairStrands::HasViewHairStrandsData(View))
-	{
-		if (FRDGTextureRef Texture = View.HairStrandsViewData.VisibilityData.CategorizationTexture)
-		{
-			Out.CategorizationTexture = Texture;
-		}
-		if (FRDGTextureRef Texture = View.HairStrandsViewData.VisibilityData.LightChannelMaskTexture)
-		{
-			Out.LightChannelMaskTexture = Texture;
-		}
-
-		Out.VoxelResources = &View.HairStrandsViewData.VirtualVoxelResources;
-	}
-	return Out;
-}
-
 /** Shader parameters to use when creating a RenderLight(...) pass. */
 BEGIN_SHADER_PARAMETER_STRUCT(FRenderLightParameters, )
 	SHADER_PARAMETER_RDG_UNIFORM_BUFFER(FSceneTextureUniformParameters, SceneTextures)
@@ -1568,8 +1549,6 @@ void FDeferredShadingSceneRenderer::RenderLights(
 
 								FRDGTextureUAV* RayTracingShadowMaskUAV = GraphBuilder.CreateUAV(FRDGTextureUAVDesc(RayTracingShadowMaskTexture));
 								FRDGTextureUAV* RayHitDistanceUAV = GraphBuilder.CreateUAV(FRDGTextureUAVDesc(RayDistanceTexture));
-								FHairStrandsOcclusionResources HairResources = GetHairStrandsResources(GraphBuilder, View);
-								HairResources.bUseHairVoxel = !BatchLightSceneInfo.Proxy->CastsHairStrandsDeepShadow();
 								{
 									RDG_EVENT_SCOPE(GraphBuilder, "%s", *BatchLightNameWithLevel);
 
@@ -1582,7 +1561,6 @@ void FDeferredShadingSceneRenderer::RenderLights(
 										BatchLightSceneInfo,
 										BatchRayTracingConfig,
 										DenoiserRequirements,
-										&HairResources,
 										LightingChannelsTexture,
 										RayTracingShadowMaskUAV,
 										RayHitDistanceUAV,
@@ -1731,9 +1709,6 @@ void FDeferredShadingSceneRenderer::RenderLights(
 								DenoiserRequirements = DenoiserToUse->GetShadowRequirements(View, LightSceneInfo, RayTracingConfig);
 							}
 
-							FHairStrandsOcclusionResources HairResources = GetHairStrandsResources(GraphBuilder, View);
-							HairResources.bUseHairVoxel = !bUseHairDeepShadow;
-
 							RenderRayTracingShadows(
 								GraphBuilder,
 								SceneTextureParameters,
@@ -1741,7 +1716,6 @@ void FDeferredShadingSceneRenderer::RenderLights(
 								LightSceneInfo,
 								RayTracingConfig,
 								DenoiserRequirements,
-								&HairResources,
 								LightingChannelsTexture,
 								RayTracingShadowMaskUAV,
 								RayHitDistanceUAV,
diff --git a/Engine/Source/Runtime/Renderer/Private/RayTracing/RayTracingShadows.cpp b/Engine/Source/Runtime/Renderer/Private/RayTracing/RayTracingShadows.cpp
--- a/Engine/Source/Runtime/Renderer/Private/RayTracing/RayTracingShadows.cpp
+++ b/Engine/Source/Runtime/Renderer/Private/RayTracing/RayTracingShadows.cpp
@@ -156,7 +156,6 @@ class FOcclusionRGS : public FGlobalShader
 		SHADER_PARAMETER_STRUCT_INCLUDE(FSceneTextureParameters, SceneTextures)
 		SHADER_PARAMETER_STRUCT_INCLUDE(FSceneLightingChannelParameters, SceneLightingChannels)
 
-		SHADER_PARAMETER_RDG_TEXTURE(Texture2D, HairCategorizationTexture)
 		SHADER_PARAMETER_RDG_TEXTURE(Texture2D, HairLightChannelMaskTexture)
 		SHADER_PARAMETER_RDG_TEXTURE(Texture2D, SSProfilesTexture)
 		SHADER_PARAMETER_SRV(RaytracingAccelerationStructure, TLAS)
@@ -164,6 +163,7 @@ class FOcclusionRGS : public FGlobalShader
 		SHADER_PARAMETER_RDG_TEXTURE_UAV(RWTexture2D<float>, RWRayDistanceUAV)
 		SHADER_PARAMETER_RDG_TEXTURE_UAV(RWTexture2D<float4>, RWSubPixelOcclusionMaskUAV)
 		SHADER_PARAMETER_STRUCT_REF(FViewUniformShaderParameters, ViewUniformBuffer)
+		SHADER_PARAMETER_RDG_UNIFORM_BUFFER(FHairStrandsViewUniformParameters, HairStrands)
 		SHADER_PARAMETER_RDG_UNIFORM_BUFFER(FVirtualVoxelParameters, VirtualVoxel)
 	END_SHADER_PARAMETER_STRUCT()
 };
@@ -231,7 +231,6 @@ void FDeferredShadingSceneRenderer::RenderRayTracingShadows(
 	const FLightSceneInfo& LightSceneInfo,
 	const IScreenSpaceDenoiser::FShadowRayTracingConfig& RayTracingConfig,
 	const IScreenSpaceDenoiser::EShadowRequirements DenoiserRequirements,
-	const FHairStrandsOcclusionResources* HairResources,
 	FRDGTextureRef LightingChannelsTexture,
 	FRDGTextureUAV* OutShadowMaskUAV,
 	FRDGTextureUAV* OutRayHitDistanceUAV,
@@ -265,11 +264,7 @@ void FDeferredShadingSceneRenderer::RenderRayTracingShadows(
 
 	// Ray generation pass for shadow occlusion.
 	{
-		const bool bUseHairLighting =
-			HairResources != nullptr &&
-			HairResources->CategorizationTexture != nullptr &&
-			HairResources->LightChannelMaskTexture != nullptr &&
-			HairResources->VoxelResources != nullptr;
+		const bool bUseHairLighting = HairStrands::HasViewHairStrandsData(View) && HairStrands::HasViewHairStrandsVoxelData(View) && !LightSceneProxy->CastsHairStrandsDeepShadow();
 
 		FOcclusionRGS::FParameters* PassParameters = GraphBuilder.AllocParameters<FOcclusionRGS::FParameters>();
 		PassParameters->RWOcclusionMaskUAV = OutShadowMaskUAV;
@@ -298,10 +293,10 @@ void FDeferredShadingSceneRenderer::RenderRayTracingShadows(
 		if (bUseHairLighting)
 		{
 			const bool bUseHairVoxel = CVarRayTracingShadowsEnableHairVoxel.GetValueOnRenderThread() > 0;
-			PassParameters->bUseHairVoxel = (HairResources->bUseHairVoxel && bUseHairVoxel) ? 1 : 0;
-			PassParameters->HairCategorizationTexture = HairResources->CategorizationTexture;
-			PassParameters->HairLightChannelMaskTexture = HairResources->LightChannelMaskTexture;
-			PassParameters->VirtualVoxel = HairResources->VoxelResources->UniformBuffer;
+			PassParameters->bUseHairVoxel = bUseHairVoxel ? 1 : 0;
+			PassParameters->HairLightChannelMaskTexture = View.HairStrandsViewData.VisibilityData.LightChannelMaskTexture;
+			PassParameters->HairStrands = HairStrands::BindHairStrandsViewUniformParameters(View);
+			PassParameters->VirtualVoxel = HairStrands::BindHairStrandsVoxelUniformParameters(View);
 
 			if (ShaderDrawDebug::IsShaderDrawDebugEnabled(View))
 			{
diff --git a/Engine/Source/Runtime/Renderer/Private/RayTracing/RaytracingSkylight.cpp b/Engine/Source/Runtime/Renderer/Private/RayTracing/RaytracingSkylight.cpp
--- a/Engine/Source/Runtime/Renderer/Private/RayTracing/RaytracingSkylight.cpp
+++ b/Engine/Source/Runtime/Renderer/Private/RayTracing/RaytracingSkylight.cpp
@@ -502,11 +502,11 @@ void FDeferredShadingSceneRenderer::RenderRayTracingSkyLight(
 
 		const bool bUseHairLighting = 
 			HairStrands::HasViewHairStrandsData(View) && 
-			View.HairStrandsViewData.VirtualVoxelResources.IsValid() &&
+			HairStrands::HasViewHairStrandsVoxelData(View) &&
 			CVarRayTracingSkyLightEnableHairVoxel.GetValueOnRenderThread() > 0;
 		if (bUseHairLighting)
 		{
-			PassParameters->VirtualVoxel = View.HairStrandsViewData.VirtualVoxelResources.UniformBuffer;
+			PassParameters->VirtualVoxel = HairStrands::BindHairStrandsVoxelUniformParameters(View);
 		}
 
 		FRayTracingSkyLightRGS::FPermutationDomain PermutationVector;
