Cleanup hair transmittance and hair shadow mask shaders/passes.

diff --git a/Engine/Shaders/Private/HairStrands/HairStrandsDeepShadowCommon.ush b/Engine/Shaders/Private/HairStrands/HairStrandsDeepShadowCommon.ush
--- a/Engine/Shaders/Private/HairStrands/HairStrandsDeepShadowCommon.ush
+++ b/Engine/Shaders/Private/HairStrands/HairStrandsDeepShadowCommon.ush
@@ -158,7 +158,6 @@ float SampleDOM_PCSS(
 	float4 InLayerDepths,
 	Texture2D<float> FrontDepthTexture,
 	Texture2D<float4> DomTexture,
-	SamplerState LinearSampler,
 	float ApexAngleInDegree,
 	float2 Jitter) // [0,1]
 {
diff --git a/Engine/Shaders/Private/HairStrands/HairStrandsDeepShadowMask.usf b/Engine/Shaders/Private/HairStrands/HairStrandsDeepShadowMask.usf
--- a/Engine/Shaders/Private/HairStrands/HairStrandsDeepShadowMask.usf
+++ b/Engine/Shaders/Private/HairStrands/HairStrandsDeepShadowMask.usf
@@ -1,5 +1,6 @@
 // Copyright Epic Games, Inc. All Rights Reserved.
 
+#include "HairStrandsVisibilityCommonStruct.ush"
 #include "HairStrandsVoxelPageCommonStruct.ush"
 #include "HairStrandsDeepShadowCommonStruct.ush"
 
@@ -15,45 +16,6 @@
 #include "HairStrandsDeepShadowCommon.ush"
 #include "HairStrandsDeepTransmittanceCommon.ush"
 
-#define OPAQUE_MASK_TYPE_DEEP_SHADOW 0
-#define OPAQUE_MASK_TYPE_VIRTUAL_VOXEL 1
-
-#define VOXEL_TRAVERSAL_DEBUG 0
-#if PERMUTATION_OPAQUEMASK_TYPE == OPAQUE_MASK_TYPE_VIRTUAL_VOXEL
-	#define VOXEL_TRAVERSAL_TYPE VOXEL_TRAVERSAL_LINEAR_MIPMAP
-	#include "HairStrandsVoxelPageTraversal.ush"
-#endif
-
-StructuredBuffer<FDeepShadowTransform> DeepShadow_WorldToLightTransformBuffer;
-uint				DeepShadow_bIsGPUDriven;
-
-float4x4			DeepShadow_CPUWorldToLightTransform;
-uint				bIsWholeSceneLight;
-int2				DeepShadow_SlotResolution;
-int2				DeepShadow_SlotOffset;
-uint				DeepShadow_SlotIndex;
-uint				DeepShadow_KernelType;
-float				DeepShadow_DepthBiasScale;
-float				DeepShadow_DensityScale;
-float4				DeepShadow_LayerDepths;
-
-float4				Voxel_LightPosition;
-float3				Voxel_LightDirection;
-uint				Voxel_MacroGroupId;
-uint				Voxel_RandomType;
-
-Texture2D<float4>	RayMarchMaskTexture;
-
-Texture2D<float>	DeepShadow_FrontDepthTexture;
-Texture2D<float4>	DeepShadow_DomTexture;
-
-//Texture2D<float>	SceneDepthTexture;
-Texture2D<uint4>	CategorizationTexture;
-RWTexture2D<float4>	OutputColor;
-
-SamplerState		LinearSampler;	
-SamplerState 		ShadowSampler;
-
 float4 WriteShadowMaskOutput(const float InFadedShadow)
 {
 	// Depending of if the light is direction (i.e., WholeSceneLight) or a local light (Point/Spot/Rect), 
@@ -62,51 +24,8 @@
 	return float4(EncodedShadow, EncodedShadow, EncodedShadow, EncodedShadow);
 }
 
-/////////////////////////////////////////////////////////////////////////////////////////////////////////////
-// Random functions
+////////////////////////////////////////////////////////////////////////////////////////////////
+// Deep shadow 
 
-#if PERMUTATION_OPAQUEMASK_TYPE == OPAQUE_MASK_TYPE_VIRTUAL_VOXEL
-// The hair shadow mask is rendered by tracing ray/cone through the voxel structure. The traversing is done with 
-// point sampling and jittering to get the smooth aspect, and avoid expensive trilinear sampling in a sparse 
-// structure. These random functions make various tradeoff between randomness and correlation, so that the output 
-// can coverge with TAA without any denoiser. They result from experimentation, and don't have a good grounding
-float4 ComputeRandom4_0(uint2 PixelPosition, uint InSeed)
-{
-	const float2 U0 = InterleavedGradientNoise(PixelPosition, InSeed);
-	const float2 U1 = InterleavedGradientNoise(PixelPosition + 17, InSeed);
-	return float4(U0.x, U0.y, U1.x, U1.y);
-}
-
-float4 ComputeRandom4_1(uint2 PixelPosition, uint InSeed)
-{
-	const uint2 Seed0 = Rand3DPCG16(int3(PixelPosition, InSeed)).xy;
-	const uint2 Seed1 = Rand3DPCG16(int3(PixelPosition + 17, InSeed)).xy;
-	return float4(
-		Hammersley16(InSeed, 8, Seed0),
-		Hammersley16(InSeed, 8, Seed1));
-}
-
-float3 GetShadowMaskRandom(uint2 PixelPosition)
-{
-	float3 Random = 0.5f;
-	if (Voxel_RandomType == 0)
-	{
-		Random = GetHairVoxelJitter(PixelPosition.xy, View.StateFrameIndexMod8);
-	}
-	else if (Voxel_RandomType == 1)
-	{
-		Random = ComputeRandom4_0(PixelPosition, View.StateFrameIndexMod8).xyz;
-	}
-	else if (Voxel_RandomType == 2)
-	{
-		Random = ComputeRandom4_1(PixelPosition, View.StateFrameIndexMod8).xyz;
-	}
-	return Random;
-}
-#endif
-
-/////////////////////////////////////////////////////////////////////////////////////////////////////////////
-// 
-
 #define KERNEL_TYPE_2x2 0 
 #define KERNEL_TYPE_4x4 1
@@ -123,7 +42,25 @@
 	return WorldPosition;
 }
 
-#if PERMUTATION_OPAQUEMASK_TYPE == OPAQUE_MASK_TYPE_DEEP_SHADOW
+#if SHADER_SHADOWMASK_DEEPSHADOW
 
+StructuredBuffer<FDeepShadowTransform> DeepShadow_WorldToLightTransformBuffer;
+uint				DeepShadow_bIsGPUDriven;
+
+float4x4			DeepShadow_CPUWorldToLightTransform;
+int2				DeepShadow_SlotResolution;
+int2				DeepShadow_SlotOffset;
+uint				DeepShadow_SlotIndex;
+uint				DeepShadow_KernelType;
+float				DeepShadow_DepthBiasScale;
+float				DeepShadow_DensityScale;
+float4				DeepShadow_LayerDepths;
+
+Texture2D<float>	DeepShadow_FrontDepthTexture;
+Texture2D<float4>	DeepShadow_DomTexture;
+
+SamplerState		LinearSampler;
+SamplerState 		ShadowSampler;
+
 // Return the light space position of a world space.
 // XY: UV of the shadow space
@@ -407,7 +344,7 @@
 	const float2 UV = Input.UV;
 	const uint2 PixelCoord = floor(Input.Position.xy);
 	const float SceneDepth = ConvertFromDeviceZ(SceneDepthTexture.Load(uint3(PixelCoord, 0)));
-	const FCategorizationData Categorization = DecodeCategorizationData(CategorizationTexture.Load(uint3(PixelCoord, 0)));
+	const FCategorizationData Categorization = DecodeCategorizationData(HairStrands.HairCategorizationTexture.Load(uint3(PixelCoord, 0)));
 	const bool bIsHairPixel = Categorization.PixelCoverage >= 1;
 
 	const float3 WorldPosition		= TotWorldPosition(UV, SceneDepth);
@@ -453,9 +390,64 @@
 	OutColor = WriteShadowMaskOutput(Visibility);
 
 }
+#endif // SHADER_SHADOWMASK_DEEPSHADOW
+
+////////////////////////////////////////////////////////////////////////////////////////////////
+// Voxel volume
+
+#if SHADER_SHADOWMASK_VOXEL
+
+#if SHADER_SHADOWMASK_VOXEL
+#define VOXEL_TRAVERSAL_DEBUG 0
+#define VOXEL_TRAVERSAL_TYPE VOXEL_TRAVERSAL_LINEAR_MIPMAP
+#include "HairStrandsVoxelPageTraversal.ush"
 #endif
 
-#if PERMUTATION_OPAQUEMASK_TYPE == OPAQUE_MASK_TYPE_VIRTUAL_VOXEL
+float4				Voxel_LightPosition;
+float3				Voxel_LightDirection;
+uint				Voxel_MacroGroupId;
+uint				Voxel_RandomType;
+Texture2D<float4>	RayMarchMaskTexture;
+Texture2D<float>	DeepShadow_FrontDepthTexture;
+Texture2D<float4>	DeepShadow_DomTexture;
 
+// The hair shadow mask is rendered by tracing ray/cone through the voxel structure. The traversing is done with 
+// point sampling and jittering to get the smooth aspect, and avoid expensive trilinear sampling in a sparse 
+// structure. These random functions make various tradeoff between randomness and correlation, so that the output 
+// can coverge with TAA without any denoiser. They result from experimentation, and don't have a good grounding
+float4 ComputeRandom4_0(uint2 PixelPosition, uint InSeed)
+{
+	const float2 U0 = InterleavedGradientNoise(PixelPosition, InSeed);
+	const float2 U1 = InterleavedGradientNoise(PixelPosition + 17, InSeed);
+	return float4(U0.x, U0.y, U1.x, U1.y);
+}
+
+float4 ComputeRandom4_1(uint2 PixelPosition, uint InSeed)
+{
+	const uint2 Seed0 = Rand3DPCG16(int3(PixelPosition, InSeed)).xy;
+	const uint2 Seed1 = Rand3DPCG16(int3(PixelPosition + 17, InSeed)).xy;
+	return float4(
+		Hammersley16(InSeed, 8, Seed0),
+		Hammersley16(InSeed, 8, Seed1));
+}
+
+float3 GetShadowMaskRandom(uint2 PixelPosition)
+{
+	float3 Random = 0.5f;
+	if (Voxel_RandomType == 0)
+	{
+		Random = GetHairVoxelJitter(PixelPosition.xy, View.StateFrameIndexMod8);
+	}
+	else if (Voxel_RandomType == 1)
+	{
+		Random = ComputeRandom4_0(PixelPosition, View.StateFrameIndexMod8).xyz;
+	}
+	else if (Voxel_RandomType == 2)
+	{
+		Random = ComputeRandom4_1(PixelPosition, View.StateFrameIndexMod8).xyz;
+	}
+	return Random;
+}
+
 void MainPS(
 	FScreenVertexOutput Input,
@@ -465,7 +457,7 @@
 	const uint2 PixelCoord = floor(Input.Position.xy);
 	const float SceneDepth = ConvertFromDeviceZ(SceneDepthTexture.Load(uint3(PixelCoord, 0)));
 
-	const FCategorizationData Categorization = DecodeCategorizationData(CategorizationTexture.Load(uint3(PixelCoord, 0)));
+	const FCategorizationData Categorization = DecodeCategorizationData(HairStrands.HairCategorizationTexture.Load(uint3(PixelCoord, 0)));
 	const bool bIsHairPixel = Categorization.PixelCoverage >= 1;
 
 	float3 WorldPosition = TotWorldPosition(UV, SceneDepth);
@@ -528,7 +520,6 @@
 				VirtualVoxel.PageIndexOccupancyBuffer,
 				VirtualVoxel.PageTexture,
 				TraversalSettings);
-
 		}
 		else
 		{
@@ -540,4 +531,4 @@
 	OutColor = WriteShadowMaskOutput(Visibility);
 }
 
-#endif
+#endif // SHADER_SHADOWMASK_VOXEL
diff --git a/Engine/Shaders/Private/HairStrands/HairStrandsDeepTransmittanceMask.usf b/Engine/Shaders/Private/HairStrands/HairStrandsDeepTransmittanceMask.usf
--- a/Engine/Shaders/Private/HairStrands/HairStrandsDeepTransmittanceMask.usf
+++ b/Engine/Shaders/Private/HairStrands/HairStrandsDeepTransmittanceMask.usf
@@ -15,26 +15,18 @@
 #include "HairStrandsVisibilityCommon.ush"
 #include "HairStrandsVoxelPageCommon.ush"
 
-#define PERMUTATION_TRANSMITTANCE_DEEP_SHADOW 0
-#define PERMUTATION_TRANSMITTANCE_VIRTUAL_VOXEL 1
-
-
-#if SHADER_TRANSMITTANCE
-
-#if PERMUTATION_TRANSMITTANCE_TYPE == PERMUTATION_TRANSMITTANCE_VIRTUAL_VOXEL
-	#if PERMUTATION_TRAVERSAL == 1
-		#define VOXEL_TRAVERSAL_TYPE VOXEL_TRAVERSAL_LINEAR_MIPMAP
+// Need to match FDeepShadowData::MaxMacroGroupCount in HairRendering.h
+#if SHADER_TRANSMITTANCE_DEEPSHADOW || SHADER_TRANSMITTANCE_VOXEL
+	#define MAX_HAIR_MACROGROUP_COUNT 16 
+	#if PERMUTATION_GROUP_SIZE == 64
+		#define TRANSMITTANCE_GROUP_SIZE 64
+	#elif PERMUTATION_GROUP_SIZE == 32
+		#define TRANSMITTANCE_GROUP_SIZE 32
 	#else
-		#define VOXEL_TRAVERSAL_TYPE VOXEL_TRAVERSAL_LINEAR
-	#endif
-	#define VOXEL_TRAVERSAL_DEBUG 0
-	#include "HairStrandsVoxelPageTraversal.ush"
-	#if VOXEL_TRAVERSAL_DEBUG
-	#include "../ShaderPrintCommon.ush"
+		#error Unknown group size	
 	#endif
 #endif
 
-
 float3 ToWorldPosition(float2 UV, float Depth)
 {
 	float2 ScreenPosition = (UV - View.ScreenPositionScaleBias.wz) / View.ScreenPositionScaleBias.xy;
@@ -44,9 +36,10 @@ float3 ToWorldPosition(float2 UV, float Depth)
 	return WorldPosition;
 }
 
+////////////////////////////////////////////////////////////////////////////////////////////////
+// Deep shadow
 
-// Need to match FDeepShadowData::MaxMacroGroupCount in HairRendering.h
-#define MAX_HAIR_MACROGROUP_COUNT 16 
+#if SHADER_TRANSMITTANCE_DEEPSHADOW
 
 StructuredBuffer<FDeepShadowTransform> DeepShadow_WorldToLightTransformBuffer;
 uint				DeepShadow_bIsGPUDriven;
@@ -67,23 +60,8 @@ int4				DeepShadow_AtlasSlotOffsets_AtlasSlotIndex[MAX_HAIR_MACROGROUP_COUNT];
 Texture2D<float>	DeepShadow_FrontDepthTexture;
 Texture2D<float4>	DeepShadow_DomTexture;
 
-Texture2D<float4> 	RayMarchMaskTexture;
-
-Buffer<uint4>				IndirectArgsBuffer;
-
-Texture3D<float4>			HairLUTTexture;
-SamplerState				LinearSampler;
-RWBuffer<uint>				OutputColor;
-
-#if PERMUTATION_GROUP_SIZE == 0
-	#define TRANSMITTANCE_GROUP_SIZE 64
-#elif PERMUTATION_GROUP_SIZE == 1
-	#define TRANSMITTANCE_GROUP_SIZE 32
-#else
-	#error Unknown group size	
-#endif
-
-#if PERMUTATION_TRANSMITTANCE_TYPE == PERMUTATION_TRANSMITTANCE_DEEP_SHADOW
+Buffer<uint4>		IndirectArgsBuffer;
+RWBuffer<uint>		OutTransmittanceMask;
 
 #if 0
 // Sample deep shadow with PCF kernel and doing transmittance evaluation for each sample
@@ -194,7 +172,7 @@ FHairTransmittanceMask ComputeTransmittance(
 		// Compute the number of hair count between light & shading point
 		if (Settings.TransmittanceKernelType == 4)
 		{
-			//Out = SampleDOM_PCF_Accurate(LightSpacePosition.xyz, DepthBias, InLayerDepths, FrontDepthTexture, DomTexture, GBuffer, SinLightAngle, HairLUTTexture, LinearSampler);
+			//Out = SampleDOM_PCF_Accurate(LightSpacePosition.xyz, DepthBias, InLayerDepths, FrontDepthTexture, DomTexture, GBuffer, SinLightAngle, HairLUTTexture);
 			Out.HairCount = 0;
 			Out.Visibility = 1;
 		}
@@ -203,7 +181,7 @@ FHairTransmittanceMask ComputeTransmittance(
 			float HairCount = 0;
 			if (Settings.TransmittanceKernelType == 3)
 			{
-				HairCount = SampleDOM_PCSS(LightSpacePosition.xyz, Settings.DeepShadowResolution, DepthBias, InLayerDepths, FrontDepthTexture, DomTexture, LinearSampler, Settings.TransmittanceKernelApertureInDegree, Settings.Random.xy);
+				HairCount = SampleDOM_PCSS(LightSpacePosition.xyz, Settings.DeepShadowResolution, DepthBias, InLayerDepths, FrontDepthTexture, DomTexture, Settings.TransmittanceKernelApertureInDegree, Settings.Random.xy);
 			}
 			else if (Settings.TransmittanceKernelType == 2)
 			{
@@ -251,7 +229,7 @@ void MainCS(uint2 DispatchThreadId : SV_DispatchThreadID)
 
 		if ((LightChannelMask & Sample.LightChannelMask) == 0)
 		{
-			OutputColor[LocalOffset] = InitNullPackedHairTransmittanceMask();
+			OutTransmittanceMask[LocalOffset] = InitNullPackedHairTransmittanceMask();
 			return;
 		}
 
@@ -284,14 +262,27 @@ void MainCS(uint2 DispatchThreadId : SV_DispatchThreadID)
 		const float3 WorldPosition = ToWorldPosition(UV, SceneDepth);
 		const FHairTransmittanceMask TransmittanceMask = ComputeTransmittance(WorldPosition, Settings, DeepShadow_LayerDepths, DeepShadow_FrontDepthTexture, DeepShadow_DomTexture);
 
-		OutputColor[LocalOffset] = PackTransmittanceMask(TransmittanceMask);
+		OutTransmittanceMask[LocalOffset] = PackTransmittanceMask(TransmittanceMask);
 	}
 }
+#endif // #if SHADER_TRANSMITTANCE_DEEPSHADOW
+
+////////////////////////////////////////////////////////////////////////////////////////////////
+// Voxel volume
+
+#if SHADER_TRANSMITTANCE_VOXEL
+
+#if PERMUTATION_TRAVERSAL == 1
+	#define VOXEL_TRAVERSAL_TYPE VOXEL_TRAVERSAL_LINEAR_MIPMAP
+#else
+	#define VOXEL_TRAVERSAL_TYPE VOXEL_TRAVERSAL_LINEAR
 #endif
 
-#if PERMUTATION_TRANSMITTANCE_TYPE == PERMUTATION_TRANSMITTANCE_VIRTUAL_VOXEL
+#define VOXEL_TRAVERSAL_DEBUG 0
+#include "HairStrandsVoxelPageTraversal.ush"
 
 #if VOXEL_TRAVERSAL_DEBUG
+#include "../ShaderPrintCommon.ush"
 void PrintInfo(uint Index, uint2 PixelCoord, FHairTransmittanceMask TransmittanceMask)
 {
 	#if 0
@@ -333,7 +324,15 @@ void PrintInfo(uint Index, uint2 PixelCoord, FHairTransmittanceMask Transmittanc
 	Pos = ShaderPrintNewline(Pos);
 	#endif
 }
-#endif
+#endif // VOXEL_TRAVERSAL_DEBUG
+
+float3				LightDirection;
+uint				LightChannelMask;
+float4				LightPosition;
+float				LightRadius;
+Texture2D<float4> 	RayMarchMaskTexture;
+Buffer<uint4>		IndirectArgsBuffer;
+RWBuffer<uint>		OutTransmittanceMask;
 
 FHairTransmittanceMask ComputeTransmittanceVoxel(
 	const uint2 PixelCoord,
@@ -432,7 +431,7 @@ void MainCS(uint2 DispatchThreadId : SV_DispatchThreadID)
 	{
 		// This pixel is already fully shadowed opaque. (RayMarchMaskTexture is computed using the closest sample as of today)
 		const FHairTransmittanceMask DefaultMask = InitHairTransmittanceMask();
-		OutputColor[LocalOffset] = PackTransmittanceMask(DefaultMask);
+		OutTransmittanceMask[LocalOffset] = PackTransmittanceMask(DefaultMask);
 		return;
 	}
 
@@ -447,7 +446,7 @@ void MainCS(uint2 DispatchThreadId : SV_DispatchThreadID)
 
 		if ((LightChannelMask & Sample.LightChannelMask) == 0)
 		{
-			OutputColor[LocalOffset] = InitNullPackedHairTransmittanceMask();
+			OutTransmittanceMask[LocalOffset] = InitNullPackedHairTransmittanceMask();
 			return;
 		}
 
@@ -481,12 +480,11 @@ void MainCS(uint2 DispatchThreadId : SV_DispatchThreadID)
 		}
 		#endif	
 
-		OutputColor[LocalOffset] = PackTransmittanceMask(Out);
+		OutTransmittanceMask[LocalOffset] = PackTransmittanceMask(Out);
 	}
 }
-#endif
 
-#endif // SHADER_TRANSMITTANCE
+#endif // SHADER_TRANSMITTANCE_VOXEL
 
 ////////////////////////////////////////////////////////////////////////////////////////////////
 
diff --git a/Engine/Source/Runtime/Renderer/Private/HairStrands/HairStrandsTransmittance.cpp b/Engine/Source/Runtime/Renderer/Private/HairStrands/HairStrandsTransmittance.cpp
--- a/Engine/Source/Runtime/Renderer/Private/HairStrands/HairStrandsTransmittance.cpp
+++ b/Engine/Source/Runtime/Renderer/Private/HairStrands/HairStrandsTransmittance.cpp
@@ -83,6 +83,14 @@ FVector4 ComputeDeepShadowLayerDepths(float LayerDistribution)
 	return Depths;
 }
 
+struct FHairStrandsTransmittanceLightParams
+{
+	FVector LightDirection = FVector::ZeroVector;
+	FVector4 LightPosition = FVector4(0, 0, 0, 0);
+	uint32 LightChannelMask = 0;
+	float LightRadius = 0;
+};
+
 ///////////////////////////////////////////////////////////////////////////////////////////////////
 // Clear transmittance Mask
 
@@ -150,30 +158,120 @@ FHairStrandsTransmittanceMaskData CreateDummyHairStrandsTransmittanceMaskData(FR
 }
 
 ///////////////////////////////////////////////////////////////////////////////////////////////////
-// Transmittance Mask from deep shadow
-
-enum FHairTransmittanceType
-{
-	FHairTransmittanceType_DeepShadow,
-	FHairTransmittanceType_VirtualVoxel,
-	FHairTransmittanceTypeCount
-};
+// Transmittance Mask from voxel
 
-class FDeepTransmittanceMaskCS : public FGlobalShader
+class FHairStrandsVoxelTransmittanceMaskCS : public FGlobalShader
 {
-	DECLARE_GLOBAL_SHADER(FDeepTransmittanceMaskCS);
-	SHADER_USE_PARAMETER_STRUCT(FDeepTransmittanceMaskCS, FGlobalShader);
+	DECLARE_GLOBAL_SHADER(FHairStrandsVoxelTransmittanceMaskCS);
+	SHADER_USE_PARAMETER_STRUCT(FHairStrandsVoxelTransmittanceMaskCS, FGlobalShader);
 
-	class FTransmittanceType : SHADER_PERMUTATION_INT("PERMUTATION_TRANSMITTANCE_TYPE", FHairTransmittanceTypeCount);
-	class FTransmittanceGroupSize : SHADER_PERMUTATION_INT("PERMUTATION_GROUP_SIZE", 2);
+	class FTransmittanceGroupSize : SHADER_PERMUTATION_SPARSE_INT("PERMUTATION_GROUP_SIZE", 32, 64);
 	class FSuperSampling : SHADER_PERMUTATION_INT("PERMUTATION_SUPERSAMPLING", 2);
 	class FTraversal : SHADER_PERMUTATION_INT("PERMUTATION_TRAVERSAL", 2);
-	using FPermutationDomain = TShaderPermutationDomain<FTransmittanceType, FTransmittanceGroupSize, FSuperSampling, FTraversal>;
+	using FPermutationDomain = TShaderPermutationDomain<FTransmittanceGroupSize, FSuperSampling, FTraversal>;
 
 	BEGIN_SHADER_PARAMETER_STRUCT(FParameters, )
 		SHADER_PARAMETER_STRUCT_INCLUDE(FSceneTextureParameters, SceneTextures)
 		SHADER_PARAMETER_STRUCT_INCLUDE(ShaderDrawDebug::FShaderDrawDebugParameters, ShaderDrawParameters)
 		SHADER_PARAMETER_RDG_UNIFORM_BUFFER(FHairStrandsViewUniformParameters, HairStrands)
+		SHADER_PARAMETER_RDG_UNIFORM_BUFFER(FVirtualVoxelParameters, VirtualVoxel)
+
+		SHADER_PARAMETER(float, LightRadius)
+		SHADER_PARAMETER(FVector, LightDirection)
+		SHADER_PARAMETER(FVector4, LightPosition)
+		SHADER_PARAMETER(uint32, LightChannelMask)
+
+		SHADER_PARAMETER_RDG_TEXTURE(Texture2D, RayMarchMaskTexture)
+		SHADER_PARAMETER_RDG_BUFFER_UAV(RWBuffer, OutTransmittanceMask)
+		SHADER_PARAMETER_STRUCT_REF(FViewUniformShaderParameters, ViewUniformBuffer)
+
+		RDG_BUFFER_ACCESS(IndirectArgsBuffer, ERHIAccess::IndirectArgs)
+
+	END_SHADER_PARAMETER_STRUCT()
+
+public:
+	static bool ShouldCompilePermutation(const FGlobalShaderPermutationParameters& Parameters)
+	{
+		return IsHairStrandsSupported(EHairStrandsShaderType::Strands, Parameters.Platform);
+	}
+
+	static void ModifyCompilationEnvironment(const FGlobalShaderPermutationParameters& Parameters, FShaderCompilerEnvironment& OutEnvironment)
+	{
+		FGlobalShader::ModifyCompilationEnvironment(Parameters, OutEnvironment);
+		OutEnvironment.SetDefine(TEXT("SHADER_TRANSMITTANCE_VOXEL"), 1);
+	}
+};
+
+IMPLEMENT_GLOBAL_SHADER(FHairStrandsVoxelTransmittanceMaskCS, "/Engine/Private/HairStrands/HairStrandsDeepTransmittanceMask.usf", "MainCS", SF_Compute);
+
+// Transmittance mask using voxel volume
+static FRDGBufferRef AddHairStrandsVoxelTransmittanceMaskPass(
+	FRDGBuilder& GraphBuilder,
+	const FSceneTextureParameters& SceneTextures,
+	const FViewInfo& View,
+	const FHairStrandsTransmittanceLightParams& Params,
+	const uint32 NodeGroupSize,
+	FRDGBufferRef IndirectArgsBuffer,
+	FRDGTextureRef ScreenShadowMaskSubPixelTexture)
+{
+	check(HairStrands::HasViewHairStrandsVoxelData(View));
+	check(NodeGroupSize == 64 || NodeGroupSize == 32);
+
+	FRDGBufferRef OutBuffer = CreateHairStrandsTransmittanceMaskBuffer(GraphBuilder, View.HairStrandsViewData.VisibilityData.MaxNodeCount);
+
+	FHairStrandsVoxelTransmittanceMaskCS::FParameters* Parameters = GraphBuilder.AllocParameters<FHairStrandsVoxelTransmittanceMaskCS::FParameters>();
+	Parameters->ViewUniformBuffer = View.ViewUniformBuffer;
+	Parameters->SceneTextures = SceneTextures;
+	Parameters->OutTransmittanceMask = GraphBuilder.CreateUAV(OutBuffer, FHairStrandsTransmittanceMaskData::Format);
+	Parameters->LightChannelMask = Params.LightChannelMask;
+	Parameters->LightDirection = Params.LightDirection;
+	Parameters->LightPosition = Params.LightPosition;
+	Parameters->LightRadius = Params.LightRadius;
+	Parameters->IndirectArgsBuffer = IndirectArgsBuffer;
+	Parameters->HairStrands = HairStrands::BindHairStrandsViewUniformParameters(View);
+	Parameters->RayMarchMaskTexture = ScreenShadowMaskSubPixelTexture ? ScreenShadowMaskSubPixelTexture : GSystemTextures.GetWhiteDummy(GraphBuilder);
+	Parameters->VirtualVoxel = HairStrands::BindHairStrandsVoxelUniformParameters(View);
+
+	if (ShaderDrawDebug::IsShaderDrawDebugEnabled(View))
+	{
+		ShaderDrawDebug::SetParameters(GraphBuilder, View.ShaderDrawData, Parameters->ShaderDrawParameters);
+	}
+
+	const bool bIsSuperSampled = GHairStrandsTransmittanceSuperSampling > 0;
+	const bool bIsMipTraversal = GHairStrandsTransmittanceMaskUseMipTraversal > 0;
+
+	FHairStrandsVoxelTransmittanceMaskCS::FPermutationDomain PermutationVector;
+	PermutationVector.Set<FHairStrandsVoxelTransmittanceMaskCS::FTransmittanceGroupSize>(NodeGroupSize);
+	PermutationVector.Set<FHairStrandsVoxelTransmittanceMaskCS::FSuperSampling>(bIsSuperSampled ? 1 : 0);
+	PermutationVector.Set<FHairStrandsVoxelTransmittanceMaskCS::FTraversal>(bIsMipTraversal ? 1 : 0);
+
+	TShaderMapRef<FHairStrandsVoxelTransmittanceMaskCS> ComputeShader(View.ShaderMap, PermutationVector);
+	FComputeShaderUtils::AddPass(
+		GraphBuilder,
+		RDG_EVENT_NAME("HairStrandsTransmittanceMask(Voxel)"),
+		ComputeShader,
+		Parameters,
+		IndirectArgsBuffer,
+		0);
+
+	return OutBuffer;
+}
+
+///////////////////////////////////////////////////////////////////////////////////////////////////
+// Transmittance Mask from deep shadow
+
+class FHairStrandsDeepShadowTransmittanceMaskCS : public FGlobalShader
+{
+	DECLARE_GLOBAL_SHADER(FHairStrandsDeepShadowTransmittanceMaskCS);
+	SHADER_USE_PARAMETER_STRUCT(FHairStrandsDeepShadowTransmittanceMaskCS, FGlobalShader);
+
+	class FTransmittanceGroupSize : SHADER_PERMUTATION_SPARSE_INT("PERMUTATION_GROUP_SIZE", 32, 64);
+	using FPermutationDomain = TShaderPermutationDomain<FTransmittanceGroupSize>;
+
+	BEGIN_SHADER_PARAMETER_STRUCT(FParameters, )
+		SHADER_PARAMETER_STRUCT_INCLUDE(FSceneTextureParameters, SceneTextures)
+		SHADER_PARAMETER_RDG_UNIFORM_BUFFER(FHairStrandsViewUniformParameters, HairStrands)
+		SHADER_PARAMETER_STRUCT_REF(FViewUniformShaderParameters, ViewUniformBuffer)
 
 		SHADER_PARAMETER_ARRAY(FIntVector4, DeepShadow_AtlasSlotOffsets_AtlasSlotIndex, [FHairStrandsDeepShadowData::MaxMacroGroupCount])
 		SHADER_PARAMETER_ARRAY(FMatrix, DeepShadow_CPUWorldToLightTransforms, [FHairStrandsDeepShadowData::MaxMacroGroupCount])
@@ -192,55 +290,31 @@ class FDeepTransmittanceMaskCS : public FGlobalShader
 		SHADER_PARAMETER(uint32, DeepShadow_bIsGPUDriven)
 
 		SHADER_PARAMETER_RDG_TEXTURE(Texture2D, RayMarchMaskTexture)
-
 		SHADER_PARAMETER_RDG_TEXTURE(Texture2D, DeepShadow_FrontDepthTexture)
 		SHADER_PARAMETER_RDG_TEXTURE(Texture2D, DeepShadow_DomTexture)
 		SHADER_PARAMETER_RDG_BUFFER_SRV(StructuredBuffer, DeepShadow_WorldToLightTransformBuffer)
+		SHADER_PARAMETER_RDG_BUFFER_UAV(RWBuffer, OutTransmittanceMask)
 
-		SHADER_PARAMETER_RDG_TEXTURE(Texture3D, HairLUTTexture)
 		RDG_BUFFER_ACCESS(IndirectArgsBuffer, ERHIAccess::IndirectArgs)
 
-		SHADER_PARAMETER_RDG_BUFFER_UAV(RWBuffer, OutputColor)
-		SHADER_PARAMETER_SAMPLER(SamplerState, LinearSampler)
-		SHADER_PARAMETER_SAMPLER(SamplerState, ShadowSampler)
-
-		SHADER_PARAMETER_STRUCT_REF(FViewUniformShaderParameters, ViewUniformBuffer)
-		SHADER_PARAMETER_RDG_UNIFORM_BUFFER(FVirtualVoxelParameters, VirtualVoxel)
 	END_SHADER_PARAMETER_STRUCT()
 
 public:
 	static bool ShouldCompilePermutation(const FGlobalShaderPermutationParameters& Parameters) 
 	{ 
-		FPermutationDomain PermutationVector(Parameters.PermutationId);
-		if (PermutationVector.Get<FTransmittanceType>() == FHairTransmittanceType_DeepShadow && 
-			(PermutationVector.Get<FSuperSampling>() == 1 || PermutationVector.Get<FTraversal>() == 1))
-		{
-			return false;
-		}
-
 		return IsHairStrandsSupported(EHairStrandsShaderType::Strands, Parameters.Platform);
 	}
 
-	static FPermutationDomain RemapPermutation(FPermutationDomain PermutationVector)
-	{
-		if (PermutationVector.Get<FTransmittanceType>() == FHairTransmittanceType_DeepShadow)
-		{
-			PermutationVector.Set<FSuperSampling>(0);
-			PermutationVector.Set<FTraversal>(0);
-		}
-		return PermutationVector;
-	}
-
 	static void ModifyCompilationEnvironment(const FGlobalShaderPermutationParameters& Parameters, FShaderCompilerEnvironment& OutEnvironment)
 	{
 		FGlobalShader::ModifyCompilationEnvironment(Parameters, OutEnvironment);
-		OutEnvironment.SetDefine(TEXT("SHADER_TRANSMITTANCE"), 1);
+		OutEnvironment.SetDefine(TEXT("SHADER_TRANSMITTANCE_DEEPSHADOW"), 1);
 	}
 };
 
-IMPLEMENT_GLOBAL_SHADER(FDeepTransmittanceMaskCS, "/Engine/Private/HairStrands/HairStrandsDeepTransmittanceMask.usf", "MainCS", SF_Compute);
+IMPLEMENT_GLOBAL_SHADER(FHairStrandsDeepShadowTransmittanceMaskCS, "/Engine/Private/HairStrands/HairStrandsDeepTransmittanceMask.usf", "MainCS", SF_Compute);
 
-struct FDeepShadowTransmittanceParams
+struct FHairStrandsDeepShadowTransmittanceLightParams : FHairStrandsTransmittanceLightParams
 {
 	FIntVector4 DeepShadow_AtlasSlotOffsets_AtlasSlotIndex[FHairStrandsDeepShadowData::MaxMacroGroupCount];
 	FMatrix DeepShadow_CPUWorldToLightTransforms[FHairStrandsDeepShadowData::MaxMacroGroupCount];
@@ -248,42 +322,32 @@ struct FDeepShadowTransmittanceParams
 	FIntPoint DeepShadow_Resolution = FIntPoint(0, 0);
 	bool DeepShadow_bIsGPUDriven = false;
 	FVector4 DeepShadow_LayerDepths = FVector4(0, 0, 0, 0);
-	FVector LightDirection = FVector::ZeroVector;
-	FVector4 LightPosition = FVector4(0, 0, 0, 0);
-	uint32 LightChannelMask = 0;
-	float LightRadius = 0;
 	float DeepShadow_DepthBiasScale = 0;
 	float DeepShadow_DensityScale = 0;
 	FMatrix DeepShadow_ShadowToWorld = FMatrix::Identity;
 	
 	FRDGTextureRef DeepShadow_FrontDepthTexture = nullptr;
 	FRDGTextureRef DeepShadow_DomTexture = nullptr;
-
-	const FHairStrandsVoxelResources* VirtualVoxelResources = nullptr;
 };
 
-// Transmittance mask
-static FRDGBufferRef AddDeepShadowTransmittanceMaskPass(
+// Transmittance mask using deep shadow
+static FRDGBufferRef AddHairStrandsDeepShadowTransmittanceMaskPass(
 	FRDGBuilder& GraphBuilder,
 	const FSceneTextureParameters& SceneTextures,
 	const FViewInfo& View,
-	const FHairTransmittanceType TransmittanceType,
-	const FDeepShadowTransmittanceParams& Params,
+	const FHairStrandsDeepShadowTransmittanceLightParams& Params,
 	const uint32 NodeGroupSize,
-	FRDGTextureRef HairLUTTexture,
 	FRDGBufferRef IndirectArgsBuffer,
 	FRDGTextureRef ScreenShadowMaskSubPixelTexture)
 {
 	FRDGBufferRef OutBuffer = CreateHairStrandsTransmittanceMaskBuffer(GraphBuilder, View.HairStrandsViewData.VisibilityData.MaxNodeCount);
 
-	FDeepTransmittanceMaskCS::FParameters* Parameters = GraphBuilder.AllocParameters<FDeepTransmittanceMaskCS::FParameters>();
+	FHairStrandsDeepShadowTransmittanceMaskCS::FParameters* Parameters = GraphBuilder.AllocParameters<FHairStrandsDeepShadowTransmittanceMaskCS::FParameters>();
 	Parameters->ViewUniformBuffer = View.ViewUniformBuffer;
 	Parameters->SceneTextures = SceneTextures;
-	Parameters->HairLUTTexture = HairLUTTexture;
 	Parameters->DeepShadow_FrontDepthTexture = Params.DeepShadow_FrontDepthTexture;
 	Parameters->DeepShadow_DomTexture = Params.DeepShadow_DomTexture;
-	Parameters->LinearSampler = TStaticSamplerState<SF_Bilinear, AM_Clamp, AM_Clamp, AM_Clamp>::GetRHI();
-	Parameters->OutputColor = GraphBuilder.CreateUAV(OutBuffer, FHairStrandsTransmittanceMaskData::Format);
+	Parameters->OutTransmittanceMask = GraphBuilder.CreateUAV(OutBuffer, FHairStrandsTransmittanceMaskData::Format);
 	Parameters->LightChannelMask = Params.LightChannelMask;
 	Parameters->DeepShadow_Resolution = Params.DeepShadow_Resolution;
 	Parameters->LightDirection = Params.LightDirection;
@@ -300,38 +364,19 @@ static FRDGBufferRef AddDeepShadowTransmittanceMaskPass(
 	Parameters->HairStrands = HairStrands::BindHairStrandsViewUniformParameters(View);
 	Parameters->DeepShadow_bIsGPUDriven = Params.DeepShadow_bIsGPUDriven ? 1 : 0;;
 	Parameters->DeepShadow_WorldToLightTransformBuffer = Params.DeepShadow_WorldToLightTransformBuffer;
-
-	if (ShaderDrawDebug::IsShaderDrawDebugEnabled(View))
-	{
-		ShaderDrawDebug::SetParameters(GraphBuilder, View.ShaderDrawData, Parameters->ShaderDrawParameters);
-	}
+	Parameters->RayMarchMaskTexture = ScreenShadowMaskSubPixelTexture ? ScreenShadowMaskSubPixelTexture : GraphBuilder.RegisterExternalTexture(GSystemTextures.WhiteDummy);
 
 	memcpy(&(Parameters->DeepShadow_AtlasSlotOffsets_AtlasSlotIndex[0]), Params.DeepShadow_AtlasSlotOffsets_AtlasSlotIndex, sizeof(FIntVector4) * FHairStrandsDeepShadowData::MaxMacroGroupCount);
 	memcpy(&(Parameters->DeepShadow_CPUWorldToLightTransforms[0]), Params.DeepShadow_CPUWorldToLightTransforms, sizeof(FMatrix) * FHairStrandsDeepShadowData::MaxMacroGroupCount);
 
-	Parameters->RayMarchMaskTexture = ScreenShadowMaskSubPixelTexture ? ScreenShadowMaskSubPixelTexture : GraphBuilder.RegisterExternalTexture(GSystemTextures.WhiteDummy);
-
-	bool bIsSuperSampled = false;
-	if (TransmittanceType == FHairTransmittanceType_VirtualVoxel)
-	{
-		check(Params.VirtualVoxelResources);
-		Parameters->VirtualVoxel = Params.VirtualVoxelResources->UniformBuffer;
-		bIsSuperSampled = GHairStrandsTransmittanceSuperSampling > 0;
-	}
-
-	const bool bIsMipTraversal = GHairStrandsTransmittanceMaskUseMipTraversal > 0;
 	check(NodeGroupSize == 64 || NodeGroupSize == 32);
-	FDeepTransmittanceMaskCS::FPermutationDomain PermutationVector;
-	PermutationVector.Set<FDeepTransmittanceMaskCS::FTransmittanceType>(TransmittanceType);
-	PermutationVector.Set<FDeepTransmittanceMaskCS::FTransmittanceGroupSize>(NodeGroupSize == 64 ? 0 : (NodeGroupSize == 32 ? 1 : 2));
-	PermutationVector.Set<FDeepTransmittanceMaskCS::FSuperSampling>(bIsSuperSampled ? 1 : 0);
-	PermutationVector.Set<FDeepTransmittanceMaskCS::FTraversal>(bIsMipTraversal ? 1 : 0);
-	PermutationVector = FDeepTransmittanceMaskCS::RemapPermutation(PermutationVector);
-
-	TShaderMapRef<FDeepTransmittanceMaskCS> ComputeShader(View.ShaderMap, PermutationVector);
+	FHairStrandsDeepShadowTransmittanceMaskCS::FPermutationDomain PermutationVector;
+	PermutationVector.Set<FHairStrandsDeepShadowTransmittanceMaskCS::FTransmittanceGroupSize>(NodeGroupSize);
+
+	TShaderMapRef<FHairStrandsDeepShadowTransmittanceMaskCS> ComputeShader(View.ShaderMap, PermutationVector);
 	FComputeShaderUtils::AddPass(
 		GraphBuilder,
-		RDG_EVENT_NAME("HairStrandsTransmittanceMask"),
+		RDG_EVENT_NAME("HairStrandsTransmittanceMask(DeepShadow)"),
 		ComputeShader,
 		Parameters,
 		IndirectArgsBuffer,
@@ -341,23 +386,150 @@ static FRDGBufferRef AddDeepShadowTransmittanceMaskPass(
 }
 
 ///////////////////////////////////////////////////////////////////////////////////////////////////
-// Opaque Mask from deep shadow
+// Opaque Mask from voxel volume
+
+class FHairStrandsVoxelShadowMaskPS : public FGlobalShader
+{
+	DECLARE_GLOBAL_SHADER(FHairStrandsVoxelShadowMaskPS);
+	SHADER_USE_PARAMETER_STRUCT(FHairStrandsVoxelShadowMaskPS, FGlobalShader);
+
+	BEGIN_SHADER_PARAMETER_STRUCT(FParameters, )
+		SHADER_PARAMETER_STRUCT_INCLUDE(ShaderDrawDebug::FShaderDrawDebugParameters, ShaderDrawParameters)
+		
+		SHADER_PARAMETER(FVector4, Voxel_LightPosition)
+		SHADER_PARAMETER(FVector, Voxel_LightDirection)
+		SHADER_PARAMETER(uint32, Voxel_MacroGroupId)
+		SHADER_PARAMETER(uint32, Voxel_RandomType)
+		SHADER_PARAMETER(uint32, bIsWholeSceneLight)
+
+		SHADER_PARAMETER_RDG_TEXTURE(Texture2D, RayMarchMaskTexture)
+		SHADER_PARAMETER_RDG_TEXTURE(Texture2D, SceneDepthTexture)
+		SHADER_PARAMETER_SAMPLER(SamplerState, LinearSampler)
+		SHADER_PARAMETER_SAMPLER(SamplerState, ShadowSampler)
+
+		SHADER_PARAMETER_STRUCT_REF(FViewUniformShaderParameters, ViewUniformBuffer)
+		SHADER_PARAMETER_RDG_UNIFORM_BUFFER(FVirtualVoxelParameters, VirtualVoxel)
+		SHADER_PARAMETER_RDG_UNIFORM_BUFFER(FHairStrandsViewUniformParameters, HairStrands)
+		RENDER_TARGET_BINDING_SLOTS()
+	END_SHADER_PARAMETER_STRUCT()
+
+public:
+	static bool ShouldCompilePermutation(const FGlobalShaderPermutationParameters& Parameters) { return IsHairStrandsSupported(EHairStrandsShaderType::Strands, Parameters.Platform); }
+	static void ModifyCompilationEnvironment(const FGlobalShaderPermutationParameters& Parameters, FShaderCompilerEnvironment& OutEnvironment)
+	{
+		FGlobalShader::ModifyCompilationEnvironment(Parameters, OutEnvironment);
+		OutEnvironment.SetRenderTargetOutputFormat(0, PF_B8G8R8A8);
+		OutEnvironment.SetDefine(TEXT("SHADER_SHADOWMASK_VOXEL"), 1);
+	}
+};
+
+IMPLEMENT_GLOBAL_SHADER(FHairStrandsVoxelShadowMaskPS, "/Engine/Private/HairStrands/HairStrandsDeepShadowMask.usf", "MainPS", SF_Pixel);
 
-enum FHairOpaqueMaskType
+struct FHairStrandsVoxelShadowParams
 {
-	FHairOpaqueMaskType_DeepShadow,
-	FHairOpaqueMaskType_VirtualVoxel,
-	FHairOpaqueMaskTypeCount
+	bool			bIsWholeSceneLight = false;
+	FVector			Voxel_LightDirection = FVector::ZeroVector;
+	FVector4		Voxel_LightPosition = FVector4(0, 0, 0, 0);
+	uint32			Voxel_MacroGroupId;
 };
 
-class FDeepShadowMaskPS : public FGlobalShader
+// Opaque mask from voxel
+static void AddHairStrandsVoxelShadowMaskPass(
+	FRDGBuilder& GraphBuilder,
+	FRDGTextureRef SceneDepthTexture,
+	const FViewInfo& View,
+	const FHairStrandsVoxelShadowParams& Params,
+	FRDGTextureRef& OutShadowMask)
 {
-	DECLARE_GLOBAL_SHADER(FDeepShadowMaskPS);
-	SHADER_USE_PARAMETER_STRUCT(FDeepShadowMaskPS, FGlobalShader);
+	check(OutShadowMask);
+	check(HairStrands::HasViewHairStrandsVoxelData(View));
+
+	FHairStrandsVoxelShadowMaskPS::FParameters* Parameters = GraphBuilder.AllocParameters<FHairStrandsVoxelShadowMaskPS::FParameters>();
+	Parameters->ViewUniformBuffer = View.ViewUniformBuffer;
+	Parameters->SceneDepthTexture = SceneDepthTexture;	
+	Parameters->HairStrands = HairStrands::BindHairStrandsViewUniformParameters(View);
+	Parameters->VirtualVoxel = HairStrands::BindHairStrandsVoxelUniformParameters(View);
+	Parameters->LinearSampler = TStaticSamplerState<SF_Bilinear, AM_Clamp, AM_Clamp, AM_Clamp>::GetRHI();
+	Parameters->ShadowSampler = TStaticSamplerState<SF_Point, AM_Clamp, AM_Clamp, AM_Clamp>::GetRHI();
+	Parameters->bIsWholeSceneLight = Params.bIsWholeSceneLight ? 1 : 0;
+	Parameters->Voxel_LightPosition		= Params.Voxel_LightPosition;
+	Parameters->Voxel_LightDirection	= Params.Voxel_LightDirection;
+	Parameters->Voxel_MacroGroupId		= Params.Voxel_MacroGroupId;
+	Parameters->Voxel_RandomType = FMath::Clamp(GHairStrandsShadowRandomTraversalType, 0, 2);	
+	Parameters->RenderTargets[0] = FRenderTargetBinding(OutShadowMask, ERenderTargetLoadAction::ELoad);
+
+	if (ShaderDrawDebug::IsShaderDrawDebugEnabled(View))
+	{
+		ShaderDrawDebug::SetParameters(GraphBuilder, View.ShaderDrawData, Parameters->ShaderDrawParameters);
+	}
+
+	FRDGTextureRef RayMarchMask = nullptr;
+	{
+		FRDGTextureDesc Desc = OutShadowMask->Desc;
+		Desc.Flags |= TexCreate_ShaderResource;
+		RayMarchMask = GraphBuilder.CreateTexture(Desc, TEXT("Hair.RayMarchMask"));
+		FRHICopyTextureInfo CopyInfo;
+		CopyInfo.Size = OutShadowMask->Desc.GetSize();
+		AddCopyTexturePass(GraphBuilder, OutShadowMask, RayMarchMask, CopyInfo);
+	}
+	Parameters->RayMarchMaskTexture = RayMarchMask;
+
+
+	FHairStrandsVoxelShadowMaskPS::FPermutationDomain PermutationVector;
+
+	const FIntPoint OutputResolution = SceneDepthTexture->Desc.Extent;
+	TShaderMapRef<FPostProcessVS> VertexShader(View.ShaderMap);
+	TShaderMapRef<FHairStrandsVoxelShadowMaskPS> PixelShader(View.ShaderMap, PermutationVector);
+	const FGlobalShaderMap* GlobalShaderMap = View.ShaderMap;
+	const FIntRect Viewport = View.ViewRect;
+
+	ClearUnusedGraphResources(PixelShader, Parameters);
+	FIntPoint Resolution = OutShadowMask->Desc.Extent;
+
+	GraphBuilder.AddPass(
+		RDG_EVENT_NAME("HairStrandsShadowMask(Voxel)"),
+		Parameters,
+		ERDGPassFlags::Raster,
+		[Parameters, VertexShader, PixelShader, Viewport, Resolution](FRHICommandList& RHICmdList)
+	{
+		FGraphicsPipelineStateInitializer GraphicsPSOInit;
+		RHICmdList.ApplyCachedRenderTargets(GraphicsPSOInit);
+		GraphicsPSOInit.BlendState = TStaticBlendState<CW_RGBA, BO_Min, BF_One, BF_One, BO_Min, BF_One, BF_One>::GetRHI(); // Min Operator
+		GraphicsPSOInit.RasterizerState = TStaticRasterizerState<>::GetRHI();
+		GraphicsPSOInit.DepthStencilState = TStaticDepthStencilState<false, CF_Always>::GetRHI();
+
+		GraphicsPSOInit.BoundShaderState.VertexDeclarationRHI = GFilterVertexDeclaration.VertexDeclarationRHI;
+		GraphicsPSOInit.BoundShaderState.VertexShaderRHI = VertexShader.GetVertexShader();
+		GraphicsPSOInit.BoundShaderState.PixelShaderRHI = PixelShader.GetPixelShader();
+		GraphicsPSOInit.PrimitiveType = PT_TriangleList;
+		SetGraphicsPipelineState(RHICmdList, GraphicsPSOInit);
+
+		RHICmdList.SetViewport(Viewport.Min.X, Viewport.Min.Y, 0.0f, Viewport.Max.X, Viewport.Max.Y, 1.0f);
+		SetShaderParameters(RHICmdList, PixelShader, PixelShader.GetPixelShader(), *Parameters);
+
+		DrawRectangle(
+			RHICmdList,
+			0, 0,
+			Viewport.Width(), Viewport.Height(),
+			Viewport.Min.X, Viewport.Min.Y,
+			Viewport.Width(), Viewport.Height(),
+			Viewport.Size(),
+			Resolution,
+			VertexShader,
+			EDRF_UseTriangleOptimization);
+	});
+}
+
+///////////////////////////////////////////////////////////////////////////////////////////////////
+// Opaque Mask from deep shadow
+
+class FHairStrandsDeepShadowMaskPS : public FGlobalShader
+{
+	DECLARE_GLOBAL_SHADER(FHairStrandsDeepShadowMaskPS);
+	SHADER_USE_PARAMETER_STRUCT(FHairStrandsDeepShadowMaskPS, FGlobalShader);
 
-	class FOpaqueMaskType : SHADER_PERMUTATION_INT("PERMUTATION_OPAQUEMASK_TYPE", FHairOpaqueMaskTypeCount);
 	class FKernelType : SHADER_PERMUTATION_INT("PERMUTATION_KERNEL_TYPE", 5);
-	using FPermutationDomain = TShaderPermutationDomain<FOpaqueMaskType, FKernelType>;
+	using FPermutationDomain = TShaderPermutationDomain<FKernelType>;
 
 	BEGIN_SHADER_PARAMETER_STRUCT(FParameters, )
 		SHADER_PARAMETER_STRUCT_INCLUDE(ShaderDrawDebug::FShaderDrawDebugParameters, ShaderDrawParameters)
@@ -372,22 +544,17 @@ class FDeepShadowMaskPS : public FGlobalShader
 		SHADER_PARAMETER(uint32, DeepShadow_bIsGPUDriven)
 		SHADER_PARAMETER(FVector4, DeepShadow_LayerDepths)
 
-		SHADER_PARAMETER(FVector4, Voxel_LightPosition)
-		SHADER_PARAMETER(FVector, Voxel_LightDirection)
-		SHADER_PARAMETER(uint32, Voxel_MacroGroupId)
-		SHADER_PARAMETER(uint32, Voxel_RandomType)
 		SHADER_PARAMETER_RDG_TEXTURE(Texture2D, RayMarchMaskTexture)
 
 		SHADER_PARAMETER_RDG_BUFFER_SRV(StructuredBuffer, DeepShadow_WorldToLightTransformBuffer)
 		SHADER_PARAMETER_RDG_TEXTURE(Texture2D, DeepShadow_FrontDepthTexture)
 		SHADER_PARAMETER_RDG_TEXTURE(Texture2D, DeepShadow_DomTexture)
-		SHADER_PARAMETER_RDG_TEXTURE(Texture2D, CategorizationTexture)
 		SHADER_PARAMETER_RDG_TEXTURE(Texture2D, SceneDepthTexture)
 		SHADER_PARAMETER_SAMPLER(SamplerState, LinearSampler)
 		SHADER_PARAMETER_SAMPLER(SamplerState, ShadowSampler)
 
 		SHADER_PARAMETER_STRUCT_REF(FViewUniformShaderParameters, ViewUniformBuffer)
-		SHADER_PARAMETER_RDG_UNIFORM_BUFFER(FVirtualVoxelParameters, VirtualVoxel)
+		SHADER_PARAMETER_RDG_UNIFORM_BUFFER(FHairStrandsViewUniformParameters, HairStrands)
 		RENDER_TARGET_BINDING_SLOTS()
 		END_SHADER_PARAMETER_STRUCT()
 
@@ -397,17 +564,16 @@ public:
 	{
 		FGlobalShader::ModifyCompilationEnvironment(Parameters, OutEnvironment);
 		OutEnvironment.SetRenderTargetOutputFormat(0, PF_B8G8R8A8);
+		OutEnvironment.SetDefine(TEXT("SHADER_SHADOWMASK_DEEPSHADOW"), 1);
 	}
 };
 
-IMPLEMENT_GLOBAL_SHADER(FDeepShadowMaskPS, "/Engine/Private/HairStrands/HairStrandsDeepShadowMask.usf", "MainPS", SF_Pixel);
+IMPLEMENT_GLOBAL_SHADER(FHairStrandsDeepShadowMaskPS, "/Engine/Private/HairStrands/HairStrandsDeepShadowMask.usf", "MainPS", SF_Pixel);
 
 
-struct FDeepShadowOpaqueParams
+struct FHairStrandsDeepShadowParams
 {
-	FRDGTextureRef	CategorizationTexture = nullptr;
 	bool			bIsWholeSceneLight = false;
-
 	FRDGBufferSRVRef DeepShadow_WorldToLightTransformBuffer = nullptr;
 	FMatrix			DeepShadow_CPUWorldToLightTransform;
 	FIntRect		DeepShadow_AtlasRect;
@@ -418,31 +584,24 @@ struct FDeepShadowOpaqueParams
 	float			DeepShadow_DensityScale = 1;
 	uint32			DeepShadow_AtlasSlotIndex = 0;
 	FVector4		DeepShadow_LayerDepths = FVector4(0, 0, 0, 0);
-
-	FVector			Voxel_LightDirection = FVector::ZeroVector;
-	FVector4		Voxel_LightPosition = FVector4(0, 0, 0, 0);
-	uint32			Voxel_MacroGroupId;
-	const FHairStrandsVoxelResources* Voxel_VirtualVoxel = nullptr;
 };
 
-// Opaque mask
-static void AddDeepShadowOpaqueMaskPass(
+// Opaque mask with deep shadow
+static void AddHairStrandsDeepShadowMaskPass(
 	FRDGBuilder& GraphBuilder,
 	FRDGTextureRef SceneDepthTexture,
 	const FViewInfo& View,
-	const FHairOpaqueMaskType HairOpaqueMaskType,
-	const FDeepShadowOpaqueParams& Params,
+	const FHairStrandsDeepShadowParams& Params,
 	FRDGTextureRef& OutShadowMask)
 {
 	check(OutShadowMask);
 
-	FDeepShadowMaskPS::FParameters* Parameters = GraphBuilder.AllocParameters<FDeepShadowMaskPS::FParameters>();
+	FHairStrandsDeepShadowMaskPS::FParameters* Parameters = GraphBuilder.AllocParameters<FHairStrandsDeepShadowMaskPS::FParameters>();
 	Parameters->ViewUniformBuffer = View.ViewUniformBuffer;
 	Parameters->SceneDepthTexture = SceneDepthTexture;
 	Parameters->DeepShadow_CPUWorldToLightTransform = Params.DeepShadow_CPUWorldToLightTransform;
 	Parameters->DeepShadow_FrontDepthTexture = Params.DeepShadow_FrontDepthTexture;
 	Parameters->DeepShadow_DomTexture = Params.DeepShadow_LayerTexture;
-	Parameters->CategorizationTexture = Params.CategorizationTexture;
 	Parameters->LinearSampler = TStaticSamplerState<SF_Bilinear, AM_Clamp, AM_Clamp, AM_Clamp>::GetRHI();
 	Parameters->ShadowSampler = TStaticSamplerState<SF_Point, AM_Clamp, AM_Clamp, AM_Clamp>::GetRHI();
 	Parameters->bIsWholeSceneLight = Params.bIsWholeSceneLight ? 1 : 0;
@@ -456,14 +615,13 @@ static void AddDeepShadowOpaqueMaskPass(
 	Parameters->DeepShadow_SlotResolution = FIntPoint(Params.DeepShadow_AtlasRect.Max.X - Params.DeepShadow_AtlasRect.Min.X, Params.DeepShadow_AtlasRect.Max.Y - Params.DeepShadow_AtlasRect.Min.Y);
 	Parameters->DeepShadow_WorldToLightTransformBuffer = Params.DeepShadow_WorldToLightTransformBuffer;
 	Parameters->DeepShadow_bIsGPUDriven = Params.DeepShadow_bIsGPUDriven ? 1 : 0;;
-
+	Parameters->HairStrands = HairStrands::BindHairStrandsViewUniformParameters(View);
 	if (ShaderDrawDebug::IsShaderDrawDebugEnabled(View))
 	{
 		ShaderDrawDebug::SetParameters(GraphBuilder, View.ShaderDrawData, Parameters->ShaderDrawParameters);
 	}
 
 	FRDGTextureRef RayMarchMask = nullptr;
-	if (HairOpaqueMaskType == FHairOpaqueMaskType_VirtualVoxel)
 	{
 		FRDGTextureDesc Desc = OutShadowMask->Desc;
 		Desc.Flags |= TexCreate_ShaderResource;
@@ -474,68 +632,55 @@ static void AddDeepShadowOpaqueMaskPass(
 	}
 	Parameters->RayMarchMaskTexture = RayMarchMask;
 
-	if (HairOpaqueMaskType == FHairOpaqueMaskType_VirtualVoxel)
-	{
-		check(Params.Voxel_VirtualVoxel);
-		Parameters->Voxel_LightPosition = Params.Voxel_LightPosition;
-		Parameters->Voxel_LightDirection = Params.Voxel_LightDirection;
-		Parameters->Voxel_MacroGroupId = Params.Voxel_MacroGroupId;
-		Parameters->VirtualVoxel = Params.Voxel_VirtualVoxel->UniformBuffer;
-		Parameters->Voxel_RandomType = FMath::Clamp(GHairStrandsShadowRandomTraversalType, 0, 2);
-	}
-
-	FDeepShadowMaskPS::FPermutationDomain PermutationVector;
-	PermutationVector.Set<FDeepShadowMaskPS::FOpaqueMaskType>(HairOpaqueMaskType);
-	PermutationVector.Set<FDeepShadowMaskPS::FKernelType>(FMath::Clamp(GStrandHairShadowMaskKernelType, 0, 4));
+	FHairStrandsDeepShadowMaskPS::FPermutationDomain PermutationVector;
+	PermutationVector.Set<FHairStrandsDeepShadowMaskPS::FKernelType>(FMath::Clamp(GStrandHairShadowMaskKernelType, 0, 4));
 
 	const FIntPoint OutputResolution = SceneDepthTexture->Desc.Extent;
 	TShaderMapRef<FPostProcessVS> VertexShader(View.ShaderMap);
-	TShaderMapRef<FDeepShadowMaskPS> PixelShader(View.ShaderMap, PermutationVector);
+	TShaderMapRef<FHairStrandsDeepShadowMaskPS> PixelShader(View.ShaderMap, PermutationVector);
 	const FGlobalShaderMap* GlobalShaderMap = View.ShaderMap;
 	const FIntRect Viewport = View.ViewRect;
 
+	ClearUnusedGraphResources(PixelShader, Parameters);
+	FIntPoint Resolution = OutShadowMask->Desc.Extent;
+
+	GraphBuilder.AddPass(
+		RDG_EVENT_NAME("HairStrandsShadowMask"),
+		Parameters,
+		ERDGPassFlags::Raster,
+		[Parameters, VertexShader, PixelShader, Viewport, Resolution](FRHICommandList& RHICmdList)
 	{
-		ClearUnusedGraphResources(PixelShader, Parameters);
-		FIntPoint Resolution = OutShadowMask->Desc.Extent;
-
-		GraphBuilder.AddPass(
-			RDG_EVENT_NAME("HairStrandsShadowMask"),
-			Parameters,
-			ERDGPassFlags::Raster,
-			[Parameters, VertexShader, PixelShader, Viewport, Resolution](FRHICommandList& RHICmdList)
-		{
-			FGraphicsPipelineStateInitializer GraphicsPSOInit;
-			RHICmdList.ApplyCachedRenderTargets(GraphicsPSOInit);
-			GraphicsPSOInit.BlendState = TStaticBlendState<CW_RGBA, BO_Min, BF_One, BF_One, BO_Min, BF_One, BF_One>::GetRHI(); // Min Operator
-			GraphicsPSOInit.RasterizerState = TStaticRasterizerState<>::GetRHI();
-			GraphicsPSOInit.DepthStencilState = TStaticDepthStencilState<false, CF_Always>::GetRHI();
-
-			GraphicsPSOInit.BoundShaderState.VertexDeclarationRHI = GFilterVertexDeclaration.VertexDeclarationRHI;
-			GraphicsPSOInit.BoundShaderState.VertexShaderRHI = VertexShader.GetVertexShader();
-			GraphicsPSOInit.BoundShaderState.PixelShaderRHI = PixelShader.GetPixelShader();
-			GraphicsPSOInit.PrimitiveType = PT_TriangleList;
-			SetGraphicsPipelineState(RHICmdList, GraphicsPSOInit);
-
-			RHICmdList.SetViewport(Viewport.Min.X, Viewport.Min.Y, 0.0f, Viewport.Max.X, Viewport.Max.Y, 1.0f);
-			SetShaderParameters(RHICmdList, PixelShader, PixelShader.GetPixelShader(), *Parameters);
-
-			DrawRectangle(
-				RHICmdList,
-				0, 0,
-				Viewport.Width(), Viewport.Height(),
-				Viewport.Min.X, Viewport.Min.Y,
-				Viewport.Width(), Viewport.Height(),
-				Viewport.Size(),
-				Resolution,
-				VertexShader,
-				EDRF_UseTriangleOptimization);
-		});
-	}
+		FGraphicsPipelineStateInitializer GraphicsPSOInit;
+		RHICmdList.ApplyCachedRenderTargets(GraphicsPSOInit);
+		GraphicsPSOInit.BlendState = TStaticBlendState<CW_RGBA, BO_Min, BF_One, BF_One, BO_Min, BF_One, BF_One>::GetRHI(); // Min Operator
+		GraphicsPSOInit.RasterizerState = TStaticRasterizerState<>::GetRHI();
+		GraphicsPSOInit.DepthStencilState = TStaticDepthStencilState<false, CF_Always>::GetRHI();
+
+		GraphicsPSOInit.BoundShaderState.VertexDeclarationRHI = GFilterVertexDeclaration.VertexDeclarationRHI;
+		GraphicsPSOInit.BoundShaderState.VertexShaderRHI = VertexShader.GetVertexShader();
+		GraphicsPSOInit.BoundShaderState.PixelShaderRHI = PixelShader.GetPixelShader();
+		GraphicsPSOInit.PrimitiveType = PT_TriangleList;
+		SetGraphicsPipelineState(RHICmdList, GraphicsPSOInit);
+
+		RHICmdList.SetViewport(Viewport.Min.X, Viewport.Min.Y, 0.0f, Viewport.Max.X, Viewport.Max.Y, 1.0f);
+		SetShaderParameters(RHICmdList, PixelShader, PixelShader.GetPixelShader(), *Parameters);
+
+		DrawRectangle(
+			RHICmdList,
+			0, 0,
+			Viewport.Width(), Viewport.Height(),
+			Viewport.Min.X, Viewport.Min.Y,
+			Viewport.Width(), Viewport.Height(),
+			Viewport.Size(),
+			Resolution,
+			VertexShader,
+			EDRF_UseTriangleOptimization);
+	});
 }
 
 ///////////////////////////////////////////////////////////////////////////////////////////////////
 
-static FHairStrandsTransmittanceMaskData RenderHairStrandsTransmittanceMask(
+static FHairStrandsTransmittanceMaskData InternalRenderHairStrandsTransmittanceMask(
 	FRDGBuilder& GraphBuilder,
 	const FViewInfo& View,
 	const FLightSceneInfo* LightSceneInfo,
@@ -563,16 +708,15 @@ static FHairStrandsTransmittanceMaskData RenderHairStrandsTransmittanceMask(
 
 	FRDGTextureRef HairLUTTexture = InHairLUT.Textures[HairLUTType_DualScattering];
 
-	FDeepShadowTransmittanceParams Params;
-	Params.DeepShadow_DensityScale = GetDeepShadowDensityScale();
-	Params.DeepShadow_DepthBiasScale = GetDeepShadowDepthBiasScale();
-	memset(Params.DeepShadow_AtlasSlotOffsets_AtlasSlotIndex, 0, sizeof(Params.DeepShadow_AtlasSlotOffsets_AtlasSlotIndex));
-	memset(Params.DeepShadow_CPUWorldToLightTransforms, 0, sizeof(Params.DeepShadow_CPUWorldToLightTransforms));
-
-
 	bool bHasFoundLight = false;
 	if (!IsHairStrandsForVoxelTransmittanceAndShadowEnable())
 	{
+		FHairStrandsDeepShadowTransmittanceLightParams Params;
+		Params.DeepShadow_DensityScale = GetDeepShadowDensityScale();
+		Params.DeepShadow_DepthBiasScale = GetDeepShadowDepthBiasScale();
+		memset(Params.DeepShadow_AtlasSlotOffsets_AtlasSlotIndex, 0, sizeof(Params.DeepShadow_AtlasSlotOffsets_AtlasSlotIndex));
+		memset(Params.DeepShadow_CPUWorldToLightTransforms, 0, sizeof(Params.DeepShadow_CPUWorldToLightTransforms));
+
 		FRDGBufferSRVRef DeepShadow_WorldToLightTransformBufferSRV = nullptr;
 		for (const FHairStrandsMacroGroupData& MacroGroupData : MacroGroupDatas)
 		{
@@ -606,14 +750,12 @@ static FHairStrandsTransmittanceMaskData RenderHairStrandsTransmittanceMask(
 		{
 			check(Params.DeepShadow_FrontDepthTexture);
 			check(Params.DeepShadow_DomTexture);
-			Out.TransmittanceMask = AddDeepShadowTransmittanceMaskPass(
+			Out.TransmittanceMask = AddHairStrandsDeepShadowTransmittanceMaskPass(
 				GraphBuilder,
 				SceneTextures,
 				View,
-				FHairTransmittanceType_DeepShadow,
 				Params,
 				VisibilityData.NodeGroupSize,
-				HairLUTTexture,
 				VisibilityData.NodeIndirectArg,
 				ScreenShadowMaskSubPixelTexture);
 		}
@@ -624,22 +766,18 @@ static FHairStrandsTransmittanceMaskData RenderHairStrandsTransmittanceMask(
 		FLightShaderParameters LightParameters;
 		LightSceneInfo->Proxy->GetLightShaderParameters(LightParameters);
 
-		Params.DeepShadow_DensityScale = 0;
-		Params.DeepShadow_DepthBiasScale = 0;
+		FHairStrandsTransmittanceLightParams Params;
 		Params.LightDirection = LightSceneInfo->Proxy->GetDirection();
 		Params.LightPosition = FVector4(LightSceneInfo->Proxy->GetPosition(), LightSceneInfo->Proxy->GetLightType() == ELightComponentType::LightType_Directional ? 0 : 1);
 		Params.LightChannelMask = LightSceneInfo->Proxy->GetLightingChannelMask();
 		Params.LightRadius = FMath::Max(LightParameters.SourceLength, LightParameters.SourceRadius);
-		Params.VirtualVoxelResources = &VoxelResources;
 
-		Out.TransmittanceMask = AddDeepShadowTransmittanceMaskPass(
+		Out.TransmittanceMask = AddHairStrandsVoxelTransmittanceMaskPass(
 			GraphBuilder,
 			SceneTextures,
 			View,
-			FHairTransmittanceType_VirtualVoxel,
 			Params,
 			VisibilityData.NodeGroupSize,
-			HairLUTTexture,
 			VisibilityData.NodeIndirectArg,
 			ScreenShadowMaskSubPixelTexture);
 	}
@@ -656,7 +794,7 @@ FHairStrandsTransmittanceMaskData RenderHairStrandsTransmittanceMask(
 	FHairStrandsTransmittanceMaskData TransmittanceMaskData;
 	if (HairStrands::HasViewHairStrandsData(View))
 	{
-		TransmittanceMaskData = RenderHairStrandsTransmittanceMask(
+		TransmittanceMaskData = InternalRenderHairStrandsTransmittanceMask(
 			GraphBuilder, 
 			View, 
 			LightSceneInfo, 
@@ -671,7 +809,7 @@ FHairStrandsTransmittanceMaskData RenderHairStrandsTransmittanceMask(
 
 ///////////////////////////////////////////////////////////////////////////////////////////////////
 
-static void RenderHairStrandsShadowMask(
+static void InternalRenderHairStrandsShadowMask(
 	FRDGBuilder& GraphBuilder,
 	const FViewInfo& View,
 	const FLightSceneInfo* LightSceneInfo,
@@ -714,8 +852,7 @@ static void RenderHairStrandsShadowMask(
 				bHasDeepShadow = true;
 				const bool bIsWholeSceneLight = LightSceneInfo->Proxy->GetLightType() == LightType_Directional;
 
-				FDeepShadowOpaqueParams Params;
-				Params.CategorizationTexture = InVisibilityData.CategorizationTexture;
+				FHairStrandsDeepShadowParams Params;
 				Params.DeepShadow_AtlasSlotIndex = DomData.AtlasSlotIndex;
 				Params.DeepShadow_WorldToLightTransformBuffer = DeepShadow_WorldToLightTransformBufferSRV;
 				Params.DeepShadow_bIsGPUDriven = DeepShadowResources.bIsGPUDriven ? 1 : 0;
@@ -727,11 +864,10 @@ static void RenderHairStrandsShadowMask(
 				Params.DeepShadow_DepthBiasScale = GetDeepShadowDepthBiasScale();
 				Params.DeepShadow_DensityScale = GetDeepShadowDensityScale();
 				Params.DeepShadow_LayerDepths = ComputeDeepShadowLayerDepths(DomData.LayerDistribution);
-				AddDeepShadowOpaqueMaskPass(
+				AddHairStrandsDeepShadowMaskPass(
 					GraphBuilder,
 					SceneTextures.Depth.Resolve,
 					View,
-					FHairOpaqueMaskType_DeepShadow,
 					Params,
 					OutShadowMask);
 			}
@@ -740,26 +876,23 @@ static void RenderHairStrandsShadowMask(
 
 	// Code is disabled for now until we have the full DOM/voxel fallback logic
 	// If there is no deep shadow for this light, fallback on the voxel representation
-	if (!bHasDeepShadow && VoxelResources.IsValid())
+	if (!bHasDeepShadow && HairStrands::HasViewHairStrandsVoxelData(View))
 	{
 		// TODO: Change this to be a single pass with virtual voxel?
 		for (const FHairStrandsMacroGroupData& MacroGroupData : InMacroGroupDatas)
 		{
 			const bool bIsWholeSceneLight = LightSceneInfo->Proxy->GetLightType() == LightType_Directional;
 
-			FDeepShadowOpaqueParams Params;
-			Params.CategorizationTexture = InVisibilityData.CategorizationTexture;
+			FHairStrandsVoxelShadowParams Params;
 			Params.Voxel_LightDirection = LightSceneInfo->Proxy->GetDirection();
 			Params.Voxel_LightPosition = FVector4(LightSceneInfo->Proxy->GetPosition(), LightSceneInfo->Proxy->GetLightType() == ELightComponentType::LightType_Directional ? 0 : 1);
 			Params.Voxel_MacroGroupId = MacroGroupData.MacroGroupId;
 			Params.bIsWholeSceneLight = bIsWholeSceneLight ? 1 : 0;
 
-			Params.Voxel_VirtualVoxel = &VoxelResources;
-			AddDeepShadowOpaqueMaskPass(
+			AddHairStrandsVoxelShadowMaskPass(
 				GraphBuilder,
 				SceneTextures.Depth.Resolve,
 				View,
-				FHairOpaqueMaskType_VirtualVoxel,
 				Params,
 				OutShadowMask);
 		}
@@ -782,7 +915,7 @@ void RenderHairStrandsShadowMask(
 		if (HairStrands::HasViewHairStrandsData(View))
 		{
 			check(View.HairStrandsViewData.VisibilityData.CategorizationTexture);
-			RenderHairStrandsShadowMask(
+			InternalRenderHairStrandsShadowMask(
 				GraphBuilder, 
 				View, 
 				LightSceneInfo, 
