Fix validation error due to invalid default texture for hair categorization texture and hair offset/count texture.

diff --git a/Engine/Source/Runtime/Renderer/Private/HairStrands/HairStrandsRendering.cpp b/Engine/Source/Runtime/Renderer/Private/HairStrands/HairStrandsRendering.cpp
--- a/Engine/Source/Runtime/Renderer/Private/HairStrands/HairStrandsRendering.cpp
+++ b/Engine/Source/Runtime/Renderer/Private/HairStrands/HairStrandsRendering.cpp
@@ -26,13 +26,17 @@ static TRDGUniformBufferRef<FHairStrandsViewUniformParameters> InternalCreateHai
 		FRDGBufferRef DummyBuffer = GraphBuilder.CreateBuffer(FRDGBufferDesc::CreateStructuredDesc(4,1), TEXT("Hair.DummyBuffer"));
 		FRDGBufferRef DummyNodeBuffer = GraphBuilder.CreateBuffer(FRDGBufferDesc::CreateStructuredDesc(20, 1), TEXT("Hair.DummyNodeBuffer"));
 		FRDGTextureRef BlackTexture = GSystemTextures.GetBlackDummy(GraphBuilder);
+		FRDGTextureRef ZeroR32_UINT = GSystemTextures.GetZeroUIntDummy(GraphBuilder);
+		FRDGTextureRef ZeroRGBA16_UINT = GSystemTextures.GetZeroUShort4Dummy(GraphBuilder);
+		FRDGTextureRef FarDepth = GSystemTextures.GetDepthDummy(GraphBuilder);
+
 		AddClearUAVPass(GraphBuilder, GraphBuilder.CreateUAV(DummyNodeBuffer), 0);
 		AddClearUAVPass(GraphBuilder, GraphBuilder.CreateUAV(DummyBuffer), 0);
 
-		Parameters->HairOnlyDepthTexture = GSystemTextures.GetDepthDummy(GraphBuilder);
-		Parameters->HairCategorizationTexture = BlackTexture;
-		Parameters->HairSampleCount = BlackTexture;
-		Parameters->HairSampleOffset = BlackTexture;
+		Parameters->HairOnlyDepthTexture = FarDepth;
+		Parameters->HairCategorizationTexture = ZeroRGBA16_UINT;
+		Parameters->HairSampleCount = ZeroR32_UINT;
+		Parameters->HairSampleOffset = ZeroR32_UINT;
 		Parameters->HairSampleCoords = GraphBuilder.CreateSRV(DummyBuffer);
 		Parameters->HairSampleData	 = GraphBuilder.CreateSRV(DummyNodeBuffer);
 		Parameters->HairSampleViewportResolution = FIntPoint(0, 0);
diff --git a/Engine/Source/Runtime/Renderer/Private/SystemTextures.cpp b/Engine/Source/Runtime/Renderer/Private/SystemTextures.cpp
--- a/Engine/Source/Runtime/Renderer/Private/SystemTextures.cpp
+++ b/Engine/Source/Runtime/Renderer/Private/SystemTextures.cpp
@@ -109,7 +109,6 @@ void FSystemTextures::InitializeCommonTextures(FRHICommandListImmediate& RHICmdL
 			RHICmdList.CopyToResolveTarget(BlackDummy->GetRenderTargetItem().TargetableTexture, BlackDummy->GetRenderTargetItem().ShaderResourceTexture, FResolveParams());
 		}
 	
-
 		// Create a texture that is a single UInt32 value set to 0
 		{
 			FPooledRenderTargetDesc Desc(FPooledRenderTargetDesc::Create2DDesc(FIntPoint(1,1), PF_R32_UINT, FClearValueBinding::Transparent, TexCreate_HideInVisualizeTexture, TexCreate_RenderTargetable | TexCreate_NoFastClear | TexCreate_ShaderResource, false));
@@ -124,6 +123,20 @@ void FSystemTextures::InitializeCommonTextures(FRHICommandListImmediate& RHICmdL
 			RHICmdList.CopyToResolveTarget(ZeroUIntDummy->GetRenderTargetItem().TargetableTexture, ZeroUIntDummy->GetRenderTargetItem().ShaderResourceTexture, FResolveParams());
 		}
 
+		// Create a texture that is a single 4xUInt16 (UShort) value set to 0
+		{
+			FPooledRenderTargetDesc Desc(FPooledRenderTargetDesc::Create2DDesc(FIntPoint(1, 1), PF_R16G16B16A16_UINT, FClearValueBinding::Transparent, TexCreate_HideInVisualizeTexture, TexCreate_RenderTargetable | TexCreate_NoFastClear | TexCreate_ShaderResource, false));
+			Desc.AutoWritable = false;
+			GRenderTargetPool.FindFreeElement(RHICmdList, Desc, ZeroUShort4Dummy, TEXT("ZeroUShort4Dummy"), ERenderTargetTransience::NonTransient);
+
+			RHICmdList.Transition(FRHITransitionInfo(ZeroUShort4Dummy->GetRenderTargetItem().TargetableTexture, ERHIAccess::SRVMask, ERHIAccess::RTV));
+
+			FRHIRenderPassInfo RPInfo(ZeroUShort4Dummy->GetRenderTargetItem().TargetableTexture, ERenderTargetActions::Clear_Store);
+			RHICmdList.BeginRenderPass(RPInfo, TEXT("ClearZeroUShort4Dummy"));
+			RHICmdList.EndRenderPass();
+			RHICmdList.CopyToResolveTarget(ZeroUShort4Dummy->GetRenderTargetItem().TargetableTexture, ZeroUShort4Dummy->GetRenderTargetItem().ShaderResourceTexture, FResolveParams());
+		}
+
 		// Create a BlackAlphaOneDummy texture
 		{
 			FPooledRenderTargetDesc Desc(FPooledRenderTargetDesc::Create2DDesc(FIntPoint(1, 1), PF_B8G8R8A8, FClearValueBinding::Black, TexCreate_HideInVisualizeTexture, TexCreate_RenderTargetable | TexCreate_NoFastClear | TexCreate_ShaderResource, false));
@@ -806,6 +819,7 @@ void FSystemTextures::ReleaseDynamicRHI()
 	DefaultNormal8Bit.SafeRelease();
 	VolumetricBlackDummy.SafeRelease();
 	ZeroUIntDummy.SafeRelease();
+	ZeroUShort4Dummy.SafeRelease();
 	MidGreyDummy.SafeRelease();
 	StencilDummy.SafeRelease();
 	StencilDummySRV.SafeRelease();
@@ -872,3 +886,8 @@ FRDGTextureRef FSystemTextures::GetZeroUIntDummy(FRDGBuilder& GraphBuilder) cons
 	return GraphBuilder.RegisterExternalTexture(ZeroUIntDummy, TEXT("ZeroUIntDummy"));
 }
 
+FRDGTextureRef FSystemTextures::GetZeroUShort4Dummy(FRDGBuilder& GraphBuilder) const
+{
+	return GraphBuilder.RegisterExternalTexture(ZeroUShort4Dummy, TEXT("ZeroUShort4Dummy"));
+}
+
diff --git a/Engine/Source/Runtime/Renderer/Private/SystemTextures.h b/Engine/Source/Runtime/Renderer/Private/SystemTextures.h
--- a/Engine/Source/Runtime/Renderer/Private/SystemTextures.h
+++ b/Engine/Source/Runtime/Renderer/Private/SystemTextures.h
@@ -105,9 +105,9 @@ public:
 	/** float4(0,0,0,0) volumetric texture. */
 	TRefCountPtr<IPooledRenderTarget> VolumetricBlackDummy;
 	
-
 	// Dummy 0 Uint texture for RHIs that need explicit overloads
 	TRefCountPtr<IPooledRenderTarget> ZeroUIntDummy;
+	TRefCountPtr<IPooledRenderTarget> ZeroUShort4Dummy;
 
 	// SRV for WhiteDummy Texture.
 	TRefCountPtr<FRHIShaderResourceView> WhiteDummySRV;
@@ -117,6 +117,7 @@ public:
 	FRDGTextureRef GetWhiteDummy(FRDGBuilder& GraphBuilder) const;
 	FRDGTextureRef GetBlackDummy(FRDGBuilder& GraphBuilder) const;
 	FRDGTextureRef GetZeroUIntDummy(FRDGBuilder& GraphBuilder) const;
+	FRDGTextureRef GetZeroUShort4Dummy(FRDGBuilder& GraphBuilder) const;
 	FRDGTextureRef GetBlackAlphaOneDummy(FRDGBuilder& GraphBuilder) const;
 	FRDGTextureRef GetMaxFP16Depth(FRDGBuilder& GraphBuilder) const;
 	FRDGTextureRef GetDepthDummy(FRDGBuilder& GraphBuilder) const;
