Fixed: The Stain Decal looks bright and firm on mobile

diff --git a/Engine/Source/Runtime/Renderer/Public/DecalRenderingCommon.h b/Engine/Source/Runtime/Renderer/Public/DecalRenderingCommon.h
--- a/Engine/Source/Runtime/Renderer/Public/DecalRenderingCommon.h
+++ b/Engine/Source/Runtime/Renderer/Public/DecalRenderingCommon.h
@@ -60,7 +60,7 @@ struct FDecalRenderingCommon
 	
 	static EDecalBlendMode ComputeFinalDecalBlendMode(EShaderPlatform Platform, EDecalBlendMode DecalBlendMode, bool bUseNormal)
 	{
-		const bool bShouldConvertToDBuffer = !IsUsingGBuffers(Platform) && !IsSimpleForwardShadingEnabled(Platform) && IsUsingDBuffers(Platform);
+		const bool bShouldConvertToDBuffer = !IsUsingGBuffers(Platform) && !IsSimpleForwardShadingEnabled(Platform) && IsUsingDBuffers(Platform) && !IsMobilePlatform(Platform);
 
 		if (bShouldConvertToDBuffer)
 		{
