Increase lifetime of scene texture resources in FXSystem to account for accesses.

diff --git a/Engine/Source/Runtime/Renderer/Private/DeferredShadingRenderer.cpp b/Engine/Source/Runtime/Renderer/Private/DeferredShadingRenderer.cpp
--- a/Engine/Source/Runtime/Renderer/Private/DeferredShadingRenderer.cpp
+++ b/Engine/Source/Runtime/Renderer/Private/DeferredShadingRenderer.cpp
@@ -433,18 +433,20 @@ static void RenderOpaqueFX(
 		RDG_GPU_STAT_SCOPE(GraphBuilder, PostRenderOpsFX);
 		RDG_CSV_STAT_EXCLUSIVE_SCOPE(GraphBuilder, RenderOpaqueFX);
 
+		const ERDGPassFlags UBPassFlags = ERDGPassFlags::Compute | ERDGPassFlags::Raster | ERDGPassFlags::SkipRenderPass | ERDGPassFlags::NeverCull;
+
 		// Add a pass which extracts the RHI handle from the scene textures UB and sends it to the FX system.
 		FRenderOpaqueFXPassParameters* ExtractUBPassParameters = GraphBuilder.AllocParameters<FRenderOpaqueFXPassParameters>();
 		ExtractUBPassParameters->SceneTextures = SceneTexturesUniformBuffer;
-		GraphBuilder.AddPass({}, ExtractUBPassParameters, ERDGPassFlags::Raster | ERDGPassFlags::SkipRenderPass | ERDGPassFlags::NeverCull, [ExtractUBPassParameters, FXSystem](FRHICommandList&)
+		GraphBuilder.AddPass({}, ExtractUBPassParameters, UBPassFlags, [ExtractUBPassParameters, FXSystem](FRHICommandList&)
 		{
 			FXSystem->SetSceneTexturesUniformBuffer(ExtractUBPassParameters->SceneTextures->GetRHIRef());
 		});
 
 		FXSystem->PostRenderOpaque(GraphBuilder, Views, true /*bAllowGPUParticleUpdate*/);
 
-		// Clear the scene textures UB pointer on the FX system.
-		AddPass(GraphBuilder, [FXSystem](FRHICommandList&)
+		// Clear the scene textures UB pointer on the FX system. Use the same pass parameters to extend resource lifetimes.
+		GraphBuilder.AddPass({}, ExtractUBPassParameters, UBPassFlags, [FXSystem](FRHICommandList&)
 		{
 			FXSystem->SetSceneTexturesUniformBuffer(nullptr);
 		});
