Fix RHI validation errors in Lumen.
 - We should store the Lumen FRDGTexture references in the FViewInfo, rather than converting the textures to external and using IPooledRenderTarget, otherwise RDG can't track the resources to issue the correct resource barriers.
 - Resolves missing UAV -> SRV barriers on the volume and history texture resources.

diff --git a/Engine/Source/Runtime/Renderer/Private/VolumetricFog.cpp b/Engine/Source/Runtime/Renderer/Private/VolumetricFog.cpp
--- a/Engine/Source/Runtime/Renderer/Private/VolumetricFog.cpp
+++ b/Engine/Source/Runtime/Renderer/Private/VolumetricFog.cpp
@@ -721,6 +721,7 @@ class TVolumetricFogLightScatteringCS : public FGlobalShader
 		SHADER_PARAMETER_RDG_TEXTURE(Texture2D, PrevConservativeDepthTexture)
 		SHADER_PARAMETER(FVector2D, PrevConservativeDepthTextureSize)
 		SHADER_PARAMETER(uint32, UseConservativeDepthTexture)
+		SHADER_PARAMETER_RDG_UNIFORM_BUFFER(FLumenTranslucencyLightingUniforms, LumenGIVolumeStruct)
 		SHADER_PARAMETER_STRUCT_INCLUDE(FVirtualShadowMapSamplingParameters, VirtualShadowMapSamplingParameters)
 		SHADER_PARAMETER_RDG_UNIFORM_BUFFER(FLumenTranslucencyLightingUniforms, LumenGIVolumeStruct)
 		SHADER_PARAMETER_RDG_TEXTURE_UAV(RWTexture2D, RWLightScattering)
