Enable SMRT (7 rays/pixel) and disable denoising by default for directional lights as well.

diff --git a/Engine/Source/Runtime/Renderer/Private/ShadowRendering.cpp b/Engine/Source/Runtime/Renderer/Private/ShadowRendering.cpp
--- a/Engine/Source/Runtime/Renderer/Private/ShadowRendering.cpp
+++ b/Engine/Source/Runtime/Renderer/Private/ShadowRendering.cpp
@@ -190,7 +190,7 @@ static TAutoConsoleVariable<float> CVarShadowMaxSlopeScaleDepthBias(
 
 static TAutoConsoleVariable<int32> CVarVirtualSmDenoiserDirectional(
 	TEXT("r.Shadow.Virtual.Denoiser.EnableDirectional"),
-	1,
+	0,
 	TEXT("Apply denoiser to directional light virtual shadow maps."),
 	ECVF_RenderThreadSafe);
 FForwardScreenSpaceShadowMaskTextureMobileOutputs GScreenSpaceShadowMaskTextureMobileOutputs;
diff --git a/Engine/Source/Runtime/Renderer/Private/VirtualShadowMaps/VirtualShadowMapProjection.cpp b/Engine/Source/Runtime/Renderer/Private/VirtualShadowMaps/VirtualShadowMapProjection.cpp
--- a/Engine/Source/Runtime/Renderer/Private/VirtualShadowMaps/VirtualShadowMapProjection.cpp
+++ b/Engine/Source/Runtime/Renderer/Private/VirtualShadowMaps/VirtualShadowMapProjection.cpp
@@ -79,7 +79,7 @@ static TAutoConsoleVariable<float> CVarSMRTMaxRayAngleFromLight(
 
 static TAutoConsoleVariable<int32> CVarSMRTRayCountDirectional(
 	TEXT( "r.Shadow.Virtual.SMRT.RayCountDirectional" ),
-	0,
+	7,
 	TEXT( "Ray count for shadow map tracing of directional lights. 0 = disabled." ),
 	ECVF_RenderThreadSafe
 );
