Fixed CIS break due to bad RM.

diff --git a/Engine/Source/Runtime/Renderer/Private/VolumetricFog.cpp b/Engine/Source/Runtime/Renderer/Private/VolumetricFog.cpp
--- a/Engine/Source/Runtime/Renderer/Private/VolumetricFog.cpp
+++ b/Engine/Source/Runtime/Renderer/Private/VolumetricFog.cpp
@@ -723,7 +723,6 @@ class TVolumetricFogLightScatteringCS : public FGlobalShader
 		SHADER_PARAMETER(uint32, UseConservativeDepthTexture)
 		SHADER_PARAMETER_RDG_UNIFORM_BUFFER(FLumenTranslucencyLightingUniforms, LumenGIVolumeStruct)
 		SHADER_PARAMETER_STRUCT_INCLUDE(FVirtualShadowMapSamplingParameters, VirtualShadowMapSamplingParameters)
-		SHADER_PARAMETER_RDG_UNIFORM_BUFFER(FLumenTranslucencyLightingUniforms, LumenGIVolumeStruct)
 		SHADER_PARAMETER_RDG_TEXTURE_UAV(RWTexture2D, RWLightScattering)
 	END_SHADER_PARAMETER_STRUCT()
 
