Add draw event for async ray tracing scene BVH build

diff --git a/Engine/Source/Runtime/Renderer/Private/DeferredShadingRenderer.cpp b/Engine/Source/Runtime/Renderer/Private/DeferredShadingRenderer.cpp
--- a/Engine/Source/Runtime/Renderer/Private/DeferredShadingRenderer.cpp
+++ b/Engine/Source/Runtime/Renderer/Private/DeferredShadingRenderer.cpp
@@ -1338,6 +1338,7 @@ bool FDeferredShadingSceneRenderer::DispatchRayTracingWorldUpdates(FRDGBuilder&
 				check(View.RayTracingSceneSRV.IsValid());
 				RHIAsyncCmdList.BindAccelerationStructureMemory(RayTracingScene, View.RayTracingSceneBuffer, 0);
 
+				SCOPED_DRAW_EVENT(RHIAsyncCmdList, RayTracingScene);
 				RHIAsyncCmdList.BuildAccelerationStructure(RayTracingScene);
 			}
 
