Update logic for dealing with VT wrap
- Now try to map both wrapped/unwrapped address when processing a wrapped address
- Previously there was nothing to map tiles to the original vAddress

diff --git a/Engine/Source/Runtime/Renderer/Private/VT/TexturePageMap.h b/Engine/Source/Runtime/Renderer/Private/VT/TexturePageMap.h
--- a/Engine/Source/Runtime/Renderer/Private/VT/TexturePageMap.h
+++ b/Engine/Source/Runtime/Renderer/Private/VT/TexturePageMap.h
@@ -41,6 +41,8 @@ union FPhysicalSpaceIDAndAddress
 		uint16 pAddress;
 	};
 };
+FORCEINLINE bool operator==(const FPhysicalSpaceIDAndAddress& Lhs, const FPhysicalSpaceIDAndAddress& Rhs) { return Lhs.Packed == Rhs.Packed; }
+FORCEINLINE bool operator!=(const FPhysicalSpaceIDAndAddress& Lhs, const FPhysicalSpaceIDAndAddress& Rhs) { return Lhs.Packed != Rhs.Packed; }
 
 struct FMappedTexturePage
 {
diff --git a/Engine/Source/Runtime/Renderer/Private/VT/VirtualTextureSystem.cpp b/Engine/Source/Runtime/Renderer/Private/VT/VirtualTextureSystem.cpp
--- a/Engine/Source/Runtime/Renderer/Private/VT/VirtualTextureSystem.cpp
+++ b/Engine/Source/Runtime/Renderer/Private/VT/VirtualTextureSystem.cpp
@@ -1549,6 +1549,7 @@ void FVirtualTextureSystem::GatherRequestsTask(const FGatherRequestsParameters&
 
 			// Wrap vAddress for the given producer
 			uint32 Wrapped_vAddress = vAddress;
+			uint32 Local_vAddress;
 			{
 				// Scale size of producer up to be relative to size of entire allocated VT
 				const uint32 ProducerScaleFactor = (1u << ProducerMipBias);
@@ -1556,20 +1557,22 @@ void FVirtualTextureSystem::GatherRequestsTask(const FGatherRequestsParameters&
 				const uint32 ProducerHeightInPages = Producer->GetHeightInTiles() * ProducerScaleFactor;
 				const uint32 AllocatedPageX = AllocatedVT->GetVirtualPageX();
 				const uint32 AllocatedPageY = AllocatedVT->GetVirtualPageY();
+				check(vPageX >= AllocatedPageX);
+				check(vPageY >= AllocatedPageY);
 
 				uint32 Local_vPageX = vPageX - AllocatedPageX;
-				uint32 Local_vTileY = vPageY - AllocatedPageY;
-				if (Local_vPageX >= ProducerWidthInPages || Local_vTileY >= ProducerHeightInPages)
+				uint32 Local_vPageY = vPageY - AllocatedPageY;
+				if (Local_vPageX >= ProducerWidthInPages || Local_vPageY >= ProducerHeightInPages)
 				{
 					Local_vPageX %= ProducerWidthInPages;
-					Local_vTileY %= ProducerHeightInPages;
-					Wrapped_vAddress = FMath::MortonCode2(Local_vPageX + AllocatedPageX) | (FMath::MortonCode2(Local_vTileY + AllocatedPageY) << 1);
+					Local_vPageY %= ProducerHeightInPages;
+					Wrapped_vAddress = FMath::MortonCode2(AllocatedPageX + Local_vPageX) | (FMath::MortonCode2(AllocatedPageY + Local_vPageY) << 1);
+
 					Wrapped_vAddress &= (0xffffffff << (vLevel * vDimensions));
 				}
+				Local_vAddress = FMath::MortonCode2(Local_vPageX >> Mapping_vLevel) | (FMath::MortonCode2(Local_vPageY >> Mapping_vLevel) << 1);
 			}
 
-			uint32 Local_vAddress = (Wrapped_vAddress - AllocatedVT->GetVirtualAddress()) >> (Mapping_vLevel * vDimensions);
-
 			const uint32 LocalMipBias = Producer->GetVirtualTexture()->GetLocalMipBias(Local_vLevel, Local_vAddress);
 			if (LocalMipBias > 0u)
 			{
@@ -1604,10 +1607,12 @@ void FVirtualTextureSystem::GatherRequestsTask(const FGatherRequestsParameters&
 					ensure(Allocated_vLevel <= AllocatedVT->GetMaxLevel());
 
 					const uint32 AllocatedMapping_vLevel = FMath::Max(Allocated_vLevel, ProducerMipBias);
-					const uint32 Allocated_vAddress = Wrapped_vAddress & (0xffffffff << (Allocated_vLevel * vDimensions));
+					const uint32 Allocated_vAddress = vAddress & (0xffffffff << (Allocated_vLevel * vDimensions));
+					const uint32 AllocatedWrapped_vAddress = Wrapped_vAddress & (0xffffffff << (Allocated_vLevel * vDimensions));
 
 					AddPageUpdate(PageUpdateBuffers, PageUpdateFlushCount, PhysicalSpace->GetID(), pAddress);
 
+					const FPhysicalSpaceIDAndAddress PhysicalSpaceIDAndAddress(PhysicalSpace->GetID(), pAddress);
 					uint32 NumMappedPages = 0u;
 					for (uint32 LoadLayerIndex = 0u; LoadLayerIndex < NumPageTableLayersToLoad; ++LoadLayerIndex)
 					{
@@ -1615,26 +1620,35 @@ void FVirtualTextureSystem::GatherRequestsTask(const FGatherRequestsParameters&
 						if (AllocatedVT->GetProducerPhysicalGroupIndexForPageTableLayer(PageTableLayerIndex) == ProducerGroupIndex &&
 							AllocatedVT->GetProducerIndexForPageTableLayer(PageTableLayerIndex) == ProducerIndex)
 						{
-							bool bPageWasMapped = false;
-
 							// if we found a lower resolution tile than was requested, it may have already been mapped, check for that first
 							const FTexturePageMap& PageMap = Space->GetPageMapForPageTableLayer(PageTableLayerIndex);
 							const FPhysicalSpaceIDAndAddress PrevPhysicalSpaceIDAndAddress = PageMap.FindPagePhysicalSpaceIDAndAddress(Allocated_vLevel, Allocated_vAddress);
-							if (PrevPhysicalSpaceIDAndAddress.Packed != ~0u)
-							{
-								// if this address was previously mapped, ensure that it was mapped by the same physical space
-								ensure(PrevPhysicalSpaceIDAndAddress.PhysicalSpaceID == PhysicalSpace->GetID());
-								// either it wasn't mapped, or it's mapped to the current physical address...
-								// otherwise that means that the same local tile is mapped to two separate physical addresses, which is an error
-								ensure(PrevPhysicalSpaceIDAndAddress.pAddress == pAddress);
-								bPageWasMapped = true;
-							}
-							if (!bPageWasMapped)
+
+							// either it wasn't mapped, or it's mapped to the current physical address...
+							// otherwise that means that the same local tile is mapped to two separate physical addresses, which is an error
+							ensure(PrevPhysicalSpaceIDAndAddress.Packed == ~0u || PrevPhysicalSpaceIDAndAddress == PhysicalSpaceIDAndAddress);
+
+							if (PrevPhysicalSpaceIDAndAddress.Packed == ~0u)
 							{
 								// map the page now if it wasn't already mapped
 								RequestList->AddDirectMappingRequest(Space->GetID(), PhysicalSpace->GetID(), PageTableLayerIndex, Allocated_vLevel, Allocated_vAddress, AllocatedMapping_vLevel, pAddress);
 							}
 
+							if (AllocatedWrapped_vAddress != Allocated_vAddress)
+							{
+								const FPhysicalSpaceIDAndAddress PrevWrappedPhysicalSpaceIDAndAddress = PageMap.FindPagePhysicalSpaceIDAndAddress(Allocated_vLevel, AllocatedWrapped_vAddress);
+
+								// Wrapped vAddress must be mapped to the same physical address unwrapped vAddress
+								ensure(PrevWrappedPhysicalSpaceIDAndAddress.Packed == ~0u || PrevWrappedPhysicalSpaceIDAndAddress == PhysicalSpaceIDAndAddress);
+								ensure(PrevPhysicalSpaceIDAndAddress.Packed == ~0u ||
+									PrevWrappedPhysicalSpaceIDAndAddress.Packed == ~0u ||
+									PrevPhysicalSpaceIDAndAddress == PrevWrappedPhysicalSpaceIDAndAddress);
+
+								if (PrevWrappedPhysicalSpaceIDAndAddress.Packed == ~0u)
+								{
+									RequestList->AddDirectMappingRequest(Space->GetID(), PhysicalSpace->GetID(), PageTableLayerIndex, Allocated_vLevel, AllocatedWrapped_vAddress, AllocatedMapping_vLevel, pAddress);
+								}
+							}
 							++NumMappedPages;
 						}
 					}
@@ -1700,7 +1714,8 @@ void FVirtualTextureSystem::GatherRequestsTask(const FGatherRequestsParameters&
 						const uint32 Prefetch_vLevel = PrefetchLocal_vLevel + ProducerMipBias;
 						ensure(Prefetch_vLevel <= AllocatedVT->GetMaxLevel());
 						const uint32 PrefetchMapping_vLevel = FMath::Max(Prefetch_vLevel, ProducerMipBias);
-						const uint32 Prefetch_vAddress = Wrapped_vAddress & (0xffffffff << (Prefetch_vLevel * vDimensions));
+						const uint32 Prefetch_vAddress = vAddress & (0xffffffff << (Prefetch_vLevel * vDimensions));
+						const uint32 PrefetchWrapped_vAddress = Wrapped_vAddress & (0xffffffff << (Prefetch_vLevel * vDimensions));
 						for (uint32 LoadLayerIndex = 0u; LoadLayerIndex < NumPageTableLayersToLoad; ++LoadLayerIndex)
 						{
 							const uint32 LayerIndex = PageTableLayersToLoad[LoadLayerIndex];
@@ -1710,6 +1725,10 @@ void FVirtualTextureSystem::GatherRequestsTask(const FGatherRequestsParameters&
 								if (ProducerPhysicalGroupMaskToPrefetch & (1u << ProducerPhysicalGroupIndex))
 								{
 									RequestList->AddMappingRequest(LoadRequestIndex, ProducerPhysicalGroupIndex, ID, LayerIndex, Prefetch_vAddress, Prefetch_vLevel, PrefetchMapping_vLevel);
+									if (PrefetchWrapped_vAddress != Prefetch_vAddress)
+									{
+										RequestList->AddMappingRequest(LoadRequestIndex, ProducerPhysicalGroupIndex, ID, LayerIndex, PrefetchWrapped_vAddress, Prefetch_vLevel, PrefetchMapping_vLevel);
+									}
 								}
 							}
 						}
@@ -1730,7 +1749,11 @@ void FVirtualTextureSystem::GatherRequestsTask(const FGatherRequestsParameters&
 							const uint32 ProducerPhysicalGroupIndex = AllocatedVT->GetProducerPhysicalGroupIndexForPageTableLayer(LayerIndex);
 							if (GroupMaskToLoad & (1u << ProducerPhysicalGroupIndex))
 							{
-								RequestList->AddMappingRequest(LoadRequestIndex, ProducerPhysicalGroupIndex, ID, LayerIndex, Wrapped_vAddress, vLevel, Mapping_vLevel);
+								RequestList->AddMappingRequest(LoadRequestIndex, ProducerPhysicalGroupIndex, ID, LayerIndex, vAddress, vLevel, Mapping_vLevel);
+								if (Wrapped_vAddress != vAddress)
+								{
+									RequestList->AddMappingRequest(LoadRequestIndex, ProducerPhysicalGroupIndex, ID, LayerIndex, Wrapped_vAddress, vLevel, Mapping_vLevel);
+								}
 							}
 						}
 					}
