Orthogonalize GPU scene transforms to remove shear that we don't properly support.

diff --git a/Engine/Source/Runtime/Renderer/Private/GPUScene.cpp b/Engine/Source/Runtime/Renderer/Private/GPUScene.cpp
--- a/Engine/Source/Runtime/Renderer/Private/GPUScene.cpp
+++ b/Engine/Source/Runtime/Renderer/Private/GPUScene.cpp
@@ -512,6 +512,27 @@ namespace
 
 };
 
+FVector OrthonormalizeTransform( FMatrix& Matrix )
+{
+	FVector X, Y, Z, Origin;
+	Matrix.GetScaledAxes( X, Y, Z );
+	Origin = Matrix.GetOrigin();
+
+	// Modified Gram-Schmidt orthogonalization
+	Y -= (Y | X) / (X | X) * X;
+	Z -= (Z | X) / (X | X) * X;
+	Z -= (Z | Y) / (Y | Y) * Y;
+
+	Matrix = FMatrix( X, Y, Z, Origin );
+
+	// Extract per axis scales
+	FVector Scale;
+	Scale.X = X.Size();
+	Scale.Y = Y.Size();
+	Scale.Z = Z.Size();
+	return Scale;
+}
+
 template<typename FUploadDataSourceAdapter>
 void FGPUScene::UploadGeneral(FRHICommandListImmediate& RHICmdList, FScene *Scene, FUploadDataSourceAdapter &UploadDataSourceAdapter)
 {
@@ -858,24 +879,19 @@ void FGPUScene::UploadGeneral(FRHICommandListImmediate& RHICmdList, FScene *Scen
 										PrimitiveInstance.LastUpdateSceneFrameNumber = SceneFrameNumber;
 
 										{
-											// Extract per axis scales from InstanceToWorld transform
-											FVector4 WorldX = FVector4(PrimitiveInstance.LocalToWorld.M[0][0], PrimitiveInstance.LocalToWorld.M[0][1], PrimitiveInstance.LocalToWorld.M[0][2], 0);
-											FVector4 WorldY = FVector4(PrimitiveInstance.LocalToWorld.M[1][0], PrimitiveInstance.LocalToWorld.M[1][1], PrimitiveInstance.LocalToWorld.M[1][2], 0);
-											FVector4 WorldZ = FVector4(PrimitiveInstance.LocalToWorld.M[2][0], PrimitiveInstance.LocalToWorld.M[2][1], PrimitiveInstance.LocalToWorld.M[2][2], 0);
-
-											const float ScaleX = FVector(WorldX).Size();
-											const float ScaleY = FVector(WorldY).Size();
-											const float ScaleZ = FVector(WorldZ).Size();
+											// Remove shear
+											FVector Scale =	OrthonormalizeTransform( PrimitiveInstance.LocalToWorld );
+															OrthonormalizeTransform( PrimitiveInstance.PrevLocalToWorld );
 
 											PrimitiveInstance.NonUniformScale = FVector4(
-												ScaleX, ScaleY, ScaleZ,
-												FMath::Max3(FMath::Abs(ScaleX), FMath::Abs(ScaleY), FMath::Abs(ScaleZ))
+												Scale.X, Scale.Y, Scale.Z,
+												FMath::Max3(FMath::Abs(Scale.X), FMath::Abs(Scale.Y), FMath::Abs(Scale.Z))
 											);
 
 											PrimitiveInstance.InvNonUniformScaleAndDeterminantSign = FVector4(
-												ScaleX > KINDA_SMALL_NUMBER ? 1.0f / ScaleX : 0.0f,
-												ScaleY > KINDA_SMALL_NUMBER ? 1.0f / ScaleY : 0.0f,
-												ScaleZ > KINDA_SMALL_NUMBER ? 1.0f / ScaleZ : 0.0f,
+												Scale.X > KINDA_SMALL_NUMBER ? 1.0f / Scale.X : 0.0f,
+												Scale.Y > KINDA_SMALL_NUMBER ? 1.0f / Scale.Y : 0.0f,
+												Scale.Z > KINDA_SMALL_NUMBER ? 1.0f / Scale.Z : 0.0f,
 												FMath::FloatSelect(PrimitiveInstance.LocalToWorld.RotDeterminant(), 1.0f, -1.0f)
 											);
 										}
