- Fix crash with DirectLighting show flag when VSMs are enabled; always generate light grid as VSM is needed for various indirect lighting effects now too
- Pass shader parameter to disable direct lighting in forward shading when show flag is off
- Fix shadowed lights still getting injected into translucent light volume even when directlighting show flag disabled

diff --git a/Engine/Shaders/Private/ForwardLightingCommon.ush b/Engine/Shaders/Private/ForwardLightingCommon.ush
--- a/Engine/Shaders/Private/ForwardLightingCommon.ush
+++ b/Engine/Shaders/Private/ForwardLightingCommon.ush
@@ -268,6 +268,13 @@ FDeferredLightingSplit GetForwardDirectLightingSplit(uint GridIndex, float3 Worl
 
 #endif
 
+	// Zero out direct lighting if show flag is disabled
+	if (ForwardLightData.DirectLightingShowFlag == 0)
+	{
+		DirectLighting.DiffuseLighting = 0.0f;
+		DirectLighting.SpecularLighting = 0.0f;
+	}
+
 	return DirectLighting;
 }
 
@@ -333,6 +340,12 @@ float3 GetForwardDirectLightingForVertexLighting(uint GridIndex, float3 WorldPos
 		DirectLighting += GetSimpleDynamicLighting(WorldPosition, CameraVector, WorldNormal, 1, DiffuseColor, SpecularColor, Roughness, LightData);
 	}
 
+	// Zero out direct lighting if show flag is disabled
+	if (ForwardLightData.DirectLightingShowFlag == 0)
+	{
+		DirectLighting = 0.0f;
+	}
+
 	return DirectLighting;
 }
 
diff --git a/Engine/Source/Runtime/Renderer/Private/LightGridInjection.cpp b/Engine/Source/Runtime/Renderer/Private/LightGridInjection.cpp
--- a/Engine/Source/Runtime/Renderer/Private/LightGridInjection.cpp
+++ b/Engine/Source/Runtime/Renderer/Private/LightGridInjection.cpp
@@ -551,6 +551,7 @@ void FSceneRenderer::ComputeLightGrid(FRDGBuilder& GraphBuilder, bool bCullLight
 		ForwardLightData.LightGridPixelSizeShift = FMath::FloorLog2(GLightGridPixelSize);
 		ForwardLightData.SimpleLightsEndIndex = SimpleLightsEnd;
 		ForwardLightData.ClusteredDeferredSupportedEndIndex = ClusteredSupportedEnd;
+		ForwardLightData.DirectLightingShowFlag = ViewFamily.EngineShowFlags.DirectLighting ? 1 : 0;
 
 		// Clamp far plane to something reasonable
 		float FarPlane = FMath::Min(FMath::Max(FurthestLight, View.FurthestReflectionCaptureDistance), (float)HALF_WORLD_MAX / 5.0f);
@@ -773,8 +774,7 @@ void FDeferredShadingSceneRenderer::GatherLightsAndComputeLightGrid(FRDGBuilder&
 	}
 	
 	const bool bCullLightsToGrid = GLightCullingQuality
-		&& (ViewFamily.EngineShowFlags.DirectLighting
-			&& (IsForwardShadingEnabled(ShaderPlatform) || bAnyViewUsesForwardLighting || IsRayTracingEnabled() || ShouldUseClusteredDeferredShading() || bAnyViewUsesLumen || ViewFamily.EngineShowFlags.VisualizeMeshDistanceFields || VirtualShadowMapArray.IsEnabled()));
+		&& (IsForwardShadingEnabled(ShaderPlatform) || bAnyViewUsesForwardLighting || IsRayTracingEnabled() || ShouldUseClusteredDeferredShading() || bAnyViewUsesLumen || ViewFamily.EngineShowFlags.VisualizeMeshDistanceFields || VirtualShadowMapArray.IsEnabled());
 
 	// Store this flag if lights are injected in the grids, check with 'AreLightsInLightGrid()'
 	bAreLightsInLightGrid = bCullLightsToGrid;
diff --git a/Engine/Source/Runtime/Renderer/Private/SceneRendering.h b/Engine/Source/Runtime/Renderer/Private/SceneRendering.h
--- a/Engine/Source/Runtime/Renderer/Private/SceneRendering.h
+++ b/Engine/Source/Runtime/Renderer/Private/SceneRendering.h
@@ -712,6 +712,7 @@ const int32 GMaxForwardShadowCascades = 4;
 	SHADER_PARAMETER(uint32, ClusteredDeferredSupportedEndIndex) \
 	SHADER_PARAMETER(FVector4, DirectionalLightStaticShadowBufferSize) \
 	SHADER_PARAMETER(FMatrix, DirectionalLightWorldToStaticShadow) \
+	SHADER_PARAMETER(uint32, DirectLightingShowFlag) \
 	SHADER_PARAMETER_TEXTURE(Texture2D, DirectionalLightShadowmapAtlas) \
 	SHADER_PARAMETER_SAMPLER(SamplerState, ShadowmapSampler) \
 	SHADER_PARAMETER_TEXTURE(Texture2D, DirectionalLightStaticShadowmap) \
diff --git a/Engine/Source/Runtime/Renderer/Private/ShadowRendering.cpp b/Engine/Source/Runtime/Renderer/Private/ShadowRendering.cpp
--- a/Engine/Source/Runtime/Renderer/Private/ShadowRendering.cpp
+++ b/Engine/Source/Runtime/Renderer/Private/ShadowRendering.cpp
@@ -2149,6 +2149,8 @@ void FDeferredShadingSceneRenderer::RenderDeferredShadowProjections(
 	RenderShadowProjections(GraphBuilder, SceneTextures, ScreenShadowMaskTexture, ScreenShadowMaskSubPixelTexture, LightSceneInfo, bProjectingForForwardShading);
 
 	// Perform injection on translucent lighting volume
+	// Translucent volume is only used for direct lighting, so don't inject the lights if that is disabled
+	if (ViewFamily.EngineShowFlags.DirectLighting)
 	{
 		const TArray<FProjectedShadowInfo*, SceneRenderingAllocator>& ShadowMaps = VisibleLightInfo.ShadowsToProject;
 	
