Export some RayTracing functions from the FViewInfo

diff --git a/Engine/Source/Runtime/Renderer/Private/SceneRendering.h b/Engine/Source/Runtime/Renderer/Private/SceneRendering.h
--- a/Engine/Source/Runtime/Renderer/Private/SceneRendering.h
+++ b/Engine/Source/Runtime/Renderer/Private/SceneRendering.h
@@ -1416,10 +1416,10 @@ public:
 	FBufferRHIRef RayTracingSceneBuffer; // #yuriy_todo: move this to RDG transient buffer system
 	FGraphEventRef CreateRayTracingSceneTask;
 
-	bool HasRayTracingScene() const;
-	FRHIRayTracingScene* GetRayTracingScene() const;
-	FRHIRayTracingScene* GetRayTracingSceneChecked() const;
-	FRHIShaderResourceView* GetRayTracingSceneViewChecked() const;
+	RENDERER_API bool HasRayTracingScene() const;
+	RENDERER_API FRHIRayTracingScene* GetRayTracingScene() const;
+	RENDERER_API FRHIRayTracingScene* GetRayTracingSceneChecked() const;
+	RENDERER_API FRHIShaderResourceView* GetRayTracingSceneViewChecked() const;
 
 	// Primary pipeline state object to be used with the ray tracing scene for this view.
 	// Material shaders are only available when using this pipeline.
