Strata on RDG. That has fixed some transition issues.
 Next will be to have buffers also be RDG transient buffers (and that should fix the last transition issue).
Tested strata enabled and disabled with rhivalidation.

diff --git a/Engine/Shaders/Private/Strata/StrataMaterialClassification.usf b/Engine/Shaders/Private/Strata/StrataMaterialClassification.usf
--- a/Engine/Shaders/Private/Strata/StrataMaterialClassification.usf
+++ b/Engine/Shaders/Private/Strata/StrataMaterialClassification.usf
@@ -9,7 +9,10 @@
 #define TILE_SIZE 8
 #define GROUP_THREAD_COUNT (TILE_SIZE*TILE_SIZE)
 
-#if SHADER_CATEGORIZATION 
+#if SHADER_CATEGORIZATION
+
+uint MaxBytesPerPixel;
+ByteAddressBuffer MaterialLobesBuffer;
 
 // Note: 
 // This code is partial copy of PackStrataOut in Strata.ush. 
@@ -140,8 +143,8 @@ void MainPS(
 	FStrataSubsurfaceData SSSData				= (FStrataSubsurfaceData)0;
 	FStrataTopLayerData TopLayerData			= (FStrataTopLayerData)0;
 
-	FStrataAddressing StrataAddressing = GetStrataPixelDataByteOffset(PixelPos, uint2(View.BufferSizeAndInvSize.xy), Strata.MaxBytesPerPixel);
-	FStrataPixelHeader Header = UnpackStrataHeaderIn(Strata.MaterialLobesBuffer, StrataAddressing);
+	FStrataAddressing StrataAddressing = GetStrataPixelDataByteOffset(PixelPos, uint2(View.BufferSizeAndInvSize.xy), MaxBytesPerPixel);
+	FStrataPixelHeader Header = UnpackStrataHeaderIn(MaterialLobesBuffer, StrataAddressing);
 
 	ClassificationData.bIsStrataMaterial = Header.BSDFCount > 0;
 	ClassificationData.MaterialAO = 1.0;
@@ -160,7 +163,7 @@ void MainPS(
 		float TopLayerRoughnessAverageCount = 0.0f;
 		for (uint BSDFIndex = 0; BSDFIndex < Header.BSDFCount; ++BSDFIndex)
 		{
-			const FStrataBSDF BSDF = UnpackStrataBSDFIn(Strata.MaterialLobesBuffer, StrataAddressing);
+			const FStrataBSDF BSDF = UnpackStrataBSDFIn(MaterialLobesBuffer, StrataAddressing);
 			ComputePostProcessBSDFData(Header, BSDF, StrataAddressing, V, ClassificationData, SSSData, TopLayerData, TopLayerRoughness, TopLayerRoughnessAverageCount);
 		}
 
diff --git a/Engine/Source/Runtime/Renderer/Private/CompositionLighting/PostProcessAmbientOcclusion.cpp b/Engine/Source/Runtime/Renderer/Private/CompositionLighting/PostProcessAmbientOcclusion.cpp
--- a/Engine/Source/Runtime/Renderer/Private/CompositionLighting/PostProcessAmbientOcclusion.cpp
+++ b/Engine/Source/Runtime/Renderer/Private/CompositionLighting/PostProcessAmbientOcclusion.cpp
@@ -523,7 +523,7 @@ public:
 		SHADER_PARAMETER(float, ThresholdInverse)
 		SHADER_PARAMETER(FVector2D, InputExtentInverse)
 
-		SHADER_PARAMETER_STRUCT_REF(FStrataGlobalUniformParameters, Strata)
+		SHADER_PARAMETER_RDG_UNIFORM_BUFFER(FStrataGlobalUniformParameters, Strata)
 
 		RENDER_TARGET_BINDING_SLOTS()
 	END_SHADER_PARAMETER_STRUCT();
@@ -707,7 +707,7 @@ public:
 
 	BEGIN_SHADER_PARAMETER_STRUCT(FParameters, )
 		SHADER_PARAMETER_STRUCT_INCLUDE(FAmbientOcclusionParameters, SharedParameters)
-		SHADER_PARAMETER_STRUCT_REF(FStrataGlobalUniformParameters, Strata)
+		SHADER_PARAMETER_RDG_UNIFORM_BUFFER(FStrataGlobalUniformParameters, Strata)
 		RENDER_TARGET_BINDING_SLOTS()
 	END_GLOBAL_SHADER_PARAMETER_STRUCT();
 };
@@ -741,7 +741,7 @@ public:
 
 	BEGIN_SHADER_PARAMETER_STRUCT(FParameters, )
 		SHADER_PARAMETER_STRUCT_INCLUDE(FAmbientOcclusionParameters, SharedParameters)
-		SHADER_PARAMETER_STRUCT_REF(FStrataGlobalUniformParameters, Strata)
+		SHADER_PARAMETER_RDG_UNIFORM_BUFFER(FStrataGlobalUniformParameters, Strata)
 		SHADER_PARAMETER_RDG_TEXTURE_UAV(RWTexture2D, OutTexture)
 	END_GLOBAL_SHADER_PARAMETER_STRUCT();
 };
@@ -1092,7 +1092,7 @@ public:
 		SHADER_PARAMETER_STRUCT_REF(FViewUniformShaderParameters, View)
 		SHADER_PARAMETER_RDG_UNIFORM_BUFFER(FSceneTextureUniformParameters, SceneTextures)
 		SHADER_PARAMETER_STRUCT_INCLUDE(FHZBParameters, HZBParameters)
-		SHADER_PARAMETER_STRUCT_REF(FStrataGlobalUniformParameters, Strata)
+		SHADER_PARAMETER_RDG_UNIFORM_BUFFER(FStrataGlobalUniformParameters, Strata)
 
 		SHADER_PARAMETER_STRUCT_INCLUDE(FSSAOShaderParameters, SSAOParameters)
 		SHADER_PARAMETER_STRUCT_INCLUDE(FGTAOShaderParameters, GTAOParameters)
@@ -1184,7 +1184,7 @@ public:
 	BEGIN_SHADER_PARAMETER_STRUCT(FParameters, )
 		SHADER_PARAMETER_STRUCT_REF(FViewUniformShaderParameters, View)
 		SHADER_PARAMETER_RDG_UNIFORM_BUFFER(FSceneTextureUniformParameters, SceneTextures)
-		SHADER_PARAMETER_STRUCT_REF(FStrataGlobalUniformParameters, Strata)
+		SHADER_PARAMETER_RDG_UNIFORM_BUFFER(FStrataGlobalUniformParameters, Strata)
 
 		SHADER_PARAMETER_STRUCT_INCLUDE(FSSAOShaderParameters, SSAOParameters)
 		SHADER_PARAMETER_STRUCT_INCLUDE(FGTAOShaderParameters, GTAOParameters)
@@ -1282,7 +1282,7 @@ public:
 		SHADER_PARAMETER_STRUCT_REF(FViewUniformShaderParameters, View)
 		SHADER_PARAMETER_RDG_UNIFORM_BUFFER(FSceneTextureUniformParameters, SceneTextures)
 		SHADER_PARAMETER_STRUCT_INCLUDE(FHZBParameters, HZBParameters)
-		SHADER_PARAMETER_STRUCT_REF(FStrataGlobalUniformParameters, Strata)
+		SHADER_PARAMETER_RDG_UNIFORM_BUFFER(FStrataGlobalUniformParameters, Strata)
 
 		SHADER_PARAMETER_STRUCT_INCLUDE(FSSAOShaderParameters, SSAOParameters)
 		SHADER_PARAMETER_STRUCT_INCLUDE(FGTAOShaderParameters, GTAOParameters)
diff --git a/Engine/Source/Runtime/Renderer/Private/DeferredShadingRenderer.cpp b/Engine/Source/Runtime/Renderer/Private/DeferredShadingRenderer.cpp
--- a/Engine/Source/Runtime/Renderer/Private/DeferredShadingRenderer.cpp
+++ b/Engine/Source/Runtime/Renderer/Private/DeferredShadingRenderer.cpp
@@ -2245,12 +2245,9 @@ void FDeferredShadingSceneRenderer::Render(FRDGBuilder& GraphBuilder)
 		Scene->AllocateAndCaptureFrameSkyEnvMap(GraphBuilder, *this, MainView, bShouldRenderSkyAtmosphere, bShouldRenderVolumetricCloud, InstanceCullingManager);
 	}
 
-	// Strata initialisation
+	// Strata initialisation is always run even when not enabled.
 	const bool bStrataEnabled = Strata::IsStrataEnabled();
-	if (bStrataEnabled)
-	{		
-		Strata::InitialiseStrataFrameSceneData(*this, GraphBuilder);
-	}
+	Strata::InitialiseStrataFrameSceneData(*this, GraphBuilder);
 
 	if (GetCustomDepthPassLocation() == ECustomDepthPassLocation::BeforeBasePass)
 	{
diff --git a/Engine/Source/Runtime/Renderer/Private/DistanceFieldAmbientOcclusion.cpp b/Engine/Source/Runtime/Renderer/Private/DistanceFieldAmbientOcclusion.cpp
--- a/Engine/Source/Runtime/Renderer/Private/DistanceFieldAmbientOcclusion.cpp
+++ b/Engine/Source/Runtime/Renderer/Private/DistanceFieldAmbientOcclusion.cpp
@@ -278,7 +278,7 @@ public:
 	BEGIN_SHADER_PARAMETER_STRUCT(FParameters, )
 		SHADER_PARAMETER_STRUCT_REF(FViewUniformShaderParameters, View)
 		SHADER_PARAMETER_RDG_UNIFORM_BUFFER(FSceneTextureUniformParameters, SceneTextures)
-		SHADER_PARAMETER_STRUCT_REF(FStrataGlobalUniformParameters, Strata)
+		SHADER_PARAMETER_RDG_UNIFORM_BUFFER(FStrataGlobalUniformParameters, Strata)
 		RENDER_TARGET_BINDING_SLOTS()
 	END_SHADER_PARAMETER_STRUCT()
 
@@ -324,7 +324,7 @@ public:
 	BEGIN_SHADER_PARAMETER_STRUCT(FParameters, )
 		SHADER_PARAMETER_STRUCT_REF(FViewUniformShaderParameters, View)
 		SHADER_PARAMETER_RDG_UNIFORM_BUFFER(FSceneTextureUniformParameters, SceneTextures)
-		SHADER_PARAMETER_STRUCT_REF(FStrataGlobalUniformParameters, Strata)
+		SHADER_PARAMETER_RDG_UNIFORM_BUFFER(FStrataGlobalUniformParameters, Strata)
 		SHADER_PARAMETER_RDG_TEXTURE_UAV(RWTexture2D<float4>, RWDistanceFieldNormal)
 	END_SHADER_PARAMETER_STRUCT()
 
diff --git a/Engine/Source/Runtime/Renderer/Private/DistanceFieldShadowing.cpp b/Engine/Source/Runtime/Renderer/Private/DistanceFieldShadowing.cpp
--- a/Engine/Source/Runtime/Renderer/Private/DistanceFieldShadowing.cpp
+++ b/Engine/Source/Runtime/Renderer/Private/DistanceFieldShadowing.cpp
@@ -138,7 +138,7 @@ class FCullObjectsForShadowCS : public FGlobalShader
 
 	BEGIN_SHADER_PARAMETER_STRUCT(FParameters, )
 		SHADER_PARAMETER_STRUCT_REF(FViewUniformShaderParameters, View)
-		SHADER_PARAMETER_STRUCT_REF(FStrataGlobalUniformParameters, Strata)
+		SHADER_PARAMETER_RDG_UNIFORM_BUFFER(FStrataGlobalUniformParameters, Strata)
 		SHADER_PARAMETER_STRUCT_INCLUDE(FDistanceFieldObjectBufferParameters, ObjectBufferParameters)
 		SHADER_PARAMETER_STRUCT_INCLUDE(FDistanceFieldCulledObjectBufferParameters, CulledObjectBufferParameters)
 		SHADER_PARAMETER(uint32, ObjectBoundingGeometryIndexCount)
diff --git a/Engine/Source/Runtime/Renderer/Private/IndirectLightRendering.cpp b/Engine/Source/Runtime/Renderer/Private/IndirectLightRendering.cpp
--- a/Engine/Source/Runtime/Renderer/Private/IndirectLightRendering.cpp
+++ b/Engine/Source/Runtime/Renderer/Private/IndirectLightRendering.cpp
@@ -302,7 +302,7 @@ class FReflectionEnvironmentSkyLightingPS : public FGlobalShader
 		SHADER_PARAMETER_STRUCT_REF(FReflectionCaptureShaderData, ReflectionCaptureData)
 		SHADER_PARAMETER_STRUCT_REF(FForwardLightData, ForwardLightData)
 
-		SHADER_PARAMETER_STRUCT_REF(FStrataGlobalUniformParameters, Strata)
+		SHADER_PARAMETER_RDG_UNIFORM_BUFFER(FStrataGlobalUniformParameters, Strata)
 
 		RENDER_TARGET_BINDING_SLOTS()
 	END_SHADER_PARAMETER_STRUCT()
diff --git a/Engine/Source/Runtime/Renderer/Private/LightRendering.cpp b/Engine/Source/Runtime/Renderer/Private/LightRendering.cpp
--- a/Engine/Source/Runtime/Renderer/Private/LightRendering.cpp
+++ b/Engine/Source/Runtime/Renderer/Private/LightRendering.cpp
@@ -586,7 +586,8 @@ private:
 		FRenderLightParams* RenderLightParams)
 	{
 		FGlobalShader::SetParameters<FViewUniformShaderParameters>(RHICmdList, ShaderRHI, View.ViewUniformBuffer);
-		FGlobalShader::SetParameters<FStrataGlobalUniformParameters>(RHICmdList, ShaderRHI, Strata::BindStrataGlobalUniformParameters(View));
+		TRDGUniformBufferRef<FStrataGlobalUniformParameters> StrataUniformBuffer = Strata::BindStrataGlobalUniformParameters(View);
+		FGlobalShader::SetParameters<FStrataGlobalUniformParameters>(RHICmdList, ShaderRHI, StrataUniformBuffer->GetRHIRef());
 
 		if(LightAttenuationTexture.IsBound())
 		{
@@ -1081,6 +1082,7 @@ void FSceneRenderer::GatherAndSortLights(FSortedLightSetSceneInfo& OutSortedLigh
 BEGIN_SHADER_PARAMETER_STRUCT(FRenderLightParameters, )
 	SHADER_PARAMETER_RDG_UNIFORM_BUFFER(FSceneTextureUniformParameters, SceneTextures)
 	SHADER_PARAMETER_RDG_UNIFORM_BUFFER(FHairStrandsViewUniformParameters, HairStrands)
+	SHADER_PARAMETER_RDG_UNIFORM_BUFFER(FStrataGlobalUniformParameters, Strata)
 	SHADER_PARAMETER_STRUCT_INCLUDE(FVolumetricCloudShadowAOParameters, CloudShadowAO)
 	RDG_TEXTURE_ACCESS(ShadowMaskTexture, ERHIAccess::SRVGraphics)
 	RDG_TEXTURE_ACCESS(LightingChannelsTexture, ERHIAccess::SRVGraphics)
@@ -1088,6 +1090,7 @@ BEGIN_SHADER_PARAMETER_STRUCT(FRenderLightParameters, )
 END_SHADER_PARAMETER_STRUCT()
 
 void GetRenderLightParameters(
+	const FViewInfo& View,
 	FRDGTextureRef SceneColorTexture,
 	FRDGTextureRef SceneDepthTexture,
 	TRDGUniformBufferRef<FSceneTextureUniformParameters> SceneTexturesUniformBuffer,
@@ -1099,6 +1102,7 @@ void GetRenderLightParameters(
 {
 	Parameters.SceneTextures = SceneTexturesUniformBuffer;
 	Parameters.HairStrands = HairStrandsUniformBuffer;
+	Parameters.Strata = Strata::BindStrataGlobalUniformParameters(View);
 	Parameters.ShadowMaskTexture = ShadowMaskTexture;
 	Parameters.LightingChannelsTexture = LightingChannelsTexture;
 	Parameters.CloudShadowAO = CloudShadowAOParameters;
@@ -1113,6 +1117,7 @@ void GetRenderLightParameters(
 FHairStrandsTransmittanceMaskData CreateDummyHairStrandsTransmittanceMaskData(FRDGBuilder& GraphBuilder, FGlobalShaderMap* ShaderMap);
 
 void GetRenderLightParameters(
+	const FViewInfo& View,
 	const FMinimalSceneTextures& SceneTextures,
 	const FHairStrandsViewData& HairViewData,
 	FRDGTextureRef ShadowMaskTexture,
@@ -1120,7 +1125,7 @@ void GetRenderLightParameters(
 	const FVolumetricCloudShadowAOParameters& CloudShadowAOParameters,
 	FRenderLightParameters& Parameters)
 {
-	GetRenderLightParameters(SceneTextures.Color.Target, SceneTextures.Depth.Target, SceneTextures.UniformBuffer, HairViewData.UniformBuffer, ShadowMaskTexture, LightingChannelsTexture, CloudShadowAOParameters, Parameters);
+	GetRenderLightParameters(View, SceneTextures.Color.Target, SceneTextures.Depth.Target, SceneTextures.UniformBuffer, HairViewData.UniformBuffer, ShadowMaskTexture, LightingChannelsTexture, CloudShadowAOParameters, Parameters);
 }
 
 void FDeferredShadingSceneRenderer::RenderLights(
@@ -1236,7 +1241,7 @@ void FDeferredShadingSceneRenderer::RenderLights(
 				{
 					const FViewInfo& View = Views[ViewIndex];
 					FRenderLightParameters* PassParameters = GraphBuilder.AllocParameters<FRenderLightParameters>();
-					GetRenderLightParameters(SceneTextures, View.HairStrandsViewData, nullptr, LightingChannelsTexture, {}, *PassParameters);
+					GetRenderLightParameters(View, SceneTextures, View.HairStrandsViewData, nullptr, LightingChannelsTexture, {}, *PassParameters);
 
 					GraphBuilder.AddPass(
 						RDG_EVENT_NAME("StandardDeferredLighting"),
@@ -1955,7 +1960,7 @@ void FDeferredShadingSceneRenderer::RenderStationaryLightOverlap(
 	if (Scene->bIsEditorScene)
 	{
 		FRenderLightParameters* PassParameters = GraphBuilder.AllocParameters<FRenderLightParameters>();
-		GetRenderLightParameters(SceneTextures, Views[0].HairStrandsViewData, nullptr, LightingChannelsTexture, {}, *PassParameters);
+		GetRenderLightParameters(Views[0], SceneTextures, Views[0].HairStrandsViewData, nullptr, LightingChannelsTexture, {}, *PassParameters);
 
 		GraphBuilder.AddPass(
 			RDG_EVENT_NAME("StationaryLightOverlap"),
@@ -2334,7 +2339,7 @@ void FDeferredShadingSceneRenderer::RenderLight(
 		const FViewInfo& View = Views[ViewIndex];
 
 		FRenderLightParameters* PassParameters = GraphBuilder.AllocParameters<FRenderLightParameters>();
-		GetRenderLightParameters(SceneTextures, View.HairStrandsViewData, ScreenShadowMaskTexture, LightingChannelsTexture, GetCloudShadowAOParameters(GraphBuilder, View, CloudInfo), *PassParameters);
+		GetRenderLightParameters(View, SceneTextures, View.HairStrandsViewData, ScreenShadowMaskTexture, LightingChannelsTexture, GetCloudShadowAOParameters(GraphBuilder, View, CloudInfo), *PassParameters);
 
 		GraphBuilder.AddPass(
 			RDG_EVENT_NAME("StandardDeferredLighting"),
@@ -2402,6 +2407,7 @@ void FDeferredShadingSceneRenderer::RenderLightForHair(
 
 		FRenderLightForHairParameters* PassParameters = GraphBuilder.AllocParameters<FRenderLightForHairParameters>();
 		GetRenderLightParameters(
+			View,
 			HairVisibilityData.SampleLightingBuffer, 
 			nullptr, 
 			SceneTexturesUniformBuffer, 
diff --git a/Engine/Source/Runtime/Renderer/Private/PostProcess/PostProcessSubsurface.cpp b/Engine/Source/Runtime/Renderer/Private/PostProcess/PostProcessSubsurface.cpp
--- a/Engine/Source/Runtime/Renderer/Private/PostProcess/PostProcessSubsurface.cpp
+++ b/Engine/Source/Runtime/Renderer/Private/PostProcess/PostProcessSubsurface.cpp
@@ -371,7 +371,7 @@ public:
 	SHADER_USE_PARAMETER_STRUCT(FSubsurfaceVisualizePS, FSubsurfaceShader);
 
 	BEGIN_SHADER_PARAMETER_STRUCT(FParameters, )
-		SHADER_PARAMETER_STRUCT_REF(FStrataGlobalUniformParameters, Strata)
+		SHADER_PARAMETER_RDG_UNIFORM_BUFFER(FStrataGlobalUniformParameters, Strata)
 		SHADER_PARAMETER_STRUCT_INCLUDE(FSubsurfaceParameters, Subsurface)
 		SHADER_PARAMETER_STRUCT(FSubsurfaceInput, SubsurfaceInput0)
 		SHADER_PARAMETER_TEXTURE(Texture2D, MiniFontTexture)
@@ -490,7 +490,7 @@ public:
 	}
 
 	BEGIN_SHADER_PARAMETER_STRUCT(FParameters, )
-		SHADER_PARAMETER_STRUCT_REF(FStrataGlobalUniformParameters, Strata)
+		SHADER_PARAMETER_RDG_UNIFORM_BUFFER(FStrataGlobalUniformParameters, Strata)
 		SHADER_PARAMETER_STRUCT_INCLUDE(FSubsurfaceParameters, Subsurface)
 		SHADER_PARAMETER_STRUCT(FScreenPassTextureViewportParameters, Output)
 		SHADER_PARAMETER_STRUCT(FSubsurfaceInput, SubsurfaceInput0)
@@ -527,7 +527,7 @@ public:
 	}
 
 	BEGIN_SHADER_PARAMETER_STRUCT(FParameters, )
-		SHADER_PARAMETER_STRUCT_REF(FStrataGlobalUniformParameters, Strata)
+		SHADER_PARAMETER_RDG_UNIFORM_BUFFER(FStrataGlobalUniformParameters, Strata)
 		SHADER_PARAMETER_STRUCT_INCLUDE(FSubsurfaceParameters, Subsurface)
 		SHADER_PARAMETER_STRUCT(FScreenPassTextureViewportParameters, Output)
 		SHADER_PARAMETER_RDG_TEXTURE_UAV(RWTexture2D, SSSColorUAV)
@@ -752,7 +752,7 @@ class FSubsurfaceRecombinePS : public FSubsurfaceShader
 	SHADER_USE_PARAMETER_STRUCT(FSubsurfaceRecombinePS, FSubsurfaceShader);
 
 	BEGIN_SHADER_PARAMETER_STRUCT(FParameters, )
-		SHADER_PARAMETER_STRUCT_REF(FStrataGlobalUniformParameters, Strata)
+		SHADER_PARAMETER_RDG_UNIFORM_BUFFER(FStrataGlobalUniformParameters, Strata)
 		SHADER_PARAMETER_STRUCT_INCLUDE(FSubsurfaceParameters, Subsurface)
 		SHADER_PARAMETER_STRUCT(FSubsurfaceInput, SubsurfaceInput0)
 		SHADER_PARAMETER_STRUCT(FSubsurfaceInput, SubsurfaceInput1)
diff --git a/Engine/Source/Runtime/Renderer/Private/ScreenSpaceRayTracing.cpp b/Engine/Source/Runtime/Renderer/Private/ScreenSpaceRayTracing.cpp
--- a/Engine/Source/Runtime/Renderer/Private/ScreenSpaceRayTracing.cpp
+++ b/Engine/Source/Runtime/Renderer/Private/ScreenSpaceRayTracing.cpp
@@ -461,7 +461,7 @@ class FScreenSpaceReflectionsStencilPS : public FGlobalShader
 	
 	BEGIN_SHADER_PARAMETER_STRUCT(FParameters, )
 		SHADER_PARAMETER_STRUCT_INCLUDE(FSSRCommonParameters, CommonParameters)
-		SHADER_PARAMETER_STRUCT_REF(FStrataGlobalUniformParameters, Strata)
+		SHADER_PARAMETER_RDG_UNIFORM_BUFFER(FStrataGlobalUniformParameters, Strata)
 		RENDER_TARGET_BINDING_SLOTS()
 	END_SHADER_PARAMETER_STRUCT()
 };
@@ -491,7 +491,7 @@ class FScreenSpaceReflectionsPS : public FGlobalShader
 		SHADER_PARAMETER_STRUCT_INCLUDE(FSSRPassCommonParameters, SSRPassCommonParameter)
 		RDG_BUFFER_ACCESS(IndirectDrawParameter, ERHIAccess::IndirectArgs)
 		SHADER_PARAMETER_RDG_BUFFER_SRV(StructuredBuffer<uint>, TileListData)		// FScreenSpaceReflectionsTileVS
-		SHADER_PARAMETER_STRUCT_REF(FStrataGlobalUniformParameters, Strata)
+		SHADER_PARAMETER_RDG_UNIFORM_BUFFER(FStrataGlobalUniformParameters, Strata)
 		RENDER_TARGET_BINDING_SLOTS()
 	END_SHADER_PARAMETER_STRUCT()
 };
diff --git a/Engine/Source/Runtime/Renderer/Private/ShadowRendering.cpp b/Engine/Source/Runtime/Renderer/Private/ShadowRendering.cpp
--- a/Engine/Source/Runtime/Renderer/Private/ShadowRendering.cpp
+++ b/Engine/Source/Runtime/Renderer/Private/ShadowRendering.cpp
@@ -422,8 +422,8 @@ static void BindShaderShaders(FRHICommandList& RHICmdList, FGraphicsPipelineStat
 	if (Strata::IsStrataEnabled())
 	{
 		FRHIPixelShader* ShaderRHI = PixelShader.GetPixelShader();
-		FRHIUniformBuffer* StrataUniformBuffer = Strata::BindStrataGlobalUniformParameters(View);
-		PixelShader->FGlobalShader::template SetParameters<FStrataGlobalUniformParameters>(RHICmdList, ShaderRHI, StrataUniformBuffer);
+		TRDGUniformBufferRef<FStrataGlobalUniformParameters> StrataUniformBuffer = Strata::BindStrataGlobalUniformParameters(View);
+		PixelShader->FGlobalShader::template SetParameters<FStrataGlobalUniformParameters>(RHICmdList, ShaderRHI, StrataUniformBuffer->GetRHIRef());
 	}
 
 	if (HairStrandsUniformBuffer)
@@ -989,6 +989,7 @@ void FProjectedShadowInfo::SetupProjectionStencilMask(
 BEGIN_SHADER_PARAMETER_STRUCT(FShadowProjectionPassParameters, )
 	SHADER_PARAMETER_STRUCT_INCLUDE(FViewShaderParameters, View)
 	SHADER_PARAMETER_RDG_UNIFORM_BUFFER(FHairStrandsViewUniformParameters, HairStrands)
+	SHADER_PARAMETER_RDG_UNIFORM_BUFFER(FStrataGlobalUniformParameters, Strata)
 	SHADER_PARAMETER_STRUCT_INCLUDE(FSceneTextureShaderParameters, SceneTextures)
 	SHADER_PARAMETER_STRUCT_INCLUDE(FInstanceCullingDrawParams, InstanceCullingDrawParams)
 	RDG_TEXTURE_ACCESS(ShadowTexture0, ERHIAccess::SRVGraphics)
@@ -1880,6 +1881,7 @@ void FSceneRenderer::RenderShadowProjections(
 		FShadowProjectionPassParameters CommonPassParameters;
 		CommonPassParameters.SceneTextures = GetSceneTextureShaderParameters(SceneTextures.UniformBuffer);
 		CommonPassParameters.HairStrands = HairStrands::BindHairStrandsViewUniformParameters(View);
+		CommonPassParameters.Strata = Strata::BindStrataGlobalUniformParameters(View);
 		CommonPassParameters.RenderTargets[0] = FRenderTargetBinding(OutputTexture, ERenderTargetLoadAction::ELoad);
 		CommonPassParameters.RenderTargets.DepthStencil =
 			bSubPixelShadow ?
diff --git a/Engine/Source/Runtime/Renderer/Private/ShadowRendering.h b/Engine/Source/Runtime/Renderer/Private/ShadowRendering.h
--- a/Engine/Source/Runtime/Renderer/Private/ShadowRendering.h
+++ b/Engine/Source/Runtime/Renderer/Private/ShadowRendering.h
@@ -1507,8 +1507,8 @@ public:
 
 		if (StrataGlobalParameters.IsBound())
 		{
-			TUniformBufferRef<FStrataGlobalUniformParameters> StrataGlobalUniformParameters = Strata::BindStrataGlobalUniformParameters(View);
-			SetUniformBufferParameter(RHICmdList, ShaderRHI, StrataGlobalParameters, StrataGlobalUniformParameters);
+			TRDGUniformBufferRef<FStrataGlobalUniformParameters> StrataUniformBuffer = Strata::BindStrataGlobalUniformParameters(View);
+			SetUniformBufferParameter(RHICmdList, ShaderRHI, StrataGlobalParameters, StrataUniformBuffer->GetRHIRef());
 		}
 
 		{
diff --git a/Engine/Source/Runtime/Renderer/Private/SingleLayerWaterRendering.cpp b/Engine/Source/Runtime/Renderer/Private/SingleLayerWaterRendering.cpp
--- a/Engine/Source/Runtime/Renderer/Private/SingleLayerWaterRendering.cpp
+++ b/Engine/Source/Runtime/Renderer/Private/SingleLayerWaterRendering.cpp
@@ -139,7 +139,7 @@ BEGIN_SHADER_PARAMETER_STRUCT(FSingleLayerWaterCommonShaderParameters, )
 	SHADER_PARAMETER_STRUCT_REF(FReflectionCaptureShaderData, ReflectionCaptureData)
 	SHADER_PARAMETER_STRUCT_REF(FReflectionUniformParameters, ReflectionsParameters)
 	SHADER_PARAMETER_STRUCT_REF(FForwardLightData, ForwardLightData)
-	SHADER_PARAMETER_STRUCT_REF(FStrataGlobalUniformParameters, Strata)
+	SHADER_PARAMETER_RDG_UNIFORM_BUFFER(FStrataGlobalUniformParameters, Strata)
 END_SHADER_PARAMETER_STRUCT()
 
 class FSingleLayerWaterCompositePS : public FGlobalShader
diff --git a/Engine/Source/Runtime/Renderer/Private/Strata/Strata.cpp b/Engine/Source/Runtime/Renderer/Private/Strata/Strata.cpp
--- a/Engine/Source/Runtime/Renderer/Private/Strata/Strata.cpp
+++ b/Engine/Source/Runtime/Renderer/Private/Strata/Strata.cpp
@@ -76,6 +76,19 @@ static TAutoConsoleVariable<int32> CVarStrataFurnaceTestSampleCount(
 
 IMPLEMENT_GLOBAL_SHADER_PARAMETER_STRUCT(FStrataGlobalUniformParameters, "Strata");
 
+
+
+void FStrataSceneData::Reset()
+{
+	ClassificationTexture = nullptr;
+	TopLayerNormalTexture = nullptr;
+	SSSTexture = nullptr;
+
+	StrataGlobalUniformParameters = nullptr;
+}
+
+
+
 namespace Strata
 {
 
@@ -112,6 +125,7 @@ static void AddStrataLUTPass(FRDGBuilder& GraphBuilder, const FViewInfo& View);
 void InitialiseStrataFrameSceneData(FSceneRenderer& SceneRenderer, FRDGBuilder& GraphBuilder)
 {
 	FStrataSceneData& StrataSceneData = SceneRenderer.Scene->StrataSceneData;
+	StrataSceneData.Reset();
 
 	uint32 ResolutionX = 1;
 	uint32 ResolutionY = 1;
@@ -130,9 +144,7 @@ void InitialiseStrataFrameSceneData(FSceneRenderer& SceneRenderer, FRDGBuilder&
 
 		// Classification texture
 		{
-			FRDGTextureRef Texture = GraphBuilder.CreateTexture(FRDGTextureDesc::Create2D(BufferSizeXY, PF_R32_UINT, FClearValueBinding::Black, TexCreate_ShaderResource | TexCreate_RenderTargetable), TEXT("StrataClassificationTexture"));
-			GraphBuilder.PreallocateTexture(Texture);
-			StrataSceneData.ClassificationTexture = GraphBuilder.GetPooledTexture(Texture);
+			StrataSceneData.ClassificationTexture = GraphBuilder.CreateTexture(FRDGTextureDesc::Create2D(BufferSizeXY, PF_R32_UINT, FClearValueBinding::Black, TexCreate_ShaderResource | TexCreate_RenderTargetable), TEXT("StrataClassificationTexture"));
 		}
 
 		// Tile classification buffers
@@ -152,16 +164,12 @@ void InitialiseStrataFrameSceneData(FSceneRenderer& SceneRenderer, FRDGBuilder&
 
 		// Top layer texture
 		{
-			FRDGTextureRef Texture = GraphBuilder.CreateTexture(FRDGTextureDesc::Create2D(BufferSizeXY, PF_R32_UINT, FClearValueBinding::Black, TexCreate_ShaderResource | TexCreate_RenderTargetable), TEXT("StrataTopLayerNormalTexture"));
-			GraphBuilder.PreallocateTexture(Texture);
-			StrataSceneData.TopLayerNormalTexture = GraphBuilder.GetPooledTexture(Texture);
+			StrataSceneData.TopLayerNormalTexture = GraphBuilder.CreateTexture(FRDGTextureDesc::Create2D(BufferSizeXY, PF_R32_UINT, FClearValueBinding::Black, TexCreate_ShaderResource | TexCreate_RenderTargetable), TEXT("StrataTopLayerNormalTexture"));
 		}
 
 		// SSS texture
 		{
-			FRDGTextureRef Texture = GraphBuilder.CreateTexture(FRDGTextureDesc::Create2D(BufferSizeXY, PF_R32G32_UINT, FClearValueBinding::Black, TexCreate_ShaderResource | TexCreate_RenderTargetable), TEXT("StrataSSSTexture"));
-			GraphBuilder.PreallocateTexture(Texture);
-			StrataSceneData.SSSTexture = GraphBuilder.GetPooledTexture(Texture);
+			StrataSceneData.SSSTexture = GraphBuilder.CreateTexture(FRDGTextureDesc::Create2D(BufferSizeXY, PF_R32G32_UINT, FClearValueBinding::Black, TexCreate_ShaderResource | TexCreate_RenderTargetable), TEXT("StrataSSSTexture"));
 		}
 
 		// Energy LUT
@@ -202,17 +210,52 @@ void InitialiseStrataFrameSceneData(FSceneRenderer& SceneRenderer, FRDGBuilder&
 		View.StrataSceneData = &SceneRenderer.Scene->StrataSceneData;
 	}
 
-	// Always reset the strata uniform buffer
-	SceneRenderer.Scene->StrataSceneData.StrataGlobalUniformParameters.SafeRelease();
-
-	AddStrataClearMaterialBufferPass(GraphBuilder, StrataSceneData.MaterialLobesBuffer.UAV, StrataSceneData.MaxBytesPerPixel, FIntPoint(ResolutionX, ResolutionY));
+	if (IsStrataEnabled())
+	{
+		AddStrataClearMaterialBufferPass(GraphBuilder, StrataSceneData.MaterialLobesBuffer.UAV, StrataSceneData.MaxBytesPerPixel, FIntPoint(ResolutionX, ResolutionY));
+	}
 
 	if (bUpdateLUT)
 	{
 		for (int32 ViewIndex = 0; ViewIndex < SceneRenderer.Views.Num(); ViewIndex++)
 		{
 			FViewInfo& View = SceneRenderer.Views[ViewIndex];
-			AddStrataLUTPass(GraphBuilder, View);
+			AddStrataLUTPass(GraphBuilder, View); // STRATA_TODO refactor because this should not be done for each view but only once.
+		}
+	}
+
+	// Create the readable uniform buffers for each views once for all
+	for (int32 ViewIndex = 0; ViewIndex < SceneRenderer.Views.Num(); ViewIndex++)
+	{
+		FViewInfo& View = SceneRenderer.Views[ViewIndex];
+
+		FStrataGlobalUniformParameters* StrataUniformParameters = GraphBuilder.AllocParameters<FStrataGlobalUniformParameters>();
+		if (IsStrataEnabled())
+		{
+			StrataUniformParameters->MaxBytesPerPixel = View.StrataSceneData->MaxBytesPerPixel;
+			StrataUniformParameters->MaterialLobesBuffer = View.StrataSceneData->MaterialLobesBuffer.SRV;
+			StrataUniformParameters->ClassificationTexture = View.StrataSceneData->ClassificationTexture;
+			StrataUniformParameters->TopLayerNormalTexture = View.StrataSceneData->TopLayerNormalTexture;
+			StrataUniformParameters->SSSTexture = View.StrataSceneData->SSSTexture;
+			StrataUniformParameters->GGXEnergyLUT3DTexture = View.StrataSceneData->GGXEnergyLUT3DTexture->GetRenderTargetItem().ShaderResourceTexture;
+			StrataUniformParameters->GGXEnergyLUT2DTexture = View.StrataSceneData->GGXEnergyLUT2DTexture->GetRenderTargetItem().ShaderResourceTexture;
+			StrataUniformParameters->GGXEnergyLUTSampler = TStaticSamplerState<SF_Bilinear, AM_Clamp, AM_Clamp, AM_Clamp>::GetRHI();
+			StrataUniformParameters->GGXEnergyLUTScaleBias = GetStrataGGXEnergyLUTScaleBias();
+			View.StrataSceneData->StrataGlobalUniformParameters = GraphBuilder.CreateUniformBuffer(StrataUniformParameters);
+		}
+		else
+		{
+			// Create each time. This path will go away when Strata is always enabled anyway.
+			StrataUniformParameters->MaxBytesPerPixel = 0;
+			StrataUniformParameters->MaterialLobesBuffer = GEmptyVertexBufferWithUAV->ShaderResourceViewRHI;
+			StrataUniformParameters->ClassificationTexture = GSystemTextures.GetZeroUIntDummy(GraphBuilder);
+			StrataUniformParameters->TopLayerNormalTexture = GSystemTextures.GetZeroUIntDummy(GraphBuilder);
+			StrataUniformParameters->SSSTexture = GSystemTextures.GetZeroUIntDummy(GraphBuilder);
+			StrataUniformParameters->GGXEnergyLUT3DTexture = GSystemTextures.VolumetricBlackDummy->GetRenderTargetItem().ShaderResourceTexture;
+			StrataUniformParameters->GGXEnergyLUT2DTexture = GSystemTextures.BlackDummy->GetRenderTargetItem().ShaderResourceTexture;
+			StrataUniformParameters->GGXEnergyLUTSampler = TStaticSamplerState<SF_Bilinear, AM_Clamp, AM_Clamp, AM_Clamp>::GetRHI();
+			StrataUniformParameters->GGXEnergyLUTScaleBias = GetStrataGGXEnergyLUTScaleBias();
+			View.StrataSceneData->StrataGlobalUniformParameters = GraphBuilder.CreateUniformBuffer(StrataUniformParameters);
 		}
 	}
 }
@@ -231,43 +274,12 @@ void BindStrataBasePassUniformParameters(const FViewInfo& View, FStrataOpaquePas
 	}
 }
 
-TUniformBufferRef<FStrataGlobalUniformParameters> BindStrataGlobalUniformParameters(const FViewInfo& View)
+TRDGUniformBufferRef<FStrataGlobalUniformParameters> BindStrataGlobalUniformParameters(const FViewInfo& View)
 {
-	// If the strata scene data has not been created this frame yet, create it.
-	FStrataGlobalUniformParameters StrataUniformParameters;
-	if (View.StrataSceneData)
-	{
-		if (View.StrataSceneData->StrataGlobalUniformParameters.IsValid())
-		{
-			return View.StrataSceneData->StrataGlobalUniformParameters;
-		}
-
-		StrataUniformParameters.MaxBytesPerPixel = View.StrataSceneData->MaxBytesPerPixel;
-		StrataUniformParameters.MaterialLobesBuffer = View.StrataSceneData->MaterialLobesBuffer.SRV;
-		StrataUniformParameters.ClassificationTexture = View.StrataSceneData->ClassificationTexture->GetRenderTargetItem().ShaderResourceTexture;
-		StrataUniformParameters.TopLayerNormalTexture = View.StrataSceneData->TopLayerNormalTexture->GetRenderTargetItem().ShaderResourceTexture;
-		StrataUniformParameters.GGXEnergyLUT3DTexture = View.StrataSceneData->GGXEnergyLUT3DTexture->GetRenderTargetItem().ShaderResourceTexture;
-		StrataUniformParameters.GGXEnergyLUT2DTexture = View.StrataSceneData->GGXEnergyLUT2DTexture->GetRenderTargetItem().ShaderResourceTexture;
-		StrataUniformParameters.GGXEnergyLUTSampler = TStaticSamplerState<SF_Bilinear, AM_Clamp, AM_Clamp, AM_Clamp>::GetRHI();
-		StrataUniformParameters.GGXEnergyLUTScaleBias = GetStrataGGXEnergyLUTScaleBias();
-		StrataUniformParameters.SSSTexture = View.StrataSceneData->SSSTexture->GetRenderTargetItem().ShaderResourceTexture;
-		View.StrataSceneData->StrataGlobalUniformParameters = CreateUniformBufferImmediate(StrataUniformParameters, UniformBuffer_SingleFrame);
-		return View.StrataSceneData->StrataGlobalUniformParameters;
-	}
-	else
-	{
-		// Create each time. This path will go away when Strata is always enabled anyway.
-		StrataUniformParameters.MaxBytesPerPixel = 0;
-		StrataUniformParameters.MaterialLobesBuffer = GEmptyVertexBufferWithUAV->ShaderResourceViewRHI;
-		StrataUniformParameters.ClassificationTexture = GSystemTextures.ZeroUIntDummy->GetRenderTargetItem().ShaderResourceTexture;
-		StrataUniformParameters.TopLayerNormalTexture = GSystemTextures.ZeroUIntDummy->GetRenderTargetItem().ShaderResourceTexture;
-		StrataUniformParameters.SSSTexture = GSystemTextures.ZeroUIntDummy->GetRenderTargetItem().ShaderResourceTexture;
-		StrataUniformParameters.GGXEnergyLUT3DTexture = GSystemTextures.VolumetricBlackDummy->GetRenderTargetItem().ShaderResourceTexture;
-		StrataUniformParameters.GGXEnergyLUT2DTexture = GSystemTextures.BlackDummy->GetRenderTargetItem().ShaderResourceTexture;
-		StrataUniformParameters.GGXEnergyLUTSampler = TStaticSamplerState<SF_Bilinear, AM_Clamp, AM_Clamp, AM_Clamp>::GetRHI();
-		StrataUniformParameters.GGXEnergyLUTScaleBias = GetStrataGGXEnergyLUTScaleBias();
-		return CreateUniformBufferImmediate(StrataUniformParameters, UniformBuffer_SingleDraw);
-	}
+	// We must always have a strata uniform buffer foreach view even when strata is not enabled.
+	check(View.StrataSceneData != nullptr);
+	check(View.StrataSceneData->StrataGlobalUniformParameters != nullptr);
+	return View.StrataSceneData->StrataGlobalUniformParameters;
 }
 
 ////////////////////////////////////////////////////////////////////////// 
@@ -283,7 +295,7 @@ class FVisualizeMaterialPS : public FGlobalShader
 
 	BEGIN_SHADER_PARAMETER_STRUCT(FParameters, )
 		SHADER_PARAMETER_STRUCT_REF(FViewUniformShaderParameters, ViewUniformBuffer)
-		SHADER_PARAMETER_STRUCT_REF(FStrataGlobalUniformParameters, Strata)
+		SHADER_PARAMETER_RDG_UNIFORM_BUFFER(FStrataGlobalUniformParameters, Strata)
 		SHADER_PARAMETER_TEXTURE(Texture2D, MiniFontTexture)
 		SHADER_PARAMETER_STRUCT_INCLUDE(FSceneTextureParameters, SceneTextures)
 		SHADER_PARAMETER_STRUCT_INCLUDE(ShaderDrawDebug::FShaderDrawDebugParameters, ShaderDrawParameters)
@@ -363,7 +375,8 @@ class FStrataMaterialClassificationPassPS : public FGlobalShader
 
 	BEGIN_SHADER_PARAMETER_STRUCT(FParameters, )
 		SHADER_PARAMETER_STRUCT_REF(FViewUniformShaderParameters, ViewUniformBuffer)
-		SHADER_PARAMETER_STRUCT_REF(FStrataGlobalUniformParameters, Strata)
+		SHADER_PARAMETER(uint32, MaxBytesPerPixel)
+		SHADER_PARAMETER_SRV(ByteAddressBuffer, MaterialLobesBuffer)
 		SHADER_PARAMETER_STRUCT_INCLUDE(FSceneTextureParameters, SceneTextures)
 		SHADER_PARAMETER_STRUCT_INCLUDE(ShaderDrawDebug::FShaderDrawDebugParameters, ShaderDrawParameters)
 		RENDER_TARGET_BINDING_SLOTS()
@@ -655,17 +668,17 @@ void AddStrataMaterialClassificationPass(FRDGBuilder& GraphBuilder, const FMinim
 		const FViewInfo& View = Views[i];
 		
 		// Classification
-		FRDGTextureRef ClassificationTexture = GraphBuilder.RegisterExternalTexture(View.StrataSceneData->ClassificationTexture);
 		{
 			FStrataMaterialClassificationPassPS::FPermutationDomain PermutationVector;
 			TShaderMapRef<FStrataMaterialClassificationPassPS> PixelShader(View.ShaderMap, PermutationVector);
 			FStrataMaterialClassificationPassPS::FParameters* PassParameters = GraphBuilder.AllocParameters<FStrataMaterialClassificationPassPS::FParameters>();
 			PassParameters->ViewUniformBuffer = View.ViewUniformBuffer;
-			PassParameters->Strata = Strata::BindStrataGlobalUniformParameters(View);
+			PassParameters->MaxBytesPerPixel = View.StrataSceneData->MaxBytesPerPixel;
+			PassParameters->MaterialLobesBuffer = View.StrataSceneData->MaterialLobesBuffer.SRV;
 			PassParameters->SceneTextures = GetSceneTextureParameters(GraphBuilder);		
-			PassParameters->RenderTargets[0] = FRenderTargetBinding(ClassificationTexture, ERenderTargetLoadAction::EClear);
-			PassParameters->RenderTargets[1] = FRenderTargetBinding(GraphBuilder.RegisterExternalTexture(View.StrataSceneData->TopLayerNormalTexture), ERenderTargetLoadAction::EClear);
-			PassParameters->RenderTargets[2] = FRenderTargetBinding(GraphBuilder.RegisterExternalTexture(View.StrataSceneData->SSSTexture), ERenderTargetLoadAction::EClear);
+			PassParameters->RenderTargets[0] = FRenderTargetBinding(View.StrataSceneData->ClassificationTexture, ERenderTargetLoadAction::EClear);
+			PassParameters->RenderTargets[1] = FRenderTargetBinding(View.StrataSceneData->TopLayerNormalTexture, ERenderTargetLoadAction::EClear);
+			PassParameters->RenderTargets[2] = FRenderTargetBinding(View.StrataSceneData->SSSTexture, ERenderTargetLoadAction::EClear);
 
 			if (ShaderDrawDebug::IsShaderDrawDebugEnabled())
 			{
@@ -689,7 +702,7 @@ void AddStrataMaterialClassificationPass(FRDGBuilder& GraphBuilder, const FMinim
 			PassParameters->TileSize = GetStrataTileSize();
 			PassParameters->bRectPrimitive = GRHISupportsRectTopology ? 1 : 0;
 			PassParameters->ViewResolution = View.ViewRect.Size();
-			PassParameters->ClassificationTexture = ClassificationTexture;
+			PassParameters->ClassificationTexture = View.StrataSceneData->ClassificationTexture;
 			PassParameters->TileListData = GraphBuilder.CreateUAV(GraphBuilder.RegisterExternalBuffer(View.StrataSceneData->ClassificationTileListBuffer), PF_R32_UINT);
 			PassParameters->TileIndirectData = GraphBuilder.CreateUAV(GraphBuilder.RegisterExternalBuffer(View.StrataSceneData->ClassificationTileIndirectBuffer), PF_R32_UINT);
 
@@ -700,7 +713,7 @@ void AddStrataMaterialClassificationPass(FRDGBuilder& GraphBuilder, const FMinim
 				RDG_EVENT_NAME("StrataMaterialTileClassification"),
 				ComputeShader,
 				PassParameters,
-				FComputeShaderUtils::GetGroupCount(ClassificationTexture->Desc.Extent, GroupSize));
+				FComputeShaderUtils::GetGroupCount(View.StrataSceneData->ClassificationTexture->Desc.Extent, GroupSize));
 		}
 	}
 }
@@ -793,6 +806,10 @@ static void AddStrataLUTPass(FRDGBuilder& GraphBuilder, const FViewInfo& View)
 		PixelShader,
 		Parameters,
 		View.ViewRect);
+
+	// Export to untracked because these textures lifetime is multiple frame, and also make sure transition is correctly done.
+	ConvertToUntrackedExternalTexture(GraphBuilder, OutLUT3D, View.StrataSceneData->GGXEnergyLUT3DTexture, ERHIAccess::SRVMask);
+	ConvertToUntrackedExternalTexture(GraphBuilder, OutLUT2D, View.StrataSceneData->GGXEnergyLUT2DTexture, ERHIAccess::SRVMask);
 }
 
 ////////////////////////////////////////////////////////////////////////////////////////////////////////////////////////////////////////////////
@@ -806,7 +823,7 @@ class FStrataFurnaceTestPassPS : public FGlobalShader
 
 	BEGIN_SHADER_PARAMETER_STRUCT(FParameters, )
 		SHADER_PARAMETER_STRUCT_REF(FViewUniformShaderParameters, ViewUniformBuffer)
-		SHADER_PARAMETER_STRUCT_REF(FStrataGlobalUniformParameters, Strata)
+		SHADER_PARAMETER_RDG_UNIFORM_BUFFER(FStrataGlobalUniformParameters, Strata)
 		SHADER_PARAMETER(uint32, NumSamples)
 		SHADER_PARAMETER(uint32, SceneType)
 		SHADER_PARAMETER(uint32, IntegratorType)
diff --git a/Engine/Source/Runtime/Renderer/Private/Strata/Strata.h b/Engine/Source/Runtime/Renderer/Private/Strata/Strata.h
--- a/Engine/Source/Runtime/Renderer/Private/Strata/Strata.h
+++ b/Engine/Source/Runtime/Renderer/Private/Strata/Strata.h
@@ -23,9 +23,9 @@ BEGIN_GLOBAL_SHADER_PARAMETER_STRUCT(FStrataGlobalUniformParameters, )
 	SHADER_PARAMETER(uint32, MaxBytesPerPixel)
 	SHADER_PARAMETER(FVector2D, GGXEnergyLUTScaleBias)
 	SHADER_PARAMETER_SRV(ByteAddressBuffer, MaterialLobesBuffer)
-	SHADER_PARAMETER_TEXTURE(Texture2D<uint>, ClassificationTexture)
-	SHADER_PARAMETER_TEXTURE(Texture2D<uint>, TopLayerNormalTexture)
-	SHADER_PARAMETER_TEXTURE(Texture2D<uint2>, SSSTexture)
+	SHADER_PARAMETER_RDG_TEXTURE(Texture2D<uint>, ClassificationTexture)
+	SHADER_PARAMETER_RDG_TEXTURE(Texture2D<uint>, TopLayerNormalTexture)
+	SHADER_PARAMETER_RDG_TEXTURE(Texture2D<uint2>, SSSTexture)
 	SHADER_PARAMETER_TEXTURE(Texture3D<float2>, GGXEnergyLUT3DTexture)
 	SHADER_PARAMETER_TEXTURE(Texture2D<float4>, GGXEnergyLUT2DTexture)
 	SHADER_PARAMETER_SAMPLER(SamplerState, GGXEnergyLUTSampler)
@@ -34,21 +34,28 @@ END_GLOBAL_SHADER_PARAMETER_STRUCT()
 struct FStrataSceneData
 {
 	uint32 MaxBytesPerPixel;
+
+	// Resources allocated and updated each frame
 	FRWByteAddressBuffer MaterialLobesBuffer;					// This should be a RDG resource
-	TRefCountPtr<IPooledRenderTarget> ClassificationTexture;	// This should be a RDG resource
-	TRefCountPtr<IPooledRenderTarget> TopLayerNormalTexture;	// This should be a RDG resource
+	FRDGTextureRef ClassificationTexture;
+	FRDGTextureRef TopLayerNormalTexture;
+	FRDGTextureRef SSSTexture;	
+
+	// Resources computed once for multiple frames
 	TRefCountPtr<IPooledRenderTarget> GGXEnergyLUT3DTexture;
 	TRefCountPtr<IPooledRenderTarget> GGXEnergyLUT2DTexture;
-	TRefCountPtr<IPooledRenderTarget> SSSTexture;				// This should be a RDG resource
 
 	TRefCountPtr<FRDGPooledBuffer> ClassificationTileListBuffer;
 	TRefCountPtr<FRDGPooledBuffer> ClassificationTileIndirectBuffer;
 
-	TUniformBufferRef<FStrataGlobalUniformParameters> StrataGlobalUniformParameters;
+	TRDGUniformBufferRef<FStrataGlobalUniformParameters> StrataGlobalUniformParameters;
 
 	FStrataSceneData()
 	{
+		Reset();
 	}
+
+	void Reset();
 };
 
 namespace Strata
@@ -64,7 +71,7 @@ void AddStrataClearMaterialBufferPass(FRDGBuilder& GraphBuilder, FUnorderedAcces
 
 void BindStrataBasePassUniformParameters(const FViewInfo& View, FStrataOpaquePassUniformParameters& OutStrataUniformParameters);
 
-TUniformBufferRef<FStrataGlobalUniformParameters> BindStrataGlobalUniformParameters(const FViewInfo& View);
+TRDGUniformBufferRef<FStrataGlobalUniformParameters> BindStrataGlobalUniformParameters(const FViewInfo& View);
 
 void AddStrataMaterialClassificationPass(FRDGBuilder& GraphBuilder, const FMinimalSceneTextures& SceneTextures, const TArray<FViewInfo>& Views);
 
