Strata uses RDG buffers.
This fixes the last transition issues.
Minor clean ups.

diff --git a/Engine/Source/Runtime/Renderer/Private/BasePassRendering.cpp b/Engine/Source/Runtime/Renderer/Private/BasePassRendering.cpp
--- a/Engine/Source/Runtime/Renderer/Private/BasePassRendering.cpp
+++ b/Engine/Source/Runtime/Renderer/Private/BasePassRendering.cpp
@@ -603,7 +603,7 @@ TRDGUniformBufferRef<FOpaqueBasePassUniformParameters> CreateOpaqueBasePassUnifo
 	SetupDistortionParams(BasePassParameters.DistortionParams, View);
 
 	// Strata
-	Strata::BindStrataBasePassUniformParameters(View, BasePassParameters.Strata);
+	Strata::BindStrataBasePassUniformParameters(GraphBuilder, View, BasePassParameters.Strata);
 
 	// Misc
 	BasePassParameters.EyeAdaptationTexture = SystemTextures.White;
diff --git a/Engine/Source/Runtime/Renderer/Private/Strata/Strata.cpp b/Engine/Source/Runtime/Renderer/Private/Strata/Strata.cpp
--- a/Engine/Source/Runtime/Renderer/Private/Strata/Strata.cpp
+++ b/Engine/Source/Runtime/Renderer/Private/Strata/Strata.cpp
@@ -84,6 +84,16 @@ void FStrataSceneData::Reset()
 	TopLayerNormalTexture = nullptr;
 	SSSTexture = nullptr;
 
+	MaterialLobesBuffer = nullptr;
+	MaterialLobesBufferUAV = nullptr;
+	MaterialLobesBufferSRV = nullptr;
+	ClassificationTileListBuffer = nullptr;
+	ClassificationTileListBufferUAV = nullptr;
+	ClassificationTileListBufferSRV = nullptr;
+	ClassificationTileIndirectBuffer = nullptr;
+	ClassificationTileIndirectBufferUAV = nullptr;
+	ClassificationTileIndirectBufferSRV = nullptr;
+
 	StrataGlobalUniformParameters = nullptr;
 }
 
@@ -92,6 +102,10 @@ void FStrataSceneData::Reset()
 namespace Strata
 {
 
+// Forward declaration
+static void AddStrataClearMaterialBufferPass(FRDGBuilder& GraphBuilder, FRDGBufferUAVRef MaterialLobesBufferUAV, uint32 MaxBytesPerPixel, FIntPoint TiledViewBufferResolution);
+static void AddStrataLUTPass(FRDGBuilder& GraphBuilder, const FViewInfo& View);
+
 static uint32 GetStrataGGXEnergyLUTResolution()
 {
 	return FMath::RoundUpToPowerOfTwo(FMath::Clamp(CVarStrataLUTResolution.GetValueOnAnyThread(), 16, 256));
@@ -120,23 +134,28 @@ uint32 GetStrataTileSize()
 	return 8;
 }
 
-static void AddStrataLUTPass(FRDGBuilder& GraphBuilder, const FViewInfo& View);
 
 void InitialiseStrataFrameSceneData(FSceneRenderer& SceneRenderer, FRDGBuilder& GraphBuilder)
 {
 	FStrataSceneData& StrataSceneData = SceneRenderer.Scene->StrataSceneData;
 	StrataSceneData.Reset();
 
-	uint32 ResolutionX = 1;
-	uint32 ResolutionY = 1;
+	auto UpdateMaterialBufferToTiledResolution = [](FIntPoint InBufferSizeXY, FIntPoint& OutMaterialBufferSizeXY)
+	{
+		// We need to allocate enough for the tiled memory addressing to always work
+		OutMaterialBufferSizeXY.X = FMath::DivideAndRoundUp(InBufferSizeXY.X, STRATA_DATA_TILE_SIZE) * STRATA_DATA_TILE_SIZE;
+		OutMaterialBufferSizeXY.Y = FMath::DivideAndRoundUp(InBufferSizeXY.Y, STRATA_DATA_TILE_SIZE) * STRATA_DATA_TILE_SIZE;
+	};
+
+	FIntPoint MaterialBufferSizeXY;
+	UpdateMaterialBufferToTiledResolution(FIntPoint(1, 1), MaterialBufferSizeXY);
 	bool bUpdateLUT = false;
 	if (IsStrataEnabled())
 	{
-		FIntPoint BufferSizeXY = GetSceneTextureExtent();
+		FIntPoint SceneTextureExtent = GetSceneTextureExtent();
 		
-		// We need to allocate enough for the tiled memory addressing to always work
-		ResolutionX = FMath::DivideAndRoundUp(BufferSizeXY.X, STRATA_DATA_TILE_SIZE) * STRATA_DATA_TILE_SIZE;
-		ResolutionY = FMath::DivideAndRoundUp(BufferSizeXY.Y, STRATA_DATA_TILE_SIZE) * STRATA_DATA_TILE_SIZE;
+		// We need to allocate enough for the tiled memory addressing of material data to always work
+		UpdateMaterialBufferToTiledResolution(SceneTextureExtent, MaterialBufferSizeXY);
 
 		const uint32 MaterialConservativeByteCountPerPixel = CVarStrataBytePerPixel.GetValueOnAnyThread();
 		const uint32 RoundToValue = 4u;
@@ -144,32 +163,33 @@ void InitialiseStrataFrameSceneData(FSceneRenderer& SceneRenderer, FRDGBuilder&
 
 		// Classification texture
 		{
-			StrataSceneData.ClassificationTexture = GraphBuilder.CreateTexture(FRDGTextureDesc::Create2D(BufferSizeXY, PF_R32_UINT, FClearValueBinding::Black, TexCreate_ShaderResource | TexCreate_RenderTargetable), TEXT("StrataClassificationTexture"));
+			StrataSceneData.ClassificationTexture = GraphBuilder.CreateTexture(FRDGTextureDesc::Create2D(SceneTextureExtent, PF_R32_UINT, FClearValueBinding::Black, TexCreate_ShaderResource | TexCreate_RenderTargetable), TEXT("StrataClassificationTexture"));
 		}
 
 		// Tile classification buffers
 		{
 			const int32 TileInPixel = GetStrataTileSize();
-			const FIntPoint TileResolution(FMath::DivideAndRoundUp(BufferSizeXY.X, TileInPixel), FMath::DivideAndRoundUp(BufferSizeXY.Y, TileInPixel));
-			FRDGBufferRef ClassificationTileListBuffer = GraphBuilder.CreateBuffer(FRDGBufferDesc::CreateBufferDesc(sizeof(uint32), TileResolution.X * TileResolution.Y), TEXT("StrataTileListBuffer"));
-			FRDGBufferRef ClassificationTileIndirectBuffer = GraphBuilder.CreateBuffer(FRDGBufferDesc::CreateIndirectDesc<FRHIDrawIndirectParameters>(), TEXT("StrataTileIndirectBuffer")); 
+			const FIntPoint TileResolution(FMath::DivideAndRoundUp(SceneTextureExtent.X, TileInPixel), FMath::DivideAndRoundUp(SceneTextureExtent.Y, TileInPixel));
+
+			StrataSceneData.ClassificationTileListBuffer = GraphBuilder.CreateBuffer(FRDGBufferDesc::CreateBufferDesc(sizeof(uint32), TileResolution.X * TileResolution.Y), TEXT("StrataTileListBuffer"));
+			StrataSceneData.ClassificationTileListBufferSRV = GraphBuilder.CreateSRV(StrataSceneData.ClassificationTileListBuffer, PF_R32_UINT);
+			StrataSceneData.ClassificationTileListBufferUAV = GraphBuilder.CreateUAV(StrataSceneData.ClassificationTileListBuffer, PF_R32_UINT);
 
-			GraphBuilder.PreallocateBuffer(ClassificationTileListBuffer);
-			GraphBuilder.PreallocateBuffer(ClassificationTileIndirectBuffer);
-			StrataSceneData.ClassificationTileListBuffer = GraphBuilder.GetPooledBuffer(ClassificationTileListBuffer);
-			StrataSceneData.ClassificationTileIndirectBuffer = GraphBuilder.GetPooledBuffer(ClassificationTileIndirectBuffer);
+			StrataSceneData.ClassificationTileIndirectBuffer = GraphBuilder.CreateBuffer(FRDGBufferDesc::CreateIndirectDesc<FRHIDispatchIndirectParameters>(1), TEXT("StrataTileIndirectBuffer"));
+			StrataSceneData.ClassificationTileIndirectBufferSRV = GraphBuilder.CreateSRV(StrataSceneData.ClassificationTileIndirectBuffer, PF_R32_UINT);
+			StrataSceneData.ClassificationTileIndirectBufferUAV = GraphBuilder.CreateUAV(StrataSceneData.ClassificationTileIndirectBuffer, PF_R32_UINT);
 
-			AddClearUAVPass(GraphBuilder, GraphBuilder.CreateUAV(GraphBuilder.RegisterExternalBuffer(StrataSceneData.ClassificationTileIndirectBuffer), PF_R32_UINT), 0);
+			AddClearUAVPass(GraphBuilder, StrataSceneData.ClassificationTileIndirectBufferUAV, 0);
 		}
 
 		// Top layer texture
 		{
-			StrataSceneData.TopLayerNormalTexture = GraphBuilder.CreateTexture(FRDGTextureDesc::Create2D(BufferSizeXY, PF_R32_UINT, FClearValueBinding::Black, TexCreate_ShaderResource | TexCreate_RenderTargetable), TEXT("StrataTopLayerNormalTexture"));
+			StrataSceneData.TopLayerNormalTexture = GraphBuilder.CreateTexture(FRDGTextureDesc::Create2D(SceneTextureExtent, PF_R32_UINT, FClearValueBinding::Black, TexCreate_ShaderResource | TexCreate_RenderTargetable), TEXT("StrataTopLayerNormalTexture"));
 		}
 
 		// SSS texture
 		{
-			StrataSceneData.SSSTexture = GraphBuilder.CreateTexture(FRDGTextureDesc::Create2D(BufferSizeXY, PF_R32G32_UINT, FClearValueBinding::Black, TexCreate_ShaderResource | TexCreate_RenderTargetable), TEXT("StrataSSSTexture"));
+			StrataSceneData.SSSTexture = GraphBuilder.CreateTexture(FRDGTextureDesc::Create2D(SceneTextureExtent, PF_R32G32_UINT, FClearValueBinding::Black, TexCreate_ShaderResource | TexCreate_RenderTargetable), TEXT("StrataSSSTexture"));
 		}
 
 		// Energy LUT
@@ -193,15 +213,11 @@ void InitialiseStrataFrameSceneData(FSceneRenderer& SceneRenderer, FRDGBuilder&
 		StrataSceneData.MaxBytesPerPixel = 4u;
 	}
 
-	const uint32 DesiredBufferSize = FMath::Max(4u, ResolutionX * ResolutionY * StrataSceneData.MaxBytesPerPixel);
-	if (StrataSceneData.MaterialLobesBuffer.NumBytes < DesiredBufferSize)
-	{
-		if (StrataSceneData.MaterialLobesBuffer.NumBytes > 0)
-		{
-			StrataSceneData.MaterialLobesBuffer.Release();
-		}
-		StrataSceneData.MaterialLobesBuffer.Initialize(TEXT("MaterialLobesBuffer"), DesiredBufferSize, BUF_Static);
-	}
+	// create the material lob buffer for all views
+	const uint32 MaterialLobesBufferByteSize = FMath::Max(4u, MaterialBufferSizeXY.X * MaterialBufferSizeXY.Y * StrataSceneData.MaxBytesPerPixel);
+	StrataSceneData.MaterialLobesBuffer = GraphBuilder.CreateBuffer(FRDGBufferDesc::CreateByteAddressDesc(MaterialLobesBufferByteSize), TEXT("StrataMaterialBuffer"));
+	StrataSceneData.MaterialLobesBufferSRV = GraphBuilder.CreateSRV(StrataSceneData.MaterialLobesBuffer);
+	StrataSceneData.MaterialLobesBufferUAV = GraphBuilder.CreateUAV(StrataSceneData.MaterialLobesBuffer);
 
 	// Set reference to the Strata data from each view
 	for (int32 ViewIndex = 0; ViewIndex < SceneRenderer.Views.Num(); ViewIndex++)
@@ -212,7 +228,7 @@ void InitialiseStrataFrameSceneData(FSceneRenderer& SceneRenderer, FRDGBuilder&
 
 	if (IsStrataEnabled())
 	{
-		AddStrataClearMaterialBufferPass(GraphBuilder, StrataSceneData.MaterialLobesBuffer.UAV, StrataSceneData.MaxBytesPerPixel, FIntPoint(ResolutionX, ResolutionY));
+		AddStrataClearMaterialBufferPass(GraphBuilder, StrataSceneData.MaterialLobesBufferUAV, StrataSceneData.MaxBytesPerPixel, MaterialBufferSizeXY);
 	}
 
 	if (bUpdateLUT)
@@ -233,7 +249,7 @@ void InitialiseStrataFrameSceneData(FSceneRenderer& SceneRenderer, FRDGBuilder&
 		if (IsStrataEnabled())
 		{
 			StrataUniformParameters->MaxBytesPerPixel = View.StrataSceneData->MaxBytesPerPixel;
-			StrataUniformParameters->MaterialLobesBuffer = View.StrataSceneData->MaterialLobesBuffer.SRV;
+			StrataUniformParameters->MaterialLobesBuffer = View.StrataSceneData->MaterialLobesBufferSRV;
 			StrataUniformParameters->ClassificationTexture = View.StrataSceneData->ClassificationTexture;
 			StrataUniformParameters->TopLayerNormalTexture = View.StrataSceneData->TopLayerNormalTexture;
 			StrataUniformParameters->SSSTexture = View.StrataSceneData->SSSTexture;
@@ -247,7 +263,7 @@ void InitialiseStrataFrameSceneData(FSceneRenderer& SceneRenderer, FRDGBuilder&
 		{
 			// Create each time. This path will go away when Strata is always enabled anyway.
 			StrataUniformParameters->MaxBytesPerPixel = 0;
-			StrataUniformParameters->MaterialLobesBuffer = GEmptyVertexBufferWithUAV->ShaderResourceViewRHI;
+			StrataUniformParameters->MaterialLobesBuffer = GraphBuilder.CreateSRV(GraphBuilder.RegisterExternalBuffer(GWhiteVertexBufferWithRDG->Buffer), PF_R32_UINT);
 			StrataUniformParameters->ClassificationTexture = GSystemTextures.GetZeroUIntDummy(GraphBuilder);
 			StrataUniformParameters->TopLayerNormalTexture = GSystemTextures.GetZeroUIntDummy(GraphBuilder);
 			StrataUniformParameters->SSSTexture = GSystemTextures.GetZeroUIntDummy(GraphBuilder);
@@ -260,17 +276,17 @@ void InitialiseStrataFrameSceneData(FSceneRenderer& SceneRenderer, FRDGBuilder&
 	}
 }
 
-void BindStrataBasePassUniformParameters(const FViewInfo& View, FStrataOpaquePassUniformParameters& OutStrataUniformParameters)
+void BindStrataBasePassUniformParameters(FRDGBuilder& GraphBuilder, const FViewInfo& View, FStrataOpaquePassUniformParameters& OutStrataUniformParameters)
 {
 	if (View.StrataSceneData)
 	{
 		OutStrataUniformParameters.MaxBytesPerPixel = View.StrataSceneData->MaxBytesPerPixel;
-		OutStrataUniformParameters.MaterialLobesBufferUAV = View.StrataSceneData->MaterialLobesBuffer.UAV;
+		OutStrataUniformParameters.MaterialLobesBufferUAV = View.StrataSceneData->MaterialLobesBufferUAV;
 	}
 	else
 	{
 		OutStrataUniformParameters.MaxBytesPerPixel = 0;
-		OutStrataUniformParameters.MaterialLobesBufferUAV = GEmptyVertexBufferWithUAV->UnorderedAccessViewRHI;
+		OutStrataUniformParameters.MaterialLobesBufferUAV = GraphBuilder.CreateUAV(GraphBuilder.RegisterExternalBuffer(GWhiteVertexBufferWithRDG->Buffer), PF_R32_UINT);
 	}
 }
 
@@ -376,7 +392,7 @@ class FStrataMaterialClassificationPassPS : public FGlobalShader
 	BEGIN_SHADER_PARAMETER_STRUCT(FParameters, )
 		SHADER_PARAMETER_STRUCT_REF(FViewUniformShaderParameters, ViewUniformBuffer)
 		SHADER_PARAMETER(uint32, MaxBytesPerPixel)
-		SHADER_PARAMETER_SRV(ByteAddressBuffer, MaterialLobesBuffer)
+		SHADER_PARAMETER_RDG_BUFFER_SRV(ByteAddressBuffer, MaterialLobesBuffer)
 		SHADER_PARAMETER_STRUCT_INCLUDE(FSceneTextureParameters, SceneTextures)
 		SHADER_PARAMETER_STRUCT_INCLUDE(ShaderDrawDebug::FShaderDrawDebugParameters, ShaderDrawParameters)
 		RENDER_TARGET_BINDING_SLOTS()
@@ -413,7 +429,7 @@ class FStrataClearMaterialBufferCS : public FGlobalShader
 	using FPermutationDomain = TShaderPermutationDomain<>;
 
 	BEGIN_SHADER_PARAMETER_STRUCT(FParameters, )
-		SHADER_PARAMETER_UAV(RWByteAddressBuffer, MaterialLobesBufferUAV)
+		SHADER_PARAMETER_RDG_BUFFER_UAV(RWByteAddressBuffer, MaterialLobesBufferUAV)
 		SHADER_PARAMETER(uint32, MaxBytesPerPixel)
 		SHADER_PARAMETER(FIntPoint, TiledViewBufferResolution)
 	END_SHADER_PARAMETER_STRUCT()
@@ -547,7 +563,7 @@ static void AddStrataInternalClassifedTilePass(
 	const FViewInfo& View,
 	const FRDGTextureRef* DepthTexture,
 	const FRDGTextureRef* ColorTexture,
-	FRDGBufferRef TileListBuffer,
+	FRDGBufferSRVRef TileListBufferSRV,
 	FRDGBufferRef TileIndirectBuffer)
 {	
 	const FIntPoint OutputResolution = View.ViewRect.Size();
@@ -561,7 +577,7 @@ static void AddStrataInternalClassifedTilePass(
 	ParametersPS->TileCount = TileCount;
 	ParametersPS->bRectPrimitive = GRHISupportsRectTopology ? 1 : 0;
 	ParametersPS->OutputResolution = OutputResolution;
-	ParametersPS->TileListBuffer = GraphBuilder.CreateSRV(TileListBuffer, PF_R32_UINT);
+	ParametersPS->TileListBuffer = TileListBufferSRV;
 	ParametersPS->TileIndirectBuffer = TileIndirectBuffer;
 
 	TShaderMapRef<FStrataMaterialStencilClassificationPassVS> VertexShader(View.ShaderMap);
@@ -634,9 +650,8 @@ void AddStrataStencilPass(
 	const FViewInfo& View,
 	const FMinimalSceneTextures& SceneTextures)
 {
-	FRDGBufferRef TileListBuffer = GraphBuilder.RegisterExternalBuffer(View.StrataSceneData->ClassificationTileListBuffer);
-	FRDGBufferRef TileIndirectBuffer = GraphBuilder.RegisterExternalBuffer(View.StrataSceneData->ClassificationTileIndirectBuffer);
-	AddStrataInternalClassifedTilePass(GraphBuilder, View, &SceneTextures.Depth.Target, nullptr, TileListBuffer, TileIndirectBuffer);
+	AddStrataInternalClassifedTilePass(GraphBuilder, View, &SceneTextures.Depth.Target, nullptr, 
+		View.StrataSceneData->ClassificationTileListBufferSRV, View.StrataSceneData->ClassificationTileIndirectBuffer);
 }
 
 void AddStrataStencilPass(
@@ -647,9 +662,8 @@ void AddStrataStencilPass(
 	for (int32 i = 0; i < Views.Num(); ++i)
 	{
 		const FViewInfo& View = Views[i];
-		FRDGBufferRef TileListBuffer = GraphBuilder.RegisterExternalBuffer(View.StrataSceneData->ClassificationTileListBuffer);
-		FRDGBufferRef TileIndirectBuffer = GraphBuilder.RegisterExternalBuffer(View.StrataSceneData->ClassificationTileIndirectBuffer);
-		AddStrataInternalClassifedTilePass(GraphBuilder, View, &SceneTextures.Depth.Target, nullptr, TileListBuffer, TileIndirectBuffer);
+		AddStrataInternalClassifedTilePass(GraphBuilder, View, &SceneTextures.Depth.Target, nullptr, 
+			View.StrataSceneData->ClassificationTileListBufferSRV, View.StrataSceneData->ClassificationTileIndirectBuffer);
 	}
 }
 
@@ -674,7 +688,7 @@ void AddStrataMaterialClassificationPass(FRDGBuilder& GraphBuilder, const FMinim
 			FStrataMaterialClassificationPassPS::FParameters* PassParameters = GraphBuilder.AllocParameters<FStrataMaterialClassificationPassPS::FParameters>();
 			PassParameters->ViewUniformBuffer = View.ViewUniformBuffer;
 			PassParameters->MaxBytesPerPixel = View.StrataSceneData->MaxBytesPerPixel;
-			PassParameters->MaterialLobesBuffer = View.StrataSceneData->MaterialLobesBuffer.SRV;
+			PassParameters->MaterialLobesBuffer = View.StrataSceneData->MaterialLobesBufferSRV;
 			PassParameters->SceneTextures = GetSceneTextureParameters(GraphBuilder);		
 			PassParameters->RenderTargets[0] = FRenderTargetBinding(View.StrataSceneData->ClassificationTexture, ERenderTargetLoadAction::EClear);
 			PassParameters->RenderTargets[1] = FRenderTargetBinding(View.StrataSceneData->TopLayerNormalTexture, ERenderTargetLoadAction::EClear);
@@ -703,8 +717,8 @@ void AddStrataMaterialClassificationPass(FRDGBuilder& GraphBuilder, const FMinim
 			PassParameters->bRectPrimitive = GRHISupportsRectTopology ? 1 : 0;
 			PassParameters->ViewResolution = View.ViewRect.Size();
 			PassParameters->ClassificationTexture = View.StrataSceneData->ClassificationTexture;
-			PassParameters->TileListData = GraphBuilder.CreateUAV(GraphBuilder.RegisterExternalBuffer(View.StrataSceneData->ClassificationTileListBuffer), PF_R32_UINT);
-			PassParameters->TileIndirectData = GraphBuilder.CreateUAV(GraphBuilder.RegisterExternalBuffer(View.StrataSceneData->ClassificationTileIndirectBuffer), PF_R32_UINT);
+			PassParameters->TileListData = View.StrataSceneData->ClassificationTileListBufferUAV;
+			PassParameters->TileIndirectData = View.StrataSceneData->ClassificationTileIndirectBufferUAV;
 
 			// Add 64 threads permutation
 			const uint32 GroupSize = 8;
@@ -718,7 +732,7 @@ void AddStrataMaterialClassificationPass(FRDGBuilder& GraphBuilder, const FMinim
 	}
 }
 
-void AddStrataClearMaterialBufferPass(FRDGBuilder& GraphBuilder, FUnorderedAccessViewRHIRef MaterialLobesBufferUAV, uint32 MaxBytesPerPixel, FIntPoint TiledViewBufferResolution)
+static void AddStrataClearMaterialBufferPass(FRDGBuilder& GraphBuilder, FRDGBufferUAVRef MaterialLobesBufferUAV, uint32 MaxBytesPerPixel, FIntPoint TiledViewBufferResolution)
 {
 	TShaderMapRef<FStrataClearMaterialBufferCS> ComputeShader(GetGlobalShaderMap(GMaxRHIFeatureLevel));
 	FStrataClearMaterialBufferCS::FParameters* PassParameters = GraphBuilder.AllocParameters<FStrataClearMaterialBufferCS::FParameters>();
@@ -896,9 +910,8 @@ void AddStrataDebugPasses(FRDGBuilder& GraphBuilder, const TArray<FViewInfo>& Vi
 		for (int32 i = 0; i < Views.Num(); ++i)
 		{
 			const FViewInfo& View = Views[i];
-			FRDGBufferRef TileListBuffer = GraphBuilder.RegisterExternalBuffer(View.StrataSceneData->ClassificationTileListBuffer);
-			FRDGBufferRef TileIndirectBuffer = GraphBuilder.RegisterExternalBuffer(View.StrataSceneData->ClassificationTileIndirectBuffer);
-			AddStrataInternalClassifedTilePass(GraphBuilder, View, nullptr, &SceneColorTexture, TileListBuffer, TileIndirectBuffer);
+			AddStrataInternalClassifedTilePass(GraphBuilder, View, nullptr, &SceneColorTexture,
+				View.StrataSceneData->ClassificationTileListBufferSRV, View.StrataSceneData->ClassificationTileIndirectBuffer);
 		}
 	}
 
diff --git a/Engine/Source/Runtime/Renderer/Private/Strata/Strata.h b/Engine/Source/Runtime/Renderer/Private/Strata/Strata.h
--- a/Engine/Source/Runtime/Renderer/Private/Strata/Strata.h
+++ b/Engine/Source/Runtime/Renderer/Private/Strata/Strata.h
@@ -16,13 +16,13 @@ struct FMinimalSceneTextures;
 
 BEGIN_SHADER_PARAMETER_STRUCT(FStrataOpaquePassUniformParameters, )
 	SHADER_PARAMETER(uint32, MaxBytesPerPixel)
-	SHADER_PARAMETER_UAV(RWByteAddressBuffer, MaterialLobesBufferUAV)
+	SHADER_PARAMETER_RDG_BUFFER_UAV(RWByteAddressBuffer, MaterialLobesBufferUAV)
 END_SHADER_PARAMETER_STRUCT()
 
 BEGIN_GLOBAL_SHADER_PARAMETER_STRUCT(FStrataGlobalUniformParameters, )
 	SHADER_PARAMETER(uint32, MaxBytesPerPixel)
 	SHADER_PARAMETER(FVector2D, GGXEnergyLUTScaleBias)
-	SHADER_PARAMETER_SRV(ByteAddressBuffer, MaterialLobesBuffer)
+	SHADER_PARAMETER_RDG_BUFFER_SRV(ByteAddressBuffer, MaterialLobesBuffer)
 	SHADER_PARAMETER_RDG_TEXTURE(Texture2D<uint>, ClassificationTexture)
 	SHADER_PARAMETER_RDG_TEXTURE(Texture2D<uint>, TopLayerNormalTexture)
 	SHADER_PARAMETER_RDG_TEXTURE(Texture2D<uint2>, SSSTexture)
@@ -36,19 +36,28 @@ struct FStrataSceneData
 	uint32 MaxBytesPerPixel;
 
 	// Resources allocated and updated each frame
-	FRWByteAddressBuffer MaterialLobesBuffer;					// This should be a RDG resource
+
+	FRDGBufferRef MaterialLobesBuffer;
+	FRDGBufferUAVRef MaterialLobesBufferUAV;
+	FRDGBufferSRVRef MaterialLobesBufferSRV;
+	FRDGBufferRef ClassificationTileListBuffer;
+	FRDGBufferUAVRef ClassificationTileListBufferUAV;
+	FRDGBufferSRVRef ClassificationTileListBufferSRV;
+	FRDGBufferRef ClassificationTileIndirectBuffer;
+	FRDGBufferUAVRef ClassificationTileIndirectBufferUAV;
+	FRDGBufferSRVRef ClassificationTileIndirectBufferSRV;
+
 	FRDGTextureRef ClassificationTexture;
 	FRDGTextureRef TopLayerNormalTexture;
-	FRDGTextureRef SSSTexture;	
+	FRDGTextureRef SSSTexture;
+
+	TRDGUniformBufferRef<FStrataGlobalUniformParameters> StrataGlobalUniformParameters;
 
 	// Resources computed once for multiple frames
+
 	TRefCountPtr<IPooledRenderTarget> GGXEnergyLUT3DTexture;
 	TRefCountPtr<IPooledRenderTarget> GGXEnergyLUT2DTexture;
 
-	TRefCountPtr<FRDGPooledBuffer> ClassificationTileListBuffer;
-	TRefCountPtr<FRDGPooledBuffer> ClassificationTileIndirectBuffer;
-
-	TRDGUniformBufferRef<FStrataGlobalUniformParameters> StrataGlobalUniformParameters;
 
 	FStrataSceneData()
 	{
@@ -67,9 +76,7 @@ bool IsClassificationEnabled();
 
 void InitialiseStrataFrameSceneData(FSceneRenderer& SceneRenderer, FRDGBuilder& GraphBuilder);
 
-void AddStrataClearMaterialBufferPass(FRDGBuilder& GraphBuilder, FUnorderedAccessViewRHIRef MaterialLobesBufferUAV, uint32 MaxBytesPerPixel, FIntPoint TiledViewBufferResolution);
-
-void BindStrataBasePassUniformParameters(const FViewInfo& View, FStrataOpaquePassUniformParameters& OutStrataUniformParameters);
+void BindStrataBasePassUniformParameters(FRDGBuilder& GraphBuilder, const FViewInfo& View, FStrataOpaquePassUniformParameters& OutStrataUniformParameters);
 
 TRDGUniformBufferRef<FStrataGlobalUniformParameters> BindStrataGlobalUniformParameters(const FViewInfo& View);
 
