Remove uncessary hair LUT binding.

diff --git a/Engine/Shaders/Private/HairStrands/HairScatter.usf b/Engine/Shaders/Private/HairStrands/HairScatter.usf
--- a/Engine/Shaders/Private/HairStrands/HairScatter.usf
+++ b/Engine/Shaders/Private/HairStrands/HairScatter.usf
@@ -53,9 +53,7 @@ Texture2D<uint4>  			CategorizationTexture;
 Texture2D<uint> 			VisibilityNodeIndex;
 StructuredBuffer<FPackedHairSample> VisibilityNodeData;
 
-Texture3D<float4>			HairLUTTexture;
 Texture3D<float4>			HairEnergyLUTTexture;
-SamplerState 				LinearSampler;
 
 FHairTransmittanceMask ComputeTransmittanceVoxel(
 	const uint2 PixelCoord,
@@ -272,7 +270,7 @@ void MainPS(
 
 				{
 					// Compute the transmittance based on precompute Hair transmittance LUT
-					FHairAverageScattering AverageScattering = SampleHairLUT(HairLUTTexture, LinearSampler, GBuffer.BaseColor, GBuffer.Roughness, SinLightAngle);
+					FHairAverageScattering AverageScattering = SampleHairLUT(View.HairScatteringLUTTexture, View.HairScatteringLUTSampler, GBuffer.BaseColor, GBuffer.Roughness, SinLightAngle);
 					TransmittanceData.A_front = AverageScattering.A_front;
 					TransmittanceData.A_back  = AverageScattering.A_back;
 				}
@@ -300,10 +298,10 @@ void MainPS(
 				TransmittanceData.ScatteringComponent = HairComponents | HAIR_COMPONENT_MULTISCATTER;
 
 				// Compute the transmittance based on precompute Hair transmittance LUT
-				FHairAverageEnergy Energy	= SampleHairEnergyLUT(HairEnergyLUTTexture, LinearSampler, GBuffer.BaseColor, GBuffer.Roughness, SinLightAngle);
+				FHairAverageEnergy Energy	= SampleHairEnergyLUT(HairEnergyLUTTexture, View.HairScatteringLUTSampler, GBuffer.BaseColor, GBuffer.Roughness, SinLightAngle);
 				const float3 TotalEnergy	= Energy.A_R + Energy.A_TT + Energy.A_TRT;
 	
-				FHairAverageEnergy SampleEnergy = SampleHairEnergyLUT(HairEnergyLUTTexture, LinearSampler, SampleGBuffer.BaseColor, SampleGBuffer.Roughness, SampleSinLightAngle);
+				FHairAverageEnergy SampleEnergy = SampleHairEnergyLUT(HairEnergyLUTTexture, View.HairScatteringLUTSampler, SampleGBuffer.BaseColor, SampleGBuffer.Roughness, SampleSinLightAngle);
 				const float3 SampleTotalEnergy = SampleEnergy.A_R + SampleEnergy.A_TT + SampleEnergy.A_TRT;
 
 				TransmittanceData.LocalScattering = TotalEnergy * SampleTotalEnergy;
@@ -321,7 +319,7 @@ void MainPS(
 			}
 			#endif
 
-			const float3 Transmittance = GetHairTransmittanceOnly(InitHairTransmittanceMask(TransmittanceMask.HairCount, TransmittanceMask.Visibility), GBuffer, SampleSinLightAngle, HairLUTTexture, LinearSampler);
+			const float3 Transmittance = GetHairTransmittanceOnly(InitHairTransmittanceMask(TransmittanceMask.HairCount, TransmittanceMask.Visibility), GBuffer, SampleSinLightAngle, View.HairScatteringLUTTexture, View.HairScatteringLUTSampler);
 
 			const float Area = 0;
 			const float3 TransmittedLight = Transmittance * TransmittanceMask.Visibility * SampleColor.xyz;
diff --git a/Engine/Shaders/Private/HairStrands/HairStrandsBsdfPlot.usf b/Engine/Shaders/Private/HairStrands/HairStrandsBsdfPlot.usf
--- a/Engine/Shaders/Private/HairStrands/HairStrandsBsdfPlot.usf
+++ b/Engine/Shaders/Private/HairStrands/HairStrandsBsdfPlot.usf
@@ -258,7 +258,7 @@ void MainPS(
 				const float SinLightAngle = dot(L, T);
 
 				// Compute the transmittance based on precompute Hair transmittance LUT
-				FHairAverageScattering AverageScattering = SampleHairLUT(HairScatteringLUTTexture, HairLUTSampler, P.BaseColor, P.Roughness, SinLightAngle);
+				FHairAverageScattering AverageScattering = SampleHairLUT(View.HairScatteringLUTTexture, View.HairScatteringLUTSampler, P.BaseColor, P.Roughness, SinLightAngle);
 
 				TransmittanceData = ComputeDualScatteringTerms(
 					InitHairTransmittanceMask(),
diff --git a/Engine/Shaders/Private/HairStrands/HairStrandsEnvironmentLighting.usf b/Engine/Shaders/Private/HairStrands/HairStrandsEnvironmentLighting.usf
--- a/Engine/Shaders/Private/HairStrands/HairStrandsEnvironmentLighting.usf
+++ b/Engine/Shaders/Private/HairStrands/HairStrandsEnvironmentLighting.usf
@@ -94,9 +94,7 @@ float3 GatherRadiance(float3 RayDirection, float Roughness, uint2 PixelPosition,
 		EyeIndex);
 }
 
-Texture3D<float4>			HairScatteringLUTTexture;
 Texture3D<float4>			HairEnergyLUTTexture;
-SamplerState 				HairLUTSampler; 
 
 //////////////////////////////////////////////////////////////////////////////////////////////////////////////////////////////
 // Debug info
@@ -282,7 +280,7 @@ float3 ReflectionEnvironment(FGBufferData GBuffer, float2 ScreenPosition, float2
 			const FHairTraversalResult Result = InternalComputeHairCountVoxel(MacroGroupId, SampleWorldPosition, SampleLightPosition, 1, 0, MaxSearchDistance, PixelRadius, true, bDebugEnabled);
 
 			const uint ScatteringComponent = (View.HairComponents & HAIR_COMPONENT_R) | (View.HairComponents & HAIR_COMPONENT_TRT) | HAIR_COMPONENT_LS | HAIR_COMPONENT_GS;// | HAIR_MULTISCATTER;
-			const FHairTransmittanceData TransmittanceData = GetHairTransmittance(V, SampleL, GBuffer, InitHairTransmittanceMask(Result.HairCount, Result.Visibility), HairScatteringLUTTexture, HairLUTSampler, ScatteringComponent);
+			const FHairTransmittanceData TransmittanceData = GetHairTransmittance(V, SampleL, GBuffer, InitHairTransmittanceMask(Result.HairCount, Result.Visibility), View.HairScatteringLUTTexture, View.HairScatteringLUTSampler, ScatteringComponent);
 
 			const float Area = 0.2;
 			const float BackLit = 1;
@@ -351,7 +349,7 @@ float3 ReflectionEnvironment(FGBufferData GBuffer, float2 ScreenPosition, float2
 		#if 0
 		{
 			const float SinViewAngle = dot(T, R);
-			FHairAverageEnergy Energy = SampleHairEnergyLUT(HairEnergyLUTTexture, HairLUTSampler, GBuffer.BaseColor, GBuffer.Roughness, SinViewAngle);
+			FHairAverageEnergy Energy = SampleHairEnergyLUT(HairEnergyLUTTexture, View.HairScatteringLUTSampler, GBuffer.BaseColor, GBuffer.Roughness, SinViewAngle);
 			ColorMS.rgb *= Energy.A_R + Energy.A_TT + Energy.A_TRT;
 		}
 		#endif	
@@ -406,7 +404,7 @@ float3 ReflectionEnvironment(FGBufferData GBuffer, float2 ScreenPosition, float2
 
 			// Compute the number of hair count between light & shading point
 			const FHairTraversalResult Result = InternalComputeHairCountVoxel(MacroGroupId, SampleWorldPosition, SampleLightPosition, 1, 0, MaxSearchDistance, PixelRadius, true, bDebugEnabled);
-			Transmittance = GetHairTransmittanceOnly(InitHairTransmittanceMask(Result.HairCount, Result.Visibility), GBuffer, SinLightAngle, HairScatteringLUTTexture, HairLUTSampler);
+			Transmittance = GetHairTransmittanceOnly(InitHairTransmittanceMask(Result.HairCount, Result.Visibility), GBuffer, SinLightAngle, View.HairScatteringLUTTexture, View.HairScatteringLUTSampler);
 			Visibility = Result.Visibility;
 		}
 
@@ -414,7 +412,7 @@ float3 ReflectionEnvironment(FGBufferData GBuffer, float2 ScreenPosition, float2
 		// Directional albedo
 		#if 0
 		{
-			FHairAverageEnergy Energy = SampleHairEnergyLUT(HairEnergyLUTTexture, HairLUTSampler, GBuffer.BaseColor, GBuffer.Roughness, SinViewAngle);
+			FHairAverageEnergy Energy = SampleHairEnergyLUT(HairEnergyLUTTexture, View.HairScatteringLUTSampler, GBuffer.BaseColor, GBuffer.Roughness, SinViewAngle);
 			ColorR.rgb *= (Energy.A_TRT + Energy.A_R) * Transmittance.xyz * Visibility;
 		}
 		#endif
@@ -473,7 +471,7 @@ float3 ReflectionEnvironment(FGBufferData GBuffer, float2 ScreenPosition, float2
 			{
 				Result = InternalComputeHairCountVoxel(MacroGroupId, SampleWorldPosition, SampleLightPosition, TransmissionDensityScaleFactor, 0, DistanceThresholdTT, PixelRadius, false, bDebugEnabled);
 			}
-			Transmittance = GetHairTransmittanceOnly(Result.HairCount, Result.Visibility, GBuffer, SinLightAngle, HairScatteringLUTTexture, HairLUTSampler);
+			Transmittance = GetHairTransmittanceOnly(Result.HairCount, Result.Visibility, GBuffer, SinLightAngle, View.HairScatteringLUTTexture, View.HairScatteringLUTSampler);
 			Visibility = Result.Visibility;
 		}
 
@@ -482,7 +480,7 @@ float3 ReflectionEnvironment(FGBufferData GBuffer, float2 ScreenPosition, float2
 		{
 			Error the Transmittance.w does not contain the correct hair count when we hit opaque geometry: TODO
 			const uint HairComponent = HAIR_COMPONENT_TT | HAIR_COMPONENT_LS | HAIR_COMPONENT_GS;
-			const FHairTransmittanceData TransmittanceData = GetHairTransmittance(V, L, GBuffer, Transmittance, HairScatteringLUTTexture, HairLUTSampler, HairComponent);
+			const FHairTransmittanceData TransmittanceData = GetHairTransmittance(V, L, GBuffer, Transmittance, View.HairScatteringLUTTexture, View.HairScatteringLUTSampler, HairComponent);
 			const float Area = 0.2;
 			ColorTT.rgb *= HairShading(GBuffer, L, V, T, 1, TransmittanceData, 1, Area, uint2(0, 0)) * Visibility;
 		}
@@ -492,7 +490,7 @@ float3 ReflectionEnvironment(FGBufferData GBuffer, float2 ScreenPosition, float2
 		#if 1
 		{
 			const float Backlit = GBuffer.CustomData.z;
-			FHairAverageEnergy Energy = SampleHairEnergyLUT(HairEnergyLUTTexture, HairLUTSampler, GBuffer.BaseColor, GBuffer.Roughness, SinViewAngle);
+			FHairAverageEnergy Energy = SampleHairEnergyLUT(HairEnergyLUTTexture, View.HairScatteringLUTSampler, GBuffer.BaseColor, GBuffer.Roughness, SinViewAngle);
 			ColorTT.rgb *= Energy.A_TT * Transmittance * Visibility * Backlit;
 		}
 		#endif	
@@ -592,8 +590,8 @@ float3 ReflectionEnvironmentUniformSampling(
 			SampleL,
 			GBuffer,
 			InitHairTransmittanceMask(Result.HairCount, Result.Visibility),
-			HairScatteringLUTTexture,
-			HairLUTSampler,
+			View.HairScatteringLUTTexture, 
+			View.HairScatteringLUTSampler,
 			View.HairComponents);	
 
 		const float Area = 0;
@@ -646,8 +644,8 @@ float3 ReflectionEnvironmentSH(
 		L,
 		GBuffer,
 		InitHairTransmittanceMask(),
-		HairScatteringLUTTexture,
-		HairLUTSampler,
+		View.HairScatteringLUTTexture, 
+		View.HairScatteringLUTSampler,
 		View.HairComponents);
 
 	bool bEvalMultiScatter = true;
@@ -719,7 +717,7 @@ float3 ComputeThinTransmission(FGBufferData GBuffer, float2 ScreenPosition, floa
 
 	// The pre-exposure is not applied onto the output as the SceneColorTexture value are already preexposed.
 	const float SinViewAngle = dot(T, V);
-	const FHairAverageScattering Energy = SampleHairLUT(HairScatteringLUTTexture, HairLUTSampler, GBuffer.BaseColor, GBuffer.Roughness, SinViewAngle);
+	const FHairAverageScattering Energy = SampleHairLUT(View.HairScatteringLUTTexture, View.HairScatteringLUTSampler, GBuffer.BaseColor, GBuffer.Roughness, SinViewAngle);
 	float3 Color = Energy.A_front * Luminance.xyz * PixelCoverage;
 
 	return -min(-Color, 0.0);
diff --git a/Engine/Shaders/Private/HairStrands/HairStrandsVisibilityCompactionComputeRaster.usf b/Engine/Shaders/Private/HairStrands/HairStrandsVisibilityCompactionComputeRaster.usf
--- a/Engine/Shaders/Private/HairStrands/HairStrandsVisibilityCompactionComputeRaster.usf
+++ b/Engine/Shaders/Private/HairStrands/HairStrandsVisibilityCompactionComputeRaster.usf
@@ -36,12 +36,6 @@ uint MaxNodeCount;
 uint SamplerPerPixel;
 float CoverageThreshold;
 
-//float LUT_HairCount;
-//float LUT_HairRadiusCount;
-//
-//SamplerState					LinearSampler;
-//Texture2D<float>				HairCoverageLUT;
-
 RWTexture2D<uint> 				OutCompactNodeCounter;
 RWTexture2D<uint> 				OutCompactNodeIndex;
 RWStructuredBuffer<FNodeVis> 	OutCompactNodeData;
diff --git a/Engine/Source/Runtime/Renderer/Private/HairStrands/HairStrandsData.h b/Engine/Source/Runtime/Renderer/Private/HairStrands/HairStrandsData.h
--- a/Engine/Source/Runtime/Renderer/Private/HairStrands/HairStrandsData.h
+++ b/Engine/Source/Runtime/Renderer/Private/HairStrands/HairStrandsData.h
@@ -127,8 +127,6 @@ BEGIN_SHADER_PARAMETER_STRUCT(FHairStrandsVoxelCommonParameters, )
 	SHADER_PARAMETER_RDG_BUFFER_SRV(Buffer<uint2>, PageIndexOccupancyBuffer)
 	SHADER_PARAMETER_RDG_BUFFER_SRV(Buffer<uint>, PageIndexCoordBuffer)
 	SHADER_PARAMETER_RDG_BUFFER_SRV(StructuredBuffer<FPackedVirtualVoxelNodeDesc>, NodeDescBuffer) // Packed into 2 x uint4
-	SHADER_PARAMETER_RDG_TEXTURE(Texture2D<float>, HairCoverageLUT)
-	SHADER_PARAMETER_SAMPLER(SamplerState, HairCoverageSampler)
 END_SHADER_PARAMETER_STRUCT()
 
 BEGIN_GLOBAL_SHADER_PARAMETER_STRUCT(FVirtualVoxelParameters, RENDERER_API)
diff --git a/Engine/Source/Runtime/Renderer/Private/HairStrands/HairStrandsDebug.cpp b/Engine/Source/Runtime/Renderer/Private/HairStrands/HairStrandsDebug.cpp
--- a/Engine/Source/Runtime/Renderer/Private/HairStrands/HairStrandsDebug.cpp
+++ b/Engine/Source/Runtime/Renderer/Private/HairStrands/HairStrandsDebug.cpp
@@ -811,9 +811,6 @@ class FHairStrandsPlotSamplePS : public FGlobalShader
 		SHADER_PARAMETER(uint32, HairComponents)
 		SHADER_PARAMETER(float, Exposure)
 		SHADER_PARAMETER_STRUCT_REF(FViewUniformShaderParameters, ViewUniformBuffer)
-
-		SHADER_PARAMETER_RDG_TEXTURE(Texture3D, HairScatteringLUTTexture)
-		SHADER_PARAMETER_SAMPLER(SamplerState, HairLUTSampler)
 		RENDER_TARGET_BINDING_SLOTS()
 	END_SHADER_PARAMETER_STRUCT()
 
@@ -838,8 +835,6 @@ static void AddPlotSamplePass(
 
 	const FIntPoint Resolution(OutputTexture->Desc.Extent);
 	FHairStrandsPlotSamplePS::FParameters* Parameters = GraphBuilder.AllocParameters<FHairStrandsPlotSamplePS::FParameters>();
-	
-	const FHairLUT InHairLUT = GetHairLUT(GraphBuilder, View);
 
 	FHairStrandsDebugData::SetParameters(GraphBuilder, DebugData, Parameters->DebugData);
 	Parameters->ViewUniformBuffer = View.ViewUniformBuffer;
@@ -848,8 +843,6 @@ static void AddPlotSamplePass(
 	Parameters->MaxResolution = OutputTexture->Desc.Extent;
 	Parameters->HairComponents = ToBitfield(GetHairComponents());
 	Parameters->Exposure = GHairStrandsDebugPlotBsdfExposure;
-	Parameters->HairScatteringLUTTexture = InHairLUT.Textures[HairLUTType_DualScattering];
-	Parameters->HairLUTSampler = TStaticSamplerState<SF_Bilinear, AM_Clamp, AM_Clamp, AM_Clamp>::GetRHI();
 	Parameters->RenderTargets[0] = FRenderTargetBinding(OutputTexture, ERenderTargetLoadAction::ELoad);
 
 	const FIntPoint OutputResolution = SceneTextures.SceneDepthTexture->Desc.Extent;
diff --git a/Engine/Source/Runtime/Renderer/Private/HairStrands/HairStrandsEnvironment.cpp b/Engine/Source/Runtime/Renderer/Private/HairStrands/HairStrandsEnvironment.cpp
--- a/Engine/Source/Runtime/Renderer/Private/HairStrands/HairStrandsEnvironment.cpp
+++ b/Engine/Source/Runtime/Renderer/Private/HairStrands/HairStrandsEnvironment.cpp
@@ -240,8 +240,6 @@ public:
 
 		SHADER_PARAMETER_RDG_TEXTURE(Texture2D, SceneColorTexture)
 		SHADER_PARAMETER_RDG_TEXTURE(Texture2D, HairEnergyLUTTexture)
-		SHADER_PARAMETER_RDG_TEXTURE(Texture2D, HairScatteringLUTTexture)
-		SHADER_PARAMETER_SAMPLER(SamplerState, HairLUTSampler)
 
 		SHADER_PARAMETER_RDG_BUFFER_UAV(RWBuffer, OutLightingBuffer)
 
@@ -321,10 +319,7 @@ static void AddHairStrandsEnvironmentLightingPassPS(
 	FHairEnvironmentLightingPS::FParameters* ParametersPS = GraphBuilder.AllocParameters<FHairEnvironmentLightingPS::FParameters>();
 	FHairEnvironmentLighting::FParameters* PassParameters = &ParametersPS->Common;
 
-	const FHairLUT InHairLUT = GetHairLUT(GraphBuilder, View);
-	PassParameters->HairEnergyLUTTexture = InHairLUT.Textures[HairLUTType_MeanEnergy];
-	PassParameters->HairScatteringLUTTexture = InHairLUT.Textures[HairLUTType_DualScattering];
-	PassParameters->HairLUTSampler = TStaticSamplerState<SF_Bilinear, AM_Clamp, AM_Clamp, AM_Clamp>::GetRHI();
+	PassParameters->HairEnergyLUTTexture = GetHairLUT(GraphBuilder, View, HairLUTType_MeanEnergy);
 
 	EHairLightingIntegrationType IntegrationType = EHairLightingIntegrationType::AdHoc;
 	const bool bUseSceneColor = SceneColorTexture != nullptr;
diff --git a/Engine/Source/Runtime/Renderer/Private/HairStrands/HairStrandsLUT.cpp b/Engine/Source/Runtime/Renderer/Private/HairStrands/HairStrandsLUT.cpp
--- a/Engine/Source/Runtime/Renderer/Private/HairStrands/HairStrandsLUT.cpp
+++ b/Engine/Source/Runtime/Renderer/Private/HairStrands/HairStrandsLUT.cpp
@@ -69,7 +69,9 @@ static FRDGTextureRef AddHairLUTPass(
 	OutputDesc.Format = PF_FloatRGBA;
 	OutputDesc.NumMips = 1;
 	OutputDesc.Flags = TexCreate_ShaderResource | TexCreate_UAV;
-	FRDGTextureRef HairLUTTexture = GraphBuilder.CreateTexture(OutputDesc, TEXT("Hair.LUT"));
+	FRDGTextureRef HairLUTTexture = GraphBuilder.CreateTexture(OutputDesc, 
+		 LUTType == FHairLUTType::HairLUTType_DualScattering ? TEXT("Hair.LUT(DualScattering)") : 
+		(LUTType == FHairLUTType::HairLUTType_MeanEnergy ? TEXT("Hair.LUT(MeanEnergy)") : TEXT("Hair.LUT(Coverage)")));
 
 	FHairLUTCS::FParameters* Parameters = GraphBuilder.AllocParameters<FHairLUTCS::FParameters>();
 	Parameters->OutputColor = GraphBuilder.CreateUAV(FRDGTextureUAVDesc(HairLUTTexture, 0));
@@ -326,80 +328,66 @@ static FRDGTextureRef AddHairCoverageLUTPass(FRDGBuilder& GraphBuilder, const FV
 	return HairLUTTexture;
 }
 
-static bool InternalHairLUT_NeedUpdateDualScatteringLUT(EShaderPlatform Platform)
+static FRDGTextureRef InternalGetHairLUT(FRDGBuilder& GraphBuilder, const FViewInfo& View, FHairLUTType Type, bool bRegister)
 {
-	const FIntVector LUT3DResolution = FIntVector(GHairLUTIncidentAngleCount, GHairLUTRoughnessCount, GHairLUTAbsorptionCount);
-	return
-		(GSystemTextures.HairLUT0.GetReference() == nullptr ||
-		 GSystemTextures.HairLUT0.GetReference()->GetRenderTargetItem().ShaderResourceTexture->GetSizeXYZ() != LUT3DResolution) &&
-		(!IsHairMobilePlatform(Platform) && (IsHairStrandsSupported(EHairStrandsShaderType::Cards, Platform) || IsHairStrandsSupported(EHairStrandsShaderType::Strands, Platform)));
-}
-
-static bool InternalHairLUT_NeedUpdateMeanEnergyLUT(EShaderPlatform Platform)
-{
-	const FIntVector LUT3DResolution = FIntVector(GHairLUTIncidentAngleCount, GHairLUTRoughnessCount, GHairLUTAbsorptionCount);
-	return 
-		(GSystemTextures.HairLUT1.GetReference() == nullptr ||
-		 GSystemTextures.HairLUT1.GetReference()->GetRenderTargetItem().ShaderResourceTexture->GetSizeXYZ() != LUT3DResolution) &&
-		IsHairStrandsSupported(EHairStrandsShaderType::Strands, Platform);
-}
-
-static bool InternalHairLUT_NeedUpdateCoverageLUT(EShaderPlatform Platform)
-{
-	const FIntVector LUT2DResolution = FIntVector(FHairCountToCoverageData::HairCount, FHairCountToCoverageData::HairRadiusCount, 1);
-	return 
-		(GSystemTextures.HairLUT2.GetReference() == nullptr || 
-		 GSystemTextures.HairLUT2.GetReference()->GetRenderTargetItem().ShaderResourceTexture->GetSizeXYZ() != LUT2DResolution) &&
-		IsHairStrandsSupported(EHairStrandsShaderType::Strands, Platform);
-}
+	FRDGTextureRef Out = nullptr;
 
-static FHairLUT InternalHairLUT_UpdateHairLUT(FRDGBuilder& GraphBuilder, const FViewInfo& View)
-{
 	const EShaderPlatform Platform = View.GetShaderPlatform();
-	FHairLUT Out;
-
-	if (InternalHairLUT_NeedUpdateDualScatteringLUT(Platform))
-	{
-		FRDGTextureRef HairDualScatteringLUTTexture = AddHairLUTPass(GraphBuilder, View, HairLUTType_DualScattering);
-		ConvertToUntrackedExternalTexture(GraphBuilder, HairDualScatteringLUTTexture, GSystemTextures.HairLUT0, ERHIAccess::SRVMask);
-		Out.Textures[HairLUTType_DualScattering] = HairDualScatteringLUTTexture;
-	}
-	else
+	const bool bIsCompatible = !IsHairMobilePlatform(Platform) && IsHairStrandsSupported(EHairStrandsShaderType::Cards, Platform);
+	if (!bIsCompatible)
 	{
-		Out.Textures[HairLUTType_DualScattering]= GraphBuilder.RegisterExternalTexture(GSystemTextures.HairLUT0, TEXT("HairLUTType_DualScattering"));
+		return Out;
 	}
 
-	if (InternalHairLUT_NeedUpdateMeanEnergyLUT(Platform))
+	if (Type == FHairLUTType::HairLUTType_DualScattering)
 	{
-		FRDGTextureRef HairMeanEnergyLUTTexture = AddHairLUTPass(GraphBuilder, View, HairLUTType_MeanEnergy);
-		ConvertToUntrackedExternalTexture(GraphBuilder, HairMeanEnergyLUTTexture, GSystemTextures.HairLUT1, ERHIAccess::SRVMask);
-		Out.Textures[HairLUTType_MeanEnergy] = HairMeanEnergyLUTTexture;
+		const bool bNeedUpdate = GSystemTextures.HairLUT0.GetReference() == nullptr || GSystemTextures.HairLUT0.GetReference()->GetRenderTargetItem().ShaderResourceTexture->GetSizeXYZ() != FIntVector(GHairLUTIncidentAngleCount, GHairLUTRoughnessCount, GHairLUTAbsorptionCount);
+		if (bNeedUpdate)
+		{
+			Out = AddHairLUTPass(GraphBuilder, View, HairLUTType_DualScattering);
+			ConvertToUntrackedExternalTexture(GraphBuilder, Out, GSystemTextures.HairLUT0, ERHIAccess::SRVMask);
+		}
+		else if (bRegister)
+		{
+			Out = GraphBuilder.RegisterExternalTexture(GSystemTextures.HairLUT0, TEXT("Hair.LUT(DualScattering)"));
+		}
 	}
-	else
-	{
-		Out.Textures[HairLUTType_MeanEnergy] = GraphBuilder.RegisterExternalTexture(GSystemTextures.HairLUT1, TEXT("HairLUTType_MeanEnergy"));
-	}
-
-	if (InternalHairLUT_NeedUpdateCoverageLUT(Platform))
+	else if (Type == FHairLUTType::HairLUTType_MeanEnergy)
 	{
-		FRDGTextureRef HairCoverageLUTTexture = AddHairCoverageLUTPass(GraphBuilder, View);
-		ConvertToUntrackedExternalTexture(GraphBuilder, HairCoverageLUTTexture, GSystemTextures.HairLUT2, ERHIAccess::SRVMask);
-		Out.Textures[HairLUTType_Coverage] = HairCoverageLUTTexture;
+		const bool bNeedUpdate = GSystemTextures.HairLUT1.GetReference() == nullptr;
+		if (bNeedUpdate)
+		{
+			Out = AddHairLUTPass(GraphBuilder, View, HairLUTType_MeanEnergy);
+			ConvertToUntrackedExternalTexture(GraphBuilder, Out, GSystemTextures.HairLUT1, ERHIAccess::SRVMask);
+		}
+		else if (bRegister)
+		{
+			Out = GraphBuilder.RegisterExternalTexture(GSystemTextures.HairLUT1, TEXT("Hair.LUT(MeanEnergy)"));
+		}
 	}
-	else
+	else if (Type == FHairLUTType::HairLUTType_Coverage)
 	{
-		Out.Textures[HairLUTType_Coverage] = GraphBuilder.RegisterExternalTexture(GSystemTextures.HairLUT2, TEXT("HairLUTType_Coverage"));
+		const bool bNeedUpdate = GSystemTextures.HairLUT2.GetReference() == nullptr;
+		if (bNeedUpdate)
+		{
+			Out = AddHairCoverageLUTPass(GraphBuilder, View);
+			ConvertToUntrackedExternalTexture(GraphBuilder, Out, GSystemTextures.HairLUT2, ERHIAccess::SRVMask);
+		}
+		else if (bRegister)
+		{
+			Out = GraphBuilder.RegisterExternalTexture(GSystemTextures.HairLUT2, TEXT("Hair.LUT(Coverage)"));
+		}
 	}
 
 	return Out;
 }
 
-FHairLUT GetHairLUT(FRDGBuilder & GraphBuilder, const FViewInfo & View)
+FRDGTextureRef GetHairLUT(FRDGBuilder& GraphBuilder, const FViewInfo& View, FHairLUTType Type)
 {
-	return InternalHairLUT_UpdateHairLUT(GraphBuilder, View);
+	return InternalGetHairLUT(GraphBuilder, View, Type, true);
 }
 
 void UpdateHairResources(FRDGBuilder& GraphBuilder, const FViewInfo& View)
 {
-	InternalHairLUT_UpdateHairLUT(GraphBuilder, View);
+	InternalGetHairLUT(GraphBuilder, View, FHairLUTType::HairLUTType_DualScattering, false);
 }
\ No newline at end of file
diff --git a/Engine/Source/Runtime/Renderer/Private/HairStrands/HairStrandsLUT.h b/Engine/Source/Runtime/Renderer/Private/HairStrands/HairStrandsLUT.h
--- a/Engine/Source/Runtime/Renderer/Private/HairStrands/HairStrandsLUT.h
+++ b/Engine/Source/Runtime/Renderer/Private/HairStrands/HairStrandsLUT.h
@@ -20,10 +20,5 @@ enum FHairLUTType
 	HairLUTTypeCount
 };
 
-struct FHairLUT
-{
-	FRDGTextureRef Textures[HairLUTTypeCount] = { nullptr, nullptr, nullptr };
-};
-
 /// Returns Hair LUTs. LUTs are generated on demand.
-FHairLUT GetHairLUT(FRDGBuilder& GraphBuilder, const FViewInfo& View);
\ No newline at end of file
+FRDGTextureRef GetHairLUT(FRDGBuilder& GraphBuilder, const FViewInfo& View, FHairLUTType Type);
diff --git a/Engine/Source/Runtime/Renderer/Private/HairStrands/HairStrandsScatter.cpp b/Engine/Source/Runtime/Renderer/Private/HairStrands/HairStrandsScatter.cpp
--- a/Engine/Source/Runtime/Renderer/Private/HairStrands/HairStrandsScatter.cpp
+++ b/Engine/Source/Runtime/Renderer/Private/HairStrands/HairStrandsScatter.cpp
@@ -124,10 +124,8 @@ class FHairScatterPS : public FGlobalShader
 		SHADER_PARAMETER(uint32, Debug)
 		SHADER_PARAMETER(uint32, SampleCount)
 		SHADER_PARAMETER(uint32, HairComponents)
-		SHADER_PARAMETER_RDG_TEXTURE(Texture3D, HairLUTTexture)
 		SHADER_PARAMETER_RDG_TEXTURE(Texture3D, HairEnergyLUTTexture)
 		SHADER_PARAMETER_RDG_BUFFER_UAV(RWStructuredBuffer, OutputColor)
-		SHADER_PARAMETER_SAMPLER(SamplerState, LinearSampler)
 		SHADER_PARAMETER_RDG_TEXTURE(Texture2D, VisibilityNodeIndex)
 		SHADER_PARAMETER_RDG_BUFFER_SRV(StructuredBuffer, VisibilityNodeData)
 		SHADER_PARAMETER_RDG_TEXTURE(Texture2D, DiffusionInputTexture)
@@ -157,7 +155,6 @@ static FRDGTextureRef AddScatterPass(
 		return nullptr;
 
 	const FIntPoint Resolution = OutSceneColorTexture->Desc.Extent;
-	const FHairLUT InHairLUT = GetHairLUT(GraphBuilder, View);
 
 	float PixelRadiusAtDepth1 = 0;
 	{
@@ -181,9 +178,7 @@ static FRDGTextureRef AddScatterPass(
 	Parameters->VisibilityNodeData = GraphBuilder.CreateSRV(InVisibilityNodeData);
 	Parameters->CategorizationTexture = InCategorizationTexture;
 	Parameters->DiffusionInputTexture = InDiffusionInput;
-	Parameters->HairLUTTexture = InHairLUT.Textures[HairLUTType_DualScattering];
-	Parameters->HairEnergyLUTTexture = InHairLUT.Textures[HairLUTType_MeanEnergy];
-	Parameters->LinearSampler = TStaticSamplerState<SF_Bilinear, AM_Clamp, AM_Clamp, AM_Clamp>::GetRHI();
+	Parameters->HairEnergyLUTTexture = GetHairLUT(GraphBuilder, View, HairLUTType_MeanEnergy);
 	Parameters->ViewUniformBuffer = View.ViewUniformBuffer;
 	Parameters->VirtualVoxel = VoxelResources.UniformBuffer;
 	if (ShaderDrawDebug::IsShaderDrawDebugEnabled(View))
diff --git a/Engine/Source/Runtime/Renderer/Private/HairStrands/HairStrandsVisibility.cpp b/Engine/Source/Runtime/Renderer/Private/HairStrands/HairStrandsVisibility.cpp
--- a/Engine/Source/Runtime/Renderer/Private/HairStrands/HairStrandsVisibility.cpp
+++ b/Engine/Source/Runtime/Renderer/Private/HairStrands/HairStrandsVisibility.cpp
@@ -2965,7 +2965,6 @@ IMPLEMENT_GLOBAL_SHADER(FHairCountToCoverageCS, "/Engine/Private/HairStrands/Hai
 static FRDGTextureRef AddHairHairCountToTransmittancePass(
 	FRDGBuilder& GraphBuilder,
 	const FViewInfo& ViewInfo,
-	const FHairLUT& HairLUT,
 	const FRDGTextureRef HairCountTexture)
 {
 	const FIntPoint OutputResolution = HairCountTexture->Desc.Extent;
@@ -2975,7 +2974,7 @@ static FRDGTextureRef AddHairHairCountToTransmittancePass(
 
 	FRDGTextureDesc Desc = FRDGTextureDesc::Create2D(OutputResolution, PF_R32_FLOAT, FClearValueBinding(FLinearColor(0.0f, 0.0f, 0.0f, 0.0f)), TexCreate_UAV | TexCreate_ShaderResource | TexCreate_RenderTargetable);
 	FRDGTextureRef OutputTexture = GraphBuilder.CreateTexture(Desc, TEXT("Hair.VisibilityTexture"));
-	FRDGTextureRef HairCoverageLUT = HairLUT.Textures[HairLUTType_Coverage];
+	FRDGTextureRef HairCoverageLUT = GetHairLUT(GraphBuilder, ViewInfo, HairLUTType_Coverage);
 
 	FHairCountToCoverageCS::FParameters* PassParameters = GraphBuilder.AllocParameters<FHairCountToCoverageCS::FParameters>();
 	PassParameters->LUT_HairCount = HairCoverageLUT->Desc.Extent.X;
@@ -3261,8 +3260,6 @@ void RenderHairStrandsVisibilityBuffer(
 	{
 		
 		{
-			FHairLUT HairLUT = GetHairLUT(GraphBuilder, View);
-
 			FHairStrandsVisibilityData& VisibilityData = View.HairStrandsViewData.VisibilityData;
 			VisibilityData.NodeGroupSize = GetVendorOptimalGroupSize1D();
 			VisibilityData.MaxSampleCount = GetMaxSamplePerPixel();
@@ -3298,7 +3295,6 @@ void RenderHairStrandsVisibilityBuffer(
 					ViewTransmittance.TransmittanceTexture = AddHairHairCountToTransmittancePass(
 						GraphBuilder,
 						View,
-						HairLUT,
 						RasterOutput.HairCountTexture);
 
 					ViewTransmittance.HairCountTextureUint = RasterOutput.HairCountTexture;
@@ -3447,7 +3443,6 @@ void RenderHairStrandsVisibilityBuffer(
 						ViewTransmittance.TransmittanceTexture = AddHairHairCountToTransmittancePass(
 							GraphBuilder,
 							View,
-							HairLUT,
 							ViewTransmittance.HairCountTexture);
 					}
 
diff --git a/Engine/Source/Runtime/Renderer/Private/HairStrands/HairStrandsVoxelization.cpp b/Engine/Source/Runtime/Renderer/Private/HairStrands/HairStrandsVoxelization.cpp
--- a/Engine/Source/Runtime/Renderer/Private/HairStrands/HairStrandsVoxelization.cpp
+++ b/Engine/Source/Runtime/Renderer/Private/HairStrands/HairStrandsVoxelization.cpp
@@ -869,8 +869,6 @@ FHairStrandsVoxelResources AllocateVirtualVoxelResources(
 	Out.Parameters.Common.Raytracing_SkyOcclusionThreshold		= FMath::Max(0.f, GHairVirtualVoxelRaytracing_SkyOcclusionThreshold);
 
 	Out.Parameters.Common.HairCoveragePixelRadiusAtDepth1	= ComputeMinStrandRadiusAtDepth1(FIntPoint(View.ViewRect.Width(), View.ViewRect.Height()), View.FOV, 1/*SampleCount*/, 1/*RasterizationScale*/).Primary;
-	Out.Parameters.Common.HairCoverageLUT					= GetHairLUT(GraphBuilder, View).Textures[HairLUTType_Coverage];
-	Out.Parameters.Common.HairCoverageSampler				= TStaticSamplerState<SF_Bilinear, AM_Clamp, AM_Clamp, AM_Clamp>::GetRHI();
 
 	AddAllocateVoxelPagesPass(
 		GraphBuilder, 
