Strata clean up to not compute multiple resources for multiple views when not necessary.
Clean up AddStrataLUTPass from unused data.

diff --git a/Engine/Source/Runtime/Renderer/Private/Strata/Strata.cpp b/Engine/Source/Runtime/Renderer/Private/Strata/Strata.cpp
--- a/Engine/Source/Runtime/Renderer/Private/Strata/Strata.cpp
+++ b/Engine/Source/Runtime/Renderer/Private/Strata/Strata.cpp
@@ -104,7 +104,7 @@ namespace Strata
 
 // Forward declaration
 static void AddStrataClearMaterialBufferPass(FRDGBuilder& GraphBuilder, FRDGBufferUAVRef MaterialLobesBufferUAV, uint32 MaxBytesPerPixel, FIntPoint TiledViewBufferResolution);
-static void AddStrataLUTPass(FRDGBuilder& GraphBuilder, const FViewInfo& View);
+static void AddStrataLUTPass(FRDGBuilder& GraphBuilder, FSceneRenderer& SceneRenderer, TRefCountPtr<IPooledRenderTarget>& GGXEnergyLUT2DTexture, TRefCountPtr<IPooledRenderTarget>& GGXEnergyLUT3DTexture);
 
 static uint32 GetStrataGGXEnergyLUTResolution()
 {
@@ -229,50 +229,40 @@ void InitialiseStrataFrameSceneData(FSceneRenderer& SceneRenderer, FRDGBuilder&
 	if (IsStrataEnabled())
 	{
 		AddStrataClearMaterialBufferPass(GraphBuilder, StrataSceneData.MaterialLobesBufferUAV, StrataSceneData.MaxBytesPerPixel, MaterialBufferSizeXY);
-	}
-
-	if (bUpdateLUT)
-	{
-		for (int32 ViewIndex = 0; ViewIndex < SceneRenderer.Views.Num(); ViewIndex++)
+		if (bUpdateLUT)
 		{
-			FViewInfo& View = SceneRenderer.Views[ViewIndex];
-			AddStrataLUTPass(GraphBuilder, View); // STRATA_TODO refactor because this should not be done for each view but only once.
+			AddStrataLUTPass(GraphBuilder, SceneRenderer, StrataSceneData.GGXEnergyLUT2DTexture, StrataSceneData.GGXEnergyLUT3DTexture);
 		}
 	}
 
-	// Create the readable uniform buffers for each views once for all
-	for (int32 ViewIndex = 0; ViewIndex < SceneRenderer.Views.Num(); ViewIndex++)
+	// Create the readable uniform buffers for each views once for all (it is view independent and all the views should be tiled into the render target textures & material buffer)
+	FStrataGlobalUniformParameters* StrataUniformParameters = GraphBuilder.AllocParameters<FStrataGlobalUniformParameters>();
+	if (IsStrataEnabled())
 	{
-		FViewInfo& View = SceneRenderer.Views[ViewIndex];
-
-		FStrataGlobalUniformParameters* StrataUniformParameters = GraphBuilder.AllocParameters<FStrataGlobalUniformParameters>();
-		if (IsStrataEnabled())
-		{
-			StrataUniformParameters->MaxBytesPerPixel = View.StrataSceneData->MaxBytesPerPixel;
-			StrataUniformParameters->MaterialLobesBuffer = View.StrataSceneData->MaterialLobesBufferSRV;
-			StrataUniformParameters->ClassificationTexture = View.StrataSceneData->ClassificationTexture;
-			StrataUniformParameters->TopLayerNormalTexture = View.StrataSceneData->TopLayerNormalTexture;
-			StrataUniformParameters->SSSTexture = View.StrataSceneData->SSSTexture;
-			StrataUniformParameters->GGXEnergyLUT3DTexture = View.StrataSceneData->GGXEnergyLUT3DTexture->GetRenderTargetItem().ShaderResourceTexture;
-			StrataUniformParameters->GGXEnergyLUT2DTexture = View.StrataSceneData->GGXEnergyLUT2DTexture->GetRenderTargetItem().ShaderResourceTexture;
-			StrataUniformParameters->GGXEnergyLUTSampler = TStaticSamplerState<SF_Bilinear, AM_Clamp, AM_Clamp, AM_Clamp>::GetRHI();
-			StrataUniformParameters->GGXEnergyLUTScaleBias = GetStrataGGXEnergyLUTScaleBias();
-			View.StrataSceneData->StrataGlobalUniformParameters = GraphBuilder.CreateUniformBuffer(StrataUniformParameters);
-		}
-		else
-		{
-			// Create each time. This path will go away when Strata is always enabled anyway.
-			StrataUniformParameters->MaxBytesPerPixel = 0;
-			StrataUniformParameters->MaterialLobesBuffer = GraphBuilder.CreateSRV(GraphBuilder.RegisterExternalBuffer(GWhiteVertexBufferWithRDG->Buffer), PF_R32_UINT);
-			StrataUniformParameters->ClassificationTexture = GSystemTextures.GetZeroUIntDummy(GraphBuilder);
-			StrataUniformParameters->TopLayerNormalTexture = GSystemTextures.GetZeroUIntDummy(GraphBuilder);
-			StrataUniformParameters->SSSTexture = GSystemTextures.GetZeroUIntDummy(GraphBuilder);
-			StrataUniformParameters->GGXEnergyLUT3DTexture = GSystemTextures.VolumetricBlackDummy->GetRenderTargetItem().ShaderResourceTexture;
-			StrataUniformParameters->GGXEnergyLUT2DTexture = GSystemTextures.BlackDummy->GetRenderTargetItem().ShaderResourceTexture;
-			StrataUniformParameters->GGXEnergyLUTSampler = TStaticSamplerState<SF_Bilinear, AM_Clamp, AM_Clamp, AM_Clamp>::GetRHI();
-			StrataUniformParameters->GGXEnergyLUTScaleBias = GetStrataGGXEnergyLUTScaleBias();
-			View.StrataSceneData->StrataGlobalUniformParameters = GraphBuilder.CreateUniformBuffer(StrataUniformParameters);
-		}
+		StrataUniformParameters->MaxBytesPerPixel = StrataSceneData.MaxBytesPerPixel;
+		StrataUniformParameters->MaterialLobesBuffer = StrataSceneData.MaterialLobesBufferSRV;
+		StrataUniformParameters->ClassificationTexture = StrataSceneData.ClassificationTexture;
+		StrataUniformParameters->TopLayerNormalTexture = StrataSceneData.TopLayerNormalTexture;
+		StrataUniformParameters->SSSTexture = StrataSceneData.SSSTexture;
+		StrataUniformParameters->GGXEnergyLUT3DTexture = StrataSceneData.GGXEnergyLUT3DTexture->GetRenderTargetItem().ShaderResourceTexture;
+		StrataUniformParameters->GGXEnergyLUT2DTexture = StrataSceneData.GGXEnergyLUT2DTexture->GetRenderTargetItem().ShaderResourceTexture;
+		StrataUniformParameters->GGXEnergyLUTSampler = TStaticSamplerState<SF_Bilinear, AM_Clamp, AM_Clamp, AM_Clamp>::GetRHI();
+		StrataUniformParameters->GGXEnergyLUTScaleBias = GetStrataGGXEnergyLUTScaleBias();
+		StrataSceneData.StrataGlobalUniformParameters = GraphBuilder.CreateUniformBuffer(StrataUniformParameters);
+	}
+	else
+	{
+		// Create each time. This path will go away when Strata is always enabled anyway.
+		StrataUniformParameters->MaxBytesPerPixel = 0;
+		StrataUniformParameters->MaterialLobesBuffer = GraphBuilder.CreateSRV(GraphBuilder.RegisterExternalBuffer(GWhiteVertexBufferWithRDG->Buffer), PF_R32_UINT);
+		StrataUniformParameters->ClassificationTexture = GSystemTextures.GetZeroUIntDummy(GraphBuilder);
+		StrataUniformParameters->TopLayerNormalTexture = GSystemTextures.GetZeroUIntDummy(GraphBuilder);
+		StrataUniformParameters->SSSTexture = GSystemTextures.GetZeroUIntDummy(GraphBuilder);
+		StrataUniformParameters->GGXEnergyLUT3DTexture = GSystemTextures.VolumetricBlackDummy->GetRenderTargetItem().ShaderResourceTexture;
+		StrataUniformParameters->GGXEnergyLUT2DTexture = GSystemTextures.BlackDummy->GetRenderTargetItem().ShaderResourceTexture;
+		StrataUniformParameters->GGXEnergyLUTSampler = TStaticSamplerState<SF_Bilinear, AM_Clamp, AM_Clamp, AM_Clamp>::GetRHI();
+		StrataUniformParameters->GGXEnergyLUTScaleBias = GetStrataGGXEnergyLUTScaleBias();
+		StrataSceneData.StrataGlobalUniformParameters = GraphBuilder.CreateUniformBuffer(StrataUniformParameters);
 	}
 }
 
@@ -759,7 +749,6 @@ class FStrataLUTPassPS : public FGlobalShader
 	using FPermutationDomain = TShaderPermutationDomain<>;
 
 	BEGIN_SHADER_PARAMETER_STRUCT(FParameters, )
-		SHADER_PARAMETER_STRUCT_REF(FViewUniformShaderParameters, ViewUniformBuffer)
 		SHADER_PARAMETER(FIntPoint, SliceXYCount)
 		SHADER_PARAMETER(uint32, EnergyLUTResolution)
 		SHADER_PARAMETER(uint32, NumSamples)
@@ -787,7 +776,8 @@ class FStrataLUTPassPS : public FGlobalShader
 
 IMPLEMENT_GLOBAL_SHADER(FStrataLUTPassPS, "/Engine/Private/Strata/StrataLUT.usf", "MainPS", SF_Pixel);
 
-static void AddStrataLUTPass(FRDGBuilder& GraphBuilder, const FViewInfo& View)
+
+static void AddStrataLUTPass(FRDGBuilder& GraphBuilder, FSceneRenderer& SceneRenderer, TRefCountPtr<IPooledRenderTarget>& GGXEnergyLUT2DTexture, TRefCountPtr<IPooledRenderTarget>& GGXEnergyLUT3DTexture)
 {
 	const uint32 LUTResolution = GetStrataGGXEnergyLUTResolution();
 	const uint32 SliceResolution = FMath::CeilToInt(FMath::Sqrt(float(LUTResolution)));
@@ -795,17 +785,18 @@ static void AddStrataLUTPass(FRDGBuilder& GraphBuilder, const FViewInfo& View)
 	FIntPoint OutputResolutionRT;
 	OutputResolutionRT.X = LUTResolution * SliceResolution;
 	OutputResolutionRT.Y = LUTResolution * SliceResolution;
+	FIntRect RenderTargetRect = FIntRect(FIntPoint(ForceInitToZero), OutputResolutionRT);
 
-	FRDGTextureRef OutLUT3D = GraphBuilder.RegisterExternalTexture(View.StrataSceneData->GGXEnergyLUT3DTexture);
-	FRDGTextureRef OutLUT2D = GraphBuilder.RegisterExternalTexture(View.StrataSceneData->GGXEnergyLUT2DTexture);
+	FRDGTextureRef OutLUT2D = GraphBuilder.RegisterExternalTexture(GGXEnergyLUT2DTexture);
+	FRDGTextureRef OutLUT3D = GraphBuilder.RegisterExternalTexture(GGXEnergyLUT3DTexture);
 
 	// For debug purpose	
 	FRDGTextureDesc UnfoldLUTDesc = FRDGTextureDesc::Create2D(OutputResolutionRT, EPixelFormat::PF_G16R16F, FClearValueBinding::Black, ETextureCreateFlags::TexCreate_RenderTargetable);
 	FRDGTextureRef UnfoldLUTTexture = GraphBuilder.CreateTexture(UnfoldLUTDesc, TEXT("StrataEnergyUnfoldLUT"));
 
-	TShaderMapRef<FStrataLUTPassPS> PixelShader(View.ShaderMap);
+	FGlobalShaderMap* GlobalShaderMap = GetGlobalShaderMap(SceneRenderer.FeatureLevel);
+	TShaderMapRef<FStrataLUTPassPS> PixelShader(GlobalShaderMap);
 	FStrataLUTPassPS::FParameters* Parameters = GraphBuilder.AllocParameters<FStrataLUTPassPS::FParameters>();
-	Parameters->ViewUniformBuffer = View.ViewUniformBuffer;
 	Parameters->NumSamples = FMath::Clamp(CVarStrataLUTSampleCount.GetValueOnAnyThread(), 16, 2048);
 	Parameters->EnergyLUTResolution = LUTResolution;
 	Parameters->SliceXYCount = FIntPoint(SliceResolution, SliceResolution);
@@ -815,15 +806,15 @@ static void AddStrataLUTPass(FRDGBuilder& GraphBuilder, const FViewInfo& View)
 
 	FPixelShaderUtils::AddFullscreenPass<FStrataLUTPassPS>(
 		GraphBuilder,
-		View.ShaderMap,
+		GlobalShaderMap,
 		RDG_EVENT_NAME("StrataLUT"),
 		PixelShader,
 		Parameters,
-		View.ViewRect);
+		RenderTargetRect);
 
 	// Export to untracked because these textures lifetime is multiple frame, and also make sure transition is correctly done.
-	ConvertToUntrackedExternalTexture(GraphBuilder, OutLUT3D, View.StrataSceneData->GGXEnergyLUT3DTexture, ERHIAccess::SRVMask);
-	ConvertToUntrackedExternalTexture(GraphBuilder, OutLUT2D, View.StrataSceneData->GGXEnergyLUT2DTexture, ERHIAccess::SRVMask);
+	ConvertToUntrackedExternalTexture(GraphBuilder, OutLUT2D, GGXEnergyLUT2DTexture, ERHIAccess::SRVMask);
+	ConvertToUntrackedExternalTexture(GraphBuilder, OutLUT3D, GGXEnergyLUT3DTexture, ERHIAccess::SRVMask);
 }
 
 ////////////////////////////////////////////////////////////////////////////////////////////////////////////////////////////////////////////////
