Fixed crash due to recent Strata change with just discovered simple light coming from instance from particle effect.

diff --git a/Engine/Source/Runtime/Renderer/Private/LightRendering.cpp b/Engine/Source/Runtime/Renderer/Private/LightRendering.cpp
--- a/Engine/Source/Runtime/Renderer/Private/LightRendering.cpp
+++ b/Engine/Source/Runtime/Renderer/Private/LightRendering.cpp
@@ -2546,6 +2546,7 @@ void FDeferredShadingSceneRenderer::RenderLightsForHair(
 
 BEGIN_SHADER_PARAMETER_STRUCT(FSimpleLightsStandardDeferredParameters, )
 	SHADER_PARAMETER_RDG_UNIFORM_BUFFER(FSceneTextureUniformParameters, SceneTextures)
+	SHADER_PARAMETER_RDG_UNIFORM_BUFFER(FStrataGlobalUniformParameters, Strata)
 	RENDER_TARGET_BINDING_SLOTS()
 END_SHADER_PARAMETER_STRUCT()
 
@@ -2559,6 +2560,10 @@ void FDeferredShadingSceneRenderer::RenderSimpleLightsStandardDeferred(
 
 	FSimpleLightsStandardDeferredParameters* PassParameters = GraphBuilder.AllocParameters<FSimpleLightsStandardDeferredParameters>();
 	PassParameters->SceneTextures = SceneTextures.UniformBuffer;
+	if (Views.Num() > 0)
+	{
+		PassParameters->Strata = Strata::BindStrataGlobalUniformParameters(Views[0]);
+	}
 	PassParameters->RenderTargets[0] = FRenderTargetBinding(SceneTextures.Color.Target, ERenderTargetLoadAction::ELoad);
 	PassParameters->RenderTargets.DepthStencil = FDepthStencilBinding(SceneTextures.Depth.Target, ERenderTargetLoadAction::ELoad, ERenderTargetLoadAction::ELoad, FExclusiveDepthStencil::DepthRead_StencilWrite);
 
