Remove an extra RDG pass just for DispatchRHICmdList cmd in a mobile renderer

diff --git a/Engine/Source/Runtime/Renderer/Private/MobileShadingRenderer.cpp b/Engine/Source/Runtime/Renderer/Private/MobileShadingRenderer.cpp
--- a/Engine/Source/Runtime/Renderer/Private/MobileShadingRenderer.cpp
+++ b/Engine/Source/Runtime/Renderer/Private/MobileShadingRenderer.cpp
@@ -822,14 +822,6 @@ void FMobileSceneRenderer::Render(FRDGBuilder& GraphBuilder)
 		{
 			GPUSortManager->OnPostRenderOpaque(GraphBuilder);
 		}
-		AddPass(
-			GraphBuilder,
-			RDG_EVENT_NAME("DispatchRHICmdList"),
-			[](FRHICommandListImmediate& RHICmdList)
-			{
-				RHICmdList.ImmediateFlush(EImmediateFlushType::DispatchToRHIThread);
-			}
-		);
 	}
 
 	if (bRequiresPixelProjectedPlanarRelfectionPass)
