Strata constant&UB binding functions now take the scene strata structure as input instead of view because the data is shared between view and sometimes it comes from scene when we do not have any views available.

diff --git a/Engine/Source/Runtime/Renderer/Private/BasePassRendering.cpp b/Engine/Source/Runtime/Renderer/Private/BasePassRendering.cpp
--- a/Engine/Source/Runtime/Renderer/Private/BasePassRendering.cpp
+++ b/Engine/Source/Runtime/Renderer/Private/BasePassRendering.cpp
@@ -603,7 +603,7 @@ TRDGUniformBufferRef<FOpaqueBasePassUniformParameters> CreateOpaqueBasePassUnifo
 	SetupDistortionParams(BasePassParameters.DistortionParams, View);
 
 	// Strata
-	Strata::BindStrataBasePassUniformParameters(GraphBuilder, View, BasePassParameters.Strata);
+	Strata::BindStrataBasePassUniformParameters(GraphBuilder, View.StrataSceneData, BasePassParameters.Strata);
 
 	// Misc
 	BasePassParameters.EyeAdaptationTexture = SystemTextures.White;
diff --git a/Engine/Source/Runtime/Renderer/Private/CompositionLighting/PostProcessAmbientOcclusion.cpp b/Engine/Source/Runtime/Renderer/Private/CompositionLighting/PostProcessAmbientOcclusion.cpp
--- a/Engine/Source/Runtime/Renderer/Private/CompositionLighting/PostProcessAmbientOcclusion.cpp
+++ b/Engine/Source/Runtime/Renderer/Private/CompositionLighting/PostProcessAmbientOcclusion.cpp
@@ -567,7 +567,7 @@ FScreenPassTexture AddAmbientOcclusionSetupPass(
 	PassParameters->SSAOParameters = GetSSAOShaderParameters(View, InputViewport, OutputViewport, CommonParameters.SceneTexturesViewport, EAOTechnique::SSAO);
 	PassParameters->ThresholdInverse = ThresholdInverseValue;
 	PassParameters->InputExtentInverse = FVector2D(1.0f) / FVector2D(InputViewport.Extent);
-	PassParameters->Strata = Strata::BindStrataGlobalUniformParameters(View);
+	PassParameters->Strata = Strata::BindStrataGlobalUniformParameters(View.StrataSceneData);
 	PassParameters->RenderTargets[0] = Output.GetRenderTargetBinding();
 
 	TShaderMapRef<FAmbientOcclusionSetupPS> PixelShader(View.ShaderMap);
@@ -847,7 +847,7 @@ FScreenPassTexture AddAmbientOcclusionPass(
 		// Compute Shader Path
 		FAmbientOcclusionCS::FParameters* PassParameters = GraphBuilder.AllocParameters<FAmbientOcclusionCS::FParameters>();
 		PassParameters->SharedParameters = MoveTemp(SharedParameters);
-		PassParameters->Strata = Strata::BindStrataGlobalUniformParameters(View);
+		PassParameters->Strata = Strata::BindStrataGlobalUniformParameters(View.StrataSceneData);
 		PassParameters->OutTexture = GraphBuilder.CreateUAV(Output.Texture);
 
 		FAmbientOcclusionCS::FPermutationDomain PermutationVector;
@@ -927,7 +927,7 @@ FScreenPassTexture AddAmbientOcclusionPass(
 		// Pixel Shader Path
 		FAmbientOcclusionPS::FParameters* PassParameters = GraphBuilder.AllocParameters<FAmbientOcclusionPS::FParameters>();
 		PassParameters->SharedParameters = MoveTemp(SharedParameters);
-		PassParameters->Strata = Strata::BindStrataGlobalUniformParameters(View);
+		PassParameters->Strata = Strata::BindStrataGlobalUniformParameters(View.StrataSceneData);
 		PassParameters->RenderTargets[0] = Output.GetRenderTargetBinding();
 		if (bDepthBoundsTestEnabled)
 		{
@@ -1137,7 +1137,7 @@ FGTAOHorizonSearchOutputs AddGTAOHorizonSearchIntegratePass(
 	PassParameters->View = View.ViewUniformBuffer;
 	PassParameters->SceneTextures = CommonParameters.SceneTexturesUniformBuffer;
 	PassParameters->HZBParameters = GetHZBParameters(View, HZBInput, SceneViewport.Extent, EAOTechnique::GTAO);
-	PassParameters->Strata = Strata::BindStrataGlobalUniformParameters(View);
+	PassParameters->Strata = Strata::BindStrataGlobalUniformParameters(View.StrataSceneData);
 	PassParameters->SSAOParameters = GetSSAOShaderParameters(View, SceneViewport, OutputViewport, CommonParameters.SceneTexturesViewport, EAOTechnique::GTAO);
 	PassParameters->GTAOParameters = GetGTAOShaderParameters(View, OutputViewport.Extent);
 
@@ -1229,7 +1229,7 @@ FScreenPassTexture AddGTAOInnerIntegratePass(
 
 	PassParameters->View = View.ViewUniformBuffer;
 	PassParameters->SceneTextures = CommonParameters.SceneTexturesUniformBuffer;
-	PassParameters->Strata = Strata::BindStrataGlobalUniformParameters(View);
+	PassParameters->Strata = Strata::BindStrataGlobalUniformParameters(View.StrataSceneData);
 	PassParameters->SSAOParameters = GetSSAOShaderParameters(View, InputViewport, OutputViewport, CommonParameters.SceneTexturesViewport, EAOTechnique::GTAO);
 	PassParameters->GTAOParameters = GetGTAOShaderParameters(View, OutputViewport.Extent);
 
@@ -1311,7 +1311,7 @@ FGTAOHorizonSearchOutputs AddGTAOHorizonSearchPass(
 	PassParameters->View = View.ViewUniformBuffer;
 	PassParameters->SceneTextures = CommonParameters.SceneTexturesUniformBuffer;
 	PassParameters->HZBParameters = GetHZBParameters(View, HZBInput, SceneViewport.Extent, EAOTechnique::GTAO);
-	PassParameters->Strata = Strata::BindStrataGlobalUniformParameters(View);
+	PassParameters->Strata = Strata::BindStrataGlobalUniformParameters(View.StrataSceneData);
 	PassParameters->SSAOParameters = GetSSAOShaderParameters(View, SceneViewport, OutputViewport, CommonParameters.SceneTexturesViewport, EAOTechnique::GTAO);
 	PassParameters->GTAOParameters = GetGTAOShaderParameters(View, OutputViewport.Extent);
 
diff --git a/Engine/Source/Runtime/Renderer/Private/DistanceFieldAmbientOcclusion.cpp b/Engine/Source/Runtime/Renderer/Private/DistanceFieldAmbientOcclusion.cpp
--- a/Engine/Source/Runtime/Renderer/Private/DistanceFieldAmbientOcclusion.cpp
+++ b/Engine/Source/Runtime/Renderer/Private/DistanceFieldAmbientOcclusion.cpp
@@ -381,7 +381,7 @@ void ComputeDistanceFieldNormal(
 			auto* PassParameters = GraphBuilder.AllocParameters<FComputeDistanceFieldNormalCS::FParameters>();
 			PassParameters->View = View.ViewUniformBuffer;
 			PassParameters->SceneTextures = SceneTexturesUniformBuffer;
-			PassParameters->Strata = Strata::BindStrataGlobalUniformParameters(View);
+			PassParameters->Strata = Strata::BindStrataGlobalUniformParameters(View.StrataSceneData);
 			PassParameters->RWDistanceFieldNormal = GraphBuilder.CreateUAV(DistanceFieldNormal);
 
 			TShaderMapRef<FComputeDistanceFieldNormalCS> ComputeShader(View.ShaderMap);
@@ -411,7 +411,7 @@ void ComputeDistanceFieldNormal(
 			auto* PassParameters = GraphBuilder.AllocParameters<FComputeDistanceFieldNormalPS::FParameters>();
 			PassParameters->View = View.ViewUniformBuffer;
 			PassParameters->SceneTextures = SceneTexturesUniformBuffer;
-			PassParameters->Strata = Strata::BindStrataGlobalUniformParameters(View);
+			PassParameters->Strata = Strata::BindStrataGlobalUniformParameters(View.StrataSceneData);
 			PassParameters->RenderTargets[0] = FRenderTargetBinding(DistanceFieldNormal, ViewIndex == 0 ? ERenderTargetLoadAction::EClear : ERenderTargetLoadAction::ELoad);
 
 			GraphBuilder.AddPass(
diff --git a/Engine/Source/Runtime/Renderer/Private/IndirectLightRendering.cpp b/Engine/Source/Runtime/Renderer/Private/IndirectLightRendering.cpp
--- a/Engine/Source/Runtime/Renderer/Private/IndirectLightRendering.cpp
+++ b/Engine/Source/Runtime/Renderer/Private/IndirectLightRendering.cpp
@@ -1227,7 +1227,7 @@ static void AddSkyReflectionPass(
 		}
 		PassParameters->ForwardLightData = View.ForwardLightingResources->ForwardLightDataUniformBuffer;
 
-		PassParameters->Strata =  Strata::BindStrataGlobalUniformParameters(View);
+		PassParameters->Strata = Strata::BindStrataGlobalUniformParameters(View.StrataSceneData);
 	}
 
 	PassParameters->RenderTargets[0] = FRenderTargetBinding(SceneColorTexture.Target, ERenderTargetLoadAction::ELoad);
diff --git a/Engine/Source/Runtime/Renderer/Private/LightRendering.cpp b/Engine/Source/Runtime/Renderer/Private/LightRendering.cpp
--- a/Engine/Source/Runtime/Renderer/Private/LightRendering.cpp
+++ b/Engine/Source/Runtime/Renderer/Private/LightRendering.cpp
@@ -586,7 +586,7 @@ private:
 		FRenderLightParams* RenderLightParams)
 	{
 		FGlobalShader::SetParameters<FViewUniformShaderParameters>(RHICmdList, ShaderRHI, View.ViewUniformBuffer);
-		TRDGUniformBufferRef<FStrataGlobalUniformParameters> StrataUniformBuffer = Strata::BindStrataGlobalUniformParameters(View);
+		TRDGUniformBufferRef<FStrataGlobalUniformParameters> StrataUniformBuffer = Strata::BindStrataGlobalUniformParameters(View.StrataSceneData);
 		FGlobalShader::SetParameters<FStrataGlobalUniformParameters>(RHICmdList, ShaderRHI, StrataUniformBuffer->GetRHIRef());
 
 		if(LightAttenuationTexture.IsBound())
@@ -1102,7 +1102,7 @@ void GetRenderLightParameters(
 {
 	Parameters.SceneTextures = SceneTexturesUniformBuffer;
 	Parameters.HairStrands = HairStrandsUniformBuffer;
-	Parameters.Strata = Strata::BindStrataGlobalUniformParameters(View);
+	Parameters.Strata = Strata::BindStrataGlobalUniformParameters(View.StrataSceneData);
 	Parameters.ShadowMaskTexture = ShadowMaskTexture;
 	Parameters.LightingChannelsTexture = LightingChannelsTexture;
 	Parameters.CloudShadowAO = CloudShadowAOParameters;
@@ -2560,10 +2560,7 @@ void FDeferredShadingSceneRenderer::RenderSimpleLightsStandardDeferred(
 
 	FSimpleLightsStandardDeferredParameters* PassParameters = GraphBuilder.AllocParameters<FSimpleLightsStandardDeferredParameters>();
 	PassParameters->SceneTextures = SceneTextures.UniformBuffer;
-	if (Views.Num() > 0)
-	{
-		PassParameters->Strata = Strata::BindStrataGlobalUniformParameters(Views[0]);
-	}
+	PassParameters->Strata = Strata::BindStrataGlobalUniformParameters(&Scene->StrataSceneData);
 	PassParameters->RenderTargets[0] = FRenderTargetBinding(SceneTextures.Color.Target, ERenderTargetLoadAction::ELoad);
 	PassParameters->RenderTargets.DepthStencil = FDepthStencilBinding(SceneTextures.Depth.Target, ERenderTargetLoadAction::ELoad, ERenderTargetLoadAction::ELoad, FExclusiveDepthStencil::DepthRead_StencilWrite);
 
diff --git a/Engine/Source/Runtime/Renderer/Private/PostProcess/PostProcessSubsurface.cpp b/Engine/Source/Runtime/Renderer/Private/PostProcess/PostProcessSubsurface.cpp
--- a/Engine/Source/Runtime/Renderer/Private/PostProcess/PostProcessSubsurface.cpp
+++ b/Engine/Source/Runtime/Renderer/Private/PostProcess/PostProcessSubsurface.cpp
@@ -992,7 +992,7 @@ void AddSubsurfaceViewPass(
 			PassParameters->RWBurleyGroupBuffer = GraphBuilder.CreateUAV(BurleyGroupBuffer, EPixelFormat::PF_R32_UINT);
 			PassParameters->RWSeparableGroupBuffer = GraphBuilder.CreateUAV(SeparableGroupBuffer, EPixelFormat::PF_R32_UINT);
 			PassParameters->SubsurfaceUniformParameters = UniformBuffer;
-			PassParameters->Strata = Strata::BindStrataGlobalUniformParameters(View);
+			PassParameters->Strata = Strata::BindStrataGlobalUniformParameters(View.StrataSceneData);
 
 			SHADER::FPermutationDomain ComputeShaderPermutationVector;
 			ComputeShaderPermutationVector.Set<SHADER::FDimensionHalfRes>(bHalfRes);
@@ -1096,7 +1096,7 @@ void AddSubsurfaceViewPass(
 				PassParameters->SubsurfaceSampler0 = SubsurfaceSamplerState;
 				PassParameters->GroupBuffer = GraphBuilder.CreateSRV(SubsurfaceBufferUsage[SubsurfaceTypeIndex], EPixelFormat::PF_R32_UINT);
 				PassParameters->IndirectDispatchArgsBuffer = SubsurfaceBufferArgs[SubsurfaceTypeIndex];
-				PassParameters->Strata = Strata::BindStrataGlobalUniformParameters(View);
+				PassParameters->Strata = Strata::BindStrataGlobalUniformParameters(View.StrataSceneData);
 
 				if (SubsurfacePassFunction == SHADER::ESubsurfacePass::PassOne && SubsurfaceType == SHADER::ESubsurfaceType::BURLEY)
 				{
@@ -1156,7 +1156,7 @@ void AddSubsurfaceViewPass(
 		PassParameters->RenderTargets[0] = FRenderTargetBinding(SceneColorTextureOutput, SceneColorTextureLoadAction);
 		PassParameters->SubsurfaceInput0 = GetSubsurfaceInput(SceneColorTexture, SceneViewportParameters);
 		PassParameters->SubsurfaceSampler0 = BilinearBorderSampler;
-		PassParameters->Strata = Strata::BindStrataGlobalUniformParameters(View);
+		PassParameters->Strata = Strata::BindStrataGlobalUniformParameters(View.StrataSceneData);
 
 		// Scattering output target is only used when scattering is enabled.
 		if (SubsurfaceMode != ESubsurfaceMode::Bypass)
@@ -1292,7 +1292,7 @@ FScreenPassTexture AddVisualizeSubsurfacePass(FRDGBuilder& GraphBuilder, const F
 	PassParameters->SubsurfaceInput0.Viewport = GetScreenPassTextureViewportParameters(InputViewport);
 	PassParameters->SubsurfaceSampler0 = TStaticSamplerState<SF_Point, AM_Clamp, AM_Clamp, AM_Clamp>::GetRHI();
 	PassParameters->MiniFontTexture = GetMiniFontTexture();
-	PassParameters->Strata = Strata::BindStrataGlobalUniformParameters(View);
+	PassParameters->Strata = Strata::BindStrataGlobalUniformParameters(View.StrataSceneData);
 
 	TShaderMapRef<FSubsurfaceVisualizePS> PixelShader(View.ShaderMap);
 
diff --git a/Engine/Source/Runtime/Renderer/Private/ScreenSpaceRayTracing.cpp b/Engine/Source/Runtime/Renderer/Private/ScreenSpaceRayTracing.cpp
--- a/Engine/Source/Runtime/Renderer/Private/ScreenSpaceRayTracing.cpp
+++ b/Engine/Source/Runtime/Renderer/Private/ScreenSpaceRayTracing.cpp
@@ -1102,7 +1102,7 @@ void RenderScreenSpaceReflections(
 
 		FScreenSpaceReflectionsStencilPS::FParameters* PassParameters = GraphBuilder.AllocParameters<FScreenSpaceReflectionsStencilPS::FParameters>();
 		PassParameters->CommonParameters = CommonParameters;
-		PassParameters->Strata = Strata::BindStrataGlobalUniformParameters(View);
+		PassParameters->Strata = Strata::BindStrataGlobalUniformParameters(View.StrataSceneData);
 		PassParameters->RenderTargets = RenderTargets;
 		
 		TShaderMapRef<FScreenSpaceReflectionsStencilPS> PixelShader(View.ShaderMap, PermutationVector);
@@ -1185,7 +1185,7 @@ void RenderScreenSpaceReflections(
 	FScreenSpaceReflectionsPS::FParameters* PassParameters = GraphBuilder.AllocParameters<FScreenSpaceReflectionsPS::FParameters>();
 	PassParameters->CommonParameters = CommonParameters;
 	SetSSRParameters(&PassParameters->SSRPassCommonParameter);
-	PassParameters->Strata = Strata::BindStrataGlobalUniformParameters(View);
+	PassParameters->Strata = Strata::BindStrataGlobalUniformParameters(View.StrataSceneData);
 	PassParameters->RenderTargets = RenderTargets;
 
 	TShaderMapRef<FScreenSpaceReflectionsPS> PixelShader(View.ShaderMap, PermutationVector);
diff --git a/Engine/Source/Runtime/Renderer/Private/ShadowRendering.cpp b/Engine/Source/Runtime/Renderer/Private/ShadowRendering.cpp
--- a/Engine/Source/Runtime/Renderer/Private/ShadowRendering.cpp
+++ b/Engine/Source/Runtime/Renderer/Private/ShadowRendering.cpp
@@ -422,7 +422,7 @@ static void BindShaderShaders(FRHICommandList& RHICmdList, FGraphicsPipelineStat
 	if (Strata::IsStrataEnabled())
 	{
 		FRHIPixelShader* ShaderRHI = PixelShader.GetPixelShader();
-		TRDGUniformBufferRef<FStrataGlobalUniformParameters> StrataUniformBuffer = Strata::BindStrataGlobalUniformParameters(View);
+		TRDGUniformBufferRef<FStrataGlobalUniformParameters> StrataUniformBuffer = Strata::BindStrataGlobalUniformParameters(View.StrataSceneData);
 		PixelShader->FGlobalShader::template SetParameters<FStrataGlobalUniformParameters>(RHICmdList, ShaderRHI, StrataUniformBuffer->GetRHIRef());
 	}
 
@@ -1881,7 +1881,7 @@ void FSceneRenderer::RenderShadowProjections(
 		FShadowProjectionPassParameters CommonPassParameters;
 		CommonPassParameters.SceneTextures = GetSceneTextureShaderParameters(SceneTextures.UniformBuffer);
 		CommonPassParameters.HairStrands = HairStrands::BindHairStrandsViewUniformParameters(View);
-		CommonPassParameters.Strata = Strata::BindStrataGlobalUniformParameters(View);
+		CommonPassParameters.Strata = Strata::BindStrataGlobalUniformParameters(View.StrataSceneData);
 		CommonPassParameters.RenderTargets[0] = FRenderTargetBinding(OutputTexture, ERenderTargetLoadAction::ELoad);
 		CommonPassParameters.RenderTargets.DepthStencil =
 			bSubPixelShadow ?
diff --git a/Engine/Source/Runtime/Renderer/Private/ShadowRendering.h b/Engine/Source/Runtime/Renderer/Private/ShadowRendering.h
--- a/Engine/Source/Runtime/Renderer/Private/ShadowRendering.h
+++ b/Engine/Source/Runtime/Renderer/Private/ShadowRendering.h
@@ -1507,7 +1507,7 @@ public:
 
 		if (StrataGlobalParameters.IsBound())
 		{
-			TRDGUniformBufferRef<FStrataGlobalUniformParameters> StrataUniformBuffer = Strata::BindStrataGlobalUniformParameters(View);
+			TRDGUniformBufferRef<FStrataGlobalUniformParameters> StrataUniformBuffer = Strata::BindStrataGlobalUniformParameters(View.StrataSceneData);
 			SetUniformBufferParameter(RHICmdList, ShaderRHI, StrataGlobalParameters, StrataUniformBuffer->GetRHIRef());
 		}
 
diff --git a/Engine/Source/Runtime/Renderer/Private/SingleLayerWaterRendering.cpp b/Engine/Source/Runtime/Renderer/Private/SingleLayerWaterRendering.cpp
--- a/Engine/Source/Runtime/Renderer/Private/SingleLayerWaterRendering.cpp
+++ b/Engine/Source/Runtime/Renderer/Private/SingleLayerWaterRendering.cpp
@@ -411,7 +411,7 @@ void FDeferredShadingSceneRenderer::RenderSingleLayerWaterReflections(
 			Parameters.ForwardLightData = View.ForwardLightingResources->ForwardLightDataUniformBuffer;
 			if (Strata::IsStrataEnabled())
 			{
-				Parameters.Strata = Strata::BindStrataGlobalUniformParameters(View);
+				Parameters.Strata = Strata::BindStrataGlobalUniformParameters(View.StrataSceneData);
 			}
 		};
 
diff --git a/Engine/Source/Runtime/Renderer/Private/Strata/Strata.cpp b/Engine/Source/Runtime/Renderer/Private/Strata/Strata.cpp
--- a/Engine/Source/Runtime/Renderer/Private/Strata/Strata.cpp
+++ b/Engine/Source/Runtime/Renderer/Private/Strata/Strata.cpp
@@ -266,12 +266,12 @@ void InitialiseStrataFrameSceneData(FSceneRenderer& SceneRenderer, FRDGBuilder&
 	}
 }
 
-void BindStrataBasePassUniformParameters(FRDGBuilder& GraphBuilder, const FViewInfo& View, FStrataOpaquePassUniformParameters& OutStrataUniformParameters)
+void BindStrataBasePassUniformParameters(FRDGBuilder& GraphBuilder, FStrataSceneData* StrataSceneData, FStrataOpaquePassUniformParameters& OutStrataUniformParameters)
 {
-	if (View.StrataSceneData)
+	if (StrataSceneData)
 	{
-		OutStrataUniformParameters.MaxBytesPerPixel = View.StrataSceneData->MaxBytesPerPixel;
-		OutStrataUniformParameters.MaterialLobesBufferUAV = View.StrataSceneData->MaterialLobesBufferUAV;
+		OutStrataUniformParameters.MaxBytesPerPixel = StrataSceneData->MaxBytesPerPixel;
+		OutStrataUniformParameters.MaterialLobesBufferUAV = StrataSceneData->MaterialLobesBufferUAV;
 	}
 	else
 	{
@@ -280,12 +280,11 @@ void BindStrataBasePassUniformParameters(FRDGBuilder& GraphBuilder, const FViewI
 	}
 }
 
-TRDGUniformBufferRef<FStrataGlobalUniformParameters> BindStrataGlobalUniformParameters(const FViewInfo& View)
+TRDGUniformBufferRef<FStrataGlobalUniformParameters> BindStrataGlobalUniformParameters(FStrataSceneData* StrataSceneData)
 {
-	// We must always have a strata uniform buffer foreach view even when strata is not enabled.
-	check(View.StrataSceneData != nullptr);
-	check(View.StrataSceneData->StrataGlobalUniformParameters != nullptr);
-	return View.StrataSceneData->StrataGlobalUniformParameters;
+	check(StrataSceneData);
+	check(StrataSceneData->StrataGlobalUniformParameters != nullptr);
+	return StrataSceneData->StrataGlobalUniformParameters;
 }
 
 ////////////////////////////////////////////////////////////////////////// 
@@ -342,7 +341,7 @@ static void AddVisualizeMaterialPasses(FRDGBuilder& GraphBuilder, const FViewInf
 	{
 		FVisualizeMaterialPS::FParameters* PassParameters = GraphBuilder.AllocParameters<FVisualizeMaterialPS::FParameters>();
 		PassParameters->ViewUniformBuffer = View.ViewUniformBuffer;
-		PassParameters->Strata = Strata::BindStrataGlobalUniformParameters(View);
+		PassParameters->Strata = Strata::BindStrataGlobalUniformParameters(View.StrataSceneData);
 		PassParameters->MiniFontTexture = GetMiniFontTexture();
 		PassParameters->SceneTextures = GetSceneTextureParameters(GraphBuilder);
 		PassParameters->RenderTargets[0] = FRenderTargetBinding(SceneColorTexture, ERenderTargetLoadAction::ELoad);
@@ -863,7 +862,7 @@ static void AddStrataFurnacePass(FRDGBuilder& GraphBuilder, const FViewInfo& Vie
 	TShaderMapRef<FStrataFurnaceTestPassPS> PixelShader(View.ShaderMap);
 	FStrataFurnaceTestPassPS::FParameters* Parameters = GraphBuilder.AllocParameters<FStrataFurnaceTestPassPS::FParameters>();
 	Parameters->ViewUniformBuffer = View.ViewUniformBuffer;
-	Parameters->Strata = Strata::BindStrataGlobalUniformParameters(View);
+	Parameters->Strata = Strata::BindStrataGlobalUniformParameters(View.StrataSceneData);
 	Parameters->SceneType = FMath::Clamp(CVarStrataFurnaceTest.GetValueOnAnyThread(), 1, 4);
 	Parameters->IntegratorType = FMath::Clamp(CVarStrataFurnaceTestIntegratorType.GetValueOnAnyThread(), 0, 2);
 	Parameters->PreIntegratedGF = GSystemTextures.PreintegratedGF->GetRenderTargetItem().ShaderResourceTexture;
diff --git a/Engine/Source/Runtime/Renderer/Private/Strata/Strata.h b/Engine/Source/Runtime/Renderer/Private/Strata/Strata.h
--- a/Engine/Source/Runtime/Renderer/Private/Strata/Strata.h
+++ b/Engine/Source/Runtime/Renderer/Private/Strata/Strata.h
@@ -76,9 +76,9 @@ bool IsClassificationEnabled();
 
 void InitialiseStrataFrameSceneData(FSceneRenderer& SceneRenderer, FRDGBuilder& GraphBuilder);
 
-void BindStrataBasePassUniformParameters(FRDGBuilder& GraphBuilder, const FViewInfo& View, FStrataOpaquePassUniformParameters& OutStrataUniformParameters);
+void BindStrataBasePassUniformParameters(FRDGBuilder& GraphBuilder, FStrataSceneData* StrataSceneData, FStrataOpaquePassUniformParameters& OutStrataUniformParameters);
 
-TRDGUniformBufferRef<FStrataGlobalUniformParameters> BindStrataGlobalUniformParameters(const FViewInfo& View);
+TRDGUniformBufferRef<FStrataGlobalUniformParameters> BindStrataGlobalUniformParameters(FStrataSceneData* StrataSceneData);
 
 void AddStrataMaterialClassificationPass(FRDGBuilder& GraphBuilder, const FMinimalSceneTextures& SceneTextures, const TArray<FViewInfo>& Views);
 
