Workaround on a (driver?) precision issue on AMD Vega serie with hair strands voxelization.

diff --git a/Engine/Source/Runtime/Renderer/Private/HairStrands/HairStrandsVoxelization.cpp b/Engine/Source/Runtime/Renderer/Private/HairStrands/HairStrandsVoxelization.cpp
--- a/Engine/Source/Runtime/Renderer/Private/HairStrands/HairStrandsVoxelization.cpp
+++ b/Engine/Source/Runtime/Renderer/Private/HairStrands/HairStrandsVoxelization.cpp
@@ -1370,7 +1370,9 @@ static void AddVirtualVoxelGenerateMipPass(
 	FRDGBufferUAVRef PageIndexBufferUAV = GraphBuilder.CreateUAV(VoxelResources.PageIndexBuffer, PF_R32_UINT);
 	FRDGBufferUAVRef PageIndexOccupancyBufferUAV = GraphBuilder.CreateUAV(VoxelResources.PageIndexOccupancyBuffer, PF_R32G32_UINT);
 	
-	const bool bPatchEmptyPage = GHairVirtualVoxelInvalidEmptyPageIndex > 0;
+	// Note: Do not clear empty page on AMD hardware as there are some issue precision or dispatch issue (to be refined)
+	const bool bAMDPC = IsPCPlatform(View.GetShaderPlatform()) && IsRHIDeviceAMD();
+	const bool bPatchEmptyPage = GHairVirtualVoxelInvalidEmptyPageIndex > 0 && !bAMDPC;
 	if (bPatchEmptyPage)
 	{
 		const uint32 LastMipIt = MipCount - 1;
