Fixed Distance Field Streaming on cooked windows
Tested TopazEntry PIE and cooked win64

diff --git a/Engine/Source/Developer/MeshUtilities/Private/MeshDistanceFieldUtilities.cpp b/Engine/Source/Developer/MeshUtilities/Private/MeshDistanceFieldUtilities.cpp
--- a/Engine/Source/Developer/MeshUtilities/Private/MeshDistanceFieldUtilities.cpp
+++ b/Engine/Source/Developer/MeshUtilities/Private/MeshDistanceFieldUtilities.cpp
@@ -481,6 +481,7 @@ void FMeshUtilities::GenerateSignedDistanceFieldVolumeData(
 				OutMip.BulkOffset = StreamableMipData.Num();
 				StreamableMipData.AddUninitialized(MipDataBytes);
 				OutMip.BulkSize = StreamableMipData.Num() - OutMip.BulkOffset;
+				checkf(OutMip.BulkSize > 0, TEXT("BulkSize was 0 for %s with %ux%ux%u indirection"), *MeshName, IndirectionDimensions.X, IndirectionDimensions.Y, IndirectionDimensions.Z);
 
 				FPlatformMemory::Memcpy(&StreamableMipData[OutMip.BulkOffset], IndirectionTable.GetData(), IndirectionTableBytes);
 
diff --git a/Engine/Source/Runtime/Renderer/Private/DistanceFieldStreaming.cpp b/Engine/Source/Runtime/Renderer/Private/DistanceFieldStreaming.cpp
--- a/Engine/Source/Runtime/Renderer/Private/DistanceFieldStreaming.cpp
+++ b/Engine/Source/Runtime/Renderer/Private/DistanceFieldStreaming.cpp
@@ -404,6 +404,7 @@ void FDistanceFieldSceneData::AsyncUpdate(FDistanceFieldAsyncUpdateParameters Up
 
 	for (FDistanceFieldReadRequest& ReadRequest : UpdateParameters.NewReadRequests)
 	{
+		check(ReadRequest.BulkSize > 0);
 		ReadRequest.ReadOutputDataPtr = (uint8*)FMemory::Malloc(ReadRequest.BulkSize);
 		const bool bIODispatcher = ReadRequest.BulkData->IsUsingIODispatcher();
 
@@ -420,7 +421,7 @@ void FDistanceFieldSceneData::AsyncUpdate(FDistanceFieldAsyncUpdateParameters Up
 		{
 			// Compatibility path without IODispatcher
 			ReadRequest.AsyncHandle = ReadRequest.BulkData->OpenAsyncReadHandle();
-			ReadRequest.AsyncRequest = ReadRequest.AsyncHandle->ReadRequest(ReadRequest.BulkOffset, ReadRequest.BulkSize, AIOP_Normal, nullptr, ReadRequest.ReadOutputDataPtr);
+			ReadRequest.AsyncRequest = ReadRequest.AsyncHandle->ReadRequest(ReadRequest.BulkData->GetBulkDataOffsetInFile() + ReadRequest.BulkOffset, ReadRequest.BulkSize, AIOP_Low, nullptr, ReadRequest.ReadOutputDataPtr);
 		}
 	}
 
@@ -542,6 +543,7 @@ void FDistanceFieldSceneData::ProcessStreamingRequestsFromGPU(
 							ReadRequest.BulkData = &AssetState.BuiltData->StreamableMips;
 							ReadRequest.BulkOffset = MipBuiltData.BulkOffset;
 							ReadRequest.BulkSize = MipBuiltData.BulkSize;
+							check(ReadRequest.BulkSize > 0);
 							NewReadRequests.Add(MoveTemp(ReadRequest));
 
 							NumAllocatedDistanceFieldBricks += MipBuiltData.NumDistanceFieldBricks;
