Don't let un-supported vertex factories for dynamic mesh batches add themselves to the command list.

- GetShader was failing because we were attempting to find FRayTracingDynamicGeometryConverterCS for GPUSkinVertexFactory (not the passthrough one).
- MaterialShouldCache was true, VFShouldCache was true, but ShaderShouldCache was false.
- This scenario can happen when the skin cache runs out of memory and skeletal meshes fallback to GPU skinning which isn't supported by ray tracing.
- In the future we should have a more robust fallback where we t-pose the skeletal mesh instead of not rendering it.

diff --git a/Engine/Source/Runtime/Renderer/Private/RayTracing/RayTracingDynamicGeometry.cpp b/Engine/Source/Runtime/Renderer/Private/RayTracing/RayTracingDynamicGeometry.cpp
--- a/Engine/Source/Runtime/Renderer/Private/RayTracing/RayTracingDynamicGeometry.cpp
+++ b/Engine/Source/Runtime/Renderer/Private/RayTracing/RayTracingDynamicGeometry.cpp
@@ -181,6 +181,14 @@ void FRayTracingDynamicGeometryCollection::AddDynamicMeshBatchForGeometryUpdate(
 
 	for (const FMeshBatch& MeshBatch : UpdateParams.MeshBatches)
 	{
+		if (!ensureMsgf(IsSupportedDynamicVertexFactoryType(MeshBatch.VertexFactory->GetType()),
+			TEXT("FRayTracingDynamicGeometryConverterCS doesn't support %s. Skipping rendering of %s.  This can happen when the skinning cache runs out of space and falls back to GPUSkinVertexFactory."),
+			MeshBatch.VertexFactory->GetType()->GetName()),
+			*PrimitiveSceneProxy->GetOwnerName().ToString())
+		{
+			continue;
+		}
+
 		const FMaterialRenderProxy* FallbackMaterialRenderProxyPtr = nullptr;
 		const FMaterial& Material = MeshBatch.MaterialRenderProxy->GetMaterialWithFallback(Scene->GetFeatureLevel(), FallbackMaterialRenderProxyPtr);
 		auto* MaterialInterface = Material.GetMaterialInterface();
