Fixed Distance Field streaming attempting to upload a completed Read Request when a different asset has been allocated at the same index since the Read Request was kicked off
Fixes crash going into PIE in Frosty -d3d12

diff --git a/Engine/Source/Runtime/Engine/Private/DistanceFieldAtlas.cpp b/Engine/Source/Runtime/Engine/Private/DistanceFieldAtlas.cpp
--- a/Engine/Source/Runtime/Engine/Private/DistanceFieldAtlas.cpp
+++ b/Engine/Source/Runtime/Engine/Private/DistanceFieldAtlas.cpp
@@ -222,6 +222,17 @@ void FDistanceFieldVolumeData::CacheDerivedData(const FString& InStaticMeshDeriv
 
 #endif
 
+uint64 NextDistanceFieldVolumeDataId = 1;
+
+FDistanceFieldVolumeData::FDistanceFieldVolumeData() :
+	LocalSpaceMeshBounds(ForceInit),
+	bMostlyTwoSided(false),
+	bAsyncBuilding(false)
+{
+	Id = NextDistanceFieldVolumeDataId;
+	NextDistanceFieldVolumeDataId++;
+}
+
 void FDistanceFieldVolumeData::Serialize(FArchive& Ar, UObject* Owner)
 {
 	// Note: this is derived data, no need for versioning (bump the DDC guid)
diff --git a/Engine/Source/Runtime/Engine/Public/DistanceFieldAtlas.h b/Engine/Source/Runtime/Engine/Public/DistanceFieldAtlas.h
--- a/Engine/Source/Runtime/Engine/Public/DistanceFieldAtlas.h
+++ b/Engine/Source/Runtime/Engine/Public/DistanceFieldAtlas.h
@@ -256,14 +256,12 @@ public:
 	// Remaining mips are streamed
 	FByteBulkData StreamableMips;
 
+	uint64 Id;
+
 	// For stats
 	FName AssetName;
 
-	FDistanceFieldVolumeData() :
-		LocalSpaceMeshBounds(ForceInit),
-		bMostlyTwoSided(false),
-		bAsyncBuilding(false)
-	{}
+	FDistanceFieldVolumeData();
 
 	void GetResourceSizeEx(FResourceSizeEx& CumulativeResourceSize) const
 	{
@@ -291,6 +289,8 @@ public:
 #endif
 
 	void Serialize(FArchive& Ar, UObject* Owner);
+
+	uint64 GetId() const { return Id; }
 };
 
 class FAsyncDistanceFieldTask;
diff --git a/Engine/Source/Runtime/Renderer/Private/DistanceFieldStreaming.cpp b/Engine/Source/Runtime/Renderer/Private/DistanceFieldStreaming.cpp
--- a/Engine/Source/Runtime/Renderer/Private/DistanceFieldStreaming.cpp
+++ b/Engine/Source/Runtime/Renderer/Private/DistanceFieldStreaming.cpp
@@ -335,9 +335,18 @@ void FDistanceFieldSceneData::AsyncUpdate(FDistanceFieldAsyncUpdateParameters Up
 			BulkDataReadPtr = (const uint8*)ReadRequest.BulkData->LockReadOnly() + ReadRequest.BulkOffset;
 		}
 #endif
-
 		const int32 NumIndirectionEntries = MipBuiltData.IndirectionDimensions.X * MipBuiltData.IndirectionDimensions.Y * MipBuiltData.IndirectionDimensions.Z;
-		check(ReadRequest.BulkSize == NumIndirectionEntries * sizeof(uint32) + ReadRequest.NumDistanceFieldBricks * BrickSizeBytes);
+		const uint32 ExpectedBulkSize = NumIndirectionEntries * sizeof(uint32) + ReadRequest.NumDistanceFieldBricks * BrickSizeBytes;
+
+		check(ReadRequest.BuiltDataId == AssetState.BuiltData->GetId());
+		checkf(ReadRequest.BulkSize == ExpectedBulkSize, 
+			TEXT("Bulk size mismatch: BulkSize %u, ExpectedSize %u, NumIndirectionEntries %u, NumBricks %u, ReversedMip %u"),
+			ReadRequest.BulkSize,
+			ExpectedBulkSize,
+			NumIndirectionEntries,
+			ReadRequest.NumDistanceFieldBricks,
+			ReversedMipIndex);
+
 		const uint32* SourceIndirectionTable = (const uint32*)BulkDataReadPtr;
 		const int32* RESTRICT GlobalBlockOffsets = MipState.AllocatedBlocks.GetData();
 		uint32* DestIndirectionTable = (uint32*)IndirectionTableUploadBuffer.Add_GetRef(MipState.IndirectionTableOffset, NumIndirectionEntries);
@@ -538,6 +547,7 @@ void FDistanceFieldSceneData::ProcessStreamingRequestsFromGPU(
 
 							FDistanceFieldReadRequest ReadRequest;
 							ReadRequest.AssetSetId = AssetSetId;
+							ReadRequest.BuiltDataId = AssetState.BuiltData->GetId();
 							ReadRequest.ReversedMipIndex = ReversedMipIndexToAdd;
 							ReadRequest.NumDistanceFieldBricks = MipBuiltData.NumDistanceFieldBricks;
 							ReadRequest.BulkData = &AssetState.BuiltData->StreamableMips;
@@ -581,6 +591,7 @@ void FDistanceFieldSceneData::ProcessReadRequests(
 
 		FDistanceFieldReadRequest NewReadRequest;
 		NewReadRequest.AssetSetId = AssetMipAdd.AssetId;
+		NewReadRequest.BuiltDataId = AssetState.BuiltData->GetId();
 		NewReadRequest.ReversedMipIndex = AssetMipAdd.ReversedMipIndex;
 		NewReadRequest.NumDistanceFieldBricks = MipBuiltData.NumDistanceFieldBricks;
 		NewReadRequest.AlwaysLoadedDataPtr = AssetState.BuiltData->AlwaysLoadedMip.GetData();
@@ -611,6 +622,8 @@ void FDistanceFieldSceneData::ProcessReadRequests(
 			RequestIndex--;
 
 			if (AssetStateArray.IsValidId(ReadRequest.AssetSetId) 
+				// Prevent attempting to upload after a different asset has been allocated at the same index
+				&& ReadRequest.BuiltDataId == AssetStateArray[ReadRequest.AssetSetId].BuiltData->GetId()
 				// Shader requires sequential reversed mips starting from 0, skip upload if the IO request got out of sync with the streaming feedback requests
 				&& ReadRequest.ReversedMipIndex == AssetStateArray[ReadRequest.AssetSetId].ReversedMips.Num())
 			{
diff --git a/Engine/Source/Runtime/Renderer/Private/ScenePrivate.h b/Engine/Source/Runtime/Renderer/Private/ScenePrivate.h
--- a/Engine/Source/Runtime/Renderer/Private/ScenePrivate.h
+++ b/Engine/Source/Runtime/Renderer/Private/ScenePrivate.h
@@ -1853,6 +1853,7 @@ struct FDistanceFieldReadRequest
 	FSetElementId AssetSetId;
 	int32 ReversedMipIndex = 0;
 	int32 NumDistanceFieldBricks = 0;
+	uint64 BuiltDataId = 0;
 
 	// Used when BulkData is nullptr
 	const uint8* AlwaysLoadedDataPtr = nullptr;
