Cloud rendering improvement:
- Fix for cloud in special project while we wait for a CL coming from UE4 main (another fallout from the special cloud/water hack)
- Added an option to run color box constrain on pixel that are reprojecting data from last frame. OFF by default.

diff --git a/Engine/Shaders/Private/VolumetricRenderTarget.usf b/Engine/Shaders/Private/VolumetricRenderTarget.usf
--- a/Engine/Shaders/Private/VolumetricRenderTarget.usf
+++ b/Engine/Shaders/Private/VolumetricRenderTarget.usf
@@ -235,12 +235,12 @@ void ReconstructVolumetricRenderTargetPS(
 					//	RGBA = NewRGBA;
 					//	Depths = NewDepths;
 					//}
-					//RGBA += float4(0, 0.1 0, 0);
+					//RGBA += float4(0, 0.5, 0, 0);
 				}
-				else //if(ReconstructDepthZ < 0.000001f)
+				else // Because of the test on bUseNewSample above, we know here that we are only dealing with reprojected data  //if(ReconstructDepthZ < 0.000001f)
 				{
 					// TODO: To use this, we need to make sure we prioritise pixe lwith under represented depth.
-#if 0
+#if PERMUTATION_REPROJECTION_BOX_CONSTRAINT
 					// Make sure that history stay in the neightborhood color/transmittance/depth box after reprojection
 					float4 ColorAABBMin = 999999999.0f;
 					float4 ColorAABBMax = 0.0f;
@@ -267,6 +267,7 @@ void ReconstructVolumetricRenderTargetPS(
 						RGBA = clamp(RGBA, ColorAABBMin, ColorAABBMax);
 						Depths = clamp(Depths, DepthsAABBMin, DepthsAABBMax);
 					}
+					//RGBA += float4(0, 0.8, 0.8, 0);
 #endif
 				}
 
@@ -676,7 +677,8 @@ void ComposeVolumetricRTOverScenePS(
 			//DataAcc = float4(1.0, 1.0, 0, 0); 
 		}
 
-		if(ValidSampleCount == 0 && bAllCloudSamplesInFrontOfWater)
+		//if (ValidSampleCount == 0 && bAllCloudSamplesInFrontOfWater)
+		if(ValidSampleCount == 0) // The above is removed because it does not always work. We need to test if pixel is water. A change list for that is coming from UE4/Main soon.
 		{
 			// Depth discontinuities edges
 			float4 weights = 1.0f / (depth4Diff * weightMultiplier + 1.0f);
diff --git a/Engine/Source/Runtime/Renderer/Private/VolumetricRenderTarget.cpp b/Engine/Source/Runtime/Renderer/Private/VolumetricRenderTarget.cpp
--- a/Engine/Source/Runtime/Renderer/Private/VolumetricRenderTarget.cpp
+++ b/Engine/Source/Runtime/Renderer/Private/VolumetricRenderTarget.cpp
@@ -40,6 +40,11 @@ static TAutoConsoleVariable<int32> CVarVolumetricRenderTargetPreferAsyncCompute(
 	TEXT("Whether to prefer using async compute to generate volumetric cloud render targets."),
 	ECVF_SetByScalability | ECVF_RenderThreadSafe);
 
+static TAutoConsoleVariable<int32> CVarVolumetricRenderTargetReprojectionBoxConstraint(
+	TEXT("r.VolumetricRenderTarget.ReprojectionBoxConstraint"), 0,
+	TEXT("Whether reprojected data should be constrained to the new incoming cloud data neighborhod value."),
+	ECVF_SetByScalability | ECVF_RenderThreadSafe);
+
 
 static bool ShouldPipelineCompileVolumetricRenderTargetShaders(EShaderPlatform ShaderPlatform)
 {
@@ -406,7 +411,8 @@ class FReconstructVolumetricRenderTargetPS : public FGlobalShader
 	SHADER_USE_PARAMETER_STRUCT(FReconstructVolumetricRenderTargetPS, FGlobalShader);
 
 	class FHistoryAvailable : SHADER_PERMUTATION_BOOL("PERMUTATION_HISTORY_AVAILABLE");
-	using FPermutationDomain = TShaderPermutationDomain<FHistoryAvailable>;
+	class FReprojectionBoxConstraint : SHADER_PERMUTATION_BOOL("PERMUTATION_REPROJECTION_BOX_CONSTRAINT");
+	using FPermutationDomain = TShaderPermutationDomain<FHistoryAvailable, FReprojectionBoxConstraint>;
 
 	BEGIN_SHADER_PARAMETER_STRUCT(FParameters, )
 		SHADER_PARAMETER_STRUCT_REF(FViewUniformShaderParameters, ViewUniformBuffer)
@@ -485,6 +491,7 @@ void ReconstructVolumetricRenderTarget(
 
 		FReconstructVolumetricRenderTargetPS::FPermutationDomain PermutationVector;
 		PermutationVector.Set<FReconstructVolumetricRenderTargetPS::FHistoryAvailable>(VolumetricCloudRT.GetHistoryValid());
+		PermutationVector.Set<FReconstructVolumetricRenderTargetPS::FReprojectionBoxConstraint>(CVarVolumetricRenderTargetReprojectionBoxConstraint.GetValueOnAnyThread() > 0);
 		TShaderMapRef<FReconstructVolumetricRenderTargetPS> PixelShader(ViewInfo.ShaderMap, PermutationVector);
 
 		FReconstructVolumetricRenderTargetPS::FParameters* PassParameters = GraphBuilder.AllocParameters<FReconstructVolumetricRenderTargetPS::FParameters>();
