Groundwork for persistent ray tracing scene.

- Added FRayTracingScene member to FScene.
- Removed ray tracing scene resources from FViewInfo.

diff --git a/Engine/Plugins/Experimental/GPULightmass/Source/GPULightmass/Private/LightmapRenderer.cpp b/Engine/Plugins/Experimental/GPULightmass/Source/GPULightmass/Private/LightmapRenderer.cpp
--- a/Engine/Plugins/Experimental/GPULightmass/Source/GPULightmass/Private/LightmapRenderer.cpp
+++ b/Engine/Plugins/Experimental/GPULightmass/Source/GPULightmass/Private/LightmapRenderer.cpp
@@ -851,7 +851,7 @@ void FSceneRenderState::SetupRayTracingScene()
 
 		SCOPED_DRAW_EVENTF(RHICmdList, GPULightmassUpdateRayTracingScene, TEXT("GPULightmass UpdateRayTracingScene %d Instances"), StaticMeshInstanceRenderStates.Elements.Num());
 
-		TArray<FRayTracingGeometryInstance, SceneRenderingAllocator>& RayTracingGeometryInstances = View.RayTracingGeometryInstances;
+		TArray<FRayTracingGeometryInstance, SceneRenderingAllocator> RayTracingGeometryInstances;
 		RayTracingGeometryInstances.Append(CachedRayTracingScene->RayTracingGeometryInstances);
 
 		int32 LandscapeStartOffset = RayTracingGeometryInstances.Num();
@@ -945,7 +945,7 @@ void FSceneRenderState::SetupRayTracingScene()
 			RayTracingScene = RHICreateRayTracingScene(Initializer);
 
 			uint32 NumTotalInstances = 0;
-			for (const FRayTracingGeometryInstance& Instance : View.RayTracingGeometryInstances)
+			for (const FRayTracingGeometryInstance& Instance : RayTracingGeometryInstances)
 			{
 				ensure(Instance.NumTransforms >= uint32(Instance.GetTransforms().Num()));
 				NumTotalInstances += Instance.NumTransforms;
diff --git a/Engine/Source/Runtime/Renderer/Private/DeferredShadingRenderer.cpp b/Engine/Source/Runtime/Renderer/Private/DeferredShadingRenderer.cpp
--- a/Engine/Source/Runtime/Renderer/Private/DeferredShadingRenderer.cpp
+++ b/Engine/Source/Runtime/Renderer/Private/DeferredShadingRenderer.cpp
@@ -37,6 +37,7 @@
 #include "GPUScene.h"
 #include "RayTracing/RayTracingMaterialHitShaders.h"
 #include "RayTracing/RayTracingLighting.h"
+#include "RayTracing/RayTracingScene.h"
 #include "RayTracingDynamicGeometryCollection.h"
 #include "SceneTextureParameters.h"
 #include "ScreenSpaceDenoise.h"
@@ -462,9 +463,9 @@ static void RenderOpaqueFX(
 
 #if RHI_RAYTRACING
 
-bool FDeferredShadingSceneRenderer::GatherRayTracingWorldInstances(FRHICommandListImmediate& RHICmdList)
+bool FDeferredShadingSceneRenderer::GatherRayTracingWorldInstancesForView(FRHICommandListImmediate& RHICmdList, FViewInfo& View, FRayTracingScene& RayTracingScene)
 {
-	if (!IsRayTracingEnabled() || Views.Num() == 0)
+	if (!IsRayTracingEnabled())
 	{
 		return false;
 	}
@@ -489,39 +490,28 @@ bool FDeferredShadingSceneRenderer::GatherRayTracingWorldInstances(FRHICommandLi
 
 	FGPUScenePrimitiveCollector DummyDynamicPrimitiveCollector;
 
-	TArray<FRayTracingInstance> RayTracingInstances;
-
-	for (int32 ViewIndex = 0; ViewIndex < Views.Num(); ViewIndex++)
-	{
-		FViewInfo& View = Views[ViewIndex];
-		View.RayTracingGeometryInstances.Reserve(Scene->Primitives.Num());
-
-		RayTracingCollector.AddViewMeshArrays(
-			&View,
-			&View.RayTracedDynamicMeshElements,
-			&View.SimpleElementCollector,
-			&DummyDynamicPrimitiveCollector,
-			ViewFamily.GetFeatureLevel(),
-			&DynamicIndexBufferForInitViews,
-			&DynamicVertexBufferForInitViews,
-			&DynamicReadBufferForInitViews
-			);
-
-		View.DynamicRayTracingMeshCommandStorage.RayTracingMeshCommands.Reserve(Scene->Primitives.Num());
-		View.VisibleRayTracingMeshCommands.Reserve(Scene->Primitives.Num());
-	}
+	RayTracingCollector.AddViewMeshArrays(
+		&View,
+		&View.RayTracedDynamicMeshElements,
+		&View.SimpleElementCollector,
+		&DummyDynamicPrimitiveCollector,
+		ViewFamily.GetFeatureLevel(),
+		&DynamicIndexBufferForInitViews,
+		&DynamicVertexBufferForInitViews,
+		&DynamicReadBufferForInitViews
+		);
 
-	const int8 ReferenceViewIndex = 0;
-	FViewInfo& ReferenceView = Views[ReferenceViewIndex];
+	View.DynamicRayTracingMeshCommandStorage.RayTracingMeshCommands.Reserve(Scene->Primitives.Num());
+	View.VisibleRayTracingMeshCommands.Reserve(Scene->Primitives.Num());
 
 	extern TSet<IPersistentViewUniformBufferExtension*> PersistentViewUniformBufferExtensions;
 
 	for (IPersistentViewUniformBufferExtension* Extension : PersistentViewUniformBufferExtensions)
 	{
-		Extension->BeginRenderView(&ReferenceView);
+		Extension->BeginRenderView(&View);
 	}
 
-	ReferenceView.RayTracingMeshResourceCollector = MakeUnique<FRayTracingMeshResourceCollector>(
+	View.RayTracingMeshResourceCollector = MakeUnique<FRayTracingMeshResourceCollector>(
 		Scene->GetFeatureLevel(),
 		&DynamicIndexBufferForInitViews,
 		&DynamicVertexBufferForInitViews,
@@ -530,13 +520,13 @@ bool FDeferredShadingSceneRenderer::GatherRayTracingWorldInstances(FRHICommandLi
 	FRayTracingMaterialGatheringContext MaterialGatheringContext
 	{
 		Scene,
-		&ReferenceView,
+		&View,
 		ViewFamily,
 		RHICmdList,
-		*ReferenceView.RayTracingMeshResourceCollector
+		*View.RayTracingMeshResourceCollector
 	};
 
-	float CurrentWorldTime = ReferenceView.Family->CurrentWorldTime;
+	const float CurrentWorldTime = View.Family->CurrentWorldTime;
 
 	struct FRelevantPrimitive
 	{
@@ -544,10 +534,9 @@ bool FDeferredShadingSceneRenderer::GatherRayTracingWorldInstances(FRHICommandLi
 		TArrayView<const int32> CachedRayTracingMeshCommandIndices;
 		uint64 StateHash = 0;
 		int32 PrimitiveIndex = -1;
-		int8 ViewIndex = -1;
 		int8 LODIndex = -1;
-		uint8 RayTracedMeshElementsMask = 0;
 		uint8 InstanceMask = 0;
+		bool bStatic = false;
 		bool bAllSegmentsOpaque = true;
 		bool bAnySegmentsCastShadow = false;
 		bool bAnySegmentsDecal = false;
@@ -577,8 +566,8 @@ bool FDeferredShadingSceneRenderer::GatherRayTracingWorldInstances(FRHICommandLi
 		const float CullingRadius = CVarRayTracingCullingRadius.GetValueOnRenderThread();
 		const float CullAngleThreshold = CVarRayTracingCullingAngle.GetValueOnRenderThread();
 		const float AngleThresholdRatio = FMath::Tan(CullAngleThreshold * PI / 180.0f);
-		const FVector ViewOrigin = ReferenceView.ViewMatrices.GetViewOrigin();
-		const FVector ViewDirection = ReferenceView.GetViewDirection();
+		const FVector ViewOrigin = View.ViewMatrices.GetViewOrigin();
+		const FVector ViewDirection = View.GetViewDirection();
 		const bool bCullAllObjects = CullInRayTracing == 2;
 
 		for (int PrimitiveIndex = 0; PrimitiveIndex < Scene->PrimitiveSceneProxies.Num(); PrimitiveIndex++)
@@ -631,69 +620,63 @@ bool FDeferredShadingSceneRenderer::GatherRayTracingWorldInstances(FRHICommandLi
 						if (bAngleIsSmallEnoughToCull)
 						{
 							continue;
-						}					
+						}
 					}
 				}
 			}
 
-			bool bIsDynamic = false;
-
-			FRelevantPrimitive Item;
-			Item.PrimitiveIndex = PrimitiveIndex;
+			if (!View.State)
+			{
+				continue;
+			}
 
-			for (int32 ViewIndex = 0; ViewIndex < Views.Num(); ++ViewIndex)
+			if (View.bIsReflectionCapture)
 			{
-				FViewInfo& View = Views[ViewIndex];
-				if (!View.State)// || View.RayTracingRenderMode == ERayTracingRenderMode::Disabled)
-				{
-					continue;
-				}
+				continue;
+			}
 
-				if (View.bIsReflectionCapture)
-				{
-					continue;
-				}
+			if (View.HiddenPrimitives.Contains(SceneInfo->PrimitiveComponentId))
+			{
+				continue;
+			}
 
-				if (View.HiddenPrimitives.Contains(SceneInfo->PrimitiveComponentId))
-				{
-					continue;
-				}
+			if (View.ShowOnlyPrimitives.IsSet() && !View.ShowOnlyPrimitives->Contains(SceneInfo->PrimitiveComponentId))
+			{
+				continue;
+			}
 
-				if (View.ShowOnlyPrimitives.IsSet() && !View.ShowOnlyPrimitives->Contains(SceneInfo->PrimitiveComponentId))
-				{
-					continue;
-				}
+			// #dxr_todo: ray tracing in scene captures should re-use the persistent RT scene. (UE-112448)
+			bool bShouldRayTraceSceneCapture = GRayTracingSceneCaptures > 0
+				|| (GRayTracingSceneCaptures == -1 && View.bSceneCaptureUsesRayTracing);
 
-				bool bShouldRayTraceSceneCapture = GRayTracingSceneCaptures > 0 || (GRayTracingSceneCaptures == -1 && View.bSceneCaptureUsesRayTracing);
-				if (View.bIsSceneCapture && (!bShouldRayTraceSceneCapture || !SceneInfo->bIsVisibleInReflectionCaptures))
-				{
-					continue;
-				}
+			if (View.bIsSceneCapture && (!bShouldRayTraceSceneCapture || !SceneInfo->bIsVisibleInReflectionCaptures))
+			{
+				continue;
+			}
 				
-				// Check if the primitive has been distance culled already during frustum culling
-				if (View.DistanceCullingPrimitiveMap[PrimitiveIndex])
-				{
-					continue;
-				}
+			// Check if the primitive has been distance culled already during frustum culling
+			if (View.DistanceCullingPrimitiveMap[PrimitiveIndex])
+			{
+				continue;
+			}
 
-				//#dxr_todo UE-68621  The Raytracing code path does not support ShowFlags since data moved to the SceneInfo. 
-				//Touching the SceneProxy to determine this would simply cost too much
-				static const auto RayTracingStaticMeshesCVar = IConsoleManager::Get().FindTConsoleVariableDataInt(TEXT("r.RayTracing.Geometry.StaticMeshes"));
+			//#dxr_todo UE-68621  The Raytracing code path does not support ShowFlags since data moved to the SceneInfo. 
+			//Touching the SceneProxy to determine this would simply cost too much
+			static const auto RayTracingStaticMeshesCVar = IConsoleManager::Get().FindTConsoleVariableDataInt(TEXT("r.RayTracing.Geometry.StaticMeshes"));
 
-				if (SceneInfo->bIsRayTracingStaticRelevant && View.Family->EngineShowFlags.StaticMeshes && RayTracingStaticMeshesCVar && RayTracingStaticMeshesCVar->GetValueOnRenderThread() > 0)
-				{
-					Item.ViewIndex = ViewIndex;
-					RelevantPrimitives.Add(Item);
-				}
-				else if (View.Family->EngineShowFlags.SkeletalMeshes)
-				{
-					Item.RayTracedMeshElementsMask |= 1 << ViewIndex;
-				}
-			}
+			FRelevantPrimitive Item;
+			Item.PrimitiveIndex = PrimitiveIndex;
 
-			if (Item.RayTracedMeshElementsMask)
+			if (SceneInfo->bIsRayTracingStaticRelevant 
+				&& View.Family->EngineShowFlags.StaticMeshes 
+				&& RayTracingStaticMeshesCVar && RayTracingStaticMeshesCVar->GetValueOnRenderThread() > 0)
+			{
+				Item.bStatic = true;
+				RelevantPrimitives.Add(Item);
+			}
+			else if (View.Family->EngineShowFlags.SkeletalMeshes)
 			{
-				Item.ViewIndex = ReferenceViewIndex;
+				Item.bStatic = false;
 				RelevantPrimitives.Add(Item);
 			}
 		}
@@ -722,7 +705,7 @@ bool FDeferredShadingSceneRenderer::GatherRayTracingWorldInstances(FRHICommandLi
 			LODTaskList.Add(FFunctionGraphTask::CreateAndDispatchWhenReady(
 			[	Items = RelevantPrimitives.GetData() + FirstTaskItemIndex,
 				NumItems = FMath::Min(ItemsPerTask, NumTotalItems - FirstTaskItemIndex),
-				Views = Views.GetData(),
+				&View,
 				Scene = this->Scene,
 				LODScaleCVarValue,
 				ForcedLODLevel
@@ -733,15 +716,13 @@ bool FDeferredShadingSceneRenderer::GatherRayTracingWorldInstances(FRHICommandLi
 				for (uint32 i = 0; i < NumItems; ++i)
 				{
 					FRelevantPrimitive& RelevantPrimitive = Items[i];
-					if (RelevantPrimitive.RayTracedMeshElementsMask)
+					if (!RelevantPrimitive.bStatic)
 					{
 						continue; // skip dynamic primitives
 					}
 
 					const int32 PrimitiveIndex = RelevantPrimitive.PrimitiveIndex;
 					const FPrimitiveSceneInfo* SceneInfo = Scene->Primitives[PrimitiveIndex];
-					const int32 ViewIndex = RelevantPrimitive.ViewIndex;
-					const FViewInfo& View = Views[ViewIndex];
 
 					const FPrimitiveBounds& Bounds = Scene->PrimitiveBounds[PrimitiveIndex];
 					const FPrimitiveSceneInfo* RESTRICT PrimitiveSceneInfo = Scene->Primitives[PrimitiveIndex];
@@ -809,13 +790,6 @@ bool FDeferredShadingSceneRenderer::GatherRayTracingWorldInstances(FRHICommandLi
 
 		Scene->GetRayTracingDynamicGeometryCollection()->BeginUpdate();
 
-		ReferenceView.RayTracingGeometryInstances.Reserve(RelevantPrimitives.Num());
-
-		for (int32 ViewIndex = 0; ViewIndex < Views.Num(); ViewIndex++)
-		{
-			Views[ViewIndex].RayTracingGeometryInstances.Reserve(RelevantPrimitives.Num());
-		}
-
 		struct FRayTracingMeshBatchWorkItem
 		{
 			const FPrimitiveSceneProxy* SceneProxy = nullptr;
@@ -851,21 +825,19 @@ bool FDeferredShadingSceneRenderer::GatherRayTracingWorldInstances(FRHICommandLi
 		uint32 NumPendingMeshBatches = 0;
 		const uint32 RayTracingParallelMeshBatchSize = GRayTracingParallelMeshBatchSize;
 
-		auto KickRayTracingMeshBatchTask = [&ReferenceView, &MeshBatchTaskHead, &MeshBatchTaskPage, &NumPendingMeshBatches, Scene = this->Scene]()
+		auto KickRayTracingMeshBatchTask = [&View, &MeshBatchTaskHead, &MeshBatchTaskPage, &NumPendingMeshBatches, Scene = this->Scene]()
 		{
 			if (MeshBatchTaskHead)
 			{
 				FDynamicRayTracingMeshCommandStorage* TaskDynamicCommandStorage = new(FMemStack::Get()) FDynamicRayTracingMeshCommandStorage;
-				ReferenceView.DynamicRayTracingMeshCommandStoragePerTask.Add(TaskDynamicCommandStorage);
+				View.DynamicRayTracingMeshCommandStoragePerTask.Add(TaskDynamicCommandStorage);
 
 				FRayTracingMeshCommandOneFrameArray* TaskVisibleCommands = new(FMemStack::Get()) FRayTracingMeshCommandOneFrameArray;
 				TaskVisibleCommands->Reserve(NumPendingMeshBatches);
-				ReferenceView.VisibleRayTracingMeshCommandsPerTask.Add(TaskVisibleCommands);
-
-				const FViewInfo* ReferenceViewPtr = &ReferenceView;
+				View.VisibleRayTracingMeshCommandsPerTask.Add(TaskVisibleCommands);
 
-				ReferenceView.AddRayTracingMeshBatchTaskList.Add(FFunctionGraphTask::CreateAndDispatchWhenReady(
-					[TaskDataHead = MeshBatchTaskHead, ReferenceViewPtr, Scene, TaskDynamicCommandStorage, TaskVisibleCommands]()
+				View.AddRayTracingMeshBatchTaskList.Add(FFunctionGraphTask::CreateAndDispatchWhenReady(
+					[TaskDataHead = MeshBatchTaskHead, &View, Scene, TaskDynamicCommandStorage, TaskVisibleCommands]()
 				{
 					FTaskTagScope TaskTagScope(ETaskTag::EParallelRenderingThread);
 					TRACE_CPUPROFILER_EVENT_SCOPE(RayTracingMeshBatchTask);
@@ -884,7 +856,7 @@ bool FDeferredShadingSceneRenderer::GatherRayTracingWorldInstances(FRHICommandLi
 									*TaskDynamicCommandStorage, *TaskVisibleCommands,
 									SegmentIndex, WorkItem.InstanceIndex);
 								FMeshPassProcessorRenderState PassDrawRenderState(Scene->UniformBuffers.ViewUniformBuffer);
-								FRayTracingMeshProcessor RayTracingMeshProcessor(&CommandContext, Scene, ReferenceViewPtr, PassDrawRenderState);
+								FRayTracingMeshProcessor RayTracingMeshProcessor(&CommandContext, Scene, &View, PassDrawRenderState);
 								RayTracingMeshProcessor.AddMeshBatch(MeshBatch, 1, WorkItem.SceneProxy);
 							}
 						}
@@ -901,12 +873,12 @@ bool FDeferredShadingSceneRenderer::GatherRayTracingWorldInstances(FRHICommandLi
 			NumPendingMeshBatches = 0;
 		};
 
+		// Local temporary array of instances used for GetDynamicRayTracingInstances()
+		TArray<FRayTracingInstance> TempRayTracingInstances;
 
 		for (const FRelevantPrimitive& RelevantPrimitive : RelevantPrimitives)
 		{
-			const uint8 RayTracedMeshElementsMask = RelevantPrimitive.RayTracedMeshElementsMask;
-
-			if (RayTracedMeshElementsMask == 0)
+			if (RelevantPrimitive.bStatic)
 			{
 				continue;
 			}
@@ -915,14 +887,14 @@ bool FDeferredShadingSceneRenderer::GatherRayTracingWorldInstances(FRHICommandLi
 			FPrimitiveSceneInfo* SceneInfo = Scene->Primitives[PrimitiveIndex];
 
 			FPrimitiveSceneProxy* SceneProxy = Scene->PrimitiveSceneProxies[PrimitiveIndex];
-			RayTracingInstances.Reset();
-			SceneProxy->GetDynamicRayTracingInstances(MaterialGatheringContext, RayTracingInstances);
+			TempRayTracingInstances.Reset();
+			SceneProxy->GetDynamicRayTracingInstances(MaterialGatheringContext, TempRayTracingInstances);
 
 			for (auto DynamicRayTracingGeometryUpdate : MaterialGatheringContext.DynamicRayTracingGeometriesToUpdate)
 			{
 				Scene->GetRayTracingDynamicGeometryCollection()->AddDynamicMeshBatchForGeometryUpdate(
 					Scene,
-					&ReferenceView,
+					&View,
 					SceneProxy,
 					DynamicRayTracingGeometryUpdate,
 					PrimitiveIndex
@@ -931,22 +903,16 @@ bool FDeferredShadingSceneRenderer::GatherRayTracingWorldInstances(FRHICommandLi
 
 			MaterialGatheringContext.DynamicRayTracingGeometriesToUpdate.Reset();
 
-			if (RayTracingInstances.Num() > 0)
+			if (TempRayTracingInstances.Num() > 0)
 			{
-				for (FRayTracingInstance& Instance : RayTracingInstances)
+				for (FRayTracingInstance& Instance : TempRayTracingInstances)
 				{
 					// If geometry still has pending build request then add to list which requires a force build
 					if (Instance.Geometry->HasPendingBuildRequest())
 					{
-						ReferenceView.ForceBuildRayTracingGeometries.Add(Instance.Geometry);
+						RayTracingScene.GeometriesToBuild.Add(Instance.Geometry);
 					}
 
-					FRayTracingGeometryInstance RayTracingInstance = { Instance.Geometry->RayTracingGeometryRHI };
-					RayTracingInstance.UserData.Add((uint32)PrimitiveIndex);
-					RayTracingInstance.Mask = Instance.Mask;
-					RayTracingInstance.bForceOpaque = Instance.bForceOpaque;
-					RayTracingInstance.bDoubleSided = Instance.bDoubleSided;
-
 					// Thin geometries like hair don't have material, as they only support shadow at the moment.
 					if (!ensureMsgf(Instance.GetMaterials().Num() == Instance.Geometry->Initializer.Segments.Num() ||
 						(Instance.Geometry->Initializer.Segments.Num() == 0 && Instance.GetMaterials().Num() == 1) ||
@@ -959,6 +925,15 @@ bool FDeferredShadingSceneRenderer::GatherRayTracingWorldInstances(FRHICommandLi
 						continue;
 					}
 
+					const uint32 InstanceIndex = RayTracingScene.Instances.Num();
+
+					FRayTracingGeometryInstance& RayTracingInstance = RayTracingScene.Instances.AddDefaulted_GetRef();
+					RayTracingInstance.GeometryRHI = Instance.Geometry->RayTracingGeometryRHI;
+					RayTracingInstance.UserData.Add((uint32)PrimitiveIndex);
+					RayTracingInstance.Mask = Instance.Mask;
+					RayTracingInstance.bForceOpaque = Instance.bForceOpaque;
+					RayTracingInstance.bDoubleSided = Instance.bDoubleSided;
+
 					if (Instance.InstanceGPUTransformsSRV.IsValid())
 					{
 						RayTracingInstance.NumTransforms = Instance.NumTransforms;
@@ -984,24 +959,12 @@ bool FDeferredShadingSceneRenderer::GatherRayTracingWorldInstances(FRHICommandLi
 						}
 					}
 
-					for (int32 ViewIndex = 1; ViewIndex < Views.Num(); ViewIndex++)
-					{
-						Views[ViewIndex].RayTracingGeometryInstances.Add(RayTracingInstance);
-					}
-
-					uint32 InstanceIndex = ReferenceView.RayTracingGeometryInstances.AddDefaulted();
-					Swap(ReferenceView.RayTracingGeometryInstances.Last(), RayTracingInstance);
-
-#if DO_CHECK
+				#if DO_CHECK
 					if (Instance.Geometry->DynamicGeometrySharedBufferGenerationID != FRayTracingGeometry::NonSharedVertexBuffers)
 					{
-						ReferenceView.RayTracingGeometriesForValidation.Add(Instance.Geometry);
-						for (int32 ViewIndex = 1; ViewIndex < Views.Num(); ViewIndex++)
-						{
-							Views[ViewIndex].RayTracingGeometriesForValidation.Add(Instance.Geometry);
-						}
+						RayTracingScene.GeometriesToValidate.Add(Instance.Geometry);
 					}
-#endif // DO_CHECK
+				#endif // DO_CHECK
 
 					if (bParallelMeshBatchSetup)
 					{
@@ -1047,9 +1010,9 @@ bool FDeferredShadingSceneRenderer::GatherRayTracingWorldInstances(FRHICommandLi
 						for (int32 SegmentIndex = 0; SegmentIndex < InstanceMaterials.Num(); SegmentIndex++)
 						{
 							const FMeshBatch& MeshBatch = InstanceMaterials[SegmentIndex];
-							FDynamicRayTracingMeshCommandContext CommandContext(ReferenceView.DynamicRayTracingMeshCommandStorage, ReferenceView.VisibleRayTracingMeshCommands, SegmentIndex, InstanceIndex);
+							FDynamicRayTracingMeshCommandContext CommandContext(View.DynamicRayTracingMeshCommandStorage, View.VisibleRayTracingMeshCommands, SegmentIndex, InstanceIndex);
 							FMeshPassProcessorRenderState PassDrawRenderState(Scene->UniformBuffers.ViewUniformBuffer);
-							FRayTracingMeshProcessor RayTracingMeshProcessor(&CommandContext, Scene, &ReferenceView, PassDrawRenderState);
+							FRayTracingMeshProcessor RayTracingMeshProcessor(&CommandContext, Scene, &View, PassDrawRenderState);
 							RayTracingMeshProcessor.AddMeshBatch(MeshBatch, 1, SceneProxy);
 						}
 					}
@@ -1057,7 +1020,7 @@ bool FDeferredShadingSceneRenderer::GatherRayTracingWorldInstances(FRHICommandLi
 
 				if (CVarRayTracingDynamicGeometryLastRenderTimeUpdateDistance.GetValueOnRenderThread() > 0.0f)
 				{
-					if (FVector::Distance(SceneProxy->GetActorPosition(), ReferenceView.ViewMatrices.GetViewOrigin()) < CVarRayTracingDynamicGeometryLastRenderTimeUpdateDistance.GetValueOnRenderThread())
+					if (FVector::Distance(SceneProxy->GetActorPosition(), View.ViewMatrices.GetViewOrigin()) < CVarRayTracingDynamicGeometryLastRenderTimeUpdateDistance.GetValueOnRenderThread())
 					{
 						// Update LastRenderTime for components so that visibility based ticking (like skeletal meshes) can get updated
 						// We are only doing this for dynamic geometries now
@@ -1092,11 +1055,9 @@ bool FDeferredShadingSceneRenderer::GatherRayTracingWorldInstances(FRHICommandLi
 		{
 			const int32 PrimitiveIndex = RelevantPrimitive.PrimitiveIndex;
 			const FPrimitiveSceneInfo* SceneInfo = Scene->Primitives[PrimitiveIndex];
-			const int32 ViewIndex = RelevantPrimitive.ViewIndex;
-			FViewInfo& View = Views[ViewIndex];
 			const int8 LODIndex = RelevantPrimitive.LODIndex;
 
-			if (LODIndex < 0 || RelevantPrimitive.RayTracedMeshElementsMask != 0)
+			if (LODIndex < 0 || !RelevantPrimitive.bStatic)
 			{
 				continue; // skip dynamic primitives and other 
 			}
@@ -1107,7 +1068,7 @@ bool FDeferredShadingSceneRenderer::GatherRayTracingWorldInstances(FRHICommandLi
 			}
 
 			// location if this is a new entry
-			const int32 NewInstanceIndex = View.RayTracingGeometryInstances.Num();
+			const int32 NewInstanceIndex = RayTracingScene.Instances.Num();
 			const uint64 InstanceKey = RelevantPrimitive.InstancingKey();
 
 			const int32 Index = bAutoInstance ? InstanceSet.FindOrAdd(InstanceKey, NewInstanceIndex) : NewInstanceIndex;
@@ -1115,7 +1076,7 @@ bool FDeferredShadingSceneRenderer::GatherRayTracingWorldInstances(FRHICommandLi
 			if (Index != NewInstanceIndex)
 			{
 				// reusing a previous entry, just append to the instance list
-				FRayTracingGeometryInstance& RayTracingInstance = View.RayTracingGeometryInstances[Index];
+				FRayTracingGeometryInstance& RayTracingInstance = RayTracingScene.Instances[Index];
 				RayTracingInstance.NumTransforms++;
 				RayTracingInstance.Transforms.Add(Scene->PrimitiveTransforms[PrimitiveIndex]);
 				RayTracingInstance.UserData.Add((uint32)PrimitiveIndex);
@@ -1140,7 +1101,7 @@ bool FDeferredShadingSceneRenderer::GatherRayTracingWorldInstances(FRHICommandLi
 					}
 				}
 
-				FRayTracingGeometryInstance& RayTracingInstance = View.RayTracingGeometryInstances.Emplace_GetRef();
+				FRayTracingGeometryInstance& RayTracingInstance = RayTracingScene.Instances.AddDefaulted_GetRef();
 				RayTracingInstance.NumTransforms = 1;
 				RayTracingInstance.Transforms.SetNumUninitialized(1);
 				RayTracingInstance.UserData.SetNumUninitialized(1);
@@ -1158,6 +1119,16 @@ bool FDeferredShadingSceneRenderer::GatherRayTracingWorldInstances(FRHICommandLi
 	return true;
 }
 
+static void DeduplicateRayGenerationShaders(TArray< FRHIRayTracingShader*>& RayGenShaders)
+{
+	TSet<FRHIRayTracingShader*> UniqueRayGenShaders;
+	for (FRHIRayTracingShader* Shader : RayGenShaders)
+	{
+		UniqueRayGenShaders.Add(Shader);
+	}
+	RayGenShaders = UniqueRayGenShaders.Array();
+}
+
 bool FDeferredShadingSceneRenderer::DispatchRayTracingWorldUpdates(FRDGBuilder& GraphBuilder)
 {
 	if (!IsRayTracingEnabled() || Views.Num() == 0)
@@ -1199,11 +1170,14 @@ bool FDeferredShadingSceneRenderer::DispatchRayTracingWorldUpdates(FRDGBuilder&
 
 	GRayTracingGeometryManager.ProcessBuildRequests(GraphBuilder.RHICmdList);
 
-	FViewInfo& ReferenceView = Views[0];
-	if (ReferenceView.ForceBuildRayTracingGeometries.Num() > 0)
+	const int32 ReferenceViewIndex = 0;
+	FViewInfo& ReferenceView = Views[ReferenceViewIndex];
+	FRayTracingScene& RayTracingScene = Scene->RayTracingScene;
+
+	if (RayTracingScene.GeometriesToBuild.Num() > 0)
 	{
 		// Force update all the collected geometries (use stack allocator?)
-		GRayTracingGeometryManager.ForceBuildIfPending(GraphBuilder.RHICmdList, ReferenceView.ForceBuildRayTracingGeometries);
+		GRayTracingGeometryManager.ForceBuildIfPending(GraphBuilder.RHICmdList, RayTracingScene.GeometriesToBuild);
 	}
 	
 	if (ReferenceView.AddRayTracingMeshBatchTaskList.Num() > 0)
@@ -1222,29 +1196,25 @@ bool FDeferredShadingSceneRenderer::DispatchRayTracingWorldUpdates(FRDGBuilder&
 
 	RDG_GPU_MASK_SCOPE(GraphBuilder, FRHIGPUMask::All());
 
-	for (int32 ViewIndex = 0; ViewIndex < Views.Num(); ++ViewIndex)
+#ifdef DO_CHECK
+	// Validate all the ray tracing geometries lifetimes
+	int64 SharedBufferGenerationID = (int64)Scene->GetRayTracingDynamicGeometryCollection()->GetSharedBufferGenerationID();
+	for (const FRayTracingGeometry* Geometry : RayTracingScene.GeometriesToValidate)
 	{
-		FViewInfo& View = Views[ViewIndex];
+		checkf(Geometry->DynamicGeometrySharedBufferGenerationID == FRayTracingGeometry::NonSharedVertexBuffers 
+			|| Geometry->DynamicGeometrySharedBufferGenerationID == SharedBufferGenerationID,
+			TEXT("GenerationID %lld, but expected to be %lld or %lld"), 
+			Geometry->DynamicGeometrySharedBufferGenerationID, SharedBufferGenerationID, FRayTracingGeometry::NonSharedVertexBuffers);
+	}
+#endif // DO_CHECK
 
-		uint32 NumTotalInstances = 0;
-		for (const FRayTracingGeometryInstance& Instance : View.RayTracingGeometryInstances)
-		{
-			ensure(Instance.NumTransforms >= uint32(Instance.GetTransforms().Num()));
-			NumTotalInstances += Instance.NumTransforms;
-		}
-		SET_DWORD_STAT(STAT_RayTracingInstances, NumTotalInstances);
+	// #dxr_todo: UE-72565: refactor ray tracing effects to not be member functions of DeferredShadingRenderer. register each effect at startup and just loop over them automatically to gather all required shaders
+	TArray<FRHIRayTracingShader*> RayGenShaders;
 
-#ifdef DO_CHECK
-		// Validate all the ray tracing geometries lifetimes
-		int64 SharedBufferGenerationID = (int64)Scene->GetRayTracingDynamicGeometryCollection()->GetSharedBufferGenerationID();
-		for (const FRayTracingGeometry* Geometry : View.RayTracingGeometriesForValidation)
-		{
-			check(Geometry->DynamicGeometrySharedBufferGenerationID == FRayTracingGeometry::NonSharedVertexBuffers || Geometry->DynamicGeometrySharedBufferGenerationID == SharedBufferGenerationID);
-		}
-#endif // DO_CHECK
+	for (int32 ViewIndex = 0; ViewIndex < Views.Num(); ++ViewIndex)
+	{
+		FViewInfo& View = Views[ViewIndex];
 
-		// #dxr_todo: UE-72565: refactor ray tracing effects to not be member functions of DeferredShadingRenderer. register each effect at startup and just loop over them automatically to gather all required shaders
-		TArray<FRHIRayTracingShader*> RayGenShaders;
 		PrepareRayTracingReflections(View, *Scene, RayGenShaders);
 		PrepareSingleLayerWaterRayTracingReflections(View, *Scene, RayGenShaders);
 		PrepareRayTracingShadows(View, RayGenShaders);
@@ -1258,57 +1228,49 @@ bool FDeferredShadingSceneRenderer::DispatchRayTracingWorldUpdates(FRDGBuilder&
 		PrepareLumenHardwareRayTracingScreenProbeGather(View, RayGenShaders);
 		PrepareLumenHardwareRayTracingReflections(View, RayGenShaders);
 		PrepareLumenHardwareRayTracingVisualize(View, RayGenShaders);
+	}
 
-		// #yuriy_todo: allocate these using RDG transient resources!
-		FRayTracingAccelerationStructureSize SizeInfo = RHICalcRayTracingSceneSize(NumTotalInstances, ERayTracingAccelerationStructureFlags::FastTrace);
-		check(SizeInfo.ResultSize < ~0u);
-		FRHIResourceCreateInfo CreateInfo(TEXT("RayTracingSceneBuffer"));
-		View.RayTracingSceneBuffer = RHICreateBuffer(uint32(SizeInfo.ResultSize), BUF_AccelerationStructure, 0, ERHIAccess::BVHWrite, CreateInfo);
-		View.RayTracingSceneSRV = RHICreateShaderResourceView(View.RayTracingSceneBuffer);
+	DeduplicateRayGenerationShaders(RayGenShaders);
 
-		View.CreateRayTracingSceneTask = FFunctionGraphTask::CreateAndDispatchWhenReady([
-			NumTotalInstances,
-			&ResultScene = View.RayTracingSceneRHI,
-			&ResultBuffer = View.RayTracingSceneBuffer,
-			Instances = MakeArrayView(View.RayTracingGeometryInstances)]()
-		{
-			FTaskTagScope TaskTagScope(ETaskTag::EParallelRenderingThread);
+	RayTracingScene.BeginCreate();
 
-			FRayTracingSceneInitializer SceneInitializer;
-			SceneInitializer.DebugName = FName(TEXT("FDeferredShadingSceneRenderer.Views[].RayTracingScene"));
-			SceneInitializer.Instances = Instances;
-			SceneInitializer.ShaderSlotsPerGeometrySegment = RAY_TRACING_NUM_SHADER_SLOTS;
-			SceneInitializer.NumMissShaderSlots = RAY_TRACING_NUM_MISS_SHADER_SLOTS;
+	if (RayGenShaders.Num())
+	{
+		auto DefaultHitShader = ReferenceView.ShaderMap->GetShader<FOpaqueShadowHitGroup>().GetRayTracingShader();
 
-			ResultScene = RHICreateRayTracingScene(SceneInitializer);
+		ReferenceView.RayTracingMaterialPipeline = BindRayTracingMaterialPipeline(GraphBuilder.RHICmdList, ReferenceView,
+			RayGenShaders,
+			DefaultHitShader
+		);
+	}
 
-		}, TStatId(), nullptr, ENamedThreads::AnyThread);
+	// Initialize common resources used for lighting in ray tracing effects
 
-		if (RayGenShaders.Num())
-		{
-			auto DefaultHitShader = View.ShaderMap->GetShader<FOpaqueShadowHitGroup>().GetRayTracingShader();
+	ReferenceView.RayTracingSubSurfaceProfileTexture = GetSubsufaceProfileTexture_RT(GraphBuilder.RHICmdList);
+	if (!ReferenceView.RayTracingSubSurfaceProfileTexture)
+	{
+		ReferenceView.RayTracingSubSurfaceProfileTexture = GSystemTextures.BlackDummy;
+	}
 
-			View.RayTracingMaterialPipeline = BindRayTracingMaterialPipeline(GraphBuilder.RHICmdList, View,
-				RayGenShaders,
-				DefaultHitShader
-			);
-		}
+	ReferenceView.RayTracingSubSurfaceProfileSRV = RHICreateShaderResourceView(ReferenceView.RayTracingSubSurfaceProfileTexture->GetRenderTargetItem().ShaderResourceTexture, 0);
+
+	for (int32 ViewIndex = 0; ViewIndex < Views.Num(); ++ViewIndex)
+	{
+		FViewInfo& View = Views[ViewIndex];
 
-		// Initialize common resources used for lighting in ray tracing effects
+		View.RayTracingLightData = CreateRayTracingLightData(GraphBuilder.RHICmdList,
+			Scene->Lights, View, EUniformBufferUsage::UniformBuffer_SingleFrame);
 
-		View.RayTracingSubSurfaceProfileTexture = GetSubsufaceProfileTexture_RT(GraphBuilder.RHICmdList);
-		if (!View.RayTracingSubSurfaceProfileTexture)
+		// Send common ray tracing resources from reference view to all others.
+		if (ViewIndex != ReferenceViewIndex)
 		{
-			View.RayTracingSubSurfaceProfileTexture = GSystemTextures.BlackDummy;
+			View.RayTracingSubSurfaceProfileTexture = ReferenceView.RayTracingSubSurfaceProfileTexture;
+			View.RayTracingSubSurfaceProfileSRV = ReferenceView.RayTracingSubSurfaceProfileSRV;
+			View.RayTracingLightData = ReferenceView.RayTracingLightData;
+			View.RayTracingMaterialPipeline = ReferenceView.RayTracingMaterialPipeline;
 		}
-
-		View.RayTracingSubSurfaceProfileSRV = RHICreateShaderResourceView(View.RayTracingSubSurfaceProfileTexture->GetRenderTargetItem().ShaderResourceTexture, 0);
-
-		View.RayTracingLightData = CreateRayTracingLightData(GraphBuilder.RHICmdList, Scene->Lights, View,
-			EUniformBufferUsage::UniformBuffer_SingleFrame);
 	}
 
-
 	const bool bRayTracingAsyncBuild = CVarRayTracingAsyncBuild.GetValueOnRenderThread() != 0;
 
 	if (bRayTracingAsyncBuild && GRHISupportsRayTracingAsyncBuildAccelerationStructure)
@@ -1329,18 +1291,12 @@ bool FDeferredShadingSceneRenderer::DispatchRayTracingWorldUpdates(FRDGBuilder&
 
 			Scene->GetRayTracingDynamicGeometryCollection()->DispatchUpdates(RHIAsyncCmdList);
 
-			for (int32 ViewIndex = 0; ViewIndex < Views.Num(); ++ViewIndex)
-			{
-				FViewInfo& View = Views[ViewIndex];
-				FRHIRayTracingScene* RayTracingScene = View.GetRayTracingSceneChecked();
+			FRHIRayTracingScene* RayTracingSceneRHI = Scene->RayTracingScene.GetRHIRayTracingSceneChecked();
 
-				check(View.RayTracingSceneBuffer.IsValid());
-				check(View.RayTracingSceneSRV.IsValid());
-				RHIAsyncCmdList.BindAccelerationStructureMemory(RayTracingScene, View.RayTracingSceneBuffer, 0);
+			RHIAsyncCmdList.BindAccelerationStructureMemory(RayTracingSceneRHI, Scene->RayTracingScene.GetBufferChecked(), 0);
 
-				SCOPED_DRAW_EVENT(RHIAsyncCmdList, RayTracingScene);
-				RHIAsyncCmdList.BuildAccelerationStructure(RayTracingScene);
-			}
+			SCOPED_DRAW_EVENT(RHIAsyncCmdList, RayTracingScene);
+			RHIAsyncCmdList.BuildAccelerationStructure(RayTracingSceneRHI);
 
 			RHIAsyncCmdList.BeginTransition(RayTracingDynamicGeometryUpdateEndTransition);
 			FRHIAsyncComputeCommandListImmediate::ImmediateDispatch(RHIAsyncCmdList);
@@ -1360,17 +1316,11 @@ bool FDeferredShadingSceneRenderer::DispatchRayTracingWorldUpdates(FRDGBuilder&
 			RDG_GPU_STAT_SCOPE(GraphBuilder, RayTracingScene);
 			AddPass(GraphBuilder, RDG_EVENT_NAME("RayTracingScene"), [this](FRHICommandList& RHICmdList)
 			{
-				for (int32 ViewIndex = 0; ViewIndex < Views.Num(); ++ViewIndex)
-				{
-					FViewInfo& View = Views[ViewIndex];
-					FRHIRayTracingScene* RayTracingScene = View.GetRayTracingSceneChecked();
+				FRHIRayTracingScene* RayTracingSceneRHI = Scene->RayTracingScene.GetRHIRayTracingSceneChecked();
 
-					check(View.RayTracingSceneBuffer.IsValid());
-					check(View.RayTracingSceneSRV.IsValid());
-					RHICmdList.BindAccelerationStructureMemory(RayTracingScene, View.RayTracingSceneBuffer, 0);
+				RHICmdList.BindAccelerationStructureMemory(RayTracingSceneRHI, Scene->RayTracingScene.GetBufferChecked(), 0);
 
-					RHICmdList.BuildAccelerationStructure(RayTracingScene);
-				}
+				RHICmdList.BuildAccelerationStructure(RayTracingSceneRHI);
 
 				// Submit potentially expensive BVH build commands to the GPU as soon as possible.
 				// Avoids a GPU bubble in some CPU-limited cases.
@@ -1387,21 +1337,25 @@ bool FDeferredShadingSceneRenderer::DispatchRayTracingWorldUpdates(FRDGBuilder&
 	return true;
 }
 
-static void ReleaseRaytracingResources(FRDGBuilder& GraphBuilder, TArrayView<FViewInfo> Views)
+static void ReleaseRaytracingResources(FRDGBuilder& GraphBuilder, TArrayView<FViewInfo> Views, FRayTracingScene &RayTracingScene)
 {
-	AddPass(GraphBuilder, [Views](FRHICommandListImmediate& RHICmdList)
+	AddPass(GraphBuilder, [Views, &RayTracingScene](FRHICommandListImmediate& RHICmdList)
 	{
+		if (RayTracingScene.IsCreated())
+		{
+			RHICmdList.ClearRayTracingBindings(RayTracingScene.GetRHIRayTracingScene());
+
+			// If we did not end up rendering anything this frame, then release all ray tracing scene resources.
+			if (RayTracingScene.Instances.Num() == 0)
+			{
+				RayTracingScene.ResetAndReleaseResources();
+			}
+		}
+
 		// Release resources that were bound to the ray tracing scene to allow them to be immediately recycled.
 		for (int32 ViewIndex = 0; ViewIndex < Views.Num(); ++ViewIndex)
 		{
 			FViewInfo& View = Views[ViewIndex];
-			checkf(!View.CreateRayTracingSceneTask.IsValid() || View.CreateRayTracingSceneTask->IsComplete(), 
-				TEXT("Ray tracing scene creation task is expected to be waited upon and ready when we get here."));
-			if (View.RayTracingSceneRHI)
-			{
-				RHICmdList.ClearRayTracingBindings(View.RayTracingSceneRHI);
-				View.RayTracingSceneRHI.SafeRelease();
-			}
 
 			// Release common lighting resources
 			View.RayTracingSubSurfaceProfileSRV.SafeRelease();
@@ -1438,100 +1392,113 @@ void FDeferredShadingSceneRenderer::WaitForRayTracingScene(FRDGBuilder& GraphBui
 
 	AddPass(GraphBuilder, [this](FRHICommandListImmediate& RHICmdList)
 	{
-		for (int32 ViewIndex = 0; ViewIndex < Views.Num(); ++ViewIndex)
-		{
-			FViewInfo& View = Views[ViewIndex];
-			if (!View.RayTracingMaterialPipeline)
-			{
-				check(View.RayTracingMaterialBindings.Num() == 0);
-				continue;
-			}
+		const int32 ReferenceViewIndex = 0;
+		FViewInfo& ReferenceView = Views[ReferenceViewIndex];
 
-			if (View.RayTracingMaterialBindings.Num())
-			{
-				FTaskGraphInterface::Get().WaitUntilTaskCompletes(View.RayTracingMaterialBindingsTask, ENamedThreads::GetRenderThread_Local());
+		check(ReferenceView.RayTracingMaterialPipeline || ReferenceView.RayTracingMaterialBindings.Num() == 0);
+
+		if (ReferenceView.RayTracingMaterialPipeline && ReferenceView.RayTracingMaterialBindings.Num())
+		{
+			FTaskGraphInterface::Get().WaitUntilTaskCompletes(ReferenceView.RayTracingMaterialBindingsTask, ENamedThreads::GetRenderThread_Local());
 
-				// Gather bindings from all chunks and submit them all as a single batch to allow RHI to bind all shader parameters in parallel.
+			// Gather bindings from all chunks and submit them all as a single batch to allow RHI to bind all shader parameters in parallel.
 
-				uint32 NumTotalBindings = 0;
+			uint32 NumTotalBindings = 0;
 
-				for (FRayTracingLocalShaderBindingWriter* BindingWriter : View.RayTracingMaterialBindings)
+			for (FRayTracingLocalShaderBindingWriter* BindingWriter : ReferenceView.RayTracingMaterialBindings)
+			{
+				const FRayTracingLocalShaderBindingWriter::FChunk* Chunk = BindingWriter->GetFirstChunk();
+				while (Chunk)
 				{
-					const FRayTracingLocalShaderBindingWriter::FChunk* Chunk = BindingWriter->GetFirstChunk();
-					while (Chunk)
-					{
-						NumTotalBindings += Chunk->Num;
-						Chunk = Chunk->Next;
-					}
+					NumTotalBindings += Chunk->Num;
+					Chunk = Chunk->Next;
 				}
+			}
 
-				const uint32 MergedBindingsSize = sizeof(FRayTracingLocalShaderBindings) * NumTotalBindings;
-				FRayTracingLocalShaderBindings* MergedBindings = (FRayTracingLocalShaderBindings*)(RHICmdList.Bypass()
-					? FMemStack::Get().Alloc(MergedBindingsSize, alignof(FRayTracingLocalShaderBindings))
-					: RHICmdList.Alloc(MergedBindingsSize, alignof(FRayTracingLocalShaderBindings)));
+			const uint32 MergedBindingsSize = sizeof(FRayTracingLocalShaderBindings) * NumTotalBindings;
+			FRayTracingLocalShaderBindings* MergedBindings = (FRayTracingLocalShaderBindings*)(RHICmdList.Bypass()
+				? FMemStack::Get().Alloc(MergedBindingsSize, alignof(FRayTracingLocalShaderBindings))
+				: RHICmdList.Alloc(MergedBindingsSize, alignof(FRayTracingLocalShaderBindings)));
 
-				uint32 MergedBindingIndex = 0;
-				for (FRayTracingLocalShaderBindingWriter* BindingWriter : View.RayTracingMaterialBindings)
+			uint32 MergedBindingIndex = 0;
+			for (FRayTracingLocalShaderBindingWriter* BindingWriter : ReferenceView.RayTracingMaterialBindings)
+			{
+				const FRayTracingLocalShaderBindingWriter::FChunk* Chunk = BindingWriter->GetFirstChunk();
+				while (Chunk)
 				{
-					const FRayTracingLocalShaderBindingWriter::FChunk* Chunk = BindingWriter->GetFirstChunk();
-					while (Chunk)
+					const uint32 Num = Chunk->Num;
+					for (uint32_t i = 0; i < Num; ++i)
 					{
-						const uint32 Num = Chunk->Num;
-						for (uint32_t i = 0; i < Num; ++i)
-						{
-							MergedBindings[MergedBindingIndex] = Chunk->Bindings[i];
-							MergedBindingIndex++;
-						}
-						Chunk = Chunk->Next;
+						MergedBindings[MergedBindingIndex] = Chunk->Bindings[i];
+						MergedBindingIndex++;
 					}
+					Chunk = Chunk->Next;
 				}
+			}
 
-				checkf(!View.CreateRayTracingSceneTask.IsValid() || View.CreateRayTracingSceneTask->IsComplete(),
-					TEXT("Ray tracing scene creation task is expected to be waited upon and ready when we get here."));
+			Scene->RayTracingScene.WaitForTasks();
 
-				const bool bCopyDataToInlineStorage = false; // Storage is already allocated from RHICmdList, no extra copy necessary
-				RHICmdList.SetRayTracingHitGroups(
-					View.GetRayTracingSceneChecked(),
-					View.RayTracingMaterialPipeline,
-					NumTotalBindings, MergedBindings,
-					bCopyDataToInlineStorage);
+			const bool bCopyDataToInlineStorage = false; // Storage is already allocated from RHICmdList, no extra copy necessary
+			RHICmdList.SetRayTracingHitGroups(
+				ReferenceView.GetRayTracingSceneChecked(),
+				ReferenceView.RayTracingMaterialPipeline,
+				NumTotalBindings, MergedBindings,
+				bCopyDataToInlineStorage);
 
-				TArray<FRHIRayTracingShader*> DeferredMaterialRayGenShaders;
 
+			TArray<FRHIRayTracingShader*> DeferredMaterialRayGenShaders;
+			for (int32 ViewIndex = 0; ViewIndex < Views.Num(); ++ViewIndex)
+			{
+				FViewInfo& View = Views[ViewIndex];
 				PrepareRayTracingReflectionsDeferredMaterial(View, *Scene, DeferredMaterialRayGenShaders);
 				PrepareRayTracingDeferredReflectionsDeferredMaterial(View, *Scene, DeferredMaterialRayGenShaders);
 				PrepareRayTracingGlobalIlluminationDeferredMaterial(View, DeferredMaterialRayGenShaders);
 				PrepareLumenHardwareRayTracingReflectionsDeferredMaterial(View, DeferredMaterialRayGenShaders);
 				PrepareLumenHardwareRayTracingScreenProbeGatherDeferredMaterial(View, DeferredMaterialRayGenShaders);
 				PrepareLumenHardwareRayTracingVisualizeDeferredMaterial(View, DeferredMaterialRayGenShaders);
+			}
+			DeduplicateRayGenerationShaders(DeferredMaterialRayGenShaders);
 
-				if (DeferredMaterialRayGenShaders.Num())
-				{
-					View.RayTracingMaterialGatherPipeline = BindRayTracingDeferredMaterialGatherPipeline(RHICmdList, View, DeferredMaterialRayGenShaders);
-				}
+			if (DeferredMaterialRayGenShaders.Num())
+			{
+				ReferenceView.RayTracingMaterialGatherPipeline = BindRayTracingDeferredMaterialGatherPipeline(RHICmdList, ReferenceView, DeferredMaterialRayGenShaders);
+			}
 
-				// Add Lumen hardware ray tracing materials
-				TArray<FRHIRayTracingShader*> LumenHardwareRayTracingRayGenShaders;
+			// Add Lumen hardware ray tracing materials
+			TArray<FRHIRayTracingShader*> LumenHardwareRayTracingRayGenShaders;
+			for (int32 ViewIndex = 0; ViewIndex < Views.Num(); ++ViewIndex)
+			{
+				FViewInfo& View = Views[ViewIndex];
 				PrepareLumenHardwareRayTracingVisualizeLumenMaterial(View, LumenHardwareRayTracingRayGenShaders);
 				PrepareLumenHardwareRayTracingReflectionsLumenMaterial(View, LumenHardwareRayTracingRayGenShaders);
 				PrepareLumenHardwareRayTracingScreenProbeGatherLumenMaterial(View, LumenHardwareRayTracingRayGenShaders);
-				if (LumenHardwareRayTracingRayGenShaders.Num())
+			}
+			DeduplicateRayGenerationShaders(DeferredMaterialRayGenShaders);
+
+			if (LumenHardwareRayTracingRayGenShaders.Num())
+			{
+				ReferenceView.LumenHardwareRayTracingMaterialPipeline = BindLumenHardwareRayTracingMaterialPipeline(RHICmdList, ReferenceView, LumenHardwareRayTracingRayGenShaders);
+			}
+
+			// Move the ray tracing binding container ownership to the command list, so that memory will be
+			// released on the RHI thread timeline, after the commands that reference it are processed.
+			RHICmdList.EnqueueLambda([Ptrs = MoveTemp(ReferenceView.RayTracingMaterialBindings)](FRHICommandListImmediate&)
+			{
+				for (auto Ptr : Ptrs)
 				{
-					View.LumenHardwareRayTracingMaterialPipeline = BindLumenHardwareRayTracingMaterialPipeline(RHICmdList, View, LumenHardwareRayTracingRayGenShaders);
+					delete Ptr;
 				}
+			});
 
-				// Move the ray tracing binding container ownership to the command list, so that memory will be
-				// released on the RHI thread timeline, after the commands that reference it are processed.
-				RHICmdList.EnqueueLambda([Ptrs = MoveTemp(View.RayTracingMaterialBindings)](FRHICommandListImmediate&)
-				{
-					for (auto Ptr : Ptrs)
-					{
-						delete Ptr;
-					}
-				});
+			// Send ray tracing resources from reference view to all others.
+			for (int32 ViewIndex = 1; ViewIndex < Views.Num(); ++ViewIndex)
+			{
+				FViewInfo& View = Views[ViewIndex];
+				View.RayTracingMaterialGatherPipeline = ReferenceView.RayTracingMaterialGatherPipeline;
+				View.LumenHardwareRayTracingMaterialPipeline = ReferenceView.LumenHardwareRayTracingMaterialPipeline;
 			}
 
-			SetupRayTracingLightingMissShader(RHICmdList, View);
+			SetupRayTracingLightingMissShader(RHICmdList, ReferenceView);
 		}
 
 		if (RayTracingDynamicGeometryUpdateEndTransition)
@@ -1540,11 +1507,8 @@ void FDeferredShadingSceneRenderer::WaitForRayTracingScene(FRDGBuilder& GraphBui
 			RayTracingDynamicGeometryUpdateEndTransition = nullptr;
 		}
 
-		for (int32 ViewIndex = 0; ViewIndex < Views.Num(); ++ViewIndex)
-		{
-			FRHIRayTracingScene* RayTracingScene = Views[ViewIndex].GetRayTracingSceneChecked();
-			RHICmdList.Transition(FRHITransitionInfo(RayTracingScene, ERHIAccess::BVHWrite, ERHIAccess::BVHRead));
-		}
+		FRHIRayTracingScene* RayTracingScene = ReferenceView.GetRayTracingSceneChecked();
+		RHICmdList.Transition(FRHITransitionInfo(RayTracingScene, ERHIAccess::BVHWrite, ERHIAccess::BVHRead));
 	});
 }
 
@@ -1767,14 +1731,24 @@ void FDeferredShadingSceneRenderer::Render(FRDGBuilder& GraphBuilder)
 	}
 
 #if RHI_RAYTRACING
+
 	// Gather mesh instances, shaders, resources, parameters, etc. and build ray tracing acceleration structure
-	GatherRayTracingWorldInstances(GraphBuilder.RHICmdList);
 
-	if (Views[0].RayTracingRenderMode != ERayTracingRenderMode::PathTracing)
+	FRayTracingScene& RayTracingScene = Scene->RayTracingScene;
+	RayTracingScene.Reset(); // Resets the internal arrays, but does not release any resources.
+
+	const int32 ReferenceViewIndex = 0;
+	FViewInfo& ReferenceView = Views[ReferenceViewIndex];
+
+	// Prepare the scene for rendering this frame.
+	GatherRayTracingWorldInstancesForView(GraphBuilder.RHICmdList, ReferenceView, RayTracingScene);
+
+	if (ReferenceView.RayTracingRenderMode != ERayTracingRenderMode::PathTracing)
 	{
 		extern ENGINE_API float GAveragePathTracedMRays;
 		GAveragePathTracedMRays = 0.0f;
 	}
+
 #endif // RHI_RAYTRACING
 
 	{
@@ -2943,7 +2917,7 @@ void FDeferredShadingSceneRenderer::Render(FRDGBuilder& GraphBuilder)
 	GEngine->GetPostRenderDelegateEx().Broadcast(GraphBuilder);
 
 #if RHI_RAYTRACING
-	ReleaseRaytracingResources(GraphBuilder, Views);
+	ReleaseRaytracingResources(GraphBuilder, Views, Scene->RayTracingScene);
 #endif //  RHI_RAYTRACING
 
 #if WITH_MGPU
diff --git a/Engine/Source/Runtime/Renderer/Private/DeferredShadingRenderer.h b/Engine/Source/Runtime/Renderer/Private/DeferredShadingRenderer.h
--- a/Engine/Source/Runtime/Renderer/Private/DeferredShadingRenderer.h
+++ b/Engine/Source/Runtime/Renderer/Private/DeferredShadingRenderer.h
@@ -951,7 +951,9 @@ private:
 	void RenderRayTracingDebug(FRDGBuilder& GraphBuilder, const FViewInfo& View, FRDGTextureRef SceneColorOutputTexture);
 	void RenderRayTracingBarycentrics(FRDGBuilder& GraphBuilder, const FViewInfo& View, FRDGTextureRef SceneColorOutputTexture);
 
-	bool GatherRayTracingWorldInstances(FRHICommandListImmediate& RHICmdList);
+	/** Fills RayTracingScene instance list for the given View and adds relevant ray tracing data to the view. Does not reset previous scene contents. */
+	bool GatherRayTracingWorldInstancesForView(FRHICommandListImmediate& RHICmdList, FViewInfo& View, FRayTracingScene& RayTracingScene);
+
 	bool DispatchRayTracingWorldUpdates(FRDGBuilder& GraphBuilder);
 	FRayTracingPipelineState* BindRayTracingMaterialPipeline(FRHICommandList& RHICmdList, FViewInfo& View, const TArrayView<FRHIRayTracingShader*>& RayGenShaderTable, FRHIRayTracingShader* DefaultClosestHitShader);
 	FRayTracingPipelineState* BindRayTracingDeferredMaterialGatherPipeline(FRHICommandList& RHICmdList, const FViewInfo& View, const TArrayView<FRHIRayTracingShader*>& RayGenShaderTable);
diff --git a/Engine/Source/Runtime/Renderer/Private/RayTracing/RayTracingScene.cpp b/Engine/Source/Runtime/Renderer/Private/RayTracing/RayTracingScene.cpp
new file mode 100644
--- /dev/null
+++ b/Engine/Source/Runtime/Renderer/Private/RayTracing/RayTracingScene.cpp
@@ -0,0 +1,139 @@
+// Copyright Epic Games, Inc. All Rights Reserved.
+
+#include "RayTracingScene.h"
+
+#if RHI_RAYTRACING
+
+#include "RenderCore.h"
+#include "RayTracingDefinitions.h"
+
+FRayTracingScene::FRayTracingScene()
+{
+
+}
+
+FRayTracingScene::~FRayTracingScene()
+{
+	WaitForTasks();
+}
+
+FGraphEventRef FRayTracingScene::BeginCreate()
+{
+	WaitForTasks();
+
+	uint32 NumTotalInstances = 0;
+	for (const FRayTracingGeometryInstance& Instance : Instances)
+	{
+		ensure(Instance.NumTransforms >= uint32(Instance.GetTransforms().Num()));
+		NumTotalInstances += Instance.NumTransforms;
+	}
+	SET_DWORD_STAT(STAT_RayTracingInstances, NumTotalInstances);
+
+	// Round up number of instances to some multiple to avoid pathological growth reallocations.
+	static constexpr uint32 AllocationGranularity = 8 * 1024;
+	NumTotalInstances = FMath::DivideAndRoundUp(NumTotalInstances, AllocationGranularity) * AllocationGranularity;
+
+	// #yuriy_todo: allocate these using RDG transient resources!
+	FRayTracingAccelerationStructureSize SizeInfo = RHICalcRayTracingSceneSize(NumTotalInstances, ERayTracingAccelerationStructureFlags::FastTrace);
+	check(SizeInfo.ResultSize < ~0u);
+
+	// Allocate GPU buffer if current one is too small or significantly larger than what we need.
+	if (!RayTracingSceneBuffer.IsValid() 
+		|| SizeInfo.ResultSize > RayTracingSceneBuffer->GetSize() 
+		|| SizeInfo.ResultSize < RayTracingSceneBuffer->GetSize() / 2)
+	{
+		RayTracingSceneSRV = nullptr;
+		RayTracingSceneBuffer = nullptr;
+
+		FRHIResourceCreateInfo CreateInfo(TEXT("RayTracingSceneBuffer"));
+		RayTracingSceneBuffer = RHICreateBuffer(uint32(SizeInfo.ResultSize), BUF_AccelerationStructure, 0, ERHIAccess::BVHWrite, CreateInfo);
+		RayTracingSceneSRV = RHICreateShaderResourceView(RayTracingSceneBuffer);
+	}
+
+	CreateRayTracingSceneTask = FFunctionGraphTask::CreateAndDispatchWhenReady([
+			&ResultScene = RayTracingSceneRHI,
+			Instances = MakeArrayView(Instances)]()
+		{
+			FTaskTagScope TaskTagScope(ETaskTag::EParallelRenderingThread);
+
+			FRayTracingSceneInitializer SceneInitializer;
+			SceneInitializer.DebugName = FName(TEXT("FRayTracingScene"));
+			SceneInitializer.Instances = Instances;
+			SceneInitializer.ShaderSlotsPerGeometrySegment = RAY_TRACING_NUM_SHADER_SLOTS;
+			SceneInitializer.NumMissShaderSlots = RAY_TRACING_NUM_MISS_SHADER_SLOTS;
+
+			ResultScene = RHICreateRayTracingScene(SceneInitializer);
+
+		}, TStatId(), nullptr, ENamedThreads::AnyThread);
+
+	return CreateRayTracingSceneTask;
+}
+
+void FRayTracingScene::WaitForTasks() const
+{
+	if (CreateRayTracingSceneTask.IsValid())
+	{
+		TRACE_CPUPROFILER_EVENT_SCOPE(WaitForCreateRayTracingScene);
+		FTaskGraphInterface::Get().WaitUntilTaskCompletes(CreateRayTracingSceneTask, ENamedThreads::GetRenderThread_Local());
+
+		CreateRayTracingSceneTask = {};
+	}
+}
+
+bool FRayTracingScene::IsCreated() const
+{
+	// Either scene creation task has been started (CreateRayTracingSceneTask != null) or we already have a valid RayTracingSceneRHI.
+	return CreateRayTracingSceneTask.IsValid() || RayTracingSceneRHI.IsValid();
+}
+
+FRHIRayTracingScene* FRayTracingScene::GetRHIRayTracingScene() const
+{
+	WaitForTasks();
+
+	checkf(!CreateRayTracingSceneTask.IsValid() || CreateRayTracingSceneTask->IsComplete(),
+		TEXT("Ray tracing scene creation task is expected to be waited upon and ready when we get here."));
+
+	return RayTracingSceneRHI.GetReference();
+}
+
+FRHIRayTracingScene* FRayTracingScene::GetRHIRayTracingSceneChecked() const
+{
+	FRHIRayTracingScene* Result = GetRHIRayTracingScene();
+	checkf(Result, TEXT("Ray tracing scene was not created. Perhaps BeginCreate() was not called."));
+	return Result;
+}
+
+FRHIShaderResourceView* FRayTracingScene::GetShaderResourceViewChecked() const
+{
+	checkf(RayTracingSceneSRV.IsValid(), TEXT("Ray tracing scene SRV was not created. Perhaps BeginCreate() was not called."));
+	return RayTracingSceneSRV.GetReference();
+}
+
+FRHIBuffer* FRayTracingScene::GetBufferChecked() const
+{
+	checkf(RayTracingSceneBuffer.IsValid(), TEXT("Ray tracing scene buffer was not created. Perhaps BeginCreate() was not called."));
+	return RayTracingSceneBuffer.GetReference();
+}
+
+void FRayTracingScene::Reset()
+{
+	WaitForTasks();
+
+	Instances.Reset();
+	GeometriesToBuild.Reset();
+#ifdef DO_CHECK
+	GeometriesToValidate.Reset();
+#endif
+}
+
+void FRayTracingScene::ResetAndReleaseResources()
+{
+	Reset();
+
+	Instances.Empty();
+	RayTracingSceneSRV = nullptr;
+	RayTracingSceneBuffer = nullptr;
+	RayTracingSceneRHI = nullptr;
+}
+
+#endif // RHI_RAYTRACING
diff --git a/Engine/Source/Runtime/Renderer/Private/RayTracing/RayTracingScene.h b/Engine/Source/Runtime/Renderer/Private/RayTracing/RayTracingScene.h
new file mode 100644
--- /dev/null
+++ b/Engine/Source/Runtime/Renderer/Private/RayTracing/RayTracingScene.h
@@ -0,0 +1,87 @@
+// Copyright Epic Games, Inc. All Rights Reserved.
+
+#pragma once
+
+#include "CoreMinimal.h"
+#include "RHIDefinitions.h"
+
+#if RHI_RAYTRACING
+
+#include "Async/TaskGraphInterfaces.h"
+#include "RHI.h"
+
+class FRHIRayTracingScene;
+class FRHIShaderResourceView;
+class FRayTracingGeometry;
+
+/**
+* Persistent representation of the scene for ray tracing.
+* Manages top level acceleration structure instances, memory and build process.
+*/
+class FRayTracingScene
+{
+public:
+
+	FRayTracingScene();
+	~FRayTracingScene();
+
+	// Starts asynchronous creation of RayTracingSceneRHI.
+	// This can be an expensive process, depending on the number of instances in the scene.
+	// Immediately allocates GPU memory to fit at least the current number of instances and updates the SRV.
+	// Returns the task that must be waited upon before accessing RayTracingSceneRHI.
+	FGraphEventRef BeginCreate();
+	void WaitForTasks() const;
+
+	// Resets the instance list and reserves memory for this frame.
+	void Reset();
+
+	// Similar to Reset(), but also releases any persistent CPU and GPU memory allocations.
+	void ResetAndReleaseResources();
+
+	// Returns true if RHI ray tracing scene has been created or asynchronous creation has been kicked.
+	// i.e. returns true after BeginCreate() and before Reset().
+	bool IsCreated() const;
+
+	// Waits for the async creation task if it's active and then returns RayTracingSceneRHI object (may return null).
+	FRHIRayTracingScene* GetRHIRayTracingScene() const;
+
+	// Similar to GetRayTracingScene, but checks that ray tracing scene RHI object is valid.
+	FRHIRayTracingScene* GetRHIRayTracingSceneChecked() const;
+
+	// Returns Buffer and SRV for this ray tracing scene.
+	// Valid to call immediately after BeginCreate() and does not block.
+	FRHIShaderResourceView* GetShaderResourceViewChecked() const;
+	FRHIBuffer* GetBufferChecked() const;
+
+public:
+
+	// Public members for initial refactoring step (previously were public members of FViewInfo).
+
+	// Persistent storage for ray tracing instance descriptors.
+	// Cleared every frame without releasing memory to avoid large heap allocations.
+	// This must be filled before calling BeginCreateTask().
+	TArray<FRayTracingGeometryInstance> Instances;
+
+	// Geometries which still have a pending build request but are used this frame and require a force build.
+	TArray<const FRayTracingGeometry*> GeometriesToBuild;
+
+#ifdef DO_CHECK
+	// Keep track of all used RT Geometries which are used to validate the vertex buffer data (see FRayTracingGeometry::DynamicGeometrySharedBufferGenerationID)
+	TSet<const FRayTracingGeometry*> GeometriesToValidate;
+#endif
+
+private:
+	// RHI object that abstracts mesh instnaces in this scene
+	FRayTracingSceneRHIRef RayTracingSceneRHI;
+
+	// Task that asynchronously creates RayTracingSceneRHI
+	mutable FGraphEventRef CreateRayTracingSceneTask;
+
+	// Persistently allocated buffer that holds the built TLAS
+	FBufferRHIRef RayTracingSceneBuffer;
+
+	// View for the TLAS buffer that should be used in ray tracing shaders
+	FShaderResourceViewRHIRef RayTracingSceneSRV;
+};
+
+#endif // RHI_RAYTRACING
diff --git a/Engine/Source/Runtime/Renderer/Private/RendererScene.cpp b/Engine/Source/Runtime/Renderer/Private/RendererScene.cpp
--- a/Engine/Source/Runtime/Renderer/Private/RendererScene.cpp
+++ b/Engine/Source/Runtime/Renderer/Private/RendererScene.cpp
@@ -62,6 +62,7 @@
 #include "VT/RuntimeVirtualTextureSceneProxy.h"
 #if RHI_RAYTRACING
 #include "RayTracingDynamicGeometryCollection.h"
+#include "RayTracing/RayTracingScene.h"
 #endif
 #include "RHIGPUReadback.h"
 
@@ -1146,14 +1147,13 @@ FScene::~FScene()
 	}
 
 #if RHI_RAYTRACING
-	if (RayTracingDynamicGeometryCollection)
-	{
-		delete RayTracingDynamicGeometryCollection;
-		RayTracingDynamicGeometryCollection = nullptr;
-	}
+
+	delete RayTracingDynamicGeometryCollection;
+	RayTracingDynamicGeometryCollection = nullptr;
 
 	BeginReleaseResource(&HaltonPrimesResource);
-#endif
+
+#endif // RHI_RAYTRACING
 
 	checkf(RemovedPrimitiveSceneInfos.Num() == 0, TEXT("Leaking %d FPrimitiveSceneInfo instances."), RemovedPrimitiveSceneInfos.Num()); // Ensure UpdateAllPrimitiveSceneInfos() is called before destruction.
 }
diff --git a/Engine/Source/Runtime/Renderer/Private/ScenePrivate.h b/Engine/Source/Runtime/Renderer/Private/ScenePrivate.h
--- a/Engine/Source/Runtime/Renderer/Private/ScenePrivate.h
+++ b/Engine/Source/Runtime/Renderer/Private/ScenePrivate.h
@@ -51,6 +51,7 @@
 #include "RayTracing/RaytracingOptions.h"
 #if RHI_RAYTRACING
 #include "RayTracing/RayTracingIESLightProfiles.h"
+#include "RayTracing/RayTracingScene.h"
 #include "Halton.h"
 #endif
 #include "Nanite/NaniteRender.h"
@@ -2724,6 +2725,11 @@ public:
 	
 	FGPUScene GPUScene;
 
+#if RHI_RAYTRACING
+	/** Persistently-allocated ray tracing scene data. */
+	FRayTracingScene RayTracingScene;
+#endif // RHI_RAYTRACING
+
 	/** Distance field object scene data. */
 	FDistanceFieldSceneData DistanceFieldSceneData;
 
diff --git a/Engine/Source/Runtime/Renderer/Private/SceneRendering.cpp b/Engine/Source/Runtime/Renderer/Private/SceneRendering.cpp
--- a/Engine/Source/Runtime/Renderer/Private/SceneRendering.cpp
+++ b/Engine/Source/Runtime/Renderer/Private/SceneRendering.cpp
@@ -66,6 +66,7 @@
 #if WITH_EDITOR
 #include "Rendering/StaticLightingSystemInterface.h"
 #endif
+#include "RayTracing/RayTracingScene.h"
 #include "FXSystem.h"
 #include "Lumen/Lumen.h"
 
@@ -1007,23 +1008,25 @@ FViewInfo::~FViewInfo()
 #if RHI_RAYTRACING
 bool FViewInfo::HasRayTracingScene() const
 {
-	return CreateRayTracingSceneTask.IsValid() || RayTracingSceneRHI.IsValid();
+	check(Family);
+	FScene* Scene = Family->Scene ? Family->Scene->GetRenderScene() : nullptr;
+	if (Scene)
+	{
+		return Scene->RayTracingScene.IsCreated();
+	}
+	return false;
 }
 FRHIRayTracingScene* FViewInfo::GetRayTracingScene() const
 {
-	if (CreateRayTracingSceneTask.IsValid())
+	check(Family);
+	if (Family->Scene)
 	{
-		// #dxr_todo: This is temporarily here, until dependency on scene during RDG setup is removed.
-		// Then the wait would only need to happen during RDG execution.
-
-		TRACE_CPUPROFILER_EVENT_SCOPE(WaitForCreateRayTracingScene);
-		FTaskGraphInterface::Get().WaitUntilTaskCompletes(CreateRayTracingSceneTask, ENamedThreads::GetRenderThread_Local());
+		if (FScene* Scene = Family->Scene->GetRenderScene())
+		{
+			return Scene->RayTracingScene.GetRHIRayTracingScene();
+		}
 	}
-
-	checkf(!CreateRayTracingSceneTask.IsValid() || CreateRayTracingSceneTask->IsComplete(),
-		TEXT("Ray tracing scene creation task is expected to be waited upon and ready when we get here."));
-
-	return RayTracingSceneRHI.GetReference();
+	return nullptr;
 }
 FRHIRayTracingScene* FViewInfo::GetRayTracingSceneChecked() const
 {
@@ -1033,8 +1036,15 @@ FRHIRayTracingScene* FViewInfo::GetRayTracingSceneChecked() const
 }
 FRHIShaderResourceView* FViewInfo::GetRayTracingSceneViewChecked() const
 {
-	FRHIShaderResourceView* Result = RayTracingSceneSRV.GetReference();
-
+	FRHIShaderResourceView* Result = nullptr;
+	check(Family);
+	if (Family->Scene)
+	{
+		if (FScene* Scene = Family->Scene->GetRenderScene())
+		{
+			Result = Scene->RayTracingScene.GetShaderResourceViewChecked();
+		}
+	}
 	checkf(Result, TEXT("Ray tracing scene SRV is expected to be created at this point."));
 	return Result;
 }
diff --git a/Engine/Source/Runtime/Renderer/Private/SceneRendering.h b/Engine/Source/Runtime/Renderer/Private/SceneRendering.h
--- a/Engine/Source/Runtime/Renderer/Private/SceneRendering.h
+++ b/Engine/Source/Runtime/Renderer/Private/SceneRendering.h
@@ -1215,7 +1215,7 @@ public:
 
 	TStaticArray<FParallelMeshDrawCommandPass, EMeshPass::Num> ParallelMeshDrawCommandPasses;
 	
-#if RHI_RAYTRACING
+#if RHI_RAYTRACING // #yuriy_todo: Move to FRayTracingScene, as this is only valid for View[0]
 	TUniquePtr<FRayTracingMeshResourceCollector> RayTracingMeshResourceCollector;
 	FRayTracingMeshCommandOneFrameArray VisibleRayTracingMeshCommands;
 	FDynamicRayTracingMeshCommandStorage DynamicRayTracingMeshCommandStorage;
@@ -1400,26 +1400,10 @@ public:
 	FLumenTranslucencyGIVolume LumenTranslucencyGIVolume;
 
 #if RHI_RAYTRACING
-	TArray<FRayTracingGeometryInstance, SceneRenderingAllocator> RayTracingGeometryInstances;
-
-	// Geometries which still have a pending build request but are used this frame and require a force build
-	TArray<const FRayTracingGeometry*> ForceBuildRayTracingGeometries;
-
-#ifdef DO_CHECK
-	// Keep track of all used RT Geometries which are used to validate the vertex buffer data (see FRayTracingGeometry::DynamicGeometrySharedBufferGenerationID)
-	TSet<const FRayTracingGeometry*> RayTracingGeometriesForValidation;
-#endif
-
-	// Ray tracing scene specific to this view
-	FRayTracingSceneRHIRef RayTracingSceneRHI;
-	FShaderResourceViewRHIRef RayTracingSceneSRV; // #yuriy_todo: move this to RDG transient buffer system
-	FBufferRHIRef RayTracingSceneBuffer; // #yuriy_todo: move this to RDG transient buffer system
-	FGraphEventRef CreateRayTracingSceneTask;
-
 	RENDERER_API bool HasRayTracingScene() const;
-	RENDERER_API FRHIRayTracingScene* GetRayTracingScene() const;
-	RENDERER_API FRHIRayTracingScene* GetRayTracingSceneChecked() const;
-	RENDERER_API FRHIShaderResourceView* GetRayTracingSceneViewChecked() const;
+	RENDERER_API FRHIRayTracingScene* GetRayTracingScene() const; // Soft-deprecated method, use FScene.RayTracingScene instead.
+	RENDERER_API FRHIRayTracingScene* GetRayTracingSceneChecked() const; // Soft-deprecated method, use FScene.RayTracingScene instead.
+	RENDERER_API FRHIShaderResourceView* GetRayTracingSceneViewChecked() const; // Soft-deprecated method, use FScene.RayTracingScene instead.
 
 	// Primary pipeline state object to be used with the ray tracing scene for this view.
 	// Material shaders are only available when using this pipeline.
