Workaround for zero size BulkData on Distance Field streaming requests

diff --git a/Engine/Source/Runtime/Renderer/Private/DistanceFieldStreaming.cpp b/Engine/Source/Runtime/Renderer/Private/DistanceFieldStreaming.cpp
--- a/Engine/Source/Runtime/Renderer/Private/DistanceFieldStreaming.cpp
+++ b/Engine/Source/Runtime/Renderer/Private/DistanceFieldStreaming.cpp
@@ -545,18 +545,22 @@ void FDistanceFieldSceneData::ProcessStreamingRequestsFromGPU(
 							const int32 MipIndexToAdd = AssetState.BuiltData->Mips.Num() - ReversedMipIndexToAdd - 1;
 							const FSparseDistanceFieldMip& MipBuiltData = AssetState.BuiltData->Mips[MipIndexToAdd];
 
-							FDistanceFieldReadRequest ReadRequest;
-							ReadRequest.AssetSetId = AssetSetId;
-							ReadRequest.BuiltDataId = AssetState.BuiltData->GetId();
-							ReadRequest.ReversedMipIndex = ReversedMipIndexToAdd;
-							ReadRequest.NumDistanceFieldBricks = MipBuiltData.NumDistanceFieldBricks;
-							ReadRequest.BulkData = &AssetState.BuiltData->StreamableMips;
-							ReadRequest.BulkOffset = MipBuiltData.BulkOffset;
-							ReadRequest.BulkSize = MipBuiltData.BulkSize;
-							check(ReadRequest.BulkSize > 0);
-							NewReadRequests.Add(MoveTemp(ReadRequest));
-
-							NumAllocatedDistanceFieldBricks += MipBuiltData.NumDistanceFieldBricks;
+							//@todo - this condition shouldn't be possible as the built data always has non-zero size, needs more investigation
+							if (MipBuiltData.BulkSize > 0)
+							{
+								FDistanceFieldReadRequest ReadRequest;
+								ReadRequest.AssetSetId = AssetSetId;
+								ReadRequest.BuiltDataId = AssetState.BuiltData->GetId();
+								ReadRequest.ReversedMipIndex = ReversedMipIndexToAdd;
+								ReadRequest.NumDistanceFieldBricks = MipBuiltData.NumDistanceFieldBricks;
+								ReadRequest.BulkData = &AssetState.BuiltData->StreamableMips;
+								ReadRequest.BulkOffset = MipBuiltData.BulkOffset;
+								ReadRequest.BulkSize = MipBuiltData.BulkSize;
+								check(ReadRequest.BulkSize > 0);
+								NewReadRequests.Add(MoveTemp(ReadRequest));
+
+								NumAllocatedDistanceFieldBricks += MipBuiltData.NumDistanceFieldBricks;
+							}
 						}
 					}
 				}
