Add input source for virtual shadow page managment pass (Gbuffer/HairStrands).

diff --git a/Engine/Source/Runtime/Renderer/Private/VirtualShadowMaps/VirtualShadowMapProjection.cpp b/Engine/Source/Runtime/Renderer/Private/VirtualShadowMaps/VirtualShadowMapProjection.cpp
--- a/Engine/Source/Runtime/Renderer/Private/VirtualShadowMaps/VirtualShadowMapProjection.cpp
+++ b/Engine/Source/Runtime/Renderer/Private/VirtualShadowMaps/VirtualShadowMapProjection.cpp
@@ -333,7 +333,8 @@ FRDGTextureRef RenderVirtualShadowMapProjectionOnePass(
 	FRDGBuilder& GraphBuilder,
 	const FMinimalSceneTextures& SceneTextures,
 	const FViewInfo& View,
-	FVirtualShadowMapArray& VirtualShadowMapArray )
+	FVirtualShadowMapArray& VirtualShadowMapArray,
+	EVirtualShadowMapProjectionInputType InputType)
 {
 	FIntRect ProjectionRect = View.ViewRect;
 
@@ -343,7 +344,7 @@ FRDGTextureRef RenderVirtualShadowMapProjectionOnePass(
 		FClearValueBinding::None,
 		TexCreate_ShaderResource | TexCreate_UAV );
 
-	FRDGTextureRef ShadowMaskBits = GraphBuilder.CreateTexture( ShadowMaskDesc, TEXT("ShadowMaskBits") );
+	FRDGTextureRef ShadowMaskBits = GraphBuilder.CreateTexture( ShadowMaskDesc, InputType == EVirtualShadowMapProjectionInputType::HairStrands ? TEXT("ShadowMaskBits(HairStrands)") : TEXT("ShadowMaskBits(Gbuffer)"));
 	
 	RenderVirtualShadowMapProjectionCommon(
 		GraphBuilder,
@@ -351,7 +352,7 @@ FRDGTextureRef RenderVirtualShadowMapProjectionOnePass(
 		View,
 		VirtualShadowMapArray,
 		ProjectionRect,
-		EVirtualShadowMapProjectionInputType::GBuffer, // TODO: Hair
+		InputType,
 		ShadowMaskBits);
 
 	return ShadowMaskBits;
