Add cluster shading support for hair strands.

diff --git a/Engine/Shaders/Private/ClusteredDeferredShadingPixelShader.usf b/Engine/Shaders/Private/ClusteredDeferredShadingPixelShader.usf
--- a/Engine/Shaders/Private/ClusteredDeferredShadingPixelShader.usf
+++ b/Engine/Shaders/Private/ClusteredDeferredShadingPixelShader.usf
@@ -15,6 +15,7 @@
 // If the textures are atlased or something like that in the future it should be ok to turn on.
 #define USE_SOURCE_TEXTURE 0
 
+#include "HairStrands/HairStrandsVisibilityCommonStruct.ush" 
 #include "Common.ush"
 #include "DeferredShadingCommon.ush"  
 #include "BRDF.ush"
@@ -22,6 +23,12 @@
 #include "LightGridCommon.ush"
 #include "DeferredLightingCommon.ush"
 
+#if USE_HAIR_LIGHTING == 1
+	#include "HairStrands/HairStrandsCommon.ush"
+	#include "HairStrands/HairStrandsVisibilityCommon.ush"
+	#include "HairStrands/HairStrandsDeepTransmittanceCommon.ush"
+	#include "HairStrands/HairStrandsDeepTransmittanceDualScattering.ush" 
+#endif // USE_HAIR_LIGHTING
 
 // Causes the shader to run one pass over the light list per shading model, which cuts down peak register pressure.
 // Moves some overhead out of the inner loop (testing for shading model) which ought to be a good thing also.
@@ -32,11 +39,16 @@
 
 Texture2D< uint > ShadowMaskBits;
 
+#if USE_HAIR_LIGHTING
+uint				HairTransmittanceBufferMaxCount;
+Buffer<uint>		HairTransmittanceBuffer;
+#endif
+
 /**
  * Adds local lighting using the light grid, does not apply directional lights, as they are done elsewhere.
  * Does not support dynamic shadows, as these require the per-light shadow mask.
  */
-float4 GetLightGridLocalLighting(const FCulledLightsGridData InLightGridData, float3 WorldPosition, float3 CameraVector, FScreenSpaceData ScreenSpaceData, float2 ScreenUV, uint EyeIndex, float Dither, uint FirstNonSimpleLightIndex)
+float4 GetLightGridLocalLighting(const FCulledLightsGridData InLightGridData, float3 WorldPosition, float3 CameraVector, FScreenSpaceData ScreenSpaceData, float2 ScreenUV, uint EyeIndex, float Dither, uint FirstNonSimpleLightIndex, uint SampleIndex=0, uint TotalSampleCount=0, uint2 PixelCoord=0)
 {
 	float4 DirectLighting = 0;
 
@@ -114,10 +126,51 @@ float4 GetLightGridLocalLighting(const FCulledLightsGridData InLightGridData, fl
 		LightData.ShadowedBits = 1;
 		// Be careful with the shift here as the shift amount is mod 32, so >>32 is a noop rather than zero
 		// We treat lights past the 10th as unshadowed
-		if (LocalLightListIndex < 10)
+		if (LocalLightListIndex < GetShadowMaskMaxLightCount())
+		{
+			LightAttenuation = UnpackShadowMask(ShadowMask, LocalLightListIndex);
+		}
+
+		LightData.HairTransmittance = InitHairTransmittanceData();
+
+		#if USE_HAIR_LIGHTING
+		LightAttenuation = 1.0f; // TODO: use shadow mask bits (e.g., subpixel shadow mask)
+		if (LightAttenuation.x > 0)
 		{
-			LightAttenuation = ( ( ShadowMask >> (LocalLightListIndex * 3) ) & 7 ) / 7.0;
+			float3 L = float3(0, 0, 1);
+			if (LightData.bRadialLight)
+				L = normalize(LightData.Position - WorldPosition);
+			else
+				L = LightData.Direction;
+
+			const float3 V = normalize(-CameraVector);
+
+			// Fetch precomputed transmittance
+			FHairTransmittanceMask TransmittanceMask = InitHairTransmittanceMask();
+			if (LocalLightListIndex < GetShadowMaskMaxLightCount())
+			{
+				const uint TransmittanceSampleIndex = SampleIndex + LocalLightListIndex * TotalSampleCount;
+				const uint TransmittanceSampleMaxCount = TotalSampleCount * GetShadowMaskMaxLightCount();
+				if (TransmittanceSampleIndex < TransmittanceSampleMaxCount)
+				{
+					TransmittanceMask = UnpackTransmittanceMask(HairTransmittanceBuffer[TransmittanceSampleIndex]);		
+				}
+			}
+
+			// Apply dual scattering
+			LightData.HairTransmittance = GetHairTransmittance(
+				V,
+				L,
+				ScreenSpaceData.GBuffer,
+				TransmittanceMask,
+				View.HairScatteringLUTTexture,
+				View.HairScatteringLUTSampler,
+				View.HairComponents);
+
+			// This shouldn't be needed any more as the virtual shadow map should offer enough spatial precision
+			LightAttenuation = min(LightAttenuation, LightData.HairTransmittance.OpaqueVisibility.xxxx);
 		}
+		#endif
 
 #if 1
 		FRectTexture RectTexture = InitRectTexture(ForwardLightData.DummyRectLightSourceTexture);
@@ -165,8 +218,7 @@ float4 GetLightGridLocalLighting(const FCulledLightsGridData InLightGridData, fl
 
 #endif // USE_PASS_PER_SHADING_MODEL
 
-
-
+#if !USE_HAIR_LIGHTING
 void ClusteredShadingPixelShader(
 	in noperspective float4 UVAndScreenPos : TEXCOORD0,
 	in float4 SvPosition : SV_Position,
@@ -292,4 +344,81 @@ void ClusteredShadingPixelShader(
 #endif
 	OutColor = CompositedLighting;
 }
+#else // USE_HAIR_LIGHTING
+void ClusteredShadingPixelShader(
+	in float4 SvPosition: SV_Position,
+	nointerpolation in uint TotalSampleCount : DISPATCH_NODECOUNT,
+	nointerpolation in uint2 Resolution : DISPATCH_RESOLUTION,
+	out float4 OutColor : SV_Target0)
+{
+	OutColor = 0;
+
+	const uint EyeIndex = 0;
+
+	const uint2 InCoord = uint2(SvPosition.xy);
+	const uint SampleIndex = InCoord.x + InCoord.y * Resolution.x;
+	if (SampleIndex >= TotalSampleCount)
+	{
+		return;
+	}
+
+	const uint2 PixelCoord = HairStrands.HairSampleCoords[SampleIndex];
+	const float2 ScreenUV  = (PixelCoord + float2(0.5f, 0.5f)) / float2(View.BufferSizeAndInvSize.xy);
+	const float2 ScreenPosition = (ScreenUV - View.ScreenPositionScaleBias.wz) / View.ScreenPositionScaleBias.xy;
+
+	const FPackedHairSample PackedSample = HairStrands.HairSampleData[SampleIndex];
+	const FHairSample Sample = UnpackHairSample(PackedSample);
+
+	const float SceneDepth = ConvertFromDeviceZ(Sample.Depth);
+	
+	const float2 LocalPosition = PixelCoord - View.ViewRectMin.xy;
+	FScreenSpaceData ScreenSpaceData = (FScreenSpaceData)0;
+	{
+		// HAIR_TODO: make a HairSampleToGbuffer helper function
+		ScreenSpaceData.AmbientOcclusion = 1;
+		ScreenSpaceData.GBuffer.WorldNormal = Sample.Tangent;
+		ScreenSpaceData.GBuffer.BaseColor = Sample.BaseColor;
+		ScreenSpaceData.GBuffer.Roughness = Sample.Roughness;
+		ScreenSpaceData.GBuffer.ShadingModelID = SHADINGMODELID_HAIR;
+		ScreenSpaceData.GBuffer.DiffuseColor = 0;
+		ScreenSpaceData.GBuffer.SpecularColor = 0;
+		ScreenSpaceData.GBuffer.Specular = Sample.Specular;
+		ScreenSpaceData.GBuffer.Metallic = 0;
+		ScreenSpaceData.GBuffer.Depth = SceneDepth;
+		ScreenSpaceData.GBuffer.GBufferAO = 1;
+		ScreenSpaceData.GBuffer.CustomData = float4(HairStrands.HairDualScatteringRoughnessOverride, 0, Sample.Backlit, 0);
+		ScreenSpaceData.GBuffer.IndirectIrradiance = 0;
+		ScreenSpaceData.GBuffer.PrecomputedShadowFactors = 1;
+		ScreenSpaceData.GBuffer.PerObjectGBufferData = 0;
+	}
+
+	// TODO: This is how it was in TiledDeferred
+	//const float Dither = 0.5f;
+	// BUT: this is how it is in trad. deferred, assumed better:
+	const float Dither = InterleavedGradientNoise(SvPosition.xy, View.StateFrameIndexMod8);
+
+	const uint GridIndex = ComputeLightGridCellIndex(uint2(LocalPosition.x, LocalPosition.y), SceneDepth);
+	const FCulledLightsGridData CulledLightGridData = GetCulledLightsGrid(GridIndex, EyeIndex);
+	// NOTE: this early out helps with the case where there are no lights in the grid cell.
+	if (CulledLightGridData.NumLocalLights == 0)
+	{		
+		return;
+	}
+
+	const float3 WorldPosition = mul(float4(ScreenPosition * SceneDepth, SceneDepth, 1), View.ScreenToWorld).xyz;
+	const float3 CameraVector = normalize(WorldPosition - View.WorldCameraOrigin);
+
+	float4 CompositedLighting = 0;
+	uint FirstNonSimpleLightIndex = 0;
+	CompositedLighting += GetLightGridLocalLighting(CulledLightGridData, WorldPosition, CameraVector, ScreenSpaceData, ScreenUV, 0, Dither, FirstNonSimpleLightIndex, SampleIndex, TotalSampleCount, PixelCoord);
+#if USE_PREEXPOSURE && !VISUALIZE_LIGHT_CULLING
+	CompositedLighting *= View.PreExposure;
+#endif
+
+	// Weight hair sample with its coverage
+	const float LocalCoverage = From8bitCoverage(Sample.Coverage8bit);
+	OutColor.rgb = CompositedLighting.xyz * LocalCoverage;
+	OutColor.a = LocalCoverage;
+}
+#endif // USE_HAIR_LIGHTING
 
diff --git a/Engine/Shaders/Private/ClusteredDeferredShadingVertexShader.usf b/Engine/Shaders/Private/ClusteredDeferredShadingVertexShader.usf
new file mode 100644
--- /dev/null
+++ b/Engine/Shaders/Private/ClusteredDeferredShadingVertexShader.usf
@@ -0,0 +1,32 @@
+// Copyright Epic Games, Inc. All Rights Reserved.
+
+/*=============================================================================
+	ClusteredDeferredShadingVertexShader.usf: 
+=============================================================================*/
+
+#include "HairStrands/HairStrandsVisibilityCommonStruct.ush" 
+#include "Common.ush"
+
+struct FHairStrandsSamplesDesc
+{
+	nointerpolation uint NodeCount : DISPATCH_NODECOUNT;
+	nointerpolation uint2 Resolution : DISPATCH_RESOLUTION;
+};
+
+// This vertex shader is only used with hair strands lighting
+void ClusteredShadingVertexShader(
+	in uint InVertexId : SV_VertexID,
+	out float4 OutPosition : SV_POSITION,
+	nointerpolation out uint OutNodeCount : DISPATCH_NODECOUNT,
+	nointerpolation out uint2 OutResolution : DISPATCH_RESOLUTION)
+{
+	OutNodeCount = HairStrands.HairSampleCount.Load(uint3(0,0,0));
+	OutResolution.x = ceil(sqrt(OutNodeCount));
+	OutResolution.y = OutResolution.x;
+
+	const float2 ClipCoord = ((OutResolution + 0.5f) / float2(HairStrands.HairSampleViewportResolution)) * 2;
+
+	const float2 UV = float2(InVertexId & 1, InVertexId >> 1);
+	const float2 Pos = float2(-UV.x, UV.y) * 4 + float2(-1, +1) + float2(ClipCoord.x, -ClipCoord.y);
+	OutPosition = float4(Pos, 0.5f, 1);
+}
diff --git a/Engine/Source/Runtime/Renderer/Private/ClusteredDeferredShadingPass.cpp b/Engine/Source/Runtime/Renderer/Private/ClusteredDeferredShadingPass.cpp
--- a/Engine/Source/Runtime/Renderer/Private/ClusteredDeferredShadingPass.cpp
+++ b/Engine/Source/Runtime/Renderer/Private/ClusteredDeferredShadingPass.cpp
@@ -21,6 +21,8 @@
 #include "SceneRenderTargetParameters.h"
 #include "DeferredShadingRenderer.h"
 #include "ScenePrivate.h"
+#include "ShaderPrintParameters.h"
+#include "ShaderPrint.h"
 
 #include "SceneFilterRendering.h"
 #include "PostProcessing.h"
@@ -55,6 +57,35 @@ bool FDeferredShadingSceneRenderer::AreLightsInLightGrid() const
 }
 
 
+/**
+ * Clustered deferred shading shader. Use a custom vertex shader for hair strands lighting, to covered all sample in sample space
+ */
+class FClusteredShadingVS : public FGlobalShader
+{
+	DECLARE_GLOBAL_SHADER(FClusteredShadingVS);
+	SHADER_USE_PARAMETER_STRUCT(FClusteredShadingVS, FGlobalShader)
+
+	BEGIN_SHADER_PARAMETER_STRUCT(FParameters, )
+		SHADER_PARAMETER_STRUCT_REF(FViewUniformShaderParameters, View)
+		SHADER_PARAMETER_RDG_UNIFORM_BUFFER(FHairStrandsViewUniformParameters, HairStrands)
+		SHADER_PARAMETER_RDG_UNIFORM_BUFFER(FSceneTextureUniformParameters, SceneTextures)
+	END_SHADER_PARAMETER_STRUCT()
+
+	static bool ShouldCompilePermutation(const FGlobalShaderPermutationParameters& Parameters)
+	{
+		// OLATODO: what level do we actually need for this?
+		return IsFeatureLevelSupported(Parameters.Platform, ERHIFeatureLevel::SM5);
+	}
+
+	static void ModifyCompilationEnvironment(const FGlobalShaderPermutationParameters& Parameters, FShaderCompilerEnvironment& OutEnvironment)
+	{
+		FGlobalShader::ModifyCompilationEnvironment(Parameters, OutEnvironment);
+		FForwardLightingParameters::ModifyCompilationEnvironment(Parameters.Platform, OutEnvironment);
+	}
+};
+
+IMPLEMENT_GLOBAL_SHADER(FClusteredShadingVS, "/Engine/Private/ClusteredDeferredShadingVertexShader.usf", "ClusteredShadingVertexShader", SF_Vertex);
+
 /**
  * Clustered deferred shading shader, used in a full-screen pass to apply all lights in the light grid.
  */
@@ -64,13 +95,18 @@ class FClusteredShadingPS : public FGlobalShader
 	SHADER_USE_PARAMETER_STRUCT(FClusteredShadingPS, FGlobalShader)
 
 	class FVisualizeLightCullingDim : SHADER_PERMUTATION_BOOL("VISUALIZE_LIGHT_CULLING");
-	using FPermutationDomain = TShaderPermutationDomain<FVisualizeLightCullingDim>;
+	class FHairStrandsLighting : SHADER_PERMUTATION_BOOL("USE_HAIR_LIGHTING");
+	using FPermutationDomain = TShaderPermutationDomain<FVisualizeLightCullingDim, FHairStrandsLighting>;
 
 	BEGIN_SHADER_PARAMETER_STRUCT(FParameters, )
 		SHADER_PARAMETER_STRUCT_REF(FForwardLightData, Forward)
 		SHADER_PARAMETER_STRUCT_REF(FViewUniformShaderParameters, View)
+		SHADER_PARAMETER_STRUCT_INCLUDE(ShaderDrawDebug::FShaderDrawDebugParameters, ShaderDrawParameters)
+		SHADER_PARAMETER_STRUCT_INCLUDE(ShaderPrint::FShaderParameters, ShaderPrintUniformBuffer)
+		SHADER_PARAMETER_RDG_UNIFORM_BUFFER(FHairStrandsViewUniformParameters, HairStrands)
 		SHADER_PARAMETER_RDG_UNIFORM_BUFFER(FSceneTextureUniformParameters, SceneTextures)
 		SHADER_PARAMETER_RDG_TEXTURE(Texture2D, ShadowMaskBits)
+		SHADER_PARAMETER_RDG_BUFFER_SRV(Buffer<uint>, HairTransmittanceBuffer)
 
 		SHADER_PARAMETER_TEXTURE(Texture2D, LTCMatTexture)
 		SHADER_PARAMETER_SAMPLER(SamplerState, LTCMatSampler)
@@ -99,79 +135,153 @@ class FClusteredShadingPS : public FGlobalShader
 
 IMPLEMENT_GLOBAL_SHADER(FClusteredShadingPS, "/Engine/Private/ClusteredDeferredShadingPixelShader.usf", "ClusteredShadingPixelShader", SF_Pixel);
 
-void FDeferredShadingSceneRenderer::AddClusteredDeferredShadingPass(
+enum class EClusterPassInputType : uint8
+{
+	GBuffer,
+	HairStrands
+};
+
+static void InternalAddClusteredDeferredShadingPass(
 	FRDGBuilder& GraphBuilder,
+	FViewInfo& View,
 	const FMinimalSceneTextures& SceneTextures,
 	const FSortedLightSetSceneInfo &SortedLightsSet,
-	FRDGTextureRef ShadowMaskBits)
+	EClusterPassInputType InputType,
+	FRDGTextureRef ShadowMaskBits,
+	FRDGBufferSRVRef HairTransmittanceBuffer)
 {
-	check(GUseClusteredDeferredShading);
+	check(SortedLightsSet.ClusteredSupportedEnd > 0);
+	const FIntPoint SceneTextureExtent = SceneTextures.Config.Extent;
+	const bool bHairStrands = InputType == EClusterPassInputType::HairStrands;
+	
+	FClusteredShadingPS::FParameters *PassParameters = GraphBuilder.AllocParameters<FClusteredShadingPS::FParameters>();
+	PassParameters->View = View.ViewUniformBuffer;
+	PassParameters->HairStrands = HairStrands::BindHairStrandsViewUniformParameters(View);
+	PassParameters->Forward = View.ForwardLightingResources->ForwardLightDataUniformBuffer;
+	PassParameters->SceneTextures = SceneTextures.UniformBuffer;
+	PassParameters->ShadowMaskBits = ShadowMaskBits;
+	PassParameters->HairTransmittanceBuffer = HairTransmittanceBuffer;
+
+	PassParameters->LTCMatTexture = GSystemTextures.LTCMat->GetShaderResourceRHI();
+	PassParameters->LTCMatSampler = TStaticSamplerState<SF_Bilinear, AM_Clamp, AM_Clamp, AM_Clamp>::GetRHI();
+	PassParameters->LTCAmpTexture = GSystemTextures.LTCAmp->GetShaderResourceRHI();
+	PassParameters->LTCAmpSampler = TStaticSamplerState<SF_Bilinear, AM_Clamp, AM_Clamp, AM_Clamp>::GetRHI();
+	PassParameters->SSProfilesTexture = GetSubsufaceProfileTexture_RT(GraphBuilder.RHICmdList)->GetShaderResourceRHI();
+	PassParameters->TransmissionProfilesLinearSampler = TStaticSamplerState<SF_Bilinear, AM_Clamp, AM_Clamp, AM_Clamp>::GetRHI();
+
+	if (ShaderDrawDebug::IsShaderDrawDebugEnabled(View))
+	{
+		ShaderDrawDebug::SetParameters(GraphBuilder, View.ShaderDrawData, PassParameters->ShaderDrawParameters);
+	}
 
-	const int32 NumLightsToRender = SortedLightsSet.ClusteredSupportedEnd;
+	if (ShaderPrint::IsSupported(View) && ShaderPrint::IsEnabled())
+	{
+		ShaderPrint::SetParameters(GraphBuilder, View, PassParameters->ShaderPrintUniformBuffer);
+	}
 
-	if (NumLightsToRender > 0)
+	PassParameters->RenderTargets[0] = FRenderTargetBinding(SceneTextures.Color.Target, ERenderTargetLoadAction::ELoad);
+	if (bHairStrands)
 	{
-		RDG_GPU_STAT_SCOPE(GraphBuilder, ClusteredShading);
-		RDG_EVENT_SCOPE(GraphBuilder, "ClusteredShading");
+		PassParameters->RenderTargets[0] = FRenderTargetBinding(View.HairStrandsViewData.VisibilityData.SampleLightingBuffer, ERenderTargetLoadAction::ELoad);
+	}
 
-		const FIntPoint SceneTextureExtent = SceneTextures.Config.Extent;
+	GraphBuilder.AddPass(
+		RDG_EVENT_NAME("ClusteredDeferredShading(%s), #Lights: %d", bHairStrands ? TEXT("HairStrands") : TEXT("GBuffer"), SortedLightsSet.ClusteredSupportedEnd),
+		PassParameters,
+		ERDGPassFlags::Raster,
+		[PassParameters, &View, SceneTextureExtent, bHairStrands](FRHICommandListImmediate& InRHICmdList)
+	{
+		TShaderMapRef<FClusteredShadingVS> HairVertexShader(View.ShaderMap);
+		TShaderMapRef<FPostProcessVS> VertexShader(View.ShaderMap);
 
-		for (int32 ViewIndex = 0, Num = Views.Num(); ViewIndex < Num; ViewIndex++)
+		FClusteredShadingPS::FPermutationDomain PermutationVector;
+		PermutationVector.Set<FClusteredShadingPS::FVisualizeLightCullingDim>(View.Family->EngineShowFlags.VisualizeLightCulling);
+		PermutationVector.Set<FClusteredShadingPS::FHairStrandsLighting>(bHairStrands);
+		TShaderMapRef<FClusteredShadingPS> PixelShader(View.ShaderMap, PermutationVector);
 		{
-			FViewInfo& View = Views[ViewIndex];
+			FGraphicsPipelineStateInitializer GraphicsPSOInit;
+			InRHICmdList.ApplyCachedRenderTargets(GraphicsPSOInit);
 
-			FClusteredShadingPS::FParameters *PassParameters = GraphBuilder.AllocParameters<FClusteredShadingPS::FParameters>();
+			// Additive blend to accumulate lighting contributions.
+			GraphicsPSOInit.BlendState = TStaticBlendState<CW_RGBA, BO_Add, BF_One, BF_One, BO_Add, BF_One, BF_One>::GetRHI();
 
-			PassParameters->View = View.ViewUniformBuffer;
-			PassParameters->Forward = View.ForwardLightingResources->ForwardLightDataUniformBuffer;
-			PassParameters->SceneTextures = SceneTextures.UniformBuffer;
-			PassParameters->ShadowMaskBits = ShadowMaskBits;
+			GraphicsPSOInit.RasterizerState = TStaticRasterizerState<FM_Solid, CM_None>::GetRHI();
+			GraphicsPSOInit.DepthStencilState = TStaticDepthStencilState<false, CF_Always>::GetRHI();
 
-			PassParameters->LTCMatTexture = GSystemTextures.LTCMat->GetShaderResourceRHI();
-			PassParameters->LTCMatSampler = TStaticSamplerState<SF_Bilinear, AM_Clamp, AM_Clamp, AM_Clamp>::GetRHI();
-			PassParameters->LTCAmpTexture = GSystemTextures.LTCAmp->GetShaderResourceRHI();
-			PassParameters->LTCAmpSampler = TStaticSamplerState<SF_Bilinear, AM_Clamp, AM_Clamp, AM_Clamp>::GetRHI();
-			PassParameters->SSProfilesTexture = GetSubsufaceProfileTexture_RT(GraphBuilder.RHICmdList)->GetShaderResourceRHI();
-			PassParameters->TransmissionProfilesLinearSampler = TStaticSamplerState<SF_Bilinear, AM_Clamp, AM_Clamp, AM_Clamp>::GetRHI();
+			GraphicsPSOInit.BoundShaderState.VertexDeclarationRHI = GFilterVertexDeclaration.VertexDeclarationRHI;
+			GraphicsPSOInit.BoundShaderState.VertexShaderRHI = bHairStrands ? HairVertexShader.GetVertexShader() : VertexShader.GetVertexShader();
+			GraphicsPSOInit.BoundShaderState.PixelShaderRHI = PixelShader.GetPixelShader();
+			GraphicsPSOInit.PrimitiveType = PT_TriangleList;
 
-			PassParameters->RenderTargets[0] = FRenderTargetBinding(SceneTextures.Color.Target, ERenderTargetLoadAction::ELoad);
+			SetGraphicsPipelineState(InRHICmdList, GraphicsPSOInit);
+		}
 
-			GraphBuilder.AddPass(
-				RDG_EVENT_NAME("ClusteredDeferredShading, #Lights: %d", NumLightsToRender),
-				PassParameters,
-				ERDGPassFlags::Raster,
-				[PassParameters, &View, SceneTextureExtent](FRHICommandListImmediate& InRHICmdList)
-			{
-				TShaderMapRef<FPostProcessVS> VertexShader(View.ShaderMap);
+		SetShaderParameters(InRHICmdList, PixelShader, PixelShader.GetPixelShader(), *PassParameters);
 
-				FClusteredShadingPS::FPermutationDomain PermutationVector;
-				PermutationVector.Set<FClusteredShadingPS::FVisualizeLightCullingDim>(View.Family->EngineShowFlags.VisualizeLightCulling);
-				TShaderMapRef<FClusteredShadingPS> PixelShader(View.ShaderMap, PermutationVector);
-				{
-					FGraphicsPipelineStateInitializer GraphicsPSOInit;
-					InRHICmdList.ApplyCachedRenderTargets(GraphicsPSOInit);
+		if (bHairStrands)
+		{
+			InRHICmdList.SetViewport(0, 0, 0.0f, View.HairStrandsViewData.VisibilityData.SampleLightingViewportResolution.X, View.HairStrandsViewData.VisibilityData.SampleLightingViewportResolution.Y, 1.0f);
+
+			FClusteredShadingVS::FParameters VertexParameters;
+			VertexParameters.View = PassParameters->View;
+			VertexParameters.HairStrands = PassParameters->HairStrands;
+			VertexParameters.SceneTextures = PassParameters->SceneTextures;
+			SetShaderParameters(InRHICmdList, HairVertexShader, HairVertexShader.GetVertexShader(), VertexParameters);
+			InRHICmdList.SetStreamSource(0, nullptr, 0);
+			InRHICmdList.DrawPrimitive(0, 1, 1);
+		}
+		else
+		{
+			InRHICmdList.SetViewport(View.ViewRect.Min.X, View.ViewRect.Min.Y, 0.0f, View.ViewRect.Max.X, View.ViewRect.Max.Y, 1.0f);
+			DrawRectangle(InRHICmdList, 0, 0, View.ViewRect.Width(), View.ViewRect.Height(),
+				View.ViewRect.Min.X, View.ViewRect.Min.Y, View.ViewRect.Width(), View.ViewRect.Height(),
+				FIntPoint(View.ViewRect.Width(), View.ViewRect.Height()), SceneTextureExtent, VertexShader);
+		}
+	});
+}
 
-					// Additive blend to accumulate lighting contributions.
-					GraphicsPSOInit.BlendState = TStaticBlendState<CW_RGBA, BO_Add, BF_One, BF_One, BO_Add, BF_One, BF_One>::GetRHI();
 
-					GraphicsPSOInit.RasterizerState = TStaticRasterizerState<FM_Solid, CM_None>::GetRHI();
-					GraphicsPSOInit.DepthStencilState = TStaticDepthStencilState<false, CF_Always>::GetRHI();
+void FDeferredShadingSceneRenderer::AddClusteredDeferredShadingPass(
+	FRDGBuilder& GraphBuilder,
+	const FMinimalSceneTextures& SceneTextures,
+	const FSortedLightSetSceneInfo &SortedLightsSet,
+	FRDGTextureRef ShadowMaskBits,
+	FRDGTextureRef HairStrandsShadowMaskBits)
+{
+	check(GUseClusteredDeferredShading);
 
-					GraphicsPSOInit.BoundShaderState.VertexDeclarationRHI = GFilterVertexDeclaration.VertexDeclarationRHI;
-					GraphicsPSOInit.BoundShaderState.VertexShaderRHI = VertexShader.GetVertexShader();
-					GraphicsPSOInit.BoundShaderState.PixelShaderRHI = PixelShader.GetPixelShader();
-					GraphicsPSOInit.PrimitiveType = PT_TriangleList;
+	const int32 NumLightsToRender = SortedLightsSet.ClusteredSupportedEnd;
 
-					SetGraphicsPipelineState(InRHICmdList, GraphicsPSOInit);
-				}
-				InRHICmdList.SetViewport(View.ViewRect.Min.X, View.ViewRect.Min.Y, 0.0f, View.ViewRect.Max.X, View.ViewRect.Max.Y, 1.0f);
+	if (NumLightsToRender > 0)
+	{
+		RDG_GPU_STAT_SCOPE(GraphBuilder, ClusteredShading);
+		RDG_EVENT_SCOPE(GraphBuilder, "ClusteredShading");
 
-				SetShaderParameters(InRHICmdList, PixelShader, PixelShader.GetPixelShader(), *PassParameters);
+		for (int32 ViewIndex = 0, Num = Views.Num(); ViewIndex < Num; ViewIndex++)
+		{
+			FViewInfo& View = Views[ViewIndex];
+
+			InternalAddClusteredDeferredShadingPass(
+				GraphBuilder,
+				View,
+				SceneTextures,
+				SortedLightsSet,
+				EClusterPassInputType::GBuffer,
+				ShadowMaskBits,
+				nullptr);
 
-				DrawRectangle(InRHICmdList, 0, 0, View.ViewRect.Width(), View.ViewRect.Height(),
-					View.ViewRect.Min.X, View.ViewRect.Min.Y, View.ViewRect.Width(), View.ViewRect.Height(),
-					FIntPoint(View.ViewRect.Width(), View.ViewRect.Height()), SceneTextureExtent, VertexShader);
-			});
+			if (HairStrands::HasViewHairStrandsData(View))
+			{
+				FHairStrandsTransmittanceMaskData TransmittanceMask = RenderHairStrandsOnePassTransmittanceMask(GraphBuilder, View, HairStrandsShadowMaskBits);
+				InternalAddClusteredDeferredShadingPass(
+					GraphBuilder,
+					View,
+					SceneTextures,
+					SortedLightsSet,
+					EClusterPassInputType::HairStrands,
+					HairStrandsShadowMaskBits,
+					GraphBuilder.CreateSRV(TransmittanceMask.TransmittanceMask, FHairStrandsTransmittanceMaskData::Format));
+			}
 		}
 	}
 }
diff --git a/Engine/Source/Runtime/Renderer/Private/DeferredShadingRenderer.h b/Engine/Source/Runtime/Renderer/Private/DeferredShadingRenderer.h
--- a/Engine/Source/Runtime/Renderer/Private/DeferredShadingRenderer.h
+++ b/Engine/Source/Runtime/Renderer/Private/DeferredShadingRenderer.h
@@ -570,7 +570,8 @@ private:
 		FRDGBuilder& GraphBuilder,
 		const FMinimalSceneTextures& SceneTextures,
 		const FSortedLightSetSceneInfo& SortedLightsSet,
-		FRDGTextureRef ShadowMaskBits);
+		FRDGTextureRef ShadowMaskBits,
+		FRDGTextureRef HairStrandsShadowMaskBits);
 
 	/** Renders the lights in SortedLights in the range [TiledDeferredLightsStart, TiledDeferredLightsEnd) using tiled deferred shading. */
 	void RenderTiledDeferredLighting(
diff --git a/Engine/Source/Runtime/Renderer/Private/LightRendering.cpp b/Engine/Source/Runtime/Renderer/Private/LightRendering.cpp
--- a/Engine/Source/Runtime/Renderer/Private/LightRendering.cpp
+++ b/Engine/Source/Runtime/Renderer/Private/LightRendering.cpp
@@ -1183,16 +1183,28 @@ void FDeferredShadingSceneRenderer::RenderLights(
 			if (ShouldUseClusteredDeferredShading() && AreLightsInLightGrid())
 			{
 				FRDGTextureRef ShadowMaskBits = nullptr;
+				FRDGTextureRef HairStrandsShadowMaskBits = nullptr;
 				if( VirtualShadowMapArray.IsAllocated() && CVarVirtualShadowOnePassProjection.GetValueOnRenderThread() )
 				{
 					// TODO: This needs to move into the view loop in clustered deferred shading pass
-					for (int32 ViewIndex = 0, Num = Views.Num(); ViewIndex < Num; ViewIndex++)
+					for (const FViewInfo& View : Views)
 					{
 						ShadowMaskBits = RenderVirtualShadowMapProjectionOnePass(
 							GraphBuilder,
 							SceneTextures,
-							Views[ ViewIndex ],
-							VirtualShadowMapArray );
+							View,
+							VirtualShadowMapArray,
+							EVirtualShadowMapProjectionInputType::GBuffer);
+
+						if (HairStrands::HasViewHairStrandsData(View))
+						{
+							HairStrandsShadowMaskBits = RenderVirtualShadowMapProjectionOnePass(
+							GraphBuilder,
+							SceneTextures,
+							View,
+							VirtualShadowMapArray,
+							EVirtualShadowMapProjectionInputType::HairStrands);
+						}
 					}
 				}
 				else
@@ -1206,7 +1218,7 @@ void FDeferredShadingSceneRenderer::RenderLights(
 				// Tell the trad. deferred that the simple lights are spoken for.
 				bRenderSimpleLightsStandardDeferred = false;
 
-				AddClusteredDeferredShadingPass(GraphBuilder, SceneTextures, SortedLightSet, ShadowMaskBits);
+				AddClusteredDeferredShadingPass(GraphBuilder, SceneTextures, SortedLightSet, ShadowMaskBits, HairStrandsShadowMaskBits);
 			}
 			else if (CanUseTiledDeferred())
 			{
diff --git a/Engine/Source/Runtime/Renderer/Private/VirtualShadowMaps/VirtualShadowMapProjection.h b/Engine/Source/Runtime/Renderer/Private/VirtualShadowMaps/VirtualShadowMapProjection.h
--- a/Engine/Source/Runtime/Renderer/Private/VirtualShadowMaps/VirtualShadowMapProjection.h
+++ b/Engine/Source/Runtime/Renderer/Private/VirtualShadowMaps/VirtualShadowMapProjection.h
@@ -41,7 +41,8 @@ FRDGTextureRef RenderVirtualShadowMapProjectionOnePass(
 	FRDGBuilder& GraphBuilder,
 	const FMinimalSceneTextures& SceneTextures,
 	const FViewInfo& View,
-	FVirtualShadowMapArray& VirtualShadowMapArray );
+	FVirtualShadowMapArray& VirtualShadowMapArray,
+	EVirtualShadowMapProjectionInputType InputType);
 
 void CompositeVirtualShadowMapMask(
 	FRDGBuilder& GraphBuilder,
