Manually transition DistanceFieldBrickVolumeTexture to final state, fixes RHI validation errors

diff --git a/Engine/Source/Runtime/Renderer/Private/DistanceFieldStreaming.cpp b/Engine/Source/Runtime/Renderer/Private/DistanceFieldStreaming.cpp
--- a/Engine/Source/Runtime/Renderer/Private/DistanceFieldStreaming.cpp
+++ b/Engine/Source/Runtime/Renderer/Private/DistanceFieldStreaming.cpp
@@ -701,7 +701,7 @@ void FDistanceFieldSceneData::ResizeBrickAtlasIfNeeded(FRDGBuilder& GraphBuilder
 		}
 
 		BrickTextureDimensionsInBricks = DesiredBrickTextureDimensionsInBricks;
-		ConvertToExternalTexture(GraphBuilder, DistanceFieldBrickVolumeTextureRDG, DistanceFieldBrickVolumeTexture);
+		ConvertToUntrackedExternalTexture(GraphBuilder, DistanceFieldBrickVolumeTextureRDG, DistanceFieldBrickVolumeTexture, ERHIAccess::SRVMask);
 	}
 }
 
@@ -956,6 +956,8 @@ void FDistanceFieldSceneData::UpdateDistanceFieldAtlas(
 					PassParameters,
 					FComputeShaderUtils::GetGroupCount(FIntVector(DistanceField::BrickSize, DistanceField::BrickSize, NumBrickUploadsThisPass * DistanceField::BrickSize), FScatterUploadDistanceFieldAtlasCS::GetGroupSize()));
 			}
+
+			ConvertToUntrackedExternalTexture(GraphBuilder, DistanceFieldBrickVolumeTextureRDG, DistanceFieldBrickVolumeTexture, ERHIAccess::SRVMask);
 		}
 	}
 
