Fix crash when running without hairstrands.

diff --git a/Engine/Source/Runtime/Renderer/Private/DeferredShadingRenderer.cpp b/Engine/Source/Runtime/Renderer/Private/DeferredShadingRenderer.cpp
--- a/Engine/Source/Runtime/Renderer/Private/DeferredShadingRenderer.cpp
+++ b/Engine/Source/Runtime/Renderer/Private/DeferredShadingRenderer.cpp
@@ -1934,13 +1934,19 @@ void FDeferredShadingSceneRenderer::Render(FRDGBuilder& GraphBuilder)
 		}
 		else
 		{
-			for (int32 ViewIndex = 0; ViewIndex < Views.Num(); ViewIndex++)
+			for (FViewInfo& View : Views)
 			{
-				FViewInfo& View = Views[ViewIndex];
 				View.HairStrandsViewData.UniformBuffer = HairStrands::CreateDefaultHairStrandsViewUniformBuffer(GraphBuilder, View);
 			}
 		}
 	}
+	else
+	{
+		for (FViewInfo& View : Views)
+		{
+			View.HairStrandsViewData.UniformBuffer = HairStrands::CreateDefaultHairStrandsViewUniformBuffer(GraphBuilder, View);
+		}
+	}
 
 	if (bNaniteEnabled)
 	{
