Change the default of r.BasePassOutputsVelocity to on, so that velocity vectors work out of the box with techniques like Nanite and Lumen that do not currenlty support a separate velocity pass.

diff --git a/Engine/Source/Runtime/Renderer/Private/VelocityRendering.cpp b/Engine/Source/Runtime/Renderer/Private/VelocityRendering.cpp
--- a/Engine/Source/Runtime/Renderer/Private/VelocityRendering.cpp
+++ b/Engine/Source/Runtime/Renderer/Private/VelocityRendering.cpp
@@ -25,7 +25,7 @@
 // Changing this causes a full shader recompile
 static TAutoConsoleVariable<int32> CVarBasePassOutputsVelocity(
 	TEXT("r.BasePassOutputsVelocity"),
-	0,
+	1,
 	TEXT("Enables rendering WPO velocities on the base pass.\n") \
 	TEXT(" 0: Renders in a separate pass/rendertarget, all movable static meshes + dynamic.\n") \
 	TEXT(" 1: Renders during the regular base pass adding an extra GBuffer, but allowing motion blur on materials with Time-based WPO."),
