Make sure MidGreyDummy texture is initialized on mobile

diff --git a/Engine/Source/Runtime/Renderer/Private/SystemTextures.cpp b/Engine/Source/Runtime/Renderer/Private/SystemTextures.cpp
--- a/Engine/Source/Runtime/Renderer/Private/SystemTextures.cpp
+++ b/Engine/Source/Runtime/Renderer/Private/SystemTextures.cpp
@@ -255,7 +255,7 @@ void FSystemTextures::InitializeCommonTextures(FRHICommandListImmediate& RHICmdL
 		StencilDummySRV = RHICreateShaderResourceView((FRHITexture2D*)StencilDummy->GetRenderTargetItem().ShaderResourceTexture.GetReference(), 0);
 	}
 
-	if (!GSupportsShaderFramebufferFetch && GPixelFormats[PF_FloatRGBA].Supported)
+	if (GPixelFormats[PF_FloatRGBA].Supported)
 	{
 		// PF_FloatRGBA to encode exactly the 0.5.
 		FPooledRenderTargetDesc Desc(FPooledRenderTargetDesc::Create2DDesc(FIntPoint(1, 1), PF_FloatRGBA, FClearValueBinding(FLinearColor(0.5f, 0.5f, 0.5f, 0.5f)), TexCreate_HideInVisualizeTexture, TexCreate_RenderTargetable | TexCreate_NoFastClear | TexCreate_ShaderResource, false));
