Remove legacy hair visibility using GBuffer.

diff --git a/Engine/Shaders/Private/HairStrands/HairStrandsVisibilityCommon.ush b/Engine/Shaders/Private/HairStrands/HairStrandsVisibilityCommon.ush
--- a/Engine/Shaders/Private/HairStrands/HairStrandsVisibilityCommon.ush
+++ b/Engine/Shaders/Private/HairStrands/HairStrandsVisibilityCommon.ush
@@ -3,11 +3,11 @@
 #pragma once
 
 // Render modes
-#define RENDER_MODE_MSAA 0
-#define RENDER_MODE_TRANSMITTANCE 1
-#define RENDER_MODE_PPLL 2
-#define RENDER_MODE_MSAA_VISIBILITY 3
-#define RENDER_MODE_TRANSMITTANCE_AND_HAIRCOUNT 4
+#define RENDER_MODE_TRANSMITTANCE 0
+#define RENDER_MODE_PPLL 1
+#define RENDER_MODE_MSAA_VISIBILITY 2
+#define RENDER_MODE_TRANSMITTANCE_AND_HAIRCOUNT 3
+#define RENDER_MODE_COMPUTE_RASTER 4
 
 #define INVALID_TILE_OFFSET ~0
 
@@ -454,67 +454,3 @@ float UnpackCoverageAndHairCount(uint InPacked)
 	return float(InPacked) / 1000.f;
 }
 
- ////////////////////////////////////////////////////////////////////////////////////////////////////////////////////////////////
- // Hair MSAA buffer 
-
-bool IsValidHairData(uint2 MRT0Data)
-{
-	const uint InvalidSample = GetInvalidHairPrimitiveId();// &0x0FFFFFFF; // PrimitiveId mask, see EncodeVisSample()
-	return MRT0Data.x != InvalidSample;
-}
-
-uint2 PackHairMRT0(uint PrimitiveId, uint MacroGroupId, uint Coverage8bit, float3 Tangent)
-{
-	uint2 Out;
-	Out.x =
-		((PrimitiveId & 0x0FFFFFFF)) |
-		((MacroGroupId & 0xF) << 28);
-
-	const float3 T = saturate(EncodeTangent(Tangent));
-	Out.y =
-		((0xFF & uint(T.x * 255))) |
-		((0xFF & uint(T.y * 255)) << 8) |
-		((0xFF & uint(T.z * 255)) << 16) |
-		(min(uint(0xFF), Coverage8bit) << 24);
-
-	return Out;
-}
-
-float4 PackHairMRT1(float3 BaseColor, float Roughness)
-{
-	return float4(sqrt(BaseColor), Roughness);
-}
-
-float4 PackHairMRT2(float Specular, uint LightChannelMask, float Backlit, bool bUseScatterSceneLighting)
-{
-	return float4(Specular, float(LightChannelMask) / 255.f, Backlit, bUseScatterSceneLighting? 1 : 0);
-}
-
-void UnpackHairMRT0(uint2 In, inout uint PrimitiveId, inout uint MacroGroupId, inout uint Coverage8bit, inout float3 Tangent)
-{
-	PrimitiveId  =  In.x & 0x0FFFFFFF;
-	MacroGroupId = (In.x >> 28) & 0xF;
-
-	Tangent = DecodeTangent(float3(
-		float((In.y) & 0xFF) / 255.f,
-		float((In.y >> 8) & 0xFF) / 255.f,
-		float((In.y >> 16) & 0xFF) / 255.f));
-	Coverage8bit = (In.y >> 24) & 0xFF;
-}
-
-void UnpackHairMRT1(float4 In, inout float3 BaseColor, inout float Roughness)
-{
-	// Cheap gamma encoding since this is encoded into 8bit target
-	BaseColor = In.xyz * In.xyz;
-	Roughness = In.w;
-}
-
-void UnpackHairMRT2(float4 In, inout float Specular, inout uint LightChannelMask, inout float Backlit, inout bool bUseScatterSceneLighting)
-{
-	Specular = In.x;
-	LightChannelMask = uint(In.y * 0xFF);
-	Backlit = In.z;
-	bUseScatterSceneLighting = In.w > 0;
-}
-
-
diff --git a/Engine/Shaders/Private/HairStrands/HairStrandsVisibilityCompaction.usf b/Engine/Shaders/Private/HairStrands/HairStrandsVisibilityCompaction.usf
--- a/Engine/Shaders/Private/HairStrands/HairStrandsVisibilityCompaction.usf
+++ b/Engine/Shaders/Private/HairStrands/HairStrandsVisibilityCompaction.usf
@@ -24,7 +24,7 @@ uint VelocityType;
 
 RWTexture2D<uint> OutCompactNodeCounter;
 RWTexture2D<uint> OutCompactNodeIndex;
-#if PERMUTATION_VISIBILITY
+#if PERMUTATION_PPLL == 0
 	RWStructuredBuffer<FNodeVis> OutCompactNodeData;
 #else
 	RWStructuredBuffer<FPackedHairSample> OutCompactNodeData;
@@ -271,7 +271,7 @@ void MainCS(uint GroupIndex : SV_GroupIndex, uint3 DispatchThreadId : SV_Dispatc
 	OutCategorizationTexture[PixelCoord] = EncodeCategorizationData(Categorization);
 }
 
-#elif PERMUTATION_VISIBILITY // PERMUTATION_PPLL or PERMUTATION_VISIBILITY or PERMUTATION_MSAA
+#else // PERMUTATION_PPLL or PERMUTATION_VISIBILITY
 
 ///////////////////////////////////////////////////////////////////////////////////////////////////////////////
 // MSAA visiblity buffer
@@ -508,311 +508,4 @@ void MainCS(uint GroupIndex : SV_GroupIndex, uint3 DispatchThreadId : SV_Dispatc
 	OutCategorizationTexture[PixelCoord] = EncodeCategorizationData(Categorization);
 }
 
-#else // PERMUTATION_PPLL or PERMUTATION_VISIBILITY or PERMUTATION_MSAA
-
-///////////////////////////////////////////////////////////////////////////////////////////////////////////////
-// MSAA buffer
-///////////////////////////////////////////////////////////////////////////////////////////////////////////////
-
-#if PERMUTATION_MSAACOUNT > 1
-void ComputeUniqueSamplesWithoutDepthTest(
-	inout uint3 OutSamples[PERMUTATION_MSAACOUNT],
-	inout FSampleSetDesc OutSet,
-	const uint2 PixelCoord,
-	const float OpaqueDepth,
-	Texture2DMS<uint2, PERMUTATION_MSAACOUNT>  InMSAA_IDTexture,
-	Texture2DMS<float, PERMUTATION_MSAACOUNT>  InMSAA_DepthTexture)
-{
-	OutSet.UniqueSampleCount = 0;
-	OutSet.ValidSampleCount = 0;
-	OutSet.HairSampleCount = PERMUTATION_MSAACOUNT;
-
-#if PERMUTATION_MATERIAL_COMPACTION == 1
-	float4 UniqueTangentAndDepths[PERMUTATION_MSAACOUNT];
-#endif
-	float ClosestDepth = 0; // Inverse-Z
-	for (uint SampleIt = 0; SampleIt < OutSet.HairSampleCount; ++SampleIt)
-	{
-		const uint2 EncodedData = InMSAA_IDTexture.Load(PixelCoord, SampleIt);
-		const bool bIsValid = IsValidHairData(EncodedData);
-		if (!bIsValid)
-			continue;
-
-		const float SampleDepth = InMSAA_DepthTexture.Load(PixelCoord, SampleIt);
-		if (OpaqueDepth > SampleDepth) // Inverse-Z
-			continue;
-
-		++OutSet.ValidSampleCount;
-
-	#if PERMUTATION_MATERIAL_COMPACTION == 1
-		const float SceneDepth = ConvertFromDeviceZ(SampleDepth);
-	#endif
-
-		uint PrimitiveId;
-		uint MacroGroupId;
-		uint Coverage8bit;
-		float3 Tangent;
-
-		UnpackHairMRT0(EncodedData, PrimitiveId, MacroGroupId, Coverage8bit, Tangent);
-		bool bAlreadyExist = false;
-		for (uint UniqueIt = 0; UniqueIt < OutSet.UniqueSampleCount; ++UniqueIt)
-		{
-		#if PERMUTATION_MATERIAL_COMPACTION == 1
-			const bool bIsSimilar  = 
-				PrimitiveId == OutSamples[UniqueIt].x ||
-				(abs(dot(UniqueTangentAndDepths[UniqueIt].xyz, Tangent)) > CosTangentThreshold &&
-					abs(UniqueTangentAndDepths[UniqueIt].w - SceneDepth) < DepthTheshold);
-		#else
-			const bool bIsSimilar = PrimitiveId == OutSamples[UniqueIt].x;
-		#endif
-
-			if (bIsSimilar)
-			{
-				OutSamples[UniqueIt].y += 1;
-				if (ClosestDepth < SampleDepth)
-				{
-					OutSamples[OutSet.UniqueSampleCount].z = SampleIt;
-					ClosestDepth = SampleDepth;
-				}
-				bAlreadyExist = true;
-				break;
-			}
-		}
-
-		if (!bAlreadyExist)
-		{
-		#if PERMUTATION_MATERIAL_COMPACTION == 1
-			UniqueTangentAndDepths[OutSet.UniqueSampleCount] = float4(Tangent, SceneDepth);
-		#endif
-
-			OutSamples[OutSet.UniqueSampleCount].x = PrimitiveId;
-			OutSamples[OutSet.UniqueSampleCount].y = 1;
-			OutSamples[OutSet.UniqueSampleCount].z = SampleIt;
-			++OutSet.UniqueSampleCount;
-		}
-	}
-}
-
-Texture2DMS<uint2, PERMUTATION_MSAACOUNT> MSAA_IDTexture;
-Texture2DMS<float, PERMUTATION_MSAACOUNT> MSAA_DepthTexture;
-Texture2DMS<float4,PERMUTATION_MSAACOUNT> MSAA_MaterialTexture;
-Texture2DMS<float4,PERMUTATION_MSAACOUNT> MSAA_AttributeTexture;
-
-uint2  GetIDTexture(uint2 PixelCoord, uint SampleIt)		{ return MSAA_IDTexture.Load(PixelCoord, SampleIt);  }
-float  GetDepthTexture(uint2 PixelCoord, uint SampleIt)		{ return MSAA_DepthTexture.Load(PixelCoord, SampleIt);  }
-float4 GetMaterialTexture(uint2 PixelCoord, uint SampleIt)	{ return MSAA_MaterialTexture.Load(PixelCoord, SampleIt); }
-float4 GetAttributeTexture(uint2 PixelCoord, uint SampleIt)	{ return MSAA_AttributeTexture.Load(PixelCoord, SampleIt); }
-
-#if PERMUTATION_VELOCITY != VELOCITY_TYPE_NONE
-Texture2DMS<float2, PERMUTATION_MSAACOUNT> MSAA_VelocityTexture;
-float2 GetVelocityTexture(uint2 PixelCoord, uint SampleIt)	{ return MSAA_VelocityTexture.Load(PixelCoord, SampleIt);  }
-#endif
-
-#else // PERMUTATION_MSAACOUNT > 1
-
-void ComputeUniqueSamplesWithoutDepthTest(
-	inout uint3 OutSamples[PERMUTATION_MSAACOUNT],
-	inout FSampleSetDesc OutSet,
-	const uint2 PixelCoord,
-	const float OpaqueDepth,
-	Texture2D<uint2>  InMSAA_IDTexture,
-	Texture2D<float>  InMSAA_DepthTexture)
-{
-	OutSet.UniqueSampleCount = 0;
-	OutSet.ValidSampleCount = 0;
-	OutSet.HairSampleCount = 1;
-
-	const uint2 EncodedData = InMSAA_IDTexture.Load(uint3(PixelCoord, 0));
-	const bool bIsValid = IsValidHairData(EncodedData);
-	if (!bIsValid)
-		return;
-
-	const float SampleDepth = InMSAA_DepthTexture.Load(uint3(PixelCoord, 0));
-	if (OpaqueDepth > SampleDepth) // Inverse-Z
-		return;
-
-	++OutSet.ValidSampleCount;
-
-	uint PrimitiveId;
-	uint MacroGroupId;
-	uint Coverage8bit;
-	float3 Tangent;
-
-	UnpackHairMRT0(EncodedData, PrimitiveId, MacroGroupId, Coverage8bit, Tangent);
-
-	OutSamples[0].x = PrimitiveId;
-	OutSamples[0].y = 1;
-	OutSamples[0].z = 0;
-	++OutSet.UniqueSampleCount;
-}
-
-Texture2D<uint2> MSAA_IDTexture;
-Texture2D<float> MSAA_DepthTexture;
-Texture2D<float4> MSAA_MaterialTexture;
-Texture2D<float4> MSAA_AttributeTexture;
-
-uint2  GetIDTexture(uint2 PixelCoord, uint SampleIt)			{ return MSAA_IDTexture.Load(uint3(PixelCoord, 0)); }
-float  GetDepthTexture(uint2 PixelCoord, uint SampleIt)			{ return MSAA_DepthTexture.Load(uint3(PixelCoord, 0));  }
-float4 GetMaterialTexture(uint2 PixelCoord, uint SampleIt)		{ return MSAA_MaterialTexture.Load(uint3(PixelCoord, 0)); }
-float4 GetAttributeTexture(uint2 PixelCoord, uint SampleIt)		{ return MSAA_AttributeTexture.Load(uint3(PixelCoord, 0)); }
-
-#if PERMUTATION_VELOCITY != VELOCITY_TYPE_NONE
-Texture2D<float2> MSAA_VelocityTexture;
-float2 GetVelocityTexture(uint2 PixelCoord, uint SampleIt)		{ return MSAA_VelocityTexture.Load(uint3(PixelCoord,0)); }
-#endif
-
-#endif // PERMUTATION_MSAACOUNT > 1
-
-#if PERMUTATION_VIEWTRANSMITTANCE == 1
-Texture2D<float> ViewTransmittanceTexture;
-#endif
-
-[numthreads(TILE_PIXEL_SIZE_X, TILE_PIXEL_SIZE_Y, 1)]
-void MainCS(uint GroupIndex : SV_GroupIndex, uint3 DispatchThreadId : SV_DispatchThreadID)
-{
-	if (GroupIndex == 0)
-	{
-		AllocationNodeCount = 0;
-		AllocationNodeOffset = 0;
-	}
-
-	uint2 PixelCoord = DispatchThreadId.xy + ResolutionOffset;
-	if (PixelCoord.x >= uint(OutputResolution.x) || PixelCoord.y >= uint(OutputResolution.y))
-		PixelCoord = uint2(0, 0);
-
-	FSampleSetDesc SampleDesc;
-
-	const float OpaqueDepth = SceneTexturesStruct.SceneDepthTexture.Load(uint3(PixelCoord, 0)).r;
-	uint3 Samples[PERMUTATION_MSAACOUNT]; // x:PrimitiveId, y:Weight, z:SampleIndex
-	ComputeUniqueSamplesWithoutDepthTest(
-		Samples,
-		SampleDesc,
-		PixelCoord,
-		OpaqueDepth,
-		MSAA_IDTexture, 
-		MSAA_DepthTexture);
-
-	FNodeDesc NodeDesc;
-	NodeDesc.Count = SampleDesc.UniqueSampleCount;
-	NodeDesc.Offset = 0;
-
-	if (NodeDesc.Count > 0)
-	{
-		InterlockedAdd(AllocationNodeCount, NodeDesc.Count, NodeDesc.Offset);
-	}
-	GroupMemoryBarrierWithGroupSync();
-	if (GroupIndex == 0 && AllocationNodeCount > 0)
-	{
-		InterlockedAdd(OutCompactNodeCounter[uint2(0, 0)], AllocationNodeCount, AllocationNodeOffset);
-	}
-	GroupMemoryBarrierWithGroupSync();
-
-	// Allocate node space
-	FCategorizationData Categorization = (FCategorizationData)0;
-	if (NodeDesc.Count > 0)
-	{
-		NodeDesc.Offset += AllocationNodeOffset;
-
-		// Store final sort node data 
-		if (NodeDesc.Offset + NodeDesc.Count < MaxNodeCount)
-		{
-		#if PERMUTATION_VELOCITY != VELOCITY_TYPE_NONE
-			float2 AverageVelocity = 0;
-			uint ValidVelocityCount = 0;
-			
-			float2 ClosestEncodedVelocity = 0;
-			float ClosestDepth = 0; // Inverse-Z
-			
-			float2 MaxEncodedVelocity = 0;
-			float MaxVelocityMagnitude2 = 0;
-		#endif		
-
-			for (uint OutIndex = 0; OutIndex < NodeDesc.Count; ++OutIndex)
-			{
-				const uint SampleIt = Samples[OutIndex].z;
-
-				// Read all material data
-				FPackedHairSample OutNodeData;
-				{
-					FHairSample HairSample;
-					HairSample.Depth = GetDepthTexture(PixelCoord, SampleIt);
-
-					UnpackHairMRT0(GetIDTexture(PixelCoord, SampleIt), HairSample.PrimitiveId, HairSample.MacroGroupId, HairSample.Coverage8bit, HairSample.Tangent);
-					UnpackHairMRT1(GetMaterialTexture(PixelCoord, SampleIt), HairSample.BaseColor, HairSample.Roughness);
-					UnpackHairMRT2(GetAttributeTexture(PixelCoord, SampleIt), HairSample.Specular, HairSample.LightChannelMask, HairSample.Backlit, HairSample.bScatterSceneLighting);
-
-					// VisibilityData.Coverage8bit is a weight normalising to 1 the contribution of all the compacted samples. Because later it is weighted by Categorization.PixelCoverage.
-					HairSample.Coverage8bit = To8bitCoverage(Samples[OutIndex].y / float(SampleDesc.ValidSampleCount));
-					Categorization.TotalCoverage8bit += HairSample.Coverage8bit;
-					Categorization.ClosestDepth = max(Categorization.ClosestDepth, HairSample.Depth); // Inverse-Z
-
-					OutNodeData = PackHairSample(HairSample);
-
-				#if PERMUTATION_VELOCITY != VELOCITY_TYPE_NONE
-					if (VelocityType == VELOCITY_TYPE_AVG)
-					{
-						const float2 EncodedVelocity = GetVelocityTexture(PixelCoord, SampleIt);
-						AverageVelocity += DecodeVelocityFromTexture(float4(EncodedVelocity, 0.0, 0.0)).xy;
-						++ValidVelocityCount;
-					}
-					else if (VelocityType == VELOCITY_TYPE_CLOSEST)
-					{
-						const float2 EncodedVelocity = GetVelocityTexture(PixelCoord, SampleIt);
-						if (ClosestDepth < HairSample.Depth) // Inverse-Z
-						{
-							ClosestEncodedVelocity = EncodedVelocity;
-							ClosestDepth = HairSample.Depth;
-						}
-					}
-					else if (VelocityType == VELOCITY_TYPE_MAX)
-					{
-						const float2 EncodedVelocity = GetVelocityTexture(PixelCoord, SampleIt);
-						const float2 CurrentVelocity = DecodeVelocityFromTexture(float4(EncodedVelocity, 0.0, 0.0)).xy;
-						const float CurrentVelocityMagnitude2 = dot(CurrentVelocity, CurrentVelocity);
-						if (CurrentVelocityMagnitude2 > MaxVelocityMagnitude2)
-						{
-							MaxEncodedVelocity = EncodedVelocity;
-							MaxVelocityMagnitude2 = CurrentVelocityMagnitude2;
-						}
-					}
-				#endif
-				}
-
-				const uint StoreIndex = NodeDesc.Offset + OutIndex;
-				OutCompactNodeData[StoreIndex] = OutNodeData;
-				OutCompactNodeCoord[StoreIndex] = PixelCoord;
-			}
-
-			#if PERMUTATION_VIEWTRANSMITTANCE == 1
-			Categorization.PixelCoverage = TransmittanceToCoverage(ViewTransmittanceTexture.Load(uint3(PixelCoord,0)), CoverageThreshold);
-			#else
-			Categorization.PixelCoverage = SampleDesc.ValidSampleCount / float(PERMUTATION_MSAACOUNT);
-			#endif
-
-		#if PERMUTATION_VELOCITY != VELOCITY_TYPE_NONE
-			float2 VelocityOutput = 0;
-			if (VelocityType == VELOCITY_TYPE_AVG)
-			{
-				VelocityOutput = EncodeVelocityToTexture(float3(AverageVelocity / (float)max(ValidVelocityCount, 1u), 0.0)).xy;
-			}
-			else if (VelocityType == VELOCITY_TYPE_CLOSEST)
-			{
-				VelocityOutput = ClosestEncodedVelocity;
-			}
-			else if (VelocityType == VELOCITY_TYPE_MAX)
-			{
-				VelocityOutput = MaxEncodedVelocity;
-			}
-			OutVelocityTexture[PixelCoord] = VelocityOutput;
-		#endif		
-
-			NodeDesc.Count = NodeDesc.Count;
-			Categorization.SampleCount = NodeDesc.Count;
-		}
-	}
-
-	OutCompactNodeIndex[PixelCoord] = EncodeNodeDesc(NodeDesc);
-	OutCategorizationTexture[PixelCoord] = EncodeCategorizationData(Categorization);
-}
-
-#endif // PERMUTATION_PPLL or PERMUTATION_VISIBILITY or PERMUTATION_MSAA
+#endif // PERMUTATION_PPLL or PERMUTATION_VISIBILITY
diff --git a/Engine/Shaders/Private/HairStrands/HairStrandsVisibilityPS.usf b/Engine/Shaders/Private/HairStrands/HairStrandsVisibilityPS.usf
--- a/Engine/Shaders/Private/HairStrands/HairStrandsVisibilityPS.usf
+++ b/Engine/Shaders/Private/HairStrands/HairStrandsVisibilityPS.usf
@@ -21,7 +21,7 @@
 
 #define SUPPORT_OPACITY_MASK 0
 
-#if HAIR_RENDER_MODE == RENDER_MODE_MSAA || HAIR_RENDER_MODE == RENDER_MODE_PPLL
+#if HAIR_RENDER_MODE == RENDER_MODE_PPLL
 #define SUPPORT_MATERIAL_PROPERTY 1
 #endif
 
@@ -47,7 +47,7 @@ void MainVisibility(
 	centroid FVertexFactoryInterpolantsVSToPS Interpolants
 #endif
 	, in float4 SvPosition	: SV_Position
-#if HAIR_RENDER_MODE == RENDER_MODE_MSAA || HAIR_RENDER_MODE == RENDER_MODE_PPLL
+#if HAIR_RENDER_MODE == RENDER_MODE_PPLL
 	, centroid in float WorldStrandRadius : TEXCOORD8
 	, const float4 PackedVelocityA : TEXCOORD9
 	, const float4 PackedVelocityC : TEXCOORD12
@@ -60,13 +60,6 @@ void MainVisibility(
 #elif HAIR_RENDER_MODE == RENDER_MODE_TRANSMITTANCE_AND_HAIRCOUNT
 	, out float OutColor0 : SV_Target0
 	, out float2 OutColor1 : SV_Target1
-#elif HAIR_RENDER_MODE == RENDER_MODE_MSAA
-	, out uint2 OutColor0	: SV_Target0
-  #if SUPPORT_MATERIAL_PROPERTY
-	, out float4 OutColor1	: SV_Target1
-	, out float4 OutColor2	: SV_Target2
-	, out float2 OutColor3	: SV_Target3
-  #endif
 #elif HAIR_RENDER_MODE == RENDER_MODE_MSAA_VISIBILITY
 	, out uint OutColor0 : SV_Target0
 #endif
@@ -106,7 +99,7 @@ void MainVisibility(
 	OutColor1 = float2(Coverage, 1);
 #endif
 
-#if HAIR_RENDER_MODE == RENDER_MODE_MSAA || HAIR_RENDER_MODE == RENDER_MODE_PPLL
+#if HAIR_RENDER_MODE == RENDER_MODE_PPLL
   #if HAIR_STRAND_MESH_FACTORY
 	float3 Tangent = normalize(Interpolants.TangentToWorld2.xyz);
 	const uint PrimitiveId = Interpolants.HairPrimitiveId;
@@ -163,13 +156,6 @@ void MainVisibility(
 	#endif
 #endif
 
-#if HAIR_RENDER_MODE == RENDER_MODE_MSAA
-	OutColor0 = PackHairMRT0(PrimitiveId, MacroGroupId, Coverage8bit, Tangent);
-	OutColor1 = PackHairMRT1(BaseColor, Roughness);
-	OutColor2 = PackHairMRT2(Specular, HairVisibilityPass_LightChannelMask, Backlit, bUseScatterSceneLighting);
-	OutColor3 = EncodedVelocity;
-#endif
-
 #if HAIR_RENDER_MODE == RENDER_MODE_PPLL
 	const uint2 PixelCoord = SvPosition.xy;
 
diff --git a/Engine/Shaders/Private/HairStrands/HairStrandsVisibilityVS.usf b/Engine/Shaders/Private/HairStrands/HairStrandsVisibilityVS.usf
--- a/Engine/Shaders/Private/HairStrands/HairStrandsVisibilityVS.usf
+++ b/Engine/Shaders/Private/HairStrands/HairStrandsVisibilityVS.usf
@@ -19,7 +19,7 @@ struct FHairVisibilityAccumulateVSToPS
 #if HAIR_RENDER_MODE != RENDER_MODE_MSAA_VISIBILITY
 	float WorldStrandRadius : TEXCOORD8;
 #endif
-#if HAIR_RENDER_MODE == RENDER_MODE_MSAA || HAIR_RENDER_MODE == RENDER_MODE_PPLL
+#if HAIR_RENDER_MODE == RENDER_MODE_PPLL
 	float4 PackedVelocityA : TEXCOORD9;
 	float4 PackedVelocityC : TEXCOORD12;
 #endif
@@ -77,7 +77,7 @@ void Main(
 	// The higher the velocity, the larger a strand becomes.
 	float VelocityScale = 0;
 	const FHairRenderInfo HairRenderInfo = GetHairRenderInfo(ResolvedView.HairRenderInfo, ResolvedView.HairRenderInfoBits, bUseStableRasterization);
-	#if HAIR_RENDER_MODE == RENDER_MODE_MSAA || HAIR_RENDER_MODE == RENDER_MODE_PPLL
+	#if HAIR_RENDER_MODE == RENDER_MODE_PPLL
 	{
 		const float4 PrevWorldPosition	= VertexFactoryGetPreviousWorldPosition(Input, VFIntermediates);
 		const float4 ScreenPos			= mul(float4(WorldPositionRaw.xyz, 1), ResolvedView.TranslatedWorldToClip);
diff --git a/Engine/Source/Runtime/Renderer/Private/HairStrands/HairStrandsVisibility.cpp b/Engine/Source/Runtime/Renderer/Private/HairStrands/HairStrandsVisibility.cpp
--- a/Engine/Source/Runtime/Renderer/Private/HairStrands/HairStrandsVisibility.cpp
+++ b/Engine/Source/Runtime/Renderer/Private/HairStrands/HairStrandsVisibility.cpp
@@ -53,9 +53,6 @@ static FAutoConsoleVariableRef CVarGHairVisibilityPPLL(TEXT("r.HairStrands.Visib
 static FAutoConsoleVariableRef CVarGHairVisibilityPPLL_MeanNodeCountPerPixel(TEXT("r.HairStrands.Visibility.PPLL.SamplePerPixel"), GHairVisibilityPPLL_MaxSamplePerPixel, TEXT("The maximum number of node allowed to be independently shaded and composited per pixel. Total amount of node will be width*height*VisibilityPPLLMaxRenderNodePerPixel. The last node is used to aggregate all furthest strands to shade into a single one."));
 static FAutoConsoleVariableRef CVarGHairVisibilityPPLL_MeanSamplePerPixel(TEXT("r.HairStrands.Visibility.PPLL.MeanSamplePerPixel"), GHairVisibilityPPLL_MeanSamplePerPixel, TEXT("Scale the maximum number of node allowed for all linked list element (0..1, default 1). It will be width*height*SamplerPerPixel*Scale."));
 
-static int32 GHairStrandsVisibilityMaterialPass = 1;
-static FAutoConsoleVariableRef CVarHairStrandsVisibilityMaterialPass(TEXT("r.HairStrands.Visibility.MaterialPass"), GHairStrandsVisibilityMaterialPass, TEXT("Enable the deferred material pass evaluation after the hair visibility is resolved."));
-
 static float GHairStrandsViewHairCountDepthDistanceThreshold = 30.f;
 static FAutoConsoleVariableRef CVarHairStrandsViewHairCountDepthDistanceThreshold(TEXT("r.HairStrands.Visibility.HairCount.DistanceThreshold"), GHairStrandsViewHairCountDepthDistanceThreshold, TEXT("Distance threshold defining if opaque depth get injected into the 'view-hair-count' buffer."));
 
@@ -110,7 +107,6 @@ namespace HairStrandsVisibilityInternal
 
 enum EHairVisibilityRenderMode
 {
-	HairVisibilityRenderMode_MSAA,
 	HairVisibilityRenderMode_Transmittance,
 	HairVisibilityRenderMode_PPLL,
 	HairVisibilityRenderMode_MSAA_Visibility,
@@ -140,14 +136,14 @@ inline EHairVisibilityRenderMode GetHairVisibilityRenderMode()
 	}
 	else
 	{
-		return HairVisibilityRenderMode_MSAA;
+		return HairVisibilityRenderMode_MSAA_Visibility;
 	}
 }
 
 inline bool IsMsaaEnabled()
 {
 	const EHairVisibilityRenderMode Mode = GetHairVisibilityRenderMode();
-	return Mode == HairVisibilityRenderMode_MSAA || Mode == HairVisibilityRenderMode_MSAA_Visibility;
+	return Mode == HairVisibilityRenderMode_MSAA_Visibility;
 }
 
 static uint32 GetMaxSamplePerPixel()
@@ -169,7 +165,6 @@ static uint32 GetMaxSamplePerPixel()
 				return 4;
 			}
 		}
-		case HairVisibilityRenderMode_MSAA:
 		case HairVisibilityRenderMode_MSAA_Visibility:
 		{
 			if (GHairVisibilityMSAA_MaxSamplePerPixel <= 1)
@@ -221,7 +216,6 @@ inline uint32 GetMeanSamplePerPixel()
 	{
 	case HairVisibilityRenderMode_ComputeRaster:
 		return FMath::Max(1, FMath::FloorToInt(SamplePerPixel * FMath::Clamp(GHairVisibilityComputeRaster_MeanSamplePerPixel, 0.f, 1.f)));
-	case HairVisibilityRenderMode_MSAA:
 	case HairVisibilityRenderMode_MSAA_Visibility:
 		return FMath::Max(1, FMath::FloorToInt(SamplePerPixel * FMath::Clamp(GHairVisibilityMSAA_MeanSamplePerPixel, 0.f, 1.f)));
 	case HairVisibilityRenderMode_PPLL:
@@ -1417,14 +1411,12 @@ protected:
 
 typedef FHairVisibilityVS<HairVisibilityRenderMode_MSAA_Visibility, false >				THairVisiblityVS_MSAAVisibility_NoCulling;
 typedef FHairVisibilityVS<HairVisibilityRenderMode_MSAA_Visibility, true >				THairVisiblityVS_MSAAVisibility_Culling;
-typedef FHairVisibilityVS<HairVisibilityRenderMode_MSAA, true >							THairVisiblityVS_MSAA;
 typedef FHairVisibilityVS<HairVisibilityRenderMode_Transmittance, true >				THairVisiblityVS_Transmittance;
 typedef FHairVisibilityVS<HairVisibilityRenderMode_TransmittanceAndHairCount, true >	THairVisiblityVS_TransmittanceAndHairCount;
 typedef FHairVisibilityVS<HairVisibilityRenderMode_PPLL, true >							THairVisiblityVS_PPLL;
 
 IMPLEMENT_MATERIAL_SHADER_TYPE(template<>, THairVisiblityVS_MSAAVisibility_NoCulling,	TEXT("/Engine/Private/HairStrands/HairStrandsVisibilityVS.usf"), TEXT("Main"), SF_Vertex);
 IMPLEMENT_MATERIAL_SHADER_TYPE(template<>, THairVisiblityVS_MSAAVisibility_Culling,		TEXT("/Engine/Private/HairStrands/HairStrandsVisibilityVS.usf"), TEXT("Main"), SF_Vertex);
-IMPLEMENT_MATERIAL_SHADER_TYPE(template<>, THairVisiblityVS_MSAA,						TEXT("/Engine/Private/HairStrands/HairStrandsVisibilityVS.usf"), TEXT("Main"), SF_Vertex);
 IMPLEMENT_MATERIAL_SHADER_TYPE(template<>, THairVisiblityVS_Transmittance,				TEXT("/Engine/Private/HairStrands/HairStrandsVisibilityVS.usf"), TEXT("Main"), SF_Vertex);
 IMPLEMENT_MATERIAL_SHADER_TYPE(template<>, THairVisiblityVS_TransmittanceAndHairCount,	TEXT("/Engine/Private/HairStrands/HairStrandsVisibilityVS.usf"), TEXT("Main"), SF_Vertex);
 IMPLEMENT_MATERIAL_SHADER_TYPE(template<>, THairVisiblityVS_PPLL,						TEXT("/Engine/Private/HairStrands/HairStrandsVisibilityVS.usf"), TEXT("Main"), SF_Vertex);
@@ -1487,12 +1479,6 @@ public:
 		{
 			OutEnvironment.SetRenderTargetOutputFormat(0, PF_R32_UINT);
 		}
-		else if (RenderMode == HairVisibilityRenderMode_MSAA)
-		{
-			OutEnvironment.SetRenderTargetOutputFormat(0, PF_R32G32_UINT);
-			OutEnvironment.SetRenderTargetOutputFormat(1, PF_R32G32_UINT);
-			OutEnvironment.SetRenderTargetOutputFormat(2, PF_R32G32_UINT);
-		}
 		else if (RenderMode == HairVisibilityRenderMode_Transmittance)
 		{
 			OutEnvironment.SetRenderTargetOutputFormat(0, PF_R32_FLOAT);
@@ -1525,7 +1511,6 @@ public:
 	LAYOUT_FIELD(FShaderParameter, HairVisibilityPass_LightChannelMask);
 };
 IMPLEMENT_MATERIAL_SHADER_TYPE(template<>, FHairVisibilityPS<HairVisibilityRenderMode_MSAA_Visibility>, TEXT("/Engine/Private/HairStrands/HairStrandsVisibilityPS.usf"), TEXT("MainVisibility"), SF_Pixel);
-IMPLEMENT_MATERIAL_SHADER_TYPE(template<>, FHairVisibilityPS<HairVisibilityRenderMode_MSAA>, TEXT("/Engine/Private/HairStrands/HairStrandsVisibilityPS.usf"), TEXT("MainVisibility"), SF_Pixel);
 IMPLEMENT_MATERIAL_SHADER_TYPE(template<>, FHairVisibilityPS<HairVisibilityRenderMode_Transmittance>, TEXT("/Engine/Private/HairStrands/HairStrandsVisibilityPS.usf"), TEXT("MainVisibility"), SF_Pixel);
 IMPLEMENT_MATERIAL_SHADER_TYPE(template<>, FHairVisibilityPS<HairVisibilityRenderMode_TransmittanceAndHairCount>, TEXT("/Engine/Private/HairStrands/HairStrandsVisibilityPS.usf"), TEXT("MainVisibility"), SF_Pixel);
 IMPLEMENT_MATERIAL_SHADER_TYPE(template<>, FHairVisibilityPS<HairVisibilityRenderMode_PPLL>, TEXT("/Engine/Private/HairStrands/HairStrandsVisibilityPS.usf"), TEXT("MainVisibility"), SF_Pixel);
@@ -1594,8 +1579,6 @@ void FHairVisibilityProcessor::AddMeshBatch(const FMeshBatch& RESTRICT MeshBatch
 			Process<HairVisibilityRenderMode_MSAA_Visibility, true>(MeshBatch, BatchElementMask, PrimitiveSceneProxy, StaticMeshId, MaterialRenderProxy, Material, HairMacroGroupId, HairMaterialId, LightChannelMask, MeshFillMode, MeshCullMode);
 		else if (RenderMode == HairVisibilityRenderMode_MSAA_Visibility && !bCullingEnable)
 			Process<HairVisibilityRenderMode_MSAA_Visibility, false>(MeshBatch, BatchElementMask, PrimitiveSceneProxy, StaticMeshId, MaterialRenderProxy, Material, HairMacroGroupId, HairMaterialId, LightChannelMask, MeshFillMode, MeshCullMode);
-		else if (RenderMode == HairVisibilityRenderMode_MSAA)
-			Process<HairVisibilityRenderMode_MSAA>(MeshBatch, BatchElementMask, PrimitiveSceneProxy, StaticMeshId, MaterialRenderProxy, Material, HairMacroGroupId, HairMaterialId, LightChannelMask, MeshFillMode, MeshCullMode);
 		else if (RenderMode == HairVisibilityRenderMode_Transmittance)
 			Process<HairVisibilityRenderMode_Transmittance>(MeshBatch, BatchElementMask, PrimitiveSceneProxy, StaticMeshId, MaterialRenderProxy, Material, HairMacroGroupId, HairMaterialId, LightChannelMask, MeshFillMode, MeshCullMode);
 		else if (RenderMode == HairVisibilityRenderMode_TransmittanceAndHairCount)
@@ -1819,9 +1802,8 @@ class FHairVisibilityPrimitiveIdCompactionCS : public FGlobalShader
 	class FViewTransmittance : SHADER_PERMUTATION_INT("PERMUTATION_VIEWTRANSMITTANCE", 2);
 	class FMaterial 	: SHADER_PERMUTATION_INT("PERMUTATION_MATERIAL_COMPACTION", 2);
 	class FPPLL 		: SHADER_PERMUTATION_SPARSE_INT("PERMUTATION_PPLL", 0, 8, 16, 32); // See GetPPLLMaxRenderNodePerPixel
-	class FVisibility 	: SHADER_PERMUTATION_INT("PERMUTATION_VISIBILITY", 2);
 	class FMSAACount 	: SHADER_PERMUTATION_SPARSE_INT("PERMUTATION_MSAACOUNT", 1, 2, 4, 8);
-	using FPermutationDomain = TShaderPermutationDomain<FGroupSize, FVelocity, FViewTransmittance, FMaterial, FPPLL, FVisibility, FMSAACount>;
+	using FPermutationDomain = TShaderPermutationDomain<FGroupSize, FVelocity, FViewTransmittance, FMaterial, FPPLL, FMSAACount>;
 
 	BEGIN_SHADER_PARAMETER_STRUCT(FParameters, )
 		SHADER_PARAMETER(FIntPoint, OutputResolution)
@@ -1864,7 +1846,6 @@ public:
 		if (PermutationVector.Get<FPPLL>() > 0)
 		{
 			PermutationVector.Set<FViewTransmittance>(0);
-			PermutationVector.Set<FVisibility>(0);
 			PermutationVector.Set<FMSAACount>(4);
 		}
 		return PermutationVector;
@@ -1877,10 +1858,6 @@ public:
 		{
 			return false;
 		}
-		if (PermutationVector.Get<FPPLL>() > 0 && PermutationVector.Get<FVisibility>() > 0)
-		{
-			return false;
-		}
 		if (PermutationVector.Get<FPPLL>() > 0 && PermutationVector.Get<FMSAACount>() == 8)
 		{
 			return false;
@@ -1893,7 +1870,6 @@ IMPLEMENT_GLOBAL_SHADER(FHairVisibilityPrimitiveIdCompactionCS, "/Engine/Private
 
 static void AddHairVisibilityPrimitiveIdCompactionPass(
 	const bool bUsePPLL,
-	const bool bUseVisibility,
 	FRDGBuilder& GraphBuilder,
 	const FViewInfo& View,
 	const FHairStrandsMacroGroupDatas& MacroGroupDatas,
@@ -1919,21 +1895,9 @@ static void AddHairVisibilityPrimitiveIdCompactionPass(
 	else
 	{
 		check(PassParameters->MSAA_DepthTexture->Desc.NumSamples == GetMaxSamplePerPixel());
-
-		if (bUseVisibility)
-		{
-			check(PassParameters->MSAA_DepthTexture);
-			check(PassParameters->MSAA_IDTexture);
-			Resolution = PassParameters->MSAA_DepthTexture->Desc.Extent;
-		}
-		else
-		{
-			check(PassParameters->MSAA_DepthTexture);
-			check(PassParameters->MSAA_IDTexture);
-			check(PassParameters->MSAA_MaterialTexture);
-			check(PassParameters->MSAA_AttributeTexture);
-			Resolution = PassParameters->MSAA_DepthTexture->Desc.Extent;
-		}
+		check(PassParameters->MSAA_DepthTexture);
+		check(PassParameters->MSAA_IDTexture);
+		Resolution = PassParameters->MSAA_DepthTexture->Desc.Extent;
 	}
 
 	{
@@ -1961,11 +1925,11 @@ static void AddHairVisibilityPrimitiveIdCompactionPass(
 	const FIntPoint EffectiveResolution(HairRect.Width(), HairRect.Height());
 
 	// Select render node count according to current mode
-	const uint32 MSAASampleCount = GetHairVisibilityRenderMode() == HairVisibilityRenderMode_MSAA ? GetMaxSamplePerPixel() : 1;
+	const uint32 MSAASampleCount = GetHairVisibilityRenderMode() == HairVisibilityRenderMode_MSAA_Visibility ? GetMaxSamplePerPixel() : 1;
 	const uint32 PPLLMaxRenderNodePerPixel = GetMaxSamplePerPixel();
 	const uint32 MaxRenderNodeCount = GetTotalSampleCountForAllocation(EffectiveResolution);
 
-	OutCompactNodeData = GraphBuilder.CreateBuffer(FRDGBufferDesc::CreateStructuredDesc(bUseVisibility ? sizeof(HairStrandsVisibilityInternal::NodeVis) : sizeof(HairStrandsVisibilityInternal::NodeData), MaxRenderNodeCount), TEXT("Hair.VisibilityPrimitiveIdCompactNodeData"));
+	OutCompactNodeData = GraphBuilder.CreateBuffer(FRDGBufferDesc::CreateStructuredDesc(sizeof(HairStrandsVisibilityInternal::NodeVis), MaxRenderNodeCount), TEXT("Hair.VisibilityPrimitiveIdCompactNodeData"));
 
 	{
 		// Pixel coord of the node. Stored as 2*uint16, packed into a single uint32
@@ -1980,7 +1944,6 @@ static void AddHairVisibilityPrimitiveIdCompactionPass(
 	PermutationVector.Set<FHairVisibilityPrimitiveIdCompactionCS::FViewTransmittance>(PassParameters->ViewTransmittanceTexture ? 1 : 0);
 	PermutationVector.Set<FHairVisibilityPrimitiveIdCompactionCS::FMaterial>(GHairStrandsMaterialCompactionEnable ? 1 : 0);
 	PermutationVector.Set<FHairVisibilityPrimitiveIdCompactionCS::FPPLL>(bUsePPLL ? PPLLMaxRenderNodePerPixel : 0);
-	PermutationVector.Set<FHairVisibilityPrimitiveIdCompactionCS::FVisibility>(bUseVisibility ? 1 : 0);	
 	PermutationVector.Set<FHairVisibilityPrimitiveIdCompactionCS::FMSAACount>(MSAASampleCount);
 	PermutationVector = FHairVisibilityPrimitiveIdCompactionCS::RemapPermutation(PermutationVector);
 
@@ -2113,7 +2076,7 @@ static void AddHairVisibilityCompactionComputeRasterPass(
 	AddClearUAVPass(GraphBuilder, GraphBuilder.CreateUAV(OutCategorizationTexture), ClearValues);
 
 	// Select render node count according to current mode
-	const uint32 MSAASampleCount = GetHairVisibilityRenderMode() == HairVisibilityRenderMode_MSAA ? GetMaxSamplePerPixel() : 1;
+	const uint32 MSAASampleCount = GetHairVisibilityRenderMode() == HairVisibilityRenderMode_MSAA_Visibility ? GetMaxSamplePerPixel() : 1;
 	const uint32 PPLLMaxRenderNodePerPixel = GetMaxSamplePerPixel();
 	const uint32 MaxRenderNodeCount = GetTotalSampleCountForAllocation(Resolution);
 	OutCompactNodeData = GraphBuilder.CreateBuffer(FRDGBufferDesc::CreateStructuredDesc(sizeof(HairStrandsVisibilityInternal::NodeVis), MaxRenderNodeCount), TEXT("Hair.VisibilityPrimitiveIdCompactNodeData"));
@@ -2264,7 +2227,7 @@ static FRDGTextureRef AddHairVisibilityFillOpaqueDepth(
 {
 	FRDGTextureRef OutVisibilityDepthTexture;
 	{
-		check(GetHairVisibilityRenderMode() == HairVisibilityRenderMode_MSAA);
+		check(GetHairVisibilityRenderMode() == HairVisibilityRenderMode_MSAA_Visibility);
 
 		FRDGTextureDesc Desc = FRDGTextureDesc::Create2D(Resolution, PF_DepthStencil, FClearValueBinding::DepthFar, TexCreate_DepthStencilTargetable | TexCreate_ShaderResource, 1, GetMaxSamplePerPixel());
 		OutVisibilityDepthTexture = GraphBuilder.CreateTexture(Desc, TEXT("Hair.VisibilityDepthTexture"));
@@ -2383,7 +2346,6 @@ static void AddHairVisibilityCommonPass(
 		switch (RenderMode)
 		{
 		case HairVisibilityRenderMode_PPLL:						return RDG_EVENT_NAME("HairStrandsVisibilityPPLLPass");
-		case HairVisibilityRenderMode_MSAA:						return RDG_EVENT_NAME("HairStrandsVisibilityMSAAPass");
 		case HairVisibilityRenderMode_MSAA_Visibility:			return RDG_EVENT_NAME("HairStrandsVisibilityMSAAVisPass");
 		case HairVisibilityRenderMode_Transmittance:			return RDG_EVENT_NAME("HairStrandsTransmittancePass");
 		case HairVisibilityRenderMode_TransmittanceAndHairCount:return RDG_EVENT_NAME("HairStrandsTransmittanceAndHairCountPass");
@@ -2422,14 +2384,7 @@ static void AddHairVisibilityCommonPass(
 
 		{
 			RHICmdList.SetViewport(0, 0, 0.0f, ViewInfo->ViewRect.Width(), ViewInfo->ViewRect.Height(), 1.0f);
-			if (RenderMode == HairVisibilityRenderMode_MSAA)
-			{
-				DrawRenderState.SetBlendState(TStaticBlendState<
-					CW_RGBA, BO_Add, BF_One, BF_Zero, BO_Add, BF_One, BF_Zero,
-					CW_RGBA, BO_Add, BF_One, BF_Zero, BO_Add, BF_One, BF_Zero>::GetRHI());
-				DrawRenderState.SetDepthStencilState(TStaticDepthStencilState<true, CF_DepthNearOrEqual>::GetRHI());
-			}
-			else if (RenderMode == HairVisibilityRenderMode_MSAA_Visibility)
+			if (RenderMode == HairVisibilityRenderMode_MSAA_Visibility)
 			{
 				DrawRenderState.SetBlendState(TStaticBlendState<CW_RGBA, BO_Add, BF_One, BF_Zero, BO_Add, BF_One, BF_Zero>::GetRHI());
 				DrawRenderState.SetDepthStencilState(TStaticDepthStencilState<true, CF_DepthNearOrEqual>::GetRHI());
@@ -2480,7 +2435,6 @@ static void AddHairVisibilityCommonPass(
 }
 
 static void AddHairVisibilityMSAAPass(
-	const bool bUseVisibility,
 	FRDGBuilder& GraphBuilder,
 	const FScene* Scene,
 	const FViewInfo* ViewInfo,
@@ -2488,22 +2442,14 @@ static void AddHairVisibilityMSAAPass(
 	const FIntPoint& Resolution,
 	FInstanceCullingManager& InstanceCullingManager,
 	FRDGTextureRef& OutVisibilityIdTexture,
-	FRDGTextureRef& OutVisibilityMaterialTexture,
-	FRDGTextureRef& OutVisibilityAttributeTexture,
-	FRDGTextureRef& OutVisibilityVelocityTexture,
 	FRDGTextureRef& OutVisibilityDepthTexture)
 {
 	const uint32 MSAASampleCount = GetMaxSamplePerPixel();
-
-	if (bUseVisibility)
 	{
 		{
 			FRDGTextureDesc Desc = FRDGTextureDesc::Create2D(Resolution, PF_R32_UINT, FClearValueBinding(EClearBinding::ENoneBound), TexCreate_NoFastClear | TexCreate_RenderTargetable | TexCreate_ShaderResource, 1, MSAASampleCount);
 			OutVisibilityIdTexture = GraphBuilder.CreateTexture(Desc, TEXT("HairVisibilityIDTexture"));
 		}
-		OutVisibilityMaterialTexture = nullptr;
-		OutVisibilityAttributeTexture = nullptr;
-		OutVisibilityVelocityTexture = nullptr;
 
 		AddClearGraphicPass(GraphBuilder, RDG_EVENT_NAME("HairStrandsClearVisibilityMSAAIdTexture"), ViewInfo, 0xFFFFFFFF, OutVisibilityIdTexture);
 
@@ -2517,55 +2463,6 @@ static void AddHairVisibilityMSAAPass(
 			FExclusiveDepthStencil::DepthWrite_StencilNop);
 		AddHairVisibilityCommonPass(GraphBuilder, Scene, ViewInfo, MacroGroupDatas, HairVisibilityRenderMode_MSAA_Visibility, PassParameters, InstanceCullingManager);
 	}
-	else
-	{
-		{
-			FRDGTextureDesc Desc = FRDGTextureDesc::Create2D(Resolution, PF_R32G32_UINT, FClearValueBinding::None, TexCreate_RenderTargetable | TexCreate_ShaderResource, 1, MSAASampleCount);
-			OutVisibilityIdTexture = GraphBuilder.CreateTexture(Desc, TEXT("Hair.VisibilityIDTexture"));
-		}
-
-		{
-			FRDGTextureDesc Desc = FRDGTextureDesc::Create2D(Resolution, PF_R8G8B8A8, FClearValueBinding(FLinearColor(0, 0, 0, 0)), TexCreate_RenderTargetable | TexCreate_ShaderResource, 1, MSAASampleCount);
-			OutVisibilityMaterialTexture = GraphBuilder.CreateTexture(Desc, TEXT("Hair.VisibilityMaterialTexture"));
-		}
-
-		{
-			FRDGTextureDesc Desc = FRDGTextureDesc::Create2D(Resolution, PF_R8G8B8A8, FClearValueBinding(FLinearColor(0, 0, 0, 0)), TexCreate_RenderTargetable | TexCreate_ShaderResource, 1, MSAASampleCount);
-			OutVisibilityAttributeTexture = GraphBuilder.CreateTexture(Desc, TEXT("Hair.VisibilityAttributeTexture"));
-		}
-
-		{
-			FRDGTextureDesc Desc = FRDGTextureDesc::Create2D(Resolution, PF_G16R16, FClearValueBinding(FLinearColor(0, 0, 0, 0)), TexCreate_RenderTargetable | TexCreate_ShaderResource, 1, MSAASampleCount);
-			OutVisibilityVelocityTexture = GraphBuilder.CreateTexture(Desc, TEXT("Hair.VisibilityVelocityTexture"));
-		}
-		AddClearGraphicPass(GraphBuilder, RDG_EVENT_NAME("HairStrandsClearVisibilityMSAAIdTexture"), ViewInfo, 0xFFFFFFFF, OutVisibilityIdTexture);
-
-		// Manually clear RTs as using the Clear action on the RT, issue a global clean on all targets, while still need a special clear 
-		// for the PrimitiveId buffer
-		// const ERenderTargetLoadAction LoadAction = GHairClearVisibilityBuffer ? ERenderTargetLoadAction::EClear : ERenderTargetLoadAction::ENoAction;
-		ERenderTargetLoadAction LoadAction = ERenderTargetLoadAction::ENoAction;
-		if (GHairClearVisibilityBuffer)
-		{
-			LoadAction = ERenderTargetLoadAction::ELoad;
-			AddClearGraphicPass(GraphBuilder, RDG_EVENT_NAME("HairStrandsClearVisibilityMSAAMaterial"), ViewInfo, 0, OutVisibilityMaterialTexture);
-			AddClearGraphicPass(GraphBuilder, RDG_EVENT_NAME("HairStrandsClearVisibilityMSAAAttribute"), ViewInfo, 0, OutVisibilityAttributeTexture);
-			AddClearGraphicPass(GraphBuilder, RDG_EVENT_NAME("HairStrandsClearVisibilityMSAAVelocity"), ViewInfo, 0, OutVisibilityVelocityTexture);
-		}
-
-		FVisibilityPassParameters* PassParameters = GraphBuilder.AllocParameters<FVisibilityPassParameters>();
-		PassParameters->UniformBuffer = CreatePassDummyTextures(GraphBuilder);
-		PassParameters->RenderTargets[0] = FRenderTargetBinding(OutVisibilityIdTexture, ERenderTargetLoadAction::ELoad, 0);
-		PassParameters->RenderTargets[1] = FRenderTargetBinding(OutVisibilityMaterialTexture, LoadAction, 0);
-		PassParameters->RenderTargets[2] = FRenderTargetBinding(OutVisibilityAttributeTexture, LoadAction, 0);
-		PassParameters->RenderTargets[3] = FRenderTargetBinding(OutVisibilityVelocityTexture, LoadAction, 0);
-
-		PassParameters->RenderTargets.DepthStencil = FDepthStencilBinding(
-			OutVisibilityDepthTexture,
-			ERenderTargetLoadAction::ELoad,
-			ERenderTargetLoadAction::ENoAction,
-			FExclusiveDepthStencil::DepthWrite_StencilNop);
-		AddHairVisibilityCommonPass(GraphBuilder, Scene, ViewInfo, MacroGroupDatas, HairVisibilityRenderMode_MSAA, PassParameters, InstanceCullingManager);
-	}
 }
 
 static void AddHairVisibilityPPLLPass(
@@ -3315,7 +3212,7 @@ void RenderHairStrandsVisibilityBuffer(
 
 			const bool bRunColorAndDepthPatching = SceneGBufferBTexture && SceneColorTexture;
 			const EHairVisibilityRenderMode RenderMode = GetHairVisibilityRenderMode();
-			check(RenderMode == HairVisibilityRenderMode_MSAA || RenderMode == HairVisibilityRenderMode_PPLL || RenderMode == HairVisibilityRenderMode_ComputeRaster);
+			check(RenderMode == HairVisibilityRenderMode_MSAA_Visibility || RenderMode == HairVisibilityRenderMode_PPLL || RenderMode == HairVisibilityRenderMode_ComputeRaster);
 
 			FRDGTextureRef HairOnlyDepthTexture = GraphBuilder.CreateTexture(SceneDepthTexture->Desc, TEXT("Hair.HairOnlyDepthTexture"));
 			FRDGTextureRef CategorizationTexture = nullptr;
@@ -3467,7 +3364,7 @@ void RenderHairStrandsVisibilityBuffer(
 					ViewTransmittance.HairCountTexture = nullptr;
 				}
 			}
-			else if (RenderMode == HairVisibilityRenderMode_MSAA)
+			else if (RenderMode == HairVisibilityRenderMode_MSAA_Visibility)
 			{
 				// Run the view transmittance pass if needed (not in PPLL mode that is already a high quality render path)
 				FHairPrimaryTransmittance ViewTransmittance;
@@ -3497,15 +3394,10 @@ void RenderHairStrandsVisibilityBuffer(
 
 				}
 
-				const bool bIsVisiblityEnable = GHairStrandsVisibilityMaterialPass > 0;
-
 				struct FRDGMsaaVisibilityResources
 				{
 					FRDGTextureRef DepthTexture;
 					FRDGTextureRef IdTexture;
-					FRDGTextureRef MaterialTexture;
-					FRDGTextureRef AttributeTexture;
-					FRDGTextureRef VelocityTexture;
 				} MsaaVisibilityResources;
 
 				MsaaVisibilityResources.DepthTexture = AddHairVisibilityFillOpaqueDepth(
@@ -3516,7 +3408,6 @@ void RenderHairStrandsVisibilityBuffer(
 					SceneDepthTexture);
 
 				AddHairVisibilityMSAAPass(
-					bIsVisiblityEnable,
 					GraphBuilder,
 					Scene,
 					&View,
@@ -3524,28 +3415,18 @@ void RenderHairStrandsVisibilityBuffer(
 					Resolution,
 					InstanceCullingManager,
 					MsaaVisibilityResources.IdTexture,
-					MsaaVisibilityResources.MaterialTexture,
-					MsaaVisibilityResources.AttributeTexture,
-					MsaaVisibilityResources.VelocityTexture,
 					MsaaVisibilityResources.DepthTexture);
 
 				// This is used when compaction is not enabled.
 				VisibilityData.MaxSampleCount = MsaaVisibilityResources.IdTexture->Desc.NumSamples;
 				VisibilityData.HairOnlyDepthTexture = HairOnlyDepthTexture;
-				if (!bIsVisiblityEnable)
-				{
-					VisibilityData.VelocityTexture = MsaaVisibilityResources.VelocityTexture;
-				}
-
+				
 				const bool bUseComplexPath = IsHairStrandsComplexLightingEnabled();
 				if (bUseComplexPath)
 				{
 					FHairVisibilityPrimitiveIdCompactionCS::FParameters* PassParameters = GraphBuilder.AllocParameters<FHairVisibilityPrimitiveIdCompactionCS::FParameters>();
 					PassParameters->MSAA_DepthTexture = MsaaVisibilityResources.DepthTexture;
 					PassParameters->MSAA_IDTexture = MsaaVisibilityResources.IdTexture;
-					PassParameters->MSAA_MaterialTexture = MsaaVisibilityResources.MaterialTexture;
-					PassParameters->MSAA_AttributeTexture = MsaaVisibilityResources.AttributeTexture;
-					PassParameters->MSAA_VelocityTexture = MsaaVisibilityResources.VelocityTexture;
 					PassParameters->ViewTransmittanceTexture = ViewTransmittance.TransmittanceTexture;
 
 					FRDGBufferRef CompactNodeCoord;
@@ -3554,7 +3435,6 @@ void RenderHairStrandsVisibilityBuffer(
 					FRDGTextureRef EmissiveTexture = nullptr;
 					AddHairVisibilityPrimitiveIdCompactionPass(
 						false, // bUsePPLL
-						bIsVisiblityEnable,
 						GraphBuilder,
 						View,
 						MacroGroupDatas,
@@ -3569,7 +3449,6 @@ void RenderHairStrandsVisibilityBuffer(
 						IndirectArgsBuffer,
 						VisibilityData.MaxNodeCount);
 
-					if (bIsVisiblityEnable)
 					{
 						const bool bUpdateSampleCoverage = GHairStrandsSortHairSampleByDepth > 0;
 
@@ -3716,7 +3595,6 @@ void RenderHairStrandsVisibilityBuffer(
 					FRDGBufferRef IndirectArgsBuffer;
 					AddHairVisibilityPrimitiveIdCompactionPass(
 						true, // bUsePPLL
-						false,
 						GraphBuilder,
 						View,
 						MacroGroupDatas,
