Remove legacy hair voxelization path.

diff --git a/Engine/Shaders/Private/HairStrands/HairStrandsDeepShadowPS.usf b/Engine/Shaders/Private/HairStrands/HairStrandsDeepShadowPS.usf
--- a/Engine/Shaders/Private/HairStrands/HairStrandsDeepShadowPS.usf
+++ b/Engine/Shaders/Private/HairStrands/HairStrandsDeepShadowPS.usf
@@ -5,8 +5,6 @@
 
 #include "../Common.ush"
 
-//#define SceneTexturesStruct VoxelRasterPass.SceneTextures
-
 #include "../SceneTexturesCommon.ush"
 #include "HairStrandsDebugCommon.ush"
 #include "HairStrandsDeepShadowCommon.ush"
@@ -53,112 +51,4 @@ void MainDom(
 	const float DistanceToFrontDepth = GetDomDistanceToFrontDepth(FrontDepth, SvPosition.z);
 
 	OutColor = ComputeDOMWeights(DistanceToFrontDepth, DeepRasterPass.LayerDepths) * HairCoverage;
-}
-
-#define DEBUG_OUTPUT 0
-
-void MainVoxel(
-#if SUPPORT_TANGENT_PROPERTY == 1 || SUPPORT_MATERIAL_PROPERTY == 1 
-	FVertexFactoryInterpolantsVSToPS Interpolants,
-#endif
-	in float4 SvPosition : SV_Position,
-	in float4 WorldPositon_Coverage : TEXCOORD6
-#if DEBUG_OUTPUT
-,	out uint OutColor : SV_Target0 // For debug purpose only
-#endif
-)
-{
-	const float3 WorldPosition = WorldPositon_Coverage.xyz;
-	const float HairCoverage = WorldPositon_Coverage.w;
-
-#if HAIR_STRAND_MESH_FACTORY  
-	const FPackedVirtualVoxelNodeDesc PackedNode = VoxelRasterPass.VirtualVoxel.NodeDescBuffer[VoxelRasterPass.MacroGroupId];
-	const FVirtualVoxelNodeDesc NodeDesc = UnpackVoxelNode(PackedNode, VoxelRasterPass.VirtualVoxel.PageResolution);
-
-	const int3 VoxelCoord = PositionToCoord(WorldPosition.xyz, NodeDesc.MinAABB, NodeDesc.MaxAABB, NodeDesc.VirtualResolution);
-	if (VoxelCoord.x >= 0 || VoxelCoord.y >= 0 || VoxelCoord.z >= 0 ||
-		VoxelCoord.x < int(NodeDesc.VirtualResolution.x) || 
-		VoxelCoord.y < int(NodeDesc.VirtualResolution.y) || 
-		VoxelCoord.z < int(NodeDesc.VirtualResolution.z))
-	{
-		const int3 PageIndexCoord		= VoxelCoord / VoxelRasterPass.VirtualVoxel.PageResolution;
-		const uint LinearPageIndexCoord = CoordToIndex(PageIndexCoord, NodeDesc.PageIndexResolution, NodeDesc.PageIndexOffset);
-		const int3 VoxelCoordBase		= PageIndexCoord * VoxelRasterPass.VirtualVoxel.PageResolution;
-		const int3 VoxelCoordOffset		= VoxelCoord - VoxelCoordBase;
-		if (LinearPageIndexCoord < VoxelRasterPass.VirtualVoxel.PageIndexCount)
-		{
-			const uint PageIndex = VoxelRasterPass.VirtualVoxel.PageIndexBuffer[LinearPageIndexCoord];
-			if (PageIndex != INVALID_VOXEL_PAGE_INDEX)
-			{
-				const uint3 PageCoord			= IndexToCoord(PageIndex, VoxelRasterPass.VirtualVoxel.PageCountResolution);
-				const uint3 PageVoxelCoordBase	= PageCoord * VoxelRasterPass.VirtualVoxel.PageResolution;
-				const uint3 PageVoxelCoord		= PageVoxelCoordBase + VoxelCoordOffset;
-
-				const float HairCount			= HairCoverage;
-				const float VoxelFixPointScale	= GetVoxelDensityFixPointScale();
-				InterlockedAdd(VoxelRasterPass.DensityTexture[PageVoxelCoord], HairCount * VoxelFixPointScale);
-
-				#if DEBUG_OUTPUT
-				OutColor = 0xFF;
-				#endif
-				return;
-			}
-			else
-			{
-				#if DEBUG_OUTPUT
-				OutColor = 0x0F;
-				#endif
-				return;
-			}
-		}
-	}
-
-	#if 0 // This code is no longer used, but it is a reference for voxelizing material
-	const int3 VoxelCoord = (WorldPosition.xyz - VoxelRasterPass.VoxelMinAABB) / (VoxelRasterPass.VoxelMaxAABB - VoxelRasterPass.VoxelMinAABB) * VoxelRasterPass.VoxelResolution;
-		 
-	if (VoxelCoord.x >= 0 || VoxelCoord.y >= 0 || VoxelCoord.z >= 0 ||
-		VoxelCoord.x < int(VoxelRasterPass.VoxelResolution.x) || 
-		VoxelCoord.y < int(VoxelRasterPass.VoxelResolution.y) || 
-		VoxelCoord.z < int(VoxelRasterPass.VoxelResolution.z))
-	{
-		const float HairCount = HairCoverage;
-		const float VoxelFixPointScale = GetVoxelDensityFixPointScale();
-		InterlockedAdd(VoxelRasterPass.DensityTexture[VoxelCoord], HairCount * VoxelFixPointScale);
-
-		#if SUPPORT_TANGENT_PROPERTY == 1
-		const half3 Tangent = GetTangentToWorld2(Interpolants).xyz;
-		const uint3 Tangent8bits = QuantizeTo8Bits(Tangent);
-
-		InterlockedAdd(VoxelRasterPass.TangentXTexture[VoxelCoord], Tangent8bits.x);
-		InterlockedAdd(VoxelRasterPass.TangentYTexture[VoxelCoord], Tangent8bits.y);
-		InterlockedAdd(VoxelRasterPass.TangentZTexture[VoxelCoord], Tangent8bits.z);
-		#endif	
-
-		#if SUPPORT_MATERIAL_PROPERTY == 1
-		{
-			ResolvedView = ResolveView();
-
-			FMaterialPixelParameters MaterialParameters = GetMaterialPixelParameters(Interpolants, SvPosition);
-			FPixelMaterialInputs PixelMaterialInputs;
-
-			float4 ScreenPosition = SvPositionToResolvedScreenPosition(SvPosition);
-			float3 TranslatedWorldPosition = SvPositionToResolvedTranslatedWorld(SvPosition);
-			CalcMaterialParametersEx(MaterialParameters, PixelMaterialInputs, SvPosition, ScreenPosition, true, TranslatedWorldPosition, TranslatedWorldPosition);
-
-			// For material we take the max color / roughness for efficiency resource instead of accumulated differences
-			const float3 BaseColor = GetMaterialBaseColor(PixelMaterialInputs);
-			const float Roughness = GetMaterialRoughness(PixelMaterialInputs);
-			const uint EncodedBaseColorAndRoughness = EncodeBaseColorAndRoughness(BaseColor, Roughness);
-			InterlockedMax(VoxelRasterPass.MaterialTexture[VoxelCoord], EncodedBaseColorAndRoughness);
-		}
-		#endif
-
-	}
-
-	#if DEBUG_OUTPUT
-	OutColor = 0xFF;
-	#endif
-	#endif
-
-#endif
 }
\ No newline at end of file
diff --git a/Engine/Shaders/Private/HairStrands/HairStrandsDeepShadowVS.usf b/Engine/Shaders/Private/HairStrands/HairStrandsDeepShadowVS.usf
--- a/Engine/Shaders/Private/HairStrands/HairStrandsDeepShadowVS.usf
+++ b/Engine/Shaders/Private/HairStrands/HairStrandsDeepShadowVS.usf
@@ -11,15 +11,6 @@
 
 #define MESH_RENDER_DEPTH 0
 #define MESH_RENDER_DOM   1
-#define MESH_RENDER_VOXEL 2
-
-#ifndef SUPPORT_TANGENT_PROPERTY
-#define SUPPORT_TANGENT_PROPERTY 0
-#endif
-
-#ifndef SUPPORT_MATERIAL_PROPERTY
-#define SUPPORT_MATERIAL_PROPERTY 0
-#endif
 
 struct FDeepShadowAccumulateVSToPS
 {
@@ -27,9 +18,6 @@ struct FDeepShadowAccumulateVSToPS
 	FVertexFactoryInterpolantsVSToPS Interpolants;
 #endif
 	float4 Position : SV_POSITION;
-#if MESH_RENDER_MODE == MESH_RENDER_VOXEL
-	float4 WorldPositon_Coverage : TEXCOORD6;
-#endif
 #if MESH_RENDER_MODE == MESH_RENDER_DOM
 	float HairCoverage : CUSTOM_COVERAGE;
 #endif
@@ -46,57 +34,12 @@ void Main(
 	OPTIONAL_VertexID
 	out VS_OUTPUT_TYPE Output
 )
-#if defined(HAIR_STRAND_MESH_FACTORY) && MESH_RENDER_MODE == MESH_RENDER_VOXEL && SUPPORT_TANGENT_PROPERTY == 0 && SUPPORT_MATERIAL_PROPERTY == 0
-{
-	ResolvedView = ResolveView();
-	
-	const FHairRenderInfo HairRenderInfo = GetHairRenderInfo(ResolvedView.HairRenderInfo, ResolvedView.HairRenderInfoBits);
-
-	FHairViewInfo HairViewInfo;
-	HairViewInfo.TranslatedWorldCameraOrigin = ResolvedView.TranslatedWorldCameraOrigin;
-	HairViewInfo.RadiusAtDepth1 = HairRenderInfo.RadiusAtDepth1Primary;
-	HairViewInfo.bIsOrthoView = true;
-	HairViewInfo.ViewForward = ResolvedView.ViewForward;
-
-	float2 ViewportScale;
-	float4x4 WorldToClipMatrix = VoxelRasterPass.WorldToClipMatrix;
-	if (HairRenderInfo.bIsGPUDriven)
-	{
-		const uint MacroGroupId = VoxelRasterPass.MacroGroupId;
-		FVoxelizationViewInfo VoxelizationViewInfo = VoxelRasterPass.VoxelizationViewInfoBuffer[MacroGroupId];
-		WorldToClipMatrix = VoxelizationViewInfo.WorldToClip;
-		ViewportScale = VoxelizationViewInfo.RasterResolution / float2(VoxelRasterPass.ViewportResolution);
-		HairViewInfo.ViewForward = VoxelizationViewInfo.ViewForward;
-	}
-
-	float HairWorldRadius = 0;
-	float HairDensity = 0;
-	float4 WorldPosition = VertexFactoryGetWorldPosition_Voxelization(Input, HairViewInfo, HairWorldRadius, HairDensity);
-
-	WorldPosition.xyz -= ResolvedView.PreViewTranslation.xyz; // Remove the camera relative position (todo move this into matrix transform)
-
-	Output.Position = mul(WorldPosition, WorldToClipMatrix);
-
-	if (HairRenderInfo.bIsGPUDriven)
-	{
-		Output.Position.xy += float2(1, 1);
-		Output.Position.xy *= ViewportScale;
-		Output.Position.xy -= float2(1, 1);
-	}
-
-	const float PixelRadius = HairRenderInfo.RadiusAtDepth1Primary;
-	const float Coverage = HairWorldRadius / max(HairWorldRadius, PixelRadius);
-	Output.WorldPositon_Coverage = float4(WorldPosition.xyz, Coverage * HairDensity);
-}
-#else // SPARSE VOXELIZAITON SPECIALIZATION
 {
 	ResolvedView = ResolveView();
 
 	FVertexFactoryIntermediates VFIntermediates = GetVertexFactoryIntermediates(Input);
 
-#if MESH_RENDER_MODE == MESH_RENDER_VOXEL
-	float4x4 WorldToClipMatrix = VoxelRasterPass.WorldToClipMatrix;
-#elif MESH_RENDER_MODE == MESH_RENDER_DEPTH || MESH_RENDER_MODE == MESH_RENDER_DOM
+#if MESH_RENDER_MODE == MESH_RENDER_DEPTH || MESH_RENDER_MODE == MESH_RENDER_DOM
 	float4x4 WorldToClipMatrix = DeepRasterPass.CPU_WorldToClipMatrix;
 #endif
 
@@ -109,18 +52,6 @@ void Main(
 	HairViewInfo.RadiusAtDepth1 = HairRenderInfo.RadiusAtDepth1Primary;
 	HairViewInfo.bIsOrthoView = HairRenderInfo.bIsOrthoView;
 
-	#if MESH_RENDER_MODE == MESH_RENDER_VOXEL
-	float2 ViewportScale;
-	if (HairRenderInfo.bIsGPUDriven)
-	{
-		const uint MacroGroupId = VoxelRasterPass.MacroGroupId;
-		FVoxelizationViewInfo VoxelizationViewInfo = VoxelRasterPass.VoxelizationViewInfoBuffer[MacroGroupId];
-		WorldToClipMatrix			= VoxelizationViewInfo.WorldToClip;
-		HairViewInfo.ViewForward	= VoxelizationViewInfo.ViewForward;
-		ViewportScale				= VoxelizationViewInfo.RasterResolution / float2(VoxelRasterPass.ViewportResolution);
-	}
-	#endif
-
 	#if MESH_RENDER_MODE == MESH_RENDER_DEPTH || MESH_RENDER_MODE == MESH_RENDER_DOM
 	if (HairRenderInfo.bIsGPUDriven)
 	{
@@ -140,23 +71,6 @@ void Main(
 
 	Output.Position = mul(WorldPosition, WorldToClipMatrix);
 
-#if MESH_RENDER_MODE == MESH_RENDER_VOXEL && defined(HAIR_STRAND_MESH_FACTORY)
-	if (HairRenderInfo.bIsGPUDriven)
-	{
-		Output.Position.xy += float2(1, 1);
-		Output.Position.xy *= ViewportScale;
-		Output.Position.xy -= float2(1, 1);
-	}
-
-	Output.WorldPositon_Coverage.xyz = WorldPosition;
-
-  #if SUPPORT_TANGENT_PROPERTY == 1 || SUPPORT_MATERIAL_PROPERTY == 1
-	float3x3 TangentToLocal = VertexFactoryGetTangentToLocal(Input, VFIntermediates);
-	FMaterialVertexParameters VertexParameters = GetMaterialVertexParameters(Input, VFIntermediates, WorldPosition.xyz, TangentToLocal);
-	Output.Interpolants = VertexFactoryGetInterpolantsVSToPS(Input, VFIntermediates, VertexParameters);
-  #endif
-#endif
-
 	float HairCoverage = 1;
   #ifdef HAIR_STRAND_MESH_FACTORY	
 	const float CurrentDepth = Output.Position.z / Output.Position.w;
@@ -170,12 +84,7 @@ void Main(
 	Output.HairCoverage = HairCoverage;
   #endif
 
-  #if MESH_RENDER_MODE == MESH_RENDER_VOXEL
-	Output.WorldPositon_Coverage.w = HairCoverage;
-  #endif
-
 	OutputVertexID( Output );
 }
-#endif // SPARSE VOXELIZAITON SPECIALIZATION
 
 #endif // VERTEXSHADER
diff --git a/Engine/Source/Runtime/Renderer/Private/HairStrands/HairStrandsRasterCommon.cpp b/Engine/Source/Runtime/Renderer/Private/HairStrands/HairStrandsRasterCommon.cpp
--- a/Engine/Source/Runtime/Renderer/Private/HairStrands/HairStrandsRasterCommon.cpp
+++ b/Engine/Source/Runtime/Renderer/Private/HairStrands/HairStrandsRasterCommon.cpp
@@ -12,7 +12,6 @@
 #include "ScenePrivate.h"
 
 IMPLEMENT_STATIC_UNIFORM_BUFFER_STRUCT(FHairDeepShadowRasterUniformParameters, "DeepRasterPass", SceneTextures);
-IMPLEMENT_STATIC_UNIFORM_BUFFER_STRUCT(FHairVoxelizationRasterUniformParameters, "VoxelRasterPass", SceneTextures);
 
 /////////////////////////////////////////////////////////////////////////////////////////
 
@@ -80,46 +79,6 @@ IMPLEMENT_MATERIAL_SHADER_TYPE(, FDeepShadowDomMeshVS, TEXT("/Engine/Private/Hai
 
 /////////////////////////////////////////////////////////////////////////////////////////
 
-template<bool bVoxelizeMaterial, bool bClusterCulling>
-class FVoxelMeshVS : public FMeshMaterialShader
-{
-	DECLARE_SHADER_TYPE(FVoxelMeshVS, MeshMaterial);
-
-protected:
-
-	FVoxelMeshVS(const ShaderMetaType::CompiledShaderInitializerType& Initializer)
-		: FMeshMaterialShader(Initializer)
-	{
-		ERHIFeatureLevel::Type FeatureLevel = GetMaxSupportedFeatureLevel((EShaderPlatform)Initializer.Target.Platform);
-		check(FSceneInterface::GetShadingPath(FeatureLevel) != EShadingPath::Mobile);
-	}
-
-	FVoxelMeshVS() {}
-
-	static bool ShouldCompilePermutation(const FMeshMaterialShaderPermutationParameters& Parameters)
-	{
-		return IsCompatibleWithHairStrands(Parameters.Platform, Parameters.MaterialParameters) && Parameters.VertexFactoryType->GetFName() == FName(TEXT("FHairStrandsVertexFactory"));;
-	}
-
-	static void ModifyCompilationEnvironment(const FMaterialShaderPermutationParameters& Parameters, FShaderCompilerEnvironment& OutEnvironment)
-	{
-		// Note: at the moment only the plain voxelization support material voxelization
-		FMaterialShader::ModifyCompilationEnvironment(Parameters, OutEnvironment);
-		OutEnvironment.SetDefine(TEXT("MESH_RENDER_MODE"), 2);
-		OutEnvironment.SetDefine(TEXT("SUPPORT_TANGENT_PROPERTY"), bVoxelizeMaterial ? 1 : 0);
-		OutEnvironment.SetDefine(TEXT("SUPPORT_MATERIAL_PROPERTY"), bVoxelizeMaterial ? 1 : 0);
-		OutEnvironment.SetDefine(TEXT("USE_CULLED_CLUSTER"), bClusterCulling ? 1 : 0);
-	}
-};
-
-typedef FVoxelMeshVS<false, false> TVoxelMeshVS_NoMaterial_NoCluster;
-typedef FVoxelMeshVS<false, true>  TVoxelMeshVS_NoMaterial_Cluster;
-
-IMPLEMENT_MATERIAL_SHADER_TYPE(template<>, TVoxelMeshVS_NoMaterial_NoCluster, TEXT("/Engine/Private/HairStrands/HairStrandsDeepShadowVS.usf"), TEXT("Main"), SF_Vertex);
-IMPLEMENT_MATERIAL_SHADER_TYPE(template<>, TVoxelMeshVS_NoMaterial_Cluster, TEXT("/Engine/Private/HairStrands/HairStrandsDeepShadowVS.usf"), TEXT("Main"), SF_Vertex);
-
-/////////////////////////////////////////////////////////////////////////////////////////
-
 class FDeepShadowDepthMeshPS : public FMeshMaterialShader
 {
 	DECLARE_SHADER_TYPE(FDeepShadowDepthMeshPS, MeshMaterial);
@@ -168,44 +127,6 @@ IMPLEMENT_MATERIAL_SHADER_TYPE(, FDeepShadowDomMeshPS, TEXT("/Engine/Private/Hai
 
 /////////////////////////////////////////////////////////////////////////////////////////
 
-enum class EVoxelMeshPSType
-{
-	Density,
-	Material
-};
-template<EVoxelMeshPSType VoxelizationType>
-class FVoxelMeshPS : public FMeshMaterialShader
-{
-	DECLARE_SHADER_TYPE(FVoxelMeshPS, MeshMaterial);
-
-public:
-
-	FVoxelMeshPS(const ShaderMetaType::CompiledShaderInitializerType& Initializer)
-		: FMeshMaterialShader(Initializer)
-	{
-		ERHIFeatureLevel::Type FeatureLevel = GetMaxSupportedFeatureLevel((EShaderPlatform)Initializer.Target.Platform);
-		check(FSceneInterface::GetShadingPath(FeatureLevel) != EShadingPath::Mobile);
-	}
-
-	FVoxelMeshPS() {}
-
-	static bool ShouldCompilePermutation(const FMeshMaterialShaderPermutationParameters& Parameters)
-	{
-		return IsCompatibleWithHairStrands(Parameters.Platform, Parameters.MaterialParameters) && Parameters.VertexFactoryType->GetFName() == FName(TEXT("FHairStrandsVertexFactory"));
-	}
-
-	static void ModifyCompilationEnvironment(const FMaterialShaderPermutationParameters& Parameters, FShaderCompilerEnvironment& OutEnvironment)
-	{
-		FMaterialShader::ModifyCompilationEnvironment(Parameters, OutEnvironment);
-		OutEnvironment.SetDefine(TEXT("SUPPORT_TANGENT_PROPERTY"), VoxelizationType == EVoxelMeshPSType::Material ? 1 : 0);
-		OutEnvironment.SetDefine(TEXT("SUPPORT_MATERIAL_PROPERTY"), VoxelizationType == EVoxelMeshPSType::Material ? 1 : 0);
-	}
-};
-IMPLEMENT_MATERIAL_SHADER_TYPE(template<>, FVoxelMeshPS<EVoxelMeshPSType::Density>, TEXT("/Engine/Private/HairStrands/HairStrandsDeepShadowPS.usf"), TEXT("MainVoxel"), SF_Pixel);
-IMPLEMENT_MATERIAL_SHADER_TYPE(template<>, FVoxelMeshPS<EVoxelMeshPSType::Material>, TEXT("/Engine/Private/HairStrands/HairStrandsDeepShadowPS.usf"), TEXT("MainVoxel"), SF_Pixel);
-
-/////////////////////////////////////////////////////////////////////////////////////////
-
 class FHairRasterMeshProcessor : public FMeshPassProcessor
 {
 public:
@@ -264,15 +185,11 @@ void FHairRasterMeshProcessor::AddMeshBatch(const FMeshBatch& RESTRICT MeshBatch
 			Process<FDeepShadowDepthMeshVS, FDeepShadowDepthMeshPS>(MeshBatch, BatchElementMask, PrimitiveSceneProxy, StaticMeshId, MaterialRenderProxy, Material, MeshFillMode, MeshCullMode);
 		else if (RasterPassType == EHairStrandsRasterPassType::DeepOpacityMap)
 			Process<FDeepShadowDomMeshVS, FDeepShadowDomMeshPS>(MeshBatch, BatchElementMask, PrimitiveSceneProxy, StaticMeshId, MaterialRenderProxy, Material, MeshFillMode, MeshCullMode);
-		else if (RasterPassType == EHairStrandsRasterPassType::VoxelizationVirtual && bCullingEnable)
-			Process<FVoxelMeshVS<false, true>, FVoxelMeshPS<EVoxelMeshPSType::Density>>(MeshBatch, BatchElementMask, PrimitiveSceneProxy, StaticMeshId, MaterialRenderProxy, Material, MeshFillMode, MeshCullMode);
-		else if (RasterPassType == EHairStrandsRasterPassType::VoxelizationVirtual && !bCullingEnable)
-			Process<FVoxelMeshVS<false, false>, FVoxelMeshPS<EVoxelMeshPSType::Density>>(MeshBatch, BatchElementMask, PrimitiveSceneProxy, StaticMeshId, MaterialRenderProxy, Material, MeshFillMode, MeshCullMode);
 	}
 }
 
-// Vertex is either FDeepShadowDepthMeshVS, FDeepShadowDomMeshVS, or FVoxelMeshVS
-// Pixel  is either FDeepShadowDepthMeshPS, FDeepShadowDomMeshPS, or FVoxelMeshPS
+// Vertex is either FDeepShadowDepthMeshVS or FDeepShadowDomMeshVS
+// Pixel  is either FDeepShadowDepthMeshPS or FDeepShadowDomMeshPS
 template<typename VertexShaderType, typename PixelShaderType>
 void FHairRasterMeshProcessor::Process(
 	const FMeshBatch& MeshBatch,
@@ -354,7 +271,6 @@ void AddHairStrandsRasterPass(
 		{
 		case EHairStrandsRasterPassType::DeepOpacityMap:		return RDG_EVENT_NAME("HairStrandsRasterDeepOpacityMap");
 		case EHairStrandsRasterPassType::FrontDepth:			return RDG_EVENT_NAME("HairStrandsRasterFrontDepth");
-		case EHairStrandsRasterPassType::VoxelizationVirtual:	return RDG_EVENT_NAME("HairStrandsRasterVoxelizationVirtual");
 		default:												return RDG_EVENT_NAME("Noname");
 		}
 	};
@@ -396,12 +312,6 @@ void AddHairStrandsRasterPass(
 				CW_RGBA, BO_Add, BF_One, BF_Zero, BO_Add, BF_One, BF_Zero>::GetRHI());
 			DrawRenderState.SetDepthStencilState(TStaticDepthStencilState<true, CF_DepthNearOrEqual>::GetRHI());
 		}
-		else if (RasterPassType == EHairStrandsRasterPassType::VoxelizationVirtual)
-		{
-			DrawRenderState.SetBlendState(TStaticBlendState<
-				CW_RGBA, BO_Add, BF_One, BF_Zero, BO_Add, BF_One, BF_Zero>::GetRHI());
-			DrawRenderState.SetDepthStencilState(TStaticDepthStencilState<false, CF_Always>::GetRHI());
-		}
 
 		FDynamicMeshDrawCommandStorage DynamicMeshDrawCommandStorage; // << Were would thid be stored?
 		FMeshCommandOneFrameArray VisibleMeshDrawCommands;
@@ -455,30 +365,4 @@ void AddHairDeepShadowRasterPass(
 		LightDirection,
 		PassParameters,
 		InstanceCullingManager);
-}
-
-void AddHairVoxelizationRasterPass(
-	FRDGBuilder& GraphBuilder,
-	const FScene* Scene,
-	const FViewInfo* ViewInfo,
-	const FHairStrandsMacroGroupData::TPrimitiveInfos& PrimitiveSceneInfos,
-	const FIntRect& ViewportRect,
-	const FVector4& HairRenderInfo,
-	const uint32 HairRenderInfoBits,
-	const FVector& RasterDirection,
-	FHairVoxelizationRasterPassParameters* PassParameters,
-	FInstanceCullingManager& InstanceCullingManager)
-{
-	AddHairStrandsRasterPass<FHairVoxelizationRasterPassParameters>(
-		GraphBuilder, 
-		Scene, 
-		ViewInfo, 
-		PrimitiveSceneInfos, 
-		EHairStrandsRasterPassType::VoxelizationVirtual,
-		ViewportRect, 
-		HairRenderInfo, 
-		HairRenderInfoBits,
-		RasterDirection,
-		PassParameters,
-		InstanceCullingManager);
 }
\ No newline at end of file
diff --git a/Engine/Source/Runtime/Renderer/Private/HairStrands/HairStrandsRasterCommon.h b/Engine/Source/Runtime/Renderer/Private/HairStrands/HairStrandsRasterCommon.h
--- a/Engine/Source/Runtime/Renderer/Private/HairStrands/HairStrandsRasterCommon.h
+++ b/Engine/Source/Runtime/Renderer/Private/HairStrands/HairStrandsRasterCommon.h
@@ -9,15 +9,13 @@
 #include "CoreMinimal.h"
 #include "RendererInterface.h"
 #include "HairStrandsCluster.h"
-#include "HairStrandsVoxelization.h"
 
 class FViewInfo;
 
 enum class EHairStrandsRasterPassType : uint8
 {
 	FrontDepth,
-	DeepOpacityMap,
-	VoxelizationVirtual
+	DeepOpacityMap
 };
 
 // ////////////////////////////////////////////////////////////////
@@ -53,36 +51,3 @@ void AddHairDeepShadowRasterPass(
 	const FVector& LightDirection,
 	FHairDeepShadowRasterPassParameters* PassParameters,
 	FInstanceCullingManager& InstanceCullingManager);
-
-// ////////////////////////////////////////////////////////////////
-// Voxelization raster pass
-BEGIN_GLOBAL_SHADER_PARAMETER_STRUCT(FHairVoxelizationRasterUniformParameters, )
-	SHADER_PARAMETER_STRUCT(FHairStrandsVoxelCommonParameters, VirtualVoxel)
-	SHADER_PARAMETER(FMatrix, WorldToClipMatrix)
-	SHADER_PARAMETER(FVector, VoxelMinAABB)
-	SHADER_PARAMETER(FVector, VoxelMaxAABB)
-	SHADER_PARAMETER(FIntVector, VoxelResolution)
-	SHADER_PARAMETER(uint32, MacroGroupId)
-	SHADER_PARAMETER(FIntPoint, ViewportResolution)
-	SHADER_PARAMETER_RDG_BUFFER_SRV(StructuredBuffer<FVoxelizationViewInfo>, VoxelizationViewInfoBuffer)
-	SHADER_PARAMETER_RDG_TEXTURE_UAV(RWTexture3D<uint>, DensityTexture)
-END_GLOBAL_SHADER_PARAMETER_STRUCT()
-
-BEGIN_SHADER_PARAMETER_STRUCT(FHairVoxelizationRasterPassParameters, )
-	SHADER_PARAMETER_STRUCT_REF(FViewUniformShaderParameters, View)
-	SHADER_PARAMETER_STRUCT_INCLUDE(FInstanceCullingDrawParams, InstanceCullingDrawParams)
-	SHADER_PARAMETER_RDG_UNIFORM_BUFFER(FHairVoxelizationRasterUniformParameters, UniformBuffer)
-	RENDER_TARGET_BINDING_SLOTS()
-END_SHADER_PARAMETER_STRUCT()
-
-void AddHairVoxelizationRasterPass(
-	FRDGBuilder& GraphBuilder,
-	const FScene* Scene,
-	const FViewInfo* ViewInfo,
-	const FHairStrandsMacroGroupData::TPrimitiveInfos& PrimitiveSceneInfos,
-	const FIntRect& ViewportRect,
-	const FVector4& HairRenderInfo,
-	const uint32 HairRenderInfoBits,
-	const FVector& RasterDirection,
-	FHairVoxelizationRasterPassParameters* PassParameters,
-	FInstanceCullingManager& InstanceCullingManager);
diff --git a/Engine/Source/Runtime/Renderer/Private/HairStrands/HairStrandsVoxelization.cpp b/Engine/Source/Runtime/Renderer/Private/HairStrands/HairStrandsVoxelization.cpp
--- a/Engine/Source/Runtime/Renderer/Private/HairStrands/HairStrandsVoxelization.cpp
+++ b/Engine/Source/Runtime/Renderer/Private/HairStrands/HairStrandsVoxelization.cpp
@@ -99,9 +99,6 @@ static const FIntPoint GPUDrivenViewportResolution = FIntPoint(4096, 4096);
 static int32 GHairVirtualVoxelInvalidEmptyPageIndex = 1;
 static FAutoConsoleVariableRef CVarHairVirtualVoxelInvalidEmptyPageIndex(TEXT("r.HairStrands.Voxelization.Virtual.InvalidateEmptyPageIndex"), GHairVirtualVoxelInvalidEmptyPageIndex, TEXT("Invalid voxel page index which does not contain any voxelized data."));
 
-static int32 GHairStrandsVoxelComputeRaster = 1;
-static FAutoConsoleVariableRef CVarHairStrandsVoxelComputeRaster(TEXT("r.HairStrands.Voxelization.Virtual.ComputeRaster"), GHairStrandsVoxelComputeRaster, TEXT("Use compute for rasterizing voxeliation (faster)."));
-
 static int32 GHairStrandsVoxelComputeRasterMaxVoxelCount = 32;
 static FAutoConsoleVariableRef CVarHairStrandsVoxelComputeRasterMaxVoxelCount(TEXT("r.HairStrands.Voxelization.Virtual.ComputeRasterMaxVoxelCount"), GHairStrandsVoxelComputeRasterMaxVoxelCount, TEXT("Max number of voxel which are rasterized for a given hair segment. This is for debug purpose only."));
 
@@ -1185,38 +1182,7 @@ static void AddVirtualVoxelizationRasterPass(
 		MacroGroup.VirtualVoxelNodeDesc.WorldToClip = WorldToClip;
 	}
 
-	const bool bUseComputeRaster = GHairStrandsVoxelComputeRaster > 0;
-	if (bIsGPUDriven && bUseComputeRaster)
-	{	
-		AddVirtualVoxelizationComputeRasterPass(GraphBuilder, ViewInfo, VoxelResources, MacroGroup);
-		return;
-	}
-
-	FHairVoxelizationRasterPassParameters* PassParameters = GraphBuilder.AllocParameters<FHairVoxelizationRasterPassParameters>();
-
-	{
-		FHairVoxelizationRasterUniformParameters* UniformParameters = GraphBuilder.AllocParameters<FHairVoxelizationRasterUniformParameters>();
-
-		UniformParameters->VirtualVoxel = VoxelResources.Parameters.Common;
-		UniformParameters->WorldToClipMatrix = WorldToClip;
-		UniformParameters->VoxelMinAABB = MacroGroup.VirtualVoxelNodeDesc.WorldMinAABB;
-		UniformParameters->VoxelMaxAABB = MacroGroup.VirtualVoxelNodeDesc.WorldMaxAABB;
-		UniformParameters->VoxelResolution = TotalVoxelResolution; // i.e., the virtual resolution
-		UniformParameters->MacroGroupId = MacroGroup.MacroGroupId;
-		UniformParameters->ViewportResolution = RasterResolution;
-		UniformParameters->VoxelizationViewInfoBuffer = GraphBuilder.CreateSRV(VoxelResources.VoxelizationViewInfoBuffer);
-		UniformParameters->DensityTexture = GraphBuilder.CreateUAV(VoxelResources.PageTexture);
-
-		PassParameters->UniformBuffer = GraphBuilder.CreateUniformBuffer(UniformParameters);
-	}
-
-	// For debug purpose
-	#if 0
-	FRDGTextureRef DebugOutputTexture = GraphBuilder.CreateTexture(FRDGTextureDesc::Create2D(RasterResolution, PF_R32_UINT, FClearValueBinding::Black, TexCreate_RenderTargetable), TEXT("Hair.DummyTexture"));
-	PassParameters->RenderTargets[0] = FRenderTargetBinding(DebugOutputTexture, ERenderTargetLoadAction::EClear);
-	#endif
-
-	AddHairVoxelizationRasterPass(GraphBuilder, Scene, ViewInfo, PrimitiveSceneInfo, ViewportRect, HairRenderInfo, HairRenderInfoBits, RasterDirection, PassParameters, InstanceCullingManager);
+	AddVirtualVoxelizationComputeRasterPass(GraphBuilder, ViewInfo, VoxelResources, MacroGroup);
 }
 
 ///////////////////////////////////////////////////////////////////////////////////////////////////
