Implements r.TemporalAA.Debug.SetupExtraPasses

diff --git a/Engine/Source/Runtime/Renderer/Private/PostProcess/TemporalAA.cpp b/Engine/Source/Runtime/Renderer/Private/PostProcess/TemporalAA.cpp
--- a/Engine/Source/Runtime/Renderer/Private/PostProcess/TemporalAA.cpp
+++ b/Engine/Source/Runtime/Renderer/Private/PostProcess/TemporalAA.cpp
@@ -96,6 +96,15 @@ TAutoConsoleVariable<int32> CVarTAAEnableAntiInterference(
 	TEXT("Enable heuristic to detect geometric interference between input pixel grid alignement and structured geometry."),
 	ECVF_RenderThreadSafe);
 
+#if COMPILE_TAA_DEBUG_PASSES
+
+TAutoConsoleVariable<int32> CVarTAASetupDebugPasses(
+	TEXT("r.TemporalAA.Debug.SetupExtraPasses"), 0,
+	TEXT("Whether to enable the debug passes"),
+	ECVF_RenderThreadSafe);
+
+#endif
+
 inline bool DoesPlatformSupportTemporalHistoryUpscale(EShaderPlatform Platform)
 {
 	return (IsPCPlatform(Platform) || FDataDrivenShaderPlatformInfo::GetSupportsTemporalHistoryUpscale(Platform))
@@ -1119,6 +1128,10 @@ static void AddGen5MainTemporalAAPasses(
 		ensureMsgf(bUsePreExposure, TEXT("r.TemporalAA.Algorithm=1 requires r.UsePreExposure=1"));
 	}
 
+#if COMPILE_TAA_DEBUG_PASSES
+	const bool bSetupDebugPasses = CVarTAASetupDebugPasses.GetValueOnRenderThread() != 0;
+#endif
+
 	// Whether to use camera cut shader permutation or not.
 	bool bCameraCut = !InputHistory.IsValid() || View.bCameraCut;
 
@@ -1825,6 +1838,7 @@ static void AddGen5MainTemporalAAPasses(
 
 	// Debug the history.
 	#if COMPILE_TAA_DEBUG_PASSES
+	if (bSetupDebugPasses)
 	{
 		const int32 kHistoryUpscalingFactor = 2;
 
