Moved ray tracing scene scratch buffer (RDG resource) to FRayTracingScene, initialized in BeginCreate(). Added reference to it in WaitForRayTracingScene(), which is where the ray tracing scene build completes and scene is transitioned to readable state.

diff --git a/Engine/Source/Runtime/Renderer/Private/DeferredShadingRenderer.cpp b/Engine/Source/Runtime/Renderer/Private/DeferredShadingRenderer.cpp
--- a/Engine/Source/Runtime/Renderer/Private/DeferredShadingRenderer.cpp
+++ b/Engine/Source/Runtime/Renderer/Private/DeferredShadingRenderer.cpp
@@ -1236,7 +1236,7 @@ bool FDeferredShadingSceneRenderer::DispatchRayTracingWorldUpdates(FRDGBuilder&
 
 	DeduplicateRayGenerationShaders(RayGenShaders);
 
-	RayTracingScene.BeginCreate();
+	RayTracingScene.BeginCreate(GraphBuilder);
 
 	if (RayGenShaders.Num())
 	{
@@ -1320,12 +1320,7 @@ bool FDeferredShadingSceneRenderer::DispatchRayTracingWorldUpdates(FRDGBuilder&
 			RDG_GPU_STAT_SCOPE(GraphBuilder, RayTracingScene);
 
 			FBuildAccelerationStructurePassParams* PassParams = GraphBuilder.AllocParameters<FBuildAccelerationStructurePassParams>();
-			FRDGBufferDesc ScratchBufferDesc;
-			ScratchBufferDesc.UnderlyingType = FRDGBufferDesc::EUnderlyingType::StructuredBuffer;
-			ScratchBufferDesc.Usage = BUF_UnorderedAccess;
-			ScratchBufferDesc.BytesPerElement = 4;
-			ScratchBufferDesc.NumElements = Scene->RayTracingScene.SizeInfo.BuildScratchSize / 4;
-			PassParams->RayTracingSceneScratchBuffer = GraphBuilder.CreateBuffer(ScratchBufferDesc, TEXT("RayTracingSceneScratchBuffer"));
+			PassParams->RayTracingSceneScratchBuffer = Scene->RayTracingScene.BuildScratchBuffer;
 
 			GraphBuilder.AddPass(RDG_EVENT_NAME("RayTracingScene"), PassParams, ERDGPassFlags::Compute | ERDGPassFlags::NeverCull,
 				[this, PassParams](FRHICommandList& RHICmdList)
@@ -1408,7 +1403,12 @@ void FDeferredShadingSceneRenderer::WaitForRayTracingScene(FRDGBuilder& GraphBui
 
 	RDG_GPU_MASK_SCOPE(GraphBuilder, FRHIGPUMask::All());
 
-	AddPass(GraphBuilder, [this](FRHICommandListImmediate& RHICmdList)
+	// Scratch buffer must be referenced in this pass, as it must live until the BVH build is complete.
+	FBuildAccelerationStructurePassParams* PassParams = GraphBuilder.AllocParameters<FBuildAccelerationStructurePassParams>();
+	PassParams->RayTracingSceneScratchBuffer = Scene->RayTracingScene.BuildScratchBuffer;
+
+	GraphBuilder.AddPass(RDG_EVENT_NAME("WaitForRayTracingScene"), PassParams, ERDGPassFlags::Compute | ERDGPassFlags::NeverCull,
+		[this, PassParams](FRHICommandListImmediate& RHICmdList)
 	{
 		const int32 ReferenceViewIndex = 0;
 		FViewInfo& ReferenceView = Views[ReferenceViewIndex];
diff --git a/Engine/Source/Runtime/Renderer/Private/RayTracing/RayTracingScene.cpp b/Engine/Source/Runtime/Renderer/Private/RayTracing/RayTracingScene.cpp
--- a/Engine/Source/Runtime/Renderer/Private/RayTracing/RayTracingScene.cpp
+++ b/Engine/Source/Runtime/Renderer/Private/RayTracing/RayTracingScene.cpp
@@ -6,6 +6,7 @@
 
 #include "RenderCore.h"
 #include "RayTracingDefinitions.h"
+#include "RenderGraphBuilder.h"
 
 FRayTracingScene::FRayTracingScene()
 {
@@ -17,7 +18,7 @@ FRayTracingScene::~FRayTracingScene()
 	WaitForTasks();
 }
 
-FGraphEventRef FRayTracingScene::BeginCreate()
+FGraphEventRef FRayTracingScene::BeginCreate(FRDGBuilder& GraphBuilder)
 {
 	WaitForTasks();
 
@@ -65,6 +66,14 @@ FGraphEventRef FRayTracingScene::BeginCreate()
 
 		}, TStatId(), nullptr, ENamedThreads::AnyThread);
 
+	FRDGBufferDesc ScratchBufferDesc;
+	ScratchBufferDesc.UnderlyingType = FRDGBufferDesc::EUnderlyingType::StructuredBuffer;
+	ScratchBufferDesc.Usage = BUF_UnorderedAccess;
+	ScratchBufferDesc.BytesPerElement = 4;
+	ScratchBufferDesc.NumElements = SizeInfo.BuildScratchSize / 4;
+
+	BuildScratchBuffer = GraphBuilder.CreateBuffer(ScratchBufferDesc, TEXT("FRayTracingScene::BuildScratchBuffer"));
+
 	return CreateRayTracingSceneTask;
 }
 
@@ -123,6 +132,8 @@ void FRayTracingScene::Reset()
 #ifdef DO_CHECK
 	GeometriesToValidate.Reset();
 #endif
+
+	BuildScratchBuffer = nullptr;
 }
 
 void FRayTracingScene::ResetAndReleaseResources()
diff --git a/Engine/Source/Runtime/Renderer/Private/RayTracing/RayTracingScene.h b/Engine/Source/Runtime/Renderer/Private/RayTracing/RayTracingScene.h
--- a/Engine/Source/Runtime/Renderer/Private/RayTracing/RayTracingScene.h
+++ b/Engine/Source/Runtime/Renderer/Private/RayTracing/RayTracingScene.h
@@ -9,10 +9,12 @@
 
 #include "Async/TaskGraphInterfaces.h"
 #include "RHI.h"
+#include "RenderGraphResources.h"
 
 class FRHIRayTracingScene;
 class FRHIShaderResourceView;
 class FRayTracingGeometry;
+class FRDGBuilder;
 
 /**
 * Persistent representation of the scene for ray tracing.
@@ -29,7 +31,7 @@ public:
 	// This can be an expensive process, depending on the number of instances in the scene.
 	// Immediately allocates GPU memory to fit at least the current number of instances and updates the SRV.
 	// Returns the task that must be waited upon before accessing RayTracingSceneRHI.
-	FGraphEventRef BeginCreate();
+	FGraphEventRef BeginCreate(FRDGBuilder& GraphBuilder);
 	void WaitForTasks() const;
 
 	// Resets the instance list and reserves memory for this frame.
@@ -72,6 +74,8 @@ public:
 
 	FRayTracingAccelerationStructureSize SizeInfo = {};
 
+	FRDGBufferRef BuildScratchBuffer;
+
 private:
 	// RHI object that abstracts mesh instnaces in this scene
 	FRayTracingSceneRHIRef RayTracingSceneRHI;
