Clean hair tile generation and add indirect draw/dispatch helpers.

diff --git a/Engine/Shaders/Private/HairStrands/HairStrandsDebug.usf b/Engine/Shaders/Private/HairStrands/HairStrandsDebug.usf
--- a/Engine/Shaders/Private/HairStrands/HairStrandsDebug.usf
+++ b/Engine/Shaders/Private/HairStrands/HairStrandsDebug.usf
@@ -20,9 +20,6 @@ SamplerState LinearSampler;
 Texture2D<uint> HairCountUintTexture;
 Texture2D<float> HairCountTexture;
 
-uint TileSize;
-Texture2D<uint> TileIndexTexture;
-
 #define DEBUG_MODE_SAMPLE_COUNT 0 
 #define DEBUG_MODE_COVERAGE_TYPE 1
 #define DEBUG_MODE_TAA_RESOLVE_TYPE 2
@@ -85,14 +82,6 @@ void MainPS(
 		Color = GetColorCode(saturate(CatData.PixelCoverage));
 	}
 
-	if (DebugMode == DEBUG_MODE_TILE)
-	{
-		const uint2 TileCoord = PixelPos / TileSize;
-		const uint  TileOffset = TileIndexTexture[TileCoord];
-		Color = TileOffset != INVALID_TILE_OFFSET ? float3(0.25, 0.85, 0.25) : float3(0.85, 0.25, 0.25);
-		Alpha = 0.5f;
-	}
-
 	const bool bIsDebugModeMaterial =
 		DebugMode == DEBUG_MODE_DEPTH ||
 		DebugMode == DEBUG_MODE_BASECOLOR ||
diff --git a/Engine/Shaders/Private/HairStrands/HairStrandsVisibilityTile.usf b/Engine/Shaders/Private/HairStrands/HairStrandsVisibilityTile.usf
--- a/Engine/Shaders/Private/HairStrands/HairStrandsVisibilityTile.usf
+++ b/Engine/Shaders/Private/HairStrands/HairStrandsVisibilityTile.usf
@@ -5,52 +5,129 @@
 #include "HairStrandsVisibilityCommon.ush"
 
 #define TILE_SIZE 8
+#define GROUP_THREAD_COUNT (TILE_SIZE*TILE_SIZE)
 
-uint2 Resolution;
-uint2 TileResolution;
-uint  TileCount;
-uint  TileSize;
+////////////////////////////////////////////////////////////////////////////////////////////////////////////
 
-Texture2D<uint4> CategorizationTexture;
-RWTexture2D<uint> OutTileCounter;
-RWTexture2D<uint> OutTileIndexTexture;
-RWBuffer<uint2>   OutTileBuffer;
+#if SHADER_TILE_GENERATION
+int bRectPrimitive;
+int2 ViewResolution;
+Texture2D<uint4> InputTexture;
+RWBuffer<uint> TileIndirectBuffer;
+RWBuffer<uint2> TileDataBuffer;
 
-groupshared uint bAllocationTile;
+groupshared bool s_HasTileHair[GROUP_THREAD_COUNT];
 
 [numthreads(TILE_SIZE, TILE_SIZE, 1)]
-void MainCS(uint GroupIndex : SV_GroupIndex, uint3 DispatchThreadId : SV_DispatchThreadID)
+void TileMainCS(uint2 DispatchThreadId : SV_DispatchThreadID, uint LinearIndex : SV_GroupIndex, uint3 GroupId : SV_GroupID)
 {
-	if (GroupIndex == 0)
+	if (all(DispatchThreadId == 0))
 	{
-		bAllocationTile = 0;
+		TileIndirectBuffer[0] = bRectPrimitive > 0 ? 4 : 6;
 	}
 
-	uint2 PixelCoord = DispatchThreadId.xy;
-	if (PixelCoord.x >= Resolution.x || PixelCoord.y >= Resolution.y)
-		PixelCoord = Resolution-1;
+	const uint2 PixelCoord = DispatchThreadId;
 
-	const uint2 TileCoord = PixelCoord / TILE_SIZE;
-	const FCategorizationData Categorization = DecodeCategorizationData(CategorizationTexture.Load(uint3(PixelCoord, 0)));
-	if (Categorization.PixelCoverage > 0)
+	bool bHasHairData = false;
+	if (all(PixelCoord < uint2(ViewResolution)))
 	{
-		InterlockedOr(bAllocationTile, 1);		
+		const FCategorizationData CategorisationData = DecodeCategorizationData(InputTexture.Load(uint3(PixelCoord, 0)));
+		bHasHairData = CategorisationData.PixelCoverage > 0.0f;
 	}
+	s_HasTileHair[LinearIndex] = bHasHairData;
 
 	GroupMemoryBarrierWithGroupSync();
+	if (LinearIndex < 32)
+	{
+		s_HasTileHair[LinearIndex] = s_HasTileHair[LinearIndex] || s_HasTileHair[LinearIndex + 32];
+	}
+	GroupMemoryBarrierWithGroupSync();
+	if (LinearIndex < 16)
+	{
+		s_HasTileHair[LinearIndex] = s_HasTileHair[LinearIndex] || s_HasTileHair[LinearIndex + 16];
+	}
+	GroupMemoryBarrierWithGroupSync();
+
+	if (LinearIndex < 8)
+	{
+		s_HasTileHair[LinearIndex] = s_HasTileHair[LinearIndex] || s_HasTileHair[LinearIndex + 8];
+	}
+	GroupMemoryBarrierWithGroupSync();
+	if (LinearIndex < 4)
+	{
+		s_HasTileHair[LinearIndex] = s_HasTileHair[LinearIndex] || s_HasTileHair[LinearIndex + 4];
+	}
+	GroupMemoryBarrierWithGroupSync();
+	if (LinearIndex < 2)
+	{
+		s_HasTileHair[LinearIndex] = s_HasTileHair[LinearIndex] || s_HasTileHair[LinearIndex + 2];
+	}
+	GroupMemoryBarrierWithGroupSync();
+	if (LinearIndex < 1 && (s_HasTileHair[LinearIndex] || s_HasTileHair[LinearIndex + 1]))
+	{
+		uint WriteToIndex;
+		InterlockedAdd(TileIndirectBuffer[1], 1, WriteToIndex);
+		TileDataBuffer[WriteToIndex] = GroupId.xy;
+	}
+}
+#endif // SHADER_TILE_GENERATION
+
+////////////////////////////////////////////////////////////////////////////////////////////////////////////
+
+#if SHADER_TILE_COPY_ARGS
+uint GroupSize;
+Buffer<uint> TileIndirectDrawBuffer;
+RWBuffer<uint> TileIndirectDispatchBuffer;
+
+[numthreads(1, 1, 1)]
+void MainCS(uint2 DispatchThreadId : SV_DispatchThreadID)
+{
+	if (all(DispatchThreadId == 0))
+	{
+		TileIndirectDispatchBuffer[0] = (TileIndirectDrawBuffer[1] + GroupSize-1) / GroupSize;
+		TileIndirectDispatchBuffer[1] = 1;
+		TileIndirectDispatchBuffer[2] = 1;
+	}
+}
+#endif // SHADER_TILE_COPY_ARGS
+
+////////////////////////////////////////////////////////////////////////////////////////////////////////////
+
+#if SHADER_TILE_DEBUG
+
+void MainPS(
+	in FScreenVertexOutput Input,
+	out float4 OutColor0 : SV_Target0)
+{
+	const uint2 InTileCoord = uint2(Input.Position.xy) >> 3u; // TILE_SIZE == 8;
+	const bool bTileX = (InTileCoord.x & 1) == 0;
+	const bool bTileY = (InTileCoord.y & 1) == 0;
+	const bool bChecker = (bTileX && bTileY) || (!bTileX && !bTileY);
+	OutColor0 = 255;
+	OutColor0 = bChecker ? float4(0.0f, 0.5f, 0.f, 1.0f) : float4(0.f, 1.0f, 0.f, 1.0f);
+}
+
+#endif //SHADER_TILE_DEBUG
+
+////////////////////////////////////////////////////////////////////////////////////////////////////////////
+
+#if SHADER_TILE_VS
+int2 TileOutputResolution;
+Buffer<uint2> TileDataBuffer;
+
+void MainVS(
+	in uint InVertexId : SV_VertexID,
+	in uint InInstanceId : SV_InstanceID,
+	out FScreenVertexOutput Out)
+{
+	Out = (FScreenVertexOutput)0;
+	const uint2 TileCoord = TileDataBuffer[InInstanceId];
+
+	uint2 TileVertex = TileCoord * TILE_SIZE;
+	TileVertex.x += InVertexId == 1 || InVertexId == 2 || InVertexId == 4 ? TILE_SIZE : 0;
+	TileVertex.y += InVertexId == 2 || InVertexId == 4 || InVertexId == 5 ? TILE_SIZE : 0;
+	Out.UV = float2(TileVertex) / float2(TileOutputResolution);
+	Out.Position = float4(Out.UV * float2(2.0f, -2.0f) + float2(-1.0, 1.0f), 0.5f, 1.0f);
+}
 
-	if (GroupIndex == 0)
-	{
-		if (bAllocationTile > 0)
-		{
-			uint AllocationOffset = 0;
-			InterlockedAdd(OutTileCounter[uint2(0, 0)], 1, AllocationOffset);
-			OutTileIndexTexture[TileCoord] = AllocationOffset;
-			OutTileBuffer[AllocationOffset] = TileCoord;
-		}
-		else
-		{
-			OutTileIndexTexture[TileCoord] = INVALID_TILE_OFFSET;
-		}
-	}
-}
\ No newline at end of file
+#endif //SHADER_TILE_VS
\ No newline at end of file
diff --git a/Engine/Source/Runtime/Renderer/Private/HairStrands/HairStrandsData.h b/Engine/Source/Runtime/Renderer/Private/HairStrands/HairStrandsData.h
--- a/Engine/Source/Runtime/Renderer/Private/HairStrands/HairStrandsData.h
+++ b/Engine/Source/Runtime/Renderer/Private/HairStrands/HairStrandsData.h
@@ -28,6 +28,25 @@ BEGIN_GLOBAL_SHADER_PARAMETER_STRUCT(FHairStrandsViewUniformParameters, )
 	SHADER_PARAMETER_RDG_BUFFER_SRV(Buffer<uint2>, HairSampleCoords)					// Screen pixel coordinate of each sample, in sample space
 END_GLOBAL_SHADER_PARAMETER_STRUCT()
 
+////////////////////////////////////////////////////////////////////////////////////
+// Tile data
+
+struct FHairStrandsTiles
+{
+	FIntPoint			Resolution = FIntPoint(0, 0);
+	static const uint32 GroupSize = 64;
+	static const uint32	TileSize = 8;
+	uint32				TileCount = 0;
+	bool				bRectPrimitive = false;
+
+	FRDGBufferSRVRef	TileDataSRV = nullptr;
+	FRDGBufferRef		TileDataBuffer = nullptr;
+	FRDGBufferRef		TileIndirectDrawBuffer = nullptr;
+	FRDGBufferRef		TileIndirectDispatchBuffer = nullptr;
+
+	bool IsValid() const { return TileCount > 0 && TileDataBuffer != nullptr; }
+};
+
 ////////////////////////////////////////////////////////////////////////////////////
 // Visibility Data
 
@@ -43,12 +62,6 @@ struct FHairStrandsVisibilityData
 
 	FRDGTextureRef LightChannelMaskTexture = nullptr;
 
-	FRDGTextureRef	TileIndexTexture = nullptr;
-	FRDGBufferRef	TileBuffer = nullptr;
-	FRDGBufferRef	TileIndirectArgs = nullptr;
-	uint32			TileSize = 8;
-	uint32			TileThreadGroupSize = 32;
-
 	uint32			MaxSampleCount = 8;
 	uint32			MaxNodeCount = 0;
 	FRDGTextureRef	NodeCount = nullptr;
@@ -58,6 +71,8 @@ struct FHairStrandsVisibilityData
 	FRDGBufferRef	NodeIndirectArg = nullptr;
 	uint32			NodeGroupSize = 0;
 
+	FHairStrandsTiles TileData;
+
 	const static EPixelFormat NodeCoordFormat = PF_R16G16_UINT;
 
 	// Hair lighting is accumulated within this buffer
diff --git a/Engine/Source/Runtime/Renderer/Private/HairStrands/HairStrandsDebug.cpp b/Engine/Source/Runtime/Renderer/Private/HairStrands/HairStrandsDebug.cpp
--- a/Engine/Source/Runtime/Renderer/Private/HairStrands/HairStrandsDebug.cpp
+++ b/Engine/Source/Runtime/Renderer/Private/HairStrands/HairStrandsDebug.cpp
@@ -10,6 +10,7 @@
 #include "HairStrandsVisibility.h"
 #include "HairStrandsInterface.h"
 #include "HairStrandsMeshProjection.h"
+#include "HairStrandsTile.h"
 
 #include "Shader.h"
 #include "GlobalShader.h"
@@ -334,13 +335,11 @@ class FHairDebugPS : public FGlobalShader
 		SHADER_PARAMETER(uint32, FastResolveMask)
 		SHADER_PARAMETER(uint32, DebugMode)
 		SHADER_PARAMETER(int32, SampleIndex)
-		SHADER_PARAMETER(uint32, TileSize)
 		SHADER_PARAMETER(uint32, MaxSampleCount)
 		SHADER_PARAMETER_RDG_TEXTURE(Texture2D, HairCountTexture)
 		SHADER_PARAMETER_RDG_TEXTURE(Texture2D, HairCountUintTexture)
 		SHADER_PARAMETER_RDG_TEXTURE(Texture2D, CategorizationTexture)
 		SHADER_PARAMETER_RDG_TEXTURE(Texture2D, NodeIndex)
-		SHADER_PARAMETER_RDG_TEXTURE(Texture2D, TileIndexTexture)
 		SHADER_PARAMETER_RDG_BUFFER_SRV(StructuredBuffer, NodeData)
 		SHADER_PARAMETER_RDG_TEXTURE_SRV(Texture2D, DepthStencilTexture)
 		SHADER_PARAMETER_SAMPLER(SamplerState, LinearSampler)
@@ -371,10 +370,9 @@ static void AddDebugHairPass(
 		InDebugMode == EHairDebugMode::MaterialBaseColor ||
 		InDebugMode == EHairDebugMode::MaterialRoughness ||
 		InDebugMode == EHairDebugMode::MaterialSpecular ||
-		InDebugMode == EHairDebugMode::MaterialTangent ||
-		InDebugMode == EHairDebugMode::Tile);
+		InDebugMode == EHairDebugMode::MaterialTangent);
 
-	if (!VisibilityData.CategorizationTexture || !VisibilityData.NodeIndex || !VisibilityData.NodeData || !VisibilityData.TileIndexTexture) return;
+	if (!VisibilityData.CategorizationTexture || !VisibilityData.NodeIndex || !VisibilityData.NodeData) return;
 	if (InDebugMode == EHairDebugMode::TAAResolveType && !InDepthStencilTexture) return;
 
 	FRDGTextureRef HairCountTexture = VisibilityData.ViewHairCountTexture ? VisibilityData.ViewHairCountTexture : GSystemTextures.GetBlackDummy(GraphBuilder);
@@ -395,7 +393,6 @@ static void AddDebugHairPass(
 		case EHairDebugMode::MaterialRoughness:	InternalDebugMode = 6; break;
 		case EHairDebugMode::MaterialSpecular:	InternalDebugMode = 7; break;
 		case EHairDebugMode::MaterialTangent:	InternalDebugMode = 8; break;
-		case EHairDebugMode::Tile:				InternalDebugMode = 9; break;
 	};
 
 	FHairDebugPS::FParameters* Parameters = GraphBuilder.AllocParameters<FHairDebugPS::FParameters>();
@@ -407,9 +404,7 @@ static void AddDebugHairPass(
 	Parameters->HairCountUintTexture = HairCountUintTexture;
 	Parameters->NodeIndex = VisibilityData.NodeIndex;
 	Parameters->NodeData = GraphBuilder.CreateSRV(VisibilityData.NodeData);
-	Parameters->TileIndexTexture = VisibilityData.TileIndexTexture;
-	Parameters->TileSize = VisibilityData.TileSize;
-	Parameters->DepthStencilTexture = InDepthStencilTexture;
+		Parameters->DepthStencilTexture = InDepthStencilTexture;
 	Parameters->LinearSampler = TStaticSamplerState<SF_Bilinear, AM_Clamp, AM_Clamp, AM_Clamp>::GetRHI();
 	Parameters->DebugMode = InternalDebugMode;
 	Parameters->SampleIndex = GHairStrandsDebugSampleIndex;
@@ -1250,16 +1245,18 @@ static void InternalRenderHairStrandsDebugInfo(
 		HairDebugMode == EHairDebugMode::MaterialBaseColor ||
 		HairDebugMode == EHairDebugMode::MaterialRoughness ||
 		HairDebugMode == EHairDebugMode::MaterialSpecular ||
-		HairDebugMode == EHairDebugMode::MaterialTangent ||
-		HairDebugMode == EHairDebugMode::Tile;
+		HairDebugMode == EHairDebugMode::MaterialTangent;
 	if (bRunDebugPass)
 	{
 		{
-			const FHairStrandsVisibilityData& VisibilityData = HairData.VisibilityData;
-			AddDebugHairPass(GraphBuilder, &View, HairDebugMode, VisibilityData, SceneTextures.Stencil, SceneColorTexture);
-			AddDebugHairPrintPass(GraphBuilder, &View, HairDebugMode, VisibilityData, HairData.MacroGroupResources, SceneTextures.Stencil);
+			AddDebugHairPass(GraphBuilder, &View, HairDebugMode, HairData.VisibilityData, SceneTextures.Stencil, SceneColorTexture);
+			AddDebugHairPrintPass(GraphBuilder, &View, HairDebugMode, HairData.VisibilityData, HairData.MacroGroupResources, SceneTextures.Stencil);
 		}
 	}
+	else if (HairDebugMode == EHairDebugMode::Tile && HairData.VisibilityData.TileData.IsValid())
+	{
+		AddHairStrandsDebugTilePass(GraphBuilder, View, SceneColorTexture, HairData.VisibilityData.TileData);
+	}
 
 	if (bIsVoxelMode)
 	{
diff --git a/Engine/Source/Runtime/Renderer/Private/HairStrands/HairStrandsTile.cpp b/Engine/Source/Runtime/Renderer/Private/HairStrands/HairStrandsTile.cpp
new file mode 100644
--- /dev/null
+++ b/Engine/Source/Runtime/Renderer/Private/HairStrands/HairStrandsTile.cpp
@@ -0,0 +1,245 @@
+// Copyright Epic Games, Inc. All Rights Reserved.
+
+#include "HairStrandsTile.h"
+#include "HairStrandsUtils.h"
+#include "HairStrandsInterface.h"
+
+#include "Shader.h"
+#include "GlobalShader.h"
+#include "ShaderParameters.h"
+#include "ShaderParameterStruct.h"
+#include "SceneTextureParameters.h"
+#include "RenderGraphUtils.h"
+#include "PostProcessing.h"
+#include "MeshPassProcessor.inl"
+#include "ScenePrivate.h"
+
+////////////////////////////////////////////////////////////////////////////////////////////////////////////////////////////////////////////////
+
+FHairStrandsTilePassVS::FParameters GetHairStrandsTileParameters(const FHairStrandsTiles& In)
+{
+	FHairStrandsTilePassVS::FParameters Out;
+	Out.bRectPrimitive			= In.bRectPrimitive ? 1 : 0;
+	Out.TileOutputResolution	= In.Resolution;
+	Out.TileDataBuffer			= In.TileDataSRV;
+	Out.TileIndirectBuffer		= In.TileIndirectDrawBuffer;
+	return Out;
+}
+
+////////////////////////////////////////////////////////////////////////////////////////////////////////////
+// Generate indirect dispatch buffer based on the indirect draw buffer
+
+class FHairStrandsTileCopyArgsPassCS : public FGlobalShader
+{
+	DECLARE_GLOBAL_SHADER(FHairStrandsTileCopyArgsPassCS);
+	SHADER_USE_PARAMETER_STRUCT(FHairStrandsTileCopyArgsPassCS, FGlobalShader);
+
+	using FPermutationDomain = TShaderPermutationDomain<>;
+
+	BEGIN_SHADER_PARAMETER_STRUCT(FParameters, )
+		SHADER_PARAMETER_STRUCT_REF(FViewUniformShaderParameters, ViewUniformBuffer)
+		SHADER_PARAMETER(uint32, GroupSize)
+		SHADER_PARAMETER_RDG_BUFFER_SRV(Buffer, TileIndirectDrawBuffer)
+		SHADER_PARAMETER_RDG_BUFFER_UAV(RWBuffer, TileIndirectDispatchBuffer)
+	END_SHADER_PARAMETER_STRUCT()
+
+	static bool ShouldCompilePermutation(const FGlobalShaderPermutationParameters& Parameters)
+	{
+		return IsHairStrandsSupported(EHairStrandsShaderType::Strands, Parameters.Platform);
+	}
+
+	static void ModifyCompilationEnvironment(const FGlobalShaderPermutationParameters& Parameters, FShaderCompilerEnvironment& OutEnvironment)
+	{
+		FGlobalShader::ModifyCompilationEnvironment(Parameters, OutEnvironment);
+		OutEnvironment.SetDefine(TEXT("SHADER_TILE_COPY_ARGS"), 1);
+	}
+};
+
+IMPLEMENT_GLOBAL_SHADER(FHairStrandsTileCopyArgsPassCS, "/Engine/Private/HairStrands/HairStrandsVisibilityTile.usf", "MainCS", SF_Compute);
+
+void AddHairStrandsCopyArgsTilesPass(
+	FRDGBuilder& GraphBuilder,
+	const FViewInfo& View,
+	FHairStrandsTiles& TileData)
+{
+	TShaderMapRef<FHairStrandsTileCopyArgsPassCS> ComputeShader(View.ShaderMap);
+	FHairStrandsTileCopyArgsPassCS::FParameters* PassParameters = GraphBuilder.AllocParameters<FHairStrandsTileCopyArgsPassCS::FParameters>();
+	PassParameters->GroupSize = FHairStrandsTiles::GroupSize;
+	PassParameters->TileIndirectDrawBuffer = GraphBuilder.CreateSRV(TileData.TileIndirectDrawBuffer, PF_R32_UINT);
+	PassParameters->TileIndirectDispatchBuffer = GraphBuilder.CreateUAV(TileData.TileIndirectDispatchBuffer, PF_R32_UINT);
+
+	FComputeShaderUtils::AddPass(
+		GraphBuilder,
+		RDG_EVENT_NAME("HairStrandsTileCopyArgs"),
+		ComputeShader,
+		PassParameters,
+		FIntVector(1, 1, 1));
+}
+
+////////////////////////////////////////////////////////////////////////////////////////////////////////////////////////////////////////////////
+// Generate tiles data based on input texture (hair pixel coverage)
+
+class FHairStrandsTileGenerationPassCS : public FGlobalShader
+{
+	DECLARE_GLOBAL_SHADER(FHairStrandsTileGenerationPassCS);
+	SHADER_USE_PARAMETER_STRUCT(FHairStrandsTileGenerationPassCS, FGlobalShader);
+
+	using FPermutationDomain = TShaderPermutationDomain<>;
+
+	BEGIN_SHADER_PARAMETER_STRUCT(FParameters, )
+		SHADER_PARAMETER_STRUCT_REF(FViewUniformShaderParameters, ViewUniformBuffer)
+		SHADER_PARAMETER(int32, bRectPrimitive)
+		SHADER_PARAMETER(FIntPoint, ViewResolution)
+		SHADER_PARAMETER_RDG_TEXTURE(Texture2D<uint4>, InputTexture)
+		SHADER_PARAMETER_RDG_BUFFER_UAV(RWBuffer, TileIndirectBuffer)
+		SHADER_PARAMETER_RDG_BUFFER_UAV(RWBuffer, TileDataBuffer)
+	END_SHADER_PARAMETER_STRUCT()
+
+	static bool ShouldCompilePermutation(const FGlobalShaderPermutationParameters& Parameters)
+	{
+		return IsHairStrandsSupported(EHairStrandsShaderType::Strands, Parameters.Platform);
+	}
+
+	static void ModifyCompilationEnvironment(const FGlobalShaderPermutationParameters& Parameters, FShaderCompilerEnvironment& OutEnvironment)
+	{
+		FGlobalShader::ModifyCompilationEnvironment(Parameters, OutEnvironment);
+		OutEnvironment.SetDefine(TEXT("SHADER_TILE_GENERATION"), 1);
+	}
+};
+IMPLEMENT_GLOBAL_SHADER(FHairStrandsTileGenerationPassCS, "/Engine/Private/HairStrands/HairStrandsVisibilityTile.usf", "TileMainCS", SF_Compute);
+
+FHairStrandsTiles AddHairStrandsGenerateTilesPass(
+	FRDGBuilder& GraphBuilder,
+	const FViewInfo& View,
+	const FRDGTextureRef& InputTexture)
+{
+	FHairStrandsTiles Out;
+
+	check(FHairStrandsTiles::TileSize == 8); // only size supported for now
+	const FIntPoint InputResolution = InputTexture->Desc.Extent;
+	const FIntPoint TileResolution = FIntPoint(FMath::CeilToInt(InputResolution.X / float(FHairStrandsTiles::TileSize)), FMath::CeilToInt(InputResolution.Y / float(FHairStrandsTiles::TileSize)));
+	Out.TileCount = TileResolution.X * TileResolution.Y;
+	Out.Resolution = InputResolution;
+	Out.bRectPrimitive = GRHISupportsRectTopology;
+	Out.TileIndirectDrawBuffer = GraphBuilder.CreateBuffer(FRDGBufferDesc::CreateIndirectDesc<FRHIDrawIndirectParameters>(1), TEXT("Hair.TileIndirectDrawBuffer"));
+	Out.TileIndirectDispatchBuffer = GraphBuilder.CreateBuffer(FRDGBufferDesc::CreateIndirectDesc<FRHIDispatchIndirectParameters>(1), TEXT("Hair.TileIndirectDispatchBuffer"));
+	Out.TileDataBuffer = GraphBuilder.CreateBuffer(FRDGBufferDesc::CreateBufferDesc(sizeof(uint32), Out.TileCount), TEXT("Hair.TileDataBuffer"));
+
+	FRDGBufferUAVRef TileIndirectUAV = GraphBuilder.CreateUAV(Out.TileIndirectDrawBuffer, PF_R32_UINT);
+	AddClearUAVPass(GraphBuilder, TileIndirectUAV, 0u);
+
+	TShaderMapRef<FHairStrandsTileGenerationPassCS> ComputeShader(View.ShaderMap);
+	FHairStrandsTileGenerationPassCS::FParameters* PassParameters = GraphBuilder.AllocParameters<FHairStrandsTileGenerationPassCS::FParameters>();
+	PassParameters->bRectPrimitive = Out.bRectPrimitive ? 1 : 0;
+	PassParameters->ViewResolution = View.ViewRect.Size();
+	PassParameters->InputTexture = InputTexture;
+	PassParameters->TileDataBuffer = GraphBuilder.CreateUAV(Out.TileDataBuffer, PF_R16G16_UINT);
+	PassParameters->TileIndirectBuffer = TileIndirectUAV;
+
+	const uint32 GroupSize = 8;
+	FComputeShaderUtils::AddPass(
+		GraphBuilder,
+		RDG_EVENT_NAME("HairStrandsTileClassification"),
+		ComputeShader,
+		PassParameters,
+		FComputeShaderUtils::GetGroupCount(InputResolution, GroupSize));
+
+	Out.TileDataSRV = GraphBuilder.CreateSRV(Out.TileDataBuffer, PF_R16G16_UINT);
+
+	// Initialize indirect dispatch buffer, based on the indirect draw bugger
+	AddHairStrandsCopyArgsTilesPass(GraphBuilder, View, Out);
+
+	return Out;
+}
+
+////////////////////////////////////////////////////////////////////////////////////////////////////////////////////////////////////////////////
+
+bool FHairStrandsTilePassVS::ShouldCompilePermutation(const FGlobalShaderPermutationParameters& Parameters)
+{
+	return IsHairStrandsSupported(EHairStrandsShaderType::Strands, Parameters.Platform);
+}
+
+void FHairStrandsTilePassVS::ModifyCompilationEnvironment(const FGlobalShaderPermutationParameters& Parameters, FShaderCompilerEnvironment& OutEnvironment)
+{
+	FGlobalShader::ModifyCompilationEnvironment(Parameters, OutEnvironment);
+	OutEnvironment.SetDefine(TEXT("SHADER_TILE_VS"), 1);
+}
+
+IMPLEMENT_GLOBAL_SHADER(FHairStrandsTilePassVS, "/Engine/Private/HairStrands/HairStrandsVisibilityTile.usf", "MainVS", SF_Vertex);
+
+////////////////////////////////////////////////////////////////////////////////////////////////////////////////////////////////////////////////
+
+class FHairStrandsTileDebugPassPS : public FGlobalShader
+{
+	DECLARE_GLOBAL_SHADER(FHairStrandsTileDebugPassPS);
+	SHADER_USE_PARAMETER_STRUCT(FHairStrandsTileDebugPassPS, FGlobalShader);
+
+	using FPermutationDomain = TShaderPermutationDomain<>;
+
+	BEGIN_SHADER_PARAMETER_STRUCT(FParameters, )
+		SHADER_PARAMETER_STRUCT_REF(FViewUniformShaderParameters, ViewUniformBuffer)
+		SHADER_PARAMETER_STRUCT_INCLUDE(FHairStrandsTilePassVS::FParameters, TileParameters)
+		SHADER_PARAMETER(FIntPoint, OutputResolution)
+		RENDER_TARGET_BINDING_SLOTS()
+	END_SHADER_PARAMETER_STRUCT()
+
+	static bool ShouldCompilePermutation(const FGlobalShaderPermutationParameters& Parameters)
+	{
+		return GetMaxSupportedFeatureLevel(Parameters.Platform) >= ERHIFeatureLevel::SM5; //TODO
+	}
+
+	static void ModifyCompilationEnvironment(const FGlobalShaderPermutationParameters& Parameters, FShaderCompilerEnvironment& OutEnvironment)
+	{
+		FGlobalShader::ModifyCompilationEnvironment(Parameters, OutEnvironment);
+		OutEnvironment.SetDefine(TEXT("SHADER_TILE_DEBUG"), 1); 
+	}
+};
+
+IMPLEMENT_GLOBAL_SHADER(FHairStrandsTileDebugPassPS, "/Engine/Private/HairStrands/HairStrandsVisibilityTile.usf", "MainPS", SF_Pixel);
+
+////////////////////////////////////////////////////////////////////////////////////////////////////////////////////////////////////////////////
+
+void AddHairStrandsDebugTilePass(
+	FRDGBuilder& GraphBuilder,
+	const FViewInfo& View,
+	const FRDGTextureRef& ColorTexture,
+	const FHairStrandsTiles& TileData)
+{	
+	const FIntPoint OutputResolution = View.ViewRect.Size();
+	
+	FHairStrandsTileDebugPassPS::FParameters* PassParameters = GraphBuilder.AllocParameters<FHairStrandsTileDebugPassPS::FParameters>();
+	PassParameters->TileParameters = GetHairStrandsTileParameters(TileData);
+	PassParameters->OutputResolution = OutputResolution;
+
+	TShaderMapRef<FHairStrandsTilePassVS> VertexShader(View.ShaderMap);
+	TShaderMapRef<FHairStrandsTileDebugPassPS> PixelShader(View.ShaderMap);
+
+	PassParameters->RenderTargets[0] = FRenderTargetBinding(ColorTexture, ERenderTargetLoadAction::ELoad);
+	
+	GraphBuilder.AddPass(
+		RDG_EVENT_NAME("HairStrandsTileDebugPass"),
+		PassParameters,
+		ERDGPassFlags::Raster,
+		[PassParameters, VertexShader, PixelShader, OutputResolution](FRHICommandList& RHICmdList)
+		{
+			FHairStrandsTilePassVS::FParameters ParametersVS = PassParameters->TileParameters;
+
+			FGraphicsPipelineStateInitializer GraphicsPSOInit;
+			RHICmdList.ApplyCachedRenderTargets(GraphicsPSOInit);
+			GraphicsPSOInit.BlendState = TStaticBlendState<CW_RGBA, BO_Add, BF_One, BF_One, BO_Max, BF_SourceAlpha, BF_DestAlpha>::GetRHI();
+			GraphicsPSOInit.RasterizerState = TStaticRasterizerState<>::GetRHI();
+			GraphicsPSOInit.DepthStencilState = TStaticDepthStencilState<false, CF_Always>::GetRHI();
+
+			GraphicsPSOInit.BoundShaderState.VertexDeclarationRHI = GFilterVertexDeclaration.VertexDeclarationRHI;
+			GraphicsPSOInit.BoundShaderState.VertexShaderRHI = VertexShader.GetVertexShader();
+			GraphicsPSOInit.BoundShaderState.PixelShaderRHI = PixelShader.GetPixelShader();
+			GraphicsPSOInit.PrimitiveType = PassParameters->TileParameters.bRectPrimitive > 0 ? PT_RectList : PT_TriangleList;
+			SetGraphicsPipelineState(RHICmdList, GraphicsPSOInit);
+			SetShaderParameters(RHICmdList, VertexShader, VertexShader.GetVertexShader(), ParametersVS);
+			SetShaderParameters(RHICmdList, PixelShader, PixelShader.GetPixelShader(), *PassParameters);
+
+			RHICmdList.SetViewport(0, 0, 0.0f, OutputResolution.X, OutputResolution.Y, 1.0f);
+			RHICmdList.SetStreamSource(0, nullptr, 0);
+			RHICmdList.DrawPrimitiveIndirect(PassParameters->TileParameters.TileIndirectBuffer->GetRHI(), 0);
+		});
+}
\ No newline at end of file
diff --git a/Engine/Source/Runtime/Renderer/Private/HairStrands/HairStrandsTile.h b/Engine/Source/Runtime/Renderer/Private/HairStrands/HairStrandsTile.h
new file mode 100644
--- /dev/null
+++ b/Engine/Source/Runtime/Renderer/Private/HairStrands/HairStrandsTile.h
@@ -0,0 +1,43 @@
+// Copyright Epic Games, Inc. All Rights Reserved.
+
+/*=============================================================================
+	HairVisibilityRendering.h: Hair strands visibility buffer implementation.
+=============================================================================*/
+
+#pragma once
+
+#include "CoreMinimal.h"
+#include "RendererInterface.h"
+#include "RenderGraphResources.h"
+#include "GlobalShader.h"
+#include "ShaderParameterStruct.h"
+#include "HairStrandsData.h"
+
+class FHairStrandsTilePassVS : public FGlobalShader
+{
+	DECLARE_GLOBAL_SHADER(FHairStrandsTilePassVS);
+	SHADER_USE_PARAMETER_STRUCT(FHairStrandsTilePassVS, FGlobalShader);
+
+	BEGIN_SHADER_PARAMETER_STRUCT(FParameters, )
+		SHADER_PARAMETER(int32, bRectPrimitive)
+		SHADER_PARAMETER(FIntPoint, TileOutputResolution)
+		SHADER_PARAMETER_RDG_BUFFER_SRV(Buffer, TileDataBuffer)
+		RDG_BUFFER_ACCESS(TileIndirectBuffer, ERHIAccess::IndirectArgs)
+	END_SHADER_PARAMETER_STRUCT()
+
+	static bool ShouldCompilePermutation(const FGlobalShaderPermutationParameters& Parameters);
+	static void ModifyCompilationEnvironment(const FGlobalShaderPermutationParameters& Parameters, FShaderCompilerEnvironment& OutEnvironment);
+};
+
+FHairStrandsTilePassVS::FParameters GetHairStrandsTileParameters(const FHairStrandsTiles& In);
+
+FHairStrandsTiles AddHairStrandsGenerateTilesPass(
+	FRDGBuilder& GraphBuilder,
+	const FViewInfo& View,
+	const FRDGTextureRef& InputTexture);
+
+void AddHairStrandsDebugTilePass(
+	FRDGBuilder& GraphBuilder,
+	const FViewInfo& View,
+	const FRDGTextureRef& ColorTexture,
+	const FHairStrandsTiles& TileData);
\ No newline at end of file
diff --git a/Engine/Source/Runtime/Renderer/Private/HairStrands/HairStrandsVisibility.cpp b/Engine/Source/Runtime/Renderer/Private/HairStrands/HairStrandsVisibility.cpp
--- a/Engine/Source/Runtime/Renderer/Private/HairStrands/HairStrandsVisibility.cpp
+++ b/Engine/Source/Runtime/Renderer/Private/HairStrands/HairStrandsVisibility.cpp
@@ -4,6 +4,7 @@
 #include "HairStrandsCluster.h"
 #include "HairStrandsUtils.h"
 #include "HairStrandsInterface.h"
+#include "HairStrandsTile.h"
 
 #include "Shader.h"
 #include "GlobalShader.h"
@@ -82,6 +83,10 @@ static FAutoConsoleVariableRef CVarHairStrandsVisibility_OutputEmissiveData(TEXT
 static int32 GHairStrandsDebugPPLL = 0;
 static FAutoConsoleVariableRef CVarHairStrandsDebugPPLL(TEXT("r.HairStrands.Visibility.PPLL.Debug"), GHairStrandsDebugPPLL, TEXT("Draw debug per pixel light list rendering."));
 
+static int32 GHairStrandsTile = 0;
+static FAutoConsoleVariableRef CVarHairStrandsTile(TEXT("r.HairStrands.Tile"), GHairStrandsTile, TEXT("Enable tile generation & usage for hair strands."));
+
+
 /////////////////////////////////////////////////////////////////////////////////////////
 
 namespace HairStrandsVisibilityInternal
@@ -2116,87 +2121,6 @@ static void AddHairVisibilityCompactionComputeRasterPass(
 	OutMaxRenderNodeCount = MaxRenderNodeCount;
 }
 
-///////////////////////////////////////////////////////////////////////////////////////////////////
-class FHairGenerateTileCS : public FGlobalShader
-{
-	DECLARE_GLOBAL_SHADER(FHairGenerateTileCS);
-	SHADER_USE_PARAMETER_STRUCT(FHairGenerateTileCS, FGlobalShader);
-
-	class FTileSize : SHADER_PERMUTATION_INT("PERMUTATION_TILESIZE", 2);
-	using FPermutationDomain = TShaderPermutationDomain<FTileSize>;
-
-	BEGIN_SHADER_PARAMETER_STRUCT(FParameters, )
-		SHADER_PARAMETER(FIntPoint, Resolution)
-		SHADER_PARAMETER(FIntPoint, TileResolution)
-		SHADER_PARAMETER_RDG_TEXTURE(Texture2D, CategorizationTexture)
-		SHADER_PARAMETER_RDG_TEXTURE_UAV(RWTexture2D, OutTileCounter)
-		SHADER_PARAMETER_RDG_TEXTURE_UAV(RWTexture2D, OutTileIndexTexture)
-		SHADER_PARAMETER_RDG_BUFFER_UAV(RWBuffer, OutTileBuffer)
-	END_SHADER_PARAMETER_STRUCT()
-
-public:
-
-	static bool ShouldCompilePermutation(const FGlobalShaderPermutationParameters& Parameters)
-	{
-		return IsHairStrandsSupported(EHairStrandsShaderType::Strands, Parameters.Platform);
-	}
-};
-
-IMPLEMENT_GLOBAL_SHADER(FHairGenerateTileCS, "/Engine/Private/HairStrands/HairStrandsVisibilityTile.usf", "MainCS", SF_Compute);
-
-static void AddGenerateTilePass(
-	FRDGBuilder& GraphBuilder,
-	const FViewInfo& View,
-	const uint32 ThreadGroupSize,
-	const uint32 TileSize,
-	const FRDGTextureRef& CategorizationTexture,
-	FRDGTextureRef& OutTileIndexTexture,
-	FRDGBufferRef& OutTileBuffer,
-	FRDGBufferRef& OutTileIndirectArgs)
-{
-	check(TileSize == 8); // only size supported for now
-	const FIntPoint Resolution = CategorizationTexture->Desc.Extent;
-	const FIntPoint TileResolution = FIntPoint(FMath::CeilToInt(Resolution.X / float(TileSize)), FMath::CeilToInt(Resolution.Y / float(TileSize)));
-	const uint32 TileCount = TileResolution.X * TileResolution.Y;
-
-	FRDGTextureRef TileCounter;
-	{
-		FRDGTextureDesc Desc = FRDGTextureDesc::Create2D(FIntPoint(1,1), PF_R32_UINT, FClearValueBinding::None, TexCreate_UAV | TexCreate_ShaderResource);
-		TileCounter = GraphBuilder.CreateTexture(Desc, TEXT("Hair.TileCounter"));
-	}
-
-	{
-		FRDGTextureDesc Desc = FRDGTextureDesc::Create2D(TileResolution, PF_R32_UINT, FClearValueBinding::None, TexCreate_UAV | TexCreate_ShaderResource);
-		OutTileIndexTexture = GraphBuilder.CreateTexture(Desc, TEXT("Hair.TileIndexTexture"));
-	}
-
-	OutTileBuffer = GraphBuilder.CreateBuffer(FRDGBufferDesc::CreateBufferDesc(sizeof(uint32), TileCount), TEXT("Hair.TileBuffer"));
-
-	uint32 ClearValues[4] = { 0,0,0,0 };
-	AddClearUAVPass(GraphBuilder, GraphBuilder.CreateUAV(TileCounter), ClearValues);
-
-	FHairGenerateTileCS::FPermutationDomain PermutationVector;
-	PermutationVector.Set<FHairGenerateTileCS::FTileSize>(0);
-
-	FHairGenerateTileCS::FParameters* PassParameters = GraphBuilder.AllocParameters<FHairGenerateTileCS::FParameters>();
-	PassParameters->Resolution = Resolution;
-	PassParameters->TileResolution = TileResolution;
-	PassParameters->CategorizationTexture = CategorizationTexture;
-	PassParameters->OutTileCounter = GraphBuilder.CreateUAV(TileCounter);
-	PassParameters->OutTileIndexTexture = GraphBuilder.CreateUAV(OutTileIndexTexture);
-	PassParameters->OutTileBuffer = GraphBuilder.CreateUAV(OutTileBuffer, PF_R16G16_UINT);
-
-	TShaderMapRef<FHairGenerateTileCS> ComputeShader(View.ShaderMap, PermutationVector);
-	FComputeShaderUtils::AddPass(
-		GraphBuilder,
-		RDG_EVENT_NAME("HairGenerateTile"),
-		ComputeShader,
-		PassParameters,
-		FIntVector(TileResolution.X, TileResolution.Y, 1));
-
-	OutTileIndirectArgs = AddCopyIndirectArgPass(GraphBuilder, &View, ThreadGroupSize, TileSize*TileSize, TileCounter);
-}
-
 ///////////////////////////////////////////////////////////////////////////////////////////////////
 class FHairVisibilityFillOpaqueDepthPS : public FGlobalShader
 {
@@ -3667,16 +3591,10 @@ void RenderHairStrandsVisibilityBuffer(
 		#endif
 
 			// Generate Tile data
-			if (CategorizationTexture)
+			const bool bGenerateTile = GHairStrandsTile > 0;
+			if (CategorizationTexture && bGenerateTile)
 			{
-				FRDGTextureRef TileIndexTexture = nullptr;
-				FRDGBufferRef TileBuffer = nullptr;
-				FRDGBufferRef TileIndirectArgs = nullptr;
-				AddGenerateTilePass(GraphBuilder, View, VisibilityData.TileThreadGroupSize, VisibilityData.TileSize, CategorizationTexture, TileIndexTexture, TileBuffer, TileIndirectArgs);
-
-				VisibilityData.TileIndexTexture = TileIndexTexture;
-				VisibilityData.TileBuffer = TileBuffer;
-				VisibilityData.TileIndirectArgs = TileIndirectArgs;
+				VisibilityData.TileData = AddHairStrandsGenerateTilesPass(GraphBuilder, View, CategorizationTexture);
 			}
 		}
 	}
