Shared screen draw code between hair postprocess/composition pass.

diff --git a/Engine/Source/Runtime/Renderer/Private/HairStrands/HairStrandsComposition.cpp b/Engine/Source/Runtime/Renderer/Private/HairStrands/HairStrandsComposition.cpp
--- a/Engine/Source/Runtime/Renderer/Private/HairStrands/HairStrandsComposition.cpp
+++ b/Engine/Source/Runtime/Renderer/Private/HairStrands/HairStrandsComposition.cpp
@@ -32,6 +32,112 @@ float GetHairFastResolveVelocityThreshold(const FIntPoint& Resolution)
 	return VelocityThreshold;
 }
 
+enum EHairStrandsCommonPassType
+{
+	Composition,
+	DOF,
+	TAAFastResolve,
+	GBuffer
+};
+
+template<typename TPassParameter, typename TPixelShader>
+void InternalCommonDrawPass(
+	FRDGBuilder& GraphBuilder, 
+	FRDGEventName&& EventName,
+	const FViewInfo& View,
+	const FIntPoint Resolution,
+	const EHairStrandsCommonPassType Type,
+	const bool bWriteDepth,
+	TPixelShader& PixelShader,
+	TPassParameter* PassParamters)
+{
+	ClearUnusedGraphResources(PixelShader, PassParamters);
+
+	const FIntRect Viewport = View.ViewRect;
+
+	TShaderMapRef<FPostProcessVS> VertexShader(View.ShaderMap);
+	GraphBuilder.AddPass(
+		Forward<FRDGEventName>(EventName),
+		PassParamters,
+		ERDGPassFlags::Raster,
+		[PassParamters, VertexShader, PixelShader, Viewport, Resolution, Type, bWriteDepth](FRHICommandList& RHICmdList)
+	{
+		FGraphicsPipelineStateInitializer GraphicsPSOInit;
+		RHICmdList.ApplyCachedRenderTargets(GraphicsPSOInit);
+
+		if (Type == EHairStrandsCommonPassType::Composition)
+		{
+			// Alpha usage/output is controlled with r.PostProcessing.PropagateAlpha. The value are:
+			// 0: disabled(default);
+			// 1: enabled in linear color space;
+			// 2: same as 1, but also enable it through the tonemapper.
+			//
+			// When enable (PorpagateAlpha is set to 1 or 2), the alpha value means:
+			// 0: valid pixel
+			// 1: invalid pixel (background)
+			GraphicsPSOInit.BlendState = TStaticBlendState<CW_RGBA, BO_Add, BF_One, BF_InverseSourceAlpha, BO_Add, BF_Zero, BF_InverseSourceAlpha>::GetRHI();
+		}
+		else
+		{
+			GraphicsPSOInit.BlendState = TStaticBlendState<CW_RGBA, BO_Add, BF_One, BF_Zero, BO_Max, BF_One, BF_Zero>::GetRHI();
+		}
+
+		GraphicsPSOInit.RasterizerState = TStaticRasterizerState<>::GetRHI();
+
+		if (Type == EHairStrandsCommonPassType::Composition)
+		{
+			GraphicsPSOInit.DepthStencilState = TStaticDepthStencilState<true, CF_DepthNearOrEqual>::GetRHI();
+		}
+		else if (Type == EHairStrandsCommonPassType::DOF)
+		{
+			GraphicsPSOInit.DepthStencilState = TStaticDepthStencilState<false, CF_Always>::GetRHI();
+		}
+		else if (Type == EHairStrandsCommonPassType::TAAFastResolve)
+		{
+			GraphicsPSOInit.DepthStencilState = TStaticDepthStencilState<
+				false, CF_Always,
+				true, CF_Always, SO_Keep, SO_Keep, SO_Replace,
+				false, CF_Always, SO_Keep, SO_Keep, SO_Keep,
+				STENCIL_TEMPORAL_RESPONSIVE_AA_MASK, STENCIL_TEMPORAL_RESPONSIVE_AA_MASK>::GetRHI();
+		}
+		else if (Type == EHairStrandsCommonPassType::GBuffer)
+		{
+			if (bWriteDepth)
+			{
+				GraphicsPSOInit.DepthStencilState = TStaticDepthStencilState<true, CF_Always>::GetRHI();
+			}
+			else
+			{
+				GraphicsPSOInit.DepthStencilState = TStaticDepthStencilState<false, CF_Always>::GetRHI();
+			}
+		}
+
+
+		GraphicsPSOInit.BoundShaderState.VertexDeclarationRHI = GFilterVertexDeclaration.VertexDeclarationRHI;
+		GraphicsPSOInit.BoundShaderState.VertexShaderRHI = VertexShader.GetVertexShader();
+		GraphicsPSOInit.BoundShaderState.PixelShaderRHI = PixelShader.GetPixelShader();
+		GraphicsPSOInit.PrimitiveType = PT_TriangleList;
+		SetGraphicsPipelineState(RHICmdList, GraphicsPSOInit);
+
+		if (Type == EHairStrandsCommonPassType::TAAFastResolve)
+		{
+			RHICmdList.SetStencilRef(STENCIL_TEMPORAL_RESPONSIVE_AA_MASK);
+		}
+		RHICmdList.SetViewport(Viewport.Min.X, Viewport.Min.Y, 0.0f, Viewport.Max.X, Viewport.Max.Y, 1.0f);
+		SetShaderParameters(RHICmdList, PixelShader, PixelShader.GetPixelShader(), *PassParamters);
+		DrawRectangle(
+			RHICmdList,
+			0, 0,
+			Viewport.Width(), Viewport.Height(),
+			Viewport.Min.X, Viewport.Min.Y,
+			Viewport.Width(), Viewport.Height(),
+			Viewport.Size(),
+			Resolution,
+			VertexShader,
+			EDRF_UseTriangleOptimization);
+	});
+}
+
 /////////////////////////////////////////////////////////////////////////////////////////
 
 class FHairVisibilityComposeSamplePS : public FGlobalShader
@@ -94,53 +200,16 @@ static void AddHairVisibilityComposeSamplePass(
 	Parameters->RenderTargets[0] = FRenderTargetBinding(OutColorTexture, ERenderTargetLoadAction::ELoad);
 	Parameters->RenderTargets.DepthStencil = FDepthStencilBinding(OutDepthTexture, ERenderTargetLoadAction::ELoad, ERenderTargetLoadAction::ENoAction, FExclusiveDepthStencil::DepthWrite_StencilNop);
 
-	TShaderMapRef<FPostProcessVS> VertexShader(View.ShaderMap);
 	TShaderMapRef<FHairVisibilityComposeSamplePS> PixelShader(View.ShaderMap);
-	FGlobalShaderMap* GlobalShaderMap = View.ShaderMap;
-	const FIntRect Viewport = View.ViewRect;
-	const FIntPoint Resolution = OutColorTexture->Desc.Extent;
-
-	ClearUnusedGraphResources(PixelShader, Parameters);
-
-	GraphBuilder.AddPass(
+	InternalCommonDrawPass(
+		GraphBuilder,
 		RDG_EVENT_NAME("HairStrandsComposeSample"),
-		Parameters,
-		ERDGPassFlags::Raster,
-		[Parameters, VertexShader, PixelShader, Viewport, Resolution](FRHICommandList& RHICmdList)
-	{
-		FGraphicsPipelineStateInitializer GraphicsPSOInit;
-		RHICmdList.ApplyCachedRenderTargets(GraphicsPSOInit);
-
-		// Alpha usage/output is controlled with r.PostProcessing.PropagateAlpha. The value are:
-		// 0: disabled(default);
-		// 1: enabled in linear color space;
-		// 2: same as 1, but also enable it through the tonemapper.
-		//
-		// When enable (PorpagateAlpha is set to 1 or 2), the alpha value means:
-		// 0: valid pixel
-		// 1: invalid pixel (background)
-		GraphicsPSOInit.BlendState = TStaticBlendState<CW_RGBA, BO_Add, BF_One, BF_InverseSourceAlpha, BO_Add, BF_Zero, BF_InverseSourceAlpha>::GetRHI();
-		GraphicsPSOInit.RasterizerState = TStaticRasterizerState<>::GetRHI();
-		GraphicsPSOInit.DepthStencilState = TStaticDepthStencilState<true, CF_DepthNearOrEqual>::GetRHI();
-		GraphicsPSOInit.BoundShaderState.VertexDeclarationRHI = GFilterVertexDeclaration.VertexDeclarationRHI;
-		GraphicsPSOInit.BoundShaderState.VertexShaderRHI = VertexShader.GetVertexShader();
-		GraphicsPSOInit.BoundShaderState.PixelShaderRHI = PixelShader.GetPixelShader();
-		GraphicsPSOInit.PrimitiveType = PT_TriangleList;
-		SetGraphicsPipelineState(RHICmdList, GraphicsPSOInit);
-
-		RHICmdList.SetViewport(Viewport.Min.X, Viewport.Min.Y, 0.0f, Viewport.Max.X, Viewport.Max.Y, 1.0f);
-		SetShaderParameters(RHICmdList, PixelShader, PixelShader.GetPixelShader(), *Parameters);
-		DrawRectangle(
-			RHICmdList,
-			0, 0,
-			Viewport.Width(), Viewport.Height(),
-			Viewport.Min.X, Viewport.Min.Y,
-			Viewport.Width(), Viewport.Height(),
-			Viewport.Size(),
-			Resolution,
-			VertexShader,
-			EDRF_UseTriangleOptimization);
-	});
+		View,
+		OutColorTexture->Desc.Extent,
+		EHairStrandsCommonPassType::Composition,
+		false,
+		PixelShader,
+		Parameters);
 }
 
 
@@ -201,44 +270,16 @@ static FRDGTextureRef AddHairDOFDepthPass(
 	Parameters->ViewUniformBuffer = View.ViewUniformBuffer;
 	Parameters->RenderTargets[0] = FRenderTargetBinding(OutDOFDepthTexture, ERenderTargetLoadAction::ENoAction);
 
-	TShaderMapRef<FPostProcessVS> VertexShader(View.ShaderMap);
 	TShaderMapRef<FHairDOFDepthPS> PixelShader(View.ShaderMap);
-	FGlobalShaderMap* GlobalShaderMap = View.ShaderMap;
-	const FIntPoint Resolution = OutputResolution;
-	const FIntRect Viewport = View.ViewRect;
-
-	ClearUnusedGraphResources(PixelShader, Parameters);
-
-	GraphBuilder.AddPass(
+	InternalCommonDrawPass(
+		GraphBuilder,
 		RDG_EVENT_NAME("HairStrandsDOFDepth"),
-		Parameters,
-		ERDGPassFlags::Raster,
-		[Parameters, VertexShader, PixelShader, Viewport, Resolution](FRHICommandList& RHICmdList)
-	{
-		FGraphicsPipelineStateInitializer GraphicsPSOInit;
-		RHICmdList.ApplyCachedRenderTargets(GraphicsPSOInit); 
-		GraphicsPSOInit.BlendState = TStaticBlendState<CW_RGBA, BO_Add, BF_One, BF_Zero, BO_Max, BF_One, BF_Zero>::GetRHI();
-		GraphicsPSOInit.RasterizerState = TStaticRasterizerState<>::GetRHI();
-		GraphicsPSOInit.DepthStencilState = TStaticDepthStencilState<false, CF_Always>::GetRHI();
-		GraphicsPSOInit.BoundShaderState.VertexDeclarationRHI = GFilterVertexDeclaration.VertexDeclarationRHI;
-		GraphicsPSOInit.BoundShaderState.VertexShaderRHI = VertexShader.GetVertexShader();
-		GraphicsPSOInit.BoundShaderState.PixelShaderRHI = PixelShader.GetPixelShader();
-		GraphicsPSOInit.PrimitiveType = PT_TriangleList;
-		SetGraphicsPipelineState(RHICmdList, GraphicsPSOInit);
-
-		RHICmdList.SetViewport(Viewport.Min.X, Viewport.Min.Y, 0.0f, Viewport.Max.X, Viewport.Max.Y, 1.0f);
-		SetShaderParameters(RHICmdList, PixelShader, PixelShader.GetPixelShader(), *Parameters);
-		DrawRectangle(
-			RHICmdList,
-			0, 0,
-			Viewport.Width(), Viewport.Height(),
-			Viewport.Min.X, Viewport.Min.Y,
-			Viewport.Width(), Viewport.Height(),
-			Viewport.Size(),
-			Resolution,
-			VertexShader,
-			EDRF_UseTriangleOptimization);
-	});
+		View,
+		OutputResolution,
+		EHairStrandsCommonPassType::DOF,
+		false,
+		PixelShader,
+		Parameters);
 
 	return OutDOFDepthTexture;
 }
@@ -306,50 +347,16 @@ static void AddHairVisibilityFastResolveMSAAPass(
 	FHairVisibilityFastResolvePS::FPermutationDomain PermutationVector;
 	PermutationVector.Set<FHairVisibilityFastResolvePS::FMSAACount>(MSAASampleCount == 4 ? 4 : 8);
 
-	TShaderMapRef<FPostProcessVS> VertexShader(View.ShaderMap);
 	TShaderMapRef<FHairVisibilityFastResolvePS> PixelShader(View.ShaderMap, PermutationVector);
-	const FGlobalShaderMap* GlobalShaderMap = View.ShaderMap;
-	const FIntRect Viewport = View.ViewRect;
-
-	{
-		ClearUnusedGraphResources(PixelShader, Parameters);
-
-		GraphBuilder.AddPass(
-			RDG_EVENT_NAME("HairStrandsVisibilityMarkTAAFastResolve"),
-			Parameters,
-			ERDGPassFlags::Raster,
-			[Parameters, VertexShader, PixelShader, Viewport, Resolution](FRHICommandList& RHICmdList)
-		{
-			FGraphicsPipelineStateInitializer GraphicsPSOInit;
-			RHICmdList.ApplyCachedRenderTargets(GraphicsPSOInit);
-			GraphicsPSOInit.BlendState = TStaticBlendState<CW_RGBA, BO_Add, BF_One, BF_Zero, BO_Add, BF_One, BF_Zero>::GetRHI();
-			GraphicsPSOInit.RasterizerState = TStaticRasterizerState<>::GetRHI();	
-			GraphicsPSOInit.DepthStencilState = TStaticDepthStencilState<
-				false, CF_Always, 
-				true, CF_Always, SO_Keep, SO_Keep, SO_Replace,
-				false, CF_Always, SO_Keep, SO_Keep, SO_Keep,
-				STENCIL_TEMPORAL_RESPONSIVE_AA_MASK, STENCIL_TEMPORAL_RESPONSIVE_AA_MASK>::GetRHI();
-			GraphicsPSOInit.BoundShaderState.VertexDeclarationRHI = GFilterVertexDeclaration.VertexDeclarationRHI;
-			GraphicsPSOInit.BoundShaderState.VertexShaderRHI = VertexShader.GetVertexShader();
-			GraphicsPSOInit.BoundShaderState.PixelShaderRHI = PixelShader.GetPixelShader();
-			GraphicsPSOInit.PrimitiveType = PT_TriangleList;
-			SetGraphicsPipelineState(RHICmdList, GraphicsPSOInit);
-
-			RHICmdList.SetStencilRef(STENCIL_TEMPORAL_RESPONSIVE_AA_MASK);
-			RHICmdList.SetViewport(Viewport.Min.X, Viewport.Min.Y, 0.0f, Viewport.Max.X, Viewport.Max.Y, 1.0f);
-			SetShaderParameters(RHICmdList, PixelShader, PixelShader.GetPixelShader(), *Parameters);
-			DrawRectangle(
-				RHICmdList,
-				0, 0,
-				Viewport.Width(), Viewport.Height(),
-				Viewport.Min.X, Viewport.Min.Y,
-				Viewport.Width(), Viewport.Height(),
-				Viewport.Size(),
-				Resolution,
-				VertexShader,
-				EDRF_UseTriangleOptimization);
-		});
-	}
+	InternalCommonDrawPass(
+		GraphBuilder,
+		RDG_EVENT_NAME("HairStrandsVisibilityMarkTAAFastResolve"),
+		View,
+		Resolution,
+		EHairStrandsCommonPassType::TAAFastResolve,
+		false,
+		PixelShader,
+		Parameters);
 }
 
 class FHairVisibilityFastResolveMaskPS : public FGlobalShader
@@ -396,50 +403,16 @@ static void AddHairVisibilityFastResolveMaskPass(
 		ERenderTargetLoadAction::ELoad,
 		FExclusiveDepthStencil::DepthNop_StencilWrite);
 
-	TShaderMapRef<FPostProcessVS> VertexShader(View.ShaderMap);
 	TShaderMapRef<FHairVisibilityFastResolveMaskPS> PixelShader(View.ShaderMap);
-	const FGlobalShaderMap* GlobalShaderMap = View.ShaderMap;
-	const FIntRect Viewport = View.ViewRect;
-
-	{
-		ClearUnusedGraphResources(PixelShader, Parameters);
-
-		GraphBuilder.AddPass(
-			RDG_EVENT_NAME("HairStrandsVisibilityMarkTAAFastResolve"),
-			Parameters,
-			ERDGPassFlags::Raster,
-			[Parameters, VertexShader, PixelShader, Viewport, Resolution](FRHICommandList& RHICmdList)
-			{
-				FGraphicsPipelineStateInitializer GraphicsPSOInit;
-				RHICmdList.ApplyCachedRenderTargets(GraphicsPSOInit);
-				GraphicsPSOInit.BlendState = TStaticBlendState<CW_RGBA, BO_Add, BF_One, BF_Zero, BO_Add, BF_One, BF_Zero>::GetRHI();
-				GraphicsPSOInit.RasterizerState = TStaticRasterizerState<>::GetRHI();
-				GraphicsPSOInit.DepthStencilState = TStaticDepthStencilState<
-					false, CF_Always,
-					true, CF_Always, SO_Keep, SO_Keep, SO_Replace,
-					false, CF_Always, SO_Keep, SO_Keep, SO_Keep,
-					STENCIL_TEMPORAL_RESPONSIVE_AA_MASK, STENCIL_TEMPORAL_RESPONSIVE_AA_MASK>::GetRHI();
-				GraphicsPSOInit.BoundShaderState.VertexDeclarationRHI = GFilterVertexDeclaration.VertexDeclarationRHI;
-				GraphicsPSOInit.BoundShaderState.VertexShaderRHI = VertexShader.GetVertexShader();
-				GraphicsPSOInit.BoundShaderState.PixelShaderRHI = PixelShader.GetPixelShader();
-				GraphicsPSOInit.PrimitiveType = PT_TriangleList;
-				SetGraphicsPipelineState(RHICmdList, GraphicsPSOInit);
-
-				RHICmdList.SetStencilRef(STENCIL_TEMPORAL_RESPONSIVE_AA_MASK);
-				RHICmdList.SetViewport(Viewport.Min.X, Viewport.Min.Y, 0.0f, Viewport.Max.X, Viewport.Max.Y, 1.0f);
-				SetShaderParameters(RHICmdList, PixelShader, PixelShader.GetPixelShader(), *Parameters);
-				DrawRectangle(
-					RHICmdList,
-					0, 0,
-					Viewport.Width(), Viewport.Height(),
-					Viewport.Min.X, Viewport.Min.Y,
-					Viewport.Width(), Viewport.Height(),
-					Viewport.Size(),
-					Resolution,
-					VertexShader,
-					EDRF_UseTriangleOptimization);
-			});
-	}
+	InternalCommonDrawPass(
+		GraphBuilder,
+		RDG_EVENT_NAME("HairStrandsVisibilityMarkTAAFastResolve"),
+		View,
+		Resolution,
+		EHairStrandsCommonPassType::TAAFastResolve,
+		false,
+		PixelShader,
+		Parameters);
 }
 
 
@@ -538,53 +511,18 @@ static void AddHairVisibilityGBufferWritePass(
 			FExclusiveDepthStencil::DepthWrite_StencilNop);
 	}
 
-	TShaderMapRef<FPostProcessVS> VertexShader(View.ShaderMap);
 	FHairVisibilityGBufferWritePS::FPermutationDomain PermutationVector;
 	PermutationVector.Set<FHairVisibilityGBufferWritePS::FOutputType>(bWriteFullGBuffer ? 1 : 0);
 	TShaderMapRef<FHairVisibilityGBufferWritePS> PixelShader(View.ShaderMap, PermutationVector);
-	const FGlobalShaderMap* GlobalShaderMap = View.ShaderMap;
-	const FIntRect Viewport = View.ViewRect;
-	const FIntPoint Resolution = OutGBufferATexture->Desc.Extent;
-	ClearUnusedGraphResources(PixelShader, Parameters);
-
-	GraphBuilder.AddPass(
-		RDG_EVENT_NAME("HairStrandsGBufferWrite"),
-		Parameters,
-		ERDGPassFlags::Raster,
-		[Parameters, VertexShader, PixelShader, Viewport, Resolution, bWriteDepth](FRHICommandList& RHICmdList)
-		{
-			FGraphicsPipelineStateInitializer GraphicsPSOInit;
-			RHICmdList.ApplyCachedRenderTargets(GraphicsPSOInit);
-			GraphicsPSOInit.BlendState = TStaticBlendState<CW_RGBA, BO_Add, BF_One, BF_Zero, BO_Add, BF_One, BF_Zero>::GetRHI();
-			GraphicsPSOInit.RasterizerState = TStaticRasterizerState<>::GetRHI();
-			if (bWriteDepth)
-			{
-				GraphicsPSOInit.DepthStencilState = TStaticDepthStencilState<true, CF_Always>::GetRHI();
-			}
-			else
-			{
-				GraphicsPSOInit.DepthStencilState = TStaticDepthStencilState<false, CF_Always>::GetRHI();
-			}
-
-			GraphicsPSOInit.BoundShaderState.VertexDeclarationRHI = GFilterVertexDeclaration.VertexDeclarationRHI;
-			GraphicsPSOInit.BoundShaderState.VertexShaderRHI = VertexShader.GetVertexShader();
-			GraphicsPSOInit.BoundShaderState.PixelShaderRHI = PixelShader.GetPixelShader();
-			GraphicsPSOInit.PrimitiveType = PT_TriangleList;
-			SetGraphicsPipelineState(RHICmdList, GraphicsPSOInit);
-
-			RHICmdList.SetViewport(Viewport.Min.X, Viewport.Min.Y, 0.0f, Viewport.Max.X, Viewport.Max.Y, 1.0f);
-			SetShaderParameters(RHICmdList, PixelShader, PixelShader.GetPixelShader(), *Parameters);
-			DrawRectangle(
-				RHICmdList,
-				0, 0,
-				Viewport.Width(), Viewport.Height(),
-				Viewport.Min.X, Viewport.Min.Y,
-				Viewport.Width(), Viewport.Height(),
-				Viewport.Size(),
-				Resolution,
-				VertexShader,
-				EDRF_UseTriangleOptimization);
-		});
+	InternalCommonDrawPass(
+		GraphBuilder,
+		RDG_EVENT_NAME("HairStrandsVisibilityMarkTAAFastResolve"),
+		View,
+		OutGBufferATexture->Desc.Extent,
+		EHairStrandsCommonPassType::GBuffer,
+		bWriteDepth,
+		PixelShader,
+		Parameters);
 }
 
 static void InternalRenderHairComposition(
