Add tile composition support for all hair post process passes.

diff --git a/Engine/Source/Runtime/Renderer/Private/HairStrands/HairStrandsComposition.cpp b/Engine/Source/Runtime/Renderer/Private/HairStrands/HairStrandsComposition.cpp
--- a/Engine/Source/Runtime/Renderer/Private/HairStrands/HairStrandsComposition.cpp
+++ b/Engine/Source/Runtime/Renderer/Private/HairStrands/HairStrandsComposition.cpp
@@ -10,6 +10,7 @@
 #include "PostProcessing.h"
 #include "HairStrandsRendering.h"
 #include "HairStrandsScatter.h"
+#include "HairStrandsTile.h"
 #include "FogRendering.h"
 
 /////////////////////////////////////////////////////////////////////////////////////////
@@ -48,20 +49,30 @@ void InternalCommonDrawPass(
 	const FIntPoint Resolution,
 	const EHairStrandsCommonPassType Type,
 	const bool bWriteDepth,
+	const FHairStrandsTiles& TileData,
 	TPixelShader& PixelShader,
 	TPassParameter* PassParamters)
 {
-	ClearUnusedGraphResources(PixelShader, PassParamters);
+	//ClearUnusedGraphResources(PixelShader, PassParamters);
 
 	const FIntRect Viewport = View.ViewRect;
 
-	TShaderMapRef<FPostProcessVS> VertexShader(View.ShaderMap);
+	TShaderMapRef<FPostProcessVS> ScreenVertexShader(View.ShaderMap);
+	TShaderMapRef<FHairStrandsTilePassVS> TileVertexShader(View.ShaderMap);
+
+	const bool bUseTile = TileData.IsValid();
+	if (TileData.IsValid())
+	{
+		PassParamters->TileData = GetHairStrandsTileParameters(TileData);
+	}
 	GraphBuilder.AddPass(
 		Forward<FRDGEventName>(EventName),
 		PassParamters,
 		ERDGPassFlags::Raster,
-		[PassParamters, VertexShader, PixelShader, Viewport, Resolution, Type, bWriteDepth](FRHICommandList& RHICmdList)
+		[PassParamters, ScreenVertexShader, TileVertexShader, PixelShader, Viewport, Resolution, Type, bWriteDepth, bUseTile](FRHICommandList& RHICmdList)
 	{
+		FHairStrandsTilePassVS::FParameters ParametersVS = PassParamters->TileData;
+
 		FGraphicsPipelineStateInitializer GraphicsPSOInit;
 		RHICmdList.ApplyCachedRenderTargets(GraphicsPSOInit);
 
@@ -112,29 +123,38 @@ void InternalCommonDrawPass(
 			}
 		}
 
-
 		GraphicsPSOInit.BoundShaderState.VertexDeclarationRHI = GFilterVertexDeclaration.VertexDeclarationRHI;
-		GraphicsPSOInit.BoundShaderState.VertexShaderRHI = VertexShader.GetVertexShader();
+		GraphicsPSOInit.BoundShaderState.VertexShaderRHI = bUseTile ? TileVertexShader.GetVertexShader() : ScreenVertexShader.GetVertexShader();
 		GraphicsPSOInit.BoundShaderState.PixelShaderRHI = PixelShader.GetPixelShader();
-		GraphicsPSOInit.PrimitiveType = PT_TriangleList;
+		GraphicsPSOInit.PrimitiveType = bUseTile && PassParamters->TileData.bRectPrimitive > 0 ? PT_RectList : PT_TriangleList;
 		SetGraphicsPipelineState(RHICmdList, GraphicsPSOInit);
 
 		if (Type == EHairStrandsCommonPassType::TAAFastResolve)
 		{
 			RHICmdList.SetStencilRef(STENCIL_TEMPORAL_RESPONSIVE_AA_MASK);
 		}
-		RHICmdList.SetViewport(Viewport.Min.X, Viewport.Min.Y, 0.0f, Viewport.Max.X, Viewport.Max.Y, 1.0f);
+
 		SetShaderParameters(RHICmdList, PixelShader, PixelShader.GetPixelShader(), *PassParamters);
-		DrawRectangle(
-			RHICmdList,
-			0, 0,
-			Viewport.Width(), Viewport.Height(),
-			Viewport.Min.X, Viewport.Min.Y,
-			Viewport.Width(), Viewport.Height(),
-			Viewport.Size(),
-			Resolution,
-			VertexShader,
-			EDRF_UseTriangleOptimization);
+		if (bUseTile)
+		{
+			SetShaderParameters(RHICmdList, TileVertexShader, TileVertexShader.GetVertexShader(), ParametersVS);
+			RHICmdList.SetStreamSource(0, nullptr, 0);
+			RHICmdList.DrawPrimitiveIndirect(PassParamters->TileData.TileIndirectBuffer->GetRHI(), 0);
+		}
+		else
+		{
+			RHICmdList.SetViewport(Viewport.Min.X, Viewport.Min.Y, 0.0f, Viewport.Max.X, Viewport.Max.Y, 1.0f);
+			DrawRectangle(
+				RHICmdList,
+				0, 0,
+				Viewport.Width(), Viewport.Height(),
+				Viewport.Min.X, Viewport.Min.Y,
+				Viewport.Width(), Viewport.Height(),
+				Viewport.Size(),
+				Resolution,
+				ScreenVertexShader,
+				EDRF_UseTriangleOptimization);
+		}
 	});
 }
 
@@ -147,6 +167,7 @@ class FHairVisibilityComposeSamplePS : public FGlobalShader
 
 	BEGIN_SHADER_PARAMETER_STRUCT(FParameters, )
 		SHADER_PARAMETER_STRUCT_REF(FViewUniformShaderParameters, ViewUniformBuffer)
+		SHADER_PARAMETER_STRUCT_INCLUDE(FHairStrandsTilePassVS::FParameters, TileData)
 		SHADER_PARAMETER(FIntPoint, OutputResolution)
 		SHADER_PARAMETER(uint32, bComposeDofDepth)
 		SHADER_PARAMETER(uint32, bEmissiveEnable)
@@ -208,6 +229,7 @@ static void AddHairVisibilityComposeSamplePass(
 		OutColorTexture->Desc.Extent,
 		EHairStrandsCommonPassType::Composition,
 		false,
+		VisibilityData.TileData,
 		PixelShader,
 		Parameters);
 }
@@ -222,6 +244,7 @@ class FHairDOFDepthPS : public FGlobalShader
 
 	BEGIN_SHADER_PARAMETER_STRUCT(FParameters, )
 		SHADER_PARAMETER_STRUCT_REF(FViewUniformShaderParameters, ViewUniformBuffer)
+		SHADER_PARAMETER_STRUCT_INCLUDE(FHairStrandsTilePassVS::FParameters, TileData)
 		SHADER_PARAMETER_RDG_TEXTURE(Texture2D, HairSampleCount)
 		SHADER_PARAMETER_RDG_TEXTURE(Texture2D, HairCategorizationTexture)
 		SHADER_PARAMETER_RDG_TEXTURE(Texture2D, HairVisibilityNodeOffsetAndCount)
@@ -278,6 +301,7 @@ static FRDGTextureRef AddHairDOFDepthPass(
 		OutputResolution,
 		EHairStrandsCommonPassType::DOF,
 		false,
+		VisibilityData.TileData,
 		PixelShader,
 		Parameters);
 
@@ -295,6 +319,7 @@ class FHairVisibilityFastResolvePS : public FGlobalShader
 	using FPermutationDomain = TShaderPermutationDomain<FMSAACount>;
 
 	BEGIN_SHADER_PARAMETER_STRUCT(FParameters, )
+		SHADER_PARAMETER_STRUCT_INCLUDE(FHairStrandsTilePassVS::FParameters, TileData)
 		SHADER_PARAMETER(float, VelocityThreshold)
 		SHADER_PARAMETER_RDG_TEXTURE(Texture2D, HairVisibilityVelocityTexture)
 		RENDER_TARGET_BINDING_SLOTS()
@@ -316,6 +341,7 @@ static void AddHairVisibilityFastResolveMSAAPass(
 	FRDGBuilder& GraphBuilder,
 	const FViewInfo& View,
 	const FRDGTextureRef& HairVisibilityVelocityTexture,
+	const FHairStrandsTiles& TileData,
 	FRDGTextureRef& OutDepthTexture)
 {
 	const FIntPoint Resolution = OutDepthTexture->Desc.Extent;
@@ -355,6 +381,7 @@ static void AddHairVisibilityFastResolveMSAAPass(
 		Resolution,
 		EHairStrandsCommonPassType::TAAFastResolve,
 		false,
+		TileData,
 		PixelShader,
 		Parameters);
 }
@@ -365,6 +392,7 @@ class FHairVisibilityFastResolveMaskPS : public FGlobalShader
 	SHADER_USE_PARAMETER_STRUCT(FHairVisibilityFastResolveMaskPS, FGlobalShader);
 
 	BEGIN_SHADER_PARAMETER_STRUCT(FParameters, )
+		SHADER_PARAMETER_STRUCT_INCLUDE(FHairStrandsTilePassVS::FParameters, TileData)
 		SHADER_PARAMETER_RDG_TEXTURE(Texture2D, ResolveMaskTexture)
 		RENDER_TARGET_BINDING_SLOTS()
 	END_SHADER_PARAMETER_STRUCT()
@@ -385,6 +413,7 @@ static void AddHairVisibilityFastResolveMaskPass(
 	FRDGBuilder& GraphBuilder,
 	const FViewInfo& View,
 	const FRDGTextureRef& HairResolveMaskTexture,
+	const FHairStrandsTiles& TileData,
 	FRDGTextureRef& OutDepthTexture)
 {
 	const FIntPoint Resolution = OutDepthTexture->Desc.Extent;
@@ -411,6 +440,7 @@ static void AddHairVisibilityFastResolveMaskPass(
 		Resolution,
 		EHairStrandsCommonPassType::TAAFastResolve,
 		false,
+		TileData,
 		PixelShader,
 		Parameters);
 }
@@ -428,6 +458,7 @@ class FHairVisibilityGBufferWritePS : public FGlobalShader
 	using FPermutationDomain = TShaderPermutationDomain<FOutputType>;
 
 	BEGIN_SHADER_PARAMETER_STRUCT(FParameters, )
+		SHADER_PARAMETER_STRUCT_INCLUDE(FHairStrandsTilePassVS::FParameters, TileData)
 		SHADER_PARAMETER(uint32, bWriteDummyData)
 		SHADER_PARAMETER_RDG_TEXTURE(Texture2D, CategorizationTexture)
 		SHADER_PARAMETER_RDG_TEXTURE(Texture2D, NodeIndex)
@@ -463,6 +494,7 @@ static void AddHairVisibilityGBufferWritePass(
 	const FRDGTextureRef& CategorizationTexture,
 	const FRDGTextureRef& NodeIndex,
 	const FRDGBufferRef& NodeData,
+	const FHairStrandsTiles& TileData,
 	FRDGTextureRef OutGBufferATexture,
 	FRDGTextureRef OutGBufferBTexture,
 	FRDGTextureRef OutGBufferCTexture,
@@ -521,6 +553,7 @@ static void AddHairVisibilityGBufferWritePass(
 		OutGBufferATexture->Desc.Extent,
 		EHairStrandsCommonPassType::GBuffer,
 		bWriteDepth,
+		TileData,
 		PixelShader,
 		Parameters);
 }
@@ -574,6 +607,7 @@ static void InternalRenderHairComposition(
 						GraphBuilder,
 						View,
 						VisibilityData.VelocityTexture,
+						VisibilityData.TileData,
 						SceneDepthTexture);
 				}
 				else if (VisibilityData.ResolveMaskTexture)
@@ -582,6 +616,7 @@ static void InternalRenderHairComposition(
 						GraphBuilder,
 						View,
 						VisibilityData.ResolveMaskTexture,
+						VisibilityData.TileData,
 						SceneDepthTexture);
 				}
 
@@ -605,6 +640,7 @@ static void InternalRenderHairComposition(
 							VisibilityData.CategorizationTexture,
 							VisibilityData.NodeIndex,
 							VisibilityData.NodeData,
+							VisibilityData.TileData,
 							GBufferATexture,
 							GBufferBTexture,
 							nullptr,
@@ -621,6 +657,7 @@ static void InternalRenderHairComposition(
 							VisibilityData.CategorizationTexture,
 							VisibilityData.NodeIndex,
 							VisibilityData.NodeData,
+							VisibilityData.TileData,
 							GBufferATexture,
 							GBufferBTexture,
 							GBufferCTexture,
