Fix out of bounds access in culled instance flags when dynamic primitives are added after culling.

diff --git a/Engine/Shaders/Private/InstanceCulling/BuildInstanceDrawCommands.usf b/Engine/Shaders/Private/InstanceCulling/BuildInstanceDrawCommands.usf
--- a/Engine/Shaders/Private/InstanceCulling/BuildInstanceDrawCommands.usf
+++ b/Engine/Shaders/Private/InstanceCulling/BuildInstanceDrawCommands.usf
@@ -36,6 +36,8 @@ StructuredBuffer<uint> PrimitiveIds;
 #if CULL_INSTANCES
 StructuredBuffer<uint> VisibleInstanceFlags;
 uint NumInstanceFlagWords;
+uint NumCulledInstances;
+uint NumCulledViews;
 #endif
 
 uint NumPrimitiveIds;
@@ -71,12 +73,14 @@ uint TranslatePrimitiveId(uint PrimitiveIdIn)
 bool TestInstanceVisibilityBit(uint InstanceId, uint ViewId)
 {
 #if CULL_INSTANCES
-	uint ViewWordOffset = ViewId * NumInstanceFlagWords;
-	uint WordOffset = ViewWordOffset + InstanceId / 32U;
-	return (VisibleInstanceFlags[WordOffset] & (1U << (InstanceId % 32U))) != 0U;
-#else
-	return true;
+	if (InstanceId < NumCulledInstances && ViewId < NumCulledViews)
+	{
+		uint ViewWordOffset = ViewId * NumInstanceFlagWords;
+		uint WordOffset = ViewWordOffset + InstanceId / 32U;
+		return (VisibleInstanceFlags[WordOffset] & (1U << (InstanceId % 32U))) != 0U;
+	}
 #endif
+	return true;
 }
 
 
diff --git a/Engine/Source/Runtime/Renderer/Private/InstanceCulling/InstanceCullingContext.cpp b/Engine/Source/Runtime/Renderer/Private/InstanceCulling/InstanceCullingContext.cpp
--- a/Engine/Source/Runtime/Renderer/Private/InstanceCulling/InstanceCullingContext.cpp
+++ b/Engine/Source/Runtime/Renderer/Private/InstanceCulling/InstanceCullingContext.cpp
@@ -133,6 +133,8 @@ public:
 		SHADER_PARAMETER(int32, NumInstanceRuns)
 		SHADER_PARAMETER(int32, NumCommands)
 		SHADER_PARAMETER(uint32, NumInstanceFlagWords)
+		SHADER_PARAMETER(uint32, NumCulledInstances)
+		SHADER_PARAMETER(uint32, NumCulledViews)
 		SHADER_PARAMETER(int32, NumViewIds)
 
 		SHADER_PARAMETER(int32, DynamicPrimitiveIdOffset)
@@ -228,6 +230,8 @@ public:
 		SHADER_PARAMETER(int32, NumInstanceRuns)
 		SHADER_PARAMETER(int32, NumCommands)
 		SHADER_PARAMETER(uint32, NumInstanceFlagWords)
+		SHADER_PARAMETER(uint32, NumCulledInstances)
+		SHADER_PARAMETER(uint32, NumCulledViews)
 		SHADER_PARAMETER(int32, NumViewIds)
 
 		SHADER_PARAMETER(int32, DynamicPrimitiveIdOffset)
@@ -294,6 +298,8 @@ public:
 		SHADER_PARAMETER(int32, NumInstanceRuns)		
 		SHADER_PARAMETER(int32, NumCommands)
 		SHADER_PARAMETER(uint32, NumInstanceFlagWords)
+		SHADER_PARAMETER(uint32, NumCulledInstances)
+		SHADER_PARAMETER(uint32, NumCulledViews)
 		SHADER_PARAMETER(int32, NumViewIds)
 
 		SHADER_PARAMETER(int32, DynamicPrimitiveIdOffset)
@@ -324,8 +330,9 @@ void FInstanceCullingContext::BuildRenderingCommands(FRDGBuilder& GraphBuilder,
 	// If there is no manager, then there is no data on culling, so set flag to skip that and ignore buffers.
 	FRDGBufferRef VisibleInstanceFlagsRDG = InstanceCullingManager != nullptr ? InstanceCullingManager->CullingIntermediate.VisibleInstanceFlags : nullptr;
 	const bool bCullInstances = InstanceCullingManager != nullptr;
-	const int32 NumInstances = InstanceCullingManager != nullptr ? InstanceCullingManager->CullingIntermediate.NumInstances : 0;
-	int32 NumInstanceFlagWords = FMath::DivideAndRoundUp(NumInstances, int32(sizeof(uint32) * 8));
+	const uint32 NumCulledInstances = InstanceCullingManager != nullptr ? uint32(InstanceCullingManager->CullingIntermediate.NumInstances) : 0U;
+	const uint32 NumCulledViews = InstanceCullingManager != nullptr ? uint32(InstanceCullingManager->CullingIntermediate.NumViews) : 0U;
+	uint32 NumInstanceFlagWords = FMath::DivideAndRoundUp(NumCulledInstances, uint32(sizeof(uint32) * 8));
 
 #if ENABLE_DETERMINISTIC_INSTANCE_CULLING
 	// 1. Compute output sizes for all commands
@@ -362,7 +369,9 @@ void FInstanceCullingContext::BuildRenderingCommands(FRDGBuilder& GraphBuilder,
 		PassParameters->NumCommands = CullingCommands.Num();
 		PassParameters->VisibleInstanceFlags = VisibleInstanceFlagsRDG ? GraphBuilder.CreateSRV(VisibleInstanceFlagsRDG) : nullptr;
 
-		PassParameters->NumInstanceFlagWords = uint32(NumInstanceFlagWords);
+		PassParameters->NumInstanceFlagWords = NumInstanceFlagWords;
+		PassParameters->NumCulledInstances = NumCulledInstances;
+		PassParameters->NumCulledViews = NumCulledViews;
 
 		FComputeInstanceIdOutputSizeCs::FPermutationDomain PermutationVector;
 		PermutationVector.Set<FComputeInstanceIdOutputSizeCs::FCullInstancesDim>(bCullInstances);
@@ -437,7 +446,9 @@ void FInstanceCullingContext::BuildRenderingCommands(FRDGBuilder& GraphBuilder,
 		PassParameters->NumInstanceRuns = InstanceRuns.Num();
 		PassParameters->NumCommands = CullingCommands.Num();
 		PassParameters->VisibleInstanceFlags = VisibleInstanceFlagsRDG ? GraphBuilder.CreateSRV(VisibleInstanceFlagsRDG) : nullptr;
-		PassParameters->NumInstanceFlagWords = uint32(NumInstanceFlagWords);
+		PassParameters->NumInstanceFlagWords = NumInstanceFlagWords;
+		PassParameters->NumCulledInstances = NumCulledInstances;
+		PassParameters->NumCulledViews = NumCulledViews;
 
 		FOutputInstanceIdsAtOffsetCs::FPermutationDomain PermutationVector;
 		// NOTE: this also switches between legacy buffer and RDG for Id output
@@ -500,7 +511,9 @@ void FInstanceCullingContext::BuildRenderingCommands(FRDGBuilder& GraphBuilder,
 	PassParameters->NumCommands = CullingCommands.Num();
 	PassParameters->VisibleInstanceFlags = VisibleInstanceFlagsRDG ? GraphBuilder.CreateSRV(VisibleInstanceFlagsRDG) : nullptr;
 
-	PassParameters->NumInstanceFlagWords = uint32(NumInstanceFlagWords);
+	PassParameters->NumInstanceFlagWords = NumInstanceFlagWords;
+	PassParameters->NumCulledInstances = NumCulledInstances;
+	PassParameters->NumCulledViews = NumCulledViews;
 
 	FBuildInstanceIdBufferAndCommandsFromPrimitiveIdsCs::FPermutationDomain PermutationVector;
 	// NOTE: this also switches between legacy buffer and RDG for Id output
@@ -612,6 +625,8 @@ void FInstanceCullingContext::BuildRenderingCommands(FRDGBuilder& GraphBuilder,
 
 	int32 NumInstanceFlagWords = FMath::DivideAndRoundUp(Intermediate.NumInstances, int32(sizeof(uint32) * 8));
 	PassParameters->NumInstanceFlagWords = uint32(NumInstanceFlagWords);
+	PassParameters->NumCulledInstances = uint32(Intermediate.NumInstances);
+	PassParameters->NumCulledViews = uint32(Intermediate.NumViews);
 
 
 	FBuildInstanceIdBufferAndCommandsFromPrimitiveIdsCs::FPermutationDomain PermutationVector;
diff --git a/Engine/Source/Runtime/Renderer/Private/InstanceCulling/InstanceCullingManager.cpp b/Engine/Source/Runtime/Renderer/Private/InstanceCulling/InstanceCullingManager.cpp
--- a/Engine/Source/Runtime/Renderer/Private/InstanceCulling/InstanceCullingManager.cpp
+++ b/Engine/Source/Runtime/Renderer/Private/InstanceCulling/InstanceCullingManager.cpp
@@ -103,7 +103,9 @@ void FInstanceCullingManager::CullInstances(FRDGBuilder& GraphBuilder, FGPUScene
 {
 #if GPUCULL_TODO
 	TRACE_CPUPROFILER_EVENT_SCOPE(FInstanceCullingManager::CullInstances);
-	RDG_EVENT_SCOPE(GraphBuilder, "CullInstances");
+	int32 NumViews = CullingViews.Num();
+	int32 NumInstances = GPUScene.InstanceDataAllocator.GetMaxSize();
+	RDG_EVENT_SCOPE(GraphBuilder, "CullInstances [%d Views X %d Instances]", NumViews, NumInstances);
 
 	check(!CullingIntermediate.InstanceIdOutOffsetBuffer);
 	check(!CullingIntermediate.VisibleInstanceFlags);
@@ -113,8 +115,6 @@ void FInstanceCullingManager::CullInstances(FRDGBuilder& GraphBuilder, FGPUScene
 
 	CullingIntermediate.InstanceIdOutOffsetBuffer = CreateStructuredBuffer(GraphBuilder, TEXT("InstanceCulling.OutputOffsetBufferOut"), NullArray);
 
-	int32 NumViews = CullingViews.Num();
-	int32 NumInstances = GPUScene.InstanceDataAllocator.GetMaxSize();
 	int32 NumInstanceFlagWords = FMath::DivideAndRoundUp(NumInstances, int32(sizeof(uint32) * 8));
 
 	CullingIntermediate.NumInstances = NumInstances;
