Add hair tile data & tile count buffer onto hair strands view uniform buffer, to ease indirect dispatch setup with compute pass.

diff --git a/Engine/Shaders/Private/HairStrands/HairStrandsVisibilityTile.usf b/Engine/Shaders/Private/HairStrands/HairStrandsVisibilityTile.usf
--- a/Engine/Shaders/Private/HairStrands/HairStrandsVisibilityTile.usf
+++ b/Engine/Shaders/Private/HairStrands/HairStrandsVisibilityTile.usf
@@ -10,10 +10,9 @@
 ////////////////////////////////////////////////////////////////////////////////////////////////////////////
 
 #if SHADER_TILE_GENERATION
-int bRectPrimitive;
 int2 ViewResolution;
 Texture2D<uint4> InputTexture;
-RWBuffer<uint> TileIndirectBuffer;
+RWBuffer<uint> TileCountBuffer;
 RWBuffer<uint2> TileDataBuffer;
 
 groupshared bool s_HasTileHair[GROUP_THREAD_COUNT];
@@ -21,11 +20,6 @@ groupshared bool s_HasTileHair[GROUP_THREAD_COUNT];
 [numthreads(TILE_SIZE, TILE_SIZE, 1)]
 void TileMainCS(uint2 DispatchThreadId : SV_DispatchThreadID, uint LinearIndex : SV_GroupIndex, uint3 GroupId : SV_GroupID)
 {
-	if (all(DispatchThreadId == 0))
-	{
-		TileIndirectBuffer[0] = bRectPrimitive > 0 ? 4 : 6;
-	}
-
 	const uint2 PixelCoord = DispatchThreadId;
 
 	bool bHasHairData = false;
@@ -66,7 +60,7 @@ void TileMainCS(uint2 DispatchThreadId : SV_DispatchThreadID, uint LinearIndex :
 	if (LinearIndex < 1 && (s_HasTileHair[LinearIndex] || s_HasTileHair[LinearIndex + 1]))
 	{
 		uint WriteToIndex;
-		InterlockedAdd(TileIndirectBuffer[1], 1, WriteToIndex);
+		InterlockedAdd(TileCountBuffer[0], 1, WriteToIndex);
 		TileDataBuffer[WriteToIndex] = GroupId.xy;
 	}
 }
@@ -76,7 +70,9 @@ void TileMainCS(uint2 DispatchThreadId : SV_DispatchThreadID, uint LinearIndex :
 
 #if SHADER_TILE_COPY_ARGS
 uint GroupSize;
-Buffer<uint> TileIndirectDrawBuffer;
+uint bRectPrimitive;
+Buffer<uint> TileCountBuffer;
+RWBuffer<uint> TileIndirectDrawBuffer;
 RWBuffer<uint> TileIndirectDispatchBuffer;
 
 [numthreads(1, 1, 1)]
@@ -84,7 +80,15 @@ void MainCS(uint2 DispatchThreadId : SV_DispatchThreadID)
 {
 	if (all(DispatchThreadId == 0))
 	{
-		TileIndirectDispatchBuffer[0] = (TileIndirectDrawBuffer[1] + GroupSize-1) / GroupSize;
+		const uint TileCount = TileCountBuffer[0];
+
+		// Indirect draw
+		TileIndirectDrawBuffer[0] = bRectPrimitive > 0 ? 4 : 6;
+		TileIndirectDrawBuffer[1] = TileCount;
+		TileIndirectDrawBuffer[2] = 1;
+
+		// Indirect dispatch
+		TileIndirectDispatchBuffer[0] = (TileCount + GroupSize-1) / GroupSize;
 		TileIndirectDispatchBuffer[1] = 1;
 		TileIndirectDispatchBuffer[2] = 1;
 	}
@@ -130,4 +134,4 @@ void MainVS(
 	Out.Position = float4(Out.UV * float2(2.0f, -2.0f) + float2(-1.0, 1.0f), 0.5f, 1.0f);
 }
 
-#endif //SHADER_TILE_VS
\ No newline at end of file
+#endif //SHADER_TILE_VS
diff --git a/Engine/Source/Runtime/Renderer/Private/HairStrands/HairStrandsData.h b/Engine/Source/Runtime/Renderer/Private/HairStrands/HairStrandsData.h
--- a/Engine/Source/Runtime/Renderer/Private/HairStrands/HairStrandsData.h
+++ b/Engine/Source/Runtime/Renderer/Private/HairStrands/HairStrandsData.h
@@ -26,6 +26,8 @@ BEGIN_GLOBAL_SHADER_PARAMETER_STRUCT(FHairStrandsViewUniformParameters, )
 	SHADER_PARAMETER_RDG_TEXTURE(Texture2D<uint>, HairSampleCount)						// Total count of hair sample, in sample space
 	SHADER_PARAMETER_RDG_BUFFER_SRV(StructuredBuffer<FPackedHairSample>, HairSampleData)// Sample data (coverage, tangent, base color, ...), in sample space // HAIRSTRANDS_TODO: change this to be a uint4 so that we don't have to include the type for generated contant buffer
 	SHADER_PARAMETER_RDG_BUFFER_SRV(Buffer<uint2>, HairSampleCoords)					// Screen pixel coordinate of each sample, in sample space
+	SHADER_PARAMETER_RDG_BUFFER_SRV(Buffer<uint2>, HairTileData)						// Tile coords (RG16F)
+	SHADER_PARAMETER_RDG_BUFFER_SRV(Buffer<uint>,  HairTileCount)						// Tile total count
 END_GLOBAL_SHADER_PARAMETER_STRUCT()
 
 ////////////////////////////////////////////////////////////////////////////////////
@@ -41,6 +43,7 @@ struct FHairStrandsTiles
 
 	FRDGBufferSRVRef	TileDataSRV = nullptr;
 	FRDGBufferRef		TileDataBuffer = nullptr;
+	FRDGBufferRef		TileCountBuffer = nullptr;
 	FRDGBufferRef		TileIndirectDrawBuffer = nullptr;
 	FRDGBufferRef		TileIndirectDispatchBuffer = nullptr;
 
diff --git a/Engine/Source/Runtime/Renderer/Private/HairStrands/HairStrandsRendering.cpp b/Engine/Source/Runtime/Renderer/Private/HairStrands/HairStrandsRendering.cpp
--- a/Engine/Source/Runtime/Renderer/Private/HairStrands/HairStrandsRendering.cpp
+++ b/Engine/Source/Runtime/Renderer/Private/HairStrands/HairStrandsRendering.cpp
@@ -20,26 +20,48 @@ static TRDGUniformBufferRef<FHairStrandsViewUniformParameters> InternalCreateHai
 		Parameters->HairSampleCoords = GraphBuilder.CreateSRV(In->NodeCoord, FHairStrandsVisibilityData::NodeCoordFormat);
 		Parameters->HairSampleCount = In->NodeCount;
 		Parameters->HairSampleViewportResolution = In->SampleLightingViewportResolution;
+
+		if (In->TileData.IsValid())
+		{
+			Parameters->HairTileData = In->TileData.TileDataSRV;
+			Parameters->HairTileCount = GraphBuilder.CreateSRV(In->TileData.TileCountBuffer, PF_R32_UINT);
+		}
+		else
+		{
+			FRDGBufferRef DummyBuffer = GraphBuilder.CreateBuffer(FRDGBufferDesc::CreateBufferDesc(4, 1), TEXT("Hair.DummyBuffer"));
+			AddClearUAVPass(GraphBuilder, GraphBuilder.CreateUAV(DummyBuffer, PF_R16G16_UINT), 0);
+			FRDGBufferSRVRef DummyBufferR32SRV = GraphBuilder.CreateSRV(DummyBuffer, PF_R32_UINT);
+			FRDGBufferSRVRef DummyBufferRG16SRV = GraphBuilder.CreateSRV(DummyBuffer, PF_R16G16_UINT);
+			Parameters->HairTileData = DummyBufferR32SRV;
+			Parameters->HairTileCount = DummyBufferRG16SRV;
+		}
 	}
 	else
 	{
-		FRDGBufferRef DummyBuffer = GraphBuilder.CreateBuffer(FRDGBufferDesc::CreateStructuredDesc(4,1), TEXT("Hair.DummyBuffer"));
+		FRDGBufferRef DummyBuffer = GraphBuilder.CreateBuffer(FRDGBufferDesc::CreateBufferDesc(4,1), TEXT("Hair.DummyBuffer"));
 		FRDGBufferRef DummyNodeBuffer = GraphBuilder.CreateBuffer(FRDGBufferDesc::CreateStructuredDesc(20, 1), TEXT("Hair.DummyNodeBuffer"));
+
 		FRDGTextureRef BlackTexture = GSystemTextures.GetBlackDummy(GraphBuilder);
 		FRDGTextureRef ZeroR32_UINT = GSystemTextures.GetZeroUIntDummy(GraphBuilder);
 		FRDGTextureRef ZeroRGBA16_UINT = GSystemTextures.GetZeroUShort4Dummy(GraphBuilder);
 		FRDGTextureRef FarDepth = GSystemTextures.GetDepthDummy(GraphBuilder);
 
 		AddClearUAVPass(GraphBuilder, GraphBuilder.CreateUAV(DummyNodeBuffer), 0);
-		AddClearUAVPass(GraphBuilder, GraphBuilder.CreateUAV(DummyBuffer), 0);
+		AddClearUAVPass(GraphBuilder, GraphBuilder.CreateUAV(DummyBuffer, PF_R16G16_UINT), 0);
+
+		FRDGBufferSRVRef DummyBufferR32SRV = GraphBuilder.CreateSRV(DummyBuffer, PF_R32_UINT);
+		FRDGBufferSRVRef DummyBufferRG16SRV = GraphBuilder.CreateSRV(DummyBuffer, PF_R16G16_UINT);
 
 		Parameters->HairOnlyDepthTexture = FarDepth;
 		Parameters->HairCategorizationTexture = ZeroRGBA16_UINT;
 		Parameters->HairSampleCount = ZeroR32_UINT;
 		Parameters->HairSampleOffset = ZeroR32_UINT;
-		Parameters->HairSampleCoords = GraphBuilder.CreateSRV(DummyBuffer);
+		Parameters->HairSampleCoords = DummyBufferRG16SRV;
 		Parameters->HairSampleData	 = GraphBuilder.CreateSRV(DummyNodeBuffer);
 		Parameters->HairSampleViewportResolution = FIntPoint(0, 0);
+
+		Parameters->HairTileData = DummyBufferR32SRV;
+		Parameters->HairTileCount = DummyBufferRG16SRV;
 	}
 
 	return GraphBuilder.CreateUniformBuffer(Parameters);
diff --git a/Engine/Source/Runtime/Renderer/Private/HairStrands/HairStrandsTile.cpp b/Engine/Source/Runtime/Renderer/Private/HairStrands/HairStrandsTile.cpp
--- a/Engine/Source/Runtime/Renderer/Private/HairStrands/HairStrandsTile.cpp
+++ b/Engine/Source/Runtime/Renderer/Private/HairStrands/HairStrandsTile.cpp
@@ -27,7 +27,7 @@ FHairStrandsTilePassVS::FParameters GetHairStrandsTileParameters(const FHairStra
 }
 
 ////////////////////////////////////////////////////////////////////////////////////////////////////////////
-// Generate indirect dispatch buffer based on the indirect draw buffer
+// Generate indirect draw and indirect dispatch buffers
 
 class FHairStrandsTileCopyArgsPassCS : public FGlobalShader
 {
@@ -39,7 +39,9 @@ class FHairStrandsTileCopyArgsPassCS : public FGlobalShader
 	BEGIN_SHADER_PARAMETER_STRUCT(FParameters, )
 		SHADER_PARAMETER_STRUCT_REF(FViewUniformShaderParameters, ViewUniformBuffer)
 		SHADER_PARAMETER(uint32, GroupSize)
-		SHADER_PARAMETER_RDG_BUFFER_SRV(Buffer, TileIndirectDrawBuffer)
+		SHADER_PARAMETER(uint32, bRectPrimitive)
+		SHADER_PARAMETER_RDG_BUFFER_SRV(Buffer, TileCountBuffer)
+		SHADER_PARAMETER_RDG_BUFFER_UAV(RWBuffer, TileIndirectDrawBuffer)
 		SHADER_PARAMETER_RDG_BUFFER_UAV(RWBuffer, TileIndirectDispatchBuffer)
 	END_SHADER_PARAMETER_STRUCT()
 
@@ -65,7 +67,9 @@ void AddHairStrandsCopyArgsTilesPass(
 	TShaderMapRef<FHairStrandsTileCopyArgsPassCS> ComputeShader(View.ShaderMap);
 	FHairStrandsTileCopyArgsPassCS::FParameters* PassParameters = GraphBuilder.AllocParameters<FHairStrandsTileCopyArgsPassCS::FParameters>();
 	PassParameters->GroupSize = FHairStrandsTiles::GroupSize;
-	PassParameters->TileIndirectDrawBuffer = GraphBuilder.CreateSRV(TileData.TileIndirectDrawBuffer, PF_R32_UINT);
+	PassParameters->bRectPrimitive = TileData.bRectPrimitive ? 1 : 0;
+	PassParameters->TileCountBuffer = GraphBuilder.CreateSRV(TileData.TileCountBuffer, PF_R32_UINT);
+	PassParameters->TileIndirectDrawBuffer = GraphBuilder.CreateUAV(TileData.TileIndirectDrawBuffer, PF_R32_UINT);
 	PassParameters->TileIndirectDispatchBuffer = GraphBuilder.CreateUAV(TileData.TileIndirectDispatchBuffer, PF_R32_UINT);
 
 	FComputeShaderUtils::AddPass(
@@ -88,10 +92,9 @@ class FHairStrandsTileGenerationPassCS : public FGlobalShader
 
 	BEGIN_SHADER_PARAMETER_STRUCT(FParameters, )
 		SHADER_PARAMETER_STRUCT_REF(FViewUniformShaderParameters, ViewUniformBuffer)
-		SHADER_PARAMETER(int32, bRectPrimitive)
 		SHADER_PARAMETER(FIntPoint, ViewResolution)
 		SHADER_PARAMETER_RDG_TEXTURE(Texture2D<uint4>, InputTexture)
-		SHADER_PARAMETER_RDG_BUFFER_UAV(RWBuffer, TileIndirectBuffer)
+		SHADER_PARAMETER_RDG_BUFFER_UAV(RWBuffer, TileCountBuffer)
 		SHADER_PARAMETER_RDG_BUFFER_UAV(RWBuffer, TileDataBuffer)
 	END_SHADER_PARAMETER_STRUCT()
 
@@ -121,28 +124,27 @@ FHairStrandsTiles AddHairStrandsGenerateTilesPass(
 	Out.TileCount = TileResolution.X * TileResolution.Y;
 	Out.Resolution = InputResolution;
 	Out.bRectPrimitive = GRHISupportsRectTopology;
+	Out.TileCountBuffer = GraphBuilder.CreateBuffer(FRDGBufferDesc::CreateBufferDesc(4, 1), TEXT("Hair.TileCountBuffer"));
 	Out.TileIndirectDrawBuffer = GraphBuilder.CreateBuffer(FRDGBufferDesc::CreateIndirectDesc<FRHIDrawIndirectParameters>(1), TEXT("Hair.TileIndirectDrawBuffer"));
 	Out.TileIndirectDispatchBuffer = GraphBuilder.CreateBuffer(FRDGBufferDesc::CreateIndirectDesc<FRHIDispatchIndirectParameters>(1), TEXT("Hair.TileIndirectDispatchBuffer"));
 	Out.TileDataBuffer = GraphBuilder.CreateBuffer(FRDGBufferDesc::CreateBufferDesc(sizeof(uint32), Out.TileCount), TEXT("Hair.TileDataBuffer"));
 
-	FRDGBufferUAVRef TileIndirectUAV = GraphBuilder.CreateUAV(Out.TileIndirectDrawBuffer, PF_R32_UINT);
-	AddClearUAVPass(GraphBuilder, TileIndirectUAV, 0u);
+	FRDGBufferUAVRef TileCountUAV = GraphBuilder.CreateUAV(Out.TileCountBuffer, PF_R32_UINT);
+	AddClearUAVPass(GraphBuilder, TileCountUAV, 0u);
 
 	TShaderMapRef<FHairStrandsTileGenerationPassCS> ComputeShader(View.ShaderMap);
 	FHairStrandsTileGenerationPassCS::FParameters* PassParameters = GraphBuilder.AllocParameters<FHairStrandsTileGenerationPassCS::FParameters>();
-	PassParameters->bRectPrimitive = Out.bRectPrimitive ? 1 : 0;
 	PassParameters->ViewResolution = View.ViewRect.Size();
 	PassParameters->InputTexture = InputTexture;
 	PassParameters->TileDataBuffer = GraphBuilder.CreateUAV(Out.TileDataBuffer, PF_R16G16_UINT);
-	PassParameters->TileIndirectBuffer = TileIndirectUAV;
+	PassParameters->TileCountBuffer = TileCountUAV;
 
-	const uint32 GroupSize = 8;
 	FComputeShaderUtils::AddPass(
 		GraphBuilder,
 		RDG_EVENT_NAME("HairStrandsTileClassification"),
 		ComputeShader,
 		PassParameters,
-		FComputeShaderUtils::GetGroupCount(InputResolution, GroupSize));
+		FComputeShaderUtils::GetGroupCount(InputResolution, FHairStrandsTiles::TileSize));
 
 	Out.TileDataSRV = GraphBuilder.CreateSRV(Out.TileDataBuffer, PF_R16G16_UINT);
 
