Fixed infinite loop on Vulkan due to integer overflow.

Vulkan actually queries the device for GRHIMaxDispatchThreadGroupsPerDimension instead of hardcoding 0xffff, so the value it gets can be as large as MAX_int32, which will overflow when multiplied.

#rnx

diff --git a/Engine/Source/Runtime/Renderer/Private/DistanceFieldStreaming.cpp b/Engine/Source/Runtime/Renderer/Private/DistanceFieldStreaming.cpp
--- a/Engine/Source/Runtime/Renderer/Private/DistanceFieldStreaming.cpp
+++ b/Engine/Source/Runtime/Renderer/Private/DistanceFieldStreaming.cpp
@@ -933,7 +933,8 @@ void FDistanceFieldSceneData::UpdateDistanceFieldAtlas(
 		{
 			FRDGTextureRef DistanceFieldBrickVolumeTextureRDG = GraphBuilder.RegisterExternalTexture(DistanceFieldBrickVolumeTexture, TEXT("DistanceFields.DistanceFieldBrickVolumeTexture"));
 
-			const int32 MaxBrickUploadsPerPass = GRHIMaxDispatchThreadGroupsPerDimension.X * FScatterUploadDistanceFieldAtlasCS::GetGroupSize() / DistanceField::BrickSize;
+			// GRHIMaxDispatchThreadGroupsPerDimension can be MAX_int32 so we need to do this math in 64-bit.
+			const int32 MaxBrickUploadsPerPass = (int32)FMath::Min<int64>((int64)GRHIMaxDispatchThreadGroupsPerDimension.X * FScatterUploadDistanceFieldAtlasCS::GetGroupSize() / DistanceField::BrickSize, MAX_int32);
 
 			for (int32 StartBrickIndex = 0; StartBrickIndex < NumBrickUploads; StartBrickIndex += MaxBrickUploadsPerPass)
 			{
