Add tile count resolution to hair strands view uniform buffer.

diff --git a/Engine/Shaders/Private/HairStrands/HairStrandsVisibilityTile.usf b/Engine/Shaders/Private/HairStrands/HairStrandsVisibilityTile.usf
--- a/Engine/Shaders/Private/HairStrands/HairStrandsVisibilityTile.usf
+++ b/Engine/Shaders/Private/HairStrands/HairStrandsVisibilityTile.usf
@@ -69,7 +69,7 @@ void TileMainCS(uint2 DispatchThreadId : SV_DispatchThreadID, uint LinearIndex :
 ////////////////////////////////////////////////////////////////////////////////////////////////////////////
 
 #if SHADER_TILE_COPY_ARGS
-uint GroupSize;
+uint2 TileCountXY;
 uint bRectPrimitive;
 Buffer<uint> TileCountBuffer;
 RWBuffer<uint> TileIndirectDrawBuffer;
@@ -88,8 +88,8 @@ void MainCS(uint2 DispatchThreadId : SV_DispatchThreadID)
 		TileIndirectDrawBuffer[2] = 1;
 
 		// Indirect dispatch
-		TileIndirectDispatchBuffer[0] = (TileCount + GroupSize-1) / GroupSize;
-		TileIndirectDispatchBuffer[1] = 1;
+		TileIndirectDispatchBuffer[0] = TileCount % TileCountXY.x;
+		TileIndirectDispatchBuffer[1] = (TileCount+TileCountXY.x-1) / TileCountXY.x;
 		TileIndirectDispatchBuffer[2] = 1;
 	}
 }
diff --git a/Engine/Source/Runtime/Renderer/Private/HairStrands/HairStrandsData.h b/Engine/Source/Runtime/Renderer/Private/HairStrands/HairStrandsData.h
--- a/Engine/Source/Runtime/Renderer/Private/HairStrands/HairStrandsData.h
+++ b/Engine/Source/Runtime/Renderer/Private/HairStrands/HairStrandsData.h
@@ -18,6 +18,7 @@ struct FMeshBatch;
 
 BEGIN_GLOBAL_SHADER_PARAMETER_STRUCT(FHairStrandsViewUniformParameters, )	
 //	SHADER_PARAMETER(uint32, SampleCount)
+	SHADER_PARAMETER(FIntPoint, HairTileCountXY)										// Tile count in X/Y
 	SHADER_PARAMETER(float, HairDualScatteringRoughnessOverride)						// Override the roughness used for dual scattering (for hack/test purpose only)
 	SHADER_PARAMETER(FIntPoint, HairSampleViewportResolution)							// Maximum viewport resolution of the sample space
 	SHADER_PARAMETER_RDG_TEXTURE(Texture2D<uint4>, HairCategorizationTexture)			// Categorization texture aggregating hair info in screen space (closest depth, coverage, ...)
@@ -27,7 +28,7 @@ BEGIN_GLOBAL_SHADER_PARAMETER_STRUCT(FHairStrandsViewUniformParameters, )
 	SHADER_PARAMETER_RDG_BUFFER_SRV(StructuredBuffer<FPackedHairSample>, HairSampleData)// Sample data (coverage, tangent, base color, ...), in sample space // HAIRSTRANDS_TODO: change this to be a uint4 so that we don't have to include the type for generated contant buffer
 	SHADER_PARAMETER_RDG_BUFFER_SRV(Buffer<uint2>, HairSampleCoords)					// Screen pixel coordinate of each sample, in sample space
 	SHADER_PARAMETER_RDG_BUFFER_SRV(Buffer<uint2>, HairTileData)						// Tile coords (RG16F)
-	SHADER_PARAMETER_RDG_BUFFER_SRV(Buffer<uint>,  HairTileCount)						// Tile total count
+	SHADER_PARAMETER_RDG_BUFFER_SRV(Buffer<uint>,  HairTileCount)						// Tile total count (actual number of tiles)
 END_GLOBAL_SHADER_PARAMETER_STRUCT()
 
 ////////////////////////////////////////////////////////////////////////////////////
@@ -39,6 +40,7 @@ struct FHairStrandsTiles
 	static const uint32 GroupSize = 64;
 	static const uint32	TileSize = 8;
 	uint32				TileCount = 0;
+	FIntPoint			TileCountXY = FIntPoint(0, 0);
 	bool				bRectPrimitive = false;
 
 	FRDGBufferSRVRef	TileDataSRV = nullptr;
diff --git a/Engine/Source/Runtime/Renderer/Private/HairStrands/HairStrandsRendering.cpp b/Engine/Source/Runtime/Renderer/Private/HairStrands/HairStrandsRendering.cpp
--- a/Engine/Source/Runtime/Renderer/Private/HairStrands/HairStrandsRendering.cpp
+++ b/Engine/Source/Runtime/Renderer/Private/HairStrands/HairStrandsRendering.cpp
@@ -20,11 +20,11 @@ static TRDGUniformBufferRef<FHairStrandsViewUniformParameters> InternalCreateHai
 		Parameters->HairSampleCoords = GraphBuilder.CreateSRV(In->NodeCoord, FHairStrandsVisibilityData::NodeCoordFormat);
 		Parameters->HairSampleCount = In->NodeCount;
 		Parameters->HairSampleViewportResolution = In->SampleLightingViewportResolution;
-
 		if (In->TileData.IsValid())
 		{
 			Parameters->HairTileData = In->TileData.TileDataSRV;
 			Parameters->HairTileCount = GraphBuilder.CreateSRV(In->TileData.TileCountBuffer, PF_R32_UINT);
+			Parameters->HairTileCountXY = In->TileData.TileCountXY;
 		}
 		else
 		{
@@ -34,6 +34,7 @@ static TRDGUniformBufferRef<FHairStrandsViewUniformParameters> InternalCreateHai
 			FRDGBufferSRVRef DummyBufferRG16SRV = GraphBuilder.CreateSRV(DummyBuffer, PF_R16G16_UINT);
 			Parameters->HairTileData = DummyBufferR32SRV;
 			Parameters->HairTileCount = DummyBufferRG16SRV;
+			Parameters->HairTileCountXY = FIntPoint(0,0);
 		}
 	}
 	else
@@ -62,6 +63,7 @@ static TRDGUniformBufferRef<FHairStrandsViewUniformParameters> InternalCreateHai
 
 		Parameters->HairTileData = DummyBufferR32SRV;
 		Parameters->HairTileCount = DummyBufferRG16SRV;
+		Parameters->HairTileCountXY = FIntPoint(0, 0);
 	}
 
 	return GraphBuilder.CreateUniformBuffer(Parameters);
diff --git a/Engine/Source/Runtime/Renderer/Private/HairStrands/HairStrandsTile.cpp b/Engine/Source/Runtime/Renderer/Private/HairStrands/HairStrandsTile.cpp
--- a/Engine/Source/Runtime/Renderer/Private/HairStrands/HairStrandsTile.cpp
+++ b/Engine/Source/Runtime/Renderer/Private/HairStrands/HairStrandsTile.cpp
@@ -38,7 +38,7 @@ class FHairStrandsTileCopyArgsPassCS : public FGlobalShader
 
 	BEGIN_SHADER_PARAMETER_STRUCT(FParameters, )
 		SHADER_PARAMETER_STRUCT_REF(FViewUniformShaderParameters, ViewUniformBuffer)
-		SHADER_PARAMETER(uint32, GroupSize)
+		SHADER_PARAMETER(FIntPoint, TileCountXY)
 		SHADER_PARAMETER(uint32, bRectPrimitive)
 		SHADER_PARAMETER_RDG_BUFFER_SRV(Buffer, TileCountBuffer)
 		SHADER_PARAMETER_RDG_BUFFER_UAV(RWBuffer, TileIndirectDrawBuffer)
@@ -66,7 +66,7 @@ void AddHairStrandsCopyArgsTilesPass(
 {
 	TShaderMapRef<FHairStrandsTileCopyArgsPassCS> ComputeShader(View.ShaderMap);
 	FHairStrandsTileCopyArgsPassCS::FParameters* PassParameters = GraphBuilder.AllocParameters<FHairStrandsTileCopyArgsPassCS::FParameters>();
-	PassParameters->GroupSize = FHairStrandsTiles::GroupSize;
+	PassParameters->TileCountXY = TileData.TileCountXY;
 	PassParameters->bRectPrimitive = TileData.bRectPrimitive ? 1 : 0;
 	PassParameters->TileCountBuffer = GraphBuilder.CreateSRV(TileData.TileCountBuffer, PF_R32_UINT);
 	PassParameters->TileIndirectDrawBuffer = GraphBuilder.CreateUAV(TileData.TileIndirectDrawBuffer, PF_R32_UINT);
@@ -120,8 +120,8 @@ FHairStrandsTiles AddHairStrandsGenerateTilesPass(
 
 	check(FHairStrandsTiles::TileSize == 8); // only size supported for now
 	const FIntPoint InputResolution = InputTexture->Desc.Extent;
-	const FIntPoint TileResolution = FIntPoint(FMath::CeilToInt(InputResolution.X / float(FHairStrandsTiles::TileSize)), FMath::CeilToInt(InputResolution.Y / float(FHairStrandsTiles::TileSize)));
-	Out.TileCount = TileResolution.X * TileResolution.Y;
+	Out.TileCountXY = FIntPoint(FMath::CeilToInt(InputResolution.X / float(FHairStrandsTiles::TileSize)), FMath::CeilToInt(InputResolution.Y / float(FHairStrandsTiles::TileSize)));
+	Out.TileCount = Out.TileCountXY.X * Out.TileCountXY.Y;
 	Out.Resolution = InputResolution;
 	Out.bRectPrimitive = GRHISupportsRectTopology;
 	Out.TileCountBuffer = GraphBuilder.CreateBuffer(FRDGBufferDesc::CreateBufferDesc(4, 1), TEXT("Hair.TileCountBuffer"));
