Fixed CIS build break due to missing file.

diff --git a/Engine/Source/Runtime/Renderer/Private/CompositionLighting/PostProcessAmbient.cpp b/Engine/Source/Runtime/Renderer/Private/CompositionLighting/PostProcessAmbient.cpp
--- a/Engine/Source/Runtime/Renderer/Private/CompositionLighting/PostProcessAmbient.cpp
+++ b/Engine/Source/Runtime/Renderer/Private/CompositionLighting/PostProcessAmbient.cpp
@@ -26,7 +26,7 @@ void SetupAmbientCubemapParameters(const FFinalPostProcessSettings::FCubemapEntr
 
 	// cubemap texture
 	{
-		FTexture* AmbientCubemapTexture = Entry.AmbientCubemap ? Entry.AmbientCubemap->Resource : GBlackTextureCube;
+		FTexture* AmbientCubemapTexture = Entry.AmbientCubemap ? Entry.AmbientCubemap->Resource : (FTexture*)GBlackTextureCube;
 		OutParameters->AmbientCubemap = AmbientCubemapTexture->TextureRHI;
 		OutParameters->AmbientCubemapSampler = AmbientCubemapTexture->SamplerStateRHI;
 	}
