Add hair tile support for hair depth/gbuffer patch passes.

diff --git a/Engine/Source/Runtime/Renderer/Private/HairStrands/HairStrandsVisibility.cpp b/Engine/Source/Runtime/Renderer/Private/HairStrands/HairStrandsVisibility.cpp
--- a/Engine/Source/Runtime/Renderer/Private/HairStrands/HairStrandsVisibility.cpp
+++ b/Engine/Source/Runtime/Renderer/Private/HairStrands/HairStrandsVisibility.cpp
@@ -2601,6 +2601,7 @@ class FHairVisibilityDepthPS : public FGlobalShader
 	using FPermutationDomain = TShaderPermutationDomain<FOutputType>;
 
 	BEGIN_SHADER_PARAMETER_STRUCT(FParameters, )
+		SHADER_PARAMETER_STRUCT_INCLUDE(FHairStrandsTilePassVS::FParameters, TileData)
 		SHADER_PARAMETER_RDG_TEXTURE(Texture2D, CategorisationTexture)
 		RENDER_TARGET_BINDING_SLOTS()
 	END_SHADER_PARAMETER_STRUCT()
@@ -2615,68 +2616,82 @@ public:
 	{
 		FGlobalShader::ModifyCompilationEnvironment(Parameters, OutEnvironment);
 		OutEnvironment.SetRenderTargetOutputFormat(0, PF_B8G8R8A8);
-		OutEnvironment.SetRenderTargetOutputFormat(1, PF_FloatRGBA);
+		OutEnvironment.SetRenderTargetOutputFormat(1, PF_B8G8R8A8);
+		OutEnvironment.SetRenderTargetOutputFormat(2, PF_FloatRGBA);
 	}
 };
 
 IMPLEMENT_GLOBAL_SHADER(FHairVisibilityDepthPS, "/Engine/Private/HairStrands/HairStrandsVisibilityDepthPS.usf", "MainPS", SF_Pixel);
 
-static void AddHairVisibilityColorAndDepthPatchPass(
+static void AddHairVisibilityCommonPatchPass(
 	FRDGBuilder& GraphBuilder,
 	const FViewInfo& View,
+	const FHairStrandsTiles& TileData,
+	const bool bDepthOnly,
 	const FRDGTextureRef& CategorisationTexture,
-	FRDGTextureRef& OutGBufferBTexture,
-	FRDGTextureRef& OutGBufferCTexture,
-	FRDGTextureRef& OutColorTexture,
-	FRDGTextureRef& OutDepthTexture)
+	FRDGTextureRef OutGBufferBTexture,
+	FRDGTextureRef OutGBufferCTexture,
+	FRDGTextureRef OutColorTexture,
+	FRDGTextureRef OutDepthTexture)
 {
-	if (!OutGBufferBTexture || !OutGBufferCTexture || !OutColorTexture || !OutDepthTexture)
-	{
-		return;
-	}
-
 	FHairVisibilityDepthPS::FParameters* Parameters = GraphBuilder.AllocParameters<FHairVisibilityDepthPS::FParameters>();
+	Parameters->TileData = GetHairStrandsTileParameters(TileData);
 	Parameters->CategorisationTexture = CategorisationTexture;
-	Parameters->RenderTargets[0] = FRenderTargetBinding(OutGBufferBTexture, ERenderTargetLoadAction::ELoad);
-	Parameters->RenderTargets[1] = FRenderTargetBinding(OutGBufferCTexture, ERenderTargetLoadAction::ELoad);
-	Parameters->RenderTargets[2] = FRenderTargetBinding(OutColorTexture, ERenderTargetLoadAction::ELoad);
 	Parameters->RenderTargets.DepthStencil = FDepthStencilBinding(
 		OutDepthTexture,
 		ERenderTargetLoadAction::ELoad,
 		ERenderTargetLoadAction::ELoad,
 		FExclusiveDepthStencil::DepthWrite_StencilNop);
+	if (!bDepthOnly)
+	{
+		check(OutGBufferBTexture && OutGBufferCTexture && OutColorTexture);
+		Parameters->RenderTargets[0] = FRenderTargetBinding(OutGBufferBTexture, ERenderTargetLoadAction::ELoad);
+		Parameters->RenderTargets[1] = FRenderTargetBinding(OutGBufferCTexture, ERenderTargetLoadAction::ELoad);
+		Parameters->RenderTargets[2] = FRenderTargetBinding(OutColorTexture, ERenderTargetLoadAction::ELoad);
+	}
+
+	TShaderMapRef<FPostProcessVS> ScreenVertexShader(View.ShaderMap);
+	TShaderMapRef<FHairStrandsTilePassVS> TileVertexShader(View.ShaderMap);
 
-	TShaderMapRef<FPostProcessVS> VertexShader(View.ShaderMap);
 	FHairVisibilityDepthPS::FPermutationDomain PermutationVector;
-	PermutationVector.Set<FHairVisibilityDepthPS::FOutputType>(0);
+	PermutationVector.Set<FHairVisibilityDepthPS::FOutputType>(bDepthOnly ? 1 : 0);
 	TShaderMapRef<FHairVisibilityDepthPS> PixelShader(View.ShaderMap, PermutationVector);
-	const FGlobalShaderMap* GlobalShaderMap = View.ShaderMap;
 	const FIntRect Viewport = View.ViewRect;
 	const FIntPoint Resolution = OutDepthTexture->Desc.Extent;
+	const bool bUseTile = TileData.IsValid();
 
+	GraphBuilder.AddPass(
+		RDG_EVENT_NAME("%s(%s)", bDepthOnly ? TEXT("HairStrandsVisibilityHairOnlyDepth") : TEXT("HairStrandsVisibilityWriteColorAndDepth"), bUseTile ? TEXT("Tile") : TEXT("Screen")),
+		Parameters,
+		ERDGPassFlags::Raster,
+		[Parameters, ScreenVertexShader, TileVertexShader, PixelShader, Viewport, Resolution, bUseTile](FRHICommandList& RHICmdList)
 	{
-		ClearUnusedGraphResources(PixelShader, Parameters);
+		FHairStrandsTilePassVS::FParameters ParametersVS = Parameters->TileData;
 
-		GraphBuilder.AddPass(
-			RDG_EVENT_NAME("HairStrandsVisibilityWriteColorAndDepth"),
-			Parameters,
-			ERDGPassFlags::Raster,
-			[Parameters, VertexShader, PixelShader, Viewport, Resolution](FRHICommandList& RHICmdList)
-		{
-			FGraphicsPipelineStateInitializer GraphicsPSOInit;
-			RHICmdList.ApplyCachedRenderTargets(GraphicsPSOInit);
-			GraphicsPSOInit.BlendState = TStaticBlendState<CW_RGBA, BO_Add, BF_One, BF_Zero, BO_Add, BF_One, BF_Zero>::GetRHI();
-			GraphicsPSOInit.RasterizerState = TStaticRasterizerState<>::GetRHI();
-			GraphicsPSOInit.DepthStencilState = TStaticDepthStencilState<true, CF_Always>::GetRHI();
+		FGraphicsPipelineStateInitializer GraphicsPSOInit;
+		RHICmdList.ApplyCachedRenderTargets(GraphicsPSOInit);
+		GraphicsPSOInit.BlendState = TStaticBlendState<CW_RGBA, BO_Add, BF_One, BF_Zero, BO_Add, BF_One, BF_Zero>::GetRHI();
+		GraphicsPSOInit.RasterizerState = TStaticRasterizerState<>::GetRHI();
+		GraphicsPSOInit.DepthStencilState = TStaticDepthStencilState<true, CF_Always>::GetRHI();
 
-			GraphicsPSOInit.BoundShaderState.VertexDeclarationRHI = GFilterVertexDeclaration.VertexDeclarationRHI;
-			GraphicsPSOInit.BoundShaderState.VertexShaderRHI = VertexShader.GetVertexShader();
-			GraphicsPSOInit.BoundShaderState.PixelShaderRHI = PixelShader.GetPixelShader();
-			GraphicsPSOInit.PrimitiveType = PT_TriangleList;
-			SetGraphicsPipelineState(RHICmdList, GraphicsPSOInit);
+		GraphicsPSOInit.BoundShaderState.VertexDeclarationRHI = GFilterVertexDeclaration.VertexDeclarationRHI;
+		GraphicsPSOInit.BoundShaderState.VertexShaderRHI = bUseTile ? TileVertexShader.GetVertexShader() : ScreenVertexShader.GetVertexShader();
+		GraphicsPSOInit.BoundShaderState.PixelShaderRHI = PixelShader.GetPixelShader();
+		GraphicsPSOInit.PrimitiveType = Parameters->TileData.bRectPrimitive > 0 ? PT_RectList : PT_TriangleList;
+		SetGraphicsPipelineState(RHICmdList, GraphicsPSOInit);
 
-			RHICmdList.SetViewport(Viewport.Min.X, Viewport.Min.Y, 0.0f, Viewport.Max.X, Viewport.Max.Y, 1.0f);
-			SetShaderParameters(RHICmdList, PixelShader, PixelShader.GetPixelShader(), *Parameters);
+		RHICmdList.SetViewport(Viewport.Min.X, Viewport.Min.Y, 0.0f, Viewport.Max.X, Viewport.Max.Y, 1.0f);
+		SetShaderParameters(RHICmdList, PixelShader, PixelShader.GetPixelShader(), *Parameters);
+
+		if (bUseTile)
+		{
+			SetShaderParameters(RHICmdList, TileVertexShader, TileVertexShader.GetVertexShader(), ParametersVS);
+//			RHICmdList.SetViewport(0, 0, 0.0f, Resolution.X, Resolution.Y, 1.0f);
+			RHICmdList.SetStreamSource(0, nullptr, 0);
+			RHICmdList.DrawPrimitiveIndirect(Parameters->TileData.TileIndirectBuffer->GetRHI(), 0);
+		}
+		else
+		{
 			DrawRectangle(
 				RHICmdList,
 				0, 0,
@@ -2685,15 +2700,44 @@ static void AddHairVisibilityColorAndDepthPatchPass(
 				Viewport.Width(), Viewport.Height(),
 				Viewport.Size(),
 				Resolution,
-				VertexShader,
+				ScreenVertexShader,
 				EDRF_UseTriangleOptimization);
-		});
+		}
+	});
+}
+
+
+static void AddHairVisibilityColorAndDepthPatchPass(
+	FRDGBuilder& GraphBuilder,
+	const FViewInfo& View,
+	const FHairStrandsTiles& TileData,
+	const FRDGTextureRef& CategorisationTexture,
+	FRDGTextureRef& OutGBufferBTexture,
+	FRDGTextureRef& OutGBufferCTexture,
+	FRDGTextureRef& OutColorTexture,
+	FRDGTextureRef& OutDepthTexture)
+{
+	if (!OutGBufferBTexture || !OutGBufferCTexture || !OutColorTexture || !OutDepthTexture)
+	{
+		return;
 	}
+
+	AddHairVisibilityCommonPatchPass(
+		GraphBuilder,
+		View,
+		TileData,
+		false,
+		CategorisationTexture,
+		OutGBufferBTexture,
+		OutGBufferCTexture,
+		OutColorTexture,
+		OutDepthTexture);
 }
 
 static void AddHairOnlyDepthPass(
 	FRDGBuilder& GraphBuilder,
 	const FViewInfo& View,
+	const FHairStrandsTiles& TileData,
 	const FRDGTextureRef& CategorisationTexture,
 	FRDGTextureRef& OutDepthTexture)
 {
@@ -2701,57 +2745,16 @@ static void AddHairOnlyDepthPass(
 	{
 		return;
 	}
-
-	FHairVisibilityDepthPS::FParameters* Parameters = GraphBuilder.AllocParameters<FHairVisibilityDepthPS::FParameters>();
-	Parameters->CategorisationTexture = CategorisationTexture;
-	Parameters->RenderTargets.DepthStencil = FDepthStencilBinding(
-		OutDepthTexture,
-		ERenderTargetLoadAction::ELoad,
-		ERenderTargetLoadAction::ELoad,
-		FExclusiveDepthStencil::DepthWrite_StencilNop);
-
-	TShaderMapRef<FPostProcessVS> VertexShader(View.ShaderMap);
-	FHairVisibilityDepthPS::FPermutationDomain PermutationVector;
-	PermutationVector.Set<FHairVisibilityDepthPS::FOutputType>(1);
-	TShaderMapRef<FHairVisibilityDepthPS> PixelShader(View.ShaderMap, PermutationVector);
-	const FGlobalShaderMap* GlobalShaderMap = View.ShaderMap;
-	const FIntRect Viewport = View.ViewRect;
-	const FIntPoint Resolution = OutDepthTexture->Desc.Extent;
-	const FViewInfo* CapturedView = &View;
-
-	ClearUnusedGraphResources(PixelShader, Parameters);
-
-	GraphBuilder.AddPass(
-		RDG_EVENT_NAME("HairStrandsVisibilityHairOnlyDepth"),
-		Parameters,
-		ERDGPassFlags::Raster,
-		[Parameters, VertexShader, PixelShader, Viewport, Resolution, CapturedView](FRHICommandList& RHICmdList)
-		{
-			FGraphicsPipelineStateInitializer GraphicsPSOInit;
-			RHICmdList.ApplyCachedRenderTargets(GraphicsPSOInit);
-			GraphicsPSOInit.BlendState = TStaticBlendState<CW_RGBA, BO_Add, BF_One, BF_Zero, BO_Add, BF_One, BF_Zero>::GetRHI();
-			GraphicsPSOInit.RasterizerState = TStaticRasterizerState<>::GetRHI();
-			GraphicsPSOInit.DepthStencilState = TStaticDepthStencilState<true, CF_Always>::GetRHI();
-
-			GraphicsPSOInit.BoundShaderState.VertexDeclarationRHI = GFilterVertexDeclaration.VertexDeclarationRHI;
-			GraphicsPSOInit.BoundShaderState.VertexShaderRHI = VertexShader.GetVertexShader();
-			GraphicsPSOInit.BoundShaderState.PixelShaderRHI = PixelShader.GetPixelShader();
-			GraphicsPSOInit.PrimitiveType = PT_TriangleList;
-			SetGraphicsPipelineState(RHICmdList, GraphicsPSOInit);
-
-			RHICmdList.SetViewport(Viewport.Min.X, Viewport.Min.Y, 0.0f, Viewport.Max.X, Viewport.Max.Y, 1.0f);
-			SetShaderParameters(RHICmdList, PixelShader, PixelShader.GetPixelShader(), *Parameters);
-			DrawRectangle(
-				RHICmdList,
-				0, 0,
-				Viewport.Width(), Viewport.Height(),
-				Viewport.Min.X, Viewport.Min.Y,
-				Viewport.Width(), Viewport.Height(),
-				Viewport.Size(),
-				Resolution,
-				VertexShader,
-				EDRF_UseTriangleOptimization);
-		});
+	AddHairVisibilityCommonPatchPass(
+		GraphBuilder,
+		View,
+		TileData,
+		true,
+		CategorisationTexture,
+		nullptr,
+		nullptr,
+		nullptr,
+		OutDepthTexture);
 }
 
 ///////////////////////////////////////////////////////////////////////////////////////////////////
@@ -3121,6 +3124,8 @@ void RenderHairStrandsVisibilityBuffer(
 	check(View.Family);
 	check(MacroGroupDatas.Num() > 0);
 
+	const bool bGenerateTile = GHairStrandsTile > 0;
+
 	{
 		
 		{
@@ -3194,6 +3199,12 @@ void RenderHairStrandsVisibilityBuffer(
 							IndirectArgsBuffer,
 							VisibilityData.MaxNodeCount);
 
+						// Generate Tile data
+						if (CategorizationTexture && bGenerateTile)
+						{
+							VisibilityData.TileData = AddHairStrandsGenerateTilesPass(GraphBuilder, View, CategorizationTexture);
+						}
+
 						// Evaluate material based on the visiblity pass result
 						// Output both complete sample data + per-sample velocity
 						FMaterialPassOutput PassOutput = AddHairMaterialPass(
@@ -3248,6 +3259,7 @@ void RenderHairStrandsVisibilityBuffer(
 						AddHairVisibilityColorAndDepthPatchPass(
 							GraphBuilder,
 							View,
+							VisibilityData.TileData,
 							CategorizationTexture,
 							SceneGBufferBTexture,
 							SceneGBufferCTexture,
@@ -3258,6 +3270,7 @@ void RenderHairStrandsVisibilityBuffer(
 					AddHairOnlyDepthPass(
 						GraphBuilder,
 						View,
+						VisibilityData.TileData,
 						CategorizationTexture,
 						HairOnlyDepthTexture);
 				}
@@ -3373,6 +3386,12 @@ void RenderHairStrandsVisibilityBuffer(
 						IndirectArgsBuffer,
 						VisibilityData.MaxNodeCount);
 
+					// Generate Tile data
+					if (CategorizationTexture && bGenerateTile)
+					{
+						VisibilityData.TileData = AddHairStrandsGenerateTilesPass(GraphBuilder, View, CategorizationTexture);
+					}
+
 					{
 						const bool bUpdateSampleCoverage = GHairStrandsSortHairSampleByDepth > 0;
 
@@ -3452,6 +3471,7 @@ void RenderHairStrandsVisibilityBuffer(
 						AddHairVisibilityColorAndDepthPatchPass(
 							GraphBuilder,
 							View,
+							VisibilityData.TileData,
 							CategorizationTexture,
 							SceneGBufferBTexture,
 							SceneGBufferCTexture,
@@ -3462,6 +3482,7 @@ void RenderHairStrandsVisibilityBuffer(
 					AddHairOnlyDepthPass(
 						GraphBuilder,
 						View,
+						VisibilityData.TileData,
 						CategorizationTexture,
 						HairOnlyDepthTexture);
 				}
@@ -3543,11 +3564,18 @@ void RenderHairStrandsVisibilityBuffer(
 					VisibilityData.NodeCount = NodeCounter;
 				}
 
+				// Generate Tile data
+				if (CategorizationTexture && bGenerateTile)
+				{
+					VisibilityData.TileData = AddHairStrandsGenerateTilesPass(GraphBuilder, View, CategorizationTexture);
+				}
+
 				if (bRunColorAndDepthPatching)
 				{
 					AddHairVisibilityColorAndDepthPatchPass(
 						GraphBuilder,
 						View,
+						VisibilityData.TileData,
 						CategorizationTexture,
 						SceneGBufferBTexture,
 						SceneGBufferCTexture,
@@ -3558,6 +3586,7 @@ void RenderHairStrandsVisibilityBuffer(
 				AddHairOnlyDepthPass(
 					GraphBuilder,
 					View,
+					VisibilityData.TileData,
 					CategorizationTexture,
 					HairOnlyDepthTexture);
 
@@ -3589,13 +3618,6 @@ void RenderHairStrandsVisibilityBuffer(
 				VisibilityData.LightChannelMaskTexture = LightingChannelMaskTexture;
 			}
 		#endif
-
-			// Generate Tile data
-			const bool bGenerateTile = GHairStrandsTile > 0;
-			if (CategorizationTexture && bGenerateTile)
-			{
-				VisibilityData.TileData = AddHairStrandsGenerateTilesPass(GraphBuilder, View, CategorizationTexture);
-			}
 		}
 	}
 }
