Added RHI globals to query stride and alignment requirements for various ray tracing buffers.

diff --git a/Engine/Source/Runtime/D3D12RHI/Private/D3D12Adapter.cpp b/Engine/Source/Runtime/D3D12RHI/Private/D3D12Adapter.cpp
--- a/Engine/Source/Runtime/D3D12RHI/Private/D3D12Adapter.cpp
+++ b/Engine/Source/Runtime/D3D12RHI/Private/D3D12Adapter.cpp
@@ -815,6 +815,11 @@ void FD3D12Adapter::InitializeDevices()
 					}
 				}
 			}
+
+			GRHIRayTracingAccelerationStructureAlignment = uint32(D3D12_RAYTRACING_ACCELERATION_STRUCTURE_BYTE_ALIGNMENT);
+			GRHIRayTracingShaderTableAlignment = uint32(D3D12_RAYTRACING_SHADER_TABLE_BYTE_ALIGNMENT);
+			GRHIRayTracingInstanceDescriptorSize = uint32(sizeof(D3D12_RAYTRACING_INSTANCE_DESC));
+
 #endif // D3D12_RHI_RAYTRACING
 
 #if PLATFORM_WINDOWS && D3D12_CORE_ENABLED
diff --git a/Engine/Source/Runtime/RHI/Private/RHI.cpp b/Engine/Source/Runtime/RHI/Private/RHI.cpp
--- a/Engine/Source/Runtime/RHI/Private/RHI.cpp
+++ b/Engine/Source/Runtime/RHI/Private/RHI.cpp
@@ -928,6 +928,9 @@ bool GRHISupportsRayTracingPSOAdditions = false;
 bool GRHISupportsRayTracingDispatchIndirect = false;
 bool GRHISupportsRayTracingAsyncBuildAccelerationStructure = false;
 bool GRHISupportsRayTracingAMDHitToken = false;
+uint32 GRHIRayTracingAccelerationStructureAlignment = 0;
+uint32 GRHIRayTracingShaderTableAlignment = 0;
+uint32 GRHIRayTracingInstanceDescriptorSize = 0;
 bool GRHISupportsWaveOperations = false;
 int32 GRHIMinimumWaveSize = 4; // Minimum supported value in SM 6.0
 int32 GRHIMaximumWaveSize = 128; // Maximum supported value in SM 6.0
diff --git a/Engine/Source/Runtime/RHI/Public/RHI.h b/Engine/Source/Runtime/RHI/Public/RHI.h
--- a/Engine/Source/Runtime/RHI/Public/RHI.h
+++ b/Engine/Source/Runtime/RHI/Public/RHI.h
@@ -578,6 +578,15 @@ extern RHI_API bool GRHISupportsRayTracingAsyncBuildAccelerationStructure;
 /** Whether or not the RHI supports the AMD Hit Token extension. */
 extern RHI_API bool GRHISupportsRayTracingAMDHitToken;
 
+/** Required alignment for ray tracing acceleration structures and scratch buffers. */
+extern RHI_API uint32 GRHIRayTracingAccelerationStructureAlignment;
+
+/** Required alignment for ray tracing shader binding table buffer. */
+extern RHI_API uint32 GRHIRayTracingShaderTableAlignment;
+
+/** Size of an individual element in the ray tracing instance buffer. This defines the required stride and alignment of structured buffers of instances. */
+extern RHI_API uint32 GRHIRayTracingInstanceDescriptorSize;
+
 /** Whether or not the RHI supports shader wave operations (shader model 6.0). */
 extern RHI_API bool GRHISupportsWaveOperations;
 
diff --git a/Engine/Source/Runtime/Renderer/Private/RayTracing/RayTracingScene.cpp b/Engine/Source/Runtime/Renderer/Private/RayTracing/RayTracingScene.cpp
--- a/Engine/Source/Runtime/Renderer/Private/RayTracing/RayTracingScene.cpp
+++ b/Engine/Source/Runtime/Renderer/Private/RayTracing/RayTracingScene.cpp
@@ -66,11 +66,12 @@ FGraphEventRef FRayTracingScene::BeginCreate(FRDGBuilder& GraphBuilder)
 
 		}, TStatId(), nullptr, ENamedThreads::AnyThread);
 
+	const uint64 ScratchAlignment = GRHIRayTracingAccelerationStructureAlignment;
 	FRDGBufferDesc ScratchBufferDesc;
 	ScratchBufferDesc.UnderlyingType = FRDGBufferDesc::EUnderlyingType::StructuredBuffer;
 	ScratchBufferDesc.Usage = BUF_UnorderedAccess;
-	ScratchBufferDesc.BytesPerElement = 4;
-	ScratchBufferDesc.NumElements = SizeInfo.BuildScratchSize / 4;
+	ScratchBufferDesc.BytesPerElement = uint32(ScratchAlignment);
+	ScratchBufferDesc.NumElements = uint32(FMath::DivideAndRoundUp(SizeInfo.BuildScratchSize, ScratchAlignment));
 
 	BuildScratchBuffer = GraphBuilder.CreateBuffer(ScratchBufferDesc, TEXT("FRayTracingScene::BuildScratchBuffer"));
 
