Added r.RayTracing.Culling mode 3, which allows culling be either distance or angle. Other modes require both angle and distance conditions.

diff --git a/Engine/Source/Runtime/Renderer/Private/DeferredShadingRenderer.cpp b/Engine/Source/Runtime/Renderer/Private/DeferredShadingRenderer.cpp
--- a/Engine/Source/Runtime/Renderer/Private/DeferredShadingRenderer.cpp
+++ b/Engine/Source/Runtime/Renderer/Private/DeferredShadingRenderer.cpp
@@ -178,12 +178,13 @@ static TAutoConsoleVariable<int32> CVarRayTracingCulling(
 	TEXT("Enable culling in ray tracing for objects that are behind the camera\n")
 	TEXT(" 0: Culling disabled (default)\n")
 	TEXT(" 1: Culling by distance and solid angle enabled. Only cull objects behind camera.\n")
-	TEXT(" 2: Culling by distance and solid angle enabled. Cull objects in front and behind camera."),
+	TEXT(" 2: Culling by distance and solid angle enabled. Cull objects in front and behind camera.\n")
+	TEXT(" 3: Culling by distance OR solid angle enabled. Cull objects in front and behind camera."),
 	ECVF_RenderThreadSafe);
 
 static TAutoConsoleVariable<float> CVarRayTracingCullingRadius(
 	TEXT("r.RayTracing.Culling.Radius"),
-	10000.0f, 
+	10000.0f,
 	TEXT("Do camera culling for objects behind the camera outside of this radius in ray tracing effects (default = 10000 (100m))"),
 	ECVF_RenderThreadSafe);
 
@@ -565,10 +566,11 @@ bool FDeferredShadingSceneRenderer::GatherRayTracingWorldInstancesForView(FRHICo
 		const int32 CullInRayTracing = CVarRayTracingCulling.GetValueOnRenderThread();
 		const float CullingRadius = CVarRayTracingCullingRadius.GetValueOnRenderThread();
 		const float CullAngleThreshold = CVarRayTracingCullingAngle.GetValueOnRenderThread();
-		const float AngleThresholdRatio = FMath::Tan(CullAngleThreshold * PI / 180.0f);
+		const float AngleThresholdRatio = FMath::Tan(FMath::Min(89.99f, CullAngleThreshold) * PI / 180.0f);
 		const FVector ViewOrigin = View.ViewMatrices.GetViewOrigin();
 		const FVector ViewDirection = View.GetViewDirection();
-		const bool bCullAllObjects = CullInRayTracing == 2;
+		const bool bCullAllObjects = CullInRayTracing == 2 || CullInRayTracing == 3;
+		const bool bCullByRadiusOrDistance = CullInRayTracing == 3;
 
 		for (int PrimitiveIndex = 0; PrimitiveIndex < Scene->PrimitiveSceneProxies.Num(); PrimitiveIndex++)
 		{
@@ -612,15 +614,16 @@ bool FDeferredShadingSceneRenderer::GatherRayTracingWorldInstancesForView(FRHICo
 					const float CameraToObjectCenterLength = CameraToObjectCenter.Size();
 					const bool bIsFarEnoughToCull = CameraToObjectCenterLength > (CullingRadius + ObjectRadius);
 
-					if (bIsFarEnoughToCull) 
-					{
-						// Cull by solid angle: check the radius of bounding sphere against angle threshold
-						const bool bAngleIsSmallEnoughToCull = ObjectRadius / CameraToObjectCenterLength < AngleThresholdRatio;
+					// Cull by solid angle: check the radius of bounding sphere against angle threshold
+					const bool bAngleIsSmallEnoughToCull = ObjectRadius / CameraToObjectCenterLength < AngleThresholdRatio;
 
-						if (bAngleIsSmallEnoughToCull)
-						{
-							continue;
-						}
+					if (bCullByRadiusOrDistance && (bIsFarEnoughToCull || bAngleIsSmallEnoughToCull))
+					{
+						continue;
+					}
+					else if (bIsFarEnoughToCull && bAngleIsSmallEnoughToCull)
+					{
+						continue;
 					}
 				}
 			}
