Properly setup SceneTextures for mobile base pass

diff --git a/Engine/Source/Runtime/Renderer/Private/MobileBasePassRendering.cpp b/Engine/Source/Runtime/Renderer/Private/MobileBasePassRendering.cpp
--- a/Engine/Source/Runtime/Renderer/Private/MobileBasePassRendering.cpp
+++ b/Engine/Source/Runtime/Renderer/Private/MobileBasePassRendering.cpp
@@ -148,8 +148,7 @@ void SetupMobileBasePassUniformParameters(
 	FRDGBuilder& GraphBuilder,
 	const FViewInfo& View, 
 	EMobileBasePass BasePass,
-	FRDGTextureRef ScreenSpaceAOTexture,
-	FTextureRHIRef PixelProjectedReflectionTexture,
+	EMobileSceneTextureSetupMode SetupMode,
 	FMobileBasePassUniformParameters& BasePassParameters)
 {
 	SetupFogUniformParameters(GraphBuilder, View, BasePassParameters.Fog);
@@ -157,15 +156,10 @@ void SetupMobileBasePassUniformParameters(
 	const FScene* Scene = View.Family->Scene ? View.Family->Scene->GetRenderScene() : nullptr;
 	const FPlanarReflectionSceneProxy* ReflectionSceneProxy = Scene ? Scene->GetForwardPassGlobalPlanarReflection() : nullptr;
 	SetupPlanarReflectionUniformParameters(View, ReflectionSceneProxy, BasePassParameters.PlanarReflection);
-
-	EMobileSceneTextureSetupMode SetupMode = EMobileSceneTextureSetupMode::None;
-	if (BasePass == EMobileBasePass::Translucent)
-	{
-		SetupMode |= EMobileSceneTextureSetupMode::SceneColor;
-	}
-	if (View.bCustomDepthStencilValid)
+	if (View.PrevViewInfo.MobilePixelProjectedReflection.IsValid())
 	{
-		SetupMode |= EMobileSceneTextureSetupMode::CustomDepth;
+		BasePassParameters.PlanarReflection.PlanarReflectionTexture = View.PrevViewInfo.MobilePixelProjectedReflection->GetRenderTargetItem().ShaderResourceTexture;
+		BasePassParameters.PlanarReflection.PlanarReflectionSampler = TStaticSamplerState<SF_Bilinear, AM_Clamp, AM_Clamp, AM_Clamp>::GetRHI();
 	}
 
 	const FRDGSystemTextures& SystemTextures = FRDGSystemTextures::Get(GraphBuilder);
@@ -174,23 +168,19 @@ void SetupMobileBasePassUniformParameters(
 
 	BasePassParameters.PreIntegratedGFTexture = GSystemTextures.PreintegratedGF->GetRenderTargetItem().ShaderResourceTexture;
 	BasePassParameters.PreIntegratedGFSampler = TStaticSamplerState<SF_Bilinear, AM_Clamp, AM_Clamp, AM_Clamp>::GetRHI();
-
-	if (PixelProjectedReflectionTexture.IsValid())
-			{
-		BasePassParameters.PlanarReflection.PlanarReflectionTexture = PixelProjectedReflectionTexture;
-				BasePassParameters.PlanarReflection.PlanarReflectionSampler = TStaticSamplerState<SF_Bilinear, AM_Clamp, AM_Clamp, AM_Clamp>::GetRHI();
-			}
-
 	BasePassParameters.EyeAdaptationBuffer = GraphBuilder.CreateSRV(GetEyeAdaptationBuffer(GraphBuilder, View), PF_A32B32G32R32F);
 
-	if (BasePass == EMobileBasePass::Opaque && ScreenSpaceAOTexture != nullptr)
+	FRDGTextureRef AmbientOcclusionTexture = SystemTextures.White;
+	if (BasePass == EMobileBasePass::Opaque)
 	{
-		BasePassParameters.AmbientOcclusionTexture = ScreenSpaceAOTexture;
-	}
-	else
-	{
-		BasePassParameters.AmbientOcclusionTexture = SystemTextures.White;
+		FRDGTextureRef LastFrameScreenSpaceAO = TryRegisterExternalTexture(GraphBuilder, View.PrevViewInfo.MobileAmbientOcclusion); 
+		if (LastFrameScreenSpaceAO)
+		{
+			AmbientOcclusionTexture = LastFrameScreenSpaceAO;
+		}
 	}
+
+	BasePassParameters.AmbientOcclusionTexture = AmbientOcclusionTexture;
 	BasePassParameters.AmbientOcclusionSampler = TStaticSamplerState<SF_Bilinear, AM_Clamp, AM_Clamp, AM_Clamp>::GetRHI();
 	BasePassParameters.AmbientOcclusionStaticFraction = FMath::Clamp(View.FinalPostProcessSettings.AmbientOcclusionStaticFraction, 0.0f, 1.0f);
 
@@ -213,11 +203,10 @@ TRDGUniformBufferRef<FMobileBasePassUniformParameters> CreateMobileBasePassUnifo
 	FRDGBuilder& GraphBuilder,
 	const FViewInfo& View,
 	EMobileBasePass BasePass,
-	FRDGTextureRef ScreenSpaceAO,
-	FTextureRHIRef PixelProjectedReflectionTexture)
+	EMobileSceneTextureSetupMode SetupMode)
 {
 	FMobileBasePassUniformParameters* BasePassParameters = GraphBuilder.AllocParameters<FMobileBasePassUniformParameters>();
-	SetupMobileBasePassUniformParameters(GraphBuilder, View, BasePass, ScreenSpaceAO, PixelProjectedReflectionTexture, *BasePassParameters);
+	SetupMobileBasePassUniformParameters(GraphBuilder, View, BasePass, SetupMode, *BasePassParameters);
 	return GraphBuilder.CreateUniformBuffer(BasePassParameters);
 }
 
diff --git a/Engine/Source/Runtime/Renderer/Private/MobileBasePassRendering.h b/Engine/Source/Runtime/Renderer/Private/MobileBasePassRendering.h
--- a/Engine/Source/Runtime/Renderer/Private/MobileBasePassRendering.h
+++ b/Engine/Source/Runtime/Renderer/Private/MobileBasePassRendering.h
@@ -39,12 +39,6 @@ BEGIN_GLOBAL_SHADER_PARAMETER_STRUCT(FMobileBasePassUniformParameters, )
 	SHADER_PARAMETER_SAMPLER(SamplerState, ScreenSpaceShadowMaskSampler)
 END_GLOBAL_SHADER_PARAMETER_STRUCT()
 
-BEGIN_SHADER_PARAMETER_STRUCT(FMobileBasePassParameters, )
-	SHADER_PARAMETER_STRUCT_INCLUDE(FViewShaderParameters, View)
-	SHADER_PARAMETER_RDG_UNIFORM_BUFFER(FMobileBasePassUniformParameters, MobileBasePass)
-	RENDER_TARGET_BINDING_SLOTS()
-END_SHADER_PARAMETER_STRUCT()
-
 enum class EMobileBasePass
 {
 	Opaque,
@@ -61,8 +55,7 @@ extern TRDGUniformBufferRef<FMobileBasePassUniformParameters> CreateMobileBasePa
 	FRDGBuilder& GraphBuilder,
 	const FViewInfo& View,
 	EMobileBasePass BasePass,
-	FRDGTextureRef ScreenSpaceAO = nullptr,
-	FTextureRHIRef PixelProjectedReflectionTexture = nullptr);
+	EMobileSceneTextureSetupMode SetupMode);
 
 extern void SetupMobileDirectionalLightUniformParameters(
 	const FScene& Scene,
diff --git a/Engine/Source/Runtime/Renderer/Private/MobileDeferredShadingPass.cpp b/Engine/Source/Runtime/Renderer/Private/MobileDeferredShadingPass.cpp
--- a/Engine/Source/Runtime/Renderer/Private/MobileDeferredShadingPass.cpp
+++ b/Engine/Source/Runtime/Renderer/Private/MobileDeferredShadingPass.cpp
@@ -202,7 +202,7 @@ static void GetLightMaterial(const FCachedLightMaterial& DefaultLightMaterial, c
 	OutShader = MaterialShaderMap->GetShader<ShaderType>(PermutationId);
 }
 
-static void RenderDirectLight(FRHICommandListImmediate& RHICmdList, const FScene& Scene, const FViewInfo& View, const FCachedLightMaterial& DefaultLightMaterial, const FTextureRHIRef MobilePixelProjectedReflection)
+static void RenderDirectLight(FRHICommandListImmediate& RHICmdList, const FScene& Scene, const FViewInfo& View, const FCachedLightMaterial& DefaultLightMaterial)
 {
 	FLightSceneInfo* DirectionalLight = nullptr;
 	for (int32 ChannelIdx = 0; ChannelIdx < UE_ARRAY_COUNT(Scene.MobileDirectionalLights) && !DirectionalLight; ChannelIdx++)
@@ -258,10 +258,9 @@ static void RenderDirectLight(FRHICommandListImmediate& RHICmdList, const FScene
 	const FPlanarReflectionSceneProxy* ReflectionSceneProxy = Scene.GetForwardPassGlobalPlanarReflection();
 	FPlanarReflectionUniformParameters PlanarReflectionUniformParameters;
 	SetupPlanarReflectionUniformParameters(View, ReflectionSceneProxy, PlanarReflectionUniformParameters);
-
-	if (MobilePixelProjectedReflection.IsValid())
+	if (View.PrevViewInfo.MobilePixelProjectedReflection.IsValid())
 	{
-		PlanarReflectionUniformParameters.PlanarReflectionTexture = MobilePixelProjectedReflection;
+		PlanarReflectionUniformParameters.PlanarReflectionTexture = View.PrevViewInfo.MobilePixelProjectedReflection->GetRenderTargetItem().ShaderResourceTexture;
 	}
 	PassParameters.PlanarReflection = TUniformBufferRef<FPlanarReflectionUniformParameters>::CreateUniformBufferImmediate(PlanarReflectionUniformParameters, UniformBuffer_SingleDraw);
 
@@ -580,8 +579,7 @@ void MobileDeferredShadingPass(
 	int32 NumViews,
 	const FViewInfo& View,
 	const FScene& Scene, 
-	const FSortedLightSetSceneInfo &SortedLightSet,
-	const FTextureRHIRef MobilePixelProjectedReflection)
+	const FSortedLightSetSceneInfo &SortedLightSet)
 {
 	RHICmdList.SetViewport(View.ViewRect.Min.X, View.ViewRect.Min.Y, 0.0f, View.ViewRect.Max.X, View.ViewRect.Max.Y, 1.0f);
 
@@ -591,7 +589,7 @@ void MobileDeferredShadingPass(
 	DefaultMaterial.Material = DefaultMaterial.MaterialProxy->GetMaterialNoFallback(ERHIFeatureLevel::ES3_1);
 	check(DefaultMaterial.Material);
 
-	RenderDirectLight(RHICmdList, Scene, View, DefaultMaterial, MobilePixelProjectedReflection);
+	RenderDirectLight(RHICmdList, Scene, View, DefaultMaterial);
 
 	if (GMobileUseClusteredDeferredShading == 0)
 	{
diff --git a/Engine/Source/Runtime/Renderer/Private/MobileDeferredShadingPass.h b/Engine/Source/Runtime/Renderer/Private/MobileDeferredShadingPass.h
--- a/Engine/Source/Runtime/Renderer/Private/MobileDeferredShadingPass.h
+++ b/Engine/Source/Runtime/Renderer/Private/MobileDeferredShadingPass.h
@@ -21,5 +21,4 @@ void MobileDeferredShadingPass(
 	int32 NumViews,
 	const FViewInfo& View,
 	const FScene& Scene, 
-	const FSortedLightSetSceneInfo &SortedLightSet,
-	const FTextureRHIRef MobilePixelProjectedReflection);
+	const FSortedLightSetSceneInfo &SortedLightSet);
diff --git a/Engine/Source/Runtime/Renderer/Private/MobileShadingRenderer.cpp b/Engine/Source/Runtime/Renderer/Private/MobileShadingRenderer.cpp
--- a/Engine/Source/Runtime/Renderer/Private/MobileShadingRenderer.cpp
+++ b/Engine/Source/Runtime/Renderer/Private/MobileShadingRenderer.cpp
@@ -155,7 +155,6 @@ END_SHADER_PARAMETER_STRUCT()
 
 BEGIN_SHADER_PARAMETER_STRUCT(FMobileRenderPassParameters, )
 	SHADER_PARAMETER_STRUCT_INCLUDE(FViewShaderParameters, View)
-	SHADER_PARAMETER_RDG_UNIFORM_BUFFER(FMobileSceneTextureUniformParameters, SceneTextures)
 	SHADER_PARAMETER_RDG_UNIFORM_BUFFER(FMobileBasePassUniformParameters, MobileBasePass)
 	SHADER_PARAMETER_RDG_UNIFORM_BUFFER(FDebugViewModePassUniformParameters, DebugViewMode)
 	SHADER_PARAMETER_STRUCT_INCLUDE(FInstanceCullingDrawParams, InstanceCullingDrawParamsDepth)
@@ -759,9 +758,6 @@ void FMobileSceneRenderer::Render(FRDGBuilder& GraphBuilder)
 		RenderCustomDepthPass(GraphBuilder, SceneTextures.CustomDepth, SceneTextures.GetSceneTextureShaderParameters(FeatureLevel));
 	}
 
-	SceneTextures.MobileSetupMode = EMobileSceneTextureSetupMode::CustomDepth;
-	SceneTextures.MobileUniformBuffer = CreateMobileSceneTextureUniformBuffer(GraphBuilder, SceneTextures.MobileSetupMode);
-
 	FRDGTextureRef ViewFamilyTexture = TryCreateViewFamilyTexture(GraphBuilder, ViewFamily);
 	
 	if (bIsFullDepthPrepassEnabled)
@@ -850,11 +846,11 @@ void FMobileSceneRenderer::Render(FRDGBuilder& GraphBuilder)
 
 				FMobilePostProcessingInputs PostProcessingInputs;
 				PostProcessingInputs.ViewFamilyTexture = ViewFamilyTexture;
+				PostProcessingInputs.SceneTextures = CreateMobileSceneTextureUniformBuffer(GraphBuilder, EMobileSceneTextureSetupMode::All);
 
 				for (int32 ViewIndex = 0; ViewIndex < Views.Num(); ViewIndex++)
 				{
 					RDG_EVENT_SCOPE_CONDITIONAL(GraphBuilder, Views.Num() > 1, "View%d", ViewIndex);
-					PostProcessingInputs.SceneTextures = SceneTextures.MobileUniformBuffer;
 					AddMobilePostProcessingPasses(GraphBuilder, Scene, Views[ViewIndex], PostProcessingInputs, InstanceCullingManager);
 				}
 			}
@@ -949,19 +945,10 @@ void FMobileSceneRenderer::RenderForward(FRDGBuilder& GraphBuilder, FRDGTextureR
 
 		UpdateDirectionalLightUniformBuffers(GraphBuilder, View);
 
-		FRDGTextureRef LastFrameScreenSpaceAO = TryRegisterExternalTexture(GraphBuilder, View.PrevViewInfo.MobileAmbientOcclusion); 
-
-		if (LastFrameScreenSpaceAO == nullptr)
-		{
-			LastFrameScreenSpaceAO = SystemTextures.White;
-		}
-
-		FTextureRHIRef MobilePixelProjectedReflection = View.PrevViewInfo.MobilePixelProjectedReflection.IsValid() ? View.PrevViewInfo.MobilePixelProjectedReflection->GetRenderTargetItem().ShaderResourceTexture : nullptr;
-
-		auto* PassParameters = GraphBuilder.AllocParameters<FMobileRenderPassParameters>();
+		EMobileSceneTextureSetupMode SetupMode = EMobileSceneTextureSetupMode::CustomDepth;
+		FMobileRenderPassParameters* PassParameters = GraphBuilder.AllocParameters<FMobileRenderPassParameters>();
 		PassParameters->View = View.GetShaderParameters();
-		PassParameters->MobileBasePass = CreateMobileBasePassUniformBuffer(GraphBuilder, View, EMobileBasePass::Opaque, LastFrameScreenSpaceAO, MobilePixelProjectedReflection);
-		PassParameters->SceneTextures = SceneTextures.MobileUniformBuffer;
+		PassParameters->MobileBasePass = CreateMobileBasePassUniformBuffer(GraphBuilder, View, EMobileBasePass::Opaque, SetupMode);
 		PassParameters->RenderTargets = BasePassRenderTargets;
 		#if WITH_DEBUG_VIEW_MODES
 		if (ViewFamily.UseDebugViewPS())
@@ -1061,8 +1048,8 @@ void FMobileSceneRenderer::RenderForwardMultiPass(FRDGBuilder& GraphBuilder, FMo
 		PostRenderBasePass(RHICmdList, View);
 	});
 
-	SceneTextures.MobileSetupMode = EMobileSceneTextureSetupMode::SceneDepth | EMobileSceneTextureSetupMode::SceneDepthAux | EMobileSceneTextureSetupMode::CustomDepth;
-	SceneTextures.MobileUniformBuffer = CreateMobileSceneTextureUniformBuffer(GraphBuilder, SceneTextures.MobileSetupMode);
+	// resolve MSAA depth for translucency
+	AddResolveSceneDepthPass(GraphBuilder, View, SceneTextures.Depth);
 
 	FExclusiveDepthStencil::Type ExclusiveDepthStencil = FExclusiveDepthStencil::DepthRead_StencilRead;
 	if (bModulatedShadowsInUse)
@@ -1078,9 +1065,10 @@ void FMobileSceneRenderer::RenderForwardMultiPass(FRDGBuilder& GraphBuilder, FMo
 	BasePassRenderTargets.NumOcclusionQueries = 0;
 	BasePassRenderTargets.SubpassHint = ESubpassHint::DepthReadSubpass;
 
-	auto* SecondPassParameters = GraphBuilder.AllocParameters<FMobileRenderPassParameters>();
+	EMobileSceneTextureSetupMode SetupMode = EMobileSceneTextureSetupMode::SceneDepth | EMobileSceneTextureSetupMode::SceneDepthAux | EMobileSceneTextureSetupMode::CustomDepth;
+	FMobileRenderPassParameters* SecondPassParameters = GraphBuilder.AllocParameters<FMobileRenderPassParameters>();
 	*SecondPassParameters = *PassParameters;
-	SecondPassParameters->SceneTextures = SceneTextures.MobileUniformBuffer;
+	SecondPassParameters->MobileBasePass = CreateMobileBasePassUniformBuffer(GraphBuilder, View, EMobileBasePass::Translucent, SetupMode);
 	SecondPassParameters->RenderTargets = BasePassRenderTargets;
 
 	GraphBuilder.AddPass(
@@ -1240,19 +1228,10 @@ void FMobileSceneRenderer::RenderDeferred(FRDGBuilder& GraphBuilder, const FSort
 
 		UpdateDirectionalLightUniformBuffers(GraphBuilder, View);
 
-		FRDGTextureRef LastFrameScreenSpaceAO = TryRegisterExternalTexture(GraphBuilder, View.PrevViewInfo.MobileAmbientOcclusion);
-
-		if (LastFrameScreenSpaceAO == nullptr)
-		{
-			LastFrameScreenSpaceAO = SystemTextures.White;
-		}
-
-		FTextureRHIRef MobilePixelProjectedReflection = View.PrevViewInfo.MobilePixelProjectedReflection.IsValid() ? View.PrevViewInfo.MobilePixelProjectedReflection->GetRenderTargetItem().ShaderResourceTexture : nullptr;
-
+		EMobileSceneTextureSetupMode SetupMode = EMobileSceneTextureSetupMode::CustomDepth;
 		auto* PassParameters = GraphBuilder.AllocParameters<FMobileRenderPassParameters>();
 		PassParameters->View = View.GetShaderParameters();
-		PassParameters->MobileBasePass = CreateMobileBasePassUniformBuffer(GraphBuilder, View, EMobileBasePass::Opaque, LastFrameScreenSpaceAO, MobilePixelProjectedReflection);
-		PassParameters->SceneTextures = SceneTextures.MobileUniformBuffer;
+		PassParameters->MobileBasePass = CreateMobileBasePassUniformBuffer(GraphBuilder, View, EMobileBasePass::Opaque, SetupMode);
 		PassParameters->RenderTargets = BasePassRenderTargets;
 		#if WITH_DEBUG_VIEW_MODES
 		if (ViewFamily.UseDebugViewPS())
@@ -1265,24 +1244,24 @@ void FMobileSceneRenderer::RenderDeferred(FRDGBuilder& GraphBuilder, const FSort
 
 		if (bRequiresMultiPass)
 		{
-			RenderDeferredMultiPass(GraphBuilder, PassParameters, BasePassRenderTargets, ColorTargets.Num(), ViewIndex, NumViews, View, SceneTextures, SortedLightSet, MobilePixelProjectedReflection);
+			RenderDeferredMultiPass(GraphBuilder, PassParameters, BasePassRenderTargets, ColorTargets.Num(), ViewIndex, NumViews, View, SceneTextures, SortedLightSet);
 		}
 		else
 		{
-			RenderDeferredSinglePass(GraphBuilder, PassParameters, ViewIndex, NumViews, View, SceneTextures, SortedLightSet, MobilePixelProjectedReflection, bUsingPixelLocalStorage);
+			RenderDeferredSinglePass(GraphBuilder, PassParameters, ViewIndex, NumViews, View, SceneTextures, SortedLightSet, bUsingPixelLocalStorage);
 		}
 	}
 
 	QueueSceneTextureExtractions(GraphBuilder, SceneTextures);
 }
 
-void FMobileSceneRenderer::RenderDeferredSinglePass(FRDGBuilder& GraphBuilder, class FMobileRenderPassParameters* PassParameters, int32 ViewIndex, int32 NumViews, FViewInfo& View, FSceneTextures& SceneTextures, const FSortedLightSetSceneInfo& SortedLightSet, FTextureRHIRef MobilePixelProjectedReflection, bool bUsingPixelLocalStorage)
+void FMobileSceneRenderer::RenderDeferredSinglePass(FRDGBuilder& GraphBuilder, class FMobileRenderPassParameters* PassParameters, int32 ViewIndex, int32 NumViews, FViewInfo& View, FSceneTextures& SceneTextures, const FSortedLightSetSceneInfo& SortedLightSet, bool bUsingPixelLocalStorage)
 {
 	GraphBuilder.AddPass(
 		RDG_EVENT_NAME("SceneColorRendering"),
 		PassParameters,
 		ERDGPassFlags::Raster,
-		[this, PassParameters, ViewIndex, NumViews, &View, &SceneTextures, &SortedLightSet, MobilePixelProjectedReflection, bUsingPixelLocalStorage](FRHICommandListImmediate& RHICmdList)
+		[this, PassParameters, ViewIndex, NumViews, &View, &SceneTextures, &SortedLightSet, bUsingPixelLocalStorage](FRHICommandListImmediate& RHICmdList)
 	{
 		if (GIsEditor && !View.bIsSceneCapture)
 		{
@@ -1310,13 +1289,13 @@ void FMobileSceneRenderer::RenderDeferredSinglePass(FRDGBuilder& GraphBuilder, c
 			// SceneColor write, SceneDepth is read only
 			RHICmdList.NextSubpass();
 		}
-		MobileDeferredShadingPass(RHICmdList, ViewIndex, NumViews, View, *Scene, SortedLightSet, MobilePixelProjectedReflection);
+		MobileDeferredShadingPass(RHICmdList, ViewIndex, NumViews, View, *Scene, SortedLightSet);
 		// Draw translucency.
 		RenderTranslucency(RHICmdList, View, &PassParameters->InstanceCullingDrawParamsTranslucency);
 	});
 }
 
-void FMobileSceneRenderer::RenderDeferredMultiPass(FRDGBuilder& GraphBuilder, class FMobileRenderPassParameters* PassParameters, FRenderTargetBindingSlots& BasePassRenderTargets, int32 NumColorTargets, int32 ViewIndex, int32 NumViews, FViewInfo& View, FSceneTextures& SceneTextures, const FSortedLightSetSceneInfo& SortedLightSet, FTextureRHIRef MobilePixelProjectedReflection)
+void FMobileSceneRenderer::RenderDeferredMultiPass(FRDGBuilder& GraphBuilder, class FMobileRenderPassParameters* PassParameters, FRenderTargetBindingSlots& BasePassRenderTargets, int32 NumColorTargets, int32 ViewIndex, int32 NumViews, FViewInfo& View, FSceneTextures& SceneTextures, const FSortedLightSetSceneInfo& SortedLightSet)
 {
 	GraphBuilder.AddPass(
 		RDG_EVENT_NAME("SceneColorRendering"),
@@ -1342,9 +1321,6 @@ void FMobileSceneRenderer::RenderDeferredMultiPass(FRDGBuilder& GraphBuilder, cl
 	});
 
 	// SceneColor + GBuffer write, SceneDepth is read only
-	SceneTextures.MobileSetupMode = EMobileSceneTextureSetupMode::SceneDepth | EMobileSceneTextureSetupMode::SceneDepthAux | EMobileSceneTextureSetupMode::CustomDepth;
-	SceneTextures.MobileUniformBuffer = CreateMobileSceneTextureUniformBuffer(GraphBuilder, SceneTextures.MobileSetupMode);
-
 	for (int32 i = 0; i < NumColorTargets; ++i)
 	{
 		BasePassRenderTargets[i].SetLoadAction(ERenderTargetLoadAction::ELoad);
@@ -1356,9 +1332,10 @@ void FMobileSceneRenderer::RenderDeferredMultiPass(FRDGBuilder& GraphBuilder, cl
 	BasePassRenderTargets.SubpassHint = ESubpassHint::None;
 	BasePassRenderTargets.NumOcclusionQueries = 0;
 
+	EMobileSceneTextureSetupMode SetupMode = EMobileSceneTextureSetupMode::SceneDepth | EMobileSceneTextureSetupMode::SceneDepthAux | EMobileSceneTextureSetupMode::CustomDepth;
 	auto* SecondPassParameters = GraphBuilder.AllocParameters<FMobileRenderPassParameters>();
 	*SecondPassParameters = *PassParameters;
-	SecondPassParameters->SceneTextures = SceneTextures.MobileUniformBuffer;
+	SecondPassParameters->MobileBasePass = CreateMobileBasePassUniformBuffer(GraphBuilder, View, EMobileBasePass::Translucent, SetupMode);
 	SecondPassParameters->RenderTargets = BasePassRenderTargets;
 
 	GraphBuilder.AddPass(
@@ -1372,27 +1349,26 @@ void FMobileSceneRenderer::RenderDeferredMultiPass(FRDGBuilder& GraphBuilder, cl
 	});
 
 	// SceneColor write, SceneDepth is read only
-	SceneTextures.MobileSetupMode = EMobileSceneTextureSetupMode::SceneDepth | EMobileSceneTextureSetupMode::SceneDepthAux | EMobileSceneTextureSetupMode::GBuffers | EMobileSceneTextureSetupMode::CustomDepth;
-	SceneTextures.MobileUniformBuffer = CreateMobileSceneTextureUniformBuffer(GraphBuilder, SceneTextures.MobileSetupMode);
-
 	for (int32 i = 1; i < NumColorTargets; ++i)
 	{
 		BasePassRenderTargets[i] = FRenderTargetBinding();
 	}
 	BasePassRenderTargets.DepthStencil.SetDepthStencilAccess(FExclusiveDepthStencil::DepthRead_StencilWrite);
+
+	SetupMode = EMobileSceneTextureSetupMode::SceneDepth | EMobileSceneTextureSetupMode::SceneDepthAux | EMobileSceneTextureSetupMode::GBuffers | EMobileSceneTextureSetupMode::CustomDepth;
 	auto* ThirdPassParameters = GraphBuilder.AllocParameters<FMobileRenderPassParameters>();
 	*ThirdPassParameters = *PassParameters;
-	ThirdPassParameters->SceneTextures = SceneTextures.MobileUniformBuffer;
+	ThirdPassParameters->MobileBasePass = CreateMobileBasePassUniformBuffer(GraphBuilder, View, EMobileBasePass::Translucent, SetupMode);
 	ThirdPassParameters->RenderTargets = BasePassRenderTargets;
 
 	GraphBuilder.AddPass(
 		{},
 		ThirdPassParameters,
 		ERDGPassFlags::Raster,
-		[this, ThirdPassParameters, ViewIndex, NumViews, &View, &SceneTextures, &SortedLightSet, MobilePixelProjectedReflection](FRHICommandListImmediate& RHICmdList)
+		[this, ThirdPassParameters, ViewIndex, NumViews, &View, &SceneTextures, &SortedLightSet](FRHICommandListImmediate& RHICmdList)
 	{
 		RHICmdList.NextSubpass();
-		MobileDeferredShadingPass(RHICmdList, ViewIndex, NumViews, View, *Scene, SortedLightSet, MobilePixelProjectedReflection);
+		MobileDeferredShadingPass(RHICmdList, ViewIndex, NumViews, View, *Scene, SortedLightSet);
 		// Draw translucency.
 		RenderTranslucency(RHICmdList, View, &ThirdPassParameters->InstanceCullingDrawParamsTranslucency);
 	});
diff --git a/Engine/Source/Runtime/Renderer/Private/MobileTranslucentRendering.cpp b/Engine/Source/Runtime/Renderer/Private/MobileTranslucentRendering.cpp
--- a/Engine/Source/Runtime/Renderer/Private/MobileTranslucentRendering.cpp
+++ b/Engine/Source/Runtime/Renderer/Private/MobileTranslucentRendering.cpp
@@ -46,16 +46,23 @@ void FMobileSceneRenderer::RenderTranslucency(FRHICommandListImmediate& RHICmdLi
 	}
 }
 
+BEGIN_SHADER_PARAMETER_STRUCT(FMobileBasePassParameters, )
+	SHADER_PARAMETER_STRUCT_INCLUDE(FViewShaderParameters, View)
+	SHADER_PARAMETER_RDG_UNIFORM_BUFFER(FMobileBasePassUniformParameters, MobileBasePass)
+	RENDER_TARGET_BINDING_SLOTS()
+END_SHADER_PARAMETER_STRUCT()
+
 void FMobileSceneRenderer::RenderInverseOpacity(FRDGBuilder& GraphBuilder, const FViewInfo& View)
 {
 	View.BeginRenderView();
 	UpdateDirectionalLightUniformBuffers(GraphBuilder, View);
 
 	const FSceneTextures& SceneTextures = FSceneTextures::Get(GraphBuilder);
+	EMobileSceneTextureSetupMode SetupMode = EMobileSceneTextureSetupMode::CustomDepth;
 
 	auto* InverseOpacityParameters = GraphBuilder.AllocParameters<FMobileBasePassParameters>();
 	InverseOpacityParameters->View = View.GetShaderParameters();
-	InverseOpacityParameters->MobileBasePass = CreateMobileBasePassUniformBuffer(GraphBuilder, View, EMobileBasePass::Translucent);
+	InverseOpacityParameters->MobileBasePass = CreateMobileBasePassUniformBuffer(GraphBuilder, View, EMobileBasePass::Translucent, SetupMode);
 	InverseOpacityParameters->RenderTargets[0] = FRenderTargetBinding(SceneTextures.Color.Target, SceneTextures.Color.Resolve, ERenderTargetLoadAction::EClear);
 	InverseOpacityParameters->RenderTargets.DepthStencil = FDepthStencilBinding(SceneTextures.Depth.Target, ERenderTargetLoadAction::EClear, FExclusiveDepthStencil::DepthWrite_StencilWrite);
 	// Opacity could fetch depth as we use exactly the same shaders as in base pass
@@ -68,7 +75,7 @@ void FMobileSceneRenderer::RenderInverseOpacity(FRDGBuilder& GraphBuilder, const
 		const FViewInfo& TranslucentViewport = (View.bIsMobileMultiViewEnabled) ? Views[0] : View;
 		RHICmdList.SetViewport(TranslucentViewport.ViewRect.Min.X, TranslucentViewport.ViewRect.Min.Y, 0.0f, TranslucentViewport.ViewRect.Max.X, TranslucentViewport.ViewRect.Max.Y, 1.0f);
 
-		// Default clear value for a SceneColor is (0,0,0,0), after this passs will blend inverse opacity into final render target with an 1-SrcAlpha op
+		// Default clear value for a SceneColor is (0,0,0,0), after this pass will blend inverse opacity into final render target with an 1-SrcAlpha op
 		// to make this blending work untouched pixels must have alpha = 1
 		DrawClearQuad(RHICmdList, FLinearColor(0, 0, 0, 1));
 		// GPUCULL_TODO: View.ParallelMeshDrawCommandPasses[EMeshPass::MobileInverseOpacity].BuildRenderingCommands(GraphBuilder, Scene->GPUScene);
@@ -77,7 +84,6 @@ void FMobileSceneRenderer::RenderInverseOpacity(FRDGBuilder& GraphBuilder, const
 		if (ShouldRenderTranslucency(ETranslucencyPass::TPT_AllTranslucency) && View.ShouldRenderView())
 		{
 			View.ParallelMeshDrawCommandPasses[EMeshPass::MobileInverseOpacity].DispatchDraw(nullptr, RHICmdList);
-			View.ParallelMeshDrawCommandPasses[EMeshPass::MobileInverseOpacity].HasAnyDraw();
 		}
 	});
 }
diff --git a/Engine/Source/Runtime/Renderer/Private/PlanarReflectionRendering.cpp b/Engine/Source/Runtime/Renderer/Private/PlanarReflectionRendering.cpp
--- a/Engine/Source/Runtime/Renderer/Private/PlanarReflectionRendering.cpp
+++ b/Engine/Source/Runtime/Renderer/Private/PlanarReflectionRendering.cpp
@@ -126,14 +126,6 @@ void SetupPlanarReflectionUniformParameters(const class FSceneView& View, const
 		OutParameters.bIsStereo = false;
 	}
 
-	const bool bIsMobilePixelProjectedReflectionEnabled = IsMobilePixelProjectedReflectionEnabled(View.GetShaderPlatform());
-
-	// We don't generate the RHI resource for the planar reflection render target if the mobile pixel projected reflection is enabled
-	if (bIsMobilePixelProjectedReflectionEnabled)
-	{
-		PlanarReflectionTextureValue = GBlackTexture;
-	}
-
 	OutParameters.PlanarReflectionTexture = PlanarReflectionTextureValue->TextureRHI;
 	OutParameters.PlanarReflectionSampler = PlanarReflectionTextureValue->SamplerStateRHI;
 }
diff --git a/Engine/Source/Runtime/Renderer/Private/PostProcess/PostProcessCompositeEditorPrimitives.cpp b/Engine/Source/Runtime/Renderer/Private/PostProcess/PostProcessCompositeEditorPrimitives.cpp
--- a/Engine/Source/Runtime/Renderer/Private/PostProcess/PostProcessCompositeEditorPrimitives.cpp
+++ b/Engine/Source/Runtime/Renderer/Private/PostProcess/PostProcessCompositeEditorPrimitives.cpp
@@ -481,7 +481,7 @@ FScreenPassTexture AddEditorPrimitivePass(
 			}
 			else
 			{
-				PassParameters->MobileBasePass = CreateMobileBasePassUniformBuffer(GraphBuilder, *EditorView, EMobileBasePass::Translucent);
+				PassParameters->MobileBasePass = CreateMobileBasePassUniformBuffer(GraphBuilder, *EditorView, EMobileBasePass::Translucent, EMobileSceneTextureSetupMode::None);
 			}
 
 			GraphBuilder.AddPass(
diff --git a/Engine/Source/Runtime/Renderer/Private/Renderer.cpp b/Engine/Source/Runtime/Renderer/Private/Renderer.cpp
--- a/Engine/Source/Runtime/Renderer/Private/Renderer.cpp
+++ b/Engine/Source/Runtime/Renderer/Private/Renderer.cpp
@@ -275,7 +275,7 @@ void FRendererModule::DrawTileMesh(FCanvasRenderContext& RenderContext, FMeshPas
 			}
 			else // Mobile
 			{
-				PassParameters->MobileBasePass = CreateMobileBasePassUniformBuffer(GraphBuilder, View, EMobileBasePass::Translucent);
+				PassParameters->MobileBasePass = CreateMobileBasePassUniformBuffer(GraphBuilder, View, EMobileBasePass::Translucent, EMobileSceneTextureSetupMode::None);
 
 				RenderContext.AddPass(RDG_EVENT_NAME("TranslucentMobile"), PassParameters,
 					[Scene, &View, DrawRenderState, &Mesh](FRHICommandListImmediate& RHICmdList)
@@ -354,7 +354,7 @@ void FRendererModule::DrawTileMesh(FCanvasRenderContext& RenderContext, FMeshPas
 				}
 				else // Mobile
 				{
-					PassParameters->MobileBasePass = CreateMobileBasePassUniformBuffer(GraphBuilder, View, EMobileBasePass::Opaque);
+					PassParameters->MobileBasePass = CreateMobileBasePassUniformBuffer(GraphBuilder, View, EMobileBasePass::Opaque, EMobileSceneTextureSetupMode::None);
 
 					RenderContext.AddPass(RDG_EVENT_NAME("OpaqueMobile"), PassParameters,
 						[Scene, &View, DrawRenderState, &Mesh](FRHICommandListImmediate& RHICmdList)
diff --git a/Engine/Source/Runtime/Renderer/Private/SceneRendering.h b/Engine/Source/Runtime/Renderer/Private/SceneRendering.h
--- a/Engine/Source/Runtime/Renderer/Private/SceneRendering.h
+++ b/Engine/Source/Runtime/Renderer/Private/SceneRendering.h
@@ -50,7 +50,6 @@ class FRaytracingLightDataPacked;
 class FRayTracingLocalShaderBindingWriter;
 class FVirtualShadowMapClipmap;
 class FShadowProjectionPassParameters;
-class FMobileBasePassParameters;
 
 struct FCloudRenderContext;
 struct FSceneWithoutWaterTextures;
@@ -2238,8 +2237,8 @@ protected:
 	void RenderForwardMultiPass(FRDGBuilder& GraphBuilder, class FMobileRenderPassParameters* PassParameters, FRenderTargetBindingSlots& BasePassRenderTargets, int32 ViewIndex, FViewInfo& View, FSceneTextures& SceneTextures);
 	
 	void RenderDeferred(FRDGBuilder& GraphBuilder, const FSortedLightSetSceneInfo& SortedLightSet, FRDGTextureRef ViewFamilyTexture, FSceneTextures& SceneTextures);
-	void RenderDeferredSinglePass(FRDGBuilder& GraphBuilder, class FMobileRenderPassParameters* PassParameters, int32 ViewIndex, int32 NumViews, FViewInfo& View, FSceneTextures& SceneTextures, const FSortedLightSetSceneInfo& SortedLightSet, FTextureRHIRef MobilePixelProjectedReflection, bool bUsingPixelLocalStorage);
-	void RenderDeferredMultiPass(FRDGBuilder& GraphBuilder, class FMobileRenderPassParameters* PassParameters, FRenderTargetBindingSlots& BasePassRenderTargets, int32 NumColorTargets, int32 ViewIndex, int32 NumViews, FViewInfo& View, FSceneTextures& SceneTextures, const FSortedLightSetSceneInfo& SortedLightSet, FTextureRHIRef MobilePixelProjectedReflection);
+	void RenderDeferredSinglePass(FRDGBuilder& GraphBuilder, class FMobileRenderPassParameters* PassParameters, int32 ViewIndex, int32 NumViews, FViewInfo& View, FSceneTextures& SceneTextures, const FSortedLightSetSceneInfo& SortedLightSet, bool bUsingPixelLocalStorage);
+	void RenderDeferredMultiPass(FRDGBuilder& GraphBuilder, class FMobileRenderPassParameters* PassParameters, FRenderTargetBindingSlots& BasePassRenderTargets, int32 NumColorTargets, int32 ViewIndex, int32 NumViews, FViewInfo& View, FSceneTextures& SceneTextures, const FSortedLightSetSceneInfo& SortedLightSet);
 	
 	void RenderAmbientOcclusion(FRDGBuilder& GraphBuilder, FRDGTextureRef SceneDepthTexture, FRDGTextureRef AmbientOcclusionTexture);
 
diff --git a/Engine/Source/Runtime/Renderer/Private/SceneTextures.cpp b/Engine/Source/Runtime/Renderer/Private/SceneTextures.cpp
--- a/Engine/Source/Runtime/Renderer/Private/SceneTextures.cpp
+++ b/Engine/Source/Runtime/Renderer/Private/SceneTextures.cpp
@@ -1096,25 +1096,16 @@ void SetupMobileSceneTextureUniformParameters(
 
 	if (const FSceneTextures* SceneTextures = GraphBuilder.Blackboard.Get<FSceneTextures>())
 	{
-		if (EnumHasAnyFlags(SetupMode, EMobileSceneTextureSetupMode::SceneColor))
+		if (EnumHasAnyFlags(SetupMode, EMobileSceneTextureSetupMode::SceneColor) && HasBeenProduced(SceneTextures->Color.Resolve))
 		{
-			if (HasBeenProduced(SceneTextures->Color.Resolve))
-			{
-				SceneTextureParameters.SceneColorTexture = SceneTextures->Color.Resolve;
-			}
-
-			if (HasBeenProduced(SceneTextures->Depth.Resolve))
-			{
-				SceneTextureParameters.SceneDepthTexture = SceneTextures->Depth.Resolve;
-			}
+			SceneTextureParameters.SceneColorTexture = SceneTextures->Color.Resolve;
 		}
 
-		if (EnumHasAnyFlags(SetupMode, EMobileSceneTextureSetupMode::SceneDepth))
+		if (EnumHasAnyFlags(SetupMode, EMobileSceneTextureSetupMode::SceneDepth) && 
+			HasBeenProduced(SceneTextures->Depth.Resolve) && 
+			!EnumHasAnyFlags(SceneTextures->Depth.Resolve->Desc.Flags, TexCreate_Memoryless))
 		{
-			if ((HasBeenProduced(SceneTextures->Depth.Resolve) && !EnumHasAnyFlags(SceneTextures->Depth.Resolve->Desc.Flags, TexCreate_Memoryless)))
-			{
-				SceneTextureParameters.SceneDepthTexture = SceneTextures->Depth.Resolve;
-			}
+			SceneTextureParameters.SceneDepthTexture = SceneTextures->Depth.Resolve;
 		}
 
 		if (SceneTextures->Config.bIsUsingGBuffers)
