Fixed: UpdateDistanceFieldAtlas was assuming that the maximum dispatch group size was the same in all dimensions.

This isn't true in Vulkan, so we need to use the correct dimension for the variable group size (Z, in this case).

#rnx

diff --git a/Engine/Source/Runtime/Renderer/Private/DistanceFieldStreaming.cpp b/Engine/Source/Runtime/Renderer/Private/DistanceFieldStreaming.cpp
--- a/Engine/Source/Runtime/Renderer/Private/DistanceFieldStreaming.cpp
+++ b/Engine/Source/Runtime/Renderer/Private/DistanceFieldStreaming.cpp
@@ -934,7 +934,7 @@ void FDistanceFieldSceneData::UpdateDistanceFieldAtlas(
 			FRDGTextureRef DistanceFieldBrickVolumeTextureRDG = GraphBuilder.RegisterExternalTexture(DistanceFieldBrickVolumeTexture, TEXT("DistanceFields.DistanceFieldBrickVolumeTexture"));
 
 			// GRHIMaxDispatchThreadGroupsPerDimension can be MAX_int32 so we need to do this math in 64-bit.
-			const int32 MaxBrickUploadsPerPass = (int32)FMath::Min<int64>((int64)GRHIMaxDispatchThreadGroupsPerDimension.X * FScatterUploadDistanceFieldAtlasCS::GetGroupSize() / DistanceField::BrickSize, MAX_int32);
+			const int32 MaxBrickUploadsPerPass = (int32)FMath::Min<int64>((int64)GRHIMaxDispatchThreadGroupsPerDimension.Z * FScatterUploadDistanceFieldAtlasCS::GetGroupSize() / DistanceField::BrickSize, MAX_int32);
 
 			for (int32 StartBrickIndex = 0; StartBrickIndex < NumBrickUploads; StartBrickIndex += MaxBrickUploadsPerPass)
 			{
