Fixed Strata draw indirect broken during refactoring.

diff --git a/Engine/Source/Runtime/Renderer/Private/Strata/Strata.cpp b/Engine/Source/Runtime/Renderer/Private/Strata/Strata.cpp
--- a/Engine/Source/Runtime/Renderer/Private/Strata/Strata.cpp
+++ b/Engine/Source/Runtime/Renderer/Private/Strata/Strata.cpp
@@ -175,7 +175,7 @@ void InitialiseStrataFrameSceneData(FSceneRenderer& SceneRenderer, FRDGBuilder&
 			StrataSceneData.ClassificationTileListBufferSRV = GraphBuilder.CreateSRV(StrataSceneData.ClassificationTileListBuffer, PF_R32_UINT);
 			StrataSceneData.ClassificationTileListBufferUAV = GraphBuilder.CreateUAV(StrataSceneData.ClassificationTileListBuffer, PF_R32_UINT);
 
-			StrataSceneData.ClassificationTileIndirectBuffer = GraphBuilder.CreateBuffer(FRDGBufferDesc::CreateIndirectDesc<FRHIDispatchIndirectParameters>(1), TEXT("StrataTileIndirectBuffer"));
+			StrataSceneData.ClassificationTileIndirectBuffer = GraphBuilder.CreateBuffer(FRDGBufferDesc::CreateIndirectDesc<FRHIDrawIndirectParameters>(), TEXT("StrataTileIndirectBuffer"));
 			StrataSceneData.ClassificationTileIndirectBufferSRV = GraphBuilder.CreateSRV(StrataSceneData.ClassificationTileIndirectBuffer, PF_R32_UINT);
 			StrataSceneData.ClassificationTileIndirectBufferUAV = GraphBuilder.CreateUAV(StrataSceneData.ClassificationTileIndirectBuffer, PF_R32_UINT);
 
@@ -630,7 +630,7 @@ static void AddStrataInternalClassifedTilePass(
 			RHICmdList.SetStencilRef(StencilBit);
 			RHICmdList.SetViewport(0, 0, 0.0f, OutputResolution.X, OutputResolution.Y, 1.0f);
 			RHICmdList.SetStreamSource(0, nullptr, 0);
-			RHICmdList.DrawPrimitiveIndirect(ParametersPS->TileIndirectBuffer->GetRHI(), 0);
+			RHICmdList.DrawPrimitiveIndirect(ParametersPS->TileIndirectBuffer->GetIndirectRHICallBuffer(), 0);
 		});
 }
 
