Use float4 instead of float3 for global input value in ReflectionEnvironmentCapture to fix translation in SPIR-V backends.

#rnx

diff --git a/Engine/Shaders/Private/ReflectionEnvironmentShaders.usf b/Engine/Shaders/Private/ReflectionEnvironmentShaders.usf
--- a/Engine/Shaders/Private/ReflectionEnvironmentShaders.usf
+++ b/Engine/Shaders/Private/ReflectionEnvironmentShaders.usf
@@ -44,7 +44,7 @@ SamplerState SceneDepthSampler;
  * Y = Sky distance threshold, 
  * Z = whether a skylight's lower hemisphere should be replaced with LowerHemisphereColor.
  */
-float3 SkyLightCaptureParameters;
+float4 SkyLightCaptureParameters;
 float4 LowerHemisphereColor;
 
 void CopySceneColorToCubeFaceColorPS(
diff --git a/Engine/Source/Runtime/Renderer/Private/ReflectionEnvironmentCapture.cpp b/Engine/Source/Runtime/Renderer/Private/ReflectionEnvironmentCapture.cpp
--- a/Engine/Source/Runtime/Renderer/Private/ReflectionEnvironmentCapture.cpp
+++ b/Engine/Source/Runtime/Renderer/Private/ReflectionEnvironmentCapture.cpp
@@ -613,7 +613,7 @@ public:
 		SHADER_PARAMETER_SAMPLER(SamplerState, SceneColorSampler)
 		SHADER_PARAMETER_RDG_TEXTURE(Texture2D, SceneDepthTexture)
 		SHADER_PARAMETER_SAMPLER(SamplerState, SceneDepthSampler)
-		SHADER_PARAMETER(FVector, SkyLightCaptureParameters)
+		SHADER_PARAMETER(FVector4, SkyLightCaptureParameters)
 		SHADER_PARAMETER(FVector4, LowerHemisphereColor)
 		RENDER_TARGET_BINDING_SLOTS()
 	END_SHADER_PARAMETER_STRUCT()
@@ -631,7 +631,7 @@ public:
 	BEGIN_SHADER_PARAMETER_STRUCT(FParameters, )
 		SHADER_PARAMETER_TEXTURE(TextureCube, SourceCubemapTexture)
 		SHADER_PARAMETER_SAMPLER(SamplerState, SourceCubemapSampler)
-		SHADER_PARAMETER(FVector, SkyLightCaptureParameters)
+		SHADER_PARAMETER(FVector4, SkyLightCaptureParameters)
 		SHADER_PARAMETER(int32, CubeFace)
 		SHADER_PARAMETER(FVector4, LowerHemisphereColor)
 		SHADER_PARAMETER(FVector2D, SinCosSourceCubemapRotation)
@@ -749,13 +749,13 @@ void CaptureSceneToScratchCubemap(FRHICommandListImmediate& RHICmdList, FSceneRe
 		PassParameters->LowerHemisphereColor = LowerHemisphereColor;
 
 		{
-			FVector SkyLightParametersValue = FVector::ZeroVector;
+			FVector4 SkyLightParametersValue(ForceInitToZero);
 			FScene* Scene = SceneRenderer->Scene;
 
 			if (bCapturingForSkyLight)
 			{
 				// When capturing reflection captures, support forcing all low hemisphere lighting to be black
-				SkyLightParametersValue = FVector(0, 0, bLowerHemisphereIsBlack ? 1.0f : 0.0f);
+				SkyLightParametersValue = FVector4(0, 0, bLowerHemisphereIsBlack ? 1.0f : 0.0f, 0);
 			}
 			else if (!bCapturingForMobile && Scene->SkyLight && !Scene->SkyLight->bHasStaticLighting)	
 			{
@@ -763,12 +763,12 @@ void CaptureSceneToScratchCubemap(FRHICommandListImmediate& RHICmdList, FSceneRe
 				
 				// When capturing reflection captures and there's a stationary sky light, mask out any pixels whose depth classify it as part of the sky
 				// This will allow changing the stationary sky light at runtime
-				SkyLightParametersValue = FVector(1, Scene->SkyLight->SkyDistanceThreshold, 0);
+				SkyLightParametersValue = FVector4(1, Scene->SkyLight->SkyDistanceThreshold, 0, 0);
 			}
 			else
 			{
 				// When capturing reflection captures and there's no sky light, or only a static sky light, capture all depth ranges
-				SkyLightParametersValue = FVector(2, 0, 0);
+				SkyLightParametersValue = FVector4(2, 0, 0, 0);
 			}
 
 			PassParameters->SkyLightCaptureParameters = SkyLightParametersValue;
