Don't add primitives to be ray traced if they use an unsupported vertex factory type.

- This fixes a handful of crashes where ray tracing code would fail to find shaders for VF types that were not supported when we went to render.
- Added early out code in each proxy's GetDynamicRayTracingInstances() function to avoid adding ray tracing instances if we don't support the VF Type.
- I attempted to to early out of the function as early as possible.
- Since skeletal mesh sections can in theory have different vertex factories we have to check each one.  The logic I've gone with is that if _any_ section uses an unsupported VF we opt out of adding any of the sections to ray trace.  The result would be the whole skeletal mesh won't draw (idea being we wouldn't want partial meshes being draw).

diff --git a/Engine/Plugins/Experimental/GeometryCache/Source/GeometryCache/Private/GeometryCacheSceneProxy.cpp b/Engine/Plugins/Experimental/GeometryCache/Source/GeometryCache/Private/GeometryCacheSceneProxy.cpp
--- a/Engine/Plugins/Experimental/GeometryCache/Source/GeometryCache/Private/GeometryCacheSceneProxy.cpp
+++ b/Engine/Plugins/Experimental/GeometryCache/Source/GeometryCache/Private/GeometryCacheSceneProxy.cpp
@@ -505,6 +505,11 @@ void FGeometryCacheSceneProxy::GetDynamicRayTracingInstances(FRayTracingMaterial
 			continue;
 		}
 
+		if (!TrackProxy->VertexFactory.GetType()->SupportsRayTracingDynamicGeometry())
+		{
+			continue;
+		}
+
 		FRayTracingInstance RayTracingInstance;
 		RayTracingInstance.Geometry = &TrackProxy->RayTracingGeometry;
 		RayTracingInstance.InstanceTransforms.Add(GetLocalToWorld());
diff --git a/Engine/Plugins/FX/Niagara/Source/Niagara/Private/NiagaraRendererMeshes.cpp b/Engine/Plugins/FX/Niagara/Source/Niagara/Private/NiagaraRendererMeshes.cpp
--- a/Engine/Plugins/FX/Niagara/Source/Niagara/Private/NiagaraRendererMeshes.cpp
+++ b/Engine/Plugins/FX/Niagara/Source/Niagara/Private/NiagaraRendererMeshes.cpp
@@ -1073,6 +1073,12 @@ void FNiagaraRendererMeshes::GetDynamicRayTracingInstances(FRayTracingMaterialGa
 		const FMeshData& MeshData = Meshes[MeshIndex];
 		const int32 LODIndex = GetLODIndex(MeshIndex);
 
+		FVertexFactory* VertexFactory = &MeshData.RenderData->LODVertexFactories[LODIndex].VertexFactory;
+		if (!VertexFactory->GetType()->SupportsRayTracingDynamicGeometry())
+		{
+			continue;
+		}
+
 		const FStaticMeshLODResources& LODModel = MeshData.RenderData->LODResources[LODIndex];
 		FRayTracingGeometry& Geometry = MeshData.RenderData->LODResources[LODIndex].RayTracingGeometry;
 		FRayTracingInstance RayTracingInstance;
@@ -1102,7 +1108,6 @@ void FNiagaraRendererMeshes::GetDynamicRayTracingInstances(FRayTracingMaterialGa
 			FMeshBatch MeshBatch;
 			const FStaticMeshLODResources& LOD = MeshData.RenderData->LODResources[LODIndex];
 			const FStaticMeshVertexFactories& VFs = MeshData.RenderData->LODVertexFactories[LODIndex];
-			FVertexFactory* VertexFactory = &MeshData.RenderData->LODVertexFactories[LODIndex].VertexFactory;
 
 			MeshBatch.VertexFactory = VertexFactory;
 			MeshBatch.MaterialRenderProxy = MaterialProxy;
diff --git a/Engine/Plugins/FX/Niagara/Source/Niagara/Private/NiagaraRendererRibbons.cpp b/Engine/Plugins/FX/Niagara/Source/Niagara/Private/NiagaraRendererRibbons.cpp
--- a/Engine/Plugins/FX/Niagara/Source/Niagara/Private/NiagaraRendererRibbons.cpp
+++ b/Engine/Plugins/FX/Niagara/Source/Niagara/Private/NiagaraRendererRibbons.cpp
@@ -1400,6 +1400,11 @@ void FNiagaraRendererRibbons::GetDynamicRayTracingInstances(FRayTracingMaterialG
 	// Setup material for our ray tracing instance
 	FNiagaraMeshCollectorResourcesRibbon& CollectorResources = Context.RayTracingMeshResourceCollector.AllocateOneFrameResource<FNiagaraMeshCollectorResourcesRibbon>();
 
+	if (!CollectorResources.VertexFactory.GetType()->SupportsRayTracingDynamicGeometry())
+	{
+		return;
+	}
+
 	FGlobalDynamicIndexBuffer::FAllocationEx DynamicIndexAllocation;
 	CreatePerViewResources(Context.ReferenceView, Context.ReferenceViewFamily, SceneProxy, Context.RayTracingMeshResourceCollector, CollectorResources.UniformBuffer, DynamicIndexAllocation);
 
diff --git a/Engine/Plugins/FX/Niagara/Source/Niagara/Private/NiagaraRendererSprites.cpp b/Engine/Plugins/FX/Niagara/Source/Niagara/Private/NiagaraRendererSprites.cpp
--- a/Engine/Plugins/FX/Niagara/Source/Niagara/Private/NiagaraRendererSprites.cpp
+++ b/Engine/Plugins/FX/Niagara/Source/Niagara/Private/NiagaraRendererSprites.cpp
@@ -994,6 +994,11 @@ void FNiagaraRendererSprites::GetDynamicRayTracingInstances(FRayTracingMaterialG
 			CollectorResources = &Context.RayTracingMeshResourceCollector.AllocateOneFrameResource<FMeshCollectorResources>();
 		}
 
+		if (!CollectorResources->GetVertexFactory().GetType()->SupportsRayTracingDynamicGeometry())
+		{
+			return;
+		}
+
 		// Setup material for our ray tracing instance
 		FCPUSimParticleDataAllocation CPUSimParticleDataAllocation = ConditionalAllocateCPUSimParticleData(DynamicDataSprites, RendererLayout, Context.RayTracingMeshResourceCollector.GetDynamicReadBuffer(), bNeedsGPUVis);
 		FNiagaraSpriteVFLooseParameters VFLooseParams;
diff --git a/Engine/Source/Runtime/Engine/Private/InstancedStaticMesh.cpp b/Engine/Source/Runtime/Engine/Private/InstancedStaticMesh.cpp
--- a/Engine/Source/Runtime/Engine/Private/InstancedStaticMesh.cpp
+++ b/Engine/Source/Runtime/Engine/Private/InstancedStaticMesh.cpp
@@ -1413,6 +1413,11 @@ void FInstancedStaticMeshSceneProxy::GetDynamicRayTracingInstances(struct FRayTr
 		return;
 	}
 
+	if (!InstancedRenderData.VertexFactories[LOD].GetType()->SupportsRayTracingDynamicGeometry())
+	{
+		return;
+	}
+
 	//setup a 'template' for the instance first, so we aren't duplicating work
 	//#dxr_todo: when multiple LODs are used, template needs to be an array of templates, probably best initialized on-demand via a lamda
 	FRayTracingInstance RayTracingInstanceTemplate;
diff --git a/Engine/Source/Runtime/Engine/Private/SkeletalMesh.cpp b/Engine/Source/Runtime/Engine/Private/SkeletalMesh.cpp
--- a/Engine/Source/Runtime/Engine/Private/SkeletalMesh.cpp
+++ b/Engine/Source/Runtime/Engine/Private/SkeletalMesh.cpp
@@ -6219,17 +6219,31 @@ void FSkeletalMeshSceneProxy::GetDynamicRayTracingInstances(FRayTracingMaterialG
 				MeshObject->SectionIndexPreview = INDEX_NONE;
 				MeshObject->MaterialIndexPreview = INDEX_NONE;
 			#endif
+
+				bool bSupportedVFType = true;
 				for (FSkeletalMeshSectionIter Iter(LODIndex, *MeshObject, LODData, LODSection); Iter; ++Iter)
 				{
 					const FSkelMeshRenderSection& Section = Iter.GetSection();
 					const int32 SectionIndex = Iter.GetSectionElementIndex();
 					const FSectionElementInfo& SectionElementInfo = Iter.GetSectionElementInfo();
 
+					if (!MeshObject->GetSkinVertexFactory(Context.ReferenceView, LODIndex, SectionIndex)->GetType()->SupportsRayTracingDynamicGeometry())
+					{
+						bSupportedVFType = false;
+						break;
+					}
+
 					FMeshBatch MeshBatch;
 					CreateBaseMeshBatch(Context.ReferenceView, LODData, LODIndex, SectionIndex, SectionElementInfo, MeshBatch);
 
 					RayTracingInstance.Materials.Add(MeshBatch);
 				}
+
+				if (!bSupportedVFType)
+				{
+					return;
+				}
+
 			#if WITH_EDITORONLY_DATA
 				MeshObject->SectionIndexPreview = SectionIndexPreview;
 				MeshObject->MaterialIndexPreview = MaterialIndexPreview;
diff --git a/Engine/Source/Runtime/Engine/Private/StaticMeshRender.cpp b/Engine/Source/Runtime/Engine/Private/StaticMeshRender.cpp
--- a/Engine/Source/Runtime/Engine/Private/StaticMeshRender.cpp
+++ b/Engine/Source/Runtime/Engine/Private/StaticMeshRender.cpp
@@ -1697,6 +1697,11 @@ void FStaticMeshSceneProxy::GetDynamicRayTracingInstances(FRayTracingMaterialGat
 	const uint32 LODIndex = FMath::Max(GetLOD(Context.ReferenceView), (int32)GetCurrentFirstLODIdx_RenderThread());
 	const FStaticMeshLODResources& LODModel = RenderData->LODResources[LODIndex];
 
+	if (!RenderData->LODVertexFactories[LODIndex].VertexFactory.GetType()->SupportsRayTracingDynamicGeometry())
+	{
+		return;
+	}
+
 	if (LODModel.GetNumVertices() <= 0)
 	{
 		return;
diff --git a/Engine/Source/Runtime/Landscape/Private/LandscapeRender.cpp b/Engine/Source/Runtime/Landscape/Private/LandscapeRender.cpp
--- a/Engine/Source/Runtime/Landscape/Private/LandscapeRender.cpp
+++ b/Engine/Source/Runtime/Landscape/Private/LandscapeRender.cpp
@@ -2472,6 +2472,11 @@ void FLandscapeComponentSceneProxy::GetDynamicRayTracingInstances(FRayTracingMat
 		return;
 	}
 
+	if (!VertexFactory->GetType()->SupportsRayTracingDynamicGeometry())
+	{
+		return;
+	}
+
 	int32 ForcedLODLevel = ForcedLOD;
 
 	int32 ViewLodOveride = GetViewLodOverride(*Context.ReferenceView);
diff --git a/Engine/Source/Runtime/Renderer/Private/RayTracing/RayTracingDynamicGeometry.cpp b/Engine/Source/Runtime/Renderer/Private/RayTracing/RayTracingDynamicGeometry.cpp
--- a/Engine/Source/Runtime/Renderer/Private/RayTracing/RayTracingDynamicGeometry.cpp
+++ b/Engine/Source/Runtime/Renderer/Private/RayTracing/RayTracingDynamicGeometry.cpp
@@ -3,6 +3,7 @@
 #include "MeshMaterialShader.h"
 #include "ScenePrivate.h"
 #include "RayTracingDynamicGeometryCollection.h"
+#include "RayTracingInstance.h"
 
 #if RHI_RAYTRACING
 
@@ -12,11 +13,6 @@ DECLARE_CYCLE_STAT(TEXT("RTDynGeomBuild"), STAT_CLM_RTDynGeomBuild, STATGROUP_Pa
 // Workaround for outstanding memory corruption on some platforms when parallel command list translation is used.
 #define USE_RAY_TRACING_DYNAMIC_GEOMETRY_PARALLEL_COMMAND_LISTS 0
 
-static bool IsSupportedDynamicVertexFactoryType(const FVertexFactoryType* VertexFactoryType)
-{
-	return VertexFactoryType->SupportsRayTracingDynamicGeometry();
-}
-
 class FRayTracingDynamicGeometryConverterCS : public FMeshMaterialShader
 {
 	DECLARE_SHADER_TYPE(FRayTracingDynamicGeometryConverterCS, MeshMaterial);
@@ -41,7 +37,7 @@ public:
 
 	static bool ShouldCompilePermutation(const FMeshMaterialShaderPermutationParameters& Parameters)
 	{
-		return IsSupportedDynamicVertexFactoryType(Parameters.VertexFactoryType) && ShouldCompileRayTracingShadersForProject(Parameters.Platform);
+		return Parameters.VertexFactoryType->SupportsRayTracingDynamicGeometry() && ShouldCompileRayTracingShadersForProject(Parameters.Platform);
 	}
 
 	void GetShaderBindings(
@@ -174,7 +170,7 @@ void FRayTracingDynamicGeometryCollection::AddDynamicMeshBatchForGeometryUpdate(
 
 	for (const FMeshBatch& MeshBatch : UpdateParams.MeshBatches)
 	{
-		if (!ensureMsgf(IsSupportedDynamicVertexFactoryType(MeshBatch.VertexFactory->GetType()),
+		if (!ensureMsgf(MeshBatch.VertexFactory->GetType()->SupportsRayTracingDynamicGeometry(),
 			TEXT("FRayTracingDynamicGeometryConverterCS doesn't support %s. Skipping rendering of %s.  This can happen when the skinning cache runs out of space and falls back to GPUSkinVertexFactory."),
 			MeshBatch.VertexFactory->GetType()->GetName(), *PrimitiveSceneProxy->GetOwnerName().ToString()))
 		{
