Fix CIS issue.

diff --git a/Engine/Source/Runtime/Renderer/Public/HairStrandsInterface.h b/Engine/Source/Runtime/Renderer/Public/HairStrandsInterface.h
--- a/Engine/Source/Runtime/Renderer/Public/HairStrandsInterface.h
+++ b/Engine/Source/Runtime/Renderer/Public/HairStrandsInterface.h
@@ -143,7 +143,7 @@ struct RENDERER_API FHairStrandsInstance
 protected:
 	mutable uint32 RefCount = 0;
 };
-typedef RENDERER_API TArray<FHairStrandsInstance*> FHairStrandsInstances;
+typedef TArray<FHairStrandsInstance*> FHairStrandsInstances;
 
 class RENDERER_API FHairGroupPublicData : public FRenderResource
 {
