Set Mobile FXAA to disabled when postprocessing is disabled

diff --git a/Engine/Source/Runtime/Renderer/Private/PostProcess/PostProcessing.cpp b/Engine/Source/Runtime/Renderer/Private/PostProcess/PostProcessing.cpp
--- a/Engine/Source/Runtime/Renderer/Private/PostProcess/PostProcessing.cpp
+++ b/Engine/Source/Runtime/Renderer/Private/PostProcess/PostProcessing.cpp
@@ -1894,6 +1894,7 @@ void AddMobilePostProcessingPasses(FRDGBuilder& GraphBuilder, FScene* Scene, con
 		PassSequence.SetEnabled(EPass::SunMerge, false);
 		PassSequence.SetEnabled(EPass::SeparateTranslucency, false);
 		PassSequence.SetEnabled(EPass::TAA, false);
+		PassSequence.SetEnabled(EPass::FXAA, false);
 		PassSequence.SetEnabled(EPass::PostProcessMaterialAfterTonemapping, false);
 		PassSequence.SetEnabled(EPass::FXAA, false);
 		PassSequence.Finalize();
