Fix CIS unity build issue.

diff --git a/Engine/Source/Runtime/Renderer/Private/HairStrands/HairStrandsData.h b/Engine/Source/Runtime/Renderer/Private/HairStrands/HairStrandsData.h
--- a/Engine/Source/Runtime/Renderer/Private/HairStrands/HairStrandsData.h
+++ b/Engine/Source/Runtime/Renderer/Private/HairStrands/HairStrandsData.h
@@ -11,6 +11,7 @@
 #include "Shader.h"
 
 class FPrimitiveSceneProxy;
+class FViewInfo;
 struct FMeshBatch;
 
 ////////////////////////////////////////////////////////////////////////////////////
@@ -194,7 +195,7 @@ struct FHairStrandsDeepShadowData
 
 	FIntPoint ShadowResolution = FIntPoint::ZeroValue;
 	uint32 LightId = ~0;
-	ELightComponentType LightType = LightType_MAX;
+	bool bIsLightDirectional = false;
 	FVector  LightDirection;
 	FVector4 LightPosition;
 	FLinearColor LightLuminance;
diff --git a/Engine/Source/Runtime/Renderer/Private/HairStrands/HairStrandsDeepShadow.cpp b/Engine/Source/Runtime/Renderer/Private/HairStrands/HairStrandsDeepShadow.cpp
--- a/Engine/Source/Runtime/Renderer/Private/HairStrands/HairStrandsDeepShadow.cpp
+++ b/Engine/Source/Runtime/Renderer/Private/HairStrands/HairStrandsDeepShadow.cpp
@@ -100,7 +100,7 @@ void AddInjectHairVoxelShadowCaster(
 	Parameters->VirtualVoxel = VoxelResources.UniformBuffer;
 	Parameters->LightDirection = DomData.LightDirection;
 	Parameters->LightPosition = DomData.LightPosition;
-	Parameters->bIsDirectional = DomData.LightType == LightType_Directional ? 1 : 0;
+	Parameters->bIsDirectional = DomData.bIsLightDirectional ? 1 : 0;
 	Parameters->MacroGroupId = DomData.MacroGroupId;
 	Parameters->DeepShadowViewInfoBuffer = DeepShadowViewInfoBufferSRV;
 	Parameters->bIsGPUDriven = GDeepShadowGPUDriven > 0;
@@ -309,7 +309,7 @@ void RenderHairStrandsDeepShadows(
 				}
 
 				const ELightComponentType LightType = (ELightComponentType)LightProxy->GetLightType();
-
+				const bool bIsDirectional = LightType == ELightComponentType::LightType_Directional;
 				FMinHairRadiusAtDepth1 MinStrandRadiusAtDepth1;
 				const FIntPoint AtlasRectOffset(
 					(TotalAtlasSlotIndex % AtlasSlotDimension.X) * AtlasSlotResolution.X,
@@ -319,10 +319,10 @@ void RenderHairStrandsDeepShadows(
 				FHairStrandsDeepShadowData& DomData = MacroGroup.DeepShadowDatas.AddZeroed_GetRef();
 				ComputeWorldToLightClip(DomData.CPU_WorldToLightTransform, MinStrandRadiusAtDepth1, MacroGroupBounds, *LightProxy, LightType, AtlasSlotResolution);
 				DomData.LightDirection = LightProxy->GetDirection();
-				DomData.LightPosition = FVector4(LightProxy->GetPosition(), LightType == ELightComponentType::LightType_Directional ? 0 : 1);
+				DomData.LightPosition = FVector4(LightProxy->GetPosition(), bIsDirectional ? 0 : 1);
 				DomData.LightLuminance = LightProxy->GetColor();
 				DomData.LayerDistribution = LightProxy->GetDeepShadowLayerDistribution();
-				DomData.LightType = LightType;
+				DomData.bIsLightDirectional = bIsDirectional;
 				DomData.LightId = LightInfo->Id;
 				DomData.ShadowResolution = AtlasSlotResolution;
 				DomData.Bounds = MacroGroupBounds;
@@ -358,7 +358,7 @@ void RenderHairStrandsDeepShadows(
 				for (FHairStrandsDeepShadowData& DomData : MacroGroup.DeepShadowDatas)
 				{				
 					Parameters->LightDirections[DomData.AtlasSlotIndex]		= FVector4(DomData.LightDirection.X, DomData.LightDirection.Y, DomData.LightDirection.Z, 0);
-					Parameters->LightPositions[DomData.AtlasSlotIndex]		= FVector4(DomData.LightPosition.X, DomData.LightPosition.Y, DomData.LightPosition.Z, DomData.LightType == LightType_Directional ? 0 : 1);
+					Parameters->LightPositions[DomData.AtlasSlotIndex]		= FVector4(DomData.LightPosition.X, DomData.LightPosition.Y, DomData.LightPosition.Z, DomData.bIsLightDirectional ? 0 : 1);
 					Parameters->MacroGroupIndices[DomData.AtlasSlotIndex]	= FIntVector4(DomData.MacroGroupId, 0,0,0);
 				}
 			}
@@ -392,7 +392,7 @@ void RenderHairStrandsDeepShadows(
 		{
 			for (FHairStrandsDeepShadowData& DomData : MacroGroup.DeepShadowDatas)
 			{
-				const bool bIsOrtho = DomData.LightType == ELightComponentType::LightType_Directional;
+				const bool bIsOrtho = DomData.bIsLightDirectional;
 				const FVector4 HairRenderInfo = PackHairRenderInfo(DomData.CPU_MinStrandRadiusAtDepth1.Primary, DomData.CPU_MinStrandRadiusAtDepth1.Stable, DomData.CPU_MinStrandRadiusAtDepth1.Primary, 1);
 				const uint32 HairRenderInfoBits = PackHairRenderInfoBits(bIsOrtho, DeepShadowResources.bIsGPUDriven);
 					
diff --git a/Engine/Source/Runtime/Renderer/Private/HairStrands/HairStrandsTile.h b/Engine/Source/Runtime/Renderer/Private/HairStrands/HairStrandsTile.h
--- a/Engine/Source/Runtime/Renderer/Private/HairStrands/HairStrandsTile.h
+++ b/Engine/Source/Runtime/Renderer/Private/HairStrands/HairStrandsTile.h
@@ -13,6 +13,8 @@
 #include "ShaderParameterStruct.h"
 #include "HairStrandsData.h"
 
+class FViewInfo;
+
 class FHairStrandsTilePassVS : public FGlobalShader
 {
 	DECLARE_GLOBAL_SHADER(FHairStrandsTilePassVS);
