Treat stationary lights as movable for the purposes of shadowing when Non-Nanite VSMs are enabled.
Remove VirtualShadowMapConfig.h which is no longer necessary.

diff --git a/Engine/Source/Runtime/Engine/Private/Components/LightComponent.cpp b/Engine/Source/Runtime/Engine/Private/Components/LightComponent.cpp
--- a/Engine/Source/Runtime/Engine/Private/Components/LightComponent.cpp
+++ b/Engine/Source/Runtime/Engine/Private/Components/LightComponent.cpp
@@ -339,6 +339,13 @@ FLightSceneProxy::FLightSceneProxy(const ULightComponent* InLightComponent)
 {
 	check(SceneInterface);
 
+	// Treat stationary lights as movable when non-nanite VSMs are enabled
+	const bool bNonNaniteVirtualShadowMaps = UseNonNaniteVirtualShadowMaps(SceneInterface->GetShaderPlatform(), SceneInterface->GetFeatureLevel());
+	if (bNonNaniteVirtualShadowMaps)
+	{
+		bStaticShadowing = bStaticLighting;
+	}
+
 	const FLightComponentMapBuildData* MapBuildData = InLightComponent->GetLightComponentMapBuildData();
 	
 	if (MapBuildData && bStaticShadowing && !bStaticLighting)
diff --git a/Engine/Source/Runtime/Engine/Public/SceneManagement.h b/Engine/Source/Runtime/Engine/Public/SceneManagement.h
--- a/Engine/Source/Runtime/Engine/Public/SceneManagement.h
+++ b/Engine/Source/Runtime/Engine/Public/SceneManagement.h
@@ -1651,7 +1651,7 @@ protected:
 	 * The light may still have dynamic brightness and color. 
 	 * The light may or may not also have static lighting.
 	 */
-	const uint8 bStaticShadowing : 1;
+	uint8 bStaticShadowing : 1;
 
 	/** True if the light casts dynamic shadows. */
 	const uint8 bCastDynamicShadow : 1;
diff --git a/Engine/Source/Runtime/RenderCore/Public/RenderUtils.h b/Engine/Source/Runtime/RenderCore/Public/RenderUtils.h
--- a/Engine/Source/Runtime/RenderCore/Public/RenderUtils.h
+++ b/Engine/Source/Runtime/RenderCore/Public/RenderUtils.h
@@ -684,3 +684,18 @@ inline bool UseVirtualShadowMaps(EShaderPlatform ShaderPlatform, const FStaticFe
 
 	return UseGPUScene(ShaderPlatform, FeatureLevel) && !bForwardShadingEnabled && (EnableVirtualSMCVar && EnableVirtualSMCVar->GetInt() != 0) && DoesPlatformSupportNanite(ShaderPlatform);
 }
+
+/**
+* Returns true if non-nanite virtual shadow maps are enbled by CVar r.Shadow.Virtual.NonNaniteVSM 
+* and UseVirtualShadowMaps is true for the given platform and feature level.
+*/
+inline bool UseNonNaniteVirtualShadowMaps(EShaderPlatform ShaderPlatform, const FStaticFeatureLevel FeatureLevel)
+{
+	static const auto EnableCVar = IConsoleManager::Get().FindConsoleVariable(TEXT("r.Shadow.Virtual.NonNaniteVSM"));
+
+#if GPUCULL_TODO
+	return EnableCVar->GetInt() != 0 && UseVirtualShadowMaps(ShaderPlatform, FeatureLevel);
+#else // !GPUCULL_TODO
+	return false;
+#endif // GPUCULL_TODO
+}
diff --git a/Engine/Source/Runtime/Renderer/Private/SceneRendering.h b/Engine/Source/Runtime/Renderer/Private/SceneRendering.h
--- a/Engine/Source/Runtime/Renderer/Private/SceneRendering.h
+++ b/Engine/Source/Runtime/Renderer/Private/SceneRendering.h
@@ -33,7 +33,6 @@
 #include "GpuDebugRendering.h"
 #include "PostProcess/PostProcessAmbientOcclusionMobile.h"
 #include "Nanite/NaniteRender.h"
-#include "VirtualShadowMaps/VirtualShadowMapConfig.h"
 #include "VirtualShadowMaps/VirtualShadowMapArray.h"
 #include "Lumen/LumenTranslucencyVolumeLighting.h"
 #include "HairStrands/HairStrandsData.h"
diff --git a/Engine/Source/Runtime/Renderer/Private/ShadowRendering.h b/Engine/Source/Runtime/Renderer/Private/ShadowRendering.h
--- a/Engine/Source/Runtime/Renderer/Private/ShadowRendering.h
+++ b/Engine/Source/Runtime/Renderer/Private/ShadowRendering.h
@@ -29,7 +29,6 @@
 #include "LightRendering.h"
 #include "HairStrands/HairStrandsRendering.h"
 #include "Strata/Strata.h"
-#include "VirtualShadowMaps/VirtualShadowMapConfig.h"
 
 ENGINE_API IPooledRenderTarget* GetSubsufaceProfileTexture_RT(FRHICommandListImmediate& RHICmdList);
 
diff --git a/Engine/Source/Runtime/Renderer/Private/ShadowSetup.cpp b/Engine/Source/Runtime/Renderer/Private/ShadowSetup.cpp
--- a/Engine/Source/Runtime/Renderer/Private/ShadowSetup.cpp
+++ b/Engine/Source/Runtime/Renderer/Private/ShadowSetup.cpp
@@ -47,6 +47,8 @@ static FAutoConsoleVariableRef CVarMinScreenRadiusForShadowCaster(
 	ECVF_Scalability | ECVF_RenderThreadSafe
 	);
 
+extern int GEnableNonNaniteVSM;
+
 int32 GCacheWholeSceneShadows = 1;
 FAutoConsoleVariableRef CVarCacheWholeSceneShadows(
 	TEXT("r.Shadow.CacheWholeSceneShadows"),
@@ -4972,6 +4974,7 @@ void FSceneRenderer::InitDynamicShadows(FRHICommandListImmediate& RHICmdList, FG
 	const bool bProjectEnablePointLightShadows = Scene->ReadOnlyCVARCache.bEnablePointLightShadows && !bMobile; // Point light shadow is unsupported on mobile for now.
 	const bool bProjectEnableMovableDirectionLightShadows = !bMobile || Scene->ReadOnlyCVARCache.bMobileAllowMovableDirectionalLights;
 	const bool bProjectEnableMovableSpotLightShadows = !bMobile || Scene->ReadOnlyCVARCache.bMobileEnableMovableSpotlightsShadow;
+	const bool bProjectEnableNonNaniteVirtualShadowMaps = UseNonNaniteVirtualShadowMaps(ShaderPlatform, FeatureLevel);
 
 	uint32 NumPointShadowCachesUpdatedThisFrame = 0;
 	uint32 NumSpotShadowCachesUpdatedThisFrame = 0;
@@ -5083,7 +5086,8 @@ void FSceneRenderer::InitDynamicShadows(FRHICommandListImmediate& RHICmdList, FG
 							AddViewDependentWholeSceneShadowsForView(ViewDependentWholeSceneShadows, ViewDependentWholeSceneShadowsThatNeedCulling, VisibleLightInfo, *LightSceneInfo);
 						}
 
-						if( !bMobile || (LightSceneInfo->Proxy->CastsModulatedShadows() && !LightSceneInfo->Proxy->UseCSMForDynamicObjects()))
+						// Disable per-object shadows when non-nanite VSMs are enabled
+						if( !bProjectEnableNonNaniteVirtualShadowMaps && (!bMobile || (LightSceneInfo->Proxy->CastsModulatedShadows() && !LightSceneInfo->Proxy->UseCSMForDynamicObjects())))
 						{
 							const TArray<FLightPrimitiveInteraction*>* InteractionShadowPrimitives = LightSceneInfo->GetInteractionShadowPrimitives();
 
diff --git a/Engine/Source/Runtime/Renderer/Private/VirtualShadowMaps/VirtualShadowMapArray.h b/Engine/Source/Runtime/Renderer/Private/VirtualShadowMaps/VirtualShadowMapArray.h
--- a/Engine/Source/Runtime/Renderer/Private/VirtualShadowMaps/VirtualShadowMapArray.h
+++ b/Engine/Source/Runtime/Renderer/Private/VirtualShadowMaps/VirtualShadowMapArray.h
@@ -4,7 +4,6 @@
 =============================================================================*/
 #pragma once
 
-#include "VirtualShadowMapConfig.h"
 #include "../Nanite/NaniteRender.h"
 #include "../MeshDrawCommands.h"
 #include "SceneTypes.h"
diff --git a/Engine/Source/Runtime/Renderer/Private/VirtualShadowMaps/VirtualShadowMapCacheManager.h b/Engine/Source/Runtime/Renderer/Private/VirtualShadowMaps/VirtualShadowMapCacheManager.h
--- a/Engine/Source/Runtime/Renderer/Private/VirtualShadowMaps/VirtualShadowMapCacheManager.h
+++ b/Engine/Source/Runtime/Renderer/Private/VirtualShadowMaps/VirtualShadowMapCacheManager.h
@@ -5,7 +5,6 @@
 =============================================================================*/
 #pragma once
 
-#include "VirtualShadowMapConfig.h"
 #include "CoreMinimal.h"
 #include "VirtualShadowMapArray.h"
 #include "SceneManagement.h"
diff --git a/Engine/Source/Runtime/Renderer/Private/VirtualShadowMaps/VirtualShadowMapClipmap.h b/Engine/Source/Runtime/Renderer/Private/VirtualShadowMaps/VirtualShadowMapClipmap.h
--- a/Engine/Source/Runtime/Renderer/Private/VirtualShadowMaps/VirtualShadowMapClipmap.h
+++ b/Engine/Source/Runtime/Renderer/Private/VirtualShadowMaps/VirtualShadowMapClipmap.h
@@ -6,7 +6,6 @@ VirtualShadowMapClipmap.h
 
 #pragma once
 
-#include "VirtualShadowMapConfig.h"
 #include "CoreMinimal.h"
 #include "VirtualShadowMapProjection.h"
 
diff --git a/Engine/Source/Runtime/Renderer/Private/VirtualShadowMaps/VirtualShadowMapConfig.h b/Engine/Source/Runtime/Renderer/Private/VirtualShadowMaps/VirtualShadowMapConfig.h
deleted file mode 100644
diff --git a/Engine/Source/Runtime/Renderer/Private/VirtualShadowMaps/VirtualShadowMapProjection.h b/Engine/Source/Runtime/Renderer/Private/VirtualShadowMaps/VirtualShadowMapProjection.h
--- a/Engine/Source/Runtime/Renderer/Private/VirtualShadowMaps/VirtualShadowMapProjection.h
+++ b/Engine/Source/Runtime/Renderer/Private/VirtualShadowMaps/VirtualShadowMapProjection.h
@@ -6,7 +6,6 @@
 
 #pragma once
 
-#include "VirtualShadowMapConfig.h"
 #include "CoreMinimal.h"
 #include "RenderGraphResources.h"
 #include "ScreenSpaceDenoise.h"
