Integrate virutal texture allocator improvements from main.

diff --git a/Engine/Source/Runtime/Renderer/Private/VT/VirtualTextureAllocator.cpp b/Engine/Source/Runtime/Renderer/Private/VT/VirtualTextureAllocator.cpp
--- a/Engine/Source/Runtime/Renderer/Private/VT/VirtualTextureAllocator.cpp
+++ b/Engine/Source/Runtime/Renderer/Private/VT/VirtualTextureAllocator.cpp
@@ -8,10 +8,12 @@
 #include "Misc/Paths.h"
 
 #if WITH_EDITOR
+#include "HAL/FileManager.h"
 #include "IImageWrapper.h"
 #include "IImageWrapperModule.h"
-#include "HAL/FileManager.h"
 #include "Misc/FileHelper.h"
+#include "Misc/Paths.h"
+#include "Modules/ModuleManager.h"
 #endif // WITH_EDITOR
 
 FVirtualTextureAllocator::FVirtualTextureAllocator(uint32 Dimensions)
diff --git a/Engine/Source/Runtime/Renderer/Private/VT/VirtualTextureSpace.cpp b/Engine/Source/Runtime/Renderer/Private/VT/VirtualTextureSpace.cpp
--- a/Engine/Source/Runtime/Renderer/Private/VT/VirtualTextureSpace.cpp
+++ b/Engine/Source/Runtime/Renderer/Private/VT/VirtualTextureSpace.cpp
@@ -76,11 +76,6 @@ FVirtualTextureSpace::FVirtualTextureSpace(FVirtualTextureSystem* InSystem, uint
 		++PageTableIndex;
 	}
 
-	//PageTableSize = FMath::Max(InSizeNeeded, VIRTUALTEXTURE_MIN_PAGETABLE_SIZE);
-	//PageTableSize = FMath::RoundUpToPowerOfTwo(PageTableSize);
-	//ensure(PageTableSize <= Description.MaxSpaceSize);
-	//ensure(Description.MaxSpaceSize <= VIRTUALTEXTURE_MAX_PAGETABLE_SIZE);
-	//NumPageTableLevels = FMath::FloorLog2(PageTableSize) + 1;
 	Allocator.Initialize(Description.MaxSpaceSize);
 
 	bNeedToAllocatePageTableIndirection = InDesc.IndirectionTextureSize > 0;
@@ -138,7 +133,10 @@ uint32 FVirtualTextureSpace::GetSizeInBytes() const
 	uint32 TotalSize = 0u;
 	for (uint32 TextureIndex = 0u; TextureIndex < GetNumPageTableTextures(); ++TextureIndex)
 	{
-		const SIZE_T TextureSize = CalcTextureSize(CachedPageTableWidth, CachedPageTableHeight, TexturePixelFormat[TextureIndex], CachedNumPageTableLevels);
+		const uint32 PageTableWidth = Align(Allocator.GetAllocatedWidth(), VIRTUALTEXTURE_MIN_PAGETABLE_SIZE);
+		const uint32 PageTableHeight = Align(Allocator.GetAllocatedHeight(), VIRTUALTEXTURE_MIN_PAGETABLE_SIZE);
+		const uint32 NumPageTableLevels = FMath::FloorLog2(FMath::Max(PageTableWidth, PageTableHeight)) + 1u;
+		const SIZE_T TextureSize = CalcTextureSize(PageTableWidth, PageTableHeight, TexturePixelFormat[TextureIndex], NumPageTableLevels);
 		TotalSize += TextureSize;
 	}
 	
@@ -273,8 +271,6 @@ void FVirtualTextureSpace::AllocateTextures(FRDGBuilder& GraphBuilder)
 
 		CachedPageTableWidth = Align(Allocator.GetAllocatedWidth(), VIRTUALTEXTURE_MIN_PAGETABLE_SIZE);
 		CachedPageTableHeight = Align(Allocator.GetAllocatedHeight(), VIRTUALTEXTURE_MIN_PAGETABLE_SIZE);
-		//CachedPageTableWidth = FMath::Max(CachedPageTableWidth, VIRTUALTEXTURE_MIN_PAGETABLE_SIZE);
-		//CachedPageTableHeight = FMath::Max(CachedPageTableHeight, VIRTUALTEXTURE_MIN_PAGETABLE_SIZE);
 		CachedNumPageTableLevels = FMath::FloorLog2(FMath::Max(CachedPageTableWidth, CachedPageTableHeight)) + 1u;
 
 		for (uint32 TextureIndex = 0u; TextureIndex < GetNumPageTableTextures(); ++TextureIndex)
diff --git a/Engine/Source/Runtime/Renderer/Private/VT/VirtualTextureSpace.h b/Engine/Source/Runtime/Renderer/Private/VT/VirtualTextureSpace.h
--- a/Engine/Source/Runtime/Renderer/Private/VT/VirtualTextureSpace.h
+++ b/Engine/Source/Runtime/Renderer/Private/VT/VirtualTextureSpace.h
@@ -125,8 +125,6 @@ private:
 	FBufferRHIRef UpdateBuffer;
 	FShaderResourceViewRHIRef UpdateBufferSRV;
 
-	//uint32 PageTableSize;
-	//uint32 NumPageTableLevels;
 	uint32 NumRefs;
 
 	uint8 ID;
diff --git a/Engine/Source/Runtime/Renderer/Private/VT/VirtualTextureSystem.cpp b/Engine/Source/Runtime/Renderer/Private/VT/VirtualTextureSystem.cpp
--- a/Engine/Source/Runtime/Renderer/Private/VT/VirtualTextureSystem.cpp
+++ b/Engine/Source/Runtime/Renderer/Private/VT/VirtualTextureSystem.cpp
@@ -406,6 +406,10 @@ void FVirtualTextureSystem::ListPhysicalPoolsFromConsole()
 	UE_LOG(LogConsoleResponse, Display, TEXT("TotalPhysicalMemory: %fMB"), (double)TotalPhysicalMemory / 1024.0 / 1024.0);
 }
 
+{
+	return FCrc::MemCrc32(&Description, sizeof(Description));
+}
+
 #if WITH_EDITOR
 void FVirtualTextureSystem::SaveAllocatorImagesFromConsole()
 {
