Fix auto merge error

diff --git a/Engine/Source/Runtime/Renderer/Private/VT/VirtualTextureSystem.cpp b/Engine/Source/Runtime/Renderer/Private/VT/VirtualTextureSystem.cpp
--- a/Engine/Source/Runtime/Renderer/Private/VT/VirtualTextureSystem.cpp
+++ b/Engine/Source/Runtime/Renderer/Private/VT/VirtualTextureSystem.cpp
@@ -406,10 +406,6 @@ void FVirtualTextureSystem::ListPhysicalPoolsFromConsole()
 	UE_LOG(LogConsoleResponse, Display, TEXT("TotalPhysicalMemory: %fMB"), (double)TotalPhysicalMemory / 1024.0 / 1024.0);
 }
 
-{
-	return FCrc::MemCrc32(&Description, sizeof(Description));
-}
-
 #if WITH_EDITOR
 void FVirtualTextureSystem::SaveAllocatorImagesFromConsole()
 {
