Fixes for UDIM VTs in texture editor
UDIMs use transformed UVs, account for these when computing visible regions

diff --git a/Engine/Source/Editor/TextureEditor/Private/Models/TextureEditorViewportClient.cpp b/Engine/Source/Editor/TextureEditor/Private/Models/TextureEditorViewportClient.cpp
--- a/Engine/Source/Editor/TextureEditor/Private/Models/TextureEditorViewportClient.cpp
+++ b/Engine/Source/Editor/TextureEditor/Private/Models/TextureEditorViewportClient.cpp
@@ -66,7 +66,7 @@ void FTextureEditorViewportClient::Draw(FViewport* Viewport, FCanvas* Canvas)
 	
 	UTexture* Texture = TextureEditorPinned->GetTexture();
 	FVector2D Ratio = FVector2D(GetViewportHorizontalScrollBarRatio(), GetViewportVerticalScrollBarRatio());
-	FVector2D ViewportSize = FVector2D(TextureEditorViewportPtr.Pin()->GetViewport()->GetSizeXY().X, TextureEditorViewportPtr.Pin()->GetViewport()->GetSizeXY().Y);
+	FVector2D ViewportSize = FVector2D(TextureEditorViewportPtr.Pin()->GetViewport()->GetSizeXY());
 	FVector2D ScrollBarPos = GetViewportScrollBarPositions();
 	int32 YOffset = (Ratio.Y > 1.0f)? ((ViewportSize.Y - (ViewportSize.Y / Ratio.Y)) * 0.5f): 0;
 	int32 YPos = YOffset - ScrollBarPos.Y;
@@ -205,28 +205,21 @@ void FTextureEditorViewportClient::Draw(FViewport* Viewport, FCanvas* Canvas)
 		if (bIsVirtualTexture && CVarEnableVTFeedback.GetValueOnAnyThread() != 0)
 		{
 			FVirtualTexture2DResource* VTResource = static_cast<FVirtualTexture2DResource*>(Texture->Resource);
-			const FVector2D ScreenSpaceSize = { (float)Width, (float)Height };
-			
-			// Calculate the rect of the texture that is visible on screen
-			// When the entire texture is visible, we pass -1 as dimensions (VirtualTextureUtils::CalculateVisibleTiles needs to calculate the effective texture size anyway so don't do it twice)
- 			auto ZeroOrAbsolute = [](int32 value) -> int32 { return value >= 0 ? 0 : FMath::Abs(value); };
-			auto GetCorrectDimension = [](int value, int32 viewportSize, int32 TextureSize) -> int32 { return value <= viewportSize ? TextureSize : viewportSize; };
-
-			const float Zoom = 1.0f / TextureEditorPinned->GetCustomZoomLevel();
-			const int32 VisibleXPos = FMath::FloorToInt(Zoom * -FMath::Min(0, XPos));
-			const int32 VisibleYPos = FMath::FloorToInt(Zoom * -FMath::Min(0, YPos));
-			
-			const FIntRect VisibleTextureRect = FIntRect(VisibleXPos, VisibleYPos,
-				VisibleXPos + GetCorrectDimension(Zoom * Width, Zoom * ViewportSize.X, VTResource->GetSizeX()),
-				VisibleYPos + GetCorrectDimension(Zoom * Height, Zoom * ViewportSize.Y, VTResource->GetSizeY()));
+			const FVector2D ScreenSpaceSize((float)Width, (float)Height);
+			const FVector2D ViewportPositon(-(float)XPos, -(float)YPos);
+			const FVector2D UV0 = TileItem.UV0;
+			const FVector2D UV1 = TileItem.UV1;
 
 			const ERHIFeatureLevel::Type InFeatureLevel = GMaxRHIFeatureLevel;
 			ENQUEUE_RENDER_COMMAND(MakeTilesResident)(
-				[InFeatureLevel, VTResource, ScreenSpaceSize, VisibleTextureRect, MipLevel](FRHICommandListImmediate& RHICmdList)
+				[InFeatureLevel, VTResource, ScreenSpaceSize, ViewportPositon, ViewportSize, UV0, UV1, MipLevel](FRHICommandListImmediate& RHICmdList)
 			{
 				// AcquireAllocatedVT() must happen on render thread
-				GetRendererModule().RequestVirtualTextureTilesForRegion(VTResource->AcquireAllocatedVT(), ScreenSpaceSize, VisibleTextureRect, MipLevel);
-				GetRendererModule().LoadPendingVirtualTextureTiles(RHICmdList, InFeatureLevel);
+				IAllocatedVirtualTexture* AllocatedVT = VTResource->AcquireAllocatedVT();
+
+				IRendererModule& RenderModule = GetRendererModule();
+				RenderModule.RequestVirtualTextureTilesForRegion(AllocatedVT, ScreenSpaceSize, ViewportPositon, ViewportSize, UV0, UV1, MipLevel);
+				RenderModule.LoadPendingVirtualTextureTiles(RHICmdList, InFeatureLevel);
 			});
 		}
 	}
diff --git a/Engine/Source/Editor/TextureEditor/Private/TextureEditorToolkit.cpp b/Engine/Source/Editor/TextureEditor/Private/TextureEditorToolkit.cpp
--- a/Engine/Source/Editor/TextureEditor/Private/TextureEditorToolkit.cpp
+++ b/Engine/Source/Editor/TextureEditor/Private/TextureEditorToolkit.cpp
@@ -267,8 +267,9 @@ void FTextureEditorToolkit::InitTextureEditor( const EToolkitMode::Type Mode, co
 
 void FTextureEditorToolkit::CalculateTextureDimensions( uint32& Width, uint32& Height, uint32& Depth, uint32& ArraySize ) const
 {
-	Width = Texture->Source.GetSizeX();
-	Height = Texture->Source.GetSizeY();
+	const FIntPoint LogicalSize = Texture->Source.GetLogicalSize();
+	Width = LogicalSize.X;
+	Height = LogicalSize.Y;
 	Depth = IsVolumeTexture() ? Texture->Source.GetNumLayers() : 0;
 	ArraySize = (Is2DArrayTexture() || IsCubeTexture()) ? Texture->Source.GetNumLayers() : 0;
 
diff --git a/Engine/Source/Runtime/RenderCore/Public/RendererInterface.h b/Engine/Source/Runtime/RenderCore/Public/RendererInterface.h
--- a/Engine/Source/Runtime/RenderCore/Public/RendererInterface.h
+++ b/Engine/Source/Runtime/RenderCore/Public/RendererInterface.h
@@ -817,7 +817,7 @@ public:
 
 	/**	Provided a list of packed virtual texture tile ids, let the VT system request them. Note this should be called as long as the tiles are needed.*/
 	virtual void RequestVirtualTextureTiles(const FVector2D& InScreenSpaceSize, int32 InMipLevel) = 0;
-	virtual void RequestVirtualTextureTilesForRegion(IAllocatedVirtualTexture* AllocatedVT, const FVector2D& InScreenSpaceSize, const FIntRect& InTextureRegion, int32 InMipLevel) = 0;
+	virtual void RequestVirtualTextureTilesForRegion(IAllocatedVirtualTexture* AllocatedVT, const FVector2D& InScreenSpaceSize, const FVector2D& InViewportPosition, const FVector2D& InViewportSize, const FVector2D& InUV0, const FVector2D& InUV1, int32 InMipLevel) = 0;
 
 	/** Ensure that any tiles requested by 'RequestVirtualTextureTilesForRegion' are loaded, must be called from render thread */
 	virtual void LoadPendingVirtualTextureTiles(FRHICommandListImmediate& RHICmdList, ERHIFeatureLevel::Type FeatureLevel) = 0;
diff --git a/Engine/Source/Runtime/RenderCore/Public/VirtualTexturing.h b/Engine/Source/Runtime/RenderCore/Public/VirtualTexturing.h
--- a/Engine/Source/Runtime/RenderCore/Public/VirtualTexturing.h
+++ b/Engine/Source/Runtime/RenderCore/Public/VirtualTexturing.h
@@ -365,6 +365,8 @@ public:
 	inline uint8 GetDimensions() const { return Description.Dimensions; }
 	inline uint32 GetWidthInBlocks() const { return WidthInBlocks; }
 	inline uint32 GetHeightInBlocks() const { return HeightInBlocks; }
+	inline uint32 GetBlockWidthInTiles() const { return BlockWidthInTiles; }
+	inline uint32 GetBlockHeightInTiles() const { return BlockHeightInTiles; }
 	inline uint32 GetWidthInTiles() const { return BlockWidthInTiles * WidthInBlocks; }
 	inline uint32 GetHeightInTiles() const { return BlockHeightInTiles * HeightInBlocks; }
 	inline uint32 GetDepthInTiles() const { return DepthInTiles; }
diff --git a/Engine/Source/Runtime/Renderer/Private/RendererModule.h b/Engine/Source/Runtime/Renderer/Private/RendererModule.h
--- a/Engine/Source/Runtime/Renderer/Private/RendererModule.h
+++ b/Engine/Source/Runtime/Renderer/Private/RendererModule.h
@@ -100,7 +100,7 @@ public:
 	virtual uint32 RemoveAllVirtualTextureProducerDestroyedCallbacks(const void* Baton) override;
 	virtual void ReleaseVirtualTexturePendingResources() override;
 	virtual void RequestVirtualTextureTiles(const FVector2D& InScreenSpaceSize, int32 InMipLevel) override;
-	virtual void RequestVirtualTextureTilesForRegion(IAllocatedVirtualTexture* AllocatedVT, const FVector2D& InScreenSpaceSize, const FIntRect& InTextureRegion, int32 InMipLevel) override;
+	virtual void RequestVirtualTextureTilesForRegion(IAllocatedVirtualTexture* AllocatedVT, const FVector2D& InScreenSpaceSize, const FVector2D& InViewportPosition, const FVector2D& InViewportSize, const FVector2D& InUV0, const FVector2D& InUV1, int32 InMipLevel) override;
 	virtual void LoadPendingVirtualTextureTiles(FRHICommandListImmediate& RHICmdList, ERHIFeatureLevel::Type FeatureLevel) override;
 	virtual void FlushVirtualTextureCache() override;
 	
diff --git a/Engine/Source/Runtime/Renderer/Private/SceneRendering.cpp b/Engine/Source/Runtime/Renderer/Private/SceneRendering.cpp
--- a/Engine/Source/Runtime/Renderer/Private/SceneRendering.cpp
+++ b/Engine/Source/Runtime/Renderer/Private/SceneRendering.cpp
@@ -4186,9 +4186,9 @@ void FRendererModule::RequestVirtualTextureTiles(const FVector2D& InScreenSpaceS
 	FVirtualTextureSystem::Get().RequestTiles(InScreenSpaceSize, InMipLevel);
 }
 
-void FRendererModule::RequestVirtualTextureTilesForRegion(IAllocatedVirtualTexture* AllocatedVT, const FVector2D& InScreenSpaceSize, const FIntRect& InTextureRegion, int32 InMipLevel)
+void FRendererModule::RequestVirtualTextureTilesForRegion(IAllocatedVirtualTexture* AllocatedVT, const FVector2D& InScreenSpaceSize, const FVector2D& InViewportPosition, const FVector2D& InViewportSize, const FVector2D& InUV0, const FVector2D& InUV1, int32 InMipLevel)
 {
-	FVirtualTextureSystem::Get().RequestTilesForRegion(AllocatedVT, InScreenSpaceSize, InTextureRegion, InMipLevel);
+	FVirtualTextureSystem::Get().RequestTilesForRegion(AllocatedVT, InScreenSpaceSize, InViewportPosition, InViewportSize, InUV0, InUV1, InMipLevel);
 }
 
 void FRendererModule::LoadPendingVirtualTextureTiles(FRHICommandListImmediate& RHICmdList, ERHIFeatureLevel::Type FeatureLevel)
diff --git a/Engine/Source/Runtime/Renderer/Private/VT/AllocatedVirtualTexture.cpp b/Engine/Source/Runtime/Renderer/Private/VT/AllocatedVirtualTexture.cpp
--- a/Engine/Source/Runtime/Renderer/Private/VT/AllocatedVirtualTexture.cpp
+++ b/Engine/Source/Runtime/Renderer/Private/VT/AllocatedVirtualTexture.cpp
@@ -21,8 +21,6 @@ FAllocatedVirtualTexture::FAllocatedVirtualTexture(FVirtualTextureSystem* InSyst
 	, RefCount(1)
 	, FrameAllocated(InFrame)
 	, Space(nullptr)
-	, VirtualPageX(~0u)
-	, VirtualPageY(~0u)
 {
 	check(IsInRenderingThread());
 	FMemory::Memzero(TextureLayers);
@@ -363,8 +361,8 @@ static inline uint32 BitcastFloatToUInt32(float v)
 
 void FAllocatedVirtualTexture::GetPackedPageTableUniform(FUintVector4* OutUniform) const
 {
-	const uint32 vPageX = FMath::ReverseMortonCode2(VirtualAddress);
-	const uint32 vPageY = FMath::ReverseMortonCode2(VirtualAddress >> 1);
+	const uint32 vPageX = VirtualPageX;
+	const uint32 vPageY = VirtualPageY;
 	const uint32 vPageSize = GetVirtualTileSize();
 	const uint32 PageBorderSize = GetTileBorderSize();
 	const uint32 WidthInPages = GetWidthInTiles();
diff --git a/Engine/Source/Runtime/Renderer/Private/VT/AllocatedVirtualTexture.h b/Engine/Source/Runtime/Renderer/Private/VT/AllocatedVirtualTexture.h
--- a/Engine/Source/Runtime/Renderer/Private/VT/AllocatedVirtualTexture.h
+++ b/Engine/Source/Runtime/Renderer/Private/VT/AllocatedVirtualTexture.h
@@ -60,9 +60,6 @@ public:
 	inline uint32 GetProducerTextureLayerMaskForPageTableLayer(uint32 InLayerIndex) const { return UniquePageTableLayers[InLayerIndex].ProducerTextureLayerMask; }
 	inline uint32 GetProducerPhysicalGroupIndexForPageTableLayer(uint32 InLayerIndex) const { return UniquePageTableLayers[InLayerIndex].ProducerPhysicalGroupIndex; }
 
-	inline uint32 GetVirtualPageX() const { return VirtualPageX; }
-	inline uint32 GetVirtualPageY() const { return VirtualPageY; }
-
 private:
 	uint32 AddUniqueProducer(FVirtualTextureProducerHandle const& InHandle, const FVirtualTextureProducer* InProducer);
 	uint32 AddUniquePhysicalSpace(FVirtualTexturePhysicalSpace* InPhysicalSpace, uint32 InUniqueProducerIndex, uint32 InProducerPhysicalSpaceIndex);
@@ -101,7 +98,4 @@ private:
 		uint8 PhysicalTextureIndex;
 	};
 	FTextureLayerDesc TextureLayers[VIRTUALTEXTURE_SPACE_MAXLAYERS];
-
-	uint32 VirtualPageX;
-	uint32 VirtualPageY;
 };
diff --git a/Engine/Source/Runtime/Renderer/Private/VT/VirtualTextureSystem.cpp b/Engine/Source/Runtime/Renderer/Private/VT/VirtualTextureSystem.cpp
--- a/Engine/Source/Runtime/Renderer/Private/VT/VirtualTextureSystem.cpp
+++ b/Engine/Source/Runtime/Renderer/Private/VT/VirtualTextureSystem.cpp
@@ -849,29 +849,36 @@ void FVirtualTextureSystem::RequestTiles(const FVector2D& InScreenSpaceSize, int
 {
 	check(IsInRenderingThread());
 
+	FScopeLock Lock(&RequestedTilesLock);
+
 	for (const auto& Pair : AllocatedVTs)
 	{
-		RequestTilesForRegion(Pair.Value, InScreenSpaceSize, FIntRect(), InMipLevel);
+		const IAllocatedVirtualTexture* AllocatedVT = Pair.Value;
+		if (InMipLevel < 0)
+		{
+			const uint32 vMaxLevel = AllocatedVT->GetMaxLevel();
+			const float vLevel = ComputeMipLevel(AllocatedVT, InScreenSpaceSize);
+			const int32 vMipLevelDown = FMath::Clamp((int32)FMath::FloorToInt(vLevel), 0, (int32)vMaxLevel);
+			RequestTilesInternal(AllocatedVT, vLevel);
+			if (vMipLevelDown + 1u <= vMaxLevel)
+			{
+				// Need to fetch 2 levels to support trilinear filtering
+				RequestTilesInternal(AllocatedVT, vMipLevelDown);
+			}
+		}
+		else
+		{
+			RequestTilesInternal(AllocatedVT, InMipLevel);
+		}
 	}
 }
 
-void FVirtualTextureSystem::RequestTilesForRegion(const IAllocatedVirtualTexture* AllocatedVT, const FVector2D& InScreenSpaceSize, const FIntRect& InTextureRegion, int32 InMipLevel)
+void FVirtualTextureSystem::RequestTilesForRegion(IAllocatedVirtualTexture* AllocatedVT, const FVector2D& InScreenSpaceSize, const FVector2D& InViewportPosition, const FVector2D& InViewportSize, const FVector2D& InUV0, const FVector2D& InUV1, int32 InMipLevel)
 {
-	FIntRect TextureRegion(InTextureRegion);
-	if (TextureRegion.IsEmpty())
-	{
-		TextureRegion.Max.X = AllocatedVT->GetWidthInPixels();
-		TextureRegion.Max.Y = AllocatedVT->GetHeightInPixels();
-	}
-	else
-	{
-		TextureRegion.Clip(FIntRect(0, 0, AllocatedVT->GetWidthInPixels(), AllocatedVT->GetHeightInPixels()));
-	}
-
 	if (InMipLevel >= 0)
 	{
 		FScopeLock Lock(&RequestedTilesLock);
-		RequestTilesForRegionInternal(AllocatedVT, TextureRegion, InMipLevel);
+		RequestTilesForRegionInternal(AllocatedVT, InScreenSpaceSize, InViewportPosition, InViewportSize, InUV0, InUV1, InMipLevel);
 	}
 	else
 	{
@@ -880,11 +887,11 @@ void FVirtualTextureSystem::RequestTilesForRegion(const IAllocatedVirtualTexture
 		const int32 vMipLevelDown = FMath::Clamp((int32)FMath::FloorToInt(vLevel), 0, (int32)vMaxLevel);
 
 		FScopeLock Lock(&RequestedTilesLock);
-		RequestTilesForRegionInternal(AllocatedVT, TextureRegion, vMipLevelDown);
+		RequestTilesForRegionInternal(AllocatedVT, InScreenSpaceSize, InViewportPosition, InViewportSize, InUV0, InUV1, vMipLevelDown);
 		if (vMipLevelDown + 1u <= vMaxLevel)
 		{
 			// Need to fetch 2 levels to support trilinear filtering
-			RequestTilesForRegionInternal(AllocatedVT, TextureRegion, vMipLevelDown + 1u);
+			RequestTilesForRegionInternal(AllocatedVT, InScreenSpaceSize, InViewportPosition, InViewportSize, InUV0, InUV1, vMipLevelDown + 1u);
 		}
 	}
 }
@@ -920,22 +927,67 @@ void FVirtualTextureSystem::LoadPendingTiles(FRDGBuilder& GraphBuilder, ERHIFeat
 	}
 }
 
-void FVirtualTextureSystem::RequestTilesForRegionInternal(const IAllocatedVirtualTexture* AllocatedVT, const FIntRect& InTextureRegion, uint32 vLevel)
+void FVirtualTextureSystem::RequestTilesInternal(const IAllocatedVirtualTexture* AllocatedVT, int32 InMipLevel)
+{
+	const int32 MipWidthInTiles = FMath::Max<int32>(AllocatedVT->GetWidthInTiles() >> InMipLevel, 1);
+	const int32 MipHeightInTiles = FMath::Max<int32>(AllocatedVT->GetHeightInTiles() >> InMipLevel, 1);
+	const uint32 vBaseTileX = AllocatedVT->GetVirtualPageX() >> InMipLevel;
+	const uint32 vBaseTileY = AllocatedVT->GetVirtualPageY() >> InMipLevel;
+
+	for (int32 TilePositionY = 0; TilePositionY < MipHeightInTiles; TilePositionY++)
+	{
+		const uint32 vGlobalTileY = vBaseTileY + TilePositionY;
+		for (int32 TilePositionX = 0; TilePositionX < MipWidthInTiles; TilePositionX++)
+		{
+			const uint32 vGlobalTileX = vBaseTileX + TilePositionX;
+			const uint32 EncodedTile = EncodePage(AllocatedVT->GetSpaceID(), InMipLevel, vGlobalTileX, vGlobalTileY);
+			RequestedPackedTiles.Add(EncodedTile);
+		}
+	}
+}
+
+static int32 WrapTilePosition(int32 Position, int32 Size)
+{
+	const int32 Result = Position % Size;
+	return (Result >= 0) ? Result : Result + Size;
+}
+
+void FVirtualTextureSystem::RequestTilesForRegionInternal(const IAllocatedVirtualTexture* AllocatedVT, const FVector2D& InScreenSpaceSize, const FVector2D& InViewportPosition, const FVector2D& InViewportSize, const FVector2D& InUV0, const FVector2D& InUV1, int32 InMipLevel)
 {
-	const FIntRect TextureRegionForLevel(InTextureRegion.Min.X >> vLevel, InTextureRegion.Min.Y >> vLevel, InTextureRegion.Max.X >> vLevel, InTextureRegion.Max.Y >> vLevel);
-	const FIntRect TileRegionForLevel = FIntRect::DivideAndRoundUp(TextureRegionForLevel, AllocatedVT->GetVirtualTileSize());
+	// Determine the screen-space coordinates of the texture we're trying to display
+	const float SrcPositionX0 = FMath::Max(InViewportPosition.X, 0.0f);
+	const float SrcPositionY0 = FMath::Max(InViewportPosition.Y, 0.0f);
+	const float SrcPositionX1 = FMath::Min(InViewportPosition.X + InViewportSize.X, InScreenSpaceSize.X);
+	const float SrcPositionY1 = FMath::Min(InViewportPosition.Y + InViewportSize.Y, InScreenSpaceSize.Y);
+
+	const int32 WidthInBlocks = AllocatedVT->GetWidthInBlocks();
+	const int32 HeightInBlocks = AllocatedVT->GetHeightInBlocks();
+
+	// Map coordinates to UV space
+	const float PositionU0 = FMath::Lerp(InUV0.X, InUV1.X, SrcPositionX0 / InScreenSpaceSize.X) / WidthInBlocks;
+	const float PositionV0 = FMath::Lerp(InUV0.Y, InUV1.Y, SrcPositionY0 / InScreenSpaceSize.Y) / HeightInBlocks;
+	const float PositionU1 = FMath::Lerp(InUV0.X, InUV1.X, SrcPositionX1 / InScreenSpaceSize.X) / WidthInBlocks;
+	const float PositionV1 = FMath::Lerp(InUV0.Y, InUV1.Y, SrcPositionY1 / InScreenSpaceSize.Y) / HeightInBlocks;
+
+	// Map UVs to tile coordinates
+	const int32 MipWidthInTiles = FMath::Max<int32>(AllocatedVT->GetWidthInTiles() >> InMipLevel, 1);
+	const int32 MipHeightInTiles = FMath::Max<int32>(AllocatedVT->GetHeightInTiles() >> InMipLevel, 1);
+	const int32 TilePositionX0 = FMath::FloorToInt(FMath::Min(PositionU0, PositionU1) * MipWidthInTiles);
+	const int32 TilePositionY0 = FMath::FloorToInt(FMath::Min(PositionV0, PositionV1) * MipHeightInTiles);
+	const int32 TilePositionX1 = FMath::CeilToInt(FMath::Max(PositionU0, PositionU1) * MipWidthInTiles);
+	const int32 TilePositionY1 = FMath::CeilToInt(FMath::Max(PositionV0, PositionV1) * MipHeightInTiles);
 
 	// RequestedPackedTiles stores packed tiles with vPosition shifted relative to current mip level
-	const uint32 vBaseTileX = FMath::ReverseMortonCode2(AllocatedVT->GetVirtualAddress()) >> vLevel;
-	const uint32 vBaseTileY = FMath::ReverseMortonCode2(AllocatedVT->GetVirtualAddress() >> 1) >> vLevel;
+	const uint32 vBaseTileX = AllocatedVT->GetVirtualPageX() >> InMipLevel;
+	const uint32 vBaseTileY = AllocatedVT->GetVirtualPageY() >> InMipLevel;
 
-	for (uint32 TileY = TileRegionForLevel.Min.Y; TileY < (uint32)TileRegionForLevel.Max.Y; ++TileY)
+	for (int32 TilePositionY = TilePositionY0; TilePositionY < TilePositionY1; TilePositionY++)
 	{
-		const uint32 vGlobalTileY = vBaseTileY + TileY;
-		for (uint32 TileX = TileRegionForLevel.Min.X; TileX < (uint32)TileRegionForLevel.Max.X; ++TileX)
+		const uint32 vGlobalTileY = vBaseTileY + WrapTilePosition(TilePositionY, MipHeightInTiles);
+		for (int32 TilePositionX = TilePositionX0; TilePositionX < TilePositionX1; TilePositionX++)
 		{
-			const uint32 vGlobalTileX = vBaseTileX + TileX;
-			const uint32 EncodedTile = EncodePage(AllocatedVT->GetSpaceID(), vLevel, vGlobalTileX, vGlobalTileY);
+			const uint32 vGlobalTileX = vBaseTileX + WrapTilePosition(TilePositionX, MipWidthInTiles);
+			const uint32 EncodedTile = EncodePage(AllocatedVT->GetSpaceID(), InMipLevel, vGlobalTileX, vGlobalTileY);
 			RequestedPackedTiles.Add(EncodedTile);
 		}
 	}
diff --git a/Engine/Source/Runtime/Renderer/Private/VT/VirtualTextureSystem.h b/Engine/Source/Runtime/Renderer/Private/VT/VirtualTextureSystem.h
--- a/Engine/Source/Runtime/Renderer/Private/VT/VirtualTextureSystem.h
+++ b/Engine/Source/Runtime/Renderer/Private/VT/VirtualTextureSystem.h
@@ -66,7 +66,7 @@ public:
 	void UnlockTile(const FVirtualTextureLocalTile& Tile, const FVirtualTextureProducer* Producer);
 	void ForceUnlockAllTiles(const FVirtualTextureProducerHandle& ProducerHandle, const FVirtualTextureProducer* Producer);
 	void RequestTiles(const FVector2D& InScreenSpaceSize, int32 InMipLevel = -1);
-	void RequestTilesForRegion(const IAllocatedVirtualTexture* AllocatedVT, const FVector2D& InScreenSpaceSize, const FIntRect& InTextureRegion, int32 InMipLevel = -1);
+	void RequestTilesForRegion(IAllocatedVirtualTexture* AllocatedVT, const FVector2D& InScreenSpaceSize, const FVector2D& InViewportPosition, const FVector2D& InViewportSize, const FVector2D& InUV0, const FVector2D& InUV1, int32 InMipLevel = -1);
 	void LoadPendingTiles(FRDGBuilder& GraphBuilder, ERHIFeatureLevel::Type FeatureLevel);
 	
 	void FlushCache();
@@ -82,7 +82,8 @@ private:
 	void DestroyPendingVirtualTextures();
 	void ReleasePendingSpaces();
 
-	void RequestTilesForRegionInternal(const IAllocatedVirtualTexture* AllocatedVT, const FIntRect& InTextureRegion, uint32 vLevel);
+	void RequestTilesForRegionInternal(const IAllocatedVirtualTexture* AllocatedVT, const FVector2D& InScreenSpaceSize, const FVector2D& InViewportPosition, const FVector2D& InViewportSize, const FVector2D& InUV0, const FVector2D& InUV1, int32 InMipLevel);
+	void RequestTilesInternal(const IAllocatedVirtualTexture* AllocatedVT, int32 InMipLevel);
 	
 	void SubmitRequestsFromLocalTileList(TArray<FVirtualTextureLocalTile>& OutDeferredTiles, const TSet<FVirtualTextureLocalTile>& LocalTileList, EVTProducePageFlags Flags, FRDGBuilder& GraphBuilder, ERHIFeatureLevel::Type FeatureLevel);
 
