Allow to render passes as tiled, according to strata tile material type.
Pixels without any materials (sky) are now skipped during tiled passes.
Next: tile the full screen directional lighting pass.
Quite a bit of data and code cleanup.

I had to use the VS/PS pattern because the clean function when using both PS and VS only take shaders with the same FParameters (the vertex shader must have using FParameters = FScreenSpaceReflectionsPS::FParameters; // Sharing parameters for proper registration with RDG for instance)

Tested with strata disabled.

diff --git a/Engine/Shaders/Private/Strata/StrataMaterialClassification.usf b/Engine/Shaders/Private/Strata/StrataMaterialClassification.usf
--- a/Engine/Shaders/Private/Strata/StrataMaterialClassification.usf
+++ b/Engine/Shaders/Private/Strata/StrataMaterialClassification.usf
@@ -198,64 +198,83 @@ int TileSize;
 int bRectPrimitive;
 int2 ViewResolution;
 Texture2D<uint> ClassificationTexture;
-RWBuffer<uint> TileIndirectData;
-RWBuffer<uint> TileListData;
+RWBuffer<uint> SimpleTileIndirectDataBuffer;
+RWBuffer<uint> SimpleTileListDataBuffer;
+RWBuffer<uint> ComplexTileIndirectDataBuffer;
+RWBuffer<uint> ComplexTileListDataBuffer;
 
-groupshared bool s_IsSimpleTile[GROUP_THREAD_COUNT];
+groupshared uint s_IsComplexIsMaterialTile[GROUP_THREAD_COUNT];
 
 [numthreads(TILE_SIZE, TILE_SIZE, 1)]
 void TileMainCS(uint2 DispatchThreadId : SV_DispatchThreadID, uint LinearIndex : SV_GroupIndex, uint3 GroupId : SV_GroupID)
 {
 	if (all(DispatchThreadId == 0))
 	{
-		TileIndirectData[0] = bRectPrimitive > 0 ? 4 : 6;
+		SimpleTileIndirectDataBuffer[0] = bRectPrimitive > 0 ? 4 : 6;
+		ComplexTileIndirectDataBuffer[0] = bRectPrimitive > 0 ? 4 : 6;
 	}
 
 	// TODO: add a SM6 permutation with ballot?
 	const uint2 PixelCoord = DispatchThreadId;
 
 	// Pixels outside of the view area are considered simple to enable screen borders to receive the simple permutation when not aligned to shader group size.
-	bool bIsSimpleTile = true;
+	bool bContainsComplexMaterial = false;
+	bool bContainsStrataMaterial = false;
 	if (all(PixelCoord < uint2(ViewResolution)))
 	{
 		const uint PackClassification = ClassificationTexture.Load(uint3(PixelCoord, 0));
 		FStrataClassification Data = StrataUnpackClassificationData(PackClassification);
-		bIsSimpleTile = Data.bIsSimple;
+		bContainsComplexMaterial = !Data.bIsSimple;
+		bContainsStrataMaterial = Data.bIsStrataMaterial;
 	}
-	s_IsSimpleTile[LinearIndex] = bIsSimpleTile;
+	s_IsComplexIsMaterialTile[LinearIndex] = (bContainsComplexMaterial ? 1 : 0) | (bContainsStrataMaterial ? 2 : 0);
 
 	GroupMemoryBarrierWithGroupSync();
 	if (LinearIndex < 32)
 	{
-		s_IsSimpleTile[LinearIndex] = s_IsSimpleTile[LinearIndex] && s_IsSimpleTile[LinearIndex + 32];
+		s_IsComplexIsMaterialTile[LinearIndex] = s_IsComplexIsMaterialTile[LinearIndex] | s_IsComplexIsMaterialTile[LinearIndex + 32];
 	}
 	GroupMemoryBarrierWithGroupSync();
 	if (LinearIndex < 16)
 	{
-		s_IsSimpleTile[LinearIndex] = s_IsSimpleTile[LinearIndex] && s_IsSimpleTile[LinearIndex + 16];
+		s_IsComplexIsMaterialTile[LinearIndex] = s_IsComplexIsMaterialTile[LinearIndex] | s_IsComplexIsMaterialTile[LinearIndex + 16];
 	}
 	GroupMemoryBarrierWithGroupSync();
 
 	if (LinearIndex < 8)
 	{
-		s_IsSimpleTile[LinearIndex] = s_IsSimpleTile[LinearIndex] && s_IsSimpleTile[LinearIndex + 8];
+		s_IsComplexIsMaterialTile[LinearIndex] = s_IsComplexIsMaterialTile[LinearIndex] | s_IsComplexIsMaterialTile[LinearIndex + 8];
 	}
 	GroupMemoryBarrierWithGroupSync();
 	if (LinearIndex < 4)
 	{
-		s_IsSimpleTile[LinearIndex] = s_IsSimpleTile[LinearIndex] && s_IsSimpleTile[LinearIndex + 4];
+		s_IsComplexIsMaterialTile[LinearIndex] = s_IsComplexIsMaterialTile[LinearIndex] | s_IsComplexIsMaterialTile[LinearIndex + 4];
 	}
 	GroupMemoryBarrierWithGroupSync();
 	if (LinearIndex < 2)
 	{
-		s_IsSimpleTile[LinearIndex] = s_IsSimpleTile[LinearIndex] && s_IsSimpleTile[LinearIndex + 2];
+		s_IsComplexIsMaterialTile[LinearIndex] = s_IsComplexIsMaterialTile[LinearIndex] | s_IsComplexIsMaterialTile[LinearIndex + 2];
 	}
 	GroupMemoryBarrierWithGroupSync();
-	if (LinearIndex < 1 && (s_IsSimpleTile[LinearIndex] && s_IsSimpleTile[LinearIndex + 1]))
+
+	const uint FinalIsSimpleIsMaterialTile = s_IsComplexIsMaterialTile[LinearIndex] | s_IsComplexIsMaterialTile[LinearIndex + 1];
+	bContainsComplexMaterial = (FinalIsSimpleIsMaterialTile & 1) > 0;
+	bContainsStrataMaterial = (FinalIsSimpleIsMaterialTile & 2) > 0;
+
+	if (LinearIndex < 1 && bContainsStrataMaterial)
 	{
-		uint WriteToIndex;
-		InterlockedAdd(TileIndirectData[1], 1, WriteToIndex);
-		TileListData[WriteToIndex] = EncodeTile(GroupId.xy);
+		if (bContainsComplexMaterial)
+		{
+			uint WriteToIndex;
+			InterlockedAdd(ComplexTileIndirectDataBuffer[1], 1, WriteToIndex);
+			ComplexTileListDataBuffer[WriteToIndex] = EncodeTile(GroupId.xy);
+		}
+		else
+		{
+			uint WriteToIndex;
+			InterlockedAdd(SimpleTileIndirectDataBuffer[1], 1, WriteToIndex);
+			SimpleTileListDataBuffer[WriteToIndex] = EncodeTile(GroupId.xy);
+		}
 	}
 }
 #endif // SHADER_TILE_CATEGORIZATION
@@ -264,13 +283,10 @@ void TileMainCS(uint2 DispatchThreadId : SV_DispatchThreadID, uint LinearIndex :
 
 #if SHADER_STENCIL_CATEGORIZATION
 
-int  TileSize;
-int2 TileCount;
-int	 bRectPrimitive;
-int2 OutputResolution;
+float4 OutputResolutionAndInv;
 Buffer<uint> TileListBuffer;
 
-void StencilMainVS(
+void StrataTilePassVS(
 	in uint InVertexId : SV_VertexID,
 	in uint InInstanceId : SV_InstanceID,
 	out float4 OutPosition : SV_POSITION,
@@ -281,7 +297,7 @@ void StencilMainVS(
 	uint2 TileVertex = TileCoord * TILE_SIZE;
 	TileVertex.x += InVertexId == 1 || InVertexId == 2 || InVertexId == 4 ? TILE_SIZE : 0;
 	TileVertex.y += InVertexId == 2 || InVertexId == 4 || InVertexId == 5 ? TILE_SIZE : 0;
-	OutPosition = float4(float2(TileVertex) / OutputResolution * float2(2.0f, -2.0f) + float2(-1.0, 1.0f), 0.5f, 1.0f);
+	OutPosition = float4(float2(TileVertex) * OutputResolutionAndInv.zw * float2(2.0f, -2.0f) + float2(-1.0, 1.0f), 0.5f, 1.0f);
 	OutTileCoord = TileCoord;
 }
 
diff --git a/Engine/Source/Runtime/Renderer/Private/IndirectLightRendering.cpp b/Engine/Source/Runtime/Renderer/Private/IndirectLightRendering.cpp
--- a/Engine/Source/Runtime/Renderer/Private/IndirectLightRendering.cpp
+++ b/Engine/Source/Runtime/Renderer/Private/IndirectLightRendering.cpp
@@ -304,10 +304,16 @@ class FReflectionEnvironmentSkyLightingPS : public FGlobalShader
 
 		SHADER_PARAMETER_RDG_UNIFORM_BUFFER(FStrataGlobalUniformParameters, Strata)
 
-		RENDER_TARGET_BINDING_SLOTS()
 	END_SHADER_PARAMETER_STRUCT()
 }; // FReflectionEnvironmentSkyLightingPS
 
+
+BEGIN_SHADER_PARAMETER_STRUCT(FReflectionEnvironmentSkyLightingParameters, )
+	SHADER_PARAMETER_STRUCT_INCLUDE(FReflectionEnvironmentSkyLightingPS::FParameters, PS)
+	SHADER_PARAMETER_STRUCT_INCLUDE(Strata::FStrataTilePassVS::FParameters, VS)
+	RENDER_TARGET_BINDING_SLOTS()
+END_SHADER_PARAMETER_STRUCT()
+
 IMPLEMENT_GLOBAL_SHADER(FDiffuseIndirectCompositePS, "/Engine/Private/DiffuseIndirectComposite.usf", "MainPS", SF_Pixel);
 IMPLEMENT_GLOBAL_SHADER(FAmbientCubemapCompositePS, "/Engine/Private/AmbientCubemapComposite.usf", "MainPS", SF_Pixel);
 IMPLEMENT_GLOBAL_SHADER(FReflectionEnvironmentSkyLightingPS, "/Engine/Private/ReflectionEnvironmentPixelShader.usf", "ReflectionEnvironmentSkyLighting", SF_Pixel);
@@ -1123,7 +1129,9 @@ static void AddSkyReflectionPass(
 	bool bSkyLight, 
 	bool bDynamicSkyLight, 
 	bool bApplySkyShadowing,
-	bool bStrataFastPath)
+	bool bEnableStrataStencilTest,
+	bool bEnableStrataTiledPass,
+	uint32 StrataTileMaterialType)	// 0 simple material, 1 all/complex material. Does not matter when strata is disabled. See EStrataTileMaterialType
 {
 	// Render the reflection environment with tiled deferred culling
 	bool bHasBoxCaptures = (View.NumBoxReflectionCaptures > 0);
@@ -1140,7 +1148,7 @@ static void AddSkyReflectionPass(
 
 	const auto& SceneColorTexture = SceneTextures.Color;
 
-	FReflectionEnvironmentSkyLightingPS::FParameters* PassParameters = GraphBuilder.AllocParameters<FReflectionEnvironmentSkyLightingPS::FParameters>();
+	FReflectionEnvironmentSkyLightingParameters* PassParameters = GraphBuilder.AllocParameters<FReflectionEnvironmentSkyLightingParameters>();
 
 	// Setup the parameters of the shader.
 	{
@@ -1168,70 +1176,70 @@ static void AddSkyReflectionPass(
 			const float Mul = 1.0f / (Max - Min);
 			const float Add = -Min / (Max - Min);
 
-			PassParameters->OcclusionTintAndMinOcclusion = SkyLightOcclusionTintAndMinOcclusion;
-			PassParameters->ContrastAndNormalizeMulAdd = FVector(SkyLightContrast, Mul, Add);
-			PassParameters->OcclusionExponent = SkyLightOcclusionExponent;
-			PassParameters->OcclusionCombineMode = SkyLightOcclusionCombineMode == OCM_Minimum ? 0.0f : 1.0f;
-			PassParameters->ApplyBentNormalAO = DynamicBentNormalAO;
-			PassParameters->InvSkySpecularOcclusionStrength = 1.0f / FMath::Max(CVarSkySpecularOcclusionStrength.GetValueOnRenderThread(), 0.1f);
+			PassParameters->PS.OcclusionTintAndMinOcclusion = SkyLightOcclusionTintAndMinOcclusion;
+			PassParameters->PS.ContrastAndNormalizeMulAdd = FVector(SkyLightContrast, Mul, Add);
+			PassParameters->PS.OcclusionExponent = SkyLightOcclusionExponent;
+			PassParameters->PS.OcclusionCombineMode = SkyLightOcclusionCombineMode == OCM_Minimum ? 0.0f : 1.0f;
+			PassParameters->PS.ApplyBentNormalAO = DynamicBentNormalAO;
+			PassParameters->PS.InvSkySpecularOcclusionStrength = 1.0f / FMath::Max(CVarSkySpecularOcclusionStrength.GetValueOnRenderThread(), 0.1f);
 		}
 
 		// Setups all shader parameters related to distance field AO
 		{
 			FIntPoint AOBufferSize = GetBufferSizeForAO();
-			PassParameters->AOBufferBilinearUVMax = FVector2D(
+			PassParameters->PS.AOBufferBilinearUVMax = FVector2D(
 				(View.ViewRect.Width() / GAODownsampleFactor - 0.51f) / AOBufferSize.X, // 0.51 - so bilateral gather4 won't sample invalid texels
 				(View.ViewRect.Height() / GAODownsampleFactor - 0.51f) / AOBufferSize.Y);
 
 			extern float GAOViewFadeDistanceScale;
-			PassParameters->AOMaxViewDistance = GetMaxAOViewDistance();
-			PassParameters->DistanceFadeScale = 1.0f / ((1.0f - GAOViewFadeDistanceScale) * GetMaxAOViewDistance());
+			PassParameters->PS.AOMaxViewDistance = GetMaxAOViewDistance();
+			PassParameters->PS.DistanceFadeScale = 1.0f / ((1.0f - GAOViewFadeDistanceScale) * GetMaxAOViewDistance());
 
-			PassParameters->BentNormalAOTexture = DynamicBentNormalAOTexture;
-			PassParameters->BentNormalAOSampler = TStaticSamplerState<SF_Bilinear>::GetRHI();
+			PassParameters->PS.BentNormalAOTexture = DynamicBentNormalAOTexture;
+			PassParameters->PS.BentNormalAOSampler = TStaticSamplerState<SF_Bilinear>::GetRHI();
 		}
 
-		PassParameters->AmbientOcclusionTexture = AmbientOcclusionTexture;
-		PassParameters->AmbientOcclusionSampler = TStaticSamplerState<SF_Point>::GetRHI();
+		PassParameters->PS.AmbientOcclusionTexture = AmbientOcclusionTexture;
+		PassParameters->PS.AmbientOcclusionSampler = TStaticSamplerState<SF_Point>::GetRHI();
 
-		PassParameters->ScreenSpaceReflectionsTexture = ReflectionsColor ? ReflectionsColor : SystemTextures.Black;
-		PassParameters->ScreenSpaceReflectionsSampler = TStaticSamplerState<SF_Point>::GetRHI();
+		PassParameters->PS.ScreenSpaceReflectionsTexture = ReflectionsColor ? ReflectionsColor : SystemTextures.Black;
+		PassParameters->PS.ScreenSpaceReflectionsSampler = TStaticSamplerState<SF_Point>::GetRHI();
 
 		if (Scene->HasVolumetricCloud())
 		{
 			FVolumetricCloudRenderSceneInfo* CloudInfo = Scene->GetVolumetricCloudSceneInfo();
 
-			PassParameters->CloudSkyAOTexture = View.VolumetricCloudSkyAO != nullptr ? View.VolumetricCloudSkyAO : SystemTextures.Black;
-			PassParameters->CloudSkyAOWorldToLightClipMatrix = CloudInfo->GetVolumetricCloudCommonShaderParameters().CloudSkyAOWorldToLightClipMatrix;
-			PassParameters->CloudSkyAOFarDepthKm = CloudInfo->GetVolumetricCloudCommonShaderParameters().CloudSkyAOFarDepthKm;
-			PassParameters->CloudSkyAOEnabled = 1;
+			PassParameters->PS.CloudSkyAOTexture = View.VolumetricCloudSkyAO != nullptr ? View.VolumetricCloudSkyAO : SystemTextures.Black;
+			PassParameters->PS.CloudSkyAOWorldToLightClipMatrix = CloudInfo->GetVolumetricCloudCommonShaderParameters().CloudSkyAOWorldToLightClipMatrix;
+			PassParameters->PS.CloudSkyAOFarDepthKm = CloudInfo->GetVolumetricCloudCommonShaderParameters().CloudSkyAOFarDepthKm;
+			PassParameters->PS.CloudSkyAOEnabled = 1;
 		}
 		else
 		{
-			PassParameters->CloudSkyAOTexture = SystemTextures.Black;
-			PassParameters->CloudSkyAOEnabled = 0;
+			PassParameters->PS.CloudSkyAOTexture = SystemTextures.Black;
+			PassParameters->PS.CloudSkyAOEnabled = 0;
 		}
-		PassParameters->CloudSkyAOSampler = TStaticSamplerState<SF_Bilinear>::GetRHI();
+		PassParameters->PS.CloudSkyAOSampler = TStaticSamplerState<SF_Bilinear>::GetRHI();
 
-		PassParameters->PreIntegratedGF = GSystemTextures.PreintegratedGF->GetRenderTargetItem().ShaderResourceTexture;
-		PassParameters->PreIntegratedGFSampler = TStaticSamplerState<SF_Bilinear, AM_Clamp, AM_Clamp, AM_Clamp>::GetRHI();
+		PassParameters->PS.PreIntegratedGF = GSystemTextures.PreintegratedGF->GetRenderTargetItem().ShaderResourceTexture;
+		PassParameters->PS.PreIntegratedGFSampler = TStaticSamplerState<SF_Bilinear, AM_Clamp, AM_Clamp, AM_Clamp>::GetRHI();
 
-		PassParameters->SceneTextures = SceneTextureParameters;
+		PassParameters->PS.SceneTextures = SceneTextureParameters;
 
-		PassParameters->ViewUniformBuffer = View.ViewUniformBuffer;
-		PassParameters->ReflectionCaptureData = View.ReflectionCaptureUniformBuffer;
+		PassParameters->PS.ViewUniformBuffer = View.ViewUniformBuffer;
+		PassParameters->PS.ReflectionCaptureData = View.ReflectionCaptureUniformBuffer;
 		{
 			FReflectionUniformParameters ReflectionUniformParameters;
 			SetupReflectionUniformParameters(View, ReflectionUniformParameters);
-			PassParameters->ReflectionsParameters = CreateUniformBufferImmediate(ReflectionUniformParameters, UniformBuffer_SingleDraw);
+			PassParameters->PS.ReflectionsParameters = CreateUniformBufferImmediate(ReflectionUniformParameters, UniformBuffer_SingleDraw);
 		}
-		PassParameters->ForwardLightData = View.ForwardLightingResources->ForwardLightDataUniformBuffer;
+		PassParameters->PS.ForwardLightData = View.ForwardLightingResources->ForwardLightDataUniformBuffer;
 
-		PassParameters->Strata = Strata::BindStrataGlobalUniformParameters(View.StrataSceneData);
+		PassParameters->PS.Strata = Strata::BindStrataGlobalUniformParameters(View.StrataSceneData);
 	}
 
 	PassParameters->RenderTargets[0] = FRenderTargetBinding(SceneColorTexture.Target, ERenderTargetLoadAction::ELoad);
-	if (Strata::IsStrataEnabled() && Strata::IsClassificationEnabled())
+	if (bEnableStrataStencilTest)
 	{		
 		PassParameters->RenderTargets.DepthStencil = FDepthStencilBinding(SceneTextureParameters.SceneDepthTexture, ERenderTargetLoadAction::ELoad, ERenderTargetLoadAction::ELoad, FExclusiveDepthStencil::DepthRead_StencilWrite);
 	}
@@ -1248,23 +1256,34 @@ static void AddSkyReflectionPass(
 		View, bHasBoxCaptures, bHasSphereCaptures, DynamicBentNormalAO != 0.0f,
 		bSkyLight, bDynamicSkyLight, bApplySkyShadowing,
 		bRequiresSpecializedReflectionEnvironmentShader,
-		bStrataFastPath);
+		StrataTileMaterialType == EStrataTileMaterialType::ESimple);
 
 	TShaderMapRef<FReflectionEnvironmentSkyLightingPS> PixelShader(View.ShaderMap, PermutationVector);
-	ClearUnusedGraphResources(PixelShader, PassParameters);
+	ClearUnusedGraphResources(PixelShader, &PassParameters->PS);
+
+	EPrimitiveType StrataTilePrimitiveType = PT_TriangleList;
+	TShaderMapRef<Strata::FStrataTilePassVS> StrataTilePassVertexShader(View.ShaderMap);
+	PassParameters->VS.TileIndirectBuffer = nullptr;
+	PassParameters->VS.TileListBuffer = nullptr;
+	if (bEnableStrataTiledPass)
+	{
+		Strata::FillUpTiledPassData((EStrataTileMaterialType)StrataTileMaterialType, View, PassParameters->VS, StrataTilePrimitiveType);
+		ClearUnusedGraphResources(StrataTilePassVertexShader, &PassParameters->VS);
+	}
 
 	GraphBuilder.AddPass(
 		RDG_EVENT_NAME("ReflectionEnvironmentAndSky %dx%d", View.ViewRect.Width(), View.ViewRect.Height()),
 		PassParameters,
 		ERDGPassFlags::Raster,
-		[PassParameters, &View, PixelShader, bCheckerboardSubsurfaceRendering, bStrataFastPath](FRHICommandList& InRHICmdList)
+		[PassParameters, &View, PixelShader, bCheckerboardSubsurfaceRendering, 
+		bEnableStrataStencilTest, bEnableStrataTiledPass, StrataTilePassVertexShader, StrataTileMaterialType, StrataTilePrimitiveType](FRHICommandList& InRHICmdList)
 	{
 		InRHICmdList.SetViewport(View.ViewRect.Min.X, View.ViewRect.Min.Y, 0.0f, View.ViewRect.Max.X, View.ViewRect.Max.Y, 1.0f);
 
 		FGraphicsPipelineStateInitializer GraphicsPSOInit;
 		FPixelShaderUtils::InitFullscreenPipelineState(InRHICmdList, View.ShaderMap, PixelShader, GraphicsPSOInit);
 
-		if (Strata::IsStrataEnabled() && Strata::IsClassificationEnabled())
+		if (bEnableStrataStencilTest)
 		{
 			GraphicsPSOInit.DepthStencilState = TStaticDepthStencilState<
 				false, CF_Always, 
@@ -1300,13 +1319,28 @@ static void AddSkyReflectionPass(
 			InRHICmdList.SetDepthBounds(Values.MinDepth, Values.MaxDepth);
 		}
 
+		if (bEnableStrataTiledPass)
+		{
+			GraphicsPSOInit.BoundShaderState.VertexShaderRHI = StrataTilePassVertexShader.GetVertexShader();
+			GraphicsPSOInit.PrimitiveType = StrataTilePrimitiveType;
+		}
+
 		SetGraphicsPipelineState(InRHICmdList, GraphicsPSOInit);
-		SetShaderParameters(InRHICmdList, PixelShader, PixelShader.GetPixelShader(), *PassParameters);
-		if (Strata::IsStrataEnabled() && Strata::IsClassificationEnabled())
+		SetShaderParameters(InRHICmdList, PixelShader, PixelShader.GetPixelShader(), PassParameters->PS);
+		if (bEnableStrataStencilTest)
+		{
+			InRHICmdList.SetStencilRef(StrataTileMaterialType == EStrataTileMaterialType::ESimple ? Strata::StencilBit : 0x0);
+		}
+
+		if (!bEnableStrataTiledPass)
+		{
+			FPixelShaderUtils::DrawFullscreenTriangle(InRHICmdList);
+		}
+		else
 		{
-			InRHICmdList.SetStencilRef(bStrataFastPath ? Strata::StencilBit : 0x0);
+			SetShaderParameters(InRHICmdList, StrataTilePassVertexShader, StrataTilePassVertexShader.GetVertexShader(), PassParameters->VS);
+			InRHICmdList.DrawPrimitiveIndirect(PassParameters->VS.TileIndirectBuffer->GetIndirectRHICallBuffer(), 0);
 		}
-		FPixelShaderUtils::DrawFullscreenTriangle(InRHICmdList);
 	});
 }
 
@@ -1535,21 +1569,52 @@ void FDeferredShadingSceneRenderer::RenderDeferredReflectionsAndSkyLighting(
 		{
 			RDG_GPU_STAT_SCOPE(GraphBuilder, ReflectionEnvironment);
 
-			AddSkyReflectionPass(
-				GraphBuilder,
-				View,
-				Scene,
-				SceneTextures,
-				DynamicBentNormalAOTexture,
-				ReflectionsColor,
-				RayTracingReflectionOptions,
-				SceneTextureParameters,
-				bSkyLight,
-				bDynamicSkyLight,
-				bApplySkyShadowing,
-				false);
-			if (Strata::IsStrataEnabled() && Strata::IsClassificationEnabled())
+			const bool bStrataClassificationEnabled = Strata::IsStrataEnabled() && Strata::IsClassificationEnabled();
+			const bool bTilePassesReadingStrataEnabled = Strata::ShouldPassesReadingStrataBeTiled(Scene->GetFeatureLevel());
+
+			if (bStrataClassificationEnabled && bTilePassesReadingStrataEnabled)
+			{
+				const bool bEnableStrataTiledPass = true;
+				const bool bEnableStrataStencilTest = false;
+
+				AddSkyReflectionPass(
+					GraphBuilder,
+					View,
+					Scene,
+					SceneTextures,
+					DynamicBentNormalAOTexture,
+					ReflectionsColor,
+					RayTracingReflectionOptions,
+					SceneTextureParameters,
+					bSkyLight,
+					bDynamicSkyLight,
+					bApplySkyShadowing,
+					bEnableStrataStencilTest,
+					bEnableStrataTiledPass,
+					EStrataTileMaterialType::ESimple);
+
+				AddSkyReflectionPass(
+					GraphBuilder,
+					View,
+					Scene,
+					SceneTextures,
+					DynamicBentNormalAOTexture,
+					ReflectionsColor,
+					RayTracingReflectionOptions,
+					SceneTextureParameters,
+					bSkyLight,
+					bDynamicSkyLight,
+					bApplySkyShadowing,
+					bEnableStrataStencilTest,
+					bEnableStrataTiledPass,
+					EStrataTileMaterialType::EComplex);
+			}
+			else
 			{
+				const bool bEnableStrataTiledPass = false;
+				const bool bEnableStrataStencilTest = bStrataClassificationEnabled;
+
+				// Typical path uses when Strata is not enabled
 				AddSkyReflectionPass(
 					GraphBuilder,
 					View,
@@ -1562,7 +1627,27 @@ void FDeferredShadingSceneRenderer::RenderDeferredReflectionsAndSkyLighting(
 					bSkyLight,
 					bDynamicSkyLight,
 					bApplySkyShadowing,
-					true);
+					bEnableStrataStencilTest,
+					bEnableStrataTiledPass,
+					EStrataTileMaterialType::EComplex);
+				if (bStrataClassificationEnabled && !bTilePassesReadingStrataEnabled)
+				{
+					AddSkyReflectionPass(
+						GraphBuilder,
+						View,
+						Scene,
+						SceneTextures,
+						DynamicBentNormalAOTexture,
+						ReflectionsColor,
+						RayTracingReflectionOptions,
+						SceneTextureParameters,
+						bSkyLight,
+						bDynamicSkyLight,
+						bApplySkyShadowing,
+						bEnableStrataStencilTest,
+						bEnableStrataTiledPass,
+						EStrataTileMaterialType::ESimple);
+				}
 			}
 		}
 
diff --git a/Engine/Source/Runtime/Renderer/Private/Strata/Strata.cpp b/Engine/Source/Runtime/Renderer/Private/Strata/Strata.cpp
--- a/Engine/Source/Runtime/Renderer/Private/Strata/Strata.cpp
+++ b/Engine/Source/Runtime/Renderer/Private/Strata/Strata.cpp
@@ -35,7 +35,13 @@ static TAutoConsoleVariable<int32> CVarStrataClassification(
 static TAutoConsoleVariable<int32> CVarStrataClassificationDebug(
 	TEXT("r.Strata.Classification.Debug"),
 	0,
-	TEXT("Enable strata classification visualization."),
+	TEXT("Enable strata classification visualization: 1 shows tile with simple material only while 2 shows tile with complex material."),
+	ECVF_RenderThreadSafe);
+
+static TAutoConsoleVariable<int32> CVarStrataClassificationPassesReadingStrataAreTiled(
+	TEXT("r.Strata.Classification.PassesReadingStrataAreTiled"),
+	1,
+	TEXT("Enable the tiling of passes reading strata material (when possible) instead of doing multiple full screen passes testing stencil."),
 	ECVF_RenderThreadSafe);
 
 static TAutoConsoleVariable<int32> CVarStrataLUTResolution(
@@ -87,12 +93,16 @@ void FStrataSceneData::Reset()
 	MaterialLobesBuffer = nullptr;
 	MaterialLobesBufferUAV = nullptr;
 	MaterialLobesBufferSRV = nullptr;
-	ClassificationTileListBuffer = nullptr;
-	ClassificationTileListBufferUAV = nullptr;
-	ClassificationTileListBufferSRV = nullptr;
-	ClassificationTileIndirectBuffer = nullptr;
-	ClassificationTileIndirectBufferUAV = nullptr;
-	ClassificationTileIndirectBufferSRV = nullptr;
+
+	for (uint32 i = 0; i < EStrataTileMaterialType::ECount; ++i)
+	{
+		ClassificationTileListBuffer[i] = nullptr;
+		ClassificationTileListBufferUAV[i] = nullptr;
+		ClassificationTileListBufferSRV[i] = nullptr;
+		ClassificationTileIndirectBuffer[i] = nullptr;
+		ClassificationTileIndirectBufferUAV[i] = nullptr;
+		ClassificationTileIndirectBufferSRV[i] = nullptr;
+	}
 
 	StrataGlobalUniformParameters = nullptr;
 }
@@ -129,7 +139,12 @@ bool IsClassificationEnabled()
 	return CVarStrataClassification.GetValueOnAnyThread() > 0;
 }
 
-uint32 GetStrataTileSize()
+bool ShouldPassesReadingStrataBeTiled(ERHIFeatureLevel::Type FeatureLevel)
+{
+	return FeatureLevel >= ERHIFeatureLevel::SM5 && CVarStrataClassificationPassesReadingStrataAreTiled.GetValueOnAnyThread() > 0;
+}
+
+uint32 GetStrataBufferTileSize()
 {
 	return 8;
 }
@@ -168,18 +183,23 @@ void InitialiseStrataFrameSceneData(FSceneRenderer& SceneRenderer, FRDGBuilder&
 
 		// Tile classification buffers
 		{
-			const int32 TileInPixel = GetStrataTileSize();
+			const int32 TileInPixel = GetStrataBufferTileSize();
 			const FIntPoint TileResolution(FMath::DivideAndRoundUp(SceneTextureExtent.X, TileInPixel), FMath::DivideAndRoundUp(SceneTextureExtent.Y, TileInPixel));
 
-			StrataSceneData.ClassificationTileListBuffer = GraphBuilder.CreateBuffer(FRDGBufferDesc::CreateBufferDesc(sizeof(uint32), TileResolution.X * TileResolution.Y), TEXT("StrataTileListBuffer"));
-			StrataSceneData.ClassificationTileListBufferSRV = GraphBuilder.CreateSRV(StrataSceneData.ClassificationTileListBuffer, PF_R32_UINT);
-			StrataSceneData.ClassificationTileListBufferUAV = GraphBuilder.CreateUAV(StrataSceneData.ClassificationTileListBuffer, PF_R32_UINT);
+			// As of today we allocate one index+indirect buffer for each EStrataTileMaterialType.
+			// This is fine for two types, later we might want to have a single list and indirect buffer with offsets.
+			for (uint32 i = 0; i < EStrataTileMaterialType::ECount; ++i)
+			{
+				StrataSceneData.ClassificationTileListBuffer[i] = GraphBuilder.CreateBuffer(FRDGBufferDesc::CreateBufferDesc(sizeof(uint32), TileResolution.X * TileResolution.Y), i==EStrataTileMaterialType::ESimple ? TEXT("SimpleStrataTileListBuffer") : TEXT("ComplexStrataTileListBuffer"));
+				StrataSceneData.ClassificationTileListBufferSRV[i] = GraphBuilder.CreateSRV(StrataSceneData.ClassificationTileListBuffer[i], PF_R32_UINT);
+				StrataSceneData.ClassificationTileListBufferUAV[i] = GraphBuilder.CreateUAV(StrataSceneData.ClassificationTileListBuffer[i], PF_R32_UINT);
 
-			StrataSceneData.ClassificationTileIndirectBuffer = GraphBuilder.CreateBuffer(FRDGBufferDesc::CreateIndirectDesc<FRHIDrawIndirectParameters>(), TEXT("StrataTileIndirectBuffer"));
-			StrataSceneData.ClassificationTileIndirectBufferSRV = GraphBuilder.CreateSRV(StrataSceneData.ClassificationTileIndirectBuffer, PF_R32_UINT);
-			StrataSceneData.ClassificationTileIndirectBufferUAV = GraphBuilder.CreateUAV(StrataSceneData.ClassificationTileIndirectBuffer, PF_R32_UINT);
+				StrataSceneData.ClassificationTileIndirectBuffer[i] = GraphBuilder.CreateBuffer(FRDGBufferDesc::CreateIndirectDesc<FRHIDrawIndirectParameters>(), i == EStrataTileMaterialType::ESimple ? TEXT("SimpleStrataTileIndirectBuffer") : TEXT("ComplexStrataTileIndirectBuffer"));
+				StrataSceneData.ClassificationTileIndirectBufferSRV[i] = GraphBuilder.CreateSRV(StrataSceneData.ClassificationTileIndirectBuffer[i], PF_R32_UINT);
+				StrataSceneData.ClassificationTileIndirectBufferUAV[i] = GraphBuilder.CreateUAV(StrataSceneData.ClassificationTileIndirectBuffer[i], PF_R32_UINT);
 
-			AddClearUAVPass(GraphBuilder, StrataSceneData.ClassificationTileIndirectBufferUAV, 0);
+				AddClearUAVPass(GraphBuilder, StrataSceneData.ClassificationTileIndirectBufferUAV[i], 0);
+			}
 		}
 
 		// Top layer texture
@@ -451,8 +471,10 @@ class FStrataMaterialTileClassificationPassCS : public FGlobalShader
 		SHADER_PARAMETER(int32, bRectPrimitive)
 		SHADER_PARAMETER(FIntPoint, ViewResolution)
 		SHADER_PARAMETER_RDG_TEXTURE(Texture2D<uint>, ClassificationTexture)
-		SHADER_PARAMETER_RDG_BUFFER_UAV(RWBuffer, TileIndirectData)
-		SHADER_PARAMETER_RDG_BUFFER_UAV(RWBuffer, TileListData)
+		SHADER_PARAMETER_RDG_BUFFER_UAV(RWBuffer, SimpleTileIndirectDataBuffer)
+		SHADER_PARAMETER_RDG_BUFFER_UAV(RWBuffer, SimpleTileListDataBuffer)
+		SHADER_PARAMETER_RDG_BUFFER_UAV(RWBuffer, ComplexTileIndirectDataBuffer)
+		SHADER_PARAMETER_RDG_BUFFER_UAV(RWBuffer, ComplexTileListDataBuffer)
 	END_SHADER_PARAMETER_STRUCT()
 
 	static FPermutationDomain RemapPermutation(FPermutationDomain PermutationVector)
@@ -475,55 +497,15 @@ IMPLEMENT_GLOBAL_SHADER(FStrataMaterialTileClassificationPassCS, "/Engine/Privat
 
 ////////////////////////////////////////////////////////////////////////////////////////////////////////////////////////////////////////////////
 
-
-class FStrataMaterialStencilClassificationPassVS : public FGlobalShader
-{
-	DECLARE_GLOBAL_SHADER(FStrataMaterialStencilClassificationPassVS);
-	SHADER_USE_PARAMETER_STRUCT(FStrataMaterialStencilClassificationPassVS, FGlobalShader);
-
-	using FPermutationDomain = TShaderPermutationDomain<>;
-
-	BEGIN_SHADER_PARAMETER_STRUCT(FParameters, )
-		SHADER_PARAMETER_STRUCT_REF(FViewUniformShaderParameters, ViewUniformBuffer)
-		SHADER_PARAMETER(int32, TileSize)
-		SHADER_PARAMETER(int32, bRectPrimitive)
-		SHADER_PARAMETER(FIntPoint, TileCount)
-		SHADER_PARAMETER(FIntPoint, OutputResolution)
-		SHADER_PARAMETER_RDG_BUFFER_SRV(Buffer, TileListBuffer)
-	END_SHADER_PARAMETER_STRUCT()
-
-	static FPermutationDomain RemapPermutation(FPermutationDomain PermutationVector)
-	{
-		return PermutationVector;
-	}
-
-	static bool ShouldCompilePermutation(const FGlobalShaderPermutationParameters& Parameters)
-	{
-		return GetMaxSupportedFeatureLevel(Parameters.Platform) >= ERHIFeatureLevel::SM5 && Strata::IsStrataEnabled();
-	}
-
-	static void ModifyCompilationEnvironment(const FGlobalShaderPermutationParameters& Parameters, FShaderCompilerEnvironment& OutEnvironment)
-	{
-		FGlobalShader::ModifyCompilationEnvironment(Parameters, OutEnvironment);
-		OutEnvironment.SetDefine(TEXT("SHADER_STENCIL_CATEGORIZATION"), 1);
-	}
-};
-
-class FStrataMaterialStencilClassificationPassPS : public FGlobalShader
+class FStrataMaterialStencilTaggingPassPS : public FGlobalShader
 {
-	DECLARE_GLOBAL_SHADER(FStrataMaterialStencilClassificationPassPS);
-	SHADER_USE_PARAMETER_STRUCT(FStrataMaterialStencilClassificationPassPS, FGlobalShader);
+	DECLARE_GLOBAL_SHADER(FStrataMaterialStencilTaggingPassPS);
+	SHADER_USE_PARAMETER_STRUCT(FStrataMaterialStencilTaggingPassPS, FGlobalShader);
 
 	using FPermutationDomain = TShaderPermutationDomain<>;
 
 	BEGIN_SHADER_PARAMETER_STRUCT(FParameters, )
-		SHADER_PARAMETER_STRUCT_REF(FViewUniformShaderParameters, ViewUniformBuffer)
-		SHADER_PARAMETER(int32, TileSize)
-		SHADER_PARAMETER(int32, bRectPrimitive)
-		SHADER_PARAMETER(FIntPoint, TileCount)
-		SHADER_PARAMETER(FIntPoint, OutputResolution)
-		SHADER_PARAMETER_RDG_BUFFER_SRV(Buffer, TileListBuffer)
-		RDG_BUFFER_ACCESS(TileIndirectBuffer, ERHIAccess::IndirectArgs)
+		SHADER_PARAMETER_STRUCT_INCLUDE(Strata::FStrataTilePassVS::FParameters, VS)
 		RENDER_TARGET_BINDING_SLOTS()
 	END_SHADER_PARAMETER_STRUCT()
 
@@ -544,38 +526,51 @@ class FStrataMaterialStencilClassificationPassPS : public FGlobalShader
 	}
 };
 
-IMPLEMENT_GLOBAL_SHADER(FStrataMaterialStencilClassificationPassVS, "/Engine/Private/Strata/StrataMaterialClassification.usf", "StencilMainVS", SF_Vertex);
-IMPLEMENT_GLOBAL_SHADER(FStrataMaterialStencilClassificationPassPS, "/Engine/Private/Strata/StrataMaterialClassification.usf", "StencilMainPS", SF_Pixel);
+IMPLEMENT_GLOBAL_SHADER(FStrataTilePassVS, "/Engine/Private/Strata/StrataMaterialClassification.usf", "StrataTilePassVS", SF_Vertex);
+IMPLEMENT_GLOBAL_SHADER(FStrataMaterialStencilTaggingPassPS, "/Engine/Private/Strata/StrataMaterialClassification.usf", "StencilMainPS", SF_Pixel);
+
+void FillUpTiledPassData(
+	EStrataTileMaterialType Type, 
+	const FViewInfo& View, 
+	FStrataTilePassVS::FParameters& ParametersVS,
+	EPrimitiveType& PrimitiveType)
+{
+	const FIntPoint OutputResolution = View.ViewRect.Size();
+	ParametersVS.OutputResolutionAndInv = FVector4(OutputResolution.X, OutputResolution.Y, 1.0f / float(OutputResolution.X), 1.0f / float(OutputResolution.Y));
+	ParametersVS.TileListBuffer = View.StrataSceneData->ClassificationTileListBufferSRV[Type];
+	ParametersVS.TileIndirectBuffer = View.StrataSceneData->ClassificationTileIndirectBuffer[Type];
+
+	PrimitiveType = GRHISupportsRectTopology ? PT_RectList : PT_TriangleList;
+}
 
-static void AddStrataInternalClassifedTilePass(
+static void AddStrataInternalClassificationTilePass(
 	FRDGBuilder& GraphBuilder,
 	const FViewInfo& View,
 	const FRDGTextureRef* DepthTexture,
 	const FRDGTextureRef* ColorTexture,
-	FRDGBufferSRVRef TileListBufferSRV,
-	FRDGBufferRef TileIndirectBuffer)
-{	
+	EStrataTileMaterialType TileMaterialType,
+	const bool bDebug = false)
+{
+	// We cannot early exit due to the fact that the local light are still rendered as mesh volumes (so cannot be tiled as such)
+	//if (ShouldPassesReadingStrataBeTiled())
+	//{
+	//	return;
+	//}
+
+	EPrimitiveType StrataTilePrimitiveType = PT_TriangleList;
 	const FIntPoint OutputResolution = View.ViewRect.Size();
-	const int32 TileSize = GetStrataTileSize();
-	
-	const FIntPoint TileCount(FMath::DivideAndRoundUp(OutputResolution.X, TileSize), FMath::DivideAndRoundUp(OutputResolution.Y, TileSize));
-	const uint32 InstanceCount = TileCount.X * TileCount.Y;
+	FVector4 OutputResolutionAndInv = FVector4(OutputResolution.X, OutputResolution.Y, 1.0f / float(OutputResolution.X), 1.0f / float(OutputResolution.Y));
 
-	FStrataMaterialStencilClassificationPassPS::FParameters* ParametersPS = GraphBuilder.AllocParameters<FStrataMaterialStencilClassificationPassPS::FParameters>();
-	ParametersPS->TileSize = TileSize;
-	ParametersPS->TileCount = TileCount;
-	ParametersPS->bRectPrimitive = GRHISupportsRectTopology ? 1 : 0;
-	ParametersPS->OutputResolution = OutputResolution;
-	ParametersPS->TileListBuffer = TileListBufferSRV;
-	ParametersPS->TileIndirectBuffer = TileIndirectBuffer;
+	FStrataMaterialStencilTaggingPassPS::FParameters* ParametersPS = GraphBuilder.AllocParameters<FStrataMaterialStencilTaggingPassPS::FParameters>();
+	FillUpTiledPassData(TileMaterialType, View, ParametersPS->VS, StrataTilePrimitiveType);
 
-	TShaderMapRef<FStrataMaterialStencilClassificationPassVS> VertexShader(View.ShaderMap);
-	TShaderMapRef<FStrataMaterialStencilClassificationPassPS> PixelShader(View.ShaderMap);
+	TShaderMapRef<FStrataTilePassVS> VertexShader(View.ShaderMap);
+	TShaderMapRef<FStrataMaterialStencilTaggingPassPS> PixelShader(View.ShaderMap);
 
 	// For debug purpose
-	const bool bDebug = ColorTexture != nullptr;
 	if (bDebug)
 	{
+		check(ColorTexture);
 		ParametersPS->RenderTargets[0] = FRenderTargetBinding(*ColorTexture, ERenderTargetLoadAction::ELoad);
 	}
 	else
@@ -594,15 +589,8 @@ static void AddStrataInternalClassifedTilePass(
 		RDG_EVENT_NAME("StrataStencilClassificationPass"),
 		ParametersPS,
 		ERDGPassFlags::Raster,
-		[ParametersPS, VertexShader, PixelShader, OutputResolution, InstanceCount, bDebug](FRHICommandList& RHICmdList)
+		[ParametersPS, VertexShader, PixelShader, OutputResolution, StrataTilePrimitiveType, bDebug](FRHICommandList& RHICmdList)
 		{
-			FStrataMaterialStencilClassificationPassVS::FParameters ParametersVS;
-			ParametersVS.TileSize = ParametersPS->TileSize;
-			ParametersVS.bRectPrimitive = ParametersPS->bRectPrimitive;
-			ParametersVS.TileCount = ParametersPS->TileCount;
-			ParametersVS.OutputResolution = ParametersPS->OutputResolution;
-			ParametersVS.TileListBuffer = ParametersPS->TileListBuffer;
-
 			FGraphicsPipelineStateInitializer GraphicsPSOInit;
 			RHICmdList.ApplyCachedRenderTargets(GraphicsPSOInit);
 			GraphicsPSOInit.BlendState = TStaticBlendState<CW_RGBA, BO_Add, BF_One, BF_One, BO_Max, BF_SourceAlpha, BF_DestAlpha>::GetRHI();
@@ -622,15 +610,14 @@ static void AddStrataInternalClassifedTilePass(
 			GraphicsPSOInit.BoundShaderState.VertexDeclarationRHI = GFilterVertexDeclaration.VertexDeclarationRHI;
 			GraphicsPSOInit.BoundShaderState.VertexShaderRHI = VertexShader.GetVertexShader();
 			GraphicsPSOInit.BoundShaderState.PixelShaderRHI = PixelShader.GetPixelShader();
-			GraphicsPSOInit.PrimitiveType = ParametersPS->bRectPrimitive > 0 ? PT_RectList : PT_TriangleList;
+			GraphicsPSOInit.PrimitiveType = StrataTilePrimitiveType;
 			SetGraphicsPipelineState(RHICmdList, GraphicsPSOInit);
-			SetShaderParameters(RHICmdList, VertexShader, VertexShader.GetVertexShader(), ParametersVS);
-			//SetShaderParameters(RHICmdList, PixelShader, PixelShader.GetPixelShader(), *ParametersPS);
+			SetShaderParameters(RHICmdList, VertexShader, VertexShader.GetVertexShader(), ParametersPS->VS);
 
 			RHICmdList.SetStencilRef(StencilBit);
 			RHICmdList.SetViewport(0, 0, 0.0f, OutputResolution.X, OutputResolution.Y, 1.0f);
 			RHICmdList.SetStreamSource(0, nullptr, 0);
-			RHICmdList.DrawPrimitiveIndirect(ParametersPS->TileIndirectBuffer->GetIndirectRHICallBuffer(), 0);
+			RHICmdList.DrawPrimitiveIndirect(ParametersPS->VS.TileIndirectBuffer->GetIndirectRHICallBuffer(), 0);
 		});
 }
 
@@ -639,8 +626,7 @@ void AddStrataStencilPass(
 	const FViewInfo& View,
 	const FMinimalSceneTextures& SceneTextures)
 {
-	AddStrataInternalClassifedTilePass(GraphBuilder, View, &SceneTextures.Depth.Target, nullptr, 
-		View.StrataSceneData->ClassificationTileListBufferSRV, View.StrataSceneData->ClassificationTileIndirectBuffer);
+	AddStrataInternalClassificationTilePass(GraphBuilder, View, &SceneTextures.Depth.Target, nullptr, EStrataTileMaterialType::ESimple);
 }
 
 void AddStrataStencilPass(
@@ -651,8 +637,7 @@ void AddStrataStencilPass(
 	for (int32 i = 0; i < Views.Num(); ++i)
 	{
 		const FViewInfo& View = Views[i];
-		AddStrataInternalClassifedTilePass(GraphBuilder, View, &SceneTextures.Depth.Target, nullptr, 
-			View.StrataSceneData->ClassificationTileListBufferSRV, View.StrataSceneData->ClassificationTileIndirectBuffer);
+		AddStrataInternalClassificationTilePass(GraphBuilder, View, &SceneTextures.Depth.Target, nullptr, EStrataTileMaterialType::ESimple);
 	}
 }
 
@@ -702,12 +687,14 @@ void AddStrataMaterialClassificationPass(FRDGBuilder& GraphBuilder, const FMinim
 		{
 			TShaderMapRef<FStrataMaterialTileClassificationPassCS> ComputeShader(View.ShaderMap);
 			FStrataMaterialTileClassificationPassCS::FParameters* PassParameters = GraphBuilder.AllocParameters<FStrataMaterialTileClassificationPassCS::FParameters>();
-			PassParameters->TileSize = GetStrataTileSize();
+			PassParameters->TileSize = GetStrataBufferTileSize();	// STRATA_TODO not sure we want to tie the buffer tile optimisation for cache and the Categotisation tile size?
 			PassParameters->bRectPrimitive = GRHISupportsRectTopology ? 1 : 0;
 			PassParameters->ViewResolution = View.ViewRect.Size();
 			PassParameters->ClassificationTexture = View.StrataSceneData->ClassificationTexture;
-			PassParameters->TileListData = View.StrataSceneData->ClassificationTileListBufferUAV;
-			PassParameters->TileIndirectData = View.StrataSceneData->ClassificationTileIndirectBufferUAV;
+			PassParameters->SimpleTileListDataBuffer = View.StrataSceneData->ClassificationTileListBufferUAV[EStrataTileMaterialType::ESimple];
+			PassParameters->SimpleTileIndirectDataBuffer = View.StrataSceneData->ClassificationTileIndirectBufferUAV[EStrataTileMaterialType::ESimple];
+			PassParameters->ComplexTileListDataBuffer = View.StrataSceneData->ClassificationTileListBufferUAV[EStrataTileMaterialType::EComplex];
+			PassParameters->ComplexTileIndirectDataBuffer = View.StrataSceneData->ClassificationTileIndirectBufferUAV[EStrataTileMaterialType::EComplex];
 
 			// Add 64 threads permutation
 			const uint32 GroupSize = 8;
@@ -894,14 +881,17 @@ void AddStrataDebugPasses(FRDGBuilder& GraphBuilder, const TArray<FViewInfo>& Vi
 		}
 	}
 
-	if (IsClassificationEnabled() && CVarStrataClassificationDebug.GetValueOnAnyThread() > 0)
+	const int32 StrataClassificationDebug = CVarStrataClassificationDebug.GetValueOnAnyThread();
+	if (IsClassificationEnabled() && StrataClassificationDebug > 0)
 	{
 		RDG_EVENT_SCOPE(GraphBuilder, "StrataVisualizeClassification");
 		for (int32 i = 0; i < Views.Num(); ++i)
 		{
 			const FViewInfo& View = Views[i];
-			AddStrataInternalClassifedTilePass(GraphBuilder, View, nullptr, &SceneColorTexture,
-				View.StrataSceneData->ClassificationTileListBufferSRV, View.StrataSceneData->ClassificationTileIndirectBuffer);
+			const int32 StrataTileMaterialTypeIndex = FMath::Clamp(StrataClassificationDebug - 1, 0, 1);
+			const bool bDebugPass = true;
+			AddStrataInternalClassificationTilePass(
+				GraphBuilder, View, nullptr, &SceneColorTexture, (EStrataTileMaterialType)StrataTileMaterialTypeIndex, bDebugPass);
 		}
 	}
 
diff --git a/Engine/Source/Runtime/Renderer/Private/Strata/Strata.h b/Engine/Source/Runtime/Renderer/Private/Strata/Strata.h
--- a/Engine/Source/Runtime/Renderer/Private/Strata/Strata.h
+++ b/Engine/Source/Runtime/Renderer/Private/Strata/Strata.h
@@ -31,6 +31,13 @@ BEGIN_GLOBAL_SHADER_PARAMETER_STRUCT(FStrataGlobalUniformParameters, )
 	SHADER_PARAMETER_SAMPLER(SamplerState, GGXEnergyLUTSampler)
 END_GLOBAL_SHADER_PARAMETER_STRUCT()
 
+enum EStrataTileMaterialType : uint32
+{
+	ESimple = 0,
+	EComplex = 1,
+	ECount
+};
+
 struct FStrataSceneData
 {
 	uint32 MaxBytesPerPixel;
@@ -40,12 +47,12 @@ struct FStrataSceneData
 	FRDGBufferRef MaterialLobesBuffer;
 	FRDGBufferUAVRef MaterialLobesBufferUAV;
 	FRDGBufferSRVRef MaterialLobesBufferSRV;
-	FRDGBufferRef ClassificationTileListBuffer;
-	FRDGBufferUAVRef ClassificationTileListBufferUAV;
-	FRDGBufferSRVRef ClassificationTileListBufferSRV;
-	FRDGBufferRef ClassificationTileIndirectBuffer;
-	FRDGBufferUAVRef ClassificationTileIndirectBufferUAV;
-	FRDGBufferSRVRef ClassificationTileIndirectBufferSRV;
+	FRDGBufferRef ClassificationTileListBuffer[EStrataTileMaterialType::ECount];
+	FRDGBufferUAVRef ClassificationTileListBufferUAV[EStrataTileMaterialType::ECount];
+	FRDGBufferSRVRef ClassificationTileListBufferSRV[EStrataTileMaterialType::ECount];
+	FRDGBufferRef ClassificationTileIndirectBuffer[EStrataTileMaterialType::ECount];
+	FRDGBufferUAVRef ClassificationTileIndirectBufferUAV[EStrataTileMaterialType::ECount];
+	FRDGBufferSRVRef ClassificationTileIndirectBufferSRV[EStrataTileMaterialType::ECount];
 
 	FRDGTextureRef ClassificationTexture;
 	FRDGTextureRef TopLayerNormalTexture;
@@ -67,12 +74,14 @@ struct FStrataSceneData
 	void Reset();
 };
 
+
 namespace Strata
 {
 constexpr uint32 StencilBit = 0x80; // In sync with SceneRenderTargets.h - GET_STENCIL_BIT_MASK(STENCIL_STRATA_FASTPATH)
 
 bool IsStrataEnabled();
 bool IsClassificationEnabled();
+bool ShouldPassesReadingStrataBeTiled(ERHIFeatureLevel::Type FeatureLevel);
 
 void InitialiseStrataFrameSceneData(FSceneRenderer& SceneRenderer, FRDGBuilder& GraphBuilder);
 
@@ -86,6 +95,35 @@ void AddStrataStencilPass(FRDGBuilder& GraphBuilder, const FViewInfo& View, cons
 void AddStrataStencilPass(FRDGBuilder& GraphBuilder, const TArray<FViewInfo>& Views, const FMinimalSceneTextures& SceneTextures);
 
 void AddStrataDebugPasses(FRDGBuilder& GraphBuilder, const TArray<FViewInfo>& Views, FRDGTextureRef SceneColorTexture, EShaderPlatform Platform);
+
+
+class FStrataTilePassVS : public FGlobalShader
+{
+	DECLARE_GLOBAL_SHADER(FStrataTilePassVS);
+	SHADER_USE_PARAMETER_STRUCT(FStrataTilePassVS, FGlobalShader);
+
+	BEGIN_SHADER_PARAMETER_STRUCT(FParameters, )
+		SHADER_PARAMETER(FVector4, OutputResolutionAndInv)
+		SHADER_PARAMETER_RDG_BUFFER_SRV(Buffer<uint>, TileListBuffer)
+		RDG_BUFFER_ACCESS(TileIndirectBuffer, ERHIAccess::IndirectArgs)
+	END_SHADER_PARAMETER_STRUCT()
+
+	static bool ShouldCompilePermutation(const FGlobalShaderPermutationParameters& Parameters)
+	{
+		return GetMaxSupportedFeatureLevel(Parameters.Platform) >= ERHIFeatureLevel::SM5; // We do not skip the compilation because we have some conditional when tiling a pass and the shader must be fetch once before hand.
+	}
+
+	static void ModifyCompilationEnvironment(const FGlobalShaderPermutationParameters& Parameters, FShaderCompilerEnvironment& OutEnvironment)
+	{
+		FGlobalShader::ModifyCompilationEnvironment(Parameters, OutEnvironment);
+		OutEnvironment.SetDefine(TEXT("SHADER_STENCIL_CATEGORIZATION"), 1);
+	}
 };
 
+void FillUpTiledPassData(EStrataTileMaterialType Type, const FViewInfo& View, FStrataTilePassVS::FParameters& ParametersVS, EPrimitiveType& PrimitiveType);
+
+};
+
+
+
 
