Made the directional light tiled when strata classification and tiling is enabled.
The VS now has some  permutation to only compute some extrapolators only when needed for pass PS compatibility.
Strata disabled tested.

diff --git a/Engine/Shaders/Private/Strata/StrataMaterialClassification.usf b/Engine/Shaders/Private/Strata/StrataMaterialClassification.usf
--- a/Engine/Shaders/Private/Strata/StrataMaterialClassification.usf
+++ b/Engine/Shaders/Private/Strata/StrataMaterialClassification.usf
@@ -283,27 +283,44 @@ void TileMainCS(uint2 DispatchThreadId : SV_DispatchThreadID, uint LinearIndex :
 
 #if SHADER_STENCIL_CATEGORIZATION
 
-float4 OutputResolutionAndInv;
+float4 OutputViewSizeAndInvSize;
+float4 OutputBufferSizeAndInvSize;
+float4x4 ViewScreenToTranslatedWorld;
 Buffer<uint> TileListBuffer;
 
 void StrataTilePassVS(
 	in uint InVertexId : SV_VertexID,
 	in uint InInstanceId : SV_InstanceID,
-	out float4 OutPosition : SV_POSITION,
-	out uint2 OutTileCoord : TILE_COORD)
+#if PERMUTATION_ENABLE_TEXCOORD_SCREENVECTOR
+	out float2 OutTexCoord : TEXCOORD0,		// Used by DeferredLightPixelMain directional light
+	out float3 OutScreenVector : TEXCOORD1,	// Used by DeferredLightPixelMain directional light
+#endif
+#if PERMUTATION_ENABLE_DEBUG
+	out uint2 OutTileCoord : TILE_COORD,	// Used by StencilMainPS below
+#endif
+	out float4 OutPosition : SV_POSITION)	// Use by StencilMainPS
 {
 	const uint2 TileCoord = DecodeTile(TileListBuffer[InInstanceId]);
 
 	uint2 TileVertex = TileCoord * TILE_SIZE;
 	TileVertex.x += InVertexId == 1 || InVertexId == 2 || InVertexId == 4 ? TILE_SIZE : 0;
 	TileVertex.y += InVertexId == 2 || InVertexId == 4 || InVertexId == 5 ? TILE_SIZE : 0;
-	OutPosition = float4(float2(TileVertex) * OutputResolutionAndInv.zw * float2(2.0f, -2.0f) + float2(-1.0, 1.0f), 0.5f, 1.0f);
+
+	OutPosition = float4(float2(TileVertex) * OutputViewSizeAndInvSize.zw * float2(2.0f, -2.0f) + float2(-1.0, 1.0f), 0.5f, 1.0f);
+
+#if PERMUTATION_ENABLE_TEXCOORD_SCREENVECTOR
+	OutTexCoord = float2(TileVertex) * OutputBufferSizeAndInvSize.zw;
+	OutScreenVector = mul(float4(OutPosition.xy, 1, 0), ViewScreenToTranslatedWorld).xyz;
+#endif
+
+#if PERMUTATION_ENABLE_DEBUG
 	OutTileCoord = TileCoord;
+#endif
 }
 
 void StencilMainPS(
-	in float4 SVPos : SV_POSITION,
 	in uint2 InTileCoord : TILE_COORD,
+	in float4 SVPos : SV_POSITION,
 	out float4 OutColor0 : SV_Target0)
 {
 	const bool bTileX = (InTileCoord.x & 1) == 0;
diff --git a/Engine/Source/Runtime/Renderer/Private/DeferredShadingRenderer.h b/Engine/Source/Runtime/Renderer/Private/DeferredShadingRenderer.h
--- a/Engine/Source/Runtime/Renderer/Private/DeferredShadingRenderer.h
+++ b/Engine/Source/Runtime/Renderer/Private/DeferredShadingRenderer.h
@@ -35,12 +35,12 @@ namespace LumenRadianceCache
 	class FRadianceCacheInputs;
 	class FRadianceCacheInterpolationParameters;
 }
+class FRenderLightParameters;
 
 struct FSceneWithoutWaterTextures;
 struct FHeightFogRenderingParameters;
 struct FRayTracingReflectionOptions;
 struct FHairStrandsTransmittanceMaskData;
-
 struct FTranslucencyLightingVolumeTextures;
 
 /**   
@@ -718,6 +718,7 @@ private:
 		FRHICommandList& RHICmdList,
 		const FViewInfo& View,
 		const FLightSceneInfo* LightSceneInfo,
+		FRenderLightParameters* PassParameters,
 		FRHITexture* ScreenShadowMaskTexture,
 		FRHITexture* LightingChannelTexture,
 		bool bRenderOverlap,
diff --git a/Engine/Source/Runtime/Renderer/Private/IndirectLightRendering.cpp b/Engine/Source/Runtime/Renderer/Private/IndirectLightRendering.cpp
--- a/Engine/Source/Runtime/Renderer/Private/IndirectLightRendering.cpp
+++ b/Engine/Source/Runtime/Renderer/Private/IndirectLightRendering.cpp
@@ -1262,7 +1262,10 @@ static void AddSkyReflectionPass(
 	ClearUnusedGraphResources(PixelShader, &PassParameters->PS);
 
 	EPrimitiveType StrataTilePrimitiveType = PT_TriangleList;
-	TShaderMapRef<Strata::FStrataTilePassVS> StrataTilePassVertexShader(View.ShaderMap);
+	Strata::FStrataTilePassVS::FPermutationDomain VSPermutationVector;
+	VSPermutationVector.Set< Strata::FStrataTilePassVS::FEnableDebug >(false);
+	VSPermutationVector.Set< Strata::FStrataTilePassVS::FEnableTexCoordScreenVector >(false);
+	TShaderMapRef<Strata::FStrataTilePassVS> StrataTilePassVertexShader(View.ShaderMap, VSPermutationVector);
 	PassParameters->VS.TileIndirectBuffer = nullptr;
 	PassParameters->VS.TileListBuffer = nullptr;
 	if (bEnableStrataTiledPass)
diff --git a/Engine/Source/Runtime/Renderer/Private/LightRendering.cpp b/Engine/Source/Runtime/Renderer/Private/LightRendering.cpp
--- a/Engine/Source/Runtime/Renderer/Private/LightRendering.cpp
+++ b/Engine/Source/Runtime/Renderer/Private/LightRendering.cpp
@@ -1086,6 +1086,11 @@ BEGIN_SHADER_PARAMETER_STRUCT(FRenderLightParameters, )
 	SHADER_PARAMETER_STRUCT_INCLUDE(FVolumetricCloudShadowAOParameters, CloudShadowAO)
 	RDG_TEXTURE_ACCESS(ShadowMaskTexture, ERHIAccess::SRVGraphics)
 	RDG_TEXTURE_ACCESS(LightingChannelsTexture, ERHIAccess::SRVGraphics)
+	// We reference all the Strata tiled resources we might need in this pass
+	SHADER_PARAMETER_RDG_BUFFER_SRV(Buffer<uint>, TileListBufferSimple)
+	SHADER_PARAMETER_RDG_BUFFER_SRV(Buffer<uint>, TileListBufferComplex)
+	RDG_BUFFER_ACCESS(TileIndirectBufferSimple, ERHIAccess::IndirectArgs)
+	RDG_BUFFER_ACCESS(TileIndirectBufferComplex, ERHIAccess::IndirectArgs)
 	RENDER_TARGET_BINDING_SLOTS()
 END_SHADER_PARAMETER_STRUCT()
 
@@ -1112,6 +1117,14 @@ void GetRenderLightParameters(
 	{
 		Parameters.RenderTargets.DepthStencil = FDepthStencilBinding(SceneDepthTexture, ERenderTargetLoadAction::ELoad, ERenderTargetLoadAction::ELoad, FExclusiveDepthStencil::DepthRead_StencilWrite);
 	}
+
+	if (Strata::ShouldPassesReadingStrataBeTiled(View.Family->GetFeatureLevel()))
+	{
+		Parameters.TileListBufferSimple = View.StrataSceneData->ClassificationTileListBufferSRV[EStrataTileMaterialType::ESimple];
+		Parameters.TileListBufferComplex = View.StrataSceneData->ClassificationTileListBufferSRV[EStrataTileMaterialType::EComplex];
+		Parameters.TileIndirectBufferSimple = View.StrataSceneData->ClassificationTileIndirectBuffer[EStrataTileMaterialType::ESimple];
+		Parameters.TileIndirectBufferComplex = View.StrataSceneData->ClassificationTileIndirectBuffer[EStrataTileMaterialType::EComplex];
+	}
 }
 
 FHairStrandsTransmittanceMaskData CreateDummyHairStrandsTransmittanceMaskData(FRDGBuilder& GraphBuilder, FGlobalShaderMap* ShaderMap);
@@ -1259,7 +1272,7 @@ void FDeferredShadingSceneRenderer::RenderLights(
 						RDG_EVENT_NAME("StandardDeferredLighting"),
 						PassParameters,
 						ERDGPassFlags::Raster,
-						[this, &View, &SortedLights, LightingChannelsTexture, StandardDeferredStart, AttenuationLightStart](FRHICommandList& RHICmdList)
+						[this, &View, &SortedLights, LightingChannelsTexture, StandardDeferredStart, AttenuationLightStart, PassParameters](FRHICommandList& RHICmdList)
 					{
 						// Draw non-shadowed non-light function lights without changing render targets between them
 						for (int32 LightIndex = StandardDeferredStart; LightIndex < AttenuationLightStart; LightIndex++)
@@ -1271,7 +1284,7 @@ void FDeferredShadingSceneRenderer::RenderLights(
 							SCOPED_GPU_MASK(RHICmdList, View.GPUMask);
 
 							// Render the light to the scene color buffer, using a 1x1 white texture as input
-							RenderLight(RHICmdList, View, LightSceneInfo, nullptr, TryGetRHI(LightingChannelsTexture), false, false);
+							RenderLight(RHICmdList, View, LightSceneInfo, PassParameters, nullptr, TryGetRHI(LightingChannelsTexture), false, false);
 						}
 					});
 				}
@@ -1959,7 +1972,7 @@ void FDeferredShadingSceneRenderer::RenderLightArrayForOverlapViewmode(
 		for (const FViewInfo& View : Views)
 		{
 			SCOPED_GPU_MASK(RHICmdList, View.GPUMask);
-			RenderLight(RHICmdList, View, LightSceneInfo, nullptr, LightingChannelsTexture, true, false);
+			RenderLight(RHICmdList, View, LightSceneInfo, nullptr, nullptr, LightingChannelsTexture, true, false);
 		}
 	}
 }
@@ -2105,6 +2118,7 @@ void FDeferredShadingSceneRenderer::RenderLight(
 	FRHICommandList& RHICmdList,
 	const FViewInfo& View,
 	const FLightSceneInfo* LightSceneInfo,
+	FRenderLightParameters* PassParameters,	// If this is null, it means we cannot use Strata tiles and fallback to previous behavior.
 	FRHITexture* ScreenShadowMaskTexture,
 	FRHITexture* LightingChannelsTexture,
 	bool bRenderOverlap, bool bIssueDrawEvent)
@@ -2119,7 +2133,7 @@ void FDeferredShadingSceneRenderer::RenderLight(
 	INC_DWORD_STAT(STAT_NumLightsUsingStandardDeferred);
 	SCOPED_CONDITIONAL_DRAW_EVENT(RHICmdList, StandardDeferredLighting, bIssueDrawEvent);
 
-	auto RenderInternalLight = [&](bool bStrataFastPath)
+	auto RenderInternalLight = [&]( bool bEnableStrataStencilTest, bool bEnableStrataTiledPass, uint32 StrataTileMaterialType)
 	{
 
 	FGraphicsPipelineStateInitializer GraphicsPSOInit;
@@ -2144,7 +2158,7 @@ void FDeferredShadingSceneRenderer::RenderLight(
 	RHICmdList.SetViewport(View.ViewRect.Min.X, View.ViewRect.Min.Y, 0.0f, View.ViewRect.Max.X, View.ViewRect.Max.Y, 1.0f);
 
 	FRenderLightParams RenderLightParams;
-	if (Strata::IsStrataEnabled() && Strata::IsClassificationEnabled())
+	if (bEnableStrataStencilTest)
 	{
 		GraphicsPSOInit.DepthStencilState = TStaticDepthStencilState<
 			false, CF_Always,
@@ -2163,6 +2177,12 @@ void FDeferredShadingSceneRenderer::RenderLight(
 		GraphicsPSOInit.bDepthBounds = false;
 		TShaderMapRef<TDeferredLightVS<false> > VertexShader(View.ShaderMap);
 
+		Strata::FStrataTilePassVS::FParameters VSParameters;
+		Strata::FStrataTilePassVS::FPermutationDomain VSPermutationVector;
+		VSPermutationVector.Set< Strata::FStrataTilePassVS::FEnableDebug >(false);
+		VSPermutationVector.Set< Strata::FStrataTilePassVS::FEnableTexCoordScreenVector >(true);
+		TShaderMapRef<Strata::FStrataTilePassVS> StrataTilePassVertexShader(View.ShaderMap, VSPermutationVector);
+
 		GraphicsPSOInit.RasterizerState = TStaticRasterizerState<FM_Solid, CM_None>::GetRHI();
 
 		if (bRenderOverlap)
@@ -2213,31 +2233,47 @@ void FDeferredShadingSceneRenderer::RenderLight(
 			// Only directional lights are rendered in this path, so we only need to check if it is use to light the atmosphere
 			PermutationVector.Set< FDeferredLightPS::FAtmosphereTransmittance >(bAtmospherePerPixelTransmittance);
 			PermutationVector.Set< FDeferredLightPS::FCloudTransmittance >(bLight0CloudPerPixelTransmittance || bLight1CloudPerPixelTransmittance);
-			PermutationVector.Set< FDeferredLightPS::FStrataFastPath >(Strata::IsStrataEnabled() && Strata::IsClassificationEnabled() && bStrataFastPath);
+			PermutationVector.Set< FDeferredLightPS::FStrataFastPath >(StrataTileMaterialType == EStrataTileMaterialType::ESimple);
 
 			TShaderMapRef< FDeferredLightPS > PixelShader( View.ShaderMap, PermutationVector );
 			GraphicsPSOInit.BoundShaderState.VertexDeclarationRHI = GFilterVertexDeclaration.VertexDeclarationRHI;
 			GraphicsPSOInit.BoundShaderState.VertexShaderRHI = VertexShader.GetVertexShader();
 			GraphicsPSOInit.BoundShaderState.PixelShaderRHI = PixelShader.GetPixelShader();
 
+			if (bEnableStrataTiledPass)
+			{
+				Strata::FillUpTiledPassData((EStrataTileMaterialType)StrataTileMaterialType, View, VSParameters, GraphicsPSOInit.PrimitiveType);
+				GraphicsPSOInit.BoundShaderState.VertexShaderRHI = StrataTilePassVertexShader.GetVertexShader();
+			}
+
 			SetGraphicsPipelineState(RHICmdList, GraphicsPSOInit);
 			PixelShader->SetParameters(RHICmdList, View, LightSceneInfo, ScreenShadowMaskTexture, LightingChannelsTexture, &RenderLightParams, nullptr);
 		}
 
-		VertexShader->SetParameters(RHICmdList, View, LightSceneInfo);
-		RHICmdList.SetStencilRef(bStrataFastPath ? Strata::StencilBit : 0u);
+		RHICmdList.SetStencilRef(StrataTileMaterialType == EStrataTileMaterialType::ESimple ? Strata::StencilBit : 0u);
 
-		// Apply the directional light as a full screen quad
-		DrawRectangle(
-			RHICmdList,
-			0, 0,
-			View.ViewRect.Width(), View.ViewRect.Height(),
-			View.ViewRect.Min.X, View.ViewRect.Min.Y,
-			View.ViewRect.Width(), View.ViewRect.Height(),
-			View.ViewRect.Size(),
-			GetSceneTextureExtent(),
-			VertexShader,
-			EDRF_UseTriangleOptimization);
+		if (!bEnableStrataTiledPass)
+		{
+			VertexShader->SetParameters(RHICmdList, View, LightSceneInfo);
+
+			// Apply the directional light as a full screen quad
+			DrawRectangle(
+				RHICmdList,
+				0, 0,
+				View.ViewRect.Width(), View.ViewRect.Height(),
+				View.ViewRect.Min.X, View.ViewRect.Min.Y,
+				View.ViewRect.Width(), View.ViewRect.Height(),
+				View.ViewRect.Size(),
+				GetSceneTextureExtent(),
+				VertexShader,
+				EDRF_UseTriangleOptimization);
+		}
+		else
+		{
+			SetShaderParameters(RHICmdList, StrataTilePassVertexShader, StrataTilePassVertexShader.GetVertexShader(), VSParameters);
+
+			RHICmdList.DrawPrimitiveIndirect(VSParameters.TileIndirectBuffer->GetIndirectRHICallBuffer(), 0);
+		}
 	}
 	else
 	{
@@ -2279,7 +2315,7 @@ void FDeferredShadingSceneRenderer::RenderLight(
 			PermutationVector.Set< FDeferredLightPS::FHairLighting>(0);
 			PermutationVector.Set < FDeferredLightPS::FAtmosphereTransmittance >(false);
 			PermutationVector.Set< FDeferredLightPS::FCloudTransmittance >(false);
-			PermutationVector.Set< FDeferredLightPS::FStrataFastPath >(Strata::IsStrataEnabled() && Strata::IsClassificationEnabled() && bStrataFastPath);
+			PermutationVector.Set< FDeferredLightPS::FStrataFastPath >(StrataTileMaterialType == EStrataTileMaterialType::ESimple);
 
 			TShaderMapRef< FDeferredLightPS > PixelShader( View.ShaderMap, PermutationVector );
 			GraphicsPSOInit.BoundShaderState.VertexDeclarationRHI = GetVertexDeclarationFVector4();
@@ -2291,7 +2327,7 @@ void FDeferredShadingSceneRenderer::RenderLight(
 		}
 
 		VertexShader->SetParameters(RHICmdList, View, LightSceneInfo);
-		RHICmdList.SetStencilRef(bStrataFastPath ? Strata::StencilBit : 0u);
+		RHICmdList.SetStencilRef(StrataTileMaterialType == EStrataTileMaterialType::ESimple ? Strata::StencilBit : 0u);
 
 		// Use DBT to allow work culling on shadow lights
 		if (GraphicsPSOInit.bDepthBounds)
@@ -2325,11 +2361,28 @@ void FDeferredShadingSceneRenderer::RenderLight(
 	}
 	};
 
+	const bool bStrataClassificationEnabled = Strata::IsStrataEnabled() && Strata::IsClassificationEnabled();
+	const bool bTilePassesReadingStrataEnabled = Strata::ShouldPassesReadingStrataBeTiled(Scene->GetFeatureLevel());
 
-	RenderInternalLight(false);
-	if (Strata::IsStrataEnabled() && Strata::IsClassificationEnabled())
+	if (bStrataClassificationEnabled && bTilePassesReadingStrataEnabled && PassParameters != nullptr)
 	{
-		RenderInternalLight(true);
+		const bool bEnableStrataTiledPass = true;
+		const bool bEnableStrataStencilTest = false;
+
+		RenderInternalLight(bEnableStrataStencilTest, bEnableStrataTiledPass, EStrataTileMaterialType::ESimple);
+		RenderInternalLight(bEnableStrataStencilTest, bEnableStrataTiledPass, EStrataTileMaterialType::EComplex);
+	}
+	else
+	{
+		const bool bEnableStrataTiledPass = false;
+		const bool bEnableStrataStencilTest = bStrataClassificationEnabled;
+
+		RenderInternalLight(bEnableStrataStencilTest, bEnableStrataTiledPass, EStrataTileMaterialType::EComplex);
+		if (Strata::IsStrataEnabled() && Strata::IsClassificationEnabled())
+		{
+			RenderInternalLight(bEnableStrataStencilTest, bEnableStrataTiledPass, EStrataTileMaterialType::ESimple);
+		}
+
 	}
 }
 
@@ -2357,12 +2410,13 @@ void FDeferredShadingSceneRenderer::RenderLight(
 			RDG_EVENT_NAME("StandardDeferredLighting"),
 			PassParameters,
 			PassFlags,
-			[this, &View, LightSceneInfo, ScreenShadowMaskTexture, LightingChannelsTexture, bRenderOverlap](FRHICommandList& RHICmdList)
+			[this, &View, LightSceneInfo, ScreenShadowMaskTexture, LightingChannelsTexture, bRenderOverlap, PassParameters](FRHICommandList& RHICmdList)
 		{
 			RenderLight(
 				RHICmdList,
 				View,
 				LightSceneInfo,
+				PassParameters,
 				TryGetRHI(ScreenShadowMaskTexture),
 				TryGetRHI(LightingChannelsTexture),
 				bRenderOverlap,
diff --git a/Engine/Source/Runtime/Renderer/Private/Strata/Strata.cpp b/Engine/Source/Runtime/Renderer/Private/Strata/Strata.cpp
--- a/Engine/Source/Runtime/Renderer/Private/Strata/Strata.cpp
+++ b/Engine/Source/Runtime/Renderer/Private/Strata/Strata.cpp
@@ -141,7 +141,7 @@ bool IsClassificationEnabled()
 
 bool ShouldPassesReadingStrataBeTiled(ERHIFeatureLevel::Type FeatureLevel)
 {
-	return FeatureLevel >= ERHIFeatureLevel::SM5 && CVarStrataClassificationPassesReadingStrataAreTiled.GetValueOnAnyThread() > 0;
+	return IsStrataEnabled() && IsClassificationEnabled() && FeatureLevel >= ERHIFeatureLevel::SM5 && CVarStrataClassificationPassesReadingStrataAreTiled.GetValueOnAnyThread() > 0;
 }
 
 uint32 GetStrataBufferTileSize()
@@ -535,8 +535,10 @@ void FillUpTiledPassData(
 	FStrataTilePassVS::FParameters& ParametersVS,
 	EPrimitiveType& PrimitiveType)
 {
-	const FIntPoint OutputResolution = View.ViewRect.Size();
-	ParametersVS.OutputResolutionAndInv = FVector4(OutputResolution.X, OutputResolution.Y, 1.0f / float(OutputResolution.X), 1.0f / float(OutputResolution.Y));
+	ParametersVS.OutputViewSizeAndInvSize = View.CachedViewUniformShaderParameters->ViewSizeAndInvSize;
+	ParametersVS.OutputBufferSizeAndInvSize = View.CachedViewUniformShaderParameters->BufferSizeAndInvSize;
+	ParametersVS.ViewScreenToTranslatedWorld = View.CachedViewUniformShaderParameters->ScreenToTranslatedWorld;
+
 	ParametersVS.TileListBuffer = View.StrataSceneData->ClassificationTileListBufferSRV[Type];
 	ParametersVS.TileIndirectBuffer = View.StrataSceneData->ClassificationTileIndirectBuffer[Type];
 
@@ -564,7 +566,10 @@ static void AddStrataInternalClassificationTilePass(
 	FStrataMaterialStencilTaggingPassPS::FParameters* ParametersPS = GraphBuilder.AllocParameters<FStrataMaterialStencilTaggingPassPS::FParameters>();
 	FillUpTiledPassData(TileMaterialType, View, ParametersPS->VS, StrataTilePrimitiveType);
 
-	TShaderMapRef<FStrataTilePassVS> VertexShader(View.ShaderMap);
+	FStrataTilePassVS::FPermutationDomain VSPermutationVector;
+	VSPermutationVector.Set< FStrataTilePassVS::FEnableDebug >(true);
+	VSPermutationVector.Set< FStrataTilePassVS::FEnableTexCoordScreenVector >(false);
+	TShaderMapRef<FStrataTilePassVS> VertexShader(View.ShaderMap, VSPermutationVector);
 	TShaderMapRef<FStrataMaterialStencilTaggingPassPS> PixelShader(View.ShaderMap);
 
 	// For debug purpose
diff --git a/Engine/Source/Runtime/Renderer/Private/Strata/Strata.h b/Engine/Source/Runtime/Renderer/Private/Strata/Strata.h
--- a/Engine/Source/Runtime/Renderer/Private/Strata/Strata.h
+++ b/Engine/Source/Runtime/Renderer/Private/Strata/Strata.h
@@ -102,8 +102,17 @@ class FStrataTilePassVS : public FGlobalShader
 	DECLARE_GLOBAL_SHADER(FStrataTilePassVS);
 	SHADER_USE_PARAMETER_STRUCT(FStrataTilePassVS, FGlobalShader);
 
+	class FEnableDebug : SHADER_PERMUTATION_BOOL("PERMUTATION_ENABLE_DEBUG");
+	class FEnableTexCoordScreenVector : SHADER_PERMUTATION_BOOL("PERMUTATION_ENABLE_TEXCOORD_SCREENVECTOR");
+	using FPermutationDomain = TShaderPermutationDomain<FEnableDebug, FEnableTexCoordScreenVector>;
+
 	BEGIN_SHADER_PARAMETER_STRUCT(FParameters, )
-		SHADER_PARAMETER(FVector4, OutputResolutionAndInv)
+		// It would be possible to use the ViewUniformBuffer instead of copying the data here, 
+		// but we would have to make sure the view UB is added to all passes using this parameter structure.
+		// We should not add it here to now have duplicated input UB.
+		SHADER_PARAMETER(FVector4, OutputViewSizeAndInvSize)
+		SHADER_PARAMETER(FVector4, OutputBufferSizeAndInvSize)
+		SHADER_PARAMETER(FMatrix, ViewScreenToTranslatedWorld)
 		SHADER_PARAMETER_RDG_BUFFER_SRV(Buffer<uint>, TileListBuffer)
 		RDG_BUFFER_ACCESS(TileIndirectBuffer, ERHIAccess::IndirectArgs)
 	END_SHADER_PARAMETER_STRUCT()
