Uses CFLAG_Wave32 on FTAAUpdateHistoryCS

diff --git a/Engine/Shaders/Private/TemporalAA/TAAUpdateHistory.usf b/Engine/Shaders/Private/TemporalAA/TAAUpdateHistory.usf
--- a/Engine/Shaders/Private/TemporalAA/TAAUpdateHistory.usf
+++ b/Engine/Shaders/Private/TemporalAA/TAAUpdateHistory.usf
@@ -1,7 +1,7 @@
 // Copyright Epic Games, Inc. All Rights Reserved.
 
 #include "TAACommon.ush"
-#include "../Random.ush"
+
 
 
 //------------------------------------------------------- CONFIG
diff --git a/Engine/Source/Runtime/Renderer/Private/PostProcess/TemporalAA.cpp b/Engine/Source/Runtime/Renderer/Private/PostProcess/TemporalAA.cpp
--- a/Engine/Source/Runtime/Renderer/Private/PostProcess/TemporalAA.cpp
+++ b/Engine/Source/Runtime/Renderer/Private/PostProcess/TemporalAA.cpp
@@ -553,6 +553,12 @@ class FTAAUpdateHistoryCS : public FTAAGen5Shader
 		SHADER_PARAMETER_RDG_TEXTURE_UAV(RWTexture2D, SceneColorOutput)
 		SHADER_PARAMETER_RDG_TEXTURE_UAV(RWTexture2D, DebugOutput)
 	END_SHADER_PARAMETER_STRUCT()
+
+	static void ModifyCompilationEnvironment(const FGlobalShaderPermutationParameters& Parameters, FShaderCompilerEnvironment& OutEnvironment)
+	{
+		FTAAGen5Shader::ModifyCompilationEnvironment(Parameters, OutEnvironment);
+		OutEnvironment.CompilerFlags.Add(CFLAG_Wave32);
+	}
 }; // class FTAAUpdateHistoryCS
 
 IMPLEMENT_GLOBAL_SHADER(FTAAStandaloneCS, "/Engine/Private/TemporalAA/TAAStandalone.usf", "MainCS", SF_Compute);
