Fixed r.MeshDrawCommands.LogDynamicInstancingStats 1 crash and secured a bit from future breakages

diff --git a/Engine/Source/Runtime/Renderer/Public/MeshPassProcessor.h b/Engine/Source/Runtime/Renderer/Public/MeshPassProcessor.h
--- a/Engine/Source/Runtime/Renderer/Public/MeshPassProcessor.h
+++ b/Engine/Source/Runtime/Renderer/Public/MeshPassProcessor.h
@@ -86,6 +86,7 @@ inline const TCHAR* GetMeshPassName(EMeshPass::Type MeshPass)
 	case EMeshPass::MobileInverseOpacity: return TEXT("MobileInverseOpacity");
 	case EMeshPass::VirtualTexture: return TEXT("VirtualTexture");
 	case EMeshPass::LumenCardCapture: return TEXT("LumenCardCapture");
+	case EMeshPass::DitheredLODFadingOutMaskPass: return TEXT("DitheredLODFadingOutMaskPass");
 #if WITH_EDITOR
 	case EMeshPass::HitProxy: return TEXT("HitProxy");
 	case EMeshPass::HitProxyOpaqueOnly: return TEXT("HitProxyOpaqueOnly");
@@ -94,6 +95,12 @@ inline const TCHAR* GetMeshPassName(EMeshPass::Type MeshPass)
 #endif
 	}
 
+#if WITH_EDITOR
+	static_assert(EMeshPass::Num == 22 + 4, "Need to update switch(MeshPass) after changing EMeshPass");
+#else
+	static_assert(EMeshPass::Num == 22, "Need to update switch(MeshPass) after changing EMeshPass");
+#endif
+
 	checkf(0, TEXT("Missing case for EMeshPass %u"), (uint32)MeshPass);
 	return nullptr;
 }
