Fix rendering of hair shadow mask twice when running with virtual shadow map.

diff --git a/Engine/Source/Runtime/Renderer/Private/ShadowRendering.cpp b/Engine/Source/Runtime/Renderer/Private/ShadowRendering.cpp
--- a/Engine/Source/Runtime/Renderer/Private/ShadowRendering.cpp
+++ b/Engine/Source/Runtime/Renderer/Private/ShadowRendering.cpp
@@ -2217,10 +2217,36 @@ void FDeferredShadingSceneRenderer::RenderDeferredShadowProjections(
 		}
 	}
 
-	// Inject deep shadow mask
+	// Inject deep shadow mask for regular shadow map. When using virtual shadow map, it is directly handled in the shadow kernel.
 	if (HairStrands::HasViewHairStrandsData(Views))
 	{
-		RenderHairStrandsShadowMask(GraphBuilder, Views, LightSceneInfo, ScreenShadowMaskTexture);
+		bool bNeedHairShadowMaskPass = false;
+		const bool bVirtualShadowOnePass = VisibleLightInfo.VirtualShadowMapClipmaps.Num() > 0;
+		if (!bVirtualShadowOnePass)
+		{
+			for (int32 ShadowIndex = 0; ShadowIndex < VisibleLightInfo.ShadowsToProject.Num(); ShadowIndex++)
+			{
+				FProjectedShadowInfo* ProjectedShadowInfo = VisibleLightInfo.ShadowsToProject[ShadowIndex];
+				if (!ProjectedShadowInfo->bIncludeInScreenSpaceShadowMask)
+				{
+					continue;
+				}
+				if (ProjectedShadowInfo->HasVirtualShadowMap())
+				{
+					bNeedHairShadowMaskPass = false;
+					break;
+				}
+				else
+				{
+					bNeedHairShadowMaskPass = true;
+					break;
+				}
+			}
+		}
+		if (bNeedHairShadowMaskPass)
+		{
+			RenderHairStrandsShadowMask(GraphBuilder, Views, LightSceneInfo, ScreenShadowMaskTexture);
+		}
 	}
 }
 
