Fix assert on startup

diff --git a/Engine/Source/Runtime/Renderer/Private/VirtualShadowMaps/VirtualShadowMapArray.cpp b/Engine/Source/Runtime/Renderer/Private/VirtualShadowMaps/VirtualShadowMapArray.cpp
--- a/Engine/Source/Runtime/Renderer/Private/VirtualShadowMaps/VirtualShadowMapArray.cpp
+++ b/Engine/Source/Runtime/Renderer/Private/VirtualShadowMaps/VirtualShadowMapArray.cpp
@@ -142,7 +142,6 @@ static TAutoConsoleVariable<int32> CVarShowClipmapStats(
 
 int32 GEnableNonNaniteVSM = 0;
 
-#if GPUCULL_TODO
 FAutoConsoleVariableRef CVarEnableNonNaniteVSM(
 	TEXT("r.Shadow.Virtual.NonNaniteVSM"),
 	GEnableNonNaniteVSM,
@@ -150,7 +149,6 @@ FAutoConsoleVariableRef CVarEnableNonNaniteVSM(
 	TEXT("Read-only and to be set in a config file (requires restart)."),
 	ECVF_RenderThreadSafe | ECVF_ReadOnly
 );
-#endif // GPUCULL_TODO
 
 FMatrix CalcTranslatedWorldToShadowUVMatrix(
 	const FMatrix& TranslatedWorldToShadowView,
