Fix crash when enabling hair strands fast path without GBufferE being allocated.

diff --git a/Engine/Source/Runtime/Renderer/Private/HairStrands/HairStrandsVisibility.cpp b/Engine/Source/Runtime/Renderer/Private/HairStrands/HairStrandsVisibility.cpp
--- a/Engine/Source/Runtime/Renderer/Private/HairStrands/HairStrandsVisibility.cpp
+++ b/Engine/Source/Runtime/Renderer/Private/HairStrands/HairStrandsVisibility.cpp
@@ -696,11 +696,18 @@ static void AddHairMaterialGBufferPass(
 		OutVelocityTexture = GraphBuilder.CreateTexture(VelocityDesc, TEXT("Hair.DummyVelocity"));
 		bIsVelocityDummy = true;
 	}
+	bool bIsBufferEDummy = false;
+	if (OutBufferETexture == nullptr)
+	{
+		FRDGTextureDesc Desc = FVelocityRendering::GetRenderTargetDesc(ViewInfo->GetShaderPlatform(), OutBufferDTexture->Desc.Extent);
+		OutBufferETexture = GraphBuilder.CreateTexture(Desc, TEXT("Hair.DummyGBufferE"));
+		bIsBufferEDummy = true;
+	}
 	PassParameters->RenderTargets[0] = FRenderTargetBinding(OutBufferATexture, ERenderTargetLoadAction::ELoad, 0);
 	PassParameters->RenderTargets[1] = FRenderTargetBinding(OutBufferBTexture, ERenderTargetLoadAction::ELoad, 0);
 	PassParameters->RenderTargets[2] = FRenderTargetBinding(OutBufferCTexture, ERenderTargetLoadAction::ELoad, 0);
 	PassParameters->RenderTargets[3] = FRenderTargetBinding(OutBufferDTexture, ERenderTargetLoadAction::ELoad, 0);
-	PassParameters->RenderTargets[4] = FRenderTargetBinding(OutBufferETexture, ERenderTargetLoadAction::ELoad, 0);
+	PassParameters->RenderTargets[4] = FRenderTargetBinding(OutBufferETexture, bIsBufferEDummy ? ERenderTargetLoadAction::ENoAction : ERenderTargetLoadAction::ELoad, 0);
 	PassParameters->RenderTargets[5] = FRenderTargetBinding(OutVelocityTexture, bIsVelocityDummy ? ERenderTargetLoadAction::ENoAction : ERenderTargetLoadAction::ELoad, 0);
 	PassParameters->RenderTargets[6] = FRenderTargetBinding(OutColorTexture, ERenderTargetLoadAction::ELoad, 0);
 	PassParameters->RenderTargets.DepthStencil = FDepthStencilBinding(
