Use frame based lifetime management for primitive occlusion queries. This allows using raw pointers to queries in primitive occlusion history, which reduces cache misses and memory allocation.

diff --git a/Engine/Source/Runtime/Renderer/Private/ScenePrivate.h b/Engine/Source/Runtime/Renderer/Private/ScenePrivate.h
--- a/Engine/Source/Runtime/Renderer/Private/ScenePrivate.h
+++ b/Engine/Source/Runtime/Renderer/Private/ScenePrivate.h
@@ -2612,6 +2612,10 @@ public:
 	/**  Packed bit array of primitives that are always visible. */
 	TBitArray<> PrimitivesAlwaysVisible;
 
+#if WITH_EDITOR
+	TBitArray<> PrimitivesSelected;
+#endif
+
 	TBitArray<> PrimitivesNeedingStaticMeshUpdate;
 	TSet<FPrimitiveSceneInfo*> PrimitivesNeedingStaticMeshUpdateWithoutVisibilityCheck;
 
