Added PrimitivesSelected TBitArray to FScene to optimize Scene visibility in editor.

diff --git a/Engine/Source/Runtime/Engine/Private/Components/PrimitiveComponent.cpp b/Engine/Source/Runtime/Engine/Private/Components/PrimitiveComponent.cpp
--- a/Engine/Source/Runtime/Engine/Private/Components/PrimitiveComponent.cpp
+++ b/Engine/Source/Runtime/Engine/Private/Components/PrimitiveComponent.cpp
@@ -1562,7 +1562,8 @@ void UPrimitiveComponent::PushSelectionToProxy()
 	//although this should only be called for attached components, some billboard components can get in without valid proxies
 	if (SceneProxy)
 	{
-		SceneProxy->SetSelection_GameThread(ShouldRenderSelected(),IsComponentIndividuallySelected());
+		SceneProxy->SetSelection_GameThread(ShouldRenderSelected(), IsComponentIndividuallySelected());
+		MarkRenderStateDirty();
 	}
 }
 
diff --git a/Engine/Source/Runtime/Renderer/Private/RendererScene.cpp b/Engine/Source/Runtime/Renderer/Private/RendererScene.cpp
--- a/Engine/Source/Runtime/Renderer/Private/RendererScene.cpp
+++ b/Engine/Source/Runtime/Renderer/Private/RendererScene.cpp
@@ -631,6 +631,9 @@ void FScene::CheckPrimitiveArrays(int MaxTypeOffsetIndex)
 	check(Primitives.Num() == PrimitiveVirtualTextureLod.Num());
 	check(Primitives.Num() == PrimitiveOcclusionBounds.Num());
 	check(Primitives.Num() == PrimitivesAlwaysVisible.Num());
+#if WITH_EDITOR
+	check(Primitives.Num() == PrimitivesSelected.Num());
+#endif
 	check(Primitives.Num() == PrimitivesNeedingStaticMeshUpdate.Num());
 
 #if UE_BUILD_DEBUG
@@ -3846,6 +3849,9 @@ void FScene::UpdateAllPrimitiveSceneInfos(FRDGBuilder& GraphBuilder, bool bAsync
 							TArraySwapElements(PrimitiveVirtualTextureLod, DestIndex, SourceIndex);
 							TArraySwapElements(PrimitiveOcclusionBounds, DestIndex, SourceIndex);
 							TBitArraySwapElements(PrimitivesAlwaysVisible, DestIndex, SourceIndex);
+						#if WITH_EDITOR
+							TBitArraySwapElements(PrimitivesSelected, DestIndex, SourceIndex);
+						#endif
 							TBitArraySwapElements(PrimitivesNeedingStaticMeshUpdate, DestIndex, SourceIndex);
 
 							GPUScene.AddPrimitiveToUpdate(SourceIndex);
@@ -3895,6 +3901,9 @@ void FScene::UpdateAllPrimitiveSceneInfos(FRDGBuilder& GraphBuilder, bool bAsync
 				PrimitiveVirtualTextureLod.Pop();
 				PrimitiveOcclusionBounds.Pop();
 				PrimitivesAlwaysVisible.RemoveAt(PrimitivesAlwaysVisible.Num() - 1);
+			#if WITH_EDITOR
+				PrimitivesSelected.RemoveAt(PrimitivesSelected.Num() - 1);
+			#endif
 				PrimitivesNeedingStaticMeshUpdate.RemoveAt(PrimitivesNeedingStaticMeshUpdate.Num() - 1);
 			}
 
@@ -3954,6 +3963,9 @@ void FScene::UpdateAllPrimitiveSceneInfos(FRDGBuilder& GraphBuilder, bool bAsync
 			PrimitiveVirtualTextureLod.Reserve(PrimitiveVirtualTextureLod.Num() + AddedLocalPrimitiveSceneInfos.Num());
 			PrimitiveOcclusionBounds.Reserve(PrimitiveOcclusionBounds.Num() + AddedLocalPrimitiveSceneInfos.Num());
 			PrimitivesAlwaysVisible.Reserve(PrimitivesAlwaysVisible.Num() + AddedLocalPrimitiveSceneInfos.Num());
+		#if WITH_EDITOR
+			PrimitivesSelected.Reserve(PrimitivesSelected.Num() + AddedLocalPrimitiveSceneInfos.Num());
+		#endif
 			PrimitivesNeedingStaticMeshUpdate.Reserve(PrimitivesNeedingStaticMeshUpdate.Num() + AddedLocalPrimitiveSceneInfos.Num());
 		}
 
@@ -3983,6 +3995,9 @@ void FScene::UpdateAllPrimitiveSceneInfos(FRDGBuilder& GraphBuilder, bool bAsync
 				PrimitiveVirtualTextureLod.AddUninitialized();
 				PrimitiveOcclusionBounds.AddUninitialized();
 				PrimitivesAlwaysVisible.Add(PrimitiveSceneInfo->Proxy->IsAlwaysVisible());
+			#if WITH_EDITOR
+				PrimitivesSelected.Add(PrimitiveSceneInfo->Proxy->IsSelected());
+			#endif
 				PrimitivesNeedingStaticMeshUpdate.Add(false);
 
 				const int32 SourceIndex = PrimitiveSceneProxies.Num() - 1;
@@ -4063,6 +4078,9 @@ void FScene::UpdateAllPrimitiveSceneInfos(FRDGBuilder& GraphBuilder, bool bAsync
 							TArraySwapElements(PrimitiveVirtualTextureLod, DestIndex, SourceIndex);
 							TArraySwapElements(PrimitiveOcclusionBounds, DestIndex, SourceIndex);
 							TBitArraySwapElements(PrimitivesAlwaysVisible, DestIndex, SourceIndex);
+						#if WITH_EDITOR
+							TBitArraySwapElements(PrimitivesSelected, DestIndex, SourceIndex);
+						#endif
 							TBitArraySwapElements(PrimitivesNeedingStaticMeshUpdate, DestIndex, SourceIndex);
 
 							GPUScene.AddPrimitiveToUpdate(DestIndex);
diff --git a/Engine/Source/Runtime/Renderer/Private/ScenePrivate.h b/Engine/Source/Runtime/Renderer/Private/ScenePrivate.h
--- a/Engine/Source/Runtime/Renderer/Private/ScenePrivate.h
+++ b/Engine/Source/Runtime/Renderer/Private/ScenePrivate.h
@@ -2613,6 +2613,7 @@ public:
 	TBitArray<> PrimitivesAlwaysVisible;
 
 #if WITH_EDITOR
+	/** Packed bit array of primitives which are selected in the editor. */
 	TBitArray<> PrimitivesSelected;
 #endif
 
diff --git a/Engine/Source/Runtime/Renderer/Private/SceneVisibility.cpp b/Engine/Source/Runtime/Renderer/Private/SceneVisibility.cpp
--- a/Engine/Source/Runtime/Renderer/Private/SceneVisibility.cpp
+++ b/Engine/Source/Runtime/Renderer/Private/SceneVisibility.cpp
@@ -917,16 +917,17 @@ static void FetchVisibilityForPrimitives_Range(FVisForPrimParams& Params, FGloba
 		//we can't allow the prim history insertion array to realloc or it will invalidate pointers in the other output arrays.
 		const bool bCanAllocPrimHistory = bSingleThreaded || InsertPrimitiveOcclusionHistory->Num() < InsertPrimitiveOcclusionHistory->Max();		
 
+	#if WITH_EDITOR
 		if (GIsEditor)
 		{
-			FPrimitiveSceneInfo* PrimitiveSceneInfo = Scene->Primitives[BitIt.GetIndex()];
-
-			if (PrimitiveSceneInfo->Proxy->IsSelected())
+			if (Scene->PrimitivesSelected[BitIt.GetIndex()])
 			{
 				// to render occluded outline for selected objects
 				bCanBeOccluded = false;
 			}
 		}
+	#endif
+
 		int32 NumSubQueries = 1;
 		bool bSubQueries = false;
 		const TArray<FBoxSphereBounds>* SubBounds = nullptr;
@@ -3167,42 +3168,41 @@ void FSceneRenderer::PreVisibilityFrameSetup(FRDGBuilder& GraphBuilder, const FS
 		FXSystem->PreInitViews(GraphBuilder, Views[0].AllowGPUParticleUpdate() && !ViewFamily.EngineShowFlags.HitProxies);
 	}
 
+#if WITH_EDITOR
 	// Draw lines to lights affecting this mesh if its selected.
 	if (ViewFamily.EngineShowFlags.LightInfluences)
 	{
-		for (TArray<FPrimitiveSceneInfo*>::TConstIterator It(Scene->Primitives); It; ++It)
+		for (TConstSetBitIterator<> It(Scene->PrimitivesSelected); It; ++It)
 		{
-			const FPrimitiveSceneInfo* PrimitiveSceneInfo = *It;
-			if (PrimitiveSceneInfo->Proxy->IsSelected())
+			const FPrimitiveSceneInfo* PrimitiveSceneInfo = Scene->Primitives[It.GetIndex()];
+			FLightPrimitiveInteraction *LightList = PrimitiveSceneInfo->LightList;
+			while (LightList)
 			{
-				FLightPrimitiveInteraction *LightList = PrimitiveSceneInfo->LightList;
-				while (LightList)
-				{
-					const FLightSceneInfo* LightSceneInfo = LightList->GetLight();
+				const FLightSceneInfo* LightSceneInfo = LightList->GetLight();
 
-					bool bDynamic = true;
-					bool bRelevant = false;
-					bool bLightMapped = true;
-					bool bShadowMapped = false;
-					PrimitiveSceneInfo->Proxy->GetLightRelevance(LightSceneInfo->Proxy, bDynamic, bRelevant, bLightMapped, bShadowMapped);
+				bool bDynamic = true;
+				bool bRelevant = false;
+				bool bLightMapped = true;
+				bool bShadowMapped = false;
+				PrimitiveSceneInfo->Proxy->GetLightRelevance(LightSceneInfo->Proxy, bDynamic, bRelevant, bLightMapped, bShadowMapped);
 
-					if (bRelevant)
+				if (bRelevant)
+				{
+					// Draw blue for light-mapped lights and orange for dynamic lights
+					const FColor LineColor = bLightMapped ? FColor(0,140,255) : FColor(255,140,0);
+					for (int32 ViewIndex = 0;ViewIndex < Views.Num();ViewIndex++)
 					{
-						// Draw blue for light-mapped lights and orange for dynamic lights
-						const FColor LineColor = bLightMapped ? FColor(0,140,255) : FColor(255,140,0);
-						for (int32 ViewIndex = 0;ViewIndex < Views.Num();ViewIndex++)
-						{
-							FViewInfo& View = Views[ViewIndex];
-							FViewElementPDI LightInfluencesPDI(&View,nullptr,&View.DynamicPrimitiveCollector);
-							LightInfluencesPDI.DrawLine(PrimitiveSceneInfo->Proxy->GetBounds().Origin, LightSceneInfo->Proxy->GetLightToWorld().GetOrigin(), LineColor, SDPG_World);
-						}
+						FViewInfo& View = Views[ViewIndex];
+						FViewElementPDI LightInfluencesPDI(&View,nullptr,&View.DynamicPrimitiveCollector);
+						LightInfluencesPDI.DrawLine(PrimitiveSceneInfo->Proxy->GetBounds().Origin, LightSceneInfo->Proxy->GetLightToWorld().GetOrigin(), LineColor, SDPG_World);
 					}
-					LightList = LightList->GetNextLight();
 				}
+				LightList = LightList->GetNextLight();
 			}
 		}
 	}
-	
+#endif
+
 #if UE_BUILD_SHIPPING
 	const bool bFreezeTemporalHistories = false;
 	const bool bFreezeTemporalSequences = false;
