Various VT fixes, mostly related to UDIMs
- Fix how local mip bias is applied
- Fix a few inconsistencies with how various vLevels are computed
- Rename some values to try to be more consistent between vLevel, Local_vLevel, Mapping_vLevel, vLogSize, etc

diff --git a/Engine/Source/Runtime/Renderer/Private/VT/TexturePageMap.cpp b/Engine/Source/Runtime/Renderer/Private/VT/TexturePageMap.cpp
--- a/Engine/Source/Runtime/Renderer/Private/VT/TexturePageMap.cpp
+++ b/Engine/Source/Runtime/Renderer/Private/VT/TexturePageMap.cpp
@@ -105,7 +105,7 @@ uint32 FTexturePageMap::FindNearestPageLevel(uint8 vLogSize, uint32 vAddress) co
 	const uint32 Index = FindNearestPageIndex(vLogSize, vAddress);
 	if (Index != ~0u)
 	{
-		return Pages[Index].vLevel;
+		return Pages[Index].Local_vLevel;
 	}
 	return 0xff;
 }
@@ -133,7 +133,7 @@ void FTexturePageMap::UnmapPage(FVirtualTextureSystem* System, FVirtualTextureSp
 			// No reason to queue a page table update to invalid pAddress, just leave it alone for now, it will be updated when the page is remapped
 			check(Pages[AncestorIndex].PhysicalSpaceID == Pages[PageIndex].PhysicalSpaceID);
 			const FVirtualTexturePhysicalSpace* PhysicalSpace = System->GetPhysicalSpace(Pages[AncestorIndex].PhysicalSpaceID);
-			const uint8 Ancestor_vLevel = Pages[AncestorIndex].vLevel;
+			const uint8 Ancestor_vLevel = Pages[AncestorIndex].Local_vLevel;
 			Space->QueueUpdate(LayerIndex, vLogSize, vAddress, Ancestor_vLevel, PhysicalSpace->GetPhysicalLocation(Pages[AncestorIndex].pAddress));
 		}
 	}
@@ -170,7 +170,7 @@ void FTexturePageMap::UnmapPage(FVirtualTextureSystem* System, FVirtualTextureSp
 	SortedKeysDirty = true;
 }
 
-void FTexturePageMap::MapPage(FVirtualTextureSpace* Space, FVirtualTexturePhysicalSpace* PhysicalSpace, uint8 vLogSize, uint32 vAddress, uint8 vLevel, uint16 pAddress)
+void FTexturePageMap::MapPage(FVirtualTextureSpace* Space, FVirtualTexturePhysicalSpace* PhysicalSpace, uint8 vLogSize, uint32 vAddress, uint8 Local_vLevel, uint16 pAddress)
 {
 #if DO_GUARD_SLOW
 	const uint32 PrevPageIndex = FindPageIndex(vLogSize, vAddress);
@@ -183,7 +183,7 @@ void FTexturePageMap::MapPage(FVirtualTextureSpace* Space, FVirtualTexturePhysic
 	FPageEntry& Entry = Pages[PageIndex];
 	Entry.Page = Page;
 	Entry.pAddress = pAddress;
-	Entry.vLevel = vLevel;
+	Entry.Local_vLevel = Local_vLevel;
 	Entry.PhysicalSpaceID = PhysicalSpace->GetID();
 
 	++MappedPageCount;
@@ -197,7 +197,7 @@ void FTexturePageMap::MapPage(FVirtualTextureSpace* Space, FVirtualTexturePhysic
 		// Map new page
 		const uint16 Hash = MurmurFinalize32(Page.Packed);
 		HashTable.Add(Hash, PageIndex);
-		Space->QueueUpdate(LayerIndex, vLogSize, vAddress, vLevel, PhysicalSpace->GetPhysicalLocation(pAddress));
+		Space->QueueUpdate(LayerIndex, vLogSize, vAddress, Local_vLevel, PhysicalSpace->GetPhysicalLocation(pAddress));
 	}
 
 	SortedKeysDirty = true;
@@ -228,7 +228,7 @@ void FTexturePageMap::GetMappedPagesInRange(uint32 vAddress, uint32 Width, uint3
 				OutPage.Page = Entry.Page;
 				OutPage.pAddress = Entry.pAddress;
 				OutPage.PhysicalSpaceID = Entry.PhysicalSpaceID;
-				OutPage.vLevel = Entry.vLevel;
+				OutPage.Local_vLevel = Entry.Local_vLevel;
 			}
 		}
 
diff --git a/Engine/Source/Runtime/Renderer/Private/VT/TexturePagePool.cpp b/Engine/Source/Runtime/Renderer/Private/VT/TexturePagePool.cpp
--- a/Engine/Source/Runtime/Renderer/Private/VT/TexturePagePool.cpp
+++ b/Engine/Source/Runtime/Renderer/Private/VT/TexturePagePool.cpp
@@ -326,11 +326,11 @@ uint32 FTexturePagePool::GetNumVisiblePages(uint32 Frame) const
 	return Count;
 }
 
-void FTexturePagePool::MapPage(FVirtualTextureSpace* Space, FVirtualTexturePhysicalSpace* PhysicalSpace, uint8 PageTableLayerIndex, uint8 vLogSize, uint32 vAddress, uint8 vLevel, uint16 pAddress)
+void FTexturePagePool::MapPage(FVirtualTextureSpace* Space, FVirtualTexturePhysicalSpace* PhysicalSpace, uint8 PageTableLayerIndex, uint8 vLogSize, uint32 vAddress, uint8 Local_vLevel, uint16 pAddress)
 {
 	check(pAddress < NumPages);
 	FTexturePageMap& PageMap = Space->GetPageMapForPageTableLayer(PageTableLayerIndex);
-	PageMap.MapPage(Space, PhysicalSpace, vLogSize, vAddress, vLevel, pAddress);
+	PageMap.MapPage(Space, PhysicalSpace, vLogSize, vAddress, Local_vLevel, pAddress);
 
 	++NumPagesMapped;
 
@@ -444,3 +444,4 @@ void FTexturePagePool::RemapPages(FVirtualTextureSystem* System, uint8 SpaceID,
 		}
 	}
 }
+
diff --git a/Engine/Source/Runtime/Renderer/Private/VT/UniqueRequestList.h b/Engine/Source/Runtime/Renderer/Private/VT/UniqueRequestList.h
--- a/Engine/Source/Runtime/Renderer/Private/VT/UniqueRequestList.h
+++ b/Engine/Source/Runtime/Renderer/Private/VT/UniqueRequestList.h
@@ -36,8 +36,8 @@ inline bool operator!=(const FMappingRequest& Lhs, const FMappingRequest& Rhs) {
 union FDirectMappingRequest
 {
 	inline FDirectMappingRequest() {}
-	inline FDirectMappingRequest(uint8 InSpaceID, uint16 InPhysicalSpaceID, uint8 InPageTableLayerIndex, uint8 InLogSize, uint32 InAddress, uint8 InLevel, uint16 InPhysicalAddress)
-		: vAddress(InAddress), vLevel(InLevel), SpaceID(InSpaceID), pAddress(InPhysicalAddress), PhysicalSpaceID(InPhysicalSpaceID), Local_vLevel(InLogSize), PageTableLayerIndex(InPageTableLayerIndex)
+	inline FDirectMappingRequest(uint8 InSpaceID, uint16 InPhysicalSpaceID, uint8 InPageTableLayerIndex, uint32 InAddress, uint8 InLevel, uint8 InLocalLevel, uint16 InPhysicalAddress)
+		: vAddress(InAddress), vLevel(InLevel), SpaceID(InSpaceID), pAddress(InPhysicalAddress), PhysicalSpaceID(InPhysicalSpaceID), Local_vLevel(InLocalLevel), PageTableLayerIndex(InPageTableLayerIndex)
 	{}
 
 	uint64 PackedValue;
@@ -109,7 +109,7 @@ public:
 
 	void AddMappingRequest(uint16 LoadRequestIndex, uint8 ProducerPhysicalGroupIndex, uint8 SpaceID, uint8 PageTableLayerIndex, uint32 vAddress, uint8 vLevel, uint8 Local_vLevel);
 
-	void AddDirectMappingRequest(uint8 InSpaceID, uint16 InPhysicalSpaceID, uint8 InPageTableLayerIndex, uint8 InLogSize, uint32 InAddress, uint8 InLevel, uint16 InPhysicalAddress);
+	void AddDirectMappingRequest(uint8 InSpaceID, uint16 InPhysicalSpaceID, uint8 InPageTableLayerIndex, uint32 InAddress, uint8 InLevel, uint8 InLocalLevel, uint16 InPhysicalAddress);
 	void AddDirectMappingRequest(const FDirectMappingRequest& Request);
 
 	void AddContinuousUpdateRequest(const FVirtualTextureLocalTile& Request);
@@ -237,9 +237,9 @@ inline void FUniqueRequestList::AddMappingRequest(uint16 LoadRequestIndex, uint8
 	}
 }
 
-inline void FUniqueRequestList::AddDirectMappingRequest(uint8 InSpaceID, uint16 InPhysicalSpaceID, uint8 InPageTableLayerIndex, uint8 InLogSize, uint32 InAddress, uint8 InLevel, uint16 InPhysicalAddress)
+inline void FUniqueRequestList::AddDirectMappingRequest(uint8 InSpaceID, uint16 InPhysicalSpaceID, uint8 InPageTableLayerIndex, uint32 InAddress, uint8 InLevel, uint8 InLocalLevel, uint16 InPhysicalAddress)
 {
-	const FDirectMappingRequest Request(InSpaceID, InPhysicalSpaceID, InPageTableLayerIndex, InLogSize, InAddress, InLevel, InPhysicalAddress);
+	const FDirectMappingRequest Request(InSpaceID, InPhysicalSpaceID, InPageTableLayerIndex, InAddress, InLevel, InLocalLevel, InPhysicalAddress);
 	AddDirectMappingRequest(Request);
 }
 
diff --git a/Engine/Source/Runtime/Renderer/Private/VT/VirtualTextureSystem.cpp b/Engine/Source/Runtime/Renderer/Private/VT/VirtualTextureSystem.cpp
--- a/Engine/Source/Runtime/Renderer/Private/VT/VirtualTextureSystem.cpp
+++ b/Engine/Source/Runtime/Renderer/Private/VT/VirtualTextureSystem.cpp
@@ -1597,9 +1597,8 @@ void FVirtualTextureSystem::GatherRequestsTask(const FGatherRequestsParameters&
 			uint32 Mapping_vLevel = FMath::Max(vLevel, ProducerMipBias);
 
 			// Local_vLevel is the level within the producer that we want to allocate/map
-			// here we subtract ProducerMipBias (clamped to ensure we don't fall below 0),
-			// which effectively matches more detailed mips of lower resolution producers with less detailed mips of higher resolution producers
-			uint32 Local_vLevel = vLevel - FMath::Min(vLevel, ProducerMipBias);
+			// here we subtract ProducerMipBias, which effectively matches more detailed mips of lower resolution producers with less detailed mips of higher resolution producers
+			uint32 Local_vLevel = Mapping_vLevel - ProducerMipBias;
 
 			// Wrap vAddress for the given producer
 			uint32 Wrapped_vAddress = vAddress;
@@ -1631,8 +1630,8 @@ void FVirtualTextureSystem::GatherRequestsTask(const FGatherRequestsParameters&
 			if (LocalMipBias > 0u)
 			{
 				Local_vLevel += LocalMipBias;
+				Mapping_vLevel += LocalMipBias;
 				Local_vAddress >>= (LocalMipBias * vDimensions);
-				Mapping_vLevel = FMath::Max(vLevel, LocalMipBias + ProducerMipBias);
 			}
 
 			uint8 ProducerPhysicalGroupMaskToPrefetchForLevel[16] = { 0u };
@@ -1650,23 +1649,36 @@ void FVirtualTextureSystem::GatherRequestsTask(const FGatherRequestsParameters&
 				const FTexturePagePool& RESTRICT PagePool = PhysicalSpace->GetPagePool();
 
 				// Find the highest resolution tile that's currently loaded
-				const uint32 pAddress = PagePool.FindNearestPageAddress(ProducerHandle, ProducerGroupIndex, Local_vAddress, Local_vLevel, MaxLevel);
-				uint32 AllocatedLocal_vLevel = MaxLevel + 1u;
-				if (pAddress != ~0u)
+				const uint32 Allocated_pAddress = PagePool.FindNearestPageAddress(ProducerHandle, ProducerGroupIndex, Local_vAddress, Local_vLevel, MaxLevel);
+
+				bool bRequestedPageWasResident = false;
+				uint32 AllocatedLocal_vLevel = MaxLevel;
+				if (Allocated_pAddress != ~0u)
 				{
-					AllocatedLocal_vLevel = PagePool.GetLocalLevelForAddress(pAddress);
+					AllocatedLocal_vLevel = PagePool.GetLocalLevelForAddress(Allocated_pAddress);
 					check(AllocatedLocal_vLevel >= Local_vLevel);
 
-					const uint32 Allocated_vLevel = AllocatedLocal_vLevel + ProducerMipBias;
-					ensure(Allocated_vLevel <= AllocatedVT->GetMaxLevel());
+					const uint32 AllocatedMapping_vLevel = AllocatedLocal_vLevel + ProducerMipBias;
+					uint32 Allocated_vLevel = AllocatedMapping_vLevel;
+					if (AllocatedLocal_vLevel == Local_vLevel)
+					{
+						// page at the requested level was already resident, no longer need to load
+						bRequestedPageWasResident = true;
+						// We can map the already resident page at the original requested vLevel
+						// This may be different from Allocated_vLevel when various biases are involved
+						// Without this, we'll never see anything mapped to the original requested level
+						Allocated_vLevel = vLevel;
+						GroupMaskToLoad &= ~(1u << ProducerGroupIndex);
+						++NumResidentPages;
+					}
 
-					const uint32 AllocatedMapping_vLevel = FMath::Max(Allocated_vLevel, ProducerMipBias);
+					ensure(Allocated_vLevel <= AllocatedVT->GetMaxLevel());
 					const uint32 Allocated_vAddress = vAddress & (0xffffffff << (Allocated_vLevel * vDimensions));
 					const uint32 AllocatedWrapped_vAddress = Wrapped_vAddress & (0xffffffff << (Allocated_vLevel * vDimensions));
 
-					AddPageUpdate(PageUpdateBuffers, PageUpdateFlushCount, PhysicalSpace->GetID(), pAddress);
+					AddPageUpdate(PageUpdateBuffers, PageUpdateFlushCount, PhysicalSpace->GetID(), Allocated_pAddress);
 
-					const FPhysicalSpaceIDAndAddress PhysicalSpaceIDAndAddress(PhysicalSpace->GetID(), pAddress);
+					const FPhysicalSpaceIDAndAddress PhysicalSpaceIDAndAddress(PhysicalSpace->GetID(), Allocated_pAddress);
 					uint32 NumMappedPages = 0u;
 					for (uint32 LoadLayerIndex = 0u; LoadLayerIndex < NumPageTableLayersToLoad; ++LoadLayerIndex)
 					{
@@ -1685,7 +1697,7 @@ void FVirtualTextureSystem::GatherRequestsTask(const FGatherRequestsParameters&
 							if (PrevPhysicalSpaceIDAndAddress.Packed == ~0u)
 							{
 								// map the page now if it wasn't already mapped
-								RequestList->AddDirectMappingRequest(Space->GetID(), PhysicalSpace->GetID(), PageTableLayerIndex, Allocated_vLevel, Allocated_vAddress, AllocatedMapping_vLevel, pAddress);
+								RequestList->AddDirectMappingRequest(Space->GetID(), PhysicalSpace->GetID(), PageTableLayerIndex, Allocated_vAddress, Allocated_vLevel, AllocatedMapping_vLevel, Allocated_pAddress);
 							}
 
 							if (AllocatedWrapped_vAddress != Allocated_vAddress)
@@ -1700,7 +1712,7 @@ void FVirtualTextureSystem::GatherRequestsTask(const FGatherRequestsParameters&
 
 								if (PrevWrappedPhysicalSpaceIDAndAddress.Packed == ~0u)
 								{
-									RequestList->AddDirectMappingRequest(Space->GetID(), PhysicalSpace->GetID(), PageTableLayerIndex, Allocated_vLevel, AllocatedWrapped_vAddress, AllocatedMapping_vLevel, pAddress);
+									RequestList->AddDirectMappingRequest(Space->GetID(), PhysicalSpace->GetID(), PageTableLayerIndex, AllocatedWrapped_vAddress, Allocated_vLevel, AllocatedMapping_vLevel, Allocated_pAddress);
 								}
 							}
 							++NumMappedPages;
@@ -1709,13 +1721,7 @@ void FVirtualTextureSystem::GatherRequestsTask(const FGatherRequestsParameters&
 					check(NumMappedPages > 0u);
 				}
 
-				if (AllocatedLocal_vLevel == Local_vLevel)
-				{
-					// page at the requested level was already resident, no longer need to load
-					GroupMaskToLoad &= ~(1u << ProducerGroupIndex);
-					++NumResidentPages;
-				}
-				else
+				if (!bRequestedPageWasResident)
 				{
 					// page not resident...see if we want to prefetch a page with resolution incrementally larger than what's currently resident
 					// this means we'll ultimately load more data, but these lower resolution pages should load much faster than the requested high resolution page
@@ -1765,11 +1771,10 @@ void FVirtualTextureSystem::GatherRequestsTask(const FGatherRequestsParameters&
 					const uint16 LoadRequestIndex = RequestList->AddLoadRequest(FVirtualTextureLocalTile(ProducerHandle, PrefetchLocal_vAddress, PrefetchLocal_vLevel), ProducerPhysicalGroupMaskToPrefetch, PageCount);
 					if (LoadRequestIndex != 0xffff)
 					{
-						const uint32 Prefetch_vLevel = PrefetchLocal_vLevel + ProducerMipBias;
-						ensure(Prefetch_vLevel <= AllocatedVT->GetMaxLevel());
-						const uint32 PrefetchMapping_vLevel = FMath::Max(Prefetch_vLevel, ProducerMipBias);
-						const uint32 Prefetch_vAddress = vAddress & (0xffffffff << (Prefetch_vLevel * vDimensions));
-						const uint32 PrefetchWrapped_vAddress = Wrapped_vAddress & (0xffffffff << (Prefetch_vLevel * vDimensions));
+						const uint32 PrefetchMapping_vLevel = PrefetchLocal_vLevel + ProducerMipBias;
+						ensure(PrefetchMapping_vLevel <= AllocatedVT->GetMaxLevel());
+						const uint32 Prefetch_vAddress = vAddress & (0xffffffff << (PrefetchMapping_vLevel * vDimensions));
+						const uint32 PrefetchWrapped_vAddress = Wrapped_vAddress & (0xffffffff << (PrefetchMapping_vLevel * vDimensions));
 						for (uint32 LoadLayerIndex = 0u; LoadLayerIndex < NumPageTableLayersToLoad; ++LoadLayerIndex)
 						{
 							const uint32 LayerIndex = PageTableLayersToLoad[LoadLayerIndex];
@@ -1778,10 +1783,10 @@ void FVirtualTextureSystem::GatherRequestsTask(const FGatherRequestsParameters&
 								const uint32 ProducerPhysicalGroupIndex = AllocatedVT->GetProducerPhysicalGroupIndexForPageTableLayer(LayerIndex);
 								if (ProducerPhysicalGroupMaskToPrefetch & (1u << ProducerPhysicalGroupIndex))
 								{
-									RequestList->AddMappingRequest(LoadRequestIndex, ProducerPhysicalGroupIndex, ID, LayerIndex, Prefetch_vAddress, Prefetch_vLevel, PrefetchMapping_vLevel);
+									RequestList->AddMappingRequest(LoadRequestIndex, ProducerPhysicalGroupIndex, ID, LayerIndex, Prefetch_vAddress, PrefetchMapping_vLevel, PrefetchMapping_vLevel);
 									if (PrefetchWrapped_vAddress != Prefetch_vAddress)
 									{
-										RequestList->AddMappingRequest(LoadRequestIndex, ProducerPhysicalGroupIndex, ID, LayerIndex, PrefetchWrapped_vAddress, Prefetch_vLevel, PrefetchMapping_vLevel);
+										RequestList->AddMappingRequest(LoadRequestIndex, ProducerPhysicalGroupIndex, ID, LayerIndex, PrefetchWrapped_vAddress, PrefetchMapping_vLevel, PrefetchMapping_vLevel);
 									}
 								}
 							}
