Missed files from previous submit

diff --git a/Engine/Source/Runtime/Renderer/Private/VT/TexturePageMap.h b/Engine/Source/Runtime/Renderer/Private/VT/TexturePageMap.h
--- a/Engine/Source/Runtime/Renderer/Private/VT/TexturePageMap.h
+++ b/Engine/Source/Runtime/Renderer/Private/VT/TexturePageMap.h
@@ -49,7 +49,7 @@ struct FMappedTexturePage
 	FTexturePage Page;
 	uint32 pAddress : 16;
 	uint32 PhysicalSpaceID : 12;
-	uint32 vLevel : 4;
+	uint32 Local_vLevel : 4;
 };
 
 /**
@@ -98,7 +98,7 @@ public:
 	/**
 	* Map the physical address to a specific virtual address.
 	*/
-	void		MapPage(FVirtualTextureSpace* Space, FVirtualTexturePhysicalSpace* PhysicalSpace, uint8 vLogSize, uint32 vAddress, uint8 vLevel, uint16 pAddress);
+	void		MapPage(FVirtualTextureSpace* Space, FVirtualTexturePhysicalSpace* PhysicalSpace, uint8 vLogSize, uint32 vAddress, uint8 Local_vLevel, uint16 pAddress);
 
 	void		GetMappedPagesInRange(uint32 vAddress, uint32 Width, uint32 Height, TArray<FMappedTexturePage>& OutMappedPages) const;
 
@@ -141,7 +141,7 @@ private:
 			{
 				uint32 pAddress : 16;
 				uint32 PhysicalSpaceID : 12;
-				uint32 vLevel : 4;
+				uint32 Local_vLevel : 4;
 			};
 		};
 	};
diff --git a/Engine/Source/Runtime/Renderer/Private/VT/TexturePagePool.h b/Engine/Source/Runtime/Renderer/Private/VT/TexturePagePool.h
--- a/Engine/Source/Runtime/Renderer/Private/VT/TexturePagePool.h
+++ b/Engine/Source/Runtime/Renderer/Private/VT/TexturePagePool.h
@@ -133,7 +133,7 @@ public:
 	/**
 	* Map the physical address to a specific virtual address.
 	*/
-	void		MapPage(FVirtualTextureSpace* Space, FVirtualTexturePhysicalSpace* PhysicalSpace, uint8 PageTableLayerIndex, uint8 vLogSize, uint32 vAddress, uint8 vLevel, uint16 pAddress);
+	void		MapPage(FVirtualTextureSpace* Space, FVirtualTexturePhysicalSpace* PhysicalSpace, uint8 PageTableLayerIndex, uint8 vLogSize, uint32 vAddress, uint8 Local_vLevel, uint16 pAddress);
 
 private:
 	// Allocate 24 bits to store next/prev indices, pack layer index into 8 bits
