Move instance transform updates out of the GPU scene upload loop and back onto the RDG setup timeline.

diff --git a/Engine/Source/Runtime/Renderer/Private/GPUScene.cpp b/Engine/Source/Runtime/Renderer/Private/GPUScene.cpp
--- a/Engine/Source/Runtime/Renderer/Private/GPUScene.cpp
+++ b/Engine/Source/Runtime/Renderer/Private/GPUScene.cpp
@@ -60,6 +60,27 @@ FAutoConsoleVariableRef CVarGPUSceneInstanceBVH(
 );
 
 
+FVector OrthonormalizeTransform(FMatrix& Matrix)
+{
+	FVector X, Y, Z, Origin;
+	Matrix.GetScaledAxes(X, Y, Z);
+	Origin = Matrix.GetOrigin();
+
+	// Modified Gram-Schmidt orthogonalization
+	Y -= (Y | X) / (X | X) * X;
+	Z -= (Z | X) / (X | X) * X;
+	Z -= (Z | Y) / (Y | Y) * Y;
+
+	Matrix = FMatrix(X, Y, Z, Origin);
+
+	// Extract per axis scales
+	FVector Scale;
+	Scale.X = X.Size();
+	Scale.Y = Y.Size();
+	Scale.Z = Z.Size();
+	return Scale;
+}
+
 static int32 GetMaxPrimitivesUpdate(uint32 NumUploads, uint32 InStrideInFloat4s)
 {
 	return FMath::Min((uint32)(GetMaxBufferDimension() / InStrideInFloat4s), NumUploads);
@@ -147,22 +168,84 @@ struct FPrimitiveUploadInfo
  */
 struct FInstanceUploadInfo
 {
-	/**
-	 * Transform of the primitive from local to world (not the instances).
-	 */
-	FMatrix PrimitiveLocalToWorld;
-	FMatrix PreviousPrimitiveLocalToWorld;
 	TArrayView<FPrimitiveInstance> PrimitiveInstances;
 	int32 InstanceDataOffset = INDEX_NONE;
-
-	int32 PrimitiveID = INDEX_NONE;
-
+	
 	// Used for primitives that need to create a dummy instance (they do not have instance data in the proxy)
 	FPrimitiveInstance DummyInstance;
+};
 
-	bool bHasPrevInstanceTransform = false;
+struct FPrimitiveTransforms
+{
+	FMatrix LocalToWorld;
+	FMatrix PreviousLocalToWorld;
 };
 
+inline FPrimitiveTransforms InitPrimitiveTransforms(const FScene& Scene, const FPrimitiveSceneProxy* PrimitiveSceneProxy, const FPrimitiveSceneInfo* PrimitiveSceneInfo)
+{
+	FPrimitiveTransforms Transforms;
+	Transforms.LocalToWorld = PrimitiveSceneProxy->GetLocalToWorld();
+
+	if (PrimitiveSceneProxy->HasPrevInstanceTransforms())
+	{
+		bool bHasPrecomputedVolumetricLightmap{};
+		bool bOutputVelocity{};
+		int32 SingleCaptureIndex{};
+		Scene.GetPrimitiveUniformShaderParameters_RenderThread(PrimitiveSceneInfo, bHasPrecomputedVolumetricLightmap, Transforms.PreviousLocalToWorld, SingleCaptureIndex, bOutputVelocity);
+	}
+	else
+	{
+		Transforms.PreviousLocalToWorld = Transforms.LocalToWorld;
+	}
+
+	return Transforms;
+}
+
+inline void InitPrimitiveInstance(FPrimitiveInstance& PrimitiveInstance, const FPrimitiveTransforms& PrimitiveTransforms, int32 PrimitiveID, uint32 SceneFrameNumber)
+{
+	PrimitiveInstance.PrimitiveId = PrimitiveID;
+	PrimitiveInstance.LastUpdateSceneFrameNumber = SceneFrameNumber;
+	PrimitiveInstance.LocalBounds = PrimitiveInstance.RenderBounds;
+	PrimitiveInstance.LocalToWorld = PrimitiveInstance.InstanceToLocal * PrimitiveTransforms.LocalToWorld;
+	PrimitiveInstance.PrevLocalToWorld = PrimitiveInstance.InstanceToLocal * PrimitiveTransforms.PreviousLocalToWorld;
+
+	// TODO: This should be propagated from the Primitive, or not exist?
+	PrimitiveInstance.Flags = 0;
+
+	// Remove shear
+	FVector Scale = OrthonormalizeTransform(PrimitiveInstance.LocalToWorld);
+					OrthonormalizeTransform(PrimitiveInstance.PrevLocalToWorld);
+
+	PrimitiveInstance.NonUniformScale = FVector4(
+		Scale.X, Scale.Y, Scale.Z,
+		FMath::Max3(FMath::Abs(Scale.X), FMath::Abs(Scale.Y), FMath::Abs(Scale.Z))
+	);
+
+	PrimitiveInstance.InvNonUniformScaleAndDeterminantSign = FVector4(
+		Scale.X > KINDA_SMALL_NUMBER ? 1.0f / Scale.X : 0.0f,
+		Scale.Y > KINDA_SMALL_NUMBER ? 1.0f / Scale.Y : 0.0f,
+		Scale.Z > KINDA_SMALL_NUMBER ? 1.0f / Scale.Z : 0.0f,
+		FMath::FloatSelect(PrimitiveInstance.LocalToWorld.RotDeterminant(), 1.0f, -1.0f)
+	);
+}
+
+inline void InitPrimitiveInstanceDummy(FPrimitiveInstance& DummyInstance, const FPrimitiveTransforms& PrimitiveTransforms, const FBoxSphereBounds& LocalBounds, int32 PrimitiveID, uint32 SceneFrameNumber)
+{
+	// We always create an instance to ensure that we can always use the same code paths in the shader
+	// In the future we should remove redundant data from the primitive, and then the instances should be
+	// provided by the proxy. However, this is a lot of work before we can just enable it in the base proxy class.
+	DummyInstance.InstanceToLocal = FMatrix::Identity;
+	DummyInstance.LocalToInstance = FMatrix::Identity;
+	DummyInstance.LocalToWorld = FMatrix::Identity;
+	DummyInstance.PrevLocalToWorld = FMatrix::Identity;
+	DummyInstance.NonUniformScale = FVector4(1.0f, 1.0f, 1.0f, 1.0f);
+	DummyInstance.InvNonUniformScaleAndDeterminantSign = FVector4(1.0f, 1.0f, 1.0f, 1.0f);
+	DummyInstance.RenderBounds = LocalBounds;
+	DummyInstance.LocalBounds = LocalBounds;
+
+	InitPrimitiveInstance(DummyInstance, PrimitiveTransforms, PrimitiveID, SceneFrameNumber);
+}
+
 /**
  * Info required by the uploader to update the lightmap data for a primitive.
  */
@@ -181,7 +264,11 @@ struct FUploadDataSourceAdapterScenePrimitives
 {
 	static constexpr bool bUpdateNaniteMaterialTables = true;
 
-	FUploadDataSourceAdapterScenePrimitives(FGPUScene& InGPUScene, FScene& InScene) : GPUScene(InGPUScene), Scene(InScene) {}
+	FUploadDataSourceAdapterScenePrimitives(FGPUScene& InGPUScene, FScene& InScene, uint32 InSceneFrameNumber)
+		: GPUScene(InGPUScene)
+		, Scene(InScene)
+		, SceneFrameNumber(InSceneFrameNumber)
+	{}
 
 	/**
 	 * Return the number of primitives to upload N, GetPrimitiveInfo will be called with ItemIndex in [0,N).
@@ -231,78 +318,73 @@ struct FUploadDataSourceAdapterScenePrimitives
 		return false;
 	}
 
-	FORCEINLINE bool GetInstanceInfo(int32 ItemIndex, FInstanceUploadInfo& InstanceUploadInfo)
+	FORCEINLINE bool GetInstanceInfo(int32 ItemIndex, FInstanceUploadInfo& InstanceUploadInfo) const
 	{
 		const int32 PrimitiveID = GPUScene.PrimitivesToUpdate[ItemIndex];
 		if (PrimitiveID < Scene.PrimitiveSceneProxies.Num())
 		{
-			InstanceUploadInfo.PrimitiveID = PrimitiveID;
 			FPrimitiveSceneProxy* PrimitiveSceneProxy = Scene.PrimitiveSceneProxies[PrimitiveID];
 			const FPrimitiveSceneInfo* PrimitiveSceneInfo = PrimitiveSceneProxy->GetPrimitiveSceneInfo();
 
-			InstanceUploadInfo.bHasPrevInstanceTransform = PrimitiveSceneProxy->HasPrevInstanceTransforms();
+			InstanceUploadInfo.InstanceDataOffset = PrimitiveSceneInfo->GetInstanceDataOffset();
 
 			if (PrimitiveSceneProxy->SupportsInstanceDataBuffer())
 			{
 				TArray<FPrimitiveInstance>* PrimitiveInstancesPtr = PrimitiveSceneProxy->GetPrimitiveInstances();
 				check(PrimitiveInstancesPtr);
 
+				// This data has already been built by UpdateInstanceTransforms.
 				InstanceUploadInfo.PrimitiveInstances = TArrayView<FPrimitiveInstance>(*PrimitiveInstancesPtr);
+				return InstanceUploadInfo.PrimitiveInstances.Num() > 0;
 			}
 #if GPUCULL_TODO
 			else
 			{
-				// We always create an instance to ensure that we can always use the same code paths in the shader
-				// In the future we should remove redundant data from the primitive, and then the instances should be
-				// provided by the proxy. However, this is a lot of work before we can just enable it in the base proxy class.
-				FPrimitiveInstance &DummyInstance = InstanceUploadInfo.DummyInstance;
-				if (!PrimitiveSceneProxy->SupportsInstanceDataBuffer())
-				{
-					DummyInstance.InstanceToLocal = FMatrix::Identity;
-					DummyInstance.LocalToInstance = FMatrix::Identity;
-					DummyInstance.LocalToWorld = FMatrix::Identity;
-					DummyInstance.PrevLocalToWorld = FMatrix::Identity;
-					DummyInstance.NonUniformScale = FVector4(1.0f, 1.0f, 1.0f, 1.0f);
-					DummyInstance.InvNonUniformScaleAndDeterminantSign = FVector4(1.0f, 1.0f, 1.0f, 1.0f);
-					DummyInstance.RenderBounds = PrimitiveSceneProxy->GetLocalBounds();
-					DummyInstance.LocalBounds = DummyInstance.RenderBounds;;
-					DummyInstance.PrimitiveId = 0xFFFFFFFF;
-					DummyInstance.LastUpdateSceneFrameNumber = 0xFFFFFFFF;
-				}
+				const FPrimitiveTransforms PrimitiveTransforms = InitPrimitiveTransforms(Scene, PrimitiveSceneProxy, PrimitiveSceneInfo);
+				InitPrimitiveInstanceDummy(InstanceUploadInfo.DummyInstance, PrimitiveTransforms, PrimitiveSceneProxy->GetLocalBounds(), PrimitiveID, SceneFrameNumber);
+				InstanceUploadInfo.PrimitiveInstances = TArrayView<FPrimitiveInstance>(&InstanceUploadInfo.DummyInstance, 1);
+				return true;
+			}
+#endif
+		}
+
+		return false;
+	}
 
-				InstanceUploadInfo.PrimitiveInstances = TArrayView<FPrimitiveInstance>(&DummyInstance, 1);
+	FORCEINLINE bool ShouldUpdateInstanceTransforms() const
+	{
+		return true;
+	}
+
+	FORCEINLINE void UpdateInstanceTransforms(int32 ItemIndex) const
+	{
+		const int32 PrimitiveID = GPUScene.PrimitivesToUpdate[ItemIndex];
+		if (PrimitiveID < Scene.PrimitiveSceneProxies.Num())
+		{
+			FPrimitiveSceneProxy* PrimitiveSceneProxy = Scene.PrimitiveSceneProxies[PrimitiveID];
+
+			if (!PrimitiveSceneProxy->SupportsInstanceDataBuffer())
+			{
+				return;
 			}
-#endif // GPUCULL_TODO
 
-			check(InstanceUploadInfo.PrimitiveInstances.Num() == PrimitiveSceneInfo->GetNumInstanceDataEntries());
-			if (InstanceUploadInfo.PrimitiveInstances.Num() == 0)
+			const FPrimitiveSceneInfo* PrimitiveSceneInfo = PrimitiveSceneProxy->GetPrimitiveSceneInfo();
+
+			if (!PrimitiveSceneInfo->GetNumInstanceDataEntries())
 			{
-				return false;
+				return;
 			}
 
-			InstanceUploadInfo.PrimitiveLocalToWorld = PrimitiveSceneProxy->GetLocalToWorld();
-			
-			InstanceUploadInfo.InstanceDataOffset = PrimitiveSceneProxy->GetPrimitiveSceneInfo()->GetInstanceDataOffset();
-			check(InstanceUploadInfo.InstanceDataOffset != INDEX_NONE);
-			check(PrimitiveSceneInfo->GetIndex() == PrimitiveID);
-
-			// Dummy data that we don't actually need
-			bool bHasPrecomputedVolumetricLightmap = false;
-			int32 SingleCaptureIndex = 0;
-			bool bOutputVelocity = false;
-
-			Scene.GetPrimitiveUniformShaderParameters_RenderThread(
-				PrimitiveSceneInfo,
-				bHasPrecomputedVolumetricLightmap,
-				InstanceUploadInfo.PreviousPrimitiveLocalToWorld,
-				SingleCaptureIndex,
-				bOutputVelocity
-			);
+			const FPrimitiveTransforms PrimitiveTransforms = InitPrimitiveTransforms(Scene, PrimitiveSceneProxy, PrimitiveSceneInfo);
 
-			return true;
-		}
+			TArray<FPrimitiveInstance>* PrimitiveInstancesPtr = PrimitiveSceneProxy->GetPrimitiveInstances();
+			check(PrimitiveInstancesPtr);
 
-		return false;
+			for (FPrimitiveInstance& PrimitiveInstance : *PrimitiveInstancesPtr)
+			{
+				InitPrimitiveInstance(PrimitiveInstance, PrimitiveTransforms, PrimitiveID, SceneFrameNumber);
+			}
+		}
 	}
 
 	FORCEINLINE bool GetLightMapInfo(int32 ItemIndex, FLightMapUploadInfo &UploadInfo) const
@@ -323,6 +405,7 @@ struct FUploadDataSourceAdapterScenePrimitives
 
 	FGPUScene& GPUScene;
 	FScene& Scene;
+	const uint32 SceneFrameNumber;
 };
 
 void FGPUScene::SetEnabled(ERHIFeatureLevel::Type InFeatureLevel)
@@ -391,21 +474,20 @@ void FGPUScene::UpdateInternal(FRDGBuilder& GraphBuilder, FScene& Scene)
 		bUpdateAllPrimitives = false;
 	}
 
-	FUploadDataSourceAdapterScenePrimitives AdapterTmp(*this, Scene);
-	FGPUSceneBufferState BufferState = UpdateBufferState(GraphBuilder, &Scene, AdapterTmp);
+	// Store in GPU-scene to enable validation that update has been carried out.
+	SceneFrameNumber = Scene.GetFrameNumber();
+
+	FUploadDataSourceAdapterScenePrimitives Adapter(*this, Scene, SceneFrameNumber);
+	FGPUSceneBufferState BufferState = UpdateBufferState(GraphBuilder, &Scene, Adapter);
 
 	AddPass(GraphBuilder, RDG_EVENT_NAME("GPUSceneUpdate"), 
-		[this, &Scene, BufferState](FRHICommandListImmediate& RHICmdList)
+		[this, &Scene, Adapter, BufferState = MoveTemp(BufferState)](FRHICommandListImmediate& RHICmdList)
 	{
 		SCOPED_NAMED_EVENT(STAT_UpdateGPUScene, FColor::Green);
 		CSV_SCOPED_TIMING_STAT_EXCLUSIVE(UpdateGPUScene);
 		QUICK_SCOPE_CYCLE_COUNTER(STAT_UpdateGPUScene);
 		SCOPE_CYCLE_COUNTER(STAT_UpdateGPUSceneTime);
 
-		// Store in GPU-scene to enable validation that update has been carried out.
-		SceneFrameNumber = Scene.GetFrameNumber();
-
-		FUploadDataSourceAdapterScenePrimitives Adapter(*this, Scene);
 		UploadGeneral<FUploadDataSourceAdapterScenePrimitives>(RHICmdList, &Scene, Adapter, BufferState);
 
 		PrimitivesMarkedToUpdate.Init(false, PrimitivesMarkedToUpdate.Num());
@@ -490,30 +572,11 @@ namespace
 
 };
 
-FVector OrthonormalizeTransform( FMatrix& Matrix )
-{
-	FVector X, Y, Z, Origin;
-	Matrix.GetScaledAxes( X, Y, Z );
-	Origin = Matrix.GetOrigin();
-
-	// Modified Gram-Schmidt orthogonalization
-	Y -= (Y | X) / (X | X) * X;
-	Z -= (Z | X) / (X | X) * X;
-	Z -= (Z | Y) / (Y | Y) * Y;
-
-	Matrix = FMatrix( X, Y, Z, Origin );
-
-	// Extract per axis scales
-	FVector Scale;
-	Scale.X = X.Size();
-	Scale.Y = Y.Size();
-	Scale.Z = Z.Size();
-	return Scale;
-}
-
 template<typename FUploadDataSourceAdapter>
-FGPUSceneBufferState FGPUScene::UpdateBufferState(FRDGBuilder& GraphBuilder, FScene* Scene, FUploadDataSourceAdapter& UploadDataSourceAdapter)
+FGPUSceneBufferState FGPUScene::UpdateBufferState(FRDGBuilder& GraphBuilder, FScene* Scene, const FUploadDataSourceAdapter& UploadDataSourceAdapter)
 {
+	UpdatePrimitiveInstances(Scene, UploadDataSourceAdapter);
+
 	FGPUSceneBufferState BufferState;
 	ensure(bInBeginEndBlock);
 	if (Scene != nullptr)
@@ -564,9 +627,41 @@ FGPUSceneBufferState FGPUScene::UpdateBufferState(FRDGBuilder& GraphBuilder, FSc
 	return BufferState;
 }
 
+template <typename FUploadDataSourceAdapter>
+void FGPUScene::UpdatePrimitiveInstances(FScene* Scene, const FUploadDataSourceAdapter& UploadDataSourceAdapter)
+{
+	const int32 NumPrimitiveDataUploads = UploadDataSourceAdapter.NumPrimitivesToUpload();
+
+	if (!NumPrimitiveDataUploads || !UploadDataSourceAdapter.ShouldUpdateInstanceTransforms())
+	{
+		return;
+	}
+
+	// Note: we iterate over the primitives, whether they have instances or not (which is a bit wasteful) but this is the way we currently get to the instance data.
+	// GPUCULL_TODO: move instance data ownership to GPU-scene such that it can be put in a compact list or something, and be tracked independent of primitives?
+
+	TRACE_CPUPROFILER_EVENT_SCOPE(GPUScene_UpdateInstanceTransforms);
+
+	FParallelUpdateRanges ParallelRanges;
+
+	const bool bExecuteInParallel = GGPUSceneParallelUpdate != 0 && FApp::ShouldUseThreadingForPerformance();
+	const int32 RangeCount = PartitionUpdateRanges(ParallelRanges, NumPrimitiveDataUploads, bExecuteInParallel);
+	const uint32 InstanceDataNumArrays = FInstanceSceneShaderData::InstanceDataStrideInFloat4s;
+
+	// Upload any out of date instance slots.
+	ParallelFor(RangeCount,
+		[&UploadDataSourceAdapter, this, &ParallelRanges, RangeCount, InstanceDataNumArrays](int32 RangeIndex)
+	{
+		for (int32 ItemIndex = ParallelRanges.Range[RangeIndex].ItemStart; ItemIndex < ParallelRanges.Range[RangeIndex].ItemStart + ParallelRanges.Range[RangeIndex].ItemCount; ++ItemIndex)
+		{
+			UploadDataSourceAdapter.UpdateInstanceTransforms(ItemIndex);
+		}
+
+	}, RangeCount == 1);
+}
 
 template<typename FUploadDataSourceAdapter>
-void FGPUScene::UploadGeneral(FRHICommandListImmediate& RHICmdList, FScene *Scene, FUploadDataSourceAdapter &UploadDataSourceAdapter, const FGPUSceneBufferState& BufferState)
+void FGPUScene::UploadGeneral(FRHICommandListImmediate& RHICmdList, FScene *Scene, const FUploadDataSourceAdapter& UploadDataSourceAdapter, const FGPUSceneBufferState& BufferState)
 {
 	if (Scene != nullptr)
 	{
@@ -897,42 +992,10 @@ void FGPUScene::UploadGeneral(FRHICommandListImmediate& RHICmdList, FScene *Scen
 								FInstanceUploadInfo UploadInfo;
 								if (UploadDataSourceAdapter.GetInstanceInfo(ItemIndex, UploadInfo))
 								{
-									// These are the bounds covering all primitives (which we don't want to use for per-instance culling).
-									//const FBoxSphereBounds& PrimitiveBounds = PrimitiveSceneProxy->GetLocalBounds();
-
 									// Update each primitive instance with current data.
 									for (int32 InstanceIndex = 0; InstanceIndex < UploadInfo.PrimitiveInstances.Num(); ++InstanceIndex)
 									{
-										FPrimitiveInstance& PrimitiveInstance = UploadInfo.PrimitiveInstances[InstanceIndex];
-										PrimitiveInstance.PrimitiveId = UploadInfo.PrimitiveID;
-										PrimitiveInstance.LocalBounds = PrimitiveInstance.RenderBounds;// .TransformBy(PrimitiveInstance.LocalToWorld);
-										PrimitiveInstance.LocalToWorld = PrimitiveInstance.InstanceToLocal * UploadInfo.PrimitiveLocalToWorld;
-										// TODO: KevinO cleanup
-										const FMatrix& PrevInstanceToLocal = UploadInfo.bHasPrevInstanceTransform ? PrimitiveInstance.PrevInstanceToLocal : PrimitiveInstance.InstanceToLocal;
-										PrimitiveInstance.PrevLocalToWorld = PrevInstanceToLocal * UploadInfo.PreviousPrimitiveLocalToWorld;
-										PrimitiveInstance.LastUpdateSceneFrameNumber = SceneFrameNumber;
-										// TODO: This should be propagated from the Primitive, or not exist?
-										PrimitiveInstance.Flags = 0;
-
-										{
-											// Remove shear
-											FVector Scale =	OrthonormalizeTransform( PrimitiveInstance.LocalToWorld );
-															OrthonormalizeTransform( PrimitiveInstance.PrevLocalToWorld );
-
-											PrimitiveInstance.NonUniformScale = FVector4(
-												Scale.X, Scale.Y, Scale.Z,
-												FMath::Max3(FMath::Abs(Scale.X), FMath::Abs(Scale.Y), FMath::Abs(Scale.Z))
-											);
-
-											PrimitiveInstance.InvNonUniformScaleAndDeterminantSign = FVector4(
-												Scale.X > KINDA_SMALL_NUMBER ? 1.0f / Scale.X : 0.0f,
-												Scale.Y > KINDA_SMALL_NUMBER ? 1.0f / Scale.Y : 0.0f,
-												Scale.Z > KINDA_SMALL_NUMBER ? 1.0f / Scale.Z : 0.0f,
-												FMath::FloatSelect(PrimitiveInstance.LocalToWorld.RotDeterminant(), 1.0f, -1.0f)
-											);
-										}
-
-										FInstanceSceneShaderData InstanceSceneData(PrimitiveInstance);
+										const FInstanceSceneShaderData InstanceSceneData(UploadInfo.PrimitiveInstances[InstanceIndex]);
 
 										void* DstRefs[FInstanceSceneShaderData::InstanceDataStrideInFloat4s];
 										if (RangeCount > 1)
@@ -1049,10 +1112,16 @@ struct FUploadDataSourceAdapterDynamicPrimitives
 {
 	static constexpr bool bUpdateNaniteMaterialTables = false;
 
-	FUploadDataSourceAdapterDynamicPrimitives(const TArray<FPrimitiveUniformShaderParameters, TInlineAllocator<8>>& InPrimitiveShaderData, int32 InPrimitiveIDStartOffset, int32 InInstanceIDStartOffset) :
-		PrimitiveShaderData(InPrimitiveShaderData), 
-		PrimitiveIDStartOffset(InPrimitiveIDStartOffset), 
-		InstanceIDStartOffset(InInstanceIDStartOffset) {}
+	FUploadDataSourceAdapterDynamicPrimitives(
+		const TArray<FPrimitiveUniformShaderParameters, TInlineAllocator<8>>& InPrimitiveShaderData,
+		int32 InPrimitiveIDStartOffset,
+		int32 InInstanceIDStartOffset,
+		uint32 InSceneFrameNumber)
+		: PrimitiveShaderData(InPrimitiveShaderData)
+		, PrimitiveIDStartOffset(InPrimitiveIDStartOffset)
+		, InstanceIDStartOffset(InInstanceIDStartOffset)
+		, SceneFrameNumber(InSceneFrameNumber)
+	{}
 
 	FORCEINLINE int32 NumPrimitivesToUpload() const
 	{ 
@@ -1088,34 +1157,33 @@ struct FUploadDataSourceAdapterDynamicPrimitives
 		return false;
 	}
 
-	FORCEINLINE bool GetInstanceInfo(int32 ItemIndex, FInstanceUploadInfo& InstanceUploadInfo)
-	{	
+	FORCEINLINE bool ShouldUpdateInstanceTransforms() const
+	{
+		return false;
+	}
+
+	FORCEINLINE void UpdateInstanceTransforms(int32 ItemIndex) const
+	{
+		// Nothing to do.
+	}
 
+	FORCEINLINE bool GetInstanceInfo(int32 ItemIndex, FInstanceUploadInfo& InstanceUploadInfo) const
+	{
 #if GPUCULL_TODO
 		if (ItemIndex < PrimitiveShaderData.Num())
 		{
-			InstanceUploadInfo.PrimitiveLocalToWorld = PrimitiveShaderData[ItemIndex].LocalToWorld;
-			InstanceUploadInfo.PreviousPrimitiveLocalToWorld = PrimitiveShaderData[ItemIndex].PreviousLocalToWorld;
+			const uint32 PrimitiveID = PrimitiveIDStartOffset + ItemIndex;
+
 			InstanceUploadInfo.InstanceDataOffset = InstanceIDStartOffset + ItemIndex;
-			InstanceUploadInfo.PrimitiveID = PrimitiveIDStartOffset + ItemIndex;;
-
-			// We always create an instance to ensure that we can always use the same code paths in the shader
-			// In the future we should remove redundant data from the primitive, and then the instances should be
-			// provided by the proxy. However, this is a lot of work before we can just enable it in the base proxy class.
-			FPrimitiveInstance& DummyInstance = InstanceUploadInfo.DummyInstance;
-			DummyInstance.InstanceToLocal = FMatrix::Identity;
-			DummyInstance.LocalToInstance = FMatrix::Identity;
-			DummyInstance.LocalToWorld = FMatrix::Identity;
-			DummyInstance.PrevLocalToWorld = FMatrix::Identity;
-			DummyInstance.NonUniformScale = FVector4(1.0f, 1.0f, 1.0f, 1.0f);
-			DummyInstance.InvNonUniformScaleAndDeterminantSign = FVector4(1.0f, 1.0f, 1.0f, 1.0f);
-			DummyInstance.RenderBounds = FBoxSphereBounds(FBox(PrimitiveShaderData[ItemIndex].LocalObjectBoundsMin, PrimitiveShaderData[ItemIndex].LocalObjectBoundsMax));
-			DummyInstance.LocalBounds = DummyInstance.RenderBounds;;
-			DummyInstance.PrimitiveId = 0xFFFFFFFF;
-			DummyInstance.LastUpdateSceneFrameNumber = 0xFFFFFFFF;
-
-			InstanceUploadInfo.PrimitiveInstances = TArrayView<FPrimitiveInstance>(&DummyInstance, 1);
 
+			FPrimitiveTransforms PrimitiveTransforms;
+			PrimitiveTransforms.LocalToWorld = PrimitiveShaderData[ItemIndex].LocalToWorld;
+			PrimitiveTransforms.PreviousLocalToWorld = PrimitiveShaderData[ItemIndex].PreviousLocalToWorld;
+
+			const FBoxSphereBounds LocalBounds = FBoxSphereBounds(FBox(PrimitiveShaderData[ItemIndex].LocalObjectBoundsMin, PrimitiveShaderData[ItemIndex].LocalObjectBoundsMax));
+
+			InitPrimitiveInstanceDummy(InstanceUploadInfo.DummyInstance, PrimitiveTransforms, LocalBounds, PrimitiveID, SceneFrameNumber);
+			InstanceUploadInfo.PrimitiveInstances = TArrayView<FPrimitiveInstance>(&InstanceUploadInfo.DummyInstance, 1);
 			return true;
 		}
 #endif // GPUCULL_TODO
@@ -1128,8 +1196,9 @@ struct FUploadDataSourceAdapterDynamicPrimitives
 	}
 
 	const TArray<FPrimitiveUniformShaderParameters, TInlineAllocator<8>> &PrimitiveShaderData;
-	int32 PrimitiveIDStartOffset;
-	int32 InstanceIDStartOffset;
+	const int32 PrimitiveIDStartOffset;
+	const int32 InstanceIDStartOffset;
+	const uint32 SceneFrameNumber;
 };
 
 void FGPUScene::UploadDynamicPrimitiveShaderDataForViewInternal(FRDGBuilder& GraphBuilder, FScene *Scene, FViewInfo& View)
@@ -1166,13 +1235,12 @@ void FGPUScene::UploadDynamicPrimitiveShaderDataForViewInternal(FRDGBuilder& Gra
 		ensure(Collector.UploadData->InstanceDataOffset != INDEX_NONE);
 #endif // GPUCULL_TODO
 
-		FUploadDataSourceAdapterDynamicPrimitives UploadAdapterTmp(Collector.UploadData->PrimitiveShaderData, UploadIdStart, Collector.UploadData->InstanceDataOffset);
-		FGPUSceneBufferState BufferState = UpdateBufferState(GraphBuilder, Scene, UploadAdapterTmp);
+		FUploadDataSourceAdapterDynamicPrimitives UploadAdapter(Collector.UploadData->PrimitiveShaderData, UploadIdStart, Collector.UploadData->InstanceDataOffset, SceneFrameNumber);
+		FGPUSceneBufferState BufferState = UpdateBufferState(GraphBuilder, Scene, UploadAdapter);
 
 		AddPass(GraphBuilder, RDG_EVENT_NAME("UploadDynamicPrimitiveShaderDataForView"),
-			[this, Scene, UploadAdapterTmp, BufferState](FRHICommandListImmediate& RHICmdList)
+			[this, Scene, UploadAdapter, BufferState = MoveTemp(BufferState)](FRHICommandListImmediate& RHICmdList)
 		{
-			FUploadDataSourceAdapterDynamicPrimitives UploadAdapter = UploadAdapterTmp;
 			UploadGeneral<FUploadDataSourceAdapterDynamicPrimitives>(RHICmdList, Scene, UploadAdapter, BufferState);
 		});
 	}
diff --git a/Engine/Source/Runtime/Renderer/Private/GPUScene.h b/Engine/Source/Runtime/Renderer/Private/GPUScene.h
--- a/Engine/Source/Runtime/Renderer/Private/GPUScene.h
+++ b/Engine/Source/Runtime/Renderer/Private/GPUScene.h
@@ -264,15 +264,17 @@ private:
 	ERHIFeatureLevel::Type FeatureLevel;
 
 	template<typename FUploadDataSourceAdapter>
-	FGPUSceneBufferState UpdateBufferState(FRDGBuilder& GraphBuilder, FScene* Scene, FUploadDataSourceAdapter& UploadDataSourceAdapter);
+	FGPUSceneBufferState UpdateBufferState(FRDGBuilder& GraphBuilder, FScene* Scene, const FUploadDataSourceAdapter& UploadDataSourceAdapter);
 
+	template <typename FUploadDataSourceAdapter>
+	void UpdatePrimitiveInstances(FScene* Scene, const FUploadDataSourceAdapter& UploadDataSourceAdapter);
 
 	/**
 	 * Generalized upload that uses an adapter to abstract the data souce. Enables uploading scene primitives & dynamic primitives using a single path.
 	 * @parameter Scene may be null, as it is only needed for the Nanite material table update (which is coupled to the Scene at the moment).
 	 */
 	template<typename FUploadDataSourceAdapter>
-	void UploadGeneral(FRHICommandListImmediate& RHICmdList, FScene* Scene, FUploadDataSourceAdapter& UploadDataSourceAdapter, const FGPUSceneBufferState &BufferState);
+	void UploadGeneral(FRHICommandListImmediate& RHICmdList, FScene* Scene, const FUploadDataSourceAdapter& UploadDataSourceAdapter, const FGPUSceneBufferState &BufferState);
 
 	void UploadDynamicPrimitiveShaderDataForViewInternal(FRDGBuilder& GraphBuilder, FScene* Scene, FViewInfo& View);
 
