Fix a bug of modulated shadow missing SceneTexture uniform buffer

diff --git a/Engine/Shaders/Private/ShadowProjectionPixelShader.usf b/Engine/Shaders/Private/ShadowProjectionPixelShader.usf
--- a/Engine/Shaders/Private/ShadowProjectionPixelShader.usf
+++ b/Engine/Shaders/Private/ShadowProjectionPixelShader.usf
@@ -39,6 +39,10 @@
 	#define FEATURE_GATHER4 0
 #endif
 
+#if MODULATED_SHADOWS
+	#define MobileSceneTextures MobileBasePass.SceneTextures
+#endif
+
 #include "ShadowProjectionCommon.ush"
 #include "ShadowFilteringCommon.ush"
 #if USE_PCSS
diff --git a/Engine/Source/Runtime/Renderer/Private/MobileBasePassRendering.cpp b/Engine/Source/Runtime/Renderer/Private/MobileBasePassRendering.cpp
--- a/Engine/Source/Runtime/Renderer/Private/MobileBasePassRendering.cpp
+++ b/Engine/Source/Runtime/Renderer/Private/MobileBasePassRendering.cpp
@@ -43,9 +43,7 @@ static TAutoConsoleVariable<int32> CVarMobileMaxVisibleMovableSpotLightsShadow(
 	TEXT("The max number of visible spotlighs can cast shadow sorted by screen size, should be as less as possible for performance reason"),
 	ECVF_RenderThreadSafe);
 
-// Specify a unique slot for mobile base pass because some rendering in the mobile base pass (e.g. modulated shadow and ViewExtensions) use SceneTextures uniform buffer, but the SceneTextures uniform buffer and MobileBasePass uniform buffer share the same slot
-IMPLEMENT_STATIC_UNIFORM_BUFFER_SLOT(MobileBasePass);
-IMPLEMENT_STATIC_UNIFORM_BUFFER_STRUCT(FMobileBasePassUniformParameters, "MobileBasePass", MobileBasePass);
+IMPLEMENT_STATIC_UNIFORM_BUFFER_STRUCT(FMobileBasePassUniformParameters, "MobileBasePass", SceneTextures);
 
 static TAutoConsoleVariable<int32> CVarMobileUseHWsRGBEncoding(
 	TEXT("r.Mobile.UseHWsRGBEncoding"),
