Fix for bad instance previous-to-world matrices.

diff --git a/Engine/Source/Runtime/Renderer/Private/GPUScene.cpp b/Engine/Source/Runtime/Renderer/Private/GPUScene.cpp
--- a/Engine/Source/Runtime/Renderer/Private/GPUScene.cpp
+++ b/Engine/Source/Runtime/Renderer/Private/GPUScene.cpp
@@ -186,28 +186,25 @@ inline FPrimitiveTransforms InitPrimitiveTransforms(const FScene& Scene, const F
 	FPrimitiveTransforms Transforms;
 	Transforms.LocalToWorld = PrimitiveSceneProxy->GetLocalToWorld();
 
-	if (PrimitiveSceneProxy->HasPrevInstanceTransforms())
 	{
 		bool bHasPrecomputedVolumetricLightmap{};
 		bool bOutputVelocity{};
 		int32 SingleCaptureIndex{};
 		Scene.GetPrimitiveUniformShaderParameters_RenderThread(PrimitiveSceneInfo, bHasPrecomputedVolumetricLightmap, Transforms.PreviousLocalToWorld, SingleCaptureIndex, bOutputVelocity);
 	}
-	else
-	{
-		Transforms.PreviousLocalToWorld = Transforms.LocalToWorld;
-	}
 
 	return Transforms;
 }
 
-inline void InitPrimitiveInstance(FPrimitiveInstance& PrimitiveInstance, const FPrimitiveTransforms& PrimitiveTransforms, int32 PrimitiveID, uint32 SceneFrameNumber)
+inline void InitPrimitiveInstance(FPrimitiveInstance& PrimitiveInstance, const FPrimitiveTransforms& PrimitiveTransforms, int32 PrimitiveID, uint32 SceneFrameNumber, bool bHasPreviousInstanceTransforms)
 {
+	const FMatrix& PreviousInstanceToLocal = bHasPreviousInstanceTransforms ? PrimitiveInstance.PrevInstanceToLocal : PrimitiveInstance.InstanceToLocal;
+
 	PrimitiveInstance.PrimitiveId = PrimitiveID;
 	PrimitiveInstance.LastUpdateSceneFrameNumber = SceneFrameNumber;
 	PrimitiveInstance.LocalBounds = PrimitiveInstance.RenderBounds;
 	PrimitiveInstance.LocalToWorld = PrimitiveInstance.InstanceToLocal * PrimitiveTransforms.LocalToWorld;
-	PrimitiveInstance.PrevLocalToWorld = PrimitiveInstance.InstanceToLocal * PrimitiveTransforms.PreviousLocalToWorld;
+	PrimitiveInstance.PrevLocalToWorld = PreviousInstanceToLocal * PrimitiveTransforms.PreviousLocalToWorld;
 
 	// TODO: This should be propagated from the Primitive, or not exist?
 	PrimitiveInstance.Flags = 0;
@@ -235,6 +232,7 @@ inline void InitPrimitiveInstanceDummy(FPrimitiveInstance& DummyInstance, const
 	// In the future we should remove redundant data from the primitive, and then the instances should be
 	// provided by the proxy. However, this is a lot of work before we can just enable it in the base proxy class.
 	DummyInstance.InstanceToLocal = FMatrix::Identity;
+	DummyInstance.PrevInstanceToLocal = FMatrix::Identity;
 	DummyInstance.LocalToInstance = FMatrix::Identity;
 	DummyInstance.LocalToWorld = FMatrix::Identity;
 	DummyInstance.PrevLocalToWorld = FMatrix::Identity;
@@ -243,7 +241,8 @@ inline void InitPrimitiveInstanceDummy(FPrimitiveInstance& DummyInstance, const
 	DummyInstance.RenderBounds = LocalBounds;
 	DummyInstance.LocalBounds = LocalBounds;
 
-	InitPrimitiveInstance(DummyInstance, PrimitiveTransforms, PrimitiveID, SceneFrameNumber);
+	const bool bHasPreviousInstanceTransforms = false;
+	InitPrimitiveInstance(DummyInstance, PrimitiveTransforms, PrimitiveID, SceneFrameNumber, bHasPreviousInstanceTransforms);
 }
 
 /**
@@ -375,6 +374,7 @@ struct FUploadDataSourceAdapterScenePrimitives
 				return;
 			}
 
+			const bool bHasPreviousInstanceTransforms = PrimitiveSceneProxy->HasPrevInstanceTransforms();
 			const FPrimitiveTransforms PrimitiveTransforms = InitPrimitiveTransforms(Scene, PrimitiveSceneProxy, PrimitiveSceneInfo);
 
 			TArray<FPrimitiveInstance>* PrimitiveInstancesPtr = PrimitiveSceneProxy->GetPrimitiveInstances();
@@ -382,7 +382,7 @@ struct FUploadDataSourceAdapterScenePrimitives
 
 			for (FPrimitiveInstance& PrimitiveInstance : *PrimitiveInstancesPtr)
 			{
-				InitPrimitiveInstance(PrimitiveInstance, PrimitiveTransforms, PrimitiveID, SceneFrameNumber);
+				InitPrimitiveInstance(PrimitiveInstance, PrimitiveTransforms, PrimitiveID, SceneFrameNumber, bHasPreviousInstanceTransforms);
 			}
 		}
 	}
@@ -1174,16 +1174,15 @@ struct FUploadDataSourceAdapterDynamicPrimitives
 		{
 			const uint32 PrimitiveID = PrimitiveIDStartOffset + ItemIndex;
 
-			InstanceUploadInfo.InstanceDataOffset = InstanceIDStartOffset + ItemIndex;
-
 			FPrimitiveTransforms PrimitiveTransforms;
 			PrimitiveTransforms.LocalToWorld = PrimitiveShaderData[ItemIndex].LocalToWorld;
-			PrimitiveTransforms.PreviousLocalToWorld = PrimitiveShaderData[ItemIndex].PreviousLocalToWorld;
+			PrimitiveTransforms.PreviousLocalToWorld = PrimitiveTransforms.LocalToWorld;
 
 			const FBoxSphereBounds LocalBounds = FBoxSphereBounds(FBox(PrimitiveShaderData[ItemIndex].LocalObjectBoundsMin, PrimitiveShaderData[ItemIndex].LocalObjectBoundsMax));
 
 			InitPrimitiveInstanceDummy(InstanceUploadInfo.DummyInstance, PrimitiveTransforms, LocalBounds, PrimitiveID, SceneFrameNumber);
 			InstanceUploadInfo.PrimitiveInstances = TArrayView<FPrimitiveInstance>(&InstanceUploadInfo.DummyInstance, 1);
+			InstanceUploadInfo.InstanceDataOffset = InstanceIDStartOffset + ItemIndex;
 			return true;
 		}
 #endif // GPUCULL_TODO
