Fixes for VT page unmapping...ensure that we don't get mismatched ancestor pages

diff --git a/Engine/Source/Runtime/Renderer/Private/VT/TexturePageMap.cpp b/Engine/Source/Runtime/Renderer/Private/VT/TexturePageMap.cpp
--- a/Engine/Source/Runtime/Renderer/Private/VT/TexturePageMap.cpp
+++ b/Engine/Source/Runtime/Renderer/Private/VT/TexturePageMap.cpp
@@ -77,9 +77,9 @@ FPhysicalSpaceIDAndAddress FTexturePageMap::FindPagePhysicalSpaceIDAndAddress(ui
 	return FindPagePhysicalSpaceIDAndAddress(CheckPage, Hash);
 }
 
-uint32 FTexturePageMap::FindNearestPageIndex(uint8 vLogSize, uint32 vAddress) const
+uint32 FTexturePageMap::FindNearestPageIndex(uint8 vLogSize, uint32 vAddress, uint8 MaxLevel) const
 {
-	while (vLogSize < 16)
+	while (vLogSize <= MaxLevel)
 	{
 		const uint32 Index = FindPageIndex(vLogSize, vAddress);
 		if (Index != ~0u)
@@ -96,13 +96,13 @@ uint32 FTexturePageMap::FindNearestPageIndex(uint8 vLogSize, uint32 vAddress) co
 
 uint32 FTexturePageMap::FindNearestPageAddress(uint8 vLogSize, uint32 vAddress) const
 {
-	const uint32 Index = FindNearestPageIndex(vLogSize, vAddress);
+	const uint32 Index = FindNearestPageIndex(vLogSize, vAddress, VIRTUALTEXTURE_LOG2_MAX_PAGETABLE_SIZE - 1u);
 	return Index != ~0u ? Pages[Index].pAddress : ~0u;
 }
 
 uint32 FTexturePageMap::FindNearestPageLevel(uint8 vLogSize, uint32 vAddress) const
 {
-	const uint32 Index = FindNearestPageIndex(vLogSize, vAddress);
+	const uint32 Index = FindNearestPageIndex(vLogSize, vAddress, VIRTUALTEXTURE_LOG2_MAX_PAGETABLE_SIZE - 1u);
 	if (Index != ~0u)
 	{
 		return Pages[Index].Local_vLevel;
@@ -116,7 +116,8 @@ void FTexturePageMap::UnmapPage(FVirtualTextureSystem* System, FVirtualTextureSp
 	check(PageIndex != ~0u);
 
 	const FTexturePage Page(vLogSize, vAddress);
-	check(Pages[PageIndex].Page == Page);
+	const FPageEntry& Entry = Pages[PageIndex];
+	check(Entry.Page == Page);
 
 	// Unmap old page
 	const uint16 Hash = MurmurFinalize32(Page.Packed);
@@ -126,15 +127,16 @@ void FTexturePageMap::UnmapPage(FVirtualTextureSystem* System, FVirtualTextureSp
 	{
 		const uint32 Parent_vLogSize = vLogSize + 1;
 		const uint32 Parent_vAddress = vAddress & (~0u << (vDimensions * Parent_vLogSize));
-		const uint32 AncestorIndex = FindNearestPageIndex(Parent_vLogSize, Parent_vAddress);
+		const uint32 AncestorIndex = FindNearestPageIndex(Parent_vLogSize, Parent_vAddress, Entry.MaxLevel);
 		if (AncestorIndex != ~0u)
 		{
 			// Root page should typically be locked in memory, so we should always find some valid ancestor pAddress, unless the entire VT is being released
 			// No reason to queue a page table update to invalid pAddress, just leave it alone for now, it will be updated when the page is remapped
-			check(Pages[AncestorIndex].PhysicalSpaceID == Pages[PageIndex].PhysicalSpaceID);
-			const FVirtualTexturePhysicalSpace* PhysicalSpace = System->GetPhysicalSpace(Pages[AncestorIndex].PhysicalSpaceID);
-			const uint8 Ancestor_vLevel = Pages[AncestorIndex].Local_vLevel;
-			Space->QueueUpdate(LayerIndex, vLogSize, vAddress, Ancestor_vLevel, PhysicalSpace->GetPhysicalLocation(Pages[AncestorIndex].pAddress));
+			const FPageEntry& AncestorPage = Pages[AncestorIndex];
+			check(AncestorPage.PackedProducerHandle == Entry.PackedProducerHandle);
+			check(AncestorPage.PhysicalSpaceID == Entry.PhysicalSpaceID);
+			const FVirtualTexturePhysicalSpace* PhysicalSpace = System->GetPhysicalSpace(AncestorPage.PhysicalSpaceID);
+			Space->QueueUpdate(LayerIndex, vLogSize, vAddress, AncestorPage.Local_vLevel, PhysicalSpace->GetPhysicalLocation(AncestorPage.pAddress));
 		}
 	}
 	
@@ -169,21 +171,40 @@ void FTexturePageMap::UnmapPage(FVirtualTextureSystem* System, FVirtualTextureSp
 	--MappedPageCount;
 
 	SortedKeysDirty = true;
+
 }
 
-void FTexturePageMap::MapPage(FVirtualTextureSpace* Space, FVirtualTexturePhysicalSpace* PhysicalSpace, uint8 vLogSize, uint32 vAddress, uint8 Local_vLevel, uint16 pAddress)
+void FTexturePageMap::MapPage(FVirtualTextureSpace* Space, FVirtualTexturePhysicalSpace* PhysicalSpace, uint32 PackedProducerHandle, uint8 MaxLevel, uint8 vLogSize, uint32 vAddress, uint8 Local_vLevel, uint16 pAddress)
 {
 #if DO_GUARD_SLOW
 	const uint32 PrevPageIndex = FindPageIndex(vLogSize, vAddress);
 	checkSlow(PrevPageIndex == ~0u);
 #endif // DO_GUARD_SLOW
 
+#if DO_CHECK
+	// If there's already an ancestor page mapped, make sure it's from the same producer/physical space
+	if (vLogSize < MaxLevel)
+	{
+		const uint32 Parent_vLogSize = vLogSize + 1;
+		const uint32 Parent_vAddress = vAddress & (~0u << (vDimensions * Parent_vLogSize));
+		const uint32 AncestorIndex = FindNearestPageIndex(Parent_vLogSize, Parent_vAddress, MaxLevel);
+		if (AncestorIndex != ~0u)
+		{
+			const FPageEntry& AncestorPage = Pages[AncestorIndex];
+			check(AncestorPage.PackedProducerHandle == PackedProducerHandle);
+			check(AncestorPage.PhysicalSpaceID == PhysicalSpace->GetID());
+		}
+	}
+#endif // DO_CHECK
+
 	const FTexturePage Page(vLogSize, vAddress);
 	const uint32 PageIndex = AcquirePage();
 	check(PageIndex > 0u);
 	FPageEntry& Entry = Pages[PageIndex];
 	Entry.Page = Page;
+	Entry.PackedProducerHandle = PackedProducerHandle;
 	Entry.pAddress = pAddress;
+	Entry.MaxLevel = MaxLevel;
 	Entry.Local_vLevel = Local_vLevel;
 	Entry.PhysicalSpaceID = PhysicalSpace->GetID();
 
@@ -206,10 +227,11 @@ void FTexturePageMap::MapPage(FVirtualTextureSpace* Space, FVirtualTexturePhysic
 
 void FTexturePageMap::GetMappedPagesInRange(uint32 vAddress, uint32 Width, uint32 Height, TArray<FMappedTexturePage>& OutMappedPages) const
 {
-	const uint32 Size = FMath::Max(Width, Height);
-	const uint32 vAddressMax = vAddress + FMath::Square(Size);
 	const uint32 vTileX0 = FMath::ReverseMortonCode2(vAddress);
 	const uint32 vTileY0 = FMath::ReverseMortonCode2(vAddress >> 1);
+	const uint32 vTileX1 = vTileX0 + Width;
+	const uint32 vTileY1 = vTileY0 + Height;
+	const uint32 vAddressMax = FMath::MortonCode2(vTileX1) | (FMath::MortonCode2(vTileY1) << 1);
 
 	uint32 PageIndex = Pages[PageListHead_Mapped].NextIndex;
 	uint32 CheckPageCount = 0u;
@@ -222,8 +244,8 @@ void FTexturePageMap::GetMappedPagesInRange(uint32 vAddress, uint32 Width, uint3
 			const uint32 vTileY = FMath::ReverseMortonCode2(Entry.Page.vAddress >> 1);
 			if (vTileX >= vTileX0 &&
 				vTileY >= vTileY0 &&
-				vTileX < vTileX0 + Width &&
-				vTileY < vTileY0 + Height)
+				vTileX < vTileX1 &&
+				vTileY < vTileY1)
 			{
 				FMappedTexturePage& OutPage = OutMappedPages.AddDefaulted_GetRef();
 				OutPage.Page = Entry.Page;
diff --git a/Engine/Source/Runtime/Renderer/Private/VT/TexturePageMap.h b/Engine/Source/Runtime/Renderer/Private/VT/TexturePageMap.h
--- a/Engine/Source/Runtime/Renderer/Private/VT/TexturePageMap.h
+++ b/Engine/Source/Runtime/Renderer/Private/VT/TexturePageMap.h
@@ -98,7 +98,7 @@ public:
 	/**
 	* Map the physical address to a specific virtual address.
 	*/
-	void		MapPage(FVirtualTextureSpace* Space, FVirtualTexturePhysicalSpace* PhysicalSpace, uint8 vLogSize, uint32 vAddress, uint8 Local_vLevel, uint16 pAddress);
+	void		MapPage(FVirtualTextureSpace* Space, FVirtualTexturePhysicalSpace* PhysicalSpace, uint32 PackedProducerHandle, uint8 MaxLevel, uint8 vLogSize, uint32 vAddress, uint8 Local_vLevel, uint16 pAddress);
 
 	void		GetMappedPagesInRange(uint32 vAddress, uint32 Width, uint32 Height, TArray<FMappedTexturePage>& OutMappedPages) const;
 
@@ -115,7 +115,7 @@ private:
 	uint64		EqualRange(uint32 Min, uint32 Max, uint32 SearchKey, uint32 Mask) const;
 
 	uint32		FindPageIndex(uint8 vLogSize, uint32 vAddress) const;
-	uint32		FindNearestPageIndex(uint8 vLogSize, uint32 vAddress) const;
+	uint32		FindNearestPageIndex(uint8 vLogSize, uint32 vAddress, uint8 MaxLevel) const;
 
 	uint32		LayerIndex;
 	uint32		vDimensions;
@@ -136,11 +136,13 @@ private:
 		uint32 PrevIndex;
 		union
 		{
-			uint32 Packed;
+			uint64 Packed;
 			struct 
 			{
+				uint32 PackedProducerHandle;
 				uint32 pAddress : 16;
-				uint32 PhysicalSpaceID : 12;
+				uint32 PhysicalSpaceID : 8;
+				uint32 MaxLevel : 4;
 				uint32 Local_vLevel : 4;
 			};
 		};
diff --git a/Engine/Source/Runtime/Renderer/Private/VT/TexturePagePool.cpp b/Engine/Source/Runtime/Renderer/Private/VT/TexturePagePool.cpp
--- a/Engine/Source/Runtime/Renderer/Private/VT/TexturePagePool.cpp
+++ b/Engine/Source/Runtime/Renderer/Private/VT/TexturePagePool.cpp
@@ -35,8 +35,8 @@ void FTexturePagePool::Initialize(uint32 InNumPages)
 	for (uint32 i = 0; i <= InNumPages; i++)
 	{
 		FPageMapping& Mapping = PageMapping[i];
-		Mapping.PackedValues = 0xffffffff;
-		Mapping.PageTableLayerIndex = 0xff;
+		FMemory::Memset(Mapping, 0xff);
+		Mapping.Pad = 0u;
 		Mapping.NextIndex = Mapping.PrevIndex = i;
 	}
 }
@@ -64,28 +64,34 @@ void FTexturePagePool::UnmapAllPagesForSpace(FVirtualTextureSystem* System, uint
 	check(Height > 0u);
 	checkf((vAddress & (0xffffffff << (MaxLevel * 2u))) == vAddress, TEXT("vAddress %08X is not aligned to max level %d"), vAddress, MaxLevel);
 
-	const uint32 Size = FMath::Max(Width, Height);
-	const uint32 vAddressMax = vAddress + FMath::Square(Size);
 	const uint32 vTileX0 = FMath::ReverseMortonCode2(vAddress);
 	const uint32 vTileY0 = FMath::ReverseMortonCode2(vAddress >> 1);
+	const uint32 vTileX1 = vTileX0 + Width;
+	const uint32 vTileY1 = vTileY0 + Height;
+	const uint32 vAddressMax = FMath::MortonCode2(vTileX1) | (FMath::MortonCode2(vTileY1) << 1);
 
 	// walk through all of our current mapping entries, and unmap any that belong to the current space
 	for (int32 MappingIndex = NumPages + 1; MappingIndex < PageMapping.Num(); ++MappingIndex)
 	{
-		FPageMapping& Mapping = PageMapping[MappingIndex];
+		const FPageMapping& Mapping = PageMapping[MappingIndex];
 		if (Mapping.PageTableLayerIndex != 0xff &&
 			Mapping.SpaceID == SpaceID &&
 			Mapping.vLogSize <= MaxLevel &&
 			Mapping.vAddress >= vAddress &&
 			Mapping.vAddress < vAddressMax)
 		{
+			check((Mapping.vAddress & (0xffffffff << (Mapping.vLogSize * 2u))) == Mapping.vAddress);
+
 			const uint32 vTileX = FMath::ReverseMortonCode2(Mapping.vAddress);
 			const uint32 vTileY = FMath::ReverseMortonCode2(Mapping.vAddress >> 1);
+
 			if (vTileX >= vTileX0 &&
 				vTileY >= vTileY0 &&
-				vTileX < vTileX0 + Width &&
-				vTileY < vTileY0 + Height)
+				vTileX < vTileX1 &&
+				vTileY < vTileY1)
 			{
+				// MaxLevel should match, if this mapping actually belongs to the region we're trying to unmap
+				check(Mapping.MaxLevel == MaxLevel);
 				// we're unmapping all pages for space, so don't try to map any ancestor pages...they'll be unmapped as well
 				UnmapPageMapping(System, MappingIndex, false);
 			}
@@ -326,11 +332,14 @@ uint32 FTexturePagePool::GetNumVisiblePages(uint32 Frame) const
 	return Count;
 }
 
-void FTexturePagePool::MapPage(FVirtualTextureSpace* Space, FVirtualTexturePhysicalSpace* PhysicalSpace, uint8 PageTableLayerIndex, uint8 vLogSize, uint32 vAddress, uint8 Local_vLevel, uint16 pAddress)
+void FTexturePagePool::MapPage(FVirtualTextureSpace* Space, FVirtualTexturePhysicalSpace* PhysicalSpace, uint8 PageTableLayerIndex, uint8 MaxLevel, uint8 vLogSize, uint32 vAddress, uint8 Local_vLevel, uint16 pAddress)
 {
 	check(pAddress < NumPages);
+	const FPageEntry& PageEntry = Pages[pAddress];
+	checkf(PageEntry.PackedProducerHandle != 0u, TEXT("Trying to map pAddress %04x that hasn't been allocated"), pAddress);
+
 	FTexturePageMap& PageMap = Space->GetPageMapForPageTableLayer(PageTableLayerIndex);
-	PageMap.MapPage(Space, PhysicalSpace, vLogSize, vAddress, Local_vLevel, pAddress);
+	PageMap.MapPage(Space, PhysicalSpace, PageEntry.PackedProducerHandle, MaxLevel, vLogSize, vAddress, Local_vLevel, pAddress);
 
 	++NumPagesMapped;
 
@@ -340,6 +349,7 @@ void FTexturePagePool::MapPage(FVirtualTextureSpace* Space, FVirtualTexturePhysi
 	Mapping.SpaceID = Space->GetID();
 	Mapping.vAddress = vAddress;
 	Mapping.vLogSize = vLogSize;
+	Mapping.MaxLevel = MaxLevel;
 	Mapping.PageTableLayerIndex = PageTableLayerIndex;
 }
 
@@ -354,16 +364,22 @@ void FTexturePagePool::UnmapPageMapping(FVirtualTextureSystem* System, uint32 Ma
 	check(NumPagesMapped > 0u);
 	--NumPagesMapped;
 
-	Mapping.PackedValues = 0xffffffff;
+	Mapping.vAddress = 0x00ffffff;
+	Mapping.vLogSize = 0x0f;
+	Mapping.SpaceID = 0x0f;
+	Mapping.MaxLevel = 0x0f;
 	Mapping.PageTableLayerIndex = 0xff;
+
 	ReleaseMapping(MappingIndex);
 }
 
 void FTexturePagePool::UnmapAllPages(FVirtualTextureSystem* System, uint16 pAddress, bool bMapAncestorPages)
 {
 	FPageEntry& PageEntry = Pages[pAddress];
+	FVirtualTextureProducer* Producer = nullptr;
 	if (PageEntry.PackedProducerHandle != 0u)
 	{
+		Producer = System->FindProducer(PageEntry.GetProducerHandle());
 		check(NumPagesAllocated > 0u);
 		--NumPagesAllocated;
 		PageHash.Remove(GetPageHash(PageEntry), pAddress);
@@ -436,7 +452,7 @@ void FTexturePagePool::RemapPages(FVirtualTextureSystem* System, uint8 SpaceID,
 					Mapping.vLogSize = (int32)Mapping.vLogSize + vLevelBias;
 
 					const int32 vLevel = PageEntry.Local_vLevel; // Deal with any producer mip bias?
-					PageMap.MapPage(Space, PhysicalSpace, Mapping.vLogSize, Mapping.vAddress, vLevel, pAddress);
+					PageMap.MapPage(Space, PhysicalSpace, NewProducerHandle.PackedValue, Mapping.MaxLevel, Mapping.vLogSize, Mapping.vAddress, vLevel, pAddress);
 
 					MappingIndex = Mapping.NextIndex;
 				}
@@ -444,4 +460,3 @@ void FTexturePagePool::RemapPages(FVirtualTextureSystem* System, uint8 SpaceID,
 		}
 	}
 }
-
diff --git a/Engine/Source/Runtime/Renderer/Private/VT/TexturePagePool.h b/Engine/Source/Runtime/Renderer/Private/VT/TexturePagePool.h
--- a/Engine/Source/Runtime/Renderer/Private/VT/TexturePagePool.h
+++ b/Engine/Source/Runtime/Renderer/Private/VT/TexturePagePool.h
@@ -133,7 +133,7 @@ public:
 	/**
 	* Map the physical address to a specific virtual address.
 	*/
-	void		MapPage(FVirtualTextureSpace* Space, FVirtualTexturePhysicalSpace* PhysicalSpace, uint8 PageTableLayerIndex, uint8 vLogSize, uint32 vAddress, uint8 Local_vLevel, uint16 pAddress);
+	void		MapPage(FVirtualTextureSpace* Space, FVirtualTexturePhysicalSpace* PhysicalSpace, uint8 PageTableLayerIndex, uint8 MaxLevel, uint8 vLogSize, uint32 vAddress, uint8 Local_vLevel, uint16 pAddress);
 
 private:
 	// Allocate 24 bits to store next/prev indices, pack layer index into 8 bits
@@ -141,18 +141,14 @@ private:
 
 	struct FPageMapping
 	{
-		union
-		{
-			uint32 PackedValues;
-			struct
-			{
-				uint32 vAddress : 24;
-				uint32 vLogSize : 4;
-				uint32 SpaceID : 4;
-			};
-		};
-	
-		uint32 NextIndex;
+		uint32 vAddress : 24;
+		uint32 vLogSize : 4;
+		uint32 SpaceID : 4;
+
+		uint32 NextIndex : 24;
+		uint32 MaxLevel : 4;
+		uint32 Pad : 4;
+
 		uint32 PrevIndex : 24;
 		uint32 PageTableLayerIndex : 8;
 	};
@@ -167,6 +163,8 @@ private:
 			uint32 Local_vLevel : 4;
 			uint32 GroupIndex : 4;
 		};
+
+		FVirtualTextureProducerHandle GetProducerHandle() const { return FVirtualTextureProducerHandle(PackedProducerHandle); }
 	};
 
 	static uint16 GetPageHash(const FPageEntry& Entry);
diff --git a/Engine/Source/Runtime/Renderer/Private/VT/UniqueRequestList.h b/Engine/Source/Runtime/Renderer/Private/VT/UniqueRequestList.h
--- a/Engine/Source/Runtime/Renderer/Private/VT/UniqueRequestList.h
+++ b/Engine/Source/Runtime/Renderer/Private/VT/UniqueRequestList.h
@@ -12,8 +12,8 @@
 union FMappingRequest
 {
 	inline FMappingRequest() {}
-	inline FMappingRequest(uint16 InLoadIndex, uint8 InPhysicalGroupIndex, uint8 InSpaceID, uint8 InPageTableLayerIndex, uint32 InAddress, uint8 InLevel, uint8 InLocalLevel)
-		: vAddress(InAddress), vLevel(InLevel), SpaceID(InSpaceID), LoadRequestIndex(InLoadIndex), Local_vLevel(InLocalLevel), ProducerPhysicalGroupIndex(InPhysicalGroupIndex), PageTableLayerIndex(InPageTableLayerIndex), Pad(0)
+	inline FMappingRequest(uint16 InLoadIndex, uint8 InPhysicalGroupIndex, uint8 InSpaceID, uint8 InPageTableLayerIndex, uint32 InMaxLevel, uint32 InAddress, uint8 InLevel, uint8 InLocalLevel)
+		: vAddress(InAddress), vLevel(InLevel), SpaceID(InSpaceID), LoadRequestIndex(InLoadIndex), Local_vLevel(InLocalLevel), ProducerPhysicalGroupIndex(InPhysicalGroupIndex), PageTableLayerIndex(InPageTableLayerIndex), MaxLevel(InMaxLevel)
 	{}
 
 	uint64 PackedValue;
@@ -26,7 +26,7 @@ union FMappingRequest
 		uint32 Local_vLevel : 4;
 		uint32 ProducerPhysicalGroupIndex : 4;
 		uint32 PageTableLayerIndex : 4;
-		uint32 Pad : 4;
+		uint32 MaxLevel : 4;
 	};
 };
 static_assert(sizeof(FMappingRequest) == sizeof(uint64), "Bad packing");
@@ -36,25 +36,29 @@ inline bool operator!=(const FMappingRequest& Lhs, const FMappingRequest& Rhs) {
 union FDirectMappingRequest
 {
 	inline FDirectMappingRequest() {}
-	inline FDirectMappingRequest(uint8 InSpaceID, uint16 InPhysicalSpaceID, uint8 InPageTableLayerIndex, uint32 InAddress, uint8 InLevel, uint8 InLocalLevel, uint16 InPhysicalAddress)
-		: vAddress(InAddress), vLevel(InLevel), SpaceID(InSpaceID), pAddress(InPhysicalAddress), PhysicalSpaceID(InPhysicalSpaceID), Local_vLevel(InLocalLevel), PageTableLayerIndex(InPageTableLayerIndex)
+	inline FDirectMappingRequest(uint8 InSpaceID, uint16 InPhysicalSpaceID, uint8 InPageTableLayerIndex, uint32 InMaxLevel, uint32 InAddress, uint8 InLevel, uint8 InLocalLevel, uint16 InPhysicalAddress)
+		: vAddress(InAddress), vLevel(InLevel), SpaceID(InSpaceID), pAddress(InPhysicalAddress), PhysicalSpaceID(InPhysicalSpaceID), Local_vLevel(InLocalLevel), MaxLevel(InMaxLevel), PageTableLayerIndex(InPageTableLayerIndex), Pad(0u)
 	{}
 
-	uint64 PackedValue;
+	uint32 PackedValue[3];
 	struct
 	{
 		uint32 vAddress : 24;
 		uint32 vLevel : 4;
 		uint32 SpaceID : 4;
+
 		uint32 pAddress : 16;
 		uint32 PhysicalSpaceID : 8;
 		uint32 Local_vLevel : 4;
+		uint32 MaxLevel : 4;
+
 		uint32 PageTableLayerIndex : 4;
+		uint32 Pad : 28;
 	};
 };
-static_assert(sizeof(FDirectMappingRequest) == sizeof(uint64), "Bad packing");
-inline bool operator==(const FDirectMappingRequest& Lhs, const FDirectMappingRequest& Rhs) { return Lhs.PackedValue == Rhs.PackedValue; }
-inline bool operator!=(const FDirectMappingRequest& Lhs, const FDirectMappingRequest& Rhs) { return Lhs.PackedValue != Rhs.PackedValue; }
+static_assert(sizeof(FDirectMappingRequest) == sizeof(uint32) * 3, "Bad packing");
+inline bool operator==(const FDirectMappingRequest& Lhs, const FDirectMappingRequest& Rhs) { return Lhs.PackedValue[0] == Rhs.PackedValue[0] && Lhs.PackedValue[1] == Rhs.PackedValue[1] && Lhs.PackedValue[2] == Rhs.PackedValue[2]; }
+inline bool operator!=(const FDirectMappingRequest& Lhs, const FDirectMappingRequest& Rhs) { return !operator==(Lhs, Rhs); }
 
 class FUniqueRequestList
 {
@@ -107,9 +111,9 @@ public:
 	uint16 AddLoadRequest(const FVirtualTextureLocalTile& Tile, uint8 GroupMask, uint16 Count);
 	uint16 LockLoadRequest(const FVirtualTextureLocalTile& Tile, uint8 GroupMask);
 
-	void AddMappingRequest(uint16 LoadRequestIndex, uint8 ProducerPhysicalGroupIndex, uint8 SpaceID, uint8 PageTableLayerIndex, uint32 vAddress, uint8 vLevel, uint8 Local_vLevel);
+	void AddMappingRequest(uint16 LoadRequestIndex, uint8 ProducerPhysicalGroupIndex, uint8 SpaceID, uint8 PageTableLayerIndex, uint32 MaxLevel, uint32 vAddress, uint8 vLevel, uint8 Local_vLevel);
 
-	void AddDirectMappingRequest(uint8 InSpaceID, uint16 InPhysicalSpaceID, uint8 InPageTableLayerIndex, uint32 InAddress, uint8 InLevel, uint8 InLocalLevel, uint16 InPhysicalAddress);
+	void AddDirectMappingRequest(uint8 InSpaceID, uint16 InPhysicalSpaceID, uint8 InPageTableLayerIndex, uint32 MaxLevel, uint32 InAddress, uint8 InLevel, uint8 InLocalLevel, uint16 InPhysicalAddress);
 	void AddDirectMappingRequest(const FDirectMappingRequest& Request);
 
 	void AddContinuousUpdateRequest(const FVirtualTextureLocalTile& Request);
@@ -215,10 +219,10 @@ inline uint16 FUniqueRequestList::LockLoadRequest(const FVirtualTextureLocalTile
 	return 0xffff;
 }
 
-inline void FUniqueRequestList::AddMappingRequest(uint16 LoadRequestIndex, uint8 ProducerPhysicalGroupIndex, uint8 SpaceID, uint8 PageTableLayerIndex, uint32 vAddress, uint8 vLevel, uint8 Local_vLevel)
+inline void FUniqueRequestList::AddMappingRequest(uint16 LoadRequestIndex, uint8 ProducerPhysicalGroupIndex, uint8 SpaceID, uint8 PageTableLayerIndex, uint32 MaxLevel, uint32 vAddress, uint8 vLevel, uint8 Local_vLevel)
 {
 	check(LoadRequestIndex < NumLoadRequests);
-	const FMappingRequest Request(LoadRequestIndex, ProducerPhysicalGroupIndex, SpaceID, PageTableLayerIndex, vAddress, vLevel, Local_vLevel);
+	const FMappingRequest Request(LoadRequestIndex, ProducerPhysicalGroupIndex, SpaceID, PageTableLayerIndex, MaxLevel, vAddress, vLevel, Local_vLevel);
 	const uint16 Hash = MurmurFinalize64(Request.PackedValue);
 
 	for (uint16 Index = MappingRequestHash.First(Hash); MappingRequestHash.IsValid(Index); Index = MappingRequestHash.Next(Index))
@@ -237,15 +241,15 @@ inline void FUniqueRequestList::AddMappingRequest(uint16 LoadRequestIndex, uint8
 	}
 }
 
-inline void FUniqueRequestList::AddDirectMappingRequest(uint8 InSpaceID, uint16 InPhysicalSpaceID, uint8 InPageTableLayerIndex, uint32 InAddress, uint8 InLevel, uint8 InLocalLevel, uint16 InPhysicalAddress)
+inline void FUniqueRequestList::AddDirectMappingRequest(uint8 InSpaceID, uint16 InPhysicalSpaceID, uint8 InPageTableLayerIndex, uint32 InMaxLevel, uint32 InAddress, uint8 InLevel, uint8 InLocalLevel, uint16 InPhysicalAddress)
 {
-	const FDirectMappingRequest Request(InSpaceID, InPhysicalSpaceID, InPageTableLayerIndex, InAddress, InLevel, InLocalLevel, InPhysicalAddress);
+	const FDirectMappingRequest Request(InSpaceID, InPhysicalSpaceID, InPageTableLayerIndex, InMaxLevel, InAddress, InLevel, InLocalLevel, InPhysicalAddress);
 	AddDirectMappingRequest(Request);
 }
 
 inline void FUniqueRequestList::AddDirectMappingRequest(const FDirectMappingRequest& Request)
 {
-	const uint16 Hash = MurmurFinalize64(Request.PackedValue);
+	const uint16 Hash = Murmur32({ Request.PackedValue[0], Request.PackedValue[1], Request.PackedValue[2] });
 	for (uint16 Index = DirectMappingRequestHash.First(Hash); DirectMappingRequestHash.IsValid(Index); Index = DirectMappingRequestHash.Next(Index))
 	{
 		if (Request == DirectMappingRequests[Index])
@@ -313,7 +317,7 @@ inline void FUniqueRequestList::MergeRequests(const FUniqueRequestList* RESTRICT
 		const uint16 LoadRequestIndex = LoadRequestIndexRemap[Request.LoadRequestIndex];
 		if (LoadRequestIndex != 0xffff)
 		{
-			AddMappingRequest(LoadRequestIndex, Request.ProducerPhysicalGroupIndex, Request.SpaceID, Request.PageTableLayerIndex, Request.vAddress, Request.vLevel, Request.Local_vLevel);
+			AddMappingRequest(LoadRequestIndex, Request.ProducerPhysicalGroupIndex, Request.SpaceID, Request.PageTableLayerIndex, Request.MaxLevel, Request.vAddress, Request.vLevel, Request.Local_vLevel);
 		}
 	}
 
diff --git a/Engine/Source/Runtime/Renderer/Private/VT/VirtualTextureSystem.cpp b/Engine/Source/Runtime/Renderer/Private/VT/VirtualTextureSystem.cpp
--- a/Engine/Source/Runtime/Renderer/Private/VT/VirtualTextureSystem.cpp
+++ b/Engine/Source/Runtime/Renderer/Private/VT/VirtualTextureSystem.cpp
@@ -1535,8 +1535,10 @@ void FVirtualTextureSystem::GatherRequestsTask(const FGatherRequestsParameters&
 			continue;
 		}
 
+		const uint32 MaxLevel = AllocatedVT->GetMaxLevel();
+
 		check(AllocatedVT->GetNumPageTableLayers() == Space->GetNumPageTableLayers());
-		if (vLevel > AllocatedVT->GetMaxLevel())
+		if (vLevel > MaxLevel)
 		{
 			// Requested level is outside the given allocated VT
 			// This can happen for requests made by expanding mips, since we don't know the current allocated VT in that context
@@ -1589,7 +1591,6 @@ void FVirtualTextureSystem::GatherRequestsTask(const FGatherRequestsParameters&
 				continue;
 			}
 
-			const uint32 MaxLevel = FMath::Min(Producer->GetMaxLevel(), AllocatedVT->GetMaxLevel());
 			const uint32 ProducerMipBias = AllocatedVT->GetUniqueProducerMipBias(ProducerIndex);
 
 			// here vLevel is clamped against ProducerMipBias, as ProducerMipBias represents the most detailed level of this producer, relative to the allocated VT
@@ -1672,7 +1673,7 @@ void FVirtualTextureSystem::GatherRequestsTask(const FGatherRequestsParameters&
 						++NumResidentPages;
 					}
 
-					ensure(Allocated_vLevel <= AllocatedVT->GetMaxLevel());
+					ensure(Allocated_vLevel <= MaxLevel);
 					const uint32 Allocated_vAddress = vAddress & (0xffffffff << (Allocated_vLevel * vDimensions));
 					const uint32 AllocatedWrapped_vAddress = Wrapped_vAddress & (0xffffffff << (Allocated_vLevel * vDimensions));
 
@@ -1697,7 +1698,7 @@ void FVirtualTextureSystem::GatherRequestsTask(const FGatherRequestsParameters&
 							if (PrevPhysicalSpaceIDAndAddress.Packed == ~0u)
 							{
 								// map the page now if it wasn't already mapped
-								RequestList->AddDirectMappingRequest(Space->GetID(), PhysicalSpace->GetID(), PageTableLayerIndex, Allocated_vAddress, Allocated_vLevel, AllocatedMapping_vLevel, Allocated_pAddress);
+								RequestList->AddDirectMappingRequest(Space->GetID(), PhysicalSpace->GetID(), PageTableLayerIndex, MaxLevel, Allocated_vAddress, Allocated_vLevel, AllocatedMapping_vLevel, Allocated_pAddress);
 							}
 
 							if (AllocatedWrapped_vAddress != Allocated_vAddress)
@@ -1712,7 +1713,7 @@ void FVirtualTextureSystem::GatherRequestsTask(const FGatherRequestsParameters&
 
 								if (PrevWrappedPhysicalSpaceIDAndAddress.Packed == ~0u)
 								{
-									RequestList->AddDirectMappingRequest(Space->GetID(), PhysicalSpace->GetID(), PageTableLayerIndex, AllocatedWrapped_vAddress, Allocated_vLevel, AllocatedMapping_vLevel, Allocated_pAddress);
+									RequestList->AddDirectMappingRequest(Space->GetID(), PhysicalSpace->GetID(), PageTableLayerIndex, MaxLevel, AllocatedWrapped_vAddress, Allocated_vLevel, AllocatedMapping_vLevel, Allocated_pAddress);
 								}
 							}
 							++NumMappedPages;
@@ -1727,7 +1728,7 @@ void FVirtualTextureSystem::GatherRequestsTask(const FGatherRequestsParameters&
 					// this means we'll ultimately load more data, but these lower resolution pages should load much faster than the requested high resolution page
 					// this should make popping less noticeable
 					uint32 PrefetchLocal_vLevel = AllocatedLocal_vLevel - FMath::Min(2u, AllocatedLocal_vLevel);
-					PrefetchLocal_vLevel = FMath::Min<uint32>(PrefetchLocal_vLevel, AllocatedVT->GetMaxLevel() - ProducerMipBias);
+					PrefetchLocal_vLevel = FMath::Min<uint32>(PrefetchLocal_vLevel, MaxLevel - ProducerMipBias);
 					if (PrefetchLocal_vLevel > Local_vLevel)
 					{
 						ProducerPhysicalGroupMaskToPrefetchForLevel[PrefetchLocal_vLevel] |= (1u << ProducerGroupIndex);
@@ -1772,7 +1773,7 @@ void FVirtualTextureSystem::GatherRequestsTask(const FGatherRequestsParameters&
 					if (LoadRequestIndex != 0xffff)
 					{
 						const uint32 PrefetchMapping_vLevel = PrefetchLocal_vLevel + ProducerMipBias;
-						ensure(PrefetchMapping_vLevel <= AllocatedVT->GetMaxLevel());
+						ensure(PrefetchMapping_vLevel <= MaxLevel);
 						const uint32 Prefetch_vAddress = vAddress & (0xffffffff << (PrefetchMapping_vLevel * vDimensions));
 						const uint32 PrefetchWrapped_vAddress = Wrapped_vAddress & (0xffffffff << (PrefetchMapping_vLevel * vDimensions));
 						for (uint32 LoadLayerIndex = 0u; LoadLayerIndex < NumPageTableLayersToLoad; ++LoadLayerIndex)
@@ -1783,10 +1784,10 @@ void FVirtualTextureSystem::GatherRequestsTask(const FGatherRequestsParameters&
 								const uint32 ProducerPhysicalGroupIndex = AllocatedVT->GetProducerPhysicalGroupIndexForPageTableLayer(LayerIndex);
 								if (ProducerPhysicalGroupMaskToPrefetch & (1u << ProducerPhysicalGroupIndex))
 								{
-									RequestList->AddMappingRequest(LoadRequestIndex, ProducerPhysicalGroupIndex, ID, LayerIndex, Prefetch_vAddress, PrefetchMapping_vLevel, PrefetchMapping_vLevel);
+									RequestList->AddMappingRequest(LoadRequestIndex, ProducerPhysicalGroupIndex, ID, LayerIndex, MaxLevel, Prefetch_vAddress, PrefetchMapping_vLevel, PrefetchMapping_vLevel);
 									if (PrefetchWrapped_vAddress != Prefetch_vAddress)
 									{
-										RequestList->AddMappingRequest(LoadRequestIndex, ProducerPhysicalGroupIndex, ID, LayerIndex, PrefetchWrapped_vAddress, PrefetchMapping_vLevel, PrefetchMapping_vLevel);
+										RequestList->AddMappingRequest(LoadRequestIndex, ProducerPhysicalGroupIndex, ID, LayerIndex, MaxLevel, PrefetchWrapped_vAddress, PrefetchMapping_vLevel, PrefetchMapping_vLevel);
 									}
 								}
 							}
@@ -1808,10 +1809,10 @@ void FVirtualTextureSystem::GatherRequestsTask(const FGatherRequestsParameters&
 							const uint32 ProducerPhysicalGroupIndex = AllocatedVT->GetProducerPhysicalGroupIndexForPageTableLayer(LayerIndex);
 							if (GroupMaskToLoad & (1u << ProducerPhysicalGroupIndex))
 							{
-								RequestList->AddMappingRequest(LoadRequestIndex, ProducerPhysicalGroupIndex, ID, LayerIndex, vAddress, vLevel, Mapping_vLevel);
+								RequestList->AddMappingRequest(LoadRequestIndex, ProducerPhysicalGroupIndex, ID, LayerIndex, MaxLevel, vAddress, vLevel, Mapping_vLevel);
 								if (Wrapped_vAddress != vAddress)
 								{
-									RequestList->AddMappingRequest(LoadRequestIndex, ProducerPhysicalGroupIndex, ID, LayerIndex, Wrapped_vAddress, vLevel, Mapping_vLevel);
+									RequestList->AddMappingRequest(LoadRequestIndex, ProducerPhysicalGroupIndex, ID, LayerIndex, MaxLevel, Wrapped_vAddress, vLevel, Mapping_vLevel);
 								}
 							}
 						}
@@ -2235,7 +2236,7 @@ void FVirtualTextureSystem::SubmitRequests(FRDGBuilder& GraphBuilder, ERHIFeatur
 			FVirtualTextureSpace* Space = GetSpace(MappingRequest.SpaceID);
 			FVirtualTexturePhysicalSpace* PhysicalSpace = GetPhysicalSpace(MappingRequest.PhysicalSpaceID);
 
-			PhysicalSpace->GetPagePool().MapPage(Space, PhysicalSpace, MappingRequest.PageTableLayerIndex, MappingRequest.vLevel, MappingRequest.vAddress, MappingRequest.Local_vLevel, MappingRequest.pAddress);
+			PhysicalSpace->GetPagePool().MapPage(Space, PhysicalSpace, MappingRequest.PageTableLayerIndex, MappingRequest.MaxLevel, MappingRequest.vLevel, MappingRequest.vAddress, MappingRequest.Local_vLevel, MappingRequest.pAddress);
 		}
 
 		// Update page mappings for any requested page that completed allocation this frame
@@ -2252,7 +2253,7 @@ void FVirtualTextureSystem::SubmitRequests(FRDGBuilder& GraphBuilder, ERHIFeatur
 				FVirtualTextureSpace* Space = GetSpace(MappingRequest.SpaceID);
 				check(RequestList->GetGroupMask(MappingRequest.LoadRequestIndex) & (1u << MappingRequest.ProducerPhysicalGroupIndex));
 
-				PhysicalSpace->GetPagePool().MapPage(Space, PhysicalSpace, MappingRequest.PageTableLayerIndex, MappingRequest.vLevel, MappingRequest.vAddress, MappingRequest.Local_vLevel, pAddress);
+				PhysicalSpace->GetPagePool().MapPage(Space, PhysicalSpace, MappingRequest.PageTableLayerIndex, MappingRequest.MaxLevel, MappingRequest.vLevel, MappingRequest.vAddress, MappingRequest.Local_vLevel, pAddress);
 			}
 		}
 	}
@@ -2287,17 +2288,18 @@ void FVirtualTextureSystem::SubmitRequests(FRDGBuilder& GraphBuilder, ERHIFeatur
 				const uint32 ProducerMipBias = AllocatedVT->GetUniqueProducerMipBias(ProducerIndex);
 				const uint32 WidthInTiles = Producer->GetWidthInTiles();
 				const uint32 HeightInTiles = Producer->GetHeightInTiles();
-				const uint32 MaxLevel = FMath::Min(Producer->GetMaxLevel(), AllocatedVT->GetMaxLevel() - ProducerMipBias);
+				const uint32 MaxLevel = AllocatedVT->GetMaxLevel();
+				const uint32 LocalMaxLevel = FMath::Min(Producer->GetMaxLevel(), MaxLevel - ProducerMipBias);
 
 				FVirtualTexturePhysicalSpace* PhysicalSpace = AllocatedVT->GetPhysicalSpaceForPageTableLayer(PageTableLayerIndex);
 				FTexturePagePool& PagePool = PhysicalSpace->GetPagePool();
 				FTexturePageMap& PageMap = Space->GetPageMapForPageTableLayer(PageTableLayerIndex);
 				
 				bool bIsLayerFullyMapped = false;
-				for (uint32 Local_vLevel = 0; Local_vLevel <= MaxLevel; ++Local_vLevel)
+				for (uint32 Local_vLevel = 0; Local_vLevel <= LocalMaxLevel; ++Local_vLevel)
 				{
 					const uint32 vLevel = Local_vLevel + ProducerMipBias;
-					check(vLevel <= AllocatedVT->GetMaxLevel());
+					check(vLevel <= MaxLevel);
 
 					const uint32 MipScaleFactor = (1u << Local_vLevel);
 					const uint32 LevelWidthInTiles = FMath::DivideAndRoundUp(WidthInTiles, MipScaleFactor);
@@ -2317,7 +2319,7 @@ void FVirtualTextureSystem::SubmitRequests(FRDGBuilder& GraphBuilder, ERHIFeatur
 								pAddress = PagePool.FindPageAddress(ProducerHandle, ProducerPhysicalGroupIndex, Local_vAddress, Local_vLevel);
 								if (pAddress != ~0u)
 								{
-									PagePool.MapPage(Space, PhysicalSpace, PageTableLayerIndex, vLevel, vAddress, vLevel, pAddress);
+									PagePool.MapPage(Space, PhysicalSpace, PageTableLayerIndex, MaxLevel, vLevel, vAddress, vLevel, pAddress);
 								}
 								else
 								{
