Added cvar to control cloud noise sample to merge according to their similarities.

diff --git a/Engine/Shaders/Private/VolumetricRenderTarget.usf b/Engine/Shaders/Private/VolumetricRenderTarget.usf
--- a/Engine/Shaders/Private/VolumetricRenderTarget.usf
+++ b/Engine/Shaders/Private/VolumetricRenderTarget.usf
@@ -339,6 +339,7 @@ float2 SafeSampleVolumetricDepthTexture(float2 UV)
 
 float4 VolumetricTextureSizeAndInvSize;
 float UvOffsetScale;
+float UvOffsetSampleAcceptanceWeight;
 float2 FullResolutionToVolumetricBufferResolutionScale;
 
 void ComposeVolumetricRTOverScenePS(
@@ -701,7 +702,7 @@ void ComposeVolumetricRTOverScenePS(
 		}
 
 		// Now continue to accumulate in DataAcc
-		const float Extinction = 30.0f;
+		const float Extinction = UvOffsetSampleAcceptanceWeight;
 		{
 			const float weight = exp2(-Extinction * abs(VolumeRGBT0.a - OutputRt0.a));
 			ValidSampleCount += weight;
diff --git a/Engine/Source/Runtime/Renderer/Private/VolumetricRenderTarget.cpp b/Engine/Source/Runtime/Renderer/Private/VolumetricRenderTarget.cpp
--- a/Engine/Source/Runtime/Renderer/Private/VolumetricRenderTarget.cpp
+++ b/Engine/Source/Runtime/Renderer/Private/VolumetricRenderTarget.cpp
@@ -25,6 +25,11 @@ static TAutoConsoleVariable<float> CVarVolumetricRenderTargetUvNoiseScale(
 	TEXT("Used when r.VolumetricRenderTarget.UpsamplingMode is in a mode using jitter - this value scales the amount of jitter."),
 	ECVF_SetByScalability);
 
+static TAutoConsoleVariable<float> CVarVolumetricRenderTargetUvNoiseSampleAcceptanceWeight(
+	TEXT("r.VolumetricRenderTarget.UvNoiseSampleAcceptanceWeight"), 20.0f,
+	TEXT("Used when r.VolumetricRenderTarget.UpsamplingMode is in a mode using jitter - this value control the acceptance of noisy cloud samples according to their similarities. A higher value means large differences will be less accepted for blending."),
+	ECVF_SetByScalability);
+
 static TAutoConsoleVariable<int32> CVarVolumetricRenderTargetMode(
 	TEXT("r.VolumetricRenderTarget.Mode"), 0,
 	TEXT("[0] trace quarter resolution + reconstruct at half resolution + upsample [1] trace half res + reconstruct full res + upsample [2] trace at quarter resolution + reconstruct full resolution (cannot intersect with opaque meshes and forces UpsamplingMode=2 [3] trace 1/8 resolution + reconstruct at half resolution + upsample)"),
@@ -46,6 +51,11 @@ static TAutoConsoleVariable<int32> CVarVolumetricRenderTargetReprojectionBoxCons
 	ECVF_SetByScalability | ECVF_RenderThreadSafe);
 
 
+static float GetUvNoiseSampleAcceptanceWeight()
+{
+	return FMath::Max(0.0f, CVarVolumetricRenderTargetUvNoiseSampleAcceptanceWeight.GetValueOnRenderThread());
+}
+
 static bool ShouldPipelineCompileVolumetricRenderTargetShaders(EShaderPlatform ShaderPlatform)
 {
 	return GetMaxSupportedFeatureLevel(ShaderPlatform) >= ERHIFeatureLevel::SM5;
@@ -553,6 +563,7 @@ class FComposeVolumetricRTOverScenePS : public FGlobalShader
 		SHADER_PARAMETER_SAMPLER(SamplerState, LinearTextureSampler)
 		RENDER_TARGET_BINDING_SLOTS()
 		SHADER_PARAMETER(float, UvOffsetScale)
+		SHADER_PARAMETER(float, UvOffsetSampleAcceptanceWeight)
 		SHADER_PARAMETER(FVector4, VolumetricTextureSizeAndInvSize)
 		SHADER_PARAMETER(FVector2D, FullResolutionToVolumetricBufferResolutionScale)
 		SHADER_PARAMETER(FVector2D, FullResolutionToWaterBufferScale)
@@ -628,6 +639,7 @@ void ComposeVolumetricRenderTargetOverScene(
 		PassParameters->VolumetricDepthTexture = VolumetricDepthTexture;
 		PassParameters->LinearTextureSampler = TStaticSamplerState<SF_Bilinear>::GetRHI();
 		PassParameters->UvOffsetScale = VolumetricCloudRT.GetUvNoiseScale();
+		PassParameters->UvOffsetSampleAcceptanceWeight = GetUvNoiseSampleAcceptanceWeight();
 		PassParameters->FullResolutionToVolumetricBufferResolutionScale = FVector2D(1.0f / float(GetMainDownsampleFactor(VRTMode)), float(GetMainDownsampleFactor(VRTMode)));
 		PassParameters->SceneTextures = SceneTextures.UniformBuffer;
 		GetTextureSafeUvCoordBound(PassParameters->VolumetricTexture, PassParameters->VolumetricTextureValidCoordRect, PassParameters->VolumetricTextureValidUvRect);
@@ -697,6 +709,7 @@ void ComposeVolumetricRenderTargetOverSceneUnderWater(
 		PassParameters->WaterLinearDepthTexture = WaterPassData.DepthTexture;
 		PassParameters->LinearTextureSampler = TStaticSamplerState<SF_Bilinear>::GetRHI();
 		PassParameters->UvOffsetScale = VolumetricCloudRT.GetUvNoiseScale();
+		PassParameters->UvOffsetSampleAcceptanceWeight = GetUvNoiseSampleAcceptanceWeight();
 		PassParameters->FullResolutionToVolumetricBufferResolutionScale = FVector2D(1.0f / float(GetMainDownsampleFactor(VRTMode)), float(GetMainDownsampleFactor(VRTMode)));
 		PassParameters->FullResolutionToWaterBufferScale = FVector2D(1.0f / WaterPassData.RefractionDownsampleFactor, WaterPassData.RefractionDownsampleFactor);
 		PassParameters->SceneWithoutSingleLayerWaterViewRect = FVector4(WaterPassViewData.ViewRect.Min.X, WaterPassViewData.ViewRect.Min.Y,
@@ -753,6 +766,7 @@ void ComposeVolumetricRenderTargetOverSceneForVisualization(
 		PassParameters->VolumetricDepthTexture = VolumetricDepthTexture;
 		PassParameters->LinearTextureSampler = TStaticSamplerState<SF_Bilinear>::GetRHI();
 		PassParameters->UvOffsetScale = VolumetricCloudRT.GetUvNoiseScale();
+		PassParameters->UvOffsetSampleAcceptanceWeight = GetUvNoiseSampleAcceptanceWeight();
 		PassParameters->FullResolutionToVolumetricBufferResolutionScale = FVector2D(1.0f / float(GetMainDownsampleFactor(VRTMode)), float(GetMainDownsampleFactor(VRTMode)));
 		PassParameters->SceneTextures = SceneTextures.UniformBuffer;
 		GetTextureSafeUvCoordBound(PassParameters->VolumetricTexture, PassParameters->VolumetricTextureValidCoordRect, PassParameters->VolumetricTextureValidUvRect);
