Avoid hair cluster resources allocation when a groom does not contain strands data.

diff --git a/Engine/Source/Runtime/Renderer/Private/HairStrands/HairStrandsInterface.cpp b/Engine/Source/Runtime/Renderer/Private/HairStrands/HairStrandsInterface.cpp
--- a/Engine/Source/Runtime/Renderer/Private/HairStrands/HairStrandsInterface.cpp
+++ b/Engine/Source/Runtime/Renderer/Private/HairStrands/HairStrandsInterface.cpp
@@ -306,7 +306,17 @@ void FHairGroupPublicData::InitRHI()
 	if (ClusterCount == 0)
 		return;
 
-	if (GUsingNullRHI) { return; }
+	bool bHasStrands = false;
+	for (const EHairGeometryType& Type : LODGeometryTypes)
+	{
+		if (Type == EHairGeometryType::Strands)
+		{
+			bHasStrands = true;
+			break;
+		}
+	}
+	
+	if (GUsingNullRHI || !bHasStrands) { return; }
 
 	FMemMark Mark(FMemStack::Get());
 	FRHICommandListImmediate& RHICmdList = FRHICommandListExecutor::GetImmediateCommandList();
