Avoid unnecessary CPU culling of nanite meshes in shadow rendering

diff --git a/Engine/Source/Runtime/Renderer/Private/PrimitiveSceneInfo.cpp b/Engine/Source/Runtime/Renderer/Private/PrimitiveSceneInfo.cpp
--- a/Engine/Source/Runtime/Renderer/Private/PrimitiveSceneInfo.cpp
+++ b/Engine/Source/Runtime/Renderer/Private/PrimitiveSceneInfo.cpp
@@ -112,6 +112,7 @@ FPrimitiveFlagsCompact::FPrimitiveFlagsCompact(const FPrimitiveSceneProxy* Proxy
 	: bCastDynamicShadow(Proxy->CastsDynamicShadow())
 	, bStaticLighting(Proxy->HasStaticLighting())
 	, bCastStaticShadow(Proxy->CastsStaticShadow())
+	, bIsNaniteMesh(Proxy->IsNaniteMesh())
 {}
 
 FPrimitiveSceneInfoCompact::FPrimitiveSceneInfoCompact(FPrimitiveSceneInfo* InPrimitiveSceneInfo) :
diff --git a/Engine/Source/Runtime/Renderer/Private/ShadowSetup.cpp b/Engine/Source/Runtime/Renderer/Private/ShadowSetup.cpp
--- a/Engine/Source/Runtime/Renderer/Private/ShadowSetup.cpp
+++ b/Engine/Source/Runtime/Renderer/Private/ShadowSetup.cpp
@@ -92,6 +92,15 @@ FAutoConsoleVariableRef CVarCachedWholeSceneShadowsCastFromMovablePrimitives(
 	ECVF_Scalability | ECVF_RenderThreadSafe
 	);
 
+// Temporary chicken bit to back out optimization if there are any issues
+int32 GSkipCullingNaniteMeshes = 1;
+FAutoConsoleVariableRef CVarSkipCullingNaniteMeshes(
+	TEXT("r.Shadow.SkipCullingNaniteMeshes"),
+	GSkipCullingNaniteMeshes,
+	TEXT("When enabled, CPU culling will ignore nanite meshes."),
+	ECVF_Scalability | ECVF_RenderThreadSafe
+);
+
 /** Can be used to visualize preshadow frustums when the shadowfrustums show flag is enabled. */
 static TAutoConsoleVariable<int32> CVarDrawPreshadowFrustum(
 	TEXT("r.Shadow.DrawPreshadowFrustums"),
@@ -3884,7 +3893,9 @@ struct FGatherShadowPrimitivesPacket
 			// Check all the primitives in this octree node.
 			for (const FPrimitiveSceneInfoCompact& PrimitiveSceneInfoCompact : Scene->PrimitiveOctree.GetElementsForNode(NodeIndex))
 			{
-				if (PrimitiveSceneInfoCompact.PrimitiveFlagsCompact.bCastDynamicShadow)
+				// Nanite has its own culling
+				if (PrimitiveSceneInfoCompact.PrimitiveFlagsCompact.bCastDynamicShadow &&
+					(!PrimitiveSceneInfoCompact.PrimitiveFlagsCompact.bIsNaniteMesh || GSkipCullingNaniteMeshes == 0))
 				{
 					FilterPrimitiveForShadows(PrimitiveSceneInfoCompact);
 				}
@@ -3899,7 +3910,9 @@ struct FGatherShadowPrimitivesPacket
 			{
 				const FPrimitiveFlagsCompact PrimitiveFlagsCompact = Scene->PrimitiveFlagsCompact[PrimitiveIndex];
 
-				if (PrimitiveFlagsCompact.bCastDynamicShadow)
+				// Nanite has its own culling
+				if (PrimitiveFlagsCompact.bCastDynamicShadow &&
+					(!PrimitiveFlagsCompact.bIsNaniteMesh || GSkipCullingNaniteMeshes == 0))
 				{
 					FPrimitiveSceneInfo* PrimitiveSceneInfo = Scene->Primitives[PrimitiveIndex];
 					const FPrimitiveSceneInfoCompact PrimitiveSceneInfoCompact(PrimitiveSceneInfo);
@@ -3986,11 +3999,6 @@ struct FGatherShadowPrimitivesPacket
 			const float CombinedRadiusSq = FMath::Square(ProjectedShadowInfo->ShadowBounds.W + PrimitiveBounds.SphereRadius);
 
 			// Note: Culling based on the primitive's bounds BEFORE dereferencing PrimitiveSceneInfo / PrimitiveProxy
-
-			bool A = PrimitiveDistanceFromCylinderAxisSq < CombinedRadiusSq;
-			bool B = !(ProjectedDistanceFromShadowOriginAlongLightDir < 0 && PrimitiveToShadowCenter.SizeSquared() > CombinedRadiusSq);
-			bool C = ProjectedShadowInfo->CascadeSettings.ShadowBoundsAccurate.IntersectBox(PrimitiveBounds.Origin, PrimitiveBounds.BoxExtent);
-
 			// Check if this primitive is in the shadow's cylinder
 			if (PrimitiveDistanceFromCylinderAxisSq < CombinedRadiusSq
 				// If the primitive is further along the cone axis than the shadow bounds origin, 
diff --git a/Engine/Source/Runtime/Renderer/Public/PrimitiveSceneInfo.h b/Engine/Source/Runtime/Renderer/Public/PrimitiveSceneInfo.h
--- a/Engine/Source/Runtime/Renderer/Public/PrimitiveSceneInfo.h
+++ b/Engine/Source/Runtime/Renderer/Public/PrimitiveSceneInfo.h
@@ -153,6 +153,8 @@ struct FPrimitiveFlagsCompact
 	uint8 bStaticLighting : 1;
 	/** True if the primitive casts static shadows. */
 	uint8 bCastStaticShadow : 1;
+	/** True if the primitive is a nanite mesh. */
+	uint8 bIsNaniteMesh : 1;
 
 	FPrimitiveFlagsCompact(const FPrimitiveSceneProxy* Proxy);
 };
