Restore post-process calibration visualization render logic overwritten by merge in CL 13834854.

diff --git a/Engine/Source/Runtime/Renderer/Private/DeferredShadingRenderer.cpp b/Engine/Source/Runtime/Renderer/Private/DeferredShadingRenderer.cpp
--- a/Engine/Source/Runtime/Renderer/Private/DeferredShadingRenderer.cpp
+++ b/Engine/Source/Runtime/Renderer/Private/DeferredShadingRenderer.cpp
@@ -2932,7 +2932,20 @@ void FDeferredShadingSceneRenderer::Render(FRDGBuilder& GraphBuilder)
 				const Nanite::FRasterResults* NaniteResults = bNaniteEnabled ? &NaniteRasterResults[ViewIndex] : nullptr;
 				RDG_GPU_MASK_SCOPE(GraphBuilder, View.GPUMask);
 				RDG_EVENT_SCOPE_CONDITIONAL(GraphBuilder, Views.Num() > 1, "View%d", ViewIndex);
-				AddPostProcessingPasses(GraphBuilder, View, PostProcessingInputs, NaniteResults, InstanceCullingManager);
+
+#if !(UE_BUILD_SHIPPING)
+				if (IsPostProcessVisualizeCalibrationMaterialEnabled(View))
+				{
+					const UMaterialInterface* DebugMaterialInterface = GetPostProcessVisualizeCalibrationMaterialInterface(View);
+					check(DebugMaterialInterface);
+
+					AddVisualizeCalibrationMaterialPostProcessingPasses(GraphBuilder, View, PostProcessingInputs, DebugMaterialInterface);
+				}
+				else
+#endif
+				{
+					AddPostProcessingPasses(GraphBuilder, View, PostProcessingInputs, NaniteResults, InstanceCullingManager);
+				}
 			}
 		}
 	}
