Fix TAA corruption caused by hair velocity mask.

diff --git a/Engine/Source/Runtime/Renderer/Private/HairStrands/HairStrandsVisibility.cpp b/Engine/Source/Runtime/Renderer/Private/HairStrands/HairStrandsVisibility.cpp
--- a/Engine/Source/Runtime/Renderer/Private/HairStrands/HairStrandsVisibility.cpp
+++ b/Engine/Source/Runtime/Renderer/Private/HairStrands/HairStrandsVisibility.cpp
@@ -1275,15 +1275,19 @@ static void AddHairVelocityPass(
 	PassParameters->OutVelocityTexture = GraphBuilder.CreateUAV(OutVelocityTexture);
 	PassParameters->OutResolveMaskTexture = GraphBuilder.CreateUAV(OutResolveMaskTexture);
 
-	FIntRect TotalRect = ComputeVisibleHairStrandsMacroGroupsRect(View.ViewRect, MacroGroupDatas);
-
-	// Snap the rect onto thread group boundary
 	const FIntPoint GroupSize = GetVendorOptimalGroupSize2D();
+	// We don't use the CPU screen projection for running the velocity pass, as we need to clear the entire 
+	// velocity mask through the UAV write, otherwise the mask will be partially invalid.
+#if 1
+	FIntRect TotalRect = View.ViewRect; 
+#else
+	// Snap the rect onto thread group boundary
+	FIntRect TotalRect = ComputeVisibleHairStrandsMacroGroupsRect(View.ViewRect, MacroGroupDatas);
 	TotalRect.Min.X = FMath::FloorToInt(float(TotalRect.Min.X) / float(GroupSize.X)) * GroupSize.X;
 	TotalRect.Min.Y = FMath::FloorToInt(float(TotalRect.Min.Y) / float(GroupSize.Y)) * GroupSize.Y;
 	TotalRect.Max.X = FMath::CeilToInt(float(TotalRect.Max.X) / float(GroupSize.X)) * GroupSize.X;
 	TotalRect.Max.Y = FMath::CeilToInt(float(TotalRect.Max.Y) / float(GroupSize.Y)) * GroupSize.Y;
-
+#endif
 	FIntPoint RectResolution(TotalRect.Width(), TotalRect.Height());
 	PassParameters->ResolutionOffset = FIntPoint(TotalRect.Min.X, TotalRect.Min.Y);
 
