Support separate pass for each view.

diff --git a/Engine/Source/Runtime/Renderer/Private/MobileShadingRenderer.cpp b/Engine/Source/Runtime/Renderer/Private/MobileShadingRenderer.cpp
--- a/Engine/Source/Runtime/Renderer/Private/MobileShadingRenderer.cpp
+++ b/Engine/Source/Runtime/Renderer/Private/MobileShadingRenderer.cpp
@@ -424,6 +424,7 @@ void FMobileSceneRenderer::InitViews(FRDGBuilder& GraphBuilder, FSceneTexturesCo
 	const bool bForceDepthResolve = (CVarMobileForceDepthResolve.GetValueOnRenderThread() == 1);
 	const bool bSeparateTranslucencyActive = IsMobileSeparateTranslucencyActive(Views.GetData(), Views.Num()); 
 	const bool bPostProcessUsesSceneDepth = PostProcessUsesSceneDepth(Views[0]) || IsMobileDistortionActive(Views[0]);
+	const bool bRequireSeparateViewPass = !Views[0].bIsMobileMultiViewEnabled && Views.Num() > 0;
 	bRequiresMultiPass = RequiresMultiPass(RHICmdList, Views[0]);
 	bKeepDepthContent =
 		bRequiresMultiPass ||
@@ -434,7 +435,8 @@ void FMobileSceneRenderer::InitViews(FRDGBuilder& GraphBuilder, FSceneTexturesCo
 		bSeparateTranslucencyActive ||
 		Views[0].bIsReflectionCapture ||
 		(bDeferredShading && bPostProcessUsesSceneDepth) ||
-		bShouldRenderVelocities;
+		bShouldRenderVelocities ||
+		bRequireSeparateViewPass;
 	// never keep MSAA depth
 	bKeepDepthContent = (NumMSAASamples > 1 ? false : bKeepDepthContent);
 
@@ -607,13 +609,13 @@ void FMobileSceneRenderer::RenderFullDepthPrepass(FRDGBuilder& GraphBuilder, FSc
 			RDG_EVENT_NAME("FullDepthPrepass"),
 			PassParameters,
 			ERDGPassFlags::Raster,
-			[this, PassParameters, &View](FRHICommandListImmediate& RHICmdList)
+			[this, PassParameters, ViewIndex, &View](FRHICommandListImmediate& RHICmdList)
 			{
 				RenderPrePass(RHICmdList, View, &PassParameters->InstanceCullingDrawParamsDepth);
 
 				// Issue occlusion queries
 				RHICmdList.SetCurrentStat(GET_STATID(STAT_CLMM_Occlusion));
-				RenderOcclusion(RHICmdList, View);
+				RenderOcclusion(RHICmdList, ViewIndex);
 			});
 	}
 }
@@ -931,6 +933,13 @@ void FMobileSceneRenderer::RenderForward(FRDGBuilder& GraphBuilder, FRDGTextureR
 
 	for (int32 ViewIndex = 0; ViewIndex < Views.Num(); ViewIndex++)
 	{
+		if (ViewIndex > 0)
+		{
+			BasePassRenderTargets[0].SetLoadAction(ERenderTargetLoadAction::ELoad);
+			BasePassRenderTargets.DepthStencil.SetDepthLoadAction(ERenderTargetLoadAction::ELoad);
+			BasePassRenderTargets.DepthStencil.SetStencilLoadAction(ERenderTargetLoadAction::ELoad);
+		}
+
 		FViewInfo& View = Views[ViewIndex];
 
 		SCOPED_GPU_MASK(GraphBuilder.RHICmdList, !View.IsInstancedStereoPass() ? View.GPUMask : (Views[0].GPUMask | Views[1].GPUMask));
@@ -978,10 +987,11 @@ void FMobileSceneRenderer::RenderForwardSinglePass(FRDGBuilder& GraphBuilder, FM
 	GraphBuilder.AddPass(
 		RDG_EVENT_NAME("SceneColorRendering"),
 		PassParameters,
-		ERDGPassFlags::Raster,
+		// the second view pass should not be merged with the first view pass on mobile since the subpass would not work properly.
+		ERDGPassFlags::Raster | ERDGPassFlags::NeverMerge,
 		[this, PassParameters, ViewIndex, &View, &SceneTextures](FRHICommandListImmediate& RHICmdList)
 	{
-		if (GIsEditor && !View.bIsSceneCapture)
+		if (GIsEditor && !View.bIsSceneCapture && ViewIndex == 0)
 		{
 			DrawClearQuad(RHICmdList, View.BackgroundColor);
 		}
@@ -999,7 +1009,7 @@ void FMobileSceneRenderer::RenderForwardSinglePass(FRDGBuilder& GraphBuilder, FM
 		{
 			// Issue occlusion queries
 			RHICmdList.SetCurrentStat(GET_STATID(STAT_CLMM_Occlusion));
-			RenderOcclusion(RHICmdList, View);
+			RenderOcclusion(RHICmdList, ViewIndex);
 		}
 		PostRenderBasePass(RHICmdList, View);
 		// scene depth is read only and can be fetched
@@ -1015,7 +1025,7 @@ void FMobileSceneRenderer::RenderForwardSinglePass(FRDGBuilder& GraphBuilder, FM
 			// flush
 			RHICmdList.SubmitCommandsHint();
 			// Issue occlusion queries
-			RenderOcclusion(RHICmdList, View);
+			RenderOcclusion(RHICmdList, ViewIndex);
 		}
 		// Pre-tonemap before MSAA resolve (iOS only)
 		PreTonemapMSAA(RHICmdList, SceneTextures);
@@ -1028,9 +1038,9 @@ void FMobileSceneRenderer::RenderForwardMultiPass(FRDGBuilder& GraphBuilder, FMo
 		RDG_EVENT_NAME("SceneColorRendering"),
 		PassParameters,
 		ERDGPassFlags::Raster,
-		[this, PassParameters, &View, &SceneTextures](FRHICommandListImmediate& RHICmdList)
+		[this, PassParameters, ViewIndex, &View, &SceneTextures](FRHICommandListImmediate& RHICmdList)
 	{
-		if (GIsEditor && !View.bIsSceneCapture)
+		if (GIsEditor && !View.bIsSceneCapture && ViewIndex == 0)
 		{
 			DrawClearQuad(RHICmdList, View.BackgroundColor);
 		}
@@ -1044,7 +1054,7 @@ void FMobileSceneRenderer::RenderForwardMultiPass(FRDGBuilder& GraphBuilder, FMo
 		RenderMobileDebugView(RHICmdList, View);
 		// Issue occlusion queries
 		RHICmdList.SetCurrentStat(GET_STATID(STAT_CLMM_Occlusion));
-		RenderOcclusion(RHICmdList, View);
+		RenderOcclusion(RHICmdList, ViewIndex);
 		PostRenderBasePass(RHICmdList, View);
 	});
 
@@ -1276,7 +1286,7 @@ void FMobileSceneRenderer::RenderDeferredSinglePass(FRDGBuilder& GraphBuilder, c
 		RHICmdList.ImmediateFlush(EImmediateFlushType::DispatchToRHIThread);
 		// Issue occlusion queries
 		RHICmdList.SetCurrentStat(GET_STATID(STAT_CLMM_Occlusion));
-		RenderOcclusion(RHICmdList, View);
+		RenderOcclusion(RHICmdList, ViewIndex);
 		PostRenderBasePass(RHICmdList, View);
 		// SceneColor + GBuffer write, SceneDepth is read only
 		RHICmdList.NextSubpass();
@@ -1301,7 +1311,7 @@ void FMobileSceneRenderer::RenderDeferredMultiPass(FRDGBuilder& GraphBuilder, cl
 		RDG_EVENT_NAME("SceneColorRendering"),
 		PassParameters,
 		ERDGPassFlags::Raster,
-		[this, PassParameters, &View, &SceneTextures](FRHICommandListImmediate& RHICmdList)
+		[this, PassParameters, ViewIndex, &View, &SceneTextures](FRHICommandListImmediate& RHICmdList)
 	{
 		if (GIsEditor && !View.bIsSceneCapture)
 		{
@@ -1316,7 +1326,7 @@ void FMobileSceneRenderer::RenderDeferredMultiPass(FRDGBuilder& GraphBuilder, cl
 		RenderMobileBasePass(RHICmdList, View, &PassParameters->InstanceCullingDrawParamsOpaque);
 		// Issue occlusion queries
 		RHICmdList.SetCurrentStat(GET_STATID(STAT_CLMM_Occlusion));
-		RenderOcclusion(RHICmdList, View);
+		RenderOcclusion(RHICmdList, ViewIndex);
 		PostRenderBasePass(RHICmdList, View);
 	});
 
