Revert the changes of CL#15987600 which contain other changes.

diff --git a/Engine/Source/Runtime/Renderer/Private/MobileShadingRenderer.cpp b/Engine/Source/Runtime/Renderer/Private/MobileShadingRenderer.cpp
--- a/Engine/Source/Runtime/Renderer/Private/MobileShadingRenderer.cpp
+++ b/Engine/Source/Runtime/Renderer/Private/MobileShadingRenderer.cpp
@@ -424,7 +424,6 @@ void FMobileSceneRenderer::InitViews(FRDGBuilder& GraphBuilder, FSceneTexturesCo
 	const bool bForceDepthResolve = (CVarMobileForceDepthResolve.GetValueOnRenderThread() == 1);
 	const bool bSeparateTranslucencyActive = IsMobileSeparateTranslucencyActive(Views.GetData(), Views.Num()); 
 	const bool bPostProcessUsesSceneDepth = PostProcessUsesSceneDepth(Views[0]) || IsMobileDistortionActive(Views[0]);
-	const bool bRequireSeparateViewPass = !Views[0].bIsMobileMultiViewEnabled && Views.Num() > 0;
 	bRequiresMultiPass = RequiresMultiPass(RHICmdList, Views[0]);
 	bKeepDepthContent =
 		bRequiresMultiPass ||
@@ -435,8 +434,7 @@ void FMobileSceneRenderer::InitViews(FRDGBuilder& GraphBuilder, FSceneTexturesCo
 		bSeparateTranslucencyActive ||
 		Views[0].bIsReflectionCapture ||
 		(bDeferredShading && bPostProcessUsesSceneDepth) ||
-		bShouldRenderVelocities ||
-		bRequireSeparateViewPass;
+		bShouldRenderVelocities;
 	// never keep MSAA depth
 	bKeepDepthContent = (NumMSAASamples > 1 ? false : bKeepDepthContent);
 
@@ -609,13 +607,13 @@ void FMobileSceneRenderer::RenderFullDepthPrepass(FRDGBuilder& GraphBuilder, FSc
 			RDG_EVENT_NAME("FullDepthPrepass"),
 			PassParameters,
 			ERDGPassFlags::Raster,
-			[this, PassParameters, ViewIndex, &View](FRHICommandListImmediate& RHICmdList)
+			[this, PassParameters, &View](FRHICommandListImmediate& RHICmdList)
 			{
 				RenderPrePass(RHICmdList, View, &PassParameters->InstanceCullingDrawParamsDepth);
 
 				// Issue occlusion queries
 				RHICmdList.SetCurrentStat(GET_STATID(STAT_CLMM_Occlusion));
-				RenderOcclusion(RHICmdList, ViewIndex);
+				RenderOcclusion(RHICmdList, View);
 			});
 	}
 }
@@ -933,13 +931,6 @@ void FMobileSceneRenderer::RenderForward(FRDGBuilder& GraphBuilder, FRDGTextureR
 
 	for (int32 ViewIndex = 0; ViewIndex < Views.Num(); ViewIndex++)
 	{
-		if (ViewIndex > 0)
-		{
-			BasePassRenderTargets[0].SetLoadAction(ERenderTargetLoadAction::ELoad);
-			BasePassRenderTargets.DepthStencil.SetDepthLoadAction(ERenderTargetLoadAction::ELoad);
-			BasePassRenderTargets.DepthStencil.SetStencilLoadAction(ERenderTargetLoadAction::ELoad);
-		}
-
 		FViewInfo& View = Views[ViewIndex];
 
 		SCOPED_GPU_MASK(GraphBuilder.RHICmdList, !View.IsInstancedStereoPass() ? View.GPUMask : (Views[0].GPUMask | Views[1].GPUMask));
@@ -987,11 +978,10 @@ void FMobileSceneRenderer::RenderForwardSinglePass(FRDGBuilder& GraphBuilder, FM
 	GraphBuilder.AddPass(
 		RDG_EVENT_NAME("SceneColorRendering"),
 		PassParameters,
-		// the second view pass should not be merged with the first view pass on mobile since the subpass would not work properly.
-		ERDGPassFlags::Raster | ERDGPassFlags::NeverMerge,
+		ERDGPassFlags::Raster,
 		[this, PassParameters, ViewIndex, &View, &SceneTextures](FRHICommandListImmediate& RHICmdList)
 	{
-		if (GIsEditor && !View.bIsSceneCapture && ViewIndex == 0)
+		if (GIsEditor && !View.bIsSceneCapture)
 		{
 			DrawClearQuad(RHICmdList, View.BackgroundColor);
 		}
@@ -1009,7 +999,7 @@ void FMobileSceneRenderer::RenderForwardSinglePass(FRDGBuilder& GraphBuilder, FM
 		{
 			// Issue occlusion queries
 			RHICmdList.SetCurrentStat(GET_STATID(STAT_CLMM_Occlusion));
-			RenderOcclusion(RHICmdList, ViewIndex);
+			RenderOcclusion(RHICmdList, View);
 		}
 		PostRenderBasePass(RHICmdList, View);
 		// scene depth is read only and can be fetched
@@ -1025,7 +1015,7 @@ void FMobileSceneRenderer::RenderForwardSinglePass(FRDGBuilder& GraphBuilder, FM
 			// flush
 			RHICmdList.SubmitCommandsHint();
 			// Issue occlusion queries
-			RenderOcclusion(RHICmdList, ViewIndex);
+			RenderOcclusion(RHICmdList, View);
 		}
 		// Pre-tonemap before MSAA resolve (iOS only)
 		PreTonemapMSAA(RHICmdList, SceneTextures);
@@ -1038,9 +1028,9 @@ void FMobileSceneRenderer::RenderForwardMultiPass(FRDGBuilder& GraphBuilder, FMo
 		RDG_EVENT_NAME("SceneColorRendering"),
 		PassParameters,
 		ERDGPassFlags::Raster,
-		[this, PassParameters, ViewIndex, &View, &SceneTextures](FRHICommandListImmediate& RHICmdList)
+		[this, PassParameters, &View, &SceneTextures](FRHICommandListImmediate& RHICmdList)
 	{
-		if (GIsEditor && !View.bIsSceneCapture && ViewIndex == 0)
+		if (GIsEditor && !View.bIsSceneCapture)
 		{
 			DrawClearQuad(RHICmdList, View.BackgroundColor);
 		}
@@ -1054,7 +1044,7 @@ void FMobileSceneRenderer::RenderForwardMultiPass(FRDGBuilder& GraphBuilder, FMo
 		RenderMobileDebugView(RHICmdList, View);
 		// Issue occlusion queries
 		RHICmdList.SetCurrentStat(GET_STATID(STAT_CLMM_Occlusion));
-		RenderOcclusion(RHICmdList, ViewIndex);
+		RenderOcclusion(RHICmdList, View);
 		PostRenderBasePass(RHICmdList, View);
 	});
 
@@ -1286,7 +1276,7 @@ void FMobileSceneRenderer::RenderDeferredSinglePass(FRDGBuilder& GraphBuilder, c
 		RHICmdList.ImmediateFlush(EImmediateFlushType::DispatchToRHIThread);
 		// Issue occlusion queries
 		RHICmdList.SetCurrentStat(GET_STATID(STAT_CLMM_Occlusion));
-		RenderOcclusion(RHICmdList, ViewIndex);
+		RenderOcclusion(RHICmdList, View);
 		PostRenderBasePass(RHICmdList, View);
 		// SceneColor + GBuffer write, SceneDepth is read only
 		RHICmdList.NextSubpass();
@@ -1311,7 +1301,7 @@ void FMobileSceneRenderer::RenderDeferredMultiPass(FRDGBuilder& GraphBuilder, cl
 		RDG_EVENT_NAME("SceneColorRendering"),
 		PassParameters,
 		ERDGPassFlags::Raster,
-		[this, PassParameters, ViewIndex, &View, &SceneTextures](FRHICommandListImmediate& RHICmdList)
+		[this, PassParameters, &View, &SceneTextures](FRHICommandListImmediate& RHICmdList)
 	{
 		if (GIsEditor && !View.bIsSceneCapture)
 		{
@@ -1326,7 +1316,7 @@ void FMobileSceneRenderer::RenderDeferredMultiPass(FRDGBuilder& GraphBuilder, cl
 		RenderMobileBasePass(RHICmdList, View, &PassParameters->InstanceCullingDrawParamsOpaque);
 		// Issue occlusion queries
 		RHICmdList.SetCurrentStat(GET_STATID(STAT_CLMM_Occlusion));
-		RenderOcclusion(RHICmdList, ViewIndex);
+		RenderOcclusion(RHICmdList, View);
 		PostRenderBasePass(RHICmdList, View);
 	});
 
